Fixes shader compilation failure in Niagara shader compiler when running QAGame on windows.

diff --git a/Engine/Shaders/Private/NiagaraSimulationShader.usf b/Engine/Shaders/Private/NiagaraSimulationShader.usf
new file mode 100644
--- /dev/null
+++ b/Engine/Shaders/Private/NiagaraSimulationShader.usf
@@ -0,0 +1,94 @@
+// Copyright 1998-2016 Epic Games, Inc. All Rights Reserved.
+
+/*=============================================================================
+NiagaraSimulationShader.usf: 
+=============================================================================*/
+
+//Should include Common.usf?
+#include "Definitions.usf"
+
+//TODO: Add common utility code for things such as reading data sources, reading from and writing to events.
+
+//Temporary random hacks
+float4 rand(float4 x) 
+{
+	return float4(rand(x.x), rand(x.y), rand(x.z), rand(x.w));
+}
+
+float3 rand(float3 x)
+{
+	return float3(rand(x.x), rand(x.y), rand(x.z));
+}
+
+float2 rand(float2 x)
+{
+	return float2(rand(x.x), rand(x.y));
+}
+
+float noise(float x);
+float noise(float2 x);
+float noise(float3 x);
+
+//Total hack to get around the cross compiler converting fmod() to "X - (Y * trunc(X/Y))";
+//On gpu just define these as fmod(x,y)
+float4 Modulo(float4 x, float4 y);
+float3 Modulo(float3 x, float3 y);
+float2 Modulo(float2 x, float2 y);
+float Modulo(float x, float y);
+
+/** Returns the index for this particle in the current execution context. On gpu this'll likely be derived from DispatchThreadId */
+int ExecIndex();
+
+/** Hooks into the stats system. */
+void EnterStatScope(int ScopeIdx);
+void ExitStatScope();
+
+//#if VM_SIMULATION
+
+//Some functions that we use to map to special VM operations for reading in data.
+//TODO: replace with proper buffer reading capability and use standard hlsl.
+
+int AcquireIndex(int DataSetID, bool DoAcquire);
+
+float InputDataFloat(int DataSetIndex, int RegisterIdx);  //InstanceIndex is 0 for main dataset
+int InputDataInt(int DataSetIndex, int RegisterIdx);
+
+float InputDataNoadvanceFloat(int DataSetIndex, int RegisterIdx);  //InstanceIndex is 0 for main dataset
+int InputDataNoadvanceInt(int DataSetIndex, int RegisterIdx);
+
+void OutputDataFloat(int DataSetIndex, int RegisterIndex, int InstanceIndex, float Value);
+void OutputDataInt(int DataSetIndex, int RegisterIndex, int InstanceIndex, int Value);
+
+
+// #else
+// #endif
+
+//void SimulateMain()
+//{
+//	float I = InputDataFloat(0);
+//	OutputDataFloat(0, I);
+//}
+
+//Include the simulation shader code generated by the node graph.
+#include "/Engine/Generated/NiagaraSimulation.ush"
+
+//The VectorVM shader needs only the generated code but a compute shader requires some wrapping.
+#if GPU_SIMULATION_SHADER
+
+//Need accessors wrapper functions for Event and other particle data buffers?
+
+StructuredBuffer<FSimulationAttrs> Input;
+RWStructuredBuffer<FSimulationAttrs> Output;
+
+[numthreads(THREADGROUP_SIZE, 1, 1)]
+void SimulateMainCompute(
+	uint3 GroupId : SV_GroupID,
+	uint3 DispatchThreadId : SV_DispatchThreadID,
+	uint3 GroupThreadId : SV_GroupThreadID)
+{
+	uint Index = GroupId.x * THREADGROUP_SIZE + (GroupThreadId.x - (GroupThreadId.x % SIMULATION_INSTANCES_PER_THREAD));
+	
+	SimulateMain(Input[Index], Output[Index]);
+}
+
+#endif//GPU_SIMULATION_SHADER
diff --git a/Engine/Source/Editor/NiagaraEditor/Private/NiagaraCompiler.cpp b/Engine/Source/Editor/NiagaraEditor/Private/NiagaraCompiler.cpp
--- a/Engine/Source/Editor/NiagaraEditor/Private/NiagaraCompiler.cpp
+++ b/Engine/Source/Editor/NiagaraEditor/Private/NiagaraCompiler.cpp
@@ -1076,7 +1076,7 @@ const FNiagaraCompileResults& FHlslNiagaraCompiler::CompileScript(UNiagaraScript
 		//But that adds the complexity of a new shader type, new shader class and a new shader map to contain them etc.
 		//Can do things simply for now.
 		FShaderCompilerInput Input;
-		Input.VirtualSourceFilePath = TEXT("NiagaraSimulationShader");
+		Input.VirtualSourceFilePath = TEXT("/Engine/Private/NiagaraSimulationShader.usf");
 		Input.EntryPointName = TEXT("SimulateMain");
 		//Input.Target = FShaderTarget(SF_Vertex, SP_VECTORVM_1_0);
 		Input.Environment.IncludeVirtualPathToContentsMap.Add(TEXT("/Engine/Generated/NiagaraSimulation.ush"), StringToArray<ANSICHAR>(*HlslOutput, HlslOutput.Len() + 1));
