Branch snapshot for CL 3555263

[CL 3555263 in Main branch]

diff --git a/Engine/Shaders/Private/ACES.ush b/Engine/Shaders/Private/ACES.ush
new file mode 100644
--- /dev/null
+++ b/Engine/Shaders/Private/ACES.ush
@@ -0,0 +1,1187 @@
+// Copyright 1998-2017 Epic Games, Inc. All Rights Reserved.
+
+// #todo-rco: Fix in the cross-compiler
+#if COMPILER_GLSL || COMPILER_GLSL_ES2 || COMPILER_GLSL_ES3_1 || COMPILER_METAL || VULKAN_PROFILE || COMPILER_HLSLCC
+#define log10(x)  log(x) / log(10.0)
+#endif
+
+/*
+=============================================================================
+	ACES: Academy Color Encoding System
+	[REDACTED_PROFILE_URL]/aces-dev/tree/v1.0
+	
+	License Terms for Academy Color Encoding System Components
+
+	Academy Color Encoding System (ACES) software and tools are provided by the Academy under 
+	the following terms and conditions: A worldwide, royalty-free, non-exclusive right to copy, modify, create
+	derivatives, and use, in source and binary forms, is hereby granted, subject to acceptance of this license.
+
+	Copyright © 2013 Academy of Motion Picture Arts and Sciences (A.M.P.A.S.). Portions contributed by
+	others as indicated. All rights reserved.
+
+	Performance of any of the aforementioned acts indicates acceptance to be bound by the following 
+	terms and conditions:
+
+	 *	Copies of source code, in whole or in part, must retain the above copyright 
+		notice, this list of conditions and the Disclaimer of Warranty.
+	 *	Use in binary form must retain the above copyright notice, this list of 
+		conditions and the Disclaimer of Warranty in the documentation and/or other 
+		materials provided with the distribution.
+	 *	Nothing in this license shall be deemed to grant any rights to trademarks,
+		copyrights, patents, trade secrets or any other intellectual property of 
+		A.M.P.A.S. or any contributors, except as expressly stated herein.
+	 *	Neither the name "A.M.P.A.S." nor the name of any other contributors to this 
+		software may be used to endorse or promote products derivative of or based on
+		this software without express prior written permission of A.M.P.A.S. or the
+		contributors, as appropriate.
+
+	This license shall be construed pursuant to the laws of the State of California,
+	and any disputes related thereto shall be subject to the jurisdiction of the courts therein.
+
+	Disclaimer of Warranty: THIS SOFTWARE IS PROVIDED BY A.M.P.A.S. AND CONTRIBUTORS "AS
+	IS" AND ANY EXPRESS OR IMPLIED WARRANTIES, INCLUDING, BUT NOT LIMITED TO, THE
+	IMPLIED WARRANTIES OF MERCHANTABILITY, FITNESS FOR A PARTICULAR PURPOSE, AND 
+	NON-INFRINGEMENT ARE DISCLAIMED. IN NO EVENT SHALL A.M.P.A.S., OR ANY 
+	CONTRIBUTORS OR DISTRIBUTORS, BE LIABLE FOR ANY DIRECT, INDIRECT, INCIDENTAL,
+	SPECIAL, EXEMPLARY, RESITUTIONARY, OR CONSEQUENTIAL DAMAGES (INCLUDING, BUT 
+	NOT LIMITED TO, PROCUREMENT OF SUBSTITUTE GOODS OR SERVICES; LOSS OF USE, 
+	DATA, OR PROFITS; OR BUSINESS INTERRUPTION) HOWEVER CAUSED AND ON ANY THEORY
+	OF LIABILITY, WHETHER IN CONTRACT, STRICT LIABILITY, OR TORT (INCLUDING 
+	NEGLIGENCE OR OTHERWISE) ARISING IN ANY WAY OUT OF THE USE OF THIS SOFTWARE,
+	EVEN IF ADVISED OF THE POSSIBILITY OF SUCH DAMAGE.
+
+	WITHOUT LIMITING THE GENERALITY OF THE FOREGOING, THE ACADEMY SPECIFICALLY 
+	DISCLAIMS ANY REPRESENTATIONS OR WARRANTIES WHATSOEVER RELATED TO PATENT OR 
+	OTHER INTELLECTUAL PROPERTY RIGHTS IN THE ACADEMY COLOR ENCODING SYSTEM, OR 
+	APPLICATIONS THEREOF, HELD BY PARTIES OTHER THAN A.M.P.A.S.,WHETHER DISCLOSED
+	OR UNDISCLOSED.
+=============================================================================
+*/
+
+static const float3x3 AP0_2_XYZ_MAT =
+{
+	0.9525523959, 0.0000000000, 0.0000936786,
+	0.3439664498, 0.7281660966,-0.0721325464,
+	0.0000000000, 0.0000000000, 1.0088251844,
+};
+
+static const float3x3 XYZ_2_AP0_MAT =
+{
+	 1.0498110175, 0.0000000000,-0.0000974845,
+	-0.4959030231, 1.3733130458, 0.0982400361,
+	 0.0000000000, 0.0000000000, 0.9912520182,
+};
+
+static const float3x3 AP1_2_XYZ_MAT = 
+{
+	 0.6624541811, 0.1340042065, 0.1561876870,
+	 0.2722287168, 0.6740817658, 0.0536895174,
+	-0.0055746495, 0.0040607335, 1.0103391003,
+};
+
+static const float3x3 XYZ_2_AP1_MAT =
+{
+	 1.6410233797, -0.3248032942, -0.2364246952,
+	-0.6636628587,  1.6153315917,  0.0167563477,
+	 0.0117218943, -0.0082844420,  0.9883948585,
+};
+
+static const float3x3 AP0_2_AP1_MAT = //mul( AP0_2_XYZ_MAT, XYZ_2_AP1_MAT );
+{
+	 1.4514393161, -0.2365107469, -0.2149285693,
+	-0.0765537734,  1.1762296998, -0.0996759264,
+	 0.0083161484, -0.0060324498,  0.9977163014,
+};
+
+static const float3x3 AP1_2_AP0_MAT = //mul( AP1_2_XYZ_MAT, XYZ_2_AP0_MAT );
+{
+	 0.6954522414,  0.1406786965,  0.1638690622,
+	 0.0447945634,  0.8596711185,  0.0955343182,
+	-0.0055258826,  0.0040252103,  1.0015006723,
+};
+
+static const float3 AP1_RGB2Y =
+{
+	0.2722287168, //AP1_2_XYZ_MAT[0][1],
+	0.6740817658, //AP1_2_XYZ_MAT[1][1],
+	0.0536895174, //AP1_2_XYZ_MAT[2][1]
+};
+
+// REC 709 primaries
+static const float3x3 XYZ_2_sRGB_MAT =
+{
+	 3.2409699419, -1.5373831776, -0.4986107603,
+	-0.9692436363,  1.8759675015,  0.0415550574,
+	 0.0556300797, -0.2039769589,  1.0569715142,
+};
+
+static const float3x3 sRGB_2_XYZ_MAT =
+{
+	0.4124564, 0.3575761, 0.1804375,
+	0.2126729, 0.7151522, 0.0721750,
+	0.0193339, 0.1191920, 0.9503041,
+};
+
+// REC 2020 primaries
+static const float3x3 XYZ_2_Rec2020_MAT =
+{
+	 1.7166084, -0.3556621, -0.2533601,
+	-0.6666829,  1.6164776,  0.0157685,
+	 0.0176422, -0.0427763,  0.94222867	
+};
+
+static const float3x3 Rec2020_2_XYZ_MAT =
+{
+	0.6369736, 0.1446172, 0.1688585,
+	0.2627066, 0.6779996, 0.0592938,
+	0.0000000, 0.0280728, 1.0608437
+};
+
+// P3, D65 primaries
+static const float3x3 XYZ_2_P3D65_MAT =
+{
+	 2.4933963, -0.9313459, -0.4026945,
+	-0.8294868,  1.7626597,  0.0236246,
+	 0.0358507, -0.0761827,  0.9570140
+};
+
+static const float3x3 P3D65_2_XYZ_MAT =
+{
+	0.4865906, 0.2656683, 0.1981905,
+	0.2289838, 0.6917402, 0.0792762,
+	0.0000000, 0.0451135, 1.0438031
+};
+
+// Bradford chromatic adaptation transforms between ACES white point (D60) and sRGB white point (D65)
+static const float3x3 D65_2_D60_CAT =
+{
+	 1.01303,    0.00610531, -0.014971,
+	 0.00769823, 0.998165,   -0.00503203,
+	-0.00284131, 0.00468516,  0.924507,
+};
+
+static const float3x3 D60_2_D65_CAT =
+{
+	 0.987224,   -0.00611327, 0.0159533,
+	-0.00759836,  1.00186,    0.00533002,
+	 0.00307257, -0.00509595, 1.08168,
+};
+
+static const float HALF_MAX = 65504.0;
+
+float rgb_2_saturation( float3 rgb )
+{
+	float minrgb = min( min(rgb.r, rgb.g ), rgb.b );
+	float maxrgb = max( max(rgb.r, rgb.g ), rgb.b );
+	return ( max( maxrgb, 1e-10 ) - max( minrgb, 1e-10 ) ) / max( maxrgb, 1e-2 );
+}
+
+float glow_fwd( float ycIn, float glowGainIn, float glowMid)
+{
+   float glowGainOut;
+
+   if (ycIn <= 2./3. * glowMid) {
+     glowGainOut = glowGainIn;
+   } else if ( ycIn >= 2 * glowMid) {
+     glowGainOut = 0;
+   } else {
+     glowGainOut = glowGainIn * (glowMid / ycIn - 0.5);
+   }
+
+   return glowGainOut;
+}
+
+float glow_inv( float ycOut, float glowGainIn, float glowMid)
+{
+    float glowGainOut;
+
+    if (ycOut <= ((1 + glowGainIn) * 2./3. * glowMid)) {
+	    glowGainOut = -glowGainIn / (1 + glowGainIn);
+    } else if ( ycOut >= (2. * glowMid)) {
+		glowGainOut = 0.;
+	} else {
+		glowGainOut = glowGainIn * (glowMid / ycOut - 1./2.) / (glowGainIn / 2. - 1.);
+    }
+
+    return glowGainOut;
+}
+
+
+float sigmoid_shaper( float x)
+{
+	// Sigmoid function in the range 0 to 1 spanning -2 to +2.
+
+	float t = max( 1 - abs( 0.5 * x ), 0 );
+	float y = 1 + sign(x) * (1 - t*t);
+	return 0.5 * y;
+}
+
+
+// ------- Red modifier functions
+float cubic_basis_shaper
+( 
+  float x, 
+  float w   // full base width of the shaper function (in degrees)
+)
+{
+	//return Square( smoothstep( 0, 1, 1 - abs( 2 * x/w ) ) );
+
+	float M[4][4] =
+	{
+		{ -1./6,  3./6, -3./6,  1./6 },
+		{  3./6, -6./6,  3./6,  0./6 },
+		{ -3./6,  0./6,  3./6,  0./6 },
+		{  1./6,  4./6,  1./6,  0./6 }
+	};
+  
+	float knots[5] = { -0.5 * w, -0.25 * w, 0, 0.25 * w, 0.5 * w };
+  
+	float y = 0;
+	if ((x > knots[0]) && (x < knots[4]))
+	{  
+		float knot_coord = (x - knots[0]) * 4.0 / w;
+		int j = knot_coord;
+		float t = knot_coord - j;
+      
+		float monomials[4] = { t*t*t, t*t, t, 1.0 };
+
+		// (if/else structure required for compatibility with CTL < v1.5.)
+		if ( j == 3) {
+			y = monomials[0] * M[0][0] + monomials[1] * M[1][0] + 
+				monomials[2] * M[2][0] + monomials[3] * M[3][0];
+		} else if ( j == 2) {
+			y = monomials[0] * M[0][1] + monomials[1] * M[1][1] + 
+				monomials[2] * M[2][1] + monomials[3] * M[3][1];
+		} else if ( j == 1) {
+			y = monomials[0] * M[0][2] + monomials[1] * M[1][2] + 
+				monomials[2] * M[2][2] + monomials[3] * M[3][2];
+		} else if ( j == 0) {
+			y = monomials[0] * M[0][3] + monomials[1] * M[1][3] + 
+				monomials[2] * M[2][3] + monomials[3] * M[3][3];
+		} else {
+			y = 0.0;
+		}
+	}
+  
+	return y * 1.5;
+}
+
+float center_hue( float hue, float centerH)
+{
+	float hueCentered = hue - centerH;
+	if (hueCentered < -180.)
+		hueCentered += 360;
+	else if (hueCentered > 180.)
+		hueCentered -= 360;
+	return hueCentered;
+}
+
+
+// Textbook monomial to basis-function conversion matrix.
+static const float3x3 M =
+{
+	{  0.5, -1.0, 0.5 },
+	{ -1.0,  1.0, 0.5 },
+	{  0.5,  0.0, 0.0 }
+};
+
+struct SegmentedSplineParams_c5
+{
+	float coefsLow[6];    // coefs for B-spline between minPoint and midPoint (units of log luminance)
+	float coefsHigh[6];   // coefs for B-spline between midPoint and maxPoint (units of log luminance)
+	float2 minPoint; // {luminance, luminance} linear extension below this
+	float2 midPoint; // {luminance, luminance} 
+	float2 maxPoint; // {luminance, luminance} linear extension above this
+	float slopeLow;       // log-log slope of low linear extension
+	float slopeHigh;      // log-log slope of high linear extension
+};
+
+struct SegmentedSplineParams_c9
+{
+	float coefsLow[10];    // coefs for B-spline between minPoint and midPoint (units of log luminance)
+	float coefsHigh[10];   // coefs for B-spline between midPoint and maxPoint (units of log luminance)
+	float2 minPoint; // {luminance, luminance} linear extension below this
+	float2 midPoint; // {luminance, luminance} 
+	float2 maxPoint; // {luminance, luminance} linear extension above this
+	float slopeLow;       // log-log slope of low linear extension
+	float slopeHigh;      // log-log slope of high linear extension
+};
+
+float segmented_spline_c5_fwd( float x )
+{
+	// RRT_PARAMS
+	const SegmentedSplineParams_c5 C =
+	{
+		// coefsLow[6]
+		{ -4.0000000000, -4.0000000000, -3.1573765773, -0.4852499958, 1.8477324706, 1.8477324706 },
+		// coefsHigh[6]
+		{ -0.7185482425, 2.0810307172, 3.6681241237, 4.0000000000, 4.0000000000, 4.0000000000 },
+		{ 0.18*exp2(-15.0), 0.0001},    // minPoint
+		{ 0.18,              4.8},    // midPoint
+		{ 0.18*exp2(18.0),  10000.},    // maxPoint
+		0.0,  // slopeLow
+		0.0   // slopeHigh
+	};
+
+	const int N_KNOTS_LOW = 4;
+	const int N_KNOTS_HIGH = 4;
+
+	// Check for negatives or zero before taking the log. If negative or zero,
+	// set to ACESMIN.1
+	float xCheck = x <= 0 ? exp2(-14.0) : x;
+
+	float logx = log10( xCheck);
+	float logy;
+
+	if ( logx <= log10(C.minPoint.x) )
+	{ 
+		logy = logx * C.slopeLow + ( log10(C.minPoint.y) - C.slopeLow * log10(C.minPoint.x) );
+	}
+	else if (( logx > log10(C.minPoint.x) ) && ( logx < log10(C.midPoint.x) ))
+	{
+		float knot_coord = (N_KNOTS_LOW-1) * (logx-log10(C.minPoint.x))/(log10(C.midPoint.x)-log10(C.minPoint.x));
+		int j = knot_coord;
+		float t = knot_coord - j;
+
+		float3 cf = { C.coefsLow[ j], C.coefsLow[ j + 1], C.coefsLow[ j + 2]};
+    
+		float3 monomials = { t * t, t, 1.0 };
+		logy = dot( monomials, mul( cf, M));
+	}
+	else if (( logx >= log10(C.midPoint.x) ) && ( logx < log10(C.maxPoint.x) ))
+	{
+		float knot_coord = (N_KNOTS_HIGH-1) * (logx-log10(C.midPoint.x))/(log10(C.maxPoint.x)-log10(C.midPoint.x));
+		int j = knot_coord;
+		float t = knot_coord - j;
+
+		float3 cf = { C.coefsHigh[ j], C.coefsHigh[ j + 1], C.coefsHigh[ j + 2]};
+
+		float3 monomials = { t * t, t, 1.0 };
+		logy = dot( monomials, mul( cf, M));
+	}
+	else
+	{ //if ( logIn >= log10(C.maxPoint.x) ) { 
+		logy = logx * C.slopeHigh + ( log10(C.maxPoint.y) - C.slopeHigh * log10(C.maxPoint.x) );
+	}
+
+	return pow( 10, logy );
+}
+
+float segmented_spline_c5_rev( float y )
+{  
+	// RRT_PARAMS
+	const SegmentedSplineParams_c5 C =
+	{
+		// coefsLow[6]
+		{ -4.0000000000, -4.0000000000, -3.1573765773, -0.4852499958, 1.8477324706, 1.8477324706 },
+		// coefsHigh[6]
+		{ -0.7185482425, 2.0810307172, 3.6681241237, 4.0000000000, 4.0000000000, 4.0000000000 },
+		{ 0.18*exp2(-15.0), 0.0001},    // minPoint
+		{ 0.18,              4.8},    // midPoint
+		{ 0.18*exp2(18.0),  10000.},    // maxPoint
+		0.0,  // slopeLow
+		0.0   // slopeHigh
+	};
+
+	const int N_KNOTS_LOW = 4;
+	const int N_KNOTS_HIGH = 4;
+
+	const float KNOT_INC_LOW = (log10(C.midPoint.x) - log10(C.minPoint.x)) / (N_KNOTS_LOW - 1.);
+	const float KNOT_INC_HIGH = (log10(C.maxPoint.x) - log10(C.midPoint.x)) / (N_KNOTS_HIGH - 1.);
+
+	int i;
+  
+	// KNOT_Y is luminance of the spline at each knot
+	float KNOT_Y_LOW[N_KNOTS_LOW];
+	for (i = 0; i < N_KNOTS_LOW; i = i+1)
+	{
+		KNOT_Y_LOW[i] = (C.coefsLow[i] + C.coefsLow[i+1]) / 2.;
+	};
+
+	float KNOT_Y_HIGH[N_KNOTS_HIGH];
+	for (i = 0; i < N_KNOTS_HIGH; i = i+1) 
+	{
+		KNOT_Y_HIGH[i] = ( C.coefsHigh[i] + C.coefsHigh[i+1]) / 2.;
+	};
+
+	float logy = log10(max(y,1e-10));
+
+	float logx;
+	if (logy <= log10(C.minPoint.y)) 
+	{
+		logx = log10(C.minPoint.x);
+	} 
+	else if ( (logy > log10(C.minPoint.y)) && (logy <= log10(C.midPoint.y)) ) 
+	{
+		uint j;
+		float3 cf;
+		if (logy > KNOT_Y_LOW[ 0] && logy <= KNOT_Y_LOW[ 1]) {
+			cf[ 0] = C.coefsLow[0];  cf[ 1] = C.coefsLow[1];  cf[ 2] = C.coefsLow[2];  j = 0;
+		} else if ( logy > KNOT_Y_LOW[ 1] && logy <= KNOT_Y_LOW[ 2]) {
+			cf[ 0] = C.coefsLow[1];  cf[ 1] = C.coefsLow[2];  cf[ 2] = C.coefsLow[3];  j = 1;
+		} else if ( logy > KNOT_Y_LOW[ 2] && logy <= KNOT_Y_LOW[ 3]) {
+			cf[ 0] = C.coefsLow[2];  cf[ 1] = C.coefsLow[3];  cf[ 2] = C.coefsLow[4];  j = 2;
+		} 
+    
+		const float3 tmp = mul( cf, M);
+
+		float a = tmp[ 0];
+		float b = tmp[ 1];
+		float c = tmp[ 2];
+		c = c - logy;
+
+		const float d = sqrt( b * b - 4. * a * c);
+
+		const float t = ( 2. * c) / ( -d - b);
+
+		logx = log10(C.minPoint.x) + ( t + j) * KNOT_INC_LOW;
+	} 
+	else if ( (logy > log10(C.midPoint.y)) && (logy < log10(C.maxPoint.y)) ) 
+	{
+		uint j;
+		float3 cf;
+		if ( logy > KNOT_Y_HIGH[ 0] && logy <= KNOT_Y_HIGH[ 1]) {
+			cf[ 0] = C.coefsHigh[0];  cf[ 1] = C.coefsHigh[1];  cf[ 2] = C.coefsHigh[2];  j = 0;
+		} else if ( logy > KNOT_Y_HIGH[ 1] && logy <= KNOT_Y_HIGH[ 2]) {
+			cf[ 0] = C.coefsHigh[1];  cf[ 1] = C.coefsHigh[2];  cf[ 2] = C.coefsHigh[3];  j = 1;
+		} else if ( logy > KNOT_Y_HIGH[ 2] && logy <= KNOT_Y_HIGH[ 3]) {
+			cf[ 0] = C.coefsHigh[2];  cf[ 1] = C.coefsHigh[3];  cf[ 2] = C.coefsHigh[4];  j = 2;
+		} 
+    
+		const float3 tmp = mul( cf, M);
+
+		float a = tmp[ 0];
+		float b = tmp[ 1];
+		float c = tmp[ 2];
+		c = c - logy;
+
+		const float d = sqrt( b * b - 4. * a * c);
+
+		const float t = ( 2. * c) / ( -d - b);
+
+		logx = log10(C.midPoint.x) + ( t + j) * KNOT_INC_HIGH;
+	} 
+	else 
+	{ //if ( logy >= log10(C.maxPoint.y) ) {
+		logx = log10(C.maxPoint.x);
+	}
+  
+	return pow(10, logx);
+}
+
+float segmented_spline_c9_fwd( float x, const SegmentedSplineParams_c9 C )
+{
+	const int N_KNOTS_LOW = 8;
+	const int N_KNOTS_HIGH = 8;
+
+	// Check for negatives or zero before taking the log. If negative or zero,
+	// set to OCESMIN.
+	float xCheck = x <= 0 ? 1e-4 : x;
+
+	float logx = log10(xCheck);
+	float logy;
+
+	if (logx <= log10(C.minPoint.x))
+	{
+		logy = logx * C.slopeLow + (log10(C.minPoint.y) - C.slopeLow * log10(C.minPoint.x));
+	}
+	else if ((logx > log10(C.minPoint.x)) && (logx < log10(C.midPoint.x)))
+	{
+		float knot_coord = (N_KNOTS_LOW - 1) * (logx - log10(C.minPoint.x)) / (log10(C.midPoint.x) - log10(C.minPoint.x));
+		int j = knot_coord;
+		float t = knot_coord - j;
+
+		float3 cf = { C.coefsLow[j], C.coefsLow[j + 1], C.coefsLow[j + 2] };
+
+		float3 monomials = { t * t, t, 1.0 };
+		logy = dot(monomials, mul(cf, M));
+	}
+	else if ((logx >= log10(C.midPoint.x)) && (logx < log10(C.maxPoint.x)))
+	{
+		float knot_coord = (N_KNOTS_HIGH - 1) * (logx - log10(C.midPoint.x)) / (log10(C.maxPoint.x) - log10(C.midPoint.x));
+		int j = knot_coord;
+		float t = knot_coord - j;
+
+		float3 cf = { C.coefsHigh[j], C.coefsHigh[j + 1], C.coefsHigh[j + 2] };
+
+		float3 monomials = { t * t, t, 1.0 };
+		logy = dot(monomials, mul(cf, M));
+	}
+	else//if ( logIn >= log10(C.maxPoint.x) )
+	{
+		logy = logx * C.slopeHigh + (log10(C.maxPoint.y) - C.slopeHigh * log10(C.maxPoint.x));
+	}
+
+	return pow( 10, logy );
+}
+
+float segmented_spline_c9_rev( float y, const SegmentedSplineParams_c9 C )
+{
+	const int N_KNOTS_LOW = 8;
+	const int N_KNOTS_HIGH = 8;
+
+	const float KNOT_INC_LOW = (log10(C.midPoint.x) - log10(C.minPoint.x)) / (N_KNOTS_LOW - 1.);
+	const float KNOT_INC_HIGH = (log10(C.maxPoint.x) - log10(C.midPoint.x)) / (N_KNOTS_HIGH - 1.);
+  
+	int i;
+
+	// KNOT_Y is luminance of the spline at each knot
+	float KNOT_Y_LOW[ N_KNOTS_LOW];
+	for (i = 0; i < N_KNOTS_LOW; i = i+1) {
+		KNOT_Y_LOW[ i] = ( C.coefsLow[i] + C.coefsLow[i+1]) / 2.;
+	};
+
+	float KNOT_Y_HIGH[ N_KNOTS_HIGH];
+	for (i = 0; i < N_KNOTS_HIGH; i = i+1) {
+		KNOT_Y_HIGH[ i] = ( C.coefsHigh[i] + C.coefsHigh[i+1]) / 2.;
+	};
+
+	float logy = log10( max( y, 1e-10));
+
+	float logx;
+	if (logy <= log10(C.minPoint.y)) {
+		logx = log10(C.minPoint.x);
+	} else if ( (logy > log10(C.minPoint.y)) && (logy <= log10(C.midPoint.y)) ) {
+		uint j;
+		float3 cf;
+		if ( logy > KNOT_Y_LOW[ 0] && logy <= KNOT_Y_LOW[ 1]) {
+			cf[ 0] = C.coefsLow[0];  cf[ 1] = C.coefsLow[1];  cf[ 2] = C.coefsLow[2];  j = 0;
+		} else if ( logy > KNOT_Y_LOW[ 1] && logy <= KNOT_Y_LOW[ 2]) {
+			cf[ 0] = C.coefsLow[1];  cf[ 1] = C.coefsLow[2];  cf[ 2] = C.coefsLow[3];  j = 1;
+		} else if ( logy > KNOT_Y_LOW[ 2] && logy <= KNOT_Y_LOW[ 3]) {
+			cf[ 0] = C.coefsLow[2];  cf[ 1] = C.coefsLow[3];  cf[ 2] = C.coefsLow[4];  j = 2;
+		} else if ( logy > KNOT_Y_LOW[ 3] && logy <= KNOT_Y_LOW[ 4]) {
+			cf[ 0] = C.coefsLow[3];  cf[ 1] = C.coefsLow[4];  cf[ 2] = C.coefsLow[5];  j = 3;
+		} else if ( logy > KNOT_Y_LOW[ 4] && logy <= KNOT_Y_LOW[ 5]) {
+			cf[ 0] = C.coefsLow[4];  cf[ 1] = C.coefsLow[5];  cf[ 2] = C.coefsLow[6];  j = 4;
+		} else if ( logy > KNOT_Y_LOW[ 5] && logy <= KNOT_Y_LOW[ 6]) {
+			cf[ 0] = C.coefsLow[5];  cf[ 1] = C.coefsLow[6];  cf[ 2] = C.coefsLow[7];  j = 5;
+		} else if ( logy > KNOT_Y_LOW[ 6] && logy <= KNOT_Y_LOW[ 7]) {
+			cf[ 0] = C.coefsLow[6];  cf[ 1] = C.coefsLow[7];  cf[ 2] = C.coefsLow[8];  j = 6;
+		}
+    
+		const float3 tmp = mul( cf, M);
+
+		float a = tmp[ 0];
+		float b = tmp[ 1];
+		float c = tmp[ 2];
+		c = c - logy;
+
+		const float d = sqrt( b * b - 4. * a * c);
+
+		const float t = ( 2. * c) / ( -d - b);
+
+		logx = log10(C.minPoint.x) + ( t + j) * KNOT_INC_LOW;
+	} else if ( (logy > log10(C.midPoint.y)) && (logy < log10(C.maxPoint.y)) ) {
+		uint j;
+		float3 cf;
+		if ( logy > KNOT_Y_HIGH[ 0] && logy <= KNOT_Y_HIGH[ 1]) {
+			cf[ 0] = C.coefsHigh[0];  cf[ 1] = C.coefsHigh[1];  cf[ 2] = C.coefsHigh[2];  j = 0;
+		} else if ( logy > KNOT_Y_HIGH[ 1] && logy <= KNOT_Y_HIGH[ 2]) {
+			cf[ 0] = C.coefsHigh[1];  cf[ 1] = C.coefsHigh[2];  cf[ 2] = C.coefsHigh[3];  j = 1;
+		} else if ( logy > KNOT_Y_HIGH[ 2] && logy <= KNOT_Y_HIGH[ 3]) {
+			cf[ 0] = C.coefsHigh[2];  cf[ 1] = C.coefsHigh[3];  cf[ 2] = C.coefsHigh[4];  j = 2;
+		} else if ( logy > KNOT_Y_HIGH[ 3] && logy <= KNOT_Y_HIGH[ 4]) {
+			cf[ 0] = C.coefsHigh[3];  cf[ 1] = C.coefsHigh[4];  cf[ 2] = C.coefsHigh[5];  j = 3;
+		} else if ( logy > KNOT_Y_HIGH[ 4] && logy <= KNOT_Y_HIGH[ 5]) {
+			cf[ 0] = C.coefsHigh[4];  cf[ 1] = C.coefsHigh[5];  cf[ 2] = C.coefsHigh[6];  j = 4;
+		} else if ( logy > KNOT_Y_HIGH[ 5] && logy <= KNOT_Y_HIGH[ 6]) {
+			cf[ 0] = C.coefsHigh[5];  cf[ 1] = C.coefsHigh[6];  cf[ 2] = C.coefsHigh[7];  j = 5;
+		} else if ( logy > KNOT_Y_HIGH[ 6] && logy <= KNOT_Y_HIGH[ 7]) {
+			cf[ 0] = C.coefsHigh[6];  cf[ 1] = C.coefsHigh[7];  cf[ 2] = C.coefsHigh[8];  j = 6;
+		}
+    
+		const float3 tmp = mul( cf, M);
+
+		float a = tmp[ 0];
+		float b = tmp[ 1];
+		float c = tmp[ 2];
+		c = c - logy;
+
+		const float d = sqrt( b * b - 4. * a * c);
+
+		const float t = ( 2. * c) / ( -d - b);
+
+		logx = log10(C.midPoint.x) + ( t + j) * KNOT_INC_HIGH;
+	} 
+	else 
+	{ //if ( logy >= log10(C.maxPoint.y) ) {
+		logx = log10(C.maxPoint.x);
+	}
+  
+	return pow(10, logx);
+}
+
+// Transformations from RGB to other color representations
+float rgb_2_hue( float3 rgb ) 
+{
+	// Returns a geometric hue angle in degrees (0-360) based on RGB values.
+	// For neutral colors, hue is undefined and the function will return a quiet NaN value.
+	float hue;
+	if (rgb[0] == rgb[1] && rgb[1] == rgb[2])
+	{
+		//hue = FLT_NAN; // RGB triplets where RGB are equal have an undefined hue
+		hue = 0;
+	}
+	else
+	{
+		hue = (180. / PI) * atan2(sqrt(3.0)*(rgb[1] - rgb[2]), 2 * rgb[0] - rgb[1] - rgb[2]);
+	}
+
+	if (hue < 0.)
+		hue = hue + 360;
+
+	return clamp( hue, 0, 360 );
+}
+
+float rgb_2_yc( float3 rgb, float ycRadiusWeight = 1.75)
+{
+	// Converts RGB to a luminance proxy, here called YC
+	// YC is ~ Y + K * Chroma
+	// Constant YC is a cone-shaped surface in RGB space, with the tip on the 
+	// neutral axis, towards white.
+	// YC is normalized: RGB 1 1 1 maps to YC = 1
+	//
+	// ycRadiusWeight defaults to 1.75, although can be overridden in function 
+	// call to rgb_2_yc
+	// ycRadiusWeight = 1 -> YC for pure cyan, magenta, yellow == YC for neutral 
+	// of same value
+	// ycRadiusWeight = 2 -> YC for pure red, green, blue  == YC for  neutral of 
+	// same value.
+
+	float r = rgb[0]; 
+	float g = rgb[1]; 
+	float b = rgb[2];
+  
+	float chroma = sqrt(b*(b-g)+g*(g-r)+r*(r-b));
+
+	return ( b + g + r + ycRadiusWeight * chroma) / 3.;
+}
+
+// 
+// Reference Rendering Transform (RRT)
+//
+//   Input is ACES
+//   Output is OCES
+//
+float3 RRT( float3 aces )
+{
+	// "Glow" module constants
+	const float RRT_GLOW_GAIN = 0.05;
+	const float RRT_GLOW_MID = 0.08;
+
+	float saturation = rgb_2_saturation( aces );
+	float ycIn = rgb_2_yc( aces );
+	float s = sigmoid_shaper( (saturation - 0.4) / 0.2);
+	float addedGlow = 1 + glow_fwd( ycIn, RRT_GLOW_GAIN * s, RRT_GLOW_MID);
+	aces *= addedGlow;
+  
+	// --- Red modifier --- //
+	const float RRT_RED_SCALE = 0.82;
+	const float RRT_RED_PIVOT = 0.03;
+	const float RRT_RED_HUE = 0;
+	const float RRT_RED_WIDTH = 135;
+	float hue = rgb_2_hue( aces );
+	float centeredHue = center_hue( hue, RRT_RED_HUE );
+	float hueWeight = cubic_basis_shaper( centeredHue, RRT_RED_WIDTH );
+		
+	aces.r += hueWeight * saturation * (RRT_RED_PIVOT - aces.r) * (1. - RRT_RED_SCALE);
+
+	// --- ACES to RGB rendering space --- //
+	aces = clamp( aces, 0, 65535 );  // avoids saturated negative colors from becoming positive in the matrix
+
+	float3 rgbPre = mul( AP0_2_AP1_MAT, aces );
+
+	rgbPre = clamp( rgbPre, 0, 65535 );
+
+	// --- Global desaturation --- //
+	const float RRT_SAT_FACTOR = 0.96;
+	rgbPre = lerp( dot( rgbPre, AP1_RGB2Y ), rgbPre, RRT_SAT_FACTOR );
+
+	// --- Apply the tonescale independently in rendering-space RGB --- //
+	float3 rgbPost;
+	rgbPost[0] = segmented_spline_c5_fwd( rgbPre[0] );
+	rgbPost[1] = segmented_spline_c5_fwd( rgbPre[1] );
+	rgbPost[2] = segmented_spline_c5_fwd( rgbPre[2] );
+
+	// --- RGB rendering space to OCES --- //
+	float3 rgbOces = mul( AP1_2_AP0_MAT, rgbPost );
+	return rgbOces;
+}
+
+// 
+// Inverse Reference Rendering Transform (RRT)
+//
+//   Input is OCES
+//   Output is ACES
+//
+
+float3 Inverse_RRT( float3 oces )
+{
+	// "Glow" module constants
+	const float RRT_GLOW_GAIN = 0.05;
+	const float RRT_GLOW_MID = 0.08;
+
+	// --- OCES to RGB rendering space --- //
+    float3 rgbPre = mul(AP0_2_AP1_MAT, oces);
+
+	// --- Apply the tonescale independently in rendering-space RGB --- //
+    float3 rgbPost;
+    rgbPost[0] = segmented_spline_c5_rev( rgbPre[0]);
+    rgbPost[1] = segmented_spline_c5_rev( rgbPre[1]);
+    rgbPost[2] = segmented_spline_c5_rev( rgbPre[2]);
+
+	// --- Global desaturation --- //
+	// rgbPost = mul( rgbPost, invert_f33(RRT_SAT_MAT));
+	const float RRT_SAT_FACTOR = 0.96;
+	rgbPost = lerp( dot( rgbPost, AP1_RGB2Y ), rgbPost, rcp(RRT_SAT_FACTOR) );
+
+    rgbPost = clamp( rgbPost, 0., HALF_MAX);
+
+	// --- RGB rendering space to ACES --- //
+    float3 aces = mul( AP1_2_AP0_MAT, rgbPost);
+
+    aces = clamp( aces, 0., HALF_MAX);
+
+	// --- Red modifier --- //
+	const float RRT_RED_SCALE = 0.82;
+	const float RRT_RED_PIVOT = 0.03;
+	const float RRT_RED_HUE = 0;
+	const float RRT_RED_WIDTH = 135;
+    float hue = rgb_2_hue( aces);
+    float centeredHue = center_hue( hue, RRT_RED_HUE);
+    float hueWeight = cubic_basis_shaper( centeredHue, RRT_RED_WIDTH);
+
+    float minChan;
+    if (centeredHue < 0) { 
+		// min_f3(aces) = aces[1] (i.e. magenta-red)
+		minChan = aces[1];
+    } else { // min_f3(aces) = aces[2] (i.e. yellow-red)
+		minChan = aces[2];
+    }
+
+    float a = hueWeight * (1. - RRT_RED_SCALE) - 1.;
+    float b = aces[0] - hueWeight * (RRT_RED_PIVOT + minChan) * (1. - RRT_RED_SCALE);
+    float c = hueWeight * RRT_RED_PIVOT * minChan * (1. - RRT_RED_SCALE);
+
+    aces[0] = ( -b - sqrt( b * b - 4. * a * c)) / ( 2. * a);
+
+	// --- Glow module --- //
+    float saturation = rgb_2_saturation( aces);
+    float ycOut = rgb_2_yc( aces);
+    float s = sigmoid_shaper( (saturation - 0.4) / 0.2);
+    float reducedGlow = 1. + glow_inv( ycOut, RRT_GLOW_GAIN * s, RRT_GLOW_MID);
+
+    aces = reducedGlow * aces;
+
+	// Assign ACES RGB to output variables (ACES)
+	return aces;
+}
+
+
+
+// Transformations between CIE XYZ tristimulus values and CIE x,y 
+// chromaticity coordinates
+float3 XYZ_2_xyY( float3 XYZ )
+{  
+	float3 xyY;
+	float divisor = (XYZ[0] + XYZ[1] + XYZ[2]);
+	if (divisor == 0.) divisor = 1e-10;
+	xyY[0] = XYZ[0] / divisor;
+	xyY[1] = XYZ[1] / divisor;  
+	xyY[2] = XYZ[1];
+  
+	return xyY;
+}
+
+float3 xyY_2_XYZ( float3 xyY )
+{
+	float3 XYZ;
+	XYZ[0] = xyY[0] * xyY[2] / max( xyY[1], 1e-10);
+	XYZ[1] = xyY[2];  
+	XYZ[2] = (1.0 - xyY[0] - xyY[1]) * xyY[2] / max( xyY[1], 1e-10);
+
+	return XYZ;
+}
+
+
+float3x3 ChromaticAdaptation( float2 src_xy, float2 dst_xy )
+{
+	// Von Kries chromatic adaptation 
+
+	// Bradford
+	const float3x3 ConeResponse =
+	{
+		 0.8951,  0.2664, -0.1614,
+		-0.7502,  1.7135,  0.0367,
+		 0.0389, -0.0685,  1.0296,
+	};
+	const float3x3 InvConeResponse =
+	{
+		 0.9869929, -0.1470543,  0.1599627,
+		 0.4323053,  0.5183603,  0.0492912,
+		-0.0085287,  0.0400428,  0.9684867,
+	};
+
+	float3 src_XYZ = xyY_2_XYZ( float3( src_xy, 1 ) );
+	float3 dst_XYZ = xyY_2_XYZ( float3( dst_xy, 1 ) );
+
+	float3 src_coneResp = mul( ConeResponse, src_XYZ );
+	float3 dst_coneResp = mul( ConeResponse, dst_XYZ );
+
+	float3x3 VonKriesMat =
+	{
+		{ dst_coneResp[0] / src_coneResp[0], 0.0, 0.0 },
+		{ 0.0, dst_coneResp[1] / src_coneResp[1], 0.0 },
+		{ 0.0, 0.0, dst_coneResp[2] / src_coneResp[2] }
+	};
+
+	return mul( InvConeResponse, mul( VonKriesMat, ConeResponse ) );
+}
+
+float Y_2_linCV( float Y, float Ymax, float Ymin) 
+{
+  return (Y - Ymin) / (Ymax - Ymin);
+}
+
+float linCV_2_Y( float linCV, float Ymax, float Ymin) 
+{
+  return linCV * (Ymax - Ymin) + Ymin;
+}
+
+// Gamma compensation factor
+static const float DIM_SURROUND_GAMMA = 0.9811;
+
+float3 darkSurround_to_dimSurround( float3 linearCV )
+{
+	float3 XYZ = mul( AP1_2_XYZ_MAT, linearCV );
+
+	float3 xyY = XYZ_2_xyY(XYZ);
+	xyY[2] = clamp( xyY[2], 0, 65535 );
+	xyY[2] = pow( xyY[2], DIM_SURROUND_GAMMA );
+	XYZ = xyY_2_XYZ(xyY);
+
+	return mul( XYZ_2_AP1_MAT, XYZ );
+}
+
+float3 dimSurround_to_darkSurround( float3 linearCV)
+{
+  float3 XYZ = mul( linearCV, AP1_2_XYZ_MAT); 
+
+  float3 xyY = XYZ_2_xyY(XYZ);
+  xyY[2] = clamp( xyY[2], 0., 65535);
+  xyY[2] = pow( xyY[2], 1./DIM_SURROUND_GAMMA);
+  XYZ = xyY_2_XYZ(xyY);
+
+  return mul( XYZ, XYZ_2_AP1_MAT);
+}
+
+// 
+// Output Device Transform - RGB computer monitor (D60 simulation)
+//
+
+//
+// Summary :
+//  This transform is intended for mapping OCES onto a desktop computer monitor 
+//  typical of those used in motion picture visual effects production used to 
+//  simulate the image appearance produced by odt_p3d60. These monitors may
+//  occasionally be referred to as "sRGB" displays, however, the monitor for
+//  which this transform is designed does not exactly match the specifications
+//  in IEC 61966-2-1:1999.
+// 
+//  The assumed observer adapted white is D60, and the viewing environment is 
+//  that of a dim surround. 
+//
+//  The monitor specified is intended to be more typical of those found in 
+//  visual effects production.
+//
+// Device Primaries : 
+//  Primaries are those specified in Rec. ITU-R BT.709
+//  CIE 1931 chromaticities:  x         y         Y
+//              Red:          0.64      0.33
+//              Green:        0.3       0.6
+//              Blue:         0.15      0.06
+//              White:        0.3217    0.329     100 cd/m^2
+//
+// Display EOTF :
+//  The reference electro-optical transfer function specified in 
+//  IEC 61966-2-1:1999.
+//
+// Signal Range:
+//    This tranform outputs full range code values.
+//
+// Assumed observer adapted white point:
+//         CIE 1931 chromaticities:    x            y
+//                                     0.32168      0.33767
+//
+// Viewing Environment:
+//   This ODT has a compensation for viewing environment variables more typical 
+//   of those associated with video mastering.
+//
+
+//
+//  Epic edits:
+//  - This ODT has been modified to target an observer adapted white of D65.
+//  - The output of the function is linear output referred values. The
+//      linear to sRGB transform should be applied after this function.
+//    
+
+float3 ODT_sRGB_D65( float3 oces )
+{
+	// OCES to RGB rendering space
+	float3 rgbPre = mul( AP0_2_AP1_MAT, oces );
+
+	const SegmentedSplineParams_c9 ODT_48nits =
+	{
+		// coefsLow[10]
+		{ -1.6989700043, -1.6989700043, -1.4779000000, -1.2291000000, -0.8648000000, -0.4480000000, 0.0051800000, 0.4511080334, 0.9113744414, 0.9113744414},
+		// coefsHigh[10]
+		{ 0.5154386965, 0.8470437783, 1.1358000000, 1.3802000000, 1.5197000000, 1.5985000000, 1.6467000000, 1.6746091357, 1.6878733390, 1.6878733390 },
+		{segmented_spline_c5_fwd( 0.18*exp2(-6.5) ),  0.02},    // minPoint
+		{segmented_spline_c5_fwd( 0.18 ),              4.8},    // midPoint  
+		{segmented_spline_c5_fwd( 0.18*exp2(6.5) ),   48.0},    // maxPoint
+		0.0,  // slopeLow
+		0.04  // slopeHigh
+	};
+
+	// Apply the tonescale independently in rendering-space RGB
+	float3 rgbPost;
+	rgbPost[0] = segmented_spline_c9_fwd( rgbPre[0], ODT_48nits );
+	rgbPost[1] = segmented_spline_c9_fwd( rgbPre[1], ODT_48nits );
+	rgbPost[2] = segmented_spline_c9_fwd( rgbPre[2], ODT_48nits );
+
+	// Target white and black points for cinema system tonescale
+	const float CINEMA_WHITE = 48.0;
+	const float CINEMA_BLACK = 0.02; // CINEMA_WHITE / 2400.
+
+	// Scale luminance to linear code value
+	float3 linearCV;
+	linearCV[0] = Y_2_linCV( rgbPost[0], CINEMA_WHITE, CINEMA_BLACK );
+	linearCV[1] = Y_2_linCV( rgbPost[1], CINEMA_WHITE, CINEMA_BLACK );
+	linearCV[2] = Y_2_linCV( rgbPost[2], CINEMA_WHITE, CINEMA_BLACK );    
+
+	// Apply gamma adjustment to compensate for dim surround
+	linearCV = darkSurround_to_dimSurround( linearCV );
+
+	// Apply desaturation to compensate for luminance difference
+	const float ODT_SAT_FACTOR = 0.93;
+	linearCV = lerp( dot( linearCV, AP1_RGB2Y ), linearCV, ODT_SAT_FACTOR );
+
+	// Convert to display primary encoding
+	// Rendering space RGB to XYZ
+	float3 XYZ = mul( AP1_2_XYZ_MAT, linearCV );
+	
+	// Apply CAT from ACES white point to assumed observer adapted white point
+	/*
+	const float3x3 D60_2_D65_CAT =
+	{
+		 0.987224,   -0.00611327, 0.0159533,
+		-0.00759836,  1.00186,    0.00533002,
+		 0.00307257, -0.00509595, 1.08168,
+	};
+	*/
+	XYZ = mul( D60_2_D65_CAT, XYZ );
+
+	// CIE XYZ to display primaries
+	linearCV = mul( XYZ_2_sRGB_MAT, XYZ );
+
+	// Handle out-of-gamut values
+	// Clip values < 0 or > 1 (i.e. projecting outside the display primaries)
+	linearCV = saturate( linearCV );
+
+	return linearCV;
+}
+
+// 
+// Inverse Output Device Transform - RGB computer monitor (D60 simulation)
+//
+
+//
+//  Epic edits:
+//  - This Inverse ODT has been modified to accept an observer adapted white of D65.
+//  - The input to the function is linear output referred values. The
+//      sRGB to linear transform should be applied before this function.
+//    
+
+float3 Inverse_ODT_sRGB_D65(float3 linearCV)
+{  
+	// Convert from display primary encoding
+    // Display primaries to CIE XYZ
+    float3 XYZ = mul( sRGB_2_XYZ_MAT, linearCV);
+
+    // CIE XYZ to rendering space RGB
+    linearCV = mul( XYZ_2_AP1_MAT, XYZ);
+
+	// Apply CAT from ACES white point to assumed observer adapted white point
+	XYZ = mul( D65_2_D60_CAT, XYZ );
+
+	// Undo desaturation to compensate for luminance difference
+    //linearCV = mul( linearCV, invert_f33( ODT_SAT_MAT));
+	const float ODT_SAT_FACTOR = 0.93;
+	linearCV = lerp( dot( linearCV, AP1_RGB2Y ), linearCV, rcp(ODT_SAT_FACTOR) );
+
+	// Undo gamma adjustment to compensate for dim surround
+    linearCV = dimSurround_to_darkSurround( linearCV);
+
+	// Undo scaling done for D60 simulation
+	const float SCALE = 0.955;
+    linearCV = linearCV * rcp(SCALE);
+
+	// Target white and black points for cinema system tonescale
+	const float CINEMA_WHITE = 48.0;
+	const float CINEMA_BLACK = 0.02; // CINEMA_WHITE / 2400.
+
+	// Scale linear code value to luminance
+    float3 rgbPre;
+    rgbPre[0] = linCV_2_Y( linearCV[0], CINEMA_WHITE, CINEMA_BLACK);
+    rgbPre[1] = linCV_2_Y( linearCV[1], CINEMA_WHITE, CINEMA_BLACK);
+    rgbPre[2] = linCV_2_Y( linearCV[2], CINEMA_WHITE, CINEMA_BLACK);
+
+	const SegmentedSplineParams_c9 ODT_48nits =
+	{
+		// coefsLow[10]
+		{ -1.6989700043, -1.6989700043, -1.4779000000, -1.2291000000, -0.8648000000, -0.4480000000, 0.0051800000, 0.4511080334, 0.9113744414, 0.9113744414},
+		// coefsHigh[10]
+		{ 0.5154386965, 0.8470437783, 1.1358000000, 1.3802000000, 1.5197000000, 1.5985000000, 1.6467000000, 1.6746091357, 1.6878733390, 1.6878733390 },
+		{segmented_spline_c5_fwd( 0.18*pow(2.,-6.5) ),  0.02},    // minPoint
+		{segmented_spline_c5_fwd( 0.18 ),                4.8},    // midPoint  
+		{segmented_spline_c5_fwd( 0.18*pow(2.,6.5) ),   48.0},    // maxPoint
+		0.0,  // slopeLow
+		0.04  // slopeHigh
+	};
+
+	// Apply the tonescale independently in rendering-space RGB
+    float3 rgbPost;
+    rgbPost[0] = segmented_spline_c9_rev( rgbPre[0], ODT_48nits );
+    rgbPost[1] = segmented_spline_c9_rev( rgbPre[1], ODT_48nits );
+    rgbPost[2] = segmented_spline_c9_rev( rgbPre[2], ODT_48nits );
+
+	// Rendering space RGB to OCES
+    float3 oces = mul(AP1_2_AP0_MAT, rgbPost);
+
+	return oces;
+}
+
+
+// 
+// Output Device Transform - Rec.2020 (1000 cd/m^2)
+//
+
+//
+// Summary :
+//  This transform is intended for mapping OCES onto a Rec.2020 HDR display calibrated 
+//  to a D65 white point at 1000 cd/m^2. The assumed observer adapted white is 
+//  D65, and the viewing environment is that of a dim surround. 
+//
+// Device Primaries : 
+//  Primaries are those specified in Rec. ITU-R BT.2020
+//  CIE 1931 chromaticities:  x         y         Y
+//              Red:          0.708     0.292
+//              Green:        0.17      0.797
+//              Blue:         0.131     0.046
+//              White:        0.3127    0.329     1000 cd/m^2
+//
+// Display EOTF :
+//  The reference electro-optical transfer function specified in SMPTE ST 
+//  2084-2014. This transform makes no attempt to address the Annex functions 
+//  which address integer quantization.
+//
+// Assumed observer adapted white point:
+//         CIE 1931 chromaticities:    x            y
+//                                     0.3127       0.329
+//
+// Viewing Environment:
+//  This ODT is designed for a viewing environment more typically associated 
+//  with video mastering.
+//
+
+float3 ODT_1000nits( float3 oces )
+{
+	// OCES to RGB rendering space
+	float3 rgbPre = mul( AP0_2_AP1_MAT, oces );
+
+	const SegmentedSplineParams_c9 ODT_1000nits =
+	{
+		// coefsLow[10]
+		{ -4.9706219331, -3.0293780669, -2.1262, -1.5105, -1.0578, -0.4668, 0.11938, 0.7088134201, 1.2911865799, 1.2911865799 },
+		// coefsHigh[10]
+		{ 0.8089132070, 1.1910867930, 1.5683, 1.9483, 2.3083, 2.6384, 2.8595, 2.9872608805, 3.0127391195, 3.0127391195 },
+		{segmented_spline_c5_fwd( 0.18*pow(2.,-12.) ), 0.0001},    // minPoint
+		{segmented_spline_c5_fwd( 0.18 ),                10.0},    // midPoint  
+		{segmented_spline_c5_fwd( 0.18*pow(2.,10.) ),  1000.0},    // maxPoint
+		3.0,  // slopeLow
+		0.06  // slopeHigh
+	};
+
+	// Apply the tonescale independently in rendering-space RGB
+	float3 rgbPost;
+	rgbPost[0] = segmented_spline_c9_fwd( rgbPre[0], ODT_1000nits );
+	rgbPost[1] = segmented_spline_c9_fwd( rgbPre[1], ODT_1000nits );
+	rgbPost[2] = segmented_spline_c9_fwd( rgbPre[2], ODT_1000nits );
+
+	// Subtract small offset to allow for a code value of 0
+	rgbPost -= 0.00003507384284432574;
+
+	return rgbPost;
+}
+
+
+// 
+// Output Device Transform - P3D60 (2000 cd/m^2)
+//
+
+//
+// Summary :
+//  This transform is intended for mapping OCES onto an HDR display calibrated 
+//  to a D60 white point at 2000 cd/m^2. The assumed observer adapted white is 
+//  D60, and the viewing environment is that of a dim surround. 
+//
+// Device Primaries : 
+//  CIE 1931 chromaticities:  x         y         Y
+//              Red:          0.68      0.32
+//              Green:        0.265     0.69
+//              Blue:         0.15      0.06
+//              White:        0.32168   0.33767   2000 cd/m^2
+//
+//  Also assumes a black level of 0.005 cd/m^2
+//
+// Display EOTF :
+//  The reference electro-optical transfer function specified in SMPTE ST 
+//  2084-2014. This transform makes no attempt to address the Annex functions 
+//  which address integer quantization.
+//
+// Assumed observer adapted white point:
+//         CIE 1931 chromaticities:    x            y
+//                                     0.32168      0.33767
+//
+// Viewing Environment:
+//  This ODT is designed for a viewing environment more typically associated 
+//  with video mastering.
+//
+
+float3 ODT_2000nits( float3 oces )
+{
+	// OCES to RGB rendering space
+	float3 rgbPre = mul( AP0_2_AP1_MAT, oces );
+
+	const SegmentedSplineParams_c9 ODT_2000nits =
+	{
+	  // coefsLow[10]
+	  { -2.3010299957, -2.3010299957, -1.9312000000, -1.5205000000, -1.0578000000, -0.4668000000, 0.1193800000, 0.7088134201, 1.2911865799, 1.2911865799 },
+	  // coefsHigh[10]
+	  { 0.8019952042, 1.1980047958, 1.5943000000, 1.9973000000, 2.3783000000, 2.7684000000, 3.0515000000, 3.2746293562, 3.3274306351, 3.3274306351 },
+	  {segmented_spline_c5_fwd( 0.18*pow(2.,-12.) ),  0.005},    // minPoint
+	  {segmented_spline_c5_fwd( 0.18 ),                10.0},    // midPoint  
+	  {segmented_spline_c5_fwd( 0.18*pow(2.,11.) ),  2000.0},    // maxPoint
+	  0.0,  // slopeLow
+	  0.12  // slopeHigh
+	};
+
+	// Apply the tonescale independently in rendering-space RGB
+	float3 rgbPost;
+	rgbPost[0] = segmented_spline_c9_fwd( rgbPre[0], ODT_2000nits );
+	rgbPost[1] = segmented_spline_c9_fwd( rgbPre[1], ODT_2000nits );
+	rgbPost[2] = segmented_spline_c9_fwd( rgbPre[2], ODT_2000nits );
+
+	return rgbPost;
+}
\ No newline at end of file
diff --git a/Engine/Shaders/Private/AtmosphereCommon.ush b/Engine/Shaders/Private/AtmosphereCommon.ush
new file mode 100644
--- /dev/null
+++ b/Engine/Shaders/Private/AtmosphereCommon.ush
@@ -0,0 +1,625 @@
+// Copyright 1998-2017 Epic Games, Inc. All Rights Reserved.
+
+/*=============================================================================
+	AtmosphereCommon.usf: Functions and variables shared by both rendering and precomputation
+
+	This code contains embedded portions of free sample source code from 
+	http://www-evasion.imag.fr/Membres/Eric.Bruneton/PrecomputedAtmosphericScattering2.zip, Author: Eric Bruneton, 
+	08/16/2011, Copyright (c) 2008 INRIA, All Rights Reserved, which have been altered from their original version.
+
+	Permission is granted to anyone to use this software for any purpose, including commercial applications, and to alter it and redistribute it freely, subject to the following restrictions:
+
+    1. Redistributions of source code must retain the above copyright notice, 
+	   this list of conditions and the following disclaimer.
+    2. Redistributions in binary form must reproduce the above copyright notice, 
+	   this list of conditions and the following disclaimer in the
+       documentation and/or other materials provided with the distribution.
+    3. Neither the name of the copyright holders nor the names of its
+       contributors may be used to endorse or promote products derived from
+       this software without specific prior written permission.
+
+=============================================================================*/
+
+#ifndef __ATMOSPHERE_COMMON__
+#define __ATMOSPHERE_COMMON__
+
+const static float AverageGroundRelectance = 0.1f; // Simple approximation ground reflection light
+
+/** Magic numbers are based on paper (http://hal.inria.fr/docs/00/28/87/58/PDF/article.pdf) */
+/** Rayleigh scattering terms */
+const static float3 BetaRayleighScattering = float3(5.8e-3f, 1.35e-2f, 3.31e-2f); // Equation 1, REK 04
+
+/** Mie scattering terms */
+const static float HeightScaleMie = 1.2f; // 1.2 km, for Mie scattering
+const static float3 BetaMieScattering = float3(4e-3f, 4e-3f, 4e-3f); // Equation 4
+const static float BetaRatio = 0.9f; // Figure 6, BetaMScattering/BetaMExtinction = 0.9
+const static float3 BetaMieExtinction = BetaMieScattering / BetaRatio.rrr; 
+const static float MieG = 0.8f; // Equation 4
+
+const static float RadiusGround = 6360;
+const static float RadiusAtmosphere = 6430;
+const static float RadiusLimit = 6431;
+
+const static int IrradianceTexWidth = 64;
+const static int IrradianceTexHeight = 16;
+
+const static int InscatterMuNum = 128;
+const static int InscatterMuSNum = 32;
+const static int InscatterNuNum = 8;
+
+#define TRANSMITTANCE_NON_LINEAR			1
+#define INSCATTER_NON_LINEAR				1
+#define ATMOSPHERIC_TEXTURE_SAMPLE_FIX		1
+
+#ifndef ATMOSPHERIC_NO_SUN_DISK
+#define	ATMOSPHERIC_NO_SUN_DISK				0
+#endif
+
+#ifndef ATMOSPHERIC_NO_GROUND_SCATTERING
+#define	ATMOSPHERIC_NO_GROUND_SCATTERING	0
+#endif
+
+#ifndef ATMOSPHERIC_NO_DECAY
+#define	ATMOSPHERIC_NO_DECAY				0
+#endif
+
+#ifndef RENDERFLAG_DISABLE_SUN_DISK
+#define RENDERFLAG_DISABLE_SUN_DISK			1		// E_DisableSunDisk = 1,
+#endif
+
+#ifndef RENDERFLAG_DISABLE_GROUND_SCATTERING
+#define RENDERFLAG_DISABLE_GROUND_SCATTERING	2	// E_DisableGroundScattering = 2,
+#endif
+
+/** Textures for rendering */
+#if IS_MATERIAL_SHADER
+    // for materials, these are in the view UB
+	#define AtmosphereTransmittanceTexture			View.AtmosphereTransmittanceTexture_UB
+	#define AtmosphereTransmittanceTextureSampler	View.AtmosphereTransmittanceTextureSampler_UB
+	#define AtmosphereIrradianceTexture				View.AtmosphereIrradianceTexture_UB
+	#define AtmosphereIrradianceTextureSampler		View.AtmosphereIrradianceTextureSampler_UB
+	#define AtmosphereInscatterTexture				View.AtmosphereInscatterTexture_UB
+	#define AtmosphereInscatterTextureSampler		View.AtmosphereInscatterTextureSampler_UB
+#else
+	// these are only used for the precomputation shader; which doesn't have a view UB
+	Texture2D AtmosphereTransmittanceTexture;
+	SamplerState AtmosphereTransmittanceTextureSampler;
+	Texture2D AtmosphereIrradianceTexture;
+	SamplerState AtmosphereIrradianceTextureSampler;
+	Texture3D AtmosphereInscatterTexture;
+	SamplerState AtmosphereInscatterTextureSampler;
+	#define ATMOSPHERIC_TEXTURE_PREFIX
+#endif
+
+struct FAtmosphereVSOutput
+{
+	float2 OutTexCoord : TEXCOORD0;
+    float4 OutPosition : SV_POSITION;
+#if USING_VERTEX_SHADER_LAYER
+    uint LayerIndex : SV_RenderTargetArrayIndex;
+#else
+    uint LayerIndex : TEXCOORD1;
+#endif
+};
+
+float2 GetTransmittanceUV(float Radius, float Mu) 
+{
+    float U, V;
+#if TRANSMITTANCE_NON_LINEAR
+	V = sqrt((Radius - RadiusGround) / (RadiusAtmosphere - RadiusGround));
+	U = atan((Mu + 0.15) / (1.0 + 0.15) * tan(1.5)) / 1.5;
+#else
+	V = (Radius - RadiusGround) / (RadiusAtmosphere - RadiusGround);
+	U = (Mu + 0.15) / (1.0 + 0.15);
+#endif
+    return float2(U, V);
+}
+
+void GetTransmittanceRMuS(float2 UV, out float Radius, out float MuS) 
+{
+    Radius = UV.y;
+    MuS = UV.x;
+#if TRANSMITTANCE_NON_LINEAR
+    Radius = RadiusGround + (Radius * Radius) * (RadiusAtmosphere - RadiusGround);
+    MuS = -0.15 + tan(1.5 * MuS) / tan(1.5) * (1.0 + 0.15);
+#else
+    r = RadiusGround + r * (RadiusAtmosphere - RadiusGround);
+    muS = -0.15 + muS * (1.0 + 0.15);
+#endif
+}
+
+float2 GetIrradianceUV(float Radius, float MuS) 
+{
+    float V = (Radius - RadiusGround) / (RadiusAtmosphere - RadiusGround);
+    float U = (MuS + 0.2) / (1.0 + 0.2);
+    return float2(U, V);
+}
+
+void GetIrradianceRMuS(float2 UV, out float Radius, out float MuS) 
+{
+    Radius = RadiusGround + (UV.y * float(IrradianceTexHeight) - 0.5) / (float(IrradianceTexHeight) - 1.0) * (RadiusAtmosphere - RadiusGround);
+    MuS = -0.2 + (UV.x * float(IrradianceTexWidth) - 0.5) / (float(IrradianceTexWidth) - 1.0) * (1.0 + 0.2);
+}
+
+float4 Texture4DSample(Texture3D Texture, SamplerState TextureSampler, float Radius, float Mu, float MuS, float Nu)
+{
+    float H = sqrt(RadiusAtmosphere * RadiusAtmosphere - RadiusGround * RadiusGround);
+    float Rho = sqrt(Radius * Radius - RadiusGround * RadiusGround);
+#if INSCATTER_NON_LINEAR
+    float RMu = Radius * Mu;
+    float Delta = RMu * RMu - Radius * Radius + RadiusGround * RadiusGround;
+    float4 TexOffset = RMu < 0.0 && Delta > 0.0 ? float4(1.0, 0.0, 0.0, 0.5 - 0.5 / float(InscatterMuNum)) : float4(-1.0, H * H, H, 0.5 + 0.5 / float(InscatterMuNum));
+	float MuR = 0.5 / float(View.AtmosphericFogInscatterAltitudeSampleNum) + Rho / H * (1.0 - 1.0 / float(View.AtmosphericFogInscatterAltitudeSampleNum));
+    float MuMu = TexOffset.w + (RMu * TexOffset.x + sqrt(Delta + TexOffset.y)) / (Rho + TexOffset.z) * (0.5 - 1.0 / float(InscatterMuNum));
+    // paper formula
+    //float MuMuS = 0.5 / float(InscatterMuSNum) + max((1.0 - exp(-3.0 * MuS - 0.6)) / (1.0 - exp(-3.6)), 0.0) * (1.0 - 1.0 / float(InscatterMuSNum));
+    // better formula
+    float MuMuS = 0.5 / float(InscatterMuSNum) + (atan(max(MuS, -0.1975) * tan(1.26 * 1.1)) / 1.1 + (1.0 - 0.26)) * 0.5 * (1.0 - 1.0 / float(InscatterMuSNum));
+#else
+	float MuR = 0.5 / float(View.AtmosphericFogInscatterAltitudeSampleNum) + Rho / H * (1.0 - 1.0 / float(View.AtmosphericFogInscatterAltitudeSampleNum));
+    float MuMu = 0.5 / float(InscatterMuNum) + (Mu + 1.0) * 0.5f * (1.0 - 1.0 / float(InscatterMuNum));
+    float MuMuS = 0.5 / float(InscatterMuSNum) + max(MuS + 0.2, 0.0) / 1.2 * (1.0 - 1.0 / float(InscatterMuSNum));
+#endif
+    float LerpValue = (Nu + 1.0) * 0.5f * (float(InscatterNuNum) - 1.0);
+    float MuNu = floor(LerpValue);
+    LerpValue = LerpValue - MuNu;
+
+    return Texture3DSampleLevel(Texture, TextureSampler, float3((MuNu + MuMuS) / float(InscatterNuNum), MuMu, MuR), 0) * (1.0 - LerpValue) +
+           Texture3DSampleLevel(Texture, TextureSampler, float3((MuNu + MuMuS + 1.0) / float(InscatterNuNum), MuMu, MuR), 0) * LerpValue;
+}
+
+float Mod(float X, float Y)
+{
+	return X - Y * floor(X/Y);
+}
+
+void GetMuMuSNu(float2 UV, float Radius, float4 DhdH, out float Mu, out float MuS, out float Nu) 
+{
+    float X = UV.x * float(InscatterMuSNum * InscatterNuNum) - 0.5;
+    float Y = UV.y * float(InscatterMuNum) - 0.5;
+#if INSCATTER_NON_LINEAR
+    if (Y < float(InscatterMuNum) * 0.5f) 
+	{
+        float D = 1.0 - Y / (float(InscatterMuNum) * 0.5f - 1.0);
+        D = min(max(DhdH.z, D * DhdH.w), DhdH.w * 0.999);
+        Mu = (RadiusGround * RadiusGround - Radius * Radius - D * D) / (2.0 * Radius * D);
+        Mu = min(Mu, -sqrt(1.0 - (RadiusGround / Radius) * (RadiusGround / Radius)) - 0.001);
+    }
+	else 
+	{
+        float D = (Y - float(InscatterMuNum) * 0.5f) / (float(InscatterMuNum) * 0.5f - 1.0);
+        D = min(max(DhdH.x, D * DhdH.y), DhdH.y * 0.999);
+        Mu = (RadiusAtmosphere * RadiusAtmosphere - Radius * Radius - D * D) / (2.0 * Radius * D);
+    }
+    MuS = Mod(X, float(InscatterMuSNum)) / (float(InscatterMuSNum) - 1.0);
+    // paper formula
+    //MuS = -(0.6 + log(1.0 - MuS * (1.0 -  exp(-3.6)))) / 3.0;
+    // better formula
+    MuS = tan((2.0 * MuS - 1.0 + 0.26) * 1.1) / tan(1.26 * 1.1);
+    Nu = -1.0 + floor(X / float(InscatterMuSNum)) / (float(InscatterNuNum) - 1.0) * 2.0;
+#else
+    Mu = -1.0 + 2.0 * Y / (float(InscatterMuNum) - 1.0);
+    MuS = Mod(X, float(InscatterMuSNum)) / (float(InscatterMuSNum) - 1.0);
+    MuS = -0.2 + MuS * 1.2;
+    Nu = -1.0 + floor(X / float(InscatterMuSNum)) / (float(InscatterNuNum) - 1.0) * 2.0;
+#endif
+}
+
+/**
+ * Nearest intersection of ray r,mu with ground or top atmosphere boundary
+ * mu=cos(ray zenith angle at ray origin)
+ */
+float Limit(float Radius, float Mu) 
+{
+    float Dout = -Radius * Mu + sqrt(Radius * Radius * (Mu * Mu - 1.0) + RadiusLimit * RadiusLimit);
+    float Delta2 = Radius * Radius * (Mu * Mu - 1.0) + RadiusGround * RadiusGround;
+    if (Delta2 >= 0.0) 
+	{
+        float Din = -Radius * Mu - sqrt(Delta2);
+        if (Din >= 0.0) 
+		{
+            Dout = min(Dout, Din);
+        }
+    }
+    return Dout;
+}
+
+/**
+ * Transmittance(=transparency) of atmosphere for infinite ray (r,mu)
+ * (mu=cos(view zenith angle)), intersections with ground ignored
+ */
+float3 Transmittance(float Radius, float Mu) 
+{
+	float2 UV = GetTransmittanceUV(Radius, Mu);
+	return Texture2DSampleLevel(AtmosphereTransmittanceTexture, AtmosphereTransmittanceTextureSampler, UV, 0).rgb;
+}
+
+/**
+ * Transmittance(=transparency) of atmosphere for infinite ray (r,mu)
+ * (mu=cos(view zenith angle)), or zero if ray intersects ground
+ */
+float3 TransmittanceWithShadow(float Radius, float Mu) 
+{
+	// Need to correct calculation based on shadow feature, currently don't consider
+    //return Mu < -sqrt(1.0 - (RadiusGround / Radius) * (RadiusGround / Radius)) ? float3(0.f, 0.f, 0.f) : Transmittance(Radius, Mu);
+	return Transmittance(Radius, Mu);
+}
+
+/**
+ * Transmittance(=transparency) of atmosphere between x and x0
+ * Assume segment x,x0 not intersecting ground
+ * D = Distance between x and x0, mu=cos(zenith angle of [x,x0) ray at x)
+ */
+float3 TransmittanceWithDistance(float Radius, float Mu, float D) 
+{
+    float3 Result;
+    float R1 = sqrt(Radius * Radius + D * D + 2.0 * Radius * Mu * D);
+    float Mu1 = (Radius * Mu + D) / R1;
+    if (Mu > 0.0) 
+	{
+        Result = min(Transmittance(Radius, Mu) / Transmittance(R1, Mu1), 1.0);
+    }
+	else 
+	{
+        Result = min(Transmittance(R1, -Mu1) / Transmittance(Radius, -Mu), 1.0);
+    }
+    return Result;
+}
+
+/**
+ * Transmittance(=transparency) of atmosphere between x and x0
+ * Assume segment x,x0 not intersecting ground
+ * Radius=||x||, Mu=cos(zenith angle of [x,x0) ray at x), v=unit direction vector of [x,x0) ray
+ */
+float3 TransmittanceWithDistance(float Radius, float Mu, float3 V, float3 X0) 
+{
+    float3 Result;
+    float R1 = length(X0);
+    float Mu1 = dot(X0, V) / Radius;
+    if (Mu > 0.0) 
+	{
+        Result = min(Transmittance(Radius, Mu) / Transmittance(R1, Mu1), 1.0);
+    } 
+	else 
+	{
+        Result = min(Transmittance(R1, -Mu1) / Transmittance(Radius, -Mu), 1.0);
+    }
+    return Result;
+}
+
+/**
+ * Optical depth for ray (r,mu) of length d, using analytic formula
+ * (mu=cos(view zenith angle)), intersections with ground ignored
+ * H=height scale of exponential density function
+ */
+float OpticalDepthWithDistance(float H, float Radius, float Mu, float D) 
+{
+	float ParticleDensity = 6.2831; // REK 04, Table 2
+    float A = sqrt((0.5/H)*Radius);
+    float2 A01 = A * float2(Mu, Mu + D / Radius);
+    float2 A01Sign = sign(A01);
+    float2 A01Squared = A01*A01;
+    float X = A01Sign.y > A01Sign.x ? exp(A01Squared.x) : 0.0;
+    float2 Y = A01Sign / (2.3193*abs(A01) + sqrt(1.52*A01Squared + 4.0)) * float2(1.0, exp(-D/H*(D/(2.0*Radius)+Mu)));
+    return sqrt((ParticleDensity*H)*Radius) * exp((RadiusGround-Radius)/H) * (X + dot(Y, float2(1.0, -1.0)));
+}
+
+float3 Irradiance(Texture2D Texture, SamplerState TextureSampler, float r, float muS) 
+{
+    float2 UV = GetIrradianceUV(r, muS);
+	return Texture2DSampleLevel(Texture, TextureSampler, UV, 0).rgb;
+}
+
+/** Rayleigh phase function */
+float PhaseFunctionR(float Mu) 
+{
+    return (3.0 / (16.0 * PI)) * (1.0 + Mu * Mu);
+}
+
+/** Mie phase function */
+float PhaseFunctionM(float Mu) 
+{
+	return 1.5 * 1.0 / (4.0 * PI) * (1.0 - MieG * MieG) * pow( abs(1.0 + (MieG * MieG) - 2.0 * MieG * Mu), -3.0/2.0) * (1.0 + Mu * Mu) / (2.0 + MieG * MieG);
+}
+
+/** Approximated single Mie scattering (cf. approximate Cm in paragraph "Angular precision") */
+float3 GetMie(float4 RayMie) 
+{ 
+	// RayMie.rgb=C*, RayMie.w=Cm,r
+	return RayMie.rgb * RayMie.w / max(RayMie.r, 1e-4) * (BetaRayleighScattering.rrr / BetaRayleighScattering.rgb);
+}
+
+/** Transmittance(=transparency) of atmosphere for ray (r,mu) of length d
+ * (mu=cos(view zenith angle)), intersections with ground ignored
+ * uses analytic formula instead of transmittance texture, REK 04, Atmospheric Transparency
+ */
+float3 AnalyticTransmittance(float R, float Mu, float D) 
+{
+    return exp(- BetaRayleighScattering * OpticalDepthWithDistance(View.AtmosphericFogHeightScaleRayleigh, R, Mu, D) - BetaMieExtinction * OpticalDepthWithDistance(HeightScaleMie, R, Mu, D));
+}
+
+static const float HeightOffset = 0.01f;
+
+/** inscattered light along ray x+tv, when sun in direction s (=S[L]-T(x,x0)S[L]|x0) */
+float3 GetInscatterColor(float FogDepth, float3 X, float T, float3 V, float3 S, float Radius, float Mu, out float3 Attenuation, bool bIsSceneGeometry) 
+{
+    float3 Result = float3(0.f, 0.f, 0.f); 	// X in space and ray looking in space, initialize
+	Attenuation = float3(1.f, 1.f, 1.f);
+
+	float D = -Radius * Mu - sqrt(Radius * Radius * (Mu * Mu - 1.0) + RadiusAtmosphere * RadiusAtmosphere);
+    if (D > 0.0) 
+	{ 
+		// if X in space and ray intersects atmosphere
+        // move X to nearest intersection of ray with top atmosphere boundary
+        X += D * V;
+        T -= D;
+        Mu = (Radius * Mu + D) / RadiusAtmosphere;
+        Radius = RadiusAtmosphere;
+    }
+
+	float Epsilon = 0.005f;
+
+	if (Radius < RadiusGround + HeightOffset + Epsilon)
+	{
+		float Diff = (RadiusGround + HeightOffset + Epsilon) - Radius;
+		X -= Diff * V;
+		T -= Diff;
+		Radius = RadiusGround + HeightOffset + Epsilon;
+		Mu = dot(X, V) / Radius;
+	}
+
+	if (Radius <= RadiusAtmosphere && FogDepth > 0.f) 
+	{ 
+		float3 X0 = X + T * V;
+		float R0 = length(X0);
+		// if ray intersects atmosphere
+		float Nu = dot(V, S);
+		float MuS = dot(X, S) / Radius;
+
+		float MuHorizon = -sqrt(1.0 - (RadiusGround / Radius) * (RadiusGround / Radius));
+
+		if (bIsSceneGeometry)
+		{
+			Mu = max(Mu, MuHorizon + Epsilon + 0.15);
+		}
+		else
+		{ 
+			Mu = max(Mu, MuHorizon + Epsilon);
+		}
+
+		float MuOriginal = Mu;
+
+		float BlendRatio = 0.0f;
+
+		if (bIsSceneGeometry)
+		{
+			BlendRatio = saturate(exp(-V.z) - 0.5);
+			if (BlendRatio < 1.f)
+			{
+				V.z = max(V.z, 0.15);
+				V = normalize(V);
+				float3 X1 = X + T * V;
+				float R1 = length(X1);
+				Mu = dot(X1, V) / R1; 
+			}
+		}
+
+		float PhaseR = PhaseFunctionR(Nu);
+		float PhaseM = PhaseFunctionM(Nu);
+		float4 Inscatter = max(Texture4DSample(AtmosphereInscatterTexture, AtmosphereInscatterTextureSampler, Radius, Mu, MuS, Nu), 0.0);
+
+		if (T > 0.0) 
+		{
+#if ATMOSPHERIC_TEXTURE_SAMPLE_FIX
+			// avoids imprecision problems in transmittance computations based on textures
+			Attenuation = AnalyticTransmittance(Radius, Mu, T);
+#else
+			Attenuation = TransmittanceWithDistance(Radius, Mu, V, X0);
+#endif
+
+			float Mu0 = dot(X0, V) / R0;
+			float MuS0 = dot(X0, S) / R0;
+
+			if (bIsSceneGeometry)
+			{
+				R0 = max(R0, Radius);
+			}
+
+			if (R0 > RadiusGround + HeightOffset)
+			{
+				if (BlendRatio < 1.f)
+				{
+					Inscatter = max(Inscatter - Attenuation.rgbr * Texture4DSample(AtmosphereInscatterTexture, AtmosphereInscatterTextureSampler, R0, Mu0, MuS0, Nu), 0.0);
+#if ATMOSPHERIC_TEXTURE_SAMPLE_FIX
+					// avoids imprecision problems near horizon by interpolating between two points above and below horizon
+					if (!bIsSceneGeometry ) 
+					{
+						if (abs(Mu - MuHorizon) < Epsilon)
+						{
+							float Alpha = ((Mu - MuHorizon) + Epsilon) * 0.5f / Epsilon;
+
+							Mu = MuHorizon - Epsilon;
+							R0 = sqrt(Radius * Radius + T * T + 2.0 * Radius * T * Mu);
+							Mu0 = (Radius * Mu + T) / R0;
+
+							Mu0 = max(MuHorizon + Epsilon, Mu0);
+							float4 Inscatter0 = Texture4DSample(AtmosphereInscatterTexture, AtmosphereInscatterTextureSampler, Radius, Mu, MuS, Nu);
+							float4 Inscatter1 = Texture4DSample(AtmosphereInscatterTexture, AtmosphereInscatterTextureSampler, R0, Mu0, MuS0, Nu);
+							float4 InscatterA = max(Inscatter0 - Attenuation.rgbr * Inscatter1, 0.0);
+
+							Mu = MuHorizon + Epsilon;
+							R0 = sqrt(Radius * Radius + T * T + 2.0 * Radius * T * Mu);
+
+							Mu0 = (Radius * Mu + T) / R0;
+							Mu0 = max(MuHorizon + Epsilon, Mu0);
+							Inscatter0 = Texture4DSample(AtmosphereInscatterTexture, AtmosphereInscatterTextureSampler, Radius, Mu, MuS, Nu);
+							Inscatter1 = Texture4DSample(AtmosphereInscatterTexture, AtmosphereInscatterTextureSampler, R0, Mu0, MuS0, Nu);
+							float4 InscatterB = max(Inscatter0 - Attenuation.rgbr * Inscatter1, 0.0);
+
+							Inscatter = lerp(InscatterA, InscatterB, Alpha);
+						}
+					}
+					else if (BlendRatio > 0.f)
+					{
+						Inscatter = lerp(Inscatter,
+							(1.f - Attenuation.rgbr) * max(Texture4DSample(AtmosphereInscatterTexture, AtmosphereInscatterTextureSampler, Radius, MuOriginal, MuS, Nu), 0.0), 
+							BlendRatio);
+					}
+#endif
+				}
+				else
+				{
+					Inscatter = (1.f - Attenuation.rgbr) * Inscatter; 
+				}
+			}
+		}
+#if ATMOSPHERIC_TEXTURE_SAMPLE_FIX
+        // avoids imprecision problems in Mie scattering when sun is below horizon
+        Inscatter.w *= smoothstep(0.00, 0.02, MuS);
+#endif
+        Result = max(Inscatter.rgb * PhaseR + GetMie(Inscatter) * PhaseM, 0.0);
+    } 
+
+	return Result;
+}
+
+/** 
+  * Ground radiance at end of ray x+tv, when sun in direction s
+  * attenuated bewteen ground and viewer (=R[L0]+R[L*])
+  */
+float3 GetGroundColor(float4 SceneColor, float3 X, float T, float3 V, float3 S, float Radius, float3 Attenuation, bool bIsSceneGeometry)
+{
+    float3 Result = float3(0.f, 0.f, 0.f); 	// ray looking at the sky (for intial value)
+    if (T > 0.0) 
+	{ 
+		// if ray hits ground surface
+        // ground Reflectance at end of ray, X0
+        float3 X0 = X + T * V;
+        float R0 = length(X0);
+        float3 N = X0 / R0;
+		N = X0 / R0;
+		SceneColor.xyz = saturate(SceneColor.xyz + 0.05);
+
+        float4 Reflectance = SceneColor * float4(0.2, 0.2, 0.2, 1.0);
+
+        // direct sun light (radiance) reaching X0
+        float MuS = dot(N, S);
+        float3 SunLight = bIsSceneGeometry ? float3(0.f, 0.f, 0.f) : TransmittanceWithShadow(R0, MuS);
+
+        // precomputed sky light (irradiance) (=E[L*]) at X0
+        float3 GroundSkyLight = Irradiance(AtmosphereIrradianceTexture, AtmosphereIrradianceTextureSampler, R0, MuS);
+
+        // light reflected at X0 (=(R[L0]+R[L*])/T(X,X0))
+        //
+		float3 GroundColor = (Reflectance.rgb * (max(MuS, 0.0) * SunLight + GroundSkyLight)) / PI;
+
+        // water specular color due to SunLight
+        if (!bIsSceneGeometry && Reflectance.w > 0.0) 
+		{
+            float3 H = normalize(S - V);
+            float Fresnel = 0.02 + 0.98 * pow(1.0 - dot(-V, H), 5.0);
+            float WaterBrdf = Fresnel * pow(max(dot(H, N), 0.0), 150.0);
+            GroundColor += Reflectance.w * max(WaterBrdf, 0.0) * SunLight;
+        }
+
+		Result = Attenuation * GroundColor; //=R[L0]+R[L*]
+    } 
+    return Result;
+}
+
+/** Direct sun light for ray x+tv, when sun in direction s (=L0) */
+float3 GetSunColor(float3 X, float T, float3 V, float3 S, float Radius, float Mu) 
+{
+	float3 TransmittanceValue = Radius <= RadiusAtmosphere ? TransmittanceWithShadow(Radius, Mu) : float3(1.0, 1.0, 1.0); // T(X,xo)
+    if (T > 0.0) 
+	{
+        return float3(0.f, 0.f, 0.f);
+    }
+	else 
+	{
+		float SunIntensity = step(cos(PI * View.AtmosphericFogSunDiscScale / 180.0), dot(V, S)); // Lsun
+        return TransmittanceValue * SunIntensity; // Eq (9)
+    }
+}
+
+float4 GetAtmosphericFog(float3 ViewPosition, float3 ViewVector, float SceneDepth, float3 SceneColor)
+{
+	float Scale = 0.00001f * View.AtmosphericFogDistanceScale;
+	ViewPosition.z = (ViewPosition.z - View.AtmosphericFogGroundOffset) * View.AtmosphericFogAltitudeScale;
+	ViewPosition *= Scale;
+	ViewPosition.z += RadiusGround + HeightOffset;
+
+	float Radius = length(ViewPosition);
+	
+	float3 V = normalize(ViewVector.xyz);
+
+	float Mu = dot(ViewPosition, V) / Radius;
+
+#if COMPILER_METAL
+	float T = -Radius * Mu - sqrt(abs(Radius * Radius * (Mu * Mu - 1.0) + RadiusGround * RadiusGround));
+#else
+	float T = -Radius * Mu - sqrt(Radius * Radius * (Mu * Mu - 1.0) + RadiusGround * RadiusGround);
+#endif
+
+	float DepthThreshold = 100.f * View.AtmosphericFogDistanceScale; // 100km limit
+	SceneDepth *= Scale;
+
+	float FogDepth = max(0.f, SceneDepth - View.AtmosphericFogStartDistance);
+	float ShadowFactor = 1.f; // shadow approximation
+	float DistanceRatio = min(FogDepth * 0.1f / View.AtmosphericFogStartDistance, 1.f);
+	bool bIsSceneGeometry = (SceneDepth < DepthThreshold); // Assume as scene geometry
+	if (bIsSceneGeometry)
+	{
+		ShadowFactor = DistanceRatio * View.AtmosphericFogPower;
+		T = max(SceneDepth + View.AtmosphericFogDistanceOffset, 1.f);
+	}
+
+	float3 Attenuation;
+
+	float3 InscatterColor = GetInscatterColor(FogDepth, ViewPosition, T, V, View.AtmosphericFogSunDirection, Radius, Mu, Attenuation, bIsSceneGeometry); //S[L]-T(ViewPosition,xs)S[l]|xs
+#if BASEPASS_ATMOSPHERIC_FOG // Transluceny rendering just follows normal render flag
+	float3 GroundColor = 0.f;
+	BRANCH
+	if ((View.AtmosphericFogRenderMask & RENDERFLAG_DISABLE_GROUND_SCATTERING) == 0)
+	{
+		GroundColor = GetGroundColor(float4(SceneColor.xyz, 1.f), ViewPosition, T, V, View.AtmosphericFogSunDirection, Radius, Attenuation, bIsSceneGeometry); //R[L0]+R[L*]
+	}
+#else
+	#if MATERIAL_ATMOSPHERIC_FOG || ATMOSPHERIC_NO_GROUND_SCATTERING // Material doesn't need this ground scattering for now 
+		float3 GroundColor = 0.f;
+	#else
+		float3 GroundColor = GetGroundColor(float4(SceneColor.xyz, 1.f), ViewPosition, T, V, View.AtmosphericFogSunDirection, Radius, Attenuation, bIsSceneGeometry); //R[L0]+R[L*]
+	#endif
+#endif
+
+#if BASEPASS_ATMOSPHERIC_FOG // Transluceny rendering just follows normal render flag
+	float3 Sun = 0.f;
+	BRANCH
+	if ((View.AtmosphericFogRenderMask & RENDERFLAG_DISABLE_SUN_DISK) == 0)
+	{
+		Sun = GetSunColor(ViewPosition, T, V, View.AtmosphericFogSunDirection, Radius, Mu); //L0
+	}
+#else
+	#if MATERIAL_ATMOSPHERIC_FOG || ATMOSPHERIC_NO_SUN_DISK // Material doesn't need to render the sun disk
+		float3 Sun = 0.f;
+	#else
+		float3 Sun = GetSunColor(ViewPosition, T, V, View.AtmosphericFogSunDirection, Radius, Mu); //L0
+	#endif
+#endif
+
+	// Decay color
+	float3 OriginalColor = Sun + GroundColor + InscatterColor;
+	OriginalColor = View.AtmosphericFogSunPower * ShadowFactor * View.AtmosphericFogSunColor.rgb * OriginalColor;
+	float4 OutColor = float4(OriginalColor, lerp(saturate(Attenuation.r * View.AtmosphericFogDensityScale - View.AtmosphericFogDensityOffset), 1.f, (1.f - DistanceRatio)) );
+	return OutColor;
+}
+
+float4 CalculateVertexAtmosphericFog(float3 WorldPosition, float3 InCameraPosition)
+{
+	float3 ViewVector = WorldPosition - InCameraPosition;
+	float Distance = length(ViewVector);
+	return GetAtmosphericFog(InCameraPosition, ViewVector, Distance, float3(0.f, 0.f, 0.f));
+}
+
+#endif
+
diff --git a/Engine/Shaders/Private/AtmospherePrecomputeCommon.ush b/Engine/Shaders/Private/AtmospherePrecomputeCommon.ush
new file mode 100644
--- /dev/null
+++ b/Engine/Shaders/Private/AtmospherePrecomputeCommon.ush
@@ -0,0 +1,47 @@
+// Copyright 1998-2017 Epic Games, Inc. All Rights Reserved.
+
+/*=============================================================================
+	AtmospherePrecomputeCommon.usf: Functions and variables only used in precomputation in editor
+
+	This code contains embedded portions of free sample source code from 
+	http://www-evasion.imag.fr/Membres/Eric.Bruneton/PrecomputedAtmosphericScattering2.zip, Author: Eric Bruneton, 
+	08/16/2011, Copyright (c) 2008 INRIA, All Rights Reserved, which have been altered from their original version.
+
+	Permission is granted to anyone to use this software for any purpose, including commercial applications, and to alter it and redistribute it freely, subject to the following restrictions:
+
+    1. Redistributions of source code must retain the above copyright notice, 
+	   this list of conditions and the following disclaimer.
+    2. Redistributions in binary form must reproduce the above copyright notice, 
+	   this list of conditions and the following disclaimer in the
+       documentation and/or other materials provided with the distribution.
+    3. Neither the name of the copyright holders nor the names of its
+       contributors may be used to endorse or promote products derived from
+       this software without specific prior written permission.
+
+=============================================================================*/
+
+#ifndef __ATMOSPHERE_PRECOMPUTE_COMMON__
+#define __ATMOSPHERE_PRECOMPUTE_COMMON__
+
+/** Integration Parameters */
+const static int TransmittanceIntegralSamples = 500;
+const static int InscatterIntegralSamples = 50;
+const static int IrradianceIntegralSamplesHalf = 16; // 32 / 2, division for integer has some problem in nVidia Card
+const static int InscatterSphericalIntegralSamples = 16;
+
+const static float DeltaPhi = PI / float(InscatterSphericalIntegralSamples);
+const static float DeltaTheta = PI / float(InscatterSphericalIntegralSamples);
+
+float FirstOrder;
+
+/** Textures only for precomputation on editor */
+Texture2D AtmosphereDeltaETexture;
+SamplerState AtmosphereDeltaETextureSampler;
+Texture3D AtmosphereDeltaSRTexture;
+SamplerState AtmosphereDeltaSRTextureSampler;
+Texture3D AtmosphereDeltaSMTexture;
+SamplerState AtmosphereDeltaSMTextureSampler;
+Texture3D AtmosphereDeltaJTexture;
+SamplerState AtmosphereDeltaJTextureSampler;
+
+#endif
diff --git a/Engine/Shaders/Private/BasePassCommon.ush b/Engine/Shaders/Private/BasePassCommon.ush
new file mode 100644
--- /dev/null
+++ b/Engine/Shaders/Private/BasePassCommon.ush
@@ -0,0 +1,145 @@
+// Copyright 1998-2017 Epic Games, Inc. All Rights Reserved.
+
+/*=============================================================================
+	BasePassCommon.usf: Base pass definitions used by both vertex and pixel shader
+=============================================================================*/
+
+#undef NEEDS_LIGHTMAP_COORDINATE
+
+#define NEEDS_LIGHTMAP_COORDINATE		(HQ_TEXTURE_LIGHTMAP || LQ_TEXTURE_LIGHTMAP)
+// Translucent materials need to compute fogging in the forward shading pass
+// Materials that read from scene color skip getting fogged, because the contents of the scene color lookup have already been fogged
+// This is not foolproof, as any additional color the material adds will then not be fogged correctly
+#define TRANSLUCENCY_NEEDS_BASEPASS_FOGGING	(ENABLE_TRANSLUCENCY_FOGGING && MATERIALBLENDING_ANY_TRANSLUCENT && !MATERIAL_USES_SCENE_COLOR_COPY)
+#define OPAQUE_NEEDS_BASEPASS_FOGGING		(!MATERIALBLENDING_ANY_TRANSLUCENT && VERTEX_FOGGING_FOR_OPAQUE && FORWARD_SHADING)
+#define NEEDS_BASEPASS_VERTEX_FOGGING		(TRANSLUCENCY_NEEDS_BASEPASS_FOGGING && !COMPUTE_FOG_PER_PIXEL || OPAQUE_NEEDS_BASEPASS_FOGGING)
+#define NEEDS_BASEPASS_PIXEL_FOGGING		(TRANSLUCENCY_NEEDS_BASEPASS_FOGGING && COMPUTE_FOG_PER_PIXEL)
+
+#define NEEDS_LIGHTMAP						(NEEDS_LIGHTMAP_COORDINATE)
+
+#define USES_GBUFFER						(FEATURE_LEVEL >= FEATURE_LEVEL_SM4 && (MATERIALBLENDING_SOLID || MATERIALBLENDING_MASKED) && !SIMPLE_FORWARD_SHADING && !FORWARD_SHADING)
+
+// Only some shader models actually need custom data.
+#define WRITES_CUSTOMDATA_TO_GBUFFER		(USES_GBUFFER && (MATERIAL_SHADINGMODEL_SUBSURFACE || MATERIAL_SHADINGMODEL_PREINTEGRATED_SKIN || MATERIAL_SHADINGMODEL_SUBSURFACE_PROFILE || MATERIAL_SHADINGMODEL_CLEAR_COAT || MATERIAL_SHADINGMODEL_TWOSIDED_FOLIAGE || MATERIAL_SHADINGMODEL_HAIR || MATERIAL_SHADINGMODEL_CLOTH || MATERIAL_SHADINGMODEL_EYE))
+
+// Based on GetPrecomputedShadowMasks()
+// Note: WRITES_PRECSHADOWFACTOR_TO_GBUFFER is currently disabled because we use the precomputed shadow factor GBuffer outside of STATICLIGHTING_TEXTUREMASK to store UseSingleSampleShadowFromStationaryLights
+#define GBUFFER_HAS_PRECSHADOWFACTOR		(USES_GBUFFER && ALLOW_STATIC_LIGHTING)
+#define WRITES_PRECSHADOWFACTOR_ZERO		(!(STATICLIGHTING_TEXTUREMASK && STATICLIGHTING_SIGNEDDISTANCEFIELD) && (HQ_TEXTURE_LIGHTMAP || LQ_TEXTURE_LIGHTMAP))
+#define WRITES_PRECSHADOWFACTOR_TO_GBUFFER	(GBUFFER_HAS_PRECSHADOWFACTOR && !WRITES_PRECSHADOWFACTOR_ZERO)
+
+// If a primitive has static lighting, we assume it is not moving. If it is, it will be rerendered in an extra renderpass.
+#define WRITES_VELOCITY_TO_GBUFFER			(USES_GBUFFER && GBUFFER_HAS_VELOCITY && (!SELECTIVE_BASEPASS_OUTPUTS || !(STATICLIGHTING_TEXTUREMASK || STATICLIGHTING_SIGNEDDISTANCEFIELD || HQ_TEXTURE_LIGHTMAP || LQ_TEXTURE_LIGHTMAP)))
+
+// Use an extra interpolator for position
+#define WRITES_VELOCITY_TO_GBUFFER_USE_POS_INTERPOLATOR		(WRITES_VELOCITY_TO_GBUFFER && USING_TESSELLATION)
+
+#define TRANSLUCENCY_ANY_VOLUMETRIC (TRANSLUCENCY_LIGHTING_VOLUMETRIC_NONDIRECTIONAL || TRANSLUCENCY_LIGHTING_VOLUMETRIC_DIRECTIONAL || TRANSLUCENCY_LIGHTING_VOLUMETRIC_PERVERTEX_NONDIRECTIONAL || TRANSLUCENCY_LIGHTING_VOLUMETRIC_PERVERTEX_DIRECTIONAL)
+#define TRANSLUCENCY_PERVERTEX_LIGHTING_VOLUME (!FORWARD_SHADING && (TRANSLUCENCY_LIGHTING_VOLUMETRIC_PERVERTEX_DIRECTIONAL || TRANSLUCENCY_LIGHTING_VOLUMETRIC_PERVERTEX_NONDIRECTIONAL))
+#define TRANSLUCENCY_PERVERTEX_FORWARD_SHADING (FORWARD_SHADING && (TRANSLUCENCY_LIGHTING_VOLUMETRIC_PERVERTEX_DIRECTIONAL || TRANSLUCENCY_LIGHTING_VOLUMETRIC_PERVERTEX_NONDIRECTIONAL))
+
+struct FSharedBasePassInterpolants
+{
+	//for texture-lightmapped translucency we can pass the vertex fog in its own interpolator
+#if NEEDS_BASEPASS_VERTEX_FOGGING
+	float4 VertexFog		: TEXCOORD7;
+#endif
+
+#if !PC_D3D && !GL4_PROFILE && !SWITCH_PROFILE && !SWITCH_PROFILE_FORWARD && !METAL_TESSELLATION_PROFILE
+	// Note: TEXCOORD8 is unused
+
+	#if USE_WORLD_POSITION_EXCLUDING_SHADER_OFFSETS
+	float3 PixelPositionExcludingWPO : TEXCOORD9;
+	#endif
+#endif
+
+#if TRANSLUCENCY_PERVERTEX_LIGHTING_VOLUME
+	float3 AmbientLightingVector : TEXCOORD12;
+#endif
+
+#if TRANSLUCENCY_PERVERTEX_LIGHTING_VOLUME && TRANSLUCENCY_LIGHTING_VOLUMETRIC_PERVERTEX_DIRECTIONAL
+	float3 DirectionalLightingVector : TEXCOORD13;
+#endif
+
+#if TRANSLUCENCY_PERVERTEX_FORWARD_SHADING
+	float3 VertexDiffuseLighting : TEXCOORD12;
+#endif
+
+#if WRITES_VELOCITY_TO_GBUFFER
+	// .xy is clip position, pre divide by w; .w is clip W; .z is 0 or 1 to mask out the velocity output
+	float4 VelocityPrevScreenPosition : VELOCITY_PREV_POS;
+	#if WRITES_VELOCITY_TO_GBUFFER_USE_POS_INTERPOLATOR
+		float4 VelocityScreenPosition : VELOCITY_POS;
+	#endif
+#endif
+};
+
+#if PC_D3D || GL4_PROFILE || SWITCH_PROFILE || SWITCH_PROFILE_FORWARD || METAL_TESSELLATION_PROFILE
+
+/** Interpolants passed from the vertex shader to the pixel shader. */
+struct FBasePassInterpolantsVSToPS : FSharedBasePassInterpolants
+{
+	// Note: TEXCOORD8 is unused
+
+	#if USE_WORLD_POSITION_EXCLUDING_SHADER_OFFSETS
+	float3 PixelPositionExcludingWPO : TEXCOORD9;
+	#endif
+};
+
+/** Interpolants passed from the vertex shader to the domain shader. */
+struct FBasePassInterpolantsVSToDS : FSharedBasePassInterpolants
+{
+	#if USE_WORLD_POSITION_EXCLUDING_SHADER_OFFSETS
+	float3 WorldPositionExcludingWPO : TEXCOORD9;
+	#endif
+};
+
+#else
+
+// Only the PC shader compiler supports HLSL inheritance
+#define FBasePassInterpolantsVSToPS FSharedBasePassInterpolants
+
+#endif
+
+Texture3D TranslucencyLightingVolumeAmbientInner;
+SamplerState TranslucencyLightingVolumeAmbientInnerSampler;
+Texture3D TranslucencyLightingVolumeAmbientOuter;
+SamplerState TranslucencyLightingVolumeAmbientOuterSampler;
+Texture3D TranslucencyLightingVolumeDirectionalInner;
+SamplerState TranslucencyLightingVolumeDirectionalInnerSampler;
+Texture3D TranslucencyLightingVolumeDirectionalOuter;
+SamplerState TranslucencyLightingVolumeDirectionalOuterSampler;
+
+void ComputeVolumeUVs(float3 WorldPosition, float3 LightingPositionOffset, out float3 InnerVolumeUVs, out float3 OuterVolumeUVs, out float FinalLerpFactor)
+{
+	// Apply a stable offset to the world position used for lighting, which breaks up artifacts from using a low res volume texture
+	InnerVolumeUVs = (WorldPosition + LightingPositionOffset - View.TranslucencyLightingVolumeMin[0].xyz) * View.TranslucencyLightingVolumeInvSize[0].xyz;
+	OuterVolumeUVs = (WorldPosition + LightingPositionOffset - View.TranslucencyLightingVolumeMin[1].xyz) * View.TranslucencyLightingVolumeInvSize[1].xyz;
+
+	// Controls how fast the lerp between the inner and outer cascades happens
+	// Larger values result in a shorter transition distance
+	float TransitionScale = 6;
+	// Setup a 3d lerp factor going to 0 at the edge of the inner volume
+	float3 LerpFactors = saturate((.5f - abs(InnerVolumeUVs - .5f)) * TransitionScale);
+	FinalLerpFactor = LerpFactors.x * LerpFactors.y * LerpFactors.z;
+}
+
+float4 GetAmbientLightingVectorFromTranslucentLightingVolume(float3 InnerVolumeUVs, float3 OuterVolumeUVs, float FinalLerpFactor)
+{
+	// Lookup the inner and outer cascade ambient lighting values
+	float4 InnerLighting = Texture3DSampleLevel(TranslucencyLightingVolumeAmbientInner, TranslucencyLightingVolumeAmbientInnerSampler, InnerVolumeUVs, 0);
+	float4 OuterLighting = Texture3DSampleLevel(TranslucencyLightingVolumeAmbientOuter, TranslucencyLightingVolumeAmbientOuterSampler, OuterVolumeUVs, 0);
+
+	// Lerp between cascades
+	return lerp(OuterLighting, InnerLighting, FinalLerpFactor);
+}
+
+float3 GetDirectionalLightingVectorFromTranslucentLightingVolume(float3 InnerVolumeUVs, float3 OuterVolumeUVs, float FinalLerpFactor)
+{
+	// Fetch both the ambient and directional values for both cascades
+	float3 InnerVector1 = Texture3DSampleLevel(TranslucencyLightingVolumeDirectionalInner, TranslucencyLightingVolumeDirectionalInnerSampler, InnerVolumeUVs, 0).rgb;
+	float3 OuterVector1 = Texture3DSampleLevel(TranslucencyLightingVolumeDirectionalOuter, TranslucencyLightingVolumeDirectionalOuterSampler, OuterVolumeUVs, 0).rgb;
+
+	// Lerp between cascades
+	return lerp(OuterVector1, InnerVector1, FinalLerpFactor);
+}
diff --git a/Engine/Shaders/Private/BasePassVertexCommon.ush b/Engine/Shaders/Private/BasePassVertexCommon.ush
new file mode 100644
--- /dev/null
+++ b/Engine/Shaders/Private/BasePassVertexCommon.ush
@@ -0,0 +1,45 @@
+// Copyright 1998-2017 Epic Games, Inc. All Rights Reserved.
+
+/*=============================================================================
+	BasePassVertexCommon.usf: Vertex-related base pass definitions, used by vertex,
+	hull and domain shader
+=============================================================================*/
+
+#include "Common.ush"
+#include "/Engine/Generated/Material.ush"
+#include "BasePassCommon.ush"
+#include "/Engine/Generated/VertexFactory.ush"
+
+#if NEEDS_BASEPASS_VERTEX_FOGGING
+	#include "HeightFogCommon.ush"
+#if BASEPASS_ATMOSPHERIC_FOG
+	#include "AtmosphereCommon.ush"
+#endif
+#endif
+
+struct FBasePassVSToPS
+{
+	FVertexFactoryInterpolantsVSToPS FactoryInterpolants;
+	FBasePassInterpolantsVSToPS BasePassInterpolants;
+	float4 Position : SV_POSITION;
+};
+
+#if USING_TESSELLATION	
+	struct FBasePassVSToDS
+	{
+		FVertexFactoryInterpolantsVSToDS FactoryInterpolants;
+		FBasePassInterpolantsVSToDS BasePassInterpolants;
+		float4 Position : VS_To_DS_Position;
+		OPTIONAL_VertexID_VS_To_DS
+	};
+	
+	#define FBasePassVSOutput FBasePassVSToDS
+	#define VertexFactoryGetInterpolants VertexFactoryGetInterpolantsVSToDS
+	#define FPassSpecificVSToDS FBasePassVSToDS
+	#define FPassSpecificVSToPS FBasePassVSToPS
+#else
+	#define FBasePassVSOutput FBasePassVSToPS
+	#define VertexFactoryGetInterpolants VertexFactoryGetInterpolantsVSToPS
+#endif
+
+
