Branch snapshot for CL 3555263

diff --git a/Engine/Shaders/Private/Common.ush b/Engine/Shaders/Private/Common.ush
new file mode 100644
--- /dev/null
+++ b/Engine/Shaders/Private/Common.ush
@@ -0,0 +1,1774 @@
+// Copyright 1998-2017 Epic Games, Inc. All Rights Reserved.
+
+/*=============================================================================
+	PlatformCommon.usf: Common shader code
+=============================================================================*/
+  
+#pragma once
+
+#include "/Engine/Public/Platform.ush"
+
+
+// These types are used for material translator generated code, or any functions the translated code can call
+#if PIXELSHADER
+	#define MaterialFloat half
+	#define MaterialFloat2 half2
+	#define MaterialFloat3 half3
+	#define MaterialFloat4 half4
+	#define MaterialFloat3x3 half3x3
+	#define MaterialFloat4x4 half4x4 
+	#define MaterialFloat4x3 half4x3 
+#else
+	// Material translated vertex shader code always uses floats, 
+	// Because it's used for things like world position and UVs
+	#define MaterialFloat float
+	#define MaterialFloat2 float2
+	#define MaterialFloat3 float3
+	#define MaterialFloat4 float4
+	#define MaterialFloat3x3 float3x3
+	#define MaterialFloat4x4 float4x4 
+	#define MaterialFloat4x3 float4x3 
+#endif
+
+#if POST_PROCESS_ALPHA
+#define SceneColorLayout float4
+#define CastFloat4ToSceneColorLayout(x) (x)
+#define SetSceneColorLayoutToFloat4(dest,value) dest = (value)
+#else
+#define SceneColorLayout float3
+#define CastFloat4ToSceneColorLayout(x) ((x).rgb)
+#define SetSceneColorLayoutToFloat4(dest,value) dest.rgb = (value)
+#endif
+
+#if (!COMPILER_GLSL_ES2 && !COMPILER_GLSL_ES3_1)
+#define TextureExternal Texture2D
+#endif
+
+// Generated file that contains uniform buffer declarations needed by the shader being compiled 
+#include "/Engine/Generated/GeneratedUniformBuffers.ush" 
+
+// uniform buffers specifics
+#include "CommonViewUniformBuffer.ush"
+
+#include "InstancedStereo.ush"
+
+#include "Definitions.usf"
+
+#define METAL_TESSELLATION_PROFILE ((METAL_PROFILE || METAL_SM5_PROFILE) && USING_TESSELLATION)
+
+#ifndef MOBILE_FORCE_DEPTH_TEXTURE_READS
+#define MOBILE_FORCE_DEPTH_TEXTURE_READS 0
+#endif
+
+// Using SV_ClipDistance has overhead (15% slower base pass in triangle bound test scene on PS4) so projects have to opt-in
+#define USE_GLOBAL_CLIP_PLANE (PLATFORM_SUPPORTS_GLOBAL_CLIP_PLANE && PROJECT_ALLOW_GLOBAL_CLIP_PLANE)
+
+// This would need to be a #define in GLSL to ignore the SamplerState, however, it is currently a function call in HLSL
+// for type checking of the parameters - ironically the type checking is really only needed in GLSL!
+MaterialFloat4 Texture1DSample(Texture1D Tex, SamplerState Sampler, float UV)
+{
+#if COMPUTESHADER
+	return Tex.SampleLevel(Sampler, UV, 0);
+#else
+	return Tex.Sample(Sampler, UV);
+#endif
+}
+MaterialFloat4 Texture2DSample(Texture2D Tex, SamplerState Sampler, float2 UV)
+{
+#if COMPUTESHADER
+	return Tex.SampleLevel(Sampler, UV, 0);
+#else
+	return Tex.Sample(Sampler, UV);
+#endif
+}
+MaterialFloat Texture2DSample_A8(Texture2D Tex, SamplerState Sampler, float2 UV)
+{
+#if COMPUTESHADER
+	return Tex.SampleLevel(Sampler, UV, 0) A8_SAMPLE_MASK;
+#else
+	return Tex.Sample(Sampler, UV) A8_SAMPLE_MASK;
+#endif
+}
+MaterialFloat4 Texture3DSample(Texture3D Tex, SamplerState Sampler, float3 UV)
+{
+#if COMPUTESHADER
+	return Tex.SampleLevel(Sampler, UV, 0);
+#else
+	return Tex.Sample(Sampler, UV);
+#endif
+}
+MaterialFloat4 TextureCubeSample(TextureCube Tex, SamplerState Sampler, float3 UV)
+{
+#if COMPUTESHADER
+	return Tex.SampleLevel(Sampler, UV, 0);
+#else
+	return Tex.Sample(Sampler, UV);
+#endif
+}
+MaterialFloat4 Texture1DSampleLevel(Texture1D Tex, SamplerState Sampler, float UV, MaterialFloat Mip)
+{
+	return Tex.SampleLevel(Sampler, UV, Mip);
+}
+MaterialFloat4 Texture2DSampleLevel(Texture2D Tex, SamplerState Sampler, float2 UV, MaterialFloat Mip)
+{
+	return Tex.SampleLevel(Sampler, UV, Mip);
+}
+MaterialFloat4 Texture2DSampleBias(Texture2D Tex, SamplerState Sampler, float2 UV, MaterialFloat MipBias)
+{
+#if COMPUTESHADER
+	return Tex.SampleLevel(Sampler, UV, 0);
+#else
+	return Tex.SampleBias(Sampler, UV, MipBias);
+#endif
+}
+MaterialFloat4 Texture2DSampleGrad(Texture2D Tex, SamplerState Sampler, float2 UV, MaterialFloat2 DDX, MaterialFloat2 DDY)
+{
+	return Tex.SampleGrad(Sampler, UV, DDX, DDY);
+}
+MaterialFloat4 Texture3DSampleLevel(Texture3D Tex, SamplerState Sampler, float3 UV, MaterialFloat Mip)
+{
+	return Tex.SampleLevel(Sampler, UV, Mip);
+}
+MaterialFloat4 TextureCubeSampleLevel(TextureCube Tex, SamplerState Sampler, float3 UV, MaterialFloat Mip)
+{
+	return Tex.SampleLevel(Sampler, UV, Mip);
+}
+MaterialFloat4 TextureCubeSampleBias(TextureCube Tex, SamplerState Sampler, float3 UV, MaterialFloat MipBias)
+{
+#if COMPUTESHADER
+	return Tex.SampleLevel(Sampler, UV, 0);
+#else
+	return Tex.SampleBias(Sampler, UV, MipBias);
+#endif
+}
+MaterialFloat4 TextureCubeSampleGrad(TextureCube Tex, SamplerState Sampler, float3 UV, MaterialFloat3 DDX, MaterialFloat3 DDY)
+{
+	return Tex.SampleGrad(Sampler, UV, DDX, DDY);
+}
+MaterialFloat4 TextureExternalSample(TextureExternal Tex, SamplerState Sampler, float2 UV)
+{
+#if COMPUTESHADER
+	return Tex.SampleLevel(Sampler, UV, 0);
+#else
+	return Tex.Sample(Sampler, UV);
+#endif
+}
+
+//converts an input 1d to 2d position. Useful for locating z frames that have been laid out in a 2d grid like a flipbook.
+float2 Tile1Dto2D(float xsize, float idx)
+{
+	float2 xyidx = 0;
+	xyidx.y = floor(idx / xsize);
+	xyidx.x = idx - xsize * xyidx.y;
+
+	return xyidx;
+}
+
+// return a pseudovolume texture sample.
+// useful for simulating 3D texturing with a 2D texture or as a texture flipbook with lerped transitions
+// treats 2d layout of frames a 3d texture and performs bilinear filtering by blending with an offset Z frame.
+// @param Tex          = Input Texture Object storing Volume Data
+// @param inPos        = Input float3 for Position, 0-1
+// @param xsize        = Input float for num frames in x,y directions
+// @param numFrames    = Input float for num total frames
+// @param mipmode      = Sampling mode: 0 = use miplevel, 1 = use UV computed gradients, 2 = Use gradients (default=0)
+// @param miplevel     = MIP level to use in mipmode=0 (default 0)
+// @param InDDX, InDDY = Texture gradients in mipmode=2
+float4 PseudoVolumeTexture(Texture2D Tex, SamplerState TexSampler, float3 inPos, float2 xysize, float numframes,
+	uint mipmode = 0, float miplevel = 0, float2 InDDX = 0, float2 InDDY = 0)
+{
+	float zframe = ceil(inPos.z * numframes);
+	float zphase = frac(inPos.z * numframes);
+
+	float2 uv = frac(inPos.xy) / xysize;
+
+	float2 curframe = Tile1Dto2D(xysize.x, zframe) / xysize;
+	float2 nextframe = Tile1Dto2D(xysize.x, zframe + 1) / xysize;
+
+	float4 sampleA = 0, sampleB = 0;
+	switch (mipmode)
+	{
+	case 0: // Mip level
+		sampleA = Tex.SampleLevel(TexSampler, uv + curframe, miplevel);
+		sampleB = Tex.SampleLevel(TexSampler, uv + nextframe, miplevel);
+		break;
+	case 1: // Gradients automatic from UV
+		sampleA = Texture2DSample(Tex, TexSampler, uv + curframe);
+		sampleB = Texture2DSample(Tex, TexSampler, uv + nextframe);
+		break;
+	case 2: // Deriviatives provided
+		sampleA = Tex.SampleGrad(TexSampler, uv + curframe,  InDDX, InDDY);
+		sampleB = Tex.SampleGrad(TexSampler, uv + nextframe, InDDX, InDDY);
+		break;
+	default:
+		break;
+	}
+
+	return lerp(sampleA, sampleB, zphase);
+}
+
+#if FEATURE_LEVEL >= FEATURE_LEVEL_SM5 // Cubemap arrays are not supported in SM4 feature level
+	MaterialFloat4 TextureCubeArraySampleLevel(TextureCubeArray Tex, SamplerState Sampler, float3 UV, float ArrayIndex, MaterialFloat Mip)
+	{
+		return Tex.SampleLevel(Sampler, float4(UV, ArrayIndex), Mip);
+	}
+#endif
+
+void Bicubic2DCatmullRom( in float2 UV, in float2 Size, out float2 Sample[3], out float2 Weight[3] )
+{
+	const float2 InvSize = 1.0 / Size;
+ 
+	UV *= Size;
+
+	float2 tc = floor( UV - 0.5 ) + 0.5;
+	float2 f = UV - tc;
+	float2 f2 = f * f;
+	float2 f3 = f2 * f;
+
+	float2 w0 = f2 - 0.5 * (f3 + f);
+	float2 w1 = 1.5 * f3 - 2.5 * f2 + 1;
+	float2 w3 = 0.5 * (f3 - f2);
+	float2 w2 = 1 - w0 - w1 - w3;
+
+	Weight[0] = w0;
+	Weight[1] = w1 + w2;
+	Weight[2] = w3;
+
+	Sample[0] = tc - 1;
+	Sample[1] = tc + w2 / Weight[1];
+	Sample[2] = tc + 2;
+ 
+	Sample[0] *= InvSize;
+	Sample[1] *= InvSize;
+	Sample[2] *= InvSize;
+}
+
+MaterialFloat4 Texture2DSampleBicubic( Texture2D Tex, SamplerState Sampler, float2 UV, float2 Size )
+{
+	float2 Weight[3];
+	float2 Sample[3];
+	Bicubic2DCatmullRom( UV, Size, Sample, Weight );
+
+	MaterialFloat4 OutColor;
+#if 0
+	OutColor  = Tex.SampleLevel( Sampler, float2( Sample[0].x, Sample[0].y ), 0 ) * Weight[0].x * Weight[0].y;
+	OutColor += Tex.SampleLevel( Sampler, float2( Sample[1].x, Sample[0].y ), 0 ) * Weight[1].x * Weight[0].y;
+	OutColor += Tex.SampleLevel( Sampler, float2( Sample[2].x, Sample[0].y ), 0 ) * Weight[2].x * Weight[0].y;
+
+	OutColor += Tex.SampleLevel( Sampler, float2( Sample[0].x, Sample[1].y ), 0 ) * Weight[0].x * Weight[1].y;
+	OutColor += Tex.SampleLevel( Sampler, float2( Sample[1].x, Sample[1].y ), 0 ) * Weight[1].x * Weight[1].y;
+	OutColor += Tex.SampleLevel( Sampler, float2( Sample[2].x, Sample[1].y ), 0 ) * Weight[2].x * Weight[1].y;
+
+	OutColor += Tex.SampleLevel( Sampler, float2( Sample[0].x, Sample[2].y ), 0 ) * Weight[0].x * Weight[2].y;
+	OutColor += Tex.SampleLevel( Sampler, float2( Sample[1].x, Sample[2].y ), 0 ) * Weight[1].x * Weight[2].y;
+	OutColor += Tex.SampleLevel( Sampler, float2( Sample[2].x, Sample[2].y ), 0 ) * Weight[2].x * Weight[2].y;
+#else
+	// Optimized by removing corner samples
+	OutColor  = Tex.SampleLevel( Sampler, float2( Sample[1].x, Sample[0].y ), 0 ) * Weight[1].x * Weight[0].y;
+	OutColor += Tex.SampleLevel( Sampler, float2( Sample[0].x, Sample[1].y ), 0 ) * Weight[0].x * Weight[1].y;
+	OutColor += Tex.SampleLevel( Sampler, float2( Sample[1].x, Sample[1].y ), 0 ) * Weight[1].x * Weight[1].y;
+	OutColor += Tex.SampleLevel( Sampler, float2( Sample[2].x, Sample[1].y ), 0 ) * Weight[2].x * Weight[1].y;
+	OutColor += Tex.SampleLevel( Sampler, float2( Sample[1].x, Sample[2].y ), 0 ) * Weight[1].x * Weight[2].y;
+
+	// Reweight after removing the corners
+	float CornerWeights;
+	CornerWeights  = Weight[1].x * Weight[0].y;
+	CornerWeights += Weight[0].x * Weight[1].y;
+	CornerWeights += Weight[1].x * Weight[1].y;
+	CornerWeights += Weight[2].x * Weight[1].y;
+	CornerWeights += Weight[1].x * Weight[2].y;
+	OutColor /= CornerWeights;
+#endif
+	return OutColor;
+}
+
+	// TANGENTTOWORLD0 is the first row of the tangent to world matrix, w might be needed for padding and is not used yet.
+	// TANGENTTOWORLD2 is the last row of the tangent to world matrix, determinant of tangent basis in w
+
+#if FEATURE_LEVEL >= FEATURE_LEVEL_SM5
+
+	#define TANGENTTOWORLD0					TEXCOORD10
+	#define TANGENTTOWORLD2					TEXCOORD11
+
+	// _centroid is needed to get better quality with MSAA
+
+	// The D3D shader compiler combines _centroid and non controid. Using float3 would results in a internal
+	// shader compiler error. This block is using float4 to prevent that.
+	#define TANGENTTOWORLD_INTERPOLATOR_BLOCK	float4 TangentToWorld0 : TEXCOORD10_centroid; float4	TangentToWorld2	: TEXCOORD11_centroid;
+
+#else
+
+	#define TANGENTTOWORLD0					TEXCOORD10
+	#define TANGENTTOWORLD2					TEXCOORD11
+
+	// TangentToWorld0 is float4 only to match D3D11
+#if METAL_PROFILE || COMPILER_GLSL_ES3_1
+		//@todo-rco: FIXME!
+		#define TANGENTTOWORLD_INTERPOLATOR_BLOCK	float4 TangentToWorld0 : TANGENTTOWORLD0; float4	TangentToWorld2	: TANGENTTOWORLD2;
+#else
+	#define TANGENTTOWORLD_INTERPOLATOR_BLOCK	MaterialFloat4 TangentToWorld0 : TANGENTTOWORLD0; MaterialFloat4	TangentToWorld2	: TANGENTTOWORLD2;
+#endif
+#endif
+
+MaterialFloat Luminance( MaterialFloat3 LinearColor )
+{
+	return dot( LinearColor, MaterialFloat3( 0.3, 0.59, 0.11 ) );
+}
+
+const static MaterialFloat PI = 3.1415926535897932f;
+
+MaterialFloat length2(MaterialFloat2 v)
+{
+	return dot(v, v);
+}
+MaterialFloat length2(MaterialFloat3 v)
+{
+	return dot(v, v);
+}
+MaterialFloat length2(MaterialFloat4 v)
+{
+	return dot(v, v);
+}
+
+uint Mod(uint a, uint b)
+{
+#if FEATURE_LEVEL >= FEATURE_LEVEL_ES3_1
+	return a % b;
+#else
+	return a - (b * (uint)((float)a / (float)b));
+#endif
+}
+
+uint2 Mod(uint2 a, uint2 b)
+{
+#if FEATURE_LEVEL >= FEATURE_LEVEL_ES3_1
+	return a % b;
+#else
+	return a - (b * (uint2)((float2)a / (float2)b));
+#endif
+}
+
+uint3 Mod(uint3 a, uint3 b)
+{
+#if FEATURE_LEVEL >= FEATURE_LEVEL_ES3_1
+	return a % b;
+#else
+	return a - (b * (uint3)((float3)a / (float3)b));
+#endif
+}
+
+MaterialFloat UnClampedPow(MaterialFloat X, MaterialFloat Y)
+{
+	return pow(X, Y);
+}
+MaterialFloat2 UnClampedPow(MaterialFloat2 X, MaterialFloat2 Y)
+{
+	return pow(X, Y);
+}
+MaterialFloat3 UnClampedPow(MaterialFloat3 X, MaterialFloat3 Y)
+{
+	return pow(X, Y);
+}
+MaterialFloat4 UnClampedPow(MaterialFloat4 X, MaterialFloat4 Y)
+{
+	return pow(X, Y);
+}
+
+// Clamp the base, so it's never <= 0.0f (INF/NaN).
+MaterialFloat ClampedPow(MaterialFloat X,MaterialFloat Y)
+{
+	return pow(max(abs(X),0.000001f),Y);
+}
+MaterialFloat2 ClampedPow(MaterialFloat2 X,MaterialFloat2 Y)
+{
+	return pow(max(abs(X),MaterialFloat2(0.000001f,0.000001f)),Y);
+}
+MaterialFloat3 ClampedPow(MaterialFloat3 X,MaterialFloat3 Y)
+{
+	return pow(max(abs(X),MaterialFloat3(0.000001f,0.000001f,0.000001f)),Y);
+}  
+MaterialFloat4 ClampedPow(MaterialFloat4 X,MaterialFloat4 Y)
+{
+	return pow(max(abs(X),MaterialFloat4(0.000001f,0.000001f,0.000001f,0.000001f)),Y);
+} 
+
+MaterialFloat PositiveClampedPow(MaterialFloat X,MaterialFloat Y)
+{
+	return pow(max(X,0.0f),Y);
+}
+MaterialFloat2 PositiveClampedPow(MaterialFloat2 X,MaterialFloat2 Y)
+{
+	return pow(max(X,MaterialFloat2(0.0f,0.0f)),Y);
+}
+MaterialFloat3 PositiveClampedPow(MaterialFloat3 X,MaterialFloat3 Y)
+{
+	return pow(max(X,MaterialFloat3(0.0f,0.0f,0.0f)),Y);
+}  
+MaterialFloat4 PositiveClampedPow(MaterialFloat4 X,MaterialFloat4 Y)
+{
+	return pow(max(X,MaterialFloat4(0.0f,0.0f,0.0f,0.0f)),Y);
+} 
+
+float DDX(float Input)
+{
+#if COMPUTESHADER
+	return 0;
+#else
+	return ddx(Input);
+#endif
+}
+
+float2 DDX(float2 Input)
+{
+#if COMPUTESHADER
+	return 0;
+#else
+	return ddx(Input);
+#endif
+}
+
+float3 DDX(float3 Input)
+{
+#if COMPUTESHADER
+	return 0;
+#else
+	return ddx(Input);
+#endif
+}
+
+float4 DDX(float4 Input)
+{
+#if COMPUTESHADER
+	return 0;
+#else
+	return ddx(Input);
+#endif
+}
+
+float DDY(float Input)
+{
+#if COMPUTESHADER
+	return 0;
+#else
+	return ddy(Input);
+#endif
+}
+
+float2 DDY(float2 Input)
+{
+#if COMPUTESHADER
+	return 0;
+#else
+	return ddy(Input);
+#endif
+}
+
+float3 DDY(float3 Input)
+{
+#if COMPUTESHADER
+	return 0;
+#else
+	return ddy(Input);
+#endif
+}
+
+float4 DDY(float4 Input)
+{
+#if COMPUTESHADER
+	return 0;
+#else
+	return ddy(Input);
+#endif
+}
+
+#include "FastMath.ush"
+#include "Random.ush"	// used by MaterialExpressionNoise
+  
+/** 
+ * Use this function to compute the pow() in the specular computation.
+ * This allows to change the implementation depending on platform or it easily can be replaced by some approxmation.
+ */
+MaterialFloat PhongShadingPow(MaterialFloat X, MaterialFloat Y)
+{
+	// The following clamping is done to prevent NaN being the result of the specular power computation.
+	// Clamping has a minor performance cost.
+
+	// In HLSL pow(a, b) is implemented as exp2(log2(a) * b).
+
+	// For a=0 this becomes exp2(-inf * 0) = exp2(NaN) = NaN.
+
+	// As seen in #TTP 160394 "QA Regression: PS3: Some maps have black pixelated artifacting."
+	// this can cause severe image artifacts (problem was caused by specular power of 0, lightshafts propagated this to other pixels).
+	// The problem appeared on PlayStation 3 but can also happen on similar PC NVidia hardware.
+
+	// In order to avoid platform differences and rarely occuring image atrifacts we clamp the base.
+
+	// Note: Clamping the exponent seemed to fix the issue mentioned TTP but we decided to fix the root and accept the
+	// minor performance cost.
+
+	return ClampedPow(X, Y);
+}
+
+#if FEATURE_LEVEL < FEATURE_LEVEL_ES3_1 && !COMPILER_METAL
+	// DX11 (feature levels >= 10) feature sets natively supports uints in shaders; we just use floats on other platforms.
+	#define uint4	int4
+#endif
+
+// Optional VertexID - used by tessellation to uniquely identify control points.
+#if USING_TESSELLATION && DISPLACEMENT_ANTICRACK
+	#define OPTIONAL_VertexID			uint VertexID : SV_VertexID,
+	#define OPTIONAL_VertexID_PARAM		VertexID,
+	#define OPTIONAL_VertexID_VS_To_DS	uint VertexID : VS_To_DS_VertexID;
+	#define OutputVertexID( Out ) Out.VertexID = VertexID
+#else // #if USING_TESSELLATION && DISPLACEMENT_ANTICRACK
+	#define OPTIONAL_VertexID
+	#define OPTIONAL_VertexID_PARAM
+	#define OPTIONAL_VertexID_VS_To_DS
+	#define OutputVertexID( Out )
+#endif // #if USING_TESSELLATION && DISPLACEMENT_ANTICRACK
+
+// Helper macro used to interpolate the given member
+#define TESSELLATION_INTERPOLATE_MEMBER(member) O.member = a.member * aInterp + b.member * bInterp
+
+
+#if FEATURE_LEVEL >= FEATURE_LEVEL_SM4
+	/** 
+		* Number of MSAA samples supported by deferred passes in D3D11. 
+		* This is hardcoded because it allows deferred passes to optimize for the given value (for example, unrolling a loop).
+		*/
+	#define NumMSAASamples 4
+#endif
+
+// depth in the red channel in DeviceZ
+Texture2D		SceneDepthTexture;
+SamplerState	SceneDepthTextureSampler;
+Texture2D		CustomDepthTexture;
+SamplerState	CustomDepthTextureSampler;
+// Custom Stencil texture used for mobile platforms
+Texture2D		MobileCustomStencilTexture;
+SamplerState	MobileCustomStencilTextureSampler;
+// scene HDR color
+Texture2D		SceneColorTexture;
+SamplerState	SceneColorTextureSampler;
+// copy of scene alpha for PC ES2 emulation
+Texture2D		SceneAlphaCopyTexture;
+SamplerState	SceneAlphaCopyTextureSampler;
+// shadow and light function
+Texture2D		LightAttenuationTexture;
+SamplerState	LightAttenuationTextureSampler;
+
+#if FEATURE_LEVEL >= FEATURE_LEVEL_SM4
+// Scene depth-stencil buffer used for stencil reads
+Texture2D<uint2>	SceneStencilTexture;
+#endif
+
+// We don't use an inline function so we can avoid type promotion/ coercion.
+#define RETURN_COLOR( Color ) ( Color )
+
+// Tangent space bias
+// We don't use a function so we can avoid type promotion/ coercion.
+#define TangentBias(X) (X * 2.0f - 1.0f)
+
+float Square( float x )
+{
+	return x*x;
+}
+
+float2 Square( float2 x )
+{
+	return x*x;
+}
+
+float3 Square( float3 x )
+{
+	return x*x;
+}
+
+float4 Square( float4 x )
+{
+	return x*x;
+}
+
+float Pow2( float x )
+{
+	return x*x;
+}
+
+float2 Pow2( float2 x )
+{
+	return x*x;
+}
+
+float3 Pow2( float3 x )
+{
+	return x*x;
+}
+
+float4 Pow2( float4 x )
+{
+	return x*x;
+}
+
+float Pow3( float x )
+{
+	return x*x*x;
+}
+
+float2 Pow3( float2 x )
+{
+	return x*x*x;
+}
+
+float3 Pow3( float3 x )
+{
+	return x*x*x;
+}
+
+float4 Pow3( float4 x )
+{
+	return x*x*x;
+}
+
+float Pow4( float x )
+{
+	float xx = x*x;
+	return xx * xx;
+}
+
+float2 Pow4( float2 x )
+{
+	float2 xx = x*x;
+	return xx * xx;
+}
+
+float3 Pow4( float3 x )
+{
+	float3 xx = x*x;
+	return xx * xx;
+}
+
+float4 Pow4( float4 x )
+{
+	float4 xx = x*x;
+	return xx * xx;
+}
+
+float Pow5( float x )
+{
+	float xx = x*x;
+	return xx * xx * x;
+}
+
+float2 Pow5( float2 x )
+{
+	float2 xx = x*x;
+	return xx * xx * x;
+}
+
+float3 Pow5( float3 x )
+{
+	float3 xx = x*x;
+	return xx * xx * x;
+}
+
+float4 Pow5( float4 x )
+{
+	float4 xx = x*x;
+	return xx * xx * x;
+}
+
+float Pow6( float x )
+{
+	float xx = x*x;
+	return xx * xx * xx;
+}
+
+float2 Pow6( float2 x )
+{
+	float2 xx = x*x;
+	return xx * xx * xx;
+}
+
+float3 Pow6( float3 x )
+{
+	float3 xx = x*x;
+	return xx * xx * xx;
+}
+
+float4 Pow6( float4 x )
+{
+	float4 xx = x*x;
+	return xx * xx * xx;
+}
+
+// Only valid for x >= 0
+MaterialFloat AtanFast( MaterialFloat x )
+{
+	// Minimax 3 approximation
+	MaterialFloat3 A = x < 1 ? MaterialFloat3( x, 0, 1 ) : MaterialFloat3( 1/x, 0.5 * PI, -1 );
+	return A.y + A.z * ( ( ( -0.130234 * A.x - 0.0954105 ) * A.x + 1.00712 ) * A.x - 0.00001203333 );
+}
+
+/**
+ * Returns the upper 3x3 portion of the LocalToWorld matrix.
+ */
+MaterialFloat3x3 GetLocalToWorld3x3()
+{
+	return (MaterialFloat3x3)Primitive.LocalToWorld;
+}
+
+/** Converts a linear input value into a value to be stored in the light attenuation buffer. */
+MaterialFloat EncodeLightAttenuation(MaterialFloat InColor)
+{
+	// Apply a 1/2 power to the input, which allocates more bits for the darks and prevents banding
+	// Similar to storing colors in gamma space, except this uses less instructions than a pow(x, 1/2.2)
+	return sqrt(InColor);
+}
+
+/** Converts a linear input value into a value to be stored in the light attenuation buffer. */
+MaterialFloat4 EncodeLightAttenuation(MaterialFloat4 InColor)
+{
+	return sqrt(InColor);
+}
+
+/** return the scene lighting texture */
+MaterialFloat3 CalcSceneColor(MaterialFloat2 ScreenUV)
+{
+#if SCENE_TEXTURES_DISABLED
+	return MaterialFloat3(0.0f,0.0f,0.0f);
+#else
+	return Texture2DSampleLevel(SceneColorTexture, SceneColorTextureSampler, ScreenUV, 0).rgb;
+#endif
+}
+
+/** return all channels of the scene lighting texture */
+MaterialFloat4 CalcFullSceneColor(MaterialFloat2 ScreenUV)
+{
+#if SCENE_TEXTURES_DISABLED
+	return MaterialFloat4(0.0f, 0.0f, 0.0f, 0.0f);
+#else
+	return Texture2DSample(SceneColorTexture, SceneColorTextureSampler,ScreenUV);
+#endif
+}
+
+/** Encodes HDR linear scene color for storage in the 8 bit light attenuation texture. */
+MaterialFloat3 EncodeSceneColorForMaterialNode(MaterialFloat3 LinearSceneColor)
+{
+	// Preserving a range from [0, 10]
+	// Remap values to get more bits of precision in the darks
+	return pow(LinearSceneColor * .1f, .25f);
+}
+
+
+//
+// MOBILE WITHOUT FP16 SUPPORT: 8-BIT/CHANNEL MOSAICING FOR LINEAR HDR
+//
+// This provides linear blending and a {0 to 2} dynamic range.
+// This works by splitting the image into a checkerboard of dark and light pixels.
+//   DLDL
+//   LDLD
+//   DLDL
+//   LDLD
+//
+// There is an extra scan line dither pattern to increase precision by 1-bit.
+// The HdrMosaic() function computes the exposure level per pixel for the forward render pass.
+// The HdrDemosaic() function reconstructs the proper image in the tonemapping pass.
+//
+
+// These have been carefully tuned and should probably not be adjusted.
+// Dark pixel range is {0 to 1/DRK_MUL}.
+#define DRK_MUL 6.0
+// Support {0 to 2} in dynamic range.
+#define HDR_MUL (1.0/2.0)
+
+#define MOSAIC_ADD   HDR_MUL
+// The blend factor causes a smooth blend between the light and dark pixels.
+#define MOSAIC_BLEND ((255.0-(DRK_MUL/HDR_MUL))/255.0)
+#define MOSAIC_MUL   ((DRK_MUL * MOSAIC_BLEND - MOSAIC_ADD) * 2.0)
+
+// Done during forward shading pass before blending.
+MaterialFloat3 HdrMosaic(MaterialFloat3 LinearColor, float2 VPos) 
+{
+	float2 V;
+	V = VPos.xy * 0.5;
+	V.y += V.x;
+	V = frac(V);
+	MaterialFloat2 C;
+	C = (V * MaterialFloat2(2.0 * (-0.5/255.0), MOSAIC_MUL)) + MaterialFloat2(-0.5/255.0, MOSAIC_ADD);
+	return (LinearColor * C.y) + C.x;
+}
+
+#define DEMOSAIC_MUL (((1.0/DRK_MUL) - (1.0/HDR_MUL)) * 2.0)
+#define DEMOSAIC_ADD (1.0/HDR_MUL)
+
+// Resolve pass to remove mosaic and restore color.
+MaterialFloat3 HdrDemosaic(MaterialFloat3 Pixel, MaterialFloat3 OtherPixel, float2 VPos)
+{
+	MaterialFloat A = frac(dot(VPos + View.DemosaicVposOffset, float2(0.5, 0.5)));
+	MaterialFloat B = 0.5 - A;
+	A = A * DEMOSAIC_MUL + DEMOSAIC_ADD;
+	B = B * DEMOSAIC_MUL + DEMOSAIC_ADD;
+
+	// On ES2 devices we demosaic during the tonemapping pass which renders upside down, account for that here.
+	#if COMPILER_GLSL_ES2 || COMPILER_GLSL_ES3_1
+		return max((Pixel * B), (OtherPixel * A)); 
+	#else
+		return max((Pixel * A), (OtherPixel * B)); 
+	#endif
+}
+
+
+// Like RGBM but this can be interpolated.
+MaterialFloat4 RGBTEncode(MaterialFloat3 Color)
+{
+	MaterialFloat4 RGBT;
+	MaterialFloat Max = max(max(Color.r, Color.g), max(Color.b, 1e-6));
+	MaterialFloat RcpMax = rcp(Max);
+	RGBT.rgb = Color.rgb * RcpMax;
+	RGBT.a = Max * rcp(1.0 + Max);
+	return RGBT;
+}
+
+MaterialFloat3 RGBTDecode(MaterialFloat4 RGBT)
+{
+	RGBT.a = RGBT.a * rcp(1.0 - RGBT.a);
+	return RGBT.rgb * RGBT.a;
+}
+
+
+
+MaterialFloat4 RGBMEncode( MaterialFloat3 Color )
+{
+	Color *= 1.0 / 64.0;
+	
+	float4 rgbm;
+	rgbm.a = saturate( max( max( Color.r, Color.g ), max( Color.b, 1e-6 ) ) );
+	rgbm.a = ceil( rgbm.a * 255.0 ) / 255.0;
+	rgbm.rgb = Color / rgbm.a;
+	return rgbm;
+}
+
+MaterialFloat4 RGBMEncodeFast( MaterialFloat3 Color )
+{
+	// 0/0 result written to fixed point buffer goes to zero
+	MaterialFloat4 rgbm;
+	rgbm.a = dot( Color, 255.0 / 64.0 );
+	rgbm.a = ceil( rgbm.a );
+	rgbm.rgb = Color / rgbm.a;
+	rgbm *= MaterialFloat4( 255.0 / 64.0, 255.0 / 64.0, 255.0 / 64.0, 1.0 / 255.0 );
+	return rgbm;
+}
+
+MaterialFloat3 RGBMDecode( MaterialFloat4 rgbm, MaterialFloat MaxValue )
+{
+	return rgbm.rgb * (rgbm.a * MaxValue);
+}
+
+MaterialFloat3 RGBMDecode( MaterialFloat4 rgbm )
+{
+	return rgbm.rgb * (rgbm.a * 64.0f);
+}
+
+MaterialFloat4 RGBTEncode8BPC(MaterialFloat3 Color, MaterialFloat Range)
+{
+	MaterialFloat Max = max(max(Color.r, Color.g), max(Color.b, 1e-6));
+	Max = min(Max, Range);
+
+	MaterialFloat4 RGBT;
+	RGBT.a = (Range + 1) / Range *  Max / (1 + Max);
+
+	// quantise alpha to 8 bit.
+	RGBT.a = ceil(RGBT.a*255.0) / 255.0;
+	Max = RGBT.a / (1 + 1 / Range - RGBT.a);
+
+	MaterialFloat RcpMax = rcp(Max);
+	RGBT.rgb = Color.rgb * RcpMax;
+	return RGBT;
+}
+
+MaterialFloat3 RGBTDecode8BPC(MaterialFloat4 RGBT, MaterialFloat Range)
+{
+	RGBT.a = RGBT.a / (1 + 1 / Range - RGBT.a);
+	return RGBT.rgb * RGBT.a;
+}
+
+#define HDR_ENCODE_NONE 0.0		// 64bpp HDR
+#define HDR_ENCODE_MOSAIC 1.0	// 32bpp HDR using Mosaic encoding
+#define HDR_ENCODE_RGBE 2.0		// 32bpp HDR using RGBE encoding
+#define HDR_ENCODE_RGBA8 3.0	// 32bpp HDR mode without any actual encoding.
+
+half GetHDR32bppEncodeMode()
+{
+#if ES2_PROFILE
+	#if COMPILER_GLSL_ES2 && !IOS && !WEBGL // ANDROID
+		return intrinsic_GetHDR32bppEncodeModeES2();
+	#endif
+	#if MOBILE_EMULATION
+		// To enable editor runtime change without recompile, PC always eats the encode cost.
+		return View.HDR32bppEncodingMode;
+	#endif
+#endif
+	return HDR_ENCODE_NONE;
+}
+
+#define DEFAULT_32BPPHDR_ENCODED_RANGE 1024.0
+
+MaterialFloat4 Encode32BPPHDR(MaterialFloat4 Color, float2 SvPosition)
+{
+	half Mode = GetHDR32bppEncodeMode();
+	if (Mode == HDR_ENCODE_MOSAIC)
+	{ 
+		return MaterialFloat4(HdrMosaic(Color.rgb, SvPosition), Color.a);
+	}
+	else if (Mode == HDR_ENCODE_RGBE)
+	{ 
+		return RGBTEncode8BPC(Color.rgb, DEFAULT_32BPPHDR_ENCODED_RANGE);
+	}
+	else
+	{
+		// Mode == HDR_ENCODE_NONE || Mode == HDR_ENCODE_RGBA8
+		return Color;
+	}
+}
+
+MaterialFloat4 Decode32BPPHDR(MaterialFloat4 Encoded, MaterialFloat3 OtherEncoded = MaterialFloat3(0, 0, 0), float2 SvPosition = float2(0, 0))
+{
+	half Mode = GetHDR32bppEncodeMode();
+	if (Mode == HDR_ENCODE_MOSAIC)
+	{
+		return MaterialFloat4(HdrDemosaic(Encoded.rgb, OtherEncoded, SvPosition), 0.0f);
+	}
+	if (Mode == HDR_ENCODE_RGBE)
+	{ 
+		return MaterialFloat4(RGBTDecode8BPC(Encoded, DEFAULT_32BPPHDR_ENCODED_RANGE), 0.0f);
+	}
+	else
+	{
+		// Mode == HDR_ENCODE_NONE || Mode == HDR_ENCODE_RGBA8
+		return Encoded; 
+	}
+}
+
+// Demosaic capable 32bpp hdr decode.
+float4 Decode32BPPHDR(float4 Pixel, float2 SvPosition, Texture2D SourceTexture, SamplerState SourceSampler, float2 InTexCoords[4])
+{
+	half3 PixelColorN = SourceTexture.Sample(SourceSampler, InTexCoords[0].xy).rgb;
+	half3 PixelColorE = SourceTexture.Sample(SourceSampler, InTexCoords[1].xy).rgb;
+	half3 PixelColorW = SourceTexture.Sample(SourceSampler, InTexCoords[2].xy).rgb;
+	half3 PixelColorS = SourceTexture.Sample(SourceSampler, InTexCoords[3].xy).rgb;
+	half3 PixelColorV = PixelColorN * 0.5 + PixelColorS * 0.5;
+	half3 PixelColorH = PixelColorW * 0.5 + PixelColorE * 0.5;
+	if(abs(PixelColorN.g - PixelColorS.g) < abs(PixelColorW.g - PixelColorE.g)) 
+	{
+		PixelColorH = PixelColorV;
+	}
+	Pixel.rgb = Decode32BPPHDR(Pixel, PixelColorH, SvPosition.xy).rgb;
+	Pixel.a = 0.0;
+	return Pixel;
+}
+
+/** Get render target write mask value
+  * This gets a bit from a write mask texture created with FRTWriteMaskDecodeCS. Only supprted on some platforms.
+  */
+#if PLATFORM_SUPPORTS_RENDERTARGET_WRITE_MASK
+uint DecodeRTWriteMaskTexture(in float2 ScreenPosition, in Texture2D<uint> RTWriteMaskTexture)
+{
+	int2 IntPosition = int2(ScreenPosition.xy);
+	uint RTWriteMaskValue = RTWriteMaskTexture.Load( int3(IntPosition.x/8, IntPosition.y/8, 0) );
+
+	int2 BitCoord = ((IntPosition / int2(4, 4)) % int2(2, 2));
+	uint BitIdx = BitCoord.x + (BitCoord.y*2);
+	uint RTWriteMaskBit = RTWriteMaskValue & (1<<BitIdx);
+
+	return RTWriteMaskBit;
+}
+#endif
+
+/** Calculates the ScreenUV given the screen position and an offset fraction. */
+float2 CalcScreenUVFromOffsetFraction(float4 ScreenPosition, float2 OffsetFraction)
+{
+	float2 NDC = ScreenPosition.xy / ScreenPosition.w;
+	// Apply the offset in NDC space so that it is consistent regardless of scene color buffer size
+	// Clamp to valid area of the screen to avoid reading garbage
+	//@todo - soft clamp
+	float2 OffsetNDC = clamp(NDC + OffsetFraction * float2(2, -2), -.999f, .999f);
+	return float2(OffsetNDC * View.ScreenPositionScaleBias.xy + View.ScreenPositionScaleBias.wz);
+}
+
+Texture2D		SceneColorCopyTexture;
+SamplerState	SceneColorCopyTextureSampler;
+
+/** Applies an offset to the scene texture lookup and decodes the HDR linear space color. */
+float3 DecodeSceneColorForMaterialNode(float2 ScreenUV)
+{
+#if HIT_PROXY_SHADER || SCENE_TEXTURES_DISABLED
+	// Hit proxies rendering pass doesn't have access to valid render buffers
+	return float3(0.0f, 0.0f, 0.0f);
+#else
+	float4 EncodedSceneColor = Texture2DSample(SceneColorCopyTexture, SceneColorCopyTextureSampler, ScreenUV);
+
+	// Undo the function in EncodeSceneColorForMaterialNode
+	return pow(EncodedSceneColor.rgb, 4) * 10;
+#endif
+}
+
+float4 GetPerPixelLightAttenuation(float2 UV)
+{
+	return Square(Texture2DSampleLevel(LightAttenuationTexture, LightAttenuationTextureSampler, UV, 0));
+}
+
+// also see ConvertToDeviceZ()
+// @param DeviceZ value that is stored in the depth buffer (Z/W)
+// @return SceneDepth (linear in world units, W)
+float ConvertFromDeviceZ(float DeviceZ)
+{
+	// Supports ortho and perspective, see CreateInvDeviceZToWorldZTransform()
+	return DeviceZ * View.InvDeviceZToWorldZTransform[0] + View.InvDeviceZToWorldZTransform[1] + 1.0f / (DeviceZ * View.InvDeviceZToWorldZTransform[2] - View.InvDeviceZToWorldZTransform[3]);
+}
+
+// inverse operation of ConvertFromDeviceZ()
+// @param SceneDepth (linear in world units, W)
+// @return DeviceZ (Z/W)
+float ConvertToDeviceZ(float SceneDepth)
+{
+	FLATTEN
+	if (View.ViewToClip[3][3] < 1.0f)
+	{
+		// Perspective
+		return 1.0f / ((SceneDepth + View.InvDeviceZToWorldZTransform[3]) * View.InvDeviceZToWorldZTransform[2]);
+	}
+	else
+	{
+		// Ortho
+		return SceneDepth * View.ViewToClip[2][2] + View.ViewToClip[3][2];
+	}
+}
+
+/** Returns clip space W, which is world space distance along the View Z axis. Note if you need DeviceZ LookupDeviceZ() is the faster option */
+float CalcSceneDepth(float2 ScreenUV)
+{
+#if SCENE_TEXTURES_DISABLED
+	return 0.0f;
+#else
+	#if FEATURE_LEVEL > FEATURE_LEVEL_ES3_1 || MOBILE_FORCE_DEPTH_TEXTURE_READS
+		return ConvertFromDeviceZ(Texture2DSampleLevel(SceneDepthTexture, SceneDepthTextureSampler, ScreenUV, 0).r);
+	#else
+		#if COMPILER_GLSL_ES2
+			#if IOS
+				// Only call FramebufferFetch when actually compiling for IOS ES2.
+				return FramebufferFetchES2().w;
+			#elif WEBGL
+				return Texture2DSampleLevel(SceneAlphaCopyTexture, SceneAlphaCopyTextureSampler, ScreenUV, 0).r;
+			#else 
+				float SceneW = ConvertFromDeviceZ(Texture2DSampleLevel(SceneDepthTexture, SceneDepthTextureSampler, ScreenUV, 0).r);		
+				return DepthbufferFetchES2(SceneW, View.InvDeviceZToWorldZTransform[2], View.InvDeviceZToWorldZTransform[3]);
+			#endif 
+		#elif METAL_PROFILE && !MAC
+				return FramebufferFetchES2().w;
+		#else
+			return ConvertFromDeviceZ(Texture2DSampleLevel(SceneDepthTexture, SceneDepthTextureSampler, ScreenUV, 0).r);
+		#endif
+	#endif
+#endif
+}
+
+/** Returns DeviceZ which is the z value stored in the depth buffer. */
+float LookupDeviceZ( float2 ScreenUV )
+{
+#if FEATURE_LEVEL > FEATURE_LEVEL_ES3_1
+	// native Depth buffer lookup
+	return Texture2DSampleLevel(SceneDepthTexture, SceneDepthTextureSampler, ScreenUV, 0).r;
+#else
+	#if COMPILER_GLSL_ES2
+		// todo: can be optimized
+		return ConvertToDeviceZ(CalcSceneDepth(ScreenUV));
+	#elif METAL_PROFILE
+		// todo: can be optimized
+		return ConvertToDeviceZ(CalcSceneDepth(ScreenUV));
+	#else
+		// native Depth buffer lookup
+		return Texture2DSampleLevel(SceneDepthTexture, SceneDepthTextureSampler, ScreenUV, 0).r;
+	#endif
+#endif
+}
+
+#if FEATURE_LEVEL >= FEATURE_LEVEL_SM4 || MOBILE_EMULATION
+
+	// depth in in DeviceZ
+	Texture2D<float> SceneDepthTextureNonMS;
+
+	/** Returns clip space W, which is world space distance along the View Z axis. */
+	float CalcSceneDepth(uint2 PixelPos)
+	{
+#if SCENE_TEXTURES_DISABLED
+		return 0.0f;
+#else
+		float DeviceZ = SceneDepthTextureNonMS.Load(int3(PixelPos, 0));
+
+		// Fetch the depth buffer Z / W value, solve for W
+		return ConvertFromDeviceZ(DeviceZ);
+#endif
+	}
+#endif
+
+/**
+* Returns scene color in rgb, depth in a
+*/
+float4 CalcSceneColorAndDepth( float2 ScreenUV )
+{
+	return float4(CalcSceneColor(ScreenUV), CalcSceneDepth(ScreenUV));
+}
+
+#if FEATURE_LEVEL >= FEATURE_LEVEL_SM4
+#if FEATURE_LEVEL == FEATURE_LEVEL_SM4
+	// SM4 requires texture size to be explicitly stated and expressed as the number of samples.
+	Texture2DMS<float4, NumMSAASamples> SceneColorSurface;
+	Texture2DMS<float, NumMSAASamples> SceneDepthSurface;
+#else
+	Texture2DMS<float4> SceneColorSurface;
+	Texture2DMS<float> SceneDepthSurface;
+#endif // FEATURE_LEVEL
+	float CalcSceneDepthMSAA(float2 ScreenUV,uint SampleIndex)
+	{
+		int2 IntUV = int2(trunc(ScreenUV * View.BufferSizeAndInvSize.xy));
+		float DeviceZ = SceneDepthSurface.Load(IntUV,SampleIndex);
+
+		return ConvertFromDeviceZ(DeviceZ);
+	}
+#endif
+	 
+// ----------------------------
+
+float2 ScreenPositionToBufferUV(float4 ScreenPosition)
+{
+	return float2(ScreenPosition.xy / ScreenPosition.w * View.ScreenPositionScaleBias.xy + View.ScreenPositionScaleBias.wz);
+}
+
+float2 SvPositionToBufferUV(float4 SvPosition)
+{
+	return SvPosition.xy * View.BufferSizeAndInvSize.zw;
+}
+
+// Used for post process shaders which don't need to resolve the view	
+float3 SvPositionToTranslatedWorld(float4 SvPosition)
+{
+	float4 HomWorldPos = mul(float4(SvPosition.xyz, 1), View.SVPositionToTranslatedWorld);
+
+	return HomWorldPos.xyz / HomWorldPos.w;
+}
+
+// Used for vertex factory shaders which need to use the resolved view
+float3 SvPositionToResolvedTranslatedWorld(float4 SvPosition)
+{
+	float4 HomWorldPos = mul(float4(SvPosition.xyz, 1), ResolvedView.SVPositionToTranslatedWorld);
+
+	return HomWorldPos.xyz / HomWorldPos.w;
+}
+
+// prefer to use SvPositionToTranslatedWorld() for better quality
+float3 SvPositionToWorld(float4 SvPosition)
+{
+	return SvPositionToTranslatedWorld(SvPosition) - View.PreViewTranslation;
+}
+
+// investigate: doesn't work for usage with View.ScreenToWorld, see SvPositionToScreenPosition2()
+float4 SvPositionToScreenPosition(float4 SvPosition)
+{
+	// todo: is already in .w or needs to be reconstructed like this:
+//	SvPosition.w = ConvertFromDeviceZ(SvPosition.z);
+
+	float2 PixelPos = SvPosition.xy - View.ViewRectMin.xy;	
+
+	// NDC (NormalizedDeviceCoordinates, after the perspective divide)
+	float3 NDCPos = float3( (PixelPos * View.ViewSizeAndInvSize.zw - 0.5f) * float2(2, -2), SvPosition.z);
+
+	// SvPosition.w: so .w has the SceneDepth, some mobile code and the DepthFade material expression wants that
+	return float4(NDCPos.xyz, 1) * SvPosition.w;
+}
+
+// Used for vertex factory shaders which need to use the resolved view
+float4 SvPositionToResolvedScreenPosition(float4 SvPosition)
+{
+	float2 PixelPos = SvPosition.xy - ResolvedView.ViewRectMin.xy;	
+
+	// NDC (NormalizedDeviceCoordinates, after the perspective divide)
+	float3 NDCPos = float3( (PixelPos * ResolvedView.ViewSizeAndInvSize.zw - 0.5f) * float2(2, -2), SvPosition.z);
+
+	// SvPosition.w: so .w has the SceneDepth, some mobile code and the DepthFade material expression wants that
+	return float4(NDCPos.xyz, 1) * SvPosition.w;
+}
+
+float2 SvPositionToViewportUV(float4 SvPosition)
+{
+	// can be optimized from 2SUB+2MUL to 2MAD
+	float2 PixelPos = SvPosition.xy - View.ViewRectMin.xy;	
+
+	return PixelPos.xy * View.ViewSizeAndInvSize.zw;
+}
+
+float2 BufferUVToViewportUV(float2 BufferUV)
+{
+	float2 PixelPos = BufferUV.xy * View.BufferSizeAndInvSize.xy - View.ViewRectMin.xy;
+	return PixelPos.xy * View.ViewSizeAndInvSize.zw;
+}
+// ----------------------------
+
+/** 
+ * aligns the clip space position so that it can be used as a texture coordinate
+ * to properly align in screen space
+ */
+MaterialFloat2 ScreenAlignedPosition( float4 ScreenPosition )
+{
+	return MaterialFloat2(ScreenPositionToBufferUV(ScreenPosition));
+}
+
+/** 
+ * Aligns the [0,1] UV to match the view within the backbuffer
+ */
+MaterialFloat2 ScreenAlignedUV( MaterialFloat2 UV )
+{
+	return (UV*MaterialFloat2(2,-2) + MaterialFloat2(-1,1))*View.ScreenPositionScaleBias.xy + View.ScreenPositionScaleBias.wz;
+}
+
+/**
+ * Compute viewport coordinates from the given fragment coordinates.
+ */
+MaterialFloat2 GetViewportCoordinates(MaterialFloat2 InFragmentCoordinates)
+{
+	return InFragmentCoordinates;
+}
+
+/**
+ * Unpack a normal stored in a normal map. The X and Y components are rescaled from [0,1] to [-1,1] and Z is reconstructed.
+ */
+MaterialFloat4 UnpackNormalMap( MaterialFloat4 TextureSample )
+{
+#if COMPILER_GLSL_ES2 && IOS
+	return MaterialFloat4(TextureSample.rgb * 2 - 1, 1);
+#else
+	#if DXT5_NORMALMAPS
+		MaterialFloat2 NormalXY = TextureSample.ag;
+	#else
+		MaterialFloat2 NormalXY = TextureSample.rg;
+	#endif
+
+	NormalXY = NormalXY * MaterialFloat2(2.0f,2.0f) - MaterialFloat2(1.0f,1.0f);
+	MaterialFloat NormalZ = sqrt( saturate( 1.0f - dot( NormalXY, NormalXY ) ) );
+	return MaterialFloat4( NormalXY.xy, NormalZ, 1.0f );
+#endif
+}
+
+// Antialiased version of a binary comparison between ThresholdConst and a texture channel.
+float AntialiasedTextureMask( Texture2D Tex, SamplerState Sampler, float2 UV, float ThresholdConst, int Channel )
+{
+	// By setting MaskConst to 0001, 0010, 0100 or 1000 individual channels can be chosen (the compiler should be able to optimize that).
+	MaterialFloat4 MaskConst = MaterialFloat4(Channel == 0, Channel == 1, Channel == 2, Channel == 3);
+
+	// border width in pixels, for antialiasing 1 .. 1.5 is good but 1.0 is good for optimizations 
+	const float WidthConst = 1.0f;			
+	float InvWidthConst = 1 / WidthConst;
+
+	// Problem:
+
+	// A simple texture lookup with a comparison against some thresold value allows to get a mask useful
+	// for many purposes (e.g. text rendering, signs, oil/water/paint). Antialiased masks look much better
+	// and mip mapping provides that but only for minification. So when the texture resolution is lower than
+	// the rendering size results get blurry.
+
+	// Idea:
+
+	// We compute the distance to the threshold line in pixels (with subpixel precision). We can visualize
+	// the problem as a heightmap that intersects a axis aligned plane at the threshold height. Only surface
+	// above the threshold plane contributes to the mask. Looking at one pixel the heightmap can be approximated
+	// by a plane. We can easily get the plane center value form a texture lookup and get the plane equation from
+	// ddx and ddy of that value (only one value per 2x2 block) or some other more precise method. We can reduce the
+	// 3d problem to 2d (looking at the steepest angle only) and the resulting value tells us how much the texture value
+	// changes for one pixel. This allows us to scale and bias (threshold) the texture value the so it maps to the
+	// distance function. We rescaling the distance to 0.5 coverage at the line, >1 MaterialFloat a pixel inside and <0 MaterialFloat
+	// a pixel outside. Clamping this value in the range from 0..1 gives us a good approximation of the pixel coverage.
+
+	// We tried multiple possible implementations - this is the cheapest and looks ok is most cases.
+	// If quality improvements are needed we can add an option to the node later on.
+	float Result;
+	{
+		// optimized, ddx/ddy only for every 2x2 block (bad for distant stuff)
+		float Sample1 = dot(MaskConst, Texture2DSample(Tex, Sampler, UV));
+
+		// compute the derivatives of the texture content
+		float2 TexDD = float2(DDX(Sample1), DDY(Sample1));
+
+		float TexDDLength = max(abs(TexDD.x), abs(TexDD.y)); 
+		float Top = InvWidthConst * (Sample1 - ThresholdConst);
+		Result = Top / TexDDLength + ThresholdConst;
+	}
+
+	Result = saturate(Result);	// no always needed (e.g. DX9 framebuffer blending)
+
+	return Result;
+}
+
+float Noise3D_Multiplexer(int Function, float3 Position, int Quality, bool bTiling, uint RepeatSize)
+{
+	// verified, HLSL compiled out the switch if Function is a constant
+	switch(Function)
+	{
+		case 0:
+			return SimplexNoise3D_TEX(Position);
+		case 1:
+			return GradientNoise3D_TEX(Position, bTiling, RepeatSize);
+		case 2:
+			return FastGradientPerlinNoise3D_TEX(Position);
+		case 3:
+			return GradientNoise3D_ALU(Position, bTiling, RepeatSize);
+		case 4:
+			return ValueNoise3D_ALU(Position, bTiling, RepeatSize);
+		default:
+			return VoronoiNoise3D_ALU(Position, Quality, bTiling, RepeatSize, true).w * 2. - 1.;
+	}
+	return 0;
+}
+
+// @param LevelScale usually 2 but higher values allow efficient use of few levels
+// @return in user defined range (OutputMin..OutputMax)
+MaterialFloat MaterialExpressionNoise(float3 Position, float Scale, int Quality, int Function, bool bTurbulence, uint Levels, float OutputMin, float OutputMax, float LevelScale, float FilterWidth, bool bTiling, float RepeatSize)
+{
+	Position *= Scale;
+	FilterWidth *= Scale;
+
+	float Out = 0.0f;
+	float OutScale = 1.0f;
+	float InvLevelScale = 1.0f / LevelScale;
+	
+	LOOP for(uint i = 0; i < Levels; ++i)
+	{
+		// fade out noise level that are too high frequent (not done through dynamic branching as it usually requires gradient instructions)
+		OutScale *= saturate(1.0 - FilterWidth);
+
+		if(bTurbulence)
+		{
+			Out += abs(Noise3D_Multiplexer(Function, Position, Quality, bTiling, RepeatSize)) * OutScale;
+		}
+		else
+		{
+			Out += Noise3D_Multiplexer(Function, Position, Quality, bTiling, RepeatSize) * OutScale;
+		}
+
+		Position *= LevelScale;
+		RepeatSize *= LevelScale;
+		OutScale *= InvLevelScale;
+		FilterWidth *= LevelScale;
+	}
+
+	if(!bTurbulence)
+	{
+		// bring -1..1 to 0..1 range
+		Out = Out * 0.5f + 0.5f;
+	}
+
+	// Out is in 0..1 range
+	return lerp(OutputMin, OutputMax, Out);
+}
+
+
+// Material node for noise functions returning a vector value
+// @param LevelScale usually 2 but higher values allow efficient use of few levels
+// @return in user defined range (OutputMin..OutputMax)
+MaterialFloat4 MaterialExpressionVectorNoise(MaterialFloat3 Position, int Quality, int Function, bool bTiling, float TileSize)
+{
+	float4 result = float4(0,0,0,1);
+	float3x4 Jacobian = JacobianSimplex_ALU(Position, bTiling, TileSize);	// compiled out if not used
+
+	// verified, HLSL compiled out the switch if Function is a constant
+	switch (Function)
+	{
+	case 0:	// Cellnoise
+		result.xyz = float3(Rand3DPCG16(int3(floor(NoiseTileWrap(Position, bTiling, TileSize))))) / 0xffff;
+		break;
+	case 1: // Color noise
+		result.xyz = float3(Jacobian[0].w, Jacobian[1].w, Jacobian[2].w);
+		break;
+	case 2: // Gradient
+		result = Jacobian[0];
+		break;
+	case 3: // Curl
+		result.xyz = float3(Jacobian[2][1] - Jacobian[1][2], Jacobian[0][2] - Jacobian[2][0], Jacobian[1][0] - Jacobian[0][1]);
+		break;
+	default: // Voronoi
+		result = VoronoiNoise3D_ALU(Position, Quality, bTiling, TileSize, false);
+		break;
+	}
+	return result;
+}
+
+
+/*
+* Clips a ray to an AABB.  Does not handle rays parallel to any of the planes.
+*
+* @param RayOrigin - The origin of the ray in world space.
+* @param RayEnd - The end of the ray in world space.  
+* @param BoxMin - The minimum extrema of the box.
+* @param BoxMax - The maximum extrema of the box.
+* @return - Returns the closest intersection along the ray in x, and furthest in y.  
+*			If the ray did not intersect the box, then the furthest intersection <= the closest intersection.
+*			The intersections will always be in the range [0,1], which corresponds to [RayOrigin, RayEnd] in worldspace.
+*			To find the world space position of either intersection, simply plug it back into the ray equation:
+*			WorldPos = RayOrigin + (RayEnd - RayOrigin) * Intersection;
+*/
+float2 LineBoxIntersect(float3 RayOrigin, float3 RayEnd, float3 BoxMin, float3 BoxMax)
+{
+	float3 InvRayDir = 1.0f / (RayEnd - RayOrigin);
+	
+	//find the ray intersection with each of the 3 planes defined by the minimum extrema.
+	float3 FirstPlaneIntersections = (BoxMin - RayOrigin) * InvRayDir;
+	//find the ray intersection with each of the 3 planes defined by the maximum extrema.
+	float3 SecondPlaneIntersections = (BoxMax - RayOrigin) * InvRayDir;
+	//get the closest of these intersections along the ray
+	float3 ClosestPlaneIntersections = min(FirstPlaneIntersections, SecondPlaneIntersections);
+	//get the furthest of these intersections along the ray
+	float3 FurthestPlaneIntersections = max(FirstPlaneIntersections, SecondPlaneIntersections);
+
+	float2 BoxIntersections;
+	//find the furthest near intersection
+	BoxIntersections.x = max(ClosestPlaneIntersections.x, max(ClosestPlaneIntersections.y, ClosestPlaneIntersections.z));
+	//find the closest far intersection
+	BoxIntersections.y = min(FurthestPlaneIntersections.x, min(FurthestPlaneIntersections.y, FurthestPlaneIntersections.z));
+	//clamp the intersections to be between RayOrigin and RayEnd on the ray
+	return saturate(BoxIntersections);
+}
+
+/** Computes distance from an AABB to a point in space. */
+MaterialFloat ComputeDistanceFromBoxToPoint(MaterialFloat3 Mins, MaterialFloat3 Maxs, MaterialFloat3 InPoint)
+{
+	MaterialFloat3 DistancesToMin = InPoint < Mins ? abs(InPoint - Mins) : 0;
+	MaterialFloat3 DistancesToMax = InPoint > Maxs ? abs(InPoint - Maxs) : 0;
+
+	//@todo - this is actually incorrect, it gives manhattan distance
+	MaterialFloat Distance = dot(DistancesToMin, 1);
+	Distance += dot(DistancesToMax, 1);
+	return Distance;
+}
+
+/** Computes squared distance from a point in space to an AABB. */
+MaterialFloat ComputeSquaredDistanceFromBoxToPoint(MaterialFloat3 BoxCenter, MaterialFloat3 BoxExtent, MaterialFloat3 InPoint)
+{
+	MaterialFloat3 AxisDistances = max(abs(InPoint - BoxCenter) - BoxExtent, 0);
+	return dot(AxisDistances, AxisDistances);
+}
+
+/** Computes distance from point inside an AABB to the AABB's surface. */
+float ComputeDistanceFromBoxToPointInside(float3 BoxCenter, float3 BoxExtent, float3 InPoint)
+{
+	float3 DistancesToMin = max(InPoint - BoxCenter + BoxExtent, 0);
+	float3 DistancesToMax = max(BoxCenter + BoxExtent - InPoint, 0);
+	float3 ClosestDistances = min(DistancesToMin, DistancesToMax);
+	return min(ClosestDistances.x, min(ClosestDistances.y, ClosestDistances.z));
+}
+
+bool RayHitSphere(float3 RayOrigin, float3 UnitRayDirection, float3 SphereCenter, float SphereRadius)
+{
+	float3 ClosestPointOnRay = max(0, dot(SphereCenter - RayOrigin, UnitRayDirection)) * UnitRayDirection;
+	float3 CenterToRay = RayOrigin + ClosestPointOnRay - SphereCenter;
+	return dot(CenterToRay, CenterToRay) <= Square(SphereRadius);
+}
+
+bool RaySegmentHitSphere(float3 RayOrigin, float3 UnitRayDirection, float RayLength, float3 SphereCenter, float SphereRadius)
+{
+	float DistanceAlongRay = dot(SphereCenter - RayOrigin, UnitRayDirection);
+	float3 ClosestPointOnRay = DistanceAlongRay * UnitRayDirection;
+	float3 CenterToRay = RayOrigin + ClosestPointOnRay - SphereCenter;
+	return dot(CenterToRay, CenterToRay) <= Square(SphereRadius) && DistanceAlongRay > -SphereRadius && DistanceAlongRay - SphereRadius < RayLength;
+}
+
+/** Transforms a vector from tangent space to world space */
+MaterialFloat3 TransformTangentVectorToWorld(MaterialFloat3x3 TangentToWorld, MaterialFloat3 InTangentVector)
+{
+	// Transform directly to world space
+	// The vector transform is optimized for this case, only one vector-matrix multiply is needed
+	return mul(InTangentVector, TangentToWorld);
+}
+
+/** Transforms a vector from world space to tangent space */
+MaterialFloat3 TransformWorldVectorToTangent(MaterialFloat3x3 TangentToWorld, MaterialFloat3 InWorldVector)
+{
+	// Transform from world to tangent space with the transpose of TangentToWorld (achieved by swapping vector / matrix multiply order)
+	// Note that the transpose is only equal to the inverse for orthonormal matrices - aka only uniform scaling
+	return mul(TangentToWorld, InWorldVector);
+}
+
+float3 TransformWorldVectorToView(float3 InTangentVector)
+{
+	// Transform from world to view space
+	return mul(InTangentVector, (float3x3)ResolvedView.TranslatedWorldToView);
+}
+
+/** Computes the distance from the center to the edge of an AABB with the given extents in the given direction. */
+MaterialFloat GetBoxPushout(MaterialFloat3 Normal,MaterialFloat3 Extent)
+{
+	return dot(abs(Normal * Extent), MaterialFloat3(1.0f, 1.0f, 1.0f));
+}
+
+/** Generates arbitrary but valid perpendicular unit vectors to ZAxis.  ZAxis should be unit length. */
+void GenerateCoordinateSystem(float3 ZAxis, out float3 XAxis, out float3 YAxis)
+{
+	if (abs(ZAxis.x) > abs(ZAxis.y))
+	{
+		float InverseLength = 1.0f / sqrt(dot(ZAxis.xz, ZAxis.xz));
+		XAxis = float3(-ZAxis.z * InverseLength, 0.0f, ZAxis.x * InverseLength);
+	}
+	else
+	{
+		float InverseLength = 1.0f / sqrt(dot(ZAxis.yz, ZAxis.yz));
+		XAxis = float3(0.0f, ZAxis.z * InverseLength, -ZAxis.y * InverseLength);
+	}
+
+	YAxis = cross(ZAxis, XAxis);
+}
+
+// Define passthrough implementations of EvaluateAttributeAtSample for non-D3D11 platforms.
+#if !SM5_PROFILE
+	float EvaluateAttributeAtSample(float Attribute,uint SampleIndex) { return Attribute; }
+	float2 EvaluateAttributeAtSample(float2 Attribute,uint SampleIndex) { return Attribute; }
+	float3 EvaluateAttributeAtSample(float3 Attribute,uint SampleIndex) { return Attribute; }
+	float4 EvaluateAttributeAtSample(float4 Attribute,uint SampleIndex) { return Attribute; }
+#endif
+
+/** Output of the screen vertex shader. */
+struct FScreenVertexOutput
+{
+#if METAL_PROFILE || COMPILER_GLSL_ES3_1
+	//@todo-rco: FIXME!
+	noperspective float2 UV : TEXCOORD0;
+#else
+	noperspective MaterialFloat2 UV : TEXCOORD0;
+#endif
+	float4 Position : SV_POSITION;
+};
+
+
+// for velocity rendering, motionblur and temporal AA
+// velocity needs to support -2..2 screen space range for x and y
+// texture is 16bit 0..1 range per channel
+float2 EncodeVelocityToTexture(float2 In)
+{
+	// 0.499f is a value smaller than 0.5f to avoid using the full range to use the clear color (0,0) as special value
+	// 0.5f to allow for a range of -2..2 instead of -1..1 for really fast motions for temporal AA
+	return In * (0.499f * 0.5f) + 32767.0f / 65535.0f;
+}
+// see EncodeVelocityToTexture()
+float2 DecodeVelocityFromTexture(float2 In)
+{
+	const float InvDiv = 1.0f / (0.499f * 0.5f);
+	// reference
+//	return (In - 32767.0f / 65535.0f ) / (0.499f * 0.5f);
+	// MAD layout to help compiler
+	return In * InvDiv - 32767.0f / 65535.0f * InvDiv;
+}
+
+// Used for the Global Illumination in the GIReplace material expression
+bool GetGIReplaceState()
+{
+#if REFLECTIVE_SHADOW_MAP
+	return true;
+#else
+	return false;
+#endif
+}
+
+
+#if FEATURE_LEVEL >= FEATURE_LEVEL_SM4
+struct FWriteToSliceGeometryOutput
+{
+	FScreenVertexOutput Vertex;
+	uint LayerIndex : SV_RenderTargetArrayIndex;
+};
+#endif
+
+// Helper macro to globally ignore requests for non-offset world positions in materials when lower than shader model 4. We do this
+// because we are using an extra interpolator for this second world position, and in < ES31 there may not be enough
+#define USE_WORLD_POSITION_EXCLUDING_SHADER_OFFSETS	(NEEDS_WORLD_POSITION_EXCLUDING_SHADER_OFFSETS && FEATURE_LEVEL >= FEATURE_LEVEL_ES3_1)
+
+/** Used for calculating vertex positions and UVs when drawing with DrawRectangle */
+void DrawRectangle( in float4 InPosition, in float2 InTexCoord, out float4 OutPosition, out float2 OutTexCoord)
+{
+	OutPosition = InPosition;
+	OutPosition.xy = -1.0f + 2.0f * (DrawRectangleParameters.PosScaleBias.zw + (InPosition.xy * DrawRectangleParameters.PosScaleBias.xy)) * DrawRectangleParameters.InvTargetSizeAndTextureSize.xy;
+	OutPosition.xy *= float2( 1, -1 );
+	OutTexCoord.xy = (DrawRectangleParameters.UVScaleBias.zw + (InTexCoord.xy * DrawRectangleParameters.UVScaleBias.xy)) * DrawRectangleParameters.InvTargetSizeAndTextureSize.zw;
+}
+
+/** Used for calculating vertex positions when drawing with DrawRectangle */
+void DrawRectangle( in float4 InPosition, out float4 OutPosition)
+{
+	OutPosition = InPosition;
+	OutPosition.xy = -1.0f + 2.0f * (DrawRectangleParameters.PosScaleBias.zw + (InPosition.xy * DrawRectangleParameters.PosScaleBias.xy)) * DrawRectangleParameters.InvTargetSizeAndTextureSize.xy;
+	OutPosition.xy *= float2( 1, -1 );
+}
+
+//Since some platforms don't remove Nans in saturate calls, 
+//SafeSaturate function will remove nan/inf.    
+//Can be expensive, only call when there's a good reason to expect Nans.
+//D3D saturate actually turns Nans -> 1  since it does the min(x, 1) first, and D3D nan rules specify the non-nand operand wins in such a case.  
+//See: http://msdn.microsoft.com/en-us/library/windows/desktop/jj218760(v=vs.85).aspx  
+#define SafeSaturate_Def(type)\
+type SafeSaturate(type In) \
+{\
+	return saturate(In);\
+}
+
+SafeSaturate_Def(float)
+SafeSaturate_Def(float2)
+SafeSaturate_Def(float3)
+SafeSaturate_Def(float4)
+
+// Experimental way to allow adjusting the OpacityMask for shadow map rendering of masked materials.
+// Can be accessed with a Custom material node. If this turns out to be very useful we can expose as MaterialFunction
+// and potentially expose other queries as well (e.g. SkeletalMesh, HitProxy, ).
+// @return 0:no, 1:yes
+float IsShadowDepthShader()
+{
+#ifdef SHADOW_DEPTH_SHADER
+	return 1;
+#else
+	return 0;
+#endif
+}
+#define TERRAIN_ZSCALE (1.0f/128.0f)
+
+// Decodes a value which was packed into two 8 bit channels
+float DecodePackedTwoChannelValue(float2 PackedHeight)
+{
+	return PackedHeight.x * 255.0 * 256.0 + PackedHeight.y * 255.0;
+}
+
+float DecodeHeightValue(float InValue)
+{
+	return (InValue - 32768.0) * TERRAIN_ZSCALE;
+}
+
+float DecodePackedHeight(float2 PackedHeight)
+{
+	return DecodeHeightValue(DecodePackedTwoChannelValue(PackedHeight));
+}
+
+uint ReverseBits32( uint bits )
+{
+#if SM5_PROFILE || COMPILER_METAL
+	return reversebits( bits );
+#else
+	bits = ( bits << 16) | ( bits >> 16);
+	bits = ( (bits & 0x00ff00ff) << 8 ) | ( (bits & 0xff00ff00) >> 8 );
+	bits = ( (bits & 0x0f0f0f0f) << 4 ) | ( (bits & 0xf0f0f0f0) >> 4 );
+	bits = ( (bits & 0x33333333) << 2 ) | ( (bits & 0xcccccccc) >> 2 );
+	bits = ( (bits & 0x55555555) << 1 ) | ( (bits & 0xaaaaaaaa) >> 1 );
+	return bits;
+#endif
+}
+
+SamplerState GetBilinearWrappedSampler()
+{
+#if PS4_PROFILE
+	sce::Gnm::Sampler BuiltinSampler = GetDefaultBuiltinSampler();
+	BuiltinSampler.setWrapMode(sce::Gnm::kWrapModeWrap, sce::Gnm::kWrapModeWrap, sce::Gnm::kWrapModeWrap);
+	BuiltinSampler.setXyFilterMode(sce::Gnm::kFilterModeBilinear, sce::Gnm::kFilterModeBilinear);
+	BuiltinSampler.setZFilterMode(sce::Gnm::kZFilterModeLinear);
+	BuiltinSampler.setMipFilterMode(sce::Gnm::kMipFilterModePoint);
+	return SamplerState(BuiltinSampler);
+#else
+	return BuiltinSamplers.Bilinear;
+#endif
+}
+
+SamplerState GetBilinearClampedSampler()
+{
+#if PS4_PROFILE
+	sce::Gnm::Sampler BuiltinSampler = GetDefaultBuiltinSampler();
+	BuiltinSampler.setXyFilterMode(sce::Gnm::kFilterModeBilinear, sce::Gnm::kFilterModeBilinear);
+	BuiltinSampler.setMipFilterMode(sce::Gnm::kMipFilterModePoint);
+	BuiltinSampler.setZFilterMode(sce::Gnm::kZFilterModeLinear);
+	BuiltinSampler.setWrapMode(sce::Gnm::kWrapModeClampLastTexel, sce::Gnm::kWrapModeClampLastTexel, sce::Gnm::kWrapModeClampLastTexel);
+	return SamplerState(BuiltinSampler);
+#else
+	return BuiltinSamplers.BilinearClamped;
+#endif
+}
+
+SamplerState GetPointWrappedSampler()
+{
+#if PS4_PROFILE
+	sce::Gnm::Sampler BuiltinSampler = GetDefaultBuiltinSampler();
+	BuiltinSampler.setXyFilterMode(sce::Gnm::kFilterModePoint, sce::Gnm::kFilterModePoint);
+	BuiltinSampler.setMipFilterMode(sce::Gnm::kMipFilterModePoint);
+	BuiltinSampler.setZFilterMode(sce::Gnm::kZFilterModePoint);
+	BuiltinSampler.setWrapMode(sce::Gnm::kWrapModeWrap, sce::Gnm::kWrapModeWrap, sce::Gnm::kWrapModeWrap);
+	return SamplerState(BuiltinSampler);
+#else
+	return BuiltinSamplers.Point;
+#endif
+}
+
+SamplerState GetPointClampedSampler()
+{
+#if PS4_PROFILE
+	sce::Gnm::Sampler BuiltinSampler = GetDefaultBuiltinSampler();
+	BuiltinSampler.setXyFilterMode(sce::Gnm::kFilterModePoint, sce::Gnm::kFilterModePoint);
+	BuiltinSampler.setMipFilterMode(sce::Gnm::kMipFilterModePoint);
+	BuiltinSampler.setZFilterMode(sce::Gnm::kZFilterModePoint);
+	BuiltinSampler.setWrapMode(sce::Gnm::kWrapModeClampLastTexel, sce::Gnm::kWrapModeClampLastTexel, sce::Gnm::kWrapModeClampLastTexel);
+	return SamplerState(BuiltinSampler);
+#else
+	return BuiltinSamplers.PointClamped;
+#endif
+}
+
+SamplerState GetTrilinearWrappedSampler()
+{
+#if PS4_PROFILE
+	sce::Gnm::Sampler BuiltinSampler = GetDefaultBuiltinSampler();
+	BuiltinSampler.setXyFilterMode(sce::Gnm::kFilterModeBilinear, sce::Gnm::kFilterModeBilinear);
+	BuiltinSampler.setMipFilterMode(sce::Gnm::kMipFilterModeLinear);
+	BuiltinSampler.setZFilterMode(sce::Gnm::kZFilterModeLinear);
+	BuiltinSampler.setWrapMode(sce::Gnm::kWrapModeWrap, sce::Gnm::kWrapModeWrap, sce::Gnm::kWrapModeWrap);
+	return SamplerState(BuiltinSampler);
+#else
+	return BuiltinSamplers.Trilinear;
+#endif
+}
+
+SamplerState GetTrilinearClampedSampler()
+{
+#if PS4_PROFILE
+	sce::Gnm::Sampler BuiltinSampler = GetDefaultBuiltinSampler();
+	BuiltinSampler.setXyFilterMode(sce::Gnm::kFilterModeBilinear, sce::Gnm::kFilterModeBilinear);
+	BuiltinSampler.setMipFilterMode(sce::Gnm::kMipFilterModeLinear);
+	BuiltinSampler.setZFilterMode(sce::Gnm::kZFilterModeLinear);
+	BuiltinSampler.setWrapMode(sce::Gnm::kWrapModeClampLastTexel, sce::Gnm::kWrapModeClampLastTexel, sce::Gnm::kWrapModeClampLastTexel);
+	return SamplerState(BuiltinSampler);
+#else
+	return BuiltinSamplers.TrilinearClamped;
+#endif
+}
+
+// see PixelShaderOutputCommon
+struct FPixelShaderIn
+{
+	// read only
+	float4 SvPosition;
+
+	// Pixel Shader InCoverage, only usable if PIXELSHADEROUTPUT_COVERAGE is 1
+	uint Coverage;
+
+	//
+	bool bIsFrontFace;
+};
+// see PixelShaderOutputCommon
+struct FPixelShaderOut
+{
+	// [0..7], only usable if PIXELSHADEROUTPUT_MRT0, PIXELSHADEROUTPUT_MRT1, ... is 1
+	float4 MRT[8];
+
+	// Pixel Shader OutCoverage, only usable if PIXELSHADEROUTPUT_COVERAGE is 1
+	uint Coverage;
+
+	// Pixel Shader OutDepth
+	float Depth;
+};
diff --git a/Engine/Shaders/Private/CommonViewUniformBuffer.ush b/Engine/Shaders/Private/CommonViewUniformBuffer.ush
new file mode 100644
--- /dev/null
+++ b/Engine/Shaders/Private/CommonViewUniformBuffer.ush
@@ -0,0 +1,32 @@
+// Copyright 1998-2017 Epic Games, Inc. All Rights Reserved.
+
+/*=============================================================================
+	CommonViewUniformBuffer.usf: Common view uniform buffer specifics
+=============================================================================*/
+
+/*
+* @return tan(View.FieldOfViewWideAngles * .5)
+*/
+float2 GetTanHalfFieldOfView()
+{
+	return float2(View.ClipToView[0][0], View.ClipToView[1][1]);
+}
+
+float2 GetPrevTanHalfFieldOfView()
+{
+	return float2(View.PrevClipToView[0][0], View.PrevClipToView[1][1]);
+}
+
+// might be used by Custom material expressions (still best to wrap the custom node in a material function)
+// @return 1 / tan(View.FieldOfViewWideAngles * .5)
+float2 GetCotanHalfFieldOfView()
+{
+	return float2(View.ViewToClip[0][0], View.ViewToClip[1][1]);
+}
+
+// might be used by Custom material expressions (still best to wrap the custom node in a material function)
+// @return previous 1 / tan(View.FieldOfViewWideAngles * .5)
+float2 GetPrevCotanHalfFieldOfView()
+{
+	return float2(View.PrevViewToClip[0][0], View.PrevViewToClip[1][1]);
+}
diff --git a/Engine/Shaders/Private/DecalCommon.ush b/Engine/Shaders/Private/DecalCommon.ush
new file mode 100644
--- /dev/null
+++ b/Engine/Shaders/Private/DecalCommon.ush
@@ -0,0 +1,97 @@
+// Copyright 1998-2017 Epic Games, Inc. All Rights Reserved.
+
+/*=============================================================================
+	DecalCommon.usf: decal shared things.
+=============================================================================*/
+
+#include "Common.ush"
+#include "DeferredShadingCommon.ush"		// GBufferData
+
+
+half3 FrameBufferDecalBlendOp(half4 Source)
+{
+#if (COMPILER_GLSL_ES2)
+	half4 Dest = FramebufferFetchES2();
+	Dest.rgb = Decode32BPPHDR(Dest).rgb;
+#else
+	half4 Dest = half4(0, 0, 0, 0);
+#endif
+    
+#if DECAL_BLEND_MODE == DECALBLENDMODEID_STAIN
+	return (Source.rgb*Dest.rgb) + (Dest.rgb*(1.0 - Source.a));
+#elif DECAL_BLEND_MODE == DECALBLENDMODEID_TRANSLUCENT
+	return (Source.rgb*Source.a) + (Dest.rgb*(1.0 - Source.a));
+#elif DECAL_BLEND_MODE == DECALBLENDMODEID_EMISSIVE
+	return (Source.rgb*Source.a) + Dest.rgb;
+#else
+	return Source.rgb;
+#endif
+}
+
+void DecalCommonOutput(inout FPixelShaderIn In, inout FPixelShaderOut Out, float3 Color, float Opacity, FGBufferData Data)
+{
+	// RETURN_COLOR not needed unless writing to SceneColor
+	Out.MRT[0] = float4(Color, Opacity);
+
+#if (ES2_PROFILE || ES3_1_PROFILE)
+#if (COMPILER_GLSL_ES2 || COMPILER_GLSL_ES3_1) || (MOBILE_EMULATION)
+	// Do decal blending if encoding requires it.
+	if (GetHDR32bppEncodeMode() == HDR_ENCODE_RGBE)
+	{
+		Out.MRT[0].rgb = FrameBufferDecalBlendOp(Out.MRT[0]);
+	}
+ 	// do 32bpp hdr encoding for mobile if required.
+	Out.MRT[0] = Encode32BPPHDR(Out.MRT[0], In.SvPosition.xy);
+#endif
+#endif
+
+#if DECAL_RENDERTARGET_COUNT > 1
+	// some MRT rendering
+
+	#if DECAL_RENDERSTAGE == 0
+	{
+		// before base pass (DBuffer)
+
+		// @param MultiOpacity .x: Color, .y:Normal, .z:Roughness
+		float3 MultiOpacity = Opacity;
+
+		EncodeDBufferData(Data, MultiOpacity, Out.MRT[0], Out.MRT[1], Out.MRT[2]);
+	}
+	#elif DECAL_RENDERSTAGE == 1
+	{
+		// after base pass (GBuffer)
+
+		// static lighting isn't updated by decals so we don't need to update that render target
+		float4 OutTarget4 = 0;
+		float4 OutTarget5 = 0;
+		float4 OutTarget6 = 0;
+
+		EncodeGBuffer(Data, Out.MRT[1], Out.MRT[2], Out.MRT[3], OutTarget4, OutTarget5, OutTarget6);
+	}
+	#elif DECAL_RENDERSTAGE == 2
+	{
+		// before lighting (GBuffer)
+
+		// static lighting isn't updated by decals so we don't need to update that render target
+		float4 OutTarget4 = 0;
+		float4 OutTarget5 = 0;
+		float4 OutTarget6 = 0;
+
+		EncodeGBuffer(Data, Out.MRT[1], Out.MRT[2], Out.MRT[3], OutTarget4, OutTarget5, OutTarget6);
+
+		#if DECAL_BLEND_MODE == DECALBLENDMODEID_STAIN
+		Out.MRT[3].rgb *= Opacity;
+		#endif
+
+		Out.MRT[0].a = Opacity;	// Emissive
+		Out.MRT[1].a = Opacity;	// Normal
+		Out.MRT[2].a = Opacity;	// Metallic, Specular, Roughness
+		Out.MRT[3].a = Opacity;	// BaseColor
+	}
+	#endif // DECAL_RENDERSTAGE
+#endif //DECAL_RENDERTARGET_COUNT > 1
+
+#if DECAL_RENDERTARGET_COUNT > 1 && !BIND_RENDERTARGET1 && COMPILER_METAL
+	Out.MRT[1] = 0;
+#endif
+}
diff --git a/Engine/Shaders/Private/Decode32bppHDR.ush b/Engine/Shaders/Private/Decode32bppHDR.ush
new file mode 100644
--- /dev/null
+++ b/Engine/Shaders/Private/Decode32bppHDR.ush
@@ -0,0 +1,49 @@
+// Copyright 1998-2017 Epic Games, Inc. All Rights Reserved.
+
+/*=============================================================================
+	Decode32bppHDR.usf: Decode 32bpp HDR encoding
+=============================================================================*/
+
+#include "Common.ush"
+
+float2 InvTexSize;
+
+void MainVS(
+	float2 InPosition	: ATTRIBUTE0,
+	float2 InUV			: ATTRIBUTE1,
+#if DECODING_MOSAIC
+	out float2 OutTexCoords[4] : TEXCOORD4,
+#endif
+	out FScreenVertexOutput Output
+	)
+{
+	DrawRectangle( float4( InPosition, 0, 1 ), InUV, Output.Position, Output.UV);
+
+#if DECODING_MOSAIC
+	OutTexCoords[0] = Output.UV + InvTexSize.xy * float2( 0,-1);
+	OutTexCoords[1] = Output.UV + InvTexSize.xy * float2( 1, 0);
+	OutTexCoords[2] = Output.UV + InvTexSize.xy * float2(-1, 0);
+	OutTexCoords[3] = Output.UV + InvTexSize.xy * float2( 0, 1);
+#endif
+}
+
+
+Texture2D InTexture;
+SamplerState InTextureSampler;
+
+void MainPS(
+#if DECODING_MOSAIC
+	in float2 InTexCoords[4] : TEXCOORD4,
+#endif
+	FScreenVertexOutput Input,
+	out float4 OutColor : SV_Target0
+	)
+{
+    float4 EncodedColor = Texture2DSample(InTexture, InTextureSampler, Input.UV);
+	// Decode EncodedColor, OutColor will saturate on 8 bit RTs.
+#if DECODING_MOSAIC
+	OutColor = Decode32BPPHDR(EncodedColor, Input.Position, InTexture, InTextureSampler, InTexCoords);
+#else
+	OutColor = Decode32BPPHDR(EncodedColor);
+#endif
+}
diff --git a/Engine/Shaders/Private/DeferredLightingCommon.ush b/Engine/Shaders/Private/DeferredLightingCommon.ush
new file mode 100644
--- /dev/null
+++ b/Engine/Shaders/Private/DeferredLightingCommon.ush
@@ -0,0 +1,594 @@
+// Copyright 1998-2017 Epic Games, Inc. All Rights Reserved.
+
+/*=============================================================================
+	DeferredLightingCommon.usf: Common definitions for deferred lighting.
+=============================================================================*/
+
+#ifndef __DEFERRED_LIGHTING_COMMON__
+#define __DEFERRED_LIGHTING_COMMON__
+
+#include "DeferredShadingCommon.ush"
+#include "DynamicLightingCommon.ush"
+#include "BRDF.ush"
+#include "MonteCarlo.ush"
+#include "IESLightProfilesCommon.ush"
+#include "ShadingModels.ush"
+
+/** 
+ * Data about a single light.
+ * Putting the light data in this struct allows the same lighting code to be used between standard deferred, 
+ * Where many light properties are known at compile time, and tiled deferred, where all light properties have to be fetched from a buffer.
+ */
+struct FDeferredLightData
+{
+	float4 LightPositionAndInvRadius;
+	float4 LightColorAndFalloffExponent;
+	float3 LightDirection;
+	float4 SpotAnglesAndSourceRadius;
+	float MinRoughness;
+	float ContactShadowLength;
+	float2 DistanceFadeMAD;
+	float4 ShadowMapChannelMask;
+	/** Whether to use inverse squared falloff. */
+	bool bInverseSquared;
+	/** Whether this is a light with radial attenuation, aka point or spot light. */
+	bool bRadialLight;
+	/** Whether this light needs spotlight attenuation. */
+	bool bSpotLight;
+	/** Whether the light should apply shadowing. */
+	uint ShadowedBits;
+};
+
+/** Data about a single light to be shaded with the simple shading model, designed for speed and limited feature set. */
+struct FSimpleDeferredLightData
+{
+	float4 LightPositionAndInvRadius;
+	float4 LightColorAndFalloffExponent;
+	/** Whether to use inverse squared falloff. */
+	bool bInverseSquared;
+};
+
+#define REFERENCE_QUALITY	0
+
+#undef LIGHT_SOURCE_SHAPE
+#define LIGHT_SOURCE_SHAPE 1
+
+#define LIGHT_SHAPE_SPHERE	1
+#define LIGHT_SHAPE_RECT	2
+
+static const uint LightSourceShape = LIGHT_SHAPE_SPHERE;
+
+bool RayHitRect( float3 R, float3 RectCenter, float3 RectX, float3 RectY, float3 RectZ, float RectExtentX, float RectExtentY )
+{
+	// Intersect ray with plane
+	float3 PointOnPlane = R * max( 0, dot( RectZ, RectCenter ) / dot( RectZ, R ) );
+
+	bool InExtentX = abs( dot( RectX, PointOnPlane - RectCenter ) ) <= RectExtentX;
+	bool InExtentY = abs( dot( RectY, PointOnPlane - RectCenter ) ) <= RectExtentY;
+	return InExtentX && InExtentY;
+}
+
+float3 SphereLightingMIS( FGBufferData GBuffer, FDeferredLightData LightData, float3 LobeRoughness, float3 ToLight, float3 V, float3 N, uint2 Random )
+{
+	float3 Lighting = 0;
+
+	LobeRoughness = max( 0.08, LobeRoughness );
+
+	const float SourceRadius = max( 1, LightData.SpotAnglesAndSourceRadius.z );
+
+	const float DistanceSqr = dot( ToLight, ToLight );
+	const float3 ConeAxis = ToLight * rsqrt( DistanceSqr );
+	const float ConeCos = sqrt( 1 - Square( SourceRadius ) / DistanceSqr );
+
+	const float Area = PI * Square(SourceRadius);
+	const float SampleColor = 1.0 / Area;
+	
+	const uint NumSets = 3;
+	const uint NumSamples[ NumSets ] =
+	{
+		0,
+		4,
+		4,
+	};
+
+	UNROLL
+	for( uint Set = 0; Set < NumSets; Set++ )
+	{
+		LOOP
+		for( uint i = 0; i < NumSamples[ Set ]; i++ )
+		{
+			float2 E = Hammersley( i, NumSamples[ Set ], Random );
+			//float2 E = CorrelatedMultiJitter2D( i, NumSamples[ Set ], Random.x * Random.y * (Set + 17) );
+			
+			float3 L, H;
+			if( Set == 0 )
+			{
+				L = TangentToWorld( CosineSampleHemisphere( E ).xyz, N );
+				H = normalize(V + L);
+			}
+			else if( Set == 1 )
+			{
+				H = TangentToWorld( ImportanceSampleGGX( E, LobeRoughness[1] ).xyz, N );
+				L = 2 * dot( V, H ) * H - V;
+			}
+			else
+			{
+				L = TangentToWorld( UniformSampleCone( E, ConeCos ).xyz, ConeAxis );
+				H = normalize(V + L);
+			}
+
+			float NoL = saturate( dot(N, L) );
+			float NoH = saturate( dot(N, H) );
+			float VoH = saturate( dot(V, H) );
+
+			if( NoL > 0 && VoH > 0 )
+			{
+				BRANCH
+				if( Set != 2 && dot( L, ConeAxis ) < ConeCos )
+				{
+					// Ray misses sphere
+					continue;
+				}
+
+				float PDF[] =
+				{
+					NoL / PI,
+					D_GGX( LobeRoughness[1], NoH ) * NoH / (4 * VoH),
+					1.0 / ( 2 * PI * (1 - ConeCos) ),
+				};
+
+				// MIS balance heuristic
+				float InvWeight = 0;
+				UNROLL for( uint j = 0; j < NumSets; j++ )
+				{
+					InvWeight += PDF[j] * NumSamples[j];
+				}
+				float Weight = rcp( InvWeight );
+
+				float3 Shading = SurfaceShading( GBuffer, LobeRoughness, 1, L, V, N, Random ) * NoL;
+				Shading += SubsurfaceShading( GBuffer, L, V, N, 1, Random );
+			
+				Lighting += SampleColor * Shading * Weight;
+			}
+		}
+	}
+
+	return Lighting;
+}
+
+// Find representative incoming light direction and energy modification
+float3 AreaLightSpecular( FDeferredLightData LightData, inout float3 LobeRoughness, inout float3 ToLight, inout float3 L, float3 V, half3 N )
+{
+	float3 LobeEnergy = 1;
+
+	LobeRoughness = max( LobeRoughness, LightData.MinRoughness );
+	float3 m = LobeRoughness * LobeRoughness;
+	
+	const float SourceRadius = LightData.SpotAnglesAndSourceRadius.z;
+	const float SourceLength = LightData.SpotAnglesAndSourceRadius.w;
+
+	// TODO early out for point lights
+	
+	float3 R = reflect( -V, N );
+	float InvDistToLight = rsqrt( dot( ToLight, ToLight ) );
+
+	// Point lobe in off-specular peak direction
+	float a = Square( LobeRoughness[1] );
+	R = lerp( N, R, (1 - a) * ( sqrt(1 - a) + a ) );
+	R = normalize( R );
+
+	BRANCH
+	if( SourceLength > 0 )
+	{
+		// Energy conservation
+		// asin(x) is angle to sphere, atan(x) is angle to disk, saturate(x) is free and in the middle
+		float LineAngle = saturate( SourceLength * InvDistToLight );
+		LobeEnergy *= m / saturate( m + 0.5 * LineAngle );
+
+		// Closest point on line segment to ray
+		float3 L01 = LightData.LightDirection * SourceLength;
+		float3 L0 = ToLight - 0.5 * L01;
+		float3 L1 = ToLight + 0.5 * L01;
+
+#if 1
+		// Shortest distance
+		float a = Square( SourceLength );
+		float b = dot( R, L01 );
+		float t = saturate( dot( L0, b*R - L01 ) / (a - b*b) );
+#else
+		// Smallest angle
+		float A = Square( SourceLength );
+		float B = 2 * dot( L0, L01 );
+		float C = dot( L0, L0 );
+		float D = dot( R, L0 );
+		float E = dot( R, L01 );
+		float t = saturate( (B*D - 2*C*E) / (B*E - 2*A*D) );
+#endif
+
+		ToLight = L0 + t * L01;
+	}
+
+	BRANCH
+	if( SourceRadius > 0 )
+	{
+		// Energy conservation
+		// asin(x) is angle to sphere, atan(x) is angle to disk, saturate(x) is free and in the middle
+		float SphereAngle = saturate( SourceRadius * InvDistToLight );
+		LobeEnergy *= Square( m / saturate( m + 0.5 * SphereAngle ) );
+		
+		// Closest point on sphere to ray
+		float3 ClosestPointOnRay = dot( ToLight, R ) * R;
+		float3 CenterToRay = ClosestPointOnRay - ToLight;
+		float3 ClosestPointOnSphere = ToLight + CenterToRay * saturate( SourceRadius * rsqrt( dot( CenterToRay, CenterToRay ) ) );
+		ToLight = ClosestPointOnSphere;
+	}
+
+	L = normalize( ToLight );
+
+	return LobeEnergy;
+}
+
+/** Returns 0 for positions closer than the fade near distance from the camera, and 1 for positions further than the fade far distance. */
+float DistanceFromCameraFade(float SceneDepth, FDeferredLightData LightData, float3 WorldPosition, float3 CameraPosition)
+{
+	// depth (non radial) based fading over distance
+	float Fade = saturate(SceneDepth * LightData.DistanceFadeMAD.x + LightData.DistanceFadeMAD.y);
+	return Fade * Fade;
+}
+
+void GetShadowTerms(FGBufferData GBuffer, FDeferredLightData LightData, float3 WorldPosition, float4 LightAttenuation, out float OpaqueShadowTerm, out float SSSShadowTerm)
+{
+	// Remapping the light attenuation buffer (see ShadowRendering.cpp)
+
+	// LightAttenuation: Light function + per-object shadows in z, per-object SSS shadowing in w, 
+	// Whole scene directional light shadows in x, whole scene directional light SSS shadows in y
+	// Get static shadowing from the appropriate GBuffer channel
+	float UsesStaticShadowMap = dot(LightData.ShadowMapChannelMask, float4(1, 1, 1, 1));
+	float StaticShadowing = lerp(1, dot(GBuffer.PrecomputedShadowFactors, LightData.ShadowMapChannelMask), UsesStaticShadowMap);
+
+	if (LightData.bRadialLight)
+	{
+		// Remapping the light attenuation buffer (see ShadowRendering.cpp)
+
+		OpaqueShadowTerm = LightAttenuation.z * StaticShadowing;
+		// SSS uses a separate shadowing term that allows light to penetrate the surface
+		//@todo - how to do static shadowing of SSS correctly?
+		SSSShadowTerm = LightAttenuation.w * StaticShadowing;
+	}
+	else
+	{
+		// Remapping the light attenuation buffer (see ShadowRendering.cpp)
+		// Also fix up the fade between dynamic and static shadows
+		// to work with plane splits rather than spheres.
+
+		float DynamicShadowFraction = DistanceFromCameraFade(GBuffer.Depth, LightData, WorldPosition, View.WorldCameraOrigin);
+		// For a directional light, fade between static shadowing and the whole scene dynamic shadowing based on distance + per object shadows
+		OpaqueShadowTerm = lerp(LightAttenuation.x, StaticShadowing, DynamicShadowFraction);
+		// Fade between SSS dynamic shadowing and static shadowing based on distance
+		SSSShadowTerm = min(lerp(LightAttenuation.y, StaticShadowing, DynamicShadowFraction), LightAttenuation.w);
+		
+		// combine with light function
+		OpaqueShadowTerm *= LightAttenuation.z;
+		SSSShadowTerm *= LightAttenuation.z;
+	}
+}
+
+float ShadowRayCast(
+	float3 RayOriginTranslatedWorld, float3 RayDirection, float RayLength,
+	int NumSteps, float StepOffset
+)
+{
+	float4 RayStartClip	= mul( float4( RayOriginTranslatedWorld, 1 ), View.TranslatedWorldToClip );
+	float4 RayDirClip	= mul( float4( RayDirection * RayLength, 0 ), View.TranslatedWorldToClip );
+	float4 RayEndClip	= RayStartClip + RayDirClip;
+
+	float3 RayStartScreen = RayStartClip.xyz / RayStartClip.w;
+	float3 RayEndScreen = RayEndClip.xyz / RayEndClip.w;
+	
+	float3 RayStepScreen = RayEndScreen - RayStartScreen;
+
+	float3 RayStartUVz = float3( RayStartScreen.xy * View.ScreenPositionScaleBias.xy + View.ScreenPositionScaleBias.wz, RayStartScreen.z );
+	float3 RayStepUVz = float3( RayStepScreen.xy * View.ScreenPositionScaleBias.xy, RayStepScreen.z );
+
+	float4 RayDepthClip	= RayStartClip + mul( float4( 0, 0, RayLength, 0 ), View.ViewToClip );
+	float3 RayDepthScreen = RayDepthClip.xyz / RayDepthClip.w;
+
+	const float Step = 1.0 / NumSteps;
+
+	// *2 to get less morie pattern in extreme cases, larger values make object appear not grounded in reflections
+	const float CompareTolerance = abs( RayDepthScreen.z - RayStartScreen.z ) * Step * 2;
+
+	float SampleTime = StepOffset * Step + Step;
+
+	float FirstHitTime = -1.0;
+
+	UNROLL
+	for( int i = 0; i < NumSteps; i++ )
+	{
+		float3 SampleUVz = RayStartUVz + RayStepUVz * SampleTime;
+		float SampleDepth = SceneDepthTexture.SampleLevel( SceneDepthTextureSampler, SampleUVz.xy, 0 ).r;
+
+		float DepthDiff = SampleUVz.z - SampleDepth;
+		bool Hit = abs( DepthDiff + CompareTolerance ) < CompareTolerance;
+
+		FirstHitTime = (Hit && FirstHitTime < 0.0) ? SampleTime : FirstHitTime;
+
+		SampleTime += Step;
+	}
+
+	float Shadow = FirstHitTime > 0.0 ? 1.0 : 0.0;
+
+	// Off screen masking
+	float2 Vignette = max(6.0 * abs(RayStartScreen.xy + RayStepScreen.xy * FirstHitTime) - 5.0, 0.0);
+	Shadow *= saturate( 1.0 - dot( Vignette, Vignette ) );
+
+	return 1 - Shadow;
+}
+
+#ifndef SUPPORT_CONTACT_SHADOWS
+#error "Must set SUPPORT_CONTACT_SHADOWS"
+#endif
+
+void GetLocalLightAttenuation(
+	float3 WorldPosition, 
+	float3 N,
+	FDeferredLightData LightData, 
+	float DistanceBiasSqr,
+	inout float3 ToLight, 
+	inout float3 L, 
+	inout float NoL, 
+	inout float DistanceAttenuation, 
+	inout float LightRadiusMask, 
+	inout float SpotFalloff)
+	{
+		ToLight = LightData.LightPositionAndInvRadius.xyz - WorldPosition;
+		
+		float DistanceSqr = dot( ToLight, ToLight );
+		L = ToLight * rsqrt( DistanceSqr );
+
+		if (LightData.bInverseSquared)
+		{
+			const float SourceRadius = LightData.SpotAnglesAndSourceRadius.z;
+			const float SourceLength = LightData.SpotAnglesAndSourceRadius.w;
+
+			BRANCH
+			if( SourceLength > 0 )
+			{
+				// Line segment irradiance
+				float3 L01 = LightData.LightDirection * SourceLength;
+				float3 ToLight0 = ToLight - 0.5 * L01;
+				float3 ToLight1 = ToLight + 0.5 * L01;
+
+				float LengthSqr0 = dot( ToLight0, ToLight0 );
+				float LengthSqr1 = dot( ToLight1, ToLight1 );
+				float rLength0 = rsqrt( LengthSqr0 );
+				float rLength1 = rsqrt( LengthSqr1 );
+				float Length0 = LengthSqr0 * rLength0;
+				float Length1 = LengthSqr1 * rLength1;
+
+				DistanceAttenuation = rcp( ( Length0 * Length1 + dot( ToLight0, ToLight1 ) ) * 0.5 + DistanceBiasSqr );
+				NoL = saturate( 0.5 * ( dot(N, ToLight0) * rLength0 + dot(N, ToLight1) * rLength1 ) );
+			}
+			else
+			{
+				DistanceAttenuation = rcp( DistanceSqr + DistanceBiasSqr );
+				NoL = saturate( dot( N, L ) );
+
+				if( SourceRadius > 0 )
+				{
+				#if 1	//HORIZON
+					NoL = dot( N, L );
+
+					float SinAlphaSqr = saturate( Square( SourceRadius ) / DistanceSqr );
+					float SinAlpha = sqrt( SinAlphaSqr );
+
+					if( NoL < SinAlpha )
+					{
+					#if 0
+						// Accurate sphere irradiance
+						float CosBeta = NoL;
+						float SinBeta = sqrt( 1 - CosBeta * CosBeta );
+						float TanBeta = SinBeta / CosBeta;
+					
+						float x = sqrt( 1 / SinAlphaSqr - 1 );
+						float y = -x / TanBeta;
+						float z = SinBeta * sqrt(1 - y*y);
+
+						DistanceAttenuation = SinAlphaSqr * ( NoL * acos(y) - x * z ) + atan( z / x );
+						DistanceAttenuation /= PI * Square( SourceRadius );
+						NoL = 1;
+					#else
+						// Hermite spline approximation
+						// Fairly accurate with SinAlpha < 0.8
+						// y=0 and dy/dx=0 at -SinAlpha
+						// y=SinAlpha and dy/dx=1 at SinAlpha
+						NoL = max( NoL, -SinAlpha );
+						NoL = Square( SinAlpha + NoL ) / ( 4 * SinAlpha );
+					#endif
+					}
+				#endif
+				}
+			}
+
+			// TODO optimize
+			LightRadiusMask = Square( saturate( 1 - Square( DistanceSqr * Square(LightData.LightPositionAndInvRadius.w) ) ) );
+		}
+		else
+		{
+			DistanceAttenuation = 1;
+			NoL = saturate( dot( N, L ) );
+			
+			LightRadiusMask = RadialAttenuation(ToLight * LightData.LightPositionAndInvRadius.w, LightData.LightColorAndFalloffExponent.w);
+
+#if REFERENCE_QUALITY
+			// anti Area
+			LightRadiusMask *= DistanceSqr + 1;
+#endif
+		}
+
+		if (LightData.bSpotLight)
+		{
+			SpotFalloff = SpotAttenuation(L, -LightData.LightDirection, LightData.SpotAnglesAndSourceRadius.xy);
+		}
+	}
+
+/** Calculates lighting for a given position, normal, etc with a fully featured lighting model designed for quality. */
+float4 GetDynamicLighting(float3 WorldPosition, float3 CameraVector, FGBufferData GBuffer, float AmbientOcclusion, uint ShadingModelID, FDeferredLightData LightData, float4 LightAttenuation, float Dither, uint2 Random)
+{
+	FLightAccumulator LightAccumulator = (FLightAccumulator)0;
+
+	float3 V = -CameraVector;
+	float3 N = GBuffer.WorldNormal;
+	float3 L = LightData.LightDirection;	// Already normalized
+	float3 ToLight = L * GBuffer.Depth;
+	float NoL = saturate( dot(N, L) );
+	float DistanceAttenuation = 1;
+	float LightRadiusMask = 1;
+	float SpotFalloff = 1;
+
+	if (LightData.bRadialLight)
+	{
+		GetLocalLightAttenuation(WorldPosition, N, LightData, 1, ToLight, L, NoL, DistanceAttenuation, LightRadiusMask, SpotFalloff);
+	}
+
+	LightAccumulator.EstimatedCost += 0.3f;		// running the PixelShader at all has a cost
+
+	BRANCH
+	if (LightRadiusMask > 0 && SpotFalloff > 0)
+	{
+		float SurfaceShadow = 1;
+		float SubsurfaceShadow = 1;
+
+		BRANCH
+		if (LightData.ShadowedBits)
+		{
+			GetShadowTerms(GBuffer, LightData, WorldPosition, LightAttenuation, SurfaceShadow, SubsurfaceShadow);
+
+			// greatly reduces shadow mapping artifacts
+			// Commented out because it reduces character shading quality.
+			//SurfaceShadow *= saturate(dot(N, L) * 6 - 0.2);
+
+#if SUPPORT_CONTACT_SHADOWS
+			BRANCH
+			if( LightData.ShadowedBits > 1 && LightData.ContactShadowLength > 0 )
+			{
+				float StepOffset = Dither - 0.5;
+				float Shadow = ShadowRayCast( WorldPosition + View.PreViewTranslation, L, LightData.ContactShadowLength * GBuffer.Depth, 8, StepOffset );
+				
+				SurfaceShadow *= Shadow;
+				SubsurfaceShadow *= Shadow * 0.5 + 0.5;
+			}
+#endif
+		}
+		else
+		{
+			SurfaceShadow = AmbientOcclusion;
+		}
+
+#if SUPPORT_CONTACT_SHADOWS
+		BRANCH
+		if( LightData.ShadowedBits < 2 && GBuffer.ShadingModelID == SHADINGMODELID_HAIR )
+		{
+			float StepOffset = Dither - 0.5;
+			SubsurfaceShadow = ShadowRayCast( WorldPosition + View.PreViewTranslation, L, 0.1 * GBuffer.Depth, 8, StepOffset );
+		}
+#endif
+
+		float SurfaceAttenuation	= (DistanceAttenuation * LightRadiusMask * SpotFalloff) * SurfaceShadow;
+		float SubsurfaceAttenuation	= (DistanceAttenuation * LightRadiusMask * SpotFalloff) * SubsurfaceShadow;
+
+		LightAccumulator.EstimatedCost += 0.3f;		// add the cost of getting the shadow terms
+
+		const bool bNeedsSeparateSubsurfaceLightAccumulation = GBuffer.ShadingModelID == SHADINGMODELID_SUBSURFACE_PROFILE;
+		const float3 LightColor = LightData.LightColorAndFalloffExponent.rgb;
+
+		#if NON_DIRECTIONAL_DIRECT_LIGHTING
+
+			float3 VolumeLighting = Diffuse_Lambert(GBuffer.DiffuseColor);
+			LightAccumulator_Add(LightAccumulator, VolumeLighting, (1.0/PI), LightColor * SurfaceAttenuation, bNeedsSeparateSubsurfaceLightAccumulation);
+
+		#else
+		{
+			const float ClearCoatRoughness	= GBuffer.CustomData.y;
+
+			float3 LobeRoughness = float3(ClearCoatRoughness, GBuffer.Roughness, 1);
+			
+			#if REFERENCE_QUALITY
+				float LightMask = LightRadiusMask * SpotFalloff * SurfaceShadow;
+				LightAccumulator_Add(LightAccumulator, SphereLightingMIS( GBuffer, LightData, LobeRoughness, ToLight, V, N, Random ), 0, LightColor * LightMask, bNeedsSeparateSubsurfaceLightAccumulation);
+			#else
+				float3 LobeEnergy = AreaLightSpecular(LightData, LobeRoughness, ToLight, L, V, N);
+
+				// accumulate surface
+				{
+					float3 SurfaceLighting = SurfaceShading(GBuffer, LobeRoughness, LobeEnergy, L, V, N, Random);
+					LightAccumulator_Add(LightAccumulator, SurfaceLighting, (1.0/PI), LightColor * (NoL * SurfaceAttenuation), bNeedsSeparateSubsurfaceLightAccumulation);
+				}
+
+				// accumulate subsurface
+				{
+					float3 SubsurfaceLighting = SubsurfaceShading(GBuffer, L, V, N, SubsurfaceShadow, Random);
+
+					LightAccumulator_Add(LightAccumulator, SubsurfaceLighting, 0, LightColor * SubsurfaceAttenuation, false);
+
+					LightAccumulator.EstimatedCost += 0.4f;		// add the cost of the lighting computations (should sum up to 1 form one light)
+				}
+			#endif
+		}
+		#endif
+	}
+
+	return LightAccumulator_GetResult(LightAccumulator);
+}
+
+/** 
+ * Calculates lighting for a given position, normal, etc with a simple lighting model designed for speed. 
+ * All lights rendered through this method are unshadowed point lights with no shadowing or light function or IES.
+ * A cheap specular is used instead of the more correct area specular, no fresnel.
+ */
+float3 GetSimpleDynamicLighting(float3 WorldPosition, float3 CameraVector, float3 WorldNormal, float AmbientOcclusion, float3 DiffuseColor, float3 SpecularColor, float Roughness, FSimpleDeferredLightData LightData)
+{
+	float3 V = -CameraVector;
+	float3 N = WorldNormal;
+	float3 ToLight = LightData.LightPositionAndInvRadius.xyz - WorldPosition;
+	float DistanceAttenuation = 1;
+	
+	float DistanceSqr = dot( ToLight, ToLight );
+	float3 L = ToLight * rsqrt( DistanceSqr );
+	float NoL = saturate( dot( N, L ) );
+
+	if (LightData.bInverseSquared)
+	{
+		// Sphere falloff (technically just 1/d2 but this avoids inf)
+		DistanceAttenuation = 1 / ( DistanceSqr + 1 );
+	
+		float LightRadiusMask = Square( saturate( 1 - Square( DistanceSqr * Square(LightData.LightPositionAndInvRadius.w) ) ) );
+		DistanceAttenuation *= LightRadiusMask;
+	}
+	else
+	{
+		DistanceAttenuation = RadialAttenuation(ToLight * LightData.LightPositionAndInvRadius.w, LightData.LightColorAndFalloffExponent.w);
+	}
+
+	float3 OutLighting = 0;
+
+	BRANCH
+	if (DistanceAttenuation > 0)
+	{
+		const float3 LightColor = LightData.LightColorAndFalloffExponent.rgb;
+
+		// Apply SSAO to the direct lighting since we're not going to have any other shadowing
+		float Attenuation = DistanceAttenuation * AmbientOcclusion;
+
+		#if NON_DIRECTIONAL_DIRECT_LIGHTING
+			float3 VolumeLighting = Diffuse_Lambert(DiffuseColor);
+			OutLighting += LightColor * Attenuation * VolumeLighting;
+		#else
+			OutLighting += LightColor * (NoL * Attenuation) * SimpleShading(DiffuseColor, SpecularColor, max(Roughness, .04f), L, V, N);
+		#endif
+	}
+
+	return OutLighting;
+}
+
+#endif
\ No newline at end of file
