Branch snapshot for CL 3555263

diff --git a/Engine/Shaders/Private/DeferredShadingCommon.ush b/Engine/Shaders/Private/DeferredShadingCommon.ush
new file mode 100644
--- /dev/null
+++ b/Engine/Shaders/Private/DeferredShadingCommon.ush
@@ -0,0 +1,913 @@
+// Copyright 1998-2017 Epic Games, Inc. All Rights Reserved.
+
+/*=============================================================================
+	DeferredShadingCommon.usf: Common definitions for deferred shading.
+=============================================================================*/
+
+#ifndef __DEFERRED_SHADING_COMMON__
+#define __DEFERRED_SHADING_COMMON__
+
+#include "LightAccumulator.ush"
+
+// TODO: for CustomGBufferResolvePS() MSAA_SAMPLE_COUNT is defined by C++ code as 2 or 4
+// bot not for any other shaders!
+#ifndef MSAA_SAMPLE_COUNT
+	#define MSAA_SAMPLE_COUNT 2
+#endif
+
+float3 RGBToYCoCg( float3 RGB )
+{
+	float Y  = dot( RGB, float3(  1, 2,  1 ) ) * 0.25;
+	float Co = dot( RGB, float3(  2, 0, -2 ) ) * 0.25 + ( 0.5 * 256.0 / 255.0 );
+	float Cg = dot( RGB, float3( -1, 2, -1 ) ) * 0.25 + ( 0.5 * 256.0 / 255.0 );
+	
+	float3 YCoCg = float3( Y, Co, Cg );
+	return YCoCg;
+}
+
+float3 YCoCgToRGB( float3 YCoCg )
+{
+	float Y  = YCoCg.x;
+	float Co = YCoCg.y - ( 0.5 * 256.0 / 255.0 );
+	float Cg = YCoCg.z - ( 0.5 * 256.0 / 255.0 );
+
+	float R = Y + Co - Cg;
+	float G = Y + Cg;
+	float B = Y - Co - Cg;
+
+	float3 RGB = float3( R, G, B );
+	return RGB;
+}
+
+// Octahedron Normal Vectors
+// [Cigolle 2014, "A Survey of Efficient Representations for Independent Unit Vectors"]
+//						Mean	Max
+// oct		8:8			0.33709 0.94424
+// snorm	8:8:8		0.17015 0.38588
+// oct		10:10		0.08380 0.23467
+// snorm	10:10:10	0.04228 0.09598
+// oct		12:12		0.02091 0.05874
+
+float2 UnitVectorToOctahedron( float3 N )
+{
+	N.xy /= dot( 1, abs(N) );
+	if( N.z <= 0 )
+	{
+		N.xy = ( 1 - abs(N.yx) ) * ( N.xy >= 0 ? float2(1,1) : float2(-1,-1) );
+	}
+	return N.xy;
+}
+
+float3 OctahedronToUnitVector( float2 Oct )
+{
+	float3 N = float3( Oct, 1 - dot( 1, abs(Oct) ) );
+	if( N.z < 0 )
+	{
+		N.xy = ( 1 - abs(N.yx) ) * ( N.xy >= 0 ? float2(1,1) : float2(-1,-1) );
+	}
+	return normalize(N);
+}
+
+float2 UnitVectorToHemiOctahedron( float3 N )
+{
+	N.xy /= dot( 1, abs(N) );
+	return float2( N.x + N.y, N.x - N.y );
+}
+
+float3 HemiOctahedronToUnitVector( float2 Oct )
+{
+	Oct = float2( Oct.x + Oct.y, Oct.x - Oct.y ) * 0.5;
+	float3 N = float3( Oct, 1 - dot( 1, abs(Oct) ) );
+	return normalize(N);
+}
+
+float3 Pack1212To888( float2 x )
+{
+	// Pack 12:12 to 8:8:8
+#if 0
+	uint2 x1212 = (uint2)( x * 4095.0 );
+	uint2 High = x1212 >> 8;
+	uint2 Low = x1212 & 255;
+	uint3 x888 = uint3( Low, High.x | (High.y << 4) );
+	return x888 / 255.0;
+#else
+	float2 x1212 = floor( x * 4095 );
+	float2 High = floor( x1212 / 256 );	// x1212 >> 8
+	float2 Low = x1212 - High * 256;	// x1212 & 255
+	float3 x888 = float3( Low, High.x + High.y * 16 );
+	return saturate( x888 / 255 );
+#endif
+}
+
+float2 Pack888To1212( float3 x )
+{
+	// Pack 8:8:8 to 12:12
+#if 0
+	uint3 x888 = (uint3)( x * 255.0 );
+	uint High = x888.z >> 4;
+	uint Low = x888.z & 15;
+	uint2 x1212 = x888.xy | uint2( Low << 8, High << 8 );
+	return x1212 / 4095.0;
+#else
+	float3 x888 = floor( x * 255 );
+	float High = floor( x888.z / 16 );	// x888.z >> 4
+	float Low = x888.z - High * 16;		// x888.z & 15
+	float2 x1212 = x888.xy + float2( Low, High ) * 256;
+	return saturate( x1212 / 4095 );
+#endif
+}
+
+float3 EncodeNormal( float3 N )
+{
+	return N * 0.5 + 0.5;
+	//return Pack1212To888( UnitVectorToOctahedron( N ) * 0.5 + 0.5 );
+}
+
+float3 DecodeNormal( float3 N )
+{
+	return N * 2 - 1;
+	//return OctahedronToUnitVector( Pack888To1212( N ) * 2 - 1 );
+}
+
+void EncodeNormal( inout float3 N, out uint Face )
+{
+#if 1
+	uint Axis = 2;
+	if( abs(N.x) >= abs(N.y) && abs(N.x) >= abs(N.z) )
+	{
+		Axis = 0;
+	}
+	else if( abs(N.y) > abs(N.z) )
+	{
+		Axis = 1;
+	}
+	Face = Axis * 2;
+#else
+	// TODO GCN
+	Face = v_cubeid_f32( N );
+	uint Axis = Face >> 1;
+#endif
+
+	N = Axis == 0 ? N.yzx : N;
+	N = Axis == 1 ? N.xzy : N;
+	
+	Face += N.z > 0 ? 0 : 1;
+	N.xy *= N.z > 0 ? 1 : -1;
+	N.xy = N.xy * 0.5 + 0.5;
+}
+
+void DecodeNormal( inout float3 N, in uint Face )
+{
+	uint Axis = Face >> 1;
+
+	N.xy = N.xy * 2 - 1;
+	N.z = sqrt( 1 - dot( N.xy, N.xy ) );
+
+	N = Axis == 0 ? N.zxy : N;
+	N = Axis == 1 ? N.xzy : N;
+	N *= (Face & 1) ? -1 : 1;
+}
+
+float3 EncodeBaseColor(float3 BaseColor)
+{
+	// we use sRGB on the render target to give more precision to the darks
+	return BaseColor;
+}
+
+float3 DecodeBaseColor(float3 BaseColor)
+{
+	// we use sRGB on the render target to give more precision to the darks
+	return BaseColor;
+}
+
+float3 EncodeSubsurfaceColor(float3 SubsurfaceColor)
+{
+	return sqrt(saturate(SubsurfaceColor));
+}
+
+// @param SubsurfaceProfile 0..1, SubsurfaceProfileId = int(x * 255)
+float3 EncodeSubsurfaceProfile(float SubsurfaceProfile)
+{
+	return float3(SubsurfaceProfile, 0, 0);
+}
+
+float EncodeIndirectIrradiance(float IndirectIrradiance)
+{
+	float L = IndirectIrradiance;
+	const float LogBlackPoint = 0.00390625;	// exp2(-8);
+	return log2( L + LogBlackPoint ) / 16 + 0.5;
+}
+
+float DecodeIndirectIrradiance(float IndirectIrradiance)
+{
+	// LogL -> L
+	float LogL = IndirectIrradiance;
+	const float LogBlackPoint = 0.00390625;	// exp2(-8);
+	return exp2( LogL * 16 - 8 ) - LogBlackPoint;		// 1 exp2, 1 smad, 1 ssub
+}
+
+float ComputeAngleFromRoughness( float Roughness, const float Threshold = 0.04f )
+{
+#if 1
+	float Angle = 3 * Square( Roughness );
+#else
+	const float LogThreshold = log2( Threshold );
+	float Power = 0.5 / pow( Roughness, 4 ) - 0.5;
+	float Angle = acos( exp2( LogThreshold / Power ) );
+#endif
+	return Angle;
+}
+
+float ComputeRoughnessFromAngle( float Angle, const float Threshold = 0.04f )
+{
+#if 1
+	float Roughness = sqrt( 0.33333 * Angle );
+#else
+	const float LogThreshold = log2( Threshold );
+	float Power = LogThreshold / log2( cos( Angle ) );
+	float Roughness = sqrt( sqrt( 2 / (Power * 4 + 2) ) );
+#endif
+	return Roughness;
+}
+
+float AddAngleToRoughness( float Angle, float Roughness )
+{
+	return saturate( sqrt( Square( Roughness ) + 0.33333 * Angle ) );
+}
+
+// @param Scalar clamped in 0..1 range
+// @param Mask 0..1
+// @return 8bit in range 0..1
+float Encode71(float Scalar, uint Mask)
+{
+	return
+		127.0f / 255.0f * saturate(Scalar) +
+		128.0f / 255.0f * Mask;
+}
+
+// 8bit reinterpretation as 7bit,1bit
+// @param Scalar 0..1
+// @param Mask 0..1
+// @return 7bit in 0.1
+float Decode71(float Scalar, out uint Mask)
+{
+	Mask = (uint)(Scalar > 0.5f);
+
+	return (Scalar - 0.5f * Mask) * 2.0f;
+}
+
+#define SHADINGMODELID_UNLIT				0
+#define SHADINGMODELID_DEFAULT_LIT			1
+#define SHADINGMODELID_SUBSURFACE			2
+#define SHADINGMODELID_PREINTEGRATED_SKIN	3
+#define SHADINGMODELID_CLEAR_COAT			4
+#define SHADINGMODELID_SUBSURFACE_PROFILE	5
+#define SHADINGMODELID_TWOSIDED_FOLIAGE		6
+#define SHADINGMODELID_HAIR					7
+#define SHADINGMODELID_CLOTH				8
+#define SHADINGMODELID_EYE					9
+#define SHADINGMODELID_NUM					10
+#define SHADINGMODELID_MASK					0xF					
+
+// The flags are defined so that 0 value has no effect!
+#define SKIP_CUSTOMDATA_MASK			(1 << 4)	// TODO remove. Can be inferred from shading model.
+#define SKIP_PRECSHADOW_MASK			(1 << 5)
+#define ZERO_PRECSHADOW_MASK			(1 << 6)
+#define SKIP_VELOCITY_MASK				(1 << 7)
+
+
+float EncodeShadingModelIdAndSelectiveOutputMask(uint ShadingModelId, uint SelectiveOutputMask)
+{
+	uint Value = (ShadingModelId & SHADINGMODELID_MASK) | SelectiveOutputMask;
+	return (float)Value / (float)0xFF;
+}
+
+uint DecodeShadingModelId(float InPackedChannel)
+{
+	return ((uint)round(InPackedChannel * (float)0xFF)) & SHADINGMODELID_MASK;
+}
+
+uint DecodeSelectiveOutputMask(float InPackedChannel)
+{
+	return ((uint)round(InPackedChannel * (float)0xFF)) & ~SHADINGMODELID_MASK;
+}
+
+bool IsSubsurfaceModel(int ShadingModel)
+{
+	return ShadingModel == SHADINGMODELID_SUBSURFACE 
+		|| ShadingModel == SHADINGMODELID_PREINTEGRATED_SKIN 
+		|| ShadingModel == SHADINGMODELID_SUBSURFACE_PROFILE
+		|| ShadingModel == SHADINGMODELID_TWOSIDED_FOLIAGE
+		|| ShadingModel == SHADINGMODELID_HAIR
+		|| ShadingModel == SHADINGMODELID_EYE;
+}
+
+// all values that are output by the forward rendering pass
+struct FGBufferData
+{
+	// normalized
+	float3 WorldNormal;
+	// 0..1 (derived from BaseColor, Metalness, Specular)
+	float3 DiffuseColor;
+	// 0..1 (derived from BaseColor, Metalness, Specular)
+	float3 SpecularColor;
+	// 0..1, white for SHADINGMODELID_SUBSURFACE_PROFILE (apply BaseColor after scattering is more correct and less blurry)
+	float3 BaseColor;
+	// 0..1
+	float Metallic;
+	// 0..1
+	float Specular;
+	// 0..1
+	float4 CustomData;
+	// Indirect irradiance luma
+	float IndirectIrradiance;
+	// Static shadow factors for channels assigned by Lightmass
+	// Lights using static shadowing will pick up the appropriate channel in their deferred pass
+	float4 PrecomputedShadowFactors;
+	// 0..1
+	float Roughness;
+	// 0..1 ambient occlusion  e.g.SSAO, wet surface mask, skylight mask, ...
+	float GBufferAO;
+	// 0..255 
+	uint ShadingModelID;
+	// 0..255 
+	uint SelectiveOutputMask;
+	// 0..1, 2 bits, use HasDistanceFieldRepresentation(GBuffer) or HasDynamicIndirectShadowCasterRepresentation(GBuffer) to extract
+	float PerObjectGBufferData;
+	// in world units
+	float CustomDepth;
+	// Custom depth stencil value
+	uint CustomStencil;
+	// in unreal units (linear), can be used to reconstruct world position,
+	// only valid when decoding the GBuffer as the value gets reconstructed from the Z buffer
+	float Depth;
+	// Velocity for motion blur (only used when WRITES_VELOCITY_TO_GBUFFER is enabled)
+	float4 Velocity;
+
+	// 0..1, only needed by SHADINGMODELID_SUBSURFACE_PROFILE which applies BaseColor later
+	float3 StoredBaseColor;
+	// 0..1, only needed by SHADINGMODELID_SUBSURFACE_PROFILE which applies Specular later
+	float StoredSpecular;
+};
+
+bool HasDistanceFieldRepresentation(FGBufferData GBufferData)
+{
+	uint PackedAlpha = (uint)(GBufferData.PerObjectGBufferData * 3.999f);
+	return PackedAlpha & 1;
+}
+
+bool HasDynamicIndirectShadowCasterRepresentation(FGBufferData GBufferData)
+{
+	uint PackedAlpha = (uint)(GBufferData.PerObjectGBufferData * 3.999f);
+	return (PackedAlpha & 2) != 0;
+}
+
+
+// all values that are output by the forward rendering pass
+struct FDBufferData
+{
+	// 0..1, premultiplied with ColorOpacity
+	float3 PreMulColor;
+	// 0:opaque ..1:see through
+	float ColorOpacity;
+
+	// -1..1, premultiplied with NormalOpacity
+	float3 PreMulWorldNormal;
+	// 0:opaque ..1:see through
+	float NormalOpacity;
+
+	// 0..1, premultiplied with RoughnessOpacity
+	float PreMulRoughness;
+	// 0:opaque ..1:see through
+	float RoughnessOpacity;
+};
+
+struct FScreenSpaceData
+{
+	// GBuffer (material attributes from forward rendering pass)
+	FGBufferData GBuffer;
+	// 0..1, only valid in some passes, 1 if off
+	float AmbientOcclusion;
+
+	float2 DirectionalOcclusion;
+};
+
+/** Populates OutGBufferA, B and C */
+void EncodeGBuffer(
+	FGBufferData GBuffer,
+	out float4 OutGBufferA,
+	out float4 OutGBufferB,
+	out float4 OutGBufferC,
+	out float4 OutGBufferD,
+	out float4 OutGBufferE,
+	out float4 OutGBufferVelocity,
+	float QuantizationBias = 0		// -0.5 to 0.5 random float. Used to bias quantization.
+	)
+{
+	if (GBuffer.ShadingModelID == SHADINGMODELID_UNLIT)
+	{
+		OutGBufferA = 0;
+		OutGBufferB = 0;
+		OutGBufferC = 0;
+		OutGBufferD = 0;
+		OutGBufferE = 0;
+	}
+	else
+	{
+#if 1
+		OutGBufferA.rgb = EncodeNormal( GBuffer.WorldNormal );
+		OutGBufferA.a = GBuffer.PerObjectGBufferData;
+#else
+		float3 Normal = GBuffer.WorldNormal;
+		uint   NormalFace = 0;
+		EncodeNormal( Normal, NormalFace );
+
+		OutGBufferA.rg = Normal.xy;
+		OutGBufferA.b = 0;
+		OutGBufferA.a = GBuffer.PerObjectGBufferData;
+#endif
+
+		OutGBufferB.r = GBuffer.Metallic;
+		OutGBufferB.g = GBuffer.Specular;
+		OutGBufferB.b = GBuffer.Roughness;
+		OutGBufferB.a = EncodeShadingModelIdAndSelectiveOutputMask(GBuffer.ShadingModelID, GBuffer.SelectiveOutputMask);
+
+		OutGBufferC.rgb = EncodeBaseColor( GBuffer.BaseColor );
+
+#if ALLOW_STATIC_LIGHTING
+		// No space for AO. Multiply IndirectIrradiance by AO instead of storing.
+		OutGBufferC.a = EncodeIndirectIrradiance(GBuffer.IndirectIrradiance * GBuffer.GBufferAO) + QuantizationBias * (1.0 / 255.0);
+#else
+		OutGBufferC.a = GBuffer.GBufferAO;
+#endif
+
+		OutGBufferD = GBuffer.CustomData;
+		OutGBufferE = GBuffer.PrecomputedShadowFactors;
+	}
+
+#if WRITES_VELOCITY_TO_GBUFFER
+	OutGBufferVelocity = GBuffer.Velocity;
+#else
+	OutGBufferVelocity = 0;
+#endif
+}
+
+// High frequency Checkerboard pattern
+// @param PixelPos relative to left top of the rendertarget (not viewport)
+// @return true/false, todo: profile if float 0/1 would be better (need to make sure it's 100% the same)
+bool CheckerFromPixelPos(uint2 PixelPos)
+{
+	// todo: Index is float and by staying float we can optimize this 
+	// We alternate the pattern to get 2x supersampling on the lower res data to get more near to full res
+	uint TemporalAASampleIndex = View.TemporalAAParams.x;
+
+#if FEATURE_LEVEL >= FEATURE_LEVEL_SM4
+	return (PixelPos.x + PixelPos.y + TemporalAASampleIndex) % 2;
+#else
+	return (uint)(fmod(PixelPos.x + PixelPos.y + TemporalAASampleIndex, 2)) != 0;
+#endif
+}
+
+// High frequency Checkerboard pattern
+// @param UVSceneColor at pixel center
+// @return true/false, todo: profile if float 0/1 would be better (need to make sure it's 100% the same)
+bool CheckerFromSceneColorUV(float2 UVSceneColor)
+{
+	// relative to left top of the rendertarget (not viewport)
+	uint2 PixelPos = uint2(UVSceneColor * View.BufferSizeAndInvSize.xy);
+
+	return CheckerFromPixelPos(PixelPos);
+}
+
+// SubsurfaceProfile does deferred lighting with a checker board pixel pattern
+// we separate the view from the non view dependent lighting and later recombine the two color constributions in a postprocess
+// We have the option to apply the BaseColor/Specular in the base pass or do it later in the postprocess (has implications to texture detail, fresnel and performance)
+void AdjustBaseColorAndSpecularColorForSubsurfaceProfileLighting(inout float3 BaseColor, inout float3 SpecularColor, inout float Specular, bool bChecker)
+{
+#if SUBSURFACE_CHANNEL_MODE == 0
+	// If SUBSURFACE_CHANNEL_MODE is 0, we can't support full-resolution lighting, so we 
+	// ignore View.bCheckerboardSubsurfaceProfileRendering
+	const bool bCheckerboardRequired = View.bSubsurfacePostprocessEnabled > 0;
+#else
+	const bool bCheckerboardRequired = View.bSubsurfacePostprocessEnabled > 0 && View.bCheckerboardSubsurfaceProfileRendering > 0;
+	BaseColor = View.bSubsurfacePostprocessEnabled ? float3(1, 1, 1) : BaseColor;
+#endif
+	if (bCheckerboardRequired)
+	{
+		// because we adjust the BaseColor here, we need StoredBaseColor
+
+		// we apply the base color later in SubsurfaceRecombinePS()
+		BaseColor = bChecker;
+		// in SubsurfaceRecombinePS() does not multiply with Specular so we do it here
+		SpecularColor *= !bChecker;
+		Specular *= !bChecker;
+	}
+}
+
+/** Populates FGBufferData */
+// @param bChecker High frequency Checkerboard pattern computed with one of the CheckerFrom.. functions, todo: profile if float 0/1 would be better (need to make sure it's 100% the same)
+FGBufferData DecodeGBufferData(
+	float4 InGBufferA,
+	float4 InGBufferB,
+	float4 InGBufferC,
+	float4 InGBufferD,
+	float4 InGBufferE,
+	float4 InGBufferVelocity,
+	float CustomNativeDepth,
+	uint CustomStencil,
+	float SceneDepth,
+	bool bGetNormalizedNormal,
+	bool bChecker)
+{
+	FGBufferData GBuffer;
+
+	GBuffer.WorldNormal = DecodeNormal( InGBufferA.xyz );
+	if(bGetNormalizedNormal)
+	{
+		GBuffer.WorldNormal = normalize(GBuffer.WorldNormal);
+	}
+
+	GBuffer.PerObjectGBufferData = InGBufferA.a;  
+	GBuffer.Metallic	= InGBufferB.r;
+	GBuffer.Specular	= InGBufferB.g;
+	GBuffer.Roughness	= InGBufferB.b;
+	// Note: must match GetShadingModelId standalone function logic
+	// Also Note: SimpleElementPixelShader directly sets SV_Target2 ( GBufferB ) to indicate unlit.
+	// An update there will be required if this layout changes.
+	GBuffer.ShadingModelID = DecodeShadingModelId(InGBufferB.a);
+	GBuffer.SelectiveOutputMask = DecodeSelectiveOutputMask(InGBufferB.a);
+
+	GBuffer.BaseColor = DecodeBaseColor(InGBufferC.rgb);
+
+#if ALLOW_STATIC_LIGHTING
+	GBuffer.GBufferAO = 1;
+	GBuffer.IndirectIrradiance = DecodeIndirectIrradiance(InGBufferC.a);
+#else
+	GBuffer.GBufferAO = InGBufferC.a;
+	GBuffer.IndirectIrradiance = 1;
+#endif
+
+	GBuffer.CustomData = !(GBuffer.SelectiveOutputMask & SKIP_CUSTOMDATA_MASK) ? InGBufferD : 0;
+
+	GBuffer.PrecomputedShadowFactors = !(GBuffer.SelectiveOutputMask & SKIP_PRECSHADOW_MASK) ? InGBufferE :  ((GBuffer.SelectiveOutputMask & ZERO_PRECSHADOW_MASK) ? 0 :  1);
+	GBuffer.CustomDepth = ConvertFromDeviceZ(CustomNativeDepth);
+	GBuffer.CustomStencil = CustomStencil;
+	GBuffer.Depth = SceneDepth;
+
+	GBuffer.StoredBaseColor = GBuffer.BaseColor;
+	GBuffer.StoredSpecular = GBuffer.Specular;
+
+	// derived from BaseColor, Metalness, Specular
+	{
+		GBuffer.SpecularColor = lerp( 0.08 * GBuffer.Specular.xxx, GBuffer.BaseColor, GBuffer.Metallic );
+
+		if (GBuffer.ShadingModelID == SHADINGMODELID_SUBSURFACE_PROFILE)
+		{
+			AdjustBaseColorAndSpecularColorForSubsurfaceProfileLighting(GBuffer.BaseColor, GBuffer.SpecularColor, GBuffer.Specular, bChecker);
+		}
+
+		GBuffer.DiffuseColor = GBuffer.BaseColor - GBuffer.BaseColor * GBuffer.Metallic;
+
+		#if USE_DEVELOPMENT_SHADERS
+		{
+			// this feature is only needed for development/editor - we can compile it out for a shipping build (see r.CompileShadersForDevelopment cvar help)
+			GBuffer.DiffuseColor = GBuffer.DiffuseColor * View.DiffuseOverrideParameter.www + View.DiffuseOverrideParameter.xyz;
+			GBuffer.SpecularColor = GBuffer.SpecularColor * View.SpecularOverrideParameter.w + View.SpecularOverrideParameter.xyz;
+		}
+		#endif //USE_DEVELOPMENT_SHADERS
+	}
+
+	GBuffer.Velocity = !(GBuffer.SelectiveOutputMask & SKIP_VELOCITY_MASK) ? InGBufferVelocity : 0;
+
+	return GBuffer;
+}
+
+
+/** Populates FDBufferData */
+FDBufferData DecodeDBufferData(
+	float4 DBufferA,
+	float4 DBufferB,
+	float2 DBufferC)
+{
+	FDBufferData ret;
+
+	// UNORM 4 channel
+	ret.PreMulColor = DBufferA.rgb;
+	ret.ColorOpacity = DBufferA.a;
+
+	// UNORM 4 channel, 128/255 represents 0
+	ret.PreMulWorldNormal = DBufferB.rgb * 2 - (256.0 / 255.0);
+	ret.NormalOpacity = DBufferB.a;
+
+	// UNORM 2 channel
+	ret.PreMulRoughness = DBufferC.r;
+	ret.RoughnessOpacity = DBufferC.g;
+
+	return ret;
+}
+
+/** Populates DBufferA, DBufferB, DBufferC as float4 and puts opacity in alpha for frame buffer blending */
+// @param MultiOpacity .x: Color, .y:Normal, .z:Roughness
+void EncodeDBufferData(FGBufferData GBufferData, float3 MultiOpacity,
+	out float4 DBufferA,
+	out float4 DBufferB,
+	out float4 DBufferC)
+{
+	// UNORM 4 channel
+	DBufferA = float4(GBufferData.BaseColor, MultiOpacity.x);
+
+	// UNORM 4 channel, 128/255 represents 0
+	DBufferB = float4(GBufferData.WorldNormal * 0.5f + 128.0f/255.0f, MultiOpacity.y);
+
+	// UNORM 2 channel
+	DBufferC = float4(GBufferData.Roughness, 0, 0, MultiOpacity.z);
+
+	// no need to output any data if the channel isn't used anyway (for DX10 we don't adjust the FrameBufferBlend state - it doesn't support different blend modes for different MRT)
+	{
+#if !MATERIAL_DBUFFERA
+		DBufferA = 0;
+#endif
+#if !MATERIAL_DBUFFERB
+		DBufferB = 0;
+#endif
+#if !MATERIAL_DBUFFERC
+		DBufferC = 0;
+#endif
+	}
+}
+
+float3 ExtractSubsurfaceColor(FGBufferData BufferData)
+{
+	return Square(BufferData.CustomData.rgb);
+}
+
+uint ExtractSubsurfaceProfileInt(FGBufferData BufferData)
+{
+	// can be optimized
+	return uint(BufferData.CustomData.r * 255.0f + 0.5f);
+}
+
+/** Populates DBufferA, DBufferB, DBufferC as float4 and puts opacity in alpha for frame buffer blending */
+void ApplyDBufferData(
+	FDBufferData DBufferData, inout float3 WorldNormal, inout float3 SubsurfaceColor, inout float Roughness, 
+	inout float3 BaseColor, inout float Metallic, inout float Specular )
+{
+	WorldNormal = WorldNormal * DBufferData.NormalOpacity + DBufferData.PreMulWorldNormal;
+	Roughness = Roughness * DBufferData.RoughnessOpacity + DBufferData.PreMulRoughness;
+	
+	SubsurfaceColor *= DBufferData.ColorOpacity;
+
+	BaseColor = BaseColor * DBufferData.ColorOpacity + DBufferData.PreMulColor;
+	Metallic = Metallic * DBufferData.ColorOpacity + 0;			// decals are always no metallic
+	Specular = lerp(0.5f, Specular, DBufferData.ColorOpacity);		// most non metal materials have a specular of 4% which is 0.5 in this scale
+}
+
+Texture2D DBufferATexture;
+SamplerState DBufferATextureSampler;
+Texture2D DBufferBTexture;
+SamplerState DBufferBTextureSampler;
+Texture2D DBufferCTexture;
+SamplerState DBufferCTextureSampler;
+
+Texture2D ScreenSpaceAOTexture;
+SamplerState ScreenSpaceAOTextureSampler;
+
+#if FEATURE_LEVEL >= FEATURE_LEVEL_SM4 
+Texture2D<uint2> CustomStencilTexture;
+#endif
+
+Texture2D		DirectionalOcclusionTexture;
+SamplerState	DirectionalOcclusionSampler;
+
+#if FEATURE_LEVEL >= FEATURE_LEVEL_SM4
+	// In all but SM5 we need to explicitly declare how many samples are in a multisampled texture.
+	#if FEATURE_LEVEL >= FEATURE_LEVEL_SM5
+		#define FMultisampledGBuffer Texture2DMS<float4>
+	#else
+		#define FMultisampledGBuffer Texture2DMS<float4,MSAA_SAMPLE_COUNT>
+	#endif
+
+	Texture2D<float4> DBufferATextureNonMS;
+	Texture2D<float4> DBufferBTextureNonMS;
+	Texture2D<float2> DBufferCTextureNonMS;
+	Texture2D<float4> ScreenSpaceAOTextureNonMS;
+	Texture2D<float> CustomDepthTextureNonMS;
+
+	// @param PixelPos relative to left top of the rendertarget (not viewport)
+	FGBufferData GetGBufferDataUint(uint2 PixelPos, bool bGetNormalizedNormal = true)
+	{
+		float4 GBufferA = GBuffers.GBufferATextureNonMS.Load(int3(PixelPos, 0));
+		float4 GBufferB = GBuffers.GBufferBTextureNonMS.Load(int3(PixelPos, 0));
+		float4 GBufferC = GBuffers.GBufferCTextureNonMS.Load(int3(PixelPos, 0));
+		float4 GBufferD = GBuffers.GBufferDTextureNonMS.Load(int3(PixelPos, 0));
+		float CustomNativeDepth = CustomDepthTextureNonMS.Load(int3(PixelPos, 0)).r;
+		uint CustomStencil = CustomStencilTexture.Load(int3(PixelPos, 0)) STENCIL_COMPONENT_SWIZZLE;
+
+		#if ALLOW_STATIC_LIGHTING
+			float4 GBufferE = GBuffers.GBufferETextureNonMS.Load(int3(PixelPos, 0));
+		#else
+			float4 GBufferE = 1;
+		#endif
+
+		#if WRITES_VELOCITY_TO_GBUFFER
+			float4 GBufferVelocity = GBuffers.GBufferVelocityTextureNonMS.Load(int3(PixelPos, 0));
+		#else
+			float4 GBufferVelocity = 0;
+		#endif
+
+		float SceneDepth = CalcSceneDepth(PixelPos);
+
+		return DecodeGBufferData(GBufferA, GBufferB, GBufferC, GBufferD, GBufferE, GBufferVelocity, CustomNativeDepth, CustomStencil, SceneDepth, bGetNormalizedNormal, CheckerFromPixelPos(PixelPos));
+	}
+
+	// @param PixelPos relative to left top of the rendertarget (not viewport)
+	FDBufferData GetDBufferData(uint2 PixelPos)
+	{
+		float4 DBufferA = DBufferATextureNonMS.Load(int3(PixelPos, 0));
+		float4 DBufferB = DBufferBTextureNonMS.Load(int3(PixelPos, 0));
+		float2 DBufferC = DBufferCTextureNonMS.Load(int3(PixelPos, 0)).rg;
+
+		return DecodeDBufferData(DBufferA, DBufferB, DBufferC);
+	}
+
+	// @param PixelPos relative to left top of the rendertarget (not viewport)
+	FScreenSpaceData GetScreenSpaceDataUint(uint2 PixelPos, bool bGetNormalizedNormal = true)
+	{
+		FScreenSpaceData Out;
+
+		Out.GBuffer = GetGBufferDataUint(PixelPos, bGetNormalizedNormal);
+
+		// todo: optimize
+		// this is what we want but because WhiteDummy (in case AO is disabled) doesn't support this lookup we do the code below
+//		Out.AmbientOcclusion = ScreenSpaceAOTextureNonMS.Load(int3(PixelPos, 0)).r;
+
+		float4 ScreenSpaceAO = Texture2DSampleLevel(ScreenSpaceAOTexture, ScreenSpaceAOTextureSampler, (PixelPos + 0.5f) * View.BufferSizeAndInvSize.zw, 0);
+		Out.AmbientOcclusion = ScreenSpaceAO.r;
+
+		Out.DirectionalOcclusion = Texture2DSampleLevel( DirectionalOcclusionTexture, DirectionalOcclusionSampler, (PixelPos + 0.5f) * View.BufferSizeAndInvSize.zw, 0 ).rg;
+
+		return Out;
+	}
+#endif
+
+// @param UV - UV space in the GBuffer textures (BufferSize resolution)
+FGBufferData GetGBufferData(float2 UV, bool bGetNormalizedNormal = true)
+{
+#if 0 //METAL_MRT_PROFILE
+	// @todo metal mrt: The deferred renderer isn't keeping these in tiled memory all the time - we don't know when this makes sense
+	// versus just sampling a bound resolved texture
+	float4 GBufferA = FramebufferFetchMRT(1);
+	float4 GBufferB = FramebufferFetchMRT(2);
+	float4 GBufferC = FramebufferFetchMRT(3);
+	float4 GBufferD = FramebufferFetchMRT(4);
+	// @todo metal mrt: We can't framebuffer fetch the depth, can we jam it in somewhere?
+	float CustomNativeDepth = 0.5;
+	#if ALLOW_STATIC_LIGHTING
+		float4 GBufferE = FramebufferFetchMRT(5);
+	#else
+		float4 GBufferE = 1;
+	#endif
+#else	
+	float4 GBufferA = Texture2DSampleLevel(GBuffers.GBufferATexture, GBuffers.GBufferATextureSampler, UV, 0);
+	float4 GBufferB = Texture2DSampleLevel(GBuffers.GBufferBTexture, GBuffers.GBufferBTextureSampler, UV, 0);
+	float4 GBufferC = Texture2DSampleLevel(GBuffers.GBufferCTexture, GBuffers.GBufferCTextureSampler, UV, 0);
+	float4 GBufferD = Texture2DSampleLevel(GBuffers.GBufferDTexture, GBuffers.GBufferDTextureSampler, UV, 0);
+	float CustomNativeDepth = Texture2DSampleLevel(CustomDepthTexture, CustomDepthTextureSampler, UV, 0).r;
+
+#if FEATURE_LEVEL >= FEATURE_LEVEL_SM4
+	int2 IntUV = (int2)trunc(UV * View.BufferSizeAndInvSize.xy);
+	uint CustomStencil = CustomStencilTexture.Load(int3(IntUV, 0)) STENCIL_COMPONENT_SWIZZLE;
+#else
+	uint CustomStencil = 0;
+#endif
+
+	#if ALLOW_STATIC_LIGHTING
+		float4 GBufferE = Texture2DSampleLevel(GBuffers.GBufferETexture, GBuffers.GBufferETextureSampler, UV, 0);
+	#else
+		float4 GBufferE = 1;
+	#endif
+
+	#if WRITES_VELOCITY_TO_GBUFFER
+		float4 GBufferVelocity = Texture2DSampleLevel(GBuffers.GBufferVelocityTexture, GBuffers.GBufferVelocityTextureSampler, UV, 0);
+	#else
+		float4 GBufferVelocity = 0;
+	#endif
+#endif
+
+	float SceneDepth = CalcSceneDepth(UV);
+	
+	return DecodeGBufferData(GBufferA, GBufferB, GBufferC, GBufferD, GBufferE, GBufferVelocity, CustomNativeDepth, CustomStencil, SceneDepth, bGetNormalizedNormal, CheckerFromSceneColorUV(UV));
+}
+
+// Minimal path for just the lighting model, used to branch around unlit pixels (skybox)
+uint GetShadingModelId(float2 UV)
+{
+	return DecodeShadingModelId(Texture2DSampleLevel(GBuffers.GBufferBTexture, GBuffers.GBufferBTextureSampler, UV, 0).a);
+}
+
+// @param UV - UV space in the DBuffer textures (BufferSize resolution)
+FDBufferData GetDBufferData(float2 UV)
+{
+#if SUPPORTS_INDEPENDENT_SAMPLERS
+	// Share samplers to give artists more
+	float4 DBufferA = Texture2DSampleLevel(DBufferATexture, DBufferATextureSampler, UV, 0);
+	float4 DBufferB = Texture2DSampleLevel(DBufferBTexture, DBufferATextureSampler, UV, 0);
+	float2 DBufferC = Texture2DSampleLevel(DBufferCTexture, DBufferATextureSampler, UV, 0).rg;
+#else
+	float4 DBufferA = Texture2DSampleLevel(DBufferATexture, DBufferATextureSampler, UV, 0);
+	float4 DBufferB = Texture2DSampleLevel(DBufferBTexture, DBufferBTextureSampler, UV, 0);
+	float2 DBufferC = Texture2DSampleLevel(DBufferCTexture, DBufferCTextureSampler, UV, 0).rg;
+#endif
+	
+	return DecodeDBufferData(DBufferA, DBufferB, DBufferC);
+}
+
+// @param UV - UV space in the GBuffer textures (BufferSize resolution)
+FScreenSpaceData GetScreenSpaceData(float2 UV, bool bGetNormalizedNormal = true)
+{
+	FScreenSpaceData Out;
+
+	Out.GBuffer = GetGBufferData(UV, bGetNormalizedNormal);
+	float4 ScreenSpaceAO = Texture2DSampleLevel(ScreenSpaceAOTexture, ScreenSpaceAOTextureSampler, UV, 0);
+
+	Out.AmbientOcclusion = ScreenSpaceAO.r;
+
+	Out.DirectionalOcclusion = Texture2DSampleLevel( DirectionalOcclusionTexture, DirectionalOcclusionSampler, UV, 0 ).rg;
+	return Out;
+}
+
+
+#if FEATURE_LEVEL >= FEATURE_LEVEL_SM5
+	// @param PixelPos relative to left top of the rendertarget (not viewport)
+	FGBufferData GetGBufferDataMS(int2 PixelPos, uint SampleIndex, bool bGetNormalizedNormal = true)
+	{
+		float4 GBufferA = GBuffers.GBufferATextureMS.Load(PixelPos, SampleIndex);
+		float4 GBufferB = GBuffers.GBufferBTextureMS.Load(PixelPos, SampleIndex);
+		float4 GBufferC = GBuffers.GBufferCTextureMS.Load(PixelPos, SampleIndex);
+		float4 GBufferD = GBuffers.GBufferDTextureMS.Load(PixelPos, SampleIndex);
+		float CustomNativeDepth = CustomDepthTextureNonMS.Load(int3(PixelPos, 0)).r;
+		uint CustomStencil = CustomStencilTexture.Load(int3(PixelPos, 0)) STENCIL_COMPONENT_SWIZZLE;
+
+		#if ALLOW_STATIC_LIGHTING
+			float4 GBufferE = GBuffers.GBufferETextureMS.Load(PixelPos, SampleIndex);
+		#else
+			float4 GBufferE = 1;
+		#endif
+
+		float DeviceZ = SceneDepthSurface.Load(PixelPos, SampleIndex);
+		float SceneDepth = ConvertFromDeviceZ(DeviceZ);
+
+	#if WRITES_VELOCITY_TO_GBUFFER
+		float4 GBufferVelocity = GBuffers.GBufferVelocityTextureMS.Load(PixelPos, SampleIndex);
+	#else
+		float4 GBufferVelocity = 0;
+	#endif
+
+		return DecodeGBufferData(GBufferA, GBufferB, GBufferC, GBufferD, GBufferE, GBufferVelocity, CustomNativeDepth, CustomStencil, SceneDepth, bGetNormalizedNormal, CheckerFromPixelPos(PixelPos));
+	}
+
+	FGBufferData GetGBufferDataMS(float2 UV, uint SampleIndex, bool bGetNormalizedNormal = true)
+	{
+		int2 IntUV = (int2)trunc(UV * View.BufferSizeAndInvSize.xy);
+
+		return GetGBufferDataMS(IntUV, SampleIndex, bGetNormalizedNormal);
+	}
+#endif
+
+
+// for debugging and to visualize
+float3 GetShadingModelColor(uint ShadingModelID)
+{
+	// TODO: PS4 doesn't optimize out correctly the switch(), so it thinks it needs all the Samplers even if they get compiled out
+	//	This will get fixed after launch per Sony...
+#if PS4_PROFILE
+		 if (ShadingModelID == SHADINGMODELID_UNLIT) return float3(0.1f, 0.1f, 0.2f); // Dark Blue
+	else if (ShadingModelID == SHADINGMODELID_DEFAULT_LIT) return float3(0.1f, 1.0f, 0.1f); // Green
+	else if (ShadingModelID == SHADINGMODELID_SUBSURFACE) return float3(1.0f, 0.1f, 0.1f); // Red
+	else if (ShadingModelID == SHADINGMODELID_PREINTEGRATED_SKIN) return float3(0.6f, 0.4f, 0.1f); // Brown
+	else if (ShadingModelID == SHADINGMODELID_CLEAR_COAT) return float3(0.1f, 0.4f, 0.4f); 
+	else if (ShadingModelID == SHADINGMODELID_SUBSURFACE_PROFILE) return float3(0.2f, 0.6f, 0.5f); // Cyan
+	else if (ShadingModelID == SHADINGMODELID_TWOSIDED_FOLIAGE) return float3(0.2f, 0.2f, 0.8f); // Blue
+	else if (ShadingModelID == SHADINGMODELID_HAIR) return float3(0.6f, 0.1f, 0.5f);
+	else if (ShadingModelID == SHADINGMODELID_CLOTH) return float3(0.7f, 1.0f, 1.0f); 
+	else if (ShadingModelID == SHADINGMODELID_EYE) return float3(0.3f, 1.0f, 1.0f); 
+	else return float3(1.0f, 1.0f, 1.0f); // White
+#else
+	switch(ShadingModelID)
+	{
+		case SHADINGMODELID_UNLIT: return float3(0.1f, 0.1f, 0.2f); // Dark Blue
+		case SHADINGMODELID_DEFAULT_LIT: return float3(0.1f, 1.0f, 0.1f); // Green
+		case SHADINGMODELID_SUBSURFACE: return float3(1.0f, 0.1f, 0.1f); // Red
+		case SHADINGMODELID_PREINTEGRATED_SKIN: return float3(0.6f, 0.4f, 0.1f); // Brown
+		case SHADINGMODELID_CLEAR_COAT: return float3(0.1f, 0.4f, 0.4f); // Brown
+		case SHADINGMODELID_SUBSURFACE_PROFILE: return float3(0.2f, 0.6f, 0.5f); // Cyan
+		case SHADINGMODELID_TWOSIDED_FOLIAGE: return float3(0.2f, 0.2f, 0.8f); // Cyan
+		case SHADINGMODELID_HAIR: return float3(0.6f, 0.1f, 0.5f);
+		case SHADINGMODELID_CLOTH: return float3(0.7f, 1.0f, 1.0f);
+		case SHADINGMODELID_EYE: return float3(0.3f, 1.0f, 1.0f);
+		default: return float3(1.0f, 1.0f, 1.0f); // White
+	}
+#endif
+}
+
+#endif // __DEFERRED_SHADING_COMMON__
diff --git a/Engine/Shaders/Private/DistanceFieldAOShared.ush b/Engine/Shaders/Private/DistanceFieldAOShared.ush
new file mode 100644
--- /dev/null
+++ b/Engine/Shaders/Private/DistanceFieldAOShared.ush
@@ -0,0 +1,231 @@
+// Copyright 1998-2017 Epic Games, Inc. All Rights Reserved.
+
+/*=============================================================================
+	DistanceFieldAOShared.usf
+=============================================================================*/
+
+#ifndef SORTTHREADGROUP_SIZEX
+#define SORTTHREADGROUP_SIZEX 1
+#endif
+
+#ifndef SORTTHREADGROUP_SIZEY
+#define SORTTHREADGROUP_SIZEY 1
+#endif
+
+#ifndef CONE_TRACE_PHASE
+#define CONE_TRACE_PHASE -1
+#endif
+
+// Must match C++
+#define NUM_CONE_STEPS 10
+
+// Must match C++
+#define NUM_CONE_DIRECTIONS 9
+
+// Must match C++
+#define AO_DOWNSAMPLE_FACTOR 2
+
+float AOObjectMaxDistance;
+float AOStepScale;
+float AOStepExponentScale;
+float AOMaxViewDistance;
+
+float AOGlobalMaxOcclusionDistance;
+
+// Gives the max occlusion distance regardless of whether the global distance field is in use
+float GetAOMaxDistance()
+{
+	return max(AOObjectMaxDistance, AOGlobalMaxOcclusionDistance);
+}
+
+// Pass a debug value through the pipeline instead of a bent normal
+#define PASS_THROUGH_DEBUG_VALUE 0
+
+// One for near, one for far
+#define NUM_CULLED_OBJECT_LISTS 2
+
+float GetStepOffset(float StepIndex)
+{
+	// Original heuristic
+	//return AOStepScale * exp2(AOStepExponentScale * StepIndex);
+
+	float temp = AOStepExponentScale * StepIndex;
+	return AOStepScale * (temp * temp + 1);
+}
+
+uint2 TileListGroupSize;
+ 
+Texture2D DistanceFieldNormalTexture;
+SamplerState DistanceFieldNormalSampler;
+
+float4 EncodeDownsampledGBuffer(FGBufferData GBufferData, float SceneDepth)
+{
+	float DepthSign = HasDistanceFieldRepresentation(GBufferData) > 0 ? 1 : -1;
+	float ZSign = GBufferData.ShadingModelID == SHADINGMODELID_UNLIT ? -1 : 1;
+	float ZChannel = GBufferData.WorldNormal.z > 0 ? .5f : 1;
+	float4 OutValue = float4(GBufferData.WorldNormal.xy, ZChannel * ZSign, DepthSign * SceneDepth);
+	return OutValue;
+}
+
+void GetDownsampledGBuffer(float2 ScreenUV, out float3 OutNormal, out float OutDepth, out bool bHasDistanceFieldRepresentation, out bool bLit)
+{
+	float4 TextureValue = Texture2DSampleLevel(DistanceFieldNormalTexture, DistanceFieldNormalSampler, ScreenUV, 0);
+	bHasDistanceFieldRepresentation = TextureValue.w > 0;
+	bLit = TextureValue.z > 0; 
+	OutDepth = abs(TextureValue.w);
+	OutNormal = float3(TextureValue.xy, sqrt(1 - dot(TextureValue.xy, TextureValue.xy)));
+	OutNormal.z *= abs(TextureValue.z) > .75f ? -1 : 1;
+}
+
+float GetDownsampledDepth(float2 ScreenUV)
+{
+	return abs(Texture2DSampleLevel(DistanceFieldNormalTexture, DistanceFieldNormalSampler, ScreenUV, 0).w);
+}
+
+uint CurrentLevelDownsampleFactor;
+float2 AOBufferSize;
+
+uint DownsampleFactorToBaseLevel;
+float2 BaseLevelTexelSize;
+
+Texture2D BentNormalAOTexture;
+SamplerState BentNormalAOSampler;
+
+Texture2D IrradianceTexture;
+SamplerState IrradianceSampler;
+
+bool SphereIntersectCone(float4 SphereCenterAndRadius, float3 ConeVertex, float3 ConeAxis, float ConeAngleCos, float ConeAngleSin)
+{
+	float3 U = ConeVertex - (SphereCenterAndRadius.w / ConeAngleSin) * ConeAxis;
+	float3 D = SphereCenterAndRadius.xyz - U;
+	float DSizeSq = dot(D, D);
+	float E = dot(ConeAxis, D);
+
+	if (E > 0 && E * E >= DSizeSq * ConeAngleCos * ConeAngleCos)
+	{
+		D = SphereCenterAndRadius.xyz - ConeVertex;
+		DSizeSq = dot(D, D);
+		E = -dot(ConeAxis, D);
+
+		if (E > 0 && E * E >= DSizeSq * ConeAngleSin * ConeAngleSin)
+		{
+			return DSizeSq <= SphereCenterAndRadius.w * SphereCenterAndRadius.w;
+		}
+		else
+		{
+			return true;
+		}
+	}
+
+	return false;
+}
+
+bool SphereIntersectConeWithDepthRanges(float4 SphereCenterAndRadius, float3 ConeVertex, float3 ConeAxis, float ConeAngleCos, float ConeAngleSin, float4 ConeAxisDepthRanges)
+{
+	if (SphereIntersectCone(SphereCenterAndRadius, ConeVertex, ConeAxis, ConeAngleCos, ConeAngleSin))
+	{
+		float ConeAxisDistance = dot(SphereCenterAndRadius.xyz - ConeVertex, ConeAxis);
+		float2 ConeAxisDistanceMinMax = float2(ConeAxisDistance + SphereCenterAndRadius.w, ConeAxisDistance - SphereCenterAndRadius.w);
+
+		if (ConeAxisDistanceMinMax.x > ConeAxisDepthRanges.x && ConeAxisDistanceMinMax.y < ConeAxisDepthRanges.y
+			|| ConeAxisDistanceMinMax.x > ConeAxisDepthRanges.z && ConeAxisDistanceMinMax.y < ConeAxisDepthRanges.w)
+		{
+			return true;
+		}
+	}
+
+	return false;
+}
+
+Buffer<float> RecordConeVisibility;
+float BentNormalNormalizeFactor;
+
+void FindBestAxisVectors2(float3 InZAxis, out float3 OutXAxis, out float3 OutYAxis )
+{
+	float3 UpVector = abs(InZAxis.z) < 0.999 ? float3(0,0,1) : float3(1,0,0);
+	OutXAxis = normalize( cross( UpVector, InZAxis ) );
+	OutYAxis = cross( InZAxis, OutXAxis );
+}
+
+float3 ComputeBentNormal(float3 RecordWorldNormal, uint RelativeRecordIndex)
+{
+	float3 TangentX;
+	float3 TangentY;
+	FindBestAxisVectors2(RecordWorldNormal, TangentX, TangentY);
+
+	float3 UnoccludedDirection = 0;
+
+	for (uint ConeIndex = 0; ConeIndex < NUM_CONE_DIRECTIONS; ConeIndex++)
+	{
+		float3 ConeDirection = AOSamples2.SampleDirections[ConeIndex].xyz;
+		float3 RotatedConeDirection = ConeDirection.x * TangentX + ConeDirection.y * TangentY + ConeDirection.z * RecordWorldNormal;
+
+		float ConeVisibility = RecordConeVisibility[RelativeRecordIndex * NUM_CONE_DIRECTIONS + ConeIndex];
+		UnoccludedDirection += ConeVisibility * RotatedConeDirection;
+	}
+
+	float InvNumSamples = 1.0f / (float)NUM_CONE_DIRECTIONS;
+	UnoccludedDirection = UnoccludedDirection * (BentNormalNormalizeFactor * InvNumSamples);
+
+	return UnoccludedDirection;
+}
+
+#ifndef CULLED_TILE_SIZEX
+#define CULLED_TILE_SIZEX 4
+#endif
+ 
+#ifndef TRACE_DOWNSAMPLE_FACTOR
+#define TRACE_DOWNSAMPLE_FACTOR 1
+#endif
+ 
+#ifndef CONE_TRACE_OBJECTS_THREADGROUP_SIZE
+#define CONE_TRACE_OBJECTS_THREADGROUP_SIZE 16
+#endif
+
+// Size of a culled tile at the resolution that cone tracing is done, in one dimension
+#define CONE_TILE_SIZEX (CULLED_TILE_SIZEX / TRACE_DOWNSAMPLE_FACTOR) 
+// Number of culled tiles per cone tracing threadgroup
+#define CONE_TRACE_TILES_PER_THREADGROUP (CONE_TRACE_OBJECTS_THREADGROUP_SIZE / (CONE_TILE_SIZEX * CONE_TILE_SIZEX))
+
+uint2 ScreenGridConeVisibilitySize;
+float2 JitterOffset;
+
+uint2 ComputeTileCoordinateFromScreenGrid(uint2 OutputCoordinate)
+{
+	uint2 TileCoordinate = OutputCoordinate * TRACE_DOWNSAMPLE_FACTOR / CULLED_TILE_SIZEX;
+	return TileCoordinate;
+}
+
+float2 GetBaseLevelScreenUVFromScreenGrid(uint2 OutputCoordinate, float JitterScale)
+{
+	float2 BaseLevelScreenUV = (OutputCoordinate * TRACE_DOWNSAMPLE_FACTOR + JitterOffset * JitterScale + float2(.5f, .5f)) * BaseLevelTexelSize;
+	return BaseLevelScreenUV;
+}
+
+float2 GetBaseLevelScreenUVFromScreenGrid(uint2 OutputCoordinate)
+{
+	return GetBaseLevelScreenUVFromScreenGrid(OutputCoordinate, 1);
+}
+
+float2 GetScreenUVFromScreenGrid(uint2 OutputCoordinate, float JitterScale)
+{
+	float2 ScreenUV = ((OutputCoordinate * TRACE_DOWNSAMPLE_FACTOR + JitterOffset * JitterScale) * AO_DOWNSAMPLE_FACTOR + View.ViewRectMin.xy + float2(.5f, .5f)) * View.BufferSizeAndInvSize.zw;
+	return ScreenUV;
+}
+
+float2 GetScreenUVFromScreenGrid(uint2 OutputCoordinate)
+{
+	return GetScreenUVFromScreenGrid(OutputCoordinate, 1);
+}
+
+Buffer<uint> CulledTilesStartOffsetArray;
+
+RWBuffer<uint> RWNumCulledTilesArray;
+RWBuffer<uint> RWCulledTilesStartOffsetArray;
+RWBuffer<uint> RWCulledTileDataArray;
+
+#define INVALID_TILE_INDEX 0xFFFF
+
+#ifndef CULLED_TILE_DATA_STRIDE
+#define CULLED_TILE_DATA_STRIDE 1
+#endif
\ No newline at end of file
diff --git a/Engine/Shaders/Private/DistanceFieldLightingShared.ush b/Engine/Shaders/Private/DistanceFieldLightingShared.ush
new file mode 100644
--- /dev/null
+++ b/Engine/Shaders/Private/DistanceFieldLightingShared.ush
@@ -0,0 +1,350 @@
+// Copyright 1998-2017 Epic Games, Inc. All Rights Reserved.
+
+/*=============================================================================
+	DistanceFieldLightingShared.usf
+=============================================================================*/
+
+#include "SurfelMaterialShared.ush"
+
+#ifndef THREADGROUP_SIZEX
+#define THREADGROUP_SIZEX 1
+#endif   
+
+#ifndef THREADGROUP_SIZEY
+#define THREADGROUP_SIZEY 1
+#endif
+
+#define THREADGROUP_TOTALSIZE (THREADGROUP_SIZEX * THREADGROUP_SIZEY)
+
+#ifndef DOWNSAMPLE_FACTOR
+#define DOWNSAMPLE_FACTOR 1
+#endif
+
+#ifndef UPDATEOBJECTS_THREADGROUP_SIZE
+#define UPDATEOBJECTS_THREADGROUP_SIZE 1
+#endif
+
+float3 DistanceFieldVolumePositionToUV(float3 VolumePosition, float3 UVScale, float3 UVAdd)
+{
+	float3 VolumeUV = VolumePosition * UVScale + UVAdd;
+	return VolumeUV;
+}
+
+Texture3D DistanceFieldTexture;
+SamplerState DistanceFieldSampler;
+
+float SampleMeshDistanceField(float3 VolumeUV, float2 ObjectDistanceFieldMAD)
+{
+	float TextureValue = Texture3DSampleLevel(DistanceFieldTexture, DistanceFieldSampler, VolumeUV, 0);
+#if EIGHT_BIT_MESH_DISTANCE_FIELDS
+	TextureValue = TextureValue * ObjectDistanceFieldMAD.x + ObjectDistanceFieldMAD.y;
+#endif
+	return TextureValue;
+}
+
+float3 DistanceFieldAtlasTexelSize;
+
+RWBuffer<uint> RWObjectIndirectArguments;
+Buffer<uint> ObjectIndirectArguments;
+
+uint GetCulledNumObjects()
+{
+	// IndexCount, NumInstances, StartIndex, BaseVertexIndex, FirstInstance
+	return ObjectIndirectArguments[1];
+}
+
+// In float4's.  Must match equivalent C++ variables.
+#define OBJECT_DATA_STRIDE 18
+
+uint NumSceneObjects;
+
+// Have to make these R32F with 4x the reads and writes because of the horrible D3D11 limitation
+// "error X3676: typed UAV loads are only allowed for single-component 32-bit element types"
+Buffer<float> ObjectBounds;
+Buffer<float> ObjectData;
+
+RWBuffer<float> RWObjectBounds;
+RWBuffer<float> RWObjectData;
+
+float4 LoadGlobalObjectPositionAndRadius(uint ObjectIndex)
+{
+	uint VectorIndex = ObjectIndex * 4;
+	return float4(ObjectBounds[VectorIndex + 0], ObjectBounds[VectorIndex + 1], ObjectBounds[VectorIndex + 2], ObjectBounds[VectorIndex + 3]);
+}
+
+float4x4 LoadGlobalObjectWorldToVolume(uint ObjectIndex)
+{
+	uint VectorIndex = (ObjectIndex * OBJECT_DATA_STRIDE + 0) * 4;
+	float4 M0 = float4(ObjectData[VectorIndex + 0], ObjectData[VectorIndex + 1], ObjectData[VectorIndex + 2], ObjectData[VectorIndex + 3]);
+	VectorIndex = (ObjectIndex * OBJECT_DATA_STRIDE + 1) * 4;
+	float4 M1 = float4(ObjectData[VectorIndex + 0], ObjectData[VectorIndex + 1], ObjectData[VectorIndex + 2], ObjectData[VectorIndex + 3]);
+	VectorIndex = (ObjectIndex * OBJECT_DATA_STRIDE + 2) * 4;
+	float4 M2 = float4(ObjectData[VectorIndex + 0], ObjectData[VectorIndex + 1], ObjectData[VectorIndex + 2], ObjectData[VectorIndex + 3]);
+	VectorIndex = (ObjectIndex * OBJECT_DATA_STRIDE + 3) * 4;
+	float4 M3 = float4(ObjectData[VectorIndex + 0], ObjectData[VectorIndex + 1], ObjectData[VectorIndex + 2], ObjectData[VectorIndex + 3]);
+
+	return float4x4(M0, M1, M2, M3);
+}
+ 
+float3 LoadGlobalObjectLocalPositionExtent(uint ObjectIndex)
+{
+	uint VectorIndex = (ObjectIndex * OBJECT_DATA_STRIDE + 4) * 4;
+	return float3(ObjectData[VectorIndex + 0], ObjectData[VectorIndex + 1], ObjectData[VectorIndex + 2]);
+}
+
+float4 LoadGlobalObjectUVScale(uint ObjectIndex)
+{
+	uint VectorIndex = (ObjectIndex * OBJECT_DATA_STRIDE + 5) * 4;
+	return float4(ObjectData[VectorIndex + 0], ObjectData[VectorIndex + 1], ObjectData[VectorIndex + 2], ObjectData[VectorIndex + 3]);
+}
+
+float3 LoadGlobalObjectUVAdd(uint ObjectIndex)
+{
+	uint VectorIndex = (ObjectIndex * OBJECT_DATA_STRIDE + 6) * 4;
+	return float3(ObjectData[VectorIndex + 0], ObjectData[VectorIndex + 1], ObjectData[VectorIndex + 2]);
+}
+
+float2 LoadGlobalObjectDistanceFieldMAD(uint ObjectIndex)
+{
+	uint VectorIndex = (ObjectIndex * OBJECT_DATA_STRIDE + 7) * 4;
+	return float2(ObjectData[VectorIndex + 0], ObjectData[VectorIndex + 1]);
+}
+
+void LoadGlobalObjectLocalVolumeBoundsMinMax(uint ObjectIndex, out float3 LocalVolumeBoundsMin, out float3 LocalVolumeBoundsMax)
+{
+	uint VectorIndex = 4 * (ObjectIndex * OBJECT_DATA_STRIDE + 16);
+	LocalVolumeBoundsMin = float3(ObjectData[VectorIndex + 0], ObjectData[VectorIndex + 1], ObjectData[VectorIndex + 2]);
+
+	VectorIndex = 4 * (ObjectIndex * OBJECT_DATA_STRIDE + 17);
+	LocalVolumeBoundsMax = float3(ObjectData[VectorIndex + 0], ObjectData[VectorIndex + 1], ObjectData[VectorIndex + 2]);
+}
+
+uint LoadGlobalObjectOftenMoving(uint ObjectIndex)
+{
+	uint VectorIndex = (ObjectIndex * OBJECT_DATA_STRIDE + 17) * 4;
+	return ObjectData[VectorIndex + 3] > 0.0f ? 1 : 0;
+}
+
+// In float4's.  Must match equivalent C++ variables.
+#define CULLED_OBJECT_DATA_STRIDE 16
+#define CULLED_OBJECT_BOX_BOUNDS_STRIDE 5
+
+float4 LoadObjectPositionAndRadiusFromBuffer(uint ObjectIndex, StructuredBuffer<float4> InBuffer)
+{
+	return InBuffer[ObjectIndex];
+}
+
+float4x4 LoadObjectWorldToVolumeFromBuffer(uint ObjectIndex, StructuredBuffer<float4> InBuffer)
+{
+	float4 M0 = InBuffer[ObjectIndex * CULLED_OBJECT_DATA_STRIDE + 0];
+	float4 M1 = InBuffer[ObjectIndex * CULLED_OBJECT_DATA_STRIDE + 1];
+	float4 M2 = InBuffer[ObjectIndex * CULLED_OBJECT_DATA_STRIDE + 2];
+	float4 M3 = InBuffer[ObjectIndex * CULLED_OBJECT_DATA_STRIDE + 3];
+
+	return float4x4(M0, M1, M2, M3);
+}
+ 
+float3 LoadObjectLocalPositionExtentFromBuffer(uint ObjectIndex, StructuredBuffer<float4> InBuffer)
+{
+	return InBuffer[ObjectIndex * CULLED_OBJECT_DATA_STRIDE + 4].xyz;
+}
+
+float4 LoadObjectUVScaleFromBuffer(uint ObjectIndex, StructuredBuffer<float4> InBuffer, out bool bGeneratedAsTwoSided)
+{
+	float4 Value = InBuffer[ObjectIndex * CULLED_OBJECT_DATA_STRIDE + 5].xyzw;
+	bGeneratedAsTwoSided = Value.w < 0;
+	return float4(Value.xyz, abs(Value.w));
+}
+
+float4 LoadObjectUVAddAndSelfShadowBiasFromBuffer(uint ObjectIndex, StructuredBuffer<float4> InBuffer)
+{
+	return InBuffer[ObjectIndex * CULLED_OBJECT_DATA_STRIDE + 6];
+}
+
+float2 LoadObjectDistanceFieldMADFromBuffer(uint ObjectIndex, StructuredBuffer<float4> InBuffer)
+{
+	return InBuffer[ObjectIndex * CULLED_OBJECT_DATA_STRIDE + 7].xy;
+}
+
+float3x3 LoadObjectVolumeToWorldFromBuffer(uint ObjectIndex, StructuredBuffer<float4> InBuffer)
+{
+	float3 M0 = InBuffer[ObjectIndex * CULLED_OBJECT_DATA_STRIDE + 8].xyz;
+	float3 M1 = InBuffer[ObjectIndex * CULLED_OBJECT_DATA_STRIDE + 9].xyz;
+	float3 M2 = InBuffer[ObjectIndex * CULLED_OBJECT_DATA_STRIDE + 10].xyz;
+
+	return float3x3(M0, M1, M2);
+}
+
+float4x4 LoadObjectLocalToWorldFromBuffer(uint ObjectIndex, StructuredBuffer<float4> InBuffer)
+{
+	float4 M0 = InBuffer[ObjectIndex * CULLED_OBJECT_DATA_STRIDE + 11];
+	float4 M1 = InBuffer[ObjectIndex * CULLED_OBJECT_DATA_STRIDE + 12];
+	float4 M2 = InBuffer[ObjectIndex * CULLED_OBJECT_DATA_STRIDE + 13];
+	float4 M3 = InBuffer[ObjectIndex * CULLED_OBJECT_DATA_STRIDE + 14];
+
+	return float4x4(M0, M1, M2, M3);
+}
+
+// These are structured buffers so they can be scalar memory loads on GCN
+StructuredBuffer<float4> CulledObjectBounds;
+StructuredBuffer<float4> CulledObjectData;
+StructuredBuffer<float4> CulledObjectBoxBounds;
+
+RWStructuredBuffer<float4> RWCulledObjectBounds;
+RWStructuredBuffer<float4> RWCulledObjectData;
+RWStructuredBuffer<float4> RWCulledObjectBoxBounds;
+
+float4 LoadObjectPositionAndRadius(uint ObjectIndex)
+{
+	return LoadObjectPositionAndRadiusFromBuffer(ObjectIndex, CulledObjectBounds);
+}
+
+float4x4 LoadObjectWorldToVolume(uint ObjectIndex)
+{
+	return LoadObjectWorldToVolumeFromBuffer(ObjectIndex, CulledObjectData);
+}
+ 
+float3 LoadObjectLocalPositionExtent(uint ObjectIndex)
+{
+	return LoadObjectLocalPositionExtentFromBuffer(ObjectIndex, CulledObjectData);
+}
+
+float4 LoadObjectUVScale(uint ObjectIndex)
+{
+	bool bUnused;
+	return LoadObjectUVScaleFromBuffer(ObjectIndex, CulledObjectData, bUnused);
+}
+
+float4 LoadObjectUVScale(uint ObjectIndex, out bool bGeneratedAsTwoSided)
+{
+	return LoadObjectUVScaleFromBuffer(ObjectIndex, CulledObjectData, bGeneratedAsTwoSided);
+}
+
+float4 LoadObjectUVAddAndSelfShadowBias(uint ObjectIndex)
+{
+	return LoadObjectUVAddAndSelfShadowBiasFromBuffer(ObjectIndex, CulledObjectData);
+}
+
+float2 LoadObjectDistanceFieldMAD(uint ObjectIndex)
+{
+	return LoadObjectDistanceFieldMADFromBuffer(ObjectIndex, CulledObjectData);
+}
+
+float3x3 LoadObjectVolumeToWorld(uint ObjectIndex)
+{
+	return LoadObjectVolumeToWorldFromBuffer(ObjectIndex, CulledObjectData);
+}
+
+float4x4 LoadObjectLocalToWorld(uint ObjectIndex)
+{
+	return LoadObjectLocalToWorldFromBuffer(ObjectIndex, CulledObjectData);
+}
+
+// x = Offset in global buffer, y = NumLOD0, z = NumSurfels (all LODs), W = instance index
+uint4 LoadObjectSurfelCoordinate(uint ObjectIndex)
+{
+	return (uint4)CulledObjectData[ObjectIndex * CULLED_OBJECT_DATA_STRIDE + 15];
+}
+
+void LoadObjectViewSpaceBox(uint ObjectIndex, out float3 ObjectViewSpaceMin, out float3 ObjectViewSpaceMax)
+{
+	ObjectViewSpaceMin = CulledObjectBoxBounds[ObjectIndex * CULLED_OBJECT_BOX_BOUNDS_STRIDE + 0].xyz;
+	ObjectViewSpaceMax = CulledObjectBoxBounds[ObjectIndex * CULLED_OBJECT_BOX_BOUNDS_STRIDE + 1].xyz;
+}
+
+void LoadObjectAxes(uint ObjectIndex, out float3 ObjectAxisX, out float3 ObjectAxisY, out float3 ObjectAxisZ)
+{
+	ObjectAxisX = CulledObjectBoxBounds[ObjectIndex * CULLED_OBJECT_BOX_BOUNDS_STRIDE + 2].xyz;
+	ObjectAxisY = CulledObjectBoxBounds[ObjectIndex * CULLED_OBJECT_BOX_BOUNDS_STRIDE + 3].xyz;
+	ObjectAxisZ = CulledObjectBoxBounds[ObjectIndex * CULLED_OBJECT_BOX_BOUNDS_STRIDE + 4].xyz;
+}
+
+// VPLs generated by raytracing from the light
+#define VPL_DATA_STRIDE 3
+
+#define FINAL_GATHER_THREADGROUP_SIZE 64
+
+// Must match C++
+#define NUM_VISIBILITY_STEPS 10
+
+// Must match C++
+#define RECORD_CONE_DATA_STRIDE NUM_VISIBILITY_STEPS
+
+// Must match C++
+#define BENT_NORMAL_STRIDE 1
+
+Buffer<float4> IrradianceCachePositionRadius;
+Buffer<float> IrradianceCacheOccluderRadius;
+Buffer<uint2> IrradianceCacheTileCoordinate;
+Buffer<float4> IrradianceCacheNormal;
+Buffer<float4> IrradianceCacheBentNormal;
+Buffer<float4> IrradianceCacheIrradiance;
+
+Buffer<uint> ScatterDrawParameters;
+Buffer<uint> SavedStartIndex;
+
+uint NumConvexHullPlanes;
+float4 ViewFrustumConvexHull[6];
+
+bool ViewFrustumIntersectSphere(float3 SphereOrigin, float SphereRadius)
+{
+	for (uint PlaneIndex = 0; PlaneIndex < NumConvexHullPlanes; PlaneIndex++)
+	{
+		float4 PlaneData = ViewFrustumConvexHull[PlaneIndex];
+		float PlaneDistance = dot(PlaneData.xyz, SphereOrigin) - PlaneData.w;
+
+		if (PlaneDistance > SphereRadius)
+		{
+			return false;
+		}
+	}
+	
+	return true;
+}
+
+Buffer<float4> SurfelData;
+
+float4 LoadSurfelPositionAndRadius(uint SurfelIndex)
+{
+	return SurfelData[SurfelIndex * SURFEL_DATA_STRIDE + 0];
+}
+
+float3 LoadSurfelNormal(uint SurfelIndex)
+{
+	return SurfelData[SurfelIndex * SURFEL_DATA_STRIDE + 1].xyz;
+}
+
+float3 LoadSurfelDiffuseColor(uint SurfelIndex)
+{
+	return SurfelData[SurfelIndex * SURFEL_DATA_STRIDE + 2].xyz;
+}
+
+float3 LoadSurfelEmissiveColor(uint SurfelIndex)
+{
+	return SurfelData[SurfelIndex * SURFEL_DATA_STRIDE + 3].xyz;
+}
+
+float ApproximateConeConeIntersection(float ArcLength0, float ArcLength1, float AngleBetweenCones)
+{
+	float AngleDifference = abs(ArcLength0 - ArcLength1);
+
+	float Intersection = smoothstep(
+		0,
+		1.0,
+		1.0 - saturate((AngleBetweenCones - AngleDifference) / (ArcLength0 + ArcLength1 - AngleDifference)));
+
+	return Intersection;
+}
+
+float GetVPLOcclusion(float3 BentNormalAO, float3 NormalizedVectorToVPL, float VPLConeAngle, float VPLOcclusionStrength)
+{
+	float BentNormalLength = length(BentNormalAO);
+	float UnoccludedAngle = BentNormalLength * PI / VPLOcclusionStrength;
+	float AngleBetween = acos(dot(BentNormalAO, NormalizedVectorToVPL) / max(BentNormalLength, .0001f));
+	float Visibility = ApproximateConeConeIntersection(VPLConeAngle, UnoccludedAngle, AngleBetween);
+
+	// Can't rely on the direction of the bent normal when close to fully occluded, lerp to shadowed
+	Visibility = lerp(0, Visibility, saturate((UnoccludedAngle - .1f) / .2f));
+
+	return Visibility;
+}
diff --git a/Engine/Shaders/Private/DistanceFieldShadowingShared.ush b/Engine/Shaders/Private/DistanceFieldShadowingShared.ush
new file mode 100644
--- /dev/null
+++ b/Engine/Shaders/Private/DistanceFieldShadowingShared.ush
@@ -0,0 +1,190 @@
+// Copyright 1998-2017 Epic Games, Inc. All Rights Reserved.
+
+/*=============================================================================
+	DistanceFieldShadowingShared.usf
+=============================================================================*/
+
+Buffer<uint> ShadowTileHeadDataUnpacked;
+Buffer<uint> ShadowTileArrayData;
+uint2 ShadowTileListGroupSize;
+uint ShadowMaxObjectsPerTile;
+
+#ifndef SHADOW_TILE_ARRAY_DATA_STRIDE
+#define SHADOW_TILE_ARRAY_DATA_STRIDE 0
+#endif
+
+uint2 GetShadowTileHead(uint2 TileCoordinate)
+{
+	uint TileIndex = TileCoordinate.y * ShadowTileListGroupSize.x + TileCoordinate.x;
+
+	return uint2(
+		ShadowTileHeadDataUnpacked[TileIndex * 2 + 0], 
+		min(ShadowTileHeadDataUnpacked[TileIndex * 2 + 1], (uint)ShadowMaxObjectsPerTile)); 
+}
+
+float4x4 WorldToShadow;
+
+void GetShadowTileCulledData(float3 WorldPosition, out uint CulledDataStart, out uint NumIntersectingObjects)
+{
+	// Transform into shadow space
+	float4 HomogeneousShadowPosition = mul(float4(WorldPosition, 1), WorldToShadow);
+	float2 NormalizedShadowPosition = HomogeneousShadowPosition.xy * .5f + .5f;
+	NormalizedShadowPosition.y = 1 - NormalizedShadowPosition.y;
+	// Quantize the shadow position to get our tile position
+	uint2 TilePosition = (uint2)(NormalizedShadowPosition * ShadowTileListGroupSize);
+	// Fetch the tile head information
+	uint2 TileHead = GetShadowTileHead(TilePosition);
+	CulledDataStart = TileHead.x;
+	NumIntersectingObjects = TileHead.y;
+}
+
+#define MAX_INTERSECTING_OBJECTS 1024
+groupshared uint IntersectingObjectIndices[MAX_INTERSECTING_OBJECTS * 2];
+
+float TwoSidedMeshDistanceBias;
+
+float ShadowRayTraceThroughCulledObjects(
+	float3 WorldRayStart, 
+	float3 WorldRayEnd, 
+	float MaxRayTime, 
+	float TanLightAngle, 
+	float MinSphereRadius, 
+	float MaxSphereRadius, 
+	float SubsurfaceDensity,
+	uint CulledDataParameter, 
+	uint NumIntersectingObjects,
+	uniform bool bUseCulling,
+	uniform bool bUseScatterTileCulling,
+	bool bUseSubsurfaceTransmission)
+{
+	float MinRayTime = MaxRayTime;
+	float MinConeVisibility = 1;
+	float3 WorldRayUnitDirection = normalize(WorldRayEnd - WorldRayStart);
+
+	LOOP
+	for (uint ListObjectIndex = 0; ListObjectIndex < NumIntersectingObjects && MinRayTime >= MaxRayTime; ListObjectIndex++)
+	{
+		uint ObjectIndex;
+
+		if (bUseCulling)
+		{
+			if (bUseScatterTileCulling)
+			{
+				uint CulledDataStart = CulledDataParameter;
+				ObjectIndex = ShadowTileArrayData.Load((ListObjectIndex * ShadowTileListGroupSize.x * ShadowTileListGroupSize.y + CulledDataStart) * SHADOW_TILE_ARRAY_DATA_STRIDE);
+			}
+			else
+			{
+				uint GroupIndex = CulledDataParameter;
+				ObjectIndex = IntersectingObjectIndices[MAX_INTERSECTING_OBJECTS * GroupIndex + ListObjectIndex];
+			}
+		}
+		else
+		{
+			ObjectIndex = ListObjectIndex;
+		}
+
+		float4 SphereCenterAndRadius = LoadObjectPositionAndRadius(ObjectIndex);
+		float ObjectCenterDistanceAlongRay = dot(SphereCenterAndRadius.xyz - WorldRayStart, WorldRayUnitDirection);
+
+		BRANCH
+		if (ObjectCenterDistanceAlongRay > -SphereCenterAndRadius.w)
+		{
+			float3 LocalPositionExtent = LoadObjectLocalPositionExtent(ObjectIndex);
+			float4x4 WorldToVolume = LoadObjectWorldToVolume(ObjectIndex);
+			bool bGeneratedAsTwoSided;
+			float4 UVScaleAndVolumeScale = LoadObjectUVScale(ObjectIndex, bGeneratedAsTwoSided);
+			float4 UVAddAndSelfShadowBias = LoadObjectUVAddAndSelfShadowBias(ObjectIndex);
+			float2 DistanceFieldMAD = LoadObjectDistanceFieldMAD(ObjectIndex);
+
+			float3 VolumeRayStart = mul(float4(WorldRayStart, 1), WorldToVolume).xyz;
+			float3 VolumeRayEnd = mul(float4(WorldRayEnd, 1), WorldToVolume).xyz;
+			float3 VolumeRayDirection = VolumeRayEnd - VolumeRayStart;
+			float VolumeRayLength = length(VolumeRayDirection);
+			VolumeRayDirection /= VolumeRayLength;
+
+			float WorldToVolumeScale = 1.0f / UVScaleAndVolumeScale.w;
+			float VolumeMinSphereRadius = MinSphereRadius * WorldToVolumeScale;
+			float VolumeMaxSphereRadius = MaxSphereRadius * WorldToVolumeScale;
+			float VolumeTwoSidedMeshDistanceBias = TwoSidedMeshDistanceBias * WorldToVolumeScale;
+			float SelfShadowScale = 1.0f / max(UVAddAndSelfShadowBias.w * WorldToVolumeScale, .0001f);
+
+			// Expand the intersection box by the radius of the cone at the distance of the object along the cone
+			float LocalConeRadiusAtObject = min(TanLightAngle * max(ObjectCenterDistanceAlongRay, 0) * WorldToVolumeScale, VolumeMaxSphereRadius);
+
+			float2 IntersectionTimes = LineBoxIntersect(VolumeRayStart, VolumeRayEnd, -LocalPositionExtent - LocalConeRadiusAtObject, LocalPositionExtent + LocalConeRadiusAtObject);
+
+			BRANCH
+			if (IntersectionTimes.x < IntersectionTimes.y && IntersectionTimes.x < 1)
+			{
+				float SampleRayTime = IntersectionTimes.x * VolumeRayLength;
+				uint MaxSteps = 64;
+				float MinStepSize = 1.0f / (4 * MaxSteps);
+
+				float MinDistance = 1000000;
+				uint StepIndex = 0;
+
+				LOOP
+				for (; StepIndex < MaxSteps; StepIndex++)
+				{
+					float3 SampleVolumePosition = VolumeRayStart + VolumeRayDirection * SampleRayTime;
+					float3 ClampedSamplePosition = clamp(SampleVolumePosition, -LocalPositionExtent, LocalPositionExtent);
+					float DistanceToClamped = length(ClampedSamplePosition - SampleVolumePosition);
+					float3 VolumeUV = DistanceFieldVolumePositionToUV(ClampedSamplePosition, UVScaleAndVolumeScale.xyz, UVAddAndSelfShadowBias.xyz);
+					float DistanceField = SampleMeshDistanceField(VolumeUV, DistanceFieldMAD).x + DistanceToClamped;
+					
+					FLATTEN
+					if (bGeneratedAsTwoSided)
+					{
+						DistanceField -= VolumeTwoSidedMeshDistanceBias;
+					}
+
+					MinDistance = min(MinDistance, DistanceField);
+
+					// Don't allow occlusion within an object's self shadow distance
+					float SelfShadowVisibility = 1 - saturate(SampleRayTime * SelfShadowScale);
+
+					float SphereRadius = clamp(TanLightAngle * SampleRayTime, VolumeMinSphereRadius, VolumeMaxSphereRadius);
+					float StepVisibility = max(saturate(DistanceField / SphereRadius), SelfShadowVisibility);
+
+					if (bUseSubsurfaceTransmission)
+					{
+						// Determine the distance that the light traveled through the subsurface object
+						// This assumes that anything between this subsurface pixel and the light was also a subsurface material
+						float Thickness = SampleRayTime * UVScaleAndVolumeScale.w;
+						float SubsurfaceVisibility = saturate(exp(-Thickness * SubsurfaceDensity));
+
+						// Prevent full occlusion in the range that SSS is effective
+						// Note: this may cause the trace to travel through negative regions of the distance field
+						// It also prevents visibility from ever going to 0
+						StepVisibility = max(StepVisibility, SubsurfaceVisibility);
+					}
+
+					MinConeVisibility = min(MinConeVisibility, StepVisibility);
+
+					float StepDistance = max(abs(DistanceField), MinStepSize);
+					SampleRayTime += StepDistance;
+
+					// Terminate the trace if we are fully occluded or went past the end of the ray
+					if (MinConeVisibility < .01f
+						|| SampleRayTime > IntersectionTimes.y * VolumeRayLength)
+					{
+						break;
+					}
+				}
+
+				if (MinConeVisibility < .01f || StepIndex == MaxSteps)
+				{
+					MinConeVisibility = 0;
+					MinRayTime = min(MinRayTime, SampleRayTime * UVScaleAndVolumeScale.w);
+				}
+
+				// Force to shadowed as we approach max steps
+				MinConeVisibility = min(MinConeVisibility, (1 - StepIndex / (float)MaxSteps));
+			}
+		}
+	}
+
+	return MinConeVisibility;
+}
+	
\ No newline at end of file
diff --git a/Engine/Shaders/Private/DynamicLightingCommon.ush b/Engine/Shaders/Private/DynamicLightingCommon.ush
new file mode 100644
--- /dev/null
+++ b/Engine/Shaders/Private/DynamicLightingCommon.ush
@@ -0,0 +1,122 @@
+// Copyright 1998-2017 Epic Games, Inc. All Rights Reserved.
+
+/*=============================================================================
+	DynamicLightingCommon.usf: Contains functions shared by dynamic light shaders.
+=============================================================================*/
+
+/** 
+ * Returns a radial attenuation factor for a point light.  
+ * WorldLightVector is the vector from the position being shaded to the light, divided by the radius of the light. 
+ */
+float RadialAttenuation(float3 WorldLightVector, half FalloffExponent)
+{
+	float NormalizeDistanceSquared = dot(WorldLightVector, WorldLightVector);
+
+	// UE3 (fast, but now we not use the default of 2 which looks quite bad):
+	return pow(1.0f - saturate(NormalizeDistanceSquared), FalloffExponent); 
+
+	// new UE4 (more physically correct but slower and has a more noticable cutoff ring in the dark):
+	// AttenFunc(x) = 1 / (x * x + 1)
+	// derived: InvAttenFunc(y) = sqrtf(1 / y - 1)
+	// FalloffExponent is ignored
+	// the following code is a normalized (scaled and biased f(0)=1 f(1)=0) and optimized
+/*
+	// light less than x % is considered 0
+	// 20% produces a bright sphere, 5 % is ok for performance, 8% looks close to the old one, smaller numbers would be more realistic but then the attenuation radius also should be increased.
+	// we can expose CutoffPercentage later, alternatively we also can compute the attenuation radius from the CutoffPercentage and the brightness
+	const float CutoffPercentage = 5.0f;  
+	    
+	float CutoffFraction = CutoffPercentage * 0.01f;  
+
+	// those could be computed on C++ side
+	float PreCompX = 1.0f - CutoffFraction;
+	float PreCompY = CutoffFraction;
+	float PreCompZ = CutoffFraction / PreCompX;
+
+	return (1 / ( NormalizeDistanceSquared * PreCompX + PreCompY) - 1) * PreCompZ;
+*/
+}
+
+/** 
+ * Calculates attenuation for a spot light.
+ * L normalize vector to light. 
+ * SpotDirection is the direction of the spot light.
+ * SpotAngles.x is CosOuterCone, SpotAngles.y is InvCosConeDifference. 
+ */
+float SpotAttenuation(float3 L, float3 SpotDirection, float2 SpotAngles)
+{
+	float ConeAngleFalloff = Square(saturate((dot(L, -SpotDirection) - SpotAngles.x) * SpotAngles.y));
+	return ConeAngleFalloff;
+}
+
+/** Calculates radial and spot attenuation. */
+float CalcLightAttenuation(float3 WorldPosition, out float3 WorldLightVector)
+{
+	WorldLightVector = DeferredLightUniforms.NormalizedLightDirection;
+	float DistanceAttenuation = 1;
+
+#if RADIAL_ATTENUATION
+	WorldLightVector = DeferredLightUniforms.LightPosition - WorldPosition;
+	float DistanceSqr = dot( WorldLightVector, WorldLightVector );
+
+	// TODO Line segment falloff
+
+	// Sphere falloff (technically just 1/d2 but this avoids inf)
+	DistanceAttenuation = 1 / ( DistanceSqr + 1 );
+	
+	float LightRadiusMask = Square( saturate( 1 - Square( DistanceSqr * DeferredLightUniforms.LightInvRadius * DeferredLightUniforms.LightInvRadius ) ) );
+	DistanceAttenuation *= LightRadiusMask;
+	
+#if !INVERSE_SQUARED_FALLOFF
+	DistanceAttenuation = RadialAttenuation(WorldLightVector * DeferredLightUniforms.LightInvRadius, DeferredLightUniforms.LightFalloffExponent);
+#endif
+#endif
+
+	float SpotFalloff = 1;
+	#if RADIAL_ATTENUATION
+		SpotFalloff = SpotAttenuation( normalize(WorldLightVector), -DeferredLightUniforms.NormalizedLightDirection, DeferredLightUniforms.SpotAngles);
+	#endif
+
+	return SpotFalloff * DistanceAttenuation;
+}
+
+float3 GetNormalizedLightVector(float3 WorldPosition)
+{
+	// assumed to be normalized
+	float3 Ret = DeferredLightUniforms.NormalizedLightDirection;
+
+	#if RADIAL_ATTENUATION
+		Ret = normalize(DeferredLightUniforms.LightPosition - WorldPosition);
+	#endif
+
+	return Ret;
+}
+
+float GetLightInfluenceMask(float3 WorldPosition)
+{
+	float LightRadiusMask = 1;
+	float SpotFalloff = 1;
+
+	if (DeferredLightUniforms.LightInvRadius > 0)
+	{
+		float3 ToLight = DeferredLightUniforms.LightPosition - WorldPosition;
+		float DistanceSqr = dot(ToLight, ToLight);
+		float3 L = ToLight * rsqrt(DistanceSqr);
+
+		if (DeferredLightUniforms.LightFalloffExponent == 0)
+		{
+			LightRadiusMask = Square(saturate( 1 - Square(DistanceSqr * Square(DeferredLightUniforms.LightInvRadius))));
+		}
+		else
+		{
+			LightRadiusMask = RadialAttenuation(ToLight * DeferredLightUniforms.LightInvRadius, DeferredLightUniforms.LightFalloffExponent);
+		}
+
+		if (DeferredLightUniforms.SpotAngles.x > -2.0f)
+		{
+			SpotFalloff = SpotAttenuation(L, -DeferredLightUniforms.NormalizedLightDirection, DeferredLightUniforms.SpotAngles);
+		}
+	}
+
+	return LightRadiusMask * SpotFalloff > 0.0f ? 1.0f : 0.0f;
+}
\ No newline at end of file
diff --git a/Engine/Shaders/Private/EyeAdaptationCommon.ush b/Engine/Shaders/Private/EyeAdaptationCommon.ush
new file mode 100644
--- /dev/null
+++ b/Engine/Shaders/Private/EyeAdaptationCommon.ush
@@ -0,0 +1,22 @@
+// Copyright 1998-2017 Epic Games, Inc. All Rights Reserved.
+
+/*=============================================================================
+	EyeAdaptationCommon.usf: PostProcessing shared functions and structures.
+=============================================================================*/
+
+#ifndef __EYE_ADAPTATION_COMMON__
+#define __EYE_ADAPTATION_COMMON__
+
+
+// Single pixel RT used for EyeAdaptation, use function EyeAdaptationLookup() to access.
+Texture2D EyeAdaptation;
+#if FEATURE_LEVEL >= FEATURE_LEVEL_SM5
+//Provides access to the EyeAdaptation RT. Only available in SM5.
+float EyeAdaptationLookup()
+{
+	return EyeAdaptation.Load(int3(0, 0, 0)).r;
+}
+#endif
+
+
+#endif // __EYE_ADAPTATION_COMMON__
diff --git a/Engine/Shaders/Private/FastMath.ush b/Engine/Shaders/Private/FastMath.ush
new file mode 100644
--- /dev/null
+++ b/Engine/Shaders/Private/FastMath.ush
@@ -0,0 +1,273 @@
+// Copyright 1998-2017 Epic Games, Inc. All Rights Reserved.
+
+/*=============================================================================
+	Common.usf: Common shader code
+=============================================================================*/
+  
+#pragma once
+
+/******************************************************************************
+    Shader Fast Math Lib (v0.41)
+    A shader math library for optimized approximate transcendental functions.
+    Optimized and tested on AMD GCN architecture.
+********************************************************************************/
+
+/******************************************************************************
+    The MIT License (MIT)
+    Copyright (c) <2014> <Michal Drobot>
+    Permission is hereby granted, free of charge, to any person obtaining a copy
+    of this software and associated documentation files (the "Software"), to deal
+    in the Software without restriction, including without limitation the rights
+    to use, copy, modify, merge, publish, distribute, sublicense, and/or sell
+    copies of the Software, and to permit persons to whom the Software is
+    furnished to do so, subject to the following conditions:
+    The above copyright notice and this permission notice shall be included in
+    all copies or substantial portions of the Software.
+    THE SOFTWARE IS PROVIDED "AS IS", WITHOUT WARRANTY OF ANY KIND, EXPRESS OR
+    IMPLIED, INCLUDING BUT NOT LIMITED TO THE WARRANTIES OF MERCHANTABILITY,
+    FITNESS FOR A PARTICULAR PURPOSE AND NONINFRINGEMENT. IN NO EVENT SHALL THE
+    AUTHORS OR COPYRIGHT HOLDERS BE LIABLE FOR ANY CLAIM, DAMAGES OR OTHER
+    LIABILITY, WHETHER IN AN ACTION OF CONTRACT, TORT OR OTHERWISE, ARISING FROM,
+    OUT OF OR IN CONNECTION WITH THE SOFTWARE OR THE USE OR OTHER DEALINGS IN
+    THE SOFTWARE.
+********************************************************************************/
+
+//
+// Normalized range [0,1] Constants
+//
+#define IEEE_INT_RCP_CONST_NR0_SNORM        0x7EEF370B
+#define IEEE_INT_SQRT_CONST_NR0_SNORM       0x1FBD1DF5
+#define IEEE_INT_RCP_SQRT_CONST_NR0_SNORM   0x5F341A43
+
+
+// Relative error : ~3.4% over full
+// Precise format : ~small float
+// 2 ALU
+float rsqrtFast( float x )
+{
+#if !GL3_PROFILE
+	int i = asint(x);
+	i = 0x5f3759df - (i >> 1);
+	return asfloat(i);
+#else
+	return rsqrt(x);
+#endif
+}
+
+// Relative error : < 0.7% over full
+// Precise format : ~small float
+// 1 ALU
+float sqrtFast( float x )
+{
+#if !GL3_PROFILE
+	int i = asint(x);
+	i = 0x1FBD1DF5 + (i >> 1);
+	return asfloat(i);
+#else
+	return sqrt(x);
+#endif
+}
+
+// Relative error : < 0.4% over full
+// Precise format : ~small float
+// 1 ALU
+float rcpFast( float x )
+{
+#if !GL3_PROFILE
+	int i = asint(x);
+	i = 0x7EF311C2 - i;
+	return asfloat(i);
+#else
+	return rcp(x);
+#endif
+}
+
+// Using 1 Newton Raphson iterations
+// Relative error : < 0.02% over full
+// Precise format : ~half float
+// 3 ALU
+float rcpFastNR1( float x )
+{
+#if !GL3_PROFILE
+	int i = asint(x);
+	i = 0x7EF311C3 - i;
+	float xRcp = asfloat(i);
+	xRcp = xRcp * (-xRcp * x + 2.0f);
+	return xRcp;
+#else
+	return rcp(x);
+#endif
+}
+
+float lengthFast( float3 v )
+{
+	float LengthSqr = dot(v,v);
+	return sqrtFast( LengthSqr );
+}
+
+float3 normalizeFast( float3 v )
+{
+	float LengthSqr = dot(v,v);
+	return v * rsqrtFast( LengthSqr );
+}
+
+float4 fastClamp(float4 x, float4 Min, float4 Max)
+{
+#if PS4_PROFILE  
+	// Warning: only correct if Min is smaller than Max
+	return med3(x, Min, Max);
+#else
+	return clamp(x, Min, Max);
+#endif
+}
+
+float3 fastClamp(float3 x, float3 Min, float3 Max)
+{
+#if PS4_PROFILE  
+	// Warning: only correct if Min is smaller than Max
+	return med3(x, Min, Max);
+#else
+	return clamp(x, Min, Max);
+#endif
+}
+
+float2 fastClamp(float2 x, float2 Min, float2 Max)
+{
+#if PS4_PROFILE  
+	// Warning: only correct if Min is smaller than Max
+	return med3(x, Min, Max);
+#else
+	return clamp(x, Min, Max);
+#endif
+}
+
+float fastClamp(float x, float Min, float Max)
+{
+#if PS4_PROFILE  
+	// Warning: only correct if Min is smaller than Max
+	return med3(x, Min, Max);
+#else
+	return clamp(x, Min, Max);
+#endif
+}
+
+//
+// Trigonometric functions
+//
+
+// max absolute error 9.0x10^-3
+// Eberly's polynomial degree 1 - respect bounds
+// 4 VGPR, 12 FR (8 FR, 1 QR), 1 scalar
+// input [-1, 1] and output [0, PI]
+float acosFast(float inX) 
+{
+    float x = abs(inX);
+    float res = -0.156583f * x + (0.5 * PI);
+    res *= sqrt(1.0f - x);
+    return (inX >= 0) ? res : PI - res;
+}
+
+// Same cost as acosFast + 1 FR
+// Same error
+// input [-1, 1] and output [-PI/2, PI/2]
+float asinFast( float x )
+{
+    return (0.5 * PI) - acosFast(x);
+}
+
+// max absolute error 1.3x10^-3
+// Eberly's odd polynomial degree 5 - respect bounds
+// 4 VGPR, 14 FR (10 FR, 1 QR), 2 scalar
+// input [0, infinity] and output [0, PI/2]
+float atanFastPos( float x ) 
+{ 
+    float t0 = (x < 1.0f) ? x : 1.0f / x;
+    float t1 = t0 * t0;
+    float poly = 0.0872929f;
+    poly = -0.301895f + poly * t1;
+    poly = 1.0f + poly * t1;
+    poly = poly * t0;
+    return (x < 1.0f) ? poly : (0.5 * PI) - poly;
+}
+
+// 4 VGPR, 16 FR (12 FR, 1 QR), 2 scalar
+// input [-infinity, infinity] and output [-PI/2, PI/2]
+float atanFast( float x )
+{
+    float t0 = atanFastPos( abs(x) );
+    return (x < 0) ? -t0: t0;
+}
+
+float atan2Fast( float y, float x )
+{
+	float t0 = max( abs(x), abs(y) );
+	float t1 = min( abs(x), abs(y) );
+	float t3 = t1 / t0;
+	float t4 = t3 * t3;
+
+	// Same polynomial as atanFastPos
+	t0 =         + 0.0872929;
+	t0 = t0 * t4 - 0.301895;
+	t0 = t0 * t4 + 1.0;
+	t3 = t0 * t3;
+
+	t3 = abs(y) > abs(x) ? (0.5 * PI) - t3 : t3;
+	t3 = x < 0 ? PI - t3 : t3;
+	t3 = y < 0 ? -t3 : t3;
+
+	return t3;
+}
+
+// 4th order polynomial approximation
+// 4 VGRP, 16 ALU Full Rate
+// 7 * 10^-5 radians precision
+// Reference : Handbook of Mathematical Functions (chapter : Elementary Transcendental Functions), M. Abramowitz and I.A. Stegun, Ed.
+float acosFast4(float inX)
+{
+	float x1 = abs(inX);
+	float x2 = x1 * x1;
+	float x3 = x2 * x1;
+	float s;
+
+	s = -0.2121144f * x1 + 1.5707288f;
+	s = 0.0742610f * x2 + s;
+	s = -0.0187293f * x3 + s;
+	s = sqrt(1.0f - x1) * s;
+
+	// acos function mirroring
+	// check per platform if compiles to a selector - no branch neeeded
+	return inX >= 0.0f ? s : PI - s;
+}
+
+// 4th order polynomial approximation
+// 4 VGRP, 16 ALU Full Rate
+// 7 * 10^-5 radians precision 
+float asinFast4( float x )
+{
+	return (0.5 * PI) - acosFast4(x);
+}
+
+// @param A doesn't have to be normalized, output could be NaN if this is near 0,0,0
+// @param B doesn't have to be normalized, output could be NaN if this is near 0,0,0
+// @return can be passed to a acosFast() or acos() to compute an angle
+float CosBetweenVectors(float3 A, float3 B)
+{
+	// unoptimized: dot(normalize(A), normalize(B))
+	return dot(A, B) * rsqrt(length2(A) * length2(B));
+}
+
+// @param A doesn't have to be normalized, output could be NaN if this is near 0,0,0
+// @param B doesn't have to be normalized, output could be NaN if this is near 0,0,0
+float AngleBetweenVectors(float3 A, float3 B)
+{
+	return acos(CosBetweenVectors(A, B));
+}
+// @param A doesn't have to be normalized, output could be NaN if this is near 0,0,0
+// @param B doesn't have to be normalized, output could be NaN if this is near 0,0,0
+float AngleBetweenVectorsFast(float3 A, float3 B)
+{
+	return acosFast(CosBetweenVectors(A, B));
+}
+
+
+
diff --git a/Engine/Shaders/Private/FlatTessellation.ush b/Engine/Shaders/Private/FlatTessellation.ush
new file mode 100644
--- /dev/null
+++ b/Engine/Shaders/Private/FlatTessellation.ush
@@ -0,0 +1,236 @@
+// Copyright 1998-2017 Epic Games, Inc. All Rights Reserved.
+
+/*=============================================================================
+	FlatTessellation.usf: Shader fragments for flat tessellation
+=============================================================================*/
+
+#if DISPLACEMENT_ANTICRACK
+#define IN_CONTROL_POINTS 12
+#else
+#define IN_CONTROL_POINTS 3
+#endif
+
+
+
+struct FHullShaderConstantOutput
+{
+	// Tess factor for the FF HW block
+	float TessFactor[3]    : SV_TessFactor;
+	float InsideTessFactor : SV_InsideTessFactor;
+};
+
+struct FFlatTessellationHSToDS
+{
+	FPassSpecificVSToDS PassSpecificData;
+
+	// scale factors in tangent space
+	float3 DisplacementScale : Flat_DisplacementScales;
+
+	// Tessellation multiplier.
+	float TessellationMultiplier : Flat_TessellationMultiplier;
+
+#if DISPLACEMENT_ANTICRACK
+	// dominant edge and vertex data used to avoid cracking when displacing
+	FHullShaderConstantDominantVertexData DominantVertex : Flat_DominantVertex;
+	FHullShaderConstantDominantEdgeData DominantEdge : Flat_DominantEdge;
+#endif
+};
+
+// GLSL needs those on the domain shader, HLSL on the hull shader, so we replicate them for both
+#define TESSELLATION_ATTRIBUTES [partitioning("fractional_odd")][outputtopology("triangle_cw")]
+
+#if HULLSHADER
+
+	FHullShaderConstantOutput HS_FlatTrianglesConstant( const OutputPatch<FFlatTessellationHSToDS, 3> I )
+	{
+		FHullShaderConstantOutput O = (FHullShaderConstantOutput)0;
+
+		#if INSTANCED_STEREO
+			ResolvedView = ResolveView(GetEyeIndex(I[0].PassSpecificData.FactoryInterpolants.InterpolantsVSToPS.PackedEyeIndex));
+		#else
+			ResolvedView = ResolveView();
+		#endif
+	    
+#if DISABLE_TESSELLATION_OVERRIDE
+		O.TessFactor[0] = 1.0f; // 1->2 edge
+		O.TessFactor[1] = 1.0f; // 2->0 edge
+		O.TessFactor[2] = 1.0f; // 0->1 edge
+		O.InsideTessFactor = 1.0f;
+#else // #if DISABLE_TESSELLATION_OVERRIDE
+		float3 WorldPostion0 = I[0].PassSpecificData.Position.xyz;
+		float3 WorldPostion1 = I[1].PassSpecificData.Position.xyz;
+		float3 WorldPostion2 = I[2].PassSpecificData.Position.xyz;
+
+		float4 TessellationMultipliers;
+		TessellationMultipliers.x = 0.5 * (I[1].TessellationMultiplier + I[2].TessellationMultiplier);
+		TessellationMultipliers.y = 0.5 * (I[2].TessellationMultiplier + I[0].TessellationMultiplier);
+		TessellationMultipliers.z = 0.5 * (I[0].TessellationMultiplier + I[1].TessellationMultiplier);
+		TessellationMultipliers.w = 0.333 * (I[0].TessellationMultiplier + I[1].TessellationMultiplier + I[2].TessellationMultiplier);
+
+		//@todo: This should be vectorized and done in the MainHull shader.
+		float4 CompositeTessellationFactors = TessellationMultipliers * CalculateCompositeTessellationFactors(WorldPostion0,WorldPostion1,WorldPostion2);
+		CompositeTessellationFactors = clamp( CompositeTessellationFactors, 1, 15 );
+
+		O.TessFactor[0] = CompositeTessellationFactors.x; // 1->2 edge
+		O.TessFactor[1] = CompositeTessellationFactors.y; // 2->0 edge
+		O.TessFactor[2] = CompositeTessellationFactors.z; // 0->1 edge
+		O.InsideTessFactor = CompositeTessellationFactors.w;
+#endif // #if DISABLE_TESSELLATION_OVERRIDE
+
+		return O;
+	}
+
+	[domain("tri")]
+	[patchconstantfunc("HS_FlatTrianglesConstant")]
+	[outputcontrolpoints(3)]
+	[maxtessfactor(15)]
+	TESSELLATION_ATTRIBUTES
+	FFlatTessellationHSToDS MainHull( InputPatch<FPassSpecificVSToDS, IN_CONTROL_POINTS> I, uint ControlPointID : SV_OutputControlPointID )
+	{
+		#if INSTANCED_STEREO
+			ResolvedView = ResolveView(GetEyeIndex(I[0].FactoryInterpolants.InterpolantsVSToPS.PackedEyeIndex));
+		#else
+			ResolvedView = ResolveView();
+		#endif
+
+		FFlatTessellationHSToDS O = (FFlatTessellationHSToDS) 0;
+		
+		float3x3 TangentToLocal = VertexFactoryGetTangentToLocalDS(I[ControlPointID].FactoryInterpolants);
+		O.DisplacementScale = GetTangentSpaceNonUniformScales(TangentToLocal);
+
+		const uint NextControlPointID = ControlPointID < 2 ? ControlPointID + 1 : 0; // (ControlPointID + 1) % 3
+
+		O.PassSpecificData = I[ControlPointID];
+		O.TessellationMultiplier = GetTessellationMultiplier( I[ControlPointID] );
+
+#if DISPLACEMENT_ANTICRACK
+		// NextControlPointID is first opposite edge vert
+		// NextNextControlPointID is second opposite edge vert
+		const uint NextNextControlPointID = NextControlPointID < 2 ? NextControlPointID + 1 : 0;
+
+		// We also need the neghbors of the opposite edge verts
+		const uint NextNeighborControlPointID = 3 + 2 * NextControlPointID;	// neighbor for first opposite edge vert
+		const uint NextNeighborNextControlPointID = NextNeighborControlPointID + 1; // neighbor for second opposite edge vert
+
+		// Dominant vertex is provided explicitly
+		const uint DominantVertexControlPointID = 9 + ControlPointID;
+		O.DominantVertex = GenerateDominantVertexData(I[DominantVertexControlPointID]);
+
+		// we need the other verts from each prim to generate the edges, clock wise winding
+		O.DominantEdge = GenerateDominantEdgeData(I[NextControlPointID],I[NextNextControlPointID],I[NextNeighborControlPointID],I[NextNeighborNextControlPointID]);
+#endif	// DISPLACEMENT_ANTICRACK
+		
+		return O;
+	}
+#endif
+
+#if DOMAINSHADER
+	[domain("tri")]
+	TESSELLATION_ATTRIBUTES
+	FPassSpecificVSToPS MainDomain(
+		FHullShaderConstantOutput HSConstantData,
+		const OutputPatch<FFlatTessellationHSToDS, 3> I,
+		float3 BarycentricCoords : SV_DomainLocation 
+#if USE_GLOBAL_CLIP_PLANE
+		, out float OutGlobalClipPlaneDistance : SV_ClipDistance
+#endif
+#if INSTANCED_STEREO && !MULTI_VIEW
+		, out float OutClipDistance : SV_ClipDistance1
+#endif
+		)
+	{
+		#if INSTANCED_STEREO
+			ResolvedView = ResolveView(GetEyeIndex(I[0].PassSpecificData.FactoryInterpolants.InterpolantsVSToPS.PackedEyeIndex));
+			#if !MULTI_VIEW
+				OutClipDistance = 0.0;
+			#endif
+		#else
+			ResolvedView = ResolveView();
+		#endif
+
+		// Get the barycentric coords
+		float U = BarycentricCoords.x;
+		float V = BarycentricCoords.y;
+		float W = BarycentricCoords.z;
+		
+		// Interpolate world position according to our tessellation scheme
+		float4 WorldPosition = I[0].PassSpecificData.Position * U + I[1].PassSpecificData.Position * V + I[2].PassSpecificData.Position * W;
+		
+		// Interp remaining attributes (construct barycentric interp from bilerp primitives)
+		// NB: the HLSL compiler resolves the 1.f to something efficient (i.e. no-op)
+		FPassSpecificVSToDS Interp = PassInterpolate(
+			PassInterpolate( I[0].PassSpecificData, U, I[1].PassSpecificData, V ), 1.f,
+			I[2].PassSpecificData, W );
+							
+		// Forward the eye index along	
+		#if INSTANCED_STEREO
+		Interp.FactoryInterpolants.InterpolantsVSToPS.PackedEyeIndex = I[0].PassSpecificData.FactoryInterpolants.InterpolantsVSToPS.PackedEyeIndex;
+		#endif
+																		
+		FMaterialTessellationParameters MaterialParameters = GetMaterialTessellationParameters(Interp.FactoryInterpolants, WorldPosition.xyz);
+
+#if DISPLACEMENT_ANTICRACK
+
+		float2 DisplacementTexCoords = 0;
+#if NUM_MATERIAL_TEXCOORDS
+		DisplacementTexCoords = GetTextureCoordinateDS( Interp );
+#endif
+		
+		FMaterialTessellationParameters DisplacementMaterialParameters = MaterialParameters;
+		
+		float3x3 TangentToLocal = VertexFactoryGetTangentToLocalDS(Interp.FactoryInterpolants);
+
+		// Override to dominant data if along edge or at control vert.  This is necessary to avoid cracks caused by primitives
+		// from disjoint locations in UV space abutting each other.
+		ApplyDominantData(
+			DisplacementTexCoords, DisplacementMaterialParameters.TangentToWorld,
+			GetNormalDS( TangentToLocal ), GetTangentDS( TangentToLocal ),
+			U, V, W,
+			I[0].DominantEdge, I[1].DominantEdge, I[2].DominantEdge,
+			I[0].DominantVertex, I[1].DominantVertex, I[2].DominantVertex );
+							
+		// Override texture coordinate 0.
+#if NUM_MATERIAL_TEXCOORDS
+		DisplacementMaterialParameters.TexCoords[0].xy = DisplacementTexCoords.xy;
+#endif
+
+#else
+
+		FMaterialTessellationParameters DisplacementMaterialParameters = MaterialParameters;
+
+#endif	// DISPLACEMENT_ANTICRACK
+
+		// tangent space displacement scaling
+		DisplacementMaterialParameters.TangentToWorldPreScale = I[0].DisplacementScale*U + I[1].DisplacementScale*V + I[2].DisplacementScale*W;
+
+		// world space displacement value.  This will likely use the tangent2world transform.
+		float3 WorldDisplacement = GetMaterialWorldDisplacement(DisplacementMaterialParameters);
+
+		WorldPosition.xyz += WorldDisplacement;
+
+		#if USE_GLOBAL_CLIP_PLANE
+			OutGlobalClipPlaneDistance = dot(ResolvedView.GlobalClippingPlane, float4(WorldPosition.xyz - ResolvedView.PreViewTranslation.xyz, 1));
+		#endif
+
+		FPassSpecificVSToPS Result = PassFinalizeTessellationOutput(Interp, WorldPosition,MaterialParameters);
+
+		#if INSTANCED_STEREO && !MULTI_VIEW
+		BRANCH
+		if (IsInstancedStereo(I[0].PassSpecificData.FactoryInterpolants.InterpolantsVSToPS.PackedEyeIndex))
+		{
+			const uint EyeIndex = GetEyeIndex(I[0].PassSpecificData.FactoryInterpolants.InterpolantsVSToPS.PackedEyeIndex);
+
+			// Clip at the center of the screen
+			OutClipDistance = dot(Result.Position, EyeClipEdge[EyeIndex]);
+			
+			// Scale to the width of a single eye viewport
+			Result.Position.x *= 0.5 * ResolvedView.HMDEyePaddingOffset;
+
+			// Shift to the eye viewport
+			Result.Position.x += (EyeOffsetScale[EyeIndex] * Result.Position.w) * (1.0f - 0.5 * ResolvedView.HMDEyePaddingOffset);
+		}
+		#endif
+
+		return Result;
+	}
+#endif
diff --git a/Engine/Shaders/Private/ForwardLightingCommon.ush b/Engine/Shaders/Private/ForwardLightingCommon.ush
new file mode 100644
--- /dev/null
+++ b/Engine/Shaders/Private/ForwardLightingCommon.ush
@@ -0,0 +1,334 @@
+// Copyright 1998-2017 Epic Games, Inc. All Rights Reserved.
+
+/*=============================================================================
+	ForwardLightingCommon.usf
+=============================================================================*/
+
+#define NON_DIRECTIONAL_DIRECT_LIGHTING (TRANSLUCENCY_LIGHTING_VOLUMETRIC_NONDIRECTIONAL || TRANSLUCENCY_LIGHTING_VOLUMETRIC_PERVERTEX_NONDIRECTIONAL)
+#define SUPPORT_CONTACT_SHADOWS 0
+
+#include "DeferredLightingCommon.ush"
+#include "LightGridCommon.ush"
+
+Texture2D ResolvedSceneDepthTexture;
+
+void UpdateNearestSample(float Z, float2 UV, float FullResZ, inout float MinDist, inout float2 NearestUV)
+{
+    float DepthDelta = abs(Z - FullResZ);
+
+	FLATTEN
+    if (DepthDelta < MinDist)
+    {
+        MinDist = DepthDelta;
+        NearestUV = UV;
+    }
+}
+
+float2 CalculateNearestResolvedDepthScreenUV(float2 ScreenUV, float SceneDepth)
+{
+	float2 EffectiveScreenUV = ScreenUV;
+
+	if (View.NumSceneColorMSAASamples > 1)
+	{
+		int2 IntScreenUV = int2(trunc(ScreenUV * View.BufferSizeAndInvSize.xy));
+
+		float DeferredShadowingDepth = ConvertFromDeviceZ(ResolvedSceneDepthTexture.Load(int3(IntScreenUV, 0)).r);
+		float RelativeDepthThreshold = .01f;
+
+		// Fragment depth doesn't match what we used for deferred shadowing, search neighbors in cross pattern
+		// Even with this nearest depth upsampling there can be edge artifacts from deferred shadowing, 
+		// Since depth testing and stencil testing used during shadow projection are done per-sample and we're fetching from the resolved light attenuation
+		if (abs(DeferredShadowingDepth - SceneDepth) / SceneDepth > RelativeDepthThreshold)
+		{
+			float2 TexelSize = View.BufferSizeAndInvSize.zw;
+			float MinDist = 1.e8f;
+
+			float2 LeftUV = ScreenUV + float2(-TexelSize.x, 0);
+			float LeftDepth = ConvertFromDeviceZ(ResolvedSceneDepthTexture.Load(int3(IntScreenUV.x - 1, IntScreenUV.y, 0)).r);
+			UpdateNearestSample(LeftDepth, LeftUV, SceneDepth, MinDist, EffectiveScreenUV);
+
+			float2 UpUV = ScreenUV + float2(0, TexelSize.y);
+			float UpDepth = ConvertFromDeviceZ(ResolvedSceneDepthTexture.Load(int3(IntScreenUV.x, IntScreenUV.y + 1, 0)).r);
+			UpdateNearestSample(UpDepth, UpUV, SceneDepth, MinDist, EffectiveScreenUV);
+
+			float2 RightUV = ScreenUV + float2(TexelSize.x, 0);
+			float RightDepth = ConvertFromDeviceZ(ResolvedSceneDepthTexture.Load(int3(IntScreenUV.x + 1, IntScreenUV.y, 0)).r);
+			UpdateNearestSample(RightDepth, RightUV, SceneDepth, MinDist, EffectiveScreenUV);
+
+			float2 BottomUV = ScreenUV + float2(0, -TexelSize.y);
+			float BottomDepth = ConvertFromDeviceZ(ResolvedSceneDepthTexture.Load(int3(IntScreenUV.x, IntScreenUV.y - 1, 0)).r);
+			UpdateNearestSample(BottomDepth, BottomUV, SceneDepth, MinDist, EffectiveScreenUV);
+		}
+	}
+
+	return EffectiveScreenUV;
+}
+
+float4 GetForwardDynamicShadowFactors(float2 ScreenUV)
+{
+	return GetPerPixelLightAttenuation(ScreenUV);
+}
+
+float4 UnpackShadowMapChannelMask(uint ShadowMapChannelMaskPacked)
+{
+	return float4((ShadowMapChannelMaskPacked & 1) ? 1.0f : 0.0f, (ShadowMapChannelMaskPacked & 2) ? 1.0f : 0.0f, (ShadowMapChannelMaskPacked & 4) ? 1.0f : 0.0f, (ShadowMapChannelMaskPacked & 8) ? 1.0f : 0.0f);
+}
+
+Texture2D IndirectOcclusionTexture;
+SamplerState IndirectOcclusionTextureSampler;
+
+float GetIndirectOcclusion(float2 ScreenUV, FGBufferData GBufferData)
+{
+	float IndirectOcclusion;
+#if SUPPORTS_INDEPENDENT_SAMPLERS
+	IndirectOcclusion = Texture2DSampleLevel(IndirectOcclusionTexture, LightAttenuationTextureSampler, ScreenUV, 0).x;
+#else
+	IndirectOcclusion = Texture2DSampleLevel(IndirectOcclusionTexture, IndirectOcclusionTextureSampler, ScreenUV, 0).x;
+#endif
+
+	// Reduce self shadowing intensity on characters (reuse distance field bit, really should be HasCapsuleShadowRepresentation)
+	IndirectOcclusion = lerp(1, IndirectOcclusion, HasDynamicIndirectShadowCasterRepresentation(GBufferData) ? View.IndirectCapsuleSelfShadowingIntensity : 1);
+
+	return IndirectOcclusion;
+}
+
+float3 GetForwardDirectLighting(uint GridIndex, float3 WorldPosition, float3 CameraVector, FGBufferData GBufferData, float2 ScreenUV, uint EyeIndex)
+{
+	float4 DynamicShadowFactors = 1;
+
+	#if MATERIALBLENDING_SOLID || MATERIALBLENDING_MASKED
+		DynamicShadowFactors = GetForwardDynamicShadowFactors(ScreenUV);
+	#endif
+
+	float3 DirectLighting = 0;
+
+	// Prevent 0 Roughness which causes NaNs in Vis_SmithJointApprox
+	float MinRoughness = .04f;
+
+#if TRANSLUCENCY_ANY_VOLUMETRIC
+	// No specular on volumetric translucency lighting modes
+	MinRoughness = 1.0f;
+#endif
+
+	BRANCH
+	if (ForwardGlobalLightData.HasDirectionalLight)
+	{
+		FDeferredLightData LightData = (FDeferredLightData)0;
+		LightData.LightColorAndFalloffExponent = float4(ForwardGlobalLightData.DirectionalLightColor, 0);
+		LightData.LightDirection = ForwardGlobalLightData.DirectionalLightDirection;
+		LightData.DistanceFadeMAD = ForwardGlobalLightData.DirectionalLightDistanceFadeMAD;
+		LightData.bRadialLight = false;
+		LightData.MinRoughness = MinRoughness;
+
+		LightData.ShadowedBits = (ForwardGlobalLightData.DirectionalLightShadowMapChannelMask & 0xFF) != 0 ? 1 : 0;
+		// Static shadowing uses ShadowMapChannel, dynamic shadows are packed into light attenuation using PreviewShadowMapChannel
+		LightData.ShadowMapChannelMask = UnpackShadowMapChannelMask(ForwardGlobalLightData.DirectionalLightShadowMapChannelMask);
+		float4 PreviewShadowMapChannelMask = UnpackShadowMapChannelMask(ForwardGlobalLightData.DirectionalLightShadowMapChannelMask >> 4);
+		float DynamicShadowing = dot(PreviewShadowMapChannelMask, DynamicShadowFactors);
+
+		// In the forward shading path we can't separate per-object shadows from CSM, since we only spend one light attenuation channel per light
+		// If CSM is enabled (distance fading to precomputed shadowing is active), treat all of our dynamic shadowing as whole scene shadows that will be faded out at the max CSM distance
+		// If CSM is not enabled, allow our dynamic shadowing to coexist with precomputed shadowing
+		float PerObjectShadowing = LightData.DistanceFadeMAD.y < 0.0f ? 1.0f : DynamicShadowing;
+		float WholeSceneShadowing = LightData.DistanceFadeMAD.y < 0.0f ? DynamicShadowing : 1.0f;
+		
+		float4 LightAttenuation = float4(WholeSceneShadowing.xx, PerObjectShadowing.xx);
+		float3 NewLighting = GetDynamicLighting(WorldPosition, -CameraVector, GBufferData, 1, GBufferData.ShadingModelID, LightData, LightAttenuation, 0.5, uint2(0,0)).xyz;
+
+		FLATTEN
+		if ((ForwardGlobalLightData.DirectionalLightShadowMapChannelMask >> 8) & Primitive.LightingChannelMask)
+		{
+			DirectLighting += NewLighting;
+		}
+	}
+	
+	const FCulledLightsGridData CulledLightsGrid = GetCulledLightsGrid(GridIndex, EyeIndex);
+
+	LOOP
+	for (uint LocalLightListIndex = 0; LocalLightListIndex < CulledLightsGrid.NumLocalLights; LocalLightListIndex++)
+	{
+		const FLocalLightData LocalLight = GetLocalLightData(CulledLightsGrid.DataStartIndex + LocalLightListIndex, EyeIndex);
+		 
+		FDeferredLightData LightData = (FDeferredLightData)0;
+		LightData.LightPositionAndInvRadius = LocalLight.LightPositionAndInvRadius;
+		LightData.LightColorAndFalloffExponent = LocalLight.LightColorAndFalloffExponent;
+		LightData.LightDirection = LocalLight.LightDirectionAndShadowMask.xyz;
+		float SourceLength = f16tof32(asuint(LocalLight.SpotAnglesAndSourceRadiusPacked.w));
+		LightData.SpotAnglesAndSourceRadius = float4(LocalLight.SpotAnglesAndSourceRadiusPacked.xyz, SourceLength);
+		LightData.bInverseSquared = LightData.LightColorAndFalloffExponent.w == 0;
+		LightData.bRadialLight = true;
+		LightData.bSpotLight = LightData.SpotAnglesAndSourceRadius.x > -2.0f;
+		LightData.MinRoughness = MinRoughness;
+
+		uint PackedShadowMapChannelMask = asuint(LocalLight.LightDirectionAndShadowMask.w);
+		LightData.ShadowedBits = (PackedShadowMapChannelMask & 0xFF) != 0 ? 1 : 0;
+		// Static shadowing uses ShadowMapChannel, dynamic shadows are packed into light attenuation using PreviewShadowMapChannel
+		LightData.ShadowMapChannelMask = UnpackShadowMapChannelMask(PackedShadowMapChannelMask);
+		float4 PreviewShadowMapChannelMask = UnpackShadowMapChannelMask(PackedShadowMapChannelMask >> 4);
+		float DynamicShadowing = dot(PreviewShadowMapChannelMask, DynamicShadowFactors);
+		float4 LightAttenuation = float4(1, 1, DynamicShadowing.x, DynamicShadowing.x);
+		float3 NewLighting = GetDynamicLighting(WorldPosition, -CameraVector, GBufferData, 1, GBufferData.ShadingModelID, LightData, LightAttenuation, 0.5, uint2(0,0)).xyz;
+
+		FLATTEN
+		if ((PackedShadowMapChannelMask >> 8) & Primitive.LightingChannelMask)
+		{
+			DirectLighting += NewLighting;
+		}
+	}
+
+	// For debugging
+	//DirectLighting = CulledLightsGrid.NumLocalLights / (float)ForwardGlobalLightData.MaxCulledLightsPerCell;
+	return DirectLighting;
+}
+
+float3 GetForwardDirectLightingForVertexLighting(uint GridIndex, float3 WorldPosition, float3 WorldNormal, uint EyeIndex)
+{
+	float3 DirectLighting = 0;
+	// Using white for diffuse color, real diffuse color will be incorporated per-pixel
+	float3 DiffuseColor = 1.0f;
+
+	BRANCH
+	if (ForwardGlobalLightData.HasDirectionalLight)
+	{
+		float3 N = WorldNormal;
+		float3 L = ForwardGlobalLightData.DirectionalLightDirection;
+		float NoL = saturate(dot(N, L));
+
+		float3 LightColor = ForwardGlobalLightData.DirectionalLightColor;
+	
+		#if NON_DIRECTIONAL_DIRECT_LIGHTING
+			NoL = 1.0f;
+		#endif
+
+		// No specular for vertex lighting
+		float3 DiffuseLighting = Diffuse_Lambert(DiffuseColor);
+		DirectLighting += LightColor * NoL * (DiffuseLighting);
+	}
+	
+	const FCulledLightsGridData CulledLightsGrid = GetCulledLightsGrid(GridIndex, EyeIndex);
+
+	LOOP
+	for (uint LocalLightListIndex = 0; LocalLightListIndex < CulledLightsGrid.NumLocalLights; LocalLightListIndex++)
+	{
+		const FLocalLightData LocalLight = GetLocalLightData(CulledLightsGrid.DataStartIndex + LocalLightListIndex, EyeIndex);
+		 
+		FSimpleDeferredLightData LightData = (FSimpleDeferredLightData)0;
+		LightData.LightPositionAndInvRadius = LocalLight.LightPositionAndInvRadius;
+		LightData.LightColorAndFalloffExponent = LocalLight.LightColorAndFalloffExponent;
+		LightData.bInverseSquared = LightData.LightColorAndFalloffExponent.w == 0;
+						
+		// No specular for vertex lighting
+		float3 CameraVector = 0;
+		float3 SpecularColor = 0;
+		float Roughness = 1.0f;
+		DirectLighting += GetSimpleDynamicLighting(WorldPosition, CameraVector, WorldNormal, 1, DiffuseColor, SpecularColor, Roughness, LightData);
+	}
+
+	return DirectLighting;
+}
+
+uint MortonCode( uint x )
+{
+	//x = (x ^ (x <<  8)) & 0x00ff00ff;
+	//x = (x ^ (x <<  4)) & 0x0f0f0f0f;
+	x = (x ^ (x <<  2)) & 0x33333333;
+	x = (x ^ (x <<  1)) & 0x55555555;
+	return x;
+}
+
+// Translucency Surface per-pixel uses blended reflection captures by default in the deferred renderer
+// Have to opt-in to blended reflection captures in the forward renderer
+#define USE_BLENDED_REFLECTION_CAPTURES_DEFERRED (!FORWARD_SHADING && (TRANSLUCENCY_LIGHTING_SURFACE_FORWARDSHADING || TRANSLUCENCY_LIGHTING_SURFACE_LIGHTINGVOLUME))
+#define USE_BLENDED_REFLECTION_CAPTURES_FORWARD (FORWARD_SHADING && MATERIAL_HQ_FORWARD_REFLECTIONS)
+#define REFLECTION_COMPOSITE_USE_BLENDED_REFLECTION_CAPTURES ((USE_BLENDED_REFLECTION_CAPTURES_DEFERRED || USE_BLENDED_REFLECTION_CAPTURES_FORWARD) && FEATURE_LEVEL >= FEATURE_LEVEL_SM5)
+#define CompositeTileReflectionCaptureIndices CulledLightDataGrid
+#define InstancedCompositeTileReflectionCaptureIndices InstancedCulledLightDataGrid
+#define REFLECTION_COMPOSITE_SUPPORT_SKYLIGHT_BLEND 0
+#define REFLECTION_COMPOSITE_HAS_BOX_CAPTURES 1
+#define REFLECTION_COMPOSITE_HAS_SPHERE_CAPTURES 1
+#include "ReflectionEnvironmentComposite.ush"
+
+half3 GetImageBasedReflectionLighting(FMaterialPixelParameters MaterialParameters, half Roughness, half3 SpecularColor, half IndirectIrradiance, uint GridIndex, uint EyeIndex)
+{
+	float3 N = MaterialParameters.WorldNormal;
+	float3 V = MaterialParameters.CameraVector;
+
+	float3 RayDirection = 2 * dot( V, N ) * N - V;
+	half NoV = saturate(dot(N, V));
+
+	uint NumLocalReflectionCaptures = 0;
+	uint DataStartIndex = 0;
+
+#if REFLECTION_COMPOSITE_USE_BLENDED_REFLECTION_CAPTURES
+	BRANCH
+	if (EyeIndex == 0)
+	{
+		uint NumCulledEntryIndex = (ForwardGlobalLightData.NumGridCells + GridIndex) * NUM_CULLED_LIGHTS_GRID_STRIDE;
+		NumLocalReflectionCaptures = NumCulledLightsGrid[NumCulledEntryIndex + 0];
+		DataStartIndex = NumCulledLightsGrid[NumCulledEntryIndex + 1];
+	}
+	else
+	{
+		uint NumCulledEntryIndex = (InstancedForwardGlobalLightData.NumGridCells + GridIndex) * NUM_CULLED_LIGHTS_GRID_STRIDE;
+		NumLocalReflectionCaptures = InstancedNumCulledLightsGrid[NumCulledEntryIndex + 0];
+		DataStartIndex = InstancedNumCulledLightsGrid[NumCulledEntryIndex + 1];
+	}
+#endif
+
+	float3 SpecularIBL = CompositeReflectionCapturesAndSkylight(1.0f, MaterialParameters.AbsoluteWorldPosition, RayDirection, Roughness, IndirectIrradiance, 1.0f, 0.0f, NumLocalReflectionCaptures, DataStartIndex, EyeIndex);
+
+	// Factors derived from EnvBRDFApprox( SpecularColor, 1, 1 ) == SpecularColor * 0.4524 - 0.0024
+	float3 SpecularBounce = 0.45f * SpecularColor * IndirectIrradiance;
+	// Replace reflection captures with indirect diffuse when we're rendering to a reflection capture, to avoid a feedback loop
+	SpecularIBL.rgb = lerp(SpecularIBL.rgb, SpecularBounce, ResolvedView.RenderingReflectionCaptureMask);
+
+#if MATERIAL_SSR && !FORWARD_SHADING
+	if( View.CameraCut == 0 )
+	{
+		float StepOffset = InterleavedGradientNoise( MaterialParameters.SvPosition.xy, View.StateFrameIndexMod8 );
+		StepOffset -= 0.5;
+		
+		float4 HitUVzTime;
+		float HCBLevel;
+
+		RayCast(
+			HZBTexture, HZBSampler, float2(1, 1),
+			MaterialParameters.WorldPosition_CamRelative, RayDirection, 0, 0, MaterialParameters.ScreenPosition.w,
+			12, StepOffset,
+			HitUVzTime, HCBLevel
+		);
+
+		// if there was a hit
+		BRANCH if( HitUVzTime.w < 1 )
+		{
+			float4 SSR = SampleScreenColor( PrevSceneColor, PrevSceneColorSampler,  ReprojectHit( HitUVzTime.xyz ) );
+			SSR *= saturate( 2 - 6.6 * Roughness );
+			SpecularIBL.rgb = SpecularIBL.rgb * (1 - SSR.a) + SSR.rgb;
+		}
+	}
+#endif
+
+	SpecularColor = EnvBRDFApprox(SpecularColor, Roughness, NoV);
+
+	float3 SpecularLighting = SpecularIBL.rgb;
+
+	// Have to opt-in to receiving planar reflections with forward shading
+#if !FORWARD_SHADING || MATERIAL_PLANAR_FORWARD_REFLECTIONS
+	// Plane normal will be zero if the feature is disabled
+	BRANCH
+	if (abs(dot(ReflectionPlane.xyz, 1)) > .0001f)
+	{
+		// Reuse ReflectionCubemapSampler to avoid reducing the sampler count available to artists
+		float4 PlanarReflection = ComputePlanarReflections(MaterialParameters.AbsoluteWorldPosition, MaterialParameters.WorldNormal, Roughness, ReflectionCubemapSampler);
+		// Planar reflections win over SSR and reflection environment
+		SpecularLighting = PlanarReflection.rgb + (1 - PlanarReflection.a) * SpecularLighting;
+	}
+#endif
+
+	return SpecularLighting * SpecularColor;
+}
+
+half3 GetImageBasedReflectionLighting(FMaterialPixelParameters MaterialParameters, half Roughness, half3 SpecularColor, half IndirectIrradiance, uint GridIndex)
+{
+	return GetImageBasedReflectionLighting(MaterialParameters, Roughness, SpecularColor, IndirectIrradiance, GridIndex, 0);
+}
\ No newline at end of file
