Branch snapshot for CL 3555263

diff --git a/Engine/Shaders/Private/GlobalDistanceFieldShared.ush b/Engine/Shaders/Private/GlobalDistanceFieldShared.ush
new file mode 100644
--- /dev/null
+++ b/Engine/Shaders/Private/GlobalDistanceFieldShared.ush
@@ -0,0 +1,263 @@
+// Copyright 1998-2017 Epic Games, Inc. All Rights Reserved.
+
+/*=============================================================================
+	GlobalDistanceFieldShared.usf
+=============================================================================*/
+
+#if FEATURE_LEVEL >= FEATURE_LEVEL_SM5
+
+#if IS_MATERIAL_SHADER
+#define MaxGlobalDistance				View.MaxGlobalDistance_UB
+#else
+float MaxGlobalDistance;
+#endif
+
+//@todo - get the clipmap branches to compile under gl
+#if SM5_PROFILE || PS4_PROFILE || METAL_SM5_PROFILE || GL4_PROFILE || VULKAN_PROFILE_SM5
+
+// Must match C++
+#define MAX_GLOBAL_DF_CLIPMAPS 4
+
+#if IS_MATERIAL_SHADER
+    // for materials, these are in the view UB
+	#define GlobalDistanceFieldTexture0		View.GlobalDistanceFieldTexture0_UB
+	#define GlobalDistanceFieldTexture1		View.GlobalDistanceFieldTexture1_UB
+	#define GlobalDistanceFieldTexture2		View.GlobalDistanceFieldTexture2_UB
+	#define GlobalDistanceFieldTexture3		View.GlobalDistanceFieldTexture3_UB
+	#define GlobalDistanceFieldSampler0		View.GlobalDistanceFieldSampler0_UB
+	#define GlobalDistanceFieldSampler1		View.GlobalDistanceFieldSampler1_UB
+	#define GlobalDistanceFieldSampler2		View.GlobalDistanceFieldSampler2_UB
+	#define GlobalDistanceFieldSampler3		View.GlobalDistanceFieldSampler3_UB
+	#define GlobalVolumeCenterAndExtent		View.GlobalVolumeCenterAndExtent_UB
+	#define GlobalVolumeWorldToUVAddAndMul	View.GlobalVolumeWorldToUVAddAndMul_UB
+	#define GlobalVolumeDimension			View.GlobalVolumeDimension_UB
+	#define GlobalVolumeTexelSize			View.GlobalVolumeTexelSize_UB
+#else
+	// these are only used for the precomputation shaders; which don't have a view UB
+	Texture3D GlobalDistanceFieldTexture0;
+	Texture3D GlobalDistanceFieldTexture1;
+	Texture3D GlobalDistanceFieldTexture2;
+	Texture3D GlobalDistanceFieldTexture3;
+	SamplerState GlobalDistanceFieldSampler0;
+	SamplerState GlobalDistanceFieldSampler1;
+	SamplerState GlobalDistanceFieldSampler2;
+	SamplerState GlobalDistanceFieldSampler3;
+	float4 GlobalVolumeCenterAndExtent[MAX_GLOBAL_DF_CLIPMAPS];
+	float4 GlobalVolumeWorldToUVAddAndMul[MAX_GLOBAL_DF_CLIPMAPS];
+	float GlobalVolumeDimension;
+	float GlobalVolumeTexelSize;
+#endif
+
+#if SUPPORTS_INDEPENDENT_SAMPLERS
+	#define SharedGlobalDistanceFieldSampler0 GlobalDistanceFieldSampler0
+	#define SharedGlobalDistanceFieldSampler1 GlobalDistanceFieldSampler0
+	#define SharedGlobalDistanceFieldSampler2 GlobalDistanceFieldSampler0
+	#define SharedGlobalDistanceFieldSampler3 GlobalDistanceFieldSampler0
+#else
+	#define SharedGlobalDistanceFieldSampler0 GlobalDistanceFieldSampler0
+	#define SharedGlobalDistanceFieldSampler1 GlobalDistanceFieldSampler1
+	#define SharedGlobalDistanceFieldSampler2 GlobalDistanceFieldSampler2
+	#define SharedGlobalDistanceFieldSampler3 GlobalDistanceFieldSampler3
+#endif
+
+float4 SampleGlobalDistanceField(int ClipmapIndex, float3 UV)
+{
+	if (ClipmapIndex == 0)
+	{
+		return Texture3DSampleLevel(GlobalDistanceFieldTexture0, SharedGlobalDistanceFieldSampler0, UV, 0);
+	}
+	else if (ClipmapIndex == 1)
+	{
+		return Texture3DSampleLevel(GlobalDistanceFieldTexture1, SharedGlobalDistanceFieldSampler1, UV, 0);
+	}
+	else if (ClipmapIndex == 2)
+	{
+		return Texture3DSampleLevel(GlobalDistanceFieldTexture2, SharedGlobalDistanceFieldSampler2, UV, 0);
+	}
+	else
+	{
+		return Texture3DSampleLevel(GlobalDistanceFieldTexture3, SharedGlobalDistanceFieldSampler3, UV, 0);
+	}
+}
+
+float3 ComputeGlobalUV(float3 WorldPosition, uint ClipmapIndex)
+{
+	//return ((WorldPosition - GlobalVolumeCenterAndExtent[ClipmapIndex].xyz + GlobalVolumeScollOffset[ClipmapIndex].xyz) / (GlobalVolumeCenterAndExtent[ClipmapIndex].w * 2) + .5f);
+	float4 WorldToUVAddAndMul = GlobalVolumeWorldToUVAddAndMul[ClipmapIndex];
+	return WorldPosition * WorldToUVAddAndMul.www + WorldToUVAddAndMul.xyz;
+}
+
+float GetDistanceToNearestSurfaceGlobalClipmap(float3 WorldPosition, uint ClipmapIndex, float OuterClipmapFade)
+{
+	float3 GlobalUV = ComputeGlobalUV(WorldPosition, ClipmapIndex);
+	float DistanceToSurface = 0;
+	if (ClipmapIndex == 0)
+	{
+		DistanceToSurface = Texture3DSampleLevel(GlobalDistanceFieldTexture0, SharedGlobalDistanceFieldSampler0, GlobalUV, 0).x;
+	}
+	else if (ClipmapIndex == 1)
+	{
+		DistanceToSurface = Texture3DSampleLevel(GlobalDistanceFieldTexture1, SharedGlobalDistanceFieldSampler1, GlobalUV, 0).x;
+	}
+	else if (ClipmapIndex == 2)
+	{
+		DistanceToSurface = Texture3DSampleLevel(GlobalDistanceFieldTexture2, SharedGlobalDistanceFieldSampler2, GlobalUV, 0).x;
+	}
+	else if (ClipmapIndex == 3)
+	{
+		DistanceToSurface = Texture3DSampleLevel(GlobalDistanceFieldTexture3, SharedGlobalDistanceFieldSampler3, GlobalUV, 0).x;
+		DistanceToSurface = lerp(MaxGlobalDistance, DistanceToSurface, OuterClipmapFade);
+	}
+	return DistanceToSurface;
+}
+
+float GetDistanceToNearestSurfaceGlobal(float3 WorldPosition)
+{
+	float DistanceToSurface = MaxGlobalDistance;
+	float DistanceFromClipmap = ComputeDistanceFromBoxToPointInside(GlobalVolumeCenterAndExtent[0].xyz, GlobalVolumeCenterAndExtent[0].www, WorldPosition);
+
+	//@todo - would be nice to atlas to get rid of branches and multiple samplers but the partial update scrolling relies on wrap addressing
+	BRANCH
+	if (DistanceFromClipmap > GlobalVolumeCenterAndExtent[0].w * GlobalVolumeTexelSize)
+	{
+		DistanceToSurface = GetDistanceToNearestSurfaceGlobalClipmap(WorldPosition, 0, 0);
+	}
+	else
+	{
+		DistanceFromClipmap = ComputeDistanceFromBoxToPointInside(GlobalVolumeCenterAndExtent[1].xyz, GlobalVolumeCenterAndExtent[1].www, WorldPosition);
+
+		BRANCH
+		if (DistanceFromClipmap > GlobalVolumeCenterAndExtent[1].w * GlobalVolumeTexelSize)
+		{
+			DistanceToSurface = GetDistanceToNearestSurfaceGlobalClipmap(WorldPosition, 1, 0);
+		}
+		else
+		{
+			DistanceFromClipmap = ComputeDistanceFromBoxToPointInside(GlobalVolumeCenterAndExtent[2].xyz, GlobalVolumeCenterAndExtent[2].www, WorldPosition);
+			float DistanceFromLastClipmap = ComputeDistanceFromBoxToPointInside(GlobalVolumeCenterAndExtent[3].xyz, GlobalVolumeCenterAndExtent[3].www, WorldPosition);
+
+			BRANCH
+			if (DistanceFromClipmap > GlobalVolumeCenterAndExtent[2].w * GlobalVolumeTexelSize)
+			{
+				DistanceToSurface = GetDistanceToNearestSurfaceGlobalClipmap(WorldPosition, 2, 0);
+			}
+			else if (DistanceFromLastClipmap > GlobalVolumeCenterAndExtent[3].w * GlobalVolumeTexelSize)
+			{
+				// Fade to max distance in the last 10% of the largest clipmap
+				float OuterClipmapFade = saturate(DistanceFromLastClipmap * 10 * GlobalVolumeWorldToUVAddAndMul[3].w);
+				DistanceToSurface = GetDistanceToNearestSurfaceGlobalClipmap(WorldPosition, 3, OuterClipmapFade);
+			}
+		}
+	}
+
+	return DistanceToSurface;
+}
+
+float3 GetDistanceFieldGradientGlobalClipmap(float3 WorldPosition, uint ClipmapIndex)
+{
+	float3 GlobalUV = ComputeGlobalUV(WorldPosition, ClipmapIndex);
+
+	float R = 0;
+	float L = 0;
+	float F = 0;
+	float B = 0;
+	float U = 0;
+	float D = 0;
+
+	if (ClipmapIndex == 0)
+	{
+		R = Texture3DSampleLevel(GlobalDistanceFieldTexture0, SharedGlobalDistanceFieldSampler0, float3(GlobalUV.x + GlobalVolumeTexelSize, GlobalUV.y, GlobalUV.z), 0).x;
+		L = Texture3DSampleLevel(GlobalDistanceFieldTexture0, SharedGlobalDistanceFieldSampler0, float3(GlobalUV.x - GlobalVolumeTexelSize, GlobalUV.y, GlobalUV.z), 0).x;
+		F = Texture3DSampleLevel(GlobalDistanceFieldTexture0, SharedGlobalDistanceFieldSampler0, float3(GlobalUV.x, GlobalUV.y + GlobalVolumeTexelSize, GlobalUV.z), 0).x;
+		B = Texture3DSampleLevel(GlobalDistanceFieldTexture0, SharedGlobalDistanceFieldSampler0, float3(GlobalUV.x, GlobalUV.y - GlobalVolumeTexelSize, GlobalUV.z), 0).x;
+		U = Texture3DSampleLevel(GlobalDistanceFieldTexture0, SharedGlobalDistanceFieldSampler0, float3(GlobalUV.x, GlobalUV.y, GlobalUV.z + GlobalVolumeTexelSize), 0).x;
+		D = Texture3DSampleLevel(GlobalDistanceFieldTexture0, SharedGlobalDistanceFieldSampler0, float3(GlobalUV.x, GlobalUV.y, GlobalUV.z - GlobalVolumeTexelSize), 0).x;
+	}
+	else if (ClipmapIndex == 1)
+	{
+		R = Texture3DSampleLevel(GlobalDistanceFieldTexture1, SharedGlobalDistanceFieldSampler1, float3(GlobalUV.x + GlobalVolumeTexelSize, GlobalUV.y, GlobalUV.z), 0).x;
+		L = Texture3DSampleLevel(GlobalDistanceFieldTexture1, SharedGlobalDistanceFieldSampler1, float3(GlobalUV.x - GlobalVolumeTexelSize, GlobalUV.y, GlobalUV.z), 0).x;
+		F = Texture3DSampleLevel(GlobalDistanceFieldTexture1, SharedGlobalDistanceFieldSampler1, float3(GlobalUV.x, GlobalUV.y + GlobalVolumeTexelSize, GlobalUV.z), 0).x;
+		B = Texture3DSampleLevel(GlobalDistanceFieldTexture1, SharedGlobalDistanceFieldSampler1, float3(GlobalUV.x, GlobalUV.y - GlobalVolumeTexelSize, GlobalUV.z), 0).x;
+		U = Texture3DSampleLevel(GlobalDistanceFieldTexture1, SharedGlobalDistanceFieldSampler1, float3(GlobalUV.x, GlobalUV.y, GlobalUV.z + GlobalVolumeTexelSize), 0).x;
+		D = Texture3DSampleLevel(GlobalDistanceFieldTexture1, SharedGlobalDistanceFieldSampler1, float3(GlobalUV.x, GlobalUV.y, GlobalUV.z - GlobalVolumeTexelSize), 0).x;
+	}
+	else if (ClipmapIndex == 2)
+	{
+		R = Texture3DSampleLevel(GlobalDistanceFieldTexture2, SharedGlobalDistanceFieldSampler2, float3(GlobalUV.x + GlobalVolumeTexelSize, GlobalUV.y, GlobalUV.z), 0).x;
+		L = Texture3DSampleLevel(GlobalDistanceFieldTexture2, SharedGlobalDistanceFieldSampler2, float3(GlobalUV.x - GlobalVolumeTexelSize, GlobalUV.y, GlobalUV.z), 0).x;
+		F = Texture3DSampleLevel(GlobalDistanceFieldTexture2, SharedGlobalDistanceFieldSampler2, float3(GlobalUV.x, GlobalUV.y + GlobalVolumeTexelSize, GlobalUV.z), 0).x;
+		B = Texture3DSampleLevel(GlobalDistanceFieldTexture2, SharedGlobalDistanceFieldSampler2, float3(GlobalUV.x, GlobalUV.y - GlobalVolumeTexelSize, GlobalUV.z), 0).x;
+		U = Texture3DSampleLevel(GlobalDistanceFieldTexture2, SharedGlobalDistanceFieldSampler2, float3(GlobalUV.x, GlobalUV.y, GlobalUV.z + GlobalVolumeTexelSize), 0).x;
+		D = Texture3DSampleLevel(GlobalDistanceFieldTexture2, SharedGlobalDistanceFieldSampler2, float3(GlobalUV.x, GlobalUV.y, GlobalUV.z - GlobalVolumeTexelSize), 0).x;
+	}
+	else if (ClipmapIndex == 3)
+	{
+		R = Texture3DSampleLevel(GlobalDistanceFieldTexture3, SharedGlobalDistanceFieldSampler3, float3(GlobalUV.x + GlobalVolumeTexelSize, GlobalUV.y, GlobalUV.z), 0).x;
+		L = Texture3DSampleLevel(GlobalDistanceFieldTexture3, SharedGlobalDistanceFieldSampler3, float3(GlobalUV.x - GlobalVolumeTexelSize, GlobalUV.y, GlobalUV.z), 0).x;
+		F = Texture3DSampleLevel(GlobalDistanceFieldTexture3, SharedGlobalDistanceFieldSampler3, float3(GlobalUV.x, GlobalUV.y + GlobalVolumeTexelSize, GlobalUV.z), 0).x;
+		B = Texture3DSampleLevel(GlobalDistanceFieldTexture3, SharedGlobalDistanceFieldSampler3, float3(GlobalUV.x, GlobalUV.y - GlobalVolumeTexelSize, GlobalUV.z), 0).x;
+		U = Texture3DSampleLevel(GlobalDistanceFieldTexture3, SharedGlobalDistanceFieldSampler3, float3(GlobalUV.x, GlobalUV.y, GlobalUV.z + GlobalVolumeTexelSize), 0).x;
+		D = Texture3DSampleLevel(GlobalDistanceFieldTexture3, SharedGlobalDistanceFieldSampler3, float3(GlobalUV.x, GlobalUV.y, GlobalUV.z - GlobalVolumeTexelSize), 0).x;
+	}
+		
+	float Extent = GlobalVolumeCenterAndExtent[ClipmapIndex].w;
+	float3 Gradient = .5f * float3(R - L, F - B, U - D) / Extent;
+	return Gradient;
+}
+
+float3 GetDistanceFieldGradientGlobal(float3 WorldPosition)
+{
+	float3 Gradient = float3(0, 0, .001f);
+	float DistanceFromClipmap = ComputeDistanceFromBoxToPointInside(GlobalVolumeCenterAndExtent[0].xyz, GlobalVolumeCenterAndExtent[0].www, WorldPosition);
+	// Don't sample near the border to avoid sampling out of bounds during partial differencing
+	float BorderTexels = GlobalVolumeTexelSize * 4;
+
+	BRANCH
+	if (DistanceFromClipmap > GlobalVolumeCenterAndExtent[0].w * BorderTexels)
+	{
+		Gradient = GetDistanceFieldGradientGlobalClipmap(WorldPosition, 0);
+	}
+	else
+	{
+		DistanceFromClipmap = ComputeDistanceFromBoxToPointInside(GlobalVolumeCenterAndExtent[1].xyz, GlobalVolumeCenterAndExtent[1].www, WorldPosition);
+
+		BRANCH
+		if (DistanceFromClipmap > GlobalVolumeCenterAndExtent[1].w * BorderTexels)
+		{
+			Gradient = GetDistanceFieldGradientGlobalClipmap(WorldPosition, 1);
+		}
+		else
+		{
+			DistanceFromClipmap = ComputeDistanceFromBoxToPointInside(GlobalVolumeCenterAndExtent[2].xyz, GlobalVolumeCenterAndExtent[2].www, WorldPosition);
+			float DistanceFromLastClipmap = ComputeDistanceFromBoxToPointInside(GlobalVolumeCenterAndExtent[3].xyz, GlobalVolumeCenterAndExtent[3].www, WorldPosition);
+
+			BRANCH
+			if (DistanceFromClipmap > GlobalVolumeCenterAndExtent[2].w * BorderTexels)
+			{
+				Gradient = GetDistanceFieldGradientGlobalClipmap(WorldPosition, 2);
+			}
+			else if (DistanceFromLastClipmap > GlobalVolumeCenterAndExtent[3].w * BorderTexels)
+			{
+				Gradient = GetDistanceFieldGradientGlobalClipmap(WorldPosition, 3);
+			}
+		}
+	}
+
+	return Gradient;
+}
+
+#else
+
+float GetDistanceToNearestSurfaceGlobal(float3 WorldPosition)
+{
+	return MaxGlobalDistance;
+}
+
+float3 GetDistanceFieldGradientGlobal(float3 WorldPosition)
+{
+	return float3(0, 0, .001f);
+}
+
+#endif
+#endif // FEATURE_LEVEL >= FEATURE_LEVEL_SM5
diff --git a/Engine/Shaders/Private/GpuSkinCommon.ush b/Engine/Shaders/Private/GpuSkinCommon.ush
new file mode 100644
--- /dev/null
+++ b/Engine/Shaders/Private/GpuSkinCommon.ush
@@ -0,0 +1,98 @@
+//// Copyright 1998-2017 Epic Games, Inc. All Rights Reserved.
+
+/*=============================================================================
+	GpuSkinCacheCommon.usf: Common properties and defines for GPU Skinning Cache
+=============================================================================*/
+
+// These offsets are in uint32/floats; they attempt to match TGPUSkinVertexBase<> & derived classes:
+//	TGPUSkinVertexBase
+//	TGPUSkinVertexFloat16Uvs<1..N>
+//	TGPUSkinVertexFloat32Uvs<1..N>
+
+// As we know the properties of the Vertices, the layout looks like this (offset, size in u32):
+//	MEMBER				TGPUSkinVertexBase	TGPUSkinVertexFloat16Uvs<1..N>	TGPUSkinVertexFloat32Uvs<1..N>
+//	TangentX			0,1					0,1								0,1
+//	TangentZ			1,1					1,1								1,1
+//	Position			2,3					2,3								2,3
+//	UVs[1..N]			N/A					5,1..N							5,(1..N)*2
+
+// TSkinWeightInfo<true>
+// TSkinWeightInfo<false>
+
+// Members of weight array
+//	MEMBER				TSkinWeightInfo<true>		TSkinWeightInfo<false>
+//	InfluenceBones		0,1							0,2	
+//	InfluenceWeights	1,1							2,2	
+
+#define	GPUSKIN_VB_OFFSET_TANGENT_X			0
+#define	GPUSKIN_VB_OFFSET_TANGENT_Z			1
+#define	GPUSKIN_VB_OFFSET_POSITION			2
+#define	GPUSKIN_VB_OFFSET_UVS				5
+
+#define	GPUSKIN_VB_OFFSET_INFLUENCEBONES	0
+#define	GPUSKIN_VB_OFFSET_INFLUENCEWEIGHTS	(1 + GPUSKIN_USE_EXTRA_INFLUENCES)
+
+
+// Should be set from FGPUSkinCacheCS!
+#ifndef GPUSKIN_RWBUFFER_OFFSET_POSITION
+#define GPUSKIN_RWBUFFER_OFFSET_POSITION	0
+#endif
+
+#ifndef GPUSKIN_RWBUFFER_OFFSET_TANGENT_X
+#define GPUSKIN_RWBUFFER_OFFSET_TANGENT_X	3
+#endif
+
+#ifndef GPUSKIN_RWBUFFER_OFFSET_TANGENT_Z
+#define GPUSKIN_RWBUFFER_OFFSET_TANGENT_Z	4
+#endif
+
+// Must match RWStrideInFloats
+#ifndef GPUSKIN_RWBUFFER_NUM_FLOATS
+#define GPUSKIN_RWBUFFER_NUM_FLOATS			5
+#endif
+
+// @return 0..255
+uint4 UnpackU8x4(uint Packed)
+{
+	uint4 Unpacked;
+	Unpacked.x = (Packed >> 0) & 0xff;
+	Unpacked.y = (Packed >> 8) & 0xff;
+	Unpacked.z = (Packed >> 16) & 0xff;
+	Unpacked.w = (Packed >> 24) & 0xff;
+	return Unpacked;
+}
+
+// @return 0..1
+float4 UnpackU8x4N(uint Packed)
+{
+	return UnpackU8x4(Packed) / 255.0f;
+}
+
+// @return -1..1
+float4 UnpackS8x4N(uint Packed)
+{
+	return UnpackU8x4(Packed) / 127.5f - 1;
+}
+
+// @param Unpacked -1..1
+// @return layout 0xWWZZYYXX where 0..0xff is mapped to -1..1
+uint PackS8x4N(float4 Unpacked)
+{
+	uint Packed = asuint(
+		(uint((Unpacked.x + 1.0f) * 127.499f) << 0) |
+		(uint((Unpacked.y + 1.0f) * 127.499f) << 8) |
+		(uint((Unpacked.z + 1.0f) * 127.499f) << 16) |
+		(uint((Unpacked.w + 1.0f) * 127.499f) << 24));
+	return Packed;
+}
+
+// @param Unpacked -1..1
+// @return layout 0x00ZZYYXX where 0..0xff is mapped to -1..1
+uint PackS8x3N(float3 Unpacked)
+{
+	uint Packed = asuint(
+		(uint((Unpacked.x + 1.0f) * 127.499f) << 0) |
+		(uint((Unpacked.y + 1.0f) * 127.499f) << 8) |
+		(uint((Unpacked.z + 1.0f) * 127.499f) << 16));
+	return Packed;
+}
\ No newline at end of file
diff --git a/Engine/Shaders/Private/GpuSkinVertexFactory.ush b/Engine/Shaders/Private/GpuSkinVertexFactory.ush
new file mode 100644
--- /dev/null
+++ b/Engine/Shaders/Private/GpuSkinVertexFactory.ush
@@ -0,0 +1,693 @@
+// Copyright 1998-2017 Epic Games, Inc. All Rights Reserved.
+
+/*=============================================================================
+	GpuSkinVertexFactory.hlsl: GPU skinned vertex factory shader code
+=============================================================================*/
+
+
+#include "VertexFactoryCommon.ush"
+#include "GpuSkinCommon.ush"
+
+#ifndef GPUSKIN_MORPH_BLEND
+#define GPUSKIN_MORPH_BLEND 0
+#endif
+
+float4x4 PreviousLocalToWorld;
+
+#ifndef GPUSKIN_APEX_CLOTH
+#define GPUSKIN_APEX_CLOTH 0
+#endif
+
+#ifndef GPUSKIN_LIMIT_2BONE_INFLUENCES
+#define GPUSKIN_LIMIT_2BONE_INFLUENCES 0 // default is 4 bones
+#endif
+
+#ifndef GPUSKIN_USE_EXTRA_INFLUENCES
+#define GPUSKIN_USE_EXTRA_INFLUENCES 0
+#endif
+
+#define FIXED_VERTEX_INDEX 0xFFFF
+
+float3 MeshOrigin;
+float3 MeshExtension;
+// if the per bone motion blur is enabled for this draw call
+bool PerBoneMotionBlur;
+
+#if GPUSKIN_APEX_CLOTH
+/** Vertex buffer from which to read simulated positions of clothing. */
+Buffer<float2> ClothSimulVertsPositionsNormals;
+Buffer<float2> PreviousClothSimulVertsPositionsNormals;
+/** blend weight between simulated positions and original key-framed animation */
+float ClothBlendWeight;
+#endif// #if GPUSKIN_APEX_CLOTH
+
+struct FVertexFactoryInput
+{
+	float4	Position		: ATTRIBUTE0;
+	// 0..1
+	half3	TangentX		: ATTRIBUTE1;
+	// 0..1
+	// TangentZ.w contains sign of tangent basis determinant
+	half4	TangentZ		: ATTRIBUTE2;	
+
+#if FEATURE_LEVEL >= FEATURE_LEVEL_ES3_1 || COMPILER_METAL
+	uint4	BlendIndices	: ATTRIBUTE3;
+	#if GPUSKIN_USE_EXTRA_INFLUENCES
+	uint4	BlendIndicesExtra	: ATTRIBUTE14;
+	#endif
+#else
+	// Continue using int for SM3, compatibility of uint is unknown across SM3 platforms
+	int4	BlendIndices	: ATTRIBUTE3;
+	#if GPUSKIN_USE_EXTRA_INFLUENCES
+	int4	BlendIndicesExtra	: ATTRIBUTE14;
+	#endif
+#endif
+	float4	BlendWeights	: ATTRIBUTE4;	
+#if GPUSKIN_USE_EXTRA_INFLUENCES
+	float4	BlendWeightsExtra	: ATTRIBUTE15;
+#endif
+
+
+#if NUM_MATERIAL_TEXCOORDS_VERTEX
+	// If this changes make sure to update LocalVertexFactory.usf
+	float2	TexCoords[NUM_MATERIAL_TEXCOORDS_VERTEX] : ATTRIBUTE5;
+
+	#if NUM_MATERIAL_TEXCOORDS_VERTEX > 4
+		#error Too many texture coordinate sets defined on GPUSkin vertex input. Max: 4.
+	#endif
+#endif
+
+#if GPUSKIN_MORPH_BLEND
+	// NOTE: TEXCOORD6,TEXCOORD7 used instead of POSITION1,NORMAL1 since those semantics are not supported by Cg 
+	/** added to the Position */
+	float3	DeltaPosition	: ATTRIBUTE9;	//POSITION1;
+	/** added to the TangentZ and then used to derive new TangentX,TangentY, .w contains the weight of the tangent blend */
+	float3	DeltaTangentZ	: ATTRIBUTE10;	//NORMAL1;
+#elif GPUSKIN_APEX_CLOTH // exclusive with GPUSKIN_MORPH_BLEND
+	// APEX cloth mesh-mesh mapping data
+	// Barycentric Coordinate Data
+	float4	BaryCoordPos	: ATTRIBUTE9;	
+	float4	BaryCoordNormal	: ATTRIBUTE10;	
+	float4	BaryCoordTangent: ATTRIBUTE11;
+	uint4	SimulIndices	: ATTRIBUTE12;
+#endif
+	/** Per vertex color */
+	float4 Color : ATTRIBUTE13;
+
+	/** Optional instance ID for vertex layered rendering */
+#if FEATURE_LEVEL >= FEATURE_LEVEL_SM4 && ONEPASS_POINTLIGHT_SHADOW && USING_VERTEX_SHADER_LAYER
+	uint InstanceId	: SV_InstanceID;
+#endif
+};
+
+struct FVertexFactoryInterpolantsVSToPS
+{
+	TANGENTTOWORLD_INTERPOLATOR_BLOCK
+
+	float4 Color : COLOR0;
+
+#if NUM_TEX_COORD_INTERPOLATORS
+	float2	TexCoords[NUM_TEX_COORD_INTERPOLATORS]	: TEXCOORD0;
+#endif
+
+#if INSTANCED_STEREO
+	nointerpolation uint PackedEyeIndex : PACKED_EYE_INDEX;
+#endif
+};
+
+struct FVertexFactoryInterpolantsVSToDS
+{
+	FVertexFactoryInterpolantsVSToPS InterpolantsVSToPS;
+};
+
+/** Converts from vertex factory specific interpolants to a FMaterialPixelParameters, which is used by material inputs. */
+FMaterialPixelParameters GetMaterialPixelParameters(FVertexFactoryInterpolantsVSToPS Interpolants, float4 SvPosition)
+{
+	// GetMaterialPixelParameters is responsible for fully initializing the result
+	FMaterialPixelParameters Result = MakeInitializedMaterialPixelParameters();
+
+#if NUM_TEX_COORD_INTERPOLATORS	
+	UNROLL
+	for(int CoordinateIndex = 0;CoordinateIndex < NUM_TEX_COORD_INTERPOLATORS;CoordinateIndex++)
+	{
+		Result.TexCoords[CoordinateIndex] = Interpolants.TexCoords[CoordinateIndex];
+	}
+#endif
+
+	half3 TangentToWorld0 = Interpolants.TangentToWorld0.xyz;
+	half4 TangentToWorld2 = Interpolants.TangentToWorld2;
+	Result.TangentToWorld = AssembleTangentToWorld( TangentToWorld0, TangentToWorld2 );
+	Result.UnMirrored = TangentToWorld2.w;	
+	Result.VertexColor = Interpolants.Color;
+	Result.TwoSidedSign = 1;
+	return Result;
+}
+
+#if NEEDS_LIGHTMAP_COORDINATE
+void GetLightMapCoordinates(FVertexFactoryInterpolantsVSToPS Interpolants, out float2 LightmapUV0, out float2 LightmapUV1)
+{
+	LightmapUV0 = LightmapUV1 = 0;
+}
+
+half2 GetShadowMapCoordinate(FVertexFactoryInterpolantsVSToPS Interpolants)
+{
+	return 0;
+}
+#endif
+
+#define FBoneMatrix float3x4
+
+#if FEATURE_LEVEL >= FEATURE_LEVEL_ES3_1
+
+// The bone matrix buffer stored as 4x3 (3 float4 behind each other), all chunks of a skeletal mesh in one, each skeletal mesh has it's own buffer
+Buffer<float4> BoneMatrices;
+// The previous bone matrix buffer stored as 4x3 (3 float4 behind each other), all chunks of a skeletal mesh in one, each skeletal mesh has it's own buffer
+Buffer<float4> PreviousBoneMatrices;
+
+#endif
+
+// Cache data to avoid multiple calculation 
+struct FVertexFactoryIntermediates
+{
+	// Blend Matrix (used for position/tangents)
+	FBoneMatrix BlendMatrix;
+
+	// Unpacked position
+	float3 UnpackedPosition;
+	
+	// Tangent Basis
+	float3x3 TangentToLocal; 
+
+	// Vertex Color
+	float4 Color;
+
+#if GPUSKIN_APEX_CLOTH
+	// in world space (non translated)
+	float3 SimulatedPosition;
+	// in world space (non translated)
+	float3 PreviousSimulatedPosition;
+#endif
+
+};
+
+/** Converts from vertex factory specific input to a FMaterialVertexParameters, which is used by vertex shader material inputs. */
+FMaterialVertexParameters GetMaterialVertexParameters(FVertexFactoryInput Input, FVertexFactoryIntermediates Intermediates, float3 WorldPosition, float3x3 TangentToLocal)
+{
+	FMaterialVertexParameters Result = (FMaterialVertexParameters)0;
+	Result.WorldPosition = WorldPosition;
+	Result.VertexColor = Intermediates.Color;
+	Result.TangentToWorld = mul(TangentToLocal, GetLocalToWorld3x3());
+	Result.PreSkinnedPosition = Intermediates.UnpackedPosition.xyz;
+	Result.PreSkinnedNormal = Input.TangentZ.xyz * 2.f - 1.f;
+
+#if NUM_MATERIAL_TEXCOORDS_VERTEX
+	for(int CoordinateIndex = 0; CoordinateIndex < NUM_MATERIAL_TEXCOORDS_VERTEX; CoordinateIndex++)
+	{
+		Result.TexCoords[CoordinateIndex] = Input.TexCoords[CoordinateIndex];
+	}
+#endif
+	return Result;
+}
+
+#if GPUSKIN_APEX_CLOTH
+
+// if false, fixed vertex
+bool IsSimulatedVertex( FVertexFactoryInput Input )
+{
+	return (Input.SimulIndices.w < FIXED_VERTEX_INDEX);
+}
+
+float3 GetClothSimulPosition(int Index, bool bPrevious)
+{
+#if FEATURE_LEVEL >= FEATURE_LEVEL_SM4
+	if(bPrevious)
+	{
+		return float3(PreviousClothSimulVertsPositionsNormals[Index * 3], PreviousClothSimulVertsPositionsNormals[Index * 3 + 1].x);
+	}
+	else
+	{
+		return float3(ClothSimulVertsPositionsNormals[Index * 3], ClothSimulVertsPositionsNormals[Index * 3 + 1].x);
+	}
+#else
+	return APEXClothParam.Positions[Index];
+#endif //FEATURE_LEVEL >= FEATURE_LEVEL_SM4
+}
+
+float3 GetClothSimulNormal(int Index, bool bPrevious)
+{
+#if FEATURE_LEVEL >= FEATURE_LEVEL_SM4
+	if(bPrevious)
+	{
+		return float3(PreviousClothSimulVertsPositionsNormals[Index * 3 + 1].y, PreviousClothSimulVertsPositionsNormals[Index * 3 + 2]);
+	}
+	else
+	{
+		return float3(ClothSimulVertsPositionsNormals[Index * 3 + 1].y, ClothSimulVertsPositionsNormals[Index * 3 + 2]);
+	}
+#else
+	return APEXClothParam.Normals[Index];
+#endif //FEATURE_LEVEL >= FEATURE_LEVEL_SM4
+
+}
+
+float3 ClothingPosition( FVertexFactoryInput Input, bool bPrevious)
+{
+	float3 Position = Input.BaryCoordPos.x*(GetClothSimulPosition(Input.SimulIndices.x, bPrevious)+GetClothSimulNormal(Input.SimulIndices.x, bPrevious)*Input.BaryCoordPos.w)
+					+ Input.BaryCoordPos.y*(GetClothSimulPosition(Input.SimulIndices.y, bPrevious)+GetClothSimulNormal(Input.SimulIndices.y, bPrevious)*Input.BaryCoordPos.w)
+					+ Input.BaryCoordPos.z*(GetClothSimulPosition(Input.SimulIndices.z, bPrevious)+GetClothSimulNormal(Input.SimulIndices.z, bPrevious)*Input.BaryCoordPos.w);
+
+	return float3(Position.xyz*MeshExtension + MeshOrigin);
+}
+
+#endif //#if GPUSKIN_APEX_CLOTH
+
+/**
+* Unpack position - uncompress xyz position to world position
+*/
+float3 UnpackedPosition( FVertexFactoryInput Input )
+{
+		return float3(Input.Position.xyz*MeshExtension + MeshOrigin);
+}
+
+#if GPUSKIN_MORPH_BLEND
+/** 
+* Adds the delta position from the combined morph targets to the vertex position 
+*/
+float3 MorphPosition( FVertexFactoryInput Input, FVertexFactoryIntermediates Intermediates )
+{
+	return Intermediates.UnpackedPosition + Input.DeltaPosition;
+}
+#endif
+
+FBoneMatrix GetBoneMatrix(int Index)
+{
+#if FEATURE_LEVEL >= FEATURE_LEVEL_ES3_1
+	float4 A = BoneMatrices[Index * 3];
+	float4 B = BoneMatrices[Index * 3 + 1];
+	float4 C = BoneMatrices[Index * 3 + 2];
+	return FBoneMatrix(A,B,C);
+#else
+	return Bones.BoneMatrices[Index];
+#endif
+}
+
+FBoneMatrix CalcBoneMatrix( FVertexFactoryInput Input )
+{
+	FBoneMatrix BoneMatrix = Input.BlendWeights.x * GetBoneMatrix(Input.BlendIndices.x);
+	BoneMatrix += Input.BlendWeights.y * GetBoneMatrix(Input.BlendIndices.y);
+
+#if !GPUSKIN_LIMIT_2BONE_INFLUENCES
+	BoneMatrix += Input.BlendWeights.z * GetBoneMatrix(Input.BlendIndices.z);
+	BoneMatrix += Input.BlendWeights.w * GetBoneMatrix(Input.BlendIndices.w);
+
+#if GPUSKIN_USE_EXTRA_INFLUENCES
+	BoneMatrix += Input.BlendWeightsExtra.x * GetBoneMatrix(Input.BlendIndicesExtra.x);
+	BoneMatrix += Input.BlendWeightsExtra.y * GetBoneMatrix(Input.BlendIndicesExtra.y);
+	BoneMatrix += Input.BlendWeightsExtra.z * GetBoneMatrix(Input.BlendIndicesExtra.z);
+	BoneMatrix += Input.BlendWeightsExtra.w * GetBoneMatrix(Input.BlendIndicesExtra.w);
+#endif
+#endif//GPUSKIN_LIMIT_2BONE_INFLUENCES
+	return BoneMatrix;
+}
+
+FBoneMatrix GetPreviousBoneMatrix(int Index)
+{
+#if FEATURE_LEVEL >= FEATURE_LEVEL_ES3_1
+	float4 A = PreviousBoneMatrices[Index * 3 + 0];
+	float4 B = PreviousBoneMatrices[Index * 3 + 1];
+	float4 C = PreviousBoneMatrices[Index * 3 + 2];
+	return FBoneMatrix(A,B,C);
+#else
+	return Bones.BoneMatrices[Index];
+#endif
+}
+
+FBoneMatrix CalcPreviousBoneMatrix( FVertexFactoryInput Input )
+{
+	FBoneMatrix BoneMatrix = Input.BlendWeights.x * GetPreviousBoneMatrix(Input.BlendIndices.x);
+	BoneMatrix += Input.BlendWeights.y * GetPreviousBoneMatrix(Input.BlendIndices.y);
+
+#if !GPUSKIN_LIMIT_2BONE_INFLUENCES
+	BoneMatrix += Input.BlendWeights.z * GetPreviousBoneMatrix(Input.BlendIndices.z);
+	BoneMatrix += Input.BlendWeights.w * GetPreviousBoneMatrix(Input.BlendIndices.w);
+
+#if GPUSKIN_USE_EXTRA_INFLUENCES
+	BoneMatrix += Input.BlendWeightsExtra.x * GetPreviousBoneMatrix(Input.BlendIndicesExtra.x);
+	BoneMatrix += Input.BlendWeightsExtra.y * GetPreviousBoneMatrix(Input.BlendIndicesExtra.y);
+	BoneMatrix += Input.BlendWeightsExtra.z * GetPreviousBoneMatrix(Input.BlendIndicesExtra.z);
+	BoneMatrix += Input.BlendWeightsExtra.w * GetPreviousBoneMatrix(Input.BlendIndicesExtra.w);
+#endif
+#endif//GPUSKIN_LIMIT_2BONE_INFLUENCES
+	return BoneMatrix;
+}
+
+/** transform position by weighted sum of skinning matrices */
+float3 SkinPosition( FVertexFactoryInput Input, FVertexFactoryIntermediates Intermediates )
+{
+#if GPUSKIN_MORPH_BLEND
+	float4 Position = float4(MorphPosition(Input, Intermediates),1);
+#else
+	float4 Position = float4(Intermediates.UnpackedPosition,1);
+#endif
+
+	// Note the use of mul(Matrix,Vector), bone matrices are stored transposed
+	// for tighter packing.
+	return mul( Intermediates.BlendMatrix, Position );
+}
+
+/** transform position by weighted sum of skinning matrices */
+float3 SkinPreviousPosition( FVertexFactoryInput Input, FVertexFactoryIntermediates Intermediates )
+{
+#if GPUSKIN_MORPH_BLEND
+	float4 Position = float4(MorphPosition(Input, Intermediates),1);
+#else
+	float4 Position = float4(Intermediates.UnpackedPosition,1);
+#endif
+
+	FBoneMatrix BlendMatrix = Intermediates.BlendMatrix;
+	
+	// Previous Blend Matrix (used for position in velocity rendering)
+
+	if(PerBoneMotionBlur)
+	{
+		BlendMatrix = CalcPreviousBoneMatrix( Input );
+	}
+
+	// Note the use of mul(Matrix,Vector), bone matrices are stored transposed
+	// for tighter packing.
+	return mul( BlendMatrix, Position );
+}
+/** transform the tangent basis vectors */
+float3x3 SkinTangents( FVertexFactoryInput Input, FVertexFactoryIntermediates Intermediates )
+{
+	float3x3 TangentToLocal;
+
+	// tangent
+	// -1..1
+	half3 LocalTangentX = TangentBias(Input.TangentX);
+	// -1..1 .xyz:normal, .w:contains sign of tangent basis determinant (left or right handed)
+	half4 LocalTangentZ = TangentBias(Input.TangentZ);
+
+#if GPUSKIN_MORPH_BLEND
+	// calc new normal by offseting it with the delta
+	LocalTangentZ.xyz = normalize(LocalTangentZ.xyz + Input.DeltaTangentZ);
+	// derive the new tangent by orthonormalizing the new normal against
+	// the base tangent vector (assuming these are normalized)
+	LocalTangentX = normalize(LocalTangentX - dot(LocalTangentX, LocalTangentZ.xyz) * LocalTangentZ.xyz);
+#else
+	#if GPUSKIN_APEX_CLOTH
+
+	float3x3 ClothTangent = float3x3(1,0,0, 0,1,0, 0,0,1);
+	float TempClothBlendWeight = 0.0f;
+	if( IsSimulatedVertex(Input) )
+	{ 
+		TempClothBlendWeight = ClothBlendWeight;
+
+		bool bPrevious = false;
+
+		float3 A = GetClothSimulPosition(Input.SimulIndices.x, bPrevious);
+		float3 B = GetClothSimulPosition(Input.SimulIndices.y, bPrevious);
+		float3 C = GetClothSimulPosition(Input.SimulIndices.z, bPrevious);
+
+		float3 NA = GetClothSimulNormal(Input.SimulIndices.x, bPrevious);
+		float3 NB = GetClothSimulNormal(Input.SimulIndices.y, bPrevious);
+		float3 NC = GetClothSimulNormal(Input.SimulIndices.z, bPrevious);
+
+		float3 NormalPosition = Input.BaryCoordNormal.x*(A+NA*Input.BaryCoordNormal.w)
+								+ Input.BaryCoordNormal.y*(B+NB*Input.BaryCoordNormal.w)
+								+ Input.BaryCoordNormal.z*(C+NC*Input.BaryCoordNormal.w);
+
+		float3 TangentPosition = Input.BaryCoordTangent.x*(A+NA*Input.BaryCoordTangent.w)
+								+ Input.BaryCoordTangent.y*(B+NB*Input.BaryCoordTangent.w)
+								+ Input.BaryCoordTangent.z*(C+NC*Input.BaryCoordTangent.w);
+
+		TangentToLocal[0] = normalize(TangentPosition*MeshExtension + MeshOrigin - Intermediates.SimulatedPosition);
+		TangentToLocal[2] = normalize(NormalPosition*MeshExtension + MeshOrigin - Intermediates.SimulatedPosition);
+
+		// cloth data are all in world space so need to change into local space
+		TangentToLocal[0] = mul(TangentToLocal[0], (half3x3)Primitive.WorldToLocal);
+		TangentToLocal[2] = mul(TangentToLocal[2], (half3x3)Primitive.WorldToLocal);
+
+		// derive the new binormal by getting the cross product of the normal and tangent
+		// and flip vector based on sign of tangent basis determinant
+		TangentToLocal[1] = cross(TangentToLocal[2], TangentToLocal[0]) * LocalTangentZ.w;
+	
+		ClothTangent = TangentToLocal;
+	}
+	#endif // GPUSKIN_APEX_CLOTH
+
+#endif // GPUSKIN_MORPH_BLEND
+
+	// Note the use of mul(Matrix,Vector), bone matrices are stored transposed
+	// for tighter packing.
+	TangentToLocal[0] = mul( Intermediates.BlendMatrix, float4(LocalTangentX, 0) );
+	TangentToLocal[2] = mul( Intermediates.BlendMatrix, float4(LocalTangentZ.xyz, 0) );
+
+	// derive the new binormal by getting the cross product of the normal and tangent
+	// and flip vector based on sign of tangent basis determinant
+	TangentToLocal[1] = cross(TangentToLocal[2], TangentToLocal[0]) * LocalTangentZ.w;
+
+#if GPUSKIN_APEX_CLOTH
+	return ClothTangent * TempClothBlendWeight + TangentToLocal * (1.0f - TempClothBlendWeight);
+#else
+	return TangentToLocal;
+#endif	// GPUSKIN_APEX_CLOTH
+}
+
+FVertexFactoryIntermediates GetVertexFactoryIntermediates(FVertexFactoryInput Input)
+{
+	FVertexFactoryIntermediates Intermediates;
+	
+	Intermediates.UnpackedPosition = UnpackedPosition(Input);
+#if GPUSKIN_APEX_CLOTH
+	if( IsSimulatedVertex(Input) )
+	{
+		Intermediates.PreviousSimulatedPosition = ClothingPosition(Input, true);
+		Intermediates.SimulatedPosition = ClothingPosition(Input, false);
+	}
+	else
+	{
+		Intermediates.PreviousSimulatedPosition = Intermediates.SimulatedPosition = Intermediates.UnpackedPosition;
+	}
+	
+#endif
+
+	Intermediates.BlendMatrix = CalcBoneMatrix( Input );
+
+	// Fill TangentToLocal
+	Intermediates.TangentToLocal = SkinTangents(Input, Intermediates);
+
+	// Swizzle vertex color.
+	Intermediates.Color = Input.Color FCOLOR_COMPONENT_SWIZZLE;
+
+	return Intermediates;
+}
+
+float4 CalcWorldPosition(FVertexFactoryInput Input, FVertexFactoryIntermediates Intermediates)
+{
+#if GPUSKIN_APEX_CLOTH
+	if( IsSimulatedVertex(Input) )
+	{
+		// skinned positions
+		float4 SkinnedPosition = float4(SkinPosition(Input, Intermediates), 1);
+		float4 TransformedSkinPos = mul(SkinnedPosition, Primitive.LocalToWorld);
+
+		// Intermediates.SimulatedPosition is a clothing position in world coord
+		float3 BlendedPos = lerp(TransformedSkinPos.xyz, Intermediates.SimulatedPosition, ClothBlendWeight);
+		return float4(BlendedPos + ResolvedView.PreViewTranslation, 1);
+	}
+	else
+#endif
+	{
+		return TransformLocalToTranslatedWorld(SkinPosition(Input, Intermediates));
+	}
+}
+
+/**
+* Get the 3x3 tangent basis vectors for this vertex factory
+*
+* @param Input - vertex input stream structure
+* @return 3x3 matrix
+*/
+float3x3 VertexFactoryGetTangentToLocal( FVertexFactoryInput Input, FVertexFactoryIntermediates Intermediates)
+{
+	return Intermediates.TangentToLocal;
+}
+
+float4 VertexFactoryGetWorldPosition(FVertexFactoryInput Input, FVertexFactoryIntermediates Intermediates)
+{
+	return CalcWorldPosition(Input, Intermediates);
+}
+
+float4 VertexFactoryGetRasterizedWorldPosition(FVertexFactoryInput Input, FVertexFactoryIntermediates Intermediates, float4 InWorldPosition)
+{
+	return InWorldPosition;
+}
+
+float3 VertexFactoryGetPositionForVertexLighting(FVertexFactoryInput Input, FVertexFactoryIntermediates Intermediates, float3 TranslatedWorldPosition)
+{
+	return TranslatedWorldPosition;
+}
+
+void CalcTangentToWorld(FVertexFactoryInput Input, FVertexFactoryIntermediates Intermediates, out float3 TangentToWorld0, out float4 TangentToWorld2)
+{
+	float3x3 LocalToWorld = GetLocalToWorld3x3();
+
+	// Remove scaling.
+	half3 InvScale = Primitive.InvNonUniformScale.xyz;
+	LocalToWorld[0] *= InvScale.x;
+	LocalToWorld[1] *= InvScale.y;
+	LocalToWorld[2] *= InvScale.z;
+
+	float3x3 TangentToWorld = mul(Intermediates.TangentToLocal, LocalToWorld);
+
+	TangentToWorld0 = TangentToWorld[0];
+	TangentToWorld2 = float4(TangentToWorld[2], TangentBias(Input.TangentZ.w) * Primitive.LocalToWorldDeterminantSign);
+}
+
+FVertexFactoryInterpolantsVSToPS VertexFactoryGetInterpolantsVSToPS(FVertexFactoryInput Input, FVertexFactoryIntermediates Intermediates, FMaterialVertexParameters VertexParameters)
+{
+	FVertexFactoryInterpolantsVSToPS Interpolants;
+
+#if NUM_TEX_COORD_INTERPOLATORS
+	float2 CustomizedUVs[NUM_TEX_COORD_INTERPOLATORS];
+	GetMaterialCustomizedUVs(VertexParameters, CustomizedUVs);
+	GetCustomInterpolators(VertexParameters, CustomizedUVs);
+
+	UNROLL
+	for (int CoordinateIndex = 0; CoordinateIndex < NUM_TEX_COORD_INTERPOLATORS; CoordinateIndex++)
+	{
+		Interpolants.TexCoords[CoordinateIndex] = CustomizedUVs[CoordinateIndex];
+	}
+#endif
+
+	Interpolants.TangentToWorld0.w = 0;
+	CalcTangentToWorld(Input, Intermediates, Interpolants.TangentToWorld0.xyz, Interpolants.TangentToWorld2);
+
+	Interpolants.Color = Intermediates.Color;
+
+#if INSTANCED_STEREO
+	Interpolants.PackedEyeIndex = 0;
+#endif
+
+	return Interpolants;
+}
+
+// @return previous translated world position
+float4 VertexFactoryGetPreviousWorldPosition(FVertexFactoryInput Input, FVertexFactoryIntermediates Intermediates)
+{
+	float4x4 PreviousLocalToWorldTranslated = PreviousLocalToWorld;
+	PreviousLocalToWorldTranslated[3][0] += ResolvedView.PrevPreViewTranslation.x;
+	PreviousLocalToWorldTranslated[3][1] += ResolvedView.PrevPreViewTranslation.y;
+	PreviousLocalToWorldTranslated[3][2] += ResolvedView.PrevPreViewTranslation.z;
+
+	float4 PrevSkinPosInWorld = mul(float4(SkinPreviousPosition(Input, Intermediates),1), PreviousLocalToWorldTranslated);
+#if GPUSKIN_APEX_CLOTH
+	if( IsSimulatedVertex(Input) )
+	{
+		float3 PrevSimulatedPos = Intermediates.PreviousSimulatedPosition + ResolvedView.PrevPreViewTranslation;
+		float3 BlendedPrevPos = lerp(PrevSkinPosInWorld.xyz, PrevSimulatedPos.xyz, ClothBlendWeight);
+		return float4(BlendedPrevPos, 1);
+	}
+#endif // GPUSKIN_APEX_CLOTH
+
+	return PrevSkinPosInWorld;
+}
+
+
+#if USING_TESSELLATION
+float2 VertexFactoryGetTextureCoordinateDS( FVertexFactoryInterpolantsVSToDS Interpolants )
+{
+#if NUM_MATERIAL_TEXCOORDS
+	return Interpolants.InterpolantsVSToPS.TexCoords[0].xy;
+#else // #if NUM_MATERIAL_TEXCOORDS
+	return float2(0,0);
+#endif // #if NUM_MATERIAL_TEXCOORDS
+}
+
+FVertexFactoryInterpolantsVSToPS VertexFactoryAssignInterpolants(FVertexFactoryInterpolantsVSToDS Input)
+{
+	return Input.InterpolantsVSToPS;
+}
+
+FVertexFactoryInterpolantsVSToDS VertexFactoryGetInterpolantsVSToDS(FVertexFactoryInput Input, FVertexFactoryIntermediates Intermediates, FMaterialVertexParameters VertexParameters)
+{
+	FVertexFactoryInterpolantsVSToDS Interpolants;
+	
+	Interpolants.InterpolantsVSToPS = VertexFactoryGetInterpolantsVSToPS(Input, Intermediates, VertexParameters);
+	
+	return Interpolants;
+}
+
+/** Converts from vertex factory specific interpolants to a FMaterialTessellationParameters, which is used by material inputs. */
+FMaterialTessellationParameters GetMaterialTessellationParameters(FVertexFactoryInterpolantsVSToDS Interpolants, float3 CameraLocalWorldPosition)
+{
+	FMaterialTessellationParameters	Result;
+#if NUM_TEX_COORD_INTERPOLATORS
+	UNROLL
+	for (int CoordinateIndex = 0;CoordinateIndex < NUM_TEX_COORD_INTERPOLATORS;CoordinateIndex++)
+	{
+		Result.TexCoords[CoordinateIndex] = Interpolants.InterpolantsVSToPS.TexCoords[CoordinateIndex];
+	}
+#endif
+
+	half3 TangentToWorld0 = Interpolants.InterpolantsVSToPS.TangentToWorld0.xyz;
+	half4 TangentToWorld2 = Interpolants.InterpolantsVSToPS.TangentToWorld2;
+
+	Result.VertexColor = Interpolants.InterpolantsVSToPS.Color;
+
+	Result.TangentToWorld = AssembleTangentToWorld( TangentToWorld0, TangentToWorld2 );
+
+	Result.TangentToWorldPreScale = 1;
+
+
+	Result.WorldPosition = CameraLocalWorldPosition + ResolvedView.WorldCameraOrigin;
+
+	return Result;
+}
+
+FVertexFactoryInterpolantsVSToDS VertexFactoryInterpolate(FVertexFactoryInterpolantsVSToDS a, float aInterp, FVertexFactoryInterpolantsVSToDS b, float bInterp)
+{
+	FVertexFactoryInterpolantsVSToDS O;
+
+	TESSELLATION_INTERPOLATE_MEMBER(InterpolantsVSToPS.TangentToWorld0.xyz);
+	TESSELLATION_INTERPOLATE_MEMBER(InterpolantsVSToPS.TangentToWorld2);
+	TESSELLATION_INTERPOLATE_MEMBER(InterpolantsVSToPS.Color);
+
+#if NUM_TEX_COORD_INTERPOLATORS
+	UNROLL
+	for(int tc = 0; tc < NUM_TEX_COORD_INTERPOLATORS; ++tc)
+	{
+		TESSELLATION_INTERPOLATE_MEMBER(InterpolantsVSToPS.TexCoords[tc]);
+	}
+#endif
+
+	return O;
+}
+
+float3x3 VertexFactoryGetTangentToLocalDS(FVertexFactoryInterpolantsVSToDS Interpolants)
+{
+	// This duplicates stuff already going on in GetMaterialTessellationParameters(), so
+	// maybe the hull shader could leverage that instead?
+
+	half3 TangentToWorld0 = Interpolants.InterpolantsVSToPS.TangentToWorld0.xyz;
+	half4 TangentToWorld2 = Interpolants.InterpolantsVSToPS.TangentToWorld2;
+
+	float3x3 TangentToWorld = AssembleTangentToWorld( TangentToWorld0, TangentToWorld2 );
+	
+	return TangentToWorld;
+}
+
+#endif // #if USING_TESSELLATION
+
+#if INSTANCED_STEREO
+uint VertexFactoryGetEyeIndex(uint InstanceId)
+{
+	return InstanceId & 1;
+}
+#endif
+
+float4 VertexFactoryGetTranslatedPrimitiveVolumeBounds(FVertexFactoryInterpolantsVSToPS Interpolants)
+{
+	return 0;
+}
\ No newline at end of file
diff --git a/Engine/Shaders/Private/HeightFogCommon.ush b/Engine/Shaders/Private/HeightFogCommon.ush
new file mode 100644
--- /dev/null
+++ b/Engine/Shaders/Private/HeightFogCommon.ush
@@ -0,0 +1,205 @@
+// Copyright 1998-2017 Epic Games, Inc. All Rights Reserved.
+
+/*=============================================================================
+	HeightFogCommon.usf: 
+=============================================================================*/
+
+#ifndef SUPPORT_FOG_INSCATTERING_TEXTURE
+#define SUPPORT_FOG_INSCATTERING_TEXTURE (FEATURE_LEVEL >= FEATURE_LEVEL_SM4)
+#endif
+
+#ifndef SUPPORT_FOG_DIRECTIONAL_LIGHT_INSCATTERING
+#define SUPPORT_FOG_DIRECTIONAL_LIGHT_INSCATTERING 1
+#endif
+
+#ifndef SUPPORT_VOLUMETRIC_FOG
+#define SUPPORT_VOLUMETRIC_FOG (FEATURE_LEVEL >= FEATURE_LEVEL_SM4)
+#endif
+
+static const float FLT_EPSILON = 0.001f;
+static const float FLT_EPSILON2 = 0.01f;
+
+/** FogDensity * exp2(-FogHeightFalloff * (CameraWorldPosition.z - FogHeight)) in x, FogHeightFalloff in y, CosTerminatorAngle in z, StartDistance in w. */
+float4 ExponentialFogParameters;
+float4 ExponentialFogColorParameter;
+/** FogDensity in x, FogHeight in y, whether to use cubemap fog color in z, FogCutoffDistance in w. */
+float4 ExponentialFogParameters3;
+float2 SinCosInscatteringColorCubemapRotation;
+half4 InscatteringLightDirection;
+half4 DirectionalInscatteringColor;
+float DirectionalInscatteringStartDistance;
+
+TextureCube FogInscatteringColorCubemap;
+SamplerState FogInscatteringColorSampler;
+/** mip distance scale in x, bias in y, num mips in z */
+float3 FogInscatteringTextureParameters;
+float ApplyVolumetricFog;
+
+float3 ComputeInscatteringColor(float3 CameraToReceiver, float CameraToReceiverLength)
+{
+	half3 Inscattering = ExponentialFogColorParameter.xyz;
+	half3 DirectionalInscattering = 0;
+
+#if SUPPORT_FOG_INSCATTERING_TEXTURE
+	BRANCH
+	if (ExponentialFogParameters3.z > 0)
+	{
+		float FadeAlpha = saturate(CameraToReceiverLength * FogInscatteringTextureParameters.x + FogInscatteringTextureParameters.y);
+		float3 CubemapLookupVector = CameraToReceiver;
+		// Rotate around Z axis
+		CubemapLookupVector.xy = float2(dot(CubemapLookupVector.xy, float2(SinCosInscatteringColorCubemapRotation.y, -SinCosInscatteringColorCubemapRotation.x)), dot(CubemapLookupVector.xy, SinCosInscatteringColorCubemapRotation.xy));
+		float3 DirectionalColor = TextureCubeSampleLevel(FogInscatteringColorCubemap, FogInscatteringColorSampler, CubemapLookupVector, 0).xyz;
+		float3 NonDirectionalColor = TextureCubeSampleLevel(FogInscatteringColorCubemap, FogInscatteringColorSampler, CubemapLookupVector, FogInscatteringTextureParameters.z).xyz;
+		Inscattering *= lerp(NonDirectionalColor, DirectionalColor, FadeAlpha);
+	}
+#endif
+
+	return Inscattering;
+}
+
+// @param WorldPositionRelativeToCamera = WorldPosition - InCameraPosition
+half4 GetExponentialHeightFog(float3 WorldPositionRelativeToCamera, float ExcludeDistance)
+{
+	const half MinFogOpacity = ExponentialFogColorParameter.w;
+
+	float3 CameraToReceiver = WorldPositionRelativeToCamera;
+	float CameraToReceiverLengthSqr = dot(CameraToReceiver, CameraToReceiver);
+	float CameraToReceiverLengthInv = rsqrt(CameraToReceiverLengthSqr);
+	float CameraToReceiverLength = CameraToReceiverLengthSqr * CameraToReceiverLengthInv;
+	half3 CameraToReceiverNormalized = CameraToReceiver * CameraToReceiverLengthInv;
+
+	float RayOriginTerms = ExponentialFogParameters.x;
+	float RayLength = CameraToReceiverLength;
+	float RayDirectionZ = CameraToReceiver.z;
+
+	// Factor in StartDistance
+	ExcludeDistance = max(ExcludeDistance, ExponentialFogParameters.w);
+	
+#if USE_GLOBAL_CLIP_PLANE
+
+	BRANCH
+	// While rendering a planar reflection with a clip plane, we must compute analytical fog using a camera path starting from the plane, rather than the virtual camera origin
+	if (dot(View.GlobalClippingPlane.xyz, 1) > 0.0f)
+	{
+		float CameraOriginPlaneDistance = dot(View.GlobalClippingPlane, float4(View.WorldCameraOrigin, 1));
+		float PlaneIntersectionTime = -CameraOriginPlaneDistance / dot(CameraToReceiver, View.GlobalClippingPlane.xyz);
+
+		// Only modify the start distance if the reflection plane is between the camera and receiver point
+		if (PlaneIntersectionTime > 0 && PlaneIntersectionTime < 1)
+		{
+			ExcludeDistance = max(ExcludeDistance, PlaneIntersectionTime * CameraToReceiverLength);
+		}
+	}
+
+#endif
+
+	if (ExcludeDistance > 0)
+	{
+		float ExcludeIntersectionTime = ExcludeDistance * CameraToReceiverLengthInv;
+		float CameraToExclusionIntersectionZ = ExcludeIntersectionTime * CameraToReceiver.z;
+		float ExclusionIntersectionZ = View.WorldCameraOrigin.z + CameraToExclusionIntersectionZ;
+		float ExclusionIntersectionToReceiverZ = CameraToReceiver.z - CameraToExclusionIntersectionZ;
+
+		// Calculate fog off of the ray starting from the exclusion distance, instead of starting from the camera
+		RayLength = (1.0f - ExcludeIntersectionTime) * CameraToReceiverLength;
+		RayDirectionZ = ExclusionIntersectionToReceiverZ;
+
+		float Exponent = max(-127.0f, ExponentialFogParameters.y * (ExclusionIntersectionZ - ExponentialFogParameters3.y));
+		RayOriginTerms = ExponentialFogParameters3.x * exp2(-Exponent);
+	}
+
+	// Calculate the line integral of the ray from the camera to the receiver position through the fog density function
+	// The exponential fog density function is d = GlobalDensity * exp(-HeightFalloff * z)
+	float Falloff = max(-127.0f, ExponentialFogParameters.y * RayDirectionZ);    // if it's lower than -127.0, then exp2() goes crazy in OpenGL's GLSL.
+	float LineIntegral = ( 1.0f - exp2(-Falloff) ) / Falloff;
+	float LineIntegralTaylor = log(2.0) - ( 0.5 * Pow2( log(2.0) ) ) * Falloff;		// Taylor expansion around 0
+	float ExponentialHeightLineIntegralShared = RayOriginTerms * ( abs(Falloff) > FLT_EPSILON2 ? LineIntegral : LineIntegralTaylor );
+	float ExponentialHeightLineIntegral = ExponentialHeightLineIntegralShared * RayLength;
+
+	half3 InscatteringColor = ComputeInscatteringColor(CameraToReceiver, CameraToReceiverLength);
+	half3 DirectionalInscattering = 0;
+
+#if SUPPORT_FOG_DIRECTIONAL_LIGHT_INSCATTERING
+	if (InscatteringLightDirection.w > 0 && ExponentialFogParameters3.z == 0)
+	{
+		// Setup a cosine lobe around the light direction to approximate inscattering from the directional light off of the ambient haze;
+		half3 DirectionalLightInscattering = DirectionalInscatteringColor.xyz * pow(saturate(dot(CameraToReceiverNormalized, InscatteringLightDirection.xyz)), DirectionalInscatteringColor.w);
+
+		// Calculate the line integral of the eye ray through the haze, using a special starting distance to limit the inscattering to the distance
+		float DirExponentialHeightLineIntegral = ExponentialHeightLineIntegralShared * max(RayLength - DirectionalInscatteringStartDistance, 0.0f);
+		// Calculate the amount of light that made it through the fog using the transmission equation
+		half DirectionalInscatteringFogFactor = saturate(exp2(-DirExponentialHeightLineIntegral));
+		// Final inscattering from the light
+		DirectionalInscattering = DirectionalLightInscattering * (1 - DirectionalInscatteringFogFactor);
+	}
+#endif
+
+	// Calculate the amount of light that made it through the fog using the transmission equation
+	half ExpFogFactor = max(saturate(exp2(-ExponentialHeightLineIntegral)), MinFogOpacity);
+
+	FLATTEN
+	if (ExponentialFogParameters3.w > 0 && CameraToReceiverLength > ExponentialFogParameters3.w)
+	{
+		ExpFogFactor = 1;
+		DirectionalInscattering = 0;
+	}
+
+	return half4((InscatteringColor) * (1 - ExpFogFactor) + DirectionalInscattering, ExpFogFactor);
+}
+
+Texture3D IntegratedLightScattering;
+SamplerState IntegratedLightScatteringSampler;
+
+float4 CombineVolumetricFog(float4 GlobalFog, float3 VolumeUV)
+{
+	float4 VolumetricFogLookup = float4(0, 0, 0, 1);
+
+	if (ApplyVolumetricFog > 0)
+	{
+		VolumetricFogLookup = Texture3DSampleLevel(IntegratedLightScattering, IntegratedLightScatteringSampler, VolumeUV, 0);
+	}
+
+	// Visualize depth distribution
+	//VolumetricFogLookup.rgb += .1f * frac(min(ZSlice, 1.0f) / View.VolumetricFogInvGridSize.z);
+
+	return float4(VolumetricFogLookup.rgb + GlobalFog.rgb * VolumetricFogLookup.a, VolumetricFogLookup.a * GlobalFog.a);
+}
+
+float4 CalculateCombinedFog(float3 WorldPositionRelativeToCamera, float3 VolumeUV)
+{
+	float ExcludeDistance = 0;
+
+#if SUPPORT_VOLUMETRIC_FOG
+	float InvCosAngle = length(WorldPositionRelativeToCamera) / dot(WorldPositionRelativeToCamera, View.ViewForward);
+	// Volumetric fog covers up to MaxDistance along ViewZ, exclude analytical fog from this range
+	ExcludeDistance = max(View.VolumetricFogMaxDistance * InvCosAngle, 0);
+#endif
+
+	float4 FogInscatteringAndOpacity = GetExponentialHeightFog(WorldPositionRelativeToCamera, ExcludeDistance);
+
+#if SUPPORT_VOLUMETRIC_FOG
+	FogInscatteringAndOpacity = CombineVolumetricFog(FogInscatteringAndOpacity, VolumeUV);
+#endif
+
+	return FogInscatteringAndOpacity;
+}
+
+float ComputeNormalizedZSliceFromDepth(float SceneDepth)
+{
+	return log2(SceneDepth * View.VolumetricFogGridZParams.x + View.VolumetricFogGridZParams.y) * View.VolumetricFogGridZParams.z * View.VolumetricFogInvGridSize.z;
+}
+
+float3 ComputeVolumeUV(float3 WorldPosition, float4x4 WorldToClip)
+{
+	float4 NDCPosition = mul(float4(WorldPosition, 1), WorldToClip);
+	NDCPosition.xy /= NDCPosition.w;
+	return float3(NDCPosition.xy * float2(.5f, -.5f) + .5f, ComputeNormalizedZSliceFromDepth(NDCPosition.w));
+}
+
+half4 CalculateHeightFog(float3 WorldPositionRelativeToCamera)
+{
+	float3 WorldPosition = WorldPositionRelativeToCamera + ResolvedView.TranslatedWorldCameraOrigin - ResolvedView.PreViewTranslation.xyz;
+	float3 VolumeUV = ComputeVolumeUV(WorldPosition, ResolvedView.WorldToClip);
+	return CalculateCombinedFog(WorldPositionRelativeToCamera, VolumeUV);
+}
+
diff --git a/Engine/Shaders/Private/HeightfieldLightingShared.ush b/Engine/Shaders/Private/HeightfieldLightingShared.ush
new file mode 100644
--- /dev/null
+++ b/Engine/Shaders/Private/HeightfieldLightingShared.ush
@@ -0,0 +1,111 @@
+// Copyright 1998-2017 Epic Games, Inc. All Rights Reserved.
+
+/*=============================================================================
+	HeightfieldLightingShared.usf
+=============================================================================*/
+
+Texture2D HeightfieldTexture;
+SamplerState HeightfieldSampler;
+Texture2D DiffuseColorTexture;
+SamplerState DiffuseColorSampler;
+uint NumHeightfields;
+Buffer<float4> HeightfieldDescriptions;
+
+// In float4's, must match C++
+#define HEIGHTFIELD_DATA_STRIDE 12
+
+float4 GetHeightfieldScaleBias(uint Index, out bool bHasSubsections)
+{
+	float4 HeightfieldVector = HeightfieldDescriptions[Index * HEIGHTFIELD_DATA_STRIDE + 0];
+	// Sign bit stored presence of subsections
+	bHasSubsections = HeightfieldVector.x < 0;
+	HeightfieldVector.x = abs(HeightfieldVector.x);
+	return HeightfieldVector;
+}
+
+float4 GetHeightfieldMinMaxUV(uint Index)
+{
+	return HeightfieldDescriptions[Index * HEIGHTFIELD_DATA_STRIDE + 1];
+}
+
+float4 GetHeightfieldGlobalScaleBias(uint Index)
+{
+	return HeightfieldDescriptions[Index * HEIGHTFIELD_DATA_STRIDE + 2];
+}
+
+float2 GetHeightfieldSize(uint Index)
+{
+	return HeightfieldDescriptions[Index * HEIGHTFIELD_DATA_STRIDE + 3].xy;
+}
+
+float2 GetGlobalTexelSize(uint Index)
+{
+	return HeightfieldDescriptions[Index * HEIGHTFIELD_DATA_STRIDE + 3].zw;
+}
+
+float4x4 GetWorldToLocal(uint Index)
+{
+	return float4x4(
+		HeightfieldDescriptions[Index * HEIGHTFIELD_DATA_STRIDE + 4],
+		HeightfieldDescriptions[Index * HEIGHTFIELD_DATA_STRIDE + 5],
+		HeightfieldDescriptions[Index * HEIGHTFIELD_DATA_STRIDE + 6],
+		HeightfieldDescriptions[Index * HEIGHTFIELD_DATA_STRIDE + 7]);
+}
+
+float4x4 GetLocalToWorld(uint Index)
+{
+	return float4x4(
+		HeightfieldDescriptions[Index * HEIGHTFIELD_DATA_STRIDE + 8],
+		HeightfieldDescriptions[Index * HEIGHTFIELD_DATA_STRIDE + 9],
+		HeightfieldDescriptions[Index * HEIGHTFIELD_DATA_STRIDE + 10],
+		HeightfieldDescriptions[Index * HEIGHTFIELD_DATA_STRIDE + 11]);
+}
+
+float2 GetHeightfieldUV(uint HeightfieldIndex, float2 LocalPosition, out float4 MinMaxHeightfieldUV)
+{
+	bool bHasSubsections;
+	float4 HeightfieldScaleBias = GetHeightfieldScaleBias(HeightfieldIndex, bHasSubsections);
+	float2 HeightfieldUV = LocalPosition * HeightfieldScaleBias.xy + HeightfieldScaleBias.zw;
+	MinMaxHeightfieldUV = GetHeightfieldMinMaxUV(HeightfieldIndex);
+
+	if (bHasSubsections)
+	{
+		float2 CenterUV = .5f * (MinMaxHeightfieldUV.xy + MinMaxHeightfieldUV.zw);
+		// One row of heights in the heightfield are duplicated between subsections
+		// The only number of subsections supported is 2
+		// If we are in subsection 0, no adjustment is needed, but if we are in subsection 1 we need to shift up by one texel
+		HeightfieldUV += HeightfieldUV + .5f * HeightfieldScaleBias.xy > CenterUV ? HeightfieldScaleBias.xy : 0;
+	}
+
+	return HeightfieldUV;
+}
+
+float2 GetHeightfieldUV(uint HeightfieldIndex, float2 LocalPosition)
+{
+	float4 MinMaxHeightfieldUV;
+	return GetHeightfieldUV(HeightfieldIndex, LocalPosition, MinMaxHeightfieldUV);
+}
+
+float3 GetHeightfieldWorldPositionAndNormal(uint HeightfieldIndex, float2 LocalPosition, float2 HeightfieldUV, out float3 WorldNormal)
+{
+	//@todo - don't call GetHeightfieldScaleBias twice
+	bool bHasSubsections;
+	float4 HeightfieldScaleBias = GetHeightfieldScaleBias(HeightfieldIndex, bHasSubsections);
+
+	float4 SampleValue = Texture2DSampleLevel(HeightfieldTexture, HeightfieldSampler, HeightfieldUV + .5f * HeightfieldScaleBias.xy, 0);
+	float HeightfieldHeight = DecodePackedHeight(SampleValue.xy);
+
+	float2 SampleNormal = float2(SampleValue.b, SampleValue.a) * float2(2.0,2.0) - float2(1.0,1.0);
+	WorldNormal = float3(SampleNormal, sqrt(max(1.0 - dot(SampleNormal, SampleNormal), 0.0)));
+
+	float3 HeightfieldShadingPosition = float3(LocalPosition, HeightfieldHeight);
+	float3 WorldHeightfieldShadingPosition = mul(float4(HeightfieldShadingPosition, 1), GetLocalToWorld(HeightfieldIndex)).xyz;
+
+	return WorldHeightfieldShadingPosition;
+}
+
+float3 GetHeightfieldWorldPosition(uint HeightfieldIndex, float2 LocalPosition, float2 HeightfieldUV)
+{
+	float3 WorldNormal;
+	return GetHeightfieldWorldPositionAndNormal(HeightfieldIndex, LocalPosition, HeightfieldUV, WorldNormal);
+}
diff --git a/Engine/Shaders/Private/IESLightProfilesCommon.ush b/Engine/Shaders/Private/IESLightProfilesCommon.ush
new file mode 100644
--- /dev/null
+++ b/Engine/Shaders/Private/IESLightProfilesCommon.ush
@@ -0,0 +1,28 @@
+// Copyright 1998-2017 Epic Games, Inc. All Rights Reserved.
+
+/*=============================================================================
+	IESLightProfiles.usf: Common shader code for IES light profiles
+=============================================================================*/
+
+// IES light profiles
+Texture2D		IESTexture;
+SamplerState	IESTextureSampler;
+
+// Apply 1D IES light profile texture
+float ComputeLightProfileMultiplier(float3 WorldPosition, float3 LightPosition, float3 LightDirection)
+{
+#if USE_IES_PROFILE
+	float3 NegLightDirection = normalize(WorldPosition - LightPosition);
+
+	// -1..1
+	float DotProd = dot(NegLightDirection, LightDirection);
+	// -PI..PI (this distortion could be put into the texture but not without quality loss or more memory)
+	float Angle = asin(DotProd);
+	// 0..1
+	float NormAngle = Angle / PI + 0.5f;
+
+	return Texture2DSampleLevel(IESTexture, IESTextureSampler, float2(NormAngle, 0), 0).r;
+#else 
+	return 1.0f;
+#endif
+}
\ No newline at end of file
diff --git a/Engine/Shaders/Private/InstancedStereo.ush b/Engine/Shaders/Private/InstancedStereo.ush
new file mode 100644
--- /dev/null
+++ b/Engine/Shaders/Private/InstancedStereo.ush
@@ -0,0 +1,61 @@
+// Copyright 1998-2017 Epic Games, Inc. All Rights Reserved..
+
+/*=============================================================================
+	InstancedStereo.usf: Resolve which view uniforms in a stereo pair to use.
+=============================================================================*/
+
+#pragma once
+
+// Explictly include view uniform buffers
+#include "/Engine/Generated/UniformBuffers/View.ush"
+#include "/Engine/Generated/UniformBuffers/InstancedView.ush"
+
+// ViewState, GetPrimaryView and GetInstancedView are generated by the shader compiler to ensure View uniform buffer changes are up to date.
+// see GenerateInstancedStereoCode()
+#include "/Engine/Generated/GeneratedInstancedStereo.ush"
+
+// Packing the eye index and instanced stereo flag into a single uint to save an interpolator
+uint PackEyeIndex(uint EyeIndex, bool bIsInstancedStereo)
+{
+	uint Result = 0;
+	Result |= (uint)bIsInstancedStereo << 31;
+	Result |= EyeIndex;
+	return Result;
+}
+
+uint GetEyeIndex(uint PackedIndex)
+{
+	return PackedIndex & 0x1;
+}
+
+bool IsInstancedStereo(uint PackedIndex)
+{
+	return (PackedIndex >> 1) > 0;
+}
+
+static ViewState ResolvedView;
+
+ViewState ResolveView()
+{
+	return GetPrimaryView();
+}
+
+#if INSTANCED_STEREO
+static const float EyeOffsetScale[2] = { -1.0, 1.0 };
+static const float4 EyeClipEdge[2] = { float4(-1.0, 0.0, 0.0, 1.0), float4(1.0, 0.0, 0.0, 1.0) };
+bool bIsInstancedStereo;
+#endif
+
+#if INSTANCED_STEREO || MOBILE_MULTI_VIEW
+ViewState ResolveView(uint ViewIndex)
+{
+	if (ViewIndex == 0)
+	{
+		return GetPrimaryView();
+	}
+	else
+	{
+		return GetInstancedView();
+	}
+}
+#endif
diff --git a/Engine/Shaders/Private/LPVCommon.ush b/Engine/Shaders/Private/LPVCommon.ush
new file mode 100644
--- /dev/null
+++ b/Engine/Shaders/Private/LPVCommon.ush
@@ -0,0 +1,338 @@
+//-----------------------------------------------------------------------------
+// File:		LPVCommon.usf
+//
+// Summary:		Common functionality for LPV shaders
+//
+// Created:		2013-03-01
+//
+// Author:		mailto:[REDACTED_EMAIL]
+//
+//				Copyright (C) Microsoft. All rights reserved.
+//-----------------------------------------------------------------------------
+
+//------------------------------------------------------------------------------
+
+#define REFINE_OVER_TIME	1
+
+#if REFINE_OVER_TIME
+  #define PREV_FRAME_MULTIPLIER			0.9f
+
+  // With REFINE_OVER_TIME enabled and a previous frame multiplier of 0.9, lighting 
+  // is effectively multiplied by 10. 
+  // Final lighting therefore needs to be divided by 10 to account for this
+  #define REFINE_LIGHTING_MULTIPLIER	0.1f
+#else
+	#define REFINE_LIGHTING_MULTIPLIER	1.0f
+#endif
+ 
+#define TILED_ADDRESSING		0
+
+#ifndef LPV_WRITE_SHADER
+#define LPV_WRITE_SHADER	0
+#endif
+
+
+//-------------------------------------------------------------------------------------------------
+SamplerState								gLpv3DTextureSampler;
+
+Texture3D<float>							gAOVolumeTexture; 
+Texture3D<float4>							gLpv3DTexture0;
+Texture3D<float4>							gLpv3DTexture1;
+Texture3D<float4>							gLpv3DTexture2;
+Texture3D<float4>							gLpv3DTexture3;
+Texture3D<float4>							gLpv3DTexture4;
+Texture3D<float4>							gLpv3DTexture5;
+Texture3D<float4>							gLpv3DTexture6;
+
+// Write shaders only
+#if LPV_WRITE_SHADER
+
+RWTexture3D<float4>							gLpv3DTextureRW0;
+RWTexture3D<float4>							gLpv3DTextureRW1;
+RWTexture3D<float4>							gLpv3DTextureRW2;
+RWTexture3D<float4>							gLpv3DTextureRW3;
+RWTexture3D<float4>							gLpv3DTextureRW4;
+RWTexture3D<float4>							gLpv3DTextureRW5;
+RWTexture3D<float4>							gLpv3DTextureRW6;
+
+RWTexture3D<float>							gAOVolumeTextureRW;
+
+#endif // LPV_WRITE_SHADER
+
+//-------------------------------------------------------------------------------------------------
+
+struct LPVCell
+{
+	float3 coeffs[9];
+	float AO;
+};
+
+//-------------------------------------------------------------------------------------------------
+// Functions
+//-------------------------------------------------------------------------------------------------
+
+void ClearCell( in out LPVCell cell )
+{ 
+	[unroll]
+	for ( int i = 0; i < 9; i++ ) 
+	{
+		cell.coeffs[i] = float3( 0.0f, 0.0f, 0.0f );
+	}
+
+	cell.AO = 0.0f;
+}
+
+//-------------------------------------------------------------------------------------------------
+
+void MultiplyCell( in out LPVCell cell, float multiplier )
+{
+	[unroll]
+	for ( int i=0; i<9; i++ ) 
+	{
+		cell.coeffs[i] *= multiplier;
+	}
+	cell.AO *= multiplier;
+}
+
+void MultiplyCellCoeffs( in out LPVCell cell, float multiplier )
+{
+	[unroll]
+	for ( int i=0; i<9; i++ ) 
+		cell.coeffs[i] *= multiplier;
+}
+
+//-------------------------------------------------------------------------------------------------
+
+void MultiplyCellV( in out LPVCell cellA, in out LPVCell cellB )
+{
+	[unroll]
+	for ( int i=0; i<9; i++ ) 
+	{
+		cellA.coeffs[i] *= cellB.coeffs[i];
+	}
+	cellA.AO *= cellB.AO;
+}
+
+//-------------------------------------------------------------------------------------------------
+
+void AddCell( in out LPVCell cellA, in LPVCell cellB )
+{
+	[unroll]
+	for ( int i=0; i<9; i++ ) 
+	{
+		cellA.coeffs[i] += cellB.coeffs[i];
+	}
+	cellA.AO += cellB.AO;
+}
+
+//-------------------------------------------------------------------------------------------------
+
+int3 IndexToGridPos( uint cellIndex )
+{
+	int3 texIndex;
+	texIndex.x = ( cellIndex % 32 );
+	texIndex.y = ( cellIndex / 32 ) %  32;
+	texIndex.z = cellIndex / ( 32 * 32 );
+	return texIndex;
+}
+
+//-------------------------------------------------------------------------------------------------
+
+uint GetGridAddress( uint3 gridPos )
+{
+#if TILED_ADDRESSING
+	uint3 tileAddr = gridPos / 4;
+	uint3 rem = gridPos % 4;
+	return ( tileAddr.z*64 + tileAddr.y*8 + tileAddr.x ) * 64 +
+		   rem.z*16 + rem.y * 4 + rem.x;
+#else
+	return gridPos.z*32*32+gridPos.y*32+gridPos.x;
+#endif
+}
+
+//-------------------------------------------------------------------------------------------------
+
+// Unpacks an HDR colour from RGY32 
+float3 UnpackRGY32( uint packedColour )
+{
+	const float ONE_OVER_255 = 1.0f/255.0f;
+	float3 rgb;
+	rgb.r = ( packedColour & 0xff000000 ) >> 24;
+	rgb.g = ( packedColour & 0x00ff0000 ) >> 16;
+	rgb.b = 255.0f-(rgb.r+rgb.g);
+	float luminance = f16tof32( packedColour );
+	return rgb * luminance * ONE_OVER_255;
+}
+
+//-------------------------------------------------------------------------------------------------
+
+// Packs an HDR colour to RGY32 
+uint PackRGY32( float3 colour )
+{
+	// @TODO: could normalise for more precision, but greater ALU in the decode
+	float luminance = (colour.r+colour.g+colour.b);
+	colour.rg /= luminance;
+
+	uint packedValue =	uint( colour.r * 255.0f + 0.5f ) << 24 |
+						uint( colour.g * 255.0f + 0.5f ) << 16; 
+
+	packedValue |= f32tof16( luminance );
+	return packedValue;
+}
+
+//-------------------------------------------------------------------------------------------------
+
+float3 SHLookup( float3 coeffs[9], float3 n ) 
+{
+	const float1 c1 = 0.429043 ;
+	const float1 c2 = 0.511664 ;
+	const float1 c3 = 0.743125 ;
+	const float1 c4 = 0.886227 ;
+	const float1 c5 = 0.247708 ;
+
+	float3 n2 = n*n;
+	float xy = n.x*n.y ; 
+	float yz = n.y*n.z ; 
+	float xz = n.x*n.z ;
+
+	float3 colour = c1*coeffs[8]*(n2.x-n2.y) + c3*coeffs[6]*n2.z + c4*coeffs[0] - c5*coeffs[6] 
+            + 2.0f*c1*(coeffs[4]*xy + coeffs[7]*xz + coeffs[5]*yz) 
+            + 2.0f*c2*(coeffs[3]*n.x+coeffs[1]*n.y+coeffs[2]*n.z) ;
+
+	//return colour;
+	return max( colour, float3(0,0,0) );
+
+}
+
+//-------------------------------------------------------------------------------------------------
+
+float3 LPVCellLookup( in LPVCell cell, in float3 normal ) 
+{
+	return SHLookup( cell.coeffs, normal );
+}
+
+//-------------------------------------------------------------------------------------------------
+
+float3 SHGetDominantDirectionApproxScalar( in float coeffs[9] )
+{
+	float3 dir = float3( -coeffs[3], -coeffs[1], -coeffs[2] );
+	float len = length( dir );
+	float den = ( len > 0.01 ) ? len : 1.0f;
+	return dir/den;
+}
+
+//-------------------------------------------------------------------------------------------------
+
+float3 WorldToGrid( float3 worldPos )
+{
+#if LPV_WRITE_SHADER
+	return worldPos * LpvWrite.OneOverLpvScale + LpvWrite.mLpvGridOffset;
+#else
+	return worldPos * LpvRead.OneOverLpvScale + LpvRead.mLpvGridOffset;
+#endif
+}
+
+//-------------------------------------------------------------------------------------------------
+
+float3 GridToWorld( float3 gridPos )
+{
+#if LPV_WRITE_SHADER
+	return ( gridPos - LpvWrite.mLpvGridOffset ) * LpvWrite.LpvScale;
+#else
+	return ( gridPos - LpvRead.mLpvGridOffset ) * LpvRead.LpvScale;
+#endif
+}
+
+//-------------------------------------------------------------------------------------------------
+
+float MaxGridExtent( float3 gridPos )
+{
+	float3 av = abs(gridPos-16.0f);
+	return max(av.x,max(av.y,av.z));
+}
+
+//-------------------------------------------------------------------------------------------------
+
+int GetGridIndex( float3 worldPos ) 
+{
+	float3 p = WorldToGrid( worldPos );
+
+	float maxExtent = MaxGridExtent( p );
+
+	if ( maxExtent > 16.0f )
+	{
+		return -1;
+	} 
+	else
+	{
+		int3 ui = int3( p );
+		return GetGridAddress( ui );
+	}
+}
+
+//-------------------------------------------------------------------------------------------------
+
+LPVCell ReadLpvCell( in uint cellIndex )
+{
+	LPVCell cell;
+	cell.AO = 0.0f;
+
+	int4 texIndex = int4( IndexToGridPos( cellIndex ), 0 );
+	float4 tex0 = gLpv3DTexture0.Load( texIndex );
+	float4 tex1 = gLpv3DTexture1.Load( texIndex );
+	float4 tex2 = gLpv3DTexture2.Load( texIndex );
+	float4 tex3 = gLpv3DTexture3.Load( texIndex );
+	float4 tex4 = gLpv3DTexture4.Load( texIndex );
+	float4 tex5 = gLpv3DTexture5.Load( texIndex );
+	float4 tex6 = gLpv3DTexture6.Load( texIndex );
+
+	cell.coeffs[0] = tex0.rgb; // tex0.a is reserved for secondary occlusion
+	cell.coeffs[1] = tex1.rgb; 
+	cell.coeffs[2] = float3( tex1.a, tex2.rg );
+	cell.coeffs[3] = float3( tex2.ba, tex3.r );
+	cell.coeffs[4] = tex3.gba;
+	cell.coeffs[5] = tex4.rgb;
+	cell.coeffs[6] = float3( tex4.a, tex5.rg );
+	cell.coeffs[7] = float3( tex5.ba, tex6.r );
+	cell.coeffs[8] = tex6.gba;
+	cell.AO = tex0.a;
+	return cell;
+}
+
+//-------------------------------------------------------------------------------------------------
+
+LPVCell ReadLpvCellVolumeTextureFiltered( float3 gridPos )
+{
+	float3 texPos = gridPos / 32.0f; // Half pixel offset?
+	float4 tex0 = gLpv3DTexture0.SampleLevel( gLpv3DTextureSampler, texPos, 0 );
+	float4 tex1 = gLpv3DTexture1.SampleLevel( gLpv3DTextureSampler, texPos, 0 );
+	float4 tex2 = gLpv3DTexture2.SampleLevel( gLpv3DTextureSampler, texPos, 0 );
+	float4 tex3 = gLpv3DTexture3.SampleLevel( gLpv3DTextureSampler, texPos, 0 );
+	float4 tex4 = gLpv3DTexture4.SampleLevel( gLpv3DTextureSampler, texPos, 0 );
+	float4 tex5 = gLpv3DTexture5.SampleLevel( gLpv3DTextureSampler, texPos, 0 );
+	float4 tex6 = gLpv3DTexture6.SampleLevel( gLpv3DTextureSampler, texPos, 0 );
+
+	LPVCell cell; 
+	cell.coeffs[0] = tex0.rgb; // tex0.a is reserved for secondary occlusion
+	cell.coeffs[1] = tex1.rgb; 
+	cell.coeffs[2] = float3( tex1.a, tex2.rg );
+	cell.coeffs[3] = float3( tex2.ba, tex3.r );
+	cell.coeffs[4] = tex3.gba;
+	cell.coeffs[5] = tex4.rgb;
+	cell.coeffs[6] = float3( tex4.a, tex5.rg );
+	cell.coeffs[7] = float3( tex5.ba, tex6.r );
+	cell.coeffs[8] = tex6.gba;
+
+	cell.AO = tex0.a;
+	return cell;
+}
+
+//-------------------------------------------------------------------------------------------------
+
+float ReadLpvCellAO( in uint cellIndex )
+{
+	int4 texIndex = int4( IndexToGridPos( cellIndex ), 0 );
+	return gLpv3DTexture0.Load( texIndex ).a;
+}
+
+//-------------------------------------------------------------------------------------------------
diff --git a/Engine/Shaders/Private/LPVFinalPass.ush b/Engine/Shaders/Private/LPVFinalPass.ush
new file mode 100644
--- /dev/null
+++ b/Engine/Shaders/Private/LPVFinalPass.ush
@@ -0,0 +1,576 @@
+//-----------------------------------------------------------------------------
+// File:		LPVFinalPass.usf
+//
+// Summary:		(SHADER HEADER) LPV Final pass functionality
+//
+// Created:		2013-03-01
+//
+// Author:		mailto:[REDACTED_EMAIL]
+//
+//				Copyright (C) Microsoft. All rights reserved.
+//-----------------------------------------------------------------------------
+
+// Copyright 2013 Lionhead Studios
+// mailto:[REDACTED_EMAIL]
+
+/*------------------------------------------------------------------------------
+	Compile time parameters:
+------------------------------------------------------------------------------*/
+
+#ifndef LPV_FINAL_PASS_USF
+#define LPV_FINAL_PASS_USF
+
+#include "LPVCommon.ush"
+
+//------------------------------------------------------------------------------
+
+#define REFLECTION_RAY_LENGTH_DEFAULT 2.0
+
+#define SHOW_GRID			0
+#define DEBUG_WEIGHTS		0
+
+// Indirect strength
+#define INDIRECT_LEVEL		( 0.00375 * REFINE_LIGHTING_MULTIPLIER )
+#define INDIRECT_DIFFUSE	1
+#define INDIRECT_SPECULAR	1
+
+// Filtering modes
+#define FILTER_NEAREST				0
+#define FILTER_TETRA				1 // Interpolate using tetrahedrons (4 points instead of 8)
+#define FILTER_TRILINEAR			2 
+#define FILTER_HARDWARE_TRILINEAR	3 // Use hardware volume texture filtering
+
+#ifndef LPV_USE_SEPARATE_AO_VOLUME
+#define LPV_USE_SEPARATE_AO_VOLUME	1 // Use the separate (low precision) AO texture for AO lookups
+#endif
+
+#define LPV_DEFAULT_AO_VALUE		0.4
+
+// Filtering settings
+	#define DIFFUSE_FILTER		FILTER_HARDWARE_TRILINEAR 
+	#define SPECULAR_FILTER		FILTER_HARDWARE_TRILINEAR 
+
+//-------------------------------------------------------------------------------------------------
+
+void GetGridIndicesAndWeightsTetra( float3 gridPos, out int indices[4], out float weights[4] )
+{
+//	gridPos.y+=0.43;
+	const float3 gridOffset[8] =
+	{
+		float3( -0.5f, -0.5f, -0.5f ), // c1
+		float3( -0.5f, -0.5f,  0.5f ),
+		float3( -0.5f,  0.5f, -0.5f ), 
+		float3( -0.5f,  0.5f,  0.5f ), // c0
+		float3(  0.5f, -0.5f, -0.5f ), 
+		float3(  0.5f, -0.5f,  0.5f ), // c3
+		float3(  0.5f,  0.5f, -0.5f ), // c2
+		float3(  0.5f,  0.5f,  0.5f ), 
+	};
+
+	int cubeIndices[8];
+	// Get cube indices
+	[unroll]
+	for ( int i=0; i<8; i++ )
+	{
+		int3 ui = int3( gridPos + gridOffset[i] );
+		ui %= 32;
+		cubeIndices[i] = GetGridAddress( ui );//ui.z*32*32 + ui.y*32 + ui.x;
+	}
+	
+
+	// Get the position in cube space
+	// Where cube is space is nearest 2x2x2 cells, with origin at 0,0,0
+	int3 cubeOriginI = int3( gridPos-float3(0.5f,0.5f,0.5f) );
+
+	float3 cubeOrigin = float3(cubeOriginI)+0.5;
+	float3 cubePos = ( gridPos - cubeOrigin );
+
+	int3 align = int3( LpvRead.mLpvGridOffset ) % 2;
+
+
+	float3 flip2 = ( (cubeOriginI + align ) % 2 ) * -2 + 1;
+
+	// Classification
+	uint tetra = 4;
+	const float3 tips[4] = 
+	{
+		float3( 0,1,0 ),
+		float3( 0,0,1 ),
+		float3( 1,0,0 ),
+		float3( 1,1,1 ),
+	};
+
+	float3 weights3 = 0;
+	if ( flip2.x < 0 ) cubePos.x = 1-cubePos.x;
+	if ( flip2.y < 0 ) cubePos.y = 1-cubePos.y;
+	if ( flip2.z < 0 ) cubePos.z = 1-cubePos.z;
+
+	[unroll]
+	for ( uint j=0; j<4; j++ )
+	{
+		float3 tip = tips[j];
+		float3 d = abs(cubePos-tip);
+		float mDist = d.x+d.y+d.z;
+		if ( mDist <= 1.0f ) tetra=j;
+	}
+
+	if ( tetra<4 )
+	{
+		float3 r4 = tips[tetra];
+		float3 diagMat = 1.0f - r4 * 2.0f;
+		float3 rMinusR4 = cubePos - r4;
+		weights3 = rMinusR4*diagMat;
+
+		// compute the 4 tetra vert positions
+		if ( flip2.x< 0 ) r4.x=1-r4.x;
+		if ( flip2.y< 0 ) r4.y=1-r4.y;
+		if ( flip2.z< 0 ) r4.z=1-r4.z;
+		int3 r4i = int3(r4);
+		int3 p[4];
+		p[0] = int3(1-r4i.x,r4i.y,r4i.z);
+		p[1] = int3(r4i.x,1-r4i.y,r4i.z);
+		p[2] = int3(r4i.x,r4i.y,1-r4i.z);
+		p[3] = r4i;
+
+		[unroll]
+		for ( int i=0; i<4; i++ )
+		{
+			int3 ui = ( cubeOriginI + p[i] ) % 32;
+			indices[i] = GetGridAddress( ui );//ui.z*32*32 + ui.y*32 + ui.x;
+		}
+	}
+	else
+	{
+		float3 centreR4 = float3( 1, 0, 1 );
+
+		// Inverse of tetrahedron matrix for the centre tetrahedron
+		const float3 centreM1 = float3( -0.5, 0.5, 0.5 );
+		const float3 centreM2 = float3( -0.5,-0.5,-0.5 );
+		const float3 centreM3 = float3(  0.5, 0.5,-0.5 );
+
+		float3 rMinusR4 = cubePos - centreR4;
+		weights3.x = dot( centreM1, rMinusR4 );
+		weights3.y = dot( centreM2, rMinusR4 );
+		weights3.z = dot( centreM3, rMinusR4 );
+
+		// compute the 4 tetra vert positions
+		//int3 r4i = int3(centreR4);
+		int3 r1i = int3(0,1,1);
+
+		if ( flip2.x< 0 ) r1i.x=1-r1i.x;
+		if ( flip2.y< 0 ) r1i.y=1-r1i.y;
+		if ( flip2.z< 0 ) r1i.z=1-r1i.z;
+
+		int3 p[4];
+		p[0] = r1i; 
+		p[1] = int3(r1i.x,1-r1i.y,1-r1i.z); 
+		p[2] = int3(1-r1i.x,r1i.y,1-r1i.z);
+		p[3] = int3(1-r1i.x,1-r1i.y,r1i.z);
+
+		[unroll]
+		for ( int i=0; i<4; i++ )
+		{
+			int3 ui = ( cubeOriginI + p[i] ) % 32;
+			indices[i] = GetGridAddress( ui );//ui.z*32*32 + ui.y*32 + ui.x;
+		}
+	}
+
+	weights[0] = weights3.x;
+	weights[1] = weights3.y;
+	weights[2] = weights3.z;
+	weights[3] = 1.0f-(weights3.x+weights3.y+weights3.z);
+
+#if DEBUG_WEIGHTS
+	float3 colour = 0;
+	float3 colours[9] =
+	{
+		float3(1,0,0),			// float3( -0.5f, -0.5f, -0.5f ), // c1
+		float3(0,1,0),			// float3( -0.5f, -0.5f,  0.5f ),
+		float3(0,0,1),			// float3( -0.5f,  0.5f, -0.5f ), 
+		float3(1,1,0),			// float3( -0.5f,  0.5f,  0.5f ), // c0
+		float3(0,2,2),			// float3(  0.5f, -0.5f, -0.5f ), 
+		float3(1,0,1),			// float3(  0.5f, -0.5f,  0.5f ), // c3
+		float3(0,0,0),			// float3(  0.5f,  0.5f, -0.5f ), // c2
+		float3(8,0.5,0.5),		// float3(  0.5f,  0.5f,  0.5f ), 
+		float3(10,0,0), // Invalid (very bright red!)
+	};
+
+	for ( int i=0;i<8; i++ )
+	{
+		LPVCell cell = ReadLpvCell( cubeIndices[i] );
+		colours[i] = cell.coeffs[0];
+
+		//int3 ui = int3( gridPos + gridOffset[i] );
+		//colours[i] = ( ( float3(ui) * 10 ) % 32 ) / 32;
+
+		int idx = cubeIndices[i];
+
+		int ix = idx % 32;
+		int iy = ( idx / 32 ) % 32;
+		int iz = ( idx / (32*32) ) % 32;
+
+		colours[i] = 0;
+		if ( ix == 16 && iy == 15 && iz == 16 ) colours[i] = float3(1,0,0);
+
+		int3 cp = int3(gridPos);//cubeOrigin-0.5;
+		if ( cp.x == 16 && cp.y == 15 && cp.z == 16 ) colours[i].b+=1;
+	}
+
+	for ( int i=0; i<4; i++ )
+	{
+		// Find the cube index
+		int cubeIndex = 8;
+		for ( int j=0; j<8; j++ )
+		{
+			if ( indices[i] == cubeIndices[j] )
+			{
+				cubeIndex = j;
+				break;
+			}
+		}
+		colour += colours[ cubeIndex ] * weights[i];
+	}
+	colour/=4.0;
+
+	//cubePos = 1-cubePos;
+	if ( flip2.x < 0 ) cubePos.x = 1-cubePos.x;
+	if ( flip2.y < 0 ) cubePos.y = 1-cubePos.y;
+	if ( flip2.z < 0 ) cubePos.z = 1-cubePos.z;
+	if ( cubePos.x < 0.0125 || cubePos.y < 0.0125 || cubePos.z < 0.0125 ) 
+	{
+		colour = 0.5;
+	}
+
+	//colour = ( ( cubeOrigin * 10 ) % 32 ) / 32;
+	weights[0] = colour.x;
+	weights[1] = colour.y;
+	weights[2] = colour.z;
+#endif
+}
+
+//-------------------------------------------------------------------------------------------------
+
+void GetGridIndicesAndWeights( float3 gridPos, out int indices[8], out float weights[8] )
+{
+	const float3 gridOffset[8] =
+	{
+		float3( -0.5f, -0.5f, -0.5f ),
+		float3( -0.5f, -0.5f,  0.5f ),
+		float3( -0.5f,  0.5f, -0.5f ),
+		float3( -0.5f,  0.5f,  0.5f ),
+		float3(  0.5f, -0.5f, -0.5f ),
+		float3(  0.5f, -0.5f,  0.5f ),
+		float3(  0.5f,  0.5f, -0.5f ),
+		float3(  0.5f,  0.5f,  0.5f ),
+	};
+
+	// Get indices
+	[unroll]
+	for ( int i=0; i<8; i++ )
+	{
+		int3 ui = int3( gridPos+gridOffset[i] );
+		float maxExtent = MaxGridExtent( ui );
+		if ( maxExtent > 15.5 )
+			indices[i] = -1;
+		else
+			indices[i] = GetGridAddress( ui );//ui.z*32*32 + ui.y*32 + ui.x;
+	}
+	
+	// Get weights
+	float3 baseTexelPos = floor( gridPos-0.5 ) + 0.5;
+	float3 texPos = gridPos-baseTexelPos; // Pos relative to bottom-left texel centre
+
+	float3 minVal = 1.0f-texPos;
+	float3 maxVal = texPos;
+	weights[0] = minVal.x * minVal.y * minVal.z;
+	weights[1] = minVal.x * minVal.y * maxVal.z;
+	weights[2] = minVal.x * maxVal.y * minVal.z;
+	weights[3] = minVal.x * maxVal.y * maxVal.z;
+	weights[4] = maxVal.x * minVal.y * minVal.z;
+	weights[5] = maxVal.x * minVal.y * maxVal.z;
+	weights[6] = maxVal.x * maxVal.y * minVal.z;
+	weights[7] = maxVal.x * maxVal.y * maxVal.z;
+
+#if 0
+	[unroll]
+	for ( int i=0; i<8; i++ )
+	{
+		weights[i] *= weights[i];
+	}
+#endif
+}
+
+//-------------------------------------------------------------------------------------------------
+
+float3 DebugGridColour( float3 gridPos, float3 originalColour )
+{
+	float3 cubePos = frac( gridPos + 0.5 );
+	float3 colour = saturate( ( 0.015 - cubePos ) * 10000 );
+
+	float3 cubePos2 = frac( gridPos );
+	float3 colour2 = saturate( ( 0.01 - cubePos2 ) * 10000 );
+	colour2 = ( colour2 * 0.05) + dot( colour2, colour2 ).xxx * 0.1;
+
+	colour += colour2;
+
+
+	if ( dot( colour,colour ) > 0.0 )
+	{
+		return colour * 8 + originalColour * 0.5;
+	}
+	return originalColour;
+}
+
+//-------------------------------------------------------------------------------------------------
+
+LPVCell LPVGetCellNearest( float3 worldPos )
+{
+	int gridIndex = GetGridIndex( worldPos );
+	return ReadLpvCell( gridIndex );
+}
+
+//-------------------------------------------------------------------------------------------------
+
+LPVCell LPVGetCellTetra( float3 worldPos )
+{
+	float3 colour = 0;
+
+	float3 gridPos = WorldToGrid( worldPos );
+
+	const int n = 4;
+	int indices[4];
+	float weights[4];
+	GetGridIndicesAndWeightsTetra( gridPos, indices, weights );
+
+#if DEBUG_WEIGHTS
+	return float3(weights[0],weights[1],weights[2]);
+#endif 
+	LPVCell accumCell;
+	ClearCell( accumCell );
+
+	[unroll]
+	for ( int i=0; i<n; i++ )
+	{
+		// Decompress the cell and apply the weight
+		LPVCell cell = ReadLpvCell( indices[i] );
+		MultiplyCell( cell, weights[i] );
+		AddCell( accumCell, cell );
+	}
+
+	return accumCell;
+}
+
+//-------------------------------------------------------------------------------------------------
+
+LPVCell LPVGetCellTrilinear( float3 worldPos )
+{
+	float3 colour = 0;
+
+	float3 gridPos = WorldToGrid( worldPos );
+
+	const int n = 8;
+	int indices[8];
+	float weights[8];
+	GetGridIndicesAndWeights( gridPos, indices, weights );
+
+#if DEBUG_WEIGHTS
+	return float3(weights[0],weights[1],weights[2]);
+#endif 
+
+	LPVCell accumCell;
+	ClearCell( accumCell );
+
+	[unroll]
+	for ( int i=0; i<n; i++ )
+	{
+		// Decompress the cell and apply the weight
+		LPVCell cell = ReadLpvCell( indices[i] );
+
+		MultiplyCell( cell, weights[i] );
+		AddCell( accumCell, cell );
+	}
+
+	return accumCell;
+}
+
+//-------------------------------------------------------------------------------------------------
+
+LPVCell LPVGetCellHardwareTrilinear( float3 worldPos )
+{
+	float3 colour = 0;
+
+	float3 gridPos = WorldToGrid( worldPos );
+
+	return ReadLpvCellVolumeTextureFiltered( gridPos );
+}
+
+//-------------------------------------------------------------------------------------------------
+
+LPVCell LPVGetCell( float3 worldPos )
+{
+#if ( DIFFUSE_FILTER == FILTER_TETRA )
+	return LPVGetCellTetra( worldPos );
+#elif ( DIFFUSE_FILTER == FILTER_NEAREST )
+	return LPVGetCellNearest( worldPos );
+#elif ( DIFFUSE_FILTER == FILTER_TRILINEAR )
+	return LPVGetCellTrilinear( worldPos );
+#elif ( DIFFUSE_FILTER == FILTER_HARDWARE_TRILINEAR )
+	return LPVGetCellHardwareTrilinear( worldPos );
+#endif
+}
+
+float LPVGetCellAO_Grid( float3 gridPos )
+{
+	float3 absGridPos = abs(gridPos-16.0f);
+	float m = max( absGridPos.x, max( absGridPos.y, absGridPos.z ));
+	if ( m<=16.0f )
+	{
+		float3 texPos = gridPos / 32.0f;
+#if LPV_USE_SEPARATE_AO_VOLUME
+		return gAOVolumeTexture.SampleLevel( gLpv3DTextureSampler, texPos, 0 ).r;
+#else
+		return gLpv3DTexture0.SampleLevel( gLpv3DTextureSampler, texPos, 0 ).a;
+#endif
+	}
+	return LPV_DEFAULT_AO_VALUE;
+}
+
+float LPVGetCellAO( float3 WorldPos )
+{
+	float3 gridPos = WorldToGrid( WorldPos );
+	return LPVGetCellAO_Grid( gridPos );
+}
+
+//-------------------------------------------------------------------------------------------------
+
+float3 LPVSpecular( float3 worldPos, float3 reflectionVec, float rOffsetLength )
+{
+#if !INDIRECT_SPECULAR
+	return 0;
+#endif
+
+	LPVCell Cell = LPVGetCell( worldPos + reflectionVec * rOffsetLength * LpvRead.LpvScale );
+
+	float3 Colour = LPVCellLookup( Cell, reflectionVec ) * INDIRECT_LEVEL * LpvRead.SpecularIntensity;
+
+#if SHOW_GRID
+	Colour = DebugGridColour( WorldToGrid( worldPos ), colour );
+#endif
+
+	return Colour;
+}
+
+//-------------------------------------------------------------------------------------------------
+
+float3 LPVSpecular( float3 worldPos, float3 reflectionVec )
+{
+	return LPVSpecular( worldPos, reflectionVec, REFLECTION_RAY_LENGTH_DEFAULT );
+}
+
+
+//-------------------------------------------------------------------------------------------------
+
+float3 LPVDiffuse( float3 worldPos, float3 normal )
+{
+#if !INDIRECT_DIFFUSE
+	return 0;
+#endif
+
+	LPVCell Cell = LPVGetCell( worldPos );
+
+	float3 Colour = LPVCellLookup( Cell, normal ) * INDIRECT_LEVEL * LpvRead.DiffuseIntensity;
+
+	return Colour;
+}
+
+//-------------------------------------------------------------------------------------------------
+
+void LPVGetDiffuseAndSpecular( float3 worldPos, float3 normal, float3 reflectionVec, out float3 diffuseOut, out float3 specularOut )
+{
+	// TODO: could implement optimised ref lookup by using SH from diffuse for reflection
+
+	diffuseOut = LPVDiffuse( worldPos, normal ) ;
+	specularOut = LPVSpecular( worldPos, reflectionVec ) ; 
+}
+
+
+//-------------------------------------------------------------------------------------------------
+
+// Directional occlusion lookup
+float2 LPVGetDirectionalOcclusion_DiffuseSpec( float3 WorldPosition, float3 R, float Roughness, float3 N )
+{
+	BRANCH
+	if ( LpvRead.DirectionalOcclusionIntensity < 0.0001f )
+	{
+		return float2(1,1);
+	}
+	float SpecularOcclusion = 0.0f;
+	float DiffuseOcclusion = 0.0f;
+	const int NumSpecularSteps = 6; 
+	const int NumDiffuseSteps  = 2;
+
+	float3 gridPos = WorldToGrid( WorldPosition );
+	float StepLength = 1.333 / float(NumDiffuseSteps);//0.301f;
+	float Weight = 1.0f;
+	float TotalWeight = 0.0f;
+
+#if DEBUG_USE_RAW_AO
+	float ao = saturate( LPVGetCellAO_Grid( gridPos ) );
+	if ( ao > 0.5 ) return 1;
+	return ao.xx;
+#endif
+
+	float DiffuseOcclusionRaw = LPVGetCellAO_Grid( gridPos + N * 0.25 );
+	DiffuseOcclusion = saturate( DiffuseOcclusionRaw * 2.0f / LpvRead.DiffuseOcclusionIntensity );
+	DiffuseOcclusion = pow( DiffuseOcclusion, LpvRead.DiffuseOcclusionExponent );
+	DiffuseOcclusion = saturate(DiffuseOcclusion);
+
+	gridPos = WorldToGrid( WorldPosition );
+	StepLength = 2.0f / (float)NumSpecularSteps;
+	SpecularOcclusion = 0.0f;
+	TotalWeight=0;
+	Weight=1.0f;
+
+	// Use the diffuse occlusion tap as the first specular tap (since it's nearly identical)
+	SpecularOcclusion = saturate( 1-DiffuseOcclusionRaw ) * Weight;
+	StepLength *= 1.3f;
+	gridPos+=R*StepLength;
+	TotalWeight+=Weight; 
+	Weight *= 0.9;
+	for  ( int i=1; i<NumSpecularSteps; i++ )
+	{
+		float AO = saturate( 1-LPVGetCellAO_Grid( gridPos ) );
+		SpecularOcclusion += AO * Weight;
+		StepLength *= 1.3f;
+		gridPos+=R*StepLength;
+		TotalWeight+=Weight; 
+		Weight *= 0.9;
+	}
+
+	// Increase intensity as power increases	
+	float LerpValue = saturate( LpvRead.SpecularOcclusionExponent * 0.1f ); // 0-1 range
+	float Strength = lerp( 0.9f, 1.3f, LerpValue ) * LpvRead.SpecularOcclusionIntensity;
+	SpecularOcclusion *= Strength;
+	SpecularOcclusion /= TotalWeight;
+	SpecularOcclusion = pow( SpecularOcclusion, LpvRead.SpecularOcclusionExponent );
+
+	SpecularOcclusion = 1-SpecularOcclusion;
+
+	float2 DirectionalOcclusion = lerp( float2( 1,1 ), saturate( float2( DiffuseOcclusion, SpecularOcclusion) ), LpvRead.DirectionalOcclusionIntensity );
+
+#if LPV_OCCLUSION_EDGE_FADE
+	float3 grid = WorldToGrid( WorldPosition );
+	float minGrid = min( grid.x, min(grid.y,grid.z ) );
+	float maxGrid = max( grid.x, max(grid.y,grid.z ) );
+
+	float inside = min(	saturate(minGrid*0.5),
+						saturate((32.0f-maxGrid)*0.5) );
+
+	DirectionalOcclusion = lerp( float2(0.5,0.5), DirectionalOcclusion, inside );
+#endif
+
+	return DirectionalOcclusion;
+}
+
+
+#endif
diff --git a/Engine/Shaders/Private/LPVGeometryVolumeCommon.ush b/Engine/Shaders/Private/LPVGeometryVolumeCommon.ush
new file mode 100644
--- /dev/null
+++ b/Engine/Shaders/Private/LPVGeometryVolumeCommon.ush
@@ -0,0 +1,228 @@
+//-----------------------------------------------------------------------------
+// File:		LPVGeometryVolumeCommon.usf
+//
+// Summary:		Common functionality for geometry volumes
+//
+// Created:		2013-03-01
+//
+// Author:		mailto:[REDACTED_EMAIL]
+//
+//				Copyright (C) Microsoft. All rights reserved.
+//-----------------------------------------------------------------------------
+
+#ifndef LPV_MULTIPLE_BOUNCES
+#define LPV_MULTIPLE_BOUNCES 0
+#endif 
+
+#define GV_ORDER 2
+
+#ifndef LPV_GV_SH_ORDER
+#error LPV_GV_SH_ORDER is not defined!
+#else
+#if ( GV_ORDER > LPV_GV_SH_ORDER )
+#undef GV_ORDER
+#define GV_ORDER LPV_GV_SH_ORDER
+#endif
+#endif
+
+//------------------------------------------------------------------------------
+
+struct GeometryVolumeEntry
+{
+#if (GV_ORDER==2)
+	float SH[9];
+#elif (GV_ORDER==1)
+	float SH[4];
+#elif (GV_ORDER==0)
+	float SH[1];
+#endif
+
+#if LPV_MULTIPLE_BOUNCES
+	//@todo : store colour as SH? Is single component good enough?
+	float3 color;
+#endif
+};
+
+//------------------------------------------------------------------------------
+
+RWTexture3D<float4>	gGv3DTextureRW0;
+Texture3D<float4>	gGv3DTexture0;
+#if ( GV_ORDER >= 1 )
+RWTexture3D<float4>	gGv3DTextureRW1;
+Texture3D<float4>	gGv3DTexture1;
+#endif
+#if (GV_ORDER >= 2 )
+RWTexture3D<float4>	gGv3DTextureRW2;
+Texture3D<float4>	gGv3DTexture2;
+#endif
+
+//-------------------------------------------------------------------------------------------------
+// Functions
+//-------------------------------------------------------------------------------------------------
+
+GeometryVolumeEntry ReadGvCell( uint index )
+{
+	GeometryVolumeEntry gv;
+
+	int4 texIndex = int4( IndexToGridPos( index ), 0 );
+	float4 gvTex0 = gGv3DTexture0.Load( texIndex );
+	gv.SH[0] = gvTex0.x;
+#if ( GV_ORDER >= 1 )
+	float4 gvTex1 = gGv3DTexture1.Load( texIndex );
+	gv.SH[1] = gvTex0.y;
+	gv.SH[2] = gvTex0.z;
+	gv.SH[3] = gvTex0.w;
+#endif
+#if ( GV_ORDER >= 2 )
+	float4 gvTex2 = gGv3DTexture2.Load( texIndex );
+	gv.SH[4] = gvTex1.x;
+	gv.SH[5] = gvTex1.y;
+	gv.SH[6] = gvTex1.z;
+	gv.SH[7] = gvTex1.w;
+	gv.SH[8] = gvTex2.x;
+#endif
+
+ #if LPV_MULTIPLE_BOUNCES
+ #if (GV_ORDER == 0 )
+  gv.color = gvTex0.yzw;
+#elif (GV_ORDER == 1)
+  gv.color = gvTex1.xyz;
+#else
+	gv.color = gvTex2.yzw;
+ #endif
+ #endif
+	return gv;
+	}
+
+//-------------------------------------------------------------------------------------------------
+
+float3 GvGetDominantDirectionApprox( in GeometryVolumeEntry cell )
+{
+#if (GV_ORDER == 0 )
+	return float3(0,0,0);
+#else
+	float3 dir = float3( -cell.SH[3], -cell.SH[1], -cell.SH[2] );
+	float len = length( dir );
+	float den = ( len > 0.01 ) ? len : 1.0f;
+	return dir/den;
+#endif
+}
+
+//-------------------------------------------------------------------------------------------------
+
+void WriteGvCell( uint index, GeometryVolumeEntry entry )
+{
+
+	int3 texIndex = IndexToGridPos( index );
+
+#if ( GV_ORDER == 0 )
+ #if LPV_MULTIPLE_BOUNCES
+	gGv3DTextureRW0[ texIndex ] = float4( entry.SH[0], 0.0f, 0.0f, 0.0f );
+ #else
+	gGv3DTextureRW2[ texIndex ] = float4( entry.SH[0], entry.color.rgb );
+ #endif
+#elif ( GV_ORDER == 1 )
+	gGv3DTextureRW0[ texIndex ] = float4( entry.SH[0], entry.SH[1], entry.SH[2], entry.SH[3] );
+ #if LPV_MULTIPLE_BOUNCES
+	gGv3DTextureRW1[ texIndex ] = float4( entry.color.rgb, 0.0f );
+ #endif
+#else
+	gGv3DTextureRW0[ texIndex ] = float4( entry.SH[0], entry.SH[1], entry.SH[2], entry.SH[3] );
+	gGv3DTextureRW1[ texIndex ] = float4( entry.SH[4], entry.SH[5], entry.SH[6], entry.SH[7] );
+ #if LPV_MULTIPLE_BOUNCES
+	gGv3DTextureRW2[ texIndex ] = float4( entry.SH[8], entry.color.rgb );
+ #else
+	gGv3DTextureRW2[ texIndex ] = float4( entry.SH[8], 0.0f, 0.0f, 0.0f );
+ #endif
+#endif
+}
+
+//-------------------------------------------------------------------------------------------------
+
+void ClearGvEntry( in out GeometryVolumeEntry entry )
+	{
+	entry.SH[0] = 0.0f;
+ #if ( GV_ORDER >= 1 )
+	entry.SH[1] = 0.0f;
+	entry.SH[2] = 0.0f;
+	entry.SH[3] = 0.0f;
+ #endif
+ #if ( GV_ORDER >= 2 )
+	entry.SH[4] = 0.0f;
+	entry.SH[5] = 0.0f;
+	entry.SH[6] = 0.0f;
+	entry.SH[7] = 0.0f;
+	entry.SH[8] = 0.0f;
+ #endif
+
+ #if LPV_MULTIPLE_BOUNCES
+	entry.color = 0.0f;
+ #endif
+}
+
+//-------------------------------------------------------------------------------------------------
+
+void ClearGvCell( uint index )
+{ 
+	int3 texIndex = IndexToGridPos( index );
+	gGv3DTextureRW0[ texIndex ] = 0.0f;
+ #if ( GV_ORDER >= 1 )
+	gGv3DTextureRW1[ texIndex ] = 0.0f;
+ #endif
+ #if ( GV_ORDER >= 2 )
+	gGv3DTextureRW2[ texIndex ] = 0.0f;
+ #endif
+}
+
+//-------------------------------------------------------------------------------------------------
+
+float SHLookupGeometryVolume( 
+	GeometryVolumeEntry cell,
+    float3 n ) 
+{
+	const float1 c1 = 0.429043 ;
+	const float1 c2 = 0.511664 ;
+	const float1 c3 = 0.743125 ;
+	const float1 c4 = 0.886227 ;
+	const float1 c5 = 0.247708 ;
+
+	float3 n2 = n*n;
+#if ( GV_ORDER == 0 )
+	float value = c4*cell.SH[0];
+#elif ( GV_ORDER == 1 )
+	float value = c4*cell.SH[0] + 2.0f*c2*(cell.SH[3]*n.x+cell.SH[1]*n.y+cell.SH[2]*n.z);
+#else
+	float xy = n.x*n.y ; 
+	float yz = n.y*n.z ; 
+	float xz = n.x*n.z ;
+	float value = c1*cell.SH[8]*(n2.x-n2.y) + c3*cell.SH[6]*n2.z + c4*cell.SH[0] - c5*cell.SH[6] 
+            + 2.0f*c1*(cell.SH[4]*xy + cell.SH[7]*xz + cell.SH[5]*yz) 
+            + 2.0f*c2*(cell.SH[3]*n.x+cell.SH[1]*n.y+cell.SH[2]*n.z) ;
+#endif
+	return max( value, 0.0f );
+}
+
+//-------------------------------------------------------------------------------------------------
+
+void AccumulateCoeffsGeometryVolume(float value, float3 n, in out GeometryVolumeEntry cell ) 
+{ 
+	const float c0 = 0.282095f;
+	const float c1 = 0.488603f;
+	const float c2 = 1.092548f;
+	const float c3 = 0.315392f;
+	const float c4 = 0.546274f;
+
+	cell.SH[0] += c0*value;
+#if ( GV_ORDER >= 1 )
+	cell.SH[1] += (c1*n.y)*value;   
+	cell.SH[2] += (c1*n.z)*value;   
+	cell.SH[3] += (c1*n.x)*value;   
+ #if ( GV_ORDER >= 2 )
+	cell.SH[4] += (c2*n.x*n.y)*value; 
+	cell.SH[5] += (c2*n.y*n.z)*value; 
+	cell.SH[7] += (c2*n.x*n.z)*value; 
+	cell.SH[6] += (c3*(3.0f*n.z*n.z-1.0f))*value; 
+	cell.SH[8] += (c4*(n.x*n.x-n.y*n.y))*value;
+ #endif
+#endif
+}
\ No newline at end of file
diff --git a/Engine/Shaders/Private/LPVPacking.ush b/Engine/Shaders/Private/LPVPacking.ush
new file mode 100644
--- /dev/null
+++ b/Engine/Shaders/Private/LPVPacking.ush
@@ -0,0 +1,112 @@
+//-----------------------------------------------------------------------------
+// File:		LPVPacking.usf
+//
+// Summary:		Common functionality for LPV compute shaders
+//
+// Created:		2015-04-15
+//
+// Author:		Ben Woodhouse - mailto:[REDACTED_EMAIL]
+//
+//				Copyright (C) Microsoft. All rights reserved.
+//-----------------------------------------------------------------------------
+
+//------------------------------------------------------------------------------
+
+struct LPVCellPacked
+{
+	uint packedCoeffs[14];
+};
+
+//-------------------------------------------------------------------------------------------------
+
+//-------------------------------------------------------------------------------------------------
+
+LPVCellPacked PackLPVCell( LPVCell cellIn )
+{
+	LPVCellPacked cellOut;
+	uint3 compressedCoeffs[9]; 
+	[unroll]
+	for ( int i=0; i<9; i++ )
+	{
+		compressedCoeffs[i] = f32tof16( cellIn.coeffs[i] );
+	}
+	uint compressedAO = f32tof16( cellIn.AO );
+	
+	cellOut.packedCoeffs[0] = compressedCoeffs[0].x | ( compressedCoeffs[0].y << 16 );
+	cellOut.packedCoeffs[1] = compressedCoeffs[0].z | ( compressedCoeffs[1].x << 16 );
+	cellOut.packedCoeffs[2] = compressedCoeffs[1].y | ( compressedCoeffs[1].z << 16 );
+
+	cellOut.packedCoeffs[3] = compressedCoeffs[2].x | ( compressedCoeffs[2].y << 16 );
+	cellOut.packedCoeffs[4] = compressedCoeffs[2].z | ( compressedCoeffs[3].x << 16 );
+	cellOut.packedCoeffs[5] = compressedCoeffs[3].y | ( compressedCoeffs[3].z << 16 );
+
+	cellOut.packedCoeffs[6] = compressedCoeffs[4].x | ( compressedCoeffs[4].y << 16 );
+	cellOut.packedCoeffs[7] = compressedCoeffs[4].z | ( compressedCoeffs[5].x << 16 );
+	cellOut.packedCoeffs[8] = compressedCoeffs[5].y | ( compressedCoeffs[5].z << 16 );
+
+	cellOut.packedCoeffs[9] = compressedCoeffs[6].x | ( compressedCoeffs[6].y << 16 );
+	cellOut.packedCoeffs[10] = compressedCoeffs[6].z | ( compressedCoeffs[7].x << 16 );
+	cellOut.packedCoeffs[11] = compressedCoeffs[7].y | ( compressedCoeffs[7].z << 16 );
+
+	cellOut.packedCoeffs[12] = compressedCoeffs[8].x | ( compressedCoeffs[8].y << 16 );
+	cellOut.packedCoeffs[13] = compressedCoeffs[8].z | ( compressedAO << 16 );
+
+
+	return cellOut;
+}
+
+//-------------------------------------------------------------------------------------------------
+
+LPVCell UnpackLPVCell( LPVCellPacked cellPacked )
+{
+	LPVCell cellOut;
+	uint3 compressedCoeffs;
+	compressedCoeffs.x = cellPacked.packedCoeffs[0];
+	compressedCoeffs.y = cellPacked.packedCoeffs[0] >> 16;
+	compressedCoeffs.z = cellPacked.packedCoeffs[1];
+	cellOut.coeffs[0] = f16tof32( compressedCoeffs );
+
+	compressedCoeffs.x = cellPacked.packedCoeffs[1] >> 16;
+	compressedCoeffs.y = cellPacked.packedCoeffs[2];
+	compressedCoeffs.z = cellPacked.packedCoeffs[2] >> 16;
+	cellOut.coeffs[1] = f16tof32( compressedCoeffs );
+
+	compressedCoeffs.x = cellPacked.packedCoeffs[3];
+	compressedCoeffs.y = cellPacked.packedCoeffs[3] >> 16;
+	compressedCoeffs.z = cellPacked.packedCoeffs[4];
+	cellOut.coeffs[2] = f16tof32( compressedCoeffs );
+
+	compressedCoeffs.x = cellPacked.packedCoeffs[4] >> 16;
+	compressedCoeffs.y = cellPacked.packedCoeffs[5];
+	compressedCoeffs.z = cellPacked.packedCoeffs[5] >> 16;
+	cellOut.coeffs[3] = f16tof32( compressedCoeffs );
+
+	compressedCoeffs.x = cellPacked.packedCoeffs[6];
+	compressedCoeffs.y = cellPacked.packedCoeffs[6] >> 16;
+	compressedCoeffs.z = cellPacked.packedCoeffs[7];
+	cellOut.coeffs[4] = f16tof32( compressedCoeffs );
+
+	compressedCoeffs.x = cellPacked.packedCoeffs[7] >> 16;
+	compressedCoeffs.y = cellPacked.packedCoeffs[8];
+	compressedCoeffs.z = cellPacked.packedCoeffs[8] >> 16;
+	cellOut.coeffs[5] = f16tof32( compressedCoeffs );
+
+	compressedCoeffs.x = cellPacked.packedCoeffs[9];
+	compressedCoeffs.y = cellPacked.packedCoeffs[9] >> 16;
+	compressedCoeffs.z = cellPacked.packedCoeffs[10];
+	cellOut.coeffs[6] = f16tof32( compressedCoeffs );
+
+	compressedCoeffs.x = cellPacked.packedCoeffs[10] >> 16;
+	compressedCoeffs.y = cellPacked.packedCoeffs[11];
+	compressedCoeffs.z = cellPacked.packedCoeffs[11] >> 16;
+	cellOut.coeffs[7] = f16tof32( compressedCoeffs );
+
+	compressedCoeffs.x = cellPacked.packedCoeffs[12];
+	compressedCoeffs.y = cellPacked.packedCoeffs[12] >> 16;
+	compressedCoeffs.z = cellPacked.packedCoeffs[13];
+	cellOut.coeffs[8] = f16tof32( compressedCoeffs );
+
+	cellOut.AO = f16tof32( cellPacked.packedCoeffs[13] >> 16 );
+
+	return cellOut;
+}
