Branch snapshot for CL 3555263

diff --git a/Engine/Shaders/Private/LPVWriteCommon.ush b/Engine/Shaders/Private/LPVWriteCommon.ush
new file mode 100644
--- /dev/null
+++ b/Engine/Shaders/Private/LPVWriteCommon.ush
@@ -0,0 +1,139 @@
+//-----------------------------------------------------------------------------
+// File:		LPVWriteCommmon.usf
+//
+// Summary:		Common functionality for LPV write
+//
+// Created:		2013-03-01
+//
+// Author:		mailto:[REDACTED_EMAIL]
+//
+//				Copyright (C) Microsoft. All rights reserved.
+//-----------------------------------------------------------------------------
+
+#define LPV_WRITE_SHADER	1
+
+//------------------------------------------------------------------------------
+
+#include "LPVCommon.ush"
+
+//------------------------------------------------------------------------------
+
+#define PROPOGATE_26TAP				1
+#define PROPOGATE_USE_LDS			0
+#define COMPRESSED_LDS_PROPOGATE	0
+
+//------------------------------------------------------------------------------
+
+#define PACKED_VPL 1
+
+struct Vpl
+{
+	float3 normal;
+	float3 flux;
+};
+
+struct VplListEntry
+{
+#if PACKED_VPL
+	uint	normalPacked;
+	uint	fluxPacked;
+#else
+	float3	normal;
+	float3  flux;
+#endif
+
+	int		nextIndex;
+};
+
+//-------------------------------------------------------------------------------------------------
+// Functions
+//-------------------------------------------------------------------------------------------------
+
+float3 UnpackVPLNormal( uint normalPacked )
+{
+	float3 n;
+	n.x = float( ( normalPacked & 0x000000ff ) );
+	n.y = float( ( normalPacked & 0x0000ff00 ) >> 8  );
+	n.z = float( ( normalPacked & 0x00ff0000 ) >> 16 );
+	return (n/254.0f)*2.0f-1.0f;
+}
+
+//-------------------------------------------------------------------------------------------------
+
+uint PackVPLNormal( float3 normal )
+{
+	uint3 n = uint3( (normal+1.0f) * 127.0f );
+	return n.x | ( n.y << 8 ) | ( n.z << 16 );
+}
+
+//-------------------------------------------------------------------------------------------------
+
+Vpl UnpackVpl( VplListEntry entry )
+{
+	Vpl vpl;
+#if PACKED_VPL
+	vpl.flux   = UnpackRGY32( entry.fluxPacked );
+	vpl.normal = UnpackVPLNormal( entry.normalPacked );
+#else
+	vpl.flux   = entry.flux;
+	vpl.normal = entry.normal;
+#endif
+	return vpl;
+}
+
+//-------------------------------------------------------------------------------------------------
+
+VplListEntry PackVpl( Vpl vpl )
+{
+	VplListEntry entry;
+
+#if PACKED_VPL
+	entry.normalPacked	= PackVPLNormal( vpl.normal );
+	entry.fluxPacked	= PackRGY32( vpl.flux );
+#else
+	entry.flux			= vpl.flux;
+	entry.normal		= vpl.normal;
+#endif
+	return entry;
+}
+
+//-------------------------------------------------------------------------------------------------
+
+void AccumulateLighting(float3 colour, float3 n, float solidAngle, in out LPVCell cell ) 
+{ 
+	const float ONE_OVER_PI = 1.0f/3.1415927f;
+	const float c0 = 0.282095f;
+	const float c1 = 0.488603f;
+	const float c2 = 1.092548f;
+	const float c3 = 0.315392f;
+	const float c4 = 0.546274f;
+
+	cell.coeffs[0] += colour*c0*solidAngle ;
+
+	cell.coeffs[1] += colour*(c1*n.y)*solidAngle;   
+	cell.coeffs[2] += colour*(c1*n.z)*solidAngle ;   
+	cell.coeffs[3] += colour*(c1*n.x)*solidAngle ;   
+
+	cell.coeffs[4] += colour*(c2*n.x*n.y)*solidAngle ; 
+	cell.coeffs[5] += colour*(c2*n.y*n.z)*solidAngle ; 
+	cell.coeffs[7] += colour*(c2*n.x*n.z)*solidAngle ; 
+	cell.coeffs[6] += colour*(c3*(3.0f*n.z*n.z-1.0f))*solidAngle ; 
+	cell.coeffs[8] += colour*(c4*(n.x*n.x-n.y*n.y))*solidAngle ;
+}
+
+//-------------------------------------------------------------------------------------------------
+
+#if LPV_WRITE_SHADER
+void WriteLpvCell( in LPVCell cell, uint cellIndex )
+{
+	int3 texIndex = IndexToGridPos( cellIndex );
+
+	gLpv3DTextureRW0[ texIndex ] = float4( cell.coeffs[0],     cell.AO ); // tex0.a is reserved for secondary occlusion;
+	gLpv3DTextureRW1[ texIndex ] = float4( cell.coeffs[1].rgb, cell.coeffs[2].r );
+	gLpv3DTextureRW2[ texIndex ] = float4( cell.coeffs[2].gb,  cell.coeffs[3].rg );
+	gLpv3DTextureRW3[ texIndex ] = float4( cell.coeffs[3].b,   cell.coeffs[4].rgb );
+	gLpv3DTextureRW4[ texIndex ] = float4( cell.coeffs[5].rgb, cell.coeffs[6].r );
+	gLpv3DTextureRW5[ texIndex ] = float4( cell.coeffs[6].gb,  cell.coeffs[7].rg );
+	gLpv3DTextureRW6[ texIndex ] = float4( cell.coeffs[7].b,   cell.coeffs[8].rgb );
+}
+#endif
\ No newline at end of file
diff --git a/Engine/Shaders/Private/LPVWriteVplCommon.ush b/Engine/Shaders/Private/LPVWriteVplCommon.ush
new file mode 100644
--- /dev/null
+++ b/Engine/Shaders/Private/LPVWriteVplCommon.ush
@@ -0,0 +1,99 @@
+//-----------------------------------------------------------------------------
+// File:		LPVWriteVPLCommmon.usf
+//
+// Summary:		Common functionality for LPV vpl list write
+//
+// Created:		2013-03-01
+//
+// Author:		mailto:[REDACTED_EMAIL]
+//
+//				Copyright (C) Microsoft. All rights reserved.
+//-----------------------------------------------------------------------------
+
+//-------------------------------------------------------------------------------------------------
+
+#include "LPVWriteCommon.ush"
+
+//-------------------------------------------------------------------------------------------------
+
+// Bind UAVs to specific indices for materials with the following flags to ensure the UAVs set match
+// up with what the shader is expecting. Only do this for these materials as it causes other compute
+// shaders which include this file to fail to compile.
+#if MATERIAL_INJECT_EMISSIVE_INTO_LPV || MATERIAL_BLOCK_GI
+RWStructuredBuffer<VplListEntry>				RWGvListBuffer : register (u2);		// 0
+RWByteAddressBuffer								RWGvListHeadBuffer : register (u3);	// 1
+
+RWStructuredBuffer<VplListEntry>				RWVplListBuffer : register (u4);		// 2
+RWByteAddressBuffer								RWVplListHeadBuffer : register (u5);	// 3
+#else
+RWStructuredBuffer<VplListEntry>				RWGvListBuffer;		// 0
+RWByteAddressBuffer								RWGvListHeadBuffer;	// 1
+
+RWStructuredBuffer<VplListEntry>				RWVplListBuffer;		// 2
+RWByteAddressBuffer								RWVplListHeadBuffer;	// 3
+#endif
+
+//-------------------------------------------------------------------------------------------------
+// Functions
+//-------------------------------------------------------------------------------------------------
+int AddToVplList( float3 worldPos, float3 flux, float3 direction, const bool emissiveInjection )
+{
+	int gridIndex = GetGridIndex( worldPos ); 
+	[branch]
+	if ( gridIndex >= 0 )
+	{
+		// Allocate a new list entry 
+		uint newHeadIndex = RWVplListBuffer.IncrementCounter();
+		uint oldHeadIndex;
+		// Note: we add one to the list index to make the list 0-terminated and prevent infinite recursion
+		RWVplListHeadBuffer.InterlockedExchange( gridIndex*4, newHeadIndex+1, oldHeadIndex );
+
+		if ( emissiveInjection )
+		{
+			// Weight according to projected texel size 1 / N dot L
+			float nDotL = abs( dot( LpvWrite.GeometryVolumeCaptureLightDirection.xyz, direction ) );
+			nDotL = max( nDotL, 0.25f ); // Clamp to prevent divBy0 and and general weirdness due to undersampling
+			flux *= 1.0f / nDotL;
+		}
+
+		// Fill in the new list node 
+		Vpl vpl;
+		vpl.normal = direction;
+		vpl.flux = flux;
+
+		VplListEntry listNode = PackVpl( vpl );
+		listNode.nextIndex = oldHeadIndex;
+		RWVplListBuffer[ newHeadIndex ] = listNode;
+	}
+	return gridIndex;
+}
+
+//-------------------------------------------------------------------------------------------------
+
+int AddToGvList( float3 worldPos, float3 flux, float3 direction )
+{
+	// Gv is offset by half 
+	float3 offset = ( LpvWrite.LpvScale * 0.5f ).xxx;
+	int gridIndex = GetGridIndex( worldPos + offset ); 
+	[branch]
+	if ( gridIndex >= 0 )
+	{
+		// Allocate a new list entry 
+		uint newHeadIndex = RWGvListBuffer.IncrementCounter();
+		uint oldHeadIndex;
+		// Note: we add one to the list index to make the list 0-terminated and prevent infinite recursion
+		RWGvListHeadBuffer.InterlockedExchange( gridIndex*4, newHeadIndex+1, oldHeadIndex );
+
+		// Fill in the new list node 
+		Vpl vpl;
+		vpl.normal = direction;
+		vpl.flux = flux;
+
+		VplListEntry listNode = PackVpl( vpl );
+		listNode.nextIndex = oldHeadIndex;
+		RWGvListBuffer[ newHeadIndex ] = listNode;
+	}
+	return gridIndex;
+}
+
+//-------------------------------------------------------------------------------------------------
diff --git a/Engine/Shaders/Private/LandscapeVertexFactory.ush b/Engine/Shaders/Private/LandscapeVertexFactory.ush
new file mode 100644
--- /dev/null
+++ b/Engine/Shaders/Private/LandscapeVertexFactory.ush
@@ -0,0 +1,874 @@
+// Copyright 1998-2017 Epic Games, Inc. All Rights Reserved.
+
+/*=============================================================================
+	LandscapeVertexFactory.usf: Landscape vertex factory.
+=============================================================================*/
+
+#include "VertexFactoryCommon.ush"
+
+#define ES2_MAXLOD	5
+#define VERTEX_FACTORY_MODIFIES_TESSELLATION 1
+
+// Set by FLightMapTexturePolicy
+#include "/Engine/Generated/UniformBuffers/PrecomputedLightingBuffer.ush"
+
+/* SM4+:
+   x = unused
+   y = unused
+   z = Heightmap texture LOD difference from current LOD to highest LOD
+   w = XYOffset texture LOD difference from current LOD to highest LOD
+
+   ES2 / ES3.1:
+   x = unused
+   y = unused
+   z = TexCoordOffset.x
+   w = TexCoordOffset.y
+*/
+float4 LodBias;
+
+/* x = current LOD
+   y = unused
+   z = SectionSizeQuads in current LOD
+   w = 1/SectionSizeQuads in current LOD
+*/
+float4 LodValues;
+
+/* Fractional Lod values for each section */
+float4 SectionLods;
+
+/* Array index is section
+   x = neighbor section's LOD in position 0               0     -> X
+   y = neighbor section's LOD in position 1        1   Current  2
+   z = neighbor section's LOD in position 2        |      3
+   w = neighbor section's LOD in position 3        V Y
+*/
+float4 NeighborSectionLod[4];
+
+float4x4 LocalToWorldNoScaling;
+
+#if FEATURE_LEVEL >= FEATURE_LEVEL_SM4
+Texture2D HeightmapTexture;
+SamplerState HeightmapTextureSampler;
+#endif
+#if PIXELSHADER
+Texture2D NormalmapTexture;
+SamplerState NormalmapTextureSampler;
+#endif
+
+#if LANDSCAPE_XYOFFSET
+#define XYOFFSET_SCALE			(1.0f/256.f)
+Texture2D XYOffsetmapTexture;
+SamplerState XYOffsetmapTextureSampler;
+#endif
+
+#if ES2_PROFILE || ES3_1_PROFILE
+float3 BlendableLayerMask;
+#endif
+
+struct FVertexFactoryInput
+{
+#if FEATURE_LEVEL >= FEATURE_LEVEL_SM4
+	float4 Position: ATTRIBUTE0;
+	/** Optional instance ID for vertex layered rendering */
+#if ONEPASS_POINTLIGHT_SHADOW && USING_VERTEX_SHADER_LAYER
+	uint InstanceId	: SV_InstanceID;
+#endif
+#else
+	float4 PackedPosition: ATTRIBUTE0; 
+	float4 LODHeights[2]: ATTRIBUTE1;  
+#endif
+};
+
+struct FVertexFactoryInterpolantsVSToPS
+{
+	float2	LayerTexCoord			: TEXCOORD0;
+#if (ES2_PROFILE || ES3_1_PROFILE)
+	float2	WeightMapTexCoord		: TEXCOORD1;
+	float4	TransformedTexCoord0	: TEXCOORD2;
+	float4  TransformedTexCoord1	: TEXCOORD3;
+#else
+	float4	WeightHeightMapTexCoord	: TEXCOORD1;
+	float4	TransformedTexCoords	: TEXCOORD2;
+#endif
+#if NEEDS_LIGHTMAP_COORDINATE
+#if (ES2_PROFILE || ES3_1_PROFILE)
+	float2	LightMapCoordinate[2]	: TEXCOORD4;
+	float2	ShadowMapCoordinate		: TEXCOORD6;
+#else
+	float4	LightMapCoordinate		: TEXCOORD3;
+#endif
+#endif
+
+#if INSTANCED_STEREO
+	nointerpolation uint PackedEyeIndex : PACKED_EYE_INDEX;
+#endif
+};
+
+struct FLandscapeTexCoords
+{
+	float2	LayerTexCoord;
+	float2	WeightMapTexCoord;
+#if !(ES2_PROFILE || ES3_1_PROFILE)
+	float2	HeightMapTexCoord;
+#endif
+#if NEEDS_LIGHTMAP_COORDINATE
+	float2	LightMapCoordinate;
+#endif
+};
+
+struct FVertexFactoryInterpolantsVSToDS
+{
+	FVertexFactoryInterpolantsVSToPS InterpolantsVSToPS;
+	
+#if USING_TESSELLATION
+	// First row of the tangent to world matrix
+	float3	TangentToWorld0	: VS_To_DS_TangentToWorld0;
+	// Last row of the tangent to world matrix in xyz
+	float4	TangentToWorld2	: VS_To_DS_TangentToWorld2;
+	// LOD of the vertex, used for fading out tessellation
+	float CalculatedLOD : VS_To_DS_CalculatedLOD;
+#endif
+};
+
+#if USING_TESSELLATION
+float2 VertexFactoryGetTextureCoordinateDS( FVertexFactoryInterpolantsVSToDS Interpolants )
+{
+#if (ES2_PROFILE || ES3_1_PROFILE)
+	return Interpolants.InterpolantsVSToPS.WeightMapTexCoord.xy;
+#else
+	return Interpolants.InterpolantsVSToPS.WeightHeightMapTexCoord.zw;
+#endif
+}
+
+float VertexFactoryModifyTessellationMultiplier(FVertexFactoryInterpolantsVSToDS Interpolants, float3 CameraLocalWorldPosition, float MaterialTessellationMultiplier)
+{
+	// The 3.0f is a tunable value for falloff speed
+	float TessellationFalloff = saturate((1.0f - Interpolants.CalculatedLOD) * 3.0f);
+
+	return MaterialTessellationMultiplier * TessellationFalloff;
+}
+#endif // #if USING_TESSELLATION
+
+FVertexFactoryInterpolantsVSToPS VertexFactoryAssignInterpolants(FVertexFactoryInterpolantsVSToDS Input)
+{
+	return Input.InterpolantsVSToPS;
+}
+
+struct FVertexFactoryIntermediates
+{
+	float4 InputPosition;
+	float3 LocalPosition;
+	float3 WorldNormal;
+};
+
+float3 GetLocalPosition(FVertexFactoryIntermediates Intermediates)
+{
+	return INVARIANT(Intermediates.LocalPosition+float3(Intermediates.InputPosition.zw * LandscapeParameters.SubsectionOffsetParams.ww,0));
+}
+
+float4 VertexFactoryGetWorldPosition(FVertexFactoryInput Input, FVertexFactoryIntermediates Intermediates)
+{
+	return INVARIANT(TransformLocalToTranslatedWorld(GetLocalPosition(Intermediates)));
+}
+
+float4 VertexFactoryGetPreviousWorldPosition(FVertexFactoryInput Input, FVertexFactoryIntermediates Intermediates)
+{
+	// Note we still use LocalToWorld. Landscape's transform never changes.
+	float3 LocalPosition = GetLocalPosition(Intermediates);
+	float3 RotatedPosition = INVARIANT(Primitive.LocalToWorld[0].xyz * LocalPosition.xxx + Primitive.LocalToWorld[1].xyz * LocalPosition.yyy + Primitive.LocalToWorld[2].xyz * LocalPosition.zzz);
+	return INVARIANT(float4(RotatedPosition + (Primitive.LocalToWorld[3].xyz + ResolvedView.PrevPreViewTranslation.xyz),1));
+}
+
+/** Calculate the texture coordinates generated by Landscape */
+FLandscapeTexCoords GetLandscapeTexCoords(FVertexFactoryInput Input, FVertexFactoryIntermediates Intermediates)
+{
+	FLandscapeTexCoords Result;
+
+	Result.LayerTexCoord = Intermediates.LocalPosition.xy + LandscapeParameters.SubsectionSizeVertsLayerUVPan.zw + Intermediates.InputPosition.zw * LandscapeParameters.SubsectionOffsetParams.ww;
+
+	Result.WeightMapTexCoord	= Intermediates.LocalPosition.xy * LandscapeParameters.WeightmapUVScaleBias.xy + LandscapeParameters.WeightmapUVScaleBias.zw + Intermediates.InputPosition.zw * LandscapeParameters.SubsectionOffsetParams.zz;
+#if !(ES2_PROFILE || ES3_1_PROFILE)
+	Result.HeightMapTexCoord	= Intermediates.LocalPosition.xy * LandscapeParameters.HeightmapUVScaleBias.xy + LandscapeParameters.HeightmapUVScaleBias.zw + 0.5*LandscapeParameters.HeightmapUVScaleBias.xy + Intermediates.InputPosition.zw * LandscapeParameters.SubsectionOffsetParams.xy;
+#endif
+
+#if NEEDS_LIGHTMAP_COORDINATE
+	Result.LightMapCoordinate			= (Intermediates.LocalPosition.xy * LandscapeParameters.LandscapeLightmapScaleBias.xy + LandscapeParameters.LandscapeLightmapScaleBias.wz + Intermediates.InputPosition.zw * LandscapeParameters.LightmapSubsectionOffsetParams.xy);
+#endif
+	
+	return Result;
+}
+
+float3x3 CalcTangentBasisFromWorldNormal(float3 Normal)
+{
+	float3 LocalTangentX = normalize(float3(Normal.z, 0, -Normal.x));
+	float3 LocalTangentY = cross(Normal, LocalTangentX);
+
+	float3x3 LocalToTangent = float3x3(LocalTangentX,LocalTangentY,Normal);
+
+	return LocalToTangent;
+}
+
+/** Lookup per-pixel tangent basis from heightmap texture */
+float3x3 VertexFactoryGetPerPixelTangentBasis(FVertexFactoryInterpolantsVSToPS Interpolants)
+{
+	float3x3 Result;
+#if PIXELSHADER
+#if FEATURE_LEVEL >= FEATURE_LEVEL_SM4
+	float4 SampleValue = Texture2DSample(NormalmapTexture, NormalmapTextureSampler, Interpolants.WeightHeightMapTexCoord.zw);
+#else
+	float4 SampleValue = Texture2DSample(NormalmapTexture, NormalmapTextureSampler, Interpolants.WeightMapTexCoord);
+#endif
+#if WEBGL
+	float2 SampleNormal = float2(SampleValue.b, SampleValue.a) * float2(2.0,2.0) - float2(1.0,1.0);
+#else
+	float2 SampleNormal = float2(SampleValue.b, SampleValue.a) * float2(2.0,2.0) - float2(1.0,1.0);
+#endif
+	float3 WorldNormal = float3( SampleNormal, sqrt(max(1.0-dot(SampleNormal,SampleNormal),0.0)) );
+
+	Result = CalcTangentBasisFromWorldNormal(WorldNormal);
+#endif
+	return Result;
+}
+
+/** Converts from vertex factory specific interpolants to a FMaterialTessellationParameters, which is used by material inputs. */
+FMaterialTessellationParameters GetMaterialTessellationParameters(FVertexFactoryInterpolantsVSToDS Interpolants, float3 CameraLocalWorldPosition)
+{
+	FMaterialTessellationParameters Result;
+
+#if NUM_MATERIAL_TEXCOORDS     // XY layer
+	Result.TexCoords[0] = Interpolants.InterpolantsVSToPS.LayerTexCoord.xy;
+#if NUM_MATERIAL_TEXCOORDS > 1 // VS calcualted TexCoord 1, default is XZ layer
+	#if (ES2_PROFILE || ES3_1_PROFILE)
+		Result.TexCoords[1] = Interpolants.InterpolantsVSToPS.TransformedTexCoord0.xy;
+	#else
+		Result.TexCoords[1] = Interpolants.InterpolantsVSToPS.TransformedTexCoords.xy;
+	#endif
+#if NUM_MATERIAL_TEXCOORDS > 2 // VS calcualted TexCoord 2, default is YZ layer
+	#if (ES2_PROFILE || ES3_1_PROFILE)
+		Result.TexCoords[2] = Interpolants.InterpolantsVSToPS.TransformedTexCoord1.xy;
+	#else
+		Result.TexCoords[2] = Interpolants.InterpolantsVSToPS.TransformedTexCoords.zw;
+	#endif
+#if NUM_MATERIAL_TEXCOORDS > 3 // Weightmap
+	#if (ES2_PROFILE || ES3_1_PROFILE)
+		Result.TexCoords[3] = Interpolants.InterpolantsVSToPS.WeightMapTexCoord;
+	#else
+		Result.TexCoords[3] = Interpolants.InterpolantsVSToPS.WeightHeightMapTexCoord.xy;
+	#endif
+#if NUM_MATERIAL_TEXCOORDS > 4 // Lightmap
+	#if NEEDS_LIGHTMAP_COORDINATE
+		#if (ES2_PROFILE || ES3_1_PROFILE)
+			Result.TexCoords[4] = Interpolants.InterpolantsVSToPS.LightMapCoordinate[0].xy;
+		#else
+			Result.TexCoords[4] = Interpolants.InterpolantsVSToPS.LightMapCoordinate.xy;
+		#endif
+	#else
+		Result.TexCoords[4] = float2(0,0);
+	#endif
+#if NUM_MATERIAL_TEXCOORDS > 5 // Height map
+	#if (ES2_PROFILE || ES3_1_PROFILE)
+		Result.TexCoords[5] = float2(0, 0); // Height map UV not supported
+	#else
+		Result.TexCoords[5] = Interpolants.InterpolantsVSToPS.WeightHeightMapTexCoord.zw;
+	#endif
+#if NUM_MATERIAL_TEXCOORDS > 6
+	UNROLL
+	for(uint CoordinateIndex = 6;CoordinateIndex < NUM_MATERIAL_TEXCOORDS;CoordinateIndex ++)
+	{
+		Result.TexCoords[CoordinateIndex] = float2(0,0);
+	}
+#endif // 6
+#endif // 5
+#endif // 4
+#endif // 3
+#endif // 2
+#endif // 1
+#endif // 0
+
+	// Use interpolated normal from VS.
+#if USING_TESSELLATION
+	Result.TangentToWorld = AssembleTangentToWorld( Interpolants.TangentToWorld0, Interpolants.TangentToWorld2 );
+#endif
+
+	Result.VertexColor = 1;
+	Result.TangentToWorldPreScale = 1;
+	Result.WorldPosition = CameraLocalWorldPosition + ResolvedView.WorldCameraOrigin;
+
+	return Result;
+}
+
+/** Converts from vertex factory specific interpolants to a FMaterialPixelParameters, which is used by material inputs. */
+FMaterialPixelParameters GetMaterialPixelParameters(FVertexFactoryInterpolantsVSToPS Interpolants, float4 SvPosition)
+{
+	// GetMaterialPixelParameters is responsible for fully initializing the result
+	FMaterialPixelParameters Result = MakeInitializedMaterialPixelParameters();
+
+#if NEEDS_LIGHTMAP_COORDINATE
+	#if (ES2_PROFILE || ES3_1_PROFILE)
+		// Not supported in pixel shader
+		float2 LightmapUVs = float2(0, 0);
+	#else
+		float2 LightmapUVs = Interpolants.LightMapCoordinate.xy;
+	#endif
+#else
+	float2 LightmapUVs = float2(0, 0);
+#endif
+
+#if NUM_MATERIAL_TEXCOORDS     // XY layer
+	Result.TexCoords[0] = Interpolants.LayerTexCoord.xy;
+#if NUM_MATERIAL_TEXCOORDS > 1 // VS calcualted TexCoord 1, default is XZ layer
+	#if (ES2_PROFILE || ES3_1_PROFILE)
+		Result.TexCoords[1] = Interpolants.TransformedTexCoord0.xy;
+	#else
+		Result.TexCoords[1] = Interpolants.TransformedTexCoords.xy;
+	#endif
+#if NUM_MATERIAL_TEXCOORDS > 2 // VS calcualted TexCoord 2, default is YZ layer
+	#if (ES2_PROFILE || ES3_1_PROFILE)
+		Result.TexCoords[2] = Interpolants.TransformedTexCoord1.xy;
+	#else
+		Result.TexCoords[2] = Interpolants.TransformedTexCoords.zw;
+	#endif
+#if NUM_MATERIAL_TEXCOORDS > 3 // Weightmap
+	#if (ES2_PROFILE || ES3_1_PROFILE)
+		Result.TexCoords[3] = Interpolants.WeightMapTexCoord;
+	#else
+		Result.TexCoords[3] = Interpolants.WeightHeightMapTexCoord.xy;
+	#endif
+#if NUM_MATERIAL_TEXCOORDS > 4 // Lightmap
+	#if NEEDS_LIGHTMAP_COORDINATE
+		Result.TexCoords[4] = LightmapUVs;
+	#else
+		Result.TexCoords[4] = float2(0,0);
+	#endif
+#if NUM_MATERIAL_TEXCOORDS > 5 // Height map
+	#if (ES2_PROFILE || ES3_1_PROFILE)
+		Result.TexCoords[5] = float2(0,0); // Height map UV not supported
+	#else
+		Result.TexCoords[5] = Interpolants.WeightHeightMapTexCoord.zw;
+	#endif
+#if NUM_MATERIAL_TEXCOORDS > 6
+	for (uint CoordinateIndex = 6; CoordinateIndex < NUM_MATERIAL_TEXCOORDS; CoordinateIndex++)
+	{
+		Result.TexCoords[CoordinateIndex] = float2(0,0);
+	}
+#endif // 6
+#endif // 5
+#endif // 4
+#endif // 3
+#endif // 2
+#endif // 1
+#endif // 0
+
+	// Calculate LocalToTangent directly from normal map texture.
+	float3x3 TangentToLocal = VertexFactoryGetPerPixelTangentBasis(Interpolants);
+	Result.TangentToWorld = mul(TangentToLocal, (float3x3)LocalToWorldNoScaling); 
+	Result.UnMirrored = 1;
+
+	Result.VertexColor = 1;
+
+#if LIGHTMAP_UV_ACCESS
+	Result.LightmapUVs = LightmapUVs;
+#endif
+
+#if (ES2_PROFILE || ES3_1_PROFILE) && PIXELSHADER
+
+#if WEBGL
+	Result.LayerWeights.xy = Texture2DSample(NormalmapTexture, NormalmapTextureSampler, Interpolants.WeightMapTexCoord).rg;
+#else
+	Result.LayerWeights.xy = Texture2DSample(NormalmapTexture, NormalmapTextureSampler, Interpolants.WeightMapTexCoord).rg;
+#endif 
+
+	Result.LayerWeights.z = saturate(BlendableLayerMask.z - dot(Result.LayerWeights.xy, BlendableLayerMask.xy));
+	Result.LayerWeights.w = 0; // Total weight need to be 1
+#endif
+
+	Result.TwoSidedSign = 1;
+	return Result;
+}
+
+/** Converts from vertex factory specific input to a FMaterialVertexParameters, which is used by vertex shader material inputs. */
+FMaterialVertexParameters GetMaterialVertexParameters(FVertexFactoryInput Input, FVertexFactoryIntermediates Intermediates, float3 WorldPosition, float3x3 TangentToLocal)
+{
+	FMaterialVertexParameters Result = (FMaterialVertexParameters)0;
+	Result.WorldPosition = WorldPosition;
+	Result.VertexColor = float4(1,1,1,1);
+	Result.TangentToWorld = mul(TangentToLocal, GetLocalToWorld3x3()); 
+	Result.PreSkinnedPosition = Intermediates.LocalPosition.xyz;
+	Result.PreSkinnedNormal = TangentToLocal[2].xyz;
+
+	FLandscapeTexCoords LandscapeTexCoords = GetLandscapeTexCoords(Input, Intermediates);
+
+#if NUM_MATERIAL_TEXCOORDS_VERTEX     // XY layer
+	Result.TexCoords[0] = LandscapeTexCoords.LayerTexCoord.xy;
+	#if (ES2_PROFILE || ES3_1_PROFILE)
+		Result.TexCoordOffset = LodBias.zw;
+	#endif
+#if NUM_MATERIAL_TEXCOORDS_VERTEX > 1 // XZ layer
+	Result.TexCoords[1] = float2(LandscapeTexCoords.LayerTexCoord.x, Intermediates.LocalPosition.z);
+#if NUM_MATERIAL_TEXCOORDS_VERTEX > 2 // YZ layer
+	Result.TexCoords[2] = float2(LandscapeTexCoords.LayerTexCoord.y, Intermediates.LocalPosition.z);
+#if NUM_MATERIAL_TEXCOORDS_VERTEX > 3 // Weightmap
+	Result.TexCoords[3] = LandscapeTexCoords.WeightMapTexCoord;
+#if NUM_MATERIAL_TEXCOORDS_VERTEX > 4 // Lightmap
+	#if NEEDS_LIGHTMAP_COORDINATE
+		Result.TexCoords[4] = LandscapeTexCoords.LightMapCoordinate.xy;
+	#else
+		Result.TexCoords[4] = float2(0,0);
+	#endif
+#if NUM_MATERIAL_TEXCOORDS_VERTEX > 5 // Height map
+	#if (ES2_PROFILE || ES3_1_PROFILE)
+		Result.TexCoords[5] = float2(0,0); // Height map UV not supported
+	#else
+		Result.TexCoords[5] = LandscapeTexCoords.HeightMapTexCoord;
+	#endif
+#if NUM_MATERIAL_TEXCOORDS_VERTEX > 6
+	UNROLL
+	for (uint CoordinateIndex = 6; CoordinateIndex < NUM_MATERIAL_TEXCOORDS; CoordinateIndex++)
+	{
+		Result.TexCoords[CoordinateIndex] = float2(0,0);
+	}
+#endif // 6
+#endif // 5
+#endif // 4
+#endif // 3
+#endif // 2
+#endif // 1
+#endif // 0
+
+	return Result;
+}
+
+#if NEEDS_LIGHTMAP_COORDINATE
+void GetLightMapCoordinates(FVertexFactoryInterpolantsVSToPS Interpolants, out float2 LightmapUV0, out float2 LightmapUV1)
+{
+#if (ES2_PROFILE || ES3_1_PROFILE)
+	LightmapUV0 = Interpolants.LightMapCoordinate[0].xy;
+	LightmapUV1 = Interpolants.LightMapCoordinate[1].xy;
+#else
+	LightmapUV0 = Interpolants.LightMapCoordinate.xy * float2( 1, 0.5 );
+	LightmapUV1 = LightmapUV0 + float2( 0, 0.5 );
+#endif
+}
+
+float2 GetShadowMapCoordinate(FVertexFactoryInterpolantsVSToPS Interpolants)
+{
+#if (ES2_PROFILE || ES3_1_PROFILE)
+	return Interpolants.ShadowMapCoordinate.xy;
+#else
+	return Interpolants.LightMapCoordinate.zw;
+#endif
+}
+#endif
+
+float CalcLOD_ES2(float2 xy, float2 Subsection)
+{
+	// LOD value based on Barycentric coordinate
+	float4 L0 = float4(xy.y, xy.x, (1 - xy.x), (1 - xy.y)) * 2;
+
+	float4 LODCalculated4 = lerp(
+								(lerp(L0 * SectionLods.x + (float4(1, 1, 1, 1) - L0) * NeighborSectionLod[0],
+									  L0 * SectionLods.y + (float4(1, 1, 1, 1) - L0) * NeighborSectionLod[1],
+									  Subsection.x)),
+								(lerp(L0 * SectionLods.z + (float4(1, 1, 1, 1) - L0) * NeighborSectionLod[2],
+									  L0 * SectionLods.w + (float4(1, 1, 1, 1) - L0) * NeighborSectionLod[3],
+									  Subsection.x)),
+								Subsection.y);
+
+	// Find which quadrant of the subsection we're in - top, left, right or bottom.
+	// 0 --------- 1 X  1 ----- 2  1          1 
+	//   | \ 0 / |        \   /    | \      / | 
+	//   | 1 + 2 |          0      |  0    0  |      0       
+	//   | / 3 \ |                 | /      \ |    /   \     
+	// 2 --------- 3               2          2  1 ----- 2   
+	float LODCalculated = ((xy.x + xy.y) > 1) ? 
+							(xy.x < xy.y ? LODCalculated4.w : LODCalculated4.z) :
+							(xy.x < xy.y ? LODCalculated4.y : LODCalculated4.x);
+
+	return LODCalculated;
+}
+
+float CalcLOD(float2 xy, float2 Subsection)
+{
+#if ES2_PROFILE
+	// less branchy version for ES2 devices
+	return CalcLOD_ES2(xy, Subsection); 
+#else
+	// LOD value based on Barycentric coordinate
+	float4 L0 = float4(xy.y, xy.x, (1 - xy.x), (1 - xy.y)) * 2;
+
+	float4 LODCalculated4;
+	if (Subsection.y > 0.5f)
+	{
+		if (Subsection.x > 0.5f)
+		{
+			float SectionLod = SectionLods.w;
+			LODCalculated4 = L0 * SectionLod + (float4(1, 1, 1, 1) - L0) * NeighborSectionLod[3];
+		}
+		else
+		{
+			float SectionLod = SectionLods.z;
+			LODCalculated4 = L0 * SectionLod + (float4(1, 1, 1, 1) - L0) * NeighborSectionLod[2];
+		}
+	}
+	else
+	{
+		if (Subsection.x > 0.5f)
+		{
+			float SectionLod = SectionLods.y;
+			LODCalculated4 = L0 * SectionLod + (float4(1, 1, 1, 1) - L0) * NeighborSectionLod[1];
+		}
+		else
+		{
+			float SectionLod = SectionLods.x;
+			LODCalculated4 = L0 * SectionLod + (float4(1, 1, 1, 1) - L0) * NeighborSectionLod[0];
+		}
+	}
+
+	// Find which quadrant of the subsection we're in - top, left, right or bottom.
+	// 0 --------- 1 X  1 ----- 2  1          1 
+	//   | \ 0 / |        \   /    | \      / | 
+	//   | 1 + 2 |          0      |  0    0  |      0       
+	//   | / 3 \ |                 | /      \ |    /   \     
+	// 2 --------- 3               2          2  1 ----- 2   
+	float LODCalculated;
+	if ((xy.x + xy.y) > 1)
+	{
+		if (xy.x < xy.y)
+		{
+			LODCalculated = LODCalculated4.w;
+		}
+		else
+		{
+			LODCalculated = LODCalculated4.z;
+		}
+	}
+	else
+	{
+		if (xy.x < xy.y)
+		{
+			LODCalculated = LODCalculated4.y;
+		}
+		else
+		{
+			LODCalculated = LODCalculated4.x;
+		}
+	}
+
+	return LODCalculated;
+#endif
+}
+
+FVertexFactoryIntermediates GetVertexFactoryIntermediates(FVertexFactoryInput Input)
+{
+	FVertexFactoryIntermediates Intermediates;
+
+#if FEATURE_LEVEL >= FEATURE_LEVEL_SM4
+	Intermediates.InputPosition = Input.Position;
+#else
+	Intermediates.InputPosition = Input.PackedPosition * 255.f;
+	float2 SubXY = frac((Intermediates.InputPosition.zw) * .5f) * 2.f;
+	float2 LOD45 = (Intermediates.InputPosition.zw - SubXY) / 255.f;
+	Intermediates.InputPosition.zw = SubXY;
+#endif
+
+	float2 xy = Intermediates.InputPosition.xy * LodValues.w;
+
+	float LODCalculated = CalcLOD(xy, Intermediates.InputPosition.zw);
+
+	float LodValue = floor(LODCalculated);
+	float MorphAlpha = frac(LODCalculated);
+
+	// InputPositionLODAdjusted : Position for actual LOD in base LOD units
+#if FEATURE_LEVEL >= FEATURE_LEVEL_SM4
+	float2 ActualLODCoordsInt = floor(Intermediates.InputPosition.xy * pow(2, -(LodValue - LodValues.x)));
+	float InvLODScaleFactor = pow(2, -LodValue);
+#else
+	float3 HeightData = (LodValue < 1 ? float3(1.0, Input.LODHeights[1][0], Input.LODHeights[1][1]) :
+						(LodValue < 2 ? float3(0.5, Input.LODHeights[1][1], Input.LODHeights[1][2]) :
+						(LodValue < 3 ? float3(0.25, Input.LODHeights[1][2], Input.LODHeights[1][3]) :
+						(LodValue < 4 ? float3(0.125, Input.LODHeights[1][3], LOD45[0]) :
+						(LodValue < 5 ? float3(0.0625, LOD45[0], LOD45[1]) :
+						(float3(0.03125, LOD45[1], LOD45[1])))))));
+
+	// we can hardcode InvLODScaleFactor
+	float InvLODScaleFactor = HeightData.x;
+	// select the height values from the vertex input corresponding to the current and next LOD
+	float InputHeight = HeightData.y;
+	float InputHeightNextLOD = HeightData.z;
+
+	float MinHeight = DecodePackedHeight(Input.LODHeights[0].xy);
+	float MaxHeight = DecodePackedHeight(Input.LODHeights[0].zw);
+
+	float Height = lerp(MinHeight, MaxHeight, InputHeight);
+	float HeightNextLOD = lerp(MinHeight, MaxHeight, InputHeightNextLOD);
+
+	// LodValues.x is always 0 on mobile.
+	float2 ActualLODCoordsInt = floor(Intermediates.InputPosition.xy * InvLODScaleFactor);
+#endif
+
+	// Base to Actual LOD, Base to Next LOD
+	float2 CoordTranslate = float2( LandscapeParameters.SubsectionSizeVertsLayerUVPan.x * InvLODScaleFactor - 1, max(LandscapeParameters.SubsectionSizeVertsLayerUVPan.x * 0.5f * InvLODScaleFactor, 2) - 1 ) * LandscapeParameters.SubsectionSizeVertsLayerUVPan.y;
+	float2 InputPositionLODAdjusted = ActualLODCoordsInt / CoordTranslate.x;
+
+	// InputPositionNextLOD : Position for next LOD in base LOD units
+	float2 NextLODCoordsInt = floor(ActualLODCoordsInt * 0.5);
+	float2 InputPositionNextLOD = NextLODCoordsInt / CoordTranslate.y;
+
+	// Get the height and normal XY for current and next LOD out of the textures
+#if FEATURE_LEVEL >= FEATURE_LEVEL_SM4
+	float2 SampleCoords = InputPositionLODAdjusted * LandscapeParameters.HeightmapUVScaleBias.xy + LandscapeParameters.HeightmapUVScaleBias.zw + 0.5*LandscapeParameters.HeightmapUVScaleBias.xy + Intermediates.InputPosition.zw * LandscapeParameters.SubsectionOffsetParams.xy;
+	float4 SampleValue = Texture2DSampleLevel(HeightmapTexture, HeightmapTextureSampler, SampleCoords, LodValue-LodBias.z);
+	float Height = DecodePackedHeight(SampleValue.xy);
+
+	float2 SampleCoordsNextLOD = InputPositionNextLOD * LandscapeParameters.HeightmapUVScaleBias.xy + LandscapeParameters.HeightmapUVScaleBias.zw + 0.5*LandscapeParameters.HeightmapUVScaleBias.xy + Intermediates.InputPosition.zw * LandscapeParameters.SubsectionOffsetParams.xy;
+	float4 SampleValueNextLOD = Texture2DSampleLevel(HeightmapTexture, HeightmapTextureSampler, SampleCoordsNextLOD, LodValue+1-LodBias.z );
+	float HeightNextLOD = DecodePackedHeight(SampleValueNextLOD.xy);
+#endif
+
+
+#if LANDSCAPE_XYOFFSET // FEATURE_LEVEL >= FEATURE_LEVEL_SM4 only
+	float2 SampleCoords2 = float2(InputPositionLODAdjusted * LandscapeParameters.WeightmapUVScaleBias.xy + LandscapeParameters.WeightmapUVScaleBias.zw + Intermediates.InputPosition.zw * LandscapeParameters.SubsectionOffsetParams.zz);
+	float4 OffsetValue = Texture2DSampleLevel( XYOffsetmapTexture, XYOffsetmapTextureSampler, SampleCoords2, LodValue-LodBias.w );
+	float2 SampleCoordsNextLOD2 = float2(InputPositionNextLOD * LandscapeParameters.WeightmapUVScaleBias.xy + LandscapeParameters.WeightmapUVScaleBias.zw + Intermediates.InputPosition.zw * LandscapeParameters.SubsectionOffsetParams.zz);
+	float4 OffsetValueNextLOD = Texture2DSampleLevel( XYOffsetmapTexture, XYOffsetmapTextureSampler, SampleCoordsNextLOD2, LodValue+1-LodBias.w );
+	float2 XYOffset = float2(((OffsetValue.r * 255.0 * 256.0 + OffsetValue.g * 255.0) - 32768.0) * XYOFFSET_SCALE, ((OffsetValue.b * 255.0 * 256.0 + OffsetValue.a * 255.0) - 32768.0) * XYOFFSET_SCALE );
+	float2 XYOffsetNextLOD = float2(((OffsetValueNextLOD.r * 255.0 * 256.0 + OffsetValueNextLOD.g * 255.0) - 32768.0) * XYOFFSET_SCALE, ((OffsetValueNextLOD.b * 255.0 * 256.0 + OffsetValueNextLOD.a * 255.0) - 32768.0) * XYOFFSET_SCALE );
+
+	InputPositionLODAdjusted = InputPositionLODAdjusted + XYOffset;
+	InputPositionNextLOD = InputPositionNextLOD + XYOffsetNextLOD;
+#endif
+
+	Intermediates.LocalPosition = lerp( float3(InputPositionLODAdjusted, Height), float3(InputPositionNextLOD, HeightNextLOD), MorphAlpha );
+
+#if FEATURE_LEVEL >= FEATURE_LEVEL_SM4
+	float2 Normal = float2(SampleValue.b, SampleValue.a);
+	float2 NormalNextLOD = float2(SampleValueNextLOD.b, SampleValueNextLOD.a);
+	float2 InterpNormal = lerp( Normal, NormalNextLOD, MorphAlpha ) * float2(2.0,2.0) - float2(1.0,1.0);
+	Intermediates.WorldNormal = float3( InterpNormal, sqrt(max(1.0-dot(InterpNormal,InterpNormal),0.0)) );
+#else
+	Intermediates.WorldNormal = float3( 0.0, 0.0, 1.0 );
+#endif
+	
+	return Intermediates;
+}
+
+#if RENDERING_LANDSCAPE_GRASS && FEATURE_LEVEL >= FEATURE_LEVEL_SM4 && VERTEXSHADER
+float2 LandscapeVertexFactorySampleHeight(FVertexFactoryInput Input, FVertexFactoryIntermediates Intermediates)
+{
+	/*
+	FLandscapeTexCoords TexCoords = GetLandscapeTexCoords(Input, Intermediates);
+
+	float4 SampleValue = Texture2DSampleLevel(HeightmapTexture, HeightmapTextureSampler, TexCoords.HeightMapTexCoord, 0);
+	return SampleValue.xy;
+	*/
+
+	float2 xy = Intermediates.InputPosition.xy * LodValues.w;
+	float LODCalculated = CalcLOD(xy, Intermediates.InputPosition.zw);
+	float LodValue = floor(LODCalculated);
+	float CoordTranslate_x = (LandscapeParameters.SubsectionSizeVertsLayerUVPan.x * pow(2, -LodValue) - 1) * LandscapeParameters.SubsectionSizeVertsLayerUVPan.y;
+
+	// InputPositionLODAdjusted : Position for actual LOD in base LOD units
+	float2 ActualLODCoordsInt = floor(Intermediates.InputPosition.xy * pow(2, -(LodValue - LodValues.x)));
+
+	float2 InputPositionLODAdjusted = ActualLODCoordsInt / CoordTranslate_x;
+	float2 SampleCoords = InputPositionLODAdjusted * LandscapeParameters.HeightmapUVScaleBias.xy + LandscapeParameters.HeightmapUVScaleBias.zw + 0.5*LandscapeParameters.HeightmapUVScaleBias.xy + Intermediates.InputPosition.zw * LandscapeParameters.SubsectionOffsetParams.xy;
+	float4 SampleValue = Texture2DSampleLevel(HeightmapTexture, HeightmapTextureSampler, SampleCoords, LodValue-LodBias.z);
+	return SampleValue.rg;
+}
+#endif
+
+
+/**
+* Get the 3x3 tangent basis vectors for this vertex factory
+* this vertex factory will calculate the binormal on-the-fly
+*
+* @param Input - vertex input stream structure
+* @return 3x3 matrix
+*/
+float3x3 VertexFactoryGetTangentToLocal( FVertexFactoryInput Input, FVertexFactoryIntermediates Intermediates )
+{
+	float3x3 Result = CalcTangentBasisFromWorldNormal(Intermediates.WorldNormal);
+	return Result;
+}
+
+float4 VertexFactoryGetRasterizedWorldPosition(FVertexFactoryInput Input, FVertexFactoryIntermediates Intermediates, float4 InWorldPosition)
+{
+#if RENDERING_LANDSCAPE_GRASS
+	float2 xy = Intermediates.InputPosition.xy * LodValues.w;
+	float LODCalculated = CalcLOD(xy, Intermediates.InputPosition.zw);
+	float LodValue = floor(LODCalculated);
+	float CoordTranslate_x = (LandscapeParameters.SubsectionSizeVertsLayerUVPan.x * pow(2, -LodValue) - 1) * LandscapeParameters.SubsectionSizeVertsLayerUVPan.y;
+
+	float4 LocalPosition = mul(InWorldPosition, Primitive.WorldToLocal);
+	LocalPosition.xy *= CoordTranslate_x;
+
+	return mul(LocalPosition, Primitive.LocalToWorld);
+#else
+	return InWorldPosition;
+#endif
+}
+
+float3 VertexFactoryGetPositionForVertexLighting(FVertexFactoryInput Input, FVertexFactoryIntermediates Intermediates, float3 TranslatedWorldPosition)
+{
+	return TranslatedWorldPosition;
+}
+
+FVertexFactoryInterpolantsVSToPS VertexFactoryGetInterpolantsVSToPS(FVertexFactoryInput Input, FVertexFactoryIntermediates Intermediates, FMaterialVertexParameters VertexParameters)
+{
+	FVertexFactoryInterpolantsVSToPS Interpolants;
+
+	Interpolants = (FVertexFactoryInterpolantsVSToPS)0;
+	FLandscapeTexCoords LandscapeTexCoords = GetLandscapeTexCoords(Input, Intermediates);
+
+#if (ES2_PROFILE || ES3_1_PROFILE)
+	Interpolants.LayerTexCoord	= LandscapeTexCoords.LayerTexCoord;
+	Interpolants.WeightMapTexCoord	= LandscapeTexCoords.WeightMapTexCoord; 
+	#if NUM_MATERIAL_TEXCOORDS
+		float2 CustomizedUVs[NUM_MATERIAL_TEXCOORDS];
+		GetMaterialCustomizedUVs(VertexParameters, CustomizedUVs);
+
+		Interpolants.LayerTexCoord = CustomizedUVs[0];
+		#if NUM_MATERIAL_TEXCOORDS > 1
+			Interpolants.TransformedTexCoord0.xy	= CustomizedUVs[1];
+			#if NUM_MATERIAL_TEXCOORDS > 2
+				Interpolants.TransformedTexCoord1.xy	= CustomizedUVs[2];
+			#endif // 2
+			// ignore 4th CustomizedUV because of weightmap UV
+		#endif // 1
+	#endif //NUM_MATERIAL_TEXCOORDS
+#else
+	Interpolants.LayerTexCoord				= LandscapeTexCoords.LayerTexCoord;
+	Interpolants.WeightHeightMapTexCoord.xy	= LandscapeTexCoords.WeightMapTexCoord;
+	Interpolants.WeightHeightMapTexCoord.zw	= LandscapeTexCoords.HeightMapTexCoord;
+	#if NUM_MATERIAL_TEXCOORDS
+		float2 CustomizedUVs[NUM_MATERIAL_TEXCOORDS];
+		GetMaterialCustomizedUVs(VertexParameters, CustomizedUVs);
+
+		Interpolants.LayerTexCoord = CustomizedUVs[0];
+		Interpolants.TransformedTexCoords = 0;
+
+		#if NUM_MATERIAL_TEXCOORDS > 1
+			Interpolants.TransformedTexCoords.xy = CustomizedUVs[1];
+			#if NUM_MATERIAL_TEXCOORDS > 2
+				Interpolants.TransformedTexCoords.zw = CustomizedUVs[2];
+			#endif // 2
+		#endif // 1
+
+	#endif // NUM_MATERIAL_TEXCOORDS
+#endif // (ES2_PROFILE || ES3_1_PROFILE)
+
+#if NEEDS_LIGHTMAP_COORDINATE
+	#if (ES2_PROFILE || ES3_1_PROFILE)
+		Interpolants.LightMapCoordinate[0].xy = LandscapeTexCoords.LightMapCoordinate * PrecomputedLightingBuffer.LightMapCoordinateScaleBias.xy + PrecomputedLightingBuffer.LightMapCoordinateScaleBias.zw;
+		Interpolants.LightMapCoordinate[0].y *= 0.5;
+		Interpolants.LightMapCoordinate[1].xy = Interpolants.LightMapCoordinate[0].xy;
+		Interpolants.LightMapCoordinate[1].y += 0.5;
+		#if STATICLIGHTING_TEXTUREMASK
+			Interpolants.ShadowMapCoordinate.xy	= LandscapeTexCoords.LightMapCoordinate.xy * PrecomputedLightingBuffer.ShadowMapCoordinateScaleBias.xy + PrecomputedLightingBuffer.ShadowMapCoordinateScaleBias.zw;
+		#else
+			Interpolants.ShadowMapCoordinate.xy = 0;
+		#endif
+	#else
+		Interpolants.LightMapCoordinate.xy		= LandscapeTexCoords.LightMapCoordinate.xy * PrecomputedLightingBuffer.LightMapCoordinateScaleBias.xy + PrecomputedLightingBuffer.LightMapCoordinateScaleBias.zw;
+		#if STATICLIGHTING_TEXTUREMASK
+			Interpolants.LightMapCoordinate.zw	= LandscapeTexCoords.LightMapCoordinate.xy * PrecomputedLightingBuffer.ShadowMapCoordinateScaleBias.xy + PrecomputedLightingBuffer.ShadowMapCoordinateScaleBias.zw;
+		#else
+			Interpolants.LightMapCoordinate.zw = 0;
+		#endif
+	#endif
+#endif
+
+#if INSTANCED_STEREO
+	Interpolants.PackedEyeIndex = 0;
+#endif
+	
+	return Interpolants;
+}
+
+FVertexFactoryInterpolantsVSToDS VertexFactoryGetInterpolantsVSToDS(FVertexFactoryInput Input, FVertexFactoryIntermediates Intermediates, FMaterialVertexParameters VertexParameters)
+{
+	FVertexFactoryInterpolantsVSToDS Interpolants;
+	
+	Interpolants.InterpolantsVSToPS = VertexFactoryGetInterpolantsVSToPS(Input, Intermediates, VertexParameters);
+
+#if USING_TESSELLATION
+	// Calculate LocalToTangent directly from normal map texture.
+	float3x3 TangentToLocal = CalcTangentBasisFromWorldNormal(Intermediates.WorldNormal);
+	float3x3 TangentToWorld = mul(TangentToLocal, (float3x3)LandscapeParameters.LocalToWorldNoScaling);
+	
+	Interpolants.TangentToWorld0 = TangentToWorld[0];
+	Interpolants.TangentToWorld2 = float4(TangentToWorld[2], 1);
+
+	float2 xy = Intermediates.InputPosition.xy * LodValues.w;
+	Interpolants.CalculatedLOD = CalcLOD(xy, Intermediates.InputPosition.zw);
+#endif
+	
+	return Interpolants;
+}
+
+FVertexFactoryInterpolantsVSToDS VertexFactoryInterpolate(FVertexFactoryInterpolantsVSToDS a, float aInterp, FVertexFactoryInterpolantsVSToDS b, float bInterp)
+{
+	FVertexFactoryInterpolantsVSToDS O;
+	
+#if NEEDS_LIGHTMAP_COORDINATE && !(ES2_PROFILE || ES3_1_PROFILE)
+	TESSELLATION_INTERPOLATE_MEMBER(InterpolantsVSToPS.LightMapCoordinate);
+#endif
+
+	TESSELLATION_INTERPOLATE_MEMBER(InterpolantsVSToPS.LayerTexCoord);
+#if (ES2_PROFILE || ES3_1_PROFILE)
+	TESSELLATION_INTERPOLATE_MEMBER(InterpolantsVSToPS.WeightMapTexCoord);
+	TESSELLATION_INTERPOLATE_MEMBER(InterpolantsVSToPS.TransformedTexCoord0);
+	TESSELLATION_INTERPOLATE_MEMBER(InterpolantsVSToPS.TransformedTexCoord1);
+#else
+	TESSELLATION_INTERPOLATE_MEMBER(InterpolantsVSToPS.WeightHeightMapTexCoord);
+	TESSELLATION_INTERPOLATE_MEMBER(InterpolantsVSToPS.TransformedTexCoords);
+#endif
+
+#if USING_TESSELLATION
+	TESSELLATION_INTERPOLATE_MEMBER(TangentToWorld0);
+	TESSELLATION_INTERPOLATE_MEMBER(TangentToWorld2);
+	TESSELLATION_INTERPOLATE_MEMBER(CalculatedLOD);
+#endif
+
+	return O;
+}
+
+#if USING_TESSELLATION
+float3x3 VertexFactoryGetTangentToLocalDS(FVertexFactoryInterpolantsVSToDS Interpolants)
+{
+	return AssembleTangentToWorld(Interpolants.TangentToWorld0, Interpolants.TangentToWorld2);
+}
+#endif // #if USING_TESSELLATION
+
+#if NUM_VF_PACKED_INTERPOLANTS > 0
+void VertexFactoryPackInterpolants(inout FVertexFactoryInterpolantsVSToPS Interpolants, float4 PackedInterpolants[NUM_VF_PACKED_INTERPOLANTS])
+{
+	Interpolants.TransformedTexCoord0.zw = PackedInterpolants[0].xy;
+	Interpolants.TransformedTexCoord1.zw = PackedInterpolants[0].zw;
+}
+
+void VertexFactoryUnpackInterpolants(FVertexFactoryInterpolantsVSToPS Interpolants, out float4 PackedInterpolants[NUM_VF_PACKED_INTERPOLANTS])
+{
+	PackedInterpolants[0].xy = Interpolants.TransformedTexCoord0.zw;
+	PackedInterpolants[0].zw = Interpolants.TransformedTexCoord1.zw;
+#if NUM_VF_PACKED_INTERPOLANTS > 1
+	UNROLL
+	for (int i = 1; i < NUM_VF_PACKED_INTERPOLANTS; ++i)
+	{
+		PackedInterpolants[i] = 0;
+	}
+#endif
+}
+#endif // NUM_VF_PACKED_INTERPOLANTS > 0
+
+#if INSTANCED_STEREO
+uint VertexFactoryGetEyeIndex(uint InstanceId)
+{
+	return InstanceId & 1;
+}
+#endif
+
+float4 VertexFactoryGetTranslatedPrimitiveVolumeBounds(FVertexFactoryInterpolantsVSToPS Interpolants)
+{
+	return 0;
+}
\ No newline at end of file
diff --git a/Engine/Shaders/Private/LightGridCommon.ush b/Engine/Shaders/Private/LightGridCommon.ush
new file mode 100644
--- /dev/null
+++ b/Engine/Shaders/Private/LightGridCommon.ush
@@ -0,0 +1,122 @@
+/*=============================================================================
+	LightGridCommon.usf
+=============================================================================*/
+
+Buffer<uint> NumCulledLightsGrid;
+Buffer<uint> CulledLightDataGrid;
+Buffer<float4> ForwardLocalLightBuffer;
+
+// ISR copies
+Buffer<uint> InstancedNumCulledLightsGrid;
+Buffer<uint> InstancedCulledLightDataGrid;
+Buffer<float4> InstancedForwardLocalLightBuffer;
+
+struct LightGridData
+{
+	uint LightGridPixelSizeShift;
+	float3 LightGridZParams;
+	int3 CulledGridSize;
+};
+
+LightGridData GetLightGridData(uint EyeIndex)
+{
+	LightGridData Result;
+
+	BRANCH
+	if (EyeIndex == 0)
+	{
+		Result.LightGridPixelSizeShift = ForwardGlobalLightData.LightGridPixelSizeShift;
+		Result.LightGridZParams = ForwardGlobalLightData.LightGridZParams;
+		Result.CulledGridSize = ForwardGlobalLightData.CulledGridSize;
+	}
+	else
+	{
+		Result.LightGridPixelSizeShift = InstancedForwardGlobalLightData.LightGridPixelSizeShift;
+		Result.LightGridZParams = InstancedForwardGlobalLightData.LightGridZParams;
+		Result.CulledGridSize = InstancedForwardGlobalLightData.CulledGridSize;
+	}
+
+	return Result;
+}
+
+uint ComputeLightGridCellIndex(uint2 PixelPos, float SceneDepth, uint EyeIndex)
+{
+	const LightGridData GridData = GetLightGridData(EyeIndex);
+	uint ZSlice = (uint)(max(0, log2(SceneDepth * GridData.LightGridZParams.x + GridData.LightGridZParams.y) * GridData.LightGridZParams.z));
+	ZSlice = min(ZSlice, (uint)(GridData.CulledGridSize.z - 1));
+	uint3 GridCoordinate = uint3(PixelPos >> GridData.LightGridPixelSizeShift, ZSlice);
+	uint GridIndex = (GridCoordinate.z * GridData.CulledGridSize.y + GridCoordinate.y) * GridData.CulledGridSize.x + GridCoordinate.x;
+	return GridIndex;
+}
+
+uint ComputeLightGridCellIndex(uint2 PixelPos, float SceneDepth)
+{
+	return ComputeLightGridCellIndex(PixelPos, SceneDepth, 0);
+}
+
+#ifndef NUM_CULLED_LIGHTS_GRID_STRIDE
+#define NUM_CULLED_LIGHTS_GRID_STRIDE 0
+#endif
+
+#ifndef LOCAL_LIGHT_DATA_STRIDE
+#define LOCAL_LIGHT_DATA_STRIDE 0
+#endif
+
+struct FCulledLightsGridData
+{
+	uint NumLocalLights;
+	uint DataStartIndex;
+};
+
+FCulledLightsGridData GetCulledLightsGrid(uint GridIndex, uint EyeIndex)
+{
+	FCulledLightsGridData Result;
+
+	BRANCH
+	if (EyeIndex == 0)
+	{
+		Result.NumLocalLights = NumCulledLightsGrid[GridIndex * NUM_CULLED_LIGHTS_GRID_STRIDE + 0];
+		Result.DataStartIndex = NumCulledLightsGrid[GridIndex * NUM_CULLED_LIGHTS_GRID_STRIDE + 1];
+	}
+	else
+	{
+		Result.NumLocalLights = InstancedNumCulledLightsGrid[GridIndex * NUM_CULLED_LIGHTS_GRID_STRIDE + 0];
+		Result.DataStartIndex = InstancedNumCulledLightsGrid[GridIndex * NUM_CULLED_LIGHTS_GRID_STRIDE + 1];
+	}
+
+	return Result;
+}
+
+struct FLocalLightData
+{
+	float4 LightPositionAndInvRadius;
+	float4 LightColorAndFalloffExponent;
+	float4 SpotAnglesAndSourceRadiusPacked;
+	float4 LightDirectionAndShadowMask;
+};
+
+FLocalLightData GetLocalLightData(uint GridIndex, uint EyeIndex)
+{
+	FLocalLightData Result;
+
+	if (EyeIndex == 0)
+	{
+		uint LocalLightIndex = CulledLightDataGrid[GridIndex];
+		uint LocalLightBaseIndex = LocalLightIndex * LOCAL_LIGHT_DATA_STRIDE;
+		Result.LightPositionAndInvRadius = ForwardLocalLightBuffer[LocalLightBaseIndex + 0];
+		Result.LightColorAndFalloffExponent = ForwardLocalLightBuffer[LocalLightBaseIndex + 1];
+		Result.LightDirectionAndShadowMask = ForwardLocalLightBuffer[LocalLightBaseIndex + 2];
+		Result.SpotAnglesAndSourceRadiusPacked = ForwardLocalLightBuffer[LocalLightBaseIndex + 3];
+	}
+	else
+	{
+		uint LocalLightIndex = InstancedCulledLightDataGrid[GridIndex];
+		uint LocalLightBaseIndex = LocalLightIndex * LOCAL_LIGHT_DATA_STRIDE;
+		Result.LightPositionAndInvRadius = InstancedForwardLocalLightBuffer[LocalLightBaseIndex + 0];
+		Result.LightColorAndFalloffExponent = InstancedForwardLocalLightBuffer[LocalLightBaseIndex + 1];
+		Result.LightDirectionAndShadowMask = InstancedForwardLocalLightBuffer[LocalLightBaseIndex + 2];
+		Result.SpotAnglesAndSourceRadiusPacked = InstancedForwardLocalLightBuffer[LocalLightBaseIndex + 3];
+	}
+
+	return Result;
+}
\ No newline at end of file
diff --git a/Engine/Shaders/Private/LightmapCommon.ush b/Engine/Shaders/Private/LightmapCommon.ush
new file mode 100644
--- /dev/null
+++ b/Engine/Shaders/Private/LightmapCommon.ush
@@ -0,0 +1,179 @@
+// Copyright 1998-2017 Epic Games, Inc. All Rights Reserved.
+
+#include "/Engine/Generated/UniformBuffers/PrecomputedLightingBuffer.ush"
+
+#if HQ_TEXTURE_LIGHTMAP || LQ_TEXTURE_LIGHTMAP
+ 
+#ifndef MATERIAL_USE_LM_DIRECTIONALITY
+#define MATERIAL_USE_LM_DIRECTIONALITY	1
+#endif
+
+// Material quality overrides
+#ifndef QL_FORCEDISABLE_LM_DIRECTIONALITY
+#define QL_FORCEDISABLE_LM_DIRECTIONALITY 0
+#endif
+
+#define USE_LM_DIRECTIONALITY (MATERIAL_USE_LM_DIRECTIONALITY && !QL_FORCEDISABLE_LM_DIRECTIONALITY)
+
+half4 GetLightMapColorLQ( float2 LightmapUV0, float2 LightmapUV1, half3 WorldNormal )
+{
+	half4 Lightmap0 = Texture2DSample( PrecomputedLightingBuffer.LightMapTexture, PrecomputedLightingBuffer.LightMapSampler, LightmapUV0 );
+	half4 Lightmap1 = Texture2DSample( PrecomputedLightingBuffer.LightMapTexture, PrecomputedLightingBuffer.LightMapSampler, LightmapUV1 );
+		
+	// Range scale
+	half3 LogRGB = Lightmap0.rgb * PrecomputedLightingBuffer.LightMapScale[0].xyz + PrecomputedLightingBuffer.LightMapAdd[0].xyz;	// 1 vmad
+
+	half LogL = Luminance( LogRGB );					// 1 dot
+
+	// LogL -> L
+	const half LogBlackPoint = 0.00390625;	// exp2(-8);
+	half L = exp2( LogL * 16 - 8 ) - LogBlackPoint;		// 1 exp2, 1 smad, 1 ssub
+
+#if USE_LM_DIRECTIONALITY
+	// Alpha doesn't matter, will scaled by zero
+	float4 SH = Lightmap1 * PrecomputedLightingBuffer.LightMapScale[1] + PrecomputedLightingBuffer.LightMapAdd[1];	// 1 vmad
+
+	// Sample SH with normal
+	half Directionality = max( 0.0, dot( SH, float4(WorldNormal.yzx, 1) ) );	// 1 dot, 1 smax
+#else
+	half Directionality = 0.6;
+#endif
+		
+	half Luma = L * Directionality;
+	half3 Color = LogRGB * (Luma / LogL);				// 1 rcp, 1 smul, 1 vmul
+
+	return half4( Color, Luma );
+}
+
+void GetLightMapColorHQ( float2 LightmapUV0, float2 LightmapUV1, half3 WorldNormal, out half3 OutDiffuseLighting, out half3 OutSubsurfaceLighting )
+{
+	OutSubsurfaceLighting = 0;
+
+	half4 Lightmap0 = Texture2DSample( PrecomputedLightingBuffer.LightMapTexture, PrecomputedLightingBuffer.LightMapSampler, LightmapUV0 );
+	half4 Lightmap1 = Texture2DSample( PrecomputedLightingBuffer.LightMapTexture, PrecomputedLightingBuffer.LightMapSampler, LightmapUV1 );
+
+	half LogL = Lightmap0.w;
+
+	// Add residual
+	LogL += Lightmap1.w * (1.0 / 255) - (0.5 / 255);
+
+	// Range scale LogL
+	LogL = LogL * PrecomputedLightingBuffer.LightMapScale[0].w + PrecomputedLightingBuffer.LightMapAdd[0].w;
+		
+	// Range scale UVW
+	half3 UVW = Lightmap0.rgb * Lightmap0.rgb * PrecomputedLightingBuffer.LightMapScale[0].rgb + PrecomputedLightingBuffer.LightMapAdd[0].rgb;
+
+	// LogL -> L
+	const half LogBlackPoint = 0.01858136;
+	half L = exp2( LogL ) - LogBlackPoint;
+
+#if USE_LM_DIRECTIONALITY
+	// Range scale SH. Alpha doesn't matter, will scale with zero
+	float4 SH = Lightmap1 * PrecomputedLightingBuffer.LightMapScale[1] + PrecomputedLightingBuffer.LightMapAdd[1];
+
+	// Sample SH with normal
+	half Directionality = max( 0.0, dot( SH, float4(WorldNormal.yzx, 1) ) );
+
+	#if MATERIAL_SHADINGMODEL_TWOSIDED_FOLIAGE
+		half SubsurfaceDirectionality = max( 0.0, dot( SH, float4(-WorldNormal.yzx, 1) ) );
+		OutSubsurfaceLighting = L * SubsurfaceDirectionality * UVW;
+	#endif
+#else
+	half Directionality = 0.6;
+#endif
+
+	half Luma = L * Directionality;
+	half3 Color = Luma * UVW;
+
+	OutDiffuseLighting = Color;
+}
+
+float4 GetSkyBentNormalAndOcclusion(float2 LightmapUV)
+{
+#if SUPPORTS_INDEPENDENT_SAMPLERS
+	// Share sampler with lightmaps to give artists more available
+	float4 TextureValue = Texture2DSample(PrecomputedLightingBuffer.SkyOcclusionTexture, PrecomputedLightingBuffer.LightMapSampler, LightmapUV.xy);
+#else
+	float4 TextureValue = Texture2DSample(PrecomputedLightingBuffer.SkyOcclusionTexture, PrecomputedLightingBuffer.SkyOcclusionSampler, LightmapUV.xy);
+#endif
+	// Unpack vector
+	TextureValue.rgb = TextureValue.rgb * 2 - 1;
+	// Undo sqrt which allocated more precision toward 0
+	TextureValue.a = TextureValue.a * TextureValue.a;
+	return TextureValue;
+}
+
+float GetAOMaterialMask(float2 LightmapUV)
+{
+#if SUPPORTS_INDEPENDENT_SAMPLERS
+	// Share sampler with lightmaps to give artists more available
+	float TextureValue = Texture2DSample(PrecomputedLightingBuffer.AOMaterialMaskTexture, PrecomputedLightingBuffer.LightMapSampler, LightmapUV.xy).x;
+#else
+	float TextureValue = Texture2DSample(PrecomputedLightingBuffer.AOMaterialMaskTexture, PrecomputedLightingBuffer.AOMaterialMaskSampler, LightmapUV.xy).x;
+#endif
+	// Undo sqrt which allocated more precision toward 0
+	return TextureValue * TextureValue;
+}
+
+#endif
+
+// Used by deferred renderer only
+half4 GetPrecomputedShadowMasks(FVertexFactoryInterpolantsVSToPS Interpolants)
+{
+	// Note: WRITES_PRECSHADOWFACTOR_ZERO have to match the logic here
+	#if STATICLIGHTING_TEXTUREMASK && STATICLIGHTING_SIGNEDDISTANCEFIELD
+
+		// Fetch the 4 channels of distance field data
+		half4 DistanceField = Texture2DSample(PrecomputedLightingBuffer.StaticShadowTexture, PrecomputedLightingBuffer.StaticShadowTextureSampler, GetShadowMapCoordinate(Interpolants));
+		float4 DistanceFieldBias = -.5f * PrecomputedLightingBuffer.InvUniformPenumbraSizes + .5f;
+
+		// Compute shadow factors by scaling and biasing the distance
+		half4 ShadowFactors = saturate(DistanceField * PrecomputedLightingBuffer.InvUniformPenumbraSizes + DistanceFieldBias);
+		return PrecomputedLightingBuffer.StaticShadowMapMasks * ShadowFactors * ShadowFactors;
+
+	#elif HQ_TEXTURE_LIGHTMAP || LQ_TEXTURE_LIGHTMAP
+
+		// Mark as shadowed for lightmapped objects with no shadowmap
+		// This is necessary because objects inside a light's influence that were determined to be completely shadowed won't be rendered with STATICLIGHTING_TEXTUREMASK==1
+		return 0;
+
+	#else
+
+		// Mark as unshadowed by default
+		// Directional light is always packed into the first static shadowmap channel, so output the per-primitive directional light shadowing there if requested
+		return half4(Primitive.UseSingleSampleShadowFromStationaryLights > 0 ? PrecomputedLightingBuffer.DirectionalLightShadowing : 1, 1, 1, 1);
+
+	#endif
+}
+
+// Used by mobile renderer only
+half GetPrimaryPrecomputedShadowMask(FVertexFactoryInterpolantsVSToPS Interpolants)
+{
+	#if STATICLIGHTING_TEXTUREMASK && STATICLIGHTING_SIGNEDDISTANCEFIELD
+
+		// Fetch the distance field data
+		half DistanceField = Texture2DSample(PrecomputedLightingBuffer.StaticShadowTexture, PrecomputedLightingBuffer.StaticShadowTextureSampler, GetShadowMapCoordinate(Interpolants)).r;
+		float DistanceFieldBias = -.5f * PrecomputedLightingBuffer.InvUniformPenumbraSizes.x + .5f;
+		// Compute shadow factors by scaling and biasing the distance
+		half ShadowFactor = saturate( DistanceField * PrecomputedLightingBuffer.InvUniformPenumbraSizes.x + DistanceFieldBias );
+		return PrecomputedLightingBuffer.StaticShadowMapMasks.r * ShadowFactor * ShadowFactor;
+
+	#elif MOVABLE_DIRECTIONAL_LIGHT
+		// Do this before checking for lightmaps as we might have a lightmap + directional light
+		//@mw todo
+		// TODO: Filter shadowmap
+		return 1;
+
+	#elif HQ_TEXTURE_LIGHTMAP || LQ_TEXTURE_LIGHTMAP
+
+		// Mark as shadowed for lightmapped objects with no shadowmap
+		// This is necessary because objects inside a light's influence that were determined to be completely shadowed won't be rendered with STATICLIGHTING_TEXTUREMASK==1
+		return 0;
+
+	#else
+
+		// output per-primitive directional light shadowing if requested
+		return Primitive.UseSingleSampleShadowFromStationaryLights > 0 ? PrecomputedLightingBuffer.DirectionalLightShadowing : 1;
+
+	#endif
+}
\ No newline at end of file
diff --git a/Engine/Shaders/Private/LocalVertexFactory.ush b/Engine/Shaders/Private/LocalVertexFactory.ush
new file mode 100644
--- /dev/null
+++ b/Engine/Shaders/Private/LocalVertexFactory.ush
@@ -0,0 +1,862 @@
+// Copyright 1998-2017 Epic Games, Inc. All Rights Reserved.
+
+/*=============================================================================
+	LocalVertexFactory.usf: Local vertex factory shader code.
+=============================================================================*/
+
+#include "VertexFactoryCommon.ush"
+#include "LocalVertexFactoryCommon.ush"
+
+#if GPUSKIN_PASS_THROUGH
+#include "GpuSkinCommon.ush"
+#endif
+
+float4x4 PreviousLocalToWorld;
+
+#include "/Engine/Generated/UniformBuffers/PrecomputedLightingBuffer.ush"
+
+#if USE_INSTANCING
+
+	// Instanced stereo sets the eye index explicitly for instanced geometry
+	#if INSTANCED_STEREO
+		uint InstancedEyeIndex;
+	#endif
+
+	#ifndef USE_DITHERED_LOD_TRANSITION_FOR_INSTANCED
+		#error "USE_DITHERED_LOD_TRANSITION_FOR_INSTANCED should have been defined"
+	#endif
+	#define USE_DITHERED_LOD_TRANSITION USE_DITHERED_LOD_TRANSITION_FOR_INSTANCED
+#else
+	#ifndef USE_DITHERED_LOD_TRANSITION_FROM_MATERIAL
+		#error "USE_DITHERED_LOD_TRANSITION_FROM_MATERIAL should have been defined"
+	#endif
+	#define USE_DITHERED_LOD_TRANSITION USE_DITHERED_LOD_TRANSITION_FROM_MATERIAL
+#endif
+
+#if USE_INSTANCING
+
+#if USE_DITHERED_LOD_TRANSITION
+	float4 InstancingViewZCompareZero;  // w contains random lod scale 
+	float4 InstancingViewZCompareOne;
+	float4 InstancingViewZConstant;
+	float4 InstancingWorldViewOriginZero;
+	float4 InstancingWorldViewOriginOne;
+#endif
+
+float4 InstancingFadeOutParams;
+
+#ifndef USE_INSTANCING_EMULATED
+	#define USE_INSTANCING_EMULATED 0
+#endif	// USE_INSTANCING_EMULATED
+
+#if USE_INSTANCING_EMULATED
+	// Required for CPU emulation of Instancing
+	float4 CPUInstanceOrigin;  // per-instance random in w 
+	float3x4 CPUInstanceTransform; // hitproxy.r + 256 * selected in .w; hitproxy.g in .w; hitproxy.b in .w
+	float4 CPUInstanceLightmapAndShadowMapBias; 
+#endif	// USE_INSTANCING_EMULATED
+
+#else // USE_INSTANCING
+	#ifndef USE_INSTANCING_EMULATED
+		#define USE_INSTANCING_EMULATED (0)
+	#endif
+#endif	// USE_INSTANCING
+
+#if USE_SPLINEDEFORM
+float3 SplineStartPos;
+float3 SplineStartTangent;
+float SplineStartRoll;
+float2 SplineStartScale;
+float2 SplineStartOffset;
+
+float3 SplineEndPos;
+float3 SplineEndTangent;
+float SplineEndRoll;
+float2 SplineEndScale;
+float2 SplineEndOffset;
+
+float3 SplineUpDir;
+#if ES2_PROFILE
+float SmoothInterpRollScale;
+#else
+bool SmoothInterpRollScale;
+#endif
+
+float SplineMeshMinZ;
+float SplineMeshScaleZ;
+
+float3 SplineMeshDir;
+float3 SplineMeshX;
+float3 SplineMeshY;
+#endif	// USE_SPLINEDEFORM
+
+struct FVertexFactoryInput
+{
+	float4	Position	: ATTRIBUTE0;
+#if METAL_PROFILE
+	//@todo-rco: FIXME!
+	float3	TangentX	: ATTRIBUTE1;
+	float4	TangentZ	: ATTRIBUTE2;
+	float4	Color		: ATTRIBUTE3;
+#else
+	half3	TangentX	: ATTRIBUTE1;
+	// TangentZ.w contains sign of tangent basis determinant
+	half4	TangentZ	: ATTRIBUTE2;
+	half4	Color		: ATTRIBUTE3;
+#endif	// METAL_PROFILE
+
+#if NUM_MATERIAL_TEXCOORDS_VERTEX
+	#if GPUSKIN_PASS_THROUGH
+		// These must match GPUSkinVertexFactory.usf
+		float2	TexCoords[NUM_MATERIAL_TEXCOORDS_VERTEX] : ATTRIBUTE4;
+		#if NUM_MATERIAL_TEXCOORDS_VERTEX > 4
+			#error Too many texture coordinate sets defined on GPUSkin vertex input. Max: 4.
+		#endif
+	#else
+		#if NUM_MATERIAL_TEXCOORDS_VERTEX > 1
+			float4	PackedTexCoords4[NUM_MATERIAL_TEXCOORDS_VERTEX/2] : ATTRIBUTE4;
+		#endif
+		#if NUM_MATERIAL_TEXCOORDS_VERTEX == 1
+			float2	PackedTexCoords2 : ATTRIBUTE4;
+		#elif NUM_MATERIAL_TEXCOORDS_VERTEX == 3
+			float2	PackedTexCoords2 : ATTRIBUTE5;
+		#elif NUM_MATERIAL_TEXCOORDS_VERTEX == 5
+			float2	PackedTexCoords2 : ATTRIBUTE6;
+		#elif NUM_MATERIAL_TEXCOORDS_VERTEX == 7
+			float2	PackedTexCoords2 : ATTRIBUTE7;
+		#endif
+	#endif
+
+#elif USE_PARTICLE_SUBUVS
+	float2	TexCoords[1] : ATTRIBUTE4;
+#endif
+
+#if USE_INSTANCING && !USE_INSTANCING_EMULATED
+	float4 InstanceOrigin : ATTRIBUTE8;  // per-instance random in w 
+	half4 InstanceTransform1 : ATTRIBUTE9;  // hitproxy.r + 256 * selected in .w
+	half4 InstanceTransform2 : ATTRIBUTE10; // hitproxy.g in .w
+	half4 InstanceTransform3 : ATTRIBUTE11; // hitproxy.b in .w
+	float4 InstanceLightmapAndShadowMapUVBias : ATTRIBUTE12; 
+#else
+
+#endif //USE_INSTANCING && !USE_INSTANCING_EMULATED
+
+#if NEEDS_LIGHTMAP_COORDINATE
+	float2	LightMapCoordinate : ATTRIBUTE15;
+#endif
+
+/** Optional instance ID for vertex layered rendering */
+#if FEATURE_LEVEL >= FEATURE_LEVEL_SM4 && ONEPASS_POINTLIGHT_SHADOW && USING_VERTEX_SHADER_LAYER
+	uint InstanceId	: SV_InstanceID;
+#endif
+
+#if GPUSKIN_PASS_THROUGH
+	uint VertexId : SV_VertexID;
+#endif
+};
+
+struct FPositionOnlyVertexFactoryInput
+{
+	float4	Position	: ATTRIBUTE0;
+	
+#if USE_INSTANCING && !USE_INSTANCING_EMULATED
+	float4 InstanceOrigin : ATTRIBUTE8;  // per-instance random in w 
+	half4 InstanceTransform1 : ATTRIBUTE9;  // hitproxy.r + 256 * selected in .w
+	half4 InstanceTransform2 : ATTRIBUTE10; // hitproxy.g in .w
+	half4 InstanceTransform3 : ATTRIBUTE11; // hitproxy.b in .w
+#endif	// USE_INSTANCING && !USE_INSTANCING_EMULATED
+
+/** Optional instance ID for vertex layered rendering */
+#if FEATURE_LEVEL >= FEATURE_LEVEL_SM4 && ONEPASS_POINTLIGHT_SHADOW && USING_VERTEX_SHADER_LAYER
+	uint InstanceId	: SV_InstanceID;
+#endif
+};
+
+struct FVertexFactoryIntermediates
+{
+	half3x3 TangentToLocal;
+	half3x3 TangentToWorld;
+	half TangentToWorldSign;
+
+	half4 Color;
+#if USE_INSTANCING
+	// x = per-instance random, y = per-instance fade out amount, z = hide/show flag, w dither fade cutoff
+	float4 PerInstanceParams;
+#endif	// USE_INSTANCING
+};
+
+#if USE_INSTANCING
+float4x4 GetInstanceTransform(FVertexFactoryInput Input)
+{
+	#if !USE_INSTANCING_EMULATED
+	return float4x4(
+		float4(Input.InstanceTransform1.xyz, 0.0f), 
+		float4(Input.InstanceTransform2.xyz, 0.0f), 
+		float4(Input.InstanceTransform3.xyz, 0.0f), 
+		float4(Input.InstanceOrigin.xyz, 1.0f));
+	#else
+	return float4x4(
+		float4(CPUInstanceTransform[0].xyz, 0.0f), 
+		float4(CPUInstanceTransform[1].xyz, 0.0f), 
+		float4(CPUInstanceTransform[2].xyz, 0.0f), 
+		float4(CPUInstanceOrigin.xyz, 1.0f));
+	#endif	// !USE_INSTANCING_EMULATED
+}
+
+float4x4 GetInstanceTransform(FPositionOnlyVertexFactoryInput Input)
+{
+	#if !USE_INSTANCING_EMULATED
+	return float4x4(
+		float4(Input.InstanceTransform1.xyz, 0.0f), 
+		float4(Input.InstanceTransform2.xyz, 0.0f), 
+		float4(Input.InstanceTransform3.xyz, 0.0f), 
+		float4(Input.InstanceOrigin.xyz, 1.0f));
+	#else
+	return float4x4(
+		float4(CPUInstanceTransform[0].xyz, 0.0f), 
+		float4(CPUInstanceTransform[1].xyz, 0.0f), 
+		float4(CPUInstanceTransform[2].xyz, 0.0f), 
+		float4(CPUInstanceOrigin.xyz, 1.0f));
+	#endif	// !USE_INSTANCING_EMULATED
+}
+
+half3x3 GetInstanceToLocal3x3(FVertexFactoryInput Input)
+{
+	#if !USE_INSTANCING_EMULATED
+	return half3x3(Input.InstanceTransform1.xyz, Input.InstanceTransform2.xyz, Input.InstanceTransform3.xyz);
+	#else
+	return half3x3(half3(CPUInstanceTransform[0].xyz), half3(CPUInstanceTransform[1].xyz), half3(CPUInstanceTransform[2].xyz));
+	#endif	// !USE_INSTANCING_EMULATED
+}
+
+half3x3 GetInstanceToLocal3x3(FPositionOnlyVertexFactoryInput Input)
+{
+	#if !USE_INSTANCING_EMULATED
+	return half3x3(Input.InstanceTransform1.xyz, Input.InstanceTransform2.xyz, Input.InstanceTransform3.xyz);
+	#else
+	return half3x3(half3(CPUInstanceTransform[0].xyz), half3(CPUInstanceTransform[1].xyz), half3(CPUInstanceTransform[2].xyz));
+	#endif	// !USE_INSTANCING_EMULATED
+}
+
+float2 GetInstanceShadowMapBias(FVertexFactoryInput Input)
+{
+	#if !USE_INSTANCING_EMULATED
+	return Input.InstanceLightmapAndShadowMapUVBias.zw;
+	#else
+	return CPUInstanceLightmapAndShadowMapBias.zw;
+	#endif	// !USE_INSTANCING_EMULATED
+}
+
+float2 GetInstanceLightMapBias(FVertexFactoryInput Input)
+{
+	#if !USE_INSTANCING_EMULATED
+	return Input.InstanceLightmapAndShadowMapUVBias.xy;
+	#else
+	return CPUInstanceLightmapAndShadowMapBias.xy;
+	#endif	// !USE_INSTANCING_EMULATED
+}
+
+float GetInstanceSelected(FVertexFactoryInput Input)
+{
+	#if !USE_INSTANCING_EMULATED
+		float SelectedValue = trunc(Input.InstanceTransform1.w * (1.0 / 256.0));
+	#else
+		float SelectedValue = trunc(CPUInstanceTransform[0].w * (1.0 / 256.0));
+	#endif	// !USE_INSTANCING_EMULATED
+	return SelectedValue;
+}
+
+float GetInstanceRandom(FVertexFactoryInput Input)
+{
+	#if !USE_INSTANCING_EMULATED
+		float RandomVal = Input.InstanceOrigin.w;
+	#else
+		float RandomVal = CPUInstanceOrigin.w;
+	#endif	// !USE_INSTANCING_EMULATED
+	return RandomVal;
+}
+
+float3 GetInstanceOrigin(FVertexFactoryInput Input)
+{
+	#if !USE_INSTANCING_EMULATED
+		float3 Origin = Input.InstanceOrigin.xyz;
+	#else
+		float3 Origin = CPUInstanceOrigin.xyz;
+	#endif	// !USE_INSTANCING_EMULATED
+	return Origin;
+}
+
+#endif	// USE_INSTANCING
+
+/** Converts from vertex factory specific interpolants to a FMaterialPixelParameters, which is used by material inputs. */
+FMaterialPixelParameters GetMaterialPixelParameters(FVertexFactoryInterpolantsVSToPS Interpolants, float4 SvPosition)
+{
+	// GetMaterialPixelParameters is responsible for fully initializing the result
+	FMaterialPixelParameters Result = MakeInitializedMaterialPixelParameters();
+
+#if NUM_TEX_COORD_INTERPOLATORS
+	UNROLL
+	for( int CoordinateIndex = 0; CoordinateIndex < NUM_TEX_COORD_INTERPOLATORS; CoordinateIndex++ )
+	{
+		Result.TexCoords[CoordinateIndex] = GetUV(Interpolants, CoordinateIndex);
+	}
+#endif
+
+#if USE_PARTICLE_SUBUVS
+	// Output TexCoord0 for when previewing materials that use ParticleSubUV.
+	Result.Particle.SubUVCoords[0] = GetUV(Interpolants, 0);
+	Result.Particle.SubUVCoords[1] = GetUV(Interpolants, 0);
+#endif	// USE_PARTICLE_SUBUVS
+
+	half3 TangentToWorld0 = GetTangentToWorld0(Interpolants).xyz;
+	half4 TangentToWorld2 = GetTangentToWorld2(Interpolants);
+	Result.UnMirrored = TangentToWorld2.w;
+
+	Result.VertexColor = GetColor(Interpolants);
+
+	// Required for previewing materials that use ParticleColor
+	Result.Particle.Color = half4(1,1,1,1);
+#if USE_INSTANCING
+	Result.PerInstanceParams = Interpolants.PerInstanceParams;
+#endif
+
+	Result.TangentToWorld = AssembleTangentToWorld( TangentToWorld0, TangentToWorld2 );
+
+#if LIGHTMAP_UV_ACCESS
+#if NEEDS_LIGHTMAP_COORDINATE
+	#if (ES2_PROFILE || ES3_1_PROFILE)
+		// Not supported in pixel shader
+		Result.LightmapUVs = float2(0, 0);
+	#else
+		Result.LightmapUVs = Interpolants.LightMapCoordinate.xy;
+	#endif	// ES2_PROFILE
+#endif	// NEEDS_LIGHTMAP_COORDINATE
+#endif	// LIGHTMAP_UV_ACCESS
+
+	Result.TwoSidedSign = 1;
+	return Result;
+}
+
+half3x3 CalcTangentToWorldNoScale(in half3x3 TangentToLocal)
+{
+	half3x3 LocalToWorld = GetLocalToWorld3x3();
+	half3 InvScale = Primitive.InvNonUniformScale.xyz;
+	LocalToWorld[0] *= InvScale.x;
+	LocalToWorld[1] *= InvScale.y;
+	LocalToWorld[2] *= InvScale.z;
+	return mul(TangentToLocal, LocalToWorld); 
+}
+
+/** Converts from vertex factory specific input to a FMaterialVertexParameters, which is used by vertex shader material inputs. */
+FMaterialVertexParameters GetMaterialVertexParameters(FVertexFactoryInput Input, FVertexFactoryIntermediates Intermediates, float3 WorldPosition, half3x3 TangentToLocal)
+{
+	FMaterialVertexParameters Result = (FMaterialVertexParameters)0;
+	Result.WorldPosition = WorldPosition;
+	Result.VertexColor = Intermediates.Color;
+
+	// does not handle instancing!
+	Result.TangentToWorld = Intermediates.TangentToWorld;
+
+#if USE_INSTANCING
+	Result.InstanceLocalToWorld = mul(GetInstanceTransform(Input), Primitive.LocalToWorld);
+	Result.InstanceLocalPosition = Input.Position.xyz;
+	Result.PerInstanceParams = Intermediates.PerInstanceParams;
+#endif	// USE_INSTANCING
+
+	Result.PreSkinnedPosition = Input.Position.xyz;
+	Result.PreSkinnedNormal = Input.TangentZ.xyz * 2.f - 1.f;
+
+#if NUM_MATERIAL_TEXCOORDS_VERTEX
+	#if GPUSKIN_PASS_THROUGH
+		for (int CoordinateIndex = 0; CoordinateIndex < NUM_MATERIAL_TEXCOORDS_VERTEX; CoordinateIndex++)
+		{
+			Result.TexCoords[CoordinateIndex] = Input.TexCoords[CoordinateIndex].xy;
+		}
+	#else
+		#if NUM_MATERIAL_TEXCOORDS_VERTEX > 1
+			UNROLL
+			for(int CoordinateIndex = 0; CoordinateIndex < NUM_MATERIAL_TEXCOORDS_VERTEX-1; CoordinateIndex+=2)
+			{
+				Result.TexCoords[CoordinateIndex] = Input.PackedTexCoords4[CoordinateIndex/2].xy;
+				if( CoordinateIndex+1 < NUM_MATERIAL_TEXCOORDS_VERTEX )
+				{
+					Result.TexCoords[CoordinateIndex+1] = Input.PackedTexCoords4[CoordinateIndex/2].zw;
+				}
+			}
+		#endif	// NUM_MATERIAL_TEXCOORDS_VERTEX > 1
+		#if NUM_MATERIAL_TEXCOORDS_VERTEX % 2 == 1
+			Result.TexCoords[NUM_MATERIAL_TEXCOORDS_VERTEX-1] = Input.PackedTexCoords2;
+		#endif	// NUM_MATERIAL_TEXCOORDS_VERTEX % 2 == 1
+	#endif
+#endif	// NUM_MATERIAL_TEXCOORDS_VERTEX
+	return Result;
+}
+
+#if USE_SPLINEDEFORM
+	float3 SplineEvalPos(float3 StartPos, float3 StartTangent, float3 EndPos, float3 EndTangent, float A)
+	{
+		float A2 = A  * A;
+		float A3 = A2 * A;
+
+		return (((2*A3)-(3*A2)+1) * StartPos) + ((A3-(2*A2)+A) * StartTangent) + ((A3-A2) * EndTangent) + (((-2*A3)+(3*A2)) * EndPos);
+	}
+
+	float3 SplineEvalDir(float3 StartPos, float3 StartTangent, float3 EndPos, float3 EndTangent, float A)
+	{
+		float3 C = (6*StartPos) + (3*StartTangent) + (3*EndTangent) - (6*EndPos);
+		float3 D = (-6*StartPos) - (4*StartTangent) - (2*EndTangent) + (6*EndPos);
+		float3 E = StartTangent;
+
+		float A2 = A  * A;
+
+		return normalize((C * A2) + (D * A) + E);
+	}
+
+	/** Calculate full transform that defines frame along spline, given the Z of a vertex. */
+	float4x3 CalcSliceTransform(float ZPos)
+	{
+		// Find how far 'along' mesh we are
+		float Alpha = ZPos * SplineMeshScaleZ - SplineMeshMinZ;
+
+		// Apply hermite interp to Alpha if desired
+		float HermiteAlpha = SmoothInterpRollScale ? smoothstep(0.0, 1.0, Alpha) : Alpha;
+
+		// Then find the point and direction of the spline at this point along
+		float3 SplinePos = SplineEvalPos( SplineStartPos, SplineStartTangent, SplineEndPos, SplineEndTangent, Alpha );
+		float3 SplineDir = SplineEvalDir( SplineStartPos, SplineStartTangent, SplineEndPos, SplineEndTangent, Alpha );
+	
+		// Find base frenet frame
+		float3 BaseXVec = normalize( cross(SplineUpDir, SplineDir) );
+		float3 BaseYVec = normalize( cross(SplineDir, BaseXVec) );
+	
+		// Offset from the spline, using the frenet frame
+		float2 SliceOffset = lerp(SplineStartOffset, SplineEndOffset, HermiteAlpha);
+		SplinePos += SliceOffset.x * BaseXVec;
+		SplinePos += SliceOffset.y * BaseYVec;
+	
+		// Apply roll to frame around spline
+		float UseRoll = lerp(SplineStartRoll, SplineEndRoll, HermiteAlpha);
+		float SinAng, CosAng;
+		sincos(UseRoll, SinAng, CosAng);
+		float3 XVec = (CosAng * BaseXVec) - (SinAng * BaseYVec);
+		float3 YVec = (CosAng * BaseYVec) + (SinAng * BaseXVec);
+
+		// Find scale at this point along spline
+		float2 UseScale = lerp(SplineStartScale, SplineEndScale, HermiteAlpha);
+
+		XVec *= UseScale.x;
+		YVec *= UseScale.y;
+
+		// Build overall transform
+		float3x3 SliceTransform3 = mul(transpose(float3x3(SplineMeshDir, SplineMeshX, SplineMeshY)), float3x3(float3(0,0,0), XVec, YVec));
+		float4x3 SliceTransform = float4x3(SliceTransform3[0], SliceTransform3[1], SliceTransform3[2], SplinePos);
+
+		return SliceTransform;
+	}
+
+	/** Calculate rotation matrix that defines frame along spline, given the Z of a vertex. */
+	half3x3 CalcSliceRot(float ZPos)
+	{
+		// Find how far 'along' mesh we are
+		half Alpha = ZPos * SplineMeshScaleZ - SplineMeshMinZ;
+
+		// Apply hermite interp to Alpha if desired
+		half HermiteAlpha = SmoothInterpRollScale ? smoothstep(0.0, 1.0, Alpha) : Alpha;
+
+		// Then find the point and direction of the spline at this point along
+		half3 SplineDir = SplineEvalDir( SplineStartPos, SplineStartTangent, SplineEndPos, SplineEndTangent, Alpha );
+
+		// Find base frenet frame
+		half3 BaseXVec = normalize( cross(SplineUpDir, (half3)SplineDir) );
+		half3 BaseYVec = normalize( cross((half3)SplineDir, BaseXVec) );
+
+		// Apply roll to frame around spline
+		half UseRoll = lerp((half)SplineStartRoll, (half)SplineEndRoll, HermiteAlpha);
+		half SinAng, CosAng;
+		sincos(UseRoll, SinAng, CosAng);
+		half3 XVec = (CosAng * BaseXVec) - (SinAng * BaseYVec);
+		half3 YVec = (CosAng * BaseYVec) + (SinAng * BaseXVec);
+
+		// Find scale at this point along spline
+		half2 UseScale = lerp(SplineStartScale, SplineEndScale, HermiteAlpha);
+
+		XVec *= sign(UseScale.x);
+		YVec *= sign(UseScale.y);
+
+		// Build rotation transform
+		half3x3 SliceTransform = mul(transpose(half3x3(SplineMeshDir, SplineMeshX, SplineMeshY)), half3x3(SplineDir, XVec, YVec));
+
+		return SliceTransform;
+	}
+#endif	// USE_SPLINEDEFORM
+
+#if USE_INSTANCING
+float4 CalcWorldPosition(float4 Position, float4x4 InstanceTransform)
+#else
+float4 CalcWorldPosition(float4 Position)
+#endif	// USE_INSTANCING
+{
+#if USE_INSTANCING
+	return TransformLocalToTranslatedWorld(mul(Position, InstanceTransform).xyz);
+#elif USE_SPLINEDEFORM
+/*
+	// Make transform for this point along spline
+	float4x3 SliceTransform = CalcSliceTransform(dot(Position.xyz, SplineMeshDir));
+
+	// Transform into mesh space
+	float4 LocalPos = float4(mul(Position, SliceTransform), Position.w);
+	// Transform from mesh to world space
+	return TransformLocalToTranslatedWorld(LocalPos.xyz);
+	*/
+	return INVARIANT(TransformLocalToTranslatedWorld(float3(mul(Position, CalcSliceTransform(dot(Position.xyz, SplineMeshDir))).xyz)));
+#else
+	return TransformLocalToTranslatedWorld(Position.xyz);
+#endif
+}
+
+half3x3 CalcTangentToLocal(FVertexFactoryInput Input)
+{
+	half3x3 Result;
+	
+	half3 TangentX = TangentBias(Input.TangentX);
+	half4 TangentZ = TangentBias(Input.TangentZ);
+
+#if USE_SPLINEDEFORM
+	// Make slice rotation matrix, and use that to transform tangents
+	half3x3 SliceRot = CalcSliceRot(dot(Input.Position.xyz, SplineMeshDir));
+
+	TangentX = mul(TangentX, SliceRot);
+	TangentZ.xyz = mul(TangentZ.xyz, SliceRot);
+#endif	// USE_SPLINEDEFORM
+
+	// derive the binormal by getting the cross product of the normal and tangent
+	half3 TangentY = cross(TangentZ.xyz, TangentX) * TangentZ.w;
+	
+	// Recalculate TangentX off of the other two vectors
+	// This corrects quantization error since TangentX was passed in as a quantized vertex input
+	// The error shows up most in specular off of a mesh with a smoothed UV seam (normal is smooth, but tangents vary across the seam)
+	Result[0] = cross(TangentY, TangentZ.xyz) * TangentZ.w;
+	Result[1] = TangentY;
+	Result[2] = TangentZ.xyz;
+
+	return Result;
+}
+
+half3x3 CalcTangentToWorld(FVertexFactoryInput Input, half3x3 TangentToLocal)
+{
+#if USE_INSTANCING
+	half3x3 InstanceToWorld = mul(GetInstanceToLocal3x3(Input), GetLocalToWorld3x3());
+	// remove scaling
+	InstanceToWorld[0] = normalize(InstanceToWorld[0]);
+	InstanceToWorld[1] = normalize(InstanceToWorld[1]);
+	InstanceToWorld[2] = normalize(InstanceToWorld[2]);
+	half3x3 TangentToWorld = mul(TangentToLocal, InstanceToWorld);
+#else
+	half3x3 TangentToWorld = CalcTangentToWorldNoScale(TangentToLocal);
+#endif	// USE_INSTANCING
+	return TangentToWorld;
+}
+
+FVertexFactoryIntermediates GetVertexFactoryIntermediates(FVertexFactoryInput Input)
+{
+	FVertexFactoryIntermediates Intermediates;
+
+	Intermediates.TangentToLocal = CalcTangentToLocal(Input);
+	Intermediates.TangentToWorld = CalcTangentToWorld(Input,Intermediates.TangentToLocal);
+	Intermediates.TangentToWorldSign = TangentBias(Input.TangentZ.w) * Primitive.LocalToWorldDeterminantSign;
+
+	// Swizzle vertex color.
+	Intermediates.Color = Input.Color FCOLOR_COMPONENT_SWIZZLE;
+#if USE_INSTANCING
+	// x = per-instance random, y = per-instance fade out factor, z = zero or one depending of if it is shown at all, w is dither cutoff 
+
+	// PerInstanceParams.z stores a hide/show flag for this instance
+	float SelectedValue = GetInstanceSelected(Input);
+	Intermediates.PerInstanceParams.x = GetInstanceRandom(Input);
+	float3 InstanceLocation = TransformLocalToWorld(GetInstanceOrigin(Input)).xyz;
+	Intermediates.PerInstanceParams.y = 1.0 - saturate((length(InstanceLocation + ResolvedView.PreViewTranslation.xyz) - InstancingFadeOutParams.x) * InstancingFadeOutParams.y);
+	// InstancingFadeOutParams.z,w are RenderSelected and RenderDeselected respectively.
+	Intermediates.PerInstanceParams.z = InstancingFadeOutParams.z * SelectedValue + InstancingFadeOutParams.w * (1-SelectedValue);
+	#if USE_DITHERED_LOD_TRANSITION
+		float RandomLOD = InstancingViewZCompareZero.w * Intermediates.PerInstanceParams.x;
+		float ViewZZero = length(InstanceLocation - InstancingWorldViewOriginZero.xyz) + RandomLOD;
+		float ViewZOne = length(InstanceLocation - InstancingWorldViewOriginOne.xyz) + RandomLOD;
+		Intermediates.PerInstanceParams.w = 
+			dot(float3(ViewZZero.xxx > InstancingViewZCompareZero.xyz), InstancingViewZConstant.xyz) * InstancingWorldViewOriginZero.w +
+			dot(float3(ViewZOne.xxx > InstancingViewZCompareOne.xyz), InstancingViewZConstant.xyz) * InstancingWorldViewOriginOne.w;
+		Intermediates.PerInstanceParams.z *= abs(Intermediates.PerInstanceParams.w) < .999;
+	#else
+		Intermediates.PerInstanceParams.w = 0;
+	#endif
+#endif	// USE_INSTANCING
+	return Intermediates;
+}
+
+/**
+* Get the 3x3 tangent basis vectors for this vertex factory
+* this vertex factory will calculate the binormal on-the-fly
+*
+* @param Input - vertex input stream structure
+* @return 3x3 matrix
+*/
+half3x3 VertexFactoryGetTangentToLocal( FVertexFactoryInput Input, FVertexFactoryIntermediates Intermediates )
+{
+	return Intermediates.TangentToLocal;
+}
+
+// @return translated world position
+float4 VertexFactoryGetWorldPosition(FVertexFactoryInput Input, FVertexFactoryIntermediates Intermediates)
+{
+#if USE_INSTANCING
+	return CalcWorldPosition(Input.Position, GetInstanceTransform(Input)) * Intermediates.PerInstanceParams.z;
+#else
+	return CalcWorldPosition(Input.Position);
+#endif	// USE_INSTANCING
+}
+
+float4 VertexFactoryGetRasterizedWorldPosition(FVertexFactoryInput Input, FVertexFactoryIntermediates Intermediates, float4 InWorldPosition)
+{
+	return InWorldPosition;
+}
+
+float3 VertexFactoryGetPositionForVertexLighting(FVertexFactoryInput Input, FVertexFactoryIntermediates Intermediates, float3 TranslatedWorldPosition)
+{
+	return TranslatedWorldPosition;
+}
+
+FVertexFactoryInterpolantsVSToPS VertexFactoryGetInterpolantsVSToPS(FVertexFactoryInput Input, FVertexFactoryIntermediates Intermediates, FMaterialVertexParameters VertexParameters)
+{
+	FVertexFactoryInterpolantsVSToPS Interpolants;
+
+	// Initialize the whole struct to 0
+	// Really only the last two components of the packed UVs have the opportunity to be uninitialized
+	Interpolants = (FVertexFactoryInterpolantsVSToPS)0;
+
+#if NUM_TEX_COORD_INTERPOLATORS
+	float2 CustomizedUVs[NUM_TEX_COORD_INTERPOLATORS];
+	GetMaterialCustomizedUVs(VertexParameters, CustomizedUVs);
+	GetCustomInterpolators(VertexParameters, CustomizedUVs);
+	
+	UNROLL
+	for (int CoordinateIndex = 0; CoordinateIndex < NUM_TEX_COORD_INTERPOLATORS; CoordinateIndex++)
+	{
+		SetUV(Interpolants, CoordinateIndex, CustomizedUVs[CoordinateIndex]);
+	}
+
+#elif NUM_MATERIAL_TEXCOORDS_VERTEX == 0 && USE_PARTICLE_SUBUVS
+	SetUV(Interpolants, 0, Input.TexCoords[0]);
+#endif
+
+#if NEEDS_LIGHTMAP_COORDINATE
+	float2 LightMapCoordinate = 0;
+	float2 ShadowMapCoordinate = 0;
+
+	#if USE_INSTANCING
+		LightMapCoordinate = Input.LightMapCoordinate * PrecomputedLightingBuffer.LightMapCoordinateScaleBias.xy + GetInstanceLightMapBias(Input);
+	#else
+		LightMapCoordinate = Input.LightMapCoordinate * PrecomputedLightingBuffer.LightMapCoordinateScaleBias.xy + PrecomputedLightingBuffer.LightMapCoordinateScaleBias.zw;
+	#endif
+	#if STATICLIGHTING_TEXTUREMASK
+		#if USE_INSTANCING
+			ShadowMapCoordinate = Input.LightMapCoordinate * PrecomputedLightingBuffer.ShadowMapCoordinateScaleBias.xy + GetInstanceShadowMapBias(Input);
+		#else
+			ShadowMapCoordinate = Input.LightMapCoordinate * PrecomputedLightingBuffer.ShadowMapCoordinateScaleBias.xy + PrecomputedLightingBuffer.ShadowMapCoordinateScaleBias.zw;
+		#endif
+	#endif	// STATICLIGHTING_TEXTUREMASK
+
+	SetLightMapCoordinate(Interpolants, LightMapCoordinate, ShadowMapCoordinate);
+#endif	// NEEDS_LIGHTMAP_COORDINATE
+
+	SetTangents(Interpolants, Intermediates.TangentToWorld[0], Intermediates.TangentToWorld[2], Intermediates.TangentToWorldSign);
+	SetColor(Interpolants, Intermediates.Color);
+#if USE_INSTANCING
+	Interpolants.PerInstanceParams = Intermediates.PerInstanceParams;
+#endif
+
+#if INSTANCED_STEREO
+	Interpolants.PackedEyeIndex = 0;
+#endif
+
+	return Interpolants;
+}
+
+/** for depth-only pass */
+float4 VertexFactoryGetWorldPosition(FPositionOnlyVertexFactoryInput Input)
+{
+#if USE_INSTANCING
+	return CalcWorldPosition(Input.Position, GetInstanceTransform(Input));
+#else
+	return CalcWorldPosition(Input.Position);
+#endif	// USE_INSTANCING
+}
+
+// Information when using pass through to get the previous position, as no matrix is available/needed
+#if GPUSKIN_PASS_THROUGH
+uint GPUSkinCachePreviousFloatOffset;
+Buffer<float> GPUSkinCachePreviousBuffer;
+#endif
+
+// @return previous translated world position
+float4 VertexFactoryGetPreviousWorldPosition(FVertexFactoryInput Input, FVertexFactoryIntermediates Intermediates)
+{
+	float4x4 PreviousLocalToWorldTranslated = PreviousLocalToWorld;
+	PreviousLocalToWorldTranslated[3][0] += ResolvedView.PrevPreViewTranslation.x;
+	PreviousLocalToWorldTranslated[3][1] += ResolvedView.PrevPreViewTranslation.y;
+	PreviousLocalToWorldTranslated[3][2] += ResolvedView.PrevPreViewTranslation.z;
+
+#if USE_INSTANCING
+	float4x4 InstanceTransform = GetInstanceTransform(Input);
+	return mul(mul(Input.Position, InstanceTransform), PreviousLocalToWorldTranslated);
+#elif GPUSKIN_PASS_THROUGH
+	uint Offset = Input.VertexId * GPUSKIN_RWBUFFER_NUM_FLOATS + GPUSkinCachePreviousFloatOffset;
+	float3 PreviousPos;
+	PreviousPos.x = GPUSkinCachePreviousBuffer[Offset + GPUSKIN_RWBUFFER_OFFSET_POSITION + 0];
+	PreviousPos.y = GPUSkinCachePreviousBuffer[Offset + GPUSKIN_RWBUFFER_OFFSET_POSITION + 1];
+	PreviousPos.z = GPUSkinCachePreviousBuffer[Offset + GPUSKIN_RWBUFFER_OFFSET_POSITION + 2];
+	return mul(float4(PreviousPos, 1), PreviousLocalToWorldTranslated);
+#elif USE_SPLINEDEFORM
+	// Just like CalcWorldPosition...
+	float4x3 SliceTransform = CalcSliceTransform(dot(Input.Position.xyz, SplineMeshDir));
+
+	// Transform into mesh space
+	float4 LocalPos = float4(mul(Input.Position, SliceTransform), Input.Position.w);
+
+	return mul(LocalPos, PreviousLocalToWorldTranslated);
+#else
+	return mul(Input.Position, PreviousLocalToWorldTranslated);
+#endif	// USE_INSTANCING
+}
+
+#if USING_TESSELLATION
+	struct FVertexFactoryInterpolantsVSToDS
+	{
+		FVertexFactoryInterpolantsVSToPS InterpolantsVSToPS;
+	};
+
+	float2 VertexFactoryGetTextureCoordinateDS( FVertexFactoryInterpolantsVSToDS Interpolants )
+	{
+	#if NUM_MATERIAL_TEXCOORDS
+		return Interpolants.InterpolantsVSToPS.TexCoords[0].xy;
+	#else // #if NUM_MATERIAL_TEXCOORDS
+		return float2(0,0);
+	#endif // #if NUM_MATERIAL_TEXCOORDS
+	}
+
+	FVertexFactoryInterpolantsVSToPS VertexFactoryAssignInterpolants(FVertexFactoryInterpolantsVSToDS Input)
+	{
+		return Input.InterpolantsVSToPS;
+	}
+
+	/** Converts from vertex factory specific interpolants to a FMaterialTessellationParameters, which is used by material inputs. */
+	FMaterialTessellationParameters GetMaterialTessellationParameters(FVertexFactoryInterpolantsVSToDS Interpolants, float3 CameraLocalWorldPosition)
+	{
+		FMaterialTessellationParameters	Result;
+	#if NUM_TEX_COORD_INTERPOLATORS
+		UNROLL
+		for(int CoordinateIndex = 0;CoordinateIndex < NUM_TEX_COORD_INTERPOLATORS;CoordinateIndex += 2)
+		{
+			Result.TexCoords[CoordinateIndex] = Interpolants.InterpolantsVSToPS.TexCoords[CoordinateIndex/2].xy;
+			if(CoordinateIndex + 1 < NUM_TEX_COORD_INTERPOLATORS)
+			{
+				Result.TexCoords[CoordinateIndex + 1] = Interpolants.InterpolantsVSToPS.TexCoords[CoordinateIndex/2].zw;
+			}
+		}
+	#endif // NUM_TEX_COORD_INTERPOLATORS
+
+		half3 TangentToWorld0 = Interpolants.InterpolantsVSToPS.TangentToWorld0.xyz;
+		half4 TangentToWorld2 = Interpolants.InterpolantsVSToPS.TangentToWorld2;
+
+	#if INTERPOLATE_VERTEX_COLOR
+		Result.VertexColor = Interpolants.InterpolantsVSToPS.Color;
+	#endif	// INTERPOLATE_VERTEX_COLOR
+
+		Result.TangentToWorld = AssembleTangentToWorld( TangentToWorld0, TangentToWorld2 );
+
+		Result.TangentToWorldPreScale = 1;
+
+		Result.WorldPosition = CameraLocalWorldPosition + ResolvedView.WorldCameraOrigin;
+
+		return Result;
+	}
+
+	FVertexFactoryInterpolantsVSToDS VertexFactoryGetInterpolantsVSToDS(FVertexFactoryInput Input, FVertexFactoryIntermediates Intermediates, FMaterialVertexParameters VertexParameters)
+	{
+		FVertexFactoryInterpolantsVSToDS Interpolants;
+		Interpolants.InterpolantsVSToPS = VertexFactoryGetInterpolantsVSToPS(Input, Intermediates, VertexParameters);
+		return Interpolants;
+	}
+
+	FVertexFactoryInterpolantsVSToDS VertexFactoryInterpolate(FVertexFactoryInterpolantsVSToDS a, float aInterp, FVertexFactoryInterpolantsVSToDS b, float bInterp)
+	{
+		FVertexFactoryInterpolantsVSToDS O;
+	
+		// Do we really need to interpolate TangentToWorld2 here? It should be replaced by the
+		// interpolated normal from 'whatever' interpolation scheme we're using
+	
+		TESSELLATION_INTERPOLATE_MEMBER(InterpolantsVSToPS.TangentToWorld0.xyz);
+		TESSELLATION_INTERPOLATE_MEMBER(InterpolantsVSToPS.TangentToWorld2);
+	#if INTERPOLATE_VERTEX_COLOR
+		TESSELLATION_INTERPOLATE_MEMBER(InterpolantsVSToPS.Color);
+	#endif
+	#if USE_INSTANCING
+		TESSELLATION_INTERPOLATE_MEMBER(InterpolantsVSToPS.PerInstanceParams);
+	#endif
+
+	#if NEEDS_LIGHTMAP_COORDINATE
+		TESSELLATION_INTERPOLATE_MEMBER(InterpolantsVSToPS.LightMapCoordinate);
+	#endif
+
+	#if NUM_TEX_COORD_INTERPOLATORS
+		UNROLL
+		for(int tc = 0; tc < (NUM_TEX_COORD_INTERPOLATORS+1)/2; ++tc)
+		{
+			TESSELLATION_INTERPOLATE_MEMBER(InterpolantsVSToPS.TexCoords[tc]);
+		}
+	#endif
+
+		return O;
+	}
+
+	float3x3 VertexFactoryGetTangentToLocalDS(FVertexFactoryInterpolantsVSToDS Interpolants)
+	{
+		// This duplicates stuff already going on in GetMaterialTessellationParameters(), so
+		// maybe the hull shader could leverage that instead?
+		half3 TangentToWorld0 = Interpolants.InterpolantsVSToPS.TangentToWorld0.xyz;
+		half4 TangentToWorld2 = Interpolants.InterpolantsVSToPS.TangentToWorld2;
+		float3x3 TangentToWorld = AssembleTangentToWorld( TangentToWorld0, TangentToWorld2 );
+
+		return TangentToWorld;
+	}
+#endif // #if USING_TESSELLATION
+
+#if USE_INSTANCING
+float4 VertexFactoryGetInstanceHitProxyId(FVertexFactoryInput Input, FVertexFactoryIntermediates Intermediates)
+{
+	#if !USE_INSTANCING_EMULATED
+		float R = Input.InstanceTransform1.w - 256.0 * GetInstanceSelected(Input);
+		float G = Input.InstanceTransform2.w;
+		float B = Input.InstanceTransform3.w;
+	#else
+		float R = CPUInstanceTransform[0].w - 256.0 * GetInstanceSelected(Input);
+		float G = CPUInstanceTransform[1].w;
+		float B = CPUInstanceTransform[2].w;
+	#endif
+	return float4(R/255.0, G/255.0, B/255.0, 0);
+}
+#endif	// USE_INSTANCING
+
+#if INSTANCED_STEREO
+uint VertexFactoryGetEyeIndex(uint InstanceId)
+{
+	#if USE_INSTANCING
+		return InstancedEyeIndex;
+	#else
+		return InstanceId & 1;
+	#endif
+}
+#endif
+
+float4 VertexFactoryGetTranslatedPrimitiveVolumeBounds(FVertexFactoryInterpolantsVSToPS Interpolants)
+{
+	return 0;
+}
diff --git a/Engine/Shaders/Private/LocalVertexFactoryCommon.ush b/Engine/Shaders/Private/LocalVertexFactoryCommon.ush
new file mode 100644
--- /dev/null
+++ b/Engine/Shaders/Private/LocalVertexFactoryCommon.ush
@@ -0,0 +1,251 @@
+// Copyright 1998-2017 Epic Games, Inc. All Rights Reserved.
+
+/*=============================================================================
+	LocalVertexFactoryCommon.usf: Local vertex factory common functionality
+=============================================================================*/
+
+// FVertexFactoryInterpolantsVSToPS is split into a PC and ES2 version, since the ES2 version has more restrictions on packing
+// Then accessors are provided for each property which try to encapsulate the packing differences
+#if FEATURE_LEVEL >= FEATURE_LEVEL_ES3_1
+
+struct FVertexFactoryInterpolantsVSToPS
+{
+	TANGENTTOWORLD_INTERPOLATOR_BLOCK
+
+#if INTERPOLATE_VERTEX_COLOR
+	half4	Color : COLOR0;
+#endif
+
+#if USE_INSTANCING
+	// x = per-instance random, y = per-instance fade out amount, z = hide/show flag, w dither fade cutoff
+	float4  PerInstanceParams : COLOR1;
+#endif
+
+#if NUM_TEX_COORD_INTERPOLATORS
+	float4	TexCoords[(NUM_TEX_COORD_INTERPOLATORS+1)/2]	: TEXCOORD0;
+#elif USE_PARTICLE_SUBUVS
+	float4	TexCoords[1] : TEXCOORD0;
+#endif
+
+#if NEEDS_LIGHTMAP_COORDINATE
+	float4	LightMapCoordinate : TEXCOORD4;
+#endif
+
+#if INSTANCED_STEREO
+	nointerpolation uint PackedEyeIndex : PACKED_EYE_INDEX;
+#endif
+};
+
+#if NUM_TEX_COORD_INTERPOLATORS || USE_PARTICLE_SUBUVS
+float2 GetUV(FVertexFactoryInterpolantsVSToPS Interpolants, int UVIndex)
+{
+	float4 UVVector = Interpolants.TexCoords[UVIndex / 2];
+	return UVIndex % 2 ? UVVector.zw : UVVector.xy;
+}
+
+void SetUV(inout FVertexFactoryInterpolantsVSToPS Interpolants, int UVIndex, float2 InValue)
+{
+	FLATTEN
+	if (UVIndex % 2)
+	{
+		Interpolants.TexCoords[UVIndex / 2].zw = InValue;
+	}
+	else
+	{
+		Interpolants.TexCoords[UVIndex / 2].xy = InValue;
+	}
+}
+#endif
+
+float4 GetColor(FVertexFactoryInterpolantsVSToPS Interpolants)
+{
+#if INTERPOLATE_VERTEX_COLOR
+	return Interpolants.Color;
+#else
+	return 0;
+#endif
+}
+
+void SetColor(inout FVertexFactoryInterpolantsVSToPS Interpolants, float4 InValue)
+{
+#if INTERPOLATE_VERTEX_COLOR
+	Interpolants.Color = InValue;
+#endif
+}
+
+#if NEEDS_LIGHTMAP_COORDINATE
+void GetLightMapCoordinates(FVertexFactoryInterpolantsVSToPS Interpolants, out float2 LightmapUV0, out float2 LightmapUV1)
+{
+	LightmapUV0 = Interpolants.LightMapCoordinate.xy * float2( 1, 0.5 );
+	LightmapUV1 = LightmapUV0 + float2( 0, 0.5 );
+}
+
+float2 GetShadowMapCoordinate(FVertexFactoryInterpolantsVSToPS Interpolants)
+{
+	return Interpolants.LightMapCoordinate.zw;
+}
+
+void SetLightMapCoordinate(inout FVertexFactoryInterpolantsVSToPS Interpolants, float2 InLightMapCoordinate, float2 InShadowMapCoordinate)
+{
+	Interpolants.LightMapCoordinate.xy = InLightMapCoordinate;
+	Interpolants.LightMapCoordinate.zw = InShadowMapCoordinate;
+}
+#endif
+
+float4 GetTangentToWorld2(FVertexFactoryInterpolantsVSToPS Interpolants)
+{
+	return Interpolants.TangentToWorld2;
+}
+
+float4 GetTangentToWorld0(FVertexFactoryInterpolantsVSToPS Interpolants)
+{
+	return Interpolants.TangentToWorld0;
+}
+
+void SetTangents(inout FVertexFactoryInterpolantsVSToPS Interpolants, float3 InTangentToWorld0, float3 InTangentToWorld2, float InTangentToWorldSign)
+{
+	Interpolants.TangentToWorld0 = float4(InTangentToWorld0,0);
+	Interpolants.TangentToWorld2 = float4(InTangentToWorld2,InTangentToWorldSign);
+}
+
+
+#else // #if FEATURE_LEVEL >= FEATURE_LEVEL_ES3_1
+
+#define NUM_PACKED_UVS (NUM_TEX_COORD_INTERPOLATORS > 1 ? 3 : 1)
+
+#define USE_INTERPOLATOR_345 ((!MATERIAL_SHADINGMODEL_UNLIT) || USES_TRANSFORM_VECTOR || USES_DISTORTION)
+
+/** Interpolants packed for ES2.  All UVs are in the xy to ensure independent texture fetches. */
+struct FVertexFactoryInterpolantsVSToPS
+{
+	// Save one interpolator if the material is only using one UV
+#if NUM_PACKED_UVS == 1
+	// UV0, normal xy
+	float4 PackedUVs[1] : TEXCOORD0;
+
+	#if INTERPOLATE_VERTEX_COLOR
+		// Only interpolate vertex color if needed by the material pixel shader
+		float4 Color : TEXCOORD1;
+	#endif
+#else
+	// UV0, normal xy
+	// UV1, color xy
+	// UV2, color zw
+	float4 PackedUVs[3] : TEXCOORD0;
+#endif
+
+#if USE_INTERPOLATOR_345
+	// LightmapUV0, normal zw
+	float4 PackedInterpolator3 : TEXCOORD3;
+
+	// LightmapUV1, tangent xy
+	float4 PackedInterpolator4 : TEXCOORD4;
+
+	// ShadowmapUV, tangent zw
+	float4 PackedInterpolator5 : TEXCOORD5;
+#endif
+
+#if USE_INSTANCING
+	// x = per-instance random, y = per-instance fade out amount, z = hide/show flag, w dither fade cutoff
+	float4  PerInstanceParams : COLOR1;
+#endif
+
+#if INSTANCED_STEREO
+	nointerpolation uint PackedEyeIndex : PACKED_EYE_INDEX;
+#endif
+};
+
+float2 GetUV(FVertexFactoryInterpolantsVSToPS Interpolants, int UVIndex)
+{
+	return Interpolants.PackedUVs[UVIndex].xy;
+}
+
+void SetUV(inout FVertexFactoryInterpolantsVSToPS Interpolants, int UVIndex, float2 InValue)
+{
+	Interpolants.PackedUVs[UVIndex].xy = InValue;
+}
+
+float4 GetColor(FVertexFactoryInterpolantsVSToPS Interpolants)
+{
+	#if NUM_PACKED_UVS == 1
+		#if INTERPOLATE_VERTEX_COLOR
+			return Interpolants.Color;
+		#else
+			return 0;
+		#endif
+	#else
+		return float4(Interpolants.PackedUVs[1].zw, Interpolants.PackedUVs[2].zw);
+	#endif
+}
+
+void SetColor(inout FVertexFactoryInterpolantsVSToPS Interpolants, float4 InValue)
+{
+	#if NUM_PACKED_UVS == 1
+		#if INTERPOLATE_VERTEX_COLOR
+			Interpolants.Color = InValue;
+		#endif
+	#else
+		Interpolants.PackedUVs[1].zw = InValue.xy;
+		Interpolants.PackedUVs[2].zw = InValue.zw;
+	#endif
+}
+
+void GetLightMapCoordinates(FVertexFactoryInterpolantsVSToPS Interpolants, out float2 LightmapUV0, out float2 LightmapUV1)
+{
+	#if USE_INTERPOLATOR_345 
+		LightmapUV0 = Interpolants.PackedInterpolator3.xy;
+		LightmapUV1 = Interpolants.PackedInterpolator4.xy;
+	#else
+		LightmapUV0 = LightmapUV1 = 0;
+	#endif
+}
+
+float2 GetShadowMapCoordinate(FVertexFactoryInterpolantsVSToPS Interpolants)
+{
+	#if USE_INTERPOLATOR_345 
+		return Interpolants.PackedInterpolator5.xy;
+	#else
+		return float2(0, 0);
+	#endif
+}
+
+void SetLightMapCoordinate(inout FVertexFactoryInterpolantsVSToPS Interpolants, float2 InLightMapCoordinate, float2 InShadowMapCoordinate)
+{
+	#if USE_INTERPOLATOR_345
+		Interpolants.PackedInterpolator3.xy = InLightMapCoordinate * float2(1, 0.5);
+		Interpolants.PackedInterpolator4.xy = Interpolants.PackedInterpolator3.xy + float2(0, 0.5);
+		Interpolants.PackedInterpolator5.xy = InShadowMapCoordinate;
+	#endif
+}
+
+float4 GetTangentToWorld2(FVertexFactoryInterpolantsVSToPS Interpolants)
+{
+	#if USE_INTERPOLATOR_345
+		return float4(Interpolants.PackedUVs[0].zw, Interpolants.PackedInterpolator3.zw);
+	#else
+		return float4(0, 0, 0, 1);
+	#endif
+}
+
+float4 GetTangentToWorld0(FVertexFactoryInterpolantsVSToPS Interpolants)
+{
+	#if USE_INTERPOLATOR_345
+		return float4(Interpolants.PackedInterpolator4.zw, Interpolants.PackedInterpolator5.zw);
+	#else
+		return float4(0, 0, 0, 0);
+	#endif
+}
+
+void SetTangents(inout FVertexFactoryInterpolantsVSToPS Interpolants, float3 InTangentToWorld0, float3 InTangentToWorld2, float InTangentToWorldSign)
+{
+	#if USE_INTERPOLATOR_345
+		Interpolants.PackedInterpolator4.zw = InTangentToWorld0.xy;
+		Interpolants.PackedInterpolator5.zw = InTangentToWorld0.zz;
+		Interpolants.PackedUVs[0].zw = InTangentToWorld2.xy;
+		Interpolants.PackedInterpolator3.zw = float2(InTangentToWorld2.z,InTangentToWorldSign);
+	#else
+		Interpolants.PackedUVs[0].zw = 0;
+	#endif
+}
+
+#endif
