Branch snapshot for CL 3555263

diff --git a/Engine/Shaders/Private/MaterialTemplate.ush b/Engine/Shaders/Private/MaterialTemplate.ush
new file mode 100644
--- /dev/null
+++ b/Engine/Shaders/Private/MaterialTemplate.ush
@@ -0,0 +1,1910 @@
+// Copyright 1998-2017 Epic Games, Inc. All Rights Reserved.
+
+/**
+ * MaterialTemplate.usf: Filled in by FHLSLMaterialTranslator::GetMaterialShaderCode for each material being compiled.
+ */
+
+#include "/Engine/Private/Random.ush"
+#include "/Engine/Private/MonteCarlo.ush"
+#include "/Engine/Generated/UniformBuffers/Material.ush"
+
+// for MaterialExpressionDepthOfFieldFunction
+#include "/Engine/Private/DepthOfFieldCommon.ush"
+#include "/Engine/Private/CircleDOFCommon.ush"
+#include "/Engine/Private/GlobalDistanceFieldShared.ush"
+
+#if USES_SPEEDTREE
+	#include "/Engine/Private/SpeedTreeCommon.ush"
+#endif
+
+#define NUM_MATERIAL_TEXCOORDS_VERTEX %s
+#define NUM_MATERIAL_TEXCOORDS %s
+#define NUM_CUSTOM_VERTEX_INTERPOLATORS %s
+#define NUM_TEX_COORD_INTERPOLATORS %s
+
+#ifdef MIN_MATERIAL_TEXCOORDS 
+	#include "/Engine/Private/MinMaterialTexCoords.ush"
+#endif
+ 
+#if MATERIAL_ATMOSPHERIC_FOG
+	#include "/Engine/Private/AtmosphereCommon.ush"
+#endif
+
+#include "/Engine/Private/PaniniProjection.ush"
+
+#ifndef USE_DITHERED_LOD_TRANSITION
+    #if USE_INSTANCING
+	    #ifndef USE_DITHERED_LOD_TRANSITION_FOR_INSTANCED
+		    #error "USE_DITHERED_LOD_TRANSITION_FOR_INSTANCED should have been defined"
+	    #endif
+	    #define USE_DITHERED_LOD_TRANSITION USE_DITHERED_LOD_TRANSITION_FOR_INSTANCED
+    #else
+	    #ifndef USE_DITHERED_LOD_TRANSITION_FROM_MATERIAL
+		    #error "USE_DITHERED_LOD_TRANSITION_FROM_MATERIAL should have been defined"
+	    #endif
+	    #define USE_DITHERED_LOD_TRANSITION USE_DITHERED_LOD_TRANSITION_FROM_MATERIAL
+    #endif
+#endif
+
+#ifndef USE_STENCIL_LOD_DITHER
+	#define USE_STENCIL_LOD_DITHER	USE_STENCIL_LOD_DITHER_DEFAULT
+#endif
+
+//Platforms that don't run the editor shouldn't need editor features in the shaders.
+#define PLATFORM_SUPPORTS_EDITOR_SHADERS (!PS4_PROFILE && !XBOXONE_PROFILE && !ESDEFERRED_PROFILE && !SWITCH_PROFILE)
+
+//Tie Editor features to platform support and the COMPILE_SHADERS_FOR_DEVELOPMENT which is set via CVAR.
+#define USE_EDITOR_SHADERS (PLATFORM_SUPPORTS_EDITOR_SHADERS && USE_DEVELOPMENT_SHADERS)
+
+//Materials also have to opt in to these features.
+#define USE_EDITOR_COMPOSITING (USE_EDITOR_SHADERS && EDITOR_PRIMITIVE_MATERIAL)
+
+#define MATERIALBLENDING_ANY_TRANSLUCENT (MATERIALBLENDING_TRANSLUCENT || MATERIALBLENDING_ADDITIVE || MATERIALBLENDING_MODULATE)
+
+#define IS_MESHPARTICLE_FACTORY (PARTICLE_MESH_FACTORY || NIAGARA_MESH_FACTORY)
+
+/**
+ * Parameters used by vertex and pixel shaders to access particle properties.
+ */
+struct FMaterialParticleParameters
+{
+	/** Relative time [0-1]. */
+	half RelativeTime;
+	/** Fade amount due to motion blur. */
+	half MotionBlurFade;
+	/** Random value per particle [0-1]. */
+	half Random;
+	/** XYZ: Direction, W: Speed. */
+	half4 Velocity;
+	/** Per-particle color. */
+	half4 Color;
+	/** Particle translated world space position and size(radius). */
+	float4 TranslatedWorldPositionAndSize;
+	/** Macro UV scale and bias. */
+	half4 MacroUV;
+	/** Dynamic parameter used by particle systems. */
+	half4 DynamicParameter;
+	/** mesh particle orientation */
+	float4x4 LocalToWorld;
+
+#if USE_PARTICLE_SUBUVS
+	/** SubUV texture coordinates*/
+	MaterialFloat2 SubUVCoords[2];
+	/** SubUV interpolation value*/
+	MaterialFloat SubUVLerp;
+#endif
+
+	/** The size of the particle. */
+	float2 Size;
+};
+
+float4 GetDynamicParameter(FMaterialParticleParameters Parameters, float4 Default)
+{
+	#if PARTICLE_FACTORY
+	return Parameters.DynamicParameter;
+	#else
+	return Default;
+	#endif
+}
+
+/** 
+ * Parameters calculated from the pixel material inputs.
+ */
+struct FPixelMaterialInputs
+{
+%s
+};
+
+/** 
+ * Parameters needed by pixel shader material inputs, related to Geometry.
+ * These are independent of vertex factory.
+ */
+struct FMaterialPixelParameters
+{
+#if NUM_TEX_COORD_INTERPOLATORS
+	float2 TexCoords[NUM_TEX_COORD_INTERPOLATORS];
+#endif
+
+	/** Interpolated vertex color, in linear color space. */
+	half4 VertexColor;
+
+	/** Normalized world space normal. */
+	half3 WorldNormal;
+
+	/** Normalized world space reflected camera vector. */
+	half3 ReflectionVector;
+
+	/** Normalized world space camera vector, which is the vector from the point being shaded to the camera position. */
+	half3 CameraVector;
+
+	/** World space light vector, only valid when rendering a light function. */
+	half3 LightVector;
+
+	/**
+	 * Like SV_Position (.xy is pixel position at pixel center, z:DeviceZ, .w:SceneDepth)
+	 * using shader generated value SV_POSITION
+	 * Note: this is not relative to the current viewport.  RelativePixelPosition = MaterialParameters.SvPosition.xy - View.ViewRectMin.xy;
+	 */
+	float4 SvPosition;
+		
+	/** Post projection position reconstructed from SvPosition, before the divide by W. left..top -1..1, bottom..top -1..1  within the viewport, W is the SceneDepth */
+	float4 ScreenPosition;
+
+	half UnMirrored;
+
+	half TwoSidedSign;
+
+	/**
+	 * Orthonormal rotation-only transform from tangent space to world space
+	 * The transpose(TangentToWorld) is WorldToTangent, and TangentToWorld[2] is WorldVertexNormal
+	 */
+	half3x3 TangentToWorld;
+
+	/** 
+	 * Interpolated worldspace position of this pixel
+	 * todo: Make this TranslatedWorldPosition and also rename the VS/DS/HS WorldPosition to be TranslatedWorldPosition
+	 */
+	float3 AbsoluteWorldPosition;
+
+	/** 
+	 * Interpolated worldspace position of this pixel, centered around the camera
+	 */
+	float3 WorldPosition_CamRelative;
+
+	/** 
+	 * Interpolated worldspace position of this pixel, not including any world position offset or displacement.
+	 * Only valid if shader is compiled with NEEDS_WORLD_POSITION_EXCLUDING_SHADER_OFFSETS, otherwise just contains 0
+	 */
+	float3 WorldPosition_NoOffsets;
+
+	/** 
+	 * Interpolated worldspace position of this pixel, not including any world position offset or displacement.
+	 * Only valid if shader is compiled with NEEDS_WORLD_POSITION_EXCLUDING_SHADER_OFFSETS, otherwise just contains 0
+	 */
+	float3 WorldPosition_NoOffsets_CamRelative;
+
+	/** Offset applied to the lighting position for translucency, used to break up aliasing artifacts. */
+	half3 LightingPositionOffset;
+
+	float AOMaterialMask;
+
+#if LIGHTMAP_UV_ACCESS
+	float2	LightmapUVs;
+#endif
+
+#if USE_INSTANCING
+	half4 PerInstanceParams;
+#endif
+
+	/** Per-particle properties. Only valid for particle vertex factories. */
+	FMaterialParticleParameters Particle;
+
+#if (ES2_PROFILE || ES3_1_PROFILE)
+	float4 LayerWeights;
+#endif
+
+#if TEX_COORD_SCALE_ANALYSIS
+	/** Parameters used by the MaterialTexCoordScales shader. */
+	FTexCoordScalesParams TexCoordScalesParams;
+#endif
+
+#if POST_PROCESS_MATERIAL && (FEATURE_LEVEL <= FEATURE_LEVEL_ES3_1)
+	/** Used in mobile custom pp material to preserve original SceneColor Alpha */
+	half BackupSceneColorAlpha;
+#endif
+
+#if COMPILER_HLSL
+	// Workaround for "error X3067: 'GetObjectWorldPosition': ambiguous function call"
+	// Which happens when FMaterialPixelParameters and FMaterialVertexParameters have the same number of floats with the HLSL compiler ver 9.29.952.3111
+	// Function overload resolution appears to identify types based on how many floats / ints / etc they contain
+	uint Dummy;
+#endif
+};
+
+// @todo compat hack
+FMaterialPixelParameters MakeInitializedMaterialPixelParameters()
+{
+	FMaterialPixelParameters MPP;
+	MPP = (FMaterialPixelParameters)0;
+	MPP.TangentToWorld = float3x3(1,0,0,0,1,0,0,0,1);
+	return MPP;
+}
+
+/** 
+ * Parameters needed by domain shader material inputs.
+ * These are independent of vertex factory.
+ */
+struct FMaterialTessellationParameters
+{
+	// Note: Customized UVs are only evaluated in the vertex shader, which is not really what you want with tessellation, but keeps the code simpler 
+	// (tessellation texcoords are the same as pixels shader texcoords)
+#if NUM_TEX_COORD_INTERPOLATORS
+	float2 TexCoords[NUM_TEX_COORD_INTERPOLATORS];
+#endif
+	float4 VertexColor;
+	// TODO: Non translated world position
+	float3 WorldPosition;
+	float3 TangentToWorldPreScale;
+
+	// TangentToWorld[2] is WorldVertexNormal, [0] and [1] are binormal and tangent
+	float3x3 TangentToWorld;
+};
+
+/** 
+ * Parameters needed by vertex shader material inputs.
+ * These are independent of vertex factory.
+ */
+struct FMaterialVertexParameters
+{
+	// Position in the translated world (VertexFactoryGetWorldPosition).
+	// Previous position in the translated world (VertexFactoryGetPreviousWorldPosition) if
+	//    computing material's output for previous frame (See {BasePassVertex,Velocity}Shader.usf).
+	float3 WorldPosition;
+	// TangentToWorld[2] is WorldVertexNormal
+	half3x3 TangentToWorld;
+#if USE_INSTANCING
+	/** Per-instance properties. */
+	float4x4 InstanceLocalToWorld;
+	float3 InstanceLocalPosition;
+	float4 PerInstanceParams;
+#elif IS_MESHPARTICLE_FACTORY 
+	/** Per-particle properties. */
+	float4x4 InstanceLocalToWorld;
+#endif
+	float3 PreSkinnedPosition;
+	float3 PreSkinnedNormal;
+
+	half4 VertexColor;
+#if NUM_MATERIAL_TEXCOORDS_VERTEX
+	float2 TexCoords[NUM_MATERIAL_TEXCOORDS_VERTEX];
+	#if (ES2_PROFILE || ES3_1_PROFILE)
+	float2 TexCoordOffset; // Offset for UV localization for large UV values
+	#endif
+#endif
+
+	/** Per-particle properties. Only valid for particle vertex factories. */
+	FMaterialParticleParameters Particle;
+};
+
+float3 GetTranslatedWorldPosition(FMaterialVertexParameters Parameters)
+{
+	return Parameters.WorldPosition;
+}
+
+float3 GetPrevTranslatedWorldPosition(FMaterialVertexParameters Parameters)
+{
+	// Previous world position and current world position are sharing the
+	// same attribute in Parameters, because in BasePassVertexShader.usf
+	// and in VelocityShader.usf, we are regenerating a Parameters from
+	// VertexFactoryGetPreviousWorldPosition() instead of
+	// VertexFactoryGetWorldPosition().
+	return GetTranslatedWorldPosition(Parameters);
+}
+
+float3 GetWorldPosition(FMaterialVertexParameters Parameters)
+{
+	return GetTranslatedWorldPosition(Parameters) - ResolvedView.PreViewTranslation;
+}
+
+float3 GetPrevWorldPosition(FMaterialVertexParameters Parameters)
+{
+	return GetPrevTranslatedWorldPosition(Parameters) - ResolvedView.PrevPreViewTranslation;
+}
+
+//TODO(bug UE-17131): We should compute world displacement for the previous frame
+float3 GetWorldPosition(FMaterialTessellationParameters Parameters)
+{
+	return Parameters.WorldPosition;
+}
+
+float3 GetTranslatedWorldPosition(FMaterialTessellationParameters Parameters)
+{
+	return Parameters.WorldPosition + ResolvedView.PreViewTranslation;
+}
+
+float3 GetWorldPosition(FMaterialPixelParameters Parameters)
+{
+	return Parameters.AbsoluteWorldPosition;
+}
+
+float3 GetWorldPosition_NoMaterialOffsets(FMaterialPixelParameters Parameters)
+{
+	return Parameters.WorldPosition_NoOffsets;
+}
+
+float3 GetTranslatedWorldPosition(FMaterialPixelParameters Parameters)
+{
+	return Parameters.WorldPosition_CamRelative;
+}
+
+float3 GetTranslatedWorldPosition_NoMaterialOffsets(FMaterialPixelParameters Parameters)
+{
+	return Parameters.WorldPosition_NoOffsets_CamRelative;
+}
+
+float4 GetScreenPosition(FMaterialVertexParameters Parameters)
+{
+	return mul(float4(Parameters.WorldPosition, 1.0f), ResolvedView.TranslatedWorldToClip);
+}
+
+float4 GetScreenPosition(FMaterialPixelParameters Parameters)
+{
+	return Parameters.ScreenPosition;
+}
+
+float2 GetSceneTextureUV(FMaterialVertexParameters Parameters)
+{
+	return ScreenAlignedPosition(GetScreenPosition(Parameters));
+}
+
+float2 GetSceneTextureUV(FMaterialPixelParameters Parameters)
+{
+	return SvPositionToBufferUV(Parameters.SvPosition);
+}
+
+float2 GetViewportUV(FMaterialVertexParameters Parameters)
+{
+	return BufferUVToViewportUV(GetSceneTextureUV(Parameters));
+}
+
+float2 GetViewportUV(FMaterialPixelParameters Parameters)
+{
+	return SvPositionToViewportUV(Parameters.SvPosition);
+}
+
+
+
+#if DECAL_PRIMITIVE && NUM_MATERIAL_TEXCOORDS
+	/*
+	 * Material node DecalMipmapLevel's code designed to avoid the 2x2 pixels artefacts on the edges around where the decal
+	 * is projected to. The technique is fetched from (http://www.humus.name/index.php?page=3D&ID=84).
+	 *
+	 * The problem around edges of the meshes, is that the hardware computes the mipmap level according to ddx(uv) and ddy(uv),
+	 * but since the pixel shader are invocated by group of 2x2 pixels, then on edges some pixel might be getting the
+	 * current depth of an differet mesh that the other pixel of the same groups. If this mesh is very far from the other
+	 * mesh of the same group of pixel, then one of the delta might be very big, leading to choosing a low mipmap level for this
+	 * group of 4 pixels, causing the artefacts.
+	 */
+	float2 ComputeDecalUVFromSvPosition(float4 SvPosition)
+	{
+		half DeviceZ = LookupDeviceZ(SvPositionToBufferUV(SvPosition));
+
+		SvPosition.z = DeviceZ;
+
+		float4 DecalVector = mul(float4(SvPosition.xyz,1), SvPositionToDecal);
+		DecalVector.xyz /= DecalVector.w;
+		DecalVector = DecalVector * 0.5f + 0.5f;
+		DecalVector.xyz = DecalVector.zyx;
+		return DecalVector.xy;
+	}
+
+	float2 ComputeDecalDDX(FMaterialPixelParameters Parameters)
+	{
+		/*
+		 * Assuming where in a pixel shader invocation, then we compute manualy compute two d(uv)/d(x)
+		 * with the pixels's left and right neighbours.
+		 */
+		float4 ScreenDeltaX = float4(1, 0, 0, 0);
+		float2 UvDiffX0 = Parameters.TexCoords[0] - ComputeDecalUVFromSvPosition(Parameters.SvPosition - ScreenDeltaX);
+		float2 UvDiffX1 = ComputeDecalUVFromSvPosition(Parameters.SvPosition + ScreenDeltaX) - Parameters.TexCoords[0];
+
+		/*
+		 * So we have two diff on the X axis, we want the one that has the smallest length
+		 * to avoid the 2x2 pixels mipmap artefacts on the edges. 
+		 */
+		return dot(UvDiffX0, UvDiffX0) < dot(UvDiffX1, UvDiffX1) ? UvDiffX0 : UvDiffX1;
+	}
+
+	float2 ComputeDecalDDY(FMaterialPixelParameters Parameters)
+	{
+		// do same for the Y axis
+		float4 ScreenDeltaY = float4(0, 1, 0, 0);
+		float2 UvDiffY0 = Parameters.TexCoords[0] - ComputeDecalUVFromSvPosition(Parameters.SvPosition - ScreenDeltaY);
+		float2 UvDiffY1 = ComputeDecalUVFromSvPosition(Parameters.SvPosition + ScreenDeltaY) - Parameters.TexCoords[0];
+
+		return dot(UvDiffY0, UvDiffY0) < dot(UvDiffY1, UvDiffY1) ? UvDiffY0 : UvDiffY1;
+	}
+
+	float ComputeDecalMipmapLevel(FMaterialPixelParameters Parameters, float2 TextureSize)
+	{
+		float2 UvPixelDiffX = ComputeDecalDDX(Parameters) * TextureSize;
+		float2 UvPixelDiffY = ComputeDecalDDY(Parameters) * TextureSize;
+
+		// Computes the mipmap level
+		float MaxDiff = max(dot(UvPixelDiffX, UvPixelDiffX), dot(UvPixelDiffY, UvPixelDiffY));
+		return 0.5 * log2(MaxDiff);
+	}
+#else // DECAL_PRIMITIVE && NUM_MATERIAL_TEXCOORDS
+	float2 ComputeDecalDDX(FMaterialPixelParameters Parameters)
+	{
+		return 0.0f;
+	}
+	
+	float2 ComputeDecalDDY(FMaterialPixelParameters Parameters)
+	{
+		return 0.0f;
+	}
+
+	float ComputeDecalMipmapLevel(FMaterialPixelParameters Parameters, float2 TextureSize)
+	{
+		return 0.0f;
+	}
+#endif // DECAL_PRIMITIVE && NUM_MATERIAL_TEXCOORDS
+
+#if DECAL_PRIMITIVE
+	/*
+	 * Deferred decal don't have a Primitive uniform buffer, because we don't know on which primitive the decal
+	 * is being projected to. But the user may still need to get the decal's actor world position.
+	 * So instead of setting up a primitive buffer that may cost to much CPU effort to be almost never used,
+	 * we directly fetch this value from the DeferredDecal.usf specific uniform variable DecalToWorld.
+	 */
+	float3 GetActorWorldPosition()
+	{
+		return DecalToWorld[3].xyz;
+	}
+#else
+	float3 GetActorWorldPosition()
+	{
+		return Primitive.ActorWorldPosition;
+	}
+#endif // DECAL_PRIMITIVE
+
+#if DECAL_PRIMITIVE
+	float DecalLifetimeOpacity()
+	{
+		return DecalParams.y;
+	}
+#else
+	float DecalLifetimeOpacity()
+	{
+		return 0.0f;
+	}
+#endif // DECAL_PRIMITIVE
+
+/** Transforms a vector from tangent space to world space, prescaling by an amount calculated previously */
+MaterialFloat3 TransformTangentVectorToWorld_PreScaled(FMaterialTessellationParameters Parameters, MaterialFloat3 InTangentVector)
+{
+#if FEATURE_LEVEL >= FEATURE_LEVEL_SM5 
+	// used optionally to scale up the vector prior to conversion
+	InTangentVector *= abs( Parameters.TangentToWorldPreScale );	
+
+	// Transform directly to world space
+	// The vector transform is optimized for this case, only one vector-matrix multiply is needed
+	return mul(InTangentVector, Parameters.TangentToWorld);
+#else
+	return TransformTangentVectorToWorld(Parameters.TangentToWorld, InTangentVector);
+#endif // #if FEATURE_LEVEL_SM5
+}
+
+/** Transforms a vector from tangent space to view space */
+MaterialFloat3 TransformTangentVectorToView(FMaterialPixelParameters Parameters, MaterialFloat3 InTangentVector)
+{
+	// Transform from tangent to world, and then to view space
+	return mul(mul(InTangentVector, Parameters.TangentToWorld), (MaterialFloat3x3)ResolvedView.TranslatedWorldToView);
+}
+
+/** Transforms a vector from local space to world space (VS version) */
+MaterialFloat3 TransformLocalVectorToWorld(FMaterialVertexParameters Parameters,MaterialFloat3 InLocalVector)
+{
+	#if USE_INSTANCING || IS_MESHPARTICLE_FACTORY
+		return mul(InLocalVector, (MaterialFloat3x3)Parameters.InstanceLocalToWorld);
+	#else
+		return mul(InLocalVector, GetLocalToWorld3x3());
+	#endif
+}
+
+/** Transforms a vector from local space to world space (PS version) */
+MaterialFloat3 TransformLocalVectorToWorld(FMaterialPixelParameters Parameters,MaterialFloat3 InLocalVector)
+{
+	return mul(InLocalVector, GetLocalToWorld3x3());
+}
+
+#if HAS_PRIMITIVE_UNIFORM_BUFFER
+
+/** Transforms a vector from world space to local space */
+MaterialFloat3 TransformWorldVectorToLocal(MaterialFloat3 InWorldVector)
+{
+	return mul(InWorldVector, (MaterialFloat3x3)Primitive.WorldToLocal);
+}
+
+/** Transforms a position from local space to absolute world space */
+float3 TransformLocalPositionToWorld(FMaterialPixelParameters Parameters,float3 InLocalPosition)
+{
+	return mul(float4(InLocalPosition, 1), Primitive.LocalToWorld).xyz;
+}
+
+/** Transforms a position from local space to absolute world space */
+float3 TransformLocalPositionToWorld(FMaterialVertexParameters Parameters,float3 InLocalPosition)
+{
+	#if USE_INSTANCING || IS_MESHPARTICLE_FACTORY
+		return mul(float4(InLocalPosition, 1), Parameters.InstanceLocalToWorld).xyz;
+	#else
+		return mul(float4(InLocalPosition, 1), Primitive.LocalToWorld).xyz;
+	#endif
+}
+
+#endif
+
+#if HAS_PRIMITIVE_UNIFORM_BUFFER
+
+/** Return the object's position in world space */
+float3 GetObjectWorldPosition(FMaterialPixelParameters Parameters)
+{
+	return Primitive.ObjectWorldPositionAndRadius.xyz;
+}
+
+float3 GetObjectWorldPosition(FMaterialTessellationParameters Parameters)
+{
+	return Primitive.ObjectWorldPositionAndRadius.xyz;
+}
+
+/** Return the object's position in world space. For instanced meshes, this returns the instance position. */
+float3 GetObjectWorldPosition(FMaterialVertexParameters Parameters)
+{
+	#if USE_INSTANCING || IS_MESHPARTICLE_FACTORY
+		return Parameters.InstanceLocalToWorld[3].xyz;
+	#else
+		return Primitive.ObjectWorldPositionAndRadius.xyz;
+	#endif
+}
+
+#endif
+
+/** Get the per-instance random value when instancing */
+float GetPerInstanceRandom(FMaterialVertexParameters Parameters)
+{
+#if USE_INSTANCING
+	return Parameters.PerInstanceParams.x;
+#else
+	return 0.0;
+#endif
+}
+
+/** Get the per-instance random value when instancing */
+float GetPerInstanceRandom(FMaterialPixelParameters Parameters)
+{
+#if USE_INSTANCING
+	return Parameters.PerInstanceParams.x;
+#else
+	return 0.0;
+#endif
+}
+
+/** Get the per-instance fade-out amount when instancing */
+float GetPerInstanceFadeAmount(FMaterialPixelParameters Parameters)
+{
+#if USE_INSTANCING
+	return float(Parameters.PerInstanceParams.y);
+#else
+	return float(1.0);
+#endif
+}
+
+/** Get the per-instance fade-out amount when instancing */
+float GetPerInstanceFadeAmount(FMaterialVertexParameters Parameters)
+{
+#if USE_INSTANCING
+	return float(Parameters.PerInstanceParams.y);
+#else
+	return float(1.0);
+#endif
+}
+ 
+MaterialFloat GetDistanceCullFade()
+{
+	return saturate( View.RealTime * PrimitiveFade.FadeTimeScaleBias.x + PrimitiveFade.FadeTimeScaleBias.y );
+}
+
+/** Rotates Position about the given axis by the given angle, in radians, and returns the offset to Position. */
+float3 RotateAboutAxis(float4 NormalizedRotationAxisAndAngle, float3 PositionOnAxis, float3 Position)
+{
+	// Project Position onto the rotation axis and find the closest point on the axis to Position
+	float3 ClosestPointOnAxis = PositionOnAxis + NormalizedRotationAxisAndAngle.xyz * dot(NormalizedRotationAxisAndAngle.xyz, Position - PositionOnAxis);
+	// Construct orthogonal axes in the plane of the rotation
+	float3 UAxis = Position - ClosestPointOnAxis;
+	float3 VAxis = cross(NormalizedRotationAxisAndAngle.xyz, UAxis);
+	float CosAngle;
+	float SinAngle;
+	sincos(NormalizedRotationAxisAndAngle.w, SinAngle, CosAngle);
+	// Rotate using the orthogonal axes
+	float3 R = UAxis * CosAngle + VAxis * SinAngle;
+	// Reconstruct the rotated world space position
+	float3 RotatedPosition = ClosestPointOnAxis + R;
+	// Convert from position to a position offset
+	return RotatedPosition - Position;
+}
+
+// Material Expression function
+float MaterialExpressionDepthOfFieldFunction(float SceneDepth, int FunctionValueIndex)
+{
+	// tryed switch() but seems that doesn't work
+
+	if(FunctionValueIndex == 0) // TDOF_NearAndFarMask
+	{
+		return CalcUnfocusedPercentCustomBound(SceneDepth, 1, 1);
+	}
+	else if(FunctionValueIndex == 1) // TDOF_Near
+	{
+		return CalcUnfocusedPercentCustomBound(SceneDepth, 1, 0);
+	}
+	else if(FunctionValueIndex == 2) // TDOF_Far
+	{
+		return CalcUnfocusedPercentCustomBound(SceneDepth, 0, 1);
+	}
+	else if(FunctionValueIndex == 3) // TDOF_CircleOfConfusionRadius
+	{
+		// * 2 to compensate for half res
+		return DepthToCoc(SceneDepth) * 2.0f;
+	}
+	return 0;
+}
+
+// TODO convert to LUT
+float3 MaterialExpressionBlackBody( float Temp )
+{
+	float u = ( 0.860117757f + 1.54118254e-4f * Temp + 1.28641212e-7f * Temp*Temp ) / ( 1.0f + 8.42420235e-4f * Temp + 7.08145163e-7f * Temp*Temp );
+	float v = ( 0.317398726f + 4.22806245e-5f * Temp + 4.20481691e-8f * Temp*Temp ) / ( 1.0f - 2.89741816e-5f * Temp + 1.61456053e-7f * Temp*Temp );
+
+	float x = 3*u / ( 2*u - 8*v + 4 );
+	float y = 2*v / ( 2*u - 8*v + 4 );
+	float z = 1 - x - y;
+
+	float Y = 1;
+	float X = Y/y * x;
+	float Z = Y/y * z;
+
+	float3x3 XYZtoRGB =
+	{
+		 3.2404542, -1.5371385, -0.4985314,
+		-0.9692660,  1.8760108,  0.0415560,
+		 0.0556434, -0.2040259,  1.0572252,
+	};
+
+	return mul( XYZtoRGB, float3( X, Y, Z ) ) * pow( 0.0004 * Temp, 4 );
+}
+
+float4 MaterialExpressionAtmosphericFog(FMaterialPixelParameters Parameters, float3 AbsoluteWorldPosition)
+{
+#if MATERIAL_ATMOSPHERIC_FOG
+	// WorldPosition default value is Parameters.AbsoluteWorldPosition if not overridden by the user
+	float3 ViewVector = AbsoluteWorldPosition - ResolvedView.WorldCameraOrigin;
+	float SceneDepth = length(ViewVector);
+	return GetAtmosphericFog(ResolvedView.WorldCameraOrigin, ViewVector, SceneDepth, float3(0.f, 0.f, 0.f));
+#else
+	return float4(0.f, 0.f, 0.f, 0.f);
+#endif
+}
+
+float3 MaterialExpressionAtmosphericLightVector(FMaterialPixelParameters Parameters)
+{
+#if MATERIAL_ATMOSPHERIC_FOG
+	return View.AtmosphericFogSunDirection;
+#else
+	return float3(0.f, 0.f, 0.f);
+#endif
+}
+
+float3 MaterialExpressionAtmosphericLightColor(FMaterialPixelParameters Parameters)
+{
+#if MATERIAL_ATMOSPHERIC_FOG
+	return View.AtmosphericFogSunColor;
+#else
+	return float3(0.f, 0.f, 0.f);
+#endif
+}
+
+/**
+ * Utility function to unmirror one coordinate value to the other side
+ * UnMirrored == 1 if normal
+ * UnMirrored == -1 if mirrored
+ *
+ * Used by most of parameter functions generated via code in this file
+ */
+MaterialFloat UnMirror( MaterialFloat Coordinate, FMaterialPixelParameters Parameters )
+{
+	return ((Coordinate)*(Parameters.UnMirrored)*0.5+0.5);
+}
+
+/**
+ * UnMirror only U
+ */
+MaterialFloat2 UnMirrorU( MaterialFloat2 UV, FMaterialPixelParameters Parameters )
+{
+	return MaterialFloat2(UnMirror(UV.x, Parameters), UV.y);
+}
+
+/**
+ * UnMirror only V
+ */
+MaterialFloat2 UnMirrorV( MaterialFloat2 UV, FMaterialPixelParameters Parameters )
+{
+	return MaterialFloat2(UV.x, UnMirror(UV.y, Parameters));
+}
+
+/**
+ * UnMirror only UV
+ */
+MaterialFloat2 UnMirrorUV( MaterialFloat2 UV, FMaterialPixelParameters Parameters )
+{
+	return MaterialFloat2(UnMirror(UV.x, Parameters), UnMirror(UV.y, Parameters));
+}
+
+/** 
+ * Transforms screen space positions into UVs with [.5, .5] centered on ObjectPostProjectionPosition,
+ * And [1, 1] at ObjectPostProjectionPosition + (ObjectRadius, ObjectRadius).
+ */
+MaterialFloat2 GetParticleMacroUV(FMaterialPixelParameters Parameters)
+{
+	return (Parameters.ScreenPosition.xy / Parameters.ScreenPosition.w - Parameters.Particle.MacroUV.xy) * Parameters.Particle.MacroUV.zw + MaterialFloat2(.5, .5);
+}
+
+MaterialFloat4 ProcessMaterialColorTextureLookup(MaterialFloat4 TextureValue)
+{
+#if (FEATURE_LEVEL == FEATURE_LEVEL_ES2) // ES2 does not support sRGB sampling
+	#if MOBILE_EMULATION
+	if( View.MobilePreviewMode > 0.5f)
+	{
+		// undo HW srgb->lin
+		TextureValue.rgb = pow(TextureValue.rgb, 1.0f / 2.2f); // TODO: replace with a more accurate lin -> sRGB conversion.
+	}
+	#endif
+	// sRGB read approximation (in highp if possible)
+	float3 LinRGB = TextureValue.rgb;
+	LinRGB *= LinRGB; 
+	return MaterialFloat4(LinRGB, TextureValue.a);
+#endif 
+	return TextureValue;
+}
+
+MaterialFloat4 ProcessMaterialLinearColorTextureLookup(MaterialFloat4 TextureValue)
+{
+	return TextureValue;
+}
+
+MaterialFloat ProcessMaterialGreyscaleTextureLookup(MaterialFloat TextureValue)
+{
+#if (FEATURE_LEVEL == FEATURE_LEVEL_ES2) || COMPILER_GLSL_ES3_1 // ES2 and OpenGLES3.1 does not support sRGB sampling from R8
+	#if MOBILE_EMULATION
+	if( View.MobilePreviewMode > 0.5f )
+	{
+		// undo HW srgb->lin
+		TextureValue = pow(TextureValue, 1.0f/2.2f); // TODO: replace with a more accurate lin -> sRGB conversion.
+	}
+	#endif
+	// sRGB read approximation (in highp if possible)
+	float LinValue = TextureValue;
+	LinValue *= LinValue;
+	return MaterialFloat(LinValue);
+#endif 
+	return TextureValue;
+}
+
+MaterialFloat ProcessMaterialLinearGreyscaleTextureLookup(MaterialFloat TextureValue)
+{
+	return TextureValue;
+}
+
+/** Accesses a shared material sampler or falls back if independent samplers are not supported. */
+SamplerState GetMaterialSharedSampler(SamplerState TextureSampler, SamplerState SharedSampler)
+{
+#if SUPPORTS_INDEPENDENT_SAMPLERS
+	return SharedSampler;
+#else
+	// Note: to match behavior on platforms that don't support SUPPORTS_INDEPENDENT_SAMPLERS, 
+	// TextureSampler should have been set to the same sampler.  This is not currently done.
+	return TextureSampler;
+#endif
+}
+
+/** Calculate a reflection vector about the specified world space normal. Optionally normalize this normal **/
+MaterialFloat3 ReflectionAboutCustomWorldNormal(FMaterialPixelParameters Parameters, MaterialFloat3 WorldNormal, bool bNormalizeInputNormal)
+{
+	if (bNormalizeInputNormal)
+	{
+		WorldNormal = normalize(WorldNormal);
+	}
+
+	return -Parameters.CameraVector + WorldNormal * dot(WorldNormal, Parameters.CameraVector) * 2.0;
+}
+
+#ifndef SPHERICAL_OPACITY_FOR_SHADOW_DEPTHS
+#define SPHERICAL_OPACITY_FOR_SHADOW_DEPTHS 0
+#endif
+
+/** 
+ * Calculates opacity for a billboard particle as if it were a sphere. 
+ * Note: Calling this function requires the vertex factory to have been compiled with SPHERICAL_PARTICLE_OPACITY set to 1
+ */
+float GetSphericalParticleOpacity(FMaterialPixelParameters Parameters, float Density)
+{
+	float Opacity = 0;
+
+#if PARTICLE_FACTORY || HAS_PRIMITIVE_UNIFORM_BUFFER
+
+#if PARTICLE_FACTORY
+
+	float3 ParticleTranslatedWorldPosition = Parameters.Particle.TranslatedWorldPositionAndSize.xyz;
+	float ParticleRadius = max(0.000001f, Parameters.Particle.TranslatedWorldPositionAndSize.w);
+
+#elif HAS_PRIMITIVE_UNIFORM_BUFFER
+
+	// Substitute object attributes if the mesh is not a particle
+	// This is mostly useful for previewing materials using spherical opacity in the material editor
+	float3 ParticleTranslatedWorldPosition = Primitive.ObjectWorldPositionAndRadius.xyz + ResolvedView.PreViewTranslation.xyz;
+	float ParticleRadius = max(0.000001f, Primitive.ObjectWorldPositionAndRadius.w);
+
+#endif
+
+	// Rescale density to make the final opacity independent of the particle radius
+	float RescaledDensity = Density / ParticleRadius;
+
+	// Distance from point being shaded to particle center
+	float DistanceToParticle = length(Parameters.WorldPosition_NoOffsets_CamRelative - ParticleTranslatedWorldPosition);
+
+	FLATTEN
+	if (DistanceToParticle < ParticleRadius) 
+	{
+		// Distance from point being shaded to the point on the sphere along the view direction
+		float HemisphericalDistance = sqrt(ParticleRadius * ParticleRadius - DistanceToParticle * DistanceToParticle);
+
+#if SPHERICAL_OPACITY_FOR_SHADOW_DEPTHS
+		// When rendering shadow depths we can't use scene depth or the near plane, just use the distance through the whole sphere
+		float DistanceThroughSphere = HemisphericalDistance * 2;
+#else
+		// Initialize near and far sphere intersection distances
+		float NearDistance = Parameters.ScreenPosition.w - HemisphericalDistance;
+		float FarDistance = Parameters.ScreenPosition.w + HemisphericalDistance;
+
+		float SceneDepth = CalcSceneDepth(SvPositionToBufferUV(Parameters.SvPosition));
+		FarDistance = min(SceneDepth, FarDistance);
+
+		// Take into account opaque objects intersecting the sphere
+		float DistanceThroughSphere = FarDistance - NearDistance;
+#endif
+
+		// Use the approximation for the extinction line integral from "Spherical Billboards and their Application to Rendering Explosions"
+		Opacity = saturate(1 - exp2(-RescaledDensity * (1 - DistanceToParticle / ParticleRadius) * DistanceThroughSphere));
+
+#if !SPHERICAL_OPACITY_FOR_SHADOW_DEPTHS
+		// Fade out as the particle approaches the near plane
+		Opacity = lerp(0, Opacity, saturate((Parameters.ScreenPosition.w - ParticleRadius - View.NearPlane) / ParticleRadius));
+#endif
+	}
+
+#endif
+
+	return Opacity;
+}
+
+#if USES_SPEEDTREE
+
+/** Vertex offset for SpeedTree wind and LOD */
+float3 GetSpeedTreeVertexOffsetInner(FMaterialVertexParameters Parameters, int GeometryType, int WindType, int LODType, float BillboardThreshold, FSpeedTreeData STData) 
+{
+	#if (NUM_MATERIAL_TEXCOORDS_VERTEX < 6) || IS_MESHPARTICLE_FACTORY
+		return float4(0,0,0);
+	#endif
+
+	#if USE_INSTANCING
+		float3x3 LocalToWorld = (float3x3)Parameters.InstanceLocalToWorld;
+		float3 LocalPosition = Parameters.InstanceLocalPosition;
+
+		// skip if this instance is hidden
+		if (Parameters.PerInstanceParams.z < 1.f)
+		{
+			return float3(0,0,0);
+		}
+	#else
+		float3x3 LocalToWorld = (float3x3)Primitive.LocalToWorld;
+		float3 LocalPosition = mul(float4(GetWorldPosition(Parameters), 1), Primitive.WorldToLocal).xyz;
+	#endif
+
+	float3 TreePos = GetObjectWorldPosition(Parameters);
+
+	// compute LOD by finding screen space size
+	float LodInterp = 1.0;
+#if !USE_INSTANCING || !USE_DITHERED_LOD_TRANSITION
+	if (LODType == SPEEDTREE_LOD_TYPE_SMOOTH) 
+	{
+		const float ScreenMultiple = 0.5 * max(View.ViewSizeAndInvSize.x * ResolvedView.ViewToClip[0][0],
+												View.ViewSizeAndInvSize.y * ResolvedView.ViewToClip[1][1]);
+		const float ScreenRadius = ScreenMultiple * Primitive.ObjectWorldPositionAndRadius.w / 
+									max(1.0, dot(TreePos - ResolvedView.WorldCameraOrigin, ResolvedView.ViewForward.xyz));
+		const float ScreenArea = PI * ScreenRadius * ScreenRadius * View.ViewSizeAndInvSize.z * View.ViewSizeAndInvSize.w;
+		LodInterp = saturate((ScreenArea - SpeedTreeLODInfo.x) / SpeedTreeLODInfo.z);
+	}
+#endif
+	TreePos *= 0.001; // The only other use of the tree position is as an offset into trig functions, but big numbers don't play nice there
+
+	// SpeedTrees should only be uniformly scaled, but if necessary, it takes a few more instructions
+	float TreeScale = length(mul(float3(0,0,1), LocalToWorld));
+					//float3(length((float3)LocalToWorld[0]),
+					//		length((float3)LocalToWorld[1]),
+					//		length((float3)LocalToWorld[2]));
+
+
+	// @todo There is probably a more optimal way to get the rotated (but not translated or scaled) vertex position needed for correct wind
+	float3 OriginalPosition = LocalPosition;
+	OriginalPosition = mul(OriginalPosition, LocalToWorld) / TreeScale;
+
+	float3 FinalPosition = OriginalPosition;
+	
+	if (GeometryType == SPEEDTREE_GEOMETRY_TYPE_BILLBOARD)
+	{
+		if (BillboardThreshold < 1.0)
+		{
+			// billboard meshes can have triangles drop out if they aren't facing the camera
+			// this rotates the view direction around so we ignore the local Z component
+			float3 LocalView2D = normalize(float3(ResolvedView.ViewForward.xy, 0));
+			float3 LocalNormal2D = normalize(float3(Parameters.TangentToWorld[2].xy, 0));
+			if (dot(LocalView2D, LocalNormal2D) > (-1.0 + BillboardThreshold * 0.25))
+			{
+				FinalPosition = float3(0,0,0);
+			}
+		}
+	}
+	else
+	{
+		// rotated normal needed in a few places
+		float3 Normal = Parameters.TangentToWorld[2];
+
+		// branches and fronds
+		if (GeometryType == SPEEDTREE_GEOMETRY_TYPE_BRANCH || GeometryType == SPEEDTREE_GEOMETRY_TYPE_FROND) 
+		{
+			// smooth LOD
+			#if !USE_INSTANCING
+				if (LODType == SPEEDTREE_LOD_TYPE_SMOOTH) 
+				{
+					float3 LODPos = float3(Parameters.TexCoords[3].x, Parameters.TexCoords[3].y, Parameters.TexCoords[4].x);
+					LODPos = mul(LODPos, LocalToWorld) / TreeScale;
+					FinalPosition = lerp(LODPos, FinalPosition, LodInterp);
+				}
+			#endif
+
+			// frond wind, if needed
+			if (GeometryType == SPEEDTREE_GEOMETRY_TYPE_FROND && WindType == SPEEDTREE_WIND_TYPE_PALM)
+			{
+				float2 TexCoords = Parameters.TexCoords[0];
+				float4 WindExtra = float4(Parameters.TexCoords[5].x, Parameters.TexCoords[5].y, Parameters.TexCoords[6].x, 0.0);
+				FinalPosition = RippleFrond(STData, FinalPosition, Normal, TexCoords.x, TexCoords.y, WindExtra.x, WindExtra.y, WindExtra.z);
+			}
+		}
+
+		// leaves and facing leaves
+		if (GeometryType == SPEEDTREE_GEOMETRY_TYPE_FACINGLEAF || 
+				(GeometryType == SPEEDTREE_GEOMETRY_TYPE_LEAF && 
+				(LODType == SPEEDTREE_LOD_TYPE_SMOOTH || (WindType > SPEEDTREE_WIND_TYPE_FASTEST && WindType != SPEEDTREE_WIND_TYPE_PALM))))
+		{
+			// remove anchor pos from vertex position
+			float3 Anchor = float3(Parameters.TexCoords[4].y, Parameters.TexCoords[5].x, Parameters.TexCoords[5].y);
+			Anchor = (mul(Anchor, LocalToWorld)) / TreeScale;
+			FinalPosition -= Anchor;
+
+			// smooth LOD
+			#if !USE_INSTANCING
+				if (LODType == SPEEDTREE_LOD_TYPE_SMOOTH) 
+				{
+					if (GeometryType == SPEEDTREE_GEOMETRY_TYPE_LEAF)
+					{
+						float3 LODPos = float3(Parameters.TexCoords[3].x, Parameters.TexCoords[3].y, Parameters.TexCoords[4].x);
+						LODPos = mul(LODPos, LocalToWorld) / TreeScale - Anchor;
+						FinalPosition = lerp(LODPos, FinalPosition, LodInterp);
+					}
+					else
+					{
+						float LODScalar = Parameters.TexCoords[3].x;
+						FinalPosition *= lerp(LODScalar, 1.0, LodInterp);
+					}
+				}
+			#endif
+
+			// face camera-facing leaves to the camera, if needed
+			if (GeometryType == SPEEDTREE_GEOMETRY_TYPE_FACINGLEAF) 
+			{
+				// have to rotate the view into local space
+				FinalPosition = FinalPosition.x * ResolvedView.ViewRight + 
+								FinalPosition.y * ResolvedView.ViewUp + 
+								FinalPosition.z * ResolvedView.ViewForward;
+			}
+
+			// leaf wind
+			if (WindType > SPEEDTREE_WIND_TYPE_FASTEST && WindType != SPEEDTREE_WIND_TYPE_PALM) 
+			{
+				float4 WindExtra = float4(Parameters.TexCoords[6].x, Parameters.TexCoords[6].y, Parameters.TexCoords[7].x, Parameters.TexCoords[7].y);
+				float LeafWindTrigOffset = Anchor.x + Anchor.y;
+				FinalPosition = LeafWind(STData, WindExtra.w > 0.0, FinalPosition, Normal, WindExtra.x, float3(0,0,0), WindExtra.y, WindExtra.z, LeafWindTrigOffset, WindType);
+			}
+				
+			// move leaf back to anchor
+			FinalPosition += Anchor;
+		}
+
+		if (WindType > SPEEDTREE_WIND_TYPE_FAST)
+		{
+			// branch wind (applies to all geometry)
+			float2 VertBranchWind = Parameters.TexCoords[2];
+			FinalPosition = BranchWind(STData, FinalPosition, TreePos, float4(VertBranchWind, 0, 0), WindType);
+		}	
+	}
+
+	// global wind can apply to the whole tree, even billboards
+	if (WindType != SPEEDTREE_WIND_TYPE_NONE)
+	{
+		FinalPosition = GlobalWind(STData, FinalPosition, TreePos, true);
+	}
+
+	// convert into a world space offset
+	return (FinalPosition - OriginalPosition) * TreeScale;
+}
+
+/** Vertex offset for SpeedTree wind and LOD */
+float3 GetSpeedTreeVertexOffset(FMaterialVertexParameters Parameters, int GeometryType, int WindType, int LODType, float BillboardThreshold, bool bUsePreviousFrame) 
+{
+#if VF_SUPPORTS_SPEEDTREE_WIND
+	if (bUsePreviousFrame)
+	{
+		return GetSpeedTreeVertexOffsetInner(Parameters, GeometryType, WindType, LODType, BillboardThreshold, GetPreviousSpeedTreeData());
+	}
+	return GetSpeedTreeVertexOffsetInner(Parameters, GeometryType, WindType, LODType, BillboardThreshold, GetCurrentSpeedTreeData());
+#else
+	return 0;
+#endif
+}
+
+#endif
+
+MaterialFloat2 GetLightmapUVs(FMaterialPixelParameters Parameters)
+{
+#if LIGHTMAP_UV_ACCESS
+	return Parameters.LightmapUVs;
+#else
+	return MaterialFloat2(0,0);
+#endif
+}
+
+#if USES_EYE_ADAPTATION
+#include "/Engine/Private/EyeAdaptationCommon.ush"
+#endif
+
+#if POST_PROCESS_MATERIAL
+// only Postprocess material bind PostprocessInput0, PostprocessInput0Size, PostprocessInput0MinMax ...
+#include "/Engine/Private/PostProcessCommon.ush"			// PostprocessInput0
+#endif
+
+#if NEEDS_SCENE_TEXTURES
+#include "/Engine/Private/DeferredShadingCommon.ush"		// GetGBufferData()
+
+float4 GetPostProcessInputSize(uint Index)
+{
+#if POST_PROCESS_MATERIAL
+		 if (Index == 0)	return PostprocessInput0Size;
+	else if (Index == 1)	return PostprocessInput1Size;
+	else if (Index == 2)	return PostprocessInput2Size;
+	else if (Index == 3)	return PostprocessInput3Size;
+	else if (Index == 4)	return PostprocessInput4Size;
+	else if (Index == 5)	return PostprocessInput5Size;
+	else if (Index == 6)	return PostprocessInput6Size;
+#endif
+	return 1.0f;
+}
+
+float4 GetPostProcessInputMinMax(uint Index)
+{
+#if POST_PROCESS_MATERIAL
+		 if (Index == 0)	return PostprocessInput0MinMax;
+	else if (Index == 1)	return PostprocessInput1MinMax;
+	else if (Index == 2)	return PostprocessInput2MinMax;
+	else if (Index == 3)	return PostprocessInput3MinMax;
+	else if (Index == 4)	return PostprocessInput4MinMax;
+	else if (Index == 5)	return PostprocessInput5MinMax;
+	else if (Index == 6)	return PostprocessInput6MinMax;
+#endif
+	return 1.0f;
+}
+
+MaterialFloat4 MobileSceneTextureLookup(inout FMaterialPixelParameters Parameters, int SceneTextureId, float2 UV)
+{
+#if (FEATURE_LEVEL <= FEATURE_LEVEL_ES3_1)
+	
+	// For fetching SceneDepth and CustomDepth after tonemapper we need to switch OpenGL coordinates
+	// SceneColor will be transformed to OpenGL coordinates during tonemapping
+	#define NEEDTOSWITCHVERTICLEAXIS (COMPILER_GLSL_ES2 || COMPILER_GLSL_ES3_1) && (POST_PROCESS_MATERIAL && !POST_PROCESS_MATERIAL_BEFORE_TONEMAP)
+	
+	//PPI_SceneDepth
+	if (SceneTextureId == 1)
+	{
+		#if NEEDTOSWITCHVERTICLEAXIS
+		UV.y = 1.0 - UV.y; 
+		#endif
+		
+		MaterialFloat Depth = ConvertFromDeviceZ(Texture2DSample(SceneDepthTexture, SceneDepthTextureSampler, UV).r);
+		return MaterialFloat4(Depth.rrr, 0.f);
+	}
+	//PPI_CustomDepth
+	else if (SceneTextureId == 13)
+	{
+		#if NEEDTOSWITCHVERTICLEAXIS
+		UV.y = 1.0 - UV.y; 
+		#endif
+
+		MaterialFloat Depth = ConvertFromDeviceZ(Texture2DSample(CustomDepthTexture, CustomDepthTextureSampler, UV).r);
+		return MaterialFloat4(Depth.rrr, 0.f);
+	}
+	// PPI_PostprocessInput
+	else if (SceneTextureId == 14)
+	{
+#if POST_PROCESS_MATERIAL	
+		MaterialFloat4 Input0 = Texture2DSample(PostprocessInput0, PostprocessInput0Sampler, UV);
+		#if POST_PROCESS_MATERIAL_BEFORE_TONEMAP
+		Input0 = Decode32BPPHDR(Input0);
+		#endif
+		// We need to preserve original SceneColor Alpha as it's used by tonemaper on mobile
+		Parameters.BackupSceneColorAlpha = Input0.a;
+		return Input0;
+#endif// POST_PROCESS_MATERIAL	
+	}
+	// PPI_CustomStencil
+	else if (SceneTextureId == 24)
+	{
+		#if NEEDTOSWITCHVERTICLEAXIS
+		UV.y = 1.0 - UV.y; 
+		#endif
+
+		MaterialFloat Stencil = Texture2DSample(MobileCustomStencilTexture, MobileCustomStencilTextureSampler, UV).r*255.0;
+		Stencil = floor(Stencil + 0.5);
+		return MaterialFloat4(Stencil.rrr, 0.f);
+	}
+#endif// FEATURE_LEVEL
+
+	return MaterialFloat4(0.0f, 0.0f, 0.0f, 0.0f);
+}
+
+/** Applies an offset to the scene texture lookup and decodes the HDR linear space color. */
+float4 SceneTextureLookup(float2 UV, int SceneTextureIndex, bool bFiltered)
+{
+#if SCENE_TEXTURES_DISABLED
+	return float4(0.0f, 0.0f, 0.0f, 0.0f);
+#endif
+
+	FScreenSpaceData ScreenSpaceData = GetScreenSpaceData(UV, false);
+	switch(SceneTextureIndex)
+	{
+		// order needs to match to ESceneTextureId
+
+		// PPI_SceneColor
+		case 0: return float4(CalcSceneColor(UV), 0);
+		// PPI_SceneDepth
+		case 1: return ScreenSpaceData.GBuffer.Depth;
+		// PPI_DiffuseColor
+		case 2: return float4(ScreenSpaceData.GBuffer.DiffuseColor, 0);
+		// PPI_SpecularColor
+		case 3: return float4(ScreenSpaceData.GBuffer.SpecularColor, 0);
+		// PPI_SubsurfaceColor
+		case 4: return IsSubsurfaceModel(ScreenSpaceData.GBuffer.ShadingModelID) ? float4( ExtractSubsurfaceColor(ScreenSpaceData.GBuffer), ScreenSpaceData.GBuffer.CustomData.a ) : ScreenSpaceData.GBuffer.CustomData;
+		// PPI_BaseColor
+		case 5: return float4(ScreenSpaceData.GBuffer.BaseColor, 0);
+		// PPI_Specular
+		case 6: return ScreenSpaceData.GBuffer.Specular;
+		// PPI_Metallic
+		case 7: return ScreenSpaceData.GBuffer.Metallic;
+		// PPI_WorldNormal
+		case 8: return float4(ScreenSpaceData.GBuffer.WorldNormal, 0);
+		// PPI_SeparateTranslucency
+		case 9: return float4(1, 1, 1, 1);	// todo
+		// PPI_Opacity
+		case 10: return ScreenSpaceData.GBuffer.CustomData.a;
+		// PPI_Roughness
+		case 11: return ScreenSpaceData.GBuffer.Roughness;
+		// PPI_MaterialAO
+		case 12: return ScreenSpaceData.GBuffer.GBufferAO;
+		// PPI_CustomDepth
+		case 13: return ScreenSpaceData.GBuffer.CustomDepth;
+#if POST_PROCESS_MATERIAL
+		// PPI_PostprocessInput0
+		case 14: return Texture2DSample(PostprocessInput0, bFiltered ? BilinearTextureSampler0 : PostprocessInput0Sampler, UV);
+		// PPI_PostprocessInput1
+		case 15: return Texture2DSample(PostprocessInput1, bFiltered ? BilinearTextureSampler0 : PostprocessInput1Sampler, UV);
+		// PPI_PostprocessInput2
+		case 16: return Texture2DSample(PostprocessInput2, bFiltered ? BilinearTextureSampler0 : PostprocessInput2Sampler, UV);
+		// PPI_PostprocessInput3
+		case 17: return Texture2DSample(PostprocessInput3, bFiltered ? BilinearTextureSampler0 : PostprocessInput3Sampler, UV);
+		// PPI_PostprocessInput4
+		case 18: return Texture2DSample(PostprocessInput4, bFiltered ? BilinearTextureSampler0 : PostprocessInput4Sampler, UV);
+		// PPI_PostprocessInput5
+		case 19: return Texture2DSample(PostprocessInput5, bFiltered ? BilinearTextureSampler0 : PostprocessInput5Sampler, UV);
+		// PPI_PostprocessInput6
+		case 20: return Texture2DSample(PostprocessInput6, bFiltered ? BilinearTextureSampler0 : PostprocessInput6Sampler, UV);
+#endif // __POST_PROCESS_COMMON__
+		// PPI_DecalMask
+		case 21: return 0;  // material compiler will return an error
+		// PPI_ShadingModel
+		case 22: return float4(GetShadingModelColor(ScreenSpaceData.GBuffer.ShadingModelID), 1);
+		// PPI_AmbientOcclusion
+		case 23: return ScreenSpaceData.AmbientOcclusion;
+		// PPI_CustomStencil
+		case 24: return ScreenSpaceData.GBuffer.CustomStencil;
+		// PPI_StoredBaseColor
+		case 25: return float4(ScreenSpaceData.GBuffer.StoredBaseColor, 0);
+		// PPI_StoredSpecular
+		case 26: return float4(ScreenSpaceData.GBuffer.StoredSpecular.rrr, 0);
+		default:
+			return float4(0, 0, 0, 0);
+	}
+}
+#endif // NEEDS_SCENE_TEXTURES
+
+// Uniform material expressions.
+%s
+
+// can return in tangent space or world space (use MATERIAL_TANGENTSPACENORMAL)
+half3 GetMaterialNormalRaw(FPixelMaterialInputs PixelMaterialInputs)
+{
+	return PixelMaterialInputs.Normal;
+}
+
+half3 GetMaterialNormal(FMaterialPixelParameters Parameters, FPixelMaterialInputs PixelMaterialInputs)
+{
+	half3 RetNormal;
+
+	RetNormal = GetMaterialNormalRaw(PixelMaterialInputs);
+		
+	#if (USE_EDITOR_SHADERS && !(ES2_PROFILE || ES3_1_PROFILE || ESDEFERRED_PROFILE)) || MOBILE_EMULATION
+	{
+		// this feature is only needed for development/editor - we can compile it out for a shipping build (see r.CompileShadersForDevelopment cvar help)
+		half3 OverrideNormal = View.NormalOverrideParameter.xyz;
+
+		#if !MATERIAL_TANGENTSPACENORMAL
+			OverrideNormal = Parameters.TangentToWorld[2] * (1 - View.NormalOverrideParameter.w);
+		#endif
+
+		RetNormal = RetNormal * View.NormalOverrideParameter.w + OverrideNormal;
+	}
+	#endif
+
+	return RetNormal;
+}
+
+half3 GetMaterialEmissiveRaw(FPixelMaterialInputs PixelMaterialInputs)
+{
+	return PixelMaterialInputs.EmissiveColor;
+}
+
+half3 GetMaterialEmissive(FPixelMaterialInputs PixelMaterialInputs)
+{
+	half3 EmissiveColor = GetMaterialEmissiveRaw(PixelMaterialInputs);
+#if !MATERIAL_ALLOW_NEGATIVE_EMISSIVECOLOR
+	EmissiveColor = max(EmissiveColor, 0.0f);
+#endif
+	return EmissiveColor;
+}
+
+half3 GetMaterialEmissiveForCS(FMaterialPixelParameters Parameters)
+{
+%s;
+}
+
+half3 GetMaterialBaseColorRaw(FPixelMaterialInputs PixelMaterialInputs)
+{
+	return PixelMaterialInputs.BaseColor;
+}
+
+half3 GetMaterialBaseColor(FPixelMaterialInputs PixelMaterialInputs)
+{
+	return saturate(GetMaterialBaseColorRaw(PixelMaterialInputs));
+}
+
+half GetMaterialMetallicRaw(FPixelMaterialInputs PixelMaterialInputs)
+{
+	return PixelMaterialInputs.Metallic;
+}
+
+half GetMaterialMetallic(FPixelMaterialInputs PixelMaterialInputs)
+{
+	return saturate(GetMaterialMetallicRaw(PixelMaterialInputs));
+}
+
+half GetMaterialSpecularRaw(FPixelMaterialInputs PixelMaterialInputs)
+{
+	return PixelMaterialInputs.Specular;
+}
+
+half GetMaterialSpecular(FPixelMaterialInputs PixelMaterialInputs)
+{
+	return saturate(GetMaterialSpecularRaw(PixelMaterialInputs));
+}
+
+half GetMaterialRoughnessRaw(FPixelMaterialInputs PixelMaterialInputs)
+{
+	return PixelMaterialInputs.Roughness;
+}
+
+half GetMaterialRoughness(FPixelMaterialInputs PixelMaterialInputs)
+{
+#if MATERIAL_FULLY_ROUGH
+	return 1;
+#endif
+#if (ES2_PROFILE || ES3_1_PROFILE)
+	// The smallest normalized value that can be represented in IEEE 754 (FP16) is 2^-14 = 6.1e-5.
+	// The code will make the following computation involving roughness: 1.0 / Roughness^4.
+	// Therefore to prevent division by zero on devices that do not support denormals, Roughness^4
+	// must be >= 6.1e-5. We will clamp to 0.09 because 0.09^4 = 6.5e-5.
+	//
+	// Note that we also clamp to 1.0 to match the deferred renderer on PC where the roughness is 
+	// stored in an 8-bit value and thus automatically clamped at 1.0.
+
+	half Roughness = GetMaterialRoughnessRaw(PixelMaterialInputs);
+
+	#if MOBILE_EMULATION
+	{
+		// this is only needed for mobile preview on PC
+		Roughness = Roughness * View.RoughnessOverrideParameter.y + View.RoughnessOverrideParameter.x;
+	}
+	#endif
+
+	// Increase value from 0.09 to 0.12 to fix missing specular lobe problem on device
+	return clamp( Roughness, 0.12, 1.0 );
+#else
+	half Roughness = saturate(GetMaterialRoughnessRaw(PixelMaterialInputs));
+	
+	#if (USE_EDITOR_SHADERS && !ESDEFERRED_PROFILE)
+	{
+		// this feature is only needed for development/editor - we can compile it out for a shipping build (see r.CompileShadersForDevelopment cvar help)
+		Roughness = Roughness * View.RoughnessOverrideParameter.y + View.RoughnessOverrideParameter.x;
+	}
+	#endif
+	
+	return Roughness;
+#endif
+}
+
+half GetMaterialTranslucencyDirectionalLightingIntensity()
+{
+%s;
+}
+
+half GetMaterialTranslucentShadowDensityScale()
+{
+%s;
+}
+
+half GetMaterialTranslucentSelfShadowDensityScale()
+{
+%s;
+}
+
+half GetMaterialTranslucentSelfShadowSecondDensityScale()
+{
+%s;
+}
+
+half GetMaterialTranslucentSelfShadowSecondOpacity()
+{
+%s;
+}
+
+half GetMaterialTranslucentBackscatteringExponent()
+{
+%s;
+}
+
+half3 GetMaterialTranslucentMultipleScatteringExtinction()
+{
+%s;
+}
+
+// This is the clip value constant that is defined in the material (range 0..1)
+// Use GetMaterialMask() to get the Material Mask combined with this.
+half GetMaterialOpacityMaskClipValue()
+{
+%s;
+}
+
+// Should only be used by GetMaterialOpacity(), returns the unmodified value generated from the shader expressions of the opacity input.
+// To compute the opacity depending on the material blending GetMaterialOpacity() should be called instead.
+half GetMaterialOpacityRaw(FPixelMaterialInputs PixelMaterialInputs)
+{
+	return PixelMaterialInputs.Opacity;
+}
+
+#if MATERIALBLENDING_MASKED || (DECAL_BLEND_MODE == DECALBLENDMODEID_VOLUMETRIC)
+// Returns the material mask value generated from the material expressions.
+// Use GetMaterialMask() to get the value altered depending on the material blend mode.
+half GetMaterialMaskInputRaw(FPixelMaterialInputs PixelMaterialInputs)
+{
+	return PixelMaterialInputs.OpacityMask;
+}
+
+// Returns the material mask value generated from the material expressions minus the used defined
+// MaskClip value constant. If this value is <=0 the pixel should be killed.
+half GetMaterialMask(FPixelMaterialInputs PixelMaterialInputs)
+{
+	return GetMaterialMaskInputRaw(PixelMaterialInputs) - GetMaterialOpacityMaskClipValue();
+}
+#endif
+
+// Returns the material opacity depending on the material blend mode.
+half GetMaterialOpacity(FPixelMaterialInputs PixelMaterialInputs)
+{
+	// Clamp to valid range to prevent negative colors from lerping
+	return saturate(GetMaterialOpacityRaw(PixelMaterialInputs));
+}
+
+float3 GetMaterialWorldPositionOffset(FMaterialVertexParameters Parameters)
+{
+	#if USE_INSTANCING
+		// skip if this instance is hidden
+		if (Parameters.PerInstanceParams.z < 1.f)
+		{
+			return float3(0,0,0);
+		}
+	#endif
+%s;
+}
+
+float3 GetMaterialPreviousWorldPositionOffset(FMaterialVertexParameters Parameters)
+{
+	#if USE_INSTANCING
+		// skip if this instance is hidden
+		if (Parameters.PerInstanceParams.z < 1.f)
+		{
+			return float3(0,0,0);
+		}
+	#endif
+%s;
+}
+
+half3 GetMaterialWorldDisplacement(FMaterialTessellationParameters Parameters)
+{
+%s;
+}
+
+half GetMaterialMaxDisplacement()
+{
+%s;
+}
+
+half GetMaterialTessellationMultiplier(FMaterialTessellationParameters Parameters)
+{
+%s;
+}
+
+// .rgb:SubsurfaceColor, .a:SSProfileId in 0..1 range
+half4 GetMaterialSubsurfaceDataRaw(FMaterialPixelParameters Parameters)
+{
+%s;
+}
+
+half4 GetMaterialSubsurfaceData(FMaterialPixelParameters Parameters)
+{
+	half4 OutSubsurface = GetMaterialSubsurfaceDataRaw(Parameters);
+	OutSubsurface.rgb = saturate(OutSubsurface.rgb);
+	return OutSubsurface;
+}
+
+half GetMaterialCustomData0(FMaterialPixelParameters Parameters)
+{
+%s;
+}
+
+half GetMaterialCustomData1(FMaterialPixelParameters Parameters)
+{
+%s;
+}
+
+half GetMaterialAmbientOcclusionRaw(FPixelMaterialInputs PixelMaterialInputs)
+{
+	return PixelMaterialInputs.AmbientOcclusion;
+}
+
+half GetMaterialAmbientOcclusion(FPixelMaterialInputs PixelMaterialInputs)
+{
+	return saturate(GetMaterialAmbientOcclusionRaw(PixelMaterialInputs));
+}
+
+half2 GetMaterialRefraction(FPixelMaterialInputs PixelMaterialInputs)
+{
+	return PixelMaterialInputs.Refraction;
+}
+
+#if NUM_TEX_COORD_INTERPOLATORS
+void GetMaterialCustomizedUVs(FMaterialVertexParameters Parameters, inout float2 OutTexCoords[NUM_TEX_COORD_INTERPOLATORS])
+{
+%s
+}
+
+void GetCustomInterpolators(FMaterialVertexParameters Parameters, inout float2 OutTexCoords[NUM_TEX_COORD_INTERPOLATORS])
+{
+%s
+}
+#endif
+
+float GetMaterialPixelDepthOffset(FPixelMaterialInputs PixelMaterialInputs)
+{
+	return PixelMaterialInputs.PixelDepthOffset;
+}
+
+#if DECAL_PRIMITIVE
+
+float3 TransformTangentNormalToWorld(in FMaterialPixelParameters Parameters, float3 TangentNormal)
+{
+	// To transform the normals use tranpose(Inverse(DecalToWorld)) = transpose(WorldToDecal)
+	// But we want to only rotate the normals (we don't want to non-uniformaly scale them).
+	// We assume the matrix is only a scale and rotation, and we remove non-uniform scale:
+	float3 lengthSqr = { length2(DecalToWorld._m00_m01_m02),
+							length2(DecalToWorld._m10_m11_m12),
+							length2(DecalToWorld._m20_m21_m22) };
+
+	float3 scale = rsqrt(lengthSqr);
+				
+	// Pre-multiply by the inverse of the non-uniform scale in DecalToWorld
+	float4 ScaledNormal = float4(-TangentNormal.z * scale.x, TangentNormal.y * scale.y, TangentNormal.x * scale.z, 0.f);
+
+	// Compute the normal 
+	return normalize(mul(ScaledNormal, DecalToWorld).xyz);
+}
+
+#else //DECAL_PRIMITIVE
+
+float3 TransformTangentNormalToWorld(in FMaterialPixelParameters Parameters, float3 TangentNormal)
+{
+	return normalize(float3(TransformTangentVectorToWorld(Parameters.TangentToWorld, TangentNormal)));
+}
+
+#endif //DECAL_PRIMITIVE
+
+void CalcPixelMaterialInputs(in out FMaterialPixelParameters Parameters, in out FPixelMaterialInputs PixelMaterialInputs)
+{
+	// Initial calculations (required for Normal)
+%s
+	// The Normal is a special case as it might have its own expressions and also be used to calculate other inputs, so perform the assignment here
+%s
+
+	// Note that here MaterialNormal can be in world space or tangent space
+	float3 MaterialNormal = GetMaterialNormal(Parameters, PixelMaterialInputs);
+
+#if MATERIAL_TANGENTSPACENORMAL
+#if SIMPLE_FORWARD_SHADING
+	Parameters.WorldNormal = float3(0, 0, 1);
+#endif
+
+#if FEATURE_LEVEL >= FEATURE_LEVEL_SM4
+	// ES2 will rely on only the final normalize for performance
+	MaterialNormal = normalize(MaterialNormal);
+#endif
+
+	// normalizing after the tangent space to world space conversion improves quality with sheared bases (UV layout to WS causes shrearing)
+	// use full precision normalize to avoid overflows
+	Parameters.WorldNormal = TransformTangentNormalToWorld(Parameters, MaterialNormal);
+#else //MATERIAL_TANGENTSPACENORMAL
+	Parameters.WorldNormal = normalize(MaterialNormal);
+#endif //MATERIAL_TANGENTSPACENORMAL
+
+#if MATERIAL_TANGENTSPACENORMAL
+	// flip the normal for backfaces being rendered with a two-sided material
+	Parameters.WorldNormal *= Parameters.TwoSidedSign;
+#endif
+
+	Parameters.ReflectionVector = ReflectionAboutCustomWorldNormal(Parameters, Parameters.WorldNormal, false);
+
+#if !PARTICLE_SPRITE_FACTORY
+	Parameters.Particle.MotionBlurFade = 1.0f;
+#endif // !PARTICLE_SPRITE_FACTORY
+
+	// Now the rest of the inputs
+%s
+}
+
+// Programmatically set the line number after all the material inputs which have a variable number of line endings
+// This allows shader error line numbers after this point to be the same regardless of which material is being compiled
+#line %s
+
+void ClipLODTransition(float2 SvPosition, float DitherFactor)
+{
+	if (abs(DitherFactor) > .001)
+	{
+		float RandCos = cos(dot(floor(SvPosition.xy), float2(347.83451793,3343.28371963)));
+		float RandomVal = frac(RandCos * 1000.0);
+		half RetVal = (DitherFactor < 0.0) ?
+			(DitherFactor + 1.0 > RandomVal) :
+			(DitherFactor < RandomVal);
+		clip(RetVal - .001);
+	}
+}
+
+void ClipLODTransition(FMaterialPixelParameters Parameters, float DitherFactor)
+{
+	ClipLODTransition(Parameters.SvPosition.xy, DitherFactor);
+}
+
+
+#define REQUIRES_VF_ATTRIBUTES_FOR_CLIPPING (USE_INSTANCING && USE_DITHERED_LOD_TRANSITION)
+
+#if USE_INSTANCING && USE_DITHERED_LOD_TRANSITION
+void ClipLODTransition(FMaterialPixelParameters Parameters)
+{
+	ClipLODTransition(Parameters, Parameters.PerInstanceParams.w);
+}
+#elif USE_DITHERED_LOD_TRANSITION && !USE_STENCIL_LOD_DITHER
+float NonInstancedDitherLODFactor;
+void ClipLODTransition(FMaterialPixelParameters Parameters)
+{
+	if (NonInstancedDitherLODFactor != 0.0)
+	{
+		ClipLODTransition(Parameters, NonInstancedDitherLODFactor);
+	}
+}
+void ClipLODTransition(float2 SvPosition)
+{
+	if (NonInstancedDitherLODFactor != 0.0)
+	{
+		ClipLODTransition(SvPosition, NonInstancedDitherLODFactor);
+	}
+}
+#else
+void ClipLODTransition(FMaterialPixelParameters Parameters)
+{
+}
+void ClipLODTransition(float2 SvPosition)
+{
+}
+#endif
+
+void GetMaterialClippingShadowDepth(FMaterialPixelParameters Parameters, FPixelMaterialInputs PixelMaterialInputs)
+{
+	ClipLODTransition(Parameters);
+	#if MATERIALBLENDING_MASKED
+		clip(GetMaterialMask(PixelMaterialInputs)); 
+	#elif MATERIALBLENDING_TRANSLUCENT
+		clip(GetMaterialOpacity(PixelMaterialInputs) - 1.0f / 255.0f);
+	#endif
+}
+
+void GetMaterialClippingVelocity(FMaterialPixelParameters Parameters, FPixelMaterialInputs PixelMaterialInputs)
+{
+	ClipLODTransition(Parameters);
+	#if MATERIALBLENDING_MASKED && MATERIAL_DITHER_OPACITY_MASK
+		clip(GetMaterialMaskInputRaw(PixelMaterialInputs) - 1.0f / 255.0f);
+	#elif MATERIALBLENDING_MASKED
+		clip(GetMaterialMask(PixelMaterialInputs)); 
+	#elif MATERIALBLENDING_TRANSLUCENT
+		clip(GetMaterialOpacity(PixelMaterialInputs) - 1.0f / 255.0f);
+	#endif
+}
+
+
+void GetMaterialCoverageAndClipping(FMaterialPixelParameters Parameters, FPixelMaterialInputs PixelMaterialInputs)
+{
+	ClipLODTransition(Parameters);
+
+#if MATERIALBLENDING_MASKED
+	#if MATERIAL_DITHER_OPACITY_MASK
+		/*
+		5 value dither. Every value present in +
+		012
+		234
+		401
+		*/
+		//float Dither5 = ( ( (uint)Parameters.SvPosition.x + (uint)Parameters.SvPosition.y * 2 + (uint)View.TemporalAAParams.x ) % 5 ) / 5.0;
+		float Dither5 = frac( ( Parameters.SvPosition.x + Parameters.SvPosition.y * 2 - 1.5 + View.TemporalAAParams.x ) / 5 );
+		float Noise = frac( dot( float2( 171.0, 231.0 ) / 71, Parameters.SvPosition.xy ) );
+		float Dither = ( Dither5 * 5 + Noise ) * (1.0 / 6.0);
+
+		float2 DepthGrad = {
+			ddx( Parameters.SvPosition.z ),
+			ddy( Parameters.SvPosition.z )
+		};
+		Dither = frac( Dither + dot( DepthGrad, float2( 4093, 3571 ) ) );
+
+		clip( GetMaterialMask(PixelMaterialInputs) + Dither - 0.5 );
+	#else
+		clip(GetMaterialMask(PixelMaterialInputs));
+	#endif
+#endif
+}
+
+#define MATERIALBLENDING_MASKED_USING_COVERAGE (FORWARD_SHADING && MATERIALBLENDING_MASKED && SUPPORTS_PIXEL_COVERAGE)
+#if MATERIALBLENDING_MASKED_USING_COVERAGE
+
+uint GetDerivativeCoverageFromMask(float MaterialMask)
+{
+	uint Coverage = 0x0;
+	if (MaterialMask > 0.01) Coverage = 0x8;
+	if (MaterialMask > 0.25) Coverage = 0x9;
+	if (MaterialMask > 0.50) Coverage = 0xD;
+	if (MaterialMask > 0.75) Coverage = 0xF;
+	return Coverage;
+}
+
+// Returns the new pixel coverage according the material's mask and the current pixel's mask.
+uint DiscardMaterialWithPixelCoverage(FMaterialPixelParameters MaterialParameters, FPixelMaterialInputs PixelMaterialInputs)
+{
+	ClipLODTransition(MaterialParameters);
+	float OriginalMask = GetMaterialMaskInputRaw(PixelMaterialInputs);
+	float MaskClip = GetMaterialOpacityMaskClipValue();
+
+	if (View.NumSceneColorMSAASamples > 1)
+	{
+		float Mask = (OriginalMask - MaskClip) / (1.0 - MaskClip);
+		uint CurrentPixelCoverage = GetDerivativeCoverageFromMask(Mask);
+		// Discard pixel shader if all sample are masked to avoid computing other material inputs.
+		clip(float(CurrentPixelCoverage) - 0.5);
+		return CurrentPixelCoverage;
+	}
+	clip(OriginalMask - MaskClip);
+	return 0xF;
+}
+
+#endif // MATERIALBLENDING_MASKED_USING_COVERAGE
+
+
+	#define FrontFaceSemantic SV_IsFrontFace
+	#define FIsFrontFace bool
+	half GetFloatFacingSign(FIsFrontFace bIsFrontFace)
+	{
+		return bIsFrontFace ? +1 : -1;
+	}
+
+#if MATERIAL_TWOSIDED_SEPARATE_PASS
+	#define OPTIONAL_IsFrontFace
+	static const FIsFrontFace bIsFrontFace = 1;
+#else
+	#define OPTIONAL_IsFrontFace , in FIsFrontFace bIsFrontFace : FrontFaceSemantic
+#endif
+
+/** Initializes the subset of Parameters that was not set in GetMaterialPixelParameters. */
+void CalcMaterialParametersEx(
+	in out FMaterialPixelParameters Parameters,
+	in out FPixelMaterialInputs PixelMaterialInputs,
+	float4 SvPosition,
+	float4 ScreenPosition,
+	FIsFrontFace bIsFrontFace,
+	float3 TranslatedWorldPosition,
+	float3 TranslatedWorldPositionExcludingShaderOffsets)
+{
+	// Remove the pre view translation
+	Parameters.WorldPosition_CamRelative = TranslatedWorldPosition.xyz;
+	Parameters.AbsoluteWorldPosition = TranslatedWorldPosition.xyz - ResolvedView.PreViewTranslation.xyz;
+
+	// If the material uses any non-offset world position expressions, calculate those parameters. If not, 
+	// the variables will have been initialised to 0 earlier.
+#if USE_WORLD_POSITION_EXCLUDING_SHADER_OFFSETS
+	Parameters.WorldPosition_NoOffsets_CamRelative = TranslatedWorldPositionExcludingShaderOffsets;
+	Parameters.WorldPosition_NoOffsets = TranslatedWorldPositionExcludingShaderOffsets - ResolvedView.PreViewTranslation.xyz;
+#endif
+
+	Parameters.SvPosition = SvPosition;
+	Parameters.ScreenPosition = ScreenPosition;
+
+#if COMPILER_GLSL_ES2 || COMPILER_GLSL_ES3_1
+	// ES2 normalize isn't done accurately. This seems to fix it.
+	// Originally this was normalize(normalize(TranslatedWorldPosition.xyz)) but tegra4 appears to optimize that out.
+	Parameters.CameraVector = normalize(-0.01 * Parameters.WorldPosition_CamRelative.xyz);
+#else
+	// TranslatedWorldPosition is the world position translated to the camera position, which is just -CameraVector
+	Parameters.CameraVector = normalize(-Parameters.WorldPosition_CamRelative.xyz);
+#endif
+
+	Parameters.LightVector = 0;
+
+	Parameters.TwoSidedSign = 1.0f;
+
+#if MATERIAL_TWOSIDED && HAS_PRIMITIVE_UNIFORM_BUFFER
+	Parameters.TwoSidedSign *= View.CullingSign * Primitive.LocalToWorldDeterminantSign;
+
+	#if !MATERIAL_TWOSIDED_SEPARATE_PASS
+		Parameters.TwoSidedSign *= GetFloatFacingSign(bIsFrontFace);
+	#endif
+#endif
+
+	// Now that we have all the pixel-related parameters setup, calculate the Material Input/Attributes and Normal
+	CalcPixelMaterialInputs(Parameters, PixelMaterialInputs);
+}
+
+// convenience function to setup CalcMaterialParameters assuming we don't support TranslatedWorldPositionExcludingShaderOffsets
+// @param SvPosition from SV_Position when rendering the view, for other projections e.g. shadowmaps this function cannot be used and you need to call CalcMaterialParametersEx()
+void CalcMaterialParameters(
+	in out FMaterialPixelParameters Parameters,
+	in out FPixelMaterialInputs PixelMaterialInputs,
+	float4 SvPosition,
+	FIsFrontFace bIsFrontFace)
+{
+	float4 ScreenPosition = SvPositionToResolvedScreenPosition(SvPosition);
+	float3 TranslatedWorldPosition = SvPositionToResolvedTranslatedWorld(SvPosition);
+
+	CalcMaterialParametersEx(Parameters, PixelMaterialInputs, SvPosition, ScreenPosition, bIsFrontFace, TranslatedWorldPosition, TranslatedWorldPosition);
+}
+
+void CalcMaterialParametersPost(
+	in out FMaterialPixelParameters Parameters,
+	in out FPixelMaterialInputs PixelMaterialInputs,
+	float4 SvPosition,
+	FIsFrontFace bIsFrontFace)
+{
+	float4 ScreenPosition = SvPositionToScreenPosition(SvPosition);
+	float3 TranslatedWorldPosition = SvPositionToTranslatedWorld(SvPosition);
+
+	CalcMaterialParametersEx(Parameters, PixelMaterialInputs, SvPosition, ScreenPosition, bIsFrontFace, TranslatedWorldPosition, TranslatedWorldPosition);
+}
+
+/** Assemble the transform from tangent space into world space */
+half3x3 AssembleTangentToWorld( half3 TangentToWorld0, half4 TangentToWorld2 )
+{
+	// Will not be orthonormal after interpolation. This perfectly matches xNormal.
+	// Any mismatch with xNormal will cause distortions for baked normal maps.
+
+	// Derive the third basis vector off of the other two.
+	// Flip based on the determinant sign
+	half3 TangentToWorld1 = cross(TangentToWorld2.xyz,TangentToWorld0) * TangentToWorld2.w;
+	// Transform from tangent space to world space
+	return half3x3(TangentToWorld0, TangentToWorld1, TangentToWorld2.xyz);
+}
+
+// Whether the material shader should output pixel depth offset
+#define OUTPUT_PIXEL_DEPTH_OFFSET (WANT_PIXEL_DEPTH_OFFSET && (MATERIALBLENDING_SOLID || MATERIALBLENDING_MASKED) && !ES2_PROFILE)
+
+// Whether to use the hidden d3d11 feature that supports depth writes with ZCull by only pushing into the screen
+//@todo - use for other SM5 platforms
+#define SUPPORTS_CONSERVATIVE_DEPTH_WRITES ((COMPILER_HLSL && FEATURE_LEVEL >= FEATURE_LEVEL_SM5) || (PS4_PROFILE))
+#define USE_CONSERVATIVE_DEPTH_WRITES (OUTPUT_PIXEL_DEPTH_OFFSET && SUPPORTS_CONSERVATIVE_DEPTH_WRITES) 
+
+#if USE_CONSERVATIVE_DEPTH_WRITES
+
+#if COMPILER_HLSL
+	// Note: for some reason using SV_DepthLessEqual without these interpolation modifiers causes a compile error in d3d
+	#define INPUT_POSITION_QUALIFIERS linear noperspective centroid
+	// Use conservative depth output so we still get Z Cull.  Note, this is a reversed Z depth surface.
+	#define DEPTH_WRITE_SEMANTIC SV_DepthLessEqual
+#elif PS4_PROFILE
+	#define INPUT_POSITION_QUALIFIERS
+	#define DEPTH_WRITE_SEMANTIC S_DEPTH_LE_OUTPUT
+#else
+	#error USE_CONSERVATIVE_DEPTH_WRITES enabled for unsupported platform
+#endif
+
+#else
+	#define INPUT_POSITION_QUALIFIERS 
+	#define DEPTH_WRITE_SEMANTIC SV_DEPTH
+#endif
+
+#if OUTPUT_PIXEL_DEPTH_OFFSET
+	#define OPTIONAL_OutDepthConservative ,out float OutDepth : DEPTH_WRITE_SEMANTIC
+	#define OPTIONAL_OutDepth ,out float OutDepth : SV_DEPTH
+#else
+	#define OPTIONAL_OutDepthConservative
+	#define OPTIONAL_OutDepth
+#endif
+
+float ApplyPixelDepthOffsetToMaterialParameters(inout FMaterialPixelParameters MaterialParameters, FPixelMaterialInputs PixelMaterialInputs, out float OutDepth)
+{
+	// Can only push into the scene for conservative depth writes
+	float PixelDepthOffset = max(GetMaterialPixelDepthOffset(PixelMaterialInputs), 0);
+
+	// Update positions used for shading
+	MaterialParameters.ScreenPosition.w += PixelDepthOffset;
+	MaterialParameters.SvPosition.w = MaterialParameters.ScreenPosition.w;
+	MaterialParameters.AbsoluteWorldPosition += MaterialParameters.CameraVector * PixelDepthOffset;
+
+	OutDepth = MaterialParameters.ScreenPosition.z / MaterialParameters.ScreenPosition.w;
+
+	return PixelDepthOffset;
+}
diff --git a/Engine/Shaders/Private/MeshParticleVertexFactory.ush b/Engine/Shaders/Private/MeshParticleVertexFactory.ush
new file mode 100644
--- /dev/null
+++ b/Engine/Shaders/Private/MeshParticleVertexFactory.ush
@@ -0,0 +1,693 @@
+// Copyright 1998-2017 Epic Games, Inc. All Rights Reserved.
+
+/*=============================================================================
+	MeshParticleVertexFactory.usf: Mesh particle vertex factory shader code.
+=============================================================================*/
+
+#include "VertexFactoryCommon.ush"
+
+// These are only enabled for SM5 due to limited interpolators elsewhere
+#define USE_PARTICLE_NODES (FEATURE_LEVEL >= FEATURE_LEVEL_SM5)
+#define USE_PARTICLE_POSITION (FEATURE_LEVEL >= FEATURE_LEVEL_SM5 && (NEEDS_PARTICLE_POSITION || PASS_NEEDS_PRIMITIVE_VOLUME_BOUNDS))
+#define USE_PARTICLE_VELOCITY (FEATURE_LEVEL >= FEATURE_LEVEL_SM5 && NEEDS_PARTICLE_VELOCITY)
+#define USE_PARTICLE_TIME (FEATURE_LEVEL >= FEATURE_LEVEL_SM5 && NEEDS_PARTICLE_TIME)
+
+// Only support previous transform for deferred feature level
+#define USE_PREVIOUS_TRANSFORM (FEATURE_LEVEL >= FEATURE_LEVEL_SM4)
+
+float4x4 PreviousLocalToWorld;
+
+
+#if USE_PREVIOUS_TRANSFORM
+#if PARTICLE_MESH_INSTANCED
+	Buffer<float4> PrevTransformBuffer;
+#else
+	float4 PrevTransform0;
+	float4 PrevTransform1;
+	float4 PrevTransform2;
+#endif
+#endif
+	
+#if !PARTICLE_MESH_INSTANCED
+	float4 Transform1;
+	float4 Transform2;
+	float4 Transform3;
+	float4 SubUVParams;
+	float SubUVLerp;
+	float4 ParticleDirection;
+	float RelativeTime;
+	float4 DynamicParameter;
+	float4 ParticleColor;
+#endif
+
+struct FVertexFactoryInput
+{
+	float4	Position	: ATTRIBUTE0;
+	half3	TangentX	: ATTRIBUTE1;
+	// TangentZ.w contains sign of tangent basis determinant
+	half4	TangentZ	: ATTRIBUTE2;
+	half4	VertexColor : ATTRIBUTE3;
+
+#if NUM_MATERIAL_TEXCOORDS_VERTEX
+	float2	TexCoords[NUM_MATERIAL_TEXCOORDS_VERTEX] : ATTRIBUTE4;
+#elif USE_PARTICLE_SUBUVS
+	float2	TexCoords[1] : ATTRIBUTE4;
+#endif
+
+#if PARTICLE_MESH_INSTANCED
+	float4 Transform1 : ATTRIBUTE8;
+	float4 Transform2 : ATTRIBUTE9;
+	float4 Transform3 : ATTRIBUTE10;
+
+	int4 SubUVParams : ATTRIBUTE11;
+
+	/** X: SubUV lerp Y: Particle relative time. */
+	float2 SubUVLerpAndRelTime : ATTRIBUTE12;
+
+	#if USE_DYNAMIC_PARAMETERS
+		float4 DynamicParameter	: ATTRIBUTE13;
+	#endif	//USE_DYNAMIC_PARAMETERS
+
+	float4 ParticleColor : ATTRIBUTE14;
+
+	/** The velocity of the particle, XYZ: direction, W: speed. */
+	float4 ParticleVelocity : ATTRIBUTE15;
+#endif
+
+/** Optional instance ID for vertex layered rendering */
+#if (PARTICLE_MESH_INSTANCED || (FEATURE_LEVEL >= FEATURE_LEVEL_SM4 && ONEPASS_POINTLIGHT_SHADOW && USING_VERTEX_SHADER_LAYER))
+	uint InstanceId	: SV_InstanceID;
+#endif
+};
+
+struct FVertexFactoryInterpolantsVSToPS
+{
+	TANGENTTOWORLD_INTERPOLATOR_BLOCK
+
+#if USE_PARTICLE_SUBUVS
+	float4 SubUV0AndTexCoord0 : TEXCOORD1;
+	float4 SubUV1AndLerp : TEXCOORD2;
+#else
+	#if NUM_TEX_COORD_INTERPOLATORS
+			#if (ES2_PROFILE || ES3_1_PROFILE)
+				// Avoid dependent texture fetches, put all UVs in xy
+				float2	TexCoords[NUM_TEX_COORD_INTERPOLATORS] : TEXCOORD0;
+			#else
+				float4	TexCoords[(NUM_TEX_COORD_INTERPOLATORS+1)/2] : TEXCOORD0;
+			#endif
+	#endif
+#endif
+
+#if INTERPOLATE_VERTEX_COLOR
+	float4	VertexColor : COLOR0;
+#endif
+
+#if NEEDS_PARTICLE_COLOR
+	float4	ParticleColor : COLOR1;
+#endif
+
+#if USE_DYNAMIC_PARAMETERS
+	float4 DynamicParameter : COLOR2;
+#endif
+	
+#if USE_PARTICLE_NODES
+#if USE_PARTICLE_POSITION
+	/** Particle position in camera-centered translated world space */
+	float3 ParticleTranslatedWorldPosition	: PARTICLE_POSITION;
+#endif
+
+#if USE_PARTICLE_VELOCITY
+	/** The velocity of the particle, XYZ: direction, W: speed. */
+	float4 ParticleVelocity	: PARTICLE_VELOCITY;
+#endif
+
+#if USE_PARTICLE_TIME
+	/** Relative alive time of the particle */
+	float ParticleTime : PARTICLE_TIME;
+#endif
+
+#if NEEDS_PARTICLE_TRANSFORM
+	float4 LocalToWorld[4] : PARTICLE_TRANSFORM;
+#endif
+#endif
+
+#if INSTANCED_STEREO
+	nointerpolation uint PackedEyeIndex : PACKED_EYE_INDEX;
+#endif
+};
+
+struct FVertexFactoryIntermediates
+{
+	/** The color of the vertex. */
+	float4 VertexColor;
+
+	/** The color of the particle. */
+	float4 ParticleColor;
+
+	/** The texture coordinates for the vertex. */
+#if NUM_MATERIAL_TEXCOORDS_VERTEX
+	float2	TexCoords[NUM_MATERIAL_TEXCOORDS_VERTEX];
+#endif 
+#if USE_PARTICLE_SUBUVS
+	float4 SubUVCoords;
+	float SubUVLerp;
+#endif
+#if USE_DYNAMIC_PARAMETERS
+	/** Optional dynamic parameter for the particle. */
+	float4 DynamicParameter;
+#endif	//USE_DYNAMIC_PARAMETERS
+
+	/** The velocity of the particle, XYZ: direction, W: speed. */
+	float4 ParticleVelocity;
+
+	/** Particle position in camera-centered translated world space. */
+	float3 ParticleTranslatedWorldPosition;
+	
+	/** Relative time. */
+	float RelativeTime;
+
+#if NEEDS_PARTICLE_TRANSFORM
+	float4 LocalToWorld[4];
+#endif
+};
+
+float4x4 GetParticleTransform(FVertexFactoryInput Input)
+{
+	//Transform is transposed on CPU and position is packed into W. Transpose again to get the correct transform back.
+#if PARTICLE_MESH_INSTANCED
+	return transpose(float4x4(Input.Transform1, Input.Transform2, Input.Transform3, float4(0, 0, 0, 1)));
+#else
+	return transpose(float4x4(Transform1, Transform2, Transform3, float4(0, 0, 0, 1)));
+#endif
+}
+
+#if USE_PREVIOUS_TRANSFORM
+float4x4 GetParticlePrevTransform(FVertexFactoryInput Input)
+{
+#if PARTICLE_MESH_INSTANCED
+	uint TransformBaseIdx = GetLayerInstanceID(Input.InstanceId) * 3;
+
+	float4x4 PrevTransform4x4 = float4x4(
+		PrevTransformBuffer[TransformBaseIdx + 0], 
+		PrevTransformBuffer[TransformBaseIdx + 1], 
+		PrevTransformBuffer[TransformBaseIdx + 2], 
+		float4(0, 0, 0, 1)
+		);
+#else
+	float4x4 PrevTransform4x4 = float4x4(
+		PrevTransform0, 
+		PrevTransform1, 
+		PrevTransform2, 
+		float4(0, 0, 0, 1)
+		);
+#endif
+	//Transform is transposed on CPU and position is packed into W. Transpose again to get the correct transform back.
+	return transpose(PrevTransform4x4);
+}
+#endif
+
+/** Converts from vertex factory specific interpolants to a FMaterialPixelParameters, which is used by material inputs. */
+FMaterialPixelParameters GetMaterialPixelParameters(FVertexFactoryInterpolantsVSToPS Interpolants, float4 SvPosition)
+{
+	// GetMaterialPixelParameters is responsible for fully initializing the result
+	FMaterialPixelParameters Result = MakeInitializedMaterialPixelParameters();
+
+#if USE_PARTICLE_SUBUVS
+	#if NUM_MATERIAL_TEXCOORDS
+		UNROLL
+		for( int CoordinateIndex = 0; CoordinateIndex < NUM_MATERIAL_TEXCOORDS; CoordinateIndex++ )
+		{
+			Result.TexCoords[CoordinateIndex] = Interpolants.SubUV0AndTexCoord0.zw;
+		}
+	#endif
+	Result.Particle.SubUVCoords[0] = Interpolants.SubUV0AndTexCoord0.xy;
+	Result.Particle.SubUVCoords[1] = Interpolants.SubUV1AndLerp.xy;
+	Result.Particle.SubUVLerp = Interpolants.SubUV1AndLerp.z;
+#elif NUM_TEX_COORD_INTERPOLATORS
+	#if (ES2_PROFILE || ES3_1_PROFILE)
+		UNROLL
+		for( int CoordinateIndex = 0; CoordinateIndex < NUM_TEX_COORD_INTERPOLATORS; CoordinateIndex++ )
+		{
+			Result.TexCoords[CoordinateIndex] = Interpolants.TexCoords[CoordinateIndex].xy;
+		}
+	#else
+		UNROLL
+		for(int CoordinateIndex = 0;CoordinateIndex < NUM_TEX_COORD_INTERPOLATORS;CoordinateIndex += 2)
+		{
+			Result.TexCoords[CoordinateIndex] = Interpolants.TexCoords[CoordinateIndex/2].xy;
+
+			if(CoordinateIndex + 1 < NUM_TEX_COORD_INTERPOLATORS)
+			{
+				Result.TexCoords[CoordinateIndex + 1] = Interpolants.TexCoords[CoordinateIndex/2].wz;
+			}
+		}
+	#endif
+#endif
+
+	half3 TangentToWorld0 = Interpolants.TangentToWorld0.xyz;
+	half4 TangentToWorld2 = Interpolants.TangentToWorld2;
+	Result.UnMirrored = TangentToWorld2.w;
+
+#if INTERPOLATE_VERTEX_COLOR
+	Result.VertexColor = Interpolants.VertexColor;
+#else
+	Result.VertexColor = 0;
+#endif
+
+#if NEEDS_PARTICLE_COLOR
+	Result.Particle.Color = Interpolants.ParticleColor;
+#endif
+
+#if USE_DYNAMIC_PARAMETERS
+	Result.Particle.DynamicParameter = Interpolants.DynamicParameter;
+#else
+	Result.Particle.DynamicParameter = float4(1,1,1,1);
+#endif //USE_DYNAMIC_PARAMETERS
+
+#if USE_PARTICLE_NODES
+	#if USE_PARTICLE_POSITION
+		Result.Particle.TranslatedWorldPositionAndSize.xyz = Interpolants.ParticleTranslatedWorldPosition.xyz;
+		Result.Particle.TranslatedWorldPositionAndSize.w = 1;
+	#endif
+
+	#if USE_PARTICLE_VELOCITY
+		Result.Particle.Velocity = Interpolants.ParticleVelocity;
+	#endif
+	#if USE_PARTICLE_TIME
+		Result.Particle.RelativeTime = Interpolants.ParticleTime;
+	#endif
+	#if NEEDS_PARTICLE_TRANSFORM
+		Result.Particle.LocalToWorld[0] = Interpolants.LocalToWorld[0];
+		Result.Particle.LocalToWorld[1] = Interpolants.LocalToWorld[1];
+		Result.Particle.LocalToWorld[2] = Interpolants.LocalToWorld[2];
+		Result.Particle.LocalToWorld[3] = Interpolants.LocalToWorld[3];
+	#endif
+#endif
+
+	Result.Particle.MotionBlurFade = 1.0f;
+	Result.TangentToWorld = AssembleTangentToWorld( TangentToWorld0, TangentToWorld2 );
+	Result.TwoSidedSign = 1;
+
+	return Result;
+}
+
+/** Converts from vertex factory specific input to a FMaterialVertexParameters, which is used by vertex shader material inputs. */
+FMaterialVertexParameters GetMaterialVertexParameters(FVertexFactoryInput Input, FVertexFactoryIntermediates Intermediates, float3 WorldPosition, float3x3 TangentToLocal)
+{
+	FMaterialVertexParameters Result = (FMaterialVertexParameters)0;
+	Result.WorldPosition = WorldPosition;
+	Result.VertexColor = Intermediates.VertexColor;
+	Result.PreSkinnedPosition = Input.Position.xyz;
+	Result.PreSkinnedNormal = Input.TangentZ.xyz * 2.f - 1.f;
+
+	Result.Particle.TranslatedWorldPositionAndSize.xyz = Intermediates.ParticleTranslatedWorldPosition.xyz;
+	Result.Particle.TranslatedWorldPositionAndSize.w = 1;
+	Result.Particle.Velocity = Intermediates.ParticleVelocity;
+	Result.Particle.RelativeTime = Intermediates.RelativeTime;
+
+#if USE_DYNAMIC_PARAMETERS
+	Result.Particle.DynamicParameter = Intermediates.DynamicParameter;
+#endif	//USE_DYNAMIC_PARAMETERS
+
+	float4x4 Transform = GetParticleTransform(Input);
+	Result.InstanceLocalToWorld = Transform;
+#if NEEDS_PARTICLE_TRANSFORM
+	Result.Particle.LocalToWorld[0] = Transform[0];
+	Result.Particle.LocalToWorld[1] = Transform[1];
+	Result.Particle.LocalToWorld[2] = Transform[2];
+	Result.Particle.LocalToWorld[3] = Transform[3];
+#endif
+
+	Result.TangentToWorld = mul(TangentToLocal, (float3x3)Transform ); 
+
+#if NUM_MATERIAL_TEXCOORDS_VERTEX
+	for(int CoordinateIndex = 0; CoordinateIndex < NUM_MATERIAL_TEXCOORDS_VERTEX; CoordinateIndex++)
+	{
+		Result.TexCoords[CoordinateIndex] = Intermediates.TexCoords[CoordinateIndex];
+	}
+#endif
+
+	return Result;
+}
+
+#if NEEDS_LIGHTMAP_COORDINATE
+void GetLightMapCoordinates(FVertexFactoryInterpolantsVSToPS Interpolants, out float2 LightmapUV0, out float2 LightmapUV1)
+{
+	LightmapUV0 = LightmapUV1 = 0;
+}
+
+half2 GetShadowMapCoordinate(FVertexFactoryInterpolantsVSToPS Interpolants)
+{
+	return 0;
+}
+#endif
+
+
+float4 CalcWorldPosition(FVertexFactoryInput Input)
+{
+	float4x4 Transform = GetParticleTransform(Input);
+	float3 WorldPosition = mul(Input.Position, Transform).xyz + ResolvedView.PreViewTranslation;
+	return float4(WorldPosition, Input.Position.w);
+}
+
+float4 CalcPrevWorldPosition(FVertexFactoryInput Input)
+{
+#if USE_PREVIOUS_TRANSFORM
+	if (MeshParticleVF.PrevTransformAvailable > 0)
+	{
+		float4x4 PrevTransform = GetParticlePrevTransform(Input);
+		float3 PrevWorldPosition = mul(Input.Position, PrevTransform).xyz + ResolvedView.PrevPreViewTranslation;
+		return float4(PrevWorldPosition, Input.Position.w);
+	}
+	else
+#endif
+	{
+		return CalcWorldPosition(Input);
+	}
+}
+
+FVertexFactoryIntermediates GetVertexFactoryIntermediates(FVertexFactoryInput Input)
+{
+	FVertexFactoryIntermediates Intermediates = (FVertexFactoryIntermediates)0;
+	// Swizzle vertex color.
+	Intermediates.VertexColor = Input.VertexColor FCOLOR_COMPONENT_SWIZZLE;
+
+	// World position and localtoworld transform.
+#if PARTICLE_MESH_INSTANCED
+	float3 ParticleWorldPosition = float3( Input.Transform1.w, Input.Transform2.w, Input.Transform3.w);
+	float4x4 LocalToWorldMat = float4x4(Input.Transform1, Input.Transform2, Input.Transform3, float4(0.0f, 0.0f, 0.0f, 1.0f));
+#if NEEDS_PARTICLE_TRANSFORM
+	Intermediates.LocalToWorld[0] = LocalToWorldMat[0];
+	Intermediates.LocalToWorld[1] = LocalToWorldMat[1];
+	Intermediates.LocalToWorld[2] = LocalToWorldMat[2];
+	Intermediates.LocalToWorld[3] = LocalToWorldMat[3];
+#endif
+#else
+	float3 ParticleWorldPosition = float3( Transform1.w, Transform2.w, Transform3.w);
+#endif
+
+	Intermediates.ParticleTranslatedWorldPosition = ParticleWorldPosition + ResolvedView.PreViewTranslation.xyz;
+
+#if PARTICLE_MESH_INSTANCED
+	Intermediates.ParticleColor = Input.ParticleColor;
+#else
+	Intermediates.ParticleColor = ParticleColor;
+#endif
+
+#if NUM_MATERIAL_TEXCOORDS_VERTEX
+	for(int CoordinateIndex = 0; CoordinateIndex < NUM_MATERIAL_TEXCOORDS_VERTEX; CoordinateIndex++)
+	{
+		Intermediates.TexCoords[CoordinateIndex] = Input.TexCoords[CoordinateIndex].xy;
+	}
+#endif
+
+#if USE_PARTICLE_SUBUVS
+	#if PARTICLE_MESH_INSTANCED
+		Intermediates.SubUVCoords.xy = (float2( Input.SubUVParams.x, Input.SubUVParams.y ) + Input.TexCoords[0].xy) * MeshParticleVF.SubImageSize.xy;
+		Intermediates.SubUVCoords.zw = (float2( Input.SubUVParams.z, Input.SubUVParams.w ) + Input.TexCoords[0].xy) * MeshParticleVF.SubImageSize.xy;
+		Intermediates.SubUVLerp = Input.SubUVLerpAndRelTime.x;
+	#else
+		Intermediates.SubUVCoords.xy = (float2( SubUVParams.x, SubUVParams.y ) + Input.TexCoords[0].xy) * MeshParticleVF.SubImageSize.xy;
+		Intermediates.SubUVCoords.zw = (float2( SubUVParams.z, SubUVParams.w ) + Input.TexCoords[0].xy) * MeshParticleVF.SubImageSize.xy;
+		Intermediates.SubUVLerp = SubUVLerp.x;
+	#endif
+#endif 
+
+#if USE_DYNAMIC_PARAMETERS
+	#if PARTICLE_MESH_INSTANCED
+		Intermediates.DynamicParameter = Input.DynamicParameter;
+	#else
+		Intermediates.DynamicParameter = DynamicParameter;
+	#endif
+#endif	//USE_DYNAMIC_PARAMETERS
+
+	Intermediates.ParticleTranslatedWorldPosition = ParticleWorldPosition.xyz + ResolvedView.PreViewTranslation.xyz;
+
+#if PARTICLE_MESH_INSTANCED
+	Intermediates.ParticleVelocity = Input.ParticleVelocity;
+	Intermediates.RelativeTime = Input.SubUVLerpAndRelTime.y;
+#else
+	Intermediates.ParticleVelocity = ParticleDirection;
+	Intermediates.RelativeTime = RelativeTime;
+#endif
+
+	return Intermediates;
+}
+
+/**
+* Get the 3x3 tangent basis vectors for this vertex factory
+* this vertex factory will calculate the binormal on-the-fly
+*
+* @param Input - vertex input stream structure
+* @return 3x3 matrix
+*/
+float3x3 VertexFactoryGetTangentToLocal( FVertexFactoryInput Input, FVertexFactoryIntermediates Intermediates )
+{
+	float3x3 Result=0;	
+	float3 TangentX = TangentBias(Input.TangentX);
+	float4 TangentZ = TangentBias(Input.TangentZ);
+
+	// pass-thru the tangent
+	Result[0] = TangentX;
+	// pass-thru the normal
+	Result[2] = TangentZ.xyz;	
+
+	
+	// derive the binormal by getting the cross product of the normal and tangent
+	Result[1] = cross(Result[2], Result[0]) * TangentZ.w;	
+	// Recalculate TangentX off of the other two vectors
+	// This corrects quantization error since TangentX was passed in as a quantized vertex input
+	// The error shows up most in specular off of a mesh with a smoothed UV seam (normal is smooth, but tangents vary across the seam)
+	Result[0] = cross(Result[1], Result[2]) * TangentZ.w;
+
+	return Result;
+}
+
+// @return translated world position
+float4 VertexFactoryGetWorldPosition(FVertexFactoryInput Input, FVertexFactoryIntermediates Intermediates)
+{
+	return CalcWorldPosition(Input);
+}
+
+float4 VertexFactoryGetRasterizedWorldPosition(FVertexFactoryInput Input, FVertexFactoryIntermediates Intermediates, float4 InWorldPosition)
+{
+	return InWorldPosition;
+}
+
+float3 VertexFactoryGetPositionForVertexLighting(FVertexFactoryInput Input, FVertexFactoryIntermediates Intermediates, float3 TranslatedWorldPosition)
+{
+	return TranslatedWorldPosition;
+}
+
+void CalcTangentToWorld(FVertexFactoryInput Input, FVertexFactoryIntermediates Intermediates, out float3 TangentToWorld0, out float4 TangentToWorld2)
+{
+	float3x3 TangentToLocal = VertexFactoryGetTangentToLocal(Input, Intermediates);
+
+	float4x4 Transform = GetParticleTransform(Input);
+	float3x3 TangentToWorld = mul(TangentToLocal, (float3x3)Transform);
+
+	// Normalize to remove scaling. Incorrect but faster than computing the inverse-transpose.
+	TangentToWorld0 = normalize(TangentToWorld[0]);
+	TangentToWorld2 = float4(normalize(TangentToWorld[2]), TangentBias(Input.TangentZ.w) * Primitive.LocalToWorldDeterminantSign);
+}
+
+FVertexFactoryInterpolantsVSToPS VertexFactoryGetInterpolantsVSToPS(FVertexFactoryInput Input, FVertexFactoryIntermediates Intermediates, FMaterialVertexParameters VertexParameters)
+{
+	FVertexFactoryInterpolantsVSToPS Interpolants;
+
+#if USE_PARTICLE_SUBUVS
+	Interpolants.SubUV0AndTexCoord0.xy = Intermediates.SubUVCoords.xy;
+	Interpolants.SubUV1AndLerp.xy = Intermediates.SubUVCoords.zw;
+	Interpolants.SubUV1AndLerp.zw = Intermediates.SubUVLerp.xx;
+
+	#if NUM_MATERIAL_TEXCOORDS
+		float2 CustomizedUVs[NUM_MATERIAL_TEXCOORDS];
+		GetMaterialCustomizedUVs(VertexParameters, CustomizedUVs);
+		Interpolants.SubUV0AndTexCoord0.zw = CustomizedUVs[0];
+	#else
+		Interpolants.SubUV0AndTexCoord0.zw = 0;
+	#endif
+#elif NUM_TEX_COORD_INTERPOLATORS
+	#if (ES2_PROFILE || ES3_1_PROFILE)
+		GetMaterialCustomizedUVs(VertexParameters, Interpolants.TexCoords);
+	#else
+		// Ensure the unused components of the last packed texture coordinate are initialized.
+		Interpolants.TexCoords[(NUM_TEX_COORD_INTERPOLATORS + 1) / 2 - 1] = 0;
+
+		float2 CustomizedUVs[NUM_TEX_COORD_INTERPOLATORS];
+		GetMaterialCustomizedUVs(VertexParameters, CustomizedUVs);
+		GetCustomInterpolators(VertexParameters, CustomizedUVs);
+
+		UNROLL
+		for(int CoordinateIndex = 0;CoordinateIndex < NUM_TEX_COORD_INTERPOLATORS;CoordinateIndex += 2)
+		{
+			Interpolants.TexCoords[CoordinateIndex / 2].xy = CustomizedUVs[CoordinateIndex];
+			if(CoordinateIndex + 1 < NUM_TEX_COORD_INTERPOLATORS)
+			{
+				Interpolants.TexCoords[CoordinateIndex / 2].wz = CustomizedUVs[CoordinateIndex + 1];
+			}
+		}
+	#endif 
+#endif
+
+
+	Interpolants.TangentToWorld0.w = 0;
+	CalcTangentToWorld(Input, Intermediates, Interpolants.TangentToWorld0.xyz, Interpolants.TangentToWorld2);
+
+#if INTERPOLATE_VERTEX_COLOR
+	Interpolants.VertexColor = Intermediates.VertexColor;
+#endif
+
+#if NEEDS_PARTICLE_COLOR
+	Interpolants.ParticleColor = Intermediates.ParticleColor;
+#endif
+
+#if USE_DYNAMIC_PARAMETERS
+	Interpolants.DynamicParameter = Intermediates.DynamicParameter;
+#endif	//USE_DYNAMIC_PARAMETERS
+
+#if USE_PARTICLE_NODES
+	#if USE_PARTICLE_POSITION
+		Interpolants.ParticleTranslatedWorldPosition = Intermediates.ParticleTranslatedWorldPosition;
+	#endif
+	#if USE_PARTICLE_VELOCITY
+		Interpolants.ParticleVelocity = Intermediates.ParticleVelocity;
+	#endif
+	#if USE_PARTICLE_TIME
+		Interpolants.ParticleTime = Intermediates.RelativeTime;
+	#endif
+	#if NEEDS_PARTICLE_TRANSFORM
+		Interpolants.LocalToWorld = Intermediates.LocalToWorld;
+	#endif
+#endif
+
+#if INSTANCED_STEREO
+	Interpolants.PackedEyeIndex = 0;
+#endif
+
+	return Interpolants;
+}
+
+// @return previous translated world position
+float4 VertexFactoryGetPreviousWorldPosition(FVertexFactoryInput Input, FVertexFactoryIntermediates Intermediates)
+{	
+	return CalcPrevWorldPosition(Input);
+}
+
+#if USING_TESSELLATION
+	struct FVertexFactoryInterpolantsVSToDS
+	{
+		FVertexFactoryInterpolantsVSToPS InterpolantsVSToPS;
+	};
+
+	float2 VertexFactoryGetTextureCoordinateDS( FVertexFactoryInterpolantsVSToDS Interpolants )
+	{
+	#if NUM_MATERIAL_TEXCOORDS
+		return Interpolants.InterpolantsVSToPS.TexCoords[0].xy;
+	#else // #if NUM_MATERIAL_TEXCOORDS
+		return float2(0,0);
+	#endif // #if NUM_MATERIAL_TEXCOORDS
+	}
+
+	FVertexFactoryInterpolantsVSToPS VertexFactoryAssignInterpolants(FVertexFactoryInterpolantsVSToDS Input)
+	{
+		return Input.InterpolantsVSToPS;
+	}
+
+	/** Converts from vertex factory specific interpolants to a FMaterialTessellationParameters, which is used by material inputs. */
+	FMaterialTessellationParameters GetMaterialTessellationParameters(FVertexFactoryInterpolantsVSToDS Interpolants, float3 CameraLocalWorldPosition)
+	{
+		FMaterialTessellationParameters	Result;
+	#if NUM_TEX_COORD_INTERPOLATORS
+		UNROLL
+		for(int CoordinateIndex = 0;CoordinateIndex < NUM_TEX_COORD_INTERPOLATORS;CoordinateIndex += 2)
+		{
+			Result.TexCoords[CoordinateIndex] = Interpolants.InterpolantsVSToPS.TexCoords[CoordinateIndex/2].xy;
+			if(CoordinateIndex + 1 < NUM_TEX_COORD_INTERPOLATORS)
+			{
+				Result.TexCoords[CoordinateIndex + 1] = Interpolants.InterpolantsVSToPS.TexCoords[CoordinateIndex/2].wz;
+			}
+		}
+	#endif
+
+		Result.VertexColor = float4(1,1,1,1);
+
+		half3 TangentToWorld0 = Interpolants.InterpolantsVSToPS.TangentToWorld0.xyz;
+		half4 TangentToWorld2 = Interpolants.InterpolantsVSToPS.TangentToWorld2;
+		Result.TangentToWorld = AssembleTangentToWorld( TangentToWorld0, TangentToWorld2 );
+		Result.TangentToWorldPreScale = 1;
+		Result.WorldPosition = CameraLocalWorldPosition + ResolvedView.WorldCameraOrigin;
+
+		return Result;
+	}
+
+	FVertexFactoryInterpolantsVSToDS VertexFactoryGetInterpolantsVSToDS(FVertexFactoryInput Input, FVertexFactoryIntermediates Intermediates, FMaterialVertexParameters VertexParameters)
+	{
+		FVertexFactoryInterpolantsVSToDS Interpolants;
+		Interpolants.InterpolantsVSToPS = VertexFactoryGetInterpolantsVSToPS(Input, Intermediates, VertexParameters);
+		return Interpolants;
+	}
+
+	FVertexFactoryInterpolantsVSToDS VertexFactoryInterpolate(FVertexFactoryInterpolantsVSToDS a, float aInterp, FVertexFactoryInterpolantsVSToDS b, float bInterp)
+	{
+		FVertexFactoryInterpolantsVSToDS O;
+	
+		// Do we really need to interpolate TangentToWorld2 here? It should be replaced by the
+		// interpolated normal from 'whatever' interpolation scheme we're using
+	
+		TESSELLATION_INTERPOLATE_MEMBER(InterpolantsVSToPS.TangentToWorld0.xyz);
+		TESSELLATION_INTERPOLATE_MEMBER(InterpolantsVSToPS.TangentToWorld2);
+#if INTERPOLATE_VERTEX_COLOR
+		TESSELLATION_INTERPOLATE_MEMBER(InterpolantsVSToPS.VertexColor);
+#endif
+
+#if NEEDS_PARTICLE_COLOR
+		TESSELLATION_INTERPOLATE_MEMBER(InterpolantsVSToPS.ParticleColor);
+#endif
+
+#if NUM_TEX_COORD_INTERPOLATORS
+		for (int i = 0; i < (NUM_TEX_COORD_INTERPOLATORS + 1) / 2; ++i)
+		{
+			TESSELLATION_INTERPOLATE_MEMBER(InterpolantsVSToPS.TexCoords[i]);
+		}
+#endif
+#if USE_DYNAMIC_PARAMETERS
+		TESSELLATION_INTERPOLATE_MEMBER(InterpolantsVSToPS.DynamicParameter);
+#endif	//USE_DYNAMIC_PARAMETERS
+
+		return O;
+	}
+
+	float3x3 VertexFactoryGetTangentToLocalDS(FVertexFactoryInterpolantsVSToDS Interpolants)
+	{
+		// This duplicates stuff already going on in GetMaterialTessellationParameters(), so
+		// maybe the hull shader could leverage that instead?
+		half3 TangentToWorld0 = Interpolants.InterpolantsVSToPS.TangentToWorld0.xyz;
+		half4 TangentToWorld2 = Interpolants.InterpolantsVSToPS.TangentToWorld2;
+		float3x3 TangentToWorld = AssembleTangentToWorld( TangentToWorld0, TangentToWorld2 );
+
+		return TangentToWorld;
+	}
+#endif // #if USING_TESSELLATION
+
+#if INSTANCED_STEREO
+
+// Instanced stereo sets the eye index explicitly for instanced geometry
+#if PARTICLE_MESH_INSTANCED
+uint InstancedEyeIndex;
+#endif
+
+uint VertexFactoryGetEyeIndex(uint InstanceId)
+{
+	#if PARTICLE_MESH_INSTANCED
+		return InstancedEyeIndex & bIsInstancedStereo;
+	#else
+		return InstanceId & 1;
+	#endif
+}
+#endif
+
+float4 VertexFactoryGetTranslatedPrimitiveVolumeBounds(FVertexFactoryInterpolantsVSToPS Interpolants)
+{
+	return 0;
+}
\ No newline at end of file
diff --git a/Engine/Shaders/Private/MiniFontCommon.ush b/Engine/Shaders/Private/MiniFontCommon.ush
new file mode 100644
--- /dev/null
+++ b/Engine/Shaders/Private/MiniFontCommon.ush
@@ -0,0 +1,124 @@
+// Copyright 1998-2017 Epic Games, Inc. All Rights Reserved.
+
+/*=============================================================================
+	TinyFontCommon.usf: TinyFont shared functions and structures.
+=============================================================================*/
+
+#ifndef __TINY_FONT_COMMON__
+#define __TINY_FONT_COMMON__
+
+Texture2D MiniFontTexture;
+
+// useful with PrintCharacter
+#define _0_ 0
+#define _1_ 1
+#define _2_ 2
+#define _3_ 3
+#define _4_ 4
+#define _5_ 5
+#define _6_ 6
+#define _7_ 7
+#define _8_ 8
+#define _9_ 9
+#define _A_ 10
+#define _B_ 11
+#define _C_ 12
+#define _D_ 13
+#define _E_ 14
+#define _F_ 15
+#define _G_ 16
+#define _H_ 17
+#define _I_ 18
+#define _J_ 19
+#define _K_ 20
+#define _L_ 21
+#define _M_ 22
+#define _N_ 23
+#define _O_ 24
+#define _P_ 25
+#define _Q_ 26
+#define _R_ 27
+#define _S_ 28
+#define _T_ 29
+#define _U_ 30
+#define _V_ 31
+#define _W_ 32
+#define _X_ 33
+#define _Y_ 34
+#define _Z_ 35
+
+// for printf debugging in the shader
+// @param LeftTop - is advanced, in pixels
+void PrintCharacter(int2 PixelPos, inout float3 OutColor, float3 FontColor, inout int2 LeftTop, int CharacterID)
+{
+	uint2 Rel = (uint2)(PixelPos - LeftTop);
+
+	FLATTEN if (Rel.x < 8 && Rel.y < 8)
+	{
+		OutColor = lerp(OutColor, FontColor, MiniFontTexture.Load(int3(CharacterID * 8 + Rel.x, Rel.y, 0)).r);
+	}
+
+	LeftTop.x += 8;
+}
+
+// only for positive numbers
+// @param DigitValue - e.g. 1 for frist digit before period, 10 for second, 0.1 for first digit behind period
+uint ExtractDigitFromFloat(float Number, float DigitValue)
+{
+	uint Temp = (uint)(Number / DigitValue);
+
+#if ES2_PROFILE
+	// Integer modulo/remainder is not an allowed operation on ES 2
+	// todo: move function to central spot, here we assume unsigned
+	return Temp - (Temp / 10) * 10;
+#else
+	return Temp % 10;
+#endif
+}
+
+
+// for printf debugging in the shader, has to be positive
+// outputs a float number in the form: xxx.yyy
+// @param LeftTop - in pixels
+void PrintFloat(int2 PixelPos, inout float3 OutColor, float3 FontColor, int2 LeftTop, float Number)
+{
+	int2 Cursor = LeftTop;
+
+	float BorderDistance = ComputeDistanceToRect(PixelPos, LeftTop, int2(7, 1) * 8 - 1);
+
+	// black border around number
+	//	OutColor = lerp(0, OutColor, saturate(BorderDistance - 2));
+
+	// before period
+	PrintCharacter(PixelPos, OutColor, FontColor, Cursor, ExtractDigitFromFloat(Number, 100));
+	PrintCharacter(PixelPos, OutColor, FontColor, Cursor, ExtractDigitFromFloat(Number, 10));
+	PrintCharacter(PixelPos, OutColor, FontColor, Cursor, ExtractDigitFromFloat(Number, 1));
+	// period
+	PrintCharacter(PixelPos, OutColor, FontColor, Cursor, 512 / 8 - 3);
+	// after period
+	PrintCharacter(PixelPos, OutColor, FontColor, Cursor, ExtractDigitFromFloat(Number, 0.1));
+	PrintCharacter(PixelPos, OutColor, FontColor, Cursor, ExtractDigitFromFloat(Number, 0.01));
+	PrintCharacter(PixelPos, OutColor, FontColor, Cursor, ExtractDigitFromFloat(Number, 0.001));
+}
+
+// for printf debugging in the shader, has to be positive
+// outputs a float number in the form: xxx.yyy
+// @param LeftTop - in pixels
+void PrintFloatNoFraction(int2 PixelPos, inout float3 OutColor, float3 FontColor, int2 LeftTop, float Number)
+{
+	int2 Cursor = LeftTop;
+
+	float BorderDistance = ComputeDistanceToRect(PixelPos, LeftTop, int2(7, 1) * 8 - 1);
+
+	// black border around number
+	//	OutColor = lerp(0, OutColor, saturate(BorderDistance - 2));
+
+	// before period
+	PrintCharacter(PixelPos, OutColor, FontColor, Cursor, ExtractDigitFromFloat(Number, 10000));
+	PrintCharacter(PixelPos, OutColor, FontColor, Cursor, ExtractDigitFromFloat(Number, 1000));
+	PrintCharacter(PixelPos, OutColor, FontColor, Cursor, ExtractDigitFromFloat(Number, 100));
+	PrintCharacter(PixelPos, OutColor, FontColor, Cursor, ExtractDigitFromFloat(Number, 10));
+	PrintCharacter(PixelPos, OutColor, FontColor, Cursor, ExtractDigitFromFloat(Number, 1));
+}
+
+#endif // __TINY_FONT_COMMON__
diff --git a/Engine/Shaders/Private/MobileBasePassCommon.ush b/Engine/Shaders/Private/MobileBasePassCommon.ush
new file mode 100644
--- /dev/null
+++ b/Engine/Shaders/Private/MobileBasePassCommon.ush
@@ -0,0 +1,53 @@
+// Copyright 1998-2017 Epic Games, Inc. All Rights Reserved.
+
+/*=============================================================================
+	MobileBasePassCommon.usf: Base pass definitions used by both vertex and pixel shader
+=============================================================================*/
+
+#define USE_PS_CLIP_PLANE PROJECT_ALLOW_GLOBAL_CLIP_PLANE
+
+#undef NEEDS_LIGHTMAP_COORDINATE
+
+#define NEEDS_LIGHTMAP_COORDINATE		(LQ_TEXTURE_LIGHTMAP)
+#define NEEDS_LIGHTMAP					(NEEDS_LIGHTMAP_COORDINATE)
+
+#ifndef PROJECT_MOBILE_DISABLE_VERTEX_FOG
+	#define PROJECT_MOBILE_DISABLE_VERTEX_FOG 0
+#endif
+
+#define USE_VERTEX_FOG					(!PROJECT_MOBILE_DISABLE_VERTEX_FOG)
+
+#define PACK_INTERPOLANTS (USE_VERTEX_FOG && NUM_VF_PACKED_INTERPOLANTS > 0 && (ES2_PROFILE || ES3_1_PROFILE))
+#define LANDSCAPE_BUG_WORKAROUND (IOS && IS_MOBILE_BASEPASS_VERTEX_SHADER && PACK_INTERPOLANTS)
+
+#if COMPILER_METAL
+// @todo-rco: @todo metal: FIXME!
+#define INTERP_TYPE float4
+#else
+#define INTERP_TYPE half4
+#endif
+
+struct FSharedMobileBasePassInterpolants
+{
+#if USE_VERTEX_FOG && !PACK_INTERPOLANTS
+	INTERP_TYPE VertexFog			: TEXCOORD7;
+#endif
+#if LANDSCAPE_BUG_WORKAROUND
+	INTERP_TYPE DummyInterp		: DUMMY_INTERP;
+#endif
+	float4 PixelPosition	: TEXCOORD8; // xyz = world position, w = clip z
+
+#if USE_WORLD_POSITION_EXCLUDING_SHADER_OFFSETS
+	float3 PixelPositionExcludingWPO	: TEXCOORD9;
+#endif
+
+#if USE_PS_CLIP_PLANE
+	float OutClipDistance : OUTCLIPDIST;
+#endif
+
+#if MOBILE_MULTI_VIEW
+	float MultiViewId : VIEW_ID;
+#endif
+};
+
+#define FMobileBasePassInterpolantsVSToPS FSharedMobileBasePassInterpolants
