Branch snapshot for CL 3555263

diff --git a/Engine/Shaders/Private/MonteCarlo.ush b/Engine/Shaders/Private/MonteCarlo.ush
new file mode 100644
--- /dev/null
+++ b/Engine/Shaders/Private/MonteCarlo.ush
@@ -0,0 +1,217 @@
+// Copyright 1998-2017 Epic Games, Inc. All Rights Reserved.
+
+#pragma once
+
+#include "Common.ush"
+
+/*=============================================================================
+	MonteCarlo.usf: Monte Carlo integration of distributions
+=============================================================================*/
+
+float3x3 GetTangentBasis( float3 TangentZ )
+{
+	float3 UpVector = abs(TangentZ.z) < 0.999 ? float3(0,0,1) : float3(1,0,0);
+	float3 TangentX = normalize( cross( UpVector, TangentZ ) );
+	float3 TangentY = cross( TangentZ, TangentX );
+	return float3x3( TangentX, TangentY, TangentZ );
+}
+
+float3 TangentToWorld( float3 Vec, float3 TangentZ )
+{
+	return mul( Vec, GetTangentBasis( TangentZ ) );
+}
+
+float2 Hammersley( uint Index, uint NumSamples, uint2 Random )
+{
+	float E1 = frac( (float)Index / NumSamples + float( Random.x & 0xffff ) / (1<<16) );
+	float E2 = float( ReverseBits32(Index) ^ Random.y ) * 2.3283064365386963e-10;
+	return float2( E1, E2 );
+}
+
+///////
+// Typical usage of the Sobol functions for one or more points per pixel
+//   uint2 SobolBase = SobolPixel(uint2(SvPosition.xy));                    // First sample for this pixel
+//   for(int i = 0; i < N; ++i)
+//     float2 Point = float2(SobolIndex(SobolBase, i)) / 0x10000;           // Points as [0,1) x [0,1)
+//
+// Typical usage for one or more points per frame
+//   uint2 SobolBase = SobolPixel(uint2(SvPosition.xy));                    // Frame 0, point 0 for this pixel
+//   uint2 SobolFrame = SobolIndex(SobolBase, View.StateFrameIndexMod8, 3); // Least significant bits for frame
+//   for(int i = 0; i < N; ++i)
+//     float2 Point = float2(SobolIndex(SobolFrame, i << 3)) / 0x10000;     // Higher-order bits for point within frame
+//
+// For additional independent point sets
+//   uint2 SobolBase2 = SobolBase ^ RandSeed; // where RandSeed is a uint2 with values in 0 to 0xffff
+
+// Compute a Sobol-distributed point from a 256x256 pixel grid one pixel in that grid
+// @param Pixel Pixel/cell position in the 256x256 grid
+// @return Sobol position relative to the pixel corner, with components in the range 0 to 0xffff
+uint2 SobolPixel(uint2 Pixel)
+{
+	// look up for pixel
+	int3 SobolLo = int3(Pixel & 0xfu, 0);
+	int3 SobolHi = int3((Pixel >> 4u) & 0xfu, 0) + int3(16, 0, 0);
+	uint Packed = View.SobolSamplingTexture.Load(SobolLo) ^ View.SobolSamplingTexture.Load(SobolHi);
+	return uint2(Packed, Packed << 8u) & 0xff00u;
+}
+
+// Evaluate additional Sobol points within a pixel
+// @param Base  Base Sobol point for this pixel
+// @param Index Which 2D Sobol point to return
+// @param Bits  Optional max bits in index (to avoid extra calculation)
+// @return Sobol position in the range 0 to 0xffff
+uint2 SobolIndex(uint2 Base, int Index, int Bits = 10)
+{
+	uint2 SobolNumbers[10] = {
+		uint2(0x8680u, 0x4c80u), uint2(0xf240u, 0x9240u), uint2(0x8220u, 0x0e20u), uint2(0x4110u, 0x1610u), uint2(0xa608u, 0x7608u),
+		uint2(0x8a02u, 0x280au), uint2(0xe204u, 0x9e04u), uint2(0xa400u, 0x4682u), uint2(0xb700u, 0x9817u), uint2(0xe300u, 0xa74du),
+	};
+
+	uint2 Result = Base;
+	UNROLL for (int b = 0; b < 10 && b < Bits; ++b)
+	{
+		Result ^= (Index & (1 << b)) ? SobolNumbers[b] : 0;
+	}
+	return Result;
+}
+
+// convert 16-bit random value in [0,0xffff] into a position in the unit circle
+// based on the approximate equal area transform from
+// http://marc-b-reynolds.github.io/math/2017/01/08/SquareDisc.html
+float2 RandToCircle(uint2 Rand)
+{
+	float2 sf = float2(Rand) * (sqrt(2.) / 0xffff) - sqrt(0.5);	// map 0..0xffff to -sqrt(0.5)..sqrt(0.5)
+	float2 sq = sf*sf;
+	float root = sqrt(2.*max(sq.x, sq.y) - min(sq.x, sq.y));
+	if (sq.x > sq.y)
+	{
+		sf.x = sf.x > 0 ? root : -root;
+	}
+	else
+	{
+		sf.y = sf.y > 0 ? root : -root;
+	}
+	return sf;
+}
+
+///////
+float4 UniformSampleSphere( float2 E )
+{
+	float Phi = 2 * PI * E.x;
+	float CosTheta = 1 - 2 * E.y;
+	float SinTheta = sqrt( 1 - CosTheta * CosTheta );
+
+	float3 H;
+	H.x = SinTheta * cos( Phi );
+	H.y = SinTheta * sin( Phi );
+	H.z = CosTheta;
+
+	float PDF = 1.0 / (4 * PI);
+
+	return float4( H, PDF );
+}
+
+float4 UniformSampleHemisphere( float2 E )
+{
+	float Phi = 2 * PI * E.x;
+	float CosTheta = E.y;
+	float SinTheta = sqrt( 1 - CosTheta * CosTheta );
+
+	float3 H;
+	H.x = SinTheta * cos( Phi );
+	H.y = SinTheta * sin( Phi );
+	H.z = CosTheta;
+
+	float PDF = 1.0 / (2 * PI);
+
+	return float4( H, PDF );
+}
+
+float2 UniformSampleDisk(float2 Random)
+{
+	const float Theta = 2.0f * (float)PI * Random.x;
+	const float Radius = sqrt(Random.y);
+	return float2(Radius * cos(Theta), Radius * sin(Theta));
+}
+
+float4 CosineSampleHemisphere( float2 E )
+{
+	float Phi = 2 * PI * E.x;
+	float CosTheta = sqrt( E.y );
+	float SinTheta = sqrt( 1 - CosTheta * CosTheta );
+
+	float3 H;
+	H.x = SinTheta * cos( Phi );
+	H.y = SinTheta * sin( Phi );
+	H.z = CosTheta;
+
+	float PDF = CosTheta / PI;
+
+	return float4( H, PDF );
+}
+
+float4 UniformSampleCone( float2 E, float CosThetaMax )
+{
+	float Phi = 2 * PI * E.x;
+	float CosTheta = lerp( CosThetaMax, 1, E.y );
+	float SinTheta = sqrt( 1 - CosTheta * CosTheta );
+
+	float3 L;
+	L.x = SinTheta * cos( Phi );
+	L.y = SinTheta * sin( Phi );
+	L.z = CosTheta;
+
+	float PDF = 1.0 / ( 2 * PI * (1 - CosThetaMax) );
+
+	return float4( L, PDF );
+}
+
+float4 ImportanceSampleBlinn( float2 E, float Roughness )
+{
+	float m = Roughness * Roughness;
+	float n = 2 / (m*m) - 2;
+
+	float Phi = 2 * PI * E.x;
+	float CosTheta = ClampedPow( E.y, 1 / (n + 1) );
+	float SinTheta = sqrt( 1 - CosTheta * CosTheta );
+
+	float3 H;
+	H.x = SinTheta * cos( Phi );
+	H.y = SinTheta * sin( Phi );
+	H.z = CosTheta;
+
+	float D = (n+2)/ (2*PI) * ClampedPow( CosTheta, n );
+	float PDF = D * CosTheta;
+
+	return float4( H, PDF );
+}
+
+float4 ImportanceSampleGGX( float2 E, float Roughness )
+{
+	float m = Roughness * Roughness;
+	float m2 = m * m;
+
+	float Phi = 2 * PI * E.x;
+	float CosTheta = sqrt( (1 - E.y) / ( 1 + (m2 - 1) * E.y ) );
+	float SinTheta = sqrt( 1 - CosTheta * CosTheta );
+
+	float3 H;
+	H.x = SinTheta * cos( Phi );
+	H.y = SinTheta * sin( Phi );
+	H.z = CosTheta;
+	
+	float d = ( CosTheta * m2 - CosTheta ) * CosTheta + 1;
+	float D = m2 / ( PI*d*d );
+	float PDF = D * CosTheta;
+
+	return float4( H, PDF );
+}
+
+// Multiple importance sampling power heuristic of two functions with a power of two. 
+// [Veach 1997, "Robust Monte Carlo Methods for Light Transport Simulation"]
+float MISWeight( uint Num, float PDF, uint OtherNum, float OtherPDF )
+{
+	float Weight = Num * PDF;
+	float OtherWeight = OtherNum * OtherPDF;
+	return Weight * Weight / (Weight * Weight + OtherWeight * OtherWeight);
+}
\ No newline at end of file
diff --git a/Engine/Shaders/Private/NiagaraMeshVertexFactory.ush b/Engine/Shaders/Private/NiagaraMeshVertexFactory.ush
new file mode 100644
--- /dev/null
+++ b/Engine/Shaders/Private/NiagaraMeshVertexFactory.ush
@@ -0,0 +1,666 @@
+// Copyright 1998-2017 Epic Games, Inc. All Rights Reserved.
+
+/*=============================================================================
+	MeshParticleVertexFactory.usf: Mesh particle vertex factory shader code.
+=============================================================================*/
+
+#include "VertexFactoryCommon.ush"
+
+// These are only enabled for SM5 due to limited interpolators elsewhere
+#define USE_PARTICLE_NODES (FEATURE_LEVEL >= FEATURE_LEVEL_SM5)
+
+// Only support previous transform for deferred feature level
+#define USE_PREVIOUS_TRANSFORM (FEATURE_LEVEL >= FEATURE_LEVEL_SM4)
+
+float4x4 PreviousLocalToWorld;
+
+
+#if USE_PREVIOUS_TRANSFORM
+#if NIAGARA_MESH_INSTANCED
+	Buffer<float4> PrevTransformBuffer;
+#else
+	float4 PrevTransform0;
+	float4 PrevTransform1;
+	float4 PrevTransform2;
+#endif
+#endif
+	
+#if !NIAGARA_MESH_INSTANCED
+	float4 Transform1;
+	float4 Transform2;
+	float4 Transform3;
+	float4 SubUVParams;
+	float SubUVLerp;
+	float4 ParticleDirection;
+	float RelativeTime;
+	float4 DynamicParameter;
+	float4 ParticleColor;
+#endif
+
+struct FVertexFactoryInput
+{
+	float4	Position	: ATTRIBUTE0;
+	half3	TangentX	: ATTRIBUTE1;
+	// TangentZ.w contains sign of tangent basis determinant
+	half4	TangentZ	: ATTRIBUTE2;
+	half4	VertexColor : ATTRIBUTE3;
+
+#if NUM_MATERIAL_TEXCOORDS_VERTEX
+	float2	TexCoords[NUM_MATERIAL_TEXCOORDS_VERTEX] : ATTRIBUTE4;
+#elif USE_PARTICLE_SUBUVS
+	float2	TexCoords[1] : ATTRIBUTE4;
+#endif
+
+#if NIAGARA_MESH_INSTANCED
+	float4 Transform1 : ATTRIBUTE8;
+	float4 Transform2 : ATTRIBUTE9;
+	float4 Transform3 : ATTRIBUTE10;
+
+	int4 SubUVParams : ATTRIBUTE11;
+
+	/** X: SubUV lerp Y: Particle relative time. */
+	float2 SubUVLerpAndRelTime : ATTRIBUTE12;
+
+	#if USE_DYNAMIC_PARAMETERS
+		float4 DynamicParameter	: ATTRIBUTE13;
+	#endif	//USE_DYNAMIC_PARAMETERS
+
+	float4 ParticleColor : ATTRIBUTE14;
+
+	/** The velocity of the particle, XYZ: direction, W: speed. */
+	float4 ParticleVelocity : ATTRIBUTE15;
+#endif
+
+/** Optional instance ID for vertex layered rendering */
+#if (NIAGARA_MESH_INSTANCED || (ONEPASS_POINTLIGHT_SHADOW && USING_VERTEX_SHADER_LAYER))
+	uint InstanceId	: SV_InstanceID;
+#endif
+};
+
+struct FVertexFactoryInterpolantsVSToPS
+{
+	TANGENTTOWORLD_INTERPOLATOR_BLOCK
+
+#if USE_PARTICLE_SUBUVS
+	float4 SubUV0AndTexCoord0 : TEXCOORD1;
+	float4 SubUV1AndLerp : TEXCOORD2;
+#else
+	#if NUM_MATERIAL_TEXCOORDS
+			#if (ES2_PROFILE || ES3_1_PROFILE)
+				// Avoid dependent texture fetches, put all UVs in xy
+				float2	TexCoords[NUM_MATERIAL_TEXCOORDS] : TEXCOORD0;
+			#else
+				float4	TexCoords[(NUM_MATERIAL_TEXCOORDS+1)/2]	: TEXCOORD0;
+			#endif
+	#endif
+#endif
+
+#if INTERPOLATE_VERTEX_COLOR
+	float4	VertexColor : COLOR0;
+#endif
+
+#if NEEDS_PARTICLE_COLOR
+	float4	ParticleColor : COLOR1;
+#endif
+
+#if USE_DYNAMIC_PARAMETERS
+	float4 DynamicParameter : COLOR2;
+#endif
+	
+#if USE_PARTICLE_NODES
+	/** Particle position in camera-centered translated world space */
+	float3 ParticleTranslatedWorldPosition	: PARTICLE_POSITION;
+
+	/** The velocity of the particle, XYZ: direction, W: speed. */
+	float4 ParticleVelocity	: PARTICLE_VELOCITY;
+
+	/** Relative alive time of the particle */
+	float ParticleTime : PARTICLE_TIME;
+#if NEEDS_PARTICLE_TRANSFORM
+	float4 LocalToWorld[4] : PARTICLE_TRANSFORM;
+#endif
+#endif
+
+#if INSTANCED_STEREO
+	nointerpolation uint PackedEyeIndex : PACKED_EYE_INDEX;
+#endif
+};
+
+struct FVertexFactoryIntermediates
+{
+	/** The color of the vertex. */
+	float4 VertexColor;
+
+	/** The color of the particle. */
+	float4 ParticleColor;
+
+	/** The texture coordinates for the vertex. */
+#if NUM_MATERIAL_TEXCOORDS_VERTEX
+	float2	TexCoords[NUM_MATERIAL_TEXCOORDS_VERTEX];
+#endif 
+#if USE_PARTICLE_SUBUVS
+	float4 SubUVCoords;
+	float SubUVLerp;
+#endif
+#if USE_DYNAMIC_PARAMETERS
+	/** Optional dynamic parameter for the particle. */
+	float4 DynamicParameter;
+#endif	//USE_DYNAMIC_PARAMETERS
+
+	/** The velocity of the particle, XYZ: direction, W: speed. */
+	float4 ParticleVelocity;
+
+	/** Particle position in camera-centered translated world space. */
+	float3 ParticleTranslatedWorldPosition;
+	
+	/** Relative time. */
+	float RelativeTime;
+
+#if NEEDS_PARTICLE_TRANSFORM
+	float4 LocalToWorld[4];
+#endif
+};
+
+float4x4 GetParticleTransform(FVertexFactoryInput Input)
+{
+	//Transform is transposed on CPU and position is packed into W. Transpose again to get the correct transform back.
+#if NIAGARA_MESH_INSTANCED
+	return transpose(float4x4(Input.Transform1, Input.Transform2, Input.Transform3, float4(0, 0, 0, 1)));
+#else
+	return transpose(float4x4(Transform1, Transform2, Transform3, float4(0, 0, 0, 1)));
+#endif
+}
+
+#if USE_PREVIOUS_TRANSFORM
+float4x4 GetParticlePrevTransform(FVertexFactoryInput Input)
+{
+#if NIAGARA_MESH_INSTANCED
+	uint TransformBaseIdx = GetLayerInstanceID(Input.InstanceId) * 3;
+
+	float4x4 PrevTransform4x4 = float4x4(
+		PrevTransformBuffer[TransformBaseIdx + 0], 
+		PrevTransformBuffer[TransformBaseIdx + 1], 
+		PrevTransformBuffer[TransformBaseIdx + 2], 
+		float4(0, 0, 0, 1)
+		);
+#else
+	float4x4 PrevTransform4x4 = float4x4(
+		PrevTransform0, 
+		PrevTransform1, 
+		PrevTransform2, 
+		float4(0, 0, 0, 1)
+		);
+#endif
+	//Transform is transposed on CPU and position is packed into W. Transpose again to get the correct transform back.
+	return transpose(PrevTransform4x4);
+}
+#endif
+
+/** Converts from vertex factory specific interpolants to a FMaterialPixelParameters, which is used by material inputs. */
+FMaterialPixelParameters GetMaterialPixelParameters(FVertexFactoryInterpolantsVSToPS Interpolants, float4 SvPosition)
+{
+	// GetMaterialPixelParameters is responsible for fully initializing the result
+	FMaterialPixelParameters Result = MakeInitializedMaterialPixelParameters();
+
+#if USE_PARTICLE_SUBUVS
+	#if NUM_MATERIAL_TEXCOORDS
+		UNROLL
+		for( int CoordinateIndex = 0; CoordinateIndex < NUM_MATERIAL_TEXCOORDS; CoordinateIndex++ )
+		{
+			Result.TexCoords[CoordinateIndex] = Interpolants.SubUV0AndTexCoord0.zw;
+		}
+	#endif
+	Result.Particle.SubUVCoords[0] = Interpolants.SubUV0AndTexCoord0.xy;
+	Result.Particle.SubUVCoords[1] = Interpolants.SubUV1AndLerp.xy;
+	Result.Particle.SubUVLerp = Interpolants.SubUV1AndLerp.z;
+#elif NUM_MATERIAL_TEXCOORDS
+	#if (ES2_PROFILE || ES3_1_PROFILE)
+		UNROLL
+		for( int CoordinateIndex = 0; CoordinateIndex < NUM_MATERIAL_TEXCOORDS; CoordinateIndex++ )
+		{
+			Result.TexCoords[CoordinateIndex] = Interpolants.TexCoords[CoordinateIndex].xy;
+		}
+	#else
+		UNROLL
+		for(int CoordinateIndex = 0;CoordinateIndex < NUM_MATERIAL_TEXCOORDS;CoordinateIndex += 2)
+		{
+			Result.TexCoords[CoordinateIndex] = Interpolants.TexCoords[CoordinateIndex/2].xy;
+
+			if(CoordinateIndex + 1 < NUM_MATERIAL_TEXCOORDS)
+			{
+				Result.TexCoords[CoordinateIndex + 1] = Interpolants.TexCoords[CoordinateIndex/2].wz;
+			}
+		}
+	#endif
+#endif
+
+	half3 TangentToWorld0 = Interpolants.TangentToWorld0.xyz;
+	half4 TangentToWorld2 = Interpolants.TangentToWorld2;
+	Result.UnMirrored = TangentToWorld2.w;
+
+#if INTERPOLATE_VERTEX_COLOR
+	Result.VertexColor = Interpolants.VertexColor;
+#else
+	Result.VertexColor = 0;
+#endif
+
+#if NEEDS_PARTICLE_COLOR
+	Result.Particle.Color = Interpolants.ParticleColor;
+#endif
+
+#if USE_DYNAMIC_PARAMETERS
+	Result.Particle.DynamicParameter = Interpolants.DynamicParameter;
+#else
+	Result.Particle.DynamicParameter = float4(1,1,1,1);
+#endif //USE_DYNAMIC_PARAMETERS
+
+#if USE_PARTICLE_NODES
+	Result.Particle.TranslatedWorldPositionAndSize.xyz = Interpolants.ParticleTranslatedWorldPosition.xyz;
+	Result.Particle.TranslatedWorldPositionAndSize.w = 1;
+	Result.Particle.Velocity = Interpolants.ParticleVelocity;
+	Result.Particle.RelativeTime = Interpolants.ParticleTime;
+#if NEEDS_PARTICLE_TRANSFORM
+	Result.Particle.LocalToWorld[0] = Interpolants.LocalToWorld[0];
+	Result.Particle.LocalToWorld[1] = Interpolants.LocalToWorld[1];
+	Result.Particle.LocalToWorld[2] = Interpolants.LocalToWorld[2];
+	Result.Particle.LocalToWorld[3] = Interpolants.LocalToWorld[3];
+#endif
+#endif
+
+	Result.Particle.MotionBlurFade = 1.0f;
+	Result.TangentToWorld = AssembleTangentToWorld( TangentToWorld0, TangentToWorld2 );
+	Result.TwoSidedSign = 1;
+
+	return Result;
+}
+
+/** Converts from vertex factory specific input to a FMaterialVertexParameters, which is used by vertex shader material inputs. */
+FMaterialVertexParameters GetMaterialVertexParameters(FVertexFactoryInput Input, FVertexFactoryIntermediates Intermediates, float3 WorldPosition, float3x3 TangentToLocal)
+{
+	FMaterialVertexParameters Result = (FMaterialVertexParameters)0;
+	Result.WorldPosition = WorldPosition;
+	Result.VertexColor = Intermediates.VertexColor;
+	Result.PreSkinnedPosition = Input.Position.xyz;
+	Result.PreSkinnedNormal = Input.TangentZ.xyz * 2.f - 1.f;
+
+	Result.Particle.TranslatedWorldPositionAndSize.xyz = Intermediates.ParticleTranslatedWorldPosition.xyz;
+	Result.Particle.TranslatedWorldPositionAndSize.w = 1;
+	Result.Particle.Velocity = Intermediates.ParticleVelocity;
+	Result.Particle.RelativeTime = Intermediates.RelativeTime;
+
+#if USE_DYNAMIC_PARAMETERS
+	Result.Particle.DynamicParameter = Intermediates.DynamicParameter;
+#endif	//USE_DYNAMIC_PARAMETERS
+
+	float4x4 Transform = GetParticleTransform(Input);
+	Result.InstanceLocalToWorld = Transform;
+#if NEEDS_PARTICLE_TRANSFORM
+	Result.Particle.LocalToWorld[0] = Transform[0];
+	Result.Particle.LocalToWorld[1] = Transform[1];
+	Result.Particle.LocalToWorld[2] = Transform[2];
+	Result.Particle.LocalToWorld[3] = Transform[3];
+#endif
+
+	Result.TangentToWorld = mul(TangentToLocal, (float3x3)Transform ); 
+
+#if NUM_MATERIAL_TEXCOORDS_VERTEX
+	for(int CoordinateIndex = 0; CoordinateIndex < NUM_MATERIAL_TEXCOORDS_VERTEX; CoordinateIndex++)
+	{
+		Result.TexCoords[CoordinateIndex] = Intermediates.TexCoords[CoordinateIndex];
+	}
+#endif
+
+	return Result;
+}
+
+#if NEEDS_LIGHTMAP_COORDINATE
+void GetLightMapCoordinates(FVertexFactoryInterpolantsVSToPS Interpolants, out float2 LightmapUV0, out float2 LightmapUV1)
+{
+	LightmapUV0 = LightmapUV1 = 0;
+}
+
+half2 GetShadowMapCoordinate(FVertexFactoryInterpolantsVSToPS Interpolants)
+{
+	return 0;
+}
+#endif
+
+
+float4 CalcWorldPosition(FVertexFactoryInput Input)
+{
+	float4x4 Transform = GetParticleTransform(Input);
+	float3 WorldPosition = mul(Input.Position, Transform).xyz + ResolvedView.PreViewTranslation;
+	return float4(WorldPosition, Input.Position.w);
+}
+
+float4 CalcPrevWorldPosition(FVertexFactoryInput Input)
+{
+#if USE_PREVIOUS_TRANSFORM
+	if (MeshParticleVF.PrevTransformAvailable > 0)
+	{
+		float4x4 PrevTransform = GetParticlePrevTransform(Input);
+		float3 PrevWorldPosition = mul(Input.Position, PrevTransform).xyz + ResolvedView.PrevPreViewTranslation;
+		return float4(PrevWorldPosition, Input.Position.w);
+	}
+	else
+#endif
+	{
+		// TODO: this is not the previous position
+		// need to add   - float4(Input.ParticleVelocity.xyz*Input.ParticleVelocity.w, 0.0)  ?
+		return CalcWorldPosition(Input)- float4(Input.ParticleVelocity.xyz*Input.ParticleVelocity.w*MeshParticleVF.DeltaSeconds, 0.0);
+	}
+}
+
+FVertexFactoryIntermediates GetVertexFactoryIntermediates(FVertexFactoryInput Input)
+{
+	FVertexFactoryIntermediates Intermediates = (FVertexFactoryIntermediates)0;
+	// Swizzle vertex color.
+	Intermediates.VertexColor = Input.VertexColor FCOLOR_COMPONENT_SWIZZLE;
+
+	// World position and localtoworld transform.
+#if NIAGARA_MESH_INSTANCED
+	float3 ParticleWorldPosition = float3( Input.Transform1.w, Input.Transform2.w, Input.Transform3.w);
+	float4x4 LocalToWorldMat = float4x4(Input.Transform1, Input.Transform2, Input.Transform3, float4(0.0f, 0.0f, 0.0f, 1.0f));
+#if NEEDS_PARTICLE_TRANSFORM
+	Intermediates.LocalToWorld[0] = LocalToWorldMat[0];
+	Intermediates.LocalToWorld[1] = LocalToWorldMat[1];
+	Intermediates.LocalToWorld[2] = LocalToWorldMat[2];
+	Intermediates.LocalToWorld[3] = LocalToWorldMat[3];
+#endif
+#else
+	float3 ParticleWorldPosition = float3( Transform1.w, Transform2.w, Transform3.w);
+#endif
+
+	Intermediates.ParticleTranslatedWorldPosition = ParticleWorldPosition + ResolvedView.PreViewTranslation.xyz;
+
+#if NIAGARA_MESH_INSTANCED
+	Intermediates.ParticleColor = Input.ParticleColor;
+#else
+	Intermediates.ParticleColor = ParticleColor;
+#endif
+
+#if NUM_MATERIAL_TEXCOORDS_VERTEX
+	for(int CoordinateIndex = 0; CoordinateIndex < NUM_MATERIAL_TEXCOORDS_VERTEX; CoordinateIndex++)
+	{
+		Intermediates.TexCoords[CoordinateIndex] = Input.TexCoords[CoordinateIndex].xy;
+	}
+#endif
+
+#if USE_PARTICLE_SUBUVS
+	#if NIAGARA_MESH_INSTANCED
+		Intermediates.SubUVCoords.xy = (float2( Input.SubUVParams.x, Input.SubUVParams.y ) + Input.TexCoords[0].xy) * MeshParticleVF.SubImageSize.xy;
+		Intermediates.SubUVCoords.zw = (float2( Input.SubUVParams.z, Input.SubUVParams.w ) + Input.TexCoords[0].xy) * MeshParticleVF.SubImageSize.xy;
+		Intermediates.SubUVLerp = Input.SubUVLerpAndRelTime.x;
+	#else
+		Intermediates.SubUVCoords.xy = (float2( SubUVParams.x, SubUVParams.y ) + Input.TexCoords[0].xy) * MeshParticleVF.SubImageSize.xy;
+		Intermediates.SubUVCoords.zw = (float2( SubUVParams.z, SubUVParams.w ) + Input.TexCoords[0].xy) * MeshParticleVF.SubImageSize.xy;
+		Intermediates.SubUVLerp = SubUVLerp.x;
+	#endif
+#endif 
+
+#if USE_DYNAMIC_PARAMETERS
+	#if NIAGARA_MESH_INSTANCED
+		Intermediates.DynamicParameter = Input.DynamicParameter;
+	#else
+		Intermediates.DynamicParameter = DynamicParameter;
+	#endif
+#endif	//USE_DYNAMIC_PARAMETERS
+
+	Intermediates.ParticleTranslatedWorldPosition = ParticleWorldPosition.xyz + ResolvedView.PreViewTranslation.xyz;
+
+#if NIAGARA_MESH_INSTANCED
+	Intermediates.ParticleVelocity = Input.ParticleVelocity;
+	Intermediates.RelativeTime = Input.SubUVLerpAndRelTime.y;
+#else
+	Intermediates.ParticleVelocity = ParticleDirection;
+	Intermediates.RelativeTime = RelativeTime;
+#endif
+
+	return Intermediates;
+}
+
+/**
+* Get the 3x3 tangent basis vectors for this vertex factory
+* this vertex factory will calculate the binormal on-the-fly
+*
+* @param Input - vertex input stream structure
+* @return 3x3 matrix
+*/
+float3x3 VertexFactoryGetTangentToLocal( FVertexFactoryInput Input, FVertexFactoryIntermediates Intermediates )
+{
+	float3x3 Result=0;	
+	float3 TangentX = TangentBias(Input.TangentX);
+	float4 TangentZ = TangentBias(Input.TangentZ);
+
+	// pass-thru the tangent
+	Result[0] = TangentX;
+	// pass-thru the normal
+	Result[2] = TangentZ.xyz;	
+
+	
+	// derive the binormal by getting the cross product of the normal and tangent
+	Result[1] = cross(Result[2], Result[0]) * TangentZ.w;	
+	// Recalculate TangentX off of the other two vectors
+	// This corrects quantization error since TangentX was passed in as a quantized vertex input
+	// The error shows up most in specular off of a mesh with a smoothed UV seam (normal is smooth, but tangents vary across the seam)
+	Result[0] = cross(Result[1], Result[2]) * TangentZ.w;
+
+	return Result;
+}
+
+// @return translated world position
+float4 VertexFactoryGetWorldPosition(FVertexFactoryInput Input, FVertexFactoryIntermediates Intermediates)
+{
+	return CalcWorldPosition(Input);
+}
+
+float4 VertexFactoryGetRasterizedWorldPosition(FVertexFactoryInput Input, FVertexFactoryIntermediates Intermediates, float4 InWorldPosition)
+{
+	return InWorldPosition;
+}
+
+float3 VertexFactoryGetPositionForVertexLighting(FVertexFactoryInput Input, FVertexFactoryIntermediates Intermediates, float3 TranslatedWorldPosition)
+{
+	return TranslatedWorldPosition;
+}
+
+void CalcTangentToWorld(FVertexFactoryInput Input, FVertexFactoryIntermediates Intermediates, out float3 TangentToWorld0, out float4 TangentToWorld2)
+{
+	float3x3 TangentToLocal = VertexFactoryGetTangentToLocal(Input, Intermediates);
+
+	float4x4 Transform = GetParticleTransform(Input);
+	float3x3 TangentToWorld = mul(TangentToLocal, (float3x3)Transform);
+
+	// Normalize to remove scaling. Incorrect but faster than computing the inverse-transpose.
+	TangentToWorld0 = normalize(TangentToWorld[0]);
+	TangentToWorld2 = float4(normalize(TangentToWorld[2]), TangentBias(Input.TangentZ.w) * Primitive.LocalToWorldDeterminantSign);
+}
+
+FVertexFactoryInterpolantsVSToPS VertexFactoryGetInterpolantsVSToPS(FVertexFactoryInput Input, FVertexFactoryIntermediates Intermediates, FMaterialVertexParameters VertexParameters)
+{
+	FVertexFactoryInterpolantsVSToPS Interpolants;
+
+#if USE_PARTICLE_SUBUVS
+	Interpolants.SubUV0AndTexCoord0.xy = Intermediates.SubUVCoords.xy;
+	Interpolants.SubUV1AndLerp.xy = Intermediates.SubUVCoords.zw;
+	Interpolants.SubUV1AndLerp.zw = Intermediates.SubUVLerp.xx;
+
+	#if NUM_MATERIAL_TEXCOORDS
+		float2 CustomizedUVs[NUM_MATERIAL_TEXCOORDS];
+		GetMaterialCustomizedUVs(VertexParameters, CustomizedUVs);
+		Interpolants.SubUV0AndTexCoord0.zw = CustomizedUVs[0];
+	#else
+		Interpolants.SubUV0AndTexCoord0.zw = 0;
+	#endif
+#elif NUM_MATERIAL_TEXCOORDS
+	#if (ES2_PROFILE || ES3_1_PROFILE)
+		GetMaterialCustomizedUVs(VertexParameters, Interpolants.TexCoords);
+	#else
+		// Ensure the unused components of the last packed texture coordinate are initialized.
+		Interpolants.TexCoords[(NUM_MATERIAL_TEXCOORDS + 1) / 2 - 1] = 0;
+
+		float2 CustomizedUVs[NUM_MATERIAL_TEXCOORDS];
+		GetMaterialCustomizedUVs(VertexParameters, CustomizedUVs);
+
+		UNROLL
+		for(int CoordinateIndex = 0;CoordinateIndex < NUM_MATERIAL_TEXCOORDS;CoordinateIndex += 2)
+		{
+				Interpolants.TexCoords[CoordinateIndex / 2].xy = CustomizedUVs[CoordinateIndex];
+			if(CoordinateIndex + 1 < NUM_MATERIAL_TEXCOORDS)
+			{
+					Interpolants.TexCoords[CoordinateIndex / 2].wz = CustomizedUVs[CoordinateIndex + 1];
+			}
+		}
+	#endif 
+#endif
+
+
+	Interpolants.TangentToWorld0.w = 0;
+	CalcTangentToWorld(Input, Intermediates, Interpolants.TangentToWorld0.xyz, Interpolants.TangentToWorld2);
+
+#if INTERPOLATE_VERTEX_COLOR
+	Interpolants.VertexColor = Intermediates.VertexColor;
+#endif
+
+#if NEEDS_PARTICLE_COLOR
+	Interpolants.ParticleColor = Intermediates.ParticleColor;
+#endif
+
+#if USE_DYNAMIC_PARAMETERS
+	Interpolants.DynamicParameter = Intermediates.DynamicParameter;
+#endif	//USE_DYNAMIC_PARAMETERS
+
+#if USE_PARTICLE_NODES
+	Interpolants.ParticleTranslatedWorldPosition = Intermediates.ParticleTranslatedWorldPosition;
+	Interpolants.ParticleVelocity = Intermediates.ParticleVelocity;
+	Interpolants.ParticleTime = Intermediates.RelativeTime;
+#if NEEDS_PARTICLE_TRANSFORM
+	Interpolants.LocalToWorld = Intermediates.LocalToWorld;
+#endif
+#endif
+
+#if INSTANCED_STEREO
+	Interpolants.PackedEyeIndex = 0;
+#endif
+
+	return Interpolants;
+}
+
+// @return previous translated world position
+float4 VertexFactoryGetPreviousWorldPosition(FVertexFactoryInput Input, FVertexFactoryIntermediates Intermediates)
+{	
+	return CalcPrevWorldPosition(Input);
+}
+
+#if USING_TESSELLATION
+	struct FVertexFactoryInterpolantsVSToDS
+	{
+		FVertexFactoryInterpolantsVSToPS InterpolantsVSToPS;
+	};
+
+	float2 VertexFactoryGetTextureCoordinateDS( FVertexFactoryInterpolantsVSToDS Interpolants )
+	{
+	#if NUM_MATERIAL_TEXCOORDS
+		return Interpolants.InterpolantsVSToPS.TexCoords[0].xy;
+	#else // #if NUM_MATERIAL_TEXCOORDS
+		return float2(0,0);
+	#endif // #if NUM_MATERIAL_TEXCOORDS
+	}
+
+	FVertexFactoryInterpolantsVSToPS VertexFactoryAssignInterpolants(FVertexFactoryInterpolantsVSToDS Input)
+	{
+		return Input.InterpolantsVSToPS;
+	}
+
+	/** Converts from vertex factory specific interpolants to a FMaterialTessellationParameters, which is used by material inputs. */
+	FMaterialTessellationParameters GetMaterialTessellationParameters(FVertexFactoryInterpolantsVSToDS Interpolants, float3 CameraLocalWorldPosition)
+	{
+		FMaterialTessellationParameters	Result;
+	#if NUM_MATERIAL_TEXCOORDS
+		UNROLL
+		for(int CoordinateIndex = 0;CoordinateIndex < NUM_MATERIAL_TEXCOORDS;CoordinateIndex += 2)
+		{
+			Result.TexCoords[CoordinateIndex] = Interpolants.InterpolantsVSToPS.TexCoords[CoordinateIndex/2].xy;
+			if(CoordinateIndex + 1 < NUM_MATERIAL_TEXCOORDS)
+			{
+				Result.TexCoords[CoordinateIndex + 1] = Interpolants.InterpolantsVSToPS.TexCoords[CoordinateIndex/2].wz;
+			}
+		}
+	#endif
+
+		Result.VertexColor = float4(1,1,1,1);
+
+		half3 TangentToWorld0 = Interpolants.InterpolantsVSToPS.TangentToWorld0.xyz;
+		half4 TangentToWorld2 = Interpolants.InterpolantsVSToPS.TangentToWorld2;
+		Result.TangentToWorld = AssembleTangentToWorld( TangentToWorld0, TangentToWorld2 );
+		Result.TangentToWorldPreScale = 1;
+		Result.WorldPosition = CameraLocalWorldPosition + ResolvedView.WorldCameraOrigin;
+
+		return Result;
+	}
+
+	FVertexFactoryInterpolantsVSToDS VertexFactoryGetInterpolantsVSToDS(FVertexFactoryInput Input, FVertexFactoryIntermediates Intermediates, FMaterialVertexParameters VertexParameters)
+	{
+		FVertexFactoryInterpolantsVSToDS Interpolants;
+		Interpolants.InterpolantsVSToPS = VertexFactoryGetInterpolantsVSToPS(Input, Intermediates, VertexParameters);
+		return Interpolants;
+	}
+
+	FVertexFactoryInterpolantsVSToDS VertexFactoryInterpolate(FVertexFactoryInterpolantsVSToDS a, float aInterp, FVertexFactoryInterpolantsVSToDS b, float bInterp)
+	{
+		FVertexFactoryInterpolantsVSToDS O;
+	
+		// Do we really need to interpolate TangentToWorld2 here? It should be replaced by the
+		// interpolated normal from 'whatever' interpolation scheme we're using
+	
+		TESSELLATION_INTERPOLATE_MEMBER(InterpolantsVSToPS.TangentToWorld0.xyz);
+		TESSELLATION_INTERPOLATE_MEMBER(InterpolantsVSToPS.TangentToWorld2);
+#if INTERPOLATE_VERTEX_COLOR
+		TESSELLATION_INTERPOLATE_MEMBER(InterpolantsVSToPS.VertexColor);
+#endif
+
+#if NEEDS_PARTICLE_COLOR
+		TESSELLATION_INTERPOLATE_MEMBER(InterpolantsVSToPS.ParticleColor);
+#endif
+
+#if NUM_MATERIAL_TEXCOORDS
+		for (int i = 0; i < (NUM_MATERIAL_TEXCOORDS + 1) / 2; ++i)
+		{
+			TESSELLATION_INTERPOLATE_MEMBER(InterpolantsVSToPS.TexCoords[i]);
+		}
+#endif
+#if USE_DYNAMIC_PARAMETERS
+		TESSELLATION_INTERPOLATE_MEMBER(InterpolantsVSToPS.DynamicParameter);
+#endif	//USE_DYNAMIC_PARAMETERS
+
+		return O;
+	}
+
+	float3x3 VertexFactoryGetTangentToLocalDS(FVertexFactoryInterpolantsVSToDS Interpolants)
+	{
+		// This duplicates stuff already going on in GetMaterialTessellationParameters(), so
+		// maybe the hull shader could leverage that instead?
+		half3 TangentToWorld0 = Interpolants.InterpolantsVSToPS.TangentToWorld0.xyz;
+		half4 TangentToWorld2 = Interpolants.InterpolantsVSToPS.TangentToWorld2;
+		float3x3 TangentToWorld = AssembleTangentToWorld( TangentToWorld0, TangentToWorld2 );
+
+		return TangentToWorld;
+	}
+#endif // #if USING_TESSELLATION
+
+#if INSTANCED_STEREO
+
+// Instanced stereo sets the eye index explicitly for instanced geometry
+#if NIAGARA_MESH_INSTANCED
+uint InstancedEyeIndex;
+#endif
+
+uint VertexFactoryGetEyeIndex(uint InstanceId)
+{
+	#if NIAGARA_MESH_INSTANCED
+		return InstancedEyeIndex & bIsInstancedStereo;
+	#else
+		return InstanceId & 1;
+	#endif
+}
+#endif
diff --git a/Engine/Shaders/Private/NiagaraSpriteVertexFactory.ush b/Engine/Shaders/Private/NiagaraSpriteVertexFactory.ush
new file mode 100644
--- /dev/null
+++ b/Engine/Shaders/Private/NiagaraSpriteVertexFactory.ush
@@ -0,0 +1,660 @@
+// Copyright 1998-2017 Epic Games, Inc. All Rights Reserved.
+
+/*=============================================================================
+	ParticleSpriteVertexFactory.hlsl: Particle vertex factory shader code.
+	Shared by standard sprite particles and SubUV sprite particles.
+=============================================================================*/
+
+#include "VertexFactoryCommon.ush"
+#include "ParticleVertexFactoryCommon.ush"
+
+// These are only enabled for SM5 due to limited interpolators elsewhere
+#define USE_PARTICLE_LIGHTING_OFFSET (FEATURE_LEVEL >= FEATURE_LEVEL_SM5 && !MATERIAL_SHADINGMODEL_UNLIT)
+#define USE_PARTICLE_POSITION (FEATURE_LEVEL >= FEATURE_LEVEL_SM5 && NEEDS_PARTICLE_POSITION)
+#define USE_PARTICLE_VELOCITY (FEATURE_LEVEL >= FEATURE_LEVEL_SM5 && NEEDS_PARTICLE_VELOCITY)
+#define USE_PARTICLE_TIME (FEATURE_LEVEL >= FEATURE_LEVEL_SM5 && NEEDS_PARTICLE_TIME)
+#define USE_PARTICLE_SIZE (FEATURE_LEVEL >= FEATURE_LEVEL_SM5 && NEEDS_PARTICLE_SIZE)
+
+bool UseCustomAlignment;
+bool UseVectorAlignment;
+bool CameraPlaneFacing;
+
+
+struct FVertexFactoryInput
+{
+	float4	Position				: ATTRIBUTE0;
+	float4	OldPosition				: ATTRIBUTE1;
+	float4	SizeRotSubImage			: ATTRIBUTE2;
+	float4	Color					: ATTRIBUTE3;
+#if USE_DYNAMIC_PARAMETERS
+	float4 DynamicParameter			: ATTRIBUTE5;
+#endif	//USE_DYNAMIC_PARAMETERS
+	float2	TexCoord				: ATTRIBUTE4;
+	uint VertexId					: SV_VertexID;
+	float3 CustomAlignmentVector	: ATTRIBUTE6;
+	float3 CustomFacingVector		: ATTRIBUTE7;
+
+/** Optional instance ID for vertex layered rendering */
+#if FEATURE_LEVEL >= FEATURE_LEVEL_SM4 && ONEPASS_POINTLIGHT_SHADOW && USING_VERTEX_SHADER_LAYER
+	uint InstanceId	: SV_InstanceID;
+#endif
+};
+
+struct FVertexFactoryInterpolantsVSToPS
+{
+	// First row of the tangent to world matrix, Interp_Sizer used by SUBUV_PARTICLES in w
+	float4	TangentToWorld0AndInterp_Sizer : TANGENTTOWORLD0;
+	// Last row of the tangent to world matrix in xyz
+	float4	TangentToWorld2	: TANGENTTOWORLD2;
+
+#if USE_DYNAMIC_PARAMETERS
+	float4 DynamicParameter	: TEXCOORD1;
+#endif	//USE_DYNAMIC_PARAMETERS
+
+#if NEEDS_PARTICLE_COLOR
+	float4	Color				: TEXCOORD2;
+#endif
+
+#if NUM_MATERIAL_TEXCOORDS
+	float4	TexCoords[(NUM_MATERIAL_TEXCOORDS + 1) / 2]	: TEXCOORD3;
+#endif
+
+	//Not sure this is actually being used now and it's awkward to slot in now we're supporting custom UVs so I'm just giving this its own interpolant.
+#if LIGHTMAP_UV_ACCESS
+	float2 LightMapUVs : LIGHTMAP_UVS;
+#endif
+
+#if USE_PARTICLE_SUBUVS
+	float4 ParticleSubUVs : PARTICLE_SUBUVS;
+#endif
+
+#if USE_PARTICLE_POSITION
+	/** Cam-relative (translated) particle center and radius */
+	float4 ParticleTranslatedWorldPositionAndSize	: PARTICLE_POSITION;
+#endif
+
+#if USE_PARTICLE_VELOCITY
+	float4 ParticleVelocity : PARTICLE_VELOCITY;
+#endif
+
+#if USE_PARTICLE_TIME
+	float ParticleTime : PARTICLE_TIME;
+#endif
+
+#if USE_PARTICLE_LIGHTING_OFFSET
+	float3 LightingPositionOffset : PARTICLE_LIGHTING_OFFSET;
+#endif
+
+#if USE_PARTICLE_SIZE
+	float2 ParticleSize : PARTICLE_SIZE;
+#endif
+
+#if INSTANCED_STEREO
+	nointerpolation uint PackedEyeIndex : PACKED_EYE_INDEX;
+#endif
+};
+
+struct FVertexFactoryIntermediates
+{
+	/** The position of the particle in non-translated world space. */
+	float3 ParticleWorldPosition;
+	/** The position of the vertex in translated world space. */
+	float3 VertexWorldPosition;
+	/** Particle translated world space position and size. */
+	float4 TranslatedWorldPositionAndSize;
+	
+#if USE_PARTICLE_LIGHTING_OFFSET
+	float3 LightingPositionOffset;
+#endif
+
+	/** The texture coordinate at this vertex. */
+	float4 TexCoord;
+#if NUM_MATERIAL_TEXCOORDS >= 3 //Unflipped UVs are placed in UVs 2-3
+	/** A second UV set. Always non-UV flipped. Allows use of UV flipping with Normal maps etc. */
+	float4 TexCoord_Unflipped;
+#endif
+	/** The sprite tangent in world space (+V). */
+	float3 TangentUp;
+	/** The sprite tangent in world space (+U). */
+	float3 TangentRight;
+	/** The color of the sprite. */
+	float4 Color;
+	/** The velocity of the particle, XYZ: direction, W: speed. */
+	float4 ParticleVelocity;
+	/** Dynamic parameter. */
+	float4 DynamicParameter;
+	/** The sub-image lerp. */
+	float SubImageLerp;
+	/** Relative time. */
+	float RelativeTime;
+	/** Transform from tangent space to local space. */
+	float3x3 TangentToLocal;
+	/** Size of the particle */
+	float2 ParticleSize;
+};
+
+#if NUM_MATERIAL_TEXCOORDS
+bool UVIndexUseZW(int UVIndex)
+{
+	#if COMPILER_GLSL_ES2
+	return (frac((float)UVIndex / 2.0f) == 0.5f);
+	#else
+	return (UVIndex % 2) != 0;
+	#endif
+}
+
+float2 GetUV(FVertexFactoryInterpolantsVSToPS Interpolants, int UVIndex)
+{
+	float4 UVVector = Interpolants.TexCoords[UVIndex / 2];
+	return UVIndexUseZW(UVIndex) ? UVVector.zw : UVVector.xy;
+}
+
+void SetUV(inout FVertexFactoryInterpolantsVSToPS Interpolants, int UVIndex, float2 InValue)
+{
+	FLATTEN
+	if (UVIndexUseZW(UVIndex))
+	{
+		Interpolants.TexCoords[UVIndex / 2].zw = InValue;
+	}
+	else
+	{
+		Interpolants.TexCoords[UVIndex / 2].xy = InValue;
+	}
+}
+#endif
+
+/** Converts from vertex factory specific interpolants to a FMaterialPixelParameters, which is used by material inputs. */
+FMaterialPixelParameters GetMaterialPixelParameters(FVertexFactoryInterpolantsVSToPS Interpolants, float4 SvPosition)
+{
+	// GetMaterialPixelParameters is responsible for fully initializing the result
+	FMaterialPixelParameters Result = MakeInitializedMaterialPixelParameters();
+
+#if NUM_MATERIAL_TEXCOORDS
+	UNROLL
+	for (int CoordinateIndex = 0; CoordinateIndex < NUM_MATERIAL_TEXCOORDS; CoordinateIndex++)
+	{
+		Result.TexCoords[CoordinateIndex] = GetUV(Interpolants, CoordinateIndex);
+	}
+#endif	// NUM_MATERIAL_TEXCOORDS
+
+	Result.VertexColor = 1;
+
+	half4 TangentToWorld0 = Interpolants.TangentToWorld0AndInterp_Sizer;
+	half4 TangentToWorld2 = Interpolants.TangentToWorld2;
+	float SubImageLerp = Interpolants.TangentToWorld0AndInterp_Sizer.w;
+
+	#if GENERATE_SPHERICAL_PARTICLE_NORMALS && USE_PARTICLE_POSITION
+		{	
+			// can be optimized
+			float4 ScreenPosition = SvPositionToResolvedScreenPosition(SvPosition);
+			float3 TranslatedWorldPosition = SvPositionToResolvedTranslatedWorld(SvPosition);
+			Result.TangentToWorld = GetSphericalParticleNormal(TranslatedWorldPosition, Interpolants.ParticleTranslatedWorldPositionAndSize.xyz, Interpolants.ParticleTranslatedWorldPositionAndSize.w);
+		}
+	#else
+		Result.TangentToWorld = AssembleTangentToWorld(TangentToWorld0.xyz, TangentToWorld2);
+	#endif
+
+	Result.UnMirrored = 1;
+	Result.Particle.MacroUV = NiagaraSpriteVF.MacroUVParameters;
+
+#if NEEDS_PARTICLE_COLOR
+	Result.Particle.Color = Interpolants.Color;
+#endif
+
+#if USE_DYNAMIC_PARAMETERS
+	Result.Particle.DynamicParameter = Interpolants.DynamicParameter;
+#endif	//USE_DYNAMIC_PARAMETERS
+
+#if USE_PARTICLE_POSITION
+	Result.Particle.TranslatedWorldPositionAndSize = Interpolants.ParticleTranslatedWorldPositionAndSize;
+#endif
+
+#if USE_PARTICLE_VELOCITY
+	Result.Particle.Velocity = Interpolants.ParticleVelocity;
+#endif
+
+#if USE_PARTICLE_TIME
+	Result.Particle.RelativeTime = Interpolants.ParticleTime;
+#endif
+
+	Result.Particle.MotionBlurFade = 1.0f;
+
+
+#if USE_PARTICLE_LIGHTING_OFFSET
+	Result.LightingPositionOffset = Interpolants.LightingPositionOffset;
+#endif
+
+#if USE_PARTICLE_SUBUVS
+	Result.Particle.SubUVCoords[0] = Interpolants.ParticleSubUVs.xy;
+	Result.Particle.SubUVCoords[1] = Interpolants.ParticleSubUVs.zw;
+	Result.Particle.SubUVLerp = SubImageLerp;
+#endif
+
+#if LIGHTMAP_UV_ACCESS
+	Result.LightmapUVs = Interpolants.LightMapUVs;
+#endif
+
+#if USE_PARTICLE_SIZE
+	Result.Particle.Size = Interpolants.ParticleSize;
+#endif
+
+	Result.TwoSidedSign = 1;
+	return Result;
+}
+
+/** Converts from vertex factory specific input to a FMaterialVertexParameters, which is used by vertex shader material inputs. */
+FMaterialVertexParameters GetMaterialVertexParameters(FVertexFactoryInput Input, FVertexFactoryIntermediates Intermediates, float3 WorldPosition, float3x3 TangentToLocal)
+{
+	FMaterialVertexParameters Result = (FMaterialVertexParameters)0;
+	Result.WorldPosition = WorldPosition;
+	Result.VertexColor = Input.Color;
+	Result.TangentToWorld = mul(TangentToLocal, GetLocalToWorld3x3()); 
+	Result.Particle.MacroUV = NiagaraSpriteVF.MacroUVParameters;
+	Result.Particle.Color = Intermediates.Color;
+	Result.Particle.MotionBlurFade = 1.0f;
+	Result.PreSkinnedPosition = Input.Position.xyz;
+	Result.PreSkinnedNormal = TangentToLocal[2].xyz;
+
+#if USE_DYNAMIC_PARAMETERS
+	Result.Particle.DynamicParameter = Intermediates.DynamicParameter;
+#endif	//USE_DYNAMIC_PARAMETERS
+
+#if USE_PARTICLE_POSITION
+	Result.Particle.TranslatedWorldPositionAndSize = Intermediates.TranslatedWorldPositionAndSize;
+#endif
+
+#if USE_PARTICLE_VELOCITY
+	Result.Particle.Velocity = Intermediates.ParticleVelocity;
+#endif
+
+#if USE_PARTICLE_TIME
+	Result.Particle.RelativeTime = Intermediates.RelativeTime;
+#endif
+
+#if NUM_MATERIAL_TEXCOORDS_VERTEX
+	#if NUM_MATERIAL_TEXCOORDS_VERTEX >= 1
+		Result.TexCoords[0] = Intermediates.TexCoord.xy;
+	#if NUM_MATERIAL_TEXCOORDS_VERTEX >= 2
+		Result.TexCoords[1] = Intermediates.TexCoord.zw;
+	#if NUM_MATERIAL_TEXCOORDS_VERTEX >= 3
+		Result.TexCoords[2] = Intermediates.TexCoord_Unflipped.xy;
+	#if NUM_MATERIAL_TEXCOORDS_VERTEX >= 4
+		Result.TexCoords[3] = Intermediates.TexCoord_Unflipped.zw;
+	#endif	// >= 4
+	#endif	// >= 3
+	#endif	// >= 2
+	#endif	// >= 1
+#endif	//NUM_MATERIAL_TEXCOORDS_VERTEX
+
+#if USE_PARTICLE_SUBUVS
+	Result.Particle.SubUVCoords[0] = Intermediates.TexCoord.xy;
+	Result.Particle.SubUVCoords[1] = Intermediates.TexCoord.zw;
+#endif
+
+#if USE_PARTICLE_SIZE
+	Result.Particle.Size = Intermediates.ParticleSize;
+#endif
+	return Result;
+}
+
+void GetLightMapCoordinates(FVertexFactoryInterpolantsVSToPS Interpolants, out float2 LightmapUV0, out float2 LightmapUV1)
+{
+	LightmapUV0 = LightmapUV1 = 0;
+}
+
+half2 GetShadowMapCoordinate(FVertexFactoryInterpolantsVSToPS Interpolants)
+{
+	return 0;
+}
+
+float3 SafeNormalize(float3 V)
+{
+	return V * rsqrt(max(dot(V,V),0.01));
+}
+
+void GetTangents(FVertexFactoryInput Input, float3 WorldPosition, float3 OldWorldPosition, out float3 OutRight, out float3 OutUp)
+{
+	// Select camera up/right vectors.
+	float3 ResolvedViewRight = lerp(ResolvedView.ViewRight, ResolvedView.HMDViewNoRollRight, NiagaraSpriteVF.RemoveHMDRoll);
+	float3 ResolvedViewUp = lerp(ResolvedView.ViewUp, ResolvedView.HMDViewNoRollUp, NiagaraSpriteVF.RemoveHMDRoll);
+
+	// camera related vectors
+	float3 CameraRight = ResolvedViewRight;
+	float3 CameraUp = -ResolvedViewUp;
+	float3 CameraVector = ResolvedView.WorldCameraOrigin.xyz - WorldPosition.xyz;
+	float3 CameraDirection = SafeNormalize(CameraVector);
+	float3 CameraFacingVector = Input.CustomFacingVector * NiagaraSpriteVF.CustomFacingVectorMask + CameraDirection * (float3(1.0,1.0,1.0)-NiagaraSpriteVF.CustomFacingVectorMask);
+	CameraFacingVector = lerp(CameraFacingVector, -ResolvedView.ViewForward, CameraPlaneFacing);
+	
+	BRANCH
+	if (NiagaraSpriteVF.CameraFacingBlend.x > 0.f)
+	{
+		float CameraDistanceSq = dot(CameraVector, CameraVector);
+		float AlignmentMode = saturate(CameraDistanceSq * NiagaraSpriteVF.CameraFacingBlend.y - NiagaraSpriteVF.CameraFacingBlend.z);
+		CameraFacingVector = -ResolvedView.ViewForward * CameraFacingVector * (1.0-AlignmentMode);
+	}
+
+	// the particle's movement direction.
+	float3 ParticleDirection	= SafeNormalize(WorldPosition.xyz - OldWorldPosition.xyz);
+	float3 DirectionVector	= lerp(ParticleDirection, Input.CustomAlignmentVector, UseCustomAlignment);
+
+	// Tangent vectors for PSA_Velocity.
+	float3 Right_Aligned	= SafeNormalize(cross(CameraDirection, DirectionVector));
+	float3 Up_Aligned		= -DirectionVector;
+
+	// Tangent vectors for camera facing position (can be based on camera position or the custom facing vector)
+	//
+	float3 Right_FacingCameraPosition = SafeNormalize(cross(CameraFacingVector,ResolvedViewUp));
+	float3 Up_FacingCameraPosition = cross(CameraFacingVector, Right_FacingCameraPosition);
+
+	float3 Right = SafeNormalize( lerp(Right_FacingCameraPosition, Right_Aligned, UseVectorAlignment) );
+	float3 Up = SafeNormalize( cross(CameraFacingVector, Right) );
+
+	// Determine the angle of rotation.
+	float SinRotation;
+	float CosRotation;
+	const float SpriteRotation = Input.SizeRotSubImage.z * NiagaraSpriteVF.RotationScale + NiagaraSpriteVF.RotationBias;
+	sincos(SpriteRotation, SinRotation, CosRotation);
+
+	// Rotate the sprite to determine final tangents.
+	OutRight	= SinRotation * Up + CosRotation * Right;
+	OutUp		= CosRotation * Up - SinRotation * Right;
+}
+
+/** derive basis vectors */
+float3x3 CalcTangentBasis(FVertexFactoryIntermediates Intermediates)
+{
+	// Using camera facing TangentX and TangentY.  The resulting tangent basis is not orthonormal with anything other than ENM_CameraFacing, 
+	// So there are artifacts with tangent space calculations like the TransformVector node,
+	// But this allows lighting based on a world space shape via the normal while still having normal maps camera aligned.
+	float3x3 Result;
+	Result[0] = Intermediates.TangentRight;
+	Result[1] = Intermediates.TangentUp;
+
+	// ENM_CameraFacing
+	//@todo - use static branching
+	if (NiagaraSpriteVF.NormalsType < .5f)
+	{
+		Result[2] = normalize(cross(Result[0],Result[1]));
+	} 
+	// ENM_Spherical
+	else if (NiagaraSpriteVF.NormalsType < 1.5f)
+	{
+		float3 TangentZ = normalize(Intermediates.ParticleWorldPosition.xyz - NiagaraSpriteVF.NormalsSphereCenter.xyz);
+		Result[2] = TangentZ;	
+	}
+	// ENM_Cylindrical
+	else
+	{
+		float3 ClosestPointOnCylinder = NiagaraSpriteVF.NormalsSphereCenter.xyz + dot(NiagaraSpriteVF.NormalsCylinderUnitDirection.xyz, Intermediates.ParticleWorldPosition.xyz - NiagaraSpriteVF.NormalsSphereCenter.xyz) * NiagaraSpriteVF.NormalsCylinderUnitDirection.xyz;
+		float3 TangentZ = normalize(Intermediates.ParticleWorldPosition.xyz - ClosestPointOnCylinder);
+		Result[2] = TangentZ;	
+	}
+
+	return Result;
+}
+
+// Note: C++ geometry setup behavior has to match this define
+/** Number of vertices in each SubUV frame of CutoutGeometry. */
+uint NumCutoutVerticesPerFrame;
+
+/** SubUV animation bounding geometry used to minimize overdraw.  Each frame has its own polygon. */
+Buffer<float2> CutoutGeometry;
+
+void ComputeBillboardUVs(FVertexFactoryInput Input, out float2 UVForPosition, out float2 UVForTexturing, out float2 UVForTexturingUnflipped)
+{
+	// Encoding the UV flip in the sign of the size data.
+	float2 UVFlip = sign(Input.SizeRotSubImage.xy);
+
+	BRANCH
+	if (NumCutoutVerticesPerFrame > 0)
+	{
+		// Avoid uint divide which is extremely slow on GCN
+		uint CutoutVertexIndex = (uint)fmod(Input.VertexId, NumCutoutVerticesPerFrame);
+		float NumFrames = NiagaraSpriteVF.SubImageSize.x * NiagaraSpriteVF.SubImageSize.y;
+		uint SubImageIndexInt = (uint)fmod(Input.SizeRotSubImage.w, NumFrames);
+
+		FLATTEN
+		if (UVFlip.x * UVFlip.y < 0)
+		{
+			// Reverse the winding order of the polygon when only flipping in one direction to counteract UV flipping
+			CutoutVertexIndex = NumCutoutVerticesPerFrame - 1 - CutoutVertexIndex;
+		}
+
+		// Fetch the billboard space positions from the appropriate frame of the preprocessed cutout geometry
+		UVForTexturing = UVForTexturingUnflipped = CutoutGeometry[SubImageIndexInt * NumCutoutVerticesPerFrame + CutoutVertexIndex];
+
+		// Invert positions on the billboard so that the cutout geometry still contains the visible portion
+		// This changes the winding order when only one direction is flipped
+		UVForPosition.x = UVFlip.x < 0.0 ? 1.0 - UVForTexturing.x : UVForTexturing.x;
+		UVForPosition.y = UVFlip.y < 0.0 ? 1.0 - UVForTexturing.y : UVForTexturing.y;
+	}
+	else
+	{
+		UVForTexturing.x = UVFlip.x < 0.0 ? 1.0 - Input.TexCoord.x : Input.TexCoord.x;
+		UVForTexturing.y = UVFlip.y < 0.0 ? 1.0 - Input.TexCoord.y : Input.TexCoord.y;
+
+		// Note: not inverting positions, as that would change the winding order
+		UVForPosition = UVForTexturingUnflipped = Input.TexCoord.xy;
+	}
+}
+
+FVertexFactoryIntermediates GetVertexFactoryIntermediates(FVertexFactoryInput Input)
+{
+	FVertexFactoryIntermediates Intermediates = (FVertexFactoryIntermediates)0;
+
+	// World position.
+	float4 ParticleWorldPosition = mul(float4(Input.Position.xyz,1), Primitive.LocalToWorld);
+	float4 ParticleOldWorldPosition = mul(float4(Input.OldPosition.xyz,1), Primitive.LocalToWorld);
+	Intermediates.ParticleWorldPosition = ParticleWorldPosition.xyz;
+
+	// Tangents.
+	float3 Right,Up;
+	GetTangents(Input,ParticleWorldPosition.xyz,ParticleOldWorldPosition.xyz,Right,Up);
+	Intermediates.TangentUp = Up;
+	Intermediates.TangentRight = Right;
+
+	float2 UVForPosition;
+	float2 UVForTexturing;
+	float2 UVForTexturingUnflipped;
+	ComputeBillboardUVs(Input, UVForPosition, UVForTexturing, UVForTexturingUnflipped);
+
+	// Vertex position.
+	float4 VertexWorldPosition = float4(ParticleWorldPosition.xyz + ResolvedView.PreViewTranslation,1);
+	float2 Size = abs(Input.SizeRotSubImage.xy);
+	VertexWorldPosition += Size.x * (UVForPosition.x + NiagaraSpriteVF.PivotOffset.x) * float4(Right,0);
+	VertexWorldPosition += Size.y * (UVForPosition.y + NiagaraSpriteVF.PivotOffset.y) * float4(Up,0);
+	Intermediates.VertexWorldPosition = VertexWorldPosition.xyz;
+	
+	// SubUV.
+	float SubImageIndex = Input.SizeRotSubImage.w;
+	float SubImageLerp = frac(SubImageIndex);
+	float SubImageA = SubImageIndex - SubImageLerp;
+	float SubImageB = SubImageA + 1;
+	float SubImageAH = fmod( SubImageA, NiagaraSpriteVF.SubImageSize.x );
+	float SubImageBH = fmod( SubImageB, NiagaraSpriteVF.SubImageSize.x );
+	float SubImageAV = floor( SubImageA * NiagaraSpriteVF.SubImageSize.z );
+	float SubImageBV = floor( SubImageB * NiagaraSpriteVF.SubImageSize.z );
+	Intermediates.TexCoord.xy = (float2( SubImageAH, SubImageAV ) + UVForTexturing) * NiagaraSpriteVF.SubImageSize.zw;
+	Intermediates.TexCoord.zw = (float2(SubImageBH, SubImageBV) + UVForTexturing) * NiagaraSpriteVF.SubImageSize.zw;
+#if NUM_MATERIAL_TEXCOORDS >= 3
+	Intermediates.TexCoord_Unflipped.xy = (float2(SubImageAH, SubImageAV) + UVForTexturingUnflipped) * NiagaraSpriteVF.SubImageSize.zw;
+	Intermediates.TexCoord_Unflipped.zw = (float2(SubImageBH, SubImageBV) + UVForTexturingUnflipped) * NiagaraSpriteVF.SubImageSize.zw;
+#endif
+	Intermediates.SubImageLerp = SubImageLerp;
+
+	Intermediates.Color = Input.Color;
+
+#if USE_DYNAMIC_PARAMETERS
+	Intermediates.DynamicParameter = Input.DynamicParameter;
+#endif	//USE_DYNAMIC_PARAMETERS
+
+
+#if USE_PARTICLE_POSITION
+	float ParticleRadius = .5f * min(Input.SizeRotSubImage.x, Input.SizeRotSubImage.y);
+	float3 ParticleTranslatedWorldPosition = ParticleWorldPosition.xyz + ResolvedView.PreViewTranslation.xyz;
+	Intermediates.TranslatedWorldPositionAndSize = float4(ParticleTranslatedWorldPosition, ParticleRadius);
+
+#if GENERATE_SPHERICAL_PARTICLE_NORMALS
+	// Recenter the particle position for GENERATE_SPHERICAL_PARTICLE_NORMALS to use the right center.
+	Intermediates.TranslatedWorldPositionAndSize.xyz += Size.x * (NiagaraSpriteVF.PivotOffset.x + .5) * Right + Size.y * (NiagaraSpriteVF.PivotOffset.y + .5) * Up;
+#endif
+#endif
+
+#if USE_PARTICLE_VELOCITY
+	float3 ParticleVelocity = (ParticleWorldPosition.xyz - ParticleOldWorldPosition.xyz) * NiagaraSpriteVF.InvDeltaSeconds;
+	float3 DirectionVector = lerp(ParticleVelocity, Input.CustomAlignmentVector, UseCustomAlignment);
+	Intermediates.ParticleVelocity.xyz = normalize(DirectionVector);
+	Intermediates.ParticleVelocity.w = length(DirectionVector);
+#endif
+
+#if USE_PARTICLE_TIME
+	Intermediates.RelativeTime = Input.Position.w;
+#endif
+
+	Intermediates.TangentToLocal = CalcTangentBasis(Intermediates);
+
+#if USE_PARTICLE_LIGHTING_OFFSET
+	// Hash function based on the particle ID to generate a uniformly distributed 3d offset
+	float3 RandomParticleOffset = frac(Square(Input.OldPosition.w + 10) * float3(1361.456345, 2333.578, 3623.983)) * 2 - 1;
+	Intermediates.LightingPositionOffset = .5f * View.TranslucencyLightingVolumeInvSize[0].w * RandomParticleOffset;
+#endif
+
+#if USE_PARTICLE_SIZE
+	Intermediates.ParticleSize = Size;
+#endif
+
+	return Intermediates;
+}
+
+float4 VertexFactoryGetWorldPosition(FVertexFactoryInput Input, FVertexFactoryIntermediates Intermediates)
+{
+	return float4(Intermediates.VertexWorldPosition,1);
+}
+
+float4 VertexFactoryGetRasterizedWorldPosition(FVertexFactoryInput Input, FVertexFactoryIntermediates Intermediates, float4 TranslatedWorldPosition)
+{
+#if SPHERICAL_PARTICLE_OPACITY
+	// For particles using spherical opacity, move the quad toward the viewer so that it lies in front of the sphere defined by the particle
+	// This avoids opaque objects intersecting the particle from causing clipping artifacts due to depth testing
+	// The downside is that the particle will clip the near plane sooner
+
+	float Radius = .5f * min(Input.SizeRotSubImage.x, Input.SizeRotSubImage.y);
+	return ReprojectPosition(TranslatedWorldPosition, Radius);
+#else
+	return TranslatedWorldPosition;
+#endif
+}
+
+float3 VertexFactoryGetPositionForVertexLighting(FVertexFactoryInput Input, FVertexFactoryIntermediates Intermediates, float3 TranslatedWorldPosition)
+{
+#if (TRANSLUCENCY_LIGHTING_VOLUMETRIC_PERVERTEX_DIRECTIONAL || TRANSLUCENCY_LIGHTING_VOLUMETRIC_PERVERTEX_NONDIRECTIONAL)
+	// Do per-vertex lighting with particle position instead of vertex position when we're using SubUV cutouts, since the vertex positions vary per-frame which would cause popping
+	return NumCutoutVerticesPerFrame > 0 ? (Intermediates.ParticleWorldPosition.xyz + ResolvedView.PreViewTranslation.xyz) : TranslatedWorldPosition;
+#else
+	return TranslatedWorldPosition;
+#endif
+}
+
+FVertexFactoryInterpolantsVSToPS VertexFactoryGetInterpolantsVSToPS(FVertexFactoryInput Input, FVertexFactoryIntermediates Intermediates, FMaterialVertexParameters VertexParameters)
+{
+	FVertexFactoryInterpolantsVSToPS Interpolants;
+
+	// Initialize the whole struct to 0
+	Interpolants = (FVertexFactoryInterpolantsVSToPS)0;
+
+#if NUM_MATERIAL_TEXCOORDS
+
+	float2 CustomizedUVs[NUM_MATERIAL_TEXCOORDS];
+	GetMaterialCustomizedUVs(VertexParameters, CustomizedUVs);
+
+	UNROLL
+	for (int CoordinateIndex = 0; CoordinateIndex < NUM_MATERIAL_TEXCOORDS; CoordinateIndex++)
+	{
+		SetUV(Interpolants, CoordinateIndex, CustomizedUVs[CoordinateIndex]);
+	}
+#endif
+
+#if LIGHTMAP_UV_ACCESS
+	Interpolants.LightMapUVs = Intermediates.TexCoord.xy;
+#endif
+
+#if USE_PARTICLE_SUBUVS
+	Interpolants.ParticleSubUVs.xy = VertexParameters.Particle.SubUVCoords[0];
+	Interpolants.ParticleSubUVs.zw = VertexParameters.Particle.SubUVCoords[1];
+#endif
+
+	// Calculate the transform from tangent to world space.
+	// Note that "local" space for particles is actually oriented in world space! Therefore no rotation is needed.
+	float3x3 TangentToWorld = Intermediates.TangentToLocal;
+
+	Interpolants.TangentToWorld0AndInterp_Sizer.xyz = TangentToWorld[0];
+	Interpolants.TangentToWorld0AndInterp_Sizer.w = Intermediates.SubImageLerp;
+	Interpolants.TangentToWorld2 = float4(TangentToWorld[2], sign(determinant(Intermediates.TangentToLocal)));
+
+#if NEEDS_PARTICLE_COLOR
+	Interpolants.Color = Intermediates.Color;
+#endif
+
+#if USE_DYNAMIC_PARAMETERS
+	Interpolants.DynamicParameter = Intermediates.DynamicParameter;
+#endif	//USE_DYNAMIC_PARAMETERS
+
+#if USE_PARTICLE_POSITION
+	Interpolants.ParticleTranslatedWorldPositionAndSize = Intermediates.TranslatedWorldPositionAndSize;
+#endif
+
+#if USE_PARTICLE_VELOCITY
+	Interpolants.ParticleVelocity = Intermediates.ParticleVelocity;
+#endif
+
+#if USE_PARTICLE_TIME
+	Interpolants.ParticleTime = Intermediates.RelativeTime;
+#endif
+
+#if USE_PARTICLE_LIGHTING_OFFSET
+	Interpolants.LightingPositionOffset = Intermediates.LightingPositionOffset;
+#endif
+
+#if USE_PARTICLE_SIZE
+	Interpolants.ParticleSize = Intermediates.ParticleSize;
+#endif
+
+#if INSTANCED_STEREO
+	Interpolants.PackedEyeIndex = 0;
+#endif
+
+	return Interpolants;
+}
+
+float4 VertexFactoryGetPreviousWorldPosition(FVertexFactoryInput Input, FVertexFactoryIntermediates Intermediates)
+{
+	return VertexFactoryGetWorldPosition(Input, Intermediates);
+}
+
+/**
+* Get the 3x3 tangent basis vectors for this vertex factory
+*
+* @param Input - vertex input stream structure
+* @return 3x3 matrix
+*/
+float3x3 VertexFactoryGetTangentToLocal(FVertexFactoryInput Input, FVertexFactoryIntermediates Intermediates)
+{
+	return Intermediates.TangentToLocal;
+}
+
+#if INSTANCED_STEREO
+
+// Instanced stereo sets the eye index explicitly for instanced geometry
+uint InstancedEyeIndex;	 
+
+uint VertexFactoryGetEyeIndex(uint InstanceId)
+{
+	return InstancedEyeIndex & bIsInstancedStereo;
+}
+#endif
+
+float4 VertexFactoryGetTranslatedPrimitiveVolumeBounds(FVertexFactoryInterpolantsVSToPS Interpolants)
+{
+	return 0;
+}
\ No newline at end of file
diff --git a/Engine/Shaders/Private/OcclusionQueryPixelShader.ush b/Engine/Shaders/Private/OcclusionQueryPixelShader.ush
new file mode 100644
--- /dev/null
+++ b/Engine/Shaders/Private/OcclusionQueryPixelShader.ush
@@ -0,0 +1,12 @@
+// Copyright 1998-2017 Epic Games, Inc. All Rights Reserved.
+
+/*=============================================================================
+	OcclusionQueryPixelShader.hlsl: Pixel shader for drawing occlusion queries.
+=============================================================================*/
+
+#include "Common.ush"
+
+void Main(out float4 OutColor : SV_Target0)
+{
+	OutColor = 0;
+}
diff --git a/Engine/Shaders/Private/PNTriangles.ush b/Engine/Shaders/Private/PNTriangles.ush
new file mode 100644
--- /dev/null
+++ b/Engine/Shaders/Private/PNTriangles.ush
@@ -0,0 +1,302 @@
+// Copyright 1998-2017 Epic Games, Inc. All Rights Reserved.
+
+/*=============================================================================
+	PNTriangles.usf: Shader fragments for PN-AEN triangles tessellation
+=============================================================================*/
+
+#define IN_CONTROL_POINTS 12
+#define OUT_CONTROL_POINTS 3  
+
+
+
+struct FHullShaderConstantOutput
+{
+	// Tess factor for the FF HW block
+	float TessFactor[3]    : SV_TessFactor;
+	float InsideTessFactor : SV_InsideTessFactor;
+
+	float4 WorldB111		: PN_POSITION9;
+};
+
+struct FPNTessellationHSToDS
+{
+	FPassSpecificVSToDS PassSpecificData;
+	//@todo - Should be storing View position, then we can do a faster projection matrix multiply in the 
+	//domain shader
+	float4 WorldPosition[3]		: PN_POSITION;
+
+	// scale factors in tangent space
+	float3 DisplacementScale	: PN_DisplacementScales;
+
+	// Tessellation multiplier.
+	float TessellationMultiplier : PN_TessellationMultiplier;
+
+#if DISPLACEMENT_ANTICRACK
+	// dominant edge and vertex data used to avoid cracking when displacing
+	FHullShaderConstantDominantVertexData DominantVertex : PN_DominantVertex;
+	FHullShaderConstantDominantEdgeData DominantEdge : PN_DominantEdge;
+#endif
+};
+
+float4 ComputeControlPoint(float4 PositionA, float4 PositionB, float4 NormalA)
+{
+    return (2 * PositionA + PositionB - (dot((PositionB - PositionA), NormalA) * NormalA)) / 3.0f;
+}
+
+// GLSL needs those on the domain shader, HLSL on the hull shader, so we replicate them for both
+#define TESSELLATION_ATTRIBUTES [partitioning("fractional_odd")][outputtopology("triangle_cw")]
+
+#if HULLSHADER
+
+	// Note that we input all control point from AEN, but only their PassSpecificData is valid (see MainHull shader)
+	FHullShaderConstantOutput PNSharedConstantHull( const OutputPatch<FPNTessellationHSToDS, OUT_CONTROL_POINTS> I )
+	{
+		FHullShaderConstantOutput O = (FHullShaderConstantOutput)0;
+
+		#if INSTANCED_STEREO
+			ResolvedView = ResolveView(GetEyeIndex(I[0].PassSpecificData.FactoryInterpolants.InterpolantsVSToPS.PackedEyeIndex));
+		#else
+			ResolvedView = ResolveView();
+		#endif
+		
+#if DISABLE_TESSELLATION_OVERRIDE
+		O.TessFactor[0] = 1.0f; // 1->2 edge
+		O.TessFactor[1] = 1.0f; // 2->0 edge
+		O.TessFactor[2] = 1.0f; // 0->1 edge
+		O.InsideTessFactor = 1.0f;
+#else // #if DISABLE_TESSELLATION_OVERRIDE
+		float4 B300 = I[0].WorldPosition[0],
+			   B210 = I[0].WorldPosition[1],
+			   B120 = I[0].WorldPosition[2],
+			   B030 = I[1].WorldPosition[0],
+			   B021 = I[1].WorldPosition[1],
+			   B012 = I[1].WorldPosition[2],
+			   B003 = I[2].WorldPosition[0],
+			   B102 = I[2].WorldPosition[1],
+			   B201 = I[2].WorldPosition[2];
+
+		float4 E = (B210 + B120 + B021 + B012 + B102 + B201) / 6.0f;
+		float4 V = (B003 + B030 + B300) / 3.0f;
+		O.WorldB111 = E + ((E - V) / 2.0f);
+
+		float4 TessellationMultipliers;
+		TessellationMultipliers.x = 0.5 * (I[1].TessellationMultiplier + I[2].TessellationMultiplier);
+		TessellationMultipliers.y = 0.5 * (I[2].TessellationMultiplier + I[0].TessellationMultiplier);
+		TessellationMultipliers.z = 0.5 * (I[0].TessellationMultiplier + I[1].TessellationMultiplier);
+		TessellationMultipliers.w = 0.333 * (I[0].TessellationMultiplier + I[1].TessellationMultiplier + I[2].TessellationMultiplier);
+
+		//@todo: This should be vectorized and done in the MainHull shader.
+		float4 CompositeTessellationFactors = TessellationMultipliers * CalculateCompositeTessellationFactors(B300.xyz,B030.xyz,B003.xyz);
+		CompositeTessellationFactors = clamp( CompositeTessellationFactors, 1, 15 );
+
+		O.TessFactor[0] = CompositeTessellationFactors.x; // 1->2 edge
+		O.TessFactor[1] = CompositeTessellationFactors.y; // 2->0 edge
+		O.TessFactor[2] = CompositeTessellationFactors.z; // 0->1 edge
+		O.InsideTessFactor = CompositeTessellationFactors.w;
+#endif // #if DISABLE_TESSELLATION_OVERRIDE
+
+		return O;
+	}
+
+	[domain("tri")]
+	[patchconstantfunc("PNSharedConstantHull")]
+	[outputcontrolpoints(OUT_CONTROL_POINTS)]
+	[maxtessfactor(15)]
+	TESSELLATION_ATTRIBUTES
+	FPNTessellationHSToDS MainHull( InputPatch<FPassSpecificVSToDS, IN_CONTROL_POINTS> I, uint ControlPointID : SV_OutputControlPointID )
+	{
+		#if INSTANCED_STEREO
+			ResolvedView = ResolveView(GetEyeIndex(I[0].FactoryInterpolants.InterpolantsVSToPS.PackedEyeIndex));
+		#else
+			ResolvedView = ResolveView();
+		#endif
+
+		FPNTessellationHSToDS O = (FPNTessellationHSToDS) 0;
+
+		float3x3 TangentToLocal = VertexFactoryGetTangentToLocalDS(I[ControlPointID].FactoryInterpolants);
+		float4 Normal = GetNormalDS(TangentToLocal);
+
+		O.DisplacementScale = GetTangentSpaceNonUniformScales(TangentToLocal);
+
+		const uint NextControlPointID = ControlPointID < 2 ? ControlPointID + 1 : 0; // (ControlPointID + 1) % 3
+		const uint NeighborControlPointID = 3 + 2 * ControlPointID;
+		const uint NeighborNextControlPointID = NeighborControlPointID + 1;
+
+		O.PassSpecificData = I[ControlPointID];
+		O.WorldPosition[0] = I[ControlPointID].Position;
+		O.TessellationMultiplier = GetTessellationMultiplier( I[ControlPointID] );
+
+#if DISPLACEMENT_ANTICRACK
+		// NextControlPointID is first opposite edge vert
+		// NextNextControlPointID is second opposite edge vert
+		const uint NextNextControlPointID = NextControlPointID < 2 ? NextControlPointID + 1 : 0; // (NextControlPointID + 1) % 3 
+
+		// we also need the neghbors of the opposite edge verts
+		const uint NextNeighborControlPointID = 3 + 2 * NextControlPointID;	// neighbor for first opposite edge vert
+		const uint NextNeighborNextControlPointID = NextNeighborControlPointID+1; // neighbor for second opposite edge vert
+
+		// Dominant vertex is provided explicitly
+		const uint DominantVertexControlPointID = 9 + ControlPointID;
+		O.DominantVertex = GenerateDominantVertexData(I[DominantVertexControlPointID]);
+
+		// we need the other verts from each prim to generate the edges, clock wise winding
+		O.DominantEdge = GenerateDominantEdgeData(I[NextControlPointID],I[NextNextControlPointID],I[NextNeighborControlPointID],I[NextNeighborNextControlPointID]);
+#endif	// DISPLACEMENT_ANTICRACK
+
+		// AEN Computations
+		float4 NextNormal = GetNormalDS(VertexFactoryGetTangentToLocalDS(I[NextControlPointID].FactoryInterpolants));
+		float4 NeighborNormal = GetNormalDS(VertexFactoryGetTangentToLocalDS(I[NeighborControlPointID].FactoryInterpolants));
+		float4 NeighborNextNormal = GetNormalDS(VertexFactoryGetTangentToLocalDS(I[NeighborNextControlPointID].FactoryInterpolants));
+
+		float4 ThisControlPoint, NeighborControlPoint;
+		
+		ThisControlPoint = ComputeControlPoint(I[ControlPointID].Position, I[NextControlPointID].Position, Normal);
+		NeighborControlPoint = ComputeControlPoint(I[NeighborControlPointID].Position, I[NeighborNextControlPointID].Position, NeighborNormal);
+		O.WorldPosition[1] = (ThisControlPoint + NeighborControlPoint) / 2;
+		
+		ThisControlPoint = ComputeControlPoint(I[NextControlPointID].Position, I[ControlPointID].Position, NextNormal);
+		NeighborControlPoint = ComputeControlPoint(I[NeighborNextControlPointID].Position, I[NeighborControlPointID].Position, NeighborNextNormal);
+		O.WorldPosition[2] = (ThisControlPoint + NeighborControlPoint) / 2;
+		
+		return O;
+	}
+#endif
+
+#if DOMAINSHADER
+	[domain("tri")]
+	TESSELLATION_ATTRIBUTES
+	FPassSpecificVSToPS MainDomain(
+		FHullShaderConstantOutput HSConstantData,
+		const OutputPatch<FPNTessellationHSToDS, OUT_CONTROL_POINTS> I,
+		float3 BarycentricCoords : SV_DomainLocation
+#if USE_GLOBAL_CLIP_PLANE
+		, out float OutGlobalClipPlaneDistance : SV_ClipDistance
+#endif
+#if INSTANCED_STEREO && !MULTI_VIEW
+		, out float OutClipDistance : SV_ClipDistance1
+#endif
+		)
+	{
+		#if INSTANCED_STEREO
+			ResolvedView = ResolveView(GetEyeIndex(I[0].PassSpecificData.FactoryInterpolants.InterpolantsVSToPS.PackedEyeIndex));
+			#if !MULTI_VIEW
+				OutClipDistance = 0.0;
+			#endif
+		#else
+			ResolvedView = ResolveView();
+		#endif
+
+		// Get the barycentric coords
+		float U = BarycentricCoords.x;
+		float V = BarycentricCoords.y;
+		float W = BarycentricCoords.z;
+		
+		// Precompute squares and squares * 3 
+		float UU = U * U;
+		float VV = V * V;
+		float WW = W * W;
+		float UU3 = UU * 3.0f;
+		float VV3 = VV * 3.0f;
+		float WW3 = WW * 3.0f;
+		
+		float4 B300 = I[0].WorldPosition[0],
+			   B210 = I[0].WorldPosition[1],
+			   B120 = I[0].WorldPosition[2],
+			   B030 = I[1].WorldPosition[0],
+			   B021 = I[1].WorldPosition[1],
+			   B012 = I[1].WorldPosition[2],
+			   B003 = I[2].WorldPosition[0],
+			   B102 = I[2].WorldPosition[1],
+			   B201 = I[2].WorldPosition[2];
+
+
+		// Compute position from cubic control points and barycentric coords
+		float4 WorldPosition =
+			B300 * UU * U +
+			B030 * VV * V + 
+			B003 * WW * W + 
+			B210 * UU3 * V +
+			B120 * VV3 * U +
+			B021 * VV3 * W +
+			B012 * WW3 * V +
+			B102 * WW3 * U +
+			B201 * UU3 * W +
+			HSConstantData.WorldB111 * 6.0f * W * U * V;
+
+
+		// Interp remaining attributes (construct barycentric interp from bilerp primitives)
+		// NB: the HLSL compiler resolves the 1.f to something efficient (i.e. no-op)
+		FPassSpecificVSToDS Interp = PassInterpolate( PassInterpolate(I[0].PassSpecificData, U, I[1].PassSpecificData, V), 1.f, I[2].PassSpecificData, W );
+
+		// Forward the eye index along
+		#if INSTANCED_STEREO
+		Interp.FactoryInterpolants.InterpolantsVSToPS.PackedEyeIndex = I[0].PassSpecificData.FactoryInterpolants.InterpolantsVSToPS.PackedEyeIndex;
+		#endif
+
+		// get our base material params that we will potentially override
+		FMaterialTessellationParameters MaterialParameters = GetMaterialTessellationParameters(Interp.FactoryInterpolants, WorldPosition.xyz); // VF ops
+
+#if DISPLACEMENT_ANTICRACK
+
+		float2 DisplacementTexCoords = 0;
+#if NUM_MATERIAL_TEXCOORDS
+		DisplacementTexCoords = GetTextureCoordinateDS( Interp );
+#endif
+
+		FMaterialTessellationParameters DisplacementMaterialParameters = MaterialParameters;
+
+		float3x3 TangentToLocal = VertexFactoryGetTangentToLocalDS(Interp.FactoryInterpolants);
+
+		// Override to dominant data if along edge or at control vert.  This is necessary to avoid cracks causes by primitives
+		// from disjoint locations in UV space abutting each other.
+		ApplyDominantData(	DisplacementTexCoords, DisplacementMaterialParameters.TangentToWorld,
+							GetNormalDS( TangentToLocal ), GetTangentDS( TangentToLocal ),
+							U, V, W,
+							I[0].DominantEdge, I[1].DominantEdge, I[2].DominantEdge,
+							I[0].DominantVertex, I[1].DominantVertex, I[2].DominantVertex
+							);
+		
+		// Override texture coordinate 0.
+#if NUM_MATERIAL_TEXCOORDS
+		DisplacementMaterialParameters.TexCoords[0].xy = DisplacementTexCoords.xy;
+#endif
+				
+#else
+
+		FMaterialTessellationParameters DisplacementMaterialParameters = MaterialParameters;
+							
+#endif	// DISPLACEMENT_ANTICRACK
+
+		// tangent space displacement scaling
+		DisplacementMaterialParameters.TangentToWorldPreScale = I[0].DisplacementScale*U + I[1].DisplacementScale*V + I[2].DisplacementScale*W;
+
+		// world space displacement value.  This will likely use the tangent2world transform.
+		float3 WorldDisplacement = GetMaterialWorldDisplacement(DisplacementMaterialParameters);
+		WorldPosition.xyz += WorldDisplacement;
+										
+		#if USE_GLOBAL_CLIP_PLANE
+			OutGlobalClipPlaneDistance = dot(ResolvedView.GlobalClippingPlane, float4(WorldPosition.xyz - ResolvedView.PreViewTranslation.xyz, 1));
+		#endif
+
+		FPassSpecificVSToPS Result = PassFinalizeTessellationOutput(Interp, WorldPosition, MaterialParameters);
+
+		#if INSTANCED_STEREO && !MULTI_VIEW
+		BRANCH
+		if (IsInstancedStereo(I[0].PassSpecificData.FactoryInterpolants.InterpolantsVSToPS.PackedEyeIndex))
+		{
+			const uint EyeIndex = GetEyeIndex(I[0].PassSpecificData.FactoryInterpolants.InterpolantsVSToPS.PackedEyeIndex);
+
+			// Clip at the center of the screen
+			OutClipDistance = dot(Result.Position, EyeClipEdge[EyeIndex]);
+			
+			// Scale to the width of a single eye viewport
+			Result.Position.x *= 0.5 * ResolvedView.HMDEyePaddingOffset;
+
+			// Shift to the eye viewport
+			Result.Position.x += (EyeOffsetScale[EyeIndex] * Result.Position.w) * (1.0f - 0.5 * ResolvedView.HMDEyePaddingOffset);
+		}
+		#endif
+
+		return Result;
+	}
+#endif
diff --git a/Engine/Shaders/Private/PixelShaderOutputCommon.ush b/Engine/Shaders/Private/PixelShaderOutputCommon.ush
new file mode 100644
--- /dev/null
+++ b/Engine/Shaders/Private/PixelShaderOutputCommon.ush
@@ -0,0 +1,174 @@
+// Copyright 1998-2017 Epic Games, Inc. All Rights Reserved.
+
+/*=============================================================================
+PixelShaderOutputCommon.usf: To allow PS input/output passed int functions through a single struct, allowing for a more readable code (less #ifdefs, reducing the boolean hell)
+=============================================================================*/
+
+#ifndef PIXELSHADEROUTPUT_INTERPOLANTS
+	#define PIXELSHADEROUTPUT_INTERPOLANTS 0
+#endif
+#ifndef PIXELSHADEROUTPUT_BASEPASS
+	#define PIXELSHADEROUTPUT_BASEPASS 0
+#endif
+#ifndef PIXELSHADEROUTPUT_MESHDECALPASS
+	#define PIXELSHADEROUTPUT_MESHDECALPASS 0
+#endif
+#ifndef PIXELSHADEROUTPUT_MRT0
+	#define PIXELSHADEROUTPUT_MRT0 0
+#endif
+#ifndef PIXELSHADEROUTPUT_MRT1
+	#define PIXELSHADEROUTPUT_MRT1 0
+#endif
+#ifndef PIXELSHADEROUTPUT_MRT2
+	#define PIXELSHADEROUTPUT_MRT2 0
+#endif
+#ifndef PIXELSHADEROUTPUT_MRT3
+	#define PIXELSHADEROUTPUT_MRT3 0
+#endif
+#ifndef PIXELSHADEROUTPUT_MRT4
+	#define PIXELSHADEROUTPUT_MRT4 0
+#endif
+#ifndef PIXELSHADEROUTPUT_MRT5
+	#define PIXELSHADEROUTPUT_MRT5 0
+#endif
+#ifndef PIXELSHADEROUTPUT_MRT6
+	#define PIXELSHADEROUTPUT_MRT6 0
+#endif
+#ifndef PIXELSHADEROUTPUT_MRT7
+	#define PIXELSHADEROUTPUT_MRT7 0
+#endif
+#ifndef PIXELSHADEROUTPUT_COVERAGE
+	#define PIXELSHADEROUTPUT_COVERAGE 0
+#endif
+#ifndef PIXELSHADEROUTPUT_A2C
+	#define PIXELSHADEROUTPUT_A2C 0
+#endif
+
+
+// ---------------------------------------------------------------------------------
+
+void MainPS
+	(
+#if PIXELSHADEROUTPUT_INTERPOLANTS || PIXELSHADEROUTPUT_BASEPASS
+		FVertexFactoryInterpolantsVSToPS Interpolants,
+#endif
+#if PIXELSHADEROUTPUT_BASEPASS
+		FBasePassInterpolantsVSToPS BasePassInterpolants,
+#elif PIXELSHADEROUTPUT_MESHDECALPASS
+		FMeshDecalInterpolants MeshDecalInterpolants,
+#endif
+
+		in INPUT_POSITION_QUALIFIERS float4 SvPosition : SV_Position		// after all interpolators
+
+		OPTIONAL_IsFrontFace
+
+#if PIXELSHADEROUTPUT_MRT0
+		, out float4 OutTarget0 : SV_Target0
+#endif
+
+#if PIXELSHADEROUTPUT_MRT1
+		, out float4 OutTarget1 : SV_Target1
+#endif
+
+#if PIXELSHADEROUTPUT_MRT2
+		, out float4 OutTarget2 : SV_Target2
+#endif
+
+#if PIXELSHADEROUTPUT_MRT3
+		, out float4 OutTarget3 : SV_Target3
+#endif
+
+#if PIXELSHADEROUTPUT_MRT4
+		, out float4 OutTarget4 : SV_Target4
+#endif
+
+#if PIXELSHADEROUTPUT_MRT5
+		, out float4 OutTarget5 : SV_Target5
+#endif
+
+#if PIXELSHADEROUTPUT_MRT6
+		, out float4 OutTarget6 : SV_Target6
+#endif
+
+#if PIXELSHADEROUTPUT_MRT7
+		, out float4 OutTarget7 : SV_Target7
+#endif
+
+		// todo: if we are not inside a volumetric decal we could use OPTIONAL_OutDepthConservative for better performance, this would require a shader permutation
+		OPTIONAL_OutDepthConservative
+
+#if PIXELSHADEROUTPUT_COVERAGE || PIXELSHADEROUTPUT_A2C
+#if PIXELSHADEROUTPUT_A2C
+		, in uint InCoverage : SV_Coverage
+#endif
+		, out uint OutCoverage : SV_Coverage
+#endif
+	)
+{
+	// ---------------------------------------------------------------------------------
+
+	FPixelShaderIn PixelShaderIn = (FPixelShaderIn)0;
+	FPixelShaderOut PixelShaderOut = (FPixelShaderOut)0;
+
+#if PIXELSHADEROUTPUT_COVERAGE || PIXELSHADEROUTPUT_A2C
+#if PIXELSHADEROUTPUT_A2C
+	PixelShaderIn.Coverage = InCoverage;
+#else
+	PixelShaderIn.Coverage = 0xF;
+#endif
+	PixelShaderOut.Coverage = PixelShaderIn.Coverage;
+#endif 
+
+	PixelShaderIn.SvPosition = SvPosition;
+	PixelShaderIn.bIsFrontFace = bIsFrontFace;
+
+#if PIXELSHADEROUTPUT_BASEPASS
+	FPixelShaderInOut_MainPS(Interpolants, BasePassInterpolants, PixelShaderIn, PixelShaderOut);
+#elif PIXELSHADEROUTPUT_MESHDECALPASS
+	FPixelShaderInOut_MainPS(Interpolants, MeshDecalInterpolants, PixelShaderIn, PixelShaderOut);
+#elif PIXELSHADEROUTPUT_INTERPOLANTS
+	FPixelShaderInOut_MainPS(Interpolants, PixelShaderIn, PixelShaderOut);
+#else
+	FPixelShaderInOut_MainPS(PixelShaderIn, PixelShaderOut);
+#endif
+
+#if PIXELSHADEROUTPUT_MRT0
+	OutTarget0 = PixelShaderOut.MRT[0];
+#endif
+
+#if PIXELSHADEROUTPUT_MRT1
+	OutTarget1 = PixelShaderOut.MRT[1];
+#endif
+
+#if PIXELSHADEROUTPUT_MRT2
+	OutTarget2 = PixelShaderOut.MRT[2];
+#endif
+
+#if PIXELSHADEROUTPUT_MRT3
+	OutTarget3 = PixelShaderOut.MRT[3];
+#endif
+
+#if PIXELSHADEROUTPUT_MRT4
+	OutTarget4 = PixelShaderOut.MRT[4];
+#endif
+
+#if PIXELSHADEROUTPUT_MRT5
+	OutTarget5 = PixelShaderOut.MRT[5];
+#endif
+
+#if PIXELSHADEROUTPUT_MRT6
+	OutTarget6 = PixelShaderOut.MRT[6];
+#endif
+
+#if PIXELSHADEROUTPUT_MRT7
+	OutTarget7 = PixelShaderOut.MRT[7];
+#endif
+
+#if PIXELSHADEROUTPUT_COVERAGE || PIXELSHADEROUTPUT_A2C
+	OutCoverage = PixelShaderOut.Coverage;
+#endif 
+
+#if OUTPUT_PIXEL_DEPTH_OFFSET
+	OutDepth = PixelShaderOut.Depth;
+#endif 
+}
\ No newline at end of file
diff --git a/Engine/Shaders/Private/PlanarReflectionShared.ush b/Engine/Shaders/Private/PlanarReflectionShared.ush
new file mode 100644
--- /dev/null
+++ b/Engine/Shaders/Private/PlanarReflectionShared.ush
@@ -0,0 +1,105 @@
+// Copyright 1998-2017 Epic Games, Inc. All Rights Reserved.
+
+/*=============================================================================
+ PlanarReflectionShared.usf
+=============================================================================*/
+
+// Must match C++
+#define NUM_MAX_PLANAR_REFLECTION_VIEWS 2
+
+// World space reflection plane, normalized
+float4 ReflectionPlane;
+
+float3 PlanarReflectionOrigin;
+float4 PlanarReflectionXAxis;
+float4 PlanarReflectionYAxis;
+
+// Used to transform normals into the mirrored space
+float3x4 InverseTransposeMirrorMatrix;
+
+// x = scale for plane distance fading, y = bias for plane distance fading, z = controls normal distortion strength
+float3 PlanarReflectionParameters;
+
+// x = scale for angle fading, y = bias for angle fading
+float2 PlanarReflectionParameters2;
+
+// Projection matrix that was used to capture PlanarReflectionTexture
+float4x4 ProjectionWithExtraFOV[NUM_MAX_PLANAR_REFLECTION_VIEWS];
+
+// Scale and bias to transform NDC into UVs for PlanarReflectionTexture
+float4 PlanarReflectionScreenScaleBias[NUM_MAX_PLANAR_REFLECTION_VIEWS];
+
+// True if we're rendering a stereo planar reflection
+bool bIsStereo;
+
+Texture2D PlanarReflectionTexture;
+SamplerState PlanarReflectionSampler;
+
+float4 ComputePlanarReflections(float3 WorldPosition, float3 WorldNormal, float Roughness, SamplerState SharedClampSampler)
+{
+	float4 OutPlanarReflection = 0;
+
+	float PlaneDistance = dot(ReflectionPlane, float4(WorldPosition, -1));
+	float DistanceFade = 1 - saturate(abs(PlaneDistance) * PlanarReflectionParameters.x + PlanarReflectionParameters.y);
+
+	float3 PlaneOriginToWorldPosition = WorldPosition - PlanarReflectionOrigin;
+	float XAxisDistance = dot(PlaneOriginToWorldPosition, PlanarReflectionXAxis.xyz);
+	float XAxisFade = saturate((PlanarReflectionXAxis.w - abs(XAxisDistance)) * PlanarReflectionParameters.x);
+	float YAxisDistance = dot(PlaneOriginToWorldPosition, PlanarReflectionYAxis.xyz);
+	float YAxisFade = saturate((PlanarReflectionYAxis.w - abs(YAxisDistance)) * PlanarReflectionParameters.x);
+	DistanceFade *= XAxisFade * YAxisFade;
+
+	BRANCH
+	if (DistanceFade > 0)
+	{
+		// CameraToPixel in the main view is what we used as ReflectionVector when rendering the reflection pass to PlanarReflectionTexture
+		float3 CameraToPixel = normalize(WorldPosition - ResolvedView.WorldCameraOrigin);
+		// Reflect the effective ReflectionVector in mirrored space to get the original camera vector
+		float3 MirroredCameraVector = reflect(CameraToPixel, -ReflectionPlane.xyz);
+		// Transform the GBuffer normal into mirrored space
+		float3 MirroredNormal = mul(WorldNormal, (float3x3)InverseTransposeMirrorMatrix);
+		// Reflect the original camera vector across the GBuffer normal in mirrored space
+		float3 MirroredReflectionVectorOffNormal = reflect(MirroredCameraVector, MirroredNormal);
+		// At this point we have a new reflection vector off of the GBuffer normal, and we need to approximate its intersection with the scene
+		// An accurate intersection would ray trace the planar reflection depth buffer
+		// As an approximation we are just intersecting with a user defined sphere
+		float3 VirtualReflectionSpherePosition = WorldPosition + MirroredReflectionVectorOffNormal * PlanarReflectionParameters.z;
+		// Transform the intersection position into view space
+		float3 ViewVirtualReflectionSpherePosition = mul(float4(VirtualReflectionSpherePosition + ResolvedView.PreViewTranslation.xyz, 1), ResolvedView.TranslatedWorldToView).xyz;
+		// Transform the intersection position into clip space using the same projection matrix used to render PlanarReflectionTexture
+		float4 ClipVirtualReflectionSpherePosition = mul(float4(ViewVirtualReflectionSpherePosition, 1), ProjectionWithExtraFOV[ResolvedView.StereoPassIndex]);
+
+		uint EyeIndex = 0;
+
+#if !MULTI_VIEW
+		if (bIsStereo)
+		{
+			EyeIndex = ResolvedView.StereoPassIndex;
+		}
+#endif
+
+		float2 NDC = clamp(ClipVirtualReflectionSpherePosition.xy / ClipVirtualReflectionSpherePosition.w, -1, 1);
+		float2 ViewportUV = NDC * PlanarReflectionScreenScaleBias[EyeIndex].xy + PlanarReflectionScreenScaleBias[EyeIndex].zw;
+
+		float AngleFade = saturate(dot(ReflectionPlane.xyz, WorldNormal) * PlanarReflectionParameters2.x + PlanarReflectionParameters2.y);
+		float RoughnessFade = 1 - saturate((Roughness - .2f) * 10.0f);
+		float FinalFade = DistanceFade * AngleFade * RoughnessFade;
+
+		float4 PlanarReflectionTextureValue = Texture2DSampleLevel(
+			PlanarReflectionTexture, 
+#if SUPPORTS_INDEPENDENT_SAMPLERS
+			SharedClampSampler, 
+#else
+			PlanarReflectionSampler,
+#endif
+			ViewportUV, 
+			0);
+
+		// Fade out in regions of the planar reflection that weren't written to, so we can composite with other reflection methods
+		FinalFade *= PlanarReflectionTextureValue.a;
+		OutPlanarReflection.rgb = PlanarReflectionTextureValue.rgb * FinalFade;
+		OutPlanarReflection.a = FinalFade;
+	}
+
+	return OutPlanarReflection;
+}
diff --git a/Engine/Shaders/Private/PostProcessCommon.ush b/Engine/Shaders/Private/PostProcessCommon.ush
new file mode 100644
--- /dev/null
+++ b/Engine/Shaders/Private/PostProcessCommon.ush
@@ -0,0 +1,209 @@
+// Copyright 1998-2017 Epic Games, Inc. All Rights Reserved.
+
+/*=============================================================================
+	PostProcessCommon.usf: PostProcessing shared functions and structures.
+=============================================================================*/
+
+#ifndef __POST_PROCESS_COMMON__
+#define __POST_PROCESS_COMMON__
+
+/// e.g.
+//   PostprocessInput0.SampleLevel(PostprocessInput0Sampler, float2 UV, 0)
+//   Texture2DSample(PostprocessInput0, PostprocessInput0Sampler, float2 UV)
+//   PostprocessInput0.Load(Texel.xyz)
+//   PostprocessInput0.GetDimensions(uint, out uint width, out uint height, out uint levels)
+//   PostprocessInput0.GetDimensions(uint, out float width, out float height, out float levels)
+//   PostprocessInput0.GetDimensions(out uint width, out uint height)
+//   PostprocessInput0.GetDimensions(out float width, out float height)
+Texture2D PostprocessInput0;
+SamplerState PostprocessInput0Sampler;
+Texture2D PostprocessInput1;
+SamplerState PostprocessInput1Sampler;
+Texture2D PostprocessInput2;
+SamplerState PostprocessInput2Sampler;
+Texture2D PostprocessInput3;
+SamplerState PostprocessInput3Sampler;
+Texture2D PostprocessInput4;
+SamplerState PostprocessInput4Sampler;
+Texture2D PostprocessInput5;
+SamplerState PostprocessInput5Sampler;
+Texture2D PostprocessInput6;
+SamplerState PostprocessInput6Sampler;
+
+// width, height, 1/width, 1/height
+float4 PostprocessInput0Size;
+float4 PostprocessInput1Size;
+float4 PostprocessInput2Size;
+float4 PostprocessInput3Size;
+float4 PostprocessInput4Size;
+float4 PostprocessInput5Size;
+float4 PostprocessInput6Size;
+
+// xy = min valid UV in PostprocessInput%d, zw = max valid UV.
+float4 PostprocessInput0MinMax;
+float4 PostprocessInput1MinMax;
+float4 PostprocessInput2MinMax;
+float4 PostprocessInput3MinMax;
+float4 PostprocessInput4MinMax;
+float4 PostprocessInput5MinMax;
+float4 PostprocessInput6MinMax;
+
+// viewport width, height, 1/width, 1/height (scaled to the current rendertarget resolution), depends on SetViewportAndCallRHI() call
+float4 ViewportSize;
+// in pixels (scaled to the current rendertarget resolution), float4(minx,miny,maxx,maxy), depends on SetViewportAndCallRHI() call
+uint4 ViewportRect;
+
+SamplerState BilinearTextureSampler0;
+SamplerState BilinearTextureSampler1;
+
+// transforms ScreenPos to the screen pixel position, not the viewport local position (*xy +wz)
+// (0,0) left top, (width-1,height-1) at right bottom, in pixels (without offset to reach the pixel center)
+// 	use ComputePixelPosCenter(UVAndScreenPos.zw, bPixelCenter) instead of directly accessing this
+float4 ScreenPosToPixel;
+
+// uses ScreenPosToPixel which is setup in FPostProcessPassParameters
+// for integer output bPixelCenter should be set to true (to avoid rounding errors), but often it's better to use "int2 PixelPos = (int2)SvPosition.xy;" instead
+// @param ScreenPos -1 .. 1, usually from UVAndScreenPos.zw
+// @param bCenter false: left top of the pixel and almost no fractional value, true: added 0.5 to be pixel center
+float2 ComputePixelPosCenter(float2 ScreenPos, bool bPixelCenter)
+{
+	float2 PixelOffset = bPixelCenter ? 0.5f : 0.0f;
+	return ScreenPos * ScreenPosToPixel.xy + ScreenPosToPixel.zw + PixelOffset;
+}
+
+// not well optimized but can be useful
+float2 ScreenPosFromPixelPos(float2 PixelPos, bool bPixelCenter)
+{
+	float2 PixelOffset = bPixelCenter ? 0.5f : 0.0f;
+
+	return (PixelPos - PixelOffset - ScreenPosToPixel.zw) / ScreenPosToPixel.xy;
+}
+
+// shape of the lens for vignette like effects and masking
+// @param ScreenPos -1 .. 1
+// @return 0..1 0:borders, 1:center
+float DiscMask(float2 ScreenPos)
+{
+	float x = saturate(1.0f - dot(ScreenPos, ScreenPos));
+
+	return x * x;
+}
+
+// black on the borders
+// @param ScreenPos -1 .. 1
+// @return 0..1 0:borders, 1:center
+float RectMask(float2 ScreenPos)
+{
+	float2 UV = saturate(ScreenPos * 0.5 + 0.5f);
+	float2 Mask2 = UV * (1 - UV);
+
+	return Mask2.x * Mask2.y * 8.0f;
+}
+
+// for rectangles with border
+// @return >=0, 0 if inside
+float ComputeDistanceToRect(int2 Pos, int2 LeftTop, int2 Extent, bool bRoundBorders = true)
+{
+	int2 RightBottom = LeftTop + Extent - 1;
+
+	// use int for more precision
+	int2 Rel = max(int2(0, 0), Pos - RightBottom) + max(int2(0, 0), LeftTop - Pos);
+
+	if(bRoundBorders)
+	{
+		// euclidian distance (round corners)
+		return length((float2)Rel);
+	}
+	else
+	{
+		// manhatten distance (90 degree corners)
+		return max(Rel.x, Rel.y);
+	}
+}
+
+
+float4 MappingPolynomial;			 // RGB = a, b, c where y = a * x*x + b * x + c
+
+// @param InLDRColor needs to be LDR (0..1) and in linear space
+half3 ColorCorrection(half3 InLDRColor)
+{
+	// final color correction to adjust for hardware differences, to make quick adjustements before a demo or simply a user setting
+	return MappingPolynomial.x * (InLDRColor * InLDRColor) + MappingPolynomial.y * InLDRColor + MappingPolynomial.z;
+}
+
+
+// Generate a mask to darken the screen borders.
+// (Camera artifact and artistic effect)
+// @param VignetteCircleSpacePos from VignetteSpace()
+float ComputeVignetteMask(float2 VignetteCircleSpacePos, float Intensity)
+{
+	// Natural vignetting
+	// cosine-fourth law
+	VignetteCircleSpacePos *= Intensity;
+	float Tan2Angle = dot( VignetteCircleSpacePos, VignetteCircleSpacePos );
+	float Cos4Angle = Square( rcp( Tan2Angle + 1 ) );
+	return Cos4Angle;
+}
+
+// Scale {-1 to 1} of the viewport space to vignette circle space.
+// Vignette space is scaled such that regardless of viewport aspect ratio, 
+// corners are at the same brightness on a circle.
+float2 VignetteSpace(float2 Pos)
+{
+	// could be optimized but this computation should be done in the vertex shader (3 or 4 vertices)
+	float Scale = sqrt(2.0) / sqrt(1.0 + (ViewportSize.y * ViewportSize.z) * (ViewportSize.y * ViewportSize.z));
+	return Pos * float2(1.0,  ViewportSize.y * ViewportSize.z) * Scale;
+}
+
+
+// gets 4 nearby SceneDepth values for one UV value, useful for depth downsample, uses Gather() if possible
+float4 GatherSceneDepth(float2 UV, float2 InvBufferSize)
+{
+#if FEATURE_LEVEL >= FEATURE_LEVEL_SM5
+	// using Gather: xyzw in counter clockwise order starting with the sample to the lower left of the queried location
+	float4 DeviceZ = SceneDepthTextureNonMS.Gather(BilinearTextureSampler0, UV);
+
+	return float4(
+		ConvertFromDeviceZ(DeviceZ.x),
+		ConvertFromDeviceZ(DeviceZ.y),
+		ConvertFromDeviceZ(DeviceZ.z),
+		ConvertFromDeviceZ(DeviceZ.w)
+		);
+#elif FEATURE_LEVEL == FEATURE_LEVEL_SM4
+	float2 TexelScale = 0.5f * InvBufferSize;
+	// @todo Optimize
+	// Manually sample neighbouring texels in counter clockwise order starting with the sample to the lower left of the queried location.
+	return float4(
+		CalcSceneDepth(UV + (float2(-1, 1) * TexelScale)),
+		CalcSceneDepth(UV + (float2( 1, 1) * TexelScale)),
+		CalcSceneDepth(UV + (float2( 1,-1) * TexelScale)),
+		CalcSceneDepth(UV + (float2(-1,-1) * TexelScale))
+		); 
+#else
+	return (float4)CalcSceneDepth(UV);
+#endif // FEATURE_LEVEL
+}
+
+half4 UnwrappedTexture3DSample( Texture2D Texture, SamplerState Sampler, float3 UVW, float Size )
+{
+	// a volume texture 16x16x16 would be unwrapped to a 2d texture 256x16
+
+	float IntW = floor( UVW.z * Size - 0.5 );
+	half FracW = UVW.z * Size - 0.5 - IntW;
+
+	float U = ( UVW.x + IntW ) / Size;
+	float V = UVW.y;
+
+	half4 RG0 = Texture2DSample( Texture, Sampler, float2(U, V) );
+	half4 RG1 = Texture2DSample( Texture, Sampler, float2(U + 1.0f / Size, V) );
+
+	return lerp(RG0, RG1, FracW);
+}
+
+bool IsComputeUVOutOfBounds(in float2 UV)
+{
+	float2 CenterDist = abs(UV - 0.5f);
+	return (max(CenterDist.x, CenterDist.y) >= 0.5f);
+}
+
+#endif // __POST_PROCESS_COMMON__
diff --git a/Engine/Shaders/Private/PostProcessHistogramCommon.ush b/Engine/Shaders/Private/PostProcessHistogramCommon.ush
new file mode 100644
--- /dev/null
+++ b/Engine/Shaders/Private/PostProcessHistogramCommon.ush
@@ -0,0 +1,144 @@
+// Copyright 1998-2017 Epic Games, Inc. All Rights Reserved.
+
+/*=============================================================================
+	PostProcessHistogramCommon.usf: PostProcessing histogram shared functions and structures.
+=============================================================================*/
+
+// [0] .x:ExposureLowPercent/100, .y:EyeAdaptationHighPercent/100, .z:EyeAdaptationMin, .w:EyeAdaptationMax
+// [1] .x:exp2(ExposureOffset), .y:DeltaWorldTime, .zw: EyeAdaptionSpeedUp/Down
+// [2] .x:Histogram multiply, .y:Histogram add, .z:HistogramMinIntensity w:unused
+float4 EyeAdaptationParams[3];
+
+// inverse of ComputeLuminanceFromHistogramPosition
+// is executed more often than ComputeLuminanceFromHistogramPosition()
+// @param Luminance
+// @return HistogramPosition 0..1
+float ComputeHistogramPositionFromLuminance(float Luminance)
+{
+	return log2(Luminance) * EyeAdaptationParams[2].x + EyeAdaptationParams[2].y;
+}
+
+// inverse of ComputeHistogramPositionFromLuminance()
+// is not as often executed as ComputeHistogramPositionFromLuminance()
+// @param HistogramPosition 0..1
+// @return Lumiance
+float ComputeLuminanceFromHistogramPosition(float HistogramPosition)
+{
+	return exp2( (HistogramPosition - EyeAdaptationParams[2].y) / EyeAdaptationParams[2].x );
+}
+
+
+
+#define HISTOGRAM_SIZE 64
+#define HISTOGRAM_TEXEL_SIZE (HISTOGRAM_SIZE / 4)
+
+float4 ComputeARGBStripeMaskInt(uint x)
+{
+#if ES2_PROFILE
+	// Integer modulo/remainder is not an allowed operation on ES 2
+	// todo: move function to central spot, here we assume unsigned
+	uint Temp = x - (x / 4) * 4;
+	return float4(
+		Temp == 0,
+		Temp == 1,
+		Temp == 2,
+		Temp == 3);
+#else
+	return float4(
+		(x % 4) == 0,
+		(x % 4) == 1,
+		(x % 4) == 2,
+		(x % 4) == 3);
+#endif
+}
+
+float GetHistogramBucket(Texture2D HistogramTexture, uint BucketIndex)
+{
+	uint Texel = BucketIndex / 4;
+	
+	float4 HistogramColor = HistogramTexture.Load(int3(Texel, 0, 0));
+
+	// WAR for a GLSL shader compiler bug in the driver
+#if 0
+	float UnweightedValue = dot(HistogramColor, ComputeARGBStripeMaskInt(BucketIndex));
+#else
+#if ES2_PROFILE
+	// Integer modulo/remainder is not an allowed operation on ES 2
+	// todo: move function to central spot, here we assume unsigned
+	uint channel = BucketIndex - (BucketIndex / 4) * 4;
+#else
+	uint channel = BucketIndex % 4;
+#endif
+	float UnweightedValue = HistogramColor.r;
+	UnweightedValue = (channel == 1) ? HistogramColor.g : UnweightedValue;
+	UnweightedValue = (channel == 2) ? HistogramColor.b : UnweightedValue;
+	UnweightedValue = (channel == 3) ? HistogramColor.a : UnweightedValue;
+#endif
+
+	// because of the logarithmic scale we need to weight the buckets ?
+
+//	return UnweightedValue * pow(2, BucketIndex);
+//	return UnweightedValue * pow(2, BucketIndex / HISTOGRAM_SIZE * (2 * HDRLogRadius));
+	return UnweightedValue;
+}
+
+float ComputeHistogramSum(Texture2D HistogramTexture)
+{
+	float Sum = 0;
+
+	for(uint i = 0; i < HISTOGRAM_SIZE; ++i)
+	{
+		Sum += GetHistogramBucket(HistogramTexture, i);
+	}
+
+	return Sum;
+}
+
+// @param MinFractionSum e.g. ComputeHistogramSum() * 0.5f for 50% percentil
+// @param MaxFractionSum e.g. ComputeHistogramSum() * 0.9f for 90% percentil
+float ComputeAverageLuminaneWithoutOutlier(Texture2D HistogramTexture, float MinFractionSum, float MaxFractionSum)
+{
+	float2 SumWithoutOutliers = 0;
+	
+	UNROLL for(uint i = 0; i < HISTOGRAM_SIZE; ++i)
+	{
+		float LocalValue = GetHistogramBucket(HistogramTexture, i);
+
+		// remove outlier at lower end
+		float Sub = min(LocalValue, MinFractionSum);
+		LocalValue = LocalValue - Sub;
+		MinFractionSum -= Sub;
+		MaxFractionSum -= Sub;
+
+		// remove outlier at upper end
+		LocalValue = min(LocalValue, MaxFractionSum);
+		MaxFractionSum -= LocalValue;
+
+		float LuminanceAtBucket = ComputeLuminanceFromHistogramPosition(i / (float)HISTOGRAM_SIZE);
+
+		SumWithoutOutliers += float2(LuminanceAtBucket, 1) * LocalValue;
+	}
+
+	return SumWithoutOutliers.x / max(0.0001f, SumWithoutOutliers.y);
+}
+
+// used by bloom and tonemapper VS
+// requires EyeAdaptationParams to be set
+float ComputeEyeAdaptationExposure(Texture2D HistogramTexture)
+{
+	float HistogramSum = ComputeHistogramSum(HistogramTexture);
+	float UnclampedAdaptedLuminance = ComputeAverageLuminaneWithoutOutlier(HistogramTexture, HistogramSum * EyeAdaptationParams[0].x, HistogramSum * EyeAdaptationParams[0].y);
+	float ClampedAdaptedLuminance = clamp(UnclampedAdaptedLuminance, EyeAdaptationParams[0].z, EyeAdaptationParams[0].w);
+
+	return ClampedAdaptedLuminance;
+}
+
+// Computation of triangle function on normalized space
+// used in Basic EyeAdaptation in weighting the center of the screen
+// @return range of 0..1 with 1 nx = 0.5
+// @param nx in 0..1 
+// @param slope assumed possitive
+float AdaptationWeight(float nx, float slope)
+{
+	return max(1.f - 2. * slope * abs(nx - 0.5f), 0.0f);
+}
