Branch snapshot for CL 3555263

diff --git a/Engine/Shaders/Private/PostProcessTemporalCommon.ush b/Engine/Shaders/Private/PostProcessTemporalCommon.ush
new file mode 100644
--- /dev/null
+++ b/Engine/Shaders/Private/PostProcessTemporalCommon.ush
@@ -0,0 +1,651 @@
+// Copyright 1998-2017 Epic Games, Inc. All Rights Reserved.
+
+/*=============================================================================
+	   PostProcessTemporalCommon.ush: Common Temporal AA Functionality
+-------------------------------------------------------------------------------
+This is the common bits of the temporal AA shader
+which can be configured via defines for various temporal AA passes.
+=============================================================================*/
+
+// 1 = Use tighter AABB clamp for history.
+// 0 = Use simple min/max clamp.
+#ifndef AA_AABB
+	#define AA_AABB 1
+#endif
+
+// 0 = Input doesn't have or need the alpha channel. Don't AA or preserve the alpha. Potentially use alpha channel of the output to improve quality.
+//     Leverages dead code removal to work in RGB instead of RGBA.
+// 1 = Anti-alias the alpha channel also. 
+#ifndef AA_ALPHA
+	#define AA_ALPHA 0
+#endif
+
+// Cross distance in pixels used in depth search X pattern.
+// 0 = Turn this feature off.
+// 2 = Is required for standard temporal AA pass.
+#ifndef AA_CROSS
+	#define AA_CROSS 2
+#endif
+
+// 1 = Use dynamic motion.
+// 0 = Skip dynamic motion, currently required for half resolution passes.
+#ifndef AA_DYNAMIC
+	#define AA_DYNAMIC 1
+#endif
+
+// 1 = Use filtered sample.
+// 0 = Use center sample.
+#ifndef AA_FILTERED
+	#define AA_FILTERED 1
+#endif
+
+// Disables filtering when high contrast
+// Useful for Catmull Rom filter kernel
+#ifndef AA_FILTERED_CONTRAST_LIMIT
+	#define AA_FILTERED_CONTRAST_LIMIT 1
+#endif
+
+// 0 = Dynamic motion based lerp value (default).
+// non-zero = Use 1/LERP fixed lerp value (used for reflections).
+#ifndef AA_LERP
+	#define AA_LERP 0
+#endif
+
+// 1 = Use higher quality round clamp.
+// 0 = Use lower quality but faster box clamp.
+#ifndef AA_ROUND
+	#define AA_ROUND 1
+#endif
+
+// 1 = Use extra clamp to avoid NANs
+// 0 = Don't use.
+#ifndef AA_NAN
+	#define AA_NAN 1
+#endif
+
+// Fix for lack of borders during current frame filter.
+#ifndef AA_BORDER
+	#define AA_BORDER 0
+#endif
+
+// Force clamp on alpha.
+#ifndef AA_FORCE_ALPHA_CLAMP
+	#define AA_FORCE_ALPHA_CLAMP 0
+#endif
+
+// Use YCoCg path.
+#ifndef AA_YCOCG
+	#define AA_YCOCG 0
+#endif
+
+// Bicubic filter history
+#ifndef AA_BICUBIC
+	#define AA_BICUBIC 0
+#endif
+
+// Tone map to kill fireflies
+#ifndef AA_TONE
+	#define AA_TONE 1
+#endif
+
+// Neighborhood clamping. Disable for testing reprojection.
+#ifndef AA_CLAMP
+	#define AA_CLAMP 1
+#endif
+
+// Antighosting using dynamic mask
+#ifndef AA_DYNAMIC_ANTIGHOST
+	#define AA_DYNAMIC_ANTIGHOST 0
+#endif
+
+// Compute alpha accurately according samples' to tone weights.
+#ifndef AA_ACCURATE_ALPHA
+	#define AA_ACCURATE_ALPHA 0
+#endif
+
+// Sample the stencil buffer inline rather than multiple masked passes.
+#ifndef AA_SINGLE_PASS_RESPONSIVE
+	#define AA_SINGLE_PASS_RESPONSIVE	0
+#endif
+
+	float2 ScreenPos = ( UV * PostprocessInput0Size.xy - 0.5 - ScreenPosToPixel.zw ) / ScreenPosToPixel.xy;
+
+	// Handle single or multi-pass responsive AA
+#if RESPONSIVE
+	const float bIsResponsiveAAPixel = 1.f;
+#elif AA_SINGLE_PASS_RESPONSIVE
+	#define STENCIL_TEMPORAL_RESPONSIVE_AA_MASK (1 << 3)
+	int2 SceneStencilUV = (int2)trunc(UV * PostprocessInput0Size.xy);
+	uint SceneStencilRef = SceneStencilTexture.Load(int3(SceneStencilUV, 0)) STENCIL_COMPONENT_SWIZZLE;
+	const float bIsResponsiveAAPixel = (SceneStencilRef & STENCIL_TEMPORAL_RESPONSIVE_AA_MASK) ? 1.f : TemporalAAComputeParams.x;
+#else
+	const float bIsResponsiveAAPixel = 0.f;
+#endif
+
+	// FIND MOTION OF PIXEL AND NEAREST IN NEIGHBORHOOD
+    // ------------------------------------------------
+	float3 PosN; // Position of this pixel, possibly later nearest pixel in neighborhood.
+	PosN.xy = ScreenPos;
+	PosN.z = SceneDepthTexture.SampleLevel(SceneDepthTextureSampler, UV, 0).r;
+	// Screen position of minimum depth.
+	float2 VelocityOffset = float2(0.0, 0.0);
+	#if AA_CROSS
+		// For motion vector, use camera/dynamic motion from min depth pixel in pattern around pixel.
+		// This enables better quality outline on foreground against different motion background.
+		// Larger 2 pixel distance "x" works best (because AA dilates surface).
+		float4 Depths;
+		Depths.x = SceneDepthTexture.SampleLevel(SceneDepthTextureSampler, UV, 0, int2(-AA_CROSS, -AA_CROSS)).r;
+		Depths.y = SceneDepthTexture.SampleLevel(SceneDepthTextureSampler, UV, 0, int2( AA_CROSS, -AA_CROSS)).r;
+		Depths.z = SceneDepthTexture.SampleLevel(SceneDepthTextureSampler, UV, 0, int2(-AA_CROSS,  AA_CROSS)).r;
+		Depths.w = SceneDepthTexture.SampleLevel(SceneDepthTextureSampler, UV, 0, int2( AA_CROSS,  AA_CROSS)).r;
+
+		float2 DepthOffset = float2(AA_CROSS, AA_CROSS);
+		float DepthOffsetXx = float(AA_CROSS);
+#if HAS_INVERTED_Z_BUFFER
+		// Nearest depth is the largest depth (depth surface 0=far, 1=near).
+		if(Depths.x > Depths.y) 
+		{
+			DepthOffsetXx = -AA_CROSS;
+		}
+		if(Depths.z > Depths.w) 
+		{
+			DepthOffset.x = -AA_CROSS;
+		}
+		float DepthsXY = max(Depths.x, Depths.y);
+		float DepthsZW = max(Depths.z, Depths.w);
+		if(DepthsXY > DepthsZW) 
+		{
+			DepthOffset.y = -AA_CROSS;
+			DepthOffset.x = DepthOffsetXx; 
+		}
+		float DepthsXYZW = max(DepthsXY, DepthsZW);
+		if(DepthsXYZW > PosN.z) 
+		{
+			// This is offset for reading from velocity texture.
+			// This supports half or fractional resolution velocity textures.
+			// With the assumption that UV position scales between velocity and color.
+			VelocityOffset = DepthOffset * PostprocessInput0Size.zw;
+			// This is [0 to 1] flipped in Y.
+			//PosN.xy = ScreenPos + DepthOffset * ViewportSize.zw * 2.0;
+			PosN.z = DepthsXYZW;
+		}
+#else
+#error Fix me!
+#endif // HAS_INVERTED_Z_BUFFER
+	#endif	// AA_CROSS
+
+	// Camera motion for pixel or nearest pixel (in ScreenPos space).
+	float4 ThisClip = float4( PosN.xy, PosN.z, 1 );
+	float4 PrevClip = mul( ThisClip, View.ClipToPrevClip );
+	float2 PrevScreen = PrevClip.xy / PrevClip.w;
+	float2 BackN = PosN.xy - PrevScreen;
+
+	float2 BackTemp = BackN * ViewportSize.xy;
+	#if AA_DYNAMIC
+		float2 VelocityN;
+		#if AA_CROSS
+			VelocityN = PostprocessInput3.SampleLevel(PostprocessInput3Sampler, UV + VelocityOffset, 0).xy;
+		#else
+			VelocityN = PostprocessInput3.SampleLevel(PostprocessInput3Sampler, UV, 0).xy;
+		#endif
+		VelocityN *= VelocityScaling.x;
+		bool DynamicN = VelocityN.x > 0.0;
+		if(DynamicN)
+		{
+			BackN = DecodeVelocityFromTexture(VelocityN);
+		}
+		BackTemp = BackN * ViewportSize.xy;
+	#endif
+	#if !AA_BICUBIC
+		// Save the amount of pixel offset of just camera motion, used later as the amount of blur introduced by history.
+		float HistoryBlurAmp = 2.0;
+		float HistoryBlur = saturate(abs(BackTemp.x) * HistoryBlurAmp + abs(BackTemp.y) * HistoryBlurAmp);
+		float Velocity = sqrt(dot(BackTemp, BackTemp));
+	#endif
+	// Easier to do off screen check before conversion.
+	// BackN is in units of 2pixels/viewportWidthInPixels
+	// This converts back projection vector to [-1 to 1] offset in viewport.
+	BackN = ScreenPos - BackN;
+	bool OffScreen = max(abs(BackN.x), abs(BackN.y)) >= 1.0;
+	// Also clamp to be on screen (fixes problem with DOF).
+	// The .z and .w is the 1/width and 1/height.
+	// This clamps to be a pixel inside the viewport.
+	BackN.x = clamp(BackN.x, -1.0 + ViewportSize.z, 1.0 - ViewportSize.z);
+	BackN.y = clamp(BackN.y, -1.0 + ViewportSize.w, 1.0 - ViewportSize.w);
+	// Convert from [-1 to 1] to view rectangle which is somewhere in [0 to 1].
+	// The extra +0.5 factor is because ScreenPosToPixel.zw is incorrectly computed
+	// as the upper left of the pixel instead of the center of the pixel.
+	BackN = (BackN * ScreenPosToPixel.xy + ScreenPosToPixel.zw + 0.5) * PostprocessInput0Size.zw;
+
+
+	// FILTER PIXEL (RESAMPLE TO REMOVE JITTER OFFSET) AND GET NEIGHBORHOOD
+    // --------------------------------------------------------------------
+	// 012
+	// 345
+	// 678
+	#if AA_YCOCG	
+		// Special case, only using 5 taps.
+		float4 Neighbor1 = PostprocessInput0.SampleLevel(PostprocessInput0Sampler, UV, 0, int2( 0, -1));
+		float4 Neighbor3 = PostprocessInput0.SampleLevel(PostprocessInput0Sampler, UV, 0, int2(-1,  0));
+		float4 Neighbor4 = PostprocessInput0.SampleLevel(PostprocessInput0Sampler, UV, 0);
+		float4 Neighbor5 = PostprocessInput0.SampleLevel(PostprocessInput0Sampler, UV, 0, int2( 1,  0));
+		float4 Neighbor7 = PostprocessInput0.SampleLevel(PostprocessInput0Sampler, UV, 0, int2( 0,  1));
+		Neighbor1.rgb = RGBToYCoCg(Neighbor1.rgb);
+		Neighbor3.rgb = RGBToYCoCg(Neighbor3.rgb);
+		Neighbor4.rgb = RGBToYCoCg(Neighbor4.rgb);
+		Neighbor5.rgb = RGBToYCoCg(Neighbor5.rgb);
+		Neighbor7.rgb = RGBToYCoCg(Neighbor7.rgb);
+
+		#if AA_FILTERED_CONTRAST_LIMIT && AA_TONE
+			float Luma_m1, Luma_m2;
+			Luma_m1  = Neighbor1.x; Luma_m2  = Neighbor1.x * Neighbor1.x;
+			Luma_m1 += Neighbor3.x; Luma_m2 += Neighbor3.x * Neighbor3.x;
+			Luma_m1 += Neighbor4.x; Luma_m2 += Neighbor4.x * Neighbor4.x;
+			Luma_m1 += Neighbor5.x; Luma_m2 += Neighbor5.x * Neighbor5.x;
+			Luma_m1 += Neighbor7.x; Luma_m2 += Neighbor7.x * Neighbor7.x;
+
+			float LumaAvg = Luma_m1 * (1.0 / 5.0);
+			float LumaVar = abs( Luma_m2 * (1.0 / 5.0) - Pow2( LumaAvg ) );
+		#else
+			float LumaVar = 0;
+		#endif
+
+		#if AA_TONE
+			float NeighborWeight1 = HdrWeightY(Neighbor1.x, InExposureScale);
+			float NeighborWeight3 = HdrWeightY(Neighbor3.x, InExposureScale);
+			float NeighborWeight4 = HdrWeightY(Neighbor4.x, InExposureScale);
+			float NeighborWeight5 = HdrWeightY(Neighbor5.x, InExposureScale);
+			float NeighborWeight7 = HdrWeightY(Neighbor7.x, InExposureScale);
+			Neighbor1.xyz *= NeighborWeight1;
+			Neighbor3.xyz *= NeighborWeight3;
+			Neighbor4.xyz *= NeighborWeight4;
+			Neighbor5.xyz *= NeighborWeight5;
+			Neighbor7.xyz *= NeighborWeight7;
+		#else
+			float NeighborWeight1 = 1;
+			float NeighborWeight3 = 1;
+			float NeighborWeight4 = 1;
+			float NeighborWeight5 = 1;
+			float NeighborWeight7 = 1;
+		#endif
+		#if AA_FILTERED
+			#if AA_ACCURATE_ALPHA
+				float NeighborFilteringWeight1 = NeighborWeight1 * PlusWeights[0];
+				float NeighborFilteringWeight3 = NeighborWeight3 * PlusWeights[1];
+				float NeighborFilteringWeight4 = NeighborWeight4 * PlusWeights[2];
+				float NeighborFilteringWeight5 = NeighborWeight5 * PlusWeights[3];
+				float NeighborFilteringWeight7 = NeighborWeight7 * PlusWeights[4];
+				float FilteredWeight = 
+					NeighborFilteringWeight1 +
+					NeighborFilteringWeight3 +
+					NeighborFilteringWeight4 +
+					NeighborFilteringWeight5 +
+					NeighborFilteringWeight7;
+
+				// Filter neighbors YCoCg in in HDR space.
+				float4 Filtered;
+				Filtered.xyz =
+					Neighbor1.xyz * PlusWeights[0] +
+					Neighbor3.xyz * PlusWeights[1] +
+					Neighbor4.xyz * PlusWeights[2] +
+					Neighbor5.xyz * PlusWeights[3] +
+					Neighbor7.xyz * PlusWeights[4];
+
+				// Filter neighbors alpha in linear space according to HDR space weights.
+				Filtered.a = (
+					Neighbor1.a * NeighborFilteringWeight1 +
+					Neighbor3.a * NeighborFilteringWeight3 +
+					Neighbor4.a * NeighborFilteringWeight4 +
+					Neighbor5.a * NeighborFilteringWeight5 +
+					Neighbor7.a * NeighborFilteringWeight7) / FilteredWeight;
+			#else
+				// Filter neighbors YCoCg and alpha. Might be in linear or HDR space.
+				float4 Filtered = 
+					Neighbor1 * PlusWeights[0] +
+					Neighbor3 * PlusWeights[1] +
+					Neighbor4 * PlusWeights[2] +
+					Neighbor5 * PlusWeights[3] +
+					Neighbor7 * PlusWeights[4];
+
+				if( LumaVar > 1 )
+					Filtered = Neighbor4;
+			#endif
+			#if AA_BORDER
+				// Use unfiltered for 1 pixel border.
+				float2 TestPos = abs(ScreenPos);
+				// Add 1 pixel and check if off screen.
+				TestPos += ViewportSize.zw * 2.0;
+				bool FilteredOffScreen = max(TestPos.x, TestPos.y) >= 1.0;
+				if(FilteredOffScreen) 
+				{
+					Filtered = Neighbor4;
+					#if AA_ACCURATE_ALPHA
+						FilteredWeight = NeighborWeight4;
+					#endif
+				}
+			#endif
+		#else // AA_FILTERED
+			// Unfiltered.
+			float4 Filtered = Neighbor4;
+			#if AA_ACCURATE_ALPHA
+				float FilteredWeight = NeighborWeight4;
+			#endif
+		#endif // AA_FILTERED
+		
+		float4 NeighborMin;
+		float4 NeighborMax;
+		NeighborMin.xyz = min(min(min(Neighbor1.xyz, Neighbor3.xyz), min(Neighbor4.xyz, Neighbor5.xyz)), Neighbor7.xyz);
+		NeighborMax.xyz = max(max(max(Neighbor1.xyz, Neighbor3.xyz), max(Neighbor4.xyz, Neighbor5.xyz)), Neighbor7.xyz);
+			
+		NeighborMin.w = Neighbor1.w;
+		NeighborMax.w = Neighbor1.w;
+		float NeighborMinAlphaWeight = NeighborWeight1;
+		float NeighborMaxAlphaWeight = NeighborWeight1;
+			
+		WeightTrackedAlphaClamping(
+			NeighborMin.a, NeighborMinAlphaWeight,
+			NeighborMax.a, NeighborMaxAlphaWeight,
+			Neighbor3.a, NeighborWeight3);
+			
+		WeightTrackedAlphaClamping(
+			NeighborMin.a, NeighborMinAlphaWeight,
+			NeighborMax.a, NeighborMaxAlphaWeight,
+			Neighbor4.a, NeighborWeight4);
+			
+		WeightTrackedAlphaClamping(
+			NeighborMin.a, NeighborMinAlphaWeight,
+			NeighborMax.a, NeighborMaxAlphaWeight,
+			Neighbor5.a, NeighborWeight5);
+			
+		WeightTrackedAlphaClamping(
+			NeighborMin.a, NeighborMinAlphaWeight,
+			NeighborMax.a, NeighborMaxAlphaWeight,
+			Neighbor7.a, NeighborWeight7);
+	#else
+		float4 Neighbor0 = PostprocessInput0.SampleLevel(PostprocessInput0Sampler, UV, 0, int2(-1, -1));
+		float4 Neighbor1 = PostprocessInput0.SampleLevel(PostprocessInput0Sampler, UV, 0, int2( 0, -1));
+		float4 Neighbor2 = PostprocessInput0.SampleLevel(PostprocessInput0Sampler, UV, 0, int2( 1, -1));
+		float4 Neighbor3 = PostprocessInput0.SampleLevel(PostprocessInput0Sampler, UV, 0, int2(-1,  0));
+		float4 Neighbor4 = PostprocessInput0.SampleLevel(PostprocessInput0Sampler, UV, 0);
+		float4 Neighbor5 = PostprocessInput0.SampleLevel(PostprocessInput0Sampler, UV, 0, int2( 1,  0));
+		float4 Neighbor6 = PostprocessInput0.SampleLevel(PostprocessInput0Sampler, UV, 0, int2(-1,  1));
+		float4 Neighbor7 = PostprocessInput0.SampleLevel(PostprocessInput0Sampler, UV, 0, int2( 0,  1));
+		float4 Neighbor8 = PostprocessInput0.SampleLevel(PostprocessInput0Sampler, UV, 0, int2( 1,  1));
+
+		Neighbor0.rgb *= HdrWeight4(Neighbor0.rgb, InExposureScale);
+		Neighbor1.rgb *= HdrWeight4(Neighbor1.rgb, InExposureScale);
+		Neighbor2.rgb *= HdrWeight4(Neighbor2.rgb, InExposureScale);
+		Neighbor3.rgb *= HdrWeight4(Neighbor3.rgb, InExposureScale);
+		Neighbor4.rgb *= HdrWeight4(Neighbor4.rgb, InExposureScale);
+		Neighbor5.rgb *= HdrWeight4(Neighbor5.rgb, InExposureScale);
+		Neighbor6.rgb *= HdrWeight4(Neighbor6.rgb, InExposureScale);
+		Neighbor7.rgb *= HdrWeight4(Neighbor7.rgb, InExposureScale);
+		Neighbor8.rgb *= HdrWeight4(Neighbor8.rgb, InExposureScale);
+
+		#if AA_FILTERED
+			float4 Filtered = 
+				Neighbor0 * SampleWeights[0] +
+				Neighbor1 * SampleWeights[1] +
+				Neighbor2 * SampleWeights[2] +
+				Neighbor3 * SampleWeights[3] +
+				Neighbor4 * SampleWeights[4] +
+				Neighbor5 * SampleWeights[5] +
+				Neighbor6 * SampleWeights[6] +
+				Neighbor7 * SampleWeights[7] +
+				Neighbor8 * SampleWeights[8];
+			#if AA_BORDER
+				// Use unfiltered for 1 pixel border.
+				float2 TestPos = abs(ScreenPos);
+				// Add 1 pixel and check if off screen.
+				TestPos += ViewportSize.zw * 2.0;
+				bool FilteredOffScreen = max(TestPos.x, TestPos.y) >= 1.0;
+				if(FilteredOffScreen) 
+				{
+					Filtered = Neighbor4;
+				}
+			#endif
+		#else
+			// Unfiltered.
+			float4 Filtered = Neighbor4;
+		#endif
+
+		#if AA_ROUND
+			float4 NeighborMin2 = min(min(Neighbor0, Neighbor2), min(Neighbor6, Neighbor8));		
+			float4 NeighborMax2 = max(max(Neighbor0, Neighbor2), max(Neighbor6, Neighbor8));		
+			float4 NeighborMin = min(min(min(Neighbor1, Neighbor3), min(Neighbor4, Neighbor5)), Neighbor7);		
+			float4 NeighborMax = max(max(max(Neighbor1, Neighbor3), max(Neighbor4, Neighbor5)), Neighbor7);		
+			NeighborMin2 = min(NeighborMin2, NeighborMin);
+			NeighborMax2 = max(NeighborMax2, NeighborMax);
+			NeighborMin = NeighborMin * 0.5 + NeighborMin2 * 0.5;
+			NeighborMax = NeighborMax * 0.5 + NeighborMax2 * 0.5;
+
+			float4 m1, m2;
+			m1  = Neighbor0; m2  = Neighbor0 * Neighbor0;
+			m1 += Neighbor1; m2 += Neighbor1 * Neighbor1;
+			m1 += Neighbor2; m2 += Neighbor2 * Neighbor2;
+			m1 += Neighbor3; m2 += Neighbor3 * Neighbor3;
+			m1 += Neighbor4; m2 += Neighbor4 * Neighbor4;
+			m1 += Neighbor5; m2 += Neighbor5 * Neighbor5;
+			m1 += Neighbor6; m2 += Neighbor6 * Neighbor6;
+			m1 += Neighbor7; m2 += Neighbor7 * Neighbor7;
+
+			float4 mu = m1 * (1.0 / 8.0);
+			float4 Sigma = sqrt( abs(m2 * (1.0 / 8.0) - mu * mu) );
+			NeighborMin = mu - 1.0 * Sigma;
+			NeighborMax = mu + 1.0 * Sigma;
+		#else
+			float4 NeighborMin = min(min(
+				min(min(Neighbor0, Neighbor1), min(Neighbor2, Neighbor3)),
+				min(min(Neighbor4, Neighbor5), min(Neighbor6, Neighbor7))), Neighbor8);
+			float4 NeighborMax = max(max(
+				max(max(Neighbor0, Neighbor1), max(Neighbor2, Neighbor3)),
+				max(max(Neighbor4, Neighbor5), max(Neighbor6, Neighbor7))), Neighbor8);
+		#endif
+	#endif
+
+	// FETCH HISTORY
+	// -------------
+	#if AA_BICUBIC
+		float4 HistoryColor = Texture2DSampleBicubic( PostprocessInput1, PostprocessInput1Sampler, BackN.xy, PostprocessInput1Size.xy );
+	#else
+		float4 HistoryColor = PostprocessInput1.SampleLevel(PostprocessInput1Sampler, BackN.xy, 0);
+	#endif
+
+	#if AA_YCOCG
+		HistoryColor.rgb = RGBToYCoCg(HistoryColor.rgb);
+		#if AA_TONE
+			float HistoryWeight = HdrWeightY(HistoryColor.x, InExposureScale);
+			HistoryColor.xyz *= HistoryWeight;
+		#endif
+	#else
+		HistoryColor.rgb *= HdrWeight4(HistoryColor.rgb, InExposureScale);
+	#endif
+
+	#if AA_DYNAMIC_ANTIGHOST && AA_DYNAMIC && !AA_ALPHA
+		bool Dynamic1 = PostprocessInput3.SampleLevel(PostprocessInput3Sampler, UV, 0, int2( 0, -1)).x > 0;
+		bool Dynamic3 = PostprocessInput3.SampleLevel(PostprocessInput3Sampler, UV, 0, int2(-1,  0)).x > 0;
+		bool Dynamic4 = PostprocessInput3.SampleLevel(PostprocessInput3Sampler, UV, 0).x > 0;
+		bool Dynamic5 = PostprocessInput3.SampleLevel(PostprocessInput3Sampler, UV, 0, int2( 1,  0)).x > 0;
+		bool Dynamic7 = PostprocessInput3.SampleLevel(PostprocessInput3Sampler, UV, 0, int2( 0,  1)).x > 0;
+
+		bool Dynamic = Dynamic1 || Dynamic3 || Dynamic4 || Dynamic5 || Dynamic7;
+		if( !Dynamic && HistoryColor.a > 0 )
+		{
+			HistoryColor = Filtered;
+			HistoryColor.a = 0;
+		}
+	#endif
+
+	// FIND LUMA OF CLAMPED HISTORY
+	// ----------------------------
+	// Save off luma of history before the clamp.
+	#if AA_YCOCG
+		float LumaMin = NeighborMin.x;
+		float LumaMax = NeighborMax.x;
+		float LumaHistory = HistoryColor.x;
+	#else
+		float LumaMin = Luma4(NeighborMin.rgb);
+		float LumaMax = Luma4(NeighborMax.rgb);
+		float LumaHistory = Luma4(HistoryColor.rgb);
+	#endif
+	float LumaContrast = LumaMax - LumaMin;
+	#if AA_CLAMP
+		#if AA_YCOCG
+			HistoryColor.rgb = clamp(HistoryColor.rgb, NeighborMin.rgb, NeighborMax.rgb);
+			#if AA_ALPHA
+				#if AA_ACCURATE_ALPHA
+					// Clamp history's alpha with sample weight tracking.
+					if (HistoryColor.a < NeighborMin.a)
+					{
+						HistoryColor.a = NeighborMin.a;
+						HistoryWeight = NeighborMinAlphaWeight;
+					}
+					else if (HistoryColor.a > NeighborMax.a)
+					{
+						HistoryColor.a = NeighborMax.a;
+						HistoryWeight = NeighborMaxAlphaWeight;
+					}
+				#else
+					HistoryColor.a = clamp(HistoryColor.a, NeighborMin.a, NeighborMax.a);
+				#endif
+			#endif
+		#else
+			#if AA_AABB
+				// Clamp history, this uses color AABB intersection for tighter fit.
+				float ClampBlend = HistoryClamp(HistoryColor.rgb, Filtered.rgb, NeighborMin.rgb, NeighborMax.rgb);
+				#if AA_ALPHA
+					HistoryColor.rgba = lerp(HistoryColor.rgba, Filtered.rgba, ClampBlend);
+				#else
+					HistoryColor.rgb = lerp(HistoryColor.rgb, Filtered.rgb, ClampBlend);
+				#endif
+			#else
+				HistoryColor = clamp(HistoryColor, NeighborMin, NeighborMax);
+			#endif
+		#endif
+	#endif
+
+	// ADD BACK IN ALIASING TO SHARPEN
+    // -------------------------------
+	#if AA_FILTERED && !AA_BICUBIC
+		// Blend in non-filtered based on the amount of sub-pixel motion.
+		float AddAliasing = saturate(HistoryBlur) * 0.5;
+		float LumaContrastFactor = 32.0;
+		#if AA_YCOCG
+			// 1/4 as bright.
+			LumaContrastFactor *= 4.0;
+		#endif
+		AddAliasing = saturate(AddAliasing + rcp(1.0 + LumaContrast * LumaContrastFactor));
+		Filtered.rgb = lerp(Filtered.rgb, Neighbor4.rgb, AddAliasing);
+	#endif
+	#if AA_YCOCG
+		float LumaFiltered = Filtered.x;
+	#else
+		float LumaFiltered = Luma4(Filtered.rgb);
+	#endif
+		
+	// COMPUTE BLEND AMOUNT 
+    // --------------------
+	float DistToClamp = min(abs(LumaMin-LumaHistory), abs(LumaMax-LumaHistory));
+	#if AA_BICUBIC
+		float HistoryFactor = CurrentFrameWeight * 3.125f * DistToClamp;
+	#else
+		float HistoryAmount = CurrentFrameWeight * 3.125f + HistoryBlur * (1.0/8.0);
+		float HistoryFactor = DistToClamp * HistoryAmount * (1.0 + HistoryBlur * HistoryAmount * 8.0);
+	#endif
+	float BlendFinal = saturate(HistoryFactor / (DistToClamp + LumaMax - LumaMin));
+	#if AA_TONE
+		BlendFinal = CurrentFrameWeight;
+	#endif
+	#if AA_NAN && (COMPILER_GLSL || COMPILER_METAL)
+		// The current Metal & GLSL compilers don't handle saturate(NaN) -> 0, instead they return NaN/INF.
+		BlendFinal = -min(-BlendFinal, 0.0);		
+	#endif
+
+	// Responsive forces 1/4 of new frame.
+	BlendFinal = bIsResponsiveAAPixel ? (1.0/4.0) : BlendFinal;
+
+	// Offscreen feedback resets.
+	#if AA_LERP 
+		float FixedLerp = 1.0/float(AA_LERP);
+	#endif
+	if(OffScreen) 
+	{
+		HistoryColor = Filtered;
+		#if AA_YCOCG && AA_ACCURATE_ALPHA
+			HistoryWeight = FilteredWeight;
+		#endif
+		#if !AA_ALPHA
+			HistoryColor.a = 0.0;
+		#endif
+		#if AA_LERP
+			FixedLerp = 1.0;
+		#endif
+	}
+	
+	// DO FINAL BLEND BETWEEN HISTORY AND FILTERED COLOR
+	// -------------------------------------------------
+	#if (AA_LERP == 0)
+		#if AA_ALPHA
+			#if AA_ACCURATE_ALPHA
+				HistoryColor.a *= HistoryWeight;
+				Filtered.a *= FilteredWeight;
+			#endif
+			OutColor = lerp(HistoryColor, Filtered, BlendFinal);
+			#if AA_YCOCG && AA_ACCURATE_ALPHA
+				float BlendWeight = lerp(HistoryWeight, FilteredWeight, BlendFinal);
+			#elif AA_FORCE_ALPHA_CLAMP
+				OutColor.a = clamp(HistoryColor.a, NeighborMin.a, NeighborMax.a);
+			#endif
+		#else
+			// Blend in linear to hit luma target.
+			OutColor.rgb = lerp(HistoryColor.rgb, Filtered.rgb, BlendFinal);
+
+			OutColor.a = bIsResponsiveAAPixel ? max(HistoryColor.a, 1.0/2.0) : OutColor.a;
+		#endif
+	#else
+		OutColor = lerp(HistoryColor, Filtered, FixedLerp);
+	#endif
+	#if AA_YCOCG
+		#if AA_TONE
+			OutColor.xyz *= HdrWeightInvY(OutColor.x, InExposureScale);
+		#endif
+		#if AA_ACCURATE_ALPHA
+			OutColor.a *= 1 / BlendWeight;
+			#if AA_FORCE_ALPHA_CLAMP
+				OutColor.a = clamp(OutColor.a, NeighborMin.a, NeighborMax.a);
+			#endif
+		#endif
+		OutColor.rgb = YCoCgToRGB(OutColor.rgb);
+	#else
+		// Convert back into linear.
+		OutColor.rgb *= HdrWeightInv4(OutColor.rgb, InExposureScale);
+	#endif
+	#if AA_NAN
+		// Transform NaNs to black, transform negative colors to black.
+		OutColor.rgb = -min(-OutColor.rgb, 0.0);
+	#endif
+	#if AA_DYNAMIC_ANTIGHOST && AA_DYNAMIC && !AA_ALPHA
+		OutColor.a = Dynamic4 ? 1 : 0;
+	#elif !AA_ALPHA
+		// Zero out to remove any prior computation of alpha
+		OutColor.a = 0;
+	#endif
+
+#undef AA_AABB
+#undef AA_ALPHA
+#undef AA_CROSS
+#undef AA_DYNAMIC
+#undef AA_FILTERED
+#undef AA_FILTERED_CONTRAST_LIMIT
+#undef AA_LERP
+#undef AA_ROUND
+#undef AA_NAN
+#undef AA_BORDER
+#undef AA_FORCE_ALPHA_CLAMP
+#undef AA_YCOCG
+#undef AA_BICUBIC
+#undef AA_TONE
+#undef AA_CLAMP
+#undef AA_DYNAMIC_ANTIGHOST
+#undef AA_ACCURATE_ALPHA
+#undef AA_SINGLE_PASS_RESPONSIVE
\ No newline at end of file
diff --git a/Engine/Shaders/Private/ReflectionEnvironmentComposite.ush b/Engine/Shaders/Private/ReflectionEnvironmentComposite.ush
new file mode 100644
--- /dev/null
+++ b/Engine/Shaders/Private/ReflectionEnvironmentComposite.ush
@@ -0,0 +1,214 @@
+// Copyright 1998-2017 Epic Games, Inc. All Rights Reserved.
+
+/*=============================================================================
+	ReflectionEnvironmentComposite
+=============================================================================*/
+
+#if FEATURE_LEVEL >= FEATURE_LEVEL_SM5
+	/** Prenormalized capture of the scene that's closest to the object being rendered, used for reflection environment on translucency. */
+	TextureCubeArray ReflectionCubemap;
+	SamplerState ReflectionCubemapSampler;
+	int SingleCubemapArrayIndex;	
+#else
+	TextureCube ReflectionCubemap;
+	SamplerState ReflectionCubemapSampler;
+#endif  
+
+float4 SingleCaptureOffsetAndAverageBrightness;
+float4 SingleCapturePositionAndRadius; 
+float SingleCaptureBrightness;
+
+float3 CompositeReflectionCapturesAndSkylight(
+	float CompositeAlpha, 
+	float3 WorldPosition, 
+	float3 RayDirection, 
+	float Roughness, 
+	float IndirectIrradiance, 
+	float IndirectSpecularOcclusion,
+	float3 ExtraIndirectSpecular,
+	uint NumCapturesAffectingTile,
+	uint CaptureDataStartIndex, 
+	uint EyeIndex)
+{
+	float Mip = ComputeReflectionCaptureMipFromRoughness(Roughness, View.ReflectionCubemapMaxMip);
+	float4 ImageBasedReflections = float4(0, 0, 0, CompositeAlpha);
+	float2 CompositedAverageBrightness = float2(0.0f, 1.0f);
+
+#if REFLECTION_COMPOSITE_USE_BLENDED_REFLECTION_CAPTURES
+	// Accumulate reflections from captures affecting this tile, applying largest captures first so that the smallest ones display on top
+	LOOP
+	for (uint TileCaptureIndex = 0; TileCaptureIndex < NumCapturesAffectingTile; TileCaptureIndex++) 
+	{
+		BRANCH
+		if (ImageBasedReflections.a < 0.001)
+		{
+			break;
+		}
+
+		uint CaptureIndex = 0;
+		BRANCH
+		if (EyeIndex == 0)
+		{
+			CaptureIndex = CompositeTileReflectionCaptureIndices[CaptureDataStartIndex + TileCaptureIndex];
+		}
+		else
+		{
+			CaptureIndex = InstancedCompositeTileReflectionCaptureIndices[CaptureDataStartIndex + TileCaptureIndex];
+		}
+			
+		float4 CapturePositionAndRadius = ReflectionCapture.PositionAndRadius[CaptureIndex];
+		float4 CaptureProperties = ReflectionCapture.CaptureProperties[CaptureIndex];
+
+		float3 CaptureVector = WorldPosition - CapturePositionAndRadius.xyz;
+		float CaptureVectorLength = sqrt(dot(CaptureVector, CaptureVector));		
+		float NormalizedDistanceToCapture = saturate(CaptureVectorLength / CapturePositionAndRadius.w);
+
+		BRANCH
+		if (CaptureVectorLength < CapturePositionAndRadius.w)
+		{
+			float3 ProjectedCaptureVector = RayDirection;
+			float4 CaptureOffsetAndAverageBrightness = ReflectionCapture.CaptureOffsetAndAverageBrightness[CaptureIndex];
+
+			// Fade out based on distance to capture
+			float DistanceAlpha = 0;
+			
+			#define PROJECT_ONTO_SHAPE 1
+			#if PROJECT_ONTO_SHAPE
+				#if REFLECTION_COMPOSITE_HAS_BOX_CAPTURES
+					#if REFLECTION_COMPOSITE_HAS_SPHERE_CAPTURES
+					// Box
+					BRANCH if (CaptureProperties.b > 0)
+					#endif
+					{
+						ProjectedCaptureVector = GetLookupVectorForBoxCapture(RayDirection, WorldPosition, CapturePositionAndRadius, ReflectionCapture.BoxTransform[CaptureIndex], ReflectionCapture.BoxScales[CaptureIndex], CaptureOffsetAndAverageBrightness.xyz, DistanceAlpha);
+					}
+				#endif
+
+				#if REFLECTION_COMPOSITE_HAS_SPHERE_CAPTURES
+					// Sphere
+					#if REFLECTION_COMPOSITE_HAS_BOX_CAPTURES
+					else
+					#endif
+					{
+						ProjectedCaptureVector = GetLookupVectorForSphereCapture(RayDirection, WorldPosition, CapturePositionAndRadius, NormalizedDistanceToCapture, CaptureOffsetAndAverageBrightness.xyz, DistanceAlpha);
+					}
+				#endif
+			#else 
+				DistanceAlpha = 1.0;
+			#endif //PROJECT_ONTO_SHAPE
+
+			float CaptureArrayIndex = CaptureProperties.g;
+
+			{
+				float4 Sample = ReflectionCubemap.SampleLevel(ReflectionCubemapSampler, float4(ProjectedCaptureVector, CaptureArrayIndex), Mip);
+
+				Sample.rgb *= CaptureProperties.r;
+				Sample *= DistanceAlpha;
+
+				// Under operator (back to front)
+				ImageBasedReflections.rgb += Sample.rgb * ImageBasedReflections.a * IndirectSpecularOcclusion;
+				ImageBasedReflections.a *= 1 - Sample.a;
+
+				float AverageBrightness = CaptureOffsetAndAverageBrightness.w;
+				CompositedAverageBrightness.x += AverageBrightness * DistanceAlpha * CompositedAverageBrightness.y;
+				CompositedAverageBrightness.y *= 1 - DistanceAlpha;
+			}
+		}
+	}
+#else
+
+	float3 ProjectedCaptureVector = RayDirection;
+
+	#define APPROXIMATE_CONTINUOUS_SINGLE_CAPTURE_PARALLAX 0
+	#if APPROXIMATE_CONTINUOUS_SINGLE_CAPTURE_PARALLAX
+		float3 CaptureVector = WorldPosition - SingleCapturePositionAndRadius.xyz;
+		float CaptureVectorLength = sqrt(dot(CaptureVector, CaptureVector));		
+		float NormalizedDistanceToCapture = saturate(CaptureVectorLength / SingleCapturePositionAndRadius.w);
+
+		float UnusedDistanceAlpha = 0;
+		ProjectedCaptureVector = GetLookupVectorForSphereCapture(RayDirection, WorldPosition, SingleCapturePositionAndRadius, NormalizedDistanceToCapture, SingleCaptureOffsetAndAverageBrightness.xyz, UnusedDistanceAlpha);
+		
+		float x = saturate(NormalizedDistanceToCapture);
+		float DistanceAlpha = 1 - x * x * (3 - 2 * x);
+		// Lerp between sphere parallax corrected and infinite based on distance to shape
+		ProjectedCaptureVector = lerp(RayDirection, normalize(ProjectedCaptureVector), DistanceAlpha);
+	#endif
+
+#if FEATURE_LEVEL >= FEATURE_LEVEL_SM5
+	float4 Sample = TextureCubeArraySampleLevel(ReflectionCubemap, ReflectionCubemapSampler, ProjectedCaptureVector, SingleCubemapArrayIndex, Mip);
+#else
+	float4 Sample = TextureCubeSampleLevel(ReflectionCubemap, ReflectionCubemapSampler, ProjectedCaptureVector, Mip);
+#endif
+	Sample.rgb *= SingleCaptureBrightness;
+	ImageBasedReflections = float4(Sample.rgb, 1 - Sample.a);
+
+	float AverageBrightness = SingleCaptureOffsetAndAverageBrightness.w;
+	CompositedAverageBrightness.x += AverageBrightness * CompositedAverageBrightness.y;
+	CompositedAverageBrightness.y = 0;
+#endif
+
+	// Apply indirect lighting scale while we have only accumulated reflection captures
+	ImageBasedReflections.rgb *= View.IndirectLightingColorScale;
+	CompositedAverageBrightness.x *= Luminance( View.IndirectLightingColorScale );
+
+#if ENABLE_SKY_LIGHT
+
+	BRANCH  
+	if (SkyLightParameters.y > 0)
+	{
+		float SkyAverageBrightness = 1.0f;
+
+		#if REFLECTION_COMPOSITE_SUPPORT_SKYLIGHT_BLEND
+			float3 SkyLighting = GetSkyLightReflectionSupportingBlend(RayDirection, Roughness, SkyAverageBrightness);
+		#else
+			float3 SkyLighting = GetSkyLightReflection(RayDirection, Roughness, SkyAverageBrightness);
+		#endif
+
+		// Normalize for static skylight types which mix with lightmaps
+		bool bNormalize = SkyLightParameters.z < 1 && ALLOW_STATIC_LIGHTING;
+
+		FLATTEN
+		if (bNormalize)
+		{
+			ImageBasedReflections.rgb += ImageBasedReflections.a * SkyLighting * IndirectSpecularOcclusion;
+			CompositedAverageBrightness.x += SkyAverageBrightness * CompositedAverageBrightness.y;
+		}
+		else
+		{
+			ExtraIndirectSpecular += SkyLighting * IndirectSpecularOcclusion;
+		}
+	}
+#endif
+
+#if ALLOW_STATIC_LIGHTING
+	ImageBasedReflections.rgb *= ComputeMixingWeight(IndirectIrradiance, CompositedAverageBrightness.x, Roughness);
+#endif
+
+	ImageBasedReflections.rgb += ImageBasedReflections.a * ExtraIndirectSpecular;
+
+	return ImageBasedReflections.rgb;
+}
+
+float3 CompositeReflectionCapturesAndSkylight(
+	float CompositeAlpha, 
+	float3 WorldPosition, 
+	float3 RayDirection, 
+	float Roughness, 
+	float IndirectIrradiance, 
+	float IndirectSpecularOcclusion,
+	float3 ExtraIndirectSpecular,
+	uint NumCapturesAffectingTile,
+	uint CaptureDataStartIndex)
+{
+	return CompositeReflectionCapturesAndSkylight(
+		CompositeAlpha, 
+		WorldPosition, 
+		RayDirection, 
+		Roughness, 
+		IndirectIrradiance, 
+		IndirectSpecularOcclusion, 
+		ExtraIndirectSpecular, 
+		NumCapturesAffectingTile, 
+		CaptureDataStartIndex, 
+		0);
+}
\ No newline at end of file
diff --git a/Engine/Shaders/Private/ReflectionEnvironmentShared.ush b/Engine/Shaders/Private/ReflectionEnvironmentShared.ush
new file mode 100644
--- /dev/null
+++ b/Engine/Shaders/Private/ReflectionEnvironmentShared.ush
@@ -0,0 +1,200 @@
+// Copyright 1998-2017 Epic Games, Inc. All Rights Reserved.
+
+/*=============================================================================
+	ReflectionEnvironmentShared
+=============================================================================*/
+  
+#define REFLECTION_CAPTURE_ROUGHEST_MIP 1
+#define REFLECTION_CAPTURE_ROUGHNESS_MIP_SCALE 1.2
+
+/** 
+ * Compute absolute mip for a reflection capture cubemap given a roughness.
+ */
+half ComputeReflectionCaptureMipFromRoughness(half Roughness, half CubemapMaxMip)
+{
+	// Heuristic that maps roughness to mip level
+	// This is done in a way such that a certain mip level will always have the same roughness, regardless of how many mips are in the texture
+	// Using more mips in the cubemap just allows sharper reflections to be supported
+	half LevelFrom1x1 = REFLECTION_CAPTURE_ROUGHEST_MIP - REFLECTION_CAPTURE_ROUGHNESS_MIP_SCALE * log2(Roughness);
+	return CubemapMaxMip - 1 - LevelFrom1x1;
+}
+
+float ComputeReflectionCaptureRoughnessFromMip(float Mip, half CubemapMaxMip)
+{
+	float LevelFrom1x1 = CubemapMaxMip - 1 - Mip;
+	return exp2( ( REFLECTION_CAPTURE_ROUGHEST_MIP - LevelFrom1x1 ) / REFLECTION_CAPTURE_ROUGHNESS_MIP_SCALE );
+}
+
+TextureCube SkyLightCubemap;
+SamplerState SkyLightCubemapSampler;
+
+/** X = max mip, Y = 1 if sky light should be rendered, 0 otherwise, Z = 1 if sky light is dynamic, 0 otherwise, W = blend fraction. */
+float4 SkyLightParameters;
+float SkyLightCubemapBrightness;
+
+float3 GetSkyLightReflection(float3 ReflectionVector, float Roughness, out float OutSkyAverageBrightness)
+{
+	float AbsoluteSpecularMip = ComputeReflectionCaptureMipFromRoughness(Roughness, SkyLightParameters.x);
+	float3 Reflection = TextureCubeSampleLevel(SkyLightCubemap, SkyLightCubemapSampler, ReflectionVector, AbsoluteSpecularMip).rgb;
+
+	OutSkyAverageBrightness = SkyLightCubemapBrightness * Luminance( View.SkyLightColor.rgb );
+	return Reflection * View.SkyLightColor.rgb;
+}
+
+TextureCube SkyLightBlendDestinationCubemap;
+SamplerState SkyLightBlendDestinationCubemapSampler;
+
+float3 GetSkyLightReflectionSupportingBlend(float3 ReflectionVector, float Roughness, out float OutSkyAverageBrightness)
+{
+	float3 Reflection = GetSkyLightReflection(ReflectionVector, Roughness, OutSkyAverageBrightness);
+
+	BRANCH
+	if (SkyLightParameters.w > 0)
+	{
+		float AbsoluteSpecularMip = ComputeReflectionCaptureMipFromRoughness(Roughness, SkyLightParameters.x);
+		float3 BlendDestinationReflection = TextureCubeSampleLevel(SkyLightBlendDestinationCubemap, SkyLightBlendDestinationCubemapSampler, ReflectionVector, AbsoluteSpecularMip).rgb;
+
+		Reflection = lerp(Reflection, BlendDestinationReflection * View.SkyLightColor.rgb, SkyLightParameters.w);
+	}
+
+	return Reflection;
+}
+
+/** 
+ * Computes sky diffuse lighting from the SH irradiance map.  
+ * This has the SH basis evaluation and diffuse convolusion weights combined for minimal ALU's - see "Stupid Spherical Harmonics (SH) Tricks" 
+ */
+float3 GetSkySHDiffuse(float3 Normal)
+{
+	float4 NormalVector = float4(Normal, 1);
+
+	float3 Intermediate0, Intermediate1, Intermediate2;
+	Intermediate0.x = dot(View.SkyIrradianceEnvironmentMap[0], NormalVector);
+	Intermediate0.y = dot(View.SkyIrradianceEnvironmentMap[1], NormalVector);
+	Intermediate0.z = dot(View.SkyIrradianceEnvironmentMap[2], NormalVector);
+
+	float4 vB = NormalVector.xyzz * NormalVector.yzzx;
+	Intermediate1.x = dot(View.SkyIrradianceEnvironmentMap[3], vB);
+	Intermediate1.y = dot(View.SkyIrradianceEnvironmentMap[4], vB);
+	Intermediate1.z = dot(View.SkyIrradianceEnvironmentMap[5], vB);
+
+	float vC = NormalVector.x * NormalVector.x - NormalVector.y * NormalVector.y;
+	Intermediate2 = View.SkyIrradianceEnvironmentMap[6].xyz * vC;
+
+	// max to not get negative colors
+	return max(0, Intermediate0 + Intermediate1 + Intermediate2);
+}
+
+/**
+* Computes sky diffuse lighting from the SH irradiance map.
+* This has the SH basis evaluation and diffuse convolusion weights combined for minimal ALU's - see "Stupid Spherical Harmonics (SH) Tricks"
+* Only does the first 3 components for speed.
+*/
+float3 GetSkySHDiffuseSimple(float3 Normal)
+{
+	float4 NormalVector = float4(Normal, 1);
+
+	float3 Intermediate0;
+	Intermediate0.x = dot(View.SkyIrradianceEnvironmentMap[0], NormalVector);
+	Intermediate0.y = dot(View.SkyIrradianceEnvironmentMap[1], NormalVector);
+	Intermediate0.z = dot(View.SkyIrradianceEnvironmentMap[2], NormalVector);		
+
+	// max to not get negative colors
+	return max(0, Intermediate0);
+}
+
+
+// Point lobe in off-specular peak direction
+float3 GetOffSpecularPeakReflectionDir(float3 Normal, float3 ReflectionVector, float Roughness)
+{
+	float a = Square(Roughness);
+	return lerp( Normal, ReflectionVector, (1 - a) * ( sqrt(1 - a) + a ) );	
+}
+
+float GetSpecularOcclusion(float NoV, float RoughnessSq, float AO)
+{
+	return saturate( pow( NoV + AO, RoughnessSq ) - 1 + AO );
+}
+
+float3 GetLookupVectorForBoxCapture(float3 ReflectionVector, float3 WorldPosition, float4 BoxCapturePositionAndRadius, float4x4 BoxTransform, float4 BoxScales, float3 LocalCaptureOffset, out float DistanceAlpha)
+{
+	// Transform the ray into the local space of the box, where it is an AABB with mins at -1 and maxs at 1
+	float3 LocalRayStart		= mul(float4(WorldPosition, 1), BoxTransform).xyz;
+	float3 LocalRayDirection	= mul(float4(ReflectionVector,  0), BoxTransform).xyz;
+
+	float3 InvRayDir = rcp(LocalRayDirection);
+	
+	//find the ray intersection with each of the 3 planes defined by the minimum extrema.
+	float3 FirstPlaneIntersections = -InvRayDir - LocalRayStart * InvRayDir;
+	//find the ray intersection with each of the 3 planes defined by the maximum extrema.
+	float3 SecondPlaneIntersections = InvRayDir - LocalRayStart * InvRayDir;
+	//get the furthest of these intersections along the ray
+	float3 FurthestPlaneIntersections = max(FirstPlaneIntersections, SecondPlaneIntersections);
+
+	//clamp the intersections to be between RayOrigin and RayEnd on the ray
+	float Intersection = min(FurthestPlaneIntersections.x, min(FurthestPlaneIntersections.y, FurthestPlaneIntersections.z));
+
+	// Compute the reprojected vector
+	float3 IntersectPosition = WorldPosition + Intersection * ReflectionVector;
+	float3 ProjectedCaptureVector = IntersectPosition - (BoxCapturePositionAndRadius.xyz + LocalCaptureOffset);
+
+	// Compute the distance from the receiving pixel to the box for masking
+	// Apply local to world scale to take scale into account without transforming back to world space
+	// Shrink the box by the transition distance (BoxScales.w) so that the fade happens inside the box influence area	
+	float BoxDistance = ComputeDistanceFromBoxToPoint(-(BoxScales.xyz - .5f * BoxScales.w), BoxScales.xyz - .5f * BoxScales.w, LocalRayStart * BoxScales.xyz);
+
+	// Setup a fade based on receiver distance to the box, hides the box influence shape
+	DistanceAlpha = 1.0 - smoothstep(0, .7f * BoxScales.w, BoxDistance);
+
+	return ProjectedCaptureVector;
+}
+
+float3 GetLookupVectorForSphereCapture(float3 ReflectionVector, float3 WorldPosition, float4 SphereCapturePositionAndRadius, float NormalizedDistanceToCapture, float3 LocalCaptureOffset, inout float DistanceAlpha)
+{
+	float3 ProjectedCaptureVector = ReflectionVector;
+	float ProjectionSphereRadius = SphereCapturePositionAndRadius.w;
+	float SphereRadiusSquared = ProjectionSphereRadius * ProjectionSphereRadius;
+
+	float3 LocalPosition = WorldPosition - SphereCapturePositionAndRadius.xyz;
+	float LocalPositionSqr = dot(LocalPosition, LocalPosition);
+
+	// Find the intersection between the ray along the reflection vector and the capture's sphere
+	float3 QuadraticCoef;
+	QuadraticCoef.x = 1;
+	QuadraticCoef.y = dot(ReflectionVector, LocalPosition);
+	QuadraticCoef.z = LocalPositionSqr - SphereRadiusSquared;
+
+	float Determinant = QuadraticCoef.y * QuadraticCoef.y - QuadraticCoef.z;
+
+	// Only continue if the ray intersects the sphere
+	FLATTEN
+	if (Determinant >= 0)
+	{
+		float FarIntersection = sqrt(Determinant) - QuadraticCoef.y;
+
+		float3 LocalIntersectionPosition = LocalPosition + FarIntersection * ReflectionVector;
+		ProjectedCaptureVector = LocalIntersectionPosition - LocalCaptureOffset;
+		// Note: some compilers don't handle smoothstep min > max (this was 1, .6)
+		//DistanceAlpha = 1.0 - smoothstep(.6, 1, NormalizedDistanceToCapture);
+
+		float x = saturate( 2.5 * NormalizedDistanceToCapture - 1.5 );
+		DistanceAlpha = 1 - x*x*(3 - 2*x);
+	}
+	return ProjectedCaptureVector;
+}
+
+float ComputeMixingWeight(float IndirectIrradiance, float AverageBrightness, float Roughness)
+{
+	// Mirror surfaces should have no mixing, so they match reflections from other sources (SSR, planar reflections)
+	float MixingAlpha = smoothstep(0, 1, saturate(Roughness * View.ReflectionEnvironmentRoughnessMixingScaleBiasAndLargestWeight.x + View.ReflectionEnvironmentRoughnessMixingScaleBiasAndLargestWeight.y));
+
+	// We have high frequency directional data but low frequency spatial data in the envmap.
+	// We have high frequency spatial data but low frequency directional data in the lightmap.
+	// So, we combine the two for the best of both. This is done by removing the low spatial frequencies from the envmap and replacing them with the lightmap data.
+	// This is only done with luma so as to not get odd color shifting.
+	float MixingWeight = IndirectIrradiance / max(AverageBrightness, .0001f);
+
+	MixingWeight = min(MixingWeight, View.ReflectionEnvironmentRoughnessMixingScaleBiasAndLargestWeight.z);
+
+	return lerp(1.0f, MixingWeight, MixingAlpha);
+}
diff --git a/Engine/Shaders/Private/ScreenSpaceRayCast.ush b/Engine/Shaders/Private/ScreenSpaceRayCast.ush
new file mode 100644
--- /dev/null
+++ b/Engine/Shaders/Private/ScreenSpaceRayCast.ush
@@ -0,0 +1,538 @@
+// Copyright 1998-2017 Epic Games, Inc. All Rights Reserved.
+
+#pragma once
+
+#include "Common.ush"	
+
+// 0:reference (slower, potentially higher quality) 1:use mips of the HZB depth (better performance)
+#define USE_HZB		1
+
+#define SCALAR_BRANCHLESS		0
+#define VECTORIZED_BRANCHLESS	0
+#define VECTORIZED_EARLY_OUT	1
+
+#if USE_HZB
+#define HZB_LEVEL_OFFSET 1.0
+#else
+#define HZB_LEVEL_OFFSET 0.0
+#endif
+
+float4 HZBUvFactorAndInvFactor;
+
+float4 SampleDepthTexture( Texture2D Texture, SamplerState Sampler, float Level, float4 SampleUV0, float4 SampleUV1 )
+{
+	float4 SampleDepth;
+#if USE_HZB
+	// SampleUV{0,1}.{xy,zw} should already be in the HZB UV frame using HZBUvFactorAndInvFactor.xy
+	SampleDepth.x = Texture2DSampleLevel( Texture, Sampler, SampleUV0.xy, Level ).r;
+	SampleDepth.y = Texture2DSampleLevel( Texture, Sampler, SampleUV0.zw, Level ).r;
+	SampleDepth.z = Texture2DSampleLevel( Texture, Sampler, SampleUV1.xy, Level ).r;
+	SampleDepth.w = Texture2DSampleLevel( Texture, Sampler, SampleUV1.zw, Level ).r;
+#else
+	SampleDepth.x = Texture2DSampleLevel( SceneDepthTexture, SceneDepthTextureSampler, SampleUV0.xy, 0 ).r;
+	SampleDepth.y = Texture2DSampleLevel( SceneDepthTexture, SceneDepthTextureSampler, SampleUV0.zw, 0 ).r;
+	SampleDepth.z = Texture2DSampleLevel( SceneDepthTexture, SceneDepthTextureSampler, SampleUV1.xy, 0 ).r;
+	SampleDepth.w = Texture2DSampleLevel( SceneDepthTexture, SceneDepthTextureSampler, SampleUV1.zw, 0 ).r;
+#endif
+	return SampleDepth;
+}
+
+void RayCast(
+	Texture2D Texture, SamplerState Sampler, float2 TextureSize,
+	float3 RayOriginTranslatedWorld, float3 RayDirection,
+	float Roughness, float ConeAngleWorld, float SceneDepth,
+	int NumSteps, float StepOffset,
+	out float4 OutHitUVzTime,
+	out float OutLevelHCB
+)
+{
+	// TODO provide RayStartUVz
+	const float4 RayStartV = mul( float4( RayOriginTranslatedWorld, 1 ), View.TranslatedWorldToView );
+	const float4 RayDirV = mul(float4(RayDirection * SceneDepth, 0), View.TranslatedWorldToView);
+	const float4 RayEndV = RayStartV + RayDirV;
+
+	const float RayEndRadiusV = length(RayDirV.xyz) * tan(ConeAngleWorld * 0.4);
+	const float2 RayEndBorderV = RayEndV.xy + RayEndRadiusV / sqrt(2.0);
+
+	const float4 RayStartClip	= mul( RayStartV, View.ViewToClip );
+	const float4 RayEndClip		= mul( RayEndV, View.ViewToClip );
+	const float2 RayEndBorderClip = mul(float4(RayEndBorderV, RayEndV.zw), View.ViewToClip).xy;
+
+	const float3 RayStartScreen = RayStartClip.xyz / RayStartClip.w;
+	const float3 RayEndScreen = RayEndClip.xyz / RayEndClip.w;
+	const float2 RayEndBorderScreen = RayEndBorderClip.xy / RayEndClip.w;
+
+	float4 RayDepthClip	= RayStartClip + mul( float4( 0, 0, SceneDepth, 0 ), View.ViewToClip );
+	float3 RayDepthScreen = RayDepthClip.xyz / RayDepthClip.w;
+
+	float3 RayStepScreen = RayEndScreen - RayStartScreen;
+	float2 RayStepRadiusScreen = RayEndBorderScreen - RayEndScreen.xy;
+
+	{
+		// Computes the scale down factor for RayStepScreen required to fit on the X and Y axis in order to clip it in the viewport
+		const float RayStepScreenInvFactor = 0.5 * length( RayStepScreen.xy );
+		const float2 AbsRayStepScreen = abs(RayStepScreen.xy);
+		const float2 S = (AbsRayStepScreen - max(abs(RayStepScreen.xy + RayStartScreen.xy * RayStepScreenInvFactor) - RayStepScreenInvFactor, 0.0f)) / AbsRayStepScreen;
+
+		// Rescales RayStepScreen accordingly
+		const float RayStepFactor = min(S.x, S.y) / RayStepScreenInvFactor;
+
+		RayStepScreen *= RayStepFactor;
+		RayStepRadiusScreen *= RayStepFactor;
+	}
+
+	const float2 HzbSize = float2(1024, 512);
+#if USE_HZB
+	float3 RayStartUVz = float3( (RayStartScreen.xy * float2( 0.5, -0.5 ) + 0.5) * HZBUvFactorAndInvFactor.xy, RayStartScreen.z );
+	float3 RayStepUVz = float3( RayStepScreen.xy * float2( 0.5, -0.5 ) * HZBUvFactorAndInvFactor.xy, RayStepScreen.z );
+	float2 RayStepRadiusUV = RayStepRadiusScreen.xy * float2( 0.5, -0.5 ) * HZBUvFactorAndInvFactor.xy;
+#else
+	float3 RayStartUVz = float3( (RayStartScreen.xy * float2( 0.5, -0.5 ) + 0.5), RayStartScreen.z );
+	float3 RayStepUVz = float3( RayStepScreen.xy * float2( 0.5, -0.5 ), RayStepScreen.z );
+	float2 RayStepRadiusUV = RayStepRadiusScreen.xy * float2( 0.5, -0.5 );
+#endif
+	const float RayStepRadiusFactor = length(RayStepRadiusUV * TextureSize);
+	
+	const float Step = 1.0 / NumSteps;
+
+	const float CompareTolerance = max(abs( RayStepUVz.z ), ( RayStartScreen.z - RayDepthScreen.z ) * 4) * Step;
+
+	// avoid bugs with early returns inside of loops on certain platform compilers.
+	float4 Result = float4( 0, 0, 0, 1 );
+
+#if SCALAR_BRANCHLESS
+
+	float MinHitTime = 1;
+	float LastDiff = 0;
+
+	float SampleTime = StepOffset * Step + Step;
+
+	UNROLL
+	for( int i = 0; i < NumSteps; i++ )
+	{
+		float3 SampleUVz = RayStartUVz + RayStepUVz * SampleTime;
+		
+		// Use lower res for farther samples
+		float Level = Roughness * (i * 4.0 / NumSteps) + HZB_LEVEL_OFFSET;
+		float SampleDepth = Texture.SampleLevel( Sampler, SampleUVz.xy, Level ).r;
+
+		float DepthDiff = SampleUVz.z - SampleDepth;
+		bool Hit = abs( DepthDiff + CompareTolerance ) < CompareTolerance;
+
+		// Find more accurate hit using line segment intersection
+		float TimeLerp = saturate( LastDiff / (LastDiff - DepthDiff) );
+		float IntersectTime = SampleTime + TimeLerp * Step - Step;
+		float HitTime = Hit ? IntersectTime : 1;
+		MinHitTime = min( MinHitTime, HitTime );
+
+		LastDiff = DepthDiff;
+
+		SampleTime += Step;
+	}
+
+	float3 HitUVz = RayStartUVz + RayStepUVz * MinHitTime;
+
+	Result = float4( HitUVz, MinHitTime );
+
+#elif VECTORIZED_BRANCHLESS
+
+	float MinHitTime = 1;
+	float LastDiff = 0;
+	float Level = HZB_LEVEL_OFFSET;
+	
+	// Vectorized to group fetches
+	float4 SampleTime = ( StepOffset + float4( 1, 2, 3, 4 ) ) * Step;
+	float4 SampleUV0 = RayStartUVz.xyxy + RayStepUVz.xyxy * SampleTime.xxyy;
+	float4 SampleUV1 = RayStartUVz.xyxy + RayStepUVz.xyxy * SampleTime.zzww;
+	float4 SampleZ   = RayStartUVz.zzzz + RayStepUVz.zzzz * SampleTime;
+
+	LOOP
+	for( int i = 0; i < NumSteps; i += 4 )
+	{
+		// Use lower res for farther samples
+		float4 SampleDepth = SampleDepthTexture( Texture, Sampler, Level, SampleUV0, SampleUV1 );
+
+		#if 1
+			float4 DepthDiff = SampleZ - SampleDepth;
+			float4 IntersectTime = ( SampleDepth - RayStartUVz.zzzz ) / RayStepUVz.zzzz;
+		
+			bool4 Hit = abs( DepthDiff + CompareTolerance ) < CompareTolerance;
+	
+			// If hit set to intersect time. If missed set to 1, beyond end of ray
+			float4 HitTime = Hit ? IntersectTime : 1;
+
+			// Take closest hit
+			HitTime.xy = min( HitTime.xy, HitTime.zw );
+			MinHitTime = min( HitTime.x, HitTime.y );
+		#else
+			// Line segment intersection
+			float4 DepthDiff1 = SampleZ - SampleDepth;
+			float4 DepthDiff0 = float4( LastDiff, DepthDiff1.xyz );
+			float4 TimeLerp = saturate( DepthDiff0 / (DepthDiff0 - DepthDiff1) );
+			float4 IntersectTime = SampleTime + (TimeLerp - 1) / (NumSteps + 1);
+		
+			bool4 Hit = abs( DepthDiff1 + CompareTolerance ) < CompareTolerance;
+	
+			// If hit set to intersect time. If missed set to 1, beyond end of ray
+			float4 HitTime = Hit ? IntersectTime : 1;
+
+			// Take closest hit
+			HitTime.xy = min( HitTime.xy, HitTime.zw );
+			MinHitTime = min( MinHitTime, min( HitTime.x, HitTime.y ) );
+
+			LastDiff = DepthDiff1.w;
+		#endif
+		Level += Roughness * (16.0 / NumSteps);
+
+		SampleTime	+= 4.0 / (NumSteps + 1);
+		SampleUV0	+= RayStepUVz.xyxy * 4.0 / (NumSteps + 1);
+		SampleUV1	+= RayStepUVz.xyxy * 4.0 / (NumSteps + 1);
+		SampleZ		+= RayStepUVz.zzzz * 4.0 / (NumSteps + 1);
+	}
+
+	float3 HitUVz = RayStartUVz + RayStepUVz * MinHitTime;
+
+	Result = float4( HitUVz, MinHitTime );
+
+#elif 1
+
+	float LastDiff = 0;
+	float Level = HZB_LEVEL_OFFSET;
+
+	RayStepUVz *= Step;
+	float3 RayUVz = RayStartUVz + RayStepUVz * StepOffset;
+
+	LOOP
+	for( int i = 0; i < NumSteps; i += 4 )
+	{
+		// Vectorized to group fetches
+		float4 SampleUV0 = RayUVz.xyxy + RayStepUVz.xyxy * float4( 1, 1, 2, 2 );
+		float4 SampleUV1 = RayUVz.xyxy + RayStepUVz.xyxy * float4( 3, 3, 4, 4 );
+		float4 SampleZ   = RayUVz.zzzz + RayStepUVz.zzzz * float4( 1, 2, 3, 4 );
+		
+		// Use lower res for farther samples
+		float4 SampleDepth = SampleDepthTexture( Texture, Sampler, Level, SampleUV0, SampleUV1 );
+
+		float4 DepthDiff = SampleZ - SampleDepth;
+		bool4 Hit = abs( -DepthDiff - CompareTolerance ) < CompareTolerance;
+
+		BRANCH if( any( Hit ) )
+		{
+			float DepthDiff0 = DepthDiff[2];
+			float DepthDiff1 = DepthDiff[3];
+			float MinTime = 3;
+
+			FLATTEN if( Hit[2] )
+			{
+				DepthDiff0 = DepthDiff[1];
+				DepthDiff1 = DepthDiff[2];
+				MinTime = 2;
+			}
+			FLATTEN if( Hit[1] )
+			{
+				DepthDiff0 = DepthDiff[0];
+				DepthDiff1 = DepthDiff[1];
+				MinTime = 1;
+			}
+			FLATTEN if( Hit[0] )
+			{
+				DepthDiff0 = LastDiff;
+				DepthDiff1 = DepthDiff[0];
+				MinTime = 0;
+			}
+
+			// Find more accurate hit using line segment intersection
+			float TimeLerp = saturate( DepthDiff0 / (DepthDiff0 - DepthDiff1) );
+			float IntersectTime = MinTime + TimeLerp;
+			float3 HitUVz = RayUVz + RayStepUVz * IntersectTime;
+
+			Result = float4( HitUVz, 0.5 );
+			break;
+		}
+
+		LastDiff = DepthDiff.w;
+		Level += Roughness * (16.0 / NumSteps);
+
+		RayUVz += 4 * RayStepUVz;
+	}
+
+#elif 1
+
+	float Level = HZB_LEVEL_OFFSET;
+
+	RayStepUVz *= Step;
+	float3 RayUVz = RayStartUVz + RayStepUVz * StepOffset;
+
+	LOOP
+	for( uint i = 0; i < NumSteps; i += 4 )
+	{
+		// Vectorized to group fetches
+		float4 SampleUV0 = RayUVz.xyxy + RayStepUVz.xyxy * float4( 1, 1, 2, 2 );
+		float4 SampleUV1 = RayUVz.xyxy + RayStepUVz.xyxy * float4( 3, 3, 4, 4 );
+		float4 SampleZ   = RayUVz.zzzz + RayStepUVz.zzzz * float4( 1, 2, 3, 4 );
+		
+		// Use lower res for farther samples
+		float4 SampleDepth = SampleDepthTexture( Texture, Sampler, Level, SampleUV0, SampleUV1 );
+
+		float4 DepthDiff = SampleZ - SampleDepth;
+		bool4 Hit = abs( -DepthDiff - CompareTolerance ) < CompareTolerance;
+
+		BRANCH if( any( Hit ) )
+		{
+			float MinTime = 3;
+			MinTime = Hit[2] ? 2 : MinTime;
+			MinTime = Hit[1] ? 1 : MinTime;
+			MinTime = Hit[0] ? 0 : MinTime;
+
+			float3 HitUVz = RayUVz + RayStepUVz * MinTime;
+
+			Result = float4( HitUVz, 0.5 );
+			break;
+		}
+
+		Level += Roughness * (16.0 / NumSteps);
+		RayUVz += 4 * RayStepUVz;
+	}
+
+#else // VECTORIZED_EARLY_OUT
+
+	float LastDiff = 0;
+	float Level = HZB_LEVEL_OFFSET;
+
+	float4 SampleTime = ( StepOffset + float4( 1, 2, 3, 4 ) ) * Step;
+
+	LOOP
+	for( int i = 0; i < NumSteps; i += 4 )
+	{
+		// Vectorized to group fetches
+		float4 SampleUV0 = RayStartUVz.xyxy + RayStepUVz.xyxy * SampleTime.xxyy;
+		float4 SampleUV1 = RayStartUVz.xyxy + RayStepUVz.xyxy * SampleTime.zzww;
+		float4 SampleZ   = RayStartUVz.zzzz + RayStepUVz.zzzz * SampleTime;
+		
+		// Use lower res for farther samples
+		float4 SampleDepth = SampleDepthTexture( Texture, Sampler, Level, SampleUV0, SampleUV1 );
+
+		float4 DepthDiff1 = SampleZ - SampleDepth;
+		bool4 Hit = abs( -DepthDiff1 - CompareTolerance ) < CompareTolerance;
+
+		BRANCH if( any( Hit ) )
+		{
+			// Find more accurate hit using line segment intersection
+			float4 DepthDiff0 = float4( LastDiff, DepthDiff1.xyz );
+			float4 TimeLerp = saturate( DepthDiff0 / (DepthDiff0 - DepthDiff1) );
+			float4 IntersectTime = SampleTime + (TimeLerp - 1) / (NumSteps + 1);
+			float4 HitTime = Hit ? IntersectTime : 1;
+
+			// Take closest hit
+			HitTime.xy = min( HitTime.xy, HitTime.zw );
+			float MinHitTime = min( HitTime.x, HitTime.y );
+			
+			float3 HitUVz = RayStartUVz + RayStepUVz * MinHitTime;
+
+			Result = float4( HitUVz, MinHitTime );
+			break;
+		}
+
+		LastDiff = DepthDiff1.w;
+		Level += Roughness * (16.0 / NumSteps);
+
+		SampleTime += 4.0 / (NumSteps + 1);
+	}
+
+#endif
+	
+#if USE_HZB
+	Result.xy *= HZBUvFactorAndInvFactor.zw;
+	Result.xy = Result.xy * float2( 2, -2 ) + float2( -1, 1 );
+	Result.xy = Result.xy * View.ScreenPositionScaleBias.xy + View.ScreenPositionScaleBias.wz;
+#endif
+
+	OutHitUVzTime = Result;
+	OutLevelHCB = log2(OutHitUVzTime.w * RayStepRadiusFactor) + 1;
+}
+
+float4 RayCastCheap( Texture2D Texture, SamplerState Sampler, float3 R, float SceneDepth, float3 PositionTranslatedWorld, float StepOffset )
+{
+	// TODO provide RayStartUVz
+
+	// NOTE could clip ray against frustum planes
+
+	// TODO use screen position and skip matrix mul
+	float4 RayStartClip	= mul( float4( PositionTranslatedWorld, 1 ), View.TranslatedWorldToClip );
+	float4 RayEndClip	= mul( float4( PositionTranslatedWorld + R * SceneDepth, 1 ), View.TranslatedWorldToClip );
+
+	float3 RayStartScreen = RayStartClip.xyz / RayStartClip.w;
+	float3 RayEndScreen = RayEndClip.xyz / RayEndClip.w;
+
+	// Normalize 2D length
+	float3 RayStepScreen = ( RayEndScreen - RayStartScreen ) / length( RayEndScreen.xy - RayStartScreen.xy );
+	RayStepScreen *= 0.75;
+
+#if USE_HZB
+	float3 RayStartUVz = float3( RayStartScreen.xy * float2( 0.5, -0.5 ) + 0.5, RayStartScreen.z );
+	float3 RayStepUVz = float3( RayStepScreen.xy * float2( 0.5, -0.5 ), RayStepScreen.z );
+#else
+	float3 RayStartUVz = float3( RayStartScreen.xy * View.ScreenPositionScaleBias.xy + View.ScreenPositionScaleBias.wz, RayStartScreen.z );
+	float3 RayStepUVz = float3( RayStepScreen.xy * View.ScreenPositionScaleBias.xy, RayStepScreen.z );
+#endif
+	
+	const float Step = 1.0 / (4 + 1);
+
+	// *2 to get less morie pattern in extreme cases, larger values make object appear not grounded in reflections
+	const float CompareTolerance = abs( RayStepUVz.z ) * Step * 2;
+
+	// Vectorized to group fetches
+	float4 SampleTime = ( StepOffset + float4( 1, 2, 3, 4 ) ) * Step;
+	float4 SampleUV0 = RayStartUVz.xyxy + RayStepUVz.xyxy * SampleTime.xxyy;
+	float4 SampleUV1 = RayStartUVz.xyxy + RayStepUVz.xyxy * SampleTime.zzww;
+	float4 SampleZ   = RayStartUVz.zzzz + RayStepUVz.zzzz * SampleTime;
+
+	float4 SampleDepth;
+#if USE_HZB
+	SampleDepth.x = Texture2DSampleLevel( Texture, Sampler, SampleUV0.xy, 0 ).r;
+	SampleDepth.y = Texture2DSampleLevel( Texture, Sampler, SampleUV0.zw, 0 ).r;
+	SampleDepth.z = Texture2DSampleLevel( Texture, Sampler, SampleUV1.xy, 1 ).r;
+	SampleDepth.w = Texture2DSampleLevel( Texture, Sampler, SampleUV1.zw, 1 ).r;
+#else
+	SampleDepth.x = Texture2DSampleLevel( SceneDepthTexture, SceneDepthTextureSampler, SampleUV0.xy, 0 ).r;
+	SampleDepth.y = Texture2DSampleLevel( SceneDepthTexture, SceneDepthTextureSampler, SampleUV0.zw, 0 ).r;
+	SampleDepth.z = Texture2DSampleLevel( SceneDepthTexture, SceneDepthTextureSampler, SampleUV1.xy, 0 ).r;
+	SampleDepth.w = Texture2DSampleLevel( SceneDepthTexture, SceneDepthTextureSampler, SampleUV1.zw, 0 ).r;
+#endif
+
+	#if 1
+		float4 DepthDiff = SampleZ - SampleDepth;
+		float4 IntersectTime = ( SampleDepth - RayStartUVz.zzzz ) / RayStepUVz.zzzz;
+		
+		bool4 Hit = abs( DepthDiff + CompareTolerance ) < CompareTolerance;
+	
+		// If hit set to intersect time. If missed set to 1, beyond end of ray
+		float4 HitTime = Hit ? IntersectTime : 1;
+
+		// Take closest hit
+		HitTime.xy = min( HitTime.xy, HitTime.zw );
+		float MinHitTime = min( HitTime.x, HitTime.y );
+	#else
+		// Line segment intersection
+		float4 DepthDiff1 = SampleZ - SampleDepth;
+		float4 DepthDiff0 = float4( 0, DepthDiff1.xyz );
+		float4 TimeLerp = saturate( DepthDiff0 / (DepthDiff0 - DepthDiff1) );
+		float4 IntersectTime = SampleTime + (TimeLerp - 1) * Step;
+		
+		bool4 Hit = abs( DepthDiff1 + CompareTolerance ) < CompareTolerance;
+	
+		// If hit set to intersect time. If missed set to 1, beyond end of ray
+		float4 HitTime = Hit ? IntersectTime : 1;
+
+		// Take closest hit
+		HitTime.xy = min( HitTime.xy, HitTime.zw );
+		float MinHitTime = min( HitTime.x, HitTime.y );
+	#endif
+
+	float3 HitUVz = RayStartUVz + RayStepUVz * MinHitTime;
+
+	BRANCH
+	if( MinHitTime == 1 && RayStepUVz.z < 0 )
+	{
+		float4 RayInfClip	= mul( float4( PositionTranslatedWorld + R * 2000000, 1 ), View.TranslatedWorldToClip );
+		float3 RayInfScreen = RayInfClip.xyz / RayInfClip.w;
+		float3 RayInfUVz = float3( RayInfScreen.xy * float2( 0.5, -0.5 ) + 0.5, RayInfScreen.z );
+
+		HitUVz = RayInfUVz;
+		MinHitTime = 0.5;
+	}
+
+#if USE_HZB
+	HitUVz.xy = HitUVz.xy = HitUVz.xy * float2( 2, -2 ) + float2( -1, 1 );
+	HitUVz.xy = HitUVz.xy * View.ScreenPositionScaleBias.xy + View.ScreenPositionScaleBias.wz;
+#endif
+	return float4( HitUVz, MinHitTime );
+}
+
+float2 ReprojectHit( float3 HitUVz )
+{
+	// Camera motion for pixel (in ScreenPos space).
+	float4 ThisClip = float4( (HitUVz.xy - View.ScreenPositionScaleBias.wz) / View.ScreenPositionScaleBias.xy, HitUVz.z, 1 );
+	float4 PrevClip = mul( ThisClip, View.ClipToPrevClip );
+	float2 PrevScreen = PrevClip.xy / PrevClip.w;
+	float2 PrevUV = PrevScreen.xy * View.ScreenPositionScaleBias.xy + View.ScreenPositionScaleBias.wz;
+	return PrevUV;
+}
+
+float2 ReprojectHit( Texture2D Texture, SamplerState Sampler, float3 HitUVz )
+{
+	// Camera motion for pixel (in ScreenPos space).
+	float4 ThisClip = float4( (HitUVz.xy - View.ScreenPositionScaleBias.wz) / View.ScreenPositionScaleBias.xy, HitUVz.z, 1 );
+	float4 PrevClip = mul( ThisClip, View.ClipToPrevClip );
+	float2 PrevScreen = PrevClip.xy / PrevClip.w;
+
+	float2 Velocity = Texture.SampleLevel( Sampler, HitUVz.xy, 0 ).xy;
+	if( Velocity.x > 0.0 )
+	{
+		PrevScreen = ThisClip.xy - DecodeVelocityFromTexture( Velocity );
+	}
+
+	float2 PrevUV = PrevScreen.xy * View.ScreenPositionScaleBias.xy + View.ScreenPositionScaleBias.wz;
+	return PrevUV;
+}
+
+float4 SampleScreenColor( Texture2D Texture, SamplerState Sampler, float2 UV )
+{
+	float4 OutColor;
+
+	OutColor.rgb = Texture.SampleLevel( Sampler, UV, 0 ).rgb;
+	OutColor.a = 1;
+
+	// Off screen masking
+	float2 ScreenPos = ( UV - View.ScreenPositionScaleBias.wz ) / View.ScreenPositionScaleBias.xy;
+	float2 Vignette = saturate( abs( ScreenPos ) * 5 - 4 );
+		
+	//PrevScreen sometimes has NaNs or Infs.  DX11 is protected because saturate turns NaNs -> 0.
+	//Do a SafeSaturate so other platforms get the same protection.
+	OutColor *= SafeSaturate( 1.0 - dot( Vignette, Vignette ) );
+
+	// Transform NaNs to black, transform negative colors to black.
+	OutColor.rgb = -min(-OutColor.rgb, 0.0);
+
+	return OutColor;
+}
+
+float4 SampleHCBLevel( Texture2D Texture, SamplerState Sampler, float3 HitUVz, float Level )
+{
+	float4 OutColor;
+
+#if PREV_FRAME_COLOR
+	// Find previous screen position for hit since color buffer is from last frame
+	// TODO combine to single matrix
+	float4 HitClip = float4( (HitUVz.xy - View.ScreenPositionScaleBias.wz) / View.ScreenPositionScaleBias.xy, HitUVz.z, 1 );
+	float4 HitTranslatedWorld = mul( HitClip, View.ClipToTranslatedWorld );
+	HitTranslatedWorld /= HitTranslatedWorld.w;
+		
+	float3 PrevTranslatedWorld = HitTranslatedWorld.xyz + (View.PrevPreViewTranslation - View.PreViewTranslation);
+	float4 PrevClip = mul( float4( PrevTranslatedWorld, 1 ), View.PrevTranslatedWorldToClip );
+	float2 PrevScreen = PrevClip.xy / PrevClip.w;
+	float2 PrevUV = PrevScreen.xy * View.ScreenPositionScaleBias.xy + View.ScreenPositionScaleBias.wz;
+
+	OutColor.rgb = Texture.SampleLevel( Sampler, PrevUV * HZBUvFactorAndInvFactor.xy, Level ).rgb;
+	OutColor.a = 1;
+
+	// Off screen masking
+	float2 Vignette = saturate( abs( PrevScreen ) * 5 - 4 );
+#else
+	OutColor.rgb = Texture.SampleLevel( Sampler, HitUVz.xy * HZBUvFactorAndInvFactor.xy, Level ).rgb;
+	OutColor.a = 1;
+
+	// Off screen masking
+	float2 HitScreenPos = ( HitUVz.xy - View.ScreenPositionScaleBias.wz ) / View.ScreenPositionScaleBias.xy;
+	float2 Vignette = saturate( abs( HitScreenPos ) * 5 - 4 );
+#endif
+		
+	//PrevScreen sometimes has NaNs or Infs.  DX11 is protected because saturate turns NaNs -> 0.
+	//Do a SafeSaturate so other platforms get the same protection.
+	OutColor *= SafeSaturate( 1.0 - dot( Vignette, Vignette ) );
+
+	// Transform NaNs to black, transform negative colors to black.
+	OutColor.rgb = -min(-OutColor.rgb, 0.0);
+
+	return OutColor;
+}
\ No newline at end of file
diff --git a/Engine/Shaders/Private/SeparateTranslucency.ush b/Engine/Shaders/Private/SeparateTranslucency.ush
new file mode 100644
--- /dev/null
+++ b/Engine/Shaders/Private/SeparateTranslucency.ush
@@ -0,0 +1,92 @@
+// Copyright 1998-2017 Epic Games, Inc. All Rights Reserved.
+
+/*=============================================================================
+	SeparateTranslucency.usf: 
+=============================================================================*/
+
+Texture2D<float> LowResDepthTexture;
+SamplerState BilinearClampedSampler;
+SamplerState PointClampedSampler;
+
+void UpdateNearestSample(float Z, float2 UV, float FullResZ, inout float MinDist, inout float2 NearestUV)
+{
+    float DepthDelta = abs(Z - FullResZ);
+
+	FLATTEN
+    if (DepthDelta < MinDist)
+    {
+        MinDist = DepthDelta;
+        NearestUV = UV;
+    }
+}
+
+float4 BilinearUpsampling(float2 UV, Texture2D LowResTex)
+{
+	return Texture2DSample(LowResTex, BilinearClampedSampler, UV);
+}
+
+float4 NearestDepthNeighborUpsampling(float2 Position, float2 UV, Texture2D LowResTex, float2 LowResTexelSize)
+{
+//@todo - support for all platforms, just skip the GatherRed optimization where not supported
+#if (SM5_PROFILE && !METAL_SM5_PROFILE || PS4_PROFILE)
+
+	// The relative depth comparison breaks down at larger distances and incorrectly causes point sampling on the skybox pixels
+	float MaxOperationDepth = 2000000.0f;
+
+	// Note: this upsample is specialized for half res to full res
+    float4 LowResDepthBuffer = LowResDepthTexture.GatherRed(BilinearClampedSampler, UV);
+    float4 LowResDepth = min(float4(ConvertFromDeviceZ(LowResDepthBuffer.x), ConvertFromDeviceZ(LowResDepthBuffer.y), ConvertFromDeviceZ(LowResDepthBuffer.z), ConvertFromDeviceZ(LowResDepthBuffer.w)), MaxOperationDepth.xxxx);
+	float FullResDepth = min(ConvertFromDeviceZ(SceneDepthTexture[uint2(Position.xy)].x), MaxOperationDepth);
+
+	float RelativeDepthThreshold = .1f;
+
+	// Search for the UV of the low res neighbor whose depth is closest to the full res depth
+    float MinDist = 1.e8f;
+
+    float2 UV00 = UV - 0.5f * LowResTexelSize;
+    float2 NearestUV = UV00;
+    UpdateNearestSample(LowResDepth.w, UV00, FullResDepth, MinDist, NearestUV);
+
+    float2 UV10 = float2(UV00.x + LowResTexelSize.x, UV00.y);
+    UpdateNearestSample(LowResDepth.z, UV10, FullResDepth, MinDist, NearestUV);
+
+    float2 UV01 = float2(UV00.x, UV00.y + LowResTexelSize.y);
+    UpdateNearestSample(LowResDepth.x, UV01, FullResDepth, MinDist, NearestUV);
+
+    float2 UV11 = float2(UV00.x + LowResTexelSize.x, UV00.y + LowResTexelSize.y);
+    UpdateNearestSample(LowResDepth.y, UV11, FullResDepth, MinDist, NearestUV);
+	 
+    float4 Output = 0.0f;
+	float InvFullResDepth = 1.0f / FullResDepth;
+
+    BRANCH
+    if (abs(LowResDepth.w - FullResDepth) * InvFullResDepth < RelativeDepthThreshold 
+		&& abs(LowResDepth.z - FullResDepth) * InvFullResDepth < RelativeDepthThreshold 
+		&& abs(LowResDepth.x - FullResDepth) * InvFullResDepth < RelativeDepthThreshold
+		&& abs(LowResDepth.y - FullResDepth) * InvFullResDepth < RelativeDepthThreshold)
+    {
+        Output = Texture2DSampleLevel(LowResTex, BilinearClampedSampler, UV, 0);
+    }
+    else
+    {
+        Output = Texture2DSampleLevel(LowResTex, PointClampedSampler, NearestUV, 0);
+    }
+	
+	return Output;
+
+#else
+
+	return Texture2DSample(LowResTex, BilinearClampedSampler, UV);
+
+#endif
+}
+
+
+float4 UpsampleSeparateTranslucency(float2 Position, float2 UV, Texture2D LowResTex, float2 LowResTexelSize)
+{
+#if NEAREST_DEPTH_NEIGHBOR_UPSAMPLE
+	return NearestDepthNeighborUpsampling(Position, UV, LowResTex, LowResTexelSize);
+#else
+	return BilinearUpsampling(UV, LowResTex);
+#endif
+}
\ No newline at end of file
diff --git a/Engine/Shaders/Private/ShadingModels.ush b/Engine/Shaders/Private/ShadingModels.ush
new file mode 100644
--- /dev/null
+++ b/Engine/Shaders/Private/ShadingModels.ush
@@ -0,0 +1,755 @@
+// Copyright 1998-2017 Epic Games, Inc. All Rights Reserved.
+
+#pragma once
+
+#include "DeferredShadingCommon.ush"
+#include "BRDF.ush"
+#include "FastMath.ush"
+
+#if 0
+void StandardShadingShared( float3 DiffuseColor, float3 SpecularColor, float Roughness, float3 V, half3 N )
+{
+	float NoV = saturate( abs( dot(N, V) ) + 1e-5 );
+
+	// Diffuse_Lambert
+	Shared.DiffuseMul = DiffuseColor * (1.0 / PI);
+
+	// D_GGX, Vis_SmithJointApprox
+	float m = Roughness * Roughness;
+	Shared.m2 = m * m;
+	Shared.SpecularMul = (0.5 / PI) * Shared.m2;
+	Shared.VisMad = float2( 2 * NoV * ( 1 - m ) + m, NoV * m );
+	
+	// F_Schlick
+	Shared.SpecularMul *= saturate( 50.0 * SpecularColor.g );
+}
+
+void StandardShadingPerLight( Shared, float3 L, float3 V, half3 N )
+{
+	float3 H = normalize(V + L);			// 3 add, 2 mad, 4 mul, 1 rsqrt
+	float NoL = saturate( dot(N, L) );		// 2 mad, 1 mul
+	float NoH = saturate( dot(N, H) );		// 2 mad, 1 mul
+	float VoH = saturate( dot(V, H) );		// 2 mad, 1 mul
+
+	// D_GGX, Vis_SmithJointApprox
+	float d = ( NoH * Shared.m2 - NoH ) * NoH + 1;			// 2 mad
+	float v = NoL * Shared.VisMad.x + Shared.VisMad.y;		// 1 mad
+	float D_Vis = Shared.SpecularMul * rcp( d * d * v );	// 3 mul, 1 rcp
+
+	// F_Schlick
+	float Fc = pow( 1 - VoH, 5 );							// 1 sub, 3 mul
+	float3 F = Fc + (1 - Fc) * SpecularColor;				// 1 sub, 3 mad
+
+	return Shared.DiffuseMul + D_Vis * F;					// 3 mad
+}
+#endif
+
+float3 StandardShading( float3 DiffuseColor, float3 SpecularColor, float3 LobeRoughness, float3 LobeEnergy, float3 L, float3 V, half3 N )
+{
+	float NoL = dot(N, L);
+	float NoV = dot(N, V);
+	float LoV = dot(L, V);
+	float InvLenH = rsqrt( 2 + 2 * LoV );
+	float NoH = saturate( ( NoL + NoV ) * InvLenH );
+	float VoH = saturate( InvLenH + InvLenH * LoV );
+	NoL = saturate(NoL);
+	NoV = saturate(abs(NoV) + 1e-5);
+
+	// Generalized microfacet specular
+	float D = D_GGX( LobeRoughness[1], NoH ) * LobeEnergy[1];
+	float Vis = Vis_SmithJointApprox( LobeRoughness[1], NoV, NoL );
+	float3 F = F_Schlick( SpecularColor, VoH );
+
+	float3 Diffuse = Diffuse_Lambert( DiffuseColor );
+	//float3 Diffuse = Diffuse_Burley( DiffuseColor, LobeRoughness[1], NoV, NoL, VoH );
+	//float3 Diffuse = Diffuse_OrenNayar( DiffuseColor, LobeRoughness[1], NoV, NoL, VoH );
+
+	return Diffuse * LobeEnergy[2] + (D * Vis) * F;
+}
+
+float3 SimpleShading( float3 DiffuseColor, float3 SpecularColor, float Roughness, float3 L, float3 V, half3 N )
+{
+	float3 H = normalize(V + L);
+	float NoH = saturate( dot(N, H) );
+	
+	// Generalized microfacet specular
+	float D = D_GGX( Roughness, NoH );
+	float Vis = Vis_Implicit();
+	float3 F = F_None( SpecularColor );
+
+	return Diffuse_Lambert( DiffuseColor ) + (D * Vis) * F;
+}
+
+float3 ClearCoatShading( FGBufferData GBuffer, float3 LobeRoughness, float3 LobeEnergy, float3 L, float3 V, half3 N )
+{
+	const float ClearCoat			= GBuffer.CustomData.x;
+	const float ClearCoatRoughness	= GBuffer.CustomData.y;
+	const float Film = 1 * ClearCoat;
+	const float MetalSpec = 0.9;
+
+#if 1
+	float3 H = normalize(V + L);
+	float NoL = saturate( dot(N, L) );
+	float NoV = saturate( abs( dot(N, V) ) + 1e-5 );
+	float NoH = saturate( dot(N, H) );
+	float VoH = saturate( dot(V, H) );
+	
+	// Generalized microfacet specular
+	float D = D_GGX( LobeRoughness[0], NoH ) * LobeEnergy[0];
+	float Vis = Vis_Kelemen( VoH );
+	
+	// F_Schlick
+	float F0 = 0.04;
+	float Fc = Pow5( 1 - VoH );
+	float F = Fc + (1 - Fc) * F0;
+	F *= ClearCoat;
+
+	float Fr1 = D * Vis * F;
+
+	float LayerAttenuation = (1 - F);
+
+
+	#if CLEAR_COAT_BOTTOM_NORMAL
+		{
+			//const float3 ClearCoatUnderNormal = OctahedronToUnitVector((float2(GBuffer.CustomData.a, GBuffer.CustomData.z) * 2) - (256.0/255.0));
+			const float2 oct1 = ((float2(GBuffer.CustomData.a, GBuffer.CustomData.z) * 2) - (256.0/255.0)) + UnitVectorToOctahedron(GBuffer.WorldNormal);
+ 			const float3 ClearCoatUnderNormal = OctahedronToUnitVector(oct1);
+			//float CNoL = saturate( dot(ClearCoatUnderNormal, L) );
+			//float CNoV = saturate( dot(ClearCoatUnderNormal, V) );
+			float CNoH = saturate( dot(ClearCoatUnderNormal, H) );
+
+			float D2 = D_GGX(LobeRoughness[1], CNoH ) * LobeEnergy[1];
+			float Vis2 = Vis_SmithJointApprox( LobeRoughness[1], NoV, NoL );
+			//float3 F2 = F_Schlick( GBuffer.SpecularColor, VoH );
+			float3 F2 = saturate( 50.0 * GBuffer.SpecularColor.g ) * Fc + (1 - Fc) * GBuffer.SpecularColor;
+
+			//Optional term taking into account Basic NdotL response of bottom normal. Not important for metallic which is the most common clearcoat case. Not energy conserving.
+			//float3 Fr2 = Diffuse_Lambert( GBuffer.DiffuseColor ) * CNoL + (D2 * Vis2) * F2;
+
+			float3 Fr2 = Diffuse_Lambert( GBuffer.DiffuseColor )  + (D2 * Vis2) * F2;
+
+			return Fr1 + Fr2 * LayerAttenuation;
+			//return float3(2,0,0);
+			//return ClearCoatUnderNormal;
+		}
+	#endif
+
+	// Generalized microfacet specular
+	float D2 = D_GGX( LobeRoughness[1], NoH ) * LobeEnergy[1];
+	float Vis2 = Vis_SmithJointApprox( LobeRoughness[1], NoV, NoL );
+	//float3 F2 = F_Schlick( GBuffer.SpecularColor, VoH );
+	float3 F2 = saturate( 50.0 * GBuffer.SpecularColor.g ) * Fc + (1 - Fc) * GBuffer.SpecularColor;
+
+	//float3 Fr2 = Diffuse_Burley( GBuffer.DiffuseColor, LobeRoughness[1], NoV, NoL, VoH ) * LobeEnergy[2] + (D2 * Vis2) * F2;
+	float3 Fr2 = Diffuse_Lambert( GBuffer.DiffuseColor ) * LobeEnergy[2] + (D2 * Vis2) * F2;
+	
+	return Fr1 + Fr2 * LayerAttenuation;
+#else
+	float3 H = normalize(V + L);
+	float NoL = saturate( dot(N, L) );
+	float NoV = saturate( abs( dot(N, V) ) + 1e-5 );
+	float NoH = saturate( dot(N, H) );
+	float VoH = saturate( dot(V, H) );
+	
+	// Hard coded IOR of 1.5
+
+	// Generalized microfacet specular
+	float D = D_GGX( ClearCoatRoughness, NoH ) * LobeEnergy[0];
+	float Vis = Vis_Kelemen( VoH );
+	
+	// F_Schlick
+	float F0 = 0.04;
+	float Fc = Pow5( 1 - VoH );
+	float F = Fc + (1 - Fc) * F0;
+
+	float Fr1 = D * Vis * F;
+
+	// Refract rays
+	//float3 L2 = refract( -L, -H, 1 / 1.5 );
+	//float3 V2 = refract( -V, -H, 1 / 1.5 );
+
+	// LoH == VoH
+	//float RefractBlend = sqrt( 4 * VoH*VoH + 5 ) / 3 + 2.0 / 3 * VoH;
+	//float3 L2 = RefractBlend * H - L / 1.5;
+	//float3 V2 = RefractBlend * H - V / 1.5;
+	//float NoL2 = saturate( dot(N, L2) );
+	//float NoV2 = saturate( dot(N, V2) );
+	
+	// Approximation
+	float RefractBlend = (0.22 * VoH + 0.7) * VoH + 0.745;	// 2 mad
+	// Dot products distribute. No need for L2 and V2.
+	float RefractNoH = RefractBlend * NoH;					// 1 mul
+	float NoL2 = saturate( RefractNoH - (1 / 1.5) * NoL );	// 1 mad
+	float NoV2 = saturate( RefractNoH - (1 / 1.5) * NoV );	// 1 mad
+	// Should refract H too but unimportant
+
+	NoL2 = max( 0.001, NoL2 );
+	NoV2 = max( 0.001, NoV2 );
+
+	float  AbsorptionDist = rcp(NoV2) + rcp(NoL2);
+	float3 Absorption = pow( AbsorptionColor, 0.5 * AbsorptionDist );
+
+	// Approximation
+	//float  AbsorptionDist = ( NoV2 + NoL2 ) / ( NoV2 * NoL2 );
+	//float3 Absorption = AbsorptionColor * ( AbsorptionColor * (AbsorptionDist * 0.5 - 1) + (2 - 0.5 * AbsorptionDist) );
+	//float3 Absorption = AbsorptionColor + AbsorptionColor * (AbsorptionColor - 1) * (AbsorptionDist * 0.5 - 1);	// use for shared version
+	
+	//float F21 = Fresnel( 1 / 1.5, saturate( dot(V2, H) ) );
+	//float TotalInternalReflection = 1 - F21 * G_Schlick( Roughness, NoV2, NoL2 );
+	//float3 LayerAttenuation = ( (1 - F12) * TotalInternalReflection ) * Absorption;
+
+	// Approximation
+	float3 LayerAttenuation = (1 - F) * Absorption;
+
+	// Approximation for IOR == 1.5
+	//SpecularColor = ChangeBaseMedium( SpecularColor, 1.5 );
+	//SpecularColor = saturate( ( 0.55 * SpecularColor + (0.45 * 1.08) ) * SpecularColor - (0.45 * 0.08) );
+	// Treat SpecularColor as relative to IOR. Artist compensates.
+	
+	// Generalized microfacet specular
+	float D2 = D_GGX( Roughness, NoH ) * LobeEnergy[2];
+	float Vis2 = Vis_SmithJointApprox( Roughness, NoV2, NoL2 );
+	float3 F2 = F_Schlick( GBuffer.SpecularColor, VoH );
+
+	float3 Fr2 = Diffuse_Lambert( GBuffer.DiffuseColor ) * LobeEnergy[2] + (D2 * Vis2) * F2;
+	
+	return Fr1 + Fr2 * LayerAttenuation;
+#endif
+}
+
+float3 ClothShading( FGBufferData GBuffer, float3 LobeRoughness, float3 LobeEnergy, float3 L, float3 V, half3 N )
+{
+	const float3 FuzzColor	= saturate(GBuffer.CustomData.rgb);
+	const float  Cloth		= saturate(GBuffer.CustomData.a);
+
+	float NoL = dot(N, L);
+	float NoV = dot(N, V);
+	float LoV = dot(L, V);
+	float InvLenH = rsqrt( 2 + 2 * LoV );
+	float NoH = saturate( ( NoL + NoV ) * InvLenH );
+	float VoH = saturate( InvLenH + InvLenH * LoV );
+	NoL = saturate(NoL);
+	NoV = saturate(abs(NoV) + 1e-5);
+
+	// Diffuse	
+	float3 Diffuse = Diffuse_Lambert( GBuffer.DiffuseColor );
+	float3 Diff = Diffuse * LobeEnergy[2];
+
+	// Cloth - Asperity Scattering - Inverse Beckmann Layer	
+	float3 F1 = F_Schlick( FuzzColor, VoH );
+	float  D1 = D_InvGGX( LobeRoughness[1], NoH );
+	float  V1 = Vis_Cloth( NoV, NoL );
+
+	float3 Spec1 = D1 * V1 * F1;
+
+	// Generalized microfacet specular
+	float3 F2 = F_Schlick( GBuffer.SpecularColor, VoH );
+	float  D2 = D_GGX( LobeRoughness[1], NoH ) * LobeEnergy[1];
+	float  V2 = Vis_SmithJointApprox( LobeRoughness[1], NoV, NoL );
+
+	float3 Spec2 = D2 * V2 * F2;
+
+	float3 Spec = lerp(Spec2, Spec1, Cloth);
+
+	return Diff + Spec;
+}
+
+float Hair_g( float B, float Theta )
+{
+	return exp( -0.5 * Pow2( Theta ) / (B*B) ) / ( sqrt(2*PI) * B );
+}
+
+float Hair_F( float CosTheta )
+{
+	const float n = 1.55;
+	const float F0 = Pow2( (1 - n) / (1 + n) );
+	return F0 + (1 - F0) * Pow5( 1 - CosTheta );
+}
+
+#define HAIR_REFERENCE	0
+#if HAIR_REFERENCE
+struct FHairTemp
+{
+	float SinThetaL;
+	float SinThetaV;
+	float CosThetaD;
+	float CosThetaT;
+	float CosPhi;
+	float CosHalfPhi;
+	float VoL;
+	float n_prime;
+};
+
+// Modified Bessel function
+float I0( float x )
+{
+	x = abs(x);
+	float a;
+	if( x < 3.75 )
+	{
+		float t = x / 3.75;
+		float t2 = t*t;
+		a =        + 0.0045813;
+		a = a * t2 + 0.0360768;
+		a = a * t2 + 0.2659732;
+		a = a * t2 + 1.2067492;
+		a = a * t2 + 3.0899424;
+		a = a * t2 + 3.5156229;
+		a = a * t2 + 1.0;
+	}
+	else
+	{
+		float t = 3.75 / x;
+		a =       + 0.00392377;
+		a = a * t - 0.01647633;
+		a = a * t + 0.02635537;
+		a = a * t - 0.02057706;
+		a = a * t + 0.00916281;
+		a = a * t - 0.00157565;
+		a = a * t + 0.00225319;
+		a = a * t + 0.01328592;
+		a = a * t + 0.39894228;
+		a *= exp(x) * rsqrt(x);
+	}
+	return a;
+}
+
+float LongitudinalScattering( float B, float SinThetaL, float SinThetaV )
+{
+	float v = B*B;
+	float CosThetaL2 = 1 - SinThetaL * SinThetaL;
+	float CosThetaV2 = 1 - SinThetaV * SinThetaV;
+	float Mp = 0;
+	if( v < 0.1 )
+	{
+		float a = sqrt(CosThetaL2 * CosThetaV2) / v;
+		float b = -SinThetaL * SinThetaV / v;
+		float logI0a = a > 12 ? a + 0.5 * ( -log(2*PI) + log(1/a) + 0.125/a ) : log( I0(a) );
+		Mp = exp( logI0a + b - rcp(v) + 0.6931 + log(0.5 / v) );
+	}
+	else
+	{
+		Mp = rcp( exp(2/v) * v - v ) * exp( ( 1 - SinThetaL * SinThetaV ) / v ) * I0( sqrt(CosThetaL2 * CosThetaV2) / v );
+	}
+	
+	return Mp;
+}
+
+float GaussianDetector( float Bp, float Phi )
+{
+	float Dp = 0;
+	for( int k = -4; k <= 4; k++ )
+	{
+		// TODO use symmetry and detect for both Phi and -Phi
+		Dp += Hair_g( Bp, Phi - (2*PI) * k );
+	}
+	return Dp;
+}
+
+float3 Attenuation( uint p, float h, float3 Color, FHairTemp HairTemp )
+{
+	float3 A;
+	if( p == 0 )
+	{
+		//A = F( cos( 0.5 * acos( HairTemp.VoL ) ) );
+		A = Hair_F( sqrt( 0.5 + 0.5 * HairTemp.VoL ) );
+	}
+	else
+	{
+		// ua is absorption
+		// ua = pe*Sigma_ae + pp*Sigma_ap
+		float3 Sigma_ae = { 0.419, 0.697, 1.37 };
+		float3 Sigma_ap = { 0.187, 0.4, 1.05 };
+		//float3 ua = 0.25 * Sigma_ae + 0.25 * Sigma_ap;
+		float3 ua = -0.25 * log( Color );
+		float3 ua_prime = ua / HairTemp.CosThetaT;
+		//float3 ua_prime = ua / sqrt( 1 - Pow2( HairTemp.CosThetaD ) / 2.4 );
+		
+		float yi = asin(h);
+		float yt = asin(h / HairTemp.n_prime);
+		
+		float f = Hair_F( HairTemp.CosThetaD * sqrt( 1 - h*h ) );		// (14)
+		//float3 T = exp( -2 * ua_prime * ( 1 + cos(2*yt) ) );
+		float3 T = exp( -2 * ua_prime * cos(yt) );
+		if( p == 1 )
+			A = Pow2(1 - f) * T;		// (13)
+		else
+			A = Pow2(1 - f) * f * T*T;	// (13)
+	}
+	return A;
+}
+
+float Omega( uint p, float h, FHairTemp HairTemp )
+{
+	float yi = asin(h);
+	float yt = asin(h / HairTemp.n_prime);
+	return 2*p*yt - 2*yi + p*PI;
+}
+
+float3 AzimuthalScattering( uint p, float Bp, float3 Color, FHairTemp HairTemp, uint2 Random )
+{
+	float Phi = acos( HairTemp.CosPhi );
+	
+	// Np = 0.5 * Integral_-1^1( A(p,h) * Dp( Phi - Omega(p,h) ) * dh )
+
+	float Offset = float( Random.x & 0xffff ) / (1<<16);
+
+	uint Num = 16;
+	float3 Np = 0;
+	for( uint i = 0; i < Num; i++ )
+	{
+		float h = ( (float)( i + Offset ) / Num ) * 2 - 1;
+		Np += Attenuation( p, h, Color, HairTemp ) * GaussianDetector( Bp, Phi - Omega( p, h, HairTemp ) );
+	}
+	Np *= 2.0 / Num;
+
+	return 0.5 * Np;
+}
+
+// [d'Eon et al. 2011, "An Energy-Conserving Hair Reflectance Model"]
+// [d'Eon et al. 2014, "A Fiber Scattering Model with Non-Separable Lobes"]
+float3 HairShadingRef( FGBufferData GBuffer, float3 L, float3 V, half3 N, uint2 Random )
+{
+	// to prevent NaN with decals
+	// OR-18489 HERO: IGGY: RMB on E ability causes blinding hair effect
+	// OR-17578 HERO: HAMMER: E causes blinding light on heroes with hair
+	float ClampedRoughness = clamp(GBuffer.Roughness, 1/255.0f, 1.0f);
+
+	float n = 1.55;
+
+	FHairTemp HairTemp;
+	
+	// N is the vector parallel to hair pointing toward root
+	HairTemp.VoL       = dot(V,L);
+	HairTemp.SinThetaL = dot(N,L);
+	HairTemp.SinThetaV = dot(N,V);
+	// SinThetaT = 1/n * SinThetaL
+	HairTemp.CosThetaT = sqrt( 1 - Pow2( (1/n) * HairTemp.SinThetaL ) );
+	HairTemp.CosThetaD = cos( 0.5 * abs( asin( HairTemp.SinThetaV ) - asin( HairTemp.SinThetaL ) ) );
+
+	float3 Lp = L - HairTemp.SinThetaL * N;
+	float3 Vp = V - HairTemp.SinThetaV * N;
+	HairTemp.CosPhi = dot(Lp,Vp) * rsqrt( dot(Lp,Lp) * dot(Vp,Vp) );
+	HairTemp.CosHalfPhi = sqrt( 0.5 + 0.5 * HairTemp.CosPhi );
+
+	HairTemp.n_prime = sqrt( n*n - 1 + Pow2( HairTemp.CosThetaD ) ) / HairTemp.CosThetaD;
+
+	float Shift = 0.035;
+	float Alpha[] =
+	{
+		-Shift * 2,
+		Shift,
+		Shift * 4,
+	};
+	float B[] =
+	{
+		Pow2( ClampedRoughness ),
+		Pow2( ClampedRoughness ) / 2,
+		Pow2( ClampedRoughness ) * 2,
+	};
+
+	float3 S = 0;
+	UNROLL for( uint p = 0; p < 3; p++ )
+	{
+		float SinThetaV = HairTemp.SinThetaV;
+		float Bp = B[p];
+		if( p == 0 )
+		{
+			Bp *= sqrt(2.0) * HairTemp.CosHalfPhi;
+			float sa, ca;
+			sincos( Alpha[p], sa, ca );
+			SinThetaV -= 2*sa * ( HairTemp.CosHalfPhi * ca * sqrt( 1 - SinThetaV * SinThetaV ) + sa * SinThetaV );
+		}
+		else
+		{
+			SinThetaV = sin( asin(SinThetaV) - Alpha[p] );
+		}
+		float Mp = LongitudinalScattering( Bp, HairTemp.SinThetaL, SinThetaV );
+		float3 Np = AzimuthalScattering( p, B[p], GBuffer.BaseColor, HairTemp, Random );
+
+		float3 Sp = Mp * Np;
+		S += Sp;
+	}
+	return S;
+}
+#endif
+
+// Approximation to HairShadingRef using concepts from the following papers:
+// [Marschner et al. 2003, "Light Scattering from Human Hair Fibers"]
+// [Pekelis et al. 2015, "A Data-Driven Light Scattering Model for Hair"]
+float3 HairShading( FGBufferData GBuffer, float3 L, float3 V, half3 N, float Shadow, float Backlit, float Area, uint2 Random )
+{
+	// to prevent NaN with decals
+	// OR-18489 HERO: IGGY: RMB on E ability causes blinding hair effect
+	// OR-17578 HERO: HAMMER: E causes blinding light on heroes with hair
+	float ClampedRoughness = clamp(GBuffer.Roughness, 1/255.0f, 1.0f);
+
+	//const float3 DiffuseN	= OctahedronToUnitVector( GBuffer.CustomData.xy * 2 - 1 );
+	//const float Backlit	= GBuffer.CustomData.z;
+
+#if HAIR_REFERENCE
+	// todo: ClampedRoughness is missing for this code path
+	float3 S = HairShadingRef( GBuffer, L, V, N, Random );
+	//float3 S = HairShadingMarschner( GBuffer, L, V, N );
+#else
+	// N is the vector parallel to hair pointing toward root
+
+	const float VoL       = dot(V,L);
+	const float SinThetaL = dot(N,L);
+	const float SinThetaV = dot(N,V);
+	float CosThetaD = cos( 0.5 * abs( asinFast( SinThetaV ) - asinFast( SinThetaL ) ) );
+
+	//CosThetaD = abs( CosThetaD ) < 0.01 ? 0.01 : CosThetaD;
+
+	const float3 Lp = L - SinThetaL * N;
+	const float3 Vp = V - SinThetaV * N;
+	const float CosPhi = dot(Lp,Vp) * rsqrt( dot(Lp,Lp) * dot(Vp,Vp) + 1e-4 );
+	const float CosHalfPhi = sqrt( saturate( 0.5 + 0.5 * CosPhi ) );
+	//const float Phi = acosFast( CosPhi );
+
+	float n = 1.55;
+	//float n_prime = sqrt( n*n - 1 + Pow2( CosThetaD ) ) / CosThetaD;
+	float n_prime = 1.19 / CosThetaD + 0.36 * CosThetaD;
+
+	float Shift = 0.035;
+	float Alpha[] =
+	{
+		-Shift * 2,
+		Shift,
+		Shift * 4,
+	};
+	float B[] =
+	{
+		Area + Pow2( ClampedRoughness ),
+		Area + Pow2( ClampedRoughness ) / 2,
+		Area + Pow2( ClampedRoughness ) * 2,
+	};
+
+	float3 S = 0;
+
+	// R
+	if(1)
+	{
+		const float sa = sin( Alpha[0] );
+		const float ca = cos( Alpha[0] );
+		float Shift = 2*sa* ( ca * CosHalfPhi * sqrt( 1 - SinThetaV * SinThetaV ) + sa * SinThetaV );
+
+		float Mp = Hair_g( B[0] * sqrt(2.0) * CosHalfPhi, SinThetaL + SinThetaV - Shift );
+		float Np = 0.25 * CosHalfPhi;
+		float Fp = Hair_F( sqrt( saturate( 0.5 + 0.5 * VoL ) ) );
+		S += Mp * Np * Fp * ( GBuffer.Specular * 2 ) * lerp( 1, Backlit, saturate(-VoL) );
+	}
+
+	// TT
+	if(1)
+	{
+		float Mp = Hair_g( B[1], SinThetaL + SinThetaV - Alpha[1] );
+
+		float a = 1 / n_prime;
+		//float h = CosHalfPhi * rsqrt( 1 + a*a - 2*a * sqrt( 0.5 - 0.5 * CosPhi ) );
+		//float h = CosHalfPhi * ( ( 1 - Pow2( CosHalfPhi ) ) * a + 1 );
+		float h = CosHalfPhi * ( 1 + a * ( 0.6 - 0.8 * CosPhi ) );
+		//float h = 0.4;
+		//float yi = asinFast(h);
+		//float yt = asinFast(h / n_prime);
+		
+		float f = Hair_F( CosThetaD * sqrt( saturate( 1 - h*h ) ) );
+		float Fp = Pow2(1 - f);
+		//float3 Tp = pow( GBuffer.BaseColor, 0.5 * ( 1 + cos(2*yt) ) / CosThetaD );
+		//float3 Tp = pow( GBuffer.BaseColor, 0.5 * cos(yt) / CosThetaD );
+		float3 Tp = pow( GBuffer.BaseColor, 0.5 * sqrt( 1 - Pow2(h * a) ) / CosThetaD );
+
+		//float t = asin( 1 / n_prime );
+		//float d = ( sqrt(2) - t ) / ( 1 - t );
+		//float s = -0.5 * PI * (1 - 1 / n_prime) * log( 2*d - 1 - 2 * sqrt( d * (d - 1) ) );
+		//float s = 0.35;
+		//float Np = exp( (Phi - PI) / s ) / ( s * Pow2( 1 + exp( (Phi - PI) / s ) ) );
+		//float Np = 0.71 * exp( -1.65 * Pow2(Phi - PI) );
+		float Np = exp( -3.65 * CosPhi - 3.98 );
+
+		S += Mp * Np * Fp * Tp * Backlit;
+	}
+
+	// TRT
+	if(1)
+	{
+		float Mp = Hair_g( B[2], SinThetaL + SinThetaV - Alpha[2] );
+		
+		//float h = 0.75;
+		float f = Hair_F( CosThetaD * 0.5 );
+		float Fp = Pow2(1 - f) * f;
+		//float3 Tp = pow( GBuffer.BaseColor, 1.6 / CosThetaD );
+		float3 Tp = pow( GBuffer.BaseColor, 0.8 / CosThetaD );
+
+		//float s = 0.15;
+		//float Np = 0.75 * exp( Phi / s ) / ( s * Pow2( 1 + exp( Phi / s ) ) );
+		float Np = exp( 17 * CosPhi - 16.78 );
+
+		S += Mp * Np * Fp * Tp;
+	}
+#endif
+
+	if(1)
+	{
+		float3 FakeNormal = normalize( V - N * dot(V,N) );
+		//N = normalize( DiffuseN + FakeNormal * 2 );
+		N = FakeNormal;
+
+		// Hack approximation for multiple scattering.
+		float Wrap = 1;
+		float NoL = saturate( ( dot(N, L) + Wrap ) / Square( 1 + Wrap ) );
+		float DiffuseScatter = (1 / PI) * NoL * GBuffer.Metallic;
+		float Luma = Luminance( GBuffer.BaseColor );
+		float3 ScatterTint = pow( GBuffer.BaseColor / Luma, 1 - Shadow );
+		S += sqrt( GBuffer.BaseColor ) * DiffuseScatter * ScatterTint;
+	}
+
+	S = -min(-S, 0.0);
+
+	return S;
+}
+
+float3 SubsurfaceShadingSubsurface( FGBufferData GBuffer, float3 L, float3 V, half3 N )
+{
+	float3 SubsurfaceColor = ExtractSubsurfaceColor(GBuffer);
+	float Opacity = GBuffer.CustomData.a;
+
+	float3 H = normalize(V + L);
+
+	// to get an effect when you see through the material
+	// hard coded pow constant
+	float InScatter = pow(saturate(dot(L, -V)), 12) * lerp(3, .1f, Opacity);
+	// wrap around lighting, /(PI*2) to be energy consistent (hack do get some view dependnt and light dependent effect)
+	// Opacity of 0 gives no normal dependent lighting, Opacity of 1 gives strong normal contribution
+	float NormalContribution = saturate(dot(N, H) * Opacity + 1 - Opacity);
+	float BackScatter = GBuffer.GBufferAO * NormalContribution / (PI * 2);
+
+	// lerp to never exceed 1 (energy conserving)
+	return SubsurfaceColor * lerp(BackScatter, 1, InScatter);
+}
+
+float3 SubsurfaceShadingTwoSided( float3 SubsurfaceColor, float3 L, float3 V, half3 N )
+{
+	// http://blog.stevemcauley.com/2011/12/03/energy-conserving-wrapped-diffuse/
+	float Wrap = 0.5;
+	float NoL = saturate( ( dot(-N, L) + Wrap ) / Square( 1 + Wrap ) );
+
+	// GGX scatter distribution
+	float VoL = saturate( dot(V, -L) );
+	float a = 0.6;
+	float a2 = a * a;
+	float d = ( VoL * a2 - VoL ) * VoL + 1;	// 2 mad
+	float GGX = (a2 / PI) / (d * d);		// 2 mul, 1 rcp
+	return NoL * GGX * SubsurfaceColor;
+}
+
+float3 EyeShading( FGBufferData GBuffer, float3 LobeRoughness, float3 LobeEnergy, float3 L, float3 V, half3 N )
+{
+	float NoL = dot(N, L);
+	float NoV = dot(N, V);
+	float LoV = dot(L, V);
+	float InvLenH = rsqrt( 2 + 2 * LoV );
+	float NoH = saturate( ( NoL + NoV ) * InvLenH );
+	float VoH = saturate( InvLenH + InvLenH * LoV );
+	NoL = saturate(NoL);
+	NoV = saturate(abs(NoV) + 1e-5);
+
+	// Generalized microfacet specular
+	float D = D_GGX( LobeRoughness[1], NoH ) * LobeEnergy[1];
+	float Vis = Vis_SmithJointApprox( LobeRoughness[1], NoV, NoL );
+	float3 F = F_Schlick( GBuffer.SpecularColor, VoH );
+
+	return D * Vis * F;
+}
+
+float3 EyeSubsurfaceShading( FGBufferData GBuffer, float3 L, float3 V, half3 N )
+{
+	float NoL = dot(N, L);
+	float LoV = dot(L, V);
+	float InvLenH = rsqrt( 2 + 2 * LoV );
+	float VoH = saturate( InvLenH + InvLenH * LoV );
+
+	// F_Schlick
+	float F0 = GBuffer.Specular * 0.08;
+	float Fc = Pow5( 1 - VoH );
+	float F = Fc + (1 - Fc) * F0;
+
+	float  IrisDistance = GBuffer.CustomData.w;
+	float  IrisMask = GBuffer.CustomData.z;
+
+	float3 IrisNormal;
+	IrisNormal = OctahedronToUnitVector( GBuffer.CustomData.xy * 2 - 1 );
+
+	// Blend in the negative intersection normal to create some concavity
+	// Not great as it ties the concavity to the convexity of the cornea surface
+	// No good justification for that. On the other hand, if we're just looking to
+	// introduce some concavity, this does the job.
+	float3 CausticNormal = normalize(lerp(IrisNormal, -N, IrisMask*IrisDistance));
+
+	float IrisNoL = saturate( dot( IrisNormal, L ) );
+	float Power = lerp( 12, 1, IrisNoL );
+	float Caustic = 0.6 + 0.2 * ( Power + 1 ) * pow( saturate( dot( CausticNormal, L ) ), Power );
+	float Iris = IrisNoL * Caustic;
+
+	// http://blog.stevemcauley.com/2011/12/03/energy-conserving-wrapped-diffuse/
+	float Wrap = 0.15;
+	float Sclera = saturate( ( NoL + Wrap ) / Square( 1 + Wrap ) );
+	
+	return (1 - F) * lerp( Sclera, Iris, IrisMask ) * GBuffer.DiffuseColor / PI;
+}
+
+Texture2D		PreIntegratedBRDF;
+SamplerState	PreIntegratedBRDFSampler;
+
+float3 SubsurfaceShadingPreintegratedSkin( FGBufferData GBuffer, float3 L, float3 V, half3 N )
+{
+	float3 SubsurfaceColor = ExtractSubsurfaceColor(GBuffer);
+	float Opacity = GBuffer.CustomData.a;
+
+	float3 PreintegratedBRDF = Texture2DSampleLevel(PreIntegratedBRDF, PreIntegratedBRDFSampler, float2(saturate(dot(N, L) * .5 + .5), 1 - Opacity), 0).rgb;
+	return PreintegratedBRDF * SubsurfaceColor;
+}
+
+// @param DiffSpecMask .r: diffuse, .g:specular e.g. float2(1,1) for both, float2(1,0) for diffuse only
+float3 SurfaceShading( FGBufferData GBuffer, float3 LobeRoughness, float3 LobeEnergy, float3 L, float3 V, half3 N, uint2 Random )
+{
+	switch( GBuffer.ShadingModelID )
+	{
+		case SHADINGMODELID_UNLIT:
+		case SHADINGMODELID_DEFAULT_LIT:
+		case SHADINGMODELID_SUBSURFACE:
+		case SHADINGMODELID_PREINTEGRATED_SKIN:
+		case SHADINGMODELID_SUBSURFACE_PROFILE:
+		case SHADINGMODELID_TWOSIDED_FOLIAGE:
+			return StandardShading( GBuffer.DiffuseColor, GBuffer.SpecularColor, LobeRoughness, LobeEnergy, L, V, N );
+		case SHADINGMODELID_CLEAR_COAT:
+			return ClearCoatShading( GBuffer, LobeRoughness, LobeEnergy, L, V, N );
+		case SHADINGMODELID_CLOTH:
+			return ClothShading( GBuffer, LobeRoughness, LobeEnergy, L, V, N );
+		case SHADINGMODELID_EYE:
+			return EyeShading( GBuffer, LobeRoughness, LobeEnergy, L, V, N );
+		default:
+			return 0;
+	}
+}
+
+float3 SubsurfaceShading( FGBufferData GBuffer, float3 L, float3 V, half3 N, float Shadow, uint2 Random )
+{
+	float3 SubsurfaceColor = ExtractSubsurfaceColor(GBuffer);
+
+	switch( GBuffer.ShadingModelID )
+	{
+		case SHADINGMODELID_SUBSURFACE:
+			return SubsurfaceShadingSubsurface( GBuffer, L, V, N );
+		case SHADINGMODELID_PREINTEGRATED_SKIN:
+			return SubsurfaceShadingPreintegratedSkin( GBuffer, L, V, N );
+		case SHADINGMODELID_TWOSIDED_FOLIAGE:
+			return SubsurfaceShadingTwoSided( SubsurfaceColor, L, V, N );
+		case SHADINGMODELID_HAIR:
+			return HairShading( GBuffer, L, V, N, Shadow, 1, 0, Random );
+		case SHADINGMODELID_EYE:
+			return EyeSubsurfaceShading( GBuffer, L, V, N );
+		default:
+			return 0;
+	}
+}
\ No newline at end of file
diff --git a/Engine/Shaders/Private/ShadingModelsMaterial.ush b/Engine/Shaders/Private/ShadingModelsMaterial.ush
new file mode 100644
--- /dev/null
+++ b/Engine/Shaders/Private/ShadingModelsMaterial.ush
@@ -0,0 +1,127 @@
+// Copyright 1998-2017 Epic Games, Inc. All Rights Reserved.
+
+/*=============================================================================
+	ShadingModelsMaterial.usf: Pixel shader function for computing a GBuffer from shading model.
+=============================================================================*/
+
+void SetGBufferForShadingModel(
+	in out FGBufferData GBuffer, 
+	in const FMaterialPixelParameters MaterialParameters,
+	const float Opacity,
+	const half3 BaseColor,
+	const half  Metallic,
+	const half  Specular,
+	const float Roughness,
+	const float4 SubsurfaceData)
+{
+	const float3 SubsurfaceColor = SubsurfaceData.rgb;
+	const float SubsurfaceProfile = SubsurfaceData.a;
+
+	GBuffer.WorldNormal = MaterialParameters.WorldNormal;
+	GBuffer.BaseColor = BaseColor;
+	GBuffer.Metallic = Metallic;
+	GBuffer.Specular = Specular;
+	GBuffer.Roughness = Roughness;
+
+#if MATERIAL_SHADINGMODEL_UNLIT
+	GBuffer.ShadingModelID = SHADINGMODELID_UNLIT;
+#elif MATERIAL_SHADINGMODEL_DEFAULT_LIT
+	GBuffer.ShadingModelID = SHADINGMODELID_DEFAULT_LIT;
+#elif MATERIAL_SHADINGMODEL_SUBSURFACE
+	GBuffer.ShadingModelID = SHADINGMODELID_SUBSURFACE;
+	GBuffer.CustomData.rgb = EncodeSubsurfaceColor(SubsurfaceColor);
+	GBuffer.CustomData.a = Opacity;
+#elif MATERIAL_SHADINGMODEL_PREINTEGRATED_SKIN
+	GBuffer.ShadingModelID = SHADINGMODELID_PREINTEGRATED_SKIN;
+	GBuffer.CustomData.rgb = EncodeSubsurfaceColor(SubsurfaceColor);
+	GBuffer.CustomData.a = Opacity;
+#elif MATERIAL_SHADINGMODEL_SUBSURFACE_PROFILE
+	// Optimization: if opacity is 0 then revert to default shading model 
+	#if SUBSURFACE_PROFILE_OPACITY_THRESHOLD
+	if (Opacity > 0.0)
+	#endif
+	{
+		GBuffer.ShadingModelID = SHADINGMODELID_SUBSURFACE_PROFILE;
+		GBuffer.CustomData.rgb = EncodeSubsurfaceProfile(SubsurfaceProfile);
+		GBuffer.CustomData.a = Opacity;
+	}
+	#if SUBSURFACE_PROFILE_OPACITY_THRESHOLD
+	else
+	{
+		GBuffer.ShadingModelID = SHADINGMODELID_DEFAULT_LIT;
+		GBuffer.CustomData = 0;
+	}
+	#endif
+#elif MATERIAL_SHADINGMODEL_CLEAR_COAT
+	GBuffer.ShadingModelID = SHADINGMODELID_CLEAR_COAT;
+
+	float ClearCoat				= saturate( GetMaterialCustomData0(MaterialParameters) );
+	float ClearCoatRoughness	= saturate( GetMaterialCustomData1(MaterialParameters) );
+	float MetalSpec = 0.9;
+
+	float NoV = saturate( dot( MaterialParameters.WorldNormal, MaterialParameters.CameraVector ) );
+
+	// Approximation of refraction's effect on EnvBRDF
+	float RefractionScale = ( (NoV * 0.5 + 0.5) * NoV - 1 ) * saturate( 1.25 - 1.25 * Roughness ) + 1;
+
+	// Approximation of absorption integral, tuned for Roughness=0.4
+	float3 AbsorptionColor = BaseColor * (1 / MetalSpec);
+	float3 Absorption = AbsorptionColor * ( (NoV - 1) * 0.85 * ( 1 - lerp( AbsorptionColor, Square(AbsorptionColor), -0.78 ) ) + 1 );
+
+	float F0 = 0.04;
+	float Fc = Pow5( 1 - NoV );
+	float F = Fc + (1 - Fc) * F0;
+	float LayerAttenuation = lerp( 1, (1 - F), ClearCoat );
+
+	GBuffer.BaseColor = lerp( BaseColor * LayerAttenuation, MetalSpec * Absorption * RefractionScale, Metallic * ClearCoat );
+	GBuffer.Specular *= lerp( 1, RefractionScale, ClearCoat );
+
+	GBuffer.CustomData.x = ClearCoat;
+	GBuffer.CustomData.y = ClearCoatRoughness;
+
+	#if CLEAR_COAT_BOTTOM_NORMAL
+	{
+		float2 oct2 = UnitVectorToOctahedron(GBuffer.WorldNormal);
+
+		#if NUM_MATERIAL_OUTPUTS_CLEARCOATBOTTOMNORMAL > 0
+			#if MATERIAL_TANGENTSPACENORMAL
+				float3 tempnormal = normalize(TransformTangentVectorToWorld( MaterialParameters.TangentToWorld, ClearCoatBottomNormal0(MaterialParameters) ));
+			#else
+				float3 tempnormal =  ClearCoatBottomNormal0(MaterialParameters);
+			#endif
+
+			float2 oct1 = UnitVectorToOctahedron(tempnormal);
+			float2 oct3 = ( (oct1 - oct2) *  0.5 ) + (128.0/255.0);
+			GBuffer.CustomData.a = oct3.x;
+			GBuffer.CustomData.z = oct3.y;
+		#else
+			GBuffer.CustomData.a = 128.0/255.0;
+			GBuffer.CustomData.z = 128.0/255.0;
+		#endif
+	}
+	#endif
+#elif MATERIAL_SHADINGMODEL_TWOSIDED_FOLIAGE
+	GBuffer.ShadingModelID = SHADINGMODELID_TWOSIDED_FOLIAGE;
+	GBuffer.CustomData.rgb = EncodeSubsurfaceColor(SubsurfaceColor);
+	GBuffer.CustomData.a = Opacity;
+#elif MATERIAL_SHADINGMODEL_HAIR
+	GBuffer.ShadingModelID = SHADINGMODELID_HAIR;
+	GBuffer.CustomData.xy = UnitVectorToOctahedron( MaterialParameters.WorldNormal ) * 0.5 + 0.5;
+	GBuffer.CustomData.z = saturate( GetMaterialCustomData0(MaterialParameters) );	// Backlit
+#elif MATERIAL_SHADINGMODEL_CLOTH
+	GBuffer.ShadingModelID = SHADINGMODELID_CLOTH;
+	GBuffer.CustomData.rgb = SubsurfaceColor;
+	GBuffer.CustomData.a = saturate( GetMaterialCustomData0(MaterialParameters) );	// Cloth
+	GBuffer.IndirectIrradiance *= 1 - GBuffer.CustomData.a;
+#elif MATERIAL_SHADINGMODEL_EYE
+	GBuffer.ShadingModelID = SHADINGMODELID_EYE;
+	#if NUM_MATERIAL_OUTPUTS_GETTANGENTOUTPUT > 0
+		float3 Tangent = GetTangentOutput0(MaterialParameters);
+		GBuffer.CustomData.xy = UnitVectorToOctahedron( normalize(Tangent) ) * 0.5 + 0.5;
+	#endif
+	GBuffer.CustomData.z = saturate( GetMaterialCustomData0(MaterialParameters) );	// Iris Mask
+	GBuffer.CustomData.w = saturate( GetMaterialCustomData1(MaterialParameters) );	// Iris Distance
+#else
+	// missing shading model, compiler should report ShadingModelID is not set
+#endif
+}
\ No newline at end of file
diff --git a/Engine/Shaders/Private/ShadowDepthCommon.ush b/Engine/Shaders/Private/ShadowDepthCommon.ush
new file mode 100644
--- /dev/null
+++ b/Engine/Shaders/Private/ShadowDepthCommon.ush
@@ -0,0 +1,43 @@
+// Copyright 1998-2017 Epic Games, Inc. All Rights Reserved.
+
+/*=============================================================================
+	ShadowDepthCommon.usf: Shared functionality for shadow depth rendering.
+=============================================================================*/
+
+#ifndef PERSPECTIVE_CORRECT_DEPTH
+#define PERSPECTIVE_CORRECT_DEPTH 0
+#endif
+
+#ifndef REFLECTIVE_SHADOW_MAP
+#define REFLECTIVE_SHADOW_MAP 0
+#endif
+
+#define INTERPOLATE_VF_ATTRIBUTES (!MATERIALBLENDING_SOLID || REFLECTIVE_SHADOW_MAP || MATERIALBLENDING_MASKED || (REQUIRES_VF_ATTRIBUTES_FOR_CLIPPING && !POSITION_ONLY))
+#define INTERPOLATE_POSITION ((REFLECTIVE_SHADOW_MAP && !ONEPASS_POINTLIGHT_SHADOW) || MATERIALBLENDING_MASKED)
+
+// x: Tweakable depth bias, y:InvMaxSubjectDepthUsed to normalize the outputted depth
+float2 ShadowParams;
+
+/** Data passed from the vertex shader to the pixel shader for the shadow depth pass. */
+struct FShadowDepthVSToPS
+{
+#if INTERPOLATE_VF_ATTRIBUTES
+	FVertexFactoryInterpolantsVSToPS FactoryInterpolants;
+#endif
+
+#if ONEPASS_POINTLIGHT_SHADOW
+	float4 GSPosition : TEXCOORD6;
+
+#elif PERSPECTIVE_CORRECT_DEPTH
+	float ShadowDepth : TEXCOORD6;
+
+#elif !COMPILER_SUPPORTS_EMPTY_STRUCTS
+	float Dummy : TEXCOORD6;
+
+#endif
+
+#if INTERPOLATE_POSITION
+	// world space position needed in some cases, could be changed to be reconsructed from SvPosition (more tricky as this is not from camera perspective)
+	float3 PixelPosition	: TEXCOORD7;
+#endif
+};
