Branch snapshot for CL 3555263

diff --git a/Engine/Shaders/Private/ShadowFilteringCommon.ush b/Engine/Shaders/Private/ShadowFilteringCommon.ush
new file mode 100644
--- /dev/null
+++ b/Engine/Shaders/Private/ShadowFilteringCommon.ush
@@ -0,0 +1,445 @@
+// Copyright 1998-2017 Epic Games, Inc. All Rights Reserved.
+
+/*=============================================================================
+	ShadowFilteringCommon.usf: Contains functions to filter a shadowmap, shared between forward/deferred shading.
+=========================================================================*/
+
+#pragma once
+
+#include "PixelQuadMessagePassing.ush"
+
+
+struct FPCFSamplerSettings
+{
+	Texture2D		ShadowDepthTexture;
+	SamplerState	ShadowDepthTextureSampler;
+
+	//XY - Pixel size of shadowmap
+	//ZW - Inverse pixel size of shadowmap
+	float4			ShadowBufferSize;
+
+	// SceneDepth in lightspace.
+	float			SceneDepth;
+
+	float			TransitionScale;
+
+	// set by the caller, constant for the code so only one code branch should be compiled
+	bool			bSubsurface;
+
+	// Whether to treat shadow depths near 1 as unshadowed.  This is useful when the shadowmap does not contain all the points being shadowed.
+	bool			bTreatMaxDepthUnshadowed;
+
+	// only used if bSubsurface is true
+	float			DensityMulConstant;
+
+	// only used if bSubsurface is true
+	float2			ProjectionDepthBiasParameters;
+};
+
+// linear PCF, input 3x3
+// @param Values0 in row 0 from left to right: x,y,z,w
+// @param Values1 in row 1 from left to right: x,y,z,w
+// @param Values2 in row 2 from left to right: x,y,z,w
+// can be optimized
+float PCF2x2(float2 Fraction, float3 Values0, float3 Values1, float3 Values2)
+{
+	float2 VerticalLerp00 = lerp(float2(Values0.x, Values1.x), float2(Values0.y, Values1.y), Fraction.xx);
+	float PCFResult00 = lerp(VerticalLerp00.x, VerticalLerp00.y, Fraction.y);
+
+	float2 VerticalLerp10 = lerp(float2(Values0.y, Values1.y), float2(Values0.z, Values1.z), Fraction.xx);
+	float PCFResult10 = lerp(VerticalLerp10.x, VerticalLerp10.y, Fraction.y);
+
+	float2 VerticalLerp01 = lerp(float2(Values1.x, Values2.x), float2(Values1.y, Values2.y), Fraction.xx);
+	float PCFResult01 = lerp(VerticalLerp01.x, VerticalLerp01.y, Fraction.y);
+
+	float2 VerticalLerp11 = lerp(float2(Values1.y, Values2.y), float2(Values1.z, Values2.z), Fraction.xx);
+	float PCFResult11 = lerp(VerticalLerp11.x, VerticalLerp11.y, Fraction.y);	
+
+	return saturate((PCFResult00 + PCFResult10 + PCFResult01 + PCFResult11) * 0.25f);
+}
+
+// linear PCF, input 4x4
+// @param Values0 in row 0 from left to right: x,y,z,w
+// @param Values1 in row 1 from left to right: x,y,z,w
+// @param Values2 in row 2 from left to right: x,y,z,w
+// @param Values3 in row 3 from left to right: x,y,z,w
+// can be optimized
+float PCF3x3(float2 Fraction, float4 Values0, float4 Values1, float4 Values2, float4 Values3)
+{
+	float2 VerticalLerp00 = lerp(float2(Values0.x, Values1.x), float2(Values0.y, Values1.y), Fraction.xx);
+	float PCFResult00 = lerp(VerticalLerp00.x, VerticalLerp00.y, Fraction.y);
+
+	float2 VerticalLerp10 = lerp(float2(Values0.y, Values1.y), float2(Values0.z, Values1.z), Fraction.xx);
+	float PCFResult10 = lerp(VerticalLerp10.x, VerticalLerp10.y, Fraction.y);
+
+	float2 VerticalLerp20 = lerp(float2(Values0.z, Values1.z), float2(Values0.w, Values1.w), Fraction.xx);
+	float PCFResult20 = lerp(VerticalLerp20.x, VerticalLerp20.y, Fraction.y);
+
+	float2 VerticalLerp01 = lerp(float2(Values1.x, Values2.x), float2(Values1.y, Values2.y), Fraction.xx);
+	float PCFResult01 = lerp(VerticalLerp01.x, VerticalLerp01.y, Fraction.y);
+
+	float2 VerticalLerp11 = lerp(float2(Values1.y, Values2.y), float2(Values1.z, Values2.z), Fraction.xx);
+	float PCFResult11 = lerp(VerticalLerp11.x, VerticalLerp11.y, Fraction.y);
+
+	float2 VerticalLerp21 = lerp(float2(Values1.z, Values2.z), float2(Values1.w, Values2.w), Fraction.xx);
+	float PCFResult21 = lerp(VerticalLerp21.x, VerticalLerp21.y, Fraction.y);
+
+	float2 VerticalLerp02 = lerp(float2(Values2.x, Values3.x), float2(Values2.y, Values3.y), Fraction.xx);
+	float PCFResult02 = lerp(VerticalLerp02.x, VerticalLerp02.y, Fraction.y);
+
+	float2 VerticalLerp12 = lerp(float2(Values2.y, Values3.y), float2(Values2.z, Values3.z), Fraction.xx);
+	float PCFResult12 = lerp(VerticalLerp12.x, VerticalLerp12.y, Fraction.y);
+
+	float2 VerticalLerp22 = lerp(float2(Values2.z, Values3.z), float2(Values2.w, Values3.w), Fraction.xx);
+	float PCFResult22 = lerp(VerticalLerp22.x, VerticalLerp22.y, Fraction.y);
+
+	return saturate((PCFResult00 + PCFResult10 + PCFResult20 + PCFResult01 + PCFResult11 + PCFResult21 + PCFResult02 + PCFResult12 + PCFResult22) * .11111f);
+}
+
+
+// linear PCF, input 4x4
+// using Gather: xyzw in counter clockwise order starting with the sample to the lower left of the queried location
+// @param Values0 left top
+// @param Values1 right top
+// @param Values2 left bottom
+// @param Values3 right bottom
+float PCF3x3gather(float2 Fraction, float4 Values0, float4 Values1, float4 Values2, float4 Values3)
+{
+	float4 Results;
+
+	Results.x = Values0.w * (1.0 - Fraction.x);
+	Results.y = Values0.x * (1.0 - Fraction.x);
+	Results.z = Values2.w * (1.0 - Fraction.x);
+	Results.w = Values2.x * (1.0 - Fraction.x);
+	Results.x += Values0.z;
+	Results.y += Values0.y;
+	Results.z += Values2.z;
+	Results.w += Values2.y;
+	Results.x += Values1.w;
+	Results.y += Values1.x;
+	Results.z += Values3.w;
+	Results.w += Values3.x;
+	Results.x += Values1.z * Fraction.x;
+	Results.y += Values1.y * Fraction.x;
+	Results.z += Values3.z * Fraction.x;
+	Results.w += Values3.y * Fraction.x;
+
+	return dot( Results, float4( 1.0 - Fraction.y, 1.0, 1.0, Fraction.y) * ( 1.0 / 9.0) );
+}
+
+// horizontal PCF, input 6x2
+float2 HorizontalPCF5x2(float2 Fraction, float4 Values00, float4 Values20, float4 Values40)
+{
+	float Results0;
+	float Results1;
+
+	Results0 = Values00.w * (1.0 - Fraction.x);
+	Results1 = Values00.x * (1.0 - Fraction.x);
+	Results0 += Values00.z;
+	Results1 += Values00.y;
+	Results0 += Values20.w;
+	Results1 += Values20.x;
+	Results0 += Values20.z;
+	Results1 += Values20.y;
+	Results0 += Values40.w;
+	Results1 += Values40.x;
+	Results0 += Values40.z * Fraction.x;
+	Results1 += Values40.y * Fraction.x;
+
+	return float2(Results0, Results1);
+}
+
+// lowest quality ith PCF
+float PCF1x1(float2 Fraction, float4 Values00)
+{
+	float2 HorizontalLerp00 = lerp(Values00.wx, Values00.zy, Fraction.xx);
+
+	return lerp(HorizontalLerp00.x, HorizontalLerp00.y, Fraction.y);
+}
+
+float4 CalculateOcclusion(float4 ShadowmapDepth, FPCFSamplerSettings Settings)
+{
+	if (Settings.bSubsurface)
+	{
+		// Determine the distance that the light traveled through the subsurface object
+		// This assumes that anything between this subsurface pixel and the light was also a subsurface material,
+		// As a result, subsurface materials receive leaked light based on how transparent they are
+		float4 Thickness = max(Settings.SceneDepth - ShadowmapDepth, 0);
+		float4 Occlusion = saturate(exp(-Thickness * Settings.DensityMulConstant));
+		// Never shadow from depths that were never written to (max depth value)
+		return ShadowmapDepth > .99f ? 1 : Occlusion;
+	}
+	else
+	{
+		// The standard comparison is SceneDepth < ShadowmapDepth
+		// Using a soft transition based on depth difference
+		// Offsets shadows a bit but reduces self shadowing artifacts considerably
+		float TransitionScale = Settings.TransitionScale; //SoftTransitionScale.z;
+		float4 ShadowFactor = saturate((ShadowmapDepth - Settings.SceneDepth) * TransitionScale + 1);
+
+		FLATTEN
+		if (Settings.bTreatMaxDepthUnshadowed)
+		{
+			ShadowFactor = saturate(ShadowFactor + (ShadowmapDepth > .99f));
+		}
+
+		return ShadowFactor;
+	}
+}
+
+float3 CalculateOcclusion(float3 ShadowmapDepth, FPCFSamplerSettings Settings)
+{
+	if (Settings.bSubsurface)
+	{
+		// Determine the distance that the light traveled through the subsurface object
+		// This assumes that anything between this subsurface pixel and the light was also a subsurface material,
+		// As a result, subsurface materials receive leaked light based on how transparent they are
+		float3 Thickness = max(Settings.SceneDepth - (ShadowmapDepth - Settings.ProjectionDepthBiasParameters.x), 0);
+		float3 Occlusion = saturate(exp(-Thickness * Settings.DensityMulConstant));
+		// Never shadow from depths that were never written to (max depth value)
+		return ShadowmapDepth > .99f ? 1 : Occlusion;
+	}
+	else
+	{
+		// The standard comparison is Settings.SceneDepth < ShadowmapDepth
+		// Using a soft transition based on depth difference
+		// Offsets shadows a bit but reduces self shadowing artifacts considerably
+		float TransitionScale = Settings.TransitionScale; //SoftTransitionScale.z;
+		float3 ShadowFactor = saturate((ShadowmapDepth - Settings.SceneDepth) * TransitionScale + 1);
+
+		FLATTEN
+		if (Settings.bTreatMaxDepthUnshadowed)
+		{
+			ShadowFactor = saturate(ShadowFactor + (ShadowmapDepth > .99f));
+		}
+
+		return ShadowFactor;
+	}
+}
+
+void FetchRowOfThree(float2 Sample00TexelCenter, float VerticalOffset, out float3 Values0, FPCFSamplerSettings Settings)
+{
+	Values0.x = Texture2DSampleLevel(Settings.ShadowDepthTexture, Settings.ShadowDepthTextureSampler, (Sample00TexelCenter + float2(0, VerticalOffset)) * Settings.ShadowBufferSize.zw, 0).r;
+	Values0.y = Texture2DSampleLevel(Settings.ShadowDepthTexture, Settings.ShadowDepthTextureSampler, (Sample00TexelCenter + float2(1, VerticalOffset)) * Settings.ShadowBufferSize.zw, 0).r;
+	Values0.z = Texture2DSampleLevel(Settings.ShadowDepthTexture, Settings.ShadowDepthTextureSampler, (Sample00TexelCenter + float2(2, VerticalOffset)) * Settings.ShadowBufferSize.zw, 0).r;	
+	Values0 = CalculateOcclusion(Values0, Settings);
+}
+
+void FetchRowOfFour(float2 Sample00TexelCenter, float VerticalOffset, out float4 Values0, FPCFSamplerSettings Settings)
+{
+	Values0.x = Texture2DSampleLevel(Settings.ShadowDepthTexture, Settings.ShadowDepthTextureSampler, (Sample00TexelCenter + float2(0, VerticalOffset)) * Settings.ShadowBufferSize.zw, 0).r;
+	Values0.y = Texture2DSampleLevel(Settings.ShadowDepthTexture, Settings.ShadowDepthTextureSampler, (Sample00TexelCenter + float2(1, VerticalOffset)) * Settings.ShadowBufferSize.zw, 0).r;
+	Values0.z = Texture2DSampleLevel(Settings.ShadowDepthTexture, Settings.ShadowDepthTextureSampler, (Sample00TexelCenter + float2(2, VerticalOffset)) * Settings.ShadowBufferSize.zw, 0).r;
+	Values0.w = Texture2DSampleLevel(Settings.ShadowDepthTexture, Settings.ShadowDepthTextureSampler, (Sample00TexelCenter + float2(3, VerticalOffset)) * Settings.ShadowBufferSize.zw, 0).r;
+	Values0 = CalculateOcclusion(Values0, Settings);
+}
+
+void FetchRowOfThreeAfterFour(float2 Sample00TexelCenter, float VerticalOffset, out float3 Values1, FPCFSamplerSettings Settings)
+{
+	Values1.x = Texture2DSampleLevel(Settings.ShadowDepthTexture, Settings.ShadowDepthTextureSampler, (Sample00TexelCenter + float2(4, VerticalOffset)) * Settings.ShadowBufferSize.zw, 0).r;
+	Values1.y = Texture2DSampleLevel(Settings.ShadowDepthTexture, Settings.ShadowDepthTextureSampler, (Sample00TexelCenter + float2(5, VerticalOffset)) * Settings.ShadowBufferSize.zw, 0).r;
+	Values1.z = Texture2DSampleLevel(Settings.ShadowDepthTexture, Settings.ShadowDepthTextureSampler, (Sample00TexelCenter + float2(6, VerticalOffset)) * Settings.ShadowBufferSize.zw, 0).r;
+	Values1 = CalculateOcclusion(Values1, Settings);
+}
+
+// break this out for forward rendering as it's not part of ManualPCF's set of shadowquality settings.
+float Manual2x2PCF(float2 ShadowPosition, FPCFSamplerSettings Settings)
+{
+	float2 TexelPos = ShadowPosition * Settings.ShadowBufferSize.xy - 0.5f;	// bias to be consistent with texture filtering hardware
+	float2 Fraction = frac(TexelPos);
+	float2 TexelCenter = floor(TexelPos) + 0.5f;	// bias to get reliable texel center content
+	
+	float3 SamplesValues0, SamplesValues1, SamplesValues2;
+
+	FetchRowOfThree(TexelCenter, 0, SamplesValues0, Settings);
+	FetchRowOfThree(TexelCenter, 1, SamplesValues1, Settings);
+	FetchRowOfThree(TexelCenter, 2, SamplesValues2, Settings);
+	
+	return PCF2x2(Fraction, SamplesValues0, SamplesValues1, SamplesValues2);
+}
+
+float ManualNoFiltering(float2 ShadowPosition, FPCFSamplerSettings Settings)
+{
+	// very low quality but very good performance, useful to profile, 1 sample, not using gather4
+	return CalculateOcclusion(Texture2DSampleLevel(Settings.ShadowDepthTexture, Settings.ShadowDepthTextureSampler, ShadowPosition, 0).rrr, Settings).r;
+}
+
+float Manual1x1PCF(float2 ShadowPosition, FPCFSamplerSettings Settings)
+{
+	float2 TexelPos = ShadowPosition * Settings.ShadowBufferSize.xy - 0.5f;	// bias to be consistent with texture filtering hardware
+	float2 Fraction = frac(TexelPos);
+	float2 TexelCenter = floor(TexelPos) + 0.5f;	// bias to get reliable texel center content
+
+	// using Gather: xyzw in counter clockwise order starting with the sample to the lower left of the queried location
+	float4 Samples;
+
+#if FEATURE_GATHER4
+	Samples = Settings.ShadowDepthTexture.Gather(Settings.ShadowDepthTextureSampler, TexelCenter * Settings.ShadowBufferSize.zw);
+#else
+	Samples.x = Texture2DSampleLevel(Settings.ShadowDepthTexture, Settings.ShadowDepthTextureSampler, (TexelCenter.xy + float2(0, 1)) * Settings.ShadowBufferSize.zw, 0).r;
+	Samples.y = Texture2DSampleLevel(Settings.ShadowDepthTexture, Settings.ShadowDepthTextureSampler, (TexelCenter.xy + float2(1, 1)) * Settings.ShadowBufferSize.zw, 0).r;
+	Samples.z = Texture2DSampleLevel(Settings.ShadowDepthTexture, Settings.ShadowDepthTextureSampler, (TexelCenter.xy + float2(1, 0)) * Settings.ShadowBufferSize.zw, 0).r;
+	Samples.w = Texture2DSampleLevel(Settings.ShadowDepthTexture, Settings.ShadowDepthTextureSampler, (TexelCenter.xy + float2(0, 0)) * Settings.ShadowBufferSize.zw, 0).r;
+#endif
+
+	float4 Values00 = CalculateOcclusion(Samples, Settings);
+	return PCF1x1(Fraction, Values00);
+}
+
+
+float ManualPCF(float2 ShadowPosition, FPCFSamplerSettings Settings)
+{
+
+#if SHADOW_QUALITY == 1
+	return ManualNoFiltering(ShadowPosition, Settings);
+#endif
+
+#if SHADOW_QUALITY == 2
+	// low quality, 2x2 samples, using and not using gather4
+	return Manual1x1PCF(ShadowPosition, Settings);
+#endif
+
+#if SHADOW_QUALITY == 3
+	// medium quality, 4x4 samples, using and not using gather4
+	{
+		float2 TexelPos = ShadowPosition * Settings.ShadowBufferSize.xy - 0.5f;	// bias to be consistent with texture filtering hardware
+		float2 Fraction = frac(TexelPos);
+		float2 TexelCenter = floor(TexelPos) + 0.5f;	// bias to get reliable texel center content
+		{
+			float2 Sample00TexelCenter = TexelCenter - float2(1, 1);
+
+			float4 SampleValues0, SampleValues1, SampleValues2, SampleValues3;
+
+#if FEATURE_GATHER4
+			float2 SamplePos = TexelCenter * Settings.ShadowBufferSize.zw;	// bias to get reliable texel center content
+			SampleValues0 = CalculateOcclusion(Settings.ShadowDepthTexture.Gather(ShadowDepthTextureSampler, SamplePos, int2(-1, -1)), Settings);
+			SampleValues1 = CalculateOcclusion(Settings.ShadowDepthTexture.Gather(ShadowDepthTextureSampler, SamplePos, int2(1, -1)), Settings);
+			SampleValues2 = CalculateOcclusion(Settings.ShadowDepthTexture.Gather(ShadowDepthTextureSampler, SamplePos, int2(-1, 1)), Settings);
+			SampleValues3 = CalculateOcclusion(Settings.ShadowDepthTexture.Gather(ShadowDepthTextureSampler, SamplePos, int2(1, 1)), Settings);
+			return PCF3x3gather(Fraction, SampleValues0, SampleValues1, SampleValues2, SampleValues3);
+#else // FEATURE_GATHER4 
+			FetchRowOfFour(Sample00TexelCenter, 0, SampleValues0, Settings);
+			FetchRowOfFour(Sample00TexelCenter, 1, SampleValues1, Settings);
+			FetchRowOfFour(Sample00TexelCenter, 2, SampleValues2, Settings);
+			FetchRowOfFour(Sample00TexelCenter, 3, SampleValues3, Settings);
+			return PCF3x3(Fraction, SampleValues0, SampleValues1, SampleValues2, SampleValues3);
+#endif // FEATURE_GATHER4
+		}
+	}
+#endif
+
+#if FEATURE_GATHER4
+	// high quality, 6x6 samples, using gather4
+	{
+		float2 TexelPos = ShadowPosition * Settings.ShadowBufferSize.xy - 0.5f;	// bias to be consistent with texture filtering hardware
+		float2 Fraction = frac(TexelPos);
+		float2 TexelCenter = floor(TexelPos);
+		float2 SamplePos = (TexelCenter + 0.5f) * Settings.ShadowBufferSize.zw;	// bias to get reliable texel center content
+
+		float Results;
+
+		float4 Values00 = CalculateOcclusion(Settings.ShadowDepthTexture.Gather(Settings.ShadowDepthTextureSampler, SamplePos, int2(-2, -2)), Settings);
+		float4 Values20 = CalculateOcclusion(Settings.ShadowDepthTexture.Gather(Settings.ShadowDepthTextureSampler, SamplePos, int2(0, -2)), Settings);
+		float4 Values40 = CalculateOcclusion(Settings.ShadowDepthTexture.Gather(Settings.ShadowDepthTextureSampler, SamplePos, int2(2, -2)), Settings);
+
+		float2 Row0 = HorizontalPCF5x2(Fraction, Values00, Values20, Values40);
+		Results = Row0.x * (1.0f - Fraction.y) + Row0.y;
+
+		float4 Values02 = CalculateOcclusion(Settings.ShadowDepthTexture.Gather(Settings.ShadowDepthTextureSampler, SamplePos, int2(-2, 0)), Settings);
+		float4 Values22 = CalculateOcclusion(Settings.ShadowDepthTexture.Gather(Settings.ShadowDepthTextureSampler, SamplePos, int2(0, 0)), Settings);
+		float4 Values42 = CalculateOcclusion(Settings.ShadowDepthTexture.Gather(Settings.ShadowDepthTextureSampler, SamplePos, int2(2, 0)), Settings);
+
+		float2 Row1 = HorizontalPCF5x2(Fraction, Values02, Values22, Values42);
+		Results += Row1.x + Row1.y;
+
+		float4 Values04 = CalculateOcclusion(Settings.ShadowDepthTexture.Gather(Settings.ShadowDepthTextureSampler, SamplePos, int2(-2, 2)), Settings);
+		float4 Values24 = CalculateOcclusion(Settings.ShadowDepthTexture.Gather(Settings.ShadowDepthTextureSampler, SamplePos, int2(0, 2)), Settings);
+		float4 Values44 = CalculateOcclusion(Settings.ShadowDepthTexture.Gather(Settings.ShadowDepthTextureSampler, SamplePos, int2(2, 2)), Settings);
+
+		float2 Row2 = HorizontalPCF5x2(Fraction, Values04, Values24, Values44);
+		Results += Row2.x + Row2.y * Fraction.y;
+
+		return 0.04f * Results;
+	}
+
+#else // FEATURE_GATHER4
+
+	// high quality, 7x7 samples, not using gather4 (todo: ideally we make this 6x6 to get same results with gather code)
+	{
+		float2 Fraction = frac(ShadowPosition * Settings.ShadowBufferSize.xy);
+		float2 Sample00TexelCenter = floor(ShadowPosition * Settings.ShadowBufferSize.xy) - float2(3, 3);
+
+		// Fetch 7x7 shadowmap point samples
+		// Do 6x6 PCF samples, sharing the point samples between neighboring PCF samples
+		float4 Results;
+
+		float4 SampleValues03;
+		float4 SampleValues13;
+
+		{
+			float4 SampleValues10;
+			float4 SampleValues11;
+			float4 SampleValues12;
+
+			// Group work to minimize temporary registers and to split texture work with PCF ALU operations to hide texture latency
+			// Without this layout (all texture lookups at the beginning, PCF ALU's at the end) this shader was 4x slower on Nvidia cards
+			{
+				float4 SampleValues00;
+				FetchRowOfFour(Sample00TexelCenter, 0, SampleValues00, Settings);
+				SampleValues10.x = SampleValues00.w;
+
+				float4 SampleValues01;
+				FetchRowOfFour(Sample00TexelCenter, 1, SampleValues01, Settings);
+				SampleValues11.x = SampleValues01.w;
+
+				float4 SampleValues02;
+				FetchRowOfFour(Sample00TexelCenter, 2, SampleValues02, Settings);
+				SampleValues12.x = SampleValues02.w;
+
+				FetchRowOfFour(Sample00TexelCenter, 3, SampleValues03, Settings);
+				SampleValues13.x = SampleValues03.w;
+				Results.x = PCF3x3(Fraction, SampleValues00, SampleValues01, SampleValues02, SampleValues03);
+			}
+
+			{
+				FetchRowOfThreeAfterFour(Sample00TexelCenter, 0, SampleValues10.yzw, Settings);
+				FetchRowOfThreeAfterFour(Sample00TexelCenter, 1, SampleValues11.yzw, Settings);
+				FetchRowOfThreeAfterFour(Sample00TexelCenter, 2, SampleValues12.yzw, Settings);
+				FetchRowOfThreeAfterFour(Sample00TexelCenter, 3, SampleValues13.yzw, Settings);
+				Results.y = PCF3x3(Fraction, SampleValues10, SampleValues11, SampleValues12, SampleValues13);
+			}
+		}
+
+		{
+			float4 SampleValues14;
+			float4 SampleValues15;
+			float4 SampleValues16;
+
+			{
+				float4 SampleValues04;
+				FetchRowOfFour(Sample00TexelCenter, 4, SampleValues04, Settings);
+				SampleValues14.x = SampleValues04.w;
+
+				float4 SampleValues05;
+				FetchRowOfFour(Sample00TexelCenter, 5, SampleValues05, Settings);
+				SampleValues15.x = SampleValues05.w;
+
+				float4 SampleValues06;
+				FetchRowOfFour(Sample00TexelCenter, 6, SampleValues06, Settings);
+				SampleValues16.x = SampleValues06.w;
+
+				Results.z = PCF3x3(Fraction, SampleValues03, SampleValues04, SampleValues05, SampleValues06);
+			}
+
+			{
+				FetchRowOfThreeAfterFour(Sample00TexelCenter, 4, SampleValues14.yzw, Settings);
+				FetchRowOfThreeAfterFour(Sample00TexelCenter, 5, SampleValues15.yzw, Settings);
+				FetchRowOfThreeAfterFour(Sample00TexelCenter, 6, SampleValues16.yzw, Settings);
+				Results.w = PCF3x3(Fraction, SampleValues13, SampleValues14, SampleValues15, SampleValues16);
+			}
+		}
+
+		return dot(Results, .25f);
+	}
+#endif	// FEATURE_GATHER4
+}
diff --git a/Engine/Shaders/Private/SlateShaderCommon.ush b/Engine/Shaders/Private/SlateShaderCommon.ush
new file mode 100644
--- /dev/null
+++ b/Engine/Shaders/Private/SlateShaderCommon.ush
@@ -0,0 +1,131 @@
+// Copyright 1998-2017 Epic Games, Inc. All Rights Reserved.
+
+#if USE_MATERIALS
+#include "/Engine/Generated/Material.ush"
+
+	#if NUM_MATERIAL_TEXCOORDS
+		#define NUM_SLATE_TEXCOORDS NUM_MATERIAL_TEXCOORDS
+	#else
+		#define NUM_SLATE_TEXCOORDS 1
+	#endif // NUM_MATERIAL_TEXCOORDS
+
+	#define HAS_MATERIAL_TEXCOORDS NUM_MATERIAL_TEXCOORDS > 0
+#else
+	#define HAS_MATERIAL_TEXCOORDS 0
+	// This is defined automatically based on the number of textures in a material
+	// In the case that this is a basic texture lookup shader, we only need 1 uv set
+	#define NUM_SLATE_TEXCOORDS 1
+#endif
+
+// On mobile renderers we don't convert textures into linear space when slate is rendering.
+// since the final output is a gamma space render target.
+// Metal iOS, Android OpenGL ES 3.1 and Vulkan support sampling from sRGB textures
+#define SOURCE_IN_LINEAR_SPACE  USE_MATERIALS || (FEATURE_LEVEL != FEATURE_LEVEL_ES2) // ES2 does blending in sRGB space
+
+
+/** Shader types (mirrored from ESlateShader::Type in RenderingCommon.h */
+#define ST_Default		0
+#define ST_Border		1
+#define ST_Font			2
+#define ST_LineSegment	3
+#define ST_Custom		4
+
+/** Color Vision Deficiency Types (mirrored in EditorUserSettings.h) */
+#define CVD_NormalVision	0
+#define CVD_Deuteranomly	1
+#define CVD_Deuteranopia	2
+#define CVD_Protanomly		3
+#define CVD_Protanopia		4
+#define CVD_Tritanomaly		5
+#define CVD_Tritanopia		6
+#define CVD_Achromatopsia	7
+
+struct VertexToPixelInterpolants
+{
+	float4 Position : SV_POSITION;
+	half4 Color : COLOR0;
+	float4 OriginalPosition : ORIGINAL_POSITION;
+	float2 MaterialTexCoords : TEXCOORD0;
+	float4 TextureCoordinates[NUM_SLATE_TEXCOORDS] : TEXCOORD1;
+};
+
+float2 GetUV(VertexToPixelInterpolants Interpolants, uint UVIndex)
+{
+	float4 UVVector = Interpolants.TextureCoordinates[UVIndex / 2];
+	return Mod(UVIndex, 2) ? UVVector.zw : UVVector.xy;
+}
+
+half4 ApplyColorVisionDeficiencyFilter(half4 Color)
+{
+	half4 OutColor = Color;
+#if COLOR_VISION_DEFICIENCY_TYPE == CVD_Deuteranomly
+	OutColor.rgb = half3(
+		OutColor.r * 0.80000 + OutColor.g * 0.20000 + OutColor.b * 0.0,
+		OutColor.r * 0.25833 + OutColor.g * 0.74167 + OutColor.b * 0.0,
+		OutColor.r * 0.0     + OutColor.g * 0.14167 + OutColor.b * 0.85833);
+#elif COLOR_VISION_DEFICIENCY_TYPE == CVD_Deuteranopia
+	OutColor.rgb = half3(
+		OutColor.r * 0.625 + OutColor.g * 0.375 + OutColor.b * 0.0,
+		OutColor.r * 0.700 + OutColor.g * 0.300 + OutColor.b * 0.0,
+		OutColor.r * 0.0   + OutColor.g * 0.300 + OutColor.b * 0.700);
+#elif COLOR_VISION_DEFICIENCY_TYPE == CVD_Protanomly
+	OutColor.rgb = half3(
+		OutColor.r * 0.81667 + OutColor.g * 0.18333 + OutColor.b * 0.0,
+		OutColor.r * 0.33333 + OutColor.g * 0.66667 + OutColor.b * 0.0,
+		OutColor.r * 0.0     + OutColor.g * 0.12500 + OutColor.b * 0.87500);
+#elif COLOR_VISION_DEFICIENCY_TYPE == CVD_Protanopia
+	OutColor.rgb = half3(
+		OutColor.r * 0.56667 + OutColor.g * 0.43333 + OutColor.b * 0.0,
+		OutColor.r * 0.55833 + OutColor.g * 0.44167 + OutColor.b * 0.0,
+		OutColor.r * 0.0     + OutColor.g * 0.24167 + OutColor.b * 0.75833);
+#elif COLOR_VISION_DEFICIENCY_TYPE == CVD_Tritanomaly
+	OutColor.rgb = half3(
+		OutColor.r * 0.96667 + OutColor.g * 0.03333 + OutColor.b * 0.0,
+		OutColor.r * 0.0     + OutColor.g * 0.73333 + OutColor.b * 0.26667,
+		OutColor.r * 0.0     + OutColor.g * 0.18333 + OutColor.b * 0.81667);
+#elif COLOR_VISION_DEFICIENCY_TYPE == CVD_Tritanopia
+	OutColor.rgb = half3(
+		OutColor.r * 0.950 + OutColor.g * 0.050 + OutColor.b * 0.0,
+		OutColor.r * 0.0   + OutColor.g * 0.433 + OutColor.b * 0.567,
+		OutColor.r * 0.0   + OutColor.g * 0.475 + OutColor.b * 0.525);
+#elif COLOR_VISION_DEFICIENCY_TYPE == CVD_Achromatopsia
+	OutColor.rgb = half3(
+		OutColor.r * 0.299 + OutColor.g * 0.587 + OutColor.b * 0.114,
+		OutColor.r * 0.299 + OutColor.g * 0.587 + OutColor.b * 0.114,
+		OutColor.r * 0.299 + OutColor.g * 0.587 + OutColor.b * 0.114);
+#endif
+	return OutColor;
+}
+
+float cross(float2 a, float2 b)
+{
+	return a.x*b.y - a.y*b.x;
+}
+
+/**
+ * Given a point p and a parallelogram defined by point a and vectors b and c, determines in p is inside the parallelogram. 
+ * returns a 4-vector that can be used with the clip instruction.
+ */
+float4 PointInParallelogram(float2 p, float2 a, float4 bc)
+{
+	// unoptomized form:
+	//float2 o = p - a;
+	//float2 b = bc.xy;
+	//float2 c = bc.zw;
+	//float d = cross(b, c);
+	//float s = -cross(o, b) / d;
+	//float t = cross(o, c) / d;
+	// test for s and t between 0 and 1
+	//return float4(s, 1 - s, t, 1 - t);
+
+	float2 o = p - a;
+	// precompute 1/d
+	float invD = 1/cross(bc.xy, bc.zw);
+	// Compute an optimized o x b and o x c, leveraging that b and c are in the same vector register already (and free swizzles):
+	//   (o.x * b .y  - o.y * b .x, o.x *  c.y - o.y *  c.x) ==
+	//   (o.x * bc.y  - o.y * bc.x, o.x * bc.w - o.y * bc.z) ==
+	//    o.x * bc.yw - o.y * bc.xz
+	float2 st = (o.x * bc.yw - o.y * bc.xz) * float2(-invD, invD);
+	// test for s and t between 0 and 1
+	return float4(st, float2(1,1) - st);
+}
diff --git a/Engine/Shaders/Private/SurfelMaterialShared.ush b/Engine/Shaders/Private/SurfelMaterialShared.ush
new file mode 100644
--- /dev/null
+++ b/Engine/Shaders/Private/SurfelMaterialShared.ush
@@ -0,0 +1,9 @@
+// Copyright 1998-2017 Epic Games, Inc. All Rights Reserved.
+
+/*=============================================================================
+	SurfelMaterialShared.usf
+=============================================================================*/
+
+// In float4's, must match C++
+#define SURFEL_DATA_STRIDE 4
+#define INTERPOLATED_VERTEX_DATA_STRIDE 6
diff --git a/Engine/Shaders/Private/Tessellation.ush b/Engine/Shaders/Private/Tessellation.ush
new file mode 100644
--- /dev/null
+++ b/Engine/Shaders/Private/Tessellation.ush
@@ -0,0 +1,271 @@
+// Copyright 1998-2017 Epic Games, Inc. All Rights Reserved.
+
+/*=============================================================================
+	Tessellation.usf: Shader #include dispatch for tessellation
+=============================================================================*/
+
+/** Returns the texture coordinate for the vertex. */
+float2 GetTextureCoordinateDS( FPassSpecificVSToDS Vertex )
+{
+	return VertexFactoryGetTextureCoordinateDS( Vertex.FactoryInterpolants );
+}
+
+/** Returns the normal for the vertex. */
+float4 GetNormalDS( float3x3 TangentToWorld )
+{
+	return float4(TangentToWorld[2], 0);
+}
+
+
+/** Returns the tangent for the vertex. */
+float3 GetTangentDS( float3x3 TangentToWorld )
+{
+	return TangentToWorld[0];
+}
+
+/** 
+ * Optionally use Dominant Edge and Corner method to avoid cracks on corners and edges due to displacement map offsets.
+ * Dominant Edges are determined dynamically from AEN data, by order on UV
+ * Dominant Corners are determined offline and appended to the AEN data
+ */
+#if DISPLACEMENT_ANTICRACK
+
+	struct FHullShaderConstantDominantEdgeData
+	{
+		float2 UV0;
+		float2 UV1;	
+		float4 Normal0;
+		float4 Normal1;
+		float3 Tangent0;
+		float3 Tangent1;
+	};
+
+	struct FHullShaderConstantDominantVertexData
+	{
+		float2 UV;
+		float4 Normal;
+		float3 Tangent;
+	};
+
+#endif
+
+#if HULLSHADER
+// Converts the Local space scales to tangent space (note flipped order on mul)
+float3 GetTangentSpaceNonUniformScales(float3x3 TangentToLocal)
+{
+	return mul(TangentToLocal, Primitive.NonUniformScale.xyz);
+}
+
+float GetTessellationMultiplier(FPassSpecificVSToDS I)
+{
+	FMaterialTessellationParameters MaterialParameters = GetMaterialTessellationParameters(I.FactoryInterpolants, I.Position.xyz);
+	float MaterialTessellationMultiplier = GetMaterialTessellationMultiplier(MaterialParameters);
+#if VERTEX_FACTORY_MODIFIES_TESSELLATION
+	return VertexFactoryModifyTessellationMultiplier(I.FactoryInterpolants, I.Position.xyz, MaterialTessellationMultiplier);
+#else
+	return MaterialTessellationMultiplier;
+#endif
+}
+
+// Flag set if inside frustum
+int3 GetClipFlag( float3 Origin, float Radius )
+{
+	//float3 ClipExpand = Radius * float3( ResolvedView.ViewToClip[0][0], ResolvedView.ViewToClip[1][1], ResolvedView.ViewToClip[2][2] );
+	float4 ClipExpand = mul( float4( Radius, Radius, Radius, 0 ), ResolvedView.ViewToClip );
+	float4 ClipPos = mul( float4( Origin, 1.0 ), ResolvedView.TranslatedWorldToClip );
+    int3 Clip0 = ClipPos.xyz - ClipExpand.xyz <  ClipPos.w + ClipExpand.w;
+    int3 Clip1 = ClipPos.xyz + ClipExpand.xyz > -ClipPos.w - ClipExpand.w;
+    return Clip0 + 2 * Clip1;
+}
+
+// This is a heuristic that maps a bounding sphere on relevant edges/primitive to post projection space and uses that to scale the tessellation on the edges/primitive.
+//
+// Returns float4 where:
+// X - 0->1 Edge tessellation factor
+// Y - 1->2 Edge tessellation factor
+// Z - 2->0 Edge tessellation factor
+// W - inside tessellation factor
+float4 CalculateCompositeTessellationFactors(float3 Control0, float3 Control1, float3 Control2)
+{
+#if USE_ADAPTIVE_TESSELLATION_FACTOR
+#if 1
+	half MaxDisplacement = GetMaterialMaxDisplacement();
+
+	// Frustum cull
+	int3 ClipFlag = 0;
+	ClipFlag  = GetClipFlag( Control0, MaxDisplacement );
+	ClipFlag |= GetClipFlag( Control1, MaxDisplacement );
+	ClipFlag |= GetClipFlag( Control2, MaxDisplacement );
+	if( any( ClipFlag != 3 ) )
+	{
+		return 0;
+	}
+#endif
+
+	float3 Edge0 = ( Control0 - Control1 );
+	float3 Edge1 = ( Control1 - Control2 );
+	float3 Edge2 = ( Control2 - Control0 );
+
+	float3 ToMidpoint0 = 0.5 * ( Control0 + Control1 ) - ResolvedView.TranslatedWorldCameraOrigin;
+	float3 ToMidpoint1 = 0.5 * ( Control1 + Control2 ) - ResolvedView.TranslatedWorldCameraOrigin;
+	float3 ToMidpoint2 = 0.5 * ( Control2 + Control0 ) - ResolvedView.TranslatedWorldCameraOrigin;
+
+	// Use spherical projection instead of planer
+	float4 CompositeFactors = float4(
+		sqrt( dot( Edge1, Edge1 ) / dot( ToMidpoint1, ToMidpoint1 ) ),
+		sqrt( dot( Edge2, Edge2 ) / dot( ToMidpoint2, ToMidpoint2 ) ),
+		sqrt( dot( Edge0, Edge0 ) / dot( ToMidpoint0, ToMidpoint0 ) ),
+		1 );
+	CompositeFactors.w = 0.333 * ( CompositeFactors.x + CompositeFactors.y + CompositeFactors.z );
+
+	// The adaptive tessellation factor is 0.5 * ResolvedView.ViewToClip[1][1] * ViewSizeY / PixelsPerEdge and CompositeFactors is 2 * PercentageOfScreen.  
+	return View.AdaptiveTessellationFactor * CompositeFactors;
+#else
+	return float4( 1.0,1.0,1.0,1.0 );
+#endif
+}
+
+#if DISPLACEMENT_ANTICRACK
+
+	/**
+	 * Use a heuristic to determine which Edge (point pair) will be the dominant Edge. 
+	 * The heuristic should be consistent across the entire mesh.
+	 * Current heuristic is the smallest UV from the control points.
+	 */
+	FHullShaderConstantDominantEdgeData GenerateDominantEdgeData(FPassSpecificVSToDS Edge0Vert0,FPassSpecificVSToDS Edge0Vert1,FPassSpecificVSToDS Edge1Vert0,FPassSpecificVSToDS Edge1Vert1)
+	{
+		FHullShaderConstantDominantEdgeData Result = (FHullShaderConstantDominantEdgeData)0;
+		FPassSpecificVSToDS DominantVert0 = Edge0Vert0;
+		FPassSpecificVSToDS DominantVert1 = Edge0Vert1;
+
+		if ( Edge1Vert0.VertexID < Edge0Vert0.VertexID ||
+			(Edge1Vert0.VertexID == Edge0Vert0.VertexID && Edge1Vert1.VertexID < Edge0Vert1.VertexID ) )
+		{
+			DominantVert0 = Edge1Vert0;	
+			DominantVert1 = Edge1Vert1;
+		}
+
+		#if NUM_MATERIAL_TEXCOORDS
+			Result.UV0 = GetTextureCoordinateDS( DominantVert0 );
+			Result.UV1 = GetTextureCoordinateDS( DominantVert1 );
+		#endif
+
+		float3x3 TangentBasis0 = VertexFactoryGetTangentToLocalDS(DominantVert0.FactoryInterpolants);
+		float3x3 TangentBasis1 = VertexFactoryGetTangentToLocalDS(DominantVert1.FactoryInterpolants);
+
+		Result.Normal0 = GetNormalDS( TangentBasis0 );
+		Result.Normal1 = GetNormalDS( TangentBasis1 );
+
+		Result.Tangent0 = GetTangentDS( TangentBasis0 );
+		Result.Tangent1 = GetTangentDS( TangentBasis1 );
+
+		return Result;
+	}
+
+	FHullShaderConstantDominantVertexData GenerateDominantVertexData(FPassSpecificVSToDS DominantVertex)
+	{
+		FHullShaderConstantDominantVertexData Result = (FHullShaderConstantDominantVertexData)0;
+
+		#if NUM_MATERIAL_TEXCOORDS
+			Result.UV = GetTextureCoordinateDS( DominantVertex );
+		#endif
+
+		float3x3 TangentToLocal = VertexFactoryGetTangentToLocalDS(DominantVertex.FactoryInterpolants);
+
+		Result.Normal = GetNormalDS( TangentToLocal );
+		Result.Tangent = GetTangentDS( TangentToLocal );
+		
+		return Result;
+	}
+	
+#endif	// DISPLACEMENT_ANTICRACK
+
+#endif // HULLSHADER
+
+#if DOMAINSHADER
+
+#if DISPLACEMENT_ANTICRACK
+void ApplyDominantData(
+	inout float2 DisplacementTexCoords,
+	inout float3x3 TangentToWorld,
+	float4 DisplacementNormal,
+	float3 DisplacementTangent,
+	float U,
+	float V,
+	float W,
+	FHullShaderConstantDominantEdgeData DE0,
+	FHullShaderConstantDominantEdgeData DE1,
+	FHullShaderConstantDominantEdgeData DE2,
+	FHullShaderConstantDominantVertexData DV0,
+	FHullShaderConstantDominantVertexData DV1,
+	FHullShaderConstantDominantVertexData DV2 )
+{
+	// some handy vars for our state
+	float EdgeU = (U == 0) ? 1 : 0;
+	float EdgeV = (V == 0) ? 1 : 0;
+	float EdgeW = (W == 0) ? 1 : 0;
+	float Corner = ((EdgeU + EdgeV + EdgeW) == 2) ? 1 : 0;
+	float Edge   = ((EdgeU + EdgeV + EdgeW) == 1) ? 1 : 0;
+	float Center = ((EdgeU + EdgeV + EdgeW) == 0) ? 1 : 0;
+
+	if(Corner==1)
+	{
+		// Corner masks
+		float CornerU = ((EdgeV + EdgeW) == 2) ? 1 : 0;
+		float CornerV = ((EdgeW + EdgeU) == 2) ? 1 : 0;
+		float CornerW = ((EdgeU + EdgeV) == 2) ? 1 : 0;
+
+		// Override the tex coord with the dominant vert
+		DisplacementTexCoords =	CornerU * DV0.UV +
+								CornerV * DV1.UV +
+								CornerW * DV2.UV;
+									 
+		// Override normal/tangent with dominant vert 
+		DisplacementNormal = CornerU * DV0.Normal +
+							 CornerV * DV1.Normal +
+							 CornerW * DV2.Normal;
+
+		DisplacementTangent = CornerU * DV0.Tangent +
+							  CornerV * DV1.Tangent +
+							  CornerW * DV2.Tangent;
+	}
+	else if(Edge)
+	{
+		// interpolant change based on which Edge we are on
+		#define CONDITIONAL_EDGE_INTERPOLATE(Data0, Data1) EdgeU*(V*Data0 + W*Data1) + EdgeV*(W*Data0 + U*Data1) + EdgeW*(U*Data0 + V*Data1)
+
+		// override tex coords interpolated along dominant Edge
+		float2 UV0 = EdgeU*DE0.UV0 + EdgeV*DE1.UV0 + EdgeW*DE2.UV0;
+		float2 UV1 = EdgeU*DE0.UV1 + EdgeV*DE1.UV1 + EdgeW*DE2.UV1;
+
+		DisplacementTexCoords = CONDITIONAL_EDGE_INTERPOLATE(UV0,UV1);
+
+		// pull out dominant normals depending on which Edge
+		float4 Normal0 = EdgeU*DE0.Normal0 + EdgeV*DE1.Normal0 + EdgeW*DE2.Normal0;
+		float4 Normal1 = EdgeU*DE0.Normal1 + EdgeV*DE1.Normal1 + EdgeW*DE2.Normal1;
+
+		// interpolate change based on which Edge we are on
+		DisplacementNormal = CONDITIONAL_EDGE_INTERPOLATE(Normal0,Normal1);
+
+		// pull out dominant normals depending on which Edge
+		float3 Tangent0 = EdgeU*DE0.Tangent0 + EdgeV*DE1.Tangent0 + EdgeW*DE2.Tangent0;
+		float3 Tangent1 = EdgeU*DE0.Tangent1 + EdgeV*DE1.Tangent1 + EdgeW*DE2.Tangent1;
+
+		// interpolate change based on which Edge we are on
+		DisplacementTangent = CONDITIONAL_EDGE_INTERPOLATE(Tangent0,Tangent1);
+	}
+
+	if(Center==0)	// Edge or Corner case will overwrite normal and tangent, thus recalc tangent2world
+	{
+		TangentToWorld = AssembleTangentToWorld( DisplacementTangent, DisplacementNormal );
+	}
+}
+#endif	// DISPLACEMENT_ANTICRACK
+
+#endif // DOMAINSHADER
+
+#if TESSELLATION_TYPE_FLAT
+#include "FlatTessellation.ush"
+#elif TESSELLATION_TYPE_PNTRIANGLES
+#include "PNTriangles.ush"
+#endif
diff --git a/Engine/Shaders/Private/TonemapCommon.ush b/Engine/Shaders/Private/TonemapCommon.ush
new file mode 100644
--- /dev/null
+++ b/Engine/Shaders/Private/TonemapCommon.ush
@@ -0,0 +1,375 @@
+// Copyright 1998-2017 Epic Games, Inc. All Rights Reserved.
+
+/*=============================================================================
+	TonemapCommon.usf: PostProcessing tone mapping common
+=============================================================================*/
+
+#include "ACES.ush"
+#include "GammaCorrectionCommon.ush"
+
+// usually 1/2.2, the .y is used for inverse gamma when "gamma only" mode is not used
+half3 InverseGamma;
+
+// Film tonal and color control.
+half4 ColorMatrixR_ColorCurveCd1;
+half4 ColorMatrixG_ColorCurveCd3Cm3;
+half4 ColorMatrixB_ColorCurveCm2;
+half4 ColorCurve_Cm0Cd0_Cd2_Ch0Cm1_Ch3;
+half4 ColorCurve_Ch1_Ch2;
+half4 ColorShadow_Luma;
+half4 ColorShadow_Tint1;
+half4 ColorShadow_Tint2;
+
+half3 FilmPostProcess(half3 LinearColor) 
+{
+	// Color and exposure control.
+	half3 MatrixColor;
+	#if USE_COLOR_MATRIX == 1
+		// Apply color matrix (channel mixer, exposure, saturation).
+		MatrixColor.r = dot(LinearColor, ColorMatrixR_ColorCurveCd1.rgb);
+		MatrixColor.g = dot(LinearColor, ColorMatrixG_ColorCurveCd3Cm3.rgb);
+		MatrixColor.b = dot(LinearColor, ColorMatrixB_ColorCurveCm2.rgb);
+		#if USE_SHADOW_TINT == 1
+			MatrixColor *= ColorShadow_Tint1.rgb + ColorShadow_Tint2.rgb * rcp(dot(LinearColor, ColorShadow_Luma.rgb) + 1.0);
+		#endif
+		// Required to insure saturation doesn't create negative colors!
+		MatrixColor = max(half3(0.0, 0.0, 0.0), MatrixColor);
+	#else
+		// Less expensive route when not using saturation and channel mixer.
+		#if USE_SHADOW_TINT == 1
+			MatrixColor = LinearColor * (ColorShadow_Tint1.rgb + ColorShadow_Tint2.rgb * rcp(dot(LinearColor, ColorShadow_Luma.rgb) + 1.0));
+		#else
+			MatrixColor = LinearColor * ColorMatrixB_ColorCurveCm2.rgb;
+		#endif
+	#endif
+	// Apply color curve (includes tonemapping).
+	#if USE_CONTRAST == 1
+		// Full path.
+		half3 MatrixColorD = max(0, ColorCurve_Cm0Cd0_Cd2_Ch0Cm1_Ch3.xxx - MatrixColor);
+		half3 MatrixColorH = max(MatrixColor, ColorCurve_Cm0Cd0_Cd2_Ch0Cm1_Ch3.zzz);
+		half3 MatrixColorM = clamp(MatrixColor, ColorCurve_Cm0Cd0_Cd2_Ch0Cm1_Ch3.xxx, ColorCurve_Cm0Cd0_Cd2_Ch0Cm1_Ch3.zzz);
+		half3 CurveColor = 
+			// Highlights
+			(MatrixColorH*ColorCurve_Ch1_Ch2.xxx + ColorCurve_Ch1_Ch2.yyy) * rcp(MatrixColorH + ColorCurve_Cm0Cd0_Cd2_Ch0Cm1_Ch3.www) + 
+				// Midtones
+				((MatrixColorM*ColorMatrixB_ColorCurveCm2.aaa + 
+					// Darks
+					((MatrixColorD*ColorMatrixR_ColorCurveCd1.aaa) * rcp(MatrixColorD + ColorCurve_Cm0Cd0_Cd2_Ch0Cm1_Ch3.yyy) + ColorMatrixG_ColorCurveCd3Cm3.aaa)));
+	#else
+		// This is for mobile, it assumes color is not negative.
+		// Fast path when contrast=1, can remove the dark part of the curve.
+		half3 MatrixColorH = max(MatrixColor, ColorCurve_Cm0Cd0_Cd2_Ch0Cm1_Ch3.zzz);
+		half3 MatrixColorM = min(MatrixColor, ColorCurve_Cm0Cd0_Cd2_Ch0Cm1_Ch3.zzz);
+		half3 CurveColor = (MatrixColorH*ColorCurve_Ch1_Ch2.xxx + ColorCurve_Ch1_Ch2.yyy) * rcp(MatrixColorH + ColorCurve_Cm0Cd0_Cd2_Ch0Cm1_Ch3.www) + MatrixColorM;
+	#endif
+
+	// Must hit black by log10(-2.474)
+	CurveColor -= 0.002;
+
+	return CurveColor;
+}
+
+half3 TonemapAndGammaCorrect(half3 LinearColor)
+{
+	// Clamp input to be positive
+	// This displays negative colors as black
+	LinearColor = max(LinearColor, 0);
+
+	half3 GammaColor = pow(LinearColor, InverseGamma.x);
+	
+	// in all cases it's good to clamp into the 0..1 range (e.g for LUT color grading)
+	GammaColor = saturate(GammaColor);
+
+	return GammaColor;
+}
+
+
+/*
+============================================
+// Uncharted settings
+Slope = 0.63;
+Toe = 0.55;
+Shoulder = 0.47;
+BlackClip= 0;
+WhiteClip = 0.01;
+
+// HP settings
+Slope = 0.65;
+Toe = 0.63;
+Shoulder = 0.45;
+BlackClip = 0;
+WhiteClip = 0;
+
+// Legacy settings
+Slope = 0.98;
+Toe = 0.3;
+Shoulder = 0.22;
+BlackClip = 0;
+WhiteClip = 0.025;
+
+// ACES settings
+Slope = 0.91;
+Toe = 0.53;
+Shoulder = 0.23;
+BlackClip = 0;
+WhiteClip = 0.035;
+===========================================
+*/
+float FilmSlope = 0.91;
+float FilmToe = 0.53;
+float FilmShoulder = 0.23;
+float FilmBlackClip = 0;
+float FilmWhiteClip = 0.035;
+
+half3 FilmToneMap( half3 LinearColor ) 
+{
+	const float3x3 sRGB_2_AP0 = mul( XYZ_2_AP0_MAT, mul( D65_2_D60_CAT, sRGB_2_XYZ_MAT ) );
+	const float3x3 sRGB_2_AP1 = mul( XYZ_2_AP1_MAT, mul( D65_2_D60_CAT, sRGB_2_XYZ_MAT ) );
+	const float3x3 AP1_2_sRGB = mul( XYZ_2_sRGB_MAT, mul( D60_2_D65_CAT, AP1_2_XYZ_MAT ) );
+	
+#if 0
+	float3 ACESColor = mul( sRGB_2_AP0, float3(LinearColor) );
+
+	// --- Red modifier --- //
+	const float RRT_RED_SCALE = 0.82;
+	const float RRT_RED_PIVOT = 0.03;
+	const float RRT_RED_HUE = 0;
+	const float RRT_RED_WIDTH = 135;
+
+	float saturation = rgb_2_saturation( ACESColor );
+	float hue = rgb_2_hue( ACESColor );
+	float centeredHue = center_hue( hue, RRT_RED_HUE );
+	float hueWeight = Square( smoothstep( 0, 1, 1 - abs( 2 * centeredHue / RRT_RED_WIDTH ) ) );
+		
+	ACESColor.r += hueWeight * saturation * (RRT_RED_PIVOT - ACESColor.r) * (1. - RRT_RED_SCALE);
+
+	// Use ACEScg primaries as working space
+	float3 WorkingColor = mul( AP0_2_AP1_MAT, ACESColor );
+#else
+	// Use ACEScg primaries as working space
+	float3 WorkingColor = mul( sRGB_2_AP1, float3(LinearColor) );
+#endif
+
+	WorkingColor = max( 0, WorkingColor );
+
+	// Pre desaturate
+	WorkingColor = lerp( dot( WorkingColor, AP1_RGB2Y ), WorkingColor, 0.96 );
+	
+	const half ToeScale			= 1 + FilmBlackClip - FilmToe;
+	const half ShoulderScale	= 1 + FilmWhiteClip - FilmShoulder;
+	
+	const float InMatch = 0.18;
+	const float OutMatch = 0.18;
+
+	float ToeMatch;
+	if( FilmToe > 0.8 )
+	{
+		// 0.18 will be on straight segment
+		ToeMatch = ( 1 - FilmToe  - OutMatch ) / FilmSlope + log10( InMatch );
+	}
+	else
+	{
+		// 0.18 will be on toe segment
+
+		// Solve for ToeMatch such that input of InMatch gives output of OutMatch.
+		const float bt = ( OutMatch + FilmBlackClip ) / ToeScale - 1;
+		ToeMatch = log10( InMatch ) - 0.5 * log( (1+bt)/(1-bt) ) * (ToeScale / FilmSlope);
+	}
+
+	float StraightMatch = ( 1 - FilmToe ) / FilmSlope - ToeMatch;
+	float ShoulderMatch = FilmShoulder / FilmSlope - StraightMatch;
+	
+	half3 LogColor = log10( WorkingColor );
+	half3 StraightColor = FilmSlope * ( LogColor + StraightMatch );
+	
+	half3 ToeColor		= (    -FilmBlackClip ) + (2 *      ToeScale) / ( 1 + exp( (-2 * FilmSlope /      ToeScale) * ( LogColor -      ToeMatch ) ) );
+	half3 ShoulderColor	= ( 1 + FilmWhiteClip ) - (2 * ShoulderScale) / ( 1 + exp( ( 2 * FilmSlope / ShoulderScale) * ( LogColor - ShoulderMatch ) ) );
+
+	ToeColor		= LogColor <      ToeMatch ?      ToeColor : StraightColor;
+	ShoulderColor	= LogColor > ShoulderMatch ? ShoulderColor : StraightColor;
+
+	half3 t = saturate( ( LogColor - ToeMatch ) / ( ShoulderMatch - ToeMatch ) );
+	t = ShoulderMatch < ToeMatch ? 1 - t : t;
+	t = (3-2*t)*t*t;
+	half3 ToneColor = lerp( ToeColor, ShoulderColor, t );
+
+	// Post desaturate
+	ToneColor = lerp( dot( float3(ToneColor), AP1_RGB2Y ), ToneColor, 0.93 );
+
+	// Returning positive AP1 values
+	return max( 0, ToneColor );
+
+	//ToneColor = mul( AP1_2_sRGB, ToneColor );
+
+	//return saturate( ToneColor );
+	//return max( 0, ToneColor );
+}
+
+half3 FilmToneMapInverse( half3 ToneColor ) 
+{
+	const float3x3 sRGB_2_AP1 = mul( XYZ_2_AP1_MAT, mul( D65_2_D60_CAT, sRGB_2_XYZ_MAT ) );
+	const float3x3 AP1_2_sRGB = mul( XYZ_2_sRGB_MAT, mul( D60_2_D65_CAT, AP1_2_XYZ_MAT ) );
+	
+	// Use ACEScg primaries as working space
+	half3 WorkingColor = mul( sRGB_2_AP1, saturate( ToneColor ) );
+
+	WorkingColor = max( 0, WorkingColor );
+	
+	// Post desaturate
+	WorkingColor = lerp( dot( WorkingColor, AP1_RGB2Y ), WorkingColor, 1.0 / 0.93 );
+
+	half3 ToeColor		= 0.374816 * pow( 0.9 / min( WorkingColor, 0.8 ) - 1, -0.588729 );
+	half3 ShoulderColor	= 0.227986 * pow( 1.56 / ( 1.04 - WorkingColor ) - 1, 1.02046 );
+
+	half3 t = saturate( ( WorkingColor - 0.35 ) / ( 0.45 - 0.35 ) );
+	t = (3-2*t)*t*t;
+	half3 LinearColor = lerp( ToeColor, ShoulderColor, t );
+
+	// Pre desaturate
+	LinearColor = lerp( dot( LinearColor, AP1_RGB2Y ), LinearColor, 1.0 / 0.96 );
+
+	LinearColor = mul( AP1_2_sRGB, LinearColor );
+
+	// Returning positive sRGB values
+	return max( 0, LinearColor );
+}
+
+//
+// ACES sRGB D65 Output Transform - Forward and Inverse
+//  Input is scene-referred linear values in the sRGB gamut
+//  Output is output-referred linear values in the sRGB gamut
+//
+float3 ACESOutputTransformsRGBD65( float3 SceneReferredLinearsRGBColor )
+{
+	const float3x3 sRGB_2_AP0 = mul( XYZ_2_AP0_MAT, mul( D65_2_D60_CAT, sRGB_2_XYZ_MAT ) );
+
+	float3 aces = mul( sRGB_2_AP0, SceneReferredLinearsRGBColor * 1.5 );
+	float3 oces = RRT( aces );
+	float3 OutputReferredLinearsRGBColor =  ODT_sRGB_D65( oces );
+	return OutputReferredLinearsRGBColor;
+}
+
+float3 InverseACESOutputTransformsRGBD65( float3 OutputReferredLinearsRGBColor )
+{
+	const float3x3 AP0_2_sRGB = mul( XYZ_2_sRGB_MAT, mul( D60_2_D65_CAT, AP0_2_XYZ_MAT ) );
+
+	float3 oces = Inverse_ODT_sRGB_D65( OutputReferredLinearsRGBColor );
+	float3 aces = Inverse_RRT( oces );
+	float3 SceneReferredLinearsRGBColor = mul( AP0_2_sRGB, aces ) * 0.6666;
+
+	return SceneReferredLinearsRGBColor;
+}
+
+//
+// ACES D65 1000 nit Output Transform - Forward
+//  Input is scene-referred linear values in the sRGB gamut
+//  Output is output-referred linear values in the AP1 gamut
+//
+float3 ACESOutputTransforms1000( float3 SceneReferredLinearsRGBColor )
+{
+	const float3x3 sRGB_2_AP0 = mul( XYZ_2_AP0_MAT, mul( D65_2_D60_CAT, sRGB_2_XYZ_MAT ) );
+
+	float3 aces = mul( sRGB_2_AP0, SceneReferredLinearsRGBColor * 1.5 );
+	float3 oces = RRT( aces );
+	float3 OutputReferredLinearAP1Color = ODT_1000nits( oces );
+	return OutputReferredLinearAP1Color;
+}
+
+//
+// ACES D65 2000 nit Output Transform - Forward
+//  Input is scene-referred linear values in the sRGB gamut
+//  Output is output-referred linear values in the AP1 gamut
+//
+float3 ACESOutputTransforms2000( float3 SceneReferredLinearsRGBColor )
+{
+	const float3x3 sRGB_2_AP0 = mul( XYZ_2_AP0_MAT, mul( D65_2_D60_CAT, sRGB_2_XYZ_MAT ) );
+
+	float3 aces = mul( sRGB_2_AP0, SceneReferredLinearsRGBColor * 1.5 );
+	float3 oces = RRT( aces );
+	float3 OutputReferredLinearAP1Color = ODT_2000nits( oces );
+	return OutputReferredLinearAP1Color;
+}
+
+static const float3x3 GamutMappingIdentityMatrix = { 1.0, 0.0, 0.0, 0.0, 1.0, 0.0, 0.0, 0.0, 1.0 };
+
+//
+// Gamut conversion matrices
+//
+float3x3 OuputGamutMappingMatrix( uint OutputGamut )
+{
+	// Gamut mapping matrices used later
+	const float3x3 AP1_2_sRGB    = mul( XYZ_2_sRGB_MAT, mul( D60_2_D65_CAT, AP1_2_XYZ_MAT ) );
+	const float3x3 AP1_2_DCI_D65 = mul( XYZ_2_P3D65_MAT, mul( D60_2_D65_CAT, AP1_2_XYZ_MAT ) );
+	const float3x3 AP1_2_Rec2020 = mul( XYZ_2_Rec2020_MAT, mul( D60_2_D65_CAT, AP1_2_XYZ_MAT ) );
+
+	// Set gamut mapping matrix 
+	// 0 = sRGB - D65
+	// 1 = P3 - D65
+	// 2 = Rec.2020 - D65
+	// 3 = ACES AP0 - D60
+	// 4 = ACES AP1 - D60
+
+	if( OutputGamut == 1 )
+		return AP1_2_DCI_D65;
+	else if( OutputGamut == 2 )
+		return AP1_2_Rec2020;
+	else if( OutputGamut == 3 )
+		return AP1_2_AP0_MAT;
+	else if( OutputGamut == 4 )
+		return GamutMappingIdentityMatrix;
+	else
+		return AP1_2_sRGB;
+}
+
+float3x3 OuputInverseGamutMappingMatrix( uint OutputGamut )
+{
+	// Gamut mapping matrices used later
+	const float3x3 sRGB_2_AP1    = mul( XYZ_2_AP1_MAT, mul( D65_2_D60_CAT, sRGB_2_XYZ_MAT ) );
+	const float3x3 DCI_D65_2_AP1 = mul( XYZ_2_AP1_MAT, mul( D65_2_D60_CAT, P3D65_2_XYZ_MAT ) );
+	const float3x3 Rec2020_2_AP1 = mul( XYZ_2_AP1_MAT, mul( D65_2_D60_CAT, Rec2020_2_XYZ_MAT ) );
+
+	// Set gamut mapping matrix 
+	// 0 = sRGB - D65
+	// 1 = P3 (DCI) - D65
+	// 2 = Rec.2020 - D65
+	// 3 = ACES AP0 - D60
+	float3x3 GamutMappingMatrix = sRGB_2_AP1;
+	if( OutputGamut == 1 )
+		GamutMappingMatrix = DCI_D65_2_AP1;
+	else if( OutputGamut == 2 )
+		GamutMappingMatrix = Rec2020_2_AP1;
+	else if( OutputGamut == 3 )
+		GamutMappingMatrix = AP0_2_AP1_MAT;
+
+	return GamutMappingMatrix;
+}
+
+float3 ST2084ToScRGB(float3 Color, uint OutputDevice)
+{
+	// Nvidia HDR encoding - Remove PQ, convert to linear scRGB
+	const float3x3 AP1_2_sRGB = mul(XYZ_2_sRGB_MAT, AP1_2_XYZ_MAT);
+	const float WhitePoint = 80.f;
+
+	// 1000.f nit display
+	float MaxODTNits = 1000.0f;
+	float MinODTNits = 0.0001f;
+
+	if (OutputDevice == 4 || OutputDevice == 6)
+	{
+		// 2000 nit display
+		MaxODTNits = 2000.0f;
+		MinODTNits = 0.005f;
+	}
+
+	float3 OutColor = ST2084ToLinear(Color);
+		
+	OutColor = clamp(OutColor, MinODTNits, MaxODTNits);
+	OutColor.x = Y_2_linCV(OutColor.x, MaxODTNits, MinODTNits);
+	OutColor.y = Y_2_linCV(OutColor.y, MaxODTNits, MinODTNits);
+	OutColor.z = Y_2_linCV(OutColor.z, MaxODTNits, MinODTNits);
+
+	float scRGBScale = MaxODTNits / WhitePoint;
+	OutColor = mul(AP1_2_sRGB, OutColor) * scRGBScale;
+
+	return OutColor;
+}
\ No newline at end of file
diff --git a/Engine/Shaders/Private/VelocityCommon.ush b/Engine/Shaders/Private/VelocityCommon.ush
new file mode 100644
--- /dev/null
+++ b/Engine/Shaders/Private/VelocityCommon.ush
@@ -0,0 +1,19 @@
+// Copyright 1998-2017 Epic Games, Inc. All Rights Reserved.
+
+/*=============================================================================
+	VelocityCommon.usf: Common functions for calculating velocity vectors.
+=============================================================================*/
+
+#pragma once
+
+float2 Calculate2DVelocity(float4 PackedVelocityA, float4 PackedVelocityC)
+{
+	float2 ScreenPos = PackedVelocityA.xy / PackedVelocityA.w - ResolvedView.TemporalAAJitter.xy;
+	float2 PrevScreenPos = PackedVelocityC.xy / PackedVelocityC.w - ResolvedView.TemporalAAJitter.zw;
+
+	// 2d velocity, includes camera an object motion
+	float2 Velocity = ScreenPos - PrevScreenPos;
+
+	// Make sure not to touch 0,0 which is clear color
+	return Velocity;
+}
diff --git a/Engine/Shaders/Private/VertexFactoryCommon.ush b/Engine/Shaders/Private/VertexFactoryCommon.ush
new file mode 100644
--- /dev/null
+++ b/Engine/Shaders/Private/VertexFactoryCommon.ush
@@ -0,0 +1,60 @@
+// Copyright 1998-2017 Epic Games, Inc. All Rights Reserved.
+
+#include "Common.ush"
+
+float3 TransformLocalToWorld(float3 LocalPosition)
+{
+	//expand the matrix multiply to avoid a code-gen bug on PS4.
+	//float3 RotatedPosition = mul(LocalPosition.xyz, (float3x3)Primitive.LocalToWorld); 
+	//invariant code!
+	return INVARIANT((Primitive.LocalToWorld[0].xyz * LocalPosition.xxx + Primitive.LocalToWorld[1].xyz * LocalPosition.yyy + Primitive.LocalToWorld[2].xyz * LocalPosition.zzz) + Primitive.LocalToWorld[3].xyz);	
+}
+
+float4 TransformLocalToTranslatedWorld(float3 LocalPosition)
+{
+	float3 RotatedPosition = Primitive.LocalToWorld[0].xyz * LocalPosition.xxx + Primitive.LocalToWorld[1].xyz * LocalPosition.yyy + Primitive.LocalToWorld[2].xyz * LocalPosition.zzz;
+	return float4(RotatedPosition + (Primitive.LocalToWorld[3].xyz + ResolvedView.PreViewTranslation.xyz),1);
+}
+
+#if VERTEXSHADER && FEATURE_LEVEL >= FEATURE_LEVEL_SM4 && ONEPASS_POINTLIGHT_SHADOW && USING_VERTEX_SHADER_LAYER
+	/** This is the instance count per-face, since we are use 6*N to render to all faces of the cube without a GS. */
+	uint InstanceCount;
+#endif
+
+uint GetLayerInstanceID( uint InstanceID )
+{
+#if VERTEXSHADER && ONEPASS_POINTLIGHT_SHADOW && USING_VERTEX_SHADER_LAYER
+	return InstanceID % InstanceCount;
+#else
+	return InstanceID;
+#endif
+}
+
+// Octahedron Normal Vectors
+// [Cigolle 2014, "A Survey of Efficient Representations for Independent Unit Vectors"]
+//						Mean	Max
+// oct		8:8			0.33709 0.94424
+// snorm	8:8:8		0.17015 0.38588
+// oct		10:10		0.08380 0.23467
+// snorm	10:10:10	0.04228 0.09598
+// oct		12:12		0.02091 0.05874
+
+float2 UnitToOct( float3 N )
+{
+	N.xy /= dot( 1, abs(N) );
+	if( N.z <= 0 )
+	{
+		N.xy = ( 1 - abs(N.yx) ) * ( N.xy >= 0 ? float2(1,1) : float2(-1,-1) );
+	}
+	return N.xy;
+}
+
+float3 OctToUnit( float2 Oct )
+{
+	float3 N = float3( Oct, 1 - dot( 1, abs(Oct) ) );
+	if( N.z < 0 )
+	{
+		N.xy = ( 1 - abs(N.yx) ) * ( N.xy >= 0 ? float2(1,1) : float2(-1,-1) );
+	}
+	return normalize(N);
+}
diff --git a/Engine/Shaders/Private/VirtualTextureCommon.ush b/Engine/Shaders/Private/VirtualTextureCommon.ush
new file mode 100644
--- /dev/null
+++ b/Engine/Shaders/Private/VirtualTextureCommon.ush
@@ -0,0 +1,78 @@
+// Copyright 1998-2017 Epic Games, Inc. All Rights Reserved.
+
+#include "Common.ush"
+
+float MipLevel2D( float2 UV, inout float2 dUVdx, inout float2 dUVdy )
+{
+	dUVdx = ddx( dUVdy );
+	dUVdy = ddy( dUVdy );
+	
+	float px = dot( dUVdx, dUVdx );
+	float py = dot( dUVdy, dUVdy );
+
+	float MinLevel = 0.5 * log2( min( px, py ) );
+	float MaxLevel = 0.5 * log2( max( px, py ) );
+	
+	float Level = MaxLevel;
+	return Level;
+}
+
+float MipLevelAniso2D( float2 UV, inout float2 dUVdx, inout float2 dUVdy, const float MaxAnisoLog2 )
+{
+	dUVdx = ddx( UV );
+	dUVdy = ddy( UV );
+	
+	float px = dot( dUVdx, dUVdx );
+	float py = dot( dUVdy, dUVdy );
+
+	float MinLevel = 0.5 * log2( min( px, py ) );
+	float MaxLevel = 0.5 * log2( max( px, py ) );
+	
+	float AnisoBias = min( MaxLevel - MinLevel, MaxAnisoLog2 );
+	float Level = MaxLevel - AnisoBias;
+	
+	return Level;
+}
+
+// Pack 16bit 6:6:4
+// 64 x 64 physical pages
+// 16 mip levels
+void PageTable2DTranslation16( Texture2D<uint> PageTable, uint PageTableSize, uint2 vPage, uint vLevel, out uint2 pPage, out uint vPagesWideInLevel )
+{
+	uint PageTableEntry = PageTable.Load( int3( vPage >> vLevel, vLevel ) );
+
+	pPage.x = ( PageTableEntry >>  0 ) & 0x3F;
+	pPage.y = ( PageTableEntry >>  6 ) & 0x3F;
+	vLevel  = ( PageTableEntry >> 12 ) & 0xF;
+
+	vPagesWideInLevel = PageTableSize >> vLevel;
+}
+
+// 8:8:8:8 page table format
+// [van Waveren 2012, "Software Virtual Textures"]
+// 256 x 256 physical pages
+void PageTable2DTranslation32( Texture2D PageTable, SamplerState PageTableSampler, float2 vUV, float vLevel, out float2 pPage, out float vPagesWideInLevel )
+{
+	float4 PageTableEntry = PageTable.SampleLevel( PageTableSampler, vUV, vLevel );
+	
+	pPage = PageTableEntry.xy * 255.0;
+	vPagesWideInLevel = dot( PageTableEntry.zw, float2( 1.0 / 16.0, 256.0 / 16.0 ) * 255.0 );
+}
+
+float2 VirtualToPhysical2D(	const float2 pTextureSize, const float vPageSize, const float pPageSize,
+							float2 pPage, float vPagesWideInLevel,
+							float2 vUV, inout float2 dUVdx, inout float2 dUVdy )
+{
+	const float2 pPageScale = pPageSize / pTextureSize;
+	const float2 vPageScale = vPageSize / pTextureSize;
+	const float2 vPageBias = ( 0.5 * ( pPageSize - vPageSize ) ) / pTextureSize;
+
+	float2 vPageFrac = frac( vUV * vPagesWideInLevel );
+	float2 pUV = vPageFrac * vPageScale + vPageBias + pPage.xy * pPageScale;
+	
+	float2 ddxyScale = vPageScale * vPagesWideInLevel;
+	dUVdx *= ddxyScale;
+	dUVdy *= ddxyScale;
+	
+	return pUV;
+}
\ No newline at end of file
diff --git a/Engine/Shaders/Private/VolumeLightingCommon.ush b/Engine/Shaders/Private/VolumeLightingCommon.ush
new file mode 100644
--- /dev/null
+++ b/Engine/Shaders/Private/VolumeLightingCommon.ush
@@ -0,0 +1,128 @@
+// Copyright 1998-2017 Epic Games, Inc. All Rights Reserved.
+
+/**
+ * VolumeLightingCommon.usf
+ */
+
+// use low quality shadow sampling on translucency for better performance
+#define SHADOW_QUALITY 2
+
+#include "ShadowProjectionCommon.ush"
+#include "ShadowFilteringCommon.ush"
+
+/** Parameters needed to access the shadow map of the light. */
+float4x4 WorldToShadowMatrix;
+float4 ShadowmapMinMax;
+// .x:1/SplitNearFadeRegion, .y:1/SplitFarFadeRegion .zw:DistanceFadeMAD
+float4 ShadowInjectParams;
+float2 DepthBiasParameters;
+
+/** Whether to compute static shadowing. */
+uint bStaticallyShadowed;
+
+/** Shadow depth map computed for static geometry by Lightmass. */
+Texture2D StaticShadowDepthTexture;
+SamplerState StaticShadowDepthTextureSampler;
+
+/** Transform used for static shadowing by spot and directional lights. */
+float4x4 WorldToStaticShadowMatrix;
+float4 StaticShadowBufferSize;
+
+/** Computes dynamic and static shadowing for a point anywhere in space. */
+float ComputeVolumeShadowing(float3 WorldPositionForLighting, bool bPointLight, bool bSpotLight)
+{
+	float ShadowFactor = 1;
+
+	BRANCH
+	if (bStaticallyShadowed)
+	{
+		bool bUsePointLightShadowing = bPointLight;
+
+		BRANCH
+		if (bUsePointLightShadowing)
+		{
+			float3 LightVector = WorldPositionForLighting - DeferredLightUniforms.LightPosition;
+			float DistanceToLight = length(LightVector);
+			float3 NormalizedLightVector = LightVector / DistanceToLight;
+
+			//@todo - use parametrization without slow inverse trig.  Dual paraboloid?
+			float NormalizedTheta = atan2(NormalizedLightVector.y, NormalizedLightVector.x) / (2 * PI);
+			// atan2 returns in the range [-PI, PI], wrap the negative portion to [.5, 1]
+			float U = NormalizedTheta > 0 ? NormalizedTheta : 1 + NormalizedTheta;
+			float V = acos(NormalizedLightVector.z) / PI;
+			float2 UnwrappedUVs = float2(U, V);
+
+			float ShadowDepth = Texture2DSampleLevel(StaticShadowDepthTexture, StaticShadowDepthTextureSampler, UnwrappedUVs, 0).x;
+			ShadowFactor = DistanceToLight * DeferredLightUniforms.LightInvRadius < ShadowDepth;
+		}
+		else
+		{
+			// This path is used for directional lights and spot lights, which only require a single projection
+			// Transform the world position into shadowmap space
+			float4 HomogeneousShadowPosition = mul(float4(WorldPositionForLighting, 1), WorldToStaticShadowMatrix);
+			float2 ShadowUVs = HomogeneousShadowPosition.xy / HomogeneousShadowPosition.w;
+
+			// Treat as unshadowed if the voxel is outside of the shadow map
+			if (all(ShadowUVs >= 0 && ShadowUVs <= 1))
+			{
+				FPCFSamplerSettings Settings;
+				Settings.ShadowDepthTexture = StaticShadowDepthTexture;
+				Settings.ShadowDepthTextureSampler = StaticShadowDepthTextureSampler;
+				Settings.ShadowBufferSize = StaticShadowBufferSize;
+				Settings.SceneDepth = HomogeneousShadowPosition.z;
+				Settings.TransitionScale = 40;
+				Settings.bSubsurface = false;
+				// We can sample outside of the static shadowmap, which is centered around the scene.  These 'infinite' depth values should not cause occlusion.
+				Settings.bTreatMaxDepthUnshadowed = true;
+				Settings.DensityMulConstant = 0;
+				Settings.ProjectionDepthBiasParameters = float2(0, 0);
+
+				ShadowFactor = Manual1x1PCF(ShadowUVs, Settings);
+
+				/*
+				// Sample the shadowmap depth and determine if this voxel is shadowed
+				float ShadowDepth = Texture2DSampleLevel(StaticShadowDepthTexture, StaticShadowDepthTextureSampler, ShadowUVs, 0).x;
+				ShadowFactor = HomogeneousShadowPosition.z < ShadowDepth;
+				*/
+			}
+		}
+	}		
+
+	#if DYNAMICALLY_SHADOWED
+		bool bUseCubemapShadowing = bPointLight;
+		float DynamicShadowFactor = 1;
+
+		if (bUseCubemapShadowing)
+		{
+			DynamicShadowFactor = CubemapHardwarePCF(WorldPositionForLighting, DeferredLightUniforms.LightPosition, DeferredLightUniforms.LightInvRadius, 0.03f * 512 * InvShadowmapResolution);
+		}
+		else
+		{
+			// Transform the world position into shadowmap space
+			float4 HomogeneousShadowPosition = mul(float4(WorldPositionForLighting, 1), WorldToShadowMatrix);
+			float2 ShadowUVs = HomogeneousShadowPosition.xy / HomogeneousShadowPosition.w;
+
+			// Treat as unshadowed if the voxel is outside of the shadow map
+			if (all(ShadowUVs >= ShadowmapMinMax.xy && ShadowUVs <= ShadowmapMinMax.zw))
+			{
+				// Sample the shadowmap depth and determine if this voxel is shadowed
+				float ShadowDepth = Texture2DSampleLevel(ShadowDepthTexture, ShadowDepthTextureSampler, ShadowUVs, 0).x;
+				DynamicShadowFactor = HomogeneousShadowPosition.z < ShadowDepth - DepthBiasParameters.x;
+			}
+		}
+
+		// fade shadows in the distance
+		if (!bPointLight && !bSpotLight)
+		{
+			float Depth = dot(WorldPositionForLighting - View.WorldCameraOrigin, View.ViewForward);
+			float DistanceFade = saturate(Depth * ShadowInjectParams.z + ShadowInjectParams.w);
+			DynamicShadowFactor = lerp(DynamicShadowFactor, 1.0f, DistanceFade * DistanceFade);
+		}
+
+		// Combine static shadowing and dynamic shadowing, important for stationary directional lights with CSM
+		ShadowFactor = min(ShadowFactor, DynamicShadowFactor);
+
+	#endif
+
+	return ShadowFactor;
+}
diff --git a/Engine/Shaders/Public/MetalCommon.ush b/Engine/Shaders/Public/MetalCommon.ush
new file mode 100644
--- /dev/null
+++ b/Engine/Shaders/Public/MetalCommon.ush
@@ -0,0 +1,20 @@
+// Copyright 1998-2017 Epic Games, Inc. All Rights Reserved.
+
+/*=============================================================================
+	MetalCommon.usf: Common Metal shader code
+=============================================================================*/
+
+// Update this GUID to improve shader recompilation for Metal only shaders
+// GUID = {36782442-E538-4A21-BEAA-7184F683FE5D}
+  
+#pragma once
+
+#if VERTEXSHADER && HULLSHADER
+#ifndef TESSELLATIONSHADER
+#define TESSELLATIONSHADER 1
+#endif
+#endif
+
+// Cross compiler doesn't support structured buffers.  This replacement works if it was using a builtin type.
+#define RWStructuredBuffer RWBuffer
+#define StructuredBuffer Buffer
diff --git a/Engine/Shaders/Public/Platform.ush b/Engine/Shaders/Public/Platform.ush
new file mode 100644
--- /dev/null
+++ b/Engine/Shaders/Public/Platform.ush
@@ -0,0 +1,457 @@
+// Copyright 1998-2017 Epic Games, Inc. All Rights Reserved.
+
+/*=============================================================================
+	Prefix.usf: USF file automatically included by shader preprocessor.
+=============================================================================*/
+
+#pragma once
+
+
+// ---------------------------------------------------- Profile or compiler specific includes
+// TODO: Have shader compiler including these platform specific USF files, that needs to work
+// with ShaderCore.cpp's GetShaderIncludes().
+
+#if PS4_PROFILE  
+	// always #include PS4Common.usf so it can #define override anything in any other included file.
+	#include "PS4/PS4Common.ush"
+#endif
+ 
+#if XBOXONE_PROFILE
+	#include "XboxOne/XboxOneCommon.ush"
+#endif
+
+#if COMPILER_METAL
+	// Helps with iteration when changing Metal shader code generation backend.
+	#include "MetalCommon.ush"
+#endif
+
+#if VULKAN_PROFILE
+	// Helps with iteration when changing Vulkan shader code generation backend.
+	#include "VulkanCommon.ush"
+#endif
+
+
+// ---------------------------------------------------- DDC invalidation
+// to support the console command "r.InvalidateShaderCache"
+#include "ShaderVersion.ush"
+
+
+// ---------------------------------------------------- COMPILE_* and *_PROFILE defaults
+
+#ifndef COMPILER_HLSLCC
+#define COMPILER_HLSLCC 0
+#endif
+
+#ifndef COMPILER_HLSL
+#define COMPILER_HLSL 0
+#endif
+
+#ifndef COMPILER_GLSL
+#define COMPILER_GLSL 0
+#endif
+
+#ifndef COMPILER_GLSL_ES2
+#define COMPILER_GLSL_ES2 0
+#endif
+
+#ifndef COMPILER_GLSL_ES3_1
+#define COMPILER_GLSL_ES3_1 0
+#endif
+
+#ifndef COMPILER_GLSL_ES3_1_EXT
+#define COMPILER_GLSL_ES3_1_EXT 0
+#endif
+
+#ifndef COMPILER_METAL
+#define COMPILER_METAL 0
+#endif
+
+#ifndef COMPILER_SUPPORTS_ATTRIBUTES
+#define COMPILER_SUPPORTS_ATTRIBUTES 0
+#endif
+
+#ifndef SM5_PROFILE
+#define SM5_PROFILE 0
+#endif
+
+#ifndef SM4_PROFILE
+#define SM4_PROFILE 0
+#endif
+
+#ifndef ES2_PROFILE
+#define ES2_PROFILE 0
+#endif
+
+#ifndef ES3_1_PROFILE
+#define ES3_1_PROFILE 0
+#endif
+
+#ifndef METAL_PROFILE
+#define METAL_PROFILE 0
+#endif
+
+#ifndef METAL_MRT_PROFILE
+#define METAL_MRT_PROFILE 0
+#endif
+
+#ifndef METAL_SM4_PROFILE
+#define METAL_SM4_PROFILE 0
+#endif
+
+#ifndef METAL_SM5_PROFILE
+#define METAL_SM5_PROFILE 0
+#endif
+
+#ifndef VULKAN_PROFILE
+#define	VULKAN_PROFILE 0
+#endif
+
+#ifndef VULKAN_PROFILE_SM4
+#define	VULKAN_PROFILE_SM4 0
+#endif
+
+#ifndef VULKAN_PROFILE_SM5
+#define	VULKAN_PROFILE_SM5 0
+#endif
+
+#ifndef IOS
+#define IOS 0
+#endif
+
+#ifndef MAC
+#define MAC 0
+#endif
+
+// 'static' asserts
+#if COMPILER_GLSL || COMPILER_GLSL_ES2 || COMPILER_GLSL_ES3_1 || VULKAN_PROFILE || METAL_PROFILE || METAL_MRT_PROFILE || METAL_SM4_PROFILE || METAL_SM5_PROFILE
+#if !COMPILER_HLSLCC
+#error "Missing COMPILER_HLSLCC define!"
+#endif
+#endif
+
+
+// ---------------------------------------------------- Alternative floating point types
+
+#ifndef FORCE_FLOATS
+#define FORCE_FLOATS 0
+#endif
+
+#if (!(COMPILER_GLSL_ES2 || COMPILER_GLSL_ES3_1) && !METAL_PROFILE) || FORCE_FLOATS
+	// Always use floats when not using the ES2 compiler, because low precision modifiers are currently only tweaked for ES2, 
+	// And we don't want potential side effects on other platforms
+	#define half float
+	#define half1 float1
+	#define half2 float2
+	#define half3 float3
+	#define half4 float4
+	#define half3x3 float3x3
+	#define half4x4 float4x4 
+	#define half4x3 float4x3 
+	#define fixed float
+	#define fixed1 float1
+	#define fixed2 float2
+	#define fixed3 float3
+	#define fixed4 float4
+	#define fixed3x3 float3x3
+	#define fixed4x4 float4x4
+	#define fixed4x3 float4x3
+#endif
+
+// ---------------------------------------------------- Profile config
+
+// Values of FEATURE_LEVEL.
+#define FEATURE_LEVEL_ES2	1
+#define FEATURE_LEVEL_ES3_1	2
+#define FEATURE_LEVEL_SM3	3
+#define FEATURE_LEVEL_SM4	4
+#define FEATURE_LEVEL_SM5	5
+#define FEATURE_LEVEL_MAX	6
+
+#if PS4_PROFILE
+	#define FEATURE_LEVEL FEATURE_LEVEL_SM5
+
+#elif SM5_PROFILE
+	// SM5 = full dx11 features (high end UE4 rendering)
+	#define FEATURE_LEVEL FEATURE_LEVEL_SM5
+
+#elif SM4_PROFILE
+	#define FEATURE_LEVEL FEATURE_LEVEL_SM4
+
+#elif SWITCH_PROFILE || SWITCH_PROFILE_FORWARD
+	#undef ES3_1_PROFILE
+
+	#if SWITCH_PROFILE
+		#define FEATURE_LEVEL FEATURE_LEVEL_SM5
+	#else
+		#define FEATURE_LEVEL FEATURE_LEVEL_ES3_1
+		// @todo switch: maybe all uses of this should check feature level not profile?
+		#define ES3_1_PROFILE 1
+	#endif
+
+
+	// A8 textures when sampled have their component in R.
+	#define A8_SAMPLE_MASK .r
+
+	// switch doesn't have vertex attribute swizzling
+	#define FCOLOR_COMPONENT_SWIZZLE .bgra
+
+#elif VULKAN_PROFILE
+	#define FEATURE_LEVEL FEATURE_LEVEL_ES3_1
+
+#elif VULKAN_PROFILE_SM4
+	#define FEATURE_LEVEL FEATURE_LEVEL_SM4
+
+#elif VULKAN_PROFILE_SM5
+	#define FEATURE_LEVEL FEATURE_LEVEL_SM5
+
+#elif GL3_PROFILE || GL4_PROFILE
+	#if GL3_PROFILE
+		#define FEATURE_LEVEL FEATURE_LEVEL_SM4
+	#elif GL4_PROFILE
+		#define FEATURE_LEVEL FEATURE_LEVEL_SM5
+	#endif
+
+	// A8 textures when sampled have their component in R.
+	#define A8_SAMPLE_MASK .r
+
+	// hacks until the shader compiler supports those
+	#if GL4_PROFILE
+		#define class struct
+	#endif
+
+#elif METAL_PROFILE
+	#define FEATURE_LEVEL  FEATURE_LEVEL_ES3_1
+	// @todo metal: remove this and make sure all uses handle METAL_PROFILE
+	#undef ES3_1_PROFILE
+	#define ES3_1_PROFILE 1
+	#define FCOLOR_COMPONENT_SWIZZLE .bgra
+	#define STENCIL_COMPONENT_SWIZZLE .x
+
+#elif METAL_ES2_PROFILE
+	#define FEATURE_LEVEL  FEATURE_LEVEL_ES2
+	// @todo metal: remove this and make sure all uses handle METAL_ES2_PROFILE
+	#undef ES2_PROFILE
+	#define ES2_PROFILE 1
+	#define FCOLOR_COMPONENT_SWIZZLE .bgra
+	#define STENCIL_COMPONENT_SWIZZLE .x
+
+#elif METAL_MRT_PROFILE
+	#define FEATURE_LEVEL FEATURE_LEVEL_SM4
+	#define FCOLOR_COMPONENT_SWIZZLE .bgra
+	#define STENCIL_COMPONENT_SWIZZLE .x
+
+#elif METAL_SM4_PROFILE
+	#define FEATURE_LEVEL FEATURE_LEVEL_SM4
+	#define FCOLOR_COMPONENT_SWIZZLE .bgra
+	#define STENCIL_COMPONENT_SWIZZLE .x
+
+#elif METAL_SM5_PROFILE
+	#define FEATURE_LEVEL FEATURE_LEVEL_SM5
+	#define FCOLOR_COMPONENT_SWIZZLE .bgra
+	#define STENCIL_COMPONENT_SWIZZLE .x
+
+#elif ES2_PROFILE || ES3_1_PROFILE
+	#if ES3_1_PROFILE
+		#define FEATURE_LEVEL FEATURE_LEVEL_ES3_1
+	#else
+		//@todo ES3_1 GL
+		#define FEATURE_LEVEL FEATURE_LEVEL_ES2
+	#endif
+
+	#if COMPILER_GLSL_ES2 || COMPILER_GLSL_ES3_1
+		// Swizzle as we only support GL_BGRA on non-ES2 platforms that have that extension
+		#define FCOLOR_COMPONENT_SWIZZLE .bgra  
+	#else 
+		#define FCOLOR_COMPONENT_SWIZZLE 
+		#if COMPILER_GLSL
+			// A8 textures when sampled have their component in R
+			#define A8_SAMPLE_MASK .r
+		#endif
+	#endif
+
+#else
+
+	#error Add your platform here
+
+	#define FEATURE_LEVEL FEATURE_LEVEL_MAX
+
+#endif
+
+
+// ---------------------------------------------------- Swizzle defaults
+
+// If we didn't request color component swizzling, just make it empty
+#ifndef FCOLOR_COMPONENT_SWIZZLE
+#define FCOLOR_COMPONENT_SWIZZLE
+#endif
+
+#ifndef STENCIL_COMPONENT_SWIZZLE
+#define STENCIL_COMPONENT_SWIZZLE .g
+#endif
+
+#ifndef A8_SAMPLE_MASK
+#define A8_SAMPLE_MASK .a
+#endif
+
+// ---------------------------------------------------- Platform dependent supports
+
+// non-editor platforms generally never want development/editor features.
+#define PLATFORM_SUPPORTS_DEVELOPMENT_SHADERS (!PS4_PROFILE && !XBOXONE_PROFILE && !ESDEFERRED_PROFILE && !SWITCH_PROFILE)
+
+#ifndef MOBILE_EMULATION
+#define MOBILE_EMULATION ((FEATURE_LEVEL == FEATURE_LEVEL_ES2 || FEATURE_LEVEL == FEATURE_LEVEL_ES3_1) && (!(COMPILER_GLSL_ES3_1 || COMPILER_GLSL_ES2) && USE_DEVELOPMENT_SHADERS && (!(METAL_PROFILE) || MAC) && !VULKAN_PROFILE && !SWITCH_PROFILE))
+#endif
+
+// Whether the platform supports independent texture and samplers
+// When enabled, different texture lookups can share samplers to allow more artist samplers in the base pass
+// Ideally this would just be enabled for all SM4 and above feature level platforms
+// @todo metal mrt: No reason this can't work with Metal, once cross compiler is fixed
+#define SUPPORTS_INDEPENDENT_SAMPLERS (PS4_PROFILE || SM5_PROFILE || SM4_PROFILE || METAL_MRT_PROFILE || METAL_SM4_PROFILE || METAL_SM5_PROFILE)
+
+// Whether the platform supports a global clip plane through SV_ClipDistance
+// Ideally this would just be enabled for all SM4 and above feature level platforms, but not tested everywhere yet
+#define PLATFORM_SUPPORTS_GLOBAL_CLIP_PLANE (PS4_PROFILE || SM5_PROFILE || SM4_PROFILE || METAL_PROFILE || METAL_MRT_PROFILE || METAL_SM4_PROFILE || METAL_SM5_PROFILE || GL4_PROFILE || GL3_PROFILE || VULKAN_PROFILE_SM4 || VULKAN_PROFILE_SM5)
+
+// Whether the platform support pixel coverage on MSAA targets (SV_Coverage).
+#define SUPPORTS_PIXEL_COVERAGE (FEATURE_LEVEL >= FEATURE_LEVEL_SM5 && !COMPILER_GLSL && !MOBILE_EMULATION)
+
+
+// ---------------------------------------------------- Compiler specific defaults and fallbacks
+
+/** Defined only for Metal's combined Vertex + Hull shader */
+#ifndef TESSELLATIONSHADER
+#define TESSELLATIONSHADER 0
+#endif
+
+// Hlslcc platforms ignore the uniform keyword as it can't properly optimize flow
+#if COMPILER_HLSLCC
+#define uniform
+#endif
+
+
+// ---------------------------------------------------- Compiler attributes
+
+#if SM5_PROFILE || COMPILER_SUPPORTS_ATTRIBUTES
+
+/** Avoids flow control constructs. */
+#define UNROLL [unroll]
+
+/** Gives preference to flow control constructs. */
+#define LOOP [loop]
+
+/** Performs branching by using control flow instructions like jmp and label. */
+#define BRANCH [branch]
+
+/** Performs branching by using the cnd instructions. */
+#define FLATTEN [flatten]
+
+/** Allows a compute shader loop termination condition to be based off of a UAV read. The loop must not contain synchronization intrinsics. */
+#define ALLOW_UAV_CONDITION [allow_uav_condition]
+
+#endif // SM5_PROFILE || COMPILER_SUPPORTS_ATTRIBUTES
+
+#if SM5_PROFILE
+#define EARLYDEPTHSTENCIL [earlydepthstencil]
+#endif
+
+
+// ---------------------------------------------------- Compiler attribute fallbacks
+
+#ifndef UNROLL
+#define UNROLL
+#endif
+
+#ifndef LOOP
+#define LOOP
+#endif
+
+#ifndef BRANCH
+#define BRANCH
+#endif
+
+#ifndef FLATTEN
+#define FLATTEN
+#endif
+
+#ifndef ALLOW_UAV_CONDITION	
+#define ALLOW_UAV_CONDITION	
+#endif
+
+#ifndef INVARIANT
+#define INVARIANT										
+#endif
+
+#ifndef ENABLE_RE_Z
+#define ENABLE_RE_Z
+#endif
+
+#ifndef EARLYDEPTHSTENCIL
+#define EARLYDEPTHSTENCIL
+#endif
+
+
+// ---------------------------------------------------- Interpolator attribute fallbacks
+
+#ifndef COMPRESSED_16_FLOAT
+#define COMPRESSED_16_FLOAT
+#endif
+
+#ifndef COMPRESSED_16_UNORM
+#define COMPRESSED_16_UNORM
+#endif
+
+#ifndef COMPRESSED_16_SNORM
+#define COMPRESSED_16_SNORM
+#endif
+
+#ifndef COMPRESSED_16_UINT
+#define COMPRESSED_16_UINT
+#endif
+
+#ifndef COMPRESSED_16_INT
+#define COMPRESSED_16_INT
+#endif
+
+#ifndef COMPRESSED_8_UNORM
+#define COMPRESSED_8_UNORM
+#endif
+
+#ifndef COMPRESSED_8_SNORM
+#define COMPRESSED_8_SNORM
+#endif
+
+#ifndef COMPRESSED_8_UINT
+#define COMPRESSED_8_UINT
+#endif
+
+
+// ---------------------------------------------------- Global uses
+
+#define USE_DEVELOPMENT_SHADERS (COMPILE_SHADERS_FOR_DEVELOPMENT && PLATFORM_SUPPORTS_DEVELOPMENT_SHADERS)
+
+
+// ---------------------------------------------------- Compiler missing implementations
+
+#if COMPILER_GLSL_ES2 || COMPILER_METAL || COMPILER_SWITCH
+
+	float determinant(float3x3 M)
+	{    
+		return	
+			M[0][0] * (M[1][1] * M[2][2] - M[1][2] * M[2][1]) -
+			M[1][0] * (M[0][1] * M[2][2] - M[0][2] * M[2][1]) +
+			M[2][0] * (M[0][1] * M[1][2] - M[0][2] * M[1][1]);
+	}
+
+#endif   
+
+
+#if COMPILER_HLSLCC
+#define ddx_fine(x) ddx(x)
+#define ddy_fine(y) ddy(y)
+#endif
+
+
+#if FEATURE_LEVEL < FEATURE_LEVEL_ES3_1 && !COMPILER_METAL
+	// DX11 (feature levels >= 10) feature sets natively supports uints in shaders; we just use floats on other platforms.
+	#define uint4	int4
+#endif
diff --git a/Engine/Shaders/Public/VulkanCommon.ush b/Engine/Shaders/Public/VulkanCommon.ush
new file mode 100644
--- /dev/null
+++ b/Engine/Shaders/Public/VulkanCommon.ush
@@ -0,0 +1,11 @@
+// Copyright 1998-2017 Epic Games, Inc. All Rights Reserved.
+
+/*=============================================================================
+	VulkanCommon.usf: Common Vulkan shader code
+=============================================================================*/
+
+// Update this GUID to improve shader recompilation for Vulkan only shaders
+// GUID = {792BDC0A-A4A8-402F-8752-3F9F642A03B6}
+// Glslang 1.0.49.0
+
+#pragma once
