Copying //UE4/Orion-Staging to //UE4/Main (Source: //Orion/Dev-General @ 3606378)

diff --git a/Engine/Shaders/Private/NiagaraMeshVertexFactory.usf b/Engine/Shaders/Private/NiagaraMeshVertexFactory.usf
new file mode 100644
--- /dev/null
+++ b/Engine/Shaders/Private/NiagaraMeshVertexFactory.usf
@@ -0,0 +1,623 @@
+// Copyright 1998-2017 Epic Games, Inc. All Rights Reserved.
+
+/*=============================================================================
+	MeshParticleVertexFactory.usf: Mesh particle vertex factory shader code.
+=============================================================================*/
+
+#include "VertexFactoryCommon.usf"
+#include "NiagaraVFParticleAccess.usf"
+
+// These are only enabled for SM5 due to limited interpolators elsewhere
+#define USE_PARTICLE_NODES (FEATURE_LEVEL >= FEATURE_LEVEL_SM5)
+
+// Only support previous transform for deferred feature level
+#define USE_PREVIOUS_TRANSFORM (FEATURE_LEVEL >= FEATURE_LEVEL_SM4)
+
+float4x4 PreviousLocalToWorld;
+
+
+#if USE_PREVIOUS_TRANSFORM
+	Buffer<float4> PrevTransformBuffer;
+#endif
+	
+
+struct FVertexFactoryInput
+{
+	float4	Position	: ATTRIBUTE0;
+	half3	TangentX	: ATTRIBUTE1;
+	// TangentZ.w contains sign of tangent basis determinant
+	half4	TangentZ	: ATTRIBUTE2;
+	half4	VertexColor : ATTRIBUTE3;
+
+#if NUM_MATERIAL_TEXCOORDS_VERTEX
+	float2	TexCoords[NUM_MATERIAL_TEXCOORDS_VERTEX] : ATTRIBUTE4;
+#elif USE_PARTICLE_SUBUVS
+	float2	TexCoords[1] : ATTRIBUTE4;
+#endif
+
+
+	uint InstanceId	: SV_InstanceID;
+};
+
+struct FVertexFactoryInterpolantsVSToPS
+{
+	TANGENTTOWORLD_INTERPOLATOR_BLOCK
+
+#if USE_PARTICLE_SUBUVS
+	float4 SubUV0AndTexCoord0 : TEXCOORD1;
+	float4 SubUV1AndLerp : TEXCOORD2;
+#else
+	#if NUM_MATERIAL_TEXCOORDS
+			#if (ES2_PROFILE || ES3_1_PROFILE)
+				// Avoid dependent texture fetches, put all UVs in xy
+				float2	TexCoords[NUM_MATERIAL_TEXCOORDS] : TEXCOORD0;
+			#else
+				float4	TexCoords[(NUM_MATERIAL_TEXCOORDS+1)/2]	: TEXCOORD0;
+			#endif
+	#endif
+#endif
+
+#if INTERPOLATE_VERTEX_COLOR
+	float4	VertexColor : COLOR0;
+#endif
+
+#if NEEDS_PARTICLE_COLOR
+	float4	ParticleColor : COLOR1;
+#endif
+
+#if USE_DYNAMIC_PARAMETERS
+	float4 DynamicParameter : COLOR2;
+#endif
+	
+#if USE_PARTICLE_NODES
+	/** Particle position in camera-centered translated world space */
+	float3 ParticleTranslatedWorldPosition	: PARTICLE_POSITION;
+
+	/** The velocity of the particle, XYZ: direction, W: speed. */
+	float4 ParticleVelocity	: PARTICLE_VELOCITY;
+
+	/** Relative alive time of the particle */
+	float ParticleTime : PARTICLE_TIME;
+#if NEEDS_PARTICLE_TRANSFORM
+	float4 LocalToWorld[4] : PARTICLE_TRANSFORM;
+#endif
+#endif
+
+#if INSTANCED_STEREO
+	nointerpolation uint PackedEyeIndex : PACKED_EYE_INDEX;
+#endif
+};
+
+struct FVertexFactoryIntermediates
+{
+	/** The color of the vertex. */
+	float4 VertexColor;
+
+	/** The color of the particle. */
+	float4 ParticleColor;
+
+	/** The texture coordinates for the vertex. */
+#if NUM_MATERIAL_TEXCOORDS_VERTEX
+	float2	TexCoords[NUM_MATERIAL_TEXCOORDS_VERTEX];
+#endif 
+#if USE_PARTICLE_SUBUVS
+	float4 SubUVCoords;
+	float SubUVLerp;
+#endif
+#if USE_DYNAMIC_PARAMETERS
+	/** Optional dynamic parameter for the particle. */
+	float4 DynamicParameter;
+#endif	//USE_DYNAMIC_PARAMETERS
+
+	/** The velocity of the particle, XYZ: direction, W: speed. */
+	float4 ParticlePosition;
+
+	/** The velocity of the particle, XYZ: direction, W: speed. */
+	float4 ParticleVelocity;
+
+	/** Particle position in camera-centered translated world space. */
+	float3 ParticleTranslatedWorldPosition;
+	
+	/** Relative time. */
+	float RelativeTime;
+
+	float4x4 LocalToWorld;
+};
+
+float4x4 GetParticleTransform(FVertexFactoryIntermediates Inters)
+{
+	//Transform is transposed on CPU and position is packed into W. Transpose again to get the correct transform back.
+	return transpose(Inters.LocalToWorld);
+}
+
+#if USE_PREVIOUS_TRANSFORM
+float4x4 GetParticlePrevTransform(FVertexFactoryInput Input)
+{
+	uint TransformBaseIdx = GetLayerInstanceID(Input.InstanceId) * 3;
+
+	float4x4 PrevTransform4x4 = float4x4(
+		PrevTransformBuffer[TransformBaseIdx + 0], 
+		PrevTransformBuffer[TransformBaseIdx + 1], 
+		PrevTransformBuffer[TransformBaseIdx + 2], 
+		float4(0, 0, 0, 1)
+		);
+
+	//Transform is transposed on CPU and position is packed into W. Transpose again to get the correct transform back.
+	return transpose(PrevTransform4x4);
+}
+#endif
+
+/** Converts from vertex factory specific interpolants to a FMaterialPixelParameters, which is used by material inputs. */
+FMaterialPixelParameters GetMaterialPixelParameters(FVertexFactoryInterpolantsVSToPS Interpolants, float4 SvPosition)
+{
+	// GetMaterialPixelParameters is responsible for fully initializing the result
+	FMaterialPixelParameters Result = MakeInitializedMaterialPixelParameters();
+
+#if USE_PARTICLE_SUBUVS
+	#if NUM_MATERIAL_TEXCOORDS
+		UNROLL
+		for( int CoordinateIndex = 0; CoordinateIndex < NUM_MATERIAL_TEXCOORDS; CoordinateIndex++ )
+		{
+			Result.TexCoords[CoordinateIndex] = Interpolants.SubUV0AndTexCoord0.zw;
+		}
+	#endif
+	Result.Particle.SubUVCoords[0] = Interpolants.SubUV0AndTexCoord0.xy;
+	Result.Particle.SubUVCoords[1] = Interpolants.SubUV1AndLerp.xy;
+	Result.Particle.SubUVLerp = Interpolants.SubUV1AndLerp.z;
+#elif NUM_MATERIAL_TEXCOORDS
+	#if (ES2_PROFILE || ES3_1_PROFILE)
+		UNROLL
+		for( int CoordinateIndex = 0; CoordinateIndex < NUM_MATERIAL_TEXCOORDS; CoordinateIndex++ )
+		{
+			Result.TexCoords[CoordinateIndex] = Interpolants.TexCoords[CoordinateIndex].xy;
+		}
+	#else
+		UNROLL
+		for(int CoordinateIndex = 0;CoordinateIndex < NUM_MATERIAL_TEXCOORDS;CoordinateIndex += 2)
+		{
+			Result.TexCoords[CoordinateIndex] = Interpolants.TexCoords[CoordinateIndex/2].xy;
+
+			if(CoordinateIndex + 1 < NUM_MATERIAL_TEXCOORDS)
+			{
+				Result.TexCoords[CoordinateIndex + 1] = Interpolants.TexCoords[CoordinateIndex/2].wz;
+			}
+		}
+	#endif
+#endif
+
+	half3 TangentToWorld0 = Interpolants.TangentToWorld0.xyz;
+	half4 TangentToWorld2 = Interpolants.TangentToWorld2;
+	Result.UnMirrored = TangentToWorld2.w;
+
+#if INTERPOLATE_VERTEX_COLOR
+	Result.VertexColor = Interpolants.VertexColor;
+#else
+	Result.VertexColor = 0;
+#endif
+
+#if NEEDS_PARTICLE_COLOR
+	Result.Particle.Color = Interpolants.ParticleColor;
+#endif
+
+#if USE_DYNAMIC_PARAMETERS
+	Result.Particle.DynamicParameter = Interpolants.DynamicParameter;
+#else
+	Result.Particle.DynamicParameter = float4(1,1,1,1);
+#endif //USE_DYNAMIC_PARAMETERS
+
+#if USE_PARTICLE_NODES
+	Result.Particle.TranslatedWorldPositionAndSize.xyz = Interpolants.ParticleTranslatedWorldPosition.xyz;
+	Result.Particle.TranslatedWorldPositionAndSize.w = 1;
+	Result.Particle.Velocity = Interpolants.ParticleVelocity;
+	Result.Particle.RelativeTime = Interpolants.ParticleTime;
+#if NEEDS_PARTICLE_TRANSFORM
+	Result.Particle.LocalToWorld[0] = Interpolants.LocalToWorld[0];
+	Result.Particle.LocalToWorld[1] = Interpolants.LocalToWorld[1];
+	Result.Particle.LocalToWorld[2] = Interpolants.LocalToWorld[2];
+	Result.Particle.LocalToWorld[3] = Interpolants.LocalToWorld[3];
+#endif
+#endif
+
+	Result.Particle.MotionBlurFade = 1.0f;
+	Result.TangentToWorld = AssembleTangentToWorld( TangentToWorld0, TangentToWorld2 );
+	Result.TwoSidedSign = 1;
+
+	return Result;
+}
+
+/** Converts from vertex factory specific input to a FMaterialVertexParameters, which is used by vertex shader material inputs. */
+FMaterialVertexParameters GetMaterialVertexParameters(FVertexFactoryInput Input, FVertexFactoryIntermediates Intermediates, float3 WorldPosition, float3x3 TangentToLocal)
+{
+	FMaterialVertexParameters Result = (FMaterialVertexParameters)0;
+	Result.WorldPosition = WorldPosition;
+	Result.VertexColor = Intermediates.VertexColor;
+	Result.PreSkinnedPosition = Input.Position.xyz;
+	Result.PreSkinnedNormal = Input.TangentZ.xyz * 2.f - 1.f;
+
+	Result.Particle.TranslatedWorldPositionAndSize.xyz = Intermediates.ParticleTranslatedWorldPosition.xyz;
+	Result.Particle.TranslatedWorldPositionAndSize.w = 1;
+	Result.Particle.Velocity = Intermediates.ParticleVelocity;
+	Result.Particle.RelativeTime = Intermediates.RelativeTime;
+	Result.Particle.Color = Intermediates.ParticleColor;
+
+#if USE_DYNAMIC_PARAMETERS
+	Result.Particle.DynamicParameter = Intermediates.DynamicParameter;
+#endif	//USE_DYNAMIC_PARAMETERS
+
+	float4x4 Transform = GetParticleTransform(Intermediates);
+	Result.InstanceLocalToWorld = Transform;
+#if NEEDS_PARTICLE_TRANSFORM
+	Result.Particle.LocalToWorld[0] = Transform[0];
+	Result.Particle.LocalToWorld[1] = Transform[1];
+	Result.Particle.LocalToWorld[2] = Transform[2];
+	Result.Particle.LocalToWorld[3] = Transform[3];
+#endif
+
+	Result.TangentToWorld = mul(TangentToLocal, (float3x3)Transform ); 
+
+#if NUM_MATERIAL_TEXCOORDS_VERTEX
+	for(int CoordinateIndex = 0; CoordinateIndex < NUM_MATERIAL_TEXCOORDS_VERTEX; CoordinateIndex++)
+	{
+		Result.TexCoords[CoordinateIndex] = Intermediates.TexCoords[CoordinateIndex];
+	}
+#endif
+
+	return Result;
+}
+
+#if NEEDS_LIGHTMAP_COORDINATE
+void GetLightMapCoordinates(FVertexFactoryInterpolantsVSToPS Interpolants, out float2 LightmapUV0, out float2 LightmapUV1)
+{
+	LightmapUV0 = LightmapUV1 = 0;
+}
+
+half2 GetShadowMapCoordinate(FVertexFactoryInterpolantsVSToPS Interpolants)
+{
+	return 0;
+}
+#endif
+
+
+float4 CalcWorldPosition(FVertexFactoryInput Input, FVertexFactoryIntermediates Intermediates)
+{
+	float4x4 Transform = GetParticleTransform(Intermediates);
+	float3 WorldPosition = mul(Input.Position, Transform).xyz + ResolvedView.PreViewTranslation;
+	return float4(WorldPosition, Input.Position.w);
+}
+
+float4 CalcPrevWorldPosition(FVertexFactoryInput Input, FVertexFactoryIntermediates Intermediates)
+{
+#if USE_PREVIOUS_TRANSFORM
+	if (MeshParticleVF.PrevTransformAvailable > 0)
+	{
+		float4x4 PrevTransform = GetParticlePrevTransform(Input);
+		float3 PrevWorldPosition = mul(Input.Position, PrevTransform).xyz + ResolvedView.PrevPreViewTranslation;
+		return float4(PrevWorldPosition, Input.Position.w);
+	}
+	else
+#endif
+	// if we don't have a full previous transform, we can at least calculate the position to get some motion blur
+	{
+		return CalcWorldPosition(Input, Intermediates)- float4(Intermediates.ParticleVelocity.xyz*Intermediates.ParticleVelocity.w*MeshParticleVF.DeltaSeconds, 0.0);
+	}
+}
+
+FVertexFactoryIntermediates GetVertexFactoryIntermediates(FVertexFactoryInput Input)
+{
+	FVertexFactoryIntermediates Intermediates = (FVertexFactoryIntermediates)0;
+
+	Intermediates.ParticlePosition = float4(GetVec3(NiagaraMeshVF.PositionDataOffset, Input.InstanceId), 1);
+	Intermediates.ParticleColor = GetVec4(NiagaraMeshVF.ColorDataOffset, Input.InstanceId);
+	Intermediates.ParticleVelocity = float4(GetVec3(NiagaraMeshVF.VelocityDataOffset, Input.InstanceId), 1);
+
+	float4 Transform1 = float4(1.0f, 0.0f, 0.0f, Intermediates.ParticlePosition.x);
+	float4 Transform2 = float4(0.0f, 1.0f, 0.0f, Intermediates.ParticlePosition.y);
+	float4 Transform3 = float4(0.0f, 0.0f, 1.0f, Intermediates.ParticlePosition.z);
+
+	if(NiagaraMeshVF.TransformDataOffset != -1)
+	{
+		float4 AxisAngle = GetVec4(NiagaraMeshVF.TransformDataOffset, Input.InstanceId);
+		AxisAngle.xyz= normalize(AxisAngle).xyz;
+		float Yaw = atan2(AxisAngle.y, AxisAngle.x) * (180.f / PI);
+		float Pitch = atan2(AxisAngle.z,sqrt(AxisAngle.x*AxisAngle.x+AxisAngle.y*AxisAngle.y)) * (180.f / PI);
+		float Roll = 0.0;
+		float4 HalfAngles = float4(Pitch, Yaw, Roll, 0.0f) * (PI/180.0f);
+		float4 SinAngles = sin(HalfAngles);
+		float4 CosAngles = cos(HalfAngles);
+
+		Transform1 =  float4(CosAngles.x * CosAngles.y, CosAngles.x * SinAngles.y, SinAngles.x, Intermediates.ParticlePosition.x);
+		Transform1.xyz = normalize(Transform1.xyz);
+		Transform2 =  float4(SinAngles.z * SinAngles.x * CosAngles.y - CosAngles.z*SinAngles.y, SinAngles.x * SinAngles.z * SinAngles.y + CosAngles.z * CosAngles.y, - SinAngles.z * CosAngles.x, Intermediates.ParticlePosition.y);
+		Transform2.xyz = normalize(Transform2.xyz);
+		Transform3 =  float4(-(CosAngles.z * SinAngles.x * CosAngles.y + SinAngles.z*SinAngles.y), CosAngles.y * SinAngles.z - CosAngles.z * SinAngles.x * SinAngles.y,  CosAngles.z * CosAngles.x, Intermediates.ParticlePosition.z);
+		Transform3.xyz = normalize(Transform3.xyz);
+	}
+
+	if(NiagaraMeshVF.ScaleDataOffset != -1)
+	{
+			float3 Scale = GetVec3(NiagaraMeshVF.ScaleDataOffset, Input.InstanceId);
+			Transform1.xyz *= Scale;
+			Transform2.xyz *= Scale;
+			Transform3.xyz *= Scale;
+	}
+	else if(NiagaraMeshVF.SizeDataOffset != -1)
+	{
+			float2 Size = GetVec2(NiagaraMeshVF.SizeDataOffset, Input.InstanceId);
+			Transform1.xyz *= Size.x;
+			Transform2.xyz *= Size.x;
+			Transform3.xyz *= Size.x;
+	}
+
+
+	// Swizzle vertex color.
+	Intermediates.VertexColor = Input.VertexColor FCOLOR_COMPONENT_SWIZZLE;
+
+	// World position and localtoworld transform.
+	float4x4 LocalToWorldMat = float4x4(Transform1, Transform2, Transform3, float4(0.0f, 0.0f, 0.0f, 1.0f));
+	Intermediates.LocalToWorld = LocalToWorldMat;
+
+	Intermediates.ParticleTranslatedWorldPosition = Intermediates.ParticlePosition + ResolvedView.PreViewTranslation.xyz;
+
+#if NUM_MATERIAL_TEXCOORDS_VERTEX
+	for(int CoordinateIndex = 0; CoordinateIndex < NUM_MATERIAL_TEXCOORDS_VERTEX; CoordinateIndex++)
+	{
+		Intermediates.TexCoords[CoordinateIndex] = Input.TexCoords[CoordinateIndex].xy;
+	}
+#endif
+
+#if USE_PARTICLE_SUBUVS
+		Intermediates.SubUVCoords.xy = (float2( Input.SubUVParams.x, Input.SubUVParams.y ) + Input.TexCoords[0].xy) * MeshParticleVF.SubImageSize.xy;
+		Intermediates.SubUVCoords.zw = (float2( Input.SubUVParams.z, Input.SubUVParams.w ) + Input.TexCoords[0].xy) * MeshParticleVF.SubImageSize.xy;
+		//Intermediates.SubUVLerp = Input.SubUVLerpAndRelTime.x;
+#endif 
+
+#if USE_DYNAMIC_PARAMETERS
+		Intermediates.DynamicParameter = Input.DynamicParameter;
+#endif	//USE_DYNAMIC_PARAMETERS
+
+	Intermediates.ParticleTranslatedWorldPosition = Intermediates.ParticlePosition.xyz + ResolvedView.PreViewTranslation.xyz;
+
+	//Intermediates.RelativeTime = Input.SubUVLerpAndRelTime.y;
+
+	return Intermediates;
+}
+
+/**
+* Get the 3x3 tangent basis vectors for this vertex factory
+* this vertex factory will calculate the binormal on-the-fly
+*
+* @param Input - vertex input stream structure
+* @return 3x3 matrix
+*/
+float3x3 VertexFactoryGetTangentToLocal( FVertexFactoryInput Input, FVertexFactoryIntermediates Intermediates )
+{
+	float3x3 Result=0;	
+	float3 TangentX = TangentBias(Input.TangentX);
+	float4 TangentZ = TangentBias(Input.TangentZ);
+
+	// pass-thru the tangent
+	Result[0] = TangentX;
+	// pass-thru the normal
+	Result[2] = TangentZ.xyz;	
+
+	
+	// derive the binormal by getting the cross product of the normal and tangent
+	Result[1] = cross(Result[2], Result[0]) * TangentZ.w;	
+	// Recalculate TangentX off of the other two vectors
+	// This corrects quantization error since TangentX was passed in as a quantized vertex input
+	// The error shows up most in specular off of a mesh with a smoothed UV seam (normal is smooth, but tangents vary across the seam)
+	Result[0] = cross(Result[1], Result[2]) * TangentZ.w;
+
+	return Result;
+}
+
+// @return translated world position
+float4 VertexFactoryGetWorldPosition(FVertexFactoryInput Input, FVertexFactoryIntermediates Intermediates)
+{
+	return CalcWorldPosition(Input, Intermediates);
+}
+
+float4 VertexFactoryGetRasterizedWorldPosition(FVertexFactoryInput Input, FVertexFactoryIntermediates Intermediates, float4 InWorldPosition)
+{
+	return InWorldPosition;
+}
+
+float3 VertexFactoryGetPositionForVertexLighting(FVertexFactoryInput Input, FVertexFactoryIntermediates Intermediates, float3 TranslatedWorldPosition)
+{
+	return TranslatedWorldPosition;
+}
+
+void CalcTangentToWorld(FVertexFactoryInput Input, FVertexFactoryIntermediates Intermediates, out float3 TangentToWorld0, out float4 TangentToWorld2)
+{
+	float3x3 TangentToLocal = VertexFactoryGetTangentToLocal(Input, Intermediates);
+
+	float4x4 Transform = GetParticleTransform(Intermediates);
+	float3x3 TangentToWorld = mul(TangentToLocal, (float3x3)Transform);
+
+	// Normalize to remove scaling. Incorrect but faster than computing the inverse-transpose.
+	TangentToWorld0 = normalize(TangentToWorld[0]);
+	TangentToWorld2 = float4(normalize(TangentToWorld[2]), TangentBias(Input.TangentZ.w) * Primitive.LocalToWorldDeterminantSign);
+}
+
+FVertexFactoryInterpolantsVSToPS VertexFactoryGetInterpolantsVSToPS(FVertexFactoryInput Input, FVertexFactoryIntermediates Intermediates, FMaterialVertexParameters VertexParameters)
+{
+	FVertexFactoryInterpolantsVSToPS Interpolants;
+
+#if USE_PARTICLE_SUBUVS
+	Interpolants.SubUV0AndTexCoord0.xy = Intermediates.SubUVCoords.xy;
+	Interpolants.SubUV1AndLerp.xy = Intermediates.SubUVCoords.zw;
+	Interpolants.SubUV1AndLerp.zw = Intermediates.SubUVLerp.xx;
+
+	#if NUM_MATERIAL_TEXCOORDS
+		float2 CustomizedUVs[NUM_MATERIAL_TEXCOORDS];
+		GetMaterialCustomizedUVs(VertexParameters, CustomizedUVs);
+		Interpolants.SubUV0AndTexCoord0.zw = CustomizedUVs[0];
+	#else
+		Interpolants.SubUV0AndTexCoord0.zw = 0;
+	#endif
+#elif NUM_MATERIAL_TEXCOORDS
+	#if (ES2_PROFILE || ES3_1_PROFILE)
+		GetMaterialCustomizedUVs(VertexParameters, Interpolants.TexCoords);
+	#else
+		// Ensure the unused components of the last packed texture coordinate are initialized.
+		Interpolants.TexCoords[(NUM_MATERIAL_TEXCOORDS + 1) / 2 - 1] = 0;
+
+		float2 CustomizedUVs[NUM_MATERIAL_TEXCOORDS];
+		GetMaterialCustomizedUVs(VertexParameters, CustomizedUVs);
+
+		UNROLL
+		for(int CoordinateIndex = 0;CoordinateIndex < NUM_MATERIAL_TEXCOORDS;CoordinateIndex += 2)
+		{
+				Interpolants.TexCoords[CoordinateIndex / 2].xy = CustomizedUVs[CoordinateIndex];
+			if(CoordinateIndex + 1 < NUM_MATERIAL_TEXCOORDS)
+			{
+					Interpolants.TexCoords[CoordinateIndex / 2].wz = CustomizedUVs[CoordinateIndex + 1];
+			}
+		}
+	#endif 
+#endif
+
+
+	Interpolants.TangentToWorld0.w = 0;
+	CalcTangentToWorld(Input, Intermediates, Interpolants.TangentToWorld0.xyz, Interpolants.TangentToWorld2);
+
+#if INTERPOLATE_VERTEX_COLOR
+	Interpolants.VertexColor = Intermediates.VertexColor;
+#endif
+
+#if NEEDS_PARTICLE_COLOR
+	Interpolants.ParticleColor = Intermediates.ParticleColor;
+#endif
+
+#if USE_DYNAMIC_PARAMETERS
+	Interpolants.DynamicParameter = Intermediates.DynamicParameter;
+#endif	//USE_DYNAMIC_PARAMETERS
+
+#if USE_PARTICLE_NODES
+	Interpolants.ParticleTranslatedWorldPosition = Intermediates.ParticleTranslatedWorldPosition;
+	Interpolants.ParticleVelocity = Intermediates.ParticleVelocity;
+	Interpolants.ParticleTime = Intermediates.RelativeTime;
+#if NEEDS_PARTICLE_TRANSFORM
+	Interpolants.LocalToWorld = Intermediates.LocalToWorld;
+#endif
+#endif
+
+#if INSTANCED_STEREO
+	Interpolants.PackedEyeIndex = 0;
+#endif
+
+	return Interpolants;
+}
+
+// @return previous translated world position
+float4 VertexFactoryGetPreviousWorldPosition(FVertexFactoryInput Input, FVertexFactoryIntermediates Intermediates)
+{	
+	return CalcPrevWorldPosition(Input, Intermediates);
+}
+
+#if USING_TESSELLATION
+	struct FVertexFactoryInterpolantsVSToDS
+	{
+		FVertexFactoryInterpolantsVSToPS InterpolantsVSToPS;
+	};
+
+	float2 VertexFactoryGetTextureCoordinateDS( FVertexFactoryInterpolantsVSToDS Interpolants )
+	{
+	#if NUM_MATERIAL_TEXCOORDS
+		return Interpolants.InterpolantsVSToPS.TexCoords[0].xy;
+	#else // #if NUM_MATERIAL_TEXCOORDS
+		return float2(0,0);
+	#endif // #if NUM_MATERIAL_TEXCOORDS
+	}
+
+	FVertexFactoryInterpolantsVSToPS VertexFactoryAssignInterpolants(FVertexFactoryInterpolantsVSToDS Input)
+	{
+		return Input.InterpolantsVSToPS;
+	}
+
+	/** Converts from vertex factory specific interpolants to a FMaterialTessellationParameters, which is used by material inputs. */
+	FMaterialTessellationParameters GetMaterialTessellationParameters(FVertexFactoryInterpolantsVSToDS Interpolants, float3 CameraLocalWorldPosition)
+	{
+		FMaterialTessellationParameters	Result;
+	#if NUM_MATERIAL_TEXCOORDS
+		UNROLL
+		for(int CoordinateIndex = 0;CoordinateIndex < NUM_MATERIAL_TEXCOORDS;CoordinateIndex += 2)
+		{
+			Result.TexCoords[CoordinateIndex] = Interpolants.InterpolantsVSToPS.TexCoords[CoordinateIndex/2].xy;
+			if(CoordinateIndex + 1 < NUM_MATERIAL_TEXCOORDS)
+			{
+				Result.TexCoords[CoordinateIndex + 1] = Interpolants.InterpolantsVSToPS.TexCoords[CoordinateIndex/2].wz;
+			}
+		}
+	#endif
+
+		Result.VertexColor = float4(1,1,1,1);
+
+		half3 TangentToWorld0 = Interpolants.InterpolantsVSToPS.TangentToWorld0.xyz;
+		half4 TangentToWorld2 = Interpolants.InterpolantsVSToPS.TangentToWorld2;
+		Result.TangentToWorld = AssembleTangentToWorld( TangentToWorld0, TangentToWorld2 );
+		Result.TangentToWorldPreScale = 1;
+		Result.WorldPosition = CameraLocalWorldPosition + ResolvedView.WorldCameraOrigin;
+
+		return Result;
+	}
+
+	FVertexFactoryInterpolantsVSToDS VertexFactoryGetInterpolantsVSToDS(FVertexFactoryInput Input, FVertexFactoryIntermediates Intermediates, FMaterialVertexParameters VertexParameters)
+	{
+		FVertexFactoryInterpolantsVSToDS Interpolants;
+		Interpolants.InterpolantsVSToPS = VertexFactoryGetInterpolantsVSToPS(Input, Intermediates, VertexParameters);
+		return Interpolants;
+	}
+
+	FVertexFactoryInterpolantsVSToDS VertexFactoryInterpolate(FVertexFactoryInterpolantsVSToDS a, float aInterp, FVertexFactoryInterpolantsVSToDS b, float bInterp)
+	{
+		FVertexFactoryInterpolantsVSToDS O;
+	
+		// Do we really need to interpolate TangentToWorld2 here? It should be replaced by the
+		// interpolated normal from 'whatever' interpolation scheme we're using
+	
+		TESSELLATION_INTERPOLATE_MEMBER(InterpolantsVSToPS.TangentToWorld0.xyz);
+		TESSELLATION_INTERPOLATE_MEMBER(InterpolantsVSToPS.TangentToWorld2);
+#if INTERPOLATE_VERTEX_COLOR
+		TESSELLATION_INTERPOLATE_MEMBER(InterpolantsVSToPS.VertexColor);
+#endif
+
+#if NEEDS_PARTICLE_COLOR
+		TESSELLATION_INTERPOLATE_MEMBER(InterpolantsVSToPS.ParticleColor);
+#endif
+
+#if NUM_MATERIAL_TEXCOORDS
+		for (int i = 0; i < (NUM_MATERIAL_TEXCOORDS + 1) / 2; ++i)
+		{
+			TESSELLATION_INTERPOLATE_MEMBER(InterpolantsVSToPS.TexCoords[i]);
+		}
+#endif
+#if USE_DYNAMIC_PARAMETERS
+		TESSELLATION_INTERPOLATE_MEMBER(InterpolantsVSToPS.DynamicParameter);
+#endif	//USE_DYNAMIC_PARAMETERS
+
+		return O;
+	}
+
+	float3x3 VertexFactoryGetTangentToLocalDS(FVertexFactoryInterpolantsVSToDS Interpolants)
+	{
+		// This duplicates stuff already going on in GetMaterialTessellationParameters(), so
+		// maybe the hull shader could leverage that instead?
+		half3 TangentToWorld0 = Interpolants.InterpolantsVSToPS.TangentToWorld0.xyz;
+		half4 TangentToWorld2 = Interpolants.InterpolantsVSToPS.TangentToWorld2;
+		float3x3 TangentToWorld = AssembleTangentToWorld( TangentToWorld0, TangentToWorld2 );
+
+		return TangentToWorld;
+	}
+#endif // #if USING_TESSELLATION
+
+#if INSTANCED_STEREO
+
+// Instanced stereo sets the eye index explicitly for instanced geometry
+uint InstancedEyeIndex;
+
+uint VertexFactoryGetEyeIndex(uint InstanceId)
+{
+	return InstancedEyeIndex & bIsInstancedStereo;
+}
+#endif
diff --git a/Engine/Shaders/Private/NiagaraSpriteVertexFactory.usf b/Engine/Shaders/Private/NiagaraSpriteVertexFactory.usf
new file mode 100644
--- /dev/null
+++ b/Engine/Shaders/Private/NiagaraSpriteVertexFactory.usf
@@ -0,0 +1,727 @@
+// Copyright 1998-2017 Epic Games, Inc. All Rights Reserved.
+
+/*=============================================================================
+	ParticleSpriteVertexFactory.hlsl: Particle vertex factory shader code.
+	Shared by standard sprite particles and SubUV sprite particles.
+=============================================================================*/
+
+#include "VertexFactoryCommon.usf"
+#include "ParticleVertexFactoryCommon.usf"
+#include "NiagaraVFParticleAccess.usf"
+
+// These are only enabled for SM5 due to limited interpolators elsewhere
+#define USE_PARTICLE_LIGHTING_OFFSET (FEATURE_LEVEL >= FEATURE_LEVEL_SM5 && !MATERIAL_SHADINGMODEL_UNLIT)
+#define USE_PARTICLE_POSITION (FEATURE_LEVEL >= FEATURE_LEVEL_SM5 && NEEDS_PARTICLE_POSITION)
+#define USE_PARTICLE_VELOCITY (FEATURE_LEVEL >= FEATURE_LEVEL_SM5 && NEEDS_PARTICLE_VELOCITY)
+#define USE_PARTICLE_TIME (FEATURE_LEVEL >= FEATURE_LEVEL_SM5 && NEEDS_PARTICLE_TIME)
+#define USE_PARTICLE_SIZE (FEATURE_LEVEL >= FEATURE_LEVEL_SM5 && NEEDS_PARTICLE_SIZE)
+
+bool UseCustomAlignment;
+bool UseVectorAlignment;
+bool CameraPlaneFacing;
+
+
+struct FVertexFactoryInput
+{
+	float4	Position				: ATTRIBUTE0;
+	float4	OldPosition				: ATTRIBUTE1;
+	float4	SizeRotSubImage			: ATTRIBUTE2;
+	float4	Color					: ATTRIBUTE3;
+#if USE_DYNAMIC_PARAMETERS
+	float4 DynamicParameter			: ATTRIBUTE5;
+#endif	//USE_DYNAMIC_PARAMETERS
+	float2	TexCoord				: ATTRIBUTE4;
+	uint VertexId					: SV_VertexID;
+	float3 CustomAlignmentVector	: ATTRIBUTE6;
+	float3 CustomFacingVector		: ATTRIBUTE7;
+
+/** Optional instance ID for vertex layered rendering */
+//#if FEATURE_LEVEL >= FEATURE_LEVEL_SM4 && ONEPASS_POINTLIGHT_SHADOW && USING_VERTEX_SHADER_LAYER
+	uint InstanceId	: SV_InstanceID;
+//#endif
+};
+
+struct FVertexFactoryInterpolantsVSToPS
+{
+	// First row of the tangent to world matrix, Interp_Sizer used by SUBUV_PARTICLES in w
+	float4	TangentToWorld0AndInterp_Sizer : TANGENTTOWORLD0;
+	// Last row of the tangent to world matrix in xyz
+	float4	TangentToWorld2	: TANGENTTOWORLD2;
+
+#if USE_DYNAMIC_PARAMETERS
+	float4 DynamicParameter	: TEXCOORD1;
+#endif	//USE_DYNAMIC_PARAMETERS
+
+#if NEEDS_PARTICLE_COLOR
+	float4	Color				: TEXCOORD2;
+#endif
+
+#if NUM_MATERIAL_TEXCOORDS
+	float4	TexCoords[(NUM_MATERIAL_TEXCOORDS + 1) / 2]	: TEXCOORD3;
+#endif
+
+	//Not sure this is actually being used now and it's awkward to slot in now we're supporting custom UVs so I'm just giving this its own interpolant.
+#if LIGHTMAP_UV_ACCESS
+	float2 LightMapUVs : LIGHTMAP_UVS;
+#endif
+
+#if USE_PARTICLE_SUBUVS
+	float4 ParticleSubUVs : PARTICLE_SUBUVS;
+#endif
+
+#if USE_PARTICLE_POSITION
+	/** Cam-relative (translated) particle center and radius */
+	float4 ParticleTranslatedWorldPositionAndSize	: PARTICLE_POSITION;
+#endif
+
+#if USE_PARTICLE_VELOCITY
+	float4 ParticleVelocity : PARTICLE_VELOCITY;
+#endif
+
+#if USE_PARTICLE_TIME
+	float ParticleTime : PARTICLE_TIME;
+#endif
+
+#if USE_PARTICLE_LIGHTING_OFFSET
+	float3 LightingPositionOffset : PARTICLE_LIGHTING_OFFSET;
+#endif
+
+#if USE_PARTICLE_SIZE
+	float2 ParticleSize : PARTICLE_SIZE;
+#endif
+
+#if INSTANCED_STEREO
+	nointerpolation uint PackedEyeIndex : PACKED_EYE_INDEX;
+#endif
+};
+
+struct FVertexFactoryIntermediates
+{
+	/** The position of the particle in non-translated world space. */
+	float3 ParticleWorldPosition;
+	/** The position of the vertex in translated world space. */
+	float3 VertexWorldPosition;
+	/** Particle translated world space position and size. */
+	float4 TranslatedWorldPositionAndSize;
+	
+#if USE_PARTICLE_LIGHTING_OFFSET
+	float3 LightingPositionOffset;
+#endif
+
+	/** The texture coordinate at this vertex. */
+	float4 TexCoord;
+#if NUM_MATERIAL_TEXCOORDS >= 3 //Unflipped UVs are placed in UVs 2-3
+	/** A second UV set. Always non-UV flipped. Allows use of UV flipping with Normal maps etc. */
+	float4 TexCoord_Unflipped;
+#endif
+	/** The sprite tangent in world space (+V). */
+	float3 TangentUp;
+	/** The sprite tangent in world space (+U). */
+	float3 TangentRight;
+	/** The color of the sprite. */
+	float4 Color;
+	/** The velocity of the particle, XYZ: direction, W: speed. */
+	float4 ParticleVelocity;
+	/** Dynamic parameter. */
+	float4 DynamicParameter;
+	/** The sub-image lerp. */
+	float SubImageLerp;
+	/** Relative time. */
+	float RelativeTime;
+	/** Transform from tangent space to local space. */
+	float3x3 TangentToLocal;
+	/** Size of the particle */
+	float2 ParticleSize;
+};
+
+#if NUM_MATERIAL_TEXCOORDS
+bool UVIndexUseZW(int UVIndex)
+{
+	#if COMPILER_GLSL_ES2
+	return (frac((float)UVIndex / 2.0f) == 0.5f);
+	#else
+	return (UVIndex % 2) != 0;
+	#endif
+}
+
+float2 GetUV(FVertexFactoryInterpolantsVSToPS Interpolants, int UVIndex)
+{
+	float4 UVVector = Interpolants.TexCoords[UVIndex / 2];
+	return UVIndexUseZW(UVIndex) ? UVVector.zw : UVVector.xy;
+}
+
+void SetUV(inout FVertexFactoryInterpolantsVSToPS Interpolants, int UVIndex, float2 InValue)
+{
+	FLATTEN
+	if (UVIndexUseZW(UVIndex))
+	{
+		Interpolants.TexCoords[UVIndex / 2].zw = InValue;
+	}
+	else
+	{
+		Interpolants.TexCoords[UVIndex / 2].xy = InValue;
+	}
+}
+#endif
+
+/** Converts from vertex factory specific interpolants to a FMaterialPixelParameters, which is used by material inputs. */
+FMaterialPixelParameters GetMaterialPixelParameters(FVertexFactoryInterpolantsVSToPS Interpolants, float4 SvPosition)
+{
+	// GetMaterialPixelParameters is responsible for fully initializing the result
+	FMaterialPixelParameters Result = MakeInitializedMaterialPixelParameters();
+
+#if NUM_MATERIAL_TEXCOORDS
+	UNROLL
+	for (int CoordinateIndex = 0; CoordinateIndex < NUM_MATERIAL_TEXCOORDS; CoordinateIndex++)
+	{
+		Result.TexCoords[CoordinateIndex] = GetUV(Interpolants, CoordinateIndex);
+	}
+#endif	// NUM_MATERIAL_TEXCOORDS
+
+	Result.VertexColor = 1;
+
+	half4 TangentToWorld0 = Interpolants.TangentToWorld0AndInterp_Sizer;
+	half4 TangentToWorld2 = Interpolants.TangentToWorld2;
+	float SubImageLerp = Interpolants.TangentToWorld0AndInterp_Sizer.w;
+
+	#if GENERATE_SPHERICAL_PARTICLE_NORMALS && USE_PARTICLE_POSITION
+		{	
+			// can be optimized
+			float4 ScreenPosition = SvPositionToResolvedScreenPosition(SvPosition);
+			float3 TranslatedWorldPosition = SvPositionToResolvedTranslatedWorld(SvPosition);
+			Result.TangentToWorld = GetSphericalParticleNormal(TranslatedWorldPosition, Interpolants.ParticleTranslatedWorldPositionAndSize.xyz, Interpolants.ParticleTranslatedWorldPositionAndSize.w);
+		}
+	#else
+		Result.TangentToWorld = AssembleTangentToWorld(TangentToWorld0.xyz, TangentToWorld2);
+	#endif
+
+	Result.UnMirrored = 1;
+	Result.Particle.MacroUV = NiagaraSpriteVF.MacroUVParameters;
+
+#if NEEDS_PARTICLE_COLOR
+	Result.Particle.Color = Interpolants.Color;
+#endif
+
+#if USE_DYNAMIC_PARAMETERS
+	Result.Particle.DynamicParameter = Interpolants.DynamicParameter;
+#endif	//USE_DYNAMIC_PARAMETERS
+
+#if USE_PARTICLE_POSITION
+	Result.Particle.TranslatedWorldPositionAndSize = Interpolants.ParticleTranslatedWorldPositionAndSize;
+#endif
+
+#if USE_PARTICLE_VELOCITY
+	Result.Particle.Velocity = Interpolants.ParticleVelocity;
+#endif
+
+#if USE_PARTICLE_TIME
+	Result.Particle.RelativeTime = Interpolants.ParticleTime;
+#endif
+
+	Result.Particle.MotionBlurFade = 1.0f;
+
+
+#if USE_PARTICLE_LIGHTING_OFFSET
+	Result.LightingPositionOffset = Interpolants.LightingPositionOffset;
+#endif
+
+#if USE_PARTICLE_SUBUVS
+	Result.Particle.SubUVCoords[0] = Interpolants.ParticleSubUVs.xy;
+	Result.Particle.SubUVCoords[1] = Interpolants.ParticleSubUVs.zw;
+	Result.Particle.SubUVLerp = SubImageLerp;
+#endif
+
+#if LIGHTMAP_UV_ACCESS
+	Result.LightmapUVs = Interpolants.LightMapUVs;
+#endif
+
+#if USE_PARTICLE_SIZE
+	Result.Particle.Size = Interpolants.ParticleSize;
+#endif
+
+	Result.TwoSidedSign = 1;
+	return Result;
+}
+
+/** Converts from vertex factory specific input to a FMaterialVertexParameters, which is used by vertex shader material inputs. */
+FMaterialVertexParameters GetMaterialVertexParameters(FVertexFactoryInput Input, FVertexFactoryIntermediates Intermediates, float3 WorldPosition, float3x3 TangentToLocal)
+{
+	FMaterialVertexParameters Result = (FMaterialVertexParameters)0;
+	Result.WorldPosition = WorldPosition;
+	Result.VertexColor = Input.Color;
+	Result.TangentToWorld = mul(TangentToLocal, GetLocalToWorld3x3()); 
+	Result.Particle.MacroUV = NiagaraSpriteVF.MacroUVParameters;
+	Result.Particle.Color = Intermediates.Color;
+	Result.Particle.MotionBlurFade = 1.0f;
+	Result.PreSkinnedPosition = Input.Position.xyz;
+	Result.PreSkinnedNormal = TangentToLocal[2].xyz;
+
+#if USE_DYNAMIC_PARAMETERS
+	Result.Particle.DynamicParameter = Intermediates.DynamicParameter;
+#endif	//USE_DYNAMIC_PARAMETERS
+
+#if USE_PARTICLE_POSITION
+	Result.Particle.TranslatedWorldPositionAndSize = Intermediates.TranslatedWorldPositionAndSize;
+#endif
+
+#if USE_PARTICLE_VELOCITY
+	Result.Particle.Velocity = Intermediates.ParticleVelocity;
+#endif
+
+#if USE_PARTICLE_TIME
+	Result.Particle.RelativeTime = Intermediates.RelativeTime;
+#endif
+
+#if NUM_MATERIAL_TEXCOORDS_VERTEX
+	#if NUM_MATERIAL_TEXCOORDS_VERTEX >= 1
+		Result.TexCoords[0] = Intermediates.TexCoord.xy;
+	#if NUM_MATERIAL_TEXCOORDS_VERTEX >= 2
+		Result.TexCoords[1] = Intermediates.TexCoord.zw;
+	#if NUM_MATERIAL_TEXCOORDS_VERTEX >= 3
+		Result.TexCoords[2] = Intermediates.TexCoord_Unflipped.xy;
+	#if NUM_MATERIAL_TEXCOORDS_VERTEX >= 4
+		Result.TexCoords[3] = Intermediates.TexCoord_Unflipped.zw;
+	#endif	// >= 4
+	#endif	// >= 3
+	#endif	// >= 2
+	#endif	// >= 1
+#endif	//NUM_MATERIAL_TEXCOORDS_VERTEX
+
+#if USE_PARTICLE_SUBUVS
+	Result.Particle.SubUVCoords[0] = Intermediates.TexCoord.xy;
+	Result.Particle.SubUVCoords[1] = Intermediates.TexCoord.zw;
+#endif
+
+#if USE_PARTICLE_SIZE
+	Result.Particle.Size = Intermediates.ParticleSize;
+#endif
+	return Result;
+}
+
+void GetLightMapCoordinates(FVertexFactoryInterpolantsVSToPS Interpolants, out float2 LightmapUV0, out float2 LightmapUV1)
+{
+	LightmapUV0 = LightmapUV1 = 0;
+}
+
+half2 GetShadowMapCoordinate(FVertexFactoryInterpolantsVSToPS Interpolants)
+{
+	return 0;
+}
+
+float3 SafeNormalize(float3 V)
+{
+	return V * rsqrt(max(dot(V,V),0.01));
+}
+
+void GetTangents(FVertexFactoryInput Input, float3 WorldPosition, float3 OldWorldPosition, out float3 OutRight, out float3 OutUp)
+{
+	// Select camera up/right vectors.
+	float3 ResolvedViewRight = lerp(ResolvedView.ViewRight, ResolvedView.HMDViewNoRollRight, NiagaraSpriteVF.RemoveHMDRoll);
+	float3 ResolvedViewUp = lerp(ResolvedView.ViewUp, ResolvedView.HMDViewNoRollUp, NiagaraSpriteVF.RemoveHMDRoll);
+
+	// camera related vectors
+	float3 CameraRight = ResolvedViewRight;
+	float3 CameraUp = -ResolvedViewUp;
+	float3 CameraVector = ResolvedView.WorldCameraOrigin.xyz - WorldPosition.xyz;
+	float3 CameraDirection = SafeNormalize(CameraVector);
+	float3 CameraFacingVector = Input.CustomFacingVector * NiagaraSpriteVF.CustomFacingVectorMask + CameraDirection * (float3(1.0,1.0,1.0)-NiagaraSpriteVF.CustomFacingVectorMask);
+	CameraFacingVector = lerp(CameraFacingVector, -ResolvedView.ViewForward, CameraPlaneFacing);
+	
+	BRANCH
+	if (NiagaraSpriteVF.CameraFacingBlend.x > 0.f)
+	{
+		float CameraDistanceSq = dot(CameraVector, CameraVector);
+		float AlignmentMode = saturate(CameraDistanceSq * NiagaraSpriteVF.CameraFacingBlend.y - NiagaraSpriteVF.CameraFacingBlend.z);
+		CameraFacingVector = -ResolvedView.ViewForward * CameraFacingVector * (1.0-AlignmentMode);
+	}
+
+	// the particle's movement direction.
+	float3 ParticleDirection	= SafeNormalize(WorldPosition.xyz - OldWorldPosition.xyz);
+	float3 DirectionVector	= lerp(ParticleDirection, Input.CustomAlignmentVector, UseCustomAlignment);
+
+	// Tangent vectors for PSA_Velocity.
+	float3 Right_Aligned	= SafeNormalize(cross(CameraDirection, DirectionVector));
+	float3 Up_Aligned		= -DirectionVector;
+
+	// Tangent vectors for camera facing position (can be based on camera position or the custom facing vector)
+	//
+	float3 Right_FacingCameraPosition = SafeNormalize(cross(CameraFacingVector,ResolvedViewUp));
+	float3 Up_FacingCameraPosition = cross(CameraFacingVector, Right_FacingCameraPosition);
+
+	float3 Right = SafeNormalize( lerp(Right_FacingCameraPosition, Right_Aligned, UseVectorAlignment) );
+	float3 Up = SafeNormalize( cross(CameraFacingVector, Right) );
+
+	// Determine the angle of rotation.
+	float SinRotation;
+	float CosRotation;
+	const float SpriteRotation = Input.SizeRotSubImage.z * NiagaraSpriteVF.RotationScale + NiagaraSpriteVF.RotationBias;
+	sincos(SpriteRotation, SinRotation, CosRotation);
+
+	// Rotate the sprite to determine final tangents.
+	OutRight	= SinRotation * Up + CosRotation * Right;
+	OutUp		= CosRotation * Up - SinRotation * Right;
+}
+
+/** derive basis vectors */
+float3x3 CalcTangentBasis(FVertexFactoryIntermediates Intermediates)
+{
+	// Using camera facing TangentX and TangentY.  The resulting tangent basis is not orthonormal with anything other than ENM_CameraFacing, 
+	// So there are artifacts with tangent space calculations like the TransformVector node,
+	// But this allows lighting based on a world space shape via the normal while still having normal maps camera aligned.
+	float3x3 Result;
+	Result[0] = Intermediates.TangentRight;
+	Result[1] = Intermediates.TangentUp;
+
+	// ENM_CameraFacing
+	//@todo - use static branching
+	if (NiagaraSpriteVF.NormalsType < .5f)
+	{
+		Result[2] = normalize(cross(Result[0],Result[1]));
+	} 
+	// ENM_Spherical
+	else if (NiagaraSpriteVF.NormalsType < 1.5f)
+	{
+		float3 TangentZ = normalize(Intermediates.ParticleWorldPosition.xyz - NiagaraSpriteVF.NormalsSphereCenter.xyz);
+		Result[2] = TangentZ;	
+	}
+	// ENM_Cylindrical
+	else
+	{
+		float3 ClosestPointOnCylinder = NiagaraSpriteVF.NormalsSphereCenter.xyz + dot(NiagaraSpriteVF.NormalsCylinderUnitDirection.xyz, Intermediates.ParticleWorldPosition.xyz - NiagaraSpriteVF.NormalsSphereCenter.xyz) * NiagaraSpriteVF.NormalsCylinderUnitDirection.xyz;
+		float3 TangentZ = normalize(Intermediates.ParticleWorldPosition.xyz - ClosestPointOnCylinder);
+		Result[2] = TangentZ;	
+	}
+
+	return Result;
+}
+
+// Note: C++ geometry setup behavior has to match this define
+/** Number of vertices in each SubUV frame of CutoutGeometry. */
+uint NumCutoutVerticesPerFrame;
+
+/** SubUV animation bounding geometry used to minimize overdraw.  Each frame has its own polygon. */
+Buffer<float2> CutoutGeometry;
+
+void ComputeBillboardUVs(FVertexFactoryInput Input, out float2 UVForPosition, out float2 UVForTexturing, out float2 UVForTexturingUnflipped)
+{
+	// Encoding the UV flip in the sign of the size data.
+	float2 UVFlip = sign(Input.SizeRotSubImage.xy);
+
+	BRANCH
+	if (NumCutoutVerticesPerFrame > 0)
+	{
+		// Avoid uint divide which is extremely slow on GCN
+		uint CutoutVertexIndex = (uint)fmod(Input.VertexId, NumCutoutVerticesPerFrame);
+		float NumFrames = NiagaraSpriteVF.SubImageSize.x * NiagaraSpriteVF.SubImageSize.y;
+		uint SubImageIndexInt = (uint)fmod(Input.SizeRotSubImage.w, NumFrames);
+
+		FLATTEN
+		if (UVFlip.x * UVFlip.y < 0)
+		{
+			// Reverse the winding order of the polygon when only flipping in one direction to counteract UV flipping
+			CutoutVertexIndex = NumCutoutVerticesPerFrame - 1 - CutoutVertexIndex;
+		}
+
+		// Fetch the billboard space positions from the appropriate frame of the preprocessed cutout geometry
+		UVForTexturing = UVForTexturingUnflipped = CutoutGeometry[SubImageIndexInt * NumCutoutVerticesPerFrame + CutoutVertexIndex];
+
+		// Invert positions on the billboard so that the cutout geometry still contains the visible portion
+		// This changes the winding order when only one direction is flipped
+		UVForPosition.x = UVFlip.x < 0.0 ? 1.0 - UVForTexturing.x : UVForTexturing.x;
+		UVForPosition.y = UVFlip.y < 0.0 ? 1.0 - UVForTexturing.y : UVForTexturing.y;
+	}
+	else
+	{
+		UVForTexturing.x = UVFlip.x < 0.0 ? 1.0 - Input.TexCoord.x : Input.TexCoord.x;
+		UVForTexturing.y = UVFlip.y < 0.0 ? 1.0 - Input.TexCoord.y : Input.TexCoord.y;
+
+		// Note: not inverting positions, as that would change the winding order
+		UVForPosition = UVForTexturingUnflipped = Input.TexCoord.xy;
+	}
+}
+
+
+
+/*********************************************************************************
+ *							Particle specific
+ *********************************************************************************/
+
+float3 GetNiagaraParticlePosition(uint InstanceID)
+{
+	return GetVec3(NiagaraSpriteVF.PositionDataOffset, InstanceID);
+}
+float3 GetNiagaraParticleVelocity(uint InstanceID)
+{
+	return GetVec3(NiagaraSpriteVF.VelocityDataOffset, InstanceID);
+}
+float2 GetNiagaraParticleSize(uint InstanceID)
+{
+	return GetVec2(NiagaraSpriteVF.SizeDataOffset, InstanceID);
+}
+float GetNiagaraParticleRotation(uint InstanceID)
+{
+	return GetFloat(NiagaraSpriteVF.RotationDataOffset, InstanceID);
+}
+float4 GetNiagaraParticleColor(uint InstanceID)
+{
+	return GetVec4(NiagaraSpriteVF.ColorDataOffset, InstanceID);
+}
+
+float GetNiagaraParticleSubimage(uint InstanceID)
+{
+	float Val = 0.0f;
+    if(NiagaraSpriteVF.SubimageDataOffset != -1)
+	{
+		Val = GetFloat(NiagaraSpriteVF.SubimageDataOffset, InstanceID);
+	}
+	return Val;
+}
+
+float3 GetNiagaraParticleFacingVector(uint InstanceID)
+{
+	float3 Val = float3(0.0f, 0.0f, 0.0f);
+    if(NiagaraSpriteVF.FacingOffset != -1)
+	{
+		Val = GetVec3(NiagaraSpriteVF.FacingOffset, InstanceID);
+	}
+	return Val;
+}
+
+float3 GetNiagaraParticleAlignmentVector(uint InstanceID)
+{
+	float3 Val = float3(0.0f, 0.0f, 0.0f);
+    if(NiagaraSpriteVF.AlignmentOffset != -1)
+	{
+		Val = GetVec3(NiagaraSpriteVF.AlignmentOffset, InstanceID);
+	}
+	return Val;
+}
+
+
+FVertexFactoryIntermediates GetVertexFactoryIntermediates(FVertexFactoryInput Input)
+{
+	FVertexFactoryIntermediates Intermediates = (FVertexFactoryIntermediates)0;
+
+	Input.Position = float4(GetNiagaraParticlePosition(Input.InstanceId), 1);
+	Input.OldPosition = float4(Input.Position.xyz - GetNiagaraParticleVelocity(Input.InstanceId), 1);
+	float2 ParticleSize = GetNiagaraParticleSize(Input.InstanceId);
+	Input.SizeRotSubImage = float4(ParticleSize.x, ParticleSize.y, GetNiagaraParticleRotation(Input.InstanceId), GetNiagaraParticleSubimage(Input.InstanceId));
+	Input.Color	= GetNiagaraParticleColor(Input.InstanceId);
+	Input.CustomFacingVector = GetNiagaraParticleFacingVector(Input.InstanceId);
+	Input.CustomAlignmentVector = GetNiagaraParticleAlignmentVector(Input.InstanceId);
+
+	// World position.
+	float4 ParticleWorldPosition = mul(float4(Input.Position.xyz,1), Primitive.LocalToWorld);
+	float4 ParticleOldWorldPosition = mul(float4(Input.OldPosition.xyz,1), Primitive.LocalToWorld); 
+	Intermediates.ParticleWorldPosition = ParticleWorldPosition.xyz;
+
+	// Tangents.
+	float3 Right,Up;
+	GetTangents(Input,ParticleWorldPosition.xyz,ParticleOldWorldPosition.xyz,Right,Up);
+	Intermediates.TangentUp = Up;
+	Intermediates.TangentRight = Right;
+
+	float2 UVForPosition;
+	float2 UVForTexturing;
+	float2 UVForTexturingUnflipped;
+	ComputeBillboardUVs(Input, UVForPosition, UVForTexturing, UVForTexturingUnflipped);
+
+	// Vertex position.
+	float4 VertexWorldPosition = float4(ParticleWorldPosition.xyz + ResolvedView.PreViewTranslation,1);
+	float2 Size = abs(Input.SizeRotSubImage.xy);
+	VertexWorldPosition += Size.x * (UVForPosition.x + NiagaraSpriteVF.PivotOffset.x) * float4(Right,0);
+	VertexWorldPosition += Size.y * (UVForPosition.y + NiagaraSpriteVF.PivotOffset.y) * float4(Up,0);
+	Intermediates.VertexWorldPosition = VertexWorldPosition.xyz;
+	
+	// SubUV.
+	float SubImageIndex = Input.SizeRotSubImage.w;
+	float SubImageLerp = frac(SubImageIndex);
+	float SubImageA = SubImageIndex - SubImageLerp;
+	float SubImageB = SubImageA + 1;
+	float SubImageAH = fmod( SubImageA, NiagaraSpriteVF.SubImageSize.x );
+	float SubImageBH = fmod( SubImageB, NiagaraSpriteVF.SubImageSize.x );
+	float SubImageAV = floor( SubImageA * NiagaraSpriteVF.SubImageSize.z );
+	float SubImageBV = floor( SubImageB * NiagaraSpriteVF.SubImageSize.z );
+	Intermediates.TexCoord.xy = (float2( SubImageAH, SubImageAV ) + UVForTexturing) * NiagaraSpriteVF.SubImageSize.zw;
+	Intermediates.TexCoord.zw = (float2(SubImageBH, SubImageBV) + UVForTexturing) * NiagaraSpriteVF.SubImageSize.zw;
+#if NUM_MATERIAL_TEXCOORDS >= 3
+	Intermediates.TexCoord_Unflipped.xy = (float2(SubImageAH, SubImageAV) + UVForTexturingUnflipped) * NiagaraSpriteVF.SubImageSize.zw;
+	Intermediates.TexCoord_Unflipped.zw = (float2(SubImageBH, SubImageBV) + UVForTexturingUnflipped) * NiagaraSpriteVF.SubImageSize.zw;
+#endif
+	Intermediates.SubImageLerp = SubImageLerp;
+
+	Intermediates.Color = Input.Color;
+
+#if USE_DYNAMIC_PARAMETERS
+	Intermediates.DynamicParameter = Input.DynamicParameter;
+#endif	//USE_DYNAMIC_PARAMETERS
+
+
+#if USE_PARTICLE_POSITION
+	float ParticleRadius = .5f * min(Input.SizeRotSubImage.x, Input.SizeRotSubImage.y);
+	float3 ParticleTranslatedWorldPosition = ParticleWorldPosition.xyz + ResolvedView.PreViewTranslation.xyz;
+	Intermediates.TranslatedWorldPositionAndSize = float4(ParticleTranslatedWorldPosition, ParticleRadius);
+
+#if GENERATE_SPHERICAL_PARTICLE_NORMALS
+	// Recenter the particle position for GENERATE_SPHERICAL_PARTICLE_NORMALS to use the right center.
+	Intermediates.TranslatedWorldPositionAndSize.xyz += Size.x * (NiagaraSpriteVF.PivotOffset.x + .5) * Right + Size.y * (NiagaraSpriteVF.PivotOffset.y + .5) * Up;
+#endif
+#endif
+
+#if USE_PARTICLE_VELOCITY
+	float3 ParticleVelocity = (ParticleWorldPosition.xyz - ParticleOldWorldPosition.xyz) * NiagaraSpriteVF.InvDeltaSeconds;
+	float3 DirectionVector = lerp(ParticleVelocity, Input.CustomAlignmentVector, UseCustomAlignment);
+	Intermediates.ParticleVelocity.xyz = normalize(DirectionVector);
+	Intermediates.ParticleVelocity.w = length(DirectionVector);
+#endif
+
+#if USE_PARTICLE_TIME
+	Intermediates.RelativeTime = Input.Position.w;
+#endif
+
+	Intermediates.TangentToLocal = CalcTangentBasis(Intermediates);
+
+#if USE_PARTICLE_LIGHTING_OFFSET
+	// Hash function based on the particle ID to generate a uniformly distributed 3d offset
+	float3 RandomParticleOffset = frac(Square(Input.OldPosition.w + 10) * float3(1361.456345, 2333.578, 3623.983)) * 2 - 1;
+	Intermediates.LightingPositionOffset = .5f * View.TranslucencyLightingVolumeInvSize[0].w * RandomParticleOffset;
+#endif
+
+#if USE_PARTICLE_SIZE
+	Intermediates.ParticleSize = Size;
+#endif
+
+	return Intermediates;
+}
+
+float4 VertexFactoryGetWorldPosition(FVertexFactoryInput Input, FVertexFactoryIntermediates Intermediates)
+{
+	return float4(Intermediates.VertexWorldPosition,1);
+}
+
+float4 VertexFactoryGetRasterizedWorldPosition(FVertexFactoryInput Input, FVertexFactoryIntermediates Intermediates, float4 TranslatedWorldPosition)
+{
+#if SPHERICAL_PARTICLE_OPACITY
+	// For particles using spherical opacity, move the quad toward the viewer so that it lies in front of the sphere defined by the particle
+	// This avoids opaque objects intersecting the particle from causing clipping artifacts due to depth testing
+	// The downside is that the particle will clip the near plane sooner
+
+	float Radius = .5f * min(Input.SizeRotSubImage.x, Input.SizeRotSubImage.y);
+	return ReprojectPosition(TranslatedWorldPosition, Radius);
+#else
+	return TranslatedWorldPosition;
+#endif
+}
+
+float3 VertexFactoryGetPositionForVertexLighting(FVertexFactoryInput Input, FVertexFactoryIntermediates Intermediates, float3 TranslatedWorldPosition)
+{
+#if (TRANSLUCENCY_LIGHTING_VOLUMETRIC_PERVERTEX_DIRECTIONAL || TRANSLUCENCY_LIGHTING_VOLUMETRIC_PERVERTEX_NONDIRECTIONAL)
+	// Do per-vertex lighting with particle position instead of vertex position when we're using SubUV cutouts, since the vertex positions vary per-frame which would cause popping
+	return NumCutoutVerticesPerFrame > 0 ? (Intermediates.ParticleWorldPosition.xyz + ResolvedView.PreViewTranslation.xyz) : TranslatedWorldPosition;
+#else
+	return TranslatedWorldPosition;
+#endif
+}
+
+FVertexFactoryInterpolantsVSToPS VertexFactoryGetInterpolantsVSToPS(FVertexFactoryInput Input, FVertexFactoryIntermediates Intermediates, FMaterialVertexParameters VertexParameters)
+{
+	FVertexFactoryInterpolantsVSToPS Interpolants;
+
+	// Initialize the whole struct to 0
+	Interpolants = (FVertexFactoryInterpolantsVSToPS)0;
+
+#if NUM_MATERIAL_TEXCOORDS
+
+	float2 CustomizedUVs[NUM_MATERIAL_TEXCOORDS];
+	GetMaterialCustomizedUVs(VertexParameters, CustomizedUVs);
+
+	UNROLL
+	for (int CoordinateIndex = 0; CoordinateIndex < NUM_MATERIAL_TEXCOORDS; CoordinateIndex++)
+	{
+		SetUV(Interpolants, CoordinateIndex, CustomizedUVs[CoordinateIndex]);
+	}
+#endif
+
+#if LIGHTMAP_UV_ACCESS
+	Interpolants.LightMapUVs = Intermediates.TexCoord.xy;
+#endif
+
+#if USE_PARTICLE_SUBUVS
+	Interpolants.ParticleSubUVs.xy = VertexParameters.Particle.SubUVCoords[0];
+	Interpolants.ParticleSubUVs.zw = VertexParameters.Particle.SubUVCoords[1];
+#endif
+
+	// Calculate the transform from tangent to world space.
+	// Note that "local" space for particles is actually oriented in world space! Therefore no rotation is needed.
+	float3x3 TangentToWorld = Intermediates.TangentToLocal;
+
+	Interpolants.TangentToWorld0AndInterp_Sizer.xyz = TangentToWorld[0];
+	Interpolants.TangentToWorld0AndInterp_Sizer.w = Intermediates.SubImageLerp;
+	Interpolants.TangentToWorld2 = float4(TangentToWorld[2], sign(determinant(Intermediates.TangentToLocal)));
+
+#if NEEDS_PARTICLE_COLOR
+	Interpolants.Color = Intermediates.Color;
+#endif
+
+#if USE_DYNAMIC_PARAMETERS
+	Interpolants.DynamicParameter = Intermediates.DynamicParameter;
+#endif	//USE_DYNAMIC_PARAMETERS
+
+#if USE_PARTICLE_POSITION
+	Interpolants.ParticleTranslatedWorldPositionAndSize = Intermediates.TranslatedWorldPositionAndSize;
+#endif
+
+#if USE_PARTICLE_VELOCITY
+	Interpolants.ParticleVelocity = Intermediates.ParticleVelocity;
+#endif
+
+#if USE_PARTICLE_TIME
+	Interpolants.ParticleTime = Intermediates.RelativeTime;
+#endif
+
+#if USE_PARTICLE_LIGHTING_OFFSET
+	Interpolants.LightingPositionOffset = Intermediates.LightingPositionOffset;
+#endif
+
+#if USE_PARTICLE_SIZE
+	Interpolants.ParticleSize = Intermediates.ParticleSize;
+#endif
+
+#if INSTANCED_STEREO
+	Interpolants.PackedEyeIndex = 0;
+#endif
+
+	return Interpolants;
+}
+
+float4 VertexFactoryGetPreviousWorldPosition(FVertexFactoryInput Input, FVertexFactoryIntermediates Intermediates)
+{
+	return VertexFactoryGetWorldPosition(Input, Intermediates);
+}
+
+/**
+* Get the 3x3 tangent basis vectors for this vertex factory
+*
+* @param Input - vertex input stream structure
+* @return 3x3 matrix
+*/
+float3x3 VertexFactoryGetTangentToLocal(FVertexFactoryInput Input, FVertexFactoryIntermediates Intermediates)
+{
+	return Intermediates.TangentToLocal;
+}
+
+#if INSTANCED_STEREO
+
+// Instanced stereo sets the eye index explicitly for instanced geometry
+uint InstancedEyeIndex;	 
+
+uint VertexFactoryGetEyeIndex(uint InstanceId)
+{
+	return InstancedEyeIndex & bIsInstancedStereo;
+}
+#endif
+
+float4 VertexFactoryGetTranslatedPrimitiveVolumeBounds(FVertexFactoryInterpolantsVSToPS Interpolants)
+{
+	return 0;
+}
\ No newline at end of file
diff --git a/Engine/Shaders/Private/NiagaraSpriteVertexFactory.ush b/Engine/Shaders/Private/NiagaraSpriteVertexFactory.ush
--- a/Engine/Shaders/Private/NiagaraSpriteVertexFactory.ush
+++ b/Engine/Shaders/Private/NiagaraSpriteVertexFactory.ush
@@ -7,6 +7,7 @@
 
 #include "VertexFactoryCommon.ush"
 #include "ParticleVertexFactoryCommon.ush"
+#include "NiagaraVFParticleAccess.usf"
 
 // These are only enabled for SM5 due to limited interpolators elsewhere
 #define USE_PARTICLE_LIGHTING_OFFSET (FEATURE_LEVEL >= FEATURE_LEVEL_SM5 && !MATERIAL_SHADINGMODEL_UNLIT)
@@ -35,9 +36,9 @@ struct FVertexFactoryInput
 	float3 CustomFacingVector		: ATTRIBUTE7;
 
 /** Optional instance ID for vertex layered rendering */
-#if FEATURE_LEVEL >= FEATURE_LEVEL_SM4 && ONEPASS_POINTLIGHT_SHADOW && USING_VERTEX_SHADER_LAYER
+//#if FEATURE_LEVEL >= FEATURE_LEVEL_SM4 && ONEPASS_POINTLIGHT_SHADOW && USING_VERTEX_SHADER_LAYER
 	uint InstanceId	: SV_InstanceID;
-#endif
+//#endif
 };
 
 struct FVertexFactoryInterpolantsVSToPS
@@ -439,13 +440,79 @@ void ComputeBillboardUVs(FVertexFactoryInput Input, out float2 UVForPosition, ou
 	}
 }
 
+
+
+/*********************************************************************************
+ *							Particle specific
+ *********************************************************************************/
+
+float3 GetNiagaraParticlePosition(uint InstanceID)
+{
+	return GetVec3(NiagaraSpriteVF.PositionDataOffset, InstanceID);
+}
+float3 GetNiagaraParticleVelocity(uint InstanceID)
+{
+	return GetVec3(NiagaraSpriteVF.VelocityDataOffset, InstanceID);
+}
+float2 GetNiagaraParticleSize(uint InstanceID)
+{
+	return GetVec2(NiagaraSpriteVF.SizeDataOffset, InstanceID);
+}
+float GetNiagaraParticleRotation(uint InstanceID)
+{
+	return GetFloat(NiagaraSpriteVF.RotationDataOffset, InstanceID);
+}
+float4 GetNiagaraParticleColor(uint InstanceID)
+{
+	return GetVec4(NiagaraSpriteVF.ColorDataOffset, InstanceID);
+}
+
+float GetNiagaraParticleSubimage(uint InstanceID)
+{
+	float Val = 0.0f;
+    if(NiagaraSpriteVF.SubimageDataOffset != -1)
+	{
+		Val = GetFloat(NiagaraSpriteVF.SubimageDataOffset, InstanceID);
+	}
+	return Val;
+}
+
+float3 GetNiagaraParticleFacingVector(uint InstanceID)
+{
+	float3 Val = float3(0.0f, 0.0f, 0.0f);
+    if(NiagaraSpriteVF.FacingOffset != -1)
+	{
+		Val = GetVec3(NiagaraSpriteVF.FacingOffset, InstanceID);
+	}
+	return Val;
+}
+
+float3 GetNiagaraParticleAlignmentVector(uint InstanceID)
+{
+	float3 Val = float3(0.0f, 0.0f, 0.0f);
+    if(NiagaraSpriteVF.AlignmentOffset != -1)
+	{
+		Val = GetVec3(NiagaraSpriteVF.AlignmentOffset, InstanceID);
+	}
+	return Val;
+}
+
+
 FVertexFactoryIntermediates GetVertexFactoryIntermediates(FVertexFactoryInput Input)
 {
 	FVertexFactoryIntermediates Intermediates = (FVertexFactoryIntermediates)0;
 
+	Input.Position = float4(GetNiagaraParticlePosition(Input.InstanceId), 1);
+	Input.OldPosition = float4(Input.Position.xyz - GetNiagaraParticleVelocity(Input.InstanceId), 1);
+	float2 ParticleSize = GetNiagaraParticleSize(Input.InstanceId);
+	Input.SizeRotSubImage = float4(ParticleSize.x, ParticleSize.y, GetNiagaraParticleRotation(Input.InstanceId), GetNiagaraParticleSubimage(Input.InstanceId));
+	Input.Color	= GetNiagaraParticleColor(Input.InstanceId);
+	Input.CustomFacingVector = GetNiagaraParticleFacingVector(Input.InstanceId);
+	Input.CustomAlignmentVector = GetNiagaraParticleAlignmentVector(Input.InstanceId);
+
 	// World position.
 	float4 ParticleWorldPosition = mul(float4(Input.Position.xyz,1), Primitive.LocalToWorld);
-	float4 ParticleOldWorldPosition = mul(float4(Input.OldPosition.xyz,1), Primitive.LocalToWorld);
+	float4 ParticleOldWorldPosition = mul(float4(Input.OldPosition.xyz,1), Primitive.LocalToWorld); 
 	Intermediates.ParticleWorldPosition = ParticleWorldPosition.xyz;
 
 	// Tangents.
diff --git a/Engine/Shaders/Private/NiagaraVFParticleAccess.usf b/Engine/Shaders/Private/NiagaraVFParticleAccess.usf
new file mode 100644
--- /dev/null
+++ b/Engine/Shaders/Private/NiagaraVFParticleAccess.usf
@@ -0,0 +1,29 @@
+
+/*********************************************************************************
+ *							Niagara Vertex Factory data access
+ *********************************************************************************/
+Buffer<float> NiagaraParticleDataFloat;
+Buffer<int> NiagaraParticleDataInt;
+uint SafeComponentBufferSize;			// stride of data in buffers
+uint VertexID : SV_VertexID;
+
+float GetFloat(int RegisterIdx, uint InstanceID)
+{
+	return NiagaraParticleDataFloat[RegisterIdx*SafeComponentBufferSize + InstanceID];
+}
+
+float2 GetVec2(int RegisterIndex, uint InstanceID)
+{
+	return float2(GetFloat(RegisterIndex, InstanceID), GetFloat(RegisterIndex+1, InstanceID));
+}
+
+float3 GetVec3(int RegisterIndex, uint InstanceID)
+{
+	return float3(GetFloat(RegisterIndex, InstanceID), GetFloat(RegisterIndex+1, InstanceID), GetFloat(RegisterIndex+2, InstanceID));
+}
+
+float4 GetVec4(int RegisterIndex, uint InstanceID)
+{
+	return float4(GetFloat(RegisterIndex, InstanceID), GetFloat(RegisterIndex+1, InstanceID), GetFloat(RegisterIndex+2, InstanceID), GetFloat(RegisterIndex+3, InstanceID));
+}
+
diff --git a/Engine/Shaders/Private/Random.usf b/Engine/Shaders/Private/Random.usf
new file mode 100644
--- /dev/null
+++ b/Engine/Shaders/Private/Random.usf
@@ -0,0 +1,897 @@
+// Copyright 1998-2017 Epic Games, Inc. All Rights Reserved.
+
+/*=============================================================================
+	Random.usf: A pseudo-random number generator.
+=============================================================================*/
+
+#ifndef __Random_usf__
+#define __Random_usf__
+
+// @param xy should be a integer position (e.g. pixel position on the screen), repeats each 128x128 pixels
+// similar to a texture lookup but is only ALU
+// ~13 ALU operations (3 frac, 6 *, 4 mad)
+float PseudoRandom(float2 xy)
+{
+	float2 pos = frac(xy / 128.0f) * 128.0f + float2(-64.340622f, -72.465622f);
+	
+	// found by experimentation
+	return frac(dot(pos.xyx * pos.xyy, float3(20.390625f, 60.703125f, 2.4281209f)));
+}
+
+// high frequency dither pattern appearing almost random without banding steps
+//note: from "NEXT GENERATION POST PROCESSING IN CALL OF DUTY: ADVANCED WARFARE"
+//      http://advances.realtimerendering.com/s2014/index.html
+// Epic extended by FrameId
+// ~7 ALU operations (2 frac, 3 mad, 2 *)
+// @return 0..1
+float InterleavedGradientNoise( float2 uv, float FrameId )
+{
+	// magic values are found by experimentation
+	uv += FrameId * (float2(47, 17) * 0.695f);
+
+    const float3 magic = float3( 0.06711056f, 0.00583715f, 52.9829189f );
+    return frac(magic.z * frac(dot(uv, magic.xy)));
+}
+
+// [0, 1[
+// ~10 ALU operations (2 frac, 5 *, 3 mad)
+float RandFast( uint2 PixelPos, float Magic = 3571.0 )
+{
+	float2 Random2 = ( 1.0 / 4320.0 ) * PixelPos + float2( 0.25, 0.0 );
+	float Random = frac( dot( Random2 * Random2, Magic ) );
+	Random = frac( Random * Random * (2 * Magic) );
+	return Random;
+}
+
+// This is the largest prime < 2^12 so s*s will fit in a 24-bit floating point mantissa
+#define BBS_PRIME24 4093
+
+// Blum-Blum-Shub-inspired pseudo random number generator
+// http://www.umbc.edu/~olano/papers/mNoise.pdf
+// real BBS uses ((s*s) mod M) with bignums and M as the product of two huge Blum primes
+// instead, we use a single prime M just small enough not to overflow
+// note that the above paper used 61, which fits in a half, but is unusably bad
+// @param Integer valued floating point seed
+// @return random number in range [0,1)
+// ~8 ALU operations (5 *, 3 frac)
+float RandBBSfloat(float seed)
+{
+	float s = frac(seed / BBS_PRIME24);
+	s = frac(s * s * BBS_PRIME24);
+	s = frac(s * s * BBS_PRIME24);
+	return s;
+}
+
+// 3D random number generator inspired by PCGs (permuted congruential generator)
+// Using a **simple** Feistel cipher in place of the usual xor shift permutation step
+// @param v = 3D integer coordinate
+// @return three elements w/ 16 random bits each (0-0xffff).
+// ~8 ALU operations for result.x    (7 mad, 1 >>)
+// ~10 ALU operations for result.xy  (8 mad, 2 >>)
+// ~12 ALU operations for result.xyz (9 mad, 3 >>)
+uint3 Rand3DPCG16(int3 p)
+{
+	// taking a signed int then reinterpreting as unsigned gives good behavior for negatives
+	uint3 v = uint3(p);
+
+	// Linear congruential step. These LCG constants are from Numerical Recipies
+	// For additional #'s, PCG would do multiple LCG steps and scramble each on output
+	// So v here is the RNG state
+	v = v * 1664525u + 1013904223u;
+
+	// PCG uses xorshift for the final shuffle, but it is expensive (and cheap
+	// versions of xorshift have visible artifacts). Instead, use simple MAD Feistel steps
+	//
+	// Feistel ciphers divide the state into separate parts (usually by bits)
+	// then apply a series of permutation steps one part at a time. The permutations
+	// use a reversible operation (usually ^) to part being updated with the result of
+	// a permutation function on the other parts and the key.
+	//
+	// In this case, I'm using v.x, v.y and v.z as the parts, using + instead of ^ for
+	// the combination function, and just multiplying the other two parts (no key) for 
+	// the permutation function.
+	//
+	// That gives a simple mad per round.
+	v.x += v.y*v.z;
+	v.y += v.z*v.x;
+	v.z += v.x*v.y;
+	v.x += v.y*v.z;
+	v.y += v.z*v.x;
+	v.z += v.x*v.y;
+
+	// only top 16 bits are well shuffled
+	return v >> 16u;
+}
+
+// 3D random number generator inspired by PCGs (permuted congruential generator)
+// Using a **simple** Feistel cipher in place of the usual xor shift permutation step
+// @param v = 3D integer coordinate
+// @return three elements w/ 32 random bits each (0-0xffffffff).
+// ~12 ALU operations for result.x   (10 mad, 3 >>)
+// ~14 ALU operations for result.xy  (11 mad, 3 >>)
+// ~15 ALU operations for result.xyz (12 mad, 3 >>)
+uint3 Rand3DPCG32(int3 p)
+{
+	// taking a signed int then reinterpreting as unsigned gives good behavior for negatives
+	uint3 v = uint3(p);
+
+	// Linear congruential step.
+	v = v * 1664525u + 1013904223u;
+
+	// swapping low and high bits makes all 32 bits pretty good
+	v = v * (1u << 16u) + (v >> 16u);
+
+	// final shuffle
+	v.x += v.y*v.z;
+	v.y += v.z*v.x;
+	v.z += v.x*v.y;
+	v.x += v.y*v.z;
+	v.y += v.z*v.x;
+	v.z += v.x*v.y;
+
+	return v;
+}
+
+
+
+// 4D random number generator inspired by PCGs (permuted congruential generator)
+// Using a **simple** Feistel cipher in place of the usual xor shift permutation step
+// @param v = 4D integer coordinate
+// @return four elements w/ 32 random bits each (0-0xffffffff).
+// ~12 ALU operations for result.x   (10 mad, 3 >>)
+// ~14 ALU operations for result.xy  (11 mad, 3 >>)
+// ~15 ALU operations for result.xyz (12 mad, 3 >>)
+uint4 Rand4DPCG32(int4 p)
+{
+	// taking a signed int then reinterpreting as unsigned gives good behavior for negatives
+	uint4 v = uint4(p);
+
+	// Linear congruential step.
+	v = v * 1664525u + 1013904223u;
+
+	// swapping low and high bits makes all 32 bits pretty good
+	v = v * (1u << 16u) + (v >> 16u);
+
+	// final shuffle
+	v.x += v.y*v.w;
+	v.y += v.z*v.x;
+	v.z += v.x*v.y;
+	v.w += v.y*v.z;
+	v.x += v.y*v.w;
+	v.y += v.z*v.x;
+	v.z += v.x*v.y;
+	v.w += v.y*v.z;
+
+	return v;
+}
+
+
+
+
+/**
+ * Find good arbitrary axis vectors to represent U and V axes of a plane,
+ * given just the normal. Ported from UnMath.h
+ */
+void FindBestAxisVectors(float3 In, out float3 Axis1, out float3 Axis2 )
+{
+	const float3 N = abs(In);
+
+	// Find best basis vectors.
+	if( N.z > N.x && N.z > N.y )
+	{
+		Axis1 = float3(1, 0, 0);
+	}
+	else
+	{
+		Axis1 = float3(0, 0, 1);
+	}
+
+	Axis1 = normalize(Axis1 - In * dot(Axis1, In));
+	Axis2 = cross(Axis1, In);
+}
+
+// References for noise:
+//
+// Improved Perlin noise
+//   http://mrl.nyu.edu/~perlin/noise/
+//   http://http.developer.nvidia.com/GPUGems/gpugems_ch05.html
+// Modified Noise for Evaluation on Graphics Hardware
+//   http://www.csee.umbc.edu/~olano/papers/mNoise.pdf
+// Perlin Noise
+//   http://mrl.nyu.edu/~perlin/doc/oscar.html
+// Fast Gradient Noise
+//   http://prettyprocs.wordpress.com/2012/10/20/fast-perlin-noise
+
+
+// -------- ALU based method ---------
+
+/*
+ * Pseudo random number generator, based on "TEA, a tiny Encrytion Algorithm"
+ * http://citeseer.ist.psu.edu/viewdoc/download?doi=10.1.1.45.281&rep=rep1&type=pdf
+ * http://www.umbc.edu/~olano/papers/index.html#GPUTEA
+ * @param v - old seed (full 32bit range)
+ * @param IterationCount - >=1, bigger numbers cost more performance but improve quality
+ * @return new seed
+ */
+uint2 ScrambleTEA(uint2 v, uint IterationCount = 3)
+{
+	// Start with some random data (numbers can be arbitrary but those have been used by others and seem to work well)
+	uint k[4] ={ 0xA341316Cu , 0xC8013EA4u , 0xAD90777Du , 0x7E95761Eu };
+	
+	uint y = v[0];
+	uint z = v[1];
+	uint sum = 0;
+	
+	UNROLL for(uint i = 0; i < IterationCount; ++i)
+	{
+		sum += 0x9e3779b9;
+		y += ((z << 4u) + k[0]) ^ (z + sum) ^ ((z >> 5u) + k[1]);
+		z += ((y << 4u) + k[2]) ^ (y + sum) ^ ((y >> 5u) + k[3]);
+	}
+
+	return uint2(y, z);
+}
+
+// Wraps noise for tiling texture creation
+// @param v = unwrapped texture parameter
+// @param bTiling = true to tile, false to not tile
+// @param RepeatSize = number of units before repeating
+// @return either original or wrapped coord
+float3 NoiseTileWrap(float3 v,  bool bTiling, float RepeatSize)
+{
+	return bTiling ? (frac(v / RepeatSize) * RepeatSize) : v;
+}
+
+// Evaluate polynomial to get smooth transitions for Perlin noise
+// only needed by Perlin functions in this file
+// scalar(per component): 2 add, 5 mul
+float4 PerlinRamp(float4 t)
+{
+	return t * t * t * (t * (t * 6 - 15) + 10); 
+}
+
+// Analytical derivative of the PerlinRamp polynomial
+// only needed by Perlin functions in this file
+// scalar(per component): 2 add, 5 mul
+float4 PerlinRampDerivative(float4 t)
+{
+	return t * t * (t * (t * 30 - 60) + 30);
+}
+
+#define MGradientMask int3(0x8000, 0x4000, 0x2000)
+#define MGradientScale float3(1. / 0x4000, 1. / 0x2000, 1. / 0x1000)
+// Modified noise gradient term
+// @param seed - random seed for integer lattice position
+// @param offset - [-1,1] offset of evaluation point from lattice point
+// @return gradient direction (xyz) and contribution (w) from this lattice point
+float4 MGradient(int seed, float3 offset)
+{
+	uint rand = Rand3DPCG16(int3(seed,0,0)).x;
+	float3 direction = float3(rand.xxx & MGradientMask) * MGradientScale - 1;
+	return float4(direction, dot(direction, offset));
+}
+
+// compute Perlin and related noise corner seed values
+// @param v = 3D noise argument, use float3(x,y,0) for 2D or float3(x,0,0) for 1D
+// @param bTiling = true to return seed values for a repeating noise pattern
+// @param RepeatSize = integer units before tiling in each dimension
+// @param seed000-seed111 = hash function seeds for the eight corners
+// @return fractional part of v
+float3 NoiseSeeds(float3 v, bool bTiling, float RepeatSize,
+	out float seed000, out float seed001, out float seed010, out float seed011,
+	out float seed100, out float seed101, out float seed110, out float seed111)
+{
+	float3 fv = frac(v);
+	float3 iv = floor(v);
+
+	const float3 primes = float3(19, 47, 101);
+
+	if (bTiling)
+	{	// can't algebraically combine with primes
+		seed000 = dot(primes, NoiseTileWrap(iv, true, RepeatSize));
+		seed100 = dot(primes, NoiseTileWrap(iv + float3(1, 0, 0), true, RepeatSize));
+		seed010 = dot(primes, NoiseTileWrap(iv + float3(0, 1, 0), true, RepeatSize));
+		seed110 = dot(primes, NoiseTileWrap(iv + float3(1, 1, 0), true, RepeatSize));
+		seed001 = dot(primes, NoiseTileWrap(iv + float3(0, 0, 1), true, RepeatSize));
+		seed101 = dot(primes, NoiseTileWrap(iv + float3(1, 0, 1), true, RepeatSize));
+		seed011 = dot(primes, NoiseTileWrap(iv + float3(0, 1, 1), true, RepeatSize));
+		seed111 = dot(primes, NoiseTileWrap(iv + float3(1, 1, 1), true, RepeatSize));
+	}
+	else
+	{	// get to combine offsets with multiplication by primes in this case
+		seed000 = dot(iv, primes);
+		seed100 = seed000 + primes.x;
+		seed010 = seed000 + primes.y;
+		seed110 = seed100 + primes.y;
+		seed001 = seed000 + primes.z;
+		seed101 = seed100 + primes.z;
+		seed011 = seed010 + primes.z;
+		seed111 = seed110 + primes.z;
+	}
+
+	return fv;
+}
+
+// Perlin-style "Modified Noise"
+// http://www.umbc.edu/~olano/papers/index.html#mNoise
+// @param v = 3D noise argument, use float3(x,y,0) for 2D or float3(x,0,0) for 1D
+// @param bTiling = repeat noise pattern
+// @param RepeatSize = integer units before tiling in each dimension
+// @return random number in the range -1 .. 1
+float GradientNoise3D_ALU(float3 v, bool bTiling, float RepeatSize)
+{
+	float seed000, seed001, seed010, seed011, seed100, seed101, seed110, seed111;
+	float3 fv = NoiseSeeds(v, bTiling, RepeatSize, seed000, seed001, seed010, seed011, seed100, seed101, seed110, seed111);
+
+	float rand000 = MGradient(int(seed000), fv - float3(0, 0, 0)).w;
+	float rand100 = MGradient(int(seed100), fv - float3(1, 0, 0)).w;
+	float rand010 = MGradient(int(seed010), fv - float3(0, 1, 0)).w;
+	float rand110 = MGradient(int(seed110), fv - float3(1, 1, 0)).w;
+	float rand001 = MGradient(int(seed001), fv - float3(0, 0, 1)).w;
+	float rand101 = MGradient(int(seed101), fv - float3(1, 0, 1)).w;
+	float rand011 = MGradient(int(seed011), fv - float3(0, 1, 1)).w;
+	float rand111 = MGradient(int(seed111), fv - float3(1, 1, 1)).w;
+
+	float3 Weights = PerlinRamp(float4(fv, 0)).xyz;
+
+	float i = lerp(lerp(rand000, rand100, Weights.x), lerp(rand010, rand110, Weights.x), Weights.y);
+	float j = lerp(lerp(rand001, rand101, Weights.x), lerp(rand011, rand111, Weights.x), Weights.y);
+	return lerp(i, j, Weights.z).x;
+}
+
+// Coordinates for corners of a Simplex tetrahedron
+// Based on McEwan et al., Efficient computation of noise in GLSL, JGT 2011
+// @param v = 3D noise argument
+// @return 4 corner locations
+float4x3 SimplexCorners(float3 v)
+{
+	// find base corner by skewing to tetrahedral space and back
+	float3 tet = floor(v + v.x/3 + v.y/3 + v.z/3);
+	float3 base = tet - tet.x/6 - tet.y/6 - tet.z/6;
+	float3 f = v - base;
+
+	// Find offsets to other corners (McEwan did this in tetrahedral space,
+	// but since skew is along x=y=z axis, this works in Euclidean space too.)
+	float3 g = step(f.yzx, f.xyz), h = 1 - g.zxy;
+	float3 a1 = min(g, h) - 1. / 6., a2 = max(g, h) - 1. / 3.;
+
+	// four corners
+	return float4x3(base, base + a1, base + a2, base + 0.5);
+}
+
+// Improved smoothing function for simplex noise
+// @param f = fractional distance to four tetrahedral corners
+// @return weight for each corner
+float4 SimplexSmooth(float4x3 f)
+{
+	const float scale = 1024. / 375.;	// scale factor to make noise -1..1
+	float4 d = float4(dot(f[0], f[0]), dot(f[1], f[1]), dot(f[2], f[2]), dot(f[3], f[3]));
+	float4 s = saturate(2 * d);
+	return (1 * scale + s*(-3 * scale + s*(3 * scale - s*scale)));
+}
+
+// Derivative of simplex noise smoothing function
+// @param f = fractional distanc eto four tetrahedral corners
+// @return derivative of smoothing function for each corner by x, y and z
+float3x4 SimplexDSmooth(float4x3 f)
+{
+	const float scale = 1024. / 375.;	// scale factor to make noise -1..1
+	float4 d = float4(dot(f[0], f[0]), dot(f[1], f[1]), dot(f[2], f[2]), dot(f[3], f[3]));
+	float4 s = saturate(2 * d);
+	s = -12 * scale + s*(24 * scale - s * 12 * scale);
+
+	return float3x4(
+		s * float4(f[0][0], f[1][0], f[2][0], f[3][0]),
+		s * float4(f[0][1], f[1][1], f[2][1], f[3][1]),
+		s * float4(f[0][2], f[1][2], f[2][2], f[3][2]));
+}
+
+// Simplex noise and its Jacobian derivative
+// @param v = 3D noise argument
+// @param bTiling = whether to repeat noise pattern
+// @param RepeatSize = integer units before tiling in each dimension, must be a multiple of 3
+// @return float3x3 Jacobian in J[*].xyz, vector noise in J[*].w
+//     J[0].w, J[1].w, J[2].w is a Perlin-style simplex noise with vector output, e.g. (Nx, Ny, Nz)
+//     J[i].x is X derivative of the i'th component of the noise so J[2].x is dNz/dx
+// You can use this to compute the noise, gradient, curl, or divergence:
+//   float3x4 J = JacobianSimplex_ALU(...);
+//   float3 VNoise = float3(J[0].w, J[1].w, J[2].w);	// 3D noise
+//   float3 Grad = J[0].xyz;							// gradient of J[0].w
+//   float3 Curl = float3(J[1][2]-J[2][1], J[2][0]-J[0][2], J[0][1]-J[1][2]);
+//   float Div = J[0][0]+J[1][1]+J[2][2];
+// All of these are confirmed to compile out all unneeded terms.
+// So Grad of X doesn't compute Y or Z components, and VNoise doesn't do any of the derivative computation.
+float3x4 JacobianSimplex_ALU(float3 v, bool bTiling, float RepeatSize)
+{
+	// corners of tetrahedron
+	float4x3 T = SimplexCorners(v);
+	uint3 rand;
+	float4x3 gvec[3], fv;
+	float3x4 grad;
+
+	// processing of tetrahedral vertices, unrolled
+	// to compute gradient at each corner
+	fv[0] = v - T[0];
+	rand = Rand3DPCG16(int3(floor(NoiseTileWrap(6 * T[0] + 0.5, bTiling, RepeatSize))));
+	gvec[0][0] = float3(rand.xxx & MGradientMask) * MGradientScale - 1;
+	gvec[1][0] = float3(rand.yyy & MGradientMask) * MGradientScale - 1;
+	gvec[2][0] = float3(rand.zzz & MGradientMask) * MGradientScale - 1;
+	grad[0][0] = dot(gvec[0][0], fv[0]);
+	grad[1][0] = dot(gvec[1][0], fv[0]);
+	grad[2][0] = dot(gvec[2][0], fv[0]);
+
+	fv[1] = v - T[1];
+	rand = Rand3DPCG16(int3(floor(NoiseTileWrap(6 * T[1] + 0.5, bTiling, RepeatSize))));
+	gvec[0][1] = float3(rand.xxx & MGradientMask) * MGradientScale - 1;
+	gvec[1][1] = float3(rand.yyy & MGradientMask) * MGradientScale - 1;
+	gvec[2][1] = float3(rand.zzz & MGradientMask) * MGradientScale - 1;
+	grad[0][1] = dot(gvec[0][1], fv[1]);
+	grad[1][1] = dot(gvec[1][1], fv[1]);
+	grad[2][1] = dot(gvec[2][1], fv[1]);
+
+	fv[2] = v - T[2];
+	rand = Rand3DPCG16(int3(floor(NoiseTileWrap(6 * T[2] + 0.5, bTiling, RepeatSize))));
+	gvec[0][2] = float3(rand.xxx & MGradientMask) * MGradientScale - 1;
+	gvec[1][2] = float3(rand.yyy & MGradientMask) * MGradientScale - 1;
+	gvec[2][2] = float3(rand.zzz & MGradientMask) * MGradientScale - 1;
+	grad[0][2] = dot(gvec[0][2], fv[2]);
+	grad[1][2] = dot(gvec[1][2], fv[2]);
+	grad[2][2] = dot(gvec[2][2], fv[2]);
+
+	fv[3] = v - T[3];
+	rand = Rand3DPCG16(int3(floor(NoiseTileWrap(6 * T[3] + 0.5, bTiling, RepeatSize))));
+	gvec[0][3] = float3(rand.xxx & MGradientMask) * MGradientScale - 1;
+	gvec[1][3] = float3(rand.yyy & MGradientMask) * MGradientScale - 1;
+	gvec[2][3] = float3(rand.zzz & MGradientMask) * MGradientScale - 1;
+	grad[0][3] = dot(gvec[0][3], fv[3]);
+	grad[1][3] = dot(gvec[1][3], fv[3]);
+	grad[2][3] = dot(gvec[2][3], fv[3]);
+
+	// blend gradients
+	float4 sv = SimplexSmooth(fv);
+	float3x4 ds = SimplexDSmooth(fv);
+
+	float3x4 jacobian;
+	jacobian[0] = float4(mul(sv, gvec[0]) + mul(ds, grad[0]), dot(sv, grad[0]));
+	jacobian[1] = float4(mul(sv, gvec[1]) + mul(ds, grad[1]), dot(sv, grad[1]));
+	jacobian[2] = float4(mul(sv, gvec[2]) + mul(ds, grad[2]), dot(sv, grad[2]));
+
+	return jacobian;
+}
+
+// 3D value noise - used to be incorrectly called Perlin noise
+// @param v = 3D noise argument, use float3(x,y,0) for 2D or float3(x,0,0) for 1D
+// @param bTiling = repeat noise pattern
+// @param RepeatSize = integer units before tiling in each dimension
+// @return random number in the range -1 .. 1
+float ValueNoise3D_ALU(float3 v, bool bTiling, float RepeatSize)
+{
+	float seed000, seed001, seed010, seed011, seed100, seed101, seed110, seed111;
+	float3 fv = NoiseSeeds(v, bTiling, RepeatSize, seed000, seed001, seed010, seed011, seed100, seed101, seed110, seed111);
+
+	float rand000 = RandBBSfloat(seed000) * 2 - 1;
+	float rand100 = RandBBSfloat(seed100) * 2 - 1;
+	float rand010 = RandBBSfloat(seed010) * 2 - 1;
+	float rand110 = RandBBSfloat(seed110) * 2 - 1;
+	float rand001 = RandBBSfloat(seed001) * 2 - 1;
+	float rand101 = RandBBSfloat(seed101) * 2 - 1;
+	float rand011 = RandBBSfloat(seed011) * 2 - 1;
+	float rand111 = RandBBSfloat(seed111) * 2 - 1;
+	
+	float3 Weights = PerlinRamp(float4(fv, 0)).xyz;
+	
+	float i = lerp(lerp(rand000, rand100, Weights.x), lerp(rand010, rand110, Weights.x), Weights.y);
+	float j = lerp(lerp(rand001, rand101, Weights.x), lerp(rand011, rand111, Weights.x), Weights.y);
+	return lerp(i, j, Weights.z).x;
+}
+
+
+// -------- TEX based methods ---------
+
+// filtered 3D noise, can be optimized
+// @param v = 3D noise argument, use float3(x,y,0) for 2D or float3(x,0,0) for 1D
+// @param bTiling = repeat noise pattern
+// @param RepeatSize = integer units before tiling in each dimension
+// @return random number in the range -1 .. 1
+float GradientNoise3D_TEX(float3 v, bool bTiling, float RepeatSize)
+{
+	bTiling = true;
+	float3 fv = frac(v);
+	float3 iv0 = NoiseTileWrap(floor(v), bTiling, RepeatSize);
+	float3 iv1 = NoiseTileWrap(iv0 + 1, bTiling, RepeatSize);
+
+	const int2 ZShear = int2(17, 89);
+	
+	float2 OffsetA = iv0.z * ZShear;
+	float2 OffsetB = OffsetA + ZShear;	// non-tiling, use relative offset
+	if (bTiling)						// tiling, have to compute from wrapped coordinates
+	{
+		OffsetB = iv1.z * ZShear;
+	}
+
+	// Texture size scale factor
+	float ts = 1 / 128.0f;
+
+	// texture coordinates for iv0.xy, as offset for both z slices
+	float2 TexA0 = (iv0.xy + OffsetA + 0.5f) * ts;
+	float2 TexB0 = (iv0.xy + OffsetB + 0.5f) * ts;
+
+	// texture coordinates for iv1.xy, as offset for both z slices
+	float2 TexA1 = TexA0 + ts;	// for non-tiling, can compute relative to existing coordinates
+	float2 TexB1 = TexB0 + ts;
+	if (bTiling)				// for tiling, need to compute from wrapped coordinates
+	{
+		TexA1 = (iv1.xy + OffsetA + 0.5f) * ts;
+		TexB1 = (iv1.xy + OffsetB + 0.5f) * ts;
+	}
+
+
+	// can be optimized to 1 or 2 texture lookups (4 or 8 channel encoded in 8, 16 or 32 bit)
+	float3 A = Texture2DSampleLevel(View.PerlinNoiseGradientTexture, View.PerlinNoiseGradientTextureSampler, float2(TexA0.x, TexA0.y), 0).xyz * 2 - 1;
+	float3 B = Texture2DSampleLevel(View.PerlinNoiseGradientTexture, View.PerlinNoiseGradientTextureSampler, float2(TexA1.x, TexA0.y), 0).xyz * 2 - 1;
+	float3 C = Texture2DSampleLevel(View.PerlinNoiseGradientTexture, View.PerlinNoiseGradientTextureSampler, float2(TexA0.x, TexA1.y), 0).xyz * 2 - 1;
+	float3 D = Texture2DSampleLevel(View.PerlinNoiseGradientTexture, View.PerlinNoiseGradientTextureSampler, float2(TexA1.x, TexA1.y), 0).xyz * 2 - 1;
+	float3 E = Texture2DSampleLevel(View.PerlinNoiseGradientTexture, View.PerlinNoiseGradientTextureSampler, float2(TexB0.x, TexB0.y), 0).xyz * 2 - 1;
+	float3 F = Texture2DSampleLevel(View.PerlinNoiseGradientTexture, View.PerlinNoiseGradientTextureSampler, float2(TexB1.x, TexB0.y), 0).xyz * 2 - 1;
+	float3 G = Texture2DSampleLevel(View.PerlinNoiseGradientTexture, View.PerlinNoiseGradientTextureSampler, float2(TexB0.x, TexB1.y), 0).xyz * 2 - 1;
+	float3 H = Texture2DSampleLevel(View.PerlinNoiseGradientTexture, View.PerlinNoiseGradientTextureSampler, float2(TexB1.x, TexB1.y), 0).xyz * 2 - 1;
+
+	float a = dot(A, fv - float3(0, 0, 0));
+	float b = dot(B, fv - float3(1, 0, 0));
+	float c = dot(C, fv - float3(0, 1, 0));
+	float d = dot(D, fv - float3(1, 1, 0));
+	float e = dot(E, fv - float3(0, 0, 1));
+	float f = dot(F, fv - float3(1, 0, 1));
+	float g = dot(G, fv - float3(0, 1, 1));
+	float h = dot(H, fv - float3(1, 1, 1));
+
+	float3 Weights = PerlinRamp(frac(float4(fv, 0))).xyz;
+	
+	float i = lerp(lerp(a, b, Weights.x), lerp(c, d, Weights.x), Weights.y);
+	float j = lerp(lerp(e, f, Weights.x), lerp(g, h, Weights.x), Weights.y);
+
+	return lerp(i, j, Weights.z);
+}
+
+// @return random number in the range -1 .. 1
+// scalar: 6 frac, 31 mul/mad, 15 add, 
+float FastGradientPerlinNoise3D_TEX(float3 xyz)
+{
+	// needs to be the same value when creating the PerlinNoise3D texture
+	float Extent = 16;
+
+	// last texel replicated and needed for filtering
+	// scalar: 3 frac, 6 mul
+	xyz = frac(xyz / (Extent - 1)) * (Extent - 1);
+
+	// scalar: 3 frac
+	float3 uvw = frac(xyz);
+	// = floor(xyz);
+	// scalar: 3 add
+	float3 p0 = xyz - uvw;
+//	float3 f = pow(uvw, 2) * 3.0f - pow(uvw, 3) * 2.0f;	// original perlin hermite (ok when used without bump mapping)
+	// scalar: 2*3 add 5*3 mul
+	float3 f = PerlinRamp(float4(uvw, 0)).xyz;	// new, better with continues second derivative for bump mapping
+	// scalar: 3 add
+	float3 p = p0 + f;
+	// scalar: 3 mad
+	float4 NoiseSample = Texture3DSampleLevel(View.PerlinNoise3DTexture, View.PerlinNoise3DTextureSampler, p / Extent + 0.5f / Extent, 0);		// +0.5f to get rid of bilinear offset
+
+	// reconstruct from 8bit (using mad with 2 constants and dot4 was same instruction count)
+	// scalar: 4 mad, 3 mul, 3 add 
+	float3 n = NoiseSample.xyz * 255.0f / 127.0f - 1.0f;
+	float d = NoiseSample.w * 255.f - 127;
+	return dot(xyz, n) - d;
+}
+
+
+// 3D jitter offset within a voronoi noise cell
+// @param pos - integer lattice corner
+// @return random offsets vector
+float3 VoronoiCornerSample(float3 pos, int Quality)
+{
+	// random values in [-0.5, 0.5]
+	float3 noise = float3(Rand3DPCG16(int3(pos))) / 0xffff - 0.5;
+
+	// quality level 1 or 2: searches a 2x2x2 neighborhood with points distributed on a sphere
+	// scale factor to guarantee jittered points will be found within a 2x2x2 search
+	if (Quality <= 2)
+	{
+		return normalize(noise) * 0.2588;
+	}
+
+	// quality level 3: searches a 3x3x3 neighborhood with points distributed on a sphere
+	// scale factor to guarantee jittered points will be found within a 3x3x3 search
+	if (Quality == 3)
+	{
+		return normalize(noise) * 0.3090;
+	}
+
+	// quality level 4: jitter to anywhere in the cell, needs 4x4x4 search
+	return noise;
+}
+
+// compare previous best with a new candidate
+// not producing point locations makes it easier for compiler to eliminate calculations when they're not needed
+// @param minval = location and distance of best candidate seed point before the new one
+// @param candidate = candidate seed point
+// @param offset = 3D offset to new candidate seed point
+// @param bDistanceOnly = if true, only set maxval.w with distance, otherwise maxval.w is distance and maxval.xyz is position
+// @return position (if bDistanceOnly is false) and distance to closest seed point so far
+float4 VoronoiCompare(float4 minval, float3 candidate, float3 offset, bool bDistanceOnly)
+{
+	if (bDistanceOnly)
+	{
+		return float4(0, 0, 0, min(minval.w, dot(offset, offset)));
+	}
+	else
+	{
+		float newdist = dot(offset, offset);
+		return newdist > minval.w ? minval : float4(candidate, newdist);
+	}
+}
+
+// 220 instruction Worley noise
+float4 VoronoiNoise3D_ALU(float3 v, int Quality, bool bTiling, float RepeatSize, bool bDistanceOnly)
+{
+	float3 fv = frac(v),  fv2 = frac(v + 0.5);
+	float3 iv = floor(v), iv2 = floor(v + 0.5);
+
+	// with initial minimum distance = infinity (or at least bigger than 4), first min is optimized away
+	float4 mindist = float4(0,0,0,100);
+	float3 p, offset;
+
+	// quality level 3: do a 3x3x3 search
+	if (Quality == 3)
+	{
+		UNROLL for (offset.x = -1; offset.x <= 1; ++offset.x)
+		{
+			UNROLL for (offset.y = -1; offset.y <= 1; ++offset.y)
+			{
+				UNROLL for (offset.z = -1; offset.z <= 1; ++offset.z)
+				{
+					p = offset + VoronoiCornerSample(NoiseTileWrap(iv2 + offset, bTiling, RepeatSize), Quality);
+					mindist = VoronoiCompare(mindist, iv2 + p, fv2 - p, bDistanceOnly);
+				}
+			}
+		}
+	}
+
+	// everybody else searches a base 2x2x2 neighborhood
+	else
+	{
+		UNROLL for (offset.x = 0; offset.x <= 1; ++offset.x)
+		{
+			UNROLL for (offset.y = 0; offset.y <= 1; ++offset.y)
+			{
+				UNROLL for (offset.z = 0; offset.z <= 1; ++offset.z)
+				{
+					p = offset + VoronoiCornerSample(NoiseTileWrap(iv + offset, bTiling, RepeatSize), Quality);
+					mindist = VoronoiCompare(mindist, iv + p, fv - p, bDistanceOnly);
+
+					// quality level 2, do extra set of points, offset by half a cell
+					if (Quality == 2)
+					{
+						// 467 is just an offset to a different area in the random number field to avoid similar neighbor artifacts
+						p = offset + VoronoiCornerSample(NoiseTileWrap(iv2 + offset, bTiling, RepeatSize) + 467, Quality);
+						mindist = VoronoiCompare(mindist, iv2 + p, fv2 - p, bDistanceOnly);
+					}
+				}
+			}
+		}
+	}
+
+	// quality level 4: add extra sets of four cells in each direction
+	if (Quality >= 4)
+	{
+		UNROLL for (offset.x = -1; offset.x <= 2; offset.x += 3)
+		{
+			UNROLL for (offset.y = 0; offset.y <= 1; ++offset.y)
+			{
+				UNROLL for (offset.z = 0; offset.z <= 1; ++offset.z)
+				{
+					// along x axis
+					p = offset.xyz + VoronoiCornerSample(NoiseTileWrap(iv + offset.xyz, bTiling, RepeatSize), Quality);
+					mindist = VoronoiCompare(mindist, iv + p, fv - p, bDistanceOnly);
+
+					// along y axis
+					p = offset.yzx + VoronoiCornerSample(NoiseTileWrap(iv + offset.yzx, bTiling, RepeatSize), Quality);
+					mindist = VoronoiCompare(mindist, iv + p, fv - p, bDistanceOnly);
+
+					// along z axis
+					p = offset.zxy + VoronoiCornerSample(NoiseTileWrap(iv + offset.zxy, bTiling, RepeatSize), Quality);
+					mindist = VoronoiCompare(mindist, iv + p, fv - p, bDistanceOnly);
+				}
+			}
+		}
+	}
+
+	// transform squared distance to real distance
+	return float4(mindist.xyz, sqrt(mindist.w));
+}
+
+
+// -------- Simplex method (faster in higher dimensions because less samples are used, uses gradient noise for quality) ---------
+// <Dimensions>D:<Normal>/<Simplex> 1D:2, 2D:4/3, 3D:8/4, 4D:16/5 
+
+// Computed weights and sample positions for simplex interpolation
+// @return float3(a,b,c) Barycentric coordianate defined as Filtered = Tex(PosA) * a + Tex(PosB) * b + Tex(PosC) * c
+float3 ComputeSimplexWeights2D(float2 OrthogonalPos, out float2 PosA, out float2 PosB, out float2 PosC)
+{
+	float2 OrthogonalPosFloor = floor(OrthogonalPos); 
+	PosA = OrthogonalPosFloor;
+	PosB = PosA + float2(1, 1); 
+
+	float2 LocalPos = OrthogonalPos - OrthogonalPosFloor;
+
+	PosC = PosA + ((LocalPos.x > LocalPos.y) ? float2(1,0) : float2(0,1));
+
+	float b = min(LocalPos.x, LocalPos.y);
+	float c = abs(LocalPos.y - LocalPos.x);
+	float a = 1.0f - b - c;
+
+	return float3(a, b, c);
+}
+
+// Computed weights and sample positions for simplex interpolation
+// @return float4(a,b,c, d) Barycentric coordinate defined as Filtered = Tex(PosA) * a + Tex(PosB) * b + Tex(PosC) * c + Tex(PosD) * d
+float4 ComputeSimplexWeights3D(float3 OrthogonalPos, out float3 PosA, out float3 PosB, out float3 PosC, out float3 PosD)
+{
+	float3 OrthogonalPosFloor = floor(OrthogonalPos);
+
+	PosA = OrthogonalPosFloor;
+	PosB = PosA + float3(1, 1, 1);
+
+	OrthogonalPos -= OrthogonalPosFloor;
+
+	float Largest = max(OrthogonalPos.x, max(OrthogonalPos.y, OrthogonalPos.z));
+	float Smallest = min(OrthogonalPos.x, min(OrthogonalPos.y, OrthogonalPos.z));
+
+	PosC = PosA + float3(Largest == OrthogonalPos.x, Largest == OrthogonalPos.y, Largest == OrthogonalPos.z);
+	PosD = PosA + float3(Smallest != OrthogonalPos.x, Smallest != OrthogonalPos.y, Smallest != OrthogonalPos.z);
+
+	float4 ret;
+
+	float RG = OrthogonalPos.x - OrthogonalPos.y;
+	float RB = OrthogonalPos.x - OrthogonalPos.z;
+	float GB = OrthogonalPos.y - OrthogonalPos.z;
+
+	ret.b = 
+		  min(max(0, RG), max(0, RB))		// X
+		+ min(max(0, -RG), max(0, GB))		// Y
+		+ min(max(0, -RB), max(0, -GB));	// Z
+	
+	ret.a = 
+		  min(max(0, -RG), max(0, -RB))		// X
+		+ min(max(0, RG), max(0, -GB))		// Y
+		+ min(max(0, RB), max(0, GB));		// Z
+
+	ret.g = Smallest;
+	ret.r = 1.0f - ret.g - ret.b - ret.a;
+
+	return ret;
+}
+
+float2 GetPerlinNoiseGradientTextureAt(float2 v)
+{
+	float2 TexA = (v.xy + 0.5f) / 128.0f;
+
+	// todo: storing random 2d unit vectors would be better
+	float3 p = Texture2DSampleLevel(View.PerlinNoiseGradientTexture, View.PerlinNoiseGradientTextureSampler, TexA, 0).xyz * 2 - 1;
+	return normalize(p.xy + p.z * 0.33f);
+}
+
+float3 GetPerlinNoiseGradientTextureAt(float3 v)
+{
+	const float2 ZShear = int2(17, 89);
+
+	float2 OffsetA = v.z * ZShear;
+	float2 TexA = (v.xy + OffsetA + 0.5f) / 128.0f;
+
+	return Texture2DSampleLevel(View.PerlinNoiseGradientTexture, View.PerlinNoiseGradientTextureSampler, TexA , 0).xyz * 2 - 1;
+}
+
+float2 SkewSimplex(float2 In)
+{
+	return In + dot(In, (sqrt(3.0f) - 1.0f) * 0.5f );
+}
+float2 UnSkewSimplex(float2 In)
+{
+	return In - dot(In, (3.0f - sqrt(3.0f)) / 6.0f );
+}
+float3 SkewSimplex(float3 In)
+{
+	return In + dot(In, 1.0 / 3.0f );
+}
+float3 UnSkewSimplex(float3 In)
+{
+	return In - dot(In, 1.0 / 6.0f );
+}
+
+// filtered 3D gradient simple noise (few texture lookups, high quality)
+// @param v >0
+// @return random number in the range -1 .. 1
+float GradientSimplexNoise2D_TEX(float2 EvalPos)
+{
+	float2 OrthogonalPos = SkewSimplex(EvalPos);
+
+	float2 PosA, PosB, PosC, PosD;
+	float3 Weights = ComputeSimplexWeights2D(OrthogonalPos, PosA, PosB, PosC);
+
+	// can be optimized to 1 or 2 texture lookups (4 or 8 channel encoded in 32 bit)
+	float2 A = GetPerlinNoiseGradientTextureAt(PosA);
+	float2 B = GetPerlinNoiseGradientTextureAt(PosB);
+	float2 C = GetPerlinNoiseGradientTextureAt(PosC);
+
+	PosA = UnSkewSimplex(PosA);
+	PosB = UnSkewSimplex(PosB);
+	PosC = UnSkewSimplex(PosC);
+
+	float DistanceWeight;
+
+	DistanceWeight = saturate(0.5f - length2(EvalPos - PosA));	DistanceWeight *= DistanceWeight; DistanceWeight *= DistanceWeight;
+	float a = dot(A, EvalPos - PosA) * DistanceWeight;
+	DistanceWeight = saturate(0.5f - length2(EvalPos - PosB));	DistanceWeight *= DistanceWeight; DistanceWeight *= DistanceWeight;
+	float b = dot(B, EvalPos - PosB) * DistanceWeight;
+	DistanceWeight = saturate(0.5f - length2(EvalPos - PosC));	DistanceWeight *= DistanceWeight; DistanceWeight *= DistanceWeight;
+	float c = dot(C, EvalPos - PosC) * DistanceWeight;
+	
+	return 70 * (a + b + c);
+}
+
+
+
+// filtered 3D gradient simple noise (few texture lookups, high quality)
+// @param v >0
+// @return random number in the range -1 .. 1
+float SimplexNoise3D_TEX(float3 EvalPos)
+{
+	float3 OrthogonalPos = SkewSimplex(EvalPos);
+
+	float3 PosA, PosB, PosC, PosD;
+	float4 Weights = ComputeSimplexWeights3D(OrthogonalPos, PosA, PosB, PosC, PosD);
+
+	// can be optimized to 1 or 2 texture lookups (4 or 8 channel encoded in 32 bit)
+	float3 A = GetPerlinNoiseGradientTextureAt(PosA);
+	float3 B = GetPerlinNoiseGradientTextureAt(PosB);
+	float3 C = GetPerlinNoiseGradientTextureAt(PosC);
+	float3 D = GetPerlinNoiseGradientTextureAt(PosD);
+	
+	PosA = UnSkewSimplex(PosA);
+	PosB = UnSkewSimplex(PosB);
+	PosC = UnSkewSimplex(PosC);
+	PosD = UnSkewSimplex(PosD);
+
+	float DistanceWeight;
+
+	DistanceWeight = saturate(0.6f - length2(EvalPos - PosA));	DistanceWeight *= DistanceWeight; DistanceWeight *= DistanceWeight;
+	float a = dot(A, EvalPos - PosA) * DistanceWeight;
+	DistanceWeight = saturate(0.6f - length2(EvalPos - PosB));	DistanceWeight *= DistanceWeight; DistanceWeight *= DistanceWeight;
+	float b = dot(B, EvalPos - PosB) * DistanceWeight;
+	DistanceWeight = saturate(0.6f - length2(EvalPos - PosC));	DistanceWeight *= DistanceWeight; DistanceWeight *= DistanceWeight;
+	float c = dot(C, EvalPos - PosC) * DistanceWeight;
+	DistanceWeight = saturate(0.6f - length2(EvalPos - PosD));	DistanceWeight *= DistanceWeight; DistanceWeight *= DistanceWeight;
+	float d = dot(D, EvalPos - PosD) * DistanceWeight;
+
+	return 32 * (a + b + c + d);
+}
+
+
+float VolumeRaymarch(float3 posPixelWS, float3 posCameraWS)
+{
+	float ret = 0;
+	int cnt = 60;
+
+	LOOP for(int i=0; i < cnt; ++i)
+	{
+		ret += saturate(FastGradientPerlinNoise3D_TEX(lerp(posPixelWS, posCameraWS, i/(float)cnt) * 0.01) - 0.2f);
+	}
+
+	return ret / cnt * (length(posPixelWS - posCameraWS) * 0.001f );
+}
+
+
+
+
+
+#endif
diff --git a/Engine/Shaders/Private/Random.ush b/Engine/Shaders/Private/Random.ush
--- a/Engine/Shaders/Private/Random.ush
+++ b/Engine/Shaders/Private/Random.ush
@@ -132,6 +132,42 @@ uint3 Rand3DPCG32(int3 p)
 	return v;
 }
 
+
+
+// 4D random number generator inspired by PCGs (permuted congruential generator)
+// Using a **simple** Feistel cipher in place of the usual xor shift permutation step
+// @param v = 4D integer coordinate
+// @return four elements w/ 32 random bits each (0-0xffffffff).
+// ~12 ALU operations for result.x   (10 mad, 3 >>)
+// ~14 ALU operations for result.xy  (11 mad, 3 >>)
+// ~15 ALU operations for result.xyz (12 mad, 3 >>)
+uint4 Rand4DPCG32(int4 p)
+{
+	// taking a signed int then reinterpreting as unsigned gives good behavior for negatives
+	uint4 v = uint4(p);
+
+	// Linear congruential step.
+	v = v * 1664525u + 1013904223u;
+
+	// swapping low and high bits makes all 32 bits pretty good
+	v = v * (1u << 16u) + (v >> 16u);
+
+	// final shuffle
+	v.x += v.y*v.w;
+	v.y += v.z*v.x;
+	v.z += v.x*v.y;
+	v.w += v.y*v.z;
+	v.x += v.y*v.w;
+	v.y += v.z*v.x;
+	v.z += v.x*v.y;
+	v.w += v.y*v.z;
+
+	return v;
+}
+
+
+
+
 /**
  * Find good arbitrary axis vectors to represent U and V axes of a plane,
  * given just the normal. Ported from UnMath.h
diff --git a/Engine/Source/Developer/AssetTools/Private/AssetTypeActions/AssetTypeActions_StaticMesh.cpp b/Engine/Source/Developer/AssetTools/Private/AssetTypeActions/AssetTypeActions_StaticMesh.cpp
--- a/Engine/Source/Developer/AssetTools/Private/AssetTypeActions/AssetTypeActions_StaticMesh.cpp
+++ b/Engine/Source/Developer/AssetTools/Private/AssetTypeActions/AssetTypeActions_StaticMesh.cpp
@@ -6,12 +6,12 @@
 #include "EditorFramework/AssetImportData.h"
 #include "ThumbnailRendering/SceneThumbnailInfo.h"
 #include "AssetTools.h"
-
 #include "Editor/StaticMeshEditor/Public/StaticMeshEditorModule.h"
-
 #include "FbxMeshUtils.h"
 #include "Framework/Notifications/NotificationManager.h"
 #include "Widgets/Notifications/SNotificationList.h"
+#include "MessageDialog.h"
+#include "ScopedSlowTask.h"
 
 #define LOCTEXT_NAMESPACE "AssetTypeActions"
 
@@ -47,6 +47,15 @@ void FAssetTypeActions_StaticMesh::GetActions( const TArray<UObject*>& InObjects
 		false,
 		FSlateIcon(FEditorStyle::GetStyleSetName(), "ContentBrowser.AssetActions")
 		);
+
+	MenuBuilder.AddMenuEntry(
+		NSLOCTEXT("AssetTypeActions_StaticMesh", "ObjectContext_ClearVertexColors", "Remove Vertex Colors"),
+		NSLOCTEXT("AssetTypeActions_StaticMesh", "ObjectContext_ClearVertexColors", "Removes vertex colors from all LODS in all selected meshes."),
+		FSlateIcon(),
+		FUIAction(
+			FExecuteAction::CreateSP(this, &FAssetTypeActions_StaticMesh::ExecuteRemoveVertexColors, Meshes)
+		)
+	);
 }
 
 void FAssetTypeActions_StaticMesh::OpenAssetEditor( const TArray<UObject*>& InObjects, TSharedPtr<IToolkitHost> EditWithinLevelEditor )
@@ -247,4 +256,22 @@ void FAssetTypeActions_StaticMesh::ExecuteSaveGeneratedLODsInPackage(TArray<TWea
 	}
 }
 
+void FAssetTypeActions_StaticMesh::ExecuteRemoveVertexColors(TArray<TWeakObjectPtr<UStaticMesh>> Objects)
+{
+	FText WarningMessage = LOCTEXT("Warning_RemoveVertexColors", "Are you sure you want to remove vertex colors from all selected meshes?  There is no undo available.");
+	if (FMessageDialog::Open(EAppMsgType::YesNo, WarningMessage) == EAppReturnType::Yes)
+	{
+		FScopedSlowTask SlowTask(1.0f, LOCTEXT("RemovingVertexColors", "Removing Vertex Colors"));
+		for (auto StaticMeshPtr : Objects)
+		{
+			bool bRemovedVertexColors = false;
+			UStaticMesh* Mesh = StaticMeshPtr.Get();
+			if (Mesh)
+			{
+				Mesh->RemoveVertexColors();
+			}
+		}
+	}
+}
+
 #undef LOCTEXT_NAMESPACE
diff --git a/Engine/Source/Developer/AssetTools/Private/AssetTypeActions/AssetTypeActions_StaticMesh.h b/Engine/Source/Developer/AssetTools/Private/AssetTypeActions/AssetTypeActions_StaticMesh.h
--- a/Engine/Source/Developer/AssetTools/Private/AssetTypeActions/AssetTypeActions_StaticMesh.h
+++ b/Engine/Source/Developer/AssetTools/Private/AssetTypeActions/AssetTypeActions_StaticMesh.h
@@ -39,9 +39,12 @@ private:
 	/** Whether there is a valid static meshes to copy LOD to*/
 	bool CanPasteLODSettings(TArray<TWeakObjectPtr<UStaticMesh>> Objects) const;
 
-	// Handler for when SaveGeneratedLODsInPackage is selected
+	/** Handler for when SaveGeneratedLODsInPackage is selected */
 	void ExecuteSaveGeneratedLODsInPackage(TArray<TWeakObjectPtr<UStaticMesh>> Objects);
 
+	/** Handler for when RemoveVertexColors is selected */
+	void ExecuteRemoveVertexColors(TArray<TWeakObjectPtr<UStaticMesh>> Objects);
+
 	/** Handler to provide the list of LODs that can be imported or reimported */
 	void GetImportLODMenu(class FMenuBuilder& MenuBuilder,TArray<TWeakObjectPtr<UStaticMesh>> Objects);
 
diff --git a/Engine/Source/Developer/FunctionalTesting/Private/Tests/ClientFuncTestPerforming.cpp b/Engine/Source/Developer/FunctionalTesting/Private/Tests/ClientFuncTestPerforming.cpp
--- a/Engine/Source/Developer/FunctionalTesting/Private/Tests/ClientFuncTestPerforming.cpp
+++ b/Engine/Source/Developer/FunctionalTesting/Private/Tests/ClientFuncTestPerforming.cpp
@@ -8,81 +8,13 @@
 #include "FunctionalTestingModule.h"
 #include "EngineGlobals.h"
 #include "Tests/AutomationCommon.h"
+#include "FunctionalTestingHelper.h"
 
 #if WITH_DEV_AUTOMATION_TESTS
 
 #define LOCTEXT_NAMESPACE "FunctionalTesting"
 
-DEFINE_LATENT_AUTOMATION_COMMAND(FWaitForFTestsToFinish);
-bool FWaitForFTestsToFinish::Update()
-{
-	return IFunctionalTestingModule::Get().IsRunning() == false;
-}
-
-DEFINE_LATENT_AUTOMATION_COMMAND(FTriggerFTests);
-bool FTriggerFTests::Update()
-{
-	IFunctionalTestingModule& Module = IFunctionalTestingModule::Get();
-	if (Module.IsFinished())
-	{
-		// if tests have been already triggered by level script just make sure it's not looping
-		if (Module.IsRunning())
-		{
-			Module.SetLooping(false);
-		}
-		else
-		{
-			Module.RunAllTestsOnMap(false, false);
-		}
-	}
-
-	return true;
-}
-
-
-DEFINE_LATENT_AUTOMATION_COMMAND_ONE_PARAMETER(FTriggerFTest, FString, TestName);
-bool FTriggerFTest::Update()
-{
-	IFunctionalTestingModule& Module = IFunctionalTestingModule::Get();
-	if (Module.IsFinished() )
-	{
-		// if tests have been already triggered by level script just make sure it's not looping
-		if (Module.IsRunning() )
-		{
-			Module.SetLooping(false);
-		}
-		else
-		{
-			Module.RunTestOnMap(TestName, false, false);
-		}
-	}
-
-	return true;
-}
-
-
-DEFINE_LATENT_AUTOMATION_COMMAND(FStartFTestsOnMap);
-bool FStartFTestsOnMap::Update()
-{
-	// This should now be handled by your IsReady override of the functional test.
-	//should really be wait until the map is properly loaded....in PIE or gameplay....
-	//ADD_LATENT_AUTOMATION_COMMAND(FWaitLatentCommand(15.f));
-
-	ADD_LATENT_AUTOMATION_COMMAND(FTriggerFTests);
-	ADD_LATENT_AUTOMATION_COMMAND(FWaitForFTestsToFinish);	
-
-	return true;
-}
-
-
-DEFINE_LATENT_AUTOMATION_COMMAND_ONE_PARAMETER(FStartFTestOnMap, FString, TestName);
-bool FStartFTestOnMap::Update()
-{
-	ADD_LATENT_AUTOMATION_COMMAND(FTriggerFTest(TestName));
-	ADD_LATENT_AUTOMATION_COMMAND(FWaitForFTestsToFinish);
-
-	return true;
-}
+DEFINE_LOG_CATEGORY_STATIC(LogFunctionalTesting, Log, All);
 
 class FClientFunctionalTestingMapsBase : public FAutomationTestBase
 {
diff --git a/Engine/Source/Developer/HierarchicalLODUtilities/Private/HierarchicalLODUtilities.cpp b/Engine/Source/Developer/HierarchicalLODUtilities/Private/HierarchicalLODUtilities.cpp
--- a/Engine/Source/Developer/HierarchicalLODUtilities/Private/HierarchicalLODUtilities.cpp
+++ b/Engine/Source/Developer/HierarchicalLODUtilities/Private/HierarchicalLODUtilities.cpp
@@ -265,34 +265,34 @@ bool FHierarchicalLODUtilities::BuildStaticMeshForLODActor(ALODActor* LODActor,
 	return false;
 }
 
-bool FHierarchicalLODUtilities::ShouldGenerateCluster(AActor* Actor)
+EClusterGenerationError FHierarchicalLODUtilities::ShouldGenerateCluster(AActor* Actor)
 {
 	if (!Actor)
 	{
-		return false;
+		return EClusterGenerationError::InvalidActor;
 	}
 
 	if (Actor->bHidden)
 	{
-		return false;
+		return EClusterGenerationError::ActorHiddenInGame;
 	}
 
 	if (!Actor->bEnableAutoLODGeneration)
 	{
-		return false;
+		return EClusterGenerationError::ExcludedActor;
 	}
 
 	ALODActor* LODActor = Cast<ALODActor>(Actor);
 	if (LODActor)
 	{
-		return false;
+		return EClusterGenerationError::LODActor;
 	}
 
 	FVector Origin, Extent;
 	Actor->GetActorBounds(false, Origin, Extent);
 	if (Extent.SizeSquared() <= 0.1)
 	{
-		return false;
+		return EClusterGenerationError::ActorTooSmall;
 	}
 
 	// for now only consider staticmesh - I don't think skel mesh would work with simplygon merge right now @fixme
@@ -303,30 +303,37 @@ bool FHierarchicalLODUtilities::ShouldGenerateCluster(AActor* Actor)
 
 	int32 ValidComponentCount = 0;
 	// now make sure you check parent primitive, so that we don't build for the actor that already has built. 
+
+	EClusterGenerationError ErrorType = EClusterGenerationError::None;
+
 	if (Components.Num() > 0)
 	{
 		for (auto& ComponentIter : Components)
 		{
 			if (ComponentIter->GetLODParentPrimitive())
 			{
-				return false;
+				return EClusterGenerationError::AlreadyClustered;
 			}
 
 			if (ComponentIter->bHiddenInGame)
 			{
-				return false;
+				return EClusterGenerationError::ComponentHiddenInGame;
 			}
 
 			// see if we should generate it
 			if (ComponentIter->ShouldGenerateAutoLOD())
 			{
 				++ValidComponentCount;
-				break;
+				ErrorType |= EClusterGenerationError::ValidActor;
+			}
+			else
+			{
+				ErrorType |= (ComponentIter->bEnableAutoLODGeneration ? EClusterGenerationError::MoveableComponent: EClusterGenerationError::ExcludedComponent);
 			}
 		}
 	}
 
-	return (ValidComponentCount > 0);
+	return ErrorType;
 }
 
 ALODActor* FHierarchicalLODUtilities::GetParentLODActor(const AActor* InActor)
