Duplicate 3640053:
Fix Desktop Forward rendering on Nvidia Metal rendering on macOS by modifying the shader and runtime to treat ForwardLocalLightBuffer data as uint rather than float and invert the use of as*() casts.

Tthis is *necessary* to avoid a bug/limitation of Nvidia's current Metal shader compiler that flushes all denorm values to zero on load from any resource. AMD & iOS Metal only flush-to-zero when sampling from textures and that's what Apple regard as the expected behaviour. I have however asked them to standardise on the D3D behaviour of preserving denorms on all load, move (incl. min/max/as*()-casts) & store operations. This won't happen in the current or imminent OS/Xcode releases.

Now only Intel Metal is broken and their problems run deeper.

diff --git a/Engine/Shaders/Private/ForwardLightingCommon.ush b/Engine/Shaders/Private/ForwardLightingCommon.ush
--- a/Engine/Shaders/Private/ForwardLightingCommon.ush
+++ b/Engine/Shaders/Private/ForwardLightingCommon.ush
@@ -152,9 +152,9 @@ float3 GetForwardDirectLighting(uint GridIndex, float3 WorldPosition, float3 Cam
 		FDeferredLightData LightData = (FDeferredLightData)0;
 		LightData.LightPositionAndInvRadius = LocalLight.LightPositionAndInvRadius;
 		LightData.LightColorAndFalloffExponent = LocalLight.LightColorAndFalloffExponent;
-		LightData.LightDirection = LocalLight.LightDirectionAndShadowMask.xyz;
-		float SourceLength = f16tof32(asuint(LocalLight.SpotAnglesAndSourceRadiusPacked.w));
-		LightData.SpotAnglesAndSourceRadius = float4(LocalLight.SpotAnglesAndSourceRadiusPacked.xyz, SourceLength);
+		LightData.LightDirection = LocalLight.LightDirection;
+		float SourceLength = f16tof32(asuint(LocalLight.SourceRadiusPacked));
+		LightData.SpotAnglesAndSourceRadius = float4(LocalLight.SpotAngles, SourceLength);
         LightData.LightTangent = LocalLight.LightTangentAndSoftSourceRadius.xyz;
         LightData.SoftSourceRadius = LocalLight.LightTangentAndSoftSourceRadius.w;
 		LightData.bInverseSquared = LightData.LightColorAndFalloffExponent.w == 0;
@@ -162,7 +162,7 @@ float3 GetForwardDirectLighting(uint GridIndex, float3 WorldPosition, float3 Cam
 		LightData.bSpotLight = LightData.SpotAnglesAndSourceRadius.x > -2.0f;
 		LightData.MinRoughness = MinRoughness;
 
-		uint PackedShadowMapChannelMask = asuint(LocalLight.LightDirectionAndShadowMask.w);
+		uint PackedShadowMapChannelMask = LocalLight.ShadowMask;
 		LightData.ShadowedBits = (PackedShadowMapChannelMask & 0xFF) != 0 ? 1 : 0;
 		// Static shadowing uses ShadowMapChannel, dynamic shadows are packed into light attenuation using PreviewShadowMapChannel
 		LightData.ShadowMapChannelMask = UnpackShadowMapChannelMask(PackedShadowMapChannelMask);
@@ -333,4 +333,4 @@ half3 GetImageBasedReflectionLighting(FMaterialPixelParameters MaterialParameter
 half3 GetImageBasedReflectionLighting(FMaterialPixelParameters MaterialParameters, half Roughness, half3 SpecularColor, half IndirectIrradiance, uint GridIndex)
 {
 	return GetImageBasedReflectionLighting(MaterialParameters, Roughness, SpecularColor, IndirectIrradiance, GridIndex, 0);
-}
\ No newline at end of file
+}
diff --git a/Engine/Shaders/Private/LightGridCommon.ush b/Engine/Shaders/Private/LightGridCommon.ush
--- a/Engine/Shaders/Private/LightGridCommon.ush
+++ b/Engine/Shaders/Private/LightGridCommon.ush
@@ -4,12 +4,12 @@
 
 Buffer<uint> NumCulledLightsGrid;
 Buffer<uint> CulledLightDataGrid;
-Buffer<float4> ForwardLocalLightBuffer;
+Buffer<uint4> ForwardLocalLightBuffer;
 
 // ISR copies
 Buffer<uint> InstancedNumCulledLightsGrid;
 Buffer<uint> InstancedCulledLightDataGrid;
-Buffer<float4> InstancedForwardLocalLightBuffer;
+Buffer<uint4> InstancedForwardLocalLightBuffer;
 
 struct LightGridData
 {
@@ -91,8 +91,10 @@ struct FLocalLightData
 {
 	float4 LightPositionAndInvRadius;
 	float4 LightColorAndFalloffExponent;
-	float4 SpotAnglesAndSourceRadiusPacked;
-    float4 LightDirectionAndShadowMask;
+	float3 SpotAngles;
+	uint   SourceRadiusPacked;
+    float3 LightDirection;
+	uint   ShadowMask;
     float4 LightTangentAndSoftSourceRadius;
 };
 
@@ -104,22 +106,26 @@ FLocalLightData GetLocalLightData(uint GridIndex, uint EyeIndex)
 	{
 		uint LocalLightIndex = CulledLightDataGrid[GridIndex];
 		uint LocalLightBaseIndex = LocalLightIndex * LOCAL_LIGHT_DATA_STRIDE;
-		Result.LightPositionAndInvRadius = ForwardLocalLightBuffer[LocalLightBaseIndex + 0];
-		Result.LightColorAndFalloffExponent = ForwardLocalLightBuffer[LocalLightBaseIndex + 1];
-		Result.LightDirectionAndShadowMask = ForwardLocalLightBuffer[LocalLightBaseIndex + 2];
-        Result.SpotAnglesAndSourceRadiusPacked = ForwardLocalLightBuffer[LocalLightBaseIndex + 3];
-        Result.LightTangentAndSoftSourceRadius = ForwardLocalLightBuffer[LocalLightBaseIndex + 4];
+		Result.LightPositionAndInvRadius = asfloat(ForwardLocalLightBuffer[LocalLightBaseIndex + 0]);
+		Result.LightColorAndFalloffExponent = asfloat(ForwardLocalLightBuffer[LocalLightBaseIndex + 1]);
+		Result.LightDirection = asfloat(ForwardLocalLightBuffer[LocalLightBaseIndex + 2].xyz);
+		Result.ShadowMask = ForwardLocalLightBuffer[LocalLightBaseIndex + 2].w;
+        Result.SpotAngles = asfloat(ForwardLocalLightBuffer[LocalLightBaseIndex + 3].xyz);
+		Result.SourceRadiusPacked = ForwardLocalLightBuffer[LocalLightBaseIndex + 3].w;
+        Result.LightTangentAndSoftSourceRadius = asfloat(ForwardLocalLightBuffer[LocalLightBaseIndex + 4]);
     }
 	else
 	{
 		uint LocalLightIndex = InstancedCulledLightDataGrid[GridIndex];
 		uint LocalLightBaseIndex = LocalLightIndex * LOCAL_LIGHT_DATA_STRIDE;
-		Result.LightPositionAndInvRadius = InstancedForwardLocalLightBuffer[LocalLightBaseIndex + 0];
-		Result.LightColorAndFalloffExponent = InstancedForwardLocalLightBuffer[LocalLightBaseIndex + 1];
-		Result.LightDirectionAndShadowMask = InstancedForwardLocalLightBuffer[LocalLightBaseIndex + 2];
-        Result.SpotAnglesAndSourceRadiusPacked = InstancedForwardLocalLightBuffer[LocalLightBaseIndex + 3];
-        Result.LightTangentAndSoftSourceRadius = InstancedForwardLocalLightBuffer[LocalLightBaseIndex + 4];
+		Result.LightPositionAndInvRadius = asfloat(InstancedForwardLocalLightBuffer[LocalLightBaseIndex + 0]);
+		Result.LightColorAndFalloffExponent = asfloat(InstancedForwardLocalLightBuffer[LocalLightBaseIndex + 1]);
+		Result.LightDirection = asfloat(InstancedForwardLocalLightBuffer[LocalLightBaseIndex + 2].xyz);
+		Result.ShadowMask = InstancedForwardLocalLightBuffer[LocalLightBaseIndex + 2].w;
+        Result.SpotAngles = asfloat(InstancedForwardLocalLightBuffer[LocalLightBaseIndex + 3].xyz);
+		Result.SourceRadiusPacked = InstancedForwardLocalLightBuffer[LocalLightBaseIndex + 3].w;
+        Result.LightTangentAndSoftSourceRadius = asfloat(InstancedForwardLocalLightBuffer[LocalLightBaseIndex + 4]);
     }
 
 	return Result;
-}
\ No newline at end of file
+}
diff --git a/Engine/Shaders/Private/LightGridInjection.usf b/Engine/Shaders/Private/LightGridInjection.usf
--- a/Engine/Shaders/Private/LightGridInjection.usf
+++ b/Engine/Shaders/Private/LightGridInjection.usf
@@ -5,7 +5,7 @@
 #include "Common.ush"
 #include "Definitions.usf"
 
-Buffer<float4> ForwardLocalLightBuffer;
+Buffer<uint4> ForwardLocalLightBuffer;
 
 RWBuffer<uint> RWNumCulledLightsGrid;
 RWBuffer<uint> RWCulledLightDataGrid;
@@ -126,18 +126,18 @@ void LightGridInjectionCS(
 		for (uint LocalLightIndex = 0; LocalLightIndex < ForwardGlobalLightData.NumLocalLights; LocalLightIndex++)
 		{
 			uint LocalLightBaseIndex = LocalLightIndex * LOCAL_LIGHT_DATA_STRIDE;
-			float4 LightPositionAndInvRadius = ForwardLocalLightBuffer[LocalLightBaseIndex + 0];
+			float4 LightPositionAndInvRadius = asfloat(ForwardLocalLightBuffer[LocalLightBaseIndex + 0]);
 			float LightRadius = 1.0f / LightPositionAndInvRadius.w;
 			bool bPassSpotlightTest = true;
 
 			#define REFINE_SPOTLIGHT_BOUNDS 1
 			#if REFINE_SPOTLIGHT_BOUNDS
-				float CosOuterCone = ForwardLocalLightBuffer[LocalLightBaseIndex + 3].x;
+				float CosOuterCone = asfloat(ForwardLocalLightBuffer[LocalLightBaseIndex + 3].x);
 
 				if (CosOuterCone > -2.0f)
 				{
 					float2 CosSinAngle = float2(CosOuterCone, sqrt(1 - CosOuterCone * CosOuterCone));
-					float3 LightDirection = ForwardLocalLightBuffer[LocalLightBaseIndex + 2].xyz;
+					float3 LightDirection = asfloat(ForwardLocalLightBuffer[LocalLightBaseIndex + 2].xyz);
 
 					bPassSpotlightTest = IntersectConeWithSphere(LightPositionAndInvRadius.xyz, LightDirection, LightRadius, CosSinAngle, WorldTileBoundingSphere);
 
@@ -294,4 +294,4 @@ void LightGridCompactCS(
 		// Compact reflection captures
 		CompactReverseLinkedList(ForwardGlobalLightData.NumGridCells + GridIndex, ForwardGlobalLightData.NumReflectionCaptures);
 	}
-}
\ No newline at end of file
+}
diff --git a/Engine/Shaders/Private/VolumetricFog.usf b/Engine/Shaders/Private/VolumetricFog.usf
--- a/Engine/Shaders/Private/VolumetricFog.usf
+++ b/Engine/Shaders/Private/VolumetricFog.usf
@@ -617,16 +617,16 @@ void LightScatteringCS(
 		{
 			const FLocalLightData LocalLight = GetLocalLightData(CulledLightsGrid.DataStartIndex + LocalLightListIndex, 0);
 
-			float VolumetricScatteringIntensity = f16tof32(asuint(LocalLight.SpotAnglesAndSourceRadiusPacked.w) >> 16);
+			float VolumetricScatteringIntensity = f16tof32(LocalLight.SourceRadiusPacked >> 16);
 
 			if (VolumetricScatteringIntensity > 0)
 			{
 				FDeferredLightData LightData = (FDeferredLightData)0;
 				LightData.LightPositionAndInvRadius = LocalLight.LightPositionAndInvRadius;
 				LightData.LightColorAndFalloffExponent = LocalLight.LightColorAndFalloffExponent;
-				LightData.LightDirection = LocalLight.LightDirectionAndShadowMask.xyz;
-				float SourceLength = f16tof32(asuint(LocalLight.SpotAnglesAndSourceRadiusPacked.w));
-				LightData.SpotAnglesAndSourceRadius = float4(LocalLight.SpotAnglesAndSourceRadiusPacked.xyz, SourceLength);
+				LightData.LightDirection = LocalLight.LightDirection;
+				float SourceLength = f16tof32(LocalLight.SourceRadiusPacked);
+				LightData.SpotAnglesAndSourceRadius = float4(LocalLight.SpotAngles, SourceLength);
                 LightData.LightTangent = LocalLight.LightTangentAndSoftSourceRadius.xyz;
                 LightData.SoftSourceRadius = LocalLight.LightTangentAndSoftSourceRadius.w;
 				LightData.bInverseSquared = LightData.LightColorAndFalloffExponent.w == 0;
@@ -720,4 +720,4 @@ void FinalIntegrationCS(
 
 		RWIntegratedLightScattering[LayerCoordinate] = float4(AccumulatedLighting, AccumulatedTransmittance);
 	}
-}
\ No newline at end of file
+}
diff --git a/Engine/Source/Runtime/Renderer/Private/LightGridInjection.cpp b/Engine/Source/Runtime/Renderer/Private/LightGridInjection.cpp
--- a/Engine/Source/Runtime/Renderer/Private/LightGridInjection.cpp
+++ b/Engine/Source/Runtime/Renderer/Private/LightGridInjection.cpp
@@ -547,7 +547,7 @@ void FDeferredShadingSceneRenderer::ComputeLightGrid(FRHICommandListImmediate& R
 				if (View.ForwardLightingResources->ForwardLocalLightBuffer.NumBytes < NumBytesRequired)
 				{
 					View.ForwardLightingResources->ForwardLocalLightBuffer.Release();
-					View.ForwardLightingResources->ForwardLocalLightBuffer.Initialize(sizeof(FVector4), NumBytesRequired / sizeof(FVector4), PF_A32B32G32R32F, BUF_Volatile);
+					View.ForwardLightingResources->ForwardLocalLightBuffer.Initialize(sizeof(FVector4), NumBytesRequired / sizeof(FVector4), PF_R32G32B32A32_UINT, BUF_Volatile);
 				}
 
 				View.ForwardLightingResources->ForwardLocalLightBuffer.Lock();
diff --git a/Engine/Source/Runtime/Renderer/Private/Renderer.cpp b/Engine/Source/Runtime/Renderer/Private/Renderer.cpp
--- a/Engine/Source/Runtime/Renderer/Private/Renderer.cpp
+++ b/Engine/Source/Runtime/Renderer/Private/Renderer.cpp
@@ -40,7 +40,7 @@ public:
 	{
 		if (GMaxRHIFeatureLevel == ERHIFeatureLevel::SM5)
 		{
-			ForwardLightingResources.ForwardLocalLightBuffer.Initialize(sizeof(FVector4), sizeof(FForwardLocalLightData) / sizeof(FVector4), PF_A32B32G32R32F, BUF_Dynamic);
+			ForwardLightingResources.ForwardLocalLightBuffer.Initialize(sizeof(FVector4), sizeof(FForwardLocalLightData) / sizeof(FVector4), PF_R32G32B32A32_UINT, BUF_Dynamic);
 			ForwardLightingResources.ForwardGlobalLightData = TUniformBufferRef<FForwardGlobalLightData>::CreateUniformBufferImmediate(FForwardGlobalLightData(), UniformBuffer_MultiFrame);
 			ForwardLightingResources.NumCulledLightsGrid.Initialize(sizeof(uint32), 1, PF_R32_UINT);
 			// @todo Metal lacks SRV/UAV format conversions in v1.1 and earlier.
