Making a fix on Chris Babcock's behalf - Correcting fallout from a bad merge & correcting the use of a 'ExposureScale' variable, which was renamed in Main to 'PreExposure'.

diff --git a/Engine/Shaders/Private/MobileBasePassPixelShader.usf b/Engine/Shaders/Private/MobileBasePassPixelShader.usf
--- a/Engine/Shaders/Private/MobileBasePassPixelShader.usf
+++ b/Engine/Shaders/Private/MobileBasePassPixelShader.usf
@@ -208,7 +208,7 @@ float4 GetPlanarReflection(float3 WorldPosition, float3 WorldNormal, float Rough
 	PlanarReflection.rgb *= PlanarReflection.rgb;
 	// the capture has applied exposure scale so need to cancel that.
 #if !MATERIALBLENDING_MODULATE
-	PlanarReflection.rgb /= ResolvedView.ExposureScale;
+	PlanarReflection.rgb /= ResolvedView.PreExposure;
 #endif
 #endif
 	return PlanarReflection;
