Copying //UE4/Dev-Niagara to //UE4/Dev-Main (Source: //UE4/Dev-Niagara @ 4074996)

diff --git a/Engine/Shaders/Private/NiagaraSpriteVertexFactory.ush b/Engine/Shaders/Private/NiagaraSpriteVertexFactory.ush
--- a/Engine/Shaders/Private/NiagaraSpriteVertexFactory.ush
+++ b/Engine/Shaders/Private/NiagaraSpriteVertexFactory.ush
@@ -39,19 +39,14 @@ uint ParticleFacingMode;
 /** Blends between FaceCamera and FaceCameraPosition.*/
 #define FACING_MODE_FACECAMERADISTANCEBLEND 4
 
+/** Sorted index buffer for particles if we need to sort their draw order (translucency etc). */
+Buffer<int> SortedIndices;
+uint SortedIndicesOffset;
+
 struct FVertexFactoryInput
 {
-	float4	Position				: ATTRIBUTE0;
-	float4	OldPosition				: ATTRIBUTE1;
-	float4	SizeRotSubImage			: ATTRIBUTE2;
-	float4	Color					: ATTRIBUTE3;
-#if USE_DYNAMIC_PARAMETERS
-	float4 DynamicParameter			: ATTRIBUTE5;
-#endif	//USE_DYNAMIC_PARAMETERS
-	float2	TexCoord				: ATTRIBUTE4;
+	float2	TexCoord				: ATTRIBUTE0;
 	uint VertexId					: SV_VertexID;
-	float3 CustomAlignmentVector	: ATTRIBUTE6;
-	float3 CustomFacingVector		: ATTRIBUTE7;
 
 /** Optional instance ID for vertex layered rendering */
 //#if FEATURE_LEVEL >= FEATURE_LEVEL_SM4 && ONEPASS_POINTLIGHT_SHADOW && USING_VERTEX_SHADER_LAYER
@@ -65,10 +60,19 @@ struct FVertexFactoryInterpolantsVSToPS
 	float4	TangentToWorld0AndInterp_Sizer : TANGENTTOWORLD0;
 	// Last row of the tangent to world matrix in xyz
 	float4	TangentToWorld2	: TANGENTTOWORLD2;
-
-#if USE_DYNAMIC_PARAMETERS
-	float4 DynamicParameter	: TEXCOORD1;
-#endif	//USE_DYNAMIC_PARAMETERS
+	
+#if NUM_DYNAMIC_PARAMETERS > 0
+	float4 DynamicParameter	: PARTICLE_DYNAMIC_PARAM0;
+#endif
+#if NUM_DYNAMIC_PARAMETERS > 1
+	float4 DynamicParameter1	: PARTICLE_DYNAMIC_PARAM1;
+#endif
+#if NUM_DYNAMIC_PARAMETERS > 2
+	float4 DynamicParameter2	: PARTICLE_DYNAMIC_PARAM2;
+#endif
+#if NUM_DYNAMIC_PARAMETERS > 3
+	float4 DynamicParameter3	: PARTICLE_DYNAMIC_PARAM3;
+#endif
 
 #if NEEDS_PARTICLE_COLOR
 	float4	Color				: TEXCOORD2;
@@ -123,6 +127,8 @@ struct FVertexFactoryIntermediates
 	float3 ParticleWorldPosition;
 	/** The position of the vertex in translated world space. */
 	float3 VertexWorldPosition;
+	/** Last frame's position of the vertex in translated world space. */
+	float3 PreviousVertexWorldPosition;
 	/** Particle translated world space position and size. */
 	float4 TranslatedWorldPositionAndSize;
 	
@@ -144,8 +150,26 @@ struct FVertexFactoryIntermediates
 	float4 Color;
 	/** The velocity of the particle, XYZ: direction, W: speed. */
 	float4 ParticleVelocity;
-	/** Dynamic parameter. */
+
+	float3 Position;
+	float3 OldPosition;
+	
+	float4 SizeRotSubImage;
+	
+	float3 CustomFacingVector;
+	float3 CustomAlignmentVector;
+
 	float4 DynamicParameter;
+#if NUM_DYNAMIC_PARAMETERS > 1
+	float4 DynamicParameter1;
+#endif
+#if NUM_DYNAMIC_PARAMETERS > 2
+	float4 DynamicParameter2;
+#endif
+#if NUM_DYNAMIC_PARAMETERS > 3
+	float4 DynamicParameter3;
+#endif
+
 	/** The sub-image lerp. */
 	float SubImageLerp;
 	/** Relative time. */
@@ -226,9 +250,18 @@ FMaterialPixelParameters GetMaterialPixelParameters(FVertexFactoryInterpolantsVS
 	Result.Particle.Color = Interpolants.Color;
 #endif
 
-#if USE_DYNAMIC_PARAMETERS
+#if NUM_DYNAMIC_PARAMETERS > 0
 	Result.Particle.DynamicParameter = Interpolants.DynamicParameter;
-#endif	//USE_DYNAMIC_PARAMETERS
+#endif
+#if NUM_DYNAMIC_PARAMETERS > 1
+	Result.Particle.DynamicParameter1 = Interpolants.DynamicParameter1;
+#endif
+#if NUM_DYNAMIC_PARAMETERS > 2
+	Result.Particle.DynamicParameter2 = Interpolants.DynamicParameter2;
+#endif
+#if NUM_DYNAMIC_PARAMETERS > 3
+	Result.Particle.DynamicParameter3 = Interpolants.DynamicParameter3;
+#endif
 
 #if USE_PARTICLE_POSITION
 	Result.Particle.TranslatedWorldPositionAndSize = Interpolants.ParticleTranslatedWorldPositionAndSize;
@@ -278,17 +311,26 @@ FMaterialVertexParameters GetMaterialVertexParameters(FVertexFactoryInput Input,
 {
 	FMaterialVertexParameters Result = (FMaterialVertexParameters)0;
 	Result.WorldPosition = WorldPosition;
-	Result.VertexColor = Input.Color;
+	Result.VertexColor = Intermediates.Color;
 	Result.TangentToWorld = mul(TangentToLocal, GetLocalToWorld3x3()); 
 	Result.Particle.MacroUV = NiagaraSpriteVF.MacroUVParameters;
 	Result.Particle.Color = Intermediates.Color;
 	Result.Particle.MotionBlurFade = 1.0f;
-	Result.PreSkinnedPosition = Input.Position.xyz;
+	Result.PreSkinnedPosition = Intermediates.Position.xyz;
 	Result.PreSkinnedNormal = TangentToLocal[2].xyz;
 
-#if USE_DYNAMIC_PARAMETERS
+#if NUM_DYNAMIC_PARAMETERS > 0
 	Result.Particle.DynamicParameter = Intermediates.DynamicParameter;
-#endif	//USE_DYNAMIC_PARAMETERS
+#endif
+#if NUM_DYNAMIC_PARAMETERS > 1
+	Result.Particle.DynamicParameter1 = Intermediates.DynamicParameter1;
+#endif
+#if NUM_DYNAMIC_PARAMETERS > 2
+	Result.Particle.DynamicParameter2 = Intermediates.DynamicParameter2;
+#endif
+#if NUM_DYNAMIC_PARAMETERS > 3
+	Result.Particle.DynamicParameter3 = Intermediates.DynamicParameter3;
+#endif
 
 #if USE_PARTICLE_POSITION
 	Result.Particle.TranslatedWorldPositionAndSize = Intermediates.TranslatedWorldPositionAndSize;
@@ -349,7 +391,7 @@ float3 SafeNormalize(float3 V)
 	return V * rsqrt(max(dot(V,V),0.01));
 }
 
-void GetTangents(FVertexFactoryInput Input, float3 WorldPosition, float3 OldWorldPosition, out float3 OutRight, out float3 OutUp)
+void GetTangents(FVertexFactoryInput Input, FVertexFactoryIntermediates Intermediates, float3 WorldPosition, float3 OldWorldPosition, out float3 OutRight, out float3 OutUp)
 {
 	// Select camera up/right vectors.
 	float3 ResolvedViewRight = lerp(ResolvedView.ViewRight, ResolvedView.HMDViewNoRollRight, NiagaraSpriteVF.RemoveHMDRoll);
@@ -368,32 +410,14 @@ void GetTangents(FVertexFactoryInput Input, float3 WorldPosition, float3 OldWorl
 	// First interpolate between the base facing mode and the custom facing mode. We assume that the code setting up this call has the facing vector mask
 	// set to 0,0,0,0 if not in FACING_MODE_CUSTOMFACINGVECTOR.
 	float3 CustomFacingVectorMask =  NiagaraSpriteVF.CustomFacingVectorMask;
-	float3 CustomFacingVector = Input.CustomFacingVector * CustomFacingVectorMask + CameraDirection * (float3(1.0,1.0,1.0)-CustomFacingVectorMask);
-	
-	float WeightDefaultCameraMode = ((ParticleAlignmentMode == ALIGNMENT_MODE_UNALIGNED) && (int)(ParticleFacingMode == FACING_MODE_FACECAMERA)) ?  1 : 0;
-	float WeightPlaneFacing = ((ParticleAlignmentMode == ALIGNMENT_MODE_UNALIGNED) && (ParticleFacingMode == FACING_MODE_FACECAMERAPLANE)) ?  1 : 0;
-	float WeightCustomFacing = ((ParticleAlignmentMode == ALIGNMENT_MODE_UNALIGNED) && (ParticleFacingMode == FACING_MODE_CUSTOMFACINGVECTOR)) ?  1 : 0;
-	float WeightPositionFacing =  ((ParticleAlignmentMode == ALIGNMENT_MODE_UNALIGNED) && (ParticleFacingMode == FACING_MODE_FACECAMERAPOSITION)) ?  1 : 0;
-	float WeightVelocityAligned = ParticleAlignmentMode != ALIGNMENT_MODE_UNALIGNED;
-	
-	// Adjust weights if you are in the blend mode based off of camera distance.
-	BRANCH
-	if (NiagaraSpriteVF.CameraFacingBlend.x > 0.f && ParticleFacingMode == FACING_MODE_FACECAMERADISTANCEBLEND)
-	{
-		float CameraDistanceSq = dot(CameraVector, CameraVector);
-		float AlignmentMode = saturate(CameraDistanceSq * NiagaraSpriteVF.CameraFacingBlend.y - NiagaraSpriteVF.CameraFacingBlend.z);
-		WeightDefaultCameraMode = 1.f - AlignmentMode;
-		WeightPositionFacing = AlignmentMode;
-		WeightVelocityAligned = 0.0f;
-	}
+	float3 CustomFacingVector = Intermediates.CustomFacingVector * CustomFacingVectorMask + CameraDirection * (float3(1.0,1.0,1.0)-CustomFacingVectorMask);
 	
 	// the particle's movement direction.
-	float3 DirectionVector	= lerp(ParticleDirection, Input.CustomAlignmentVector, ParticleAlignmentMode == ALIGNMENT_MODE_CUSTOMALIGNMENT);
+	float3 DirectionVector	= lerp(ParticleDirection, Intermediates.CustomAlignmentVector, ParticleAlignmentMode == ALIGNMENT_MODE_CUSTOMALIGNMENT);
 	
 	// Tangent vectors for any non-default alignment mode.
 	float3 CameraRightAligned	= SafeNormalize(cross(CustomFacingVector, DirectionVector));
 	float3 CameraUpAligned = lerp(-DirectionVector, cross(CustomFacingVector, CameraRightAligned), ParticleFacingMode == FACING_MODE_CUSTOMFACINGVECTOR);
-	//float3 CameraUpAligned = -DirectionVector;
 
 	// Tangent vectors for camera plane facing mode
 	float3 CameraRightPlaneFacing = SafeNormalize(cross(CameraFacingVector,ResolvedViewUp));
@@ -407,30 +431,63 @@ void GetTangents(FVertexFactoryInput Input, float3 WorldPosition, float3 OldWorl
 	float3 CameraRightPosition = SafeNormalize(cross(CameraDirection,float3(0,0,1)));
 	float3 CameraUpPosition = cross(CameraDirection, CameraRightPosition);
 
-	float3 RightCombined =	WeightDefaultCameraMode * CameraRight.xyz + 
-							WeightCustomFacing * CameraRightCustomFacing.xyz +
-							WeightVelocityAligned * CameraRightAligned.xyz +
-							WeightPlaneFacing * CameraRightPlaneFacing.xyz + 
-							WeightPositionFacing * CameraRightPosition.xyz;
-
-	float3 UpCombined = WeightDefaultCameraMode * CameraUp.xyz + 
-						WeightCustomFacing * CameraUpCustomFacing.xyz +
-						WeightVelocityAligned * CameraUpAligned.xyz +
-						WeightPlaneFacing * CameraUpPlaneFacing.xyz +
-						WeightPositionFacing * CameraUpPosition.xyz;
+	float3 Right;
+	float3 Up;
+	
+	/**This could well be better as a branch*/
+	FLATTEN
+	if(ParticleAlignmentMode == ALIGNMENT_MODE_UNALIGNED)
+	{
+		FLATTEN
+		if (ParticleFacingMode == FACING_MODE_FACECAMERADISTANCEBLEND)
+		{
+			float CameraDistanceSq = dot(CameraVector, CameraVector);
+			float Interp = saturate(CameraDistanceSq * NiagaraSpriteVF.CameraFacingBlend.y - NiagaraSpriteVF.CameraFacingBlend.z);
+			Right = cross(CameraDirection,float3(0,0,1));
+			Up = cross(CameraDirection, Right);
+			Right = lerp(Right, CameraRight.xyz, Interp);
+			Up = lerp(Up, CameraUp.xyz, Interp);
+		} 
+		else if(ParticleFacingMode == FACING_MODE_FACECAMERA)
+		{
+			Right = CameraRight.xyz;
+			Up = CameraUp.xyz;
+		}		 
+		else if(ParticleFacingMode == FACING_MODE_FACECAMERAPLANE)
+		{
+			Right = CameraRightPlaneFacing;
+			Up = CameraUpPlaneFacing;
+		}
+		else if(ParticleFacingMode == FACING_MODE_CUSTOMFACINGVECTOR)
+		{
+			Right = CameraRightCustomFacing;
+			Up = CameraUpCustomFacing;
+		}
+		else//(ParticleFacingMode == FACING_MODE_FACECAMERAPOSITION)
+		{
+			Right = CameraRightPosition;
+			Up = CameraUpPosition;
+		}
+	}
+	else 
+	{	
+		Right = CameraRightAligned;
+		Up = CameraUpAligned;
+	}
 
-	float3 Right = SafeNormalize( RightCombined );
-	float3 Up = SafeNormalize( UpCombined );
 
 	// Determine the angle of rotation.
 	float SinRotation;
 	float CosRotation;
-	const float SpriteRotation = Input.SizeRotSubImage.z * NiagaraSpriteVF.RotationScale + NiagaraSpriteVF.RotationBias;
+	const float SpriteRotation = Intermediates.SizeRotSubImage.z * NiagaraSpriteVF.RotationScale + NiagaraSpriteVF.RotationBias;
 	sincos(SpriteRotation, SinRotation, CosRotation);
 
-	// Rotate the sprite to determine final tangents.
-	OutRight	= SinRotation * Up + CosRotation * Right;
-	OutUp		= CosRotation * Up - SinRotation * Right;
+	float3 SR = Right * SinRotation;
+	float3 SU = Up * SinRotation;
+	float3 CR = Right * CosRotation;
+	float3 CU = Up * CosRotation;
+	OutRight = SU + CR;
+	OutUp = CU - SR;
 }
 
 /** derive basis vectors */
@@ -473,10 +530,10 @@ uint NumCutoutVerticesPerFrame;
 /** SubUV animation bounding geometry used to minimize overdraw.  Each frame has its own polygon. */
 Buffer<float2> CutoutGeometry;
 
-void ComputeBillboardUVs(FVertexFactoryInput Input, out float2 UVForPosition, out float2 UVForTexturing, out float2 UVForTexturingUnflipped)
+void ComputeBillboardUVs(FVertexFactoryInput Input, FVertexFactoryIntermediates Intermediates, out float2 UVForPosition, out float2 UVForTexturing, out float2 UVForTexturingUnflipped)
 {
 	// Encoding the UV flip in the sign of the size data.
-	float2 UVFlip = sign(Input.SizeRotSubImage.xy);
+	float2 UVFlip = sign(Intermediates.SizeRotSubImage.xy);
 
 	BRANCH
 	if (NumCutoutVerticesPerFrame > 0)
@@ -484,7 +541,7 @@ void ComputeBillboardUVs(FVertexFactoryInput Input, out float2 UVForPosition, ou
 		// Avoid uint divide which is extremely slow on GCN
 		uint CutoutVertexIndex = (uint)fmod(Input.VertexId, NumCutoutVerticesPerFrame);
 		float NumFrames = NiagaraSpriteVF.SubImageSize.x * NiagaraSpriteVF.SubImageSize.y;
-		uint SubImageIndexInt = (uint)fmod(Input.SizeRotSubImage.w, NumFrames);
+		uint SubImageIndexInt = (uint)fmod(Intermediates.SizeRotSubImage.w, NumFrames);
 
 		FLATTEN
 		if (UVFlip.x * UVFlip.y < 0)
@@ -527,22 +584,47 @@ float3 TransformVector(float3 InVector)
 
 float3 GetNiagaraParticlePosition(uint InstanceID)
 {
+	if(NiagaraSpriteVF.PositionDataOffset  == -1)
+	{
+		// See FNiagaraConstants for a synchronized value...
+		return TransformPosition(NiagaraSpriteVF.DefaultPos.xyz);
+	}
 	return TransformPosition(GetVec3(NiagaraSpriteVF.PositionDataOffset, InstanceID));
 }
 float3 GetNiagaraParticleVelocity(uint InstanceID)
 {
+	if(NiagaraSpriteVF.VelocityDataOffset  == -1)
+	{
+		// See FNiagaraConstants for a synchronized value...
+		return TransformVector(float3(0.0f,0.0f,0.0f));
+	}
 	return TransformVector(GetVec3(NiagaraSpriteVF.VelocityDataOffset, InstanceID));
 }
 float2 GetNiagaraParticleSize(uint InstanceID)
 {
+	if(NiagaraSpriteVF.SizeDataOffset  == -1)
+	{
+		// See FNiagaraConstants for a synchronized value...
+		return float2(50.0f, 50.0f);
+	}
 	return GetVec2(NiagaraSpriteVF.SizeDataOffset, InstanceID);
 }
 float GetNiagaraParticleRotation(uint InstanceID)
 {
+	if(NiagaraSpriteVF.RotationDataOffset  == -1)
+	{
+		// See FNiagaraConstants for a synchronized value...
+		return 0.0f;
+	}
 	return (GetFloat(NiagaraSpriteVF.RotationDataOffset, InstanceID) / 180.0f) * PI;
 }
 float4 GetNiagaraParticleColor(uint InstanceID)
 {
+	if(NiagaraSpriteVF.ColorDataOffset  == -1)
+	{
+		// See FNiagaraConstants for a synchronized value...
+		return float4(1.0f, 1.0f, 1.0f, 1.0f);
+	}
 	return GetVec4(NiagaraSpriteVF.ColorDataOffset, InstanceID);
 }
 
@@ -550,6 +632,7 @@ float2 GetNiagaraUVScale(uint InstanceID)
 {
 	if(NiagaraSpriteVF.UVScaleOffset  == -1)
 	{
+		// See FNiagaraConstants for a synchronized value...
 		return float2(1.0f, 1.0f);
 	}
 	return GetVec2(NiagaraSpriteVF.UVScaleOffset, InstanceID);	
@@ -559,35 +642,51 @@ float GetNiagaraParticleRandom(uint InstanceID)
 {
 	if(NiagaraSpriteVF.ParticleRandomOffset  == -1)
 	{
+		// See FNiagaraConstants for a synchronized value...
 		return 0.0f;
 	}
 	return GetFloat(NiagaraSpriteVF.ParticleRandomOffset, InstanceID);	
 }
 
-
-
 float GetNiagaraCameraOffset(uint InstanceID)
 {
 	if (NiagaraSpriteVF.CameraOffsetOffset == -1)
 	{
+		// See FNiagaraConstants for a synchronized value...
 		return 0.0f;
 	}
 	return GetFloat(NiagaraSpriteVF.CameraOffsetOffset, InstanceID);
 }
 
-#if USE_DYNAMIC_PARAMETERS
+#if NUM_DYNAMIC_PARAMETERS > 0
 float4 GetNiagaraParticleDynamicParameters(uint InstanceID)
 {
-	if (NiagaraSpriteVF.MaterialParamDataOffset == -1)
-	{
-		return float4(1.0f, 1.0f, 1.0f, 1.0f);
-	}
-	return GetVec4(NiagaraSpriteVF.MaterialParamDataOffset, InstanceID);
+	return NiagaraSpriteVF.MaterialParamDataOffset != -1 ? GetVec4(NiagaraSpriteVF.MaterialParamDataOffset, InstanceID) : float4(1.0f, 1.0f, 1.0f, 1.0f);;
 }
 #endif
+#if NUM_DYNAMIC_PARAMETERS > 1
+float4 GetNiagaraParticleDynamicParameters1(uint InstanceID)
+{
+	return NiagaraSpriteVF.MaterialParam1DataOffset != -1 ? GetVec4(NiagaraSpriteVF.MaterialParam1DataOffset, InstanceID) : float4(1.0f, 1.0f, 1.0f, 1.0f);;
+}
+#endif
+#if NUM_DYNAMIC_PARAMETERS > 2
+float4 GetNiagaraParticleDynamicParameters2(uint InstanceID)
+{
+	return NiagaraSpriteVF.MaterialParam2DataOffset != -1 ? GetVec4(NiagaraSpriteVF.MaterialParam2DataOffset, InstanceID) : float4(1.0f, 1.0f, 1.0f, 1.0f);;
+}
+#endif
+#if NUM_DYNAMIC_PARAMETERS > 3
+float4 GetNiagaraParticleDynamicParameters3(uint InstanceID)
+{
+	return NiagaraSpriteVF.MaterialParam3DataOffset != -1 ? GetVec4(NiagaraSpriteVF.MaterialParam3DataOffset, InstanceID) : float4(1.0f, 1.0f, 1.0f, 1.0f);;
+}
+#endif
+
 
 float GetNiagaraParticleSubimage(uint InstanceID)
 {
+	// See FNiagaraConstants for a synchronized value...
 	float Val = 0.0f;
     if(NiagaraSpriteVF.SubimageDataOffset != -1)
 	{
@@ -598,7 +697,8 @@ float GetNiagaraParticleSubimage(uint InstanceID)
 
 float3 GetNiagaraParticleFacingVector(uint InstanceID)
 {
-	float3 Val = float3(0.0f, 0.0f, 0.0f);
+	// See FNiagaraConstants for a synchronized value...
+	float3 Val = float3(1.0f, 0.0f, 0.0f);
     if(NiagaraSpriteVF.FacingOffset != -1)
 	{
 		Val = GetVec3(NiagaraSpriteVF.FacingOffset, InstanceID);
@@ -608,7 +708,7 @@ float3 GetNiagaraParticleFacingVector(uint InstanceID)
 
 float3 GetNiagaraParticleAlignmentVector(uint InstanceID)
 {
-	float3 Val = float3(0.0f, 0.0f, 0.0f);
+	float3 Val = float3(1.0f, 0.0f, 0.0f);
     if(NiagaraSpriteVF.AlignmentOffset != -1)
 	{
 		Val = GetVec3(NiagaraSpriteVF.AlignmentOffset, InstanceID);
@@ -621,23 +721,39 @@ FVertexFactoryIntermediates GetVertexFactoryIntermediates(FVertexFactoryInput In
 {
 	FVertexFactoryIntermediates Intermediates = (FVertexFactoryIntermediates)0;
 
-	Input.Position = float4(GetNiagaraParticlePosition(Input.InstanceId), 1);
-	Input.OldPosition = float4(Input.Position.xyz - GetNiagaraParticleVelocity(Input.InstanceId), 1);
-	float2 ParticleSize = GetNiagaraParticleSize(Input.InstanceId);
-	Input.SizeRotSubImage = float4(ParticleSize.x, ParticleSize.y, GetNiagaraParticleRotation(Input.InstanceId), GetNiagaraParticleSubimage(Input.InstanceId));
-	Input.Color	= GetNiagaraParticleColor(Input.InstanceId);
-	Input.CustomFacingVector = GetNiagaraParticleFacingVector(Input.InstanceId);
-	Input.CustomAlignmentVector = GetNiagaraParticleAlignmentVector(Input.InstanceId);
-#if USE_DYNAMIC_PARAMETERS
-	Input.DynamicParameter = GetNiagaraParticleDynamicParameters(Input.InstanceId);
+	uint ParticleID = Input.InstanceId;
+	if(SortedIndicesOffset != 0xFFFFFFFF)
+	{
+		ParticleID = SortedIndices[SortedIndicesOffset + Input.InstanceId];
+	}
+
+	Intermediates.Position = GetNiagaraParticlePosition(ParticleID);
+	Intermediates.OldPosition = Intermediates.Position.xyz - GetNiagaraParticleVelocity(ParticleID) * NiagaraSpriteVF.DeltaSeconds;
+	float2 ParticleSize = GetNiagaraParticleSize(ParticleID);
+	Intermediates.SizeRotSubImage = float4(ParticleSize.x, ParticleSize.y, GetNiagaraParticleRotation(ParticleID), GetNiagaraParticleSubimage(ParticleID));
+	Intermediates.Color	= GetNiagaraParticleColor(ParticleID);
+	Intermediates.CustomFacingVector = GetNiagaraParticleFacingVector(ParticleID);
+	Intermediates.CustomAlignmentVector = GetNiagaraParticleAlignmentVector(ParticleID);
+
+#if NUM_DYNAMIC_PARAMETERS > 0
+	Intermediates.DynamicParameter = GetNiagaraParticleDynamicParameters(ParticleID);
+#endif
+#if NUM_DYNAMIC_PARAMETERS > 1
+	Intermediates.DynamicParameter1 = GetNiagaraParticleDynamicParameters1(ParticleID);
+#endif
+#if NUM_DYNAMIC_PARAMETERS > 2
+	Intermediates.DynamicParameter2 = GetNiagaraParticleDynamicParameters2(ParticleID);
+#endif
+#if NUM_DYNAMIC_PARAMETERS > 3
+	Intermediates.DynamicParameter3 = GetNiagaraParticleDynamicParameters3(ParticleID);
 #endif
 
 
 	// World position.
 	
-	float4 ParticleWorldPosition = mul(float4(Input.Position.xyz,1), Primitive.LocalToWorld);
-	float4 ParticleOldWorldPosition = mul(float4(Input.OldPosition.xyz,1), Primitive.LocalToWorld); 
-	float CameraOffset = GetNiagaraCameraOffset(Input.InstanceId);
+	float4 ParticleWorldPosition = mul(float4(Intermediates.Position.xyz,1), Primitive.LocalToWorld);
+	float4 ParticleOldWorldPosition = mul(float4(Intermediates.OldPosition.xyz,1), Primitive.LocalToWorld); 
+	float CameraOffset = GetNiagaraCameraOffset(ParticleID);
 	float3 CameraVector = ResolvedView.WorldCameraOrigin.xyz - ParticleWorldPosition.xyz;
 	float3 CameraDirection = SafeNormalize(CameraVector);
 
@@ -651,27 +767,30 @@ FVertexFactoryIntermediates GetVertexFactoryIntermediates(FVertexFactoryInput In
 
 	// Tangents.
 	float3 Right,Up;
-	GetTangents(Input,ParticleWorldPosition.xyz,ParticleOldWorldPosition.xyz,Right,Up);
+	GetTangents(Input, Intermediates, ParticleWorldPosition.xyz,ParticleOldWorldPosition.xyz,Right,Up);
 	Intermediates.TangentUp = Up;
 	Intermediates.TangentRight = Right;
 
 	float2 UVForPosition;
 	float2 UVForTexturing;
 	float2 UVForTexturingUnflipped;
-	ComputeBillboardUVs(Input, UVForPosition, UVForTexturing, UVForTexturingUnflipped);
+	ComputeBillboardUVs(Input, Intermediates, UVForPosition, UVForTexturing, UVForTexturingUnflipped);
 
-	float2 UVScale = GetNiagaraUVScale(Input.InstanceId);
+	float2 UVScale = GetNiagaraUVScale(ParticleID);
 	UVForTexturing.xy = float2(0.5f, 0.5f) + (UVForTexturing.xy - float2(0.5f, 0.5f) ) *  UVScale.xy; 
 
 	// Vertex position.
 	float4 VertexWorldPosition = float4(ParticleWorldPosition.xyz + ResolvedView.PreViewTranslation,1);
-	float2 Size = abs(Input.SizeRotSubImage.xy);
+	float2 Size = abs(Intermediates.SizeRotSubImage.xy);
 	VertexWorldPosition += Size.x * (UVForPosition.x + NiagaraSpriteVF.PivotOffset.x) * float4(Right,0);
 	VertexWorldPosition += Size.y * (UVForPosition.y + NiagaraSpriteVF.PivotOffset.y) * float4(Up,0);
 	Intermediates.VertexWorldPosition = VertexWorldPosition.xyz;
+
+	// Previous vertex position.
+	Intermediates.PreviousVertexWorldPosition = VertexWorldPosition.xyz + (ParticleOldWorldPosition.xyz - ParticleWorldPosition.xyz);
 	
 	// SubUV.
-	float SubImageIndex = Input.SizeRotSubImage.w;
+	float SubImageIndex = Intermediates.SizeRotSubImage.w;
 	float SubImageLerp = frac(SubImageIndex);
 	float SubImageA = SubImageIndex - SubImageLerp;
 	float SubImageB = SubImageA + 1;
@@ -686,16 +805,10 @@ FVertexFactoryIntermediates GetVertexFactoryIntermediates(FVertexFactoryInput In
 	Intermediates.TexCoord_Unflipped.zw = (float2(SubImageBH, SubImageBV) + UVForTexturingUnflipped) * NiagaraSpriteVF.SubImageSize.zw;
 #endif
 	Intermediates.SubImageLerp = (NiagaraSpriteVF.SubImageBlendMode == 1) ? SubImageLerp : 0.0f;
-
-	Intermediates.Color = Input.Color;
-
-#if USE_DYNAMIC_PARAMETERS
-	Intermediates.DynamicParameter = Input.DynamicParameter;
-#endif	//USE_DYNAMIC_PARAMETERS
-
+	
 
 #if USE_PARTICLE_POSITION
-	float ParticleRadius = .5f * min(Input.SizeRotSubImage.x, Input.SizeRotSubImage.y);
+	float ParticleRadius = .5f * min(Intermediates.SizeRotSubImage.x, Intermediates.SizeRotSubImage.y);
 	float3 ParticleTranslatedWorldPosition = ParticleWorldPosition.xyz + ResolvedView.PreViewTranslation.xyz;
 	Intermediates.TranslatedWorldPositionAndSize = float4(ParticleTranslatedWorldPosition, ParticleRadius);
 
@@ -706,25 +819,25 @@ FVertexFactoryIntermediates GetVertexFactoryIntermediates(FVertexFactoryInput In
 #endif
 
 #if USE_PARTICLE_VELOCITY
-	float3 ParticleVelocity = (ParticleWorldPosition.xyz - ParticleOldWorldPosition.xyz) * NiagaraSpriteVF.InvDeltaSeconds;
-	float3 DirectionVector = lerp(ParticleVelocity, Input.CustomAlignmentVector, ParticleAlignmentMode == ALIGNMENT_MODE_CUSTOMALIGNMENT);
+	float3 ParticleVelocity = (ParticleWorldPosition.xyz - ParticleOldWorldPosition.xyz) * (1.0f/NiagaraSpriteVF.DeltaSeconds);
+	float3 DirectionVector = lerp(ParticleVelocity, Intermediates.CustomAlignmentVector, ParticleAlignmentMode == ALIGNMENT_MODE_CUSTOMALIGNMENT);
 	Intermediates.ParticleVelocity.xyz = normalize(DirectionVector);
 	Intermediates.ParticleVelocity.w = length(DirectionVector);
 #endif
 
 #if USE_PARTICLE_TIME
-	Intermediates.RelativeTime = Input.Position.w;
+	Intermediates.RelativeTime = Intermediates.Position.w;
 #endif
 
 #if USE_PARTICLE_RANDOM
-	Intermediates.ParticleRandom = GetNiagaraParticleRandom(Input.InstanceId);
+	Intermediates.ParticleRandom = GetNiagaraParticleRandom(ParticleID);
 #endif
 
 	Intermediates.TangentToLocal = CalcTangentBasis(Intermediates);
 
 #if USE_PARTICLE_LIGHTING_OFFSET
 	// Hash function based on the particle ID to generate a uniformly distributed 3d offset
-	float3 RandomParticleOffset = frac(Square(Input.OldPosition.w + 10) * float3(1361.456345, 2333.578, 3623.983)) * 2 - 1;
+	float3 RandomParticleOffset = frac(Square(ParticleID + 10) * float3(1361.456345, 2333.578, 3623.983)) * 2 - 1;
 	Intermediates.LightingPositionOffset = .5f * View.TranslucencyLightingVolumeInvSize[0].w * RandomParticleOffset;
 #endif
 
@@ -747,7 +860,7 @@ float4 VertexFactoryGetRasterizedWorldPosition(FVertexFactoryInput Input, FVerte
 	// This avoids opaque objects intersecting the particle from causing clipping artifacts due to depth testing
 	// The downside is that the particle will clip the near plane sooner
 
-	float Radius = .5f * min(Input.SizeRotSubImage.x, Input.SizeRotSubImage.y);
+	float Radius = .5f * min(Intermediates.SizeRotSubImage.x, Intermediates.SizeRotSubImage.y);
 	return ReprojectPosition(TranslatedWorldPosition, Radius);
 #else
 	return TranslatedWorldPosition;
@@ -804,9 +917,18 @@ FVertexFactoryInterpolantsVSToPS VertexFactoryGetInterpolantsVSToPS(FVertexFacto
 	Interpolants.Color = Intermediates.Color;
 #endif
 
-#if USE_DYNAMIC_PARAMETERS
+#if NUM_DYNAMIC_PARAMETERS > 0
 	Interpolants.DynamicParameter = Intermediates.DynamicParameter;
-#endif	//USE_DYNAMIC_PARAMETERS
+#endif
+#if NUM_DYNAMIC_PARAMETERS > 1
+	Interpolants.DynamicParameter1 = Intermediates.DynamicParameter1;
+#endif
+#if NUM_DYNAMIC_PARAMETERS > 2
+	Interpolants.DynamicParameter2 = Intermediates.DynamicParameter2;
+#endif
+#if NUM_DYNAMIC_PARAMETERS > 3
+	Interpolants.DynamicParameter3 = Intermediates.DynamicParameter3;
+#endif
 
 #if USE_PARTICLE_POSITION
 	Interpolants.ParticleTranslatedWorldPositionAndSize = Intermediates.TranslatedWorldPositionAndSize;
@@ -842,7 +964,7 @@ FVertexFactoryInterpolantsVSToPS VertexFactoryGetInterpolantsVSToPS(FVertexFacto
 
 float4 VertexFactoryGetPreviousWorldPosition(FVertexFactoryInput Input, FVertexFactoryIntermediates Intermediates)
 {
-	return VertexFactoryGetWorldPosition(Input, Intermediates);
+	return float4(Intermediates.PreviousVertexWorldPosition, 1.0f);
 }
 
 /**
diff --git a/Engine/Shaders/Private/NiagaraVFParticleAccess.usf b/Engine/Shaders/Private/NiagaraVFParticleAccess.usf
--- a/Engine/Shaders/Private/NiagaraVFParticleAccess.usf
+++ b/Engine/Shaders/Private/NiagaraVFParticleAccess.usf
@@ -4,11 +4,13 @@
  *********************************************************************************/
 Buffer<float> NiagaraParticleDataFloat;
 Buffer<int> NiagaraParticleDataInt;
-uint SafeComponentBufferSize;			// stride of data in buffers
+uint NiagaraFloatDataOffset;
+uint NiagaraFloatDataStride;
+
 
 float GetFloat(int RegisterIdx, uint InstanceID)
 {
-	return NiagaraParticleDataFloat[RegisterIdx*SafeComponentBufferSize + InstanceID];
+	return NiagaraParticleDataFloat[NiagaraFloatDataOffset + (RegisterIdx * NiagaraFloatDataStride + InstanceID)];
 }
 
 float2 GetVec2(int RegisterIndex, uint InstanceID)
diff --git a/Engine/Source/Developer/ShaderFormatVectorVM/Private/VectorVMBackend.cpp b/Engine/Source/Developer/ShaderFormatVectorVM/Private/VectorVMBackend.cpp
--- a/Engine/Source/Developer/ShaderFormatVectorVM/Private/VectorVMBackend.cpp
+++ b/Engine/Source/Developer/ShaderFormatVectorVM/Private/VectorVMBackend.cpp
@@ -86,6 +86,9 @@ char* FVectorVMCodeBackend::GenerateCode(exec_list* ir, _mesa_glsl_parse_state*
 	if (state->error) return nullptr;
 
 	bool progress = false;
+
+	FlattenUniformBufferStructures(ir, state);
+
 	{
 		SCOPE_CYCLE_COUNTER(STAT_VVMInitMisc);
 		vm_debug_print("== Initial misc ==\n");
diff --git a/Engine/Source/Developer/ShaderFormatVectorVM/Private/VectorVMShaderCompiler.cpp b/Engine/Source/Developer/ShaderFormatVectorVM/Private/VectorVMShaderCompiler.cpp
--- a/Engine/Source/Developer/ShaderFormatVectorVM/Private/VectorVMShaderCompiler.cpp
+++ b/Engine/Source/Developer/ShaderFormatVectorVM/Private/VectorVMShaderCompiler.cpp
@@ -138,31 +138,31 @@ bool CompileShader_VectorVM(const FShaderCompilerInput& Input, FShaderCompilerOu
 			) ? 1 : 0;
 	}
 
-		if (ErrorLog)
-		{
-			int32 SrcLen = FPlatformString::Strlen(ErrorLog);
-			int32 DestLen = FPlatformString::ConvertedLength<TCHAR, ANSICHAR>(ErrorLog, SrcLen);
-			TArray<TCHAR> Converted;
-			Converted.AddUninitialized(DestLen);
+	if (ErrorLog)
+	{
+		int32 SrcLen = FPlatformString::Strlen(ErrorLog);
+		int32 DestLen = FPlatformString::ConvertedLength<TCHAR, ANSICHAR>(ErrorLog, SrcLen);
+		TArray<TCHAR> Converted;
+		Converted.AddUninitialized(DestLen);
 
-			FPlatformString::Convert<ANSICHAR, TCHAR>(Converted.GetData(), DestLen, ErrorLog, SrcLen);
-			Converted.Add(0);
+		FPlatformString::Convert<ANSICHAR, TCHAR>(Converted.GetData(), DestLen, ErrorLog, SrcLen);
+		Converted.Add(0);
 
-			VMCompilationOutput.Errors = Converted.GetData();
+		VMCompilationOutput.Errors = Converted.GetData();
 
-			TArray<FString> OutputByLines;
-			PreprocessedShader.ParseIntoArrayLines(OutputByLines, false);
+		TArray<FString> OutputByLines;
+		PreprocessedShader.ParseIntoArrayLines(OutputByLines, false);
 
-			FString OutputHlsl;
-			for (int32 i = 0; i < OutputByLines.Num(); i++)
-			{
-				UE_LOG(LogVectorVMShaderCompiler, Warning, TEXT("/*%d*/%s"), i, *OutputByLines[i]);
-			}
-		}
-		else
+		FString OutputHlsl;
+		for (int32 i = 0; i < OutputByLines.Num(); i++)
 		{
-			VMCompilationOutput.Errors.Empty();
+			UE_LOG(LogVectorVMShaderCompiler, Warning, TEXT("/*%d*/%s"), i, *OutputByLines[i]);
 		}
+	}
+	else
+	{
+		VMCompilationOutput.Errors.Empty();
+	}
 
 	//TODO: Try to get rid of the CompilationOutput and have the vm bytecode life in the shader eco-system as the compute shader version will.
 //		int32 SourceLen = VVMBackEnd.ByteCode.Num();//ShaderSource ? FCStringAnsi::Strlen(ShaderSource) : 0;
@@ -237,9 +237,9 @@ bool CompileShader_VectorVM(const FShaderCompilerInput& Input, FShaderCompilerOu
 		UE_LOG(LogVectorVMShaderCompiler, Warning, TEXT("%s"), (const char*)ShaderSource);
 		free(ShaderSource);
 	}
-		if (VMCompilationOutput.Errors.Len() != 0)
+	if (VMCompilationOutput.Errors.Len() != 0)
 	{
-			UE_LOG(LogVectorVMShaderCompiler, Warning, TEXT("%s"), *VMCompilationOutput.Errors);
+		UE_LOG(LogVectorVMShaderCompiler, Warning, TEXT("%s"), *VMCompilationOutput.Errors);
 		free(ErrorLog);
 	}
 
diff --git a/Engine/Source/Developer/ShaderFormatVectorVM/Private/ir_vm_gen_bytecode_visitor.cpp b/Engine/Source/Developer/ShaderFormatVectorVM/Private/ir_vm_gen_bytecode_visitor.cpp
--- a/Engine/Source/Developer/ShaderFormatVectorVM/Private/ir_vm_gen_bytecode_visitor.cpp
+++ b/Engine/Source/Developer/ShaderFormatVectorVM/Private/ir_vm_gen_bytecode_visitor.cpp
@@ -369,6 +369,14 @@ EVectorVMOp get_special_vm_opcode(ir_function_signature* signature)
 	{
 		VVMOpCode = EVectorVMOp::acquireindex;
 	}
+	else if (strcmp(signature->function_name(), "AcquireID") == 0)
+	{
+		VVMOpCode = EVectorVMOp::acquire_id;
+	}
+	else if (strcmp(signature->function_name(), "UpdateID") == 0)
+	{
+		VVMOpCode = EVectorVMOp::update_id;
+	}
 	else if (strcmp(signature->function_name(), "ExecIndex") == 0)
 	{
 		VVMOpCode = EVectorVMOp::exec_index;
@@ -557,32 +565,59 @@ struct op_base
 	virtual struct op_input* as_input() { return nullptr; }
 	virtual struct op_output* as_output() { return nullptr; }
 	virtual struct op_index_acquire* as_index_acquire() { return nullptr; }
+	virtual struct op_id_acquire* as_id_acquire() { return nullptr; }
+	virtual struct op_id_update* as_id_update() { return nullptr; }
 	virtual struct op_external_func* as_external_func() { return nullptr; }
 
-	static bool finalize_component_temporary_allocation(variable_info_node* component, const unsigned num_temp_registers, unsigned *registers, unsigned op_idx)
+	static void validate_component_offset(_mesa_glsl_parse_state* parse_state, variable_info_node* component)
+	{
+		if (component->offset == INDEX_NONE)
+		{
+			FString ErrorStr = FString::Printf(TEXT("Component %s of variable %s has no valid offset. Possibly uninitialized data being used."), ANSI_TO_TCHAR(component->name), ANSI_TO_TCHAR(component->owner->root->name));
+
+			_mesa_glsl_error(parse_state, TCHAR_TO_ANSI(*ErrorStr));
+		}
+
+	}
+
+	static bool finalize_component_temporary_allocation(_mesa_glsl_parse_state* parse_state, variable_info_node* component, const unsigned num_temp_registers, unsigned *registers, unsigned op_idx)
 	{
 		check(component);
 		if (ir_variable* var = component->owner->ir->as_variable())
 		{
-			if (component->offset == INDEX_NONE && (var->mode == ir_var_temporary || var->mode == ir_var_auto))
+			if (var->mode == ir_var_temporary || var->mode == ir_var_auto)
 			{
-				for (unsigned j = 0; j < num_temp_registers; ++j)
+				if (component->offset == INDEX_NONE)
 				{
-					if (registers[j] < op_idx || registers[j] == 0)
+					for (unsigned j = 0; j < num_temp_registers; ++j)
 					{
-						component->offset = VectorVM::FirstTempRegister + j;
-						registers[j] = component->last_read;
+						if (registers[j] < op_idx || registers[j] == INDEX_NONE)
+						{
+							component->offset = VectorVM::FirstTempRegister + j;
+							registers[j] = component->last_read;
 #if VM_VERBOSE_LOGGING
-						UE_LOG(LogVVMBackend, Log, TEXT("OP:%d | Comp:%p allocated Reg: %d | Last Read: %d |"), op_idx, component, j, component->last_read);
+							UE_LOG(LogVVMBackend, Log, TEXT("OP:%d | Comp:%p allocated Reg: %d | Last Read: %d |"), op_idx, component, j, component->last_read);
 #endif
-						break;
+							break;
+						}
+					}
+
+					if (component->offset == INDEX_NONE)
+					{
+						FString ErrorStr = FString::Printf(TEXT("Failed to allocate temporary register for component %s of variable %s."), ANSI_TO_TCHAR(component->name), ANSI_TO_TCHAR(component->owner->root->name));
+						_mesa_glsl_error(parse_state, TCHAR_TO_ANSI(*ErrorStr));
 					}
-				}
-				if (component->offset == INDEX_NONE)
-				{
-					return false;
 				}
 			}
+			else
+			{
+				FString ErrorStr = FString::Printf(TEXT("Attempted to allocate temporary register for invalid variable. Component %s of variable %s."), ANSI_TO_TCHAR(component->name), ANSI_TO_TCHAR(component->owner->root->name));
+				_mesa_glsl_error(parse_state, TCHAR_TO_ANSI(*ErrorStr));
+			}
+		}
+		else
+		{
+			_mesa_glsl_error(parse_state, "Attempted to allocate temporary register for null variable.");
 		}
 		return true;
 	}
@@ -635,6 +670,10 @@ struct op_hook : public op_base
 
 	virtual void validate(_mesa_glsl_parse_state* parse_state)override
 	{
+		if (source_component)
+		{
+			validate_component_offset(parse_state, source_component);
+		}
 	}
 };
 
@@ -710,27 +749,17 @@ struct op_standard : public op_base
 
 	virtual bool finalize(_mesa_glsl_parse_state* parse_state, FVectorVMCompilationOutput&CompilationOutput, const unsigned num_temp_registers, unsigned *registers, unsigned op_idx)
 	{
-		if (!finalize_component_temporary_allocation(dest_component, num_temp_registers, registers, op_idx))
-		{
-			_mesa_glsl_error(parse_state, "Failed to allocate temporary registers");
-			return false;
-		}
-		return true;
+		return finalize_component_temporary_allocation(parse_state, dest_component, num_temp_registers, registers, op_idx);
 	}
 
 	virtual void validate(_mesa_glsl_parse_state* parse_state)override
 	{
 		for (unsigned operand_idx = 0; operand_idx < num_operands; ++operand_idx)
 		{
-			if (source_components[operand_idx]->offset == INDEX_NONE)
-			{
-				_mesa_glsl_error(parse_state, "unknown error");//temp cop out, we have info to produce a good error message here.
-			}
-		}
-		if(dest_component->offset == INDEX_NONE)
-		{
-			_mesa_glsl_error(parse_state, "unknown error");//temp cop out, we have info to produce a good error message here.
+			validate_component_offset(parse_state, source_components[operand_idx]);
 		}
+
+		validate_component_offset(parse_state, dest_component);
 	}
 };
 
@@ -778,20 +807,12 @@ struct op_input : public op_base
 
 	virtual bool finalize(_mesa_glsl_parse_state* parse_state, FVectorVMCompilationOutput&CompilationOutput, const unsigned num_temp_registers, unsigned *registers, unsigned op_idx)
 	{
-		if (!finalize_component_temporary_allocation(dest_component, num_temp_registers, registers, op_idx))
-		{
-			_mesa_glsl_error(parse_state, "Failed to allocate temporary registers");
-			return false;
-		}
-		return true;
+		return finalize_component_temporary_allocation(parse_state, dest_component, num_temp_registers, registers, op_idx);
 	}
 
 	virtual void validate(_mesa_glsl_parse_state* parse_state)override
 	{
-		if (dest_component->offset == INDEX_NONE)
-		{
-			_mesa_glsl_error(parse_state, "unknown error");//temp cop out, we have info to produce a good error message here.
-		}
+		validate_component_offset(parse_state, dest_component);
 	}
 };
 
@@ -833,20 +854,12 @@ struct op_input_noadvance : public op_input
 
 	virtual bool finalize(_mesa_glsl_parse_state* parse_state, FVectorVMCompilationOutput&CompilationOutput, const unsigned num_temp_registers, unsigned *registers, unsigned op_idx)
 	{
-		if (!finalize_component_temporary_allocation(dest_component, num_temp_registers, registers, op_idx))
-		{
-			_mesa_glsl_error(parse_state, "Failed to allocate temporary registers");
-			return false;
-		}
-		return true;
+		return finalize_component_temporary_allocation(parse_state, dest_component, num_temp_registers, registers, op_idx);
 	}
 
 	virtual void validate(_mesa_glsl_parse_state* parse_state)override
 	{
-		if (dest_component->offset == INDEX_NONE)
-		{
-			_mesa_glsl_error(parse_state, "unknown error");//temp cop out, we have info to produce a good error message here.
-		}
+		validate_component_offset(parse_state, dest_component);
 	}
 };
 
@@ -868,7 +881,8 @@ struct op_output : public op_base
 		check(instance_idx_component);
 		check(value_component);
 
-		return FString::Printf(TEXT("OutputData(%d, [%s][%d], [%s][%d]);\n"),
+		return FString::Printf(TEXT("OutputData(%d, %d, [%s][%d], [%s][%d]);\n"),
+			dataset_index,
 			register_idx,
 			get_location_string(instance_idx_component->owner->location),
 			instance_idx_component->offset,
@@ -891,11 +905,7 @@ struct op_output : public op_base
 
 	virtual void validate(_mesa_glsl_parse_state* parse_state)override
 	{
-		//TODO: Check registers etc
-		if (value_component->offset == INDEX_NONE || instance_idx_component->offset == INDEX_NONE)
-		{
-			_mesa_glsl_error(parse_state, "unknown error");//temp cop out, we have info to produce a good error message here.
-		}
+		validate_component_offset(parse_state, value_component);
 	}
 };
 
@@ -937,7 +947,55 @@ struct op_index_acquire : public op_base
 
 	virtual bool finalize(_mesa_glsl_parse_state* parse_state, FVectorVMCompilationOutput&CompilationOutput, const unsigned num_temp_registers, unsigned *registers, unsigned op_idx)
 	{
-		if (!finalize_component_temporary_allocation(dest_component, num_temp_registers, registers, op_idx))
+		return finalize_component_temporary_allocation(parse_state, dest_component, num_temp_registers, registers, op_idx);
+	}
+
+	virtual void validate(_mesa_glsl_parse_state* parse_state)override
+	{
+		validate_component_offset(parse_state, dest_component);
+	}
+};
+
+struct op_id_acquire : public op_base
+{
+	op_id_acquire()
+		: dataset_idx(INDEX_NONE)
+		, id_index_component(nullptr)
+		, id_tag_component(nullptr)
+	{}
+
+	virtual op_id_acquire* as_id_acquire() { return this; }
+
+	unsigned dataset_idx;
+	variable_info_node* id_index_component;
+	variable_info_node* id_tag_component;
+
+	virtual FString to_string()
+	{
+		check(id_index_component);
+		check(id_tag_component);
+
+		return FString::Printf(TEXT("AcquireID(%d, [R][%d], [R][%d]);\n"), dataset_idx, id_index_component->offset, id_tag_component->offset);
+	}
+
+	virtual void add_to_bytecode(TArray<uint8>& bytecode)
+	{
+		check(id_index_component);
+		check(id_tag_component);
+		emit_byte((uint8)op_code, bytecode);
+		emit_short(dataset_idx, bytecode);
+		emit_short(id_index_component->offset, bytecode);
+		emit_short(id_tag_component->offset, bytecode);
+	}
+
+	virtual bool finalize(_mesa_glsl_parse_state* parse_state, FVectorVMCompilationOutput&CompilationOutput, const unsigned num_temp_registers, unsigned *registers, unsigned op_idx)
+	{
+		if (!finalize_component_temporary_allocation(parse_state, id_index_component, num_temp_registers, registers, op_idx))
+		{
+			_mesa_glsl_error(parse_state, "Failed to allocate temporary registers");
+			return false;
+		}
+		if (!finalize_component_temporary_allocation(parse_state, id_tag_component, num_temp_registers, registers, op_idx))
 		{
 			_mesa_glsl_error(parse_state, "Failed to allocate temporary registers");
 			return false;
@@ -946,11 +1004,56 @@ struct op_index_acquire : public op_base
 	}
 
 	virtual void validate(_mesa_glsl_parse_state* parse_state)override
-	{		
-		if (dest_component->offset == INDEX_NONE)
-		{
-			_mesa_glsl_error(parse_state, "unknown error");//temp cop out, we have info to produce a good error message here.
-		}
+	{
+		validate_component_offset(parse_state, id_index_component);
+		validate_component_offset(parse_state, id_tag_component);
+	}
+};
+
+struct op_id_update : public op_base
+{
+	op_id_update()
+		: dataset_idx(INDEX_NONE)
+		, id_component(nullptr)
+		, index_component(nullptr)
+	{}
+
+	virtual op_id_update* as_id_update() { return this; }
+
+	unsigned dataset_idx;
+	variable_info_node* id_component;
+	variable_info_node* index_component;
+
+	virtual FString to_string()
+	{
+		check(id_component);
+		check(index_component);
+
+		return FString::Printf(TEXT("UpdateID(%d, [%d], [%d]);\n"),
+			dataset_idx,
+			id_component->offset,
+			index_component->offset);
+	}
+
+	virtual void add_to_bytecode(TArray<uint8>& bytecode)
+	{
+		check(id_component);
+		check(index_component);
+		emit_byte((uint8)op_code, bytecode);
+		emit_short(dataset_idx, bytecode);
+		emit_short(id_component->offset, bytecode);
+		emit_short(index_component->offset, bytecode);
+	}
+
+	virtual bool finalize(_mesa_glsl_parse_state* parse_state, FVectorVMCompilationOutput&CompilationOutput, const unsigned num_temp_registers, unsigned *registers, unsigned op_idx)
+	{
+		return true;
+	}
+
+	virtual void validate(_mesa_glsl_parse_state* parse_state)override
+	{
+		validate_component_offset(parse_state, id_component);
+		validate_component_offset(parse_state, index_component);
 	}
 };
 
@@ -1044,6 +1147,7 @@ struct op_external_func : public op_base
 
 	struct finalize_temp_register_ctx
 	{
+		_mesa_glsl_parse_state* parse_state;
 		unsigned num_temp_registers;
 		unsigned *registers;
 		unsigned op_idx;
@@ -1054,7 +1158,7 @@ struct op_external_func : public op_base
 	{
 		finalize_temp_register_ctx* ctx = (finalize_temp_register_ctx*)user_ptr;
 		check(node->is_scalar());
-		ctx->success &= finalize_component_temporary_allocation(node, ctx->num_temp_registers, ctx->registers, ctx->op_idx);
+		ctx->success &= finalize_component_temporary_allocation(ctx->parse_state, node, ctx->num_temp_registers, ctx->registers, ctx->op_idx);
 	}
 
 	virtual bool finalize(_mesa_glsl_parse_state* parse_state, FVectorVMCompilationOutput&CompilationOutput, const unsigned num_temp_registers, unsigned *registers, unsigned op_idx)
@@ -1078,6 +1182,7 @@ struct op_external_func : public op_base
 		}
 
 		finalize_temp_register_ctx ctx;
+		ctx.parse_state = parse_state;
 		ctx.num_temp_registers = num_temp_registers;
 		ctx.registers = registers;
 		ctx.op_idx = op_idx;
@@ -1102,17 +1207,11 @@ struct op_external_func : public op_base
 
 		for (variable_info_node* input : inputs)
 		{
-			if (input->offset == INDEX_NONE)
-			{
-				_mesa_glsl_error(parse_state, "unknown error");//temp cop out, we have info to produce a good error message here.
-			}
+			validate_component_offset(parse_state, input);
 		}
 		for (variable_info_node* output : outputs)
 		{
-			if (output->offset == INDEX_NONE)
-			{
-				_mesa_glsl_error(parse_state, "unknown error");//temp cop out, we have info to produce a good error message here.
-			}
+			validate_component_offset(parse_state, output);
 		}
 	}
 };
@@ -1255,14 +1354,17 @@ class ir_gen_vvm_visitor : public ir_hierarchical_visitor
 		case ir_var_uniform:
 		{
 			variable_info *varinfo = (variable_info*)hash_table_find(var_info_table, var);
-			check(varinfo == nullptr);
-
-			varinfo = ralloc(mem_ctx, variable_info);
-			hash_table_insert(var_info_table, varinfo, var);
+			
+			//We can possibly hit the same variables twice here as we do a pre pass to ensure we visit them in the correct order.
+			if (varinfo == nullptr)
+			{
+				varinfo = ralloc(mem_ctx, variable_info);
+				hash_table_insert(var_info_table, varinfo, var);
 
-			varinfo->Init(mem_ctx, var, EVectorVMOperandLocation::Constant);
+				varinfo->Init(mem_ctx, var, EVectorVMOperandLocation::Constant);
 
-			for_each_component(varinfo->root, &assign_uniform_loc, true, this);
+				for_each_component(varinfo->root, &assign_uniform_loc, true, this);
+			}
 		};
 		break;
 		}
@@ -1581,6 +1683,8 @@ class ir_gen_vvm_visitor : public ir_hierarchical_visitor
 		case EVectorVMOp::inputdata_noadvance_32bit: allocate<op_input_noadvance>(); break;
 		case EVectorVMOp::outputdata_32bit: allocate<op_output>(); break;
 		case EVectorVMOp::acquireindex: allocate<op_index_acquire>(); break;
+		case EVectorVMOp::acquire_id: allocate<op_id_acquire>(); break;
+		case EVectorVMOp::update_id: allocate<op_id_update>(); break;
 		case EVectorVMOp::external_func_call: allocate<op_external_func>(); break;
 
 		case EVectorVMOp::enter_stat_scope:
@@ -1746,6 +1850,64 @@ class ir_gen_vvm_visitor : public ir_hierarchical_visitor
 			acquire->dest_component = curr_node;
 			curr_node = nullptr;
 		}
+		else if (op_id_acquire* id_acquire = curr_op()->as_id_acquire())
+		{
+			check(num_operands == 3);
+
+			//Handle dataset index
+			exec_node* curr = call->actual_parameters.get_head();
+			check(curr);
+			ir_constant* dataset_idx_const = ((ir_rvalue*)curr)->as_constant();
+			check(dataset_idx_const);
+			id_acquire->dataset_idx = dataset_idx_const->get_int_component(0);
+
+			//Handle id index.
+			curr = curr->get_next();
+			check(curr);
+			ir_rvalue* id_index = (ir_rvalue*)curr;
+			id_index->accept(this);
+			check(curr_node);
+			id_acquire->id_index_component = curr_node;
+			curr_node = nullptr;
+
+			//Handle id tag.
+			curr = curr->get_next();
+			check(curr);
+			ir_rvalue* id_tag = (ir_rvalue*)curr;
+			id_tag->accept(this);
+			check(curr_node);
+			id_acquire->id_tag_component = curr_node;
+			curr_node = nullptr;
+		}
+		else if (op_id_update* update = curr_op()->as_id_update())
+		{
+			check(num_operands == 3);
+
+			//Handle dataset index
+			exec_node* curr = call->actual_parameters.get_head();
+			check(curr);
+			ir_constant* dataset_idx_const = ((ir_rvalue*)curr)->as_constant();
+			check(dataset_idx_const);
+			update->dataset_idx = dataset_idx_const->get_int_component(0);
+
+			//Handle ID parameter.
+			curr = curr->get_next();
+			check(curr);
+			ir_rvalue* id = (ir_rvalue*)curr;
+			id->accept(this);
+			check(curr_node);
+			update->id_component = curr_node;
+			curr_node = nullptr;
+
+			//Handle index parameter.
+			curr = curr->get_next();
+			check(curr);
+			ir_rvalue* index = (ir_rvalue*)curr;
+			index->accept(this);
+			check(curr_node);
+			update->index_component = curr_node;
+			curr_node = nullptr;
+		}
 		else if (op_external_func* func = curr_op()->as_external_func())
 		{
 			unsigned num_inputs = 0;
@@ -1838,6 +2000,7 @@ class ir_gen_vvm_visitor : public ir_hierarchical_visitor
 		//Allocate temp registers
 		const unsigned num_temp_registers = VectorVM::NumTempRegisters;
 		unsigned *registers = rzalloc_array(mem_ctx, unsigned, num_temp_registers);
+		FMemory::Memset(registers, 0xFF, sizeof(unsigned) * num_temp_registers);
 #if VM_VERBOSE_LOGGING
 		UE_LOG(LogVVMBackend, Log, TEXT("\n-------------------------------\nTemporary Allocations\n------------------------------\n"));
 #endif
@@ -1975,7 +2138,7 @@ public:
 		ir_gen_vvm_visitor genv(state, InCompOutput);
 
 		//Have to manually visit the uniform variables from cbuffers first.		
-		for (SCBuffer& buff : state->CBuffersOriginal)
+		for (SCBuffer& buff : state->CBuffersStructuresFlattened)
 		{
 			for (SCBufferMember& member : buff.Members)
 			{
diff --git a/Engine/Source/Editor/GraphEditor/Private/GraphEditorSettings.cpp b/Engine/Source/Editor/GraphEditor/Private/GraphEditorSettings.cpp
--- a/Engine/Source/Editor/GraphEditor/Private/GraphEditorSettings.cpp
+++ b/Engine/Source/Editor/GraphEditor/Private/GraphEditorSettings.cpp
@@ -16,6 +16,7 @@ UGraphEditorSettings::UGraphEditorSettings( const FObjectInitializer& ObjectInit
 	, PaddingRightOfInput(10.0f)
 	, PaddingLeftOfOutput(10.0f)
 	, PaddingTowardsNodeEdge(10.0f)
+	, bShowCommentBubbles(false)
 	, bTreatSplinesLikePins(true)
 	, SplineHoverTolerance(2.0f)
 	, ForwardSplineHorizontalDeltaRange(1000.0f)
diff --git a/Engine/Source/Editor/GraphEditor/Private/NodeFactory.cpp b/Engine/Source/Editor/GraphEditor/Private/NodeFactory.cpp
--- a/Engine/Source/Editor/GraphEditor/Private/NodeFactory.cpp
+++ b/Engine/Source/Editor/GraphEditor/Private/NodeFactory.cpp
@@ -400,3 +400,18 @@ FConnectionDrawingPolicy* FNodeFactory::CreateConnectionPolicy(const UEdGraphSch
 	// If we never picked a custom policy, use the uncustomized standard policy
 	return ConnectionDrawingPolicy;
 }
+
+TSharedPtr<SGraphNode> FGraphNodeFactory::CreateNodeWidget(UEdGraphNode* InNode)
+{
+	return FNodeFactory::CreateNodeWidget(InNode);
+}
+
+TSharedPtr<SGraphPin> FGraphNodeFactory::CreatePinWidget(UEdGraphPin* InPin)
+{
+	return FNodeFactory::CreatePinWidget(InPin);
+}
+
+FConnectionDrawingPolicy* FGraphNodeFactory::CreateConnectionPolicy(const UEdGraphSchema* Schema, int32 InBackLayerID, int32 InFrontLayerID, float ZoomFactor, const FSlateRect& InClippingRect, FSlateWindowElementList& InDrawElements, UEdGraph* InGraphObj)
+{
+	return FNodeFactory::CreateConnectionPolicy(Schema, InBackLayerID, InFrontLayerID, ZoomFactor, InClippingRect, InDrawElements, InGraphObj);
+}
diff --git a/Engine/Source/Editor/GraphEditor/Private/SGraphEditorImpl.cpp b/Engine/Source/Editor/GraphEditor/Private/SGraphEditorImpl.cpp
--- a/Engine/Source/Editor/GraphEditor/Private/SGraphEditorImpl.cpp
+++ b/Engine/Source/Editor/GraphEditor/Private/SGraphEditorImpl.cpp
@@ -794,6 +794,11 @@ void SGraphEditorImpl::CaptureKeyboard()
 	FSlateApplication::Get().SetKeyboardFocus(GraphPanel);
 }
 
+void SGraphEditorImpl::SetNodeFactory(const TSharedRef<class FGraphNodeFactory>& NewNodeFactory)
+{
+	GraphPanel->SetNodeFactory(NewNodeFactory);
+}
+
 /////////////////////////////////////////////////////
 
 #undef LOCTEXT_NAMESPACE 
diff --git a/Engine/Source/Editor/GraphEditor/Private/SGraphEditorImpl.h b/Engine/Source/Editor/GraphEditor/Private/SGraphEditorImpl.h
--- a/Engine/Source/Editor/GraphEditor/Private/SGraphEditorImpl.h
+++ b/Engine/Source/Editor/GraphEditor/Private/SGraphEditorImpl.h
@@ -157,6 +157,7 @@ public:
 	virtual void StraightenConnections() override;
 	virtual void StraightenConnections(UEdGraphPin* SourcePin, UEdGraphPin* PinToAlign) override;
 	virtual void CaptureKeyboard() override;
+	virtual void SetNodeFactory(const TSharedRef<class FGraphNodeFactory>& NewNodeFactory) override;
 	// End of SGraphEditor interface
 protected:
 	//
diff --git a/Engine/Source/Editor/GraphEditor/Private/SGraphNodeComment.cpp b/Engine/Source/Editor/GraphEditor/Private/SGraphNodeComment.cpp
--- a/Engine/Source/Editor/GraphEditor/Private/SGraphNodeComment.cpp
+++ b/Engine/Source/Editor/GraphEditor/Private/SGraphNodeComment.cpp
@@ -68,22 +68,37 @@ void SGraphNodeComment::Tick( const FGeometry& AllottedGeometry, const double In
 		CachedCommentTitle = CurrentCommentTitle;
 	}
 
+
 	const int32 CurrentWidth = static_cast<int32>(UserSize.X);
 	if (CurrentWidth != CachedWidth)
 	{
 		CachedWidth = CurrentWidth;
 	}
 
-	UEdGraphNode_Comment* CommentNode = CastChecked<UEdGraphNode_Comment>(GraphNode);
-	if (bCachedBubbleVisibility != CommentNode->bCommentBubbleVisible_InDetailsPanel)
+	const bool bShowCommentBubbles = GetDefault<UGraphEditorSettings>()->bShowCommentBubbles;
+	if (bShowCommentBubbles)
 	{
-		CommentBubble->UpdateBubble();
-		bCachedBubbleVisibility = CommentNode->bCommentBubbleVisible_InDetailsPanel;
-	}
+		UEdGraphNode_Comment* CommentNode = CastChecked<UEdGraphNode_Comment>(GraphNode);
+		if (bCachedBubbleVisibility != CommentNode->bCommentBubbleVisible_InDetailsPanel)
+		{
+			GetCommentBubble()->UpdateBubble();
+			bCachedBubbleVisibility = CommentNode->bCommentBubbleVisible_InDetailsPanel;
+		}
 
-	if (CachedFontSize != CommentNode->FontSize)
+		if (CachedFontSize != CommentNode->FontSize)
+		{
+			UpdateGraphNode();
+		}
+	}
+	else
 	{
-		UpdateGraphNode();
+		if (bCachedBubbleVisibility && CommentBubble.IsValid())
+		{
+			CommentBubble->SetVisibility(EVisibility::Collapsed);
+		}
+
+		CommentBubble.Reset();
+		bCachedBubbleVisibility = false;
 	}
 }
 
@@ -102,6 +117,36 @@ float SGraphNodeComment::GetWrapAt() const
 	return (float)(CachedWidth - 32.0f);
 }
 
+TSharedPtr<SCommentBubble> SGraphNodeComment::GetCommentBubble()
+{
+	if (!CommentBubble.IsValid())
+	{
+		// Create comment bubble
+		CommentBubble = SNew(SCommentBubble)
+			.GraphNode(GraphNode)
+			.Text(this, &SGraphNodeComment::GetNodeComment)
+			.OnTextCommitted(this, &SGraphNodeComment::OnNameTextCommited)
+			.ColorAndOpacity(this, &SGraphNodeComment::GetCommentBubbleColor)
+			.AllowPinning(true)
+			.EnableTitleBarBubble(false)
+			.EnableBubbleCtrls(false)
+			.GraphLOD(this, &SGraphNode::GetCurrentLOD)
+			.InvertLODCulling(true)
+			.IsGraphNodeHovered(this, &SGraphNode::IsHovered);
+
+		GetOrAddSlot(ENodeZone::TopCenter)
+			.SlotOffset(TAttribute<FVector2D>(CommentBubble.Get(), &SCommentBubble::GetOffset))
+			.SlotSize(TAttribute<FVector2D>(CommentBubble.Get(), &SCommentBubble::GetSize))
+			.AllowScaling(TAttribute<bool>(CommentBubble.Get(), &SCommentBubble::IsScalingAllowed))
+			.VAlign(VAlign_Top)
+			[
+				CommentBubble.ToSharedRef()
+			];
+	}
+
+	return CommentBubble;
+}
+
 bool SGraphNodeComment::IsNameReadOnly() const
 {
 	return !IsEditable.Get() || SGraphNode::IsNameReadOnly();
@@ -109,6 +154,8 @@ bool SGraphNodeComment::IsNameReadOnly() const
 
 void SGraphNodeComment::UpdateGraphNode()
 {
+	const bool bShowCommentBubbles = GetDefault<UGraphEditorSettings>()->bShowCommentBubbles;
+
 	// No pins in a comment box
 	InputPins.Empty();
 	OutputPins.Empty();
@@ -119,7 +166,7 @@ void SGraphNodeComment::UpdateGraphNode()
 
 	// Remember if we should be showing the bubble
 	UEdGraphNode_Comment* CommentNode = CastChecked<UEdGraphNode_Comment>(GraphNode);
-	bCachedBubbleVisibility = CommentNode->bCommentBubbleVisible_InDetailsPanel;
+	bCachedBubbleVisibility = bShowCommentBubbles ? CommentNode->bCommentBubbleVisible_InDetailsPanel : false;
 
 	// Setup a tag for this node
 	FString TagName;
@@ -201,27 +248,11 @@ void SGraphNodeComment::UpdateGraphNode()
 			]			
 		];
 
-	// Create comment bubble
-	CommentBubble = SNew(SCommentBubble)
-	.GraphNode(GraphNode)
-	.Text(this, &SGraphNodeComment::GetNodeComment)
-	.OnTextCommitted(this, &SGraphNodeComment::OnNameTextCommited)
-	.ColorAndOpacity(this, &SGraphNodeComment::GetCommentBubbleColor )
-	.AllowPinning(true)
-	.EnableTitleBarBubble(false)
-	.EnableBubbleCtrls(false)
-	.GraphLOD(this, &SGraphNode::GetCurrentLOD)
-	.InvertLODCulling(true)
-	.IsGraphNodeHovered(this, &SGraphNode::IsHovered);
-
-	GetOrAddSlot(ENodeZone::TopCenter)
-	.SlotOffset(TAttribute<FVector2D>(CommentBubble.Get(), &SCommentBubble::GetOffset))
-	.SlotSize(TAttribute<FVector2D>(CommentBubble.Get(), &SCommentBubble::GetSize))
-	.AllowScaling(TAttribute<bool>(CommentBubble.Get(), &SCommentBubble::IsScalingAllowed))
-	.VAlign(VAlign_Top)
-	[
-		CommentBubble.ToSharedRef()
-	];
+	if (bShowCommentBubbles)
+	{
+		// Create comment bubble
+		GetCommentBubble();
+	}
 }
 
 FVector2D SGraphNodeComment::ComputeDesiredSize( float ) const
diff --git a/Engine/Source/Editor/GraphEditor/Private/SGraphPalette.cpp b/Engine/Source/Editor/GraphEditor/Private/SGraphPalette.cpp
--- a/Engine/Source/Editor/GraphEditor/Private/SGraphPalette.cpp
+++ b/Engine/Source/Editor/GraphEditor/Private/SGraphPalette.cpp
@@ -164,7 +164,7 @@ FText SGraphPaletteItem::GetDisplayText() const
 
 FText SGraphPaletteItem::GetItemTooltip() const
 {
-	return ActionPtr.Pin()->GetMenuDescription();
+	return ActionPtr.Pin()->GetTooltipDescription();
 }
 
 
diff --git a/Engine/Source/Editor/GraphEditor/Private/SGraphPanel.cpp b/Engine/Source/Editor/GraphEditor/Private/SGraphPanel.cpp
--- a/Engine/Source/Editor/GraphEditor/Private/SGraphPanel.cpp
+++ b/Engine/Source/Editor/GraphEditor/Private/SGraphPanel.cpp
@@ -297,7 +297,15 @@ int32 SGraphPanel::OnPaint( const FPaintArgs& Args, const FGeometry& AllottedGeo
 	// Draw connections between pins 
 	if (Children.Num() > 0 )
 	{
-        FConnectionDrawingPolicy* ConnectionDrawingPolicy = FNodeFactory::CreateConnectionPolicy(Schema, WireLayerId, MaxLayerId, ZoomFactor, MyCullingRect, OutDrawElements, GraphObj);
+		FConnectionDrawingPolicy* ConnectionDrawingPolicy = nullptr;
+		if (NodeFactory.IsValid())
+		{
+			ConnectionDrawingPolicy = NodeFactory->CreateConnectionPolicy(Schema, WireLayerId, MaxLayerId, ZoomFactor, MyCullingRect, OutDrawElements, GraphObj);
+		}
+		else
+		{
+			ConnectionDrawingPolicy = FNodeFactory::CreateConnectionPolicy(Schema, WireLayerId, MaxLayerId, ZoomFactor, MyCullingRect, OutDrawElements, GraphObj);
+		}
 
 		TArray<TSharedPtr<SGraphPin>> OverridePins;
 		for (const FGraphPinHandle& Handle : PreviewConnectorFromPins)
@@ -1335,7 +1343,16 @@ const TSharedRef<SGraphNode> SGraphPanel::GetChild(int32 ChildIndex)
 
 void SGraphPanel::AddNode(UEdGraphNode* Node, AddNodeBehavior Behavior)
 {
-	TSharedPtr<SGraphNode> NewNode = FNodeFactory::CreateNodeWidget(Node);
+	TSharedPtr<SGraphNode> NewNode;
+	if (NodeFactory.IsValid())
+	{
+		NewNode = NodeFactory->CreateNodeWidget(Node);
+	}
+	else
+	{
+		NewNode = FNodeFactory::CreateNodeWidget(Node);
+	}
+
 	check(NewNode.IsValid());
 
 	const bool bWasUserAdded = 
@@ -1692,3 +1709,8 @@ EActiveTimerReturnType SGraphPanel::InvalidatePerTick(double InCurrentTime, floa
 	}
 	return EActiveTimerReturnType::Continue;
 }
+
+void SGraphPanel::SetNodeFactory(const TSharedRef<class FGraphNodeFactory>& NewNodeFactory)
+{
+	NodeFactory = NewNodeFactory;
+}
diff --git a/Engine/Source/Editor/GraphEditor/Public/GraphEditorDragDropAction.h b/Engine/Source/Editor/GraphEditor/Public/GraphEditorDragDropAction.h
--- a/Engine/Source/Editor/GraphEditor/Public/GraphEditorDragDropAction.h
+++ b/Engine/Source/Editor/GraphEditor/Public/GraphEditorDragDropAction.h
@@ -58,11 +58,10 @@ protected:
 	/** Constructs the window and widget if applicable */
 	virtual void Construct() override;
 
-private:
-
-	EVisibility GetIconVisible() const;
-	EVisibility GetErrorIconVisible() const;
+	virtual EVisibility GetIconVisible() const;
+	virtual EVisibility GetErrorIconVisible() const;
 
+private:
 	// The pin that the drag action is currently hovering over
 	FEdGraphPinReference HoveredPin;
 
diff --git a/Engine/Source/Editor/GraphEditor/Public/GraphEditorSettings.h b/Engine/Source/Editor/GraphEditor/Public/GraphEditorSettings.h
--- a/Engine/Source/Editor/GraphEditor/Public/GraphEditorSettings.h
+++ b/Engine/Source/Editor/GraphEditor/Public/GraphEditorSettings.h
@@ -65,6 +65,10 @@ public:
 	UPROPERTY(config, EditAnywhere, Category=GeneralStyle, AdvancedDisplay)
 	float PaddingTowardsNodeEdge;
 
+	/** Whether the comment bubbles for comment nodes are visible or hidden. */
+	UPROPERTY(config, EditAnywhere, Category = GeneralStyle, AdvancedDisplay, meta = (DisplayName = "Show Comment Bubbles"))
+	bool bShowCommentBubbles;
+
 	/** Returns the padding to use for an input pin on a graph node */
 	FMargin GetInputPinPadding() const;
 
diff --git a/Engine/Source/Editor/GraphEditor/Public/NodeFactory.h b/Engine/Source/Editor/GraphEditor/Public/NodeFactory.h
--- a/Engine/Source/Editor/GraphEditor/Public/NodeFactory.h
+++ b/Engine/Source/Editor/GraphEditor/Public/NodeFactory.h
@@ -30,3 +30,17 @@ public:
     /** Create a pin connection factory for the supplied schema */
     static FConnectionDrawingPolicy* CreateConnectionPolicy(const UEdGraphSchema* Schema, int32 InBackLayerID, int32 InFrontLayerID, float ZoomFactor, const FSlateRect& InClippingRect, FSlateWindowElementList& InDrawElements, UEdGraph* InGraphObj);
 };
+
+
+class GRAPHEDITOR_API FGraphNodeFactory
+{
+public:
+	/** Create a widget for the supplied node */
+	virtual TSharedPtr<SGraphNode> CreateNodeWidget(UEdGraphNode* InNode);
+
+	/** Create a widget for the supplied pin */
+	virtual TSharedPtr<SGraphPin> CreatePinWidget(UEdGraphPin* InPin);
+
+	/** Create a pin connection factory for the supplied schema */
+	virtual FConnectionDrawingPolicy* CreateConnectionPolicy(const UEdGraphSchema* Schema, int32 InBackLayerID, int32 InFrontLayerID, float ZoomFactor, const FSlateRect& InClippingRect, FSlateWindowElementList& InDrawElements, UEdGraph* InGraphObj);
+};
\ No newline at end of file
diff --git a/Engine/Source/Editor/GraphEditor/Public/SGraphNodeComment.h b/Engine/Source/Editor/GraphEditor/Public/SGraphNodeComment.h
--- a/Engine/Source/Editor/GraphEditor/Public/SGraphNodeComment.h
+++ b/Engine/Source/Editor/GraphEditor/Public/SGraphNodeComment.h
@@ -92,6 +92,10 @@ private:
 	float GetWrapAt() const;
 
 private:
+
+	/** Gets the comment bubble and creates one if it doesn't exist. */
+	TSharedPtr<SCommentBubble> GetCommentBubble();
+
 	/** The comment bubble widget (used when zoomed out) */
 	TSharedPtr<SCommentBubble> CommentBubble;
 
diff --git a/Engine/Source/Editor/GraphEditor/Public/SGraphNodeDefault.h b/Engine/Source/Editor/GraphEditor/Public/SGraphNodeDefault.h
new file mode 100644
--- /dev/null
+++ b/Engine/Source/Editor/GraphEditor/Public/SGraphNodeDefault.h
@@ -0,0 +1,21 @@
+// Copyright 1998-2018 Epic Games, Inc. All Rights Reserved.
+
+#pragma once
+
+#include "CoreMinimal.h"
+#include "Widgets/DeclarativeSyntaxSupport.h"
+#include "SGraphNode.h"
+
+class GRAPHEDITOR_API SGraphNodeDefault : public SGraphNode
+{
+public:
+
+	SLATE_BEGIN_ARGS( SGraphNodeDefault )
+		: _GraphNodeObj( static_cast<UEdGraphNode*>(NULL) )
+		{}
+
+		SLATE_ARGUMENT( UEdGraphNode*, GraphNodeObj )
+	SLATE_END_ARGS()
+
+	void Construct( const FArguments& InArgs );
+};
diff --git a/Engine/Source/Editor/GraphEditor/Public/SGraphNodeKnot.h b/Engine/Source/Editor/GraphEditor/Public/SGraphNodeKnot.h
new file mode 100644
--- /dev/null
+++ b/Engine/Source/Editor/GraphEditor/Public/SGraphNodeKnot.h
@@ -0,0 +1,55 @@
+// Copyright 1998-2018 Epic Games, Inc. All Rights Reserved.
+
+#pragma once
+
+#include "CoreMinimal.h"
+#include "Widgets/DeclarativeSyntaxSupport.h"
+#include "SGraphPin.h"
+#include "SGraphNodeDefault.h"
+
+class SCommentBubble;
+
+/** The visual representation of a control point meant to adjust how connections are routed, also known as a Reroute node.
+ * The input knot node should have properly implemented ShouldDrawNodeAsControlPointOnly to return true with valid indices for its pins.
+ */
+class GRAPHEDITOR_API SGraphNodeKnot : public SGraphNodeDefault
+{
+public:
+	SLATE_BEGIN_ARGS(SGraphNodeKnot) {}
+	SLATE_END_ARGS()
+
+	void Construct(const FArguments& InArgs, class UEdGraphNode* InKnot);
+
+	// SGraphNode interface
+	virtual void UpdateGraphNode() override;
+	virtual const FSlateBrush* GetShadowBrush(bool bSelected) const override;
+	virtual TSharedPtr<SGraphPin> CreatePinWidget(UEdGraphPin* Pin) const override;
+	virtual void AddPin(const TSharedRef<SGraphPin>& PinToAdd) override;
+	virtual void RequestRenameOnSpawn() override { }
+	// End of SGraphNode interface
+
+	// SWidget interface
+	virtual void OnMouseEnter( const FGeometry& MyGeometry, const FPointerEvent& MouseEvent ) override;
+	virtual void OnMouseLeave( const FPointerEvent& MouseEvent ) override;
+	// End of SWidget interface
+
+protected:
+	/** Returns Offset to center comment on the node's only pin */
+	FVector2D GetCommentOffset() const;
+protected:
+
+	/** Toggles the hovered visibility state */
+	virtual void OnCommentBubbleToggled(bool bInCommentBubbleVisible) override;
+
+	/** If bHoveredCommentVisibility is true, hides the comment bubble after a change is committed */
+	virtual void OnCommentTextCommitted(const FText& NewComment, ETextCommit::Type CommitInfo) override;
+
+	/** The hovered visibility state. If false, comment bubble will only appear on hover. */
+	bool bAlwaysShowCommentBubble;
+
+	/** SharedPtr to comment bubble */
+	TSharedPtr<SCommentBubble> CommentBubble;
+
+	const FSlateBrush* ShadowBrush;
+	const FSlateBrush* ShadowBrushSelected;
+};
diff --git a/Engine/Source/Editor/GraphEditor/Public/SGraphPanel.h b/Engine/Source/Editor/GraphEditor/Public/SGraphPanel.h
--- a/Engine/Source/Editor/GraphEditor/Public/SGraphPanel.h
+++ b/Engine/Source/Editor/GraphEditor/Public/SGraphPanel.h
@@ -176,6 +176,9 @@ public:
 	/** When the graph panel needs to be dynamically refreshing for animations, this function is registered to tick and invalidate the UI. */
 	EActiveTimerReturnType InvalidatePerTick(double InCurrentTime, float InDeltaTime);
 
+	/** Sets the current widget factory. */
+	void SetNodeFactory(const TSharedRef<class FGraphNodeFactory>& NewNodeFactory);
+
 protected:
 
 	void NotifyGraphChanged ( const struct FEdGraphEditAction& InAction);
@@ -311,4 +314,7 @@ private:
 
 	/** Amount of time left to invalidate the UI per tick */
 	float TimeLeftToInvalidatePerTick;
+
+	/** The current node factory to create nodes, pins and connections. Uses the static FNodeFactory if not set. */
+	TSharedPtr<class FGraphNodeFactory> NodeFactory;
 };
diff --git a/Engine/Source/Editor/MeshPaint/Private/MeshPaintHelpers.cpp b/Engine/Source/Editor/MeshPaint/Private/MeshPaintHelpers.cpp
--- a/Engine/Source/Editor/MeshPaint/Private/MeshPaintHelpers.cpp
+++ b/Engine/Source/Editor/MeshPaint/Private/MeshPaintHelpers.cpp
@@ -1679,6 +1679,7 @@ void MeshPaintHelpers::ApplyVertexColorsToAllLODs(IMeshPaintGeometryAdapter& Geo
 	if (Mesh)
 	{
 		FSkeletalMeshRenderData* Resource = Mesh->GetResourceForRendering();
+		FSkeletalMeshModel* SrcMesh = Mesh->GetImportedModel();
 		if (Resource)
 		{
 			const int32 NumLODs = Resource->LODRenderData.Num();
@@ -1712,6 +1713,8 @@ void MeshPaintHelpers::ApplyVertexColorsToAllLODs(IMeshPaintGeometryAdapter& Geo
 				{
 					// Do something
 					FSkeletalMeshLODRenderData& ApplyLOD = Resource->LODRenderData[LODIndex];
+					FSkeletalMeshLODModel& SrcLOD = SrcMesh->LODModels[LODIndex];
+
 					FBox CombinedBounds = BaseBounds;
 					Mesh->GetLODInfo(LODIndex)->bHasPerLODVertexColors = false;
 
@@ -1737,6 +1740,7 @@ void MeshPaintHelpers::ApplyVertexColorsToAllLODs(IMeshPaintGeometryAdapter& Geo
 					// Iterate over each new vertex position, attempting to find the old vertex it is closest to, applying
 					// the color of the old vertex to the new position if possible.
 					const float DistanceOverNormalThreshold = KINDA_SMALL_NUMBER;
+					check(SrcLOD.NumVertices == ApplyLOD.GetNumVertices());
 					for (uint32 VertexIndex = 0; VertexIndex < ApplyLOD.GetNumVertices(); ++VertexIndex)
 					{
 						TArray<FPaintedMeshVertex> PointsToConsider;
@@ -1818,6 +1822,11 @@ void MeshPaintHelpers::ApplyVertexColorsToAllLODs(IMeshPaintGeometryAdapter& Geo
 							}
 
 							ApplyLOD.StaticVertexBuffers.ColorVertexBuffer.VertexColor(VertexIndex) = PointsToConsider[BestVertexIndex].Color;
+							// Also apply to the skeletal mesh source mesh
+							int32 SectionIndex = INDEX_NONE;
+							int32 SectionVertexIndex = INDEX_NONE;
+							SrcLOD.GetSectionFromVertexIndex(VertexIndex, SectionIndex, SectionVertexIndex);
+							SrcLOD.Sections[SectionIndex].SoftVertices[SectionVertexIndex].Color = PointsToConsider[BestVertexIndex].Color;
 						}
 					}
 				}
diff --git a/Engine/Source/Editor/MeshPaintMode/Private/PaintModeCommands.cpp b/Engine/Source/Editor/MeshPaintMode/Private/PaintModeCommands.cpp
--- a/Engine/Source/Editor/MeshPaintMode/Private/PaintModeCommands.cpp
+++ b/Engine/Source/Editor/MeshPaintMode/Private/PaintModeCommands.cpp
@@ -45,8 +45,17 @@ void FPaintModeCommands::RegisterCommands()
 	Commands.Add(SaveTexturePaint);
 
 	UI_COMMAND(SwitchForeAndBackgroundColor, "Switch Fore and Background Color", "Switches the Fore and Background Colors used for Vertex Painting", EUserInterfaceActionType::None, FInputChord(EKeys::X));
-	Commands.Add(SaveTexturePaint);
+	Commands.Add(SwitchForeAndBackgroundColor);
+
+	UI_COMMAND(PropagateVertexColorsToLODs, "Apply Vertex Colors to LODs", "Applied the Vertex Colors from LOD0 to all LOD levels", EUserInterfaceActionType::Button, FInputChord());
+	Commands.Add(PropagateVertexColorsToLODs);
+
+
+	UI_COMMAND(CycleToNextLOD, "Cycle to next Mesh LOD", "Cycles to the next possible Mesh LOD to Paint on", EUserInterfaceActionType::None, FInputChord(EKeys::N));
+	Commands.Add(CycleToNextLOD);
 
+	UI_COMMAND(CycleToPreviousLOD, "Cycle to previous Mesh LOD", "Cycles to the previous possible Mesh LOD to Paint on", EUserInterfaceActionType::None, FInputChord(EKeys::B));
+	Commands.Add(CycleToPreviousLOD);
 }
 
 #undef LOCTEXT_NAMESPACE
diff --git a/Engine/Source/Editor/MeshPaintMode/Private/PaintModeCommands.h b/Engine/Source/Editor/MeshPaintMode/Private/PaintModeCommands.h
--- a/Engine/Source/Editor/MeshPaintMode/Private/PaintModeCommands.h
+++ b/Engine/Source/Editor/MeshPaintMode/Private/PaintModeCommands.h
@@ -31,9 +31,13 @@ public:
 	TSharedPtr<FUICommandInfo> Save;
 
 	TSharedPtr<FUICommandInfo> SwitchForeAndBackgroundColor;
+	TSharedPtr<FUICommandInfo> CycleToNextLOD;
+	TSharedPtr<FUICommandInfo> CycleToPreviousLOD;
 
 	TSharedPtr<FUICommandInfo> PropagateTexturePaint;
 	TSharedPtr<FUICommandInfo> SaveTexturePaint;
 
+	TSharedPtr<FUICommandInfo> PropagateVertexColorsToLODs;
+
 	TArray<TSharedPtr<FUICommandInfo>> Commands;
 };
\ No newline at end of file
diff --git a/Engine/Source/Editor/MeshPaintMode/Private/PaintModePainter.cpp b/Engine/Source/Editor/MeshPaintMode/Private/PaintModePainter.cpp
--- a/Engine/Source/Editor/MeshPaintMode/Private/PaintModePainter.cpp
+++ b/Engine/Source/Editor/MeshPaintMode/Private/PaintModePainter.cpp
@@ -117,6 +117,8 @@ void FPaintModePainter::RegisterVertexPaintCommands()
 	UICommandList->MapAction(FPaintModeCommands::Get().Remove, FUIAction(FExecuteAction::CreateRaw(this, &FPaintModePainter::RemoveVertexColors), FCanExecuteAction::CreateRaw(this, &FPaintModePainter::CanRemoveInstanceColors)));
 
 	UICommandList->MapAction(FPaintModeCommands::Get().Fix, FUIAction(FExecuteAction::CreateRaw(this, &FPaintModePainter::FixVertexColors), FCanExecuteAction::CreateRaw(this, &FPaintModePainter::DoesRequireVertexColorsFixup)));
+
+	UICommandList->MapAction(FPaintModeCommands::Get().PropagateVertexColorsToLODs, FUIAction(FExecuteAction::CreateRaw(this, &FPaintModePainter::PropagateVertexColorsToLODs), FCanExecuteAction::CreateRaw(this, &FPaintModePainter::CanPropagateVertexColorsToLODs)));
 }
 
 void FPaintModePainter::Render(const FSceneView* View, FViewport* Viewport, FPrimitiveDrawInterface* PDI)
@@ -255,6 +257,12 @@ bool FPaintModePainter::CanCopyInstanceVertexColors() const
 	return bValidSelection && (NumValidMeshes != 0);
 }
 
+bool FPaintModePainter::CanPropagateVertexColorsToLODs() const
+{
+	// Can propagate when the mesh contains per-lod vertex colors or when we are not painting to a specific lod
+	return bSelectionContainsPerLODColors || !PaintSettings->VertexPaintSettings.bPaintOnSpecificLOD;
+}
+
 void FPaintModePainter::CopyVertexColors()
 {
 	const TArray<UStaticMeshComponent*> StaticMeshComponents = GetSelectedComponents<UStaticMeshComponent>();
@@ -401,6 +409,91 @@ void FPaintModePainter::RemoveVertexColors()
 	}
 }
 
+void FPaintModePainter::PropagateVertexColorsToLODs()
+{	
+	//Only show the lost data warning if there is actually some data to loose
+	bool bAbortChange = false;
+	if (bSelectionContainsPerLODColors)
+	{
+		//Warn the user he will loose is custom painting data
+		FSuppressableWarningDialog::FSetupInfo SetupInfo(LOCTEXT("LooseLowersLODsVertexColorsPrompt_Message", "Propagating Vertex Colors from LOD0 to all lower LODs. This mean all lower LODs custom vertex painting will be lost."),
+			LOCTEXT("LooseLowersLODsVertexColorsPrompt_Title", "Warning: Lowers LODs custom vertex painting will be lost!"), "Warning_LooseLowersLODsVertexColorsPrompt");
+
+		SetupInfo.ConfirmText = LOCTEXT("LooseLowersLODsVertexColorsPrompt_ConfirmText", "Continue");
+		SetupInfo.CancelText = LOCTEXT("LooseLowersLODsVertexColorsPrompt_CancelText", "Abort");
+		SetupInfo.CheckBoxText = LOCTEXT("LooseLowersLODsVertexColorsPrompt_CheckBoxText", "Always copy vertex colors without prompting");
+
+		FSuppressableWarningDialog LooseLowersLODsVertexColorsWarning(SetupInfo);
+
+		// Prompt the user to see if they really want to propagate the base lod vert colors to the lowers LODs.
+		if (LooseLowersLODsVertexColorsWarning.ShowModal() == FSuppressableWarningDialog::Cancel)
+		{
+			bAbortChange = true;
+		}
+		else
+		{
+			//Remove painting on all lowers LODs before doing the propagation
+			for (UMeshComponent* SelectedComponent : PaintableComponents)
+			{
+				UStaticMeshComponent *StaticMeshComponent = Cast<UStaticMeshComponent>(SelectedComponent);
+				if (StaticMeshComponent && StaticMeshComponent->GetStaticMesh())
+				{
+					// Mark the mesh component as modified
+					StaticMeshComponent->Modify();
+
+					// If this is called from the Remove button being clicked the SMC wont be in a Reregister context,
+					// but when it gets called from a Paste or Copy to Source operation it's already inside a more specific
+					// SMCRecreateScene context so we shouldn't put it inside another one.
+					if (StaticMeshComponent->IsRenderStateCreated())
+					{
+						// Detach all instances of this static mesh from the scene.
+						FComponentReregisterContext ComponentReregisterContext(StaticMeshComponent);
+
+						for (int32 LODIndex = 1; LODIndex < StaticMeshComponent->LODData.Num(); ++LODIndex)
+						{
+							StaticMeshComponent->RemoveInstanceVertexColorsFromLOD(LODIndex);
+						}
+					}
+					else
+					{
+						for (int32 LODIndex = 1; LODIndex < StaticMeshComponent->LODData.Num(); ++LODIndex)
+						{
+							StaticMeshComponent->RemoveInstanceVertexColorsFromLOD(LODIndex);
+						}
+					}
+				}
+			}
+		}
+	}
+
+	//The user cancel the change, avoid changing the value
+	if (bAbortChange)
+	{
+		return;
+	}	
+
+	for (UMeshComponent* SelectedComponent : PaintableComponents)
+	{
+		TSharedPtr<IMeshPaintGeometryAdapter> MeshAdapter = ComponentToAdapterMap.FindChecked(SelectedComponent);
+		MeshPaintHelpers::ApplyVertexColorsToAllLODs(*MeshAdapter, SelectedComponent);
+		FComponentReregisterContext ReregisterContext(SelectedComponent);
+	}
+
+	Refresh();
+}
+
+void FPaintModePainter::CycleMeshLODs(int32 Direction)
+{
+	if (bCachedForceLOD)
+	{
+		const int32 MaxLODIndex = GetMaxLODIndexToPaint() + 1;
+		const int32 NewLODIndex = PaintSettings->VertexPaintSettings.LODIndex + Direction;
+		const int32 AdjustedLODIndex = NewLODIndex < 0 ? MaxLODIndex + NewLODIndex : NewLODIndex % MaxLODIndex;
+		PaintSettings->VertexPaintSettings.LODIndex = AdjustedLODIndex;
+		PaintLODChanged();
+	}
+}
+
 bool FPaintModePainter::CanSaveMeshPackages() const
 {
 	// Check whether or not any of our selected mesh components contain mesh objects which require saving
@@ -544,6 +637,9 @@ void FPaintModePainter::RegisterCommands(TSharedRef<FUICommandList> CommandList)
 		}
 	}));
 
+	CommandList->MapAction(Commands.CycleToNextLOD, FExecuteAction::CreateRaw(this, &FPaintModePainter::CycleMeshLODs, 1));
+	CommandList->MapAction(Commands.CycleToPreviousLOD, FExecuteAction::CreateRaw(this, &FPaintModePainter::CycleMeshLODs, -1));
+
 	/** Map commit texture painting to commiting all the outstanding paint changes */
 	auto CanCommitLambda = [this]() -> bool { return GetNumberOfPendingPaintChanges() > 0; };
 	CommandList->MapAction(Commands.CommitTexturePainting, FExecuteAction::CreateLambda([this]() { CommitAllPaintedTextures(); }), FCanExecuteAction::CreateLambda([this]() -> bool { return GetNumberOfPendingPaintChanges() > 0; }));
@@ -646,10 +742,6 @@ void FPaintModePainter::ActorDeselected(AActor* Actor)
 			/** If in vertex painting mode ensure we reset forced LOD rendering and propagate vertex colors to other LODs if necessary */
 			else if (PaintSettings->PaintMode == EPaintMode::Vertices)
 			{
-				if (!bCachedForceLOD)
-				{
-					MeshPaintHelpers::ApplyVertexColorsToAllLODs(*Adapter, MeshComponent);
-				}
 				MeshPaintHelpers::ForceRenderMeshLOD(MeshComponent, -1);
 				FComponentReregisterContext ReregisterContext(MeshComponent);
 			}
@@ -684,6 +776,11 @@ void FPaintModePainter::FinishPainting()
 	/** Reset state and apply outstanding paint data*/
 	IMeshPainter::FinishPainting();	
 	FinishPaintingTexture();	
+
+	if (PaintSettings->PaintMode == EPaintMode::Vertices && !PaintSettings->VertexPaintSettings.bPaintOnSpecificLOD)
+	{
+		PropagateVertexColorsToLODs();
+	}
 }
 
 bool FPaintModePainter::PaintInternal(const FVector& InCameraOrigin, const FVector& InRayOrigin, const FVector& InRayDirection, EMeshPaintAction PaintAction, float PaintStrength)
@@ -930,13 +1027,6 @@ void FPaintModePainter::Reset()
 	//If we're painting vertex colors then propagate the painting done on LOD0 to all lower LODs. 
 	//Then stop forcing the LOD level of the mesh to LOD0.
 	ApplyForcedLODIndex(-1);
-	if (!PaintSettings->VertexPaintSettings.bPaintOnSpecificLOD)
-	{
-		for (auto Pair : ComponentToAdapterMap)
-		{
-			MeshPaintHelpers::ApplyVertexColorsToAllLODs(*Pair.Value.Get(), Pair.Key);
-		}
-	}
 
 	// If the user has pending changes and the editor is not exiting, we want to do the commit for all the modified textures.
 	if ((GetNumberOfPendingPaintChanges() > 0) && !GIsRequestingExit)
@@ -1025,83 +1115,16 @@ void FPaintModePainter::LODPaintStateChanged(const bool bLODPaintingEnabled)
 {
 	checkf(PaintSettings->PaintMode == EPaintMode::Vertices, TEXT("Can only change this state in vertex paint mode"));
 	bool AbortChange = false;
-	if (!bLODPaintingEnabled)
-	{
-		//Only show the lost data warning if there is actually some data to loose
-		if (bSelectionContainsPerLODColors)
-		{
-			//Warn the user he will loose is custom painting data
-			FSuppressableWarningDialog::FSetupInfo SetupInfo(LOCTEXT("LooseLowersLODsVertexColorsPrompt_Message", "Changing from custom LODs to base LOD only painting will propagate the base lod vertex color to all lowers LODs. This mean all lowers LODs custom vertex painting will be lost."),
-				LOCTEXT("LooseLowersLODsVertexColorsPrompt_Title", "Warning: Lowers LODs custom vertex painting will be lost!"), "Warning_LooseLowersLODsVertexColorsPrompt");
-
-			SetupInfo.ConfirmText = LOCTEXT("LooseLowersLODsVertexColorsPrompt_ConfirmText", "Continue");
-			SetupInfo.CancelText = LOCTEXT("LooseLowersLODsVertexColorsPrompt_CancelText", "Abort");
-			SetupInfo.CheckBoxText = LOCTEXT("LooseLowersLODsVertexColorsPrompt_CheckBoxText", "Always copy vertex colors without prompting");
-
-			FSuppressableWarningDialog LooseLowersLODsVertexColorsWarning(SetupInfo);
-
-			// Prompt the user to see if they really want to propagate the base lod vert colors to the lowers LODs.
-			if (LooseLowersLODsVertexColorsWarning.ShowModal() == FSuppressableWarningDialog::Cancel)
-			{
-				AbortChange = true;
-			}
-			else
-			{
-				//Remove painting on all lowers LODs before doing the propagation
-				for (UMeshComponent* SelectedComponent : PaintableComponents)
-				{
-					UStaticMeshComponent *StaticMeshComponent = Cast<UStaticMeshComponent>(SelectedComponent);
-					if (StaticMeshComponent && StaticMeshComponent->GetStaticMesh())
-					{
-						// Mark the mesh component as modified
-						StaticMeshComponent->Modify();
-
-						// If this is called from the Remove button being clicked the SMC wont be in a Reregister context,
-						// but when it gets called from a Paste or Copy to Source operation it's already inside a more specific
-						// SMCRecreateScene context so we shouldn't put it inside another one.
-						if (StaticMeshComponent->IsRenderStateCreated())
-						{
-							// Detach all instances of this static mesh from the scene.
-							FComponentReregisterContext ComponentReregisterContext(StaticMeshComponent);
-
-							for (int32 LODIndex = 1; LODIndex < StaticMeshComponent->LODData.Num(); ++LODIndex)
-							{
-								StaticMeshComponent->RemoveInstanceVertexColorsFromLOD(LODIndex);
-							}
-						}
-						else
-						{
-							for (int32 LODIndex = 1; LODIndex < StaticMeshComponent->LODData.Num(); ++LODIndex)
-							{
-								StaticMeshComponent->RemoveInstanceVertexColorsFromLOD(LODIndex);
-							}
-						}
-					}
-				}
-			}
-		}
 	
-		//The user cancel the change, avoid changing the value
-		if (AbortChange)
-		{
-			return;
-		}
-	}
+	// Set actual flag in the settings struct
+	PaintSettings->VertexPaintSettings.bPaintOnSpecificLOD = bLODPaintingEnabled;
 
-	for (UMeshComponent* SelectedComponent : PaintableComponents)
+	if (!bLODPaintingEnabled)
 	{
-		//Always propagate the base LOD when we switch the painting to all LODs
-		//In case we go from false to true, we want to tweak the propagation of the base LOD
-		//In case we go from true to false, we want to force a propagate to be sure the custom painting is done over what the user already paint
-		//In case the staticmesh has some data comming from fbx, the propagate will not do anything
-		//if there is nothing that was overrite in the base LOD. So the user will be able to paint
-		//over the painting done in a DCC.
-		TSharedPtr<IMeshPaintGeometryAdapter> MeshAdapter = ComponentToAdapterMap.FindChecked(SelectedComponent);
-		MeshPaintHelpers::ApplyVertexColorsToAllLODs(*MeshAdapter, SelectedComponent);
-	}
+		// Reset painting LOD index
+		PaintSettings->VertexPaintSettings.LODIndex = 0;
 
-	// Set actual flag in the settings struct
-	PaintSettings->VertexPaintSettings.bPaintOnSpecificLOD = bLODPaintingEnabled;
+	}
 
 	ApplyForcedLODIndex(bLODPaintingEnabled ? CachedLODIndex : -1);
 	TUniquePtr< FComponentReregisterContext > ComponentReregisterContext;
@@ -1688,6 +1711,12 @@ void FPaintModePainter::FinishPaintingTexture()
 		PaintingTexture2D = nullptr;
 		TexturePaintingCurrentMeshComponent = nullptr;
 	}
+
+	if (bCachedForceLOD)
+	{
+		// Make assumption here that when we paint while having a forced lod the mesh will contain per-lod vertex colors after this
+		bSelectionContainsPerLODColors = true;
+	}
 }
 
 FPaintTexture2DData* FPaintModePainter::GetPaintTargetData(const UTexture2D* InTexture)
@@ -2155,7 +2184,7 @@ void FPaintModePainter::CacheSelectionData()
 	// Update (cached) Paint LOD level if necessary
 	PaintSettings->VertexPaintSettings.LODIndex = FMath::Min<int32>(PaintSettings->VertexPaintSettings.LODIndex, GetMaxLODIndexToPaint());
 	CachedLODIndex = PaintSettings->VertexPaintSettings.LODIndex;
-
+	bCachedForceLOD = PaintSettings->VertexPaintSettings.bPaintOnSpecificLOD;
 	// Determine LOD level to use for painting (can only paint on LODs in vertex mode)
 	const int32 PaintLODIndex = (PaintSettings->PaintMode == EPaintMode::Vertices && PaintSettings->VertexPaintSettings.bPaintOnSpecificLOD) ? PaintSettings->VertexPaintSettings.LODIndex : 0;
 	// Determine UV channel to use while painting textures
diff --git a/Engine/Source/Editor/MeshPaintMode/Private/PaintModePainter.h b/Engine/Source/Editor/MeshPaintMode/Private/PaintModePainter.h
--- a/Engine/Source/Editor/MeshPaintMode/Private/PaintModePainter.h
+++ b/Engine/Source/Editor/MeshPaintMode/Private/PaintModePainter.h
@@ -116,6 +116,7 @@ protected:
 	bool CanRemoveInstanceColors() const;
 	bool CanPasteInstanceVertexColors() const;
 	bool CanCopyInstanceVertexColors() const;
+	bool CanPropagateVertexColorsToLODs() const;
 
 	/** Functions to determine whether or not an vertex color action is valid */
 	bool CanSaveMeshPackages() const;
@@ -152,6 +153,8 @@ protected:
 	void PasteVertexColors();
 	void FixVertexColors();
 	void RemoveVertexColors();
+	void PropagateVertexColorsToLODs();
+	void CycleMeshLODs(int32 Direction);
 	
 	/** Checks whether or not the current selection contains components which reference the same (static/skeletal)-mesh */
 	bool ContainsDuplicateMeshes(TArray<UMeshComponent*>& Components) const;
diff --git a/Engine/Source/Editor/MeshPaintMode/Private/SPaintModeWidget.cpp b/Engine/Source/Editor/MeshPaintMode/Private/SPaintModeWidget.cpp
--- a/Engine/Source/Editor/MeshPaintMode/Private/SPaintModeWidget.cpp
+++ b/Engine/Source/Editor/MeshPaintMode/Private/SPaintModeWidget.cpp
@@ -102,6 +102,9 @@ TSharedPtr<SWidget> SPaintModeWidget::CreateVertexPaintWidget()
 	TSharedPtr<SWidget> VertexColorWidget;
 	TSharedPtr<SHorizontalBox> VertexColorActionBox;
 	TSharedPtr<SHorizontalBox> InstanceColorActionBox;
+
+	static const FText SkelMeshNotificationText = LOCTEXT("SkelMeshAssetPaintInfo", "Paint is directly propagated to Skeletal Mesh Asset(s)");
+	static const FText StaticMeshNotificationText = LOCTEXT("StaticMeshAssetPaintInfo", "Paint is directly applied to all LODs");	
 		
 	SAssignNew(VertexColorWidget, SVerticalBox)
 	.Visibility(this, &SPaintModeWidget::IsVertexPaintModeVisible)
@@ -134,7 +137,8 @@ TSharedPtr<SWidget> SPaintModeWidget::CreateVertexPaintWidget()
 			SNew(SHorizontalBox)
 			.Visibility_Lambda([this]() -> EVisibility 
 			{
-				const bool bVisible = MeshPainter && MeshPainter->GetSelectedComponents<USkeletalMeshComponent>().Num();
+				bool bVisible = MeshPainter && MeshPainter->GetSelectedComponents<USkeletalMeshComponent>().Num();
+				bVisible |= ((PaintModeSettings->PaintMode == EPaintMode::Vertices) && !PaintModeSettings->VertexPaintSettings.bPaintOnSpecificLOD);
 				return bVisible ? EVisibility::Visible : EVisibility::Collapsed;
 			})
 		
@@ -154,7 +158,12 @@ TSharedPtr<SWidget> SPaintModeWidget::CreateVertexPaintWidget()
 			[
 				SNew(STextBlock)
 				.AutoWrapText(true)
-				.Text(LOCTEXT("SkelMeshAssetPaintInfo", "Paint is directly propagated to Skeletal Mesh Asset(s)"))
+				.Text_Lambda([this]() -> FText
+				{
+					const bool bSkelMeshText = MeshPainter && MeshPainter->GetSelectedComponents<USkeletalMeshComponent>().Num();
+					const bool bLODPaintText = (PaintModeSettings->PaintMode == EPaintMode::Vertices) && !PaintModeSettings->VertexPaintSettings.bPaintOnSpecificLOD;
+					return FText::Format(FTextFormat::FromString(TEXT("{0}{1}{2}")), bSkelMeshText ? SkelMeshNotificationText : FText::GetEmpty(), bSkelMeshText && bLODPaintText ? FText::FromString(TEXT("\n")) : FText::GetEmpty(), bLODPaintText ? StaticMeshNotificationText : FText::GetEmpty());
+				})
 			]
 		]
 	];
@@ -178,6 +187,7 @@ TSharedPtr<SWidget> SPaintModeWidget::CreateVertexPaintWidget()
 	InstanceToolbarBuilder.AddToolBarButton(FPaintModeCommands::Get().Paste, NAME_None, FText::GetEmpty(), TAttribute<FText>(), FSlateIcon(FEditorStyle::GetStyleSetName(), "MeshPaint.Paste"));
 	InstanceToolbarBuilder.AddToolBarButton(FPaintModeCommands::Get().Remove, NAME_None, FText::GetEmpty(), TAttribute<FText>(), FSlateIcon(FEditorStyle::GetStyleSetName(), "MeshPaint.Remove"));
 	InstanceToolbarBuilder.AddToolBarButton(FPaintModeCommands::Get().Fix, NAME_None, FText::GetEmpty(), TAttribute<FText>(), FSlateIcon(FEditorStyle::GetStyleSetName(), "MeshPaint.Fix"));
+	InstanceToolbarBuilder.AddToolBarButton(FPaintModeCommands::Get().PropagateVertexColorsToLODs, NAME_None, FText::GetEmpty(), TAttribute<FText>(), FSlateIcon(FEditorStyle::GetStyleSetName(), "MeshPaint.Propagate"));
 
 	InstanceColorActionBox->AddSlot()
 	.FillWidth(1.0f)
diff --git a/Engine/Source/Editor/PropertyEditor/Private/DetailCategoryBuilderImpl.cpp b/Engine/Source/Editor/PropertyEditor/Private/DetailCategoryBuilderImpl.cpp
--- a/Engine/Source/Editor/PropertyEditor/Private/DetailCategoryBuilderImpl.cpp
+++ b/Engine/Source/Editor/PropertyEditor/Private/DetailCategoryBuilderImpl.cpp
@@ -193,7 +193,14 @@ void FDetailCategoryImpl::SetCategoryVisibility(bool bIsVisible)
 	{
 		bIsCategoryVisible = bIsVisible;
 
-		GetDetailsView()->RerunCurrentFilter();
+		if (GetDetailsView())
+		{
+			GetDetailsView()->RerunCurrentFilter();
+		}
+		if (DetailLayoutBuilder.IsValid())
+		{
+			DetailLayoutBuilder.Pin()->NotifyNodeVisibilityChanged();
+		}
 	}
 }
 
@@ -588,7 +595,9 @@ bool FDetailCategoryImpl::GenerateStandaloneWidget(FDetailWidgetRow& OutRow) con
 
 	if(HeaderContentWidget.IsValid())
 	{
-		OutRow.ValueContent()
+		OutRow
+			.ValueContent()
+			.HAlign(HAlign_Fill)
 			[
 				HeaderContentWidget.ToSharedRef()
 			];
@@ -597,6 +606,11 @@ bool FDetailCategoryImpl::GenerateStandaloneWidget(FDetailWidgetRow& OutRow) con
 	return true;
 }
 
+void FDetailCategoryImpl::GetFilterStrings(TArray<FString>& OutFilterStrings) const
+{
+	OutFilterStrings.Add(GetDisplayName().ToString());
+}
+
 void FDetailCategoryImpl::OnItemExpansionChanged(bool bIsExpanded, bool bShouldSaveState)
 {
 	if (bRestoreExpansionState && bShouldSaveState)
@@ -853,6 +867,21 @@ void FDetailCategoryImpl::FilterNode(const FDetailFilter& InFilter)
 	}
 }
 
+FCustomPropertyTypeLayoutMap FDetailCategoryImpl::GetCustomPropertyTypeLayoutMap() const
+{
+	if (GetDetailsView())
+	{
+		return GetDetailsView()->GetCustomPropertyTypeLayoutMap();
+	}
+
+	if (DetailLayoutBuilder.IsValid())
+	{
+		return DetailLayoutBuilder.Pin()->GetInstancedPropertyTypeLayoutMap();
+	}
+
+	return FCustomPropertyTypeLayoutMap();
+}
+
 void FDetailCategoryImpl::GenerateLayout()
 {
 	// Reset all children
diff --git a/Engine/Source/Editor/PropertyEditor/Private/DetailCategoryBuilderImpl.h b/Engine/Source/Editor/PropertyEditor/Private/DetailCategoryBuilderImpl.h
--- a/Engine/Source/Editor/PropertyEditor/Private/DetailCategoryBuilderImpl.h
+++ b/Engine/Source/Editor/PropertyEditor/Private/DetailCategoryBuilderImpl.h
@@ -180,6 +180,7 @@ public:
 	/** IDetailTreeNode interface */
 	virtual EDetailNodeType GetNodeType() const override { return EDetailNodeType::Category; }
 	virtual TSharedPtr<IPropertyHandle> CreatePropertyHandle() const override { return nullptr; }
+	virtual void GetFilterStrings(TArray<FString>& OutFilterStrings) const override;
 
 	virtual void GetChildren(FDetailNodeList& OutChildren) override;
 	virtual bool ShouldBeExpanded() const override;
@@ -189,6 +190,8 @@ public:
 	virtual bool ShouldShowOnlyChildren() const override { return bShowOnlyChildren; }
 	virtual FName GetNodeName() const override { return GetCategoryName(); }
 
+	FCustomPropertyTypeLayoutMap GetCustomPropertyTypeLayoutMap() const;
+
 	/**
 	 * @return true if the parent layout is valid or has been destroyed by a refresh.
 	 */
diff --git a/Engine/Source/Editor/PropertyEditor/Private/DetailCategoryGroupNode.h b/Engine/Source/Editor/PropertyEditor/Private/DetailCategoryGroupNode.h
--- a/Engine/Source/Editor/PropertyEditor/Private/DetailCategoryGroupNode.h
+++ b/Engine/Source/Editor/PropertyEditor/Private/DetailCategoryGroupNode.h
@@ -65,12 +65,13 @@ private:
 
 	virtual EDetailNodeType GetNodeType() const override { return EDetailNodeType::Category; }
 	virtual TSharedPtr<IPropertyHandle> CreatePropertyHandle() const override { return nullptr; }
+	virtual void GetFilterStrings(TArray<FString>& OutFilterStrings) const override { OutFilterStrings.Add(GroupName.ToString()); };
 
 	virtual void GetChildren(FDetailNodeList& OutChildren )  override;
 	virtual void FilterNode( const FDetailFilter& InFilter ) override;
 	virtual void Tick( float DeltaTime ) override {}
 	virtual bool ShouldShowOnlyChildren() const override { return false; }
-	virtual FName GetNodeName() const override { return NAME_None; }
+	virtual FName GetNodeName() const override { return GroupName; }
 private:
 	FDetailNodeList ChildNodes;
 	FDetailCategoryImpl& ParentCategory;
diff --git a/Engine/Source/Editor/PropertyEditor/Private/DetailItemNode.cpp b/Engine/Source/Editor/PropertyEditor/Private/DetailItemNode.cpp
--- a/Engine/Source/Editor/PropertyEditor/Private/DetailItemNode.cpp
+++ b/Engine/Source/Editor/PropertyEditor/Private/DetailItemNode.cpp
@@ -106,6 +106,27 @@ TSharedPtr<IPropertyHandle> FDetailItemNode::CreatePropertyHandle() const
 	}
 }
 
+void FDetailItemNode::GetFilterStrings(TArray<FString>& OutFilterStrings) const
+{
+	if (Customization.HasCustomWidget())
+	{
+		OutFilterStrings.Add(Customization.GetWidgetRow().FilterTextString.ToString());
+	}
+
+	if (Customization.HasPropertyNode())
+	{
+		TSharedPtr<FPropertyNode> PropertyNode = Customization.GetPropertyNode();
+		if (PropertyNode.IsValid())
+		{
+			OutFilterStrings.Add(GetPropertyNode()->GetDisplayName().ToString());
+			if (PropertyNode->GetDisplayName().ToString() != PropertyNode->GetProperty()->GetName())
+			{
+				OutFilterStrings.Add(PropertyNode->GetProperty()->GetName());
+			}
+		}
+	}
+}
+
 void FDetailItemNode::InitPropertyEditor()
 {
 	UProperty* NodeProperty = Customization.GetPropertyNode()->GetProperty();
@@ -317,13 +338,25 @@ void FDetailItemNode::GetChildren(FDetailNodeList& OutChildren)
 
 void FDetailItemNode::GenerateChildren( bool bUpdateFilteredNodes )
 {
+	FDetailNodeList OldChildren = Children;
 	Children.Empty();
 
-	if (!ParentCategory.IsValid())
+	TSharedPtr<FDetailCategoryImpl> ParentCategoryPinned = ParentCategory.Pin();
+	if (ParentCategoryPinned.IsValid() == false)
 	{
 		return;
 	}
 
+	// Make sure to remove the root properties referenced by the old children, otherwise they will leak.
+	for (TSharedRef<FDetailTreeNode> OldChild : OldChildren)
+	{
+		TSharedPtr<FComplexPropertyNode> OldChildExternalRootPropertyNode = OldChild->GetExternalRootPropertyNode();
+		if (OldChildExternalRootPropertyNode.IsValid())
+		{
+			ParentCategoryPinned->GetParentLayoutImpl().RemoveExternalRootPropertyNode(OldChildExternalRootPropertyNode.ToSharedRef());
+		}
+	}
+
 	if( Customization.HasPropertyNode() )
 	{
 		Customization.PropertyRow->OnGenerateChildren( Children );
@@ -464,9 +497,14 @@ static bool PassesAllFilters( FDetailItemNode* ItemNode, const FDetailLayoutCust
 		const bool bPassesCategoryFilter = !bSearchFilterIsEmpty && InFilter.bShowAllChildrenIfCategoryMatches ? Local::StringPassesFilter(InFilter, InCategoryName) : false;
 
 		bPassesAllFilters = false;
-		if( PropertyNodePin.IsValid() && !PropertyNodePin->AsCategoryNode() )
+		if( PropertyNodePin.IsValid() && !PropertyNodePin->AsCategoryNode())
 		{
-			UClass* ObjectClass = PropertyNodePin->FindObjectItemParent()->GetObjectBaseClass();
+			UClass* ObjectClass = nullptr;
+			FObjectPropertyNode* PropertyNodeParent = PropertyNodePin->FindObjectItemParent();
+			if (PropertyNodeParent)
+			{
+				ObjectClass = PropertyNodeParent->GetObjectBaseClass();
+			}
 
 			const bool bIsNotBeingFiltered = PropertyNodePin->HasNodeFlags(EPropertyNodeFlags::IsBeingFiltered) == 0;
 			const bool bIsSeenDueToFiltering = PropertyNodePin->HasNodeFlags(EPropertyNodeFlags::IsSeenDueToFiltering) != 0;
@@ -475,7 +513,7 @@ static bool PassesAllFilters( FDetailItemNode* ItemNode, const FDetailLayoutCust
 			const bool bPassesSearchFilter = bSearchFilterIsEmpty || ( bIsNotBeingFiltered || bIsSeenDueToFiltering || bIsParentSeenDueToFiltering );
 			const bool bPassesModifiedFilter = bPassesSearchFilter && ( InFilter.bShowOnlyModifiedProperties == false || PropertyNodePin->GetDiffersFromDefault() == true );
 			const bool bPassesDifferingFilter = InFilter.bShowOnlyDiffering ? InFilter.WhitelistedProperties.Find(*FPropertyNode::CreatePropertyPath(PropertyNodePin.ToSharedRef())) != nullptr : true;
-			const bool bPassesKeyableFilter = (InFilter.bShowKeyable == false || Local::ItemIsKeyable(ItemNode, ObjectClass, PropertyNodePin));
+			const bool bPassesKeyableFilter = (InFilter.bShowKeyable == false || (ObjectClass ? Local::ItemIsKeyable(ItemNode, ObjectClass, PropertyNodePin) : false));
 			const bool bPassesAnimatedFilter = (InFilter.bShowAnimated == false || Local::ItemIsAnimated(ItemNode, PropertyNodePin));
 
 			// The property node is visible (note categories are never visible unless they have a child that is visible )
@@ -552,6 +590,15 @@ FName FDetailItemNode::GetNodeName() const
 
 		return CustomName;
 	}
+	else if (Customization.HasPropertyNode())
+	{
+		TSharedPtr<FPropertyNode> PropertyNode = Customization.GetPropertyNode();
+		if (PropertyNode.IsValid())
+		{
+			FName PropertyName = PropertyNode->GetProperty()->GetFName();
+			return PropertyName;
+		}
+	}
 	return NAME_None;
 }
 
@@ -580,6 +627,15 @@ TSharedPtr<IDetailPropertyRow> FDetailItemNode::GetRow() const
 	return nullptr;
 }
 
+TSharedPtr<FComplexPropertyNode> FDetailItemNode::GetExternalRootPropertyNode() const
+{
+	if (Customization.IsValidCustomization() && Customization.PropertyRow.IsValid())
+	{
+		return Customization.PropertyRow->GetExternalRootNode();
+	}
+	return nullptr;
+}
+
 void FDetailItemNode::FilterNode(const FDetailFilter& InFilter)
 {
 	bShouldBeVisibleDueToFiltering = PassesAllFilters( this, Customization, InFilter, ParentCategory.Pin()->GetDisplayName().ToString() );
diff --git a/Engine/Source/Editor/PropertyEditor/Private/DetailItemNode.h b/Engine/Source/Editor/PropertyEditor/Private/DetailItemNode.h
--- a/Engine/Source/Editor/PropertyEditor/Private/DetailItemNode.h
+++ b/Engine/Source/Editor/PropertyEditor/Private/DetailItemNode.h
@@ -23,6 +23,7 @@ public:
 	/** IDetailTreeNode interface */
 	virtual EDetailNodeType GetNodeType() const override;
 	virtual TSharedPtr<IPropertyHandle> CreatePropertyHandle() const override;
+	virtual void GetFilterStrings(TArray<FString>& OutFilterStrings) const override;
 
 	/**
 	 * Initializes this node                                                              
@@ -51,7 +52,7 @@ public:
 	 */
 	bool HasGeneratedChildren() const { return Children.Num() > 0;}
 
-	/** IDetailTreeNode interface */
+	/** FDetailTreeNode interface */
 	virtual IDetailsViewPrivate* GetDetailsView() const override{ return ParentCategory.Pin()->GetDetailsView(); }
 	virtual TSharedRef< ITableRow > GenerateWidgetForTableView( const TSharedRef<STableViewBase>& OwnerTable, const FDetailColumnSizeData& ColumnSizeData, bool bAllowFavoriteSystem) override;
 	virtual bool GenerateStandaloneWidget(FDetailWidgetRow& OutRow) const override;
@@ -70,6 +71,7 @@ public:
 	virtual bool IsLeaf() override { return true; }
 	virtual TAttribute<bool> IsPropertyEditingEnabled() const override { return IsParentEnabled; }
 	virtual TSharedPtr<FPropertyNode> GetPropertyNode() const override;
+	virtual TSharedPtr<FComplexPropertyNode> GetExternalRootPropertyNode() const override;
 	virtual TSharedPtr<IDetailPropertyRow> GetRow() const override;
 
 private:
