Added subsurface profile for eye shading model.

diff --git a/Engine/Shaders/Private/BasePassPixelShader.usf b/Engine/Shaders/Private/BasePassPixelShader.usf
--- a/Engine/Shaders/Private/BasePassPixelShader.usf
+++ b/Engine/Shaders/Private/BasePassPixelShader.usf
@@ -69,6 +69,7 @@ float NormalCurvatureToRoughness(float3 WorldNormal)
 
 #define FORCE_FULLY_ROUGH (SIMPLE_FORWARD_SHADING || MATERIAL_FULLY_ROUGH)
 #define EDITOR_ALPHA2COVERAGE (USE_EDITOR_COMPOSITING && SUPPORTS_PIXEL_COVERAGE)
+#define POST_PROCESS_SUBSURFACE ((MATERIAL_SHADINGMODEL_SUBSURFACE_PROFILE || MATERIAL_SHADINGMODEL_EYE) && USES_GBUFFER)
 
 #if TRANSLUCENCY_LIGHTING_SURFACE_LIGHTINGVOLUME || TRANSLUCENCY_LIGHTING_SURFACE_FORWARDSHADING || FORWARD_SHADING
 #include "ForwardLightingCommon.ush" 
@@ -769,7 +770,7 @@ void FPixelShaderInOut_MainPS(
 
 	// If we don't use this shading model the color should be black (don't generate shader code for unused data, don't do indirectlighting cache lighting with this color).
 	float3 SubsurfaceColor = 0;
-#if MATERIAL_SHADINGMODEL_SUBSURFACE || MATERIAL_SHADINGMODEL_PREINTEGRATED_SKIN || MATERIAL_SHADINGMODEL_SUBSURFACE_PROFILE || MATERIAL_SHADINGMODEL_TWOSIDED_FOLIAGE || MATERIAL_SHADINGMODEL_CLOTH
+#if MATERIAL_SHADINGMODEL_SUBSURFACE || MATERIAL_SHADINGMODEL_PREINTEGRATED_SKIN || MATERIAL_SHADINGMODEL_SUBSURFACE_PROFILE || MATERIAL_SHADINGMODEL_TWOSIDED_FOLIAGE || MATERIAL_SHADINGMODEL_CLOTH || MATERIAL_SHADINGMODEL_EYE
 	{
 		float4 SubsurfaceData = GetMaterialSubsurfaceData(PixelMaterialInputs);
 
@@ -875,11 +876,11 @@ void FPixelShaderInOut_MainPS(
 #endif
 #endif
 
-#if MATERIAL_SHADINGMODEL_SUBSURFACE_PROFILE && USES_GBUFFER
+#if POST_PROCESS_SUBSURFACE
 	// SubsurfaceProfile applies the BaseColor in a later pass. Any lighting output in the base pass needs
 	// to separate specular and diffuse lighting in a checkerboard pattern
 	bool bChecker = CheckerFromPixelPos(MaterialParameters.SvPosition.xy);
-	if ( GBuffer.ShadingModelID == SHADINGMODELID_SUBSURFACE_PROFILE ) 
+	if (UseSubsurfaceProfile(GBuffer.ShadingModelID))
 	{
 		AdjustBaseColorAndSpecularColorForSubsurfaceProfileLighting(BaseColor, GBuffer.SpecularColor, Specular, bChecker);
 	}
@@ -1052,7 +1053,7 @@ void FPixelShaderInOut_MainPS(
 	#endif
 #endif
 
-#if !(MATERIAL_SHADINGMODEL_SUBSURFACE_PROFILE && USES_GBUFFER) 
+#if !POST_PROCESS_SUBSURFACE
  	// For non-skin, we just add diffuse and non-diffuse color together, otherwise we need to keep them separate
 	Color += DiffuseColor;
 #endif
@@ -1085,18 +1086,18 @@ void FPixelShaderInOut_MainPS(
 			// Apply vertex fog
 			Color = Color * Fogging.a + Fogging.rgb;
 
-#if MATERIAL_SHADINGMODEL_SUBSURFACE_PROFILE && USES_GBUFFER
+#if POST_PROCESS_SUBSURFACE
 			// Apply vertex fog to diffuse color
 			DiffuseColor = DiffuseColor * Fogging.a + Fogging.rgb;
 
-			if (GBuffer.ShadingModelID == SHADINGMODELID_SUBSURFACE_PROFILE && 
+			if (UseSubsurfaceProfile(GBuffer.ShadingModelID) && 
                 View.bSubsurfacePostprocessEnabled > 0 && View.bCheckerboardSubsurfaceProfileRendering > 0 )
 			{
 				// Adjust for checkerboard. only apply non-diffuse lighting (including emissive) 
 				// to the specular component, otherwise lighting is applied twice
 				Color *= !bChecker;
 			}
-			LightAccumulator_Add(LightAccumulator, Color + DiffuseColor, DiffuseColor, 1.0f, GBuffer.ShadingModelID == SHADINGMODELID_SUBSURFACE_PROFILE);
+			LightAccumulator_Add(LightAccumulator, Color + DiffuseColor, DiffuseColor, 1.0f, UseSubsurfaceProfile(GBuffer.ShadingModelID));
 #else
 			LightAccumulator_Add(LightAccumulator, Color, 0, 1.0f, false);
 #endif
diff --git a/Engine/Shaders/Private/DeferredLightingCommon.ush b/Engine/Shaders/Private/DeferredLightingCommon.ush
--- a/Engine/Shaders/Private/DeferredLightingCommon.ush
+++ b/Engine/Shaders/Private/DeferredLightingCommon.ush
@@ -296,7 +296,7 @@ float4 GetDynamicLighting(float3 WorldPosition, float3 CameraVector, FGBufferDat
 		BRANCH
 		if( Shadow.SurfaceShadow + Shadow.TransmissionShadow > 0 )
 		{
-			const bool bNeedsSeparateSubsurfaceLightAccumulation = GBuffer.ShadingModelID == SHADINGMODELID_SUBSURFACE_PROFILE;
+			const bool bNeedsSeparateSubsurfaceLightAccumulation = UseSubsurfaceProfile(GBuffer.ShadingModelID);
 			float3 LightColor = LightData.LightColorAndFalloffExponent.rgb;
 
 		#if NON_DIRECTIONAL_DIRECT_LIGHTING
diff --git a/Engine/Shaders/Private/DeferredShadingCommon.ush b/Engine/Shaders/Private/DeferredShadingCommon.ush
--- a/Engine/Shaders/Private/DeferredShadingCommon.ush
+++ b/Engine/Shaders/Private/DeferredShadingCommon.ush
@@ -312,6 +312,11 @@ bool IsSubsurfaceModel(int ShadingModel)
 		|| ShadingModel == SHADINGMODELID_EYE;
 }
 
+bool UseSubsurfaceProfile(int ShadingModel)
+{
+	return ShadingModel == SHADINGMODELID_SUBSURFACE_PROFILE || ShadingModel == SHADINGMODELID_EYE;
+}
+
 // all values that are output by the forward rendering pass
 struct FGBufferData
 {
@@ -321,7 +326,7 @@ struct FGBufferData
 	float3 DiffuseColor;
 	// 0..1 (derived from BaseColor, Metalness, Specular)
 	float3 SpecularColor;
-	// 0..1, white for SHADINGMODELID_SUBSURFACE_PROFILE (apply BaseColor after scattering is more correct and less blurry)
+	// 0..1, white for SHADINGMODELID_SUBSURFACE_PROFILE and SHADINGMODELID_EYE (apply BaseColor after scattering is more correct and less blurry)
 	float3 BaseColor;
 	// 0..1
 	float Metallic;
@@ -354,10 +359,12 @@ struct FGBufferData
 	// Velocity for motion blur (only used when WRITES_VELOCITY_TO_GBUFFER is enabled)
 	float4 Velocity;
 
-	// 0..1, only needed by SHADINGMODELID_SUBSURFACE_PROFILE which applies BaseColor later
+	// 0..1, only needed by SHADINGMODELID_SUBSURFACE_PROFILE and SHADINGMODELID_EYE which apply BaseColor later
 	float3 StoredBaseColor;
-	// 0..1, only needed by SHADINGMODELID_SUBSURFACE_PROFILE which applies Specular later
+	// 0..1, only needed by SHADINGMODELID_SUBSURFACE_PROFILE and SHADINGMODELID_EYE which apply Specular later
 	float StoredSpecular;
+	// 0..1, only needed by SHADINGMODELID_EYE which encodes Iris Distance inside Metallic
+	float StoredMetallic;
 };
 
 bool HasDistanceFieldRepresentation(FGBufferData GBufferData)
@@ -551,12 +558,16 @@ FGBufferData DecodeGBufferData(
 
 	GBuffer.StoredBaseColor = GBuffer.BaseColor;
 	GBuffer.StoredSpecular = GBuffer.Specular;
+	GBuffer.StoredMetallic = GBuffer.Metallic;
+
+	// SHADINGMODELID_EYE stores Iris Distance inside Metallic.
+	GBuffer.Metallic = GBuffer.ShadingModelID == SHADINGMODELID_EYE ? 0.0f : GBuffer.Metallic;
 
 	// derived from BaseColor, Metalness, Specular
 	{
 		GBuffer.SpecularColor = lerp( 0.08 * GBuffer.Specular.xxx, GBuffer.BaseColor, GBuffer.Metallic );
 
-		if (GBuffer.ShadingModelID == SHADINGMODELID_SUBSURFACE_PROFILE)
+		if (UseSubsurfaceProfile(GBuffer.ShadingModelID))
 		{
 			AdjustBaseColorAndSpecularColorForSubsurfaceProfileLighting(GBuffer.BaseColor, GBuffer.SpecularColor, GBuffer.Specular, bChecker);
 		}
diff --git a/Engine/Shaders/Private/PostProcessAmbient.usf b/Engine/Shaders/Private/PostProcessAmbient.usf
--- a/Engine/Shaders/Private/PostProcessAmbient.usf
+++ b/Engine/Shaders/Private/PostProcessAmbient.usf
@@ -458,7 +458,7 @@ void MainPS(in noperspective float4 UVAndScreenPos : TEXCOORD0, float4 SvPositio
 
 	FLightAccumulator LightAccumulator = (FLightAccumulator)0;
 
-	const bool bNeedsSeparateSubsurfaceLightAccumulation = GBuffer.ShadingModelID == SHADINGMODELID_SUBSURFACE_PROFILE;
+	const bool bNeedsSeparateSubsurfaceLightAccumulation = UseSubsurfaceProfile(GBuffer.ShadingModelID);
 
 	LightAccumulator_Add(LightAccumulator, NonSpecularContribution + SpecularContribution, NonSpecularContribution, AmbientCubemapColor.rgb, bNeedsSeparateSubsurfaceLightAccumulation);
 
diff --git a/Engine/Shaders/Private/PostProcessSubsurface.usf b/Engine/Shaders/Private/PostProcessSubsurface.usf
--- a/Engine/Shaders/Private/PostProcessSubsurface.usf
+++ b/Engine/Shaders/Private/PostProcessSubsurface.usf
@@ -36,7 +36,7 @@ float GetSubsurfaceStrength(float2 UV)
 {
 	FScreenSpaceData ScreenSpaceData = GetScreenSpaceData(UV);
 
-	float Mask = ScreenSpaceData.GBuffer.ShadingModelID == SHADINGMODELID_SUBSURFACE_PROFILE;
+	float Mask = UseSubsurfaceProfile(ScreenSpaceData.GBuffer.ShadingModelID) ? 1.0f : 0.0f;
 	return Mask * ScreenSpaceData.GBuffer.CustomData.a;
 }
 
@@ -195,7 +195,7 @@ void VisualizePS(in noperspective float4 UVAndScreenPos : TEXCOORD0, float4 SvPo
 
 		OutColor = float4(0.5f, 0.5f, 0.5f, 0);
 
-		BRANCH if (ScreenSpaceData.GBuffer.ShadingModelID == SHADINGMODELID_SUBSURFACE_PROFILE)
+		BRANCH if (UseSubsurfaceProfile(ScreenSpaceData.GBuffer.ShadingModelID))
 		{
 			if (SubsurfaceProfileInt == 0)
 			{
@@ -355,7 +355,7 @@ float4 SetupSubsurfaceForOnePixel(float2 UVSceneColor)
 
 	FScreenSpaceData ScreenSpaceData = GetScreenSpaceData(UVSceneColor);
 
-	BRANCH if (ScreenSpaceData.GBuffer.ShadingModelID == SHADINGMODELID_SUBSURFACE_PROFILE)
+	BRANCH if (UseSubsurfaceProfile(ScreenSpaceData.GBuffer.ShadingModelID))
 	{
 		// '1' is lower quality but that is acceptable here
 		SDiffuseAndSpecular DiffuseAndSpecular = ReconstructLighting(UVSceneColor, 1);
@@ -462,7 +462,7 @@ void SubsurfaceRecombinePS(noperspective float4 UVAndScreenPos : TEXCOORD0, floa
 
 	FScreenSpaceData ScreenSpaceData = GetScreenSpaceData(BufferUV);
 	
-	if (ScreenSpaceData.GBuffer.ShadingModelID != SHADINGMODELID_SUBSURFACE_PROFILE)
+	if (!UseSubsurfaceProfile(ScreenSpaceData.GBuffer.ShadingModelID))
 	{
 		OutColor = Texture2DSample(PostprocessInput0, PostprocessInput0Sampler, BufferUV);
 		return;
diff --git a/Engine/Shaders/Private/SeparableSSS.ush b/Engine/Shaders/Private/SeparableSSS.ush
--- a/Engine/Shaders/Private/SeparableSSS.ush
+++ b/Engine/Shaders/Private/SeparableSSS.ush
@@ -114,7 +114,7 @@
 #define SSSSSampleLevelZero(tex, coord) tex.SampleLevel(LinearSampler, coord, 0)
 #define SSSSSampleLevelZeroPoint(tex, coord) tex.SampleLevel(PointSampler, coord, 0)
 #define SSSSSampleShadowmap(coord)	float4(0,0,0,0)	// todo: not used yet
-#define SSSSSampleMasking(texcoord)	(GetScreenSpaceData(texcoord).GBuffer.ShadingModelID == SHADINGMODELID_SUBSURFACE_PROFILE)	// todo: masking doesn't work with filtering ideally put in alpha to get it filtered
+#define SSSSSampleMasking(texcoord)	(HasSubsurfaceProfile(GetScreenSpaceData(texcoord).GBuffer.ShadingModelID))	// todo: masking doesn't work with filtering ideally put in alpha to get it filtered
 #define SSSSSampleSceneColor(coord) SampleSetup(coord)
 #define SSSSSampleSceneColorPoint(coord) SampleSetup(coord) // todo: could be made point
 #define SSSSSampleProfileId(coord) SampleSetupProfileId(coord)
diff --git a/Engine/Shaders/Private/ShadingModels.ush b/Engine/Shaders/Private/ShadingModels.ush
--- a/Engine/Shaders/Private/ShadingModels.ush
+++ b/Engine/Shaders/Private/ShadingModels.ush
@@ -779,10 +779,10 @@ FDirectLighting TwoSidedBxDF( FGBufferData GBuffer, half3 N, half3 V, half3 L, f
 
 FDirectLighting EyeBxDF( FGBufferData GBuffer, half3 N, half3 V, half3 L, float Falloff, float NoL, FAreaLight AreaLight, FShadowTerms Shadow )
 {
-	const float3 IrisNormal		= OctahedronToUnitVector( GBuffer.CustomData.xy * 2 - 1 );
-	const float  IrisDistance	= GBuffer.CustomData.w;
-	const float  IrisMask		= GBuffer.CustomData.z;
-	
+	const float3 IrisNormal		= OctahedronToUnitVector( GBuffer.CustomData.yz * 2 - 1 );
+	const float  IrisDistance	= GBuffer.StoredMetallic;
+	const float  IrisMask		= 1.0f - GBuffer.CustomData.w;
+
 	BxDFContext Context;
 	Init( Context, N, V, L );
 	SphereMaxNoH( Context, AreaLight.SphereSinAlpha, false );
diff --git a/Engine/Shaders/Private/ShadingModelsMaterial.ush b/Engine/Shaders/Private/ShadingModelsMaterial.ush
--- a/Engine/Shaders/Private/ShadingModelsMaterial.ush
+++ b/Engine/Shaders/Private/ShadingModelsMaterial.ush
@@ -118,12 +118,14 @@ void SetGBufferForShadingModel(
 	GBuffer.IndirectIrradiance *= 1 - GBuffer.CustomData.a;
 #elif MATERIAL_SHADINGMODEL_EYE
 	GBuffer.ShadingModelID = SHADINGMODELID_EYE;
+	GBuffer.CustomData.x = EncodeSubsurfaceProfile(SubsurfaceProfile).x;
+	GBuffer.CustomData.w = 1.0f - saturate(GetMaterialCustomData0(MaterialParameters));	// Opacity = 1.0 - Iris Mask
+	GBuffer.Metallic = saturate(GetMaterialCustomData1(MaterialParameters));			// Iris Distance
+
 	#if NUM_MATERIAL_OUTPUTS_GETTANGENTOUTPUT > 0
 		float3 Tangent = GetTangentOutput0(MaterialParameters);
-		GBuffer.CustomData.xy = UnitVectorToOctahedron( normalize(Tangent) ) * 0.5 + 0.5;
+		GBuffer.CustomData.yz = UnitVectorToOctahedron( normalize(Tangent) ) * 0.5 + 0.5;
 	#endif
-	GBuffer.CustomData.z = saturate( GetMaterialCustomData0(MaterialParameters) );	// Iris Mask
-	GBuffer.CustomData.w = saturate( GetMaterialCustomData1(MaterialParameters) );	// Iris Distance
 #else
 	// missing shading model, compiler should report ShadingModelID is not set
 #endif
diff --git a/Engine/Source/Editor/MaterialEditor/Private/MaterialEditorInstanceDetailCustomization.cpp b/Engine/Source/Editor/MaterialEditor/Private/MaterialEditorInstanceDetailCustomization.cpp
--- a/Engine/Source/Editor/MaterialEditor/Private/MaterialEditorInstanceDetailCustomization.cpp
+++ b/Engine/Source/Editor/MaterialEditor/Private/MaterialEditorInstanceDetailCustomization.cpp
@@ -703,7 +703,7 @@ EVisibility FMaterialInstanceParameterDetails::ShouldShowSubsurfaceProfile() con
 {
 	EMaterialShadingModel Model = MaterialEditorInstance->SourceInstance->GetShadingModel();
 
-	return (Model == MSM_SubsurfaceProfile) ? EVisibility::Visible : EVisibility::Collapsed;
+	return UseSubsurfaceProfile(Model) ? EVisibility::Visible : EVisibility::Collapsed;
 }
 
 
diff --git a/Engine/Source/Runtime/Engine/Private/Materials/Material.cpp b/Engine/Source/Runtime/Engine/Private/Materials/Material.cpp
--- a/Engine/Source/Runtime/Engine/Private/Materials/Material.cpp
+++ b/Engine/Source/Runtime/Engine/Private/Materials/Material.cpp
@@ -4080,7 +4080,7 @@ bool UMaterial::CanEditChange(const UProperty* InProperty) const
 
 		if (PropertyName == GET_MEMBER_NAME_STRING_CHECKED(UMaterial, SubsurfaceProfile))
 		{
-			return MaterialDomain == MD_Surface && ShadingModel == MSM_SubsurfaceProfile && (BlendMode == BLEND_Opaque || BlendMode == BLEND_Masked);
+			return MaterialDomain == MD_Surface && UseSubsurfaceProfile(ShadingModel) && (BlendMode == BLEND_Opaque || BlendMode == BLEND_Masked);
 		}
 
 		if (PropertyName == GET_MEMBER_NAME_STRING_CHECKED(FLightmassMaterialInterfaceSettings, bCastShadowAsMasked))
diff --git a/Engine/Source/Runtime/Engine/Private/Materials/MaterialInterface.cpp b/Engine/Source/Runtime/Engine/Private/Materials/MaterialInterface.cpp
--- a/Engine/Source/Runtime/Engine/Private/Materials/MaterialInterface.cpp
+++ b/Engine/Source/Runtime/Engine/Private/Materials/MaterialInterface.cpp
@@ -460,7 +460,7 @@ void UMaterialInterface::UpdateMaterialRenderProxy(FMaterialRenderProxy& Proxy)
 	EMaterialShadingModel MaterialShadingModel = GetShadingModel();
 
 	// for better performance we only update SubsurfaceProfileRT if the feature is used
-	if (MaterialShadingModel == MSM_SubsurfaceProfile)
+	if (UseSubsurfaceProfile(MaterialShadingModel))
 	{
 		FSubsurfaceProfileStruct Settings;
 
diff --git a/Engine/Source/Runtime/Engine/Public/MaterialShared.h b/Engine/Source/Runtime/Engine/Public/MaterialShared.h
--- a/Engine/Source/Runtime/Engine/Public/MaterialShared.h
+++ b/Engine/Source/Runtime/Engine/Public/MaterialShared.h
@@ -66,9 +66,19 @@ extern FName GetMaterialQualityLevelFName(EMaterialQualityLevel::Type InMaterial
 
 inline bool IsSubsurfaceShadingModel(EMaterialShadingModel ShadingModel)
 {
-	return ShadingModel == MSM_Subsurface || ShadingModel == MSM_PreintegratedSkin || 
-		ShadingModel == MSM_SubsurfaceProfile || ShadingModel == MSM_TwoSidedFoliage || 
-		ShadingModel == MSM_Cloth;
+	return ShadingModel == MSM_Subsurface || ShadingModel == MSM_PreintegratedSkin ||
+		ShadingModel == MSM_SubsurfaceProfile || ShadingModel == MSM_TwoSidedFoliage ||
+		ShadingModel == MSM_Cloth || ShadingModel == MSM_Eye;
+}
+
+inline bool UseSubsurfaceProfile(EMaterialShadingModel ShadingModel)
+{
+	return ShadingModel == MSM_SubsurfaceProfile || ShadingModel == MSM_Eye;
+}
+
+inline uint32 GetUseSubsurfaceProfileShadingModelMask()
+{
+	return (1 << MSM_SubsurfaceProfile) | (1 << MSM_Eye);
 }
 
 /**
diff --git a/Engine/Source/Runtime/Renderer/Private/CompositionLighting/CompositionLighting.cpp b/Engine/Source/Runtime/Renderer/Private/CompositionLighting/CompositionLighting.cpp
--- a/Engine/Source/Runtime/Renderer/Private/CompositionLighting/CompositionLighting.cpp
+++ b/Engine/Source/Runtime/Renderer/Private/CompositionLighting/CompositionLighting.cpp
@@ -446,7 +446,7 @@ void FCompositionLighting::ProcessAfterLighting(FRHICommandListImmediate& RHICmd
 		{
 			float Radius = CVarSSSScale.GetValueOnRenderThread();
 			bool bSimpleDynamicLighting = IsAnyForwardShadingEnabled(View.GetShaderPlatform());
-			bool bScreenSpaceSubsurfacePassNeeded = ((View.ShadingModelMaskInView & (1 << MSM_SubsurfaceProfile)) != 0) && IsSubsurfacePostprocessRequired();
+			bool bScreenSpaceSubsurfacePassNeeded = ((View.ShadingModelMaskInView & GetUseSubsurfaceProfileShadingModelMask()) != 0) && IsSubsurfacePostprocessRequired();
 			bool bSubsurfaceAllowed = CVarSubsurfaceScattering.GetValueOnRenderThread() == 1;
 
 			if (bScreenSpaceSubsurfacePassNeeded 
