Missing file for CL 4097655: "Added subsurface profile for eye shading model."

diff --git a/Engine/Shaders/Private/ReflectionEnvironmentPixelShader.usf b/Engine/Shaders/Private/ReflectionEnvironmentPixelShader.usf
--- a/Engine/Shaders/Private/ReflectionEnvironmentPixelShader.usf
+++ b/Engine/Shaders/Private/ReflectionEnvironmentPixelShader.usf
@@ -317,7 +317,7 @@ void ReflectionEnvironmentSkyLighting(
 		// Scene color alpha is used for ScreenSpaceSubsurfaceScattering (if that is not needed it can be disabled with SUBSURFACE_CHANNEL_MODE)
 		FLightAccumulator LightAccumulator = (FLightAccumulator)0;
 
-		const bool bNeedsSeparateSubsurfaceLightAccumulation = ShadingModelID == SHADINGMODELID_SUBSURFACE_PROFILE;
+		const bool bNeedsSeparateSubsurfaceLightAccumulation = UseSubsurfaceProfile(ShadingModelID);
 
 		LightAccumulator_Add(LightAccumulator, SkyLighting, SkyLighting, 1.0f, bNeedsSeparateSubsurfaceLightAccumulation);
 		OutColor = LightAccumulator_GetResult(LightAccumulator);
