Added dual specular for subsurface profile shading model. It is controlled by parameters inside subsurface profile.

diff --git a/Engine/Shaders/Private/BasePassCommon.ush b/Engine/Shaders/Private/BasePassCommon.ush
--- a/Engine/Shaders/Private/BasePassCommon.ush
+++ b/Engine/Shaders/Private/BasePassCommon.ush
@@ -8,10 +8,12 @@
 	#define ForwardLightData TranslucentBasePass.Shared.Forward
 	#define ReflectionStruct TranslucentBasePass.Shared.Reflection
 	#define FogStruct TranslucentBasePass.Shared.Fog
+	#define BasePassSSProfilesTexture TranslucentBasePass.SSProfilesTexture
 #else
 	#define ForwardLightData OpaqueBasePass.Shared.Forward
 	#define ReflectionStruct OpaqueBasePass.Shared.Reflection
 	#define FogStruct OpaqueBasePass.Shared.Fog
+	#define BasePassSSProfilesTexture OpaqueBasePass.Shared.SSProfilesTexture
 #endif 
 
 #undef NEEDS_LIGHTMAP_COORDINATE
diff --git a/Engine/Shaders/Private/BasePassPixelShader.usf b/Engine/Shaders/Private/BasePassPixelShader.usf
--- a/Engine/Shaders/Private/BasePassPixelShader.usf
+++ b/Engine/Shaders/Private/BasePassPixelShader.usf
@@ -836,7 +836,8 @@ void FPixelShaderInOut_MainPS(
 		Metallic, 
 		Specular, 
 		Roughness, 
-		float4(SubsurfaceColor, SubsurfaceProfile));
+		SubsurfaceColor,
+		SubsurfaceProfile);
 
 #if USES_GBUFFER
 	GBuffer.SelectiveOutputMask = GetSelectiveOutputMask();
diff --git a/Engine/Shaders/Private/ShadingModels.ush b/Engine/Shaders/Private/ShadingModels.ush
--- a/Engine/Shaders/Private/ShadingModels.ush
+++ b/Engine/Shaders/Private/ShadingModels.ush
@@ -120,6 +120,22 @@ float3 SpecularGGX( float Roughness, float3 SpecularColor, BxDFContext Context,
 	return (D * Vis) * F;
 }
 
+float3 DualSpecularGGX(float AverageRoughness, float Lobe0Roughness, float Lobe1Roughness, float LobeMix, float3 SpecularColor, BxDFContext Context, float NoL, FAreaLight AreaLight)
+{
+	float AverageAlpha2 = Pow4(AverageRoughness);
+	float Lobe0Alpha2 = Pow4(Lobe0Roughness);
+	float Lobe1Alpha2 = Pow4(Lobe1Roughness);
+
+	float Lobe0Energy = EnergyNormalization(Lobe0Alpha2, Context.VoH, AreaLight);
+	float Lobe1Energy = EnergyNormalization(Lobe1Alpha2, Context.VoH, AreaLight);
+
+	// Generalized microfacet specular
+	float D = lerp(D_GGX(Lobe0Alpha2, Context.NoH) * Lobe0Energy, D_GGX(Lobe1Alpha2, Context.NoH) * Lobe1Energy, LobeMix);
+	float Vis = Vis_SmithJointApprox(AverageAlpha2, Context.NoV, NoL); // Average visibility well approximates using two separate ones (one per lobe).
+	float3 F = F_Schlick(SpecularColor, Context.VoH);
+
+	return (D * Vis) * F;
+}
 
 struct FDirectLighting
 {
@@ -303,6 +319,17 @@ float ApproximateHG(float cosJ, float g)
 	return (ISO_PHASE_FUNC_Normalized * (1.0f - g2) / gcos2);
 }
 
+void GetProfileDualSpecular(FGBufferData GBuffer, out float AverageToRoughness0, out float AverageToRoughness1, out float LobeMix)
+{
+	// 0..255, which SubSurface profile to pick
+	uint SubsurfaceProfileInt = ExtractSubsurfaceProfileInt(GBuffer);
+
+	float4 Data = SSProfilesTexture.Load(int3(SSSS_DUAL_SPECULAR_OFFSET, SubsurfaceProfileInt, 0));
+	AverageToRoughness0 = Data.x * SSSS_MAX_DUAL_SPECULAR_ROUGHNESS;
+	AverageToRoughness1 = Data.y * SSSS_MAX_DUAL_SPECULAR_ROUGHNESS;
+	LobeMix = Data.z;
+}
+
 FDirectLighting SubsurfaceProfileBxDF( FGBufferData GBuffer, half3 N, half3 V, half3 L, float Falloff, float NoL, FAreaLight AreaLight, FShadowTerms Shadow )
 {
 	BxDFContext Context;
@@ -310,13 +337,34 @@ FDirectLighting SubsurfaceProfileBxDF( FGBufferData GBuffer, half3 N, half3 V, h
 	SphereMaxNoH( Context, AreaLight.SphereSinAlpha, true );
 	Context.NoV = saturate( abs( Context.NoV ) + 1e-5 );
 
+	float AverageToRoughness0;
+	float AverageToRoughness1;
+	float LobeMix;
+	GetProfileDualSpecular(GBuffer, AverageToRoughness0, AverageToRoughness1, LobeMix);
+
+	float AverageRoughness = GBuffer.Roughness;
+	float Lobe0Roughness = saturate(AverageRoughness * AverageToRoughness0);
+	float Lobe1Roughness = saturate(AverageRoughness * AverageToRoughness1);
+
+	// Smoothly lerp to default single GGX lobe as Opacity approaches 0, before reverting to SHADINGMODELID_DEFAULT_LIT.
+	// See SUBSURFACE_PROFILE_OPACITY_THRESHOLD in ShadingModelsMaterial.ush.
+	float Opacity = GBuffer.CustomData.a;
+	Lobe0Roughness = lerp(1.0f, Lobe0Roughness, saturate(Opacity * 10.0f));
+	Lobe1Roughness = lerp(1.0f, Lobe1Roughness, saturate(Opacity * 10.0f));
+
 	FDirectLighting Lighting;
 	Lighting.Diffuse  = AreaLight.FalloffColor * (Falloff * NoL) * Diffuse_Burley( GBuffer.DiffuseColor, GBuffer.Roughness, Context.NoV, NoL, Context.VoH );
 
-	if( AreaLight.bIsRect )
-		Lighting.Specular = RectGGXApproxLTC( GBuffer.Roughness, GBuffer.SpecularColor, N, V, AreaLight.Rect, AreaLight.Texture );
+	if (AreaLight.bIsRect)
+	{
+		float3 Lobe0Specular = RectGGXApproxLTC(Lobe0Roughness, GBuffer.SpecularColor, N, V, AreaLight.Rect, AreaLight.Texture);
+		float3 Lobe1Specular = RectGGXApproxLTC(Lobe1Roughness, GBuffer.SpecularColor, N, V, AreaLight.Rect, AreaLight.Texture);
+		Lighting.Specular = lerp(Lobe0Specular, Lobe1Specular, LobeMix);
+	}
 	else
-		Lighting.Specular = AreaLight.FalloffColor * (Falloff * NoL) * SpecularGGX( GBuffer.Roughness, GBuffer.SpecularColor, Context, NoL, AreaLight );
+	{
+		Lighting.Specular = AreaLight.FalloffColor * (Falloff * NoL) * DualSpecularGGX(AverageRoughness, Lobe0Roughness, Lobe1Roughness, LobeMix, GBuffer.SpecularColor, Context, NoL, AreaLight);
+	}
 
 
 #if USE_TRANSMISSION
diff --git a/Engine/Shaders/Private/ShadingModelsMaterial.ush b/Engine/Shaders/Private/ShadingModelsMaterial.ush
--- a/Engine/Shaders/Private/ShadingModelsMaterial.ush
+++ b/Engine/Shaders/Private/ShadingModelsMaterial.ush
@@ -4,6 +4,8 @@
 	ShadingModelsMaterial.usf: Pixel shader function for computing a GBuffer from shading model.
 =============================================================================*/
 
+#include "SubsurfaceProfileCommon.ush"
+
 // Optimization: if opacity is 0 then revert to default shading model 
 #define SUBSURFACE_PROFILE_OPACITY_THRESHOLD 1
 
@@ -15,11 +17,9 @@ void SetGBufferForShadingModel(
 	const half  Metallic,
 	const half  Specular,
 	const float Roughness,
-	const float4 SubsurfaceData)
+	const float3 SubsurfaceColor,
+	const float SubsurfaceProfile)
 {
-	const float3 SubsurfaceColor = SubsurfaceData.rgb;
-	const float SubsurfaceProfile = SubsurfaceData.a;
-
 	GBuffer.WorldNormal = MaterialParameters.WorldNormal;
 	GBuffer.BaseColor = BaseColor;
 	GBuffer.Metallic = Metallic;
@@ -47,6 +47,15 @@ void SetGBufferForShadingModel(
 		GBuffer.ShadingModelID = SHADINGMODELID_SUBSURFACE_PROFILE;
 		GBuffer.CustomData.rgb = EncodeSubsurfaceProfile(SubsurfaceProfile);
 		GBuffer.CustomData.a = Opacity;
+
+		// Average roughness for dual specular.
+		uint SubsurfaceProfileUInt = uint(SubsurfaceProfile * 255.0f + 0.5f);
+		float MaterialRoughnessToAverage = BasePassSSProfilesTexture.Load(int3(SSSS_DUAL_SPECULAR_OFFSET, SubsurfaceProfileUInt, 0)).w * SSSS_MAX_DUAL_SPECULAR_ROUGHNESS;
+
+		// Smooth blend out dual specular when opacity is low.
+		MaterialRoughnessToAverage = lerp(1.0f, MaterialRoughnessToAverage, saturate(Opacity * 10.0f));
+
+		GBuffer.Roughness = saturate(GBuffer.Roughness * MaterialRoughnessToAverage);
 	}
 	#if SUBSURFACE_PROFILE_OPACITY_THRESHOLD
 	else
diff --git a/Engine/Shaders/Private/SubsurfaceProfileCommon.ush b/Engine/Shaders/Private/SubsurfaceProfileCommon.ush
--- a/Engine/Shaders/Private/SubsurfaceProfileCommon.ush
+++ b/Engine/Shaders/Private/SubsurfaceProfileCommon.ush
@@ -12,8 +12,8 @@
 #define SSSS_SUBSURFACE_COLOR_OFFSET			0
 #define SSSS_TRANSMISSION_OFFSET				(SSSS_SUBSURFACE_COLOR_OFFSET+1)
 #define SSSS_BOUNDARY_COLOR_BLEED_OFFSET		(SSSS_TRANSMISSION_OFFSET+1)
-#define SSSS_RESERVED_PARAM2_OFFSET				(SSSS_BOUNDARY_COLOR_BLEED_OFFSET+1)
-#define SSSS_KERNEL0_OFFSET						(SSSS_RESERVED_PARAM2_OFFSET+1)
+#define SSSS_DUAL_SPECULAR_OFFSET				(SSSS_BOUNDARY_COLOR_BLEED_OFFSET+1)
+#define SSSS_KERNEL0_OFFSET						(SSSS_DUAL_SPECULAR_OFFSET+1)
 #define SSSS_KERNEL0_SIZE						13
 #define SSSS_KERNEL1_OFFSET						(SSSS_KERNEL0_OFFSET + SSSS_KERNEL0_SIZE)
 #define SSSS_KERNEL1_SIZE						9
@@ -23,6 +23,7 @@
 #define SSSS_TRANSMISSION_PROFILE_OFFSET		(SSSS_KERNEL0_OFFSET + SSSS_KERNEL_TOTAL_SIZE)
 #define SSSS_TRANSMISSION_PROFILE_SIZE			32
 #define	SSSS_MAX_TRANSMISSION_PROFILE_DISTANCE	5.0f // See MaxTransmissionProfileDistance in ComputeTransmissionProfile(), SeparableSSS.cpp
+#define SSSS_MAX_DUAL_SPECULAR_ROUGHNESS		2.0f
 
 /** One profile per line, encoded using constants above. See FSubsurfaceProfileTexture::CreateTexture() in SubsurfaceProfile.cpp. */
 Texture2D SSProfilesTexture;
diff --git a/Engine/Source/Runtime/Engine/Classes/Engine/SubsurfaceProfile.h b/Engine/Source/Runtime/Engine/Classes/Engine/SubsurfaceProfile.h
--- a/Engine/Source/Runtime/Engine/Classes/Engine/SubsurfaceProfile.h
+++ b/Engine/Source/Runtime/Engine/Classes/Engine/SubsurfaceProfile.h
@@ -48,6 +48,15 @@ struct FSubsurfaceProfileStruct
 	UPROPERTY(Category = "Transmission", EditAnywhere, BlueprintReadOnly, meta = (ClampMin = "1.0", UIMax = "3.0", ClampMax = "3.0"))
 	float IOR;
 
+	UPROPERTY(Category = "Dual Specular", EditAnywhere, BlueprintReadOnly, meta = (ClampMin = "0.5", UIMax = "2.0", ClampMax = "2.0"))
+	float Roughness0;
+
+	UPROPERTY(Category = "Dual Specular", EditAnywhere, BlueprintReadOnly, meta = (ClampMin = "0.5", UIMax = "2.0", ClampMax = "2.0"))
+	float Roughness1;
+
+	UPROPERTY(Category = "Dual Specular", EditAnywhere, BlueprintReadOnly, meta = (ClampMin = "0.1", UIMax = "0.9", ClampMax = "0.9"))
+	float LobeMix;
+
 	// constructor
 	FSubsurfaceProfileStruct()
 	{
@@ -59,7 +68,10 @@ struct FSubsurfaceProfileStruct
 		ExtinctionScale = 1.0f;
 		ScatteringDistribution = 0.93f;
 		NormalScale = 0.08f;
-		IOR = 1.55;
+		IOR = 1.55f;
+		Roughness0 = 0.75f;
+		Roughness1 = 1.30f;
+		LobeMix = 0.85f;
 	}
 
 	void Invalidate()
@@ -72,7 +84,10 @@ struct FSubsurfaceProfileStruct
 		ExtinctionScale = 1.0f;
 		ScatteringDistribution = 0.93f;
 		NormalScale = 0.08f;
-		IOR = 1.55;
+		IOR = 1.55f;
+		Roughness0 = 0.75f;
+		Roughness1 = 1.30f;
+		LobeMix = 0.85f;
 	}
 };
 
diff --git a/Engine/Source/Runtime/Engine/Private/Rendering/SubsurfaceProfile.cpp b/Engine/Source/Runtime/Engine/Private/Rendering/SubsurfaceProfile.cpp
--- a/Engine/Source/Runtime/Engine/Private/Rendering/SubsurfaceProfile.cpp
+++ b/Engine/Source/Runtime/Engine/Private/Rendering/SubsurfaceProfile.cpp
@@ -127,8 +127,8 @@ static float GetNextSmallerPositiveFloat(float x)
 #define SSSS_SUBSURFACE_COLOR_OFFSET			0
 #define SSSS_TRANSMISSION_OFFSET				(SSSS_SUBSURFACE_COLOR_OFFSET+1)
 #define SSSS_BOUNDARY_COLOR_BLEED_OFFSET		(SSSS_TRANSMISSION_OFFSET+1)
-#define SSSS_RESERVED_PARAM2_OFFSET				(SSSS_BOUNDARY_COLOR_BLEED_OFFSET+1)
-#define SSSS_KERNEL0_OFFSET						(SSSS_RESERVED_PARAM2_OFFSET+1)
+#define SSSS_DUAL_SPECULAR_OFFSET				(SSSS_BOUNDARY_COLOR_BLEED_OFFSET+1)
+#define SSSS_KERNEL0_OFFSET						(SSSS_DUAL_SPECULAR_OFFSET+1)
 #define SSSS_KERNEL0_SIZE						13
 #define SSSS_KERNEL1_OFFSET						(SSSS_KERNEL0_OFFSET + SSSS_KERNEL0_SIZE)
 #define SSSS_KERNEL1_SIZE						9
@@ -138,6 +138,17 @@ static float GetNextSmallerPositiveFloat(float x)
 #define SSSS_TRANSMISSION_PROFILE_OFFSET		(SSSS_KERNEL0_OFFSET + SSSS_KERNEL_TOTAL_SIZE)
 #define SSSS_TRANSMISSION_PROFILE_SIZE			32
 #define	SSSS_MAX_TRANSMISSION_PROFILE_DISTANCE	5.0f // See MaxTransmissionProfileDistance in ComputeTransmissionProfile(), SeparableSSS.cpp
+#define	SSSS_MAX_DUAL_SPECULAR_ROUGHNESS		2.0f
+
+float Sqrt2(float X)
+{
+	return sqrtf(sqrtf(X));
+}
+
+float Pow4(float X)
+{
+	return X * X * X * X;
+}
 
 void FSubsurfaceProfileTexture::CreateTexture(FRHICommandListImmediate& RHICmdList)
 {
@@ -188,6 +199,15 @@ void FSubsurfaceProfileTexture::CreateTexture(FRHICommandListImmediate& RHICmdLi
 
 		TextureRow[SSSS_BOUNDARY_COLOR_BLEED_OFFSET] = Data.BoundaryColorBleed;
 
+		float MaterialRoughnessToAverage = Data.Roughness0 * (1.0f - Data.LobeMix) + Data.Roughness1 * Data.LobeMix;
+		float AverageToRoughness0 = Data.Roughness0 / MaterialRoughnessToAverage;
+		float AverageToRoughness1 = Data.Roughness1 / MaterialRoughnessToAverage;
+
+		TextureRow[SSSS_DUAL_SPECULAR_OFFSET].R = FMath::Clamp(AverageToRoughness0 / SSSS_MAX_DUAL_SPECULAR_ROUGHNESS, 0.0f, 1.0f);
+		TextureRow[SSSS_DUAL_SPECULAR_OFFSET].G = FMath::Clamp(AverageToRoughness1 / SSSS_MAX_DUAL_SPECULAR_ROUGHNESS, 0.0f, 1.0f);
+		TextureRow[SSSS_DUAL_SPECULAR_OFFSET].B = Data.LobeMix;
+		TextureRow[SSSS_DUAL_SPECULAR_OFFSET].A = FMath::Clamp(MaterialRoughnessToAverage / SSSS_MAX_DUAL_SPECULAR_ROUGHNESS, 0.0f, 1.0f);
+
 		//X:ExtinctionScale, Y:Normal Scale, Z:ScatteringDistribution, W:OneOverIOR
 		TextureRow[SSSS_TRANSMISSION_OFFSET].R = Data.ExtinctionScale;
 		TextureRow[SSSS_TRANSMISSION_OFFSET].G = Data.NormalScale;
diff --git a/Engine/Source/Runtime/Renderer/Private/BasePassRendering.cpp b/Engine/Source/Runtime/Renderer/Private/BasePassRendering.cpp
--- a/Engine/Source/Runtime/Renderer/Private/BasePassRendering.cpp
+++ b/Engine/Source/Runtime/Renderer/Private/BasePassRendering.cpp
@@ -780,6 +780,7 @@ FTextureRHIRef& GetEyeAdaptation(const FViewInfo& View)
 }
 
 void SetupSharedBasePassParameters(
+	FRHICommandListImmediate& RHICmdList,
 	const FViewInfo& View,
 	FSceneRenderTargets& SceneRenderTargets,
 	FSharedBasePassUniformParameters& SharedParameters)
@@ -798,6 +799,17 @@ void SetupSharedBasePassParameters(
 
 	SetupReflectionUniformParameters(View, SharedParameters.Reflection);
 	SetupFogUniformParameters(View, SharedParameters.Fog);
+
+	const IPooledRenderTarget* PooledRT = GetSubsufaceProfileTexture_RT(RHICmdList);
+
+	if (!PooledRT)
+	{
+		// no subsurface profile was used yet
+		PooledRT = GSystemTextures.BlackDummy;
+	}
+
+	const FSceneRenderTargetItem& Item = PooledRT->GetRenderTargetItem();
+	SharedParameters.SSProfilesTexture = Item.ShaderResourceTexture;
 }
 
 void CreateOpaqueBasePassUniformBuffer(
@@ -809,7 +821,7 @@ void CreateOpaqueBasePassUniformBuffer(
 	FSceneRenderTargets& SceneRenderTargets = FSceneRenderTargets::Get(RHICmdList);
 
 	FOpaqueBasePassUniformParameters BasePassParameters;
-	SetupSharedBasePassParameters(View, SceneRenderTargets, BasePassParameters.Shared);
+	SetupSharedBasePassParameters(RHICmdList, View, SceneRenderTargets, BasePassParameters.Shared);
 
 	// Forward shading
 	{
diff --git a/Engine/Source/Runtime/Renderer/Private/BasePassRendering.h b/Engine/Source/Runtime/Renderer/Private/BasePassRendering.h
--- a/Engine/Source/Runtime/Renderer/Private/BasePassRendering.h
+++ b/Engine/Source/Runtime/Renderer/Private/BasePassRendering.h
@@ -55,6 +55,7 @@ BEGIN_UNIFORM_BUFFER_STRUCT(FSharedBasePassUniformParameters,)
 	UNIFORM_MEMBER_STRUCT(FForwardLightData, ForwardISR)
 	UNIFORM_MEMBER_STRUCT(FReflectionUniformParameters, Reflection)
 	UNIFORM_MEMBER_STRUCT(FFogUniformParameters, Fog)
+	UNIFORM_MEMBER_TEXTURE(Texture2D, SSProfilesTexture)
 END_UNIFORM_BUFFER_STRUCT(FSharedBasePassUniformParameters)
 
 BEGIN_UNIFORM_BUFFER_STRUCT(FOpaqueBasePassUniformParameters,)
@@ -102,6 +103,7 @@ END_UNIFORM_BUFFER_STRUCT(FTranslucentBasePassUniformParameters)
 extern FTextureRHIRef& GetEyeAdaptation(const FViewInfo& View);
 
 extern void SetupSharedBasePassParameters(
+	FRHICommandListImmediate& RHICmdList,
 	const FViewInfo& View,
 	FSceneRenderTargets& SceneRenderTargets,
 	FSharedBasePassUniformParameters& BasePassParameters);
diff --git a/Engine/Source/Runtime/Renderer/Private/TiledDeferredLightRendering.cpp b/Engine/Source/Runtime/Renderer/Private/TiledDeferredLightRendering.cpp
--- a/Engine/Source/Runtime/Renderer/Private/TiledDeferredLightRendering.cpp
+++ b/Engine/Source/Runtime/Renderer/Private/TiledDeferredLightRendering.cpp
@@ -110,6 +110,8 @@ public:
 		LTCMatSampler.Bind(Initializer.ParameterMap, TEXT("LTCMatSampler"));
 		LTCAmpTexture.Bind(Initializer.ParameterMap, TEXT("LTCAmpTexture"));
 		LTCAmpSampler.Bind(Initializer.ParameterMap, TEXT("LTCAmpSampler"));
+		TransmissionProfilesTexture.Bind(Initializer.ParameterMap, TEXT("SSProfilesTexture"));
+		TransmissionProfilesLinearSampler.Bind(Initializer.ParameterMap, TEXT("TransmissionProfilesLinearSampler"));
 	}
 
 	FTiledDeferredLightingCS()
@@ -160,6 +162,27 @@ public:
 			GSystemTextures.LTCAmp->GetRenderTargetItem().ShaderResourceTexture
 			);
 
+		if (TransmissionProfilesTexture.IsBound())
+		{
+			FSceneRenderTargets& SceneContext = FSceneRenderTargets::Get(RHICmdList);
+			const IPooledRenderTarget* PooledRT = GetSubsufaceProfileTexture_RT((FRHICommandListImmediate&)RHICmdList);
+
+			if (!PooledRT)
+			{
+				// no subsurface profile was used yet
+				PooledRT = GSystemTextures.BlackDummy;
+			}
+
+			const FSceneRenderTargetItem& Item = PooledRT->GetRenderTargetItem();
+
+			SetTextureParameter(RHICmdList,
+				ShaderRHI,
+				TransmissionProfilesTexture,
+				TransmissionProfilesLinearSampler,
+				TStaticSamplerState<SF_Bilinear, AM_Clamp, AM_Clamp, AM_Clamp>::GetRHI(),
+				Item.ShaderResourceTexture);
+		}
+
 		static const auto AllowStaticLightingVar = IConsoleManager::Get().FindTConsoleVariableDataInt(TEXT("r.AllowStaticLighting"));
 		const bool bAllowStaticLighting = (!AllowStaticLightingVar || AllowStaticLightingVar->GetValueOnRenderThread() != 0);
 
@@ -257,6 +280,8 @@ public:
 		Ar << LTCMatSampler;
 		Ar << LTCAmpTexture;
 		Ar << LTCAmpSampler;
+		Ar << TransmissionProfilesTexture;
+		Ar << TransmissionProfilesLinearSampler;
 		return bShaderHasOutdatedParameters;
 	}
 
@@ -281,6 +306,8 @@ private:
 	FShaderResourceParameter LTCMatSampler;
 	FShaderResourceParameter LTCAmpTexture;
 	FShaderResourceParameter LTCAmpSampler;
+	FShaderResourceParameter TransmissionProfilesTexture;
+	FShaderResourceParameter TransmissionProfilesLinearSampler;
 };
 
 // #define avoids a lot of code duplication
diff --git a/Engine/Source/Runtime/Renderer/Private/TranslucentRendering.cpp b/Engine/Source/Runtime/Renderer/Private/TranslucentRendering.cpp
--- a/Engine/Source/Runtime/Renderer/Private/TranslucentRendering.cpp
+++ b/Engine/Source/Runtime/Renderer/Private/TranslucentRendering.cpp
@@ -1231,7 +1231,7 @@ void CreateTranslucentBasePassUniformBuffer(
 	FSceneRenderTargets& SceneRenderTargets = FSceneRenderTargets::Get(RHICmdList);
 
 	FTranslucentBasePassUniformParameters BasePassParameters;
-	SetupSharedBasePassParameters(View, SceneRenderTargets, BasePassParameters.Shared);
+	SetupSharedBasePassParameters(RHICmdList, View, SceneRenderTargets, BasePassParameters.Shared);
 
 	{
 		SetupSceneTextureUniformParameters(SceneRenderTargets, View.FeatureLevel, SceneTextureSetupMode, BasePassParameters.SceneTextures);
