Cherry-pick 4074767: Fixes DOF's scattered bokeh changing of intensity dynamic resolution + TAAU.

diff --git a/Engine/Shaders/Private/DiaphragmDOF/DOFReduce.usf b/Engine/Shaders/Private/DiaphragmDOF/DOFReduce.usf
--- a/Engine/Shaders/Private/DiaphragmDOF/DOFReduce.usf
+++ b/Engine/Shaders/Private/DiaphragmDOF/DOFReduce.usf
@@ -294,14 +294,14 @@ void ReduceCS(
 	}
 	#endif
 
-	// Convert Coc radius from preprocessing basis to processing basis.
-	CocRadius *= PreProcessingToProcessingCocRadiusFactor.x;
-
 	// Indirect dispatch 
 	#if DIM_HYBRID_SCATTER_FGD || DIM_HYBRID_SCATTER_BGD
 	{
 		const uint IndirectParameterSize = SCATTERING_INDIRECT_PARAMETER_SIZE;
 		
+		// Convert Coc radius from preprocessing basis to processing basis.
+		float ScatterCocRadius = CocRadius * PreProcessingToProcessingCocRadiusFactor.x;
+
 		// Decides whether should hybrid scatter or not.
 		// It is fine to do it for sample outside viewport UV because they are discarded when the bokeh
 		// is getting close to viewport edge.
@@ -315,7 +315,7 @@ void ReduceCS(
 			Parameters.FrameExposureScale = EyeAdaptationLookup();
 
 			Parameters.Color = GatheredColor;
-			Parameters.CocRadius = CocRadius;
+			Parameters.CocRadius = ScatterCocRadius;
 			Parameters.BufferUV = BufferUV;
 
 			ScatterFactor = ComputeLuminanceOnlyBasedScatterFactor(Parameters);
@@ -324,13 +324,13 @@ void ReduceCS(
 			if (1)
 			{
 				#if DIM_HYBRID_SCATTER_FGD && !DIM_HYBRID_SCATTER_BGD
-					ScatterFactor *= saturate(-Parameters.CocRadius - MinScatteringCocRadius);
+					ScatterFactor *= saturate(-ScatterCocRadius - MinScatteringCocRadius);
 
 				#elif !DIM_HYBRID_SCATTER_FGD && DIM_HYBRID_SCATTER_BGD
-					ScatterFactor *= saturate(Parameters.CocRadius - MinScatteringCocRadius);
+					ScatterFactor *= saturate(ScatterCocRadius - MinScatteringCocRadius);
 
 				#else
-					ScatterFactor *= saturate(abs(Parameters.CocRadius) - MinScatteringCocRadius);
+					ScatterFactor *= saturate(abs(ScatterCocRadius) - MinScatteringCocRadius);
 
 				#endif
 			}
@@ -344,7 +344,7 @@ void ReduceCS(
 					min(SvPosition.x, SvPosition.y),
 					min(ViewportRect.z - SvPosition.x, ViewportRect.w - SvPosition.y)));
 	
-				ScatterFactor *= saturate(ClosestBorderDistance - abs(Parameters.CocRadius));
+				ScatterFactor *= saturate(ClosestBorderDistance - abs(ScatterCocRadius));
 			}
 
 			// Compare this sample with neighbor if worth it. 
@@ -457,8 +457,8 @@ void ReduceCS(
 					OutForegroundScatterDrawList[5 * IndirectDrawIndex + 0] = float4(DispatchThreadId + 0.5, 0.0, 0.0);
 				}
 				
-				float4 ScatterColor = GatheredColor * (ScatterFactor * IntensityLossFactor * (IsForeground(CocRadius) ? 1 : 0));
-				OutForegroundScatterDrawList[5 * IndirectDrawIndex + ScatteringGroupThreadIndex + 1] = float4(ScatterColor.rgb, abs(CocRadius));
+				float4 ScatterColor = GatheredColor * (ScatterFactor * IntensityLossFactor * (IsForeground(ScatterCocRadius) ? 1 : 0));
+				OutForegroundScatterDrawList[5 * IndirectDrawIndex + ScatteringGroupThreadIndex + 1] = float4(ScatterColor.rgb, abs(ScatterCocRadius));
 			}
 			else
 			{
@@ -483,8 +483,8 @@ void ReduceCS(
 					OutBackgroundScatterDrawList[5 * IndirectDrawIndex + 0] = float4(DispatchThreadId + 0.5, 0.0, 0.0);
 				}
 				
-				float4 ScatterColor = GatheredColor * (ScatterFactor * IntensityLossFactor * (IsForeground(CocRadius) ? 0 : 1));
-				OutBackgroundScatterDrawList[5 * IndirectDrawIndex + ScatteringGroupThreadIndex + 1] = float4(ScatterColor.rgb, abs(CocRadius));
+				float4 ScatterColor = GatheredColor * (ScatterFactor * IntensityLossFactor * (IsForeground(ScatterCocRadius) ? 0 : 1));
+				OutBackgroundScatterDrawList[5 * IndirectDrawIndex + ScatteringGroupThreadIndex + 1] = float4(ScatterColor.rgb, abs(ScatterCocRadius));
 			}
 			else
 			{
@@ -497,10 +497,6 @@ void ReduceCS(
 		GatheredColor.rgb *= (1 - ScatterFactor);
 	}
 	#endif // DIM_HYBRID_SCATTER_FGD || DIM_HYBRID_SCATTER_BGD
-	
-	// Convert back.
-	// TODO: could clean that a little bit?
-	CocRadius *= PreProcessingToProcessingCocRadiusFactor.y;
 
 	// Unconditionaly output mip 0 in output buffer that is always a multiple of group size,
 	// for gathering pass to avoid having to clamp UV in gathering pass. 
