Cherry-pick 4075014: Works arround HLSLCC's issue with InterlockedOr in DOF's reduce pass.

diff --git a/Engine/Shaders/Private/DiaphragmDOF/DOFReduce.usf b/Engine/Shaders/Private/DiaphragmDOF/DOFReduce.usf
--- a/Engine/Shaders/Private/DiaphragmDOF/DOFReduce.usf
+++ b/Engine/Shaders/Private/DiaphragmDOF/DOFReduce.usf
@@ -8,6 +8,19 @@
 #include "DOFHybridScatterCompilation.ush"
 
 
+//------------------------------------------------------- ENUM VALUES
+
+/** Method used to allocate scatter group globally. */
+	// Uses bit or atomic
+	#define SCATTER_ALLOC_METHOD_ATOMIC 0
+
+	// Uses wave bit or instruction.
+	#define SCATTER_ALLOC_METHOD_WAVE 1
+	
+	// Uses LDS entry for each thread.
+	#define SCATTER_ALLOC_METHOD_LDS 2
+
+
 //------------------------------------------------------- COMPILE TIME CONSTANTS
 
 #define SHARED_REDUCE_COUNT 3
@@ -28,6 +41,22 @@
 #define CONFIG_WAVE_BROADCAST_REDUCTION (PLATFORM_SUPPORTS_WAVE_BROADCAST)
 
 
+// Configures the neighborhood analysis method to use for slight out of focus early out.
+#if COMPILER_SUPPORTS_WAVE_BIT_ORAND && (PS4_PROFILE || XBOXONE_PROFILE)
+	// GCN only optimisation 
+	#define SCATTER_ALLOC_METHOD (SCATTER_ALLOC_METHOD_WAVE)
+	
+#elif COMPILER_HLSLCC
+	// Compiler does not like InterlockedOr().
+	#define SCATTER_ALLOC_METHOD (SCATTER_ALLOC_METHOD_LDS)
+
+#else
+	#define SCATTER_ALLOC_METHOD (SCATTER_ALLOC_METHOD_ATOMIC)
+
+#endif
+
+
+
 //------------------------------------------------------- PARAMETERS
 
 uint MaxScatteringGroupCount;
@@ -66,8 +95,15 @@ RWTexture2D<float>	Output1Mip4;
 
 //------------------------------------------------------- LDS
 
-groupshared uint SharedForegroundScatterGroupMask;
-groupshared uint SharedBackgroundScatterGroupMask;
+#if SCATTER_ALLOC_METHOD == SCATTER_ALLOC_METHOD_ATOMIC
+	groupshared uint SharedForegroundScatterGroupMask;
+	groupshared uint SharedBackgroundScatterGroupMask;
+
+#elif SCATTER_ALLOC_METHOD == SCATTER_ALLOC_METHOD_LDS
+	groupshared uint SharedForegroundScatterGroupMask[THREADGROUP_TOTALSIZE];
+	groupshared uint SharedBackgroundScatterGroupMask[THREADGROUP_TOTALSIZE];
+
+#endif
 
 groupshared uint SharedForegroundAtomic;
 groupshared uint SharedBackgroundAtomic;
@@ -242,11 +278,11 @@ void ReduceCS(
 	uint GroupThreadIndex : SV_GroupIndex) 
 {
 	// Init LDS.
-	#if (DIM_HYBRID_SCATTER_FGD || DIM_HYBRID_SCATTER_BGD) && !COMPILER_SUPPORTS_WAVE_BIT_ORAND
+	#if (DIM_HYBRID_SCATTER_FGD || DIM_HYBRID_SCATTER_BGD) && SCATTER_ALLOC_METHOD == SCATTER_ALLOC_METHOD_ATOMIC
 	{
 		SharedForegroundScatterGroupMask = 0;
 		SharedBackgroundScatterGroupMask = 0;
-		
+
 		GroupMemoryBarrierWithGroupSync();
 	}
 	#endif
@@ -365,37 +401,37 @@ void ReduceCS(
 		uint ScatteringGroupMask = 1 << ScatteringGroupIndex;
 		
 		const float ScatterFactorThreshold = 0.01;
-		uint bScatterThisColor = ScatterFactor > ScatterFactorThreshold;
-		uint bIsForeground = IsForeground(CocRadius);
+		bool bScatterThisColor = ScatterFactor > ScatterFactorThreshold;
+		bool bIsForeground = IsForeground(CocRadius);
 			
 		// Whether scattering group actually scatter or not.
 		uint ForegroundScatterGroupMask = 0;
 		uint BackgroundScatterGroupMask = 0;
-		#if COMPILER_SUPPORTS_WAVE_BIT_ORAND
+		#if SCATTER_ALLOC_METHOD == SCATTER_ALLOC_METHOD_WAVE
 		{
 			#if DIM_HYBRID_SCATTER_FGD
-				ForegroundScatterGroupMask = WaveAllBitOr((ScatteringGroupMask * bScatterThisColor) * bIsForeground);
+				ForegroundScatterGroupMask = WaveAllBitOr((bScatterThisColor && bIsForeground) ? ScatteringGroupIndex : 0);
 			#endif
 			
 			#if DIM_HYBRID_SCATTER_BGD
-				BackgroundScatterGroupMask = WaveAllBitOr((ScatteringGroupMask * bScatterThisColor) * !bIsForeground);
+				BackgroundScatterGroupMask = WaveAllBitOr((bScatterThisColor && !bIsForeground) ? ScatteringGroupIndex : 0);
 			#endif
 		}
-		#else // if !COMPILER_SUPPORTS_WAVE_BIT_ORAND
+		#elif SCATTER_ALLOC_METHOD == SCATTER_ALLOC_METHOD_ATOMIC
 		// Atomic or whether scattering group actually scatter or not.
 		{
 			BRANCH
-			if (DIM_HYBRID_SCATTER_FGD && bIsForeground)
+			if (DIM_HYBRID_SCATTER_FGD && bIsForeground && bScatterThisColor)
 			{
 				uint Unused;
-				InterlockedOr(SharedForegroundScatterGroupMask, (ScatteringGroupMask * bScatterThisColor), Unused);
+				InterlockedOr(SharedForegroundScatterGroupMask, ScatteringGroupMask, Unused);
 			}
 			
 			BRANCH
-			if (DIM_HYBRID_SCATTER_BGD && !bIsForeground)
+			if (DIM_HYBRID_SCATTER_BGD && !bIsForeground && bScatterThisColor)
 			{
 				uint Unused;
-				InterlockedOr(SharedBackgroundScatterGroupMask, (ScatteringGroupMask * bScatterThisColor), Unused);
+				InterlockedOr(SharedBackgroundScatterGroupMask, ScatteringGroupMask, Unused);
 			}
 			
 			GroupMemoryBarrierWithGroupSync();
@@ -407,6 +443,34 @@ void ReduceCS(
 				BackgroundScatterGroupMask = SharedBackgroundScatterGroupMask;
 			#endif
 		}
+		#elif SCATTER_ALLOC_METHOD == SCATTER_ALLOC_METHOD_LDS
+		{
+			#if DIM_HYBRID_SCATTER_FGD
+				SharedForegroundScatterGroupMask[GroupThreadIndex] = (
+					(bScatterThisColor && bIsForeground) ? ScatteringGroupIndex : 0);
+			#endif
+
+			#if DIM_HYBRID_SCATTER_BGD
+				SharedBackgroundScatterGroupMask[GroupThreadIndex] = (
+					(bScatterThisColor && !bIsForeground) ? ScatteringGroupIndex : 0);
+			#endif
+
+			GroupMemoryBarrierWithGroupSync();
+
+			UNROLL
+			for (uint i = 0; i < THREADGROUP_TOTALSIZE; i++)
+			{
+				#if DIM_HYBRID_SCATTER_FGD
+					ForegroundScatterGroupMask |= SharedForegroundScatterGroupMask[i];
+				#endif
+				
+				#if DIM_HYBRID_SCATTER_BGD
+					BackgroundScatterGroupMask |= SharedBackgroundScatterGroupMask[i];
+				#endif
+			}
+		}
+		#else
+			#error Unknown scatter group allocation method.
 		#endif // !COMPILER_SUPPORTS_WAVE_BIT_ORAND
 
 		// Allocate globally.
