[NO NON-PII COMMIT MESSAGE]

diff --git a/Engine/Shaders/Private/PostProcessSubsurface.usf b/Engine/Shaders/Private/PostProcessSubsurface.usf
--- a/Engine/Shaders/Private/PostProcessSubsurface.usf
+++ b/Engine/Shaders/Private/PostProcessSubsurface.usf
@@ -396,8 +396,10 @@ void SetupPS(in noperspective float4 UVAndScreenPos : TEXCOORD0, out float4 OutC
 
 	OutColor.a = SingleDepth;
 #else // HALF_RES
-	OutColor.rgb = SetupSubsurfaceForOnePixel(BufferUV).rgb;
-	OutColor.a = CalcSceneDepth(BufferUV);
+	OutColor = SetupSubsurfaceForOnePixel(BufferUV);
+	if( OutColor.a > 0){
+		OutColor.a = CalcSceneDepth(BufferUV);
+	}
 #endif // HALF_RES
 
 #if 0
