Moving CL4119043 into UE4-Main. (Workaround for metal shader compiler bug on NV)

diff --git a/Engine/Shaders/Private/ReflectionEnvironmentPixelShader.usf b/Engine/Shaders/Private/ReflectionEnvironmentPixelShader.usf
--- a/Engine/Shaders/Private/ReflectionEnvironmentPixelShader.usf
+++ b/Engine/Shaders/Private/ReflectionEnvironmentPixelShader.usf
@@ -298,7 +298,7 @@ void ReflectionEnvironmentSkyLighting(
 {
 	float2 ScreenUV = UVAndScreenPos.xy;
 	FScreenSpaceData ScreenSpaceData = GetScreenSpaceData(ScreenUV);
-	uint ShadingModelID = GetShadingModelId(ScreenUV);
+	uint ShadingModelID = ScreenSpaceData.GBuffer.ShadingModelID;
 
 	OutColor = 0.0f;
 
