Always normalizing custom alignment and facing vectors

diff --git a/Engine/Shaders/Private/NiagaraSpriteVertexFactory.ush b/Engine/Shaders/Private/NiagaraSpriteVertexFactory.ush
--- a/Engine/Shaders/Private/NiagaraSpriteVertexFactory.ush
+++ b/Engine/Shaders/Private/NiagaraSpriteVertexFactory.ush
@@ -732,8 +732,8 @@ FVertexFactoryIntermediates GetVertexFactoryIntermediates(FVertexFactoryInput In
 	float2 ParticleSize = GetNiagaraParticleSize(ParticleID);
 	Intermediates.SizeRotSubImage = float4(ParticleSize.x, ParticleSize.y, GetNiagaraParticleRotation(ParticleID), GetNiagaraParticleSubimage(ParticleID));
 	Intermediates.Color	= GetNiagaraParticleColor(ParticleID);
-	Intermediates.CustomFacingVector = GetNiagaraParticleFacingVector(ParticleID);
-	Intermediates.CustomAlignmentVector = GetNiagaraParticleAlignmentVector(ParticleID);
+	Intermediates.CustomFacingVector = SafeNormalize(GetNiagaraParticleFacingVector(ParticleID));
+	Intermediates.CustomAlignmentVector = SafeNormalize(GetNiagaraParticleAlignmentVector(ParticleID));
 
 #if NUM_DYNAMIC_PARAMETERS > 0
 	Intermediates.DynamicParameter = GetNiagaraParticleDynamicParameters(ParticleID);
