Fixed: Modulated shadows in Mali device, ES3_1, movable object leads to shadow artifact when moving camera
Always use framebuffer_fetch on ES3.1
[RN] none

diff --git a/Engine/Shaders/Private/SceneTexturesCommon.ush b/Engine/Shaders/Private/SceneTexturesCommon.ush
--- a/Engine/Shaders/Private/SceneTexturesCommon.ush
+++ b/Engine/Shaders/Private/SceneTexturesCommon.ush
@@ -105,6 +105,8 @@ float CalcSceneDepth(float2 ScreenUV)
 			#endif 
 		#elif METAL_PROFILE && !MAC
 				return FramebufferFetchES2().w;
+		#elif COMPILER_GLSL_ES3_1
+				return FramebufferFetchES2().w;
 		#else
 			return ConvertFromDeviceZ(Texture2DSampleLevel(MobileSceneTextures.SceneDepthTexture, MobileSceneTextures.SceneDepthTextureSampler, ScreenUV, 0).r);
 		#endif
@@ -115,7 +117,7 @@ float CalcSceneDepth(float2 ScreenUV)
 /** Returns DeviceZ which is the z value stored in the depth buffer. */
 float LookupDeviceZ( float2 ScreenUV )
 {
-#if COMPILER_GLSL_ES2
+#if COMPILER_GLSL_ES2 || COMPILER_GLSL_ES3_1
 	// todo: can be optimized
 	return ConvertToDeviceZ(CalcSceneDepth(ScreenUV));
 #elif METAL_PROFILE
