Eye shading update.
Added Iris normal (disabled). Removed wrap. Fixed contact shadows.

diff --git a/Engine/Shaders/Private/DeferredLightingCommon.ush b/Engine/Shaders/Private/DeferredLightingCommon.ush
--- a/Engine/Shaders/Private/DeferredLightingCommon.ush
+++ b/Engine/Shaders/Private/DeferredLightingCommon.ush
@@ -187,17 +187,14 @@ void GetShadowTerms(FGBufferData GBuffer, FDeferredLightData LightData, float3 W
 	{
 		float StepOffset = Dither - 0.5;
 		float ContactShadow = ShadowRayCast( WorldPosition + View.PreViewTranslation, L, ContactShadowLength, 8, StepOffset );
+		
+		Shadow.SurfaceShadow *= ContactShadow;
 
 		FLATTEN
-		if (GBuffer.ShadingModelID == SHADINGMODELID_HAIR)
-		{
+		if( GBuffer.ShadingModelID == SHADINGMODELID_HAIR )
 			Shadow.TransmissionShadow *= ContactShadow;
-		}
-		else
-		{
-			Shadow.SurfaceShadow *= ContactShadow;
+		else if( GBuffer.ShadingModelID != SHADINGMODELID_EYE )
 			Shadow.TransmissionShadow *= ContactShadow * 0.5 + 0.5;
-		}
 	}
 #endif
 }
@@ -342,7 +339,7 @@ float4 GetDynamicLighting(float3 WorldPosition, float3 CameraVector, FGBufferDat
 			}
 				
 			LightAccumulator_Add( LightAccumulator, Lighting.Diffuse + Lighting.Specular, Lighting.Diffuse, LightColor * LightMask * Shadow.SurfaceShadow, bNeedsSeparateSubsurfaceLightAccumulation );
-			LightAccumulator_Add( LightAccumulator, Lighting.Transmission, 0, LightColor * LightMask * Shadow.TransmissionShadow, false );
+			LightAccumulator_Add( LightAccumulator, Lighting.Transmission, Lighting.Transmission, LightColor * LightMask * Shadow.TransmissionShadow, bNeedsSeparateSubsurfaceLightAccumulation );
 
 			LightAccumulator.EstimatedCost += 0.4f;		// add the cost of the lighting computations (should sum up to 1 form one light)
 		#endif
diff --git a/Engine/Shaders/Private/DeferredShadingCommon.ush b/Engine/Shaders/Private/DeferredShadingCommon.ush
--- a/Engine/Shaders/Private/DeferredShadingCommon.ush
+++ b/Engine/Shaders/Private/DeferredShadingCommon.ush
@@ -557,11 +557,15 @@ FGBufferData DecodeGBufferData(
 	GBuffer.Depth = SceneDepth;
 
 	GBuffer.StoredBaseColor = GBuffer.BaseColor;
-	GBuffer.StoredSpecular = GBuffer.Specular;
 	GBuffer.StoredMetallic = GBuffer.Metallic;
+	GBuffer.StoredSpecular = GBuffer.Specular;
 
-	// SHADINGMODELID_EYE stores Iris Distance inside Metallic.
-	GBuffer.Metallic = GBuffer.ShadingModelID == SHADINGMODELID_EYE ? 0.0f : GBuffer.Metallic;
+	FLATTEN
+	if( GBuffer.ShadingModelID == SHADINGMODELID_EYE )
+	{
+		GBuffer.Metallic = 0.0;
+		//GBuffer.Specular = 0.25;
+	}
 
 	// derived from BaseColor, Metalness, Specular
 	{
diff --git a/Engine/Shaders/Private/ShadingModels.ush b/Engine/Shaders/Private/ShadingModels.ush
--- a/Engine/Shaders/Private/ShadingModels.ush
+++ b/Engine/Shaders/Private/ShadingModels.ush
@@ -755,6 +755,9 @@ float3 HairShading( FGBufferData GBuffer, float3 L, float3 V, half3 N, float Sha
 
 	if(1)
 	{
+		// Use soft Kajiya Kay diffuse attenuation
+		float KajiyaDiffuse = 1 - abs( dot(N,L) );
+
 		float3 FakeNormal = normalize( V - N * dot(V,N) );
 		//N = normalize( DiffuseN + FakeNormal * 2 );
 		N = FakeNormal;
@@ -762,7 +765,7 @@ float3 HairShading( FGBufferData GBuffer, float3 L, float3 V, half3 N, float Sha
 		// Hack approximation for multiple scattering.
 		float Wrap = 1;
 		float NoL = saturate( ( dot(N, L) + Wrap ) / Square( 1 + Wrap ) );
-		float DiffuseScatter = (1 / PI) * NoL * GBuffer.Metallic;
+		float DiffuseScatter = (1 / PI) * lerp( NoL, KajiyaDiffuse, 0.33 ) * GBuffer.Metallic;
 		float Luma = Luminance( GBuffer.BaseColor );
 		float3 ScatterTint = pow( GBuffer.BaseColor / Luma, 1 - Shadow );
 		S += sqrt( GBuffer.BaseColor ) * DiffuseScatter * ScatterTint;
@@ -778,7 +781,6 @@ FDirectLighting HairBxDF( FGBufferData GBuffer, half3 N, half3 V, half3 L, float
 	FDirectLighting Lighting;
 	Lighting.Diffuse = 0;
 	Lighting.Specular = 0;
-	// FIXME Not passing Shadow!!!!!!
 	Lighting.Transmission = AreaLight.FalloffColor * Falloff * HairShading( GBuffer, L, V, N, Shadow.TransmissionShadow, 1, 0, uint2(0,0) );
 	return Lighting;
 }
@@ -827,10 +829,26 @@ FDirectLighting TwoSidedBxDF( FGBufferData GBuffer, half3 N, half3 V, half3 L, f
 
 FDirectLighting EyeBxDF( FGBufferData GBuffer, half3 N, half3 V, half3 L, float Falloff, float NoL, FAreaLight AreaLight, FShadowTerms Shadow )
 {
+#if 0
+	const float2 CausticNormalDelta	= float2( GBuffer.StoredMetallic, GBuffer.StoredSpecular )	* 2 - (256.0/255.0);
+	const float2 IrisNormalDelta	= float2( GBuffer.CustomData.y, GBuffer.CustomData.z )		* 2 - (256.0/255.0);
+	const float  IrisMask = 1.0f - GBuffer.CustomData.w;
+
+	const float2 WorldNormalOct	= UnitVectorToOctahedron( GBuffer.WorldNormal );
+	const float3 CausticNormal	= OctahedronToUnitVector( WorldNormalOct + CausticNormalDelta );
+	const float3 IrisNormal		= OctahedronToUnitVector( WorldNormalOct + IrisNormalDelta );
+#else
 	const float3 IrisNormal		= OctahedronToUnitVector( GBuffer.CustomData.yz * 2 - 1 );
 	const float  IrisDistance	= GBuffer.StoredMetallic;
 	const float  IrisMask		= 1.0f - GBuffer.CustomData.w;
 
+	// Blend in the negative intersection normal to create some concavity
+	// Not great as it ties the concavity to the convexity of the cornea surface
+	// No good justification for that. On the other hand, if we're just looking to
+	// introduce some concavity, this does the job.
+	const float3 CausticNormal = normalize(lerp(IrisNormal, -N, IrisMask*IrisDistance));
+#endif
+
 	BxDFContext Context;
 	Init( Context, N, V, L );
 	SphereMaxNoH( Context, AreaLight.SphereSinAlpha, false );
@@ -860,25 +878,14 @@ FDirectLighting EyeBxDF( FGBufferData GBuffer, half3 N, half3 V, half3 L, float
 		Lighting.Specular = AreaLight.FalloffColor * (Falloff * NoL) * D * Vis * F;
 	}
 
-
-	// Blend in the negative intersection normal to create some concavity
-	// Not great as it ties the concavity to the convexity of the cornea surface
-	// No good justification for that. On the other hand, if we're just looking to
-	// introduce some concavity, this does the job.
-	float3 CausticNormal = normalize(lerp(IrisNormal, -N, IrisMask*IrisDistance));
-
 	float IrisNoL = saturate( dot( IrisNormal, L ) );
 	float Power = lerp( 12, 1, IrisNoL );
-	float Caustic = 0.6 + 0.2 * ( Power + 1 ) * pow( saturate( dot( CausticNormal, L ) ), Power );
+	float Caustic = 0.8 + 0.2 * ( Power + 1 ) * pow( saturate( dot( CausticNormal, L ) ), Power );
 	float Iris = IrisNoL * Caustic;
-
-	// http://blog.stevemcauley.com/2011/12/03/energy-conserving-wrapped-diffuse/
-	float Wrap = 0.15;
-	float Sclera = saturate( ( NoL + Wrap ) / Square( 1 + Wrap ) );
-	
+	float Sclera = NoL;
 	
-	Lighting.Diffuse  = AreaLight.FalloffColor * (Falloff * lerp( Sclera, Iris, IrisMask ) * (1 - F) ) * Diffuse_Lambert( GBuffer.DiffuseColor );
-	Lighting.Transmission = 0;
+	Lighting.Diffuse = 0;
+	Lighting.Transmission = AreaLight.FalloffColor * ( Falloff * lerp( Sclera, Iris, IrisMask ) * (1 - F) ) * Diffuse_Lambert( GBuffer.DiffuseColor );
 	return Lighting;
 }
 
diff --git a/Engine/Shaders/Private/ShadingModelsMaterial.ush b/Engine/Shaders/Private/ShadingModelsMaterial.ush
--- a/Engine/Shaders/Private/ShadingModelsMaterial.ush
+++ b/Engine/Shaders/Private/ShadingModelsMaterial.ush
@@ -129,12 +129,54 @@ void SetGBufferForShadingModel(
 	GBuffer.ShadingModelID = SHADINGMODELID_EYE;
 	GBuffer.CustomData.x = EncodeSubsurfaceProfile(SubsurfaceProfile).x;
 	GBuffer.CustomData.w = 1.0f - saturate(GetMaterialCustomData0(MaterialParameters));	// Opacity = 1.0 - Iris Mask
+#if 0
+	float IrisMask		= saturate( GetMaterialCustomData0(MaterialParameters) );
+	float IrisDistance	= saturate( GetMaterialCustomData1(MaterialParameters) );
+
+	GBuffer.CustomData.x = EncodeSubsurfaceProfile(SubsurfaceProfile).x;
+	GBuffer.CustomData.w = 1.0 - IrisMask;	// Opacity
+
+	float2 WorldNormalOct = UnitVectorToOctahedron( GBuffer.WorldNormal );
+
+	// CausticNormal stored as octahedron
+	#if NUM_MATERIAL_OUTPUTS_GETTANGENTOUTPUT > 0
+		// Blend in the negative intersection normal to create some concavity
+		// Not great as it ties the concavity to the convexity of the cornea surface
+		// No good justification for that. On the other hand, if we're just looking to
+		// introduce some concavity, this does the job.
+		float3 PlaneNormal = normalize( GetTangentOutput0(MaterialParameters) );
+		float3 CausticNormal = normalize( lerp( PlaneNormal, -GBuffer.WorldNormal, IrisMask*IrisDistance ) );
+		float2 CausticNormalOct  = UnitVectorToOctahedron( CausticNormal );
+		float2 CausticNormalDelta = ( CausticNormalOct - WorldNormalOct ) * 0.5 + (128.0/255.0);
+		GBuffer.Metallic = CausticNormalDelta.x;
+		GBuffer.Specular = CausticNormalDelta.y;
+	#else
+		float3 PlaneNormal = GBuffer.WorldNormal;
+		GBuffer.Metallic = 128.0/255.0;
+		GBuffer.Specular = 128.0/255.0;
+	#endif
+
+	// IrisNormal CustomData.yz
+	#if NUM_MATERIAL_OUTPUTS_CLEARCOATBOTTOMNORMAL > 0
+		float3 IrisNormal = normalize( ClearCoatBottomNormal0(MaterialParameters) );
+		#if MATERIAL_TANGENTSPACENORMAL
+			IrisNormal = normalize( TransformTangentVectorToWorld( MaterialParameters.TangentToWorld, IrisNormal ) );
+		#endif
+	#else
+		float3 IrisNormal = PlaneNormal;
+	#endif
+
+	float2 IrisNormalOct  = UnitVectorToOctahedron( IrisNormal );
+	float2 IrisNormalDelta = ( IrisNormalOct - WorldNormalOct ) * 0.5 + (128.0/255.0);
+	GBuffer.CustomData.yz = IrisNormalDelta;
+#else
 	GBuffer.Metallic = saturate(GetMaterialCustomData1(MaterialParameters));			// Iris Distance
 
 	#if NUM_MATERIAL_OUTPUTS_GETTANGENTOUTPUT > 0
 		float3 Tangent = GetTangentOutput0(MaterialParameters);
 		GBuffer.CustomData.yz = UnitVectorToOctahedron( normalize(Tangent) ) * 0.5 + 0.5;
 	#endif
+#endif
 #else
 	// missing shading model, compiler should report ShadingModelID is not set
 #endif
