Fix for black pixels with rect light and roughness=0.

diff --git a/Engine/Shaders/Private/RectLight.ush b/Engine/Shaders/Private/RectLight.ush
--- a/Engine/Shaders/Private/RectLight.ush
+++ b/Engine/Shaders/Private/RectLight.ush
@@ -405,7 +405,7 @@ float IntegrateEdge( float3 L0, float3 L1 )
 	float w01 = ( 0.8543985 + (0.4965155 + 0.0145206 * abs(c01)) * abs(c01) ) /
 				( 3.4175940 + (4.1616724 + abs(c01)) * abs(c01) );
 	
-	w01 = c01 > 0 ? w01 : 0.5 * rsqrt( 1 - c01 * c01 ) - w01;
+	w01 = c01 > 0 ? w01 : 0.5 * rsqrt( saturate( 1 - c01 * c01 ) ) - w01;
 #endif
 
 	return w01;
@@ -495,7 +495,7 @@ float3 RectGGXApproxLTC( float Roughness, float3 SpecularColor, half3 N, float3
 	Poly[3] = mul( LTC, Rect.Origin - Rect.Axis[0] * Rect.Extent.x + Rect.Axis[1] * Rect.Extent.y );
 
 	// Vector irradiance
-	float3 L = PolygonIrradiance( Poly );
+	float3 L = PolygonIrradiance( Poly ) * (0.5 / PI);
 
 	float LengthSqr = dot( L, L );
 	float InvLength = rsqrt( LengthSqr );
@@ -512,6 +512,9 @@ float3 RectGGXApproxLTC( float Roughness, float3 SpecularColor, half3 N, float3
 	float NoL = SphereHorizonCosWrap( L.z, SinAlphaSqr );
 	float Irradiance = SinAlphaSqr * NoL;
 
+	// Kill negative and NaN
+	Irradiance = -min(-Irradiance, 0.0);
+
 #if 0
 	float NoL;
 	float Falloff;
