Integrated TAAU shader performance tweaks for Switch

Change 4108528 by [REDACTED_PERSON][REDACTED_HANDLE] on 2018/06/03 13:09:53

	TAAU performance tweaks for Switch: disabled YCoCg path, dropped AA_SAMPLES from 6 to 5.

diff --git a/Engine/Shaders/Private/PostProcessTemporalAA.usf b/Engine/Shaders/Private/PostProcessTemporalAA.usf
--- a/Engine/Shaders/Private/PostProcessTemporalAA.usf
+++ b/Engine/Shaders/Private/PostProcessTemporalAA.usf
@@ -129,7 +129,11 @@
 	#define AA_FILTERED TAA_UPSAMPLE_FILTERED
 	#define AA_MANUALLY_CLAMP_HISTORY_UV 1
 	#define AA_TONE 1
-	#define AA_YCOCG 1
+	#if SWITCH_PROFILE || SWITCH_PROFILE_FORWARD
+		#define AA_YCOCG 0
+	#else
+		#define AA_YCOCG 1
+	#endif
 	
 	#define AA_UPSAMPLE 1
 	#define AA_UPSAMPLE_ADAPTIVE_FILTERING 1
@@ -139,7 +143,11 @@
 		#if AA_DOWNSAMPLE
 			#define AA_SAMPLES 9
 		#else
-			#define AA_SAMPLES 6
+			#if SWITCH_PROFILE || SWITCH_PROFILE_FORWARD
+				#define AA_SAMPLES 5
+			#else
+				#define AA_SAMPLES 6
+			#endif
 		#endif
 
 	#else
