Fixed min roughness for dual lobe specular. Before it could go down to 0 resulting in a NaN pixel (D_GGX returns +Inf for roughness = 0).

diff --git a/Engine/Shaders/Private/ShadingModels.ush b/Engine/Shaders/Private/ShadingModels.ush
--- a/Engine/Shaders/Private/ShadingModels.ush
+++ b/Engine/Shaders/Private/ShadingModels.ush
@@ -343,7 +343,7 @@ FDirectLighting SubsurfaceProfileBxDF( FGBufferData GBuffer, half3 N, half3 V, h
 	GetProfileDualSpecular(GBuffer, AverageToRoughness0, AverageToRoughness1, LobeMix);
 
 	float AverageRoughness = GBuffer.Roughness;
-	float Lobe0Roughness = saturate(AverageRoughness * AverageToRoughness0);
+	float Lobe0Roughness = max(saturate(AverageRoughness * AverageToRoughness0), 0.02f);
 	float Lobe1Roughness = saturate(AverageRoughness * AverageToRoughness1);
 
 	// Smoothly lerp to default single GGX lobe as Opacity approaches 0, before reverting to SHADINGMODELID_DEFAULT_LIT.
