Cherry-pick 4149636: Implements depth blur using Diaphragm DOF.

diff --git a/Engine/Shaders/Private/DiaphragmDOF/DOFCommon.ush b/Engine/Shaders/Private/DiaphragmDOF/DOFCommon.ush
--- a/Engine/Shaders/Private/DiaphragmDOF/DOFCommon.ush
+++ b/Engine/Shaders/Private/DiaphragmDOF/DOFCommon.ush
@@ -145,14 +145,27 @@ static const int2 kOffsetsSquare3x3[9] = {
 
 // is wrapping the DepthToCoc() function in this file for faster iteration
 // @return half res pixel radius
-float DepthToHalfResCocRadius(float4 CocModelParameters, float SceneDepth)
+float DepthToHalfResCocRadius(float4 CocModelParameters, float2 DepthBlurParameters, float SceneDepth)
 {
 	const float Focus = View.DepthOfFieldFocalDistance; // TODO: rename that guy to DepthOfFieldFocusDistance.
 	const float InfinityCocRadius = CocModelParameters.x;
 	const float MinCocRadius = CocModelParameters.y;
 	const float MaxCocRadius = CocModelParameters.z;
 	
-	return clamp(((SceneDepth - Focus) / SceneDepth) * InfinityCocRadius, MinCocRadius, MaxCocRadius);
+	// Radius from the circle of confusion.
+	float CocRadius = ((SceneDepth - Focus) / SceneDepth) * InfinityCocRadius;
+
+	// Depth blur's radius.
+	float DepthBlurAbsRadius = (1.0 - exp2(-SceneDepth * DepthBlurParameters.x)) * DepthBlurParameters.y;
+
+	float ReturnCoc = max(abs(CocRadius), DepthBlurAbsRadius);
+	if(CocRadius < 0.0) 
+	{
+		// Near CoC is using negative values.
+		ReturnCoc = -ReturnCoc;
+	}
+
+	return clamp(ReturnCoc, MinCocRadius, MaxCocRadius);
 }
 
 // Return whether this is front element from CocRadius.
diff --git a/Engine/Shaders/Private/DiaphragmDOF/DOFRecombine.usf b/Engine/Shaders/Private/DiaphragmDOF/DOFRecombine.usf
--- a/Engine/Shaders/Private/DiaphragmDOF/DOFRecombine.usf
+++ b/Engine/Shaders/Private/DiaphragmDOF/DOFRecombine.usf
@@ -131,6 +131,7 @@
 
 float2 TemporalJitterPixels;
 float4 CocModelParameters;
+float2 DepthBlurParameters;
 float2 DOFBufferUVMax;
 
 RWTexture2D<float4>	Output0;
@@ -174,7 +175,7 @@ float2 RotVec(float Radius, float Radians)
 
 float DepthToHalfResCocRadius(float WorldDepth)
 {
-	return DepthToHalfResCocRadius(CocModelParameters, WorldDepth);
+	return DepthToHalfResCocRadius(CocModelParameters, DepthBlurParameters, WorldDepth);
 }
 
 void AmendAdditiveColorWithMaxTranslucency(inout float4 Color, inout float Translucency, float MaxTranslucency)
diff --git a/Engine/Shaders/Private/DiaphragmDOF/DOFSetup.usf b/Engine/Shaders/Private/DiaphragmDOF/DOFSetup.usf
--- a/Engine/Shaders/Private/DiaphragmDOF/DOFSetup.usf
+++ b/Engine/Shaders/Private/DiaphragmDOF/DOFSetup.usf
@@ -19,6 +19,7 @@
 //------------------------------------------------------- PARAMETERS
 
 float4 CocModelParameters;
+float2 DepthBlurParameters;
 
 // x: full res basis; y: half res basis.
 float2 CocRadiusBasis;
@@ -75,7 +76,7 @@ void SetupCS(
 
 	float4 SceneColor = PostprocessInput0.SampleLevel(GlobalPointClampedSampler, SceneBufferUV, 0);
 	float DeviceZ = PostprocessInput1.SampleLevel(GlobalPointClampedSampler, SceneBufferUV, 0).r;
-	float CocRadius = DepthToHalfResCocRadius(CocModelParameters, ConvertFromDeviceZ(DeviceZ));
+	float CocRadius = DepthToHalfResCocRadius(CocModelParameters, DepthBlurParameters, ConvertFromDeviceZ(DeviceZ));
 	
 	// Convert from translucency to opacity.
 	SceneColor.a = 1 - SceneColor.a;
@@ -171,7 +172,7 @@ void SetupCS(uint2 DispatchThreadId : SV_DispatchThreadID)
 	float4 OutColor = PostprocessInput0.SampleLevel(GlobalPointClampedSampler, SceneBufferUV, 0);
 	
 	float DeviceZ = PostprocessInput1.SampleLevel(GlobalPointClampedSampler, SceneBufferUV, 0).r;
-	float OutCocRadius = CocRadiusBasis.x * DepthToHalfResCocRadius(CocModelParameters, ConvertFromDeviceZ(DeviceZ));
+	float OutCocRadius = CocRadiusBasis.x * DepthToHalfResCocRadius(CocModelParameters, DepthBlurParameters, ConvertFromDeviceZ(DeviceZ));
 
 	if (all(DispatchThreadId < ViewportRect.zw))
 	{
@@ -206,10 +207,10 @@ void SetupCS(
 		// Fetch Coc
 		float4 SceneDepthGather = PostprocessInput1.Gather(GlobalPointClampedSampler, SceneBufferUV);
 		float CocRadii[4];
-		CocRadii[0] = DepthToHalfResCocRadius(CocModelParameters, ConvertFromDeviceZ(SceneDepthGather.w));
-		CocRadii[1] = DepthToHalfResCocRadius(CocModelParameters, ConvertFromDeviceZ(SceneDepthGather.z));
-		CocRadii[2] = DepthToHalfResCocRadius(CocModelParameters, ConvertFromDeviceZ(SceneDepthGather.x));
-		CocRadii[3] = DepthToHalfResCocRadius(CocModelParameters, ConvertFromDeviceZ(SceneDepthGather.y));
+		CocRadii[0] = DepthToHalfResCocRadius(CocModelParameters, DepthBlurParameters, ConvertFromDeviceZ(SceneDepthGather.w));
+		CocRadii[1] = DepthToHalfResCocRadius(CocModelParameters, DepthBlurParameters, ConvertFromDeviceZ(SceneDepthGather.z));
+		CocRadii[2] = DepthToHalfResCocRadius(CocModelParameters, DepthBlurParameters, ConvertFromDeviceZ(SceneDepthGather.x));
+		CocRadii[3] = DepthToHalfResCocRadius(CocModelParameters, DepthBlurParameters, ConvertFromDeviceZ(SceneDepthGather.y));
 
 		// Fetch color.
 		float4 SceneColors[4];
@@ -258,7 +259,7 @@ void SetupCS(
 			SceneColors[i].a = 1 - SceneColors[i].a;
 
 			float SampleDeviceZ = PostprocessInput1.SampleLevel(GlobalPointClampedSampler, SampleUV, 0).r;
-			CocRadii[i] = DepthToHalfResCocRadius(CocModelParameters, ConvertFromDeviceZ(SampleDeviceZ));
+			CocRadii[i] = DepthToHalfResCocRadius(CocModelParameters, DepthBlurParameters, ConvertFromDeviceZ(SampleDeviceZ));
 		}
 
 		// Downsample.
diff --git a/Engine/Source/Runtime/Renderer/Private/PostProcess/DiaphragmDOF.h b/Engine/Source/Runtime/Renderer/Private/PostProcess/DiaphragmDOF.h
--- a/Engine/Source/Runtime/Renderer/Private/PostProcess/DiaphragmDOF.h
+++ b/Engine/Source/Runtime/Renderer/Private/PostProcess/DiaphragmDOF.h
@@ -33,6 +33,9 @@ struct FPhysicalCocModel
 	// Focus distance.
 	float FocusDistance;
 
+	// The maximum radius of depth blur.
+	float MaxDepthBlurRadius;
+	float DepthBlurExponent;
 
 	/** Compile the coc model from a view. */
 	void Compile(const FViewInfo& View);
@@ -46,7 +49,7 @@ struct FPhysicalCocModel
 	/** Returns limit(DepthToHalfResCocRadius) for SceneDepth -> Infinity. */
 	FORCEINLINE float ComputeViewMaxBackgroundCocRadius(float HorizontalResolution) const
 	{
-		return FMath::Min(InfinityBackgroundCocRadius, MaxBackgroundCocRadius) * HorizontalResolution;
+		return FMath::Min(FMath::Max(InfinityBackgroundCocRadius, MaxDepthBlurRadius), MaxBackgroundCocRadius) * HorizontalResolution;
 	}
 	
 	/** Returns limit(DepthToHalfResCocRadius) for SceneDepth -> 0.
@@ -55,7 +58,7 @@ struct FPhysicalCocModel
 	 */
 	FORCEINLINE float ComputeViewMinForegroundCocRadius(float HorizontalResolution) const
 	{
-		return MinForegroundCocRadius * HorizontalResolution;
+		return DepthToResCocRadius(GNearClippingPlane, HorizontalResolution);
 	}
 };
 
diff --git a/Engine/Source/Runtime/Renderer/Private/PostProcess/DiaphragmDOFGraph.cpp b/Engine/Source/Runtime/Renderer/Private/PostProcess/DiaphragmDOFGraph.cpp
--- a/Engine/Source/Runtime/Renderer/Private/PostProcess/DiaphragmDOFGraph.cpp
+++ b/Engine/Source/Runtime/Renderer/Private/PostProcess/DiaphragmDOFGraph.cpp
@@ -141,6 +141,9 @@ enum class EHybridScatterMode
 }
 
 
+FVector4 CircleDofHalfCoc(const FViewInfo& View);
+
+
 bool DiaphragmDOF::WireSceneColorPasses(FPostprocessContext& Context, const FRenderingCompositeOutputRef& VelocityInput, const FRenderingCompositeOutputRef& SeparateTranslucency)
 {
 	if (Context.View.Family->EngineShowFlags.VisualizeDOF)
diff --git a/Engine/Source/Runtime/Renderer/Private/PostProcess/DiaphragmDOFPasses.cpp b/Engine/Source/Runtime/Renderer/Private/PostProcess/DiaphragmDOFPasses.cpp
--- a/Engine/Source/Runtime/Renderer/Private/PostProcess/DiaphragmDOFPasses.cpp
+++ b/Engine/Source/Runtime/Renderer/Private/PostProcess/DiaphragmDOFPasses.cpp
@@ -523,6 +523,11 @@ void SetCocModelParameters(FRenderingCompositePassContext& Context, DispatchCont
 	CocModelParameters.Y = CocRadiusBasis * CocModel.MinForegroundCocRadius;
 	CocModelParameters.Z = CocRadiusBasis * CocModel.MaxBackgroundCocRadius;
 	SetShaderValue(Context.RHICmdList, DispatchCtx.ShaderRHI, DispatchCtx->CocModelParameters, CocModelParameters);
+
+	FVector2D DepthBlurParameters(0, 0);
+	DepthBlurParameters.X = CocModel.DepthBlurExponent;
+	DepthBlurParameters.Y = CocRadiusBasis * CocModel.MaxDepthBlurRadius;
+	SetShaderValue(Context.RHICmdList, DispatchCtx.ShaderRHI, DispatchCtx->DepthBlurParameters, DepthBlurParameters);
 }
 
 
@@ -643,6 +648,7 @@ FPooledRenderTargetDesc FRCPassDiaphragmDOFDilateCoc::ComputeOutputDesc(EPassOut
 
 #define SETUP_SHADER_PARAMS(PARAMETER) \
 	PARAMETER(FShaderParameter, CocModelParameters) \
+	PARAMETER(FShaderParameter, DepthBlurParameters) \
 	PARAMETER(FShaderParameter, CocRadiusBasis) \
 
 
@@ -1700,6 +1706,7 @@ FPooledRenderTargetDesc FRCPassDiaphragmDOFHybridScatter::ComputeOutputDesc(EPas
 	PARAMETER(FSceneTextureShaderParameters, SceneTextureParameters) \
 	PARAMETER(FShaderParameter, TemporalJitterPixels) \
 	PARAMETER(FShaderParameter, CocModelParameters) \
+	PARAMETER(FShaderParameter, DepthBlurParameters) \
 	PARAMETER(FShaderParameter, DOFBufferUVMax) \
 
 class FPostProcessDiaphragmDOFRecombineCS : public FPostProcessDiaphragmDOFShader
diff --git a/Engine/Source/Runtime/Renderer/Private/PostProcess/DiaphragmDOFUtils.cpp b/Engine/Source/Runtime/Renderer/Private/PostProcess/DiaphragmDOFUtils.cpp
--- a/Engine/Source/Runtime/Renderer/Private/PostProcess/DiaphragmDOFUtils.cpp
+++ b/Engine/Source/Runtime/Renderer/Private/PostProcess/DiaphragmDOFUtils.cpp
@@ -34,6 +34,13 @@ void DiaphragmDOF::FPhysicalCocModel::Compile(const FViewInfo& View)
 		// -because foreground Coc are negative.
 		MinForegroundCocRadius = -CVarMaxForegroundRadius.GetValueOnRenderThread();
 		MaxBackgroundCocRadius = CVarMaxBackgroundRadius.GetValueOnRenderThread();
+
+		MaxDepthBlurRadius = View.FinalPostProcessSettings.DepthOfFieldDepthBlurRadius / 1920.0f;
+
+		// Circle DOF was actually computing in this depth blur radius in half res.
+		MaxDepthBlurRadius *= 2.0f;
+
+		DepthBlurExponent = 1.0f / (View.FinalPostProcessSettings.DepthOfFieldDepthBlurAmount * 100000.0f);
 	}
 
 	// Compile coc model equation.
@@ -64,7 +71,18 @@ void DiaphragmDOF::FPhysicalCocModel::Compile(const FViewInfo& View)
 
 float DiaphragmDOF::FPhysicalCocModel::DepthToResCocRadius(float SceneDepth, float HorizontalResolution) const
 {
-	return HorizontalResolution * FMath::Clamp(((SceneDepth - FocusDistance) / SceneDepth) * InfinityBackgroundCocRadius, MinForegroundCocRadius, MaxBackgroundCocRadius);
+	float CocRadius = ((SceneDepth - FocusDistance) / SceneDepth) * InfinityBackgroundCocRadius;
+
+	// Depth blur based.
+	float DepthBlurAbsRadius = (1.0 - FMath::Exp2(-SceneDepth * DepthBlurExponent)) * MaxDepthBlurRadius;
+
+	float ReturnCoc = FMath::Max(FMath::Abs(CocRadius), DepthBlurAbsRadius);
+	if (CocRadius < 0.0)
+	{
+		// near CoC is using negative values
+		ReturnCoc = -ReturnCoc;
+	}
+	return HorizontalResolution * FMath::Clamp(ReturnCoc, MinForegroundCocRadius, MaxBackgroundCocRadius);
 }
 
 void DiaphragmDOF::FBokehModel::Compile(const FViewInfo& View)
