Cherry-pick 4152079: Fixes long standing TODO in gathering kernel of DOF to take the CocRadius error into account when binning gathering samples.

diff --git a/Engine/Shaders/Private/DiaphragmDOF/DOFGatherAccumulator.ush b/Engine/Shaders/Private/DiaphragmDOF/DOFGatherAccumulator.ush
--- a/Engine/Shaders/Private/DiaphragmDOF/DOFGatherAccumulator.ush
+++ b/Engine/Shaders/Private/DiaphragmDOF/DOFGatherAccumulator.ush
@@ -678,8 +678,11 @@ void MergeCurrentBucketIntoPreviousBucket(inout FGatherAccumulator Accumulator,
 /** Accumulates center sample. */
 void AccumulateCenterSample(inout FGatherAccumulator Accumulator, in FGatherSample A)
 {
+	// Error in coc radius error introduced by the under sampling of the kernel
+	const float CocRadiusError = CONFIG_COC_RADIUS_CORRECTION ? 1.0 : 0.0;
+
 	// Set the bordering radius at exactly the radius of the center sample.
-	Accumulator.BorderingRadius = 0.5 * Accumulator.GatherParameters.KernelRadius * rcp(0.5 + Accumulator.GatherParameters.RingCount);
+	Accumulator.BorderingRadius = (0.5 + CocRadiusError) * Accumulator.GatherParameters.KernelRadius * rcp(0.5 + Accumulator.GatherParameters.RingCount);
 
 	// Center sample is first ring only when smallest ring first.
 	Accumulator.bIsFirstRing = LARGEST_RINGS_FIRST == 0;
@@ -708,7 +711,11 @@ void AmendAccumulatedRingsWeight(
 void StartRingSamples(inout FGatherAccumulator Accumulator, in FGatherRingInfos RingInfos)
 {
 	Accumulator.bIsFirstRing = RingInfos.bIsFirstRing;
-	Accumulator.BorderingRadius = (RingInfos.RingId + 1.5) * (Accumulator.GatherParameters.KernelRadius * rcp(0.5 + Accumulator.GatherParameters.RingCount)); 
+	
+	// Error in coc radius error introduced by the under sampling of the kernel
+	const float CocRadiusError = CONFIG_COC_RADIUS_CORRECTION ? 1.0 : 0.0;
+
+	Accumulator.BorderingRadius = (RingInfos.RingId + (1.5 + CocRadiusError)) * (Accumulator.GatherParameters.KernelRadius * rcp(0.5 + Accumulator.GatherParameters.RingCount)); 
 
 	#if CONFIG_SGPR_HINTS
 	{
diff --git a/Engine/Shaders/Private/DiaphragmDOF/DOFGatherCommon.ush b/Engine/Shaders/Private/DiaphragmDOF/DOFGatherCommon.ush
--- a/Engine/Shaders/Private/DiaphragmDOF/DOFGatherCommon.ush
+++ b/Engine/Shaders/Private/DiaphragmDOF/DOFGatherCommon.ush
@@ -100,6 +100,9 @@
 	// Background DOF can suffer under sampling on smaller out of focus, because
 	// the gathering kernel is too large.
 	#define CONFIG_KERNEL_SAMPLING_DENSITY_DRIVER 1
+	
+	// Background need coc radius error correction to reduce amount of background leaking.
+	#define CONFIG_COC_RADIUS_CORRECTION 1
 
 #elif DIM_LAYER_PROCESSING == LAYER_PROCESSING_SLIGHT_OUT_OF_FOCUS
 	#define GATHER_ACCUMULATOR_METHOD (GATHER_ACCUMULATOR_METHOD_SIMPLE)
@@ -202,6 +205,18 @@
 #endif
 
 
+// D3D shader compiler is rediculously slow to compile this shader.
+#if SM5_PROFILE
+	#define CONFIG_UNROLL_GATHER 0
+#else
+	#define CONFIG_UNROLL_GATHER 1
+#endif
+
+// Substract the CocRadius error to the sample distance to fight against layering opacity issues caused by undersampling.
+#ifndef CONFIG_COC_RADIUS_CORRECTION
+	#define CONFIG_COC_RADIUS_CORRECTION 0
+#endif
+
 
 //------------------------------------------------------- CONFIG CHECKS
 
@@ -215,11 +230,6 @@
 
 //------------------------------------------------------- CONFIG DERIVED VALUES
 
-// Sets the gather pass's half resolution divisor.
-#define GATHER_HALF_RES_DIVISOR 1
-
-#define GATHER_HALF_RES_DIVISOR_INV (1.0/float(GATHER_HALF_RES_DIVISOR))
-
 // Wether the group size is larger than COC tile size.
 #define GROUP_SIZE_IS_LARGER_THAN_COC_TILE_SIZE 0
 
diff --git a/Engine/Shaders/Private/DiaphragmDOF/DOFGatherKernel.ush b/Engine/Shaders/Private/DiaphragmDOF/DOFGatherKernel.ush
--- a/Engine/Shaders/Private/DiaphragmDOF/DOFGatherKernel.ush
+++ b/Engine/Shaders/Private/DiaphragmDOF/DOFGatherKernel.ush
@@ -203,7 +203,7 @@ void UpdateRuntimeKernelScope(
 	}
 
 	// Compute the randomising range of the kernel offset in pixels.
-	float KernelRandomRange = max((GatherParameters.KernelRadius - GATHER_HALF_RES_DIVISOR * GatherParameters.MaxRingCount) * UnitSampleRadius, 0.0);
+	float KernelRandomRange = max((GatherParameters.KernelRadius - GatherParameters.MaxRingCount) * UnitSampleRadius, 0.0);
 	
 	// No need to do kernel randomization when there is enough samples.
 	if (static_condition(KernelScope.bKenrelHasEnoughSamples))
@@ -289,7 +289,11 @@ void GatherRing(
 	const uint BatchCount = RingSamplePairCount / SamplePairBatchSize;
 
 	// Iterate over the batch of samples.
-	UNROLL
+	#if CONFIG_UNROLL_GATHER
+		UNROLL
+	#else
+		LOOP
+	#endif
 	for (uint BatchId = 0; BatchId < BatchCount; BatchId++)
 	#if defined(WORKARROUND_UE_58394)
 		WORKARROUND_UE_58394 // TODO: nuke that guy once original issue is fixed.
@@ -434,13 +438,6 @@ void GatherRing(
 				// TODO: SampleDistanceSquare for perf reasons.
 				//float SampleDistance = length(RandomKernelPosOffset + SamplePosOffset[k]);
 				//float SampleDistance = length(SamplePosOffset[0]);
-					
-				// Substract the CocRadius error to the sample distance to fight against layering opacity issues caused by undersampling.
-				// TODO: technically, should be done on the Coc radius instead so that it also happens on FGatherAccumulator::CurrentTranslucency 
-				//if (static_condition(KernelScope.bKenrelHasEnoughSamples)) // TODO: bug here.
-				//{
-				//	SampleDistance = max(SampleDistance - KernelScope.CocRadiusError, 0);
-				//}
 
 				// Compute sample's input buffer UV
 				// TODO: VGPR could be optimised.
diff --git a/Engine/Shaders/Private/DiaphragmDOF/DOFGatherPass.usf b/Engine/Shaders/Private/DiaphragmDOF/DOFGatherPass.usf
--- a/Engine/Shaders/Private/DiaphragmDOF/DOFGatherPass.usf
+++ b/Engine/Shaders/Private/DiaphragmDOF/DOFGatherPass.usf
@@ -225,7 +225,7 @@ void GatherMainCS(
 	{
 		// Actually sample the Coc tile buffer.
 		#if GROUP_SIZE_IS_LARGER_THAN_COC_TILE_SIZE
-			const uint2 TilePos = (DispatchThreadId * GATHER_HALF_RES_DIVISOR) / COC_TILE_SIZE;
+			const uint2 TilePos = DispatchThreadId / COC_TILE_SIZE;
 		#else
 			const uint2 TilePos = GroupId;
 		#endif
@@ -249,9 +249,6 @@ void GatherMainCS(
 		#endif
 	}
 	
-	// Coc error accepted for fast gathering, set to 5% of the kernel size.
-	float FastGatheringCocError = GatherParameters.KernelRadius * 0.05;
-
 	// Whether can entirely skeep gathering.
 	bool bSkipGathering = TileSuggestion.bCanEarlyReturn;
 
@@ -386,8 +383,18 @@ void GatherMainCS(
 		}
 
 		#if DEBUG_FAST_ACCUMULATOR
+		{
+			#if CONFIG_KERNEL_SAMPLING_DENSITY_DRIVER == 1
+			if (GatherParameters.bDoHighDensityGatheringAtCenter)
+			{
+				AccumulatorOutput.ForegroundColor.rgb *= float3(1, 0.5, 0.5);
+				AccumulatorOutput.BackgroundColor.rgb *= float3(1, 0.5, 0.5);
+			}
+			#endif
+
 			AccumulatorOutput.ForegroundColor.rgb *= float3(1, 0.5, 0.5);
 			AccumulatorOutput.BackgroundColor.rgb *= float3(1, 0.5, 0.5);
+		}
 		#endif
 		
 		#if 1 // Lower VGPR footprint.
