Cherry-pick 4152166: Implements CONFIG_MAX_SAMPLING_DENSITY_CHANGES to address when doing macro photography-like shot.

diff --git a/Engine/Shaders/Private/DiaphragmDOF/DOFGatherCommon.ush b/Engine/Shaders/Private/DiaphragmDOF/DOFGatherCommon.ush
--- a/Engine/Shaders/Private/DiaphragmDOF/DOFGatherCommon.ush
+++ b/Engine/Shaders/Private/DiaphragmDOF/DOFGatherCommon.ush
@@ -217,6 +217,10 @@
 	#define CONFIG_COC_RADIUS_CORRECTION 0
 #endif
 
+/** Defines the maximum number of gathering sampling density changes supported in the kernel. */
+#ifndef CONFIG_MAX_SAMPLING_DENSITY_CHANGES
+	#define CONFIG_MAX_SAMPLING_DENSITY_CHANGES 3
+#endif
 
 //------------------------------------------------------- CONFIG CHECKS
 
@@ -346,7 +350,7 @@ struct FGatherInputParameters
 	uint KernelSamplingDensityMode;
 
 	// Runtime toggle when KernelSamplingDensityMode == KERNEL_DENSITY_HIGHER_IN_CENTER_DISK 
-	bool bDoHighDensityGatheringAtCenter;
+	bool bDoHighDensityGatheringAtCenter[CONFIG_MAX_SAMPLING_DENSITY_CHANGES];
 };
 
 /** Informations about a ring. */
diff --git a/Engine/Shaders/Private/DiaphragmDOF/DOFGatherKernel.ush b/Engine/Shaders/Private/DiaphragmDOF/DOFGatherKernel.ush
--- a/Engine/Shaders/Private/DiaphragmDOF/DOFGatherKernel.ush
+++ b/Engine/Shaders/Private/DiaphragmDOF/DOFGatherKernel.ush
@@ -572,39 +572,47 @@ void GatherSamplesAndResolve(
 
 			bIsFirstRing = false;
 		} // for (int RingId)
-
-		BRANCH
-		if (GatherParameters.bDoHighDensityGatheringAtCenter)
+		
+		UNROLL
+		for (uint SamplingDensityChangeId = 0; SamplingDensityChangeId < CONFIG_MAX_SAMPLING_DENSITY_CHANGES; SamplingDensityChangeId++)
 		{
-			// Decrease kernel size, and changes the already accumulated weight.
+			BRANCH
+			if (GatherParameters.bDoHighDensityGatheringAtCenter[SamplingDensityChangeId])
 			{
-				const int LargerRingCount = GatherParameters.MaxRingCount - LowerRingId;
+				// Decrease kernel size, and changes the already accumulated weight.
+				{
+					const int LargerRingCount = GatherParameters.MaxRingCount - LowerRingId;
 
-				const float LargeKernelRadius = 1 + 2 * GatherParameters.MaxRingCount;
-				const float SmallKernelRadius = 1 + 2 * (GatherParameters.MaxRingCount - LargerRingCount);
+					const float LargeKernelRadius = 1 + 2 * GatherParameters.MaxRingCount;
+					const float SmallKernelRadius = 1 + 2 * (GatherParameters.MaxRingCount - LargerRingCount);
 
-				const float RadiusDownScaleFactor = SmallKernelRadius / LargeKernelRadius;
+					const float RadiusDownScaleFactor = SmallKernelRadius / LargeKernelRadius;
 
-				const float AccumulatedWeightMultiplier = rcp(RadiusDownScaleFactor * RadiusDownScaleFactor);
+					const float AccumulatedWeightMultiplier = rcp(RadiusDownScaleFactor * RadiusDownScaleFactor);
 				
-				GatherParameters.KernelRadius *= RadiusDownScaleFactor;
+					GatherParameters.KernelRadius *= RadiusDownScaleFactor;
 
-				AmendAccumulatedRingsWeight(/* inout */ Accumulator, GatherParameters, AccumulatedWeightMultiplier);
-				AmendAccumulatedRingsWeight(/* inout */ FastAccumulator, AccumulatedWeightMultiplier);
-				UpdateRuntimeKernelScope(GatherParameters, /* inout */ KernelScope);
-			}
+					AmendAccumulatedRingsWeight(/* inout */ Accumulator, GatherParameters, AccumulatedWeightMultiplier);
+					AmendAccumulatedRingsWeight(/* inout */ FastAccumulator, AccumulatedWeightMultiplier);
+					UpdateRuntimeKernelScope(GatherParameters, /* inout */ KernelScope);
+				}
 
-			UNROLL
-			for (int RingId2 = GatherParameters.MaxRingCount - 1; RingId2 >= LowerRingId; RingId2--)
+				UNROLL
+				for (int RingId2 = GatherParameters.MaxRingCount - 1; RingId2 >= LowerRingId; RingId2--)
+				{
+					GatherRing(
+						GatherParameters,
+						KernelScope,
+						uint(RingId2),
+						/* bIsFirstRing = */ false,
+						/* inout */ Accumulator,
+						/* inout */ FastAccumulator);
+				} // for (int RingId)
+			}
+			else
 			{
-				GatherRing(
-					GatherParameters,
-					KernelScope,
-					uint(RingId2),
-					/* bIsFirstRing = */ false,
-					/* inout */ Accumulator,
-					/* inout */ FastAccumulator);
-			} // for (int RingId)
+				break;
+			}
 		}
 
 		UNROLL
diff --git a/Engine/Shaders/Private/DiaphragmDOF/DOFGatherPass.usf b/Engine/Shaders/Private/DiaphragmDOF/DOFGatherPass.usf
--- a/Engine/Shaders/Private/DiaphragmDOF/DOFGatherPass.usf
+++ b/Engine/Shaders/Private/DiaphragmDOF/DOFGatherPass.usf
@@ -177,7 +177,9 @@ void GatherMainCS(
 		GatherParameters.RingCount = GatherParameters.MaxRingCount;
 		
 		GatherParameters.KernelSamplingDensityMode = KERNEL_DENSITY_CONSTANT;
-		GatherParameters.bDoHighDensityGatheringAtCenter = false;
+		UNROLL
+		for (uint i = 0; i < CONFIG_MAX_SAMPLING_DENSITY_CHANGES; i++)
+			GatherParameters.bDoHighDensityGatheringAtCenter[i] = false;
 		
 		GatherParameters.ConsiderCocRadiusAffineTransformation0 = kContantlyPassingAffineTransformation;
 		GatherParameters.ConsiderCocRadiusAffineTransformation1 = kContantlyPassingAffineTransformation;
@@ -362,14 +364,24 @@ void GatherMainCS(
 	else
 	{
 		// Drives kernel sampling density.
-		#if CONFIG_KERNEL_SAMPLING_DENSITY_DRIVER == 1
+		#if CONFIG_KERNEL_SAMPLING_DENSITY_DRIVER > 0
 		{
 			#if DIM_LAYER_PROCESSING != LAYER_PROCESSING_BACKGROUND_ONLY
 				#error Kernel sampling density driver only available for background.
 			#endif
 
 			GatherParameters.KernelSamplingDensityMode = KERNEL_DENSITY_HIGHER_IN_CENTER_DISK;
-			GatherParameters.bDoHighDensityGatheringAtCenter = TileSuggestion.MaxCocRadiusAbs > 8.0 && TileSuggestion.MinCocRadiusAbs < 0.5 * TileSuggestion.MaxCocRadiusAbs;
+
+			UNROLL
+			for (uint i = 0; i < CONFIG_MAX_SAMPLING_DENSITY_CHANGES; i++)
+			{
+				const float MinConvolutionRadius = float(8 << i);
+				const float SmallestConvolutionRadius = rcp(float(2 << i)) * TileSuggestion.MaxCocRadiusAbs;
+
+				GatherParameters.bDoHighDensityGatheringAtCenter[i] = (
+					TileSuggestion.MaxCocRadiusAbs > MinConvolutionRadius &&
+					TileSuggestion.MinCocRadiusAbs < SmallestConvolutionRadius);
+			}
 		}
 		#endif
 
@@ -384,11 +396,15 @@ void GatherMainCS(
 
 		#if DEBUG_FAST_ACCUMULATOR
 		{
-			#if CONFIG_KERNEL_SAMPLING_DENSITY_DRIVER == 1
-			if (GatherParameters.bDoHighDensityGatheringAtCenter)
+			#if CONFIG_KERNEL_SAMPLING_DENSITY_DRIVER > 0
+			UNROLL
+			for (uint i = 0; i < CONFIG_MAX_SAMPLING_DENSITY_CHANGES; i++)
 			{
-				AccumulatorOutput.ForegroundColor.rgb *= float3(1, 0.5, 0.5);
-				AccumulatorOutput.BackgroundColor.rgb *= float3(1, 0.5, 0.5);
+				if (GatherParameters.bDoHighDensityGatheringAtCenter[i])
+				{
+					AccumulatorOutput.ForegroundColor.rgb *= float3(1, 0.5, 0.5);
+					AccumulatorOutput.BackgroundColor.rgb *= float3(1, 0.5, 0.5);
+				}
 			}
 			#endif
 
