Removed MinRoughness. Added SpecularScale in it's place.
Connected directional light SourceAngle to specular and added SourceSoftAngle.

diff --git a/Engine/Shaders/Private/BRDF.ush b/Engine/Shaders/Private/BRDF.ush
--- a/Engine/Shaders/Private/BRDF.ush
+++ b/Engine/Shaders/Private/BRDF.ush
@@ -32,48 +32,51 @@ void Init( inout BxDFContext Context, half3 N, half3 V, half3 L )
 // [ de Carpentier 2017, "Decima Engine: Advances in Lighting and AA" ]
 void SphereMaxNoH( inout BxDFContext Context, float SinAlpha, bool bNewtonIteration )
 {
-	float CosAlpha = sqrt( 1 - Pow2( SinAlpha ) );
-	
-	float RoL = 2 * Context.NoL * Context.NoV - Context.VoL;
-	if( RoL >= CosAlpha )
-	{
-		Context.NoH = 1;
-		Context.VoH = abs( Context.NoV );
-	}
-	else
+	if( SinAlpha > 0 )
 	{
-		float rInvLengthT = SinAlpha * rsqrt( 1 - RoL*RoL );
-		float NoTr = rInvLengthT * ( Context.NoV - RoL * Context.NoL );
-		float VoTr = rInvLengthT * ( 2 * Context.NoV*Context.NoV - 1 - RoL * Context.VoL );
-
-		if( bNewtonIteration && SinAlpha != 0 )
+		float CosAlpha = sqrt( 1 - Pow2( SinAlpha ) );
+	
+		float RoL = 2 * Context.NoL * Context.NoV - Context.VoL;
+		if( RoL >= CosAlpha )
 		{
-			// dot( cross(N,L), V )
-			float NxLoV = sqrt( saturate( 1 - Pow2(Context.NoL) - Pow2(Context.NoV) - Pow2(Context.VoL) + 2 * Context.NoL * Context.NoV * Context.VoL ) );
-
-			float NoBr = rInvLengthT * NxLoV;
-			float VoBr = rInvLengthT * NxLoV * 2 * Context.NoV;
-			float NoLVTr = Context.NoL * CosAlpha + Context.NoV + NoTr;
-			float VoLVTr = Context.VoL * CosAlpha + 1   + VoTr;
-
-			float p = NoBr   * VoLVTr;
-			float q = NoLVTr * VoLVTr;
-			float s = VoBr   * NoLVTr;
-
-			float xNum = q * ( -0.5 * p + 0.25 * VoBr * NoLVTr );
-			float xDenom = p*p + s * (s - 2*p) + NoLVTr * ( (Context.NoL * CosAlpha + Context.NoV) * Pow2(VoLVTr) + q * (-0.5 * (VoLVTr + Context.VoL * CosAlpha) - 0.5) );
-			float TwoX1 = 2 * xNum / ( Pow2(xDenom) + Pow2(xNum) );
-			float SinTheta = TwoX1 * xDenom;
-			float CosTheta = 1.0 - TwoX1 * xNum;
-			NoTr = CosTheta * NoTr + SinTheta * NoBr;
-			VoTr = CosTheta * VoTr + SinTheta * VoBr;
+			Context.NoH = 1;
+			Context.VoH = abs( Context.NoV );
 		}
+		else
+		{
+			float rInvLengthT = SinAlpha * rsqrt( 1 - RoL*RoL );
+			float NoTr = rInvLengthT * ( Context.NoV - RoL * Context.NoL );
+			float VoTr = rInvLengthT * ( 2 * Context.NoV*Context.NoV - 1 - RoL * Context.VoL );
+
+			if( bNewtonIteration )
+			{
+				// dot( cross(N,L), V )
+				float NxLoV = sqrt( saturate( 1 - Pow2(Context.NoL) - Pow2(Context.NoV) - Pow2(Context.VoL) + 2 * Context.NoL * Context.NoV * Context.VoL ) );
+
+				float NoBr = rInvLengthT * NxLoV;
+				float VoBr = rInvLengthT * NxLoV * 2 * Context.NoV;
+				float NoLVTr = Context.NoL * CosAlpha + Context.NoV + NoTr;
+				float VoLVTr = Context.VoL * CosAlpha + 1   + VoTr;
+
+				float p = NoBr   * VoLVTr;
+				float q = NoLVTr * VoLVTr;
+				float s = VoBr   * NoLVTr;
+
+				float xNum = q * ( -0.5 * p + 0.25 * VoBr * NoLVTr );
+				float xDenom = p*p + s * (s - 2*p) + NoLVTr * ( (Context.NoL * CosAlpha + Context.NoV) * Pow2(VoLVTr) + q * (-0.5 * (VoLVTr + Context.VoL * CosAlpha) - 0.5) );
+				float TwoX1 = 2 * xNum / ( Pow2(xDenom) + Pow2(xNum) );
+				float SinTheta = TwoX1 * xDenom;
+				float CosTheta = 1.0 - TwoX1 * xNum;
+				NoTr = CosTheta * NoTr + SinTheta * NoBr;
+				VoTr = CosTheta * VoTr + SinTheta * VoBr;
+			}
 
-		Context.NoL = Context.NoL * CosAlpha + NoTr;
-		Context.VoL = Context.VoL * CosAlpha + VoTr;
-		float InvLenH = rsqrt( 2 + 2 * Context.VoL );
-		Context.NoH = saturate( ( Context.NoL + Context.NoV ) * InvLenH );
-		Context.VoH = saturate( InvLenH + InvLenH * Context.VoL );
+			Context.NoL = Context.NoL * CosAlpha + NoTr;
+			Context.VoL = Context.VoL * CosAlpha + VoTr;
+			float InvLenH = rsqrt( 2 + 2 * Context.VoL );
+			Context.NoH = saturate( ( Context.NoL + Context.NoV ) * InvLenH );
+			Context.VoH = saturate( InvLenH + InvLenH * Context.VoL );
+		}
 	}
 }
 #else
diff --git a/Engine/Shaders/Private/DeferredLightPixelShaders.usf b/Engine/Shaders/Private/DeferredLightPixelShaders.usf
--- a/Engine/Shaders/Private/DeferredLightPixelShaders.usf
+++ b/Engine/Shaders/Private/DeferredLightPixelShaders.usf
@@ -21,7 +21,7 @@ FDeferredLightData SetupLightDataForStandardDeferred()
 	LightData.LightTangent = DeferredLightUniforms.NormalizedLightTangent;
 	LightData.SpotAnglesAndSourceRadius = float4(DeferredLightUniforms.SpotAngles, DeferredLightUniforms.SourceRadius, DeferredLightUniforms.SourceLength);
     LightData.SoftSourceRadius = DeferredLightUniforms.SoftSourceRadius;
-	LightData.MinRoughness = DeferredLightUniforms.MinRoughness;
+	LightData.SpecularScale = DeferredLightUniforms.SpecularScale;
 	LightData.ContactShadowLength = abs(DeferredLightUniforms.ContactShadowLength);
 	LightData.ContactShadowLengthInWS = DeferredLightUniforms.ContactShadowLength < 0.0f;
 	LightData.DistanceFadeMAD = DeferredLightUniforms.DistanceFadeMAD;
diff --git a/Engine/Shaders/Private/DeferredLightingCommon.ush b/Engine/Shaders/Private/DeferredLightingCommon.ush
--- a/Engine/Shaders/Private/DeferredLightingCommon.ush
+++ b/Engine/Shaders/Private/DeferredLightingCommon.ush
@@ -26,7 +26,7 @@ struct FDeferredLightData
 	float3 LightTangent;
     float SoftSourceRadius;
 	float4 SpotAnglesAndSourceRadius;
-	float MinRoughness;
+	float SpecularScale;
 	float ContactShadowLength;
 	float2 DistanceFadeMAD;
 	float4 ShadowMapChannelMask;
@@ -269,7 +269,7 @@ float4 GetDynamicLighting(float3 WorldPosition, float3 CameraVector, FGBufferDat
 	float3 V = -CameraVector;
 	float3 N = GBuffer.WorldNormal;
 	float3 L = LightData.LightDirection;	// Already normalized
-	float3 ToLight = L * GBuffer.Depth;
+	float3 ToLight = L;
 	
 	float LightMask = 1;
 	if (LightData.bRadialLight)
@@ -337,6 +337,8 @@ float4 GetDynamicLighting(float3 WorldPosition, float3 CameraVector, FGBufferDat
 					Lighting = IntegrateBxDF( GBuffer, N, V, Capsule, Shadow, LightData.bInverseSquared );
 				#endif
 			}
+
+			Lighting.Specular *= LightData.SpecularScale;
 				
 			LightAccumulator_Add( LightAccumulator, Lighting.Diffuse + Lighting.Specular, Lighting.Diffuse, LightColor * LightMask * Shadow.SurfaceShadow, bNeedsSeparateSubsurfaceLightAccumulation );
 			LightAccumulator_Add( LightAccumulator, Lighting.Transmission, Lighting.Transmission, LightColor * LightMask * Shadow.TransmissionShadow, bNeedsSeparateSubsurfaceLightAccumulation );
diff --git a/Engine/Shaders/Private/ForwardLightingCommon.ush b/Engine/Shaders/Private/ForwardLightingCommon.ush
--- a/Engine/Shaders/Private/ForwardLightingCommon.ush
+++ b/Engine/Shaders/Private/ForwardLightingCommon.ush
@@ -103,12 +103,10 @@ float3 GetForwardDirectLighting(uint GridIndex, float3 WorldPosition, float3 Cam
 
 	float3 DirectLighting = 0;
 
-	// Prevent 0 Roughness which causes NaNs in Vis_SmithJointApprox
-	float MinRoughness = .04f;
-
+	float SpecularScale = 1;
 #if TRANSLUCENCY_ANY_VOLUMETRIC
 	// No specular on volumetric translucency lighting modes
-	MinRoughness = 1.0f;
+	SpecularScale = 0;
 #endif
 
 	const FDirectionalLightData DirectionalLightData = GetDirectionalLightData(EyeIndex);
@@ -121,7 +119,7 @@ float3 GetForwardDirectLighting(uint GridIndex, float3 WorldPosition, float3 Cam
 		LightData.LightDirection = DirectionalLightData.DirectionalLightDirection;
 		LightData.DistanceFadeMAD = DirectionalLightData.DirectionalLightDistanceFadeMAD;
 		LightData.bRadialLight = false;
-		LightData.MinRoughness = MinRoughness;
+		LightData.SpecularScale = SpecularScale;
 
 		LightData.ShadowedBits = (DirectionalLightData.DirectionalLightShadowMapChannelMask & 0xFF) != 0 ? 1 : 0;
 		// Static shadowing uses ShadowMapChannel, dynamic shadows are packed into light attenuation using PreviewShadowMapChannel
@@ -177,7 +175,7 @@ float3 GetForwardDirectLighting(uint GridIndex, float3 WorldPosition, float3 Cam
 		LightData.bInverseSquared = LightData.LightColorAndFalloffExponent.w == 0;
 		LightData.bRadialLight = true;
 		LightData.bSpotLight = LightData.SpotAnglesAndSourceRadius.x > -2.0f;
-		LightData.MinRoughness = MinRoughness;
+		LightData.SpecularScale = SpecularScale;
 
 		uint PackedShadowMapChannelMask = asuint(LocalLight.LightDirectionAndShadowMask.w);
 		LightData.ShadowedBits = (PackedShadowMapChannelMask & 0xFF) != 0 ? 1 : 0;
diff --git a/Engine/Shaders/Private/ShadingModels.ush b/Engine/Shaders/Private/ShadingModels.ush
--- a/Engine/Shaders/Private/ShadingModels.ush
+++ b/Engine/Shaders/Private/ShadingModels.ush
@@ -69,8 +69,9 @@ struct FShadowTerms
 
 float New_a2( float a2, float SinAlpha, float VoH )
 {
-	//return a2 + 0.25 * SinAlpha * (3.0 * sqrtFast(a2) + SinAlpha) / ( VoH + 0.001 );
-	return a2 + 0.25 * SinAlpha * ( saturate(12 * a2 + 0.125) + SinAlpha ) / ( VoH + 0.001 );
+	return a2 + 0.25 * SinAlpha * (3.0 * sqrtFast(a2) + SinAlpha) / ( VoH + 0.001 );
+	//return a2 + 0.25 * SinAlpha * ( saturate(12 * a2 + 0.125) + SinAlpha ) / ( VoH + 0.001 );
+	//return a2 + 0.25 * SinAlpha * ( a2 * 2 + 1 + SinAlpha ) / ( VoH + 0.001 );
 }
 
 float EnergyNormalization( inout float a2, float VoH, FAreaLight AreaLight )
@@ -78,7 +79,7 @@ float EnergyNormalization( inout float a2, float VoH, FAreaLight AreaLight )
 	if( AreaLight.SphereSinAlphaSoft > 0 )
 	{
 		// Modify Roughness
-		a2 = a2 + Pow2( AreaLight.SphereSinAlphaSoft ) / ( VoH * 3.6 + 0.4 );
+		a2 = saturate( a2 + Pow2( AreaLight.SphereSinAlphaSoft ) / ( VoH * 3.6 + 0.4 ) );
 	}
 
 	float Sphere_a2 = a2;
diff --git a/Engine/Shaders/Private/TiledDeferredLightShaders.usf b/Engine/Shaders/Private/TiledDeferredLightShaders.usf
--- a/Engine/Shaders/Private/TiledDeferredLightShaders.usf
+++ b/Engine/Shaders/Private/TiledDeferredLightShaders.usf
@@ -256,10 +256,10 @@ void TiledDeferredLightingMain(
 			LightData.LightPositionAndInvRadius = TiledDeferred.LightPositionAndInvRadius[LightIndex];
 			LightData.LightColorAndFalloffExponent = TiledDeferred.LightColorAndFalloffExponent[LightIndex];
 			{
-				float4 Value = TiledDeferred2.LightDirectionAndSpotlightMaskAndMinRoughness[LightIndex];
+				float4 Value = TiledDeferred2.LightDirectionAndSpotlightMaskAndSpecularScale[LightIndex];
 				LightData.LightDirection = Value.xyz;
 				LightData.bSpotLight = Value.w > 0;
-				LightData.MinRoughness = abs(Value.w);
+				LightData.SpecularScale = abs(Value.w);
 			}
 			LightData.SpotAnglesAndSourceRadius = TiledDeferred2.SpotAnglesAndSourceRadiusAndSimpleLighting[LightIndex];
 			LightData.ShadowMapChannelMask = TiledDeferred2.ShadowMapChannelMask[LightIndex];
diff --git a/Engine/Shaders/Private/VolumetricFog.usf b/Engine/Shaders/Private/VolumetricFog.usf
--- a/Engine/Shaders/Private/VolumetricFog.usf
+++ b/Engine/Shaders/Private/VolumetricFog.usf
@@ -209,7 +209,7 @@ void InjectShadowedLocalLightPS(
         LightData.LightTangent = DeferredLightUniforms.NormalizedLightTangent;
 		LightData.SpotAnglesAndSourceRadius = float4(DeferredLightUniforms.SpotAngles, DeferredLightUniforms.SourceRadius, DeferredLightUniforms.SourceLength);
         LightData.SoftSourceRadius = DeferredLightUniforms.SoftSourceRadius;
-		LightData.MinRoughness = DeferredLightUniforms.MinRoughness;
+		LightData.SpecularScale = DeferredLightUniforms.SpecularScale;
 		LightData.ContactShadowLength = abs(DeferredLightUniforms.ContactShadowLength);
 		LightData.ContactShadowLengthInWS = DeferredLightUniforms.ContactShadowLength < 0.0f;
 		LightData.DistanceFadeMAD = DeferredLightUniforms.DistanceFadeMAD;
diff --git a/Engine/Source/Runtime/Engine/Classes/Components/DirectionalLightComponent.h b/Engine/Source/Runtime/Engine/Classes/Components/DirectionalLightComponent.h
--- a/Engine/Source/Runtime/Engine/Classes/Components/DirectionalLightComponent.h
+++ b/Engine/Source/Runtime/Engine/Classes/Components/DirectionalLightComponent.h
@@ -116,12 +116,18 @@ class ENGINE_API UDirectionalLightComponent : public ULightComponent
 	float DistanceFieldShadowDistance;
 
 	/** 
-	 * Light source angle in degrees, used for dynamic shadowing methods. 
-	 * Currently only Ray Traced Distance Field shadows and Capsule shadows support area shadows, and therefore make use of LightSourceAngle.
+	 * Angle subtended by light source in degrees (also known as angular diameter).
+	 * Defaults to 0.5357 which is the angle for our sun.
 	 */
-	UPROPERTY(EditAnywhere, BlueprintReadOnly, Category=Light, meta=(UIMin = "0", UIMax = "5"))
+	UPROPERTY(EditAnywhere, BlueprintReadOnly, Category=Light, meta=(UIMin = "0", UIMax = "5"), DisplayName = "Source Angle")
 	float LightSourceAngle;
 
+	/** 
+	 * Angle subtended by soft light source in degrees.
+	 */
+	UPROPERTY(EditAnywhere, BlueprintReadOnly, Category=Light, meta=(UIMin = "0", UIMax = "5"), DisplayName = "Source Soft Angle")
+	float LightSourceSoftAngle;
+
 	/** Determines how far shadows can be cast, in world units.  Larger values increase the shadowing cost. */
 	UPROPERTY(EditAnywhere, AdvancedDisplay, BlueprintReadOnly, Category=DistanceFieldShadows, meta=(UIMin = "1000", UIMax = "100000"), DisplayName = "DistanceField Trace Distance")
 	float TraceDistance;
diff --git a/Engine/Source/Runtime/Engine/Classes/Components/LightComponent.h b/Engine/Source/Runtime/Engine/Classes/Components/LightComponent.h
--- a/Engine/Source/Runtime/Engine/Classes/Components/LightComponent.h
+++ b/Engine/Source/Runtime/Engine/Classes/Components/LightComponent.h
@@ -68,8 +68,15 @@ class ENGINE_API ULightComponent : public ULightComponentBase
 	int32 PreviewShadowMapChannel;
 	
 	/** Min roughness effective for this light. Used for softening specular highlights. */
-	UPROPERTY(EditAnywhere, BlueprintReadOnly, Category=Light, AdvancedDisplay, meta=(UIMin = "0.08", UIMax = "1.0"))
-	float MinRoughness;
+	UPROPERTY()
+	float MinRoughness_DEPRECATED;
+
+	/** 
+	 * Multiplier on specular highlights. Use only with great care! Any value besides 1 is not physical!
+	 * Can be used to artistically remove highlights mimicking polarizing filters or photo touch up.
+	 */
+	UPROPERTY(EditAnywhere, BlueprintReadOnly, Category=Light, AdvancedDisplay, meta=(UIMin = "0", UIMax = "1"))
+	float SpecularScale;
 
 	/** 
 	 * Scales the resolution of shadowmaps used to shadow this light.  By default shadowmap resolution is chosen based on screen size of the caster. 
diff --git a/Engine/Source/Runtime/Engine/Private/Components/DirectionalLightComponent.cpp b/Engine/Source/Runtime/Engine/Private/Components/DirectionalLightComponent.cpp
--- a/Engine/Source/Runtime/Engine/Private/Components/DirectionalLightComponent.cpp
+++ b/Engine/Source/Runtime/Engine/Private/Components/DirectionalLightComponent.cpp
@@ -89,6 +89,9 @@ public:
 	/** Light source angle in degrees. */
 	float LightSourceAngle;
 
+	/** Light source angle in degrees. */
+	float LightSourceSoftAngle;
+
 	/** Determines how far shadows can be cast, in world units.  Larger values increase the shadowing cost. */
 	float TraceDistance;
 
@@ -106,6 +109,7 @@ public:
 		bUseInsetShadowsForMovableObjects(Component->bUseInsetShadowsForMovableObjects),
 		DistanceFieldShadowDistance(Component->bUseRayTracedDistanceFieldShadows ? Component->DistanceFieldShadowDistance : 0),
 		LightSourceAngle(Component->LightSourceAngle),
+		LightSourceSoftAngle(Component->LightSourceSoftAngle),
 		TraceDistance(FMath::Clamp(Component->TraceDistance, 1000.0f, 1000000.0f))
 	{
 		LightShaftOverrideDirection.Normalize();
@@ -167,11 +171,10 @@ public:
 		LightParameters.NormalizedLightTangent = -GetDirection();
 
 		LightParameters.SpotAngles = FVector2D(0, 0);
-		LightParameters.LightSourceRadius = 0.0f;
-		LightParameters.LightSoftSourceRadius = 0.0f;
+		LightParameters.SpecularScale = SpecularScale;
+		LightParameters.LightSourceRadius = FMath::Sin( 0.5f * FMath::DegreesToRadians( LightSourceAngle ) );
+		LightParameters.LightSoftSourceRadius = FMath::Sin( 0.5f * FMath::DegreesToRadians( LightSourceSoftAngle ) );
 		LightParameters.LightSourceLength = 0.0f;
-		// Prevent 0 Roughness which causes NaNs in Vis_SmithJointApprox
-		LightParameters.LightMinRoughness = FMath::Max(MinRoughness, .02f);
 		LightParameters.SourceTexture = GWhiteTexture;
 	}
 
@@ -727,7 +730,8 @@ UDirectionalLightComponent::UDirectionalLightComponent(const FObjectInitializer&
 	DistanceFieldShadowDistance = 30000.0f;
 	TraceDistance = 10000.0f;
 	FarShadowDistance = 300000.0f;
-	LightSourceAngle = 1;
+	LightSourceAngle = 0.5357f;		// Angle of earth's sun
+	LightSourceSoftAngle = 0.0f;
 
 	DynamicShadowCascades = 3;
 	CascadeDistributionExponent = 3.0f;
diff --git a/Engine/Source/Runtime/Engine/Private/Components/LightComponent.cpp b/Engine/Source/Runtime/Engine/Private/Components/LightComponent.cpp
--- a/Engine/Source/Runtime/Engine/Private/Components/LightComponent.cpp
+++ b/Engine/Source/Runtime/Engine/Private/Components/LightComponent.cpp
@@ -216,7 +216,7 @@ FLightSceneProxy::FLightSceneProxy(const ULightComponent* InLightComponent)
 	, ShadowSharpen(InLightComponent->ShadowSharpen)
 	, ContactShadowLength(InLightComponent->ContactShadowLength)
 	, bContactShadowLengthInWS(InLightComponent->ContactShadowLengthInWS ? true : false)
-	, MinRoughness(InLightComponent->MinRoughness)
+	, SpecularScale(InLightComponent->SpecularScale)
 	, LightGuid(InLightComponent->LightGuid)
 	, IESTexture(0)
 	, bMovable(InLightComponent->IsMovable())
@@ -364,7 +364,7 @@ ULightComponent::ULightComponent(const FObjectInitializer& ObjectInitializer)
 
 	LightFunctionFadeDistance = 100000.0f;
 	DisabledBrightness = 0.5f;
-	MinRoughness = 0.08f;
+	SpecularScale = 1.0f;
 
 	bEnableLightShaftBloom = false;
 	BloomScale = .2f;
@@ -439,6 +439,12 @@ void ULightComponent::Serialize(FArchive& Ar)
 			GLightComponentsWithLegacyBuildData.AddAnnotation(this, LegacyLightData);
 		}
 	}
+
+	if( MinRoughness_DEPRECATED == 1.0f )
+	{
+		MinRoughness_DEPRECATED = 0.0f;
+		SpecularScale = 0.0f;
+	}
 }
 
 /**
@@ -546,6 +552,7 @@ void ULightComponent::PostEditChangeProperty(FPropertyChangedEvent& PropertyChan
 	const FString PropertyName = PropertyThatChanged ? PropertyThatChanged->GetName() : TEXT("");
 
 	Intensity = FMath::Max(0.0f, Intensity);
+	SpecularScale = FMath::Clamp( SpecularScale, 0.0f, 1.0f );
 
 	if (HasStaticLighting())
 	{
@@ -561,7 +568,7 @@ void ULightComponent::PostEditChangeProperty(FPropertyChangedEvent& PropertyChan
 		PropertyName != GET_MEMBER_NAME_STRING_CHECKED(ULightComponent, CastDynamicShadows) &&
 		PropertyName != GET_MEMBER_NAME_STRING_CHECKED(ULightComponent, bAffectTranslucentLighting) &&
 		PropertyName != GET_MEMBER_NAME_STRING_CHECKED(ULightComponent, bTransmission) &&
-		PropertyName != GET_MEMBER_NAME_STRING_CHECKED(ULightComponent, MinRoughness) &&
+		PropertyName != GET_MEMBER_NAME_STRING_CHECKED(ULightComponent, SpecularScale) &&
 		PropertyName != GET_MEMBER_NAME_STRING_CHECKED(ULightComponent, LightFunctionMaterial) &&
 		PropertyName != GET_MEMBER_NAME_STRING_CHECKED(ULightComponent, LightFunctionScale) &&
 		PropertyName != GET_MEMBER_NAME_STRING_CHECKED(ULightComponent, LightFunctionFadeDistance) &&
diff --git a/Engine/Source/Runtime/Engine/Private/Components/PointLightComponent.cpp b/Engine/Source/Runtime/Engine/Private/Components/PointLightComponent.cpp
--- a/Engine/Source/Runtime/Engine/Private/Components/PointLightComponent.cpp
+++ b/Engine/Source/Runtime/Engine/Private/Components/PointLightComponent.cpp
@@ -47,11 +47,10 @@ void FPointLightSceneProxy::GetParameters(FLightParameters& LightParameters) con
 	LightParameters.NormalizedLightDirection = -GetDirection();
 	LightParameters.NormalizedLightTangent = ZAxis;
 	LightParameters.SpotAngles = FVector2D( -2.0f, 1.0f );
+	LightParameters.SpecularScale = SpecularScale;
 	LightParameters.LightSourceRadius = SourceRadius;
 	LightParameters.LightSoftSourceRadius = SoftSourceRadius;
 	LightParameters.LightSourceLength = SourceLength;
-	// Prevent 0 Roughness which causes NaNs in Vis_SmithJointApprox
-	LightParameters.LightMinRoughness = FMath::Max(MinRoughness, .04f);
 	LightParameters.SourceTexture = GWhiteTexture;
 }
 
diff --git a/Engine/Source/Runtime/Engine/Private/Components/RectLightComponent.cpp b/Engine/Source/Runtime/Engine/Private/Components/RectLightComponent.cpp
--- a/Engine/Source/Runtime/Engine/Private/Components/RectLightComponent.cpp
+++ b/Engine/Source/Runtime/Engine/Private/Components/RectLightComponent.cpp
@@ -53,10 +53,10 @@ public:
 		LightParameters.NormalizedLightDirection = -GetDirection();
 		LightParameters.NormalizedLightTangent = ZAxis;
 		LightParameters.SpotAngles = FVector2D( -2.0f, 1.0f );
+		LightParameters.SpecularScale = SpecularScale;
 		LightParameters.LightSourceRadius = SourceWidth * 0.5f;
 		LightParameters.LightSoftSourceRadius = 0.0f;
 		LightParameters.LightSourceLength = SourceHeight * 0.5f;
-		LightParameters.LightMinRoughness = 0.0f;
 		LightParameters.SourceTexture = GWhiteTexture;
 		//LightParameters.SourceTexture = SourceTexture ? SourceTexture->Resource : GWhiteTexture;
 	}
diff --git a/Engine/Source/Runtime/Engine/Private/Components/SpotLightComponent.cpp b/Engine/Source/Runtime/Engine/Private/Components/SpotLightComponent.cpp
--- a/Engine/Source/Runtime/Engine/Private/Components/SpotLightComponent.cpp
+++ b/Engine/Source/Runtime/Engine/Private/Components/SpotLightComponent.cpp
@@ -80,11 +80,10 @@ public:
 		LightParameters.NormalizedLightDirection = -GetDirection();
 		LightParameters.NormalizedLightTangent = ZAxis;
 		LightParameters.SpotAngles = FVector2D(CosOuterCone, InvCosConeDifference);
+		LightParameters.SpecularScale = SpecularScale;
 		LightParameters.LightSourceRadius = SourceRadius;
 		LightParameters.LightSoftSourceRadius = SoftSourceRadius;
 		LightParameters.LightSourceLength = SourceLength;
-		// Prevent 0 Roughness which causes NaNs in Vis_SmithJointApprox
-		LightParameters.LightMinRoughness = FMath::Max(MinRoughness, .04f);
 		LightParameters.SourceTexture = GWhiteTexture;
 	}
 
diff --git a/Engine/Source/Runtime/Engine/Public/SceneManagement.h b/Engine/Source/Runtime/Engine/Public/SceneManagement.h
--- a/Engine/Source/Runtime/Engine/Public/SceneManagement.h
+++ b/Engine/Source/Runtime/Engine/Public/SceneManagement.h
@@ -888,10 +888,10 @@ struct FLightParameters
 	FVector		NormalizedLightDirection;
 	FVector		NormalizedLightTangent;
 	FVector2D	SpotAngles;
+	float		SpecularScale;
 	float		LightSourceRadius;
 	float		LightSoftSourceRadius;
 	float		LightSourceLength;
-	float		LightMinRoughness;
 	FTexture*	SourceTexture;
 };
 
@@ -1058,7 +1058,7 @@ public:
 	inline float GetShadowSharpen() const { return ShadowSharpen; }
 	inline float GetContactShadowLength() const { return ContactShadowLength; }
 	inline bool IsContactShadowLengthInWS() const { return bContactShadowLengthInWS; }
-	inline float GetMinRoughness() const { return MinRoughness; }
+	inline float GetSpecularScale() const { return SpecularScale; }
 	inline FVector GetLightFunctionScale() const { return LightFunctionScale; }
 	inline float GetLightFunctionFadeDistance() const { return LightFunctionFadeDistance; }
 	inline float GetLightFunctionDisabledBrightness() const { return LightFunctionDisabledBrightness; }
@@ -1152,8 +1152,8 @@ protected:
 	/** True: length of screen space ray trace for sharp contact shadows is in world space. False: in screen space. */
 	bool bContactShadowLengthInWS : 1;
 
-	/** Min roughness */
-	float MinRoughness;
+	/** Specular scale */
+	float SpecularScale;
 
 	/** The light's persistent shadowing GUID. */
 	FGuid LightGuid;
diff --git a/Engine/Source/Runtime/Renderer/Private/LightRendering.h b/Engine/Source/Runtime/Renderer/Private/LightRendering.h
--- a/Engine/Source/Runtime/Renderer/Private/LightRendering.h
+++ b/Engine/Source/Runtime/Renderer/Private/LightRendering.h
@@ -27,10 +27,10 @@ BEGIN_UNIFORM_BUFFER_STRUCT(FDeferredLightUniformStruct,)
 	UNIFORM_MEMBER(FVector,NormalizedLightDirection)
 	UNIFORM_MEMBER(FVector,NormalizedLightTangent)
 	UNIFORM_MEMBER(FVector2D,SpotAngles)
+	UNIFORM_MEMBER(float,SpecularScale)
 	UNIFORM_MEMBER(float,SourceRadius)
 	UNIFORM_MEMBER(float,SoftSourceRadius)
 	UNIFORM_MEMBER(float,SourceLength)
-	UNIFORM_MEMBER(float,MinRoughness)
 	UNIFORM_MEMBER(float,ContactShadowLength)
 	UNIFORM_MEMBER(FVector2D,DistanceFadeMAD)
 	UNIFORM_MEMBER(FVector4,ShadowMapChannelMask)
@@ -65,10 +65,10 @@ void SetDeferredLightParameters(
 	DeferredLightUniformsValue.NormalizedLightDirection = LightParameters.NormalizedLightDirection;
 	DeferredLightUniformsValue.NormalizedLightTangent = LightParameters.NormalizedLightTangent;
 	DeferredLightUniformsValue.SpotAngles = LightParameters.SpotAngles;
+	DeferredLightUniformsValue.SpecularScale = LightParameters.SpecularScale;
 	DeferredLightUniformsValue.SourceRadius = LightParameters.LightSourceRadius;
 	DeferredLightUniformsValue.SoftSourceRadius = LightParameters.LightSoftSourceRadius;
 	DeferredLightUniformsValue.SourceLength = LightParameters.LightSourceLength;
-	DeferredLightUniformsValue.MinRoughness = LightParameters.LightMinRoughness;
 	DeferredLightUniformsValue.SourceTexture = LightParameters.SourceTexture->TextureRHI;
 
 	const FVector2D FadeParams = LightSceneInfo->Proxy->GetDirectionalLightDistanceFadeParameters(View.GetFeatureLevel(), LightSceneInfo->IsPrecomputedLightingValid(), View.MaxShadowCascades);
@@ -144,10 +144,10 @@ void SetSimpleDeferredLightParameters(
 	DeferredLightUniformsValue.NormalizedLightDirection = FVector(1, 0, 0);
 	DeferredLightUniformsValue.NormalizedLightTangent = FVector(1, 0, 0);
 	DeferredLightUniformsValue.SpotAngles = FVector2D(-2, 1);
+	DeferredLightUniformsValue.SpecularScale = 1.0f;
 	DeferredLightUniformsValue.SourceRadius = 0.0f;
 	DeferredLightUniformsValue.SoftSourceRadius = 0.0f;
 	DeferredLightUniformsValue.SourceLength = 0.0f;
-	DeferredLightUniformsValue.MinRoughness = 0.08f;
 	DeferredLightUniformsValue.ContactShadowLength = 0.0f;
 	DeferredLightUniformsValue.DistanceFadeMAD = FVector2D(0, 0);
 	DeferredLightUniformsValue.ShadowMapChannelMask = FVector4(0, 0, 0, 0);
diff --git a/Engine/Source/Runtime/Renderer/Private/SceneVisibility.cpp b/Engine/Source/Runtime/Renderer/Private/SceneVisibility.cpp
--- a/Engine/Source/Runtime/Renderer/Private/SceneVisibility.cpp
+++ b/Engine/Source/Runtime/Renderer/Private/SceneVisibility.cpp
@@ -3209,9 +3209,7 @@ void FSceneRenderer::PostVisibilityFrameSetup(FILCUpdatePrimTaskData& OutILCTask
 			if( View.bIsReflectionCapture 
 				&& VisibleLightViewInfo.bInViewFrustum
 				&& Proxy->HasStaticLighting() 
-				&& Proxy->GetLightType() != LightType_Directional 
-				// Min roughness is used to hide the specular response of virtual area lights, so skip drawing the source shape when Min Roughness is 1
-				&& Proxy->GetMinRoughness() < 1.0f)
+				&& Proxy->GetLightType() != LightType_Directional )
 			{
 				FVector Origin = Proxy->GetOrigin();
 				FVector ToLight = Origin - View.ViewMatrices.GetViewOrigin();
@@ -3268,6 +3266,8 @@ void FSceneRenderer::PostVisibilityFrameSetup(FILCUpdatePrimTaskData& OutILCTask
 					FVector L = ToLight.GetSafeNormal();
 					Color.A *= FMath::Square( FMath::Clamp( ( (L | LightParameters.NormalizedLightDirection) - LightParameters.SpotAngles.X ) * LightParameters.SpotAngles.Y, 0.0f, 1.0f ) );
 
+					Color.A *= LightParameters.SpecularScale;
+
 					// Rect is one sided
 					if( Proxy->IsRectLight() && (L | LightParameters.NormalizedLightDirection) < 0.0f )
 						continue;
diff --git a/Engine/Source/Runtime/Renderer/Private/TiledDeferredLightRendering.cpp b/Engine/Source/Runtime/Renderer/Private/TiledDeferredLightRendering.cpp
--- a/Engine/Source/Runtime/Renderer/Private/TiledDeferredLightRendering.cpp
+++ b/Engine/Source/Runtime/Renderer/Private/TiledDeferredLightRendering.cpp
@@ -68,7 +68,7 @@ IMPLEMENT_UNIFORM_BUFFER_STRUCT(FTiledDeferredLightData,TEXT("TiledDeferred"));
 
 /** Second constant buffer of light data for tiled deferred. */
 BEGIN_UNIFORM_BUFFER_STRUCT(FTiledDeferredLightData2,)
-	UNIFORM_MEMBER_ARRAY(FVector4,LightDirectionAndSpotlightMaskAndMinRoughness,[GMaxNumTiledDeferredLights])
+	UNIFORM_MEMBER_ARRAY(FVector4,LightDirectionAndSpotlightMaskAndSpecularScale,[GMaxNumTiledDeferredLights])
 	UNIFORM_MEMBER_ARRAY(FVector4,SpotAnglesAndSourceRadiusAndSimpleLighting,[GMaxNumTiledDeferredLights])
 	UNIFORM_MEMBER_ARRAY(FVector4,ShadowMapChannelMask,[GMaxNumTiledDeferredLights])
 END_UNIFORM_BUFFER_STRUCT(FTiledDeferredLightData2)
@@ -217,10 +217,10 @@ public:
 				}
 
 				{
-					// SpotlightMaskAndMinRoughness, >0:Spotlight, MinRoughness = abs();
-					float W = FMath::Max(0.0001f, LightParameters.LightMinRoughness) * ((LightSceneInfo->Proxy->GetLightType() == LightType_Spot) ? 1 : -1);
+					// SignBit:Spotlight, SpecularScale = abs();
+					float W = LightParameters.SpecularScale * ((LightSceneInfo->Proxy->GetLightType() == LightType_Spot) ? 1 : -1);
 
-					LightData2.LightDirectionAndSpotlightMaskAndMinRoughness[LightIndex] = FVector4(LightParameters.NormalizedLightDirection, W);
+					LightData2.LightDirectionAndSpotlightMaskAndSpecularScale[LightIndex] = FVector4(LightParameters.NormalizedLightDirection, W);
 				}
 
 				LightData2.SpotAnglesAndSourceRadiusAndSimpleLighting[LightIndex] = FVector4(
@@ -249,7 +249,7 @@ public:
 				const FSimpleLightPerViewEntry& SimpleLightPerViewData = SimpleLights.GetViewDependentData(SimpleLightIndex, ViewIndex, NumViews);
 				LightData.LightPositionAndInvRadius[LightIndex] = FVector4(SimpleLightPerViewData.Position, 1.0f / FMath::Max(SimpleLight.Radius, KINDA_SMALL_NUMBER));
 				LightData.LightColorAndFalloffExponent[LightIndex] = FVector4(SimpleLight.Color, SimpleLight.Exponent);
-				LightData2.LightDirectionAndSpotlightMaskAndMinRoughness[LightIndex] = FVector4(FVector(1, 0, 0), 0);
+				LightData2.LightDirectionAndSpotlightMaskAndSpecularScale[LightIndex] = FVector4(FVector(1, 0, 0), 0);
 				LightData2.SpotAnglesAndSourceRadiusAndSimpleLighting[LightIndex] = FVector4(-2, 1, 0, 1);
 				LightData2.ShadowMapChannelMask[LightIndex] = FVector4(0, 0, 0, 0);
 			}
