Cherry-pick 4155027: Temporally stabilizes DOF on dithered opaque material slightly out of focus

diff --git a/Engine/Shaders/Private/DiaphragmDOF/DOFDownsample.ush b/Engine/Shaders/Private/DiaphragmDOF/DOFDownsample.ush
--- a/Engine/Shaders/Private/DiaphragmDOF/DOFDownsample.ush
+++ b/Engine/Shaders/Private/DiaphragmDOF/DOFDownsample.ush
@@ -10,9 +10,13 @@
 
 //------------------------------------------------------- COMPILE TIME CONFIG
 
+#if defined(CONFIG_SETUP)
+#define DOWNSAMPLE_COC_OPERATOR 4
+#else
 // Defines the COC operator to choose when downsampling.
 // Same as circle DOF.
 #define DOWNSAMPLE_COC_OPERATOR 2
+#endif
 
 
 //------------------------------------------------------- COMMON FUNCTIONS.
@@ -116,6 +120,14 @@ float ComputeDownsamplingBilateralWeight(
 	// TODO: Tune the ScaleFactor. Looks like to large.
 	float ScaleFactor = 64.0 * CocRadiusMultiplier;
 
+	// not doing abs(OutCocRadius - SampleCocRadius) because: it is fine to leak background on foreground arround geometric
+	// edges because going to be close to the hole filling as well. This allow dither opacity material such as dithered human hair
+	// to keep consistent opacity, thickness, and temporal stability as the Coc changes between slight out focus, foreground
+	// and background.
+	#if defined(CONFIG_SETUP)
+		return saturate(1.0 - (OutCocRadius - SampleCocRadius) * ScaleFactor);
+	#endif
+	
 	// GCN Hint: one subtract with abs() post modifier + one MAD with saturate post modifier.
 	return saturate(1.0 - abs(OutCocRadius - SampleCocRadius) * ScaleFactor);
 }
@@ -144,6 +156,7 @@ void DownsampleSceneColorWithCoc(
 	{
 		OutputWeight += (OutCocRadius * CocRadii[i]) >= 0.0 ? 0.25 : 0.0;
 	}
+	OutputWeight = 1;
 
 	float WeightSum = dot(float4(1, 1, 1, 1), BilateralWeights);
 
@@ -201,6 +214,22 @@ float DownsampleCoc_Wave8x8(float CocRadius, const uint ReductionScale)
 				/* InnerLaneGroupSize = */ 8 * ReductionScale,
 				/* InnerLaneGroupId = */ 0);
 		}
+	#elif DOWNSAMPLE_COC_OPERATOR == 4
+		// Horizontal.
+		float OtherCocRadius = WaveSwapWithinLaneGroup(
+			OutCocRadius,
+			/* LaneGroupSize = */ 2 * ReductionScale,
+			/* InnerLaneGroupSize = */ 1 * ReductionScale,
+			/* InnerLaneGroupId = */ 1);
+		OutCocRadius = min(OutCocRadius, OtherCocRadius);
+		
+		// Vertical.
+		OtherCocRadius = WaveSwapWithinLaneGroup(
+			OutCocRadius,
+			/* LaneGroupSize = */ 16 * ReductionScale,
+			/* InnerLaneGroupSize = */ 8 * ReductionScale,
+			/* InnerLaneGroupId = */ 1);
+		OutCocRadius = min(OutCocRadius, OtherCocRadius);
 	#else
 		#error Unimplemented.
 	#endif
diff --git a/Engine/Shaders/Private/DiaphragmDOF/DOFGatherAccumulator.ush b/Engine/Shaders/Private/DiaphragmDOF/DOFGatherAccumulator.ush
--- a/Engine/Shaders/Private/DiaphragmDOF/DOFGatherAccumulator.ush
+++ b/Engine/Shaders/Private/DiaphragmDOF/DOFGatherAccumulator.ush
@@ -28,6 +28,9 @@ struct FGatherAccumulator
 		float4 Color;
 		float ColorWeight;
 
+		float Opacity;
+		float OpacityWeight;
+
 		// Distance to the closest foreground sample.
 		float DistanceToClosestForeground;
 
@@ -127,6 +130,8 @@ FGatherAccumulator CreateGatherAccumulator(
 	{
 		Accumulator.Color = 0;
 		Accumulator.ColorWeight = 0;
+		Accumulator.Opacity = 0;
+		Accumulator.OpacityWeight = 0;
 		Accumulator.DistanceToClosestForeground = EXTREMELY_LARGE_COC_RADIUS;
 	}
 	#endif
@@ -226,9 +231,6 @@ void AccumulateSample(
 	in FGatherSample A,
 	in FGatherSampleDerivedParameters DerivedA)
 {
-	float ColorWeight = A.Intersection * DerivedA.Weight * DerivedA.IsConsidered;
-	float OpacityWeight = A.Intersection;
-	
 	// TODO: soften that.
 	if (Accumulator.LayerProcessing == LAYER_PROCESSING_FOREGROUND_HOLE_FILLING &&
 		A.CocRadius < -(Accumulator.GatherParameters.MaxRecombineAbsCocRadius - 1))
@@ -237,11 +239,14 @@ void AccumulateSample(
 	}
 	else
 	{
+		float ColorWeight = A.Intersection * DerivedA.Weight * DerivedA.IsConsidered;
+		float OpacityWeight = A.Intersection;
+	
 		Accumulator.Color += A.Color * ColorWeight;
 		Accumulator.ColorWeight += ColorWeight;
 	
-		//Accumulator.Opacity += OpacityWeight * DerivedA.Opacity;
-		//Accumulator.OpacityWeight += OpacityWeight;
+		Accumulator.Opacity += OpacityWeight * DerivedA.Opacity;
+		Accumulator.OpacityWeight += OpacityWeight;
 	}
 }
 
@@ -322,7 +327,8 @@ void ResolveAccumulatorOutputs(
 	Output = CreateAccumulatorOutput();
 	
 	float4 NormalizedColor = Accumulator.Color * SafeRcp(Accumulator.ColorWeight);
-	
+	float NormalizedOpacity = Accumulator.Opacity * SafeRcp(Accumulator.OpacityWeight);
+
 	#if DIM_LAYER_PROCESSING == LAYER_PROCESSING_FOREGROUND_HOLE_FILLING
 	if (Accumulator.DistanceToClosestForeground != EXTREMELY_LARGE_COC_RADIUS && Accumulator.ColorWeight > 0)
 	{
@@ -335,7 +341,7 @@ void ResolveAccumulatorOutputs(
 	#elif DIM_LAYER_PROCESSING == LAYER_PROCESSING_SLIGHT_OUT_OF_FOCUS
 	{
 		Output.SlightFocusColor = NormalizedColor;
-		Output.SlightFocusWeight = Accumulator.ColorWeight;
+		Output.SlightFocusOpacity = NormalizedOpacity;
 	}
 	#else
 		#error Unsupported.
diff --git a/Engine/Shaders/Private/DiaphragmDOF/DOFGatherCommon.ush b/Engine/Shaders/Private/DiaphragmDOF/DOFGatherCommon.ush
--- a/Engine/Shaders/Private/DiaphragmDOF/DOFGatherCommon.ush
+++ b/Engine/Shaders/Private/DiaphragmDOF/DOFGatherCommon.ush
@@ -418,9 +418,9 @@ struct FAccumulatorOutput
 	// Coc average and variance for background.
 	float2 BackgroundCocAvgAndSquareAvg;
 
-	// Slight out of focus color and weight.
+	// Slight out of focus color and opacity.
 	float4 SlightFocusColor;
-	float SlightFocusWeight;
+	float SlightFocusOpacity;
 };
 
 
@@ -451,6 +451,6 @@ FAccumulatorOutput CreateAccumulatorOutput()
 	AccumulatorOutput.BackgroundWeight = 0;
 	AccumulatorOutput.BackgroundCocAvgAndSquareAvg = 0;
 	AccumulatorOutput.SlightFocusColor = 0;
-	AccumulatorOutput.SlightFocusWeight = 0;
+	AccumulatorOutput.SlightFocusOpacity = 0;
 	return AccumulatorOutput;
 }
diff --git a/Engine/Shaders/Private/DiaphragmDOF/DOFGatherPass.usf b/Engine/Shaders/Private/DiaphragmDOF/DOFGatherPass.usf
--- a/Engine/Shaders/Private/DiaphragmDOF/DOFGatherPass.usf
+++ b/Engine/Shaders/Private/DiaphragmDOF/DOFGatherPass.usf
@@ -102,7 +102,7 @@ void StoreAccumulatorOutput(in uint2 OutputPixelPosition, in FAccumulatorOutput
 		#elif DIM_LAYER_PROCESSING == LAYER_PROCESSING_SLIGHT_OUT_OF_FOCUS
 			Output0[OutputPixelPosition] = float4(
 				AccumulatorOutput.SlightFocusColor.rgb,
-				AccumulatorOutput.SlightFocusWeight);
+				AccumulatorOutput.SlightFocusOpacity);
 		
 		#else
 			#error Unknown layer processing.
@@ -129,7 +129,7 @@ void StoreAccumulatorOutput(in uint2 OutputPixelPosition, in FAccumulatorOutput
 		#elif DIM_LAYER_PROCESSING == LAYER_PROCESSING_SLIGHT_OUT_OF_FOCUS
 			Output0[OutputPixelPosition] = float4(
 				AccumulatorOutput.SlightFocusColor.rgb,
-				AccumulatorOutput.SlightFocusWeight);
+				AccumulatorOutput.SlightFocusOpacity);
 		
 		#else
 			#error Unknown layer processing.
@@ -337,10 +337,6 @@ void GatherMainCS(
 			AccumulatorOutput.BackgroundCocAvgAndSquareAvg.x = GatherParameters.KernelRadius;
 			AccumulatorOutput.BackgroundCocAvgAndSquareAvg.y = 0;
 
-		#elif DIM_LAYER_PROCESSING == LAYER_PROCESSING_SLIGHT_OUT_OF_FOCUS
-			AccumulatorOutput.SlightFocusColor = FastAccumulatorOutput.Color;
-			AccumulatorOutput.SlightFocusWeight = FastAccumulatorOutput.Weight;
-
 		#else
 			#error Fast gathering is not supported.
 		#endif
diff --git a/Engine/Shaders/Private/DiaphragmDOF/DOFRecombine.usf b/Engine/Shaders/Private/DiaphragmDOF/DOFRecombine.usf
--- a/Engine/Shaders/Private/DiaphragmDOF/DOFRecombine.usf
+++ b/Engine/Shaders/Private/DiaphragmDOF/DOFRecombine.usf
@@ -472,14 +472,18 @@ void AccumulateMirrorSamples(inout FFullResGatherAccumulator Accumulator, in FGa
 void AccumulateCenterSample(inout FFullResGatherAccumulator Accumulator, in FGatherSample A)
 {
 	FGatherSampleDerivedParameters DerivedA = ComputeSampleDerivates(Accumulator, A);
+
+	// Force this sample to be considered to guareentee their is a resolved color if in focus or background.
+	DerivedA.IsConsidered = 1;
+
 	AccumulateSample(Accumulator, A, DerivedA);
 }
 
 /** Resolves the slightly out of focus. */
 void ResolveAccumulator(
 	in FFullResGatherAccumulator Accumulator,
-	out float4 OutAdditiveColor,
-	out float OutTranslucency)
+	out float4 OutGatherBackgroundUnpremultipliedColor,
+	out float OutGatherBackgroundOpacity)
 {
 	const float SampleCount = 1.0 + 2.0 * Accumulator.Parameters.SamplePairCount;
 
@@ -498,9 +502,8 @@ void ResolveAccumulator(
 		Opacity = saturate(Accumulator.Opacity * SafeRcp(SampleCount));
 	}
 
-	OutTranslucency = Accumulator.ColorWeight > 0 ? 1 - Opacity : 1;
-	//OutTranslucency = 1 - Opacity;
-	OutAdditiveColor = Accumulator.Color * (Opacity * SafeRcp(Accumulator.ColorWeight));
+	OutGatherBackgroundOpacity = Accumulator.ColorWeight > 0 ? Opacity : 0;
+	OutGatherBackgroundUnpremultipliedColor = Accumulator.Color * (SafeRcp(Accumulator.ColorWeight));
 }
 
 
@@ -848,16 +851,12 @@ void RecombineMainCS(
 		return;
 	}
 
-	float4 GatherBackgroundAdditiveColor = 0;
-	float4 GatherForegroundAdditiveColor = 0;
-	float GatherBackgroundTranslucency = 1;
-	float GatherForegroundTranslucency = 1;
-	
+	float GatherBackgroundOpacity = ComputeInFocusOpacity(CenterSample.CocRadius);
+	float4 GatherBackgroundUnpremultipliedColor = CenterSample.Color;
+
 	#if CONFIG_SLIGHT_FOCUS_METHOD == SLIGHT_FOCUS_METHOD_DISABLED
 	{
-		float CenterOpacity = ComputeBackgroundSampleOpacity(CenterSample.CocRadius);
-		GatherBackgroundAdditiveColor = CenterSample.Color * CenterOpacity;
-		GatherBackgroundTranslucency = 1 - CenterOpacity;
+		GatherBackgroundOpacity = ComputeBackgroundSampleOpacity(CenterSample.CocRadius);
 	}
 	#elif CONFIG_SLIGHT_FOCUS_METHOD == SLIGHT_FOCUS_METHOD_UNIQUE_CONVOLUTIONS
 	BRANCH
@@ -917,32 +916,35 @@ void RecombineMainCS(
 				float Small = 0.125 * (1.0 - saturate(CenterSample.CocRadius * CenterSample.CocRadius * rcp(64.0)));
 				Max += HD * Small;
 				Min -= HD * Small;
+
+				// Ensures the temporally stable opacity remains between 0-1.
+				// Uses saturate() instead of min(0 and max() to be optimised as
+				// saturate() MAD post modifier on GCN.
+				Min.a = saturate(Min.a);
+				Max.a = saturate(Max.a);
 			}
 			#endif
 
+			// Clamp color
 			float3 ClampedColor = clamp(Accumulator.Color.rgb, Min.rgb * Accumulator.ColorWeight, Max.rgb * Accumulator.ColorWeight);
+
 			Accumulator.Color.rgb = lerp(Accumulator.Color.rgb, ClampedColor, saturate(CenterSample.CocRadius * CenterSample.CocRadius * 4.0));
+
+			// Clamp opacity.			
+			float ClampedOpacity = clamp(Accumulator.Opacity, Min.a * Accumulator.OpacityWeight, Max.a * Accumulator.OpacityWeight);
+
+			Accumulator.Opacity = lerp(Accumulator.Opacity, ClampedOpacity, saturate(CenterSample.CocRadius * CenterSample.CocRadius * 4.0));
 		}
 		#endif
 
 		// Resolve full res gather.
-		ResolveAccumulator(Accumulator, GatherBackgroundAdditiveColor, GatherBackgroundTranslucency);
-	}
-	else
-	{
-		float CenterOpacity = saturate(MAX_RECOMBINE_ABS_COC_RADIUS - abs(CenterSample.CocRadius));
-		GatherBackgroundAdditiveColor = CenterSample.Color * CenterOpacity;
-		GatherBackgroundTranslucency = 1 - CenterOpacity;
+		ResolveAccumulator(Accumulator, GatherBackgroundUnpremultipliedColor, GatherBackgroundOpacity);
 	}
-	#else
-		#error Unknown full res convolution.
 	#endif
 	
 	// Compose lower res foreground with full res gather foreground.
-	{
-		GatherForegroundAdditiveColor = GatherForegroundAdditiveColor * ForegroundTranslucency + ForegroundColor;
-		GatherForegroundTranslucency *= ForegroundTranslucency;
-	}
+	float4 GatherForegroundAdditiveColor = ForegroundColor;
+	float GatherForegroundTranslucency = ForegroundTranslucency;
 
 	// Sample lower res background, if necessary.
 	float4 BackgroundColor = 0.0;
@@ -986,17 +988,33 @@ void RecombineMainCS(
 			float MaxTranslucency = 1 - IsConsideredForegroundSample(CenterSample.CocRadius);
 
 			// Force the hole filling translucency to 0 if the background is unrendered.
-			MaxTranslucency = min(MaxTranslucency, BackgroundValidity);
+			if (BackgroundValidity <= 0.001)
+			{
+				if (HoleFillingTranslucency < 1.0)
+				{
+					float MaxTranslucency = 1 - IsConsideredForegroundSample(CenterSample.CocRadius);
 
-			AmendAdditiveColorWithMaxTranslucency(HoleFillingAdditiveColor, HoleFillingTranslucency, MaxTranslucency);
+					MaxTranslucency = min(MaxTranslucency, BackgroundValidity);
+					AmendAdditiveColorWithMaxTranslucency(HoleFillingAdditiveColor, HoleFillingTranslucency, MaxTranslucency);
+				}
+				else
+				{
+					GatherBackgroundOpacity = 1;
+				}
+			}
 
 			BackgroundColor = BackgroundColor * HoleFillingTranslucency + HoleFillingAdditiveColor;
 		}
+		#else
+		if (BackgroundValidity <= 0.001)
+		{
+			GatherBackgroundOpacity = 1;
+		}
 		#endif
 	}
 
 	// Compose background loweer res gather and full res gather.
-	float4 OutputFinalSceneColor = BackgroundColor * GatherBackgroundTranslucency + GatherBackgroundAdditiveColor;
+	float4 OutputFinalSceneColor = BackgroundColor * (1 - GatherBackgroundOpacity) + GatherBackgroundOpacity * GatherBackgroundUnpremultipliedColor;
 	
 	// Forces foreground translucency to 0 when large out of focus high res foreground, to
 	// avoid background leaking.
@@ -1061,11 +1079,13 @@ void RecombineMainCS(
 		OutputFinalSceneColor = GatherParameters.KernelPixelRadius;
 	}
 	#elif 0 // DEBUG
+	if (View.GeneralPurposeTweak > 0.5)
 	{
 		//OutputFinalSceneColor = ForegroundColor;
 		//OutputFinalSceneColor = HoleFillingAdditiveColor;
-		//OutputFinalSceneColor = BackgroundColor;
-		//OutputFinalSceneColor = GatherBackgroundAdditiveColor;
+		//OutputFinalSceneColor = BackgroundColor.a;
+		//OutputFinalSceneColor = PostprocessInput9.SampleLevel(GlobalPointClampedSampler, InputParameters.DOFBufferUV, 0).a;
+		//OutputFinalSceneColor = GatherBackgroundUnpremultipliedColor;
 		//OutputFinalSceneColor.a = GatherBackgroundTranslucency;
 	}
 	#elif 0 // DEBUG
diff --git a/Engine/Shaders/Private/DiaphragmDOF/DOFSetup.usf b/Engine/Shaders/Private/DiaphragmDOF/DOFSetup.usf
--- a/Engine/Shaders/Private/DiaphragmDOF/DOFSetup.usf
+++ b/Engine/Shaders/Private/DiaphragmDOF/DOFSetup.usf
@@ -4,6 +4,8 @@
 	DiaphragmDOF/DOFSetup.usf: Diaphragm DOF's setup shader. 
 =============================================================================*/
 
+#define CONFIG_SETUP 1
+
 #include "DOFDownsample.ush"
 
 
diff --git a/Engine/Shaders/Private/PostProcessTemporalAA.usf b/Engine/Shaders/Private/PostProcessTemporalAA.usf
--- a/Engine/Shaders/Private/PostProcessTemporalAA.usf
+++ b/Engine/Shaders/Private/PostProcessTemporalAA.usf
@@ -604,18 +604,17 @@ float2 WeightedLerpFactors(float WeightA, float WeightB, float Blend)
 // Computes the bilateral weight according to two Coc radii.
 float ComputeBilateralWeight(float RefCocRadius, float SampleCocRadius)
 {
-	#if 0
-		float ShapestCocRadius = abs(RefCocRadius) < abs(SampleCocRadius) ? RefCocRadius : SampleCocRadius;
-
-		return max(0, 1 - abs(RefCocRadius - SampleCocRadius) + abs(ShapestCocRadius));
-
-	#else
-		float Factor = (abs(RefCocRadius) > 1 ? rcp(abs(RefCocRadius)) : 1.0);
-		return max(0, 1 - abs(RefCocRadius - SampleCocRadius) * Factor);
+	float Factor = (abs(RefCocRadius) > 1 ? rcp(abs(RefCocRadius)) : 1.0);
+	return saturate(1 - abs(RefCocRadius - SampleCocRadius) * Factor);
+}
 
-	#endif
+float ComputeNeightborSampleBilateralWeight(float CenterCocRadius, float SampleCocRadius)
+{
+	float Factor = (abs(CenterCocRadius) > 1 ? rcp(abs(CenterCocRadius)) : 1.0);
+	return saturate(1 - (CenterCocRadius - SampleCocRadius) * Factor);
 }
 
+
 #endif
 
 
@@ -1524,7 +1523,7 @@ void FilterCurrentFrameInputSamples(
 
 			// Finds out the sample's bilateral weight according to the payload.
 			#if AA_HISTORY_PAYLOAD == HISTORY_PAYLOAD_RGB_COC
-				float BilateralWeight = ComputeBilateralWeight(CenterCocRadius, Sample.CocRadius);
+				float BilateralWeight = ComputeNeightborSampleBilateralWeight(CenterCocRadius, Sample.CocRadius);
 
 			#else
 				float BilateralWeight = 1;
@@ -2181,6 +2180,7 @@ FTAAHistoryPayload TemporalAASample(uint2 GroupId, uint2 GroupThreadId, uint Gro
 		#if AA_HISTORY_PAYLOAD == HISTORY_PAYLOAD_RGB_COC || AA_HISTORY_PAYLOAD == HISTORY_PAYLOAD_RGB_OPACITY_COC
 		{
 			float BilateralWeight = ComputeBilateralWeight(IntermediaryResult.Filtered.CocRadius, History.CocRadius);
+
 			BlendFinal = lerp(1, BlendFinal, BilateralWeight);
 		}
 		#endif
