Cherry-pick 4155356: Fixes shader compilation failure for platforms that supports PLATFORM_SUPPORTS_WAVE_BROADCAST.

diff --git a/Engine/Shaders/Private/DiaphragmDOF/DOFDownsample.ush b/Engine/Shaders/Private/DiaphragmDOF/DOFDownsample.ush
--- a/Engine/Shaders/Private/DiaphragmDOF/DOFDownsample.ush
+++ b/Engine/Shaders/Private/DiaphragmDOF/DOFDownsample.ush
@@ -218,17 +218,13 @@ float DownsampleCoc_Wave8x8(float CocRadius, const uint ReductionScale)
 		// Horizontal.
 		float OtherCocRadius = WaveSwapWithinLaneGroup(
 			OutCocRadius,
-			/* LaneGroupSize = */ 2 * ReductionScale,
-			/* InnerLaneGroupSize = */ 1 * ReductionScale,
-			/* InnerLaneGroupId = */ 1);
+			/* LaneGroupSize = */ 2 * ReductionScale);
 		OutCocRadius = min(OutCocRadius, OtherCocRadius);
 		
 		// Vertical.
 		OtherCocRadius = WaveSwapWithinLaneGroup(
 			OutCocRadius,
-			/* LaneGroupSize = */ 16 * ReductionScale,
-			/* InnerLaneGroupSize = */ 8 * ReductionScale,
-			/* InnerLaneGroupId = */ 1);
+			/* LaneGroupSize = */ 16 * ReductionScale);
 		OutCocRadius = min(OutCocRadius, OtherCocRadius);
 	#else
 		#error Unimplemented.
