Cherry-pick 4157123: Fixes legacy DOF's TAA pass clamping coc radius, causing wrong foreground out of focus.

diff --git a/Engine/Shaders/Private/PostProcessTemporalAA.usf b/Engine/Shaders/Private/PostProcessTemporalAA.usf
--- a/Engine/Shaders/Private/PostProcessTemporalAA.usf
+++ b/Engine/Shaders/Private/PostProcessTemporalAA.usf
@@ -33,6 +33,9 @@
 	
 	// Have RGB, opacity in alpha and DOF's CocRadius in separate's Red.
 	#define HISTORY_PAYLOAD_RGB_OPACITY_COC 3
+	
+	// Have RGB, opacity in alpha and old DOF's CocRadius in alpha.
+	#define HISTORY_PAYLOAD_RGB_COC_LEGACY 4
 
 
 /** Caching method for scene color. */
@@ -65,7 +68,7 @@
 #endif
 
 #if TAA_PASS_CONFIG == 0 // Used for DOF
-	#define AA_HISTORY_PAYLOAD (HISTORY_PAYLOAD_RGB_TRANSLUCENCY) // This isn't right be do it anyway for legacy reasons.
+	#define AA_HISTORY_PAYLOAD (HISTORY_PAYLOAD_RGB_COC_LEGACY)
 	#define AA_BORDER 1
 	#define AA_CROSS 4 // because running at half res.
 	#define AA_DYNAMIC 1
@@ -855,7 +858,7 @@ float GetSceneColorHdrWeight(
 		Sample.Color = InputParams.CachedNeighbors0[NeighborsId];
 		Sample.CocRadius = 0;
 
-		#if AA_HISTORY_PAYLOAD == HISTORY_PAYLOAD_RGB_COC
+		#if AA_HISTORY_PAYLOAD == HISTORY_PAYLOAD_RGB_COC || AA_HISTORY_PAYLOAD == HISTORY_PAYLOAD_RGB_COC_LEGACY
 			Sample.CocRadius = Sample.Color.a;
 		#elif AA_HISTORY_PAYLOAD == HISTORY_PAYLOAD_RGB_OPACITY_COC
 			Sample.CocRadius = InputParams.CachedNeighbors1[NeighborsId];
@@ -1334,7 +1337,7 @@ FTAASceneColorSample SampleCachedSceneColorTexture(
 
 	Sample.CocRadius = 0;
 
-	#if AA_HISTORY_PAYLOAD == HISTORY_PAYLOAD_RGB_COC
+	#if AA_HISTORY_PAYLOAD == HISTORY_PAYLOAD_RGB_COC || AA_HISTORY_PAYLOAD == HISTORY_PAYLOAD_RGB_COC_LEGACY
 		Sample.CocRadius = Sample.Color.a;
 	#elif AA_HISTORY_PAYLOAD == HISTORY_PAYLOAD_RGB_OPACITY_COC
 		Sample.CocRadius = PostprocessInput1.SampleLevel(
@@ -2259,7 +2262,7 @@ void MainCS(
 	float4 OutColor0 = 0;
 	float4 OutColor1 = 0;
 
-	#if AA_HISTORY_PAYLOAD == HISTORY_PAYLOAD_RGB_COC
+	#if AA_HISTORY_PAYLOAD == HISTORY_PAYLOAD_RGB_COC || AA_HISTORY_PAYLOAD == HISTORY_PAYLOAD_RGB_COC_LEGACY
 	{
 		OutColor0.rgb = OutputPayload.Color.rgb;
 		OutColor0.a = OutputPayload.CocRadius;
@@ -2348,7 +2351,7 @@ void MainPS(
 
 	FTAAHistoryPayload OutputPayload = TemporalAASample(GroupId, GroupThreadId, GroupThreadIndex, ViewportUV, FrameExposureScale);
 
-	#if AA_HISTORY_PAYLOAD == HISTORY_PAYLOAD_RGB_COC
+	#if AA_HISTORY_PAYLOAD == HISTORY_PAYLOAD_RGB_COC || AA_HISTORY_PAYLOAD == HISTORY_PAYLOAD_RGB_COC_LEGACY
 	{
 		OutColor0.rgb = OutputPayload.Color.rgb;
 		OutColor0.a = OutputPayload.CocRadius;
