Cherry-pick 4157461: Fixes potential bug where DOF's dilate could jump over some tiles at high resolution, and optimises recurisve kernel density by tracking bgd min intersectable coc.

diff --git a/Engine/Shaders/Private/DiaphragmDOF/DOFCocTileCommon.ush b/Engine/Shaders/Private/DiaphragmDOF/DOFCocTileCommon.ush
--- a/Engine/Shaders/Private/DiaphragmDOF/DOFCocTileCommon.ush
+++ b/Engine/Shaders/Private/DiaphragmDOF/DOFCocTileCommon.ush
@@ -11,9 +11,6 @@
 // Maximum number of buffer for Coc tile.
 #define COC_TILE_MAX_BUFFER_COUNT 2
 
-// Maximum number of Coc ranges per foreground/backrgound.
-#define MAX_COC_TILE_RANGES 2
-
 
 //------------------------------------------------------- ENUM VALUES
 
@@ -26,44 +23,40 @@
 // be used fprergound only pass, which means all background code would be compiled out.
 #define COC_TILE_LAYOUT_FGD (COC_TILE_LAYOUT_FGD_BGD)
 
-// 1 background encoded in RG channels.
-#define COC_TILE_LAYOUT_BGD 1
-
-// 2 background encoded in RG & BA channels.
-#define COC_TILE_LAYOUT_2BGD 2
-
-// TODO: COC_TILE_LAYOUT_FGD_2BGD
-
-#define COC_TILE_LAYOUT_MAX 3
+// 1 foreground encoded in RG channels if RT 0, 1 background encoded in RGB channels if RT 1.
+#define COC_TILE_LAYOUT_FGD_SEP_BGD 1
 
 
 /** Returns the number of buffers for a Coc tile layout. */
 uint GetBufferCountForCocTileLayour(uint Layout)
 {
+	if (Layout == COC_TILE_LAYOUT_FGD_SEP_BGD)
+	{
+		return 2;
+	}
 	return 1;
 }
 
 
 //------------------------------------------------------- STRUCTS
 
-/** Coc range. */
-struct FCocRange
-{
-	// The minimum Coc range.
-	float MinCoc;
-
-	// The maximum Coc range.
-	float MaxCoc;
-};
-
 /** Tile buffer. */
 struct FCocTileSample
 {
-	// Foreground Coc ranges.
-	FCocRange Foreground[MAX_COC_TILE_RANGES];
+	// The minimum Coc for foreground <= 0.
+	float FgdMinCoc;
+
+	// The maximum Coc for forground <= 0.
+	float FgdMaxCoc;
+
+	// The minimum Coc for background >= 0.
+	float BgdMinCoc;
 
-	// Background Coc ranges.
-	FCocRange Background[MAX_COC_TILE_RANGES];
+	// The maximum Coc for background >= 0.
+	float BgdMaxCoc;
+
+	// The maximum intersectable Coc radius.
+	float BgdMinIntersectableCoc;
 };
 
 
@@ -71,13 +64,11 @@ struct FCocTileSample
 
 void InitCocTileSample(out FCocTileSample Sample)
 {
-	for (uint i = 0; i < MAX_COC_TILE_RANGES; i++)
-	{
-		Sample.Foreground[i].MinCoc = 0;
-		Sample.Foreground[i].MaxCoc = 0;
-		Sample.Background[i].MinCoc = 0;
-		Sample.Background[i].MaxCoc = 0;
-	}
+	Sample.FgdMinCoc = 0;
+	Sample.FgdMaxCoc = 0;
+	Sample.BgdMinCoc = 0;
+	Sample.BgdMaxCoc = 0;
+	Sample.BgdMinIntersectableCoc = EXTREMELY_LARGE_COC_RADIUS;
 }
 
 /** Encode coc tile sample.
@@ -92,28 +83,16 @@ void EncodeCocTileSample(in FCocTileSample Sample, in uint Layout, out float4 Ra
 	
 	if (static_condition(Layout == COC_TILE_LAYOUT_FGD_BGD))
 	{
-		Raw[0].rg = float2(Sample.Foreground[0].MinCoc, Sample.Foreground[0].MaxCoc);
-		Raw[0].ba = float2(Sample.Background[0].MinCoc, Sample.Background[0].MaxCoc);
+		Raw[0].rg = float2(Sample.FgdMinCoc, Sample.FgdMaxCoc);
+		Raw[0].ba = float2(Sample.BgdMinCoc, Sample.BgdMaxCoc);
 	}
-	else if (static_condition(Layout == COC_TILE_LAYOUT_BGD))
+	else if (static_condition(Layout == COC_TILE_LAYOUT_FGD_SEP_BGD))
 	{
-		Raw[0].rg = float2(Sample.Background[0].MinCoc, Sample.Background[0].MaxCoc);
-	}
-	else if (static_condition(Layout == COC_TILE_LAYOUT_2BGD))
-	{
-		Raw[0].rg = float2(Sample.Background[0].MinCoc, Sample.Background[0].MaxCoc);
-		Raw[0].ba = float2(Sample.Background[1].MinCoc, Sample.Background[1].MaxCoc);
+		Raw[0].rg = float2(Sample.FgdMinCoc, Sample.FgdMaxCoc);
+		Raw[1].rgb = float3(Sample.BgdMaxCoc, Sample.BgdMinCoc, Sample.BgdMinIntersectableCoc);
 	}
 }
 
-/** Gets the output. */
-float4 GetEncodedCocTileSample(in FCocTileSample Sample, in uint Layout, uint BufferId)
-{
-	float4 Raw[COC_TILE_MAX_BUFFER_COUNT];
-	EncodeCocTileSample(Sample, Layout, Raw);
-	return Raw[BufferId];
-}
-
 /** Decode coc tile sample.
  * @param <InputLayout>: the layout of the coc buffer.
  */
@@ -123,26 +102,23 @@ void DecodeCocTileSample(in float4 Raw[COC_TILE_MAX_BUFFER_COUNT], in uint Layou
 
 	if (static_condition(Layout == COC_TILE_LAYOUT_FGD_BGD))
 	{
-		Sample.Foreground[0].MinCoc = Raw[0].r;
-		Sample.Foreground[0].MaxCoc = Raw[0].g;
-		Sample.Background[0].MinCoc = Raw[0].b;
-		Sample.Background[0].MaxCoc = Raw[0].a;
+		Sample.FgdMinCoc = Raw[0].r;
+		Sample.FgdMaxCoc = Raw[0].g;
+		Sample.BgdMinCoc = Raw[0].b;
+		Sample.BgdMaxCoc = Raw[0].a;
 	}
-	else if (static_condition(Layout == COC_TILE_LAYOUT_BGD))
+	else if (static_condition(Layout == COC_TILE_LAYOUT_FGD_SEP_BGD))
 	{
-		Sample.Background[0].MinCoc = Raw[0].r;
-		Sample.Background[0].MaxCoc = Raw[0].g;
-	}
-	else if (static_condition(Layout == COC_TILE_LAYOUT_2BGD))
-	{
-		Sample.Background[0].MinCoc = Raw[0].r;
-		Sample.Background[0].MaxCoc = Raw[0].g;
-		Sample.Background[1].MinCoc = Raw[0].b;
-		Sample.Background[1].MaxCoc = Raw[0].a;
+		Sample.FgdMinCoc = Raw[0].r;
+		Sample.FgdMaxCoc = Raw[0].g;
+		Sample.BgdMinCoc = Raw[1].g;
+		Sample.BgdMaxCoc = Raw[1].r;
+		Sample.BgdMinIntersectableCoc = Raw[1].b;
 	}
 }
 
 
+
 //------------------------------------------------------- Higher level API.
 
 /** Sample and decode tile buffer. */
diff --git a/Engine/Shaders/Private/DiaphragmDOF/DOFCocTileDilate.usf b/Engine/Shaders/Private/DiaphragmDOF/DOFCocTileDilate.usf
--- a/Engine/Shaders/Private/DiaphragmDOF/DOFCocTileDilate.usf
+++ b/Engine/Shaders/Private/DiaphragmDOF/DOFCocTileDilate.usf
@@ -8,6 +8,14 @@
 #include "DOFCocTileCommon.ush"
 
 
+//------------------------------------------------------- ENUM VALUES
+
+/** Matches FRCPassDiaphragmDOFDilateCoc::EMode. */
+	#define DILATE_MODE_STANDALONE 0
+	#define GATHER_MODE_MIN_MAX 1
+	#define GATHER_MODE_MIN_ABS 2
+
+
 //------------------------------------------------------- PASSES.
 
 uint SampleOffsetMultipler = 1;
@@ -28,7 +36,7 @@ RWTexture2D<float4>	Output1;
 [numthreads(DEFAULT_GROUP_BORDER_SIZE, DEFAULT_GROUP_BORDER_SIZE, 1)]
 void CocDilateMainCS(uint2 DispatchThreadId : SV_DispatchThreadID) 
 {
-	const uint CocTileLayout = COC_TILE_LAYOUT_FGD_BGD;
+	const uint CocTileLayout = COC_TILE_LAYOUT_FGD_SEP_BGD;
 	const int Radius = DIM_DILATE_RADIUS;
 
 	const uint2 PixelPos = DispatchThreadId;
@@ -45,10 +53,19 @@ void CocDilateMainCS(uint2 DispatchThreadId : SV_DispatchThreadID)
 		OutSampleBucketsisInitialized[j] = false;
 	}
 
+	// Standalone does not have any radius multiplier since there is only one dilate pass.
+	#if DIM_DILATE_MODE == DILATE_MODE_STANDALONE
+		const uint RadiusMultiplier = 1;
+	#else
+		const uint RadiusMultiplier = SampleOffsetMultipler;
+	#endif
+
 	// Iterate over all input tiles and assign them into their buckets.
-	UNROLL for (int x = -Radius; x <= Radius; x++)
+	UNROLL
+	for (int x = -Radius; x <= Radius; x++)
 	{
-		UNROLL for (int y = -Radius; y <= Radius; y++)
+		UNROLL
+		for (int y = -Radius; y <= Radius; y++)
 		{
 			// Don't sampling the center as we do this latter on.
 			if ( x == 0 && y == 0 )
@@ -57,7 +74,7 @@ void CocDilateMainCS(uint2 DispatchThreadId : SV_DispatchThreadID)
 			// Id of the bucket to dilate this sample into.
 			const uint BucketId = uint(max(abs(x), abs(y)) - 1);
 		
-			int2 Offset = asint(SampleOffsetMultipler) * int2(x, y);
+			int2 Offset = asint(RadiusMultiplier) * int2(x, y);
 			int2 SampleIndex = PixelPos + Offset;
 			
 			// If sampling outside the tiles, ignore this sample.
@@ -68,97 +85,115 @@ void CocDilateMainCS(uint2 DispatchThreadId : SV_DispatchThreadID)
 			// Minimal CoC radius to consider because done at full res.
 			const float MinimalCoC = 0.5 * kCocRadiusToFullResFactor;
 
-			// Closest distance from sampling tile to output tide.
+			// Closest distance from sampling tile to output tile.
 			float ClosestDistance = COC_TILE_SIZE * length(max(abs(float2(Offset)) - 1.0, 0.0)) + MinimalCoC;
 
 			// Sample tile.
 			const FCocTileSample NeightborSample = LoadCocTile(CocTileLayout, PostprocessInput0, PostprocessInput1, SampleIndex);
 
-			// TODO: Could be smarter here like in old code to optimise better the gather pass?
-			UNROLL
-			for (uint i = 0; i < MAX_COC_TILE_RANGES; i++)
+			if (OutSampleBucketsisInitialized[BucketId] == false)
 			{
-				if (OutSampleBucketsisInitialized[BucketId] == false)
-				{
-					OutSampleBucketsisInitialized[BucketId] = true;
-					OutSampleBuckets[BucketId] = NeightborSample;
-				}
-				else
-				{
-					OutSampleBuckets[BucketId].Foreground[i].MinCoc = min(
-						OutSampleBuckets[BucketId].Foreground[i].MinCoc, NeightborSample.Foreground[i].MinCoc);
-					OutSampleBuckets[BucketId].Foreground[i].MaxCoc = max(
-						OutSampleBuckets[BucketId].Foreground[i].MaxCoc, NeightborSample.Foreground[i].MaxCoc);
-
-					OutSampleBuckets[BucketId].Background[i].MinCoc = min(
-						OutSampleBuckets[BucketId].Background[i].MinCoc, NeightborSample.Background[i].MinCoc);
-					OutSampleBuckets[BucketId].Background[i].MaxCoc = max(
-						OutSampleBuckets[BucketId].Background[i].MaxCoc, NeightborSample.Background[i].MaxCoc);
-				}
+				OutSampleBucketsisInitialized[BucketId] = true;
+				OutSampleBuckets[BucketId] = NeightborSample;
+			}
+			else
+			{
+				OutSampleBuckets[BucketId].FgdMinCoc = min(
+					OutSampleBuckets[BucketId].FgdMinCoc,
+					NeightborSample.FgdMinCoc);
+
+				OutSampleBuckets[BucketId].FgdMaxCoc = max(
+					OutSampleBuckets[BucketId].FgdMaxCoc,
+					NeightborSample.FgdMaxCoc);
+
+				OutSampleBuckets[BucketId].BgdMinCoc = min(
+					OutSampleBuckets[BucketId].BgdMinCoc,
+					NeightborSample.BgdMinCoc);
+
+				OutSampleBuckets[BucketId].BgdMaxCoc = max(
+					OutSampleBuckets[BucketId].BgdMaxCoc,
+					NeightborSample.BgdMaxCoc);
+					
+				OutSampleBuckets[BucketId].BgdMinIntersectableCoc = min(
+					OutSampleBuckets[BucketId].BgdMinIntersectableCoc,
+					NeightborSample.BgdMinIntersectableCoc);
 			}
 		}
 	}
 	
-	FCocTileSample OutSample;
-	
 	// Sample the center bucket.
-	FCocTileSample CenterSample = LoadCocTile(CocTileLayout, PostprocessInput0, PostprocessInput1, PixelPos);
+	FCocTileSample OutSample = LoadCocTile(CocTileLayout, PostprocessInput0, PostprocessInput1, PixelPos);
 
-	UNROLL
-	for (uint i = 0; i < MAX_COC_TILE_RANGES; i++)
 	{
-		// First, find out what the maximum absolute Coc radii will be for the output sample.
-		float MinForegroundCocRadius = CenterSample.Foreground[i].MinCoc;
-		float MaxBackgroundCocRadius = CenterSample.Background[i].MaxCoc;
-		
 		int k;
-
+	
+		// Find out what the maximum absolute Coc radii will be for the output sample.
 		UNROLL
 		for (k = 0; k < Radius; k++)
 		{
-			float BucketDistance = (float(SampleOffsetMultipler) * float(k + 1) - 1.0) * 8.0;
+			float BucketDistance = (float(RadiusMultiplier) * float(k + 1) - 1.0) * 8.0;
 			
 			FLATTEN
-			if (-OutSampleBuckets[k].Foreground[i].MinCoc * CocRadiusToBucketDistance > BucketDistance)
+			if (-OutSampleBuckets[k].FgdMinCoc * CocRadiusToBucketDistance > BucketDistance)
 			{
-				MinForegroundCocRadius = min(MinForegroundCocRadius, OutSampleBuckets[k].Foreground[i].MinCoc);
+				OutSample.FgdMinCoc = min(OutSample.FgdMinCoc, OutSampleBuckets[k].FgdMinCoc);
 			}
 			
 			FLATTEN
-			if (OutSampleBuckets[k].Background[i].MaxCoc * CocRadiusToBucketDistance > BucketDistance)
+			if (OutSampleBuckets[k].BgdMaxCoc * CocRadiusToBucketDistance > BucketDistance)
 			{
-				MaxBackgroundCocRadius = max(MaxBackgroundCocRadius, OutSampleBuckets[k].Background[i].MaxCoc);
+				OutSample.BgdMaxCoc = max(OutSample.BgdMaxCoc, OutSampleBuckets[k].BgdMaxCoc);
 			}
 		}
 		
-		// Then, find out what is minimum absolute Coc radii that will be encountered for these maximum Coc radii.
-		float MaxForegroundCocRadius = CenterSample.Foreground[i].MaxCoc;
-		float MinBackgroundCocRadius = CenterSample.Background[i].MinCoc;
-		
+		// Fetches the already dilated min foreground and max background COC radius.
+		#if DIM_DILATE_MODE == GATHER_MODE_MIN_ABS
+		{
+			FCocTileSample DilatedMinMaxSample = LoadCocTile(CocTileLayout, PostprocessInput2, PostprocessInput3, PixelPos);
+
+			OutSample.FgdMinCoc = DilatedMinMaxSample.FgdMinCoc;
+			OutSample.BgdMaxCoc = DilatedMinMaxSample.BgdMaxCoc;
+		}
+		#endif
+
+		const float CocRadiusErrorFraction = 0.05;
+		float ForegroundCocRadiusError = -OutSample.FgdMinCoc * CocRadiusErrorFraction;
+		float BackgroundCocRadiusError = OutSample.BgdMaxCoc * CocRadiusErrorFraction;
+
+
+		// Then, find out what is minimum absolute Coc radii that will be intersected for these maximum Coc radii.
 		UNROLL
 		for (k = 0; k < Radius; k++)
 		{
-			float BucketDistance = (float(SampleOffsetMultipler) * float(k + 1) - 1.0) * float(COC_TILE_SIZE);
+			float BucketDistance = (float(RadiusMultiplier) * float(k + 1) - 1.0) * float(COC_TILE_SIZE);
+			
+			const float MinimalCoC = 0.5 * kCocRadiusToFullResFactor;
+			float MaxIntersectableAbsCocRadius = COC_TILE_SIZE * (float(RadiusMultiplier * (k + 1)) - 1) + MinimalCoC;
 			
 			FLATTEN
-			if (-MinForegroundCocRadius * CocRadiusToBucketDistance > BucketDistance)
+			if (-OutSample.FgdMinCoc * CocRadiusToBucketDistance > BucketDistance)
 			{
-				MaxForegroundCocRadius = max(MaxForegroundCocRadius, OutSampleBuckets[k].Foreground[i].MaxCoc);
+				OutSample.FgdMaxCoc = max(OutSample.FgdMaxCoc, OutSampleBuckets[k].FgdMaxCoc);
 			}
 
 			FLATTEN
-			if (MaxBackgroundCocRadius * CocRadiusToBucketDistance > BucketDistance)
+			if (OutSample.BgdMaxCoc * CocRadiusToBucketDistance > BucketDistance)
 			{
-				MinBackgroundCocRadius = min(MinBackgroundCocRadius, OutSampleBuckets[k].Background[i].MinCoc);
+				OutSample.BgdMinCoc = min(OutSample.BgdMinCoc, OutSampleBuckets[k].BgdMinCoc);
+				OutSample.BgdMinIntersectableCoc = min(OutSample.BgdMinIntersectableCoc, max(OutSampleBuckets[k].BgdMinIntersectableCoc, MaxIntersectableAbsCocRadius));
 			}
 		}
-		
-		OutSample.Foreground[i].MinCoc = MinForegroundCocRadius;
-		OutSample.Foreground[i].MaxCoc = MaxForegroundCocRadius;
-		OutSample.Background[i].MinCoc = MinBackgroundCocRadius;
-		OutSample.Background[i].MaxCoc = MaxBackgroundCocRadius;
 	}
 
+	// If dilating min foreground and max background coc, zero output other to prune the unecessary ALU.
+	#if DIM_DILATE_MODE == GATHER_MODE_MIN_MAX
+	{
+		OutSample.FgdMaxCoc = 0;
+		OutSample.BgdMinCoc = 0;
+		OutSample.BgdMinIntersectableCoc = 0;
+	}
+	#endif
+
 	// Output that.
 	if (all(PixelPos < ViewportRect.zw))
 	{
diff --git a/Engine/Shaders/Private/DiaphragmDOF/DOFCocTileFlatten.usf b/Engine/Shaders/Private/DiaphragmDOF/DOFCocTileFlatten.usf
--- a/Engine/Shaders/Private/DiaphragmDOF/DOFCocTileFlatten.usf
+++ b/Engine/Shaders/Private/DiaphragmDOF/DOFCocTileFlatten.usf
@@ -37,32 +37,50 @@ RWTexture2D<float4>	Output1;
 
 //------------------------------------------------------- LDS.
 
-#define LDS_LAYOUT (COC_TILE_LAYOUT_FGD_BGD)
+#define LDS_LAYOUT (COC_TILE_LAYOUT_FGD_SEP_BGD)
 
-#if  CONFIG_FLATTEN_METHOD == FLATTEN_METHOD_REDUCE_LDS
-
-// TODO: could technically only need half of it: would this be faster?
-groupshared float4 Shared[DEFAULT_GROUP_BORDER_SIZE * DEFAULT_GROUP_BORDER_SIZE];
 
+void InitCocTileSampleFromCocRadius(float CocRadius, out FCocTileSample OutA)
+{
+	InitCocTileSample(OutA);
+
+	OutA.FgdMinCoc = min(CocRadius, 0);
+	OutA.FgdMaxCoc = OutA.FgdMinCoc;
+				
+	OutA.BgdMinCoc = max(CocRadius, 0);
+	OutA.BgdMaxCoc = OutA.BgdMinCoc;
+	OutA.BgdMinIntersectableCoc = (CocRadius > 0) ? CocRadius : EXTREMELY_LARGE_COC_RADIUS;
+}
 
 void ReduceOperator(inout FCocTileSample A, in const FCocTileSample B)
 {
-	for (uint i = 0; i < MAX_COC_TILE_RANGES; i++)
-	{
-		A.Foreground[i].MinCoc = min(A.Foreground[i].MinCoc, B.Foreground[i].MinCoc);
-		A.Foreground[i].MaxCoc = max(A.Foreground[i].MaxCoc, B.Foreground[i].MaxCoc);
+	A.FgdMinCoc = min(A.FgdMinCoc, B.FgdMinCoc);
+	A.FgdMaxCoc = max(A.FgdMaxCoc, B.FgdMaxCoc);
 
-		A.Background[i].MinCoc = min(A.Background[i].MinCoc, B.Background[i].MinCoc);
-		A.Background[i].MaxCoc = max(A.Background[i].MaxCoc, B.Background[i].MaxCoc);
-	}
+	A.BgdMinCoc = min(A.BgdMinCoc, B.BgdMinCoc);
+	A.BgdMaxCoc = max(A.BgdMaxCoc, B.BgdMaxCoc);
 }
 
+
+#if CONFIG_FLATTEN_METHOD == FLATTEN_METHOD_REDUCE_LDS
+
+// TODO: could technically only need half of it: would this be faster?
+groupshared float4 Shared0[DEFAULT_GROUP_BORDER_SIZE * DEFAULT_GROUP_BORDER_SIZE];
+groupshared float Shared1[DEFAULT_GROUP_BORDER_SIZE * DEFAULT_GROUP_BORDER_SIZE];
+
 FCocTileSample LoadFromLds(uint i)
 {
 	float4 Raw[2];
-	Raw[0] = Shared[i];
+	Raw[0] = Shared0[i];
 	Raw[1] = 0;
-
+	
+	if (LDS_LAYOUT == COC_TILE_LAYOUT_FGD_SEP_BGD)
+	{
+		Raw[0] = Shared0[i];
+		Raw[1].rg = Shared0[i].ba;
+		Raw[1].b = Shared1[i];
+	}
+	
 	FCocTileSample Sample;
 	DecodeCocTileSample(Raw, LDS_LAYOUT, Sample);
 	return Sample;
@@ -70,7 +88,18 @@ FCocTileSample LoadFromLds(uint i)
 
 void StoreToLds(uint i, in const FCocTileSample Sample)
 {
-	Shared[i] = GetEncodedCocTileSample(Sample, LDS_LAYOUT, /** BufferId = */ 0);
+	float4 Raw[COC_TILE_MAX_BUFFER_COUNT];
+	EncodeCocTileSample(Sample, LDS_LAYOUT, Raw);
+	
+	if (LDS_LAYOUT == COC_TILE_LAYOUT_FGD_SEP_BGD)
+	{
+		Shared0[i] = float4(Raw[0].rg, Raw[1].rg);
+		Shared1[i] = Raw[1].b;
+	}
+	else // if (LDS_LAYOUT == COC_TILE_LAYOUT_FGD_BGD)
+	{
+		Shared0[i] = Raw[0];
+	}
 }
 
 
@@ -79,6 +108,7 @@ void StoreToLds(uint i, in const FCocTileSample Sample)
 groupshared uint SharedForegroundMinCocRadius;
 groupshared uint SharedForegroundMaxCocRadius;
 groupshared uint SharedBackgroundMinCocRadius;
+groupshared uint SharedBackgroundMinIntersectableCocRadius;
 groupshared uint SharedBackgroundMaxCocRadius;
 
 #endif
@@ -92,7 +122,10 @@ void CocFlattenMainCS(
 	uint2 DispatchThreadId : SV_DispatchThreadID,
 	uint GroupIndex : SV_GroupIndex)
 {
+	FCocTileSample TileSample;
+
 	#if DIM_DO_COC_GATHER4
+	{
 		float2 BufferUV = (DispatchThreadId + 0.5) * ThreadIdToBufferUV;
 
 		if (true)
@@ -107,33 +140,35 @@ void CocFlattenMainCS(
 			float4 CocRadii = PostprocessInput0.GatherAlpha(GlobalPointClampedSampler, BufferUV, 0);
 		#endif
 
-		float MinCocRadius = min(min(CocRadii.x, CocRadii.y), min(CocRadii.z, CocRadii.w));
-		float MaxCocRadius = max(max(CocRadii.x, CocRadii.y), max(CocRadii.z, CocRadii.w));
+		InitCocTileSampleFromCocRadius(CocRadii.x, TileSample);
 
+		UNROLL
+		for (uint i = 1; i < 4; i++)
+		{
+			FCocTileSample SecondTileSample;
+			InitCocTileSampleFromCocRadius(CocRadii[i], SecondTileSample);
+			ReduceOperator(TileSample, SecondTileSample);
+		}
+	}
 	#else // !DIM_DO_COC_GATHER4
+	{
 		uint2 PixelPos = min(DispatchThreadId, ViewportRect.zw - 1);
 
 		#if CONFIG_DOF_ALPHA
 			// The coc radius is within separate buffer's red channel.
-			float MinCocRadius = PostprocessInput0[PixelPos].r;
+			float CocRadius = PostprocessInput0[PixelPos].r;
 		#else
-			float MinCocRadius = PostprocessInput0[PixelPos].w;
+			float CocRadius = PostprocessInput0[PixelPos].w;
 		#endif
 
-		float MaxCocRadius = MinCocRadius;
+		InitCocTileSampleFromCocRadius(CocRadius, TileSample);
+	}
 	#endif
 
-	#if  CONFIG_FLATTEN_METHOD == FLATTEN_METHOD_REDUCE_LDS
+	#if CONFIG_FLATTEN_METHOD == FLATTEN_METHOD_REDUCE_LDS
 	{
 		// Setup local sample.
-		FCocTileSample Local;
-		for (uint i = 0; i < MAX_COC_TILE_RANGES; i++)
-		{
-			Local.Foreground[i].MinCoc = min(MinCocRadius, 0);
-			Local.Foreground[i].MaxCoc = min(MaxCocRadius, 0);
-			Local.Background[i].MinCoc = max(MinCocRadius, 0);
-			Local.Background[i].MaxCoc = max(MaxCocRadius, 0);
-		}
+		FCocTileSample Local = TileSample;
 
 		// LDS setup.
 		StoreToLds(GroupIndex, Local);
@@ -161,7 +196,7 @@ void CocFlattenMainCS(
 		{	
 			// Final reduction.
 			ReduceOperator(Local, LoadFromLds(1));
-		
+			
 			// Encode.
 			float4 RawOutput[COC_TILE_MAX_BUFFER_COUNT];
 			EncodeCocTileSample(Local, LDS_LAYOUT, RawOutput);
@@ -185,6 +220,7 @@ void CocFlattenMainCS(
 				SharedForegroundMinCocRadius = 0;
 				SharedForegroundMaxCocRadius = ~0;
 				SharedBackgroundMinCocRadius = ~0;
+				SharedBackgroundMinIntersectableCocRadius = ~0;
 				SharedBackgroundMaxCocRadius = 0;
 			}
 		
@@ -193,12 +229,16 @@ void CocFlattenMainCS(
 
 		// Perform wave flattening first.
 		#if CONFIG_FLATTEN_METHOD == FLATTEN_METHOD_WAVE_MINMAX
-			MinCocRadius = WaveAllMin(MinCocRadius);
-			MaxCocRadius = WaveAllMax(MaxCocRadius);
+			TileSample.FgdMinCoc = WaveAllMin(TileSample.FgdMinCoc);
+			TileSample.FgdMaxCoc = WaveAllMax(TileSample.FgdMaxCoc);
+			TileSample.BgdMinCoc = WaveAllMin(TileSample.BgdMinCoc);
+			TileSample.BgdMaxCoc = WaveAllMax(TileSample.BgdMaxCoc);
+			TileSample.BgdMinIntersectableCoc = WaveAllMin(TileSample.BgdMinIntersectableCoc);
 
 			if (WaveGetLaneCount() == DEFAULT_GROUP_BORDER_SIZE * DEFAULT_GROUP_BORDER_SIZE)
 			{
-				// NOP: because no need to gather the result of all waves of the group.
+				// NOP: because no need to gather the result of all waves of the group, because there is
+				// only one wave per group.
 			}
 			else if (WaveOnce())
 		#endif
@@ -210,42 +250,38 @@ void CocFlattenMainCS(
 			uint Unused;
 			
 			// Foreground CoCRadius <= 0.0.
-			InterlockedMax(SharedForegroundMinCocRadius, asuint(min(MinCocRadius, 0)), Unused);
-			InterlockedMin(SharedForegroundMaxCocRadius, asuint(min(MaxCocRadius, 0)), Unused);
+			InterlockedMax(SharedForegroundMinCocRadius, asuint(TileSample.FgdMinCoc), Unused);
+			InterlockedMin(SharedForegroundMaxCocRadius, asuint(TileSample.FgdMaxCoc), Unused);
 
 			// Background CoCRadius >= 0.0.
-			InterlockedMin(SharedBackgroundMinCocRadius, asuint(max(MinCocRadius, 0)), Unused);
-			InterlockedMax(SharedBackgroundMaxCocRadius, asuint(max(MaxCocRadius, 0)), Unused);
+			InterlockedMin(SharedBackgroundMinCocRadius, asuint(TileSample.BgdMinCoc), Unused);
+			InterlockedMax(SharedBackgroundMaxCocRadius, asuint(TileSample.BgdMaxCoc), Unused);
+			InterlockedMin(SharedBackgroundMinIntersectableCocRadius, asuint(TileSample.BgdMinIntersectableCoc), Unused);
+			
+			GroupMemoryBarrierWithGroupSync();
 		}
-		
-		GroupMemoryBarrierWithGroupSync();
-	
+
 		// Output.
 		BRANCH
 		if (GroupIndex == 0)
 		{	
-			// Fetch from LDS.
+			// Finialize tile sample to output.
 			FCocTileSample OutputTileSample;
 			#if CONFIG_FLATTEN_METHOD == FLATTEN_METHOD_WAVE_MINMAX
 				if (WaveGetLaneCount() == DEFAULT_GROUP_BORDER_SIZE * DEFAULT_GROUP_BORDER_SIZE)
 				{
-					OutputTileSample.Foreground[0].MinCoc = min(MinCocRadius, 0);
-					OutputTileSample.Foreground[0].MaxCoc = min(MaxCocRadius, 0);
-					OutputTileSample.Background[0].MinCoc = max(MinCocRadius, 0);
-					OutputTileSample.Background[0].MaxCoc = max(MaxCocRadius, 0);
+					OutputTileSample = TileSample;
 				}
 				else
 			#endif
 			{
-				OutputTileSample.Foreground[0].MinCoc = asfloat(SharedForegroundMinCocRadius);
-				OutputTileSample.Foreground[0].MaxCoc = asfloat(SharedForegroundMaxCocRadius);
-				OutputTileSample.Background[0].MinCoc = asfloat(SharedBackgroundMinCocRadius);
-				OutputTileSample.Background[0].MaxCoc = asfloat(SharedBackgroundMaxCocRadius);
+				OutputTileSample.FgdMinCoc = asfloat(SharedForegroundMinCocRadius);
+				OutputTileSample.FgdMaxCoc = asfloat(SharedForegroundMaxCocRadius);
+				OutputTileSample.BgdMinCoc = asfloat(SharedBackgroundMinCocRadius);
+				OutputTileSample.BgdMaxCoc = asfloat(SharedBackgroundMaxCocRadius);
+				OutputTileSample.BgdMinIntersectableCoc = asfloat(SharedBackgroundMinIntersectableCocRadius);
 			}
 
-			OutputTileSample.Foreground[1] = OutputTileSample.Foreground[0];
-			OutputTileSample.Background[1] = OutputTileSample.Background[0];
-
 			// Encode.
 			float4 RawOutput[COC_TILE_MAX_BUFFER_COUNT];
 			EncodeCocTileSample(OutputTileSample, LDS_LAYOUT, RawOutput);
diff --git a/Engine/Shaders/Private/DiaphragmDOF/DOFCommon.ush b/Engine/Shaders/Private/DiaphragmDOF/DOFCommon.ush
--- a/Engine/Shaders/Private/DiaphragmDOF/DOFCommon.ush
+++ b/Engine/Shaders/Private/DiaphragmDOF/DOFCommon.ush
@@ -77,8 +77,8 @@
 
 //------------------------------------------------------- COMPILE TIME CONSTANTS.
 
-// A very large Coc radius that is assumed to be unreachable.
-#define EXTREMELY_LARGE_COC_RADIUS 100000
+// A very large Coc radius that is assumed to be unreachable, but still encodable in 16bits float.
+#define EXTREMELY_LARGE_COC_RADIUS 16384
 
 
 // Default border size for group.
diff --git a/Engine/Shaders/Private/DiaphragmDOF/DOFGatherPass.usf b/Engine/Shaders/Private/DiaphragmDOF/DOFGatherPass.usf
--- a/Engine/Shaders/Private/DiaphragmDOF/DOFGatherPass.usf
+++ b/Engine/Shaders/Private/DiaphragmDOF/DOFGatherPass.usf
@@ -231,7 +231,7 @@ void GatherMainCS(
 		#else
 			const uint2 TilePos = GroupId;
 		#endif
-		const FCocTileSample CocTileInfos = LoadCocTile(COC_TILE_LAYOUT_FGD_BGD, PostprocessInput2, PostprocessInput2, int2(TilePos));
+		const FCocTileSample CocTileInfos = LoadCocTile(COC_TILE_LAYOUT_FGD_SEP_BGD, PostprocessInput2, PostprocessInput3, int2(TilePos));
 
 		TileSuggestion = InferGatherTileSuggestion(CocTileInfos,
 			/* LayerProcessing = */ DIM_LAYER_PROCESSING,
@@ -376,7 +376,7 @@ void GatherMainCS(
 
 				GatherParameters.bDoHighDensityGatheringAtCenter[i] = (
 					TileSuggestion.MaxCocRadiusAbs > MinConvolutionRadius &&
-					TileSuggestion.MinCocRadiusAbs < SmallestConvolutionRadius);
+					TileSuggestion.MinIntersectableCocRadiusAbs < SmallestConvolutionRadius);
 			}
 		}
 		#endif
diff --git a/Engine/Shaders/Private/DiaphragmDOF/DOFGatherTileSuggest.ush b/Engine/Shaders/Private/DiaphragmDOF/DOFGatherTileSuggest.ush
--- a/Engine/Shaders/Private/DiaphragmDOF/DOFGatherTileSuggest.ush
+++ b/Engine/Shaders/Private/DiaphragmDOF/DOFGatherTileSuggest.ush
@@ -23,6 +23,9 @@ struct FGatheringTileSuggestion
 	// Smallest absolute coc radius
 	float MinCocRadiusAbs;
 	
+	// Smallest absolute coc radius
+	float MinIntersectableCocRadiusAbs;
+	
 	// Largest absolute coc radius
 	float MaxCocRadiusAbs;
 
@@ -39,63 +42,71 @@ FGatheringTileSuggestion InferGatherTileSuggestion(
 {
 	FGatheringTileSuggestion OutSuggestion;
 
+	bool bIsComplexTile;
+	
+	// May not have this information, and it may not be needed.
+	OutSuggestion.MinIntersectableCocRadiusAbs = EXTREMELY_LARGE_COC_RADIUS;
+
 	if (LayerProcessing == LAYER_PROCESSING_FOREGROUND_ONLY)
 	{
 		// Smallest Coc radius is the foreground minimum Coc.
-		OutSuggestion.MinCocRadiusAbs = -CocTileInfos.Foreground[0].MaxCoc;
+		OutSuggestion.MinCocRadiusAbs = -CocTileInfos.FgdMaxCoc;
 		
 		// Kernel radius is the foreground maximum Coc.
-		OutSuggestion.MaxCocRadiusAbs = -CocTileInfos.Foreground[0].MinCoc;
+		OutSuggestion.MaxCocRadiusAbs = -CocTileInfos.FgdMinCoc;
 		
 		// Closest Coc radius is the foreground minimum Coc.
-		OutSuggestion.ClosestCocRadius = CocTileInfos.Foreground[0].MinCoc;
+		OutSuggestion.ClosestCocRadius = CocTileInfos.FgdMinCoc;
 	}
 	else if (LayerProcessing == LAYER_PROCESSING_FOREGROUND_HOLE_FILLING)
 	{
 		// Smallest Coc radius is the foreground minimum Coc.
-		OutSuggestion.MinCocRadiusAbs = -CocTileInfos.Foreground[0].MaxCoc;
-		
+		OutSuggestion.MinCocRadiusAbs = -CocTileInfos.FgdMaxCoc;
+
 		// Kernel radius is the foreground maximum Coc.
 		// TODO: max() with background.
-		OutSuggestion.MaxCocRadiusAbs = -CocTileInfos.Foreground[0].MinCoc;
+		OutSuggestion.MaxCocRadiusAbs = -CocTileInfos.FgdMinCoc;
 			
 		// Closest Coc radius is the foreground minimum Coc.
-		OutSuggestion.ClosestCocRadius = CocTileInfos.Foreground[0].MinCoc;
+		OutSuggestion.ClosestCocRadius = CocTileInfos.FgdMinCoc;
 	}
 	else if (LayerProcessing == LAYER_PROCESSING_BACKGROUND_ONLY)
 	{
 		// Smallest Coc radius is the backrgound minimum Coc.
-		OutSuggestion.MinCocRadiusAbs = CocTileInfos.Background[0].MinCoc;
+		OutSuggestion.MinCocRadiusAbs = CocTileInfos.BgdMinCoc;
 		
+		// Don't have this information, so just assume this is min absolute coc radius
+		OutSuggestion.MinIntersectableCocRadiusAbs = CocTileInfos.BgdMinIntersectableCoc;
+
 		// Kernel radius is the foreground maximum Coc.
-		OutSuggestion.MaxCocRadiusAbs = CocTileInfos.Background[0].MaxCoc;
+		OutSuggestion.MaxCocRadiusAbs = CocTileInfos.BgdMaxCoc;
 		
 		// Closest Coc radius is the background minimum Coc.
-		OutSuggestion.ClosestCocRadius = CocTileInfos.Background[0].MinCoc;
+		OutSuggestion.ClosestCocRadius = CocTileInfos.BgdMinCoc;
 	}
 	else // if (LayerProcessing == LAYER_PROCESSING_SLIGHT_OUT_OF_FOCUS)
 	{
 		// TODO: could be optimised on transition between very close foreground and very far background to bCanSkipGathering == true.
-		if (CocTileInfos.Foreground[0].MinCoc == 0 && CocTileInfos.Background[0].MaxCoc == 0)
+		if (CocTileInfos.FgdMinCoc == 0 && CocTileInfos.BgdMaxCoc == 0)
 		{
 			OutSuggestion.MinCocRadiusAbs = EXTREMELY_LARGE_COC_RADIUS;
 		}
-		else if (CocTileInfos.Foreground[0].MinCoc == 0)
+		else if (CocTileInfos.FgdMinCoc == 0)
 		{
-			OutSuggestion.MinCocRadiusAbs = CocTileInfos.Background[0].MinCoc;
+			OutSuggestion.MinCocRadiusAbs = CocTileInfos.BgdMinCoc;
 		}
-		else if (CocTileInfos.Background[0].MaxCoc == 0)
+		else if (CocTileInfos.BgdMaxCoc == 0)
 		{
-			OutSuggestion.MinCocRadiusAbs = -CocTileInfos.Foreground[0].MaxCoc;
+			OutSuggestion.MinCocRadiusAbs = -CocTileInfos.FgdMaxCoc;
 		}
 		else
 		{
-			OutSuggestion.MinCocRadiusAbs = min(-CocTileInfos.Foreground[0].MaxCoc, CocTileInfos.Background[0].MinCoc);
+			OutSuggestion.MinCocRadiusAbs = min(-CocTileInfos.FgdMaxCoc, CocTileInfos.BgdMinCoc);
 		}
+		
+		OutSuggestion.MaxCocRadiusAbs = max(-CocTileInfos.FgdMinCoc, CocTileInfos.BgdMaxCoc);
 
-		OutSuggestion.MaxCocRadiusAbs = max(-CocTileInfos.Foreground[0].MinCoc, CocTileInfos.Background[0].MaxCoc);
-
-		OutSuggestion.ClosestCocRadius = CocTileInfos.Foreground[0].MinCoc < 0 ? CocTileInfos.Foreground[0].MinCoc : CocTileInfos.Background[0].MinCoc;
+		OutSuggestion.ClosestCocRadius = CocTileInfos.FgdMinCoc < 0 ? CocTileInfos.FgdMinCoc : CocTileInfos.BgdMinCoc;
 	}
 	
 	// Coc error accepted for fast gathering, set to 5% of the kernel size.
diff --git a/Engine/Shaders/Private/DiaphragmDOF/DOFPostfiltering.usf b/Engine/Shaders/Private/DiaphragmDOF/DOFPostfiltering.usf
--- a/Engine/Shaders/Private/DiaphragmDOF/DOFPostfiltering.usf
+++ b/Engine/Shaders/Private/DiaphragmDOF/DOFPostfiltering.usf
@@ -163,7 +163,7 @@ void PostfilterMainCS(
 	#if DIM_TILE_PERMUTATION
 	{
 		const uint2 TilePos = GroupId;
-		const FCocTileSample CocTileInfos = LoadCocTile(COC_TILE_LAYOUT_FGD_BGD, PostprocessInput2, PostprocessInput2, TilePos);
+		const FCocTileSample CocTileInfos = LoadCocTile(COC_TILE_LAYOUT_FGD_SEP_BGD, PostprocessInput2, PostprocessInput3, TilePos);
 
 		FGatheringTileSuggestion TileSuggestion = InferGatherTileSuggestion(CocTileInfos,
 			/* LayerProcessing = */ DIM_LAYER_PROCESSING,
diff --git a/Engine/Source/Runtime/Renderer/Private/PostProcess/DiaphragmDOFGraph.cpp b/Engine/Source/Runtime/Renderer/Private/PostProcess/DiaphragmDOFGraph.cpp
--- a/Engine/Source/Runtime/Renderer/Private/PostProcess/DiaphragmDOFGraph.cpp
+++ b/Engine/Source/Runtime/Renderer/Private/PostProcess/DiaphragmDOFGraph.cpp
@@ -334,7 +334,8 @@ bool DiaphragmDOF::WireSceneColorPasses(FPostprocessContext& Context, const FRen
 	}
 
 	// Generate conservative coc tiles.
-	FRenderingCompositeOutputRef CocTileOutput;
+	FRenderingCompositeOutputRef CocTileOutput0;
+	FRenderingCompositeOutputRef CocTileOutput1;
 	{
 		// Flatten half res CoC to lower res tiles.
 		FRCPassDiaphragmDOFFlattenCoc::FParameters FlattenParams;
@@ -343,9 +344,11 @@ bool DiaphragmDOF::WireSceneColorPasses(FPostprocessContext& Context, const FRen
 
 		FRenderingCompositePass* CocFlatten = Context.Graph.RegisterPass(new(FMemStack::Get()) FRCPassDiaphragmDOFFlattenCoc(FlattenParams));
 		CocFlatten->SetInput(ePId_Input0, GatherColorSetup1.IsValid() ? GatherColorSetup1 : GatherColorSetup0);
-		CocTileOutput = CocFlatten;
+		CocTileOutput0 = FRenderingCompositeOutputRef(CocFlatten, ePId_Output0);
+		CocTileOutput1 = FRenderingCompositeOutputRef(CocFlatten, ePId_Output1);
 
 		// Parameters for the dilate Coc passes.
+		int32 DilateCount = 1;
 		FRCPassDiaphragmDOFDilateCoc::FParameters DilateParams[3];
 		{
 			const int32 MaxSampleRadiusCount = FRCPassDiaphragmDOFDilateCoc::MaxSampleRadiusCount;
@@ -383,20 +386,59 @@ bool DiaphragmDOF::WireSceneColorPasses(FPostprocessContext& Context, const FRen
 				DilateParams[i].SampleDistanceMultiplier = FMath::Min(IdealMultiplier, HighestPossibleMultiplierUpperBound);
 
 				CurrentConvolutionRadius += DilateParams[i].SampleRadiusCount * DilateParams[i].SampleDistanceMultiplier;
+
+				DilateCount++;
 			}
 		}
 
-		// Creates the dilates passes.
-		for (int32 i = 0; i < ARRAY_COUNT(DilateParams); i++)
+		// Setup common parameters.
+		for (int32 i = 0; i < DilateCount; i++)
 		{
-			if (!DilateParams[i].SampleRadiusCount && i != 0) break;
-
 			DilateParams[i].GatherViewSize = GatheringViewSize;
 			DilateParams[i].PreProcessingToProcessingCocRadiusFactor = PreProcessingToProcessingCocRadiusFactor;
+		}
 
-			FRenderingCompositePass* CocDilate = Context.Graph.RegisterPass(new(FMemStack::Get()) FRCPassDiaphragmDOFDilateCoc(DilateParams[i]));
-			CocDilate->SetInput(ePId_Input0, CocTileOutput);
-			CocTileOutput = CocDilate;
+		if (DilateCount > 1)
+		{
+			FRenderingCompositeOutputRef CocTileMinmaxOutput0 = CocTileOutput0;
+			FRenderingCompositeOutputRef CocTileMinmaxOutput1 = CocTileOutput1;
+
+			// Dilate min foreground and max background coc radii first.
+			for (int32 i = 0; i < DilateCount; i++)
+			{
+				FRCPassDiaphragmDOFDilateCoc::FParameters Params = DilateParams[i];
+				Params.Mode = FRCPassDiaphragmDOFDilateCoc::EMode::MinForegroundAndMaxBackground;
+
+				FRenderingCompositePass* CocDilate = Context.Graph.RegisterPass(new(FMemStack::Get()) FRCPassDiaphragmDOFDilateCoc(Params));
+				CocDilate->SetInput(ePId_Input0, CocTileMinmaxOutput0);
+				CocDilate->SetInput(ePId_Input1, CocTileMinmaxOutput1);
+				CocTileMinmaxOutput0 = FRenderingCompositeOutputRef(CocDilate, ePId_Output0);
+				CocTileMinmaxOutput1 = FRenderingCompositeOutputRef(CocDilate, ePId_Output1);
+			}
+
+			// Dilates everything else.
+			for (int32 i = 0; i < DilateCount; i++)
+			{
+				FRCPassDiaphragmDOFDilateCoc::FParameters Params = DilateParams[i];
+				Params.Mode = FRCPassDiaphragmDOFDilateCoc::EMode::MinimalAbsoluteRadiuses;
+
+				FRenderingCompositePass* CocDilate = Context.Graph.RegisterPass(new(FMemStack::Get()) FRCPassDiaphragmDOFDilateCoc(Params));
+				CocDilate->SetInput(ePId_Input0, CocTileOutput0);
+				CocDilate->SetInput(ePId_Input1, CocTileOutput1);
+				CocDilate->SetInput(ePId_Input2, CocTileMinmaxOutput0);
+				CocDilate->SetInput(ePId_Input3, CocTileMinmaxOutput1);
+				CocTileOutput0 = FRenderingCompositeOutputRef(CocDilate, ePId_Output0);
+				CocTileOutput1 = FRenderingCompositeOutputRef(CocDilate, ePId_Output1);
+			}
+		}
+		else
+		{
+			// Just dilate everything in one single pass.
+			FRenderingCompositePass* CocDilate = Context.Graph.RegisterPass(new(FMemStack::Get()) FRCPassDiaphragmDOFDilateCoc(DilateParams[0]));
+			CocDilate->SetInput(ePId_Input0, CocTileOutput0);
+			CocDilate->SetInput(ePId_Input1, CocTileOutput1);
+			CocTileOutput0 = FRenderingCompositeOutputRef(CocDilate, ePId_Output0);
+			CocTileOutput1 = FRenderingCompositeOutputRef(CocDilate, ePId_Output1);
 		}
 	}
 
@@ -495,7 +537,8 @@ bool DiaphragmDOF::WireSceneColorPasses(FPostprocessContext& Context, const FRen
 			FRenderingCompositePass* GatherPass = Context.Graph.RegisterPass(new(FMemStack::Get()) FRCPassDiaphragmDOFGather(GatherParameters));
 			GatherPass->SetInput(ePId_Input0, GatherInput0);
 			GatherPass->SetInput(ePId_Input1, GatherInput1);
-			GatherPass->SetInput(ePId_Input2, CocTileOutput);
+			GatherPass->SetInput(ePId_Input2, CocTileOutput0);
+			GatherPass->SetInput(ePId_Input3, CocTileOutput1);
 
 			if (GatherParameters.BokehSimulation != EDiaphragmDOFBokehSimulation::Disabled)
 				GatherPass->SetInput(ePId_Input4, GatheringBokehLUTOutput);
@@ -509,7 +552,8 @@ bool DiaphragmDOF::WireSceneColorPasses(FPostprocessContext& Context, const FRen
 			{
 				FRenderingCompositePass* Postfilter = Context.Graph.RegisterPass(new(FMemStack::Get()) FRCPassDiaphragmDOFPostfilter(GatherParameters));
 				Postfilter->SetInput(ePId_Input0, Input);
-				Postfilter->SetInput(ePId_Input2, CocTileOutput);
+				Postfilter->SetInput(ePId_Input2, CocTileOutput0);
+				Postfilter->SetInput(ePId_Input3, CocTileOutput1);
 				Input = FRenderingCompositeOutputRef(Postfilter, ePId_Output0);
 			}
 			return Input;
@@ -623,7 +667,8 @@ bool DiaphragmDOF::WireSceneColorPasses(FPostprocessContext& Context, const FRen
 			GatherParameters));
 		GatherPass->SetInput(ePId_Input0, GatherInput0); // TODO: take TAA input instead?
 		GatherPass->SetInput(ePId_Input1, GatherInput1);
-		GatherPass->SetInput(ePId_Input2, CocTileOutput);
+		GatherPass->SetInput(ePId_Input2, CocTileOutput0);
+		GatherPass->SetInput(ePId_Input3, CocTileOutput1);
 
 		// Slight out of focus gather pass use exact same LUT as scattering because all samples of ther kernel are used.
 		if (bEnableSlightOutOfFocusBokeh)
diff --git a/Engine/Source/Runtime/Renderer/Private/PostProcess/DiaphragmDOFPasses.cpp b/Engine/Source/Runtime/Renderer/Private/PostProcess/DiaphragmDOFPasses.cpp
--- a/Engine/Source/Runtime/Renderer/Private/PostProcess/DiaphragmDOFPasses.cpp
+++ b/Engine/Source/Runtime/Renderer/Private/PostProcess/DiaphragmDOFPasses.cpp
@@ -139,6 +139,18 @@ const TCHAR* GetEventName(FRCPassDiaphragmDOFGather::EQualityConfig e)
 	return kArray[i];
 }
 
+const TCHAR* GetEventName(FRCPassDiaphragmDOFDilateCoc::EMode e)
+{
+	static const TCHAR* const kArray[] = {
+		TEXT("StandAlone"),
+		TEXT("MinMax"),
+		TEXT("MinAbs"),
+	};
+	int32 i = int32(e);
+	check(i < ARRAY_COUNT(kArray));
+	return kArray[i];
+}
+
 // Returns X and Y for F(M) = saturate(M * X + Y) so that F(LowM) = 0 and F(HighM) = 1
 FVector2D GenerateSaturatedAffineTransformation(float LowM, float HighM)
 {
@@ -507,6 +519,7 @@ namespace
 {
 
 class FDDOFDilateRadiusDim     : SHADER_PERMUTATION_RANGE_INT("DIM_DILATE_RADIUS", 1, 3);
+class FDDOFDilateModeDim       : SHADER_PERMUTATION_ENUM_CLASS("DIM_DILATE_MODE", FRCPassDiaphragmDOFDilateCoc::EMode);
 
 class FDDOFLayerProcessingDim  : SHADER_PERMUTATION_ENUM_CLASS("DIM_LAYER_PROCESSING", EDiaphragmDOFLayerProcessing);
 class FDDOFGatherRingCountDim  : SHADER_PERMUTATION_RANGE_INT("DIM_GATHER_RING_COUNT", FRCPassDiaphragmDOFGather::kMinRingCount, 3);
@@ -601,7 +614,8 @@ FPooledRenderTargetDesc FRCPassDiaphragmDOFFlattenCoc::ComputeOutputDesc(EPassOu
 	FIntPoint TileCount = CocTileGridSize(UnmodifiedRet.Extent);
 
 	FPooledRenderTargetDesc Ret(FPooledRenderTargetDesc::Create2DDesc(TileCount, PF_FloatRGBA, FClearValueBinding::None, TexCreate_None, TexCreate_RenderTargetable | TexCreate_UAV, false));
-	Ret.DebugName = TEXT("DOFFlattenCoc");
+	Ret.DebugName = InPassOutputId == ePId_Output0 ? TEXT("DOFFlattenFgdCoc") : TEXT("DOFFlattenBgdCoc");
+	Ret.Format = InPassOutputId == ePId_Output0 ? PF_G16R16F : PF_FloatRGB;
 	return Ret;
 }
 
@@ -617,7 +631,7 @@ class FPostProcessCocDilateCS : public FPostProcessDiaphragmDOFShader
 	DECLARE_GLOBAL_SHADER(FPostProcessCocDilateCS);
 	SHADER_TYPE_PARAMETERS(FPostProcessCocDilateCS, FPostProcessDiaphragmDOFShader, COC_DILATE_SHADER_PARAMS);
 
-	using FPermutationDomain = TShaderPermutationDomain<FDDOFDilateRadiusDim>;
+	using FPermutationDomain = TShaderPermutationDomain<FDDOFDilateRadiusDim, FDDOFDilateModeDim>;
 };
 
 IMPLEMENT_GLOBAL_SHADER(FPostProcessCocDilateCS, "/Engine/Private/DiaphragmDOF/DOFCocTileDilate.usf", "CocDilateMainCS", SF_Compute);
@@ -627,12 +641,14 @@ void FRCPassDiaphragmDOFDilateCoc::Process(FRenderingCompositePassContext& Conte
 {
 	FPostProcessCocDilateCS::FPermutationDomain PermutationVector;
 	PermutationVector.Set<FDDOFDilateRadiusDim>(Params.SampleRadiusCount);
+	PermutationVector.Set<FDDOFDilateModeDim>(Params.Mode);
+	// TODO: permutation to do foreground and background separately, to have higher occupancy?
 
 	FDispatchDiaphragmDOFPass<FPostProcessCocDilateCS, FRCPassDiaphragmDOFDilateCoc> DispatchCtx(this, Context, PermutationVector);
 	DispatchCtx.DestViewport = FIntRect(FIntPoint::ZeroValue, FIntPoint::DivideAndRoundUp(Params.GatherViewSize, kCocTileSize));
 
-	SCOPED_DRAW_EVENTF(Context.RHICmdList, DiaphragmDOFDilateCoc, TEXT("DiaphragmDOF DilateCoc(1/16 radius=%d step=%d) %dx%d"),
-		Params.SampleRadiusCount, Params.SampleDistanceMultiplier,
+	SCOPED_DRAW_EVENTF(Context.RHICmdList, DiaphragmDOFDilateCoc, TEXT("DiaphragmDOF DilateCoc(1/16 %s radius=%d step=%d) %dx%d"),
+		GetEventName(Params.Mode), Params.SampleRadiusCount, Params.SampleDistanceMultiplier,
 		DispatchCtx.DestViewport.Width(), DispatchCtx.DestViewport.Height());
 
 	{
@@ -647,7 +663,19 @@ void FRCPassDiaphragmDOFDilateCoc::Process(FRenderingCompositePassContext& Conte
 FPooledRenderTargetDesc FRCPassDiaphragmDOFDilateCoc::ComputeOutputDesc(EPassOutputId InPassOutputId) const
 {
 	FPooledRenderTargetDesc Ret = GetInput(ePId_Input0)->GetOutput()->RenderTargetDesc;
-	Ret.DebugName = TEXT("DOFDilateCoc");
+
+	// When dilating only min foreground and max background, only one channel.
+	if (Params.Mode == EMode::MinForegroundAndMaxBackground)
+	{
+		Ret.DebugName = InPassOutputId == ePId_Output0 ? TEXT("DOFDilateMinFgdCoc") : TEXT("DOFDilateMaxBgdCoc");
+		Ret.Format = PF_R16F;
+	}
+	else
+	{
+		Ret.DebugName = InPassOutputId == ePId_Output0 ? TEXT("DOFDilateFgdCoc") : TEXT("DOFDilateBgdCoc");
+		Ret.Format = InPassOutputId == ePId_Output0 ? PF_G16R16F : PF_FloatRGB;
+	}
+
 	return Ret;
 }
 
diff --git a/Engine/Source/Runtime/Renderer/Private/PostProcess/DiaphragmDOFPasses.h b/Engine/Source/Runtime/Renderer/Private/PostProcess/DiaphragmDOFPasses.h
--- a/Engine/Source/Runtime/Renderer/Private/PostProcess/DiaphragmDOFPasses.h
+++ b/Engine/Source/Runtime/Renderer/Private/PostProcess/DiaphragmDOFPasses.h
@@ -69,7 +69,7 @@ enum class EDiaphragmDOFBokehSimulation
 
 // 
 // ePId_Input0: SceneDepth
-class FRCPassDiaphragmDOFFlattenCoc : public TRenderingCompositePassBase<1, 1>
+class FRCPassDiaphragmDOFFlattenCoc : public TRenderingCompositePassBase<1, 2>
 {
 public:
 	/** Resolution divisor of the Coc tiles. */
@@ -104,15 +104,33 @@ private:
 
 // 
 // ePId_Input0: SceneDepth
-class FRCPassDiaphragmDOFDilateCoc : public TRenderingCompositePassBase<1, 1>
+class FRCPassDiaphragmDOFDilateCoc : public TRenderingCompositePassBase<4, 2>
 {
 public:
 	/** Resolution divisor of the Coc tiles. */
 	static constexpr int32 MaxSampleRadiusCount = 3;
 
+	/** Dilate mode of the pass. */
+	enum class EMode
+	{
+		// One single dilate pass.
+		StandAlone,
+
+		// Dilate min foreground and max background coc radius.
+		MinForegroundAndMaxBackground,
+
+		// Dilate everything else from dilated min foreground and max background coc radius.
+		MinimalAbsoluteRadiuses,
+
+		MAX
+	};
+
 	/** Configuration parameters of the dilate pass. */
 	struct FParameters
 	{
+		/** Mode of teh dilate. */
+		EMode Mode;
+
 		/** Radius in number of samples. 0 means the dilate pass should not be run. */
 		int32 SampleRadiusCount;
 
@@ -126,7 +144,8 @@ public:
 		float PreProcessingToProcessingCocRadiusFactor;
 
 		FParameters()
-			: SampleRadiusCount(0)
+			: Mode(EMode::StandAlone)
+			, SampleRadiusCount(0)
 			, SampleDistanceMultiplier(1)
 			, PreProcessingToProcessingCocRadiusFactor(1.0f)
 		{ }
@@ -136,6 +155,7 @@ public:
 		: Params(InParams)
 	{
 		check(Params.SampleRadiusCount > 0);
+		check(Params.SampleDistanceMultiplier == 1 || Params.Mode != EMode::StandAlone);
 	}
 
 	// interface FRenderingCompositePass ---------
@@ -418,7 +438,7 @@ private:
 };
 
 // 
-class FRCPassDiaphragmDOFPostfilter : public TRenderingCompositePassBase<3, 2>
+class FRCPassDiaphragmDOFPostfilter : public TRenderingCompositePassBase<4, 2>
 {
 public:
 	/** Configuration parameters of the gather pass to post filter.
