Cherry-pick 4158896: Optimises the variadic number of sample density changes on background, and work arround shader compiler bug in coc dilate pass.

diff --git a/Engine/Shaders/Private/DiaphragmDOF/DOFCocTileDilate.usf b/Engine/Shaders/Private/DiaphragmDOF/DOFCocTileDilate.usf
--- a/Engine/Shaders/Private/DiaphragmDOF/DOFCocTileDilate.usf
+++ b/Engine/Shaders/Private/DiaphragmDOF/DOFCocTileDilate.usf
@@ -19,9 +19,11 @@
 //------------------------------------------------------- PASSES.
 
 uint SampleOffsetMultipler = 1;
+float fSampleOffsetMultipler = 1;
 
 // Multiplier to compare CocRadius and bucket distance.
 float CocRadiusToBucketDistance;
+float BucketDistanceToCocRadius;
 
 
 //------------------------------------------------------- OUTPUTS.
@@ -37,101 +39,125 @@ RWTexture2D<float4>	Output1;
 void CocDilateMainCS(uint2 DispatchThreadId : SV_DispatchThreadID) 
 {
 	const uint CocTileLayout = COC_TILE_LAYOUT_FGD_SEP_BGD;
-	const int Radius = DIM_DILATE_RADIUS;
+	const uint Radius = DIM_DILATE_RADIUS;
 
 	const uint2 PixelPos = DispatchThreadId;
 
 	// Different buckets of CocRadius per gathering radius.
 	FCocTileSample OutSampleBuckets[Radius];
 
-	// Whether OutSampleBuckets[BucketId] has been initialized. Going to be compile
-	// out at runtime tahnks to the UNROLL over x and y variables.
-	bool OutSampleBucketsisInitialized[Radius];
-	UNROLL
-	for (int j = 0; j < Radius; j++)
-	{
-		OutSampleBucketsisInitialized[j] = false;
-	}
-
 	// Standalone does not have any radius multiplier since there is only one dilate pass.
 	#if DIM_DILATE_MODE == DILATE_MODE_STANDALONE
 		const uint RadiusMultiplier = 1;
+		const uint fRadiusMultiplier = 1;
 	#else
 		const uint RadiusMultiplier = SampleOffsetMultipler;
+		const uint fRadiusMultiplier = fSampleOffsetMultipler;
 	#endif
 
 	// Iterate over all input tiles and assign them into their buckets.
 	UNROLL
-	for (int x = -Radius; x <= Radius; x++)
+	for (uint BucketId = 0; BucketId < Radius; BucketId++)
 	{
-		UNROLL
-		for (int y = -Radius; y <= Radius; y++)
-		{
-			// Don't sampling the center as we do this latter on.
-			if ( x == 0 && y == 0 )
-				continue;
-	
-			// Id of the bucket to dilate this sample into.
-			const uint BucketId = uint(max(abs(x), abs(y)) - 1);
+		uint RingDistance = BucketId + 1;
 		
-			int2 Offset = asint(RadiusMultiplier) * int2(x, y);
-			int2 SampleIndex = PixelPos + Offset;
-			
-			// If sampling outside the tiles, ignore this sample.
-			bool2 bInsideViewport = SampleIndex >= 0 && uint2(SampleIndex) < uint2(ViewportRect.zw);
-			if (!all(bInsideViewport))
-				continue;
-			
-			// Minimal CoC radius to consider because done at full res.
-			const float MinimalCoC = 0.5 * kCocRadiusToFullResFactor;
-
-			// Closest distance from sampling tile to output tile.
-			float ClosestDistance = COC_TILE_SIZE * length(max(abs(float2(Offset)) - 1.0, 0.0)) + MinimalCoC;
-
-			// Sample tile.
-			const FCocTileSample NeightborSample = LoadCocTile(CocTileLayout, PostprocessInput0, PostprocessInput1, SampleIndex);
+		// Init the bucket.
+		OutSampleBuckets[BucketId].FgdMinCoc = 0;
+		OutSampleBuckets[BucketId].FgdMaxCoc = -EXTREMELY_LARGE_COC_RADIUS;
+		OutSampleBuckets[BucketId].BgdMinCoc = EXTREMELY_LARGE_COC_RADIUS;
+		OutSampleBuckets[BucketId].BgdMaxCoc = 0;
+		OutSampleBuckets[BucketId].BgdMinIntersectableCoc = EXTREMELY_LARGE_COC_RADIUS;
 
-			if (OutSampleBucketsisInitialized[BucketId] == false)
+		UNROLL
+		for (uint BatchId = 0; BatchId < 4 * RingDistance; BatchId++)
+		{
+			// Generates at compile time sample location if square shaped ring in sample count distance from the kernel's center.
+			//
+			// RingId = 0:
+			//      Y
+			//      |
+			//    3 2 1
+			//    4   0 -X
+			//
+			// RingId = 1:
+			//      Y
+			//      |
+			//  6 5 4 3 2
+			//  7       1
+			//  8       0 -X
+			//
+			// ...
+			int2 BatchSquarePos;
+			if (BatchId < RingDistance)
+			{
+				BatchSquarePos.x = int(RingDistance);
+				BatchSquarePos.y = int(BatchId);
+			}
+			else if (BatchId < RingDistance * 3)
 			{
-				OutSampleBucketsisInitialized[BucketId] = true;
-				OutSampleBuckets[BucketId] = NeightborSample;
+				BatchSquarePos.x = int(RingDistance) - int(BatchId - RingDistance);
+				BatchSquarePos.y = int(RingDistance);
 			}
 			else
 			{
-				OutSampleBuckets[BucketId].FgdMinCoc = min(
-					OutSampleBuckets[BucketId].FgdMinCoc,
-					NeightborSample.FgdMinCoc);
-
-				OutSampleBuckets[BucketId].FgdMaxCoc = max(
-					OutSampleBuckets[BucketId].FgdMaxCoc,
-					NeightborSample.FgdMaxCoc);
-
-				OutSampleBuckets[BucketId].BgdMinCoc = min(
-					OutSampleBuckets[BucketId].BgdMinCoc,
-					NeightborSample.BgdMinCoc);
+				BatchSquarePos.x = -int(RingDistance);
+				BatchSquarePos.y = int(RingDistance) - (int(BatchId) - int(RingDistance) * 3);
+			}
 
-				OutSampleBuckets[BucketId].BgdMaxCoc = max(
-					OutSampleBuckets[BucketId].BgdMaxCoc,
-					NeightborSample.BgdMaxCoc);
+			UNROLL
+			for (uint SampleId = 0; SampleId < 2; SampleId++)
+			{
+				const int2 Offset = (SampleId == 0 ? 1 : -1) * BatchSquarePos;
+				int2 SampleIndex = int2(PixelPos) + Offset;
+			
+				// If sampling outside the tiles, ignore this sample.
+				bool2 bInsideViewport = uint2(SampleIndex) < uint2(ViewportRect.zw);
+				if (all(bInsideViewport))
+				{
+					// Minimal CoC radius to consider because done at full res.
+					const float MinimalCoC = 0.5 * kCocRadiusToFullResFactor;
+
+					// Closest distance from sampling tile to output tile.
+					float ClosestDistance = fRadiusMultiplier * (COC_TILE_SIZE * length(max(abs(float2(Offset)) - 1.0, 0.0)));
+
+					// Sample tile.
+					const FCocTileSample NeightborSample = LoadCocTile(CocTileLayout, PostprocessInput0, PostprocessInput1, SampleIndex);
+
+					OutSampleBuckets[BucketId].FgdMinCoc = min(
+						OutSampleBuckets[BucketId].FgdMinCoc,
+						NeightborSample.FgdMinCoc);
+
+					OutSampleBuckets[BucketId].FgdMaxCoc = max(
+						OutSampleBuckets[BucketId].FgdMaxCoc,
+						NeightborSample.FgdMaxCoc);
+
+					OutSampleBuckets[BucketId].BgdMinCoc = min(
+						OutSampleBuckets[BucketId].BgdMinCoc,
+						NeightborSample.BgdMinCoc);
+
+					OutSampleBuckets[BucketId].BgdMaxCoc = max(
+						OutSampleBuckets[BucketId].BgdMaxCoc,
+						NeightborSample.BgdMaxCoc);
 					
-				OutSampleBuckets[BucketId].BgdMinIntersectableCoc = min(
-					OutSampleBuckets[BucketId].BgdMinIntersectableCoc,
-					NeightborSample.BgdMinIntersectableCoc);
-			}
-		}
-	}
+					OutSampleBuckets[BucketId].BgdMinIntersectableCoc = min(
+						OutSampleBuckets[BucketId].BgdMinIntersectableCoc,
+						NeightborSample.BgdMinIntersectableCoc + ClosestDistance * BucketDistanceToCocRadius);
+				}
+			} // for (uint SampleId = 0; SampleId < 2; SampleId++)
+		} // for (uint BatchId = 0; BatchId < 4 * RingDistance; BatchId++)
+	} // for (uint BucketId = 0; BucketId < Radius; BucketId++)
 	
 	// Sample the center bucket.
 	FCocTileSample OutSample = LoadCocTile(CocTileLayout, PostprocessInput0, PostprocessInput1, PixelPos);
 
 	{
-		int k;
+		uint k;
 	
 		// Find out what the maximum absolute Coc radii will be for the output sample.
 		UNROLL
 		for (k = 0; k < Radius; k++)
 		{
-			float BucketDistance = (float(RadiusMultiplier) * float(k + 1) - 1.0) * 8.0;
+			float BucketDistance = fRadiusMultiplier * (COC_TILE_SIZE * float(k + 1)) - COC_TILE_SIZE;
 			
 			FLATTEN
 			if (-OutSampleBuckets[k].FgdMinCoc * CocRadiusToBucketDistance > BucketDistance)
@@ -165,10 +191,9 @@ void CocDilateMainCS(uint2 DispatchThreadId : SV_DispatchThreadID)
 		UNROLL
 		for (k = 0; k < Radius; k++)
 		{
-			float BucketDistance = (float(RadiusMultiplier) * float(k + 1) - 1.0) * float(COC_TILE_SIZE);
+			float BucketDistance = fRadiusMultiplier * ((float(k + 1) - 1.0) * float(COC_TILE_SIZE));
 			
 			const float MinimalCoC = 0.5 * kCocRadiusToFullResFactor;
-			float MaxIntersectableAbsCocRadius = COC_TILE_SIZE * (float(RadiusMultiplier * (k + 1)) - 1) + MinimalCoC;
 			
 			FLATTEN
 			if (-OutSample.FgdMinCoc * CocRadiusToBucketDistance > BucketDistance)
@@ -180,7 +205,7 @@ void CocDilateMainCS(uint2 DispatchThreadId : SV_DispatchThreadID)
 			if (OutSample.BgdMaxCoc * CocRadiusToBucketDistance > BucketDistance)
 			{
 				OutSample.BgdMinCoc = min(OutSample.BgdMinCoc, OutSampleBuckets[k].BgdMinCoc);
-				OutSample.BgdMinIntersectableCoc = min(OutSample.BgdMinIntersectableCoc, max(OutSampleBuckets[k].BgdMinIntersectableCoc, MaxIntersectableAbsCocRadius));
+				OutSample.BgdMinIntersectableCoc = min(OutSample.BgdMinIntersectableCoc, OutSampleBuckets[k].BgdMinIntersectableCoc);
 			}
 		}
 	}
diff --git a/Engine/Shaders/Private/DiaphragmDOF/DOFGatherCommon.ush b/Engine/Shaders/Private/DiaphragmDOF/DOFGatherCommon.ush
--- a/Engine/Shaders/Private/DiaphragmDOF/DOFGatherCommon.ush
+++ b/Engine/Shaders/Private/DiaphragmDOF/DOFGatherCommon.ush
@@ -104,6 +104,9 @@
 	// Background need coc radius error correction to reduce amount of background leaking.
 	#define CONFIG_COC_RADIUS_CORRECTION 1
 
+	// TODO: investigate this for foreground.
+	#define CONFIG_FULL_KERNEL_CENTER_RANDOM (DIM_GATHER_QUALITY != GATHER_QUALITY_LOW_QUALITY)
+
 #elif DIM_LAYER_PROCESSING == LAYER_PROCESSING_SLIGHT_OUT_OF_FOCUS
 	#define GATHER_ACCUMULATOR_METHOD (GATHER_ACCUMULATOR_METHOD_SIMPLE)
 
@@ -204,6 +207,10 @@
 	#define CONFIG_POINT_MIRROR_BOKEH 1
 #endif
 
+// Whether to randomise the kernel center on the full displacement.
+#ifndef CONFIG_FULL_KERNEL_CENTER_RANDOM
+	#define CONFIG_FULL_KERNEL_CENTER_RANDOM 0
+#endif
 
 // D3D shader compiler is rediculously slow to compile this shader.
 #if SM5_PROFILE
@@ -239,50 +246,6 @@
 
 //------------------------------------------------------- STAND ALONE FUNCTIONS
 
-// Computes intersection of sample from its Coc and its distance.
-// Returns 0: no intersection, 1: intersection.
-float ComputeSampleIntersection(float SampleCocRadius, float SampleDistance)
-{
-	// Intersection of output pixel's disk area with sample's Coc.
-	#if 1
-		// To have an intersection of 50% when the Sample's Coc radius == SampleDistance.
-		#if 0 // TODO: clean this hack.
-			// Poor out of focus, great alpha channel, no leaking.
-			const float Offset = 1;
-		#elif 1
-			// Best slight out of focus.
-			const float Offset = 0;
-		#else
-			// Mathematically ideal.
-			const float Offset = 0.5;
-		#endif
-
-		// Uses a sharper intersection for slight out of focus like Recombine pass since it is being used
-		// for full res temporal statbility clamping box, so the aliasing this may introduce at half res is
-		// Fine.
-		const float Multiplier = (DIM_LAYER_PROCESSING == LAYER_PROCESSING_SLIGHT_OUT_OF_FOCUS) ? 4.0 : 1.0;
-
-		return saturate((abs(SampleCocRadius) - SampleDistance) * Multiplier + Offset);
-
-	// Same as Circle DOF.
-	#elif 0
-		// 0.5 because Coc radius is at half res, and 0.5 so that the fallof is only the radius of a ful res pixel.
-		const float IntersectionSharpen = 0.5 * 0.5;
-
-		// TODO: missing +0.5 to be energy conservative.
-		return saturate((abs(SampleCocRadius) - SampleDistance) * IntersectionSharpen);
-
-	#else // Hard cut of.
-		if (SampleDistance > abs(SampleCocRadius))
-		{
-			return 0.0;
-		}
-		return 1.0;
-
-	#endif
-}
-
-
 // Compute sample weight according's to its Coc radius.
 float ComputeSampleWeight(float CocRadius)
 {
@@ -351,6 +314,9 @@ struct FGatherInputParameters
 
 	// Runtime toggle when KernelSamplingDensityMode == KERNEL_DENSITY_HIGHER_IN_CENTER_DISK 
 	bool bDoHighDensityGatheringAtCenter[CONFIG_MAX_SAMPLING_DENSITY_CHANGES];
+
+	// Compile time number of rings for density changes.
+	uint DensityChangeRingCount;
 };
 
 /** Informations about a ring. */
diff --git a/Engine/Shaders/Private/DiaphragmDOF/DOFGatherKernel.ush b/Engine/Shaders/Private/DiaphragmDOF/DOFGatherKernel.ush
--- a/Engine/Shaders/Private/DiaphragmDOF/DOFGatherKernel.ush
+++ b/Engine/Shaders/Private/DiaphragmDOF/DOFGatherKernel.ush
@@ -68,10 +68,7 @@ struct FPrivateKernelSharedScope
 	// Multiplier to transform sample count unit to sample position in pixels of GatherInput's mip level 0.
 	// Set at compile time to 1 when bKenrelHasEnoughSamples==true.
 	float SampleCountToSamplePosition;
-			
-	// CocRadius error caused by undersampling.
-	float CocRadiusError;
-
+	
 	// Sampling mip level for GatherInput.
 	float GatherInputMipLevel;
 
@@ -81,6 +78,56 @@ struct FPrivateKernelSharedScope
 };
 
 
+// Computes intersection of sample from its Coc and its distance.
+// Returns 0: no intersection, 1: intersection.
+float ComputeSampleIntersection(FPrivateKernelSharedScope KernelScope, float SampleCocRadius, float SampleDistance)
+{
+	// Intersection of output pixel's disk area with sample's Coc.
+	#if 1
+		// To have an intersection of 50% when the Sample's Coc radius == SampleDistance.
+		#if 0 // TODO: clean this hack.
+			// Poor out of focus, great alpha channel, no leaking.
+			const float Offset = 1;
+		#elif 1
+			// Best slight out of focus.
+			const float Offset = 0;
+		#else
+			// Mathematically ideal.
+			const float Offset = 0.5;
+		#endif
+
+		#if DIM_LAYER_PROCESSING == LAYER_PROCESSING_SLIGHT_OUT_OF_FOCUS
+			// Uses a sharper intersection for slight out of focus like Recombine pass since it is being used
+			// for full res temporal statbility clamping box, so the aliasing this may introduce at half res is
+			// Fine.
+			const float Multiplier = 4.0;
+		#else
+			const float Multiplier = 1;
+			// TODO: investigate infficient way to reduce Multiplier according to MipLevel.
+			//const float Multiplier = SampleCocRadius > 4 ? 4.0 / abs(SampleCocRadius) : 1.0;
+		#endif
+
+		return saturate((abs(SampleCocRadius) - SampleDistance) * Multiplier + Offset);
+
+	// Same as Circle DOF.
+	#elif 0
+		// 0.5 because Coc radius is at half res, and 0.5 so that the fallof is only the radius of a ful res pixel.
+		const float IntersectionSharpen = 0.5 * 0.5;
+
+		// TODO: missing +0.5 to be energy conservative.
+		return saturate((abs(SampleCocRadius) - SampleDistance) * IntersectionSharpen);
+
+	#else // Hard cut of.
+		if (SampleDistance > abs(SampleCocRadius))
+		{
+			return 0.0;
+		}
+		return 1.0;
+
+	#endif
+}
+
+
 // Returns the total number of sampling iteration for a given ring id.
 uint GetRingSamplingPairCount(const uint KernelSamplingDensityMode, uint RingId)
 {
@@ -172,7 +219,7 @@ void UpdateRuntimeKernelScope(
 	// Compute the mip level to sample.
 	#if CONFIG_MIP_LEVEL_METHOD == MIP_LEVEL_METHOD_INTERLEAVED_OFFSET	
 		// Compute random uniformly distributed on disk (centered, diameter = 1).
-		float2 RandomCircle = (0.5 * sqrt(GatherParameters.Random[0])) * float2(cos(2 * PI * GatherParameters.Random[1]), sin(2 * PI * GatherParameters.Random[1]));
+		float2 RandomCircle = (0.48 * sqrt(GatherParameters.Random[0])) * float2(cos(2 * PI * GatherParameters.Random[1]), sin(2 * PI * GatherParameters.Random[1]));
 
 		float MipLevel = floor(0.5 + log2(GatherParameters.KernelRadius * UnitSampleRadius));
 
@@ -203,8 +250,12 @@ void UpdateRuntimeKernelScope(
 	}
 
 	// Compute the randomising range of the kernel offset in pixels.
-	float KernelRandomRange = max((GatherParameters.KernelRadius - GatherParameters.MaxRingCount) * UnitSampleRadius, 0.0);
-	
+	#if CONFIG_FULL_KERNEL_CENTER_RANDOM
+		float KernelRandomRange = max(GatherParameters.KernelRadius - GatherParameters.MaxRingCount - pow(2.0, uint(MipLevel)), 0.0);
+	#else
+		float KernelRandomRange = max((GatherParameters.KernelRadius - GatherParameters.MaxRingCount) * UnitSampleRadius, 0.0);
+	#endif
+
 	// No need to do kernel randomization when there is enough samples.
 	if (static_condition(KernelScope.bKenrelHasEnoughSamples))
 	{
@@ -214,7 +265,6 @@ void UpdateRuntimeKernelScope(
 	KernelScope.KernelCenterUV = GatherParameters.InputBufferCenterUV + (RandomCircle * KernelRandomRange) * PostprocessInput0Size.zw;
 
 	KernelScope.SampleCountToSamplePosition = GatherParameters.KernelRadius * UnitSampleRadius;
-	KernelScope.CocRadiusError = 0.5 * KernelRandomRange;
 
 	KernelScope.GatherInputMaxUV = (ViewportSize.xy - 0.5 * pow(2.0, uint(MipLevel))) * GatherParameters.DispatchThreadIdToInputBufferUV;
 	KernelScope.GatherInputMipLevel = MipLevel;
@@ -229,7 +279,6 @@ void UpdateRuntimeKernelScope(
 	#if CONFIG_SGPR_HINTS
 	{
 		KernelScope.SampleCountToSamplePosition = ToScalarMemory(KernelScope.SampleCountToSamplePosition);
-		KernelScope.CocRadiusError = ToScalarMemory(KernelScope.CocRadiusError);
 		KernelScope.GatherInputMaxUV = ToScalarMemory(KernelScope.GatherInputMaxUV);
 		KernelScope.GatherInputMipLevel = ToScalarMemory(KernelScope.GatherInputMipLevel);
 	}
@@ -466,7 +515,7 @@ void GatherRing(
 				#endif
 
 				Sample[k].Distance = SampleDistance;
-				Sample[k].Intersection = ComputeSampleIntersection(Sample[k].CocRadius, SampleDistance);
+				Sample[k].Intersection = ComputeSampleIntersection(KernelScope, Sample[k].CocRadius, SampleDistance);
 				
 				FastAccumulator.Color += Sample[k].Color;
 			}
@@ -555,7 +604,7 @@ void GatherSamplesAndResolve(
 	{
 		//static_assert(GatherParameters.MaxRingCount >= 2, "Assumes this LARGEST_RINGS_FIRST.");
 		
-		int LowerRingId = GatherParameters.MaxRingCount / 2;
+		int LowerRingId = GatherParameters.MaxRingCount - GatherParameters.DensityChangeRingCount;
 
 		int RingId;
 
diff --git a/Engine/Shaders/Private/DiaphragmDOF/DOFGatherPass.usf b/Engine/Shaders/Private/DiaphragmDOF/DOFGatherPass.usf
--- a/Engine/Shaders/Private/DiaphragmDOF/DOFGatherPass.usf
+++ b/Engine/Shaders/Private/DiaphragmDOF/DOFGatherPass.usf
@@ -172,6 +172,7 @@ void GatherMainCS(
 	{
 		GatherParameters.MaxRingCount = DIM_GATHER_RING_COUNT;
 		GatherParameters.bFastGathering = false;
+		GatherParameters.DensityChangeRingCount = GatherParameters.MaxRingCount / 2;
 
 		// TODO: variable number of ring.
 		GatherParameters.RingCount = GatherParameters.MaxRingCount;
@@ -368,12 +369,16 @@ void GatherMainCS(
 
 			GatherParameters.KernelSamplingDensityMode = KERNEL_DENSITY_HIGHER_IN_CENTER_DISK;
 
+			// Scaling factor applied on the samples to decrease their size when increasing density.
+			const float SampleScaleFactor = 1.0 - GatherParameters.DensityChangeRingCount / (GatherParameters.MaxRingCount + 0.5);
+			
+			// Compute whether need to change the kernel density into SGPR.
 			UNROLL
 			for (uint i = 0; i < CONFIG_MAX_SAMPLING_DENSITY_CHANGES; i++)
 			{
-				const float MinConvolutionRadius = float(8 << i);
-				const float SmallestConvolutionRadius = rcp(float(2 << i)) * TileSuggestion.MaxCocRadiusAbs;
-
+				const float MinConvolutionRadius = GatherParameters.MaxRingCount * pow(1 / SampleScaleFactor, i);
+				const float SmallestConvolutionRadius = TileSuggestion.MaxCocRadiusAbs * pow(SampleScaleFactor, i + 1);
+	
 				GatherParameters.bDoHighDensityGatheringAtCenter[i] = (
 					TileSuggestion.MaxCocRadiusAbs > MinConvolutionRadius &&
 					TileSuggestion.MinIntersectableCocRadiusAbs < SmallestConvolutionRadius);
diff --git a/Engine/Source/Runtime/Renderer/Private/PostProcess/DiaphragmDOFPasses.cpp b/Engine/Source/Runtime/Renderer/Private/PostProcess/DiaphragmDOFPasses.cpp
--- a/Engine/Source/Runtime/Renderer/Private/PostProcess/DiaphragmDOFPasses.cpp
+++ b/Engine/Source/Runtime/Renderer/Private/PostProcess/DiaphragmDOFPasses.cpp
@@ -624,7 +624,9 @@ FPooledRenderTargetDesc FRCPassDiaphragmDOFFlattenCoc::ComputeOutputDesc(EPassOu
 
 #define COC_DILATE_SHADER_PARAMS(PARAMETER) \
 	PARAMETER(FShaderParameter, SampleOffsetMultipler) \
+	PARAMETER(FShaderParameter, fSampleOffsetMultipler) \
 	PARAMETER(FShaderParameter, CocRadiusToBucketDistance) \
+	PARAMETER(FShaderParameter, BucketDistanceToCocRadius) \
 
 class FPostProcessCocDilateCS : public FPostProcessDiaphragmDOFShader
 {
@@ -654,8 +656,14 @@ void FRCPassDiaphragmDOFDilateCoc::Process(FRenderingCompositePassContext& Conte
 	{
 		SetShaderValue(Context.RHICmdList, DispatchCtx.ShaderRHI, DispatchCtx->SampleOffsetMultipler, Params.SampleDistanceMultiplier);
 	
+		float fSampleOffsetMultipler = Params.SampleDistanceMultiplier;
+		SetShaderValue(Context.RHICmdList, DispatchCtx.ShaderRHI, DispatchCtx->fSampleOffsetMultipler, fSampleOffsetMultipler);
+
 		float CocRadiusToBucketDistance = Params.PreProcessingToProcessingCocRadiusFactor;
 		SetShaderValue(Context.RHICmdList, DispatchCtx.ShaderRHI, DispatchCtx->CocRadiusToBucketDistance, CocRadiusToBucketDistance);
+
+		float BucketDistanceToCocRadius = 1.0f / CocRadiusToBucketDistance;
+		SetShaderValue(Context.RHICmdList, DispatchCtx.ShaderRHI, DispatchCtx->BucketDistanceToCocRadius, BucketDistanceToCocRadius);
 	}
 	DispatchCtx.Dispatch();
 }
