Cherry-pick 4074770: Implements mirror sampling for less distraction on screen edges.

diff --git a/Engine/Shaders/Private/DiaphragmDOF/DOFGatherCommon.ush b/Engine/Shaders/Private/DiaphragmDOF/DOFGatherCommon.ush
--- a/Engine/Shaders/Private/DiaphragmDOF/DOFGatherCommon.ush
+++ b/Engine/Shaders/Private/DiaphragmDOF/DOFGatherCommon.ush
@@ -127,10 +127,12 @@
 
 #elif DIM_GATHER_QUALITY == GATHER_QUALITY_HIGH_QUALITY
 	#define CONFIG_WEIGHT_SAMPLE_FROM_COC 1
+	#define CONFIG_MIRRORED_SAMPLER 1
 
 #elif DIM_GATHER_QUALITY == GATHER_QUALITY_SCATTER_OCCLUSION_OUTPUT
 	#define CONFIG_WEIGHT_SAMPLE_FROM_COC 1
 	#define CONFIG_GENERATE_SCATTER_OCCLUSION_BUFFER 1
+	#define CONFIG_MIRRORED_SAMPLER 1
 
 #else
 	#error Unknon gathering quality.
@@ -212,6 +214,11 @@
 	#define CONFIG_FULL_KERNEL_CENTER_RANDOM 0
 #endif
 
+// Uses mirrored sampler when gathering outside viewport rect.
+#ifndef CONFIG_MIRRORED_SAMPLER
+	#define CONFIG_MIRRORED_SAMPLER 0
+#endif
+
 // D3D shader compiler is rediculously slow to compile this shader.
 #if SM5_PROFILE
 	#define CONFIG_UNROLL_GATHER 0
diff --git a/Engine/Shaders/Private/DiaphragmDOF/DOFGatherKernel.ush b/Engine/Shaders/Private/DiaphragmDOF/DOFGatherKernel.ush
--- a/Engine/Shaders/Private/DiaphragmDOF/DOFGatherKernel.ush
+++ b/Engine/Shaders/Private/DiaphragmDOF/DOFGatherKernel.ush
@@ -158,7 +158,16 @@ uint GetKernelSampleCount(const uint KernelSamplingDensityMode, uint RingCount)
 // Fetch and unpack a sample from GatherInput.
 FGatherSample FetchGatherSample(in FPrivateKernelSharedScope KernelScope, float2 SampleUV)
 {
-	#if CONFIG_CLAMP_INPUT_BUFFER_UV
+	#if CONFIG_MIRRORED_SAMPLER
+	{
+		SampleUV = max(SampleUV, -SampleUV);
+		#if CONFIG_CLAMP_INPUT_BUFFER_UV
+			SampleUV = min(SampleUV, 2.0 * KernelScope.GatherInputMaxUV - SampleUV);
+		#else
+			SampleUV = min(SampleUV, 2.0 - SampleUV);
+		#endif
+	}
+	#elif CONFIG_CLAMP_INPUT_BUFFER_UV
 	{
 		SampleUV = min(SampleUV, KernelScope.GatherInputMaxUV);
 	}
