Cherry-pick 4159257: Reduces amount of tilling artifacts of bokehs.

diff --git a/Engine/Shaders/Private/DiaphragmDOF/DOFCocTileCommon.ush b/Engine/Shaders/Private/DiaphragmDOF/DOFCocTileCommon.ush
--- a/Engine/Shaders/Private/DiaphragmDOF/DOFCocTileCommon.ush
+++ b/Engine/Shaders/Private/DiaphragmDOF/DOFCocTileCommon.ush
@@ -6,6 +6,9 @@
 
 #pragma once
 
+#include "DOFCommon.ush"
+
+
 //------------------------------------------------------- COMPILE TIME CONSTANTS.
 
 // Maximum number of buffer for Coc tile.
diff --git a/Engine/Shaders/Private/DiaphragmDOF/DOFCommon.ush b/Engine/Shaders/Private/DiaphragmDOF/DOFCommon.ush
--- a/Engine/Shaders/Private/DiaphragmDOF/DOFCommon.ush
+++ b/Engine/Shaders/Private/DiaphragmDOF/DOFCommon.ush
@@ -106,6 +106,9 @@
 /** Whether alpha channel should be processed or not. */
 #define CONFIG_DOF_ALPHA (POST_PROCESS_ALPHA)
 
+/** Amount of error tolerated to quick fast gathering in. */
+#define FAST_GATHERING_COC_ERROR 0.05
+
 
 //------------------------------------------------------- COMPILE TIME CONSTANTS.
 
diff --git a/Engine/Shaders/Private/DiaphragmDOF/DOFGatherPass.usf b/Engine/Shaders/Private/DiaphragmDOF/DOFGatherPass.usf
--- a/Engine/Shaders/Private/DiaphragmDOF/DOFGatherPass.usf
+++ b/Engine/Shaders/Private/DiaphragmDOF/DOFGatherPass.usf
@@ -73,7 +73,12 @@ void StoreAccumulatorOutput(in uint2 OutputPixelPosition, in FAccumulatorOutput
 		#if 1
 			float OcclusionStepness = rcp(max(CocRadiusStdDeviation * (0.5 * CocRadiusAvg), 1));
 			float OcclusionOffset = AccumulatorOutput.BackgroundCocAvgAndSquareAvg.x * OcclusionStepness + 1;
-		
+
+			// Multiply OcclusionStepness, to be equivalent as multiplying AbsCocRadius. This will result in a
+			// pull of the background bokeh closer to the camera when doing occlusion, reducing the amount of tiling
+			// artifacts between fast and slow gathered tiles that can generate noticeably different occlusion buffer.
+			OcclusionStepness *= 1 - FAST_GATHERING_COC_ERROR;
+
 			OutScatterOcclusion.x = OcclusionStepness;
 			OutScatterOcclusion.y = OcclusionOffset;
 		#else
diff --git a/Engine/Shaders/Private/DiaphragmDOF/DOFGatherTileSuggest.ush b/Engine/Shaders/Private/DiaphragmDOF/DOFGatherTileSuggest.ush
--- a/Engine/Shaders/Private/DiaphragmDOF/DOFGatherTileSuggest.ush
+++ b/Engine/Shaders/Private/DiaphragmDOF/DOFGatherTileSuggest.ush
@@ -5,10 +5,10 @@
 	gather a given coc tile.
 =============================================================================*/
 
-#include "DOFCocTileCommon.ush"
-
 #pragma once
 
+#include "DOFCocTileCommon.ush"
+
 
 //------------------------------------------------------- FUNCTIONS
 
@@ -110,7 +110,7 @@ FGatheringTileSuggestion InferGatherTileSuggestion(
 	}
 	
 	// Coc error accepted for fast gathering, set to 5% of the kernel size.
-	float FastGatheringCocError = OutSuggestion.MaxCocRadiusAbs * 0.05;
+	float FastGatheringCocError = OutSuggestion.MaxCocRadiusAbs * FAST_GATHERING_COC_ERROR;
 
 	// Logic whether can entirely skeep gathering.
 	if (LayerProcessing == LAYER_PROCESSING_FOREGROUND_HOLE_FILLING)
diff --git a/Engine/Shaders/Private/DiaphragmDOF/DOFHybridScatterPixelShader.usf b/Engine/Shaders/Private/DiaphragmDOF/DOFHybridScatterPixelShader.usf
--- a/Engine/Shaders/Private/DiaphragmDOF/DOFHybridScatterPixelShader.usf
+++ b/Engine/Shaders/Private/DiaphragmDOF/DOFHybridScatterPixelShader.usf
@@ -74,6 +74,11 @@ void ScatterMainPS(
 			float CocRadiusStdDeviation = max(ScatterOcclusionSample.y * (0.5 * CocRadiusAvg), 1);
 			float OcclusionStepness = rcp(CocRadiusStdDeviation);
 			float OcclusionOffset = (CocRadiusAvg + CocRadiusStdDeviation) * OcclusionStepness;
+
+			// Multiply OcclusionStepness, to be equivalent as multiplying AbsCocRadius. This will result in a
+			// pull of the background bokeh closer to the camera when doing occlusion, reducing the amount of tiling
+			// artifacts between fast and slow gathered tiles that can generate noticeably different occlusion buffer.
+			OcclusionStepness *= 1 - FAST_GATHERING_COC_ERROR;
 		#endif
 	#endif
 
