Cherry-pick 4169332: Fixes a couple of bugs on coc tile flattening and gathering kernel offset.

diff --git a/Engine/Shaders/Private/DiaphragmDOF/DOFCocTileCommon.ush b/Engine/Shaders/Private/DiaphragmDOF/DOFCocTileCommon.ush
--- a/Engine/Shaders/Private/DiaphragmDOF/DOFCocTileCommon.ush
+++ b/Engine/Shaders/Private/DiaphragmDOF/DOFCocTileCommon.ush
@@ -58,8 +58,11 @@ struct FCocTileSample
 	// The maximum Coc for background >= 0.
 	float BgdMaxCoc;
 
-	// The maximum intersectable Coc radius.
+	// The minimum intersectable Coc radius >= 0.
 	float BgdMinIntersectableCoc;
+
+	// The conservative distance field to any closer geometry.
+	float BgdMinConservativeCloserGoemetryDistance;
 };
 
 
@@ -72,6 +75,7 @@ void InitCocTileSample(out FCocTileSample Sample)
 	Sample.BgdMinCoc = 0;
 	Sample.BgdMaxCoc = 0;
 	Sample.BgdMinIntersectableCoc = EXTREMELY_LARGE_COC_RADIUS;
+	Sample.BgdMinConservativeCloserGoemetryDistance = 0;
 }
 
 /** Encode coc tile sample.
@@ -92,7 +96,7 @@ void EncodeCocTileSample(in FCocTileSample Sample, in uint Layout, out float4 Ra
 	else if (static_condition(Layout == COC_TILE_LAYOUT_FGD_SEP_BGD))
 	{
 		Raw[0].rg = float2(Sample.FgdMinCoc, Sample.FgdMaxCoc);
-		Raw[1].rgb = float3(Sample.BgdMaxCoc, Sample.BgdMinCoc, Sample.BgdMinIntersectableCoc);
+		Raw[1] = float4(Sample.BgdMaxCoc, Sample.BgdMinCoc, Sample.BgdMinIntersectableCoc, Sample.BgdMinConservativeCloserGoemetryDistance);
 	}
 }
 
@@ -117,6 +121,7 @@ void DecodeCocTileSample(in float4 Raw[COC_TILE_MAX_BUFFER_COUNT], in uint Layou
 		Sample.BgdMinCoc = Raw[1].g;
 		Sample.BgdMaxCoc = Raw[1].r;
 		Sample.BgdMinIntersectableCoc = Raw[1].b;
+		Sample.BgdMinConservativeCloserGoemetryDistance = Raw[1].a;
 	}
 }
 
diff --git a/Engine/Shaders/Private/DiaphragmDOF/DOFCocTileDilate.usf b/Engine/Shaders/Private/DiaphragmDOF/DOFCocTileDilate.usf
--- a/Engine/Shaders/Private/DiaphragmDOF/DOFCocTileDilate.usf
+++ b/Engine/Shaders/Private/DiaphragmDOF/DOFCocTileDilate.usf
@@ -12,8 +12,8 @@
 
 /** Matches FRCPassDiaphragmDOFDilateCoc::EMode. */
 	#define DILATE_MODE_STANDALONE 0
-	#define GATHER_MODE_MIN_MAX 1
-	#define GATHER_MODE_MIN_ABS 2
+	#define DILATE_MODE_MIN_MAX 1
+	#define DILATE_MODE_MIN_ABS 2
 
 
 //------------------------------------------------------- PASSES.
@@ -22,7 +22,7 @@ uint SampleOffsetMultipler = 1;
 float fSampleOffsetMultipler = 1;
 
 // Multiplier to compare CocRadius and bucket distance.
-float CocRadiusToBucketDistance;
+float CocRadiusToBucketDistanceUpperBound;
 float BucketDistanceToCocRadius;
 
 
@@ -49,10 +49,15 @@ void CocDilateMainCS(uint2 DispatchThreadId : SV_DispatchThreadID)
 	// Standalone does not have any radius multiplier since there is only one dilate pass.
 	#if DIM_DILATE_MODE == DILATE_MODE_STANDALONE
 		const uint RadiusMultiplier = 1;
-		const uint fRadiusMultiplier = 1;
+		const float fRadiusMultiplier = 1;
 	#else
 		const uint RadiusMultiplier = SampleOffsetMultipler;
-		const uint fRadiusMultiplier = fSampleOffsetMultipler;
+		const float fRadiusMultiplier = fSampleOffsetMultipler;
+	#endif
+
+	// Fetches the already dilated min foreground and max background COC radius.
+	#if DIM_DILATE_MODE == DILATE_MODE_MIN_ABS
+		FCocTileSample DilatedMinMaxSample = LoadCocTile(CocTileLayout, PostprocessInput2, PostprocessInput3, PixelPos);
 	#endif
 
 	// Iterate over all input tiles and assign them into their buckets.
@@ -67,6 +72,7 @@ void CocDilateMainCS(uint2 DispatchThreadId : SV_DispatchThreadID)
 		OutSampleBuckets[BucketId].BgdMinCoc = EXTREMELY_LARGE_COC_RADIUS;
 		OutSampleBuckets[BucketId].BgdMaxCoc = 0;
 		OutSampleBuckets[BucketId].BgdMinIntersectableCoc = EXTREMELY_LARGE_COC_RADIUS;
+		OutSampleBuckets[BucketId].BgdMinConservativeCloserGoemetryDistance = EXTREMELY_LARGE_COC_RADIUS;
 
 		UNROLL
 		for (uint BatchId = 0; BatchId < 4 * RingDistance; BatchId++)
@@ -107,7 +113,7 @@ void CocDilateMainCS(uint2 DispatchThreadId : SV_DispatchThreadID)
 			UNROLL
 			for (uint SampleId = 0; SampleId < 2; SampleId++)
 			{
-				const int2 Offset = (SampleId == 0 ? 1 : -1) * BatchSquarePos;
+				const int2 Offset = (SampleId == 0 ? 1 : -1) * (int(RadiusMultiplier) * BatchSquarePos);
 				int2 SampleIndex = int2(PixelPos) + Offset;
 			
 				// If sampling outside the tiles, ignore this sample.
@@ -118,7 +124,8 @@ void CocDilateMainCS(uint2 DispatchThreadId : SV_DispatchThreadID)
 					const float MinimalCoC = 0.5 * kCocRadiusToFullResFactor;
 
 					// Closest distance from sampling tile to output tile.
-					float ClosestDistance = fRadiusMultiplier * (COC_TILE_SIZE * length(max(abs(float2(Offset)) - 1.0, 0.0)));
+					float ClosestDistance = length(max(fRadiusMultiplier * (COC_TILE_SIZE * abs(float2(Offset))) - COC_TILE_SIZE, 0.0));
+					//ClosestDistance = fRadiusMultiplier * float((BucketId + 1) * COC_TILE_SIZE) - COC_TILE_SIZE;
 
 					// Sample tile.
 					const FCocTileSample NeightborSample = LoadCocTile(CocTileLayout, PostprocessInput0, PostprocessInput1, SampleIndex);
@@ -131,17 +138,26 @@ void CocDilateMainCS(uint2 DispatchThreadId : SV_DispatchThreadID)
 						OutSampleBuckets[BucketId].FgdMaxCoc,
 						NeightborSample.FgdMaxCoc);
 
-					OutSampleBuckets[BucketId].BgdMinCoc = min(
-						OutSampleBuckets[BucketId].BgdMinCoc,
-						NeightborSample.BgdMinCoc);
-
-					OutSampleBuckets[BucketId].BgdMaxCoc = max(
-						OutSampleBuckets[BucketId].BgdMaxCoc,
-						NeightborSample.BgdMaxCoc);
+					#if DIM_DILATE_MODE == DILATE_MODE_MIN_ABS && 0 // TODO.
+						if (NeightborSample.BgdMinConservativeCloserGoemetryDistance + ClosestDistance * BucketDistanceToCocRadius < 2 * DilatedMinMaxSample.BgdMaxCoc)
+					#endif
+					{
+						OutSampleBuckets[BucketId].BgdMinCoc = min(
+							OutSampleBuckets[BucketId].BgdMinCoc,
+							NeightborSample.BgdMinCoc);
+
+						OutSampleBuckets[BucketId].BgdMaxCoc = max(
+							OutSampleBuckets[BucketId].BgdMaxCoc,
+							NeightborSample.BgdMaxCoc);
 					
-					OutSampleBuckets[BucketId].BgdMinIntersectableCoc = min(
-						OutSampleBuckets[BucketId].BgdMinIntersectableCoc,
-						NeightborSample.BgdMinIntersectableCoc + ClosestDistance * BucketDistanceToCocRadius);
+						OutSampleBuckets[BucketId].BgdMinIntersectableCoc = min(
+							OutSampleBuckets[BucketId].BgdMinIntersectableCoc,
+							NeightborSample.BgdMinIntersectableCoc + ClosestDistance * BucketDistanceToCocRadius);
+					}
+							
+					OutSampleBuckets[BucketId].BgdMinConservativeCloserGoemetryDistance = min(
+						OutSampleBuckets[BucketId].BgdMinConservativeCloserGoemetryDistance,
+						NeightborSample.BgdMinConservativeCloserGoemetryDistance + ClosestDistance * BucketDistanceToCocRadius);
 				}
 			} // for (uint SampleId = 0; SampleId < 2; SampleId++)
 		} // for (uint BatchId = 0; BatchId < 4 * RingDistance; BatchId++)
@@ -150,72 +166,71 @@ void CocDilateMainCS(uint2 DispatchThreadId : SV_DispatchThreadID)
 	// Sample the center bucket.
 	FCocTileSample OutSample = LoadCocTile(CocTileLayout, PostprocessInput0, PostprocessInput1, PixelPos);
 
+	#if DIM_DILATE_MODE == DILATE_MODE_MIN_ABS
+	{
+		OutSample.FgdMinCoc = DilatedMinMaxSample.FgdMinCoc;
+		OutSample.BgdMaxCoc = DilatedMinMaxSample.BgdMaxCoc;
+	}
+	#else // DIM_DILATE_MODE != DILATE_MODE_MIN_ABS
 	{
-		uint k;
-	
 		// Find out what the maximum absolute Coc radii will be for the output sample.
 		UNROLL
-		for (k = 0; k < Radius; k++)
+		for (uint k = 0; k < Radius; k++)
 		{
-			float BucketDistance = fRadiusMultiplier * (COC_TILE_SIZE * float(k + 1)) - COC_TILE_SIZE;
-			
+			float BucketDistance = fRadiusMultiplier * (float(k + 1) * COC_TILE_SIZE) - COC_TILE_SIZE;
+
 			FLATTEN
-			if (-OutSampleBuckets[k].FgdMinCoc * CocRadiusToBucketDistance > BucketDistance)
+			if (-OutSampleBuckets[k].FgdMinCoc * CocRadiusToBucketDistanceUpperBound > BucketDistance)
 			{
 				OutSample.FgdMinCoc = min(OutSample.FgdMinCoc, OutSampleBuckets[k].FgdMinCoc);
 			}
 			
 			FLATTEN
-			if (OutSampleBuckets[k].BgdMaxCoc * CocRadiusToBucketDistance > BucketDistance)
+			if (OutSampleBuckets[k].BgdMaxCoc * CocRadiusToBucketDistanceUpperBound > BucketDistance)
 			{
 				OutSample.BgdMaxCoc = max(OutSample.BgdMaxCoc, OutSampleBuckets[k].BgdMaxCoc);
 			}
 		}
+	}
+	#endif
 		
-		// Fetches the already dilated min foreground and max background COC radius.
-		#if DIM_DILATE_MODE == GATHER_MODE_MIN_ABS
-		{
-			FCocTileSample DilatedMinMaxSample = LoadCocTile(CocTileLayout, PostprocessInput2, PostprocessInput3, PixelPos);
-
-			OutSample.FgdMinCoc = DilatedMinMaxSample.FgdMinCoc;
-			OutSample.BgdMaxCoc = DilatedMinMaxSample.BgdMaxCoc;
-		}
-		#endif
-
+	// Then, find out what is minimum absolute Coc radii that will be intersected for these maximum Coc radii.
+	{
 		const float CocRadiusErrorFraction = 0.05;
 		float ForegroundCocRadiusError = -OutSample.FgdMinCoc * CocRadiusErrorFraction;
 		float BackgroundCocRadiusError = OutSample.BgdMaxCoc * CocRadiusErrorFraction;
 
-
-		// Then, find out what is minimum absolute Coc radii that will be intersected for these maximum Coc radii.
 		UNROLL
-		for (k = 0; k < Radius; k++)
+		for (uint k = 0; k < Radius; k++)
 		{
-			float BucketDistance = fRadiusMultiplier * ((float(k + 1) - 1.0) * float(COC_TILE_SIZE));
+			float BucketDistance = fRadiusMultiplier * (float(k + 1) * COC_TILE_SIZE) - COC_TILE_SIZE;
 			
 			const float MinimalCoC = 0.5 * kCocRadiusToFullResFactor;
-			
+
 			FLATTEN
-			if (-OutSample.FgdMinCoc * CocRadiusToBucketDistance > BucketDistance)
+			if (-OutSample.FgdMinCoc * CocRadiusToBucketDistanceUpperBound > BucketDistance)
 			{
 				OutSample.FgdMaxCoc = max(OutSample.FgdMaxCoc, OutSampleBuckets[k].FgdMaxCoc);
 			}
 
 			FLATTEN
-			if (OutSample.BgdMaxCoc * CocRadiusToBucketDistance > BucketDistance)
+			if (OutSample.BgdMaxCoc * CocRadiusToBucketDistanceUpperBound > BucketDistance)
 			{
 				OutSample.BgdMinCoc = min(OutSample.BgdMinCoc, OutSampleBuckets[k].BgdMinCoc);
 				OutSample.BgdMinIntersectableCoc = min(OutSample.BgdMinIntersectableCoc, OutSampleBuckets[k].BgdMinIntersectableCoc);
 			}
+
+			OutSample.BgdMinConservativeCloserGoemetryDistance = min(OutSample.BgdMinConservativeCloserGoemetryDistance, OutSampleBuckets[k].BgdMinConservativeCloserGoemetryDistance);
 		}
 	}
 
 	// If dilating min foreground and max background coc, zero output other to prune the unecessary ALU.
-	#if DIM_DILATE_MODE == GATHER_MODE_MIN_MAX
+	#if DIM_DILATE_MODE == DILATE_MODE_MIN_MAX
 	{
 		OutSample.FgdMaxCoc = 0;
 		OutSample.BgdMinCoc = 0;
 		OutSample.BgdMinIntersectableCoc = 0;
+		OutSample.BgdMinConservativeCloserGoemetryDistance = 0;
 	}
 	#endif
 
diff --git a/Engine/Shaders/Private/DiaphragmDOF/DOFCocTileFlatten.usf b/Engine/Shaders/Private/DiaphragmDOF/DOFCocTileFlatten.usf
--- a/Engine/Shaders/Private/DiaphragmDOF/DOFCocTileFlatten.usf
+++ b/Engine/Shaders/Private/DiaphragmDOF/DOFCocTileFlatten.usf
@@ -59,6 +59,7 @@ void ReduceOperator(inout FCocTileSample A, in const FCocTileSample B)
 
 	A.BgdMinCoc = min(A.BgdMinCoc, B.BgdMinCoc);
 	A.BgdMaxCoc = max(A.BgdMaxCoc, B.BgdMaxCoc);
+	A.BgdMinIntersectableCoc = min(A.BgdMinIntersectableCoc, B.BgdMinIntersectableCoc);
 }
 
 
@@ -197,6 +198,8 @@ void CocFlattenMainCS(
 			// Final reduction.
 			ReduceOperator(Local, LoadFromLds(1));
 			
+			Local.BgdMinConservativeCloserGoemetryDistance = Local.BgdMinCoc;
+
 			// Encode.
 			float4 RawOutput[COC_TILE_MAX_BUFFER_COUNT];
 			EncodeCocTileSample(Local, LDS_LAYOUT, RawOutput);
@@ -282,6 +285,8 @@ void CocFlattenMainCS(
 				OutputTileSample.BgdMinIntersectableCoc = asfloat(SharedBackgroundMinIntersectableCocRadius);
 			}
 
+			OutputTileSample.BgdMinConservativeCloserGoemetryDistance = OutputTileSample.BgdMinCoc;
+
 			// Encode.
 			float4 RawOutput[COC_TILE_MAX_BUFFER_COUNT];
 			EncodeCocTileSample(OutputTileSample, LDS_LAYOUT, RawOutput);
diff --git a/Engine/Shaders/Private/DiaphragmDOF/DOFGatherCommon.ush b/Engine/Shaders/Private/DiaphragmDOF/DOFGatherCommon.ush
--- a/Engine/Shaders/Private/DiaphragmDOF/DOFGatherCommon.ush
+++ b/Engine/Shaders/Private/DiaphragmDOF/DOFGatherCommon.ush
@@ -104,9 +104,6 @@
 	// Background need coc radius error correction to reduce amount of background leaking.
 	#define CONFIG_COC_RADIUS_CORRECTION 1
 
-	// TODO: investigate this for foreground.
-	#define CONFIG_FULL_KERNEL_CENTER_RANDOM (DIM_GATHER_QUALITY != GATHER_QUALITY_LOW_QUALITY)
-
 #elif DIM_LAYER_PROCESSING == LAYER_PROCESSING_SLIGHT_OUT_OF_FOCUS
 	#define GATHER_ACCUMULATOR_METHOD (GATHER_ACCUMULATOR_METHOD_SIMPLE)
 
@@ -209,11 +206,6 @@
 	#define CONFIG_POINT_MIRROR_BOKEH 1
 #endif
 
-// Whether to randomise the kernel center on the full displacement.
-#ifndef CONFIG_FULL_KERNEL_CENTER_RANDOM
-	#define CONFIG_FULL_KERNEL_CENTER_RANDOM 0
-#endif
-
 // Uses mirrored sampler when gathering outside viewport rect.
 #ifndef CONFIG_MIRRORED_SAMPLER
 	#define CONFIG_MIRRORED_SAMPLER 0
@@ -319,11 +311,14 @@ struct FGatherInputParameters
 	// Sampling density compile time mode. KERNEL_DENSITY_*
 	uint KernelSamplingDensityMode;
 
-	// Runtime toggle when KernelSamplingDensityMode == KERNEL_DENSITY_HIGHER_IN_CENTER_DISK 
-	bool bDoHighDensityGatheringAtCenter[CONFIG_MAX_SAMPLING_DENSITY_CHANGES];
-
 	// Compile time number of rings for density changes.
 	uint DensityChangeRingCount;
+
+	// Compile time max number of density changes allowed.
+	uint MaxDensityChangeCount;
+	
+	// Tile constant about what is the smallest intersectable radius.
+	float MinIntersectableCocRadiusAbs;
 };
 
 /** Informations about a ring. */
diff --git a/Engine/Shaders/Private/DiaphragmDOF/DOFGatherKernel.ush b/Engine/Shaders/Private/DiaphragmDOF/DOFGatherKernel.ush
--- a/Engine/Shaders/Private/DiaphragmDOF/DOFGatherKernel.ush
+++ b/Engine/Shaders/Private/DiaphragmDOF/DOFGatherKernel.ush
@@ -75,6 +75,9 @@ struct FPrivateKernelSharedScope
 	// GatherInput's max UVs.
 	// TODO: this is currently MipLevel  depent, but could save some VGPR by making it independent.
 	float2 GatherInputMaxUV;
+
+	// Multiplier to apply when computing sample intersection. Higher value sharpen, lower value feather.
+	float IntersectionMultiplier;
 };
 
 
@@ -102,9 +105,7 @@ float ComputeSampleIntersection(FPrivateKernelSharedScope KernelScope, float Sam
 			// Fine.
 			const float Multiplier = 4.0;
 		#else
-			const float Multiplier = 1;
-			// TODO: investigate infficient way to reduce Multiplier according to MipLevel.
-			//const float Multiplier = SampleCocRadius > 4 ? 4.0 / abs(SampleCocRadius) : 1.0;
+			const float Multiplier = KernelScope.IntersectionMultiplier;
 		#endif
 
 		return saturate((abs(SampleCocRadius) - SampleDistance) * Multiplier + Offset);
@@ -259,11 +260,11 @@ void UpdateRuntimeKernelScope(
 	}
 
 	// Compute the randomising range of the kernel offset in pixels.
-	#if CONFIG_FULL_KERNEL_CENTER_RANDOM
-		float KernelRandomRange = max(GatherParameters.KernelRadius - GatherParameters.MaxRingCount - pow(2.0, uint(MipLevel)), 0.0);
-	#else
-		float KernelRandomRange = max((GatherParameters.KernelRadius - GatherParameters.MaxRingCount) * UnitSampleRadius, 0.0);
-	#endif
+	float KernelRandomRange = GatherParameters.KernelRadius * UnitSampleRadius;
+	if (static_condition(KernelScope.bBilinearSampleInput))
+	{
+		KernelRandomRange = max(GatherParameters.KernelRadius - GatherParameters.MaxRingCount - pow(2.0, uint(MipLevel)), 0.0) * UnitSampleRadius;
+	}
 
 	// No need to do kernel randomization when there is enough samples.
 	if (static_condition(KernelScope.bKenrelHasEnoughSamples))
@@ -275,8 +276,10 @@ void UpdateRuntimeKernelScope(
 
 	KernelScope.SampleCountToSamplePosition = GatherParameters.KernelRadius * UnitSampleRadius;
 
-	KernelScope.GatherInputMaxUV = (ViewportSize.xy - 0.5 * pow(2.0, uint(MipLevel))) * GatherParameters.DispatchThreadIdToInputBufferUV;
+	float MipTexelSize = pow(2.0, uint(MipLevel));
+	KernelScope.GatherInputMaxUV = (ViewportSize.xy - 0.5 * MipTexelSize) * GatherParameters.DispatchThreadIdToInputBufferUV;
 	KernelScope.GatherInputMipLevel = MipLevel;
+	KernelScope.IntersectionMultiplier = rcp(MipTexelSize);
 		
 	// When there is enough samples, the sample count measurement and the sample position measurement are same.
 	if (static_condition(KernelScope.bKenrelHasEnoughSamples))
@@ -290,6 +293,7 @@ void UpdateRuntimeKernelScope(
 		KernelScope.SampleCountToSamplePosition = ToScalarMemory(KernelScope.SampleCountToSamplePosition);
 		KernelScope.GatherInputMaxUV = ToScalarMemory(KernelScope.GatherInputMaxUV);
 		KernelScope.GatherInputMipLevel = ToScalarMemory(KernelScope.GatherInputMipLevel);
+		KernelScope.IntersectionMultiplier = ToScalarMemory(KernelScope.IntersectionMultiplier);
 	}
 	#endif
 }
@@ -323,6 +327,13 @@ void GatherRing(
 	inout FGatherAccumulator Accumulator,
 	inout FFastAccumulatorOutput FastAccumulator)
 {
+	// Gives hints to the compiler that these var can be shifted to SGPR.
+	#if CONFIG_SGPR_HINTS
+	{
+		GatherParameters.KernelRadius = ToScalarMemory(GatherParameters.KernelRadius);
+	}
+	#endif
+
 	// Number of sample iteration for this ring.
 	const uint RingSamplePairCount = GetRingSamplingPairCount(GatherParameters.KernelSamplingDensityMode, uint(RingId));
 
@@ -540,12 +551,14 @@ void GatherRing(
 
 
 // Gathers samples with an preconfigured accumulator and resolve its output.
-void GatherSamplesAndResolve(
+uint GatherSamplesAndResolve(
 	in FGatherInputParameters GatherParameters,
 	in FGatherAccumulator Accumulator,
 	out FAccumulatorOutput AccumulatorOutput, 
 	out FFastAccumulatorOutput FastAccumulator)
 {
+	uint DebugGatherDensityChanges = 0;
+
 	#if DEBUG_FORCE_LOW_RING_COUNT
 		GatherParameters.MaxRingCount = 2;
 	#endif
@@ -631,24 +644,40 @@ void GatherSamplesAndResolve(
 			bIsFirstRing = false;
 		} // for (int RingId)
 		
-		UNROLL
-		for (uint SamplingDensityChangeId = 0; SamplingDensityChangeId < CONFIG_MAX_SAMPLING_DENSITY_CHANGES; SamplingDensityChangeId++)
+		UNROLL // TODO: LOOP instead? should be totally fine perf wise...
+		for (uint SamplingDensityChangeId = 0; SamplingDensityChangeId < GatherParameters.MaxDensityChangeCount; SamplingDensityChangeId++)
 		{
+			const int LargerRingCount = GatherParameters.MaxRingCount - LowerRingId;
+
+			const float LargeKernelRadius = 1 + 2 * GatherParameters.MaxRingCount;
+			const float SmallKernelRadius = 1 + 2 * (GatherParameters.MaxRingCount - LargerRingCount);
+
+			const float RadiusDownScaleFactor = SmallKernelRadius / LargeKernelRadius;
+					
+			float MipTexelSize = pow(2.0, uint(KernelScope.GatherInputMipLevel));
+			
+			const float MinConvolutionRadius = GatherParameters.MaxRingCount;
+			const float SmallestConvolutionRadius = GatherParameters.KernelRadius * RadiusDownScaleFactor + MipTexelSize;
+	
+			bool bChangeDensity = (
+				GatherParameters.KernelRadius * RadiusDownScaleFactor > MinConvolutionRadius &&
+				GatherParameters.MinIntersectableCocRadiusAbs < SmallestConvolutionRadius);
+
 			BRANCH
-			if (GatherParameters.bDoHighDensityGatheringAtCenter[SamplingDensityChangeId])
+			if (bChangeDensity)
 			{
 				// Decrease kernel size, and changes the already accumulated weight.
 				{
-					const int LargerRingCount = GatherParameters.MaxRingCount - LowerRingId;
-
-					const float LargeKernelRadius = 1 + 2 * GatherParameters.MaxRingCount;
-					const float SmallKernelRadius = 1 + 2 * (GatherParameters.MaxRingCount - LargerRingCount);
-
-					const float RadiusDownScaleFactor = SmallKernelRadius / LargeKernelRadius;
-
-					const float AccumulatedWeightMultiplier = rcp(RadiusDownScaleFactor * RadiusDownScaleFactor);
-				
 					GatherParameters.KernelRadius *= RadiusDownScaleFactor;
+					
+					// Gives hints to the compiler that these var can be shifted to SGPR.
+					#if CONFIG_SGPR_HINTS
+					{
+						GatherParameters.KernelRadius = ToScalarMemory(GatherParameters.KernelRadius);
+					}
+					#endif
+					
+					const float AccumulatedWeightMultiplier = rcp(RadiusDownScaleFactor * RadiusDownScaleFactor);
 
 					AmendAccumulatedRingsWeight(/* inout */ Accumulator, GatherParameters, AccumulatedWeightMultiplier);
 					AmendAccumulatedRingsWeight(/* inout */ FastAccumulator, AccumulatedWeightMultiplier);
@@ -666,6 +695,8 @@ void GatherSamplesAndResolve(
 						/* inout */ Accumulator,
 						/* inout */ FastAccumulator);
 				} // for (int RingId)
+
+				DebugGatherDensityChanges++;
 			}
 			else
 			{
@@ -756,4 +787,6 @@ void GatherSamplesAndResolve(
 	
 	// Resolves the accumulator.
 	ResolveAccumulatorOutputs(Accumulator, GatherResolveInfos, /* out */ AccumulatorOutput);
+
+	return DebugGatherDensityChanges;
 }
diff --git a/Engine/Shaders/Private/DiaphragmDOF/DOFGatherPass.usf b/Engine/Shaders/Private/DiaphragmDOF/DOFGatherPass.usf
--- a/Engine/Shaders/Private/DiaphragmDOF/DOFGatherPass.usf
+++ b/Engine/Shaders/Private/DiaphragmDOF/DOFGatherPass.usf
@@ -177,16 +177,14 @@ void GatherMainCS(
 	{
 		GatherParameters.MaxRingCount = DIM_GATHER_RING_COUNT;
 		GatherParameters.bFastGathering = false;
+
+		GatherParameters.KernelSamplingDensityMode = KERNEL_DENSITY_CONSTANT;
 		GatherParameters.DensityChangeRingCount = GatherParameters.MaxRingCount / 2;
+		GatherParameters.MaxDensityChangeCount = CONFIG_MAX_SAMPLING_DENSITY_CHANGES;
 
 		// TODO: variable number of ring.
 		GatherParameters.RingCount = GatherParameters.MaxRingCount;
 		
-		GatherParameters.KernelSamplingDensityMode = KERNEL_DENSITY_CONSTANT;
-		UNROLL
-		for (uint i = 0; i < CONFIG_MAX_SAMPLING_DENSITY_CHANGES; i++)
-			GatherParameters.bDoHighDensityGatheringAtCenter[i] = false;
-		
 		GatherParameters.ConsiderCocRadiusAffineTransformation0 = kContantlyPassingAffineTransformation;
 		GatherParameters.ConsiderCocRadiusAffineTransformation1 = kContantlyPassingAffineTransformation;
 		GatherParameters.ConsiderAbsCocRadiusAffineTransformation = kContantlyPassingAffineTransformation;
@@ -249,6 +247,7 @@ void GatherMainCS(
 		// Kernel radius is simply set to the maximum absolute Coc size.
 		GatherParameters.ClosestCocRadius = TileSuggestion.ClosestCocRadius;
 		GatherParameters.KernelRadius = TileSuggestion.MaxCocRadiusAbs;
+		GatherParameters.MinIntersectableCocRadiusAbs = TileSuggestion.MinIntersectableCocRadiusAbs;
 
 		#if DIM_LAYER_PROCESSING == LAYER_PROCESSING_SLIGHT_OUT_OF_FOCUS
 		{
@@ -373,31 +372,17 @@ void GatherMainCS(
 			#endif
 
 			GatherParameters.KernelSamplingDensityMode = KERNEL_DENSITY_HIGHER_IN_CENTER_DISK;
-
-			// Scaling factor applied on the samples to decrease their size when increasing density.
-			const float SampleScaleFactor = 1.0 - GatherParameters.DensityChangeRingCount / (GatherParameters.MaxRingCount + 0.5);
-			
-			// Compute whether need to change the kernel density into SGPR.
-			UNROLL
-			for (uint i = 0; i < CONFIG_MAX_SAMPLING_DENSITY_CHANGES; i++)
-			{
-				const float MinConvolutionRadius = GatherParameters.MaxRingCount * pow(1 / SampleScaleFactor, i);
-				const float SmallestConvolutionRadius = TileSuggestion.MaxCocRadiusAbs * pow(SampleScaleFactor, i + 1);
-	
-				GatherParameters.bDoHighDensityGatheringAtCenter[i] = (
-					TileSuggestion.MaxCocRadiusAbs > MinConvolutionRadius &&
-					TileSuggestion.MinIntersectableCocRadiusAbs < SmallestConvolutionRadius);
-			}
 		}
 		#endif
 
 		// Create, use and resolve gather accumulator.
 		FAccumulatorOutput AccumulatorOutput = CreateAccumulatorOutput();
+		uint DebugGatherDensityChanges;
 		{
 			FFastAccumulatorOutput UnusedFastAccumulatorOutput;
 
 			FGatherAccumulator FirstAccumulator = CreateGatherAccumulator(GatherParameters, LargerConvolution);
-			GatherSamplesAndResolve(GatherParameters, FirstAccumulator, AccumulatorOutput, UnusedFastAccumulatorOutput);
+			DebugGatherDensityChanges = GatherSamplesAndResolve(GatherParameters, FirstAccumulator, AccumulatorOutput, UnusedFastAccumulatorOutput);
 		}
 
 		#if DEBUG_FAST_ACCUMULATOR
@@ -406,7 +391,7 @@ void GatherMainCS(
 			UNROLL
 			for (uint i = 0; i < CONFIG_MAX_SAMPLING_DENSITY_CHANGES; i++)
 			{
-				if (GatherParameters.bDoHighDensityGatheringAtCenter[i])
+				if (i < DebugGatherDensityChanges)
 				{
 					AccumulatorOutput.ForegroundColor.rgb *= float3(1, 0.5, 0.5);
 					AccumulatorOutput.BackgroundColor.rgb *= float3(1, 0.5, 0.5);
diff --git a/Engine/Source/Runtime/Renderer/Private/PostProcess/DiaphragmDOFGraph.cpp b/Engine/Source/Runtime/Renderer/Private/PostProcess/DiaphragmDOFGraph.cpp
--- a/Engine/Source/Runtime/Renderer/Private/PostProcess/DiaphragmDOFGraph.cpp
+++ b/Engine/Source/Runtime/Renderer/Private/PostProcess/DiaphragmDOFGraph.cpp
@@ -347,6 +347,9 @@ bool DiaphragmDOF::WireSceneColorPasses(FPostprocessContext& Context, const FRen
 		CocTileOutput0 = FRenderingCompositeOutputRef(CocFlatten, ePId_Output0);
 		CocTileOutput1 = FRenderingCompositeOutputRef(CocFlatten, ePId_Output1);
 
+		// Error introduced by the random offset of the gathering kernel's center.
+		const float BluringRadiusErrorMultiplier = 1.0f + 1.0f / (HalfResRingCount + 0.5f);
+
 		// Parameters for the dilate Coc passes.
 		int32 DilateCount = 1;
 		FRCPassDiaphragmDOFDilateCoc::FParameters DilateParams[3];
@@ -355,7 +358,7 @@ bool DiaphragmDOF::WireSceneColorPasses(FPostprocessContext& Context, const FRen
 
 			// Compute the maximum tile dilation.
 			int32 MaximumTileDilation = FMath::CeilToInt(
-				MaxBluringRadius / FRCPassDiaphragmDOFFlattenCoc::CocTileResolutionDivisor);
+				(MaxBluringRadius * BluringRadiusErrorMultiplier) / FRCPassDiaphragmDOFFlattenCoc::CocTileResolutionDivisor);
 
 			// There is always at least one dilate pass so that even small Coc radius conservatively dilate on next neighboor.
 			DilateParams[0].SampleRadiusCount = FMath::Min(
@@ -396,6 +399,7 @@ bool DiaphragmDOF::WireSceneColorPasses(FPostprocessContext& Context, const FRen
 		{
 			DilateParams[i].GatherViewSize = GatheringViewSize;
 			DilateParams[i].PreProcessingToProcessingCocRadiusFactor = PreProcessingToProcessingCocRadiusFactor;
+			DilateParams[i].BluringRadiusErrorMultiplier = BluringRadiusErrorMultiplier;
 		}
 
 		if (DilateCount > 1)
diff --git a/Engine/Source/Runtime/Renderer/Private/PostProcess/DiaphragmDOFPasses.cpp b/Engine/Source/Runtime/Renderer/Private/PostProcess/DiaphragmDOFPasses.cpp
--- a/Engine/Source/Runtime/Renderer/Private/PostProcess/DiaphragmDOFPasses.cpp
+++ b/Engine/Source/Runtime/Renderer/Private/PostProcess/DiaphragmDOFPasses.cpp
@@ -615,7 +615,7 @@ FPooledRenderTargetDesc FRCPassDiaphragmDOFFlattenCoc::ComputeOutputDesc(EPassOu
 
 	FPooledRenderTargetDesc Ret(FPooledRenderTargetDesc::Create2DDesc(TileCount, PF_FloatRGBA, FClearValueBinding::None, TexCreate_None, TexCreate_RenderTargetable | TexCreate_UAV, false));
 	Ret.DebugName = InPassOutputId == ePId_Output0 ? TEXT("DOFFlattenFgdCoc") : TEXT("DOFFlattenBgdCoc");
-	Ret.Format = InPassOutputId == ePId_Output0 ? PF_G16R16F : PF_FloatRGB;
+	Ret.Format = InPassOutputId == ePId_Output0 ? PF_G16R16F : PF_FloatRGBA;
 	return Ret;
 }
 
@@ -625,7 +625,7 @@ FPooledRenderTargetDesc FRCPassDiaphragmDOFFlattenCoc::ComputeOutputDesc(EPassOu
 #define COC_DILATE_SHADER_PARAMS(PARAMETER) \
 	PARAMETER(FShaderParameter, SampleOffsetMultipler) \
 	PARAMETER(FShaderParameter, fSampleOffsetMultipler) \
-	PARAMETER(FShaderParameter, CocRadiusToBucketDistance) \
+	PARAMETER(FShaderParameter, CocRadiusToBucketDistanceUpperBound) \
 	PARAMETER(FShaderParameter, BucketDistanceToCocRadius) \
 
 class FPostProcessCocDilateCS : public FPostProcessDiaphragmDOFShader
@@ -659,10 +659,10 @@ void FRCPassDiaphragmDOFDilateCoc::Process(FRenderingCompositePassContext& Conte
 		float fSampleOffsetMultipler = Params.SampleDistanceMultiplier;
 		SetShaderValue(Context.RHICmdList, DispatchCtx.ShaderRHI, DispatchCtx->fSampleOffsetMultipler, fSampleOffsetMultipler);
 
-		float CocRadiusToBucketDistance = Params.PreProcessingToProcessingCocRadiusFactor;
-		SetShaderValue(Context.RHICmdList, DispatchCtx.ShaderRHI, DispatchCtx->CocRadiusToBucketDistance, CocRadiusToBucketDistance);
+		float CocRadiusToBucketDistanceUpperBound = Params.PreProcessingToProcessingCocRadiusFactor * Params.BluringRadiusErrorMultiplier;
+		SetShaderValue(Context.RHICmdList, DispatchCtx.ShaderRHI, DispatchCtx->CocRadiusToBucketDistanceUpperBound, CocRadiusToBucketDistanceUpperBound);
 
-		float BucketDistanceToCocRadius = 1.0f / CocRadiusToBucketDistance;
+		float BucketDistanceToCocRadius = 1.0f / CocRadiusToBucketDistanceUpperBound;
 		SetShaderValue(Context.RHICmdList, DispatchCtx.ShaderRHI, DispatchCtx->BucketDistanceToCocRadius, BucketDistanceToCocRadius);
 	}
 	DispatchCtx.Dispatch();
@@ -681,7 +681,7 @@ FPooledRenderTargetDesc FRCPassDiaphragmDOFDilateCoc::ComputeOutputDesc(EPassOut
 	else
 	{
 		Ret.DebugName = InPassOutputId == ePId_Output0 ? TEXT("DOFDilateFgdCoc") : TEXT("DOFDilateBgdCoc");
-		Ret.Format = InPassOutputId == ePId_Output0 ? PF_G16R16F : PF_FloatRGB;
+		Ret.Format = InPassOutputId == ePId_Output0 ? PF_G16R16F : PF_FloatRGBA;
 	}
 
 	return Ret;
diff --git a/Engine/Source/Runtime/Renderer/Private/PostProcess/DiaphragmDOFPasses.h b/Engine/Source/Runtime/Renderer/Private/PostProcess/DiaphragmDOFPasses.h
--- a/Engine/Source/Runtime/Renderer/Private/PostProcess/DiaphragmDOFPasses.h
+++ b/Engine/Source/Runtime/Renderer/Private/PostProcess/DiaphragmDOFPasses.h
@@ -143,11 +143,15 @@ public:
 		/** Convert pre-processing coc radius to processing coc radius */
 		float PreProcessingToProcessingCocRadiusFactor;
 
+		// Error introduced by the random offset of the gathering kernel.
+		float BluringRadiusErrorMultiplier;
+
 		FParameters()
 			: Mode(EMode::StandAlone)
 			, SampleRadiusCount(0)
 			, SampleDistanceMultiplier(1)
 			, PreProcessingToProcessingCocRadiusFactor(1.0f)
+			, BluringRadiusErrorMultiplier(1.0f)
 		{ }
 	};
 
