Cherry-pick 4169377: Reduces under sampling artifact at high resolutions by doing reduction in Karis color space.

diff --git a/Engine/Shaders/Private/DiaphragmDOF/DOFCommon.ush b/Engine/Shaders/Private/DiaphragmDOF/DOFCommon.ush
--- a/Engine/Shaders/Private/DiaphragmDOF/DOFCommon.ush
+++ b/Engine/Shaders/Private/DiaphragmDOF/DOFCommon.ush
@@ -177,6 +177,12 @@ bool IsForeground(float CocRadius)
 	return CocRadius < 0.0f;
 }
 
+// Faster but less accurate luma computation with a scaling by 4.
+float Luma4(float3 Color)
+{
+	return (Color.g * 2.0) + (Color.r + Color.b);
+}
+
 // Returns 1/x if x > 0, d otherwise.
 float SafeRcp(float x, float d = 0.0)
 {
diff --git a/Engine/Shaders/Private/DiaphragmDOF/DOFDownsample.usf b/Engine/Shaders/Private/DiaphragmDOF/DOFDownsample.usf
--- a/Engine/Shaders/Private/DiaphragmDOF/DOFDownsample.usf
+++ b/Engine/Shaders/Private/DiaphragmDOF/DOFDownsample.usf
@@ -57,11 +57,16 @@ void DownsampleCS(uint2 DispatchThreadId : SV_DispatchThreadID)
 		}
 		#endif
 	}
+	
+	FCocDownsampleParams DownsampleParams;
+	DownsampleParams.CocRadiusMultiplier = 1.0;
+	DownsampleParams.FrameExposureScale = 1.0;
+	DownsampleParams.bDoColorBasedWeighting = false;
 
 	// Downsample.
 	float4 OutColor;
 	float OutCocRadius;
-	DownsampleSceneColorWithCoc(SceneColors, CocRadii, OutColor, OutCocRadius);
+	DownsampleSceneColorWithCoc(DownsampleParams, SceneColors, CocRadii, OutColor, OutCocRadius);
 	
 	if (all(DispatchThreadId < ViewportRect.zw))
 	{
diff --git a/Engine/Shaders/Private/DiaphragmDOF/DOFDownsample.ush b/Engine/Shaders/Private/DiaphragmDOF/DOFDownsample.ush
--- a/Engine/Shaders/Private/DiaphragmDOF/DOFDownsample.ush
+++ b/Engine/Shaders/Private/DiaphragmDOF/DOFDownsample.ush
@@ -110,58 +110,77 @@ float DownsampleCoc(float CocRadii[4])
 }
 
 
+// Optimized HDR weighting function.
+float HdrWeight4(float3 Color, float Exposure) 
+{
+	return rcp(Luma4(Color) * Exposure + 4.0);
+}
+
+
+struct FCocDownsampleParams
+{
+	// Multiplier to apply on the bilateral weights.
+	float CocRadiusMultiplier;
+
+	// The exposure scale of the frame.
+	float FrameExposureScale;
+
+	// Whether should color based weightings to reduce highlight contamination.
+	bool bDoColorBasedWeighting;
+};
+
+
 /** Compute the bilateral weight of sample to downsample. */
 float ComputeDownsamplingBilateralWeight(
+	FCocDownsampleParams DownsampleParams,
 	float OutCocRadius,
 	float SampleCocRadius,
-	in float CocRadiusMultiplier = 1.0)
+	float3 SampleColor = 0)
 {
 	// Remove samples which are outside the size.
 	// TODO: Tune the ScaleFactor. Looks like to large.
-	float ScaleFactor = 64.0 * CocRadiusMultiplier;
+	float ScaleFactor = 64.0 * DownsampleParams.CocRadiusMultiplier;
 
 	// not doing abs(OutCocRadius - SampleCocRadius) because: it is fine to leak background on foreground arround geometric
 	// edges because going to be close to the hole filling as well. This allow dither opacity material such as dithered human hair
 	// to keep consistent opacity, thickness, and temporal stability as the Coc changes between slight out focus, foreground
 	// and background.
-	#if defined(CONFIG_SETUP)
-		return saturate(1.0 - (OutCocRadius - SampleCocRadius) * ScaleFactor);
+	#if 1
+		float BilateralWeight = saturate(1.0 - (OutCocRadius - SampleCocRadius) * ScaleFactor);
+	#else
+		// GCN Hint: one subtract with abs() post modifier + one MAD with saturate post modifier.
+		float BilateralWeight = saturate(1.0 - abs(OutCocRadius - SampleCocRadius) * ScaleFactor);
 	#endif
 	
-	// GCN Hint: one subtract with abs() post modifier + one MAD with saturate post modifier.
-	return saturate(1.0 - abs(OutCocRadius - SampleCocRadius) * ScaleFactor);
+	float ColorWeight = 1;
+	if (DownsampleParams.bDoColorBasedWeighting)
+	{
+		ColorWeight = HdrWeight4(SampleColor, DownsampleParams.FrameExposureScale);
+	}
+
+	return BilateralWeight * ColorWeight;
 }
 
 
 /** Operator to use to downsample 4 scene sample to 1. */
 void DownsampleSceneColorWithCoc(
+	FCocDownsampleParams DownsampleParams,
 	float4 Color[4], float CocRadii[4],
-	out float4 OutColor, out float OutCocRadius,
-	in float CocRadiusMultiplier = 1.0)
+	out float4 OutColor, out float OutCocRadius)
 {
 	// Choose the best the coc to use.
 	OutCocRadius = DownsampleCoc(CocRadii);
 
 	float4 BilateralWeights = float4(
-		ComputeDownsamplingBilateralWeight(OutCocRadius, CocRadii[0], CocRadiusMultiplier),
-		ComputeDownsamplingBilateralWeight(OutCocRadius, CocRadii[1], CocRadiusMultiplier),
-		ComputeDownsamplingBilateralWeight(OutCocRadius, CocRadii[2], CocRadiusMultiplier),
-		ComputeDownsamplingBilateralWeight(OutCocRadius, CocRadii[3], CocRadiusMultiplier));
-
-	// Reduce the output weight if there was information loss on the oposite layer from the one chosen
-	// by DownsampleCoc(). This improve temporaly stability and low res aliasing of foreground out focus.
-	float OutputWeight = 0;
-	UNROLL
-	for (uint i = 0; i < 4; i++)
-	{
-		OutputWeight += (OutCocRadius * CocRadii[i]) >= 0.0 ? 0.25 : 0.0;
-	}
-	OutputWeight = 1;
+		ComputeDownsamplingBilateralWeight(DownsampleParams, OutCocRadius, CocRadii[0], Color[0].rgb),
+		ComputeDownsamplingBilateralWeight(DownsampleParams, OutCocRadius, CocRadii[1], Color[1].rgb),
+		ComputeDownsamplingBilateralWeight(DownsampleParams, OutCocRadius, CocRadii[2], Color[2].rgb),
+		ComputeDownsamplingBilateralWeight(DownsampleParams, OutCocRadius, CocRadii[3], Color[3].rgb));
 
 	float WeightSum = dot(float4(1, 1, 1, 1), BilateralWeights);
 
 	// Normalize weights.
-	float WeightNormalizationFactor = OutputWeight * rcp(WeightSum);
+	float WeightNormalizationFactor = rcp(WeightSum);
 
 	// Do the multiply of WeightNormalizationFactor, because save one mad when alpha channel is disabled.
 	OutColor = WeightNormalizationFactor * (
diff --git a/Engine/Shaders/Private/DiaphragmDOF/DOFHybridScatterCompilation.ush b/Engine/Shaders/Private/DiaphragmDOF/DOFHybridScatterCompilation.ush
--- a/Engine/Shaders/Private/DiaphragmDOF/DOFHybridScatterCompilation.ush
+++ b/Engine/Shaders/Private/DiaphragmDOF/DOFHybridScatterCompilation.ush
@@ -34,12 +34,6 @@ struct FHybridScatterInputs
 	float2 BufferUV;
 };
 
-// Faster but less accurate luma computation with a scaling by 4.
-float Luma4(float3 Color)
-{
-	return (Color.g * 2.0) + (Color.r + Color.b);
-}
-
 // Return the minimal luma for a pixel to be scattered.
 float ComputeLuminanceOnlyBasedScatterFactor(in const FHybridScatterInputs Parameters)
 {
diff --git a/Engine/Shaders/Private/DiaphragmDOF/DOFReduce.usf b/Engine/Shaders/Private/DiaphragmDOF/DOFReduce.usf
--- a/Engine/Shaders/Private/DiaphragmDOF/DOFReduce.usf
+++ b/Engine/Shaders/Private/DiaphragmDOF/DOFReduce.usf
@@ -155,7 +155,7 @@ uint2 InitialTilePixelPositionForReduction2x2(const uint TileSizeLog2, uint Shar
 }
 
 // Reduce using GroupSharedArray
-void ReduceOperator(uint OutId, uint ReduceBankSize, float CocRadiusMultiplier)
+void ReduceOperator(FCocDownsampleParams DownsampleParams, uint OutId, uint ReduceBankSize)
 {
 	// Gather the 4 imput samples. Eaches samples are on banks of size <ReduceBankSize> for coherent LDS memory access.
 	float4 Colors[4];
@@ -181,7 +181,7 @@ void ReduceOperator(uint OutId, uint ReduceBankSize, float CocRadiusMultiplier)
 	// Downsample the 4 sample to one according to Coc.
 	float4 OutColor;
 	float OutCocRadius;
-	DownsampleSceneColorWithCoc(Colors, CocRadii, OutColor, OutCocRadius, CocRadiusMultiplier); 
+	DownsampleSceneColorWithCoc(DownsampleParams, Colors, CocRadii, OutColor, OutCocRadius); 
 	
 	// Technically need a barrier, but in practice with a warp size >= 32, no need.
 	// GroupMemoryBarrierWithGroupSync();
@@ -330,6 +330,9 @@ void ReduceCS(
 	}
 	#endif
 
+	// Sample the frame exposure to SGPR.
+	float FrameExposureScale = ToScalarMemory(EyeAdaptationLookup());
+
 	// Indirect dispatch 
 	#if DIM_HYBRID_SCATTER_FGD || DIM_HYBRID_SCATTER_BGD
 	{
@@ -348,7 +351,7 @@ void ReduceCS(
 			Parameters.NeighborhoodComparisonBufferInvSize = PostprocessInput2Size.zw;
 			Parameters.NeighborhoodComparisonMaxBufferUV = MaxInputBufferUV;
 
-			Parameters.FrameExposureScale = EyeAdaptationLookup();
+			Parameters.FrameExposureScale = FrameExposureScale;
 
 			Parameters.Color = GatheredColor;
 			Parameters.CocRadius = ScatterCocRadius;
@@ -567,13 +570,18 @@ void ReduceCS(
 	uint2 OutputPixelId = DispatchThreadId;
 	OutputMipLevel(/* MipLevel = */ 0, OutputPixelId, GatheredColor, CocRadius);
 	
+	FCocDownsampleParams DownsampleParams;
+	DownsampleParams.CocRadiusMultiplier = 1.0;
+	DownsampleParams.FrameExposureScale = FrameExposureScale;
+	DownsampleParams.bDoColorBasedWeighting = false;
+
 	#if CONFIG_WAVE_BROADCAST_REDUCTION
 	{
 		UNROLL
 		for (uint i = 0; i < (DIM_REDUCE_MIP_COUNT - 1); i++)
 		{
 			// Multiplier that needs to be applied on Coc computation to be resolution independent.
-			const float CocRadiusMultiplier = 0.5 / float(1 << (i));
+			DownsampleParams.CocRadiusMultiplier = 0.5 / float(1 << (i));
 
 			// Scale of the reduction within the 8x8 tile.
 			const uint ReductionScale = 1 << i;
@@ -582,7 +590,7 @@ void ReduceCS(
 			float OutputCocRadius = DownsampleCoc_Wave8x8(CocRadius, ReductionScale);
 			
 			// Compute the sample's bilateral weight.
-			float SampleBilateralWeight = ComputeDownsamplingBilateralWeight(OutputCocRadius, CocRadius, CocRadiusMultiplier);
+			float SampleBilateralWeight = ComputeDownsamplingBilateralWeight(DownsampleParams, OutputCocRadius, CocRadius, GatheredColor.rgb);
 
 			// Weight sum the color across the lanes.
 			GatheredColor.r = Sum2x2WithinWave8x8(GatheredColor.r * SampleBilateralWeight, ReductionScale);
@@ -593,11 +601,8 @@ void ReduceCS(
 			// Compute normalising weight across lanes.
 			float InvTotalWeight = rcp(Sum2x2WithinWave8x8(SampleBilateralWeight, ReductionScale));
 			
-			// Compute output weight across lanes.
-			float OutputWeight = Sum2x2WithinWave8x8((OutputCocRadius * CocRadius) >= 0.0 ? 0.25 : 0.0, ReductionScale);
-			
 			// Normalize the color.
-			GatheredColor *= OutputWeight * InvTotalWeight;
+			GatheredColor *= InvTotalWeight;
 			CocRadius = OutputCocRadius;
 
 			// Output mip map.
@@ -637,10 +642,10 @@ void ReduceCS(
 		
 			if (GroupThreadIndex < ReduceBankSize)
 			{
-				const float CocRadiusMultiplier = 0.5 / float(1 << (i));
+				DownsampleParams.CocRadiusMultiplier = 0.5 / float(1 << (i));
 
 				// Reduce.
-				ReduceOperator(GroupThreadIndex, ReduceBankSize, CocRadiusMultiplier);
+				ReduceOperator(DownsampleParams, GroupThreadIndex, ReduceBankSize);
 
 				// Read LDS (compiler are generally smart here to actualy keep value on VGPR.
 				#if CONFIG_GATHER_INPUT_LAYOUT == GATHER_INPUT_LAYOUT_RGB_ALPHA_COC
diff --git a/Engine/Shaders/Private/DiaphragmDOF/DOFSetup.usf b/Engine/Shaders/Private/DiaphragmDOF/DOFSetup.usf
--- a/Engine/Shaders/Private/DiaphragmDOF/DOFSetup.usf
+++ b/Engine/Shaders/Private/DiaphragmDOF/DOFSetup.usf
@@ -141,11 +141,16 @@ void SetupCS(
 			}
 			#endif
 		}
+		
+		FCocDownsampleParams DownsampleParams;
+		DownsampleParams.CocRadiusMultiplier = 1.0;
+		DownsampleParams.FrameExposureScale = 1.0;
+		DownsampleParams.bDoColorBasedWeighting = false;
 
 		// Downsample.
 		float4 OutColor;
 		float OutCocRadius;
-		DownsampleSceneColorWithCoc(SceneColors, CocRadii, OutColor, OutCocRadius);
+		DownsampleSceneColorWithCoc(DownsampleParams, SceneColors, CocRadii, OutColor, OutCocRadius);
 	
 		uint2 HalfResOutputPixelPosition = GroupId * (GROUP_BORDER_SIZE / 2) + SubGroupPos;
 
@@ -202,7 +207,12 @@ void SetupCS(
 	// because this pass is running at half resolution.
 	float2 SceneBufferUV = (float2(View.ViewRectMin.xy) + 2.0 * (DispatchThreadId + 0.5)) * View.BufferSizeAndInvSize.zw;
 	float2 SceneBufferUVOffset = View.BufferSizeAndInvSize.zw * 0.5;
-	
+			
+	FCocDownsampleParams DownsampleParams;
+	DownsampleParams.CocRadiusMultiplier = 1.0;
+	DownsampleParams.FrameExposureScale = 1.0;
+	DownsampleParams.bDoColorBasedWeighting = false;
+
 	BRANCH
 	if (all(GroupId * GROUP_BORDER_SIZE + (GROUP_BORDER_SIZE + 1) <= ViewportRect.zw))
 	{
@@ -228,7 +238,7 @@ void SetupCS(
 		// Downsample.
 		float4 OutColor;
 		float OutCocRadius;
-		DownsampleSceneColorWithCoc(SceneColors, CocRadii, OutColor, OutCocRadius);
+		DownsampleSceneColorWithCoc(DownsampleParams, SceneColors, CocRadii, OutColor, OutCocRadius);
 	
 		#if CONFIG_DOF_ALPHA
 		{
@@ -267,7 +277,7 @@ void SetupCS(
 		// Downsample.
 		float4 OutColor;
 		float OutCocRadius;
-		DownsampleSceneColorWithCoc(SceneColors, CocRadii, OutColor, OutCocRadius);
+		DownsampleSceneColorWithCoc(DownsampleParams, SceneColors, CocRadii, OutColor, OutCocRadius);
 	
 		#if 0 // Lower VGPR footprint.
 			uint2 OutputPixelPosition = 0.5 * SceneBufferUV * View.BufferSizeAndInvSize.xy;
