Cherry-pick 4169624: Fixes bokeh mirroring mistmatch between gather and scatter.

diff --git a/Engine/Shaders/Private/DiaphragmDOF/DOFHybridScatterPixelShader.usf b/Engine/Shaders/Private/DiaphragmDOF/DOFHybridScatterPixelShader.usf
--- a/Engine/Shaders/Private/DiaphragmDOF/DOFHybridScatterPixelShader.usf
+++ b/Engine/Shaders/Private/DiaphragmDOF/DOFHybridScatterPixelShader.usf
@@ -12,9 +12,9 @@
 
 // Whether to point mirror the bokeh.
 #if DIM_LAYER_PROCESSING == LAYER_PROCESSING_FOREGROUND_ONLY
-	#define CONFIG_POINT_MIRROR_BOKEH 1
-#elif DIM_LAYER_PROCESSING == LAYER_PROCESSING_BACKGROUND_ONLY
 	#define CONFIG_POINT_MIRROR_BOKEH -1
+#elif DIM_LAYER_PROCESSING == LAYER_PROCESSING_BACKGROUND_ONLY
+	#define CONFIG_POINT_MIRROR_BOKEH 1
 #endif
 
 
