Cherry-pick 4173101: Fixes tilling artifacts on DOF gathering pass caused by feather change as resolution goes up, creating area conservativeness issue between fast and complex gathering technics.

diff --git a/Engine/Shaders/Private/DiaphragmDOF/DOFGatherKernel.ush b/Engine/Shaders/Private/DiaphragmDOF/DOFGatherKernel.ush
--- a/Engine/Shaders/Private/DiaphragmDOF/DOFGatherKernel.ush
+++ b/Engine/Shaders/Private/DiaphragmDOF/DOFGatherKernel.ush
@@ -87,17 +87,9 @@ float ComputeSampleIntersection(FPrivateKernelSharedScope KernelScope, float Sam
 {
 	// Intersection of output pixel's disk area with sample's Coc.
 	#if 1
-		// To have an intersection of 50% when the Sample's Coc radius == SampleDistance.
-		#if 0 // TODO: clean this hack.
-			// Poor out of focus, great alpha channel, no leaking.
-			const float Offset = 1;
-		#elif 1
-			// Best slight out of focus.
-			const float Offset = 0;
-		#else
-			// Mathematically ideal.
-			const float Offset = 0.5;
-		#endif
+		// To have an intersection of 50% when the Sample's Coc radius == SampleDistance to be area conservative
+		// with bilinear sampled gathering or hybrid scattering.
+		const float Offset = 0.5;
 
 		#if DIM_LAYER_PROCESSING == LAYER_PROCESSING_SLIGHT_OUT_OF_FOCUS
 			// Uses a sharper intersection for slight out of focus like Recombine pass since it is being used
