#rnx

diff --git a/Engine/Shaders/Private/LocalVertexFactory.ush b/Engine/Shaders/Private/LocalVertexFactory.ush
--- a/Engine/Shaders/Private/LocalVertexFactory.ush
+++ b/Engine/Shaders/Private/LocalVertexFactory.ush
@@ -636,9 +636,15 @@ FVertexFactoryIntermediates GetVertexFactoryIntermediates(FVertexFactoryInput In
 	Intermediates.InstanceOrigin = VertexFetch_InstanceOriginBuffer[(Input.InstanceId + InstanceOffset)];
 	Intermediates.InstanceLightmapAndShadowMapUVBias = VertexFetch_InstanceLightmapBuffer[(Input.InstanceId + InstanceOffset)];
 #elif USE_INSTANCING && !USE_INSTANCING_EMULATED
-	Intermediates.InstanceTransform1 = Input.InstanceTransform1;
-	Intermediates.InstanceTransform2 = Input.InstanceTransform2;
-	Intermediates.InstanceTransform3 = Input.InstanceTransform3;
+	#if WEBGL
+		Intermediates.InstanceTransform1 = f16tof32(Input.InstanceTransform1);
+		Intermediates.InstanceTransform2 = f16tof32(Input.InstanceTransform2);
+		Intermediates.InstanceTransform3 = f16tof32(Input.InstanceTransform3);
+	#else
+		Intermediates.InstanceTransform1 = Input.InstanceTransform1;
+		Intermediates.InstanceTransform2 = Input.InstanceTransform2;
+		Intermediates.InstanceTransform3 = Input.InstanceTransform3;
+	#endif
 	Intermediates.InstanceOrigin = Input.InstanceOrigin;
 	Intermediates.InstanceLightmapAndShadowMapUVBias = Input.InstanceLightmapAndShadowMapUVBias;
 #endif
