Protect clearcoat materials from nans when the clearcoat layer's roughness is 0

diff --git a/Engine/Shaders/Private/ShadingModels.ush b/Engine/Shaders/Private/ShadingModels.ush
--- a/Engine/Shaders/Private/ShadingModels.ush
+++ b/Engine/Shaders/Private/ShadingModels.ush
@@ -180,7 +180,7 @@ float3 SimpleShading( float3 DiffuseColor, float3 SpecularColor, float Roughness
 FDirectLighting ClearCoatBxDF( FGBufferData GBuffer, half3 N, half3 V, half3 L, float Falloff, float NoL, FAreaLight AreaLight, FShadowTerms Shadow )
 {
 	const float ClearCoat			= GBuffer.CustomData.x;
-	const float ClearCoatRoughness	= GBuffer.CustomData.y;
+	const float ClearCoatRoughness	= max(GBuffer.CustomData.y, 0.02f);
 	const float Film = 1 * ClearCoat;
 	const float MetalSpec = 0.9;
 
@@ -229,9 +229,13 @@ FDirectLighting ClearCoatBxDF( FGBufferData GBuffer, half3 N, half3 V, half3 L,
 	Lighting.Diffuse = AreaLight.FalloffColor * (LayerAttenuation * Falloff * NoL) * Diffuse_Lambert( GBuffer.DiffuseColor );
 
 	if( AreaLight.bIsRect )
+	{
 		Lighting.Specular += LayerAttenuation * RectGGXApproxLTC( GBuffer.Roughness, GBuffer.SpecularColor, N, V, AreaLight.Rect, AreaLight.Texture );
+	}
 	else
+	{
 		Lighting.Specular += AreaLight.FalloffColor * (LayerAttenuation * Falloff * NoL) * SpecularGGX( GBuffer.Roughness, GBuffer.SpecularColor, Context, NoL, AreaLight );
+	}
 
 	Lighting.Transmission = 0;
 	return Lighting;
