Fixed SSS transmission phase function clamping.
Added cvar for iris normal. Default behavior is unchanged.

diff --git a/Engine/Shaders/Private/DeferredShadingCommon.ush b/Engine/Shaders/Private/DeferredShadingCommon.ush
--- a/Engine/Shaders/Private/DeferredShadingCommon.ush
+++ b/Engine/Shaders/Private/DeferredShadingCommon.ush
@@ -564,7 +564,9 @@ FGBufferData DecodeGBufferData(
 	if( GBuffer.ShadingModelID == SHADINGMODELID_EYE )
 	{
 		GBuffer.Metallic = 0.0;
-		//GBuffer.Specular = 0.25;
+#if IRIS_NORMAL
+		GBuffer.Specular = 0.25;
+#endif
 	}
 
 	// derived from BaseColor, Metalness, Specular
diff --git a/Engine/Shaders/Private/Definitions.usf b/Engine/Shaders/Private/Definitions.usf
--- a/Engine/Shaders/Private/Definitions.usf
+++ b/Engine/Shaders/Private/Definitions.usf
@@ -120,6 +120,10 @@
 #define CLEAR_COAT_BOTTOM_NORMAL						0
 #endif
 
+#ifndef IRIS_NORMAL
+#define IRIS_NORMAL										0
+#endif
+
 #ifndef TEX_COORD_SCALE_ANALYSIS
 #define TEX_COORD_SCALE_ANALYSIS						0
 #endif
diff --git a/Engine/Shaders/Private/ShadingModels.ush b/Engine/Shaders/Private/ShadingModels.ush
--- a/Engine/Shaders/Private/ShadingModels.ush
+++ b/Engine/Shaders/Private/ShadingModels.ush
@@ -321,7 +321,7 @@ float ApproximateHG(float cosJ, float g)
 
 	const float ISO_PHASE_FUNC_Normalized = 0.5;
 
-	return (ISO_PHASE_FUNC_Normalized * (1.0f - g2) / gcos2);
+	return (ISO_PHASE_FUNC_Normalized * (1.0f - g2) / max( 1e-5, gcos2));
 }
 
 void GetProfileDualSpecular(FGBufferData GBuffer, out float AverageToRoughness0, out float AverageToRoughness1, out float LobeMix)
@@ -382,7 +382,7 @@ FDirectLighting SubsurfaceProfileBxDF( FGBufferData GBuffer, half3 N, half3 V, h
 	float3 Profile = GetTransmissionProfile(GBuffer, Thickness).rgb;
 
 	float3 RefracV = refract(V, -N, TransmissionParams.OneOverIOR);
-	float PhaseFunction = saturate( ApproximateHG( dot(-L, RefracV), TransmissionParams.ScatteringDistribution ) );
+	float PhaseFunction = ApproximateHG( dot(-L, RefracV), TransmissionParams.ScatteringDistribution );
 	Lighting.Transmission = AreaLight.FalloffColor * Profile * (Falloff * PhaseFunction); // TODO: This probably should also include cosine term (NoL)
 
 #else // USE_TRANSMISSION
@@ -834,7 +834,7 @@ FDirectLighting TwoSidedBxDF( FGBufferData GBuffer, half3 N, half3 V, half3 L, f
 
 FDirectLighting EyeBxDF( FGBufferData GBuffer, half3 N, half3 V, half3 L, float Falloff, float NoL, FAreaLight AreaLight, FShadowTerms Shadow )
 {
-#if 0
+#if IRIS_NORMAL
 	const float2 CausticNormalDelta	= float2( GBuffer.StoredMetallic, GBuffer.StoredSpecular )	* 2 - (256.0/255.0);
 	const float2 IrisNormalDelta	= float2( GBuffer.CustomData.y, GBuffer.CustomData.z )		* 2 - (256.0/255.0);
 	const float  IrisMask = 1.0f - GBuffer.CustomData.w;
diff --git a/Engine/Shaders/Private/ShadingModelsMaterial.ush b/Engine/Shaders/Private/ShadingModelsMaterial.ush
--- a/Engine/Shaders/Private/ShadingModelsMaterial.ush
+++ b/Engine/Shaders/Private/ShadingModelsMaterial.ush
@@ -129,7 +129,7 @@ void SetGBufferForShadingModel(
 	GBuffer.ShadingModelID = SHADINGMODELID_EYE;
 	GBuffer.CustomData.x = EncodeSubsurfaceProfile(SubsurfaceProfile).x;
 	GBuffer.CustomData.w = 1.0f - saturate(GetMaterialCustomData0(MaterialParameters));	// Opacity = 1.0 - Iris Mask
-#if 0
+#if IRIS_NORMAL
 	float IrisMask		= saturate( GetMaterialCustomData0(MaterialParameters) );
 	float IrisDistance	= saturate( GetMaterialCustomData1(MaterialParameters) );
 
diff --git a/Engine/Source/Runtime/Engine/Private/ShaderCompiler/ShaderCompiler.cpp b/Engine/Source/Runtime/Engine/Private/ShaderCompiler/ShaderCompiler.cpp
--- a/Engine/Source/Runtime/Engine/Private/ShaderCompiler/ShaderCompiler.cpp
+++ b/Engine/Source/Runtime/Engine/Private/ShaderCompiler/ShaderCompiler.cpp
@@ -3142,6 +3142,11 @@ void GlobalBeginCompileShader(
 		Input.Environment.SetDefine(TEXT("CLEAR_COAT_BOTTOM_NORMAL"), CVar ? (CVar->GetValueOnGameThread() != 0) : 0);
 	}
 
+	{
+		static const auto CVar = IConsoleManager::Get().FindTConsoleVariableDataInt(TEXT("r.IrisNormal"));
+		Input.Environment.SetDefine(TEXT("IRIS_NORMAL"), CVar ? (CVar->GetValueOnGameThread() != 0) : 0);
+	}
+
 	{
 		static const auto CVar = IConsoleManager::Get().FindTConsoleVariableDataInt(TEXT("Compat.UseDXT5NormalMaps"));
 		Input.Environment.SetDefine(TEXT("DXT5_NORMALMAPS"), CVar ? (CVar->GetValueOnGameThread() != 0) : 0);
diff --git a/Engine/Source/Runtime/Renderer/Private/DeferredShadingRenderer.cpp b/Engine/Source/Runtime/Renderer/Private/DeferredShadingRenderer.cpp
--- a/Engine/Source/Runtime/Renderer/Private/DeferredShadingRenderer.cpp
+++ b/Engine/Source/Runtime/Renderer/Private/DeferredShadingRenderer.cpp
@@ -91,6 +91,14 @@ static TAutoConsoleVariable<int32> CVarClearCoatNormal(
 	TEXT(" 1: on"),
 	ECVF_ReadOnly);
 
+static TAutoConsoleVariable<int32> CVarIrisNormal(
+	TEXT("r.IrisNormal"),
+	0,
+	TEXT("0 to disable iris normal.\n")
+	TEXT(" 0: off\n")
+	TEXT(" 1: on"),
+	ECVF_ReadOnly);
+
 static TAutoConsoleVariable<int32> CVarFXSystemPreRenderAfterPrepass(
 	TEXT("r.FXSystemPreRenderAfterPrepass"),
 	0,
diff --git a/Engine/Source/Runtime/ShaderCore/Private/Shader.cpp b/Engine/Source/Runtime/ShaderCore/Private/Shader.cpp
--- a/Engine/Source/Runtime/ShaderCore/Private/Shader.cpp
+++ b/Engine/Source/Runtime/ShaderCore/Private/Shader.cpp
@@ -1814,6 +1814,11 @@ void ShaderMapAppendKeyString(EShaderPlatform Platform, FString& KeyString)
 		KeyString += (CVar && CVar->GetValueOnAnyThread() != 0) ? TEXT("_CCBN") : TEXT("_NoCCBN");
 	}
 
+	{
+		static const auto CVar = IConsoleManager::Get().FindTConsoleVariableDataInt(TEXT("r.IrisNormal"));
+		KeyString += (CVar && CVar->GetValueOnAnyThread() != 0) ? TEXT("_Iris") : TEXT("_NoIris");
+	}
+
 	{
 		static const auto CVar = IConsoleManager::Get().FindTConsoleVariableDataInt(TEXT("r.CompileShadersForDevelopment"));
 		KeyString += (CVar && CVar->GetValueOnAnyThread() != 0) ? TEXT("_DEV") : TEXT("_NoDEV");
