Making division by 0xFFFFFFFF explicitly unsigned, because OpenGL otherwise assumes it's a signed int, just dividing by -1

diff --git a/Engine/Plugins/FX/Niagara/Source/Niagara/Private/NiagaraScriptDerivedData.h b/Engine/Plugins/FX/Niagara/Source/Niagara/Private/NiagaraScriptDerivedData.h
--- a/Engine/Plugins/FX/Niagara/Source/Niagara/Private/NiagaraScriptDerivedData.h
+++ b/Engine/Plugins/FX/Niagara/Source/Niagara/Private/NiagaraScriptDerivedData.h
@@ -54,7 +54,7 @@ public:
 		// This is a version string that mimics the old versioning scheme. If you
 		// want to bump this version, generate a new guid using VS->Tools->Create GUID and
 		// return it here. Ex.
-		return TEXT("F7CCEA20D9A5451F8658F37985856110");
+		return TEXT("B19605DB7417452D85E5BD2E13EE370B");
 	}
 
 	virtual FString GetPluginSpecificCacheKeySuffix() const override;
diff --git a/Engine/Shaders/Private/NiagaraEmitterInstanceShader.usf b/Engine/Shaders/Private/NiagaraEmitterInstanceShader.usf
--- a/Engine/Shaders/Private/NiagaraEmitterInstanceShader.usf
+++ b/Engine/Shaders/Private/NiagaraEmitterInstanceShader.usf
@@ -72,7 +72,7 @@ NiagaraSimulationShader.usf:
 	float4 rand4()
 	{
 		RandomVal = NiagaraRand4DPCG32( int4(GDispatchThreadId+(RandomVal/573).xyz, RandomVal.x/921) + int4(EmitterTickCounter.xxxx));
-		return float4(RandomVal.xyzw) / float4(0xFFFFFFFF,0xFFFFFFFF,0xFFFFFFFF,0xFFFFFFFF);
+		return float4(RandomVal.xyzw) / float4(0xFFFFFFFFu,0xFFFFFFFFu,0xFFFFFFFFu,0xFFFFFFFFu);
 	}
 
 	float rand(float x)
@@ -82,12 +82,14 @@ NiagaraSimulationShader.usf:
 
 	float2 rand(float2 x)
 	{
-		return float2(rand4().x, rand4().y) * x;
+		float4 rnd = rand4();
+		return float2(rnd.x, rnd.y) * x;
 	}
 
 	float3 rand(float3 x)
 	{
-		return float3(rand4().x, rand4().y, rand4().z) * x;
+		float4 rnd = rand4();
+		return float3(rnd.x, rnd.y, rnd.z) * x;
 	}
 
 	float4 rand(float4 x) 
@@ -226,10 +228,10 @@ void ExitStatScope();
 	 */
 	int AcquireIndex(int DataSetID, bool bDoAcquire)
 	{
-		int PrevIdx = GSpawnStartInstance + SpawnedInstances;
+		int PrevIdx = GSpawnStartInstance + SpawnedInstances;	// scratch instance as default; write to that for dead particles
 		if(bDoAcquire == true)
 		{
-			InterlockedAdd(RWDataSetIndices[DataSetID*5+1], 1, PrevIdx);
+			InterlockedAdd(RWDataSetIndices[DataSetID*5+1], 1U, PrevIdx);
 		}
 		return PrevIdx;
 	}
