Fix to prevent PS4 compiler warning during cook on ambiguous uses of AtomicAdd.

#tests doesn't cause cook errors now

diff --git a/Engine/Shaders/Private/NiagaraEmitterInstanceShader.usf b/Engine/Shaders/Private/NiagaraEmitterInstanceShader.usf
--- a/Engine/Shaders/Private/NiagaraEmitterInstanceShader.usf
+++ b/Engine/Shaders/Private/NiagaraEmitterInstanceShader.usf
@@ -226,7 +226,10 @@ NiagaraSimulationShader.usf:
 		int PrevIdx = GSpawnStartInstance + SpawnedInstances;	// scratch instance as default; write to that for dead particles
 		if(bDoAcquire == true)
 		{
-			InterlockedAdd(RWDataSetIndices[DataSetID*5+1], 1U, PrevIdx);
+			// Have to use uint's here to avoid PS4 compiler warnings about InterlockedAdd, cannot propagate uint due to CPU VM limitations...
+			uint RetPrevIdx;
+			InterlockedAdd((uint)RWDataSetIndices[DataSetID*5+1], (uint)1U, RetPrevIdx);
+			PrevIdx = (int)RetPrevIdx;
 		}
 		return PrevIdx;
 	}
