Niagara - Fix the passthrough of data to material particle nodes for our renderers, also fix ribbon width so that it's the actual width of the ribbon and not the width from the center to the edge of the ribbon.

diff --git a/Engine/Plugins/FX/Niagara/Source/Niagara/Private/NiagaraMeshRendererProperties.cpp b/Engine/Plugins/FX/Niagara/Source/Niagara/Private/NiagaraMeshRendererProperties.cpp
--- a/Engine/Plugins/FX/Niagara/Source/Niagara/Private/NiagaraMeshRendererProperties.cpp
+++ b/Engine/Plugins/FX/Niagara/Source/Niagara/Private/NiagaraMeshRendererProperties.cpp
@@ -49,6 +49,9 @@ void UNiagaraMeshRendererProperties::InitBindings()
 		DynamicMaterial3Binding = FNiagaraConstants::GetAttributeDefaultBinding(SYS_PARAM_PARTICLES_DYNAMIC_MATERIAL_PARAM_3);
 		MeshOrientationBinding = FNiagaraConstants::GetAttributeDefaultBinding(SYS_PARAM_PARTICLES_MESH_ORIENTATION);
 		ScaleBinding = FNiagaraConstants::GetAttributeDefaultBinding(SYS_PARAM_PARTICLES_SCALE);
+		MaterialRandomBinding = FNiagaraConstants::GetAttributeDefaultBinding(SYS_PARAM_PARTICLES_MATERIAL_RANDOM);
+		NormalizedAgeBinding = FNiagaraConstants::GetAttributeDefaultBinding(SYS_PARAM_PARTICLES_NORMALIZED_AGE);
+
 		//Default custom sorting to age
 		CustomSortingBinding = FNiagaraConstants::GetAttributeDefaultBinding(SYS_PARAM_PARTICLES_NORMALIZED_AGE);
 	}
diff --git a/Engine/Plugins/FX/Niagara/Source/Niagara/Private/NiagaraRendererMeshes.cpp b/Engine/Plugins/FX/Niagara/Source/Niagara/Private/NiagaraRendererMeshes.cpp
--- a/Engine/Plugins/FX/Niagara/Source/Niagara/Private/NiagaraRendererMeshes.cpp
+++ b/Engine/Plugins/FX/Niagara/Source/Niagara/Private/NiagaraRendererMeshes.cpp
@@ -43,6 +43,8 @@ NiagaraRendererMeshes::NiagaraRendererMeshes(ERHIFeatureLevel::Type FeatureLevel
 	, MaterialParamOffset2(INDEX_NONE)
 	, MaterialParamOffset3(INDEX_NONE)
 	, TransformOffset(INDEX_NONE)
+	, NormalizedAgeOffset(INDEX_NONE)
+	, MaterialRandomOffset(INDEX_NONE)
 	, CustomSortingOffset(INDEX_NONE)
 	, LastSyncedId(INDEX_NONE)
 {
@@ -188,6 +190,8 @@ void NiagaraRendererMeshes::GetDynamicMeshElements(const TArray<const FSceneView
 				PerViewUniformParameters.MaterialParam1DataOffset = MaterialParamOffset1;
 				PerViewUniformParameters.MaterialParam2DataOffset = MaterialParamOffset2;
 				PerViewUniformParameters.MaterialParam3DataOffset = MaterialParamOffset3;
+				PerViewUniformParameters.NormalizedAgeDataOffset = NormalizedAgeOffset;
+				PerViewUniformParameters.MaterialRandomDataOffset = MaterialRandomOffset;
 				PerViewUniformParameters.DefaultPos = bLocalSpace ? FVector4(0.0f, 0.0f, 0.0f, 1.0f) : FVector4(SceneProxy->GetLocalToWorld().GetOrigin());
 				/*
 				if (Properties)
@@ -377,6 +381,8 @@ FNiagaraDynamicDataBase *NiagaraRendererMeshes::GenerateVertexData(const FNiagar
 		Data.GetVariableComponentOffsets(Properties->DynamicMaterial2Binding.DataSetVariable, MaterialParamOffset2, IntDummy);
 		Data.GetVariableComponentOffsets(Properties->DynamicMaterial3Binding.DataSetVariable, MaterialParamOffset3, IntDummy);
 		Data.GetVariableComponentOffsets(Properties->MeshOrientationBinding.DataSetVariable, TransformOffset, IntDummy);
+		Data.GetVariableComponentOffsets(Properties->NormalizedAgeBinding.DataSetVariable, NormalizedAgeOffset, IntDummy);
+		Data.GetVariableComponentOffsets(Properties->MaterialRandomBinding.DataSetVariable, MaterialRandomOffset, IntDummy);
 		Data.GetVariableComponentOffsets(Properties->CustomSortingBinding.DataSetVariable, CustomSortingOffset, IntDummy);
 		LastSyncedId = Properties->SyncId;
 	}
diff --git a/Engine/Plugins/FX/Niagara/Source/Niagara/Private/NiagaraRendererRibbons.cpp b/Engine/Plugins/FX/Niagara/Source/Niagara/Private/NiagaraRendererRibbons.cpp
--- a/Engine/Plugins/FX/Niagara/Source/Niagara/Private/NiagaraRendererRibbons.cpp
+++ b/Engine/Plugins/FX/Niagara/Source/Niagara/Private/NiagaraRendererRibbons.cpp
@@ -55,9 +55,13 @@ public:
 NiagaraRendererRibbons::NiagaraRendererRibbons(ERHIFeatureLevel::Type FeatureLevel, UNiagaraRendererProperties *InProps) :
 	NiagaraRenderer()
 	, PositionDataOffset(INDEX_NONE)
+	, VelocityDataOffset(INDEX_NONE)
 	, WidthDataOffset(INDEX_NONE)
 	, TwistDataOffset(INDEX_NONE)
+	, FacingDataOffset(INDEX_NONE)
 	, ColorDataOffset(INDEX_NONE)
+	, NormalizedAgeDataOffset(INDEX_NONE)
+	, MaterialRandomDataOffset(INDEX_NONE)
 	, LastSyncedId(INDEX_NONE)
 {
 	VertexFactory = new FNiagaraRibbonVertexFactory(NVFT_Ribbon, FeatureLevel);
@@ -168,6 +172,8 @@ void NiagaraRendererRibbons::GetDynamicMeshElements(const TArray<const FSceneVie
 				PerViewUniformParameters.WidthDataOffset = WidthDataOffset;
 				PerViewUniformParameters.TwistDataOffset = TwistDataOffset;
 				PerViewUniformParameters.FacingDataOffset = FacingDataOffset;
+				PerViewUniformParameters.NormalizedAgeDataOffset = NormalizedAgeDataOffset;
+				PerViewUniformParameters.MaterialRandomDataOffset = MaterialRandomDataOffset;
 				PerViewUniformParameters.MaterialParamDataOffset = MaterialParamOffset;
 				PerViewUniformParameters.MaterialParam1DataOffset = MaterialParamOffset1;
 				PerViewUniformParameters.MaterialParam2DataOffset = MaterialParamOffset2;
@@ -411,6 +417,8 @@ FNiagaraDynamicDataBase *NiagaraRendererRibbons::GenerateVertexData(const FNiaga
 		Data.GetVariableComponentOffsets(Properties->RibbonWidthBinding.DataSetVariable, WidthDataOffset, IntDummy);
 		Data.GetVariableComponentOffsets(Properties->RibbonTwistBinding.DataSetVariable, TwistDataOffset, IntDummy);
 		Data.GetVariableComponentOffsets(Properties->RibbonFacingBinding.DataSetVariable, FacingDataOffset, IntDummy);
+		Data.GetVariableComponentOffsets(Properties->NormalizedAgeBinding.DataSetVariable, NormalizedAgeDataOffset, IntDummy);
+		Data.GetVariableComponentOffsets(Properties->MaterialRandomBinding.DataSetVariable, MaterialRandomDataOffset, IntDummy);
 
 		Data.GetVariableComponentOffsets(Properties->DynamicMaterialBinding.DataSetVariable, MaterialParamOffset, IntDummy);
 		Data.GetVariableComponentOffsets(Properties->DynamicMaterial1Binding.DataSetVariable, MaterialParamOffset1, IntDummy);
diff --git a/Engine/Plugins/FX/Niagara/Source/Niagara/Private/NiagaraRendererSprites.cpp b/Engine/Plugins/FX/Niagara/Source/Niagara/Private/NiagaraRendererSprites.cpp
--- a/Engine/Plugins/FX/Niagara/Source/Niagara/Private/NiagaraRendererSprites.cpp
+++ b/Engine/Plugins/FX/Niagara/Source/Niagara/Private/NiagaraRendererSprites.cpp
@@ -51,7 +51,8 @@ NiagaraRendererSprites::NiagaraRendererSprites(ERHIFeatureLevel::Type FeatureLev
 	, MaterialParamOffset3(INDEX_NONE)
 	, CameraOffsetOffset(INDEX_NONE)
 	, UVScaleOffset(INDEX_NONE)
-	, ParticleRandomOffset(INDEX_NONE)
+	, NormalizedAgeOffset(INDEX_NONE)
+	, MaterialRandomOffset(INDEX_NONE)
 	, CustomSortingOffset(INDEX_NONE)
 	, LastSyncId(INDEX_NONE)
 {
@@ -164,12 +165,13 @@ void NiagaraRendererSprites::GetDynamicMeshElements(const TArray<const FSceneVie
 				PerViewUniformParameters.MaterialParam2DataOffset = MaterialParamOffset2;
 				PerViewUniformParameters.MaterialParam3DataOffset = MaterialParamOffset3;
 				PerViewUniformParameters.SubimageDataOffset = SubImageOffset;
-				PerViewUniformParameters.FacingOffset = FacingOffset;
-				PerViewUniformParameters.AlignmentOffset = AlignmentOffset;
+				PerViewUniformParameters.FacingDataOffset = FacingOffset;
+				PerViewUniformParameters.AlignmentDataOffset = AlignmentOffset;
 				PerViewUniformParameters.SubImageBlendMode = Properties->bSubImageBlend;
-				PerViewUniformParameters.CameraOffsetOffset = CameraOffsetOffset;
-				PerViewUniformParameters.UVScaleOffset = UVScaleOffset;
-				PerViewUniformParameters.ParticleRandomOffset = ParticleRandomOffset;
+				PerViewUniformParameters.CameraOffsetDataOffset = CameraOffsetOffset;
+				PerViewUniformParameters.UVScaleDataOffset = UVScaleOffset;
+				PerViewUniformParameters.NormalizedAgeDataOffset = NormalizedAgeOffset;
+				PerViewUniformParameters.MaterialRandomDataOffset = MaterialRandomOffset;
 				PerViewUniformParameters.DefaultPos = bLocalSpace ? FVector4(0.0f, 0.0f, 0.0f, 1.0f) : FVector4(SceneProxy->GetLocalToWorld().GetOrigin());
 
 				//UE_LOG(LogNiagara, Log, TEXT("SubImageSize: %f %f offset: %d"), Properties->SubImageSize.X, Properties->SubImageSize.Y, SubImageOffset);
@@ -330,7 +332,8 @@ FNiagaraDynamicDataBase *NiagaraRendererSprites::GenerateVertexData(const FNiaga
 			Data.GetVariableComponentOffsets(Properties->DynamicMaterial3Binding.DataSetVariable, MaterialParamOffset3, IntDummy);
 			Data.GetVariableComponentOffsets(Properties->CameraOffsetBinding.DataSetVariable, CameraOffsetOffset, IntDummy);
 			Data.GetVariableComponentOffsets(Properties->UVScaleBinding.DataSetVariable, UVScaleOffset, IntDummy);
-			Data.GetVariableComponentOffsets(Properties->MaterialRandomBinding.DataSetVariable, ParticleRandomOffset, IntDummy);
+			Data.GetVariableComponentOffsets(Properties->NormalizedAgeBinding.DataSetVariable, NormalizedAgeOffset, IntDummy);
+			Data.GetVariableComponentOffsets(Properties->MaterialRandomBinding.DataSetVariable, MaterialRandomOffset, IntDummy);
 			Data.GetVariableComponentOffsets(Properties->CustomSortingBinding.DataSetVariable, CustomSortingOffset, IntDummy);
 
 			if (CustomSortingOffset == INDEX_NONE && (Properties->SortMode == ENiagaraSortMode::CustomAscending || Properties->SortMode == ENiagaraSortMode::CustomDecending))
diff --git a/Engine/Plugins/FX/Niagara/Source/Niagara/Private/NiagaraRibbonRendererProperties.cpp b/Engine/Plugins/FX/Niagara/Source/Niagara/Private/NiagaraRibbonRendererProperties.cpp
--- a/Engine/Plugins/FX/Niagara/Source/Niagara/Private/NiagaraRibbonRendererProperties.cpp
+++ b/Engine/Plugins/FX/Niagara/Source/Niagara/Private/NiagaraRibbonRendererProperties.cpp
@@ -48,6 +48,7 @@ void UNiagaraRibbonRendererProperties::InitBindings()
 		RibbonFacingBinding = FNiagaraConstants::GetAttributeDefaultBinding(SYS_PARAM_PARTICLES_RIBBONFACING);
 		RibbonIdBinding = FNiagaraConstants::GetAttributeDefaultBinding(SYS_PARAM_PARTICLES_RIBBONID);
 		RibbonLinkOrderBinding = FNiagaraConstants::GetAttributeDefaultBinding(SYS_PARAM_PARTICLES_RIBBONLINKORDER);
+		MaterialRandomBinding = FNiagaraConstants::GetAttributeDefaultBinding(SYS_PARAM_PARTICLES_MATERIAL_RANDOM);
 	}
 }
 
diff --git a/Engine/Plugins/FX/Niagara/Source/Niagara/Private/NiagaraSpriteRendererProperties.cpp b/Engine/Plugins/FX/Niagara/Source/Niagara/Private/NiagaraSpriteRendererProperties.cpp
--- a/Engine/Plugins/FX/Niagara/Source/Niagara/Private/NiagaraSpriteRendererProperties.cpp
+++ b/Engine/Plugins/FX/Niagara/Source/Niagara/Private/NiagaraSpriteRendererProperties.cpp
@@ -69,6 +69,7 @@ void UNiagaraSpriteRendererProperties::InitBindings()
 		CameraOffsetBinding = FNiagaraConstants::GetAttributeDefaultBinding(SYS_PARAM_PARTICLES_CAMERA_OFFSET);
 		UVScaleBinding = FNiagaraConstants::GetAttributeDefaultBinding(SYS_PARAM_PARTICLES_UV_SCALE);
 		MaterialRandomBinding = FNiagaraConstants::GetAttributeDefaultBinding(SYS_PARAM_PARTICLES_MATERIAL_RANDOM);
+		NormalizedAgeBinding = FNiagaraConstants::GetAttributeDefaultBinding(SYS_PARAM_PARTICLES_NORMALIZED_AGE);
 		
 		//Default custom sorting to age
 		CustomSortingBinding = FNiagaraConstants::GetAttributeDefaultBinding(SYS_PARAM_PARTICLES_NORMALIZED_AGE);
diff --git a/Engine/Plugins/FX/Niagara/Source/Niagara/Public/NiagaraMeshRendererProperties.h b/Engine/Plugins/FX/Niagara/Source/Niagara/Public/NiagaraMeshRendererProperties.h
--- a/Engine/Plugins/FX/Niagara/Source/Niagara/Public/NiagaraMeshRendererProperties.h
+++ b/Engine/Plugins/FX/Niagara/Source/Niagara/Public/NiagaraMeshRendererProperties.h
@@ -110,10 +110,18 @@ public:
 	UPROPERTY(EditAnywhere, Category = "Bindings")
 	FNiagaraVariableAttributeBinding DynamicMaterial3Binding;
 
+	/** Which attribute should we use for material randoms when generating instanced meshes?*/
+	UPROPERTY(EditAnywhere, Category = "Bindings")
+	FNiagaraVariableAttributeBinding MaterialRandomBinding;
+
 	/** Which attribute should we use custom sorting of particles in this emitter. */
 	UPROPERTY(EditAnywhere, Category = "Bindings")
 	FNiagaraVariableAttributeBinding CustomSortingBinding;
 
+	/** Which attribute should we use for Normalized Age? */
+	UPROPERTY(EditAnywhere, Category = "Bindings")
+	FNiagaraVariableAttributeBinding NormalizedAgeBinding;
+
 	UPROPERTY(Transient)
 	int32 SyncId;
 protected:
diff --git a/Engine/Plugins/FX/Niagara/Source/Niagara/Public/NiagaraRenderer.h b/Engine/Plugins/FX/Niagara/Source/Niagara/Public/NiagaraRenderer.h
--- a/Engine/Plugins/FX/Niagara/Source/Niagara/Public/NiagaraRenderer.h
+++ b/Engine/Plugins/FX/Niagara/Source/Niagara/Public/NiagaraRenderer.h
@@ -237,7 +237,8 @@ private:
 	int32 MaterialParamOffset3;
 	int32 CameraOffsetOffset;
 	int32 UVScaleOffset;
-	int32 ParticleRandomOffset;
+	int32 NormalizedAgeOffset;
+	int32 MaterialRandomOffset;
 	int32 CustomSortingOffset;
 
 	int32 LastSyncId;
diff --git a/Engine/Plugins/FX/Niagara/Source/Niagara/Public/NiagaraRendererMeshes.h b/Engine/Plugins/FX/Niagara/Source/Niagara/Public/NiagaraRendererMeshes.h
--- a/Engine/Plugins/FX/Niagara/Source/Niagara/Public/NiagaraRendererMeshes.h
+++ b/Engine/Plugins/FX/Niagara/Source/Niagara/Public/NiagaraRendererMeshes.h
@@ -98,6 +98,8 @@ private:
 	int32 MaterialParamOffset2;
 	int32 MaterialParamOffset3;
 	int32 TransformOffset;
+	int32 NormalizedAgeOffset;
+	int32 MaterialRandomOffset;
 	int32 CustomSortingOffset;
 	int32 LastSyncedId;
 };
diff --git a/Engine/Plugins/FX/Niagara/Source/Niagara/Public/NiagaraRendererRibbons.h b/Engine/Plugins/FX/Niagara/Source/Niagara/Public/NiagaraRendererRibbons.h
--- a/Engine/Plugins/FX/Niagara/Source/Niagara/Public/NiagaraRendererRibbons.h
+++ b/Engine/Plugins/FX/Niagara/Source/Niagara/Public/NiagaraRendererRibbons.h
@@ -74,6 +74,8 @@ private:
 	int32 TwistDataOffset;
 	int32 FacingDataOffset;
 	int32 ColorDataOffset;
+	int32 NormalizedAgeDataOffset;
+	int32 MaterialRandomDataOffset;
 	int32 LastSyncedId;
 	int32 MaterialParamOffset;
 	int32 MaterialParamOffset1;
diff --git a/Engine/Plugins/FX/Niagara/Source/Niagara/Public/NiagaraRibbonRendererProperties.h b/Engine/Plugins/FX/Niagara/Source/Niagara/Public/NiagaraRibbonRendererProperties.h
--- a/Engine/Plugins/FX/Niagara/Source/Niagara/Public/NiagaraRibbonRendererProperties.h
+++ b/Engine/Plugins/FX/Niagara/Source/Niagara/Public/NiagaraRibbonRendererProperties.h
@@ -102,55 +102,59 @@ public:
 
 
 	/** Which attribute should we use for position when generating ribbons?*/
-	UPROPERTY(EditAnywhere, AdvancedDisplay, Category = "Ribbon Rendering")
+	UPROPERTY(EditAnywhere, AdvancedDisplay, Category = "Bindings")
 	FNiagaraVariableAttributeBinding PositionBinding;
 
 	/** Which attribute should we use for color when generating ribbons?*/
-	UPROPERTY(EditAnywhere, AdvancedDisplay, Category = "Ribbon Rendering")
+	UPROPERTY(EditAnywhere, AdvancedDisplay, Category = "Bindings")
 	FNiagaraVariableAttributeBinding ColorBinding;
 
 	/** Which attribute should we use for velocity when generating ribbons?*/
-	UPROPERTY(EditAnywhere, AdvancedDisplay, Category = "Ribbon Rendering")
+	UPROPERTY(EditAnywhere, AdvancedDisplay, Category = "Bindings")
 	FNiagaraVariableAttributeBinding VelocityBinding;
 
 	/** Which attribute should we use for normalized age when generating ribbons?*/
-	UPROPERTY(EditAnywhere, AdvancedDisplay, Category = "Ribbon Rendering")
+	UPROPERTY(EditAnywhere, AdvancedDisplay, Category = "Bindings")
 	FNiagaraVariableAttributeBinding NormalizedAgeBinding;
 
 	/** Which attribute should we use for ribbon twist when generating ribbons?*/
-	UPROPERTY(EditAnywhere, AdvancedDisplay, Category = "Ribbon Rendering")
+	UPROPERTY(EditAnywhere, AdvancedDisplay, Category = "Bindings")
 	FNiagaraVariableAttributeBinding RibbonTwistBinding;
 
 	/** Which attribute should we use for ribbon width when generating ribbons?*/
-	UPROPERTY(EditAnywhere, AdvancedDisplay, Category = "Ribbon Rendering")
+	UPROPERTY(EditAnywhere, AdvancedDisplay, Category = "Bindings")
 	FNiagaraVariableAttributeBinding RibbonWidthBinding;
 
 	/** Which attribute should we use for ribbon facing when generating ribbons?*/
-	UPROPERTY(EditAnywhere, AdvancedDisplay, Category = "Ribbon Rendering")
+	UPROPERTY(EditAnywhere, AdvancedDisplay, Category = "Bindings")
 	FNiagaraVariableAttributeBinding RibbonFacingBinding;
 	
 	/** Which attribute should we use for ribbon id when generating ribbons?*/
-	UPROPERTY(EditAnywhere, AdvancedDisplay, Category = "Ribbon Rendering")
+	UPROPERTY(EditAnywhere, AdvancedDisplay, Category = "Bindings")
 	FNiagaraVariableAttributeBinding RibbonIdBinding;
 
 	/** Which attribute should we use for RibbonLinkOrder when generating ribbons?*/
-	UPROPERTY(EditAnywhere, AdvancedDisplay, Category = "Ribbon Rendering")
+	UPROPERTY(EditAnywhere, AdvancedDisplay, Category = "Bindings")
 	FNiagaraVariableAttributeBinding RibbonLinkOrderBinding;
 
+	/** Which attribute should we use for MaterialRandom when generating ribbons?*/
+	UPROPERTY(EditAnywhere, AdvancedDisplay, Category = "Bindings")
+	FNiagaraVariableAttributeBinding MaterialRandomBinding;
+
 	/** Which attribute should we use for dynamic material parameters when generating ribbons?*/
-	UPROPERTY(EditAnywhere, AdvancedDisplay, Category = "Ribbon Rendering")
+	UPROPERTY(EditAnywhere, AdvancedDisplay, Category = "Bindings")
 	FNiagaraVariableAttributeBinding DynamicMaterialBinding;
 
 	/** Which attribute should we use for dynamic material parameters when generating ribbons?*/
-	UPROPERTY(EditAnywhere, AdvancedDisplay, Category = "Ribbon Rendering")
+	UPROPERTY(EditAnywhere, AdvancedDisplay, Category = "Bindings")
 	FNiagaraVariableAttributeBinding DynamicMaterial1Binding;
 
 	/** Which attribute should we use for dynamic material parameters when generating ribbons?*/
-	UPROPERTY(EditAnywhere, AdvancedDisplay, Category = "Ribbon Rendering")
+	UPROPERTY(EditAnywhere, AdvancedDisplay, Category = "Bindings")
 	FNiagaraVariableAttributeBinding DynamicMaterial2Binding;
 
 	/** Which attribute should we use for dynamic material parameters when generating ribbons?*/
-	UPROPERTY(EditAnywhere, AdvancedDisplay, Category = "Ribbon Rendering")
+	UPROPERTY(EditAnywhere, AdvancedDisplay, Category = "Bindings")
 	FNiagaraVariableAttributeBinding DynamicMaterial3Binding;
 
 	UPROPERTY(Transient)
diff --git a/Engine/Plugins/FX/Niagara/Source/Niagara/Public/NiagaraSpriteRendererProperties.h b/Engine/Plugins/FX/Niagara/Source/Niagara/Public/NiagaraSpriteRendererProperties.h
--- a/Engine/Plugins/FX/Niagara/Source/Niagara/Public/NiagaraSpriteRendererProperties.h
+++ b/Engine/Plugins/FX/Niagara/Source/Niagara/Public/NiagaraSpriteRendererProperties.h
@@ -178,6 +178,10 @@ public:
 	UPROPERTY(EditAnywhere, Category = "Bindings")
 	FNiagaraVariableAttributeBinding CustomSortingBinding;
 
+	/** Which attribute should we use for Normalized Age? */
+	UPROPERTY(EditAnywhere, Category = "Bindings")
+	FNiagaraVariableAttributeBinding NormalizedAgeBinding;
+
 	UPROPERTY(Transient)
 	int32 SyncId;
 
diff --git a/Engine/Plugins/FX/Niagara/Source/NiagaraVertexFactories/Public/NiagaraMeshVertexFactory.h b/Engine/Plugins/FX/Niagara/Source/NiagaraVertexFactories/Public/NiagaraMeshVertexFactory.h
--- a/Engine/Plugins/FX/Niagara/Source/NiagaraVertexFactories/Public/NiagaraMeshVertexFactory.h
+++ b/Engine/Plugins/FX/Niagara/Source/NiagaraVertexFactories/Public/NiagaraMeshVertexFactory.h
@@ -46,6 +46,8 @@ BEGIN_UNIFORM_BUFFER_STRUCT(FNiagaraMeshUniformParameters, NIAGARAVERTEXFACTORIE
 	UNIFORM_MEMBER(int, MaterialParam1DataOffset)
 	UNIFORM_MEMBER(int, MaterialParam2DataOffset)
 	UNIFORM_MEMBER(int, MaterialParam3DataOffset)
+	UNIFORM_MEMBER(int, NormalizedAgeDataOffset)
+	UNIFORM_MEMBER(int, MaterialRandomDataOffset)
 	UNIFORM_MEMBER(FVector4, DefaultPos)
 END_UNIFORM_BUFFER_STRUCT(FNiagaraMeshUniformParameters)
 
diff --git a/Engine/Plugins/FX/Niagara/Source/NiagaraVertexFactories/Public/NiagaraRibbonVertexFactory.h b/Engine/Plugins/FX/Niagara/Source/NiagaraVertexFactories/Public/NiagaraRibbonVertexFactory.h
--- a/Engine/Plugins/FX/Niagara/Source/NiagaraVertexFactories/Public/NiagaraRibbonVertexFactory.h
+++ b/Engine/Plugins/FX/Niagara/Source/NiagaraVertexFactories/Public/NiagaraRibbonVertexFactory.h
@@ -45,6 +45,8 @@ BEGIN_UNIFORM_BUFFER_STRUCT(FNiagaraRibbonUniformParameters, NIAGARAVERTEXFACTOR
 	UNIFORM_MEMBER(int, TwistDataOffset)
 	UNIFORM_MEMBER(int, ColorDataOffset)
 	UNIFORM_MEMBER(int, FacingDataOffset)
+	UNIFORM_MEMBER(int, NormalizedAgeDataOffset)
+	UNIFORM_MEMBER(int, MaterialRandomDataOffset)
 	UNIFORM_MEMBER(int, MaterialParamDataOffset)
 	UNIFORM_MEMBER(int, MaterialParam1DataOffset)
 	UNIFORM_MEMBER(int, MaterialParam2DataOffset)
diff --git a/Engine/Plugins/FX/Niagara/Source/NiagaraVertexFactories/Public/NiagaraSpriteVertexFactory.h b/Engine/Plugins/FX/Niagara/Source/NiagaraVertexFactories/Public/NiagaraSpriteVertexFactory.h
--- a/Engine/Plugins/FX/Niagara/Source/NiagaraVertexFactories/Public/NiagaraSpriteVertexFactory.h
+++ b/Engine/Plugins/FX/Niagara/Source/NiagaraVertexFactories/Public/NiagaraSpriteVertexFactory.h
@@ -45,12 +45,13 @@ BEGIN_UNIFORM_BUFFER_STRUCT( FNiagaraSpriteUniformParameters, NIAGARAVERTEXFACTO
 	UNIFORM_MEMBER(int, MaterialParam1DataOffset)
 	UNIFORM_MEMBER(int, MaterialParam2DataOffset)
 	UNIFORM_MEMBER(int, MaterialParam3DataOffset)
-	UNIFORM_MEMBER(int, FacingOffset)
-	UNIFORM_MEMBER(int, AlignmentOffset)
+	UNIFORM_MEMBER(int, FacingDataOffset)
+	UNIFORM_MEMBER(int, AlignmentDataOffset)
 	UNIFORM_MEMBER(int, SubImageBlendMode)
-	UNIFORM_MEMBER(int, CameraOffsetOffset)
-	UNIFORM_MEMBER(int, UVScaleOffset)
-	UNIFORM_MEMBER(int, ParticleRandomOffset)
+	UNIFORM_MEMBER(int, CameraOffsetDataOffset)
+	UNIFORM_MEMBER(int, UVScaleDataOffset)
+	UNIFORM_MEMBER(int, NormalizedAgeDataOffset)
+	UNIFORM_MEMBER(int, MaterialRandomDataOffset)
 	UNIFORM_MEMBER(FVector4, DefaultPos)
 	END_UNIFORM_BUFFER_STRUCT(FNiagaraSpriteUniformParameters)
 
diff --git a/Engine/Shaders/Private/NiagaraMeshVertexFactory.ush b/Engine/Shaders/Private/NiagaraMeshVertexFactory.ush
--- a/Engine/Shaders/Private/NiagaraMeshVertexFactory.ush
+++ b/Engine/Shaders/Private/NiagaraMeshVertexFactory.ush
@@ -8,7 +8,11 @@
 #include "NiagaraVFParticleAccess.usf"
 
 // These are only enabled for SM5 due to limited interpolators elsewhere
-#define USE_PARTICLE_NODES (FEATURE_LEVEL >= FEATURE_LEVEL_SM5)
+#define USE_PARTICLE_POSITION (FEATURE_LEVEL >= FEATURE_LEVEL_SM5 && NEEDS_PARTICLE_POSITION)
+#define USE_PARTICLE_VELOCITY (FEATURE_LEVEL >= FEATURE_LEVEL_SM5 && NEEDS_PARTICLE_VELOCITY)
+#define USE_PARTICLE_TIME (FEATURE_LEVEL >= FEATURE_LEVEL_SM5 && NEEDS_PARTICLE_TIME)
+#define USE_PARTICLE_RANDOM (FEATURE_LEVEL >= FEATURE_LEVEL_SM5 && NEEDS_PARTICLE_RANDOM)
+#define USE_PARTICLE_TRANSFORM (FEATURE_LEVEL >= FEATURE_LEVEL_SM5 && NEEDS_PARTICLE_TRANSFORM)
 
 // Only support previous transform for deferred feature level
 //#define USE_PREVIOUS_TRANSFORM (FEATURE_LEVEL >= FEATURE_LEVEL_SM4)
@@ -26,7 +30,7 @@ uint MeshFacingMode;
 #define MESH_FACING_VELOCITY 1
 #define MESH_FACING_CAMERA_POSITION 2
 #define MESH_FACING_CAMERA_PLANE 3
-	
+
 /** Sorted index buffer for particles if we need to sort their draw order (translucency etc). */
 Buffer<int> SortedIndices;
 uint SortedIndicesOffset;
@@ -87,19 +91,28 @@ struct FVertexFactoryInterpolantsVSToPS
 #if NUM_DYNAMIC_PARAMETERS > 3
 	float4 DynamicParameter3 : COLOR5;
 #endif
-	
-#if USE_PARTICLE_NODES
+
+#if USE_PARTICLE_POSITION
 	/** Particle position in camera-centered translated world space */
 	float3 ParticleTranslatedWorldPosition	: PARTICLE_POSITION;
+#endif
 
+#if USE_PARTICLE_VELOCITY
 	/** The velocity of the particle, XYZ: direction, W: speed. */
 	float4 ParticleVelocity	: PARTICLE_VELOCITY;
+#endif
 
+#if USE_PARTICLE_TIME
 	/** Relative alive time of the particle */
-	float ParticleTime : PARTICLE_TIME;
-#if NEEDS_PARTICLE_TRANSFORM
-	float4 LocalToWorld[4] : PARTICLE_TRANSFORM;
+	float RelativeTime : PARTICLE_TIME;
+#endif
+
+#if USE_PARTICLE_RANDOM
+	float ParticleRandom : PARTICLE_RANDOM;
 #endif
+
+#if USE_PARTICLE_TRANSFORM
+	float4 LocalToWorld[4] : PARTICLE_TRANSFORM;
 #endif
 
 #if INSTANCED_STEREO
@@ -118,7 +131,7 @@ struct FVertexFactoryIntermediates
 	/** The texture coordinates for the vertex. */
 #if NUM_MATERIAL_TEXCOORDS_VERTEX
 	float2	TexCoords[NUM_MATERIAL_TEXCOORDS_VERTEX];
-#endif 
+#endif
 #if USE_PARTICLE_SUBUVS
 	float4 SubUVCoords;
 	float SubUVLerp;
@@ -146,9 +159,16 @@ struct FVertexFactoryIntermediates
 
 	/** Particle position in camera-centered translated world space. */
 	float3 ParticleTranslatedWorldPosition;
-	
+
+#if USE_PARTICLE_TIME
 	/** Relative time. */
 	float RelativeTime;
+#endif
+
+#if USE_PARTICLE_RANDOM
+	/** Particle Random value */
+	float ParticleRandom;
+#endif
 
 	float4x4 LocalToWorld;
 };
@@ -164,9 +184,9 @@ float4x4 GetParticlePrevTransform(FVertexFactoryInput Input)
 	uint TransformBaseIdx = GetLayerInstanceID(Input.InstanceId) * 3;
 
 	float4x4 PrevTransform4x4 = float4x4(
-		PrevTransformBuffer[TransformBaseIdx + 0], 
-		PrevTransformBuffer[TransformBaseIdx + 1], 
-		PrevTransformBuffer[TransformBaseIdx + 2], 
+		PrevTransformBuffer[TransformBaseIdx + 0],
+		PrevTransformBuffer[TransformBaseIdx + 1],
+		PrevTransformBuffer[TransformBaseIdx + 2],
 		float4(0, 0, 0, 1)
 		);
 
@@ -240,17 +260,30 @@ FMaterialPixelParameters GetMaterialPixelParameters(FVertexFactoryInterpolantsVS
 	Result.Particle.DynamicParameter3 = Interpolants.DynamicParameter3;
 #endif
 
-#if USE_PARTICLE_NODES
+#if USE_PARTICLE_POSITION
 	Result.Particle.TranslatedWorldPositionAndSize.xyz = Interpolants.ParticleTranslatedWorldPosition.xyz;
 	Result.Particle.TranslatedWorldPositionAndSize.w = 1;
+#endif
+
+#if USE_PARTICLE_VELOCITY
 	Result.Particle.Velocity = Interpolants.ParticleVelocity;
-	Result.Particle.RelativeTime = Interpolants.ParticleTime;
-#if NEEDS_PARTICLE_TRANSFORM
+#endif
+
+#if USE_PARTICLE_TIME
+	Result.Particle.RelativeTime = Interpolants.RelativeTime;
+#endif
+
+#if USE_PARTICLE_RANDOM
+	Result.Particle.Random = Interpolants.ParticleRandom;
+#else
+	Result.Particle.Random = 0.0f;
+#endif
+
+#if USE_PARTICLE_TRANSFORM
 	Result.Particle.LocalToWorld[0] = Interpolants.LocalToWorld[0];
 	Result.Particle.LocalToWorld[1] = Interpolants.LocalToWorld[1];
 	Result.Particle.LocalToWorld[2] = Interpolants.LocalToWorld[2];
 	Result.Particle.LocalToWorld[3] = Interpolants.LocalToWorld[3];
-#endif
 #endif
 
 	Result.Particle.MotionBlurFade = 1.0f;
@@ -273,10 +306,25 @@ FMaterialVertexParameters GetMaterialVertexParameters(FVertexFactoryInput Input,
 	Result.PreSkinnedPosition = Input.Position.xyz;
 	Result.PreSkinnedNormal = TangentBias(Input.TangentZ.xyz);
 
+#if USE_PARTICLE_POSITION
 	Result.Particle.TranslatedWorldPositionAndSize.xyz = Intermediates.ParticleTranslatedWorldPosition.xyz;
 	Result.Particle.TranslatedWorldPositionAndSize.w = 1;
+#endif
+
+#if USE_PARTICLE_VELOCITY
 	Result.Particle.Velocity = Intermediates.ParticleVelocity;
+#endif
+
+#if USE_PARTICLE_TIME
 	Result.Particle.RelativeTime = Intermediates.RelativeTime;
+#endif
+
+#if USE_PARTICLE_RANDOM
+	Result.Particle.Random = Intermediates.ParticleRandom;
+#else
+	Result.Particle.Random = 0.0f;
+#endif
+
 	Result.Particle.Color = Intermediates.ParticleColor;
 
 #if NUM_DYNAMIC_PARAMETERS > 0
@@ -294,6 +342,7 @@ FMaterialVertexParameters GetMaterialVertexParameters(FVertexFactoryInput Input,
 
 	float4x4 Transform = GetParticleTransform(Intermediates);
 	Result.InstanceLocalToWorld = Transform;
+
 #if NEEDS_PARTICLE_TRANSFORM
 	Result.Particle.LocalToWorld[0] = Transform[0];
 	Result.Particle.LocalToWorld[1] = Transform[1];
@@ -301,7 +350,7 @@ FMaterialVertexParameters GetMaterialVertexParameters(FVertexFactoryInput Input,
 	Result.Particle.LocalToWorld[3] = Transform[3];
 #endif
 
-	Result.TangentToWorld = mul(TangentToLocal, (float3x3)Transform ); 
+	Result.TangentToWorld = mul(TangentToLocal, (float3x3)Transform );
 
 #if NUM_MATERIAL_TEXCOORDS_VERTEX
 	for(int CoordinateIndex = 0; CoordinateIndex < NUM_MATERIAL_TEXCOORDS_VERTEX; CoordinateIndex++)
@@ -407,7 +456,7 @@ void AxisAngleToPitchYawRoll(float4 AxisAngle, out float Pitch, out float Yaw, o
 }
 
 
-float4 EulerToQuaternion(float Pitch, float Yaw, float Roll) 
+float4 EulerToQuaternion(float Pitch, float Yaw, float Roll)
 {
 	float DEG_TO_RAD = PI/(180.f);
 	float DIVIDE_BY_2 = DEG_TO_RAD/2.f;
@@ -437,7 +486,7 @@ float4x4 EulerToMat4x4(float Pitch, float Yaw, float Roll)
 FVertexFactoryIntermediates GetVertexFactoryIntermediates(FVertexFactoryInput Input)
 {
 	FVertexFactoryIntermediates Intermediates = (FVertexFactoryIntermediates)0;
-	
+
 	if(SortedIndicesOffset != 0xFFFFFFFF)
 	{
 		Input.InstanceId = SortedIndices[SortedIndicesOffset + Input.InstanceId];
@@ -461,13 +510,14 @@ FVertexFactoryIntermediates GetVertexFactoryIntermediates(FVertexFactoryInput In
 		Intermediates.ParticleColor = GetVec4(NiagaraMeshVF.ColorDataOffset, Input.InstanceId);
 	}
 
+	float3 ParticleVelocity;
 	if (NiagaraMeshVF.VelocityDataOffset == -1)
 	{
-		Intermediates.ParticleVelocity = float4(0.0f, 0.0f, 0.0f, 1);
+		ParticleVelocity = float4(0.0f, 0.0f, 0.0f, 1);
 	}
 	else
 	{
-		Intermediates.ParticleVelocity = float4(GetVec3(NiagaraMeshVF.VelocityDataOffset, Input.InstanceId), 1);
+		ParticleVelocity = float4(GetVec3(NiagaraMeshVF.VelocityDataOffset, Input.InstanceId), 1);
 	}
 
 #if NUM_DYNAMIC_PARAMETERS > 0
@@ -483,6 +533,33 @@ FVertexFactoryIntermediates GetVertexFactoryIntermediates(FVertexFactoryInput In
 	Intermediates.DynamicParameter3  = NiagaraMeshVF.MaterialParam3DataOffset != -1 ? GetVec4(NiagaraMeshVF.MaterialParam3DataOffset, Input.InstanceId) : float4(1,1,1,1);
 #endif
 
+#if USE_PARTICLE_TIME
+	if (NiagaraMeshVF.NormalizedAgeDataOffset == -1)
+	{
+		Intermediates.RelativeTime = 0;
+	}
+	else
+	{
+		Intermediates.RelativeTime = GetFloat(NiagaraMeshVF.NormalizedAgeDataOffset, Input.InstanceId);
+	}
+#endif
+
+#if USE_PARTICLE_RANDOM
+	if (NiagaraMeshVF.MaterialRandomDataOffset == -1)
+	{
+		Intermediates.ParticleRandom = 0;
+	}
+	else
+	{
+		Intermediates.ParticleRandom = GetFloat(NiagaraMeshVF.MaterialRandomDataOffset, Input.InstanceId);
+	}
+#endif
+
+#if USE_PARTICLE_VELOCITY
+	Intermediates.ParticleVelocity.xyz = normalize(ParticleVelocity);
+	Intermediates.ParticleVelocity.w = length(ParticleVelocity);
+#endif
+
 	float4 Transform1 = float4(1.0f, 0.0f, 0.0f, Intermediates.ParticlePosition.x);
 	float4 Transform2 = float4(0.0f, 1.0f, 0.0f, Intermediates.ParticlePosition.y);
 	float4 Transform3 = float4(0.0f, 0.0f, 1.0f, Intermediates.ParticlePosition.z);
@@ -490,11 +567,11 @@ FVertexFactoryIntermediates GetVertexFactoryIntermediates(FVertexFactoryInput In
 	if(NiagaraMeshVF.TransformDataOffset != -1)
 	{
 		float4 Quat = GetVec4(NiagaraMeshVF.TransformDataOffset, Input.InstanceId);
-		
+
 		//Transform axis to handle local space.
 		//AxisAngle.xyz = mul(float4(AxisAngle.xyz, 0), NiagaraMeshVF.LocalToWorldInverseTransposed);
 
-		
+
 		float4x4 EulerToMat = QuatTo4x4(normalize(Quat));
 		Transform1.xyz = EulerToMat[0].xyz;
 		Transform2.xyz = EulerToMat[1].xyz;
@@ -507,20 +584,20 @@ FVertexFactoryIntermediates GetVertexFactoryIntermediates(FVertexFactoryInput In
 		ScaleMatrix[0][0] *= Scale.x;
 		ScaleMatrix[1][1] *= Scale.y;
 		ScaleMatrix[2][2] *= Scale.z;
-	}	
+	}
 	// Swizzle vertex color.
 	Intermediates.VertexColor = Input.VertexColor FCOLOR_COMPONENT_SWIZZLE;
 
 	// World position and localtoworld transform.
 	float4x4 LocalToWorldMat = float4x4(Transform1, Transform2, Transform3, float4(0.0f, 0.0f, 0.0f, 1.0f));
-	
+
 	Intermediates.ParticleTranslatedWorldPosition = Intermediates.ParticlePosition + ResolvedView.PreViewTranslation.xyz;
 
 	//LocalToWorldMat = mul(ScaleMatrix, LocalToWorldMat);
 
 	if (MeshFacingMode == MESH_FACING_VELOCITY)
 	{
-		float4 ToAlignmentQuat = FindQuatBetween(float3(1.0f, 0.0f, 0.0f), normalize(Intermediates.ParticleVelocity.xyz));
+		float4 ToAlignmentQuat = FindQuatBetween(float3(1.0f, 0.0f, 0.0f), normalize(ParticleVelocity.xyz));
 		float4x4 ToAlignmentMat = QuatTo4x4( normalize(ToAlignmentQuat) );
 		Intermediates.LocalToWorld =  mul(LocalToWorldMat, mul(ToAlignmentMat, ScaleMatrix));
 	}
@@ -542,7 +619,7 @@ FVertexFactoryIntermediates GetVertexFactoryIntermediates(FVertexFactoryInput In
 	}
 	else
 	{
-		Intermediates.LocalToWorld =  mul(LocalToWorldMat, ScaleMatrix); 
+		Intermediates.LocalToWorld =  mul(LocalToWorldMat, ScaleMatrix);
 	}
 
 	//Handle local space. TODO: Optimize behind define.
@@ -550,8 +627,8 @@ FVertexFactoryIntermediates GetVertexFactoryIntermediates(FVertexFactoryInput In
 	Intermediates.LocalToWorld = transpose(Intermediates.LocalToWorld);
 	Intermediates.LocalToWorld = mul(Intermediates.LocalToWorld, NiagaraMeshVF.LocalToWorld);
 	Intermediates.ParticlePosition = mul(Intermediates.ParticlePosition, NiagaraMeshVF.LocalToWorld);
-	Intermediates.ParticleVelocity.xyz = mul(float4(Intermediates.ParticleVelocity.xyz,0), NiagaraMeshVF.LocalToWorldInverseTransposed);
-	
+	//Intermediates.ParticleVelocity.xyz = mul(float4(Intermediates.ParticleVelocity.xyz,0), NiagaraMeshVF.LocalToWorldInverseTransposed);
+
 #if NUM_MATERIAL_TEXCOORDS_VERTEX
 	for(int CoordinateIndex = 0; CoordinateIndex < NUM_MATERIAL_TEXCOORDS_VERTEX; CoordinateIndex++)
 	{
@@ -563,7 +640,7 @@ FVertexFactoryIntermediates GetVertexFactoryIntermediates(FVertexFactoryInput In
 		Intermediates.SubUVCoords.xy = (float2( Input.SubUVParams.x, Input.SubUVParams.y ) + Input.TexCoords[0].xy) * MeshParticleVF.SubImageSize.xy;
 		Intermediates.SubUVCoords.zw = (float2( Input.SubUVParams.z, Input.SubUVParams.w ) + Input.TexCoords[0].xy) * MeshParticleVF.SubImageSize.xy;
 		//Intermediates.SubUVLerp = Input.SubUVLerpAndRelTime.x;
-#endif 
+#endif
 
 	//Intermediates.RelativeTime = Input.SubUVLerpAndRelTime.y;
 
@@ -579,18 +656,18 @@ FVertexFactoryIntermediates GetVertexFactoryIntermediates(FVertexFactoryInput In
 */
 float3x3 VertexFactoryGetTangentToLocal( FVertexFactoryInput Input, FVertexFactoryIntermediates Intermediates )
 {
-	float3x3 Result=0;	
+	float3x3 Result=0;
 	float3 TangentX = TangentBias(Input.TangentX);
 	float4 TangentZ = TangentBias(Input.TangentZ);
 
 	// pass-thru the tangent
 	Result[0] = TangentX;
 	// pass-thru the normal
-	Result[2] = TangentZ.xyz;	
+	Result[2] = TangentZ.xyz;
+
 
-	
 	// derive the binormal by getting the cross product of the normal and tangent
-	Result[1] = cross(Result[2], Result[0]) * TangentZ.w;	
+	Result[1] = cross(Result[2], Result[0]) * TangentZ.w;
 	// Recalculate TangentX off of the other two vectors
 	// This corrects quantization error since TangentX was passed in as a quantized vertex input
 	// The error shows up most in specular off of a mesh with a smoothed UV seam (normal is smooth, but tangents vary across the seam)
@@ -662,7 +739,7 @@ FVertexFactoryInterpolantsVSToPS VertexFactoryGetInterpolantsVSToPS(FVertexFacto
 					Interpolants.TexCoords[CoordinateIndex / 2].wz = CustomizedUVs[CoordinateIndex + 1];
 			}
 		}
-	#endif 
+	#endif
 #endif
 
 
@@ -694,13 +771,24 @@ FVertexFactoryInterpolantsVSToPS VertexFactoryGetInterpolantsVSToPS(FVertexFacto
 #endif
 
 
-#if USE_PARTICLE_NODES
+#if USE_PARTICLE_POSITION
 	Interpolants.ParticleTranslatedWorldPosition = Intermediates.ParticleTranslatedWorldPosition;
+#endif
+
+#if USE_PARTICLE_VELOCITY
 	Interpolants.ParticleVelocity = Intermediates.ParticleVelocity;
-	Interpolants.ParticleTime = Intermediates.RelativeTime;
-#if NEEDS_PARTICLE_TRANSFORM
-	Interpolants.LocalToWorld = Intermediates.LocalToWorld;
 #endif
+
+#if USE_PARTICLE_TIME
+	Interpolants.RelativeTime = Intermediates.RelativeTime;
+#endif
+
+#if USE_PARTICLE_RANDOM
+	Interpolants.ParticleRandom = Intermediates.ParticleRandom;
+#endif
+
+#if USE_PARTICLE_TRANSFORM
+	Interpolants.LocalToWorld = Intermediates.LocalToWorld;
 #endif
 
 #if INSTANCED_STEREO
@@ -712,7 +800,7 @@ FVertexFactoryInterpolantsVSToPS VertexFactoryGetInterpolantsVSToPS(FVertexFacto
 
 // @return previous translated world position
 float4 VertexFactoryGetPreviousWorldPosition(FVertexFactoryInput Input, FVertexFactoryIntermediates Intermediates)
-{	
+{
 	return CalcPrevWorldPosition(Input, Intermediates);
 }
 
@@ -773,10 +861,10 @@ float4 VertexFactoryGetPreviousWorldPosition(FVertexFactoryInput Input, FVertexF
 	FVertexFactoryInterpolantsVSToDS VertexFactoryInterpolate(FVertexFactoryInterpolantsVSToDS a, float aInterp, FVertexFactoryInterpolantsVSToDS b, float bInterp)
 	{
 		FVertexFactoryInterpolantsVSToDS O;
-	
+
 		// Do we really need to interpolate TangentToWorld2 here? It should be replaced by the
 		// interpolated normal from 'whatever' interpolation scheme we're using
-	
+
 		TESSELLATION_INTERPOLATE_MEMBER(InterpolantsVSToPS.TangentToWorld0.xyz);
 		TESSELLATION_INTERPOLATE_MEMBER(InterpolantsVSToPS.TangentToWorld2);
 #if INTERPOLATE_VERTEX_COLOR
diff --git a/Engine/Shaders/Private/NiagaraRibbonVertexFactory.ush b/Engine/Shaders/Private/NiagaraRibbonVertexFactory.ush
--- a/Engine/Shaders/Private/NiagaraRibbonVertexFactory.ush
+++ b/Engine/Shaders/Private/NiagaraRibbonVertexFactory.ush
@@ -7,8 +7,11 @@
 #include "VertexFactoryCommon.ush"
 #include "NiagaraVFParticleAccess.usf"
 
-
-
+#define USE_PARTICLE_POSITION (FEATURE_LEVEL >= FEATURE_LEVEL_SM5 && NEEDS_PARTICLE_POSITION)
+#define USE_PARTICLE_VELOCITY (FEATURE_LEVEL >= FEATURE_LEVEL_SM5 && NEEDS_PARTICLE_VELOCITY)
+#define USE_PARTICLE_TIME (FEATURE_LEVEL >= FEATURE_LEVEL_SM5 && NEEDS_PARTICLE_TIME)
+#define USE_PARTICLE_SIZE (FEATURE_LEVEL >= FEATURE_LEVEL_SM5 && NEEDS_PARTICLE_SIZE)
+#define USE_PARTICLE_RANDOM (FEATURE_LEVEL >= FEATURE_LEVEL_SM5 && NEEDS_PARTICLE_RANDOM)
 
 float3 TransformPosition(float3 InPosition)
 {
@@ -67,6 +70,28 @@ float4 GetNiagaraParticleColor(uint InstanceID)
 	return GetVec4(NiagaraRibbonVF.ColorDataOffset, InstanceID);
 }
 
+#if USE_PARTICLE_TIME
+float GetNiagaraNormalizedAge(uint InstanceID)
+{
+	if (NiagaraRibbonVF.NormalizedAgeDataOffset == -1)
+	{
+		return 0;
+	}
+	return GetFloat(NiagaraRibbonVF.NormalizedAgeDataOffset, InstanceID);
+}
+#endif
+
+#if USE_PARTICLE_RANDOM
+float GetNiagaraMaterialRandom(uint InstanceID)
+{
+	if (NiagaraRibbonVF.MaterialRandomDataOffset == -1)
+	{
+		return 0;
+	}
+	return GetFloat(NiagaraRibbonVF.MaterialRandomDataOffset, InstanceID);
+}
+#endif
+
 #if NUM_DYNAMIC_PARAMETERS > 0
 float4 GetNiagaraParticleDynamicParameters(uint InstanceID)
 {
@@ -166,6 +191,27 @@ struct FVertexFactoryInterpolantsVSToPS
 	float4  Color				: TEXCOORD2;
 #endif
 
+#if USE_PARTICLE_POSITION
+	/** Cam-relative (translated) particle center and radius */
+	float4 TranslatedWorldPositionAndSize : PARTICLE_POSITION;
+#endif
+
+#if USE_PARTICLE_VELOCITY
+	float4 ParticleVelocity : PARTICLE_VELOCITY;
+#endif
+
+#if USE_PARTICLE_TIME
+	float RelativeTime : PARTICLE_TIME;
+#endif
+
+#if USE_PARTICLE_SIZE
+	float2 ParticleSize : PARTICLE_SIZE;
+#endif
+
+#if USE_PARTICLE_RANDOM
+	float ParticleRandom : PARTICLE_RANDOM;
+#endif
+
 #if INSTANCED_STEREO
 	nointerpolation uint PackedEyeIndex : PACKED_EYE_INDEX;
 #endif
@@ -184,6 +230,26 @@ struct FVertexFactoryIntermediates
 	float3 Up;
 	float3 CameraDirection;
 	float4 TexCoord;
+	
+#if USE_PARTICLE_POSITION
+	float4 TranslatedWorldPositionAndSize;
+#endif
+
+#if USE_PARTICLE_VELOCITY
+	float4 ParticleVelocity;
+#endif
+
+#if USE_PARTICLE_TIME
+	float RelativeTime;
+#endif
+
+#if USE_PARTICLE_SIZE
+	float2 ParticleSize;
+#endif
+
+#if USE_PARTICLE_RANDOM
+	float ParticleRandom;
+#endif
 
 #if NUM_DYNAMIC_PARAMETERS > 0
 	float4 DynamicParameter;
@@ -317,6 +383,28 @@ FMaterialPixelParameters GetMaterialPixelParameters(FVertexFactoryInterpolantsVS
 	Result.Particle.DynamicParameter3 = Interpolants.DynamicParameter3;
 #endif
 
+#if USE_PARTICLE_POSITION
+	Result.Particle.TranslatedWorldPositionAndSize = Interpolants.TranslatedWorldPositionAndSize;
+#endif
+
+#if USE_PARTICLE_VELOCITY
+	Result.Particle.Velocity = Interpolants.ParticleVelocity;
+#endif
+
+#if USE_PARTICLE_TIME
+	Result.Particle.RelativeTime = Interpolants.RelativeTime;
+#endif
+
+#if USE_PARTICLE_SIZE
+	Result.Particle.Size = Interpolants.ParticleSize;
+#endif
+
+#if USE_PARTICLE_RANDOM
+	Result.Particle.Random = Interpolants.ParticleRandom;
+#else
+	Result.Particle.Random = 0.0f;
+#endif
+
 	return Result;
 }
 
@@ -344,10 +432,26 @@ FMaterialVertexParameters GetMaterialVertexParameters(FVertexFactoryInput Input,
 	Result.Particle.DynamicParameter3 = Intermediates.DynamicParameter3;
 #endif
 
+#if USE_PARTICLE_POSITION
+	Result.Particle.TranslatedWorldPositionAndSize = Intermediates.TranslatedWorldPositionAndSize;
+#endif
+
 #if USE_PARTICLE_VELOCITY
 	Result.Particle.Velocity = Intermediates.ParticleVelocity;
 #endif
 
+#if USE_PARTICLE_TIME
+	Result.Particle.RelativeTime = Intermediates.RelativeTime;
+#endif
+
+#if USE_PARTICLE_SIZE
+	Result.Particle.Size = Intermediates.ParticleSize;
+#endif
+
+#if USE_PARTICLE_RANDOM
+	Result.Particle.Random = Intermediates.ParticleRandom;
+#endif
+
 #if NUM_MATERIAL_TEXCOORDS_VERTEX
 	Result.TexCoords[0] = Intermediates.TexCoord.xy;
 	#if NUM_MATERIAL_TEXCOORDS_VERTEX > 1
@@ -416,7 +520,7 @@ float4 CalcWorldPosition(FVertexFactoryInput Input, FVertexFactoryIntermediates
 	float3 NormDir = SafeNormalize(Intermediates.Direction.xyz);
 	float3 RightDir = Intermediates.Right;
 	float3 UpDir = Intermediates.Up;
-	float3 Pos = Intermediates.Position.xyz + RightDir * Sgn * Intermediates.RibbonWidth.xxx;
+	float3 Pos = Intermediates.Position.xyz + RightDir * Sgn * (Intermediates.RibbonWidth.xxx / 2);
 	return TransformLocalToTranslatedWorld(Pos);
 }
 
@@ -464,8 +568,8 @@ FVertexFactoryIntermediates GetVertexFactoryIntermediates(FVertexFactoryInput In
 		Intermediates.Direction = Intermediates.Position - PreviousParticlePos;
 	}
 	
-	float SegmentLength  = length(Intermediates.Direction);
-	Intermediates.OldPosition = Intermediates.Position.xyz - GetNiagaraParticleVelocity(ParticleId) * NiagaraRibbonVF.DeltaSeconds;
+	float3 ParticleVelocity = GetNiagaraParticleVelocity(ParticleId);
+	Intermediates.OldPosition = Intermediates.Position.xyz - ParticleVelocity * NiagaraRibbonVF.DeltaSeconds;
 	Intermediates.Color	= GetNiagaraParticleColor(ParticleId);
 	
 	if(NiagaraRibbonVF.UseCustomFacing)
@@ -519,6 +623,31 @@ FVertexFactoryIntermediates GetVertexFactoryIntermediates(FVertexFactoryInput In
 	Intermediates.TexCoord.xy = UV0ForSegment * UV0Scale + UV0Offset;
 	Intermediates.TexCoord.zw = UV1ForSegment * UV1Scale + UV1Offset;
 
+#if USE_PARTICLE_POSITION
+	float4 ParticleWorldPosition = mul(float4(Intermediates.Position.xyz,1), Primitive.LocalToWorld);
+	float ParticleRadius = .5 * Intermediates.RibbonWidth;
+	float3 ParticleTranslatedWorldPosition = ParticleWorldPosition.xyz + ResolvedView.PreViewTranslation.xyz;
+	Intermediates.TranslatedWorldPositionAndSize = float4(ParticleTranslatedWorldPosition, ParticleRadius);
+#endif
+
+#if USE_PARTICLE_VELOCITY
+	Intermediates.ParticleVelocity.xyz = normalize(ParticleVelocity);
+	Intermediates.ParticleVelocity.w = length(ParticleVelocity);
+#endif
+
+#if USE_PARTICLE_TIME
+	Intermediates.RelativeTime = GetNiagaraNormalizedAge(ParticleId);
+#endif
+
+#if USE_PARTICLE_SIZE
+	float SegmentLength  = length(Intermediates.Direction);
+	Intermediates.ParticleSize = float2(SegmentLength, Intermediates.RibbonWidth);
+#endif
+
+#if USE_PARTICLE_RANDOM
+	Intermediates.ParticleRandom = GetNiagaraMaterialRandom(ParticleId);
+#endif
+
 	return Intermediates;
 }
 
@@ -566,6 +695,26 @@ FVertexFactoryInterpolantsVSToPS VertexFactoryGetInterpolantsVSToPS(FVertexFacto
 	Interpolants.DynamicParameter3 = Intermediates.DynamicParameter3;
 #endif
 
+#if USE_PARTICLE_POSITION
+	Interpolants.TranslatedWorldPositionAndSize = Intermediates.TranslatedWorldPositionAndSize;
+#endif
+
+#if USE_PARTICLE_VELOCITY
+	Interpolants.ParticleVelocity = Intermediates.ParticleVelocity;
+#endif
+
+#if USE_PARTICLE_TIME
+	Interpolants.RelativeTime = Intermediates.RelativeTime;
+#endif
+
+#if USE_PARTICLE_SIZE
+	Interpolants.ParticleSize = Intermediates.ParticleSize;
+#endif
+
+#if USE_PARTICLE_RANDOM
+	Interpolants.ParticleRandom = Intermediates.ParticleRandom;
+#endif
+
 	float3x3 TangentToLocal = CalcTangentBasis(Input, Intermediates);
 	float3x3 TangentToWorld = mul(TangentToLocal, GetLocalToWorld3x3());
 
diff --git a/Engine/Shaders/Private/NiagaraSpriteVertexFactory.ush b/Engine/Shaders/Private/NiagaraSpriteVertexFactory.ush
--- a/Engine/Shaders/Private/NiagaraSpriteVertexFactory.ush
+++ b/Engine/Shaders/Private/NiagaraSpriteVertexFactory.ush
@@ -101,7 +101,7 @@ struct FVertexFactoryInterpolantsVSToPS
 #endif
 
 #if USE_PARTICLE_TIME
-	float ParticleTime : PARTICLE_TIME;
+	float RelativeTime : PARTICLE_TIME;
 #endif
 
 #if USE_PARTICLE_LIGHTING_OFFSET
@@ -172,14 +172,21 @@ struct FVertexFactoryIntermediates
 
 	/** The sub-image lerp. */
 	float SubImageLerp;
+
+#if USE_PARTICLE_TIME
 	/** Relative time. */
 	float RelativeTime;
+#endif
+
 	/** Transform from tangent space to local space. */
 	float3x3 TangentToLocal;
 	/** Size of the particle */
 	float2 ParticleSize;
+
+#if USE_PARTICLE_RANDOM
 	/** Particle Random value */
 	float ParticleRandom;
+#endif
 };
 
 #if NUM_MATERIAL_TEXCOORDS
@@ -272,7 +279,7 @@ FMaterialPixelParameters GetMaterialPixelParameters(FVertexFactoryInterpolantsVS
 #endif
 
 #if USE_PARTICLE_TIME
-	Result.Particle.RelativeTime = Interpolants.ParticleTime;
+	Result.Particle.RelativeTime = Interpolants.RelativeTime;
 #endif
 
 	Result.Particle.MotionBlurFade = 1.0f;
@@ -630,34 +637,45 @@ float4 GetNiagaraParticleColor(uint InstanceID)
 
 float2 GetNiagaraUVScale(uint InstanceID)
 {
-	if(NiagaraSpriteVF.UVScaleOffset  == -1)
+	if(NiagaraSpriteVF.UVScaleDataOffset  == -1)
 	{
 		// See FNiagaraConstants for a synchronized value...
 		return float2(1.0f, 1.0f);
 	}
-	return GetVec2(NiagaraSpriteVF.UVScaleOffset, InstanceID);	
+	return GetVec2(NiagaraSpriteVF.UVScaleDataOffset, InstanceID);	
 }
 
 float GetNiagaraParticleRandom(uint InstanceID)
 {
-	if(NiagaraSpriteVF.ParticleRandomOffset  == -1)
+	if(NiagaraSpriteVF.MaterialRandomDataOffset  == -1)
 	{
 		// See FNiagaraConstants for a synchronized value...
 		return 0.0f;
 	}
-	return GetFloat(NiagaraSpriteVF.ParticleRandomOffset, InstanceID);	
+	return GetFloat(NiagaraSpriteVF.MaterialRandomDataOffset, InstanceID);	
 }
 
 float GetNiagaraCameraOffset(uint InstanceID)
 {
-	if (NiagaraSpriteVF.CameraOffsetOffset == -1)
+	if (NiagaraSpriteVF.CameraOffsetDataOffset == -1)
 	{
 		// See FNiagaraConstants for a synchronized value...
 		return 0.0f;
 	}
-	return GetFloat(NiagaraSpriteVF.CameraOffsetOffset, InstanceID);
+	return GetFloat(NiagaraSpriteVF.CameraOffsetDataOffset, InstanceID);
 }
 
+#if USE_PARTICLE_TIME
+float GetNiagaraNormalizedAge(uint InstanceID)
+{
+	if (NiagaraSpriteVF.NormalizedAgeDataOffset == -1)
+	{
+		return 0;
+	}
+	return GetFloat(NiagaraSpriteVF.NormalizedAgeDataOffset, InstanceID);
+}
+#endif
+
 #if NUM_DYNAMIC_PARAMETERS > 0
 float4 GetNiagaraParticleDynamicParameters(uint InstanceID)
 {
@@ -699,9 +717,9 @@ float3 GetNiagaraParticleFacingVector(uint InstanceID)
 {
 	// See FNiagaraConstants for a synchronized value...
 	float3 Val = float3(1.0f, 0.0f, 0.0f);
-    if(NiagaraSpriteVF.FacingOffset != -1)
+    if(NiagaraSpriteVF.FacingDataOffset != -1)
 	{
-		Val = GetVec3(NiagaraSpriteVF.FacingOffset, InstanceID);
+		Val = GetVec3(NiagaraSpriteVF.FacingDataOffset, InstanceID);
 	}
 	return TransformVector(Val);
 }
@@ -709,9 +727,9 @@ float3 GetNiagaraParticleFacingVector(uint InstanceID)
 float3 GetNiagaraParticleAlignmentVector(uint InstanceID)
 {
 	float3 Val = float3(1.0f, 0.0f, 0.0f);
-    if(NiagaraSpriteVF.AlignmentOffset != -1)
+    if(NiagaraSpriteVF.AlignmentDataOffset != -1)
 	{
-		Val = GetVec3(NiagaraSpriteVF.AlignmentOffset, InstanceID);
+		Val = GetVec3(NiagaraSpriteVF.AlignmentDataOffset, InstanceID);
 	}
 	return TransformVector(Val);
 }
@@ -826,7 +844,7 @@ FVertexFactoryIntermediates GetVertexFactoryIntermediates(FVertexFactoryInput In
 #endif
 
 #if USE_PARTICLE_TIME
-	Intermediates.RelativeTime = Intermediates.Position.w;
+	Intermediates.RelativeTime = GetNiagaraNormalizedAge(ParticleID);
 #endif
 
 #if USE_PARTICLE_RANDOM
@@ -939,7 +957,7 @@ FVertexFactoryInterpolantsVSToPS VertexFactoryGetInterpolantsVSToPS(FVertexFacto
 #endif
 
 #if USE_PARTICLE_TIME
-	Interpolants.ParticleTime = Intermediates.RelativeTime;
+	Interpolants.RelativeTime = Intermediates.RelativeTime;
 #endif
 
 #if USE_PARTICLE_LIGHTING_OFFSET
@@ -992,4 +1010,4 @@ uint VertexFactoryGetEyeIndex(uint InstanceId)
 float4 VertexFactoryGetTranslatedPrimitiveVolumeBounds(FVertexFactoryInterpolantsVSToPS Interpolants)
 {
 	return 0;
-}
\ No newline at end of file
+}
