Fix HQ DoF downsample scaling for mobile.

diff --git a/Engine/Shaders/Private/PostProcessDOF.usf b/Engine/Shaders/Private/PostProcessDOF.usf
--- a/Engine/Shaders/Private/PostProcessDOF.usf
+++ b/Engine/Shaders/Private/PostProcessDOF.usf
@@ -83,11 +83,17 @@ void SetupPS(
 #endif
 	)
 {
+#if MOBILE_SHADING
+	float DoFToScreenScale = 4;
+#else
+	float DoFToScreenScale = 2;
+#endif
+
 	// Manually compute UV and ScreenSpacePos from SvPosition to avoid south east shiting of the pixel in bottom right corner
 	// caused by the DestRect.Size() + 1 viewport size done for odd view rect resolution.
-	float2 ViewportUV = min(SvPositionToViewportUV(SvPosition * 2), 1 - View.ViewSizeAndInvSize.zw * 0.5);
+	float2 ViewportUV = min(SvPositionToViewportUV(SvPosition * DoFToScreenScale), 1 - View.ViewSizeAndInvSize.zw * 0.5);
 	float2 UV = ViewportUVToBufferUV(ViewportUV);
-	float2 ScreenSpacePos = 2 * float2(ViewportUV.x, 1 - ViewportUV.y) - 1;
+	float2 ScreenSpacePos = DoFToScreenScale * float2(ViewportUV.x, 1 - ViewportUV.y) - 1;
 
 	float2 Offset = 0.5f * PostprocessInput0Size.zw;
 
