Add particle cutout support for feature level es31.

diff --git a/Engine/Shaders/Private/ParticleSpriteVertexFactory.ush b/Engine/Shaders/Private/ParticleSpriteVertexFactory.ush
--- a/Engine/Shaders/Private/ParticleSpriteVertexFactory.ush
+++ b/Engine/Shaders/Private/ParticleSpriteVertexFactory.ush
@@ -425,7 +425,7 @@ float3x3 CalcTangentBasis(FVertexFactoryIntermediates Intermediates)
 }
 
 // Note: C++ geometry setup behavior has to match this define
-#define SUPPORTS_PARTICLE_CUTOUTS (FEATURE_LEVEL >= FEATURE_LEVEL_SM4)
+#define SUPPORTS_PARTICLE_CUTOUTS (FEATURE_LEVEL >= FEATURE_LEVEL_ES3_1)
 #if SUPPORTS_PARTICLE_CUTOUTS
 
 /** Number of vertices in each SubUV frame of CutoutGeometry. */
