in order for UNROLL to work -- break statement needs to be removed

#rnx

diff --git a/Engine/Shaders/Private/MobileBasePassPixelShader.usf b/Engine/Shaders/Private/MobileBasePassPixelShader.usf
--- a/Engine/Shaders/Private/MobileBasePassPixelShader.usf
+++ b/Engine/Shaders/Private/MobileBasePassPixelShader.usf
@@ -556,13 +556,8 @@ void Main(
 			UNROLL
 			for (int i = 0; i < MAX_DYNAMIC_POINT_LIGHTS; i++)
 			{
-				if (NumDynamicPointLights < MAX_DYNAMIC_POINT_LIGHTS)
+				if (i < NumDynamicPointLights)
 				{
-					if (i >= NumDynamicPointLights)
-					{
-						break;
-					}
-				}
 			#else
 			for (int i = 0; i < NumDynamicPointLights; i++)
 			{
@@ -600,6 +595,9 @@ void Main(
 				#else
 					Color += (Attenuation * PointNoL) * LightColorAndFalloffExponent[i].rgb * DiffuseColor;
 				#endif
+			#if VARIABLE_NUM_DYNAMIC_POINT_LIGHTS && ( WEBGL || ES2_PROFILE || VULKAN_PROFILE )
+				}
+			#endif
 			}
 		#endif
 
diff --git a/Engine/Source/Developer/ShaderFormatOpenGL/Private/OpenGLShaderCompiler.cpp b/Engine/Source/Developer/ShaderFormatOpenGL/Private/OpenGLShaderCompiler.cpp
--- a/Engine/Source/Developer/ShaderFormatOpenGL/Private/OpenGLShaderCompiler.cpp
+++ b/Engine/Source/Developer/ShaderFormatOpenGL/Private/OpenGLShaderCompiler.cpp
@@ -1354,7 +1354,7 @@ uint32 FOpenGLFrontend::CalculateCrossCompilerFlags(GLSLVersion Version, bool bC
 		CCFlags |= HLSLCC_ApplyCommonSubexpressionElimination;
 	}
 
-	if (bUseFullPrecisionInPS)
+	if (bUseFullPrecisionInPS || Version == GLSL_ES2_WEBGL)
 	{
 		CCFlags |= HLSLCC_UseFullPrecisionInPS;
 	}
