UE-54791 Planar Reflection gamma / brightness incorrect on Samsung Galaxy S5

#Android
#4.21

The planar texture contains encoded data:
- decode it immediately after sampling the texture
- alpha channel was inverted during the PrefilterPlanarReflectionPS step

diff --git a/Engine/Shaders/Private/PlanarReflectionShaders.usf b/Engine/Shaders/Private/PlanarReflectionShaders.usf
--- a/Engine/Shaders/Private/PlanarReflectionShaders.usf
+++ b/Engine/Shaders/Private/PlanarReflectionShaders.usf
@@ -13,7 +13,7 @@ SamplerState SceneColorInputSampler;
 
 float4 SampleSceneColorInput(float2 UV)
 {
-	return Texture2DSampleLevel(SceneColorInputTexture, SceneColorInputSampler, UV, 0);
+	return Decode32BPPHDR(Texture2DSampleLevel(SceneColorInputTexture, SceneColorInputSampler, UV, 0));
 }
 
 float GaussianWeight1d(float u, float KernelRadius)
