Merging //UE4/Dev-Main to Dev-Build (//UE4/Dev-Build)

diff --git a/Engine/Plugins/FX/Niagara/Source/NiagaraEditor/Public/NiagaraNodeOp.h b/Engine/Plugins/FX/Niagara/Source/NiagaraEditor/Public/NiagaraNodeOp.h
--- a/Engine/Plugins/FX/Niagara/Source/NiagaraEditor/Public/NiagaraNodeOp.h
+++ b/Engine/Plugins/FX/Niagara/Source/NiagaraEditor/Public/NiagaraNodeOp.h
@@ -9,14 +9,29 @@
 #include "Widgets/DeclarativeSyntaxSupport.h"
 #include "Widgets/SOverlay.h"
 #include "NiagaraNode.h"
+#include "NiagaraNodeWithDynamicPins.h"
 #include "NiagaraNodeOp.generated.h"
 
 class SGraphNode;
 class SGraphPin;
 class SVerticalBox;
 
+USTRUCT()
+struct FAddedPinData
+{
+	GENERATED_BODY()
+		
+	/** The data type of the added pin */
+	UPROPERTY()
+	FEdGraphPinType PinType;
+
+	/** The name type of the added pin */
+	UPROPERTY()
+	FName PinName;
+};
+
 UCLASS(MinimalAPI)
-class UNiagaraNodeOp : public UNiagaraNode
+class UNiagaraNodeOp : public UNiagaraNodeWithDynamicPins
 {
 	GENERATED_UCLASS_BODY()
 
@@ -26,6 +41,9 @@ public:
 	UPROPERTY()
 	FName OpName;
 
+	UPROPERTY()
+	TArray<FAddedPinData> AddedPins;
+
 	//~ Begin UObject interface
 	virtual void PostLoad() override;
 	//~ End UObject interface
@@ -42,6 +60,26 @@ public:
 	virtual bool RefreshFromExternalChanges() override;
 	virtual ENiagaraNumericOutputTypeSelectionMode GetNumericOutputTypeSelectionMode() const override;
 	//~ End UNiagaraNode Interface
+
+	//~ Begin UNiagaraNodeWithDynamicPins Interface
+	virtual bool AllowNiagaraTypeForAddPin(const FNiagaraTypeDefinition& InType) override;
+	//~ End UNiagaraNodeWithDynamicPins Interface
+
+protected:
+	//~ Begin EdGraphNode Interface
+	virtual void OnPinRemoved(UEdGraphPin* PinToRemove) override;
+	//~ End EdGraphNode Interface
+
+	//~ Begin UNiagaraNodeWithDynamicPins Interface
+	virtual bool AllowDynamicPins() const override;
+	virtual bool CanMovePin(const UEdGraphPin* Pin) const override { return false; }
+	virtual void OnNewTypedPinAdded(UEdGraphPin* NewPin) override;
+	virtual void OnPinRenamed(UEdGraphPin* RenamedPin, const FString& OldName) override;
+	virtual bool CanRemovePin(const UEdGraphPin* Pin) const override;
+	//~ End UNiagaraNodeWithDynamicPins Interface
+	
+private:
+	FName GetUniqueAdditionalPinName() const;
 };
 
 
diff --git a/Engine/Plugins/FX/Niagara/Source/NiagaraEditor/Public/NiagaraNodeOutput.h b/Engine/Plugins/FX/Niagara/Source/NiagaraEditor/Public/NiagaraNodeOutput.h
--- a/Engine/Plugins/FX/Niagara/Source/NiagaraEditor/Public/NiagaraNodeOutput.h
+++ b/Engine/Plugins/FX/Niagara/Source/NiagaraEditor/Public/NiagaraNodeOutput.h
@@ -15,7 +15,7 @@ class UNiagaraNodeOutput : public UNiagaraNode
 {
 	GENERATED_UCLASS_BODY()
 
-	UPROPERTY(EditAnywhere, Category = Output)
+	UPROPERTY(VisibleAnywhere, Category = Output)
 	TArray<FNiagaraVariable> Outputs;
 
 	UPROPERTY()
diff --git a/Engine/Plugins/FX/Niagara/Source/NiagaraEditor/Public/NiagaraNodeWithDynamicPins.h b/Engine/Plugins/FX/Niagara/Source/NiagaraEditor/Public/NiagaraNodeWithDynamicPins.h
--- a/Engine/Plugins/FX/Niagara/Source/NiagaraEditor/Public/NiagaraNodeWithDynamicPins.h
+++ b/Engine/Plugins/FX/Niagara/Source/NiagaraEditor/Public/NiagaraNodeWithDynamicPins.h
@@ -15,7 +15,7 @@ class SNiagaraGraphPinAdd;
 DECLARE_DELEGATE_OneParam(FOnAddParameter, const FNiagaraVariable&);
 
 /** A base node for niagara nodes with pins which can be dynamically added and removed by the user. */
-UCLASS()
+UCLASS(Abstract)
 class UNiagaraNodeWithDynamicPins : public UNiagaraNode
 {
 public:
diff --git a/Engine/Plugins/FX/Niagara/Source/NiagaraEditor/Public/NiagaraParameterMapHistory.h b/Engine/Plugins/FX/Niagara/Source/NiagaraEditor/Public/NiagaraParameterMapHistory.h
--- a/Engine/Plugins/FX/Niagara/Source/NiagaraEditor/Public/NiagaraParameterMapHistory.h
+++ b/Engine/Plugins/FX/Niagara/Source/NiagaraEditor/Public/NiagaraParameterMapHistory.h
@@ -270,7 +270,7 @@ public:
 	/**
 	* Use the current alias map to resolve any aliases in this input variable name.
 	*/
-	FNiagaraVariable ResolveAliases(const FNiagaraVariable& InVar);
+	FNiagaraVariable ResolveAliases(const FNiagaraVariable& InVar) const;
 	
 	/**
 	* Has RegisterNodeVisitation been called yet on the owning node of this pin?
diff --git a/Engine/Plugins/FX/Niagara/Source/NiagaraEditor/Public/NiagaraStackEditorData.h b/Engine/Plugins/FX/Niagara/Source/NiagaraEditor/Public/NiagaraStackEditorData.h
--- a/Engine/Plugins/FX/Niagara/Source/NiagaraEditor/Public/NiagaraStackEditorData.h
+++ b/Engine/Plugins/FX/Niagara/Source/NiagaraEditor/Public/NiagaraStackEditorData.h
@@ -39,6 +39,20 @@ public:
 	 */
 	void SetStackEntryIsExpanded(const FString& StackEntryKey, bool bIsExpanded);
 
+	/*
+	 * Gets whether or not a stack entry was Expanded before triggering a stack search.
+	 * @param bWasExpandedPreSearchDefault The default value to to return if the pre-search expanded state hasn't been set for the stack entry.
+	 * @param StackEntryKey a unique key for the entry.
+	 */
+	bool GetStackEntryWasExpandedPreSearch(const FString& StackEntryKey, bool bWasExpandedPreSearchDefault) const;
+
+	/*
+	 * Sets whether or not a stack entry was Expanded before a stack search was triggered.
+	 * @param StackEntryKey A unique key for the entry.
+	 * @param bWasExpandedPreSearch Whether or not the entry was expanded pre-search.
+	 */
+	void SetStackEntryWasExpandedPreSearch(const FString& StackEntryKey, bool bWasExpandedPreSearch);
+
 	/*
 	* Gets whether or not a stack item is showing advanced items.
 	* @param StackItemKey A unique key for the entry.
@@ -65,6 +79,12 @@ public:
 	/* Sets whether or not item outputs should be shown in the stack. */
 	void SetShowOutputs(bool bInShowOutputs);
 
+	/* Gets whether or not item linked script inputs should be shown in the stack. */
+	bool GetShowLinkedInputs() const;
+
+	/* Sets whether or not item linked script inputs should be shown in the stack. */
+	void SetShowLinkedInputs(bool bInShowLinkedInputs);
+
 	/* Gets the last scroll position for the associated stack. */
 	double GetLastScrollPosition() const;
 
@@ -89,6 +109,9 @@ private:
 	UPROPERTY()
 	TMap<FString, bool> StackEntryKeyToExpandedMap;
 
+	UPROPERTY()
+	TMap<FString, bool> StackEntryKeyToPreSearchExpandedMap;
+
 	UPROPERTY()
 	TMap<FString, bool> StackItemKeyToShowAdvancedMap;
 
@@ -98,6 +121,9 @@ private:
 	UPROPERTY()
 	bool bShowOutputs;
 
+	UPROPERTY()
+	bool bShowLinkedInputs;
+
 	UPROPERTY()
 	double LastScrollPosition;
 
diff --git a/Engine/Plugins/FX/Niagara/Source/NiagaraEditor/Public/ViewModels/Stack/NiagaraStackEntry.h b/Engine/Plugins/FX/Niagara/Source/NiagaraEditor/Public/ViewModels/Stack/NiagaraStackEntry.h
--- a/Engine/Plugins/FX/Niagara/Source/NiagaraEditor/Public/ViewModels/Stack/NiagaraStackEntry.h
+++ b/Engine/Plugins/FX/Niagara/Source/NiagaraEditor/Public/ViewModels/Stack/NiagaraStackEntry.h
@@ -149,6 +149,8 @@ public:
 
 		const TArray<FStackIssueFix>& GetFixes() const;
 
+		void InsertFix(int32 InsertionIdx, const FStackIssueFix& Fix);
+
 	private:
 		EStackIssueSeverity Severity;
 		FText ShortDescription;
@@ -195,9 +197,7 @@ public:
 
 	virtual bool GetShouldShowInStack() const;
 
-	void GetFilteredChildren(TArray<UNiagaraStackEntry*>& OutFilteredChildren);
-
-	void GetUnfilteredChildren(TArray<UNiagaraStackEntry*>& OutUnfilteredChildren);
+	void GetFilteredChildren(TArray<UNiagaraStackEntry*>& OutFilteredChildren) const;
 
 	void GetUnfilteredChildren(TArray<UNiagaraStackEntry*>& OutUnfilteredChildren) const;
 
@@ -232,7 +232,7 @@ public:
 		return nullptr;
 	}
 
-	void GetSearchItems(TArray<FStackSearchItem>& SearchItems) const;
+	virtual void GetSearchItems(TArray<FStackSearchItem>& SearchItems) const;
 
 	virtual UObject* GetExternalAsset() const;
 
@@ -246,6 +246,10 @@ public:
 
 	void SetOnRequestDrop(FOnRequestDrop InOnRequestCanDrop);
 
+	const bool GetIsSearchResult() const;
+
+	void SetIsSearchResult(bool bInIsSearchResult);
+
 protected:
 	virtual void BeginDestroy() override;
 
@@ -257,8 +261,6 @@ protected:
 
 	virtual int32 GetChildIndentLevel() const;
 
-	virtual void GetAdditionalSearchItemsInternal(TArray<FStackSearchItem>& SearchItems) const;
-
 	virtual TOptional<FDropResult> CanDropInternal(const TArray<UNiagaraStackEntry*>& DraggedEntries);
 
 	virtual TOptional<FDropResult> DropInternal(const TArray<UNiagaraStackEntry*>& DraggedEntries);
@@ -325,4 +327,6 @@ private:
 	TArray<FStackIssue> StackIssues;
 
 	bool bIsFinalized;
+
+	bool bIsSearchResult;
 };
\ No newline at end of file
diff --git a/Engine/Plugins/FX/Niagara/Source/NiagaraEditor/Public/ViewModels/Stack/NiagaraStackFunctionInput.h b/Engine/Plugins/FX/Niagara/Source/NiagaraEditor/Public/ViewModels/Stack/NiagaraStackFunctionInput.h
--- a/Engine/Plugins/FX/Niagara/Source/NiagaraEditor/Public/ViewModels/Stack/NiagaraStackFunctionInput.h
+++ b/Engine/Plugins/FX/Niagara/Source/NiagaraEditor/Public/ViewModels/Stack/NiagaraStackFunctionInput.h
@@ -198,6 +198,22 @@ public:
 	/** Gets whether or not this input is used as an edit condition for another input and should be hidden. */
 	bool GetIsInlineEditConditionToggle() const;
 
+	/** Gets whether or not a dynamic input script reassignment is pending.  This can happen when trying to fix dynamic inputs which are missing their scripts. */
+	bool GetIsDynamicInputScriptReassignmentPending() const;
+
+	/** Gets whether or not a dynamic input script reassignment should be be pending. */
+	void SetIsDynamicInputScriptReassignmentPending(bool bIsPending);
+
+	/** Reassigns the function script for the current dynamic input without resetting the sub-inputs. */
+	void ReassignDynamicInputScript(UNiagaraScript* DynamicInputScript);
+
+	/** Gets whether or not this input is filtered from search results and appearing in stack due to visibility metadata*/
+	bool GetShouldPassFilterForVisibleCondition() const;
+
+public:
+	//~ UNiagaraStackEntry interface
+	virtual void GetSearchItems(TArray<FStackSearchItem>& SearchItems) const override;
+
 protected:
 	//~ UNiagaraStackEntry interface
 	virtual void FinalizeInternal() override;
@@ -431,4 +447,7 @@ private:
 
 	/** Whether or not this input is an edit condition toggle. */
 	bool bIsInlineEditConditionToggle;
+
+	/** Whether or not the dynamic input for this input has a function script reassignment pending due to a request to fix a missing script. */
+	bool bIsDynamicInputScriptReassignmentPending;
 };
\ No newline at end of file
diff --git a/Engine/Plugins/FX/Niagara/Source/NiagaraEditor/Public/ViewModels/Stack/NiagaraStackGraphUtilities.h b/Engine/Plugins/FX/Niagara/Source/NiagaraEditor/Public/ViewModels/Stack/NiagaraStackGraphUtilities.h
--- a/Engine/Plugins/FX/Niagara/Source/NiagaraEditor/Public/ViewModels/Stack/NiagaraStackGraphUtilities.h
+++ b/Engine/Plugins/FX/Niagara/Source/NiagaraEditor/Public/ViewModels/Stack/NiagaraStackGraphUtilities.h
@@ -128,7 +128,7 @@ namespace FNiagaraStackGraphUtilities
 
 	void GetNewParameterAvailableTypes(TArray<FNiagaraTypeDefinition>& OutAvailableTypes);
 
-	void GetScriptAssetsByUsage(ENiagaraScriptUsage AssetUsage, ENiagaraScriptUsage TargetUsage, TArray<FAssetData>& OutAssets);
+	NIAGARAEDITOR_API void GetScriptAssetsByUsage(ENiagaraScriptUsage AssetUsage, ENiagaraScriptUsage TargetUsage, TArray<FAssetData>& OutAssets);
 
 	void GetScriptAssetsByDependencyProvided(ENiagaraScriptUsage AssetUsage, FName DependencyName, TArray<FAssetData>& OutAssets);
 
diff --git a/Engine/Plugins/FX/Niagara/Source/NiagaraEditor/Public/ViewModels/Stack/NiagaraStackModuleItem.h b/Engine/Plugins/FX/Niagara/Source/NiagaraEditor/Public/ViewModels/Stack/NiagaraStackModuleItem.h
--- a/Engine/Plugins/FX/Niagara/Source/NiagaraEditor/Public/ViewModels/Stack/NiagaraStackModuleItem.h
+++ b/Engine/Plugins/FX/Niagara/Source/NiagaraEditor/Public/ViewModels/Stack/NiagaraStackModuleItem.h
@@ -7,8 +7,10 @@
 #include "NiagaraStackModuleItem.generated.h"
 
 class UNiagaraNodeFunctionCall;
+class UNiagaraStackModuleItemLinkedInputCollection;
 class UNiagaraStackFunctionInputCollection;
 class UNiagaraStackModuleItemOutputCollection;
+class UNiagaraScript;
 class INiagaraStackItemGroupAddUtilities;
 struct FAssetData;
 
@@ -55,11 +57,23 @@ public:
 
 	void AddInput(FNiagaraVariable InputParameter);
 
+	/** Gets whether or not a module script reassignment is pending.  This can happen when trying to fix modules which are missing their scripts. */
+	bool GetIsModuleScriptReassignmentPending() const;
+
+	/** Gets whether or not a module script reassignment should be be pending. */
+	void SetIsModuleScriptReassignmentPending(bool bIsPending);
+
+	/** Reassigns the function script for the module without resetting the inputs. */
+	void ReassignModuleScript(UNiagaraScript* ModuleScript);
+
 protected:
 	virtual void RefreshChildrenInternal(const TArray<UNiagaraStackEntry*>& CurrentChildren, TArray<UNiagaraStackEntry*>& NewChildren, TArray<FStackIssue>& NewIssues) override;
 
 private:
 	bool FilterOutputCollection(const UNiagaraStackEntry& Child) const;
+	bool FilterOutputCollectionChild(const UNiagaraStackEntry& Child) const;
+	bool FilterLinkedInputCollection(const UNiagaraStackEntry& Child) const;
+	bool FilterLinkedInputCollectionChild(const UNiagaraStackEntry& Child) const;
 	void RefreshIssues(TArray<FStackIssue>& NewIssues);
 
 private:
@@ -72,6 +86,9 @@ private:
 	bool bIsEnabled;
 	bool bCanRefresh;
 
+	UPROPERTY()
+	UNiagaraStackModuleItemLinkedInputCollection* LinkedInputCollection;
+
 	UPROPERTY()
 	UNiagaraStackFunctionInputCollection* InputCollection;
 
@@ -79,4 +96,6 @@ private:
 	UNiagaraStackModuleItemOutputCollection* OutputCollection;
 
 	INiagaraStackItemGroupAddUtilities* GroupAddUtilities;
+
+	bool bIsModuleScriptReassignmentPending;
 };
diff --git a/Engine/Plugins/FX/Niagara/Source/NiagaraEditor/Public/ViewModels/Stack/NiagaraStackModuleItemLinkedInputCollection.h b/Engine/Plugins/FX/Niagara/Source/NiagaraEditor/Public/ViewModels/Stack/NiagaraStackModuleItemLinkedInputCollection.h
new file mode 100644
--- /dev/null
+++ b/Engine/Plugins/FX/Niagara/Source/NiagaraEditor/Public/ViewModels/Stack/NiagaraStackModuleItemLinkedInputCollection.h
@@ -0,0 +1,32 @@
+// Copyright 1998-2019 Epic Games, Inc. All Rights Reserved.
+
+#pragma once
+
+#include "ViewModels/Stack/NiagaraStackEntry.h"
+#include "NiagaraStackModuleItemLinkedInputCollection.generated.h"
+
+class UNiagaraNodeFunctionCall;
+
+UCLASS()
+class NIAGARAEDITOR_API UNiagaraStackModuleItemLinkedInputCollection : public UNiagaraStackEntry
+{
+	GENERATED_BODY()
+
+public:
+	UNiagaraStackModuleItemLinkedInputCollection();
+
+	void Initialize(FRequiredEntryData InRequiredEntryData, UNiagaraNodeFunctionCall& InFunctionCallNode);
+
+	//~ UNiagaraStackEntry interface
+	virtual FText GetDisplayName() const override;
+	virtual bool IsExpandedByDefault() const override;
+	virtual bool GetIsEnabled() const override;
+	virtual EStackRowStyle GetStackRowStyle() const override;
+	virtual bool GetShouldShowInStack() const override;
+
+protected:
+	virtual void RefreshChildrenInternal(const TArray<UNiagaraStackEntry*>& CurrentChildren, TArray<UNiagaraStackEntry*>& NewChildren, TArray<FStackIssue>& NewIssues) override;
+
+private:
+	UNiagaraNodeFunctionCall* FunctionCallNode;
+};
diff --git a/Engine/Plugins/FX/Niagara/Source/NiagaraEditor/Public/ViewModels/Stack/NiagaraStackPropertyRow.h b/Engine/Plugins/FX/Niagara/Source/NiagaraEditor/Public/ViewModels/Stack/NiagaraStackPropertyRow.h
--- a/Engine/Plugins/FX/Niagara/Source/NiagaraEditor/Public/ViewModels/Stack/NiagaraStackPropertyRow.h
+++ b/Engine/Plugins/FX/Niagara/Source/NiagaraEditor/Public/ViewModels/Stack/NiagaraStackPropertyRow.h
@@ -23,7 +23,7 @@ public:
 protected:
 	virtual void RefreshChildrenInternal(const TArray<UNiagaraStackEntry*>& CurrentChildren, TArray<UNiagaraStackEntry*>& NewChildren, TArray<FStackIssue>& NewIssues) override;
 
-	virtual void GetAdditionalSearchItemsInternal(TArray<FStackSearchItem>& SearchItems) const override;
+	virtual void GetSearchItems(TArray<FStackSearchItem>& SearchItems) const override;
 
 private:
 	TSharedPtr<IDetailTreeNode> DetailTreeNode;
diff --git a/Engine/Plugins/FX/Niagara/Source/NiagaraEditor/Public/ViewModels/Stack/NiagaraStackViewModel.h b/Engine/Plugins/FX/Niagara/Source/NiagaraEditor/Public/ViewModels/Stack/NiagaraStackViewModel.h
--- a/Engine/Plugins/FX/Niagara/Source/NiagaraEditor/Public/ViewModels/Stack/NiagaraStackViewModel.h
+++ b/Engine/Plugins/FX/Niagara/Source/NiagaraEditor/Public/ViewModels/Stack/NiagaraStackViewModel.h
@@ -52,6 +52,8 @@ public:
 
 	bool GetShowOutputs() const;
 	void SetShowOutputs(bool bInShowOutputs);
+	bool GetShowLinkedInputs() const;
+	void SetShowLinkedInputs(bool bInShowLinkedInputs);
 
 	double GetLastScrollPosition() const;
 	void SetLastScrollPosition(double InLastScrollPosition);
@@ -110,6 +112,7 @@ private:
 		TArray<FSearchWorkItem>& TraversedArray);
 	bool ItemMatchesSearchCriteria(UNiagaraStackEntry::FStackSearchItem SearchItem);
 	void GeneratePathForEntry(UNiagaraStackEntry* Root, UNiagaraStackEntry* Entry, TArray<UNiagaraStackEntry*> CurrentPath, TArray<UNiagaraStackEntry*>& EntryPath);
+	void RestoreStackEntryExpansionPreSearch();
 
 private:
 	TSharedPtr<FNiagaraEmitterHandleViewModel> EmitterHandleViewModel;
@@ -127,6 +130,7 @@ private:
 	int CurrentFocusedSearchMatchIndex;
 	FOnSearchCompleted SearchCompletedDelegate;
 	TArray<FSearchWorkItem> ItemsToSearch;
+	TArray<FSearchWorkItem> ItemsToRestoreExpansionState;
 	TArray<FSearchResult> CurrentSearchResults;
 	static const double MaxSearchTime;
 	bool bRestartSearch;
diff --git a/Engine/Plugins/FX/Niagara/Source/NiagaraEditorWidgets/Private/SNiagaraStack.cpp b/Engine/Plugins/FX/Niagara/Source/NiagaraEditorWidgets/Private/SNiagaraStack.cpp
--- a/Engine/Plugins/FX/Niagara/Source/NiagaraEditorWidgets/Private/SNiagaraStack.cpp
+++ b/Engine/Plugins/FX/Niagara/Source/NiagaraEditorWidgets/Private/SNiagaraStack.cpp
@@ -14,6 +14,7 @@
 #include "ViewModels/Stack/NiagaraStackViewModel.h"
 #include "ViewModels/Stack/NiagaraStackSpacer.h"
 #include "ViewModels/Stack/NiagaraStackModuleItemOutputCollection.h"
+#include "ViewModels/Stack/NiagaraStackModuleItemLinkedInputCollection.h"
 #include "ViewModels/Stack/NiagaraStackModuleItemOutput.h"
 #include "ViewModels/Stack/NiagaraStackErrorItem.h"
 #include "ViewModels/Stack/NiagaraStackInputCategory.h" 
@@ -53,6 +54,7 @@
 #include "ScopedTransaction.h"
 #include "Widgets/Layout/SWrapBox.h"
 #include "Stack/SNiagaraStackSpacer.h"
+#include "ViewModels/Stack/NiagaraStackRoot.h"
 
 /** Contains data for a socket drag and drop operation in the StackEntry node. */
 class FNiagaraStackEntryDragDropOp : public FDecoratedDragDropOp
@@ -220,10 +222,10 @@ void SNiagaraStack::Construct(const FArguments& InArgs, UNiagaraStackViewModel*
 	};
 	HeaderBox->SetOnMouseButtonUp(FPointerEventHandler::CreateLambda(OnHeaderMouseButtonUp));
 	
-	PrimeTreeExpansion();
+	SynchronizeTreeExpansion();
 }
 
-void SNiagaraStack::PrimeTreeExpansion()
+void SNiagaraStack::SynchronizeTreeExpansion()
 {
 	TArray<UNiagaraStackEntry*> EntriesToProcess(StackViewModel->GetRootEntries());
 	while (EntriesToProcess.Num() > 0)
@@ -413,33 +415,73 @@ FReply SNiagaraStack::PinButtonPressed()
 
 void SNiagaraStack::OnSearchTextChanged(const FText& SearchText)
 {
-	if (StackViewModel->GetCurrentSearchText().CompareTo(SearchText) != 0)
+	bNeedsJumpToNextOccurence = true;
+	StackViewModel->OnSearchTextChanged(SearchText);
+}
+
+FReply SNiagaraStack::ScrollToNextMatch()
+{
+	
+	const int NextMatchIndex = StackViewModel->GetCurrentFocusedMatchIndex() + 1;
+	TArray<UNiagaraStackViewModel::FSearchResult> CurrentSearchResults = StackViewModel->GetCurrentSearchResults();
+	if (CurrentSearchResults.Num() != 0)
 	{
-		if (SearchExpandTimer.IsValid())
+		if (NextMatchIndex < CurrentSearchResults.Num())
 		{
-			UnRegisterActiveTimer(SearchExpandTimer.ToSharedRef());
+			for (auto SearchResultEntry : CurrentSearchResults[NextMatchIndex].EntryPath)
+			{
+				if (!SearchResultEntry->IsA<UNiagaraStackRoot>())
+				{
+					SearchResultEntry->SetIsExpanded(true);
+				}
+			}
 		}
-		// restore expansion state of previous search
-		for (auto SearchResult : StackViewModel->GetCurrentSearchResults())
+		else
 		{
-			for (UNiagaraStackEntry* ParentalUnit : SearchResult.EntryPath)
+			for (auto SearchResultEntry : CurrentSearchResults[0].EntryPath)
 			{
-				StackTree->SetItemExpansion(ParentalUnit, ParentalUnit->GetIsExpanded());
+				if (!SearchResultEntry->IsA<UNiagaraStackRoot>())
+				{
+					SearchResultEntry->SetIsExpanded(true);
+				}
 			}
 		}
-		bNeedsJumpToNextOccurence = true;
-		StackViewModel->OnSearchTextChanged(SearchText);
+		SynchronizeTreeExpansion();
 	}
-}
 
-FReply SNiagaraStack::ScrollToNextMatch()
-{
 	AddSearchScrollOffset(1);
 	return FReply::Handled();
 }
 
 FReply SNiagaraStack::ScrollToPreviousMatch()
 {
+	const int PreviousMatchIndex = StackViewModel->GetCurrentFocusedMatchIndex() - 1;
+	TArray<UNiagaraStackViewModel::FSearchResult> CurrentSearchResults = StackViewModel->GetCurrentSearchResults();
+	if (CurrentSearchResults.Num() != 0)
+	{
+		if (PreviousMatchIndex > 0)
+		{
+			for (auto SearchResultEntry : CurrentSearchResults[PreviousMatchIndex].EntryPath)
+			{
+				if (!SearchResultEntry->IsA<UNiagaraStackRoot>())
+				{
+					SearchResultEntry->SetIsExpanded(true);
+				}
+			}
+		}
+		else
+		{
+			for (auto SearchResultEntry : CurrentSearchResults.Last().EntryPath)
+			{
+				if (!SearchResultEntry->IsA<UNiagaraStackRoot>())
+				{
+					SearchResultEntry->SetIsExpanded(true);
+				}
+			}
+		}
+		SynchronizeTreeExpansion();
+	}
+
 	// move current match to the previous one in the StackTree, wrap around
 	AddSearchScrollOffset(-1);
 	return FReply::Handled();
@@ -589,9 +631,19 @@ TSharedRef<SWidget> SNiagaraStack::GetViewOptionsMenu() const
 			FGetActionCheckState::CreateLambda([=]() { return StackViewModel->GetShowAllAdvanced() ? ECheckBoxState::Checked : ECheckBoxState::Unchecked; })),
 		NAME_None, EUserInterfaceActionType::Check);
 
+	MenuBuilder.AddMenuEntry(
+		LOCTEXT("ShowLinkedInputsLabel", "Show Linked Script Inputs"),
+		LOCTEXT("ShowLinkedInputsToolTip", "Whether or not to show internal module linked inputs in the stack."),
+		FSlateIcon(),
+		FUIAction(
+			FExecuteAction::CreateLambda([=]() { StackViewModel->SetShowLinkedInputs(!StackViewModel->GetShowLinkedInputs()); }),
+			FCanExecuteAction(),
+			FGetActionCheckState::CreateLambda([=]() { return StackViewModel->GetShowLinkedInputs() ? ECheckBoxState::Checked : ECheckBoxState::Unchecked; })),
+		NAME_None, EUserInterfaceActionType::Check);
+
 	MenuBuilder.AddMenuEntry(
 		LOCTEXT("ShowOutputsLabel", "Show Outputs"),
-		LOCTEXT("ShowOutputsToolTip", "Whether or now to show module outputs in the stack."),
+		LOCTEXT("ShowOutputsToolTip", "Whether or not to show module outputs in the stack."),
 		FSlateIcon(),
 		FUIAction(
 			FExecuteAction::CreateLambda([=]() { StackViewModel->SetShowOutputs(!StackViewModel->GetShowOutputs()); }),
@@ -664,12 +716,6 @@ FReply SNiagaraStack::OnRowAcceptDrop(const FDragDropEvent& InDragDropEvent, EIt
 }
 
 void SNiagaraStack::OnStackSearchComplete()
-{
-	// fire up timer to expand all parentchains!!
-	SearchExpandTimer = RegisterActiveTimer(0.7f, FWidgetActiveTimerDelegate::CreateSP(this, &SNiagaraStack::TriggerExpandSearchResults));
-}
-
-EActiveTimerReturnType SNiagaraStack::TriggerExpandSearchResults(double InCurrentTime, float InDeltaTime)
 {
 	ExpandSearchResults();
 	if (bNeedsJumpToNextOccurence)
@@ -677,7 +723,6 @@ EActiveTimerReturnType SNiagaraStack::TriggerExpandSearchResults(double InCurren
 		ScrollToNextMatch();
 		bNeedsJumpToNextOccurence = false;
 	}
-	return EActiveTimerReturnType::Stop;
 }
 
 void SNiagaraStack::ExpandSearchResults()
@@ -686,22 +731,17 @@ void SNiagaraStack::ExpandSearchResults()
 	{
 		for (UNiagaraStackEntry* ParentalUnit : SearchResult.EntryPath)
 		{
-			StackTree->SetItemExpansion(ParentalUnit, true);
+			if (!ParentalUnit->IsA<UNiagaraStackRoot>()) // should not alter expansion state of root entries
+			{
+				ParentalUnit->SetIsExpanded(true);
+			}
 		}
 	}
+	SynchronizeTreeExpansion();
 }
 
 void SNiagaraStack::OnSearchBoxTextCommitted(const FText& NewText, ETextCommit::Type CommitInfo)
 {
-	if (StackViewModel->GetCurrentSearchText().CompareTo(NewText) != 0)
-	{
-		if (SearchExpandTimer.IsValid())
-		{
-			UnRegisterActiveTimer(SearchExpandTimer.ToSharedRef());
-			ExpandSearchResults();
-			SearchExpandTimer.Reset();
-		}
-	}
 	if (bNeedsJumpToNextOccurence || CommitInfo == ETextCommit::OnEnter) // hasn't been autojumped yet or we hit enter
 	{
 		AddSearchScrollOffset(+1);
@@ -933,7 +973,8 @@ SNiagaraStack::FRowWidgets SNiagaraStack::ConstructNameAndValueWidgetsForItem(UN
 			.HighlightText_UObject(StackViewModel, &UNiagaraStackViewModel::GetCurrentSearchText));
 	}
 	else if (Item->IsA<UNiagaraStackFunctionInputCollection>() ||
-		Item->IsA<UNiagaraStackModuleItemOutputCollection>())
+		Item->IsA<UNiagaraStackModuleItemOutputCollection>() ||
+		Item->IsA<UNiagaraStackModuleItemLinkedInputCollection>())
 	{
 		return FRowWidgets(
 			SNew(STextBlock)
@@ -1015,7 +1056,7 @@ void SNiagaraStack::OnContentColumnWidthChanged(float Width)
 
 void SNiagaraStack::StackStructureChanged()
 {
-	PrimeTreeExpansion();
+	SynchronizeTreeExpansion();
 	StackTree->RequestTreeRefresh();
 }
 
diff --git a/Engine/Plugins/FX/Niagara/Source/NiagaraEditorWidgets/Private/SNiagaraStack.h b/Engine/Plugins/FX/Niagara/Source/NiagaraEditorWidgets/Private/SNiagaraStack.h
--- a/Engine/Plugins/FX/Niagara/Source/NiagaraEditorWidgets/Private/SNiagaraStack.h
+++ b/Engine/Plugins/FX/Niagara/Source/NiagaraEditorWidgets/Private/SNiagaraStack.h
@@ -46,7 +46,7 @@ private:
 		TSharedPtr<SWidget> ValueWidget;
 	};
 
-	void PrimeTreeExpansion();
+	void SynchronizeTreeExpansion();
 
 	TSharedRef<ITableRow> OnGenerateRowForStackItem(UNiagaraStackEntry* Item, const TSharedRef<STableViewBase>& OwnerTable);
 
@@ -96,7 +96,6 @@ private:
 	FSlateColor GetTextColorForItem(UNiagaraStackEntry* Item) const;
 	void AddSearchScrollOffset(int NumberOfSteps);
 	void OnStackSearchComplete();
-	EActiveTimerReturnType TriggerExpandSearchResults(double InCurrentTime, float InDeltaTime);
 	void ExpandSearchResults();
 	bool IsEntryFocusedInSearch(UNiagaraStackEntry* Entry) const;
 	
@@ -139,7 +138,6 @@ private:
 
 	// ~ search stuff
 	TSharedPtr<SSearchBox> SearchBox;
-	TSharedPtr<FActiveTimerHandle> SearchExpandTimer;
 	static const FText OccurencesFormat;
 	bool bNeedsJumpToNextOccurence;
 	
diff --git a/Engine/Plugins/FX/Niagara/Source/NiagaraEditorWidgets/Private/Stack/SNiagaraStackFunctionInputValue.cpp b/Engine/Plugins/FX/Niagara/Source/NiagaraEditorWidgets/Private/Stack/SNiagaraStackFunctionInputValue.cpp
--- a/Engine/Plugins/FX/Niagara/Source/NiagaraEditorWidgets/Private/Stack/SNiagaraStackFunctionInputValue.cpp
+++ b/Engine/Plugins/FX/Niagara/Source/NiagaraEditorWidgets/Private/Stack/SNiagaraStackFunctionInputValue.cpp
@@ -23,6 +23,8 @@
 #include "IStructureDetailsView.h"
 #include "SDropTarget.h"
 #include "Modules/ModuleManager.h"
+#include "AssetRegistryModule.h"
+#include "NiagaraParameterCollection.h"
 
 #define LOCTEXT_NAMESPACE "NiagaraStackFunctionInputValue"
 
@@ -95,6 +97,7 @@ void SNiagaraStackFunctionInputValue::Construct(const FArguments& InArgs, UNiaga
 						SNew(STextBlock)
 						.TextStyle(FNiagaraEditorStyle::Get(), "NiagaraEditor.ParameterText")
 						.Text(this, &SNiagaraStackFunctionInputValue::GetLinkedValueHandleText)
+						.OnDoubleClicked(this, &SNiagaraStackFunctionInputValue::OnLinkedInputDoubleClicked)
 					]
 				]
 
@@ -222,6 +225,15 @@ void SNiagaraStackFunctionInputValue::Construct(const FArguments& InArgs, UNiaga
 	];
 }
 
+void SNiagaraStackFunctionInputValue::Tick(const FGeometry& AllottedGeometry, const double InCurrentTime, const float InDeltaTime)
+{
+	if (FunctionInput->GetIsDynamicInputScriptReassignmentPending())
+	{
+		FunctionInput->SetIsDynamicInputScriptReassignmentPending(false);
+		ShowReassignDynamicInputScriptMenu();
+	}
+}
+
 bool SNiagaraStackFunctionInputValue::GetInputEnabled() const
 {
 	return FunctionInput->GetHasEditCondition() == false || FunctionInput->GetEditConditionEnabled();
@@ -413,6 +425,29 @@ FReply SNiagaraStackFunctionInputValue::DynamicInputTextDoubleClicked()
 	return FReply::Unhandled();
 }
 
+FReply SNiagaraStackFunctionInputValue::OnLinkedInputDoubleClicked()
+{
+	FString ParamCollection;
+	FString ParamName;
+	FunctionInput->GetLinkedValueHandle().GetName().ToString().Split(TEXT("."), &ParamCollection, &ParamName);
+
+	FAssetRegistryModule& AssetRegistryModule = FModuleManager::LoadModuleChecked<FAssetRegistryModule>("AssetRegistry");
+	TArray<FAssetData> CollectionAssets;
+	AssetRegistryModule.Get().GetAssetsByClass(UNiagaraParameterCollection::StaticClass()->GetFName(), CollectionAssets);
+
+	for (FAssetData& CollectionAsset : CollectionAssets)
+	{
+		UNiagaraParameterCollection* Collection = CastChecked<UNiagaraParameterCollection>(CollectionAsset.GetAsset());
+		if (Collection && Collection->GetNamespace() == *ParamCollection)
+		{
+			FAssetEditorManager::Get().OpenEditorForAsset(Collection);
+			return FReply::Handled();
+		}
+	}
+
+	return FReply::Unhandled();
+}
+
 TSharedRef<SExpanderArrow> SNiagaraStackFunctionInputValue::CreateCustomNiagaraFunctionInputActionExpander(const FCustomExpanderData& ActionMenuData)
 {
 	return SNew(SNiagaraFunctionInputActionMenuExpander, ActionMenuData);
@@ -497,6 +532,7 @@ void SNiagaraStackFunctionInputValue::CollectAllActions(FGraphActionListBuilderB
 	TArray<FNiagaraParameterHandle> AvailableHandles;
 	FunctionInput->GetAvailableParameterHandles(AvailableHandles);
 
+	TArray<FNiagaraParameterHandle> ParameterCollectionHandles;
 	TArray<FNiagaraParameterHandle> UserHandles;
 	TArray<FNiagaraParameterHandle> EngineHandles;
 	TArray<FNiagaraParameterHandle> SystemHandles;
@@ -505,7 +541,11 @@ void SNiagaraStackFunctionInputValue::CollectAllActions(FGraphActionListBuilderB
 	TArray<FNiagaraParameterHandle> OtherHandles;
 	for (const FNiagaraParameterHandle AvailableHandle : AvailableHandles)
 	{
-		if (AvailableHandle.IsUserHandle())
+		 if (AvailableHandle.IsParameterCollectionHandle())
+		{
+			ParameterCollectionHandles.Add(AvailableHandle);
+		}
+		else if (AvailableHandle.IsUserHandle())
 		{
 			UserHandles.Add(AvailableHandle);
 		}
@@ -546,6 +586,7 @@ void SNiagaraStackFunctionInputValue::CollectAllActions(FGraphActionListBuilderB
 			}
 		};
 
+		AddMenuItemsForHandleList(ParameterCollectionHandles, LOCTEXT("NPC", "Parameter Collections"));
 		AddMenuItemsForHandleList(UserHandles, LOCTEXT("UserSection", "User Exposed"));
 		AddMenuItemsForHandleList(EngineHandles, LOCTEXT("EngineSection", "Engine"));
 		AddMenuItemsForHandleList(SystemHandles, LOCTEXT("SystemSection", "System"));
@@ -766,4 +807,49 @@ bool SNiagaraStackFunctionInputValue::OnFunctionInputAllowDrop(TSharedPtr<FDragD
 	return false;
 }
 
+void ReassignDynamicInputScript(UNiagaraStackFunctionInput* FunctionInput, UNiagaraScript* NewDynamicInputScript)
+{
+	FunctionInput->ReassignDynamicInputScript(NewDynamicInputScript);
+}
+
+void CollectDynamicInputActions(FGraphActionListBuilderBase& DynamicInputActions, UNiagaraStackFunctionInput* FunctionInput)
+{
+	const FText CategoryName = LOCTEXT("DynamicInputValueCategory", "Dynamic Inputs");
+	TArray<UNiagaraScript*> DynamicInputScripts;
+	FunctionInput->GetAvailableDynamicInputs(DynamicInputScripts);
+	for (UNiagaraScript* DynamicInputScript : DynamicInputScripts)
+	{
+		const FText DynamicInputText = FText::FromString(FName::NameToDisplayString(DynamicInputScript->GetName(), false));
+		const FText Tooltip = FNiagaraEditorUtilities::FormatScriptAssetDescription(DynamicInputScript->Description, *DynamicInputScript->GetPathName());
+		TSharedPtr<FNiagaraMenuAction> DynamicInputAction(new FNiagaraMenuAction(CategoryName, DynamicInputText, Tooltip, 0, DynamicInputScript->Keywords,
+			FNiagaraMenuAction::FOnExecuteStackAction::CreateStatic(&ReassignDynamicInputScript, FunctionInput, DynamicInputScript)));
+		DynamicInputActions.AddAction(DynamicInputAction);
+	}
+}
+
+void SNiagaraStackFunctionInputValue::ShowReassignDynamicInputScriptMenu()
+{
+	TSharedRef<SBorder> MenuWidget = SNew(SBorder)
+		.BorderImage(FEditorStyle::GetBrush("Menu.Background"))
+		.Padding(5)
+		[
+			SNew(SBox)
+			.WidthOverride(300)
+			.HeightOverride(400)
+			[
+				SNew(SGraphActionMenu)
+				.OnActionSelected(this, &SNiagaraStackFunctionInputValue::OnActionSelected)
+				.OnCollectAllActions_Static(CollectDynamicInputActions, FunctionInput)
+				.AutoExpandActionMenu(true)
+				.ShowFilterTextBox(true)
+				.OnCreateCustomRowExpander_Static(&CreateCustomNiagaraFunctionInputActionExpander)
+			]
+		];
+
+	FGeometry ThisGeometry = GetCachedGeometry();
+	bool bAutoAdjustForDpiScale = false; // Don't adjust for dpi scale because the push menu command is expecting an unscaled position.
+	FVector2D MenuPosition = FSlateApplication::Get().CalculatePopupWindowPosition(ThisGeometry.GetLayoutBoundingRect(), MenuWidget->GetDesiredSize(), bAutoAdjustForDpiScale);
+	FSlateApplication::Get().PushMenu(AsShared(), FWidgetPath(), MenuWidget, MenuPosition, FPopupTransitionEffect::ContextMenu);
+}
+
 #undef LOCTEXT_NAMESPACE
diff --git a/Engine/Plugins/FX/Niagara/Source/NiagaraEditorWidgets/Private/Stack/SNiagaraStackFunctionInputValue.h b/Engine/Plugins/FX/Niagara/Source/NiagaraEditorWidgets/Private/Stack/SNiagaraStackFunctionInputValue.h
--- a/Engine/Plugins/FX/Niagara/Source/NiagaraEditorWidgets/Private/Stack/SNiagaraStackFunctionInputValue.h
+++ b/Engine/Plugins/FX/Niagara/Source/NiagaraEditorWidgets/Private/Stack/SNiagaraStackFunctionInputValue.h
@@ -25,6 +25,8 @@ public:
 
 	void Construct(const FArguments& InArgs, UNiagaraStackFunctionInput* InFunctionInput);
 
+	virtual void Tick(const FGeometry& AllottedGeometry, const double InCurrentTime, const float InDeltaTime) override;
+
 private:
 
 	bool GetInputEnabled() const;
@@ -58,6 +60,7 @@ private:
 	FText GetInvalidValueToolTipText() const;
 
 	FReply DynamicInputTextDoubleClicked();
+	FReply OnLinkedInputDoubleClicked();
 
 	class SNiagaraFunctionInputActionMenuExpander: public SExpanderArrow
 	{
@@ -130,6 +133,8 @@ private:
 
 	bool OnFunctionInputAllowDrop(TSharedPtr<FDragDropOperation> DragDropOperation);
 
+	void ShowReassignDynamicInputScriptMenu();
+
 private:
 	UNiagaraStackFunctionInput* FunctionInput;
 
diff --git a/Engine/Plugins/FX/Niagara/Source/NiagaraEditorWidgets/Private/Stack/SNiagaraStackModuleItem.cpp b/Engine/Plugins/FX/Niagara/Source/NiagaraEditorWidgets/Private/Stack/SNiagaraStackModuleItem.cpp
--- a/Engine/Plugins/FX/Niagara/Source/NiagaraEditorWidgets/Private/Stack/SNiagaraStackModuleItem.cpp
+++ b/Engine/Plugins/FX/Niagara/Source/NiagaraEditorWidgets/Private/Stack/SNiagaraStackModuleItem.cpp
@@ -23,6 +23,8 @@
 #include "Toolkits/AssetEditorManager.h"
 #include "SDropTarget.h"
 #include "SNiagaraStackErrorButton.h"
+#include "NiagaraEditorUtilities.h"
+#include "SGraphActionMenu.h"
 
 #define LOCTEXT_NAMESPACE "NiagaraStackModuleItem"
 
@@ -180,6 +182,15 @@ FReply SNiagaraStackModuleItem::OnMouseButtonDoubleClick(const FGeometry& InMyGe
 	return FReply::Unhandled();
 }
 
+void SNiagaraStackModuleItem::Tick(const FGeometry& AllottedGeometry, const double InCurrentTime, const float InDeltaTime)
+{
+	if (ModuleItem->GetIsModuleScriptReassignmentPending())
+	{
+		ModuleItem->SetIsModuleScriptReassignmentPending(false);
+		ShowReassignModuleScriptMenu();
+	}
+}
+
 ECheckBoxState SNiagaraStackModuleItem::GetCheckState() const
 {
 	return ModuleItem->GetIsEnabled() ? ECheckBoxState::Checked : ECheckBoxState::Unchecked;
@@ -273,7 +284,8 @@ void SNiagaraStackModuleItem::ShowInsertModuleMenu(int32 InsertIndex)
 {
 	TSharedRef<SWidget> MenuContent = SNew(SNiagaraStackItemGroupAddMenu, ModuleItem->GetGroupAddUtilities(), InsertIndex);
 	FGeometry ThisGeometry = GetCachedGeometry();
-	FVector2D MenuPosition = FSlateApplication::Get().CalculatePopupWindowPosition(ThisGeometry.GetLayoutBoundingRect(), MenuContent->GetDesiredSize());
+	bool bAutoAdjustForDpiScale = false; // Don't adjust for dpi scale because the push menu command is expecting an unscaled position.
+	FVector2D MenuPosition = FSlateApplication::Get().CalculatePopupWindowPosition(ThisGeometry.GetLayoutBoundingRect(), MenuContent->GetDesiredSize(), bAutoAdjustForDpiScale);
 	FSlateApplication::Get().PushMenu(AsShared(), FWidgetPath(), MenuContent, MenuPosition, FPopupTransitionEffect::ContextMenu);
 }
 
@@ -315,4 +327,86 @@ FText SNiagaraStackModuleItem::GetErrorButtonTooltipText() const
 	return FText::Format(LOCTEXT("ModuleIssuesTooltip", "This module has {0} issues, click to expand."), ModuleItem->GetRecursiveStackIssuesCount());
 }
 
+void ReassignModuleScript(UNiagaraStackModuleItem* ModuleItem, FAssetData NewModuleScriptAsset)
+{
+	UNiagaraScript* NewModuleScript = Cast<UNiagaraScript>(NewModuleScriptAsset.GetAsset());
+	if (NewModuleScript != nullptr)
+	{
+		ModuleItem->ReassignModuleScript(NewModuleScript);
+	}
+}
+
+void OnActionSelected(const TArray<TSharedPtr<FEdGraphSchemaAction>>& SelectedActions, ESelectInfo::Type InSelectionType)
+{
+	if (SelectedActions.Num() == 1 && (InSelectionType == ESelectInfo::OnKeyPress || InSelectionType == ESelectInfo::OnMouseClick))
+	{
+		TSharedPtr<FNiagaraMenuAction> Action = StaticCastSharedPtr<FNiagaraMenuAction>(SelectedActions[0]);
+		if (Action.IsValid())
+		{
+			FSlateApplication::Get().DismissAllMenus();
+			Action->ExecuteAction();
+		}
+	}
+}
+
+void CollectModuleActions(FGraphActionListBuilderBase& ModuleActions, UNiagaraStackModuleItem* ModuleItem)
+{
+	TArray<FAssetData> ModuleAssets;
+	FNiagaraStackGraphUtilities::GetScriptAssetsByUsage(ENiagaraScriptUsage::Module, ModuleItem->GetOutputNode()->GetUsage(), ModuleAssets);
+	for (const FAssetData& ModuleAsset : ModuleAssets)
+	{
+		FText Category;
+		ModuleAsset.GetTagValue(GET_MEMBER_NAME_CHECKED(UNiagaraScript, Category), Category);
+		if (Category.IsEmptyOrWhitespace())
+		{
+			Category = LOCTEXT("ModuleNotCategorized", "Uncategorized Modules");
+		}
+
+		uint32 bDeprecatedModule;
+		ModuleAsset.GetTagValue(GET_MEMBER_NAME_CHECKED(UNiagaraScript, bDeprecated), bDeprecatedModule);
+
+		if (bDeprecatedModule)
+		{
+			continue;
+		}
+
+		FString DisplayNameString = FName::NameToDisplayString(ModuleAsset.AssetName.ToString(), false);
+		FText DisplayName = FText::FromString(DisplayNameString);
+
+		FText AssetDescription;
+		ModuleAsset.GetTagValue(GET_MEMBER_NAME_CHECKED(UNiagaraScript, Description), AssetDescription);
+		FText Description = FNiagaraEditorUtilities::FormatScriptAssetDescription(AssetDescription, ModuleAsset.ObjectPath);
+
+		FText Keywords;
+		ModuleAsset.GetTagValue(GET_MEMBER_NAME_CHECKED(UNiagaraScript, Keywords), Keywords);
+
+		TSharedPtr<FNiagaraMenuAction> ModuleAction(new FNiagaraMenuAction(Category, DisplayName, Description, 0, Keywords,
+			FNiagaraMenuAction::FOnExecuteStackAction::CreateStatic(&ReassignModuleScript, ModuleItem, ModuleAsset)));
+		ModuleActions.AddAction(ModuleAction);
+	}
+}
+
+void SNiagaraStackModuleItem::ShowReassignModuleScriptMenu()
+{
+	TSharedRef<SBorder> MenuWidget = SNew(SBorder)
+		.BorderImage(FEditorStyle::GetBrush("Menu.Background"))
+		.Padding(5)
+		[
+			SNew(SBox)
+			.WidthOverride(300)
+			.HeightOverride(400)
+			[
+				SNew(SGraphActionMenu)
+				.OnActionSelected_Static(OnActionSelected)
+				.OnCollectAllActions_Static(CollectModuleActions, ModuleItem)
+				.ShowFilterTextBox(true)
+			]
+		];
+
+	FGeometry ThisGeometry = GetCachedGeometry();
+	bool bAutoAdjustForDpiScale = false; // Don't adjust for dpi scale because the push menu command is expecting an unscaled position.
+	FVector2D MenuPosition = FSlateApplication::Get().CalculatePopupWindowPosition(ThisGeometry.GetLayoutBoundingRect(), MenuWidget->GetDesiredSize(), bAutoAdjustForDpiScale);
+	FSlateApplication::Get().PushMenu(AsShared(), FWidgetPath(), MenuWidget, MenuPosition, FPopupTransitionEffect::ContextMenu);
+}
+
 #undef LOCTEXT_NAMESPACE
\ No newline at end of file
diff --git a/Engine/Plugins/FX/Niagara/Source/NiagaraEditorWidgets/Private/Stack/SNiagaraStackModuleItem.h b/Engine/Plugins/FX/Niagara/Source/NiagaraEditorWidgets/Private/Stack/SNiagaraStackModuleItem.h
--- a/Engine/Plugins/FX/Niagara/Source/NiagaraEditorWidgets/Private/Stack/SNiagaraStackModuleItem.h
+++ b/Engine/Plugins/FX/Niagara/Source/NiagaraEditorWidgets/Private/Stack/SNiagaraStackModuleItem.h
@@ -25,6 +25,7 @@ public:
 
 	//~ SWidget interface
 	virtual FReply OnMouseButtonDoubleClick(const FGeometry& InMyGeometry, const FPointerEvent& InMouseEvent) override;
+	virtual void Tick(const FGeometry& AllottedGeometry, const double InCurrentTime, const float InDeltaTime) override;
 
 private:
 	ECheckBoxState GetCheckState() const;
@@ -59,6 +60,8 @@ private:
 
 	FText GetErrorButtonTooltipText() const;
 
+	void ShowReassignModuleScriptMenu();
+
 private:
 	UNiagaraStackModuleItem* ModuleItem;
 };
\ No newline at end of file
diff --git a/Engine/Plugins/FX/Niagara/Source/NiagaraShader/Private/NiagaraShaderCompilationManager.cpp b/Engine/Plugins/FX/Niagara/Source/NiagaraShader/Private/NiagaraShaderCompilationManager.cpp
--- a/Engine/Plugins/FX/Niagara/Source/NiagaraShader/Private/NiagaraShaderCompilationManager.cpp
+++ b/Engine/Plugins/FX/Niagara/Source/NiagaraShader/Private/NiagaraShaderCompilationManager.cpp
@@ -13,7 +13,7 @@
 
 DEFINE_LOG_CATEGORY_STATIC(LogNiagaraShaderCompiler, All, All);
 
-static int32 GShowNiagaraShaderWarnings = 1;
+static int32 GShowNiagaraShaderWarnings = 0;
 static FAutoConsoleVariableRef CVarShowNiagaraShaderWarnings(
 	TEXT("niagara.ShowShaderCompilerWarnings"),
 	GShowNiagaraShaderWarnings,
@@ -104,6 +104,12 @@ void FNiagaraShaderCompilationManager::RunCompileJobs()
 					UE_LOG(LogNiagaraShaderCompiler, Fatal, TEXT("Can't compile shaders for format %s, couldn't load compiler dll"), *Format.ToString());
 				}
 				CA_ASSUME(Compiler != NULL);
+				
+				// Fast math breaks The ExecGrid layout script because floor(x/y) returns a bad value if x == y. Yay.
+				if(IsMetalPlatform((EShaderPlatform)CurrentJob.Input.Target.Platform))
+				{
+					CurrentJob.Input.Environment.CompilerFlags.Add(CFLAG_NoFastMath);
+				}
 
 				UE_LOG(LogNiagaraShaderCompiler, Log, TEXT("Compile Job processing... %s"), *CurrentJob.Input.DebugGroupName);
 
diff --git a/Engine/Plugins/FX/Niagara/Source/NiagaraShader/Private/NiagaraShaderDerivedDataVersion.h b/Engine/Plugins/FX/Niagara/Source/NiagaraShader/Private/NiagaraShaderDerivedDataVersion.h
--- a/Engine/Plugins/FX/Niagara/Source/NiagaraShader/Private/NiagaraShaderDerivedDataVersion.h
+++ b/Engine/Plugins/FX/Niagara/Source/NiagaraShader/Private/NiagaraShaderDerivedDataVersion.h
@@ -11,5 +11,5 @@ NiagaraShaderDerivedDataVersion.h: Shader derived data version for Niagara.
 // In case of merge conflicts with DDC versions, you MUST generate a new GUID and set this new
 // guid as version
 
-#define NIAGARASHADERMAP_DERIVEDDATA_VER		TEXT("9B98AE21EB7347A99AB0862D8B1EB8A6}")
+#define NIAGARASHADERMAP_DERIVEDDATA_VER		TEXT("183CB3659FE7414CB5DC9CB93D9A13D9}")
 													  
diff --git a/Engine/Plugins/FX/Niagara/Source/NiagaraShader/Private/NiagaraShared.cpp b/Engine/Plugins/FX/Niagara/Source/NiagaraShader/Private/NiagaraShared.cpp
--- a/Engine/Plugins/FX/Niagara/Source/NiagaraShader/Private/NiagaraShared.cpp
+++ b/Engine/Plugins/FX/Niagara/Source/NiagaraShader/Private/NiagaraShared.cpp
@@ -359,7 +359,10 @@ void FNiagaraShaderScript::FinishCompilation()
 		// Shouldn't have anything left to do...
 		TArray<int32> ShaderMapIdsToFinish2;
 		GetShaderMapIDsWithUnfinishedCompilation(ShaderMapIdsToFinish2);
-		ensure(ShaderMapIdsToFinish2.Num() == 0);
+		if (ShaderMapIdsToFinish2.Num() != 0)
+		{
+			UE_LOG(LogShaders, Warning, TEXT("Skipped multiple Niagara shader maps for compilation! May be indicative of no support for a given platform. Count: %d"), ShaderMapIdsToFinish2.Num());
+		}
 	}
 }
 
diff --git a/Engine/Plugins/FX/Niagara/Source/NiagaraVertexFactories/Private/NiagaraRibbonVertexFactory.cpp b/Engine/Plugins/FX/Niagara/Source/NiagaraVertexFactories/Private/NiagaraRibbonVertexFactory.cpp
--- a/Engine/Plugins/FX/Niagara/Source/NiagaraVertexFactories/Private/NiagaraRibbonVertexFactory.cpp
+++ b/Engine/Plugins/FX/Niagara/Source/NiagaraVertexFactories/Private/NiagaraRibbonVertexFactory.cpp
@@ -38,6 +38,7 @@ public:
 		SortedIndices.Bind(ParameterMap, TEXT("SortedIndices"));
 		SortedIndicesOffset.Bind(ParameterMap, TEXT("SortedIndicesOffset"));
 		SegmentDistances.Bind(ParameterMap, TEXT("SegmentDistances"));
+		MultiRibbonIndices.Bind(ParameterMap, TEXT("MultiRibbonIndices"));
 		PackedPerRibbonDataByIndex.Bind(ParameterMap, TEXT("PackedPerRibbonDataByIndex"));
 
 		ensure(NiagaraParticleDataFloat.IsBound());
@@ -55,6 +56,7 @@ public:
 		Ar << SortedIndices;
 		Ar << SortedIndicesOffset;
 		Ar << SegmentDistances;
+		Ar << MultiRibbonIndices;
 		Ar << PackedPerRibbonDataByIndex;
 	}
 
@@ -70,6 +72,7 @@ public:
 
 		SetSRVParameter(RHICmdList, VertexShaderRHI, SortedIndices, RibbonVF->GetSortedIndicesSRV());
 		SetSRVParameter(RHICmdList, VertexShaderRHI, SegmentDistances, RibbonVF->GetSegmentDistancesSRV());
+		SetSRVParameter(RHICmdList, VertexShaderRHI, MultiRibbonIndices, RibbonVF->GetMultiRibbonIndicesSRV());
 		SetSRVParameter(RHICmdList, VertexShaderRHI, PackedPerRibbonDataByIndex, RibbonVF->GetPackedPerRibbonDataByIndexSRV());
 		SetShaderValue(RHICmdList, VertexShaderRHI, SortedIndicesOffset, RibbonVF->GetSortedIndicesOffset());
 	}
@@ -81,6 +84,7 @@ private:
 
 	FShaderResourceParameter SortedIndices;
 	FShaderResourceParameter SegmentDistances;
+	FShaderResourceParameter MultiRibbonIndices;
 	FShaderResourceParameter PackedPerRibbonDataByIndex;
 	FShaderParameter SortedIndicesOffset;
 };
@@ -123,10 +127,6 @@ public:
 
 	virtual void FillDeclElements(FVertexDeclarationElementList& Elements, int32& Offset)
 	{
-		uint32 Stride = sizeof(FNiagaraRibbonVertex);
-		/** The stream to read the ribbon index from. */
-		Elements.Add(FVertexElement(0, Offset, VET_UByte4, 0, Stride));
-		Offset += sizeof(int32);
 	}
 
 	virtual void InitDynamicRHI()
diff --git a/Engine/Plugins/FX/Niagara/Source/NiagaraVertexFactories/Public/NiagaraRibbonVertexFactory.h b/Engine/Plugins/FX/Niagara/Source/NiagaraVertexFactories/Public/NiagaraRibbonVertexFactory.h
--- a/Engine/Plugins/FX/Niagara/Source/NiagaraVertexFactories/Public/NiagaraRibbonVertexFactory.h
+++ b/Engine/Plugins/FX/Niagara/Source/NiagaraVertexFactories/Public/NiagaraRibbonVertexFactory.h
@@ -15,16 +15,6 @@
 #
 class FMaterial;
 
-
-
-//	FNiagaraRibbonVertex
-struct FNiagaraRibbonVertex
-{
-	/** The index of the ribbon for multi-ribbon emitters */
-	uint32 RibbonIndex;
-};
-
-
 //	FNiagaraRibbonVertexDynamicParameter
 struct FNiagaraRibbonVertexDynamicParameter
 {
@@ -147,6 +137,12 @@ public:
 		SegmentDistancesSRV = InSegmentDistancesSRV;
 	}
 
+
+	void SetMultiRibbonIndicesSRV(const FShaderResourceViewRHIRef& InMultiRibbonIndicesSRV)
+	{
+		MultiRibbonIndicesSRV = InMultiRibbonIndicesSRV;
+	}
+
 	void SetPackedPerRibbonDataByIndexSRV(const FShaderResourceViewRHIRef& InPackedPerRibbonDataByIndexSRV)
 	{
 		PackedPerRibbonDataByIndexSRV = InPackedPerRibbonDataByIndexSRV;
@@ -182,6 +178,11 @@ public:
 		return SegmentDistancesSRV;
 	}
 
+	FORCEINLINE FShaderResourceViewRHIRef GetMultiRibbonIndicesSRV()
+	{
+		return MultiRibbonIndicesSRV;
+	}
+
 	FORCEINLINE FShaderResourceViewRHIRef GetPackedPerRibbonDataByIndexSRV()
 	{
 		return PackedPerRibbonDataByIndexSRV;
@@ -225,6 +226,7 @@ private:
 
 	FShaderResourceViewRHIRef SortedIndicesSRV;
 	FShaderResourceViewRHIRef SegmentDistancesSRV;
+	FShaderResourceViewRHIRef MultiRibbonIndicesSRV;
 	FShaderResourceViewRHIRef PackedPerRibbonDataByIndexSRV;
 
 	uint32 SortedIndicesOffset;
diff --git a/Engine/Shaders/Private/NiagaraEmitterInstanceShader.usf b/Engine/Shaders/Private/NiagaraEmitterInstanceShader.usf
--- a/Engine/Shaders/Private/NiagaraEmitterInstanceShader.usf
+++ b/Engine/Shaders/Private/NiagaraEmitterInstanceShader.usf
@@ -8,6 +8,7 @@ NiagaraSimulationShader.usf:
 #if GPU_SIMULATION
 	#include "/Engine/Public/Platform.ush"
 	#include "Common.ush"
+	#include "GlobalDistanceFieldShared.ush"
 #else
 	const static float PI = 3.1415926535897932f;
 #endif
@@ -15,6 +16,20 @@ NiagaraSimulationShader.usf:
 
 #include "Definitions.usf"
 
+// Most of the vector implementations work this way. This helps us keep proper precision.
+float4 ModuloPrecise(float4 x, float4 y){ return x - y * trunc(x/y); }
+float3 ModuloPrecise(float3 x, float3 y){ return x - y * trunc(x/y); }
+float2 ModuloPrecise(float2 x, float2 y){ return x - y * trunc(x/y);}
+float ModuloPrecise(float x, float y){ return x - y * trunc(x/y); }
+int ModuloPrecise(int x, int y){ return x - y * (x/y); }
+int Modulo(int x, int y){  return x - y * (x/y);  }
+
+	
+// using rcp is only 12 bits of precision, we should usually pay for perf
+float4 Reciprocal(float4 x){ return 1.0f/x;}
+float3 Reciprocal(float3 x){ return 1.0f/x; }
+float2 Reciprocal(float2 x){ return 1.0f/x;}
+float Reciprocal(float x){ return 1.0f/x; }
 
 /* -----------------------------------------------------------------
  * GPU simulation utility functions
@@ -35,77 +50,169 @@ NiagaraSimulationShader.usf:
 	float2 Modulo(float2 x, float2 y){ return fmod(x,y); }
 	float Modulo(float x, float y){ return fmod(x,y); }
 
-
-	// 4D random number generator inspired by PCGs (permuted congruential generator)
-	// COPY FROM RANDOM.USH, which we can't easily include here due to a sizable dependency chain
-	uint4 NiagaraRand4DPCG32(int4 p)
+	// utility function used for scene depth calculations
+	float3 WorldPositionFromSceneDepth(float2 ScreenPosition, float SceneDepth)
 	{
-		// taking a signed int then reinterpreting as unsigned gives good behavior for negatives
-		uint4 v = uint4(p);
-
-		// Linear congruential step.
-		v = v * 1664525u + 1013904223u;
-
-		// swapping low and high bits makes all 32 bits pretty good
-		v = v * (1u << 16u) + (v >> 16u);
-
-		// final shuffle
-		v.x += v.y*v.w;
-		v.y += v.z*v.x;
-		v.z += v.x*v.y;
-		v.w += v.y*v.z;
-		v.x += v.y*v.w;
-		v.y += v.z*v.x;
-		v.z += v.x*v.y;
-		v.w += v.y*v.z;
-
-		return v;
+		float4 HomogeneousWorldPosition = mul(float4(ScreenPosition * SceneDepth, SceneDepth, 1), View.ScreenToWorld);
+		return HomogeneousWorldPosition.xyz / HomogeneousWorldPosition.w;
 	}
+#endif
 
+/* ----------------------------------------------------------------------------
+ * Seeded/Deterministic random number generation functions
+ *
+ * This is a variant of NiagaraRand4DPCG32 from Random.ush. 
+ *
+ * uint is not fully supported in the VM so we simply use ints and drop the 
+ * top and bottom bit swap. This should be fine since signed overflow should 
+ * produce the same results as unsigned overflow when comparing bit-by-bit on 
+ * all relevant architectures.
+ *
+ * Warning: Only contains 24 bits of randomness, since we produce values in 
+ *          the unit interval. 
+ *
+ * Note: Deterministic as long as Seed1, Seed2 and Seed3 are the same *and* as 
+ *       long as no Random or Seeded Random ops are inserted into or removed 
+ *       from the current script.
+ *
+ * TODO(mv): Might want to use a special case based on NiagaraRand3DPCG32 if 
+ *           only 3 or less components are needed, and we can drop one of the 
+ *           seeds, to reduce the operation count
+ * ----------------------------------------------------------------------------
+ */
 
+// static counter that is incremented every time rand3 or rand4 is called. 
+static int RandomCounter = 0;
 
-	// using Marc's PCG random implementation, sequentially; initially seed with DispatchThreadID and tick counter,
-	// then repeat seeding with the result from the PRNG; should be deterministic
-	// TODO: use this in the VM for CPU sims as well
-	//
-	static uint4 RandomVal = 2308975;
-	float4 rand4()
-	{
-		RandomVal = NiagaraRand4DPCG32( int4(GDispatchThreadId+(RandomVal/573).xyz, RandomVal.x/921) + int4(EmitterTickCounter.xxxx));
-		return float4(RandomVal.xyzw) / float4(0xFFFFFFFFu,0xFFFFFFFFu,0xFFFFFFFFu,0xFFFFFFFFu);
-	}
+// Returns 4 random normalized floats based on 3 explicit integer seeds and one 
+// implicit integer seed. 
+float4 rand4(int Seed1, int Seed2, int Seed3)
+{
+	RandomCounter += 1;
+
+	int4 v = int4(RandomCounter, Seed1, Seed2, Seed3) * 1664525 + 1013904223;
+
+	v.x += v.y*v.w;
+	v.y += v.z*v.x;
+	v.z += v.x*v.y;
+	v.w += v.y*v.z;
+	v.x += v.y*v.w;
+	v.y += v.z*v.x;
+	v.z += v.x*v.y;
+	v.w += v.y*v.z;
+
+	// NOTE(mv): We can use 24 bits of randomness, as all integers in [0, 2^24] 
+    //           are exactly representable in single precision floats.
+    //           We use the upper 24 bits as they tend to be higher quality.
+
+    // NOTE(mv): The divide can often be folded with the range scale in the rand functions
+
+    return float4((v >> 8) & 0x00ffffff) / 16777216.0; // 0x01000000 == 16777216
+    // return float4((v >> 8) & 0x00ffffff) * (1.0/16777216.0); // bugged, see UE-67738
+}
+
+// float3 specialization of the above:
+// Returns 3 random normalized floats based on 3 explicit integer seeds and one 
+// implicit integer seed. The implicit seed takes the upper 16 bits of the third 
+// explicit seed, as this is often a global seed. 
+// The expression 
+// 
+//    RandomCounter | (Seed3 << 16)
+//
+// should be optimized away at compile time, as the third seed is typically 
+// a global constant.
+float3 rand3(int Seed1, int Seed2, int Seed3)
+{
+    RandomCounter += 1;
+
+    int3 v = int3(Seed1, Seed2, RandomCounter | (Seed3 << 16)) * 1664525 + 1013904223;
+
+    v.x += v.y*v.z;
+    v.y += v.z*v.x;
+    v.z += v.x*v.y;
+    v.x += v.y*v.z;
+    v.y += v.z*v.x;
+    v.z += v.x*v.y;
+
+    return float3((v >> 8) & 0x00ffffff) / 16777216.0; // 0x01000000 == 16777216
+}
+
+// Cost using rand4: 6 imad, 1 itof, 1 ishr, 1 add, 2 mul
+float rand(float x, int Seed1, int Seed2, int Seed3)
+{
+	return rand3(Seed1, Seed2, Seed3).x * x;
+}
+
+// Cost using rand4: 7 imad, 1 itof, 1 ishr, 1 add, 2 mul
+float2 rand(float2 x, int Seed1, int Seed2, int Seed3)
+{
+	return rand3(Seed1, Seed2, Seed3).xy * x;
+}
+
+// Cost using rand4: 8 imad, 1 itof, 1 ishr, 1 add, 2 mul
+float3 rand(float3 x, int Seed1, int Seed2, int Seed3)
+{
+	return rand3(Seed1, Seed2, Seed3).xyz * x;
+}
+
+// Cost using rand4: 9 imad, 1 itof, 1 ishr, 1 and, 2 mul
+float4 rand(float4 x, int Seed1, int Seed2, int Seed3) 
+{
+    return rand4(Seed1, Seed2, Seed3).xyzw * x;
+}
+
+// Cost using rand4: 6 imad, 2 itof, 1 ishr, 1 add, 2 mul, 1 ftoi
+int rand(int x, int Seed1, int Seed2, int Seed3)
+{
+	// NOTE: Scaling a uniform float range provides better distribution of 
+	//       numbers than using %.
+	// NOTE: Inclusive! So [0, x] instead of [0, x)
+	return int(rand3(Seed1, Seed2, Seed3).x * (x+1));
+}
+
+/* -----------------------------------------------------------------
+ * Un-seeded/Non-deterministic random number generation functions
+ * -----------------------------------------------------------------
+ */
+#if GPU_SIMULATION
+	// NOTE(mv): This simply calls the deterministic random number functions 
+    //           from the Seeded RNG section, but uses non-deterministic seeds 
+    //           as input. In particular, GDispatchThreadId for the first seed 
+	//           and EmitterTickCounter for the second seed. Also note that 
+	//           EmitterTickCounter does *not* reset for each frame, which 
+	//           secures non-determinism across simulations.
+
+	// TODO(mv): This could perhaps be optimized a little bit using slightly 
+	//           cheaper functions, but the difference is likely negligible. 
+	//           So for the sake of brevity and cleanliness, keep it like this.
 
 	float rand(float x)
 	{
-		return rand4().x * x;
+		return rand4(GDispatchThreadId.x, EmitterTickCounter, 0).x * x;
 	}
 
 	float2 rand(float2 x)
 	{
-		float4 rnd = rand4();
-		return float2(rnd.x, rnd.y) * x;
+		return rand4(GDispatchThreadId.x, EmitterTickCounter, 0).xy * x;
 	}
 
 	float3 rand(float3 x)
 	{
-		float4 rnd = rand4();
-		return float3(rnd.x, rnd.y, rnd.z) * x;
+		return rand4(GDispatchThreadId.x, EmitterTickCounter, 0).xyz * x;
 	}
 
 	float4 rand(float4 x) 
 	{
-		return rand4() * x;
+		return rand4(GDispatchThreadId.x, EmitterTickCounter, 0).xyzw * x;
 	}
 
+	// Note: integer randoms are INCLUSIVE, i.e. includes both the upper and lower limit
 	int rand(int x)
 	{
-		RandomVal = NiagaraRand4DPCG32(int4(GDispatchThreadId + (RandomVal / 57843).xyz, 54121) + int4(EmitterTickCounter.xxxx));
-		return RandomVal.x % x+1;
+		return int(rand4(GDispatchThreadId.x, EmitterTickCounter, 0).x * (x+1));
 	}
-
 #else
-
-	//Temporary random hacks
+	// Old unseeded, passthrough to FRandomStream
 	float2 rand(float2 x)
 	{
 		return float2(rand(x.x), rand(x.y));
@@ -120,7 +227,8 @@ NiagaraSimulationShader.usf:
 	{
 		return float4(rand(x.x), rand(x.y), rand(x.z), rand(x.w));
 	}
-
+	
+	// NOTE(mv): Where's rand(float)? Seems to work regardless..
 	int rand(int x);
 #endif
 
diff --git a/Engine/Shaders/Private/NiagaraMeshVertexFactory.ush b/Engine/Shaders/Private/NiagaraMeshVertexFactory.ush
--- a/Engine/Shaders/Private/NiagaraMeshVertexFactory.ush
+++ b/Engine/Shaders/Private/NiagaraMeshVertexFactory.ush
@@ -61,12 +61,12 @@ struct FVertexFactoryInterpolantsVSToPS
 	float4 SubUV0AndTexCoord0 : TEXCOORD1;
 	float4 SubUV1AndLerp : TEXCOORD2;
 #else
-	#if NUM_MATERIAL_TEXCOORDS
+	#if  NUM_TEX_COORD_INTERPOLATORS
 			#if (ES2_PROFILE || ES3_1_PROFILE)
 				// Avoid dependent texture fetches, put all UVs in xy
-				float2	TexCoords[NUM_MATERIAL_TEXCOORDS] : TEXCOORD0;
+				float2	TexCoords[ NUM_TEX_COORD_INTERPOLATORS] : TEXCOORD0;
 			#else
-				float4	TexCoords[(NUM_MATERIAL_TEXCOORDS+1)/2]	: TEXCOORD0;
+				float4	TexCoords[( NUM_TEX_COORD_INTERPOLATORS+1)/2]	: TEXCOORD0;
 			#endif
 	#endif
 #endif
@@ -202,9 +202,9 @@ FMaterialPixelParameters GetMaterialPixelParameters(FVertexFactoryInterpolantsVS
 	FMaterialPixelParameters Result = MakeInitializedMaterialPixelParameters();
 
 #if USE_PARTICLE_SUBUVS
-	#if NUM_MATERIAL_TEXCOORDS
+	#if  NUM_TEX_COORD_INTERPOLATORS
 		UNROLL
-		for( int CoordinateIndex = 0; CoordinateIndex < NUM_MATERIAL_TEXCOORDS; CoordinateIndex++ )
+		for( int CoordinateIndex = 0; CoordinateIndex <  NUM_TEX_COORD_INTERPOLATORS; CoordinateIndex++ )
 		{
 			Result.TexCoords[CoordinateIndex] = Interpolants.SubUV0AndTexCoord0.zw;
 		}
@@ -212,20 +212,20 @@ FMaterialPixelParameters GetMaterialPixelParameters(FVertexFactoryInterpolantsVS
 	Result.Particle.SubUVCoords[0] = Interpolants.SubUV0AndTexCoord0.xy;
 	Result.Particle.SubUVCoords[1] = Interpolants.SubUV1AndLerp.xy;
 	Result.Particle.SubUVLerp = Interpolants.SubUV1AndLerp.z;
-#elif NUM_MATERIAL_TEXCOORDS
+#elif  NUM_TEX_COORD_INTERPOLATORS
 	#if (ES2_PROFILE || ES3_1_PROFILE)
 		UNROLL
-		for( int CoordinateIndex = 0; CoordinateIndex < NUM_MATERIAL_TEXCOORDS; CoordinateIndex++ )
+		for( int CoordinateIndex = 0; CoordinateIndex <  NUM_TEX_COORD_INTERPOLATORS; CoordinateIndex++ )
 		{
 			Result.TexCoords[CoordinateIndex] = Interpolants.TexCoords[CoordinateIndex].xy;
 		}
 	#else
 		UNROLL
-		for(int CoordinateIndex = 0;CoordinateIndex < NUM_MATERIAL_TEXCOORDS;CoordinateIndex += 2)
+		for(int CoordinateIndex = 0;CoordinateIndex <  NUM_TEX_COORD_INTERPOLATORS;CoordinateIndex += 2)
 		{
 			Result.TexCoords[CoordinateIndex] = Interpolants.TexCoords[CoordinateIndex/2].xy;
 
-			if(CoordinateIndex + 1 < NUM_MATERIAL_TEXCOORDS)
+			if(CoordinateIndex + 1 <  NUM_TEX_COORD_INTERPOLATORS)
 			{
 				Result.TexCoords[CoordinateIndex + 1] = Interpolants.TexCoords[CoordinateIndex/2].wz;
 			}
@@ -719,28 +719,28 @@ FVertexFactoryInterpolantsVSToPS VertexFactoryGetInterpolantsVSToPS(FVertexFacto
 	Interpolants.SubUV1AndLerp.xy = Intermediates.SubUVCoords.zw;
 	Interpolants.SubUV1AndLerp.zw = Intermediates.SubUVLerp.xx;
 
-	#if NUM_MATERIAL_TEXCOORDS
-		float2 CustomizedUVs[NUM_MATERIAL_TEXCOORDS];
+	#if  NUM_TEX_COORD_INTERPOLATORS
+		float2 CustomizedUVs[ NUM_TEX_COORD_INTERPOLATORS];
 		GetMaterialCustomizedUVs(VertexParameters, CustomizedUVs);
 		Interpolants.SubUV0AndTexCoord0.zw = CustomizedUVs[0];
 	#else
 		Interpolants.SubUV0AndTexCoord0.zw = 0;
 	#endif
-#elif NUM_MATERIAL_TEXCOORDS
+#elif  NUM_TEX_COORD_INTERPOLATORS
 	#if (ES2_PROFILE || ES3_1_PROFILE)
 		GetMaterialCustomizedUVs(VertexParameters, Interpolants.TexCoords);
 	#else
 		// Ensure the unused components of the last packed texture coordinate are initialized.
-		Interpolants.TexCoords[(NUM_MATERIAL_TEXCOORDS + 1) / 2 - 1] = 0;
+		Interpolants.TexCoords[( NUM_TEX_COORD_INTERPOLATORS + 1) / 2 - 1] = 0;
 
-		float2 CustomizedUVs[NUM_MATERIAL_TEXCOORDS];
+		float2 CustomizedUVs[ NUM_TEX_COORD_INTERPOLATORS];
 		GetMaterialCustomizedUVs(VertexParameters, CustomizedUVs);
 
 		UNROLL
-		for(int CoordinateIndex = 0;CoordinateIndex < NUM_MATERIAL_TEXCOORDS;CoordinateIndex += 2)
+		for(int CoordinateIndex = 0;CoordinateIndex <  NUM_TEX_COORD_INTERPOLATORS;CoordinateIndex += 2)
 		{
 				Interpolants.TexCoords[CoordinateIndex / 2].xy = CustomizedUVs[CoordinateIndex];
-			if(CoordinateIndex + 1 < NUM_MATERIAL_TEXCOORDS)
+			if(CoordinateIndex + 1 <  NUM_TEX_COORD_INTERPOLATORS)
 			{
 					Interpolants.TexCoords[CoordinateIndex / 2].wz = CustomizedUVs[CoordinateIndex + 1];
 			}
@@ -834,12 +834,12 @@ float4 VertexFactoryGetPreviousWorldPosition(FVertexFactoryInput Input, FVertexF
 	FMaterialTessellationParameters GetMaterialTessellationParameters(FVertexFactoryInterpolantsVSToDS Interpolants, float3 CameraLocalWorldPosition)
 	{
 		FMaterialTessellationParameters	Result;
-	#if NUM_MATERIAL_TEXCOORDS
+	#if  NUM_TEX_COORD_INTERPOLATORS
 		UNROLL
-		for(int CoordinateIndex = 0;CoordinateIndex < NUM_MATERIAL_TEXCOORDS;CoordinateIndex += 2)
+		for(int CoordinateIndex = 0;CoordinateIndex <  NUM_TEX_COORD_INTERPOLATORS;CoordinateIndex += 2)
 		{
 			Result.TexCoords[CoordinateIndex] = Interpolants.InterpolantsVSToPS.TexCoords[CoordinateIndex/2].xy;
-			if(CoordinateIndex + 1 < NUM_MATERIAL_TEXCOORDS)
+			if(CoordinateIndex + 1 <  NUM_TEX_COORD_INTERPOLATORS)
 			{
 				Result.TexCoords[CoordinateIndex + 1] = Interpolants.InterpolantsVSToPS.TexCoords[CoordinateIndex/2].wz;
 			}
@@ -881,8 +881,8 @@ float4 VertexFactoryGetPreviousWorldPosition(FVertexFactoryInput Input, FVertexF
 		TESSELLATION_INTERPOLATE_MEMBER(InterpolantsVSToPS.ParticleColor);
 #endif
 
-#if NUM_MATERIAL_TEXCOORDS
-		for (int i = 0; i < (NUM_MATERIAL_TEXCOORDS + 1) / 2; ++i)
+#if  NUM_TEX_COORD_INTERPOLATORS
+		for (int i = 0; i < ( NUM_TEX_COORD_INTERPOLATORS + 1) / 2; ++i)
 		{
 			TESSELLATION_INTERPOLATE_MEMBER(InterpolantsVSToPS.TexCoords[i]);
 		}
diff --git a/Engine/Shaders/Private/NiagaraRibbonVertexFactory.ush b/Engine/Shaders/Private/NiagaraRibbonVertexFactory.ush
--- a/Engine/Shaders/Private/NiagaraRibbonVertexFactory.ush
+++ b/Engine/Shaders/Private/NiagaraRibbonVertexFactory.ush
@@ -149,14 +149,13 @@ Buffer<int> SortedIndices;
 uint SortedIndicesOffset;
 
 Buffer<float> SegmentDistances;
+Buffer<uint> MultiRibbonIndices;
 Buffer<float> PackedPerRibbonDataByIndex;
 
 struct FVertexFactoryInput
 {
 	uint VertexId					: SV_VertexID;
 
-	uint RibbonIndex				: ATTRIBUTE0;
-
 	// Optional instance ID for vertex layered rendering
 #if FEATURE_LEVEL >= FEATURE_LEVEL_SM4 && ONEPASS_POINTLIGHT_SHADOW && USING_VERTEX_SHADER_LAYER
 	uint InstanceId	: SV_InstanceID;
@@ -502,13 +501,10 @@ void GetTangents(FVertexFactoryInput Input, FVertexFactoryIntermediates Intermed
 	float cosa = cos(Intermediates.Twist);
 	Facing = Facing * cosa + cross(Facing, NormDir) * sin(Intermediates.Twist) + NormDir * dot(NormDir, Facing) * (1-cosa);
 
-	if( abs(dot(NormDir, Facing)) < 0.999f)
+	RightDir = SafeNormalize(cross(NormDir, Facing));
+	if(!any(RightDir))
 	{
-		RightDir = SafeNormalize(cross(NormDir, Facing));
-	}
-	else
-	{
-		RightDir = SafeNormalize(cross(NormDir, float3(0.0f, 0.0f, -1.0f)));
+		RightDir = SafeNormalize(cross(NormDir, float3(0.0f, 0.0f, 1.0f)));
 	}
 
 	Right = RightDir;
@@ -553,7 +549,9 @@ FVertexFactoryIntermediates GetVertexFactoryIntermediates(FVertexFactoryInput In
 	float2 UV1Offset;
 	uint NumSegments;
 	uint StartVertexId;
-	UnpackPerRibbonDataByIndex(Input.RibbonIndex, UV0Scale, UV0Offset, UV1Scale, UV1Offset, NumSegments,  StartVertexId);
+
+	uint RibbonIndex = MultiRibbonIndices[Input.VertexId / 2];
+	UnpackPerRibbonDataByIndex(RibbonIndex, UV0Scale, UV0Offset, UV1Scale, UV1Offset, NumSegments,  StartVertexId);
 	
 	uint StartParticleId = GetParticleIdByVertexId(StartVertexId);
 
diff --git a/Engine/Shaders/Private/NiagaraSpriteVertexFactory.ush b/Engine/Shaders/Private/NiagaraSpriteVertexFactory.ush
--- a/Engine/Shaders/Private/NiagaraSpriteVertexFactory.ush
+++ b/Engine/Shaders/Private/NiagaraSpriteVertexFactory.ush
@@ -75,11 +75,11 @@ struct FVertexFactoryInterpolantsVSToPS
 #endif
 
 #if NEEDS_PARTICLE_COLOR
-	float4	Color				: TEXCOORD2;
+	float4	Color				: TEXCOORD0;
 #endif
 
-#if NUM_MATERIAL_TEXCOORDS
-	float4	TexCoords[(NUM_MATERIAL_TEXCOORDS + 1) / 2]	: TEXCOORD3;
+#if NUM_TEX_COORD_INTERPOLATORS
+	float4	TexCoords[(NUM_TEX_COORD_INTERPOLATORS + 1) / 2]	: TEXCOORD1;
 #endif
 
 	//Not sure this is actually being used now and it's awkward to slot in now we're supporting custom UVs so I'm just giving this its own interpolant.
@@ -189,7 +189,7 @@ struct FVertexFactoryIntermediates
 #endif
 };
 
-#if NUM_MATERIAL_TEXCOORDS
+#if NUM_TEX_COORD_INTERPOLATORS
 bool UVIndexUseZW(int UVIndex)
 {
 	#if COMPILER_GLSL_ES2
@@ -225,13 +225,13 @@ FMaterialPixelParameters GetMaterialPixelParameters(FVertexFactoryInterpolantsVS
 	// GetMaterialPixelParameters is responsible for fully initializing the result
 	FMaterialPixelParameters Result = MakeInitializedMaterialPixelParameters();
 
-#if NUM_MATERIAL_TEXCOORDS
+#if NUM_TEX_COORD_INTERPOLATORS
 	UNROLL
-	for (int CoordinateIndex = 0; CoordinateIndex < NUM_MATERIAL_TEXCOORDS; CoordinateIndex++)
+	for (int CoordinateIndex = 0; CoordinateIndex < NUM_TEX_COORD_INTERPOLATORS; CoordinateIndex++)
 	{
 		Result.TexCoords[CoordinateIndex] = GetUV(Interpolants, CoordinateIndex);
 	}
-#endif	// NUM_MATERIAL_TEXCOORDS
+#endif	// NUM_TEX_COORD_INTERPOLATORS
 
 	Result.VertexColor = 1;
 
@@ -906,13 +906,13 @@ FVertexFactoryInterpolantsVSToPS VertexFactoryGetInterpolantsVSToPS(FVertexFacto
 	// Initialize the whole struct to 0
 	Interpolants = (FVertexFactoryInterpolantsVSToPS)0;
 
-#if NUM_MATERIAL_TEXCOORDS
+#if NUM_TEX_COORD_INTERPOLATORS
 
-	float2 CustomizedUVs[NUM_MATERIAL_TEXCOORDS];
+	float2 CustomizedUVs[NUM_TEX_COORD_INTERPOLATORS];
 	GetMaterialCustomizedUVs(VertexParameters, CustomizedUVs);
 
 	UNROLL
-	for (int CoordinateIndex = 0; CoordinateIndex < NUM_MATERIAL_TEXCOORDS; CoordinateIndex++)
+	for (int CoordinateIndex = 0; CoordinateIndex < NUM_TEX_COORD_INTERPOLATORS; CoordinateIndex++)
 	{
 		SetUV(Interpolants, CoordinateIndex, CustomizedUVs[CoordinateIndex]);
 	}
@@ -1013,5 +1013,9 @@ uint VertexFactoryGetEyeIndex(uint InstanceId)
 
 float4 VertexFactoryGetTranslatedPrimitiveVolumeBounds(FVertexFactoryInterpolantsVSToPS Interpolants)
 {
+#if USE_PARTICLE_POSITION
+	return Interpolants.ParticleTranslatedWorldPositionAndSize;
+#else
 	return 0;
+#endif
 }
diff --git a/Engine/Source/Developer/ShaderFormatVectorVM/Private/ir_vm_gen_bytecode_visitor.cpp b/Engine/Source/Developer/ShaderFormatVectorVM/Private/ir_vm_gen_bytecode_visitor.cpp
--- a/Engine/Source/Developer/ShaderFormatVectorVM/Private/ir_vm_gen_bytecode_visitor.cpp
+++ b/Engine/Source/Developer/ShaderFormatVectorVM/Private/ir_vm_gen_bytecode_visitor.cpp
@@ -221,7 +221,8 @@ void BuildExpressionMap()
 		Info->Add(FVMExpresssionInfo(EVectorVMOp::muli, glsl_type::int_type, glsl_type::int_type, glsl_type::int_type));
 		Info = &VMExpressionMap.Add(ir_binop_div);
 		Info->Add(FVMExpresssionInfo(EVectorVMOp::div, glsl_type::float_type, glsl_type::float_type, glsl_type::float_type));
-		//Currently don't have an integer division operation.
+		Info->Add(FVMExpresssionInfo(EVectorVMOp::divi, glsl_type::int_type, glsl_type::int_type, glsl_type::int_type));
+	
 
 		/**
 		* Takes one of two combinations of arguments:
@@ -277,8 +278,10 @@ void BuildExpressionMap()
 		* \name Bit-wise binary operations.
 		*/
 		/*@{*/
-		//Info = &VMExpressionMap.Add(ir_binop_lshift);
-		//Info = &VMExpressionMap.Add(ir_binop_rshift);
+		Info = &VMExpressionMap.Add(ir_binop_lshift);
+		Info->Add(FVMExpresssionInfo(EVectorVMOp::bit_lshift, glsl_type::int_type, glsl_type::int_type, glsl_type::int_type));
+		Info = &VMExpressionMap.Add(ir_binop_rshift);
+		Info->Add(FVMExpresssionInfo(EVectorVMOp::bit_rshift, glsl_type::int_type, glsl_type::int_type, glsl_type::int_type));
 		Info = &VMExpressionMap.Add(ir_binop_bit_and);
 		Info->Add(FVMExpresssionInfo(EVectorVMOp::bit_and, glsl_type::int_type, glsl_type::int_type, glsl_type::int_type));
 		Info = &VMExpressionMap.Add(ir_binop_bit_xor);
@@ -2115,7 +2118,7 @@ class ir_gen_vvm_visitor : public ir_hierarchical_visitor
 			case EVectorVMBaseTypes::Float:
 			{
 				float Val = *(float*)(CompilationOutput.InternalConstantData.GetData() + Offset);
-				OpsConstantTable += FString::Printf(TEXT("%d | %f\n"), TableOffset, Val);
+				OpsConstantTable += FString::Printf(TEXT("%d | %.9g\n"), TableOffset, Val);
 				NumConstants++;
 			}
 			break;
diff --git a/Engine/Source/Editor/ContentBrowser/Private/ContentBrowserUtils.cpp b/Engine/Source/Editor/ContentBrowser/Private/ContentBrowserUtils.cpp
--- a/Engine/Source/Editor/ContentBrowser/Private/ContentBrowserUtils.cpp
+++ b/Engine/Source/Editor/ContentBrowser/Private/ContentBrowserUtils.cpp
@@ -1790,19 +1790,31 @@ int32 ContentBrowserUtils::GetPackageLengthForCooking(const FString& PackageName
 	// We use "WindowsNoEditor" below as it's the longest platform name, so will also prove that any shorter platform names will validate correctly
 	const FString AbsoluteRootPath = FPaths::ConvertRelativePathToFull(FPaths::RootDir());
 	const FString AbsoluteGamePath = FPaths::ConvertRelativePathToFull(FPaths::ProjectDir());
-	const FString AbsoluteCookPath = AbsoluteGamePath / TEXT("Saved") / TEXT("Cooked") / TEXT("WindowsNoEditor") / GameName;
+	const FString AbsoluteEnginePath = FPaths::ConvertRelativePathToFull(FPaths::EngineDir());
+	const FString AbsoluteEngineCookPath = AbsoluteGamePath / TEXT("Saved") / TEXT("Cooked") / TEXT("WindowsNoEditor") / TEXT("Engine");
+	const FString AbsoluteGameCookPath = AbsoluteGamePath / TEXT("Saved") / TEXT("Cooked") / TEXT("WindowsNoEditor") / GameName;
+
+	EPluginLoadedFrom PluginLoadedFrom;
+	const bool bIsPluginAsset = ContentBrowserUtils::IsPluginFolder(PackageName, &PluginLoadedFrom);
+	const bool bIsEngineAsset = ContentBrowserUtils::IsEngineFolder(PackageName) || (bIsPluginAsset && PluginLoadedFrom == EPluginLoadedFrom::Engine);
+	const bool bIsProjectAsset = !bIsEngineAsset;
 
 	int32 AbsoluteCookPathToAssetLength = 0;
 
 	FString RelativePathToAsset;
 
+	const FString AbsolutePath = bIsEngineAsset ? AbsoluteEnginePath : AbsoluteGamePath;
+
+	const FString& AbsoluteCookPath = bIsEngineAsset ? AbsoluteEngineCookPath : AbsoluteGameCookPath;
+
 	if(FPackageName::TryConvertLongPackageNameToFilename(PackageName, RelativePathToAsset, FPackageName::GetAssetPackageExtension()))
 	{
 		const FString AbsolutePathToAsset = FPaths::ConvertRelativePathToFull(RelativePathToAsset);
 
 		FString AssetPathWithinCookDir = AbsolutePathToAsset;
 		FPaths::RemoveDuplicateSlashes(AssetPathWithinCookDir);
-		AssetPathWithinCookDir.RemoveFromStart(AbsoluteGamePath, ESearchCase::CaseSensitive);
+		AssetPathWithinCookDir.RemoveFromStart(AbsolutePath, ESearchCase::CaseSensitive);
+
 
 		if (IsInternalBuild)
 		{
@@ -1815,11 +1827,16 @@ int32 ContentBrowserUtils::GetPackageLengthForCooking(const FString& PackageName
 			}
 			else
 			{
-				CookDirWithoutBasePath.RemoveFromStart(AbsoluteGamePath, ESearchCase::CaseSensitive);
+				CookDirWithoutBasePath.RemoveFromStart(AbsoluteCookPath, ESearchCase::CaseSensitive);
 			}
-
+			
 			FString AbsoluteBuildMachineCookPathToAsset = FString(TEXT("D:/BuildFarm/buildmachine_++depot+UE4-Releases+4.10")) / CookDirWithoutBasePath / AssetPathWithinCookDir;
-			AbsoluteBuildMachineCookPathToAsset.ReplaceInline(*GameName, *GameNamePadded, ESearchCase::CaseSensitive);
+
+			// only add game name padding if it is not an engine asset, otherwise it is considered portable already
+			if(!bIsEngineAsset)
+			{
+				AbsoluteBuildMachineCookPathToAsset.ReplaceInline(*GameName, *GameNamePadded, ESearchCase::CaseSensitive);
+			}
 
 			AbsoluteCookPathToAssetLength = AbsoluteBuildMachineCookPathToAsset.Len();
 		}
@@ -1827,7 +1844,12 @@ int32 ContentBrowserUtils::GetPackageLengthForCooking(const FString& PackageName
 		{
 			// Test that the package can be cooked based on the current project path
 			FString AbsoluteCookPathToAsset = AbsoluteCookPath / AssetPathWithinCookDir;
-			AbsoluteCookPathToAsset.ReplaceInline(*GameName, *GameNamePadded, ESearchCase::CaseSensitive);
+
+			// only add game name padding if it is not an engine asset, otherwise it is considered portable already
+			if (!bIsEngineAsset)
+			{
+				AbsoluteCookPathToAsset.ReplaceInline(*GameName, *GameNamePadded, ESearchCase::CaseSensitive);
+			}
 
 			AbsoluteCookPathToAssetLength = AbsoluteCookPathToAsset.Len();
 		}
diff --git a/Engine/Source/Editor/DetailCustomizations/Private/MathStructProxyCustomizations.cpp b/Engine/Source/Editor/DetailCustomizations/Private/MathStructProxyCustomizations.cpp
--- a/Engine/Source/Editor/DetailCustomizations/Private/MathStructProxyCustomizations.cpp
+++ b/Engine/Source/Editor/DetailCustomizations/Private/MathStructProxyCustomizations.cpp
@@ -688,7 +688,11 @@ bool FQuatStructCustomization::FlushValues(TWeakPtr<IPropertyHandle> PropertyHan
 			
 			const FQuat NewValue = Rotation.Quaternion();
 
-			if (!bNotifiedPreChange && (!QuatValue->Equals(NewValue, 0.0f) || (!bIsUsingSlider && bIsInteractiveChangeInProgress)))
+			// In some cases the FQuat pointed to in RawData is no longer aligned to 16 bytes.
+			// Make a local copy to guarantee the alignment criterions of the vector intrinsics inside FQuat::Equals
+			const FQuat AlignedQuatValue = *QuatValue; 
+
+			if (!bNotifiedPreChange && (!AlignedQuatValue.Equals(NewValue, 0.0f) || (!bIsUsingSlider && bIsInteractiveChangeInProgress)))
 			{
 				if (!bIsInteractiveChangeInProgress)
 				{
diff --git a/Engine/Source/Runtime/Apple/MetalRHI/Private/MetalUAV.cpp b/Engine/Source/Runtime/Apple/MetalRHI/Private/MetalUAV.cpp
--- a/Engine/Source/Runtime/Apple/MetalRHI/Private/MetalUAV.cpp
+++ b/Engine/Source/Runtime/Apple/MetalRHI/Private/MetalUAV.cpp
@@ -851,32 +851,27 @@ void FMetalGPUFence::WriteInternal(mtlpp::CommandBuffer& CmdBuffer)
 void FMetalRHICommandContext::RHIEnqueueStagedRead(FStagingBufferRHIParamRef StagingBuffer, FGPUFenceRHIParamRef InFence, uint32 Offset, uint32 NumBytes)
 {
 	@autoreleasepool {
-	
 		check(StagingBuffer);
-
-		FMetalStagingBuffer* StageBuffer = ResourceCast(StagingBuffer);
-		FMetalVertexBuffer* VertexBuffer = ResourceCast(StageBuffer->GetBackingBuffer());
-		switch (VertexBuffer->Buffer.GetStorageMode())
+		
+		FMetalStagingBuffer* MetalStagingBuffer = ResourceCast(StagingBuffer);
+		FMetalVertexBuffer* BackingBuffer = ResourceCast(MetalStagingBuffer->GetBackingBuffer());
+		FMetalBuffer& ReadbackBuffer = MetalStagingBuffer->ReadbackStagingBuffer;
+		
+		// Need a shadow buffer for this read. If it hasn't been allocated in our FStagingBuffer or if
+		// it's not big enough to hold our readback we need to allocate.
+		if(!ReadbackBuffer || ReadbackBuffer.GetLength() < NumBytes)
 		{
-	#if PLATFORM_MAC
-			case mtlpp::StorageMode::Managed:
-			{
-				GetMetalDeviceContext().SynchroniseResource(VertexBuffer->Buffer);
-				break;
-			}
-	#endif
-			case mtlpp::StorageMode::Private:
+			if(ReadbackBuffer)
 			{
-				VertexBuffer->Alloc(VertexBuffer->Buffer.GetLength(), RLM_ReadOnly);
-				GetMetalDeviceContext().CopyFromBufferToBuffer(VertexBuffer->Buffer, Offset, VertexBuffer->CPUBuffer, Offset, NumBytes);
-				break;
-			}
-			default:
-			{
-				break;
+				SafeReleaseMetalBuffer(ReadbackBuffer);
 			}
+			FMetalPooledBufferArgs ArgsCPU(GetMetalDeviceContext().GetDevice(), NumBytes, mtlpp::StorageMode::Shared);
+			ReadbackBuffer = GetMetalDeviceContext().CreatePooledBuffer(ArgsCPU);
 		}
-
+		
+		// Inline copy from the actual buffer to the shadow
+		GetMetalDeviceContext().CopyFromBufferToBuffer(BackingBuffer->Buffer, Offset, ReadbackBuffer, 0, NumBytes);
+		
 		if (InFence)
 		{
 			FMetalGPUFence* Fence = ResourceCast(InFence);
diff --git a/Engine/Source/Runtime/Apple/MetalRHI/Private/MetalVertexBuffer.cpp b/Engine/Source/Runtime/Apple/MetalRHI/Private/MetalVertexBuffer.cpp
--- a/Engine/Source/Runtime/Apple/MetalRHI/Private/MetalVertexBuffer.cpp
+++ b/Engine/Source/Runtime/Apple/MetalRHI/Private/MetalVertexBuffer.cpp
@@ -774,36 +774,25 @@ FStagingBufferRHIRef FMetalDynamicRHI::RHICreateStagingBuffer(FVertexBufferRHIPa
 	return new FMetalStagingBuffer(VertexBuffer);
 }
 
-// Call this to lock the vertex-buffer for the given mode.
-// A read-only lock must have the same buffer used to call EnqueueStagedRead, and that fence must have passed or the behaviour is undefined.
-// A write-only lock must not have had the EnqueueStagedRead function called and must supply the buffer.
-void *FMetalStagingBuffer::Lock(uint32 Offset, uint32 NumBytes)
+FMetalStagingBuffer::~FMetalStagingBuffer()
 {
-	check(BackingBuffer);
-	FMetalVertexBuffer* VertexBuffer = ResourceCast(BackingBuffer.GetReference());
-	uint8* BytePtr = nullptr;
-	if (VertexBuffer->CPUBuffer)
-	{
-		BytePtr = (uint8*)VertexBuffer->CPUBuffer.GetContents();
-	}
-	else
+	if(ReadbackStagingBuffer)
 	{
-		check(VertexBuffer->Buffer.GetStorageMode() != mtlpp::StorageMode::Private);
-		BytePtr = (uint8*)VertexBuffer->Buffer.GetContents();
+		SafeReleaseMetalBuffer(ReadbackStagingBuffer);
 	}
-	BytePtr += Offset;
-	return BytePtr;
 }
 
-// Releases the mapped memory for a lock.
+// Returns the pointer to read the buffer. There is no locking; the buffer is always shared.
+// If this was not fenced correctly it will not have the expected data.
+void *FMetalStagingBuffer::Lock(uint32 Offset, uint32 NumBytes)
+{
+	check(ReadbackStagingBuffer);
+	
+	uint8* BackingPtr = (uint8*) ReadbackStagingBuffer.GetContents();
+	return BackingPtr+Offset;
+}
+
 void FMetalStagingBuffer::Unlock()
 {
-	check(BackingBuffer);
-	FMetalVertexBuffer* VertexBuffer = ResourceCast(BackingBuffer.GetReference());
-	if (VertexBuffer->CPUBuffer && (VertexBuffer->UsePrivateMemory()))
-	{
-		LLM_SCOPE(ELLMTag::VertexBuffer);
-		SafeReleaseMetalBuffer(VertexBuffer->CPUBuffer);
-		VertexBuffer->CPUBuffer = nil;
-	}
+	// does nothing in metal.
 }
diff --git a/Engine/Source/Runtime/Apple/MetalRHI/Public/MetalResources.h b/Engine/Source/Runtime/Apple/MetalRHI/Public/MetalResources.h
--- a/Engine/Source/Runtime/Apple/MetalRHI/Public/MetalResources.h
+++ b/Engine/Source/Runtime/Apple/MetalRHI/Public/MetalResources.h
@@ -1048,10 +1048,14 @@ public:
 	: FRHIStagingBuffer(InBuffer)
 	{
 	}
+	
+	~FMetalStagingBuffer();
 
 	void *Lock(uint32 Offset, uint32 NumBytes);
 
 	void Unlock();
+	
+	FMetalBuffer ReadbackStagingBuffer;
 };
 
 class FMetalShaderLibrary final : public FRHIShaderLibrary
diff --git a/Engine/Source/Runtime/Core/Public/Math/UnrealMathNeon.h b/Engine/Source/Runtime/Core/Public/Math/UnrealMathNeon.h
--- a/Engine/Source/Runtime/Core/Public/Math/UnrealMathNeon.h
+++ b/Engine/Source/Runtime/Core/Public/Math/UnrealMathNeon.h
@@ -1423,7 +1423,7 @@ FORCEINLINE VectorRegisterInt VectorIntSelect(const VectorRegisterInt& Mask, con
 * @param Vec	Vector to store
 * @param Ptr	Memory pointer
 */
-#define VectorIntStore( Vec, Ptr )			vst1q_s32( (VectorRegisterInt*)(Ptr), Vec )
+#define VectorIntStore( Vec, Ptr )			vst1q_s32( (int32*)(Ptr), Vec )
 
 /**
 * Loads 4 int32s from unaligned memory.
diff --git a/Engine/Source/Runtime/Engine/Classes/Engine/TextureRenderTarget.h b/Engine/Source/Runtime/Engine/Classes/Engine/TextureRenderTarget.h
--- a/Engine/Source/Runtime/Engine/Classes/Engine/TextureRenderTarget.h
+++ b/Engine/Source/Runtime/Engine/Classes/Engine/TextureRenderTarget.h
@@ -26,6 +26,9 @@ class UTextureRenderTarget : public UTexture
 	/** If true, there will be two copies in memory - one for the texture and one for the render target. If false, they will share memory if possible. This is useful for scene capture textures that are used in the scene. */
 	uint32 bNeedsTwoCopies:1;
 
+	/** If true, it will be possible to create a FUnorderedAccessViewRHIRef using RHICreateUnorderedAccessView and the internal FTexture2DRHIRef. */
+	uint32 bCanCreateUAV : 1;
+
 	/**
 	 * Render thread: Access the render target resource for this texture target object
 	 * @return pointer to resource or NULL if not initialized
diff --git a/Engine/Source/Runtime/Engine/Classes/VectorField/VectorFieldStatic.h b/Engine/Source/Runtime/Engine/Classes/VectorField/VectorFieldStatic.h
--- a/Engine/Source/Runtime/Engine/Classes/VectorField/VectorFieldStatic.h
+++ b/Engine/Source/Runtime/Engine/Classes/VectorField/VectorFieldStatic.h
@@ -12,6 +12,7 @@
 #include "VectorField/VectorField.h"
 #include "VectorFieldStatic.generated.h"
 
+class FRHITexture;
 struct FPropertyChangedEvent;
 
 UCLASS(hidecategories=VectorFieldBounds, MinimalAPI)
@@ -31,6 +32,9 @@ class UVectorFieldStatic : public UVectorField
 	UPROPERTY(Category=VectorFieldStatic, VisibleAnywhere)
 	int32 SizeZ;
 
+	/** Whether to keep vector field data accessible to the CPU. */
+	UPROPERTY(Category=VectorFieldStatic, EditAnywhere)
+	bool bAllowCPUAccess;
 
 public:
 	/** The resource for this vector field. */
@@ -39,6 +43,11 @@ public:
 	/** Source vector data. */
 	FByteBulkData SourceData;
 
+	/** Local copy of the source vector data. */
+	UPROPERTY(Transient)
+	TArray<FVector4> CPUData; 
+
+
 #if WITH_EDITORONLY_DATA
 	UPROPERTY()
 	FString SourceFilePath_DEPRECATED;
@@ -68,6 +77,17 @@ public:
 	 * Initialize resources.
 	 */
 	ENGINE_API void InitResource();
+
+	/** Takes a local copy of the source bulk data so that it is readable at runtime on the CPU. */
+	ENGINE_API void UpdateCPUData();
+
+#if WITH_EDITOR
+	/** Sets the bAllowCPUAccess flag and calls UpdateCPUData(). */
+	ENGINE_API void SetCPUAccessEnabled();
+#endif // WITH_EDITOR
+
+	/** Returns a reference to a 3D texture handle for the GPU data. */
+	ENGINE_API FRHITexture* GetVolumeTextureRef();
 private:
 
 	/** Permit the factory class to update and release resources externally. */
@@ -84,5 +104,6 @@ private:
 	 */
 	ENGINE_API void ReleaseResource();
 
+	
 };
 
diff --git a/Engine/Source/Runtime/Engine/Private/ActorConstruction.cpp b/Engine/Source/Runtime/Engine/Private/ActorConstruction.cpp
--- a/Engine/Source/Runtime/Engine/Private/ActorConstruction.cpp
+++ b/Engine/Source/Runtime/Engine/Private/ActorConstruction.cpp
@@ -1044,7 +1044,8 @@ UActorComponent* AActor::CreateComponentFromTemplateData(const FBlueprintCookedC
 
 UActorComponent* AActor::AddComponent(FName TemplateName, bool bManualAttachment, const FTransform& RelativeTransform, const UObject* ComponentTemplateContext)
 {
-	if (GetWorld()->bIsTearingDown)
+	UWorld* World = GetWorld();
+	if (World->bIsTearingDown)
 	{
 		UE_LOG(LogActor, Warning, TEXT("AddComponent failed because we are in the process of tearing down the world"));
 		return nullptr;
@@ -1115,7 +1116,6 @@ UActorComponent* AActor::AddComponent(FName TemplateName, bool bManualAttachment
 			NewActorComp->RegisterComponent();
 		}
 
-		UWorld* World = GetWorld();
 		if (!bRunningUserConstructionScript && World && bIsSceneComponent)
 		{
 			UPrimitiveComponent* NewPrimitiveComponent = Cast<UPrimitiveComponent>(NewActorComp);
diff --git a/Engine/Source/Runtime/Engine/Private/Particles/ParticleComponents.cpp b/Engine/Source/Runtime/Engine/Private/Particles/ParticleComponents.cpp
--- a/Engine/Source/Runtime/Engine/Private/Particles/ParticleComponents.cpp
+++ b/Engine/Source/Runtime/Engine/Private/Particles/ParticleComponents.cpp
@@ -114,6 +114,10 @@ DECLARE_CYCLE_STAT(TEXT("ParticleComponent CreateRenderState Concurrent GT"), ST
 DECLARE_CYCLE_STAT(TEXT("PSys Comp Marshall Time GT"), STAT_UParticleSystemComponent_Marshall, STATGROUP_Particles);
 
 CSV_DECLARE_CATEGORY_MODULE_EXTERN(CORE_API, Basic);
+DEFINE_STAT(STAT_ParticlesOverview_GT);
+DEFINE_STAT(STAT_ParticlesOverview_GT_CNC);
+DEFINE_STAT(STAT_ParticlesOverview_RT);
+DEFINE_STAT(STAT_ParticlesOverview_RT_CNC);
 
 #include "InGamePerformanceTracker.h"
 
@@ -3839,6 +3843,7 @@ void UParticleSystemComponent::OnUnregister()
 void UParticleSystemComponent::CreateRenderState_Concurrent()
 {
 	SCOPE_CYCLE_COUNTER(STAT_ParticleSystemComponent_CreateRenderState_Concurrent);
+	SCOPE_CYCLE_COUNTER(STAT_ParticlesOverview_GT_CNC);
 
 	ForceAsyncWorkCompletion(ENSURE_AND_STALL);
 	check( GetWorld() );
@@ -3873,6 +3878,7 @@ void UParticleSystemComponent::CreateRenderState_Concurrent()
 void UParticleSystemComponent::SendRenderTransform_Concurrent()
 {
 	SCOPE_CYCLE_COUNTER(STAT_ParticleSystemComponent_SendRenderTransform_Concurrent);
+	SCOPE_CYCLE_COUNTER(STAT_ParticlesOverview_GT_CNC);
 
 	ForceAsyncWorkCompletion(ENSURE_AND_STALL);
 	if (bIsActive)
@@ -3890,6 +3896,7 @@ void UParticleSystemComponent::SendRenderTransform_Concurrent()
 void UParticleSystemComponent::SendRenderDynamicData_Concurrent()
 {
 	SCOPE_CYCLE_COUNTER(STAT_ParticleSystemComponent_SendRenderDynamicData_Concurrent);
+	SCOPE_CYCLE_COUNTER(STAT_ParticlesOverview_GT_CNC);
 
 	ForceAsyncWorkCompletion(ENSURE_AND_STALL);
 	Super::SendRenderDynamicData_Concurrent();
@@ -3924,6 +3931,7 @@ void UParticleSystemComponent::SendRenderDynamicData_Concurrent()
 void UParticleSystemComponent::DestroyRenderState_Concurrent()
 {
 	SCOPE_CYCLE_COUNTER(STAT_ParticleSystemComponent_DestroyRenderState_Concurrent);
+	SCOPE_CYCLE_COUNTER(STAT_ParticlesOverview_GT_CNC);
 
 	ForceAsyncWorkCompletion(ENSURE_AND_STALL);
 	check( GetWorld() );
@@ -4927,6 +4935,7 @@ void UParticleSystemComponent::TickComponent(float DeltaTime, enum ELevelTick Ti
 	CSV_SCOPED_TIMING_STAT_EXCLUSIVE(Effects);
 	LLM_SCOPE(ELLMTag::Particles);
 	FInGameScopedCycleCounter InGameCycleCounter(GetWorld(), EInGamePerfTrackers::VFXSignificance, EInGamePerfTrackerThreads::GameThread, bIsManagingSignificance);
+	SCOPE_CYCLE_COUNTER(STAT_ParticlesOverview_GT);
 
 	if (Template == nullptr || Template->Emitters.Num() == 0)
 	{
@@ -5214,6 +5223,8 @@ void UParticleSystemComponent::ComputeTickComponent_Concurrent()
 
 	SCOPE_CYCLE_COUNTER(STAT_ParticleComputeTickTime);
 	FScopeCycleCounterUObject AdditionalScope(AdditionalStatObject(), GET_STATID(STAT_ParticleComputeTickTime));
+	SCOPE_CYCLE_COUNTER(STAT_ParticlesOverview_GT_CNC);
+
 	// Tick Subemitters.
 	int32 EmitterIndex;
 	NumSignificantEmitters = 0;
@@ -5287,6 +5298,7 @@ void UParticleSystemComponent::FinalizeTickComponent()
 	FInGameScopedCycleCounter InGameCycleCounter(GetWorld(), EInGamePerfTrackers::VFXSignificance, IsInGameThread() ? EInGamePerfTrackerThreads::GameThread : EInGamePerfTrackerThreads::OtherThread, bIsManagingSignificance);
 
 	SCOPE_CYCLE_COUNTER(STAT_ParticleFinalizeTickTime);
+	SCOPE_CYCLE_COUNTER(STAT_ParticlesOverview_GT);
 
 	if(bAsyncDataCopyIsValid)
 	{
@@ -5786,6 +5798,7 @@ void UParticleSystemComponent::ActivateSystem(bool bFlagAsJustAttached)
 {
 	CSV_SCOPED_TIMING_STAT_EXCLUSIVE(Effects);
 	SCOPE_CYCLE_COUNTER(STAT_ParticleActivateTime);
+	SCOPE_CYCLE_COUNTER(STAT_ParticlesOverview_GT);
 	ForceAsyncWorkCompletion(STALL);
 
 	if (IsTemplate() == true || !IsRegistered() || 	!FApp::CanEverRender())
@@ -6033,6 +6046,7 @@ void UParticleSystemComponent::DeactivateSystem()
 {
 	UWorld* World = GetWorld();
 	FInGameScopedCycleCounter InGameCycleCounter(World, EInGamePerfTrackers::VFXSignificance, EInGamePerfTrackerThreads::GameThread, bIsManagingSignificance);
+	SCOPE_CYCLE_COUNTER(STAT_ParticlesOverview_GT);
 
 	if (IsTemplate() == true)
 	{
diff --git a/Engine/Source/Runtime/Engine/Private/Particles/ParticleSystemRender.cpp b/Engine/Source/Runtime/Engine/Private/Particles/ParticleSystemRender.cpp
--- a/Engine/Source/Runtime/Engine/Private/Particles/ParticleSystemRender.cpp
+++ b/Engine/Source/Runtime/Engine/Private/Particles/ParticleSystemRender.cpp
@@ -1122,6 +1122,8 @@ void FDynamicSpriteEmitterData::GetDynamicMeshElementsEmitter(const FParticleSys
 								[this, SourceData, View, Proxy, Allocation, DynamicParameterAllocation, bInstanced, bSort, ParticleCount, NumVerticesPerParticleInBuffer]()
 								{
 									SCOPE_CYCLE_COUNTER(STAT_FDynamicSpriteEmitterData_GetDynamicMeshElementsEmitter_Task);
+									SCOPE_CYCLE_COUNTER(STAT_ParticlesOverview_RT_CNC);
+
 									FMemMark Mark(FMemStack::Get());
 									FParticleOrder* ParticleOrder = NULL;
 									if (bSort)
@@ -2397,6 +2399,7 @@ void FDynamicMeshEmitterData::CalculateParticleTransform(
 void FDynamicMeshEmitterData::GetInstanceData(void* InstanceData, void* DynamicParameterData, void* PrevTransformBuffer, const FParticleSystemSceneProxy* Proxy, const FSceneView* View) const
 {
 	SCOPE_CYCLE_COUNTER(STAT_ParticlePackingTime);
+	SCOPE_CYCLE_COUNTER(STAT_ParticlesOverview_RT_CNC);
 
 	int32 SubImagesX = Source.SubImages_Horizontal;
 	int32 SubImagesY = Source.SubImages_Vertical;
@@ -6933,6 +6936,7 @@ void FParticleSystemSceneProxy::GetDynamicMeshElements(const TArray<const FScene
 	FInGameScopedCycleCounter InGameCycleCounter(GetScene().GetWorld(), EInGamePerfTrackers::VFXSignificance, EInGamePerfTrackerThreads::RenderThread, bManagingSignificance);
 
 	SCOPE_CYCLE_COUNTER(STAT_FParticleSystemSceneProxy_GetMeshElements);
+	SCOPE_CYCLE_COUNTER(STAT_ParticlesOverview_RT);
 
 	if ((GIsEditor == true) || (GbEnableGameThreadLODCalculation == false))
 	{
@@ -7072,6 +7076,7 @@ void FParticleSystemSceneProxy::UpdateData(FParticleDynamicData* NewDynamicData)
 void FParticleSystemSceneProxy::UpdateData_RenderThread(FParticleDynamicData* NewDynamicData)
 {
 	FInGameScopedCycleCounter InGameCycleCounter(GetScene().GetWorld(), EInGamePerfTrackers::VFXSignificance, EInGamePerfTrackerThreads::RenderThread, bManagingSignificance);
+	SCOPE_CYCLE_COUNTER(STAT_ParticlesOverview_RT);
 
 	ReleaseRenderThreadResourcesForEmitterData();
 	if (DynamicData != NewDynamicData)
@@ -7259,6 +7264,7 @@ void FParticleSystemSceneProxy::UpdateWorldSpacePrimitiveUniformBuffer() const
 void FParticleSystemSceneProxy::GatherSimpleLights(const FSceneViewFamily& ViewFamily, FSimpleLightArray& OutParticleLights) const
 {
 	FInGameScopedCycleCounter InGameCycleCounter(GetScene().GetWorld(), EInGamePerfTrackers::VFXSignificance, EInGamePerfTrackerThreads::RenderThread, bManagingSignificance);
+	SCOPE_CYCLE_COUNTER(STAT_ParticlesOverview_RT);
 	if (DynamicData != NULL)
 	{
 		FScopeCycleCounter Context(GetStatId());
@@ -7277,6 +7283,8 @@ void FParticleSystemSceneProxy::GatherSimpleLights(const FSceneViewFamily& ViewF
 FPrimitiveSceneProxy* UParticleSystemComponent::CreateSceneProxy()
 {
 	SCOPE_CYCLE_COUNTER(STAT_FParticleSystemSceneProxy_Create);
+	SCOPE_CYCLE_COUNTER(STAT_ParticlesOverview_GT);
+
 	FParticleSystemSceneProxy* NewProxy = NULL;
 
 	//@fixme EmitterInstances.Num() check should be here to avoid proxies for dead emitters but there are some edge cases where it happens for emitters that have just activated...
diff --git a/Engine/Source/Runtime/Engine/Private/Particles/WorldPSCPool.cpp b/Engine/Source/Runtime/Engine/Private/Particles/WorldPSCPool.cpp
--- a/Engine/Source/Runtime/Engine/Private/Particles/WorldPSCPool.cpp
+++ b/Engine/Source/Runtime/Engine/Private/Particles/WorldPSCPool.cpp
@@ -87,10 +87,13 @@ UParticleSystemComponent* FPSCPool::Acquire(UWorld* World, UParticleSystem* Temp
 		check(RetElem.PSC->Template == Template);
 		check(!RetElem.PSC->IsPendingKill());
 
-		// Rename the PSC to move it into the current PersistentLevel - it may have been spawned in one
-		// level but is now needed in another level.
-		// Use the REN_ForceNoResetLoaders flag to prevent the rename from potentially calling FlushAsyncLoading.
-		RetElem.PSC->Rename(nullptr, World, REN_ForceNoResetLoaders);
+		if (RetElem.PSC->GetWorld() != World)
+		{
+			// Rename the PSC to move it into the current PersistentLevel - it may have been spawned in one
+			// level but is now needed in another level.
+			// Use the REN_ForceNoResetLoaders flag to prevent the rename from potentially calling FlushAsyncLoading.
+			RetElem.PSC->Rename(nullptr, World, REN_ForceNoResetLoaders);
+		}
 	}
 	else
 	{
diff --git a/Engine/Source/Runtime/Engine/Private/TextureRenderTarget.cpp b/Engine/Source/Runtime/Engine/Private/TextureRenderTarget.cpp
--- a/Engine/Source/Runtime/Engine/Private/TextureRenderTarget.cpp
+++ b/Engine/Source/Runtime/Engine/Private/TextureRenderTarget.cpp
@@ -18,6 +18,7 @@ UTextureRenderTarget::UTextureRenderTarget(const FObjectInitializer& ObjectIniti
 	SRGB = true;
 	LODGroup = TEXTUREGROUP_RenderTarget;	
 	bNeedsTwoCopies = false;
+	bCanCreateUAV = false;
 #if WITH_EDITORONLY_DATA
 	CompressionNone = true;
 #endif // #if WITH_EDITORONLY_DATA
diff --git a/Engine/Source/Runtime/Engine/Private/VectorField.cpp b/Engine/Source/Runtime/Engine/Private/VectorField.cpp
--- a/Engine/Source/Runtime/Engine/Private/VectorField.cpp
+++ b/Engine/Source/Runtime/Engine/Private/VectorField.cpp
@@ -244,19 +244,19 @@ public:
 			FVectorFieldStaticResource*, Resource, this,
 			FUpdateParams, UpdateParams, UpdateParams,
 		{
-			// Free any existing volume data on the resource.
-			FMemory::Free(Resource->VolumeData);
-			
-			// Update settings on this resource.
-			Resource->SizeX = UpdateParams.SizeX;
-			Resource->SizeY = UpdateParams.SizeY;
-			Resource->SizeZ = UpdateParams.SizeZ;
-			Resource->Intensity = UpdateParams.Intensity;
-			Resource->LocalBounds = UpdateParams.Bounds;
-			Resource->VolumeData = UpdateParams.VolumeData;
-
-			// Update RHI resources.
-			Resource->UpdateRHI();
+				// Free any existing volume data on the resource.
+				FMemory::Free(Resource->VolumeData);
+
+				// Update settings on this resource.
+				Resource->SizeX = UpdateParams.SizeX;
+				Resource->SizeY = UpdateParams.SizeY;
+				Resource->SizeZ = UpdateParams.SizeZ;
+				Resource->Intensity = UpdateParams.Intensity;
+				Resource->LocalBounds = UpdateParams.Bounds;
+				Resource->VolumeData = UpdateParams.VolumeData;
+
+				// Update RHI resources.
+				Resource->UpdateRHI();
 		});
 	}
 
@@ -268,6 +268,7 @@ private:
 
 UVectorFieldStatic::UVectorFieldStatic(const FObjectInitializer& ObjectInitializer)
 	: Super(ObjectInitializer)
+	, bAllowCPUAccess(false)
 {
 }
 
@@ -280,18 +281,100 @@ void UVectorFieldStatic::InitInstance(FVectorFieldInstance* Instance, bool bPrev
 void UVectorFieldStatic::InitResource()
 {
 	check(Resource == NULL);
-	Resource = new FVectorFieldStaticResource( this );
-	BeginInitResource( Resource );
+
+	// Loads and copies the bulk data into CPUData if bAllowCPUAccess is set, otherwise clear CPUData. 
+	UpdateCPUData();
+
+	Resource = new FVectorFieldStaticResource(this);
+	BeginInitResource(Resource);
 }
 
 
 void UVectorFieldStatic::UpdateResource()
 {
 	check(Resource != NULL);
+
+	// Loads and copies the bulk data into CPUData if bAllowCPUAccess is set, otherwise clears CPUData. 
+	UpdateCPUData();
+
 	FVectorFieldStaticResource* StaticResource = (FVectorFieldStaticResource*)Resource;
 	StaticResource->UpdateResource(this);
 }
 
+#if WITH_EDITOR
+ENGINE_API void UVectorFieldStatic::SetCPUAccessEnabled()
+{
+	bAllowCPUAccess = true;
+	UpdateCPUData();
+}
+#endif // WITH_EDITOR
+
+void UVectorFieldStatic::UpdateCPUData()
+{
+	if (bAllowCPUAccess)
+	{
+		// Grab a copy of the bulk vector data. 
+		// If the data is already loaded it makes a copy and discards the old content,
+		// otherwise it simply loads the data directly from file into the pointer after allocating.
+		FFloat16Color *Ptr = nullptr;
+		SourceData.GetCopy((void**)&Ptr, /* bDiscardInternalCopy */ true);
+
+		// Make sure the data is actually valid. 
+		if (!ensure(Ptr))
+		{
+			UE_LOG(LogVectorField, Error, TEXT("Vector field data is not loaded."));
+			return;
+		}
+
+		// Make sure the size actually match what we expect
+		if (!ensure(SourceData.GetBulkDataSize() == (SizeX*SizeY*SizeZ) * sizeof(FFloat16Color)))
+		{
+			UE_LOG(LogVectorField, Error, TEXT("Vector field bulk data size is different than expected. Expected %d bytes, got %d."), SizeX*SizeY*SizeZ, SourceData.GetBulkDataSize());
+			FMemory::Free(Ptr);
+			return;
+		}
+
+		// GetCopy should free/unload the data.
+		if (SourceData.IsBulkDataLoaded())
+		{
+			// NOTE(mv): This assertion will fail in the case where the bulk data is still available even though the bDiscardInternalCopy
+			//           flag is toggled when FUntypedBulkData::CanLoadFromDisk() also fail. This happens when the user tries to allow 
+			//           CPU access to a newly imported file that isn't reloaded. We still have our valid data, so we just issue a 
+			//           warning and move on. See FUntypedBulkData::GetCopy() for more details. 
+			UE_LOG(LogVectorField, Warning, TEXT("SourceData.GetCopy() is supposed to unload the data after copying, but it is still loaded."));
+		}
+
+		// Convert from 16-bit to 32-bit floats.
+			// Use vec4s instead of vec3s because of alignment, which in principle would be better for 
+			// cache and automatic or manual vectorization, even if the memory usage is 33% larger. 
+			// Need to profile to to make sure.
+		CPUData.SetNumUninitialized(SizeX*SizeY*SizeZ);
+		for (size_t i = 0; i < (size_t)(SizeX*SizeY*SizeZ); i++)
+		{
+			CPUData[i] = FVector4(float(Ptr[i].R), float(Ptr[i].G), float(Ptr[i].B), 0.0f);
+		}
+
+		FMemory::Free(Ptr);
+	}
+	else
+	{
+		// If there's no need to access the CPU data just empty the array.
+		CPUData.Empty();
+	}
+}
+
+FRHITexture* UVectorFieldStatic::GetVolumeTextureRef()
+{
+	if (Resource)
+	{
+		return Resource->VolumeTextureRHI;
+	}
+	else
+	{	
+		// Fallback to a global 1x1x1 black texture when no vector field is loaded or unavailable
+		return GBlackVolumeTexture->TextureRHI;
+	}
+}
 
 void UVectorFieldStatic::ReleaseResource()
 {
diff --git a/Engine/Source/Runtime/Engine/Public/ParticleEmitterInstances.h b/Engine/Source/Runtime/Engine/Public/ParticleEmitterInstances.h
--- a/Engine/Source/Runtime/Engine/Public/ParticleEmitterInstances.h
+++ b/Engine/Source/Runtime/Engine/Public/ParticleEmitterInstances.h
@@ -1122,6 +1122,10 @@ struct FParticleBeam2EmitterInstance : public FParticleEmitterInstance
 		{
 			SourceEmitter = NULL;
 		}
+		if (TargetEmitter == Instance)
+		{
+			TargetEmitter = NULL;
+		}
 	}
 
 protected:
