Fix incorrect instance ID being used in VFs when using ISR.

diff --git a/Engine/Shaders/Private/MeshParticleVertexFactory.ush b/Engine/Shaders/Private/MeshParticleVertexFactory.ush
--- a/Engine/Shaders/Private/MeshParticleVertexFactory.ush
+++ b/Engine/Shaders/Private/MeshParticleVertexFactory.ush
@@ -182,7 +182,7 @@ float4x4 GetParticleTransform(FVertexFactoryInput Input)
 float4x4 GetParticlePrevTransform(FVertexFactoryInput Input)
 {
 #if PARTICLE_MESH_INSTANCED
-	uint TransformBaseIdx = Input.InstanceId * 3;
+	uint TransformBaseIdx = GetInstanceId(Input.InstanceId) * 3;
 
 	float4x4 PrevTransform4x4 = float4x4(
 		PrevTransformBuffer[TransformBaseIdx + 0], 
diff --git a/Engine/Shaders/Private/ParticleGPUSpriteVertexFactory.ush b/Engine/Shaders/Private/ParticleGPUSpriteVertexFactory.ush
--- a/Engine/Shaders/Private/ParticleGPUSpriteVertexFactory.ush
+++ b/Engine/Shaders/Private/ParticleGPUSpriteVertexFactory.ush
@@ -229,7 +229,7 @@ void GetSpriteTangents(
 FVertexFactoryIntermediates GetVertexFactoryIntermediates( FVertexFactoryInput Input )
 {
 	// Sample the position and velocity textures to get the current state of the particle.
-	uint InstanceId = Input.InstanceId * PARTICLES_PER_INSTANCE + Input.VertexId / 4;
+	uint InstanceId = GetInstanceId(Input.InstanceId) * PARTICLES_PER_INSTANCE + Input.VertexId / 4;
 	float2 ParticleIndex = ParticleIndices[ParticleIndicesOffset + InstanceId];
 	float4 PositionSample = Texture2DSampleLevel(PositionTexture, PositionTextureSampler, ParticleIndex, 0);
 	float4 VelocitySample = Texture2DSampleLevel(VelocityTexture, VelocityTextureSampler, ParticleIndex, 0);
diff --git a/Engine/Shaders/Private/VectorFieldVisualizationVertexFactory.ush b/Engine/Shaders/Private/VectorFieldVisualizationVertexFactory.ush
--- a/Engine/Shaders/Private/VectorFieldVisualizationVertexFactory.ush
+++ b/Engine/Shaders/Private/VectorFieldVisualizationVertexFactory.ush
@@ -56,7 +56,7 @@ FVertexFactoryIntermediates GetVertexFactoryIntermediates( FVertexFactoryInput I
 	const float3 VoxelSize = VectorFieldVis.VoxelSize.xyz;
 	const float3 HalfVoxelOffset = VoxelSize.xyz * 0.5f;
 	const float3 VoxelMultipliers = float3(1, VoxelSize.x, VoxelSize.x * VoxelSize.y);
-	const float3 VoxelIndex = floor(float3(Input.InstanceId.xxx) * VoxelMultipliers.xyz);
+	const float3 VoxelIndex = floor((float3)GetInstanceId(Input.InstanceId.xxx) * VoxelMultipliers.xyz);
 	const float3 VoxelUV = frac(VoxelIndex.xyz * VoxelSize.xyz) + HalfVoxelOffset;
 
 	// Read vector for this voxel.
