Normalize half vector in full precision to avoid precision issues when camera and light verctors have similar directrion (mobile only)

diff --git a/Engine/Shaders/Private/MobileBasePassPixelShader.usf b/Engine/Shaders/Private/MobileBasePassPixelShader.usf
--- a/Engine/Shaders/Private/MobileBasePassPixelShader.usf
+++ b/Engine/Shaders/Private/MobileBasePassPixelShader.usf
@@ -541,7 +541,7 @@ void Main(
 
 		float NoL = max(0, dot(float3(MaterialParameters.WorldNormal), float3(MobileDirectionalLight.DirectionalLightDirectionAndShadowTransition.xyz)));
 		float RoL = max(0, dot(float3(MaterialParameters.ReflectionVector), float3(MobileDirectionalLight.DirectionalLightDirectionAndShadowTransition.xyz)));
-		float3 H = normalize(MaterialParameters.CameraVector + MobileDirectionalLight.DirectionalLightDirectionAndShadowTransition.xyz);
+		float3 H = normalize(float3(MaterialParameters.CameraVector) + float3(MobileDirectionalLight.DirectionalLightDirectionAndShadowTransition.xyz));
 		float NoH = max(0,dot(MaterialParameters.WorldNormal, H));
 
 		#if FULLY_ROUGH
