Build fix, adding missing pragma once.

diff --git a/Engine/Shaders/Private/PostProcessHistogramCommon.ush b/Engine/Shaders/Private/PostProcessHistogramCommon.ush
--- a/Engine/Shaders/Private/PostProcessHistogramCommon.ush
+++ b/Engine/Shaders/Private/PostProcessHistogramCommon.ush
@@ -128,10 +128,7 @@ float ComputeAverageLuminaneWithoutOutlier(Texture2D HistogramTexture, float Min
 float ComputeEyeAdaptationExposure(Texture2D HistogramTexture)
 {
 	float HistogramSum = ComputeHistogramSum(HistogramTexture);
-	float UnclampedAdaptedLuminance = ComputeAverageLuminaneWithoutOutlier(HistogramTexture, HistogramSum * EyeAdaptationParams[0].x, HistogramSum * EyeAdaptationParams[0].y);
-	float ClampedAdaptedLuminance = clamp(UnclampedAdaptedLuminance, EyeAdaptationParams[0].z, EyeAdaptationParams[0].w);
-
-	return ClampedAdaptedLuminance;
+	return ComputeAverageLuminaneWithoutOutlier(HistogramTexture, HistogramSum * EyeAdaptationParams[0].x, HistogramSum * EyeAdaptationParams[0].y);
 }
 
 // Computation of triangle function on normalized space
