Apply Indirect Ligthing Color setting from PPV GI section in mobile render
Fixes: Post process Global Illumination doesn't work on mobile preview, device

diff --git a/Engine/Shaders/Private/MobileBasePassPixelShader.usf b/Engine/Shaders/Private/MobileBasePassPixelShader.usf
--- a/Engine/Shaders/Private/MobileBasePassPixelShader.usf
+++ b/Engine/Shaders/Private/MobileBasePassPixelShader.usf
@@ -433,7 +433,7 @@ void Main(
 		GetLightMapCoordinates(Interpolants, LightmapUV0, LightmapUV1, LightmapDataIndex);
 
 		half4 LightmapColor = GetLightMapColorLQ( LightmapUV0, LightmapUV1, LightmapDataIndex, MaterialParameters.WorldNormal );
-		Color += LightmapColor.rgb * DiffuseColor;
+		Color += LightmapColor.rgb * DiffuseColor * View.IndirectLightingColorScale;
 		IndirectIrradiance = LightmapColor.a;
 	#elif CACHED_POINT_INDIRECT_LIGHTING
 		#if MATERIALBLENDING_MASKED || MATERIALBLENDING_SOLID
@@ -458,9 +458,7 @@ void Main(
 			half3 DiffuseGI = max(half3(0,0,0), DotSH3(PointIndirectLighting, DiffuseTransferSH));
 
 			IndirectIrradiance = Luminance(DiffuseGI);
-			Color += DiffuseColor * DiffuseGI;			
-
-
+			Color += DiffuseColor * DiffuseGI * View.IndirectLightingColorScale;			
 		#else 
 			
 			// Non-directional for translucency
@@ -470,7 +468,7 @@ void Main(
 			half3 DiffuseGI = PointIndirectLighting;
 
 			IndirectIrradiance = Luminance(DiffuseGI);
-			Color += DiffuseColor * DiffuseGI;
+			Color += DiffuseColor * DiffuseGI * View.IndirectLightingColorScale;
 
 		#endif
 	#endif
