Fix a d3d12 PSO creation failure caused by pixel shader reading an interpolant that is never written by the domain shader

diff --git a/Engine/Shaders/Private/ShadowDepthVertexShader.usf b/Engine/Shaders/Private/ShadowDepthVertexShader.usf
--- a/Engine/Shaders/Private/ShadowDepthVertexShader.usf
+++ b/Engine/Shaders/Private/ShadowDepthVertexShader.usf
@@ -184,10 +184,15 @@ FShadowDepthDSToPS PassFinalizeTessellationOutput(FShadowDepthVSToDS Interpolant
 		);
 #endif
 
-#if !(ONEPASS_POINTLIGHT_SHADOW || PERSPECTIVE_CORRECT_DEPTH) && COMPILER_SUPPORTS_EMPTY_STRUCTS
+// Conditions below are derived. See struct FShadowDepthVSToPS in ShadowDepthCommon.ush
+#if !PERSPECTIVE_CORRECT_DEPTH && !COMPILER_SUPPORTS_EMPTY_STRUCTS && !ONEPASS_POINTLIGHT_SHADOW
 	O.Dummy = 0;
 #endif
 
+#if INTERPOLATE_POSITION
+	O.PixelPosition = WorldPosition.xyz;
+#endif
+
 	return O;
 }
 	
