Fix another d3d12 PSO creation failure caused by DS-GS linkage error

diff --git a/Engine/Shaders/Private/ShadowDepthVertexShader.usf b/Engine/Shaders/Private/ShadowDepthVertexShader.usf
--- a/Engine/Shaders/Private/ShadowDepthVertexShader.usf
+++ b/Engine/Shaders/Private/ShadowDepthVertexShader.usf
@@ -136,10 +136,14 @@ struct FShadowDepthVSToDS
 
 struct FShadowDepthDSToPS : FShadowDepthVSToPS
 {
-	float4 Position : SV_POSITION;
-#if ONEPASS_POINTLIGHT_SHADOW && USING_VERTEX_SHADER_LAYER
+#if ONEPASS_POINTLIGHT_SHADOW
+	float4 Position : TEXCOORD6;
+#if USING_VERTEX_SHADER_LAYER
 	uint LayerIndex : SV_RenderTargetArrayIndex;
 #endif
+#else
+	float4 Position : SV_POSITION;
+#endif
 };
 
 #define FPassSpecificVSToDS FShadowDepthVSToDS
