Fix 5263539 for Metal - there is no GS on Metal so this change was causing an interface mismatch and silent Metal shader compile failure - Metal needs the DS to emit SV_Position.

diff --git a/Engine/Shaders/Private/ShadowDepthVertexShader.usf b/Engine/Shaders/Private/ShadowDepthVertexShader.usf
--- a/Engine/Shaders/Private/ShadowDepthVertexShader.usf
+++ b/Engine/Shaders/Private/ShadowDepthVertexShader.usf
@@ -137,9 +137,12 @@ struct FShadowDepthVSToDS
 struct FShadowDepthDSToPS : FShadowDepthVSToPS
 {
 #if ONEPASS_POINTLIGHT_SHADOW
-	float4 Position : TEXCOORD6;
 #if USING_VERTEX_SHADER_LAYER
+	float4 Position : SV_POSITION;
+	float4 InterpPosition : TEXCOORD6;
 	uint LayerIndex : SV_RenderTargetArrayIndex;
+#else
+	float4 Position : TEXCOORD6;
 #endif
 #else
 	float4 Position : SV_POSITION;
