Framebuffer fetch is fragment-shader only for Metal - you can't use it in vertex shaders.

diff --git a/Engine/Shaders/Private/SceneTexturesCommon.ush b/Engine/Shaders/Private/SceneTexturesCommon.ush
--- a/Engine/Shaders/Private/SceneTexturesCommon.ush
+++ b/Engine/Shaders/Private/SceneTexturesCommon.ush
@@ -108,7 +108,11 @@ float CalcSceneDepth(float2 ScreenUV)
 				return DepthbufferFetchES2(SceneW, View.InvDeviceZToWorldZTransform[2], View.InvDeviceZToWorldZTransform[3]);
 			#endif 
 		#elif METAL_PROFILE && !MAC
-			return FramebufferFetchES2().w;
+			#if PIXELSHADER
+				return FramebufferFetchES2().w;
+			#else
+				return ConvertFromDeviceZ(Texture2DSampleLevel(MobileSceneTextures.SceneDepthTexture, MobileSceneTextures.SceneDepthTextureSampler, ScreenUV, 0).r);
+			#endif
 		#elif COMPILER_GLSL_ES3_1
 			#if OUTPUT_GAMMA_SPACE
 				// Try framebuffer_fetch_depth_stencil to get the depth directly if the extension is available.
diff --git a/Engine/Source/Developer/Apple/MetalShaderFormat/Private/MetalBackend.cpp b/Engine/Source/Developer/Apple/MetalShaderFormat/Private/MetalBackend.cpp
--- a/Engine/Source/Developer/Apple/MetalShaderFormat/Private/MetalBackend.cpp
+++ b/Engine/Source/Developer/Apple/MetalShaderFormat/Private/MetalBackend.cpp
@@ -4935,7 +4935,11 @@ char* FMetalCodeBackend::GenerateCode(exec_list* ir, _mesa_glsl_parse_state* sta
 	ExpandArrayAssignments(ir, state);
 
 	// Fix any special language extensions (FrameBufferFetchES2() intrinsic)
-	FixIntrinsics(ir, state);
+	// Only applicable to fragment shaders - this intrinsic is not available elsewhere
+	if (Frequency == HSF_PixelShader)
+	{
+		FixIntrinsics(ir, state);
+	}
 
 	// Remove half->float->half or float->half->float
 	FixRedundantCasts(ir);
