Clamping observer height to avoid precision issues. Fixes the issue where the height fog disapears when the observer is too far up.

diff --git a/Engine/Shaders/Private/HeightFogCommon.ush b/Engine/Shaders/Private/HeightFogCommon.ush
--- a/Engine/Shaders/Private/HeightFogCommon.ush
+++ b/Engine/Shaders/Private/HeightFogCommon.ush
@@ -19,7 +19,7 @@
 static const float FLT_EPSILON = 0.001f;
 static const float FLT_EPSILON2 = 0.01f;
 
-// FogStruct.ExponentialFogParameters: FogDensity * exp2(-FogHeightFalloff * (CameraWorldPosition.z - FogHeight)) in x, FogHeightFalloff in y, CosTerminatorAngle in z, StartDistance in w. 
+// FogStruct.ExponentialFogParameters: FogDensity * exp2(-FogHeightFalloff * (CameraWorldPosition.z - FogHeight)) in x, FogHeightFalloff in y, MaxWorldObserverHeight in z, StartDistance in w. 
 // FogStruct.ExponentialFogParameters2: FogDensitySecond * exp2(-FogHeightFalloffSecond * (CameraWorldPosition.z - FogHeightSecond)) in x, FogHeightFalloffSecond in y, FogDensitySecond in z, FogHeightSecond in w 
 // FogStruct.ExponentialFogParameters3: FogDensity in x, FogHeight in y, whether to use cubemap fog color in z, FogCutoffDistance in w. 
 // FogStruct.FogInscatteringTextureParameters: mip distance scale in x, bias in y, num mips in z 
@@ -61,8 +61,12 @@ float CalculateLineIntegralShared(float FogHeightFalloff, float RayDirectionZ, f
 half4 GetExponentialHeightFog(float3 WorldPositionRelativeToCamera, float ExcludeDistance)
 {
 	const half MinFogOpacity = FogStruct.ExponentialFogColorParameter.w;
+	const float MaxWorldObserverHeight = FogStruct.ExponentialFogParameters.z;
+
+	const float3 WorldObserverOrigin = float3(View.WorldCameraOrigin.xy, min(View.WorldCameraOrigin.z, MaxWorldObserverHeight)); // Clamp z to max height
 
 	float3 CameraToReceiver = WorldPositionRelativeToCamera;
+	CameraToReceiver.z += View.WorldCameraOrigin.z - WorldObserverOrigin.z; // Compensate this vector for clamping the observer height 
 	float CameraToReceiverLengthSqr = dot(CameraToReceiver, CameraToReceiver);
 	float CameraToReceiverLengthInv = rsqrt(CameraToReceiverLengthSqr);
 	float CameraToReceiverLength = CameraToReceiverLengthSqr * CameraToReceiverLengthInv;
@@ -82,7 +86,7 @@ half4 GetExponentialHeightFog(float3 WorldPositionRelativeToCamera, float Exclud
 	// While rendering a planar reflection with a clip plane, we must compute analytical fog using a camera path starting from the plane, rather than the virtual camera origin
 	if (dot(View.GlobalClippingPlane.xyz, 1) > 0.0f)
 	{
-		float CameraOriginPlaneDistance = dot(View.GlobalClippingPlane, float4(View.WorldCameraOrigin, 1));
+		float CameraOriginPlaneDistance = dot(View.GlobalClippingPlane, float4(WorldObserverOrigin, 1));
 		float PlaneIntersectionTime = -CameraOriginPlaneDistance / dot(CameraToReceiver, View.GlobalClippingPlane.xyz);
 
 		// Only modify the start distance if the reflection plane is between the camera and receiver point
@@ -98,7 +102,7 @@ half4 GetExponentialHeightFog(float3 WorldPositionRelativeToCamera, float Exclud
 	{
 		float ExcludeIntersectionTime = ExcludeDistance * CameraToReceiverLengthInv;
 		float CameraToExclusionIntersectionZ = ExcludeIntersectionTime * CameraToReceiver.z;
-		float ExclusionIntersectionZ = View.WorldCameraOrigin.z + CameraToExclusionIntersectionZ;
+		float ExclusionIntersectionZ = WorldObserverOrigin.z + CameraToExclusionIntersectionZ;
 		float ExclusionIntersectionToReceiverZ = CameraToReceiver.z - CameraToExclusionIntersectionZ;
 
 		// Calculate fog off of the ray starting from the exclusion distance, instead of starting from the camera
diff --git a/Engine/Source/Runtime/Renderer/Private/FogRendering.cpp b/Engine/Source/Runtime/Renderer/Private/FogRendering.cpp
--- a/Engine/Source/Runtime/Renderer/Private/FogRendering.cpp
+++ b/Engine/Source/Runtime/Renderer/Private/FogRendering.cpp
@@ -283,20 +283,33 @@ void FSceneRenderer::InitFogConstants()
 			if (Scene->ExponentialFogs.Num() > 0)
 			{
 				const FExponentialHeightFogSceneInfo& FogInfo = Scene->ExponentialFogs[0];
-				const float CosTerminatorAngle = FMath::Clamp(FMath::Cos(FogInfo.LightTerminatorAngle * PI / 180.0f), -1.0f + DELTA, 1.0f - DELTA);
-				float CollapsedFogParameter[2];
+				float CollapsedFogParameter[FExponentialHeightFogSceneInfo::NumFogs];
+				static constexpr float MaxObserverHeightDifference = 65536.0f;
+				float MaxObserverHeight = FLT_MAX;
+				for (int i = 0; i < FExponentialHeightFogSceneInfo::NumFogs; i++)
+				{
+					// Only limit the observer height to fog if it has any density
+					if (FogInfo.FogData[i].Density > 0.0f)
+					{
+						MaxObserverHeight = FMath::Min(MaxObserverHeight, FogInfo.FogData[i].Height + MaxObserverHeightDifference);
+					}
+				}
+				
+				// Clamping the observer height to avoid numerical precision issues in the height fog equation. The max observer height is relative to the fog height.
+				const float ObserverHeight = FMath::Min(View.ViewMatrices.GetViewOrigin().Z, MaxObserverHeight);
+
 				for (int i = 0; i < FExponentialHeightFogSceneInfo::NumFogs; i++)
 				{
 					const float CollapsedFogParameterPower = FMath::Clamp(
-						-FogInfo.FogData[i].HeightFalloff * (View.ViewMatrices.GetViewOrigin().Z - FogInfo.FogData[i].Height),
+						-FogInfo.FogData[i].HeightFalloff * (ObserverHeight - FogInfo.FogData[i].Height),
 						-126.f + 1.f, // min and max exponent values for IEEE floating points (http://en.wikipedia.org/wiki/IEEE_floating_point)
 						+127.f - 1.f
 					);
 
 					CollapsedFogParameter[i] = FogInfo.FogData[i].Density * FMath::Pow(2.0f, CollapsedFogParameterPower);
-
 				}
-				View.ExponentialFogParameters = FVector4(CollapsedFogParameter[0], FogInfo.FogData[0].HeightFalloff, CosTerminatorAngle, FogInfo.StartDistance);
+
+				View.ExponentialFogParameters = FVector4(CollapsedFogParameter[0], FogInfo.FogData[0].HeightFalloff, MaxObserverHeight, FogInfo.StartDistance);
 				View.ExponentialFogParameters2 = FVector4(CollapsedFogParameter[1], FogInfo.FogData[1].HeightFalloff, FogInfo.FogData[1].Density, FogInfo.FogData[1].Height);
 				View.ExponentialFogColor = FVector(FogInfo.FogColor.R, FogInfo.FogColor.G, FogInfo.FogColor.B);
 				View.FogMaxOpacity = FogInfo.FogMaxOpacity;
diff --git a/Engine/Source/Runtime/Renderer/Private/SceneCore.cpp b/Engine/Source/Runtime/Renderer/Private/SceneCore.cpp
--- a/Engine/Source/Runtime/Renderer/Private/SceneCore.cpp
+++ b/Engine/Source/Runtime/Renderer/Private/SceneCore.cpp
@@ -368,7 +368,6 @@ FExponentialHeightFogSceneInfo::FExponentialHeightFogSceneInfo(const UExponentia
 	FogMaxOpacity(InComponent->FogMaxOpacity),
 	StartDistance(InComponent->StartDistance),
 	FogCutoffDistance(InComponent->FogCutoffDistance),
-	LightTerminatorAngle(0),
 	DirectionalInscatteringExponent(InComponent->DirectionalInscatteringExponent),
 	DirectionalInscatteringStartDistance(InComponent->DirectionalInscatteringStartDistance),
 	DirectionalInscatteringColor(InComponent->DirectionalInscatteringColor)
diff --git a/Engine/Source/Runtime/Renderer/Private/SceneCore.h b/Engine/Source/Runtime/Renderer/Private/SceneCore.h
--- a/Engine/Source/Runtime/Renderer/Private/SceneCore.h
+++ b/Engine/Source/Runtime/Renderer/Private/SceneCore.h
@@ -245,7 +245,6 @@ public:
 	float FogMaxOpacity;
 	float StartDistance; 
 	float FogCutoffDistance;
-	float LightTerminatorAngle;
 	FLinearColor FogColor;
 	float DirectionalInscatteringExponent; 
 	float DirectionalInscatteringStartDistance;
