When SetParentEditorOnly is called we no longer recompute shader in most case as we will call UpdateStaticPermutation that will do so.
Add Procedural Layer support for Weightmaps, including Brush support that affect Weightmap and/or Heightmap.

diff --git a/Engine/Shaders/Private/LandscapeProceduralCS.usf b/Engine/Shaders/Private/LandscapeProceduralCS.usf
new file mode 100644
--- /dev/null
+++ b/Engine/Shaders/Private/LandscapeProceduralCS.usf
@@ -0,0 +1,90 @@
+// Copyright 1998-2018 Epic Games, Inc. All Rights Reserved.
+
+#include "Common.ush"
+
+uint ComponentSize;
+
+struct FLandscapeProceduralWeightmapExtractLayersComponentData
+{
+	int2 ComponentVertexPosition;		// Section Base of the component converted to vertex
+	uint DestinationPaintLayerIndex;	// correspond to which layer info object index the data should be stored in the texture 2d array
+	uint WeightmapChannelToProcess;		// correspond to which RGBA channel to process
+	int2 AtlasTexturePositionOutput;	// This represent the location we will write layer information
+};
+
+StructuredBuffer<FLandscapeProceduralWeightmapExtractLayersComponentData> InExtractLayersComponentsData;
+Texture2D<float4> InComponentWeightMaps;
+RWTexture2DArray<float> OutAtlasPaintLayers;
+
+[numthreads(THREADGROUP_SIZEX, THREADGROUP_SIZEY, 1)]
+void ComputeWeightmapPerPaintLayer(uint3 DispatchThreadId : SV_DispatchThreadID)
+{
+	if (DispatchThreadId.x < ComponentSize && DispatchThreadId.y < ComponentSize)
+	{
+		FLandscapeProceduralWeightmapExtractLayersComponentData ComponentData = InExtractLayersComponentsData[DispatchThreadId.z];
+		int2 CurrentTexel = DispatchThreadId.xy + ComponentData.ComponentVertexPosition;
+
+		float4 LayerBlend = InComponentWeightMaps.Load(int3(CurrentTexel, 0));
+		OutAtlasPaintLayers[int3(DispatchThreadId.xy + ComponentData.AtlasTexturePositionOutput, ComponentData.DestinationPaintLayerIndex)] = LayerBlend[ComponentData.WeightmapChannelToProcess];
+	}
+}
+
+struct FLandscapeProceduralWeightmapPackLayersComponentData
+{
+	int4 ComponenVertexPositionX;	// Section Base of the component converted to vertex for each rgba value (as a texture can contain data from different component)
+	int4 ComponenVertexPositionY;	// Section Base of the component converted to vertex for each rgba value (as a texture can contain data from different component)
+	int4 SourcePaintLayerIndex;		// correspond to which layer info object index the data should be stored in the texture 2d array
+	int4 WeightmapChannelToProcess;	// correspond to which RGBA channel to process
+};
+
+StructuredBuffer<FLandscapeProceduralWeightmapPackLayersComponentData> InPackLayersComponentsData;
+Buffer<float> InWeightmapWeightBlendMode;
+Buffer<float2> InWeightmapTextureOutputOffset;
+Texture2DArray<float> InAtlasPaintLayers;
+RWTexture2D<float4> OutComponentWeightMaps;
+
+float ComputeWeightmapChannel(int ChannelIndexToProcess, uint2 TexelPosition, uint SourcePaintLayerIndex)
+{
+	float Channel = 0;
+
+	if (ChannelIndexToProcess != -1)
+	{
+		Channel = InAtlasPaintLayers.Load(int4(TexelPosition, SourcePaintLayerIndex, 0));
+		float AllLayerWeightSum = 0;
+
+		float Unused = 0.0;
+		float ArraySize = 0.0;
+		InAtlasPaintLayers.GetDimensions(Unused, Unused, ArraySize);
+
+		for (int i = 0; i < ArraySize; ++i)
+		{
+			if (InWeightmapWeightBlendMode[i] == 1.0) // Only include the one that are weight blended
+			{
+				AllLayerWeightSum += InAtlasPaintLayers.Load(int4(TexelPosition, i, 0));
+			}
+		}
+
+		Channel = AllLayerWeightSum == 0 ? Channel : clamp(Channel / AllLayerWeightSum, 0.0, 1.0);
+	}
+
+	return Channel;
+}
+
+[numthreads(THREADGROUP_SIZEX, THREADGROUP_SIZEY, 1)]
+void PackPaintLayerToWeightmap(uint3 DispatchThreadId : SV_DispatchThreadID)
+{
+	if (DispatchThreadId.x < ComponentSize && DispatchThreadId.y < ComponentSize)
+	{
+		FLandscapeProceduralWeightmapPackLayersComponentData ComponentData = InPackLayersComponentsData[DispatchThreadId.z];
+
+		float4 Channels;
+
+		for (int i = 0; i < 4; ++i)
+		{
+			Channels[i] = ComputeWeightmapChannel(ComponentData.WeightmapChannelToProcess[i], DispatchThreadId.xy + int2(ComponentData.ComponenVertexPositionX[i], ComponentData.ComponenVertexPositionY[i]), ComponentData.SourcePaintLayerIndex[i]);
+		}
+
+		float2 TextureOutputOffset = InWeightmapTextureOutputOffset[DispatchThreadId.z];
+		OutComponentWeightMaps[DispatchThreadId.xy + TextureOutputOffset] = Channels;
+	}
+}
\ No newline at end of file
diff --git a/Engine/Shaders/Private/LandscapeProceduralPS.usf b/Engine/Shaders/Private/LandscapeProceduralPS.usf
--- a/Engine/Shaders/Private/LandscapeProceduralPS.usf
+++ b/Engine/Shaders/Private/LandscapeProceduralPS.usf
@@ -2,39 +2,52 @@
 
 #include "Common.ush"
 
-Texture2D ReadHeightmapTexture1;
-Texture2D ReadHeightmapTexture2;
-SamplerState ReadHeightmapTexture1Sampler;
-SamplerState ReadHeightmapTexture2Sampler;
-float2 LayerInfo; // x == weight, y == visibility
-float4 OutputConfig; // x == ApplyLayerModifiers, y == OutputAsDelta, z == Use ReadHeightmapTexture2 and out Delta with ReadHeightmapTexture1, w == Output Normals
-float2 HeightmapTextureSize; // x == source heightmap width, y == source heightmap height
-float3 LandscapeGridScale; // x == LS Actor DrawScale.X, y == LS Actor DrawScale.y, z == LS Actor DrawScale.z / 128.0f (ZSCALE)
+Texture2D ReadTexture1;
+Texture2D ReadTexture2;					// This one is optional, it's valid when OutputConfig.z == 1
+SamplerState ReadTexture1Sampler;		
+SamplerState ReadTexture2Sampler;
+float2 LayerInfo;					// x == weight, y == visibility
+float4 OutputConfig;					// x == ApplyLayerModifiers, y == OutputAsDelta or Output weightmap as substractive/additive, z == Use ReadTexture2 and out Delta with ReadTexture1, w == Output Normals or Output weightmap as normalized weight blend
+float2 TextureSize;					// x == read texture 1 width, y == read texture 1 height
+float3 LandscapeGridScale;				// x == LS Actor DrawScale.X, y == LS Actor DrawScale.y, z == LS Actor DrawScale.z / 128.0f (ZSCALE)
 float CurrentMipComponentVertexCount;
+float2 CurrentMipTextureSize;
+float2 ParentMipTextureSize;
+
+float ExtractHeight(float2 InPackedHeight)
+{
+	return float(((int)round(InPackedHeight.r * 255.0) << 8) | (int)round(InPackedHeight.g * 255.0));
+}
+
+float2 PackHeight(float InExtractedHeight)
+{
+	InExtractedHeight = clamp(InExtractedHeight, 0.0f, 65536.0f);
+	int iHeight = (int)InExtractedHeight;
+	return float2((float)((iHeight - (iHeight & 255)) >> 8) / 255.0, (float)(iHeight & 255) / 255.0);
+}
 
 float3 SampleTextureYNormal(float2 InTextureCoordinates, float2 InTexelSize, bool InUp)
 {
 	float2 SamplingUV = InTextureCoordinates;
 	SamplingUV.y = InUp ? SamplingUV.y - InTexelSize.y : SamplingUV.y + InTexelSize.y;
-	float2 PosSample = ReadHeightmapTexture1.SampleLevel(ReadHeightmapTexture1Sampler, SamplingUV, 0).rg;
-	float PosHeight = float(((int)round(PosSample.r * 255.0) << 8) | (int)round(PosSample.g * 255.0)) - 32768.0;
-	return float3(SamplingUV * HeightmapTextureSize * LandscapeGridScale.xy, PosHeight * LandscapeGridScale.z * TERRAIN_ZSCALE);
+	float2 PosSample = ReadTexture1.SampleLevel(ReadTexture1Sampler, SamplingUV, 0).rg;
+	float PosHeight = ExtractHeight(PosSample) - 32768.0;
+	return float3(SamplingUV * TextureSize * LandscapeGridScale.xy, PosHeight * LandscapeGridScale.z * TERRAIN_ZSCALE);
 }
 
 float3 SampleTextureXNormal(float2 InTextureCoordinates, float2 InTexelSize, bool InLeft)
 {
 	float2 SamplingUV = InTextureCoordinates;
 	SamplingUV.x = InLeft ? SamplingUV.x - InTexelSize.x : SamplingUV.x + InTexelSize.x;
-	float2 PosSample = ReadHeightmapTexture1.SampleLevel(ReadHeightmapTexture1Sampler, SamplingUV, 0).rg;
-	float PosHeight = float(((int)round(PosSample.r * 255.0) << 8) | (int)round(PosSample.g * 255.0)) - 32768.0;
-	return float3(SamplingUV * HeightmapTextureSize * LandscapeGridScale.xy, PosHeight * LandscapeGridScale.z * TERRAIN_ZSCALE);
+	float2 PosSample = ReadTexture1.SampleLevel(ReadTexture1Sampler, SamplingUV, 0).rg;
+	float PosHeight = ExtractHeight(PosSample) - 32768.0;
+	return float3(SamplingUV * TextureSize * LandscapeGridScale.xy, PosHeight * LandscapeGridScale.z * TERRAIN_ZSCALE);
 }
 
-void PSMain(float2 InTextureCoordinates : TEXCOORD0,
-			out float4 OutColor : SV_Target0)
+void PSHeightmapMain(float2 InTextureCoordinates : TEXCOORD0, out float4 OutColor : SV_Target0)
 {
-	float4 SourceColor = ReadHeightmapTexture1.Sample(ReadHeightmapTexture1Sampler, InTextureCoordinates);
-	float Height = float(((int)round(SourceColor.r * 255.0) << 8) | (int)round(SourceColor.g * 255.0));
+	float4 SourceColor = ReadTexture1.Sample(ReadTexture1Sampler, InTextureCoordinates);
+	float Height = ExtractHeight(SourceColor.rg);
 
 	// Perform calculation 0 based
 	Height -= 32768.0f;
@@ -48,9 +61,8 @@ void PSMain(float2 InTextureCoordinates : TEXCOORD0,
 	// Output using 2nd heightmap
 	if (OutputConfig.z == 1.0)
 	{
-		float2 BaseColor = ReadHeightmapTexture2.Sample(ReadHeightmapTexture2Sampler, InTextureCoordinates).rg;
-		float BaseHeight = float(((int)round(BaseColor.r * 255.0) << 8) | (int)round(BaseColor.g * 255.0));
-
+		float2 BaseColor = ReadTexture2.Sample(ReadTexture2Sampler, InTextureCoordinates).rg;
+		float BaseHeight = ExtractHeight(BaseColor.rg);
 		Height = OutputConfig.y == 1 ? BaseHeight + (Height - BaseHeight) : BaseHeight + Height;
 	}
 	else
@@ -62,9 +74,7 @@ void PSMain(float2 InTextureCoordinates : TEXCOORD0,
 		}
 	}
 
-	Height = clamp(Height, 0.0f, 65536.0f);
-	int iHeight = (int)Height;
-	float2 PackedHeight = float2((float)((iHeight - (iHeight & 255)) >> 8) / 255.0, (float)(iHeight & 255) / 255.0);
+	float2 PackedHeight = PackHeight(Height);
 
 	OutColor.r = PackedHeight.x;
 	OutColor.g = PackedHeight.y;
@@ -74,14 +84,14 @@ void PSMain(float2 InTextureCoordinates : TEXCOORD0,
 	// Output normals
 	if (OutputConfig.w == 1.0)
 	{
-		float2 TexelSize = 1.0 / HeightmapTextureSize;
+		float2 TexelSize = 1.0 / TextureSize;
 
 		bool IsMinBorderTexelX = InTextureCoordinates.x <= TexelSize.x;
 		bool IsMinBorderTexelY = InTextureCoordinates.y <= TexelSize.y;
 		bool IsMaxBorderTexelX = InTextureCoordinates.x >= TexelSize.x * min(CurrentMipComponentVertexCount - 1, 1.0);
 		bool IsMaxBorderTexelY = InTextureCoordinates.y >= TexelSize.y * min(CurrentMipComponentVertexCount - 1, 1.0);
 
-		float3 CurrentPos = float3(InTextureCoordinates * HeightmapTextureSize * LandscapeGridScale.xy, (Height - 32768.0f) * LandscapeGridScale.z * TERRAIN_ZSCALE);
+		float3 CurrentPos = float3(InTextureCoordinates * TextureSize * LandscapeGridScale.xy, (Height - 32768.0f) * LandscapeGridScale.z * TERRAIN_ZSCALE);
 		float3 FinalNormal = 0.0;
 
 		if (IsMinBorderTexelX) // left border
@@ -179,11 +189,7 @@ void PSMain(float2 InTextureCoordinates : TEXCOORD0,
 	}
 }
 
-float2 CurrentMipTextureSize;
-float2 ParentMipTextureSize;
-
-void PSMainMips(float2 InTextureCoordinates : TEXCOORD0,
-				out float4 OutColor : SV_Target0)
+void PSHeightmapMainMips(float2 InTextureCoordinates : TEXCOORD0, out float4 OutColor : SV_Target0)
 {
 	bool IsMinBorderTexelX = false;
 	bool IsMinBorderTexelY = false;
@@ -219,16 +225,17 @@ void PSMainMips(float2 InTextureCoordinates : TEXCOORD0,
 
 	float2 ParentMipTexelSize = 1.0 / ParentMipTextureSize;
 	float2 SourceUV = InTextureCoordinates - (ParentMipTexelSize * 0.5);
-	float2 HeightResult = 0;
+	float HeightResult = 0;
 	float2 NormalResult = 0;
+	float Alpha = 0.5;
 	OutColor = 0; // Default to 0
 
 	if (IsMinBorderTexelX) // on left border
 	{
-		float4 SourceSample = ReadHeightmapTexture1.SampleLevel(ReadHeightmapTexture1Sampler, SourceUV, 0);
-		float4 DownSample = ReadHeightmapTexture1.SampleLevel(ReadHeightmapTexture1Sampler, SourceUV + float2(0.0, ParentMipTexelSize.y), 0);
-		float SourceHeight = clamp(float(((int)round(SourceSample.r * 255.0) << 8) | (int)round(SourceSample.g * 255.0)), 0.0, 65536.0);
-		float DownHeight = clamp(float(((int)round(DownSample.r * 255.0) << 8) | (int)round(DownSample.g * 255.0)), 0.0, 65536.0);
+		float4 SourceSample = ReadTexture1.SampleLevel(ReadTexture1Sampler, SourceUV, 0);
+		float4 DownSample = ReadTexture1.SampleLevel(ReadTexture1Sampler, SourceUV + float2(0.0, ParentMipTexelSize.y), 0);
+		float SourceHeight = ExtractHeight(SourceSample.rg);
+		float DownHeight = ExtractHeight(DownSample.rg);
 
 		if (IsMinBorderTexelY) // on top border
 		{
@@ -242,18 +249,17 @@ void PSMainMips(float2 InTextureCoordinates : TEXCOORD0,
 		}
 		else
 		{
-			float FracTopDown = 0.5;
-			HeightResult = clamp(round(lerp(SourceHeight, DownHeight, FracTopDown)), 0.0, 65536.0);
-			NormalResult.x = lerp(SourceSample.b, DownSample.b, FracTopDown);
-			NormalResult.y = lerp(SourceSample.a, DownSample.a, FracTopDown);
+			HeightResult = clamp(round(lerp(SourceHeight, DownHeight, Alpha)), 0.0, 65536.0);
+			NormalResult.x = lerp(SourceSample.b, DownSample.b, Alpha);
+			NormalResult.y = lerp(SourceSample.a, DownSample.a, Alpha);
 		}
 	}
 	else if (IsMaxBorderTexelX) // on right border
 	{
-		float4 RightSample = ReadHeightmapTexture1.SampleLevel(ReadHeightmapTexture1Sampler, SourceUV + float2(ParentMipTexelSize.x, 0.0), 0);
-		float4 DownRightSample = ReadHeightmapTexture1.SampleLevel(ReadHeightmapTexture1Sampler, SourceUV + ParentMipTexelSize, 0);
-		float RightHeight = clamp(float(((int)round(RightSample.r * 255.0) << 8) | (int)round(RightSample.g * 255.0)), 0.0, 65536.0);
-		float DownRightHeight = clamp(float(((int)round(DownRightSample.r * 255.0) << 8) | (int)round(DownRightSample.g * 255.0)), 0.0, 65536.0);
+		float4 RightSample = ReadTexture1.SampleLevel(ReadTexture1Sampler, SourceUV + float2(ParentMipTexelSize.x, 0.0), 0);
+		float4 DownRightSample = ReadTexture1.SampleLevel(ReadTexture1Sampler, SourceUV + ParentMipTexelSize, 0);
+		float RightHeight = ExtractHeight(RightSample.rg);
+		float DownRightHeight = ExtractHeight(DownRightSample.rg);
 
 		if (IsMinBorderTexelY) // on top border
 		{
@@ -267,57 +273,133 @@ void PSMainMips(float2 InTextureCoordinates : TEXCOORD0,
 		}
 		else
 		{
-			float FracTopDown = 0.5;
-			HeightResult = clamp(round(lerp(RightHeight, DownRightHeight, FracTopDown)), 0.0, 65536.0);;
-			NormalResult.x = lerp(RightSample.b, DownRightSample.b, FracTopDown);
-			NormalResult.y = lerp(RightSample.a, DownRightSample.a, FracTopDown);
+			HeightResult = clamp(round(lerp(RightHeight, DownRightHeight, Alpha)), 0.0, 65536.0);;
+			NormalResult.x = lerp(RightSample.b, DownRightSample.b, Alpha);
+			NormalResult.y = lerp(RightSample.a, DownRightSample.a, Alpha);
 		}
 	}
 	else // center texel
 	{
-		float4 SourceSample = ReadHeightmapTexture1.SampleLevel(ReadHeightmapTexture1Sampler, SourceUV, 0);
-		float4 DownSample = ReadHeightmapTexture1.SampleLevel(ReadHeightmapTexture1Sampler, SourceUV + float2(0.0, ParentMipTexelSize.y), 0);
-		float4 RightSample = ReadHeightmapTexture1.SampleLevel(ReadHeightmapTexture1Sampler, SourceUV + float2(ParentMipTexelSize.x, 0.0), 0);
-		float4 DownRightSample = ReadHeightmapTexture1.SampleLevel(ReadHeightmapTexture1Sampler, SourceUV + ParentMipTexelSize, 0);
+		float4 SourceSample = ReadTexture1.SampleLevel(ReadTexture1Sampler, SourceUV, 0);
+		float4 DownSample = ReadTexture1.SampleLevel(ReadTexture1Sampler, SourceUV + float2(0.0, ParentMipTexelSize.y), 0);
+		float4 RightSample = ReadTexture1.SampleLevel(ReadTexture1Sampler, SourceUV + float2(ParentMipTexelSize.x, 0.0), 0);
+		float4 DownRightSample = ReadTexture1.SampleLevel(ReadTexture1Sampler, SourceUV + ParentMipTexelSize, 0);
 
-		float SourceHeight = clamp(float(((int)round(SourceSample.r * 255.0) << 8) | (int)round(SourceSample.g * 255.0)), 0.0, 65536.0);
-		float DownHeight = clamp(float(((int)round(DownSample.r * 255.0) << 8) | (int)round(DownSample.g * 255.0)), 0.0, 65536.0);
-		float RightHeight = clamp(float(((int)round(RightSample.r * 255.0) << 8) | (int)round(RightSample.g * 255.0)), 0.0, 65536.0);
-		float DownRightHeight = clamp(float(((int)round(DownRightSample.r * 255.0) << 8) | (int)round(DownRightSample.g * 255.0)), 0.0, 65536.0);
+		float SourceHeight = ExtractHeight(SourceSample.rg);
+		float DownHeight = ExtractHeight(DownSample.rg);
+		float RightHeight = ExtractHeight(RightSample.rg);
+		float DownRightHeight = ExtractHeight(DownRightSample.rg);
 
 		if (IsMinBorderTexelY) // on top border
 		{
-			float FracSourceToRight = 0.5;
-			HeightResult = clamp(round(lerp(SourceHeight, RightHeight, FracSourceToRight)), 0.0, 65536.0);
-			NormalResult.x = lerp(SourceSample.b, RightSample.b, FracSourceToRight);
-			NormalResult.y = lerp(SourceSample.a, RightSample.a, FracSourceToRight);
+			HeightResult = clamp(round(lerp(SourceHeight, RightHeight, Alpha)), 0.0, 65536.0);
+			NormalResult.x = lerp(SourceSample.b, RightSample.b, Alpha);
+			NormalResult.y = lerp(SourceSample.a, RightSample.a, Alpha);
 		}
 		else if (IsMaxBorderTexelY) // on bottom border
 		{
-			float FracDownToDownRight = 0.5;
-			HeightResult = clamp(round(lerp(DownHeight, DownRightHeight, FracDownToDownRight)), 0.0, 65536.0);
-			NormalResult.x = lerp(DownSample.b, DownRightSample.b, FracDownToDownRight);
-			NormalResult.y = lerp(DownSample.a, DownRightSample.a, FracDownToDownRight);
+			HeightResult = clamp(round(lerp(DownHeight, DownRightHeight, Alpha)), 0.0, 65536.0);
+			NormalResult.x = lerp(DownSample.b, DownRightSample.b, Alpha);
+			NormalResult.y = lerp(DownSample.a, DownRightSample.a, Alpha);
 		}
 		else
 		{
-			float FracSourceToRight = 0.5;
-			float FracDownToDownRight = 0.5;
+			float SourceToRightHeight = lerp(SourceHeight, RightHeight, Alpha);
+			float DownToDownRightHeight = lerp(DownHeight, DownRightHeight, Alpha);
 
-			float SourceToRightHeight = lerp(SourceHeight, RightHeight, 0.5);
-			float DownToDownRightHeight = lerp(DownHeight, DownRightHeight, 0.5);
+			HeightResult = clamp(round(lerp(SourceToRightHeight, DownToDownRightHeight, Alpha)), 0.0, 65536.0);
+			NormalResult.x = lerp(lerp(SourceSample.b, RightSample.b, Alpha), lerp(DownSample.b, DownRightSample.b, Alpha), Alpha);
+			NormalResult.y = lerp(lerp(SourceSample.a, RightSample.a, Alpha), lerp(DownSample.a, DownRightSample.a, Alpha), Alpha);
+		}
+	}	
 
-			float SourceToRightDeltaHeight = RightHeight - 32768.0;
-			float DownToDownRightDeltaHeight = DownRightHeight - 32768.0;
+	float2 PackedHeightResult = PackHeight(HeightResult);
+	OutColor = float4(PackedHeightResult.xy, NormalResult.xy);
+}
 
-			float FracTopDown = 0.5;
+float4 WeightmapLayerWeightBlend; // each xyzw == 1 texture channel
+
+void PSWeightmapMain(float2 InTextureCoordinates : TEXCOORD0, out float4 OutColor : SV_Target0)
+{
+	float4 SourceColor = ReadTexture1.Sample(ReadTexture1Sampler, InTextureCoordinates);
+
+	// Output as Layer, so apply Layer info
+	SourceColor = OutputConfig.x == 1.0 ? SourceColor * LayerInfo.x * LayerInfo.y : SourceColor;
 
-			HeightResult = clamp(round(lerp(SourceToRightHeight, DownToDownRightHeight, 0.5)), 0.0, 65536.0);
-			NormalResult.x = lerp(lerp(SourceSample.b, RightSample.b, FracSourceToRight), lerp(DownSample.b, DownRightSample.b, FracDownToDownRight), FracTopDown);
-			NormalResult.y = lerp(lerp(SourceSample.a, RightSample.a, FracSourceToRight), lerp(DownSample.a, DownRightSample.a, FracDownToDownRight), FracTopDown);
+	// Output using 2nd heightmap
+	if (OutputConfig.z == 1.0)
+	{
+		float4 BaseColor = ReadTexture2.Sample(ReadTexture2Sampler, InTextureCoordinates);
+
+		// Apply addtive/substractive logic
+		SourceColor = OutputConfig.y == 0.0f ? clamp(BaseColor + SourceColor, 0.0, 1.0) : clamp(BaseColor - SourceColor, 0.0, 1.0);
+	}	
+
+	OutColor = SourceColor;
+}
+
+void PSWeightmapMainMips(float2 InTextureCoordinates : TEXCOORD0, out float4 OutColor : SV_Target0)
+{
+	bool IsMinBorderTexelX = false;
+	bool IsMinBorderTexelY = false;
+	bool IsMaxBorderTexelX = false;
+	bool IsMaxBorderTexelY = false;
+
+	// Special case of 1 texel size component
+	if (CurrentMipComponentVertexCount == 1)
+	{		
+		if (CurrentMipTextureSize.x >= CurrentMipTextureSize.y) // x size 1 texel
+		{
+			IsMinBorderTexelY = true;
+			IsMinBorderTexelX = (InTextureCoordinates.x * CurrentMipTextureSize.x) <= 1.0;
+			IsMaxBorderTexelX = (InTextureCoordinates.x * CurrentMipTextureSize.x) >= CurrentMipTextureSize.x - 1.0;
+		}
+		else
+		{
+			IsMinBorderTexelX = true;
+			IsMinBorderTexelY = (InTextureCoordinates.y * CurrentMipTextureSize.y) <= 1.0;
+			IsMaxBorderTexelY = (InTextureCoordinates.y * CurrentMipTextureSize.y) >= CurrentMipTextureSize.y - 1.0;
 		}
 	}
+	else
+	{
+		float2 CurrentMipQuadTexelSize = 1.0 / CurrentMipComponentVertexCount;
+		float2 LandscapeQuadUV = frac(InTextureCoordinates * CurrentMipTextureSize / CurrentMipComponentVertexCount);
+
+		IsMinBorderTexelX = LandscapeQuadUV.x <= CurrentMipQuadTexelSize.x;
+		IsMinBorderTexelY = LandscapeQuadUV.y <= CurrentMipQuadTexelSize.y;
+		IsMaxBorderTexelX = LandscapeQuadUV.x >= (CurrentMipQuadTexelSize.x * max(CurrentMipComponentVertexCount - 1, 1.0));
+		IsMaxBorderTexelY = LandscapeQuadUV.y >= (CurrentMipQuadTexelSize.y * max(CurrentMipComponentVertexCount - 1, 1.0));
+	}
 
-	float2 PackedHeightResult = float2((float)(((int)HeightResult - ((int)HeightResult & 255)) >> 8) / 255.0, (float)((int)HeightResult & 255) / 255.0);
-	OutColor = float4(PackedHeightResult.x, PackedHeightResult.y, NormalResult.x, NormalResult.y);
+	float2 ParentMipTexelSize = 1.0 / ParentMipTextureSize;
+	float2 SourceUV = InTextureCoordinates - (ParentMipTexelSize * 0.5);
+	float2 HeightResult = 0;
+	float2 NormalResult = 0;
+	float Alpha = 0.5;
+	OutColor = 0; // Default to 0
+
+	if (IsMinBorderTexelX) // on left border
+	{
+		float4 SourceSample = ReadTexture1.SampleLevel(ReadTexture1Sampler, SourceUV, 0);
+		float4 DownSample = ReadTexture1.SampleLevel(ReadTexture1Sampler, SourceUV + float2(0.0, ParentMipTexelSize.y), 0);
+
+		OutColor = IsMinBorderTexelY ? SourceSample : (IsMaxBorderTexelY ? DownSample : lerp(SourceSample, DownSample, Alpha));		
+	}
+	else if (IsMaxBorderTexelX) // on right border
+	{
+		float4 RightSample = ReadTexture1.SampleLevel(ReadTexture1Sampler, SourceUV + float2(ParentMipTexelSize.x, 0.0), 0);
+		float4 DownRightSample = ReadTexture1.SampleLevel(ReadTexture1Sampler, SourceUV + ParentMipTexelSize, 0);
+
+		OutColor = IsMinBorderTexelY ? RightSample : (IsMaxBorderTexelY ? DownRightSample : lerp(RightSample, DownRightSample, Alpha));
+	}
+	else // center texel
+	{
+		float4 SourceSample = ReadTexture1.SampleLevel(ReadTexture1Sampler, SourceUV, 0);
+		float4 DownSample = ReadTexture1.SampleLevel(ReadTexture1Sampler, SourceUV + float2(0.0, ParentMipTexelSize.y), 0);
+		float4 RightSample = ReadTexture1.SampleLevel(ReadTexture1Sampler, SourceUV + float2(ParentMipTexelSize.x, 0.0), 0);
+		float4 DownRightSample = ReadTexture1.SampleLevel(ReadTexture1Sampler, SourceUV + ParentMipTexelSize, 0);
+
+		OutColor = IsMinBorderTexelY ? lerp(SourceSample, RightSample, Alpha) : (IsMaxBorderTexelY ? lerp(DownSample, DownRightSample, Alpha) : lerp(lerp(SourceSample, RightSample, Alpha), lerp(DownSample, DownRightSample, Alpha), Alpha));
+	}
 }
\ No newline at end of file
diff --git a/Engine/Source/Editor/LandscapeEditor/Private/LandscapeEdMode.cpp b/Engine/Source/Editor/LandscapeEditor/Private/LandscapeEdMode.cpp
--- a/Engine/Source/Editor/LandscapeEditor/Private/LandscapeEdMode.cpp
+++ b/Engine/Source/Editor/LandscapeEditor/Private/LandscapeEdMode.cpp
@@ -1897,7 +1897,9 @@ bool FEdModeLandscape::InputKey(FEditorViewportClient* ViewportClient, FViewport
 					FVector HitLocation;
 					if (LandscapeMouseTrace(ViewportClient, HitLocation))
 					{
-						if (CurrentToolTarget.TargetType == ELandscapeToolTargetType::Weightmap && CurrentToolTarget.LayerInfo == NULL)
+						const FName CurrentToolName = CurrentTool->GetToolName();
+
+						if (CurrentToolTarget.TargetType == ELandscapeToolTargetType::Weightmap && CurrentToolTarget.LayerInfo == NULL && CurrentToolName != TEXT("BPCustom"))
 						{
 							FMessageDialog::Open(EAppMsgType::Ok,
 								NSLOCTEXT("UnrealEd", "LandscapeNeedToCreateLayerInfo", "This layer has no layer info assigned yet. You must create or assign a layer info before you can paint this layer."));
@@ -3567,6 +3569,11 @@ void FEdModeLandscape::ImportData(const FLandscapeTargetListInfo& TargetInfo, co
 				return;
 			}
 
+			if (GetMutableDefault<UEditorExperimentalSettings>()->bProceduralLandscape)
+			{
+				ChangeWeightmapsToCurrentProceduralLayerWeightmaps();
+			}
+
 			TArray<uint8> Data;
 			if (ImportResolution != LandscapeResolution)
 			{
@@ -3591,6 +3598,14 @@ void FEdModeLandscape::ImportData(const FLandscapeTargetListInfo& TargetInfo, co
 
 			FAlphamapAccessor<false, false> AlphamapAccessor(LandscapeInfo, TargetInfo.LayerInfoObj.Get());
 			AlphamapAccessor.SetData(MinX, MinY, MaxX, MaxY, Data.GetData(), ELandscapeLayerPaintingRestriction::None);
+
+			if (GetMutableDefault<UEditorExperimentalSettings>()->bProceduralLandscape)
+			{
+				ChangeWeightmapsToCurrentProceduralLayerWeightmaps(true);
+
+				check(CurrentToolTarget.LandscapeInfo->LandscapeActor.IsValid());
+				CurrentToolTarget.LandscapeInfo->LandscapeActor.Get()->RequestProceduralContentUpdate(EProceduralContentUpdateFlag::Weightmap_All);
+			}
 		}
 	}
 }
@@ -3720,12 +3735,14 @@ void FEdModeLandscape::DeleteLandscapeComponents(ULandscapeInfo* LandscapeInfo,
 			HeightmapTexture->ClearFlags(RF_Standalone); // Remove when there is no reference for this Heightmap...
 		}
 
-		for (int32 i = 0; i < Component->WeightmapTextures.Num(); ++i)
+		TArray<UTexture2D*>& ComponentWeightmapTextures = Component->GetWeightmapTextures();
+
+		for (UTexture2D* WeightmapTexture : ComponentWeightmapTextures)
 		{
-			Component->WeightmapTextures[i]->SetFlags(RF_Transactional);
-			Component->WeightmapTextures[i]->Modify();
-			Component->WeightmapTextures[i]->MarkPackageDirty();
-			Component->WeightmapTextures[i]->ClearFlags(RF_Standalone);
+			WeightmapTexture->SetFlags(RF_Transactional);
+			WeightmapTexture->Modify();
+			WeightmapTexture->MarkPackageDirty();
+			WeightmapTexture->ClearFlags(RF_Standalone);
 		}
 
 		if (Component->XYOffsetmapTexture)
@@ -3928,6 +3945,22 @@ ALandscape* FEdModeLandscape::ChangeComponentSetting(int32 NumComponentsX, int32
 			Landscape->bUsedForNavigation = OldLandscapeProxy->bUsedForNavigation;
 			Landscape->MaxPaintedLayersPerComponent = OldLandscapeProxy->MaxPaintedLayersPerComponent;
 
+			if (GetMutableDefault<UEditorExperimentalSettings>()->bProceduralLandscape)
+			{
+				ALandscape* OldLandscapeActor = Cast<ALandscape>(OldLandscapeProxy);
+				ALandscape* NewLandscapeActor = Cast<ALandscape>(Landscape);
+
+				NewLandscapeActor->PreviousExperimentalLandscapeProcedural = OldLandscapeActor->PreviousExperimentalLandscapeProcedural;
+
+				// Moves specific config including brushes to new landscape actor
+				for (const FProceduralLayer& Layer : OldLandscapeActor->ProceduralLayers)
+				{
+					FProceduralLayer* NewLayer = NewLandscapeActor->ProceduralLayers.FindByPredicate( [&](const FProceduralLayer& NewLayer)	{ return NewLayer.Name == Layer.Name; }	);
+					check(NewLayer != nullptr);
+					*NewLayer = Layer;
+				}
+			}
+
 			Landscape->CreateLandscapeInfo();
 
 			// Clone landscape splines
@@ -4163,7 +4196,7 @@ void FEdModeLandscape::SetProceduralLayerName(int32 InLayerIndex, const FName& I
 	Landscape->ProceduralLayers[InLayerIndex].Name = InName;
 }
 
-float FEdModeLandscape::GetProceduralLayerWeight(int32 InLayerIndex) const
+float FEdModeLandscape::GetProceduralLayerAlpha(int32 InLayerIndex) const
 {
 	if (CurrentToolTarget.LandscapeInfo == nullptr)
 	{
@@ -4177,10 +4210,10 @@ float FEdModeLandscape::GetProceduralLayerWeight(int32 InLayerIndex) const
 		return 1.0f;
 	}
 
-	return Landscape->ProceduralLayers[InLayerIndex].Weight;
+	return CurrentToolTarget.TargetType == ELandscapeToolTargetType::Heightmap ? Landscape->ProceduralLayers[InLayerIndex].HeightmapAlpha : Landscape->ProceduralLayers[InLayerIndex].WeightmapAlpha;
 }
 
-void FEdModeLandscape::SetProceduralLayerWeight(float InWeight, int32 InLayerIndex)
+void FEdModeLandscape::SetProceduralLayerAlpha(float InWeight, int32 InLayerIndex)
 {
 	if (CurrentToolTarget.LandscapeInfo == nullptr)
 	{
@@ -4194,9 +4227,17 @@ void FEdModeLandscape::SetProceduralLayerWeight(float InWeight, int32 InLayerInd
 		return;
 	}
 
-	Landscape->ProceduralLayers[InLayerIndex].Weight = InWeight;
+	if (CurrentToolTarget.TargetType == ELandscapeToolTargetType::Heightmap)
+	{
+		Landscape->ProceduralLayers[InLayerIndex].HeightmapAlpha = InWeight;
+	}
+	else
+	{
+		Landscape->ProceduralLayers[InLayerIndex].WeightmapAlpha = InWeight;
+	}
+
 
-	Landscape->RequestProceduralContentUpdate(EProceduralContentUpdateFlag::Heightmap_All);
+	Landscape->RequestProceduralContentUpdate(EProceduralContentUpdateFlag::All, true);
 }
 
 void FEdModeLandscape::SetProceduralLayerVisibility(bool InVisible, int32 InLayerIndex)
@@ -4215,7 +4256,7 @@ void FEdModeLandscape::SetProceduralLayerVisibility(bool InVisible, int32 InLaye
 
 	Landscape->ProceduralLayers[InLayerIndex].Visible = InVisible;
 
-	Landscape->RequestProceduralContentUpdate(EProceduralContentUpdateFlag::Heightmap_All);
+	Landscape->RequestProceduralContentUpdate(EProceduralContentUpdateFlag::All, true);
 }
 
 bool FEdModeLandscape::IsProceduralLayerVisible(int32 InLayerIndex) const
@@ -4264,10 +4305,10 @@ void FEdModeLandscape::AddBrushToCurrentProceduralLayer(int32 InTargetType, ALan
 
 	InBrush->SetOwningLandscape(Landscape);
 
-	Landscape->RequestProceduralContentUpdate(InTargetType == ELandscapeToolTargetType::Type::Heightmap ? EProceduralContentUpdateFlag::Heightmap_All : EProceduralContentUpdateFlag::Weightmap_All);
+	Landscape->RequestProceduralContentUpdate(EProceduralContentUpdateFlag::All_Render);
 }
 
-void FEdModeLandscape::RequestProceduralContentUpdate()
+void FEdModeLandscape::RequestProceduralContentUpdate(bool InUpdateAllMaterials)
 {
 	if (CurrentToolTarget.LandscapeInfo == nullptr || !CurrentToolTarget.LandscapeInfo->LandscapeActor.IsValid())
 	{
@@ -4275,8 +4316,7 @@ void FEdModeLandscape::RequestProceduralContentUpdate()
 	}
 
 	ALandscape* Landscape = CurrentToolTarget.LandscapeInfo->LandscapeActor.Get();
-
-	Landscape->RequestProceduralContentUpdate(CurrentToolTarget.TargetType == ELandscapeToolTargetType::Type::Heightmap ? EProceduralContentUpdateFlag::Heightmap_All : EProceduralContentUpdateFlag::Weightmap_All);	
+	Landscape->RequestProceduralContentUpdate(EProceduralContentUpdateFlag::All_Render, InUpdateAllMaterials);
 }
 
 void FEdModeLandscape::RemoveBrushFromCurrentProceduralLayer(int32 InTargetType, ALandscapeBlueprintCustomBrush* InBrush)
@@ -4288,72 +4328,64 @@ void FEdModeLandscape::RemoveBrushFromCurrentProceduralLayer(int32 InTargetType,
 
 	ALandscape* Landscape = CurrentToolTarget.LandscapeInfo->LandscapeActor.Get();
 
-	FProceduralLayer* Layer = GetCurrentProceduralLayer();
-
-	if (Layer == nullptr)
-	{
-		return;
-	}
-
-	int32 IndexToRemove = INDEX_NONE;
-	for (int32 i = 0; i < Layer->Brushes.Num(); ++i)
+	for (FProceduralLayer& Layer : Landscape->ProceduralLayers)
 	{
-		if (Layer->Brushes[i].BPCustomBrush == InBrush)
+		int32 IndexToRemove = INDEX_NONE;
+		for (int32 i = 0; i < Layer.Brushes.Num(); ++i)
 		{
-			IndexToRemove = i;
-			break;
+			if (Layer.Brushes[i].BPCustomBrush == InBrush)
+			{
+				IndexToRemove = i;
+				break;
+			}
 		}
-	}
 
-	if (IndexToRemove != INDEX_NONE)
-	{
-		Layer->Brushes.RemoveAt(IndexToRemove);
-
-		if (InTargetType == ELandscapeToolTargetType::Type::Heightmap)
+		if (IndexToRemove != INDEX_NONE)
 		{
-			for (int32 i = 0; i < Layer->HeightmapBrushOrderIndices.Num(); ++i)
+			Layer.Brushes.RemoveAt(IndexToRemove);
+
+			for (int32 i = 0; i < Layer.HeightmapBrushOrderIndices.Num(); ++i)
 			{
-				if (Layer->HeightmapBrushOrderIndices[i] == IndexToRemove)
+				if (Layer.HeightmapBrushOrderIndices[i] == IndexToRemove)
 				{
 					// Update the value of the index of all the one after the one we removed, so index still correctly match actual brushes list
-					for (int32 j = 0; j < Layer->HeightmapBrushOrderIndices.Num(); ++j)
+					for (int32 j = 0; j < Layer.HeightmapBrushOrderIndices.Num(); ++j)
 					{
-						if (Layer->HeightmapBrushOrderIndices[j] > IndexToRemove)
+						if (Layer.HeightmapBrushOrderIndices[j] > IndexToRemove)
 						{
-							--Layer->HeightmapBrushOrderIndices[j];
+							--Layer.HeightmapBrushOrderIndices[j];
 						}
 					}
 
-					Layer->HeightmapBrushOrderIndices.RemoveAt(i);
+					Layer.HeightmapBrushOrderIndices.RemoveAt(i);
 					break;
 				}
 			}
-		}
-		else
-		{
-			for (int32 i = 0; i < Layer->WeightmapBrushOrderIndices.Num(); ++i)
+
+			for (int32 i = 0; i < Layer.WeightmapBrushOrderIndices.Num(); ++i)
 			{
-				if (Layer->WeightmapBrushOrderIndices[i] == IndexToRemove)
+				if (Layer.WeightmapBrushOrderIndices[i] == IndexToRemove)
 				{
 					// Update the value of the index of all the one after the one we removed, so index still correctly match actual brushes list
-					for (int32 j = 0; j < Layer->WeightmapBrushOrderIndices.Num(); ++j)
+					for (int32 j = 0; j < Layer.WeightmapBrushOrderIndices.Num(); ++j)
 					{
-						if (Layer->WeightmapBrushOrderIndices[j] > IndexToRemove)
+						if (Layer.WeightmapBrushOrderIndices[j] > IndexToRemove)
 						{
-							--Layer->HeightmapBrushOrderIndices[j];
+							--Layer.WeightmapBrushOrderIndices[j];
 						}
 					}
 
-					Layer->WeightmapBrushOrderIndices.RemoveAt(i);
+					Layer.WeightmapBrushOrderIndices.RemoveAt(i);
 					break;
 				}
 			}
-		}
 
-		InBrush->SetOwningLandscape(nullptr);
+			InBrush->SetOwningLandscape(nullptr);
+			break;
+		}
 	}
 
-	Landscape->RequestProceduralContentUpdate(InTargetType == ELandscapeToolTargetType::Type::Heightmap ? EProceduralContentUpdateFlag::Heightmap_All : EProceduralContentUpdateFlag::Weightmap_All);
+	Landscape->RequestProceduralContentUpdate(EProceduralContentUpdateFlag::All_Render);
 }
 
 bool FEdModeLandscape::AreAllBrushesCommitedToCurrentProceduralLayer(int32 InTargetType)
@@ -4367,7 +4399,7 @@ bool FEdModeLandscape::AreAllBrushesCommitedToCurrentProceduralLayer(int32 InTar
 
 	for (FLandscapeProceduralLayerBrush& Brush : Layer->Brushes)
 	{
-		if (!Brush.BPCustomBrush->IsCommited() 
+		if (Brush.BPCustomBrush != nullptr && !Brush.BPCustomBrush->IsCommited()
 			&& ((InTargetType == ELandscapeToolTargetType::Type::Heightmap && Brush.BPCustomBrush->IsAffectingHeightmap()) || (InTargetType == ELandscapeToolTargetType::Type::Weightmap && Brush.BPCustomBrush->IsAffectingWeightmap())))
 		{
 			return false;
@@ -4388,7 +4420,10 @@ void FEdModeLandscape::SetCurrentProceduralLayerBrushesCommitState(int32 InTarge
 
 	for (FLandscapeProceduralLayerBrush& Brush : Layer->Brushes)
 	{
-		Brush.BPCustomBrush->SetCommitState(InCommited);
+		if (Brush.BPCustomBrush != nullptr)
+		{
+			Brush.BPCustomBrush->SetCommitState(InCommited);
+		}
 	}
 
 	GEngine->BroadcastLevelActorListChanged();
@@ -4459,8 +4494,8 @@ TArray<ALandscapeBlueprintCustomBrush*> FEdModeLandscape::GetBrushesForCurrentPr
 
 	for (const FLandscapeProceduralLayerBrush& Brush : Layer->Brushes)
 	{
-		if ((Brush.BPCustomBrush->IsAffectingHeightmap() && InTargetType == ELandscapeToolTargetType::Type::Heightmap) 
-			|| (Brush.BPCustomBrush->IsAffectingWeightmap() && InTargetType == ELandscapeToolTargetType::Type::Weightmap))
+		if (Brush.BPCustomBrush != nullptr && ((Brush.BPCustomBrush->IsAffectingHeightmap() && InTargetType == ELandscapeToolTargetType::Type::Heightmap)
+			|| (Brush.BPCustomBrush->IsAffectingWeightmap() && InTargetType == ELandscapeToolTargetType::Type::Weightmap)))
 		{
 			Brushes.Add(Brush.BPCustomBrush);
 		}
@@ -4469,6 +4504,48 @@ TArray<ALandscapeBlueprintCustomBrush*> FEdModeLandscape::GetBrushesForCurrentPr
 	return Brushes;
 }
 
+bool FEdModeLandscape::IsCurrentProceduralLayerBlendSubstractive(const TWeakObjectPtr<ULandscapeLayerInfoObject>& InLayerInfoObj) const
+{
+	FProceduralLayer* Layer = GetCurrentProceduralLayer();
+
+	if (Layer == nullptr)
+	{
+		return false;
+	}
+
+	bool* AllocationBlend = Layer->WeightmapLayerAllocationBlend.Find(InLayerInfoObj.Get());
+
+	if (AllocationBlend != nullptr)
+	{
+		return (*AllocationBlend);
+	}
+
+	return false;
+}
+
+void FEdModeLandscape::SetCurrentProceduralLayerSubstractiveBlendStatus(bool InStatus, const TWeakObjectPtr<ULandscapeLayerInfoObject>& InLayerInfoObj)
+{
+	FProceduralLayer* Layer = GetCurrentProceduralLayer();
+
+	if (Layer == nullptr)
+	{
+		return;
+	}
+
+	bool* AllocationBlend = Layer->WeightmapLayerAllocationBlend.Find(InLayerInfoObj.Get());
+
+	if (AllocationBlend == nullptr)
+	{
+		Layer->WeightmapLayerAllocationBlend.Add(InLayerInfoObj.Get(), InStatus);
+	}
+	else
+	{
+		*AllocationBlend = InStatus;
+	}
+
+	RequestProceduralContentUpdate(true);
+}
+
 FProceduralLayer* FEdModeLandscape::GetCurrentProceduralLayer() const
 {
 	if (!CurrentToolTarget.LandscapeInfo.IsValid())
@@ -4552,6 +4629,103 @@ void FEdModeLandscape::ChangeHeightmapsToCurrentProceduralLayerHeightmaps(bool I
 		}
 	}
 }
+void FEdModeLandscape::ChangeWeightmapsToCurrentProceduralLayerWeightmaps(bool InResetCurrentEditingHeightmap)
+{
+	if (!CurrentToolTarget.LandscapeInfo.IsValid() || !CurrentToolTarget.LandscapeInfo->LandscapeActor.IsValid())
+	{
+		return;
+	}
+
+	TArray<ALandscapeProxy*> AllLandscapes;
+	AllLandscapes.Add(CurrentToolTarget.LandscapeInfo->LandscapeActor.Get());
+
+	for (const auto& It : CurrentToolTarget.LandscapeInfo->Proxies)
+	{
+		AllLandscapes.Add(It);
+	}
+
+	FName CurrentLayerName = GetCurrentProceduralLayerName();
+
+	if (CurrentLayerName == NAME_None)
+	{
+		return;
+	}
+
+	for (ALandscapeProxy* LandscapeProxy : AllLandscapes)
+	{
+		FProceduralLayerData* CurrentLayerData = LandscapeProxy->ProceduralLayersData.Find(CurrentLayerName);
+
+		if (CurrentLayerData == nullptr)
+		{
+			continue;
+		}
+
+		for (ULandscapeComponent* Component : LandscapeProxy->LandscapeComponents)
+		{
+			if (InResetCurrentEditingHeightmap)
+			{
+				Component->SetCurrentEditingWeightmapLayerAllocations(nullptr);
+				Component->SetCurrentEditingWeightmaps(nullptr);
+				Component->SetCurrentProceduralLayerName(NAME_None);
+				Component->SetCurrentEditingWeightmapTexturesUsage(nullptr);			
+			}
+			else
+			{
+				FWeightmapLayerData* WeightmapData = CurrentLayerData->WeightmapData.Find(Component);
+
+				if (WeightmapData != nullptr)
+				{
+					Component->SetCurrentEditingWeightmaps(&WeightmapData->Weightmaps);
+					Component->SetCurrentEditingWeightmapLayerAllocations(&WeightmapData->WeightmapLayerAllocations);
+					Component->SetCurrentProceduralLayerName(CurrentLayerName);
+					Component->SetCurrentEditingWeightmapTexturesUsage(&WeightmapData->WeightmapTextureUsages);
+				}
+			}
+
+			Component->MarkRenderStateDirty();
+		}
+	}
+}
+
+
+bool FEdModeLandscape::NeedToFillEmptyLayersForProcedural() const
+{
+	if (!CurrentToolTarget.LandscapeInfo.IsValid() || !CurrentToolTarget.LandscapeInfo->LandscapeActor.IsValid())
+	{
+		return false;
+	}
+
+	TArray<ALandscapeProxy*> AllLandscapes;
+	AllLandscapes.Add(CurrentToolTarget.LandscapeInfo->LandscapeActor.Get());
+
+	for (const auto& It : CurrentToolTarget.LandscapeInfo->Proxies)
+	{
+		AllLandscapes.Add(It);
+	}
+
+	for (const ALandscapeProxy* LandscapeProxy : AllLandscapes)
+	{
+		for (const auto& ItLayerPair : LandscapeProxy->ProceduralLayersData)
+		{
+			const FProceduralLayerData& ProceduralLayerData = ItLayerPair.Value;
+
+			for (const auto& ItWeightmapPair : ProceduralLayerData.WeightmapData)
+			{
+				const FWeightmapLayerData& WeightmapData = ItWeightmapPair.Value;
+
+				for (const FWeightmapLayerAllocationInfo& Alloc : WeightmapData.WeightmapLayerAllocations)
+				{
+					if (Alloc.LayerInfo != nullptr)
+					{
+						return false;
+					}
+				}
+			}
+		}
+	}
+
+	return true;
+}
 
 void FEdModeLandscape::OnLevelActorAdded(AActor* InActor)
 {
diff --git a/Engine/Source/Editor/LandscapeEditor/Private/LandscapeEdMode.h b/Engine/Source/Editor/LandscapeEditor/Private/LandscapeEdMode.h
--- a/Engine/Source/Editor/LandscapeEditor/Private/LandscapeEdMode.h
+++ b/Engine/Source/Editor/LandscapeEditor/Private/LandscapeEdMode.h
@@ -98,11 +98,12 @@ struct FLandscapeTargetListInfo
 	{
 	}
 
-	FLandscapeInfoLayerSettings* GetLandscapeInfoLayerSettings() const
+	int32 GetLandscapeInfoLayerIndex() const
 	{
+		int32 Index = INDEX_NONE;
+
 		if (TargetType == ELandscapeToolTargetType::Weightmap)
 		{
-			int32 Index = INDEX_NONE;
 			if (LayerInfoObj.IsValid())
 			{
 				Index = LandscapeInfo->GetLayerInfoIndex(LayerInfoObj.Get(), Owner.Get());
@@ -111,11 +112,20 @@ struct FLandscapeTargetListInfo
 			{
 				Index = LandscapeInfo->GetLayerInfoIndex(LayerName, Owner.Get());
 			}
-			if (ensure(Index != INDEX_NONE))
-			{
-				return &LandscapeInfo->Layers[Index];
-			}
 		}
+
+		return Index;
+	}
+
+	FLandscapeInfoLayerSettings* GetLandscapeInfoLayerSettings() const
+	{
+		int32 Index = GetLandscapeInfoLayerIndex();
+
+		if (Index != INDEX_NONE)
+		{
+			return &LandscapeInfo->Layers[Index];
+		}
+
 		return NULL;
 	}
 
@@ -486,10 +496,14 @@ public:
 	FName GetCurrentProceduralLayerName() const;
 	FName GetProceduralLayerName(int32 InLayerIndex) const;
 	void SetProceduralLayerName(int32 InLayerIndex, const FName& InName);
-	float GetProceduralLayerWeight(int32 InLayerIndex) const;
-	void SetProceduralLayerWeight(float InWeight, int32 InLayerIndex);
+	float GetProceduralLayerAlpha(int32 InLayerIndex) const;
+	void SetProceduralLayerAlpha(float InWeight, int32 InLayerIndex);
 	void SetProceduralLayerVisibility(bool InVisible, int32 InLayerIndex);
 	bool IsProceduralLayerVisible(int32 InLayerIndex) const;
+	struct FProceduralLayer* GetCurrentProceduralLayer() const;
+	bool IsCurrentProceduralLayerBlendSubstractive(const TWeakObjectPtr<ULandscapeLayerInfoObject>& InLayerInfoObj) const;
+	void SetCurrentProceduralLayerSubstractiveBlendStatus(bool InStatus, const TWeakObjectPtr<ULandscapeLayerInfoObject>& InLayerInfoObj);
+	
 	void AddBrushToCurrentProceduralLayer(int32 InTargetType, class ALandscapeBlueprintCustomBrush* InBrush);
 	void RemoveBrushFromCurrentProceduralLayer(int32 InTargetType, class ALandscapeBlueprintCustomBrush* InBrush);
 	bool AreAllBrushesCommitedToCurrentProceduralLayer(int32 InTargetType);
@@ -497,9 +511,12 @@ public:
 	TArray<int8>& GetBrushesOrderForCurrentProceduralLayer(int32 InTargetType) const;
 	class ALandscapeBlueprintCustomBrush* GetBrushForCurrentProceduralLayer(int32 InTargetType, int8 BrushIndex) const;
 	TArray<class ALandscapeBlueprintCustomBrush*> GetBrushesForCurrentProceduralLayer(int32 InTargetType);
-	struct FProceduralLayer* GetCurrentProceduralLayer() const;
+	
 	void ChangeHeightmapsToCurrentProceduralLayerHeightmaps(bool InResetCurrentEditingHeightmap = false);
-	void RequestProceduralContentUpdate();
+	void ChangeWeightmapsToCurrentProceduralLayerWeightmaps(bool InResetCurrentEditingHeightmap = false);
+
+	bool NeedToFillEmptyLayersForProcedural() const;
+	void RequestProceduralContentUpdate(bool InUpdateAllMaterials = false);
 
 	void OnLevelActorAdded(AActor* InActor);
 	void OnLevelActorRemoved(AActor* InActor);
diff --git a/Engine/Source/Editor/LandscapeEditor/Private/LandscapeEdModeBrushes.cpp b/Engine/Source/Editor/LandscapeEditor/Private/LandscapeEdModeBrushes.cpp
--- a/Engine/Source/Editor/LandscapeEditor/Private/LandscapeEdModeBrushes.cpp
+++ b/Engine/Source/Editor/LandscapeEditor/Private/LandscapeEdModeBrushes.cpp
@@ -212,12 +212,14 @@ public:
 				}
 				else
 				{
-					bool bExisting = Component->WeightmapLayerAllocations.ContainsByPredicate([LayerInfo](const FWeightmapLayerAllocationInfo& Allocation) { return Allocation.LayerInfo == LayerInfo; });
+					TArray<FWeightmapLayerAllocationInfo>& ComponentWeightmapLayerAllocations = Component->GetWeightmapLayerAllocations(true);
+
+					bool bExisting = ComponentWeightmapLayerAllocations.ContainsByPredicate([LayerInfo](const FWeightmapLayerAllocationInfo& Allocation) { return Allocation.LayerInfo == LayerInfo; });
 					if (!bExisting)
 					{
 						if (EdMode->UISettings->PaintingRestriction == ELandscapeLayerPaintingRestriction::ExistingOnly ||
 							(EdMode->UISettings->PaintingRestriction == ELandscapeLayerPaintingRestriction::UseMaxLayers &&
-							 LandscapeProxy->MaxPaintedLayersPerComponent > 0 && Component->WeightmapLayerAllocations.Num() >= LandscapeProxy->MaxPaintedLayersPerComponent))
+							 LandscapeProxy->MaxPaintedLayersPerComponent > 0 && ComponentWeightmapLayerAllocations.Num() >= LandscapeProxy->MaxPaintedLayersPerComponent))
 						{
 							bCanPaint = false;
 						}
diff --git a/Engine/Source/Editor/LandscapeEditor/Private/LandscapeEdModeComponentTools.cpp b/Engine/Source/Editor/LandscapeEditor/Private/LandscapeEdModeComponentTools.cpp
--- a/Engine/Source/Editor/LandscapeEditor/Private/LandscapeEdModeComponentTools.cpp
+++ b/Engine/Source/Editor/LandscapeEditor/Private/LandscapeEdModeComponentTools.cpp
@@ -342,9 +342,48 @@ public:
 			return false;
 		}
 
+		if (GetMutableDefault<UEditorExperimentalSettings>()->bProceduralLandscape)
+		{
+			ALandscape* Landscape = this->EdMode->CurrentToolTarget.LandscapeInfo->LandscapeActor.Get();
+			if (Landscape != nullptr)
+			{
+				Landscape->RequestProceduralContentUpdate(EProceduralContentUpdateFlag::Weightmap_Render);
+				this->EdMode->ChangeWeightmapsToCurrentProceduralLayerWeightmaps();
+			}
+		}
+
 		return FLandscapeToolBase<FLandscapeToolStrokeVisibility>::BeginTool(ViewportClient, InTarget, InHitLocation);
 	}
 
+	virtual void Tick(FEditorViewportClient* ViewportClient, float DeltaTime) override
+	{
+		FLandscapeToolBase<FLandscapeToolStrokeVisibility>::Tick(ViewportClient, DeltaTime);
+
+		if (GetMutableDefault<UEditorExperimentalSettings>()->bProceduralLandscape && this->IsToolActive())
+		{
+			ALandscape* Landscape = this->EdMode->CurrentToolTarget.LandscapeInfo->LandscapeActor.Get();
+			if (Landscape != nullptr)
+			{
+				Landscape->RequestProceduralContentUpdate(EProceduralContentUpdateFlag::Weightmap_Render);
+			}
+		}
+	}
+
+	virtual void EndTool(FEditorViewportClient* ViewportClient) override
+	{
+		if (GetMutableDefault<UEditorExperimentalSettings>()->bProceduralLandscape)
+		{
+			this->EdMode->ChangeWeightmapsToCurrentProceduralLayerWeightmaps(true);
+
+			if (this->EdMode->CurrentToolTarget.LandscapeInfo->LandscapeActor.IsValid())
+			{
+				this->EdMode->CurrentToolTarget.LandscapeInfo->LandscapeActor->RequestProceduralContentUpdate(EProceduralContentUpdateFlag::Weightmap_All);
+			}
+		}
+
+		FLandscapeToolBase<FLandscapeToolStrokeVisibility>::EndTool(ViewportClient);
+	}
+
 	virtual const TCHAR* GetToolName() override { return TEXT("Visibility"); }
 	virtual FText GetDisplayName() override { return NSLOCTEXT("UnrealEd", "LandscapeMode_Visibility", "Visibility"); };
 
@@ -596,7 +635,9 @@ public:
 					FLandscapeEditDataInterface LandscapeEdit(LandscapeInfo);
 					for (ULandscapeComponent* Component : TargetSelectedComponents)
 					{
-						int32 TotalNeededChannels = Component->WeightmapLayerAllocations.Num();
+						TArray<FWeightmapLayerAllocationInfo>& ComponentWeightmapLayerAllocations = Component->GetWeightmapLayerAllocations();
+
+						int32 TotalNeededChannels = ComponentWeightmapLayerAllocations.Num();
 						int32 CurrentLayer = 0;
 						TArray<UTexture2D*> NewWeightmapTextures;
 
@@ -617,8 +658,8 @@ public:
 								int32 BestDistanceSquared = MAX_int32;
 								for (auto& WeightmapUsagePair : LandscapeProxy->WeightmapUsageMap)
 								{
-									FLandscapeWeightmapUsage* TryWeightmapUsage = &WeightmapUsagePair.Value;
-									if (TryWeightmapUsage->FreeChannelCount() >= TotalNeededChannels)
+									FLandscapeWeightmapUsage* TryWeightmapUsage = WeightmapUsagePair.Value;
+									if (TryWeightmapUsage->FreeChannelCount() >= TotalNeededChannels) // TODO: handle procedural layer
 									{
 										// See if this candidate is closer than any others we've found
 										for (int32 ChanIdx = 0; ChanIdx < 4; ChanIdx++)
@@ -654,7 +695,7 @@ public:
 								CurrentWeightmapTexture->PostEditChange();
 
 								// Store it in the usage map
-								CurrentWeightmapUsage = &LandscapeProxy->WeightmapUsageMap.Add(CurrentWeightmapTexture, FLandscapeWeightmapUsage());
+								CurrentWeightmapUsage = LandscapeProxy->WeightmapUsageMap.Add(CurrentWeightmapTexture, new FLandscapeWeightmapUsage());
 
 								// UE_LOG(LogLandscape, Log, TEXT("Making a new texture %s"), *CurrentWeightmapTexture->GetName());
 							}
@@ -668,7 +709,7 @@ public:
 								if (CurrentWeightmapUsage->ChannelUsage[ChanIdx] == nullptr)
 								{
 									// Use this allocation
-									FWeightmapLayerAllocationInfo& AllocInfo = Component->WeightmapLayerAllocations[CurrentLayer];
+									FWeightmapLayerAllocationInfo& AllocInfo = ComponentWeightmapLayerAllocations[CurrentLayer];
 
 									if (AllocInfo.WeightmapTextureIndex == 255)
 									{
@@ -677,15 +718,20 @@ public:
 									}
 									else
 									{
-										UTexture2D* OldWeightmapTexture = Component->WeightmapTextures[AllocInfo.WeightmapTextureIndex];
+										TArray<UTexture2D*>& ComponentWeightmapTextures = Component->GetWeightmapTextures();
+										UTexture2D* OldWeightmapTexture = ComponentWeightmapTextures[AllocInfo.WeightmapTextureIndex];
 
 										// Copy the data
 										LandscapeEdit.CopyTextureChannel(CurrentWeightmapTexture, ChanIdx, OldWeightmapTexture, AllocInfo.WeightmapTextureChannel);
 										LandscapeEdit.ZeroTextureChannel(OldWeightmapTexture, AllocInfo.WeightmapTextureChannel);
 
 										// Remove the old allocation
-										FLandscapeWeightmapUsage* OldWeightmapUsage = Component->GetLandscapeProxy()->WeightmapUsageMap.Find(OldWeightmapTexture);
-										OldWeightmapUsage->ChannelUsage[AllocInfo.WeightmapTextureChannel] = nullptr;
+										FLandscapeWeightmapUsage** OldWeightmapUsage = Component->GetLandscapeProxy()->WeightmapUsageMap.Find(OldWeightmapTexture);
+
+										if (OldWeightmapUsage != nullptr)
+										{
+											(*OldWeightmapUsage)->ChannelUsage[AllocInfo.WeightmapTextureChannel] = nullptr;
+										}
 									}
 
 									// Assign the new allocation
@@ -699,10 +745,10 @@ public:
 						}
 
 						// Replace the weightmap textures
-						Component->WeightmapTextures = NewWeightmapTextures;
+						Component->SetWeightmapTextures(NewWeightmapTextures);
 
 						// Update the mipmaps for the textures we edited
-						for (UTexture2D* WeightmapTexture : Component->WeightmapTextures)
+						for (UTexture2D* WeightmapTexture : NewWeightmapTextures)
 						{
 							FLandscapeTextureDataInfo* WeightmapDataInfo = LandscapeEdit.GetTextureDataInfo(WeightmapTexture);
 
diff --git a/Engine/Source/Editor/LandscapeEditor/Private/LandscapeEdModePaintTools.cpp b/Engine/Source/Editor/LandscapeEditor/Private/LandscapeEdModePaintTools.cpp
--- a/Engine/Source/Editor/LandscapeEditor/Private/LandscapeEdModePaintTools.cpp
+++ b/Engine/Source/Editor/LandscapeEditor/Private/LandscapeEdModePaintTools.cpp
@@ -80,10 +80,17 @@ public:
 			ALandscape* Landscape = this->EdMode->CurrentToolTarget.LandscapeInfo->LandscapeActor.Get();
 			if (Landscape != nullptr)
 			{
-				Landscape->RequestProceduralContentUpdate(this->EdMode->CurrentToolTarget.TargetType == ELandscapeToolTargetType::Type::Heightmap ? EProceduralContentUpdateFlag::Heightmap_Render : EProceduralContentUpdateFlag::Weightmap_Render);
+				if (this->EdMode->CurrentToolTarget.TargetType == ELandscapeToolTargetType::Type::Heightmap)
+				{
+					Landscape->RequestProceduralContentUpdate(EProceduralContentUpdateFlag::Heightmap_Render);
+					this->EdMode->ChangeHeightmapsToCurrentProceduralLayerHeightmaps();
+				}
+				else
+				{
+					Landscape->RequestProceduralContentUpdate(EProceduralContentUpdateFlag::Weightmap_Render);
+					this->EdMode->ChangeWeightmapsToCurrentProceduralLayerWeightmaps();
+				}
 			}
-
-			this->EdMode->ChangeHeightmapsToCurrentProceduralLayerHeightmaps(false);
 		}
 
 		return FLandscapeToolBase<TStrokeClass>::BeginTool(ViewportClient, InTarget, InHitLocation);
@@ -104,7 +111,8 @@ public:
 			}
 			else
 			{
-				// TODO: Activate/Deactivate weightmap layers
+				this->EdMode->ChangeWeightmapsToCurrentProceduralLayerWeightmaps(true);
+
 				if (this->EdMode->CurrentToolTarget.LandscapeInfo->LandscapeActor.IsValid())
 				{
 					this->EdMode->CurrentToolTarget.LandscapeInfo->LandscapeActor->RequestProceduralContentUpdate(EProceduralContentUpdateFlag::Weightmap_All);
diff --git a/Engine/Source/Editor/LandscapeEditor/Private/LandscapeEditorDetailCustomization_ProceduralLayers.cpp b/Engine/Source/Editor/LandscapeEditor/Private/LandscapeEditorDetailCustomization_ProceduralLayers.cpp
--- a/Engine/Source/Editor/LandscapeEditor/Private/LandscapeEditorDetailCustomization_ProceduralLayers.cpp
+++ b/Engine/Source/Editor/LandscapeEditor/Private/LandscapeEditorDetailCustomization_ProceduralLayers.cpp
@@ -204,31 +204,40 @@ TSharedPtr<SWidget> FLandscapeEditorCustomNodeBuilder_ProceduralLayers::Generate
 					.Text(LOCTEXT("FLandscapeEditorCustomNodeBuilder_ProceduralLayerVisibility", "Visibility"))
 				]
 			]
+
 			+ SHorizontalBox::Slot()
-			.Padding(0)
 			.FillWidth(1.0f)
 			.VAlign(VAlign_Center)
-			.HAlign(HAlign_Left)
-			[
-				SNew(STextBlock)
-				.Text(LOCTEXT("FLandscapeEditorCustomNodeBuilder_ProceduralLayerWeight", "Weight"))
-			]
-			+ SHorizontalBox::Slot()
-			.VAlign(VAlign_Center)
 			.Padding(0, 2)
-			.HAlign(HAlign_Left)
-			.FillWidth(1.0f)
+			.HAlign(HAlign_Right)
 			[
-				SNew(SNumericEntryBox<float>)
-				.AllowSpin(true)
-				.MinValue(0.0f)
-				.MaxValue(65536.0f)
-				.MaxSliderValue(65536.0f)
-				.MinDesiredValueWidth(25.0f)
-				.Value(this, &FLandscapeEditorCustomNodeBuilder_ProceduralLayers::GetLayerWeight, InLayerIndex)
-				.OnValueChanged(this, &FLandscapeEditorCustomNodeBuilder_ProceduralLayers::SetLayerWeight, InLayerIndex)
-				.IsEnabled(true)
-			]			
+				SNew(SHorizontalBox)
+				+ SHorizontalBox::Slot()
+				.Padding(0)
+				.FillWidth(1.0f)
+				.VAlign(VAlign_Center)
+				.HAlign(HAlign_Left)
+				[
+					SNew(STextBlock)
+					.Text(LOCTEXT("FLandscapeEditorCustomNodeBuilder_ProceduralLayerAlpha", "Alpha"))
+				]
+				+ SHorizontalBox::Slot()
+				.VAlign(VAlign_Center)
+				.Padding(0, 2)
+				.HAlign(HAlign_Left)
+				.FillWidth(1.0f)
+				[
+					SNew(SNumericEntryBox<float>)
+					.AllowSpin(true)
+					.MinValue(0.0f)
+					.MaxValue(100.0f)
+					.MaxSliderValue(100.0f)
+					.MinDesiredValueWidth(60.0f)
+					.Value(this, &FLandscapeEditorCustomNodeBuilder_ProceduralLayers::GetLayerAlpha, InLayerIndex)
+					.OnValueChanged(this, &FLandscapeEditorCustomNodeBuilder_ProceduralLayers::SetLayerAlpha, InLayerIndex)
+					.IsEnabled(true)
+				]		
+			]
 		];	
 
 	return RowWidget;
@@ -278,25 +287,25 @@ void FLandscapeEditorCustomNodeBuilder_ProceduralLayers::OnLayerSelectionChanged
 	}
 }
 
-TOptional<float> FLandscapeEditorCustomNodeBuilder_ProceduralLayers::GetLayerWeight(int32 InLayerIndex) const
+TOptional<float> FLandscapeEditorCustomNodeBuilder_ProceduralLayers::GetLayerAlpha(int32 InLayerIndex) const
 {
 	FEdModeLandscape* LandscapeEdMode = GetEditorMode();
 
 	if (LandscapeEdMode)
 	{
-		return LandscapeEdMode->GetProceduralLayerWeight(InLayerIndex);
+		return LandscapeEdMode->GetProceduralLayerAlpha(InLayerIndex);
 	}
 
 	return 1.0f;
 }
 
-void FLandscapeEditorCustomNodeBuilder_ProceduralLayers::SetLayerWeight(float InWeight, int32 InLayerIndex)
+void FLandscapeEditorCustomNodeBuilder_ProceduralLayers::SetLayerAlpha(float InAlpha, int32 InLayerIndex)
 {
 	FEdModeLandscape* LandscapeEdMode = GetEditorMode();
 
 	if (LandscapeEdMode)
 	{
-		LandscapeEdMode->SetProceduralLayerWeight(InWeight, InLayerIndex);
+		LandscapeEdMode->SetProceduralLayerAlpha(InAlpha, InLayerIndex);
 	}
 }
 
diff --git a/Engine/Source/Editor/LandscapeEditor/Private/LandscapeEditorDetailCustomization_ProceduralLayers.h b/Engine/Source/Editor/LandscapeEditor/Private/LandscapeEditorDetailCustomization_ProceduralLayers.h
--- a/Engine/Source/Editor/LandscapeEditor/Private/LandscapeEditorDetailCustomization_ProceduralLayers.h
+++ b/Engine/Source/Editor/LandscapeEditor/Private/LandscapeEditorDetailCustomization_ProceduralLayers.h
@@ -69,8 +69,8 @@ protected:
 	void OnLayerTextCommitted(const FText& InText, ETextCommit::Type InCommitType, int32 InLayerIndex);
 	FText GetLayerText(int32 InLayerIndex) const;
 
-	TOptional<float> GetLayerWeight(int32 InLayerIndex) const;
-	void SetLayerWeight(float InWeight, int32 InLayerIndex);
+	TOptional<float> GetLayerAlpha(int32 InLayerIndex) const;
+	void SetLayerAlpha(float InAlpha, int32 InLayerIndex);
 
 	void OnLayerVisibilityChanged(ECheckBoxState NewState, int32 InLayerIndex);
 	ECheckBoxState IsLayerVisible(int32 InLayerIndex) const;
diff --git a/Engine/Source/Editor/LandscapeEditor/Private/LandscapeEditorDetailCustomization_TargetLayers.cpp b/Engine/Source/Editor/LandscapeEditor/Private/LandscapeEditorDetailCustomization_TargetLayers.cpp
--- a/Engine/Source/Editor/LandscapeEditor/Private/LandscapeEditorDetailCustomization_TargetLayers.cpp
+++ b/Engine/Source/Editor/LandscapeEditor/Private/LandscapeEditorDetailCustomization_TargetLayers.cpp
@@ -41,6 +41,7 @@
 #include "IDetailGroup.h"
 #include "Widgets/SBoxPanel.h"
 #include "Editor/EditorStyle/Private/SlateEditorStyle.h"
+#include "Settings/EditorExperimentalSettings.h"
 
 #define LOCTEXT_NAMESPACE "LandscapeEditor.TargetLayers"
 
@@ -698,6 +699,24 @@ TSharedPtr<SWidget> FLandscapeEditorCustomNodeBuilder_TargetLayers::GenerateRow(
 					]
 					+ SVerticalBox::Slot()
 					.AutoHeight()
+					[
+						SNew(SHorizontalBox)
+						.Visibility_Static(&FLandscapeEditorCustomNodeBuilder_TargetLayers::GetProceduralLayerSubstractiveBlendVisibility, Target)
+						+ SHorizontalBox::Slot()
+						.AutoWidth()
+						.Padding(0, 2, 2, 2)
+						[
+							SNew(SCheckBox)
+							.IsChecked_Static(&FLandscapeEditorCustomNodeBuilder_TargetLayers::ProceduralLayerSubstractiveBlendIsChecked, Target)
+							.OnCheckStateChanged_Static(&FLandscapeEditorCustomNodeBuilder_TargetLayers::OnProceduralLayerSubstractiveBlendChanged, Target)
+							[
+								SNew(STextBlock)
+								.Text(LOCTEXT("ViewMode.Debug_None", "Substractive Blend"))
+							]
+						]
+					]
+					+ SVerticalBox::Slot()
+					.AutoHeight()
 					[
 						SNew(SHorizontalBox)
 						.Visibility_Static(&FLandscapeEditorCustomNodeBuilder_TargetLayers::GetDebugModeColorChannelVisibility, Target)
@@ -919,7 +938,7 @@ TSharedPtr<SWidget> FLandscapeEditorCustomNodeBuilder_TargetLayers::OnTargetLaye
 
 				// Rebuild material instances
 				FUIAction RebuildAction = FUIAction(FExecuteAction::CreateStatic(&FLandscapeEditorCustomNodeBuilder_TargetLayers::OnRebuildMICs, Target));
-				MenuBuilder.AddMenuEntry(LOCTEXT("LayerContextMenu.Rebuild", "Rebuild Materials"), LOCTEXT("LayerContextMenu.Rebuild_Tooltip", "Rebuild material instances used for this landscape."), FSlateIcon(), ClearAction);
+				MenuBuilder.AddMenuEntry(LOCTEXT("LayerContextMenu.Rebuild", "Rebuild Materials"), LOCTEXT("LayerContextMenu.Rebuild_Tooltip", "Rebuild material instances used for this landscape."), FSlateIcon(), RebuildAction);
 			}
 			else if (Target->TargetType == ELandscapeToolTargetType::Visibility)
 			{
@@ -1060,7 +1079,17 @@ void FLandscapeEditorCustomNodeBuilder_TargetLayers::OnReimportLayer(const TShar
 	FEdModeLandscape* LandscapeEdMode = GetEditorMode();
 	if (LandscapeEdMode)
 	{
-		LandscapeEdMode->ReimportData(*Target);
+		if (GetMutableDefault<UEditorExperimentalSettings>()->bProceduralLandscape)
+		{
+			LandscapeEdMode->ChangeWeightmapsToCurrentProceduralLayerWeightmaps();
+			LandscapeEdMode->ReimportData(*Target);
+			LandscapeEdMode->ChangeWeightmapsToCurrentProceduralLayerWeightmaps(true);
+			LandscapeEdMode->RequestProceduralContentUpdate();
+		}
+		else
+		{
+			LandscapeEdMode->ReimportData(*Target);
+		}
 	}
 }
 
@@ -1070,14 +1099,47 @@ void FLandscapeEditorCustomNodeBuilder_TargetLayers::OnFillLayer(const TSharedRe
 	if (Target->LandscapeInfo.IsValid() && Target->LayerInfoObj.IsValid())
 	{
 		FLandscapeEditDataInterface LandscapeEdit(Target->LandscapeInfo.Get());
-		LandscapeEdit.FillLayer(Target->LayerInfoObj.Get());
+
+		FEdModeLandscape* LandscapeEdMode = GetEditorMode();
+		if (LandscapeEdMode)
+		{
+			if (GetMutableDefault<UEditorExperimentalSettings>()->bProceduralLandscape)
+			{
+				LandscapeEdMode->ChangeWeightmapsToCurrentProceduralLayerWeightmaps();
+				LandscapeEdit.FillLayer(Target->LayerInfoObj.Get());
+				LandscapeEdMode->ChangeWeightmapsToCurrentProceduralLayerWeightmaps(true);
+				LandscapeEdMode->RequestProceduralContentUpdate(true);
+			}
+			else
+			{
+				LandscapeEdit.FillLayer(Target->LayerInfoObj.Get());
+			}
+		}
 	}
 }
 
 void FLandscapeEditorCustomNodeBuilder_TargetLayers::FillEmptyLayers(ULandscapeInfo* LandscapeInfo, ULandscapeLayerInfoObject* LandscapeInfoObject)
 {
-	FLandscapeEditDataInterface LandscapeEdit(LandscapeInfo);
-	LandscapeEdit.FillEmptyLayers(LandscapeInfoObject);
+	FEdModeLandscape* LandscapeEdMode = GetEditorMode();
+	if (LandscapeEdMode)
+	{
+		FLandscapeEditDataInterface LandscapeEdit(LandscapeInfo);
+
+		if (GetMutableDefault<UEditorExperimentalSettings>()->bProceduralLandscape)
+		{
+			if (LandscapeEdMode->NeedToFillEmptyLayersForProcedural())
+			{
+				LandscapeEdMode->ChangeWeightmapsToCurrentProceduralLayerWeightmaps();
+				LandscapeEdit.FillEmptyLayers(LandscapeInfoObject);
+				LandscapeEdMode->ChangeWeightmapsToCurrentProceduralLayerWeightmaps(true);
+			}
+		}
+		else
+		{
+			LandscapeEdit.FillEmptyLayers(LandscapeInfoObject);
+		}
+	}
+
 }
 
 
@@ -1086,8 +1148,23 @@ void FLandscapeEditorCustomNodeBuilder_TargetLayers::OnClearLayer(const TSharedR
 	FScopedTransaction Transaction(LOCTEXT("Undo_ClearLayer", "Clearing Landscape Layer"));
 	if (Target->LandscapeInfo.IsValid() && Target->LayerInfoObj.IsValid())
 	{
-		FLandscapeEditDataInterface LandscapeEdit(Target->LandscapeInfo.Get());
-		LandscapeEdit.DeleteLayer(Target->LayerInfoObj.Get());
+		FEdModeLandscape* LandscapeEdMode = GetEditorMode();
+		if (LandscapeEdMode)
+		{
+			FLandscapeEditDataInterface LandscapeEdit(Target->LandscapeInfo.Get());
+
+			if (GetMutableDefault<UEditorExperimentalSettings>()->bProceduralLandscape)
+			{
+				LandscapeEdMode->ChangeWeightmapsToCurrentProceduralLayerWeightmaps();
+				LandscapeEdit.DeleteLayer(Target->LayerInfoObj.Get());
+				LandscapeEdMode->ChangeWeightmapsToCurrentProceduralLayerWeightmaps(true);
+				LandscapeEdMode->RequestProceduralContentUpdate(true);
+			}
+			else
+			{
+				LandscapeEdit.DeleteLayer(Target->LayerInfoObj.Get());
+			}
+		}
 	}
 }
 
@@ -1377,6 +1454,36 @@ EVisibility FLandscapeEditorCustomNodeBuilder_TargetLayers::GetDebugModeLayerUsa
 	return EVisibility::Visible;
 }
 
+EVisibility FLandscapeEditorCustomNodeBuilder_TargetLayers::GetProceduralLayerSubstractiveBlendVisibility(const TSharedRef<FLandscapeTargetListInfo> Target)
+{
+	if (GetMutableDefault<UEditorExperimentalSettings>()->bProceduralLandscape && Target->TargetType != ELandscapeToolTargetType::Heightmap && Target->LayerInfoObj.IsValid())
+	{
+		return EVisibility::Visible;
+	}
+
+	return EVisibility::Collapsed;
+}
+
+ECheckBoxState FLandscapeEditorCustomNodeBuilder_TargetLayers::ProceduralLayerSubstractiveBlendIsChecked(const TSharedRef<FLandscapeTargetListInfo> Target)
+{
+	FEdModeLandscape* LandscapeEdMode = GetEditorMode();
+	if (LandscapeEdMode)
+	{
+		return LandscapeEdMode->IsCurrentProceduralLayerBlendSubstractive(Target.Get().LayerInfoObj) ? ECheckBoxState::Checked : ECheckBoxState::Unchecked;
+	}
+
+	return ECheckBoxState::Unchecked;
+}
+
+void FLandscapeEditorCustomNodeBuilder_TargetLayers::OnProceduralLayerSubstractiveBlendChanged(ECheckBoxState NewCheckedState, const TSharedRef<FLandscapeTargetListInfo> Target)
+{
+	FEdModeLandscape* LandscapeEdMode = GetEditorMode();
+	if (LandscapeEdMode)
+	{
+		LandscapeEdMode->SetCurrentProceduralLayerSubstractiveBlendStatus(NewCheckedState == ECheckBoxState::Checked, Target.Get().LayerInfoObj);
+	}
+}
+
 EVisibility FLandscapeEditorCustomNodeBuilder_TargetLayers::GetDebugModeColorChannelVisibility(const TSharedRef<FLandscapeTargetListInfo> Target)
 {
 	if (GLandscapeViewMode == ELandscapeViewMode::DebugLayer && Target->TargetType != ELandscapeToolTargetType::Heightmap && Target->LayerInfoObj.IsValid())
diff --git a/Engine/Source/Editor/LandscapeEditor/Private/LandscapeEditorDetailCustomization_TargetLayers.h b/Engine/Source/Editor/LandscapeEditor/Private/LandscapeEditorDetailCustomization_TargetLayers.h
--- a/Engine/Source/Editor/LandscapeEditor/Private/LandscapeEditorDetailCustomization_TargetLayers.h
+++ b/Engine/Source/Editor/LandscapeEditor/Private/LandscapeEditorDetailCustomization_TargetLayers.h
@@ -94,6 +94,10 @@ protected:
 	static ECheckBoxState DebugModeColorChannelIsChecked(const TSharedRef<FLandscapeTargetListInfo> Target, int32 Channel);
 	static void OnDebugModeColorChannelChanged(ECheckBoxState NewCheckedState, const TSharedRef<FLandscapeTargetListInfo> Target, int32 Channel);
 
+	static EVisibility GetProceduralLayerSubstractiveBlendVisibility(const TSharedRef<FLandscapeTargetListInfo> Target);
+	static ECheckBoxState ProceduralLayerSubstractiveBlendIsChecked(const TSharedRef<FLandscapeTargetListInfo> Target);
+	static void OnProceduralLayerSubstractiveBlendChanged(ECheckBoxState NewCheckedState, const TSharedRef<FLandscapeTargetListInfo> Target);
+
 	// Drag/Drop handling
 	FReply HandleDragDetected(const FGeometry& MyGeometry, const FPointerEvent& MouseEvent, int32 SlotIndex, SVerticalBox::FSlot* Slot);
 	TOptional<SDragAndDropVerticalBox::EItemDropZone> HandleCanAcceptDrop(const FDragDropEvent& DragDropEvent, SDragAndDropVerticalBox::EItemDropZone DropZone, SVerticalBox::FSlot* Slot);
diff --git a/Engine/Source/Editor/StatsViewer/Private/StatsPages/PrimitiveStatsPage.cpp b/Engine/Source/Editor/StatsViewer/Private/StatsPages/PrimitiveStatsPage.cpp
--- a/Engine/Source/Editor/StatsViewer/Private/StatsPages/PrimitiveStatsPage.cpp
+++ b/Engine/Source/Editor/StatsViewer/Private/StatsPages/PrimitiveStatsPage.cpp
@@ -352,7 +352,6 @@ struct PrimitiveStatsGenerator
 						NewStatsEntry->Sections += FMath::Square(CurrentComponent->NumSubsections);
 
 						// count resource usage of landscape
-						//TODO: take into consideration all the editing RT/heightmap, etc.
 						bool bNotUnique = false;
 						UniqueTextures.Add(CurrentComponent->GetHeightmap(), &bNotUnique);
 						if (!bNotUnique)
@@ -370,12 +369,14 @@ struct PrimitiveStatsGenerator
 							}
 						}
 
-						for (auto ItWeightmaps = CurrentComponent->WeightmapTextures.CreateConstIterator(); ItWeightmaps; ++ItWeightmaps)
+						const TArray<UTexture2D*>& ComponentWeightmapTextures = CurrentComponent->GetWeightmapTextures();
+
+						for (UTexture2D* Weightmap : ComponentWeightmapTextures)
 						{
-							UniqueTextures.Add((*ItWeightmaps), &bNotUnique);
+							UniqueTextures.Add(Weightmap, &bNotUnique);
 							if (!bNotUnique)
 							{
-								const SIZE_T WeightmapResourceSize = (*ItWeightmaps)->GetResourceSizeBytes(EResourceSizeMode::EstimatedTotal);
+								const SIZE_T WeightmapResourceSize = Weightmap->GetResourceSizeBytes(EResourceSizeMode::EstimatedTotal);
 								NewStatsEntry->ResourceSize += (float)WeightmapResourceSize / 1024.f;
 							}
 						}
diff --git a/Engine/Source/Editor/UnrealEd/Private/EditorExporters.cpp b/Engine/Source/Editor/UnrealEd/Private/EditorExporters.cpp
--- a/Engine/Source/Editor/UnrealEd/Private/EditorExporters.cpp
+++ b/Engine/Source/Editor/UnrealEd/Private/EditorExporters.cpp
@@ -853,13 +853,17 @@ static void AddActorToOBJs(AActor* Actor, TArray<FOBJGeom*>& Objects, TSet<UMate
 			int32 WeightMapSize = (SubsectionSizeQuads + 1) * Component->NumSubsections;
 			int32 ChannelOffsets[4] = {(int32)STRUCT_OFFSET(FColor,R),(int32)STRUCT_OFFSET(FColor,G),(int32)STRUCT_OFFSET(FColor,B),(int32)STRUCT_OFFSET(FColor,A)};
 
-			for( int32 AllocIdx=0;AllocIdx < Component->WeightmapLayerAllocations.Num(); AllocIdx++ )
+			TArray<FWeightmapLayerAllocationInfo>& ComponentWeightmapLayerAllocations = Component->GetWeightmapLayerAllocations();
+			TArray<UTexture2D*>& ComponentWeightmapTextures = Component->GetWeightmapTextures();
+
+			for( int32 AllocIdx=0;AllocIdx < ComponentWeightmapLayerAllocations.Num(); AllocIdx++ )
 			{
-				FWeightmapLayerAllocationInfo& AllocInfo = Component->WeightmapLayerAllocations[AllocIdx];
+				FWeightmapLayerAllocationInfo& AllocInfo = ComponentWeightmapLayerAllocations[AllocIdx];
 				if( AllocInfo.LayerInfo == ALandscapeProxy::VisibilityLayer )
 				{
 					TexIndex = AllocInfo.WeightmapTextureIndex;
-					Component->WeightmapTextures[TexIndex]->Source.GetMipData(RawVisData, 0);
+
+					ComponentWeightmapTextures[TexIndex]->Source.GetMipData(RawVisData, 0);
 					VisDataMap = RawVisData.GetData() + ChannelOffsets[AllocInfo.WeightmapTextureChannel];
 				}
 			}
diff --git a/Engine/Source/Editor/UnrealEd/Private/Fbx/FbxMainExport.cpp b/Engine/Source/Editor/UnrealEd/Private/Fbx/FbxMainExport.cpp
--- a/Engine/Source/Editor/UnrealEd/Private/Fbx/FbxMainExport.cpp
+++ b/Engine/Source/Editor/UnrealEd/Private/Fbx/FbxMainExport.cpp
@@ -4272,10 +4272,11 @@ void FFbxExporter::ExportLandscapeToFbx(ALandscapeProxy* Landscape, const TCHAR*
 		FLandscapeComponentDataInterface CDI(Component, Landscape->ExportLOD);
 		const int32 BaseVertIndex = SelectedComponentIndex++ * VertexCountPerComponent;
 
+		TArray<FWeightmapLayerAllocationInfo>& ComponentWeightmapLayerAllocations = Component->GetWeightmapLayerAllocations();
 		TArray<uint8> CompVisData;
-		for (int32 AllocIdx = 0; AllocIdx < Component->WeightmapLayerAllocations.Num(); AllocIdx++)
+		for (int32 AllocIdx = 0; AllocIdx < ComponentWeightmapLayerAllocations.Num(); AllocIdx++)
 		{
-			FWeightmapLayerAllocationInfo& AllocInfo = Component->WeightmapLayerAllocations[AllocIdx];
+			FWeightmapLayerAllocationInfo& AllocInfo = ComponentWeightmapLayerAllocations[AllocIdx];
 			if (AllocInfo.LayerInfo == ALandscapeProxy::VisibilityLayer)
 			{
 				CDI.GetWeightmapTextureData(AllocInfo.LayerInfo, CompVisData);
diff --git a/Engine/Source/Editor/UnrealEd/Private/Lightmass/LightmassLandscapeRender.cpp b/Engine/Source/Editor/UnrealEd/Private/Lightmass/LightmassLandscapeRender.cpp
--- a/Engine/Source/Editor/UnrealEd/Private/Lightmass/LightmassLandscapeRender.cpp
+++ b/Engine/Source/Editor/UnrealEd/Private/Lightmass/LightmassLandscapeRender.cpp
@@ -195,7 +195,9 @@ void GetLandscapeOpacityData(const FLandscapeStaticLightingMesh* LandscapeMesh,
 			ULandscapeComponent* Component = LandscapeInfo->XYtoComponentMap.FindRef(FIntPoint(ComponentX, ComponentY));
 			if (Component)
 			{
-				if (Component->WeightmapLayerAllocations.ContainsByPredicate([](const FWeightmapLayerAllocationInfo& Allocation) { return Allocation.LayerInfo == ALandscapeProxy::VisibilityLayer; }))
+				TArray<FWeightmapLayerAllocationInfo>& ComponentWeightmapLayerAllocations = Component->GetWeightmapLayerAllocations();
+
+				if (ComponentWeightmapLayerAllocations.ContainsByPredicate([](const FWeightmapLayerAllocationInfo& Allocation) { return Allocation.LayerInfo == ALandscapeProxy::VisibilityLayer; }))
 				{
 					// filled by GetWeightDataFast()
 				}
diff --git a/Engine/Source/Runtime/Engine/Classes/Materials/MaterialInstanceConstant.h b/Engine/Source/Runtime/Engine/Classes/Materials/MaterialInstanceConstant.h
--- a/Engine/Source/Runtime/Engine/Classes/Materials/MaterialInstanceConstant.h
+++ b/Engine/Source/Runtime/Engine/Classes/Materials/MaterialInstanceConstant.h
@@ -52,8 +52,9 @@ class UMaterialInstanceConstant : public UMaterialInstance
 	 * Set the parent of this material instance. This function may only be called in the Editor!
 	 *   WARNING: You MUST call PostEditChange afterwards to propagate changes to other materials in the chain!
 	 * @param NewParent - The new parent for this material instance.
+ 	 * @param RecacheShader - Will recache required shaders.
 	 */
-	ENGINE_API void SetParentEditorOnly(class UMaterialInterface* NewParent);
+	ENGINE_API void SetParentEditorOnly(class UMaterialInterface* NewParent, bool RecacheShader = true);
 
 	/**
 	* Copies the uniform parameters (scalar, vector and texture) from a material or instance hierarchy.
diff --git a/Engine/Source/Runtime/Engine/Private/Materials/MaterialInstanceConstant.cpp b/Engine/Source/Runtime/Engine/Private/Materials/MaterialInstanceConstant.cpp
--- a/Engine/Source/Runtime/Engine/Private/Materials/MaterialInstanceConstant.cpp
+++ b/Engine/Source/Runtime/Engine/Private/Materials/MaterialInstanceConstant.cpp
@@ -49,10 +49,10 @@ void UMaterialInstanceConstant::PostEditChangeProperty(FPropertyChangedEvent& Pr
 	ParameterStateId = FGuid::NewGuid();
 }
 
-void UMaterialInstanceConstant::SetParentEditorOnly(UMaterialInterface* NewParent)
+void UMaterialInstanceConstant::SetParentEditorOnly(UMaterialInterface* NewParent, bool RecacheShader)
 {
 	check(GIsEditor || IsRunningCommandlet());
-	SetParentInternal(NewParent, true);
+	SetParentInternal(NewParent, RecacheShader);
 }
 
 void UMaterialInstanceConstant::CopyMaterialUniformParametersEditorOnly(UMaterialInterface* Source, bool bIncludeStaticParams)
diff --git a/Engine/Source/Runtime/Landscape/Classes/Landscape.h b/Engine/Source/Runtime/Landscape/Classes/Landscape.h
--- a/Engine/Source/Runtime/Landscape/Classes/Landscape.h
+++ b/Engine/Source/Runtime/Landscape/Classes/Landscape.h
@@ -36,20 +36,40 @@ enum class ERTDrawingType : uint8
 
 enum EHeightmapRTType : uint8
 {
-	LandscapeSizeCombinedAtlas,
-	LandscapeSizeCombinedNonAtlas,
-	LandscapeSizeScratch1,
-	LandscapeSizeScratch2,
-	LandscapeSizeScratch3,
+	HeightmapRT_CombinedAtlas,
+	HeightmapRT_CombinedNonAtlas,
+	HeightmapRT_Scratch1,
+	HeightmapRT_Scratch2,
+	HeightmapRT_Scratch3,
 	// Mips RT
-	LandscapeSizeMip1,
-	LandscapeSizeMip2,
-	LandscapeSizeMip3,
-	LandscapeSizeMip4,
-	LandscapeSizeMip5,
-	LandscapeSizeMip6,
-	LandscapeSizeMip7,
-	Count
+	HeightmapRT_Mip1,
+	HeightmapRT_Mip2,
+	HeightmapRT_Mip3,
+	HeightmapRT_Mip4,
+	HeightmapRT_Mip5,
+	HeightmapRT_Mip6,
+	HeightmapRT_Mip7,
+	HeightmapRT_Count
+};
+
+enum EWeightmapRTType : uint8
+{
+	WeightmapRT_Scratch_RGBA,
+	WeightmapRT_Scratch1,
+	WeightmapRT_Scratch2,
+	WeightmapRT_Scratch3,
+
+	// Mips RT
+	WeightmapRT_Mip0,
+	WeightmapRT_Mip1,
+	WeightmapRT_Mip2,
+	WeightmapRT_Mip3,
+	WeightmapRT_Mip4,
+	WeightmapRT_Mip5,
+	WeightmapRT_Mip6,
+	WeightmapRT_Mip7,
+	
+	WeightmapRT_Count
 };
 
 enum EProceduralContentUpdateFlag : uint32
@@ -58,21 +78,16 @@ enum EProceduralContentUpdateFlag : uint32
 	Heightmap_Render				= 0x00000002u,
 	Heightmap_BoundsAndCollision	= 0x00000004u,
 	Heightmap_ResolveToTexture		= 0x00000008u,
-	Heightmap_ResolveToTextureDDC	= 0x00000010u,
 
-	// TODO: add weightmap update type
 	Weightmap_Setup					= 0x00000100u,
 	Weightmap_Render				= 0x00000200u,
-	Weightmap_ResolveToTexture		= 0x00000400u,
-	Weightmap_ResolveToTextureDDC	= 0x00000800u,
+	Weightmap_Collision				= 0x00000400u,
+	Weightmap_ResolveToTexture		= 0x00000800u,
 
 	// Combinations
 	Heightmap_All = Heightmap_Render | Heightmap_BoundsAndCollision | Heightmap_ResolveToTexture,
-	Heightmap_All_WithDDCUpdate = Heightmap_Render | Heightmap_BoundsAndCollision | Heightmap_ResolveToTextureDDC,
-	Weightmap_All = Weightmap_Render | Weightmap_ResolveToTexture,
-	Weightmap_All_WithDDCUpdate = Weightmap_Render | Weightmap_ResolveToTextureDDC,
+	Weightmap_All = Weightmap_Render | Weightmap_Collision | Weightmap_ResolveToTexture,
 
-	All_WithDDCUpdate = Heightmap_All_WithDDCUpdate | Weightmap_All_WithDDCUpdate,
 	All = Heightmap_All | Weightmap_All,
 	All_Setup = Heightmap_Setup | Weightmap_Setup,
 	All_Render = Heightmap_Render | Weightmap_Render,
@@ -94,21 +109,29 @@ struct FLandscapeProceduralLayerBrush
 #if WITH_EDITOR
 	UTextureRenderTarget2D* Render(bool InIsHeightmap, UTextureRenderTarget2D* InCombinedResult)
 	{
-		TGuardValue<bool> AutoRestore(GAllowActorScriptExecutionInEditor, true);
-		return BPCustomBrush->Render(InIsHeightmap, InCombinedResult);
+		if (BPCustomBrush != nullptr)
+		{
+			TGuardValue<bool> AutoRestore(GAllowActorScriptExecutionInEditor, true);
+			return BPCustomBrush->Render(InIsHeightmap, InCombinedResult);
+		}
+
+		return nullptr;
 	}
 
 	bool IsInitialized() const 
 	{
-		return BPCustomBrush->IsInitialized();
+		return BPCustomBrush != nullptr ? BPCustomBrush->IsInitialized() : false;
 	}
 
 	void Initialize(const FIntRect& InBoundRect, const FIntPoint& InLandscapeRenderTargetSize)
 	{
-		TGuardValue<bool> AutoRestore(GAllowActorScriptExecutionInEditor, true);
-		FIntPoint LandscapeSize = InBoundRect.Max - InBoundRect.Min;
-		BPCustomBrush->Initialize(LandscapeSize, InLandscapeRenderTargetSize);
-		BPCustomBrush->SetIsInitialized(true);
+		if (BPCustomBrush != nullptr)
+		{
+			TGuardValue<bool> AutoRestore(GAllowActorScriptExecutionInEditor, true);
+			FIntPoint LandscapeSize = InBoundRect.Max - InBoundRect.Min;
+			BPCustomBrush->Initialize(LandscapeSize, InLandscapeRenderTargetSize);
+			BPCustomBrush->SetIsInitialized(true);
+		}
 	}
 #endif
 
@@ -124,7 +147,8 @@ struct FProceduralLayer
 	FProceduralLayer()
 		: Name(NAME_None)
 		, Visible(true)
-		, Weight(1.0f)
+		, HeightmapAlpha(1.0f)
+		, WeightmapAlpha(1.0f)
 	{}
 
 	UPROPERTY()
@@ -134,7 +158,10 @@ struct FProceduralLayer
 	bool Visible;
 
 	UPROPERTY()
-	float Weight;
+	float HeightmapAlpha;
+
+	UPROPERTY()
+	float WeightmapAlpha;
 
 	UPROPERTY()
 	TArray<FLandscapeProceduralLayerBrush> Brushes;
@@ -144,6 +171,9 @@ struct FProceduralLayer
 
 	UPROPERTY()
 	TArray<int8> WeightmapBrushOrderIndices;
+
+	UPROPERTY()
+	TMap<ULandscapeLayerInfoObject*, bool> WeightmapLayerAllocationBlend; // True -> Substractive, False -> Additive
 };
 
 UCLASS(MinimalAPI, showcategories=(Display, Movement, Collision, Lighting, LOD, Input), hidecategories=(Mobility))
@@ -185,50 +215,103 @@ public:
 #endif
 	virtual void PostLoad() override;
 	virtual void BeginDestroy() override;
+	virtual void FinishDestroy() override;
 	//~ End UObject Interface
 
 
-#if WITH_EDITOR
 	// Procedural stuff
-	LANDSCAPE_API void RegenerateProceduralContent();
-	LANDSCAPE_API void RegenerateProceduralHeightmaps();
-	LANDSCAPE_API void ResolveProceduralHeightmapTexture(bool InUpdateDDC);
-	LANDSCAPE_API void RegenerateProceduralWeightmaps();
-
-	LANDSCAPE_API void RequestProceduralContentUpdate(uint32 InDataFlags);
-
-	void GenerateHeightmapQuad(const FIntPoint& InVertexPosition, const float InVertexSize, const FVector2D& InUVStart, const FVector2D& InUVSize, TArray<struct FLandscapeProceduralTriangle>& OutTriangles) const;
-	void GenerateHeightmapQuadsAtlas(const FIntPoint& InSectionBase, const FVector2D& InScaleBias, float InSubSectionSizeQuad, const FIntPoint& InReadSize, const FIntPoint& InWriteSize, TArray<struct FLandscapeProceduralTriangle>& OutTriangles) const;
-	void GenerateHeightmapQuadsAtlasToNonAtlas(const FIntPoint& InSectionBase, const FVector2D& InScaleBias, float InSubSectionSizeQuad, const FIntPoint& InHeightmapReadTextureSize, const FIntPoint& InHeightmapWriteTextureSize, TArray<struct FLandscapeProceduralTriangle>& OutTriangles) const;
-	void GenerateHeightmapQuadsNonAtlas(const FIntPoint& InSectionBase, const FVector2D& InScaleBias, float InSubSectionSizeQuad, const FIntPoint& InHeightmapReadTextureSize, const FIntPoint& InHeightmapWriteTextureSize, TArray<struct FLandscapeProceduralTriangle>& OutTriangles) const;
-	void GenerateHeightmapQuadsNonAtlasToAtlas(const FIntPoint& InSectionBase, const FVector2D& InScaleBias, float InSubSectionSizeQuad, const FIntPoint& InHeightmapReadTextureSize, const FIntPoint& InHeightmapWriteTextureSize, TArray<struct FLandscapeProceduralTriangle>& OutTriangles) const;
-	void GenerateHeightmapQuadsMip(const FIntPoint& InSectionBase, const FVector2D& InScaleBias, float InSubSectionSizeQuad, const FIntPoint& InReadSize, const FIntPoint& InWriteSize, int32 CurrentMip, TArray<FLandscapeProceduralTriangle>& OutTriangles) const;
-
-	void DrawHeightmapComponentsToRenderTarget(const FString& InDebugName, TArray<ULandscapeComponent*>& InComponentsToDraw, UTexture* InHeightmapRTRead, UTextureRenderTarget2D* InOptionalHeightmapRTRead2, UTextureRenderTarget2D* InHeightmapRTWrite, ERTDrawingType InDrawType,
-											   bool InClearRTWrite, struct FLandscapeHeightmapProceduralShaderParameters& InShaderParams, int32 InMipRender = 0) const;
+#if WITH_EDITOR
+	LANDSCAPE_API void RequestProceduralContentUpdate(uint32 InDataFlags, bool InUpdateAllMaterials = false);
+
+private:
+	void TickProcedural(float DeltaTime, ELevelTick TickType, FActorTickFunction& ThisTickFunction);
+	void RegenerateProceduralContent();
+	void RegenerateProceduralHeightmaps();
+	void ResolveProceduralHeightmapTexture(const TArray<ALandscapeProxy*>& InAllLandscapes);
+	void ResolveProceduralWeightmapTexture(const TArray<ALandscapeProxy*>& InAllLandscapes);
+	void ResolveProceduralTexture(FLandscapeProceduralTexture2DCPUReadBackResource* InCPUReadBackTexture, UTexture2D* InOriginalTexture);
+	void RegenerateProceduralWeightmaps();
+
+	bool AreHeightmapTextureResourcesReady(const TArray<ALandscapeProxy*>& InAllLandscapes) const;
+	bool AreWeightmapTextureResourcesReady(const TArray<ALandscapeProxy*>& InAllLandscapes) const;
+
+	void UpdateProceduralMaterialInstances(const TArray<ULandscapeComponent*>& InComponentsToUpdate, const TMap<ULandscapeComponent*, TArray<ULandscapeLayerInfoObject*>>& InZeroAllocationsPerComponents);
+
+	void PrepareProceduralComponentDataForExtractLayersCS(const FName& InProceduralLayerName, int32 InCurrentWeightmapToProcessIndex, bool InOutputDebugName, const TArray<ALandscapeProxy*>& InAllLandscape, class FLandscapeTexture2DResource* InOutTextureData,
+														  TArray<struct FLandscapeProceduralWeightmapExtractLayersComponentData>& OutComponentData, TMap<ULandscapeLayerInfoObject*, int32>& OutLayerInfoObjects);
+	void PrepareProceduralComponentDataForPackLayersCS(int32 InCurrentWeightmapToProcessIndex, bool InOutputDebugName, const TArray<ULandscapeComponent*>& InAllLandscapeComponents, 
+													   TArray<UTexture2D*>& InOutProcessedWeightmaps, TArray<class FLandscapeProceduralTexture2DCPUReadBackResource*>& InOutProcessedWeightmapCPUCopy, TArray<struct FLandscapeProceduralWeightmapPackLayersComponentData>& OutComponentData);
+	void ReallocateProceduralWeightmaps(const TArray<ALandscapeProxy*>& InAllLandscape, const TArray<ULandscapeLayerInfoObject*>& InBrushRequiredAllocations, TArray<ULandscapeComponent*>& OutComponentThatNeedMaterialRebuild);
+	void InitProceduralWeightmapResources(uint8 InLayerCount);
+	bool GenerateZeroAllocationPerComponents(const TArray<ALandscapeProxy*>& InAllLandscape, const TMap<ULandscapeLayerInfoObject*, bool>& InWeightmapLayersBlendSubstractive, TMap<ULandscapeComponent*, TArray<ULandscapeLayerInfoObject*>>& OutZeroAllocationsPerComponents);
+
+	void GenerateProceduralRenderQuad(const FIntPoint& InVertexPosition, float InVertexSize, const FVector2D& InUVStart, const FVector2D& InUVSize, TArray<struct FLandscapeProceduralTriangle>& OutTriangles) const;
+	void GenerateProceduralRenderQuadsAtlas(const FIntPoint& InSectionBase, const FVector2D& InScaleBias, float InSubSectionSizeQuad, const FIntPoint& InReadSize, const FIntPoint& InWriteSize, TArray<struct FLandscapeProceduralTriangle>& OutTriangles) const;
+	void GenerateProceduralRenderQuadsAtlasToNonAtlas(const FIntPoint& InSectionBase, const FVector2D& InScaleBias, float InSubSectionSizeQuad, const FIntPoint& InReadSize, const FIntPoint& InWriteSize, TArray<struct FLandscapeProceduralTriangle>& OutTriangles) const;
+	void GenerateProceduralRenderQuadsNonAtlas(const FIntPoint& InSectionBase, const FVector2D& InScaleBias, float InSubSectionSizeQuad, const FIntPoint& InReadSize, const FIntPoint& InWriteSize, TArray<struct FLandscapeProceduralTriangle>& OutTriangles) const;
+	void GenerateProceduralRenderQuadsNonAtlasToAtlas(const FIntPoint& InSectionBase, const FVector2D& InScaleBias, float InSubSectionSizeQuad, const FIntPoint& InReadSize, const FIntPoint& InWriteSize, TArray<struct FLandscapeProceduralTriangle>& OutTriangles) const;
+	void GenerateProceduralRenderQuadsMip(const FIntPoint& InSectionBase, const FVector2D& InScaleBias, float InSubSectionSizeQuad, const FIntPoint& InReadSize, const FIntPoint& InWriteSize, uint8 InCurrentMip, TArray<FLandscapeProceduralTriangle>& OutTriangles) const;
+
+	void ClearWeightmapTextureResource(const FString& InDebugName, FTextureRenderTargetResource* InTextureResourceToClear);
+	void DrawHeightmapComponentsToRenderTarget(const FString& InDebugName, const TArray<ULandscapeComponent*>& InComponentsToDraw, UTexture* InHeightmapRTRead, UTextureRenderTarget2D* InOptionalHeightmapRTRead2, UTextureRenderTarget2D* InHeightmapRTWrite, ERTDrawingType InDrawType,
+											   bool InClearRTWrite, struct FLandscapeHeightmapProceduralShaderParameters& InShaderParams, uint8 InMipRender = 0) const;
+
+	void DrawWeightmapComponentsToRenderTarget(const FString& InDebugName, const TArray<ULandscapeComponent*>& InComponentsToDraw, UTexture* InWeightmapRTRead, UTextureRenderTarget2D* InOptionalWeightmapRTRead2, UTextureRenderTarget2D* InWeightmapRTWrite,
+											   bool InClearRTWrite, struct FLandscapeWeightmapProceduralShaderParameters& InShaderParams, uint8 InMipRender) const;
+
+	void DrawWeightmapComponentsToRenderTarget(const FString& InDebugName, const FIntPoint& InSectionBase, const FVector2D& InScaleBias, UTexture* InWeightmapRTRead, UTextureRenderTarget2D* InOptionalWeightmapRTRead2, UTextureRenderTarget2D* InWeightmapRTWrite,
+									 		   bool InClearRTWrite, struct FLandscapeWeightmapProceduralShaderParameters& InShaderParams, uint8 InMipRender) const;
 
 	void DrawHeightmapComponentsToRenderTargetMips(TArray<ULandscapeComponent*>& InComponentsToDraw, UTexture* InReadHeightmap, bool InClearRTWrite, struct FLandscapeHeightmapProceduralShaderParameters& InShaderParams) const;
+	void DrawWeightmapComponentToRenderTargetMips(const FIntPoint& TopLeftTexturePosition, UTexture* InReadWeightmap, bool InClearRTWrite, struct FLandscapeWeightmapProceduralShaderParameters& InShaderParams) const;
 
-	void CopyProceduralTargetToResolveTarget(UTexture* InHeightmapRTRead, UTexture* InCopyResolveTarget, FTextureResource* InCopyResolveTargetCPUResource, const FIntPoint& InFirstComponentSectionBase, int32 InCurrentMip) const;
+	void CopyProceduralTexture(UTexture* InSourceTexture, UTexture* InDestTexture, FTextureResource* InDestCPUResource = nullptr, const FIntPoint& InFirstComponentSectionBase = FIntPoint(0, 0), uint8 InSourceCurrentMip = 0, uint8 InDestCurrentMip = 0,
+							   uint32 InSourceArrayIndex = 0, uint32 InDestArrayIndex = 0) const;
+	void CopyProceduralTexture(const FString& InSourceDebugName, FTextureResource* InSourceResource, const FString& InDestDebugName, FTextureResource* InDestResource, FTextureResource* InDestCPUResource = nullptr, const FIntPoint& InFirstComponentSectionBase = FIntPoint(0, 0),
+							   uint8 InSourceCurrentMip = 0, uint8 InDestCurrentMip = 0, uint32 InSourceArrayIndex = 0, uint32 uInDestArrayIndex = 0) const;
 
-	void PrintDebugRTHeightmap(FString Context, UTextureRenderTarget2D* InDebugRT, int32 InMipRender = 0, bool InOutputNormals = false) const;
-	void PrintDebugHeightData(const FString& InContext, const TArray<FColor>& InHeightmapData, const FIntPoint& InDataSize, int32 InMipRender, bool InOutputNormals = false) const;
-
-	void OnPreSaveWorld(uint32 SaveFlags, UWorld* World);
-	void OnPostSaveWorld(uint32 SaveFlags, UWorld* World, bool bSuccess);
+	void PrintProceduralDebugRT(const FString& InContext, UTextureRenderTarget2D* InDebugRT, uint8 InMipRender = 0, bool InOutputHeight = true, bool InOutputNormals = false) const;
+	void PrintProceduralDebugTextureResource(const FString& InContext, FTextureResource* InTextureResource, uint8 InMipRender = 0, bool InOutputHeight = true, bool InOutputNormals = false) const;
+	void PrintProceduralDebugHeightData(const FString& InContext, const TArray<FColor>& InHeightmapData, const FIntPoint& InDataSize, uint8 InMipRender, bool InOutputNormals = false) const;
+	void PrintProceduralDebugWeightData(const FString& InContext, const TArray<FColor>& InWeightmapData, const FIntPoint& InDataSize, uint8 InMipRender) const;
 #endif
 
+public:
+
 #if WITH_EDITORONLY_DATA
 	UPROPERTY(TextExportTransient)
 	TArray<FProceduralLayer> ProceduralLayers;
 
+	UPROPERTY(Transient)
+	TArray<UTextureRenderTarget2D*> HeightmapRTList;
+
+	UPROPERTY(Transient)
+	TArray<UTextureRenderTarget2D*> WeightmapRTList;
+
 	UPROPERTY(Transient)
 	bool PreviousExperimentalLandscapeProcedural;
 
+private:
+
+	UPROPERTY(Transient)
+	bool WasCompilingShaders;
+
 	UPROPERTY(Transient)
 	uint32 ProceduralContentUpdateFlags;
 
 	UPROPERTY(Transient)
-	TArray<UTextureRenderTarget2D*> HeightmapRTList;
+	bool ProceduralUpdateAllMaterials;
+
+	// Represent all the resolved paint layer, from all procedural layer blended together (size of the landscape x paint layer count)
+	class FLandscapeTexture2DArrayResource* CombinedProcLayerWeightmapAllLayersResource;
+	
+	// Represent all the resolved paint layer, from the current procedual layer only (size of the landscape x paint layer count)
+	class FLandscapeTexture2DArrayResource* CurrentProcLayerWeightmapAllLayersResource;	
+	
+	// Used in extracting the paint layers data from procedural layer weightmaps (size of the landscape)
+	class FLandscapeTexture2DResource* WeightmapScratchExtractLayerTextureResource;	
+	
+	// Used in packing the paint layers data contained into CombinedProcLayerWeightmapAllLayersResource to be set again for each component weightmap (size of the landscape)
+	class FLandscapeTexture2DResource* WeightmapScratchPackLayerTextureResource;	
 #endif
 };
diff --git a/Engine/Source/Runtime/Landscape/Classes/LandscapeBPCustomBrush.h b/Engine/Source/Runtime/Landscape/Classes/LandscapeBPCustomBrush.h
--- a/Engine/Source/Runtime/Landscape/Classes/LandscapeBPCustomBrush.h
+++ b/Engine/Source/Runtime/Landscape/Classes/LandscapeBPCustomBrush.h
@@ -20,6 +20,9 @@ private:
 	UPROPERTY(Category= "Settings", EditAnywhere, NonTransactional)
 	bool AffectWeightmap;
 
+	UPROPERTY(Category = "Settings", EditAnywhere, NonTransactional)
+	TArray<FName> AffectedWeightmapLayers;
+
 #if WITH_EDITORONLY_DATA
 	UPROPERTY(NonTransactional, DuplicateTransient)
 	class ALandscape* OwningLandscape;
@@ -43,6 +46,7 @@ public:
 
 	bool IsAffectingHeightmap() const { return AffectHeightmap; }
 	bool IsAffectingWeightmap() const { return AffectWeightmap; }
+	bool IsAffectingWeightmapLayer(const FName& InLayerName) const;
 
 	UFUNCTION(BlueprintImplementableEvent)
 	UTextureRenderTarget2D* Render(bool InIsHeightmap, UTextureRenderTarget2D* InCombinedResult);
diff --git a/Engine/Source/Runtime/Landscape/Classes/LandscapeComponent.h b/Engine/Source/Runtime/Landscape/Classes/LandscapeComponent.h
--- a/Engine/Source/Runtime/Landscape/Classes/LandscapeComponent.h
+++ b/Engine/Source/Runtime/Landscape/Classes/LandscapeComponent.h
@@ -251,6 +251,46 @@ struct FLandscapeComponentMaterialOverride
 	UMaterialInterface* Material;
 };
 
+/** Structure storing channel usage for weightmap textures */
+USTRUCT(NotBlueprintable)
+struct FLandscapeWeightmapUsage
+{
+	GENERATED_USTRUCT_BODY()
+
+	UPROPERTY()
+	ULandscapeComponent* ChannelUsage[4];
+
+	UPROPERTY()
+	FName ProceduralLayerName;
+
+	FLandscapeWeightmapUsage()
+	{
+		ClearUsage();
+	}
+
+	friend FArchive& operator<<(FArchive& Ar, FLandscapeWeightmapUsage& U);
+
+	int32 FreeChannelCount() const
+	{
+		int32 Count = 0;
+
+		for (int8 i = 0; i < 4; ++i)
+		{
+			Count += (ChannelUsage[i] == nullptr) ? 1 : 0;
+		}
+
+		return	Count;
+	}
+
+	void ClearUsage()
+	{
+		for (int8 i = 0; i < 4; ++i)
+		{
+			ChannelUsage[i] = nullptr;
+		}
+	}
+};
+
 UCLASS(hidecategories=(Display, Attachment, Physics, Debug, Collision, Movement, Rendering, PrimitiveComponent, Object, Transform, Mobility), showcategories=("Rendering|Material"), MinimalAPI, Within=LandscapeProxy)
 class ULandscapeComponent : public UPrimitiveComponent
 {
@@ -304,14 +344,6 @@ class ULandscapeComponent : public UPrimitiveComponent
 	UPROPERTY(TextExportTransient)
 	TArray<int8> MaterialIndexToDisabledTessellationMaterial;
 
-	/** List of layers, and the weightmap and channel they are stored */
-	UPROPERTY()
-	TArray<FWeightmapLayerAllocationInfo> WeightmapLayerAllocations;
-
-	/** Weightmap texture reference */
-	UPROPERTY(TextExportTransient)
-	TArray<UTexture2D*> WeightmapTextures;
-
 	/** XYOffsetmap texture reference */
 	UPROPERTY(TextExportTransient)
 	UTexture2D* XYOffsetmapTexture;
@@ -342,15 +374,28 @@ private:
 	UPROPERTY()
 	FGuid LightingGuid;
 
-	/** Heightmap texture reference */
-	UPROPERTY(Transient, TextExportTransient)
+	/** Current data we're working on (only used in Procedural mode) This data will get set by tools as needed */
 	UTexture2D* CurrentEditingHeightmapTexture;
+	TArray<FWeightmapLayerAllocationInfo>* CurrentEditingWeightmapLayerAllocations;
+	TArray<UTexture2D*>* CurrentEditingWeightmapTextures;
+	TArray<FLandscapeWeightmapUsage*>* CurrentEditingWeightmapTexturesUsage;
+	FName CurrentProceduralLayerName;
+
+	TArray<FLandscapeWeightmapUsage*> WeightmapTexturesUsage;
 #endif // WITH_EDITORONLY_DATA
 
 	/** Heightmap texture reference */
 	UPROPERTY(TextExportTransient)
 	UTexture2D* HeightmapTexture;
 
+	/** List of layers, and the weightmap and channel they are stored */
+	UPROPERTY()
+	TArray<FWeightmapLayerAllocationInfo> WeightmapLayerAllocations;
+
+	/** Weightmap texture reference */
+	UPROPERTY(TextExportTransient)
+	TArray<UTexture2D*> WeightmapTextures;
+
 public:
 
 	/** Uniquely identifies this component's built map data. */
@@ -501,9 +546,31 @@ public:
 	virtual bool IsPrecomputedLightingValid() const override;
 
 	LANDSCAPE_API UTexture2D* GetHeightmap(bool InReturnCurrentEditingHeightmap = false) const;
+	LANDSCAPE_API TArray<UTexture2D*>& GetWeightmapTextures(bool InReturnCurrentEditingWeightmap = false);
+	LANDSCAPE_API const TArray<UTexture2D*>& GetWeightmapTextures(bool InReturnCurrentEditingWeightmap = false) const;
+
+	LANDSCAPE_API TArray<FWeightmapLayerAllocationInfo>& GetWeightmapLayerAllocations(bool InReturnCurrentEditingWeightmap = false);
+	LANDSCAPE_API const TArray<FWeightmapLayerAllocationInfo>& GetWeightmapLayerAllocations(bool InReturnCurrentEditingWeightmap = false) const;
+
+#if WITH_EDITOR
 	LANDSCAPE_API void SetHeightmap(UTexture2D* NewHeightmap);
 	LANDSCAPE_API void SetCurrentEditingHeightmap(UTexture2D* InNewHeightmap);
 
+	LANDSCAPE_API void SetWeightmapTextures(const TArray<UTexture2D*>& InNewWeightmapTextures, bool InApplyToCurrentEditingWeightmap = false);
+	LANDSCAPE_API void SetCurrentEditingWeightmaps(TArray<UTexture2D*>* InNewWeightmapTextures);
+	LANDSCAPE_API void SetCurrentProceduralLayerName(const FName& InName);
+	LANDSCAPE_API const FName& GetCurrentProceduralLayerName() const;
+
+	LANDSCAPE_API void SetWeightmapLayerAllocations(const TArray<FWeightmapLayerAllocationInfo>& InNewWeightmapLayerAllocations);
+	LANDSCAPE_API void SetCurrentEditingWeightmapLayerAllocations(TArray<FWeightmapLayerAllocationInfo>* InNewWeightmapLayerAllocations);
+	LANDSCAPE_API void SetWeightmapTexturesUsage(const TArray<FLandscapeWeightmapUsage*>& InNewWeightmapTexturesUsage, bool InApplyToCurrentEditingWeightmap = false);
+	LANDSCAPE_API void SetCurrentEditingWeightmapTexturesUsage(TArray<FLandscapeWeightmapUsage*>* InNewWeightmapTexturesUsage);
+
+	LANDSCAPE_API TArray<FLandscapeWeightmapUsage*>& GetWeightmapTexturesUsage(bool InReturnCurrentEditingWeightmap = false);
+	LANDSCAPE_API const TArray<FLandscapeWeightmapUsage*>& GetWeightmapTexturesUsage(bool InReturnCurrentEditingWeightmap = false) const;
+
+#endif 
+
 #if WITH_EDITOR
 	virtual int32 GetNumMaterials() const override;
 	virtual UMaterialInterface* GetMaterial(int32 ElementIndex) const override;
@@ -721,7 +788,7 @@ public:
 	/**
 	 * Create weightmaps for this component for the layers specified in the WeightmapLayerAllocations array
 	 */
-	void ReallocateWeightmaps(FLandscapeEditDataInterface* DataInterface=NULL);
+	void ReallocateWeightmaps(FLandscapeEditDataInterface* DataInterface= nullptr, bool InCanUseCurrentEditingWeightmap = true, bool InSaveToTransactionBuffer = true, bool InInitPlatformDataAsync = false, TArray<UTexture2D*>* OutNewCreatedTextures = nullptr);
 
 	/** Returns the actor's LandscapeMaterial, or the Component's OverrideLandscapeMaterial if set */
 	LANDSCAPE_API UMaterialInterface* GetLandscapeMaterial(int8 InLODIndex = INDEX_NONE) const;
diff --git a/Engine/Source/Runtime/Landscape/Classes/LandscapeProxy.h b/Engine/Source/Runtime/Landscape/Classes/LandscapeProxy.h
--- a/Engine/Source/Runtime/Landscape/Classes/LandscapeProxy.h
+++ b/Engine/Source/Runtime/Landscape/Classes/LandscapeProxy.h
@@ -13,6 +13,7 @@
 #include "Engine/Texture.h"
 #include "PerPlatformProperties.h"
 #include "LandscapeBPCustomBrush.h"
+#include "LandscapeComponent.h"
 #include "LandscapeProxy.generated.h"
 
 class ALandscape;
@@ -35,32 +36,6 @@ struct FLandscapeInfoLayerSettings;
 struct FMeshDescription;
 enum class ENavDataGatheringMode : uint8;
 
-/** Structure storing channel usage for weightmap textures */
-USTRUCT()
-struct FLandscapeWeightmapUsage
-{
-	GENERATED_USTRUCT_BODY()
-
-	UPROPERTY()
-	ULandscapeComponent* ChannelUsage[4];
-
-	FLandscapeWeightmapUsage()
-	{
-		ChannelUsage[0] = nullptr;
-		ChannelUsage[1] = nullptr;
-		ChannelUsage[2] = nullptr;
-		ChannelUsage[3] = nullptr;
-	}
-	friend FArchive& operator<<( FArchive& Ar, FLandscapeWeightmapUsage& U );
-	int32 FreeChannelCount() const
-	{
-		return	((ChannelUsage[0] == nullptr) ? 1 : 0) +
-				((ChannelUsage[1] == nullptr) ? 1 : 0) +
-				((ChannelUsage[2] == nullptr) ? 1 : 0) +
-				((ChannelUsage[3] == nullptr) ? 1 : 0);
-	}
-};
-
 USTRUCT()
 struct FLandscapeEditorLayerSettings
 {
@@ -399,6 +374,20 @@ struct FRenderDataPerHeightmap
 	FIntPoint TopLeftSectionBase;
 };
 
+USTRUCT()
+struct FWeightmapLayerData
+{
+	GENERATED_USTRUCT_BODY()
+
+	UPROPERTY()
+	TArray<UTexture2D*> Weightmaps;
+
+	UPROPERTY()
+	TArray<FWeightmapLayerAllocationInfo> WeightmapLayerAllocations;
+
+	TArray<FLandscapeWeightmapUsage*> WeightmapTextureUsages;	// Easy Access ref to data stored into the LandscapeProxy weightmap usage map
+};
+
 USTRUCT()
 struct FProceduralLayerData
 {
@@ -408,9 +397,10 @@ struct FProceduralLayerData
 	{}
 
 	UPROPERTY()
-	TMap<UTexture2D*, UTexture2D*> Heightmaps;
+	TMap<UTexture2D*, UTexture2D*> Heightmaps; // Mapping between Original Heightmap -> Layer Heightmap
 
-	// TODO: add weightmap data
+	UPROPERTY()
+	TMap<ULandscapeComponent*, FWeightmapLayerData> WeightmapData; // Weightmaps per components
 };
 
 UCLASS(Abstract, MinimalAPI, NotBlueprintable, hidecategories=(Display, Attachment, Physics, Debug, Lighting, LOD), showcategories=(Lighting, Rendering, "Utilities|Transformation"), hidecategories=(Mobility))
@@ -659,6 +649,8 @@ public:
 	UPROPERTY(Transient)
 	TMap<UTexture2D*, FRenderDataPerHeightmap> RenderDataPerHeightmap; // Mapping between Original heightmap and general render data
 
+	TMap<UTexture2D*, FLandscapeProceduralTexture2DCPUReadBackResource*> WeightmapCPUReadBackTextures; // Mapping between Original weightmap and tyhe CPU readback resource
+
 	FRenderCommandFence ReleaseResourceFence;
 #endif
 
@@ -704,7 +696,7 @@ public:
 #endif
 
 	/** Map of weightmap usage */
-	TMap<UTexture2D*, FLandscapeWeightmapUsage> WeightmapUsageMap;
+	TMap<UTexture2D*, FLandscapeWeightmapUsage*> WeightmapUsageMap;
 
 	// Blueprint functions
 
