Fixed copyright issue

diff --git a/Engine/Shaders/Private/LandscapeProceduralCS.usf b/Engine/Shaders/Private/LandscapeProceduralCS.usf
--- a/Engine/Shaders/Private/LandscapeProceduralCS.usf
+++ b/Engine/Shaders/Private/LandscapeProceduralCS.usf
@@ -1,4 +1,4 @@
-// Copyright 1998-2018 Epic Games, Inc. All Rights Reserved.
+// Copyright 1998-2019 Epic Games, Inc. All Rights Reserved.
 
 #include "Common.ush"
 
