Fixed many issues related to reflection
- Clarified top and bottom layers. Default material only has top layer.
- Reverted clear coat code doing an extra loop for sky sampling causing issues on other materials (wrongly receiving sky color)
- Fixed clear coat appearance when MaxDistance!=-1
- Fixed black reflexions due to R pointing into the surface using the clear coat normal (assuming it is smooth so only a half kind of a fix, a proper fix would involve sending the geometry normal in the payload)
- Tested: works with fog, with/without reflection volume, and with/without deferred material evaluation

diff --git a/Engine/Shaders/Private/RayTracing/RayTracingLightingCommon.ush b/Engine/Shaders/Private/RayTracing/RayTracingLightingCommon.ush
--- a/Engine/Shaders/Private/RayTracing/RayTracingLightingCommon.ush
+++ b/Engine/Shaders/Private/RayTracing/RayTracingLightingCommon.ush
@@ -476,7 +476,7 @@ float3 ComputeDirectLightingCulled(
 	return DirectLighting;
 }
 
-void ApplyViewDependentMaterialPayloadModifications(RayDesc Ray, inout FMaterialClosestHitPayload Payload)
+void ComputeBottomLayerMaterialProperties(RayDesc Ray, inout FMaterialClosestHitPayload Payload)
 {
 	if (Payload.ShadingModelID == SHADINGMODELID_CLEAR_COAT)
 	{
@@ -503,7 +503,7 @@ void ApplyViewDependentMaterialPayloadModifications(RayDesc Ray, inout FMaterial
 	}
 }
 
-FMaterialClosestHitPayload TraceRayAndAccumulateResults(
+FMaterialClosestHitPayload TraceRayResults(
 	in RayDesc Ray,
 	in RandomSequence RandSequence,
 	in uint RayFlags,
@@ -522,20 +522,25 @@ FMaterialClosestHitPayload TraceRayAndAccumulateResults(
 		Ray,								// RayDesc
 		RayCone
 	);
-
-	uint2 LightCullMask = (uint2)0;
-	if (IsHit(Ray, PackedPayload))
-	{
-		const float3 WorldPosition = Ray.Origin + PackedPayload.HitT * Ray.Direction;
-		const float3 WorldNormal = UnpackRayTracingNormal(PackedPayload);
-		LightCullMask = CullDirectLighting(Ray, RayCone, WorldPosition, WorldNormal, RandSequence, ReflectedShadowsType > 0, ShadowMaxNormalBias);
-	}
-
 	FMaterialClosestHitPayload Payload = UnpackRayTracingPayload(PackedPayload, Ray);
+	return Payload;
+}
 
-	if (IsHit(Ray, PackedPayload))
+void AccumulateResults(
+	in FMaterialClosestHitPayload Payload,
+	in RayDesc Ray,
+	in RandomSequence RandSequence,
+	in uint RayFlags,
+	in float ShadowMaxNormalBias,
+	in uint Mask,
+	inout FRayCone RayCone,
+	inout float3 Radiance)
+{
+	if (IsHit(Ray, Payload))
 	{
-		ApplyViewDependentMaterialPayloadModifications(Ray, Payload);
+		uint2 LightCullMask = (uint2)0;
+		const float3 WorldPosition = Ray.Origin + Payload.HitT * Ray.Direction;
+		LightCullMask = CullDirectLighting(Ray, RayCone, WorldPosition, Payload.WorldNormal, RandSequence, ReflectedShadowsType > 0, ShadowMaxNormalBias);
 
 		Payload.DiffuseColor = Payload.BaseColor - Payload.BaseColor * Payload.Metallic;
 		Payload.SpecularColor = ComputeF0(Payload.Specular, Payload.BaseColor, Payload.Metallic);
@@ -556,6 +561,68 @@ FMaterialClosestHitPayload TraceRayAndAccumulateResults(
 			Radiance += Payload.DiffuseColor * Payload.IndirectIrradiance;
 		}
 	}
+}
+
+FMaterialClosestHitPayload TraceRayAndAccumulateResults(
+	in RayDesc Ray,
+	in RandomSequence RandSequence,
+	in uint RayFlags,
+	in float ShadowMaxNormalBias,
+	in uint Mask,
+	inout FRayCone RayCone,
+	inout float3 Radiance)
+{
+	FMaterialClosestHitPayload Payload = TraceRayResults(
+		Ray,
+		RandSequence,
+		RayFlags,
+		ShadowMaxNormalBias,
+		Mask,
+		RayCone,
+		Radiance);
+
+	AccumulateResults(
+		Payload,
+		Ray,
+		RandSequence,
+		RayFlags,
+		ShadowMaxNormalBias,
+		Mask,
+		RayCone,
+		Radiance);
 
 	return Payload;
 }
+
+FMaterialClosestHitPayload TraceRayAndAccumulateBottomLayerResults(
+	in RayDesc Ray,
+	in RandomSequence RandSequence,
+	in uint RayFlags,
+	in float ShadowMaxNormalBias,
+	in uint Mask,
+	inout FRayCone RayCone,
+	inout float3 Radiance)
+{
+	FMaterialClosestHitPayload BottomLayerPayload = TraceRayResults(
+		Ray,
+		RandSequence,
+		RayFlags,
+		ShadowMaxNormalBias,
+		Mask,
+		RayCone,
+		Radiance);
+
+	ComputeBottomLayerMaterialProperties(Ray, BottomLayerPayload);
+
+	AccumulateResults(
+		BottomLayerPayload,
+		Ray,
+		RandSequence,
+		RayFlags,
+		ShadowMaxNormalBias,
+		Mask,
+		RayCone,
+		Radiance);
+
+	return BottomLayerPayload;
+}
diff --git a/Engine/Shaders/Private/RayTracing/RayTracingReflections.usf b/Engine/Shaders/Private/RayTracing/RayTracingReflections.usf
--- a/Engine/Shaders/Private/RayTracing/RayTracingReflections.usf
+++ b/Engine/Shaders/Private/RayTracing/RayTracingReflections.usf
@@ -83,10 +83,46 @@ uint2 GetPixelCoord(uint2 DispatchThreadId)
 	return DispatchThreadId * UpscaleFactor + uint2(SubPixelId & (UpscaleFactor - 1), SubPixelId / UpscaleFactor);
 }
 
+// Generate a random direction to sample according to world normal and roughness.
+float3 GenerateReflectionDirection(inout RandomSequence RandSequence, uint2 DispatchThreadId, uint BounceIndex, uint SamplePerPixel, float3 IncidentDirection, float3 WorldNormal, float Roughness)
+{
+	uint DummyVariable;
+	int RetryCount = 0;
+	float3 outDirection;
+	do
+	{
+		float2 RandSample = RandomSequence_GenerateSample2D(RandSequence, DummyVariable);
+
+		// Sobol can only be used on the first sample of the first bounce
+		if (BounceIndex == 0 && SamplePerPixel == 1 && RetryCount == 0)
+		{
+			uint2 SobolFrame = ComputePixelUniqueSobolRandSample(DispatchThreadId);
+			RandSample = SobolIndexToUniformUnitSquare(SobolFrame);
+		}
+
+		outDirection = GenerateReflectedRayDirection(IncidentDirection, WorldNormal, Roughness, RandSample);
+	} while (dot(WorldNormal, outDirection) < 0.0 && ++RetryCount < 16);
+
+	return outDirection;
+}
+
+// Helper function to ensure a reflection direction is not pointing inside a geometry.
+// It uses the geometry smooth normal to reflect R out of it.
+void FixSampleDirectionIfNeeded(float3 SmoothSurfaceNormal, inout float3 SampleDirection)
+{
+	if (dot(SmoothSurfaceNormal, SampleDirection) < 0.0)
+	{
+		// The sampling direction is pointing towards the surface, so revert it along the normal axis.
+		SampleDirection = SampleDirection - 2.0f * dot(SmoothSurfaceNormal, SampleDirection) * SmoothSurfaceNormal;
+	}
+}
 
 [shader("raygeneration")]
 void RayTracingReflectionsRGS()
 {
+	uint DummyVariable;
+	const float ClearCoatSpecularColor = 0.04; // See ReflectionEnvironmentPixelShader.usf
+
 	FDeferredMaterialPayload DeferredMaterialPayload;
 	DeferredMaterialPayload.SortKey = RAY_TRACING_DEFERRED_MATERIAL_KEY_INVALID;
 
@@ -164,25 +200,27 @@ void RayTracingReflectionsRGS()
 			float3 CameraOrigin = ReconstructWorldPositionFromDepth(UV, 0.0f);
 			float3 WorldPosition = ReconstructWorldPositionFromDepth(UV, Depth);
 			float3 CameraDirection = normalize(WorldPosition - CameraOrigin);
-			float3 WorldNormal = ScreenSpaceData.GBuffer.WorldNormal;
-			float Roughness = ScreenSpaceData.GBuffer.Roughness;
-			float3 SpecularColor = ScreenSpaceData.GBuffer.SpecularColor;
 			float3 IncidentDirection = CameraDirection;
-			float ClearCoat = 1.0f;
-			float ClearCoatRoughness = 1.0f;
-			#if CLEAR_COAT_BOTTOM_NORMAL
-			float3 ClearCoatBottomNormal = WorldNormal;
-			#endif
-			
+
 			uint ShadingModelID = ScreenSpaceData.GBuffer.ShadingModelID;
+			float3 TopLayerWorldNormal = ScreenSpaceData.GBuffer.WorldNormal;
+			float  TopLayerRoughness = ScreenSpaceData.GBuffer.Roughness;
+			float3 TopLayerSpecularColor = ScreenSpaceData.GBuffer.SpecularColor;
+
+			float  ClearCoat = 1.0f;
+			float  BottomLayerRoughness = TopLayerRoughness;
+			float3 BottomLayerSpecularColor = TopLayerSpecularColor;
+			float3 BottomLayerWorldNormal = TopLayerWorldNormal;
 			if (ShadingModelID == SHADINGMODELID_CLEAR_COAT)
 			{
-				ClearCoat = ScreenSpaceData.GBuffer.CustomData.x;
-				ClearCoatRoughness = ScreenSpaceData.GBuffer.CustomData.y;
 				#if CLEAR_COAT_BOTTOM_NORMAL
-					const float2 oct1 = ((float2(ScreenSpaceData.GBuffer.CustomData.a, ScreenSpaceData.GBuffer.CustomData.z) * 2) - (256.0/255.0)) + UnitVectorToOctahedron(ScreenSpaceData.GBuffer.WorldNormal);
-					ClearCoatBottomNormal = OctahedronToUnitVector(oct1);
+				const float2 oct1 = ((float2(ScreenSpaceData.GBuffer.CustomData.a, ScreenSpaceData.GBuffer.CustomData.z) * 2) - (256.0/255.0)) + UnitVectorToOctahedron(TopLayerWorldNormal);
+				BottomLayerWorldNormal = OctahedronToUnitVector(oct1);
 				#endif
+
+				ClearCoat = ScreenSpaceData.GBuffer.CustomData.x;			// Clear coat weight
+				TopLayerRoughness = ScreenSpaceData.GBuffer.CustomData.y;	// Clear coat roughness
+				TopLayerSpecularColor = ClearCoatSpecularColor;				// Hard coded top layer specular color
 			}
 
 			FRayCone RayCone = (FRayCone)0;
@@ -190,68 +228,27 @@ void RayTracingReflectionsRGS()
 			float SurfaceCurvature = 0.0f; /* #todo_dxr assume no curvature */
 			RayCone = PropagateRayCone(RayCone, SurfaceCurvature, Depth);
 
-			uint DummyVariable;
-
-			// Add one more 'fake bounce' and we will evaluate cubemap there
-			for (uint BounceIndex = 0; BounceIndex < MaxBounces + 1; ++BounceIndex)
+			//
+			// Reflections are trace by always following the top layer reflection vector. This is recursive and up to MaxBounces recursion.
+			// In the case of coated material, we also trace a reflection ray for the bottom layer which can have different normal and roughness. This is not recursive.
+			// We fix normal pointing away towards the geometry by assuming coated layer has smooth normal (assuming it is the mesh smooth normal).
+			// Fog is applied as a function of the top layer path, for both top and botto; refelctions combined (as an optimisation approximation).
+			//
+			for (uint BounceIndex = 0; BounceIndex < MaxBounces; ++BounceIndex)
 			{
-				float NoV = saturate(dot(-IncidentDirection, WorldNormal));
-				float FinalRoughness = Roughness;
-				if (ShadingModelID == SHADINGMODELID_CLEAR_COAT)
-				{
-					float F = EnvBRDF( 0.04, ClearCoatRoughness, NoV ).x;
-					F *= ClearCoat;
-					
-					float w = RandomSequence_GenerateSample1D(RandSequence, DummyVariable);
-					
-					if (w <= F)
-					{
-						SpecularColor = 0.04;
-						FinalRoughness = ClearCoatRoughness;
-					}
-					else
-					{
-						#if CLEAR_COAT_BOTTOM_NORMAL
-							WorldNormal = ClearCoatBottomNormal;
-							NoV = saturate(dot(-IncidentDirection, WorldNormal));
-						#endif
-						PathThroughput *= EnvBRDF( SpecularColor, Roughness, NoV );
-					}
-				}
-				else
+				// Prepare a ray to trace the top layer reflection
+				RayDesc TopLayerRay;
+				TopLayerRay.Origin = WorldPosition;
+				TopLayerRay.TMin = 0.0;
+				TopLayerRay.TMax = LocalMaxRayDistance;
+
+				TopLayerRay.Direction = GenerateReflectionDirection(RandSequence, DispatchThreadId, BounceIndex, LocalSamplesPerPixel, IncidentDirection, TopLayerWorldNormal, TopLayerRoughness);
+				ApplyPositionBias(TopLayerRay, TopLayerWorldNormal, ReflectionMaxNormalBias);
+				if (BounceIndex == 0)
 				{
-					// If it is not clear coat, and we're on the first bounce, the G F terms will be applied by ReflectionEnvironmentPixelShader
-					if (BounceIndex > 0)
-						PathThroughput *= EnvBRDF( SpecularColor, Roughness, NoV );
-				}
-
-				RayDesc Ray;
-				Ray.Origin = WorldPosition;
-				Ray.TMin = 0.0;
-				Ray.TMax = LocalMaxRayDistance;
-
-				int RetryCount = 0;
-				do
-				{
-					float2 RandSample = RandomSequence_GenerateSample2D(RandSequence, DummyVariable);
-
-					// Sobol can only be used on the first sample of the first bounce
-					if (BounceIndex == 0 && LocalSamplesPerPixel == 1 && RetryCount == 0)
-					{
-						uint2 SobolFrame = ComputePixelUniqueSobolRandSample(DispatchThreadId);
-						RandSample = SobolIndexToUniformUnitSquare(SobolFrame);
-					}
-				
-					Ray.Direction = GenerateReflectedRayDirection(IncidentDirection, WorldNormal, FinalRoughness, RandSample);
-				} while (dot(WorldNormal, Ray.Direction) < 0.0 && ++RetryCount < 16);
-				
-				ApplyPositionBias(Ray, WorldNormal, ReflectionMaxNormalBias);
-				
-				if (dot(WorldNormal, Ray.Direction) < 0.0)
-				{
-					// Refraction rays are handled separately
-					bIsValidSample = false;
-					break;
+					// This fixes black reflexions due to R pointing into the surface.
+					// It assumes that the clear coat relies on a smooth normal following the mesh surface.
+					FixSampleDirectionIfNeeded(ScreenSpaceData.GBuffer.WorldNormal, TopLayerRay.Direction);
 				}
 
 				//#dxr-todo: apply heuristics to do different shading computations depending on BounceIndex
@@ -262,13 +259,13 @@ void RayTracingReflectionsRGS()
 					{
 						uint RayFlags = RAY_FLAG_CULL_BACK_FACING_TRIANGLES;
 						TraceRay(
-							TLAS,      // AccelerationStructure
+							TLAS,								// AccelerationStructure
 							RayFlags,
-							RAY_TRACING_MASK_OPAQUE,          // InstanceInclusionMask -- NOTE: Reflection rays are traced only against opaque geometry. Blended geometry needs special handling (#dxr_todo).
-							RAY_TRACING_SHADER_SLOT_MATERIAL, // RayContributionToHitGroupIndex
-							RAY_TRACING_NUM_SHADER_SLOTS,     // MultiplierForGeometryContributionToShaderIndex
-							0,         // MissShaderIndex
-							Ray,       // RayDesc
+							RAY_TRACING_MASK_OPAQUE,			// InstanceInclusionMask -- NOTE: Reflection rays are traced only against opaque geometry. Blended geometry needs special handling (#dxr_todo).
+							RAY_TRACING_SHADER_SLOT_MATERIAL,	// RayContributionToHitGroupIndex
+							RAY_TRACING_NUM_SHADER_SLOTS,		// MultiplierForGeometryContributionToShaderIndex
+							0,									// MissShaderIndex
+							TopLayerRay,						// RayDesc
 							DeferredMaterialPayload
 						);
 
@@ -279,60 +276,58 @@ void RayTracingReflectionsRGS()
 						if (DeferredMaterialPayload.SortKey < RAY_TRACING_DEFERRED_MATERIAL_KEY_RAY_MISS)
 						{
 							const float ShortRayLength = 1.0f; // 1cm is arbitrarily chosen
-							Ray.TMin = max(0, DeferredMaterialPayload.HitT - ShortRayLength * 0.5f);
+							TopLayerRay.TMin = max(0, DeferredMaterialPayload.HitT - ShortRayLength * 0.5f);
 
 							// TMax can only be used if alpha masked materials are handled during initial ray tracing (material gathering)
 							// Ray.TMax = DeferredMaterialPayload.HitT + ShortRayLength * 0.5f;
 						}
 						else
 						{
-							Ray.TMax = 0;
+							TopLayerRay.TMax = 0;
 						}
 					}
 				}
 
-				float3 Radiance = float3(0, 0, 0);
+
+
+				// Trace the top layer reflection
+				float3 TopLayerRadiance = float3(0, 0, 0);
 				FMaterialClosestHitPayload Payload;
+				const uint RayFlags = 0;					
+				Payload = TraceRayAndAccumulateResults(TopLayerRay, RandSequence, RayFlags, ReflectionMaxNormalBias, RAY_TRACING_MASK_OPAQUE, RayCone, TopLayerRadiance);
+				bool isTopLayerRayValid = Payload.HitT >= 0 || (bAllowSkySampling && ReflectionStruct.SkyLightParameters.y > 0);
 
-				// Force the ray to miss on the MaxBounces + 1 fake bounce, so that we fallback to cubemap there
-				if (BounceIndex < MaxBounces)
-				{
-					const uint RayFlags = 0;					
-					Payload = TraceRayAndAccumulateResults(Ray, RandSequence, RayFlags, ReflectionMaxNormalBias, RAY_TRACING_MASK_OPAQUE, RayCone, Radiance);
-				}
 
+				// Compute some flags and values
 				bool bApplyHeightFog = false;
+				bool bSkyWasHit = false;
 				float3 OriginToCollider = 0.0f;
-
-				// Force the ray to miss on the MaxBounces + 1 fake bounce, so that we fallback to cubemap there
-				float3 PathSegmentRadiance = 0.0f;
-				if (Payload.HitT >= 0 && BounceIndex < MaxBounces)
+				if (Payload.HitT >= 0)
 				{
-					PathSegmentRadiance += Radiance;
 					bIsValidSample = true;
-
 					if (BounceIndex == 0)
 					{
 						HitDistance = Payload.HitT;
 					}
 
 					bApplyHeightFog = HeightFog > 0;
-					OriginToCollider = Payload.WorldPos - Ray.Origin;
+					OriginToCollider = Payload.WorldPos - TopLayerRay.Origin;
 				}
 				else if (bAllowSkySampling && ReflectionStruct.SkyLightParameters.y > 0)
 				{
 					float SkyAverageBrightness = 1.0f;
-					float3 SkyLighting = GetSkyLightReflection(Ray.Direction, FinalRoughness, SkyAverageBrightness);
-					PathSegmentRadiance += SkyLighting;
+					float3 SkyLighting = GetSkyLightReflection(TopLayerRay.Direction, TopLayerRoughness, SkyAverageBrightness);
+					TopLayerRadiance = SkyLighting;
 					bIsValidSample = true;
+					bSkyWasHit = true;	// To stop the main path recusion early
 
 					bApplyHeightFog = HeightFog > 0;
 					const float SkyFarDistance = 100000.0f;
-					OriginToCollider = Ray.Direction * SkyFarDistance;
+					OriginToCollider = TopLayerRay.Direction * SkyFarDistance;
 
 					if (BounceIndex == 0)
 					{
-						// To tell the denoiser the black input is valid even if no valid ray trace result was found. Uses a high ray hit distance
+						// To tell the denoiser the sky radiance is valid even if no valid ray trace result was found. Uses a high ray hit distance
 						// so that the sample may be reused by another neighbor pixel in the denoise as if RoughnessFade in ]0;1[.
 						HitDistance = 1.0e20;
 					}
@@ -344,40 +339,127 @@ void RayTracingReflectionsRGS()
 						// Tell to the denoiser this sample reached outer bounds of the scene without any intersection.
 						HitDistance = -1;
 					}
+					TopLayerRadiance = 0;
+				}
+
+				float3 BottomLayerRadiance = float3(0, 0, 0);
+				float3 TopLayerReflectionEventThroughput = float3(1, 1, 1);
+				float NoV = saturate(dot(-IncidentDirection, TopLayerWorldNormal));
+				BRANCH
+				if (ShadingModelID == SHADINGMODELID_CLEAR_COAT)
+				{
+					// Trace reflection ray for the bottom layer
+					RayDesc BottomLayerRay;
+					BottomLayerRay.Origin = WorldPosition;
+					BottomLayerRay.TMin = 0.0;
+					BottomLayerRay.TMax = LocalMaxRayDistance;
+					BottomLayerRay.Direction = GenerateReflectionDirection(RandSequence, DispatchThreadId, BounceIndex, LocalSamplesPerPixel, IncidentDirection, BottomLayerWorldNormal, BottomLayerRoughness);
+					ApplyPositionBias(BottomLayerRay, BottomLayerWorldNormal, ReflectionMaxNormalBias);
+					if (BounceIndex == 0 && dot(ScreenSpaceData.GBuffer.WorldNormal, BottomLayerRay.Direction) < 0.0)
+					{
+						FixSampleDirectionIfNeeded(ScreenSpaceData.GBuffer.WorldNormal, BottomLayerRay.Direction);
+					}
+
+					
+					FRayCone BottomLayerRayCone = (FRayCone)0;
+					FMaterialClosestHitPayload BottomLayerPayload = TraceRayAndAccumulateBottomLayerResults(BottomLayerRay, RandSequence, RayFlags, ReflectionMaxNormalBias, RAY_TRACING_MASK_OPAQUE, RayCone, BottomLayerRadiance);
+
+					// Distance and sky: similar processing to main ray.
+					float BottomLayerHitDistance = -1;
+					bool IsBottomSampleValid = false;
+					if (BottomLayerPayload.HitT >= 0)
+					{
+						IsBottomSampleValid = true;
+						BottomLayerHitDistance = BottomLayerPayload.HitT;
+					}
+					else if (bAllowSkySampling && ReflectionStruct.SkyLightParameters.y > 0)
+					{
+						IsBottomSampleValid = true;
+						float SkyAverageBrightness = 1.0f;
+						float3 SkyLighting = GetSkyLightReflection(BottomLayerRay.Direction, BottomLayerRoughness, SkyAverageBrightness);
+						BottomLayerRadiance += SkyLighting;
+						BottomLayerHitDistance = 1.0e20;
+					}
+					
+					if (BounceIndex == 0)
+					{
+						HitDistance = IsBottomSampleValid ? min(bIsValidSample ? HitDistance : 1.0e27f, BottomLayerHitDistance) : HitDistance;
+					}
+
+
+					// Apply correct weights to the bottom and top layers (See ReflectionEnvironmentPixelShader)
+					{
+						float NoVBottom = saturate(dot(-IncidentDirection, BottomLayerWorldNormal));
+
+						float TopLayerFSChlick = EnvBRDF(TopLayerSpecularColor, TopLayerRoughness, NoV).x;
+						float BottomLayerFSChlick = EnvBRDF(BottomLayerSpecularColor, BottomLayerRoughness, NoVBottom).x;
+
+						float TopLayerRefractionTransmittance = (1 - TopLayerFSChlick);
+
+						float2 AB = PreIntegratedGF.SampleLevel(PreIntegratedGFSampler, float2(NoV, TopLayerRoughness), 0).rg;
+						float3 TopLayerSpecularLobIntegral = TopLayerSpecularColor * AB.x + AB.y * saturate(50.0 * TopLayerSpecularColor.g) * (1 - ClearCoat);
+
+						// dxr-todo: double check the math when we can compare with path tracer.
+						// Matches what is found in ReflectionEnvironmentPixelShader.usf.
+						//PathSegmentRadiance.rgb = ClearCoat * (TopLayerRadiance * TopLayerSpecularLobIntegral * TopLayerRefractionTransmittance + TopLayerRadiance * TopLayerFSChlick);
+						TopLayerReflectionEventThroughput = ClearCoat * (TopLayerSpecularLobIntegral * TopLayerRefractionTransmittance + TopLayerFSChlick);
+						float BottomToTopRefractionThroughput = 1.0; // This is already evaluated in TraceRayAndAccumulateBottomLayerResults
+						BottomLayerRadiance = BottomLayerRadiance * BottomLayerFSChlick * BottomToTopRefractionThroughput; // Bottom layer refelciton accounts for its throughput.
+					}
 				}
+				else
+				{
+					// If it is not clear coat, and we are on the first bounce, the G F terms will be applied by ReflectionEnvironmentPixelShader
+					if (BounceIndex > 0)
+						TopLayerReflectionEventThroughput *= EnvBRDF(TopLayerSpecularColor, TopLayerRoughness, NoV);
+				}
+
 
+				// Apply height fog on the sample. (for clear coat: we apply the fog the same way on each layer reflection)
 				// If there is a contribution on this path, apply the required fog on the sample. 
 				// Also accumulate the path contribution.
 				// In case no contribution is accumulated, the ReflectionEnvironmentAndSky pass will fill it up using reflection volumes.
-				//#dxr - todo: when MaxBounces>1, we will need to take into account a path throughput affected from the fog trnasmittance.
+				//#dxr - todo: when MaxBounces>1, we will need to take into account a path throughput affected from the fog transmittance.
 				float4 HeightFogInscatteringAndTransmittance = float4(0.0f, 0.0f, 0.0f, 1.0f);	// no fog 
 				if (bApplyHeightFog)
 				{
 					HeightFogInscatteringAndTransmittance = CalculateHeightFog(OriginToCollider);
 				}
 				// Now accumulate the radiance coming through this path segment into the full path radiance.
-				PathSegmentRadiance.rgb = PathSegmentRadiance.rgb * HeightFogInscatteringAndTransmittance.a + HeightFogInscatteringAndTransmittance.rgb;
+				// In the case of clear coat material: the fog is applied once for both bottom and top layer using the main top layer refelction path.
+				float3 PathSegmentRadiance = TopLayerRadiance * TopLayerReflectionEventThroughput + BottomLayerRadiance;
+				PathSegmentRadiance = PathSegmentRadiance * HeightFogInscatteringAndTransmittance.a + HeightFogInscatteringAndTransmittance.rgb;
 				PathRadiance += PathSegmentRadiance * PathThroughput;
-				PathThroughput *= HeightFogInscatteringAndTransmittance.a;
-				
-				IncidentDirection = Ray.Direction;
-				WorldPosition = Ray.Origin + Payload.HitT * Ray.Direction;
-				WorldNormal = Payload.WorldNormal;
-				Roughness = Payload.Roughness;
-				SpecularColor = Payload.SpecularColor;
+				// Update the path throughput according to fog and material top layer reflection event.
+				PathThroughput *= HeightFogInscatteringAndTransmittance.a * TopLayerReflectionEventThroughput;
 				
+
+
+				// Setup next iteration
+				IncidentDirection = TopLayerRay.Direction;
+				WorldPosition = TopLayerRay.Origin + Payload.HitT * TopLayerRay.Direction;
+
 				ShadingModelID = Payload.ShadingModelID;
+				TopLayerWorldNormal = Payload.WorldNormal;
+				TopLayerRoughness = Payload.Roughness;
+				TopLayerSpecularColor = Payload.SpecularColor;
+
 				if (ShadingModelID == SHADINGMODELID_CLEAR_COAT)
 				{
-					ClearCoat = Payload.CustomData.x;
-					ClearCoatRoughness = Payload.CustomData.y;
+					BottomLayerRoughness = TopLayerRoughness;
+					BottomLayerSpecularColor = TopLayerSpecularColor;
+					BottomLayerWorldNormal = TopLayerWorldNormal;
 					#if CLEAR_COAT_BOTTOM_NORMAL
-						const float2 oct1 = ((float2(Payload.CustomData.a, Payload.CustomData.z) * 2) - (256.0/255.0)) + UnitVectorToOctahedron(Payload.WorldNormal);
-						ClearCoatBottomNormal = OctahedronToUnitVector(oct1);
+					const float2 oct1 = ((float2(Payload.CustomData.a, Payload.CustomData.z) * 2) - (256.0 / 255.0)) + UnitVectorToOctahedron(TopLayerWorldNormal);
+					BottomLayerWorldNormal = OctahedronToUnitVector(oct1);
 					#endif
+
+					ClearCoat = Payload.CustomData.x;				// Clear coat weight
+					TopLayerRoughness = Payload.CustomData.y;		// Clear coat roughness
+					TopLayerSpecularColor = ClearCoatSpecularColor;	// Hard coded top layer specular color
 				}
-				
-				if (all(PathThroughput < 0.1)) break;
+
+				if (all(PathThroughput < 0.001) || bSkyWasHit || HitDistance<0.0f) break;
 			}
 			
 			if (bIsValidSample)
diff --git a/Engine/Shaders/Private/ReflectionEnvironmentPixelShader.usf b/Engine/Shaders/Private/ReflectionEnvironmentPixelShader.usf
--- a/Engine/Shaders/Private/ReflectionEnvironmentPixelShader.usf
+++ b/Engine/Shaders/Private/ReflectionEnvironmentPixelShader.usf
@@ -235,9 +235,14 @@ float3 ReflectionEnvironment(FScreenSpaceData ScreenSpaceData, float4 UVAndScree
 	Color.a = 1 - SSR.a;
 #endif
 
-#if ENABLE_CLEAR_COAT
+#if RAY_TRACED_REFLECTIONS
+	float4 SavedColor = Color;	// When a clear coat material is encountered, we save the reflection buffer color for it to not be affected by operations.
+#endif
 	if(GBuffer.ShadingModelID == SHADINGMODELID_CLEAR_COAT )
 	{
+#if RAY_TRACED_REFLECTIONS
+		Color = float4(0, 0, 0, 0); // Clear coat reflection is expected to be computed on a black background
+#endif
 		const float ClearCoat = GBuffer.CustomData.x;
 		Color = lerp( Color, float4(0,0,0,1), ClearCoat );
 
@@ -249,7 +254,6 @@ float3 ReflectionEnvironment(FScreenSpaceData ScreenSpaceData, float4 UVAndScree
 		R = 2 * dot( V, ClearCoatUnderNormal ) * ClearCoatUnderNormal - V;
 #endif
 	}
-#endif
 
 	float AO = GBuffer.GBufferAO * ScreenSpaceData.AmbientOcclusion;
 	float RoughnessSq = GBuffer.Roughness * GBuffer.Roughness;
@@ -268,13 +272,12 @@ float3 ReflectionEnvironment(FScreenSpaceData ScreenSpaceData, float4 UVAndScree
 	uint NumCulledReflectionCaptures = 0;
 #endif
 
-	//bottom for clearcoat or the only reflection.
+	//Top of regular reflection or bottom layer of clear coat.
 	Color.rgb += PreExposure * GatherRadiance(Color.a, WorldPosition, R, GBuffer.Roughness, BentNormal, IndirectIrradiance, GBuffer.ShadingModelID, NumCulledReflectionCaptures, DataStartIndex);
 
 	BRANCH
-	if( GBuffer.ShadingModelID == SHADINGMODELID_CLEAR_COAT )
+	if( GBuffer.ShadingModelID == SHADINGMODELID_CLEAR_COAT)
 	{
-#if ENABLE_CLEAR_COAT
 		const float ClearCoat			= GBuffer.CustomData.x;
 		const float ClearCoatRoughness	= GBuffer.CustomData.y;
 
@@ -299,9 +302,10 @@ float3 ReflectionEnvironment(FScreenSpaceData ScreenSpaceData, float4 UVAndScree
 
 		float3 TopLayerR = 2 * dot( V, N ) * N - V;
 		Color.rgb += PreExposure * GatherRadiance(Color.a, WorldPosition, TopLayerR, ClearCoatRoughness, BentNormal, IndirectIrradiance, GBuffer.ShadingModelID, NumCulledReflectionCaptures, DataStartIndex);
-#else
-		// When clear coat is disabled in ReflectionEnvironmentPS (both layers evaluated in ray traced reflections), do nothing here
-#endif
+
+	#if RAY_TRACED_REFLECTIONS
+		Color.rgb = SavedColor.rgb +  Color.rgb * SavedColor.a; // Compose default clear coat reflection over regular refelction (using Premultiplied alpha where SaveColor.a=transmittance)
+	#endif
 	}
 	else
 	{
diff --git a/Engine/Source/Runtime/Renderer/Private/ReflectionEnvironment.cpp b/Engine/Source/Runtime/Renderer/Private/ReflectionEnvironment.cpp
--- a/Engine/Source/Runtime/Renderer/Private/ReflectionEnvironment.cpp
+++ b/Engine/Source/Runtime/Renderer/Private/ReflectionEnvironment.cpp
@@ -519,7 +519,7 @@ class FReflectionEnvironmentSkyLightingPS : public FGlobalShader
 	class FSkyLight					: SHADER_PERMUTATION_BOOL("ENABLE_SKY_LIGHT");
 	class FDynamicSkyLight			: SHADER_PERMUTATION_BOOL("ENABLE_DYNAMIC_SKY_LIGHT");
 	class FSkyShadowing				: SHADER_PERMUTATION_BOOL("APPLY_SKY_SHADOWING");
-	class FClearCoat				: SHADER_PERMUTATION_BOOL("ENABLE_CLEAR_COAT");
+	class FRayTracedReflections		: SHADER_PERMUTATION_BOOL("RAY_TRACED_REFLECTIONS");
 
 	using FPermutationDomain = TShaderPermutationDomain<
 		FHasBoxCaptures,
@@ -529,7 +529,7 @@ class FReflectionEnvironmentSkyLightingPS : public FGlobalShader
 		FSkyLight,
 		FDynamicSkyLight,
 		FSkyShadowing,
-		FClearCoat>;
+		FRayTracedReflections>;
 
 	static FPermutationDomain RemapPermutation(FPermutationDomain PermutationVector)
 	{
@@ -557,7 +557,7 @@ class FReflectionEnvironmentSkyLightingPS : public FGlobalShader
 		return PermutationVector;
 	}
 
-	static FPermutationDomain BuildPermutationVector(const FViewInfo& View, bool bBoxCapturesOnly, bool bSphereCapturesOnly, bool bSupportDFAOIndirectOcclusion, bool bSpecularBounce, bool bEnableSkyLight, bool bEnableDynamicSkyLight, bool bApplySkyShadowing, bool bEnableClearCoat)
+	static FPermutationDomain BuildPermutationVector(const FViewInfo& View, bool bBoxCapturesOnly, bool bSphereCapturesOnly, bool bSupportDFAOIndirectOcclusion, bool bSpecularBounce, bool bEnableSkyLight, bool bEnableDynamicSkyLight, bool bApplySkyShadowing, bool bRayTracedReflections)
 	{
 		FPermutationDomain PermutationVector;
 
@@ -568,7 +568,7 @@ class FReflectionEnvironmentSkyLightingPS : public FGlobalShader
 		PermutationVector.Set<FSkyLight>(bEnableSkyLight);
 		PermutationVector.Set<FDynamicSkyLight>(bEnableDynamicSkyLight);
 		PermutationVector.Set<FSkyShadowing>(bApplySkyShadowing);
-		PermutationVector.Set<FClearCoat>(bEnableClearCoat);
+		PermutationVector.Set<FRayTracedReflections>(bRayTracedReflections);
 
 		return RemapPermutation(PermutationVector);
 	}
@@ -855,7 +855,7 @@ void FDeferredShadingSceneRenderer::RenderDeferredReflectionsAndSkyLighting(FRHI
 
 			TShaderMapRef<FPostProcessVS> VertexShader(View.ShaderMap);
 
-			auto PermutationVector = FReflectionEnvironmentSkyLightingPS::BuildPermutationVector(View, bHasBoxCaptures, bHasSphereCaptures, DynamicBentNormalAO != NULL, bReflectionCapture, bSkyLight, bDynamicSkyLight, bApplySkyShadowing, !bRayTracedReflections);
+			auto PermutationVector = FReflectionEnvironmentSkyLightingPS::BuildPermutationVector(View, bHasBoxCaptures, bHasSphereCaptures, DynamicBentNormalAO != NULL, bReflectionCapture, bSkyLight, bDynamicSkyLight, bApplySkyShadowing, bRayTracedReflections);
 
 			TShaderMapRef<FReflectionEnvironmentSkyLightingPS> PixelShader(View.ShaderMap, PermutationVector);
 
