Fixed ray traced shadows from masked materials

diff --git a/Engine/Shaders/Private/RayTracing/RayTracingOcclusionRGS.usf b/Engine/Shaders/Private/RayTracing/RayTracingOcclusionRGS.usf
--- a/Engine/Shaders/Private/RayTracing/RayTracingOcclusionRGS.usf
+++ b/Engine/Shaders/Private/RayTracing/RayTracingOcclusionRGS.usf
@@ -223,11 +223,11 @@ void OcclusionRGS()
 			}
 		}
 
-		uint RayFlags = 0
+		uint RayFlags = 0;
+
 #if !ENABLE_TWO_SIDED_GEOMETRY
-			| RAY_FLAG_CULL_BACK_FACING_TRIANGLES
+		RayFlags |= RAY_FLAG_CULL_BACK_FACING_TRIANGLES;
 #endif
-			| RAY_FLAG_FORCE_OPAQUE;
 
 		FPackedMaterialClosestHitPayload Payload = (FPackedMaterialClosestHitPayload)0;
 		
