Fix an issue where instanced draws could exhibit visual glitches on certain platforms. This was caused by reading PrimitiveID into a scalar register in the pixel shader. Pixel waves are packed for instanced drawing, so this was actually invalid.

#rnx

diff --git a/Engine/Shaders/Private/LocalVertexFactory.ush b/Engine/Shaders/Private/LocalVertexFactory.ush
--- a/Engine/Shaders/Private/LocalVertexFactory.ush
+++ b/Engine/Shaders/Private/LocalVertexFactory.ush
@@ -397,10 +397,6 @@ FMaterialPixelParameters GetMaterialPixelParameters(FVertexFactoryInterpolantsVS
 	Result.TwoSidedSign = 1;
 	Result.PrimitiveId = GetPrimitiveId(Interpolants);
 
-#if COMPILER_SUPPORTS_TO_SCALAR_MEMORY
-	Result.PrimitiveId = ToScalarMemory(Result.PrimitiveId);
-#endif
-
 	return Result;
 }
 
