Cherry-pick 5774624: Fix path tracer not using RAY_TRACING_SHADER_SLOT_MATERIAL and RAY_TRACING_NUM_SHADER_SLOTS

diff --git a/Engine/Shaders/Private/PathTracing/PathTracing.usf b/Engine/Shaders/Private/PathTracing/PathTracing.usf
--- a/Engine/Shaders/Private/PathTracing/PathTracing.usf
+++ b/Engine/Shaders/Private/PathTracing/PathTracing.usf
@@ -87,8 +87,8 @@ void PathTracingMainRG()
 			TLAS,   // AccelerationStructure
 			RayFlags,
 			RAY_TRACING_MASK_ALL,
-			0,      // RayContributionToHitGroupIndex
-			1,      // MultiplierForGeometryContributionToShaderIndex #dxr_todo: we currently only support shader binding per mesh instance
+			RAY_TRACING_SHADER_SLOT_MATERIAL, // RayContributionToHitGroupIndex
+			RAY_TRACING_NUM_SHADER_SLOTS,     // MultiplierForGeometryContributionToShaderIndex
 			0,      // MissShaderIndex
 			Ray,    // RayDesc
 			RayCone
