Cherry-pick 5738970: Fix wrong composition between RT reflection and cubemap/sky when pre-exposition is enabled.

diff --git a/Engine/Shaders/Private/RayTracing/RayTracingReflections.usf b/Engine/Shaders/Private/RayTracing/RayTracingReflections.usf
--- a/Engine/Shaders/Private/RayTracing/RayTracingReflections.usf
+++ b/Engine/Shaders/Private/RayTracing/RayTracingReflections.usf
@@ -512,7 +512,7 @@ void RayTracingReflectionsRGS()
 		}
 
 	#if USE_PREEXPOSURE
-		ReflectedColor *= View.PreExposure;
+		ReflectedColor.xyz *= View.PreExposure;
 	#endif
 
 		ColorOutput[DispatchThreadId] = ClampToHalfFloatRange(ReflectedColor);
