Cherry-pick 5832319: Fixed regression that made pure specular brdfs render black in path tracing mode

diff --git a/Engine/Shaders/Private/PathTracing/Material/PathTracingSpecularReflection.ush b/Engine/Shaders/Private/PathTracing/Material/PathTracingSpecularReflection.ush
--- a/Engine/Shaders/Private/PathTracing/Material/PathTracingSpecularReflection.ush
+++ b/Engine/Shaders/Private/PathTracing/Material/PathTracingSpecularReflection.ush
@@ -16,7 +16,7 @@ void SpecularReflection_SampleMaterial(
 	float3 SpecularColor = GetSpecularColor(Payload);
 
 	float3 V_World = normalize(-RayDirection);
-	OutDirection = reflect(V_World, N_World);
+	OutDirection = reflect(-V_World, N_World);
 
 	float NoV = saturate(dot(V_World, N_World));
 	float VoH = NoV;
