Fixed PCSS.

diff --git a/Engine/Shaders/Private/ShadowPercentageCloserFiltering.ush b/Engine/Shaders/Private/ShadowPercentageCloserFiltering.ush
--- a/Engine/Shaders/Private/ShadowPercentageCloserFiltering.ush
+++ b/Engine/Shaders/Private/ShadowPercentageCloserFiltering.ush
@@ -192,7 +192,7 @@ float DirectionalPCSS(FPCSSSamplerSettings Settings, float2 ShadowPosition, floa
 	float UVCrossSum = 0;
 	for (int i = 0; i < PCSS_SEARCH_SAMPLES; i++)
 	{
-		float2 E = SobolIndex(SobolRandom ^ uint2(0x4B05, 0xB0CD), i << 3, PCSS_SEARCH_BITS + 3) / 0x10000;
+		float2 E = SobolIndex(SobolRandom ^ uint2(0x4B05, 0xB0CD), i << 3, PCSS_SEARCH_BITS + 3) / float(0x10000);
 		float2 PCSSSample = UniformSampleDiskConcentricApprox(E);
 		float2 SampleUVOffset = PCSSSample * SearchRadius;
 		float2 SampleUV = ShadowPosition + SampleUVOffset * Settings.ShadowTileOffsetAndSize.zw;
@@ -391,7 +391,7 @@ float DirectionalPCSS(FPCSSSamplerSettings Settings, float2 ShadowPosition, floa
 	float ShadowOcclusionSum = 0;
 	for (int j = 0; j < PCSS_SAMPLES; j++)
 	{
-		float2 E = SobolIndex(SobolRandom, j << 3, PCSS_SAMPLE_BITS + 3) / 0x10000;
+		float2 E = SobolIndex(SobolRandom, j << 3, PCSS_SAMPLE_BITS + 3) / float(0x10000);
 		float2 PCFSample = UniformSampleDiskConcentricApprox(E);
 
 		float OriginalUVLength = length(PCFSample);
@@ -417,7 +417,7 @@ float DirectionalPCSS(FPCSSSamplerSettings Settings, float2 ShadowPosition, floa
 	float VisibleLightAccumulation = 0;
 	for (int j = 0; j < PCSS_SAMPLES; j++)
 	{
-		float2 E = SobolIndex(SobolRandom, j << 3, PCSS_SAMPLE_BITS + 3) / 0x10000;
+		float2 E = SobolIndex(SobolRandom, j << 3, PCSS_SAMPLE_BITS + 3) / float(0x10000);
 		float2 PCFSample = UniformSampleDiskConcentricApprox(E);
 		float2 SampleUVOffset = mul(PCFUVMatrix, PCFSample) * FilterRadius;
 		float2 SampleUV = ShadowPosition + SampleUVOffset * Settings.ShadowTileOffsetAndSize.zw;
