Fixed landscape shader warning related to LayerTexCoord

diff --git a/Engine/Shaders/Private/LandscapeVertexFactory.ush b/Engine/Shaders/Private/LandscapeVertexFactory.ush
--- a/Engine/Shaders/Private/LandscapeVertexFactory.ush
+++ b/Engine/Shaders/Private/LandscapeVertexFactory.ush
@@ -115,7 +115,11 @@ struct FVertexFactoryInterpolantsVSToPS
 
 struct FLandscapeTexCoords
 {
-	float4	LayerTexCoord; // xy == texcoord, zw == tess falloff params
+#if USING_TESSELLATION
+	float4	LayerTexCoord : TEXCOORD0; // xy == texcoord, zw == tess falloff params
+#else
+	float2	LayerTexCoord : TEXCOORD0; // xy == texcoord
+#endif
 	float2	WeightMapTexCoord;
 #if !(ES2_PROFILE || ES3_1_PROFILE)
 	float2	HeightMapTexCoord;
