- Niagara fix issue with depth only rendering position not matching base pass rendering due to shader compiler optimizations
- Change edigrated from Dev-Niagara

diff --git a/Engine/Shaders/Private/DepthOnlyVertexShader.usf b/Engine/Shaders/Private/DepthOnlyVertexShader.usf
--- a/Engine/Shaders/Private/DepthOnlyVertexShader.usf
+++ b/Engine/Shaders/Private/DepthOnlyVertexShader.usf
@@ -166,9 +166,6 @@ void Main(
 		// XBOXONE_BIAS_HACK is defined only in a custom node in a particular material
 		// This should be removed with a future shader compiler update
 		Output.Position.z -= 0.0001 * Output.Position.w;
-	#elif NIAGARA_BIAS_HACK
-		// Similar to above, only for Niagara. (UE-68456)
-		Output.Position.z -= 0.0001 * Output.Position.w;
 	#endif
 
 #if USE_GLOBAL_CLIP_PLANE
diff --git a/Engine/Shaders/Private/NiagaraSpriteVertexFactory.ush b/Engine/Shaders/Private/NiagaraSpriteVertexFactory.ush
--- a/Engine/Shaders/Private/NiagaraSpriteVertexFactory.ush
+++ b/Engine/Shaders/Private/NiagaraSpriteVertexFactory.ush
@@ -446,7 +446,7 @@ void GetTangents(FVertexFactoryInput Input, FVertexFactoryIntermediates Intermed
 	float3 Up;
 	
 	/**This could well be better as a branch*/
-	FLATTEN
+	BRANCH
 	if(NiagaraSpriteVFLooseParameters.ParticleAlignmentMode == ALIGNMENT_MODE_UNALIGNED)
 	{
 		FLATTEN
diff --git a/Engine/Shaders/Private/NiagaraVFParticleAccess.usf b/Engine/Shaders/Private/NiagaraVFParticleAccess.usf
--- a/Engine/Shaders/Private/NiagaraVFParticleAccess.usf
+++ b/Engine/Shaders/Private/NiagaraVFParticleAccess.usf
@@ -8,9 +8,6 @@
 //uint NiagaraIntDataOffset;
 //uint NiagaraIntDataStride;
 
-// This is a hack to fix cases where depth values are invariant due to compiler optimizations/operation reordering between the depth pass and base pass. (UE-68456)
-#define NIAGARA_BIAS_HACK 1
-
 #ifndef NIAGARA_KEYGEN
 #define NIAGARA_KEYGEN 0
 #endif
