Disable EyeAdaptation material node output for materialtexcoordscale debug shader.

diff --git a/Engine/Shaders/Private/MaterialTemplate.ush b/Engine/Shaders/Private/MaterialTemplate.ush
--- a/Engine/Shaders/Private/MaterialTemplate.ush
+++ b/Engine/Shaders/Private/MaterialTemplate.ush
@@ -1350,7 +1350,11 @@ MaterialFloat2 GetLightmapUVs(FMaterialPixelParameters Parameters)
 //Provides access to the EyeAdaptation RT. Only available in SM5.
 float EyeAdaptationLookup()
 {
+#if EYE_ADAPTATION_DISABLED
+	return 0.0;
+#else
 	return EyeAdaptationStruct.EyeAdaptation.Load(int3(0, 0, 0)).x;
+#endif
 }
 #endif
 
diff --git a/Engine/Shaders/Private/MaterialTexCoordScalesPixelShader.usf b/Engine/Shaders/Private/MaterialTexCoordScalesPixelShader.usf
--- a/Engine/Shaders/Private/MaterialTexCoordScalesPixelShader.usf
+++ b/Engine/Shaders/Private/MaterialTexCoordScalesPixelShader.usf
@@ -12,6 +12,9 @@
 #undef SCENE_TEXTURES_DISABLED
 #define SCENE_TEXTURES_DISABLED 1
 
+#undef EYE_ADAPTATION_DISABLED
+#define EYE_ADAPTATION_DISABLED 1
+
 #include "Common.ush"
 #include "SHCommon.ush"
 
