Fixed d3d error DEVICE_UNORDEREDACCESSVIEW_RETURN_TYPE_MISMATCH when running in -d3ddebug (caused by compute shader LandscapeLayerCS)

diff --git a/Engine/Shaders/Private/LandscapeLayersCS.usf b/Engine/Shaders/Private/LandscapeLayersCS.usf
--- a/Engine/Shaders/Private/LandscapeLayersCS.usf
+++ b/Engine/Shaders/Private/LandscapeLayersCS.usf
@@ -14,7 +14,7 @@ struct FLandscapeLayerWeightmapExtractMaterialLayersComponentData
 
 StructuredBuffer<FLandscapeLayerWeightmapExtractMaterialLayersComponentData> InExtractLayersComponentsData;
 Texture2D<float4> InComponentWeightMaps;
-RWTexture2DArray<float> OutAtlasPaintLayers;
+RWTexture2DArray<UNORM float> OutAtlasPaintLayers;
 
 [numthreads(THREADGROUP_SIZEX, THREADGROUP_SIZEY, 1)]
 void ComputeWeightmapPerPaintLayer(uint3 DispatchThreadId : SV_DispatchThreadID)
@@ -41,7 +41,7 @@ StructuredBuffer<FLandscapeLayerWeightmapPackMaterialLayersComponentData> InPack
 Buffer<float> InWeightmapWeightBlendMode;
 Buffer<float2> InWeightmapTextureOutputOffset;
 Texture2DArray<float> InAtlasPaintLayers;
-RWTexture2D<float4> OutComponentWeightMaps;
+RWTexture2D<UNORM float4> OutComponentWeightMaps;
 
 float ComputeWeightmapChannel(int ChannelIndexToProcess, uint2 TexelPosition, uint SourcePaintLayerIndex)
 {
