Fix gamma correction in GLES 3.1 external texture sampling

diff --git a/Engine/Shaders/Private/MaterialTemplate.ush b/Engine/Shaders/Private/MaterialTemplate.ush
--- a/Engine/Shaders/Private/MaterialTemplate.ush
+++ b/Engine/Shaders/Private/MaterialTemplate.ush
@@ -1022,6 +1022,20 @@ MaterialFloat4 ProcessMaterialColorTextureLookup(MaterialFloat4 TextureValue)
 	return TextureValue;
 }
 
+MaterialFloat4 ProcessMaterialExternalTextureLookup(MaterialFloat4 TextureValue)
+{
+#if (FEATURE_LEVEL == FEATURE_LEVEL_ES2) || COMPILER_GLSL_ES3_1
+	// As per https://www.khronos.org/registry/OpenGL/extensions/OES/OES_EGL_image_external.txt
+	// sampling an external texture does not change the color space, so we always need to
+	// perform manual approximate sRGB conversion
+	float3 LinRGB = TextureValue.rgb;
+	LinRGB *= LinRGB; 
+	return MaterialFloat4(LinRGB, TextureValue.a);
+#else
+	return ProcessMaterialColorTextureLookup(TextureValue);
+#endif
+}
+
 MaterialFloat4 ProcessMaterialLinearColorTextureLookup(MaterialFloat4 TextureValue)
 {
 	return TextureValue;
diff --git a/Engine/Source/Runtime/Engine/Private/Materials/HLSLMaterialTranslator.h b/Engine/Source/Runtime/Engine/Private/Materials/HLSLMaterialTranslator.h
--- a/Engine/Source/Runtime/Engine/Private/Materials/HLSLMaterialTranslator.h
+++ b/Engine/Source/Runtime/Engine/Private/Materials/HLSLMaterialTranslator.h
@@ -3540,7 +3540,9 @@ protected:
 		switch( SamplerType )
 		{
 			case SAMPLERTYPE_External:
-				// fall through since should be treated same as SAMPLERTYPE_Color
+				SampleCode = FString::Printf(TEXT("ProcessMaterialExternalTextureLookup(%s)"), *SampleCode);
+				break;
+
 			case SAMPLERTYPE_Color:
 				SampleCode = FString::Printf( TEXT("ProcessMaterialColorTextureLookup(%s)"), *SampleCode );
 				break;
