Optimize Tonemap function for non mobile

Use Quad intrinsics to cheapen the cost of the sharpen function. Gives around 15% speed up.

diff --git a/Engine/Shaders/Private/PixelQuadMessagePassing.ush b/Engine/Shaders/Private/PixelQuadMessagePassing.ush
--- a/Engine/Shaders/Private/PixelQuadMessagePassing.ush
+++ b/Engine/Shaders/Private/PixelQuadMessagePassing.ush
@@ -16,6 +16,9 @@ struct FPQMPContext
 {
 	// Position of the pixel within the 2x2 pixel quad.
 	float2 PixelPos;
+
+	// Signed normalized  2D vector for each pixel in a 2x2 quad from the centre. Will be either +1 or -1 for each axis
+	float2  QuadVector;
 };
 
 
@@ -24,9 +27,12 @@ FPQMPContext PQMPInit(float2 SvPosition)
 {
 	FPQMPContext PQMP;
 #if HAS_PIXEL_QUAD_MESSAGE_PASSING_SUPPORT
-	PQMP.PixelPos = floor(2.0 * frac(SvPosition * 0.5));
+	PQMP.PixelPos   = floor(2.0 * frac(SvPosition * 0.5));
+	uint2 ScreenPos = uint2(SvPosition);
+	PQMP.QuadVector = float2( float(ScreenPos.x & 1) * 2.0 - 1.0, float(ScreenPos.y & 1) * 2.0 - 1.0);
 #else
-	PQMP.PixelPos = float2(0, 0);
+	PQMP.PixelPos	= float2(0, 0);
+	PQMP.QuadVector = float2(0, 0);
 #endif
 	return PQMP;
 }
@@ -36,7 +42,7 @@ FPQMPContext PQMPInit(float2 SvPosition)
 float PQMPAverage(FPQMPContext PQMP, float v)
 {
 #if HAS_PIXEL_QUAD_MESSAGE_PASSING_SUPPORT
-	v = v + (0.5 - PQMP.PixelPos.x) * ddx_fine(v);
+	v    = v + (0.5 - PQMP.PixelPos.x) * ddx_fine(v);
 	return v + (0.5 - PQMP.PixelPos.y) * ddy_fine(v);
 #else
 	return v;
@@ -46,9 +52,62 @@ float PQMPAverage(FPQMPContext PQMP, float v)
 float2 PQMPAverage(FPQMPContext PQMP, float2 v)
 {
 #if HAS_PIXEL_QUAD_MESSAGE_PASSING_SUPPORT
-	v = v + (0.5 - PQMP.PixelPos.x) * ddx_fine(v);
+	v    = v + (0.5 - PQMP.PixelPos.x) * ddx_fine(v);
 	return v + (0.5 - PQMP.PixelPos.y) * ddy_fine(v);
 #else
 	return v;
 #endif
 }
+
+
+// Get the value of the horizontal neighbour in the same pixel quad
+float PQMPReadX(FPQMPContext PQMP, float v)
+{
+#if HAS_PIXEL_QUAD_MESSAGE_PASSING_SUPPORT
+	#if COMPILER_SUPPORTS_QUAD_PASS
+		return QuadReadAcrossX(v);
+	#else
+		float dX = ddx_fine(v);
+		return v - (dX * PQMP.QuadVector.x);
+	#endif
+#else
+	return v;
+#endif
+}
+
+// Get the value of the vertical neighbour in the same pixel quad
+float PQMPReadY(FPQMPContext PQMP, float v)
+{
+#if HAS_PIXEL_QUAD_MESSAGE_PASSING_SUPPORT
+	#if COMPILER_SUPPORTS_QUAD_PASS
+		return QuadReadAcrossY(v);
+	#else
+		float dY = ddy_fine(v);
+		return v - (dY * PQMP.QuadVector.y);
+	#endif
+#else
+	return v;
+#endif
+}
+
+float2 PQMPReadX(FPQMPContext PQMP, float2 v)
+{
+	return float2( PQMPReadX(PQMP, v.x), PQMPReadX(PQMP, v.y));
+}
+
+float2 PQMPReadY(FPQMPContext PQMP, float2 v)
+{
+	return float2( PQMPReadY(PQMP, v.x), PQMPReadY(PQMP, v.y));
+}
+
+float3 PQMPReadX(FPQMPContext PQMP, float3 v)
+{
+	return float3( PQMPReadX(PQMP, v.x), PQMPReadX(PQMP, v.y), PQMPReadX(PQMP, v.z));
+}
+
+float3 PQMPReadY(FPQMPContext PQMP, float3 v)
+{
+	return float3( PQMPReadY(PQMP, v.x), PQMPReadY(PQMP, v.y), PQMPReadY(PQMP, v.z));
+}
+
+
diff --git a/Engine/Shaders/Private/PostProcessTonemap.usf b/Engine/Shaders/Private/PostProcessTonemap.usf
--- a/Engine/Shaders/Private/PostProcessTonemap.usf
+++ b/Engine/Shaders/Private/PostProcessTonemap.usf
@@ -8,7 +8,7 @@
 #include "PostProcessCommon.ush"
 #include "TonemapCommon.ush"
 #include "EyeAdaptationCommon.ush"
-
+#include "PixelQuadMessagePassing.ush"
 
 
 
@@ -400,28 +400,36 @@ float4 TonemapCommonPS(
 
 #if USE_SHARPEN
 	{
-		half3 C1 = LookupSceneColor(SceneUV, int2(-1,  0)) * OneOverPreExposure;
-		half3 C2 = LookupSceneColor(SceneUV, int2( 1,  0)) * OneOverPreExposure;
-		half3 C3 = LookupSceneColor(SceneUV, int2(-1, -1)) * OneOverPreExposure;
-		half3 C4 = LookupSceneColor(SceneUV, int2(-1,  1)) * OneOverPreExposure;
 
-		// reference solution
-//		half A0 = max3(SceneColor);
-//		half4 LuminanceNeightbors = half4(max3(C1), max3(C2), max3(C3), max3(C4));
-		// 5 instructiopns less than reference
 		half A0 = Luminance(SceneColor.rgb);
-		half4 LuminanceNeightbors = half4(Luminance(C1), Luminance(C2), Luminance(C3), Luminance(C4));
-		// 9 instructions less than reference (green is often a good approximation for the luminance)
-		//half A0 = SceneColor.g;
-		//half4 LuminanceNeightbors = half4(C1.g, C2.g, C3.g, C4.g);
+
+	#if HAS_PIXEL_QUAD_MESSAGE_PASSING_SUPPORT 
+		// Use Wave Intrinsics to reduce texture taps
+		FPQMPContext PQMP = PQMPInit(ScreenPos);
+		half4 LuminanceNeighbors;
+
+		half3 C1 = Texture2DSample(PostprocessInput0, PostprocessInput0Sampler, SceneUV + float2(PQMP.QuadVector.x, 0) * PostprocessInput0Size.zw).rgb * OneOverPreExposure;
+		half3 C3 = Texture2DSample(PostprocessInput0, PostprocessInput0Sampler, SceneUV + float2(0, PQMP.QuadVector.y) * PostprocessInput0Size.zw).rgb * OneOverPreExposure;
+		LuminanceNeighbors.x = Luminance(C1);
+		LuminanceNeighbors.y = Luminance(C3);
+		
+		half3 C2 = PQMPReadX(PQMP, SceneColor.rgb);
+		half3 C4 = PQMPReadY(PQMP, SceneColor.rgb);
+		LuminanceNeighbors.z =  PQMPReadX(PQMP,A0);
+		LuminanceNeighbors.w =  PQMPReadY(PQMP,A0);
+	#else
+		half3 C1 = LookupSceneColor(SceneUV, int2(-1,  0)) * OneOverPreExposure;
+		half3 C2 = LookupSceneColor(SceneUV, int2( 1,  0)) * OneOverPreExposure;
+		half3 C3 = LookupSceneColor(SceneUV, int2( 0, -1)) * OneOverPreExposure;
+		half3 C4 = LookupSceneColor(SceneUV, int2( 0,  1)) * OneOverPreExposure;
+		half4 LuminanceNeighbors = half4(Luminance(C1), Luminance(C2), Luminance(C3), Luminance(C4));
+	#endif
 
 		// compute mask to not sharpen near very bright HDR content
 		// Note: using max instead of summming up saves 1 instruction
 		// Note: We could multiply this to tweak the edge weight but it might instroduce a cost
-		half HDREdge = ExposureScale * max4(abs(A0 - LuminanceNeightbors));
-		// 3 instructions less but less accurate (faint black dots near bright edges)
-		//half HDREdge = ExposureScale * max(abs(A0 * 2 - (LuminanceNeightbors.r + LuminanceNeightbors.g)), abs(A0 * 2 - (LuminanceNeightbors.b + LuminanceNeightbors.a)));
-
+		half HDREdge = ExposureScale * max4(abs(A0 - LuminanceNeighbors));
+	
 		// 0..1
 		half EdgeMask = saturate(1.0f - HDREdge);
 
@@ -429,17 +437,14 @@ float4 TonemapCommonPS(
 		if (all(DebugTile == int2(1, 0)))
 		{
 			// right to the center is: Sharpen EdgeMask
-			OutColor = EdgeMask; return;
+			OutColor = EdgeMask; 
+			return OutColor;
 		}
 #endif //DEBUG_TONEMAPPER
 
 		// -1:sharpen, 0:no effect, 1:blur
 		float LerpFactor = -EdgeMask * SharpenMultiplierDiv6;
 
-		// reference
-		// same instruction count as if the center weight is 1.0
-//		half3 BlurredSceneColor = ((C1 + C2 + C3 + C4) + SceneColor * 2) / 6;
-//		SceneColor = lerp(SceneColor, BlurredSceneColor, LerpFactor);
 		// optimized, Div6 went into the C++ code
 		half3 DeltaColor = (C1 + C2 + C3 + C4) - SceneColor.rgb * 4;
 		SceneColor.rgb += DeltaColor * LerpFactor;
diff --git a/Engine/Shaders/Public/Platform.ush b/Engine/Shaders/Public/Platform.ush
--- a/Engine/Shaders/Public/Platform.ush
+++ b/Engine/Shaders/Public/Platform.ush
@@ -76,6 +76,10 @@
 #define COMPILER_SUPPORTS_ATTRIBUTES 0
 #endif
 
+#ifndef COMPILER_SUPPORTS_QUAD_PASS
+#define COMPILER_SUPPORTS_QUAD_PASS 0
+#endif
+
 #ifndef PLATFORM_SUPPORTS_SRV_UB
 #define PLATFORM_SUPPORTS_SRV_UB 0
 #endif
