Merge CL 6298294 from Dev-VR - Work around shader optimization bug on Vulkan

diff --git a/Engine/Shaders/Private/PostProcessCommon.ush b/Engine/Shaders/Private/PostProcessCommon.ush
--- a/Engine/Shaders/Private/PostProcessCommon.ush
+++ b/Engine/Shaders/Private/PostProcessCommon.ush
@@ -222,7 +222,8 @@ half4 UnwrappedTexture3DSample( Texture2D Texture, SamplerState Sampler, float3
 	float V = UVW.y;
 
 	half4 RG0 = Texture2DSample( Texture, Sampler, float2(U, V) );
-	half4 RG1 = Texture2DSample( Texture, Sampler, float2(U + 1.0f / Size, V) );
+    U += 1.0f / Size;
+	half4 RG1 = Texture2DSample( Texture, Sampler, float2(U, V) );
 
 	return lerp(RG0, RG1, FracW);
 }
