Adds support for multiple shading models per material. Each material now has a bit field of which shading models it is using. Shading models for a material are selected through a new Shading Model material expression, which can be connected to the new Shading Model material output pin. This is opt-in by selecting "From Material Expression" in the Shading Model drop down on the material. With no changes, everything should behave like before, and the old workflow is still there. Optimized shader code should be as performant as before if not using more than one shading model.

diff --git a/Engine/Plugins/2D/Paper2D/Source/Paper2DEditor/Private/SpriteEditor/SpriteDetailsCustomization.cpp b/Engine/Plugins/2D/Paper2D/Source/Paper2DEditor/Private/SpriteEditor/SpriteDetailsCustomization.cpp
--- a/Engine/Plugins/2D/Paper2D/Source/Paper2DEditor/Private/SpriteEditor/SpriteDetailsCustomization.cpp
+++ b/Engine/Plugins/2D/Paper2D/Source/Paper2DEditor/Private/SpriteEditor/SpriteDetailsCustomization.cpp
@@ -504,20 +504,19 @@ FText FSpriteDetailsCustomization::GetRenderingHeaderContentText(TWeakObjectPtr<
 		FText MaterialType;
 		if (UMaterialInterface* Material = SpriteBeingEdited->GetDefaultMaterial())
 		{
-			switch (Material->GetShadingModel())
+			FMaterialShadingModelField ShadingModels = Material->GetShadingModels();
+			if (ShadingModels.IsUnlit())
 			{
-			case MSM_Unlit:
 				MaterialType = LOCTEXT("Unlit", "Unlit");
-				break;
-			case MSM_DefaultLit:
+			}
+			else if (ShadingModels.HasOnlyShadingModel(MSM_DefaultLit))
+			{
 				MaterialType = LOCTEXT("Lit", "Lit");
-				break;
-			default:
+			}
+			else
+			{
 				MaterialType = LOCTEXT("Exotic", "Exotic");
-				break;
 			}
-			//@TODO: This isn't exported, and it would include the MSM_ prefix as well...
-			// MaterialType = FText::AsCultureInvariant(UMaterial::GetMaterialShadingModelString(Material->GetShadingModel()));
 		}
 
 		static const FText Opaque = LOCTEXT("OpaqueMaterial", "Opaque");
diff --git a/Engine/Plugins/Editor/MaterialAnalyzer/Source/Private/SMaterialAnalyzer.cpp b/Engine/Plugins/Editor/MaterialAnalyzer/Source/Private/SMaterialAnalyzer.cpp
--- a/Engine/Plugins/Editor/MaterialAnalyzer/Source/Private/SMaterialAnalyzer.cpp
+++ b/Engine/Plugins/Editor/MaterialAnalyzer/Source/Private/SMaterialAnalyzer.cpp
@@ -694,7 +694,16 @@ void FAnalyzeMaterialTreeAsyncTask::DoWork()
 		}
 		else if (BasePropertyOverrideName.Key.IsEqual(TEXT("bOverride_ShadingModel")))
 		{
-			TempValue = (float)CurrentMaterialInterface->GetShadingModel();
+			if (CurrentMaterialInterface->IsShadingModelFromMaterialExpression())
+			{
+				TempValue = (float)MSM_FromMaterialExpression;
+			}
+			else
+			{
+				ensure(CurrentMaterialInterface->GetShadingModels().CountShadingModels() == 1);
+				TempValue = (float)CurrentMaterialInterface->GetShadingModels().GetFirstShadingModel();
+			}
+
 			if (CurrentMaterialInstance)
 			{
 				bIsOverridden = CurrentMaterialInstance->BasePropertyOverrides.bOverride_ShadingModel;
diff --git a/Engine/Shaders/Private/BasePassPixelShader.usf b/Engine/Shaders/Private/BasePassPixelShader.usf
--- a/Engine/Shaders/Private/BasePassPixelShader.usf
+++ b/Engine/Shaders/Private/BasePassPixelShader.usf
@@ -185,7 +185,8 @@ float3 GetTranslucencyVolumeLighting(
 			float3 BackscatteredLighting = 0;
 
 			#if MATERIAL_SHADINGMODEL_SUBSURFACE
-
+			if (GBuffer.ShadingModelID == SHADINGMODELID_SUBSURFACE)
+			{
 				float InScatterPower = GetMaterialTranslucentBackscatteringExponent();
 				// Setup a pow lobe to approximate anisotropic in-scattering near to the light direction
 				float InScattering = pow(saturate(dot(TranslucentSelfShadow.DirectionalLightDirection.xyz, MaterialParameters.CameraVector)), InScatterPower);
@@ -202,6 +203,7 @@ float3 GetTranslucencyVolumeLighting(
 					// Mask by shadowing, exaggerated
 					* SelfShadowing * SelfShadowing
 					* VolumeLighting.a;
+			}
 			#endif
 
 			// The volume lighting already contains the contribution of the directional light, 
@@ -227,7 +229,7 @@ float3 GetTranslucencyVolumeLighting(
 #endif
 
 /** Computes sky diffuse lighting, including precomputed shadowing. */
-void GetSkyLighting(FMaterialPixelParameters MaterialParameters, float3 WorldNormal, float2 LightmapUV, float3 SkyOcclusionUV3D, out float3 OutDiffuseLighting, out float3 OutSubsurfaceLighting)
+void GetSkyLighting(FMaterialPixelParameters MaterialParameters, FGBufferData GBuffer, float3 WorldNormal, float2 LightmapUV, float3 SkyOcclusionUV3D, out float3 OutDiffuseLighting, out float3 OutSubsurfaceLighting)
 {
 	OutDiffuseLighting = 0;
 	OutSubsurfaceLighting = 0;
@@ -297,8 +299,11 @@ void GetSkyLighting(FMaterialPixelParameters MaterialParameters, float3 WorldNor
 	OutDiffuseLighting += DiffuseLookup * (SkyVisibility * GeometryTerm);
 
 	#if MATERIAL_SHADINGMODEL_TWOSIDED_FOLIAGE
+	if (GBuffer.ShadingModelID == SHADINGMODELID_TWOSIDED_FOLIAGE)
+	{
 		float3 BackfaceDiffuseLookup = GetEffectiveSkySHDiffuse(-WorldNormal) * ResolvedView.SkyLightColor.rgb;
 		OutSubsurfaceLighting += BackfaceDiffuseLookup * SkyVisibility;
+	}
 	#endif
 
 #endif
@@ -317,6 +322,7 @@ void GetPrecomputedIndirectLightingAndSkyLight(
 	FMaterialPixelParameters MaterialParameters, 
 	FVertexFactoryInterpolantsVSToPS Interpolants,
 	FBasePassInterpolantsVSToPS BasePassInterpolants,
+	FGBufferData GBuffer,
 	float3 DiffuseDir,
 	float3 VolumetricLightmapBrickTextureUVs,
 	out float3 OutDiffuseLighting,
@@ -367,8 +373,11 @@ void GetPrecomputedIndirectLightingAndSkyLight(
 			OutDiffuseLighting = max(float3(0,0,0), DotSH3(IrradianceSH, DiffuseTransferSH)) / PI;
 
 			#if MATERIAL_SHADINGMODEL_TWOSIDED_FOLIAGE
+			if (GBuffer.ShadingModelID == SHADINGMODELID_TWOSIDED_FOLIAGE)
+			{
 				FThreeBandSHVector SubsurfaceTransferSH = CalcDiffuseTransferSH3(-DiffuseDir, 1);
 				OutSubsurfaceLighting += max(float3(0,0,0), DotSH3(IrradianceSH, SubsurfaceTransferSH)) / PI;
+			}
 			#endif
 		#endif
 
@@ -408,8 +417,11 @@ void GetPrecomputedIndirectLightingAndSkyLight(
 				OutDiffuseLighting = max(half3(0,0,0), DotSH(CachedSH, DiffuseTransferSH)) / PI;
 
 				#if MATERIAL_SHADINGMODEL_TWOSIDED_FOLIAGE
+				if (GBuffer.ShadingModelID == SHADINGMODELID_TWOSIDED_FOLIAGE)
+				{
 					FTwoBandSHVector SubsurfaceTransferSH = CalcDiffuseTransferSH(-DiffuseDir, 1);
 					OutSubsurfaceLighting += max(half3(0,0,0), DotSH(CachedSH, SubsurfaceTransferSH)) / PI;
+				}
 				#endif
 
 			#endif
@@ -430,8 +442,11 @@ void GetPrecomputedIndirectLightingAndSkyLight(
 			OutDiffuseLighting = DotSH1(PointIndirectLighting, DiffuseTransferSH);
 
 			#if MATERIAL_SHADINGMODEL_TWOSIDED_FOLIAGE
+			if (GBuffer.ShadingModelID == SHADINGMODELID_TWOSIDED_FOLIAGE)
+			{
 				FOneBandSHVector SubsurfaceTransferSH = CalcDiffuseTransferSH1(1);
 				OutSubsurfaceLighting += DotSH1(PointIndirectLighting, SubsurfaceTransferSH);
+			}
 			#endif
 
 		#else
@@ -454,8 +469,11 @@ void GetPrecomputedIndirectLightingAndSkyLight(
 			OutDiffuseLighting = max(half3(0,0,0), DotSH3(PointIndirectLighting, DiffuseTransferSH));
 
 			#if MATERIAL_SHADINGMODEL_TWOSIDED_FOLIAGE
+			if (GBuffer.ShadingModelID == SHADINGMODELID_TWOSIDED_FOLIAGE)
+			{
 				FThreeBandSHVector SubsurfaceTransferSH = CalcDiffuseTransferSH3(-DiffuseDir, 1);
-				OutSubsurfaceLighting += max(half3(0,0,0), DotSH3(PointIndirectLighting, SubsurfaceTransferSH));
+				OutSubsurfaceLighting += max(half3(0, 0, 0), DotSH3(PointIndirectLighting, SubsurfaceTransferSH));
+			}
 			#endif
 		
 		#endif
@@ -470,7 +488,7 @@ void GetPrecomputedIndirectLightingAndSkyLight(
 		SkyOcclusionUV = LightmapUV0;
 
 		uint VTRequest = 0;
-		GetLightMapColorHQ(LightmapUV0, LightmapUV1, LightmapDataIndex, DiffuseDir, MaterialParameters.SvPosition.xy, OutDiffuseLighting, OutSubsurfaceLighting, VTRequest);
+		GetLightMapColorHQ(LightmapUV0, LightmapUV1, LightmapDataIndex, DiffuseDir, MaterialParameters.SvPosition.xy, GBuffer.ShadingModelID, OutDiffuseLighting, OutSubsurfaceLighting, VTRequest);
 
 	#if LIGHTMAP_VT_ENABLED
 		uint2 PixelPos = (uint2)MaterialParameters.SvPosition.xy;
@@ -498,7 +516,7 @@ void GetPrecomputedIndirectLightingAndSkyLight(
 
 	float3 SkyDiffuseLighting;
 	float3 SkySubsurfaceLighting;
-	GetSkyLighting(MaterialParameters, DiffuseDir, SkyOcclusionUV, VolumetricLightmapBrickTextureUVs, SkyDiffuseLighting, SkySubsurfaceLighting);
+	GetSkyLighting(MaterialParameters, GBuffer, DiffuseDir, SkyOcclusionUV, VolumetricLightmapBrickTextureUVs, SkyDiffuseLighting, SkySubsurfaceLighting);
 
 	OutSubsurfaceLighting += SkySubsurfaceLighting;
 
@@ -724,7 +742,7 @@ uint GetSelectiveOutputMask()
 void FPixelShaderInOut_MainPS(
 	FVertexFactoryInterpolantsVSToPS Interpolants,
 	FBasePassInterpolantsVSToPS BasePassInterpolants,
-	in FPixelShaderIn In, 
+	in FPixelShaderIn In,
 	inout FPixelShaderOut Out)
 {
 #if INSTANCED_STEREO
@@ -742,7 +760,7 @@ void FPixelShaderInOut_MainPS(
 
 	FMaterialPixelParameters MaterialParameters = GetMaterialPixelParameters(Interpolants, In.SvPosition);
 	FPixelMaterialInputs PixelMaterialInputs;
-	
+
 #if HQ_TEXTURE_LIGHTMAP && USES_AO_MATERIAL_MASK && !MATERIAL_SHADINGMODEL_UNLIT
 	{
 		float2 LightmapUV0, LightmapUV1;
@@ -753,39 +771,39 @@ void FPixelShaderInOut_MainPS(
 	}
 #endif
 
-	#if USE_WORLD_POSITION_EXCLUDING_SHADER_OFFSETS
-		{
-			float4 ScreenPosition = SvPositionToResolvedScreenPosition(In.SvPosition);
-			float3 TranslatedWorldPosition = SvPositionToResolvedTranslatedWorld(In.SvPosition);
-			CalcMaterialParametersEx(MaterialParameters, PixelMaterialInputs, In.SvPosition, ScreenPosition, In.bIsFrontFace, TranslatedWorldPosition, BasePassInterpolants.PixelPositionExcludingWPO);
-		}
-	#else
-		{
-			float4 ScreenPosition = SvPositionToResolvedScreenPosition(In.SvPosition);
-			float3 TranslatedWorldPosition = SvPositionToResolvedTranslatedWorld(In.SvPosition);
-			CalcMaterialParametersEx(MaterialParameters, PixelMaterialInputs, In.SvPosition, ScreenPosition, In.bIsFrontFace, TranslatedWorldPosition, TranslatedWorldPosition);
-		}
-	#endif
+#if USE_WORLD_POSITION_EXCLUDING_SHADER_OFFSETS
+	{
+		float4 ScreenPosition = SvPositionToResolvedScreenPosition(In.SvPosition);
+		float3 TranslatedWorldPosition = SvPositionToResolvedTranslatedWorld(In.SvPosition);
+		CalcMaterialParametersEx(MaterialParameters, PixelMaterialInputs, In.SvPosition, ScreenPosition, In.bIsFrontFace, TranslatedWorldPosition, BasePassInterpolants.PixelPositionExcludingWPO);
+	}
+#else
+	{
+		float4 ScreenPosition = SvPositionToResolvedScreenPosition(In.SvPosition);
+		float3 TranslatedWorldPosition = SvPositionToResolvedTranslatedWorld(In.SvPosition);
+		CalcMaterialParametersEx(MaterialParameters, PixelMaterialInputs, In.SvPosition, ScreenPosition, In.bIsFrontFace, TranslatedWorldPosition, TranslatedWorldPosition);
+	}
+#endif
 
 #if USE_EDITOR_COMPOSITING && (FEATURE_LEVEL >= FEATURE_LEVEL_SM4 || MOBILE_EMULATION)
 	const bool bEditorWeightedZBuffering = true;
 #else
 	const bool bEditorWeightedZBuffering = false;
 #endif
-		
+
 #if OUTPUT_PIXEL_DEPTH_OFFSET
 	ApplyPixelDepthOffsetForBasePass(MaterialParameters, PixelMaterialInputs, BasePassInterpolants, Out.Depth);
 #endif
 
 	//Clip if the blend mode requires it.
 #if !EARLY_Z_PASS_ONLY_MATERIAL_MASKING
-	if(!bEditorWeightedZBuffering)
+	if (!bEditorWeightedZBuffering)
 	{
-		#if MATERIALBLENDING_MASKED_USING_COVERAGE
-			Out.Coverage = DiscardMaterialWithPixelCoverage(MaterialParameters, PixelMaterialInputs);
-		#else
-			GetMaterialCoverageAndClipping(MaterialParameters, PixelMaterialInputs);
-		#endif
+#if MATERIALBLENDING_MASKED_USING_COVERAGE
+		Out.Coverage = DiscardMaterialWithPixelCoverage(MaterialParameters, PixelMaterialInputs);
+#else
+		GetMaterialCoverageAndClipping(MaterialParameters, PixelMaterialInputs);
+#endif
 	}
 #endif
 	// Store the results in local variables and reuse instead of calling the functions multiple times.
@@ -795,6 +813,7 @@ void FPixelShaderInOut_MainPS(
 
 	float MaterialAO = GetMaterialAmbientOcclusion(PixelMaterialInputs);
 	float Roughness = GetMaterialRoughness(PixelMaterialInputs);
+	uint ShadingModel = GetMaterialShadingModel(PixelMaterialInputs);
 
 	// 0..1, SubsurfaceProfileId = int(x * 255)
 	float SubsurfaceProfile = 0;
@@ -802,14 +821,26 @@ void FPixelShaderInOut_MainPS(
 	// If we don't use this shading model the color should be black (don't generate shader code for unused data, don't do indirectlighting cache lighting with this color).
 	float3 SubsurfaceColor = 0;
 #if MATERIAL_SHADINGMODEL_SUBSURFACE || MATERIAL_SHADINGMODEL_PREINTEGRATED_SKIN || MATERIAL_SHADINGMODEL_SUBSURFACE_PROFILE || MATERIAL_SHADINGMODEL_TWOSIDED_FOLIAGE || MATERIAL_SHADINGMODEL_CLOTH || MATERIAL_SHADINGMODEL_EYE
+	if (ShadingModel == SHADINGMODELID_SUBSURFACE || ShadingModel == SHADINGMODELID_PREINTEGRATED_SKIN || ShadingModel == SHADINGMODELID_SUBSURFACE_PROFILE || ShadingModel == SHADINGMODELID_TWOSIDED_FOLIAGE || ShadingModel == SHADINGMODELID_CLOTH || ShadingModel == SHADINGMODELID_EYE)
 	{
 		float4 SubsurfaceData = GetMaterialSubsurfaceData(PixelMaterialInputs);
 
+		if (false) // Dummy if to make the ifdef logic play nicely
+		{
+		}
 #if MATERIAL_SHADINGMODEL_SUBSURFACE || MATERIAL_SHADINGMODEL_PREINTEGRATED_SKIN || MATERIAL_SHADINGMODEL_TWOSIDED_FOLIAGE
-		SubsurfaceColor = SubsurfaceData.rgb * View.DiffuseOverrideParameter.w + View.DiffuseOverrideParameter.xyz;
-#elif MATERIAL_SHADINGMODEL_CLOTH
-		SubsurfaceColor = SubsurfaceData.rgb;
+		else if (ShadingModel == SHADINGMODELID_SUBSURFACE || ShadingModel == SHADINGMODELID_PREINTEGRATED_SKIN || ShadingModel == SHADINGMODELID_TWOSIDED_FOLIAGE)
+		{
+			SubsurfaceColor = SubsurfaceData.rgb * View.DiffuseOverrideParameter.w + View.DiffuseOverrideParameter.xyz;
+		}
+#endif
+#if MATERIAL_SHADINGMODEL_CLOTH
+		else if (ShadingModel == SHADINGMODELID_CLOTH)
+		{
+			SubsurfaceColor = SubsurfaceData.rgb;
+		}
 #endif
+
 		SubsurfaceProfile = SubsurfaceData.a;
 	}
 #endif
@@ -822,30 +853,30 @@ void FPixelShaderInOut_MainPS(
 #else
 	BRANCH
 #endif
-	if (GetPrimitiveData(MaterialParameters.PrimitiveId).DecalReceiverMask > 0 && View.ShowDecalsMask > 0)
-	{
-		uint DBufferMask = 0x07;
-	#if PLATFORM_SUPPORTS_RENDERTARGET_WRITE_MASK
-		DBufferMask = DecodeRTWriteMaskTexture(In.SvPosition.xy, OpaqueBasePass.DBufferRenderMask);
-		BRANCH
-	#elif PLATFORM_SUPPORTS_PER_PIXEL_DBUFFER_MASK
-		DBufferMask = GetDBufferMask(In.SvPosition.xy, OpaqueBasePass.DBufferRenderMask);
-	#endif
-		if(DBufferMask)
+		if (GetPrimitiveData(MaterialParameters.PrimitiveId).DecalReceiverMask > 0 && View.ShowDecalsMask > 0)
 		{
-			float2 NDC = MaterialParameters.ScreenPosition.xy / MaterialParameters.ScreenPosition.w;
-			float2 ScreenUV = NDC * ResolvedView.ScreenPositionScaleBias.xy + ResolvedView.ScreenPositionScaleBias.wz;
-			FDBufferData DBufferData = GetDBufferData(ScreenUV, DBufferMask);
+			uint DBufferMask = 0x07;
+#if PLATFORM_SUPPORTS_RENDERTARGET_WRITE_MASK
+			DBufferMask = DecodeRTWriteMaskTexture(In.SvPosition.xy, OpaqueBasePass.DBufferRenderMask);
+			BRANCH
+#elif PLATFORM_SUPPORTS_PER_PIXEL_DBUFFER_MASK
+			DBufferMask = GetDBufferMask(In.SvPosition.xy, OpaqueBasePass.DBufferRenderMask);
+#endif
+			if (DBufferMask)
+			{
+				float2 NDC = MaterialParameters.ScreenPosition.xy / MaterialParameters.ScreenPosition.w;
+				float2 ScreenUV = NDC * ResolvedView.ScreenPositionScaleBias.xy + ResolvedView.ScreenPositionScaleBias.wz;
+				FDBufferData DBufferData = GetDBufferData(ScreenUV, DBufferMask);
 
-			ApplyDBufferData(DBufferData, MaterialParameters.WorldNormal, SubsurfaceColor, Roughness, BaseColor, Metallic, Specular);
+				ApplyDBufferData(DBufferData, MaterialParameters.WorldNormal, SubsurfaceColor, Roughness, BaseColor, Metallic, Specular);
+			}
 		}
-	}
 #endif
 
 	half Opacity = GetMaterialOpacity(PixelMaterialInputs);
 
 	float3 VolumetricLightmapBrickTextureUVs;
-		
+
 #if PRECOMPUTED_IRRADIANCE_VOLUME_LIGHTING
 	VolumetricLightmapBrickTextureUVs = ComputeVolumetricLightmapBrickTextureUVs(MaterialParameters.AbsoluteWorldPosition);
 #endif
@@ -858,17 +889,19 @@ void FPixelShaderInOut_MainPS(
 	GBuffer.PrecomputedShadowFactors = GetPrecomputedShadowMasks(Interpolants, MaterialParameters.PrimitiveId, MaterialParameters.AbsoluteWorldPosition, VolumetricLightmapBrickTextureUVs);
 
 	const float GBufferDither = InterleavedGradientNoise(MaterialParameters.SvPosition.xy, View.StateFrameIndexMod8);
+	// Use GBuffer.ShadingModelID after SetGBufferForShadingModel(..) because the ShadingModel input might not be the same as the output
 	SetGBufferForShadingModel(
-		GBuffer, 
-		MaterialParameters, 
-		Opacity, 
-		BaseColor, 
-		Metallic, 
-		Specular, 
-		Roughness, 
+		GBuffer,
+		MaterialParameters,
+		Opacity,
+		BaseColor,
+		Metallic,
+		Specular,
+		Roughness,
 		SubsurfaceColor,
 		SubsurfaceProfile,
-		GBufferDither);
+		GBufferDither,
+ShadingModel);
 
 #if USES_GBUFFER
 	GBuffer.SelectiveOutputMask = GetSelectiveOutputMask();
@@ -903,14 +936,17 @@ void FPixelShaderInOut_MainPS(
 	// Curvature-to-roughness uses derivatives of the WorldVertexNormal, which is incompatible with centroid interpolation because
 	// the samples are not uniformly distributed. Therefore we use WorldVertexNormal_Center which is guaranteed to be center interpolated.
 #if USE_WORLDVERTEXNORMAL_CENTER_INTERPOLATION
-    float GeometricAARoughness = NormalCurvatureToRoughness(MaterialParameters.WorldVertexNormal_Center);
+	float GeometricAARoughness = NormalCurvatureToRoughness(MaterialParameters.WorldVertexNormal_Center);
 #else
-    float GeometricAARoughness = NormalCurvatureToRoughness(MaterialParameters.TangentToWorld[2].xyz);
+	float GeometricAARoughness = NormalCurvatureToRoughness(MaterialParameters.TangentToWorld[2].xyz);
 #endif
-    GBuffer.Roughness = max(GBuffer.Roughness, GeometricAARoughness);
+	GBuffer.Roughness = max(GBuffer.Roughness, GeometricAARoughness);
 
 #if MATERIAL_SHADINGMODEL_CLEAR_COAT
-	GBuffer.CustomData.y = max(GBuffer.CustomData.y, GeometricAARoughness);
+	if (GBuffer.ShadingModelID == SHADINGMODELID_CLEAR_COAT)
+	{
+		GBuffer.CustomData.y = max(GBuffer.CustomData.y, GeometricAARoughness);
+	}	
 #endif
 #endif
 
@@ -965,16 +1001,23 @@ void FPixelShaderInOut_MainPS(
 		float3 DiffuseColorForIndirect = GBuffer.DiffuseColor;
 
 		#if MATERIAL_SHADINGMODEL_SUBSURFACE || MATERIAL_SHADINGMODEL_PREINTEGRATED_SKIN
+		if (GBuffer.ShadingModelID == SHADINGMODELID_SUBSURFACE || GBuffer.ShadingModelID == SHADINGMODELID_PREINTEGRATED_SKIN)
+		{
 			// Add subsurface energy to diffuse
 			//@todo - better subsurface handling for these shading models with skylight and precomputed GI
 			DiffuseColorForIndirect += SubsurfaceColor;
+		}
 		#endif
 
 		#if MATERIAL_SHADINGMODEL_CLOTH
-			DiffuseColorForIndirect += SubsurfaceColor * saturate( GetMaterialCustomData0(MaterialParameters) );
+		if (GBuffer.ShadingModelID == SHADINGMODELID_CLOTH)
+		{
+			DiffuseColorForIndirect += SubsurfaceColor * saturate(GetMaterialCustomData0(MaterialParameters));
+		}
 		#endif
 
 		#if MATERIAL_SHADINGMODEL_HAIR
+		if (GBuffer.ShadingModelID == SHADINGMODELID_HAIR)
 		{
 			float3 N = MaterialParameters.WorldNormal;
 			float3 V = MaterialParameters.CameraVector;
@@ -986,7 +1029,7 @@ void FPixelShaderInOut_MainPS(
 
 		float3 DiffuseIndirectLighting;
 		float3 SubsurfaceIndirectLighting;
-		GetPrecomputedIndirectLightingAndSkyLight(MaterialParameters, Interpolants, BasePassInterpolants, DiffuseDir, VolumetricLightmapBrickTextureUVs, DiffuseIndirectLighting, SubsurfaceIndirectLighting, IndirectIrradiance);
+		GetPrecomputedIndirectLightingAndSkyLight(MaterialParameters, Interpolants, BasePassInterpolants, GBuffer, DiffuseDir, VolumetricLightmapBrickTextureUVs, DiffuseIndirectLighting, SubsurfaceIndirectLighting, IndirectIrradiance);
 		
 		float IndirectOcclusion = 1.0f;
 		float2 NearestResolvedDepthScreenUV = 0;
@@ -1025,13 +1068,15 @@ void FPixelShaderInOut_MainPS(
 					Color += GetForwardDirectLighting(GridIndex, MaterialParameters.AbsoluteWorldPosition, MaterialParameters.CameraVector, GBuffer, NearestResolvedDepthScreenUV, MaterialParameters.PrimitiveId, EyeIndex, Dither);
 				#endif
 			#endif
-            
-            #if !MATERIAL_SHADINGMODEL_HAIR
-				int SingleCaptureIndex = GetPrimitiveData(MaterialParameters.PrimitiveId).SingleCaptureIndex;
-			    Color += GetImageBasedReflectionLighting(MaterialParameters, GBuffer.Roughness, GBuffer.SpecularColor, IndirectIrradiance, GridIndex, SingleCaptureIndex, EyeIndex) 
-					* IndirectOcclusion 
-					* AOMultiBounce( GBuffer.SpecularColor, SpecOcclusion );
-		    #endif
+			#if !(MATERIAL_SINGLE_SHADINGMODEL && MATERIAL_SHADINGMODEL_HAIR)
+				if (GBuffer.ShadingModelID != SHADINGMODELID_HAIR)
+				{
+					int SingleCaptureIndex = GetPrimitiveData(MaterialParameters.PrimitiveId).SingleCaptureIndex;
+					Color += GetImageBasedReflectionLighting(MaterialParameters, GBuffer.Roughness, GBuffer.SpecularColor, IndirectIrradiance, GridIndex, SingleCaptureIndex, EyeIndex)
+						* IndirectOcclusion
+						* AOMultiBounce(GBuffer.SpecularColor, SpecOcclusion);
+				}
+			#endif
         #endif
 
 		#if SIMPLE_FORWARD_DIRECTIONAL_LIGHT
@@ -1076,9 +1121,12 @@ void FPixelShaderInOut_MainPS(
 #endif
 
 	// Volume lighting for lit translucency
-	#if (MATERIAL_SHADINGMODEL_DEFAULT_LIT || MATERIAL_SHADINGMODEL_SUBSURFACE) && (MATERIALBLENDING_TRANSLUCENT || MATERIALBLENDING_ADDITIVE) && !SIMPLE_FORWARD_SHADING && !FORWARD_SHADING
+#if (MATERIAL_SHADINGMODEL_DEFAULT_LIT || MATERIAL_SHADINGMODEL_SUBSURFACE) && (MATERIALBLENDING_TRANSLUCENT || MATERIALBLENDING_ADDITIVE) && !SIMPLE_FORWARD_SHADING && !FORWARD_SHADING
+	if (GBuffer.ShadingModelID == SHADINGMODELID_DEFAULT_LIT || GBuffer.ShadingModelID == SHADINGMODELID_SUBSURFACE)
+	{
 		Color += GetTranslucencyVolumeLighting(MaterialParameters, PixelMaterialInputs, BasePassInterpolants, GBuffer, IndirectIrradiance);
-	#endif
+	}
+#endif
 
 	#if !MATERIAL_SHADINGMODEL_UNLIT && USE_DEVELOPMENT_SHADERS
 		Color = lerp(Color, GBuffer.DiffuseColor + GBuffer.SpecularColor * 0.45, View.UnlitViewmodeMask);
diff --git a/Engine/Shaders/Private/Definitions.usf b/Engine/Shaders/Private/Definitions.usf
--- a/Engine/Shaders/Private/Definitions.usf
+++ b/Engine/Shaders/Private/Definitions.usf
@@ -58,6 +58,10 @@
 #define	MATERIAL_SHADINGMODEL_UNLIT						0
 #endif
 
+#ifndef MATERIAL_SINGLE_SHADINGMODEL
+#define	MATERIAL_SINGLE_SHADINGMODEL					0
+#endif
+
 #ifndef HAS_PRIMITIVE_UNIFORM_BUFFER
 #define	HAS_PRIMITIVE_UNIFORM_BUFFER					0
 #endif
diff --git a/Engine/Shaders/Private/LightmapCommon.ush b/Engine/Shaders/Private/LightmapCommon.ush
--- a/Engine/Shaders/Private/LightmapCommon.ush
+++ b/Engine/Shaders/Private/LightmapCommon.ush
@@ -1,5 +1,6 @@
 // Copyright 1998-2019 Epic Games, Inc. All Rights Reserved.
 
+#include "ShadingCommon.ush"
 #include "/Engine/Generated/UniformBuffers/PrecomputedLightingBuffer.ush"
 #include "VolumetricLightmapShared.ush"
 #include "VirtualTextureCommon.ush"
@@ -79,7 +80,7 @@ half4 GetLightMapColorLQ( float2 LightmapUV0, float2 LightmapUV1, uint LightmapD
 	return half4( Color, Luma );
 }
 
-void GetLightMapColorHQ( float2 LightmapUV0, float2 LightmapUV1, uint LightmapDataIndex, half3 WorldNormal, float2 SvPositionXY, out half3 OutDiffuseLighting, out half3 OutSubsurfaceLighting, out uint VTRequest )
+void GetLightMapColorHQ( float2 LightmapUV0, float2 LightmapUV1, uint LightmapDataIndex, half3 WorldNormal, float2 SvPositionXY, uint ShadingModel, out half3 OutDiffuseLighting, out half3 OutSubsurfaceLighting, out uint VTRequest )
 {
 	OutSubsurfaceLighting = 0;
 
@@ -123,14 +124,20 @@ void GetLightMapColorHQ( float2 LightmapUV0, float2 LightmapUV1, uint LightmapDa
 	half Directionality = max( 0.0, dot( SH, float4(WorldNormal.yzx, 1) ) );
 
 	#if MATERIAL_SHADINGMODEL_TWOSIDED_FOLIAGE
-		half SubsurfaceDirectionality = max( 0.0, dot( SH, float4(-WorldNormal.yzx, 1) ) );
+	if (ShadingModel == SHADINGMODELID_TWOSIDED_FOLIAGE)
+	{
+		half SubsurfaceDirectionality = max(0.0, dot(SH, float4(-WorldNormal.yzx, 1)));
 		OutSubsurfaceLighting = L * SubsurfaceDirectionality * UVW;
+	}
 	#endif
 #else
 	half Directionality = 0.6;
 
 	#if MATERIAL_SHADINGMODEL_TWOSIDED_FOLIAGE
+	if (ShadingModel == SHADINGMODELID_TWOSIDED_FOLIAGE)
+	{
 		OutSubsurfaceLighting = L * Directionality * UVW;
+	}
 	#endif
 #endif
 
diff --git a/Engine/Shaders/Private/MaterialTemplate.ush b/Engine/Shaders/Private/MaterialTemplate.ush
--- a/Engine/Shaders/Private/MaterialTemplate.ush
+++ b/Engine/Shaders/Private/MaterialTemplate.ush
@@ -1621,6 +1621,12 @@ half3 GetMaterialEmissiveForCS(FMaterialPixelParameters Parameters)
 %s;
 }
 
+// Shading Model is an uint and represents a SHADINGMODELID_* in ShadingCommon.ush 
+uint GetMaterialShadingModel(FPixelMaterialInputs PixelMaterialInputs)
+{
+	return PixelMaterialInputs.ShadingModel;
+}
+
 half3 GetMaterialBaseColorRaw(FPixelMaterialInputs PixelMaterialInputs)
 {
 	return PixelMaterialInputs.BaseColor;
diff --git a/Engine/Shaders/Private/RayTracing/RayTracingMaterialHitShaders.usf b/Engine/Shaders/Private/RayTracing/RayTracingMaterialHitShaders.usf
--- a/Engine/Shaders/Private/RayTracing/RayTracingMaterialHitShaders.usf
+++ b/Engine/Shaders/Private/RayTracing/RayTracingMaterialHitShaders.usf
@@ -29,7 +29,7 @@
 #include "/Engine/Private/LightmapCommon.ush"
 
 /** Computes sky diffuse lighting, including precomputed shadowing. */
-void GetSkyLighting(FMaterialPixelParameters MaterialParameters, float3 WorldNormal, float2 LightmapUV, float3 SkyOcclusionUV3D, out float3 OutDiffuseLighting, out float3 OutSubsurfaceLighting)
+void GetSkyLighting(FMaterialPixelParameters MaterialParameters, float3 WorldNormal, float2 LightmapUV, float3 SkyOcclusionUV3D, uint ShadingModel, out float3 OutDiffuseLighting, out float3 OutSubsurfaceLighting)
 {
 	OutDiffuseLighting = 0;
 	OutSubsurfaceLighting = 0;
@@ -94,8 +94,11 @@ void GetSkyLighting(FMaterialPixelParameters MaterialParameters, float3 WorldNor
 	OutDiffuseLighting += DiffuseLookup * (SkyVisibility * GeometryTerm);
 
 	#if MATERIAL_SHADINGMODEL_TWOSIDED_FOLIAGE
+	if (ShadingModel == SHADINGMODELID_TWOSIDED_FOLIAGE)
+	{
 		float3 BackfaceDiffuseLookup = GetSkySHDiffuse(-WorldNormal) * ResolvedView.SkyLightColor.rgb;
 		OutSubsurfaceLighting += BackfaceDiffuseLookup * SkyVisibility;
+	}
 	#endif
 
 #endif
@@ -107,6 +110,7 @@ void GetPrecomputedIndirectLightingAndSkyLight(
 	FVertexFactoryInterpolantsVSToPS Interpolants,
 	float3 DiffuseDir,
 	float3 VolumetricLightmapBrickTextureUVs,
+	uint ShadingModel,
 	out float3 OutDiffuseLighting,
 	out float3 OutSubsurfaceLighting,
 	out float OutIndirectIrradiance)
@@ -124,8 +128,11 @@ void GetPrecomputedIndirectLightingAndSkyLight(
 		OutDiffuseLighting = max(float3(0,0,0), DotSH3(IrradianceSH, DiffuseTransferSH)) / PI;
 
 		#if MATERIAL_SHADINGMODEL_TWOSIDED_FOLIAGE
+		if (ShadingModel == SHADINGMODELID_TWOSIDED_FOLIAGE)
+		{
 			FThreeBandSHVector SubsurfaceTransferSH = CalcDiffuseTransferSH3(-DiffuseDir, 1);
 			OutSubsurfaceLighting += max(float3(0,0,0), DotSH3(IrradianceSH, SubsurfaceTransferSH)) / PI;
+		}
 		#endif
 
 	// High quality texture lightmaps
@@ -154,7 +161,7 @@ void GetPrecomputedIndirectLightingAndSkyLight(
 
 	float3 SkyDiffuseLighting;
 	float3 SkySubsurfaceLighting;
-	GetSkyLighting(MaterialParameters, DiffuseDir, SkyOcclusionUV, VolumetricLightmapBrickTextureUVs, SkyDiffuseLighting, SkySubsurfaceLighting);
+	GetSkyLighting(MaterialParameters, DiffuseDir, SkyOcclusionUV, VolumetricLightmapBrickTextureUVs, ShadingModel, SkyDiffuseLighting, SkySubsurfaceLighting);
 
 	OutSubsurfaceLighting += SkySubsurfaceLighting;
 
@@ -310,7 +317,7 @@ void MaterialCHS(inout FPackedMaterialClosestHitPayload PackedPayload, in FDefau
 	}
 	#endif
 	Payload.Opacity = GetMaterialOpacity(PixelMaterialInputs);
-    Payload.ShadingModelID = GetShadingModelID();
+    Payload.ShadingModelID = GetMaterialShadingModel(PixelMaterialInputs);
 
  #if MATERIALBLENDING_ALPHACOMPOSITE
 	uint MaterialBlendingMode = RAY_TRACING_BLEND_MODE_ALPHA_COMPOSITE;
@@ -342,7 +349,7 @@ void MaterialCHS(inout FPackedMaterialClosestHitPayload PackedPayload, in FDefau
 #if PRECOMPUTED_IRRADIANCE_VOLUME_LIGHTING
 	VolumetricLightmapBrickTextureUVs = ComputeVolumetricLightmapBrickTextureUVs(MaterialParameters.AbsoluteWorldPosition);
 #endif
-	GetPrecomputedIndirectLightingAndSkyLight(MaterialParameters, Interpolants, Payload.WorldNormal, VolumetricLightmapBrickTextureUVs, DiffuseIndirectLighting, SubsurfaceIndirectLighting, IndirectIrradiance);
+	GetPrecomputedIndirectLightingAndSkyLight(MaterialParameters, Interpolants, Payload.WorldNormal, VolumetricLightmapBrickTextureUVs, Payload.ShadingModelID, DiffuseIndirectLighting, SubsurfaceIndirectLighting, IndirectIrradiance);
 	
 	Payload.DiffuseColor = BaseColor - BaseColor * Metallic;
 	Payload.SpecularColor = ComputeF0(Specular, BaseColor, Metallic);
diff --git a/Engine/Shaders/Private/ShadingCommon.ush b/Engine/Shaders/Private/ShadingCommon.ush
--- a/Engine/Shaders/Private/ShadingCommon.ush
+++ b/Engine/Shaders/Private/ShadingCommon.ush
@@ -4,6 +4,7 @@
 #ifndef __SHADING_COMMON__
 #define __SHADING_COMMON__
 
+// SHADINGMODELID_* occupy the 4 low bits of an 8bit channel and SKIP_* occupy the 4 high bits
 #define SHADINGMODELID_UNLIT				0
 #define SHADINGMODELID_DEFAULT_LIT			1
 #define SHADINGMODELID_SUBSURFACE			2
@@ -15,41 +16,15 @@
 #define SHADINGMODELID_CLOTH				8
 #define SHADINGMODELID_EYE					9
 #define SHADINGMODELID_NUM					10
-#define SHADINGMODELID_MASK					0xF					
+#define SHADINGMODELID_MASK					0xF		// 4 bits reserved for ShadingModelID			
 
 // The flags are defined so that 0 value has no effect!
+// These occupy the 4 high bits in the same channel as the SHADINGMODELID_*
 #define SKIP_CUSTOMDATA_MASK			(1 << 4)	// TODO remove. Can be inferred from shading model.
 #define SKIP_PRECSHADOW_MASK			(1 << 5)
 #define ZERO_PRECSHADOW_MASK			(1 << 6)
 #define SKIP_VELOCITY_MASK				(1 << 7)
 
-uint GetShadingModelID()
-{        
-#if MATERIAL_SHADINGMODEL_UNLIT
-	return SHADINGMODELID_UNLIT;
-#elif MATERIAL_SHADINGMODEL_DEFAULT_LIT
-	return SHADINGMODELID_DEFAULT_LIT;
-#elif MATERIAL_SHADINGMODEL_SUBSURFACE
-	return SHADINGMODELID_SUBSURFACE;
-#elif MATERIAL_SHADINGMODEL_PREINTEGRATED_SKIN
-	return SHADINGMODELID_PREINTEGRATED_SKIN;
-#elif MATERIAL_SHADINGMODEL_SUBSURFACE_PROFILE
-	return SHADINGMODELID_SUBSURFACE_PROFILE;
-#elif MATERIAL_SHADINGMODEL_CLEAR_COAT
-	return SHADINGMODELID_CLEAR_COAT;
-#elif MATERIAL_SHADINGMODEL_TWOSIDED_FOLIAGE
-	return SHADINGMODELID_TWOSIDED_FOLIAGE;
-#elif MATERIAL_SHADINGMODEL_HAIR
-	return SHADINGMODELID_HAIR;
-#elif MATERIAL_SHADINGMODEL_CLOTH
-	return SHADINGMODELID_CLOTH;
-#elif MATERIAL_SHADINGMODEL_EYE
-	return SHADINGMODELID_EYE;
-#else
-	return SHADINGMODELID_MASK;
-#endif
-}
-
 // for debugging and to visualize
 float3 GetShadingModelColor(uint ShadingModelID)
 {
diff --git a/Engine/Shaders/Private/ShadingModelsMaterial.ush b/Engine/Shaders/Private/ShadingModelsMaterial.ush
--- a/Engine/Shaders/Private/ShadingModelsMaterial.ush
+++ b/Engine/Shaders/Private/ShadingModelsMaterial.ush
@@ -19,169 +19,192 @@ void SetGBufferForShadingModel(
 	const float Roughness,
 	const float3 SubsurfaceColor,
 	const float SubsurfaceProfile,
-	const float dither)
+	const float Dither,
+        const uint ShadingModel)
 {
 	GBuffer.WorldNormal = MaterialParameters.WorldNormal;
 	GBuffer.BaseColor = BaseColor;
 	GBuffer.Metallic = Metallic;
 	GBuffer.Specular = Specular;
 	GBuffer.Roughness = Roughness;
+	GBuffer.ShadingModelID = ShadingModel;
 
-#if MATERIAL_SHADINGMODEL_UNLIT
-	GBuffer.ShadingModelID = SHADINGMODELID_UNLIT;
-#elif MATERIAL_SHADINGMODEL_DEFAULT_LIT
-	GBuffer.ShadingModelID = SHADINGMODELID_DEFAULT_LIT;
-#elif MATERIAL_SHADINGMODEL_SUBSURFACE
-	GBuffer.ShadingModelID = SHADINGMODELID_SUBSURFACE;
-	GBuffer.CustomData.rgb = EncodeSubsurfaceColor(SubsurfaceColor);
-	GBuffer.CustomData.a = Opacity;
-#elif MATERIAL_SHADINGMODEL_PREINTEGRATED_SKIN
-	GBuffer.ShadingModelID = SHADINGMODELID_PREINTEGRATED_SKIN;
-	GBuffer.CustomData.rgb = EncodeSubsurfaceColor(SubsurfaceColor);
-	GBuffer.CustomData.a = Opacity;
-#elif MATERIAL_SHADINGMODEL_SUBSURFACE_PROFILE
-	// Optimization: if opacity is 0 then revert to default shading model 
-	#if SUBSURFACE_PROFILE_OPACITY_THRESHOLD
-	if (Opacity > 0.0)
-	#endif
+	// Calculate and set custom data for the shading models that need it
+	
+	// Dummy initial if statement to play nicely with the #ifdef logic
+	if (false)
+	{
+	}
+#if MATERIAL_SHADINGMODEL_SUBSURFACE
+	else if (ShadingModel == SHADINGMODELID_SUBSURFACE)
 	{
-		GBuffer.ShadingModelID = SHADINGMODELID_SUBSURFACE_PROFILE;
-		GBuffer.CustomData.rgb = EncodeSubsurfaceProfile(SubsurfaceProfile);
+		GBuffer.CustomData.rgb = EncodeSubsurfaceColor(SubsurfaceColor);
 		GBuffer.CustomData.a = Opacity;
-
-		// Average roughness for dual specular.
-		uint SubsurfaceProfileUInt = uint(SubsurfaceProfile * 255.0f + 0.5f);
-		float MaterialRoughnessToAverage = ActualSSProfilesTexture.Load(int3(SSSS_DUAL_SPECULAR_OFFSET, SubsurfaceProfileUInt, 0)).w * SSSS_MAX_DUAL_SPECULAR_ROUGHNESS;
-
-		// Smooth blend out dual specular when opacity is low.
-		MaterialRoughnessToAverage = lerp(1.0f, MaterialRoughnessToAverage, saturate(Opacity * 10.0f));
-
-		GBuffer.Roughness = saturate(GBuffer.Roughness * MaterialRoughnessToAverage);
 	}
-	#if SUBSURFACE_PROFILE_OPACITY_THRESHOLD
-	else
+#endif
+#if MATERIAL_SHADINGMODEL_PREINTEGRATED_SKIN
+	else if (ShadingModel == SHADINGMODELID_PREINTEGRATED_SKIN)
 	{
-		GBuffer.ShadingModelID = SHADINGMODELID_DEFAULT_LIT;
-		GBuffer.CustomData = 0;
+		GBuffer.CustomData.rgb = EncodeSubsurfaceColor(SubsurfaceColor);
+		GBuffer.CustomData.a = Opacity;
 	}
+#endif
+#if MATERIAL_SHADINGMODEL_SUBSURFACE_PROFILE
+	else if (ShadingModel == SHADINGMODELID_SUBSURFACE_PROFILE)
+	{
+		// Optimization: if opacity is 0 then revert to default shading model 
+	#if SUBSURFACE_PROFILE_OPACITY_THRESHOLD
+		if (Opacity > 0.0)
 	#endif
-#elif MATERIAL_SHADINGMODEL_CLEAR_COAT
-	GBuffer.ShadingModelID = SHADINGMODELID_CLEAR_COAT;
+		{
+			GBuffer.CustomData.rgb = EncodeSubsurfaceProfile(SubsurfaceProfile);
+			GBuffer.CustomData.a = Opacity;
 
-	float ClearCoat				= saturate( GetMaterialCustomData0(MaterialParameters) );
-	float ClearCoatRoughness	= saturate( GetMaterialCustomData1(MaterialParameters) );
-	float MetalSpec = 0.9;
+			// Average roughness for dual specular.
+			uint SubsurfaceProfileUInt = uint(SubsurfaceProfile * 255.0f + 0.5f);
+			float MaterialRoughnessToAverage = ActualSSProfilesTexture.Load(int3(SSSS_DUAL_SPECULAR_OFFSET, SubsurfaceProfileUInt, 0)).w * SSSS_MAX_DUAL_SPECULAR_ROUGHNESS;
 
-	float NoV = saturate( dot( MaterialParameters.WorldNormal, MaterialParameters.CameraVector ) );
+			// Smooth blend out dual specular when opacity is low.
+			MaterialRoughnessToAverage = lerp(1.0f, MaterialRoughnessToAverage, saturate(Opacity * 10.0f));
 
-	// Approximation of refraction's effect on EnvBRDF
-	float RefractionScale = ( (NoV * 0.5 + 0.5) * NoV - 1 ) * saturate( 1.25 - 1.25 * Roughness ) + 1;
+			GBuffer.Roughness = saturate(GBuffer.Roughness * MaterialRoughnessToAverage);
+		}
+	#if SUBSURFACE_PROFILE_OPACITY_THRESHOLD
+		else
+		{
+			GBuffer.ShadingModelID = SHADINGMODELID_DEFAULT_LIT;
+			GBuffer.CustomData = 0;
+		}
+	#endif
+	}
+#endif
+#if MATERIAL_SHADINGMODEL_CLEAR_COAT
+	else if (ShadingModel == SHADINGMODELID_CLEAR_COAT)
+	{
+		float ClearCoat				= saturate( GetMaterialCustomData0(MaterialParameters) );
+		float ClearCoatRoughness	= saturate( GetMaterialCustomData1(MaterialParameters) );
+		float MetalSpec = 0.9;
+
+		float NoV = saturate( dot( MaterialParameters.WorldNormal, MaterialParameters.CameraVector ) );
 
-	// Approximation of absorption integral, tuned for Roughness=0.4
-	float3 AbsorptionColor = BaseColor * (1 / MetalSpec);
-	float3 Absorption = AbsorptionColor * ( (NoV - 1) * 0.85 * ( 1 - lerp( AbsorptionColor, Square(AbsorptionColor), -0.78 ) ) + 1 );
+		// Approximation of refraction's effect on EnvBRDF
+		float RefractionScale = ( (NoV * 0.5 + 0.5) * NoV - 1 ) * saturate( 1.25 - 1.25 * Roughness ) + 1;
 
-	float F0 = 0.04;
-	float Fc = Pow5( 1 - NoV );
-	float F = Fc + (1 - Fc) * F0;
-	float LayerAttenuation = lerp( 1, (1 - F), ClearCoat );
+		// Approximation of absorption integral, tuned for Roughness=0.4
+		float3 AbsorptionColor = BaseColor * (1 / MetalSpec);
+		float3 Absorption = AbsorptionColor * ( (NoV - 1) * 0.85 * ( 1 - lerp( AbsorptionColor, Square(AbsorptionColor), -0.78 ) ) + 1 );
+
+		float F0 = 0.04;
+		float Fc = Pow5( 1 - NoV );
+		float F = Fc + (1 - Fc) * F0;
+		float LayerAttenuation = lerp( 1, (1 - F), ClearCoat );
 
 	GBuffer.BaseColor = lerp( BaseColor * LayerAttenuation, MetalSpec * Absorption * RefractionScale, Metallic * ClearCoat );
-	GBuffer.BaseColor+= dither / 255.f;
+	GBuffer.BaseColor+= Dither / 255.f;
 	GBuffer.Specular *= lerp( 1, RefractionScale, ClearCoat );
-
-	GBuffer.CustomData.x = ClearCoat;
-	GBuffer.CustomData.y = ClearCoatRoughness;
+		GBuffer.CustomData.x = ClearCoat;
+		GBuffer.CustomData.y = ClearCoatRoughness;
 
 	#if CLEAR_COAT_BOTTOM_NORMAL
+		{
+			float2 oct2 = UnitVectorToOctahedron(GBuffer.WorldNormal);
+
+			#if NUM_MATERIAL_OUTPUTS_CLEARCOATBOTTOMNORMAL > 0
+				#if MATERIAL_TANGENTSPACENORMAL
+					float3 tempnormal = normalize(TransformTangentVectorToWorld( MaterialParameters.TangentToWorld, ClearCoatBottomNormal0(MaterialParameters) ));
+				#else
+					float3 tempnormal =  ClearCoatBottomNormal0(MaterialParameters);
+				#endif
+
+				float2 oct1 = UnitVectorToOctahedron(tempnormal);
+				float2 oct3 = ( (oct1 - oct2) *  0.5 ) + (128.0/255.0);
+				GBuffer.CustomData.a = oct3.x;
+				GBuffer.CustomData.z = oct3.y;
+			#else
+				GBuffer.CustomData.a = 128.0/255.0;
+				GBuffer.CustomData.z = 128.0/255.0;
+			#endif
+		}
+	#endif
+	}
+#endif
+#if MATERIAL_SHADINGMODEL_TWOSIDED_FOLIAGE
+	else if (ShadingModel == SHADINGMODELID_TWOSIDED_FOLIAGE)
 	{
-		float2 oct2 = UnitVectorToOctahedron(GBuffer.WorldNormal);
+		GBuffer.CustomData.rgb = EncodeSubsurfaceColor(SubsurfaceColor);
+		GBuffer.CustomData.a = Opacity;
+	}
+#endif
+#if MATERIAL_SHADINGMODEL_HAIR
+	else if (ShadingModel == SHADINGMODELID_HAIR)
+	{
+		GBuffer.CustomData.xy = UnitVectorToOctahedron( MaterialParameters.WorldNormal ) * 0.5 + 0.5;
+		GBuffer.CustomData.z = saturate( GetMaterialCustomData0(MaterialParameters) );	// Backlit
+	}
+#endif
+#if MATERIAL_SHADINGMODEL_CLOTH
+	else if (ShadingModel == SHADINGMODELID_CLOTH)
+	{
+		GBuffer.CustomData.rgb = SubsurfaceColor;
+		GBuffer.CustomData.a = saturate( GetMaterialCustomData0(MaterialParameters) );	// Cloth
+		GBuffer.IndirectIrradiance *= 1 - GBuffer.CustomData.a;
+	}
+#endif
+#if MATERIAL_SHADINGMODEL_EYE
+	else if (ShadingModel == SHADINGMODELID_EYE)
+	{
+		GBuffer.CustomData.x = EncodeSubsurfaceProfile(SubsurfaceProfile).x;
+		GBuffer.CustomData.w = 1.0f - saturate(GetMaterialCustomData0(MaterialParameters));	// Opacity = 1.0 - Iris Mask
+		GBuffer.Metallic = saturate(GetMaterialCustomData1(MaterialParameters));			// Iris Distance
+
+	#if IRIS_NORMAL
+		float IrisMask		= saturate( GetMaterialCustomData0(MaterialParameters) );
+		float IrisDistance	= saturate( GetMaterialCustomData1(MaterialParameters) );
+
+		GBuffer.CustomData.x = EncodeSubsurfaceProfile(SubsurfaceProfile).x;
+		GBuffer.CustomData.w = 1.0 - IrisMask;	// Opacity
+
+		float2 WorldNormalOct = UnitVectorToOctahedron( GBuffer.WorldNormal );
+
+		// CausticNormal stored as octahedron
+		#if NUM_MATERIAL_OUTPUTS_GETTANGENTOUTPUT > 0
+			// Blend in the negative intersection normal to create some concavity
+			// Not great as it ties the concavity to the convexity of the cornea surface
+			// No good justification for that. On the other hand, if we're just looking to
+			// introduce some concavity, this does the job.
+			float3 PlaneNormal = normalize( GetTangentOutput0(MaterialParameters) );
+			float3 CausticNormal = normalize( lerp( PlaneNormal, -GBuffer.WorldNormal, IrisMask*IrisDistance ) );
+			float2 CausticNormalOct  = UnitVectorToOctahedron( CausticNormal );
+			float2 CausticNormalDelta = ( CausticNormalOct - WorldNormalOct ) * 0.5 + (128.0/255.0);
+			GBuffer.Metallic = CausticNormalDelta.x;
+			GBuffer.Specular = CausticNormalDelta.y;
+		#else
+			float3 PlaneNormal = GBuffer.WorldNormal;
+			GBuffer.Metallic = 128.0/255.0;
+			GBuffer.Specular = 128.0/255.0;
+		#endif
 
+		// IrisNormal CustomData.yz
 		#if NUM_MATERIAL_OUTPUTS_CLEARCOATBOTTOMNORMAL > 0
+			float3 IrisNormal = normalize( ClearCoatBottomNormal0(MaterialParameters) );
 			#if MATERIAL_TANGENTSPACENORMAL
-				float3 tempnormal = normalize(TransformTangentVectorToWorld( MaterialParameters.TangentToWorld, ClearCoatBottomNormal0(MaterialParameters) ));
-			#else
-				float3 tempnormal =  ClearCoatBottomNormal0(MaterialParameters);
+				IrisNormal = normalize( TransformTangentVectorToWorld( MaterialParameters.TangentToWorld, IrisNormal ) );
 			#endif
-
-			float2 oct1 = UnitVectorToOctahedron(tempnormal);
-			float2 oct3 = ( (oct1 - oct2) *  0.5 ) + (128.0/255.0);
-			GBuffer.CustomData.a = oct3.x;
-			GBuffer.CustomData.z = oct3.y;
 		#else
-			GBuffer.CustomData.a = 128.0/255.0;
-			GBuffer.CustomData.z = 128.0/255.0;
+			float3 IrisNormal = PlaneNormal;
 		#endif
-	}
-	#endif
-#elif MATERIAL_SHADINGMODEL_TWOSIDED_FOLIAGE
-	GBuffer.ShadingModelID = SHADINGMODELID_TWOSIDED_FOLIAGE;
-	GBuffer.CustomData.rgb = EncodeSubsurfaceColor(SubsurfaceColor);
-	GBuffer.CustomData.a = Opacity;
-#elif MATERIAL_SHADINGMODEL_HAIR
-	GBuffer.ShadingModelID = SHADINGMODELID_HAIR;
-	GBuffer.CustomData.xy = UnitVectorToOctahedron( MaterialParameters.WorldNormal ) * 0.5 + 0.5;
-	GBuffer.CustomData.z = saturate( GetMaterialCustomData0(MaterialParameters) );	// Backlit
-#elif MATERIAL_SHADINGMODEL_CLOTH
-	GBuffer.ShadingModelID = SHADINGMODELID_CLOTH;
-	GBuffer.CustomData.rgb = SubsurfaceColor;
-	GBuffer.CustomData.a = saturate( GetMaterialCustomData0(MaterialParameters) );	// Cloth
-	GBuffer.IndirectIrradiance *= 1 - GBuffer.CustomData.a;
-#elif MATERIAL_SHADINGMODEL_EYE
-	GBuffer.ShadingModelID = SHADINGMODELID_EYE;
-	GBuffer.CustomData.x = EncodeSubsurfaceProfile(SubsurfaceProfile).x;
-	GBuffer.CustomData.w = 1.0f - saturate(GetMaterialCustomData0(MaterialParameters));	// Opacity = 1.0 - Iris Mask
-	GBuffer.Metallic = saturate(GetMaterialCustomData1(MaterialParameters));			// Iris Distance
-
-#if IRIS_NORMAL
-	float IrisMask		= saturate( GetMaterialCustomData0(MaterialParameters) );
-	float IrisDistance	= saturate( GetMaterialCustomData1(MaterialParameters) );
-
-	GBuffer.CustomData.x = EncodeSubsurfaceProfile(SubsurfaceProfile).x;
-	GBuffer.CustomData.w = 1.0 - IrisMask;	// Opacity
-
-	float2 WorldNormalOct = UnitVectorToOctahedron( GBuffer.WorldNormal );
-
-	// CausticNormal stored as octahedron
-	#if NUM_MATERIAL_OUTPUTS_GETTANGENTOUTPUT > 0
-		// Blend in the negative intersection normal to create some concavity
-		// Not great as it ties the concavity to the convexity of the cornea surface
-		// No good justification for that. On the other hand, if we're just looking to
-		// introduce some concavity, this does the job.
-		float3 PlaneNormal = normalize( GetTangentOutput0(MaterialParameters) );
-		float3 CausticNormal = normalize( lerp( PlaneNormal, -GBuffer.WorldNormal, IrisMask*IrisDistance ) );
-		float2 CausticNormalOct  = UnitVectorToOctahedron( CausticNormal );
-		float2 CausticNormalDelta = ( CausticNormalOct - WorldNormalOct ) * 0.5 + (128.0/255.0);
-		GBuffer.Metallic = CausticNormalDelta.x;
-		GBuffer.Specular = CausticNormalDelta.y;
-	#else
-		float3 PlaneNormal = GBuffer.WorldNormal;
-		GBuffer.Metallic = 128.0/255.0;
-		GBuffer.Specular = 128.0/255.0;
-	#endif
 
-	// IrisNormal CustomData.yz
-	#if NUM_MATERIAL_OUTPUTS_CLEARCOATBOTTOMNORMAL > 0
-		float3 IrisNormal = normalize( ClearCoatBottomNormal0(MaterialParameters) );
-		#if MATERIAL_TANGENTSPACENORMAL
-			IrisNormal = normalize( TransformTangentVectorToWorld( MaterialParameters.TangentToWorld, IrisNormal ) );
-		#endif
+		float2 IrisNormalOct  = UnitVectorToOctahedron( IrisNormal );
+		float2 IrisNormalDelta = ( IrisNormalOct - WorldNormalOct ) * 0.5 + (128.0/255.0);
+		GBuffer.CustomData.yz = IrisNormalDelta;
 	#else
-		float3 IrisNormal = PlaneNormal;
-	#endif
-
-	float2 IrisNormalOct  = UnitVectorToOctahedron( IrisNormal );
-	float2 IrisNormalDelta = ( IrisNormalOct - WorldNormalOct ) * 0.5 + (128.0/255.0);
-	GBuffer.CustomData.yz = IrisNormalDelta;
-#else
-	GBuffer.Metallic = saturate(GetMaterialCustomData1(MaterialParameters));			// Iris Distance
+		GBuffer.Metallic = saturate(GetMaterialCustomData1(MaterialParameters));			// Iris Distance
 
-	#if NUM_MATERIAL_OUTPUTS_GETTANGENTOUTPUT > 0
-		float3 Tangent = GetTangentOutput0(MaterialParameters);
-		GBuffer.CustomData.yz = UnitVectorToOctahedron( normalize(Tangent) ) * 0.5 + 0.5;
+		#if NUM_MATERIAL_OUTPUTS_GETTANGENTOUTPUT > 0
+			float3 Tangent = GetTangentOutput0(MaterialParameters);
+			GBuffer.CustomData.yz = UnitVectorToOctahedron( normalize(Tangent) ) * 0.5 + 0.5;
+		#endif
 	#endif
+	}
 #endif
-#else
-	// missing shading model, compiler should report ShadingModelID is not set
-#endif
-}
+}
\ No newline at end of file
diff --git a/Engine/Source/Developer/MaterialBaking/Private/ExportMaterialProxy.h b/Engine/Source/Developer/MaterialBaking/Private/ExportMaterialProxy.h
--- a/Engine/Source/Developer/MaterialBaking/Private/ExportMaterialProxy.h
+++ b/Engine/Source/Developer/MaterialBaking/Private/ExportMaterialProxy.h
@@ -30,7 +30,7 @@ struct FExportMaterialCompiler : public FProxyMaterialCompiler
 		return SF_Pixel;
 	}
 
-	virtual EMaterialShadingModel GetMaterialShadingModel() const override
+	virtual FMaterialShadingModelField GetMaterialShadingModels() const override
 	{
 		// not used by Lightmass
 		return MSM_MAX;
@@ -250,7 +250,7 @@ public:
 			GetDependentShaderAndVFTypes(GMaxRHIShaderPlatform, ShaderTypes, ShaderPipelineTypes, VFTypes);
 
 			// Overwrite the shader map Id's dependencies with ones that came from the FMaterial actually being compiled (this)
-			// This is necessary as we change FMaterial attributes like GetShadingModel(), which factor into the ShouldCache functions that determine dependent shader types
+			// This is necessary as we change FMaterial attributes like GetShadingModels(), which factor into the ShouldCache functions that determine dependent shader types
 			ResourceId.SetShaderDependencies(ShaderTypes, ShaderPipelineTypes, VFTypes, GMaxRHIShaderPlatform);
 		}
 
@@ -347,7 +347,6 @@ public:
 		if (Property == MP_EmissiveColor)
 		{
 			const EBlendMode BlendMode = MaterialInterface->GetBlendMode();
-			const EMaterialShadingModel ShadingModel = MaterialInterface->GetShadingModel();
 			FExportMaterialCompiler ProxyCompiler(Compiler);
 			const uint32 ForceCast_Exact_Replicate = MFCF_ForceCast | MFCF_ExactMatch | MFCF_ReplicateValue;
 
@@ -468,7 +467,7 @@ public:
 	}
 	virtual bool IsMasked() const override { return false; }
 	virtual enum EBlendMode GetBlendMode() const override { return BLEND_Opaque; }
-	virtual enum EMaterialShadingModel GetShadingModel() const override { return MSM_DefaultLit; }
+	virtual FMaterialShadingModelField GetShadingModels() const override { return MSM_DefaultLit; }
 	virtual float GetOpacityMaskClipValue() const override { return 0.5f; }
 	virtual bool GetCastDynamicShadowAsMasked() const override { return false; }
 	virtual FString GetFriendlyName() const override { return FString::Printf(TEXT("FExportMaterialRenderer %s"), MaterialInterface ? *MaterialInterface->GetName() : TEXT("NULL")); }
diff --git a/Engine/Source/Developer/MaterialUtilities/Private/MaterialUtilities.cpp b/Engine/Source/Developer/MaterialUtilities/Private/MaterialUtilities.cpp
--- a/Engine/Source/Developer/MaterialUtilities/Private/MaterialUtilities.cpp
+++ b/Engine/Source/Developer/MaterialUtilities/Private/MaterialUtilities.cpp
@@ -206,7 +206,7 @@ struct FExportMaterialCompiler : public FProxyMaterialCompiler
 		return SF_Pixel;
 	}
 
-	virtual EMaterialShadingModel GetMaterialShadingModel() const override
+	virtual FMaterialShadingModelField GetMaterialShadingModels() const override
 	{
 		// not used by Lightmass
 		return MSM_MAX;
@@ -428,7 +428,7 @@ public:
 			GetDependentShaderAndVFTypes(GMaxRHIShaderPlatform, ShaderTypes, ShaderPipelineTypes, VFTypes);
 
 			// Overwrite the shader map Id's dependencies with ones that came from the FMaterial actually being compiled (this)
-			// This is necessary as we change FMaterial attributes like GetShadingModel(), which factor into the ShouldCache functions that determine dependent shader types
+			// This is necessary as we change FMaterial attributes like GetShadingModels(), which factor into the ShouldCache functions that determine dependent shader types
 			ResourceId.SetShaderDependencies(ShaderTypes, ShaderPipelineTypes, VFTypes, GMaxRHIShaderPlatform);
 		}
 
@@ -529,7 +529,6 @@ public:
 			UMaterial* ProxyMaterial = MaterialInterface->GetMaterial();
 			check(ProxyMaterial);
 			EBlendMode BlendMode = MaterialInterface->GetBlendMode();
-			EMaterialShadingModel ShadingModel = MaterialInterface->GetShadingModel();
 			FExportMaterialCompiler ProxyCompiler(Compiler);
 			const uint32 ForceCast_Exact_Replicate = MFCF_ForceCast | MFCF_ExactMatch | MFCF_ReplicateValue;
 									
@@ -643,11 +642,11 @@ public:
 		}
 		return false;
 	}
-	virtual bool IsMasked() const override								{ return false; }
-	virtual enum EBlendMode GetBlendMode() const override				{ return BLEND_Opaque; }
-	virtual enum EMaterialShadingModel GetShadingModel() const override	{ return MSM_Unlit; }
-	virtual float GetOpacityMaskClipValue() const override				{ return 0.5f; }
-	virtual bool GetCastDynamicShadowAsMasked() const override					{ return false; }
+	virtual bool IsMasked() const override									{ return false; }
+	virtual enum EBlendMode GetBlendMode() const override					{ return BLEND_Opaque; }
+	virtual FMaterialShadingModelField GetShadingModels() const override	{ return MSM_Unlit; }
+	virtual float GetOpacityMaskClipValue() const override					{ return 0.5f; }
+	virtual bool GetCastDynamicShadowAsMasked() const override				{ return false; }
 	virtual FString GetFriendlyName() const override { return FString::Printf(TEXT("FExportMaterialRenderer %s"), MaterialInterface ? *MaterialInterface->GetName() : TEXT("NULL")); }
 	/**
 	* Should shaders compiled for this material be saved to disk?
diff --git a/Engine/Source/Editor/DetailCustomizations/Private/DetailCustomizations.cpp b/Engine/Source/Editor/DetailCustomizations/Private/DetailCustomizations.cpp
--- a/Engine/Source/Editor/DetailCustomizations/Private/DetailCustomizations.cpp
+++ b/Engine/Source/Editor/DetailCustomizations/Private/DetailCustomizations.cpp
@@ -135,7 +135,7 @@
 #include "MaterialExpressionLandscapeGrassCustomization.h"
 #include "TimecodeDetailsCustomization.h"
 #include "SkeletonDetails.h"
-
+#include "MaterialShadingModelCustomization.h"
 IMPLEMENT_MODULE( FDetailCustomizationsModule, DetailCustomizations );
 
 
@@ -268,6 +268,7 @@ void FDetailCustomizationsModule::RegisterPropertyTypeCustomizations()
 	RegisterCustomPropertyTypeLayout("PerPlatformBool", FOnGetPropertyTypeCustomizationInstance::CreateStatic(&FPerPlatformPropertyCustomization<FPerPlatformBool>::MakeInstance));
 	RegisterCustomPropertyTypeLayout("SkeletalMeshOptimizationSettings", FOnGetPropertyTypeCustomizationInstance::CreateStatic(&FSkeletalMeshReductionSettingsDetails::MakeInstance));
 	RegisterCustomPropertyTypeLayout("GrassInput", FOnGetPropertyTypeCustomizationInstance::CreateStatic(&FMaterialExpressionLandscapeGrassInputCustomization::MakeInstance));
+	RegisterCustomPropertyTypeLayout("EMaterialShadingModel", FOnGetPropertyTypeCustomizationInstance::CreateStatic(&FMaterialShadingModelCustomization::MakeInstance));
 }
 
 void FDetailCustomizationsModule::RegisterObjectCustomizations()
diff --git a/Engine/Source/Editor/DetailCustomizations/Private/MaterialShadingModelCustomization.cpp b/Engine/Source/Editor/DetailCustomizations/Private/MaterialShadingModelCustomization.cpp
new file mode 100644
--- /dev/null
+++ b/Engine/Source/Editor/DetailCustomizations/Private/MaterialShadingModelCustomization.cpp
@@ -0,0 +1,56 @@
+// Copyright 1998-2019 Epic Games, Inc. All Rights Reserved.
+
+#include "MaterialShadingModelCustomization.h"
+#include "IPropertyUtilities.h"
+#include "PropertyRestriction.h"
+#include "Engine/EngineTypes.h"
+#include "IDetailChildrenBuilder.h"
+#include "DetailLayoutBuilder.h"
+#include "DetailCategoryBuilder.h"
+#include "DetailWidgetRow.h"
+#include "IDetailPropertyRow.h"
+#include "PropertyCustomizationHelpers.h"
+
+TSharedRef<IPropertyTypeCustomization> FMaterialShadingModelCustomization::MakeInstance() 
+{
+	return MakeShareable( new FMaterialShadingModelCustomization );
+}
+
+void FMaterialShadingModelCustomization::CustomizeHeader(TSharedRef<IPropertyHandle> PropertyHandle, class FDetailWidgetRow& HeaderRow, IPropertyTypeCustomizationUtils& CustomizationUtils)
+{
+	// Hide special value by default
+	bool bShouldHideFromMaterialExpression = true;
+	
+	// This piece of code enables the use of the Shading Model material output pin through selecting "From Material Expression" in the Shading Model drop down menu.
+	if (PropertyHandle->GetNumOuterObjects() == 1)
+	{
+		TArray<UObject*> OuterObjects;
+		PropertyHandle->GetOuterObjects(OuterObjects);
+		
+		// Don't hide MSM_FromMaterialExpression on a UMaterial
+		if (OuterObjects[0]->IsA<UMaterial>())
+		{
+			bShouldHideFromMaterialExpression = false;
+		}
+	}
+	
+	// Add restrictin to property if asked for
+	if (bShouldHideFromMaterialExpression)
+	{
+		TSharedPtr<FPropertyRestriction> EnumRestriction = MakeShareable(new FPropertyRestriction(NSLOCTEXT("MaterialShadingModel", "FromMaterialExpression", "FromMaterialExpression is only available on UMaterial")));
+		const UEnum* const MaterialShadingModelEnum = StaticEnum<EMaterialShadingModel>();		
+		EnumRestriction->AddHiddenValue(MaterialShadingModelEnum->GetNameStringByValue((int64)EMaterialShadingModel::MSM_FromMaterialExpression));
+		PropertyHandle->AddRestriction(EnumRestriction.ToSharedRef());
+	}
+
+	// Implement the copmbobox for the enum
+	HeaderRow
+	.NameContent()
+	[
+		PropertyHandle->CreatePropertyNameWidget()
+	]
+	.ValueContent()
+	[
+		PropertyCustomizationHelpers::MakePropertyComboBox(PropertyHandle)
+	];
+}
diff --git a/Engine/Source/Editor/DetailCustomizations/Private/MaterialShadingModelCustomization.h b/Engine/Source/Editor/DetailCustomizations/Private/MaterialShadingModelCustomization.h
new file mode 100644
--- /dev/null
+++ b/Engine/Source/Editor/DetailCustomizations/Private/MaterialShadingModelCustomization.h
@@ -0,0 +1,20 @@
+// Copyright 1998-2019 Epic Games, Inc. All Rights Reserved.
+
+#pragma once
+
+#include "CoreMinimal.h"
+#include "Misc/Attribute.h"
+#include "Layout/Visibility.h"
+#include "IPropertyTypeCustomization.h"
+#include "PropertyHandle.h"
+
+class FMaterialShadingModelCustomization : public IPropertyTypeCustomization
+{
+public:
+	static TSharedRef<IPropertyTypeCustomization> MakeInstance();
+
+	//~ Begin IPropertyTypeCustomization
+	virtual void CustomizeHeader(TSharedRef<IPropertyHandle> PropertyHandle, class FDetailWidgetRow& HeaderRow, IPropertyTypeCustomizationUtils& CustomizationUtils) override;
+	virtual void CustomizeChildren(TSharedRef<IPropertyHandle> StructPropertyHandle, class IDetailChildrenBuilder& ChildBuilder, IPropertyTypeCustomizationUtils& StructCustomizationUtils) override {}
+	//~ End IPropertyTypeCustomization
+};
\ No newline at end of file
diff --git a/Engine/Source/Editor/MaterialEditor/Private/MaterialEditor.cpp b/Engine/Source/Editor/MaterialEditor/Private/MaterialEditor.cpp
--- a/Engine/Source/Editor/MaterialEditor/Private/MaterialEditor.cpp
+++ b/Engine/Source/Editor/MaterialEditor/Private/MaterialEditor.cpp
@@ -207,6 +207,11 @@ int32 FMatExpressionPreview::CompilePropertyAndSetMaterialProperty(EMaterialProp
 		const int32 TextureCoordinateIndex = Property - MP_CustomizedUVs0;
 		Ret = Compiler->TextureCoordinate(TextureCoordinateIndex, false, false);
 	}
+	else if (Property == MP_ShadingModel)
+	{
+		FMaterialShadingModelField ShadingModels = Compiler->GetMaterialShadingModels();
+		Ret = Compiler->ShadingModel(ShadingModels.GetFirstShadingModel());
+	}
 	else
 	{
 		Ret = Compiler->Constant(1.0f);
@@ -2908,6 +2913,7 @@ UClass* FMaterialEditor::GetOnPromoteToParameterClass(UEdGraphPin* TargetPin)
 			case MP_AmbientOcclusion:
 			case MP_Refraction:
 			case MP_PixelDepthOffset:
+			case MP_ShadingModel:
 			case MP_OpacityMask: return UMaterialExpressionScalarParameter::StaticClass();
 
 			case MP_WorldPositionOffset:
@@ -4103,7 +4109,7 @@ void FMaterialEditor::NotifyPostChange( const FPropertyChangedEvent& PropertyCha
 			}
 		}
 		else if( NameOfPropertyThatChanged == GET_MEMBER_NAME_CHECKED(UMaterial, MaterialDomain) ||
-				 NameOfPropertyThatChanged == GET_MEMBER_NAME_CHECKED(UMaterial, ShadingModel) )
+				 NameOfPropertyThatChanged == GET_MEMBER_NAME_CHECKED(UMaterial, ShadingModel))
 		{
 			Material->MaterialGraph->RebuildGraph();
 			TArray<TWeakObjectPtr<UObject>> SelectedObjects = MaterialDetailsView->GetSelectedObjects();
diff --git a/Engine/Source/Editor/MaterialEditor/Private/MaterialEditor.h b/Engine/Source/Editor/MaterialEditor/Private/MaterialEditor.h
--- a/Engine/Source/Editor/MaterialEditor/Private/MaterialEditor.h
+++ b/Engine/Source/Editor/MaterialEditor/Private/MaterialEditor.h
@@ -151,7 +151,7 @@ public:
 	virtual bool IsWireframe() const override { return false; }
 	virtual bool IsMasked() const override { return false; }
 	virtual enum EBlendMode GetBlendMode() const override { return BLEND_Opaque; }
-	virtual enum EMaterialShadingModel GetShadingModel() const override { return MSM_Unlit; }
+	virtual FMaterialShadingModelField GetShadingModels() const override { return MSM_Unlit; }
 	virtual float GetOpacityMaskClipValue() const override { return 0.5f; }
 	virtual bool GetCastDynamicShadowAsMasked() const override { return false; }
 	virtual FString GetFriendlyName() const override { return FString::Printf(TEXT("FMatExpressionPreview %s"), Expression.IsValid() ? *Expression->GetName() : TEXT("NULL")); }
diff --git a/Engine/Source/Editor/MaterialEditor/Private/MaterialEditorInstanceDetailCustomization.cpp b/Engine/Source/Editor/MaterialEditor/Private/MaterialEditorInstanceDetailCustomization.cpp
--- a/Engine/Source/Editor/MaterialEditor/Private/MaterialEditorInstanceDetailCustomization.cpp
+++ b/Engine/Source/Editor/MaterialEditor/Private/MaterialEditorInstanceDetailCustomization.cpp
@@ -737,9 +737,9 @@ EVisibility FMaterialInstanceParameterDetails::ShouldShowMaterialRefractionSetti
 
 EVisibility FMaterialInstanceParameterDetails::ShouldShowSubsurfaceProfile() const
 {
-	EMaterialShadingModel Model = MaterialEditorInstance->SourceInstance->GetShadingModel();
+	FMaterialShadingModelField ShadingModels = MaterialEditorInstance->SourceInstance->GetShadingModels();
 
-	return UseSubsurfaceProfile(Model) ? EVisibility::Visible : EVisibility::Collapsed;
+	return UseSubsurfaceProfile(ShadingModels) ? EVisibility::Visible : EVisibility::Collapsed;
 }
 
 
@@ -908,12 +908,33 @@ void FMaterialInstanceParameterDetails::CreateBasePropertyOverrideWidgets(IDetai
 		.OverrideResetToDefault(ResetBlendModePropertyOverride);
 
 	FIsResetToDefaultVisible IsShadingModelPropertyResetVisible = FIsResetToDefaultVisible::CreateLambda([this](TSharedPtr<IPropertyHandle> InHandle) {
-		return MaterialEditorInstance->Parent != nullptr ? MaterialEditorInstance->BasePropertyOverrides.ShadingModel != MaterialEditorInstance->Parent->GetShadingModel() : false;
+		if (MaterialEditorInstance->Parent != nullptr)
+		{	
+			if (MaterialEditorInstance->Parent->IsShadingModelFromMaterialExpression())
+			{
+				return MaterialEditorInstance->BasePropertyOverrides.ShadingModel != MSM_FromMaterialExpression;
+			}
+			else
+			{
+				return MaterialEditorInstance->BasePropertyOverrides.ShadingModel != MaterialEditorInstance->Parent->GetShadingModels().GetFirstShadingModel();
+			}
+		}
+		else
+		{
+			return false;
+		}
 	});
 	FResetToDefaultHandler ResetShadingModelPropertyHandler = FResetToDefaultHandler::CreateLambda([this](TSharedPtr<IPropertyHandle> InHandle) {
 		if (MaterialEditorInstance->Parent != nullptr)
 		{
-			MaterialEditorInstance->BasePropertyOverrides.ShadingModel = MaterialEditorInstance->Parent->GetShadingModel();
+			if (MaterialEditorInstance->Parent->IsShadingModelFromMaterialExpression())
+			{
+				MaterialEditorInstance->BasePropertyOverrides.ShadingModel = MSM_FromMaterialExpression;
+			}
+			else
+			{
+				MaterialEditorInstance->BasePropertyOverrides.ShadingModel = MaterialEditorInstance->Parent->GetShadingModels().GetFirstShadingModel();
+			}
 		}
 	});
 	FResetToDefaultOverride ResetShadingModelPropertyOverride = FResetToDefaultOverride::Create(IsShadingModelPropertyResetVisible, ResetShadingModelPropertyHandler);
diff --git a/Engine/Source/Editor/MaterialEditor/Private/MaterialStatsCommon.cpp b/Engine/Source/Editor/MaterialEditor/Private/MaterialStatsCommon.cpp
--- a/Engine/Source/Editor/MaterialEditor/Private/MaterialStatsCommon.cpp
+++ b/Engine/Source/Editor/MaterialEditor/Private/MaterialStatsCommon.cpp
@@ -379,7 +379,7 @@ void FMaterialStatsUtils::GetRepresentativeShaderTypesAndDescriptions(TMap<FName
 	}
 	else if (TargetMaterial->GetFeatureLevel() >= ERHIFeatureLevel::SM4)
 	{
-		if (TargetMaterial->GetShadingModel() == MSM_Unlit)
+		if (TargetMaterial->GetShadingModels().IsUnlit())
 		{
 			//unlit materials are never lightmapped
 			static FName TBasePassPSFNoLightMapPolicyName = TEXT("TBasePassPSFNoLightMapPolicy");
@@ -422,7 +422,7 @@ void FMaterialStatsUtils::GetRepresentativeShaderTypesAndDescriptions(TMap<FName
 	{
 		const TCHAR* DescSuffix = bMobileHDR ? TEXT(" (HDR)") : TEXT(" (LDR)");
 
-		if (TargetMaterial->GetShadingModel() == MSM_Unlit)
+		if (TargetMaterial->GetShadingModels().IsUnlit())
 		{
 			//unlit materials are never lightmapped
 			static FName Name_HDRLinear64 = TEXT("TMobileBasePassPSFNoLightMapPolicy0HDRLinear64");
diff --git a/Engine/Source/Editor/UnrealEd/Private/Fbx/FbxMainExport.cpp b/Engine/Source/Editor/UnrealEd/Private/Fbx/FbxMainExport.cpp
--- a/Engine/Source/Editor/UnrealEd/Private/Fbx/FbxMainExport.cpp
+++ b/Engine/Source/Editor/UnrealEd/Private/Fbx/FbxMainExport.cpp
@@ -1342,7 +1342,7 @@ FbxSurfaceMaterial* FFbxExporter::ExportMaterial(UMaterialInterface* MaterialInt
 	UMaterial *Material = MaterialInterface->GetMaterial();
 	
 	// Set the shading model
-	if (Material->GetShadingModel() == MSM_DefaultLit)
+	if (Material->GetShadingModels().HasOnlyShadingModel(MSM_DefaultLit))
 	{
 		FbxMaterial = FbxSurfacePhong::Create(Scene, TCHAR_TO_UTF8(*MaterialInterface->GetName()));
 		//((FbxSurfacePhong*)FbxMaterial)->Specular.Set(Material->Specular));
diff --git a/Engine/Source/Editor/UnrealEd/Private/Lightmass/LightmassRender.cpp b/Engine/Source/Editor/UnrealEd/Private/Lightmass/LightmassRender.cpp
--- a/Engine/Source/Editor/UnrealEd/Private/Lightmass/LightmassRender.cpp
+++ b/Engine/Source/Editor/UnrealEd/Private/Lightmass/LightmassRender.cpp
@@ -41,7 +41,7 @@ struct FLightmassMaterialCompiler : public FProxyMaterialCompiler
 		return SF_Pixel;
 	}
 
-	virtual EMaterialShadingModel GetMaterialShadingModel() const override
+	virtual FMaterialShadingModelField GetMaterialShadingModels() const override
 	{ 
 		// not used by Lightmass
 		return MSM_MAX;
@@ -255,7 +255,7 @@ public:
 					GetDependentShaderAndVFTypes(GMaxRHIShaderPlatform, ShaderTypes, ShaderPipelineTypes, VFTypes);
 
 					// Overwrite the shader map Id's dependencies with ones that came from the FMaterial actually being compiled (this)
-					// This is necessary as we change FMaterial attributes like GetShadingModel(), which factor into the ShouldCache functions that determine dependent shader types
+					// This is necessary as we change FMaterial attributes like GetShadingModels(), which factor into the ShouldCache functions that determine dependent shader types
 					ResourceId.SetShaderDependencies(ShaderTypes, ShaderPipelineTypes, VFTypes, GMaxRHIShaderPlatform);
 				}
 
@@ -367,7 +367,7 @@ public:
 		{
 			UMaterial* ProxyMaterial = MaterialInterface->GetMaterial();
 			EBlendMode BlendMode = MaterialInterface->GetBlendMode();
-			EMaterialShadingModel ShadingModel = MaterialInterface->GetShadingModel();
+			bool bIsMaterialUnlit = MaterialInterface->GetShadingModels().IsUnlit();
 			check(ProxyMaterial);
 			FLightmassMaterialCompiler ProxyCompiler(Compiler);
 
@@ -413,7 +413,7 @@ public:
 				}
 				else if (BlendMode == BLEND_Modulate)
 				{
-					if (ShadingModel == MSM_Unlit)
+					if (bIsMaterialUnlit)
 					{
 						return MaterialInterface->CompileProperty(Compiler, MP_EmissiveColor, ForceCast_Exact_Replicate);
 					}
@@ -425,7 +425,7 @@ public:
 				else if ((BlendMode == BLEND_Translucent) || (BlendMode == BLEND_Additive || (BlendMode == BLEND_AlphaComposite)))
 				{
 					int32 ColoredOpacity = INDEX_NONE;
-					if (ShadingModel == MSM_Unlit)
+					if (bIsMaterialUnlit)
 					{
 						ColoredOpacity = MaterialInterface->CompileProperty(Compiler, MP_EmissiveColor, ForceCast_Exact_Replicate);
 					}
@@ -523,11 +523,11 @@ public:
 		}
 		return false;
 	}
-	virtual bool IsMasked() const override								{ return false; }
-	virtual enum EBlendMode GetBlendMode() const override				{ return BLEND_Opaque; }
-	virtual enum EMaterialShadingModel GetShadingModel() const override	{ return MSM_Unlit; }
-	virtual float GetOpacityMaskClipValue() const override				{ return 0.5f; }
-	virtual bool GetCastDynamicShadowAsMasked() const override					{ return false; }
+	virtual bool IsMasked() const override									{ return false; }
+	virtual enum EBlendMode GetBlendMode() const override					{ return BLEND_Opaque; }
+	virtual FMaterialShadingModelField GetShadingModels() const override	{ return MSM_Unlit; }
+	virtual float GetOpacityMaskClipValue() const override					{ return 0.5f; }
+	virtual bool GetCastDynamicShadowAsMasked() const override				{ return false; }
 	virtual FString GetFriendlyName() const override { return FString::Printf(TEXT("FLightmassMaterialRenderer %s"), MaterialInterface ? *MaterialInterface->GetName() : TEXT("NULL")); }
 
 	/**
@@ -605,7 +605,7 @@ public:
 		OutUniformValue.A = 0.0f;
 
 		EBlendMode BlendMode = MaterialInterface->GetBlendMode();
-		EMaterialShadingModel ShadingModel = MaterialInterface->GetShadingModel();
+		bool bIsMaterialUnlit = MaterialInterface->GetShadingModels().IsUnlit();
 		
 		check(Material);
 		bool bExpressionIsNULL = false;
@@ -654,7 +654,7 @@ public:
 				(BlendMode == BLEND_AlphaComposite))
 			{
 				bool bColorInputIsNULL = false;
-				if (ShadingModel == MSM_Unlit)
+				if (bIsMaterialUnlit)
 				{
 					bColorInputIsNULL = !IsMaterialInputConnected(Material, MP_EmissiveColor);
 				}
@@ -924,15 +924,15 @@ bool FLightmassMaterialRenderer::GenerateMaterialData(
 	check(BaseMaterial);
 
 	EBlendMode BlendMode = InMaterial.GetBlendMode();
-	EMaterialShadingModel ShadingModel = InMaterial.GetShadingModel();
- 	if ((ShadingModel != MSM_DefaultLit) &&
-		(ShadingModel != MSM_Unlit) &&
-		(ShadingModel != MSM_Subsurface) &&
-		(ShadingModel != MSM_PreintegratedSkin) &&
-		(ShadingModel != MSM_SubsurfaceProfile))
+	FMaterialShadingModelField ShadingModels = InMaterial.GetShadingModels();
+ 	if (!ShadingModels.HasShadingModel(MSM_DefaultLit) &&
+		!ShadingModels.HasShadingModel(MSM_Unlit) &&
+		!ShadingModels.HasShadingModel(MSM_Subsurface) &&
+		!ShadingModels.HasShadingModel(MSM_PreintegratedSkin) &&
+		!ShadingModels.HasShadingModel(MSM_SubsurfaceProfile))
 	{
 		UE_LOG(LogLightmassRender, Warning, TEXT("LIGHTMASS: Material has an unsupported shading model: %d on %s"), 
-			(int32)ShadingModel,
+			(int32)ShadingModels.GetShadingModelField(),
 			*(InMaterial.GetPathName()));
 	}
 
diff --git a/Engine/Source/Editor/UnrealEd/Private/MaterialGraph.cpp b/Engine/Source/Editor/UnrealEd/Private/MaterialGraph.cpp
--- a/Engine/Source/Editor/UnrealEd/Private/MaterialGraph.cpp
+++ b/Engine/Source/Editor/UnrealEd/Private/MaterialGraph.cpp
@@ -32,6 +32,9 @@ void UMaterialGraph::RebuildGraph()
 
 	if (!MaterialFunction)
 	{
+		// This needs to be done before building the new material inputs to guarantee that the shading model field is up to date
+		Material->RebuildShadingModelField();
+
 		// Initialize the material input list.
 		MaterialInputs.Add( FMaterialInputInfo( GetBaseColorPinName(), MP_BaseColor, LOCTEXT( "BaseColorToolTip", "Defines the overall color of the Material. Each channel is automatically clamped between 0 and 1" ) ) );
 		MaterialInputs.Add( FMaterialInputInfo( GetMetallicPinName(), MP_Metallic, LOCTEXT( "MetallicToolTip", "Controls how \"metal-like\" your surface looks like") ) );
@@ -56,7 +59,8 @@ void UMaterialGraph::RebuildGraph()
 			MaterialInputs.Add( FMaterialInputInfo( FText::FromString(FString::Printf(TEXT("Customized UV%u"), UVIndex)), (EMaterialProperty)(MP_CustomizedUVs0 + UVIndex), FText::FromString(FString::Printf( TEXT( "CustomizedUV%uToolTip" ), UVIndex ) ) ) );
 		}
 
-		MaterialInputs.Add(FMaterialInputInfo(LOCTEXT("PixelDepthOffset", "Pixel Depth Offset"), MP_PixelDepthOffset, LOCTEXT( "PixelDepthOffsetToolTip", "Pixel Depth Offset" ) ));
+		MaterialInputs.Add(FMaterialInputInfo(LOCTEXT("PixelDepthOffset", "Pixel Depth Offset"), MP_PixelDepthOffset, LOCTEXT("PixelDepthOffsetToolTip", "Pixel Depth Offset")));
+		MaterialInputs.Add(FMaterialInputInfo(LOCTEXT("ShadingModel", "Shading Model"), MP_ShadingModel, LOCTEXT("ShadingModelToolTip", "Selects which shading model should be used per pixel")));
 
 		//^^^ New material properties go above here. ^^^^
 		MaterialInputs.Add(FMaterialInputInfo(LOCTEXT("MaterialAttributes", "Material Attributes"), MP_MaterialAttributes, LOCTEXT( "MaterialAttributesToolTip", "Material Attributes" ) ));
@@ -485,6 +489,42 @@ int32 UMaterialGraph::GetValidOutputIndex(FExpressionInput* Input) const
 	return OutputIndex;
 }
 
+/** Builds a pin name based on the custom list of <EMaterialShadingModel, FString> pairs and the default pin name */
+static FString GetPinNameFromShadingModelField(FMaterialShadingModelField InShadingModels, const TArray<TKeyValuePair<EMaterialShadingModel, FString>>& InCustomShadingModelPinNames, const FString& InDefaultPinName)
+{
+	FString OutPinName;
+	for (const TKeyValuePair<EMaterialShadingModel, FString>& CustomShadingModelPinName : InCustomShadingModelPinNames)
+	{
+		if (InShadingModels.HasShadingModel(CustomShadingModelPinName.Key))
+		{
+			// Add delimiter
+			if (!OutPinName.IsEmpty())
+			{
+				OutPinName.Append(" or ");
+			}
+
+			// Append the name and remove the shading model from the temp field
+			OutPinName.Append(CustomShadingModelPinName.Value);
+			InShadingModels.RemoveShadingModel(CustomShadingModelPinName.Key);
+		}
+	}
+
+	// There are other shading models present, these don't have their own specific name for this pin, so use a default one
+	if (InShadingModels.CountShadingModels() != 0)
+	{
+		// Add delimiter
+		if (!OutPinName.IsEmpty())
+		{
+			OutPinName.Append(" or ");
+		}
+
+		OutPinName.Append(InDefaultPinName);
+	}
+
+	ensure(!OutPinName.IsEmpty());
+	return OutPinName;
+}
+
 FText UMaterialGraph::GetEmissivePinName() const
 {
 	return Material->IsUIMaterial() ? LOCTEXT("UIOutputColor", "Final Color") : LOCTEXT("EmissiveColor", "Emissive Color");
@@ -502,12 +542,14 @@ FText UMaterialGraph::GetOpacityPinName() const
 
 FText UMaterialGraph::GetMetallicPinName() const
 {
-	return Material->GetShadingModel() == MSM_Hair ? LOCTEXT("Scatter", "Scatter") : LOCTEXT("Metallic", "Metallic");
+	TArray<TKeyValuePair<EMaterialShadingModel, FString>> CustomPinNames({{MSM_Hair, "Scatter"}});
+	return FText::FromString(GetPinNameFromShadingModelField(Material->GetShadingModels(), CustomPinNames, "Metallic"));
 }
 
 FText UMaterialGraph::GetNormalPinName() const
 {
-	return Material->GetShadingModel() == MSM_Hair ? LOCTEXT("Tangent", "Tangent") : LOCTEXT("Normal", "Normal");
+	TArray<TKeyValuePair<EMaterialShadingModel, FString>> CustomPinNames({{MSM_Hair, "Tangent"}});
+	return FText::FromString(GetPinNameFromShadingModelField(Material->GetShadingModels(), CustomPinNames, "Normal"));
 }
 
 FText UMaterialGraph::GetWorldPositionOffsetPinName() const
@@ -517,47 +559,32 @@ FText UMaterialGraph::GetWorldPositionOffsetPinName() const
 
 FText UMaterialGraph::GetSubsurfacePinName() const
 {
-	switch (Material->GetShadingModel())
-	{
-	case MSM_Cloth:
-		return LOCTEXT("FuzzColor", "Fuzz Color");
-	default:
-		return LOCTEXT("SubsurfaceColor", "Subsurface Color");
-	}
+	TArray<TKeyValuePair<EMaterialShadingModel, FString>> CustomPinNames({{MSM_Cloth, "Fuzz Color"}});
+	return FText::FromString(GetPinNameFromShadingModelField(Material->GetShadingModels(), CustomPinNames, "Subsurface Color"));
 }
 
 FText UMaterialGraph::GetCustomDataPinName( uint32 Index ) const
-{
+{ 
 	if( Index == 0 )
 	{
-		switch( Material->GetShadingModel() )
-		{
-		case MSM_ClearCoat:
-			return LOCTEXT("ClearCoat", "Clear Coat");
-		case MSM_Hair:
-			return LOCTEXT("Backlit", "Backlit");
-		case MSM_Cloth:
-			return LOCTEXT("Cloth", "Cloth");
-		case MSM_Eye:
-			return LOCTEXT("IrisMask", "Iris Mask");
-		default:
-			return LOCTEXT("CustomData0", "Custom Data 0");
-		}
+		TArray<TKeyValuePair<EMaterialShadingModel, FString>> CustomPinNames(
+			{{MSM_ClearCoat, "Clear Coat"},
+			{MSM_Hair, "Backlit"},
+			{MSM_Cloth, "Cloth"},
+			{MSM_Eye, "Iris Mask"}});
+
+		return FText::FromString(GetPinNameFromShadingModelField(Material->GetShadingModels(), CustomPinNames, "Custom Data 0"));
 	}
 	else if( Index == 1 )
 	{
-		switch( Material->GetShadingModel() )
-		{
-		case MSM_ClearCoat:
-			return LOCTEXT("ClearCoatRoughness", "Clear Coat Roughness");
-		case MSM_Eye:
-			return LOCTEXT("IrisDistance", "Iris Distance");
-		default:
-			return LOCTEXT("CustomData1", "Custom Data 1");
-		}
+		TArray<TKeyValuePair<EMaterialShadingModel, FString>> CustomPinNames(
+			{{MSM_ClearCoat, "Clear Coat Roughness"},
+			{MSM_Eye, "Iris Distance"}});
+
+		return FText::FromString(GetPinNameFromShadingModelField(Material->GetShadingModels(), CustomPinNames, "Custom Data 1"));
 	}
 
-	return LOCTEXT("CustomData", "Custom Data");	
+	return LOCTEXT("CustomData", "Custom Data");
 }
 
 #undef LOCTEXT_NAMESPACE
diff --git a/Engine/Source/Editor/UnrealEd/Private/PreviewMaterial.cpp b/Engine/Source/Editor/UnrealEd/Private/PreviewMaterial.cpp
--- a/Engine/Source/Editor/UnrealEd/Private/PreviewMaterial.cpp
+++ b/Engine/Source/Editor/UnrealEd/Private/PreviewMaterial.cpp
@@ -1524,7 +1524,14 @@ void UMaterialEditorInstanceConstant::CopyBasePropertiesFromParent()
 	}
 	if (!BasePropertyOverrides.bOverride_ShadingModel)
 	{
-		BasePropertyOverrides.ShadingModel = SourceInstance->GetShadingModel();
+		if (SourceInstance->IsShadingModelFromMaterialExpression())
+		{
+			BasePropertyOverrides.ShadingModel = MSM_FromMaterialExpression;
+		}
+		else
+		{
+			BasePropertyOverrides.ShadingModel = SourceInstance->GetShadingModels().GetFirstShadingModel(); 
+		}
 	}
 	if (!BasePropertyOverrides.bOverride_TwoSided)
 	{
diff --git a/Engine/Source/Runtime/Engine/Classes/Engine/EngineTypes.h b/Engine/Source/Runtime/Engine/Classes/Engine/EngineTypes.h
--- a/Engine/Source/Runtime/Engine/Classes/Engine/EngineTypes.h
+++ b/Engine/Source/Runtime/Engine/Classes/Engine/EngineTypes.h
@@ -460,17 +460,68 @@ namespace EParticleCollisionMode
 UENUM()
 enum EMaterialShadingModel
 {
-	MSM_Unlit				UMETA(DisplayName="Unlit"),
-	MSM_DefaultLit			UMETA(DisplayName="Default Lit"),
-	MSM_Subsurface			UMETA(DisplayName="Subsurface"),
-	MSM_PreintegratedSkin	UMETA(DisplayName="Preintegrated Skin"),
-	MSM_ClearCoat			UMETA(DisplayName="Clear Coat"),
-	MSM_SubsurfaceProfile	UMETA(DisplayName="Subsurface Profile"),
-	MSM_TwoSidedFoliage		UMETA(DisplayName="Two Sided Foliage"),
-	MSM_Hair				UMETA(DisplayName="Hair"),
-	MSM_Cloth				UMETA(DisplayName="Cloth"),
-	MSM_Eye					UMETA(DisplayName="Eye"),
-	MSM_MAX,
+	MSM_Unlit					UMETA(DisplayName="Unlit"),
+	MSM_DefaultLit				UMETA(DisplayName="Default Lit"),
+	MSM_Subsurface				UMETA(DisplayName="Subsurface"),
+	MSM_PreintegratedSkin		UMETA(DisplayName="Preintegrated Skin"),
+	MSM_ClearCoat				UMETA(DisplayName="Clear Coat"),
+	MSM_SubsurfaceProfile		UMETA(DisplayName="Subsurface Profile"),
+	MSM_TwoSidedFoliage			UMETA(DisplayName="Two Sided Foliage"),
+	MSM_Hair					UMETA(DisplayName="Hair"),
+	MSM_Cloth					UMETA(DisplayName="Cloth"),
+	MSM_Eye						UMETA(DisplayName="Eye"),
+	/** Number of unique shading models. */
+	MSM_NUM						UMETA(Hidden),
+	/** Shading model will be determined by the Material Expression Graph,
+		by utilizing the 'Shading Model' MaterialAttribute output pin. */
+	MSM_FromMaterialExpression	UMETA(DisplayName="From Material Expression"),
+	MSM_MAX
+};
+
+static_assert(MSM_NUM <= 16, "Do not exceed 16 shading models without expanding FMaterialShadingModelField to support uint32 instead of uint16!");
+
+/** Wrapper for a bitfield of shading models. A material contains one of these to describe what possible shading models can be used by that material. */
+USTRUCT(BlueprintType)
+struct ENGINE_API FMaterialShadingModelField
+{
+	GENERATED_USTRUCT_BODY()
+
+public:
+	FMaterialShadingModelField() {}
+	FMaterialShadingModelField(EMaterialShadingModel InShadingModel)		{ AddShadingModel(InShadingModel); }
+
+	void AddShadingModel(EMaterialShadingModel InShadingModel)				{ check(InShadingModel < MSM_NUM); ShadingModelField |= (1 << (uint16)InShadingModel); }
+	void RemoveShadingModel(EMaterialShadingModel InShadingModel)			{ ShadingModelField &= ~(1 << (uint16)InShadingModel); }
+	void ClearShadingModels()												{ ShadingModelField = 0; }
+
+	// Check if any of the given shading models are present
+	bool HasAnyShadingModel(TArray<EMaterialShadingModel> InShadingModels) const	
+	{ 
+		for (EMaterialShadingModel ShadingModel : InShadingModels)
+		{
+			if (HasShadingModel(ShadingModel))
+			{
+				return true;
+			}
+		}
+		return false; 
+	}
+
+	bool HasShadingModel(EMaterialShadingModel InShadingModel) const		{ return (ShadingModelField & (1 << (uint16)InShadingModel)) != 0; }
+	bool HasOnlyShadingModel(EMaterialShadingModel InShadingModel) const	{ return ShadingModelField == (1 << (uint16)InShadingModel); }
+	bool IsUnlit() const													{ return HasShadingModel(MSM_Unlit); }
+	bool IsLit() const														{ return !IsUnlit(); }
+	bool IsValid() const													{ return (ShadingModelField > 0) && (ShadingModelField < (1 << MSM_NUM)); }
+	uint16 GetShadingModelField() const										{ return ShadingModelField; }
+	int32 CountShadingModels() const										{ return FMath::CountBits(ShadingModelField); }
+	EMaterialShadingModel GetFirstShadingModel() const						{ check(IsValid()); return (EMaterialShadingModel)FMath::CountTrailingZeros(ShadingModelField); }
+
+	bool operator==(const FMaterialShadingModelField& Other) const			{ return ShadingModelField == Other.GetShadingModelField(); }
+	bool operator!=(const FMaterialShadingModelField& Other) const			{ return ShadingModelField != Other.GetShadingModelField(); }
+
+private:
+	UPROPERTY()
+	uint16 ShadingModelField = 0;
 };
 
 /** This is used by the drawing passes to determine tessellation policy, so changes here need to be supported in native code. */
diff --git a/Engine/Source/Runtime/Engine/Classes/Materials/Material.h b/Engine/Source/Runtime/Engine/Classes/Materials/Material.h
--- a/Engine/Source/Runtime/Engine/Classes/Materials/Material.h
+++ b/Engine/Source/Runtime/Engine/Classes/Materials/Material.h
@@ -221,6 +221,14 @@ struct FScalarMaterialInput : public FMaterialInput
 };
 #endif
 
+#if !CPP      //noexport struct
+USTRUCT(noexport)
+struct FShadingModelMaterialInput : public FMaterialInput
+{
+	// No support for constant
+};
+#endif
+
 #if !CPP      //noexport struct
 USTRUCT(noexport)
 struct FVectorMaterialInput : public FMaterialInput
@@ -401,7 +409,11 @@ class UMaterial : public UMaterialInterface
 private:
 	/** Determines how inputs are combined to create the material's final color. */
 	UPROPERTY(EditAnywhere, Category=Material, AssetRegistrySearchable)
-	TEnumAsByte<enum EMaterialShadingModel> ShadingModel;
+	TEnumAsByte<enum EMaterialShadingModel> ShadingModel; 
+
+	UPROPERTY(AssetRegistrySearchable)
+	FMaterialShadingModelField ShadingModels;
+
 public:
 
 	/** If BlendMode is BLEND_Masked, the surface is not rendered where OpacityMask < OpacityMaskClipValue. */
@@ -467,6 +479,9 @@ public:
 	UPROPERTY()
 	FScalarMaterialInput PixelDepthOffset;
 
+	UPROPERTY()
+	FShadingModelMaterialInput ShadingModelFromMaterialExpression;
+
 	/** Indicates that the material should be rendered in the SeparateTranslucency Pass (not affected by DOF, requires bAllowSeparateTranslucency to be set in .ini). */
 	UPROPERTY(EditAnywhere, Category=Translucency, meta=(DisplayName = "Render After DOF"), AdvancedDisplay)
 	uint32 bEnableSeparateTranslucency:1;
@@ -978,11 +993,12 @@ public:
 		ERHIFeatureLevel::Type InFeatureLevel, EMaterialQualityLevel::Type InQuality) override;
 #endif
 	ENGINE_API virtual void RecacheUniformExpressions(bool bRecreateUniformBuffer) const override;
-
+	
 	ENGINE_API virtual float GetOpacityMaskClipValue() const override;
 	ENGINE_API virtual bool GetCastDynamicShadowAsMasked() const override;
 	ENGINE_API virtual EBlendMode GetBlendMode() const override;
-	ENGINE_API virtual EMaterialShadingModel GetShadingModel() const override;
+	ENGINE_API virtual FMaterialShadingModelField GetShadingModels() const override;
+	ENGINE_API virtual bool IsShadingModelFromMaterialExpression() const override;
 	ENGINE_API virtual bool IsTwoSided() const override;
 	ENGINE_API virtual bool IsDitheredLODTransition() const override;
 	ENGINE_API virtual bool IsTranslucencyWritingCustomDepth() const override;
@@ -991,7 +1007,7 @@ public:
 	ENGINE_API virtual bool IsPostProcessMaterial() const { return MaterialDomain == MD_PostProcess; }
 	ENGINE_API virtual USubsurfaceProfile* GetSubsurfaceProfile_Internal() const override;
 
-	ENGINE_API void SetShadingModel(EMaterialShadingModel NewModel) {ShadingModel = NewModel;}
+	ENGINE_API void SetShadingModel(EMaterialShadingModel NewModel) { ensure(ShadingModel < MSM_NUM); ShadingModel = NewModel; ShadingModels = FMaterialShadingModelField(ShadingModel); }
 
 	/** Checks to see if an input property should be active, based on the state of the material */
 	ENGINE_API virtual bool IsPropertyActive(EMaterialProperty InProperty) const override;
@@ -1544,6 +1560,9 @@ public:
 	 * @param	Expression	The expression dynamic parameter to check for duplicates.
 	 */
 	ENGINE_API virtual bool HasDuplicateDynamicParameters(const UMaterialExpression* Expression);
+
+	/** Walk the material expression graph from the Shading Model input pin (or from the Material Attribute) and find all used Shading Models for this material and update the field in this material. */
+	ENGINE_API void RebuildShadingModelField();
 #endif // WITH_EDITOR
 
 	/**
diff --git a/Engine/Source/Runtime/Engine/Classes/Materials/MaterialExpressionGetMaterialAttributes.h b/Engine/Source/Runtime/Engine/Classes/Materials/MaterialExpressionGetMaterialAttributes.h
--- a/Engine/Source/Runtime/Engine/Classes/Materials/MaterialExpressionGetMaterialAttributes.h
+++ b/Engine/Source/Runtime/Engine/Classes/Materials/MaterialExpressionGetMaterialAttributes.h
@@ -40,8 +40,9 @@ class UMaterialExpressionGetMaterialAttributes : public UMaterialExpression
 	virtual const TArray<FExpressionInput*> GetInputs()override;
 	virtual FExpressionInput* GetInput(int32 InputIndex)override;
 	virtual FName GetInputName(int32 InputIndex) const override;
+	virtual uint32 GetOutputType(int32 OutputIndex) override;
 	virtual bool IsInputConnectionRequired(int32 InputIndex) const override {return true;}
-	virtual uint32 GetInputType(int32 InputIndex) override {return MCT_MaterialAttributes;}
+	virtual uint32 GetInputType(int32 InputIndex) override { return MCT_MaterialAttributes; }
 	virtual bool IsResultMaterialAttributes(int32 OutputIndex) override {return OutputIndex == 0;}
 #endif
 	//~ End UMaterialExpression Interface
diff --git a/Engine/Source/Runtime/Engine/Classes/Materials/MaterialExpressionMakeMaterialAttributes.h b/Engine/Source/Runtime/Engine/Classes/Materials/MaterialExpressionMakeMaterialAttributes.h
--- a/Engine/Source/Runtime/Engine/Classes/Materials/MaterialExpressionMakeMaterialAttributes.h
+++ b/Engine/Source/Runtime/Engine/Classes/Materials/MaterialExpressionMakeMaterialAttributes.h
@@ -68,6 +68,9 @@ class UMaterialExpressionMakeMaterialAttributes : public UMaterialExpression
 	UPROPERTY()
 	FExpressionInput PixelDepthOffset;
 
+	UPROPERTY()
+	FExpressionInput ShadingModel;
+
 	//~ Begin UObject Interface
 	virtual void Serialize(FStructuredArchive::FRecord Record) override;
 	//~ End UObject Interface
diff --git a/Engine/Source/Runtime/Engine/Classes/Materials/MaterialExpressionSetMaterialAttributes.h b/Engine/Source/Runtime/Engine/Classes/Materials/MaterialExpressionSetMaterialAttributes.h
--- a/Engine/Source/Runtime/Engine/Classes/Materials/MaterialExpressionSetMaterialAttributes.h
+++ b/Engine/Source/Runtime/Engine/Classes/Materials/MaterialExpressionSetMaterialAttributes.h
@@ -39,9 +39,9 @@ class UMaterialExpressionSetMaterialAttributes : public UMaterialExpression
 	virtual const TArray<FExpressionInput*> GetInputs()override;
 	virtual FExpressionInput* GetInput(int32 InputIndex)override;
 	virtual FName GetInputName(int32 InputIndex) const override;
+	virtual uint32 GetInputType(int32 InputIndex) override;
 	virtual bool IsInputConnectionRequired(int32 InputIndex) const override {return true;}
 	virtual bool IsResultMaterialAttributes(int32 OutputIndex) override {return true;}
-	virtual uint32 GetInputType(int32 InputIndex) override {return InputIndex == 0 ? MCT_MaterialAttributes : MCT_Float3;}
 #endif
 	//~ End UMaterialExpression Interface
 };
diff --git a/Engine/Source/Runtime/Engine/Classes/Materials/MaterialExpressionShadingModel.h b/Engine/Source/Runtime/Engine/Classes/Materials/MaterialExpressionShadingModel.h
new file mode 100644
--- /dev/null
+++ b/Engine/Source/Runtime/Engine/Classes/Materials/MaterialExpressionShadingModel.h
@@ -0,0 +1,36 @@
+// Copyright 1998-2018 Epic Games, Inc. All Rights Reserved.
+
+#pragma once
+
+#include "CoreMinimal.h"
+#include "Engine/EngineTypes.h"
+#include "Materials/MaterialExpression.h"
+#include "MaterialExpressionShadingModel.generated.h"
+
+/**
+ * Compile a select "blend" between ShadingModels
+ *
+ * @param Compiler				The compiler to add code to
+ * @param A						Select A if Alpha is less than 0.5f
+ * @param B						Select B if Alpha is greater or equal to 0.5f
+ * @param Alpha					Bland factor [0..1]
+ * @return						Index to a new code chunk
+ */
+extern int32 CompileShadingModelBlendFunction(FMaterialCompiler* Compiler, int32 A, int32 B, int32 Alpha);
+
+UCLASS(collapsecategories, hidecategories = Object, MinimalAPI)
+class UMaterialExpressionShadingModel : public UMaterialExpression
+{
+	GENERATED_UCLASS_BODY()
+
+	//~ Begin UMaterialExpression Interface
+#if WITH_EDITOR
+	virtual int32 Compile(class FMaterialCompiler* Compiler, int32 OutputIndex) override;
+	virtual uint32 GetOutputType(int32 OutputIndex) override;
+	virtual void GetCaption(TArray<FString>& OutCaptions) const override;
+#endif
+public:
+	UPROPERTY(EditAnywhere, Category=ShadingModel,  meta=(ValidEnumValues="MSM_DefaultLit, MSM_Subsurface, MSM_PreintegratedSkin, MSM_ClearCoat, MSM_SubsurfaceProfile, MSM_TwoSidedFoliage, MSM_Hair, MSM_Cloth, MSM_Eye"))
+	TEnumAsByte<enum EMaterialShadingModel> ShadingModel = MSM_DefaultLit;
+	//~ End UMaterialExpression Interface
+};
diff --git a/Engine/Source/Runtime/Engine/Classes/Materials/MaterialInstance.h b/Engine/Source/Runtime/Engine/Classes/Materials/MaterialInstance.h
--- a/Engine/Source/Runtime/Engine/Classes/Materials/MaterialInstance.h
+++ b/Engine/Source/Runtime/Engine/Classes/Materials/MaterialInstance.h
@@ -313,9 +313,10 @@ class UMaterialInstance : public UMaterialInterface
 	uint8 TwoSided : 1;
 	uint8 DitheredLODTransition : 1;
 	uint8 bCastDynamicShadowAsMasked : 1;
+	uint8 bIsShadingModelFromMaterialExpression : 1;
 
 	TEnumAsByte<EBlendMode> BlendMode;
-	TEnumAsByte<EMaterialShadingModel> ShadingModel;
+	FMaterialShadingModelField ShadingModels;
 
 	FORCEINLINE bool GetReentrantFlag() const
 	{
@@ -459,10 +460,11 @@ public:
 
 	ENGINE_API virtual float GetOpacityMaskClipValue() const override;
 	ENGINE_API virtual EBlendMode GetBlendMode() const override;
-	ENGINE_API virtual EMaterialShadingModel GetShadingModel() const override;
+	ENGINE_API virtual FMaterialShadingModelField GetShadingModels() const override;
+	ENGINE_API virtual bool IsShadingModelFromMaterialExpression() const override;
 	ENGINE_API virtual bool IsTwoSided() const override;
 	ENGINE_API virtual bool IsDitheredLODTransition() const override;
-	ENGINE_API virtual bool IsMasked() const override;;
+	ENGINE_API virtual bool IsMasked() const override;
 	
 	ENGINE_API virtual USubsurfaceProfile* GetSubsurfaceProfile_Internal() const override;
 
diff --git a/Engine/Source/Runtime/Engine/Classes/Materials/MaterialInstanceDynamic.h b/Engine/Source/Runtime/Engine/Classes/Materials/MaterialInstanceDynamic.h
--- a/Engine/Source/Runtime/Engine/Classes/Materials/MaterialInstanceDynamic.h
+++ b/Engine/Source/Runtime/Engine/Classes/Materials/MaterialInstanceDynamic.h
@@ -129,8 +129,9 @@ class ENGINE_API UMaterialInstanceDynamic : public UMaterialInstance
 	//Material base property overrides. MIDs cannot override these so they just grab from their parent.
 	virtual float GetOpacityMaskClipValue() const override;
 	virtual bool GetCastDynamicShadowAsMasked() const override;
+	virtual FMaterialShadingModelField GetShadingModels() const override;
+	virtual bool IsShadingModelFromMaterialExpression() const override;
 	virtual EBlendMode GetBlendMode() const override;
-	virtual EMaterialShadingModel GetShadingModel() const override;
 	virtual bool IsTwoSided() const override;
 	virtual bool IsDitheredLODTransition() const override;
 	virtual bool IsMasked() const override;
diff --git a/Engine/Source/Runtime/Engine/Classes/Materials/MaterialInterface.h b/Engine/Source/Runtime/Engine/Classes/Materials/MaterialInterface.h
--- a/Engine/Source/Runtime/Engine/Classes/Materials/MaterialInterface.h
+++ b/Engine/Source/Runtime/Engine/Classes/Materials/MaterialInterface.h
@@ -708,7 +708,8 @@ public:
 	ENGINE_API virtual float GetOpacityMaskClipValue() const;
 	ENGINE_API virtual bool GetCastDynamicShadowAsMasked() const;
 	ENGINE_API virtual EBlendMode GetBlendMode() const;
-	ENGINE_API virtual EMaterialShadingModel GetShadingModel() const;
+	ENGINE_API virtual FMaterialShadingModelField GetShadingModels() const;
+	ENGINE_API virtual bool IsShadingModelFromMaterialExpression() const;
 	ENGINE_API virtual bool IsTwoSided() const;
 	ENGINE_API virtual bool IsDitheredLODTransition() const;
 	ENGINE_API virtual bool IsTranslucencyWritingCustomDepth() const;
diff --git a/Engine/Source/Runtime/Engine/Private/Components/PrimitiveComponent.cpp b/Engine/Source/Runtime/Engine/Private/Components/PrimitiveComponent.cpp
--- a/Engine/Source/Runtime/Engine/Private/Components/PrimitiveComponent.cpp
+++ b/Engine/Source/Runtime/Engine/Private/Components/PrimitiveComponent.cpp
@@ -948,7 +948,7 @@ bool UPrimitiveComponent::CanEditChange(const UProperty* InProperty) const
 
 				if (Material)
 				{
-					if (Material->GetShadingModel() != MSM_Unlit)
+					if (Material->GetShadingModels().IsLit())
 					{
 						bHasAnyLitMaterials = true;
 					}
diff --git a/Engine/Source/Runtime/Engine/Private/Components/StaticMeshComponent.cpp b/Engine/Source/Runtime/Engine/Private/Components/StaticMeshComponent.cpp
--- a/Engine/Source/Runtime/Engine/Private/Components/StaticMeshComponent.cpp
+++ b/Engine/Source/Runtime/Engine/Private/Components/StaticMeshComponent.cpp
@@ -1913,7 +1913,7 @@ bool UStaticMeshComponent::UsesOnlyUnlitMaterials() const
 				UMaterialInterface*	MaterialInterface	= GetMaterial(LOD.Sections[ElementIndex].MaterialIndex);
 				UMaterial*			Material			= MaterialInterface ? MaterialInterface->GetMaterial() : NULL;
 
-				bUsesOnlyUnlitMaterials = Material && Material->GetShadingModel() == MSM_Unlit;
+				bUsesOnlyUnlitMaterials = Material && Material->GetShadingModels().IsUnlit();
 			}
 		}
 		return bUsesOnlyUnlitMaterials;
diff --git a/Engine/Source/Runtime/Engine/Private/DebugViewModeMaterialProxy.cpp b/Engine/Source/Runtime/Engine/Private/DebugViewModeMaterialProxy.cpp
--- a/Engine/Source/Runtime/Engine/Private/DebugViewModeMaterialProxy.cpp
+++ b/Engine/Source/Runtime/Engine/Private/DebugViewModeMaterialProxy.cpp
@@ -71,7 +71,7 @@ FDebugViewModeMaterialProxy::FDebugViewModeMaterialProxy(
 			GetDependentShaderAndVFTypes(GMaxRHIShaderPlatform, ShaderTypes, ShaderPipelineTypes, VFTypes);
 
 			// Overwrite the shader map Id's dependencies with ones that came from the FMaterial actually being compiled (this)
-			// This is necessary as we change FMaterial attributes like GetShadingModel(), which factor into the ShouldCache functions that determine dependent shader types
+			// This is necessary as we change FMaterial attributes like GetShadingModels(), which factor into the ShouldCache functions that determine dependent shader types
 			ResourceId.SetShaderDependencies(ShaderTypes, ShaderPipelineTypes, VFTypes, GMaxRHIShaderPlatform);
 		}
 
@@ -163,9 +163,9 @@ enum EBlendMode FDebugViewModeMaterialProxy::GetBlendMode() const
 	return MaterialInterface ? MaterialInterface->GetBlendMode() : BLEND_Opaque;
 }
 
-enum EMaterialShadingModel FDebugViewModeMaterialProxy::GetShadingModel() const
+FMaterialShadingModelField FDebugViewModeMaterialProxy::GetShadingModels() const
 { 
-	return Material ? Material->GetShadingModel() : MSM_Unlit;
+	return Material ? Material->GetShadingModels() : MSM_Unlit;
 }
 
 float FDebugViewModeMaterialProxy::GetOpacityMaskClipValue() const
diff --git a/Engine/Source/Runtime/Engine/Private/DebugViewModeMaterialProxy.h b/Engine/Source/Runtime/Engine/Private/DebugViewModeMaterialProxy.h
--- a/Engine/Source/Runtime/Engine/Private/DebugViewModeMaterialProxy.h
+++ b/Engine/Source/Runtime/Engine/Private/DebugViewModeMaterialProxy.h
@@ -141,7 +141,7 @@ public:
 	virtual bool IsWireframe() const override;
 	virtual bool IsMasked() const override;
 	virtual enum EBlendMode GetBlendMode() const override;
-	virtual enum EMaterialShadingModel GetShadingModel() const override;
+	virtual FMaterialShadingModelField GetShadingModels() const override;
 	virtual float GetOpacityMaskClipValue() const override;
 	virtual bool GetCastDynamicShadowAsMasked() const override;
 	virtual void GatherCustomOutputExpressions(TArray<class UMaterialExpressionCustomOutput*>& OutCustomOutputs) const override;
diff --git a/Engine/Source/Runtime/Engine/Private/Materials/HLSLMaterialTranslator.h b/Engine/Source/Runtime/Engine/Private/Materials/HLSLMaterialTranslator.h
--- a/Engine/Source/Runtime/Engine/Private/Materials/HLSLMaterialTranslator.h
+++ b/Engine/Source/Runtime/Engine/Private/Materials/HLSLMaterialTranslator.h
@@ -321,6 +321,7 @@ public:
 		SharedPixelProperties[MP_Refraction] = true;
 		SharedPixelProperties[MP_PixelDepthOffset] = true;
 		SharedPixelProperties[MP_SubsurfaceColor] = true;
+		SharedPixelProperties[MP_ShadingModel] = true;
 
 		for (int32 Frequency = 0; Frequency < SF_NumFrequencies; ++Frequency)
 		{
@@ -591,7 +592,7 @@ public:
 			}
 
 			const EShaderFrequency NormalShaderFrequency = FMaterialAttributeDefinitionMap::GetShaderFrequency(MP_Normal);
-			const EMaterialShadingModel MaterialShadingModel = Material->GetShadingModel();
+			const FMaterialShadingModelField MaterialShadingModels = Material->GetShadingModels();
 			const EMaterialDomain Domain = Material->GetMaterialDomain();
 			const EBlendMode BlendMode = Material->GetBlendMode();
 
@@ -625,7 +626,7 @@ public:
 			Chunk[MP_WorldDisplacement]				= Material->CompilePropertyAndSetMaterialProperty(MP_WorldDisplacement		,this);
 			Chunk[MP_TessellationMultiplier]		= Material->CompilePropertyAndSetMaterialProperty(MP_TessellationMultiplier	,this);			
 
-			if (Domain == MD_Surface && IsSubsurfaceShadingModel(MaterialShadingModel))
+			if (Domain == MD_Surface && IsSubsurfaceShadingModel(MaterialShadingModels))
 			{
 				// Note we don't test for the blend mode as you can have a translucent material using the subsurface shading model
 
@@ -666,6 +667,8 @@ public:
 
 			Chunk[MP_PixelDepthOffset] = Material->CompilePropertyAndSetMaterialProperty(MP_PixelDepthOffset, this);
 
+			Chunk[MP_ShadingModel] = Material->CompilePropertyAndSetMaterialProperty(MP_ShadingModel, this);
+
 			// No more calls to non-vertex shader CompilePropertyAndSetMaterialProperty beyond this point
 			const uint32 SavedNumUserTexCoords = NumUserTexCoords;
 
@@ -692,7 +695,7 @@ public:
 			// Fully rough if we have a roughness code chunk and it's constant and evaluates to 1.
 			bIsFullyRough = Chunk[MP_Roughness] != INDEX_NONE && IsMaterialPropertyUsed(MP_Roughness, Chunk[MP_Roughness], FLinearColor(1, 0, 0, 0), 1) == false;
 
-			if (BlendMode == BLEND_Modulate && MaterialShadingModel != MSM_Unlit && !Material->IsDeferredDecal())
+			if (BlendMode == BLEND_Modulate && MaterialShadingModels.IsLit() && !Material->IsDeferredDecal())
 			{
 				Errorf(TEXT("Dynamically lit translucency is not supported for BLEND_Modulate materials."));
 			}
@@ -737,17 +740,17 @@ public:
 				Errorf(TEXT("Light function materials must be opaque."));
 			}
 
-			if (Material->IsLightFunction() && MaterialShadingModel != MSM_Unlit)
+			if (Material->IsLightFunction() && MaterialShadingModels.IsLit())
 			{
 				Errorf(TEXT("Light function materials must use unlit."));
 			}
 
-			if (Domain == MD_PostProcess && MaterialShadingModel != MSM_Unlit)
+			if (Domain == MD_PostProcess && MaterialShadingModels.IsLit())
 			{
 				Errorf(TEXT("Post process materials must use unlit."));
 			}
 
-			if (Material->AllowNegativeEmissiveColor() && MaterialShadingModel != MSM_Unlit)
+			if (Material->AllowNegativeEmissiveColor() && MaterialShadingModels.IsLit())
 			{
 				Errorf(TEXT("Only unlit materials can output negative emissive color."));
 			}
@@ -1429,6 +1432,7 @@ protected:
 		case MCT_StaticBool:	return TEXT("static bool");
 		case MCT_MaterialAttributes:	return TEXT("MaterialAttributes");
 		case MCT_TextureExternal:	return TEXT("TextureExternal");
+		case MCT_ShadingModel:	return TEXT("ShadingModel");
 		default:				return TEXT("unknown");
 		};
 	}
@@ -1449,6 +1453,7 @@ protected:
 		case MCT_StaticBool:	return TEXT("static bool");
 		case MCT_MaterialAttributes:	return TEXT("MaterialAttributes");
 		case MCT_TextureExternal:	return TEXT("TextureExternal");
+		case MCT_ShadingModel:	return TEXT("uint");
 		default:				return TEXT("unknown");
 		};
 	}
@@ -1521,8 +1526,8 @@ protected:
 			new(*CurrentScopeChunks) FShaderCodeChunk(FormattedCode,TEXT(""),Type,true);
 			return CodeIndex;
 		}
-		// Can only create temporaries for float and material attribute types.
-		else if (Type & (MCT_Float))
+		// Can only create temporaries for float and shading model types.
+		else if (Type & (MCT_Float) || Type == MCT_ShadingModel)
 		{
 			const int32 CodeIndex = CurrentScopeChunks->Num();
 			// Allocate a local variable name
@@ -1851,9 +1856,9 @@ protected:
 	}
 
 	// GetArithmeticResultType
-	EMaterialValueType GetArithmeticResultType(EMaterialValueType TypeA,EMaterialValueType TypeB)
+	EMaterialValueType GetArithmeticResultType(EMaterialValueType TypeA, EMaterialValueType TypeB)
 	{
-		if(!(TypeA & MCT_Float) || !(TypeB & MCT_Float))
+		if (!((TypeA & MCT_Float) || TypeA == MCT_ShadingModel) || !((TypeB & MCT_Float) || TypeB == MCT_ShadingModel))
 		{
 			Errorf(TEXT("Attempting to perform arithmetic on non-numeric types: %s %s"),DescribeType(TypeA),DescribeType(TypeB));
 			return MCT_Unknown;
@@ -1964,10 +1969,10 @@ protected:
 		return ShaderFrequency;
 	}
 
-	virtual EMaterialShadingModel GetMaterialShadingModel() const override
+	virtual FMaterialShadingModelField GetMaterialShadingModels() const override
 	{
 		check(Material);
-		return Material->GetShadingModel();
+		return Material->GetShadingModels();
 	}
 
 	virtual int32 Error(const TCHAR* Text) override
@@ -3210,13 +3215,18 @@ protected:
 
 	virtual int32 If(int32 A,int32 B,int32 AGreaterThanB,int32 AEqualsB,int32 ALessThanB,int32 ThresholdArg) override
 	{
-		if(A == INDEX_NONE || B == INDEX_NONE || AGreaterThanB == INDEX_NONE || ALessThanB == INDEX_NONE || ThresholdArg == INDEX_NONE)
+		if(A == INDEX_NONE || B == INDEX_NONE || AGreaterThanB == INDEX_NONE || ALessThanB == INDEX_NONE)
 		{
 			return INDEX_NONE;
 		}
 
 		if (AEqualsB != INDEX_NONE)
 		{
+			if (ThresholdArg == INDEX_NONE)
+			{
+				return INDEX_NONE;
+			}
+
 			EMaterialValueType ResultType = GetArithmeticResultType(GetParameterType(AGreaterThanB),GetArithmeticResultType(AEqualsB,ALessThanB));
 
 			int32 CoercedAGreaterThanB = ForceCast(AGreaterThanB,ResultType);
@@ -5481,6 +5491,17 @@ protected:
 		return GetPrimitiveProperty(MCT_Float, TEXT("CustomPrimitiveData"), *HlslName);
 	}
 
+virtual int32 ShadingModel(EMaterialShadingModel InSelectedShadingModel) override
+	{
+		// Don't allow unlit shading model, since it has to be mutually exclusive with the other shading models due to #ifdef logic in BasePassPixelShader.usf
+		if (InSelectedShadingModel == MSM_Unlit)
+		{
+			return INDEX_NONE;
+		}
+
+		return AddInlinedCodeChunk(MCT_ShadingModel, TEXT("%d"), InSelectedShadingModel);
+	}
+
 	virtual int32 CustomExpression( class UMaterialExpressionCustom* Custom, TArray<int32>& CompiledInputs ) override
 	{
 		int32 ResultIdx = INDEX_NONE;
diff --git a/Engine/Source/Runtime/Engine/Private/Materials/Material.cpp b/Engine/Source/Runtime/Engine/Private/Materials/Material.cpp
--- a/Engine/Source/Runtime/Engine/Private/Materials/Material.cpp
+++ b/Engine/Source/Runtime/Engine/Private/Materials/Material.cpp
@@ -37,6 +37,7 @@
 #include "Materials/MaterialExpressionVectorParameter.h"
 #include "Materials/MaterialExpressionVertexInterpolator.h"
 #include "Materials/MaterialExpressionSceneColor.h"
+#include "Materials/MaterialExpressionShadingModel.h"
 #include "Materials/MaterialFunction.h"
 #include "Materials/MaterialFunctionInstance.h"
 #include "Materials/MaterialExpressionMaterialFunctionCall.h"
@@ -152,7 +153,7 @@ int32 FMaterialResource::CompilePropertyAndSetMaterialProperty(EMaterialProperty
 		case MP_OpacityMask:
 			// Force basic opaque surfaces to skip masked/translucent-only attributes.
 			// Some features can force the material to create a masked variant which unintentionally runs this dormant code
-			if (GetMaterialDomain() != MD_Surface || GetBlendMode() != BLEND_Opaque || !(GetShadingModel() == MSM_Unlit || GetShadingModel() == MSM_DefaultLit))
+			if (GetMaterialDomain() != MD_Surface || GetBlendMode() != BLEND_Opaque || (GetShadingModels().IsLit() && !GetShadingModels().HasShadingModel(MSM_DefaultLit)))
 			{
 				Ret = MaterialInterface->CompileProperty(Compiler, Property);
 			}
@@ -267,9 +268,9 @@ public:
 		FDefaultMaterialInstance* Resource = this;
 		ENQUEUE_RENDER_COMMAND(FDestroyDefaultMaterialInstanceCommand)(
 			[Resource](FRHICommandList& RHICmdList)
-			{
-				delete Resource;
-			});
+		{
+			delete Resource;
+		});
 	}
 
 	// FMaterialRenderProxy interface.
@@ -3482,7 +3483,7 @@ void UMaterial::Serialize(FArchive& Ar)
 	}
 #endif // #if WITH_EDITOR
 
-	static_assert(MP_MAX == 29, "New material properties must have DoMaterialAttributesReorder called on them to ensure that any future reordering of property pins is correctly applied.");
+	static_assert(MP_MAX == 30, "New material properties must have DoMaterialAttributeReorder called on them to ensure that any future reordering of property pins is correctly applied.");
 
 	if (Ar.UE4Ver() < VER_UE4_MATERIAL_MASKED_BLENDMODE_TIDY)
 	{
@@ -3829,6 +3830,7 @@ void UMaterial::PostLoad()
 	DoMaterialAttributeReorder(&CustomizedUVs[6], UE4Ver);
 	DoMaterialAttributeReorder(&CustomizedUVs[7], UE4Ver);
 	DoMaterialAttributeReorder(&PixelDepthOffset, UE4Ver);
+	DoMaterialAttributeReorder(&ShadingModelFromMaterialExpression, UE4Ver);
 #endif // WITH_EDITORONLY_DATA
 
 	if (!IsDefaultMaterial())
@@ -3884,6 +3886,12 @@ void UMaterial::PostLoad()
 		ShadingModel = MSM_DefaultLit;
 	}
 
+	// If this material hasn't been compiled and saved with support for multiple shading models yet, propagate the existing shading model
+	if (!ShadingModels.IsValid())
+	{
+		ShadingModels = FMaterialShadingModelField(ShadingModel);
+	}
+
 	if(DecalBlendMode == DBM_MAX)
 	{
 		DecalBlendMode = DBM_Translucent;
@@ -4249,7 +4257,7 @@ bool UMaterial::CanEditChange(const UProperty* InProperty) const
 	
 		if (PropertyName == GET_MEMBER_NAME_STRING_CHECKED(UMaterial, ShadingModel))
 		{
-			return MaterialDomain == MD_Surface || (MaterialDomain == MD_DeferredDecal && DecalBlendMode == DBM_Volumetric_DistanceFunction);
+			return (MaterialDomain == MD_Surface || (MaterialDomain == MD_DeferredDecal && DecalBlendMode == DBM_Volumetric_DistanceFunction));
 		}
 
 		if (PropertyName == GET_MEMBER_NAME_STRING_CHECKED(UMaterial, DecalBlendMode))
@@ -4289,12 +4297,12 @@ bool UMaterial::CanEditChange(const UProperty* InProperty) const
 			|| PropertyName == GET_MEMBER_NAME_STRING_CHECKED(UMaterial, TranslucentMultipleScatteringExtinction)
 			|| PropertyName == GET_MEMBER_NAME_STRING_CHECKED(UMaterial, TranslucentShadowStartOffset))
 		{
-			return MaterialDomain != MD_DeferredDecal && IsTranslucentBlendMode(BlendMode) && ShadingModel != MSM_Unlit;
+			return MaterialDomain != MD_DeferredDecal && IsTranslucentBlendMode(BlendMode) && GetShadingModels().IsLit();
 		}
 
 		if (PropertyName == GET_MEMBER_NAME_STRING_CHECKED(UMaterial, SubsurfaceProfile))
 		{
-			return MaterialDomain == MD_Surface && UseSubsurfaceProfile(ShadingModel) && (BlendMode == BLEND_Opaque || BlendMode == BLEND_Masked);
+			return MaterialDomain == MD_Surface && UseSubsurfaceProfile(ShadingModels) && (BlendMode == BLEND_Opaque || BlendMode == BLEND_Masked);
 		}
 
 		if (PropertyName == GET_MEMBER_NAME_STRING_CHECKED(FLightmassMaterialInterfaceSettings, bCastShadowAsMasked))
@@ -4325,7 +4333,7 @@ void UMaterial::PostEditChangeProperty(FPropertyChangedEvent& PropertyChangedEve
 	Super::PostEditChangeProperty(PropertyChangedEvent);
 
 	UProperty* PropertyThatChanged = PropertyChangedEvent.Property;
-	
+
 	//Cancel any current compilation jobs that are in flight for this material.
 	CancelOutstandingCompilation();
 
@@ -4355,7 +4363,7 @@ void UMaterial::PostEditChangeProperty(FPropertyChangedEvent& PropertyChangedEve
 
 	TranslucencyDirectionalLightingIntensity = FMath::Clamp(TranslucencyDirectionalLightingIntensity, .1f, 10.0f);
 
-	// Don't want to recompile after a duplicate because it's just been done by PostLoad, nor during interactive changes to prevent constant recompliation while spinning properties.
+	// Don't want to recompile after a duplicate because it's just been done by PostLoad, nor during interactive changes to prevent constant recompilation while spinning properties.
 	if( PropertyChangedEvent.ChangeType == EPropertyChangeType::Duplicate || PropertyChangedEvent.ChangeType == EPropertyChangeType::Interactive )
 	{
 		bRequiresCompilation = false;
@@ -4363,6 +4371,9 @@ void UMaterial::PostEditChangeProperty(FPropertyChangedEvent& PropertyChangedEve
 	
 	if (bRequiresCompilation)
 	{
+		// Always rebuild the shading model field on recompile
+		RebuildShadingModelField();
+
 		// When redirecting an object pointer, we trust that the DDC hash will detect the change and that we don't need to force a recompile.
 		const bool bRegenerateId = PropertyChangedEvent.ChangeType != EPropertyChangeType::Redirected;
 		CacheResourceShadersForRendering(bRegenerateId);
@@ -4611,6 +4622,45 @@ void UMaterial::UpdateExpressionParameterName(UMaterialExpression* Expression)
 		}
 	}
 }
+
+void UMaterial::RebuildShadingModelField()
+{
+	ShadingModels.ClearShadingModels();
+	// If using shading model from material expression, go through the expressions and look for the ShadingModel expression to figure out what shading models need to be supported in this material
+	if (ShadingModel == MSM_FromMaterialExpression)
+	{
+		TArray<UMaterialExpression*> ShadingModelPropertyExpressions;
+
+		// Depending on if we use Material Attributes or not, we have to start at the right property
+		if (bUseMaterialAttributes)
+		{
+			GetExpressionsInPropertyChain(MP_MaterialAttributes, ShadingModelPropertyExpressions, nullptr);
+		}
+		else
+		{
+			GetExpressionsInPropertyChain(MP_ShadingModel, ShadingModelPropertyExpressions, nullptr);
+		}
+
+		for (UMaterialExpression* MatExpr: ShadingModelPropertyExpressions)
+		{
+			if (UMaterialExpressionShadingModel* ShadingModelExpression = Cast<UMaterialExpressionShadingModel>(MatExpr))
+			{
+				ShadingModels.AddShadingModel(ShadingModelExpression->ShadingModel);
+			}
+		}
+
+		// If no expressions have been found, set a default
+		if (!ShadingModels.IsValid())
+		{
+			ShadingModels.AddShadingModel(MSM_DefaultLit);
+		}
+	}
+	else 
+	{
+		// If a shading model has been selected directly for the material, set it here
+		ShadingModels.AddShadingModel(ShadingModel);
+	}
+}
 #endif // WITH_EDITOR
 
 bool UMaterial::GetExpressionParameterName(const UMaterialExpression* Expression, FName& OutName)
@@ -5100,6 +5150,7 @@ FExpressionInput* UMaterial::GetExpressionInputForProperty(EMaterialProperty InP
 		case MP_Refraction:				return &Refraction;
 		case MP_MaterialAttributes:		return &MaterialAttributes;
 		case MP_PixelDepthOffset:		return &PixelDepthOffset;
+		case MP_ShadingModel:			return &ShadingModelFromMaterialExpression;
 	}
 
 	if (InProperty >= MP_CustomizedUVs0 && InProperty <= MP_CustomizedUVs7)
@@ -5638,6 +5689,7 @@ int32 UMaterial::CompilePropertyEx( FMaterialCompiler* Compiler, const FGuid& At
 		case MP_WorldPositionOffset:	return WorldPositionOffset.CompileWithDefault(Compiler, Property);
 		case MP_WorldDisplacement:		return WorldDisplacement.CompileWithDefault(Compiler, Property);
 		case MP_PixelDepthOffset:		return PixelDepthOffset.CompileWithDefault(Compiler, Property);
+		case MP_ShadingModel:			return ShadingModelFromMaterialExpression.CompileWithDefault(Compiler, Property);
 
 		default:
 			if (Property >= MP_CustomizedUVs0 && Property <= MP_CustomizedUVs7)
@@ -5805,13 +5857,13 @@ EBlendMode UMaterial::GetBlendMode() const
 	}
 }
 
-EMaterialShadingModel UMaterial::GetShadingModel() const
+FMaterialShadingModelField UMaterial::GetShadingModels() const
 {
 	switch (MaterialDomain)
 	{
 		case MD_Surface:
 		case MD_Volume:
-			return ShadingModel;
+			return ShadingModels;
 		case MD_DeferredDecal:
 			return MSM_DefaultLit;
 
@@ -5827,6 +5879,11 @@ EMaterialShadingModel UMaterial::GetShadingModel() const
 	}
 }
 
+bool UMaterial::IsShadingModelFromMaterialExpression() const
+{
+	return ShadingModel == MSM_FromMaterialExpression;
+}
+
 bool UMaterial::IsTwoSided() const
 {
 	return TwoSided != 0;
@@ -5856,12 +5913,13 @@ USubsurfaceProfile* UMaterial::GetSubsurfaceProfile_Internal() const
 static bool IsPropertyActive_Internal(EMaterialProperty InProperty,
 	EMaterialDomain Domain,
 	EBlendMode BlendMode,
-	EMaterialShadingModel ShadingModel,
+	FMaterialShadingModelField ShadingModels,
 	ETranslucencyLightingMode TranslucencyLightingMode,
 	EDecalBlendMode DecalBlendMode,
 	bool bBlendableOutputAlpha,
 	bool bHasTessellation,
-	bool bHasRefraction )
+	bool bHasRefraction,
+	bool bUsesShadingModelFromMaterialExpression)
 {
 	if (Domain == MD_PostProcess)
 	{
@@ -6042,7 +6100,7 @@ static bool IsPropertyActive_Internal(EMaterialProperty InProperty,
 		break;
 	case MP_Opacity:
 		Active = bIsTranslucentBlendMode && BlendMode != BLEND_Modulate;
-		if (IsSubsurfaceShadingModel(ShadingModel))
+		if (IsSubsurfaceShadingModel(ShadingModels))
 		{
 			Active = true;
 		}
@@ -6052,27 +6110,27 @@ static bool IsPropertyActive_Internal(EMaterialProperty InProperty,
 		break;
 	case MP_BaseColor:
 	case MP_AmbientOcclusion:
-		Active = ShadingModel != MSM_Unlit;
+		Active = ShadingModels.IsLit();
 		break;
 	case MP_Specular:
 	case MP_Roughness:
-		Active = ShadingModel != MSM_Unlit && (!bIsTranslucentBlendMode || !bIsVolumetricTranslucencyLightingMode);
+		Active = ShadingModels.IsLit() && (!bIsTranslucentBlendMode || !bIsVolumetricTranslucencyLightingMode);
 		break;
 	case MP_Metallic:
 		// Subsurface models store opacity in place of Metallic in the GBuffer
-		Active = ShadingModel != MSM_Unlit && (!bIsTranslucentBlendMode || !bIsVolumetricTranslucencyLightingMode);
+		Active = ShadingModels.IsLit() && (!bIsTranslucentBlendMode || !bIsVolumetricTranslucencyLightingMode);
 		break;
 	case MP_Normal:
-		Active = (ShadingModel != MSM_Unlit && (!bIsTranslucentBlendMode || !bIsNonDirectionalTranslucencyLightingMode)) || bHasRefraction;
+		Active = (ShadingModels.IsLit() && (!bIsTranslucentBlendMode || !bIsNonDirectionalTranslucencyLightingMode)) || bHasRefraction;
 		break;
 	case MP_SubsurfaceColor:
-		Active = ShadingModel == MSM_Subsurface || ShadingModel == MSM_PreintegratedSkin || ShadingModel == MSM_TwoSidedFoliage || ShadingModel == MSM_Cloth;
+		Active = ShadingModels.HasAnyShadingModel({ MSM_Subsurface, MSM_PreintegratedSkin, MSM_TwoSidedFoliage, MSM_Cloth });
 		break;
 	case MP_CustomData0:
-		Active = ShadingModel == MSM_ClearCoat || ShadingModel == MSM_Hair || ShadingModel == MSM_Cloth || ShadingModel == MSM_Eye;
+		Active = ShadingModels.HasAnyShadingModel({ MSM_ClearCoat, MSM_Hair, MSM_Cloth, MSM_Eye });
 		break;
 	case MP_CustomData1:
-		Active = ShadingModel == MSM_ClearCoat || ShadingModel == MSM_Eye;
+		Active = ShadingModels.HasAnyShadingModel({ MSM_ClearCoat, MSM_Eye });
 		break;
 	case MP_TessellationMultiplier:
 	case MP_WorldDisplacement:
@@ -6088,6 +6146,9 @@ static bool IsPropertyActive_Internal(EMaterialProperty InProperty,
 	case MP_PixelDepthOffset:
 		Active = !bIsTranslucentBlendMode;
 		break;
+	case MP_ShadingModel:
+		Active = bUsesShadingModelFromMaterialExpression;
+                break;
 	case MP_MaterialAttributes:
 	default:
 		Active = true;
@@ -6109,12 +6170,13 @@ bool UMaterial::IsPropertyActiveInEditor(EMaterialProperty InProperty) const
 	return IsPropertyActive_Internal(InProperty,
 		MaterialDomain,
 		BlendMode,
-		ShadingModel,
+		ShadingModels,
 		TranslucencyLightingMode,
 		DecalBlendMode,
 		BlendableOutputAlpha,
 		D3D11TessellationMode != MTM_NoTessellation,
-		Refraction.IsConnected());
+		Refraction.IsConnected(),
+		IsShadingModelFromMaterialExpression());
 }
 #endif // WITH_EDITOR
 
@@ -6123,12 +6185,13 @@ bool UMaterial::IsPropertyActiveInDerived(EMaterialProperty InProperty, const UM
 	return IsPropertyActive_Internal(InProperty,
 		MaterialDomain,
 		DerivedMaterial->GetBlendMode(),
-		DerivedMaterial->GetShadingModel(),
+		DerivedMaterial->GetShadingModels(),
 		TranslucencyLightingMode,
 		DecalBlendMode,
 		BlendableOutputAlpha,
 		D3D11TessellationMode != MTM_NoTessellation,
-		Refraction.IsConnected());
+		Refraction.IsConnected(),
+		DerivedMaterial->IsShadingModelFromMaterialExpression());
 }
 
 #if WITH_EDITORONLY_DATA
diff --git a/Engine/Source/Runtime/Engine/Private/Materials/MaterialExpressions.cpp b/Engine/Source/Runtime/Engine/Private/Materials/MaterialExpressions.cpp
--- a/Engine/Source/Runtime/Engine/Private/Materials/MaterialExpressions.cpp
+++ b/Engine/Source/Runtime/Engine/Private/Materials/MaterialExpressions.cpp
@@ -164,6 +164,7 @@
 #include "Materials/MaterialExpressionSceneTexelSize.h"
 #include "Materials/MaterialExpressionSceneTexture.h"
 #include "Materials/MaterialExpressionScreenPosition.h"
+#include "Materials/MaterialExpressionShadingModel.h"
 #include "Materials/MaterialExpressionSine.h"
 #include "Materials/MaterialExpressionSobol.h"
 #include "Materials/MaterialExpressionSpeedTree.h"
@@ -350,6 +351,8 @@ void GetMaterialValueTypeDescriptions(uint32 MaterialValueType, TArray<FText>& O
 		OutDescriptions.Add(LOCTEXT("StaticBool", "Bool"));
 	if (MaterialValueType & MCT_MaterialAttributes)
 		OutDescriptions.Add(LOCTEXT("MaterialAttributes", "Material Attributes"));
+	if (MaterialValueType & MCT_ShadingModel)
+		OutDescriptions.Add(LOCTEXT("ShadingModel", "Shading Model"));
 	if (MaterialValueType & MCT_Unknown)
 		OutDescriptions.Add(LOCTEXT("Unknown", "Unknown"));
 }
@@ -529,6 +532,28 @@ int32 CompileTextureSample(
 }
 #endif // WITH_EDITOR
 
+/**
+ * Compile a select "blend" between ShadingModels
+ *
+ * @param Compiler				The compiler to add code to
+ * @param A						Select A if Alpha is less than 0.5f
+ * @param B						Select B if Alpha is greater or equal to 0.5f
+ * @param Alpha					Bland factor [0..1]
+ * @return						Index to a new code chunk
+ */
+int32 CompileShadingModelBlendFunction(FMaterialCompiler* Compiler, int32 A, int32 B, int32 Alpha)
+{
+	if (A == INDEX_NONE || B == INDEX_NONE || Alpha == INDEX_NONE)
+	{
+		return INDEX_NONE;
+	}
+
+	int32  MidPoint = Compiler->Constant(0.5f);
+
+	return Compiler->If(Alpha, MidPoint, B, INDEX_NONE, A, INDEX_NONE);
+}
+
+
 UMaterialExpression::UMaterialExpression(const FObjectInitializer& ObjectInitializer)
 	: Super(ObjectInitializer)
 #if WITH_EDITORONLY_DATA
@@ -4516,7 +4541,7 @@ int32 UMaterialExpressionMakeMaterialAttributes::Compile(class FMaterialCompiler
 	int32 Ret = INDEX_NONE;
 	UMaterialExpression* Expression = NULL;
 
- 	static_assert(MP_MAX == 29, 
+ 	static_assert(MP_MAX == 30, 
 		"New material properties should be added to the end of the inputs for this expression. \
 		The order of properties here should match the material results pins, the make material attriubtes node inputs and the mapping of IO indices to properties in GetMaterialPropertyFromInputOutputIndex().\
 		Insertions into the middle of the properties or a change in the order of properties will also require that existing data is fixed up in DoMaterialAttriubtesReorder().\
@@ -4543,6 +4568,7 @@ int32 UMaterialExpressionMakeMaterialAttributes::Compile(class FMaterialCompiler
 	case MP_AmbientOcclusion: Ret = AmbientOcclusion.Compile(Compiler); Expression = AmbientOcclusion.Expression; break;
 	case MP_Refraction: Ret = Refraction.Compile(Compiler); Expression = Refraction.Expression; break;
 	case MP_PixelDepthOffset: Ret = PixelDepthOffset.Compile(Compiler); Expression = PixelDepthOffset.Expression; break;
+	case MP_ShadingModel: Ret = ShadingModel.Compile(Compiler); Expression = ShadingModel.Expression; break;
 	};
 
 	if (Property >= MP_CustomizedUVs0 && Property <= MP_CustomizedUVs7)
@@ -4588,7 +4614,7 @@ UMaterialExpressionBreakMaterialAttributes::UMaterialExpressionBreakMaterialAttr
 
 	MenuCategories.Add(ConstructorStatics.NAME_MaterialAttributes);
 	
- 	static_assert(MP_MAX == 29, 
+ 	static_assert(MP_MAX == 30, 
 		"New material properties should be added to the end of the outputs for this expression. \
 		The order of properties here should match the material results pins, the make material attriubtes node inputs and the mapping of IO indices to properties in GetMaterialPropertyFromInputOutputIndex().\
 		Insertions into the middle of the properties or a change in the order of properties will also require that existing data is fixed up in DoMaterialAttriubtesReorder().\
@@ -4618,6 +4644,7 @@ UMaterialExpressionBreakMaterialAttributes::UMaterialExpressionBreakMaterialAttr
 	}
 
 	Outputs.Add(FExpressionOutput(TEXT("PixelDepthOffset"), 1, 1, 0, 0, 0));
+	Outputs.Add(FExpressionOutput(TEXT("ShadingModel"), 1, 1, 0, 0, 0));
 #endif
 }
 
@@ -4656,7 +4683,8 @@ void UMaterialExpressionBreakMaterialAttributes::Serialize(FStructuredArchive::F
 			Outputs[OutputIndex].SetMask(1, 1, 1, 0, 0);
 		}
 
-		Outputs[OutputIndex].SetMask(1, 1, 0, 0, 0); // PixelDepthOffset
+		Outputs[OutputIndex].SetMask(1, 1, 0, 0, 0); ++OutputIndex;// PixelDepthOffset
+		Outputs[OutputIndex].SetMask(1, 1, 0, 0, 0); // ShadingModelFromMaterialExpression
 	}
 #endif // WITH_EDITOR
 }
@@ -4692,6 +4720,7 @@ int32 UMaterialExpressionBreakMaterialAttributes::Compile(class FMaterialCompile
 		PropertyToIOIndexMap.Add(MP_CustomizedUVs6, 22);
 		PropertyToIOIndexMap.Add(MP_CustomizedUVs7, 23);
 		PropertyToIOIndexMap.Add(MP_PixelDepthOffset, 24);
+		PropertyToIOIndexMap.Add(MP_ShadingModel, 25);
 	}
 
 	// Here we don't care about any multiplex index coming in.
@@ -4839,6 +4868,25 @@ FName UMaterialExpressionGetMaterialAttributes::GetInputName(int32 InputIndex) c
 	return NAME_None;
 }
 
+uint32 UMaterialExpressionGetMaterialAttributes::GetOutputType(int32 OutputIndex)
+{
+	// Call base class impl to get the type
+	uint32 OutputType = Super::GetOutputType(OutputIndex);
+
+	// Override the type if it's a ShadingModel type
+	if (OutputIndex > 0) // "0th" place is the mandatory MaterialAttribute itself, skip it
+	{
+		ensure(OutputIndex < AttributeGetTypes.Num() + 1);
+		EMaterialValueType PinType = FMaterialAttributeDefinitionMap::GetValueType(AttributeGetTypes[OutputIndex - 1]);
+		if (PinType == MCT_ShadingModel)
+		{
+			OutputType = PinType;
+		}
+	}
+
+	return OutputType;
+}
+
 void UMaterialExpressionGetMaterialAttributes::PreEditChange(UProperty* PropertyAboutToChange)
 {
 	// Backup attribute array so we can re-connect pins
@@ -5064,6 +5112,31 @@ FName UMaterialExpressionSetMaterialAttributes::GetInputName(int32 InputIndex) c
 	return Name;
 }
 
+uint32 UMaterialExpressionSetMaterialAttributes::GetInputType(int32 InputIndex)
+{
+	uint32 InputType = MCT_Unknown;
+
+	if (InputIndex == 0)
+	{
+		InputType = MCT_MaterialAttributes;
+	}
+	else
+	{
+		ensure(InputIndex > 0 && InputIndex < AttributeSetTypes.Num() + 1);
+		InputType = FMaterialAttributeDefinitionMap::GetValueType(AttributeSetTypes[InputIndex - 1]);
+		if (InputType == MCT_ShadingModel)
+		{
+			InputType = MCT_ShadingModel;
+		}
+		else
+		{
+			InputType = MCT_Float3;
+		}
+	}
+
+	return InputType;
+}
+
 void UMaterialExpressionSetMaterialAttributes::PreEditChange(UProperty* PropertyAboutToChange)
 {
 	// Backup attribute array so we can re-connect pins
@@ -5098,7 +5171,7 @@ void UMaterialExpressionSetMaterialAttributes::PostEditChangeProperty(FPropertyC
 		
 			// Copy final defaults to new input
 			Inputs.Add(FExpressionInput());
-			Inputs.Last().ExpressionName = FName(*FMaterialAttributeDefinitionMap::GetDisplayName(AttributeSetTypes.Last()));
+			Inputs.Last().InputName = FName(*FMaterialAttributeDefinitionMap::GetDisplayName(AttributeSetTypes.Last()));
 			GraphNode->ReconstructNode();
 		}	 
 		else if (PreEditAttributeSetTypes.Num() > AttributeSetTypes.Num())
@@ -8537,10 +8610,15 @@ uint32 UMaterialExpressionIf::GetInputType(int32 InputIndex)
 		{
 			return MCT_MaterialAttributes;
 		}
+		else if ((A.GetTracedInput().Expression && !A.Expression->ContainsInputLoop() && A.Expression->GetOutputType(0) == MCT_ShadingModel) &&
+			(B.GetTracedInput().Expression && !B.Expression->ContainsInputLoop() && B.Expression->GetOutputType(0) == MCT_ShadingModel))
+		{
+			return MCT_ShadingModel;
+		}	
 		else
 		{
 			return MCT_Float;
-		}	
+		}
 	}
 
 	return MCT_Unknown;
@@ -15280,6 +15358,48 @@ void UMaterialExpressionCustomPrimitiveData::PostEditChangeProperty(FPropertyCha
 		Super::PostEditChangeProperty(PropertyChangedEvent);
 	}
 }
+#endif // WITH_EDITOR
 
+//
+//	UMaterialExpressionShadingModel
+//
+UMaterialExpressionShadingModel::UMaterialExpressionShadingModel(const FObjectInitializer& ObjectInitializer)
+	: Super(ObjectInitializer)
+{
+#if WITH_EDITORONLY_DATA
+
+	MenuCategories.Add(LOCTEXT("Shading Model", "Shading Model"));
+
+	// bShaderInputData = true;
+	bShowOutputNameOnPin = true;
+	
+	Outputs.Reset();
+	Outputs.Add(FExpressionOutput(TEXT("")));
+#endif
+}
+
+#if WITH_EDITOR
+int32 UMaterialExpressionShadingModel::Compile(class FMaterialCompiler* Compiler, int32 OutputIndex)
+{
+	// Maybe add the shading model to the material here. Don't forget to clear the material shading model list before compilation though
+	return Compiler->ShadingModel(ShadingModel);
+}
+
+void UMaterialExpressionShadingModel::GetCaption(TArray<FString>& OutCaptions) const
+{
+	const UEnum* ShadingModelEnum = FindObject<UEnum>(NULL, TEXT("Engine.EMaterialShadingModel"));
+	check(ShadingModelEnum);
+
+	const FString ShadingModelDisplayName = ShadingModelEnum->GetDisplayNameTextByValue(ShadingModel).ToString();
+
+	// Add as a stack, last caption to be added will be the main (bold) caption
+	OutCaptions.Add(ShadingModelDisplayName);
+	OutCaptions.Add(TEXT("Shading Model"));
+}
+
+uint32 UMaterialExpressionShadingModel::GetOutputType(int32 OutputIndex)
+{
+	return MCT_ShadingModel;
+}
 #endif // WITH_EDITOR
 #undef LOCTEXT_NAMESPACE
diff --git a/Engine/Source/Runtime/Engine/Private/Materials/MaterialInstance.cpp b/Engine/Source/Runtime/Engine/Private/Materials/MaterialInstance.cpp
--- a/Engine/Source/Runtime/Engine/Private/Materials/MaterialInstance.cpp
+++ b/Engine/Source/Runtime/Engine/Private/Materials/MaterialInstance.cpp
@@ -1537,9 +1537,14 @@ bool UMaterialInstanceDynamic::IsMasked() const
 	return Parent ? Parent->IsMasked() : false;
 }
 
-EMaterialShadingModel UMaterialInstanceDynamic::GetShadingModel() const
+FMaterialShadingModelField UMaterialInstanceDynamic::GetShadingModels() const
 {
-	return Parent ? Parent->GetShadingModel() : MSM_DefaultLit;
+	return Parent ? Parent->GetShadingModels() : MSM_DefaultLit;
+}
+
+bool UMaterialInstanceDynamic::IsShadingModelFromMaterialExpression() const
+{
+	return Parent ? Parent->IsShadingModelFromMaterialExpression() : false;
 }
 
 void UMaterialInstance::CopyMaterialInstanceParameters(UMaterialInterface* Source)
@@ -2390,9 +2395,10 @@ void UMaterialInstance::UpdateOverridableBaseProperties()
 	{
 		OpacityMaskClipValue = 0.0f;
 		BlendMode = BLEND_Opaque;
-		ShadingModel = MSM_DefaultLit;
+		ShadingModels = MSM_DefaultLit;
 		TwoSided = 0;
 		DitheredLODTransition = 0;
+		bIsShadingModelFromMaterialExpression = 0;
 		return;
 	}
 
@@ -2428,12 +2434,24 @@ void UMaterialInstance::UpdateOverridableBaseProperties()
 
 	if (BasePropertyOverrides.bOverride_ShadingModel)
 	{
-		ShadingModel = BasePropertyOverrides.ShadingModel;
+		// It's only possible to override using a single shading model
+		ShadingModels = FMaterialShadingModelField(BasePropertyOverrides.ShadingModel);
+		bIsShadingModelFromMaterialExpression = 0;
 	}
 	else
 	{
-		ShadingModel = Parent->GetShadingModel();
-		BasePropertyOverrides.ShadingModel = ShadingModel;
+		ShadingModels = Parent->GetShadingModels();
+		bIsShadingModelFromMaterialExpression = Parent->IsShadingModelFromMaterialExpression();
+
+		if (bIsShadingModelFromMaterialExpression)
+		{
+			BasePropertyOverrides.ShadingModel = MSM_FromMaterialExpression; 
+		}
+		else
+		{
+			ensure(ShadingModels.CountShadingModels() == 1);
+			BasePropertyOverrides.ShadingModel = ShadingModels.GetFirstShadingModel(); 
+		}
 	}
 
 	if (BasePropertyOverrides.bOverride_TwoSided)
@@ -3981,12 +3999,12 @@ void UMaterialInstance::GetBasePropertyOverridesHash(FSHAHash& OutHash)const
 		bHasOverrides = true;
 	}
 	
-	EMaterialShadingModel UsedShadingModel = GetShadingModel();
-	if (UsedShadingModel != Mat->GetShadingModel())
+	FMaterialShadingModelField UsedShadingModels = GetShadingModels();
+	if (UsedShadingModels != Mat->GetShadingModels())
 	{
 		const FString HashString = TEXT("bOverride_ShadingModel");
 		Hash.UpdateWithString(*HashString, HashString.Len());
-		Hash.Update((const uint8*)&UsedShadingModel, sizeof(UsedShadingModel));
+		Hash.Update((const uint8*)&UsedShadingModels, sizeof(UsedShadingModels));
 		bHasOverrides = true;
 	}
 
@@ -4022,7 +4040,7 @@ bool UMaterialInstance::HasOverridenBaseProperties()const
 	if (Parent && Material && Material->bUsedAsSpecialEngineMaterial == false &&
 		((FMath::Abs(GetOpacityMaskClipValue() - Parent->GetOpacityMaskClipValue()) > SMALL_NUMBER) ||
 		(GetBlendMode() != Parent->GetBlendMode()) ||
-		(GetShadingModel() != Parent->GetShadingModel()) ||
+		(GetShadingModels() != Parent->GetShadingModels()) ||
 		(IsTwoSided() != Parent->IsTwoSided()) ||
 		(IsDitheredLODTransition() != Parent->IsDitheredLODTransition()) ||
 		(GetCastDynamicShadowAsMasked() != Parent->GetCastDynamicShadowAsMasked())
@@ -4044,9 +4062,14 @@ EBlendMode UMaterialInstance::GetBlendMode() const
 	return BlendMode;
 }
 
-EMaterialShadingModel UMaterialInstance::GetShadingModel() const
+FMaterialShadingModelField UMaterialInstance::GetShadingModels() const
+{
+	return ShadingModels;
+}
+
+bool UMaterialInstance::IsShadingModelFromMaterialExpression() const
 {
-	return ShadingModel;
+	return bIsShadingModelFromMaterialExpression;
 }
 
 bool UMaterialInstance::IsTwoSided() const
diff --git a/Engine/Source/Runtime/Engine/Private/Materials/MaterialInterface.cpp b/Engine/Source/Runtime/Engine/Private/Materials/MaterialInterface.cpp
--- a/Engine/Source/Runtime/Engine/Private/Materials/MaterialInterface.cpp
+++ b/Engine/Source/Runtime/Engine/Private/Materials/MaterialInterface.cpp
@@ -122,14 +122,13 @@ FMaterialRelevance UMaterialInterface::GetRelevance_Internal(const UMaterial* Ma
 		const EBlendMode BlendMode = (EBlendMode)GetBlendMode();
 		const bool bIsTranslucent = IsTranslucentBlendMode(BlendMode);
 
-		EMaterialShadingModel ShadingModel = GetShadingModel();
 		EMaterialDomain Domain = (EMaterialDomain)MaterialResource->GetMaterialDomain();
 		bool bDecal = (Domain == MD_DeferredDecal);
 
 		// Determine the material's view relevance.
 		FMaterialRelevance MaterialRelevance;
 
-		MaterialRelevance.ShadingModelMask = 1 << ShadingModel;
+		MaterialRelevance.ShadingModelMask = GetShadingModels().GetShadingModelField();
 
 		if(bDecal)
 		{
@@ -423,11 +422,16 @@ bool UMaterialInterface::GetCastDynamicShadowAsMasked() const
 	return false;
 }
 
-EMaterialShadingModel UMaterialInterface::GetShadingModel() const
+FMaterialShadingModelField UMaterialInterface::GetShadingModels() const
 {
 	return MSM_DefaultLit;
 }
 
+bool UMaterialInterface::IsShadingModelFromMaterialExpression() const
+{
+	return false;
+}
+
 USubsurfaceProfile* UMaterialInterface::GetSubsurfaceProfile_Internal() const
 {
 	return NULL;
@@ -473,10 +477,10 @@ void UMaterialInterface::UpdateMaterialRenderProxy(FMaterialRenderProxy& Proxy)
 	// no 0 pointer
 	check(&Proxy);
 
-	EMaterialShadingModel MaterialShadingModel = GetShadingModel();
+	FMaterialShadingModelField MaterialShadingModels = GetShadingModels();
 
 	// for better performance we only update SubsurfaceProfileRT if the feature is used
-	if (UseSubsurfaceProfile(MaterialShadingModel))
+	if (UseSubsurfaceProfile(MaterialShadingModels))
 	{
 		FSubsurfaceProfileStruct Settings;
 
diff --git a/Engine/Source/Runtime/Engine/Private/Materials/MaterialShader.cpp b/Engine/Source/Runtime/Engine/Private/Materials/MaterialShader.cpp
--- a/Engine/Source/Runtime/Engine/Private/Materials/MaterialShader.cpp
+++ b/Engine/Source/Runtime/Engine/Private/Materials/MaterialShader.cpp
@@ -76,6 +76,35 @@ FString GetShadingModelString(EMaterialShadingModel ShadingModel)
 	return ShadingModelName;
 }
 
+
+/** Converts an FMaterialShadingModelField to a string description containing all the shading models present, delimited by "|" */
+FString GetShadingModelFieldString(FMaterialShadingModelField ShadingModels, FString (*ShadingModelToString)(EMaterialShadingModel), const FString& Delimiter)
+{
+	FString ShadingModelsName;
+	uint32 TempShadingModels = (uint32)ShadingModels.GetShadingModelField();
+
+	while (TempShadingModels)
+	{
+		uint32 BitIndex = FMath::CountTrailingZeros(TempShadingModels); // Find index of first set bit
+		TempShadingModels &= ~(1 << BitIndex); // Flip first set bit to 0
+		ShadingModelsName += ShadingModelToString((EMaterialShadingModel)BitIndex); // Add the name of the shading model corresponding to that bit
+
+		// If there are more bits left, add a pipe limiter to the string 
+		if (TempShadingModels)
+		{
+			ShadingModelsName.Append(TEXT("|"));
+		}
+	}
+
+	return ShadingModelsName;
+}
+
+/** Converts an FMaterialShadingModelField to a string description containing all the shading models present, delimited by "|" */
+FString GetShadingModelFieldString(FMaterialShadingModelField ShadingModels)
+{
+	return GetShadingModelFieldString(ShadingModels, &GetShadingModelString, TEXT("|"));
+}
+
 /** Converts an EBlendMode to a string description. */
 FString GetBlendModeString(EBlendMode BlendMode)
 {
@@ -118,22 +147,17 @@ void UpdateMaterialShaderCompilingStats(const FMaterial* Material)
 		default: INC_DWORD_STAT_BY(STAT_ShaderCompiling_NumTransparentMaterialShaders,1); break;
 	}
 
-	switch(Material->GetShadingModel())
-	{
-		case MSM_Unlit:
-			INC_DWORD_STAT_BY(STAT_ShaderCompiling_NumUnlitMaterialShaders,1);
-			break;
-		case MSM_DefaultLit:
-		case MSM_Subsurface:
-		case MSM_PreintegratedSkin:
-		case MSM_ClearCoat:
-		case MSM_Cloth:
-		case MSM_SubsurfaceProfile:
-		case MSM_TwoSidedFoliage:
-			INC_DWORD_STAT_BY(STAT_ShaderCompiling_NumLitMaterialShaders,1);
-			break;
-		default: break;
-	};
+	FMaterialShadingModelField ShadingModels = Material->GetShadingModels();
+	
+	if (ShadingModels.HasOnlyShadingModel(MSM_Unlit))
+	{
+		INC_DWORD_STAT_BY(STAT_ShaderCompiling_NumUnlitMaterialShaders, 1);
+	}
+	else if (ShadingModels.HasAnyShadingModel({ MSM_DefaultLit, MSM_Subsurface, MSM_PreintegratedSkin, MSM_ClearCoat, MSM_Cloth, MSM_SubsurfaceProfile, MSM_TwoSidedFoliage }))
+	{
+		INC_DWORD_STAT_BY(STAT_ShaderCompiling_NumLitMaterialShaders, 1);
+	}
+
 
 	if (Material->IsSpecialEngineMaterial())
 	{
diff --git a/Engine/Source/Runtime/Engine/Private/Materials/MaterialShared.cpp b/Engine/Source/Runtime/Engine/Private/Materials/MaterialShared.cpp
--- a/Engine/Source/Runtime/Engine/Private/Materials/MaterialShared.cpp
+++ b/Engine/Source/Runtime/Engine/Private/Materials/MaterialShared.cpp
@@ -19,6 +19,7 @@
 #include "UObject/UObjectIterator.h"
 #include "ComponentReregisterContext.h"
 #include "Materials/MaterialExpressionBreakMaterialAttributes.h"
+#include "Materials/MaterialExpressionShadingModel.h"
 #include "Materials/MaterialExpressionReroute.h"
 #include "ShaderCompiler.h"
 #include "MaterialCompiler.h"
@@ -331,6 +332,11 @@ bool FScalarMaterialInput::Serialize(FArchive& Ar)
 	return SerializeMaterialInput<float>(Ar, *this);
 }
 
+bool FShadingModelMaterialInput::Serialize(FArchive& Ar)
+{
+	return SerializeMaterialInput<uint32>(Ar, *this);
+}
+
 bool FVectorMaterialInput::Serialize(FArchive& Ar)
 {
 	return SerializeMaterialInput<FVector>(Ar, *this);
@@ -384,6 +390,20 @@ int32 FScalarMaterialInput::CompileWithDefault(class FMaterialCompiler* Compiler
 	return Compiler->ForceCast(FMaterialAttributeDefinitionMap::CompileDefaultExpression(Compiler, Property), MCT_Float1);
 }
 
+int32 FShadingModelMaterialInput::CompileWithDefault(class FMaterialCompiler* Compiler, EMaterialProperty Property)
+{
+	if (Expression)
+	{
+		int32 ResultIndex = FExpressionInput::Compile(Compiler);
+		if (ResultIndex != INDEX_NONE)
+		{
+			return ResultIndex;
+		}
+	}
+
+	return Compiler->ForceCast(FMaterialAttributeDefinitionMap::CompileDefaultExpression(Compiler, Property), MCT_ShadingModel, MFCF_ExactMatch);
+}
+
 int32 FVectorMaterialInput::CompileWithDefault(class FMaterialCompiler* Compiler, EMaterialProperty Property)
 {
 	if (UseConstant)
@@ -975,9 +995,9 @@ void FMaterial::ReleaseShaderMap()
 		FMaterial* Material = this;
 		ENQUEUE_RENDER_COMMAND(ReleaseShaderMap)(
 			[Material](FRHICommandList& RHICmdList)
-			{
-				Material->SetRenderingThreadShaderMap(nullptr);
-			});
+		{
+			Material->SetRenderingThreadShaderMap(nullptr);
+		});
 	}
 }
 
@@ -1010,6 +1030,7 @@ bool FMaterialResource::IsVolumetricPrimitive() const { return Material->Materia
 bool FMaterialResource::IsSpecialEngineMaterial() const { return Material->bUsedAsSpecialEngineMaterial; }
 bool FMaterialResource::HasVertexPositionOffsetConnected() const { return HasMaterialAttributesConnected() || (!Material->bUseMaterialAttributes && Material->WorldPositionOffset.IsConnected()); }
 bool FMaterialResource::HasPixelDepthOffsetConnected() const { return HasMaterialAttributesConnected() || (!Material->bUseMaterialAttributes && Material->PixelDepthOffset.IsConnected()); }
+bool FMaterialResource::HasShadingModelFromMaterialExpressionConnected() const { return HasMaterialAttributesConnected() || (!Material->bUseMaterialAttributes && Material->ShadingModelFromMaterialExpression.IsConnected()); }
 bool FMaterialResource::HasMaterialAttributesConnected() const { return Material->bUseMaterialAttributes && Material->MaterialAttributes.IsConnected(); }
 FString FMaterialResource::GetBaseMaterialPathName() const { return Material->GetPathName(); }
 FString FMaterialResource::GetDebugName() const
@@ -1196,9 +1217,9 @@ ERefractionMode FMaterialResource::GetRefractionMode() const
 	return Material->RefractionMode;
 }
 
-EMaterialShadingModel FMaterialResource::GetShadingModel() const 
+FMaterialShadingModelField FMaterialResource::GetShadingModels() const 
 {
-	return MaterialInstance ? MaterialInstance->GetShadingModel() : Material->GetShadingModel();
+	return MaterialInstance ? MaterialInstance->GetShadingModels() : Material->GetShadingModels();
 }
 
 bool FMaterialResource::IsTwoSided() const 
@@ -1561,22 +1582,78 @@ void FMaterial::SetupMaterialEnvironment(
 			OutEnvironment.SetDefine(TEXT("MATERIAL_DOMAIN_SURFACE"),TEXT("1"));
 	};
 
-	switch(GetShadingModel())
-	{
-		case MSM_Unlit:				OutEnvironment.SetDefine(TEXT("MATERIAL_SHADINGMODEL_UNLIT"),				TEXT("1")); break;
-		case MSM_DefaultLit:		OutEnvironment.SetDefine(TEXT("MATERIAL_SHADINGMODEL_DEFAULT_LIT"),			TEXT("1")); break;
-		case MSM_Subsurface:		OutEnvironment.SetDefine(TEXT("MATERIAL_SHADINGMODEL_SUBSURFACE"),			TEXT("1")); break;
-		case MSM_PreintegratedSkin: OutEnvironment.SetDefine(TEXT("MATERIAL_SHADINGMODEL_PREINTEGRATED_SKIN"),	TEXT("1")); break;
-		case MSM_SubsurfaceProfile: OutEnvironment.SetDefine(TEXT("MATERIAL_SHADINGMODEL_SUBSURFACE_PROFILE"),	TEXT("1")); break;
-		case MSM_ClearCoat:			OutEnvironment.SetDefine(TEXT("MATERIAL_SHADINGMODEL_CLEAR_COAT"),			TEXT("1")); break;
-		case MSM_TwoSidedFoliage:	OutEnvironment.SetDefine(TEXT("MATERIAL_SHADINGMODEL_TWOSIDED_FOLIAGE"),	TEXT("1")); break;
-		case MSM_Hair:				OutEnvironment.SetDefine(TEXT("MATERIAL_SHADINGMODEL_HAIR"),				TEXT("1")); break;
-		case MSM_Cloth:				OutEnvironment.SetDefine(TEXT("MATERIAL_SHADINGMODEL_CLOTH"),				TEXT("1")); break;
-		case MSM_Eye:				OutEnvironment.SetDefine(TEXT("MATERIAL_SHADINGMODEL_EYE"),					TEXT("1")); break;
-		default:
-			UE_LOG(LogMaterial, Warning, TEXT("Unknown material shading model: %u  Setting to MSM_DefaultLit"),(int32)GetShadingModel());
+	// Set all the shading models for this material here 
+	FMaterialShadingModelField ShadingModels = GetShadingModels();
+	ensure(ShadingModels.IsValid());
+
+	if (ShadingModels.IsLit())
+	{	
+		int NumSetMaterials = 0;
+		if (ShadingModels.HasShadingModel(MSM_DefaultLit))
+		{
+			OutEnvironment.SetDefine(TEXT("MATERIAL_SHADINGMODEL_DEFAULT_LIT"), TEXT("1"));
+			NumSetMaterials++;
+		}
+		if (ShadingModels.HasShadingModel(MSM_Subsurface))
+		{
+			OutEnvironment.SetDefine(TEXT("MATERIAL_SHADINGMODEL_SUBSURFACE"), TEXT("1"));
+			NumSetMaterials++;
+		}
+		if (ShadingModels.HasShadingModel(MSM_PreintegratedSkin))
+		{
+			OutEnvironment.SetDefine(TEXT("MATERIAL_SHADINGMODEL_PREINTEGRATED_SKIN"), TEXT("1"));
+			NumSetMaterials++;
+		}
+		if (ShadingModels.HasShadingModel(MSM_SubsurfaceProfile))
+		{
+			OutEnvironment.SetDefine(TEXT("MATERIAL_SHADINGMODEL_SUBSURFACE_PROFILE"), TEXT("1"));
+			NumSetMaterials++;
+		}
+		if (ShadingModels.HasShadingModel(MSM_ClearCoat))
+		{
+			OutEnvironment.SetDefine(TEXT("MATERIAL_SHADINGMODEL_CLEAR_COAT"), TEXT("1"));
+			NumSetMaterials++;
+		}
+		if (ShadingModels.HasShadingModel(MSM_TwoSidedFoliage))
+		{
+			OutEnvironment.SetDefine(TEXT("MATERIAL_SHADINGMODEL_TWOSIDED_FOLIAGE"), TEXT("1"));
+			NumSetMaterials++;
+		}
+		if (ShadingModels.HasShadingModel(MSM_Hair))
+		{
+			OutEnvironment.SetDefine(TEXT("MATERIAL_SHADINGMODEL_HAIR"), TEXT("1"));
+			NumSetMaterials++;
+		}
+		if (ShadingModels.HasShadingModel(MSM_Cloth))
+		{
+			OutEnvironment.SetDefine(TEXT("MATERIAL_SHADINGMODEL_CLOTH"), TEXT("1"));
+			NumSetMaterials++;
+		}
+		if (ShadingModels.HasShadingModel(MSM_Eye))
+		{
+			OutEnvironment.SetDefine(TEXT("MATERIAL_SHADINGMODEL_EYE"), TEXT("1"));
+			NumSetMaterials++;
+		}
+
+		if (NumSetMaterials == 1)
+		{
+			OutEnvironment.SetDefine(TEXT("MATERIAL_SINGLE_SHADINGMODEL"), TEXT("1"));
+		}
+
+		ensure(NumSetMaterials != 0);
+		if (NumSetMaterials == 0)
+		{
+			// Should not really end up here
+			UE_LOG(LogMaterial, Warning, TEXT("Unknown material shading model(s). Setting to MSM_DefaultLit"));
 			OutEnvironment.SetDefine(TEXT("MATERIAL_SHADINGMODEL_DEFAULT_LIT"),TEXT("1"));
-	};
+		}
+	}
+	else
+	{
+		// Unlit shading model can only exist by itself
+		OutEnvironment.SetDefine(TEXT("MATERIAL_SINGLE_SHADINGMODEL"), TEXT("1"));
+		OutEnvironment.SetDefine(TEXT("MATERIAL_SHADINGMODEL_UNLIT"), TEXT("1"));
+	}
 
 	if (IsTranslucentBlendMode(GetBlendMode()))
 	{
@@ -2159,7 +2236,7 @@ void FMaterialRenderProxy::UpdateUniformExpressionCacheIfNeeded(ERHIFeatureLevel
 
 		// Don't cache uniform expressions if an entirely different FMaterialRenderProxy is going to be used for rendering
 		if (!FallbackMaterialRenderProxy)
-		{
+{
 			FMaterialRenderContext MaterialRenderContext(this, Material, nullptr);
 			MaterialRenderContext.bShowSelection = GIsEditor;
 			EvaluateUniformExpressions(UniformExpressionCache[InFeatureLevel], MaterialRenderContext);
@@ -2358,7 +2435,7 @@ FString FMaterialResource::GetMaterialUsageDescription() const
 	check(Material);
 	FString BaseDescription = FString::Printf(
 		TEXT("LightingModel=%s, BlendMode=%s, "),
-		*GetShadingModelString(GetShadingModel()), *GetBlendModeString(GetBlendMode()));
+		*GetShadingModelFieldString(GetShadingModels()), *GetBlendModeString(GetBlendMode()));
 
 	// this changed from ",SpecialEngine, TwoSided" to ",SpecialEngine=1, TwoSided=1, TSNormal=0, ..." to be more readable
 	BaseDescription += FString::Printf(
@@ -3075,7 +3152,7 @@ FMaterialAttributeDefintion::FMaterialAttributeDefintion(
 	, BlendFunction(InBlendFunction)
 	, bIsHidden(bInIsHidden)
 {
-	checkf(ValueType & MCT_Float, TEXT("Unsupported type, only Float1 through Float4 allowed."));
+	checkf(ValueType & MCT_Float || ValueType == MCT_ShadingModel , TEXT("Unsupported type, only Float1 through Float4 or MCT_ShadingModel are allowed."));
 }
 
 int32 FMaterialAttributeDefintion::CompileDefaultValue(FMaterialCompiler* Compiler)
@@ -3083,12 +3160,19 @@ int32 FMaterialAttributeDefintion::CompileDefaultValue(FMaterialCompiler* Compil
 	int32 Ret;
 
 	// TODO: Temporarily preserving hack from 4.13 to change default value for two-sided foliage model 
-	if (Property == MP_SubsurfaceColor && Compiler->GetMaterialShadingModel() == MSM_TwoSidedFoliage)
+	if (Property == MP_SubsurfaceColor && Compiler->GetMaterialShadingModels().HasShadingModel(MSM_TwoSidedFoliage))
 	{
 		check(ValueType == MCT_Float3);
 		return Compiler->Constant3(0, 0, 0);
 	}
 
+	if (Property == MP_ShadingModel)
+	{
+		check(ValueType == MCT_ShadingModel);
+		// Default to the first shading model of the material. If the material is using a single shading model selected through the dropdown, this is how it gets written to the shader as a constant (optimizing out all the dynamic branches)
+		return Compiler->ShadingModel(Compiler->GetMaterialShadingModels().GetFirstShadingModel());
+	}
+
 	if (TexCoordIndex == INDEX_NONE)
 	{
 		// Standard value type
@@ -3152,6 +3236,7 @@ void FMaterialAttributeDefinitionMap::InitializeAttributeMap()
 	Add(FGuid(0xE8EBD0AD, 0xB1654CBE, 0xB079C3A8, 0xB39B9F15), TEXT("AmbientOcclusion"),		MP_AmbientOcclusion,		MCT_Float,	FVector4(1,0,0,0),	SF_Pixel);
 	Add(FGuid(0xD0B0FA03, 0x14D74455, 0xA851BAC5, 0x81A0788B), TEXT("Refraction"),				MP_Refraction,				MCT_Float2,	FVector4(1,0,0,0),	SF_Pixel);
 	Add(FGuid(0x0AC97EC3, 0xE3D047BA, 0xB610167D, 0xC4D919FF), TEXT("PixelDepthOffset"),		MP_PixelDepthOffset,		MCT_Float,	FVector4(0,0,0,0),	SF_Pixel);
+	Add(FGuid(0xD9423FFF, 0xD77E4D82, 0x8FF9CF5E, 0x055D1255), TEXT("ShadingModel"),			MP_ShadingModel,			MCT_ShadingModel,	FVector4(0,0,0,0),	SF_Pixel, INDEX_NONE, false, &CompileShadingModelBlendFunction);
 
 	// Texture coordinates
 	Add(FGuid(0xD30EC284, 0xE13A4160, 0x87BB5230, 0x2ED115DC), TEXT("CustomizedUV0"), MP_CustomizedUVs0, MCT_Float2, FVector4(0,0,0,0), SF_Vertex, 0);
diff --git a/Engine/Source/Runtime/Engine/Private/PrimitiveDrawingUtils.cpp b/Engine/Source/Runtime/Engine/Private/PrimitiveDrawingUtils.cpp
--- a/Engine/Source/Runtime/Engine/Private/PrimitiveDrawingUtils.cpp
+++ b/Engine/Source/Runtime/Engine/Private/PrimitiveDrawingUtils.cpp
@@ -1437,7 +1437,7 @@ void ApplyViewModeOverrides(
 		{
 			int32 lodColorationIndex = FMath::Clamp((int32)Mesh.VisualizeLODIndex, 0, GEngine->LODColorationColors.Num() - 1);
 	
-			bool bLit = Mesh.MaterialRenderProxy->GetMaterial(FeatureLevel)->GetShadingModel() != MSM_Unlit;
+			bool bLit = Mesh.MaterialRenderProxy->GetMaterial(FeatureLevel)->GetShadingModels().IsLit();
 			const UMaterial* LODColorationMaterial = (bLit && EngineShowFlags.Lighting) ? GEngine->LevelColorationLitMaterial : GEngine->LevelColorationUnlitMaterial;
 
 			auto LODColorationMaterialInstance = new FColoredMaterialRenderProxy(
@@ -1455,7 +1455,7 @@ void ApplyViewModeOverrides(
 		{
 			int32 hlodColorationIndex = FMath::Clamp((int32)Mesh.VisualizeHLODIndex, 0, GEngine->HLODColorationColors.Num() - 1);
 
-			bool bLit = Mesh.MaterialRenderProxy->GetMaterial(FeatureLevel)->GetShadingModel() != MSM_Unlit;
+			bool bLit = Mesh.MaterialRenderProxy->GetMaterial(FeatureLevel)->GetShadingModels().IsLit();
 			const UMaterial* HLODColorationMaterial = (bLit && EngineShowFlags.Lighting) ? GEngine->LevelColorationLitMaterial : GEngine->LevelColorationUnlitMaterial;
 
 			auto HLODColorationMaterialInstance = new FColoredMaterialRenderProxy(
@@ -1470,7 +1470,7 @@ void ApplyViewModeOverrides(
 	else if (!EngineShowFlags.Materials)
 	{
 		// Don't render unlit translucency when in 'lighting only' viewmode.
-		if (Mesh.MaterialRenderProxy->GetMaterial(FeatureLevel)->GetShadingModel() != MSM_Unlit
+		if (Mesh.MaterialRenderProxy->GetMaterial(FeatureLevel)->GetShadingModels().IsLit()
 			// Don't render translucency in 'lighting only', since the viewmode works by overriding with an opaque material
 			// This would cause a mismatch of the material's blend mode with the primitive's view relevance,
 			// And make faint particles block the view
diff --git a/Engine/Source/Runtime/Engine/Private/StaticMeshRender.cpp b/Engine/Source/Runtime/Engine/Private/StaticMeshRender.cpp
--- a/Engine/Source/Runtime/Engine/Private/StaticMeshRender.cpp
+++ b/Engine/Source/Runtime/Engine/Private/StaticMeshRender.cpp
@@ -217,7 +217,7 @@ FStaticMeshSceneProxy::FStaticMeshSceneProxy(UStaticMeshComponent* InComponent,
 			// #dxr_todo warning this is a hack:
 			//  - The only reason we exclude unlit from any RT acceleration structure right now is to not hit the sky dome. 
 			//  - It works but this would need consideration (e.g. infinite RTAO would not work if the sky dome is hit, double reflection due sky light cubemap, reflection testing no intersection, etc.).
-			if (SectionInfo.Material->GetShadingModel() != MSM_Unlit)
+			if (SectionInfo.Material->GetShadingModels().IsLit())
 			{
 				bHasNonUnlitSections = true;
 			}
diff --git a/Engine/Source/Runtime/Engine/Public/MaterialCompiler.h b/Engine/Source/Runtime/Engine/Public/MaterialCompiler.h
--- a/Engine/Source/Runtime/Engine/Public/MaterialCompiler.h
+++ b/Engine/Source/Runtime/Engine/Public/MaterialCompiler.h
@@ -71,7 +71,7 @@ public:
 
 	virtual EShaderPlatform GetShaderPlatform() = 0;
 
-	virtual EMaterialShadingModel GetMaterialShadingModel() const = 0;
+	virtual FMaterialShadingModelField GetMaterialShadingModels() const = 0;
 
 	virtual EMaterialValueType GetParameterType(int32 Index) const = 0;
 
@@ -294,6 +294,8 @@ public:
 	virtual int32 AtmosphericLightVector() = 0;
 	virtual int32 AtmosphericLightColor() = 0;
 	virtual int32 CustomPrimitiveData(int32 OutputIndex) = 0;
+	virtual int32 ShadingModel(EMaterialShadingModel InSelectedShadingModel) = 0;
+
 	// The compiler can run in a different state and this affects caching of sub expression, Expressions are different (e.g. View.PrevWorldViewOrigin) when using previous frame's values
 	// If possible we should re-factor this to avoid having to deal with compiler state
 	virtual bool IsCurrentlyCompilingForPreviousFrame() const { return false; }
@@ -315,7 +317,7 @@ public:
 
 	// Simple pass through all other material operations unmodified.
 
-	virtual EMaterialShadingModel GetMaterialShadingModel() const { return Compiler->GetMaterialShadingModel();  }
+	virtual FMaterialShadingModelField GetMaterialShadingModels() const { return Compiler->GetMaterialShadingModels(); }
 	virtual EMaterialValueType GetParameterType(int32 Index) const { return Compiler->GetParameterType(Index); }
 	virtual FMaterialUniformExpression* GetParameterUniformExpression(int32 Index) const { return Compiler->GetParameterUniformExpression(Index); }
 	virtual void SetMaterialProperty(EMaterialProperty InProperty, EShaderFrequency OverrideShaderFrequency, bool bUsePreviousFrameTime) override { Compiler->SetMaterialProperty(InProperty, OverrideShaderFrequency, bUsePreviousFrameTime); }
@@ -555,6 +557,11 @@ public:
 		return Compiler->CustomPrimitiveData(OutputIndex);
 	}
 
+	virtual int32 ShadingModel(EMaterialShadingModel InSelectedShadingModel) override
+	{
+		return Compiler->ShadingModel(InSelectedShadingModel);
+	}
+
 	virtual int32 TextureCoordinateOffset() override
 	{
 		return Compiler->TextureCoordinateOffset();
diff --git a/Engine/Source/Runtime/Engine/Public/MaterialExpressionIO.h b/Engine/Source/Runtime/Engine/Public/MaterialExpressionIO.h
--- a/Engine/Source/Runtime/Engine/Public/MaterialExpressionIO.h
+++ b/Engine/Source/Runtime/Engine/Public/MaterialExpressionIO.h
@@ -210,6 +210,25 @@ struct TStructOpsTypeTraits<FScalarMaterialInput>
 	};
 };
 
+struct FShadingModelMaterialInput : FMaterialInput<uint32>
+{
+#if WITH_EDITOR
+	ENGINE_API int32 CompileWithDefault(class FMaterialCompiler* Compiler, EMaterialProperty Property);
+#endif  // WITH_EDITOR
+	/** ICPPStructOps interface */
+	ENGINE_API bool Serialize(FArchive& Ar);
+};
+
+template<>
+struct TStructOpsTypeTraits<FShadingModelMaterialInput>
+	: public TStructOpsTypeTraitsBase2<FShadingModelMaterialInput>
+{
+	enum
+	{
+		WithSerializer = true,
+	};
+};
+
 struct FVectorMaterialInput : FMaterialInput<FVector>
 {
 #if WITH_EDITOR
diff --git a/Engine/Source/Runtime/Engine/Public/MaterialShaderType.h b/Engine/Source/Runtime/Engine/Public/MaterialShaderType.h
--- a/Engine/Source/Runtime/Engine/Public/MaterialShaderType.h
+++ b/Engine/Source/Runtime/Engine/Public/MaterialShaderType.h
@@ -30,6 +30,12 @@ class FVertexFactoryType;
 /** Converts an EMaterialShadingModel to a string description. */
 extern ENGINE_API FString GetShadingModelString(EMaterialShadingModel ShadingModel);
 
+/** Converts an FMaterialShadingModelField to a string description, base on the passed in function. */
+extern ENGINE_API FString GetShadingModelFieldString(FMaterialShadingModelField ShadingModels, FString(*ShadingModelToString)(EMaterialShadingModel), const FString& Delimiter = " ");
+
+/** Converts an FMaterialShadingModelField to a string description, base on a default function. */
+extern ENGINE_API FString GetShadingModelFieldString(FMaterialShadingModelField ShadingModels);
+
 /** Converts an EBlendMode to a string description. */
 extern ENGINE_API FString GetBlendModeString(EBlendMode BlendMode);
 
diff --git a/Engine/Source/Runtime/Engine/Public/MaterialShared.h b/Engine/Source/Runtime/Engine/Public/MaterialShared.h
--- a/Engine/Source/Runtime/Engine/Public/MaterialShared.h
+++ b/Engine/Source/Runtime/Engine/Public/MaterialShared.h
@@ -70,16 +70,16 @@ DECLARE_LOG_CATEGORY_EXTERN(LogMaterial,Log,Verbose);
 extern void GetMaterialQualityLevelName(EMaterialQualityLevel::Type InMaterialQualityLevel, FString& OutName);
 extern FName GetMaterialQualityLevelFName(EMaterialQualityLevel::Type InMaterialQualityLevel);
 
-inline bool IsSubsurfaceShadingModel(EMaterialShadingModel ShadingModel)
+inline bool IsSubsurfaceShadingModel(FMaterialShadingModelField ShadingModel)
 {
-	return ShadingModel == MSM_Subsurface || ShadingModel == MSM_PreintegratedSkin ||
-		ShadingModel == MSM_SubsurfaceProfile || ShadingModel == MSM_TwoSidedFoliage ||
-		ShadingModel == MSM_Cloth || ShadingModel == MSM_Eye;
+	return ShadingModel.HasShadingModel(MSM_Subsurface) || ShadingModel.HasShadingModel(MSM_PreintegratedSkin) ||
+		ShadingModel.HasShadingModel(MSM_SubsurfaceProfile) || ShadingModel.HasShadingModel(MSM_TwoSidedFoliage) ||
+		ShadingModel.HasShadingModel(MSM_Cloth) || ShadingModel.HasShadingModel(MSM_Eye);
 }
 
-inline bool UseSubsurfaceProfile(EMaterialShadingModel ShadingModel)
+inline bool UseSubsurfaceProfile(FMaterialShadingModelField ShadingModel)
 {
-	return ShadingModel == MSM_SubsurfaceProfile || ShadingModel == MSM_Eye;
+	return ShadingModel.HasShadingModel(MSM_SubsurfaceProfile) || ShadingModel.HasShadingModel(MSM_Eye);
 }
 
 inline uint32 GetUseSubsurfaceProfileShadingModelMask()
@@ -118,6 +118,7 @@ enum EMaterialValueType
 	MCT_MaterialAttributes	= 512,
 	MCT_TextureExternal = 1024,
 	MCT_Texture		= 16|32|64|1024,
+	MCT_ShadingModel = 2048,
 };
 
 /**
@@ -485,6 +486,7 @@ public:
 
 	/** true if the material uses the Velocity SceneTexture lookup */
 	bool bUsesVelocitySceneTexture;
+
 	/** true if the material uses distance cull fade */
 	bool bUsesDistanceCullFade;
 };
@@ -1302,7 +1304,7 @@ public:
 	virtual bool AllowNegativeEmissiveColor() const { return false; }
 	virtual enum EBlendMode GetBlendMode() const = 0;
 	ENGINE_API virtual enum ERefractionMode GetRefractionMode() const;
-	virtual enum EMaterialShadingModel GetShadingModel() const = 0;
+	virtual FMaterialShadingModelField GetShadingModels() const = 0;
 	virtual enum ETranslucencyLightingMode GetTranslucencyLightingMode() const { return TLM_VolumetricNonDirectional; };
 	virtual float GetOpacityMaskClipValue() const = 0;
 	virtual bool GetCastDynamicShadowAsMasked() const = 0;
@@ -1324,6 +1326,7 @@ public:
 	virtual FString GetFriendlyName() const = 0;
 	virtual bool HasVertexPositionOffsetConnected() const { return false; }
 	virtual bool HasPixelDepthOffsetConnected() const { return false; }
+	virtual bool HasShadingModelFromMaterialExpressionConnected() const { return false; }
 	virtual bool HasMaterialAttributesConnected() const { return false; }
 	virtual uint32 GetDecalBlendMode() const { return 0; }
 	virtual uint32 GetMaterialDecalResponse() const { return 0; }
@@ -1377,7 +1380,7 @@ public:
 	/** Returns whether this material should be considered for casting dynamic shadows. */
 	inline bool ShouldCastDynamicShadows() const
 	{
-		return GetShadingModel() != MSM_Unlit &&
+		return GetShadingModels().IsLit() &&
 			(GetBlendMode() == BLEND_Opaque ||
 			 GetBlendMode() == BLEND_Masked ||
 			(GetBlendMode() == BLEND_Translucent && GetCastDynamicShadowAsMasked()));
@@ -1990,7 +1993,7 @@ public:
 	ENGINE_API virtual uint32 GetMaterialDecalResponse() const override;
 	ENGINE_API virtual bool HasNormalConnected() const override;
 	ENGINE_API virtual bool HasEmissiveColorConnected() const override;
-	ENGINE_API virtual enum EMaterialShadingModel GetShadingModel() const override;
+	ENGINE_API virtual FMaterialShadingModelField GetShadingModels() const override;
 	ENGINE_API virtual enum ETranslucencyLightingMode GetTranslucencyLightingMode() const override;
 	ENGINE_API virtual float GetOpacityMaskClipValue() const override;
 	ENGINE_API virtual bool GetCastDynamicShadowAsMasked() const override;
@@ -2058,6 +2061,7 @@ protected:
 	ENGINE_API virtual void GatherExpressionsForCustomInterpolators(TArray<class UMaterialExpression*>& OutExpressions) const override;
 	ENGINE_API virtual bool HasVertexPositionOffsetConnected() const override;
 	ENGINE_API virtual bool HasPixelDepthOffsetConnected() const override;
+	ENGINE_API virtual bool HasShadingModelFromMaterialExpressionConnected() const override;
 	ENGINE_API virtual bool HasMaterialAttributesConnected() const override;
 	/** Useful for debugging. */
 	ENGINE_API virtual FString GetBaseMaterialPathName() const override;
diff --git a/Engine/Source/Runtime/Engine/Public/MeshBatch.h b/Engine/Source/Runtime/Engine/Public/MeshBatch.h
--- a/Engine/Source/Runtime/Engine/Public/MeshBatch.h
+++ b/Engine/Source/Runtime/Engine/Public/MeshBatch.h
@@ -225,7 +225,7 @@ struct FMeshBatch
 	FORCEINLINE bool CastsDeepShadow(/*ERHIFeatureLevel::Type InFeatureLevel*/) const
 	{
 		const FMaterial* Mat = MaterialRenderProxy->GetMaterial(ERHIFeatureLevel::SM5);
-		return Mat->GetShadingModel() == EMaterialShadingModel::MSM_Hair;
+		return Mat->GetShadingModels().HasOnlyShadingModel(EMaterialShadingModel::MSM_Hair);
 	}
 
 	FORCEINLINE bool IsMasked(ERHIFeatureLevel::Type InFeatureLevel) const
diff --git a/Engine/Source/Runtime/Engine/Public/SceneTypes.h b/Engine/Source/Runtime/Engine/Public/SceneTypes.h
--- a/Engine/Source/Runtime/Engine/Public/SceneTypes.h
+++ b/Engine/Source/Runtime/Engine/Public/SceneTypes.h
@@ -175,6 +175,7 @@ enum EMaterialProperty
 	MP_CustomizedUVs6 UMETA(Hidden),
 	MP_CustomizedUVs7 UMETA(Hidden),
 	MP_PixelDepthOffset UMETA(Hidden),
+	MP_ShadingModel UMETA(Hidden),
 
 	//^^^ New material properties go above here ^^^^
 	MP_MaterialAttributes UMETA(Hidden),
diff --git a/Engine/Source/Runtime/Renderer/Private/BasePassRendering.cpp b/Engine/Source/Runtime/Renderer/Private/BasePassRendering.cpp
--- a/Engine/Source/Runtime/Renderer/Private/BasePassRendering.cpp
+++ b/Engine/Source/Runtime/Renderer/Private/BasePassRendering.cpp
@@ -907,7 +907,7 @@ void FBasePassMeshProcessor::Process(
 	const FMaterialRenderProxy& RESTRICT MaterialRenderProxy,
 	const FMaterial& RESTRICT MaterialResource,
 	EBlendMode BlendMode,
-	EMaterialShadingModel ShadingModel,
+	FMaterialShadingModelField ShadingModels,
 	const LightMapPolicyType& RESTRICT LightMapPolicy,
 	const typename LightMapPolicyType::ElementDataType& RESTRICT LightMapElementData,
 	ERasterizerFillMode MeshFillMode,
@@ -915,7 +915,7 @@ void FBasePassMeshProcessor::Process(
 {
 	const FVertexFactory* VertexFactory = MeshBatch.VertexFactory;
 
-	const bool bRenderSkylight = Scene && Scene->ShouldRenderSkylightInBasePass(BlendMode) && ShadingModel != MSM_Unlit;
+	const bool bRenderSkylight = Scene && Scene->ShouldRenderSkylightInBasePass(BlendMode) && ShadingModels.IsLit();
 	const bool bRenderAtmosphericFog = IsTranslucentBlendMode(BlendMode) && (Scene && Scene->HasAtmosphericFog() && Scene->ReadOnlyCVARCache.bEnableAtmosphericFog);
 
 	TMeshProcessorShaders<
@@ -994,7 +994,7 @@ void FBasePassMeshProcessor::AddMeshBatchForSimpleForwardShading(
 	ERasterizerCullMode MeshCullMode)
 {
 	const EBlendMode BlendMode = Material.GetBlendMode();
-	const EMaterialShadingModel ShadingModel = Material.GetShadingModel();
+	const FMaterialShadingModelField ShadingModels = Material.GetShadingModels();
 
 	if (bAllowStaticLighting && LightMapInteraction.GetType() == LMIT_Texture)
 	{
@@ -1012,7 +1012,7 @@ void FBasePassMeshProcessor::AddMeshBatchForSimpleForwardShading(
 				MaterialRenderProxy,
 				Material,
 				BlendMode,
-				ShadingModel,
+				ShadingModels,
 				FUniformLightMapPolicy(LMP_SIMPLE_STATIONARY_PRECOMPUTED_SHADOW_LIGHTING),
 				MeshBatch.LCI,
 				MeshFillMode,
@@ -1028,7 +1028,7 @@ void FBasePassMeshProcessor::AddMeshBatchForSimpleForwardShading(
 				MaterialRenderProxy,
 				Material,
 				BlendMode,
-				ShadingModel,
+				ShadingModels,
 				FUniformLightMapPolicy(LMP_SIMPLE_LIGHTMAP_ONLY_LIGHTING),
 				MeshBatch.LCI,
 				MeshFillMode,
@@ -1048,7 +1048,7 @@ void FBasePassMeshProcessor::AddMeshBatchForSimpleForwardShading(
 			MaterialRenderProxy,
 			Material,
 			BlendMode,
-			ShadingModel,
+			ShadingModels,
 			FUniformLightMapPolicy(LMP_SIMPLE_STATIONARY_VOLUMETRICLIGHTMAP_SHADOW_LIGHTING),
 			MeshBatch.LCI,
 			MeshFillMode,
@@ -1080,7 +1080,7 @@ void FBasePassMeshProcessor::AddMeshBatchForSimpleForwardShading(
 				MaterialRenderProxy,
 				Material,
 				BlendMode,
-				ShadingModel,
+				ShadingModels,
 				FUniformLightMapPolicy(LMP_SIMPLE_STATIONARY_SINGLESAMPLE_SHADOW_LIGHTING),
 				MeshBatch.LCI,
 				MeshFillMode,
@@ -1096,7 +1096,7 @@ void FBasePassMeshProcessor::AddMeshBatchForSimpleForwardShading(
 				MaterialRenderProxy,
 				Material,
 				BlendMode,
-				ShadingModel,
+				ShadingModels,
 				FUniformLightMapPolicy(LMP_SIMPLE_NO_LIGHTMAP),
 				MeshBatch.LCI,
 				MeshFillMode,
@@ -1114,7 +1114,7 @@ void FBasePassMeshProcessor::AddMeshBatchForSimpleForwardShading(
 			MaterialRenderProxy,
 			Material,
 			BlendMode,
-			ShadingModel,
+			ShadingModels,
 			FUniformLightMapPolicy(LMP_SIMPLE_DIRECTIONAL_LIGHT_LIGHTING),
 			MeshBatch.LCI,
 			MeshFillMode,
@@ -1130,7 +1130,7 @@ void FBasePassMeshProcessor::AddMeshBatchForSimpleForwardShading(
 			MaterialRenderProxy,
 			Material,
 			BlendMode,
-			ShadingModel,
+			ShadingModels,
 			FUniformLightMapPolicy(LMP_SIMPLE_NO_LIGHTMAP),
 			MeshBatch.LCI,
 			MeshFillMode,
@@ -1149,7 +1149,7 @@ void FBasePassMeshProcessor::AddMeshBatch(const FMeshBatch& RESTRICT MeshBatch,
 		const FMaterialRenderProxy& MaterialRenderProxy = FallbackMaterialRenderProxyPtr ? *FallbackMaterialRenderProxyPtr : *MeshBatch.MaterialRenderProxy;
 
 		const EBlendMode BlendMode = Material.GetBlendMode();
-		const EMaterialShadingModel ShadingModel = Material.GetShadingModel();
+		const FMaterialShadingModelField ShadingModels = Material.GetShadingModels();
 		const bool bIsTranslucent = IsTranslucentBlendMode(BlendMode);
 		const ERasterizerFillMode MeshFillMode = ComputeMeshFillMode(MeshBatch, Material);
 		const ERasterizerCullMode MeshCullMode = ComputeMeshCullMode(MeshBatch, Material);
@@ -1189,7 +1189,7 @@ void FBasePassMeshProcessor::AddMeshBatch(const FMeshBatch& RESTRICT MeshBatch,
 			&& ShouldIncludeDomainInMeshPass(Material.GetMaterialDomain()))
 		{
 			// Check for a cached light-map.
-			const bool bIsLitMaterial = (ShadingModel != MSM_Unlit);
+			const bool bIsLitMaterial = ShadingModels.IsLit();
 			static const auto AllowStaticLightingVar = IConsoleManager::Get().FindTConsoleVariableDataInt(TEXT("r.AllowStaticLighting"));
 			const bool bAllowStaticLighting = (!AllowStaticLightingVar || AllowStaticLightingVar->GetValueOnRenderThread() != 0);
 
@@ -1256,7 +1256,7 @@ void FBasePassMeshProcessor::AddMeshBatch(const FMeshBatch& RESTRICT MeshBatch,
 						MaterialRenderProxy,
 						Material,
 						BlendMode,
-						ShadingModel,
+						ShadingModels,
 						FSelfShadowedVolumetricLightmapPolicy(),
 						ElementData,
 						MeshFillMode,
@@ -1275,7 +1275,7 @@ void FBasePassMeshProcessor::AddMeshBatch(const FMeshBatch& RESTRICT MeshBatch,
 						MaterialRenderProxy,
 						Material,
 						BlendMode,
-						ShadingModel,
+						ShadingModels,
 						FSelfShadowedCachedPointIndirectLightingPolicy(),
 						ElementData,
 						MeshFillMode,
@@ -1291,7 +1291,7 @@ void FBasePassMeshProcessor::AddMeshBatch(const FMeshBatch& RESTRICT MeshBatch,
 						MaterialRenderProxy,
 						Material,
 						BlendMode,
-						ShadingModel,
+						ShadingModels,
 						FSelfShadowedTranslucencyPolicy(),
 						ElementData.SelfShadowTranslucencyUniformBuffer,
 						MeshFillMode,
@@ -1322,7 +1322,7 @@ void FBasePassMeshProcessor::AddMeshBatch(const FMeshBatch& RESTRICT MeshBatch,
 								MaterialRenderProxy,
 								Material,
 								BlendMode,
-								ShadingModel,
+								ShadingModels,
 								FUniformLightMapPolicy(LMP_DISTANCE_FIELD_SHADOWS_AND_HQ_LIGHTMAP),
 								MeshBatch.LCI,
 								MeshFillMode,
@@ -1338,7 +1338,7 @@ void FBasePassMeshProcessor::AddMeshBatch(const FMeshBatch& RESTRICT MeshBatch,
 								MaterialRenderProxy,
 								Material,
 								BlendMode,
-								ShadingModel,
+								ShadingModels,
 								FUniformLightMapPolicy(LMP_HQ_LIGHTMAP),
 								MeshBatch.LCI,
 								MeshFillMode,
@@ -1355,7 +1355,7 @@ void FBasePassMeshProcessor::AddMeshBatch(const FMeshBatch& RESTRICT MeshBatch,
 							MaterialRenderProxy,
 							Material,
 							BlendMode,
-							ShadingModel,
+							ShadingModels,
 							FUniformLightMapPolicy(LMP_LQ_LIGHTMAP),
 							MeshBatch.LCI,
 							MeshFillMode,
@@ -1371,7 +1371,7 @@ void FBasePassMeshProcessor::AddMeshBatch(const FMeshBatch& RESTRICT MeshBatch,
 							MaterialRenderProxy,
 							Material,
 							BlendMode,
-							ShadingModel,
+							ShadingModels,
 							FUniformLightMapPolicy(LMP_NO_LIGHTMAP),
 							MeshBatch.LCI,
 							MeshFillMode,
@@ -1396,7 +1396,7 @@ void FBasePassMeshProcessor::AddMeshBatch(const FMeshBatch& RESTRICT MeshBatch,
 							MaterialRenderProxy,
 							Material,
 							BlendMode,
-							ShadingModel,
+							ShadingModels,
 							FUniformLightMapPolicy(LMP_PRECOMPUTED_IRRADIANCE_VOLUME_INDIRECT_LIGHTING),
 							MeshBatch.LCI,
 							MeshFillMode,
@@ -1435,7 +1435,7 @@ void FBasePassMeshProcessor::AddMeshBatch(const FMeshBatch& RESTRICT MeshBatch,
 									MaterialRenderProxy,
 									Material,
 									BlendMode,
-									ShadingModel,
+									ShadingModels,
 									FUniformLightMapPolicy(LMP_CACHED_VOLUME_INDIRECT_LIGHTING),
 									MeshBatch.LCI,
 									MeshFillMode,
@@ -1452,7 +1452,7 @@ void FBasePassMeshProcessor::AddMeshBatch(const FMeshBatch& RESTRICT MeshBatch,
 									MaterialRenderProxy,
 									Material,
 									BlendMode,
-									ShadingModel,
+									ShadingModels,
 									FUniformLightMapPolicy(LMP_CACHED_POINT_INDIRECT_LIGHTING),
 									MeshBatch.LCI,
 									MeshFillMode,
@@ -1469,7 +1469,7 @@ void FBasePassMeshProcessor::AddMeshBatch(const FMeshBatch& RESTRICT MeshBatch,
 								MaterialRenderProxy,
 								Material,
 								BlendMode,
-								ShadingModel,
+								ShadingModels,
 								FUniformLightMapPolicy(LMP_NO_LIGHTMAP),
 								MeshBatch.LCI,
 								MeshFillMode,
@@ -1486,7 +1486,7 @@ void FBasePassMeshProcessor::AddMeshBatch(const FMeshBatch& RESTRICT MeshBatch,
 							MaterialRenderProxy,
 							Material,
 							BlendMode,
-							ShadingModel,
+							ShadingModels,
 							FUniformLightMapPolicy(LMP_NO_LIGHTMAP),
 							MeshBatch.LCI,
 							MeshFillMode,
diff --git a/Engine/Source/Runtime/Renderer/Private/BasePassRendering.h b/Engine/Source/Runtime/Renderer/Private/BasePassRendering.h
--- a/Engine/Source/Runtime/Renderer/Private/BasePassRendering.h
+++ b/Engine/Source/Runtime/Renderer/Private/BasePassRendering.h
@@ -495,7 +495,7 @@ public:
 			|| bTranslucent
 			// Some lightmap policies (eg Simple Forward) always require skylight support
 			|| LightMapPolicyType::RequiresSkylight()
-			|| ((bProjectSupportsStationarySkylight || IsForwardShadingEnabled(Platform)) && (Material->GetShadingModel() != MSM_Unlit));
+			|| ((bProjectSupportsStationarySkylight || IsForwardShadingEnabled(Platform)) && Material->GetShadingModels().IsLit());
 		return bCacheShaders
 			&& (IsFeatureLevelSupported(Platform, ERHIFeatureLevel::SM4))
 			&& TBasePassPixelShaderBaseType<LightMapPolicyType>::ShouldCompilePermutation(Platform, Material, VertexFactoryType);
@@ -640,7 +640,7 @@ private:
 		const FMaterialRenderProxy& RESTRICT MaterialRenderProxy,
 		const FMaterial& RESTRICT MaterialResource,
 		EBlendMode BlendMode,
-		EMaterialShadingModel ShadingModel,
+		FMaterialShadingModelField ShadingModels,
 		const LightMapPolicyType& RESTRICT LightMapPolicy,
 		const typename LightMapPolicyType::ElementDataType& RESTRICT LightMapElementData,
 		ERasterizerFillMode MeshFillMode,
diff --git a/Engine/Source/Runtime/Renderer/Private/ConvertToUniformMesh.cpp b/Engine/Source/Runtime/Renderer/Private/ConvertToUniformMesh.cpp
--- a/Engine/Source/Runtime/Renderer/Private/ConvertToUniformMesh.cpp
+++ b/Engine/Source/Runtime/Renderer/Private/ConvertToUniformMesh.cpp
@@ -212,7 +212,7 @@ bool ShouldGenerateSurfelsOnMesh(const FMeshBatch& Mesh, ERHIFeatureLevel::Type
 	//@todo - support for tessellated meshes
 	return Mesh.Type == PT_TriangleList 
 		&& !Mesh.IsTranslucent(FeatureLevel) 
-		&& Mesh.MaterialRenderProxy->GetMaterial(FeatureLevel)->GetShadingModel() != MSM_Unlit;
+		&& Mesh.MaterialRenderProxy->GetMaterial(FeatureLevel)->GetShadingModels().IsLit();
 }
 
 bool ShouldConvertMesh(const FMeshBatch& Mesh)
diff --git a/Engine/Source/Runtime/Renderer/Private/DebugViewModeRendering.cpp b/Engine/Source/Runtime/Renderer/Private/DebugViewModeRendering.cpp
--- a/Engine/Source/Runtime/Renderer/Private/DebugViewModeRendering.cpp
+++ b/Engine/Source/Runtime/Renderer/Private/DebugViewModeRendering.cpp
@@ -376,7 +376,7 @@ void FDebugViewModeMeshProcessor::UpdateInstructionCount(FDebugViewModeShaderEle
 	check(InBatchMaterial && InVertexFactoryType);
 
 	const bool bDeferred = !IsAnyForwardShadingEnabled(GetFeatureLevelShaderPlatform(InBatchMaterial->GetFeatureLevel()));
-	const bool bLit = InBatchMaterial->GetShadingModel() != MSM_Unlit;
+	const bool bLit = InBatchMaterial->GetShadingModels().IsLit();
 
 	OutShaderElementData.NumVSInstructions = InBatchMaterial->GetShader<TBasePassVS<TUniformLightMapPolicy<LMP_NO_LIGHTMAP>, false>>(InVertexFactoryType)->GetNumInstructions();
 	OutShaderElementData.NumPSInstructions = InBatchMaterial->GetShader<TBasePassPS<TUniformLightMapPolicy<LMP_NO_LIGHTMAP>, false>>(InVertexFactoryType)->GetNumInstructions();
diff --git a/Engine/Source/Runtime/Renderer/Private/LightMapDensityRendering.cpp b/Engine/Source/Runtime/Renderer/Private/LightMapDensityRendering.cpp
--- a/Engine/Source/Runtime/Renderer/Private/LightMapDensityRendering.cpp
+++ b/Engine/Source/Runtime/Renderer/Private/LightMapDensityRendering.cpp
@@ -240,7 +240,7 @@ void FLightmapDensityMeshProcessor::AddMeshBatch(const FMeshBatch& RESTRICT Mesh
 		const FMaterial* Material = &MeshBatch.MaterialRenderProxy->GetMaterialWithFallback(FeatureLevel, MaterialRenderProxy);
 		const bool bMaterialMasked = Material->IsMasked();
 		const bool bTranslucentBlendMode = IsTranslucentBlendMode(Material->GetBlendMode());
-		const bool bIsLitMaterial = Material->GetShadingModel() != MSM_Unlit;
+		const bool bIsLitMaterial = Material->GetShadingModels().IsLit();
 		const ERasterizerFillMode MeshFillMode = ComputeMeshFillMode(MeshBatch, *Material);
 		const ERasterizerCullMode MeshCullMode = ComputeMeshCullMode(MeshBatch, *Material);
 		const FLightMapInteraction LightMapInteraction = (MeshBatch.LCI && bIsLitMaterial) ? MeshBatch.LCI->GetLightMapInteraction(FeatureLevel) : FLightMapInteraction();
diff --git a/Engine/Source/Runtime/Renderer/Private/LightMapRendering.h b/Engine/Source/Runtime/Renderer/Private/LightMapRendering.h
--- a/Engine/Source/Runtime/Renderer/Private/LightMapRendering.h
+++ b/Engine/Source/Runtime/Renderer/Private/LightMapRendering.h
@@ -131,7 +131,7 @@ struct TLightMapPolicy
 		const bool bShouldCacheQuality = (LightmapQuality != ELightmapQuality::LQ_LIGHTMAP) || bProjectCanHaveLowQualityLightmaps;
 
 		// GetValueOnAnyThread() as it's possible that ShouldCache is called from rendering thread. That is to output some error message.
-		return (Material->GetShadingModel() != MSM_Unlit)
+		return (Material->GetShadingModels().IsLit())
 			&& bShouldCacheQuality
 			&& VertexFactoryType->SupportsStaticLighting() 
 			&& (!AllowStaticLightingVar || AllowStaticLightingVar->GetValueOnAnyThread() != 0)
@@ -167,7 +167,7 @@ public:
 
 	static bool ShouldCompilePermutation(EShaderPlatform Platform,const FMaterial* Material,const FVertexFactoryType* VertexFactoryType)
 	{
-		return Material->GetShadingModel() != MSM_Unlit && VertexFactoryType->SupportsStaticLighting();
+		return Material->GetShadingModels().IsLit() && VertexFactoryType->SupportsStaticLighting();
 	}
 };
 
@@ -205,7 +205,7 @@ public:
 
 	static bool ShouldCompilePermutation(EShaderPlatform Platform,const FMaterial* Material,const FVertexFactoryType* VertexFactoryType)
 	{
-		return Material->GetShadingModel() != MSM_Unlit && IsTranslucentBlendMode(Material->GetBlendMode()) && IsFeatureLevelSupported(Platform, ERHIFeatureLevel::SM4);
+		return Material->GetShadingModels().IsLit() && IsTranslucentBlendMode(Material->GetBlendMode()) && IsFeatureLevelSupported(Platform, ERHIFeatureLevel::SM4);
 	}
 
 	static void ModifyCompilationEnvironment(EShaderPlatform Platform, const FMaterial* Material, FShaderCompilerEnvironment& OutEnvironment)
@@ -251,7 +251,7 @@ struct FPrecomputedVolumetricLightmapLightingPolicy
 	{
 		static const auto AllowStaticLightingVar = IConsoleManager::Get().FindTConsoleVariableDataInt(TEXT("r.AllowStaticLighting"));
 	
-		return Material->GetShadingModel() != MSM_Unlit
+		return Material->GetShadingModels().IsLit()
 			&& (!AllowStaticLightingVar || AllowStaticLightingVar->GetValueOnAnyThread() != 0);
 	}
 
@@ -270,7 +270,7 @@ struct FCachedVolumeIndirectLightingPolicy
 	{
 		static const auto AllowStaticLightingVar = IConsoleManager::Get().FindTConsoleVariableDataInt(TEXT("r.AllowStaticLighting"));
 
-		return Material->GetShadingModel() != MSM_Unlit 
+		return Material->GetShadingModels().IsLit()
 			&& !IsTranslucentBlendMode(Material->GetBlendMode()) 
 			&& (!AllowStaticLightingVar || AllowStaticLightingVar->GetValueOnAnyThread() != 0)
 			&& IsFeatureLevelSupported(Platform, ERHIFeatureLevel::SM4);
@@ -297,7 +297,7 @@ struct FCachedPointIndirectLightingPolicy
 	{
 		static const auto AllowStaticLightingVar = IConsoleManager::Get().FindTConsoleVariableDataInt(TEXT("r.AllowStaticLighting"));
 	
-		return Material->GetShadingModel() != MSM_Unlit
+		return Material->GetShadingModels().IsLit()
 			&& (!AllowStaticLightingVar || AllowStaticLightingVar->GetValueOnAnyThread() != 0);
 	}
 
@@ -366,7 +366,7 @@ struct FSimpleDirectionalLightLightingPolicy
 	static bool ShouldCompilePermutation(EShaderPlatform Platform,const FMaterial* Material,const FVertexFactoryType* VertexFactoryType)
 	{
 		return PlatformSupportsSimpleForwardShading(Platform)
-			&& Material->GetShadingModel() != MSM_Unlit;
+			&& Material->GetShadingModels().IsLit();
 	}
 
 	static void ModifyCompilationEnvironment(EShaderPlatform Platform, const FMaterial* Material, FShaderCompilerEnvironment& OutEnvironment)
@@ -488,7 +488,7 @@ public:
 	{
 		static auto* CVarMobileEnableStaticAndCSMShadowReceivers = IConsoleManager::Get().FindTConsoleVariableDataInt(TEXT("r.Mobile.EnableStaticAndCSMShadowReceivers"));
 		const bool bMobileEnableStaticAndCSMShadowReceivers = CVarMobileEnableStaticAndCSMShadowReceivers->GetValueOnAnyThread() == 1;
-		return bMobileEnableStaticAndCSMShadowReceivers && (Material->GetShadingModel() != MSM_Unlit) && Super::ShouldCompilePermutation(Platform, Material, VertexFactoryType);
+		return bMobileEnableStaticAndCSMShadowReceivers && (Material->GetShadingModels().IsLit()) && Super::ShouldCompilePermutation(Platform, Material, VertexFactoryType);
 	}
 
 	static void ModifyCompilationEnvironment(EShaderPlatform Platform, const FMaterial* Material, FShaderCompilerEnvironment& OutEnvironment)
@@ -513,7 +513,7 @@ struct FMobileDirectionalLightAndSHIndirectPolicy
 		static auto* CVarAllowStaticLighting = IConsoleManager::Get().FindTConsoleVariableDataInt(TEXT("r.AllowStaticLighting"));
 		const bool bAllowStaticLighting = CVarAllowStaticLighting->GetValueOnAnyThread() != 0;
 
-		return bAllowStaticLighting && Material->GetShadingModel() != MSM_Unlit && FCachedPointIndirectLightingPolicy::ShouldCompilePermutation(Platform, Material, VertexFactoryType);
+		return bAllowStaticLighting && Material->GetShadingModels().IsLit() && FCachedPointIndirectLightingPolicy::ShouldCompilePermutation(Platform, Material, VertexFactoryType);
 	}
 
 	static void ModifyCompilationEnvironment(EShaderPlatform Platform, const FMaterial* Material, FShaderCompilerEnvironment& OutEnvironment)
@@ -592,7 +592,7 @@ struct FMobileMovableDirectionalLightLightingPolicy
 		static auto* CVarMobileAllowMovableDirectionalLights = IConsoleManager::Get().FindTConsoleVariableDataInt(TEXT("r.Mobile.AllowMovableDirectionalLights"));
 		const bool bMobileAllowMovableDirectionalLights = CVarMobileAllowMovableDirectionalLights->GetValueOnAnyThread() != 0;
 
-		return bMobileAllowMovableDirectionalLights && Material->GetShadingModel() != MSM_Unlit;
+		return bMobileAllowMovableDirectionalLights && Material->GetShadingModels().IsLit();
 	}
 
 	static void ModifyCompilationEnvironment(EShaderPlatform Platform, const FMaterial* Material, FShaderCompilerEnvironment& OutEnvironment)
@@ -614,7 +614,7 @@ struct FMobileMovableDirectionalLightCSMLightingPolicy
 		static auto* CVarMobileAllowMovableDirectionalLights = IConsoleManager::Get().FindTConsoleVariableDataInt(TEXT("r.Mobile.AllowMovableDirectionalLights"));
 		const bool bMobileAllowMovableDirectionalLights = CVarMobileAllowMovableDirectionalLights->GetValueOnAnyThread() != 0;
 
-		return bMobileAllowMovableDirectionalLights && Material->GetShadingModel() != MSM_Unlit;
+		return bMobileAllowMovableDirectionalLights && Material->GetShadingModels().IsLit();
 	}	
 
 	static void ModifyCompilationEnvironment(EShaderPlatform Platform, const FMaterial* Material, FShaderCompilerEnvironment& OutEnvironment)
@@ -645,7 +645,7 @@ struct FMobileMovableDirectionalLightWithLightmapPolicy : public TLightMapPolicy
 		static auto* CVarAllowStaticLighting = IConsoleManager::Get().FindTConsoleVariableDataInt(TEXT("r.AllowStaticLighting"));
 		const bool bAllowStaticLighting = (!CVarAllowStaticLighting || CVarAllowStaticLighting->GetValueOnAnyThread() != 0);
 
-		return bAllowStaticLighting && bMobileAllowMovableDirectionalLights && (Material->GetShadingModel() != MSM_Unlit) && Super::ShouldCompilePermutation(Platform, Material, VertexFactoryType);
+		return bAllowStaticLighting && bMobileAllowMovableDirectionalLights && (Material->GetShadingModels().IsLit()) && Super::ShouldCompilePermutation(Platform, Material, VertexFactoryType);
 	}
 
 	static void ModifyCompilationEnvironment(EShaderPlatform Platform, const FMaterial* Material, FShaderCompilerEnvironment& OutEnvironment)
@@ -1001,7 +1001,7 @@ public:
 	{
 		static IConsoleVariable* AllowStaticLightingVar = IConsoleManager::Get().FindConsoleVariable(TEXT("r.AllowStaticLighting"));
 
-		return Material->GetShadingModel() != MSM_Unlit 
+		return Material->GetShadingModels().IsLit()
 			&& IsTranslucentBlendMode(Material->GetBlendMode()) 
 			&& (!AllowStaticLightingVar || AllowStaticLightingVar->GetInt() != 0)
 			&& FSelfShadowedTranslucencyPolicy::ShouldCompilePermutation(Platform, Material, VertexFactoryType);
@@ -1050,7 +1050,7 @@ public:
 	{
 		static IConsoleVariable* AllowStaticLightingVar = IConsoleManager::Get().FindConsoleVariable(TEXT("r.AllowStaticLighting"));
 
-		return Material->GetShadingModel() != MSM_Unlit 
+		return Material->GetShadingModels().IsLit()
 			&& IsTranslucentBlendMode(Material->GetBlendMode()) 
 			&& (!AllowStaticLightingVar || AllowStaticLightingVar->GetInt() != 0)
 			&& FSelfShadowedTranslucencyPolicy::ShouldCompilePermutation(Platform, Material, VertexFactoryType);
diff --git a/Engine/Source/Runtime/Renderer/Private/MobileBasePass.cpp b/Engine/Source/Runtime/Renderer/Private/MobileBasePass.cpp
--- a/Engine/Source/Runtime/Renderer/Private/MobileBasePass.cpp
+++ b/Engine/Source/Runtime/Renderer/Private/MobileBasePass.cpp
@@ -112,7 +112,7 @@ void MobileBasePass::GetShaders(
 	TMobileBasePassVSPolicyParamType<FUniformLightMapPolicy>*& VertexShader,
 	TMobileBasePassPSPolicyParamType<FUniformLightMapPolicy>*& PixelShader)
 {
-	bool bIsLit = (MaterialResource.GetShadingModel() != MSM_Unlit);
+	bool bIsLit = (MaterialResource.GetShadingModels().IsLit());
 	if (bIsLit && !UseSkylightPermutation(bEnableSkyLight, FReadOnlyCVARCache::Get().MobileSkyLightPermutation))	
 	{
 		bEnableSkyLight = !bEnableSkyLight;
@@ -206,7 +206,7 @@ const FLightSceneInfo* MobileBasePass::GetDirectionalLightInfo(const FScene* Sce
 int32 MobileBasePass::CalcNumMovablePointLights(const FMaterial& InMaterial, const FPrimitiveSceneProxy* InPrimitiveSceneProxy)
 {
 	const FReadOnlyCVARCache& ReadOnlyCVARCache = FReadOnlyCVARCache::Get();
-	const bool bIsUnlit = InMaterial.GetShadingModel() == MSM_Unlit;
+	const bool bIsUnlit = InMaterial.GetShadingModels().IsUnlit();
 	int32 OutNumMovablePointLights = (InPrimitiveSceneProxy && !bIsUnlit) ? FMath::Min<int32>(InPrimitiveSceneProxy->GetPrimitiveSceneInfo()->NumMobileMovablePointLights, ReadOnlyCVARCache.NumMobileMovablePointLights) : 0;
 	if (OutNumMovablePointLights > 0 && ReadOnlyCVARCache.bMobileMovablePointLightsUseStaticBranch)
 	{
@@ -234,7 +234,7 @@ ELightMapPolicyType MobileBasePass::SelectMeshLightmapPolicy(
 	const FMeshBatch& Mesh, 
 	const FPrimitiveSceneProxy* PrimitiveSceneProxy,
 	const FLightSceneInfo* MobileDirectionalLight,
-	EMaterialShadingModel ShadingModel, 
+	FMaterialShadingModelField ShadingModels, 
 	bool bPrimReceivesCSM, 
 	ERHIFeatureLevel::Type FeatureLevel)
 {
@@ -242,7 +242,7 @@ ELightMapPolicyType MobileBasePass::SelectMeshLightmapPolicy(
 	ELightMapPolicyType SelectedLightmapPolicy = LMP_NO_LIGHTMAP;
 
 	// Check for a cached light-map.
-	const bool bIsLitMaterial = ShadingModel != MSM_Unlit;
+	const bool bIsLitMaterial = ShadingModels.IsLit();
 	if (bIsLitMaterial)
 	{
 		const FLightMapInteraction LightMapInteraction = (Mesh.LCI && bIsLitMaterial)
@@ -558,7 +558,7 @@ void FMobileBasePassMeshProcessor::AddMeshBatch(const FMeshBatch& RESTRICT MeshB
 	const FMaterial& Material = MeshBatch.MaterialRenderProxy->GetMaterialWithFallback(FeatureLevel, FallbackMaterialRenderProxyPtr);
 	const FMaterialRenderProxy& MaterialRenderProxy = FallbackMaterialRenderProxyPtr ? *FallbackMaterialRenderProxyPtr : *MeshBatch.MaterialRenderProxy;
 	const EBlendMode BlendMode = Material.GetBlendMode();
-	const EMaterialShadingModel ShadingModel = Material.GetShadingModel();
+	const FMaterialShadingModelField ShadingModels = Material.GetShadingModels();
 	const bool bIsTranslucent = IsTranslucentBlendMode(BlendMode);
 	
 	if (bTranslucentBasePass)
@@ -572,8 +572,8 @@ void FMobileBasePassMeshProcessor::AddMeshBatch(const FMeshBatch& RESTRICT MeshB
 		{
 			check(bCanReceiveCSM == false);
 			const FLightSceneInfo* MobileDirectionalLight = MobileBasePass::GetDirectionalLightInfo(Scene, PrimitiveSceneProxy);
-			ELightMapPolicyType LightmapPolicyType = MobileBasePass::SelectMeshLightmapPolicy(Scene, MeshBatch, PrimitiveSceneProxy, MobileDirectionalLight, ShadingModel, bCanReceiveCSM, FeatureLevel);
-			Process(MeshBatch, BatchElementMask, StaticMeshId, PrimitiveSceneProxy, MaterialRenderProxy, Material, BlendMode, ShadingModel, LightmapPolicyType, MeshBatch.LCI);
+			ELightMapPolicyType LightmapPolicyType = MobileBasePass::SelectMeshLightmapPolicy(Scene, MeshBatch, PrimitiveSceneProxy, MobileDirectionalLight, ShadingModels, bCanReceiveCSM, FeatureLevel);
+			Process(MeshBatch, BatchElementMask, StaticMeshId, PrimitiveSceneProxy, MaterialRenderProxy, Material, BlendMode, ShadingModels, LightmapPolicyType, MeshBatch.LCI);
 		}
 	}
 	else
@@ -582,8 +582,8 @@ void FMobileBasePassMeshProcessor::AddMeshBatch(const FMeshBatch& RESTRICT MeshB
 		if (!bIsTranslucent)
 		{
 			const FLightSceneInfo* MobileDirectionalLight = MobileBasePass::GetDirectionalLightInfo(Scene, PrimitiveSceneProxy);
-			ELightMapPolicyType LightmapPolicyType = MobileBasePass::SelectMeshLightmapPolicy(Scene, MeshBatch, PrimitiveSceneProxy, MobileDirectionalLight, ShadingModel, bCanReceiveCSM, FeatureLevel);
-			Process(MeshBatch, BatchElementMask, StaticMeshId, PrimitiveSceneProxy, MaterialRenderProxy, Material, BlendMode, ShadingModel, LightmapPolicyType, MeshBatch.LCI);
+			ELightMapPolicyType LightmapPolicyType = MobileBasePass::SelectMeshLightmapPolicy(Scene, MeshBatch, PrimitiveSceneProxy, MobileDirectionalLight, ShadingModels, bCanReceiveCSM, FeatureLevel);
+			Process(MeshBatch, BatchElementMask, StaticMeshId, PrimitiveSceneProxy, MaterialRenderProxy, Material, BlendMode, ShadingModels, LightmapPolicyType, MeshBatch.LCI);
 		}
 	}
 }
@@ -596,7 +596,7 @@ void FMobileBasePassMeshProcessor::Process(
 		const FMaterialRenderProxy& RESTRICT MaterialRenderProxy,
 		const FMaterial& RESTRICT MaterialResource,
 		EBlendMode BlendMode,
-		EMaterialShadingModel ShadingModel,
+		FMaterialShadingModelField ShadingModels,
 		const ELightMapPolicyType LightMapPolicyType,
 		const FUniformLightMapPolicy::ElementDataType& RESTRICT LightMapElementData)
 {
@@ -606,7 +606,7 @@ void FMobileBasePassMeshProcessor::Process(
 		FBaseDS,
 		TMobileBasePassPSPolicyParamType<FUniformLightMapPolicy>> BasePassShaders;
 	
-	bool bEnableSkyLight = ShadingModel != MSM_Unlit && Scene && Scene->ShouldRenderSkylightInBasePass(BlendMode);
+	bool bEnableSkyLight = ShadingModels.IsLit() && Scene && Scene->ShouldRenderSkylightInBasePass(BlendMode);
 	int32 NumMovablePointLights = MobileBasePass::CalcNumMovablePointLights(MaterialResource, PrimitiveSceneProxy);
 
 	MobileBasePass::GetShaders(
diff --git a/Engine/Source/Runtime/Renderer/Private/MobileBasePassRendering.h b/Engine/Source/Runtime/Renderer/Private/MobileBasePassRendering.h
--- a/Engine/Source/Runtime/Renderer/Private/MobileBasePassRendering.h
+++ b/Engine/Source/Runtime/Renderer/Private/MobileBasePassRendering.h
@@ -345,7 +345,7 @@ namespace MobileBasePass
 		const FMeshBatch& MeshBatch, 
 		const FPrimitiveSceneProxy* PrimitiveSceneProxy, 
 		const FLightSceneInfo* MobileDirectionalLight, 
-		EMaterialShadingModel ShadingModel, 
+		FMaterialShadingModelField ShadingModels, 
 		bool bPrimReceivesCSM, 
 		ERHIFeatureLevel::Type FeatureLevel);
 
@@ -395,7 +395,7 @@ public:
 		const bool bMobileDynamicPointLightsUseStaticBranch = (MobileDynamicPointLightsUseStaticBranchCVar->GetValueOnAnyThread() == 1);
 		const int32 MobileNumDynamicPointLights = MobileNumDynamicPointLightsCVar->GetValueOnAnyThread();
 		const int32 MobileSkyLightPermutationOptions = MobileSkyLightPermutationCVar->GetValueOnAnyThread();
-		const bool bIsUnlit = Material->GetShadingModel() == MSM_Unlit;
+		const bool bIsUnlit = Material->GetShadingModels().IsUnlit();
 
 		// Only compile skylight version for lit materials on ES2 (Metal) or higher
 		const bool bShouldCacheBySkylight = !bEnableSkyLight || !bIsUnlit;
@@ -481,7 +481,7 @@ private:
 		const FMaterialRenderProxy& RESTRICT MaterialRenderProxy,
 		const FMaterial& RESTRICT MaterialResource,
 		EBlendMode BlendMode,
-		EMaterialShadingModel ShadingModel,
+		FMaterialShadingModelField ShadingModels,
 		const ELightMapPolicyType LightMapPolicyType,
 		const FUniformLightMapPolicy::ElementDataType& RESTRICT LightMapElementData);
 			
diff --git a/Engine/Source/Runtime/Renderer/Private/PostProcess/VisualizeShadingModels.cpp b/Engine/Source/Runtime/Renderer/Private/PostProcess/VisualizeShadingModels.cpp
--- a/Engine/Source/Runtime/Renderer/Private/PostProcess/VisualizeShadingModels.cpp
+++ b/Engine/Source/Runtime/Renderer/Private/PostProcess/VisualizeShadingModels.cpp
@@ -167,7 +167,7 @@ void FRCPassPostProcessVisualizeShadingModels::Process(FRenderingCompositePassCo
 
 	Y += 5;
 
-	for(uint32 i = 0; i < MSM_MAX; ++i)
+	for(uint32 i = 0; i < MSM_NUM; ++i)
 	{
 		FString Name = Enum->GetNameStringByValue(i);
 		Line = FString::Printf(TEXT("%d.  %s"), i, *Name);
diff --git a/Engine/Source/Runtime/Renderer/Private/RayTracing/RayTracingMaterialHitShaders.cpp b/Engine/Source/Runtime/Renderer/Private/RayTracing/RayTracingMaterialHitShaders.cpp
--- a/Engine/Source/Runtime/Renderer/Private/RayTracing/RayTracingMaterialHitShaders.cpp
+++ b/Engine/Source/Runtime/Renderer/Private/RayTracing/RayTracingMaterialHitShaders.cpp
@@ -243,7 +243,7 @@ void FRayTracingMeshProcessor::Process(
 	const FPrimitiveSceneProxy* RESTRICT PrimitiveSceneProxy,
 	const FMaterialRenderProxy& RESTRICT MaterialRenderProxy,
 	const FMaterial& RESTRICT MaterialResource,
-	EMaterialShadingModel ShadingModel,
+	FMaterialShadingModelField ShadingModels,
 	const FUniformLightMapPolicy& RESTRICT LightMapPolicy,
 	const typename FUniformLightMapPolicy::ElementDataType& RESTRICT LightMapElementData)
 {
@@ -309,7 +309,7 @@ void FRayTracingMeshProcessor::AddMeshBatch(const FMeshBatch& RESTRICT MeshBatch
 		const FMaterialRenderProxy& MaterialRenderProxy = FallbackMaterialRenderProxyPtr ? *FallbackMaterialRenderProxyPtr : *MeshBatch.MaterialRenderProxy;
 
 		const EBlendMode BlendMode = Material.GetBlendMode();
-		const EMaterialShadingModel ShadingModel = Material.GetShadingModel();
+		const FMaterialShadingModelField ShadingModels = Material.GetShadingModels();
 		const bool bIsTranslucent = IsTranslucentBlendMode(BlendMode);
 
 		// Only draw opaque materials.
@@ -317,7 +317,7 @@ void FRayTracingMeshProcessor::AddMeshBatch(const FMeshBatch& RESTRICT MeshBatch
 			&& ShouldIncludeDomainInMeshPass(Material.GetMaterialDomain()))
 		{
 			// Check for a cached light-map.
-			const bool bIsLitMaterial = (ShadingModel != MSM_Unlit);
+			const bool bIsLitMaterial = ShadingModels.IsLit();
 			static const auto AllowStaticLightingVar = IConsoleManager::Get().FindTConsoleVariableDataInt(TEXT("r.AllowStaticLighting"));
 			const bool bAllowStaticLighting = (!AllowStaticLightingVar || AllowStaticLightingVar->GetValueOnRenderThread() != 0);
 
@@ -353,7 +353,7 @@ void FRayTracingMeshProcessor::AddMeshBatch(const FMeshBatch& RESTRICT MeshBatch
 								PrimitiveSceneProxy,
 								MaterialRenderProxy,
 								Material,
-								ShadingModel,
+								ShadingModels,
 								FUniformLightMapPolicy(LMP_DISTANCE_FIELD_SHADOWS_AND_HQ_LIGHTMAP),
 								MeshBatch.LCI);
 						}
@@ -365,7 +365,7 @@ void FRayTracingMeshProcessor::AddMeshBatch(const FMeshBatch& RESTRICT MeshBatch
 								PrimitiveSceneProxy,
 								MaterialRenderProxy,
 								Material,
-								ShadingModel,
+								ShadingModels,
 								FUniformLightMapPolicy(LMP_HQ_LIGHTMAP),
 								MeshBatch.LCI);
 						}
@@ -378,7 +378,7 @@ void FRayTracingMeshProcessor::AddMeshBatch(const FMeshBatch& RESTRICT MeshBatch
 							PrimitiveSceneProxy,
 							MaterialRenderProxy,
 							Material,
-							ShadingModel,
+							ShadingModels,
 							FUniformLightMapPolicy(LMP_LQ_LIGHTMAP),
 							MeshBatch.LCI);
 					}
@@ -390,7 +390,7 @@ void FRayTracingMeshProcessor::AddMeshBatch(const FMeshBatch& RESTRICT MeshBatch
 							PrimitiveSceneProxy,
 							MaterialRenderProxy,
 							Material,
-							ShadingModel,
+							ShadingModels,
 							FUniformLightMapPolicy(LMP_NO_LIGHTMAP),
 							MeshBatch.LCI);
 					}
@@ -411,7 +411,7 @@ void FRayTracingMeshProcessor::AddMeshBatch(const FMeshBatch& RESTRICT MeshBatch
 							PrimitiveSceneProxy,
 							MaterialRenderProxy,
 							Material,
-							ShadingModel,
+							ShadingModels,
 							FUniformLightMapPolicy(LMP_PRECOMPUTED_IRRADIANCE_VOLUME_INDIRECT_LIGHTING),
 							MeshBatch.LCI);
 					}
@@ -423,7 +423,7 @@ void FRayTracingMeshProcessor::AddMeshBatch(const FMeshBatch& RESTRICT MeshBatch
 							PrimitiveSceneProxy,
 							MaterialRenderProxy,
 							Material,
-							ShadingModel,
+							ShadingModels,
 							FUniformLightMapPolicy(LMP_NO_LIGHTMAP),
 							MeshBatch.LCI);
 					}
diff --git a/Engine/Source/Runtime/Renderer/Private/RayTracing/RayTracingMaterialHitShaders.h b/Engine/Source/Runtime/Renderer/Private/RayTracing/RayTracingMaterialHitShaders.h
--- a/Engine/Source/Runtime/Renderer/Private/RayTracing/RayTracingMaterialHitShaders.h
+++ b/Engine/Source/Runtime/Renderer/Private/RayTracing/RayTracingMaterialHitShaders.h
@@ -34,7 +34,7 @@ private:
 		const FPrimitiveSceneProxy* RESTRICT PrimitiveSceneProxy,
 		const FMaterialRenderProxy& RESTRICT MaterialRenderProxy,
 		const FMaterial& RESTRICT MaterialResource,
-		EMaterialShadingModel ShadingModel,
+		FMaterialShadingModelField ShadingModels,
 		const FUniformLightMapPolicy& RESTRICT LightMapPolicy,
 		const typename FUniformLightMapPolicy::ElementDataType& RESTRICT LightMapElementData);
 
diff --git a/Engine/Source/Runtime/Renderer/Private/ShadowSetupMobile.cpp b/Engine/Source/Runtime/Renderer/Private/ShadowSetupMobile.cpp
--- a/Engine/Source/Runtime/Renderer/Private/ShadowSetupMobile.cpp
+++ b/Engine/Source/Runtime/Renderer/Private/ShadowSetupMobile.cpp
@@ -68,9 +68,7 @@ static bool EnableStaticMeshCSMVisibilityState(bool bMovableLight, const FPrimit
 		{
 			const FMaterialRenderProxy* MaterialRenderProxy = StaticMesh.MaterialRenderProxy;
 			const FMaterial* Material = MaterialRenderProxy->GetMaterial(View.GetFeatureLevel());
-			const EMaterialShadingModel ShadingModel = Material->GetShadingModel();
-			const bool bIsLitMaterial = ShadingModel != MSM_Unlit;
-			if (bIsLitMaterial)
+			if (Material->GetShadingModels().IsLit())
 			{
 				// CSM enabled list
 				MobileCSMVisibilityInfo.MobileCSMStaticMeshVisibilityMap[StaticMesh.Id] = MobileCSMVisibilityInfo.MobileNonCSMStaticMeshVisibilityMap[StaticMesh.Id];
