FPackedMaterialClosestHitPayload now stores uints instead of ints. That should fix hitKind() that was not working properly.

diff --git a/Engine/Shaders/Private/RayTracing/RayTracingCommon.ush b/Engine/Shaders/Private/RayTracing/RayTracingCommon.ush
--- a/Engine/Shaders/Private/RayTracing/RayTracingCommon.ush
+++ b/Engine/Shaders/Private/RayTracing/RayTracingCommon.ush
@@ -121,12 +121,12 @@ struct FPackedMaterialClosestHitPayload
 {
 	float HitT;                                         // 4  bytes
 	FRayCone RayCone;                                   // 8  bytes
-	int RadianceAndNormal[3];                           // 12 bytes
-	int BaseColorAndOpacity[2];                         // 8  bytes
-	int MetallicAndSpecularAndRoughness;                // 4  bytes
-	int IorAndShadingModelIDAndBlendingModeAndHitKind;  // 4  bytes
-	int IndirectIrradiance[2];                          // 8  bytes
-	int CustomData;                                     // 4  bytes
+	uint RadianceAndNormal[3];                           // 12 bytes
+	uint BaseColorAndOpacity[2];                         // 8  bytes
+	uint MetallicAndSpecularAndRoughness;                // 4  bytes
+	uint IorAndShadingModelIDAndBlendingModeAndHitKind;  // 4  bytes
+	uint IndirectIrradiance[2];                          // 8  bytes
+	uint CustomData;                                     // 4  bytes
 };                                                      // 52 bytes total
 
 FPackedMaterialClosestHitPayload PackRayTracingPayload(FMaterialClosestHitPayload Input, in FRayCone RayCone)
