Cherry-pick 6078870: Use DispatchRaysIndex() to initialize a fake SV_Position in ray tracing material shaders

diff --git a/Engine/Shaders/Private/RayTracing/RayTracingMaterialHitShaders.usf b/Engine/Shaders/Private/RayTracing/RayTracingMaterialHitShaders.usf
--- a/Engine/Shaders/Private/RayTracing/RayTracingMaterialHitShaders.usf
+++ b/Engine/Shaders/Private/RayTracing/RayTracingMaterialHitShaders.usf
@@ -173,19 +173,7 @@ void GetPrecomputedIndirectLightingAndSkyLight(
 	#endif
 }
 
-float4 ClipSpacePositionToSvPosition(float4 ClipSpacePosition)
-{
-	float3 NDCPos = ClipSpacePosition.xyz / ClipSpacePosition.w;
-
-	float4 SvPosition;
-	SvPosition.xy = ((NDCPos.xy * float2(0.5, -0.5) + 0.5f) * ResolvedView.ViewSizeAndInvSize.xy) + ResolvedView.ViewRectMin.xy;
-	SvPosition.z = NDCPos.z;
-	SvPosition.w = ClipSpacePosition.w;
-
-	return SvPosition;
-}
-
-void CalcInterpolants(in FRayCone RayCone, in FDefaultAttributes Attributes, out FVertexFactoryInterpolantsVSToPS Interpolants, out float4 FakeSvPosition)
+void CalcInterpolants(in FRayCone RayCone, in FDefaultAttributes Attributes, out FVertexFactoryInterpolantsVSToPS Interpolants)
 {
 	FVertexFactoryInterpolantsVSToDS PerVertexInterpolants[3];
 	float4 PerVertexClipSpacePositions[3]; // Used to fake SvPosition
@@ -225,7 +213,6 @@ void CalcInterpolants(in FRayCone RayCone, in FDefaultAttributes Attributes, out
 	Interpolants = VertexFactoryAssignInterpolants(Interpolated);
 	
 	float4 ClipSpacePosition = PerVertexClipSpacePositions[0] * (1 - Attributes.Barycentrics.x - Attributes.Barycentrics.y) + PerVertexClipSpacePositions[1] * Attributes.Barycentrics.x + PerVertexClipSpacePositions[2] * Attributes.Barycentrics.y;
-	FakeSvPosition = ClipSpacePosition / ClipSpacePosition.w;
 
 #if NUM_TEX_COORD_INTERPOLATORS || USE_PARTICLE_SUBUVS
 	float2 TA = TexCoords[1] - TexCoords[0];
@@ -253,10 +240,10 @@ void MaterialCHS(inout FPackedMaterialClosestHitPayload PackedPayload, in FDefau
 	ResolvedView = ResolveView();
 
 	FVertexFactoryInterpolantsVSToPS Interpolants;
-	float4 SvPosition;
+	float4 SvPosition = float4(float2(DispatchRaysIndex().xy) + ResolvedView.ViewRectMin.xy, 0.0, 1.0);
 	
 	FRayCone PropagatedCone = PropagateRayCone(PackedPayload.RayCone, 0 /* surface curvature */, RayTCurrent());
-	CalcInterpolants(PropagatedCone, Attributes, Interpolants, SvPosition);
+	CalcInterpolants(PropagatedCone, Attributes, Interpolants);
 
 	FMaterialPixelParameters MaterialParameters = GetMaterialPixelParameters(Interpolants, SvPosition);
 	FPixelMaterialInputs PixelMaterialInputs;
@@ -366,9 +353,9 @@ void MaterialAHS(inout FPackedMaterialClosestHitPayload PackedPayload, in FDefau
 	ResolvedView = ResolveView();
 
 	FVertexFactoryInterpolantsVSToPS Interpolants;
-	float4 SvPosition;
+	float4 SvPosition = float4(float2(DispatchRaysIndex().xy) + ResolvedView.ViewRectMin.xy, 0.0, 1.0);
 	
-	CalcInterpolants(PackedPayload.RayCone, Attributes, Interpolants, SvPosition);
+	CalcInterpolants(PackedPayload.RayCone, Attributes, Interpolants);
 
 	FMaterialPixelParameters MaterialParameters = GetMaterialPixelParameters(Interpolants, SvPosition);
 	FPixelMaterialInputs PixelMaterialInputs;
