Cherry-pick 6411512: Fix path tracer PreExposure

diff --git a/Engine/Shaders/Private/PathTracing/PathTracingCompositingPixelShader.usf b/Engine/Shaders/Private/PathTracing/PathTracingCompositingPixelShader.usf
--- a/Engine/Shaders/Private/PathTracing/PathTracingCompositingPixelShader.usf
+++ b/Engine/Shaders/Private/PathTracing/PathTracingCompositingPixelShader.usf
@@ -86,7 +86,7 @@ void CompositeMain(
 		OutVariance /= OutCumulativeSampleCount - 1;
 	}
 	OutCumulativeIrradiance.a = OutVariance;
-	OutColor = float4(OutCumulativeIrradiance.rgb, 1.0);
+	OutColor = float4(OutCumulativeIrradiance.rgb * View.PreExposure, 1.0);
 
 	// To visualize the current frame:
 	// OutColor = float4(Irradiance.rgb / SampleCount, 1.0);
