Fixes multi-bounce energy explosion in ray tracing global illumination pass.

diff --git a/Engine/Shaders/Private/RayTracing/RayTracingGlobalIlluminationRGS.usf b/Engine/Shaders/Private/RayTracing/RayTracingGlobalIlluminationRGS.usf
--- a/Engine/Shaders/Private/RayTracing/RayTracingGlobalIlluminationRGS.usf
+++ b/Engine/Shaders/Private/RayTracing/RayTracingGlobalIlluminationRGS.usf
@@ -67,7 +67,7 @@ void GenerateCosineNormalRay(
 
 	RayOrigin = WorldPosition;
 	RayDirection = Direction_World;
-	RayTMin = 0.0;
+	RayTMin = 0.01;
 	RayTMax = max(MaxRayDistanceForGI, MaxRayDistanceForAO);
 	RayPdf = Direction_Tangent.w;
 }
@@ -336,7 +336,6 @@ void GlobalIlluminationRGS()
 
 				// Update ray
 				Ray.Direction = Direction;
-				ApplyPositionBias(Ray, WorldNormal, MaxNormalBias);
 				RayThroughput *= Throughput;
 				RayPdf *= MaterialPdf;
 
