Backing out Stu's dynamic material defaults change for fear that it might cause patch bloat. It can wait for Season 10.

#tests engine tests pass for Niagara

diff --git a/Engine/Plugins/FX/Niagara/Shaders/Private/NiagaraMeshVertexFactory.ush b/Engine/Plugins/FX/Niagara/Shaders/Private/NiagaraMeshVertexFactory.ush
--- a/Engine/Plugins/FX/Niagara/Shaders/Private/NiagaraMeshVertexFactory.ush
+++ b/Engine/Plugins/FX/Niagara/Shaders/Private/NiagaraMeshVertexFactory.ush
@@ -245,9 +245,6 @@ FMaterialPixelParameters GetMaterialPixelParameters(FVertexFactoryInterpolantsVS
 	Result.Particle.Color = Interpolants.ParticleColor;
 #endif
 
-#if (DYNAMIC_PARAMETERS_MASK != 0)
-	Result.Particle.DynamicParameterValidMask = NiagaraMeshVF.MaterialParamValidMask;
-#endif
 #if (DYNAMIC_PARAMETERS_MASK & 1)
 	Result.Particle.DynamicParameter = Interpolants.DynamicParameter;
 #endif
@@ -328,9 +325,6 @@ FMaterialVertexParameters GetMaterialVertexParameters(FVertexFactoryInput Input,
 
 	Result.Particle.Color = Intermediates.ParticleColor;
 
-#if (DYNAMIC_PARAMETERS_MASK != 0)
-	Result.Particle.DynamicParameterValidMask = NiagaraMeshVF.MaterialParamValidMask;
-#endif
 #if (DYNAMIC_PARAMETERS_MASK & 1)
 	Result.Particle.DynamicParameter = Intermediates.DynamicParameter;
 #endif
diff --git a/Engine/Plugins/FX/Niagara/Shaders/Private/NiagaraRibbonVertexFactory.ush b/Engine/Plugins/FX/Niagara/Shaders/Private/NiagaraRibbonVertexFactory.ush
--- a/Engine/Plugins/FX/Niagara/Shaders/Private/NiagaraRibbonVertexFactory.ush
+++ b/Engine/Plugins/FX/Niagara/Shaders/Private/NiagaraRibbonVertexFactory.ush
@@ -448,9 +448,6 @@ FMaterialPixelParameters GetMaterialPixelParameters(FVertexFactoryInterpolantsVS
 
 	Result.TwoSidedSign = 1;
 
-#if (DYNAMIC_PARAMETERS_MASK != 0)
-	Result.Particle.DynamicParameterValidMask = NiagaraRibbonVF.MaterialParamValidMask;
-#endif
 #if (DYNAMIC_PARAMETERS_MASK & 1)
 	Result.Particle.DynamicParameter = Interpolants.DynamicParameter;
 #endif
@@ -500,9 +497,6 @@ FMaterialVertexParameters GetMaterialVertexParameters(FVertexFactoryInput Input,
 	Result.PreSkinnedPosition = Intermediates.Position.xyz;
 	Result.PreSkinnedNormal = TangentToLocal[2].xyz;
 
-#if (DYNAMIC_PARAMETERS_MASK != 0)
-	Result.Particle.DynamicParameterValidMask = NiagaraRibbonVF.MaterialParamValidMask;
-#endif
 #if (DYNAMIC_PARAMETERS_MASK & 1)
 	Result.Particle.DynamicParameter = Intermediates.DynamicParameter;
 #endif
diff --git a/Engine/Plugins/FX/Niagara/Shaders/Private/NiagaraSpriteVertexFactory.ush b/Engine/Plugins/FX/Niagara/Shaders/Private/NiagaraSpriteVertexFactory.ush
--- a/Engine/Plugins/FX/Niagara/Shaders/Private/NiagaraSpriteVertexFactory.ush
+++ b/Engine/Plugins/FX/Niagara/Shaders/Private/NiagaraSpriteVertexFactory.ush
@@ -269,9 +269,6 @@ FMaterialPixelParameters GetMaterialPixelParameters(FVertexFactoryInterpolantsVS
 	Result.Particle.Color = Interpolants.Color;
 #endif
 
-#if (DYNAMIC_PARAMETERS_MASK != 0)
-	Result.Particle.DynamicParameterValidMask = NiagaraSpriteVF.MaterialParamValidMask;
-#endif
 #if (DYNAMIC_PARAMETERS_MASK & 1)
 	Result.Particle.DynamicParameter = Interpolants.DynamicParameter;
 #endif
@@ -345,9 +342,6 @@ FMaterialVertexParameters GetMaterialVertexParameters(FVertexFactoryInput Input,
 	// primitives using this VF are usually transparent and hence don't output velocity
 	Result.PrevFrameLocalToWorld = Primitive.LocalToWorld;
 
-#if (DYNAMIC_PARAMETERS_MASK != 0)
-	Result.Particle.DynamicParameterValidMask = NiagaraSpriteVF.MaterialParamValidMask;
-#endif
 #if (DYNAMIC_PARAMETERS_MASK & 1)
 	Result.Particle.DynamicParameter = Intermediates.DynamicParameter;
 #endif
diff --git a/Engine/Plugins/FX/Niagara/Source/Niagara/Private/NiagaraRendererMeshes.cpp b/Engine/Plugins/FX/Niagara/Source/Niagara/Private/NiagaraRendererMeshes.cpp
--- a/Engine/Plugins/FX/Niagara/Source/Niagara/Private/NiagaraRendererMeshes.cpp
+++ b/Engine/Plugins/FX/Niagara/Source/Niagara/Private/NiagaraRendererMeshes.cpp
@@ -219,7 +219,7 @@ void FNiagaraRendererMeshes::GetDynamicMeshElements(const TArray<const FSceneVie
 				PerViewUniformParameters.TransformDataOffset = TransformOffset;
 				PerViewUniformParameters.ScaleDataOffset = ScaleOffset;
 				PerViewUniformParameters.SizeDataOffset = SizeOffset;
-				PerViewUniformParameters.MaterialParamValidMask = MaterialParamValidMask;
+				//PerViewUniformParameters.MaterialParamValidMask = MaterialParamValidMask;
 				PerViewUniformParameters.MaterialParamDataOffset = MaterialParamOffset;
 				PerViewUniformParameters.MaterialParam1DataOffset = MaterialParamOffset1;
 				PerViewUniformParameters.MaterialParam2DataOffset = MaterialParamOffset2;
diff --git a/Engine/Plugins/FX/Niagara/Source/Niagara/Private/NiagaraRendererRibbons.cpp b/Engine/Plugins/FX/Niagara/Source/Niagara/Private/NiagaraRendererRibbons.cpp
--- a/Engine/Plugins/FX/Niagara/Source/Niagara/Private/NiagaraRendererRibbons.cpp
+++ b/Engine/Plugins/FX/Niagara/Source/Niagara/Private/NiagaraRendererRibbons.cpp
@@ -396,7 +396,7 @@ void FNiagaraRendererRibbons::GetDynamicMeshElements(const TArray<const FSceneVi
 			PerViewUniformParameters.FacingDataOffset = FacingMode == ENiagaraRibbonFacingMode::Custom ? FacingDataOffset : -1;
 			PerViewUniformParameters.NormalizedAgeDataOffset = NormalizedAgeDataOffset;
 			PerViewUniformParameters.MaterialRandomDataOffset = MaterialRandomDataOffset;
-			PerViewUniformParameters.MaterialParamValidMask = MaterialParamValidMask;
+			//PerViewUniformParameters.MaterialParamValidMask = MaterialParamValidMask;
 			PerViewUniformParameters.MaterialParamDataOffset = MaterialParamOffset;
 			PerViewUniformParameters.MaterialParam1DataOffset = MaterialParamOffset1;
 			PerViewUniformParameters.MaterialParam2DataOffset = MaterialParamOffset2;
diff --git a/Engine/Plugins/FX/Niagara/Source/Niagara/Private/NiagaraRendererSprites.cpp b/Engine/Plugins/FX/Niagara/Source/Niagara/Private/NiagaraRendererSprites.cpp
--- a/Engine/Plugins/FX/Niagara/Source/Niagara/Private/NiagaraRendererSprites.cpp
+++ b/Engine/Plugins/FX/Niagara/Source/Niagara/Private/NiagaraRendererSprites.cpp
@@ -254,7 +254,7 @@ FNiagaraSpriteUniformBufferRef FNiagaraRendererSprites::CreatePerViewUniformBuff
 	PerViewUniformParameters.RotationDataOffset = RotationOffset;
 	PerViewUniformParameters.SizeDataOffset = SizeOffset;
 	PerViewUniformParameters.ColorDataOffset = ColorOffset;
-	PerViewUniformParameters.MaterialParamValidMask = MaterialParamValidMask;
+	//PerViewUniformParameters.MaterialParamValidMask = MaterialParamValidMask;
 	PerViewUniformParameters.MaterialParamDataOffset = MaterialParamOffset;
 	PerViewUniformParameters.MaterialParam1DataOffset = MaterialParamOffset1;
 	PerViewUniformParameters.MaterialParam2DataOffset = MaterialParamOffset2;
diff --git a/Engine/Plugins/FX/Niagara/Source/NiagaraVertexFactories/Public/NiagaraMeshVertexFactory.h b/Engine/Plugins/FX/Niagara/Source/NiagaraVertexFactories/Public/NiagaraMeshVertexFactory.h
--- a/Engine/Plugins/FX/Niagara/Source/NiagaraVertexFactories/Public/NiagaraMeshVertexFactory.h
+++ b/Engine/Plugins/FX/Niagara/Source/NiagaraVertexFactories/Public/NiagaraMeshVertexFactory.h
@@ -41,7 +41,6 @@ BEGIN_GLOBAL_SHADER_PARAMETER_STRUCT(FNiagaraMeshUniformParameters, NIAGARAVERTE
 	SHADER_PARAMETER(int, TransformDataOffset)
 	SHADER_PARAMETER(int, ScaleDataOffset)
 	SHADER_PARAMETER(int, SizeDataOffset)
-	SHADER_PARAMETER(uint32, MaterialParamValidMask)
 	SHADER_PARAMETER(int, MaterialParamDataOffset)
 	SHADER_PARAMETER(int, MaterialParam1DataOffset)
 	SHADER_PARAMETER(int, MaterialParam2DataOffset)
diff --git a/Engine/Plugins/FX/Niagara/Source/NiagaraVertexFactories/Public/NiagaraRibbonVertexFactory.h b/Engine/Plugins/FX/Niagara/Source/NiagaraVertexFactories/Public/NiagaraRibbonVertexFactory.h
--- a/Engine/Plugins/FX/Niagara/Source/NiagaraVertexFactories/Public/NiagaraRibbonVertexFactory.h
+++ b/Engine/Plugins/FX/Niagara/Source/NiagaraVertexFactories/Public/NiagaraRibbonVertexFactory.h
@@ -37,7 +37,6 @@ BEGIN_GLOBAL_SHADER_PARAMETER_STRUCT(FNiagaraRibbonUniformParameters, NIAGARAVER
 	SHADER_PARAMETER(int, FacingDataOffset)
 	SHADER_PARAMETER(int, NormalizedAgeDataOffset)
 	SHADER_PARAMETER(int, MaterialRandomDataOffset)
-	SHADER_PARAMETER(uint32, MaterialParamValidMask)
 	SHADER_PARAMETER(int, MaterialParamDataOffset)
 	SHADER_PARAMETER(int, MaterialParam1DataOffset)
 	SHADER_PARAMETER(int, MaterialParam2DataOffset)
diff --git a/Engine/Plugins/FX/Niagara/Source/NiagaraVertexFactories/Public/NiagaraSpriteVertexFactory.h b/Engine/Plugins/FX/Niagara/Source/NiagaraVertexFactories/Public/NiagaraSpriteVertexFactory.h
--- a/Engine/Plugins/FX/Niagara/Source/NiagaraVertexFactories/Public/NiagaraSpriteVertexFactory.h
+++ b/Engine/Plugins/FX/Niagara/Source/NiagaraVertexFactories/Public/NiagaraSpriteVertexFactory.h
@@ -40,7 +40,6 @@ BEGIN_GLOBAL_SHADER_PARAMETER_STRUCT( FNiagaraSpriteUniformParameters, NIAGARAVE
 	SHADER_PARAMETER(int, SizeDataOffset)
 	SHADER_PARAMETER(int, SubimageDataOffset)
 	SHADER_PARAMETER(int, ColorDataOffset)
-	SHADER_PARAMETER(uint32, MaterialParamValidMask)
 	SHADER_PARAMETER(int, MaterialParamDataOffset)
 	SHADER_PARAMETER(int, MaterialParam1DataOffset)
 	SHADER_PARAMETER(int, MaterialParam2DataOffset)
diff --git a/Engine/Shaders/Private/MaterialTemplate.ush b/Engine/Shaders/Private/MaterialTemplate.ush
--- a/Engine/Shaders/Private/MaterialTemplate.ush
+++ b/Engine/Shaders/Private/MaterialTemplate.ush
@@ -91,9 +91,6 @@ struct FMaterialParticleParameters
 	half4 MacroUV;
 
 	/** Dynamic parameters used by particle systems. */
-#if NIAGARA_PARTICLE_FACTORY && (DYNAMIC_PARAMETERS_MASK != 0)
-	uint DynamicParameterValidMask;
-#endif
 	half4 DynamicParameter;
 
 #if( DYNAMIC_PARAMETERS_MASK & 2)
@@ -125,26 +122,37 @@ struct FMaterialParticleParameters
 float4 GetDynamicParameter(FMaterialParticleParameters Parameters, float4 Default, int ParameterIndex=0)
 {
 #if (NIAGARA_PARTICLE_FACTORY)
-	switch ( ParameterIndex)
+	if(ParameterIndex == 0)
 	{
-	#if (DYNAMIC_PARAMETERS_MASK & 1)
-		case 0:	return (Parameters.DynamicParameterValidMask & 1) != 0 ? Parameters.DynamicParameter : Default;
-	#endif
-	#if (DYNAMIC_PARAMETERS_MASK & 2)
-		case 1:	return (Parameters.DynamicParameterValidMask & 2) != 0 ? Parameters.DynamicParameter1 : Default;
-	#endif
-	#if (DYNAMIC_PARAMETERS_MASK & 4)
-		case 2:	return (Parameters.DynamicParameterValidMask & 4) != 0 ? Parameters.DynamicParameter2 : Default;
-	#endif	
-	#if (DYNAMIC_PARAMETERS_MASK & 8)
-		case 3:	return (Parameters.DynamicParameterValidMask & 8) != 0 ? Parameters.DynamicParameter3 : Default;
-	#endif
-		default: return Default;
+		return Parameters.DynamicParameter;
+	}
+#if (DYNAMIC_PARAMETERS_MASK & 2)
+	else if(ParameterIndex == 1)
+	{
+		return Parameters.DynamicParameter1;
+	}
+#endif
+#if (DYNAMIC_PARAMETERS_MASK & 4)
+	else if(ParameterIndex == 2)
+	{
+		return Parameters.DynamicParameter2;
+	}
+#endif	
+#if (DYNAMIC_PARAMETERS_MASK & 8)
+	else if(ParameterIndex == 3)
+	{
+		return Parameters.DynamicParameter3;
+	}
+#endif
+	else
+	{
+		return Default;
 	}
 #elif (PARTICLE_FACTORY)
 	return Parameters.DynamicParameter;
-#endif
+#else
 	return Default;
+#endif
 
 }
 
