Fixed raytracing shader compilation error

diff --git a/Engine/Shaders/Private/RayTracing/RayTracingMaterialHitShaders.usf b/Engine/Shaders/Private/RayTracing/RayTracingMaterialHitShaders.usf
--- a/Engine/Shaders/Private/RayTracing/RayTracingMaterialHitShaders.usf
+++ b/Engine/Shaders/Private/RayTracing/RayTracingMaterialHitShaders.usf
@@ -143,7 +143,7 @@ void GetPrecomputedIndirectLightingAndSkyLight(
 		SkyOcclusionUV = LightmapUV0;
 
 		uint VTRequest = 0;
-		GetLightMapColorHQ(LightmapUV0, LightmapUV1, LightmapDataIndex, DiffuseDir, MaterialParameters.SvPosition.xy, OutDiffuseLighting, OutSubsurfaceLighting, VTRequest);
+		GetLightMapColorHQ(LightmapUV0, LightmapUV1, LightmapDataIndex, DiffuseDir, MaterialParameters.SvPosition.xy, ShadingModel, OutDiffuseLighting, OutSubsurfaceLighting, VTRequest);
 
 	// Low quality texture lightmaps
 	#elif LQ_TEXTURE_LIGHTMAP
