Workaround for decal rendering artifacts (wireframes) on iOS

diff --git a/Engine/Shaders/Private/Common.ush b/Engine/Shaders/Private/Common.ush
--- a/Engine/Shaders/Private/Common.ush
+++ b/Engine/Shaders/Private/Common.ush
@@ -207,6 +207,38 @@ MaterialFloat4 TextureExternalSampleLevel(TextureExternal Tex, SamplerState Samp
 	return Tex.SampleLevel(Sampler, UV, Mip);
 }
 
+// Re-routed texture sampling for decals
+// On iOS decal UVs has issues with deriavatives on polygon edges resulting in 'wireframe' rendering artifacts (UE-50971)
+// to workaround we always sample top mip level
+MaterialFloat4 Texture1DSample_Decal(Texture1D Tex, SamplerState Sampler, float UV)
+{
+	return Texture1DSample(Tex, Sampler, UV);
+}
+MaterialFloat4 Texture2DSample_Decal(Texture2D Tex, SamplerState Sampler, float2 UV)
+{
+#if METAL_PROFILE
+	return Texture2DSampleGrad(Tex, Sampler, UV, 0.f, 0.f);
+#else
+	return Texture2DSample(Tex, Sampler, UV);
+#endif
+}
+MaterialFloat4 Texture3DSample_Decal(Texture3D Tex, SamplerState Sampler, float3 UV)
+{
+#if METAL_PROFILE
+	return Texture3DSampleGrad(Tex, Sampler, UV, 0.f, 0.f);
+#else
+	return Texture3DSample(Tex, Sampler, UV);
+#endif
+}
+MaterialFloat4 TextureCubeSample_Decal(TextureCube Tex, SamplerState Sampler, float3 UV)
+{
+#if METAL_PROFILE
+	return TextureCubeSampleGrad(Tex, Sampler, UV, 0.f, 0.f);
+#else
+	return TextureCubeSample(Tex, Sampler, UV);
+#endif
+}
+
 //converts an input 1d to 2d position. Useful for locating z frames that have been laid out in a 2d grid like a flipbook.
 float2 Tile1Dto2D(float xsize, float idx)
 {
diff --git a/Engine/Source/Runtime/Engine/Private/Materials/HLSLMaterialTranslator.h b/Engine/Source/Runtime/Engine/Private/Materials/HLSLMaterialTranslator.h
--- a/Engine/Source/Runtime/Engine/Private/Materials/HLSLMaterialTranslator.h
+++ b/Engine/Source/Runtime/Engine/Private/Materials/HLSLMaterialTranslator.h
@@ -4246,7 +4246,13 @@ protected:
 		}
 		else
 		{
-			// Non-VT MipValueMode logic
+			// Non-VT MipValueMode logic			
+
+			// Re-route decal texture sampling so platforms may add specific workarounds there
+			if (ShaderFrequency == SF_Pixel && Material->GetMaterialDomain() == MD_DeferredDecal && MipValueMode == TMVM_None)
+			{
+				SampleCode += TEXT("_Decal");
+			}
 
 			SamplerStateCode = ", " + SamplerStateCode;
 
