Add workaround for decal rendering artifacts (wireframes) on Android ES3.1 (Mali)

diff --git a/Engine/Shaders/Private/Common.ush b/Engine/Shaders/Private/Common.ush
--- a/Engine/Shaders/Private/Common.ush
+++ b/Engine/Shaders/Private/Common.ush
@@ -208,7 +208,7 @@ MaterialFloat4 TextureExternalSampleLevel(TextureExternal Tex, SamplerState Samp
 }
 
 // Re-routed texture sampling for decals
-// On iOS decal UVs has issues with deriavatives on polygon edges resulting in 'wireframe' rendering artifacts (UE-50971)
+// On iOS and Android MALI devices decal UVs has issues with deriavatives on polygon edges resulting in 'wireframe' rendering artifacts (UE-50971)
 // to workaround we always sample top mip level
 MaterialFloat4 Texture1DSample_Decal(Texture1D Tex, SamplerState Sampler, float UV)
 {
@@ -216,7 +216,7 @@ MaterialFloat4 Texture1DSample_Decal(Texture1D Tex, SamplerState Sampler, float
 }
 MaterialFloat4 Texture2DSample_Decal(Texture2D Tex, SamplerState Sampler, float2 UV)
 {
-#if METAL_PROFILE
+#if METAL_PROFILE || COMPILER_GLSL_ES3_1
 	return Texture2DSampleGrad(Tex, Sampler, UV, 0.f, 0.f);
 #else
 	return Texture2DSample(Tex, Sampler, UV);
@@ -224,7 +224,7 @@ MaterialFloat4 Texture2DSample_Decal(Texture2D Tex, SamplerState Sampler, float2
 }
 MaterialFloat4 Texture3DSample_Decal(Texture3D Tex, SamplerState Sampler, float3 UV)
 {
-#if METAL_PROFILE
+#if METAL_PROFILE || COMPILER_GLSL_ES3_1
 	return Texture3DSampleGrad(Tex, Sampler, UV, 0.f, 0.f);
 #else
 	return Texture3DSample(Tex, Sampler, UV);
@@ -232,7 +232,7 @@ MaterialFloat4 Texture3DSample_Decal(Texture3D Tex, SamplerState Sampler, float3
 }
 MaterialFloat4 TextureCubeSample_Decal(TextureCube Tex, SamplerState Sampler, float3 UV)
 {
-#if METAL_PROFILE
+#if METAL_PROFILE || COMPILER_GLSL_ES3_1
 	return TextureCubeSampleGrad(Tex, Sampler, UV, 0.f, 0.f);
 #else
 	return TextureCubeSample(Tex, Sampler, UV);
