Cherry-pick 6980350: Fixed ray tracing debug modes where using pre exposure incorrectly

diff --git a/Engine/Shaders/Private/RayTracing/RayTracingBarycentrics.usf b/Engine/Shaders/Private/RayTracing/RayTracingBarycentrics.usf
--- a/Engine/Shaders/Private/RayTracing/RayTracingBarycentrics.usf
+++ b/Engine/Shaders/Private/RayTracing/RayTracingBarycentrics.usf
@@ -48,7 +48,7 @@ void RayTracingBarycentricsMainRGS()
 	{
 		Result = float4(Payload.Barycentrics, 1.0f - Payload.Barycentrics.x - Payload.Barycentrics.y, 1.0f);
 	}
-	Output[PixelCoord] = Result * View.PreExposure;
+	Output[PixelCoord] = Result;
 }
 
 [shader("closesthit")]
diff --git a/Engine/Shaders/Private/RayTracing/RayTracingDebug.usf b/Engine/Shaders/Private/RayTracing/RayTracingDebug.usf
--- a/Engine/Shaders/Private/RayTracing/RayTracingDebug.usf
+++ b/Engine/Shaders/Private/RayTracing/RayTracingDebug.usf
@@ -106,7 +106,7 @@ void RayTracingDebugMainRGS()
 		}
 	}
 
-	Output[PixelCoord] = Result * View.PreExposure;
+	Output[PixelCoord] = Result;
 }
 
 [shader("closesthit")]
