Fix texel offset in RTV compression
Integrated from //UE4/Dev-Rendering
#rnx

diff --git a/Engine/Shaders/Private/BCCompressionCommon.ush b/Engine/Shaders/Private/BCCompressionCommon.ush
--- a/Engine/Shaders/Private/BCCompressionCommon.ush
+++ b/Engine/Shaders/Private/BCCompressionCommon.ush
@@ -15,7 +15,7 @@
 #include "GammaCorrectionCommon.ush"
 
 // Simple convert float3 color to 565 uint using 'round' arithmetic
-uint ColorTo565( in float3 Color )
+uint Float3ToUint565( in float3 Color )
 {
 	float3 Scale       = float3(31.f, 63.f, 31.f);
 	float3 ColorScaled = round(saturate(Color) * Scale);
@@ -26,7 +26,7 @@ uint ColorTo565( in float3 Color )
 
 // Convert float3 color to 565 uint using 'ceil' arithmetic
 // Color parameter is inout and is modified to match the converted value
-uint ColorTo565Max( inout float3 Color )
+uint Float3ToUint565_Ceil( inout float3 Color )
 {
 	float3 Scale         = float3(31.f, 63.f, 31.f);
 	float3 ColorScaled   = ceil(saturate(Color) * Scale);
@@ -38,7 +38,7 @@ uint ColorTo565Max( inout float3 Color )
 
 // Convert float3 color to 565 uint using 'floor' arithmetic
 // Color parameter is inout and is modified to match the converted value
-uint ColorTo565Min( inout float3 Color )
+uint Float3ToUint565_Floor( inout float3 Color )
 {
 	float3 Scale         = float3(31.f, 63.f, 31.f);
 	float3 ColorScaled   = floor(saturate(Color) * Scale);
@@ -92,81 +92,70 @@ void GetMinMax( in float3 Block[16], out float3 OutMin, out float3 OutMax )
 }
 
 // Calculate the final packed indices for a color block
-uint GetIndicesRGB( in float3 Block[16], in float3 MinColor, in float3 MaxColor )
+uint GetPackedColorBlockIndices( in float3 Block[16], in float3 MinColor, in float3 MaxColor )
 {
-	uint OutIndices = 0;
-
-	// We project (Block[i] - MaxColor) onto (MinColor - MaxColor) and split into final values in range (0,3)
-	// For efficiency we pull out the constant terms here
-	float3 Diff    = MinColor - MaxColor;
-	float StepInc  = 3.f / dot(Diff, Diff);
-	Diff          *= StepInc;
-	float Offset   = StepInc * (dot(MaxColor, MaxColor) - dot(MaxColor, MinColor));
+	uint PackedIndices = 0;
+
+	// Project onto max->min color vector and segment into range [0,3]
+	float3 Range = MinColor - MaxColor;
+	float Scale = 3.f / dot(Range, Range);
+	float3 ScaledRange  = Range * Scale;
+	float Bias = (dot(MaxColor, MaxColor) - dot(MaxColor, MinColor)) * Scale;
 	
 	for (int i=15; i>=0; --i)
 	{
-		// Compute the index for this Block element
-		uint Index = round(dot(Block[i], Diff) + Offset);
-
-		// Now we need to convert our index into the BC indexing scheme:
-		//  0: MaxColor
-		//  1: MinColor
-		//  2: (2/3)*MaxColor + (1/3)*MinColor
-		//  3: (1/3)*MaxColor + (2/3)*MinColor
+		// Compute the distance index for this element
+		uint Index = round(dot(Block[i], ScaledRange) + Bias);
+		// Convert distance index into the BC index
 		Index += (Index > 0) - (3 * (Index == 3));
-
-		// OR into the final Indices
-		OutIndices = (OutIndices<<2)|Index;
+		// OR into the final PackedIndices
+		PackedIndices = (PackedIndices<<2)|Index;
 	}
 
-	return OutIndices;
+	return PackedIndices;
 }
 
 // Calculate the final packed indices for an alpha block
 // The results are in their final location of the uint2 indices and will need ORing with the min and max alpha
-uint2 GetIndicesAlpha( in float Block[16], in float MinAlpha, in float MaxAlpha )
+uint2 GetPackedAlphaBlockIndices( in float Block[16], in float MinAlpha, in float MaxAlpha )
 {
-	// For each input alpha value, we need to select an index in range (0,7)
-	float Diff    = MinAlpha - MaxAlpha;
-	float StepInc = 7.f / Diff;
-	float Offset  = -StepInc * MaxAlpha;
-
-	// We need two loops for the two 32 bit components of our output
-	// The first 6 elements of the block will go into the top 16 bits of the x component
-	// The 6th block straddles the x and y components
-	// The rest of the elements fill the y component
-	uint2 Indices = 0;
-	uint Index = 0;
-	uint Shift = 16;
-	
-	int i = 0;
-	for (i=0; i<6; ++i)
-	{
-		// Compute the index for this Block element
-		Index = round(StepInc * Block[i] + Offset);
+	uint2 PackedIndices = 0;
 
-		// Now we need to convert our index into the BC indexing scheme
-		// The logic is similar to the GetIndicesRGB case above
-		Index += (Index > 0) - (7 * (Index == 7));
+	// Segment alpha max->min into range [0,7]
+	float Range = MinAlpha - MaxAlpha;
+	float Scale = 7.f / Range;
+	float Bias = -Scale * MaxAlpha;
 
-		// OR into the final Indices
-		Indices.x |= (Index << Shift);
-		Shift     += 3;
+	uint i = 0;
+	// The first 5 elements of the block will go into the top 16 bits of the x component
+	for (i=0; i<5; ++i)
+	{
+		// Compute the distance index for this element
+		uint Index = round(Block[i] * Scale + Bias);
+		// Convert distance index into the BC index
+		Index += (Index > 0) - (7 * (Index == 7));
+		// OR into the final PackedIndices
+		PackedIndices.x |= (Index << (i*3 + 16));
 	}
 
-	// The 6th index straddles the two components
-	Indices.y |= (Index >> 1);
+	// The 6th element is split across the x and y components
+	{
+		i = 5;
+		uint Index = round(Block[i] * Scale + Bias);
+		Index += (Index > 0) - (7 * (Index == 7));
+		PackedIndices.x |= (Index << 31);
+		PackedIndices.y |= (Index >> 1);
+	}
 
-	Shift = 2;
+	// The rest of the elements fill the y component
 	for (i=6; i<16; ++i)
 	{
-		Index      = round(StepInc * Block[i] + Offset);
-		Index     += (Index > 0) - (7 * (Index == 7 ? 1 : 0));
-		Indices.y |= (Index << Shift);
-		Shift     += 3;
+		uint Index = round(Block[i] * Scale + Bias);
+		Index += (Index > 0) - (7 * (Index == 7 ? 1 : 0));
+		PackedIndices.y |= (Index << (i*3 - 16));
 	}
 
-	return Indices;
+	return PackedIndices;
 }
 
 // Compress a BC1 block
@@ -175,13 +164,13 @@ uint2 CompressBC1Block( in float3 Block[16] )
 	float3 MinColor, MaxColor;
 	GetMinMax( Block, MinColor, MaxColor );
 
-	uint MinColor565 = ColorTo565Min( MinColor );
-	uint MaxColor565 = ColorTo565Max( MaxColor );
+	uint MinColor565 = Float3ToUint565_Floor( MinColor );
+	uint MaxColor565 = Float3ToUint565_Ceil( MaxColor );
 
 	uint Indices = 0;
 	if (MinColor565 < MaxColor565)
 	{
-		Indices = GetIndicesRGB( Block, MinColor, MaxColor );
+		Indices = GetPackedColorBlockIndices( Block, MinColor, MaxColor );
 	}
 
 	return uint2((MinColor565 << 16) | MaxColor565, Indices);
@@ -194,13 +183,13 @@ uint2 CompressBC1Block_SRGB( in float3 Block[16] )
 	float3 MinColor, MaxColor;
 	GetMinMax( Block, MinColor, MaxColor );
 
-	uint MinColor565 = ColorTo565( LinearToSrgb(MinColor) );
-	uint MaxColor565 = ColorTo565( LinearToSrgb(MaxColor) );
+	uint MinColor565 = Float3ToUint565( LinearToSrgb(MinColor) );
+	uint MaxColor565 = Float3ToUint565( LinearToSrgb(MaxColor) );
 
 	uint Indices = 0;
 	if (MinColor565 < MaxColor565)
 	{
-		Indices = GetIndicesRGB( Block, MinColor, MaxColor );
+		Indices = GetPackedColorBlockIndices( Block, MinColor, MaxColor );
 	}
 
 	return uint2((MinColor565 << 16) | MaxColor565, Indices);
@@ -212,8 +201,8 @@ uint4 CompressBC3Block( in float3 BlockRGB[16], in float BlockA[16] )
 	float3 MinColor, MaxColor;
 	GetMinMax( BlockRGB, MinColor, MaxColor );
 
-	uint MinColor565 = ColorTo565Min( MinColor );
-	uint MaxColor565 = ColorTo565Max( MaxColor );
+	uint MinColor565 = Float3ToUint565_Floor( MinColor );
+	uint MaxColor565 = Float3ToUint565_Ceil( MaxColor );
 
 	float MinAlpha, MaxAlpha;
 	GetMinMax( BlockA, MinAlpha, MaxAlpha );
@@ -224,13 +213,13 @@ uint4 CompressBC3Block( in float3 BlockRGB[16], in float BlockA[16] )
 	uint ColorIndices = 0;
 	if (MinColor565 < MaxColor565)
 	{
-		ColorIndices = GetIndicesRGB( BlockRGB, MinColor, MaxColor );
+		ColorIndices = GetPackedColorBlockIndices( BlockRGB, MinColor, MaxColor );
 	}
 
 	uint2 AlphaIndices = 0;
 	if (MinAlphaUint < MaxAlphaUint)
 	{
-		AlphaIndices = GetIndicesAlpha(BlockA, MinAlpha, MaxAlpha);
+		AlphaIndices = GetPackedAlphaBlockIndices(BlockA, MinAlpha, MaxAlpha);
 	}
 
 	return uint4((MinAlphaUint << 8) | MaxAlphaUint | AlphaIndices.x, AlphaIndices.y, (MinColor565 << 16) | MaxColor565, ColorIndices);
@@ -243,8 +232,8 @@ uint4 CompressBC3Block_SRGB( in float3 BlockRGB[16], in float BlockA[16] )
 	float3 MinColor, MaxColor;
 	GetMinMax( BlockRGB, MinColor, MaxColor );
 
-	uint MinColor565 = ColorTo565( LinearToSrgb( MinColor ) );
-	uint MaxColor565 = ColorTo565( LinearToSrgb( MaxColor ) );
+	uint MinColor565 = Float3ToUint565( LinearToSrgb( MinColor ) );
+	uint MaxColor565 = Float3ToUint565( LinearToSrgb( MaxColor ) );
 
 	float MinAlpha, MaxAlpha;
 	GetMinMax( BlockA, MinAlpha, MaxAlpha );
@@ -255,13 +244,13 @@ uint4 CompressBC3Block_SRGB( in float3 BlockRGB[16], in float BlockA[16] )
 	uint ColorIndices = 0;
 	if (MinColor565 < MaxColor565)
 	{
-		ColorIndices = GetIndicesRGB(BlockRGB, MinColor, MaxColor);
+		ColorIndices = GetPackedColorBlockIndices(BlockRGB, MinColor, MaxColor);
 	}
 
 	uint2 AlphaIndices = 0;
 	if (MinAlphaUint < MaxAlphaUint)
 	{
-		AlphaIndices = GetIndicesAlpha(BlockA, MinAlpha, MaxAlpha);
+		AlphaIndices = GetPackedAlphaBlockIndices(BlockA, MinAlpha, MaxAlpha);
 	}
 
 	return uint4((MinAlphaUint << 8) | MaxAlphaUint | AlphaIndices.x, AlphaIndices.y, (MinColor565 << 16) | MaxColor565, ColorIndices);
@@ -279,7 +268,7 @@ uint2 CompressBC4Block( in float Block[16] )
 	uint2 Indices = 0;
 	if (MinAlphaUint < MaxAlphaUint)
 	{
-		Indices = GetIndicesAlpha( Block, MinAlpha, MaxAlpha );
+		Indices = GetPackedAlphaBlockIndices( Block, MinAlpha, MaxAlpha );
 	}
 
 	return uint2((MinAlphaUint << 8) | MaxAlphaUint | Indices.x, Indices.y);
@@ -299,13 +288,13 @@ uint4 CompressBC5Block( in float BlockU[16], in float BlockV[16] )
 	uint2 IndicesU = 0;
 	if (MinUUint < MaxUUint)
 	{
-		IndicesU = GetIndicesAlpha( BlockU, MinU, MaxU );
+		IndicesU = GetPackedAlphaBlockIndices( BlockU, MinU, MaxU );
 	}
 
 	uint2 IndicesV = 0;
 	if (MinVUint < MaxVUint)
 	{
-		IndicesV = GetIndicesAlpha( BlockV, MinV, MaxV );
+		IndicesV = GetPackedAlphaBlockIndices( BlockV, MinV, MaxV );
 	}
 
 	return uint4((MinUUint << 8) | MaxUUint | IndicesU.x, IndicesU.y, (MinVUint << 8) | MaxVUint | IndicesV.x, IndicesV.y);
diff --git a/Engine/Shaders/Private/VirtualTextureCompress.usf b/Engine/Shaders/Private/VirtualTextureCompress.usf
--- a/Engine/Shaders/Private/VirtualTextureCompress.usf
+++ b/Engine/Shaders/Private/VirtualTextureCompress.usf
@@ -21,39 +21,42 @@ RWTexture2D<uint4> OutCompressTexture1;
 RWTexture2D<float4> OutCopyTexture0;
 
 // Read a 4x4 color block ready for BC1 compression
-void ReadBlockRGB( Texture2D<float4> SourceTexture, SamplerState TextureSampler, float2 UV, out float3 Block[16] )
+void ReadBlockRGB( Texture2D<float4> SourceTexture, SamplerState TextureSampler, float2 UV, float2 TexelUVSize, out float3 Block[16] )
 {
 	{
-		float4 Red = SourceTexture.GatherRed(TextureSampler, UV, int2(0,0));
-		float4 Green = SourceTexture.GatherGreen(TextureSampler, UV, int2(0,0));
-		float4 Blue = SourceTexture.GatherBlue(TextureSampler, UV, int2(0,0));
+		float4 Red = SourceTexture.GatherRed(TextureSampler, UV);
+		float4 Green = SourceTexture.GatherGreen(TextureSampler, UV);
+		float4 Blue = SourceTexture.GatherBlue(TextureSampler, UV);
 		Block[0] = float3(Red[3], Green[3], Blue[3]);
 		Block[1] = float3(Red[2], Green[2], Blue[2]);
 		Block[4] = float3(Red[0], Green[0], Blue[0]);
 		Block[5] = float3(Red[1], Green[1], Blue[1]);
 	}
 	{
-		float4 Red = SourceTexture.GatherRed(TextureSampler, UV, int2(2,0));
-		float4 Green = SourceTexture.GatherGreen(TextureSampler, UV, int2(2,0));
-		float4 Blue = SourceTexture.GatherBlue(TextureSampler, UV, int2(2,0));
+		float2 UVOffset = UV + float2(2.f * TexelUVSize.x, 0);
+		float4 Red = SourceTexture.GatherRed(TextureSampler, UVOffset);
+		float4 Green = SourceTexture.GatherGreen(TextureSampler, UVOffset);
+		float4 Blue = SourceTexture.GatherBlue(TextureSampler, UVOffset);
 		Block[2] = float3(Red[3], Green[3], Blue[3]);
 		Block[3] = float3(Red[2], Green[2], Blue[2]);
 		Block[6] = float3(Red[0], Green[0], Blue[0]);
 		Block[7] = float3(Red[1], Green[1], Blue[1]);
 	}
 	{
-		float4 Red = SourceTexture.GatherRed(TextureSampler, UV, int2(0,2));
-		float4 Green = SourceTexture.GatherGreen(TextureSampler, UV, int2(0,2));
-		float4 Blue = SourceTexture.GatherBlue(TextureSampler, UV, int2(0,2));
+		float2 UVOffset = UV + float2(0, 2.f * TexelUVSize.y);
+		float4 Red = SourceTexture.GatherRed(TextureSampler, UVOffset);
+		float4 Green = SourceTexture.GatherGreen(TextureSampler, UVOffset);
+		float4 Blue = SourceTexture.GatherBlue(TextureSampler, UVOffset);
 		Block[8] = float3(Red[3], Green[3], Blue[3]);
 		Block[9] = float3(Red[2], Green[2], Blue[2]);
 		Block[12] = float3(Red[0], Green[0], Blue[0]);
 		Block[13] = float3(Red[1], Green[1], Blue[1]);
 	}
 	{
-		float4 Red = SourceTexture.GatherRed(TextureSampler, UV, int2(2,2));
-		float4 Green = SourceTexture.GatherGreen(TextureSampler, UV, int2(2,2));
-		float4 Blue = SourceTexture.GatherBlue(TextureSampler, UV, int2(2,2));
+		float2 UVOffset = UV + 2.f * TexelUVSize;
+		float4 Red = SourceTexture.GatherRed(TextureSampler, UVOffset);
+		float4 Green = SourceTexture.GatherGreen(TextureSampler, UVOffset);
+		float4 Blue = SourceTexture.GatherBlue(TextureSampler, UVOffset);
 		Block[10] = float3(Red[3], Green[3], Blue[3]);
 		Block[11] = float3(Red[2], Green[2], Blue[2]);
 		Block[14] = float3(Red[0], Green[0], Blue[0]);
@@ -62,35 +65,38 @@ void ReadBlockRGB( Texture2D<float4> SourceTexture, SamplerState TextureSampler,
 }
 
 // Read a 4x4 block of XY channels from a normal texture ready for BC5 compression
-void ReadBlockRG( Texture2D<float4> SourceTexture, SamplerState TextureSampler, float2 UV, out float BlockR[16], out float BlockG[16] )
+void ReadBlockRG( Texture2D<float4> SourceTexture, SamplerState TextureSampler, float2 UV, float2 TexelUVSize, out float BlockR[16], out float BlockG[16] )
 {
 	{
-		float4 Red = SourceTexture.GatherRed(TextureSampler, UV, int2(0,0));
-		float4 Green = SourceTexture.GatherGreen(TextureSampler, UV, int2(0,0));
+		float4 Red = SourceTexture.GatherRed(TextureSampler, UV);
+		float4 Green = SourceTexture.GatherGreen(TextureSampler, UV);
 		BlockR[0] = Red[3]; BlockG[0] = Green[3];
 		BlockR[1] = Red[2]; BlockG[1] = Green[2];
 		BlockR[4] = Red[0]; BlockG[4] = Green[0];
 		BlockR[5] = Red[1]; BlockG[5] = Green[1];
 	}
 	{
-		float4 Red = SourceTexture.GatherRed(TextureSampler, UV, int2(2,0));
-		float4 Green = SourceTexture.GatherGreen(TextureSampler, UV, int2(2,0));
+		float2 UVOffset = UV + float2(2.f * TexelUVSize.x, 0);
+		float4 Red = SourceTexture.GatherRed(TextureSampler, UVOffset);
+		float4 Green = SourceTexture.GatherGreen(TextureSampler, UVOffset);
 		BlockR[2] = Red[3]; BlockG[2] = Green[3];
 		BlockR[3] = Red[2]; BlockG[3] = Green[2];
 		BlockR[6] = Red[0]; BlockG[6] = Green[0];
 		BlockR[7] = Red[1]; BlockG[7] = Green[1];
 	}
 	{
-		float4 Red = SourceTexture.GatherRed(TextureSampler, UV, int2(0,2));
-		float4 Green = SourceTexture.GatherGreen(TextureSampler, UV, int2(0,2));
+		float2 UVOffset = UV + float2(0, 2.f * TexelUVSize.y);
+		float4 Red = SourceTexture.GatherRed(TextureSampler, UVOffset);
+		float4 Green = SourceTexture.GatherGreen(TextureSampler, UVOffset);
 		BlockR[8]  = Red[3]; BlockG[8]  = Green[3];
 		BlockR[9]  = Red[2]; BlockG[9]  = Green[2];
 		BlockR[12] = Red[0]; BlockG[12] = Green[0];
 		BlockR[13] = Red[1]; BlockG[13] = Green[1];
 	}
 	{
-		float4 Red = SourceTexture.GatherRed(TextureSampler, UV, int2(2,2));
-		float4 Green = SourceTexture.GatherGreen(TextureSampler, UV, int2(2,2));
+		float2 UVOffset = UV + 2.f * TexelUVSize;
+		float4 Red = SourceTexture.GatherRed(TextureSampler, UVOffset);
+		float4 Green = SourceTexture.GatherGreen(TextureSampler, UVOffset);
 		BlockR[10] = Red[3]; BlockG[10] = Green[3];
 		BlockR[11] = Red[2]; BlockG[11] = Green[2];
 		BlockR[14] = Red[0]; BlockG[14] = Green[0];
@@ -117,10 +123,11 @@ void CompressBaseColorCS(
 	if (any(ThreadId.xy >= DestRect.zw))
 		return;
 
-	float2 SampleUV = float2(SamplePos) / float2(DestRect.zw);
+	float2 TexelUVSize = 1.f / float2(DestRect.zw);
+	float2 SampleUV = (float2(SamplePos) + 0.5f) * TexelUVSize;
 
 	float3 BlockBaseColor[16];
-	ReadBlockRGB(RenderTexture0, TextureSampler0, SampleUV, BlockBaseColor);
+	ReadBlockRGB(RenderTexture0, TextureSampler0, SampleUV, TexelUVSize, BlockBaseColor);
 
 	OutCompressTexture0[ThreadId.xy + DestRect.xy] = CompressBC1Block_SRGB(BlockBaseColor);
 }
@@ -135,14 +142,15 @@ void CompressBaseColorNormalCS(
 	if (any(SamplePos >= DestRect.zw))
 		return;
 
-	float2 SampleUV = float2(SamplePos) / float2(DestRect.zw);
+	float2 TexelUVSize = 1.f / float2(DestRect.zw);
+	float2 SampleUV = (float2(SamplePos) + 0.5f) * TexelUVSize;
 
 	float3 BlockBaseColor[16];
-	ReadBlockRGB(RenderTexture0, TextureSampler0, SampleUV, BlockBaseColor);
+	ReadBlockRGB(RenderTexture0, TextureSampler0, SampleUV, TexelUVSize, BlockBaseColor);
 
 	float BlockNormalX[16];
 	float BlockNormalY[16];
-	ReadBlockRG(RenderTexture1, TextureSampler1, SampleUV, BlockNormalX, BlockNormalY);
+	ReadBlockRG(RenderTexture1, TextureSampler1, SampleUV, TexelUVSize, BlockNormalX, BlockNormalY);
 
 	OutCompressTexture0[ThreadId.xy + DestRect.xy] = CompressBC1Block_SRGB(BlockBaseColor);
 	OutCompressTexture1[ThreadId.xy + DestRect.xy] = CompressBC5Block(BlockNormalX, BlockNormalY);
@@ -158,17 +166,18 @@ void CompressBaseColorNormalSpecularCS(
 	if (any(SamplePos >= DestRect.zw))
 		return;
 
-	float2 SampleUV = float2(SamplePos) / float2(DestRect.zw);
+	float2 TexelUVSize = 1.f / float2(DestRect.zw);
+	float2 SampleUV = (float2(SamplePos) + 0.5f) * TexelUVSize;
 
 	float3 BlockBaseColor[16];
-	ReadBlockRGB(RenderTexture0, TextureSampler0, SampleUV, BlockBaseColor);
+	ReadBlockRGB(RenderTexture0, TextureSampler0, SampleUV, TexelUVSize, BlockBaseColor);
 
 	float BlockNormalX[16];
 	float BlockNormalY[16];
-	ReadBlockRG(RenderTexture1, TextureSampler1, SampleUV, BlockNormalX, BlockNormalY);
+	ReadBlockRG(RenderTexture1, TextureSampler1, SampleUV, TexelUVSize, BlockNormalX, BlockNormalY);
 
 	float3 BlockSpecular[16];
-	ReadBlockRGB(RenderTexture2, TextureSampler2, SampleUV, BlockSpecular);
+	ReadBlockRGB(RenderTexture2, TextureSampler2, SampleUV, TexelUVSize, BlockSpecular);
 
 	CombineNormalSpecularBlocksForBC3(BlockNormalY, BlockSpecular);
 
