Cherry-pick 7097701: Fixes regression in SSGI

diff --git a/Engine/Shaders/Private/SSRT/SSRTDiffuseIndirect.usf b/Engine/Shaders/Private/SSRT/SSRTDiffuseIndirect.usf
--- a/Engine/Shaders/Private/SSRT/SSRTDiffuseIndirect.usf
+++ b/Engine/Shaders/Private/SSRT/SSRTDiffuseIndirect.usf
@@ -41,7 +41,7 @@ void MainCS(
 {
 	uint2 PixelPos = DispatchThreadId + View.ViewRectMin.xy;
 	float2 BufferUV = (PixelPos + 0.5) * View.BufferSizeAndInvSize.zw;
-	float2 ScreenPos = (DispatchThreadId + 0.5) * View.ViewSizeAndInvSize.zw;
+    float2 ScreenPos = ViewportUVToScreenPos((DispatchThreadId + 0.5) * View.ViewSizeAndInvSize.zw);
 
 	FGBufferData GBuffer = GetGBufferDataFromSceneTextures(BufferUV);
 
