Cherry-pick 7076114: Adding methods for WorldPosition reconstruction, primary ray creation, and camera-relative depth biasing

diff --git a/Engine/Shaders/Private/RayTracing/RayTracingCommon.ush b/Engine/Shaders/Private/RayTracing/RayTracingCommon.ush
--- a/Engine/Shaders/Private/RayTracing/RayTracingCommon.ush
+++ b/Engine/Shaders/Private/RayTracing/RayTracingCommon.ush
@@ -73,6 +73,28 @@ FRayCone PropagateRayCone(in FRayCone Cone, in float SurfaceSpreadAngle, in floa
 	return NewCone;
 }
 
+float3 ReconstructWorldPositionFromDeviceZ(uint2 PixelCoord, float DeviceZ)
+{
+	float4 TranslatedWorldPosition = mul(float4(PixelCoord + 0.5, DeviceZ, 1), View.SVPositionToTranslatedWorld);
+	TranslatedWorldPosition.xyz /= TranslatedWorldPosition.w;
+	return TranslatedWorldPosition.xyz + View.WorldCameraOrigin;
+}
+
+void ReconstructWorldPositionAndCameraDirectionFromDeviceZ(uint2 PixelCoord, float DeviceZ, out float3 OutWorldPosition, out float3 OutCameraDirection)
+{
+	float4 TranslatedWorldPosition = mul(float4(PixelCoord + 0.5, DeviceZ, 1), View.SVPositionToTranslatedWorld);
+	TranslatedWorldPosition.xyz /= TranslatedWorldPosition.w;
+	OutWorldPosition = TranslatedWorldPosition.xyz + View.WorldCameraOrigin;
+	OutCameraDirection = normalize(TranslatedWorldPosition.xyz - View.TranslatedWorldCameraOrigin);
+}
+
+float3 ReconstructTranslatedWorldPositionFromDeviceZ(uint2 PixelCoord, float DeviceZ)
+{
+	float4 TranslatedWorldPosition = mul(float4(PixelCoord + 0.5, DeviceZ, 1), View.SVPositionToTranslatedWorld);
+	TranslatedWorldPosition.xyz /= TranslatedWorldPosition.w;
+	return TranslatedWorldPosition.xyz;
+}
+
 float3 ReconstructWorldPositionFromDepth(float2 UV, float Depth)
 {
 	float2 ScreenPosition = (UV - View.ScreenPositionScaleBias.wz) / View.ScreenPositionScaleBias.xy;
@@ -140,6 +162,22 @@ struct FMaterialClosestHitPayload : FMinimalPayload
 #define USE_PACKED_PAYLOAD 1
 
 #if !COMPUTESHADER
+
+RayDesc CreatePrimaryRay(float2 UV)
+{
+	float2 ScreenPosition = (UV - View.ScreenPositionScaleBias.wz) / View.ScreenPositionScaleBias.xy;
+	float3 RayStart_TranslatedWorld = View.TranslatedWorldCameraOrigin;
+	float4 RayEnd_TranslatedWorld = mul(float4(ScreenPosition, 1, 1), View.ScreenToTranslatedWorld);
+	RayEnd_TranslatedWorld *= rcp(RayEnd_TranslatedWorld.w);
+
+	RayDesc Ray;
+	Ray.Origin = View.WorldCameraOrigin;
+	Ray.Direction = normalize(RayEnd_TranslatedWorld.xyz - RayStart_TranslatedWorld);
+	Ray.TMin = 0.0;
+	Ray.TMax = 1.0e27;
+	return Ray;
+}
+
 #if USE_PACKED_PAYLOAD
 struct FPackedMaterialClosestHitPayload : FMinimalPayload
 {
@@ -288,6 +326,21 @@ void ApplyPositionBias(inout RayDesc Ray, const float3 WorldNormal, const float
 	Ray.Origin += WorldNormal * NormalBias;
 }
 
+void ApplyCameraRelativeDepthBias(inout RayDesc Ray, uint2 PixelCoord, float DeviceZ, const float3 WorldNormal, const float AbsoluteNormalBias)
+{
+	float3 WorldPosition = ReconstructTranslatedWorldPositionFromDeviceZ(PixelCoord, DeviceZ);
+	float3 CameraDirection = WorldPosition - View.TranslatedWorldCameraOrigin;
+	float DistanceToCamera = length(CameraDirection);
+	CameraDirection = normalize(CameraDirection);
+	float Epsilon = 1.0e-4;
+	float RelativeBias = DistanceToCamera * Epsilon;
+	//float ProjectedBias = RelativeBias / dot(Ray.Direction, WorldNormal);
+
+	float RayBias = max(RelativeBias, AbsoluteNormalBias);
+	Ray.Origin -= CameraDirection * RayBias;
+	ApplyPositionBias(Ray, WorldNormal, RayBias);
+}
+
 #endif // !COMPUTESHADER
 
 #endif // RAYTRACINGCOMMON_USH_INCLUDED // Workarround for UE-66460
