Cherry-pick 7076124: Integrate with CreatePrimaryRay

diff --git a/Engine/Shaders/Private/PathTracing/Sensor/PathTracingAdaptivePinholeCamera.ush b/Engine/Shaders/Private/PathTracing/Sensor/PathTracingAdaptivePinholeCamera.ush
--- a/Engine/Shaders/Private/PathTracing/Sensor/PathTracingAdaptivePinholeCamera.ush
+++ b/Engine/Shaders/Private/PathTracing/Sensor/PathTracingAdaptivePinholeCamera.ush
@@ -119,11 +119,7 @@ void AdaptivePinholeCamera_SampleEmitter(
 	// Correct for ViewSize/BufferSize discrepancy
 	MipPixelUV *= View.ViewSizeAndInvSize.xy * View.BufferSizeAndInvSize.zw;
 
-	OutRay.Origin = ReconstructWorldPositionFromDepth(MipPixelUV, 0.0f);
-	OutRay.Direction = normalize(ReconstructWorldPositionFromDepth(MipPixelUV, 1.f) - OutRay.Origin);
-	OutRay.TMin = 0.01;
-	OutRay.TMax = 1e12;
-
+	OutRay = CreatePrimaryRay(MipPixelUV);
 	OutPixel = MipPixelUV * View.BufferSizeAndInvSize.xy;
 }
 
diff --git a/Engine/Shaders/Private/PathTracing/Sensor/PathTracingPinholeCamera.ush b/Engine/Shaders/Private/PathTracing/Sensor/PathTracingPinholeCamera.ush
--- a/Engine/Shaders/Private/PathTracing/Sensor/PathTracingPinholeCamera.ush
+++ b/Engine/Shaders/Private/PathTracing/Sensor/PathTracingPinholeCamera.ush
@@ -20,11 +20,7 @@ void PinholeCamera_SampleEmitter(
 	float2 RandSample = RandomSequence_GenerateSample2D(RandSequence, SampleIndex);
 	float2 ViewportUV = (LaunchIndex.xy + RandSample.xy) * View.BufferSizeAndInvSize.zw;
 
-	OutRay.Origin = ReconstructWorldPositionFromDepth(ViewportUV, 0.0f);
-	OutRay.Direction = normalize(ReconstructWorldPositionFromDepth(ViewportUV, 1.f) - OutRay.Origin);
-	OutRay.TMin = 0.01;
-	OutRay.TMax = 1e12;
-
+	OutRay = CreatePrimaryRay(ViewportUV);
 	OutPixel = LaunchIndex.xy;
 	OutPdf = 1.0;
 }
@@ -35,4 +31,4 @@ void PinholeCamera_PdfEmitter(
 )
 {
 	OutPdf = 1.0;
-}
\ No newline at end of file
+}
diff --git a/Engine/Shaders/Private/RayTracing/RayTracingBarycentrics.usf b/Engine/Shaders/Private/RayTracing/RayTracingBarycentrics.usf
--- a/Engine/Shaders/Private/RayTracing/RayTracingBarycentrics.usf
+++ b/Engine/Shaders/Private/RayTracing/RayTracingBarycentrics.usf
@@ -19,13 +19,7 @@ void RayTracingBarycentricsMainRGS()
 
 	float2 RenderTargetUV = (float2(PixelCoord) + .5f) * View.BufferSizeAndInvSize.zw;
 
-	RayDesc Ray;
-
-	Ray.Origin = View.WorldCameraOrigin;
-	Ray.Direction = normalize(ReconstructWorldPositionFromDepth(RenderTargetUV, 10000.f) - Ray.Origin);
-	Ray.TMin = 0.5f; // 1 unit in UE4 = 1 cm -0> 0.5 ==> 5 mm
-	Ray.TMax = 1e27f;
-
+	RayDesc Ray = CreatePrimaryRay(RenderTargetUV);
 	uint RayFlags = 0
 	              | RAY_FLAG_CULL_BACK_FACING_TRIANGLES      // use back face culling
 	              | RAY_FLAG_FORCE_OPAQUE;                   // don't run anyhit shader
diff --git a/Engine/Shaders/Private/RayTracing/RayTracingDebug.usf b/Engine/Shaders/Private/RayTracing/RayTracingDebug.usf
--- a/Engine/Shaders/Private/RayTracing/RayTracingDebug.usf
+++ b/Engine/Shaders/Private/RayTracing/RayTracingDebug.usf
@@ -20,13 +20,7 @@ void RayTracingDebugMainRGS()
 
 	float2 RenderTargetUV = (float2(PixelCoord) + .5f) * View.BufferSizeAndInvSize.zw;
 
-	RayDesc Ray;
-
-	Ray.Origin = View.WorldCameraOrigin;
-	Ray.Direction = normalize(ReconstructWorldPositionFromDepth(RenderTargetUV, 10000.f) - Ray.Origin);
-	Ray.TMin = 0.5f; // 1 unit in UE4 = 1 cm -0> 0.5 ==> 5 mm
-	Ray.TMax = 1e27f;
-
+	RayDesc Ray = CreatePrimaryRay(RenderTargetUV);
     FRayCone RayCone = (FRayCone)0;
 	RayCone.SpreadAngle = View.EyeToPixelSpreadAngle;
 
diff --git a/Engine/Shaders/Private/RayTracing/RayTracingTranslucency.usf b/Engine/Shaders/Private/RayTracing/RayTracingTranslucency.usf
--- a/Engine/Shaders/Private/RayTracing/RayTracingTranslucency.usf
+++ b/Engine/Shaders/Private/RayTracing/RayTracingTranslucency.usf
@@ -111,17 +111,10 @@ void RayTracingTranslucencyRGS()
 	float2 UV = (float2(PixelCoord) + 0.5) * InvBufferSize;
 	FGBufferData GBufferData = GetGBufferDataFromSceneTextures(UV);
 	float Depth = GBufferData.Depth;
-	float3 CameraOrigin = View.WorldCameraOrigin;
-	float3 WorldPosition = ReconstructWorldPositionFromDepth(UV, Depth);
-	float3 CameraDirection = normalize(WorldPosition - CameraOrigin);	
-	const float GBufferHitT = length(WorldPosition - CameraOrigin);
 
 	// Trace rays from camera origin to (Gbuffer - epsilon) to only intersect translucent objects
-	RayDesc Ray;
-	Ray.Origin = CameraOrigin;
-	Ray.TMin = 0.0;
-	Ray.TMax = GBufferHitT - 0.1;
-	Ray.Direction = CameraDirection;
+	RayDesc Ray = CreatePrimaryRay(UV);
+	Ray.TMax = Depth - 0.1;
 	FRayCone RayCone = (FRayCone)0;
 	RayCone.SpreadAngle = View.EyeToPixelSpreadAngle;
 
