Cherry-pick 7076127: Integrating with WorldPosition and ray bias changes

diff --git a/Engine/Shaders/Private/RayTracing/RayTracingAmbientOcclusionRGS.usf b/Engine/Shaders/Private/RayTracing/RayTracingAmbientOcclusionRGS.usf
--- a/Engine/Shaders/Private/RayTracing/RayTracingAmbientOcclusionRGS.usf
+++ b/Engine/Shaders/Private/RayTracing/RayTracingAmbientOcclusionRGS.usf
@@ -55,7 +55,8 @@ void AmbientOcclusionRGS()
 	float2 UV = (float2(PixelCoord) + 0.5) * InvBufferSize;
 	FGBufferData GBufferData = GetGBufferDataFromSceneTextures(UV);	
 	float Depth = GBufferData.Depth;
-	float3 WorldPosition = ReconstructWorldPositionFromDepth(UV, Depth);
+	float DeviceZ = SceneDepthBuffer.Load(int3(PixelCoord, 0)).r;
+	float3 WorldPosition = ReconstructWorldPositionFromDeviceZ(PixelCoord, DeviceZ);
 	float3 WorldNormal = GBufferData.WorldNormal;
 
 	float RayCount = 0.0;
@@ -66,7 +67,6 @@ void AmbientOcclusionRGS()
 	uint SamplesPerPixelLocal = SamplesPerPixel;
 	
 	// Mask out depth values beyond far plane
-	float DeviceZ = SceneDepthBuffer.Load(int3(PixelCoord, 0)).r;
 	bool IsFiniteDepth = DeviceZ > 0.0;
 	bool bTraceRay = (IsFiniteDepth && GBufferData.ShadingModelID != SHADINGMODELID_UNLIT);
 	if (!bTraceRay)
@@ -80,7 +80,7 @@ void AmbientOcclusionRGS()
 	{
 		RayDesc Ray;
 		GenerateCosineNormalRay(WorldPosition, WorldNormal, RandSequence, Ray.Origin, Ray.Direction, Ray.TMin, Ray.TMax);
-		ApplyPositionBias(Ray, WorldNormal, MaxNormalBias);
+		ApplyCameraRelativeDepthBias(Ray, PixelCoord, DeviceZ, WorldNormal, MaxNormalBias);
 
 		if (dot(WorldNormal, Ray.Direction) <= 0.0) // TODO: does it needs to be handled by the denoiser?
 			continue;
diff --git a/Engine/Shaders/Private/RayTracing/RayTracingGlobalIlluminationRGS.usf b/Engine/Shaders/Private/RayTracing/RayTracingGlobalIlluminationRGS.usf
--- a/Engine/Shaders/Private/RayTracing/RayTracingGlobalIlluminationRGS.usf
+++ b/Engine/Shaders/Private/RayTracing/RayTracingGlobalIlluminationRGS.usf
@@ -74,10 +74,10 @@ void GlobalIlluminationRGS()
 	{
 		GBufferData.DiffuseColor = GBufferData.StoredBaseColor;
 	}
-	float Depth = GBufferData.Depth;
-	float3 WorldPosition = ReconstructWorldPositionFromDepth(UV, Depth);
-	float3 CameraOrigin = View.WorldCameraOrigin;
-	float3 CameraDirection = normalize(WorldPosition - CameraOrigin);
+	float DeviceZ = SceneDepthBuffer.Load(int3(PixelCoord, 0)).r;
+	float3 WorldPosition;
+	float3 CameraDirection;
+	ReconstructWorldPositionAndCameraDirectionFromDeviceZ(PixelCoord, DeviceZ, WorldPosition, CameraDirection);
 	float3 WorldNormal = GBufferData.WorldNormal;
 	uint ShadingModelID = GBufferData.ShadingModelID;
 	if (ShadingModelID == SHADINGMODELID_UNLIT
@@ -141,7 +141,7 @@ void GlobalIlluminationRGS()
 		float3 RayThroughput = 1.0;
 #endif
 		Ray.TMax = max(MaxRayDistanceForGI, MaxRayDistanceForAO);
-        ApplyPositionBias(Ray, WorldNormal, MaxNormalBias);
+		ApplyCameraRelativeDepthBias(Ray, PixelCoord, DeviceZ, WorldNormal, MaxNormalBias);
 		
 		float MaterialPdf = 0.0;
 		uint Bounce = 0;
diff --git a/Engine/Shaders/Private/RayTracing/RayTracingOcclusionRGS.usf b/Engine/Shaders/Private/RayTracing/RayTracingOcclusionRGS.usf
--- a/Engine/Shaders/Private/RayTracing/RayTracingOcclusionRGS.usf
+++ b/Engine/Shaders/Private/RayTracing/RayTracingOcclusionRGS.usf
@@ -134,14 +134,13 @@ void OcclusionRGS()
 		float2 InvBufferSize = View.BufferSizeAndInvSize.zw;
 		float2 BufferUV = (float2(PixelCoord) + 0.5) * InvBufferSize;
 		FGBufferData GBufferData = GetGBufferDataFromSceneTextures(BufferUV);
-		float Depth = GBufferData.Depth;
-		float3 WorldPosition = ReconstructWorldPositionFromDepth(BufferUV, Depth);
+		float DeviceZ = SceneDepthBuffer.SampleLevel(SceneDepthBufferSampler, BufferUV, 0).r;
+		float3 WorldPosition = ReconstructWorldPositionFromDeviceZ(PixelCoord, DeviceZ);
 		float3 WorldNormal = GBufferData.WorldNormal;
 
 		uint LightChannel = GetSceneLightingChannel(PixelCoord);
 
 		// Mask out depth values that are infinitely far away
-		float DeviceZ = SceneDepthBuffer.Load(int3(PixelCoord, 0)).r;
 		bool IsFiniteDepth = DeviceZ > 0.0;
 		bool bTraceRay = (
 			IsFiniteDepth &&
@@ -192,7 +191,7 @@ void OcclusionRGS()
 				/* out */ Ray.Direction,
 				/* out */ Ray.TMin,
 				/* out */ Ray.TMax);
-			ApplyPositionBias(Ray, WorldNormal, NormalBias);
+			ApplyCameraRelativeDepthBias(Ray, PixelCoord, DeviceZ, WorldNormal, NormalBias);
 
 			BRANCH
 			if (!bIsValidRay)
diff --git a/Engine/Shaders/Private/RayTracing/RayTracingRectLightRGS.usf b/Engine/Shaders/Private/RayTracing/RayTracingRectLightRGS.usf
--- a/Engine/Shaders/Private/RayTracing/RayTracingRectLightRGS.usf
+++ b/Engine/Shaders/Private/RayTracing/RayTracingRectLightRGS.usf
@@ -164,10 +164,10 @@ void RectLightRGS()
 	float2 UV = (float2(PixelCoord) + 0.5) * InvBufferSize;
 	FScreenSpaceData ScreenSpaceData = GetScreenSpaceData(UV);
 	FGBufferData GBuffer = ScreenSpaceData.GBuffer;
-	float Depth = ScreenSpaceData.GBuffer.Depth;
-	float3 WorldPosition = ReconstructWorldPositionFromDepth(UV, Depth);
-	float3 CameraOrigin = View.WorldCameraOrigin;
-	float3 CameraDirection = normalize(WorldPosition - CameraOrigin);
+	float DeviceZ = SceneTexturesStruct.SceneDepthTexture.Load(int3(PixelCoord, 0)).r;
+	float3 WorldPosition;
+	float3 CameraDirection;
+	ReconstructWorldPositionAndCameraDirectionFromDeviceZ(PixelCoord, DeviceZ, WorldPosition, CameraDirection);
 	float3 WorldNormal = ScreenSpaceData.GBuffer.WorldNormal;
 	float3 BaseColor = ScreenSpaceData.GBuffer.BaseColor;
 	float3 DiffuseColor = ScreenSpaceData.GBuffer.DiffuseColor;
@@ -179,7 +179,6 @@ void RectLightRGS()
 	uint SamplesPerPixel = RectLight.SamplesPerPixel;
 
 	// Mask out depth values that are infinitely far away
-	float DeviceZ = SceneTexturesStruct.SceneDepthTexture.Load(int3(PixelCoord, 0)).r;
 	bool IsFiniteDepth = DeviceZ > 0.0;
 	bool bTraceRay = (
 		IsFiniteDepth &&
@@ -256,7 +255,7 @@ void RectLightRGS()
 		{
 			continue;
 		}
-		ApplyPositionBias(Ray, WorldNormal, RectLight.MaxNormalBias);
+		ApplyCameraRelativeDepthBias(Ray, PixelCoord, DeviceZ, WorldNormal, RectLight.MaxNormalBias);
 
 		float NoL = dot(WorldNormal, Ray.Direction);
 		if (NoL <= 0.0)
diff --git a/Engine/Shaders/Private/RayTracing/RayTracingReflections.usf b/Engine/Shaders/Private/RayTracing/RayTracingReflections.usf
--- a/Engine/Shaders/Private/RayTracing/RayTracingReflections.usf
+++ b/Engine/Shaders/Private/RayTracing/RayTracingReflections.usf
@@ -214,9 +214,9 @@ void RayTracingReflectionsRGS()
 			float3 PathThroughput = float3(1.0f, 1.0f, 1.0f);
 			
 			// Initialize path states according to GBuffer
-			float3 CameraOrigin = View.WorldCameraOrigin;
-			float3 WorldPosition = ReconstructWorldPositionFromDepth(UV, Depth);
-			float3 CameraDirection = normalize(WorldPosition - CameraOrigin);
+			float3 WorldPosition;
+			float3 CameraDirection;
+			ReconstructWorldPositionAndCameraDirectionFromDeviceZ(PixelCoord, DeviceZ, WorldPosition, CameraDirection);
 			float3 IncidentDirection = CameraDirection;
 
 			uint ShadingModelID = GBufferData.ShadingModelID;
@@ -283,7 +283,7 @@ void RayTracingReflectionsRGS()
 				TopLayerRay.TMax = LocalMaxRayDistance;
 
 				TopLayerRay.Direction = GenerateReflectionDirection(RandSequence, DispatchThreadId, BounceIndex, LocalSamplesPerPixel, IncidentDirection, TopLayerWorldNormal, TopLayerRoughness);
-				ApplyPositionBias(TopLayerRay, TopLayerWorldNormal, ReflectionMaxNormalBias);
+				ApplyCameraRelativeDepthBias(TopLayerRay, PixelCoord, DeviceZ, TopLayerWorldNormal, ReflectionMaxNormalBias);
 				if (BounceIndex == 0)
 				{
 					// This fixes black reflexions due to R pointing into the surface.
@@ -641,8 +641,9 @@ void RayTracingReflectionsRGS()
 		// Imaginary depth computation for Nvidia's denoiser.
 		// TODO: shader permutation?
 		{
-			float3 OriginalWorldPos = ReconstructWorldPositionFromDepth(UV, Depth);
-			float3 CameraToPixel = normalize(OriginalWorldPos - View.WorldCameraOrigin.xyz);
+			float3 OriginalWorldPos;
+			float3 CameraToPixel;
+			ReconstructWorldPositionAndCameraDirectionFromDeviceZ(PixelCoord, DeviceZ, OriginalWorldPos, CameraToPixel);
 
 			float4 ImaginaryWorldPos = float4(OriginalWorldPos + CameraToPixel * HitDistance, 1.0f);
 			float4 ImaginaryClipPos = mul(ImaginaryWorldPos, View.WorldToClip);
diff --git a/Engine/Shaders/Private/RayTracing/RayTracingSkyLightRGS.usf b/Engine/Shaders/Private/RayTracing/RayTracingSkyLightRGS.usf
--- a/Engine/Shaders/Private/RayTracing/RayTracingSkyLightRGS.usf
+++ b/Engine/Shaders/Private/RayTracing/RayTracingSkyLightRGS.usf
@@ -231,10 +231,10 @@ void SkyLightRGS()
 	float2 UV = (float2(PixelCoord) + 0.5) * InvBufferSize;
 	FScreenSpaceData ScreenSpaceData = GetScreenSpaceData(UV);
 	FGBufferData GBufferData = GetGBufferDataFromSceneTextures(UV);
-	float Depth = GBufferData.Depth;
-	float3 WorldPosition = ReconstructWorldPositionFromDepth(UV, Depth);
-	float3 CameraOrigin = View.WorldCameraOrigin;
-	float3 CameraDirection = normalize(WorldPosition - CameraOrigin);
+	float DeviceZ = SceneDepthBuffer.Load(int3(PixelCoord, 0)).r;
+	float3 WorldPosition;
+	float3 CameraDirection;
+	ReconstructWorldPositionAndCameraDirectionFromDeviceZ(PixelCoord, DeviceZ, WorldPosition, CameraDirection);
 	float3 WorldNormal = GBufferData.WorldNormal;
 	float3 Albedo = GBufferData.DiffuseColor;
 
@@ -252,7 +252,6 @@ void SkyLightRGS()
 	uint SamplesPerPixel = SkyLight.SamplesPerPixel;
 
 	// Mask out depth values that are infinitely far away
-	float DeviceZ = SceneDepthBuffer.Load(int3(PixelCoord, 0)).r;
 	bool IsFiniteDepth = DeviceZ > 0.0;
 	bool bTraceRay = (
 		IsFiniteDepth &&
@@ -301,7 +300,7 @@ void SkyLightRGS()
 			MisWeight = RayPdf / ((1.0 - SkyLightSamplingStrategyPdf) * RayPdf + SkyLightSamplingStrategyPdf * PdfSkyLightOcclusionRay(FaceCdf, WorldNormal, Ray.Direction));
 		}
 
-		ApplyPositionBias(Ray, WorldNormal, SkyLight.MaxNormalBias);
+		ApplyCameraRelativeDepthBias(Ray, PixelCoord, DeviceZ, WorldNormal, SkyLight.MaxNormalBias);
 
 		uint RayFlags = 0;
 #if !ENABLE_MATERIALS
