Cherry-pick 7096110: Parameterize DirectionalLight next-event ray to be consistent with other "infinite" rays

diff --git a/Engine/Shaders/Private/PathTracing/Light/PathTracingDirectionalLight.ush b/Engine/Shaders/Private/PathTracing/Light/PathTracingDirectionalLight.ush
--- a/Engine/Shaders/Private/PathTracing/Light/PathTracingDirectionalLight.ush
+++ b/Engine/Shaders/Private/PathTracing/Light/PathTracingDirectionalLight.ush
@@ -16,9 +16,9 @@ void DirectionalLight_GenerateLightRay(
 )
 {
 	OutLightRay.Origin = Ray.Origin;
-	OutLightRay.TMin = 0.001;
+	OutLightRay.TMin = 0.1;
 	OutLightRay.Direction = normalize(LightUV);
-	OutLightRay.TMax = 1e6;
+	OutLightRay.TMax = 1.0e27;
 }
 
 void DirectionalLight_PdfLight(
