Cherry-pick 7104639: Forcing shadow-variant of CHS for visibility queries

diff --git a/Engine/Shaders/Private/PathTracing/PathTracing.usf b/Engine/Shaders/Private/PathTracing/PathTracing.usf
--- a/Engine/Shaders/Private/PathTracing/PathTracing.usf
+++ b/Engine/Shaders/Private/PathTracing/PathTracing.usf
@@ -179,7 +179,7 @@ void PathTracingMainRG()
 						TLAS,      // AccelerationStructure
 						RayFlags,
 						RAY_TRACING_MASK_SHADOW,
-						RAY_TRACING_SHADER_SLOT_MATERIAL, // RayContributionToHitGroupIndex
+						RAY_TRACING_SHADER_SLOT_SHADOW,   // RayContributionToHitGroupIndex
 						RAY_TRACING_NUM_SHADER_SLOTS,     // MultiplierForGeometryContributionToShaderIndex
 						0,         // MissShaderIndex
 						LightRay,
diff --git a/Engine/Shaders/Private/RayTracing/RayTracingAmbientOcclusionRGS.usf b/Engine/Shaders/Private/RayTracing/RayTracingAmbientOcclusionRGS.usf
--- a/Engine/Shaders/Private/RayTracing/RayTracingAmbientOcclusionRGS.usf
+++ b/Engine/Shaders/Private/RayTracing/RayTracingAmbientOcclusionRGS.usf
@@ -99,7 +99,7 @@ void AmbientOcclusionRGS()
 			TLAS,
 			RayFlags,
 			RAY_TRACING_MASK_SHADOW,          // InstanceInclusionMask
-			RAY_TRACING_SHADER_SLOT_MATERIAL, // RayContributionToHitGroupIndex
+			RAY_TRACING_SHADER_SLOT_SHADOW,   // RayContributionToHitGroupIndex
 			RAY_TRACING_NUM_SHADER_SLOTS,     // MultiplierForGeometryContributionToShaderIndex
 			0,    // MissShaderIndex
 			Ray,
diff --git a/Engine/Shaders/Private/RayTracing/RayTracingGlobalIlluminationRGS.usf b/Engine/Shaders/Private/RayTracing/RayTracingGlobalIlluminationRGS.usf
--- a/Engine/Shaders/Private/RayTracing/RayTracingGlobalIlluminationRGS.usf
+++ b/Engine/Shaders/Private/RayTracing/RayTracingGlobalIlluminationRGS.usf
@@ -238,7 +238,7 @@ void GlobalIlluminationRGS()
 						TLAS,   // AccelerationStructure
 						RayFlags,
 						RAY_TRACING_MASK_OPAQUE,
-						RAY_TRACING_SHADER_SLOT_MATERIAL, // RayContributionToHitGroupIndex
+						RAY_TRACING_SHADER_SLOT_SHADOW,  // RayContributionToHitGroupIndex
 						RAY_TRACING_NUM_SHADER_SLOTS,     // MultiplierForGeometryContributionToShaderIndex
 						0,      // MissShaderIndex
 						LightRay,    // RayDesc
diff --git a/Engine/Shaders/Private/RayTracing/RayTracingRectLightRGS.usf b/Engine/Shaders/Private/RayTracing/RayTracingRectLightRGS.usf
--- a/Engine/Shaders/Private/RayTracing/RayTracingRectLightRGS.usf
+++ b/Engine/Shaders/Private/RayTracing/RayTracingRectLightRGS.usf
@@ -273,7 +273,7 @@ void RectLightRGS()
 			TLAS,      // AccelerationStructure
 			RayFlags,
 			RAY_TRACING_MASK_SHADOW,          // InstanceInclusionMask
-			RAY_TRACING_SHADER_SLOT_MATERIAL, // RayContributionToHitGroupIndex
+			RAY_TRACING_SHADER_SLOT_SHADOW,   // RayContributionToHitGroupIndex
 			RAY_TRACING_NUM_SHADER_SLOTS,     // MultiplierForGeometryContributionToShaderIndex
 			0,         // MissShaderIndex
 			Ray,       // RayDesc
diff --git a/Engine/Shaders/Private/RayTracing/RayTracingSkyLightRGS.usf b/Engine/Shaders/Private/RayTracing/RayTracingSkyLightRGS.usf
--- a/Engine/Shaders/Private/RayTracing/RayTracingSkyLightRGS.usf
+++ b/Engine/Shaders/Private/RayTracing/RayTracingSkyLightRGS.usf
@@ -315,7 +315,7 @@ void SkyLightRGS()
 			TLAS,      // AccelerationStructure
 			RayFlags,
 			RAY_TRACING_MASK_SHADOW,          // InstanceInclusionMask
-			RAY_TRACING_SHADER_SLOT_MATERIAL, // RayContributionToHitGroupIndex
+			RAY_TRACING_SHADER_SLOT_SHADOW,   // RayContributionToHitGroupIndex
 			RAY_TRACING_NUM_SHADER_SLOTS,     // MultiplierForGeometryContributionToShaderIndex
 			0,         // MissShaderIndex
 			Ray,       // RayDesc
