Cherry-pick 7106785: Added an early out path to ray tracing material hit shaders. This is useful when tracing visibility rays against masked geometry, where closest hit shaders only need to fill the hit distance.

Added 8 bit flags field to the packed ray tracing payload.

diff --git a/Engine/Shaders/Private/RayTracing/RayTracingAmbientOcclusionRGS.usf b/Engine/Shaders/Private/RayTracing/RayTracingAmbientOcclusionRGS.usf
--- a/Engine/Shaders/Private/RayTracing/RayTracingAmbientOcclusionRGS.usf
+++ b/Engine/Shaders/Private/RayTracing/RayTracingAmbientOcclusionRGS.usf
@@ -94,6 +94,7 @@ void AmbientOcclusionRGS()
 #endif
 
 		FPackedMaterialClosestHitPayload Payload = (FPackedMaterialClosestHitPayload)0;
+		Payload.SetMinimalPayloadMode(); // Only HitT will be filled by closest hit shader, skipping full material evaluation.
 
 		TraceRay(
 			TLAS,
diff --git a/Engine/Shaders/Private/RayTracing/RayTracingCommon.ush b/Engine/Shaders/Private/RayTracing/RayTracingCommon.ush
--- a/Engine/Shaders/Private/RayTracing/RayTracingCommon.ush
+++ b/Engine/Shaders/Private/RayTracing/RayTracingCommon.ush
@@ -126,6 +126,8 @@ uint2 GetPixelCoord(uint2 DispatchThreadId, uint UpscaleFactor)
 #define RAY_TRACING_BLEND_MODE_MODULATE				4
 #define RAY_TRACING_BLEND_MODE_ALPHA_HOLDOUT		5
 
+#define RAY_TRACING_PAYLOAD_FLAG_FRONT_FACE         1 // Indicates that ray has hit front face of a primitive. Set by closest hit shader.
+#define RAY_TRACING_PAYLOAD_FLAG_MINIMAL_PAYLOAD    2 // Indicates that closest hit shader should only fill FMinimalPayload (HitT), skipping full material evaluation. Set by RayGen shader, before TraceRay().
 
 struct FMaterialClosestHitPayload : FMinimalPayload
 {
@@ -154,9 +156,14 @@ struct FMaterialClosestHitPayload : FMinimalPayload
 	// So keep it here and force to re-construct it in Unpack call using ray information.
 	// It is not packed in FRayHitInfoPacked
 	float3 WorldPos;					// 136       0       (derived)
-
-	uint HitKind;						// 148       0       1bit (packed)
+	uint Flags;							// 148       4       32bits
 										// 152 total
+
+	void SetMinimalPayloadMode() { Flags |= RAY_TRACING_PAYLOAD_FLAG_MINIMAL_PAYLOAD; }
+	bool IsMinimalPayloadMode() { return (Flags & RAY_TRACING_PAYLOAD_FLAG_MINIMAL_PAYLOAD) != 0; }
+
+	void SetFrontFace() { Flags |= RAY_TRACING_PAYLOAD_FLAG_FRONT_FACE; }
+	bool IsFrontFace() { return (Flags & RAY_TRACING_PAYLOAD_FLAG_FRONT_FACE) != 0; }
 };
 
 #define USE_PACKED_PAYLOAD 1
@@ -186,10 +193,15 @@ struct FPackedMaterialClosestHitPayload : FMinimalPayload
 	uint RadianceAndNormal[3];                           // 12 bytes
 	uint BaseColorAndOpacity[2];                         // 8  bytes
 	uint MetallicAndSpecularAndRoughness;                // 4  bytes
-	uint IorAndShadingModelIDAndBlendingModeAndHitKind;  // 4  bytes
+	uint IorAndShadingModelIDAndBlendingModeAndFlags;    // 4  bytes
 	uint IndirectIrradiance[2];                          // 8  bytes
 	uint CustomData;                                     // 4  bytes
-};                                                       // 52 bytes total
+	                                                     // 52 bytes total
+
+	void SetMinimalPayloadMode() { IorAndShadingModelIDAndBlendingModeAndFlags |= RAY_TRACING_PAYLOAD_FLAG_MINIMAL_PAYLOAD << 24; }
+	bool IsMinimalPayloadMode() { return ((IorAndShadingModelIDAndBlendingModeAndFlags >> 24) & RAY_TRACING_PAYLOAD_FLAG_MINIMAL_PAYLOAD) != 0; }
+
+};
 
 FPackedMaterialClosestHitPayload PackRayTracingPayload(FMaterialClosestHitPayload Input, in FRayCone RayCone)
 {
@@ -209,10 +221,10 @@ FPackedMaterialClosestHitPayload PackRayTracingPayload(FMaterialClosestHitPayloa
 	Output.MetallicAndSpecularAndRoughness  = (uint(round(Input.Metallic * 255.0f)) & 0xFF);
 	Output.MetallicAndSpecularAndRoughness |= (uint(round(Input.Specular * 255.0f)) & 0xFF) << 8;
 	Output.MetallicAndSpecularAndRoughness |= f32tof16(Input.Roughness) << 16;
-	Output.IorAndShadingModelIDAndBlendingModeAndHitKind  = f32tof16(Input.Ior);
-	Output.IorAndShadingModelIDAndBlendingModeAndHitKind |= (Input.ShadingModelID & 0xF) << 16;
-	Output.IorAndShadingModelIDAndBlendingModeAndHitKind |= (Input.BlendingMode & 0x7) << 20;
-	Output.IorAndShadingModelIDAndBlendingModeAndHitKind |= (Input.HitKind & 0x1) << 31;
+	Output.IorAndShadingModelIDAndBlendingModeAndFlags  = f32tof16(Input.Ior);                // 16 bits
+	Output.IorAndShadingModelIDAndBlendingModeAndFlags |= (Input.ShadingModelID & 0xF) << 16; // 4 bits
+	Output.IorAndShadingModelIDAndBlendingModeAndFlags |= (Input.BlendingMode & 0xF) << 20;   // 4 bits
+	Output.IorAndShadingModelIDAndBlendingModeAndFlags |= (Input.Flags & 0xF) << 24;          // 8 bits
 	Output.IndirectIrradiance[0] = f32tof16(Input.IndirectIrradiance.x);
 	Output.IndirectIrradiance[0] |= f32tof16(Input.IndirectIrradiance.y) << 16;
 	Output.IndirectIrradiance[1] = f32tof16(Input.IndirectIrradiance.z);
@@ -248,10 +260,10 @@ FMaterialClosestHitPayload UnpackRayTracingPayload(FPackedMaterialClosestHitPayl
 	Output.Metallic	   = float(Input.MetallicAndSpecularAndRoughness & 0xFF) / 255.0f;
     Output.Specular	   = float((Input.MetallicAndSpecularAndRoughness >> 8) & 0xFF) / 255.0f;
     Output.Roughness   = f16tof32(Input.MetallicAndSpecularAndRoughness >> 16);
-	Output.Ior			  = f16tof32(Input.IorAndShadingModelIDAndBlendingModeAndHitKind);
-	Output.ShadingModelID = (Input.IorAndShadingModelIDAndBlendingModeAndHitKind >> 16) & 0xF;
-	Output.BlendingMode	  = (Input.IorAndShadingModelIDAndBlendingModeAndHitKind >> 20) & 0x7;
-	Output.HitKind		  = (Input.IorAndShadingModelIDAndBlendingModeAndHitKind >> 31);
+	Output.Ior            = f16tof32(Input.IorAndShadingModelIDAndBlendingModeAndFlags);      // 16 bits
+	Output.ShadingModelID = (Input.IorAndShadingModelIDAndBlendingModeAndFlags >> 16) & 0xF;  // 4 bits
+	Output.BlendingMode	  = (Input.IorAndShadingModelIDAndBlendingModeAndFlags >> 20) & 0xF;  // 4 bits
+	Output.Flags          = (Input.IorAndShadingModelIDAndBlendingModeAndFlags >> 24) & 0xFF; // 8 bits
 	Output.IndirectIrradiance.x	= f16tof32(Input.IndirectIrradiance[0]);
 	Output.IndirectIrradiance.y	= f16tof32(Input.IndirectIrradiance[0] >> 16);
 	Output.IndirectIrradiance.z	= f16tof32(Input.IndirectIrradiance[1]);
diff --git a/Engine/Shaders/Private/RayTracing/RayTracingDebug.usf b/Engine/Shaders/Private/RayTracing/RayTracingDebug.usf
--- a/Engine/Shaders/Private/RayTracing/RayTracingDebug.usf
+++ b/Engine/Shaders/Private/RayTracing/RayTracingDebug.usf
@@ -96,7 +96,7 @@ void RayTracingDebugMainRGS()
 			Result = float4(Payload.WorldPos, 1.0f);
 			break;
 		case RAY_TRACING_DEBUG_VIZ_HITKIND:
-			Result = float4(Payload.HitKind, Payload.HitKind, Payload.HitKind, 1.0f);
+			Result = Payload.IsFrontFace() ? float4(0.0, 1.0, 0.0, 1.0f) : float4(1.0, 0.0, 0.0, 1.0f);
 			break;
 		}
 	}
diff --git a/Engine/Shaders/Private/RayTracing/RayTracingMaterialHitShaders.usf b/Engine/Shaders/Private/RayTracing/RayTracingMaterialHitShaders.usf
--- a/Engine/Shaders/Private/RayTracing/RayTracingMaterialHitShaders.usf
+++ b/Engine/Shaders/Private/RayTracing/RayTracingMaterialHitShaders.usf
@@ -241,6 +241,15 @@ void CalcInterpolants(in FRayCone RayCone, in FDefaultAttributes Attributes, out
 [shader("closesthit")]
 void MaterialCHS(inout FPackedMaterialClosestHitPayload PackedPayload, in FDefaultAttributes Attributes)
 {
+	if (PackedPayload.IsMinimalPayloadMode())
+	{
+		// Minimal payload mode only fills FMinimalPayload::HitT, skipping actual material evaluation.
+		// This mode is used when tracing shadow rays against masked geometry.
+		// Dynamic branch is used to avoid compiling an extra shader permutation.
+		PackedPayload.HitT = RayTCurrent();
+		return;
+	}
+
 	ResolvedView = ResolveView();
 
 	FVertexFactoryInterpolantsVSToPS Interpolants;
@@ -326,9 +335,12 @@ void MaterialCHS(inout FPackedMaterialClosestHitPayload PackedPayload, in FDefau
 
 	Payload.HitT = RayTCurrent();
 #if !MATERIAL_TWOSIDED
-    Payload.HitKind = HitKind();
+	if (HitKind() == HIT_KIND_TRIANGLE_FRONT_FACE)
+	{
+		Payload.SetFrontFace();
+	}
 #else
-    Payload.HitKind = HIT_KIND_TRIANGLE_FRONT_FACE;
+	Payload.SetFrontFace();
 #endif
 
 	float3 DiffuseIndirectLighting = 0;
diff --git a/Engine/Shaders/Private/RayTracing/RayTracingOcclusionRGS.usf b/Engine/Shaders/Private/RayTracing/RayTracingOcclusionRGS.usf
--- a/Engine/Shaders/Private/RayTracing/RayTracingOcclusionRGS.usf
+++ b/Engine/Shaders/Private/RayTracing/RayTracingOcclusionRGS.usf
@@ -221,6 +221,7 @@ void OcclusionRGS()
 	#endif
 
 			FPackedMaterialClosestHitPayload Payload = (FPackedMaterialClosestHitPayload)0;
+			Payload.SetMinimalPayloadMode(); // Only HitT will be filled by closest hit shader, skipping full material evaluation.
 		
 			TraceRay(
 				TLAS,      // AccelerationStructure
diff --git a/Engine/Shaders/Private/RayTracing/RayTracingRectLightRGS.usf b/Engine/Shaders/Private/RayTracing/RayTracingRectLightRGS.usf
--- a/Engine/Shaders/Private/RayTracing/RayTracingRectLightRGS.usf
+++ b/Engine/Shaders/Private/RayTracing/RayTracingRectLightRGS.usf
@@ -267,7 +267,8 @@ void RectLightRGS()
 			| RAY_FLAG_CULL_BACK_FACING_TRIANGLES
 			| RAY_FLAG_FORCE_OPAQUE; // #dxr_todo: UE-72563 generate a special shader permutation for occlusion ray tracing that includes any-hit shaders for masked geometry
 
-		FDefaultPayload Payload = (FDefaultPayload)0;
+		FPackedMaterialClosestHitPayload Payload = (FPackedMaterialClosestHitPayload)0;
+		Payload.SetMinimalPayloadMode(); // Only HitT will be filled by closest hit shader, skipping full material evaluation.
 
 		TraceRay(
 			TLAS,      // AccelerationStructure
diff --git a/Engine/Shaders/Private/RayTracing/RayTracingSkyLightRGS.usf b/Engine/Shaders/Private/RayTracing/RayTracingSkyLightRGS.usf
--- a/Engine/Shaders/Private/RayTracing/RayTracingSkyLightRGS.usf
+++ b/Engine/Shaders/Private/RayTracing/RayTracingSkyLightRGS.usf
@@ -310,6 +310,7 @@ void SkyLightRGS()
 		RayFlags |= RAY_FLAG_CULL_BACK_FACING_TRIANGLES;
 #endif
 		FPackedMaterialClosestHitPayload Payload = (FPackedMaterialClosestHitPayload)0;
+		Payload.SetMinimalPayloadMode(); // Only HitT will be filled by closest hit shader, skipping full material evaluation.
 
 		TraceRay(
 			TLAS,      // AccelerationStructure
diff --git a/Engine/Shaders/Private/RayTracing/RayTracingTranslucency.usf b/Engine/Shaders/Private/RayTracing/RayTracingTranslucency.usf
--- a/Engine/Shaders/Private/RayTracing/RayTracingTranslucency.usf
+++ b/Engine/Shaders/Private/RayTracing/RayTracingTranslucency.usf
@@ -223,7 +223,7 @@ void RayTracingTranslucencyRGS()
 			float Ior2 = LastIor;
 			bHasScattered |= Ior1 != Ior2;
 			float ReflectionRefractionEventThroughput;
-			bool bIsEntering = Payload.HitKind == HIT_KIND_TRIANGLE_FRONT_FACE;
+			bool bIsEntering = Payload.IsFrontFace();
 
 			if (!RefractRay(Ray.Direction, Payload.WorldNormal, Ior1, Ior2, bIsEntering, RefractedDirection, ReflectionRefractionEventThroughput))
 			{
diff --git a/Engine/Source/Runtime/Renderer/Private/RayTracing/RayTracingMaterialHitShaders.cpp b/Engine/Source/Runtime/Renderer/Private/RayTracing/RayTracingMaterialHitShaders.cpp
--- a/Engine/Source/Runtime/Renderer/Private/RayTracing/RayTracingMaterialHitShaders.cpp
+++ b/Engine/Source/Runtime/Renderer/Private/RayTracing/RayTracingMaterialHitShaders.cpp
@@ -516,7 +516,7 @@ FRayTracingPipelineState* FDeferredShadingSceneRenderer::BindRayTracingMaterialP
 			else
 			{
 				// Masked materials require full material evaluation with any-hit shader.
-				// #dxr_todo UE-72029: we need to generate a shadow-specific closest hit shader for this!
+				// Full CHS is bound, however material evaluation is skipped for shadow rays using a dynamic branch on a ray payload flag.
 				MeshCommand.ShaderBindings.SetRayTracingShaderBindingsForHitGroup(RHICmdList,
 					View.RayTracingScene.RayTracingSceneRHI,
 					VisibleMeshCommand.InstanceIndex,
