Fixed StationaryLightOverlap viewmode in forward shading, by ignoring the shadermodel test, since it is not available in forward.

diff --git a/Engine/Shaders/Private/StationaryLightOverlapShaders.usf b/Engine/Shaders/Private/StationaryLightOverlapShaders.usf
--- a/Engine/Shaders/Private/StationaryLightOverlapShaders.usf
+++ b/Engine/Shaders/Private/StationaryLightOverlapShaders.usf
@@ -18,10 +18,11 @@ void OverlapDirectionalPixelMain(
 	)
 {
 	OutColor = 0;
-	FGBufferData GBuffer = GetGBufferData(InUV);
-
+#if !FORWARD_SHADING
+	FGBufferData GBuffer = GetGBufferData(InUV); 
 	// Only light pixels marked as using deferred shading
 	if( GBuffer.ShadingModelID > 0 )
+#endif
 	{
 		// Output minor complexity to visualize overlap
 		OutColor.rgb = .04f;
@@ -44,11 +45,12 @@ void OverlapRadialPixelMain(
 	float2 ScreenUV = SvPositionToBufferUV(SvPosition);
 	SvPosition.z = LookupDeviceZ(ScreenUV);
 	
+#if !FORWARD_SHADING
 	FGBufferData GBuffer = GetGBufferData(ScreenUV); 
-	
 	// Only light pixels marked as using deferred shading
 	BRANCH if( GBuffer.ShadingModelID > 0 )
-	{
+#endif
+{
 		float3 WorldPosition = SvPositionToWorld(SvPosition);
 
 		float3 ToLight = DeferredLightUniforms.Direction;
diff --git a/Engine/Source/Runtime/Renderer/Private/DeferredShadingRenderer.cpp b/Engine/Source/Runtime/Renderer/Private/DeferredShadingRenderer.cpp
--- a/Engine/Source/Runtime/Renderer/Private/DeferredShadingRenderer.cpp
+++ b/Engine/Source/Runtime/Renderer/Private/DeferredShadingRenderer.cpp
@@ -1954,8 +1954,7 @@ void FDeferredShadingSceneRenderer::Render(FRHICommandListImmediate& RHICmdList)
 	}
 
 	if (ViewFamily.EngineShowFlags.StationaryLightOverlap &&
-		FeatureLevel >= ERHIFeatureLevel::SM4 && 
-		bUseGBuffer)
+		FeatureLevel >= ERHIFeatureLevel::SM4)
 	{
 		RenderStationaryLightOverlap(RHICmdList);
 		ServiceLocalQueue();
