Fixed: Bloom artifacts on newer Apple GPUs

diff --git a/Engine/Shaders/Private/PostProcessMobile.usf b/Engine/Shaders/Private/PostProcessMobile.usf
--- a/Engine/Shaders/Private/PostProcessMobile.usf
+++ b/Engine/Shaders/Private/PostProcessMobile.usf
@@ -304,7 +304,7 @@ void BloomSmallPS_ES2(
 	half3 N13 = PostprocessInput0.Sample(PostprocessInput0Sampler, InUVs[6].zw).rgb;
 	half3 N14 = PostprocessInput0.Sample(PostprocessInput0Sampler, InUVs[7].xy).rgb;
 
-	half W = 1.0/15.0;
+	float W = 1.0/15.0;
 	OutColor.rgb = 
 		(N0 * W) +
 		(N1 * W) +
@@ -391,7 +391,7 @@ void BloomDownPS_ES2(
 	half4 N12 = PostprocessInput0.Sample(PostprocessInput0Sampler, InUVs[6].xy).rgba;
 	half4 N13 = PostprocessInput0.Sample(PostprocessInput0Sampler, InUVs[6].zw).rgba;
 	half4 N14 = PostprocessInput0.Sample(PostprocessInput0Sampler, InUVs[7].xy).rgba;
-	half W = 1.0/15.0;
+	float W = 1.0/15.0;
 	OutColor.rgb = 
 		(HDRDecode(N0) * W) +
 		(HDRDecode(N1) * W) +
@@ -992,10 +992,9 @@ void SunMergePS_ES2(
 
 		OutColor.rgb = HDRDecode(PostprocessInput2.Sample(PostprocessInput2Sampler, InUVVignette.xy));
 
-		OutColor.rgb += Bloom2 * BloomColor;
-
 		// Have 5 layers on mobile.
-		OutColor.rgb *= 1.0/5.0;
+		float Scale3 = 1.0/5.0;	
+		OutColor.rgb+= (Bloom2 * Scale3 * BloomColor);
 
 	#else
 		OutColor.rgb = half3(0.0, 0.0, 0.0);
