Cherry-pick 7296040

Decouples sample generation from ray traversal in the SkyLight pass.

Created GenerateSkyLightVisibilityRays compute shader. We leverage the QuasiRandom sequence and infinite light properties of SkyLight to cache importance-sampled ray directions for the RayTracingSkyLight ray generation shader. This amounts to noticeable performance improvement without observable increases in variance.

r.RayTracing.SkyLight.DecoupleSampleGeneration (default = 1): decouples sample generation from ray traversal.

diff --git a/Engine/Shaders/Private/RayTracing/GenerateSkyLightVisibilityRaysCS.usf b/Engine/Shaders/Private/RayTracing/GenerateSkyLightVisibilityRaysCS.usf
new file mode 100644
--- /dev/null
+++ b/Engine/Shaders/Private/RayTracing/GenerateSkyLightVisibilityRaysCS.usf
@@ -0,0 +1,68 @@
+// Copyright 1998-2019 Epic Games, Inc. All Rights Reserved.
+
+/*=============================================================================================
+	GenerateSkyLightVisibilityRaysCS.usf: Build variance map
+===============================================================================================*/
+
+#include "../Common.ush"
+#include "../MonteCarlo.ush"
+#include "../QuasiRandom.ush"
+#include "SkyLightMipTreeCommon.ush"
+#include "RayTracingCommon.ush"
+#include "RayTracingSkyLightCommon.ush"
+
+struct FSkyLightVisibilityRays
+{
+	float4 DirectionAndPdf;
+};
+
+uint3 Dimensions;
+RWStructuredBuffer<FSkyLightVisibilityRays> OutSkyLightVisibilityRays;
+
+[numthreads(TILE_SIZE, TILE_SIZE, 1)]
+void MainCS(
+	uint3 DispatchThreadId : SV_DispatchThreadID
+)
+{
+	uint2 PixelCoord = uint2(DispatchThreadId.x, DispatchThreadId.y);
+
+	FQuasiRandom QuasiRandom;
+	QuasiRandom_Initialize(QuasiRandom, PixelCoord, 0, SkyLight.SamplesPerPixel);
+	float3 BlueNoiseShift = EvalBlueNoise(QuasiRandom.BlueNoiseCoord, 0);
+
+	// Cache Cubeface CDF
+	uint MipCount = log2(SkyLight.MipDimensions.x);
+	float FaceCdf[7];
+	BuildFaceCdf(MipCount, FaceCdf);
+
+	for (uint SampleIndex = 0; SampleIndex < SkyLight.SamplesPerPixel; ++SampleIndex)
+	{
+		QuasiRandom_Update(QuasiRandom, SampleIndex, SkyLight.SamplesPerPixel);
+
+		// Determine sky light or lambert ray
+		float3 RandSample = QuasiRandom_Sample3(QuasiRandom);
+		RandSample = HaltonSequence3(QuasiRandom.HaltonSequence, QuasiRandom.HaltonIteration, 0);
+		RandSample += BlueNoiseShift;
+		RandSample = fmod(RandSample, 1.0);
+
+		float3 RayDirection;
+		float SkyLightSamplingStrategyPdf = 0.5;
+		float MisWeight = 1.0;
+		float RayPdf = 0.0;
+		BRANCH
+			if (RandSample.x < SkyLightSamplingStrategyPdf)
+			{
+				RandSample.x /= SkyLightSamplingStrategyPdf;
+				GenerateSkyLightOcclusionRay(FaceCdf, RandSample, RayDirection, RayPdf);
+				MisWeight = RayPdf / (SkyLightSamplingStrategyPdf * RayPdf + (1.0 - SkyLightSamplingStrategyPdf) * PdfSphericalOcclusionRay(RayDirection));
+			}
+			else
+			{
+				GenerateSphericalOcclusionRay(RandSample.yz, RayDirection, RayPdf);
+				MisWeight = RayPdf / ((1.0 - SkyLightSamplingStrategyPdf) * RayPdf + SkyLightSamplingStrategyPdf * PdfSkyLightOcclusionRay(FaceCdf, RayDirection));
+			}
+
+		uint LinearIndex = (PixelCoord.y * Dimensions.x + PixelCoord.x) * SkyLight.SamplesPerPixel + SampleIndex;
+		OutSkyLightVisibilityRays[LinearIndex].DirectionAndPdf = float4(RayDirection.xyz, MisWeight / RayPdf);
+	}
+}
diff --git a/Engine/Shaders/Private/RayTracing/RayTracingSkyLightCommon.ush b/Engine/Shaders/Private/RayTracing/RayTracingSkyLightCommon.ush
new file mode 100644
--- /dev/null
+++ b/Engine/Shaders/Private/RayTracing/RayTracingSkyLightCommon.ush
@@ -0,0 +1,203 @@
+// Copyright 1998-2019 Epic Games, Inc. All Rights Reserved.
+
+/*=============================================================================================
+	SkyLightMipTreeCommon.ush: Common utilities for MipTree-based SkyLight sampling
+===============================================================================================*/
+
+float PdfSkyLightOcclusionRay(
+	float FaceCdf[7],
+	float3 WorldDirection
+)
+{
+	// Determine primary direction
+	uint MipCount = log2(SkyLight.MipDimensions.x);
+	uint3 TextureCoord = GetTextureCubeCoordinate(WorldDirection, SkyLight.TextureDimensions);
+	float2 TextureUV = TextureCoord.xy / (float2)SkyLight.TextureDimensions;
+	uint2 MipCoord = TextureUV * SkyLight.MipDimensions.xy;
+	float FacePdf = PdfFace(FaceCdf, TextureCoord.z);
+
+	uint StopLevel = min(SkyLight.SamplingStopLevel, MipCount);
+#if USE_MIP_PDF_TEXTURE
+	float MipPdf = PdfMipTree(uint3(MipCoord, TextureCoord.z), StopLevel);
+
+#if USE_SOLID_ANGLE_TEXTURE
+	uint LinearIndex = BufferOffsetAtPixel(MipCoord, StopLevel, SkyLight.MipDimensions.xy);
+	float AreaPdf = SkyLight.SolidAnglePdf[LinearIndex];
+#else
+	float AreaPdf = 1.0 / TexelCoordSolidAngle(MipCoord, DimensionsAtLevel(SkyLight.MipDimensions, StopLevel));
+#endif
+	return FacePdf * MipPdf * AreaPdf;
+#else
+	uint2 MipPixel = 0;
+	float MipPdf = 1.0;
+	for (uint Index = 0; Index < MipCount - StopLevel; ++Index)
+	{
+		// Build Cdf
+		uint MipIndex = MipCount - Index - 1;
+#if USE_MORTON_CODE
+		uint BufferOffset = BufferOffsetAtPixel(MipPixel, MipIndex, SkyLight.MipDimensions);
+		uint4 BufferOffsets = uint4(BufferOffset, BufferOffset + 1, BufferOffset + 2, BufferOffset + 3);
+#else
+		uint4 BufferOffsets;
+		BufferOffsets.x = BufferOffsetAtPixel(MipPixel, MipIndex, SkyLight.MipDimensions.xy);
+		BufferOffsets.y = BufferOffsetAtPixel(MipPixel + uint2(1, 0), MipIndex, SkyLight.MipDimensions.xy);
+		BufferOffsets.z = BufferOffsetAtPixel(MipPixel + uint2(0, 1), MipIndex, SkyLight.MipDimensions.xy);
+		BufferOffsets.w = BufferOffsetAtPixel(MipPixel + uint2(1, 1), MipIndex, SkyLight.MipDimensions.xy);
+#endif
+		float4 Cdf = SampleMipTree(TextureCoord.z, BufferOffsets);
+		Cdf.y += Cdf.x;
+		Cdf.z += Cdf.y;
+		Cdf.w += Cdf.z;
+		Cdf /= Cdf.w;
+
+		// Determine texel choice
+		uint2 MipOffset = (MipCoord.xy >> MipIndex) - MipPixel;
+		uint LinearIndex = MipOffset.y << 1 | MipOffset.x;
+		MipPdf *= LinearIndex > 0 ? Cdf[LinearIndex] - Cdf[LinearIndex - 1] : Cdf[LinearIndex];
+		MipPixel += MipOffset;
+
+		MipPixel = MipPixel << 1;
+	}
+	MipPdf = max(MipPdf, 0.0);
+	// DEBUG: assert that PdfMipTree is equivalent for evaluating the decision tree..
+	//MipPdf = PdfMipTree(uint3(MipCoord, TextureCoord.z), StopLevel);
+
+#if USE_SOLID_ANGLE_TEXTURE
+	uint LinearIndex = BufferOffsetAtPixel(MipCoord, StopLevel, SkyLight.MipDimensions.xy);
+	float AreaPdf = SkyLight.SolidAnglePdf[LinearIndex];
+#else
+	float AreaPdf = 1.0 / TexelCoordSolidAngle(MipCoord, DimensionsAtLevel(SkyLight.MipDimensions, StopLevel));
+#endif
+	return FacePdf * MipPdf * AreaPdf;
+#endif
+}
+
+bool GenerateSkyLightOcclusionRay(
+	float FaceCdf[7],
+	float3 RandSample,
+	out float3 RayDirection,
+	out float RayPdf
+)
+{
+	// Determine the appropriate face to use
+	uint MipCount = log2(SkyLight.MipDimensions.x);
+	float FacePdf = 0.0;
+	uint FaceIndex = SampleFace(FaceCdf, RandSample.x, FacePdf);
+
+	uint2 MipPixel = 0;
+	float MipPdf = 1.0;
+	uint StopLevel = min(SkyLight.SamplingStopLevel, MipCount);
+	for (uint Index = 0; Index < MipCount - StopLevel; ++Index)
+	{
+		uint MipIndex = MipCount - Index - 1;
+#if USE_MORTON_CODE
+		uint BufferOffset = BufferOffsetAtPixel(MipPixel, MipIndex, SkyLight.MipDimensions.xy);
+		uint4 BufferOffsets = uint4(BufferOffset, BufferOffset + 1, BufferOffset + 2, BufferOffset + 3);
+#else
+		uint4 BufferOffsets;
+		BufferOffsets.x = BufferOffsetAtPixel(MipPixel, MipIndex, SkyLight.MipDimensions);
+		BufferOffsets.y = BufferOffsetAtPixel(MipPixel + uint2(1, 0), MipIndex, SkyLight.MipDimensions.xy);
+		BufferOffsets.z = BufferOffsetAtPixel(MipPixel + uint2(0, 1), MipIndex, SkyLight.MipDimensions.xy);
+		BufferOffsets.w = BufferOffsetAtPixel(MipPixel + uint2(1, 1), MipIndex, SkyLight.MipDimensions.xy);
+#endif
+		float4 Cdf = SampleMipTree(FaceIndex, BufferOffsets);
+		Cdf.y += Cdf.x;
+		Cdf.z += Cdf.y;
+		Cdf.w += Cdf.z;
+		Cdf /= Cdf.w;
+
+		if (RandSample.y < Cdf.x)
+		{
+			MipPdf *= Cdf.x;
+			RandSample.y /= Cdf.x;
+		}
+		else if (RandSample.y < Cdf.y)
+		{
+			MipPdf *= Cdf.y - Cdf.x;
+			MipPixel.x += 1;
+			RandSample.y = (RandSample.y - Cdf.x) / (Cdf.y - Cdf.x);
+		}
+		else if (RandSample.y < Cdf.z)
+		{
+			MipPdf *= Cdf.z - Cdf.y;
+			MipPixel.y += 1;
+			RandSample.y = (RandSample.y - Cdf.y) / (Cdf.z - Cdf.y);
+		}
+		else // if (RandSample.y < Cdf.w)
+		{
+			MipPdf *= Cdf.w - Cdf.z;
+			MipPixel += 1;
+			RandSample.y = (RandSample.y - Cdf.z) / (Cdf.w - Cdf.z);
+		}
+
+		MipPixel = uint2(MipPixel.x << 1, MipPixel.y << 1);
+	}
+	MipPixel = uint2(MipPixel.x >> 1, MipPixel.y >> 1);
+
+	float2 PixelOffset = RandSample.yz;
+	float2 MipPixelUV = (MipPixel + PixelOffset) / float2(DimensionsAtLevel(SkyLight.MipDimensions, StopLevel));
+
+	// Construct Cubemap coordinate
+	float3 TexelCoord = float3(MipPixelUV * SkyLight.TextureDimensions, FaceIndex);
+
+	// Determine area pdf
+#if USE_SOLID_ANGLE_TEXTURE
+	uint LinearIndex = BufferOffsetAtPixel(MipPixel, StopLevel, SkyLight.MipDimensions.xy);
+	float AreaPdf = SkyLight.SolidAnglePdf[LinearIndex];
+#else
+	float AreaPdf = 1.0 / TexelCoordSolidAngle(MipPixel, DimensionsAtLevel(SkyLight.MipDimensions, StopLevel));
+#endif
+
+	RayDirection = normalize(GetTextureCubeVector(TexelCoord, SkyLight.TextureDimensions));
+	RayPdf = FacePdf * MipPdf * AreaPdf;
+
+	return true;
+}
+
+float PdfSphericalOcclusionRay(
+	float3 WorldDirection
+)
+{
+	float Pdf = 1.0 / (4.0 * PI);
+	return Pdf;
+}
+
+bool GenerateSphericalOcclusionRay(
+	float2 RandSample,
+	out float3 RayDirection,
+	out float RayPdf
+)
+{
+	// Perform uniform spherical sampling and convert to world-space
+	float4 Direction_World = UniformSampleSphere(RandSample);
+
+	RayDirection = Direction_World.xyz;
+	RayPdf = Direction_World.w;
+
+	return true;
+}
+
+float PdfCosineHemisphericalOcclusionRay(
+	float3 WorldNormal,
+	float3 WorldDirection
+)
+{
+	float Pdf = max(dot(WorldDirection, WorldNormal), 0.0) / PI;
+	return Pdf;
+}
+
+bool GenerateCosineHemisphericalOcclusionRay(
+	float3 WorldNormal,
+	float2 RandSample,
+	out float3 RayDirection,
+	out float RayPdf
+)
+{
+	// Perform cosine-hemispherical sampling
+	float4 Direction_World = CosineSampleHemisphere(RandSample, WorldNormal);
+
+	RayDirection = Direction_World.xyz;
+	RayPdf = Direction_World.w;
+
+	return true;
+}
diff --git a/Engine/Shaders/Private/RayTracing/RayTracingSkyLightRGS.usf b/Engine/Shaders/Private/RayTracing/RayTracingSkyLightRGS.usf
--- a/Engine/Shaders/Private/RayTracing/RayTracingSkyLightRGS.usf
+++ b/Engine/Shaders/Private/RayTracing/RayTracingSkyLightRGS.usf
@@ -7,219 +7,29 @@
 #include "../QuasiRandom.ush"
 #include "SkyLightMipTreeCommon.ush"
 #include "RayTracingCommon.ush"
+#include "RayTracingSkyLightCommon.ush"
 #include "../SceneTextureParameters.ush"
 
+struct FSkyLightVisibilityRays
+{
+	float4 DirectionAndPdf;
+};
+
 RaytracingAccelerationStructure TLAS;
+uint3 SkyLightVisibilityRaysDimensions;
+StructuredBuffer<FSkyLightVisibilityRays> SkyLightVisibilityRays;
+
 RWTexture2D<float4> RWOcclusionMaskUAV;
 RWTexture2D<float2> RWRayDistanceUAV;
 
 #define USE_MIP_PDF_TEXTURE 1
 #define USE_SOLID_ANGLE_TEXTURE 1
 
-
 float3 EvalSkyLight(float3 Direction)
 {
 	return SkyLight.Color * TextureCubeSampleLevel(SkyLight.Texture, SkyLight.TextureSampler, Direction, 0).rgb;
 }
 
-float PdfSkyLightOcclusionRay(
-	float FaceCdf[7],
-	float3 WorldNormal,
-	float3 WorldDirection
-)
-{
-	// Determine primary direction
-	uint MipCount = log2(SkyLight.MipDimensions.x);
-	uint3 TextureCoord = GetTextureCubeCoordinate(WorldDirection, SkyLight.TextureDimensions);
-	float2 TextureUV = TextureCoord.xy / (float2)SkyLight.TextureDimensions;
-	uint2 MipCoord = TextureUV * SkyLight.MipDimensions.xy;
-	float FacePdf = PdfFace(FaceCdf, TextureCoord.z);
-
-	uint StopLevel = min(SkyLight.SamplingStopLevel, MipCount);
-#if USE_MIP_PDF_TEXTURE
-	float MipPdf = PdfMipTree(uint3(MipCoord, TextureCoord.z), StopLevel);
-
-	#if USE_SOLID_ANGLE_TEXTURE
-		uint LinearIndex = BufferOffsetAtPixel(MipCoord, StopLevel, SkyLight.MipDimensions.xy);
-		float AreaPdf = SkyLight.SolidAnglePdf[LinearIndex];
-	#else
-		float AreaPdf = 1.0 / TexelCoordSolidAngle(MipCoord, DimensionsAtLevel(SkyLight.MipDimensions, StopLevel));
-	#endif
-	return FacePdf * MipPdf * AreaPdf;
-#else
-	uint2 MipPixel = 0;
-	float MipPdf = 1.0;
-	for (uint Index = 0; Index < MipCount - StopLevel; ++Index)
-	{
-		// Build Cdf
-		uint MipIndex = MipCount - Index - 1;
-#if USE_MORTON_CODE
-		uint BufferOffset = BufferOffsetAtPixel(MipPixel, MipIndex, SkyLight.MipDimensions);
-		uint4 BufferOffsets = uint4(BufferOffset, BufferOffset + 1, BufferOffset + 2, BufferOffset + 3);
-#else
-		uint4 BufferOffsets;
-		BufferOffsets.x = BufferOffsetAtPixel(MipPixel, MipIndex, SkyLight.MipDimensions.xy);
-		BufferOffsets.y = BufferOffsetAtPixel(MipPixel + uint2(1, 0), MipIndex, SkyLight.MipDimensions.xy);
-		BufferOffsets.z = BufferOffsetAtPixel(MipPixel + uint2(0, 1), MipIndex, SkyLight.MipDimensions.xy);
-		BufferOffsets.w = BufferOffsetAtPixel(MipPixel + uint2(1, 1), MipIndex, SkyLight.MipDimensions.xy);
-#endif
-		float4 Cdf = SampleMipTree(TextureCoord.z, BufferOffsets);
-		Cdf.y += Cdf.x;
-		Cdf.z += Cdf.y;
-		Cdf.w += Cdf.z;
-		Cdf /= Cdf.w;
-
-		// Determine texel choice
-		uint2 MipOffset = (MipCoord.xy >> MipIndex) - MipPixel;
-		uint LinearIndex = MipOffset.y << 1 | MipOffset.x;
-		MipPdf *= LinearIndex > 0 ? Cdf[LinearIndex] - Cdf[LinearIndex - 1] : Cdf[LinearIndex];
-		MipPixel += MipOffset;
-
-		MipPixel = MipPixel << 1;
-	}
-	MipPdf = max(MipPdf, 0.0);
-	// DEBUG: assert that PdfMipTree is equivalent for evaluating the decision tree..
-	//MipPdf = PdfMipTree(uint3(MipCoord, TextureCoord.z), StopLevel);
-
-	#if USE_SOLID_ANGLE_TEXTURE
-		uint LinearIndex = BufferOffsetAtPixel(MipCoord, StopLevel, SkyLight.MipDimensions.xy);
-		float AreaPdf = SkyLight.SolidAnglePdf[LinearIndex];
-	#else
-		float AreaPdf = 1.0 / TexelCoordSolidAngle(MipCoord, DimensionsAtLevel(SkyLight.MipDimensions, StopLevel));
-	#endif
-	return FacePdf * MipPdf * AreaPdf;
-#endif
-}
-
-bool GenerateSkyLightOcclusionRay(
-	float FaceCdf[7],
-	float3 WorldPosition,
-	float3 WorldNormal,
-	float3 RandSample,
-	out float3 RayOrigin,
-	out float3 RayDirection,
-	out float RayTMin,
-	out float RayTMax,
-	out float RayPdf
-)
-{
-	// Determine the appropriate face to use
-	uint MipCount = log2(SkyLight.MipDimensions.x);
-	float FacePdf = 0.0;
-	uint FaceIndex = SampleFace(FaceCdf, RandSample.x, FacePdf);
-
-	uint2 MipPixel = 0;
-	float MipPdf = 1.0;
-	uint StopLevel = min(SkyLight.SamplingStopLevel, MipCount);
-	for (uint Index = 0; Index < MipCount - StopLevel; ++Index)
-	{
-		uint MipIndex = MipCount - Index - 1;
-#if USE_MORTON_CODE
-		uint BufferOffset = BufferOffsetAtPixel(MipPixel, MipIndex, SkyLight.MipDimensions.xy);
-		uint4 BufferOffsets = uint4(BufferOffset, BufferOffset + 1, BufferOffset + 2, BufferOffset + 3);
-#else
-		uint4 BufferOffsets;
-		BufferOffsets.x = BufferOffsetAtPixel(MipPixel, MipIndex, SkyLight.MipDimensions);
-		BufferOffsets.y = BufferOffsetAtPixel(MipPixel + uint2(1, 0), MipIndex, SkyLight.MipDimensions.xy);
-		BufferOffsets.z = BufferOffsetAtPixel(MipPixel + uint2(0, 1), MipIndex, SkyLight.MipDimensions.xy);
-		BufferOffsets.w = BufferOffsetAtPixel(MipPixel + uint2(1, 1), MipIndex, SkyLight.MipDimensions.xy);
-#endif
-		float4 Cdf = SampleMipTree(FaceIndex, BufferOffsets);
-		Cdf.y += Cdf.x;
-		Cdf.z += Cdf.y;
-		Cdf.w += Cdf.z;
-		Cdf /= Cdf.w;
-
-		if (RandSample.y < Cdf.x)
-		{
-			MipPdf *= Cdf.x;
-			RandSample.y /= Cdf.x;
-		}
-		else if (RandSample.y < Cdf.y)
-		{
-			MipPdf *= Cdf.y - Cdf.x;
-			MipPixel.x += 1;
-			RandSample.y = (RandSample.y - Cdf.x) / (Cdf.y - Cdf.x);
-		}
-		else if (RandSample.y < Cdf.z)
-		{
-			MipPdf *= Cdf.z - Cdf.y;
-			MipPixel.y += 1;
-			RandSample.y = (RandSample.y - Cdf.y) / (Cdf.z - Cdf.y);
-		}
-		else // if (RandSample.y < Cdf.w)
-		{
-			MipPdf *= Cdf.w - Cdf.z;
-			MipPixel += 1;
-			RandSample.y = (RandSample.y - Cdf.z) / (Cdf.w - Cdf.z);
-		}
-
-		MipPixel = uint2(MipPixel.x << 1, MipPixel.y << 1);
-	}
-	MipPixel = uint2(MipPixel.x >> 1, MipPixel.y >> 1);
-
-	// DEBUG: assert cached PDF is equivalent to re-evaluation
-	//MipPdf = PdfSkyLightOcclusionRay(WorldNormal, RayDirection);
-
-	float2 PixelOffset = RandSample.yz;
-	float2 MipPixelUV = (MipPixel + PixelOffset) / float2(DimensionsAtLevel(SkyLight.MipDimensions, StopLevel));
-
-	// Construct Cubemap coordinate
-	float3 TexelCoord = float3(MipPixelUV * SkyLight.TextureDimensions, FaceIndex);
-
-	// Determine area pdf
-#if USE_SOLID_ANGLE_TEXTURE
-	uint LinearIndex = BufferOffsetAtPixel(MipPixel, StopLevel, SkyLight.MipDimensions.xy);
-	float AreaPdf = SkyLight.SolidAnglePdf[LinearIndex];
-#else
-	float AreaPdf = 1.0 / TexelCoordSolidAngle(MipPixel, DimensionsAtLevel(SkyLight.MipDimensions, StopLevel));
-#endif
-
-	RayOrigin = WorldPosition;
-	RayDirection = normalize(GetTextureCubeVector(TexelCoord, SkyLight.TextureDimensions));
-	RayTMin = 0.0;
-	RayTMax = SkyLight.MaxRayDistance;
-	RayPdf = FacePdf * MipPdf * AreaPdf;
-
-	return true;
-}
-
-float PdfLambertOcclusionRay(
-	float3 WorldNormal,
-	float3 WorldDirection
-)
-{
-	float Pdf = max(dot(WorldDirection, WorldNormal), 0.0) / PI;
-	return Pdf;
-}
-
-bool GenerateLambertOcclusionRay(
-	float3 WorldPosition,
-	float3 WorldNormal,
-	float2 RandSample,
-	out float3 RayOrigin,
-	out float3 RayDirection,
-	out float RayTMin,
-	out float RayTMax,
-	out float RayPdf
-)
-{
-	// Draw random variable
-	float2 BufferSize = View.BufferSizeAndInvSize.xy;
-
-	// Perform cosine-hemispherical sampling and convert to world-space
-	float4 Direction_Tangent = CosineSampleHemisphere(RandSample);
-	float3 Direction_World = TangentToWorld(Direction_Tangent.xyz, WorldNormal);
-
-	RayOrigin = WorldPosition;
-	RayDirection = Direction_World;
-	RayTMin = 0.0;
-	RayTMax = SkyLight.MaxRayDistance;
-	RayPdf = Direction_Tangent.w;
-
-	return true;
-}
-
 [shader("raygeneration")]
 void SkyLightRGS()
 {
@@ -260,7 +70,7 @@ void SkyLightRGS()
 	{
 		SamplesPerPixel = 0.0;
 	}
-
+#if !DECOUPLE_SAMPLE_GENERATION
 	// Cache Cubeface CDF
 	uint MipCount = log2(SkyLight.MipDimensions.x);
 	float FaceCdf[7];
@@ -270,14 +80,27 @@ void SkyLightRGS()
 	const uint Dimension = 0;
 	QuasiRandom_Initialize(QuasiRandom, PixelCoord, 0, SamplesPerPixel);
 	float3 BlueNoiseShift = EvalBlueNoise(QuasiRandom.BlueNoiseCoord, Dimension);
+#endif
 
 	float3 ExitantRadiance = 0.0;
 	float3 DiffuseExitantRadiance = 0.0;
 	for (uint SampleIndex = 0; SampleIndex < SamplesPerPixel; ++SampleIndex)
 	{
-		QuasiRandom_Update(QuasiRandom, SampleIndex, SamplesPerPixel);
 		RayDesc Ray;
-		float RayPdf = 0.0;
+		float RayWeight;
+
+#if DECOUPLE_SAMPLE_GENERATION
+		float2 RayCoord = PixelCoord % SkyLightVisibilityRaysDimensions.xy;
+		uint RayCoordLinearIndex = (RayCoord.y * SkyLightVisibilityRaysDimensions.x + RayCoord.x) * SkyLight.SamplesPerPixel + SampleIndex;
+
+		float4 DirectionAndPdf = SkyLightVisibilityRays[RayCoordLinearIndex].DirectionAndPdf;
+		Ray.Origin = WorldPosition;
+		Ray.Direction = DirectionAndPdf.xyz;
+		Ray.TMin = 0.0;
+		Ray.TMax = SkyLight.MaxRayDistance;
+		RayWeight = DirectionAndPdf.w;
+#else
+		QuasiRandom_Update(QuasiRandom, SampleIndex, SamplesPerPixel);
 
 		// Determine sky light or lambert ray
 		float3 RandSample = QuasiRandom_Sample3(QuasiRandom);
@@ -287,18 +110,24 @@ void SkyLightRGS()
 
 		float SkyLightSamplingStrategyPdf = 0.5;
 		float MisWeight = 1.0;
+		float RayPdf = 0.0;
 		BRANCH
 		if (RandSample.x < SkyLightSamplingStrategyPdf)
 		{
 			RandSample.x /= SkyLightSamplingStrategyPdf;
-			GenerateSkyLightOcclusionRay(FaceCdf, WorldPosition, WorldNormal, RandSample, Ray.Origin, Ray.Direction, Ray.TMin, Ray.TMax, RayPdf);
-			MisWeight = RayPdf / (SkyLightSamplingStrategyPdf * RayPdf + (1.0 - SkyLightSamplingStrategyPdf) * PdfLambertOcclusionRay(WorldNormal, Ray.Direction));
+			GenerateSkyLightOcclusionRay(FaceCdf, RandSample, Ray.Direction, RayPdf);
+			MisWeight = RayPdf / (SkyLightSamplingStrategyPdf * RayPdf + (1.0 - SkyLightSamplingStrategyPdf) * PdfCosineHemisphericalOcclusionRay(WorldNormal, Ray.Direction));
 		}
 		else
 		{
-			GenerateLambertOcclusionRay(WorldPosition, WorldNormal, RandSample.yz, Ray.Origin, Ray.Direction, Ray.TMin, Ray.TMax, RayPdf);
-			MisWeight = RayPdf / ((1.0 - SkyLightSamplingStrategyPdf) * RayPdf + SkyLightSamplingStrategyPdf * PdfSkyLightOcclusionRay(FaceCdf, WorldNormal, Ray.Direction));
+			GenerateCosineHemisphericalOcclusionRay(WorldNormal, RandSample.yz, Ray.Direction, RayPdf);
+			MisWeight = RayPdf / ((1.0 - SkyLightSamplingStrategyPdf) * RayPdf + SkyLightSamplingStrategyPdf * PdfSkyLightOcclusionRay(FaceCdf, Ray.Direction));
 		}
+		Ray.Origin = WorldPosition;
+		Ray.TMin = 0.0;
+		Ray.TMax = SkyLight.MaxRayDistance;
+		RayWeight = MisWeight / RayPdf;
+#endif
 
 		ApplyCameraRelativeDepthBias(Ray, PixelCoord, DeviceZ, WorldNormal, SkyLight.MaxNormalBias);
 
@@ -340,8 +169,8 @@ void SkyLightRGS()
 			float3 Brdf = LightingSample.Diffuse + LightingSample.Transmission + LightingSample.Specular;
 
 			float3 IncomingRadiance = EvalSkyLight(Ray.Direction);
-			ExitantRadiance += IncomingRadiance * Brdf * MisWeight / RayPdf;
-			DiffuseExitantRadiance += IncomingRadiance * (LightingSample.Diffuse + LightingSample.Transmission) * (MisWeight / RayPdf);
+			ExitantRadiance += IncomingRadiance * Brdf * RayWeight;
+			DiffuseExitantRadiance += IncomingRadiance * (LightingSample.Diffuse + LightingSample.Transmission) * RayWeight;
 		}
 	}
 
diff --git a/Engine/Source/Runtime/Renderer/Private/DeferredShadingRenderer.h b/Engine/Source/Runtime/Renderer/Private/DeferredShadingRenderer.h
--- a/Engine/Source/Runtime/Renderer/Private/DeferredShadingRenderer.h
+++ b/Engine/Source/Runtime/Renderer/Private/DeferredShadingRenderer.h
@@ -529,6 +529,7 @@ private:
 	);
 	void BuildSolidAnglePdf(FRHICommandListImmediate& RHICmdList, const FIntVector& Dimensions, FRWBuffer& SolidAnglePdf);
 
+	void GenerateSkyLightVisibilityRays(FRDGBuilder& GraphBuilder, FRDGBufferRef& SkyLightVisibilityRays, FIntVector& Dimensions);
 	void VisualizeSkyLightMipTree(FRHICommandListImmediate& RHICmdList, const FViewInfo& View, FRWBuffer& SkyLightMipTreePosX, FRWBuffer& SkyLightMipTreePosY, FRWBuffer& SkyLightMipTreePosZ, FRWBuffer& SkyLightMipTreeNegX, FRWBuffer& SkyLightMipTreeNegY, FRWBuffer& SkyLightMipTreeNegZ, const FIntVector& SkyLightMipDimensions);
 	void RenderRayTracingSkyLight(FRHICommandListImmediate& RHICmdList, TRefCountPtr<IPooledRenderTarget>& SkyLightRT, TRefCountPtr<IPooledRenderTarget>& HitDistanceRT);
 	
diff --git a/Engine/Source/Runtime/Renderer/Private/RayTracing/RaytracingSkylight.cpp b/Engine/Source/Runtime/Renderer/Private/RayTracing/RaytracingSkylight.cpp
--- a/Engine/Source/Runtime/Renderer/Private/RayTracing/RaytracingSkylight.cpp
+++ b/Engine/Source/Runtime/Renderer/Private/RayTracing/RaytracingSkylight.cpp
@@ -78,6 +78,13 @@ static TAutoConsoleVariable<int32> CVarRayTracingSkyLightEnableMaterials(
 	ECVF_RenderThreadSafe
 );
 
+static TAutoConsoleVariable<int32> CVarRayTracingSkyLightDecoupleSampleGeneration(
+	TEXT("r.RayTracing.SkyLight.DecoupleSampleGeneration"),
+	1,
+	TEXT("Decouples sample generation from ray traversal (default = 1)"),
+	ECVF_RenderThreadSafe
+);
+
 bool ShouldRenderRayTracingSkyLight(const FSkyLightSceneProxy* SkyLightSceneProxy)
 {
 	if (SkyLightSceneProxy != nullptr)
@@ -173,8 +180,9 @@ class FRayTracingSkyLightRGS : public FGlobalShader
 
 	class FEnableTwoSidedGeometryDim : SHADER_PERMUTATION_BOOL("ENABLE_TWO_SIDED_GEOMETRY");
 	class FEnableMaterialsDim : SHADER_PERMUTATION_BOOL("ENABLE_MATERIALS");
+	class FDecoupleSampleGeneration : SHADER_PERMUTATION_BOOL("DECOUPLE_SAMPLE_GENERATION");
 
-	using FPermutationDomain = TShaderPermutationDomain<FEnableTwoSidedGeometryDim, FEnableMaterialsDim>;
+	using FPermutationDomain = TShaderPermutationDomain<FEnableTwoSidedGeometryDim, FEnableMaterialsDim, FDecoupleSampleGeneration>;
 
 	static bool ShouldCompilePermutation(const FGlobalShaderPermutationParameters& Parameters)
 	{
@@ -191,6 +199,8 @@ class FRayTracingSkyLightRGS : public FGlobalShader
 		SHADER_PARAMETER_STRUCT_REF(FHaltonIteration, HaltonIteration)
 		SHADER_PARAMETER_STRUCT_REF(FHaltonPrimes, HaltonPrimes)
 		SHADER_PARAMETER_STRUCT_REF(FBlueNoise, BlueNoise)
+		SHADER_PARAMETER_RDG_BUFFER_SRV(StructuredBuffer<SkyLightVisibilityRays>, SkyLightVisibilityRays)
+		SHADER_PARAMETER(FIntVector, SkyLightVisibilityRaysDimensions)
 
 		SHADER_PARAMETER_STRUCT_INCLUDE(FSceneTextureParameters, SceneTextures)
 		SHADER_PARAMETER_RDG_TEXTURE(Texture2D, SSProfilesTexture)
@@ -781,14 +791,117 @@ void FDeferredShadingSceneRenderer::PrepareRayTracingSkyLight(const FViewInfo& V
 	{
 		for (uint32 EnableMaterialsIndex = 0; EnableMaterialsIndex < 2; ++EnableMaterialsIndex)
 		{
-			PermutationVector.Set<FRayTracingSkyLightRGS::FEnableTwoSidedGeometryDim>(TwoSidedGeometryIndex != 0);
-			PermutationVector.Set<FRayTracingSkyLightRGS::FEnableMaterialsDim>(EnableMaterialsIndex != 0);
-			TShaderMapRef<FRayTracingSkyLightRGS> RayGenerationShader(View.ShaderMap, PermutationVector);
-			OutRayGenShaders.Add(RayGenerationShader->GetRayTracingShader());
+			for (uint32 DecoupleSampleGeneration = 0; DecoupleSampleGeneration < 2; ++DecoupleSampleGeneration)
+			{
+				PermutationVector.Set<FRayTracingSkyLightRGS::FEnableTwoSidedGeometryDim>(TwoSidedGeometryIndex != 0);
+				PermutationVector.Set<FRayTracingSkyLightRGS::FEnableMaterialsDim>(EnableMaterialsIndex != 0);
+				PermutationVector.Set<FRayTracingSkyLightRGS::FDecoupleSampleGeneration>(DecoupleSampleGeneration != 0);
+				TShaderMapRef<FRayTracingSkyLightRGS> RayGenerationShader(View.ShaderMap, PermutationVector);
+				OutRayGenShaders.Add(RayGenerationShader->GetRayTracingShader());
+			}
 		}
 	}
 }
 
+struct FSkyLightVisibilityRays
+{
+	FVector4 DirectionAndPdf;
+};
+
+class FGenerateSkyLightVisibilityRaysCS : public FGlobalShader
+{
+	DECLARE_GLOBAL_SHADER(FGenerateSkyLightVisibilityRaysCS);
+	SHADER_USE_PARAMETER_STRUCT(FGenerateSkyLightVisibilityRaysCS, FGlobalShader);
+
+	static const uint32 kGroupSize = 16;
+	using FPermutationDomain = FShaderPermutationNone;
+
+	static bool ShouldCompilePermutation(const FGlobalShaderPermutationParameters& Parameters)
+	{
+		return ShouldCompileRayTracingShadersForProject(Parameters.Platform);
+	}
+
+	static void ModifyCompilationEnvironment(const FGlobalShaderPermutationParameters& Parameters, FShaderCompilerEnvironment& OutEnvironment)
+	{
+		FGlobalShader::ModifyCompilationEnvironment(Parameters, OutEnvironment);
+		OutEnvironment.SetDefine(TEXT("TILE_SIZE"), kGroupSize);
+	}
+
+	BEGIN_SHADER_PARAMETER_STRUCT(FParameters, )
+		// Input
+		SHADER_PARAMETER(FIntVector, Dimensions)
+		SHADER_PARAMETER(int32, SamplesPerPixel)
+		SHADER_PARAMETER_STRUCT_REF(FSkyLightData, SkyLightData)
+
+		SHADER_PARAMETER_STRUCT_REF(FHaltonIteration, HaltonIteration)
+		SHADER_PARAMETER_STRUCT_REF(FHaltonPrimes, HaltonPrimes)
+		SHADER_PARAMETER_STRUCT_REF(FBlueNoise, BlueNoise)
+
+		// Output
+		SHADER_PARAMETER_RDG_BUFFER_UAV(RWStructuredBuffer<FSkyLightVisibilityRays>, OutSkyLightVisibilityRays)
+		END_SHADER_PARAMETER_STRUCT()
+};
+
+IMPLEMENT_GLOBAL_SHADER(FGenerateSkyLightVisibilityRaysCS, "/Engine/Private/RayTracing/GenerateSkyLightVisibilityRaysCS.usf", "MainCS", SF_Compute);
+
+void FDeferredShadingSceneRenderer::GenerateSkyLightVisibilityRays(
+	FRDGBuilder& GraphBuilder,
+	FRDGBufferRef& SkyLightVisibilityRaysBuffer,
+	FIntVector& Dimensions
+)
+{
+	FRHIResourceCreateInfo CreateInfo;
+	CreateInfo.DebugName = TEXT("SkyLightVisibilityRays");
+
+	// SkyLight data setup
+	FSkyLightData SkyLightData;
+	SetupSkyLightParameters(*Scene, &SkyLightData);
+	uint32 SamplesPerPixel = GRayTracingSkyLightSamplesPerPixel >= 0 ? GRayTracingSkyLightSamplesPerPixel : Scene->SkyLight->SamplesPerPixel;
+	SkyLightData.SamplesPerPixel = SamplesPerPixel;
+	SkyLightData.MaxRayDistance = GRayTracingSkyLightMaxRayDistance;
+	SkyLightData.SamplingStopLevel = GRayTracingSkyLightSamplingStopLevel;
+
+	// Sampling state
+	FHaltonPrimes HaltonPrimes;
+	InitializeHaltonPrimes(Scene->HaltonPrimesResource, HaltonPrimes);
+
+	FBlueNoise BlueNoise;
+	InitializeBlueNoise(BlueNoise);
+	Dimensions = FIntVector(BlueNoise.Dimensions.X, BlueNoise.Dimensions.Y, 0);
+
+	// Halton iteration
+	uint32 IterationCount = FMath::Max(SamplesPerPixel, 1u);
+	uint32 SequenceCount = 1;
+	uint32 DimensionCount = 3;
+	FHaltonSequenceIteration HaltonSequenceIteration(Scene->HaltonSequence, IterationCount, SequenceCount, DimensionCount, Views[0].ViewState ? (Views[0].ViewState->FrameIndex % 1024) : 0);
+
+	FHaltonIteration HaltonIteration;
+	InitializeHaltonSequenceIteration(HaltonSequenceIteration, HaltonIteration);
+
+	// Output structured buffer creation
+	FRDGBufferDesc BufferDesc = FRDGBufferDesc::CreateStructuredDesc(sizeof(FSkyLightVisibilityRays), Dimensions.X * Dimensions.Y * SkyLightData.SamplesPerPixel);
+	SkyLightVisibilityRaysBuffer = GraphBuilder.CreateBuffer(BufferDesc, TEXT("SkyLightVisibilityRays"));
+
+	// Compute Pass parameter definition
+	FGenerateSkyLightVisibilityRaysCS::FParameters* PassParameters = GraphBuilder.AllocParameters<FGenerateSkyLightVisibilityRaysCS::FParameters>();
+	PassParameters->Dimensions = FIntVector(Dimensions.X, Dimensions.Y, 0);
+	PassParameters->SkyLightData = CreateUniformBufferImmediate(SkyLightData, EUniformBufferUsage::UniformBuffer_SingleDraw);
+	PassParameters->HaltonPrimes = CreateUniformBufferImmediate(HaltonPrimes, EUniformBufferUsage::UniformBuffer_SingleDraw);
+	PassParameters->BlueNoise = CreateUniformBufferImmediate(BlueNoise, EUniformBufferUsage::UniformBuffer_SingleDraw);
+	PassParameters->HaltonIteration = CreateUniformBufferImmediate(HaltonIteration, EUniformBufferUsage::UniformBuffer_SingleDraw);
+	PassParameters->OutSkyLightVisibilityRays = GraphBuilder.CreateUAV(SkyLightVisibilityRaysBuffer, EPixelFormat::PF_R32_UINT);
+
+	TShaderMapRef<FGenerateSkyLightVisibilityRaysCS> ComputeShader(GetGlobalShaderMap(FeatureLevel));
+
+	FComputeShaderUtils::AddPass(
+		GraphBuilder,
+		RDG_EVENT_NAME("GenerateSkyLightVisibilityRays"),
+		*ComputeShader,
+		PassParameters,
+		FComputeShaderUtils::GetGroupCount(FIntPoint(Dimensions.X, Dimensions.Y), FGenerateSkyLightVisibilityRaysCS::kGroupSize)
+	);
+}
+
 void FDeferredShadingSceneRenderer::RenderRayTracingSkyLight(
 	FRHICommandListImmediate& RHICmdList,
 	TRefCountPtr<IPooledRenderTarget>& SkyLightRT,
@@ -823,6 +936,18 @@ void FDeferredShadingSceneRenderer::RenderRayTracingSkyLight(
 	}
 
 	FRDGBuilder GraphBuilder(RHICmdList);
+	FRDGBufferRef SkyLightVisibilityRaysBuffer;
+	FIntVector SkyLightVisibilityRaysDimensions;
+	if (CVarRayTracingSkyLightDecoupleSampleGeneration.GetValueOnRenderThread() == 1)
+	{
+		GenerateSkyLightVisibilityRays(GraphBuilder, SkyLightVisibilityRaysBuffer, SkyLightVisibilityRaysDimensions);
+	}
+	else
+	{
+		FRDGBufferDesc BufferDesc = FRDGBufferDesc::CreateStructuredDesc(sizeof(FSkyLightVisibilityRays), 1);
+		SkyLightVisibilityRaysBuffer = GraphBuilder.CreateBuffer(BufferDesc, TEXT("SkyLightVisibilityRays"));
+		SkyLightVisibilityRaysDimensions = FIntVector(1);
+	}
 
 	// Define RDG targets
 	FRDGTextureRef SkyLightTexture;
@@ -872,6 +997,11 @@ void FDeferredShadingSceneRenderer::RenderRayTracingSkyLight(
 	PassParameters->SkyLightData = CreateUniformBufferImmediate(SkyLightData, EUniformBufferUsage::UniformBuffer_SingleDraw);
 	PassParameters->HaltonPrimes = CreateUniformBufferImmediate(HaltonPrimes, EUniformBufferUsage::UniformBuffer_SingleDraw);
 	PassParameters->BlueNoise = CreateUniformBufferImmediate(BlueNoise, EUniformBufferUsage::UniformBuffer_SingleDraw);
+	PassParameters->SkyLightVisibilityRaysDimensions = SkyLightVisibilityRaysDimensions;
+	if (CVarRayTracingSkyLightDecoupleSampleGeneration.GetValueOnRenderThread() == 1)
+	{
+		PassParameters->SkyLightVisibilityRays = GraphBuilder.CreateSRV(SkyLightVisibilityRaysBuffer, EPixelFormat::PF_R32_UINT);
+	}
 	PassParameters->SSProfilesTexture = GraphBuilder.RegisterExternalTexture(SubsurfaceProfileRT);
 	PassParameters->TransmissionProfilesLinearSampler = TStaticSamplerState<SF_Bilinear, AM_Clamp, AM_Clamp, AM_Clamp>::GetRHI();
 	PassParameters->SceneTextures = SceneTextures;
@@ -880,7 +1010,6 @@ void FDeferredShadingSceneRenderer::RenderRayTracingSkyLight(
 	{
 		FViewInfo& View = Views[ViewIndex];
 
-
 		// Halton iteration
 		uint32 IterationCount = FMath::Max(SamplesPerPixel, 1u);
 		uint32 SequenceCount = 1;
@@ -898,6 +1027,7 @@ void FDeferredShadingSceneRenderer::RenderRayTracingSkyLight(
 		FRayTracingSkyLightRGS::FPermutationDomain PermutationVector;
 		PermutationVector.Set<FRayTracingSkyLightRGS::FEnableTwoSidedGeometryDim>(CVarRayTracingSkyLightEnableTwoSidedGeometry.GetValueOnRenderThread() != 0);
 		PermutationVector.Set<FRayTracingSkyLightRGS::FEnableMaterialsDim>(CVarRayTracingSkyLightEnableMaterials.GetValueOnRenderThread() != 0);
+		PermutationVector.Set<FRayTracingSkyLightRGS::FDecoupleSampleGeneration>(CVarRayTracingSkyLightDecoupleSampleGeneration.GetValueOnRenderThread() != 0);
 		TShaderMapRef<FRayTracingSkyLightRGS> RayGenerationShader(GetGlobalShaderMap(FeatureLevel), PermutationVector);
 		ClearUnusedGraphResources(*RayGenerationShader, PassParameters);
 
