Cherry-pick 7308853

Disable normal-culling when casting a shadow ray from geometry with hair material.
Use proxy light direction and hair tangent direction to reconstruct a plausible normal for ray biasing.

diff --git a/Engine/Shaders/Private/RayTracing/RayTracingOcclusionRGS.usf b/Engine/Shaders/Private/RayTracing/RayTracingOcclusionRGS.usf
--- a/Engine/Shaders/Private/RayTracing/RayTracingOcclusionRGS.usf
+++ b/Engine/Shaders/Private/RayTracing/RayTracingOcclusionRGS.usf
@@ -165,6 +165,32 @@ void OcclusionRGS()
 				WorldNormal = -WorldNormal;
 		}
 
+		// For hair shading model, WorldNormal is actually the shading tangent.
+		// To generate proper bias, we compute a normal oriented toward the light.
+		// Normal clipping is removed from hair since the BSDF is spherical, rather 
+		// than hemispherical
+		//
+		// Note: 
+		// Since we don't have a notion of backfacing here, there is no correct way 
+		// to avoid self-intersection with hair. Either we push towards the light, 
+		// (this creates some issue in triangle facing back the light), or we push 
+		// toward the view point (this create issue for transmitted light)
+		// It seems having proper transmission is more important, thus the light 
+		// version is enabled by default
+		bool bApplyNormalCulling = true;
+		if (GBufferData.ShadingModelID == SHADINGMODELID_HAIR)
+		{
+#if LIGHT_TYPE == LIGHT_TYPE_DIRECTIONAL
+			const float3 LightDirection = LightParameters.Direction;
+#else
+			const float3 LightDirection = normalize(LightParameters.Position - WorldPosition);
+#endif
+
+			const float3 TangentDirection = normalize(WorldNormal);
+			WorldNormal = normalize(LightDirection - TangentDirection * dot(LightDirection, TangentDirection));
+			bApplyNormalCulling = false;
+		}
+
 		LOOP
 		for (uint SampleIndex = 0; SampleIndex < LocalSamplesPerPixel; ++SampleIndex)
 		{
@@ -199,7 +225,7 @@ void OcclusionRGS()
 			{
 				continue;
 			}
-			else if (dot(WorldNormal, Ray.Direction) <= 0.0)
+			else if (bApplyNormalCulling && dot(WorldNormal, Ray.Direction) <= 0.0)
 			{
 				ClosestRayDistance = 0.001;
 				continue;
