Duplicate CL #4396224: Remove another broken bit of legacy Metal code.

diff --git a/Engine/Shaders/Private/DecalCommon.ush b/Engine/Shaders/Private/DecalCommon.ush
--- a/Engine/Shaders/Private/DecalCommon.ush
+++ b/Engine/Shaders/Private/DecalCommon.ush
@@ -97,7 +97,7 @@ void DecalCommonOutput(inout FPixelShaderIn In, inout FPixelShaderOut Out, float
 	#endif // DECAL_RENDERSTAGE
 #endif //DECAL_RENDERTARGET_COUNT > 1
 
-#if DECAL_RENDERTARGET_COUNT > 1 && !BIND_RENDERTARGET1 && (COMPILER_METAL || COMPILER_SWITCH)
+#if DECAL_RENDERTARGET_COUNT > 1 && !BIND_RENDERTARGET1 && (COMPILER_SWITCH)
 	Out.MRT[1] = 0;
 #endif
 }
