[CONSOLE] [!] Minor PS4 shader compile fix.

diff --git a/Engine/Shaders/Private/GPUFastFourierTransformCore.ush b/Engine/Shaders/Private/GPUFastFourierTransformCore.ush
--- a/Engine/Shaders/Private/GPUFastFourierTransformCore.ush
+++ b/Engine/Shaders/Private/GPUFastFourierTransformCore.ush
@@ -575,6 +575,10 @@ void GroupSharedFFT(in const bool bIsForward, inout Complex Local[RADIX], in con
 
 	// NumCols = ArrayLength / RADIX: NumCols is an int power of RADIX (ie. 0, 1, 2..)
 	uint Ns = RADIX;
+
+#if PS4_PROFILE
+	LOOP
+#endif
 	for ( ; Ns < NumCols; Ns *= RADIX ) 
 	{
 		
