Reduced the amount of interpolants in FVertexFactoryInterpolantsVSToPS by the Niagara when using DynamicParameters.
Now only the param required are used where as before the count was defined by the max index.

diff --git a/Engine/Shaders/Private/MaterialTemplate.ush b/Engine/Shaders/Private/MaterialTemplate.ush
--- a/Engine/Shaders/Private/MaterialTemplate.ush
+++ b/Engine/Shaders/Private/MaterialTemplate.ush
@@ -94,15 +94,15 @@ struct FMaterialParticleParameters
 	/** Dynamic parameters used by particle systems. */
 	half4 DynamicParameter;
 
-#if NUM_DYNAMIC_PARAMETERS > 1
+#if( DYNAMIC_PARAMETERS_MASK & 2)
 	half4 DynamicParameter1;
 #endif
 
-#if NUM_DYNAMIC_PARAMETERS > 2
+#if (DYNAMIC_PARAMETERS_MASK & 4)
 	half4 DynamicParameter2;
 #endif
 
-#if NUM_DYNAMIC_PARAMETERS > 3
+#if (DYNAMIC_PARAMETERS_MASK & 8)
 	half4 DynamicParameter3;
 #endif
 
@@ -127,19 +127,19 @@ float4 GetDynamicParameter(FMaterialParticleParameters Parameters, float4 Defaul
 	{
 		return Parameters.DynamicParameter;
 	}
-#if NUM_DYNAMIC_PARAMETERS > 1
+#if (DYNAMIC_PARAMETERS_MASK & 2)
 	else if(ParameterIndex == 1)
 	{
 		return Parameters.DynamicParameter1;
 	}
 #endif
-#if NUM_DYNAMIC_PARAMETERS > 2
+#if (DYNAMIC_PARAMETERS_MASK & 4)
 	else if(ParameterIndex == 2)
 	{
 		return Parameters.DynamicParameter2;
 	}
 #endif	
-#if NUM_DYNAMIC_PARAMETERS > 3
+#if (DYNAMIC_PARAMETERS_MASK & 8)
 	else if(ParameterIndex == 3)
 	{
 		return Parameters.DynamicParameter3;
diff --git a/Engine/Shaders/Private/NiagaraMeshVertexFactory.ush b/Engine/Shaders/Private/NiagaraMeshVertexFactory.ush
--- a/Engine/Shaders/Private/NiagaraMeshVertexFactory.ush
+++ b/Engine/Shaders/Private/NiagaraMeshVertexFactory.ush
@@ -79,16 +79,16 @@ struct FVertexFactoryInterpolantsVSToPS
 	float4	ParticleColor : COLOR1;
 #endif
 
-#if NUM_DYNAMIC_PARAMETERS > 0
+#if (DYNAMIC_PARAMETERS_MASK & 1)
 	float4 DynamicParameter : COLOR2;
 #endif
-#if NUM_DYNAMIC_PARAMETERS > 1
+#if (DYNAMIC_PARAMETERS_MASK & 2)
 	float4 DynamicParameter1 : COLOR3;
 #endif
-#if NUM_DYNAMIC_PARAMETERS > 2
+#if (DYNAMIC_PARAMETERS_MASK & 4)
 	float4 DynamicParameter2 : COLOR4;
 #endif
-#if NUM_DYNAMIC_PARAMETERS > 3
+#if (DYNAMIC_PARAMETERS_MASK & 8)
 	float4 DynamicParameter3 : COLOR5;
 #endif
 
@@ -138,16 +138,16 @@ struct FVertexFactoryIntermediates
 #endif
 
 	/** Optional dynamic parameters for the particle. */
-#if NUM_DYNAMIC_PARAMETERS > 0
+#if (DYNAMIC_PARAMETERS_MASK & 1)
 	float4 DynamicParameter;
 #endif
-#if NUM_DYNAMIC_PARAMETERS > 1
+#if (DYNAMIC_PARAMETERS_MASK & 2)
 	float4 DynamicParameter1;
 #endif
-#if NUM_DYNAMIC_PARAMETERS > 2
+#if (DYNAMIC_PARAMETERS_MASK & 4)
 	float4 DynamicParameter2;
 #endif
-#if NUM_DYNAMIC_PARAMETERS > 3
+#if (DYNAMIC_PARAMETERS_MASK & 8)
 	float4 DynamicParameter3;
 #endif
 
@@ -247,16 +247,16 @@ FMaterialPixelParameters GetMaterialPixelParameters(FVertexFactoryInterpolantsVS
 	Result.Particle.Color = Interpolants.ParticleColor;
 #endif
 
-#if NUM_DYNAMIC_PARAMETERS > 0
+#if (DYNAMIC_PARAMETERS_MASK & 1)
 	Result.Particle.DynamicParameter = Interpolants.DynamicParameter;
 #endif
-#if NUM_DYNAMIC_PARAMETERS > 1
+#if (DYNAMIC_PARAMETERS_MASK & 2)
 	Result.Particle.DynamicParameter1 = Interpolants.DynamicParameter1;
 #endif
-#if NUM_DYNAMIC_PARAMETERS > 2
+#if (DYNAMIC_PARAMETERS_MASK & 4)
 	Result.Particle.DynamicParameter2 = Interpolants.DynamicParameter2;
 #endif
-#if NUM_DYNAMIC_PARAMETERS > 3
+#if (DYNAMIC_PARAMETERS_MASK & 8)
 	Result.Particle.DynamicParameter3 = Interpolants.DynamicParameter3;
 #endif
 
@@ -327,16 +327,16 @@ FMaterialVertexParameters GetMaterialVertexParameters(FVertexFactoryInput Input,
 
 	Result.Particle.Color = Intermediates.ParticleColor;
 
-#if NUM_DYNAMIC_PARAMETERS > 0
+#if (DYNAMIC_PARAMETERS_MASK & 1)
 	Result.Particle.DynamicParameter = Intermediates.DynamicParameter;
 #endif
-#if NUM_DYNAMIC_PARAMETERS > 1
+#if (DYNAMIC_PARAMETERS_MASK & 2)
 	Result.Particle.DynamicParameter1 = Intermediates.DynamicParameter1;
 #endif
-#if NUM_DYNAMIC_PARAMETERS > 2
+#if (DYNAMIC_PARAMETERS_MASK & 4)
 	Result.Particle.DynamicParameter2 = Intermediates.DynamicParameter2;
 #endif
-#if NUM_DYNAMIC_PARAMETERS > 3
+#if (DYNAMIC_PARAMETERS_MASK & 8)
 	Result.Particle.DynamicParameter3 = Intermediates.DynamicParameter3;
 #endif
 
@@ -526,16 +526,16 @@ FVertexFactoryIntermediates GetVertexFactoryIntermediates(FVertexFactoryInput In
 		ParticleVelocity = float3(GetVec3(NiagaraMeshVF.VelocityDataOffset, Input.InstanceId));
 	}
 
-#if NUM_DYNAMIC_PARAMETERS > 0
+#if (DYNAMIC_PARAMETERS_MASK & 1)
 	Intermediates.DynamicParameter  = NiagaraMeshVF.MaterialParamDataOffset != -1 ? GetVec4(NiagaraMeshVF.MaterialParamDataOffset, Input.InstanceId) : float4(1,1,1,1);
 #endif
-#if NUM_DYNAMIC_PARAMETERS > 1
+#if (DYNAMIC_PARAMETERS_MASK & 2)
 	Intermediates.DynamicParameter1  = NiagaraMeshVF.MaterialParam1DataOffset != -1 ? GetVec4(NiagaraMeshVF.MaterialParam1DataOffset, Input.InstanceId) : float4(1,1,1,1);
 #endif
-#if NUM_DYNAMIC_PARAMETERS > 2
+#if (DYNAMIC_PARAMETERS_MASK & 4)
 	Intermediates.DynamicParameter2  = NiagaraMeshVF.MaterialParam2DataOffset != -1 ? GetVec4(NiagaraMeshVF.MaterialParam2DataOffset, Input.InstanceId) : float4(1,1,1,1);
 #endif
-#if NUM_DYNAMIC_PARAMETERS > 3
+#if (DYNAMIC_PARAMETERS_MASK & 8)
 	Intermediates.DynamicParameter3  = NiagaraMeshVF.MaterialParam3DataOffset != -1 ? GetVec4(NiagaraMeshVF.MaterialParam3DataOffset, Input.InstanceId) : float4(1,1,1,1);
 #endif
 
@@ -763,16 +763,16 @@ FVertexFactoryInterpolantsVSToPS VertexFactoryGetInterpolantsVSToPS(FVertexFacto
 	Interpolants.ParticleColor = Intermediates.ParticleColor;
 #endif
 
-#if NUM_DYNAMIC_PARAMETERS > 0
+#if (DYNAMIC_PARAMETERS_MASK & 1)
 	Interpolants.DynamicParameter = Intermediates.DynamicParameter;
 #endif
-#if NUM_DYNAMIC_PARAMETERS > 1
+#if (DYNAMIC_PARAMETERS_MASK & 2)
 	Interpolants.DynamicParameter1 = Intermediates.DynamicParameter1;
 #endif
-#if NUM_DYNAMIC_PARAMETERS > 2
+#if (DYNAMIC_PARAMETERS_MASK & 4)
 	Interpolants.DynamicParameter2 = Intermediates.DynamicParameter2;
 #endif
-#if NUM_DYNAMIC_PARAMETERS > 3
+#if (DYNAMIC_PARAMETERS_MASK & 8)
 	Interpolants.DynamicParameter3 = Intermediates.DynamicParameter3;
 #endif
 
@@ -888,16 +888,16 @@ float4 VertexFactoryGetPreviousWorldPosition(FVertexFactoryInput Input, FVertexF
 		}
 #endif
 
-#if NUM_DYNAMIC_PARAMETERS > 0
+#if (DYNAMIC_PARAMETERS_MASK & 1)
 		TESSELLATION_INTERPOLATE_MEMBER(InterpolantsVSToPS.DynamicParameter);
 #endif
-#if NUM_DYNAMIC_PARAMETERS > 1
+#if (DYNAMIC_PARAMETERS_MASK & 2)
 		TESSELLATION_INTERPOLATE_MEMBER(InterpolantsVSToPS.DynamicParameter1);
 #endif
-#if NUM_DYNAMIC_PARAMETERS > 2
+#if (DYNAMIC_PARAMETERS_MASK & 4)
 		TESSELLATION_INTERPOLATE_MEMBER(InterpolantsVSToPS.DynamicParameter2);
 #endif
-#if NUM_DYNAMIC_PARAMETERS > 3
+#if (DYNAMIC_PARAMETERS_MASK & 8)
 		TESSELLATION_INTERPOLATE_MEMBER(InterpolantsVSToPS.DynamicParameter3);
 #endif
 
diff --git a/Engine/Shaders/Private/NiagaraRibbonVertexFactory.ush b/Engine/Shaders/Private/NiagaraRibbonVertexFactory.ush
--- a/Engine/Shaders/Private/NiagaraRibbonVertexFactory.ush
+++ b/Engine/Shaders/Private/NiagaraRibbonVertexFactory.ush
@@ -92,7 +92,7 @@ float GetNiagaraMaterialRandom(uint InstanceID)
 }
 #endif
 
-#if NUM_DYNAMIC_PARAMETERS > 0
+#if (DYNAMIC_PARAMETERS_MASK & 1)
 float4 GetNiagaraParticleDynamicParameters(uint InstanceID)
 {
 	if (NiagaraRibbonVF.MaterialParamDataOffset == -1)
@@ -102,7 +102,7 @@ float4 GetNiagaraParticleDynamicParameters(uint InstanceID)
 	return GetVec4(NiagaraRibbonVF.MaterialParamDataOffset, InstanceID);
 }
 #endif
-#if NUM_DYNAMIC_PARAMETERS > 1
+#if (DYNAMIC_PARAMETERS_MASK & 2)
 float4 GetNiagaraParticleDynamicParameters1(uint InstanceID)
 {
 	if (NiagaraRibbonVF.MaterialParam1DataOffset == -1)
@@ -112,7 +112,7 @@ float4 GetNiagaraParticleDynamicParameters1(uint InstanceID)
 	return GetVec4(NiagaraRibbonVF.MaterialParam1DataOffset, InstanceID);
 }
 #endif
-#if NUM_DYNAMIC_PARAMETERS > 2
+#if (DYNAMIC_PARAMETERS_MASK & 4)
 float4 GetNiagaraParticleDynamicParameters2(uint InstanceID)
 {
 	if (NiagaraRibbonVF.MaterialParam2DataOffset == -1)
@@ -122,7 +122,7 @@ float4 GetNiagaraParticleDynamicParameters2(uint InstanceID)
 	return GetVec4(NiagaraRibbonVF.MaterialParam2DataOffset, InstanceID);
 }
 #endif
-#if NUM_DYNAMIC_PARAMETERS > 3
+#if (DYNAMIC_PARAMETERS_MASK & 8)
 float4 GetNiagaraParticleDynamicParameters3(uint InstanceID)
 {
 	if (NiagaraRibbonVF.MaterialParam3DataOffset == -1)
@@ -170,16 +170,16 @@ struct FVertexFactoryInterpolantsVSToPS
 	// Last row of the tangent to world matrix in xyz
 	float4	TangentToWorld2	: TANGENTTOWORLD2;
 
-#if NUM_DYNAMIC_PARAMETERS > 0
+#if (DYNAMIC_PARAMETERS_MASK & 1)
 	float4 DynamicParameter	: PARTICLE_DYNAMIC_PARAM0;
 #endif
-#if NUM_DYNAMIC_PARAMETERS > 1
+#if (DYNAMIC_PARAMETERS_MASK & 2)
 	float4 DynamicParameter1	: PARTICLE_DYNAMIC_PARAM1;
 #endif
-#if NUM_DYNAMIC_PARAMETERS > 2
+#if (DYNAMIC_PARAMETERS_MASK & 4)
 	float4 DynamicParameter2	: PARTICLE_DYNAMIC_PARAM2;
 #endif
-#if NUM_DYNAMIC_PARAMETERS > 3
+#if (DYNAMIC_PARAMETERS_MASK & 8)
 	float4 DynamicParameter3	: PARTICLE_DYNAMIC_PARAM3;
 #endif
 
@@ -251,16 +251,16 @@ struct FVertexFactoryIntermediates
 	float ParticleRandom;
 #endif
 
-#if NUM_DYNAMIC_PARAMETERS > 0
+#if (DYNAMIC_PARAMETERS_MASK & 1)
 	float4 DynamicParameter;
 #endif
-#if NUM_DYNAMIC_PARAMETERS > 1
+#if (DYNAMIC_PARAMETERS_MASK & 2)
 	float4 DynamicParameter1;
 #endif
-#if NUM_DYNAMIC_PARAMETERS > 2
+#if (DYNAMIC_PARAMETERS_MASK & 4)
 	float4 DynamicParameter2;
 #endif
-#if NUM_DYNAMIC_PARAMETERS > 3
+#if (DYNAMIC_PARAMETERS_MASK & 8)
 	float4 DynamicParameter3;
 #endif
 };
@@ -370,16 +370,16 @@ FMaterialPixelParameters GetMaterialPixelParameters(FVertexFactoryInterpolantsVS
 
 	Result.TwoSidedSign = 1;
 
-#if NUM_DYNAMIC_PARAMETERS > 0
+#if (DYNAMIC_PARAMETERS_MASK & 1)
 	Result.Particle.DynamicParameter = Interpolants.DynamicParameter;
 #endif
-#if NUM_DYNAMIC_PARAMETERS > 1
+#if (DYNAMIC_PARAMETERS_MASK & 2)
 	Result.Particle.DynamicParameter1 = Interpolants.DynamicParameter1;
 #endif
-#if NUM_DYNAMIC_PARAMETERS > 2
+#if (DYNAMIC_PARAMETERS_MASK & 4)
 	Result.Particle.DynamicParameter2 = Interpolants.DynamicParameter2;
 #endif
-#if NUM_DYNAMIC_PARAMETERS > 3
+#if (DYNAMIC_PARAMETERS_MASK & 8)
 	Result.Particle.DynamicParameter3 = Interpolants.DynamicParameter3;
 #endif
 
@@ -419,16 +419,16 @@ FMaterialVertexParameters GetMaterialVertexParameters(FVertexFactoryInput Input,
 	Result.PreSkinnedPosition = Intermediates.Position.xyz;
 	Result.PreSkinnedNormal = TangentToLocal[2].xyz;
 
-#if NUM_DYNAMIC_PARAMETERS > 0
+#if (DYNAMIC_PARAMETERS_MASK & 1)
 	Result.Particle.DynamicParameter = Intermediates.DynamicParameter;
 #endif
-#if NUM_DYNAMIC_PARAMETERS > 1
+#if (DYNAMIC_PARAMETERS_MASK & 2)
 	Result.Particle.DynamicParameter1 = Intermediates.DynamicParameter1;
 #endif
-#if NUM_DYNAMIC_PARAMETERS > 2
+#if (DYNAMIC_PARAMETERS_MASK & 4)
 	Result.Particle.DynamicParameter2 = Intermediates.DynamicParameter2;
 #endif
-#if NUM_DYNAMIC_PARAMETERS > 3
+#if (DYNAMIC_PARAMETERS_MASK & 8)
 	Result.Particle.DynamicParameter3 = Intermediates.DynamicParameter3;
 #endif
 
@@ -580,16 +580,16 @@ FVertexFactoryIntermediates GetVertexFactoryIntermediates(FVertexFactoryInput In
 		Intermediates.CustomFacingVector = GetNiagaraParticleFacingVector(ParticleId);
 	}
 
-#if NUM_DYNAMIC_PARAMETERS > 0
+#if (DYNAMIC_PARAMETERS_MASK & 1)
 	Intermediates.DynamicParameter = GetNiagaraParticleDynamicParameters(ParticleId);
 #endif
-#if NUM_DYNAMIC_PARAMETERS > 1
+#if (DYNAMIC_PARAMETERS_MASK & 2)
 	Intermediates.DynamicParameter1 = GetNiagaraParticleDynamicParameters1(ParticleId);
 #endif
-#if NUM_DYNAMIC_PARAMETERS > 2
+#if (DYNAMIC_PARAMETERS_MASK & 4)
 	Intermediates.DynamicParameter2 = GetNiagaraParticleDynamicParameters2(ParticleId);
 #endif
-#if NUM_DYNAMIC_PARAMETERS > 3
+#if (DYNAMIC_PARAMETERS_MASK & 8)
 	Intermediates.DynamicParameter3 = GetNiagaraParticleDynamicParameters3(ParticleId);
 #endif
 
@@ -685,16 +685,16 @@ FVertexFactoryInterpolantsVSToPS VertexFactoryGetInterpolantsVSToPS(FVertexFacto
 	Interpolants.Color		= Intermediates.Color;
 #endif
 
-#if NUM_DYNAMIC_PARAMETERS > 0
+#if (DYNAMIC_PARAMETERS_MASK & 1)
 	Interpolants.DynamicParameter = Intermediates.DynamicParameter;
 #endif
-#if NUM_DYNAMIC_PARAMETERS > 1
+#if (DYNAMIC_PARAMETERS_MASK & 2)
 	Interpolants.DynamicParameter1 = Intermediates.DynamicParameter1;
 #endif
-#if NUM_DYNAMIC_PARAMETERS > 2
+#if (DYNAMIC_PARAMETERS_MASK & 4)
 	Interpolants.DynamicParameter2 = Intermediates.DynamicParameter2;
 #endif
-#if NUM_DYNAMIC_PARAMETERS > 3
+#if (DYNAMIC_PARAMETERS_MASK & 8)
 	Interpolants.DynamicParameter3 = Intermediates.DynamicParameter3;
 #endif
 
diff --git a/Engine/Shaders/Private/NiagaraSpriteVertexFactory.ush b/Engine/Shaders/Private/NiagaraSpriteVertexFactory.ush
--- a/Engine/Shaders/Private/NiagaraSpriteVertexFactory.ush
+++ b/Engine/Shaders/Private/NiagaraSpriteVertexFactory.ush
@@ -61,16 +61,16 @@ struct FVertexFactoryInterpolantsVSToPS
 	// Last row of the tangent to world matrix in xyz
 	float4	TangentToWorld2	: TANGENTTOWORLD2;
 	
-#if NUM_DYNAMIC_PARAMETERS > 0
+#if (DYNAMIC_PARAMETERS_MASK & 1)
 	float4 DynamicParameter	: PARTICLE_DYNAMIC_PARAM0;
 #endif
-#if NUM_DYNAMIC_PARAMETERS > 1
+#if (DYNAMIC_PARAMETERS_MASK & 2)
 	float4 DynamicParameter1	: PARTICLE_DYNAMIC_PARAM1;
 #endif
-#if NUM_DYNAMIC_PARAMETERS > 2
+#if (DYNAMIC_PARAMETERS_MASK & 4)
 	float4 DynamicParameter2	: PARTICLE_DYNAMIC_PARAM2;
 #endif
-#if NUM_DYNAMIC_PARAMETERS > 3
+#if (DYNAMIC_PARAMETERS_MASK & 8)
 	float4 DynamicParameter3	: PARTICLE_DYNAMIC_PARAM3;
 #endif
 
@@ -160,13 +160,13 @@ struct FVertexFactoryIntermediates
 	float3 CustomAlignmentVector;
 
 	float4 DynamicParameter;
-#if NUM_DYNAMIC_PARAMETERS > 1
+#if (DYNAMIC_PARAMETERS_MASK & 2)
 	float4 DynamicParameter1;
 #endif
-#if NUM_DYNAMIC_PARAMETERS > 2
+#if (DYNAMIC_PARAMETERS_MASK & 4)
 	float4 DynamicParameter2;
 #endif
-#if NUM_DYNAMIC_PARAMETERS > 3
+#if (DYNAMIC_PARAMETERS_MASK & 8)
 	float4 DynamicParameter3;
 #endif
 
@@ -257,16 +257,16 @@ FMaterialPixelParameters GetMaterialPixelParameters(FVertexFactoryInterpolantsVS
 	Result.Particle.Color = Interpolants.Color;
 #endif
 
-#if NUM_DYNAMIC_PARAMETERS > 0
+#if (DYNAMIC_PARAMETERS_MASK & 1)
 	Result.Particle.DynamicParameter = Interpolants.DynamicParameter;
 #endif
-#if NUM_DYNAMIC_PARAMETERS > 1
+#if (DYNAMIC_PARAMETERS_MASK & 2)
 	Result.Particle.DynamicParameter1 = Interpolants.DynamicParameter1;
 #endif
-#if NUM_DYNAMIC_PARAMETERS > 2
+#if (DYNAMIC_PARAMETERS_MASK & 4)
 	Result.Particle.DynamicParameter2 = Interpolants.DynamicParameter2;
 #endif
-#if NUM_DYNAMIC_PARAMETERS > 3
+#if (DYNAMIC_PARAMETERS_MASK & 8)
 	Result.Particle.DynamicParameter3 = Interpolants.DynamicParameter3;
 #endif
 
@@ -330,16 +330,16 @@ FMaterialVertexParameters GetMaterialVertexParameters(FVertexFactoryInput Input,
 	// primitives using this VF are usually transparent and hence don't output velocity
 	Result.PrevFrameLocalToWorld = Primitive.LocalToWorld;
 
-#if NUM_DYNAMIC_PARAMETERS > 0
+#if (DYNAMIC_PARAMETERS_MASK & 1)
 	Result.Particle.DynamicParameter = Intermediates.DynamicParameter;
 #endif
-#if NUM_DYNAMIC_PARAMETERS > 1
+#if (DYNAMIC_PARAMETERS_MASK & 2)
 	Result.Particle.DynamicParameter1 = Intermediates.DynamicParameter1;
 #endif
-#if NUM_DYNAMIC_PARAMETERS > 2
+#if (DYNAMIC_PARAMETERS_MASK & 4)
 	Result.Particle.DynamicParameter2 = Intermediates.DynamicParameter2;
 #endif
-#if NUM_DYNAMIC_PARAMETERS > 3
+#if (DYNAMIC_PARAMETERS_MASK & 8)
 	Result.Particle.DynamicParameter3 = Intermediates.DynamicParameter3;
 #endif
 
@@ -680,25 +680,25 @@ float GetNiagaraNormalizedAge(uint InstanceID)
 }
 #endif
 
-#if NUM_DYNAMIC_PARAMETERS > 0
+#if (DYNAMIC_PARAMETERS_MASK & 1)
 float4 GetNiagaraParticleDynamicParameters(uint InstanceID)
 {
 	return NiagaraSpriteVF.MaterialParamDataOffset != -1 ? GetVec4(NiagaraSpriteVF.MaterialParamDataOffset, InstanceID) : float4(1.0f, 1.0f, 1.0f, 1.0f);;
 }
 #endif
-#if NUM_DYNAMIC_PARAMETERS > 1
+#if (DYNAMIC_PARAMETERS_MASK & 2)
 float4 GetNiagaraParticleDynamicParameters1(uint InstanceID)
 {
 	return NiagaraSpriteVF.MaterialParam1DataOffset != -1 ? GetVec4(NiagaraSpriteVF.MaterialParam1DataOffset, InstanceID) : float4(1.0f, 1.0f, 1.0f, 1.0f);;
 }
 #endif
-#if NUM_DYNAMIC_PARAMETERS > 2
+#if (DYNAMIC_PARAMETERS_MASK & 4)
 float4 GetNiagaraParticleDynamicParameters2(uint InstanceID)
 {
 	return NiagaraSpriteVF.MaterialParam2DataOffset != -1 ? GetVec4(NiagaraSpriteVF.MaterialParam2DataOffset, InstanceID) : float4(1.0f, 1.0f, 1.0f, 1.0f);;
 }
 #endif
-#if NUM_DYNAMIC_PARAMETERS > 3
+#if (DYNAMIC_PARAMETERS_MASK & 8)
 float4 GetNiagaraParticleDynamicParameters3(uint InstanceID)
 {
 	return NiagaraSpriteVF.MaterialParam3DataOffset != -1 ? GetVec4(NiagaraSpriteVF.MaterialParam3DataOffset, InstanceID) : float4(1.0f, 1.0f, 1.0f, 1.0f);;
@@ -757,16 +757,16 @@ FVertexFactoryIntermediates GetVertexFactoryIntermediates(FVertexFactoryInput In
 	Intermediates.CustomFacingVector = SafeNormalize(GetNiagaraParticleFacingVector(ParticleID));
 	Intermediates.CustomAlignmentVector = SafeNormalize(GetNiagaraParticleAlignmentVector(ParticleID));
 
-#if NUM_DYNAMIC_PARAMETERS > 0
+#if (DYNAMIC_PARAMETERS_MASK & 1)
 	Intermediates.DynamicParameter = GetNiagaraParticleDynamicParameters(ParticleID);
 #endif
-#if NUM_DYNAMIC_PARAMETERS > 1
+#if (DYNAMIC_PARAMETERS_MASK & 2)
 	Intermediates.DynamicParameter1 = GetNiagaraParticleDynamicParameters1(ParticleID);
 #endif
-#if NUM_DYNAMIC_PARAMETERS > 2
+#if (DYNAMIC_PARAMETERS_MASK & 4)
 	Intermediates.DynamicParameter2 = GetNiagaraParticleDynamicParameters2(ParticleID);
 #endif
-#if NUM_DYNAMIC_PARAMETERS > 3
+#if (DYNAMIC_PARAMETERS_MASK & 8)
 	Intermediates.DynamicParameter3 = GetNiagaraParticleDynamicParameters3(ParticleID);
 #endif
 
@@ -939,16 +939,16 @@ FVertexFactoryInterpolantsVSToPS VertexFactoryGetInterpolantsVSToPS(FVertexFacto
 	Interpolants.Color = Intermediates.Color;
 #endif
 
-#if NUM_DYNAMIC_PARAMETERS > 0
+#if (DYNAMIC_PARAMETERS_MASK & 1)
 	Interpolants.DynamicParameter = Intermediates.DynamicParameter;
 #endif
-#if NUM_DYNAMIC_PARAMETERS > 1
+#if (DYNAMIC_PARAMETERS_MASK & 2)
 	Interpolants.DynamicParameter1 = Intermediates.DynamicParameter1;
 #endif
-#if NUM_DYNAMIC_PARAMETERS > 2
+#if (DYNAMIC_PARAMETERS_MASK & 4)
 	Interpolants.DynamicParameter2 = Intermediates.DynamicParameter2;
 #endif
-#if NUM_DYNAMIC_PARAMETERS > 3
+#if (DYNAMIC_PARAMETERS_MASK & 8)
 	Interpolants.DynamicParameter3 = Intermediates.DynamicParameter3;
 #endif
 
diff --git a/Engine/Source/Runtime/Engine/Private/Materials/HLSLMaterialTranslator.h b/Engine/Source/Runtime/Engine/Private/Materials/HLSLMaterialTranslator.h
--- a/Engine/Source/Runtime/Engine/Private/Materials/HLSLMaterialTranslator.h
+++ b/Engine/Source/Runtime/Engine/Private/Materials/HLSLMaterialTranslator.h
@@ -251,7 +251,7 @@ protected:
 	/** Tracks the number of texture coordinates used by the vertex shader in this material. */
 	uint32 NumUserVertexTexCoords;
 
-	uint32 NumParticleDynamicParameters;
+	uint32 DynamicParticleParameterMask;
 public: 
 
 	FHLSLMaterialTranslator(FMaterial* InMaterial,
@@ -303,7 +303,7 @@ public:
 	,	bIsFullyRough(0)
 	,	NumUserTexCoords(0)
 	,	NumUserVertexTexCoords(0)
-	,	NumParticleDynamicParameters(0)
+	,	DynamicParticleParameterMask(0)
 	{
 		FMemory::Memzero(SharedPixelProperties);
 
@@ -991,10 +991,10 @@ public:
 			OutEnvironment.SetDefine(TEXT("NEEDS_PARTICLE_VELOCITY"), 1);
 		}
 
-		if (NumParticleDynamicParameters > 0)
+		if (DynamicParticleParameterMask)
 		{
 			OutEnvironment.SetDefine(TEXT("USE_DYNAMIC_PARAMETERS"), 1);
-			OutEnvironment.SetDefine(TEXT("NUM_DYNAMIC_PARAMETERS"), NumParticleDynamicParameters);
+			OutEnvironment.SetDefine(TEXT("DYNAMIC_PARAMETERS_MASK"), DynamicParticleParameterMask);
 		}
 
 		if (bNeedsParticleTime)
@@ -4964,7 +4964,7 @@ protected:
 			return NonVertexOrPixelShaderExpressionError();
 		}
 
-		NumParticleDynamicParameters = FMath::Max(NumParticleDynamicParameters, ParameterIndex + 1);
+		DynamicParticleParameterMask |= (1 << ParameterIndex);
 
 		int32 Default = Constant4(DefaultValue.R, DefaultValue.G, DefaultValue.B, DefaultValue.A);
 		return AddInlinedCodeChunk(
