Fixed error when compiling some Niagara GPU particles

diff --git a/Engine/Shaders/Private/NiagaraEmitterInstanceShader.usf b/Engine/Shaders/Private/NiagaraEmitterInstanceShader.usf
--- a/Engine/Shaders/Private/NiagaraEmitterInstanceShader.usf
+++ b/Engine/Shaders/Private/NiagaraEmitterInstanceShader.usf
@@ -228,7 +228,7 @@ NiagaraSimulationShader.usf:
 		{
 			// Have to use uint's here to avoid PS4 compiler warnings about InterlockedAdd, cannot propagate uint due to CPU VM limitations...
 			uint RetPrevIdx;
-			InterlockedAdd((uint)RWDataSetIndices[DataSetID*5+1], (uint)1U, RetPrevIdx);
+			InterlockedAdd(RWDataSetIndices[DataSetID*5+1], (uint)1U, RetPrevIdx);
 			PrevIdx = (int)RetPrevIdx;
 		}
 		return PrevIdx;
