Changed GetForwardDynamicShadowFactors implementation to use Texture2DSampleLevel instead of Load.
This is because we sometime bind the GWhiteDummy when the light attenuation buffer is not required, which is a 1x1 texture.

diff --git a/Engine/Shaders/Private/ForwardLightingCommon.ush b/Engine/Shaders/Private/ForwardLightingCommon.ush
--- a/Engine/Shaders/Private/ForwardLightingCommon.ush
+++ b/Engine/Shaders/Private/ForwardLightingCommon.ush
@@ -69,8 +69,7 @@ float2 CalculateNearestResolvedDepthScreenUV(float2 ScreenUV, float SceneDepth)
 
 float4 GetForwardDynamicShadowFactors(float2 ScreenUV)
 {
-	int2 IntScreenUV = int2(trunc(ScreenUV * View_BufferSizeAndInvSize.xy));
-	float4 Value = OpaqueBasePass.ForwardScreenSpaceShadowMaskTexture.Load(int3(IntScreenUV, 0));
+	float4 Value = Texture2DSampleLevel(OpaqueBasePass.ForwardScreenSpaceShadowMaskTexture, GlobalPointClampedSampler, ScreenUV, 0);
 	return Square(Value);
 }
 
