Fix for movable light wth lightmap ignoring scene capture's average brightness.

diff --git a/Engine/Shaders/Private/MobileBasePassPixelShader.usf b/Engine/Shaders/Private/MobileBasePassPixelShader.usf
--- a/Engine/Shaders/Private/MobileBasePassPixelShader.usf
+++ b/Engine/Shaders/Private/MobileBasePassPixelShader.usf
@@ -269,8 +269,8 @@ half3 GetImageBasedReflectionLighting(FMaterialPixelParameters MaterialParameter
 	{
 		SpecularIBL = RGBMDecode(SpecularIBLSample, 16.0);
 		SpecularIBL = SpecularIBL * SpecularIBL;
-		// divide by average brightness for lightmap use cases only.
-		#if !MOVABLE_DIRECTIONAL_LIGHT
+		#if LQ_TEXTURE_LIGHTMAP || CACHED_POINT_INDIRECT_LIGHTING
+			// divide by average brightness for lightmap use cases only.
 			SpecularIBL *= InvReflectionCubemapAverageBrightness.x;
 		#endif
 	}
