Increases the bluring radius of lowest DOF quality.

diff --git a/Engine/Config/ConsoleVariables.ini b/Engine/Config/ConsoleVariables.ini
--- a/Engine/Config/ConsoleVariables.ini
+++ b/Engine/Config/ConsoleVariables.ini
@@ -33,7 +33,7 @@
 
 
 ; Uncomment to get detailed logs on shader compiles and the opportunity to retry on errors
-; r.ShaderDevelopmentMode=1
+r.ShaderDevelopmentMode=1
 ; Uncomment to dump shaders in the Saved folder
 ; Warning: leaving this on for a while will fill your hard drive with many small files and folders
 ;r.DumpShaderDebugInfo=1
diff --git a/Engine/Shaders/Private/DiaphragmDOF/DOFGatherCommon.ush b/Engine/Shaders/Private/DiaphragmDOF/DOFGatherCommon.ush
--- a/Engine/Shaders/Private/DiaphragmDOF/DOFGatherCommon.ush
+++ b/Engine/Shaders/Private/DiaphragmDOF/DOFGatherCommon.ush
@@ -122,6 +122,9 @@
 	#define CONFIG_FAST_ACCUMULATOR_KERNEL_DENSITY (KERNEL_DENSITY_CONSTANT_HEXAWEB)
 	//#define CONFIG_FAST_ACCUMULATOR_KERNEL_DENSITY (KERNEL_DENSITY_HEXAWEB_LOWER_IN_CENTER)
 
+	// Highlights needs a bit of help.
+	#define CONFIG_SLIGHT_RING_ROTATION 1
+
 #elif DIM_GATHER_QUALITY == GATHER_QUALITY_HIGH_QUALITY
 	#define CONFIG_WEIGHT_SAMPLE_FROM_COC 1
 	#define CONFIG_MIRRORED_SAMPLER 1
@@ -228,6 +231,13 @@
 	#define CONFIG_MAX_SAMPLING_DENSITY_CHANGES 3
 #endif
 
+/** Adds a slight rotation to the samples to reduce interferance sample location and texel grid with the point sampler,
+ * that improves highlight post filtering. */
+#ifndef CONFIG_SLIGHT_RING_ROTATION
+	#define CONFIG_SLIGHT_RING_ROTATION 0
+#endif
+
+
 //------------------------------------------------------- CONFIG CHECKS
 
 #ifndef GATHER_ACCUMULATOR_METHOD
diff --git a/Engine/Shaders/Private/DiaphragmDOF/DOFGatherKernel.ush b/Engine/Shaders/Private/DiaphragmDOF/DOFGatherKernel.ush
--- a/Engine/Shaders/Private/DiaphragmDOF/DOFGatherKernel.ush
+++ b/Engine/Shaders/Private/DiaphragmDOF/DOFGatherKernel.ush
@@ -21,16 +21,19 @@
 // Returns the position of the sample on the unit circle (radius = 1) for a given ring.
 float2 GetDiskSampleOnUnitCirle(uint RingId, uint RingSampleIteration, uint RingSampleId)
 {
+	float SampleRingPos = RingSampleId;
+
+	// Do not allign all j == 0 samples of the different ring on the X axis to increase minimal distance between all
+	// samples, that reduce variance to clean by post filtering.
 	#if 1
-		// Do not allign all j == 0 samples of the different ring on the X axus to increase minimal distance between all
-		// samples, that reduce variance to clean by post filtering.
-		float SampleAngle = PI * (RingSampleId + (RingId - 2 * (RingId / 2)) * 0.5) / float(RingSampleIteration);
-			
-	#else
-		// Alligns all j == 0 samples of the different ring on the X axus.
-		float SampleAngle = PI * RingSampleId / float(RingSampleIteration);
+		SampleRingPos += (RingId - 2 * (RingId / 2)) * 0.5;
+	#endif
 
+	#if CONFIG_SLIGHT_RING_ROTATION
+		SampleRingPos += (RingId + 1) * 0.2;
 	#endif
+
+	float SampleAngle = PI * SampleRingPos / float(RingSampleIteration);
 			
 	return float2(cos(SampleAngle), sin(SampleAngle));
 }
