Optimises DOF's CONFIG_UNROLL_GATHER==0 to be potentially used on other platforms

diff --git a/Engine/Shaders/Private/DiaphragmDOF/DOFGatherKernel.ush b/Engine/Shaders/Private/DiaphragmDOF/DOFGatherKernel.ush
--- a/Engine/Shaders/Private/DiaphragmDOF/DOFGatherKernel.ush
+++ b/Engine/Shaders/Private/DiaphragmDOF/DOFGatherKernel.ush
@@ -38,6 +38,20 @@ float2 GetDiskSampleOnUnitCirle(uint RingId, uint RingSampleIteration, uint Ring
 	return float2(cos(SampleAngle), sin(SampleAngle));
 }
 
+// Returns the rotation matrix to use between sample of the ring.
+float2x2 GetSampleRotationMatrix(uint RingSampleIteration)
+{
+	float RotationAngle = PI / float(RingSampleIteration);
+
+	float C = cos(RotationAngle);
+	float S = sin(RotationAngle);
+
+	return float2x2(
+		float2( C,  S),
+		float2(-S,  C));
+}
+
+
 struct FFastAccumulatorOutput
 {
 	// Accumulated Color.
@@ -351,6 +365,35 @@ void GatherRing(
 
 	// Number of batch to process.
 	const uint BatchCount = RingSamplePairCount / SamplePairBatchSize;
+	
+	#if CONFIG_SGPR_HINTS
+	{
+		RingInfos.Radius = ToScalarMemory(RingInfos.Radius);
+	}
+	#endif
+
+	// when the batch loop is not unroll, use successive rotation matrix to avoid doing cos & sin in inner loop.
+	#if !CONFIG_UNROLL_GATHER
+		// Generate at compile time sample rotation matrix.
+		const float2x2 SampleRotationMatrix = GetSampleRotationMatrix(RingSamplePairCount);
+
+		// Generates at compile time first sample location on circle (radius = 1).
+		const float2 FirstCircleUnitPos = GetDiskSampleOnUnitCirle(RingId, RingSamplePairCount, /* BatchId = */ 0);
+
+		// Position of the first sample on circle with radius according to KernelRadius.
+		float2 FirstCircleSamplePosOffset = FirstCircleUnitPos * RingInfos.Radius;
+
+		// Setup iteratable SGPR
+		float2 CurrentCircleUnitPos = FirstCircleUnitPos;
+		float2 CurrentCircleSamplePosOffset = FirstCircleSamplePosOffset;
+
+		#if CONFIG_SGPR_HINTS
+		{
+			CurrentCircleUnitPos = ToScalarMemory(CurrentCircleUnitPos);
+			CurrentCircleSamplePosOffset = ToScalarMemory(CurrentCircleSamplePosOffset);
+		}
+		#endif
+	#endif
 
 	// Iterate over the batch of samples.
 	#if CONFIG_UNROLL_GATHER
@@ -363,8 +406,28 @@ void GatherRing(
 		WORKARROUND_UE_58394 // TODO: nuke that guy once original issue is fixed.
 	#endif
 	{
-		// Generates at compile time sample location on circle (radius = 1).
-		const float2 BatchCirleUnitPos = GetDiskSampleOnUnitCirle(RingId, RingSamplePairCount, BatchId);
+		#if CONFIG_UNROLL_GATHER
+			// Generates at compile time sample location on circle (radius = 1).
+			const float2 BatchCirleUnitPos = GetDiskSampleOnUnitCirle(RingId, RingSamplePairCount, BatchId);
+
+			// Position of sample on circle with radius according to KernelRadius.
+			// TODO: ALU could be optimised since SampleCountToSamplePosition has some constant part as well.
+			float2 BatchCircleSamplePosOffset = BatchCirleUnitPos * RingInfos.Radius;
+		#else
+			// Rotate the samples position along the ring.
+			CurrentCircleUnitPos = mul(CurrentCircleUnitPos, SampleRotationMatrix);
+			CurrentCircleSamplePosOffset = mul(CurrentCircleSamplePosOffset, SampleRotationMatrix);
+
+			#if CONFIG_SGPR_HINTS
+			{
+				CurrentCircleUnitPos = ToScalarMemory(CurrentCircleUnitPos);
+				CurrentCircleSamplePosOffset = ToScalarMemory(CurrentCircleSamplePosOffset);
+			}
+			#endif
+
+			float2 BatchCirleUnitPos = CurrentCircleUnitPos;
+			float2 BatchCircleSamplePosOffset = CurrentCircleSamplePosOffset;
+		#endif
 
 		// Generates at compile time sample location if square shaped ring in sample count distance from the kernel's center.
 		//
@@ -399,19 +462,14 @@ void GatherRing(
 			BatchSquarePos.y = RingDistance - (BatchId - RingDistance * 3);
 		}
 
-		// Iterate over the pair of sample of the batch.
+		// Iterate over the pair of sample of the batch to share vgpr setup and rotating ALU
 		UNROLL
 		for (uint BatchSampleId = 0; BatchSampleId < SamplePairBatchSize; BatchSampleId++)
 		{
-			// Generate at compile time the circle unit pos and square pos from the batch
-			// to share const vgpr setup on all the batch's samples.
-			const float2 CirleUnitPos = SampleConstFromBatchConst(BatchSampleId, BatchCirleUnitPos);
+			// Swizzle the batch positions' X and Y components.
+			float2 CirleUnitPos = SampleConstFromBatchConst(BatchSampleId, BatchCirleUnitPos);
 			const float2 SquarePos = SampleConstFromBatchConst(BatchSampleId, BatchSquarePos);
-
-			// Position of sample on circle with radius according to KernelRadius.
-			// TODO: ALU could be optimised since SampleCountToSamplePosition has some constant part as well.
-			float2 CircleSamplePosOffset = (CirleUnitPos * RingDistance) * KernelScope.SampleCountToSamplePosition;
-
+			float2 CircleSamplePosOffset = SampleConstFromBatchConst(BatchSampleId, BatchCircleSamplePosOffset);
 
 			float2 SamplePosOffset[2];
 			float SampleDistance;
@@ -432,7 +490,7 @@ void GatherRing(
 						GlobalBilinearWrappedSampler, LookUpUV[1], 0).xy;
 				#endif
 				
-				SampleDistance = length(CirleUnitPos * RingDistance) * KernelScope.SampleCountToSamplePosition;
+				SampleDistance = RingInfos.Radius;
 			}
 			#elif 0 // TODO
 			{
@@ -484,7 +542,7 @@ void GatherRing(
 				// Always keep a disk shaped kernel.
 				SamplePosOffset[0] = CircleSamplePosOffset;
 				SamplePosOffset[1] = -CircleSamplePosOffset;
-				SampleDistance = length(CirleUnitPos * RingDistance) * KernelScope.SampleCountToSamplePosition;
+				SampleDistance = RingInfos.Radius;
 			}
 			#else
 				#error Unknown gathering kernel layout.
@@ -498,13 +556,7 @@ void GatherRing(
 			{
 				const float SampleSign = (k == 0) ? 1.0 : -1.0;
 
-				// Compute sample's euclidian distance.
-				// TODO: SampleDistanceSquare for perf reasons.
-				//float SampleDistance = length(RandomKernelPosOffset + SamplePosOffset[k]);
-				//float SampleDistance = length(SamplePosOffset[0]);
-
-				// Compute sample's input buffer UV
-				// TODO: VGPR could be optimised.
+				// Compute sample's input buffer UV.
 				float2 SampleUV = KernelScope.KernelCenterUV + SamplePosOffset[k] * KernelScope.GatherInputInvSize;
 
 				// Sample input buffer.
