UE4.21 - Force recompile
#rnx

diff --git a/Engine/Shaders/Private/DiaphragmDOF/DOFDownsample.ush b/Engine/Shaders/Private/DiaphragmDOF/DOFDownsample.ush
--- a/Engine/Shaders/Private/DiaphragmDOF/DOFDownsample.ush
+++ b/Engine/Shaders/Private/DiaphragmDOF/DOFDownsample.ush
@@ -1,6 +1,6 @@
 // Copyright 1998-2018 Epic Games, Inc. All Rights Reserved.
 
-/*=============================================================================
+/*============================================================================
 	DiaphragmDOF/DOFDownsample.ush: Diaphragm DOF's downsampling code/
 =============================================================================*/
 
