Scene static shadows on translucency are now sampled based on the material  "translucency.lighting mode" in forward shading, as it is in deferred shading.

diff --git a/Engine/Shaders/Private/ForwardLightingCommon.ush b/Engine/Shaders/Private/ForwardLightingCommon.ush
--- a/Engine/Shaders/Private/ForwardLightingCommon.ush
+++ b/Engine/Shaders/Private/ForwardLightingCommon.ush
@@ -146,7 +146,11 @@ float3 GetForwardDirectLighting(uint GridIndex, float3 WorldPosition, float3 Cam
 			float4 LightAttenuation = float4(WholeSceneShadowing.xx, PerObjectShadowing.xx);
 		#else
 			LightData.ShadowMapChannelMask.x = 1;
-			GBufferData.PrecomputedShadowFactors.x = ComputeDirectionalLightStaticShadowing(WorldPosition).x;
+			#if TRANSLUCENCY_LIGHTING_SURFACE_FORWARDSHADING
+				GBufferData.PrecomputedShadowFactors.x = ComputeDirectionalLightStaticShadowing(WorldPosition).x;
+			#else
+				GBufferData.PrecomputedShadowFactors.x = 1;
+			#endif
 			float DynamicShadowFactor = ComputeDirectionalLightDynamicShadowing(WorldPosition, GBufferData.Depth);
 			float4 LightAttenuation = float4(DynamicShadowFactor.x, DynamicShadowFactor.x, 1, 1);
 		#endif
