Vulkan: fix black artifacts in Volumetric Fog (UE-66144).

diff --git a/Engine/Shaders/Private/VolumetricFog.usf b/Engine/Shaders/Private/VolumetricFog.usf
--- a/Engine/Shaders/Private/VolumetricFog.usf
+++ b/Engine/Shaders/Private/VolumetricFog.usf
@@ -636,7 +636,7 @@ void LightScatteringCS(
 
 	if (all(GridCoordinate < VolumetricFog.GridSizeInt))
 	{
-		ScatteringAndExtinction = isnan(ScatteringAndExtinction) || isinf(ScatteringAndExtinction) ? 0 : ScatteringAndExtinction;
+		ScatteringAndExtinction = -min(-ScatteringAndExtinction, float4(0.f, 0.f, 0.f, 0.f));	// vector isnan/isinf are miscompiled by hlslcc (see UE-66179)
 		RWLightScattering[GridCoordinate] = ScatteringAndExtinction;
 	}
 }
@@ -681,4 +681,4 @@ void FinalIntegrationCS(
 
 		RWIntegratedLightScattering[LayerCoordinate] = float4(AccumulatedLighting, AccumulatedTransmittance);
 	}
-}
\ No newline at end of file
+}
