Fix rect light with clear coat

diff --git a/Engine/Shaders/Private/ShadingModels.ush b/Engine/Shaders/Private/ShadingModels.ush
--- a/Engine/Shaders/Private/ShadingModels.ush
+++ b/Engine/Shaders/Private/ShadingModels.ush
@@ -201,7 +201,7 @@ FDirectLighting ClearCoatBxDF( FGBufferData GBuffer, half3 N, half3 V, half3 L,
 
 	if( AreaLight.bIsRect )
 	{
-		Lighting.Specular = RectGGXApproxLTC( ClearCoatRoughness, F0, N, V, AreaLight.Rect, AreaLight.Texture );
+		Lighting.Specular = ClearCoat * RectGGXApproxLTC( ClearCoatRoughness, F0, N, V, AreaLight.Rect, AreaLight.Texture );
 	}
 	else
 	{
