Adding some parameters on the ExponentialHeightFogComponent to drive a second height fog

diff --git a/Engine/Shaders/Private/HeightFogCommon.ush b/Engine/Shaders/Private/HeightFogCommon.ush
--- a/Engine/Shaders/Private/HeightFogCommon.ush
+++ b/Engine/Shaders/Private/HeightFogCommon.ush
@@ -20,6 +20,7 @@ static const float FLT_EPSILON = 0.001f;
 static const float FLT_EPSILON2 = 0.01f;
 
 // FogStruct.ExponentialFogParameters: FogDensity * exp2(-FogHeightFalloff * (CameraWorldPosition.z - FogHeight)) in x, FogHeightFalloff in y, CosTerminatorAngle in z, StartDistance in w. 
+// FogStruct.ExponentialFogParameters2: FogDensitySecond * exp2(-FogHeightFalloffSecond * (CameraWorldPosition.z - FogHeightSecond)) in x, FogHeightFalloffSecond in y, FogDensitySecond in z, FogHeightSecond in w 
 // FogStruct.ExponentialFogParameters3: FogDensity in x, FogHeight in y, whether to use cubemap fog color in z, FogCutoffDistance in w. 
 // FogStruct.FogInscatteringTextureParameters: mip distance scale in x, bias in y, num mips in z 
 
@@ -45,6 +46,17 @@ float3 ComputeInscatteringColor(float3 CameraToReceiver, float CameraToReceiverL
 	return Inscattering;
 }
 
+// Calculate the line integral of the ray from the camera to the receiver position through the fog density function
+// The exponential fog density function is d = GlobalDensity * exp(-HeightFalloff * z)
+float CalculateLineIntegralShared(float FogHeightFalloff, float RayDirectionZ, float RayOriginTerms)
+{
+	float Falloff = max(-127.0f, FogHeightFalloff * RayDirectionZ);    // if it's lower than -127.0, then exp2() goes crazy in OpenGL's GLSL.
+	float LineIntegral = ( 1.0f - exp2(-Falloff) ) / Falloff;
+	float LineIntegralTaylor = log(2.0) - ( 0.5 * Pow2( log(2.0) ) ) * Falloff;		// Taylor expansion around 0
+	
+	return RayOriginTerms * ( abs(Falloff) > FLT_EPSILON2 ? LineIntegral : LineIntegralTaylor );
+}
+
 // @param WorldPositionRelativeToCamera = WorldPosition - InCameraPosition
 half4 GetExponentialHeightFog(float3 WorldPositionRelativeToCamera, float ExcludeDistance)
 {
@@ -57,6 +69,7 @@ half4 GetExponentialHeightFog(float3 WorldPositionRelativeToCamera, float Exclud
 	half3 CameraToReceiverNormalized = CameraToReceiver * CameraToReceiverLengthInv;
 
 	float RayOriginTerms = FogStruct.ExponentialFogParameters.x;
+	float RayOriginTermsSecond = FogStruct.ExponentialFogParameters2.x;
 	float RayLength = CameraToReceiverLength;
 	float RayDirectionZ = CameraToReceiver.z;
 
@@ -94,14 +107,14 @@ half4 GetExponentialHeightFog(float3 WorldPositionRelativeToCamera, float Exclud
 
 		float Exponent = max(-127.0f, FogStruct.ExponentialFogParameters.y * (ExclusionIntersectionZ - FogStruct.ExponentialFogParameters3.y));
 		RayOriginTerms = FogStruct.ExponentialFogParameters3.x * exp2(-Exponent);
+		
+		float ExponentSecond = max(-127.0f, FogStruct.ExponentialFogParameters2.y * (ExclusionIntersectionZ - FogStruct.ExponentialFogParameters2.w)); 	 
+		RayOriginTermsSecond = FogStruct.ExponentialFogParameters2.z * exp2(-ExponentSecond);
 	}
 
-	// Calculate the line integral of the ray from the camera to the receiver position through the fog density function
-	// The exponential fog density function is d = GlobalDensity * exp(-HeightFalloff * z)
-	float Falloff = max(-127.0f, FogStruct.ExponentialFogParameters.y * RayDirectionZ);    // if it's lower than -127.0, then exp2() goes crazy in OpenGL's GLSL.
-	float LineIntegral = ( 1.0f - exp2(-Falloff) ) / Falloff;
-	float LineIntegralTaylor = log(2.0) - ( 0.5 * Pow2( log(2.0) ) ) * Falloff;		// Taylor expansion around 0
-	float ExponentialHeightLineIntegralShared = RayOriginTerms * ( abs(Falloff) > FLT_EPSILON2 ? LineIntegral : LineIntegralTaylor );
+	// Calculate the "shared" line integral (this term is also used for the directional light inscattering) by adding the two line integrals together (from two different height falloffs and densities)
+	float ExponentialHeightLineIntegralShared = CalculateLineIntegralShared(FogStruct.ExponentialFogParameters.y, RayDirectionZ, RayOriginTerms) + CalculateLineIntegralShared(FogStruct.ExponentialFogParameters2.y, RayDirectionZ, RayOriginTermsSecond);
+
 	float ExponentialHeightLineIntegral = ExponentialHeightLineIntegralShared * RayLength;
 
 	half3 InscatteringColor = ComputeInscatteringColor(CameraToReceiver, CameraToReceiverLength);
diff --git a/Engine/Source/Runtime/Engine/Classes/Components/ExponentialHeightFogComponent.h b/Engine/Source/Runtime/Engine/Classes/Components/ExponentialHeightFogComponent.h
--- a/Engine/Source/Runtime/Engine/Classes/Components/ExponentialHeightFogComponent.h
+++ b/Engine/Source/Runtime/Engine/Classes/Components/ExponentialHeightFogComponent.h
@@ -8,6 +8,38 @@
 #include "Components/SceneComponent.h"
 #include "ExponentialHeightFogComponent.generated.h"
 
+/**
+*	Data for an individual fog line integral.
+*	This is the data which is not shared between fogs when multiple fogs are set up on a single UExponentialHeightFogComponent
+*/
+USTRUCT(BlueprintType)
+struct FExponentialHeightFogData
+{
+	GENERATED_USTRUCT_BODY()
+
+	/** Global density factor for this fog. */
+	UPROPERTY(EditAnywhere, BlueprintReadOnly, interp, Category=ExponentialHeightFogComponent, meta = (UIMin = "0", UIMax = ".05"))
+	float FogDensity = 0.02f;
+
+	/**
+	* Height density factor, controls how the density increases as height decreases.
+	* Smaller values make the visible transition larger.
+	*/
+	UPROPERTY(EditAnywhere, BlueprintReadOnly, interp, Category=ExponentialHeightFogComponent, meta = (UIMin = "0.001", UIMax = "2"))
+	float FogHeightFalloff = 0.2f;
+
+	/** Height offset, relative to the actor position Z. */
+	UPROPERTY(EditAnywhere, BlueprintReadOnly, interp, Category=ExponentialHeightFogComponent)
+	float FogHeightOffset = 0.0f;
+
+	/** Clamp to valid ranges. This might be different from the UI clamp. */
+	void ClampToValidRanges()
+	{
+		FogDensity = FMath::Clamp(FogDensity, 0.0f, 10.0f);
+		FogHeightFalloff = FMath::Clamp(FogHeightFalloff, 0.0f, 2.0f);
+	}
+};
+
 /**
  *	Used to create fogging effects such as clouds but with a density that is related to the height of the fog.
  */
@@ -20,6 +52,17 @@ class ENGINE_API UExponentialHeightFogComponent : public USceneComponent
 	UPROPERTY(BlueprintReadOnly, interp, Category=ExponentialHeightFogComponent, meta=(UIMin = "0", UIMax = ".05"))
 	float FogDensity;
 
+	/**
+	 * Height density factor, controls how the density increases as height decreases.
+	 * Smaller values make the visible transition larger.
+	 */
+	UPROPERTY(BlueprintReadOnly, interp, Category = ExponentialHeightFogComponent, meta = (UIMin = "0.001", UIMax = "2"))
+	float FogHeightFalloff;
+
+	/** Settings for the second fog. Setting the density of this to 0 means it doesn't have any influence. */
+	UPROPERTY(EditAnywhere, BlueprintReadOnly, Category=ExponentialHeightFogComponent)
+	FExponentialHeightFogData SecondFogData;
+
 	UPROPERTY(BlueprintReadOnly, interp, Category=ExponentialHeightFogComponent)
 	FLinearColor FogInscatteringColor;
 
@@ -67,13 +110,6 @@ class ENGINE_API UExponentialHeightFogComponent : public USceneComponent
 	UPROPERTY(BlueprintReadOnly, interp, Category=DirectionalInscattering)
 	FLinearColor DirectionalInscatteringColor;
 
-	/** 
-	 * Height density factor, controls how the density increases as height decreases.  
-	 * Smaller values make the visible transition larger.
-	 */
-	UPROPERTY(BlueprintReadOnly, interp, Category=ExponentialHeightFogComponent, meta=(UIMin = "0.001", UIMax = "2"))
-	float FogHeightFalloff;
-
 	/** 
 	 * Maximum opacity of the fog.  
 	 * A value of 1 means the fog can become fully opaque at a distance and replace scene color completely,
@@ -175,7 +211,7 @@ public:
 	
 	UFUNCTION(BlueprintCallable, Category="Rendering|Components|ExponentialHeightFog")
 	void SetFogHeightFalloff(float Value);
-	
+
 	UFUNCTION(BlueprintCallable, Category="Rendering|Components|ExponentialHeightFog")
 	void SetFogMaxOpacity(float Value);
 	
diff --git a/Engine/Source/Runtime/Engine/Private/Components/HeightFogComponent.cpp b/Engine/Source/Runtime/Engine/Private/Components/HeightFogComponent.cpp
--- a/Engine/Source/Runtime/Engine/Private/Components/HeightFogComponent.cpp
+++ b/Engine/Source/Runtime/Engine/Private/Components/HeightFogComponent.cpp
@@ -31,6 +31,9 @@ UExponentialHeightFogComponent::UExponentialHeightFogComponent(const FObjectInit
 
 	FogDensity = 0.02f;
 	FogHeightFalloff = 0.2f;
+	// No influence from the second fog as default
+	SecondFogData.FogDensity = 0.0f;
+
 	FogMaxOpacity = 1.0f;
 	StartDistance = 0.0f;
 
@@ -102,6 +105,7 @@ bool UExponentialHeightFogComponent::CanEditChange(const UProperty* InProperty)
 
 void UExponentialHeightFogComponent::PostEditChangeProperty(FPropertyChangedEvent& PropertyChangedEvent)
 {
+	SecondFogData.ClampToValidRanges();
 	FogDensity = FMath::Clamp(FogDensity, 0.0f, 10.0f);
 	FogHeightFalloff = FMath::Clamp(FogHeightFalloff, 0.0f, 2.0f);
 	FogMaxOpacity = FMath::Clamp(FogMaxOpacity, 0.0f, 1.0f);
diff --git a/Engine/Source/Runtime/Renderer/Private/FogRendering.cpp b/Engine/Source/Runtime/Renderer/Private/FogRendering.cpp
--- a/Engine/Source/Runtime/Renderer/Private/FogRendering.cpp
+++ b/Engine/Source/Runtime/Renderer/Private/FogRendering.cpp
@@ -56,6 +56,7 @@ void SetupFogUniformParameters(const FViewInfo& View, FFogUniformParameters& Out
 
 		OutParameters.ExponentialFogParameters = View.ExponentialFogParameters;
 		OutParameters.ExponentialFogColorParameter = FVector4(View.ExponentialFogColor, 1.0f - View.FogMaxOpacity);
+		OutParameters.ExponentialFogParameters2 = View.ExponentialFogParameters2;
 		OutParameters.ExponentialFogParameters3 = View.ExponentialFogParameters3;
 		OutParameters.SinCosInscatteringColorCubemapRotation = View.SinCosInscatteringColorCubemapRotation;
 		OutParameters.FogInscatteringTextureParameters = View.FogInscatteringTextureParameters;
@@ -277,16 +278,23 @@ void FSceneRenderer::InitFogConstants()
 			{
 				const FExponentialHeightFogSceneInfo& FogInfo = Scene->ExponentialFogs[0];
 				const float CosTerminatorAngle = FMath::Clamp(FMath::Cos(FogInfo.LightTerminatorAngle * PI / 180.0f), -1.0f + DELTA, 1.0f - DELTA);
-				const float CollapsedFogParameterPower = FMath::Clamp(
-						-FogInfo.FogHeightFalloff * (View.ViewMatrices.GetViewOrigin().Z - FogInfo.FogHeight),
+				float CollapsedFogParameter[2];
+				for (int i = 0; i < FExponentialHeightFogSceneInfo::NumFogs; i++)
+				{
+					const float CollapsedFogParameterPower = FMath::Clamp(
+						-FogInfo.FogData[i].HeightFalloff * (View.ViewMatrices.GetViewOrigin().Z - FogInfo.FogData[i].Height),
 						-126.f + 1.f, // min and max exponent values for IEEE floating points (http://en.wikipedia.org/wiki/IEEE_floating_point)
 						+127.f - 1.f
-						);
-				const float CollapsedFogParameter = FogInfo.FogDensity * FMath::Pow(2.0f, CollapsedFogParameterPower);
-				View.ExponentialFogParameters = FVector4(CollapsedFogParameter, FogInfo.FogHeightFalloff, CosTerminatorAngle, FogInfo.StartDistance);
+					);
+
+					CollapsedFogParameter[i] = FogInfo.FogData[i].Density * FMath::Pow(2.0f, CollapsedFogParameterPower);
+
+				}
+				View.ExponentialFogParameters = FVector4(CollapsedFogParameter[0], FogInfo.FogData[0].HeightFalloff, CosTerminatorAngle, FogInfo.StartDistance);
+				View.ExponentialFogParameters2 = FVector4(CollapsedFogParameter[1], FogInfo.FogData[1].HeightFalloff, FogInfo.FogData[1].Density, FogInfo.FogData[1].Height);
 				View.ExponentialFogColor = FVector(FogInfo.FogColor.R, FogInfo.FogColor.G, FogInfo.FogColor.B);
 				View.FogMaxOpacity = FogInfo.FogMaxOpacity;
-				View.ExponentialFogParameters3 = FVector4(FogInfo.FogDensity, FogInfo.FogHeight, FogInfo.InscatteringColorCubemap ? 1.0f : 0.0f, FogInfo.FogCutoffDistance);
+				View.ExponentialFogParameters3 = FVector4(FogInfo.FogData[0].Density, FogInfo.FogData[0].Height, FogInfo.InscatteringColorCubemap ? 1.0f : 0.0f, FogInfo.FogCutoffDistance);
 				View.SinCosInscatteringColorCubemapRotation = FVector2D(FMath::Sin(FogInfo.InscatteringColorCubemapAngle), FMath::Cos(FogInfo.InscatteringColorCubemapAngle));
 				View.FogInscatteringColorCubemap = FogInfo.InscatteringColorCubemap;
 				const float InvRange = 1.0f / FMath::Max(FogInfo.FullyDirectionalInscatteringColorDistance - FogInfo.NonDirectionalInscatteringColorDistance, .00001f);
diff --git a/Engine/Source/Runtime/Renderer/Private/FogRendering.h b/Engine/Source/Runtime/Renderer/Private/FogRendering.h
--- a/Engine/Source/Runtime/Renderer/Private/FogRendering.h
+++ b/Engine/Source/Runtime/Renderer/Private/FogRendering.h
@@ -1,4 +1,4 @@
-// Copyright 1998-2019 Epic Games, Inc. All Rights Reserved.
+// Copyright 1998-2018 Epic Games, Inc. All Rights Reserved.
 
 /*=============================================================================
 	FogRendering.h: 
@@ -13,6 +13,7 @@
 
 BEGIN_GLOBAL_SHADER_PARAMETER_STRUCT(FFogUniformParameters,)
 	SHADER_PARAMETER(FVector4, ExponentialFogParameters)
+	SHADER_PARAMETER(FVector4, ExponentialFogParameters2)
 	SHADER_PARAMETER(FVector4, ExponentialFogColorParameter)
 	SHADER_PARAMETER(FVector4, ExponentialFogParameters3)
 	SHADER_PARAMETER(FVector4, InscatteringLightDirection) // non negative DirectionalInscatteringStartDistance in .W
diff --git a/Engine/Source/Runtime/Renderer/Private/RendererScene.cpp b/Engine/Source/Runtime/Renderer/Private/RendererScene.cpp
--- a/Engine/Source/Runtime/Renderer/Private/RendererScene.cpp
+++ b/Engine/Source/Runtime/Renderer/Private/RendererScene.cpp
@@ -3076,7 +3076,10 @@ void FScene::ApplyWorldOffset_RenderThread(FVector InOffset)
 	// Exponential Fog
 	for (FExponentialHeightFogSceneInfo& FogInfo : ExponentialFogs)
 	{
-		FogInfo.FogHeight+= InOffset.Z;
+		for (FExponentialHeightFogSceneInfo::FExponentialHeightFogSceneData& FogData : FogInfo.FogData)
+		{
+			FogData.Height += InOffset.Z;
+		}
 	}
 	
 	// StaticMeshDrawLists
diff --git a/Engine/Source/Runtime/Renderer/Private/SceneCore.cpp b/Engine/Source/Runtime/Renderer/Private/SceneCore.cpp
--- a/Engine/Source/Runtime/Renderer/Private/SceneCore.cpp
+++ b/Engine/Source/Runtime/Renderer/Private/SceneCore.cpp
@@ -440,11 +440,6 @@ FStaticMesh::~FStaticMesh()
 /** Initialization constructor. */
 FExponentialHeightFogSceneInfo::FExponentialHeightFogSceneInfo(const UExponentialHeightFogComponent* InComponent):
 	Component(InComponent),
-	FogHeight(InComponent->GetComponentLocation().Z),
-	// Scale the densities back down to their real scale
-	// Artists edit the densities scaled up so they aren't entering in minuscule floating point numbers
-	FogDensity(InComponent->FogDensity / 1000.0f),
-	FogHeightFalloff(InComponent->FogHeightFalloff / 1000.0f),
 	FogMaxOpacity(InComponent->FogMaxOpacity),
 	StartDistance(InComponent->StartDistance),
 	FogCutoffDistance(InComponent->FogCutoffDistance),
@@ -453,6 +448,16 @@ FExponentialHeightFogSceneInfo::FExponentialHeightFogSceneInfo(const UExponentia
 	DirectionalInscatteringStartDistance(InComponent->DirectionalInscatteringStartDistance),
 	DirectionalInscatteringColor(InComponent->DirectionalInscatteringColor)
 {
+	FogData[0].Height = InComponent->GetComponentLocation().Z;
+	FogData[1].Height = InComponent->GetComponentLocation().Z + InComponent->SecondFogData.FogHeightOffset;
+
+	// Scale the densities back down to their real scale
+	// Artists edit the densities scaled up so they aren't entering in minuscule floating point numbers
+	FogData[0].Density = InComponent->FogDensity / 1000.0f;	
+	FogData[0].HeightFalloff = InComponent->FogHeightFalloff / 1000.0f;
+	FogData[1].Density = InComponent->SecondFogData.FogDensity / 1000.0f;
+	FogData[1].HeightFalloff = InComponent->SecondFogData.FogHeightFalloff / 1000.0f;
+
 	FogColor = InComponent->InscatteringColorCubemap ? InComponent->InscatteringTextureTint : InComponent->FogInscatteringColor;
 	InscatteringColorCubemap = InComponent->InscatteringColorCubemap;
 	InscatteringColorCubemapAngle = InComponent->InscatteringColorCubemapAngle * (PI / 180.f);
diff --git a/Engine/Source/Runtime/Renderer/Private/SceneCore.h b/Engine/Source/Runtime/Renderer/Private/SceneCore.h
--- a/Engine/Source/Runtime/Renderer/Private/SceneCore.h
+++ b/Engine/Source/Runtime/Renderer/Private/SceneCore.h
@@ -202,17 +202,25 @@ class FExponentialHeightFogSceneInfo
 {
 public:
 
+	struct FExponentialHeightFogSceneData
+	{
+		float Height;
+		float Density;
+		float HeightFalloff;
+	};
+
+	/** Number of supported individual fog settings on this ExponentialHeightFog */
+	static constexpr int NumFogs = 2;
+
 	/** The fog component the scene info is for. */
 	const UExponentialHeightFogComponent* Component;
-	float FogHeight;
-	float FogDensity;
-	float FogHeightFalloff;
+	FExponentialHeightFogSceneData FogData[NumFogs];
 	float FogMaxOpacity;
-	float StartDistance;
+	float StartDistance; 
 	float FogCutoffDistance;
 	float LightTerminatorAngle;
 	FLinearColor FogColor;
-	float DirectionalInscatteringExponent;
+	float DirectionalInscatteringExponent; 
 	float DirectionalInscatteringStartDistance;
 	FLinearColor DirectionalInscatteringColor;
 	UTextureCube* InscatteringColorCubemap;
diff --git a/Engine/Source/Runtime/Renderer/Private/SceneRendering.cpp b/Engine/Source/Runtime/Renderer/Private/SceneRendering.cpp
--- a/Engine/Source/Runtime/Renderer/Private/SceneRendering.cpp
+++ b/Engine/Source/Runtime/Renderer/Private/SceneRendering.cpp
@@ -785,6 +785,7 @@ void FViewInfo::Init()
 	bUsesSceneDepth = false;
 
 	ExponentialFogParameters = FVector4(0,1,1,0);
+	ExponentialFogParameters2 = FVector4(0, 1, 0, 0);
 	ExponentialFogColor = FVector::ZeroVector;
 	FogMaxOpacity = 1;
 	ExponentialFogParameters3 = FVector4(0, 0, 0, 0);
diff --git a/Engine/Source/Runtime/Renderer/Private/SceneRendering.h b/Engine/Source/Runtime/Renderer/Private/SceneRendering.h
--- a/Engine/Source/Runtime/Renderer/Private/SceneRendering.h
+++ b/Engine/Source/Runtime/Renderer/Private/SceneRendering.h
@@ -991,6 +991,7 @@ public:
 
 	/** Parameters for exponential height fog. */
 	FVector4 ExponentialFogParameters;
+	FVector4 ExponentialFogParameters2;
 	FVector ExponentialFogColor;
 	float FogMaxOpacity;
 	FVector4 ExponentialFogParameters3;
