Fix for second height fog density not influencing the volumetric fog

diff --git a/Engine/Shaders/Private/VolumetricFog.usf b/Engine/Shaders/Private/VolumetricFog.usf
--- a/Engine/Shaders/Private/VolumetricFog.usf
+++ b/Engine/Shaders/Private/VolumetricFog.usf
@@ -83,7 +83,9 @@ void MaterialSetupCS(
 	 
 	float3 WorldPosition = ComputeCellWorldPosition(GridCoordinate, VoxelOffset);
 
-	float GlobalDensity = FogStruct.ExponentialFogParameters3.x * exp2(-FogStruct.ExponentialFogParameters.y * (WorldPosition.z - FogStruct.ExponentialFogParameters3.y));
+	float GlobalDensityFirst = FogStruct.ExponentialFogParameters3.x * exp2(-FogStruct.ExponentialFogParameters.y * (WorldPosition.z - FogStruct.ExponentialFogParameters3.y));
+	float GlobalDensitySecond = FogStruct.ExponentialFogParameters2.z * exp2(-FogStruct.ExponentialFogParameters2.y * (WorldPosition.z - FogStruct.ExponentialFogParameters2.w));
+	float GlobalDensity = GlobalDensityFirst + GlobalDensitySecond;
 
 	float3 Albedo = GlobalAlbedo;
 
