[NO NON-PII COMMIT MESSAGE]

diff --git a/Engine/Shaders/Private/LandscapeProceduralPS.usf b/Engine/Shaders/Private/LandscapeProceduralPS.usf
--- a/Engine/Shaders/Private/LandscapeProceduralPS.usf
+++ b/Engine/Shaders/Private/LandscapeProceduralPS.usf
@@ -4,7 +4,8 @@
 
 Texture2D ReadHeightmapTexture1;
 Texture2D ReadHeightmapTexture2;
-SamplerState ReadHeightmapTextureSampler;
+SamplerState ReadHeightmapTexture1Sampler;
+SamplerState ReadHeightmapTexture2Sampler;
 float2 LayerInfo; // x == weight, y == visibility
 float4 OutputConfig; // x == ApplyLayerModifiers, y == OutputAsDelta, z == Use ReadHeightmapTexture2 and out Delta with ReadHeightmapTexture1, w == Output Normals
 float2 HeightmapTextureSize; // x == source heightmap width, y == source heightmap height
@@ -15,7 +16,7 @@ float3 SampleTextureYNormal(float2 InTextureCoordinates, float2 InTexelSize, boo
 {
 	float2 SamplingUV = InTextureCoordinates;
 	SamplingUV.y = InUp ? SamplingUV.y - InTexelSize.y : SamplingUV.y + InTexelSize.y;
-	float2 PosSample = ReadHeightmapTexture1.SampleLevel(ReadHeightmapTextureSampler, SamplingUV, 0).rg;
+	float2 PosSample = ReadHeightmapTexture1.SampleLevel(ReadHeightmapTexture1Sampler, SamplingUV, 0).rg;
 	float PosHeight = float(((int)round(PosSample.r * 255.0) << 8) | (int)round(PosSample.g * 255.0)) - 32768.0;
 	return float3(SamplingUV * HeightmapTextureSize * LandscapeGridScale.xy, PosHeight * LandscapeGridScale.z * TERRAIN_ZSCALE);
 }
@@ -24,7 +25,7 @@ float3 SampleTextureXNormal(float2 InTextureCoordinates, float2 InTexelSize, boo
 {
 	float2 SamplingUV = InTextureCoordinates;
 	SamplingUV.x = InLeft ? SamplingUV.x - InTexelSize.x : SamplingUV.x + InTexelSize.x;
-	float2 PosSample = ReadHeightmapTexture1.SampleLevel(ReadHeightmapTextureSampler, SamplingUV, 0).rg;
+	float2 PosSample = ReadHeightmapTexture1.SampleLevel(ReadHeightmapTexture1Sampler, SamplingUV, 0).rg;
 	float PosHeight = float(((int)round(PosSample.r * 255.0) << 8) | (int)round(PosSample.g * 255.0)) - 32768.0;
 	return float3(SamplingUV * HeightmapTextureSize * LandscapeGridScale.xy, PosHeight * LandscapeGridScale.z * TERRAIN_ZSCALE);
 }
@@ -32,7 +33,7 @@ float3 SampleTextureXNormal(float2 InTextureCoordinates, float2 InTexelSize, boo
 void PSMain(float2 InTextureCoordinates : TEXCOORD0,
 			out float4 OutColor : SV_Target0)
 {
-	float4 SourceColor = ReadHeightmapTexture1.Sample(ReadHeightmapTextureSampler, InTextureCoordinates);
+	float4 SourceColor = ReadHeightmapTexture1.Sample(ReadHeightmapTexture1Sampler, InTextureCoordinates);
 	float Height = float(((int)round(SourceColor.r * 255.0) << 8) | (int)round(SourceColor.g * 255.0));
 
 	// Perform calculation 0 based
@@ -47,7 +48,7 @@ void PSMain(float2 InTextureCoordinates : TEXCOORD0,
 	// Output using 2nd heightmap
 	if (OutputConfig.z == 1.0)
 	{
-		float2 BaseColor = ReadHeightmapTexture2.Sample(ReadHeightmapTextureSampler, InTextureCoordinates).rg;
+		float2 BaseColor = ReadHeightmapTexture2.Sample(ReadHeightmapTexture2Sampler, InTextureCoordinates).rg;
 		float BaseHeight = float(((int)round(BaseColor.r * 255.0) << 8) | (int)round(BaseColor.g * 255.0));
 
 		Height = OutputConfig.y == 1 ? BaseHeight + (Height - BaseHeight) : BaseHeight + Height;
@@ -224,8 +225,8 @@ void PSMainMips(float2 InTextureCoordinates : TEXCOORD0,
 
 	if (IsMinBorderTexelX) // on left border
 	{
-		float4 SourceSample = ReadHeightmapTexture1.SampleLevel(ReadHeightmapTextureSampler, SourceUV, 0);
-		float4 DownSample = ReadHeightmapTexture1.SampleLevel(ReadHeightmapTextureSampler, SourceUV + float2(0.0, ParentMipTexelSize.y), 0);
+		float4 SourceSample = ReadHeightmapTexture1.SampleLevel(ReadHeightmapTexture1Sampler, SourceUV, 0);
+		float4 DownSample = ReadHeightmapTexture1.SampleLevel(ReadHeightmapTexture1Sampler, SourceUV + float2(0.0, ParentMipTexelSize.y), 0);
 		float SourceHeight = clamp(float(((int)round(SourceSample.r * 255.0) << 8) | (int)round(SourceSample.g * 255.0)), 0.0, 65536.0);
 		float DownHeight = clamp(float(((int)round(DownSample.r * 255.0) << 8) | (int)round(DownSample.g * 255.0)), 0.0, 65536.0);
 
@@ -249,8 +250,8 @@ void PSMainMips(float2 InTextureCoordinates : TEXCOORD0,
 	}
 	else if (IsMaxBorderTexelX) // on right border
 	{
-		float4 RightSample = ReadHeightmapTexture1.SampleLevel(ReadHeightmapTextureSampler, SourceUV + float2(ParentMipTexelSize.x, 0.0), 0);
-		float4 DownRightSample = ReadHeightmapTexture1.SampleLevel(ReadHeightmapTextureSampler, SourceUV + ParentMipTexelSize, 0);
+		float4 RightSample = ReadHeightmapTexture1.SampleLevel(ReadHeightmapTexture1Sampler, SourceUV + float2(ParentMipTexelSize.x, 0.0), 0);
+		float4 DownRightSample = ReadHeightmapTexture1.SampleLevel(ReadHeightmapTexture1Sampler, SourceUV + ParentMipTexelSize, 0);
 		float RightHeight = clamp(float(((int)round(RightSample.r * 255.0) << 8) | (int)round(RightSample.g * 255.0)), 0.0, 65536.0);
 		float DownRightHeight = clamp(float(((int)round(DownRightSample.r * 255.0) << 8) | (int)round(DownRightSample.g * 255.0)), 0.0, 65536.0);
 
@@ -274,10 +275,10 @@ void PSMainMips(float2 InTextureCoordinates : TEXCOORD0,
 	}
 	else // center texel
 	{
-		float4 SourceSample = ReadHeightmapTexture1.SampleLevel(ReadHeightmapTextureSampler, SourceUV, 0);
-		float4 DownSample = ReadHeightmapTexture1.SampleLevel(ReadHeightmapTextureSampler, SourceUV + float2(0.0, ParentMipTexelSize.y), 0);
-		float4 RightSample = ReadHeightmapTexture1.SampleLevel(ReadHeightmapTextureSampler, SourceUV + float2(ParentMipTexelSize.x, 0.0), 0);
-		float4 DownRightSample = ReadHeightmapTexture1.SampleLevel(ReadHeightmapTextureSampler, SourceUV + ParentMipTexelSize, 0);
+		float4 SourceSample = ReadHeightmapTexture1.SampleLevel(ReadHeightmapTexture1Sampler, SourceUV, 0);
+		float4 DownSample = ReadHeightmapTexture1.SampleLevel(ReadHeightmapTexture1Sampler, SourceUV + float2(0.0, ParentMipTexelSize.y), 0);
+		float4 RightSample = ReadHeightmapTexture1.SampleLevel(ReadHeightmapTexture1Sampler, SourceUV + float2(ParentMipTexelSize.x, 0.0), 0);
+		float4 DownRightSample = ReadHeightmapTexture1.SampleLevel(ReadHeightmapTexture1Sampler, SourceUV + ParentMipTexelSize, 0);
 
 		float SourceHeight = clamp(float(((int)round(SourceSample.r * 255.0) << 8) | (int)round(SourceSample.g * 255.0)), 0.0, 65536.0);
 		float DownHeight = clamp(float(((int)round(DownSample.r * 255.0) << 8) | (int)round(DownSample.g * 255.0)), 0.0, 65536.0);
diff --git a/Engine/Source/Runtime/Landscape/Private/LandscapeEditProcedural.cpp b/Engine/Source/Runtime/Landscape/Private/LandscapeEditProcedural.cpp
--- a/Engine/Source/Runtime/Landscape/Private/LandscapeEditProcedural.cpp
+++ b/Engine/Source/Runtime/Landscape/Private/LandscapeEditProcedural.cpp
@@ -270,7 +270,8 @@ public:
 	{
 		ReadHeightmapTexture1Param.Bind(Initializer.ParameterMap, TEXT("ReadHeightmapTexture1"));
 		ReadHeightmapTexture2Param.Bind(Initializer.ParameterMap, TEXT("ReadHeightmapTexture2"));
-		ReadHeightmapTextureSamplerParam.Bind(Initializer.ParameterMap, TEXT("ReadHeightmapTextureSampler"));
+		ReadHeightmapTexture1SamplerParam.Bind(Initializer.ParameterMap, TEXT("ReadHeightmapTexture1Sampler"));
+		ReadHeightmapTexture2SamplerParam.Bind(Initializer.ParameterMap, TEXT("ReadHeightmapTexture2Sampler"));
 		LayerInfoParam.Bind(Initializer.ParameterMap, TEXT("LayerInfo"));
 		OutputConfigParam.Bind(Initializer.ParameterMap, TEXT("OutputConfig"));
 		TextureSizeParam.Bind(Initializer.ParameterMap, TEXT("HeightmapTextureSize"));
@@ -283,11 +284,11 @@ public:
 
 	void SetParameters(FRHICommandList& RHICmdList, const FLandscapeHeightmapProceduralShaderParameters& InParams)
 	{
-		SetTextureParameter(RHICmdList, GetPixelShader(), ReadHeightmapTexture1Param, ReadHeightmapTextureSamplerParam, TStaticSamplerState<SF_Point, AM_Clamp, AM_Clamp, AM_Clamp>::GetRHI(), InParams.ReadHeightmap1->Resource->TextureRHI);
+		SetTextureParameter(RHICmdList, GetPixelShader(), ReadHeightmapTexture1Param, ReadHeightmapTexture1SamplerParam, TStaticSamplerState<SF_Point, AM_Clamp, AM_Clamp, AM_Clamp>::GetRHI(), InParams.ReadHeightmap1->Resource->TextureRHI);
 
 		if (InParams.ReadHeightmap2 != nullptr)
 		{
-			SetTextureParameter(RHICmdList, GetPixelShader(), ReadHeightmapTexture2Param, ReadHeightmapTextureSamplerParam, TStaticSamplerState<SF_Point, AM_Clamp, AM_Clamp, AM_Clamp>::GetRHI(), InParams.ReadHeightmap2->Resource->TextureRHI);
+			SetTextureParameter(RHICmdList, GetPixelShader(), ReadHeightmapTexture2Param, ReadHeightmapTexture2SamplerParam, TStaticSamplerState<SF_Point, AM_Clamp, AM_Clamp, AM_Clamp>::GetRHI(), InParams.ReadHeightmap2->Resource->TextureRHI);
 		}
 
 		FVector2D LayerInfo(InParams.LayerWeight, InParams.LayerVisible ? 1.0f : 0.0f);
@@ -306,7 +307,8 @@ public:
 		bool bShaderHasOutdatedParameters = FShader::Serialize(Ar);
 		Ar << ReadHeightmapTexture1Param;
 		Ar << ReadHeightmapTexture2Param;
-		Ar << ReadHeightmapTextureSamplerParam;
+		Ar << ReadHeightmapTexture1SamplerParam;
+		Ar << ReadHeightmapTexture2SamplerParam;
 		Ar << LayerInfoParam;
 		Ar << OutputConfigParam;
 		Ar << TextureSizeParam;
@@ -318,7 +320,8 @@ public:
 private:
 	FShaderResourceParameter ReadHeightmapTexture1Param;
 	FShaderResourceParameter ReadHeightmapTexture2Param;
-	FShaderResourceParameter ReadHeightmapTextureSamplerParam;
+	FShaderResourceParameter ReadHeightmapTexture1SamplerParam;
+	FShaderResourceParameter ReadHeightmapTexture2SamplerParam;
 	FShaderParameter LayerInfoParam;
 	FShaderParameter OutputConfigParam;
 	FShaderParameter TextureSizeParam;
@@ -346,7 +349,7 @@ public:
 		: FGlobalShader(Initializer)
 	{
 		ReadHeightmapTexture1Param.Bind(Initializer.ParameterMap, TEXT("ReadHeightmapTexture1"));
-		ReadHeightmapTextureSamplerParam.Bind(Initializer.ParameterMap, TEXT("ReadHeightmapTextureSampler"));
+		ReadHeightmapTexture1SamplerParam.Bind(Initializer.ParameterMap, TEXT("ReadHeightmapTexture1Sampler"));
 		CurrentMipHeightmapSizeParam.Bind(Initializer.ParameterMap, TEXT("CurrentMipTextureSize"));
 		ParentMipHeightmapSizeParam.Bind(Initializer.ParameterMap, TEXT("ParentMipTextureSize"));
 		CurrentMipComponentVertexCountParam.Bind(Initializer.ParameterMap, TEXT("CurrentMipComponentVertexCount"));
@@ -357,7 +360,7 @@ public:
 
 	void SetParameters(FRHICommandList& RHICmdList, const FLandscapeHeightmapProceduralShaderParameters& InParams)
 	{
-		SetTextureParameter(RHICmdList, GetPixelShader(), ReadHeightmapTexture1Param, ReadHeightmapTextureSamplerParam, TStaticSamplerState<SF_Point, AM_Clamp, AM_Clamp, AM_Clamp>::GetRHI(), InParams.ReadHeightmap1->Resource->TextureRHI);
+		SetTextureParameter(RHICmdList, GetPixelShader(), ReadHeightmapTexture1Param, ReadHeightmapTexture1SamplerParam, TStaticSamplerState<SF_Point, AM_Clamp, AM_Clamp, AM_Clamp>::GetRHI(), InParams.ReadHeightmap1->Resource->TextureRHI);
 
 		SetShaderValue(RHICmdList, GetPixelShader(), CurrentMipHeightmapSizeParam, FVector2D(InParams.CurrentMipHeightmapSize.X, InParams.CurrentMipHeightmapSize.Y));
 		SetShaderValue(RHICmdList, GetPixelShader(), ParentMipHeightmapSizeParam, FVector2D(InParams.ParentMipHeightmapSize.X, InParams.ParentMipHeightmapSize.Y));
@@ -368,7 +371,7 @@ public:
 	{
 		bool bShaderHasOutdatedParameters = FShader::Serialize(Ar);
 		Ar << ReadHeightmapTexture1Param;
-		Ar << ReadHeightmapTextureSamplerParam;
+		Ar << ReadHeightmapTexture1SamplerParam;
 		Ar << CurrentMipHeightmapSizeParam;
 		Ar << ParentMipHeightmapSizeParam;
 		Ar << CurrentMipComponentVertexCountParam;
@@ -378,7 +381,7 @@ public:
 
 private:
 	FShaderResourceParameter ReadHeightmapTexture1Param;
-	FShaderResourceParameter ReadHeightmapTextureSamplerParam;
+	FShaderResourceParameter ReadHeightmapTexture1SamplerParam;
 	FShaderParameter CurrentMipHeightmapSizeParam;
 	FShaderParameter ParentMipHeightmapSizeParam;
 	FShaderParameter CurrentMipComponentVertexCountParam;
