Fix Vulkan shader compilation errors.

#rnx

diff --git a/Engine/Shaders/Private/Common.ush b/Engine/Shaders/Private/Common.ush
--- a/Engine/Shaders/Private/Common.ush
+++ b/Engine/Shaders/Private/Common.ush
@@ -172,7 +172,7 @@ MaterialFloat4 TextureExternalSample(TextureExternal Tex, SamplerState Sampler,
 	return Tex.Sample(Sampler, UV);
 #endif
 }
-MaterialFloat4 TextureExternalSampleLevel(Texture2D Tex, SamplerState Sampler, float2 UV, MaterialFloat Mip)
+MaterialFloat4 TextureExternalSampleLevel(TextureExternal Tex, SamplerState Sampler, float2 UV, MaterialFloat Mip)
 {
 	return Tex.SampleLevel(Sampler, UV, Mip);
 }
diff --git a/Engine/Shaders/Private/VolumetricFogLightFunction.usf b/Engine/Shaders/Private/VolumetricFogLightFunction.usf
--- a/Engine/Shaders/Private/VolumetricFogLightFunction.usf
+++ b/Engine/Shaders/Private/VolumetricFogLightFunction.usf
@@ -42,7 +42,7 @@ void Main(
 	float3 AbsoluteWorldPosition = SvPositionToWorld(SvPosition);
 	*/
 
-	float2 ScreenPosition = (SvPosition * LightFunctionTexelSize - .5f) * float2(2, -2);
+	float2 ScreenPosition = (SvPosition.xy * LightFunctionTexelSize - .5f) * float2(2, -2);
 	float3 AbsoluteWorldPosition = GetWorldPositionForLightFunctionTexturePosition(ScreenPosition);
 
 	float4 LightVector = mul(float4(AbsoluteWorldPosition, 1), LightFunctionWorldToLight);
