Fix HQ reflections with movable light wth lightmap ignoring scene capture's average brightness.

diff --git a/Engine/Shaders/Private/MobileBasePassPixelShader.usf b/Engine/Shaders/Private/MobileBasePassPixelShader.usf
--- a/Engine/Shaders/Private/MobileBasePassPixelShader.usf
+++ b/Engine/Shaders/Private/MobileBasePassPixelShader.usf
@@ -202,7 +202,7 @@ half3 BlendReflectionCapture(FMaterialPixelParameters MaterialParameters, half R
 	float4 Reflection = ReflectionCube.SampleLevel(ReflectionSampler, ProjectedCaptureVector.xyz, MipLevel);
 	half3 SpecularIBL = RGBMDecode(Reflection, 16.0f);
 
-#if MOVABLE_DIRECTIONAL_LIGHT
+#if !LQ_TEXTURE_LIGHTMAP && !CACHED_POINT_INDIRECT_LIGHTING
 	// ignore average brightness for non-lightmap use case.
 	InvReflectionCubemapAverageBrightness = 1.0f;
 #endif
