In ES 3.1, try framebuffer_fetch_depth_stencil to get the depth directly, and if unavailable fall back to framebuffer_fetch from the alpha channel.

diff --git a/Engine/Shaders/Private/SceneTexturesCommon.ush b/Engine/Shaders/Private/SceneTexturesCommon.ush
--- a/Engine/Shaders/Private/SceneTexturesCommon.ush
+++ b/Engine/Shaders/Private/SceneTexturesCommon.ush
@@ -110,7 +110,14 @@ float CalcSceneDepth(float2 ScreenUV)
 		#elif METAL_PROFILE && !MAC
 			return FramebufferFetchES2().w;
 		#elif COMPILER_GLSL_ES3_1
-			return FramebufferFetchES2().w;
+			#if OUTPUT_GAMMA_SPACE
+				// Try framebuffer_fetch_depth_stencil to get the depth directly if the extension is available.
+				// We cannot fall back to fetching the alpha channel when MobileHDR=false because the alpha channel is only 8-bit.
+				return DepthbufferFetchES2(65000.0, View.InvDeviceZToWorldZTransform[2], View.InvDeviceZToWorldZTransform[3]);
+			#else
+				// Try framebuffer_fetch_depth_stencil to get the depth directly, and if unavailable fall back to framebuffer_fetch from the alpha channel.
+				return DepthbufferFetchES2(FramebufferFetchES2().w, View.InvDeviceZToWorldZTransform[2], View.InvDeviceZToWorldZTransform[3]);
+			#endif
 		#elif VULKAN_SUBPASS_DEPTHFETCH
 			// Special intrinsic to read from the current depth buffer
 			return ConvertFromDeviceZ(VulkanSubpassDepthFetch());
