Fix numeric value for custom stencil case for mobile scene texture fetch

diff --git a/Engine/Shaders/Private/MaterialTemplate.ush b/Engine/Shaders/Private/MaterialTemplate.ush
--- a/Engine/Shaders/Private/MaterialTemplate.ush
+++ b/Engine/Shaders/Private/MaterialTemplate.ush
@@ -1390,7 +1390,7 @@ MaterialFloat4 MobileSceneTextureLookup(inout FMaterialPixelParameters Parameter
 #endif// POST_PROCESS_MATERIAL	
 	}
 	// PPI_CustomStencil
-	else if (SceneTextureId == 24)
+	else if (SceneTextureId == 25)
 	{
 		MaterialFloat Stencil = Texture2DSample(MobileSceneTextures.MobileCustomStencilTexture, MobileSceneTextures.MobileCustomStencilTextureSampler, UV).r*255.0;
 		Stencil = floor(Stencil + 0.5);
