DR - Fix for potential TDR due to very big values
#rnx

diff --git a/Engine/Shaders/Private/DistanceFieldObjectCulling.usf b/Engine/Shaders/Private/DistanceFieldObjectCulling.usf
--- a/Engine/Shaders/Private/DistanceFieldObjectCulling.usf
+++ b/Engine/Shaders/Private/DistanceFieldObjectCulling.usf
@@ -347,8 +347,17 @@ void BuildTileConesMain(
 		}
 
 		uint TileIndex = GroupId.y * NumGroups.x + GroupId.x;
-		RWTileConeAxisAndCos[TileIndex] = float4(TileConeAxis, TileConeAngleCos);
-		RWTileConeDepthRanges[TileIndex] = ConeAxisDepthRanges;
+		if (IntegerTileMinZ > IntegerTileMaxZ)
+		{
+			// Guard against IntegerTileMinZ never getting updated
+			RWTileConeAxisAndCos[TileIndex] = float4(0, 0, 0, 1);
+			RWTileConeDepthRanges[TileIndex] = 0;
+		}
+		else
+		{
+			RWTileConeAxisAndCos[TileIndex] = float4(TileConeAxis, TileConeAngleCos);
+			RWTileConeDepthRanges[TileIndex] = ConeAxisDepthRanges;
+		}
 	}
 }
 
