DR - Convert all global float3/half3/uint3 to float4/half4/unit4
#rnx

diff --git a/Engine/Shaders/Private/CapsuleShadowShaders.usf b/Engine/Shaders/Private/CapsuleShadowShaders.usf
--- a/Engine/Shaders/Private/CapsuleShadowShaders.usf
+++ b/Engine/Shaders/Private/CapsuleShadowShaders.usf
@@ -179,11 +179,11 @@ bool SphereIntersectSphere(float4 SphereCenterAndRadius, float4 OtherSphereCente
 
 #if LIGHT_SOURCE_MODE == LIGHT_SOURCE_PUNCTUAL
 /** From point being shaded toward light, for directional lights. */
-float3 LightDirection;
+float4 LightDirection;
 float4 LightPositionAndInvRadius;
 float LightSourceRadius;
 float RayStartOffsetDepthScale;
-float3 LightAngleAndNormalThreshold;
+float4 LightAngleAndNormalThreshold;
 #endif
 
 uint4 ScissorRectMinAndSize;
@@ -377,7 +377,7 @@ void SetupTileCullingData(
 		bTileShouldComputeShadowing = bTileShouldComputeShadowing && bTileInLightInfluenceBounds;
 
 	#else
-		TileCullingData0.ConeAxis = TileCullingData1.ConeAxis = LightDirection;
+		TileCullingData0.ConeAxis = TileCullingData1.ConeAxis = LightDirection.xyz;
 		TileCullingData0.ConeAngleCos = TileCullingData1.ConeAngleCos = cos(LightAngleAndNormalThreshold.x);
 		TileCullingData0.ConeAngleSin = TileCullingData1.ConeAngleSin = sin(LightAngleAndNormalThreshold.x);
 	#endif
diff --git a/Engine/Shaders/Private/DeferredLightVertexShaders.usf b/Engine/Shaders/Private/DeferredLightVertexShaders.usf
--- a/Engine/Shaders/Private/DeferredLightVertexShaders.usf
+++ b/Engine/Shaders/Private/DeferredLightVertexShaders.usf
@@ -9,7 +9,7 @@
 float4 StencilingGeometryPosAndScale;
 float4 StencilingConeParameters;	// .x NumSides (0 if not cone), .y NumSlices, .z ConeAngle, .w ConeSphereRadius
 float4x4 StencilingConeTransform;
-float3 StencilingPreViewTranslation;
+float4 StencilingPreViewTranslation;
 
 /** Vertex shader for rendering a directional light using a full screen quad. */
 void DirectionalVertexMain(
@@ -62,7 +62,7 @@ void RadialVertexMain(
 			// Create a position in the local space of the cone, forming a circle in the XY plane, and offsetting along the Z axis
 			const float3 LocalPosition = float3(ZRadius * SliceIndex / (float)(NumSlices - 1), SliceRadius * sin(CurrentAngle), SliceRadius * cos(CurrentAngle));
 			// Transform to world space and apply pre-view translation, since these vertices will be used with a shader that has pre-view translation removed
-			WorldPosition = mul(float4(LocalPosition, 1), StencilingConeTransform).xyz + StencilingPreViewTranslation;
+			WorldPosition = mul(float4(LocalPosition, 1), StencilingConeTransform).xyz + StencilingPreViewTranslation.xyz;
 		}
 		else
 		{
@@ -81,7 +81,7 @@ void RadialVertexMain(
 
 			const float CurrentAngle = SideIndex * 2 * PI / (float)NumSides;
 			const float3 LocalPosition = float3(ZDistance, SliceRadius * sin(CurrentAngle), SliceRadius * cos(CurrentAngle));
-			WorldPosition = mul(float4(LocalPosition, 1), StencilingConeTransform).xyz + StencilingPreViewTranslation;
+			WorldPosition = mul(float4(LocalPosition, 1), StencilingConeTransform).xyz + StencilingPreViewTranslation.xyz;
 		}
 	}
 	else
diff --git a/Engine/Shaders/Private/DistanceFieldGlobalIllumination.usf b/Engine/Shaders/Private/DistanceFieldGlobalIllumination.usf
--- a/Engine/Shaders/Private/DistanceFieldGlobalIllumination.usf
+++ b/Engine/Shaders/Private/DistanceFieldGlobalIllumination.usf
@@ -310,7 +310,7 @@ void SetupLightVPLsIndirectArgumentsCS(
 }
 
 // From receiver to to light
-float3 LightDirection;
+float4 LightDirection;
 float4 LightPositionAndInvRadius;
 float LightSourceRadius;
 float2 TanLightAngleAndNormalThreshold;
@@ -443,7 +443,7 @@ float3 ComputeVPLFlux(uint SurfelIndex, float4x4 InstanceToWorld)
 	float3 AccumulatedLighting = 0;
 
 	{
-		float SurfelDotLight = dot(SurfelNormal, LightDirection);
+		float SurfelDotLight = dot(SurfelNormal, LightDirection.xyz);
 
 		BRANCH
 		if (SurfelDotLight > 0)
@@ -479,8 +479,8 @@ float3 ComputeVPLFlux(uint SurfelIndex, float4x4 InstanceToWorld)
 				}
 				else
 				{
-					WorldRayStart = PositionAndRadius.xyz + LightDirection * RayStartOffset;
-					WorldRayEnd = PositionAndRadius.xyz + LightDirection * TraceDistance;
+					WorldRayStart = PositionAndRadius.xyz + LightDirection.xyz * RayStartOffset;
+					WorldRayEnd = PositionAndRadius.xyz + LightDirection.xyz * TraceDistance;
 					MaxRayTime = TraceDistance;
 					TanLightAngle = TanLightAngleAndNormalThreshold.x;
 
diff --git a/Engine/Shaders/Private/DistanceFieldLightingShared.ush b/Engine/Shaders/Private/DistanceFieldLightingShared.ush
--- a/Engine/Shaders/Private/DistanceFieldLightingShared.ush
+++ b/Engine/Shaders/Private/DistanceFieldLightingShared.ush
@@ -42,7 +42,7 @@ float SampleMeshDistanceField(float3 VolumeUV, float2 ObjectDistanceFieldMAD)
 	return TextureValue;
 }
 
-float3 DistanceFieldAtlasTexelSize;
+float4 DistanceFieldAtlasTexelSize;
 
 RWBuffer<uint> RWObjectIndirectArguments;
 Buffer<uint> ObjectIndirectArguments;
diff --git a/Engine/Shaders/Private/DistanceFieldShadowing.usf b/Engine/Shaders/Private/DistanceFieldShadowing.usf
--- a/Engine/Shaders/Private/DistanceFieldShadowing.usf
+++ b/Engine/Shaders/Private/DistanceFieldShadowing.usf
@@ -301,11 +301,11 @@ RWTexture2D<float2> RWShadowFactors;
 float2 NumGroups;
 
 /** From point being shaded toward light, for directional lights. */
-float3 LightDirection;
+float4 LightDirection;
 float4 LightPositionAndInvRadius;
 float LightSourceRadius;
 float RayStartOffsetDepthScale;
-float3 TanLightAngleAndNormalThreshold;
+float4 TanLightAngleAndNormalThreshold;
 int4 ScissorRectMinAndSize;
 
 /** Min and Max depth for this tile. */
@@ -426,8 +426,8 @@ void CullObjectsToTileWithGather(
 		|| LightVector1Length < 1.0f / LightPositionAndInvRadius.w + Group1BoundingRadius;
 
 #else
-	float3 LightDirection0 = LightDirection;
-	float3 LightDirection1 = LightDirection;
+	float3 LightDirection0 = LightDirection.xyz;
+	float3 LightDirection1 = LightDirection.xyz;
 	float RayLength0 = TraceDistance;
 	float RayLength1 = TraceDistance;
 
@@ -547,8 +547,8 @@ void DistanceFieldShadowingCS(
 
 			#else
 
-				float3 WorldRayStart = OpaqueWorldPosition + LightDirection * RayStartOffset;
-				float3 WorldRayEnd = OpaqueWorldPosition + LightDirection * TraceDistance;
+				float3 WorldRayStart = OpaqueWorldPosition + LightDirection.xyz * RayStartOffset;
+				float3 WorldRayEnd = OpaqueWorldPosition + LightDirection.xyz * TraceDistance;
 				float MaxRayTime = TraceDistance;
 				float TanLightAngle = TanLightAngleAndNormalThreshold.x;
 
diff --git a/Engine/Shaders/Private/GPUFastFourierTransform.usf b/Engine/Shaders/Private/GPUFastFourierTransform.usf
--- a/Engine/Shaders/Private/GPUFastFourierTransform.usf
+++ b/Engine/Shaders/Private/GPUFastFourierTransform.usf
@@ -88,7 +88,7 @@ float ConvertToLuma(in float3 ColorValue)
 // BrightPixelGain.y = Max Threshold: the maximal value brightness
 // BrightPixelGain.z = Multiplier applied to brightness above Min Threshold
 
-float3 BrightPixelGain;
+float4 BrightPixelGain;
 
 bool FilterPixel(in float3 Filter, inout float4 PixelValue)
 {
@@ -269,7 +269,7 @@ void ModifyInput(inout Complex LocalBuffer[2][RADIX])
 	{
 		float4 InputColor = float4(LocalBuffer[0][r].x, LocalBuffer[0][r].y, LocalBuffer[1][r].x, LocalBuffer[1][r].y);
 
-		bool bIsChanged = FilterPixel(BrightPixelGain, InputColor);
+		bool bIsChanged = FilterPixel(BrightPixelGain.xyz, InputColor);
 		
 		if (bIsChanged)
 		{
diff --git a/Engine/Shaders/Private/GammaCorrection.usf b/Engine/Shaders/Private/GammaCorrection.usf
--- a/Engine/Shaders/Private/GammaCorrection.usf
+++ b/Engine/Shaders/Private/GammaCorrection.usf
@@ -17,7 +17,7 @@ void MainVS(
 	DrawRectangle(InPosition, InTexCoord, OutPosition, OutTexCoord);
 }
 
-half3 ColorScale;
+half4 ColorScale;
 half4 OverlayColor;
 half InverseGamma;
 
@@ -46,7 +46,7 @@ void MainPS(float2 UV : TEXCOORD0, out float4 OutColor : SV_Target0)
 	half3 LDRColor = TonemapAndGammaCorrect(LinearColor.rgb);
 	
 	// blend with custom LDR color, used for Fade track in Matinee
-	LDRColor = lerp(LDRColor * ColorScale, OverlayColor.rgb, OverlayColor.a);
+	LDRColor = lerp(LDRColor * ColorScale.xyz, OverlayColor.rgb, OverlayColor.a);
 	
 	// RETURN_COLOR not needed unless writing to SceneColor
 	OutColor = float4(LDRColor, LinearColor.a);
diff --git a/Engine/Shaders/Private/GeometryCacheVertexFactory.ush b/Engine/Shaders/Private/GeometryCacheVertexFactory.ush
--- a/Engine/Shaders/Private/GeometryCacheVertexFactory.ush
+++ b/Engine/Shaders/Private/GeometryCacheVertexFactory.ush
@@ -9,15 +9,15 @@
 
 /*
     Meshes vertices can optionally still be packed and will be unpacked using the following formula:
-    VertextBuffer.Position*MeshExtension+MeshOrigin
+    VertextBuffer.Position * MeshExtension + MeshOrigin
 */
-float3 MeshOrigin;
-float3 MeshExtension;
+float4 MeshOrigin;
+float4 MeshExtension;
 
 /* This vertex type is passed in an extra float4 MotionBlurData. This allows specifying motion blur data in a generic way. The previous
     frame's object space position is derived from the MotionBlurData and the current unpacked object space position fields in the following way:
 
-    (Vertex.MotionBlurData * MotionBlurDataExtension + MotionBlurDataOrigin) + (UnpackePosition * MotionBlurPositionScale)
+    (Vertex.MotionBlurData * MotionBlurDataExtension + MotionBlurDataOrigin.xyz) + (UnpackePosition * MotionBlurPositionScale)
     
     Using this formula a number of effects can be achieved with a single vertex shader. Some common examples
     - Packed previous vertex positions are bound to MotionBlurData
@@ -33,8 +33,8 @@ float3 MeshExtension;
         - MotionBlurDataExtension: 0
         - MotionBlurPositionScale: 1
 */
-float3 MotionBlurDataOrigin;
-float3 MotionBlurDataExtension;
+float4 MotionBlurDataOrigin;
+float4 MotionBlurDataExtension;
 float MotionBlurPositionScale;
 
 struct FVertexFactoryInput
@@ -160,7 +160,7 @@ FMaterialVertexParameters GetMaterialVertexParameters(FVertexFactoryInput Input,
 */
 float3 UnpackedPosition( FVertexFactoryInput Input )
 {
-		return float3(Input.Position.xyz*MeshExtension + MeshOrigin);
+		return float3(Input.Position.xyz * MeshExtension.xyz + MeshOrigin.xyz);
 }
 
 /**
@@ -276,7 +276,7 @@ FVertexFactoryInterpolantsVSToPS VertexFactoryGetInterpolantsVSToPS(FVertexFacto
 // @return The previous position of the vertex in object space.
 float3 CalcPreviousPosition(FVertexFactoryInput Input, FVertexFactoryIntermediates Intermediates)
 {
-    float3 PreviousPosition =  (Input.MotionBlurData.xyz * MotionBlurDataExtension + MotionBlurDataOrigin) + Intermediates.UnpackedPosition * MotionBlurPositionScale;
+    float3 PreviousPosition =  (Input.MotionBlurData.xyz * MotionBlurDataExtension.xyz + MotionBlurDataOrigin.xyz) + Intermediates.UnpackedPosition * MotionBlurPositionScale;
 #if 0    
     float ExaggerateMotion = 10.0;
     float3 Delta = Intermediates.UnpackedPosition - PreviousPosition;
diff --git a/Engine/Shaders/Private/GlobalDistanceField.usf b/Engine/Shaders/Private/GlobalDistanceField.usf
--- a/Engine/Shaders/Private/GlobalDistanceField.usf
+++ b/Engine/Shaders/Private/GlobalDistanceField.usf
@@ -110,10 +110,10 @@ void CullObjectsForVolumeCS(
 }
 
 uint ClipmapIndex;
-float3 VolumeTexelSize;
-float3 UpdateRegionVolumeMin;
+float4 VolumeTexelSize;
+float4 UpdateRegionVolumeMin;
 float UpdateRegionVolumeStep;
-uint3 CullGridDimension;
+uint4 CullGridDimension;
 
 RWBuffer<uint> RWCullGridObjectNum;
 RWBuffer<uint> RWCullGridObjectArray;
@@ -142,7 +142,7 @@ void CullObjectsToGridCS(
 	uint ThreadIndex = GroupThreadId.x;
 	uint CullTileIndex = ComputeCullTileIndex(GroupId);
 
-	if (ThreadIndex == 0 && all(GroupId < CullGridDimension))
+	if (ThreadIndex == 0 && all(GroupId < CullGridDimension.xyz))
 	{
 		// Clear array num to 0 for all cull grid cells
 		RWCullGridObjectNum[CullTileIndex] = 0;
@@ -150,12 +150,12 @@ void CullObjectsToGridCS(
 
 	GroupMemoryBarrierWithGroupSync();
 
-	if (all(GroupId < CullGridDimension))
+	if (all(GroupId < CullGridDimension.xyz))
 	{
 		uint NumCulledObjects = GetCulledNumObjects();
 		float GlobalVolumeExtent = GlobalVolumeCenterAndExtent[ClipmapIndex].w;
-		float3 TileBoxExtent = .5f * CULL_GRID_TILE_SIZE * VolumeTexelSize * GlobalVolumeExtent * 2;
-		float3 TileBoxCenter = UpdateRegionVolumeMin + ((GroupId.xyz + .5f) * CULL_GRID_TILE_SIZE + .5f) * VolumeTexelSize * GlobalVolumeExtent * 2;
+		float3 TileBoxExtent = .5f * CULL_GRID_TILE_SIZE * VolumeTexelSize.xyz * GlobalVolumeExtent * 2;
+		float3 TileBoxCenter = UpdateRegionVolumeMin.xyz + ((GroupId.xyz + .5f) * CULL_GRID_TILE_SIZE + .5f) * VolumeTexelSize.xyz * GlobalVolumeExtent * 2;
 
 		for (uint ObjectIndex = ThreadIndex; ObjectIndex < NumCulledObjects; ObjectIndex += GRID_CULL_THREADGORUP_TOTALSIZE)
 		{
@@ -180,7 +180,7 @@ void CullObjectsToGridCS(
 	}
 }
 
-uint3 UpdateRegionSize;
+uint4 UpdateRegionSize;
 Buffer<uint> CullGridObjectNum;
 Buffer<uint> CullGridObjectArray;
 Texture3D<float> ParentGlobalDistanceFieldTexture;
@@ -228,8 +228,8 @@ void CompositeObjectDistanceFieldsCS(
 	uint ThreadIndex = (GroupThreadId.z * COMPOSITE_THREADGROUP_SIZEY + GroupThreadId.y) * COMPOSITE_THREADGROUP_SIZEX + GroupThreadId.x;
 
 	float3 GroupSize = float3(COMPOSITE_THREADGROUP_SIZEX, COMPOSITE_THREADGROUP_SIZEY, COMPOSITE_THREADGROUP_SIZEZ);
-	float3 TileBoxExtent = .5f * GroupSize * VolumeTexelSize * GlobalVolumeExtent * 2;
-	float3 TileBoxCenter = UpdateRegionVolumeMin + ((GroupId.xyz + .5f) * GroupSize + .5f) * VolumeTexelSize * GlobalVolumeExtent * 2;
+	float3 TileBoxExtent = .5f * GroupSize * VolumeTexelSize.xyz * GlobalVolumeExtent * 2;
+	float3 TileBoxCenter = UpdateRegionVolumeMin.xyz + ((GroupId.xyz + .5f) * GroupSize + .5f) * VolumeTexelSize.xyz * GlobalVolumeExtent * 2;
 
 #if USE_CULL_GRID
 	uint CullTileIndex = ComputeCullTileIndex(DispatchThreadId / CULL_GRID_TILE_SIZE);
@@ -266,7 +266,7 @@ void CompositeObjectDistanceFieldsCS(
 
 #endif
 
-	float3 WorldPosition = UpdateRegionVolumeMin + (DispatchThreadId.xyz + .5f) * VolumeTexelSize * GlobalVolumeExtent * 2;
+	float3 WorldPosition = UpdateRegionVolumeMin.xyz + (DispatchThreadId.xyz + .5f) * VolumeTexelSize.xyz * GlobalVolumeExtent * 2;
 	float3 GlobalUV = ComputeGlobalUV(WorldPosition, ClipmapIndex);
 	uint3 IntCoordinate = (uint3)(frac(GlobalUV) * GlobalVolumeDimension);
 
@@ -308,7 +308,7 @@ void CompositeObjectDistanceFieldsCS(
 		}
 	}
 
-	if (all(DispatchThreadId.xyz < UpdateRegionSize))
+	if (all(DispatchThreadId.xyz < UpdateRegionSize.xyz))
 	{
 		RWGlobalDistanceFieldTexture[IntCoordinate] = MinDistance;
 	}
@@ -325,7 +325,7 @@ void CompositeHeightfieldsIntoGlobalDistanceFieldCS(
 	uint3 GroupThreadId : SV_GroupThreadID)
 {
 	// We only dispatch one depth layer, so the z value is incorrect
-	float3 LookupWorldPosition = UpdateRegionVolumeMin + float3(DispatchThreadId.xy, 0.0f) * UpdateRegionVolumeStep;
+	float3 LookupWorldPosition = UpdateRegionVolumeMin.xyz + float3(DispatchThreadId.xy, 0.0f) * UpdateRegionVolumeStep;
 
 	uint ValidHeightfieldIndex = 0;
 
diff --git a/Engine/Shaders/Private/HZBOcclusion.usf b/Engine/Shaders/Private/HZBOcclusion.usf
--- a/Engine/Shaders/Private/HZBOcclusion.usf
+++ b/Engine/Shaders/Private/HZBOcclusion.usf
@@ -54,7 +54,7 @@ SamplerState	BoundsCenterSampler;
 Texture2D		BoundsExtentTexture;
 SamplerState	BoundsExtentSampler;
 
-float3 HZBUvFactor;
+float4 HZBUvFactor;
 float4 HZBSize;
 
 void HZBTestPS(float2 InUV : TEXCOORD0, out float4 OutColor : SV_Target0)
diff --git a/Engine/Shaders/Private/HeightfieldLighting.usf b/Engine/Shaders/Private/HeightfieldLighting.usf
--- a/Engine/Shaders/Private/HeightfieldLighting.usf
+++ b/Engine/Shaders/Private/HeightfieldLighting.usf
@@ -236,7 +236,7 @@ void ShadowHeightfieldsPS(
 	OutColor = float4(Visibility, 0, 0, 0);
 }
 
-float3 LightDirection;
+float4 LightDirection;
 float TanLightAngle;
 
 void RayTracedShadowHeightfieldsPS(
@@ -270,7 +270,7 @@ void RayTracedShadowHeightfieldsPS(
 		float TraceDistance = 10000;
 
 		float3 WorldRayStart = OpaqueWorldPosition;
-		float3 WorldRayEnd = OpaqueWorldPosition + -LightDirection * TraceDistance;
+		float3 WorldRayEnd = OpaqueWorldPosition + -LightDirection.xyz * TraceDistance;
 		float MaxRayTime = TraceDistance;
 
 		Shadow = ShadowRayTraceThroughCulledObjects(
@@ -295,7 +295,7 @@ void RayTracedShadowHeightfieldsPS(
 	OutColor = float4(Shadow, 0, 0, 0);
 }
 
-float3 LightColor;
+float4 LightColor;
 float SkyLightIndirectScale;
 
 #if APPLY_LIGHT_FUNCTION
@@ -321,7 +321,7 @@ void LightHeightfieldsPS(
 	float Shadowing = Texture2DSampleLevel(HeightfieldShadowing, HeightfieldShadowingSampler, Input.GlobalUV, 0).x;
 
 	float3 DiffuseColor = GetGlobalDiffuseColor(Input.GlobalUV);
-	float3 Lighting = (Shadowing * GreyScaleLightFunction * saturate(dot(-LightDirection, WorldHeightfieldNormal))) * LightColor;
+	float3 Lighting = (Shadowing * GreyScaleLightFunction * saturate(dot(-LightDirection.xyz, WorldHeightfieldNormal))) * LightColor.xyz;
 	 
 	float3 SkylightDiffuse = GetSkySHDiffuse(WorldHeightfieldNormal) * View.SkyLightColor.rgb * SkyLightIndirectScale;
 	Lighting += SkylightDiffuse;
diff --git a/Engine/Shaders/Private/LightFunctionPixelShader.usf b/Engine/Shaders/Private/LightFunctionPixelShader.usf
--- a/Engine/Shaders/Private/LightFunctionPixelShader.usf
+++ b/Engine/Shaders/Private/LightFunctionPixelShader.usf
@@ -13,7 +13,7 @@
 float4x4 SvPositionToLight;
 
 /** Fade distance in x, disabled brightness in y, output for preview shadows mask in z. */
-float3 LightFunctionParameters2;
+float4 LightFunctionParameters2;
 
 void Main(
 	// in float4 ScreenPosition : TEXCOORD0, // Not used and not output in VS
diff --git a/Engine/Shaders/Private/LightMapDensityShader.usf b/Engine/Shaders/Private/LightMapDensityShader.usf
--- a/Engine/Shaders/Private/LightMapDensityShader.usf
+++ b/Engine/Shaders/Private/LightMapDensityShader.usf
@@ -140,7 +140,7 @@ float2		LightMapResolutionScale;
  *	x = 0.0, y = 1.0 for unbuilt lighting
  *	z = 1.0 if the object is selected
  */
-float3		BuiltLightingAndSelectedFlags;
+float4		BuiltLightingAndSelectedFlags;
 
 /** 
  *	X = scalar to multiply density against...
diff --git a/Engine/Shaders/Private/LightShaftShader.usf b/Engine/Shaders/Private/LightShaftShader.usf
--- a/Engine/Shaders/Private/LightShaftShader.usf
+++ b/Engine/Shaders/Private/LightShaftShader.usf
@@ -27,7 +27,7 @@ float4 WorldSpaceBlurOriginAndRadius;
 float LightSourceRadius;
 
 /** Spot light direction. */
-float3 WorldSpaceSpotDirection;
+float4 WorldSpaceSpotDirection;
 
 /** SpotLightInfo->CosOuterCone in x, SpotLightInfo->InvCosConeDifference in y. */
 float2 SpotAngles;
@@ -92,7 +92,7 @@ void DownsampleLightShaftsPixelMain(
 		#if SPOT_LIGHT_SHAFTS
 	
 			// Setup a mask that only allows occlusion or blooming from pixels in the spotlight's cone
-			float SpotDotProduct = dot(WorldLightVector, WorldSpaceSpotDirection);
+			float SpotDotProduct = dot(WorldLightVector, WorldSpaceSpotDirection.xyz);
 
 			// Filter the spotlight mask instead of the depths
 			SpotLightAttenuation = Square(saturate((SpotDotProduct / Distance - SpotAngles.x) * SpotAngles.y));
@@ -137,7 +137,7 @@ void DownsampleLightShaftsPixelMain(
 		#if SPOT_LIGHT_SHAFTS
 	
 			// Setup a mask that only allows occlusion or blooming from pixels in the spotlight's cone
-			float SpotDotProduct = dot(WorldLightVector, WorldSpaceSpotDirection);
+			float SpotDotProduct = dot(WorldLightVector, WorldSpaceSpotDirection.xyz);
 
 			// Filter the spotlight mask instead of the depths
 			SpotLightAttenuation = Square(saturate((SpotDotProduct / Distance - SpotAngles.x) * SpotAngles.y));
diff --git a/Engine/Shaders/Private/LocalVertexFactory.ush b/Engine/Shaders/Private/LocalVertexFactory.ush
--- a/Engine/Shaders/Private/LocalVertexFactory.ush
+++ b/Engine/Shaders/Private/LocalVertexFactory.ush
@@ -62,19 +62,19 @@ float4 InstancingFadeOutParams;
 #endif	// USE_INSTANCING
 
 #if USE_SPLINEDEFORM
-float3 SplineStartPos;
-float3 SplineStartTangent;
+float4 SplineStartPos;
+float4 SplineStartTangent;
 float SplineStartRoll;
 float2 SplineStartScale;
 float2 SplineStartOffset;
 
-float3 SplineEndPos;
-float3 SplineEndTangent;
+float4 SplineEndPos;
+float4 SplineEndTangent;
 float SplineEndRoll;
 float2 SplineEndScale;
 float2 SplineEndOffset;
 
-float3 SplineUpDir;
+float4 SplineUpDir;
 #if ES2_PROFILE
 float SmoothInterpRollScale;
 #else
@@ -84,9 +84,9 @@ bool SmoothInterpRollScale;
 float SplineMeshMinZ;
 float SplineMeshScaleZ;
 
-float3 SplineMeshDir;
-float3 SplineMeshX;
-float3 SplineMeshY;
+float4 SplineMeshDir;
+float4 SplineMeshX;
+float4 SplineMeshY;
 #endif	// USE_SPLINEDEFORM
 
 #ifndef MANUAL_VERTEX_FETCH
@@ -500,11 +500,11 @@ FMaterialVertexParameters GetMaterialVertexParameters(FVertexFactoryInput Input,
 		float HermiteAlpha = SmoothInterpRollScale ? smoothstep(0.0, 1.0, Alpha) : Alpha;
 
 		// Then find the point and direction of the spline at this point along
-		float3 SplinePos = SplineEvalPos( SplineStartPos, SplineStartTangent, SplineEndPos, SplineEndTangent, Alpha );
-		float3 SplineDir = SplineEvalDir( SplineStartPos, SplineStartTangent, SplineEndPos, SplineEndTangent, Alpha );
+		float3 SplinePos = SplineEvalPos( SplineStartPos.xyz, SplineStartTangent.xyz, SplineEndPos.xyz, SplineEndTangent.xyz, Alpha );
+		float3 SplineDir = SplineEvalDir( SplineStartPos.xyz, SplineStartTangent.xyz, SplineEndPos.xyz, SplineEndTangent.xyz, Alpha );
 	
 		// Find base frenet frame
-		float3 BaseXVec = normalize( cross(SplineUpDir, SplineDir) );
+		float3 BaseXVec = normalize( cross(SplineUpDir.xyz, SplineDir) );
 		float3 BaseYVec = normalize( cross(SplineDir, BaseXVec) );
 	
 		// Offset from the spline, using the frenet frame
@@ -526,7 +526,7 @@ FMaterialVertexParameters GetMaterialVertexParameters(FVertexFactoryInput Input,
 		YVec *= UseScale.y;
 
 		// Build overall transform
-		float3x3 SliceTransform3 = mul(transpose(float3x3(SplineMeshDir, SplineMeshX, SplineMeshY)), float3x3(float3(0,0,0), XVec, YVec));
+		float3x3 SliceTransform3 = mul(transpose(float3x3(SplineMeshDir.xyz, SplineMeshX.xyz, SplineMeshY.xyz)), float3x3(float3(0,0,0), XVec, YVec));
 		float4x3 SliceTransform = float4x3(SliceTransform3[0], SliceTransform3[1], SliceTransform3[2], SplinePos);
 
 		return SliceTransform;
@@ -542,10 +542,10 @@ FMaterialVertexParameters GetMaterialVertexParameters(FVertexFactoryInput Input,
 		half HermiteAlpha = SmoothInterpRollScale ? smoothstep(0.0, 1.0, Alpha) : Alpha;
 
 		// Then find the point and direction of the spline at this point along
-		half3 SplineDir = SplineEvalDir( SplineStartPos, SplineStartTangent, SplineEndPos, SplineEndTangent, Alpha );
+		half3 SplineDir = SplineEvalDir( SplineStartPos.xyz, SplineStartTangent.xyz, SplineEndPos.xyz, SplineEndTangent.xyz, Alpha );
 
 		// Find base frenet frame
-		half3 BaseXVec = normalize( cross(SplineUpDir, (half3)SplineDir) );
+		half3 BaseXVec = normalize( cross(SplineUpDir.xyz, (half3)SplineDir) );
 		half3 BaseYVec = normalize( cross((half3)SplineDir, BaseXVec) );
 
 		// Apply roll to frame around spline
@@ -562,7 +562,7 @@ FMaterialVertexParameters GetMaterialVertexParameters(FVertexFactoryInput Input,
 		YVec *= sign(UseScale.y);
 
 		// Build rotation transform
-		half3x3 SliceTransform = mul(transpose(half3x3(SplineMeshDir, SplineMeshX, SplineMeshY)), half3x3(SplineDir, XVec, YVec));
+		half3x3 SliceTransform = mul(transpose(half3x3(SplineMeshDir.xyz, SplineMeshX.xyz, SplineMeshY.xyz)), half3x3(SplineDir, XVec, YVec));
 
 		return SliceTransform;
 	}
@@ -579,14 +579,14 @@ float4 CalcWorldPosition(float4 Position, uint PrimitiveId)
 #elif USE_SPLINEDEFORM
 /*
 	// Make transform for this point along spline
-	float4x3 SliceTransform = CalcSliceTransform(dot(Position.xyz, SplineMeshDir));
+	float4x3 SliceTransform = CalcSliceTransform(dot(Position.xyz, SplineMeshDir.xyz));
 
 	// Transform into mesh space
 	float4 LocalPos = float4(mul(Position, SliceTransform), Position.w);
 	// Transform from mesh to world space
 	return TransformLocalToTranslatedWorld(LocalPos.xyz);
 	*/
-	return INVARIANT(TransformLocalToTranslatedWorld(float3(mul(Position, CalcSliceTransform(dot(Position.xyz, SplineMeshDir))).xyz), PrimitiveId));
+	return INVARIANT(TransformLocalToTranslatedWorld(float3(mul(Position, CalcSliceTransform(dot(Position.xyz, SplineMeshDir.xyz))).xyz), PrimitiveId));
 #else
 	return TransformLocalToTranslatedWorld(Position.xyz, PrimitiveId);
 #endif
@@ -616,7 +616,7 @@ half3x3 CalcTangentToLocal(FVertexFactoryInput Input, out float TangentSign)
 
 #if USE_SPLINEDEFORM
 	// Make slice rotation matrix, and use that to transform tangents
-	half3x3 SliceRot = CalcSliceRot(dot(Input.Position.xyz, SplineMeshDir));
+	half3x3 SliceRot = CalcSliceRot(dot(Input.Position.xyz, SplineMeshDir.xyz));
 
 	TangentX = mul(TangentX, SliceRot);
 	TangentZ.xyz = mul(TangentZ.xyz, SliceRot);
@@ -870,7 +870,7 @@ float4 VertexFactoryGetPreviousWorldPosition(FVertexFactoryInput Input, FVertexF
 	return mul(float4(PreviousPos, 1), PreviousLocalToWorldTranslated);
 #elif USE_SPLINEDEFORM
 	// Just like CalcWorldPosition...
-	float4x3 SliceTransform = CalcSliceTransform(dot(Input.Position.xyz, SplineMeshDir));
+	float4x3 SliceTransform = CalcSliceTransform(dot(Input.Position.xyz, SplineMeshDir.xyz));
 
 	// Transform into mesh space
 	float4 LocalPos = float4(mul(Input.Position, SliceTransform), Input.Position.w);
diff --git a/Engine/Shaders/Private/PostProcessCombineLUTs.usf b/Engine/Shaders/Private/PostProcessCombineLUTs.usf
--- a/Engine/Shaders/Private/PostProcessCombineLUTs.usf
+++ b/Engine/Shaders/Private/PostProcessCombineLUTs.usf
@@ -25,7 +25,7 @@ Texture2D Texture4;
 SamplerState Texture4Sampler;
 // 0 is for neutral, 1 for Texture1, 2 for ...
 float LUTWeights[5];
-half3 ColorScale;
+half4 ColorScale;
 half4 OverlayColor;
 
 static const float LUTSize = 32;
@@ -372,7 +372,7 @@ float4 CombineLUTsCommon(float2 InUV, uint InLayerIndex)
 	FilmColor = ColorCorrection( FilmColor );
 
 	// blend with custom LDR color, used for Fade track in Matinee
-	FilmColor = lerp( FilmColor * ColorScale, OverlayColor.rgb, OverlayColor.a );
+	FilmColor = lerp( FilmColor * ColorScale.xyz, OverlayColor.rgb, OverlayColor.a );
 
 	// Apply "gamma" curve adjustment.
 	FilmColor = pow( max(0, FilmColor), InverseGamma.y );
diff --git a/Engine/Shaders/Private/PostProcessMobile.usf b/Engine/Shaders/Private/PostProcessMobile.usf
--- a/Engine/Shaders/Private/PostProcessMobile.usf
+++ b/Engine/Shaders/Private/PostProcessMobile.usf
@@ -964,7 +964,7 @@ void SunMergeVS_ES2(
 #undef SUN_M
 
 float4 SunColorVignetteIntensity;
-float3 BloomColor;
+float4 BloomColor;
 
 void SunMergePS_ES2(
 	float4 InUVVignette : TEXCOORD0,
@@ -989,7 +989,7 @@ void SunMergePS_ES2(
 
 		OutColor.rgb = HDRDecode(PostprocessInput2.Sample(PostprocessInput2Sampler, InUVVignette.xy));
 
-		OutColor.rgb += Bloom2 * BloomColor;
+		OutColor.rgb += Bloom2 * BloomColor.xyz;
 
 		// Have 5 layers on mobile.
 		OutColor.rgb *= 1.0/5.0;
@@ -1062,7 +1062,7 @@ void SunMergeSmallVS_ES2(
 	OutTexCoords[5] = InTexCoord.xy + Circle(Start, 6.0, 5.0) * Scale * PostprocessInput0Size.zw;
 }
 
-float3 BloomColor2;
+float4 BloomColor2;
 
 void SunMergeSmallPS_ES2(
 	float4 InUVVignette : TEXCOORD0,
@@ -1092,9 +1092,9 @@ void SunMergeSmallPS_ES2(
 		PostprocessInput1.Sample(PostprocessInput1Sampler, InUVs[4].xy).rgb * Scale2 + 
 		PostprocessInput1.Sample(PostprocessInput1Sampler, InUVs[5].xy).rgb * Scale2) * rcp(Scale1 * 1.0 + Scale2 * 6.0);
 
-	OutColor.rgb *= BloomColor2;
+	OutColor.rgb *= BloomColor2.xyz;
 
-	OutColor.rgb += Bloom1 * BloomColor;
+	OutColor.rgb += Bloom1 * BloomColor.xyz;
 
 	// Have 5 layers on mobile (actually only 2 in this shader).
 	OutColor.rgb *= 1.0/5.0;
diff --git a/Engine/Shaders/Private/PostProcessTonemap.usf b/Engine/Shaders/Private/PostProcessTonemap.usf
--- a/Engine/Shaders/Private/PostProcessTonemap.usf
+++ b/Engine/Shaders/Private/PostProcessTonemap.usf
@@ -43,8 +43,8 @@ float4 ColorScale1;
 float2 TonemapperParams;
 
 // Fine film grain
-float3 GrainRandomFull;
-float3 GrainScaleBiasJitter;
+float2 GrainRandomFull;
+float4 GrainScaleBiasJitter;
 
 float4 SceneUVMinMax;
 float4 SceneBloomUVMinMax;
diff --git a/Engine/Shaders/Private/PostProcessUpscale.usf b/Engine/Shaders/Private/PostProcessUpscale.usf
--- a/Engine/Shaders/Private/PostProcessUpscale.usf
+++ b/Engine/Shaders/Private/PostProcessUpscale.usf
@@ -13,7 +13,7 @@ float UpscaleSoftness;
 
 // only for MainVS()
 // .x:HalfFOV .y:tan(HalfFov) .z:CylinderDistortion(0=none, 1:full) .w:Correction to scale Y the same as X is scaled in the center
-float3 PaniniParams;
+float4 PaniniParams;
 
 // Min and max UV coordingate for input scene color.
 float4 SceneColorBilinearUVMinMax;
diff --git a/Engine/Shaders/Private/ReflectionEnvironmentPixelShader.usf b/Engine/Shaders/Private/ReflectionEnvironmentPixelShader.usf
--- a/Engine/Shaders/Private/ReflectionEnvironmentPixelShader.usf
+++ b/Engine/Shaders/Private/ReflectionEnvironmentPixelShader.usf
@@ -47,7 +47,7 @@ float3 GatherRadiance(float CompositeAlpha, float3 WorldPosition, float3 RayDire
 Texture2D ScreenSpaceReflectionsTexture;
 SamplerState ScreenSpaceReflectionsSampler;
 
-float3 ContrastAndNormalizeMulAdd;
+float4 ContrastAndNormalizeMulAdd;
 float OcclusionExponent;
 float OcclusionCombineMode;
 
diff --git a/Engine/Shaders/Private/ReflectionEnvironmentShaders.usf b/Engine/Shaders/Private/ReflectionEnvironmentShaders.usf
--- a/Engine/Shaders/Private/ReflectionEnvironmentShaders.usf
+++ b/Engine/Shaders/Private/ReflectionEnvironmentShaders.usf
@@ -39,7 +39,7 @@ SamplerState InTextureSampler;
  * Y = Sky distance threshold, 
  * Z = whether a skylight's lower hemisphere should be replaced with LowerHemisphereColor.
  */
-float3 SkyLightCaptureParameters;
+float4 SkyLightCaptureParameters;
 float4 LowerHemisphereColor;
 
 void CopySceneColorToCubeFaceColorPS(
diff --git a/Engine/Shaders/Private/ShadowProjectionCommon.ush b/Engine/Shaders/Private/ShadowProjectionCommon.ush
--- a/Engine/Shaders/Private/ShadowProjectionCommon.ush
+++ b/Engine/Shaders/Private/ShadowProjectionCommon.ush
@@ -7,7 +7,7 @@
 Texture2D ShadowDepthTexture;
 SamplerState ShadowDepthTextureSampler;
 // xy:unused, z:SoftTransitionScale
-float3 SoftTransitionScale;
+float4 SoftTransitionScale;
 // xy:ShadowTexelSize.xy, zw:1/ShadowTexelSize.xy
 float4 ShadowBufferSize;
 
diff --git a/Engine/Shaders/Private/SimpleElementVolumeTexturePreviewPixelShader.usf b/Engine/Shaders/Private/SimpleElementVolumeTexturePreviewPixelShader.usf
--- a/Engine/Shaders/Private/SimpleElementVolumeTexturePreviewPixelShader.usf
+++ b/Engine/Shaders/Private/SimpleElementVolumeTexturePreviewPixelShader.usf
@@ -26,7 +26,7 @@ float4 NumTilesPerSide;
 float4 TraceVolumeScaling;
 
 // Texture dimension : width, heigth and depth.
-float3 TextureDimension;
+float4 TextureDimension;
 
 // Transform into view space.
 float4x4 TraceViewMatrix;
@@ -240,7 +240,7 @@ void TraceMain(
 	// Discard when NumSteps = 0
 	clip(NumSteps - .5f);
 
-	const float NumTexelsPerSample = dot(abs(Increment3D), TextureDimension);
+	const float NumTexelsPerSample = dot(abs(Increment3D), TextureDimension.xyz);
 
 	LOOP
 	for (int i = 0; i < NumSteps && Transmittance > .02; i++)
diff --git a/Engine/Shaders/Private/TonemapCommon.ush b/Engine/Shaders/Private/TonemapCommon.ush
--- a/Engine/Shaders/Private/TonemapCommon.ush
+++ b/Engine/Shaders/Private/TonemapCommon.ush
@@ -8,7 +8,7 @@
 #include "GammaCorrectionCommon.ush"
 
 // usually 1/2.2, the .y is used for inverse gamma when "gamma only" mode is not used
-half3 InverseGamma;
+half4 InverseGamma;
 
 // Film tonal and color control.
 half4 ColorMatrixR_ColorCurveCd1;
diff --git a/Engine/Shaders/Private/Tools/VisualizeTexture.usf b/Engine/Shaders/Private/Tools/VisualizeTexture.usf
--- a/Engine/Shaders/Private/Tools/VisualizeTexture.usf
+++ b/Engine/Shaders/Private/Tools/VisualizeTexture.usf
@@ -32,7 +32,7 @@ SamplerState VisualizeTextureCubeArraySampler;
 // [2]: r:0=normal 1=depth, 2=shadow depth, g:unwrapcubemap, b:singlechannel(0-3), a:unused
 float4 VisualizeParam[3];
 // x, y (0 for 2D textures), z(only !=0 if 3d texture), w unused
-float3 TextureExtent;
+float4 TextureExtent;
 
 void VisualizeTexturePS(
 	FScreenVertexOutput Input,
diff --git a/Engine/Shaders/Private/VisualizeVolumetricLightmap.usf b/Engine/Shaders/Private/VisualizeVolumetricLightmap.usf
--- a/Engine/Shaders/Private/VisualizeVolumetricLightmap.usf
+++ b/Engine/Shaders/Private/VisualizeVolumetricLightmap.usf
@@ -78,7 +78,7 @@ void VisualizeVolumetricLightmapVS(
 
 #define WRITE_TO_SHADING_MODEL (FEATURE_LEVEL >= FEATURE_LEVEL_SM4 && !FORWARD_SHADING)
 
-float3 DiffuseColor;
+float4 DiffuseColor;
 
 void VisualizeVolumetricLightmapPS(
 	FVisualizeVolumetricLightmapVSToPS Input,
@@ -108,7 +108,7 @@ void VisualizeVolumetricLightmapPS(
 #endif
 
 	// Use a diffuse color that matches the rest of the scene for LightingOnly and DetailLighting view modes
-	float3 EffectiveDiffuseColor = DiffuseColor * View.DiffuseOverrideParameter.w + View.DiffuseOverrideParameter.xyz;
+	float3 EffectiveDiffuseColor = DiffuseColor.xyz * View.DiffuseOverrideParameter.w + View.DiffuseOverrideParameter.xyz;
 	OutColor = float4(EffectiveDiffuseColor * FinalLighting, 0);
 
 	//OutColor = float4(WorldNormal, 0);
diff --git a/Engine/Shaders/Private/VolumetricFog.usf b/Engine/Shaders/Private/VolumetricFog.usf
--- a/Engine/Shaders/Private/VolumetricFog.usf
+++ b/Engine/Shaders/Private/VolumetricFog.usf
@@ -62,8 +62,8 @@ float3 ScatteringFunction()
 	//return RaleighScattering();
 }
 
-float3 GlobalAlbedo;
-float3 GlobalEmissive;
+float4 GlobalAlbedo;
+float4 GlobalEmissive;
 float GlobalExtinctionScale;
 
 #ifndef THREADGROUP_SIZE 
@@ -85,7 +85,7 @@ void MaterialSetupCS(
 
 	float GlobalDensity = FogStruct.ExponentialFogParameters3.x * exp2(-FogStruct.ExponentialFogParameters.y * (WorldPosition.z - FogStruct.ExponentialFogParameters3.y));
 
-	float3 Albedo = GlobalAlbedo;
+	float3 Albedo = GlobalAlbedo.xyz;
 
 	// Exponential height fog interprets density differently, match its behavior
 	float MatchHeightFogFactor = .5f;
@@ -97,7 +97,7 @@ void MaterialSetupCS(
 	if (all((int3)GridCoordinate < VolumetricFog.GridSizeInt))
 	{
 		RWVBufferA[GridCoordinate] = float4(Scattering, Absorption);
-		RWVBufferB[GridCoordinate] = float4(GlobalEmissive, 0);
+		RWVBufferB[GridCoordinate] = float4(GlobalEmissive.xyz, 0);
 	}
 }
 
diff --git a/Engine/Shaders/Private/VolumetricFogLightFunction.usf b/Engine/Shaders/Private/VolumetricFogLightFunction.usf
--- a/Engine/Shaders/Private/VolumetricFogLightFunction.usf
+++ b/Engine/Shaders/Private/VolumetricFogLightFunction.usf
@@ -7,7 +7,7 @@
 #include "LightFunctionCommon.ush"
 
 /** Fade distance in x, disabled brightness in y, output for preview shadows mask in z. */
-float3 LightFunctionParameters2;
+float4 LightFunctionParameters2;
 float2 LightFunctionTexelSize;
 
 float4x4 ShadowToWorld;
diff --git a/Engine/Shaders/Public/ShaderVersion.ush b/Engine/Shaders/Public/ShaderVersion.ush
--- a/Engine/Shaders/Public/ShaderVersion.ush
+++ b/Engine/Shaders/Public/ShaderVersion.ush
@@ -2,4 +2,4 @@
 // Each time the console command is executed it generates a new GUID. As this file is included
 // in Platform.ush (which should be included in any shader) it allows to invalidate the shader DDC.
 // 
-// GUID = 373DFDBD095D4D8E91B55377386723F3
\ No newline at end of file
+// GUID = 339F87C6-80DF-4A3F-8089-8AA5E9A73864
