DR - Missed float3->float4
#rnx

diff --git a/Engine/Plugins/Compositing/LensDistortion/Shaders/Private/UVGeneration.usf b/Engine/Plugins/Compositing/LensDistortion/Shaders/Private/UVGeneration.usf
--- a/Engine/Plugins/Compositing/LensDistortion/Shaders/Private/UVGeneration.usf
+++ b/Engine/Plugins/Compositing/LensDistortion/Shaders/Private/UVGeneration.usf
@@ -24,7 +24,7 @@
 float2 PixelUVSize;
 
 // K1, K2, K3
-float3 RadialDistortionCoefs;
+float4 RadialDistortionCoefs;
 
 // P1, P2
 float2 TangentialDistortionCoefs;
diff --git a/Engine/Shaders/Private/DeferredDecal.usf b/Engine/Shaders/Private/DeferredDecal.usf
--- a/Engine/Shaders/Private/DeferredDecal.usf
+++ b/Engine/Shaders/Private/DeferredDecal.usf
@@ -17,7 +17,7 @@ float4x4 DecalToWorld;
 // to transform position from world space to decal space, for normals use transpose of inverse
 float4x4 WorldToDecal;
 // decal orientation in world space, X axis. Used for ObjectOrientation material expression
-float3 DecalOrientation;
+float4 DecalOrientation;
 
 // DECAL_PRIMITIVE informs material templates which functions to expose when rendering decals.
 #define DECAL_PRIMITIVE 1
