Handle NaNs and INFs properly

diff --git a/Engine/Shaders/Private/Common.ush b/Engine/Shaders/Private/Common.ush
--- a/Engine/Shaders/Private/Common.ush
+++ b/Engine/Shaders/Private/Common.ush
@@ -1569,8 +1569,8 @@ void DrawRectangle( in float4 InPosition, out float4 OutPosition)
 //Since some platforms don't remove Nans in saturate calls, 
 //SafeSaturate function will remove nan/inf.    
 //Can be expensive, only call when there's a good reason to expect Nans.
-//D3D saturate actually turns Nans -> 1  since it does the min(x, 1) first, and D3D nan rules specify the non-nand operand wins in such a case.  
-//See: http://msdn.microsoft.com/en-us/library/windows/desktop/jj218760(v=vs.85).aspx  
+//D3D saturate actually turns NaNs -> 0  since it does the max(0.0f, value) first, and D3D NaN rules specify the non-NaN operand wins in such a case.  
+//See: https://docs.microsoft.com/en-us/windows/desktop/direct3dhlsl/saturate
 #define SafeSaturate_Def(type)\
 type SafeSaturate(type In) \
 {\
@@ -1582,6 +1582,35 @@ SafeSaturate_Def(float2)
 SafeSaturate_Def(float3)
 SafeSaturate_Def(float4)
 
+// By default HLSL compilers make assumptions about the non-NaN nature of most inputs. Otherwise, most of its optimizations would be invalid, e.g. assuming multiply by 0 equals 0.
+// Unfortunately, as a consequence of this, the native infinite()/isnan()/isinf() intrinsics frequently get silently optimized away. To avoid this hazard we manually implement the same logic.
+// These are expected to generate identical code as native intrinsics would (when working properly). The PositiveFinite versions are faster than the native.
+// Also note that FXC has a quirk in which the pattern b? a : 0 is implemented with an AND instead of the more natural MOVC, i.e. optimizing for D3D bytecode instead of real-world HW instructions.
+// So drivers detect this pattern and convert back to a conditional move. HLSLcc instead converts to mix(), so NaNs end up propagating through despite our efforts to kill them (see UE-66179).
+// So we break this pattern using !b? 0 : a instead, which generates a MOVC.
+#define IsAndMakeFinite_Def(type, booltype)\
+booltype IsFinite(type In) \
+{\
+	return (asuint(In) & 0x7F800000) != 0x7F800000; \
+}\
+booltype IsPositiveFinite(type In) \
+{\
+	return asuint(In) < 0x7F800000; \
+}\
+type MakeFinite(type In) \
+{\
+    return !IsFinite(In)? 0 : In; \
+}\
+type MakePositiveFinite(type In) \
+{\
+    return !IsPositiveFinite(In)? 0 : In; \
+}
+
+IsAndMakeFinite_Def(float,  bool)
+IsAndMakeFinite_Def(float2, bool2)
+IsAndMakeFinite_Def(float3, bool3)
+IsAndMakeFinite_Def(float4, bool4)
+
 // Experimental way to allow adjusting the OpacityMask for shadow map rendering of masked materials.
 // Can be accessed with a Custom material node. If this turns out to be very useful we can expose as MaterialFunction
 // and potentially expose other queries as well (e.g. SkeletalMesh, HitProxy, ).
diff --git a/Engine/Shaders/Private/DistanceFieldLightingPost.usf b/Engine/Shaders/Private/DistanceFieldLightingPost.usf
--- a/Engine/Shaders/Private/DistanceFieldLightingPost.usf
+++ b/Engine/Shaders/Private/DistanceFieldLightingPost.usf
@@ -281,14 +281,7 @@ void UpdateHistoryDepthRejectionPS(
 	OutBentNormal.rgb = lerp(NewValue.rgb, HistoryValue.rgb, EffectiveHistoryWeight);
 #endif
 
-#if XBOXONE_PROFILE
-	// Xbox shader compiler (FXC) does non-IEEE optimization by default unless __XBOX_CONTROL_NONIEEE is
-	// defined to 0, which is not done currently due to performance reason. As a result, we need to check
-	// for Nan manually on Xbox.
-	OutBentNormal.rgb = (asuint(OutBentNormal.rgb) & 0x7f800000) == 0x7f800000 ? 0 : OutBentNormal.rgb;
-#else
-	OutBentNormal.rgb = isnan(OutBentNormal.rgb) || isinf(OutBentNormal.rgb) ? 0 : OutBentNormal.rgb;
-#endif
+	OutBentNormal.rgb = MakeFinite(OutBentNormal.rgb);
 
 	OutBentNormal.a = SceneDepth;
 
diff --git a/Engine/Shaders/Private/LPVClear.usf b/Engine/Shaders/Private/LPVClear.usf
--- a/Engine/Shaders/Private/LPVClear.usf
+++ b/Engine/Shaders/Private/LPVClear.usf
@@ -47,15 +47,15 @@ void CSClear(uint3 DTid : SV_DispatchThreadID)
 
 	// corruption detection (workaround for brightlodge lighting corruption)
 	{
-		if ( isnan( cell.AO ) || isinf( cell.AO ) )
+		if ( !IsFinite(cell.AO ) )
 		{
 			mul = 0.0f;
 		}
 		[unroll]
 		for ( int i=0; i<9; i++ )
 		{
-			float df = dot( cell.coeffs[i], float3(1.0f,1.0f,1.0f) );
-			if ( isnan( df ) || isinf( df ) )
+			float df = cell.coeffs[i].x + cell.coeffs[i].y + cell.coeffs[i].z;
+			if ( !IsFinite(df) )
 			{
 				mul = 0.0f;
 			}
diff --git a/Engine/Shaders/Private/VolumetricFog.usf b/Engine/Shaders/Private/VolumetricFog.usf
--- a/Engine/Shaders/Private/VolumetricFog.usf
+++ b/Engine/Shaders/Private/VolumetricFog.usf
@@ -643,7 +643,7 @@ void LightScatteringCS(
 
 	if (all(GridCoordinate < VolumetricFog.GridSizeInt))
 	{
-		ScatteringAndExtinction = -min(-ScatteringAndExtinction, float4(0.f, 0.f, 0.f, 0.f));	// vector isnan/isinf are miscompiled by hlslcc (see UE-66179)
+		ScatteringAndExtinction = MakePositiveFinite(ScatteringAndExtinction);
 		RWLightScattering[GridCoordinate] = ScatteringAndExtinction;
 	}
 }
