Renaming Raytracing to RayTracing to be consistent. To overcome case insensitivity, this will be done in two submissions.

diff --git a/Engine/Shaders/Private/Raytracing/BuildMipTreeCS.usf b/Engine/Shaders/Private/Raytracing/BuildMipTreeCS.usf
deleted file mode 100644
diff --git a/Engine/Shaders/Private/Raytracing/BuildMipTreeCS_mv.usf b/Engine/Shaders/Private/Raytracing/BuildMipTreeCS_mv.usf
new file mode 100644
--- /dev/null
+++ b/Engine/Shaders/Private/Raytracing/BuildMipTreeCS_mv.usf
@@ -0,0 +1,95 @@
+// Copyright 1998-2018 Epic Games, Inc. All Rights Reserved.
+
+/*=============================================================================================
+	BuildMipTreeCS.usf: Build variance map
+===============================================================================================*/
+
+#include "../Common.ush"
+#include "SkyLightMipTreeCommon.ush"
+
+Texture2D RectLightTexture;
+TextureCube Texture;
+SamplerState TextureSampler;
+
+uint2 Dimensions;
+uint FaceIndex;
+uint MipLevel;
+RWBuffer<float> RWMipTree;
+
+[numthreads(THREADGROUP_SIZE, THREADGROUP_SIZE, 1)]
+void BuildMipTreeCS(
+	uint3 DispatchThreadId : SV_DispatchThreadID
+)
+{
+	uint3 GridCoordinate = DispatchThreadId;
+	uint2 Bounds = DimensionsAtLevel(Dimensions, MipLevel);
+	if (GridCoordinate.x < 0 || GridCoordinate.x >= Bounds.x || GridCoordinate.y < 0 || GridCoordinate.y >= Bounds.y)
+	{
+		return;
+	}
+
+	float Sum = 0.0;
+	if (MipLevel == 0)
+	{
+		// #dxr_todo: TextureCubeSampleLevel() may be enough.. could do live during query??
+		float3 Direction = GetTextureCubeVector(uint3(GridCoordinate.xy, FaceIndex), Dimensions);
+		float3 Value = TextureCubeSample(Texture, TextureSampler, Direction).rgb;
+		Sum += Luminance(Value.rgb);
+	}
+	else
+	{
+		// Report the mean as the aggregate.
+		uint2 PixelOffset = uint2(GridCoordinate.x << 1, GridCoordinate.y << 1);
+		for (int i = 0; i < 4; ++i)
+		{
+			uint2 Pixel = PixelOffset + uint2(i >> 1, i & 1);
+			uint InputIndex = BufferOffsetAtPixel(Pixel, MipLevel - 1, Dimensions);
+			// #dxr_todo: Experiment between average and max
+			Sum += RWMipTree[InputIndex];
+		}
+		Sum /= 4;
+	}
+	uint OutputIndex = BufferOffsetAtPixel(GridCoordinate.xy, MipLevel, Dimensions);
+	RWMipTree[OutputIndex] = Sum;
+}
+
+[numthreads(THREADGROUP_SIZE, THREADGROUP_SIZE, 1)]
+void BuildRectLightMipTreeCS(
+	uint3 DispatchThreadId : SV_DispatchThreadID
+)
+{
+	uint3 GridCoordinate = DispatchThreadId;
+	uint2 Bounds = DimensionsAtLevel(Dimensions, MipLevel);
+	if (GridCoordinate.x < 0 || GridCoordinate.x >= Bounds.x || GridCoordinate.y < 0 || GridCoordinate.y >= Bounds.y)
+	{
+		return;
+	}
+
+	float Sum = 0.0;
+	if (MipLevel == 0)
+	{
+		// #dxr_todo: TextureCubeSampleLevel() may be enough.. could do live during query??
+		uint2 TextureDimensions;
+		RectLightTexture.GetDimensions(TextureDimensions.x, TextureDimensions.y);
+		float2 UV = GridCoordinate.xy / float2(TextureDimensions);
+		UV = 1.0 - UV;
+		float3 Value = Texture2DSample(RectLightTexture, TextureSampler, UV).rgb;
+		Sum += Luminance(Value.rgb);
+	}
+	else
+	{
+		// Report the mean as the aggregate.
+		uint2 PixelOffset = uint2(GridCoordinate.x << 1, GridCoordinate.y << 1);
+		for (int i = 0; i < 4; ++i)
+		{
+			uint2 Pixel = PixelOffset + uint2(i >> 1, i & 1);
+			uint InputIndex = BufferOffsetAtPixel(Pixel, MipLevel - 1, Dimensions);
+			// #dxr_todo: Experiment between average and max
+			//Sum = max(Sum, RWMipTree[InputIndex]);
+			Sum += RWMipTree[InputIndex];
+		}
+		Sum /= 4;
+	}
+	uint OutputIndex = BufferOffsetAtPixel(GridCoordinate.xy, MipLevel, Dimensions);
+	RWMipTree[OutputIndex] = Sum;
+}
\ No newline at end of file
diff --git a/Engine/Shaders/Private/Raytracing/BuildMipTreePdfCS.usf b/Engine/Shaders/Private/Raytracing/BuildMipTreePdfCS.usf
deleted file mode 100644
diff --git a/Engine/Shaders/Private/Raytracing/BuildMipTreePdfCS_mv.usf b/Engine/Shaders/Private/Raytracing/BuildMipTreePdfCS_mv.usf
new file mode 100644
--- /dev/null
+++ b/Engine/Shaders/Private/Raytracing/BuildMipTreePdfCS_mv.usf
@@ -0,0 +1,87 @@
+// Copyright 1998-2018 Epic Games, Inc. All Rights Reserved.
+
+/*=============================================================================================
+BuildMipTreeCS.usf: Build MipTree PDF structure
+===============================================================================================*/
+
+#include "../Common.ush"
+#include "SkyLightMipTreeCommon.ush"
+
+uint2 Dimensions;
+uint FaceIndex;
+uint MipLevel;
+
+Buffer<float> MipTree;
+RWBuffer<float> RWMipTreePdf;
+
+float SampleMipTree(uint2 MipPixel, uint MipIndex)
+{
+	return MipTree[BufferOffsetAtPixel(MipPixel, MipIndex, Dimensions)];
+}
+
+float CalcSkyLightPdf(uint2 Pixel)
+{
+	uint2 MipPixel = 0;
+	float MipPdf = 1.0;
+	uint StopLevel = MipLevel;
+	uint MipCount = Log2(Dimensions.x);
+	for (uint Index = 0; Index < MipCount - StopLevel; ++Index)
+	{
+		uint MipIndex = MipCount - Index - 1;
+
+		// #dxr_todo: Use float4?
+		float Cdf[4];
+		Cdf[0] = SampleMipTree(MipPixel, MipIndex);
+		Cdf[1] = Cdf[0] + SampleMipTree(MipPixel + uint2(1, 0), MipIndex);
+		Cdf[2] = Cdf[1] + SampleMipTree(MipPixel + uint2(0, 1), MipIndex);
+		Cdf[3] = Cdf[2] + SampleMipTree(MipPixel + uint2(1, 1), MipIndex);
+
+		float Sum = Cdf[3];
+		if (Sum > 0.0)
+		{
+			for (uint i = 0; i < 3; ++i)
+			{
+				Cdf[i] /= Sum;
+			}
+			Cdf[3] = 1.0;
+		}
+
+		// Determine probability
+		uint2 MipOffset = (Pixel.xy >> (MipIndex - StopLevel)) - MipPixel;
+		uint LinearIndex = MipOffset.y << 1 | MipOffset.x;
+		MipPdf *= LinearIndex > 0 ? Cdf[LinearIndex] - Cdf[LinearIndex - 1] : Cdf[LinearIndex];
+#if 0
+		float LevelPdf = Cdf[LinearIndex];
+		if (LinearIndex > 0)
+			LevelPdf -= Cdf[LinearIndex - 1];
+		MipPdf *= LevelPdf;
+#endif
+		MipPixel += MipOffset;
+
+		MipPixel = MipPixel << 1;
+	}
+	//MipPdf = max(MipPdf, 0.0);
+	return MipPdf;
+
+	MipPixel = MipPixel >> 1;
+	float AreaPdf = 1.0 / TexelCoordSolidAngle(MipPixel, DimensionsAtLevel(Dimensions, StopLevel));
+	return MipPdf * AreaPdf;
+}
+
+[numthreads(THREADGROUP_SIZE, THREADGROUP_SIZE, 1)]
+void BuildMipTreePdfCS(
+	uint3 DispatchThreadId : SV_DispatchThreadID
+)
+{
+	uint3 GridCoordinate = DispatchThreadId;
+	uint2 Bounds = DimensionsAtLevel(Dimensions, MipLevel);
+	if (GridCoordinate.x < 0 || GridCoordinate.x >= Bounds.x || GridCoordinate.y < 0 || GridCoordinate.y >= Bounds.y)
+	{
+		return;
+	}
+
+	float Pdf = CalcSkyLightPdf(GridCoordinate.xy);
+	uint BufferOffset = BufferOffsetAtPixel(GridCoordinate.xy, MipLevel, Dimensions);
+	RWMipTreePdf[BufferOffset] = Pdf;
+	//RWMipTreePdf[BufferOffset] = SampleMipTree(GridCoordinate.xy, MipLevel);
+}
\ No newline at end of file
diff --git a/Engine/Shaders/Private/Raytracing/BuildSolidAnglePdfCS.usf b/Engine/Shaders/Private/Raytracing/BuildSolidAnglePdfCS.usf
deleted file mode 100644
diff --git a/Engine/Shaders/Private/Raytracing/BuildSolidAnglePdfCS_mv.usf b/Engine/Shaders/Private/Raytracing/BuildSolidAnglePdfCS_mv.usf
new file mode 100644
--- /dev/null
+++ b/Engine/Shaders/Private/Raytracing/BuildSolidAnglePdfCS_mv.usf
@@ -0,0 +1,26 @@
+// Copyright 1998-2018 Epic Games, Inc. All Rights Reserved.
+
+/*=============================================================================================
+BuildMipTreeCS.usf: Build MipTree PDF structure
+===============================================================================================*/
+
+#include "SkyLightMipTreeCommon.ush"
+
+uint MipLevel;
+uint2 Dimensions;
+RWBuffer<float> RWSolidAnglePdf;
+
+[numthreads(THREADGROUP_SIZE, THREADGROUP_SIZE, 1)]
+void BuildSolidAnglePdfCS(
+	uint3 DispatchThreadId : SV_DispatchThreadID
+)
+{
+	uint2 Pixel = DispatchThreadId.xy;
+	if (Pixel.x >= Dimensions.x || Pixel.y >= Dimensions.y)
+	{
+		return;
+	}
+
+	uint LinearIndex = BufferOffsetAtPixel(Pixel, MipLevel, Dimensions);
+	RWSolidAnglePdf[LinearIndex] = 1.0 / TexelCoordSolidAngle(Pixel, Dimensions);
+}
diff --git a/Engine/Shaders/Private/Raytracing/CompositeSkyLightPS.usf b/Engine/Shaders/Private/Raytracing/CompositeSkyLightPS.usf
deleted file mode 100644
diff --git a/Engine/Shaders/Private/Raytracing/CompositeSkyLightPS_mv.usf b/Engine/Shaders/Private/Raytracing/CompositeSkyLightPS_mv.usf
new file mode 100644
--- /dev/null
+++ b/Engine/Shaders/Private/Raytracing/CompositeSkyLightPS_mv.usf
@@ -0,0 +1,19 @@
+// Copyright 1998-2018 Epic Games, Inc. All Rights Reserved.
+
+#include "../Common.ush"
+#include "../DeferredShadingCommon.ush"
+
+Texture2D<float4> SkyLightTexture;
+Texture2D<float> HitDistanceTexture;
+
+SamplerState SkyLightTextureSampler;
+SamplerState HitDistanceTextureSampler;
+
+void CompositeSkyLightPS(
+	in noperspective float2 UV : TEXCOORD0,
+	out float4 OutColor : SV_Target0
+)
+{
+	float4 SkyLight = SkyLightTexture.Sample(SkyLightTextureSampler, UV);
+	OutColor = float4(SkyLight.rgb, 1.0);
+}
diff --git a/Engine/Shaders/Private/Raytracing/MipTreeCommon.ush b/Engine/Shaders/Private/Raytracing/MipTreeCommon.ush
deleted file mode 100644
diff --git a/Engine/Shaders/Private/Raytracing/MipTreeCommon_mv.ush b/Engine/Shaders/Private/Raytracing/MipTreeCommon_mv.ush
new file mode 100644
--- /dev/null
+++ b/Engine/Shaders/Private/Raytracing/MipTreeCommon_mv.ush
@@ -0,0 +1,76 @@
+// Copyright 1998-2018 Epic Games, Inc. All Rights Reserved.
+
+/*=============================================================================================
+	MipTreeCommon.ush: Common utilities for MipTree-based sampling
+===============================================================================================*/
+
+#define USE_MORTON_CODE 1
+
+uint2 DimensionsAtLevel(uint2 Dimensions, uint Level)
+{
+	return uint2(Dimensions.x >> Level, Dimensions.y >> Level);
+}
+
+// #dxr_todo: Discussing with Brian where these utility functions should live.
+uint MortonEncode(uint2 Pixel)
+{
+	Pixel.x &= 0x0000ffff;
+	Pixel.x = (Pixel.x ^ (Pixel.x << 8)) & 0x00ff00ff;
+	Pixel.x = (Pixel.x ^ (Pixel.x << 4)) & 0x0f0f0f0f;
+	Pixel.x = (Pixel.x ^ (Pixel.x << 2)) & 0x33333333;
+	Pixel.x = (Pixel.x ^ (Pixel.x << 1)) & 0x55555555;
+
+	Pixel.y &= 0x0000ffff;
+	Pixel.y = (Pixel.y ^ (Pixel.y << 8)) & 0x00ff00ff;
+	Pixel.y = (Pixel.y ^ (Pixel.y << 4)) & 0x0f0f0f0f;
+	Pixel.y = (Pixel.y ^ (Pixel.y << 2)) & 0x33333333;
+	Pixel.y = (Pixel.y ^ (Pixel.y << 1)) & 0x55555555;
+	return Pixel.x | (Pixel.y << 1);
+}
+
+uint2 MortonDecode(uint Morton)
+{
+	uint2 Pixel = uint2(Morton, Morton >> 1);
+	Pixel.x &= 0x55555555;
+	Pixel.x = (Pixel.x ^ (Pixel.x >> 1)) & 0x33333333;
+	Pixel.x = (Pixel.x ^ (Pixel.x >> 2)) & 0x0f0f0f0f;
+	Pixel.x = (Pixel.x ^ (Pixel.x >> 4)) & 0x00ff00ff;
+	Pixel.x = (Pixel.x ^ (Pixel.x >> 8)) & 0x0000ffff;
+
+	Pixel.y &= 0x55555555;
+	Pixel.y = (Pixel.y ^ (Pixel.y >> 1)) & 0x33333333;
+	Pixel.y = (Pixel.y ^ (Pixel.y >> 2)) & 0x0f0f0f0f;
+	Pixel.y = (Pixel.y ^ (Pixel.y >> 4)) & 0x00ff00ff;
+	Pixel.y = (Pixel.y ^ (Pixel.y >> 8)) & 0x0000ffff;
+	return Pixel;
+}
+
+uint BufferOffsetAtPixel(uint2 Pixel, uint Level, uint2 Dimensions)
+{
+	uint Offset = 0;
+	for (uint LevelIndex = 0; LevelIndex < Level; ++LevelIndex)
+	{
+		uint2 MipLevelDimensions = DimensionsAtLevel(Dimensions, LevelIndex);
+		Offset += MipLevelDimensions.x * MipLevelDimensions.y;
+	}
+
+	uint2 MipLevelDimensions = DimensionsAtLevel(Dimensions, Level);
+#if USE_MORTON_CODE
+	// NOTE: Morton encoding requires power-of-two dimensionality..
+	Offset += MortonEncode(Pixel);
+#else
+	Offset += Pixel.y * MipLevelDimensions.x + Pixel.x;
+#endif
+	return Offset;
+}
+
+uint Log2(uint Value)
+{
+	uint Result = 0;
+	while (Value)
+	{
+		Result++;
+		Value = Value >> 1;
+	}
+	return Result - 1;
+}
diff --git a/Engine/Shaders/Private/Raytracing/RayTracingDebugDefinitions_mv.ush b/Engine/Shaders/Private/Raytracing/RayTracingDebugDefinitions_mv.ush
new file mode 100644
--- /dev/null
+++ b/Engine/Shaders/Private/Raytracing/RayTracingDebugDefinitions_mv.ush
@@ -0,0 +1,29 @@
+// Copyright 1998-2018 Epic Games, Inc. All Rights Reserved.
+
+/*================================================================================================
+	RayTracingDebugDefinitions.ush: used in ray tracing debug shaders  and C++ code to define common constants
+	!!! Changing this file requires recompilation of the engine !!!
+=================================================================================================*/
+
+#pragma once
+
+#define RAY_TRACING_DEBUG_VIZ_RADIANCE					0 
+#define RAY_TRACING_DEBUG_VIZ_WORLD_NORMAL				1 
+#define RAY_TRACING_DEBUG_VIZ_BASE_COLOR				2 
+#define RAY_TRACING_DEBUG_VIZ_DIFFUSE_COLOR				3 
+#define RAY_TRACING_DEBUG_VIZ_SPECULAR_COLOR			4 
+#define RAY_TRACING_DEBUG_VIZ_OPACITY					5 
+#define RAY_TRACING_DEBUG_VIZ_METALLIC					6 
+#define RAY_TRACING_DEBUG_VIZ_SPECULAR					7 
+#define RAY_TRACING_DEBUG_VIZ_ROUGHNESS					8 
+#define RAY_TRACING_DEBUG_VIZ_IOR						9 
+#define RAY_TRACING_DEBUG_VIZ_SHADING_MODEL				10 
+#define RAY_TRACING_DEBUG_VIZ_BLENDING_MODE				11  
+#define RAY_TRACING_DEBUG_VIZ_LIGHTING_CHANNEL_MASK		12  
+#define RAY_TRACING_DEBUG_VIZ_CUSTOM_DATA				13  
+#define RAY_TRACING_DEBUG_VIZ_GBUFFER_AO				14  
+#define RAY_TRACING_DEBUG_VIZ_INDIRECT_IRRADIANCE		15  
+#define RAY_TRACING_DEBUG_VIZ_WORLD_POSITION			16  
+#define RAY_TRACING_DEBUG_VIZ_HITKIND					17  
+#define RAY_TRACING_DEBUG_VIZ_BARYCENTRICS				18  
+
diff --git a/Engine/Shaders/Private/Raytracing/RayTracingDynamicMesh.usf b/Engine/Shaders/Private/Raytracing/RayTracingDynamicMesh.usf
deleted file mode 100644
diff --git a/Engine/Shaders/Private/Raytracing/RayTracingDynamicMesh_mv.usf b/Engine/Shaders/Private/Raytracing/RayTracingDynamicMesh_mv.usf
new file mode 100644
--- /dev/null
+++ b/Engine/Shaders/Private/Raytracing/RayTracingDynamicMesh_mv.usf
@@ -0,0 +1,67 @@
+// Copyright 1998-2018 Epic Games, Inc. All Rights Reserved.
+
+/*=============================================================================
+	RayTracingDynamicMesh.usf
+=============================================================================*/
+
+#include "/Engine/Private/Common.ush"
+#include "/Engine/Generated/Material.ush"
+#include "/Engine/Generated/VertexFactory.ush"
+ 
+struct FRayTracingDynamicMeshVSToGS
+{
+	float4 WorldPosition : SV_POSITION;
+};
+
+void RayTracingDynamicMeshVS(
+	FVertexFactoryInput Input,
+	out FRayTracingDynamicMeshVSToGS Output
+	)
+{
+	ResolvedView = ResolveView();
+
+	FVertexFactoryIntermediates VFIntermediates = GetVertexFactoryIntermediates(Input);
+	Output.WorldPosition = VertexFactoryGetWorldPosition(Input, VFIntermediates) - float4(ResolvedView.PreViewTranslation, 0.0f);
+}
+
+// Must match GetUniformMeshStreamOutLayout and TRIANGLE_VERTEX_DATA_STRIDE
+struct FRayTracingDynamicMeshVertex
+{
+	float4 WorldPosition : SV_POSITION;
+};
+
+[maxvertexcount(3)]
+void RayTracingDynamicMeshGS(triangle FRayTracingDynamicMeshVSToGS Inputs[3], inout TriangleStream<FRayTracingDynamicMeshVertex> OutStream)
+{ 
+	for (int i = 0; i < 3; i++)
+	{
+		FRayTracingDynamicMeshVSToGS Input = Inputs[i];
+
+		FRayTracingDynamicMeshVertex Vertex;
+		Vertex.WorldPosition = Input.WorldPosition;
+
+		OutStream.Append(Vertex);
+	}
+}
+
+#if COMPUTESHADER
+uint NumVertices;
+RWBuffer<float> RWVertexPositions;
+
+[numthreads(256, 1, 1)]
+void RayTracingDynamicGeometryConverterCS(uint3 DispatchThreadId : SV_DispatchThreadID)
+{ 
+	ResolvedView = ResolveView();
+
+	uint VertexIndex = DispatchThreadId.x;
+	
+	if (VertexIndex >= NumVertices) return;
+
+	FVertexFactoryInput Input = LoadVertexFactoryInput(VertexIndex / 3, VertexIndex % 3);
+	FVertexFactoryIntermediates VFIntermediates = GetVertexFactoryIntermediates(Input);
+	float3 WorldPosition = (VertexFactoryGetWorldPosition(Input, VFIntermediates) - float4(ResolvedView.PreViewTranslation, 0.0f)).xyz;
+	RWVertexPositions[VertexIndex * 3 + 0] = WorldPosition.x;
+	RWVertexPositions[VertexIndex * 3 + 1] = WorldPosition.y;
+	RWVertexPositions[VertexIndex * 3 + 2] = WorldPosition.z;
+}
+#endif
diff --git a/Engine/Shaders/Private/Raytracing/RayTracingSkyLightRG_mv.usf b/Engine/Shaders/Private/Raytracing/RayTracingSkyLightRG_mv.usf
new file mode 100644
--- /dev/null
+++ b/Engine/Shaders/Private/Raytracing/RayTracingSkyLightRG_mv.usf
@@ -0,0 +1,477 @@
+// Copyright 1998-2018 Epic Games, Inc. All Rights Reserved.
+
+#include "../Common.ush"
+#include "../MonteCarlo.ush"
+#include "../DeferredShadingCommon.ush"
+#include "SkyLightMipTreeCommon.ush"
+#include "RayTracingCommon.ush"
+#include "RayTracingDirectionalLight.ush"
+#include "RayTracingRectLight.ush"
+
+RaytracingAccelerationStructure TLAS;
+RWTexture2D<float4> RWOcclusionMaskUAV;
+RWTexture2D<float> RWRayDistanceUAV;
+
+#define USE_MIP_PDF_TEXTURE 1
+#define USE_SOLID_ANGLE_TEXTURE 1
+
+uint CalcLinearIndex(uint2 PixelCoord)
+{
+	return PixelCoord.y * View.BufferSizeAndInvSize.x + PixelCoord.x;
+}
+
+uint SelectFace(uint MipCount, inout RandomSequence RandSequence, out float FacePdf)
+{
+	uint BufferOffset = BufferOffsetAtPixel(uint2(0, 0), MipCount, SkyLight.MipDimensions);
+
+	float FaceCdf[6];
+	FaceCdf[0] = SkyLight.MipTreePosX[BufferOffset];
+	FaceCdf[1] = FaceCdf[0] + SkyLight.MipTreeNegX[BufferOffset];
+	FaceCdf[2] = FaceCdf[1] + SkyLight.MipTreePosY[BufferOffset];
+	FaceCdf[3] = FaceCdf[2] + SkyLight.MipTreeNegY[BufferOffset];
+	FaceCdf[4] = FaceCdf[3] + SkyLight.MipTreePosZ[BufferOffset];
+	FaceCdf[5] = FaceCdf[4] + SkyLight.MipTreeNegZ[BufferOffset];
+	float Sum = FaceCdf[5];
+
+	for (uint Index = 0; Index < 5; ++Index)
+	{
+		FaceCdf[Index] /= Sum;
+	}
+	FaceCdf[5] = 1.0;
+
+	// Draw random variable
+	uint DummyVariable;
+	float RandSample = RandomSequence_GenerateSample1D(RandSequence, DummyVariable);
+
+	// Determine CDF entry
+	uint FaceIndex = 0;
+	while (RandSample > FaceCdf[FaceIndex])
+	{
+		FaceIndex++;
+	}
+
+	// Calculate PDF and return entry
+	FacePdf = FaceCdf[FaceIndex];
+	if (FaceIndex > 0)
+		FacePdf -= FaceCdf[FaceIndex - 1];
+	return FaceIndex;
+}
+
+float4 SampleMipTree(uint FaceIndex, uint4 BufferOffset)
+{
+	float4 Values;
+	if (FaceIndex == 0)
+	{
+		Values.x = SkyLight.MipTreePosX[BufferOffset.x];
+		Values.y = SkyLight.MipTreePosX[BufferOffset.y];
+		Values.z = SkyLight.MipTreePosX[BufferOffset.z];
+		Values.w = SkyLight.MipTreePosX[BufferOffset.w];
+	}
+	else if (FaceIndex == 1)
+	{
+		Values.x = SkyLight.MipTreeNegX[BufferOffset.x];
+		Values.y = SkyLight.MipTreeNegX[BufferOffset.y];
+		Values.z = SkyLight.MipTreeNegX[BufferOffset.z];
+		Values.w = SkyLight.MipTreeNegX[BufferOffset.w];
+	}
+	else if (FaceIndex == 2)
+	{
+		Values.x = SkyLight.MipTreePosY[BufferOffset.x];
+		Values.y = SkyLight.MipTreePosY[BufferOffset.y];
+		Values.z = SkyLight.MipTreePosY[BufferOffset.z];
+		Values.w = SkyLight.MipTreePosY[BufferOffset.w];
+	}
+	else if (FaceIndex == 3)
+	{
+		Values.x = SkyLight.MipTreeNegY[BufferOffset.x];
+		Values.y = SkyLight.MipTreeNegY[BufferOffset.y];
+		Values.z = SkyLight.MipTreeNegY[BufferOffset.z];
+		Values.w = SkyLight.MipTreeNegY[BufferOffset.w];
+	}
+	else if (FaceIndex == 4)
+	{
+		Values.x = SkyLight.MipTreePosZ[BufferOffset.x];
+		Values.y = SkyLight.MipTreePosZ[BufferOffset.y];
+		Values.z = SkyLight.MipTreePosZ[BufferOffset.z];
+		Values.w = SkyLight.MipTreePosZ[BufferOffset.w];
+	}
+	else
+	{
+		Values.x = SkyLight.MipTreeNegZ[BufferOffset.x];
+		Values.y = SkyLight.MipTreeNegZ[BufferOffset.y];
+		Values.z = SkyLight.MipTreeNegZ[BufferOffset.z];
+		Values.w = SkyLight.MipTreeNegZ[BufferOffset.w];
+	}
+
+	return Values;
+}
+
+// #dxr_todo: Copied from PathTracingSkyLight.ush
+float3 EvalSkyLight(float3 Direction)
+{
+	return SkyLight.Color * TextureCubeSampleLevel(SkyLight.Texture, SkyLight.TextureSampler, Direction, 0).rgb;
+}
+
+float PdfFace(uint MipCount, uint FaceIndex)
+{
+	// Build Cdf
+	uint BufferOffset = BufferOffsetAtPixel(uint2(0, 0), MipCount, SkyLight.MipDimensions);
+
+	float FaceCdf[6];
+	FaceCdf[0] = SkyLight.MipTreePosX[BufferOffset];
+	FaceCdf[1] = FaceCdf[0] + SkyLight.MipTreeNegX[BufferOffset];
+	FaceCdf[2] = FaceCdf[1] + SkyLight.MipTreePosY[BufferOffset];
+	FaceCdf[3] = FaceCdf[2] + SkyLight.MipTreeNegY[BufferOffset];
+	FaceCdf[4] = FaceCdf[3] + SkyLight.MipTreePosZ[BufferOffset];
+	FaceCdf[5] = FaceCdf[4] + SkyLight.MipTreeNegZ[BufferOffset];
+	float Sum = FaceCdf[5];
+
+	for (uint Index = 0; Index < 5; ++Index)
+	{
+		FaceCdf[Index] /= Sum;
+	}
+	FaceCdf[5] = 1.0;
+
+	float Pdf = FaceCdf[FaceIndex];
+	if (FaceIndex > 0)
+		Pdf -= FaceCdf[FaceIndex - 1];
+	return Pdf;
+}
+
+float SampleMipPdf(uint3 TextureCoord, uint StopLevel)
+{
+	float MipPdf = 1.0;
+	uint BufferOffset = BufferOffsetAtPixel(TextureCoord.xy, StopLevel, SkyLight.MipDimensions);
+	if (TextureCoord.z == 0)
+	{
+		MipPdf = SkyLight.MipTreePdfPosX[BufferOffset];
+	}
+	else if (TextureCoord.z == 1)
+	{
+		MipPdf = SkyLight.MipTreePdfNegX[BufferOffset];
+	}
+	else if (TextureCoord.z == 2)
+	{
+		MipPdf = SkyLight.MipTreePdfPosY[BufferOffset];
+	}
+	else if (TextureCoord.z == 3)
+	{
+		MipPdf = SkyLight.MipTreePdfNegY[BufferOffset];
+	}
+	else if (TextureCoord.z == 4)
+	{
+		MipPdf = SkyLight.MipTreePdfPosZ[BufferOffset];
+	}
+	else // if (TextureCoord.z == 5)
+	{
+		MipPdf = SkyLight.MipTreePdfNegZ[BufferOffset];
+	}
+	return MipPdf;
+}
+
+float PdfSkyLightOcclusionRay(
+	float3 WorldNormal,
+	float3 WorldDirection
+)
+{
+	// Determine primary direction
+	uint MipCount = Log2(SkyLight.MipDimensions.x);
+	uint2 TextureRes;
+	SkyLight.Texture.GetDimensions(TextureRes.x, TextureRes.y);
+	uint3 TextureCoord = GetTextureCubeCoordinate(WorldDirection, TextureRes);
+	float2 TextureUV = TextureCoord.xy / (float2)TextureRes;
+	uint2 MipCoord = TextureUV * SkyLight.MipDimensions;
+	float FacePdf = PdfFace(MipCount, TextureCoord.z);
+
+	// #dxr_todo: StopLevel can be customized for run-time performance
+	uint StopLevel = min(SkyLight.SamplingStopLevel, MipCount);
+#if USE_MIP_PDF_TEXTURE
+	float MipPdf = SampleMipPdf(uint3(MipCoord, TextureCoord.z), StopLevel);
+
+#if USE_SOLID_ANGLE_TEXTURE
+	uint LinearIndex = BufferOffsetAtPixel(MipCoord, StopLevel, SkyLight.MipDimensions);
+	float AreaPdf = SkyLight.SolidAnglePdf[LinearIndex];
+#else
+	float AreaPdf = 1.0 / TexelCoordSolidAngle(MipCoord, DimensionsAtLevel(SkyLight.MipDimensions, StopLevel));
+#endif
+	return FacePdf * MipPdf * AreaPdf;
+#else
+	uint2 MipPixel = 0;
+	float MipPdf = 1.0;
+	for (uint Index = 0; Index < MipCount - StopLevel; ++Index)
+	{
+		// Build Cdf
+		uint MipIndex = MipCount - Index - 1;
+#if USE_MORTON_CODE
+		uint BufferOffset = BufferOffsetAtPixel(MipPixel, MipIndex, SkyLight.MipDimensions);
+		uint4 BufferOffsets = uint4(BufferOffset, BufferOffset + 1, BufferOffset + 2, BufferOffset + 3);
+#else
+		uint4 BufferOffsets;
+		BufferOffsets.x = BufferOffsetAtPixel(MipPixel, MipIndex, SkyLight.MipDimensions);
+		BufferOffsets.y = BufferOffsetAtPixel(MipPixel + uint2(1, 0), MipIndex, SkyLight.MipDimensions);
+		BufferOffsets.z = BufferOffsetAtPixel(MipPixel + uint2(0, 1), MipIndex, SkyLight.MipDimensions);
+		BufferOffsets.w = BufferOffsetAtPixel(MipPixel + uint2(1, 1), MipIndex, SkyLight.MipDimensions);
+#endif
+		float4 Cdf = SampleMipTree(TextureCoord.z, BufferOffsets);
+		Cdf.y += Cdf.x;
+		Cdf.z += Cdf.y;
+		Cdf.w += Cdf.z;
+		Cdf /= Cdf.w;
+
+		// Determine texel choice
+		uint2 MipOffset = (MipCoord.xy >> MipIndex) - MipPixel;
+		uint LinearIndex = MipOffset.y << 1 | MipOffset.x;
+		MipPdf *= LinearIndex > 0 ? Cdf[LinearIndex] - Cdf[LinearIndex - 1] : Cdf[LinearIndex];
+		MipPixel += MipOffset;
+
+		MipPixel = MipPixel << 1;
+	}
+	MipPdf = max(MipPdf, 0.0);
+	// DEBUG: assert that SampleMipPdf is equivalent for evaluating the decision tree..
+	//MipPdf = SampleMipPdf(uint3(MipCoord, TextureCoord.z), StopLevel);
+
+#if USE_SOLID_ANGLE_TEXTURE
+	uint LinearIndex = BufferOffsetAtPixel(MipCoord, StopLevel, SkyLight.MipDimensions);
+	float AreaPdf = SkyLight.SolidAnglePdf[LinearIndex];
+#else
+	float AreaPdf = 1.0 / TexelCoordSolidAngle(MipCoord, DimensionsAtLevel(SkyLight.MipDimensions, StopLevel));
+#endif
+	return FacePdf * MipPdf * AreaPdf;
+#endif
+}
+
+bool GenerateSkyLightOcclusionRay(
+	float3 WorldPosition,
+	float3 WorldNormal,
+	inout RandomSequence RandSequence,
+	out float3 RayOrigin,
+	out float3 RayDirection,
+	out float RayTMin,
+	out float RayTMax,
+	out float RayPdf
+)
+{
+	// Determine the appropriate face to use
+	uint MipCount = Log2(SkyLight.MipDimensions.x);
+	float FacePdf = 0.0;
+	uint FaceIndex = SelectFace(MipCount, RandSequence, FacePdf);
+
+	uint2 MipPixel = 0;
+	float MipPdf = 1.0;
+	uint StopLevel = min(SkyLight.SamplingStopLevel, MipCount);
+	for (uint Index = 0; Index < MipCount - StopLevel; ++Index)
+	{
+		uint MipIndex = MipCount - Index - 1;
+#if USE_MORTON_CODE
+		uint BufferOffset = BufferOffsetAtPixel(MipPixel, MipIndex, SkyLight.MipDimensions);
+		uint4 BufferOffsets = uint4(BufferOffset, BufferOffset + 1, BufferOffset + 2, BufferOffset + 3);
+#else
+		uint4 BufferOffsets;
+		BufferOffsets.x = BufferOffsetAtPixel(MipPixel, MipIndex, SkyLight.MipDimensions);
+		BufferOffsets.y = BufferOffsetAtPixel(MipPixel + uint2(1, 0), MipIndex, SkyLight.MipDimensions);
+		BufferOffsets.z = BufferOffsetAtPixel(MipPixel + uint2(0, 1), MipIndex, SkyLight.MipDimensions);
+		BufferOffsets.w = BufferOffsetAtPixel(MipPixel + uint2(1, 1), MipIndex, SkyLight.MipDimensions);
+#endif
+		float4 Cdf = SampleMipTree(FaceIndex, BufferOffsets);
+		Cdf.y += Cdf.x;
+		Cdf.z += Cdf.y;
+		Cdf.w += Cdf.z;
+		Cdf /= Cdf.w;
+
+		uint DummyVariable;
+		float RandSample = RandomSequence_GenerateSample1D(RandSequence, DummyVariable);
+		if (RandSample < Cdf.x)
+		{
+			MipPdf *= Cdf.x;
+		}
+		else if (RandSample < Cdf.y)
+		{
+			MipPdf *= Cdf.y - Cdf.x;
+			MipPixel.x += 1;
+		}
+		else if (RandSample < Cdf.z)
+		{
+			MipPdf *= Cdf.z - Cdf.y;
+			MipPixel.y += 1;
+		}
+		else // if (RandSample < Cdf.w)
+		{
+			MipPdf *= Cdf.w - Cdf.z;
+			MipPixel += 1;
+		}
+
+		MipPixel = uint2(MipPixel.x << 1, MipPixel.y << 1);
+	}
+	MipPixel = uint2(MipPixel.x >> 1, MipPixel.y >> 1);
+
+	// DEBUG: assert cached PDF is equivalent to re-evaluation
+	//MipPdf = PdfSkyLightOcclusionRay(WorldNormal, RayDirection);
+
+	uint DummyVariable;
+	float2 PixelOffset = RandomSequence_GenerateSample2D(RandSequence, DummyVariable);
+	float2 MipPixelUV = (MipPixel + PixelOffset) / float2(DimensionsAtLevel(SkyLight.MipDimensions, StopLevel));
+
+	// Construct Cubemap coordinate
+	uint2 TextureDimensions;
+	SkyLight.Texture.GetDimensions(TextureDimensions.x, TextureDimensions.y);
+	float3 TexelCoord = float3(MipPixelUV * TextureDimensions, FaceIndex);
+
+	// Determine area pdf
+#if USE_SOLID_ANGLE_TEXTURE
+	uint LinearIndex = BufferOffsetAtPixel(MipPixel, StopLevel, SkyLight.MipDimensions);
+	float AreaPdf = SkyLight.SolidAnglePdf[LinearIndex];
+#else
+	float AreaPdf = 1.0 / TexelCoordSolidAngle(MipPixel, DimensionsAtLevel(SkyLight.MipDimensions, StopLevel));
+#endif
+
+	RayOrigin = WorldPosition + WorldNormal * 0.01;
+	RayDirection = GetTextureCubeVector(TexelCoord, TextureDimensions);
+	RayTMin = 0.0;
+	RayTMax = 1.0e27;
+	RayPdf = FacePdf * MipPdf * AreaPdf;
+
+	return true;
+}
+
+float PdfLambertOcclusionRay(
+	float3 WorldNormal,
+	float3 WorldDirection
+)
+{
+	float Pdf = max(dot(WorldDirection, WorldNormal), 0.0) / PI;
+	return Pdf;
+}
+
+bool GenerateLambertOcclusionRay(
+	float3 WorldPosition,
+	float3 WorldNormal,
+	inout RandomSequence RandSequence,
+	out float3 RayOrigin,
+	out float3 RayDirection,
+	out float RayTMin,
+	out float RayTMax,
+	out float RayPdf
+)
+{
+	// Draw random variable
+	float2 BufferSize = View.BufferSizeAndInvSize.xy;
+	uint DummyVariable;
+	float2 RandSample = RandomSequence_GenerateSample2D(RandSequence, DummyVariable);
+
+	// Perform cosine-hemispherical sampling and convert to world-space
+	float4 Direction_Tangent = CosineSampleHemisphere(RandSample);
+	float3 Direction_World = TangentToWorld(Direction_Tangent.xyz, WorldNormal);
+
+	RayOrigin = WorldPosition + WorldNormal * 0.01;
+	RayDirection = Direction_World;
+	RayTMin = 0.0;
+	RayTMax = 1.0e27;
+	RayPdf = Direction_Tangent.w;
+
+	return true;
+}
+
+[shader("raygeneration")]
+void SkyLightRG()
+{
+	uint2 PixelCoord = DispatchRaysIndex().xy;
+
+	RandomSequence RandSequence;
+	uint LinearIndex = CalcLinearIndex(PixelCoord);
+	RandomSequence_Initialize(RandSequence, LinearIndex, View.FrameNumber * SkyLight.SamplesPerPixel);
+	//RandSequence.Type = 2;
+
+	// Get G-Buffer surface data
+	float2 InvBufferSize = View.BufferSizeAndInvSize.zw;
+	float2 UV = (float2(PixelCoord) + 0.5) * InvBufferSize;
+	FScreenSpaceData ScreenSpaceData = GetScreenSpaceData(UV);
+	float Depth = ScreenSpaceData.GBuffer.Depth;
+	float3 WorldPosition = ReconstructWorldPositionFromDepth(UV, Depth);
+	float3 WorldNormal = ScreenSpaceData.GBuffer.WorldNormal;
+	float3 BaseColor = ScreenSpaceData.GBuffer.BaseColor;
+	float Roughness = ScreenSpaceData.GBuffer.Roughness;
+
+	float RayDistance = 0.0;
+	float HitCount = 0.0;
+	uint SamplesPerPixel = SkyLight.SamplesPerPixel;
+
+	// Mask out depth values at or beyond far plane
+	if (Depth >= SCENE_TEXTURES_DISABLED_SCENE_DEPTH_VALUE)
+	{
+		SamplesPerPixel = 0.0;
+	}
+
+	float3 Irradiance = 0.0;
+	for (uint SampleIndex = 0; SampleIndex < SamplesPerPixel; ++SampleIndex)
+	{
+		RayDesc Ray;
+		float RayPdf = 0.0;
+
+		// Determine sky light or lambert ray
+		uint DummyVariable;
+		float RandSample = RandomSequence_GenerateSample1D(RandSequence, DummyVariable);
+		// #dxr_todo: Should strategy selection be driven by CVar?
+		float SkyLightSamplingStrategyPdf = 0.5;
+		float MisWeight = 1.0;
+		if (RandSample < SkyLightSamplingStrategyPdf)
+		{
+			GenerateSkyLightOcclusionRay(WorldPosition, WorldNormal, RandSequence, Ray.Origin, Ray.Direction, Ray.TMin, Ray.TMax, RayPdf);
+			MisWeight = RayPdf / (SkyLightSamplingStrategyPdf * RayPdf + (1.0 - SkyLightSamplingStrategyPdf) * PdfLambertOcclusionRay(WorldNormal, Ray.Direction));
+		}
+		else
+		{
+			GenerateLambertOcclusionRay(WorldPosition, WorldNormal, RandSequence, Ray.Origin, Ray.Direction, Ray.TMin, Ray.TMax, RayPdf);
+			MisWeight = RayPdf / ((1.0 - SkyLightSamplingStrategyPdf) * RayPdf + SkyLightSamplingStrategyPdf * PdfSkyLightOcclusionRay(WorldNormal, Ray.Direction));
+		}
+		
+		// Surface-normal culling
+		float NdotL = dot(WorldNormal, Ray.Direction);
+		if (RayPdf < 0.0 || NdotL <= 0.0)
+		{
+			continue;
+		}
+
+		uint RayFlags = 0
+			| RAY_FLAG_FORCE_OPAQUE; // #dxr_todo: generate a special shader permutation for occlusion ray tracing that includes any-hit shaders for masked geometry
+		FDefaultPayload Payload = (FDefaultPayload)0;
+
+		TraceRay(
+			TLAS,      // AccelerationStructure
+			RayFlags,
+			0xFF,      // InstanceInclusionMask
+			0,         // RayContributionToHitGroupIndex
+			1,         // MultiplierForGeometryContributionToShaderIndex
+			0,         // MissShaderIndex
+			Ray,       // RayDesc
+			Payload    // Payload
+		);
+
+		if (Payload.HitT >= 0.0)
+		{
+			RayDistance += Payload.HitT;
+			HitCount += 1.0;
+		}
+		else
+		{
+			float3 Radiance = EvalSkyLight(Ray.Direction);
+			// #dxr_todo: Apply actual material BRDF
+			float3 Brdf = BaseColor * Roughness / PI;
+			Irradiance += Radiance * Brdf * NdotL * MisWeight / RayPdf;
+		}
+	}
+
+	if (SamplesPerPixel > 0)
+	{
+		Irradiance /= SamplesPerPixel;
+	}
+	RWOcclusionMaskUAV[PixelCoord] = float4(Irradiance, 1.0);
+
+	if (HitCount > 0.0)
+	{
+		RWRayDistanceUAV[PixelCoord] = RayDistance / HitCount;
+	}
+	else
+	{
+		RWRayDistanceUAV[PixelCoord] = 1.0e27;
+	}
+}
\ No newline at end of file
diff --git a/Engine/Shaders/Private/Raytracing/RaytracingDebugDefinitions.ush b/Engine/Shaders/Private/Raytracing/RaytracingDebugDefinitions.ush
deleted file mode 100644
diff --git a/Engine/Shaders/Private/Raytracing/RaytracingSkyLightRG.usf b/Engine/Shaders/Private/Raytracing/RaytracingSkyLightRG.usf
deleted file mode 100644
diff --git a/Engine/Shaders/Private/Raytracing/SkyLightMipTreeCommon.ush b/Engine/Shaders/Private/Raytracing/SkyLightMipTreeCommon.ush
deleted file mode 100644
diff --git a/Engine/Shaders/Private/Raytracing/SkyLightMipTreeCommon_mv.ush b/Engine/Shaders/Private/Raytracing/SkyLightMipTreeCommon_mv.ush
new file mode 100644
--- /dev/null
+++ b/Engine/Shaders/Private/Raytracing/SkyLightMipTreeCommon_mv.ush
@@ -0,0 +1,132 @@
+// Copyright 1998-2018 Epic Games, Inc. All Rights Reserved.
+
+/*=============================================================================================
+	SkyLightMipTreeCommon.ush: Common utilities for MipTree-based SkyLight sampling
+===============================================================================================*/
+
+#include "MipTreeCommon.ush"
+
+float3 GetTextureCubeVector(float3 TexelCoord, uint2 TextureRes)
+{
+	float3 Vector;
+
+	float2 UV = TexelCoord.xy / float2(TextureRes.xy);
+	UV = UV * 2.0 - 1.0;
+
+	uint Cubeface = TexelCoord.z;
+	if (Cubeface == 0)
+	{
+		Vector = float3(1, -UV.y, -UV.x);
+	}
+	else if (Cubeface == 1)
+	{
+		Vector = float3(-1, -UV.y, UV.x);
+	}
+	else if (Cubeface == 2)
+	{
+		Vector = float3(UV.x, 1, UV.y);
+	}
+	else if (Cubeface == 3)
+	{
+		Vector = float3(UV.x, -1, -UV.y);
+	}
+	else if (Cubeface == 4)
+	{
+		Vector = float3(UV.x, -UV.y, 1);
+	}
+	else if (Cubeface == 5)
+	{
+		Vector = float3(-UV.x, -UV.y, -1);
+	}
+
+	return Vector;
+}
+
+uint GetMaximumComponentIndex(float3 Vector)
+{
+	float AbsX = abs(Vector.x);
+	float AbsY = abs(Vector.y);
+	float AbsZ = abs(Vector.z);
+	if (AbsX > AbsY)
+	{
+		return AbsX > AbsZ ? 0 : 2;
+	}
+	else
+	{
+		return AbsY > AbsZ ? 1 : 2;
+	}
+}
+
+uint3 GetTextureCubeCoordinate(float3 WorldDirection, uint2 TextureRes)
+{
+	// Strongest coordinate determines cube face
+	uint MaxComponent = GetMaximumComponentIndex(WorldDirection);
+	WorldDirection /= abs(WorldDirection[MaxComponent]);
+	float3 TextureCubeCoordinate;
+	if (MaxComponent == 0)
+	{
+		if (WorldDirection[MaxComponent] > 0)
+		{
+			TextureCubeCoordinate = float3(-WorldDirection.z, -WorldDirection.y, 0);
+		}
+		else
+		{
+			TextureCubeCoordinate = float3(WorldDirection.z, -WorldDirection.y, 1);
+		}
+	}
+	else if (MaxComponent == 1)
+	{
+		if (WorldDirection[MaxComponent] > 0)
+		{
+			TextureCubeCoordinate = float3(WorldDirection.x, WorldDirection.z, 2);
+		}
+		else
+		{
+			TextureCubeCoordinate = float3(WorldDirection.x, -WorldDirection.z, 3);
+		}
+	}
+	else if (MaxComponent == 2)
+	{
+		if (WorldDirection[MaxComponent] > 0)
+		{
+			TextureCubeCoordinate = float3(WorldDirection.x, -WorldDirection.y, 4);
+		}
+		else
+		{
+			TextureCubeCoordinate = float3(-WorldDirection.x, -WorldDirection.y, 5);
+		}
+	}
+
+	TextureCubeCoordinate.xy = (TextureCubeCoordinate.xy + 1.0) / 2.0;
+	TextureCubeCoordinate.xy *= TextureRes;
+	return TextureCubeCoordinate;
+}
+
+/**
+* Cube texel solid angle computation based on formula from Manne Öhrström's
+* thesis - "Spherical Harmonics, Precomputed Radiance Transfer and Realtime Radiosity in Computer Games".
+*/
+float AreaElement(float x, float y)
+{
+	// Equation 7.12
+	return atan2(x * y, sqrt(x * x + y * y + 1));
+}
+
+float TexelCoordSolidAngle(uint2 TextureCoord, uint2 TextureRes)
+{
+	float2 InvResolution = 1.0f / (float)TextureRes;
+
+	// Convert texture coordinate to [-1, 1] range, offset to texel center.
+	float U = (2.0f * ((float)TextureCoord.x + 0.5f) * InvResolution.x) - 1.0f;
+	float V = (2.0f * ((float)TextureCoord.y + 0.5f) * InvResolution.y) - 1.0f;
+
+	// Projected area
+	float x0 = U - InvResolution.x;
+	float y0 = V - InvResolution.y;
+	float x1 = U + InvResolution.x;
+	float y1 = V + InvResolution.y;
+
+	// Equation 7.12
+	float SolidAngle = AreaElement(x0, y0) - AreaElement(x0, y1) - AreaElement(x1, y0) + AreaElement(x1, y1);
+	return SolidAngle;
+}
diff --git a/Engine/Shaders/Private/Raytracing/VisualizeSkyLightMipTreePS.usf b/Engine/Shaders/Private/Raytracing/VisualizeSkyLightMipTreePS.usf
deleted file mode 100644
diff --git a/Engine/Shaders/Private/Raytracing/VisualizeSkyLightMipTreePS_mv.usf b/Engine/Shaders/Private/Raytracing/VisualizeSkyLightMipTreePS_mv.usf
new file mode 100644
--- /dev/null
+++ b/Engine/Shaders/Private/Raytracing/VisualizeSkyLightMipTreePS_mv.usf
@@ -0,0 +1,62 @@
+// Copyright 1998-2018 Epic Games, Inc. All Rights Reserved.
+
+/*=============================================================================================
+	VisualizeSkyLightMipTreePixelShader.usf: Interactive visualization of sky light mip tree.
+===============================================================================================*/
+
+#include "../Common.ush"
+#include "MipTreeCommon.ush"
+
+uint3 Dimensions;
+Buffer<float> MipTreePosX;
+Buffer<float> MipTreeNegX;
+Buffer<float> MipTreePosY;
+Buffer<float> MipTreeNegY;
+Buffer<float> MipTreePosZ;
+Buffer<float> MipTreeNegZ;
+
+float SampleMipTree(uint FaceIndex, uint BufferOffset)
+{
+	float Value = 0.0;
+	if (FaceIndex == 0)
+	{
+		Value = MipTreePosX[BufferOffset];
+	}
+	else if (FaceIndex == 1)
+	{
+		Value = MipTreeNegX[BufferOffset];
+	}
+	else if (FaceIndex == 2)
+	{
+		Value = MipTreePosY[BufferOffset];
+	}
+	else if (FaceIndex == 3)
+	{
+		Value = MipTreeNegY[BufferOffset];
+	}
+	else if (FaceIndex == 4)
+	{
+		Value = MipTreePosZ[BufferOffset];
+	}
+	else
+	{
+		Value = MipTreeNegZ[BufferOffset];
+	}
+
+	return Value;
+}
+
+void VisualizeSkyLightMipTreePS(
+	in noperspective float2 UV : TEXCOORD0,
+	out float4 OutColor : SV_Target0,
+	out float OutMipTree : SV_Target1
+)
+{
+	uint FaceIndex = 0;
+	uint MipLevel = 0;
+	uint2 TexCoord = UV * DimensionsAtLevel(Dimensions, MipLevel);
+	uint BufferOffset = BufferOffsetAtPixel(TexCoord, MipLevel, Dimensions);
+
+	OutColor = 0.f;
+	OutMipTree = SampleMipTree(FaceIndex, BufferOffset);
+}
\ No newline at end of file
