Renaming Raytracing/ to RayTracing/. This will be done in two submissions. (Part 1)

diff --git a/Engine/Shaders/Private/RayTracing_mv/BuildMipTreeCS.usf b/Engine/Shaders/Private/RayTracing_mv/BuildMipTreeCS.usf
new file mode 100644
--- /dev/null
+++ b/Engine/Shaders/Private/RayTracing_mv/BuildMipTreeCS.usf
@@ -0,0 +1,95 @@
+// Copyright 1998-2018 Epic Games, Inc. All Rights Reserved.
+
+/*=============================================================================================
+	BuildMipTreeCS.usf: Build variance map
+===============================================================================================*/
+
+#include "../Common.ush"
+#include "SkyLightMipTreeCommon.ush"
+
+Texture2D RectLightTexture;
+TextureCube Texture;
+SamplerState TextureSampler;
+
+uint2 Dimensions;
+uint FaceIndex;
+uint MipLevel;
+RWBuffer<float> RWMipTree;
+
+[numthreads(THREADGROUP_SIZE, THREADGROUP_SIZE, 1)]
+void BuildMipTreeCS(
+	uint3 DispatchThreadId : SV_DispatchThreadID
+)
+{
+	uint3 GridCoordinate = DispatchThreadId;
+	uint2 Bounds = DimensionsAtLevel(Dimensions, MipLevel);
+	if (GridCoordinate.x < 0 || GridCoordinate.x >= Bounds.x || GridCoordinate.y < 0 || GridCoordinate.y >= Bounds.y)
+	{
+		return;
+	}
+
+	float Sum = 0.0;
+	if (MipLevel == 0)
+	{
+		// #dxr_todo: TextureCubeSampleLevel() may be enough.. could do live during query??
+		float3 Direction = GetTextureCubeVector(uint3(GridCoordinate.xy, FaceIndex), Dimensions);
+		float3 Value = TextureCubeSample(Texture, TextureSampler, Direction).rgb;
+		Sum += Luminance(Value.rgb);
+	}
+	else
+	{
+		// Report the mean as the aggregate.
+		uint2 PixelOffset = uint2(GridCoordinate.x << 1, GridCoordinate.y << 1);
+		for (int i = 0; i < 4; ++i)
+		{
+			uint2 Pixel = PixelOffset + uint2(i >> 1, i & 1);
+			uint InputIndex = BufferOffsetAtPixel(Pixel, MipLevel - 1, Dimensions);
+			// #dxr_todo: Experiment between average and max
+			Sum += RWMipTree[InputIndex];
+		}
+		Sum /= 4;
+	}
+	uint OutputIndex = BufferOffsetAtPixel(GridCoordinate.xy, MipLevel, Dimensions);
+	RWMipTree[OutputIndex] = Sum;
+}
+
+[numthreads(THREADGROUP_SIZE, THREADGROUP_SIZE, 1)]
+void BuildRectLightMipTreeCS(
+	uint3 DispatchThreadId : SV_DispatchThreadID
+)
+{
+	uint3 GridCoordinate = DispatchThreadId;
+	uint2 Bounds = DimensionsAtLevel(Dimensions, MipLevel);
+	if (GridCoordinate.x < 0 || GridCoordinate.x >= Bounds.x || GridCoordinate.y < 0 || GridCoordinate.y >= Bounds.y)
+	{
+		return;
+	}
+
+	float Sum = 0.0;
+	if (MipLevel == 0)
+	{
+		// #dxr_todo: TextureCubeSampleLevel() may be enough.. could do live during query??
+		uint2 TextureDimensions;
+		RectLightTexture.GetDimensions(TextureDimensions.x, TextureDimensions.y);
+		float2 UV = GridCoordinate.xy / float2(TextureDimensions);
+		UV = 1.0 - UV;
+		float3 Value = Texture2DSample(RectLightTexture, TextureSampler, UV).rgb;
+		Sum += Luminance(Value.rgb);
+	}
+	else
+	{
+		// Report the mean as the aggregate.
+		uint2 PixelOffset = uint2(GridCoordinate.x << 1, GridCoordinate.y << 1);
+		for (int i = 0; i < 4; ++i)
+		{
+			uint2 Pixel = PixelOffset + uint2(i >> 1, i & 1);
+			uint InputIndex = BufferOffsetAtPixel(Pixel, MipLevel - 1, Dimensions);
+			// #dxr_todo: Experiment between average and max
+			//Sum = max(Sum, RWMipTree[InputIndex]);
+			Sum += RWMipTree[InputIndex];
+		}
+		Sum /= 4;
+	}
+	uint OutputIndex = BufferOffsetAtPixel(GridCoordinate.xy, MipLevel, Dimensions);
+	RWMipTree[OutputIndex] = Sum;
+}
\ No newline at end of file
diff --git a/Engine/Shaders/Private/RayTracing_mv/BuildMipTreePdfCS.usf b/Engine/Shaders/Private/RayTracing_mv/BuildMipTreePdfCS.usf
new file mode 100644
--- /dev/null
+++ b/Engine/Shaders/Private/RayTracing_mv/BuildMipTreePdfCS.usf
@@ -0,0 +1,87 @@
+// Copyright 1998-2018 Epic Games, Inc. All Rights Reserved.
+
+/*=============================================================================================
+BuildMipTreeCS.usf: Build MipTree PDF structure
+===============================================================================================*/
+
+#include "../Common.ush"
+#include "SkyLightMipTreeCommon.ush"
+
+uint2 Dimensions;
+uint FaceIndex;
+uint MipLevel;
+
+Buffer<float> MipTree;
+RWBuffer<float> RWMipTreePdf;
+
+float SampleMipTree(uint2 MipPixel, uint MipIndex)
+{
+	return MipTree[BufferOffsetAtPixel(MipPixel, MipIndex, Dimensions)];
+}
+
+float CalcSkyLightPdf(uint2 Pixel)
+{
+	uint2 MipPixel = 0;
+	float MipPdf = 1.0;
+	uint StopLevel = MipLevel;
+	uint MipCount = Log2(Dimensions.x);
+	for (uint Index = 0; Index < MipCount - StopLevel; ++Index)
+	{
+		uint MipIndex = MipCount - Index - 1;
+
+		// #dxr_todo: Use float4?
+		float Cdf[4];
+		Cdf[0] = SampleMipTree(MipPixel, MipIndex);
+		Cdf[1] = Cdf[0] + SampleMipTree(MipPixel + uint2(1, 0), MipIndex);
+		Cdf[2] = Cdf[1] + SampleMipTree(MipPixel + uint2(0, 1), MipIndex);
+		Cdf[3] = Cdf[2] + SampleMipTree(MipPixel + uint2(1, 1), MipIndex);
+
+		float Sum = Cdf[3];
+		if (Sum > 0.0)
+		{
+			for (uint i = 0; i < 3; ++i)
+			{
+				Cdf[i] /= Sum;
+			}
+			Cdf[3] = 1.0;
+		}
+
+		// Determine probability
+		uint2 MipOffset = (Pixel.xy >> (MipIndex - StopLevel)) - MipPixel;
+		uint LinearIndex = MipOffset.y << 1 | MipOffset.x;
+		MipPdf *= LinearIndex > 0 ? Cdf[LinearIndex] - Cdf[LinearIndex - 1] : Cdf[LinearIndex];
+#if 0
+		float LevelPdf = Cdf[LinearIndex];
+		if (LinearIndex > 0)
+			LevelPdf -= Cdf[LinearIndex - 1];
+		MipPdf *= LevelPdf;
+#endif
+		MipPixel += MipOffset;
+
+		MipPixel = MipPixel << 1;
+	}
+	//MipPdf = max(MipPdf, 0.0);
+	return MipPdf;
+
+	MipPixel = MipPixel >> 1;
+	float AreaPdf = 1.0 / TexelCoordSolidAngle(MipPixel, DimensionsAtLevel(Dimensions, StopLevel));
+	return MipPdf * AreaPdf;
+}
+
+[numthreads(THREADGROUP_SIZE, THREADGROUP_SIZE, 1)]
+void BuildMipTreePdfCS(
+	uint3 DispatchThreadId : SV_DispatchThreadID
+)
+{
+	uint3 GridCoordinate = DispatchThreadId;
+	uint2 Bounds = DimensionsAtLevel(Dimensions, MipLevel);
+	if (GridCoordinate.x < 0 || GridCoordinate.x >= Bounds.x || GridCoordinate.y < 0 || GridCoordinate.y >= Bounds.y)
+	{
+		return;
+	}
+
+	float Pdf = CalcSkyLightPdf(GridCoordinate.xy);
+	uint BufferOffset = BufferOffsetAtPixel(GridCoordinate.xy, MipLevel, Dimensions);
+	RWMipTreePdf[BufferOffset] = Pdf;
+	//RWMipTreePdf[BufferOffset] = SampleMipTree(GridCoordinate.xy, MipLevel);
+}
\ No newline at end of file
diff --git a/Engine/Shaders/Private/RayTracing_mv/BuildSolidAnglePdfCS.usf b/Engine/Shaders/Private/RayTracing_mv/BuildSolidAnglePdfCS.usf
new file mode 100644
--- /dev/null
+++ b/Engine/Shaders/Private/RayTracing_mv/BuildSolidAnglePdfCS.usf
@@ -0,0 +1,26 @@
+// Copyright 1998-2018 Epic Games, Inc. All Rights Reserved.
+
+/*=============================================================================================
+BuildMipTreeCS.usf: Build MipTree PDF structure
+===============================================================================================*/
+
+#include "SkyLightMipTreeCommon.ush"
+
+uint MipLevel;
+uint2 Dimensions;
+RWBuffer<float> RWSolidAnglePdf;
+
+[numthreads(THREADGROUP_SIZE, THREADGROUP_SIZE, 1)]
+void BuildSolidAnglePdfCS(
+	uint3 DispatchThreadId : SV_DispatchThreadID
+)
+{
+	uint2 Pixel = DispatchThreadId.xy;
+	if (Pixel.x >= Dimensions.x || Pixel.y >= Dimensions.y)
+	{
+		return;
+	}
+
+	uint LinearIndex = BufferOffsetAtPixel(Pixel, MipLevel, Dimensions);
+	RWSolidAnglePdf[LinearIndex] = 1.0 / TexelCoordSolidAngle(Pixel, Dimensions);
+}
diff --git a/Engine/Shaders/Private/RayTracing_mv/CommonRayTracingBuiltInResources.ush b/Engine/Shaders/Private/RayTracing_mv/CommonRayTracingBuiltInResources.ush
new file mode 100644
--- /dev/null
+++ b/Engine/Shaders/Private/RayTracing_mv/CommonRayTracingBuiltInResources.ush
@@ -0,0 +1,82 @@
+// Copyright 1998-2018 Epic Games, Inc. All Rights Reserved.
+
+/*************************************************************
+* Changeing this file requires recompillation of the engine! *
+*************************************************************/
+
+#pragma once
+
+#ifndef RAYTRACINGBUILTINRESOURCES_USH_INCLUDED
+#define RAYTRACINGBUILTINRESOURCES_USH_INCLUDED // Workarround for UE-66460
+
+#include "RayTracingDefinitions.ush"
+
+#if defined(PLATFORM_WINDOWS) || defined(PLATFORM_XBOXONE)
+	#define INCLUDED_FROM_C_CODE 1
+	#define INCLUDED_FROM_HLSL_CODE 0
+#elif defined(SM5_PROFILE)
+	// #dxr_todo: we should use a built-in macro to detect if this shader is compiled using DXC (depends on [REDACTED_PROFILE_URL]/DirectXShaderCompiler/issues/1686)
+	#define INCLUDED_FROM_C_CODE 0
+	#define INCLUDED_FROM_HLSL_CODE 1
+#else
+	#error Unknown Compiler
+#endif
+
+#if INCLUDED_FROM_HLSL_CODE
+	#define UINT_TYPE uint
+#elif INCLUDED_FROM_C_CODE
+	#define UINT_TYPE uint32
+#endif
+
+struct FHitGroupSystemVertexFetchParameters
+{
+	// Config is a bitfield defined in FHitGroupSystemParameters
+	// uint IndexStride  : 8; // Can be just 1 bit to indicate 16 or 32 bit indices
+	// uint VertexStride : 8; // Can be just 2 bits to indicate float3, float2 or half2 format
+	// uint Unused       : 16;
+	UINT_TYPE Config;
+	UINT_TYPE IndexBufferOffsetInBytes;
+
+	UINT_TYPE GetIndexStride()
+	{
+		return Config & 0xFF;
+	}
+
+	UINT_TYPE GetVertexStride()
+	{
+		return (Config >> 8) & 0xFF;
+	}
+
+	#if INCLUDED_FROM_C_CODE
+		void SetVertexAndIndexStride(UINT_TYPE Vertex, UINT_TYPE Index)
+		{
+			Config = (Index & 0xFF) | ((Vertex & 0xFF) << 8);
+		}
+	#endif
+};
+
+#define RAY_TRACING_SYSTEM_INDEXBUFFER_REGISTER  0
+#define RAY_TRACING_SYSTEM_VERTEXBUFFER_REGISTER 1
+#define RAY_TRACING_SYSTEM_ROOTCONSTANT_REGISTER 0
+
+#if INCLUDED_FROM_HLSL_CODE
+	#define RT_CONCATENATE2(a, b) a##b
+	#define RT_CONCATENATE(a, b) RT_CONCATENATE2(a, b)
+	#define RT_REGISTER(InType, InIndex, InSpace) register(RT_CONCATENATE(InType, InIndex), RT_CONCATENATE(space, InSpace))
+
+	// Built-in local root parameters that are always bound to all hit shaders
+	ByteAddressBuffer									 HitGroupSystemIndexBuffer   : RT_REGISTER(t, RAY_TRACING_SYSTEM_INDEXBUFFER_REGISTER,  RAY_TRACING_REGISTER_SPACE_SYSTEM);
+	ByteAddressBuffer									 HitGroupSystemVertexBuffer  : RT_REGISTER(t, RAY_TRACING_SYSTEM_VERTEXBUFFER_REGISTER, RAY_TRACING_REGISTER_SPACE_SYSTEM);
+	ConstantBuffer<FHitGroupSystemVertexFetchParameters> HitGroupSystemRootConstants : RT_REGISTER(b, RAY_TRACING_SYSTEM_ROOTCONSTANT_REGISTER, RAY_TRACING_REGISTER_SPACE_SYSTEM);
+
+	#undef RT_REGISTER
+	#undef RT_CONCATENATE
+	#undef RT_CONCATENATE2
+#endif // INCLUDED_FROM_HLSL_CODE
+
+
+#undef INCLUDED_FROM_C_CODE
+#undef INCLUDED_FROM_HLSL_CODE
+#undef UINT_TYPE
+
+#endif // RAYTRACINGBUILTINRESOURCES_USH_INCLUDED // Workarround for UE-66460
diff --git a/Engine/Shaders/Private/RayTracing_mv/CompositeAmbientOcclusionPS.usf b/Engine/Shaders/Private/RayTracing_mv/CompositeAmbientOcclusionPS.usf
new file mode 100644
--- /dev/null
+++ b/Engine/Shaders/Private/RayTracing_mv/CompositeAmbientOcclusionPS.usf
@@ -0,0 +1,16 @@
+// Copyright 1998-2018 Epic Games, Inc. All Rights Reserved.
+
+#include "../Common.ush"
+#include "../DeferredShadingCommon.ush"
+
+Texture2D<float4> AmbientOcclusionTexture;
+SamplerState AmbientOcclusionTextureSampler;
+
+void CompositeAmbientOcclusionPS(
+	in noperspective float2 UV : TEXCOORD0,
+	out float4 OutColor : SV_Target0
+)
+{
+	float4 AmbientOcclusion = AmbientOcclusionTexture.Sample(AmbientOcclusionTextureSampler, UV);
+	OutColor = float4(AmbientOcclusion.rgb, 1.0);
+}
diff --git a/Engine/Shaders/Private/RayTracing_mv/CompositeSkyLightPS.usf b/Engine/Shaders/Private/RayTracing_mv/CompositeSkyLightPS.usf
new file mode 100644
--- /dev/null
+++ b/Engine/Shaders/Private/RayTracing_mv/CompositeSkyLightPS.usf
@@ -0,0 +1,19 @@
+// Copyright 1998-2018 Epic Games, Inc. All Rights Reserved.
+
+#include "../Common.ush"
+#include "../DeferredShadingCommon.ush"
+
+Texture2D<float4> SkyLightTexture;
+Texture2D<float> HitDistanceTexture;
+
+SamplerState SkyLightTextureSampler;
+SamplerState HitDistanceTextureSampler;
+
+void CompositeSkyLightPS(
+	in noperspective float2 UV : TEXCOORD0,
+	out float4 OutColor : SV_Target0
+)
+{
+	float4 SkyLight = SkyLightTexture.Sample(SkyLightTextureSampler, UV);
+	OutColor = float4(SkyLight.rgb, 1.0);
+}
diff --git a/Engine/Shaders/Private/RayTracing_mv/MipTreeCommon.ush b/Engine/Shaders/Private/RayTracing_mv/MipTreeCommon.ush
new file mode 100644
--- /dev/null
+++ b/Engine/Shaders/Private/RayTracing_mv/MipTreeCommon.ush
@@ -0,0 +1,76 @@
+// Copyright 1998-2018 Epic Games, Inc. All Rights Reserved.
+
+/*=============================================================================================
+	MipTreeCommon.ush: Common utilities for MipTree-based sampling
+===============================================================================================*/
+
+#define USE_MORTON_CODE 1
+
+uint2 DimensionsAtLevel(uint2 Dimensions, uint Level)
+{
+	return uint2(Dimensions.x >> Level, Dimensions.y >> Level);
+}
+
+// #dxr_todo: Discussing with Brian where these utility functions should live.
+uint MortonEncode(uint2 Pixel)
+{
+	Pixel.x &= 0x0000ffff;
+	Pixel.x = (Pixel.x ^ (Pixel.x << 8)) & 0x00ff00ff;
+	Pixel.x = (Pixel.x ^ (Pixel.x << 4)) & 0x0f0f0f0f;
+	Pixel.x = (Pixel.x ^ (Pixel.x << 2)) & 0x33333333;
+	Pixel.x = (Pixel.x ^ (Pixel.x << 1)) & 0x55555555;
+
+	Pixel.y &= 0x0000ffff;
+	Pixel.y = (Pixel.y ^ (Pixel.y << 8)) & 0x00ff00ff;
+	Pixel.y = (Pixel.y ^ (Pixel.y << 4)) & 0x0f0f0f0f;
+	Pixel.y = (Pixel.y ^ (Pixel.y << 2)) & 0x33333333;
+	Pixel.y = (Pixel.y ^ (Pixel.y << 1)) & 0x55555555;
+	return Pixel.x | (Pixel.y << 1);
+}
+
+uint2 MortonDecode(uint Morton)
+{
+	uint2 Pixel = uint2(Morton, Morton >> 1);
+	Pixel.x &= 0x55555555;
+	Pixel.x = (Pixel.x ^ (Pixel.x >> 1)) & 0x33333333;
+	Pixel.x = (Pixel.x ^ (Pixel.x >> 2)) & 0x0f0f0f0f;
+	Pixel.x = (Pixel.x ^ (Pixel.x >> 4)) & 0x00ff00ff;
+	Pixel.x = (Pixel.x ^ (Pixel.x >> 8)) & 0x0000ffff;
+
+	Pixel.y &= 0x55555555;
+	Pixel.y = (Pixel.y ^ (Pixel.y >> 1)) & 0x33333333;
+	Pixel.y = (Pixel.y ^ (Pixel.y >> 2)) & 0x0f0f0f0f;
+	Pixel.y = (Pixel.y ^ (Pixel.y >> 4)) & 0x00ff00ff;
+	Pixel.y = (Pixel.y ^ (Pixel.y >> 8)) & 0x0000ffff;
+	return Pixel;
+}
+
+uint BufferOffsetAtPixel(uint2 Pixel, uint Level, uint2 Dimensions)
+{
+	uint Offset = 0;
+	for (uint LevelIndex = 0; LevelIndex < Level; ++LevelIndex)
+	{
+		uint2 MipLevelDimensions = DimensionsAtLevel(Dimensions, LevelIndex);
+		Offset += MipLevelDimensions.x * MipLevelDimensions.y;
+	}
+
+	uint2 MipLevelDimensions = DimensionsAtLevel(Dimensions, Level);
+#if USE_MORTON_CODE
+	// NOTE: Morton encoding requires power-of-two dimensionality..
+	Offset += MortonEncode(Pixel);
+#else
+	Offset += Pixel.y * MipLevelDimensions.x + Pixel.x;
+#endif
+	return Offset;
+}
+
+uint Log2(uint Value)
+{
+	uint Result = 0;
+	while (Value)
+	{
+		Result++;
+		Value = Value >> 1;
+	}
+	return Result - 1;
+}
diff --git a/Engine/Shaders/Private/RayTracing_mv/RayTracingAmbientOcclusionRG.usf b/Engine/Shaders/Private/RayTracing_mv/RayTracingAmbientOcclusionRG.usf
new file mode 100644
--- /dev/null
+++ b/Engine/Shaders/Private/RayTracing_mv/RayTracingAmbientOcclusionRG.usf
@@ -0,0 +1,134 @@
+// Copyright 1998-2018 Epic Games, Inc. All Rights Reserved.
+
+#include "../Common.ush"
+#include "../Montecarlo.ush"
+#include "../DeferredShadingCommon.ush"
+#include "../PathTracing/Utilities/PathTracingRandomSequence.ush"
+#include "RayTracingCommon.ush"
+
+RaytracingAccelerationStructure TLAS;
+RWTexture2D<float4> RWOcclusionMaskUAV;
+RWTexture2D<float> RWRayDistanceUAV;
+
+uint CalcLinearIndex(uint2 PixelCoord)
+{
+	return PixelCoord.y * View.BufferSizeAndInvSize.x + PixelCoord.x;
+}
+
+// #patrick todo: move to RayTracingCommon.ush
+float3 ReconstructWorldPositionFromDepth(float2 UV, float Depth)
+{
+	float2 ScreenPosition = (UV - View.ScreenPositionScaleBias.wz) / View.ScreenPositionScaleBias.xy;
+	float4 HomogeneousWorldPosition = mul(float4(ScreenPosition * Depth, Depth, 1), View.ScreenToWorld);
+	float3 WorldPosition = HomogeneousWorldPosition.xyz / HomogeneousWorldPosition.w;
+
+	return WorldPosition;
+}
+
+void GenerateCosineNormalRay(
+	float3 WorldPosition,
+	float3 WorldNormal,
+	inout RandomSequence RandSequence,
+	out float3 RayOrigin,
+	out float3 RayDirection,
+	out float RayTMin,
+	out float RayTMax
+)
+{
+	// Draw random variable
+	float2 BufferSize = View.BufferSizeAndInvSize.xy;
+	uint DummyVariable;
+	float2 RandSample = RandomSequence_GenerateSample2D(RandSequence, DummyVariable);
+
+	// Perform cosine-hemispherical sampling and convert to world-space
+	float4 Direction_Tangent = CosineSampleHemisphere(RandSample);
+	float3 Direction_World = TangentToWorld(Direction_Tangent.xyz, WorldNormal);
+
+	RayOrigin = WorldPosition + WorldNormal * 0.01;
+	RayDirection = Direction_World;
+	RayTMin = 0.0;
+	RayTMax = AmbientOcclusion.MaxRayDistance;
+}
+
+[shader("raygeneration")]
+void AmbientOcclusionRG()
+{
+	uint2 PixelCoord = DispatchRaysIndex().xy;
+
+	RandomSequence RandSequence;
+	uint LinearIndex = CalcLinearIndex(PixelCoord);
+	RandomSequence_Initialize(RandSequence, LinearIndex, View.FrameNumber);
+
+	// Get G-Buffer surface data
+	float2 InvBufferSize = View.BufferSizeAndInvSize.zw;
+	float2 UV = (float2(PixelCoord) + 0.5) * InvBufferSize;
+	FScreenSpaceData ScreenSpaceData = GetScreenSpaceData(UV);
+	float Depth = ScreenSpaceData.GBuffer.Depth;
+	float3 WorldPosition = ReconstructWorldPositionFromDepth(UV, Depth);
+	float3 WorldNormal = ScreenSpaceData.GBuffer.WorldNormal;
+
+	float Visibility = 0.0;
+	float RayDistance = 0.0;
+	float HitCount = 0.0;
+	uint SamplesPerPixel = AmbientOcclusion.SamplesPerPixel;
+	
+	// Mask out depth values beyond far plane
+	if (Depth >= SCENE_TEXTURES_DISABLED_SCENE_DEPTH_VALUE)
+	{
+		SamplesPerPixel = 0.0;
+		Visibility = 1.0;
+	}
+
+	for (uint SampleIndex = 0; SampleIndex < SamplesPerPixel; ++SampleIndex)
+	{
+		RayDesc Ray;
+		GenerateCosineNormalRay(WorldPosition, WorldNormal, RandSequence, Ray.Origin, Ray.Direction, Ray.TMin, Ray.TMax);
+		if (dot(WorldNormal, Ray.Direction) > 0.0)
+		{
+			uint RayFlags = 0
+				| RAY_FLAG_CULL_BACK_FACING_TRIANGLES
+				| RAY_FLAG_FORCE_OPAQUE; // #dxr_todo: generate a special shader permutation for occlusion ray tracing that includes any-hit shaders for masked geometry
+
+			FDefaultPayload Payload = (FDefaultPayload)0;
+
+			TraceRay(
+				TLAS,
+				RayFlags,
+				0xFF, // InstanceInclusionMask
+				0,    // RayContributionToHitGroupIndex
+				1,    // MultiplierForGeometryContributionToShaderIndex
+				0,    // MissShaderIndex
+				Ray,
+				Payload
+			);
+
+			Visibility += Payload.HitT < 0 ? 1.0 : 0.0;
+			if (Payload.HitT >= 0.0)
+			{
+				RayDistance += Payload.HitT;
+				HitCount += 1.0;
+			}
+		}
+		else
+		{
+			Visibility += 0.0;
+			RayDistance += 0.0;
+			HitCount += 1.0;
+		}
+	}
+
+	RWOcclusionMaskUAV[PixelCoord] = Visibility;
+	if (SamplesPerPixel > 0)
+	{
+		RWOcclusionMaskUAV[PixelCoord] /= SamplesPerPixel;
+	}
+
+	if (HitCount > 0.0)
+	{
+		RWRayDistanceUAV[PixelCoord] = RayDistance / HitCount;
+	}
+	else
+	{
+		RWRayDistanceUAV[PixelCoord] = 1.0e27;
+	}
+}
\ No newline at end of file
diff --git a/Engine/Shaders/Private/RayTracing_mv/RayTracingBarycentrics.usf b/Engine/Shaders/Private/RayTracing_mv/RayTracingBarycentrics.usf
new file mode 100644
--- /dev/null
+++ b/Engine/Shaders/Private/RayTracing_mv/RayTracingBarycentrics.usf
@@ -0,0 +1,75 @@
+// Copyright 1998-2018 Epic Games, Inc. All Rights Reserved.
+
+#include "../Common.ush"
+#include "RayTracingCommon.ush"
+
+RWTexture2D<float4> Output;
+RaytracingAccelerationStructure TLAS;
+
+float3 ReconstructWorldPositionFromDepth(float2 UV, float Depth)
+{
+	float2 ScreenPosition = (UV - View.ScreenPositionScaleBias.wz) / View.ScreenPositionScaleBias.xy;
+	float4 HomogeneousWorldPosition = mul(float4(ScreenPosition * Depth, Depth, 1), View.ScreenToWorld);
+	float3 WorldPosition = HomogeneousWorldPosition.xyz / HomogeneousWorldPosition.w;
+
+	return WorldPosition;
+}
+
+struct FBarycentricsPayload
+{
+	float2 Barycentrics;
+	bool Hit;
+};
+
+
+[shader("raygeneration")]
+void RayTracingBarycentricsMainRG()
+{
+	uint2 PixelCoord = DispatchRaysIndex().xy;
+
+	float2 RenderTargetUV = (float2(PixelCoord) + .5f) * View.BufferSizeAndInvSize.zw;
+
+	RayDesc Ray;
+
+	Ray.Origin = ReconstructWorldPositionFromDepth(RenderTargetUV, 0.0f);
+	Ray.Direction = normalize(ReconstructWorldPositionFromDepth(RenderTargetUV, 10000.f) - Ray.Origin);
+	Ray.TMin = 0.5f; // 1 unit in UE4 = 1 cm -0> 0.5 ==> 5 mm
+	Ray.TMax = 1e27f;
+
+	uint RayFlags = 0
+	              | RAY_FLAG_CULL_BACK_FACING_TRIANGLES      // use back face culling
+	              | RAY_FLAG_FORCE_OPAQUE;                   // don't run anyhit shader
+
+	FBarycentricsPayload Payload = (FBarycentricsPayload)0;
+
+	TraceRay(
+		TLAS,   // AccelerationStructure
+		RayFlags,
+		0xFF,   // InstanceInclusionMask
+		0,      // RayContributionToHitGroupIndex
+		1,      // MultiplierForGeometryContributionToShaderIndex
+		0,      // MissShaderIndex
+		Ray,    // RayDesc
+		Payload // Payload
+	);
+
+	float4 Result = float4(0, 0, 0, 1);
+	if (Payload.Hit)
+	{
+		Result = float4(Payload.Barycentrics, 1.0f - Payload.Barycentrics.x - Payload.Barycentrics.y, 1.0f);
+	}
+	Output[PixelCoord] = Result;
+}
+
+[shader("miss")]
+void RayTracingBarycentricsMainMS(inout FBarycentricsPayload Payload)
+{
+	Payload.Hit = false;
+}
+
+[shader("closesthit")]
+void RayTracingBarycentricsMainCHS(inout FBarycentricsPayload Payload, in FDefaultAttributes Attributes)
+{
+	Payload.Barycentrics = Attributes.Barycentrics;
+	Payload.Hit = true;
+}
diff --git a/Engine/Shaders/Private/RayTracing_mv/RayTracingBuiltInShaders.usf b/Engine/Shaders/Private/RayTracing_mv/RayTracingBuiltInShaders.usf
new file mode 100644
--- /dev/null
+++ b/Engine/Shaders/Private/RayTracing_mv/RayTracingBuiltInShaders.usf
@@ -0,0 +1,128 @@
+// Copyright 1998-2018 Epic Games, Inc. All Rights Reserved.
+
+#include "RayTracingCommon.ush"
+
+// #dxr_todo: move RT_REGISTER macro declaration into some a shared header, as it's now used in a few places
+#define RT_CONCATENATE2(a, b) a##b
+#define RT_CONCATENATE(a, b) RT_CONCATENATE2(a, b)
+#define RT_REGISTER(InType, InIndex, InSpace) register(RT_CONCATENATE(InType, InIndex), RT_CONCATENATE(space, InSpace))
+
+RaytracingAccelerationStructure               TLAS               : RT_REGISTER(t, 0, RAY_TRACING_REGISTER_SPACE_GLOBAL);
+StructuredBuffer<FBasicRayData>               Rays               : RT_REGISTER(t, 1, RAY_TRACING_REGISTER_SPACE_GLOBAL);
+RWStructuredBuffer<uint>                      OcclusionOutput    : RT_REGISTER(u, 0, RAY_TRACING_REGISTER_SPACE_GLOBAL);
+RWStructuredBuffer<FBasicRayIntersectionData> IntersectionOutput : RT_REGISTER(u, 1, RAY_TRACING_REGISTER_SPACE_GLOBAL);
+
+#undef RT_REGISTER
+#undef RT_CONCATENATE
+#undef RT_CONCATENATE2
+
+[shader("raygeneration")]
+void OcclusionMainRG()
+{
+	const uint RayIndex = DispatchRaysIndex().x;
+	FBasicRayData InputRay = Rays[RayIndex];
+
+	RayDesc Ray;
+	Ray.Origin = InputRay.Origin;
+	Ray.Direction = InputRay.Direction;
+	Ray.TMin = 0.0f;
+	Ray.TMax = InputRay.TFar;
+
+	uint RayFlags = 0
+	              | RAY_FLAG_CULL_BACK_FACING_TRIANGLES      // use back face culling
+	              | RAY_FLAG_ACCEPT_FIRST_HIT_AND_END_SEARCH // terminate immediately upon detecting primitive intersection
+	              | RAY_FLAG_FORCE_OPAQUE                    // don't run anyhit shader
+	              | RAY_FLAG_SKIP_CLOSEST_HIT_SHADER;        // don't run closesthit shader
+
+	FOcclusionPayload Payload;
+	Payload.Hit = true;
+
+	TraceRay(
+		TLAS,   // AccelerationStructure
+		RayFlags,
+		0xFF,   // InstanceInclusionMask
+		0,      // RayContributionToHitGroupIndex
+		1,      // MultiplierForGeometryContributionToShaderIndex
+		0,      // MissShaderIndex
+		Ray,    // RayDesc
+		Payload // Payload
+	);
+
+	OcclusionOutput[RayIndex] = Payload.Hit ? ~0 : 0;
+}
+
+[shader("miss")]
+void OcclusionMainMS(inout FOcclusionPayload Payload)
+{
+	Payload.Hit = false;
+}
+
+[shader("raygeneration")]
+void IntersectionMainRG()
+{
+	const uint RayIndex = DispatchRaysIndex().x;
+	FBasicRayData InputRay = Rays[RayIndex];
+
+	RayDesc Ray;
+	Ray.Origin = InputRay.Origin;
+	Ray.Direction = InputRay.Direction;
+	Ray.TMin = 0.0f;
+	Ray.TMax = InputRay.TFar;
+
+	uint RayFlags = 0
+	              | RAY_FLAG_CULL_BACK_FACING_TRIANGLES      // use back face culling
+	              | RAY_FLAG_FORCE_OPAQUE;                   // don't run anyhit shader
+
+	FBasicRayIntersectionData Payload = (FBasicRayIntersectionData)0;
+
+	TraceRay(
+		TLAS,   // AccelerationStructure
+		RayFlags,
+		0xFF,   // InstanceInclusionMask
+		0,      // RayContributionToHitGroupIndex
+		1,      // MultiplierForGeometryContributionToShaderIndex
+		0,      // MissShaderIndex
+		Ray,    // RayDesc
+		Payload // Payload
+	);
+
+	IntersectionOutput[RayIndex] = Payload;
+}
+
+[shader("miss")]
+void IntersectionMainMS(inout FBasicRayIntersectionData Payload)
+{
+	Payload.Barycentrics = float2(0, 0);
+	Payload.InstanceIndex = ~0;
+	Payload.PrimitiveIndex = ~0;
+	
+}
+
+[shader("closesthit")]
+void IntersectionMainCHS(inout FBasicRayIntersectionData Payload, in FDefaultAttributes Attributes)
+{
+	Payload.Barycentrics = Attributes.Barycentrics;
+	Payload.InstanceIndex = InstanceIndex();
+	Payload.PrimitiveIndex = PrimitiveIndex();
+}
+
+[shader("closesthit")]
+void DefaultMainCHS(inout FDefaultPayload Payload, in FDefaultAttributes Attributes)
+{
+	Payload.Barycentrics = Attributes.Barycentrics;
+	Payload.InstanceID = InstanceID();
+	Payload.InstanceIndex = InstanceIndex();
+	Payload.PrimitiveIndex = PrimitiveIndex();
+	Payload.HitT = RayTCurrent();
+}
+
+[shader("miss")]
+void DefaultMainMS(inout FDefaultPayload Payload)
+{
+	Payload.Barycentrics = float2(0, 0);
+	Payload.InstanceID = ~0;
+	Payload.InstanceIndex = ~0;
+	Payload.PrimitiveIndex = ~0;
+	Payload.HitT = -1;
+}
+
diff --git a/Engine/Shaders/Private/RayTracing_mv/RayTracingCommon.ush b/Engine/Shaders/Private/RayTracing_mv/RayTracingCommon.ush
new file mode 100644
--- /dev/null
+++ b/Engine/Shaders/Private/RayTracing_mv/RayTracingCommon.ush
@@ -0,0 +1,185 @@
+// Copyright 1998-2018 Epic Games, Inc. All Rights Reserved.
+
+/*=============================================================================
+	RayTracingCommon.ush: common header used in multiple ray generation shaders
+=============================================================================*/
+
+#pragma once
+
+#ifndef RAYTRACINGCOMMON_USH_INCLUDED
+#define RAYTRACINGCOMMON_USH_INCLUDED // Workarround for UE-66460
+
+#include "RayTracingDefinitions.ush"
+
+struct FBasicRayData
+{
+	float3 Origin;
+	float TFar;
+	float3 Direction;
+	uint Mask;
+};
+
+struct FBasicRayIntersectionData
+{
+	float2 Barycentrics;
+	uint   InstanceIndex;  // 0xFFFFFFFF if no intersection if found
+	uint   PrimitiveIndex; // 0xFFFFFFFF if no intersection if found
+};
+
+struct FOcclusionPayload
+{
+	bool Hit;
+};
+
+struct FDefaultPayload
+{
+	float2 Barycentrics;
+	uint   InstanceID;
+	uint   InstanceIndex;
+	uint   PrimitiveIndex;
+	float  HitT;
+	float2 Padding;
+};
+
+struct FDefaultAttributes
+{
+	float2 Barycentrics;
+};
+
+struct FMaterialClosestHitPayload
+{
+	float HitT;							//4		4bytes	32bits
+
+	float3 Radiance;					//16	6bytes	48bits	
+	float3 WorldNormal;					//28	6bytes	48bits
+	float3 BaseColor;					//40	6bytes	48bits
+	float3 DiffuseColor;				//52	0bytes	(derived)
+	float3 SpecularColor;				//64	0bytes	(derived)
+	float Opacity;						//68	2bytes	16bits
+	float Metallic;						//72	1bytes	8bits
+	float Specular;						//76	1bytes	8bits
+	float Roughness;					//80	2bytes	16bits
+	float Ior;							//84	2bytes	16bits
+	uint ShadingModelID;				//88	1byte	4bits
+	uint BlendingMode;					//92	0bytes	4bits (packed)
+	uint PrimitiveLightingChannelMask;	//96	0bytes	(removed)
+	float4 CustomData;					//112	0bytes	(removed)
+	float GBufferAO;					//116	0bytes	(removed)
+	float IndirectIrradiance;			//120	0bytes	(removed)
+
+	// Quite some code assume FRayHitInfo has a WorldPos
+	// So keep it here and force to re-construct it in Unpack call using ray information.
+	// It is not packed in FRayHitInfoPacked
+	float3 WorldPos;					//132	0bytes	(derived)
+
+	uint HitKind;						//136	0bytes 1bit	(packed)
+};
+
+#define USE_PACKED_PAYLOAD 0
+
+#if !COMPUTESHADER
+#if USE_PACKED_PAYLOAD
+struct FPackedMaterialClosestHitPayload
+{
+	float HitT;
+	int RadienceAndNormal[3];
+	int BaseColorAndOpacity[2];
+	int MetallicAndSpecularAndRoughness;
+	int IorAndShadingModelIDAndBlendingModeAndHitKind;
+};
+
+FPackedMaterialClosestHitPayload PackRayTracingPayload(FMaterialClosestHitPayload Input)
+{
+    FPackedMaterialClosestHitPayload Output = (FPackedMaterialClosestHitPayload)0;
+	Output.HitT = Input.HitT;
+    Output.RadienceAndNormal[0]  = f32tof16(Input.Radiance.x);
+	Output.RadienceAndNormal[0] |= f32tof16(Input.Radiance.y) << 16;
+	Output.RadienceAndNormal[1]  = f32tof16(Input.Radiance.z);
+    Output.RadienceAndNormal[1] |= f32tof16(Input.WorldNormal.x) << 16;
+	Output.RadienceAndNormal[2]  = f32tof16(Input.WorldNormal.y);
+	Output.RadienceAndNormal[2] |= f32tof16(Input.WorldNormal.z) << 16;
+	Output.BaseColorAndOpacity[0]  = f32tof16(Input.BaseColor.x);
+	Output.BaseColorAndOpacity[0] |= f32tof16(Input.BaseColor.y) << 16;
+	Output.BaseColorAndOpacity[1]  = f32tof16(Input.BaseColor.z);
+	Output.BaseColorAndOpacity[1] |= f32tof16(Input.Opacity) << 16;
+    Output.MetallicAndSpecularAndRoughness  = (uint(round(Input.Metallic * 255.0f)) & 0xFF);
+    Output.MetallicAndSpecularAndRoughness |= (uint(round(Input.Specular * 255.0f)) & 0xFF) << 8;
+	Output.MetallicAndSpecularAndRoughness |= f32tof16(Input.Roughness) << 16;
+	Output.IorAndShadingModelIDAndBlendingModeAndHitKind  = f32tof16(Input.Ior);
+    Output.IorAndShadingModelIDAndBlendingModeAndHitKind |= (Input.ShadingModelID & 0xF) << 16;
+	Output.IorAndShadingModelIDAndBlendingModeAndHitKind |= (Input.BlendingMode & 0x7) << 20;
+	Output.IorAndShadingModelIDAndBlendingModeAndHitKind |= (Input.HitKind & 0x1) << 31;
+
+	return Output;
+}
+
+FMaterialClosestHitPayload UnpackRayTracingPayload(FPackedMaterialClosestHitPayload Input, RayDesc Ray)
+{
+    FMaterialClosestHitPayload Output = (FMaterialClosestHitPayload)0;
+	Output.HitT = Input.HitT;
+	Output.Radiance.x = f16tof32(Input.RadienceAndNormal[0]);
+	Output.Radiance.y = f16tof32(Input.RadienceAndNormal[0] >> 16);
+	Output.Radiance.z = f16tof32(Input.RadienceAndNormal[1]);
+	Output.WorldNormal.x = f16tof32(Input.RadienceAndNormal[1] >> 16);
+	Output.WorldNormal.y = f16tof32(Input.RadienceAndNormal[2]);
+	Output.WorldNormal.z = f16tof32(Input.RadienceAndNormal[2] >> 16);
+	Output.BaseColor.x = f16tof32(Input.BaseColorAndOpacity[0]);
+	Output.BaseColor.y = f16tof32(Input.BaseColorAndOpacity[0] >> 16);
+	Output.BaseColor.z = f16tof32(Input.BaseColorAndOpacity[1]);
+	Output.Opacity	   = f16tof32(Input.BaseColorAndOpacity[1] >> 16);
+	Output.Metallic	   = float(Input.MetallicAndSpecularAndRoughness & 0xFF) / 255.0f;
+    Output.Specular	   = float((Input.MetallicAndSpecularAndRoughness >> 8) & 0xFF) / 255.0f;
+    Output.Roughness   = f16tof32(Input.MetallicAndSpecularAndRoughness >> 16);
+	Output.Ior			  = f16tof32(Input.IorAndShadingModelIDAndBlendingModeAndHitKind);
+	Output.ShadingModelID = (Input.IorAndShadingModelIDAndBlendingModeAndHitKind >> 16) & 0xF;
+	Output.BlendingMode	  = (Input.IorAndShadingModelIDAndBlendingModeAndHitKind >> 20) & 0x7;
+	Output.HitKind		  = (Input.IorAndShadingModelIDAndBlendingModeAndHitKind >> 31);
+
+	Output.DiffuseColor  = Output.BaseColor - Output.BaseColor * Output.Metallic;
+	Output.SpecularColor = lerp(0.08 * Output.Specular.xxx, Output.BaseColor, Output.Metallic);
+	Output.WorldPos		 = Ray.Origin + Output.HitT * Ray.Direction;
+	return Output;
+}
+
+#else
+#define FPackedMaterialClosestHitPayload FMaterialClosestHitPayload
+
+FPackedMaterialClosestHitPayload PackRayTracingPayload(FMaterialClosestHitPayload Input)
+{
+	return Input;
+}
+
+FMaterialClosestHitPayload UnpackRayTracingPayload(FPackedMaterialClosestHitPayload Input, RayDesc Ray)
+{
+	return Input;
+}
+
+#endif
+
+FMaterialClosestHitPayload TraceRayInternal(in RaytracingAccelerationStructure TLAS, 
+											in uint RayFlags, 
+											in uint InstanceInclusionMask, 
+											in uint RayContributionToHitGroupIndex,  
+											in uint MultiplierForGeometryContributionToShaderIndex, 
+											in uint MissShaderIndex, 
+											in RayDesc Ray)
+{
+	FPackedMaterialClosestHitPayload Payload = (FPackedMaterialClosestHitPayload)0;
+	Payload.HitT = -1.0f;
+
+    TraceRay
+	(
+			TLAS,
+			RayFlags,
+			InstanceInclusionMask,
+			RayContributionToHitGroupIndex,
+			MultiplierForGeometryContributionToShaderIndex,
+			MissShaderIndex,
+			Ray,
+			Payload
+	);
+
+    return UnpackRayTracingPayload(Payload, Ray);
+}
+#endif
+#endif // RAYTRACINGCOMMON_USH_INCLUDED // Workarround for UE-66460
diff --git a/Engine/Shaders/Private/RayTracing_mv/RayTracingDebug.usf b/Engine/Shaders/Private/RayTracing_mv/RayTracingDebug.usf
new file mode 100644
--- /dev/null
+++ b/Engine/Shaders/Private/RayTracing_mv/RayTracingDebug.usf
@@ -0,0 +1,138 @@
+// Copyright 1998-2018 Epic Games, Inc. All Rights Reserved.
+
+#include "../Common.ush"
+#include "RayTracingDebugDefinitions.ush"
+#include "RayTracingCommon.ush"
+#include "RayTracingHitGroupCommon.ush"
+
+RWTexture2D<float4> Output;
+RaytracingAccelerationStructure TLAS;
+
+uint VisualizationMode;
+
+float3 ReconstructWorldPositionFromDepth(float2 UV, float Depth)
+{
+	float2 ScreenPosition = (UV - View.ScreenPositionScaleBias.wz) / View.ScreenPositionScaleBias.xy;
+	float4 HomogeneousWorldPosition = mul(float4(ScreenPosition * Depth, Depth, 1), View.ScreenToWorld);
+	float3 WorldPosition = HomogeneousWorldPosition.xyz / HomogeneousWorldPosition.w;
+
+	return WorldPosition;
+}
+
+[shader("raygeneration")]
+void RayTracingDebugMainRG()
+{
+	uint2 PixelCoord = DispatchRaysIndex().xy;
+
+	float2 RenderTargetUV = (float2(PixelCoord) + .5f) * View.BufferSizeAndInvSize.zw;
+
+	RayDesc Ray;
+
+	Ray.Origin = ReconstructWorldPositionFromDepth(RenderTargetUV, 0.0f);
+	Ray.Direction = normalize(ReconstructWorldPositionFromDepth(RenderTargetUV, 10000.f) - Ray.Origin);
+	Ray.TMin = 0.5f; // 1 unit in UE4 = 1 cm -0> 0.5 ==> 5 mm
+	Ray.TMax = 1e27f;
+
+	uint RayFlags = RAY_FLAG_CULL_BACK_FACING_TRIANGLES;
+
+	FMaterialClosestHitPayload Payload = TraceRayInternal(
+		TLAS,   // AccelerationStructure
+		RayFlags,
+		0xFF,   // InstanceInclusionMask
+		0,      // RayContributionToHitGroupIndex
+		1,      // MultiplierForGeometryContributionToShaderIndex
+		0,      // MissShaderIndex
+		Ray     // RayDesc
+	);
+
+	float4 Result = float4(0, 0, 0, 1);
+
+	if (Payload.HitT > 0)
+	{
+		switch (VisualizationMode)
+		{
+		default:
+		case RAY_TRACING_DEBUG_VIZ_RADIANCE:
+			Result = float4(Payload.Radiance, 1.0f);
+			break;
+		case RAY_TRACING_DEBUG_VIZ_WORLD_NORMAL:
+			Result = float4(Payload.WorldNormal * 0.5 + 0.5, 1.0f);
+			break;
+		case RAY_TRACING_DEBUG_VIZ_BASE_COLOR:
+			Result = float4(Payload.BaseColor.rgb, 1.0f);
+			break;
+		case RAY_TRACING_DEBUG_VIZ_DIFFUSE_COLOR:
+			Result = float4(Payload.DiffuseColor.rgb, 1.0f);
+			break;
+		case RAY_TRACING_DEBUG_VIZ_SPECULAR_COLOR:
+			Result = float4(Payload.SpecularColor.rgb, 1.0f);
+			break;
+		case RAY_TRACING_DEBUG_VIZ_OPACITY:
+			Result = float4(Payload.Opacity, Payload.Opacity, Payload.Opacity, 1.0f);
+			break;
+		case RAY_TRACING_DEBUG_VIZ_METALLIC:
+			Result = float4(Payload.Metallic, Payload.Metallic, Payload.Metallic, 1.0f);
+			break;
+		case RAY_TRACING_DEBUG_VIZ_SPECULAR:
+			Result = float4(Payload.Specular, Payload.Specular, Payload.Specular, 1.0f);
+			break;
+		case RAY_TRACING_DEBUG_VIZ_ROUGHNESS:
+			Result = float4(Payload.Roughness, Payload.Roughness, Payload.Roughness, 1.0f);
+			break;
+		case RAY_TRACING_DEBUG_VIZ_IOR:
+			Result = float4(Payload.Ior, Payload.Ior, Payload.Ior, 1.0f);
+			break;
+		case RAY_TRACING_DEBUG_VIZ_SHADING_MODEL:
+			Result = float4(Payload.ShadingModelID, Payload.ShadingModelID, Payload.ShadingModelID, 1.0f);
+			break;
+		case RAY_TRACING_DEBUG_VIZ_BLENDING_MODE:
+			Result = float4(Payload.BlendingMode, Payload.BlendingMode, Payload.BlendingMode, 1.0f);
+			break;
+		case RAY_TRACING_DEBUG_VIZ_LIGHTING_CHANNEL_MASK:
+			Result = float4(Payload.PrimitiveLightingChannelMask, Payload.PrimitiveLightingChannelMask, Payload.PrimitiveLightingChannelMask, 1.0f);
+			break;
+		case RAY_TRACING_DEBUG_VIZ_CUSTOM_DATA:
+			Result = float4(Payload.CustomData);
+			break;
+		case RAY_TRACING_DEBUG_VIZ_GBUFFER_AO:
+			Result = float4(Payload.GBufferAO, Payload.GBufferAO, Payload.GBufferAO, 1.0f);
+			break;
+		case RAY_TRACING_DEBUG_VIZ_INDIRECT_IRRADIANCE:
+			Result = float4(Payload.IndirectIrradiance, Payload.IndirectIrradiance, Payload.IndirectIrradiance, 1.0f);
+			break;
+		case RAY_TRACING_DEBUG_VIZ_WORLD_POSITION:
+			Result = float4(Payload.WorldPos, 1.0f);
+			break;
+		case RAY_TRACING_DEBUG_VIZ_HITKIND:
+			Result = float4(Payload.HitKind, Payload.HitKind, Payload.HitKind, 1.0f);
+			break;
+		}
+	}
+
+	Output[PixelCoord] = Result;
+}
+
+[shader("miss")]
+void RayTracingDebugMainMS(inout FPackedMaterialClosestHitPayload PackedPayload)
+{
+	PackedPayload.HitT = -1;
+}
+
+
+[shader("closesthit")]
+void RayTracingDebugMainCHS(inout FPackedMaterialClosestHitPayload PackedPayload, in FDefaultAttributes Attributes)
+{
+	FMaterialClosestHitPayload Payload = (FMaterialClosestHitPayload)0;
+
+	Payload.HitT = RayTCurrent();
+	Payload.BaseColor = float3(0.2, 0.3, 0.4);
+
+	FTriangleBaseAttributes Triangle = LoadTriangleBaseAttributes(PrimitiveIndex());
+	float3 Edge0 = Triangle.LocalPositions[2] - Triangle.LocalPositions[0];
+	float3 Edge1 = Triangle.LocalPositions[1] - Triangle.LocalPositions[0];
+	float3x3 WorldToLocal = (float3x3)WorldToObject();
+	float3x3 LocalToWorldNormal = transpose(WorldToLocal);
+	Payload.WorldNormal = normalize(mul(LocalToWorldNormal, cross(Edge0, Edge1)));
+
+    PackedPayload = PackRayTracingPayload(Payload);
+}
diff --git a/Engine/Shaders/Private/RayTracing_mv/RayTracingDebugDefinitions.ush b/Engine/Shaders/Private/RayTracing_mv/RayTracingDebugDefinitions.ush
new file mode 100644
--- /dev/null
+++ b/Engine/Shaders/Private/RayTracing_mv/RayTracingDebugDefinitions.ush
@@ -0,0 +1,29 @@
+// Copyright 1998-2018 Epic Games, Inc. All Rights Reserved.
+
+/*================================================================================================
+	RayTracingDebugDefinitions.ush: used in ray tracing debug shaders  and C++ code to define common constants
+	!!! Changing this file requires recompilation of the engine !!!
+=================================================================================================*/
+
+#pragma once
+
+#define RAY_TRACING_DEBUG_VIZ_RADIANCE					0 
+#define RAY_TRACING_DEBUG_VIZ_WORLD_NORMAL				1 
+#define RAY_TRACING_DEBUG_VIZ_BASE_COLOR				2 
+#define RAY_TRACING_DEBUG_VIZ_DIFFUSE_COLOR				3 
+#define RAY_TRACING_DEBUG_VIZ_SPECULAR_COLOR			4 
+#define RAY_TRACING_DEBUG_VIZ_OPACITY					5 
+#define RAY_TRACING_DEBUG_VIZ_METALLIC					6 
+#define RAY_TRACING_DEBUG_VIZ_SPECULAR					7 
+#define RAY_TRACING_DEBUG_VIZ_ROUGHNESS					8 
+#define RAY_TRACING_DEBUG_VIZ_IOR						9 
+#define RAY_TRACING_DEBUG_VIZ_SHADING_MODEL				10 
+#define RAY_TRACING_DEBUG_VIZ_BLENDING_MODE				11  
+#define RAY_TRACING_DEBUG_VIZ_LIGHTING_CHANNEL_MASK		12  
+#define RAY_TRACING_DEBUG_VIZ_CUSTOM_DATA				13  
+#define RAY_TRACING_DEBUG_VIZ_GBUFFER_AO				14  
+#define RAY_TRACING_DEBUG_VIZ_INDIRECT_IRRADIANCE		15  
+#define RAY_TRACING_DEBUG_VIZ_WORLD_POSITION			16  
+#define RAY_TRACING_DEBUG_VIZ_HITKIND					17  
+#define RAY_TRACING_DEBUG_VIZ_BARYCENTRICS				18  
+
diff --git a/Engine/Shaders/Private/RayTracing_mv/RayTracingDefinitions.ush b/Engine/Shaders/Private/RayTracing_mv/RayTracingDefinitions.ush
new file mode 100644
--- /dev/null
+++ b/Engine/Shaders/Private/RayTracing_mv/RayTracingDefinitions.ush
@@ -0,0 +1,16 @@
+// Copyright 1998-2018 Epic Games, Inc. All Rights Reserved.
+
+/*================================================================================================
+	RayTracingDefinitions.ush: used in ray tracing shaders  and C++ code to define common constants
+	!!! Changeing this file requires recompillation of the engine !!!
+=================================================================================================*/
+
+#pragma once
+
+#define RAY_TRACING_REGISTER_SPACE_LOCAL  0 // default register space for hit group (closest hit, any hit, intersection) shader resources
+#define RAY_TRACING_REGISTER_SPACE_GLOBAL 1 // register space for ray generation and miss shaders (#dxr_todo: make global resources also available to hit shaders)
+#define RAY_TRACING_REGISTER_SPACE_SYSTEM 2 // register space for "system" parameters (index buffer, vertex buffer, fetch parameters)
+
+// Change this to force recompilation of all ray tracing shaders (use https://www.random.org/cgi-bin/randbyte?nbytes=4&format=h)
+// This avoids changing the global ShaderVersion.ush and forcing recompilation of all shaders in the engine (only RT shaders will be affected)
+#define RAY_TRACING_SHADER_VERSION 0x510cacb6
diff --git a/Engine/Shaders/Private/RayTracing_mv/RayTracingDirectionalLight.ush b/Engine/Shaders/Private/RayTracing_mv/RayTracingDirectionalLight.ush
new file mode 100644
--- /dev/null
+++ b/Engine/Shaders/Private/RayTracing_mv/RayTracingDirectionalLight.ush
@@ -0,0 +1,77 @@
+// Copyright 1998-2018 Epic Games, Inc. All Rights Reserved.
+
+#pragma once
+
+#include "../PathTracing/Utilities/PathTracingRandomSequence.ush"
+#include "../MonteCarlo.ush"
+
+// #patrick todo: move to RayTracingCommon.ush
+float3 ReconstructWorldPositionFromDepth(float2 UV, float Depth)
+{
+	float2 ScreenPosition = (UV - View.ScreenPositionScaleBias.wz) / View.ScreenPositionScaleBias.xy;
+	float4 HomogeneousWorldPosition = mul(float4(ScreenPosition * Depth, Depth, 1), View.ScreenToWorld);
+	float3 WorldPosition = HomogeneousWorldPosition.xyz / HomogeneousWorldPosition.w;
+
+	return WorldPosition;
+}
+
+// Adapted from "A Low Distortion Map Between Disk and Square." Improved based
+// on comments from Peter Shirley's blog post: "Improved code for concentric map."
+float2 ToConcentricMap(float2 RectangularCoords)
+{
+	float R;
+	float Phi;
+
+	RectangularCoords = 2.0 * RectangularCoords - 1.0;
+	float2 RectangularCoordsSquared = RectangularCoords * RectangularCoords;
+	if (RectangularCoordsSquared.x > RectangularCoordsSquared.y)
+	{
+		R = RectangularCoords.x;
+		Phi = (PI / 4.0) * (RectangularCoords.y / RectangularCoords.x);
+	}
+	else
+	{
+		R = RectangularCoords.y;
+		Phi = (PI / 2.0) - (PI / 4.0) * (RectangularCoords.x / RectangularCoords.y);
+	}
+
+	float2 PolarCoords = float2(R, Phi);
+	return PolarCoords;
+}
+
+void GenerateDirectionalLightOcclusionRay(
+	float3 WorldPosition,
+	float3 WorldNormal,
+	float3 LightDirection,
+	float LightRadius,
+	float2 RandSample,
+	out float3 RayOrigin,
+	out float3 RayDirection,
+	out float RayTMin,
+	out float RayTMax)
+{
+	// Draw random variable and choose a point on a unit disk
+	float2 BufferSize = View.BufferSizeAndInvSize.xy;
+
+	float2 DiskUV = UniformSampleDiskConcentric(RandSample) * LightRadius;
+
+	// Permute light direction by user-defined radius on unit sphere
+	float3 N = LightDirection;
+	float3 dPdu = float3(1, 0, 0);
+	if (dot(N, dPdu) != 0)
+	{
+		dPdu = cross(N, dPdu);
+	}
+	else
+	{
+		dPdu = cross(N, float3(0, 1, 0));
+	}
+	float3 dPdv = cross(dPdu, N);
+	LightDirection += dPdu * DiskUV.x + dPdv * DiskUV.y;
+
+	// Apply normal perturbation when defining ray
+	RayOrigin = WorldPosition + WorldNormal * 0.01;
+	RayDirection = normalize(LightDirection);
+	RayTMin = 0.0;
+	RayTMax = 1.0e27;
+}
diff --git a/Engine/Shaders/Private/RayTracing_mv/RayTracingDynamicMesh.usf b/Engine/Shaders/Private/RayTracing_mv/RayTracingDynamicMesh.usf
new file mode 100644
--- /dev/null
+++ b/Engine/Shaders/Private/RayTracing_mv/RayTracingDynamicMesh.usf
@@ -0,0 +1,67 @@
+// Copyright 1998-2018 Epic Games, Inc. All Rights Reserved.
+
+/*=============================================================================
+	RayTracingDynamicMesh.usf
+=============================================================================*/
+
+#include "/Engine/Private/Common.ush"
+#include "/Engine/Generated/Material.ush"
+#include "/Engine/Generated/VertexFactory.ush"
+ 
+struct FRayTracingDynamicMeshVSToGS
+{
+	float4 WorldPosition : SV_POSITION;
+};
+
+void RayTracingDynamicMeshVS(
+	FVertexFactoryInput Input,
+	out FRayTracingDynamicMeshVSToGS Output
+	)
+{
+	ResolvedView = ResolveView();
+
+	FVertexFactoryIntermediates VFIntermediates = GetVertexFactoryIntermediates(Input);
+	Output.WorldPosition = VertexFactoryGetWorldPosition(Input, VFIntermediates) - float4(ResolvedView.PreViewTranslation, 0.0f);
+}
+
+// Must match GetUniformMeshStreamOutLayout and TRIANGLE_VERTEX_DATA_STRIDE
+struct FRayTracingDynamicMeshVertex
+{
+	float4 WorldPosition : SV_POSITION;
+};
+
+[maxvertexcount(3)]
+void RayTracingDynamicMeshGS(triangle FRayTracingDynamicMeshVSToGS Inputs[3], inout TriangleStream<FRayTracingDynamicMeshVertex> OutStream)
+{ 
+	for (int i = 0; i < 3; i++)
+	{
+		FRayTracingDynamicMeshVSToGS Input = Inputs[i];
+
+		FRayTracingDynamicMeshVertex Vertex;
+		Vertex.WorldPosition = Input.WorldPosition;
+
+		OutStream.Append(Vertex);
+	}
+}
+
+#if COMPUTESHADER
+uint NumVertices;
+RWBuffer<float> RWVertexPositions;
+
+[numthreads(256, 1, 1)]
+void RayTracingDynamicGeometryConverterCS(uint3 DispatchThreadId : SV_DispatchThreadID)
+{ 
+	ResolvedView = ResolveView();
+
+	uint VertexIndex = DispatchThreadId.x;
+	
+	if (VertexIndex >= NumVertices) return;
+
+	FVertexFactoryInput Input = LoadVertexFactoryInput(VertexIndex / 3, VertexIndex % 3);
+	FVertexFactoryIntermediates VFIntermediates = GetVertexFactoryIntermediates(Input);
+	float3 WorldPosition = (VertexFactoryGetWorldPosition(Input, VFIntermediates) - float4(ResolvedView.PreViewTranslation, 0.0f)).xyz;
+	RWVertexPositions[VertexIndex * 3 + 0] = WorldPosition.x;
+	RWVertexPositions[VertexIndex * 3 + 1] = WorldPosition.y;
+	RWVertexPositions[VertexIndex * 3 + 2] = WorldPosition.z;
+}
+#endif
diff --git a/Engine/Shaders/Private/RayTracing_mv/RayTracingHitGroupCommon.ush b/Engine/Shaders/Private/RayTracing_mv/RayTracingHitGroupCommon.ush
new file mode 100644
--- /dev/null
+++ b/Engine/Shaders/Private/RayTracing_mv/RayTracingHitGroupCommon.ush
@@ -0,0 +1,101 @@
+// Copyright 1998-2018 Epic Games, Inc. All Rights Reserved.
+
+/*=============================================================================
+	RayTracingHitGroupCommon.ush: common header used in all hit group shaders
+=============================================================================*/
+
+#pragma once
+
+#ifndef RAYTRACINGHITGROUPCOMMON_USH_INCLUDED
+#define RAYTRACINGHITGROUPCOMMON_USH_INCLUDED // Workarround for UE-66460
+
+#include "CommonRayTracingBuiltInResources.ush"
+
+uint3 LoadIndices16Bit(ByteAddressBuffer SourceBuffer, uint OffsetInBytes)
+{
+	uint3 Result;
+
+	// ByteAddressBuffer loads must be aligned to DWORD boundary.
+	// We can load 2 DWORDs (4 SHORTs) at a time and extract 3 SHORTs (first 3 or second 3).
+	uint AlignedOffsetInBytes = OffsetInBytes & (~3);
+	const uint2 PackedIndices = SourceBuffer.Load2(AlignedOffsetInBytes);
+
+	if (AlignedOffsetInBytes == OffsetInBytes)
+	{
+		// Extract first 3 SHORTs from 2 DWORDs
+		Result[0] = PackedIndices[0] & 0xffff;
+		Result[1] = PackedIndices[0] >> 16;
+		Result[2] = PackedIndices[1] & 0xffff;
+	}
+	else
+	{
+		// Extract second 3 SHORTs from 2 DWORDs
+		Result[0] = PackedIndices[0] >> 16;
+		Result[1] = PackedIndices[1] & 0xffff;
+		Result[2] = PackedIndices[1] >> 16;
+	}
+
+	return Result;
+}
+
+uint3 LoadIndices32Bit(ByteAddressBuffer SourceBuffer, uint OffsetInBytes)
+{
+	return SourceBuffer.Load3(OffsetInBytes);
+}
+
+float3 LoadVertexPositionFloat3(ByteAddressBuffer SourceBuffer, uint Index, uint StrideInBytes)
+{
+	uint OffsetInBytes = Index * StrideInBytes;
+	return asfloat(SourceBuffer.Load3(OffsetInBytes));
+}
+
+// Explicitly declare "system" resources that will be available to all hit group shaders.
+
+struct FTriangleBaseAttributes
+{
+	uint3 Indices;
+	float3 LocalPositions[3];
+};
+
+FTriangleBaseAttributes LoadTriangleBaseAttributes(uint PrimitiveId)
+{
+	FTriangleBaseAttributes Result = (FTriangleBaseAttributes)0;
+
+	// Fetch vertex indices and positions, then compute local space normal and transform it to world space
+
+	const uint IndexBufferOffsetInBytes = HitGroupSystemRootConstants.IndexBufferOffsetInBytes;
+	const uint IndexBufferStride = HitGroupSystemRootConstants.GetIndexStride();
+	const uint BaseIndex = PrimitiveId * 3;
+
+	if (IndexBufferStride == 0)
+	{
+		// Non-indexed geometry (implicit triangle list indices)
+		Result.Indices = uint3(BaseIndex, BaseIndex + 1, BaseIndex + 2);
+	}
+	else if (IndexBufferStride == 2)
+	{
+		Result.Indices = LoadIndices16Bit(HitGroupSystemIndexBuffer, IndexBufferOffsetInBytes + BaseIndex * IndexBufferStride);
+	}
+	else
+	{
+		Result.Indices = LoadIndices32Bit(HitGroupSystemIndexBuffer, IndexBufferOffsetInBytes + BaseIndex * IndexBufferStride);
+	}
+
+	// Fetch vertex positions (in local space)
+	const uint VertexStride = HitGroupSystemRootConstants.GetVertexStride();
+
+	// #dxr_todo: handle various vertex formats, for now only supporting float3
+	// #dxr_todo: handle BaseVertexIndex
+	Result.LocalPositions[0] = LoadVertexPositionFloat3(HitGroupSystemVertexBuffer, Result.Indices[0], VertexStride);
+	Result.LocalPositions[1] = LoadVertexPositionFloat3(HitGroupSystemVertexBuffer, Result.Indices[1], VertexStride);
+	Result.LocalPositions[2] = LoadVertexPositionFloat3(HitGroupSystemVertexBuffer, Result.Indices[2], VertexStride);
+
+	return Result;
+}
+
+uint GetInstanceUserData()
+{
+	return InstanceID();
+}
+
+#endif // RAYTRACINGHITGROUPCOMMON_USH_INCLUDED // Workarround for UE-66460
diff --git a/Engine/Shaders/Private/RayTracing_mv/RayTracingMaterialHitShaders.usf b/Engine/Shaders/Private/RayTracing_mv/RayTracingMaterialHitShaders.usf
new file mode 100644
--- /dev/null
+++ b/Engine/Shaders/Private/RayTracing_mv/RayTracingMaterialHitShaders.usf
@@ -0,0 +1,352 @@
+// Copyright 1998-2018 Epic Games, Inc. All Rights Reserved.
+
+#pragma one
+
+#define NEEDS_LIGHTMAP_COORDINATE		(HQ_TEXTURE_LIGHTMAP || LQ_TEXTURE_LIGHTMAP)
+#ifdef USING_TESSELLATION
+#undef USING_TESSELLATION
+// Needed for VertexFactoryInterpolate
+#define USING_TESSELLATION 1
+#endif
+
+// #dxr_todo: Temporarily hacking clip() away to prevent driver crash. Should be integrated into MaterialTemplate.ush in the future.
+#define clip(x) (x)
+#define ddx(x) 0
+#define ddy(x) 0
+
+#include "/Engine/Private/Common.ush"
+#include "RayTracingCommon.ush"
+#include "RayTracingHitGroupCommon.ush"
+
+#include "/Engine/Generated/Material.ush"
+#include "/Engine/Generated/VertexFactory.ush"
+
+#include "/Engine/Private/SHCommon.ush"
+#include "/Engine/Private/LightmapCommon.ush"
+
+#if SIMPLE_FORWARD_SHADING
+	#define GetEffectiveSkySHDiffuse GetSkySHDiffuseSimple
+#else
+	#define GetEffectiveSkySHDiffuse GetSkySHDiffuse
+#endif
+
+/** Computes sky diffuse lighting, including precomputed shadowing. */
+void GetSkyLighting(FMaterialPixelParameters MaterialParameters, float3 WorldNormal, float2 LightmapUV, float3 SkyOcclusionUV3D, out float3 OutDiffuseLighting, out float3 OutSubsurfaceLighting)
+{
+	OutDiffuseLighting = 0;
+	OutSubsurfaceLighting = 0;
+
+#if ENABLE_SKY_LIGHT
+
+	float SkyVisibility = 1;
+	float GeometryTerm = 1;
+	float3 SkyLightingNormal = WorldNormal;
+	
+	#if HQ_TEXTURE_LIGHTMAP || CACHED_POINT_INDIRECT_LIGHTING || CACHED_VOLUME_INDIRECT_LIGHTING || PRECOMPUTED_IRRADIANCE_VOLUME_LIGHTING
+		BRANCH
+		if (View.SkyLightParameters.x > 0)
+		{
+			#if PRECOMPUTED_IRRADIANCE_VOLUME_LIGHTING
+			
+				float3 SkyBentNormal = GetVolumetricLightmapSkyBentNormal(SkyOcclusionUV3D);
+				SkyVisibility = length(SkyBentNormal);
+				float3 NormalizedBentNormal = SkyBentNormal / max(SkyVisibility, .0001f);
+
+			#elif HQ_TEXTURE_LIGHTMAP
+
+				// Bent normal from precomputed texture
+				float4 WorldSkyBentNormalAndOcclusion = GetSkyBentNormalAndOcclusion(LightmapUV * float2(1, 2), MaterialParameters.SvPosition.xy);
+				// Renormalize as vector was quantized and compressed
+				float3 NormalizedBentNormal = normalize(WorldSkyBentNormalAndOcclusion.xyz);
+				SkyVisibility = WorldSkyBentNormalAndOcclusion.w;
+
+			#elif CACHED_POINT_INDIRECT_LIGHTING || CACHED_VOLUME_INDIRECT_LIGHTING
+ 
+				// Bent normal from the indirect lighting cache - one value for the whole object
+				float3 NormalizedBentNormal = IndirectLightingCache.PointSkyBentNormal.xyz;
+				SkyVisibility = IndirectLightingCache.PointSkyBentNormal.w;
+
+			#endif
+
+			#if (MATERIALBLENDING_TRANSLUCENT || MATERIALBLENDING_ADDITIVE) && (TRANSLUCENCY_LIGHTING_VOLUMETRIC_NONDIRECTIONAL || TRANSLUCENCY_LIGHTING_VOLUMETRIC_PERVERTEX_NONDIRECTIONAL)
+				// NonDirectional lighting can't depend on the normal
+				SkyLightingNormal = NormalizedBentNormal;
+			#else
+				
+				// Weight toward the material normal to increase directionality
+				float BentNormalWeightFactor = 1 - (1 - SkyVisibility) * (1 - SkyVisibility);
+
+				// We are lerping between the inputs of two lighting scenarios based on occlusion
+				// In the mostly unoccluded case, evaluate sky lighting with the material normal, because it has higher detail
+				// In the mostly occluded case, evaluate sky lighting with the bent normal, because it is a better representation of the incoming lighting
+				// Then treat the lighting evaluated along the bent normal as an area light, so we must apply the lambert term
+				SkyLightingNormal = lerp(NormalizedBentNormal, WorldNormal, BentNormalWeightFactor);
+
+				float DotProductFactor = lerp(saturate(dot(NormalizedBentNormal, WorldNormal)), 1, BentNormalWeightFactor);
+				// Account for darkening due to the geometry term
+				GeometryTerm = DotProductFactor;
+			#endif
+		}
+	#endif
+			
+	// Compute the preconvolved incoming lighting with the bent normal direction
+	float3 DiffuseLookup = GetEffectiveSkySHDiffuse(SkyLightingNormal) * ResolvedView.SkyLightColor.rgb;
+
+	// Apply AO to the sky diffuse
+	OutDiffuseLighting += DiffuseLookup * (SkyVisibility * GeometryTerm);
+
+	#if MATERIAL_SHADINGMODEL_TWOSIDED_FOLIAGE
+		float3 BackfaceDiffuseLookup = GetEffectiveSkySHDiffuse(-WorldNormal) * ResolvedView.SkyLightColor.rgb;
+		OutSubsurfaceLighting += BackfaceDiffuseLookup * SkyVisibility;
+	#endif
+
+#endif
+}
+
+/** Calculates indirect lighting contribution on this object from precomputed data. */
+void GetPrecomputedIndirectLightingAndSkyLight(
+	FMaterialPixelParameters MaterialParameters, 
+	FVertexFactoryInterpolantsVSToPS Interpolants,
+	float3 DiffuseDir,
+	float3 VolumetricLightmapBrickTextureUVs,
+	out float3 OutDiffuseLighting,
+	out float3 OutSubsurfaceLighting,
+	out float OutIndirectIrradiance)
+{
+	OutIndirectIrradiance = 0;
+	OutDiffuseLighting = 0;
+	OutSubsurfaceLighting = 0;
+	float2 SkyOcclusionUV = 0;
+
+	#if PRECOMPUTED_IRRADIANCE_VOLUME_LIGHTING
+		FThreeBandSHVectorRGB IrradianceSH = GetVolumetricLightmapSH3(VolumetricLightmapBrickTextureUVs);
+
+		// Diffuse convolution
+		FThreeBandSHVector DiffuseTransferSH = CalcDiffuseTransferSH3(DiffuseDir, 1);
+		OutDiffuseLighting = max(float3(0,0,0), DotSH3(IrradianceSH, DiffuseTransferSH)) / PI;
+
+		#if MATERIAL_SHADINGMODEL_TWOSIDED_FOLIAGE
+			FThreeBandSHVector SubsurfaceTransferSH = CalcDiffuseTransferSH3(-DiffuseDir, 1);
+			OutSubsurfaceLighting += max(float3(0,0,0), DotSH3(IrradianceSH, SubsurfaceTransferSH)) / PI;
+		#endif
+
+	// High quality texture lightmaps
+	#elif HQ_TEXTURE_LIGHTMAP
+		float2 LightmapUV0, LightmapUV1;
+		uint LightmapDataIndex;
+		GetLightMapCoordinates(Interpolants, LightmapUV0, LightmapUV1, LightmapDataIndex);
+		SkyOcclusionUV = LightmapUV0;
+
+		uint VTRequest = 0;
+		GetLightMapColorHQ(LightmapUV0, LightmapUV1, LightmapDataIndex, DiffuseDir, MaterialParameters.SvPosition.xy, OutDiffuseLighting, OutSubsurfaceLighting, VTRequest);
+
+	// Low quality texture lightmaps
+	#elif LQ_TEXTURE_LIGHTMAP
+
+		float2 LightmapUV0, LightmapUV1;
+		uint LightmapDataIndex;
+		GetLightMapCoordinates(Interpolants, LightmapUV0, LightmapUV1, LightmapDataIndex);
+		OutDiffuseLighting = GetLightMapColorLQ(LightmapUV0, LightmapUV1, LightmapDataIndex, DiffuseDir).rgb;
+
+	#endif
+
+	// Apply indirect lighting scale while we have only accumulated lightmaps
+	OutDiffuseLighting *= View.IndirectLightingColorScale;
+	OutSubsurfaceLighting *= View.IndirectLightingColorScale;
+
+	float3 SkyDiffuseLighting;
+	float3 SkySubsurfaceLighting;
+	GetSkyLighting(MaterialParameters, DiffuseDir, SkyOcclusionUV, VolumetricLightmapBrickTextureUVs, SkyDiffuseLighting, SkySubsurfaceLighting);
+
+	OutSubsurfaceLighting += SkySubsurfaceLighting;
+
+	// Sky lighting must contribute to IndirectIrradiance for ReflectionEnvironment lightmap mixing
+	OutDiffuseLighting += SkyDiffuseLighting;
+
+	#if HQ_TEXTURE_LIGHTMAP || LQ_TEXTURE_LIGHTMAP || CACHED_VOLUME_INDIRECT_LIGHTING || CACHED_POINT_INDIRECT_LIGHTING || PRECOMPUTED_IRRADIANCE_VOLUME_LIGHTING
+		OutIndirectIrradiance = Luminance(OutDiffuseLighting);
+	#endif
+}
+
+float4 ClipSpacePositionToSvPosition(float4 ClipSpacePosition)
+{
+	float3 NDCPos = ClipSpacePosition.xyz / ClipSpacePosition.w;
+
+	float4 SvPosition;
+	SvPosition.xy = ((NDCPos.xy * float2(0.5, -0.5) + 0.5f) * ResolvedView.ViewSizeAndInvSize.xy) + ResolvedView.ViewRectMin.xy;
+	SvPosition.z = NDCPos.z;
+	SvPosition.w = ClipSpacePosition.w;
+
+	return SvPosition;
+}
+
+void CalcInterpolants(in FDefaultAttributes Attributes, out FVertexFactoryInterpolantsVSToPS Interpolants, out float4 FakeSvPosition)
+{
+	FVertexFactoryInterpolantsVSToDS PerVertexInterpolants[3];
+	float4 PerVertexClipSpacePositions[3]; // Used to fake SvPosition
+	
+	for (int i = 0; i < 3; i++)
+	{
+		FVertexFactoryInput Input = LoadVertexFactoryInput(PrimitiveIndex(), i);
+	
+		FVertexFactoryIntermediates VFIntermediates = GetVertexFactoryIntermediates(Input);
+		float3x3 TangentToLocal = VertexFactoryGetTangentToLocal(Input, VFIntermediates);	
+		float4 WorldPositionExcludingWPO = VertexFactoryGetWorldPosition(Input, VFIntermediates);
+		FMaterialVertexParameters VertexParameters = GetMaterialVertexParameters(Input, VFIntermediates, WorldPositionExcludingWPO, TangentToLocal);
+		
+		float4 WorldPosition = WorldPositionExcludingWPO + float4(GetMaterialWorldPositionOffset(VertexParameters), 0.0f);
+		
+		PerVertexInterpolants[i] = VertexFactoryGetInterpolantsVSToDS(Input, VFIntermediates, VertexParameters);
+		
+		float4 RasterizedWorldPosition = VertexFactoryGetRasterizedWorldPosition(Input, VFIntermediates, WorldPosition);
+		PerVertexClipSpacePositions[i] = mul(RasterizedWorldPosition, ResolvedView.TranslatedWorldToClip);
+
+	}
+	
+	FVertexFactoryInterpolantsVSToDS Interpolated = 
+		VertexFactoryInterpolate(
+			VertexFactoryInterpolate(
+				PerVertexInterpolants[0], 1 - Attributes.Barycentrics.x - Attributes.Barycentrics.y,
+				PerVertexInterpolants[1], Attributes.Barycentrics.x),
+			1.0f,
+			PerVertexInterpolants[2], Attributes.Barycentrics.y);
+	
+	Interpolants = VertexFactoryAssignInterpolants(Interpolated);
+	
+	float4 ClipSpacePosition = PerVertexClipSpacePositions[0] * (1 - Attributes.Barycentrics.x - Attributes.Barycentrics.y) + PerVertexClipSpacePositions[1] * Attributes.Barycentrics.x + PerVertexClipSpacePositions[2] * Attributes.Barycentrics.y;
+	FakeSvPosition = ClipSpacePosition / ClipSpacePosition.w;
+}
+
+[shader("closesthit")]
+void MaterialCHS(inout FPackedMaterialClosestHitPayload PackedPayload, in FDefaultAttributes Attributes)
+{
+	ResolvedView = ResolveView();
+
+	FVertexFactoryInterpolantsVSToPS Interpolants;
+	float4 SvPosition;
+	
+	CalcInterpolants(Attributes, Interpolants, SvPosition);
+
+	FMaterialPixelParameters MaterialParameters = GetMaterialPixelParameters(Interpolants, SvPosition);
+	FPixelMaterialInputs PixelMaterialInputs;
+
+	{
+		float4 ScreenPosition = SvPositionToResolvedScreenPosition(SvPosition);
+		float3 TranslatedWorldPosition = WorldRayOrigin() + WorldRayDirection() * RayTCurrent() + ResolvedView.PreViewTranslation.xyz;
+		
+		bool bIsFrontFace = HitKind() == HIT_KIND_TRIANGLE_FRONT_FACE;
+
+		MaterialParameters.CameraVector = -WorldRayDirection();
+
+		// #dxr_todo: support world position offset at some point
+		// #if USE_WORLD_POSITION_EXCLUDING_SHADER_OFFSETS
+		// CalcMaterialParametersEx(MaterialParameters, PixelMaterialInputs, SvPosition, ScreenPosition, bIsFrontFace, TranslatedWorldPosition, BasePassInterpolants.PixelPositionExcludingWPO);
+		// #else
+		CalcMaterialParametersEx(MaterialParameters, PixelMaterialInputs, SvPosition, ScreenPosition, bIsFrontFace, TranslatedWorldPosition, TranslatedWorldPosition);
+		// #endif
+	}
+
+	FMaterialClosestHitPayload Payload = (FMaterialClosestHitPayload)0;
+
+	// Store the results in local variables and reuse instead of calling the functions multiple times.
+	half3 BaseColor = GetMaterialBaseColor(PixelMaterialInputs);
+	half2 Refraction = GetMaterialRefraction(PixelMaterialInputs);
+	half  Metallic = GetMaterialMetallic(PixelMaterialInputs);
+	half  Specular = GetMaterialSpecular(PixelMaterialInputs);
+	half Roughness = GetMaterialRoughness(PixelMaterialInputs);
+
+	Payload.WorldPos = MaterialParameters.AbsoluteWorldPosition;
+	Payload.WorldNormal = normalize(MaterialParameters.WorldNormal);
+	Payload.Radiance = GetMaterialEmissive(PixelMaterialInputs);
+	Payload.BaseColor = BaseColor;
+	Payload.Specular = Specular;
+	Payload.Roughness = Roughness;
+	Payload.Metallic = Metallic;
+	Payload.Ior = Refraction.x;
+	//Payload.Opacity = Opacity;
+
+// #if MATERIAL_SHADINGMODEL_UNLIT
+	// Payload.ShadingModelID = SHADINGMODELID_UNLIT;
+// #else
+	// Payload.ShadingModelID = SHADINGMODELID_DEFAULT_LIT;
+	// // TODO - support other ShadingModelIDs in path tracer
+// #endif
+
+	// TODO - support different blending modes in path tracer?
+	// uint MaterialBlendingMode = 0;
+// #if MATERIALBLENDING_ALPHACOMPOSITE
+	// MaterialBlendingMode = RT_MATERIAL_BLEND_MODE_ALPHA_COMPOSITE;
+// #elif MATERIALBLENDING_TRANSLUCENT
+	// MaterialBlendingMode = RT_MATERIAL_BLEND_MODE_TRANSLUCENT;
+// #elif MATERIALBLENDING_ADDITIVE
+	// MaterialBlendingMode = RT_MATERIAL_BLEND_MODE_ADDITIVE;
+// #elif MATERIALBLENDING_MODULATE
+	// MaterialBlendingMode = RT_MATERIAL_BLEND_MODE_MODULATE;
+// #endif
+	// Payload.BlendingMode = MaterialBlendingMode;
+
+	Payload.HitT = RayTCurrent();
+#if !MATERIAL_TWOSIDED
+    Payload.HitKind = HitKind();
+#else
+    Payload.HitKind = HIT_KIND_TRIANGLE_FRONT_FACE;
+#endif
+
+	float3 DiffuseIndirectLighting = 0;
+	float3 SubsurfaceIndirectLighting;
+	float IndirectIrradiance = 0;
+	
+	float3 VolumetricLightmapBrickTextureUVs;
+	
+#if PRECOMPUTED_IRRADIANCE_VOLUME_LIGHTING
+	VolumetricLightmapBrickTextureUVs = ComputeVolumetricLightmapBrickTextureUVs(MaterialParameters.AbsoluteWorldPosition);
+#endif
+	GetPrecomputedIndirectLightingAndSkyLight(MaterialParameters, Interpolants, Payload.WorldNormal, VolumetricLightmapBrickTextureUVs, DiffuseIndirectLighting, SubsurfaceIndirectLighting, IndirectIrradiance);
+	
+	Payload.DiffuseColor = BaseColor - BaseColor * Metallic;
+	Payload.SpecularColor = lerp(0.08 * Specular.xxx, BaseColor, Metallic);
+
+	PackedPayload = PackRayTracingPayload(Payload);
+}
+
+[shader("anyhit")] 
+void MaterialAHS(inout FPackedMaterialClosestHitPayload PackedPayload, in FDefaultAttributes Attributes)
+{
+#if MATERIALBLENDING_MASKED
+	ResolvedView = ResolveView();
+
+	FVertexFactoryInterpolantsVSToPS Interpolants;
+	float4 SvPosition;
+	
+	CalcInterpolants(Attributes, Interpolants, SvPosition);
+
+	FMaterialPixelParameters MaterialParameters = GetMaterialPixelParameters(Interpolants, SvPosition);
+	FPixelMaterialInputs PixelMaterialInputs;
+
+	{
+		float4 ScreenPosition = SvPositionToResolvedScreenPosition(SvPosition);
+		float3 TranslatedWorldPosition = WorldRayOrigin() + WorldRayDirection() * RayTCurrent() + ResolvedView.PreViewTranslation.xyz;
+		
+		bool bIsFrontFace = HitKind() == HIT_KIND_TRIANGLE_FRONT_FACE;
+
+		MaterialParameters.CameraVector = -WorldRayDirection();
+		
+		// #dxr_todo: support world position offset at some point
+		// #if USE_WORLD_POSITION_EXCLUDING_SHADER_OFFSETS
+		// CalcMaterialParametersEx(MaterialParameters, PixelMaterialInputs, SvPosition, ScreenPosition, bIsFrontFace, TranslatedWorldPosition, BasePassInterpolants.PixelPositionExcludingWPO);
+		// #else
+		CalcMaterialParametersEx(MaterialParameters, PixelMaterialInputs, SvPosition, ScreenPosition, bIsFrontFace, TranslatedWorldPosition, TranslatedWorldPosition);
+		// #endif
+	}
+
+	if (GetMaterialMask(PixelMaterialInputs) < 0)
+	{
+		IgnoreHit();
+	}
+	else
+	{
+		// Continue traversal by simply returning from this shader
+		return;
+	}
+#endif
+}  
diff --git a/Engine/Shaders/Private/RayTracing_mv/RayTracingOcclusionRG.usf b/Engine/Shaders/Private/RayTracing_mv/RayTracingOcclusionRG.usf
new file mode 100644
--- /dev/null
+++ b/Engine/Shaders/Private/RayTracing_mv/RayTracingOcclusionRG.usf
@@ -0,0 +1,285 @@
+// Copyright 1998-2018 Epic Games, Inc. All Rights Reserved.
+
+#include "../Common.ush"
+#include "../MonteCarlo.ush"
+#include "../DeferredShadingCommon.ush"
+#include "../LightShaderParameters.ush"
+#include "RayTracingCommon.ush"
+#include "RayTracingDirectionalLight.ush"
+#include "RayTracingRectLight.ush"
+
+/** Light types */
+#define LIGHT_TYPE_DIRECTIONAL 0
+#define LIGHT_TYPE_POINT 1
+#define LIGHT_TYPE_SPOT 2
+#define LIGHT_TYPE_RECT 3
+
+RaytracingAccelerationStructure TLAS;
+RWTexture2D<float4> RWOcclusionMaskUAV;
+RWTexture2D<float> RWRayDistanceUAV;
+
+uint SamplesPerPixel;
+
+float2 SobolIndexToUniformUnitSquare(uint2 SobolRand)
+{
+	return float2(SobolRand) * rcp(65536.0) + rcp(65536.0 * 2.0);
+}
+
+uint CalcLinearIndex(uint2 PixelCoord)
+{
+	return PixelCoord.y * View.BufferSizeAndInvSize.x + PixelCoord.x;
+}
+
+bool GeneratePointLightOcclusionRay(
+	FLightShaderParameters LightParameters,
+	float3 WorldPosition,
+	float3 WorldNormal,
+	float2 RandSample,
+	out float3 RayOrigin,
+	out float3 RayDirection,
+	out float RayTMin,
+	out float RayTMax
+)
+{
+	float3 LightDirection = LightParameters.Position - WorldPosition;
+	float RayLength = length(LightDirection);
+	LightDirection /= RayLength;
+
+	// Apply normal perturbation when defining ray
+	RayOrigin = WorldPosition + WorldNormal * 0.01;
+	RayDirection = LightDirection;
+	RayTMin = 0.0;
+	RayTMax = RayLength;
+	return true;
+}
+
+bool GenerateSpotLightOcclusionRay(
+	FLightShaderParameters LightParameters,
+	float3 WorldPosition,
+	float3 WorldNormal,
+	float2 RandSample,
+	out float3 RayOrigin,
+	out float3 RayDirection,
+	out float RayTMin,
+	out float RayTMax
+)
+{
+	float3 LightDirection = LightParameters.Position - WorldPosition;
+	float RayLength = length(LightDirection);
+	LightDirection /= RayLength;
+
+	// Apply culling
+	float CosAngle = LightParameters.SpotAngles.x;
+	if (dot(-LightDirection, LightParameters.Direction) < CosAngle)
+	{
+		return false;
+	}
+
+	// Apply normal perturbation when defining ray
+	RayOrigin = WorldPosition + WorldNormal * 0.01;
+	RayDirection = LightDirection;
+	RayTMin = 0.0;
+	RayTMax = RayLength;
+	return true;
+}
+
+bool GenerateOcclusionRay(
+	FLightShaderParameters LightParameters,
+	float3 WorldPosition,
+	float3 WorldNormal,
+	float2 RandSample,
+	out float3 RayOrigin,
+	out float3 RayDirection,
+	out float RayTMin,
+	out float RayTMax
+)
+{
+	#if LIGHT_TYPE == LIGHT_TYPE_DIRECTIONAL
+	{
+		// TODO: just passdown LightParameters to GenerateDirectionalLightOcclusionRay().
+		float3 LightDirection = LightParameters.Direction;
+		float LightRadius = LightParameters.SourceRadius;
+
+		GenerateDirectionalLightOcclusionRay(
+			WorldPosition, WorldNormal,
+			LightDirection, LightRadius,
+			RandSample,
+			/* out */ RayOrigin,
+			/* out */ RayDirection,
+			/* out */ RayTMin,
+			/* out */ RayTMax);
+	}
+	#elif LIGHT_TYPE == LIGHT_TYPE_POINT
+	{
+		return GeneratePointLightOcclusionRay(
+			LightParameters,
+			WorldPosition, WorldNormal,
+			RandSample,
+			/* out */ RayOrigin,
+			/* out */ RayDirection,
+			/* out */ RayTMin,
+			/* out */ RayTMax);
+	}
+	#elif LIGHT_TYPE == LIGHT_TYPE_SPOT
+	{
+		return GenerateSpotLightOcclusionRay(
+			LightParameters,
+			WorldPosition, WorldNormal,
+			RandSample,
+			/* out */ RayOrigin,
+			/* out */ RayDirection,
+			/* out */ RayTMin,
+			/* out */ RayTMax);
+	}
+	#elif LIGHT_TYPE == LIGHT_TYPE_RECT
+	{
+		float3 Tangent;
+		float3 BiTangent;
+		GetRectLightTangents(LightParameters, /* out */ Tangent, /* out */ BiTangent);
+
+		float2 LightDimensions = GetRectLightDimensions(LightParameters);
+		float RayPdf = 0.0;
+
+		// TODO: just passdown LightParameters to GenerateRectLightOcclusionRay().
+		return GenerateRectLightOcclusionRay(
+			WorldPosition, WorldNormal,
+			LightParameters.Position,
+			-LightParameters.Direction,
+			Tangent,
+			BiTangent,
+			LightDimensions.x, LightDimensions.y,
+			RandSample,
+			/* out */ RayOrigin,
+			/* out */ RayDirection,
+			/* out */ RayTMin,
+			/* out */ RayTMax,
+			/* out */ RayPdf);
+	}
+	#else
+		#error Unknown light type.
+	#endif
+	return true;
+}
+
+uint ReverseNBits(uint Bitfield, const uint BitCount)
+{
+	uint Out = 0;
+	UNROLL
+	for (uint i = 0; i < BitCount; i++)
+	{
+		uint bIsSet = ((Bitfield & (0x1 << i)) >> i);
+		Out |= bIsSet << (BitCount - 1 - i);
+	}
+	return Out;
+}
+
+[shader("raygeneration")]
+void OcclusionRG()
+{
+	uint2 PixelCoord = DispatchRaysIndex().xy;
+
+	RandomSequence RandSequence;
+	uint LinearIndex = CalcLinearIndex(PixelCoord);
+	RandomSequence_Initialize(RandSequence, LinearIndex, View.FrameNumber);
+
+	FLightShaderParameters LightParameters = GetRootLightShaderParameters();
+
+	// Get G-Buffer surface data
+	float2 InvBufferSize = View.BufferSizeAndInvSize.zw;
+	float2 UV = (float2(PixelCoord) + 0.5) * InvBufferSize;
+	FScreenSpaceData ScreenSpaceData = GetScreenSpaceData(UV);
+	float Depth = ScreenSpaceData.GBuffer.Depth;
+	float3 WorldPosition = ReconstructWorldPositionFromDepth(UV, Depth);
+	float3 WorldNormal = ScreenSpaceData.GBuffer.WorldNormal;
+
+	float Visibility = 0.0;
+	float RayDistance = 0.0;
+	float HitCount = 0.0;
+	float RayCount = 0.0;
+	uint LocalSamplesPerPixel = SamplesPerPixel;
+
+	// Mask out depth values at or beyond far plane
+	if (Depth >= SCENE_TEXTURES_DISABLED_SCENE_DEPTH_VALUE)
+	{
+		LocalSamplesPerPixel = 0.0;
+	}
+
+	LOOP
+	for (uint SampleIndex = 0; SampleIndex < LocalSamplesPerPixel; ++SampleIndex)
+	{
+		uint DummyVariable;
+		float2 RandSample = RandomSequence_GenerateSample2D(RandSequence, DummyVariable);
+
+		BRANCH
+		if (LocalSamplesPerPixel == 1 && 1)
+		{
+			const uint TemporalBits = 10;
+			uint FrameIndexMod1024 = ReverseNBits(GetPowerOfTwoModulatedFrameIndex(1 << TemporalBits), TemporalBits);
+
+			uint2 SobolBase = SobolPixel(PixelCoord);
+			uint2 SobolFrame = SobolIndex(SobolBase, FrameIndexMod1024, TemporalBits);
+
+			RandSample = SobolIndexToUniformUnitSquare(SobolFrame);
+		}
+
+		RayDesc Ray;
+		bool bIsValidRay = GenerateOcclusionRay(
+			LightParameters,
+			WorldPosition, WorldNormal,
+			RandSample,
+			/* out */ Ray.Origin,
+			/* out */ Ray.Direction,
+			/* out */ Ray.TMin,
+			/* out */ Ray.TMax);
+
+		BRANCH
+		if (!bIsValidRay)
+		{
+			continue;
+		}
+		else if (dot(WorldNormal, Ray.Direction) <= 0.0)
+		{
+			continue;
+		}
+
+		uint RayFlags = 0
+			| RAY_FLAG_CULL_BACK_FACING_TRIANGLES
+			| RAY_FLAG_FORCE_OPAQUE; // #dxr_todo: generate a special shader permutation for occlusion ray tracing that includes any-hit shaders for masked geometry
+
+		FDefaultPayload Payload = (FDefaultPayload)0;
+
+		TraceRay(
+			TLAS,      // AccelerationStructure
+			RayFlags,
+			0xFF,      // InstanceInclusionMask
+			0,         // RayContributionToHitGroupIndex
+			1,         // MultiplierForGeometryContributionToShaderIndex
+			0,         // MissShaderIndex
+			Ray,       // RayDesc
+			Payload    // Payload
+		);
+
+		RayCount += 1.0;
+		Visibility += Payload.HitT < 0 ? 1.0 : 0.0;
+		if (Payload.HitT >= 0.0)
+		{
+			RayDistance += Payload.HitT;
+			HitCount += 1.0;
+		}
+	}
+
+	RWOcclusionMaskUAV[PixelCoord] = (LocalSamplesPerPixel > 0) ? Visibility / LocalSamplesPerPixel : Visibility;
+
+	if (HitCount > 0.0)
+	{
+		RWRayDistanceUAV[PixelCoord] = RayDistance / HitCount;
+	}
+	else if (RayCount > 0.0)
+	{
+		RWRayDistanceUAV[PixelCoord] = 1.0e27;
+	}
+	else
+	{
+		RWRayDistanceUAV[PixelCoord] = -1.0;
+	}
+}
\ No newline at end of file
diff --git a/Engine/Shaders/Private/RayTracing_mv/RayTracingRectLight.ush b/Engine/Shaders/Private/RayTracing_mv/RayTracingRectLight.ush
new file mode 100644
--- /dev/null
+++ b/Engine/Shaders/Private/RayTracing_mv/RayTracingRectLight.ush
@@ -0,0 +1,50 @@
+// Copyright 1998-2018 Epic Games, Inc. All Rights Reserved.
+
+#pragma once
+
+#include "../PathTracing/Utilities/PathTracingRandomSequence.ush"
+
+bool GenerateRectLightOcclusionRay(
+	float3 WorldPosition,
+	float3 WorldNormal,
+	float3 LightPosition,
+	float3 LightNormal,
+	float3 LightdPdu,
+	float3 LightdPdv,
+	float LightWidth,
+	float LightHeight,
+	float2 RandSample,
+	out float3 RayOrigin,
+	out float3 RayDirection,
+	out float RayTMin,
+	out float RayTMax,
+	out float RayPdf
+)
+{
+	// Draw random variable
+	RandSample -= 0.5;
+
+	// Map sample point to quad
+	LightPosition += LightdPdu * LightWidth * RandSample.x + LightdPdv * LightHeight * RandSample.y;
+	float3 LightDirection = normalize(LightPosition - WorldPosition);
+
+	// Light-normal culling
+	if (dot(-LightDirection, LightNormal) <= 0.0)
+	{
+		return false;
+	}
+
+	// Apply normal perturbation when defining ray
+	RayOrigin = WorldPosition + WorldNormal * 0.01;
+	RayDirection = LightDirection;
+	RayTMin = 0.0;
+	RayTMax = length(LightPosition - WorldPosition);
+
+	// Convert PDF to solid angle measure
+	float Area = LightWidth * LightHeight;
+	float Distance2 = RayTMax * RayTMax;
+	float SolidAngle = Area * saturate(dot(-LightDirection, LightNormal)) / Distance2;
+	RayPdf = 1.0 / SolidAngle;
+
+	return true;
+}
diff --git a/Engine/Shaders/Private/RayTracing_mv/RayTracingRectLightOcclusionRG.usf b/Engine/Shaders/Private/RayTracing_mv/RayTracingRectLightOcclusionRG.usf
new file mode 100644
--- /dev/null
+++ b/Engine/Shaders/Private/RayTracing_mv/RayTracingRectLightOcclusionRG.usf
@@ -0,0 +1,274 @@
+// Copyright 1998-2018 Epic Games, Inc. All Rights Reserved.
+
+#include "../Common.ush"
+#include "../DeferredShadingCommon.ush"
+#include "../PathTracing/Utilities/PathTracingRandomSequence.ush"
+#include "RayTracingCommon.ush"
+#include "RayTracingRectLight.ush"
+#include "MipTreeCommon.ush"
+
+RaytracingAccelerationStructure TLAS;
+RWTexture2D<float4> RWOcclusionMaskUAV;
+RWTexture2D<float> RWRayDistanceUAV;
+
+uint CalcLinearIndex(uint2 PixelCoord)
+{
+	return PixelCoord.y * View.BufferSizeAndInvSize.x + PixelCoord.x;
+}
+
+// #patrick todo: move to RayTracingCommon.ush
+float3 ReconstructWorldPositionFromDepth(float2 UV, float Depth)
+{
+	float2 ScreenPosition = (UV - View.ScreenPositionScaleBias.wz) / View.ScreenPositionScaleBias.xy;
+	float4 HomogeneousWorldPosition = mul(float4(ScreenPosition * Depth, Depth, 1), View.ScreenToWorld);
+	float3 WorldPosition = HomogeneousWorldPosition.xyz / HomogeneousWorldPosition.w;
+
+	return WorldPosition;
+}
+
+bool GenerateRectLightImportanceSampledRay(
+	float3 WorldPosition,
+	float3 WorldNormal,
+	float3 LightPosition,
+	float3 LightNormal,
+	float3 LightdPdu,
+	float3 LightdPdv,
+	float LightWidth,
+	float LightHeight,
+	inout RandomSequence RandSequence,
+	out float3 RayOrigin,
+	out float3 RayDirection,
+	out float RayTMin,
+	out float RayTMax,
+	out float RayPdf
+)
+{
+	uint MipCount = Log2(RectLight.MipTreeDimensions.x);
+	uint2 MipPixel = 0;
+	float MipPdf = 1.0;
+
+	uint StopLevel = 0;
+	for (uint Index = 0; Index < MipCount - StopLevel; ++Index)
+	{
+		uint MipIndex = MipCount - Index - 1;
+		uint BufferOffset = BufferOffsetAtPixel(MipPixel, MipIndex, RectLight.MipTreeDimensions);
+		uint4 BufferOffsets = BufferOffset + uint4(0, 1, 2, 3);
+		float4 PowerTerm = float4(RectLight.MipTree[BufferOffsets.x], RectLight.MipTree[BufferOffsets.y], RectLight.MipTree[BufferOffsets.z], RectLight.MipTree[BufferOffsets.w]);
+		PowerTerm = max(PowerTerm, 0.0);
+
+		// Determine sample positions
+		uint DummyVariable;
+		float2 Jitter = RandomSequence_GenerateSample2D(RandSequence, DummyVariable);
+		float2 DimInv = 1.0 / DimensionsAtLevel(RectLight.MipTreeDimensions, MipIndex);
+		float2 UV00 = (MipPixel + Jitter) * DimInv - 0.5;
+		float2 UV10 = (MipPixel + Jitter + uint2(1, 0)) * DimInv - 0.5;
+		float2 UV01 = (MipPixel + Jitter + uint2(0, 1)) * DimInv - 0.5;
+		float2 UV11 = (MipPixel + Jitter + uint2(1, 1)) * DimInv - 0.5;
+
+		float3 Position00 = LightPosition + LightdPdu * LightWidth * UV00.x + LightdPdv * LightHeight * UV00.y;
+		float3 Position10 = LightPosition + LightdPdu * LightWidth * UV10.x + LightdPdv * LightHeight * UV10.y;
+		float3 Position01 = LightPosition + LightdPdu * LightWidth * UV01.x + LightdPdv * LightHeight * UV01.y;
+		float3 Position11 = LightPosition + LightdPdu * LightWidth * UV11.x + LightdPdv * LightHeight * UV11.y;
+
+		// Evaluate irradiance coefficients
+		float3 LightDirection00 = WorldPosition - Position00;
+		float3 LightDirection10 = WorldPosition - Position10;
+		float3 LightDirection01 = WorldPosition - Position01;
+		float3 LightDirection11 = WorldPosition - Position11;
+
+		float4 DistanceSquaredTerm = float4(
+			dot(LightDirection00, LightDirection00),
+			dot(LightDirection10, LightDirection10),
+			dot(LightDirection01, LightDirection01),
+			dot(LightDirection11, LightDirection11)
+		);
+
+		float4 CosineTerm = float4(
+			saturate(dot(normalize(LightDirection00), LightNormal)),
+			saturate(dot(normalize(LightDirection10), LightNormal)),
+			saturate(dot(normalize(LightDirection01), LightNormal)),
+			saturate(dot(normalize(LightDirection11), LightNormal))
+		);
+
+		// Build CDF on approximate irradiance
+		float4 Cdf = PowerTerm * CosineTerm / DistanceSquaredTerm;
+		Cdf.y += Cdf.x;
+		Cdf.z += Cdf.y;
+		Cdf.w += Cdf.z;
+		Cdf /= Cdf.w;
+
+		// Sample CDF
+		float RandSample = RandomSequence_GenerateSample1D(RandSequence, DummyVariable);
+		if (RandSample < Cdf.x)
+		{
+			MipPdf *= Cdf.x;
+		}
+		else if (RandSample < Cdf.y)
+		{
+			MipPdf *= Cdf.y - Cdf.x;
+			MipPixel.x += 1;
+		}
+		else if (RandSample < Cdf.z)
+		{
+			MipPdf *= Cdf.z - Cdf.y;
+			MipPixel.y += 1;
+		}
+		else // if (RandSample < Cdf.w)
+		{
+			MipPdf *= Cdf.w - Cdf.z;
+			MipPixel += 1;
+		}
+
+		MipPixel = uint2(MipPixel.x << 1, MipPixel.y << 1);
+	}
+	MipPixel = uint2(MipPixel.x >> 1, MipPixel.y >> 1);
+
+	// Sample within stop-level texel
+	uint DummyVariable;
+	float2 PixelOffset = RandomSequence_GenerateSample2D(RandSequence, DummyVariable);
+	float2 DimensionsAtStopLevel = DimensionsAtLevel(RectLight.MipTreeDimensions, StopLevel);
+	float2 MipPixelUV = (MipPixel + PixelOffset) / DimensionsAtStopLevel;
+	MipPdf *= DimensionsAtStopLevel;
+
+	// Map sample point to quad
+	MipPixelUV -= 0.5;
+	LightPosition += LightdPdu * LightWidth * MipPixelUV.x + LightdPdv * LightHeight * MipPixelUV.y;
+	float3 LightDirection = normalize(LightPosition - WorldPosition);
+
+	// Light-normal culling
+	if (dot(-LightDirection, LightNormal) <= 0.0)
+	{
+		RayPdf = 0.0;
+		return false;
+	}
+
+	// #dxr_todo: Performing direct lighting here could save some recalculation
+
+	// Define ray
+	RayOrigin = WorldPosition + WorldNormal * 0.01;
+	RayDirection = LightDirection;
+	RayTMin = 0.0;
+	RayTMax = length(LightPosition - WorldPosition);
+
+	// Determine area pdf
+	float Area = (RectLight.Width * RectLight.Height) / DimensionsAtStopLevel;
+	float Distance2 = RayTMax * RayTMax;
+	float SolidAngle = Area * saturate(dot(-LightDirection, LightNormal)) / Distance2;
+	float SolidAnglePdf = 1.0 / SolidAngle;
+	RayPdf = MipPdf * SolidAnglePdf;
+
+	return RayTMax > RayTMin;
+}
+
+[shader("raygeneration")]
+void RectLightOcclusionRG()
+{
+	uint2 PixelCoord = DispatchRaysIndex().xy;
+
+	RandomSequence RandSequence;
+	uint LinearIndex = CalcLinearIndex(PixelCoord);
+	RandomSequence_Initialize(RandSequence, LinearIndex, View.FrameNumber);
+
+	// Get G-Buffer surface data
+	float2 InvBufferSize = View.BufferSizeAndInvSize.zw;
+	float2 UV = (float2(PixelCoord) + 0.5) * InvBufferSize;
+	FScreenSpaceData ScreenSpaceData = GetScreenSpaceData(UV);
+	float Depth = ScreenSpaceData.GBuffer.Depth;
+	float3 WorldPosition = ReconstructWorldPositionFromDepth(UV, Depth);
+	float3 WorldNormal = ScreenSpaceData.GBuffer.WorldNormal;
+	float3 DiffuseColor = ScreenSpaceData.GBuffer.DiffuseColor;
+
+	float3 LightRadiance = RectLight.Color;
+
+	float RayDistance = 0.0;
+	float HitCount = 0.0;
+	uint SamplesPerPixel = RectLight.SamplesPerPixel;
+
+	// Mask out depth values at or beyond far plane
+	if (Depth >= SCENE_TEXTURES_DISABLED_SCENE_DEPTH_VALUE)
+	{
+		SamplesPerPixel = 0.0;
+	}
+
+	float3 ExitantRadiance = 0.0;
+	for (uint SampleIndex = 0; SampleIndex < SamplesPerPixel; ++SampleIndex)
+	{
+		RayDesc Ray;
+		float3 LightPosition = RectLight.Position;
+		float3 LightNormal = RectLight.Normal;
+		float3 LightdPdu = RectLight.dPdu;
+		float3 LightdPdv = RectLight.dPdv;
+		float LightWidth = RectLight.Width;
+		float LightHeight = RectLight.Height;
+
+		uint DummyVariable;
+		float2 RandSample = RandomSequence_GenerateSample2D(RandSequence, DummyVariable);
+		float RayPdf = 0.0;
+
+		bool bIsValid = RectLight.bIsTextureImportanceSampling ?
+			GenerateRectLightImportanceSampledRay(WorldPosition, WorldNormal, LightPosition, LightNormal, LightdPdu, LightdPdv, LightWidth, LightHeight, RandSequence, Ray.Origin, Ray.Direction, Ray.TMin, Ray.TMax, RayPdf) :
+			GenerateRectLightOcclusionRay(WorldPosition, WorldNormal, LightPosition, LightNormal, LightdPdu, LightdPdv, LightWidth, LightHeight, RandSample, Ray.Origin, Ray.Direction, Ray.TMin, Ray.TMax, RayPdf);
+		if (!bIsValid)
+		{
+			continue;
+		}
+
+		float NoL = dot(WorldNormal, Ray.Direction);
+		if (NoL <= 0.0)
+		{
+			continue;
+		}
+
+		uint RayFlags = 0
+			| RAY_FLAG_CULL_BACK_FACING_TRIANGLES
+			| RAY_FLAG_FORCE_OPAQUE; // #dxr_todo: generate a special shader permutation for occlusion ray tracing that includes any-hit shaders for masked geometry
+
+		FDefaultPayload Payload = (FDefaultPayload)0;
+
+		TraceRay(
+			TLAS,      // AccelerationStructure
+			RayFlags,
+			0xFF,      // InstanceInclusionMask
+			0,         // RayContributionToHitGroupIndex
+			1,         // MultiplierForGeometryContributionToShaderIndex
+			0,         // MissShaderIndex
+			Ray,       // RayDesc
+			Payload    // Payload
+		);
+
+		if (Payload.HitT >= 0.0)
+		{
+			RayDistance += Payload.HitT;
+			HitCount += 1.0;
+		}
+		else
+		{
+#if CALC_DIRECT_LIGHTING
+			// Sample light texture
+			float3 LightPosition = Ray.Origin + Ray.Direction * Ray.TMax;
+			float3 RectLightCorner = RectLight.Position - 0.5 * (RectLight.dPdu * RectLight.Width + RectLight.dPdv * RectLight.Height);
+			float3 DeltaPosition = LightPosition - RectLightCorner;
+			float2 UV = float2(dot(DeltaPosition, RectLight.dPdu) / RectLight.Width, dot(DeltaPosition, RectLight.dPdv) / RectLight.Height);
+			UV = 1.0 - UV;
+			float3 TextureColor = Texture2DSample(RectLight.Texture, RectLight.TextureSampler, UV);
+			float3 LightRadiance = RectLight.Color * TextureColor;
+			// #dxr_todo: Replace with actual BRDF
+			float3 Brdf = DiffuseColor;
+			float Distance2 = Ray.TMax * Ray.TMax;
+			ExitantRadiance += LightRadiance * Brdf * saturate(NoL) / RayPdf;
+#else
+			// Record exitant radiance as visibility score
+			ExitantRadiance += 1.0;
+#endif
+		}
+	}
+
+	float3 Result = (SamplesPerPixel > 0) ? ExitantRadiance / SamplesPerPixel : ExitantRadiance;
+#if	ENCODE_VISIBILITY
+	// Encode visibility term to match shadow map pipeline
+	Result = sqrt(Result);
+#endif
+	RWOcclusionMaskUAV[PixelCoord].rgb = Result;
+	RWOcclusionMaskUAV[PixelCoord].a = 1.0;
+	RWRayDistanceUAV[PixelCoord] = (HitCount > 0.0) ? RayDistance / HitCount : 1.0e27;
+}
diff --git a/Engine/Shaders/Private/RayTracing_mv/RayTracingReflectionsRG.usf b/Engine/Shaders/Private/RayTracing_mv/RayTracingReflectionsRG.usf
new file mode 100644
--- /dev/null
+++ b/Engine/Shaders/Private/RayTracing_mv/RayTracingReflectionsRG.usf
@@ -0,0 +1,380 @@
+// Copyright 1998-2018 Epic Games, Inc. All Rights Reserved.
+
+#ifndef DIM_DENOISER_OUTPUT
+	#define DIM_DENOISER_OUTPUT 0
+#endif
+
+#include "../Common.ush"
+
+float2 SobolIndexToUniformUnitSquare(uint2 SobolRand)
+{
+	return float2(SobolRand) * rcp(65536.0) + rcp(65536.0 * 2.0);
+}
+
+#define SUPPORT_CONTACT_SHADOWS		0
+#define USE_SOURCE_TEXTURE			0 //Used in RectLight.ush  //#dxr_todo: enable textures
+
+// Light types: should match SceneTypes.h until there is a common header
+#define LIGHT_TYPE_DIRECTIONAL		0 
+#define LIGHT_TYPE_POINT			1 
+#define LIGHT_TYPE_SPOT				2 
+#define LIGHT_TYPE_RECT				3 
+#define LIGHT_TYPE_MAX				4 
+
+#include "../DeferredShadingCommon.ush"
+#include "../DeferredLightingCommon.ush"
+#include "../Montecarlo.ush"
+#include "../PathTracing/Utilities/PathTracingRandomSequence.ush" // #dxr-todo: move to a common folder
+#include "RayTracingCommon.ush"
+#include "RayTracingHitGroupCommon.ush"
+
+uint SamplesPerPixel;
+
+RaytracingAccelerationStructure TLAS;
+
+RWTexture2D<float4> ColorOutput;
+
+uint CalcLinearIndex(uint2 PixelCoord)
+{
+	return PixelCoord.y * View.BufferSizeAndInvSize.x + PixelCoord.x;
+}
+
+// #dxr-todo: move to RayTracingCommon.ush
+float3 ReconstructWorldPositionFromDepth(float2 UV, float Depth)
+{
+	float2 ScreenPosition = (UV - View.ScreenPositionScaleBias.wz) / View.ScreenPositionScaleBias.xy;
+	float4 HomogeneousWorldPosition = mul(float4(ScreenPosition * Depth, Depth, 1), View.ScreenToWorld);
+	float3 WorldPosition = HomogeneousWorldPosition.xyz / HomogeneousWorldPosition.w;
+
+	return WorldPosition;
+}
+
+float3 GenerateReflectedRayDirection(
+	float3 IncidentDirection,
+	float3 WorldNormal,
+	float Roughness,
+	float2 RandSample
+)
+{
+	float3 RayDirection;
+	if (Roughness < 0.05 ) //ReflectionSmoothClamp)
+	{
+		RayDirection = reflect(IncidentDirection, WorldNormal);
+	}
+	else
+	{
+		float NoV = saturate(dot(-IncidentDirection, WorldNormal));
+		float2 E = RandSample;
+		float4 Sample = ImportanceSampleGGX(E, Pow4(Roughness));
+		float3 H = TangentToWorld(Sample.xyz, WorldNormal);
+		RayDirection = reflect(IncidentDirection, H);
+	}
+
+	return RayDirection;
+}
+
+#if 0
+float3 GatherRadiance(float CompositeAlpha, float3 WorldPosition, float3 RayDirection, float Roughness, float IndirectIrradiance)
+{
+	float IndirectSpecularOcclusion = 1.0f;
+	float3 ExtraIndirectSpecular = 0;
+	return CompositeReflectionCapturesAndSkylight(CompositeAlpha, WorldPosition, RayDirection, Roughness, IndirectIrradiance, IndirectSpecularOcclusion, ExtraIndirectSpecular, 0, 0);
+}
+
+float4 BlendWithEnvironment(inout float3 BRDFThroughput, const in FMaterialClosestHitPayload CurrentBounceHitInfo, float3 WorldPosition, float4 ReflectionColor, float3 RayDirection, float3 Normal)
+{
+#if USE_PREEXPOSURE
+	const float PreExposure = View.PreExposure;
+#else
+	const float PreExposure = 1.f;
+#endif
+
+	float3 ProbeReflectDirection = reflect(RayDirection, Normal);
+	float  NoV = saturate(dot(-RayDirection, Normal));
+	float IndirectIrradiance = CurrentBounceHitInfo.IndirectIrradiance;
+
+	// Point lobe in off-specular peak direction
+	ProbeReflectDirection = GetOffSpecularPeakReflectionDir(Normal, ProbeReflectDirection, CurrentBounceHitInfo.Roughness);
+
+	float RoughnessFade = saturate(CurrentBounceHitInfo.Roughness * RTRRoughnessFadeFactor + 2);
+	ReflectionColor.a = RoughnessFade; 	
+	  
+	float4 FinalReflectionColor = float4(ReflectionColor.rgb, 1 - ReflectionColor.a); 
+	float AO = CurrentBounceHitInfo.GBufferAO;
+	float RoughnessSq = CurrentBounceHitInfo.Roughness * CurrentBounceHitInfo.Roughness;
+	float SpecularOcclusion = GetSpecularOcclusion(NoV, RoughnessSq, AO);
+	FinalReflectionColor.a *= SpecularOcclusion;
+
+	FinalReflectionColor.rgb += PreExposure * GatherRadiance(FinalReflectionColor.a, WorldPosition, ProbeReflectDirection, CurrentBounceHitInfo.Roughness, IndirectIrradiance);
+
+	float3 BRDF = EnvBRDF(CurrentBounceHitInfo.SpecularColor, CurrentBounceHitInfo.Roughness, NoV);
+	FinalReflectionColor.rgb *= BRDF;
+
+	FinalReflectionColor.rgb *= BRDFThroughput;
+	BRDFThroughput *= BRDF;
+
+	return FinalReflectionColor;
+}
+#endif
+
+float3 ComputeDirectLighting(RayDesc Ray, FMaterialClosestHitPayload Payload)
+{
+	float3 DirectLighting = float3(0.0, 0.0, 0.0);
+
+	float3 WorldPosition = Payload.WorldPos;
+	float3 CameraVector = -Ray.Direction;
+	float AmbientOcclusion = 1.0;
+
+	// #dxr_todo: review when interpolants fill the Payload with all the data we need
+	FGBufferData GBufferData = (FGBufferData)0;
+	GBufferData.Depth = 1.f; // Do not use depth
+	GBufferData.WorldNormal = Payload.WorldNormal;
+	GBufferData.BaseColor = Payload.BaseColor;
+	GBufferData.Metallic = Payload.Metallic;
+	GBufferData.Specular = Payload.Specular;
+	GBufferData.Roughness = Payload.Roughness;
+	GBufferData.ShadingModelID = Payload.ShadingModelID;
+	GBufferData.CustomData = Payload.CustomData;
+	GBufferData.GBufferAO = Payload.GBufferAO;
+	GBufferData.IndirectIrradiance = Payload.IndirectIrradiance;
+	GBufferData.ShadingModelID = 1; 
+	   
+	//#dxr_todo: calculate on the fly or store in Payload?
+#if 1
+	GBufferData.SpecularColor = lerp(0.08 * GBufferData.Specular.xxx, GBufferData.BaseColor, GBufferData.Metallic);
+	GBufferData.DiffuseColor = GBufferData.BaseColor - GBufferData.BaseColor * GBufferData.Metallic;
+#else	
+	GBufferData.SpecularColor = Payload.SpecularColor;	
+	GBufferData.DiffuseColor = Payload.DiffuseColor;	
+#endif
+
+	//#dxr_todo: think on the more efficient way of iterate through lights (split in light types, etc)
+	for (uint DirIndex = 0; DirIndex < ReflectionLightsData.Count; DirIndex++)
+	{
+		FDeferredLightData LightData = (FDeferredLightData)0;
+
+		uint LightType = ReflectionLightsData.Type[DirIndex];
+		LightData.Position = ReflectionLightsData.LightPosition[DirIndex];
+		LightData.InvRadius	= ReflectionLightsData.LightInvRadius[DirIndex];
+		LightData.Color = ReflectionLightsData.LightColor[DirIndex];
+		LightData.FalloffExponent = ReflectionLightsData.LightFalloffExponent[DirIndex];
+		LightData.Direction = ReflectionLightsData.Direction[DirIndex];
+		LightData.Tangent = ReflectionLightsData.Tangent[DirIndex];
+		LightData.SpotAngles = ReflectionLightsData.SpotAngles[DirIndex];
+		LightData.SourceRadius = ReflectionLightsData.SourceRadius[DirIndex];
+		LightData.SourceLength = ReflectionLightsData.SourceLength[DirIndex];
+		LightData.SoftSourceRadius = ReflectionLightsData.SoftSourceRadius[DirIndex];
+		LightData.SpecularScale = ReflectionLightsData.SpecularScale[DirIndex];
+		LightData.ContactShadowLength = 0.0;
+		LightData.DistanceFadeMAD = ReflectionLightsData.DistanceFadeMAD[DirIndex];
+		LightData.ShadowMapChannelMask = float4(0, 0, 0, 0); // #dxr_todo;
+		LightData.ShadowedBits = 0; // Not lit dynamic shadows 
+		LightData.ContactShadowLengthInWS = false;
+
+		LightData.bRadialLight = (LightType != LIGHT_TYPE_DIRECTIONAL);
+		LightData.bSpotLight = (LightType == LIGHT_TYPE_SPOT);
+		LightData.bRectLight = (LightType == LIGHT_TYPE_RECT);
+
+		if (LightType == LIGHT_TYPE_DIRECTIONAL)
+		{
+			LightData.bInverseSquared = false;
+		}
+		else
+		{
+			LightData.bInverseSquared = LightData.FalloffExponent == 0;
+		}
+
+		// #dxr_todo: get a valid texture once they are properly bound
+		Texture2D LightSourceTexture = ReflectionLightsData.DummyRectLightTexture;
+
+		float4 LightAttenuation = 1.0f;
+		float LightProfileMultiplier = 1.0;//#dxr_todo
+
+		float3 LightContribution = GetDynamicLighting(WorldPosition, -CameraVector, GBufferData, AmbientOcclusion, GBufferData.ShadingModelID, LightData, LightAttenuation, 0.5, uint2(0, 0), LightSourceTexture).xyz;
+		DirectLighting += LightContribution * LightProfileMultiplier;
+	}
+
+	return DirectLighting;
+}
+
+FMaterialClosestHitPayload TraceRayAndAccumulateResults(
+	RayDesc Ray,
+	inout float3 Radiance)
+{
+	// Call TraceRay iteratively until alpha is 0 and accumulate the result
+	//#dxr-todo: iterate until opacity == 1 if the material is not opaque
+	//#dxr-todo: add optimizations (TMax depending on roughness, etc)
+
+	uint RayFlags = RAY_FLAG_CULL_BACK_FACING_TRIANGLES;
+
+	FMaterialClosestHitPayload Payload = TraceRayInternal(
+		TLAS,      // AccelerationStructure
+		RayFlags,
+		0xFF,      // InstanceInclusionMask
+		0,         // RayContributionToHitGroupIndex
+		1,         // MultiplierForGeometryContributionToShaderIndex
+		0,         // MissShaderIndex
+		Ray        // RayDesc
+	);
+
+
+	float3 DirectLighting = ComputeDirectLighting(Ray, Payload);
+
+	// Transform NaNs to black, transform negative colors to black.
+	DirectLighting = -min(-DirectLighting, 0.0);
+
+	Radiance += DirectLighting;
+
+	return Payload;
+}
+
+uint ReverseNBits(uint Bitfield, const uint BitCount) // TODO: where to put that guy?
+{
+	uint Out = 0;
+	UNROLL
+	for (uint i = 0; i < BitCount; i++)
+	{
+		uint bIsSet = ((Bitfield & (0x1 << i)) >> i);
+		Out |= bIsSet << (BitCount - 1 - i);
+	}
+	return Out;
+}
+
+
+[shader("raygeneration")]
+void RayTracingReflectionsRG()
+{
+	uint2 PixelCoord = DispatchRaysIndex().xy;
+
+	RandomSequence RandSequence;
+	uint LinearIndex = CalcLinearIndex(PixelCoord);
+	RandomSequence_Initialize(RandSequence, LinearIndex, View.FrameNumber);
+
+	float2 InvBufferSize = View.BufferSizeAndInvSize.zw;
+	float2 UV = (float2(PixelCoord) + 0.5) * InvBufferSize;
+	FScreenSpaceData ScreenSpaceData = GetScreenSpaceData(UV);
+	float Depth = ScreenSpaceData.GBuffer.Depth;
+	float3 CameraOrigin = ReconstructWorldPositionFromDepth(UV, 0.0f);
+	float3 WorldPosition = ReconstructWorldPositionFromDepth(UV, Depth);
+	float3 CameraDirection = normalize(WorldPosition - CameraOrigin);
+	//#dxr - todo: use also the shading normal (interpolate, or apply a more sophisticated approach)
+	float3 WorldNormal = ScreenSpaceData.GBuffer.WorldNormal;
+	float Roughness = ScreenSpaceData.GBuffer.Roughness;
+	
+	float4 ReflectedColor = 0.0;
+	float3 BRDFThroughput = float3(1.0f, 1.0f, 1.0f);
+
+	uint MaxBounces = 1; //#dxr-todo: only works with 1 bounce until hit shaders are supported
+	uint LocalSamplesPerPixel = SamplesPerPixel;
+	float RoughnessThreshold = 1.0f; //#dxr-todo: for internal testing only, remove later
+
+	// Mask out depth values at or beyond far plane
+	if (Depth >= SCENE_TEXTURES_DISABLED_SCENE_DEPTH_VALUE)
+	{
+		LocalSamplesPerPixel = 0;
+	}
+
+	float HitDistance = 0.0;
+	float3 Radiance = float3(0, 0, 0);
+
+	if (Roughness < RoughnessThreshold)
+	{
+		for (uint SampleIndex = 0; SampleIndex < LocalSamplesPerPixel; ++SampleIndex)
+		{
+			float3 IncidentDirection = CameraDirection;
+			for (uint BounceIndex = 0; BounceIndex < MaxBounces; ++BounceIndex)
+			{
+				RayDesc Ray;
+				Ray.Origin = WorldPosition + WorldNormal * 0.01;
+				Ray.TMin = 0.0;
+				Ray.TMax = 1e27f;
+
+				uint DummyVariable;
+				float2 RandSample = RandomSequence_GenerateSample2D(RandSequence, DummyVariable);
+
+				if (LocalSamplesPerPixel == 1)
+				{
+					const uint TemporalBits = 10;
+					uint FrameIndexMod1024 = ReverseNBits(GetPowerOfTwoModulatedFrameIndex(1 << TemporalBits), TemporalBits);
+
+					uint2 SobolBase = SobolPixel(PixelCoord);
+					uint2 SobolFrame = SobolIndex(SobolBase, FrameIndexMod1024, TemporalBits);
+
+					RandSample = SobolIndexToUniformUnitSquare(SobolFrame);
+				}
+
+				Ray.Direction = GenerateReflectedRayDirection(IncidentDirection, WorldNormal, Roughness, RandSample);
+				if (dot(WorldNormal, Ray.Direction) <= 0.0)
+				{
+					// Refraction rays is handled separately
+					break;
+				}
+
+				//#dxr-todo: apply heuristics to do different shading computations depending on BounceIndex
+
+				uint RayFlags = RAY_FLAG_CULL_BACK_FACING_TRIANGLES;
+
+				FMaterialClosestHitPayload Payload = TraceRayAndAccumulateResults(Ray, Radiance);
+
+				if (Payload.HitT < 0)
+				{
+					break;
+				}
+
+#if 1
+				float4 PathRadiance = float4(Radiance, 1.0); 
+				if (BounceIndex == 0)
+				{
+					ReflectedColor += PathRadiance;
+				}
+				else
+				{
+					float4 FinalReflectionColor = PathRadiance;
+					//float4 FinalReflectionColor = BlendWithEnvironment(BRDFThroughput, Payload, Ray.Origin, PathRadiance, Ray.Direction, Payload.WorldNormal);
+					ReflectedColor += FinalReflectionColor;
+				}
+#else // Debugging, remove later
+				ReflectedColor += float4(Payload.BaseColor, 1.0); 
+#endif
+
+				IncidentDirection = Ray.Direction;
+
+				HitDistance = Payload.HitT;
+			}
+		}
+	}
+
+	////uint NumLights = ForwardLightData.NumDirectionalLights;
+
+	if (SamplesPerPixel > 0)
+	{
+		ReflectedColor /= SamplesPerPixel;
+	}
+	ColorOutput[PixelCoord] = float4(ReflectedColor);
+}
+
+
+[shader("miss")]
+void RayTracingReflectionsMainMS(inout FPackedMaterialClosestHitPayload PackedPayload)
+{
+	PackedPayload.HitT = -1;
+}
+
+
+[shader("closesthit")]
+void RayTracingReflectionsMainCHS(inout FPackedMaterialClosestHitPayload PackedPayload, in FDefaultAttributes Attributes)
+{
+	FMaterialClosestHitPayload Payload = (FMaterialClosestHitPayload)0;
+
+	Payload.HitT = RayTCurrent();
+	Payload.BaseColor = float3(0.2, 0.3, 0.4);
+
+	FTriangleBaseAttributes Triangle = LoadTriangleBaseAttributes(PrimitiveIndex());
+	float3 Edge0 = Triangle.LocalPositions[2] - Triangle.LocalPositions[0];
+	float3 Edge1 = Triangle.LocalPositions[1] - Triangle.LocalPositions[0];
+	float3x3 WorldToLocal = (float3x3)WorldToObject();
+	float3x3 LocalToWorldNormal = transpose(WorldToLocal);
+	Payload.WorldNormal = normalize(mul(LocalToWorldNormal, cross(Edge0, Edge1)));
+
+	PackedPayload = PackRayTracingPayload(Payload);
+}
diff --git a/Engine/Shaders/Private/RayTracing_mv/RayTracingSkyLightRG.usf b/Engine/Shaders/Private/RayTracing_mv/RayTracingSkyLightRG.usf
new file mode 100644
--- /dev/null
+++ b/Engine/Shaders/Private/RayTracing_mv/RayTracingSkyLightRG.usf
@@ -0,0 +1,477 @@
+// Copyright 1998-2018 Epic Games, Inc. All Rights Reserved.
+
+#include "../Common.ush"
+#include "../MonteCarlo.ush"
+#include "../DeferredShadingCommon.ush"
+#include "SkyLightMipTreeCommon.ush"
+#include "RayTracingCommon.ush"
+#include "RayTracingDirectionalLight.ush"
+#include "RayTracingRectLight.ush"
+
+RaytracingAccelerationStructure TLAS;
+RWTexture2D<float4> RWOcclusionMaskUAV;
+RWTexture2D<float> RWRayDistanceUAV;
+
+#define USE_MIP_PDF_TEXTURE 1
+#define USE_SOLID_ANGLE_TEXTURE 1
+
+uint CalcLinearIndex(uint2 PixelCoord)
+{
+	return PixelCoord.y * View.BufferSizeAndInvSize.x + PixelCoord.x;
+}
+
+uint SelectFace(uint MipCount, inout RandomSequence RandSequence, out float FacePdf)
+{
+	uint BufferOffset = BufferOffsetAtPixel(uint2(0, 0), MipCount, SkyLight.MipDimensions);
+
+	float FaceCdf[6];
+	FaceCdf[0] = SkyLight.MipTreePosX[BufferOffset];
+	FaceCdf[1] = FaceCdf[0] + SkyLight.MipTreeNegX[BufferOffset];
+	FaceCdf[2] = FaceCdf[1] + SkyLight.MipTreePosY[BufferOffset];
+	FaceCdf[3] = FaceCdf[2] + SkyLight.MipTreeNegY[BufferOffset];
+	FaceCdf[4] = FaceCdf[3] + SkyLight.MipTreePosZ[BufferOffset];
+	FaceCdf[5] = FaceCdf[4] + SkyLight.MipTreeNegZ[BufferOffset];
+	float Sum = FaceCdf[5];
+
+	for (uint Index = 0; Index < 5; ++Index)
+	{
+		FaceCdf[Index] /= Sum;
+	}
+	FaceCdf[5] = 1.0;
+
+	// Draw random variable
+	uint DummyVariable;
+	float RandSample = RandomSequence_GenerateSample1D(RandSequence, DummyVariable);
+
+	// Determine CDF entry
+	uint FaceIndex = 0;
+	while (RandSample > FaceCdf[FaceIndex])
+	{
+		FaceIndex++;
+	}
+
+	// Calculate PDF and return entry
+	FacePdf = FaceCdf[FaceIndex];
+	if (FaceIndex > 0)
+		FacePdf -= FaceCdf[FaceIndex - 1];
+	return FaceIndex;
+}
+
+float4 SampleMipTree(uint FaceIndex, uint4 BufferOffset)
+{
+	float4 Values;
+	if (FaceIndex == 0)
+	{
+		Values.x = SkyLight.MipTreePosX[BufferOffset.x];
+		Values.y = SkyLight.MipTreePosX[BufferOffset.y];
+		Values.z = SkyLight.MipTreePosX[BufferOffset.z];
+		Values.w = SkyLight.MipTreePosX[BufferOffset.w];
+	}
+	else if (FaceIndex == 1)
+	{
+		Values.x = SkyLight.MipTreeNegX[BufferOffset.x];
+		Values.y = SkyLight.MipTreeNegX[BufferOffset.y];
+		Values.z = SkyLight.MipTreeNegX[BufferOffset.z];
+		Values.w = SkyLight.MipTreeNegX[BufferOffset.w];
+	}
+	else if (FaceIndex == 2)
+	{
+		Values.x = SkyLight.MipTreePosY[BufferOffset.x];
+		Values.y = SkyLight.MipTreePosY[BufferOffset.y];
+		Values.z = SkyLight.MipTreePosY[BufferOffset.z];
+		Values.w = SkyLight.MipTreePosY[BufferOffset.w];
+	}
+	else if (FaceIndex == 3)
+	{
+		Values.x = SkyLight.MipTreeNegY[BufferOffset.x];
+		Values.y = SkyLight.MipTreeNegY[BufferOffset.y];
+		Values.z = SkyLight.MipTreeNegY[BufferOffset.z];
+		Values.w = SkyLight.MipTreeNegY[BufferOffset.w];
+	}
+	else if (FaceIndex == 4)
+	{
+		Values.x = SkyLight.MipTreePosZ[BufferOffset.x];
+		Values.y = SkyLight.MipTreePosZ[BufferOffset.y];
+		Values.z = SkyLight.MipTreePosZ[BufferOffset.z];
+		Values.w = SkyLight.MipTreePosZ[BufferOffset.w];
+	}
+	else
+	{
+		Values.x = SkyLight.MipTreeNegZ[BufferOffset.x];
+		Values.y = SkyLight.MipTreeNegZ[BufferOffset.y];
+		Values.z = SkyLight.MipTreeNegZ[BufferOffset.z];
+		Values.w = SkyLight.MipTreeNegZ[BufferOffset.w];
+	}
+
+	return Values;
+}
+
+// #dxr_todo: Copied from PathTracingSkyLight.ush
+float3 EvalSkyLight(float3 Direction)
+{
+	return SkyLight.Color * TextureCubeSampleLevel(SkyLight.Texture, SkyLight.TextureSampler, Direction, 0).rgb;
+}
+
+float PdfFace(uint MipCount, uint FaceIndex)
+{
+	// Build Cdf
+	uint BufferOffset = BufferOffsetAtPixel(uint2(0, 0), MipCount, SkyLight.MipDimensions);
+
+	float FaceCdf[6];
+	FaceCdf[0] = SkyLight.MipTreePosX[BufferOffset];
+	FaceCdf[1] = FaceCdf[0] + SkyLight.MipTreeNegX[BufferOffset];
+	FaceCdf[2] = FaceCdf[1] + SkyLight.MipTreePosY[BufferOffset];
+	FaceCdf[3] = FaceCdf[2] + SkyLight.MipTreeNegY[BufferOffset];
+	FaceCdf[4] = FaceCdf[3] + SkyLight.MipTreePosZ[BufferOffset];
+	FaceCdf[5] = FaceCdf[4] + SkyLight.MipTreeNegZ[BufferOffset];
+	float Sum = FaceCdf[5];
+
+	for (uint Index = 0; Index < 5; ++Index)
+	{
+		FaceCdf[Index] /= Sum;
+	}
+	FaceCdf[5] = 1.0;
+
+	float Pdf = FaceCdf[FaceIndex];
+	if (FaceIndex > 0)
+		Pdf -= FaceCdf[FaceIndex - 1];
+	return Pdf;
+}
+
+float SampleMipPdf(uint3 TextureCoord, uint StopLevel)
+{
+	float MipPdf = 1.0;
+	uint BufferOffset = BufferOffsetAtPixel(TextureCoord.xy, StopLevel, SkyLight.MipDimensions);
+	if (TextureCoord.z == 0)
+	{
+		MipPdf = SkyLight.MipTreePdfPosX[BufferOffset];
+	}
+	else if (TextureCoord.z == 1)
+	{
+		MipPdf = SkyLight.MipTreePdfNegX[BufferOffset];
+	}
+	else if (TextureCoord.z == 2)
+	{
+		MipPdf = SkyLight.MipTreePdfPosY[BufferOffset];
+	}
+	else if (TextureCoord.z == 3)
+	{
+		MipPdf = SkyLight.MipTreePdfNegY[BufferOffset];
+	}
+	else if (TextureCoord.z == 4)
+	{
+		MipPdf = SkyLight.MipTreePdfPosZ[BufferOffset];
+	}
+	else // if (TextureCoord.z == 5)
+	{
+		MipPdf = SkyLight.MipTreePdfNegZ[BufferOffset];
+	}
+	return MipPdf;
+}
+
+float PdfSkyLightOcclusionRay(
+	float3 WorldNormal,
+	float3 WorldDirection
+)
+{
+	// Determine primary direction
+	uint MipCount = Log2(SkyLight.MipDimensions.x);
+	uint2 TextureRes;
+	SkyLight.Texture.GetDimensions(TextureRes.x, TextureRes.y);
+	uint3 TextureCoord = GetTextureCubeCoordinate(WorldDirection, TextureRes);
+	float2 TextureUV = TextureCoord.xy / (float2)TextureRes;
+	uint2 MipCoord = TextureUV * SkyLight.MipDimensions;
+	float FacePdf = PdfFace(MipCount, TextureCoord.z);
+
+	// #dxr_todo: StopLevel can be customized for run-time performance
+	uint StopLevel = min(SkyLight.SamplingStopLevel, MipCount);
+#if USE_MIP_PDF_TEXTURE
+	float MipPdf = SampleMipPdf(uint3(MipCoord, TextureCoord.z), StopLevel);
+
+#if USE_SOLID_ANGLE_TEXTURE
+	uint LinearIndex = BufferOffsetAtPixel(MipCoord, StopLevel, SkyLight.MipDimensions);
+	float AreaPdf = SkyLight.SolidAnglePdf[LinearIndex];
+#else
+	float AreaPdf = 1.0 / TexelCoordSolidAngle(MipCoord, DimensionsAtLevel(SkyLight.MipDimensions, StopLevel));
+#endif
+	return FacePdf * MipPdf * AreaPdf;
+#else
+	uint2 MipPixel = 0;
+	float MipPdf = 1.0;
+	for (uint Index = 0; Index < MipCount - StopLevel; ++Index)
+	{
+		// Build Cdf
+		uint MipIndex = MipCount - Index - 1;
+#if USE_MORTON_CODE
+		uint BufferOffset = BufferOffsetAtPixel(MipPixel, MipIndex, SkyLight.MipDimensions);
+		uint4 BufferOffsets = uint4(BufferOffset, BufferOffset + 1, BufferOffset + 2, BufferOffset + 3);
+#else
+		uint4 BufferOffsets;
+		BufferOffsets.x = BufferOffsetAtPixel(MipPixel, MipIndex, SkyLight.MipDimensions);
+		BufferOffsets.y = BufferOffsetAtPixel(MipPixel + uint2(1, 0), MipIndex, SkyLight.MipDimensions);
+		BufferOffsets.z = BufferOffsetAtPixel(MipPixel + uint2(0, 1), MipIndex, SkyLight.MipDimensions);
+		BufferOffsets.w = BufferOffsetAtPixel(MipPixel + uint2(1, 1), MipIndex, SkyLight.MipDimensions);
+#endif
+		float4 Cdf = SampleMipTree(TextureCoord.z, BufferOffsets);
+		Cdf.y += Cdf.x;
+		Cdf.z += Cdf.y;
+		Cdf.w += Cdf.z;
+		Cdf /= Cdf.w;
+
+		// Determine texel choice
+		uint2 MipOffset = (MipCoord.xy >> MipIndex) - MipPixel;
+		uint LinearIndex = MipOffset.y << 1 | MipOffset.x;
+		MipPdf *= LinearIndex > 0 ? Cdf[LinearIndex] - Cdf[LinearIndex - 1] : Cdf[LinearIndex];
+		MipPixel += MipOffset;
+
+		MipPixel = MipPixel << 1;
+	}
+	MipPdf = max(MipPdf, 0.0);
+	// DEBUG: assert that SampleMipPdf is equivalent for evaluating the decision tree..
+	//MipPdf = SampleMipPdf(uint3(MipCoord, TextureCoord.z), StopLevel);
+
+#if USE_SOLID_ANGLE_TEXTURE
+	uint LinearIndex = BufferOffsetAtPixel(MipCoord, StopLevel, SkyLight.MipDimensions);
+	float AreaPdf = SkyLight.SolidAnglePdf[LinearIndex];
+#else
+	float AreaPdf = 1.0 / TexelCoordSolidAngle(MipCoord, DimensionsAtLevel(SkyLight.MipDimensions, StopLevel));
+#endif
+	return FacePdf * MipPdf * AreaPdf;
+#endif
+}
+
+bool GenerateSkyLightOcclusionRay(
+	float3 WorldPosition,
+	float3 WorldNormal,
+	inout RandomSequence RandSequence,
+	out float3 RayOrigin,
+	out float3 RayDirection,
+	out float RayTMin,
+	out float RayTMax,
+	out float RayPdf
+)
+{
+	// Determine the appropriate face to use
+	uint MipCount = Log2(SkyLight.MipDimensions.x);
+	float FacePdf = 0.0;
+	uint FaceIndex = SelectFace(MipCount, RandSequence, FacePdf);
+
+	uint2 MipPixel = 0;
+	float MipPdf = 1.0;
+	uint StopLevel = min(SkyLight.SamplingStopLevel, MipCount);
+	for (uint Index = 0; Index < MipCount - StopLevel; ++Index)
+	{
+		uint MipIndex = MipCount - Index - 1;
+#if USE_MORTON_CODE
+		uint BufferOffset = BufferOffsetAtPixel(MipPixel, MipIndex, SkyLight.MipDimensions);
+		uint4 BufferOffsets = uint4(BufferOffset, BufferOffset + 1, BufferOffset + 2, BufferOffset + 3);
+#else
+		uint4 BufferOffsets;
+		BufferOffsets.x = BufferOffsetAtPixel(MipPixel, MipIndex, SkyLight.MipDimensions);
+		BufferOffsets.y = BufferOffsetAtPixel(MipPixel + uint2(1, 0), MipIndex, SkyLight.MipDimensions);
+		BufferOffsets.z = BufferOffsetAtPixel(MipPixel + uint2(0, 1), MipIndex, SkyLight.MipDimensions);
+		BufferOffsets.w = BufferOffsetAtPixel(MipPixel + uint2(1, 1), MipIndex, SkyLight.MipDimensions);
+#endif
+		float4 Cdf = SampleMipTree(FaceIndex, BufferOffsets);
+		Cdf.y += Cdf.x;
+		Cdf.z += Cdf.y;
+		Cdf.w += Cdf.z;
+		Cdf /= Cdf.w;
+
+		uint DummyVariable;
+		float RandSample = RandomSequence_GenerateSample1D(RandSequence, DummyVariable);
+		if (RandSample < Cdf.x)
+		{
+			MipPdf *= Cdf.x;
+		}
+		else if (RandSample < Cdf.y)
+		{
+			MipPdf *= Cdf.y - Cdf.x;
+			MipPixel.x += 1;
+		}
+		else if (RandSample < Cdf.z)
+		{
+			MipPdf *= Cdf.z - Cdf.y;
+			MipPixel.y += 1;
+		}
+		else // if (RandSample < Cdf.w)
+		{
+			MipPdf *= Cdf.w - Cdf.z;
+			MipPixel += 1;
+		}
+
+		MipPixel = uint2(MipPixel.x << 1, MipPixel.y << 1);
+	}
+	MipPixel = uint2(MipPixel.x >> 1, MipPixel.y >> 1);
+
+	// DEBUG: assert cached PDF is equivalent to re-evaluation
+	//MipPdf = PdfSkyLightOcclusionRay(WorldNormal, RayDirection);
+
+	uint DummyVariable;
+	float2 PixelOffset = RandomSequence_GenerateSample2D(RandSequence, DummyVariable);
+	float2 MipPixelUV = (MipPixel + PixelOffset) / float2(DimensionsAtLevel(SkyLight.MipDimensions, StopLevel));
+
+	// Construct Cubemap coordinate
+	uint2 TextureDimensions;
+	SkyLight.Texture.GetDimensions(TextureDimensions.x, TextureDimensions.y);
+	float3 TexelCoord = float3(MipPixelUV * TextureDimensions, FaceIndex);
+
+	// Determine area pdf
+#if USE_SOLID_ANGLE_TEXTURE
+	uint LinearIndex = BufferOffsetAtPixel(MipPixel, StopLevel, SkyLight.MipDimensions);
+	float AreaPdf = SkyLight.SolidAnglePdf[LinearIndex];
+#else
+	float AreaPdf = 1.0 / TexelCoordSolidAngle(MipPixel, DimensionsAtLevel(SkyLight.MipDimensions, StopLevel));
+#endif
+
+	RayOrigin = WorldPosition + WorldNormal * 0.01;
+	RayDirection = GetTextureCubeVector(TexelCoord, TextureDimensions);
+	RayTMin = 0.0;
+	RayTMax = 1.0e27;
+	RayPdf = FacePdf * MipPdf * AreaPdf;
+
+	return true;
+}
+
+float PdfLambertOcclusionRay(
+	float3 WorldNormal,
+	float3 WorldDirection
+)
+{
+	float Pdf = max(dot(WorldDirection, WorldNormal), 0.0) / PI;
+	return Pdf;
+}
+
+bool GenerateLambertOcclusionRay(
+	float3 WorldPosition,
+	float3 WorldNormal,
+	inout RandomSequence RandSequence,
+	out float3 RayOrigin,
+	out float3 RayDirection,
+	out float RayTMin,
+	out float RayTMax,
+	out float RayPdf
+)
+{
+	// Draw random variable
+	float2 BufferSize = View.BufferSizeAndInvSize.xy;
+	uint DummyVariable;
+	float2 RandSample = RandomSequence_GenerateSample2D(RandSequence, DummyVariable);
+
+	// Perform cosine-hemispherical sampling and convert to world-space
+	float4 Direction_Tangent = CosineSampleHemisphere(RandSample);
+	float3 Direction_World = TangentToWorld(Direction_Tangent.xyz, WorldNormal);
+
+	RayOrigin = WorldPosition + WorldNormal * 0.01;
+	RayDirection = Direction_World;
+	RayTMin = 0.0;
+	RayTMax = 1.0e27;
+	RayPdf = Direction_Tangent.w;
+
+	return true;
+}
+
+[shader("raygeneration")]
+void SkyLightRG()
+{
+	uint2 PixelCoord = DispatchRaysIndex().xy;
+
+	RandomSequence RandSequence;
+	uint LinearIndex = CalcLinearIndex(PixelCoord);
+	RandomSequence_Initialize(RandSequence, LinearIndex, View.FrameNumber * SkyLight.SamplesPerPixel);
+	//RandSequence.Type = 2;
+
+	// Get G-Buffer surface data
+	float2 InvBufferSize = View.BufferSizeAndInvSize.zw;
+	float2 UV = (float2(PixelCoord) + 0.5) * InvBufferSize;
+	FScreenSpaceData ScreenSpaceData = GetScreenSpaceData(UV);
+	float Depth = ScreenSpaceData.GBuffer.Depth;
+	float3 WorldPosition = ReconstructWorldPositionFromDepth(UV, Depth);
+	float3 WorldNormal = ScreenSpaceData.GBuffer.WorldNormal;
+	float3 BaseColor = ScreenSpaceData.GBuffer.BaseColor;
+	float Roughness = ScreenSpaceData.GBuffer.Roughness;
+
+	float RayDistance = 0.0;
+	float HitCount = 0.0;
+	uint SamplesPerPixel = SkyLight.SamplesPerPixel;
+
+	// Mask out depth values at or beyond far plane
+	if (Depth >= SCENE_TEXTURES_DISABLED_SCENE_DEPTH_VALUE)
+	{
+		SamplesPerPixel = 0.0;
+	}
+
+	float3 Irradiance = 0.0;
+	for (uint SampleIndex = 0; SampleIndex < SamplesPerPixel; ++SampleIndex)
+	{
+		RayDesc Ray;
+		float RayPdf = 0.0;
+
+		// Determine sky light or lambert ray
+		uint DummyVariable;
+		float RandSample = RandomSequence_GenerateSample1D(RandSequence, DummyVariable);
+		// #dxr_todo: Should strategy selection be driven by CVar?
+		float SkyLightSamplingStrategyPdf = 0.5;
+		float MisWeight = 1.0;
+		if (RandSample < SkyLightSamplingStrategyPdf)
+		{
+			GenerateSkyLightOcclusionRay(WorldPosition, WorldNormal, RandSequence, Ray.Origin, Ray.Direction, Ray.TMin, Ray.TMax, RayPdf);
+			MisWeight = RayPdf / (SkyLightSamplingStrategyPdf * RayPdf + (1.0 - SkyLightSamplingStrategyPdf) * PdfLambertOcclusionRay(WorldNormal, Ray.Direction));
+		}
+		else
+		{
+			GenerateLambertOcclusionRay(WorldPosition, WorldNormal, RandSequence, Ray.Origin, Ray.Direction, Ray.TMin, Ray.TMax, RayPdf);
+			MisWeight = RayPdf / ((1.0 - SkyLightSamplingStrategyPdf) * RayPdf + SkyLightSamplingStrategyPdf * PdfSkyLightOcclusionRay(WorldNormal, Ray.Direction));
+		}
+		
+		// Surface-normal culling
+		float NdotL = dot(WorldNormal, Ray.Direction);
+		if (RayPdf < 0.0 || NdotL <= 0.0)
+		{
+			continue;
+		}
+
+		uint RayFlags = 0
+			| RAY_FLAG_FORCE_OPAQUE; // #dxr_todo: generate a special shader permutation for occlusion ray tracing that includes any-hit shaders for masked geometry
+		FDefaultPayload Payload = (FDefaultPayload)0;
+
+		TraceRay(
+			TLAS,      // AccelerationStructure
+			RayFlags,
+			0xFF,      // InstanceInclusionMask
+			0,         // RayContributionToHitGroupIndex
+			1,         // MultiplierForGeometryContributionToShaderIndex
+			0,         // MissShaderIndex
+			Ray,       // RayDesc
+			Payload    // Payload
+		);
+
+		if (Payload.HitT >= 0.0)
+		{
+			RayDistance += Payload.HitT;
+			HitCount += 1.0;
+		}
+		else
+		{
+			float3 Radiance = EvalSkyLight(Ray.Direction);
+			// #dxr_todo: Apply actual material BRDF
+			float3 Brdf = BaseColor * Roughness / PI;
+			Irradiance += Radiance * Brdf * NdotL * MisWeight / RayPdf;
+		}
+	}
+
+	if (SamplesPerPixel > 0)
+	{
+		Irradiance /= SamplesPerPixel;
+	}
+	RWOcclusionMaskUAV[PixelCoord] = float4(Irradiance, 1.0);
+
+	if (HitCount > 0.0)
+	{
+		RWRayDistanceUAV[PixelCoord] = RayDistance / HitCount;
+	}
+	else
+	{
+		RWRayDistanceUAV[PixelCoord] = 1.0e27;
+	}
+}
\ No newline at end of file
diff --git a/Engine/Shaders/Private/RayTracing_mv/RayTracingTest.usf b/Engine/Shaders/Private/RayTracing_mv/RayTracingTest.usf
new file mode 100644
--- /dev/null
+++ b/Engine/Shaders/Private/RayTracing_mv/RayTracingTest.usf
@@ -0,0 +1,42 @@
+// Copyright 1998-2018 Epic Games, Inc. All Rights Reserved.
+
+#include "RayTracingCommon.ush"
+
+RaytracingAccelerationStructure  TLAS;
+StructuredBuffer<FBasicRayData>  Rays;
+RWStructuredBuffer<uint>         Output;
+
+[shader("raygeneration")]
+void TestMainRGS()
+{
+	const uint RayIndex = DispatchRaysIndex().x;
+	FBasicRayData InputRay = Rays[RayIndex];
+
+	RayDesc Ray;
+	Ray.Origin = InputRay.Origin;
+	Ray.Direction = InputRay.Direction;
+	Ray.TMin = 0.0f;
+	Ray.TMax = InputRay.TFar;
+
+	uint RayFlags = 0
+	              | RAY_FLAG_CULL_BACK_FACING_TRIANGLES      // use back face culling
+	              | RAY_FLAG_ACCEPT_FIRST_HIT_AND_END_SEARCH // terminate immediately upon detecting primitive intersection
+	              | RAY_FLAG_FORCE_OPAQUE                    // don't run anyhit shader
+	              | RAY_FLAG_SKIP_CLOSEST_HIT_SHADER;        // don't run closesthit shader
+
+	FOcclusionPayload Payload;
+	Payload.Hit = true;
+
+	TraceRay(
+		TLAS,   // AccelerationStructure
+		RayFlags,
+		0xFF,   // InstanceInclusionMask
+		0,      // RayContributionToHitGroupIndex
+		1,      // MultiplierForGeometryContributionToShaderIndex
+		0,      // MissShaderIndex
+		Ray,    // RayDesc
+		Payload // Payload
+	);
+
+	Output[RayIndex] = Payload.Hit ? ~0 : 0;
+}
diff --git a/Engine/Shaders/Private/RayTracing_mv/SkyLightMipTreeCommon.ush b/Engine/Shaders/Private/RayTracing_mv/SkyLightMipTreeCommon.ush
new file mode 100644
--- /dev/null
+++ b/Engine/Shaders/Private/RayTracing_mv/SkyLightMipTreeCommon.ush
@@ -0,0 +1,132 @@
+// Copyright 1998-2018 Epic Games, Inc. All Rights Reserved.
+
+/*=============================================================================================
+	SkyLightMipTreeCommon.ush: Common utilities for MipTree-based SkyLight sampling
+===============================================================================================*/
+
+#include "MipTreeCommon.ush"
+
+float3 GetTextureCubeVector(float3 TexelCoord, uint2 TextureRes)
+{
+	float3 Vector;
+
+	float2 UV = TexelCoord.xy / float2(TextureRes.xy);
+	UV = UV * 2.0 - 1.0;
+
+	uint Cubeface = TexelCoord.z;
+	if (Cubeface == 0)
+	{
+		Vector = float3(1, -UV.y, -UV.x);
+	}
+	else if (Cubeface == 1)
+	{
+		Vector = float3(-1, -UV.y, UV.x);
+	}
+	else if (Cubeface == 2)
+	{
+		Vector = float3(UV.x, 1, UV.y);
+	}
+	else if (Cubeface == 3)
+	{
+		Vector = float3(UV.x, -1, -UV.y);
+	}
+	else if (Cubeface == 4)
+	{
+		Vector = float3(UV.x, -UV.y, 1);
+	}
+	else if (Cubeface == 5)
+	{
+		Vector = float3(-UV.x, -UV.y, -1);
+	}
+
+	return Vector;
+}
+
+uint GetMaximumComponentIndex(float3 Vector)
+{
+	float AbsX = abs(Vector.x);
+	float AbsY = abs(Vector.y);
+	float AbsZ = abs(Vector.z);
+	if (AbsX > AbsY)
+	{
+		return AbsX > AbsZ ? 0 : 2;
+	}
+	else
+	{
+		return AbsY > AbsZ ? 1 : 2;
+	}
+}
+
+uint3 GetTextureCubeCoordinate(float3 WorldDirection, uint2 TextureRes)
+{
+	// Strongest coordinate determines cube face
+	uint MaxComponent = GetMaximumComponentIndex(WorldDirection);
+	WorldDirection /= abs(WorldDirection[MaxComponent]);
+	float3 TextureCubeCoordinate;
+	if (MaxComponent == 0)
+	{
+		if (WorldDirection[MaxComponent] > 0)
+		{
+			TextureCubeCoordinate = float3(-WorldDirection.z, -WorldDirection.y, 0);
+		}
+		else
+		{
+			TextureCubeCoordinate = float3(WorldDirection.z, -WorldDirection.y, 1);
+		}
+	}
+	else if (MaxComponent == 1)
+	{
+		if (WorldDirection[MaxComponent] > 0)
+		{
+			TextureCubeCoordinate = float3(WorldDirection.x, WorldDirection.z, 2);
+		}
+		else
+		{
+			TextureCubeCoordinate = float3(WorldDirection.x, -WorldDirection.z, 3);
+		}
+	}
+	else if (MaxComponent == 2)
+	{
+		if (WorldDirection[MaxComponent] > 0)
+		{
+			TextureCubeCoordinate = float3(WorldDirection.x, -WorldDirection.y, 4);
+		}
+		else
+		{
+			TextureCubeCoordinate = float3(-WorldDirection.x, -WorldDirection.y, 5);
+		}
+	}
+
+	TextureCubeCoordinate.xy = (TextureCubeCoordinate.xy + 1.0) / 2.0;
+	TextureCubeCoordinate.xy *= TextureRes;
+	return TextureCubeCoordinate;
+}
+
+/**
+* Cube texel solid angle computation based on formula from Manne Öhrström's
+* thesis - "Spherical Harmonics, Precomputed Radiance Transfer and Realtime Radiosity in Computer Games".
+*/
+float AreaElement(float x, float y)
+{
+	// Equation 7.12
+	return atan2(x * y, sqrt(x * x + y * y + 1));
+}
+
+float TexelCoordSolidAngle(uint2 TextureCoord, uint2 TextureRes)
+{
+	float2 InvResolution = 1.0f / (float)TextureRes;
+
+	// Convert texture coordinate to [-1, 1] range, offset to texel center.
+	float U = (2.0f * ((float)TextureCoord.x + 0.5f) * InvResolution.x) - 1.0f;
+	float V = (2.0f * ((float)TextureCoord.y + 0.5f) * InvResolution.y) - 1.0f;
+
+	// Projected area
+	float x0 = U - InvResolution.x;
+	float y0 = V - InvResolution.y;
+	float x1 = U + InvResolution.x;
+	float y1 = V + InvResolution.y;
+
+	// Equation 7.12
+	float SolidAngle = AreaElement(x0, y0) - AreaElement(x0, y1) - AreaElement(x1, y0) + AreaElement(x1, y1);
+	return SolidAngle;
+}
diff --git a/Engine/Shaders/Private/RayTracing_mv/VisualizeSkyLightMipTreePS.usf b/Engine/Shaders/Private/RayTracing_mv/VisualizeSkyLightMipTreePS.usf
new file mode 100644
--- /dev/null
+++ b/Engine/Shaders/Private/RayTracing_mv/VisualizeSkyLightMipTreePS.usf
@@ -0,0 +1,62 @@
+// Copyright 1998-2018 Epic Games, Inc. All Rights Reserved.
+
+/*=============================================================================================
+	VisualizeSkyLightMipTreePixelShader.usf: Interactive visualization of sky light mip tree.
+===============================================================================================*/
+
+#include "../Common.ush"
+#include "MipTreeCommon.ush"
+
+uint3 Dimensions;
+Buffer<float> MipTreePosX;
+Buffer<float> MipTreeNegX;
+Buffer<float> MipTreePosY;
+Buffer<float> MipTreeNegY;
+Buffer<float> MipTreePosZ;
+Buffer<float> MipTreeNegZ;
+
+float SampleMipTree(uint FaceIndex, uint BufferOffset)
+{
+	float Value = 0.0;
+	if (FaceIndex == 0)
+	{
+		Value = MipTreePosX[BufferOffset];
+	}
+	else if (FaceIndex == 1)
+	{
+		Value = MipTreeNegX[BufferOffset];
+	}
+	else if (FaceIndex == 2)
+	{
+		Value = MipTreePosY[BufferOffset];
+	}
+	else if (FaceIndex == 3)
+	{
+		Value = MipTreeNegY[BufferOffset];
+	}
+	else if (FaceIndex == 4)
+	{
+		Value = MipTreePosZ[BufferOffset];
+	}
+	else
+	{
+		Value = MipTreeNegZ[BufferOffset];
+	}
+
+	return Value;
+}
+
+void VisualizeSkyLightMipTreePS(
+	in noperspective float2 UV : TEXCOORD0,
+	out float4 OutColor : SV_Target0,
+	out float OutMipTree : SV_Target1
+)
+{
+	uint FaceIndex = 0;
+	uint MipLevel = 0;
+	uint2 TexCoord = UV * DimensionsAtLevel(Dimensions, MipLevel);
+	uint BufferOffset = BufferOffsetAtPixel(TexCoord, MipLevel, Dimensions);
+
+	OutColor = 0.f;
+	OutMipTree = SampleMipTree(FaceIndex, BufferOffset);
+}
\ No newline at end of file
diff --git a/Engine/Shaders/Private/Raytracing/BuildMipTreeCS.usf b/Engine/Shaders/Private/Raytracing/BuildMipTreeCS.usf
deleted file mode 100644
diff --git a/Engine/Shaders/Private/Raytracing/BuildMipTreePdfCS.usf b/Engine/Shaders/Private/Raytracing/BuildMipTreePdfCS.usf
deleted file mode 100644
diff --git a/Engine/Shaders/Private/Raytracing/BuildSolidAnglePdfCS.usf b/Engine/Shaders/Private/Raytracing/BuildSolidAnglePdfCS.usf
deleted file mode 100644
diff --git a/Engine/Shaders/Private/Raytracing/CommonRayTracingBuiltInResources.ush b/Engine/Shaders/Private/Raytracing/CommonRayTracingBuiltInResources.ush
deleted file mode 100644
diff --git a/Engine/Shaders/Private/Raytracing/CompositeAmbientOcclusionPS.usf b/Engine/Shaders/Private/Raytracing/CompositeAmbientOcclusionPS.usf
deleted file mode 100644
diff --git a/Engine/Shaders/Private/Raytracing/CompositeSkyLightPS.usf b/Engine/Shaders/Private/Raytracing/CompositeSkyLightPS.usf
deleted file mode 100644
diff --git a/Engine/Shaders/Private/Raytracing/MipTreeCommon.ush b/Engine/Shaders/Private/Raytracing/MipTreeCommon.ush
deleted file mode 100644
diff --git a/Engine/Shaders/Private/Raytracing/RayTracingAmbientOcclusionRG.usf b/Engine/Shaders/Private/Raytracing/RayTracingAmbientOcclusionRG.usf
deleted file mode 100644
diff --git a/Engine/Shaders/Private/Raytracing/RayTracingBarycentrics.usf b/Engine/Shaders/Private/Raytracing/RayTracingBarycentrics.usf
deleted file mode 100644
diff --git a/Engine/Shaders/Private/Raytracing/RayTracingBuiltInShaders.usf b/Engine/Shaders/Private/Raytracing/RayTracingBuiltInShaders.usf
deleted file mode 100644
diff --git a/Engine/Shaders/Private/Raytracing/RayTracingCommon.ush b/Engine/Shaders/Private/Raytracing/RayTracingCommon.ush
deleted file mode 100644
diff --git a/Engine/Shaders/Private/Raytracing/RayTracingDebug.usf b/Engine/Shaders/Private/Raytracing/RayTracingDebug.usf
deleted file mode 100644
diff --git a/Engine/Shaders/Private/Raytracing/RayTracingDebugDefinitions.ush b/Engine/Shaders/Private/Raytracing/RayTracingDebugDefinitions.ush
deleted file mode 100644
diff --git a/Engine/Shaders/Private/Raytracing/RayTracingDefinitions.ush b/Engine/Shaders/Private/Raytracing/RayTracingDefinitions.ush
deleted file mode 100644
diff --git a/Engine/Shaders/Private/Raytracing/RayTracingDirectionalLight.ush b/Engine/Shaders/Private/Raytracing/RayTracingDirectionalLight.ush
deleted file mode 100644
diff --git a/Engine/Shaders/Private/Raytracing/RayTracingDynamicMesh.usf b/Engine/Shaders/Private/Raytracing/RayTracingDynamicMesh.usf
deleted file mode 100644
diff --git a/Engine/Shaders/Private/Raytracing/RayTracingHitGroupCommon.ush b/Engine/Shaders/Private/Raytracing/RayTracingHitGroupCommon.ush
deleted file mode 100644
diff --git a/Engine/Shaders/Private/Raytracing/RayTracingMaterialHitShaders.usf b/Engine/Shaders/Private/Raytracing/RayTracingMaterialHitShaders.usf
deleted file mode 100644
diff --git a/Engine/Shaders/Private/Raytracing/RayTracingOcclusionRG.usf b/Engine/Shaders/Private/Raytracing/RayTracingOcclusionRG.usf
deleted file mode 100644
diff --git a/Engine/Shaders/Private/Raytracing/RayTracingRectLight.ush b/Engine/Shaders/Private/Raytracing/RayTracingRectLight.ush
deleted file mode 100644
diff --git a/Engine/Shaders/Private/Raytracing/RayTracingRectLightOcclusionRG.usf b/Engine/Shaders/Private/Raytracing/RayTracingRectLightOcclusionRG.usf
deleted file mode 100644
diff --git a/Engine/Shaders/Private/Raytracing/RayTracingReflectionsRG.usf b/Engine/Shaders/Private/Raytracing/RayTracingReflectionsRG.usf
deleted file mode 100644
diff --git a/Engine/Shaders/Private/Raytracing/RayTracingSkyLightRG.usf b/Engine/Shaders/Private/Raytracing/RayTracingSkyLightRG.usf
deleted file mode 100644
diff --git a/Engine/Shaders/Private/Raytracing/RayTracingTest.usf b/Engine/Shaders/Private/Raytracing/RayTracingTest.usf
deleted file mode 100644
diff --git a/Engine/Shaders/Private/Raytracing/SkyLightMipTreeCommon.ush b/Engine/Shaders/Private/Raytracing/SkyLightMipTreeCommon.ush
deleted file mode 100644
diff --git a/Engine/Shaders/Private/Raytracing/VisualizeSkyLightMipTreePS.usf b/Engine/Shaders/Private/Raytracing/VisualizeSkyLightMipTreePS.usf
deleted file mode 100644
