WorldToTangent updated to be consistent with the updated version of TangentToWorld that uses the new GetTangentBasis function.

diff --git a/Engine/Shaders/Private/MonteCarlo.ush b/Engine/Shaders/Private/MonteCarlo.ush
--- a/Engine/Shaders/Private/MonteCarlo.ush
+++ b/Engine/Shaders/Private/MonteCarlo.ush
@@ -21,40 +21,14 @@ float3x3 GetTangentBasis( float3 TangentZ )
 	return float3x3( TangentX, TangentY, TangentZ );
 }
 
-void CoordinateSystem(float3 vector1, inout float3 vector2, inout float3 vector3)
-{
-	if (abs(vector1.x) > abs(vector1.y))
-	{
-		float fInvLen = 1.0f / sqrt(vector1.x * vector1.x + vector1.z * vector1.z);
-		vector2 = float3(-vector1.z * fInvLen, 0.0f, vector1.x * fInvLen);
-	}
-	else
-	{
-		float fInvLen = 1.0f / sqrt(vector1.y * vector1.y + vector1.z * vector1.z);
-		vector2 = float3(0.0f, vector1.z * fInvLen, -vector1.y * fInvLen);
-	}
-	vector3 = cross(vector1, vector2);
-}
-
-#if 1 //dxr_todo: new TangentToWorld that calls GetTangentBasis needs further investigation
-float3 TangentToWorld(float3 Vec, float3 TangentZ)
-{
-	float3 TangentX, TangentY;
-	CoordinateSystem(TangentZ, TangentX, TangentY);
-	return TangentX * Vec.x + TangentY * Vec.y + TangentZ * Vec.z;
-}
-#else
 float3 TangentToWorld( float3 Vec, float3 TangentZ )
 {
 	return mul( Vec, GetTangentBasis( TangentZ ) );
 }
-#endif
 
 float3 WorldToTangent(float3 Vec, float3 TangentZ)
 {
-	float3 TangentX, TangentY;
-	CoordinateSystem(TangentZ, TangentX, TangentY);
-	return float3(dot(TangentX, Vec), dot(TangentY, Vec), dot(TangentZ, Vec));
+	return mul(GetTangentBasis(TangentZ), Vec);
 }
 
 float2 Hammersley( uint Index, uint NumSamples, uint2 Random )
