Renaming ray-gen shader file suffix from RG to RGS per new conventions. This is only true for files that exclusively contain ray-gen shaders (and not closest-hit or miss shaders, eg: RayTracingReflections.usf).

Renaming ray-gen entry points from <Function>RG() to <Function>RGS().

diff --git a/Engine/Shaders/Private/RayTracing/RayTracingAmbientOcclusionRG.usf b/Engine/Shaders/Private/RayTracing/RayTracingAmbientOcclusionRG.usf
deleted file mode 100644
diff --git a/Engine/Shaders/Private/RayTracing/RayTracingAmbientOcclusionRGS.usf b/Engine/Shaders/Private/RayTracing/RayTracingAmbientOcclusionRGS.usf
new file mode 100644
--- /dev/null
+++ b/Engine/Shaders/Private/RayTracing/RayTracingAmbientOcclusionRGS.usf
@@ -0,0 +1,134 @@
+// Copyright 1998-2018 Epic Games, Inc. All Rights Reserved.
+
+#include "../Common.ush"
+#include "../Montecarlo.ush"
+#include "../DeferredShadingCommon.ush"
+#include "../PathTracing/Utilities/PathTracingRandomSequence.ush"
+#include "RayTracingCommon.ush"
+
+RaytracingAccelerationStructure TLAS;
+RWTexture2D<float4> RWOcclusionMaskUAV;
+RWTexture2D<float> RWRayDistanceUAV;
+
+uint CalcLinearIndex(uint2 PixelCoord)
+{
+	return PixelCoord.y * View.BufferSizeAndInvSize.x + PixelCoord.x;
+}
+
+// #patrick todo: move to RayTracingCommon.ush
+float3 ReconstructWorldPositionFromDepth(float2 UV, float Depth)
+{
+	float2 ScreenPosition = (UV - View.ScreenPositionScaleBias.wz) / View.ScreenPositionScaleBias.xy;
+	float4 HomogeneousWorldPosition = mul(float4(ScreenPosition * Depth, Depth, 1), View.ScreenToWorld);
+	float3 WorldPosition = HomogeneousWorldPosition.xyz / HomogeneousWorldPosition.w;
+
+	return WorldPosition;
+}
+
+void GenerateCosineNormalRay(
+	float3 WorldPosition,
+	float3 WorldNormal,
+	inout RandomSequence RandSequence,
+	out float3 RayOrigin,
+	out float3 RayDirection,
+	out float RayTMin,
+	out float RayTMax
+)
+{
+	// Draw random variable
+	float2 BufferSize = View.BufferSizeAndInvSize.xy;
+	uint DummyVariable;
+	float2 RandSample = RandomSequence_GenerateSample2D(RandSequence, DummyVariable);
+
+	// Perform cosine-hemispherical sampling and convert to world-space
+	float4 Direction_Tangent = CosineSampleHemisphere(RandSample);
+	float3 Direction_World = TangentToWorld(Direction_Tangent.xyz, WorldNormal);
+
+	RayOrigin = WorldPosition + WorldNormal * 0.01;
+	RayDirection = Direction_World;
+	RayTMin = 0.0;
+	RayTMax = AmbientOcclusion.MaxRayDistance;
+}
+
+[shader("raygeneration")]
+void AmbientOcclusionRGS()
+{
+	uint2 PixelCoord = DispatchRaysIndex().xy;
+
+	RandomSequence RandSequence;
+	uint LinearIndex = CalcLinearIndex(PixelCoord);
+	RandomSequence_Initialize(RandSequence, LinearIndex, View.FrameNumber);
+
+	// Get G-Buffer surface data
+	float2 InvBufferSize = View.BufferSizeAndInvSize.zw;
+	float2 UV = (float2(PixelCoord) + 0.5) * InvBufferSize;
+	FScreenSpaceData ScreenSpaceData = GetScreenSpaceData(UV);
+	float Depth = ScreenSpaceData.GBuffer.Depth;
+	float3 WorldPosition = ReconstructWorldPositionFromDepth(UV, Depth);
+	float3 WorldNormal = ScreenSpaceData.GBuffer.WorldNormal;
+
+	float Visibility = 0.0;
+	float RayDistance = 0.0;
+	float HitCount = 0.0;
+	uint SamplesPerPixel = AmbientOcclusion.SamplesPerPixel;
+	
+	// Mask out depth values beyond far plane
+	if (Depth >= SCENE_TEXTURES_DISABLED_SCENE_DEPTH_VALUE)
+	{
+		SamplesPerPixel = 0.0;
+		Visibility = 1.0;
+	}
+
+	for (uint SampleIndex = 0; SampleIndex < SamplesPerPixel; ++SampleIndex)
+	{
+		RayDesc Ray;
+		GenerateCosineNormalRay(WorldPosition, WorldNormal, RandSequence, Ray.Origin, Ray.Direction, Ray.TMin, Ray.TMax);
+		if (dot(WorldNormal, Ray.Direction) > 0.0)
+		{
+			uint RayFlags = 0
+				| RAY_FLAG_CULL_BACK_FACING_TRIANGLES
+				| RAY_FLAG_FORCE_OPAQUE; // #dxr_todo: generate a special shader permutation for occlusion ray tracing that includes any-hit shaders for masked geometry
+
+			FDefaultPayload Payload = (FDefaultPayload)0;
+
+			TraceRay(
+				TLAS,
+				RayFlags,
+				0xFF, // InstanceInclusionMask
+				0,    // RayContributionToHitGroupIndex
+				1,    // MultiplierForGeometryContributionToShaderIndex
+				0,    // MissShaderIndex
+				Ray,
+				Payload
+			);
+
+			Visibility += Payload.HitT < 0 ? 1.0 : 0.0;
+			if (Payload.HitT >= 0.0)
+			{
+				RayDistance += Payload.HitT;
+				HitCount += 1.0;
+			}
+		}
+		else
+		{
+			Visibility += 0.0;
+			RayDistance += 0.0;
+			HitCount += 1.0;
+		}
+	}
+
+	RWOcclusionMaskUAV[PixelCoord] = Visibility;
+	if (SamplesPerPixel > 0)
+	{
+		RWOcclusionMaskUAV[PixelCoord] /= SamplesPerPixel;
+	}
+
+	if (HitCount > 0.0)
+	{
+		RWRayDistanceUAV[PixelCoord] = RayDistance / HitCount;
+	}
+	else
+	{
+		RWRayDistanceUAV[PixelCoord] = 1.0e27;
+	}
+}
\ No newline at end of file
diff --git a/Engine/Shaders/Private/RayTracing/RayTracingBarycentrics.usf b/Engine/Shaders/Private/RayTracing/RayTracingBarycentrics.usf
--- a/Engine/Shaders/Private/RayTracing/RayTracingBarycentrics.usf
+++ b/Engine/Shaders/Private/RayTracing/RayTracingBarycentrics.usf
@@ -23,7 +23,7 @@ struct FBarycentricsPayload
 
 
 [shader("raygeneration")]
-void RayTracingBarycentricsMainRG()
+void RayTracingBarycentricsMainRGS()
 {
 	uint2 PixelCoord = DispatchRaysIndex().xy;
 
diff --git a/Engine/Shaders/Private/RayTracing/RayTracingDebug.usf b/Engine/Shaders/Private/RayTracing/RayTracingDebug.usf
--- a/Engine/Shaders/Private/RayTracing/RayTracingDebug.usf
+++ b/Engine/Shaders/Private/RayTracing/RayTracingDebug.usf
@@ -20,7 +20,7 @@ float3 ReconstructWorldPositionFromDepth(float2 UV, float Depth)
 }
 
 [shader("raygeneration")]
-void RayTracingDebugMainRG()
+void RayTracingDebugMainRGS()
 {
 	uint2 PixelCoord = DispatchRaysIndex().xy;
 
diff --git a/Engine/Shaders/Private/RayTracing/RayTracingOcclusionRG.usf b/Engine/Shaders/Private/RayTracing/RayTracingOcclusionRG.usf
deleted file mode 100644
diff --git a/Engine/Shaders/Private/RayTracing/RayTracingOcclusionRGS.usf b/Engine/Shaders/Private/RayTracing/RayTracingOcclusionRGS.usf
new file mode 100644
--- /dev/null
+++ b/Engine/Shaders/Private/RayTracing/RayTracingOcclusionRGS.usf
@@ -0,0 +1,285 @@
+// Copyright 1998-2018 Epic Games, Inc. All Rights Reserved.
+
+#include "../Common.ush"
+#include "../MonteCarlo.ush"
+#include "../DeferredShadingCommon.ush"
+#include "../LightShaderParameters.ush"
+#include "RayTracingCommon.ush"
+#include "RayTracingDirectionalLight.ush"
+#include "RayTracingRectLight.ush"
+
+/** Light types */
+#define LIGHT_TYPE_DIRECTIONAL 0
+#define LIGHT_TYPE_POINT 1
+#define LIGHT_TYPE_SPOT 2
+#define LIGHT_TYPE_RECT 3
+
+RaytracingAccelerationStructure TLAS;
+RWTexture2D<float4> RWOcclusionMaskUAV;
+RWTexture2D<float> RWRayDistanceUAV;
+
+uint SamplesPerPixel;
+
+float2 SobolIndexToUniformUnitSquare(uint2 SobolRand)
+{
+	return float2(SobolRand) * rcp(65536.0) + rcp(65536.0 * 2.0);
+}
+
+uint CalcLinearIndex(uint2 PixelCoord)
+{
+	return PixelCoord.y * View.BufferSizeAndInvSize.x + PixelCoord.x;
+}
+
+bool GeneratePointLightOcclusionRay(
+	FLightShaderParameters LightParameters,
+	float3 WorldPosition,
+	float3 WorldNormal,
+	float2 RandSample,
+	out float3 RayOrigin,
+	out float3 RayDirection,
+	out float RayTMin,
+	out float RayTMax
+)
+{
+	float3 LightDirection = LightParameters.Position - WorldPosition;
+	float RayLength = length(LightDirection);
+	LightDirection /= RayLength;
+
+	// Apply normal perturbation when defining ray
+	RayOrigin = WorldPosition + WorldNormal * 0.01;
+	RayDirection = LightDirection;
+	RayTMin = 0.0;
+	RayTMax = RayLength;
+	return true;
+}
+
+bool GenerateSpotLightOcclusionRay(
+	FLightShaderParameters LightParameters,
+	float3 WorldPosition,
+	float3 WorldNormal,
+	float2 RandSample,
+	out float3 RayOrigin,
+	out float3 RayDirection,
+	out float RayTMin,
+	out float RayTMax
+)
+{
+	float3 LightDirection = LightParameters.Position - WorldPosition;
+	float RayLength = length(LightDirection);
+	LightDirection /= RayLength;
+
+	// Apply culling
+	float CosAngle = LightParameters.SpotAngles.x;
+	if (dot(-LightDirection, LightParameters.Direction) < CosAngle)
+	{
+		return false;
+	}
+
+	// Apply normal perturbation when defining ray
+	RayOrigin = WorldPosition + WorldNormal * 0.01;
+	RayDirection = LightDirection;
+	RayTMin = 0.0;
+	RayTMax = RayLength;
+	return true;
+}
+
+bool GenerateOcclusionRay(
+	FLightShaderParameters LightParameters,
+	float3 WorldPosition,
+	float3 WorldNormal,
+	float2 RandSample,
+	out float3 RayOrigin,
+	out float3 RayDirection,
+	out float RayTMin,
+	out float RayTMax
+)
+{
+	#if LIGHT_TYPE == LIGHT_TYPE_DIRECTIONAL
+	{
+		// TODO: just passdown LightParameters to GenerateDirectionalLightOcclusionRay().
+		float3 LightDirection = LightParameters.Direction;
+		float LightRadius = LightParameters.SourceRadius;
+
+		GenerateDirectionalLightOcclusionRay(
+			WorldPosition, WorldNormal,
+			LightDirection, LightRadius,
+			RandSample,
+			/* out */ RayOrigin,
+			/* out */ RayDirection,
+			/* out */ RayTMin,
+			/* out */ RayTMax);
+	}
+	#elif LIGHT_TYPE == LIGHT_TYPE_POINT
+	{
+		return GeneratePointLightOcclusionRay(
+			LightParameters,
+			WorldPosition, WorldNormal,
+			RandSample,
+			/* out */ RayOrigin,
+			/* out */ RayDirection,
+			/* out */ RayTMin,
+			/* out */ RayTMax);
+	}
+	#elif LIGHT_TYPE == LIGHT_TYPE_SPOT
+	{
+		return GenerateSpotLightOcclusionRay(
+			LightParameters,
+			WorldPosition, WorldNormal,
+			RandSample,
+			/* out */ RayOrigin,
+			/* out */ RayDirection,
+			/* out */ RayTMin,
+			/* out */ RayTMax);
+	}
+	#elif LIGHT_TYPE == LIGHT_TYPE_RECT
+	{
+		float3 Tangent;
+		float3 BiTangent;
+		GetRectLightTangents(LightParameters, /* out */ Tangent, /* out */ BiTangent);
+
+		float2 LightDimensions = GetRectLightDimensions(LightParameters);
+		float RayPdf = 0.0;
+
+		// TODO: just passdown LightParameters to GenerateRectLightOcclusionRay().
+		return GenerateRectLightOcclusionRay(
+			WorldPosition, WorldNormal,
+			LightParameters.Position,
+			-LightParameters.Direction,
+			Tangent,
+			BiTangent,
+			LightDimensions.x, LightDimensions.y,
+			RandSample,
+			/* out */ RayOrigin,
+			/* out */ RayDirection,
+			/* out */ RayTMin,
+			/* out */ RayTMax,
+			/* out */ RayPdf);
+	}
+	#else
+		#error Unknown light type.
+	#endif
+	return true;
+}
+
+uint ReverseNBits(uint Bitfield, const uint BitCount)
+{
+	uint Out = 0;
+	UNROLL
+	for (uint i = 0; i < BitCount; i++)
+	{
+		uint bIsSet = ((Bitfield & (0x1 << i)) >> i);
+		Out |= bIsSet << (BitCount - 1 - i);
+	}
+	return Out;
+}
+
+[shader("raygeneration")]
+void OcclusionRGS()
+{
+	uint2 PixelCoord = DispatchRaysIndex().xy;
+
+	RandomSequence RandSequence;
+	uint LinearIndex = CalcLinearIndex(PixelCoord);
+	RandomSequence_Initialize(RandSequence, LinearIndex, View.FrameNumber);
+
+	FLightShaderParameters LightParameters = GetRootLightShaderParameters();
+
+	// Get G-Buffer surface data
+	float2 InvBufferSize = View.BufferSizeAndInvSize.zw;
+	float2 UV = (float2(PixelCoord) + 0.5) * InvBufferSize;
+	FScreenSpaceData ScreenSpaceData = GetScreenSpaceData(UV);
+	float Depth = ScreenSpaceData.GBuffer.Depth;
+	float3 WorldPosition = ReconstructWorldPositionFromDepth(UV, Depth);
+	float3 WorldNormal = ScreenSpaceData.GBuffer.WorldNormal;
+
+	float Visibility = 0.0;
+	float RayDistance = 0.0;
+	float HitCount = 0.0;
+	float RayCount = 0.0;
+	uint LocalSamplesPerPixel = SamplesPerPixel;
+
+	// Mask out depth values at or beyond far plane
+	if (Depth >= SCENE_TEXTURES_DISABLED_SCENE_DEPTH_VALUE)
+	{
+		LocalSamplesPerPixel = 0.0;
+	}
+
+	LOOP
+	for (uint SampleIndex = 0; SampleIndex < LocalSamplesPerPixel; ++SampleIndex)
+	{
+		uint DummyVariable;
+		float2 RandSample = RandomSequence_GenerateSample2D(RandSequence, DummyVariable);
+
+		BRANCH
+		if (LocalSamplesPerPixel == 1 && 1)
+		{
+			const uint TemporalBits = 10;
+			uint FrameIndexMod1024 = ReverseNBits(GetPowerOfTwoModulatedFrameIndex(1 << TemporalBits), TemporalBits);
+
+			uint2 SobolBase = SobolPixel(PixelCoord);
+			uint2 SobolFrame = SobolIndex(SobolBase, FrameIndexMod1024, TemporalBits);
+
+			RandSample = SobolIndexToUniformUnitSquare(SobolFrame);
+		}
+
+		RayDesc Ray;
+		bool bIsValidRay = GenerateOcclusionRay(
+			LightParameters,
+			WorldPosition, WorldNormal,
+			RandSample,
+			/* out */ Ray.Origin,
+			/* out */ Ray.Direction,
+			/* out */ Ray.TMin,
+			/* out */ Ray.TMax);
+
+		BRANCH
+		if (!bIsValidRay)
+		{
+			continue;
+		}
+		else if (dot(WorldNormal, Ray.Direction) <= 0.0)
+		{
+			continue;
+		}
+
+		uint RayFlags = 0
+			| RAY_FLAG_CULL_BACK_FACING_TRIANGLES
+			| RAY_FLAG_FORCE_OPAQUE; // #dxr_todo: generate a special shader permutation for occlusion ray tracing that includes any-hit shaders for masked geometry
+
+		FDefaultPayload Payload = (FDefaultPayload)0;
+
+		TraceRay(
+			TLAS,      // AccelerationStructure
+			RayFlags,
+			0xFF,      // InstanceInclusionMask
+			0,         // RayContributionToHitGroupIndex
+			1,         // MultiplierForGeometryContributionToShaderIndex
+			0,         // MissShaderIndex
+			Ray,       // RayDesc
+			Payload    // Payload
+		);
+
+		RayCount += 1.0;
+		Visibility += Payload.HitT < 0 ? 1.0 : 0.0;
+		if (Payload.HitT >= 0.0)
+		{
+			RayDistance += Payload.HitT;
+			HitCount += 1.0;
+		}
+	}
+
+	RWOcclusionMaskUAV[PixelCoord] = (LocalSamplesPerPixel > 0) ? Visibility / LocalSamplesPerPixel : Visibility;
+
+	if (HitCount > 0.0)
+	{
+		RWRayDistanceUAV[PixelCoord] = RayDistance / HitCount;
+	}
+	else if (RayCount > 0.0)
+	{
+		RWRayDistanceUAV[PixelCoord] = 1.0e27;
+	}
+	else
+	{
+		RWRayDistanceUAV[PixelCoord] = -1.0;
+	}
+}
\ No newline at end of file
diff --git a/Engine/Shaders/Private/RayTracing/RayTracingRectLightOcclusionRG.usf b/Engine/Shaders/Private/RayTracing/RayTracingRectLightOcclusionRG.usf
deleted file mode 100644
diff --git a/Engine/Shaders/Private/RayTracing/RayTracingRectLightOcclusionRGS.usf b/Engine/Shaders/Private/RayTracing/RayTracingRectLightOcclusionRGS.usf
new file mode 100644
--- /dev/null
+++ b/Engine/Shaders/Private/RayTracing/RayTracingRectLightOcclusionRGS.usf
@@ -0,0 +1,274 @@
+// Copyright 1998-2018 Epic Games, Inc. All Rights Reserved.
+
+#include "../Common.ush"
+#include "../DeferredShadingCommon.ush"
+#include "../PathTracing/Utilities/PathTracingRandomSequence.ush"
+#include "RayTracingCommon.ush"
+#include "RayTracingRectLight.ush"
+#include "MipTreeCommon.ush"
+
+RaytracingAccelerationStructure TLAS;
+RWTexture2D<float4> RWOcclusionMaskUAV;
+RWTexture2D<float> RWRayDistanceUAV;
+
+uint CalcLinearIndex(uint2 PixelCoord)
+{
+	return PixelCoord.y * View.BufferSizeAndInvSize.x + PixelCoord.x;
+}
+
+// #patrick todo: move to RayTracingCommon.ush
+float3 ReconstructWorldPositionFromDepth(float2 UV, float Depth)
+{
+	float2 ScreenPosition = (UV - View.ScreenPositionScaleBias.wz) / View.ScreenPositionScaleBias.xy;
+	float4 HomogeneousWorldPosition = mul(float4(ScreenPosition * Depth, Depth, 1), View.ScreenToWorld);
+	float3 WorldPosition = HomogeneousWorldPosition.xyz / HomogeneousWorldPosition.w;
+
+	return WorldPosition;
+}
+
+bool GenerateRectLightImportanceSampledRay(
+	float3 WorldPosition,
+	float3 WorldNormal,
+	float3 LightPosition,
+	float3 LightNormal,
+	float3 LightdPdu,
+	float3 LightdPdv,
+	float LightWidth,
+	float LightHeight,
+	inout RandomSequence RandSequence,
+	out float3 RayOrigin,
+	out float3 RayDirection,
+	out float RayTMin,
+	out float RayTMax,
+	out float RayPdf
+)
+{
+	uint MipCount = Log2(RectLight.MipTreeDimensions.x);
+	uint2 MipPixel = 0;
+	float MipPdf = 1.0;
+
+	uint StopLevel = 0;
+	for (uint Index = 0; Index < MipCount - StopLevel; ++Index)
+	{
+		uint MipIndex = MipCount - Index - 1;
+		uint BufferOffset = BufferOffsetAtPixel(MipPixel, MipIndex, RectLight.MipTreeDimensions);
+		uint4 BufferOffsets = BufferOffset + uint4(0, 1, 2, 3);
+		float4 PowerTerm = float4(RectLight.MipTree[BufferOffsets.x], RectLight.MipTree[BufferOffsets.y], RectLight.MipTree[BufferOffsets.z], RectLight.MipTree[BufferOffsets.w]);
+		PowerTerm = max(PowerTerm, 0.0);
+
+		// Determine sample positions
+		uint DummyVariable;
+		float2 Jitter = RandomSequence_GenerateSample2D(RandSequence, DummyVariable);
+		float2 DimInv = 1.0 / DimensionsAtLevel(RectLight.MipTreeDimensions, MipIndex);
+		float2 UV00 = (MipPixel + Jitter) * DimInv - 0.5;
+		float2 UV10 = (MipPixel + Jitter + uint2(1, 0)) * DimInv - 0.5;
+		float2 UV01 = (MipPixel + Jitter + uint2(0, 1)) * DimInv - 0.5;
+		float2 UV11 = (MipPixel + Jitter + uint2(1, 1)) * DimInv - 0.5;
+
+		float3 Position00 = LightPosition + LightdPdu * LightWidth * UV00.x + LightdPdv * LightHeight * UV00.y;
+		float3 Position10 = LightPosition + LightdPdu * LightWidth * UV10.x + LightdPdv * LightHeight * UV10.y;
+		float3 Position01 = LightPosition + LightdPdu * LightWidth * UV01.x + LightdPdv * LightHeight * UV01.y;
+		float3 Position11 = LightPosition + LightdPdu * LightWidth * UV11.x + LightdPdv * LightHeight * UV11.y;
+
+		// Evaluate irradiance coefficients
+		float3 LightDirection00 = WorldPosition - Position00;
+		float3 LightDirection10 = WorldPosition - Position10;
+		float3 LightDirection01 = WorldPosition - Position01;
+		float3 LightDirection11 = WorldPosition - Position11;
+
+		float4 DistanceSquaredTerm = float4(
+			dot(LightDirection00, LightDirection00),
+			dot(LightDirection10, LightDirection10),
+			dot(LightDirection01, LightDirection01),
+			dot(LightDirection11, LightDirection11)
+		);
+
+		float4 CosineTerm = float4(
+			saturate(dot(normalize(LightDirection00), LightNormal)),
+			saturate(dot(normalize(LightDirection10), LightNormal)),
+			saturate(dot(normalize(LightDirection01), LightNormal)),
+			saturate(dot(normalize(LightDirection11), LightNormal))
+		);
+
+		// Build CDF on approximate irradiance
+		float4 Cdf = PowerTerm * CosineTerm / DistanceSquaredTerm;
+		Cdf.y += Cdf.x;
+		Cdf.z += Cdf.y;
+		Cdf.w += Cdf.z;
+		Cdf /= Cdf.w;
+
+		// Sample CDF
+		float RandSample = RandomSequence_GenerateSample1D(RandSequence, DummyVariable);
+		if (RandSample < Cdf.x)
+		{
+			MipPdf *= Cdf.x;
+		}
+		else if (RandSample < Cdf.y)
+		{
+			MipPdf *= Cdf.y - Cdf.x;
+			MipPixel.x += 1;
+		}
+		else if (RandSample < Cdf.z)
+		{
+			MipPdf *= Cdf.z - Cdf.y;
+			MipPixel.y += 1;
+		}
+		else // if (RandSample < Cdf.w)
+		{
+			MipPdf *= Cdf.w - Cdf.z;
+			MipPixel += 1;
+		}
+
+		MipPixel = uint2(MipPixel.x << 1, MipPixel.y << 1);
+	}
+	MipPixel = uint2(MipPixel.x >> 1, MipPixel.y >> 1);
+
+	// Sample within stop-level texel
+	uint DummyVariable;
+	float2 PixelOffset = RandomSequence_GenerateSample2D(RandSequence, DummyVariable);
+	float2 DimensionsAtStopLevel = DimensionsAtLevel(RectLight.MipTreeDimensions, StopLevel);
+	float2 MipPixelUV = (MipPixel + PixelOffset) / DimensionsAtStopLevel;
+	MipPdf *= DimensionsAtStopLevel;
+
+	// Map sample point to quad
+	MipPixelUV -= 0.5;
+	LightPosition += LightdPdu * LightWidth * MipPixelUV.x + LightdPdv * LightHeight * MipPixelUV.y;
+	float3 LightDirection = normalize(LightPosition - WorldPosition);
+
+	// Light-normal culling
+	if (dot(-LightDirection, LightNormal) <= 0.0)
+	{
+		RayPdf = 0.0;
+		return false;
+	}
+
+	// #dxr_todo: Performing direct lighting here could save some recalculation
+
+	// Define ray
+	RayOrigin = WorldPosition + WorldNormal * 0.01;
+	RayDirection = LightDirection;
+	RayTMin = 0.0;
+	RayTMax = length(LightPosition - WorldPosition);
+
+	// Determine area pdf
+	float Area = (RectLight.Width * RectLight.Height) / DimensionsAtStopLevel;
+	float Distance2 = RayTMax * RayTMax;
+	float SolidAngle = Area * saturate(dot(-LightDirection, LightNormal)) / Distance2;
+	float SolidAnglePdf = 1.0 / SolidAngle;
+	RayPdf = MipPdf * SolidAnglePdf;
+
+	return RayTMax > RayTMin;
+}
+
+[shader("raygeneration")]
+void RectLightOcclusionRGS()
+{
+	uint2 PixelCoord = DispatchRaysIndex().xy;
+
+	RandomSequence RandSequence;
+	uint LinearIndex = CalcLinearIndex(PixelCoord);
+	RandomSequence_Initialize(RandSequence, LinearIndex, View.FrameNumber);
+
+	// Get G-Buffer surface data
+	float2 InvBufferSize = View.BufferSizeAndInvSize.zw;
+	float2 UV = (float2(PixelCoord) + 0.5) * InvBufferSize;
+	FScreenSpaceData ScreenSpaceData = GetScreenSpaceData(UV);
+	float Depth = ScreenSpaceData.GBuffer.Depth;
+	float3 WorldPosition = ReconstructWorldPositionFromDepth(UV, Depth);
+	float3 WorldNormal = ScreenSpaceData.GBuffer.WorldNormal;
+	float3 DiffuseColor = ScreenSpaceData.GBuffer.DiffuseColor;
+
+	float3 LightRadiance = RectLight.Color;
+
+	float RayDistance = 0.0;
+	float HitCount = 0.0;
+	uint SamplesPerPixel = RectLight.SamplesPerPixel;
+
+	// Mask out depth values at or beyond far plane
+	if (Depth >= SCENE_TEXTURES_DISABLED_SCENE_DEPTH_VALUE)
+	{
+		SamplesPerPixel = 0.0;
+	}
+
+	float3 ExitantRadiance = 0.0;
+	for (uint SampleIndex = 0; SampleIndex < SamplesPerPixel; ++SampleIndex)
+	{
+		RayDesc Ray;
+		float3 LightPosition = RectLight.Position;
+		float3 LightNormal = RectLight.Normal;
+		float3 LightdPdu = RectLight.dPdu;
+		float3 LightdPdv = RectLight.dPdv;
+		float LightWidth = RectLight.Width;
+		float LightHeight = RectLight.Height;
+
+		uint DummyVariable;
+		float2 RandSample = RandomSequence_GenerateSample2D(RandSequence, DummyVariable);
+		float RayPdf = 0.0;
+
+		bool bIsValid = RectLight.bIsTextureImportanceSampling ?
+			GenerateRectLightImportanceSampledRay(WorldPosition, WorldNormal, LightPosition, LightNormal, LightdPdu, LightdPdv, LightWidth, LightHeight, RandSequence, Ray.Origin, Ray.Direction, Ray.TMin, Ray.TMax, RayPdf) :
+			GenerateRectLightOcclusionRay(WorldPosition, WorldNormal, LightPosition, LightNormal, LightdPdu, LightdPdv, LightWidth, LightHeight, RandSample, Ray.Origin, Ray.Direction, Ray.TMin, Ray.TMax, RayPdf);
+		if (!bIsValid)
+		{
+			continue;
+		}
+
+		float NoL = dot(WorldNormal, Ray.Direction);
+		if (NoL <= 0.0)
+		{
+			continue;
+		}
+
+		uint RayFlags = 0
+			| RAY_FLAG_CULL_BACK_FACING_TRIANGLES
+			| RAY_FLAG_FORCE_OPAQUE; // #dxr_todo: generate a special shader permutation for occlusion ray tracing that includes any-hit shaders for masked geometry
+
+		FDefaultPayload Payload = (FDefaultPayload)0;
+
+		TraceRay(
+			TLAS,      // AccelerationStructure
+			RayFlags,
+			0xFF,      // InstanceInclusionMask
+			0,         // RayContributionToHitGroupIndex
+			1,         // MultiplierForGeometryContributionToShaderIndex
+			0,         // MissShaderIndex
+			Ray,       // RayDesc
+			Payload    // Payload
+		);
+
+		if (Payload.HitT >= 0.0)
+		{
+			RayDistance += Payload.HitT;
+			HitCount += 1.0;
+		}
+		else
+		{
+#if CALC_DIRECT_LIGHTING
+			// Sample light texture
+			float3 LightPosition = Ray.Origin + Ray.Direction * Ray.TMax;
+			float3 RectLightCorner = RectLight.Position - 0.5 * (RectLight.dPdu * RectLight.Width + RectLight.dPdv * RectLight.Height);
+			float3 DeltaPosition = LightPosition - RectLightCorner;
+			float2 UV = float2(dot(DeltaPosition, RectLight.dPdu) / RectLight.Width, dot(DeltaPosition, RectLight.dPdv) / RectLight.Height);
+			UV = 1.0 - UV;
+			float3 TextureColor = Texture2DSample(RectLight.Texture, RectLight.TextureSampler, UV);
+			float3 LightRadiance = RectLight.Color * TextureColor;
+			// #dxr_todo: Replace with actual BRDF
+			float3 Brdf = DiffuseColor;
+			float Distance2 = Ray.TMax * Ray.TMax;
+			ExitantRadiance += LightRadiance * Brdf * saturate(NoL) / RayPdf;
+#else
+			// Record exitant radiance as visibility score
+			ExitantRadiance += 1.0;
+#endif
+		}
+	}
+
+	float3 Result = (SamplesPerPixel > 0) ? ExitantRadiance / SamplesPerPixel : ExitantRadiance;
+#if	ENCODE_VISIBILITY
+	// Encode visibility term to match shadow map pipeline
+	Result = sqrt(Result);
+#endif
+	RWOcclusionMaskUAV[PixelCoord].rgb = Result;
+	RWOcclusionMaskUAV[PixelCoord].a = 1.0;
+	RWRayDistanceUAV[PixelCoord] = (HitCount > 0.0) ? RayDistance / HitCount : 1.0e27;
+}
diff --git a/Engine/Shaders/Private/RayTracing/RayTracingReflections.usf b/Engine/Shaders/Private/RayTracing/RayTracingReflections.usf
new file mode 100644
--- /dev/null
+++ b/Engine/Shaders/Private/RayTracing/RayTracingReflections.usf
@@ -0,0 +1,380 @@
+// Copyright 1998-2018 Epic Games, Inc. All Rights Reserved.
+
+#ifndef DIM_DENOISER_OUTPUT
+	#define DIM_DENOISER_OUTPUT 0
+#endif
+
+#include "../Common.ush"
+
+float2 SobolIndexToUniformUnitSquare(uint2 SobolRand)
+{
+	return float2(SobolRand) * rcp(65536.0) + rcp(65536.0 * 2.0);
+}
+
+#define SUPPORT_CONTACT_SHADOWS		0
+#define USE_SOURCE_TEXTURE			0 //Used in RectLight.ush  //#dxr_todo: enable textures
+
+// Light types: should match SceneTypes.h until there is a common header
+#define LIGHT_TYPE_DIRECTIONAL		0 
+#define LIGHT_TYPE_POINT			1 
+#define LIGHT_TYPE_SPOT				2 
+#define LIGHT_TYPE_RECT				3 
+#define LIGHT_TYPE_MAX				4 
+
+#include "../DeferredShadingCommon.ush"
+#include "../DeferredLightingCommon.ush"
+#include "../Montecarlo.ush"
+#include "../PathTracing/Utilities/PathTracingRandomSequence.ush" // #dxr-todo: move to a common folder
+#include "RayTracingCommon.ush"
+#include "RayTracingHitGroupCommon.ush"
+
+uint SamplesPerPixel;
+
+RaytracingAccelerationStructure TLAS;
+
+RWTexture2D<float4> ColorOutput;
+
+uint CalcLinearIndex(uint2 PixelCoord)
+{
+	return PixelCoord.y * View.BufferSizeAndInvSize.x + PixelCoord.x;
+}
+
+// #dxr-todo: move to RayTracingCommon.ush
+float3 ReconstructWorldPositionFromDepth(float2 UV, float Depth)
+{
+	float2 ScreenPosition = (UV - View.ScreenPositionScaleBias.wz) / View.ScreenPositionScaleBias.xy;
+	float4 HomogeneousWorldPosition = mul(float4(ScreenPosition * Depth, Depth, 1), View.ScreenToWorld);
+	float3 WorldPosition = HomogeneousWorldPosition.xyz / HomogeneousWorldPosition.w;
+
+	return WorldPosition;
+}
+
+float3 GenerateReflectedRayDirection(
+	float3 IncidentDirection,
+	float3 WorldNormal,
+	float Roughness,
+	float2 RandSample
+)
+{
+	float3 RayDirection;
+	if (Roughness < 0.05 ) //ReflectionSmoothClamp)
+	{
+		RayDirection = reflect(IncidentDirection, WorldNormal);
+	}
+	else
+	{
+		float NoV = saturate(dot(-IncidentDirection, WorldNormal));
+		float2 E = RandSample;
+		float4 Sample = ImportanceSampleGGX(E, Pow4(Roughness));
+		float3 H = TangentToWorld(Sample.xyz, WorldNormal);
+		RayDirection = reflect(IncidentDirection, H);
+	}
+
+	return RayDirection;
+}
+
+#if 0
+float3 GatherRadiance(float CompositeAlpha, float3 WorldPosition, float3 RayDirection, float Roughness, float IndirectIrradiance)
+{
+	float IndirectSpecularOcclusion = 1.0f;
+	float3 ExtraIndirectSpecular = 0;
+	return CompositeReflectionCapturesAndSkylight(CompositeAlpha, WorldPosition, RayDirection, Roughness, IndirectIrradiance, IndirectSpecularOcclusion, ExtraIndirectSpecular, 0, 0);
+}
+
+float4 BlendWithEnvironment(inout float3 BRDFThroughput, const in FMaterialClosestHitPayload CurrentBounceHitInfo, float3 WorldPosition, float4 ReflectionColor, float3 RayDirection, float3 Normal)
+{
+#if USE_PREEXPOSURE
+	const float PreExposure = View.PreExposure;
+#else
+	const float PreExposure = 1.f;
+#endif
+
+	float3 ProbeReflectDirection = reflect(RayDirection, Normal);
+	float  NoV = saturate(dot(-RayDirection, Normal));
+	float IndirectIrradiance = CurrentBounceHitInfo.IndirectIrradiance;
+
+	// Point lobe in off-specular peak direction
+	ProbeReflectDirection = GetOffSpecularPeakReflectionDir(Normal, ProbeReflectDirection, CurrentBounceHitInfo.Roughness);
+
+	float RoughnessFade = saturate(CurrentBounceHitInfo.Roughness * RTRRoughnessFadeFactor + 2);
+	ReflectionColor.a = RoughnessFade; 	
+	  
+	float4 FinalReflectionColor = float4(ReflectionColor.rgb, 1 - ReflectionColor.a); 
+	float AO = CurrentBounceHitInfo.GBufferAO;
+	float RoughnessSq = CurrentBounceHitInfo.Roughness * CurrentBounceHitInfo.Roughness;
+	float SpecularOcclusion = GetSpecularOcclusion(NoV, RoughnessSq, AO);
+	FinalReflectionColor.a *= SpecularOcclusion;
+
+	FinalReflectionColor.rgb += PreExposure * GatherRadiance(FinalReflectionColor.a, WorldPosition, ProbeReflectDirection, CurrentBounceHitInfo.Roughness, IndirectIrradiance);
+
+	float3 BRDF = EnvBRDF(CurrentBounceHitInfo.SpecularColor, CurrentBounceHitInfo.Roughness, NoV);
+	FinalReflectionColor.rgb *= BRDF;
+
+	FinalReflectionColor.rgb *= BRDFThroughput;
+	BRDFThroughput *= BRDF;
+
+	return FinalReflectionColor;
+}
+#endif
+
+float3 ComputeDirectLighting(RayDesc Ray, FMaterialClosestHitPayload Payload)
+{
+	float3 DirectLighting = float3(0.0, 0.0, 0.0);
+
+	float3 WorldPosition = Payload.WorldPos;
+	float3 CameraVector = -Ray.Direction;
+	float AmbientOcclusion = 1.0;
+
+	// #dxr_todo: review when interpolants fill the Payload with all the data we need
+	FGBufferData GBufferData = (FGBufferData)0;
+	GBufferData.Depth = 1.f; // Do not use depth
+	GBufferData.WorldNormal = Payload.WorldNormal;
+	GBufferData.BaseColor = Payload.BaseColor;
+	GBufferData.Metallic = Payload.Metallic;
+	GBufferData.Specular = Payload.Specular;
+	GBufferData.Roughness = Payload.Roughness;
+	GBufferData.ShadingModelID = Payload.ShadingModelID;
+	GBufferData.CustomData = Payload.CustomData;
+	GBufferData.GBufferAO = Payload.GBufferAO;
+	GBufferData.IndirectIrradiance = Payload.IndirectIrradiance;
+	GBufferData.ShadingModelID = 1; 
+	   
+	//#dxr_todo: calculate on the fly or store in Payload?
+#if 1
+	GBufferData.SpecularColor = lerp(0.08 * GBufferData.Specular.xxx, GBufferData.BaseColor, GBufferData.Metallic);
+	GBufferData.DiffuseColor = GBufferData.BaseColor - GBufferData.BaseColor * GBufferData.Metallic;
+#else	
+	GBufferData.SpecularColor = Payload.SpecularColor;	
+	GBufferData.DiffuseColor = Payload.DiffuseColor;	
+#endif
+
+	//#dxr_todo: think on the more efficient way of iterate through lights (split in light types, etc)
+	for (uint DirIndex = 0; DirIndex < ReflectionLightsData.Count; DirIndex++)
+	{
+		FDeferredLightData LightData = (FDeferredLightData)0;
+
+		uint LightType = ReflectionLightsData.Type[DirIndex];
+		LightData.Position = ReflectionLightsData.LightPosition[DirIndex];
+		LightData.InvRadius	= ReflectionLightsData.LightInvRadius[DirIndex];
+		LightData.Color = ReflectionLightsData.LightColor[DirIndex];
+		LightData.FalloffExponent = ReflectionLightsData.LightFalloffExponent[DirIndex];
+		LightData.Direction = ReflectionLightsData.Direction[DirIndex];
+		LightData.Tangent = ReflectionLightsData.Tangent[DirIndex];
+		LightData.SpotAngles = ReflectionLightsData.SpotAngles[DirIndex];
+		LightData.SourceRadius = ReflectionLightsData.SourceRadius[DirIndex];
+		LightData.SourceLength = ReflectionLightsData.SourceLength[DirIndex];
+		LightData.SoftSourceRadius = ReflectionLightsData.SoftSourceRadius[DirIndex];
+		LightData.SpecularScale = ReflectionLightsData.SpecularScale[DirIndex];
+		LightData.ContactShadowLength = 0.0;
+		LightData.DistanceFadeMAD = ReflectionLightsData.DistanceFadeMAD[DirIndex];
+		LightData.ShadowMapChannelMask = float4(0, 0, 0, 0); // #dxr_todo;
+		LightData.ShadowedBits = 0; // Not lit dynamic shadows 
+		LightData.ContactShadowLengthInWS = false;
+
+		LightData.bRadialLight = (LightType != LIGHT_TYPE_DIRECTIONAL);
+		LightData.bSpotLight = (LightType == LIGHT_TYPE_SPOT);
+		LightData.bRectLight = (LightType == LIGHT_TYPE_RECT);
+
+		if (LightType == LIGHT_TYPE_DIRECTIONAL)
+		{
+			LightData.bInverseSquared = false;
+		}
+		else
+		{
+			LightData.bInverseSquared = LightData.FalloffExponent == 0;
+		}
+
+		// #dxr_todo: get a valid texture once they are properly bound
+		Texture2D LightSourceTexture = ReflectionLightsData.DummyRectLightTexture;
+
+		float4 LightAttenuation = 1.0f;
+		float LightProfileMultiplier = 1.0;//#dxr_todo
+
+		float3 LightContribution = GetDynamicLighting(WorldPosition, -CameraVector, GBufferData, AmbientOcclusion, GBufferData.ShadingModelID, LightData, LightAttenuation, 0.5, uint2(0, 0), LightSourceTexture).xyz;
+		DirectLighting += LightContribution * LightProfileMultiplier;
+	}
+
+	return DirectLighting;
+}
+
+FMaterialClosestHitPayload TraceRayAndAccumulateResults(
+	RayDesc Ray,
+	inout float3 Radiance)
+{
+	// Call TraceRay iteratively until alpha is 0 and accumulate the result
+	//#dxr-todo: iterate until opacity == 1 if the material is not opaque
+	//#dxr-todo: add optimizations (TMax depending on roughness, etc)
+
+	uint RayFlags = RAY_FLAG_CULL_BACK_FACING_TRIANGLES;
+
+	FMaterialClosestHitPayload Payload = TraceRayInternal(
+		TLAS,      // AccelerationStructure
+		RayFlags,
+		0xFF,      // InstanceInclusionMask
+		0,         // RayContributionToHitGroupIndex
+		1,         // MultiplierForGeometryContributionToShaderIndex
+		0,         // MissShaderIndex
+		Ray        // RayDesc
+	);
+
+
+	float3 DirectLighting = ComputeDirectLighting(Ray, Payload);
+
+	// Transform NaNs to black, transform negative colors to black.
+	DirectLighting = -min(-DirectLighting, 0.0);
+
+	Radiance += DirectLighting;
+
+	return Payload;
+}
+
+uint ReverseNBits(uint Bitfield, const uint BitCount) // TODO: where to put that guy?
+{
+	uint Out = 0;
+	UNROLL
+	for (uint i = 0; i < BitCount; i++)
+	{
+		uint bIsSet = ((Bitfield & (0x1 << i)) >> i);
+		Out |= bIsSet << (BitCount - 1 - i);
+	}
+	return Out;
+}
+
+
+[shader("raygeneration")]
+void RayTracingReflectionsRGS()
+{
+	uint2 PixelCoord = DispatchRaysIndex().xy;
+
+	RandomSequence RandSequence;
+	uint LinearIndex = CalcLinearIndex(PixelCoord);
+	RandomSequence_Initialize(RandSequence, LinearIndex, View.FrameNumber);
+
+	float2 InvBufferSize = View.BufferSizeAndInvSize.zw;
+	float2 UV = (float2(PixelCoord) + 0.5) * InvBufferSize;
+	FScreenSpaceData ScreenSpaceData = GetScreenSpaceData(UV);
+	float Depth = ScreenSpaceData.GBuffer.Depth;
+	float3 CameraOrigin = ReconstructWorldPositionFromDepth(UV, 0.0f);
+	float3 WorldPosition = ReconstructWorldPositionFromDepth(UV, Depth);
+	float3 CameraDirection = normalize(WorldPosition - CameraOrigin);
+	//#dxr - todo: use also the shading normal (interpolate, or apply a more sophisticated approach)
+	float3 WorldNormal = ScreenSpaceData.GBuffer.WorldNormal;
+	float Roughness = ScreenSpaceData.GBuffer.Roughness;
+	
+	float4 ReflectedColor = 0.0;
+	float3 BRDFThroughput = float3(1.0f, 1.0f, 1.0f);
+
+	uint MaxBounces = 1; //#dxr-todo: only works with 1 bounce until hit shaders are supported
+	uint LocalSamplesPerPixel = SamplesPerPixel;
+	float RoughnessThreshold = 1.0f; //#dxr-todo: for internal testing only, remove later
+
+	// Mask out depth values at or beyond far plane
+	if (Depth >= SCENE_TEXTURES_DISABLED_SCENE_DEPTH_VALUE)
+	{
+		LocalSamplesPerPixel = 0;
+	}
+
+	float HitDistance = 0.0;
+	float3 Radiance = float3(0, 0, 0);
+
+	if (Roughness < RoughnessThreshold)
+	{
+		for (uint SampleIndex = 0; SampleIndex < LocalSamplesPerPixel; ++SampleIndex)
+		{
+			float3 IncidentDirection = CameraDirection;
+			for (uint BounceIndex = 0; BounceIndex < MaxBounces; ++BounceIndex)
+			{
+				RayDesc Ray;
+				Ray.Origin = WorldPosition + WorldNormal * 0.01;
+				Ray.TMin = 0.0;
+				Ray.TMax = 1e27f;
+
+				uint DummyVariable;
+				float2 RandSample = RandomSequence_GenerateSample2D(RandSequence, DummyVariable);
+
+				if (LocalSamplesPerPixel == 1)
+				{
+					const uint TemporalBits = 10;
+					uint FrameIndexMod1024 = ReverseNBits(GetPowerOfTwoModulatedFrameIndex(1 << TemporalBits), TemporalBits);
+
+					uint2 SobolBase = SobolPixel(PixelCoord);
+					uint2 SobolFrame = SobolIndex(SobolBase, FrameIndexMod1024, TemporalBits);
+
+					RandSample = SobolIndexToUniformUnitSquare(SobolFrame);
+				}
+
+				Ray.Direction = GenerateReflectedRayDirection(IncidentDirection, WorldNormal, Roughness, RandSample);
+				if (dot(WorldNormal, Ray.Direction) <= 0.0)
+				{
+					// Refraction rays is handled separately
+					break;
+				}
+
+				//#dxr-todo: apply heuristics to do different shading computations depending on BounceIndex
+
+				uint RayFlags = RAY_FLAG_CULL_BACK_FACING_TRIANGLES;
+
+				FMaterialClosestHitPayload Payload = TraceRayAndAccumulateResults(Ray, Radiance);
+
+				if (Payload.HitT < 0)
+				{
+					break;
+				}
+
+#if 1
+				float4 PathRadiance = float4(Radiance, 1.0); 
+				if (BounceIndex == 0)
+				{
+					ReflectedColor += PathRadiance;
+				}
+				else
+				{
+					float4 FinalReflectionColor = PathRadiance;
+					//float4 FinalReflectionColor = BlendWithEnvironment(BRDFThroughput, Payload, Ray.Origin, PathRadiance, Ray.Direction, Payload.WorldNormal);
+					ReflectedColor += FinalReflectionColor;
+				}
+#else // Debugging, remove later
+				ReflectedColor += float4(Payload.BaseColor, 1.0); 
+#endif
+
+				IncidentDirection = Ray.Direction;
+
+				HitDistance = Payload.HitT;
+			}
+		}
+	}
+
+	////uint NumLights = ForwardLightData.NumDirectionalLights;
+
+	if (SamplesPerPixel > 0)
+	{
+		ReflectedColor /= SamplesPerPixel;
+	}
+	ColorOutput[PixelCoord] = float4(ReflectedColor);
+}
+
+
+[shader("miss")]
+void RayTracingReflectionsMainMS(inout FPackedMaterialClosestHitPayload PackedPayload)
+{
+	PackedPayload.HitT = -1;
+}
+
+
+[shader("closesthit")]
+void RayTracingReflectionsMainCHS(inout FPackedMaterialClosestHitPayload PackedPayload, in FDefaultAttributes Attributes)
+{
+	FMaterialClosestHitPayload Payload = (FMaterialClosestHitPayload)0;
+
+	Payload.HitT = RayTCurrent();
+	Payload.BaseColor = float3(0.2, 0.3, 0.4);
+
+	FTriangleBaseAttributes Triangle = LoadTriangleBaseAttributes(PrimitiveIndex());
+	float3 Edge0 = Triangle.LocalPositions[2] - Triangle.LocalPositions[0];
+	float3 Edge1 = Triangle.LocalPositions[1] - Triangle.LocalPositions[0];
+	float3x3 WorldToLocal = (float3x3)WorldToObject();
+	float3x3 LocalToWorldNormal = transpose(WorldToLocal);
+	Payload.WorldNormal = normalize(mul(LocalToWorldNormal, cross(Edge0, Edge1)));
+
+	PackedPayload = PackRayTracingPayload(Payload);
+}
diff --git a/Engine/Shaders/Private/RayTracing/RayTracingReflectionsRG.usf b/Engine/Shaders/Private/RayTracing/RayTracingReflectionsRG.usf
deleted file mode 100644
diff --git a/Engine/Shaders/Private/RayTracing/RayTracingSkyLightRG.usf b/Engine/Shaders/Private/RayTracing/RayTracingSkyLightRG.usf
deleted file mode 100644
diff --git a/Engine/Shaders/Private/RayTracing/RayTracingSkyLightRGS.usf b/Engine/Shaders/Private/RayTracing/RayTracingSkyLightRGS.usf
new file mode 100644
--- /dev/null
+++ b/Engine/Shaders/Private/RayTracing/RayTracingSkyLightRGS.usf
@@ -0,0 +1,477 @@
+// Copyright 1998-2018 Epic Games, Inc. All Rights Reserved.
+
+#include "../Common.ush"
+#include "../MonteCarlo.ush"
+#include "../DeferredShadingCommon.ush"
+#include "SkyLightMipTreeCommon.ush"
+#include "RayTracingCommon.ush"
+#include "RayTracingDirectionalLight.ush"
+#include "RayTracingRectLight.ush"
+
+RaytracingAccelerationStructure TLAS;
+RWTexture2D<float4> RWOcclusionMaskUAV;
+RWTexture2D<float> RWRayDistanceUAV;
+
+#define USE_MIP_PDF_TEXTURE 1
+#define USE_SOLID_ANGLE_TEXTURE 1
+
+uint CalcLinearIndex(uint2 PixelCoord)
+{
+	return PixelCoord.y * View.BufferSizeAndInvSize.x + PixelCoord.x;
+}
+
+uint SelectFace(uint MipCount, inout RandomSequence RandSequence, out float FacePdf)
+{
+	uint BufferOffset = BufferOffsetAtPixel(uint2(0, 0), MipCount, SkyLight.MipDimensions);
+
+	float FaceCdf[6];
+	FaceCdf[0] = SkyLight.MipTreePosX[BufferOffset];
+	FaceCdf[1] = FaceCdf[0] + SkyLight.MipTreeNegX[BufferOffset];
+	FaceCdf[2] = FaceCdf[1] + SkyLight.MipTreePosY[BufferOffset];
+	FaceCdf[3] = FaceCdf[2] + SkyLight.MipTreeNegY[BufferOffset];
+	FaceCdf[4] = FaceCdf[3] + SkyLight.MipTreePosZ[BufferOffset];
+	FaceCdf[5] = FaceCdf[4] + SkyLight.MipTreeNegZ[BufferOffset];
+	float Sum = FaceCdf[5];
+
+	for (uint Index = 0; Index < 5; ++Index)
+	{
+		FaceCdf[Index] /= Sum;
+	}
+	FaceCdf[5] = 1.0;
+
+	// Draw random variable
+	uint DummyVariable;
+	float RandSample = RandomSequence_GenerateSample1D(RandSequence, DummyVariable);
+
+	// Determine CDF entry
+	uint FaceIndex = 0;
+	while (RandSample > FaceCdf[FaceIndex])
+	{
+		FaceIndex++;
+	}
+
+	// Calculate PDF and return entry
+	FacePdf = FaceCdf[FaceIndex];
+	if (FaceIndex > 0)
+		FacePdf -= FaceCdf[FaceIndex - 1];
+	return FaceIndex;
+}
+
+float4 SampleMipTree(uint FaceIndex, uint4 BufferOffset)
+{
+	float4 Values;
+	if (FaceIndex == 0)
+	{
+		Values.x = SkyLight.MipTreePosX[BufferOffset.x];
+		Values.y = SkyLight.MipTreePosX[BufferOffset.y];
+		Values.z = SkyLight.MipTreePosX[BufferOffset.z];
+		Values.w = SkyLight.MipTreePosX[BufferOffset.w];
+	}
+	else if (FaceIndex == 1)
+	{
+		Values.x = SkyLight.MipTreeNegX[BufferOffset.x];
+		Values.y = SkyLight.MipTreeNegX[BufferOffset.y];
+		Values.z = SkyLight.MipTreeNegX[BufferOffset.z];
+		Values.w = SkyLight.MipTreeNegX[BufferOffset.w];
+	}
+	else if (FaceIndex == 2)
+	{
+		Values.x = SkyLight.MipTreePosY[BufferOffset.x];
+		Values.y = SkyLight.MipTreePosY[BufferOffset.y];
+		Values.z = SkyLight.MipTreePosY[BufferOffset.z];
+		Values.w = SkyLight.MipTreePosY[BufferOffset.w];
+	}
+	else if (FaceIndex == 3)
+	{
+		Values.x = SkyLight.MipTreeNegY[BufferOffset.x];
+		Values.y = SkyLight.MipTreeNegY[BufferOffset.y];
+		Values.z = SkyLight.MipTreeNegY[BufferOffset.z];
+		Values.w = SkyLight.MipTreeNegY[BufferOffset.w];
+	}
+	else if (FaceIndex == 4)
+	{
+		Values.x = SkyLight.MipTreePosZ[BufferOffset.x];
+		Values.y = SkyLight.MipTreePosZ[BufferOffset.y];
+		Values.z = SkyLight.MipTreePosZ[BufferOffset.z];
+		Values.w = SkyLight.MipTreePosZ[BufferOffset.w];
+	}
+	else
+	{
+		Values.x = SkyLight.MipTreeNegZ[BufferOffset.x];
+		Values.y = SkyLight.MipTreeNegZ[BufferOffset.y];
+		Values.z = SkyLight.MipTreeNegZ[BufferOffset.z];
+		Values.w = SkyLight.MipTreeNegZ[BufferOffset.w];
+	}
+
+	return Values;
+}
+
+// #dxr_todo: Copied from PathTracingSkyLight.ush
+float3 EvalSkyLight(float3 Direction)
+{
+	return SkyLight.Color * TextureCubeSampleLevel(SkyLight.Texture, SkyLight.TextureSampler, Direction, 0).rgb;
+}
+
+float PdfFace(uint MipCount, uint FaceIndex)
+{
+	// Build Cdf
+	uint BufferOffset = BufferOffsetAtPixel(uint2(0, 0), MipCount, SkyLight.MipDimensions);
+
+	float FaceCdf[6];
+	FaceCdf[0] = SkyLight.MipTreePosX[BufferOffset];
+	FaceCdf[1] = FaceCdf[0] + SkyLight.MipTreeNegX[BufferOffset];
+	FaceCdf[2] = FaceCdf[1] + SkyLight.MipTreePosY[BufferOffset];
+	FaceCdf[3] = FaceCdf[2] + SkyLight.MipTreeNegY[BufferOffset];
+	FaceCdf[4] = FaceCdf[3] + SkyLight.MipTreePosZ[BufferOffset];
+	FaceCdf[5] = FaceCdf[4] + SkyLight.MipTreeNegZ[BufferOffset];
+	float Sum = FaceCdf[5];
+
+	for (uint Index = 0; Index < 5; ++Index)
+	{
+		FaceCdf[Index] /= Sum;
+	}
+	FaceCdf[5] = 1.0;
+
+	float Pdf = FaceCdf[FaceIndex];
+	if (FaceIndex > 0)
+		Pdf -= FaceCdf[FaceIndex - 1];
+	return Pdf;
+}
+
+float SampleMipPdf(uint3 TextureCoord, uint StopLevel)
+{
+	float MipPdf = 1.0;
+	uint BufferOffset = BufferOffsetAtPixel(TextureCoord.xy, StopLevel, SkyLight.MipDimensions);
+	if (TextureCoord.z == 0)
+	{
+		MipPdf = SkyLight.MipTreePdfPosX[BufferOffset];
+	}
+	else if (TextureCoord.z == 1)
+	{
+		MipPdf = SkyLight.MipTreePdfNegX[BufferOffset];
+	}
+	else if (TextureCoord.z == 2)
+	{
+		MipPdf = SkyLight.MipTreePdfPosY[BufferOffset];
+	}
+	else if (TextureCoord.z == 3)
+	{
+		MipPdf = SkyLight.MipTreePdfNegY[BufferOffset];
+	}
+	else if (TextureCoord.z == 4)
+	{
+		MipPdf = SkyLight.MipTreePdfPosZ[BufferOffset];
+	}
+	else // if (TextureCoord.z == 5)
+	{
+		MipPdf = SkyLight.MipTreePdfNegZ[BufferOffset];
+	}
+	return MipPdf;
+}
+
+float PdfSkyLightOcclusionRay(
+	float3 WorldNormal,
+	float3 WorldDirection
+)
+{
+	// Determine primary direction
+	uint MipCount = Log2(SkyLight.MipDimensions.x);
+	uint2 TextureRes;
+	SkyLight.Texture.GetDimensions(TextureRes.x, TextureRes.y);
+	uint3 TextureCoord = GetTextureCubeCoordinate(WorldDirection, TextureRes);
+	float2 TextureUV = TextureCoord.xy / (float2)TextureRes;
+	uint2 MipCoord = TextureUV * SkyLight.MipDimensions;
+	float FacePdf = PdfFace(MipCount, TextureCoord.z);
+
+	// #dxr_todo: StopLevel can be customized for run-time performance
+	uint StopLevel = min(SkyLight.SamplingStopLevel, MipCount);
+#if USE_MIP_PDF_TEXTURE
+	float MipPdf = SampleMipPdf(uint3(MipCoord, TextureCoord.z), StopLevel);
+
+#if USE_SOLID_ANGLE_TEXTURE
+	uint LinearIndex = BufferOffsetAtPixel(MipCoord, StopLevel, SkyLight.MipDimensions);
+	float AreaPdf = SkyLight.SolidAnglePdf[LinearIndex];
+#else
+	float AreaPdf = 1.0 / TexelCoordSolidAngle(MipCoord, DimensionsAtLevel(SkyLight.MipDimensions, StopLevel));
+#endif
+	return FacePdf * MipPdf * AreaPdf;
+#else
+	uint2 MipPixel = 0;
+	float MipPdf = 1.0;
+	for (uint Index = 0; Index < MipCount - StopLevel; ++Index)
+	{
+		// Build Cdf
+		uint MipIndex = MipCount - Index - 1;
+#if USE_MORTON_CODE
+		uint BufferOffset = BufferOffsetAtPixel(MipPixel, MipIndex, SkyLight.MipDimensions);
+		uint4 BufferOffsets = uint4(BufferOffset, BufferOffset + 1, BufferOffset + 2, BufferOffset + 3);
+#else
+		uint4 BufferOffsets;
+		BufferOffsets.x = BufferOffsetAtPixel(MipPixel, MipIndex, SkyLight.MipDimensions);
+		BufferOffsets.y = BufferOffsetAtPixel(MipPixel + uint2(1, 0), MipIndex, SkyLight.MipDimensions);
+		BufferOffsets.z = BufferOffsetAtPixel(MipPixel + uint2(0, 1), MipIndex, SkyLight.MipDimensions);
+		BufferOffsets.w = BufferOffsetAtPixel(MipPixel + uint2(1, 1), MipIndex, SkyLight.MipDimensions);
+#endif
+		float4 Cdf = SampleMipTree(TextureCoord.z, BufferOffsets);
+		Cdf.y += Cdf.x;
+		Cdf.z += Cdf.y;
+		Cdf.w += Cdf.z;
+		Cdf /= Cdf.w;
+
+		// Determine texel choice
+		uint2 MipOffset = (MipCoord.xy >> MipIndex) - MipPixel;
+		uint LinearIndex = MipOffset.y << 1 | MipOffset.x;
+		MipPdf *= LinearIndex > 0 ? Cdf[LinearIndex] - Cdf[LinearIndex - 1] : Cdf[LinearIndex];
+		MipPixel += MipOffset;
+
+		MipPixel = MipPixel << 1;
+	}
+	MipPdf = max(MipPdf, 0.0);
+	// DEBUG: assert that SampleMipPdf is equivalent for evaluating the decision tree..
+	//MipPdf = SampleMipPdf(uint3(MipCoord, TextureCoord.z), StopLevel);
+
+#if USE_SOLID_ANGLE_TEXTURE
+	uint LinearIndex = BufferOffsetAtPixel(MipCoord, StopLevel, SkyLight.MipDimensions);
+	float AreaPdf = SkyLight.SolidAnglePdf[LinearIndex];
+#else
+	float AreaPdf = 1.0 / TexelCoordSolidAngle(MipCoord, DimensionsAtLevel(SkyLight.MipDimensions, StopLevel));
+#endif
+	return FacePdf * MipPdf * AreaPdf;
+#endif
+}
+
+bool GenerateSkyLightOcclusionRay(
+	float3 WorldPosition,
+	float3 WorldNormal,
+	inout RandomSequence RandSequence,
+	out float3 RayOrigin,
+	out float3 RayDirection,
+	out float RayTMin,
+	out float RayTMax,
+	out float RayPdf
+)
+{
+	// Determine the appropriate face to use
+	uint MipCount = Log2(SkyLight.MipDimensions.x);
+	float FacePdf = 0.0;
+	uint FaceIndex = SelectFace(MipCount, RandSequence, FacePdf);
+
+	uint2 MipPixel = 0;
+	float MipPdf = 1.0;
+	uint StopLevel = min(SkyLight.SamplingStopLevel, MipCount);
+	for (uint Index = 0; Index < MipCount - StopLevel; ++Index)
+	{
+		uint MipIndex = MipCount - Index - 1;
+#if USE_MORTON_CODE
+		uint BufferOffset = BufferOffsetAtPixel(MipPixel, MipIndex, SkyLight.MipDimensions);
+		uint4 BufferOffsets = uint4(BufferOffset, BufferOffset + 1, BufferOffset + 2, BufferOffset + 3);
+#else
+		uint4 BufferOffsets;
+		BufferOffsets.x = BufferOffsetAtPixel(MipPixel, MipIndex, SkyLight.MipDimensions);
+		BufferOffsets.y = BufferOffsetAtPixel(MipPixel + uint2(1, 0), MipIndex, SkyLight.MipDimensions);
+		BufferOffsets.z = BufferOffsetAtPixel(MipPixel + uint2(0, 1), MipIndex, SkyLight.MipDimensions);
+		BufferOffsets.w = BufferOffsetAtPixel(MipPixel + uint2(1, 1), MipIndex, SkyLight.MipDimensions);
+#endif
+		float4 Cdf = SampleMipTree(FaceIndex, BufferOffsets);
+		Cdf.y += Cdf.x;
+		Cdf.z += Cdf.y;
+		Cdf.w += Cdf.z;
+		Cdf /= Cdf.w;
+
+		uint DummyVariable;
+		float RandSample = RandomSequence_GenerateSample1D(RandSequence, DummyVariable);
+		if (RandSample < Cdf.x)
+		{
+			MipPdf *= Cdf.x;
+		}
+		else if (RandSample < Cdf.y)
+		{
+			MipPdf *= Cdf.y - Cdf.x;
+			MipPixel.x += 1;
+		}
+		else if (RandSample < Cdf.z)
+		{
+			MipPdf *= Cdf.z - Cdf.y;
+			MipPixel.y += 1;
+		}
+		else // if (RandSample < Cdf.w)
+		{
+			MipPdf *= Cdf.w - Cdf.z;
+			MipPixel += 1;
+		}
+
+		MipPixel = uint2(MipPixel.x << 1, MipPixel.y << 1);
+	}
+	MipPixel = uint2(MipPixel.x >> 1, MipPixel.y >> 1);
+
+	// DEBUG: assert cached PDF is equivalent to re-evaluation
+	//MipPdf = PdfSkyLightOcclusionRay(WorldNormal, RayDirection);
+
+	uint DummyVariable;
+	float2 PixelOffset = RandomSequence_GenerateSample2D(RandSequence, DummyVariable);
+	float2 MipPixelUV = (MipPixel + PixelOffset) / float2(DimensionsAtLevel(SkyLight.MipDimensions, StopLevel));
+
+	// Construct Cubemap coordinate
+	uint2 TextureDimensions;
+	SkyLight.Texture.GetDimensions(TextureDimensions.x, TextureDimensions.y);
+	float3 TexelCoord = float3(MipPixelUV * TextureDimensions, FaceIndex);
+
+	// Determine area pdf
+#if USE_SOLID_ANGLE_TEXTURE
+	uint LinearIndex = BufferOffsetAtPixel(MipPixel, StopLevel, SkyLight.MipDimensions);
+	float AreaPdf = SkyLight.SolidAnglePdf[LinearIndex];
+#else
+	float AreaPdf = 1.0 / TexelCoordSolidAngle(MipPixel, DimensionsAtLevel(SkyLight.MipDimensions, StopLevel));
+#endif
+
+	RayOrigin = WorldPosition + WorldNormal * 0.01;
+	RayDirection = GetTextureCubeVector(TexelCoord, TextureDimensions);
+	RayTMin = 0.0;
+	RayTMax = 1.0e27;
+	RayPdf = FacePdf * MipPdf * AreaPdf;
+
+	return true;
+}
+
+float PdfLambertOcclusionRay(
+	float3 WorldNormal,
+	float3 WorldDirection
+)
+{
+	float Pdf = max(dot(WorldDirection, WorldNormal), 0.0) / PI;
+	return Pdf;
+}
+
+bool GenerateLambertOcclusionRay(
+	float3 WorldPosition,
+	float3 WorldNormal,
+	inout RandomSequence RandSequence,
+	out float3 RayOrigin,
+	out float3 RayDirection,
+	out float RayTMin,
+	out float RayTMax,
+	out float RayPdf
+)
+{
+	// Draw random variable
+	float2 BufferSize = View.BufferSizeAndInvSize.xy;
+	uint DummyVariable;
+	float2 RandSample = RandomSequence_GenerateSample2D(RandSequence, DummyVariable);
+
+	// Perform cosine-hemispherical sampling and convert to world-space
+	float4 Direction_Tangent = CosineSampleHemisphere(RandSample);
+	float3 Direction_World = TangentToWorld(Direction_Tangent.xyz, WorldNormal);
+
+	RayOrigin = WorldPosition + WorldNormal * 0.01;
+	RayDirection = Direction_World;
+	RayTMin = 0.0;
+	RayTMax = 1.0e27;
+	RayPdf = Direction_Tangent.w;
+
+	return true;
+}
+
+[shader("raygeneration")]
+void SkyLightRGS()
+{
+	uint2 PixelCoord = DispatchRaysIndex().xy;
+
+	RandomSequence RandSequence;
+	uint LinearIndex = CalcLinearIndex(PixelCoord);
+	RandomSequence_Initialize(RandSequence, LinearIndex, View.FrameNumber * SkyLight.SamplesPerPixel);
+	//RandSequence.Type = 2;
+
+	// Get G-Buffer surface data
+	float2 InvBufferSize = View.BufferSizeAndInvSize.zw;
+	float2 UV = (float2(PixelCoord) + 0.5) * InvBufferSize;
+	FScreenSpaceData ScreenSpaceData = GetScreenSpaceData(UV);
+	float Depth = ScreenSpaceData.GBuffer.Depth;
+	float3 WorldPosition = ReconstructWorldPositionFromDepth(UV, Depth);
+	float3 WorldNormal = ScreenSpaceData.GBuffer.WorldNormal;
+	float3 BaseColor = ScreenSpaceData.GBuffer.BaseColor;
+	float Roughness = ScreenSpaceData.GBuffer.Roughness;
+
+	float RayDistance = 0.0;
+	float HitCount = 0.0;
+	uint SamplesPerPixel = SkyLight.SamplesPerPixel;
+
+	// Mask out depth values at or beyond far plane
+	if (Depth >= SCENE_TEXTURES_DISABLED_SCENE_DEPTH_VALUE)
+	{
+		SamplesPerPixel = 0.0;
+	}
+
+	float3 Irradiance = 0.0;
+	for (uint SampleIndex = 0; SampleIndex < SamplesPerPixel; ++SampleIndex)
+	{
+		RayDesc Ray;
+		float RayPdf = 0.0;
+
+		// Determine sky light or lambert ray
+		uint DummyVariable;
+		float RandSample = RandomSequence_GenerateSample1D(RandSequence, DummyVariable);
+		// #dxr_todo: Should strategy selection be driven by CVar?
+		float SkyLightSamplingStrategyPdf = 0.5;
+		float MisWeight = 1.0;
+		if (RandSample < SkyLightSamplingStrategyPdf)
+		{
+			GenerateSkyLightOcclusionRay(WorldPosition, WorldNormal, RandSequence, Ray.Origin, Ray.Direction, Ray.TMin, Ray.TMax, RayPdf);
+			MisWeight = RayPdf / (SkyLightSamplingStrategyPdf * RayPdf + (1.0 - SkyLightSamplingStrategyPdf) * PdfLambertOcclusionRay(WorldNormal, Ray.Direction));
+		}
+		else
+		{
+			GenerateLambertOcclusionRay(WorldPosition, WorldNormal, RandSequence, Ray.Origin, Ray.Direction, Ray.TMin, Ray.TMax, RayPdf);
+			MisWeight = RayPdf / ((1.0 - SkyLightSamplingStrategyPdf) * RayPdf + SkyLightSamplingStrategyPdf * PdfSkyLightOcclusionRay(WorldNormal, Ray.Direction));
+		}
+		
+		// Surface-normal culling
+		float NdotL = dot(WorldNormal, Ray.Direction);
+		if (RayPdf < 0.0 || NdotL <= 0.0)
+		{
+			continue;
+		}
+
+		uint RayFlags = 0
+			| RAY_FLAG_FORCE_OPAQUE; // #dxr_todo: generate a special shader permutation for occlusion ray tracing that includes any-hit shaders for masked geometry
+		FDefaultPayload Payload = (FDefaultPayload)0;
+
+		TraceRay(
+			TLAS,      // AccelerationStructure
+			RayFlags,
+			0xFF,      // InstanceInclusionMask
+			0,         // RayContributionToHitGroupIndex
+			1,         // MultiplierForGeometryContributionToShaderIndex
+			0,         // MissShaderIndex
+			Ray,       // RayDesc
+			Payload    // Payload
+		);
+
+		if (Payload.HitT >= 0.0)
+		{
+			RayDistance += Payload.HitT;
+			HitCount += 1.0;
+		}
+		else
+		{
+			float3 Radiance = EvalSkyLight(Ray.Direction);
+			// #dxr_todo: Apply actual material BRDF
+			float3 Brdf = BaseColor * Roughness / PI;
+			Irradiance += Radiance * Brdf * NdotL * MisWeight / RayPdf;
+		}
+	}
+
+	if (SamplesPerPixel > 0)
+	{
+		Irradiance /= SamplesPerPixel;
+	}
+	RWOcclusionMaskUAV[PixelCoord] = float4(Irradiance, 1.0);
+
+	if (HitCount > 0.0)
+	{
+		RWRayDistanceUAV[PixelCoord] = RayDistance / HitCount;
+	}
+	else
+	{
+		RWRayDistanceUAV[PixelCoord] = 1.0e27;
+	}
+}
\ No newline at end of file
diff --git a/Engine/Source/Runtime/Renderer/Private/RayTracing/RayTracingAmbientOcclusion.cpp b/Engine/Source/Runtime/Renderer/Private/RayTracing/RayTracingAmbientOcclusion.cpp
--- a/Engine/Source/Runtime/Renderer/Private/RayTracing/RayTracingAmbientOcclusion.cpp
+++ b/Engine/Source/Runtime/Renderer/Private/RayTracing/RayTracingAmbientOcclusion.cpp
@@ -51,9 +51,9 @@ IMPLEMENT_GLOBAL_SHADER_PARAMETER_STRUCT(FAmbientOcclusionData, "AmbientOcclusio
 DECLARE_GPU_STAT_NAMED(RayTracingAmbientOcclusion, TEXT("RTAO"));
 DECLARE_GPU_STAT_NAMED(CompositeRayTracingAmbientOcclusion, TEXT("Denoising: RTAO"));
 
-class FAmbientOcclusionRG : public FGlobalShader
+class FAmbientOcclusionRGS : public FGlobalShader
 {
-	DECLARE_SHADER_TYPE(FAmbientOcclusionRG, Global)
+	DECLARE_SHADER_TYPE(FAmbientOcclusionRGS, Global)
 
 	static void ModifyCompilationEnvironment(const FGlobalShaderPermutationParameters& Parameters, FShaderCompilerEnvironment& OutEnvironment)
 	{
@@ -65,10 +65,10 @@ public:
 		return IsFeatureLevelSupported(Parameters.Platform, ERHIFeatureLevel::SM5) && IsRayTracingSupportedForThisProject();
 	}
 
-	FAmbientOcclusionRG() {}
-	virtual ~FAmbientOcclusionRG() {}
+	FAmbientOcclusionRGS() {}
+	virtual ~FAmbientOcclusionRGS() {}
 
-	FAmbientOcclusionRG(const ShaderMetaType::CompiledShaderInitializerType& Initializer)
+	FAmbientOcclusionRGS(const ShaderMetaType::CompiledShaderInitializerType& Initializer)
 		: FGlobalShader(Initializer)
 	{
 		ViewParameter.Bind(Initializer.ParameterMap, TEXT("View"));
@@ -131,7 +131,7 @@ private:
 	FShaderResourceParameter RayDistanceUAVParameter;
 };
 
-IMPLEMENT_SHADER_TYPE(, FAmbientOcclusionRG, TEXT("/Engine/Private/RayTracing/RayTracingAmbientOcclusionRG.usf"), TEXT("AmbientOcclusionRG"), SF_RayGen)
+IMPLEMENT_SHADER_TYPE(, FAmbientOcclusionRGS, TEXT("/Engine/Private/RayTracing/RayTracingAmbientOcclusionRGS.usf"), TEXT("AmbientOcclusionRGS"), SF_RayGen)
 
 void FDeferredShadingSceneRenderer::RenderRayTracingAmbientOcclusion(
 	FRHICommandListImmediate& RHICmdList,
@@ -167,7 +167,7 @@ void FDeferredShadingSceneRenderer::RenderRayTracingAmbientOcclusion(
 		FViewInfo& View = Views[ViewIndex];
 		FIntPoint ViewSize = View.ViewRect.Size();
 
-		TShaderMapRef<FAmbientOcclusionRG> AmbientOcclusionRayGenerationShader(GetGlobalShaderMap(FeatureLevel));
+		TShaderMapRef<FAmbientOcclusionRGS> AmbientOcclusionRayGenerationShader(GetGlobalShaderMap(FeatureLevel));
 		FSceneTexturesUniformParameters SceneTextures;
 		SetupSceneTextureUniformParameters(SceneContext, FeatureLevel, ESceneTextureSetupMode::All, SceneTextures);
 		FUniformBufferRHIRef SceneTexturesUniformBuffer = RHICreateUniformBuffer(&SceneTextures, FSceneTexturesUniformParameters::StaticStructMetadata.GetLayout(), EUniformBufferUsage::UniformBuffer_SingleDraw);
diff --git a/Engine/Source/Runtime/Renderer/Private/RayTracing/RayTracingBarycentrics.cpp b/Engine/Source/Runtime/Renderer/Private/RayTracing/RayTracingBarycentrics.cpp
--- a/Engine/Source/Runtime/Renderer/Private/RayTracing/RayTracingBarycentrics.cpp
+++ b/Engine/Source/Runtime/Renderer/Private/RayTracing/RayTracingBarycentrics.cpp
@@ -9,10 +9,10 @@
 #include "SceneRenderTargets.h"
 #include "RenderGraphBuilder.h"
 
-class FRayTracingBarycentricsRG : public FGlobalShader
+class FRayTracingBarycentricsRGS : public FGlobalShader
 {
-	DECLARE_GLOBAL_SHADER(FRayTracingBarycentricsRG)
-	SHADER_USE_ROOT_PARAMETER_STRUCT(FRayTracingBarycentricsRG, FGlobalShader)
+	DECLARE_GLOBAL_SHADER(FRayTracingBarycentricsRGS)
+	SHADER_USE_ROOT_PARAMETER_STRUCT(FRayTracingBarycentricsRGS, FGlobalShader)
 
 	BEGIN_SHADER_PARAMETER_STRUCT(FParameters, )
 		SHADER_PARAMETER_SRV(RaytracingAccelerationStructure, TLAS)
@@ -25,7 +25,7 @@ class FRayTracingBarycentricsRG : public FGlobalShader
 		return IsFeatureLevelSupported(Parameters.Platform, ERHIFeatureLevel::SM5) && IsRayTracingSupportedForThisProject();
 	}
 };
-IMPLEMENT_GLOBAL_SHADER(FRayTracingBarycentricsRG, "/Engine/Private/RayTracing/RayTracingBarycentrics.usf", "RayTracingBarycentricsMainRG", SF_RayGen);
+IMPLEMENT_GLOBAL_SHADER(FRayTracingBarycentricsRGS, "/Engine/Private/RayTracing/RayTracingBarycentrics.usf", "RayTracingBarycentricsMainRGS", SF_RayGen);
 
 // Example ray miss shader
 class FRayTracingBarycentricsMS : public FGlobalShader
@@ -79,7 +79,7 @@ void FDeferredShadingSceneRenderer::RenderRayTracedBarycentrics(FRHICommandListI
 
 	TShaderMap<FGlobalShaderType>* ShaderMap = GetGlobalShaderMap(FeatureLevel);
 
-	auto RayGenShader = ShaderMap->GetShader<FRayTracingBarycentricsRG>();
+	auto RayGenShader = ShaderMap->GetShader<FRayTracingBarycentricsRGS>();
 	auto ClosestHitShader = ShaderMap->GetShader<FRayTracingBarycentricsCHS>();
 	auto MissShader = ShaderMap->GetShader<FRayTracingBarycentricsMS>();
 
@@ -92,7 +92,7 @@ void FDeferredShadingSceneRenderer::RenderRayTracedBarycentrics(FRHICommandListI
 
 	FRayTracingSceneRHIParamRef RayTracingSceneRHI = View.PerViewRayTracingScene.RayTracingSceneRHI;
 
-	FRayTracingBarycentricsRG::FParameters* RayGenParameters = GraphBuilder.AllocParameters<FRayTracingBarycentricsRG::FParameters>();
+	FRayTracingBarycentricsRGS::FParameters* RayGenParameters = GraphBuilder.AllocParameters<FRayTracingBarycentricsRGS::FParameters>();
 
 	RayGenParameters->TLAS = RHIGetAccelerationStructureShaderResourceView(RayTracingSceneRHI);
 	RayGenParameters->ViewUniformBuffer = View.ViewUniformBuffer;
diff --git a/Engine/Source/Runtime/Renderer/Private/RayTracing/RayTracingDebug.cpp b/Engine/Source/Runtime/Renderer/Private/RayTracing/RayTracingDebug.cpp
--- a/Engine/Source/Runtime/Renderer/Private/RayTracing/RayTracingDebug.cpp
+++ b/Engine/Source/Runtime/Renderer/Private/RayTracing/RayTracingDebug.cpp
@@ -15,10 +15,10 @@
 
 DECLARE_GPU_STAT(RayTracingDebug);
 
-class FRayTracingDebugRG : public FGlobalShader
+class FRayTracingDebugRGS : public FGlobalShader
 {
-	DECLARE_GLOBAL_SHADER(FRayTracingDebugRG)
-	SHADER_USE_ROOT_PARAMETER_STRUCT(FRayTracingDebugRG, FGlobalShader)
+	DECLARE_GLOBAL_SHADER(FRayTracingDebugRGS)
+	SHADER_USE_ROOT_PARAMETER_STRUCT(FRayTracingDebugRGS, FGlobalShader)
 
 	BEGIN_SHADER_PARAMETER_STRUCT(FParameters, )
 		SHADER_PARAMETER(uint32, VisualizationMode)
@@ -32,7 +32,7 @@ class FRayTracingDebugRG : public FGlobalShader
 		return IsFeatureLevelSupported(Parameters.Platform, ERHIFeatureLevel::SM5) && IsRayTracingSupportedForThisProject();
 	}
 };
-IMPLEMENT_GLOBAL_SHADER(FRayTracingDebugRG, "/Engine/Private/RayTracing/RayTracingDebug.usf", "RayTracingDebugMainRG", SF_RayGen);
+IMPLEMENT_GLOBAL_SHADER(FRayTracingDebugRGS, "/Engine/Private/RayTracing/RayTracingDebug.usf", "RayTracingDebugMainRGS", SF_RayGen);
 
 class FRayTracingDebugMS : public FGlobalShader
 {
@@ -112,7 +112,7 @@ void FDeferredShadingSceneRenderer::RenderRayTracedDebug(FRHICommandListImmediat
 
 	TShaderMap<FGlobalShaderType>* ShaderMap = GetGlobalShaderMap(FeatureLevel);	
 
-	auto RayGenShader = ShaderMap->GetShader<FRayTracingDebugRG>();
+	auto RayGenShader = ShaderMap->GetShader<FRayTracingDebugRGS>();
 	auto ClosestHitShader = ShaderMap->GetShader<FRayTracingDebugCHS>();
 	auto MissShader = ShaderMap->GetShader<FRayTracingDebugMS>();
 
@@ -123,7 +123,7 @@ void FDeferredShadingSceneRenderer::RenderRayTracedDebug(FRHICommandListImmediat
 
 	FRayTracingSceneRHIParamRef RayTracingSceneRHI = View.PerViewRayTracingScene.RayTracingSceneRHI;
 
-	FRayTracingDebugRG::FParameters* RayGenParameters = GraphBuilder.AllocParameters<FRayTracingDebugRG::FParameters>();
+	FRayTracingDebugRGS::FParameters* RayGenParameters = GraphBuilder.AllocParameters<FRayTracingDebugRGS::FParameters>();
 
 	RayGenParameters->VisualizationMode = DebugVisualizationMode;
 	RayGenParameters->TLAS = RHIGetAccelerationStructureShaderResourceView(RayTracingSceneRHI);
diff --git a/Engine/Source/Runtime/Renderer/Private/RayTracing/RayTracingReflections.cpp b/Engine/Source/Runtime/Renderer/Private/RayTracing/RayTracingReflections.cpp
--- a/Engine/Source/Runtime/Renderer/Private/RayTracing/RayTracingReflections.cpp
+++ b/Engine/Source/Runtime/Renderer/Private/RayTracing/RayTracingReflections.cpp
@@ -146,9 +146,9 @@ class FRayTracingReflectionsMS : public FGlobalShader
 	using FParameters = FEmptyShaderParameters;
 };
 
-IMPLEMENT_GLOBAL_SHADER(FRayTracingReflectionsRG, "/Engine/Private/RayTracing/RayTracingReflectionsRG.usf", "RayTracingReflectionsRG", SF_RayGen);
-IMPLEMENT_GLOBAL_SHADER(FRayTracingReflectionsCHS, "/Engine/Private/RayTracing/RayTracingReflectionsRG.usf", "RayTracingReflectionsMainCHS", SF_RayHitGroup);
-IMPLEMENT_GLOBAL_SHADER(FRayTracingReflectionsMS, "/Engine/Private/RayTracing/RayTracingReflectionsRG.usf", "RayTracingReflectionsMainMS", SF_RayMiss);
+IMPLEMENT_GLOBAL_SHADER(FRayTracingReflectionsRG, "/Engine/Private/RayTracing/RayTracingReflections.usf", "RayTracingReflectionsRGS", SF_RayGen);
+IMPLEMENT_GLOBAL_SHADER(FRayTracingReflectionsCHS, "/Engine/Private/RayTracing/RayTracingReflections.usf", "RayTracingReflectionsMainCHS", SF_RayHitGroup);
+IMPLEMENT_GLOBAL_SHADER(FRayTracingReflectionsMS, "/Engine/Private/RayTracing/RayTracingReflections.usf", "RayTracingReflectionsMainMS", SF_RayMiss);
 
 #endif // RHI_RAYTRACING
 
diff --git a/Engine/Source/Runtime/Renderer/Private/RayTracing/RayTracingShadows.cpp b/Engine/Source/Runtime/Renderer/Private/RayTracing/RayTracingShadows.cpp
--- a/Engine/Source/Runtime/Renderer/Private/RayTracing/RayTracingShadows.cpp
+++ b/Engine/Source/Runtime/Renderer/Private/RayTracing/RayTracingShadows.cpp
@@ -36,10 +36,10 @@ static TAutoConsoleVariable<int32> CVarShadowUseDenoiser(
 	TEXT(" 2: GScreenSpaceDenoiser witch may be overriden by a third party plugin.\n"),
 	ECVF_RenderThreadSafe);
 
-class FOcclusionRG : public FGlobalShader
+class FOcclusionRGS : public FGlobalShader
 {
-	DECLARE_GLOBAL_SHADER(FOcclusionRG)
-	SHADER_USE_ROOT_PARAMETER_STRUCT(FOcclusionRG, FGlobalShader)
+	DECLARE_GLOBAL_SHADER(FOcclusionRGS)
+	SHADER_USE_ROOT_PARAMETER_STRUCT(FOcclusionRGS, FGlobalShader)
 
 	class FLightTypeDim : SHADER_PERMUTATION_INT("LIGHT_TYPE", LightType_MAX);
 
@@ -63,7 +63,7 @@ class FOcclusionRG : public FGlobalShader
 	END_SHADER_PARAMETER_STRUCT()
 };
 
-IMPLEMENT_GLOBAL_SHADER(FOcclusionRG, "/Engine/Private/RayTracing/RayTracingOcclusionRG.usf", "OcclusionRG", SF_RayGen);
+IMPLEMENT_GLOBAL_SHADER(FOcclusionRGS, "/Engine/Private/RayTracing/RayTracingOcclusionRGS.usf", "OcclusionRGS", SF_RayGen);
 
 void FDeferredShadingSceneRenderer::RenderRayTracingOcclusion(
 	FRHICommandListImmediate& RHICmdList,
@@ -105,7 +105,7 @@ void FDeferredShadingSceneRenderer::RenderRayTracingOcclusion(
 		FSceneTexturesUniformParameters SceneTextures;
 		SetupSceneTextureUniformParameters(SceneContext, FeatureLevel, ESceneTextureSetupMode::All, SceneTextures);
 
-		FOcclusionRG::FParameters* PassParameters = GraphBuilder.AllocParameters<FOcclusionRG::FParameters>();
+		FOcclusionRGS::FParameters* PassParameters = GraphBuilder.AllocParameters<FOcclusionRGS::FParameters>();
 		PassParameters->RWOcclusionMaskUAV = GraphBuilder.CreateUAV(FRDGTextureUAVDesc(ScreenShadowMaskTexture));
 		PassParameters->RWRayDistanceUAV = GraphBuilder.CreateUAV(FRDGTextureUAVDesc(RayDistanceTexture));
 		PassParameters->SamplesPerPixel = GRayTracingOcclusionSamplesPerPixel;
@@ -114,10 +114,10 @@ void FDeferredShadingSceneRenderer::RenderRayTracingOcclusion(
 		PassParameters->ViewUniformBuffer = View.ViewUniformBuffer;
 		PassParameters->SceneTexturesStruct = CreateUniformBufferImmediate(SceneTextures, EUniformBufferUsage::UniformBuffer_SingleDraw);
 
-		FOcclusionRG::FPermutationDomain PermutationVector;
-		PermutationVector.Set<FOcclusionRG::FLightTypeDim>(LightSceneProxy->GetLightType());
+		FOcclusionRGS::FPermutationDomain PermutationVector;
+		PermutationVector.Set<FOcclusionRGS::FLightTypeDim>(LightSceneProxy->GetLightType());
 
-		TShaderMapRef<FOcclusionRG> RayGenerationShader(GetGlobalShaderMap(FeatureLevel), PermutationVector);
+		TShaderMapRef<FOcclusionRGS> RayGenerationShader(GetGlobalShaderMap(FeatureLevel), PermutationVector);
 
 		FRayGenShaderUtils::AddRayTraceDispatchPass(
 			GraphBuilder,
diff --git a/Engine/Source/Runtime/Renderer/Private/RayTracing/RayTracingShadowsForRectLight.cpp b/Engine/Source/Runtime/Renderer/Private/RayTracing/RayTracingShadowsForRectLight.cpp
--- a/Engine/Source/Runtime/Renderer/Private/RayTracing/RayTracingShadowsForRectLight.cpp
+++ b/Engine/Source/Runtime/Renderer/Private/RayTracing/RayTracingShadowsForRectLight.cpp
@@ -82,9 +82,9 @@ DECLARE_GPU_STAT_NAMED(BuildRectLightMipTree, TEXT("Build RectLight Mip Tree"));
 IMPLEMENT_GLOBAL_SHADER_PARAMETER_STRUCT(FRectLightData, "RectLight");
 
 template <int CalcDirectLighting, int EncodeVisibility>
-class FRectLightOcclusionRG : public FGlobalShader
+class FRectLightOcclusionRGS : public FGlobalShader
 {
-	DECLARE_SHADER_TYPE(FRectLightOcclusionRG, Global)
+	DECLARE_SHADER_TYPE(FRectLightOcclusionRGS, Global)
 
 	static void ModifyCompilationEnvironment(const FGlobalShaderPermutationParameters& Parameters, FShaderCompilerEnvironment& OutEnvironment)
 	{
@@ -98,10 +98,10 @@ public:
 		return IsFeatureLevelSupported(Parameters.Platform, ERHIFeatureLevel::SM5) && IsRayTracingSupportedForThisProject();
 	}
 
-	FRectLightOcclusionRG() {}
-	virtual ~FRectLightOcclusionRG() {}
+	FRectLightOcclusionRGS() {}
+	virtual ~FRectLightOcclusionRGS() {}
 
-	FRectLightOcclusionRG(const ShaderMetaType::CompiledShaderInitializerType& Initializer)
+	FRectLightOcclusionRGS(const ShaderMetaType::CompiledShaderInitializerType& Initializer)
 		: FGlobalShader(Initializer)
 	{
 		ViewParameter.Bind(Initializer.ParameterMap, TEXT("View"));
@@ -165,8 +165,8 @@ private:
 };
 
 #define IMPLEMENT_RECT_LIGHT_OCCLUSION_TYPE(CalcDirectLightingType, EncodeVisbilityType) \
-	typedef FRectLightOcclusionRG<CalcDirectLightingType, EncodeVisbilityType> FRectLightOcclusionRG##CalcDirectLightingType##EncodeVisbilityType; \
-	IMPLEMENT_SHADER_TYPE(template<>, FRectLightOcclusionRG##CalcDirectLightingType##EncodeVisbilityType, TEXT("/Engine/Private/RayTracing/RayTracingRectLightOcclusionRG.usf"), TEXT("RectLightOcclusionRG"), SF_RayGen);
+	typedef FRectLightOcclusionRGS<CalcDirectLightingType, EncodeVisbilityType> FRectLightOcclusionRGS##CalcDirectLightingType##EncodeVisbilityType; \
+	IMPLEMENT_SHADER_TYPE(template<>, FRectLightOcclusionRGS##CalcDirectLightingType##EncodeVisbilityType, TEXT("/Engine/Private/RayTracing/RayTracingRectLightOcclusionRGS.usf"), TEXT("RectLightOcclusionRGS"), SF_RayGen);
 
 IMPLEMENT_RECT_LIGHT_OCCLUSION_TYPE(0, 0);
 IMPLEMENT_RECT_LIGHT_OCCLUSION_TYPE(0, 1);
@@ -523,7 +523,7 @@ void FDeferredShadingSceneRenderer::RenderRayTracingRectLightInternal(
 		FViewInfo& View = Views[ViewIndex];
 		FIntPoint ViewSize = View.ViewRect.Size();
 
-		TShaderMapRef<FRectLightOcclusionRG<CalcDirectLighting, EncodeVisibility>> RectLightOcclusionRayGenerationShader(GetGlobalShaderMap(FeatureLevel));
+		TShaderMapRef<FRectLightOcclusionRGS<CalcDirectLighting, EncodeVisibility>> RectLightOcclusionRayGenerationShader(GetGlobalShaderMap(FeatureLevel));
 
 		FSceneRenderTargets& SceneContext = FSceneRenderTargets::Get(RHICmdList);
 		FSceneTexturesUniformParameters SceneTextures;
diff --git a/Engine/Source/Runtime/Renderer/Private/RayTracing/RaytracingSkylight.cpp b/Engine/Source/Runtime/Renderer/Private/RayTracing/RaytracingSkylight.cpp
--- a/Engine/Source/Runtime/Renderer/Private/RayTracing/RaytracingSkylight.cpp
+++ b/Engine/Source/Runtime/Renderer/Private/RayTracing/RaytracingSkylight.cpp
@@ -80,7 +80,7 @@ class FSkylightRG : public FGlobalShader
 	END_SHADER_PARAMETER_STRUCT()
 };
 
-IMPLEMENT_GLOBAL_SHADER(FSkylightRG, "/Engine/Private/Raytracing/RaytracingSkylightRG.usf", "SkyLightRG", SF_RayGen);
+IMPLEMENT_GLOBAL_SHADER(FSkylightRG, "/Engine/Private/Raytracing/RaytracingSkylightRGS.usf", "SkyLightRGS", SF_RayGen);
 
 void FDeferredShadingSceneRenderer::RenderRayTracingSkyLight(
 	FRHICommandListImmediate& RHICmdList,
@@ -248,7 +248,7 @@ private:
 	FShaderResourceParameter RayDistanceUAVParameter;
 };
 
-IMPLEMENT_SHADER_TYPE(, FSkylightRG, TEXT("/Engine/Private/Raytracing/RaytracingSkylightRG.usf"), TEXT("SkyLightRG"), SF_RayGen);
+IMPLEMENT_SHADER_TYPE(, FSkylightRG, TEXT("/Engine/Private/Raytracing/RaytracingSkylightRGS.usf"), TEXT("SkyLightRGS"), SF_RayGen);
 
 class FBuildMipTreeCS : public FGlobalShader
 {
