Fix some more pedantic shader compilers

diff --git a/Engine/Shaders/Private/RayTracing/MipTreeCommon.ush b/Engine/Shaders/Private/RayTracing/MipTreeCommon.ush
--- a/Engine/Shaders/Private/RayTracing/MipTreeCommon.ush
+++ b/Engine/Shaders/Private/RayTracing/MipTreeCommon.ush
@@ -67,7 +67,7 @@ uint BufferOffsetAtPixel(uint2 Pixel, uint Level, uint2 Dimensions)
 uint Log2(uint Value)
 {
 	uint Result = 0;
-	while (Value)
+	while (Value != 0)
 	{
 		Result++;
 		Value = Value >> 1;
