Some cleaning of dead code in the path tracer. Forced to use lambert materials works properly now.

diff --git a/Engine/Shaders/Private/PathTracing/Material/PathTracingMaterialCommon.ush b/Engine/Shaders/Private/PathTracing/Material/PathTracingMaterialCommon.ush
--- a/Engine/Shaders/Private/PathTracing/Material/PathTracingMaterialCommon.ush
+++ b/Engine/Shaders/Private/PathTracing/Material/PathTracingMaterialCommon.ush
@@ -48,7 +48,9 @@ bool IsPureSpecularTransmission(FMaterialClosestHitPayload HitInfo)
 
 bool IsLambert(FMaterialClosestHitPayload HitInfo)
 {
-	return GetRoughness(HitInfo) > GLOSSY_MAX_ROUGHNESS;
+	//#dxr_todo: force lambert material until other BRDF types are finished
+	return true;
+	//return GetRoughness(HitInfo) > GLOSSY_MAX_ROUGHNESS;
 }
 
 bool BSDF_Refract(float3 V_Local, float3 N_Local, float IorV, float IorL, out float3 L_Local)
diff --git a/Engine/Shaders/Private/PathTracing/Material/PathTracingMaterialSampling.ush b/Engine/Shaders/Private/PathTracing/Material/PathTracingMaterialSampling.ush
--- a/Engine/Shaders/Private/PathTracing/Material/PathTracingMaterialSampling.ush
+++ b/Engine/Shaders/Private/PathTracing/Material/PathTracingMaterialSampling.ush
@@ -37,10 +37,6 @@ void SampleMaterial(
 	out float3 OutThroughput,
 	out float OutPdf)
 {
-	//#dxr_todo: force lambert material until other BRDF types are finished
-	Lambert_SampleMaterial(Payload, RandSample, OutDirection, OutThroughput, OutPdf);
-	return;
-
 	if (IsPureSpecularReflection(Payload))
 	{
 		SpecularReflection_SampleMaterial(RayDirection, Payload, OutDirection, OutThroughput, OutPdf);
diff --git a/Engine/Shaders/Private/PathTracing/Material/PathTracingSpecularReflection.ush b/Engine/Shaders/Private/PathTracing/Material/PathTracingSpecularReflection.ush
--- a/Engine/Shaders/Private/PathTracing/Material/PathTracingSpecularReflection.ush
+++ b/Engine/Shaders/Private/PathTracing/Material/PathTracingSpecularReflection.ush
@@ -5,19 +5,6 @@ PathTracingSpecularReflection.usf: Lambertian BRDF sampling functions
 ===============================================================================================*/
 #pragma once
 
-//// Fixed version of GetTangentBasis that follows UE axis conventions
-//float3x3 GetTangentBasis_NEW(float3 TangentZ)
-//{
-//	const float Sign = TangentZ.z >= 0 ? 1 : -1;
-//	const float a = -rcp(Sign + TangentZ.z);
-//	const float b = TangentZ.x * TangentZ.y * a;
-//
-//	float3 TangentY = { -1*(1 + Sign * a * Pow2(TangentZ.x)), -Sign * b, Sign * TangentZ.x };
-//	float3 TangentX = { b,  Sign + a * Pow2(TangentZ.y), -TangentZ.y };
-//
-//	return float3x3(TangentX, TangentY, TangentZ);
-//}
-
 void SpecularReflection_SampleMaterial(
 	float3 RayDirection,
 	FMaterialClosestHitPayload Payload,
