Add support for emissive & indirect lighting in ray traced reflections

diff --git a/Engine/Shaders/Private/RayTracing/RayTracingMaterialHitShaders.usf b/Engine/Shaders/Private/RayTracing/RayTracingMaterialHitShaders.usf
--- a/Engine/Shaders/Private/RayTracing/RayTracingMaterialHitShaders.usf
+++ b/Engine/Shaders/Private/RayTracing/RayTracingMaterialHitShaders.usf
@@ -1,12 +1,13 @@
-// Copyright 1998-2018 Epic Games, Inc. All Rights Reserved.
+// Copyright 1998-2019 Epic Games, Inc. All Rights Reserved.
 
-#pragma one
+#pragma once
 
-#define NEEDS_LIGHTMAP_COORDINATE		(HQ_TEXTURE_LIGHTMAP || LQ_TEXTURE_LIGHTMAP)
+#define ENABLE_SKY_LIGHT 1
+#define NEEDS_LIGHTMAP_COORDINATE (HQ_TEXTURE_LIGHTMAP || LQ_TEXTURE_LIGHTMAP)
 #ifdef USING_TESSELLATION
-#undef USING_TESSELLATION
-// Needed for VertexFactoryInterpolate
-#define USING_TESSELLATION 1
+	#undef USING_TESSELLATION
+	// Needed for VertexFactoryInterpolate
+	#define USING_TESSELLATION 1
 #endif
 
 // #dxr_todo: Temporarily hacking clip() away to prevent driver crash. Should be integrated into MaterialTemplate.ush in the future.
@@ -22,14 +23,9 @@
 #include "/Engine/Generated/VertexFactory.ush"
 
 #include "/Engine/Private/SHCommon.ush"
+#include "/Engine/Private/ReflectionEnvironmentShared.ush"
 #include "/Engine/Private/LightmapCommon.ush"
 
-#if SIMPLE_FORWARD_SHADING
-	#define GetEffectiveSkySHDiffuse GetSkySHDiffuseSimple
-#else
-	#define GetEffectiveSkySHDiffuse GetSkySHDiffuse
-#endif
-
 /** Computes sky diffuse lighting, including precomputed shadowing. */
 void GetSkyLighting(FMaterialPixelParameters MaterialParameters, float3 WorldNormal, float2 LightmapUV, float3 SkyOcclusionUV3D, out float3 OutDiffuseLighting, out float3 OutSubsurfaceLighting)
 {
@@ -90,13 +86,13 @@ void GetSkyLighting(FMaterialPixelParameters MaterialParameters, float3 WorldNor
 	#endif
 			
 	// Compute the preconvolved incoming lighting with the bent normal direction
-	float3 DiffuseLookup = GetEffectiveSkySHDiffuse(SkyLightingNormal) * ResolvedView.SkyLightColor.rgb;
+	float3 DiffuseLookup = GetSkySHDiffuse(SkyLightingNormal) * ResolvedView.SkyLightColor.rgb;
 
 	// Apply AO to the sky diffuse
 	OutDiffuseLighting += DiffuseLookup * (SkyVisibility * GeometryTerm);
 
 	#if MATERIAL_SHADINGMODEL_TWOSIDED_FOLIAGE
-		float3 BackfaceDiffuseLookup = GetEffectiveSkySHDiffuse(-WorldNormal) * ResolvedView.SkyLightColor.rgb;
+		float3 BackfaceDiffuseLookup = GetSkySHDiffuse(-WorldNormal) * ResolvedView.SkyLightColor.rgb;
 		OutSubsurfaceLighting += BackfaceDiffuseLookup * SkyVisibility;
 	#endif
 
@@ -305,6 +301,8 @@ void MaterialCHS(inout FPackedMaterialClosestHitPayload PackedPayload, in FDefau
 	
 	Payload.DiffuseColor = BaseColor - BaseColor * Metallic;
 	Payload.SpecularColor = lerp(0.08 * Specular.xxx, BaseColor, Metallic);
+	
+	Payload.Radiance += Payload.DiffuseColor * DiffuseIndirectLighting; // Contains both lightmap & skylight
 
 	PackedPayload = PackRayTracingPayload(Payload);
 }
diff --git a/Engine/Shaders/Private/RayTracing/RayTracingReflections.usf b/Engine/Shaders/Private/RayTracing/RayTracingReflections.usf
--- a/Engine/Shaders/Private/RayTracing/RayTracingReflections.usf
+++ b/Engine/Shaders/Private/RayTracing/RayTracingReflections.usf
@@ -30,6 +30,7 @@ float2 SobolIndexToUniformUnitSquare(uint2 SobolRand)
 
 uint SamplesPerPixel;
 uint ShouldDoDirectLighting;
+uint ShouldDoEmissiveAndIndirectLighting;
 float ReflectionMaxRayDistance;
 
 RaytracingAccelerationStructure TLAS;
@@ -227,6 +228,12 @@ FMaterialClosestHitPayload TraceRayAndAccumulateResults(
 		DirectLighting = -min(-DirectLighting, 0.0);
 		Radiance += DirectLighting;
 	}
+    
+	if (ShouldDoEmissiveAndIndirectLighting)
+	{
+	    // Emissive & indirect contribution
+	    Radiance += Payload.Radiance;
+    }
 
 	return Payload;
 }
diff --git a/Engine/Source/Runtime/Renderer/Private/RayTracing/RayTracingReflections.cpp b/Engine/Source/Runtime/Renderer/Private/RayTracing/RayTracingReflections.cpp
--- a/Engine/Source/Runtime/Renderer/Private/RayTracing/RayTracingReflections.cpp
+++ b/Engine/Source/Runtime/Renderer/Private/RayTracing/RayTracingReflections.cpp
@@ -23,6 +23,13 @@ static FAutoConsoleVariableRef CVarRayTracingReflectionsSamplesPerPixel(
 	TEXT("Sets the samples-per-pixel for reflections (default = 1)")
 );
 
+static int32 GRayTracingReflectionsEmissiveAndIndirectLighting = 1;
+static FAutoConsoleVariableRef CVarRayTracingReflectionsEmissiveAndIndirectLighting(
+	TEXT("r.RayTracing.Reflections.EmissiveAndIndirectLighting"),
+	GRayTracingReflectionsEmissiveAndIndirectLighting,
+	TEXT("Enables ray tracing reflections emissive and indirect lighting (default = 1)")
+);
+
 static int32 GRayTracingReflectionsDirectLighting = 1;
 static FAutoConsoleVariableRef CVarRayTracingReflectionsDirectLighting(
 	TEXT("r.RayTracing.Reflections.DirectLighting"),
@@ -113,6 +120,7 @@ class FRayTracingReflectionsRG : public FGlobalShader
 	BEGIN_SHADER_PARAMETER_STRUCT(FParameters, )
 		SHADER_PARAMETER(int32, SamplesPerPixel)
 		SHADER_PARAMETER(int32, ShouldDoDirectLighting)
+		SHADER_PARAMETER(int32, ShouldDoEmissiveAndIndirectLighting)
 		SHADER_PARAMETER(float, ReflectionMaxRayDistance)
 
 		SHADER_PARAMETER_SRV(RaytracingAccelerationStructure, TLAS)
@@ -188,6 +196,7 @@ void FDeferredShadingSceneRenderer::RayTraceReflections(
 
 	PassParameters->SamplesPerPixel = GRayTracingReflectionsSamplesPerPixel;
 	PassParameters->ShouldDoDirectLighting = GRayTracingReflectionsDirectLighting;
+	PassParameters->ShouldDoEmissiveAndIndirectLighting = GRayTracingReflectionsEmissiveAndIndirectLighting;
 	PassParameters->ReflectionMaxRayDistance = GRayTracingReflectionsMaxRayDistance;
 	PassParameters->LTCMatTexture = GSystemTextures.LTCMat->GetRenderTargetItem().ShaderResourceTexture;
 	PassParameters->LTCMatSampler = TStaticSamplerState<SF_Bilinear, AM_Clamp, AM_Clamp, AM_Clamp>::GetRHI();
