Fixed SamplesPerPixel was not working properly in ray traced reflections

diff --git a/Engine/Shaders/Private/RayTracing/RayTracingReflections.usf b/Engine/Shaders/Private/RayTracing/RayTracingReflections.usf
--- a/Engine/Shaders/Private/RayTracing/RayTracingReflections.usf
+++ b/Engine/Shaders/Private/RayTracing/RayTracingReflections.usf
@@ -285,7 +285,6 @@ void RayTracingReflectionsRGS()
 	}
 
 	float HitDistance = 0.0;
-	float3 Radiance = float3(0, 0, 0);
 
 	if (Roughness < RoughnessThreshold)
 	{
@@ -324,6 +323,7 @@ void RayTracingReflectionsRGS()
 
 				uint RayFlags = RAY_FLAG_CULL_BACK_FACING_TRIANGLES;
 
+				float3 Radiance = float3(0, 0, 0);
 				FMaterialClosestHitPayload Payload = TraceRayAndAccumulateResults(Ray, Radiance);
 
 				if (Payload.HitT < 0)
@@ -331,7 +331,6 @@ void RayTracingReflectionsRGS()
 					break;
 				}
 
-#if 1
 				float4 PathRadiance = float4(Radiance, 1.0); 
 				if (BounceIndex == 0)
 				{
@@ -343,9 +342,6 @@ void RayTracingReflectionsRGS()
 					//float4 FinalReflectionColor = BlendWithEnvironment(BRDFThroughput, Payload, Ray.Origin, PathRadiance, Ray.Direction, Payload.WorldNormal);
 					ReflectedColor += FinalReflectionColor;
 				}
-#else // Debugging, remove later
-				ReflectedColor += float4(Payload.BaseColor, 1.0); 
-#endif
 
 				IncidentDirection = Ray.Direction;
 
@@ -356,7 +352,7 @@ void RayTracingReflectionsRGS()
 
 	if (SamplesPerPixel > 0)
 	{
-		ReflectedColor /= SamplesPerPixel;
+		ReflectedColor /= LocalSamplesPerPixel;
 	}
 
 	ColorOutput[PixelCoord] = float4(ReflectedColor);
