Added support for reflected shadows. TODO: add optimizations and support for stochastic area sampling

diff --git a/Engine/Shaders/Private/RayTracing/RayTracingReflections.usf b/Engine/Shaders/Private/RayTracing/RayTracingReflections.usf
--- a/Engine/Shaders/Private/RayTracing/RayTracingReflections.usf
+++ b/Engine/Shaders/Private/RayTracing/RayTracingReflections.usf
@@ -30,6 +30,7 @@ float2 SobolIndexToUniformUnitSquare(uint2 SobolRand)
 
 uint SamplesPerPixel;
 uint ShouldDoDirectLighting;
+uint ShouldDoReflectedShadows;
 uint ShouldDoEmissiveAndIndirectLighting;
 float ReflectionMaxRayDistance;
 
@@ -120,7 +121,33 @@ float4 BlendWithEnvironment(inout float3 BRDFThroughput, const in FMaterialClose
 }
 #endif
 
-float3 ComputeDirectLighting(RayDesc Ray, FMaterialClosestHitPayload Payload)
+float TraceShadowRay(
+	float3 Origin,
+	float3 Direction,
+	float TMax)
+{
+	RayDesc Ray;
+	Ray.Origin = Origin;
+	Ray.Direction = Direction;
+	Ray.TMin = 1e-4f;
+	Ray.TMax = TMax;
+
+	uint RayFlags = RAY_FLAG_CULL_BACK_FACING_TRIANGLES | RAY_FLAG_ACCEPT_FIRST_HIT_AND_END_SEARCH;
+
+	FMaterialClosestHitPayload Payload = TraceRayInternal(
+		TLAS,      // AccelerationStructure
+		RayFlags,
+		0xFF,      // InstanceInclusionMask
+		0,         // RayContributionToHitGroupIndex
+		1,         // MultiplierForGeometryContributionToShaderIndex
+		0,         // MissShaderIndex
+		Ray        // RayDesc
+	);
+
+	return Payload.HitT;
+}
+
+float3 ComputeDirectLighting(RayDesc Ray, FMaterialClosestHitPayload Payload, bool bReflectedShadows)
 {
 	float3 DirectLighting = float3(0.0, 0.0, 0.0);
 
@@ -193,6 +220,27 @@ float3 ComputeDirectLighting(RayDesc Ray, FMaterialClosestHitPayload Payload)
 		float4 LightAttenuation = 1.0f;
 		float LightProfileMultiplier = 1.0;//#dxr_todo
 
+		//#dxr_todo: check perf improvements forcing TMin == TMax when bReflectedShadows = false, to avoid tracing rays inside if statements
+		if (bReflectedShadows)
+		{
+			//#dxr_todo: add early outs if backfacing or if light energy is not sufficient, we are inside a light volume, etc
+			//#dxr_todo: add support for sampling area lights (and perhaps add CVar to enable stocastic shadows?)
+
+			float3 ShadowRayDirection;
+			
+			if (LightType == LIGHT_TYPE_DIRECTIONAL)
+			{
+				ShadowRayDirection = LightData.Direction;
+			}
+			else
+			{
+				ShadowRayDirection = normalize(LightData.Position - WorldPosition);
+			}
+
+			float HitT = TraceShadowRay(WorldPosition + Payload.WorldNormal*0.001, ShadowRayDirection, ReflectionMaxRayDistance);
+			AmbientOcclusion = HitT < 0.0f;
+		}
+
 		float3 LightContribution = GetDynamicLighting(WorldPosition, -CameraVector, GBufferData, AmbientOcclusion, GBufferData.ShadingModelID, LightData, LightAttenuation, 0.5, uint2(0, 0), LightSourceTexture).xyz;
 		DirectLighting += LightContribution * LightProfileMultiplier;
 	}
@@ -222,7 +270,7 @@ FMaterialClosestHitPayload TraceRayAndAccumulateResults(
 
 	if (ShouldDoDirectLighting)
 	{
-		float3 DirectLighting = ComputeDirectLighting(Ray, Payload);
+		float3 DirectLighting = ComputeDirectLighting(Ray, Payload, ShouldDoReflectedShadows > 0);
 
 		// Transform NaNs to black, transform negative colors to black.
 		DirectLighting = -min(-DirectLighting, 0.0);
diff --git a/Engine/Source/Runtime/Renderer/Private/RayTracing/RayTracingReflections.cpp b/Engine/Source/Runtime/Renderer/Private/RayTracing/RayTracingReflections.cpp
--- a/Engine/Source/Runtime/Renderer/Private/RayTracing/RayTracingReflections.cpp
+++ b/Engine/Source/Runtime/Renderer/Private/RayTracing/RayTracingReflections.cpp
@@ -37,6 +37,13 @@ static FAutoConsoleVariableRef CVarRayTracingReflectionsDirectLighting(
 	TEXT("Enables ray tracing reflections direct lighting (default = 1)")
 );
 
+static int32 GRayTracingReflectionsShadows = 1;
+static FAutoConsoleVariableRef CVarRayTracingReflectionsShadows(
+	TEXT("r.RayTracing.Reflections.Shadows"),
+	GRayTracingReflectionsShadows,
+	TEXT("Enables shadows in ray tracing reflections (default = 1)")
+);
+
 static float GRayTracingReflectionsMaxRayDistance = 1.0e27;
 static FAutoConsoleVariableRef CVarRayTracingReflectionsMaxRayDistance(
 	TEXT("r.RayTracing.Reflections.MaxRayDistance"),
@@ -120,6 +127,7 @@ class FRayTracingReflectionsRG : public FGlobalShader
 	BEGIN_SHADER_PARAMETER_STRUCT(FParameters, )
 		SHADER_PARAMETER(int32, SamplesPerPixel)
 		SHADER_PARAMETER(int32, ShouldDoDirectLighting)
+		SHADER_PARAMETER(int32, ShouldDoReflectedShadows)
 		SHADER_PARAMETER(int32, ShouldDoEmissiveAndIndirectLighting)
 		SHADER_PARAMETER(float, ReflectionMaxRayDistance)
 
@@ -196,6 +204,7 @@ void FDeferredShadingSceneRenderer::RayTraceReflections(
 
 	PassParameters->SamplesPerPixel = GRayTracingReflectionsSamplesPerPixel;
 	PassParameters->ShouldDoDirectLighting = GRayTracingReflectionsDirectLighting;
+	PassParameters->ShouldDoReflectedShadows = GRayTracingReflectionsShadows;
 	PassParameters->ShouldDoEmissiveAndIndirectLighting = GRayTracingReflectionsEmissiveAndIndirectLighting;
 	PassParameters->ReflectionMaxRayDistance = GRayTracingReflectionsMaxRayDistance;
 	PassParameters->LTCMatTexture = GSystemTextures.LTCMat->GetRenderTargetItem().ShaderResourceTexture;
