DR - Fix for Niagara shader, as NiagaraSpriteVFLooseParameters.CutoutGeometry is not defined

diff --git a/Engine/Shaders/Private/NiagaraSpriteVertexFactory.ush b/Engine/Shaders/Private/NiagaraSpriteVertexFactory.ush
--- a/Engine/Shaders/Private/NiagaraSpriteVertexFactory.ush
+++ b/Engine/Shaders/Private/NiagaraSpriteVertexFactory.ush
@@ -535,6 +535,7 @@ void ComputeBillboardUVs(FVertexFactoryInput Input, FVertexFactoryIntermediates
 	// Encoding the UV flip in the sign of the size data.
 	float2 UVFlip = sign(Intermediates.SizeRotSubImage.xy);
 
+#if FEATURE_LEVEL >= FEATURE_LEVEL_SM4
 	BRANCH
 	if (NiagaraSpriteVFLooseParameters.NumCutoutVerticesPerFrame > 0)
 	{
@@ -559,6 +560,7 @@ void ComputeBillboardUVs(FVertexFactoryInput Input, FVertexFactoryIntermediates
 		UVForPosition.y = UVFlip.y < 0.0 ? 1.0 - UVForTexturing.y : UVForTexturing.y;
 	}
 	else
+#endif
 	{
 		UVForTexturing.x = UVFlip.x < 0.0 ? 1.0 - Input.TexCoord.x : Input.TexCoord.x;
 		UVForTexturing.y = UVFlip.y < 0.0 ? 1.0 - Input.TexCoord.y : Input.TexCoord.y;
@@ -733,11 +735,12 @@ FVertexFactoryIntermediates GetVertexFactoryIntermediates(FVertexFactoryInput In
 	FVertexFactoryIntermediates Intermediates = (FVertexFactoryIntermediates)0;
 
 	uint ParticleID = Input.InstanceId;
+#if FEATURE_LEVEL >= FEATURE_LEVEL_SM4
 	if(NiagaraSpriteVFLooseParameters.SortedIndicesOffset != 0xFFFFFFFF)
 	{
 		ParticleID = NiagaraSpriteVFLooseParameters.SortedIndices[NiagaraSpriteVFLooseParameters.SortedIndicesOffset + Input.InstanceId];
 	}
-
+#endif
 	Intermediates.Position = GetNiagaraParticlePosition(ParticleID);
 	Intermediates.OldPosition = Intermediates.Position.xyz - GetNiagaraParticleVelocity(ParticleID) * NiagaraSpriteVF.DeltaSeconds;
 	float2 ParticleSize = GetNiagaraParticleSize(ParticleID);
diff --git a/Engine/Shaders/Public/ShaderVersion.ush b/Engine/Shaders/Public/ShaderVersion.ush
--- a/Engine/Shaders/Public/ShaderVersion.ush
+++ b/Engine/Shaders/Public/ShaderVersion.ush
@@ -3,4 +3,4 @@
 // in Platform.ush (which should be included in any shader) it allows to invalidate the shader DDC.
 // 
 // If you are merging streams and there is a conflict with this GUID you should make a new GUID rather than taking one or the other.
-// GUID = 487F1C22-EBE7-4CEA-A565-E19FE49FE8EE
\ No newline at end of file
+// GUID = 487F1C22-EBE7-4CEA-A565-E19FE49FE8EF
