Local copyright updates

diff --git a/Engine/Build/Android/Java/src/com/epicgames/ue4/GameActivity.java.template b/Engine/Build/Android/Java/src/com/epicgames/ue4/GameActivity.java.template
--- a/Engine/Build/Android/Java/src/com/epicgames/ue4/GameActivity.java.template
+++ b/Engine/Build/Android/Java/src/com/epicgames/ue4/GameActivity.java.template
@@ -1,4 +1,4 @@
-// Copyright 1998-2018 Epic Games, Inc. All Rights Reserved.
+// Copyright 1998-2019 Epic Games, Inc. All Rights Reserved.
 //This file needs to be here so the "ant" build step doesnt fail when looking for a /src folder.
 
 package com.epicgames.ue4;
diff --git a/Engine/Build/Target.cs.template b/Engine/Build/Target.cs.template
--- a/Engine/Build/Target.cs.template
+++ b/Engine/Build/Target.cs.template
@@ -1,4 +1,4 @@
-// Copyright 1998-2018 Epic Games, Inc. All Rights Reserved.
+// Copyright 1998-2019 Epic Games, Inc. All Rights Reserved.
 
 using UnrealBuildTool;
 using System.Collections.Generic;
diff --git a/Engine/Shaders/Private/DataStructures/SceneViewFamilyBlackboard.ush b/Engine/Shaders/Private/DataStructures/SceneViewFamilyBlackboard.ush
--- a/Engine/Shaders/Private/DataStructures/SceneViewFamilyBlackboard.ush
+++ b/Engine/Shaders/Private/DataStructures/SceneViewFamilyBlackboard.ush
@@ -1,4 +1,4 @@
-// Copyright 1998-2018 Epic Games, Inc. All Rights Reserved.
+// Copyright 1998-2019 Epic Games, Inc. All Rights Reserved.
 
 /*=============================================================================
 	SceneViewFamilyBlackboard.ush: 
diff --git a/Engine/Shaders/Private/LightShaderParameters.ush b/Engine/Shaders/Private/LightShaderParameters.ush
--- a/Engine/Shaders/Private/LightShaderParameters.ush
+++ b/Engine/Shaders/Private/LightShaderParameters.ush
@@ -1,4 +1,4 @@
-// Copyright 1998-2018 Epic Games, Inc. All Rights Reserved.
+// Copyright 1998-2019 Epic Games, Inc. All Rights Reserved.
 
 #pragma once
 
diff --git a/Engine/Shaders/Private/LightmapData.ush b/Engine/Shaders/Private/LightmapData.ush
--- a/Engine/Shaders/Private/LightmapData.ush
+++ b/Engine/Shaders/Private/LightmapData.ush
@@ -1,4 +1,4 @@
-// Copyright 1998-2018 Epic Games, Inc. All Rights Reserved.
+// Copyright 1998-2019 Epic Games, Inc. All Rights Reserved.
 
 /**
  * LightmapData.ush
diff --git a/Engine/Shaders/Private/PathTracing/BuildSkyLightColumnCdfComputeShader.usf b/Engine/Shaders/Private/PathTracing/BuildSkyLightColumnCdfComputeShader.usf
--- a/Engine/Shaders/Private/PathTracing/BuildSkyLightColumnCdfComputeShader.usf
+++ b/Engine/Shaders/Private/PathTracing/BuildSkyLightColumnCdfComputeShader.usf
@@ -1,4 +1,4 @@
-// Copyright 1998-2018 Epic Games, Inc. All Rights Reserved.
+// Copyright 1998-2019 Epic Games, Inc. All Rights Reserved.
 
 /*=============================================================================================
 	BuildSkyLightColumnCdfComputeShader.usf: Intermediate SkyLight CDF build.
diff --git a/Engine/Shaders/Private/PathTracing/BuildSkyLightCubeFaceCdfComputeShader.usf b/Engine/Shaders/Private/PathTracing/BuildSkyLightCubeFaceCdfComputeShader.usf
--- a/Engine/Shaders/Private/PathTracing/BuildSkyLightCubeFaceCdfComputeShader.usf
+++ b/Engine/Shaders/Private/PathTracing/BuildSkyLightCubeFaceCdfComputeShader.usf
@@ -1,4 +1,4 @@
-// Copyright 1998-2018 Epic Games, Inc. All Rights Reserved.
+// Copyright 1998-2019 Epic Games, Inc. All Rights Reserved.
 
 /*=============================================================================================
 BuildSkyLightCubeFaceCdfComputeShader.usf: Intermediate SkyLight CDF build.
diff --git a/Engine/Shaders/Private/PathTracing/BuildSkyLightRowCdfComputeShader.usf b/Engine/Shaders/Private/PathTracing/BuildSkyLightRowCdfComputeShader.usf
--- a/Engine/Shaders/Private/PathTracing/BuildSkyLightRowCdfComputeShader.usf
+++ b/Engine/Shaders/Private/PathTracing/BuildSkyLightRowCdfComputeShader.usf
@@ -1,4 +1,4 @@
-// Copyright 1998-2018 Epic Games, Inc. All Rights Reserved.
+// Copyright 1998-2019 Epic Games, Inc. All Rights Reserved.
 
 /*=============================================================================================
 BuildSkyLightRowCdfComputeShader.usf: Intermediate SkyLight CDF build.
diff --git a/Engine/Shaders/Private/PathTracing/Light/PathTracingDirectionalLight.ush b/Engine/Shaders/Private/PathTracing/Light/PathTracingDirectionalLight.ush
--- a/Engine/Shaders/Private/PathTracing/Light/PathTracingDirectionalLight.ush
+++ b/Engine/Shaders/Private/PathTracing/Light/PathTracingDirectionalLight.ush
@@ -1,4 +1,4 @@
-// Copyright 1998-2018 Epic Games, Inc. All Rights Reserved.
+// Copyright 1998-2019 Epic Games, Inc. All Rights Reserved.
 
 /*=============================================================================================
 DirectionalLight.usf: Light sampling functions for Directional light implementation
diff --git a/Engine/Shaders/Private/PathTracing/Light/PathTracingLightCommon.ush b/Engine/Shaders/Private/PathTracing/Light/PathTracingLightCommon.ush
--- a/Engine/Shaders/Private/PathTracing/Light/PathTracingLightCommon.ush
+++ b/Engine/Shaders/Private/PathTracing/Light/PathTracingLightCommon.ush
@@ -1,4 +1,4 @@
-// Copyright 1998-2018 Epic Games, Inc. All Rights Reserved.
+// Copyright 1998-2019 Epic Games, Inc. All Rights Reserved.
 
 /*=============================================================================================
 LightCommon.usf: Common utility functions for light sampling
diff --git a/Engine/Shaders/Private/PathTracing/Light/PathTracingLightSampling.ush b/Engine/Shaders/Private/PathTracing/Light/PathTracingLightSampling.ush
--- a/Engine/Shaders/Private/PathTracing/Light/PathTracingLightSampling.ush
+++ b/Engine/Shaders/Private/PathTracing/Light/PathTracingLightSampling.ush
@@ -1,4 +1,4 @@
-// Copyright 1998-2018 Epic Games, Inc. All Rights Reserved.
+// Copyright 1998-2019 Epic Games, Inc. All Rights Reserved.
 
 /*=============================================================================================
 LightSampling.usf: Light sampling for next-event estimation 
diff --git a/Engine/Shaders/Private/PathTracing/Light/PathTracingPointLight.ush b/Engine/Shaders/Private/PathTracing/Light/PathTracingPointLight.ush
--- a/Engine/Shaders/Private/PathTracing/Light/PathTracingPointLight.ush
+++ b/Engine/Shaders/Private/PathTracing/Light/PathTracingPointLight.ush
@@ -1,4 +1,4 @@
-// Copyright 1998-2018 Epic Games, Inc. All Rights Reserved.
+// Copyright 1998-2019 Epic Games, Inc. All Rights Reserved.
 
 /*=============================================================================================
 PointLight.usf: Light sampling functions for Point light implementation
diff --git a/Engine/Shaders/Private/PathTracing/Light/PathTracingRectLight.ush b/Engine/Shaders/Private/PathTracing/Light/PathTracingRectLight.ush
--- a/Engine/Shaders/Private/PathTracing/Light/PathTracingRectLight.ush
+++ b/Engine/Shaders/Private/PathTracing/Light/PathTracingRectLight.ush
@@ -1,4 +1,4 @@
-// Copyright 1998-2018 Epic Games, Inc. All Rights Reserved.
+// Copyright 1998-2019 Epic Games, Inc. All Rights Reserved.
 
 /*=============================================================================================
 	RectLight.usf: Light sampling functions for Rect light implementation
diff --git a/Engine/Shaders/Private/PathTracing/Light/PathTracingSkyLight.ush b/Engine/Shaders/Private/PathTracing/Light/PathTracingSkyLight.ush
--- a/Engine/Shaders/Private/PathTracing/Light/PathTracingSkyLight.ush
+++ b/Engine/Shaders/Private/PathTracing/Light/PathTracingSkyLight.ush
@@ -1,4 +1,4 @@
-// Copyright 1998-2018 Epic Games, Inc. All Rights Reserved.
+// Copyright 1998-2019 Epic Games, Inc. All Rights Reserved.
 
 /*=============================================================================================
 SkyLight.usf: Light sampling functions for SkyLight implementation
diff --git a/Engine/Shaders/Private/PathTracing/Material/PathTracingFresnel.ush b/Engine/Shaders/Private/PathTracing/Material/PathTracingFresnel.ush
--- a/Engine/Shaders/Private/PathTracing/Material/PathTracingFresnel.ush
+++ b/Engine/Shaders/Private/PathTracing/Material/PathTracingFresnel.ush
@@ -1,4 +1,4 @@
-// Copyright 1998-2018 Epic Games, Inc. All Rights Reserved.
+// Copyright 1998-2019 Epic Games, Inc. All Rights Reserved.
 
 /*=============================================================================================
 PathTracingFresnel.usf: Fresnel utilities for material sampling functions 
diff --git a/Engine/Shaders/Private/PathTracing/Material/PathTracingGlossy.ush b/Engine/Shaders/Private/PathTracing/Material/PathTracingGlossy.ush
--- a/Engine/Shaders/Private/PathTracing/Material/PathTracingGlossy.ush
+++ b/Engine/Shaders/Private/PathTracing/Material/PathTracingGlossy.ush
@@ -1,4 +1,4 @@
-// Copyright 1998-2018 Epic Games, Inc. All Rights Reserved.
+// Copyright 1998-2019 Epic Games, Inc. All Rights Reserved.
 
 /*=============================================================================================
 Glossy.usf: Microfacet GGX BRDF sampling functions 
diff --git a/Engine/Shaders/Private/PathTracing/Material/PathTracingLambert.ush b/Engine/Shaders/Private/PathTracing/Material/PathTracingLambert.ush
--- a/Engine/Shaders/Private/PathTracing/Material/PathTracingLambert.ush
+++ b/Engine/Shaders/Private/PathTracing/Material/PathTracingLambert.ush
@@ -1,4 +1,4 @@
-// Copyright 1998-2018 Epic Games, Inc. All Rights Reserved.
+// Copyright 1998-2019 Epic Games, Inc. All Rights Reserved.
 
 /*=============================================================================================
 PathTracingLambert.usf: Lambertian BRDF sampling functions 
diff --git a/Engine/Shaders/Private/PathTracing/Material/PathTracingMaterialCommon.ush b/Engine/Shaders/Private/PathTracing/Material/PathTracingMaterialCommon.ush
--- a/Engine/Shaders/Private/PathTracing/Material/PathTracingMaterialCommon.ush
+++ b/Engine/Shaders/Private/PathTracing/Material/PathTracingMaterialCommon.ush
@@ -1,4 +1,4 @@
-// Copyright 1998-2018 Epic Games, Inc. All Rights Reserved.
+// Copyright 1998-2019 Epic Games, Inc. All Rights Reserved.
 
 /*=============================================================================================
 PathTracingMaterialCommon.usf: Brdf utility functions
diff --git a/Engine/Shaders/Private/PathTracing/Material/PathTracingMaterialSampling.ush b/Engine/Shaders/Private/PathTracing/Material/PathTracingMaterialSampling.ush
--- a/Engine/Shaders/Private/PathTracing/Material/PathTracingMaterialSampling.ush
+++ b/Engine/Shaders/Private/PathTracing/Material/PathTracingMaterialSampling.ush
@@ -1,4 +1,4 @@
-// Copyright 1998-2018 Epic Games, Inc. All Rights Reserved.
+// Copyright 1998-2019 Epic Games, Inc. All Rights Reserved.
 
 /*=============================================================================================
 PathTracingMaterialSampling.usf: Material sampling functions
diff --git a/Engine/Shaders/Private/PathTracing/Material/PathTracingSpecularReflection.ush b/Engine/Shaders/Private/PathTracing/Material/PathTracingSpecularReflection.ush
--- a/Engine/Shaders/Private/PathTracing/Material/PathTracingSpecularReflection.ush
+++ b/Engine/Shaders/Private/PathTracing/Material/PathTracingSpecularReflection.ush
@@ -1,4 +1,4 @@
-// Copyright 1998-2018 Epic Games, Inc. All Rights Reserved.
+// Copyright 1998-2019 Epic Games, Inc. All Rights Reserved.
 
 /*=============================================================================================
 PathTracingSpecularReflection.usf: Lambertian BRDF sampling functions
diff --git a/Engine/Shaders/Private/PathTracing/Material/PathTracingSpecularTransmission.ush b/Engine/Shaders/Private/PathTracing/Material/PathTracingSpecularTransmission.ush
--- a/Engine/Shaders/Private/PathTracing/Material/PathTracingSpecularTransmission.ush
+++ b/Engine/Shaders/Private/PathTracing/Material/PathTracingSpecularTransmission.ush
@@ -1,4 +1,4 @@
-// Copyright 1998-2018 Epic Games, Inc. All Rights Reserved.
+// Copyright 1998-2019 Epic Games, Inc. All Rights Reserved.
 
 /*=============================================================================================
 PathTracingSpecularTransmission.usf: Pure specular transmission BTDF sampling functions
diff --git a/Engine/Shaders/Private/PathTracing/Material/PathTracingTestBrdfs.usf b/Engine/Shaders/Private/PathTracing/Material/PathTracingTestBrdfs.usf
--- a/Engine/Shaders/Private/PathTracing/Material/PathTracingTestBrdfs.usf
+++ b/Engine/Shaders/Private/PathTracing/Material/PathTracingTestBrdfs.usf
@@ -1,4 +1,4 @@
-// Copyright 1998-2018 Epic Games, Inc. All Rights Reserved.
+// Copyright 1998-2019 Epic Games, Inc. All Rights Reserved.
 
 /*=============================================================================================
 	PathTracingTestBrdfs.usf: BRDF tests 
diff --git a/Engine/Shaders/Private/PathTracing/Material/PathTracingTestGenerateWiSamples.usf b/Engine/Shaders/Private/PathTracing/Material/PathTracingTestGenerateWiSamples.usf
--- a/Engine/Shaders/Private/PathTracing/Material/PathTracingTestGenerateWiSamples.usf
+++ b/Engine/Shaders/Private/PathTracing/Material/PathTracingTestGenerateWiSamples.usf
@@ -1,4 +1,4 @@
-// Copyright 1998-2018 Epic Games, Inc. All Rights Reserved.
+// Copyright 1998-2019 Epic Games, Inc. All Rights Reserved.
 
 /*=============================================================================================
 	PathTracingTestGenerateWiSamples.usf: BRDF tests 
diff --git a/Engine/Shaders/Private/PathTracing/Material/PathTracingTestIntegrateHemispherePatch.usf b/Engine/Shaders/Private/PathTracing/Material/PathTracingTestIntegrateHemispherePatch.usf
--- a/Engine/Shaders/Private/PathTracing/Material/PathTracingTestIntegrateHemispherePatch.usf
+++ b/Engine/Shaders/Private/PathTracing/Material/PathTracingTestIntegrateHemispherePatch.usf
@@ -1,4 +1,4 @@
-// Copyright 1998-2018 Epic Games, Inc. All Rights Reserved.
+// Copyright 1998-2019 Epic Games, Inc. All Rights Reserved.
 
 /*=============================================================================================
 	PathTracingTestIntegrateHemispherePatch.usf: BRDF tests 
diff --git a/Engine/Shaders/Private/PathTracing/Material/PathTracingTestPdfIntegration.usf b/Engine/Shaders/Private/PathTracing/Material/PathTracingTestPdfIntegration.usf
--- a/Engine/Shaders/Private/PathTracing/Material/PathTracingTestPdfIntegration.usf
+++ b/Engine/Shaders/Private/PathTracing/Material/PathTracingTestPdfIntegration.usf
@@ -1,4 +1,4 @@
-// Copyright 1998-2018 Epic Games, Inc. All Rights Reserved.
+// Copyright 1998-2019 Epic Games, Inc. All Rights Reserved.
 
 /*=============================================================================================
 	PathTracingTestPdfIntegration.usf: BRDF tests 
diff --git a/Engine/Shaders/Private/PathTracing/PathCompaction.usf b/Engine/Shaders/Private/PathTracing/PathCompaction.usf
--- a/Engine/Shaders/Private/PathTracing/PathCompaction.usf
+++ b/Engine/Shaders/Private/PathTracing/PathCompaction.usf
@@ -1,4 +1,4 @@
-// Copyright 1998-2018 Epic Games, Inc. All Rights Reserved.
+// Copyright 1998-2019 Epic Games, Inc. All Rights Reserved.
 
 /*=============================================================================================
 	PathCompaction.usf: Reference path tracing
diff --git a/Engine/Shaders/Private/PathTracing/PathTracing.usf b/Engine/Shaders/Private/PathTracing/PathTracing.usf
--- a/Engine/Shaders/Private/PathTracing/PathTracing.usf
+++ b/Engine/Shaders/Private/PathTracing/PathTracing.usf
@@ -1,4 +1,4 @@
-// Copyright 1998-2018 Epic Games, Inc. All Rights Reserved.
+// Copyright 1998-2019 Epic Games, Inc. All Rights Reserved.
 
 /*=============================================================================================
 PathTracingRayGenShader.usf: Reference path tracing
diff --git a/Engine/Shaders/Private/PathTracing/PathTracingCHS.usf b/Engine/Shaders/Private/PathTracing/PathTracingCHS.usf
--- a/Engine/Shaders/Private/PathTracing/PathTracingCHS.usf
+++ b/Engine/Shaders/Private/PathTracing/PathTracingCHS.usf
@@ -1,4 +1,4 @@
-// Copyright 1998-2018 Epic Games, Inc. All Rights Reserved.
+// Copyright 1998-2019 Epic Games, Inc. All Rights Reserved.
 
 #include "PathTracingDefinitions.ush"
 #include "../RayTracing/RayTracingCommon.ush"
diff --git a/Engine/Shaders/Private/PathTracing/PathTracingCamera.ush b/Engine/Shaders/Private/PathTracing/PathTracingCamera.ush
--- a/Engine/Shaders/Private/PathTracing/PathTracingCamera.ush
+++ b/Engine/Shaders/Private/PathTracing/PathTracingCamera.ush
@@ -1,4 +1,4 @@
-// Copyright 1998-2018 Epic Games, Inc. All Rights Reserved.
+// Copyright 1998-2019 Epic Games, Inc. All Rights Reserved.
 
 /*=============================================================================================
 	PathTracingCamera.ush: Reference path tracing camera functions
diff --git a/Engine/Shaders/Private/PathTracing/PathTracingCommon.ush b/Engine/Shaders/Private/PathTracing/PathTracingCommon.ush
--- a/Engine/Shaders/Private/PathTracing/PathTracingCommon.ush
+++ b/Engine/Shaders/Private/PathTracing/PathTracingCommon.ush
@@ -1,4 +1,4 @@
-// Copyright 1998-2018 Epic Games, Inc. All Rights Reserved.
+// Copyright 1998-2019 Epic Games, Inc. All Rights Reserved.
 
 /*=============================================================================
 	PathTracingCommon.ush: common header used in multiple ray generation shaders
diff --git a/Engine/Shaders/Private/PathTracing/PathTracingCompositingPixelShader.usf b/Engine/Shaders/Private/PathTracing/PathTracingCompositingPixelShader.usf
--- a/Engine/Shaders/Private/PathTracing/PathTracingCompositingPixelShader.usf
+++ b/Engine/Shaders/Private/PathTracing/PathTracingCompositingPixelShader.usf
@@ -1,4 +1,4 @@
-// Copyright 1998-2018 Epic Games, Inc. All Rights Reserved.
+// Copyright 1998-2019 Epic Games, Inc. All Rights Reserved.
 
 /*=============================================================================================
 ReferencePathTracinPixelShader.ush: Reference path tracing  compositing pixel shader for progressive update.
diff --git a/Engine/Shaders/Private/PathTracing/PathTracingDefinitions.ush b/Engine/Shaders/Private/PathTracing/PathTracingDefinitions.ush
--- a/Engine/Shaders/Private/PathTracing/PathTracingDefinitions.ush
+++ b/Engine/Shaders/Private/PathTracing/PathTracingDefinitions.ush
@@ -1,4 +1,4 @@
-// Copyright 1998-2018 Epic Games, Inc. All Rights Reserved.
+// Copyright 1998-2019 Epic Games, Inc. All Rights Reserved.
 
 /*================================================================================================
 	PathTracingDefinitions.ush: used in ray tracing shaders  and C++ code to define common constants
diff --git a/Engine/Shaders/Private/PathTracing/PathTracingMS.usf b/Engine/Shaders/Private/PathTracing/PathTracingMS.usf
--- a/Engine/Shaders/Private/PathTracing/PathTracingMS.usf
+++ b/Engine/Shaders/Private/PathTracing/PathTracingMS.usf
@@ -1,4 +1,4 @@
-// Copyright 1998-2018 Epic Games, Inc. All Rights Reserved.
+// Copyright 1998-2019 Epic Games, Inc. All Rights Reserved.
 
 #include "PathTracingDefinitions.ush"
 #include "../RayTracing/RayTracingCommon.ush"
diff --git a/Engine/Shaders/Private/PathTracing/PathTracingRayCounterComputeShader.usf b/Engine/Shaders/Private/PathTracing/PathTracingRayCounterComputeShader.usf
--- a/Engine/Shaders/Private/PathTracing/PathTracingRayCounterComputeShader.usf
+++ b/Engine/Shaders/Private/PathTracing/PathTracingRayCounterComputeShader.usf
@@ -1,4 +1,4 @@
-// Copyright 1998-2018 Epic Games, Inc. All Rights Reserved.
+// Copyright 1998-2019 Epic Games, Inc. All Rights Reserved.
 
 /*=============================================================================================
 	RayCounterComputeShader.usf: Ray counter for the reference path tracing
diff --git a/Engine/Shaders/Private/PathTracing/Sensor/PathTracingAdaptivePinholeCamera.ush b/Engine/Shaders/Private/PathTracing/Sensor/PathTracingAdaptivePinholeCamera.ush
--- a/Engine/Shaders/Private/PathTracing/Sensor/PathTracingAdaptivePinholeCamera.ush
+++ b/Engine/Shaders/Private/PathTracing/Sensor/PathTracingAdaptivePinholeCamera.ush
@@ -1,4 +1,4 @@
-// Copyright 1998-2018 Epic Games, Inc. All Rights Reserved.
+// Copyright 1998-2019 Epic Games, Inc. All Rights Reserved.
 
 /*=============================================================================================
 	PathTracingAdaptivePinholeCamera.ush: Reference path tracing camera functions
diff --git a/Engine/Shaders/Private/PathTracing/Sensor/PathTracingPinholeCamera.ush b/Engine/Shaders/Private/PathTracing/Sensor/PathTracingPinholeCamera.ush
--- a/Engine/Shaders/Private/PathTracing/Sensor/PathTracingPinholeCamera.ush
+++ b/Engine/Shaders/Private/PathTracing/Sensor/PathTracingPinholeCamera.ush
@@ -1,4 +1,4 @@
-// Copyright 1998-2018 Epic Games, Inc. All Rights Reserved.
+// Copyright 1998-2019 Epic Games, Inc. All Rights Reserved.
 
 /*=============================================================================================
 	PathTracingPinholeCamera.ush: Reference path tracing camera functions
diff --git a/Engine/Shaders/Private/PathTracing/Utilities/PathTracingFixedPointArithmetic.ush b/Engine/Shaders/Private/PathTracing/Utilities/PathTracingFixedPointArithmetic.ush
--- a/Engine/Shaders/Private/PathTracing/Utilities/PathTracingFixedPointArithmetic.ush
+++ b/Engine/Shaders/Private/PathTracing/Utilities/PathTracingFixedPointArithmetic.ush
@@ -1,4 +1,4 @@
-// Copyright 1998-2018 Epic Games, Inc. All Rights Reserved.
+// Copyright 1998-2019 Epic Games, Inc. All Rights Reserved.
 
 /*=============================================================================================
 	PathTracingFixedPointArithmetic.ush: Utilities for fixed-point arithmetic. Necessary for encoding
diff --git a/Engine/Shaders/Private/PathTracing/Utilities/PathTracingRandomSequence.ush b/Engine/Shaders/Private/PathTracing/Utilities/PathTracingRandomSequence.ush
--- a/Engine/Shaders/Private/PathTracing/Utilities/PathTracingRandomSequence.ush
+++ b/Engine/Shaders/Private/PathTracing/Utilities/PathTracingRandomSequence.ush
@@ -1,4 +1,4 @@
-// Copyright 1998-2018 Epic Games, Inc. All Rights Reserved.
+// Copyright 1998-2019 Epic Games, Inc. All Rights Reserved.
 
 /*=============================================================================================
 	PathTracingRandomSequence.ush: Reference path tracing
diff --git a/Engine/Shaders/Private/PathTracing/VisualizeMipTreePixelShader.usf b/Engine/Shaders/Private/PathTracing/VisualizeMipTreePixelShader.usf
--- a/Engine/Shaders/Private/PathTracing/VisualizeMipTreePixelShader.usf
+++ b/Engine/Shaders/Private/PathTracing/VisualizeMipTreePixelShader.usf
@@ -1,4 +1,4 @@
-// Copyright 1998-2018 Epic Games, Inc. All Rights Reserved.
+// Copyright 1998-2019 Epic Games, Inc. All Rights Reserved.
 
 /*=============================================================================================
 	VisualizeMipTrreePixelShader.usf: Interactive visualization of variance map.
diff --git a/Engine/Shaders/Private/RayTracing/BuildMipTreeCS.usf b/Engine/Shaders/Private/RayTracing/BuildMipTreeCS.usf
--- a/Engine/Shaders/Private/RayTracing/BuildMipTreeCS.usf
+++ b/Engine/Shaders/Private/RayTracing/BuildMipTreeCS.usf
@@ -1,4 +1,4 @@
-// Copyright 1998-2018 Epic Games, Inc. All Rights Reserved.
+// Copyright 1998-2019 Epic Games, Inc. All Rights Reserved.
 
 /*=============================================================================================
 	BuildMipTreeCS.usf: Build variance map
diff --git a/Engine/Shaders/Private/RayTracing/BuildMipTreePdfCS.usf b/Engine/Shaders/Private/RayTracing/BuildMipTreePdfCS.usf
--- a/Engine/Shaders/Private/RayTracing/BuildMipTreePdfCS.usf
+++ b/Engine/Shaders/Private/RayTracing/BuildMipTreePdfCS.usf
@@ -1,4 +1,4 @@
-// Copyright 1998-2018 Epic Games, Inc. All Rights Reserved.
+// Copyright 1998-2019 Epic Games, Inc. All Rights Reserved.
 
 /*=============================================================================================
 BuildMipTreeCS.usf: Build MipTree PDF structure
diff --git a/Engine/Shaders/Private/RayTracing/BuildSolidAnglePdfCS.usf b/Engine/Shaders/Private/RayTracing/BuildSolidAnglePdfCS.usf
--- a/Engine/Shaders/Private/RayTracing/BuildSolidAnglePdfCS.usf
+++ b/Engine/Shaders/Private/RayTracing/BuildSolidAnglePdfCS.usf
@@ -1,4 +1,4 @@
-// Copyright 1998-2018 Epic Games, Inc. All Rights Reserved.
+// Copyright 1998-2019 Epic Games, Inc. All Rights Reserved.
 
 /*=============================================================================================
 BuildMipTreeCS.usf: Build MipTree PDF structure
diff --git a/Engine/Shaders/Private/RayTracing/CommonRayTracingBuiltInResources.ush b/Engine/Shaders/Private/RayTracing/CommonRayTracingBuiltInResources.ush
--- a/Engine/Shaders/Private/RayTracing/CommonRayTracingBuiltInResources.ush
+++ b/Engine/Shaders/Private/RayTracing/CommonRayTracingBuiltInResources.ush
@@ -1,4 +1,4 @@
-// Copyright 1998-2018 Epic Games, Inc. All Rights Reserved.
+// Copyright 1998-2019 Epic Games, Inc. All Rights Reserved.
 
 /*************************************************************
 * Changeing this file requires recompillation of the engine! *
diff --git a/Engine/Shaders/Private/RayTracing/CompositeAmbientOcclusionPS.usf b/Engine/Shaders/Private/RayTracing/CompositeAmbientOcclusionPS.usf
--- a/Engine/Shaders/Private/RayTracing/CompositeAmbientOcclusionPS.usf
+++ b/Engine/Shaders/Private/RayTracing/CompositeAmbientOcclusionPS.usf
@@ -1,4 +1,4 @@
-// Copyright 1998-2018 Epic Games, Inc. All Rights Reserved.
+// Copyright 1998-2019 Epic Games, Inc. All Rights Reserved.
 
 #include "../Common.ush"
 #include "../DeferredShadingCommon.ush"
diff --git a/Engine/Shaders/Private/RayTracing/CompositeSkyLightPS.usf b/Engine/Shaders/Private/RayTracing/CompositeSkyLightPS.usf
--- a/Engine/Shaders/Private/RayTracing/CompositeSkyLightPS.usf
+++ b/Engine/Shaders/Private/RayTracing/CompositeSkyLightPS.usf
@@ -1,4 +1,4 @@
-// Copyright 1998-2018 Epic Games, Inc. All Rights Reserved.
+// Copyright 1998-2019 Epic Games, Inc. All Rights Reserved.
 
 #include "../Common.ush"
 #include "../DeferredShadingCommon.ush"
diff --git a/Engine/Shaders/Private/RayTracing/MipTreeCommon.ush b/Engine/Shaders/Private/RayTracing/MipTreeCommon.ush
--- a/Engine/Shaders/Private/RayTracing/MipTreeCommon.ush
+++ b/Engine/Shaders/Private/RayTracing/MipTreeCommon.ush
@@ -1,4 +1,4 @@
-// Copyright 1998-2018 Epic Games, Inc. All Rights Reserved.
+// Copyright 1998-2019 Epic Games, Inc. All Rights Reserved.
 
 /*=============================================================================================
 	MipTreeCommon.ush: Common utilities for MipTree-based sampling
diff --git a/Engine/Shaders/Private/RayTracing/RayTracingAmbientOcclusionRGS.usf b/Engine/Shaders/Private/RayTracing/RayTracingAmbientOcclusionRGS.usf
--- a/Engine/Shaders/Private/RayTracing/RayTracingAmbientOcclusionRGS.usf
+++ b/Engine/Shaders/Private/RayTracing/RayTracingAmbientOcclusionRGS.usf
@@ -1,4 +1,4 @@
-// Copyright 1998-2018 Epic Games, Inc. All Rights Reserved.
+// Copyright 1998-2019 Epic Games, Inc. All Rights Reserved.
 
 #include "../Common.ush"
 #include "../Montecarlo.ush"
diff --git a/Engine/Shaders/Private/RayTracing/RayTracingBarycentrics.usf b/Engine/Shaders/Private/RayTracing/RayTracingBarycentrics.usf
--- a/Engine/Shaders/Private/RayTracing/RayTracingBarycentrics.usf
+++ b/Engine/Shaders/Private/RayTracing/RayTracingBarycentrics.usf
@@ -1,4 +1,4 @@
-// Copyright 1998-2018 Epic Games, Inc. All Rights Reserved.
+// Copyright 1998-2019 Epic Games, Inc. All Rights Reserved.
 
 #include "../Common.ush"
 #include "RayTracingCommon.ush"
diff --git a/Engine/Shaders/Private/RayTracing/RayTracingBuiltInShaders.usf b/Engine/Shaders/Private/RayTracing/RayTracingBuiltInShaders.usf
--- a/Engine/Shaders/Private/RayTracing/RayTracingBuiltInShaders.usf
+++ b/Engine/Shaders/Private/RayTracing/RayTracingBuiltInShaders.usf
@@ -1,4 +1,4 @@
-// Copyright 1998-2018 Epic Games, Inc. All Rights Reserved.
+// Copyright 1998-2019 Epic Games, Inc. All Rights Reserved.
 
 #include "RayTracingCommon.ush"
 
diff --git a/Engine/Shaders/Private/RayTracing/RayTracingCommon.ush b/Engine/Shaders/Private/RayTracing/RayTracingCommon.ush
--- a/Engine/Shaders/Private/RayTracing/RayTracingCommon.ush
+++ b/Engine/Shaders/Private/RayTracing/RayTracingCommon.ush
@@ -1,4 +1,4 @@
-// Copyright 1998-2018 Epic Games, Inc. All Rights Reserved.
+// Copyright 1998-2019 Epic Games, Inc. All Rights Reserved.
 
 /*=============================================================================
 	RayTracingCommon.ush: common header used in multiple ray generation shaders
diff --git a/Engine/Shaders/Private/RayTracing/RayTracingDebug.usf b/Engine/Shaders/Private/RayTracing/RayTracingDebug.usf
--- a/Engine/Shaders/Private/RayTracing/RayTracingDebug.usf
+++ b/Engine/Shaders/Private/RayTracing/RayTracingDebug.usf
@@ -1,4 +1,4 @@
-// Copyright 1998-2018 Epic Games, Inc. All Rights Reserved.
+// Copyright 1998-2019 Epic Games, Inc. All Rights Reserved.
 
 #include "../Common.ush"
 #include "RayTracingDebugDefinitions.ush"
diff --git a/Engine/Shaders/Private/RayTracing/RayTracingDebugDefinitions.ush b/Engine/Shaders/Private/RayTracing/RayTracingDebugDefinitions.ush
--- a/Engine/Shaders/Private/RayTracing/RayTracingDebugDefinitions.ush
+++ b/Engine/Shaders/Private/RayTracing/RayTracingDebugDefinitions.ush
@@ -1,4 +1,4 @@
-// Copyright 1998-2018 Epic Games, Inc. All Rights Reserved.
+// Copyright 1998-2019 Epic Games, Inc. All Rights Reserved.
 
 /*================================================================================================
 	RayTracingDebugDefinitions.ush: used in ray tracing debug shaders  and C++ code to define common constants
diff --git a/Engine/Shaders/Private/RayTracing/RayTracingDefinitions.ush b/Engine/Shaders/Private/RayTracing/RayTracingDefinitions.ush
--- a/Engine/Shaders/Private/RayTracing/RayTracingDefinitions.ush
+++ b/Engine/Shaders/Private/RayTracing/RayTracingDefinitions.ush
@@ -1,4 +1,4 @@
-// Copyright 1998-2018 Epic Games, Inc. All Rights Reserved.
+// Copyright 1998-2019 Epic Games, Inc. All Rights Reserved.
 
 /*================================================================================================
 	RayTracingDefinitions.ush: used in ray tracing shaders  and C++ code to define common constants
diff --git a/Engine/Shaders/Private/RayTracing/RayTracingDirectionalLight.ush b/Engine/Shaders/Private/RayTracing/RayTracingDirectionalLight.ush
--- a/Engine/Shaders/Private/RayTracing/RayTracingDirectionalLight.ush
+++ b/Engine/Shaders/Private/RayTracing/RayTracingDirectionalLight.ush
@@ -1,4 +1,4 @@
-// Copyright 1998-2018 Epic Games, Inc. All Rights Reserved.
+// Copyright 1998-2019 Epic Games, Inc. All Rights Reserved.
 
 #pragma once
 
diff --git a/Engine/Shaders/Private/RayTracing/RayTracingDynamicMesh.usf b/Engine/Shaders/Private/RayTracing/RayTracingDynamicMesh.usf
--- a/Engine/Shaders/Private/RayTracing/RayTracingDynamicMesh.usf
+++ b/Engine/Shaders/Private/RayTracing/RayTracingDynamicMesh.usf
@@ -1,4 +1,4 @@
-// Copyright 1998-2018 Epic Games, Inc. All Rights Reserved.
+// Copyright 1998-2019 Epic Games, Inc. All Rights Reserved.
 
 /*=============================================================================
 	RayTracingDynamicMesh.usf
diff --git a/Engine/Shaders/Private/RayTracing/RayTracingHitGroupCommon.ush b/Engine/Shaders/Private/RayTracing/RayTracingHitGroupCommon.ush
--- a/Engine/Shaders/Private/RayTracing/RayTracingHitGroupCommon.ush
+++ b/Engine/Shaders/Private/RayTracing/RayTracingHitGroupCommon.ush
@@ -1,4 +1,4 @@
-// Copyright 1998-2018 Epic Games, Inc. All Rights Reserved.
+// Copyright 1998-2019 Epic Games, Inc. All Rights Reserved.
 
 /*=============================================================================
 	RayTracingHitGroupCommon.ush: common header used in all hit group shaders
diff --git a/Engine/Shaders/Private/RayTracing/RayTracingMaterialHitShaders.usf b/Engine/Shaders/Private/RayTracing/RayTracingMaterialHitShaders.usf
--- a/Engine/Shaders/Private/RayTracing/RayTracingMaterialHitShaders.usf
+++ b/Engine/Shaders/Private/RayTracing/RayTracingMaterialHitShaders.usf
@@ -1,4 +1,4 @@
-// Copyright 1998-2018 Epic Games, Inc. All Rights Reserved.
+// Copyright 1998-2019 Epic Games, Inc. All Rights Reserved.
 
 #pragma once
 
diff --git a/Engine/Shaders/Private/RayTracing/RayTracingOcclusionRGS.usf b/Engine/Shaders/Private/RayTracing/RayTracingOcclusionRGS.usf
--- a/Engine/Shaders/Private/RayTracing/RayTracingOcclusionRGS.usf
+++ b/Engine/Shaders/Private/RayTracing/RayTracingOcclusionRGS.usf
@@ -1,4 +1,4 @@
-// Copyright 1998-2018 Epic Games, Inc. All Rights Reserved.
+// Copyright 1998-2019 Epic Games, Inc. All Rights Reserved.
 
 #include "../Common.ush"
 #include "../MonteCarlo.ush"
diff --git a/Engine/Shaders/Private/RayTracing/RayTracingRectLight.ush b/Engine/Shaders/Private/RayTracing/RayTracingRectLight.ush
--- a/Engine/Shaders/Private/RayTracing/RayTracingRectLight.ush
+++ b/Engine/Shaders/Private/RayTracing/RayTracingRectLight.ush
@@ -1,4 +1,4 @@
-// Copyright 1998-2018 Epic Games, Inc. All Rights Reserved.
+// Copyright 1998-2019 Epic Games, Inc. All Rights Reserved.
 
 #pragma once
 
diff --git a/Engine/Shaders/Private/RayTracing/RayTracingRectLightOcclusionRGS.usf b/Engine/Shaders/Private/RayTracing/RayTracingRectLightOcclusionRGS.usf
--- a/Engine/Shaders/Private/RayTracing/RayTracingRectLightOcclusionRGS.usf
+++ b/Engine/Shaders/Private/RayTracing/RayTracingRectLightOcclusionRGS.usf
@@ -1,4 +1,4 @@
-// Copyright 1998-2018 Epic Games, Inc. All Rights Reserved.
+// Copyright 1998-2019 Epic Games, Inc. All Rights Reserved.
 
 #include "../Common.ush"
 #include "../DeferredShadingCommon.ush"
diff --git a/Engine/Shaders/Private/RayTracing/RayTracingReflections.usf b/Engine/Shaders/Private/RayTracing/RayTracingReflections.usf
--- a/Engine/Shaders/Private/RayTracing/RayTracingReflections.usf
+++ b/Engine/Shaders/Private/RayTracing/RayTracingReflections.usf
@@ -1,4 +1,4 @@
-// Copyright 1998-2018 Epic Games, Inc. All Rights Reserved.
+// Copyright 1998-2019 Epic Games, Inc. All Rights Reserved.
 
 #ifndef DIM_DENOISER_OUTPUT
 	#define DIM_DENOISER_OUTPUT 0
diff --git a/Engine/Shaders/Private/RayTracing/RayTracingSkyLightRGS.usf b/Engine/Shaders/Private/RayTracing/RayTracingSkyLightRGS.usf
--- a/Engine/Shaders/Private/RayTracing/RayTracingSkyLightRGS.usf
+++ b/Engine/Shaders/Private/RayTracing/RayTracingSkyLightRGS.usf
@@ -1,4 +1,4 @@
-// Copyright 1998-2018 Epic Games, Inc. All Rights Reserved.
+// Copyright 1998-2019 Epic Games, Inc. All Rights Reserved.
 
 #include "../Common.ush"
 #include "../MonteCarlo.ush"
diff --git a/Engine/Shaders/Private/RayTracing/RayTracingTest.usf b/Engine/Shaders/Private/RayTracing/RayTracingTest.usf
--- a/Engine/Shaders/Private/RayTracing/RayTracingTest.usf
+++ b/Engine/Shaders/Private/RayTracing/RayTracingTest.usf
@@ -1,4 +1,4 @@
-// Copyright 1998-2018 Epic Games, Inc. All Rights Reserved.
+// Copyright 1998-2019 Epic Games, Inc. All Rights Reserved.
 
 #include "RayTracingCommon.ush"
 
diff --git a/Engine/Shaders/Private/RayTracing/SkyLightMipTreeCommon.ush b/Engine/Shaders/Private/RayTracing/SkyLightMipTreeCommon.ush
--- a/Engine/Shaders/Private/RayTracing/SkyLightMipTreeCommon.ush
+++ b/Engine/Shaders/Private/RayTracing/SkyLightMipTreeCommon.ush
@@ -1,4 +1,4 @@
-// Copyright 1998-2018 Epic Games, Inc. All Rights Reserved.
+// Copyright 1998-2019 Epic Games, Inc. All Rights Reserved.
 
 /*=============================================================================================
 	SkyLightMipTreeCommon.ush: Common utilities for MipTree-based SkyLight sampling
diff --git a/Engine/Shaders/Private/RayTracing/VisualizeSkyLightMipTreePS.usf b/Engine/Shaders/Private/RayTracing/VisualizeSkyLightMipTreePS.usf
--- a/Engine/Shaders/Private/RayTracing/VisualizeSkyLightMipTreePS.usf
+++ b/Engine/Shaders/Private/RayTracing/VisualizeSkyLightMipTreePS.usf
@@ -1,4 +1,4 @@
-// Copyright 1998-2018 Epic Games, Inc. All Rights Reserved.
+// Copyright 1998-2019 Epic Games, Inc. All Rights Reserved.
 
 /*=============================================================================================
 	VisualizeSkyLightMipTreePixelShader.usf: Interactive visualization of sky light mip tree.
diff --git a/Engine/Shaders/Private/SceneData.ush b/Engine/Shaders/Private/SceneData.ush
--- a/Engine/Shaders/Private/SceneData.ush
+++ b/Engine/Shaders/Private/SceneData.ush
@@ -1,4 +1,4 @@
-// Copyright 1998-2018 Epic Games, Inc. All Rights Reserved.
+// Copyright 1998-2019 Epic Games, Inc. All Rights Reserved.
 
 /**
  * SceneData.ush
diff --git a/Engine/Shaders/Private/ScreenSpaceDenoise/SSDBufferLayouts.ush b/Engine/Shaders/Private/ScreenSpaceDenoise/SSDBufferLayouts.ush
--- a/Engine/Shaders/Private/ScreenSpaceDenoise/SSDBufferLayouts.ush
+++ b/Engine/Shaders/Private/ScreenSpaceDenoise/SSDBufferLayouts.ush
@@ -1,4 +1,4 @@
-// Copyright 1998-2018 Epic Games, Inc. All Rights Reserved.
+// Copyright 1998-2019 Epic Games, Inc. All Rights Reserved.
 
 /*=============================================================================
 	PostProcessTemporalAA.usf: Temporal AA
diff --git a/Engine/Shaders/Private/ScreenSpaceDenoise/SSDBuildHistoryRejection.usf b/Engine/Shaders/Private/ScreenSpaceDenoise/SSDBuildHistoryRejection.usf
--- a/Engine/Shaders/Private/ScreenSpaceDenoise/SSDBuildHistoryRejection.usf
+++ b/Engine/Shaders/Private/ScreenSpaceDenoise/SSDBuildHistoryRejection.usf
@@ -1,4 +1,4 @@
-// Copyright 1998-2018 Epic Games, Inc. All Rights Reserved.
+// Copyright 1998-2019 Epic Games, Inc. All Rights Reserved.
 
 /*=============================================================================
 SSDReduce.usf:
diff --git a/Engine/Shaders/Private/ScreenSpaceDenoise/SSDCommon.ush b/Engine/Shaders/Private/ScreenSpaceDenoise/SSDCommon.ush
--- a/Engine/Shaders/Private/ScreenSpaceDenoise/SSDCommon.ush
+++ b/Engine/Shaders/Private/ScreenSpaceDenoise/SSDCommon.ush
@@ -1,4 +1,4 @@
-// Copyright 1998-2018 Epic Games, Inc. All Rights Reserved.
+// Copyright 1998-2019 Epic Games, Inc. All Rights Reserved.
 
 /*=============================================================================
 	SSDCommon.ush: All the code commonly shared between uniform buffers.
diff --git a/Engine/Shaders/Private/ScreenSpaceDenoise/SSDDebug.usf b/Engine/Shaders/Private/ScreenSpaceDenoise/SSDDebug.usf
--- a/Engine/Shaders/Private/ScreenSpaceDenoise/SSDDebug.usf
+++ b/Engine/Shaders/Private/ScreenSpaceDenoise/SSDDebug.usf
@@ -1,4 +1,4 @@
-// Copyright 1998-2018 Epic Games, Inc. All Rights Reserved.
+// Copyright 1998-2019 Epic Games, Inc. All Rights Reserved.
 
 /*=============================================================================
 SSDOutput.usf: Output the denoising signal for external consumer
diff --git a/Engine/Shaders/Private/ScreenSpaceDenoise/SSDOutput.usf b/Engine/Shaders/Private/ScreenSpaceDenoise/SSDOutput.usf
--- a/Engine/Shaders/Private/ScreenSpaceDenoise/SSDOutput.usf
+++ b/Engine/Shaders/Private/ScreenSpaceDenoise/SSDOutput.usf
@@ -1,4 +1,4 @@
-// Copyright 1998-2018 Epic Games, Inc. All Rights Reserved.
+// Copyright 1998-2019 Epic Games, Inc. All Rights Reserved.
 
 /*=============================================================================
 SSDOutput.usf: Output the denoising signal for external consumer
diff --git a/Engine/Shaders/Private/ScreenSpaceDenoise/SSDPenumbraSignal.ush b/Engine/Shaders/Private/ScreenSpaceDenoise/SSDPenumbraSignal.ush
--- a/Engine/Shaders/Private/ScreenSpaceDenoise/SSDPenumbraSignal.ush
+++ b/Engine/Shaders/Private/ScreenSpaceDenoise/SSDPenumbraSignal.ush
@@ -1,4 +1,4 @@
-// Copyright 1998-2018 Epic Games, Inc. All Rights Reserved.
+// Copyright 1998-2019 Epic Games, Inc. All Rights Reserved.
 
 /*=============================================================================
 	SSDPenumbraSignal.ush: All the signal implementation for penumbra processing.
diff --git a/Engine/Shaders/Private/ScreenSpaceDenoise/SSDReduce.usf b/Engine/Shaders/Private/ScreenSpaceDenoise/SSDReduce.usf
--- a/Engine/Shaders/Private/ScreenSpaceDenoise/SSDReduce.usf
+++ b/Engine/Shaders/Private/ScreenSpaceDenoise/SSDReduce.usf
@@ -1,4 +1,4 @@
-// Copyright 1998-2018 Epic Games, Inc. All Rights Reserved.
+// Copyright 1998-2019 Epic Games, Inc. All Rights Reserved.
 
 /*=============================================================================
 SSDReduce.usf:
diff --git a/Engine/Shaders/Private/ScreenSpaceDenoise/SSDSignalCore.ush b/Engine/Shaders/Private/ScreenSpaceDenoise/SSDSignalCore.ush
--- a/Engine/Shaders/Private/ScreenSpaceDenoise/SSDSignalCore.ush
+++ b/Engine/Shaders/Private/ScreenSpaceDenoise/SSDSignalCore.ush
@@ -1,4 +1,4 @@
-// Copyright 1998-2018 Epic Games, Inc. All Rights Reserved.
+// Copyright 1998-2019 Epic Games, Inc. All Rights Reserved.
 
 /*=============================================================================
 	SSDCommon.ush: All the code commonly shared between uniform buffers.
diff --git a/Engine/Shaders/Private/ScreenSpaceDenoise/SSDSignalInterface.ush b/Engine/Shaders/Private/ScreenSpaceDenoise/SSDSignalInterface.ush
--- a/Engine/Shaders/Private/ScreenSpaceDenoise/SSDSignalInterface.ush
+++ b/Engine/Shaders/Private/ScreenSpaceDenoise/SSDSignalInterface.ush
@@ -1,4 +1,4 @@
-// Copyright 1998-2018 Epic Games, Inc. All Rights Reserved.
+// Copyright 1998-2019 Epic Games, Inc. All Rights Reserved.
 
 /*=============================================================================
 	SSDSignalInterface.ush: Selects the required signal API.
diff --git a/Engine/Shaders/Private/ScreenSpaceDenoise/SSDSpatialAccumulation.usf b/Engine/Shaders/Private/ScreenSpaceDenoise/SSDSpatialAccumulation.usf
--- a/Engine/Shaders/Private/ScreenSpaceDenoise/SSDSpatialAccumulation.usf
+++ b/Engine/Shaders/Private/ScreenSpaceDenoise/SSDSpatialAccumulation.usf
@@ -1,4 +1,4 @@
-// Copyright 1998-2018 Epic Games, Inc. All Rights Reserved.
+// Copyright 1998-2019 Epic Games, Inc. All Rights Reserved.
 
 /*=============================================================================
 PostProcessTemporalAA.usf: Temporal AA
diff --git a/Engine/Shaders/Private/ScreenSpaceDenoise/SSDTemporalAccumulation.usf b/Engine/Shaders/Private/ScreenSpaceDenoise/SSDTemporalAccumulation.usf
--- a/Engine/Shaders/Private/ScreenSpaceDenoise/SSDTemporalAccumulation.usf
+++ b/Engine/Shaders/Private/ScreenSpaceDenoise/SSDTemporalAccumulation.usf
@@ -1,4 +1,4 @@
-// Copyright 1998-2018 Epic Games, Inc. All Rights Reserved.
+// Copyright 1998-2019 Epic Games, Inc. All Rights Reserved.
 
 /*=============================================================================
 	PostProcessTemporalAA.usf: Temporal AA
diff --git a/Engine/Shaders/Private/ScreenSpaceDenoise/SSDTileDilate.usf b/Engine/Shaders/Private/ScreenSpaceDenoise/SSDTileDilate.usf
--- a/Engine/Shaders/Private/ScreenSpaceDenoise/SSDTileDilate.usf
+++ b/Engine/Shaders/Private/ScreenSpaceDenoise/SSDTileDilate.usf
@@ -1,4 +1,4 @@
-// Copyright 1998-2018 Epic Games, Inc. All Rights Reserved.
+// Copyright 1998-2019 Epic Games, Inc. All Rights Reserved.
 
 /*=============================================================================
 SSDReduce.usf:
diff --git a/Engine/Shaders/Private/TextureSampling.ush b/Engine/Shaders/Private/TextureSampling.ush
--- a/Engine/Shaders/Private/TextureSampling.ush
+++ b/Engine/Shaders/Private/TextureSampling.ush
@@ -1,4 +1,4 @@
-// Copyright 1998-2018 Epic Games, Inc. All Rights Reserved.
+// Copyright 1998-2019 Epic Games, Inc. All Rights Reserved.
 
 /*=============================================================================
 	TextureSampling.usf: Hold helper to do custom texture sampling.
diff --git a/Engine/Shaders/Private/VisualizeRayBuffer.usf b/Engine/Shaders/Private/VisualizeRayBuffer.usf
--- a/Engine/Shaders/Private/VisualizeRayBuffer.usf
+++ b/Engine/Shaders/Private/VisualizeRayBuffer.usf
@@ -1,4 +1,4 @@
-// Copyright 1998-2018 Epic Games, Inc. All Rights Reserved.
+// Copyright 1998-2019 Epic Games, Inc. All Rights Reserved.
 
 /*=============================================================================================
 	VisualizeRayBuffer.usf: Prepare occlusion rays for a given directional light
diff --git a/Engine/Source/Editor/UnrealEd/Private/RayTracingDebugVisualizationMenuCommands.cpp b/Engine/Source/Editor/UnrealEd/Private/RayTracingDebugVisualizationMenuCommands.cpp
--- a/Engine/Source/Editor/UnrealEd/Private/RayTracingDebugVisualizationMenuCommands.cpp
+++ b/Engine/Source/Editor/UnrealEd/Private/RayTracingDebugVisualizationMenuCommands.cpp
@@ -1,4 +1,4 @@
-// Copyright 1998-2018 Epic Games, Inc. All Rights Reserved.
+// Copyright 1998-2019 Epic Games, Inc. All Rights Reserved.
 
 #include "RayTracingDebugVisualizationMenuCommands.h"
 #include "Containers/UnrealString.h"
diff --git a/Engine/Source/Editor/UnrealEd/Public/RayTracingDebugVisualizationMenuCommands.h b/Engine/Source/Editor/UnrealEd/Public/RayTracingDebugVisualizationMenuCommands.h
--- a/Engine/Source/Editor/UnrealEd/Public/RayTracingDebugVisualizationMenuCommands.h
+++ b/Engine/Source/Editor/UnrealEd/Public/RayTracingDebugVisualizationMenuCommands.h
@@ -1,4 +1,4 @@
-// Copyright 1998-2018 Epic Games, Inc. All Rights Reserved.
+// Copyright 1998-2019 Epic Games, Inc. All Rights Reserved.
 
 #pragma once
 
diff --git a/Engine/Source/Runtime/D3D12RHI/Private/D3D12RayTracing.cpp b/Engine/Source/Runtime/D3D12RHI/Private/D3D12RayTracing.cpp
--- a/Engine/Source/Runtime/D3D12RHI/Private/D3D12RayTracing.cpp
+++ b/Engine/Source/Runtime/D3D12RHI/Private/D3D12RayTracing.cpp
@@ -1,4 +1,4 @@
-// Copyright 1998-2018 Epic Games, Inc. All Rights Reserved.
+// Copyright 1998-2019 Epic Games, Inc. All Rights Reserved.
 
 #include "D3D12RayTracing.h"
 
diff --git a/Engine/Source/Runtime/D3D12RHI/Private/D3D12RayTracing.h b/Engine/Source/Runtime/D3D12RHI/Private/D3D12RayTracing.h
--- a/Engine/Source/Runtime/D3D12RHI/Private/D3D12RayTracing.h
+++ b/Engine/Source/Runtime/D3D12RHI/Private/D3D12RayTracing.h
@@ -1,4 +1,4 @@
-// Copyright 1998-2018 Epic Games, Inc. All Rights Reserved.
+// Copyright 1998-2019 Epic Games, Inc. All Rights Reserved.
 
 #pragma once
 
diff --git a/Engine/Source/Runtime/Engine/Classes/Materials/MaterialExpressionDeltaTime.h b/Engine/Source/Runtime/Engine/Classes/Materials/MaterialExpressionDeltaTime.h
--- a/Engine/Source/Runtime/Engine/Classes/Materials/MaterialExpressionDeltaTime.h
+++ b/Engine/Source/Runtime/Engine/Classes/Materials/MaterialExpressionDeltaTime.h
@@ -1,4 +1,4 @@
-// Copyright 1998-2018 Epic Games, Inc. All Rights Reserved.
+// Copyright 1998-2019 Epic Games, Inc. All Rights Reserved.
 
 
 #pragma once
diff --git a/Engine/Source/Runtime/Engine/Classes/Materials/MaterialExpressionShadowReplace.h b/Engine/Source/Runtime/Engine/Classes/Materials/MaterialExpressionShadowReplace.h
--- a/Engine/Source/Runtime/Engine/Classes/Materials/MaterialExpressionShadowReplace.h
+++ b/Engine/Source/Runtime/Engine/Classes/Materials/MaterialExpressionShadowReplace.h
@@ -1,4 +1,4 @@
-// Copyright 1998-2018 Epic Games, Inc. All Rights Reserved.
+// Copyright 1998-2019 Epic Games, Inc. All Rights Reserved.
 
 
 #pragma once
diff --git a/Engine/Source/Runtime/Engine/Private/DebugViewModeInterface.cpp b/Engine/Source/Runtime/Engine/Private/DebugViewModeInterface.cpp
--- a/Engine/Source/Runtime/Engine/Private/DebugViewModeInterface.cpp
+++ b/Engine/Source/Runtime/Engine/Private/DebugViewModeInterface.cpp
@@ -1,4 +1,4 @@
-// Copyright 1998-2018 Epic Games, Inc. All Rights Reserved.
+// Copyright 1998-2019 Epic Games, Inc. All Rights Reserved.
 
 /*=============================================================================
 DebugViewModeInterface.cpp: Contains definitions for rendering debug viewmodes.
diff --git a/Engine/Source/Runtime/Engine/Private/PrimitiveUniformShaderParameters.cpp b/Engine/Source/Runtime/Engine/Private/PrimitiveUniformShaderParameters.cpp
--- a/Engine/Source/Runtime/Engine/Private/PrimitiveUniformShaderParameters.cpp
+++ b/Engine/Source/Runtime/Engine/Private/PrimitiveUniformShaderParameters.cpp
@@ -1,4 +1,4 @@
-// Copyright 1998-2018 Epic Games, Inc. All Rights Reserved.
+// Copyright 1998-2019 Epic Games, Inc. All Rights Reserved.
 
 #include "PrimitiveUniformShaderParameters.h"
 #include "PrimitiveSceneProxy.h"
diff --git a/Engine/Source/Runtime/Engine/Public/DebugViewModeInterface.h b/Engine/Source/Runtime/Engine/Public/DebugViewModeInterface.h
--- a/Engine/Source/Runtime/Engine/Public/DebugViewModeInterface.h
+++ b/Engine/Source/Runtime/Engine/Public/DebugViewModeInterface.h
@@ -1,4 +1,4 @@
-// Copyright 1998-2018 Epic Games, Inc. All Rights Reserved.
+// Copyright 1998-2019 Epic Games, Inc. All Rights Reserved.
 
 /*=============================================================================
 DebugViewModeInterface.h: Contains definitions for rendering debug viewmodes.
diff --git a/Engine/Source/Runtime/Engine/Public/Matrix3x4.h b/Engine/Source/Runtime/Engine/Public/Matrix3x4.h
--- a/Engine/Source/Runtime/Engine/Public/Matrix3x4.h
+++ b/Engine/Source/Runtime/Engine/Public/Matrix3x4.h
@@ -1,4 +1,4 @@
-// Copyright 1998-2018 Epic Games, Inc. All Rights Reserved.
+// Copyright 1998-2019 Epic Games, Inc. All Rights Reserved.
 
 #pragma once
 
diff --git a/Engine/Source/Runtime/RenderCore/Public/RayGenShaderUtils.h b/Engine/Source/Runtime/RenderCore/Public/RayGenShaderUtils.h
--- a/Engine/Source/Runtime/RenderCore/Public/RayGenShaderUtils.h
+++ b/Engine/Source/Runtime/RenderCore/Public/RayGenShaderUtils.h
@@ -1,4 +1,4 @@
-// Copyright 1998-2018 Epic Games, Inc. All Rights Reserved.
+// Copyright 1998-2019 Epic Games, Inc. All Rights Reserved.
 
 /*=============================================================================
 	RayGenShaderUtils.h: Utilities for ray generation shaders shaders.
diff --git a/Engine/Source/Runtime/Renderer/Private/EditorPrimitivesRendering.cpp b/Engine/Source/Runtime/Renderer/Private/EditorPrimitivesRendering.cpp
--- a/Engine/Source/Runtime/Renderer/Private/EditorPrimitivesRendering.cpp
+++ b/Engine/Source/Runtime/Renderer/Private/EditorPrimitivesRendering.cpp
@@ -1,4 +1,4 @@
-// Copyright 1998-2018 Epic Games, Inc. All Rights Reserved.
+// Copyright 1998-2019 Epic Games, Inc. All Rights Reserved.
 
 /*=============================================================================
 	EditorPrimitivesRendering
diff --git a/Engine/Source/Runtime/Renderer/Private/EditorPrimitivesRendering.h b/Engine/Source/Runtime/Renderer/Private/EditorPrimitivesRendering.h
--- a/Engine/Source/Runtime/Renderer/Private/EditorPrimitivesRendering.h
+++ b/Engine/Source/Runtime/Renderer/Private/EditorPrimitivesRendering.h
@@ -1,4 +1,4 @@
-// Copyright 1998-2018 Epic Games, Inc. All Rights Reserved.
+// Copyright 1998-2019 Epic Games, Inc. All Rights Reserved.
 
 /*=============================================================================
 	EditorPrimitivesRendering
diff --git a/Engine/Source/Runtime/Renderer/Private/GPUScene.cpp b/Engine/Source/Runtime/Renderer/Private/GPUScene.cpp
--- a/Engine/Source/Runtime/Renderer/Private/GPUScene.cpp
+++ b/Engine/Source/Runtime/Renderer/Private/GPUScene.cpp
@@ -1,4 +1,4 @@
-// Copyright 1998-2018 Epic Games, Inc. All Rights Reserved.
+// Copyright 1998-2019 Epic Games, Inc. All Rights Reserved.
 
 /*=============================================================================
 	GPUScene.cpp
diff --git a/Engine/Source/Runtime/Renderer/Private/GPUScene.h b/Engine/Source/Runtime/Renderer/Private/GPUScene.h
--- a/Engine/Source/Runtime/Renderer/Private/GPUScene.h
+++ b/Engine/Source/Runtime/Renderer/Private/GPUScene.h
@@ -1,4 +1,4 @@
-// Copyright 1998-2018 Epic Games, Inc. All Rights Reserved.
+// Copyright 1998-2019 Epic Games, Inc. All Rights Reserved.
 
 #pragma once
 
diff --git a/Engine/Source/Runtime/Renderer/Private/MobileBasePass.cpp b/Engine/Source/Runtime/Renderer/Private/MobileBasePass.cpp
--- a/Engine/Source/Runtime/Renderer/Private/MobileBasePass.cpp
+++ b/Engine/Source/Runtime/Renderer/Private/MobileBasePass.cpp
@@ -1,4 +1,4 @@
-// Copyright 1998-2018 Epic Games, Inc. All Rights Reserved.
+// Copyright 1998-2019 Epic Games, Inc. All Rights Reserved.
 
 /*=============================================================================
 	MobileBasePassRendering.cpp: Base pass rendering implementation.
diff --git a/Engine/Source/Runtime/Renderer/Private/PathTracing.cpp b/Engine/Source/Runtime/Renderer/Private/PathTracing.cpp
--- a/Engine/Source/Runtime/Renderer/Private/PathTracing.cpp
+++ b/Engine/Source/Runtime/Renderer/Private/PathTracing.cpp
@@ -1,4 +1,4 @@
-// Copyright 1998-2018 Epic Games, Inc. All Rights Reserved.
+// Copyright 1998-2019 Epic Games, Inc. All Rights Reserved.
 
 #include "RHI.h"
 
diff --git a/Engine/Source/Runtime/Renderer/Private/PathTracing/PathCompactionCompute.cpp b/Engine/Source/Runtime/Renderer/Private/PathTracing/PathCompactionCompute.cpp
--- a/Engine/Source/Runtime/Renderer/Private/PathTracing/PathCompactionCompute.cpp
+++ b/Engine/Source/Runtime/Renderer/Private/PathTracing/PathCompactionCompute.cpp
@@ -1,4 +1,4 @@
-// Copyright 1998-2018 Epic Games, Inc. All Rights Reserved.
+// Copyright 1998-2019 Epic Games, Inc. All Rights Reserved.
 
 /*=============================================================================
 PathCompactionCompute.cpp: Compute path continuation shader.
diff --git a/Engine/Source/Runtime/Renderer/Private/PathTracing/PathTracingTestbed.h b/Engine/Source/Runtime/Renderer/Private/PathTracing/PathTracingTestbed.h
--- a/Engine/Source/Runtime/Renderer/Private/PathTracing/PathTracingTestbed.h
+++ b/Engine/Source/Runtime/Renderer/Private/PathTracing/PathTracingTestbed.h
@@ -1,4 +1,4 @@
-// Copyright 1998-2018 Epic Games, Inc. All Rights Reserved.
+// Copyright 1998-2019 Epic Games, Inc. All Rights Reserved.
 
 #pragma once 
 
diff --git a/Engine/Source/Runtime/Renderer/Private/PathTracing/RayCounterCompute.cpp b/Engine/Source/Runtime/Renderer/Private/PathTracing/RayCounterCompute.cpp
--- a/Engine/Source/Runtime/Renderer/Private/PathTracing/RayCounterCompute.cpp
+++ b/Engine/Source/Runtime/Renderer/Private/PathTracing/RayCounterCompute.cpp
@@ -1,4 +1,4 @@
-// Copyright 1998-2018 Epic Games, Inc. All Rights Reserved.
+// Copyright 1998-2019 Epic Games, Inc. All Rights Reserved.
 
 /*=============================================================================
 	SkyLightCdfBuild.cpp: SkyLight CDF build algorithm.
diff --git a/Engine/Source/Runtime/Renderer/Private/PathTracingUniformBuffers.h b/Engine/Source/Runtime/Renderer/Private/PathTracingUniformBuffers.h
--- a/Engine/Source/Runtime/Renderer/Private/PathTracingUniformBuffers.h
+++ b/Engine/Source/Runtime/Renderer/Private/PathTracingUniformBuffers.h
@@ -1,4 +1,4 @@
-// Copyright 1998-2018 Epic Games, Inc. All Rights Reserved.
+// Copyright 1998-2019 Epic Games, Inc. All Rights Reserved.
 
 /*=============================================================================
 	PathTracingUniformBuffers.h
diff --git a/Engine/Source/Runtime/Renderer/Private/RayTracing/RayTracingAmbientOcclusion.cpp b/Engine/Source/Runtime/Renderer/Private/RayTracing/RayTracingAmbientOcclusion.cpp
--- a/Engine/Source/Runtime/Renderer/Private/RayTracing/RayTracingAmbientOcclusion.cpp
+++ b/Engine/Source/Runtime/Renderer/Private/RayTracing/RayTracingAmbientOcclusion.cpp
@@ -1,4 +1,4 @@
-// Copyright 1998-2018 Epic Games, Inc. All Rights Reserved.
+// Copyright 1998-2019 Epic Games, Inc. All Rights Reserved.
 
 #include "DeferredShadingRenderer.h"
 
diff --git a/Engine/Source/Runtime/Renderer/Private/RayTracing/RayTracingBarycentrics.cpp b/Engine/Source/Runtime/Renderer/Private/RayTracing/RayTracingBarycentrics.cpp
--- a/Engine/Source/Runtime/Renderer/Private/RayTracing/RayTracingBarycentrics.cpp
+++ b/Engine/Source/Runtime/Renderer/Private/RayTracing/RayTracingBarycentrics.cpp
@@ -1,4 +1,4 @@
-// Copyright 1998-2018 Epic Games, Inc. All Rights Reserved.
+// Copyright 1998-2019 Epic Games, Inc. All Rights Reserved.
 
 #include "RHI.h"
 
diff --git a/Engine/Source/Runtime/Renderer/Private/RayTracing/RayTracingDebug.cpp b/Engine/Source/Runtime/Renderer/Private/RayTracing/RayTracingDebug.cpp
--- a/Engine/Source/Runtime/Renderer/Private/RayTracing/RayTracingDebug.cpp
+++ b/Engine/Source/Runtime/Renderer/Private/RayTracing/RayTracingDebug.cpp
@@ -1,4 +1,4 @@
-// Copyright 1998-2018 Epic Games, Inc. All Rights Reserved.
+// Copyright 1998-2019 Epic Games, Inc. All Rights Reserved.
 
 #include "RHI.h"
 
diff --git a/Engine/Source/Runtime/Renderer/Private/RayTracing/RayTracingDynamicGeometryCollection.h b/Engine/Source/Runtime/Renderer/Private/RayTracing/RayTracingDynamicGeometryCollection.h
--- a/Engine/Source/Runtime/Renderer/Private/RayTracing/RayTracingDynamicGeometryCollection.h
+++ b/Engine/Source/Runtime/Renderer/Private/RayTracing/RayTracingDynamicGeometryCollection.h
@@ -1,4 +1,4 @@
-// Copyright 1998-2018 Epic Games, Inc. All Rights Reserved.
+// Copyright 1998-2019 Epic Games, Inc. All Rights Reserved.
 
 #pragma once
 
diff --git a/Engine/Source/Runtime/Renderer/Private/RayTracing/RayTracingMaterialHitShaders.cpp b/Engine/Source/Runtime/Renderer/Private/RayTracing/RayTracingMaterialHitShaders.cpp
--- a/Engine/Source/Runtime/Renderer/Private/RayTracing/RayTracingMaterialHitShaders.cpp
+++ b/Engine/Source/Runtime/Renderer/Private/RayTracing/RayTracingMaterialHitShaders.cpp
@@ -1,4 +1,4 @@
-// Copyright 1998-2018 Epic Games, Inc. All Rights Reserved.
+// Copyright 1998-2019 Epic Games, Inc. All Rights Reserved.
 
 #include "RayTracing/RayTracingMaterialHitShaders.h"
 
diff --git a/Engine/Source/Runtime/Renderer/Private/RayTracing/RayTracingMaterialHitShaders.h b/Engine/Source/Runtime/Renderer/Private/RayTracing/RayTracingMaterialHitShaders.h
--- a/Engine/Source/Runtime/Renderer/Private/RayTracing/RayTracingMaterialHitShaders.h
+++ b/Engine/Source/Runtime/Renderer/Private/RayTracing/RayTracingMaterialHitShaders.h
@@ -1,4 +1,4 @@
-// Copyright 1998-2018 Epic Games, Inc. All Rights Reserved.
+// Copyright 1998-2019 Epic Games, Inc. All Rights Reserved.
 
 #pragma once
 #include "MeshMaterialShader.h"
diff --git a/Engine/Source/Runtime/Renderer/Private/RayTracing/RayTracingReflections.cpp b/Engine/Source/Runtime/Renderer/Private/RayTracing/RayTracingReflections.cpp
--- a/Engine/Source/Runtime/Renderer/Private/RayTracing/RayTracingReflections.cpp
+++ b/Engine/Source/Runtime/Renderer/Private/RayTracing/RayTracingReflections.cpp
@@ -1,4 +1,4 @@
-// Copyright 1998-2018 Epic Games, Inc. All Rights Reserved.
+// Copyright 1998-2019 Epic Games, Inc. All Rights Reserved.
 
 #include "RendererPrivate.h"
 #include "GlobalShader.h"
diff --git a/Engine/Source/Runtime/Renderer/Private/RayTracing/RayTracingShadows.cpp b/Engine/Source/Runtime/Renderer/Private/RayTracing/RayTracingShadows.cpp
--- a/Engine/Source/Runtime/Renderer/Private/RayTracing/RayTracingShadows.cpp
+++ b/Engine/Source/Runtime/Renderer/Private/RayTracing/RayTracingShadows.cpp
@@ -1,4 +1,4 @@
-// Copyright 1998-2018 Epic Games, Inc. All Rights Reserved.
+// Copyright 1998-2019 Epic Games, Inc. All Rights Reserved.
 
 #include "DeferredShadingRenderer.h"
 
diff --git a/Engine/Source/Runtime/Renderer/Private/RayTracing/RayTracingShadowsForRectLight.cpp b/Engine/Source/Runtime/Renderer/Private/RayTracing/RayTracingShadowsForRectLight.cpp
--- a/Engine/Source/Runtime/Renderer/Private/RayTracing/RayTracingShadowsForRectLight.cpp
+++ b/Engine/Source/Runtime/Renderer/Private/RayTracing/RayTracingShadowsForRectLight.cpp
@@ -1,4 +1,4 @@
-// Copyright 1998-2018 Epic Games, Inc. All Rights Reserved.
+// Copyright 1998-2019 Epic Games, Inc. All Rights Reserved.
 
 #include "DeferredShadingRenderer.h"
 
diff --git a/Engine/Source/Runtime/Renderer/Private/RayTracing/RayTracingTestbed.cpp b/Engine/Source/Runtime/Renderer/Private/RayTracing/RayTracingTestbed.cpp
--- a/Engine/Source/Runtime/Renderer/Private/RayTracing/RayTracingTestbed.cpp
+++ b/Engine/Source/Runtime/Renderer/Private/RayTracing/RayTracingTestbed.cpp
@@ -1,4 +1,4 @@
-// Copyright 1998-2018 Epic Games, Inc. All Rights Reserved.
+// Copyright 1998-2019 Epic Games, Inc. All Rights Reserved.
 
 #include "RayTracingTestbed.h"
 #include "RHI.h"
diff --git a/Engine/Source/Runtime/Renderer/Private/RayTracing/RayTracingTestbed.h b/Engine/Source/Runtime/Renderer/Private/RayTracing/RayTracingTestbed.h
--- a/Engine/Source/Runtime/Renderer/Private/RayTracing/RayTracingTestbed.h
+++ b/Engine/Source/Runtime/Renderer/Private/RayTracing/RayTracingTestbed.h
@@ -1,4 +1,4 @@
-// Copyright 1998-2018 Epic Games, Inc. All Rights Reserved.
+// Copyright 1998-2019 Epic Games, Inc. All Rights Reserved.
 
 #pragma once 
 
diff --git a/Engine/Source/Runtime/Renderer/Private/RayTracing/RaytracingOptions.h b/Engine/Source/Runtime/Renderer/Private/RayTracing/RaytracingOptions.h
--- a/Engine/Source/Runtime/Renderer/Private/RayTracing/RaytracingOptions.h
+++ b/Engine/Source/Runtime/Renderer/Private/RayTracing/RaytracingOptions.h
@@ -1,4 +1,4 @@
-// Copyright 1998-2018 Epic Games, Inc. All Rights Reserved.
+// Copyright 1998-2019 Epic Games, Inc. All Rights Reserved.
 
 /*=============================================================================
 	RaytracingOptions.h declares ray tracing options for use in rendering
diff --git a/Engine/Source/Runtime/Renderer/Private/RayTracing/RaytracingSkylight.cpp b/Engine/Source/Runtime/Renderer/Private/RayTracing/RaytracingSkylight.cpp
--- a/Engine/Source/Runtime/Renderer/Private/RayTracing/RaytracingSkylight.cpp
+++ b/Engine/Source/Runtime/Renderer/Private/RayTracing/RaytracingSkylight.cpp
@@ -1,4 +1,4 @@
-// Copyright 1998-2018 Epic Games, Inc. All Rights Reserved.
+// Copyright 1998-2019 Epic Games, Inc. All Rights Reserved.
 
 /*=============================================================================
 	RaytracingSkyLight.cpp implements sky lighting with ray tracing
diff --git a/Engine/Source/Runtime/Renderer/Private/SceneViewFamilyBlackboard.h b/Engine/Source/Runtime/Renderer/Private/SceneViewFamilyBlackboard.h
--- a/Engine/Source/Runtime/Renderer/Private/SceneViewFamilyBlackboard.h
+++ b/Engine/Source/Runtime/Renderer/Private/SceneViewFamilyBlackboard.h
@@ -1,4 +1,4 @@
-// Copyright 1998-2018 Epic Games, Inc. All Rights Reserved.
+// Copyright 1998-2019 Epic Games, Inc. All Rights Reserved.
 
 /*=============================================================================
 	DeferredShadingRenderer.h: Scene rendering definitions.
diff --git a/Engine/Source/Runtime/Renderer/Private/ScreenSpaceDenoise.h b/Engine/Source/Runtime/Renderer/Private/ScreenSpaceDenoise.h
--- a/Engine/Source/Runtime/Renderer/Private/ScreenSpaceDenoise.h
+++ b/Engine/Source/Runtime/Renderer/Private/ScreenSpaceDenoise.h
@@ -1,4 +1,4 @@
-// Copyright 1998-2018 Epic Games, Inc. All Rights Reserved.
+// Copyright 1998-2019 Epic Games, Inc. All Rights Reserved.
 
 #pragma once
 
diff --git a/Engine/Source/Runtime/Renderer/Public/MeshDrawShaderBindings.h b/Engine/Source/Runtime/Renderer/Public/MeshDrawShaderBindings.h
--- a/Engine/Source/Runtime/Renderer/Public/MeshDrawShaderBindings.h
+++ b/Engine/Source/Runtime/Renderer/Public/MeshDrawShaderBindings.h
@@ -1,4 +1,4 @@
-// Copyright 1998-2018 Epic Games, Inc. All Rights Reserved.
+// Copyright 1998-2019 Epic Games, Inc. All Rights Reserved.
 
 /*=============================================================================
 	MeshDrawShaderBindings.h: 
diff --git a/Engine/Source/Runtime/Renderer/Public/MeshPassProcessor.inl b/Engine/Source/Runtime/Renderer/Public/MeshPassProcessor.inl
--- a/Engine/Source/Runtime/Renderer/Public/MeshPassProcessor.inl
+++ b/Engine/Source/Runtime/Renderer/Public/MeshPassProcessor.inl
@@ -1,4 +1,4 @@
-// Copyright 1998-2018 Epic Games, Inc. All Rights Reserved.
+// Copyright 1998-2019 Epic Games, Inc. All Rights Reserved.
 
 /*=============================================================================
 MeshPassProcessor.inl:
diff --git a/Engine/Source/Runtime/UtilityShaders/Private/BuiltInRayTracingShaders.cpp b/Engine/Source/Runtime/UtilityShaders/Private/BuiltInRayTracingShaders.cpp
--- a/Engine/Source/Runtime/UtilityShaders/Private/BuiltInRayTracingShaders.cpp
+++ b/Engine/Source/Runtime/UtilityShaders/Private/BuiltInRayTracingShaders.cpp
@@ -1,4 +1,4 @@
-// Copyright 1998-2018 Epic Games, Inc. All Rights Reserved.
+// Copyright 1998-2019 Epic Games, Inc. All Rights Reserved.
 
 #include "BuiltInRayTracingShaders.h"
 
diff --git a/Engine/Source/Runtime/UtilityShaders/Public/BuiltInRayTracingShaders.h b/Engine/Source/Runtime/UtilityShaders/Public/BuiltInRayTracingShaders.h
--- a/Engine/Source/Runtime/UtilityShaders/Public/BuiltInRayTracingShaders.h
+++ b/Engine/Source/Runtime/UtilityShaders/Public/BuiltInRayTracingShaders.h
@@ -1,4 +1,4 @@
-// Copyright 1998-2018 Epic Games, Inc. All Rights Reserved.
+// Copyright 1998-2019 Epic Games, Inc. All Rights Reserved.
 
 #pragma once
 
