Ray payload default behavior change: HitT is now default to 0.0 (hit state), and miss shader is responsible to change it to -1.0 (miss state). This allows us to skip closest hit shader entirely for shadow rays if hit distance is not needed

diff --git a/Engine/Shaders/Private/RayTracing/RayTracingCommon.ush b/Engine/Shaders/Private/RayTracing/RayTracingCommon.ush
--- a/Engine/Shaders/Private/RayTracing/RayTracingCommon.ush
+++ b/Engine/Shaders/Private/RayTracing/RayTracingCommon.ush
@@ -165,7 +165,6 @@ FMaterialClosestHitPayload TraceRayInternal(in RaytracingAccelerationStructure T
 											in RayDesc Ray)
 {
 	FPackedMaterialClosestHitPayload Payload = (FPackedMaterialClosestHitPayload)0;
-	Payload.HitT = -1.0f;
 
     TraceRay
 	(
