Skip closest hit shaders on reflected shadow rays
Enable any hit shaders for RTAO

diff --git a/Engine/Shaders/Private/RayTracing/RayTracingAmbientOcclusionRGS.usf b/Engine/Shaders/Private/RayTracing/RayTracingAmbientOcclusionRGS.usf
--- a/Engine/Shaders/Private/RayTracing/RayTracingAmbientOcclusionRGS.usf
+++ b/Engine/Shaders/Private/RayTracing/RayTracingAmbientOcclusionRGS.usf
@@ -86,8 +86,7 @@ void AmbientOcclusionRGS()
 		if (dot(WorldNormal, Ray.Direction) > 0.0)
 		{
 			uint RayFlags = 0
-				| RAY_FLAG_CULL_BACK_FACING_TRIANGLES
-				| RAY_FLAG_FORCE_OPAQUE; // #dxr_todo: generate a special shader permutation for occlusion ray tracing that includes any-hit shaders for masked geometry
+				| RAY_FLAG_CULL_BACK_FACING_TRIANGLES;
 
 			FDefaultPayload Payload = (FDefaultPayload)0;
 
diff --git a/Engine/Shaders/Private/RayTracing/RayTracingReflections.usf b/Engine/Shaders/Private/RayTracing/RayTracingReflections.usf
--- a/Engine/Shaders/Private/RayTracing/RayTracingReflections.usf
+++ b/Engine/Shaders/Private/RayTracing/RayTracingReflections.usf
@@ -132,7 +132,7 @@ float TraceShadowRay(
 	Ray.TMin = 1e-4f;
 	Ray.TMax = TMax;
 
-	uint RayFlags = RAY_FLAG_CULL_BACK_FACING_TRIANGLES | RAY_FLAG_ACCEPT_FIRST_HIT_AND_END_SEARCH;
+	uint RayFlags = RAY_FLAG_CULL_BACK_FACING_TRIANGLES | RAY_FLAG_ACCEPT_FIRST_HIT_AND_END_SEARCH | RAY_FLAG_SKIP_CLOSEST_HIT_SHADER;
 
 	FMaterialClosestHitPayload Payload = TraceRayInternal(
 		TLAS,      // AccelerationStructure
