Fixes ray tracing occlusion from SpotLight.

Light direction was incorrectly reversed when shading parameters were unified.

diff --git a/Engine/Shaders/Private/RayTracing/RayTracingOcclusionRGS.usf b/Engine/Shaders/Private/RayTracing/RayTracingOcclusionRGS.usf
--- a/Engine/Shaders/Private/RayTracing/RayTracingOcclusionRGS.usf
+++ b/Engine/Shaders/Private/RayTracing/RayTracingOcclusionRGS.usf
@@ -70,7 +70,7 @@ bool GenerateSpotLightOcclusionRay(
 
 	// Apply culling
 	float CosAngle = LightParameters.SpotAngles.x;
-	if (dot(-LightDirection, LightParameters.Direction) < CosAngle)
+	if (dot(LightDirection, LightParameters.Direction) < CosAngle)
 	{
 		return false;
 	}
