DR - Remove workaround for rect lights
#rnx

diff --git a/Engine/Shaders/Private/DeferredLightingCommon.ush b/Engine/Shaders/Private/DeferredLightingCommon.ush
--- a/Engine/Shaders/Private/DeferredLightingCommon.ush
+++ b/Engine/Shaders/Private/DeferredLightingCommon.ush
@@ -323,12 +323,7 @@ float4 GetDynamicLighting(float3 WorldPosition, float3 CameraVector, FGBufferDat
 		#else
 			FDirectLighting Lighting;
 
-			//temp workaround for HLSLCC bug with propagating SourceTexture parameter down
-			#if COMPILER_HLSLCC
-			if (false)
-			#else
 			if (LightData.bRectLight)
-			#endif
 			{
 				FRect Rect = GetRect( ToLight, LightData );
 
