Fixed variable ghosting issue

diff --git a/Engine/Shaders/Private/PathTracing/PathTracing.usf b/Engine/Shaders/Private/PathTracing/PathTracing.usf
--- a/Engine/Shaders/Private/PathTracing/PathTracing.usf
+++ b/Engine/Shaders/Private/PathTracing/PathTracing.usf
@@ -215,8 +215,8 @@ void PathTracingMainRG()
 
 						// Evaluate material
 						float3 MaterialThroughput;
-						float MaterialPdf = 0.0;
-						EvalMaterial(Ray.Direction, LightRay.Direction, RayHitInfo, MaterialThroughput, MaterialPdf);
+						float MaterialEvalPdf = 0.0;
+						EvalMaterial(Ray.Direction, LightRay.Direction, RayHitInfo, MaterialThroughput, MaterialEvalPdf);
 
 						// Apply material Pdf for correct MIS weight
 						float MisWeight = 1.0;
@@ -232,7 +232,7 @@ void PathTracingMainRG()
 							}
 #endif
 
-							MisWeight = NeePdf / (NeePdf + MaterialPdf);
+							MisWeight = NeePdf / (NeePdf + MaterialEvalPdf);
 						}
 
 						// Record the contribution
