Implements denoiser temporal accumulation pass for reflections

diff --git a/Engine/Shaders/Private/ScreenSpaceDenoise/SSDPenumbraSignal.ush b/Engine/Shaders/Private/ScreenSpaceDenoise/SSDPenumbraSignal.ush
--- a/Engine/Shaders/Private/ScreenSpaceDenoise/SSDPenumbraSignal.ush
+++ b/Engine/Shaders/Private/ScreenSpaceDenoise/SSDPenumbraSignal.ush
@@ -151,3 +151,11 @@ float GetSamplePenumbra(FSSFSignalSample Sample)
 	}
 	return Sample.LeavingRayCount / Sample.AccumulatedSampleCount;
 }
+
+float GetMultiImportanceSamplingWeight(
+	FSSFSampleSceneInfos RefSceneMetadata,
+	FSSFSampleSceneInfos SampleSceneMetadata,
+	FSSFSignalSample Sample)
+{
+	return 1.0;
+}
diff --git a/Engine/Shaders/Private/ScreenSpaceDenoise/SSDReflectionsSignal.ush b/Engine/Shaders/Private/ScreenSpaceDenoise/SSDReflectionsSignal.ush
new file mode 100644
--- /dev/null
+++ b/Engine/Shaders/Private/ScreenSpaceDenoise/SSDReflectionsSignal.ush
@@ -0,0 +1,85 @@
+// Copyright 1998-2018 Epic Games, Inc. All Rights Reserved.
+
+#pragma once
+
+
+#include "SSDCommon.ush"
+#include "SSDSignalCore.ush"
+
+
+
+/** COmpute the weight of the bilateral filter to use between a reference sample and a neighbor. */
+float ComputeBilateralWeight(
+	float MaxWorldBluringDistance,
+	FSSFSampleSceneInfos RefSceneMetadata,
+	FSSFSampleSceneInfos NeighborSceneMetadata)
+{
+	// TODO: duplicated code.
+	if (DEBUG_DISABLE_BILATERAL)
+	{
+		return 1;
+	}
+	
+	float WorldBilateralWeight;
+	{
+		const float WorldRadius = MaxWorldBluringDistance;
+
+		float3 V = RefSceneMetadata.TranslatedWorldPosition - NeighborSceneMetadata.TranslatedWorldPosition;
+
+		float RefAnisotropyInvFactor = ComputeAnisotropyInvFactor(RefSceneMetadata);
+
+		float3 CameraToRef = View.ViewForward;
+
+		float Z = dot(CameraToRef, V);
+		float XY = length2(V - CameraToRef * Z * (1 - RefAnisotropyInvFactor));
+
+		float DistSquare = XY;
+
+		float Multiplier = rcp(WorldRadius * WorldRadius);
+
+		WorldBilateralWeight = saturate(1 - DistSquare * Multiplier);
+	}
+
+	return WorldBilateralWeight;
+}
+
+float GetSignalWorldBluringRadius(FSSFSignalSample Sample, FSSFSampleSceneInfos SceneMetadata)
+{
+	// Not implemented.
+	return WORLD_RADIUS_INVALID;
+}
+
+FSSFSignalSample GetConstantSignalSample(uint TileMode)
+{
+	FSSFSignalSample OutputSample = CreateSignalSample();
+	if (TileMode == TILE_MODE_FULLY_INVALID)
+	{
+		OutputSample.AccumulatedSampleCount = 0.0;
+		OutputSample.LeavingRayCount = 0.0;
+	}
+	return OutputSample;
+}
+
+// Returns the penumbra of this sample, or 1 if invalid.
+float GetSamplePenumbraSafe(FSSFSignalSample Sample)
+{
+	return (Sample.AccumulatedSampleCount > 0 ? Sample.LeavingRayCount / Sample.AccumulatedSampleCount : 1);
+}
+
+// Get penumbra, but be aware of NaN.
+float GetSamplePenumbra(FSSFSignalSample Sample)
+{
+	if (1) // TODO: For debugging purpose of division by 0.
+	{
+		return GetSamplePenumbraSafe(Sample);
+	}
+	return Sample.LeavingRayCount / Sample.AccumulatedSampleCount;
+}
+
+float GetMultiImportanceSamplingWeight(
+	FSSFSampleSceneInfos RefSceneMetadata,
+	FSSFSampleSceneInfos SampleSceneMetadata,
+	FSSFSignalSample Sample)
+{
+	return 1.0;
+}
diff --git a/Engine/Shaders/Private/ScreenSpaceDenoise/SSDSignalCore.ush b/Engine/Shaders/Private/ScreenSpaceDenoise/SSDSignalCore.ush
--- a/Engine/Shaders/Private/ScreenSpaceDenoise/SSDSignalCore.ush
+++ b/Engine/Shaders/Private/ScreenSpaceDenoise/SSDSignalCore.ush
@@ -21,22 +21,31 @@
 	/** Internal buffer layout used for history rejection for the shadow penumbra. */
 	#define SIGNAL_BUFFER_LAYOUT_PENUMBRA_HISTORY_REJECTION 2
 
-	/** Buffer layout taken as input by */
-	#define SIGNAL_BUFFER_LAYOUT_PENUMBRA_OUTPUT 3
+	/** Buffer layout taken as input by */ // TODO: unused, remove it.
+	// #define SIGNAL_BUFFER_LAYOUT_PENUMBRA_OUTPUT 3
 
+	/** Buffer layout for the reflection given by the ray generation shader. */
+	#define SIGNAL_BUFFER_LAYOUT_REFLECTIONS_INPUT 3
 
-/** Defines how the signal is being processed. */
+	/** Buffer layout for the reflection output. */
+	#define SIGNAL_BUFFER_LAYOUT_REFLECTIONS_HISTORY 4
+
+
+/** Defines how the signal is being processed. Matches C++'s ESignalProcessing. */
 	/** Scene color anti aliasing processing. */
-	#define SIGNAL_PROCESSING_COLORAA 0
+	//#define SIGNAL_PROCESSING_COLORAA 0
 
 	/** Depth of field stabilization processing. */
-	#define SIGNAL_PROCESSING_DOF 1
+	//#define SIGNAL_PROCESSING_DOF 1
 
 	/** Shadow penumbra denoising. */
-	#define SIGNAL_PROCESSING_PENUMBRA 2
+	#define SIGNAL_PROCESSING_PENUMBRA 0
 
 	/** Shadow penumbra denoising. */
-	#define SIGNAL_PROCESSING_DIFUSE_GI 3
+	//#define SIGNAL_PROCESSING_DIFUSE_GI 3
+
+	/** Reflection denoising. */
+	#define SIGNAL_PROCESSING_REFLECTIONS 1
 
 
 //------------------------------------------------------- STRUCTURE
@@ -47,12 +56,9 @@ struct FSSFSignalSample
 	// Number of sample accumulated.
 	float AccumulatedSampleCount;
 
-	// Scene color and alpha for anti aliasing.
+	// Scene color and alpha.
 	float4 SceneColor;
 
-	// The circle of confusion for depth of field.
-	float CocRadius;
-
 	// Shadow denoising specific.
 	float LeavingRayCount;
 	float PenumbraClosestHit;
@@ -64,7 +70,6 @@ FSSFSignalSample CreateSignalSample()
 	FSSFSignalSample OutSample;
 	OutSample.AccumulatedSampleCount = 0;
 	OutSample.SceneColor = 0;
-	OutSample.CocRadius = 0;
 	OutSample.LeavingRayCount = 0;
 	OutSample.PenumbraClosestHit = 0;
 	OutSample.UncertainRayCount = 0;
@@ -109,6 +114,19 @@ FSSFSignalSample DecodeInputSignal(const uint SignalBufferLayout, float4 Samples
 			OutSample.LeavingRayCount = OutSample.AccumulatedSampleCount * Samples[0].r; // *(255.0 / 254.0);
 		}
 	}
+	else if (SignalBufferLayout == SIGNAL_BUFFER_LAYOUT_REFLECTIONS_INPUT)
+	{
+		OutSample.AccumulatedSampleCount = 1;
+		OutSample.LeavingRayCount = OutSample.AccumulatedSampleCount * (1 - Samples[0].a);
+		OutSample.SceneColor = Samples[0];
+		OutSample.PenumbraClosestHit = Samples[1].r;
+	}
+	else if (SignalBufferLayout == SIGNAL_BUFFER_LAYOUT_REFLECTIONS_HISTORY)
+	{
+		OutSample.AccumulatedSampleCount = Samples[1].r;
+		OutSample.LeavingRayCount = OutSample.AccumulatedSampleCount * (1 - Samples[0].a);
+		OutSample.SceneColor = OutSample.AccumulatedSampleCount * Samples[0];
+	}
 
 	return OutSample;
 }
@@ -155,9 +173,11 @@ void EncodeOutputSignal(const uint SignalBufferLayout, FSSFSignalSample Sample,
 
 		OutSamples[0].r = Sample.AccumulatedSampleCount == 0 ? -1 : Sample.LeavingRayCount / Sample.AccumulatedSampleCount;
 	}
-	else if (SignalBufferLayout == SIGNAL_BUFFER_LAYOUT_PENUMBRA_OUTPUT)
+	else if (SignalBufferLayout == SIGNAL_BUFFER_LAYOUT_REFLECTIONS_HISTORY)
 	{
-		OutSamples[0] = float4(1, 1, 1, 1) * Sample.LeavingRayCount;
+		OutSamples[0] = Sample.SceneColor * (Sample.AccumulatedSampleCount > 0 ? rcp(Sample.AccumulatedSampleCount) : 0);
+		OutSamples[1].r = Sample.AccumulatedSampleCount;
+		OutBufferCount = 2;
 	}
 }
 
@@ -190,7 +210,6 @@ FSSFSignalSample Mul(FSSFSignalSample Sample, float Scalar)
 {
 	FSSFSignalSample OutSample;
 	OutSample.SceneColor             = Sample.SceneColor * Scalar;
-	OutSample.CocRadius              = Sample.CocRadius * Scalar;
 	OutSample.AccumulatedSampleCount = Sample.AccumulatedSampleCount * Scalar;
 	OutSample.LeavingRayCount        = Sample.LeavingRayCount * Scalar;
 	OutSample.PenumbraClosestHit     = Sample.PenumbraClosestHit * Scalar;
@@ -201,7 +220,6 @@ FSSFSignalSample Mul(FSSFSignalSample Sample, float Scalar)
 FSSFSignalSample Add(FSSFSignalSample SampleA, FSSFSignalSample SampleB)
 {
 	SampleA.SceneColor             += SampleB.SceneColor;
-	SampleA.CocRadius              += SampleB.CocRadius;
 	SampleA.AccumulatedSampleCount += SampleB.AccumulatedSampleCount;
 	SampleA.LeavingRayCount        += SampleB.LeavingRayCount;
 	SampleA.PenumbraClosestHit     += SampleB.PenumbraClosestHit;
@@ -212,7 +230,6 @@ FSSFSignalSample Add(FSSFSignalSample SampleA, FSSFSignalSample SampleB)
 FSSFSignalSample Minus(FSSFSignalSample Sample)
 {
 	Sample.SceneColor             = -Sample.SceneColor;
-	Sample.CocRadius              = -Sample.CocRadius;
 	Sample.AccumulatedSampleCount = -Sample.AccumulatedSampleCount;
 	Sample.LeavingRayCount        = -Sample.LeavingRayCount;
 	Sample.PenumbraClosestHit     = -Sample.PenumbraClosestHit;
@@ -229,7 +246,6 @@ FSSFSignalSample Min(FSSFSignalSample SampleA, FSSFSignalSample SampleB)
 {
 	FSSFSignalSample OutSample;
 	OutSample.SceneColor = min(SampleA.SceneColor, SampleB.SceneColor);
-	OutSample.CocRadius = min(SampleA.CocRadius, SampleB.CocRadius);
 	OutSample.AccumulatedSampleCount = min(SampleA.AccumulatedSampleCount, SampleB.AccumulatedSampleCount);
 	OutSample.LeavingRayCount = min(SampleA.LeavingRayCount, SampleB.LeavingRayCount);
 	OutSample.PenumbraClosestHit = min(SampleA.PenumbraClosestHit, SampleB.PenumbraClosestHit);
@@ -241,7 +257,6 @@ FSSFSignalSample Max(FSSFSignalSample SampleA, FSSFSignalSample SampleB)
 {
 	FSSFSignalSample OutSample;
 	OutSample.SceneColor = max(SampleA.SceneColor, SampleB.SceneColor);
-	OutSample.CocRadius = max(SampleA.CocRadius, SampleB.CocRadius);
 	OutSample.AccumulatedSampleCount = max(SampleA.AccumulatedSampleCount, SampleB.AccumulatedSampleCount);
 	OutSample.LeavingRayCount = max(SampleA.LeavingRayCount, SampleB.LeavingRayCount);
 	OutSample.PenumbraClosestHit = max(SampleA.PenumbraClosestHit, SampleB.PenumbraClosestHit);
diff --git a/Engine/Shaders/Private/ScreenSpaceDenoise/SSDSignalInterface.ush b/Engine/Shaders/Private/ScreenSpaceDenoise/SSDSignalInterface.ush
--- a/Engine/Shaders/Private/ScreenSpaceDenoise/SSDSignalInterface.ush
+++ b/Engine/Shaders/Private/ScreenSpaceDenoise/SSDSignalInterface.ush
@@ -9,17 +9,11 @@
 #if !defined(CONFIG_SIGNAL_PROCESSING)
 	#error Missing CONFIG_SIGNAL_PROCESSING
 
-#elif CONFIG_SIGNAL_PROCESSING == SIGNAL_PROCESSING_COLORAA
-	#error Signal processing for color AA not implemented.
-
-#elif CONFIG_SIGNAL_PROCESSING == SIGNAL_PROCESSING_DOF
-	#error Signal processing for DOF not implemented.
-
 #elif CONFIG_SIGNAL_PROCESSING == SIGNAL_PROCESSING_PENUMBRA
 	#include "SSDPenumbraSignal.ush"
 
-#elif CONFIG_SIGNAL_PROCESSING == SIGNAL_PROCESSING_DIFUSE_GI
-	#error Signal processing for diffuse GI not implemented.
+#elif CONFIG_SIGNAL_PROCESSING == SIGNAL_PROCESSING_REFLECTIONS
+	#include "SSDReflectionsSignal.ush"
 
 #else
 	#error Unknown CONFIG_SIGNAL_PROCESSING
diff --git a/Engine/Shaders/Private/ScreenSpaceDenoise/SSDSpatialAccumulation.usf b/Engine/Shaders/Private/ScreenSpaceDenoise/SSDSpatialAccumulation.usf
--- a/Engine/Shaders/Private/ScreenSpaceDenoise/SSDSpatialAccumulation.usf
+++ b/Engine/Shaders/Private/ScreenSpaceDenoise/SSDSpatialAccumulation.usf
@@ -20,17 +20,40 @@ PostProcessTemporalAA.usf: Temporal AA
 #define TILE_PIXEL_SIZE 8
 
 
-#define CONFIG_SIGNAL_INPUT_LAYOUT \
-	SIGNAL_BUFFER_LAYOUT_PENUMBRA_INTERNAL
+#if DIM_SIGNAL_PROCESSING == SIGNAL_PROCESSING_PENUMBRA
+	#define CONFIG_SIGNAL_INPUT_LAYOUT \
+		SIGNAL_BUFFER_LAYOUT_PENUMBRA_INTERNAL
 
-#define CONFIG_SIGNAL_OUTPUT_LAYOUT \
-	SIGNAL_BUFFER_LAYOUT_PENUMBRA_HISTORY_REJECTION
+	#define CONFIG_SIGNAL_OUTPUT_LAYOUT \
+		SIGNAL_BUFFER_LAYOUT_PENUMBRA_HISTORY_REJECTION
 
-#define CONFIG_SIGNAL_PROCESSING \
-	SIGNAL_PROCESSING_PENUMBRA
+	#define CONFIG_SIGNAL_PROCESSING \
+		SIGNAL_PROCESSING_PENUMBRA
 
-#define CONFIG_SAMPLE_SET \
-	SAMPLE_SET_3X3
+	#define CONFIG_SAMPLE_SET \
+		SAMPLE_SET_3X3
+
+
+#elif DIM_SIGNAL_PROCESSING == SIGNAL_PROCESSING_REFLECTIONS
+	#define CONFIG_SIGNAL_INPUT_LAYOUT \
+		SIGNAL_BUFFER_LAYOUT_REFLECTIONS_INPUT
+
+	#define CONFIG_SIGNAL_OUTPUT_LAYOUT \
+		SIGNAL_BUFFER_LAYOUT_REFLECTIONS_OUTPUT
+
+	#define CONFIG_SIGNAL_PROCESSING \
+		SIGNAL_PROCESSING_REFLECTIONS
+
+	#define CONFIG_SAMPLE_SET \
+		SAMPLE_SET_3X3
+
+
+#else
+	#error Unknown signal processing.
+#endif
+
+
+//------------------------------------------------------- CONFIGS
 
 #include "SSDSignalInterface.ush"
 
@@ -122,6 +145,9 @@ void MainCS(
 			CurrentSceneMetadata,
 			NeighborSceneMetadata);
 
+		// Multiple importance sampling weights.
+		float MISWeight = GetMultiImportanceSamplingWeight(CurrentSceneMetadata, NeighborSceneMetadata, NeighborSample);
+
 		// Compute the weight of the sample in the kernel.
 		#if CONFIG_SAMPLE_SET == SAMPLE_SET_3X3
 			float KernelWeight = 1;
@@ -137,7 +163,7 @@ void MainCS(
 		#endif
 
 		// Compute the final weight of that sample.
-		float FinalWeight = BilateralWeight * KernelWeight;
+		float FinalWeight = BilateralWeight * MISWeight * KernelWeight;
 		
 		Accumulator = Add(Accumulator, Mul(NeighborSample, FinalWeight));
 	}
diff --git a/Engine/Shaders/Private/ScreenSpaceDenoise/SSDTemporalAccumulation.usf b/Engine/Shaders/Private/ScreenSpaceDenoise/SSDTemporalAccumulation.usf
--- a/Engine/Shaders/Private/ScreenSpaceDenoise/SSDTemporalAccumulation.usf
+++ b/Engine/Shaders/Private/ScreenSpaceDenoise/SSDTemporalAccumulation.usf
@@ -13,22 +13,61 @@
 
 #define TILE_PIXEL_SIZE 8
 
-// TODO: Do SIGNAL_BUFFER_LAYOUT_PENUMBRA_INTERNAL that does not store certainty for the history.
-#define CONFIG_SIGNAL_INPUT_LAYOUT\
-	SIGNAL_BUFFER_LAYOUT_PENUMBRA_INTERNAL
-#define CONFIG_SIGNAL_HISTORY_LAYOUT \
-	SIGNAL_BUFFER_LAYOUT_PENUMBRA_INTERNAL
-#define CONFIG_SIGNAL_HISTORY_REJECTION_LAYOUT \
-	SIGNAL_BUFFER_LAYOUT_PENUMBRA_HISTORY_REJECTION
 
-#define CONFIG_SIGNAL_PROCESSING \
-	SIGNAL_PROCESSING_PENUMBRA
+#define CONFIG_SIGNAL_PROCESSING DIM_SIGNAL_PROCESSING
+
+
+#if CONFIG_SIGNAL_PROCESSING == SIGNAL_PROCESSING_PENUMBRA
+	#define CONFIG_INPUT_TILE_CLASSIFICATION 1
+	#define CONFIG_USE_TILE_CLASSIFICATION 1
+	#define CONFIG_USE_REJECTION_BUFFER (DIM_IS_MIP_0 == 1)
+
+	// Uses a larger kernel when brute forcing to increase the speed of convergence by bluring more spatially.
+	#define CONFIG_LARGE_HISTORY_KERNEL 1
+
+	// TODO: Do SIGNAL_BUFFER_LAYOUT_PENUMBRA_INTERNAL that does not store certainty for the history.
+	#define CONFIG_SIGNAL_INPUT_LAYOUT   SIGNAL_BUFFER_LAYOUT_PENUMBRA_INTERNAL
+	#define CONFIG_SIGNAL_HISTORY_LAYOUT SIGNAL_BUFFER_LAYOUT_PENUMBRA_INTERNAL
+
+	#define CONFIG_SIGNAL_HISTORY_REJECTION_LAYOUT SIGNAL_BUFFER_LAYOUT_PENUMBRA_HISTORY_REJECTION
+
+#elif CONFIG_SIGNAL_PROCESSING == SIGNAL_PROCESSING_REFLECTIONS
+	#define CONFIG_SIGNAL_INPUT_LAYOUT   SIGNAL_BUFFER_LAYOUT_REFLECTIONS_INPUT
+	#define CONFIG_SIGNAL_HISTORY_LAYOUT SIGNAL_BUFFER_LAYOUT_REFLECTIONS_HISTORY
+
+	#define CONFIG_HISTORY_REJECTION 0
+
+#else
+	#error Unimplemented signal processing.
+#endif
 
 #include "SSDSignalInterface.ush"
 
 
 //------------------------------------------------------- CONFIG DISABLED DEFAULTS
 
+/** Whether this pass make use of a tile classification. */
+#ifndef CONFIG_USE_TILE_CLASSIFICATION
+	#define CONFIG_USE_TILE_CLASSIFICATION 0
+#endif
+
+/** Whether this pass uses custom rejection buffer. */
+#ifndef CONFIG_USE_REJECTION_BUFFER
+	#define CONFIG_USE_REJECTION_BUFFER 0
+	#define CONFIG_SIGNAL_HISTORY_REJECTION_LAYOUT CONFIG_SIGNAL_INPUT_LAYOUT
+#endif
+
+/** Whether this pass as a tile classified input. */
+#ifndef CONFIG_INPUT_TILE_CLASSIFICATION
+	#define CONFIG_INPUT_TILE_CLASSIFICATION 0
+#endif
+
+/** Uses a kernel that is larger to sample the history. */
+#ifndef CONFIG_LARGE_HISTORY_KERNEL
+	#define CONFIG_LARGE_HISTORY_KERNEL 0
+#endif
+
+
 //------------------------------------------------------- CONFIG ENABLED DEFAULTS
 
 /** Whether the history rejection should be taking into account the certainty. */
@@ -36,6 +75,10 @@
 	#define CONFIG_CERTAINTY_BASED_REJECTION 1
 #endif
 
+/** Whether should do history rejection. */
+#ifndef CONFIG_HISTORY_REJECTION
+	#define CONFIG_HISTORY_REJECTION 1
+#endif
 
 
 //------------------------------------------------------- MENDATORY CONFIG
@@ -171,7 +214,7 @@ void TemporallyAccumulate(
 				MipLevelPow2, rcp(MipLevelPow2));
 
 			// Number of sample to use.
-			const bool bUseLargerTAAKernel = bBruteForce && 1;
+			const bool bUseLargerTAAKernel = bBruteForce && (CONFIG_HISTORY_REJECTION != 0);
 			const uint SampleCount = bUseLargerTAAKernel ? 16 : 4;
 
 			// Disabled because it is useful to count the sample we add together when brute forcing,
@@ -235,6 +278,8 @@ void TemporallyAccumulate(
 
 				// Sample history.
 				FSSFSignalSample HistorySample = GetConstantSignalSample(TILE_MODE_FULLY_INVALID);
+
+				#if CONFIG_SIGNAL_PROCESSING == SIGNAL_PROCESSING_PENUMBRA
 				if (bBruteForce)
 				{
 					BRANCH
@@ -383,6 +428,16 @@ void TemporallyAccumulate(
 						HistorySample = Mul(GetConstantSignalSample(HistoryTileInfos.TileMode), MipLevelPow2 * MipLevelPow2);
 					}
 				} // if (!bBruteForce)
+				#else // CONFIG_SIGNAL_PROCESSING != SIGNAL_PROCESSING_PENUMBRA
+				{
+					HistorySample = SampleInputSignal(
+						PrevHistory0, PrevHistory1,
+						GlobalPointClampedSampler,
+						CONFIG_SIGNAL_HISTORY_LAYOUT,
+						SampleHistoryBufferUV,
+						/* MipLevel = */ 0);
+				}
+				#endif
 
 				if (HistorySample.AccumulatedSampleCount > 0)
 				{
@@ -436,28 +491,19 @@ void TemporallyAccumulate(
 
 	if (!bPostRejectionBlending)
 	{
-		// CurrentFrameWeight = 0 can be 0 in the case where no valid input was given, because shooting ray at grazing angles
-		// geometry, and may be discarded.
-		float CurrentFrameWeight = CurrentFrame.AccumulatedSampleCount;
-		float PreviousFrameWeight = min(HistorySample.AccumulatedSampleCount, TargetedSampleCount - CurrentFrameWeight);
-
-		float FinalLerpFactor = CurrentFrameWeight / (CurrentFrameWeight + PreviousFrameWeight);
+		float PreviousFrameWeight = min(HistorySample.AccumulatedSampleCount, TargetedSampleCount - CurrentFrame.AccumulatedSampleCount);
+		float PreviousFrameMultiplier = HistorySample.AccumulatedSampleCount > 0 ? PreviousFrameWeight / HistorySample.AccumulatedSampleCount : 0;
 
-		FSSFSignalSample Final = CreateSignalSample();
-		Final.AccumulatedSampleCount = CurrentFrameWeight + PreviousFrameWeight;
-		Final.LeavingRayCount = Final.AccumulatedSampleCount * lerp(
-			GetSamplePenumbraSafe(HistorySample),
-			GetSamplePenumbraSafe(CurrentFrame), FinalLerpFactor);
-
-		HistorySample = Final;
+		HistorySample = Add(Mul(PreviousFrame, PreviousFrameMultiplier), CurrentFrame);
 	}
 
 	// History rejection.
+	#if CONFIG_HISTORY_REJECTION
 	BRANCH
 	if (HistorySample.AccumulatedSampleCount > 0 && 1 || 1)
 	{
 		// Whether the history is rejected using rejection history signal.
-		const bool bUsesHistoryRejectionSignal = DIM_IS_MIP_0;
+		const bool bUsesHistoryRejectionSignal = CONFIG_USE_REJECTION_BUFFER;
 
 		// Compile time size of the neighborhood to look for.
 		int NeighborSize = 2; // TODO: based on anisotropy.
@@ -500,8 +546,13 @@ void TemporallyAccumulate(
 				}
 			}
 
-			// TODO: unless being lowest mip level, can reduce this to a gather4.
-			FSSFTileInfos SampleTileInfos = LoadTileInfos(uint2(floor(SampleBufferUV * View.BufferSizeAndInvSize.xy)) / 8);
+			FSSFTileInfos SampleTileInfos;
+			#if CONFIG_INPUT_TILE_CLASSIFICATION
+				// TODO: unless being lowest mip level, can reduce this to a gather4.
+				SampleTileInfos = LoadTileInfos(uint2(floor(SampleBufferUV * View.BufferSizeAndInvSize.xy)) / 8);
+			#else // !CONFIG_INPUT_TILE_CLASSIFICATION
+				SampleTileInfos = CreateTileInfos();
+			#endif
 
 			float NeighborPenumbra = -1;
 			float NeighborUncertainty = 1.0;
@@ -538,7 +589,7 @@ void TemporallyAccumulate(
 				}
 
 				FLATTEN
-				if (bUseSample)
+				if (bUseSample || !CONFIG_INPUT_TILE_CLASSIFICATION)
 				{
 					float2 SampleViewportUV = BufferUVToViewportUV(SampleBufferUV);
 					float2 SampleScreenPosition = ViewportUVToScreenPos(SampleViewportUV);
@@ -601,24 +652,15 @@ void TemporallyAccumulate(
 		HistorySample.AccumulatedSampleCount -= RejectedSampleCount;
 		HistorySample.LeavingRayCount = HistorySample.AccumulatedSampleCount * NewHistoryPenumbra;
 	}
+	#endif // CONFIG_HISTORY_REJECTION
 	
 	// Evaluate the final sample to output.
 	if (bPostRejectionBlending)
 	{
-		// CurrentFrameWeight = 0 can be 0 in the case where no valid input was given, because shooting ray at grazing angles
-		// geometry, and may be discarded.
-		float CurrentFrameWeight = CurrentFrame.AccumulatedSampleCount;
-		float PreviousFrameWeight = min(HistorySample.AccumulatedSampleCount, TargetedSampleCount - CurrentFrameWeight);
-
-		float FinalLerpFactor = CurrentFrameWeight / (CurrentFrameWeight + PreviousFrameWeight);
+		float PreviousFrameWeight = min(HistorySample.AccumulatedSampleCount, TargetedSampleCount - CurrentFrame.AccumulatedSampleCount);
+		float PreviousFrameMultiplier = HistorySample.AccumulatedSampleCount > 0 ? PreviousFrameWeight / HistorySample.AccumulatedSampleCount : 0;
 
-		FSSFSignalSample Final = CreateSignalSample();
-		Final.AccumulatedSampleCount = CurrentFrameWeight + PreviousFrameWeight;
-		Final.LeavingRayCount = Final.AccumulatedSampleCount * lerp(
-			GetSamplePenumbraSafe(HistorySample),
-			GetSamplePenumbraSafe(CurrentFrame), FinalLerpFactor);
-
-		HistorySample = Final;
+		HistorySample = Add(Mul(PreviousFrame, PreviousFrameMultiplier), CurrentFrame);
 	}
 
 	if (all(DispatchThreadId < uint2(View.ViewSizeAndInvSize.xy)))
@@ -636,6 +678,7 @@ void MainCS(
 	uint2 GroupId : SV_GroupID,
 	uint2 GroupThreadId : SV_GroupThreadID,
 	uint GroupThreadIndex : SV_GroupIndex)
+#if CONFIG_USE_TILE_CLASSIFICATION
 {
 	// TODO: tile classification.
 	FSSFTileInfos TileInfos = LoadTileInfos((iMipLevelPow2 * DispatchThreadId) / TILE_PIXEL_SIZE);
@@ -684,3 +727,12 @@ void MainCS(
 		#endif
 	}
 }
+#else // !CONFIG_USE_TILE_CLASSIFICATION
+{
+	FSSFTileInfos TileInfos = CreateTileInfos();
+	TemporallyAccumulate(
+		DispatchThreadId, GroupId, GroupThreadId, GroupThreadIndex,
+		TileInfos,
+		/* bBruteForce = */ true);
+}
+#endif // !CONFIG_USE_TILE_CLASSIFICATION
diff --git a/Engine/Source/Runtime/Renderer/Private/SceneRendering.h b/Engine/Source/Runtime/Renderer/Private/SceneRendering.h
--- a/Engine/Source/Runtime/Renderer/Private/SceneRendering.h
+++ b/Engine/Source/Runtime/Renderer/Private/SceneRendering.h
@@ -858,8 +858,8 @@ struct FPreviousViewInfo
 	// input post process material.
 	TRefCountPtr<IPooledRenderTarget> CustomSSRInput;
 
-	// History scene color
-	FScreenSpaceFilteringHistory RayTracedGIHistory;
+	// History for the reflections
+	FScreenSpaceFilteringHistory ReflectionsHistory;
 
 	// History for shadow denoising.
 	TMap<const FLightSceneInfo*, FScreenSpaceFilteringHistory> ShadowHistories;
@@ -874,7 +874,7 @@ struct FPreviousViewInfo
 		DOFPostGatherForegroundHistory.SafeRelease();
 		DOFPostGatherBackgroundHistory.SafeRelease();
 		CustomSSRInput.SafeRelease();
-		RayTracedGIHistory.SafeRelease();
+		ReflectionsHistory.SafeRelease();
 		ShadowHistories.Reset();
 	}
 };
diff --git a/Engine/Source/Runtime/Renderer/Private/ScreenSpaceDenoise.cpp b/Engine/Source/Runtime/Renderer/Private/ScreenSpaceDenoise.cpp
--- a/Engine/Source/Runtime/Renderer/Private/ScreenSpaceDenoise.cpp
+++ b/Engine/Source/Runtime/Renderer/Private/ScreenSpaceDenoise.cpp
@@ -57,6 +57,20 @@ static bool IsSupportedLightType(ELightComponentType LightType)
 namespace
 {
 
+/** Different signals to denoise. */
+enum class ESignalProcessing
+{
+	Penumbra,
+	Reflections,
+
+	MAX,
+};
+
+
+class FSignalProcessingDim : SHADER_PERMUTATION_ENUM_CLASS("DIM_SIGNAL_PROCESSING", ESignalProcessing);
+
+
+/** Base class for a screen space denoising shader. */
 class FScreenSpaceDenoisingShader : public FGlobalShader
 {
 public:
@@ -176,7 +190,23 @@ class FSSDTemporalAccumulationCS : public FScreenSpaceDenoisingShader
 
 	class FIsMip0Dim : SHADER_PERMUTATION_BOOL("DIM_IS_MIP_0");
 
-	using FPermutationDomain = TShaderPermutationDomain<FIsMip0Dim>;
+	using FPermutationDomain = TShaderPermutationDomain<
+		FIsMip0Dim,
+		FSignalProcessingDim>;
+
+	static bool ShouldCompilePermutation(const FGlobalShaderPermutationParameters& Parameters)
+	{
+		FPermutationDomain PermutationVector(Parameters.PermutationId);
+
+		// Reflections are only processed at mip 0.
+		if (PermutationVector.Get<FSignalProcessingDim>() == ESignalProcessing::Reflections && 
+			!PermutationVector.Get<FIsMip0Dim>())
+		{
+			return false;
+		}
+
+		return FScreenSpaceDenoisingShader::ShouldCompilePermutation(Parameters);
+	}
 
 	BEGIN_SHADER_PARAMETER_STRUCT(FParameters, )
 		SHADER_PARAMETER(uint32, iMipLevel)
@@ -443,6 +473,7 @@ static void DenoiseShadowPenumbra(
 		{
 			FSSDTemporalAccumulationCS::FPermutationDomain PermutationVector;
 			PermutationVector.Set<FSSDTemporalAccumulationCS::FIsMip0Dim>(MipLevel == 0);
+			PermutationVector.Set<FSignalProcessingDim>(ESignalProcessing::Penumbra);
 
 			TShaderMapRef<FSSDTemporalAccumulationCS> ComputeShader(View.ShaderMap, PermutationVector);
 
@@ -536,7 +567,98 @@ static void DenoiseShadowPenumbra(
 	#endif // SSD_DEBUG_PASS
 
 	*OutPenumbraMask = OutputSignal;
-}
+} // DenoiseShadowPenumbra()
+
+
+IScreenSpaceDenoiser::FReflectionOutputs DenoiseReflections(
+	FRDGBuilder& GraphBuilder,
+	const FViewInfo& View,
+	const FSceneViewFamilyBlackboard& SceneBlackboard,
+	const IScreenSpaceDenoiser::FReflectionInputs& ReflectionInputs)
+{
+	const FIntPoint DenoiseResolution = View.ViewRect.Size();
+
+	// Descriptor to allocate internal denoising buffer.
+	FRDGTextureDesc SignalProcessingDesc = ReflectionInputs.Color->Desc;
+	SignalProcessingDesc.Format = PF_FloatRGBA;
+	SignalProcessingDesc.NumMips = 1;
+
+	// Setup common shader parameters.
+	FSSDCommonParameters CommonParameters;
+	{
+		CommonParameters.SceneBlackboard = SceneBlackboard;
+		CommonParameters.ViewUniformBuffer = View.ViewUniformBuffer;
+	}
+
+	FScreenSpaceFilteringHistory PrevFrameHistory;
+	if (!View.bCameraCut)
+	{
+		PrevFrameHistory = View.PrevViewInfo.ReflectionsHistory;
+	}
+
+	FRDGTextureRef SignalHistory0 = GraphBuilder.CreateTexture(SignalProcessingDesc, TEXT("SSDReflectionsHistory0"));
+	FRDGTextureRef SignalHistory1 = GraphBuilder.CreateTexture(SignalProcessingDesc, TEXT("SSDReflectionsHistory1"));
+
+	{
+		const int32 MipLevel = 0;
+
+		FSSDTemporalAccumulationCS::FPermutationDomain PermutationVector;
+		PermutationVector.Set<FSSDTemporalAccumulationCS::FIsMip0Dim>(MipLevel == 0);
+		PermutationVector.Set<FSignalProcessingDim>(ESignalProcessing::Reflections);
+
+		TShaderMapRef<FSSDTemporalAccumulationCS> ComputeShader(View.ShaderMap, PermutationVector);
+
+		FSSDTemporalAccumulationCS::FParameters* PassParameters = GraphBuilder.AllocParameters<FSSDTemporalAccumulationCS::FParameters>();
+		PassParameters->CommonParameters = CommonParameters;
+		PassParameters->iMipLevel = MipLevel;
+		PassParameters->iMipLevelPow2 = (1 << MipLevel);
+		PassParameters->MipLevel = MipLevel;
+		PassParameters->MipLevelPow2 = (1 << MipLevel);
+		PassParameters->InvMipLevelPow2 = 1.0f / PassParameters->MipLevelPow2;
+		PassParameters->bCameraCut = View.bCameraCut || !PrevFrameHistory.RT[0].IsValid();
+		PassParameters->SignalInput0 = ReflectionInputs.Color;
+		PassParameters->SignalInput1 = ReflectionInputs.RayHitDistance;
+
+		PassParameters->PrevScreenToTranslatedWorld = View.PrevViewInfo.ViewMatrices.GetInvTranslatedViewProjectionMatrix();
+		PassParameters->PrevHistory0 = RegisterExternalTextureWithFallback(GraphBuilder, PrevFrameHistory.RT[0], GSystemTextures.BlackDummy);
+		PassParameters->PrevHistory1 = RegisterExternalTextureWithFallback(GraphBuilder, PrevFrameHistory.RT[1], GSystemTextures.BlackDummy);
+		PassParameters->PrevTileClassificationTexture = RegisterExternalTextureWithFallback(GraphBuilder, PrevFrameHistory.TileClassification, GSystemTextures.BlackDummy);
+		PassParameters->PrevDepthBuffer = RegisterExternalTextureWithFallback(GraphBuilder, View.PrevViewInfo.DepthBuffer, GSystemTextures.BlackDummy);
+		PassParameters->PrevGBufferA = RegisterExternalTextureWithFallback(GraphBuilder, View.PrevViewInfo.GBufferA, GSystemTextures.BlackDummy);
+
+		PassParameters->SignalHistoryOutput0 = GraphBuilder.CreateUAV(FRDGTextureUAVDesc(SignalHistory0, MipLevel));
+		PassParameters->SignalHistoryOutput1 = GraphBuilder.CreateUAV(FRDGTextureUAVDesc(SignalHistory1, MipLevel));
+
+		FIntPoint Resolution = FIntPoint::DivideAndRoundUp(DenoiseResolution, 1 << MipLevel);
+
+		FComputeShaderUtils::AddPass(
+			GraphBuilder,
+			RDG_EVENT_NAME("SSD TemporalAccumulation(Mip=%d)", MipLevel),
+			*ComputeShader,
+			PassParameters,
+			FComputeShaderUtils::GetGroupCount(Resolution, FComputeShaderUtils::kGolden2DGroupSize));
+	}
+
+	if (View.ViewState)
+	{
+		// Keep depth buffer and GBufferA around for next frame.
+		{
+			GraphBuilder.QueueTextureExtraction(SceneBlackboard.SceneDepthBuffer, &View.ViewState->PendingPrevFrameViewInfo.DepthBuffer);
+			GraphBuilder.QueueTextureExtraction(SceneBlackboard.SceneGBufferA, &View.ViewState->PendingPrevFrameViewInfo.GBufferA);
+		}
+
+		// Saves history.
+		{
+			FScreenSpaceFilteringHistory& NewHistory = View.ViewState->PendingPrevFrameViewInfo.ReflectionsHistory;
+			GraphBuilder.QueueTextureExtraction(SignalHistory0, &NewHistory.RT[0]);
+			GraphBuilder.QueueTextureExtraction(SignalHistory1, &NewHistory.RT[1]);
+		}
+	}
+
+	IScreenSpaceDenoiser::FReflectionOutputs Outputs;
+	Outputs.Color = SignalHistory0; // TODO.
+	return Outputs;
+} // DenoiseReflections()
 
 
 /** The implementation of the default denoiser of the renderer. */
@@ -606,12 +728,10 @@ public:
 		const FSceneViewFamilyBlackboard& SceneBlackboard,
 		const FReflectionInputs& ReflectionInputs) const override
 	{
-		FReflectionOutputs Outputs;
-		Outputs.Color = ReflectionInputs.Color; // TODO.
-		return Outputs;
+		return ::DenoiseReflections(GraphBuilder, View, SceneBlackboard, ReflectionInputs);
 	}
 
-}; // FDefaultScreenSpaceDenoiser
+}; // class FDefaultScreenSpaceDenoiser
 
 
 // static
