Multiple importance sampling is incorrectly being invoked when hitting analytical light sources.

diff --git a/Engine/Shaders/Private/PathTracing/Light/PathTracingLightCommon.ush b/Engine/Shaders/Private/PathTracing/Light/PathTracingLightCommon.ush
--- a/Engine/Shaders/Private/PathTracing/Light/PathTracingLightCommon.ush
+++ b/Engine/Shaders/Private/PathTracing/Light/PathTracingLightCommon.ush
@@ -31,6 +31,12 @@ bool IsRectLight(int LightId)
 	return SceneLightsData.Type[LightId] == 3;
 }
 
+// A light is a physical light if it can be intersected by a ray.
+bool IsPhysicalLight(int LightId)
+{
+	return IsEnvironmentLight(LightId);
+}
+
 float3 GetLightUV(int LightId, RayDesc Ray, FMaterialClosestHitPayload RayHitInfo)
 {
 	float3 LightUV = 0.0;
diff --git a/Engine/Shaders/Private/PathTracing/PathTracing.usf b/Engine/Shaders/Private/PathTracing/PathTracing.usf
--- a/Engine/Shaders/Private/PathTracing/PathTracing.usf
+++ b/Engine/Shaders/Private/PathTracing/PathTracing.usf
@@ -220,7 +220,7 @@ void PathTracingMainRG()
 
 						// Apply material Pdf for correct MIS weight
 						float MisWeight = 1.0;
-						if (IsUnidirectionalEnabled)
+						if (IsUnidirectionalEnabled && IsPhysicalLight(LightId))
 						{
 #if 0
 							// DEBUG: Validate Material PDF
