Fixing energy discrepancy with RectLight when path tracing. Converting sampling PDF to solid angle measure.

diff --git a/Engine/Shaders/Private/PathTracing/Light/PathTracingRectLight.ush b/Engine/Shaders/Private/PathTracing/Light/PathTracingRectLight.ush
--- a/Engine/Shaders/Private/PathTracing/Light/PathTracingRectLight.ush
+++ b/Engine/Shaders/Private/PathTracing/Light/PathTracingRectLight.ush
@@ -44,14 +44,7 @@ void RectLight_EvalLight(
 )
 {
 	float3 Color = GetColor(LightId);
-	float2 DeltaP = float2(GetWidth(LightId), GetHeight(LightId)) * LightUV.xy;
-	float3 LightDir = GetPosition(LightId) + DeltaP.x * GetdPdu(LightId) + DeltaP.y * GetdPdv(LightId) - Ray.Origin;
-	float CosTheta = saturate(dot(normalize(-LightDir), GetNormal(LightId)));
-
-	float DistanceSquared = dot(LightDir, LightDir);
-	float Area = GetWidth(LightId) * GetHeight(LightId);
-	//Color /= Area;
-	OutRadiance = Color * CosTheta / DistanceSquared;
+	OutRadiance = Color;
 }
 
 void RectLight_SampleLight(
@@ -65,7 +58,15 @@ void RectLight_SampleLight(
 {
 	// Pick a random position on the light
 	float2 UV = RandSample.yz - 0.5;
-	float Area = GetWidth(LightId) * GetHeight(LightId);
 	OutLightUV = float3(UV, 0);
-	OutPdf = 1.0 / Area;
+
+	// Convert PDF to solid angle measure
+	float Area = GetWidth(LightId) * GetHeight(LightId);
+	float2 DeltaP = float2(GetWidth(LightId), GetHeight(LightId)) * OutLightUV.xy;
+	float3 LightDir = GetPosition(LightId) + DeltaP.x * GetdPdu(LightId) + DeltaP.y * GetdPdv(LightId) - Payload.WorldPos;
+	float Distance2 = dot(LightDir, LightDir);
+	LightDir /= length(LightDir);
+	float CosTheta = saturate(dot(normalize(-LightDir), GetNormal(LightId)));
+	float SolidAngle = Area * CosTheta / Distance2;
+	OutPdf = 1.0 / SolidAngle;
 }
\ No newline at end of file
diff --git a/Engine/Source/Runtime/Renderer/Private/PathTracing.cpp b/Engine/Source/Runtime/Renderer/Private/PathTracing.cpp
--- a/Engine/Source/Runtime/Renderer/Private/PathTracing.cpp
+++ b/Engine/Source/Runtime/Renderer/Private/PathTracing.cpp
@@ -172,7 +172,7 @@ public:
 						LightData.Normal[LightData.Count] = -LightParameters.Direction;
 						LightData.dPdu[LightData.Count] = FVector::CrossProduct(LightParameters.Tangent, LightParameters.Direction);
 						LightData.dPdv[LightData.Count] = LightParameters.Tangent;
-						LightData.Color[LightData.Count] = LightParameters.Color * (2.0f * LightParameters.SourceRadius * LightParameters.SourceLength);
+						LightData.Color[LightData.Count] = LightParameters.Color;
 						LightData.Dimensions[LightData.Count] = FVector(2.0f * LightParameters.SourceRadius, 2.0f * LightParameters.SourceLength, 0.0f);
 						break;
 					}
