Implements variance based history rejection for denoiser's temporal accumulation.

diff --git a/Engine/Shaders/Private/PostProcessTemporalAA.usf b/Engine/Shaders/Private/PostProcessTemporalAA.usf
--- a/Engine/Shaders/Private/PostProcessTemporalAA.usf
+++ b/Engine/Shaders/Private/PostProcessTemporalAA.usf
@@ -1650,9 +1650,9 @@ void ComputeNeighborhoodBoundingbox(
 		m1 *= (1.0 / AA_SAMPLES);
 		m2 *= (1.0 / AA_SAMPLES);
 
-		float4 Variance = sqrt( abs(m2 - m1 * m1) );
-		NeighborMin = m1 - 1.25 * Variance;
-		NeighborMax = m1 + 1.25 * Variance;
+		float4 StdDev = sqrt( abs(m2 - m1 * m1) );
+		NeighborMin = m1 - 1.25 * StdDev;
+		NeighborMax = m1 + 1.25 * StdDev;
 
 		NeighborMin = min( NeighborMin, IntermediaryResult.Filtered );
 		NeighborMax = max( NeighborMax, IntermediaryResult.Filtered );
diff --git a/Engine/Shaders/Private/ScreenSpaceDenoise/SSDSignalCore.ush b/Engine/Shaders/Private/ScreenSpaceDenoise/SSDSignalCore.ush
--- a/Engine/Shaders/Private/ScreenSpaceDenoise/SSDSignalCore.ush
+++ b/Engine/Shaders/Private/ScreenSpaceDenoise/SSDSignalCore.ush
@@ -65,17 +65,22 @@ struct FSSFSignalSample
 	float UncertainRayCount;
 };
 
-FSSFSignalSample CreateSignalSample()
+FSSFSignalSample CreateSignalSampleFromScalarValue(float Scalar)
 {
 	FSSFSignalSample OutSample;
-	OutSample.AccumulatedSampleCount = 0;
-	OutSample.SceneColor = 0;
-	OutSample.LeavingRayCount = 0;
-	OutSample.PenumbraClosestHit = 0;
-	OutSample.UncertainRayCount = 0;
+	OutSample.SceneColor = Scalar;
+	OutSample.AccumulatedSampleCount = Scalar;
+	OutSample.LeavingRayCount = Scalar;
+	OutSample.PenumbraClosestHit = Scalar;
+	OutSample.UncertainRayCount = Scalar;
 	return OutSample;
 }
 
+FSSFSignalSample CreateSignalSample()
+{
+	return CreateSignalSampleFromScalarValue(0.0);
+}
+
 
 //------------------------------------------------------- INPUT/OUTPUT
 
@@ -237,6 +242,44 @@ FSSFSignalSample Minus(FSSFSignalSample Sample)
 	return Sample;
 }
 
+FSSFSignalSample Subtract(FSSFSignalSample SampleA, FSSFSignalSample SampleB)
+{
+	return Add(SampleA, Minus(SampleB));
+}
+
+FSSFSignalSample Abs(FSSFSignalSample Sample)
+{
+	Sample.SceneColor             = abs(Sample.SceneColor);
+	Sample.AccumulatedSampleCount = abs(Sample.AccumulatedSampleCount);
+	Sample.LeavingRayCount        = abs(Sample.LeavingRayCount);
+	Sample.PenumbraClosestHit     = abs(Sample.PenumbraClosestHit);
+	Sample.UncertainRayCount      = abs(Sample.UncertainRayCount);
+	return Sample;
+}
+
+FSSFSignalSample Sqrt(FSSFSignalSample Sample)
+{
+	Sample.SceneColor             = sqrt(Sample.SceneColor);
+	Sample.AccumulatedSampleCount = sqrt(Sample.AccumulatedSampleCount);
+	Sample.LeavingRayCount        = sqrt(Sample.LeavingRayCount);
+	Sample.PenumbraClosestHit     = sqrt(Sample.PenumbraClosestHit);
+	Sample.UncertainRayCount      = sqrt(Sample.UncertainRayCount);
+	return Sample;
+}
+
+FSSFSignalSample Power(FSSFSignalSample Sample, float Exponent)
+{
+	Sample.SceneColor.r           = pow(Sample.SceneColor.r, Exponent);
+	Sample.SceneColor.g           = pow(Sample.SceneColor.g, Exponent);
+	Sample.SceneColor.b           = pow(Sample.SceneColor.b, Exponent);
+	Sample.SceneColor.a           = pow(Sample.SceneColor.a, Exponent);
+	Sample.AccumulatedSampleCount = pow(Sample.AccumulatedSampleCount, Exponent);
+	Sample.LeavingRayCount        = pow(Sample.LeavingRayCount, Exponent);
+	Sample.PenumbraClosestHit     = pow(Sample.PenumbraClosestHit, Exponent);
+	Sample.UncertainRayCount      = pow(Sample.UncertainRayCount, Exponent);
+	return Sample;
+}
+
 FSSFSignalSample Suctract(FSSFSignalSample SampleA, FSSFSignalSample SampleB)
 {
 	return Add(SampleA, Minus(SampleB));
@@ -277,8 +320,7 @@ FSSFSignalSample Lerp(FSSFSignalSample Sample0, FSSFSignalSample Sample1, float
 /** Normalize the signal sample to one. */
 FSSFSignalSample NormalizeToOneSample(FSSFSignalSample Sample)
 {
-	float NormalizationFactor = -min(-rcp(Sample.AccumulatedSampleCount), 0);
-
-	FSSFSignalSample OutSample = Mul(Sample, NormalizationFactor);
+	FSSFSignalSample OutSample = Mul(Sample, Sample.AccumulatedSampleCount > 0 ? rcp(Sample.AccumulatedSampleCount) : 0);
+	OutSample.AccumulatedSampleCount = 1;
 	return OutSample;
 }
diff --git a/Engine/Shaders/Private/ScreenSpaceDenoise/SSDTemporalAccumulation.usf b/Engine/Shaders/Private/ScreenSpaceDenoise/SSDTemporalAccumulation.usf
--- a/Engine/Shaders/Private/ScreenSpaceDenoise/SSDTemporalAccumulation.usf
+++ b/Engine/Shaders/Private/ScreenSpaceDenoise/SSDTemporalAccumulation.usf
@@ -25,6 +25,8 @@
 	// Uses a larger kernel when brute forcing to increase the speed of convergence by bluring more spatially.
 	#define CONFIG_LARGE_HISTORY_KERNEL 1
 
+	#define CONFIG_CERTAINTY_BASED_REJECTION 1
+
 	// TODO: Do SIGNAL_BUFFER_LAYOUT_PENUMBRA_INTERNAL that does not store certainty for the history.
 	#define CONFIG_SIGNAL_INPUT_LAYOUT   SIGNAL_BUFFER_LAYOUT_PENUMBRA_INTERNAL
 	#define CONFIG_SIGNAL_HISTORY_LAYOUT SIGNAL_BUFFER_LAYOUT_PENUMBRA_INTERNAL
@@ -35,7 +37,7 @@
 	#define CONFIG_SIGNAL_INPUT_LAYOUT   SIGNAL_BUFFER_LAYOUT_REFLECTIONS_INPUT
 	#define CONFIG_SIGNAL_HISTORY_LAYOUT SIGNAL_BUFFER_LAYOUT_REFLECTIONS_HISTORY
 
-	#define CONFIG_HISTORY_REJECTION 0
+	#define CONFIG_HISTORY_REJECTION 2
 
 #else
 	#error Unimplemented signal processing.
@@ -67,14 +69,14 @@
 	#define CONFIG_LARGE_HISTORY_KERNEL 0
 #endif
 
-
-//------------------------------------------------------- CONFIG ENABLED DEFAULTS
-
 /** Whether the history rejection should be taking into account the certainty. */
 #ifndef CONFIG_CERTAINTY_BASED_REJECTION
-	#define CONFIG_CERTAINTY_BASED_REJECTION 1
+	#define CONFIG_CERTAINTY_BASED_REJECTION 0
 #endif
 
+
+//------------------------------------------------------- CONFIG ENABLED DEFAULTS
+
 /** Whether should do history rejection. */
 #ifndef CONFIG_HISTORY_REJECTION
 	#define CONFIG_HISTORY_REJECTION 1
@@ -498,7 +500,7 @@ void TemporallyAccumulate(
 	}
 
 	// History rejection.
-	#if CONFIG_HISTORY_REJECTION
+	#if CONFIG_HISTORY_REJECTION > 0
 	BRANCH
 	if (HistorySample.AccumulatedSampleCount > 0 && 1 || 1)
 	{
@@ -507,17 +509,32 @@ void TemporallyAccumulate(
 
 		// Compile time size of the neighborhood to look for.
 		int NeighborSize = 2; // TODO: based on anisotropy.
-		if (bBruteForce)
+		#if CONFIG_SIGNAL_PROCESSING == SIGNAL_PROCESSING_PENUMBRA
 		{
-			// Brute forcing, want to have very large clamping box.
-			NeighborSize = bUsesHistoryRejectionSignal ? 3 : 4;
+			if (bBruteForce)
+			{
+				// Brute forcing, want to have very large clamping box.
+				NeighborSize = bUsesHistoryRejectionSignal ? 3 : 4;
+			}
+			else if (bUsesHistoryRejectionSignal)
+			{
+				NeighborSize = 2;
+			}
 		}
-		else if (bUsesHistoryRejectionSignal)
+		#else
 		{
-			NeighborSize = 2;
+			NeighborSize = 3;
 		}
+		#endif
 
 		// The minimum and maximum penumbra.
+		FSSFSignalSample NeighborMin = CreateSignalSampleFromScalarValue(1.#INF);
+		FSSFSignalSample NeighborMax = CreateSignalSampleFromScalarValue(0.0);
+
+		float NeighborMomentWeight = 0;
+		FSSFSignalSample NeighborMoment1 = CreateSignalSampleFromScalarValue(0.0);
+		FSSFSignalSample NeighborMoment2 = CreateSignalSampleFromScalarValue(0.0);
+
 		float NeighborMinPenumbra = 1.0;
 		float NeighborMaxPenumbra = 0.0;
 
@@ -526,6 +543,13 @@ void TemporallyAccumulate(
 		{
 			NeighborMinPenumbra = GetSamplePenumbraSafe(CurrentFrame);
 			NeighborMaxPenumbra = GetSamplePenumbraSafe(CurrentFrame);
+
+			NeighborMin = NormalizeToOneSample(CurrentFrame);
+			NeighborMax = NormalizeToOneSample(CurrentFrame);
+
+			NeighborMomentWeight = 1;
+			NeighborMoment1 = NormalizeToOneSample(CurrentFrame);
+			NeighborMoment2 = Power(NeighborMoment1, 2.0);
 		}
 
 		// Look for neighborhood to build clamping box.
@@ -546,80 +570,116 @@ void TemporallyAccumulate(
 				}
 			}
 
+			FSSFSignalSample NormalizedNeighborSample = CreateSignalSample();
+			float NeighborPenumbra = -1;
+			float NeighborUncertainty = 1.0;
+
 			FSSFTileInfos SampleTileInfos;
 			#if CONFIG_INPUT_TILE_CLASSIFICATION
+			{
 				// TODO: unless being lowest mip level, can reduce this to a gather4.
 				SampleTileInfos = LoadTileInfos(uint2(floor(SampleBufferUV * View.BufferSizeAndInvSize.xy)) / 8);
-			#else // !CONFIG_INPUT_TILE_CLASSIFICATION
-				SampleTileInfos = CreateTileInfos();
-			#endif
-
-			float NeighborPenumbra = -1;
-			float NeighborUncertainty = 1.0;
 
-			// Already know this sample will be invalid.
-			if (SampleTileInfos.TileMode == TILE_MODE_FULLY_INVALID && 0) // TODO: enable
-			{
-				continue;
-			}
-			else if (SampleTileInfos.TileMode == TILE_MODE_FULLY_OCCLUDED && 0)
-			{
-				NeighborPenumbra = 0;
-				NeighborUncertainty = 0.0;
-			}
-			else if (SampleTileInfos.TileMode == TILE_MODE_FULLY_LIGHTED && 0)
-			{
-				NeighborPenumbra = 1;
-				NeighborUncertainty = 0.0;
-			}
-			else
-			{
-				bool bUseSample;
-				if (bBruteForce)
+				// Already know this sample will be invalid.
+				if (SampleTileInfos.TileMode == TILE_MODE_FULLY_INVALID && 0) // TODO: enable
+				{
+					continue;
+				}
+				else if (SampleTileInfos.TileMode == TILE_MODE_FULLY_OCCLUDED && 0)
 				{
-					bUseSample = SampleTileInfos.TileMode == TILE_MODE_BRUTE_FORCE;
+					NeighborPenumbra = 0;
+					NeighborUncertainty = 0.0;
 				}
-				else if (DIM_IS_MIP_0)
+				else if (SampleTileInfos.TileMode == TILE_MODE_FULLY_LIGHTED && 0)
 				{
-					bUseSample = SampleTileInfos.TileMode != TILE_MODE_BRUTE_FORCE && SampleTileInfos.bProcessMipLevel[0];
+					NeighborPenumbra = 1;
+					NeighborUncertainty = 0.0;
 				}
 				else
 				{
-					bUseSample = SampleTileInfos.TileMode != TILE_MODE_BRUTE_FORCE && SampleTileInfos.bProcessMipLevel[iMipLevel];
+					bool bUseSample;
+					if (bBruteForce)
+					{
+						bUseSample = SampleTileInfos.TileMode == TILE_MODE_BRUTE_FORCE;
+					}
+					else if (DIM_IS_MIP_0)
+					{
+						bUseSample = SampleTileInfos.TileMode != TILE_MODE_BRUTE_FORCE && SampleTileInfos.bProcessMipLevel[0];
+					}
+					else
+					{
+						bUseSample = SampleTileInfos.TileMode != TILE_MODE_BRUTE_FORCE && SampleTileInfos.bProcessMipLevel[iMipLevel];
+					}
+
+					FLATTEN
+					if (bUseSample)
+					{
+						float2 SampleViewportUV = BufferUVToViewportUV(SampleBufferUV);
+						float2 SampleScreenPosition = ViewportUVToScreenPos(SampleViewportUV);
+
+						// Sample the neighor for history rejection.
+						FSSFSignalSample Neighbor = SampleInputSignal(
+							bUsesHistoryRejectionSignal ? HistoryRejectionSignal0 : SignalInput0,
+							bUsesHistoryRejectionSignal ? HistoryRejectionSignal1 : SignalInput1,
+							GlobalPointClampedSampler,
+							bUsesHistoryRejectionSignal ? CONFIG_SIGNAL_HISTORY_REJECTION_LAYOUT : CONFIG_SIGNAL_INPUT_LAYOUT,
+							SampleBufferUV, MipLevel);
+
+						// Sample the meta data of the neighbor.	
+						FSSFSampleSceneInfos NeighborSceneMetadata = FetchCurrentSceneInfos(SampleScreenPosition, SampleBufferUV);
+
+						// TODO: needs to be done for TILE_MODE_FULLY_OCCLUDED and TILE_MODE_FULLY_LIGHTED
+						float BilateralWeight = ComputeBilateralWeight(
+							WorldBluringRadius,
+							CurrentSceneMetadata,
+							NeighborSceneMetadata);
+
+						if (BilateralWeight > 0 && Neighbor.AccumulatedSampleCount > 0)
+						{
+							NeighborPenumbra = GetSamplePenumbra(Neighbor);
+							NeighborUncertainty = bUsesHistoryRejectionSignal
+								? 0.0
+								: Neighbor.UncertainRayCount / (Neighbor.AccumulatedSampleCount + Neighbor.UncertainRayCount);
+						}
+
+						NormalizedNeighborSample = NormalizeToOneSample(Neighbor);
+					} // if (bUseSample)
 				}
 
-				FLATTEN
-				if (bUseSample || !CONFIG_INPUT_TILE_CLASSIFICATION)
-				{
-					float2 SampleViewportUV = BufferUVToViewportUV(SampleBufferUV);
-					float2 SampleScreenPosition = ViewportUVToScreenPos(SampleViewportUV);
+			}
+			#else // !CONFIG_INPUT_TILE_CLASSIFICATION
+			{
+				float2 SampleViewportUV = BufferUVToViewportUV(SampleBufferUV);
+				float2 SampleScreenPosition = ViewportUVToScreenPos(SampleViewportUV);
 
-					// Sample the neighor for history rejection.
-					FSSFSignalSample Neighbor = SampleInputSignal(
-						bUsesHistoryRejectionSignal ? HistoryRejectionSignal0 : SignalInput0,
-						bUsesHistoryRejectionSignal ? HistoryRejectionSignal1 : SignalInput1,
-						GlobalPointClampedSampler,
-						bUsesHistoryRejectionSignal ? CONFIG_SIGNAL_HISTORY_REJECTION_LAYOUT : CONFIG_SIGNAL_INPUT_LAYOUT,
-						SampleBufferUV, MipLevel);
+				// Sample the neighor for history rejection.
+				FSSFSignalSample Neighbor = SampleInputSignal(
+					bUsesHistoryRejectionSignal ? HistoryRejectionSignal0 : SignalInput0,
+					bUsesHistoryRejectionSignal ? HistoryRejectionSignal1 : SignalInput1,
+					GlobalPointClampedSampler,
+					bUsesHistoryRejectionSignal ? CONFIG_SIGNAL_HISTORY_REJECTION_LAYOUT : CONFIG_SIGNAL_INPUT_LAYOUT,
+					SampleBufferUV, MipLevel);
 
-					// Sample the meta data of the neighbor.	
-					FSSFSampleSceneInfos NeighborSceneMetadata = FetchCurrentSceneInfos(SampleScreenPosition, SampleBufferUV);
+				// Sample the meta data of the neighbor.	
+				FSSFSampleSceneInfos NeighborSceneMetadata = FetchCurrentSceneInfos(SampleScreenPosition, SampleBufferUV);
 
-					// TODO: needs to be done for TILE_MODE_FULLY_OCCLUDED and TILE_MODE_FULLY_LIGHTED
-					float BilateralWeight = ComputeBilateralWeight(
-						WorldBluringRadius,
-						CurrentSceneMetadata,
-						NeighborSceneMetadata);
+				// TODO: needs to be done for TILE_MODE_FULLY_OCCLUDED and TILE_MODE_FULLY_LIGHTED
+				float BilateralWeight = ComputeBilateralWeight(
+					WorldBluringRadius,
+					CurrentSceneMetadata,
+					NeighborSceneMetadata);
 
-					if (BilateralWeight > 0 && Neighbor.AccumulatedSampleCount > 0)
-					{
-						NeighborPenumbra = GetSamplePenumbra(Neighbor);
-						NeighborUncertainty = bUsesHistoryRejectionSignal
-							? 0.0
-							: Neighbor.UncertainRayCount / (Neighbor.AccumulatedSampleCount + Neighbor.UncertainRayCount);
-					}
+				if (BilateralWeight > 0 && Neighbor.AccumulatedSampleCount > 0)
+				{
+					NeighborPenumbra = GetSamplePenumbra(Neighbor);
+					NeighborUncertainty = bUsesHistoryRejectionSignal
+						? 0.0
+						: Neighbor.UncertainRayCount / (Neighbor.AccumulatedSampleCount + Neighbor.UncertainRayCount);
 				}
+
+				NormalizedNeighborSample = NormalizeToOneSample(Neighbor);
 			}
+			#endif // !CONFIG_INPUT_TILE_CLASSIFICATION
 
 			if (CONFIG_CERTAINTY_BASED_REJECTION && NeighborUncertainty < 1.0)
 			{
@@ -631,6 +691,18 @@ void TemporallyAccumulate(
 				NeighborMinPenumbra = min(NeighborMinPenumbra, NeighborPenumbra);
 				NeighborMaxPenumbra = max(NeighborMaxPenumbra, NeighborPenumbra);
 			}
+
+			if (NeighborUncertainty < 1.0)
+			{
+				NeighborMin = Min(NeighborMin, NormalizedNeighborSample);
+				NeighborMax = Max(NeighborMax, NormalizedNeighborSample);
+
+				float Weight = 1;
+
+				NeighborMomentWeight += Weight;
+				NeighborMoment1 = Add(NeighborMoment1, Mul(NormalizedNeighborSample, Weight));
+				NeighborMoment2 = Add(NeighborMoment2, Mul(Power(NormalizedNeighborSample, 2.0), Weight));
+			}
 		}
 
 		// The clamping box might be uncertain, in witch case we just take the middle.
@@ -640,19 +712,45 @@ void TemporallyAccumulate(
 			NeighborMaxPenumbra = (NeighborMinPenumbra + NeighborMaxPenumbra) * 0.5;
 		}
 
+		NeighborMoment1 = Mul(NeighborMoment1, rcp(NeighborMomentWeight));
+		NeighborMoment2 = Mul(NeighborMoment2, rcp(NeighborMomentWeight));
+
 		// Clamp the history.
-		float HistoryPenumbra = GetSamplePenumbraSafe(HistorySample);
-		float NewHistoryPenumbra = clamp(
-			HistoryPenumbra, NeighborMinPenumbra, NeighborMaxPenumbra);
+		#if CONFIG_SIGNAL_PROCESSING == SIGNAL_PROCESSING_PENUMBRA
+		{
+			float HistoryPenumbra = GetSamplePenumbraSafe(HistorySample);
+			float NewHistoryPenumbra = clamp(
+				HistoryPenumbra, NeighborMinPenumbra, NeighborMaxPenumbra);
+
+			float PreviousSampleCount = HistorySample.AccumulatedSampleCount;
+
+			// Count the number of samples (hit or miss) that have been rejected by the clamp.
+			float RejectedSampleCount = abs(NewHistoryPenumbra - HistoryPenumbra) * PreviousSampleCount;
+
+			// Amend the previous frame with new sample count and clamped penumbra.
+			HistorySample.AccumulatedSampleCount -= RejectedSampleCount;
+			HistorySample.LeavingRayCount = HistorySample.AccumulatedSampleCount * NewHistoryPenumbra;
+		}
+		#elif CONFIG_HISTORY_REJECTION == 2
+		{
+			const float StdDevMultiplier = 1.25;
 
-		// Count the number of samples (hit or miss) that have been rejected by the clamp.
-		float RejectedSampleCount = abs(NewHistoryPenumbra - HistoryPenumbra) * HistorySample.AccumulatedSampleCount;
+			// StdDev = sqrt(abs(NeighborSquareAvg - NeighborAvg * NeighborAvg));
+			FSSFSignalSample StdDev = Sqrt(Abs(Subtract(NeighborMoment2, Power(NeighborMoment1, 2))));
+			NeighborMin = Add(NeighborMoment1, Mul(StdDev, -StdDevMultiplier));
+			NeighborMax = Add(NeighborMoment1, Mul(StdDev,  StdDevMultiplier));
 
-		// Amend the previous frame with new sample count and clamped penumbra.
-		HistorySample.AccumulatedSampleCount -= RejectedSampleCount;
-		HistorySample.LeavingRayCount = HistorySample.AccumulatedSampleCount * NewHistoryPenumbra;
+			float RejectedSampleCount = 0;
+			HistorySample = Mul(Clamp(NormalizeToOneSample(HistorySample), NeighborMin, NeighborMax), HistorySample.AccumulatedSampleCount - RejectedSampleCount);
+		}
+		#else // CONFIG_HISTORY_REJECTION == 1
+		{
+			float RejectedSampleCount = 0;
+			HistorySample = Mul(Clamp(NormalizeToOneSample(HistorySample), NeighborMin, NeighborMax), HistorySample.AccumulatedSampleCount - RejectedSampleCount);
+		}
+		#endif
 	}
-	#endif // CONFIG_HISTORY_REJECTION
+	#endif // CONFIG_HISTORY_REJECTION > 0
 	
 	// Evaluate the final sample to output.
 	if (bPostRejectionBlending)
