Avoid using `Vector` variable name, as it is a reserved keyword in some compilers.

diff --git a/Engine/Shaders/Private/RayTracing/SkyLightMipTreeCommon.ush b/Engine/Shaders/Private/RayTracing/SkyLightMipTreeCommon.ush
--- a/Engine/Shaders/Private/RayTracing/SkyLightMipTreeCommon.ush
+++ b/Engine/Shaders/Private/RayTracing/SkyLightMipTreeCommon.ush
@@ -8,7 +8,7 @@
 
 float3 GetTextureCubeVector(float3 TexelCoord, uint2 TextureRes)
 {
-	float3 Vector;
+	float3 Result;
 
 	float2 UV = TexelCoord.xy / float2(TextureRes.xy);
 	UV = UV * 2.0 - 1.0;
@@ -16,37 +16,37 @@ float3 GetTextureCubeVector(float3 TexelCoord, uint2 TextureRes)
 	uint Cubeface = TexelCoord.z;
 	if (Cubeface == 0)
 	{
-		Vector = float3(1, -UV.y, -UV.x);
+		Result = float3(1, -UV.y, -UV.x);
 	}
 	else if (Cubeface == 1)
 	{
-		Vector = float3(-1, -UV.y, UV.x);
+		Result = float3(-1, -UV.y, UV.x);
 	}
 	else if (Cubeface == 2)
 	{
-		Vector = float3(UV.x, 1, UV.y);
+		Result = float3(UV.x, 1, UV.y);
 	}
 	else if (Cubeface == 3)
 	{
-		Vector = float3(UV.x, -1, -UV.y);
+		Result = float3(UV.x, -1, -UV.y);
 	}
 	else if (Cubeface == 4)
 	{
-		Vector = float3(UV.x, -UV.y, 1);
+		Result = float3(UV.x, -UV.y, 1);
 	}
 	else if (Cubeface == 5)
 	{
-		Vector = float3(-UV.x, -UV.y, -1);
+		Result = float3(-UV.x, -UV.y, -1);
 	}
 
-	return Vector;
+	return Result;
 }
 
-uint GetMaximumComponentIndex(float3 Vector)
+uint GetMaximumComponentIndex(float3 Input)
 {
-	float AbsX = abs(Vector.x);
-	float AbsY = abs(Vector.y);
-	float AbsZ = abs(Vector.z);
+	float AbsX = abs(Input.x);
+	float AbsY = abs(Input.y);
+	float AbsZ = abs(Input.z);
 	if (AbsX > AbsY)
 	{
 		return AbsX > AbsZ ? 0 : 2;
