Fixes a bug in reflection ray generation shader where opacity of the reflection would decrease on high rougness material.

diff --git a/Engine/Shaders/Private/RayTracing/RayTracingReflections.usf b/Engine/Shaders/Private/RayTracing/RayTracingReflections.usf
--- a/Engine/Shaders/Private/RayTracing/RayTracingReflections.usf
+++ b/Engine/Shaders/Private/RayTracing/RayTracingReflections.usf
@@ -132,7 +132,7 @@ float TraceShadowRay(
 	Ray.TMin = 1e-4f;
 	Ray.TMax = TMax;
 
-	uint RayFlags = RAY_FLAG_CULL_BACK_FACING_TRIANGLES | RAY_FLAG_ACCEPT_FIRST_HIT_AND_END_SEARCH | RAY_FLAG_SKIP_CLOSEST_HIT_SHADER;
+	uint RayFlags = RAY_FLAG_CULL_BACK_FACING_TRIANGLES | RAY_FLAG_ACCEPT_FIRST_HIT_AND_END_SEARCH;
 
 	FMaterialClosestHitPayload Payload = TraceRayInternal(
 		TLAS,      // AccelerationStructure
@@ -319,7 +319,6 @@ void RayTracingReflectionsRGS()
 	float3 WorldNormal = ScreenSpaceData.GBuffer.WorldNormal;
 	float Roughness = ScreenSpaceData.GBuffer.Roughness;
 	
-	float4 ReflectedColor = 0.0;
 	float3 BRDFThroughput = float3(1.0f, 1.0f, 1.0f);
 
 	uint MaxBounces = 1; //#dxr-todo: only works with 1 bounce until hit shaders are supported
@@ -332,6 +331,8 @@ void RayTracingReflectionsRGS()
 		LocalSamplesPerPixel = 0;
 	}
 
+	float4 ReflectedColor = 0.0;
+	float ReflectedWeight = 0.0;
 	float HitDistance = 0.0;
 
 	if (Roughness < RoughnessThreshold)
@@ -374,21 +375,18 @@ void RayTracingReflectionsRGS()
 				float3 Radiance = float3(0, 0, 0);
 				FMaterialClosestHitPayload Payload = TraceRayAndAccumulateResults(Ray, Radiance);
 
-				if (Payload.HitT < 0)
-				{
-					break;
-				}
-
-				float4 PathRadiance = float4(Radiance, 1.0); 
+				float4 PathRadiance = float4(Radiance, Payload.HitT < 0 ? 0.0 : 1.0); 
 				if (BounceIndex == 0)
 				{
 					ReflectedColor += PathRadiance;
+					ReflectedWeight += 1.0;
 				}
 				else
 				{
 					float4 FinalReflectionColor = PathRadiance;
 					//float4 FinalReflectionColor = BlendWithEnvironment(BRDFThroughput, Payload, Ray.Origin, PathRadiance, Ray.Direction, Payload.WorldNormal);
 					ReflectedColor += FinalReflectionColor;
+					ReflectedWeight += 1.0;
 				}
 
 				IncidentDirection = Ray.Direction;
@@ -398,9 +396,9 @@ void RayTracingReflectionsRGS()
 		}
 	}
 
-	if (SamplesPerPixel > 0)
+	if (ReflectedWeight > 0)
 	{
-		ReflectedColor /= LocalSamplesPerPixel;
+		ReflectedColor *= rcp(ReflectedWeight);
 	}
 
 	ColorOutput[PixelCoord] = float4(ReflectedColor);
