Hopefully last round of copyright updates

diff --git a/Engine/Shaders/Private/PathTracing/BuildVarianceMipTreeComputeShader.usf b/Engine/Shaders/Private/PathTracing/BuildVarianceMipTreeComputeShader.usf
--- a/Engine/Shaders/Private/PathTracing/BuildVarianceMipTreeComputeShader.usf
+++ b/Engine/Shaders/Private/PathTracing/BuildVarianceMipTreeComputeShader.usf
@@ -1,4 +1,4 @@
-// Copyright 1998-2018 Epic Games, Inc. All Rights Reserved.
+// Copyright 1998-2019 Epic Games, Inc. All Rights Reserved.
 
 /*=============================================================================================
 	BuildVarianceMapComputeShader.usf: Build variance map
diff --git a/Engine/Shaders/Private/PathTracing/VisualizeSkyLightCdfPixelShader.usf b/Engine/Shaders/Private/PathTracing/VisualizeSkyLightCdfPixelShader.usf
--- a/Engine/Shaders/Private/PathTracing/VisualizeSkyLightCdfPixelShader.usf
+++ b/Engine/Shaders/Private/PathTracing/VisualizeSkyLightCdfPixelShader.usf
@@ -1,4 +1,4 @@
-// Copyright 1998-2018 Epic Games, Inc. All Rights Reserved.
+// Copyright 1998-2019 Epic Games, Inc. All Rights Reserved.
 
 /*=============================================================================================
 VisualizeSkyLightCdfPixelShader.usf: Interactive visualization of SkyLight CDFs.
diff --git a/Engine/Source/Runtime/Renderer/Private/MeshPassProcessor.cpp b/Engine/Source/Runtime/Renderer/Private/MeshPassProcessor.cpp
--- a/Engine/Source/Runtime/Renderer/Private/MeshPassProcessor.cpp
+++ b/Engine/Source/Runtime/Renderer/Private/MeshPassProcessor.cpp
@@ -1,4 +1,4 @@
-// Copyright 1998-2018 Epic Games, Inc. All Rights Reserved.
+// Copyright 1998-2019 Epic Games, Inc. All Rights Reserved.
 
 /*=============================================================================
 	MeshPassProcessor.cpp: 
diff --git a/Engine/Source/Runtime/Renderer/Private/PathTracing/PathTracingTestbed.cpp b/Engine/Source/Runtime/Renderer/Private/PathTracing/PathTracingTestbed.cpp
--- a/Engine/Source/Runtime/Renderer/Private/PathTracing/PathTracingTestbed.cpp
+++ b/Engine/Source/Runtime/Renderer/Private/PathTracing/PathTracingTestbed.cpp
@@ -1,4 +1,4 @@
-// Copyright 1998-2018 Epic Games, Inc. All Rights Reserved.
+// Copyright 1998-2019 Epic Games, Inc. All Rights Reserved.
 
 #include "PathTracingTestbed.h"
 #include "RHI.h"
diff --git a/Engine/Source/Runtime/Renderer/Private/PathTracing/SkyLightCdfBuild.cpp b/Engine/Source/Runtime/Renderer/Private/PathTracing/SkyLightCdfBuild.cpp
--- a/Engine/Source/Runtime/Renderer/Private/PathTracing/SkyLightCdfBuild.cpp
+++ b/Engine/Source/Runtime/Renderer/Private/PathTracing/SkyLightCdfBuild.cpp
@@ -1,4 +1,4 @@
-// Copyright 1998-2018 Epic Games, Inc. All Rights Reserved.
+// Copyright 1998-2019 Epic Games, Inc. All Rights Reserved.
 
 /*=============================================================================
 	SkyLightCdfBuild.cpp: SkyLight CDF build algorithm.
diff --git a/Engine/Source/Runtime/Renderer/Private/PathTracing/VarianceMipTreeBuild.cpp b/Engine/Source/Runtime/Renderer/Private/PathTracing/VarianceMipTreeBuild.cpp
--- a/Engine/Source/Runtime/Renderer/Private/PathTracing/VarianceMipTreeBuild.cpp
+++ b/Engine/Source/Runtime/Renderer/Private/PathTracing/VarianceMipTreeBuild.cpp
@@ -1,4 +1,4 @@
-// Copyright 1998-2018 Epic Games, Inc. All Rights Reserved.
+// Copyright 1998-2019 Epic Games, Inc. All Rights Reserved.
 
 /*=============================================================================
 	VarianceMipTreeBuild.cpp: SkyLight CDF build algorithm.
diff --git a/Engine/Source/Runtime/Renderer/Private/RayTracing/RayTracingDynamicGeometry.cpp b/Engine/Source/Runtime/Renderer/Private/RayTracing/RayTracingDynamicGeometry.cpp
--- a/Engine/Source/Runtime/Renderer/Private/RayTracing/RayTracingDynamicGeometry.cpp
+++ b/Engine/Source/Runtime/Renderer/Private/RayTracing/RayTracingDynamicGeometry.cpp
@@ -1,4 +1,4 @@
-// Copyright 1998-2018 Epic Games, Inc. All Rights Reserved.
+// Copyright 1998-2019 Epic Games, Inc. All Rights Reserved.
 
 #include "MeshMaterialShader.h"
 #include "ScenePrivate.h"
diff --git a/Engine/Source/Runtime/Renderer/Private/SceneViewFamilyBlackboard.cpp b/Engine/Source/Runtime/Renderer/Private/SceneViewFamilyBlackboard.cpp
--- a/Engine/Source/Runtime/Renderer/Private/SceneViewFamilyBlackboard.cpp
+++ b/Engine/Source/Runtime/Renderer/Private/SceneViewFamilyBlackboard.cpp
@@ -1,4 +1,4 @@
-// Copyright 1998-2018 Epic Games, Inc. All Rights Reserved.
+// Copyright 1998-2019 Epic Games, Inc. All Rights Reserved.
 
 /*=============================================================================
 	DeferredShadingRenderer.h: Scene rendering definitions.
diff --git a/Engine/Source/Runtime/Renderer/Private/ScreenSpaceDenoise.cpp b/Engine/Source/Runtime/Renderer/Private/ScreenSpaceDenoise.cpp
--- a/Engine/Source/Runtime/Renderer/Private/ScreenSpaceDenoise.cpp
+++ b/Engine/Source/Runtime/Renderer/Private/ScreenSpaceDenoise.cpp
@@ -1,4 +1,4 @@
-// Copyright 1998-2018 Epic Games, Inc. All Rights Reserved.
+// Copyright 1998-2019 Epic Games, Inc. All Rights Reserved.
 
 /*=============================================================================
 	ScreenSpaceDenoise.cpp: Denoise in screen space.
diff --git a/Engine/Source/Runtime/Renderer/Public/MeshPassProcessor.h b/Engine/Source/Runtime/Renderer/Public/MeshPassProcessor.h
--- a/Engine/Source/Runtime/Renderer/Public/MeshPassProcessor.h
+++ b/Engine/Source/Runtime/Renderer/Public/MeshPassProcessor.h
@@ -1,4 +1,4 @@
-// Copyright 1998-2018 Epic Games, Inc. All Rights Reserved.
+// Copyright 1998-2019 Epic Games, Inc. All Rights Reserved.
 
 /*=============================================================================
 	MeshPassProcessor.h
