Local copyright updates

diff --git a/Engine/Shaders/Private/ScreenSpaceDenoise/SSDReflectionsSignal.ush b/Engine/Shaders/Private/ScreenSpaceDenoise/SSDReflectionsSignal.ush
--- a/Engine/Shaders/Private/ScreenSpaceDenoise/SSDReflectionsSignal.ush
+++ b/Engine/Shaders/Private/ScreenSpaceDenoise/SSDReflectionsSignal.ush
@@ -1,4 +1,4 @@
-// Copyright 1998-2018 Epic Games, Inc. All Rights Reserved.
+// Copyright 1998-2019 Epic Games, Inc. All Rights Reserved.
 
 #pragma once
 
