Evaluate skylight directly in the ray traced reflection pass, instead of falling back to ReflectionEnvironmentPS

diff --git a/Engine/Shaders/Private/RayTracing/RayTracingReflections.usf b/Engine/Shaders/Private/RayTracing/RayTracingReflections.usf
--- a/Engine/Shaders/Private/RayTracing/RayTracingReflections.usf
+++ b/Engine/Shaders/Private/RayTracing/RayTracingReflections.usf
@@ -23,6 +23,7 @@ float2 SobolIndexToUniformUnitSquare(uint2 SobolRand)
 
 #include "../DeferredShadingCommon.ush"
 #include "../DeferredLightingCommon.ush"
+#include "../ReflectionEnvironmentShared.ush"
 #include "../Montecarlo.ush"
 #include "../PathTracing/Utilities/PathTracingRandomSequence.ush" // #dxr-todo: move to a common folder
 #include "RayTracingCommon.ush"
@@ -375,18 +376,34 @@ void RayTracingReflectionsRGS()
 				float3 Radiance = float3(0, 0, 0);
 				FMaterialClosestHitPayload Payload = TraceRayAndAccumulateResults(Ray, Radiance);
 
-				float4 PathRadiance = float4(Radiance, Payload.HitT < 0 ? 0.0 : 1.0); 
-				if (BounceIndex == 0)
+				float4 PathRadiance;
+				if (Payload.HitT >= 0)
 				{
-					ReflectedColor += PathRadiance;
-					ReflectedWeight += 1.0;
+					PathRadiance = float4(Radiance, 1.0); 
+					if (BounceIndex == 0)
+					{
+						ReflectedColor += PathRadiance;
+						ReflectedWeight += 1.0;
+					}
+					else
+					{
+						float4 FinalReflectionColor = PathRadiance;
+						//float4 FinalReflectionColor = BlendWithEnvironment(BRDFThroughput, Payload, Ray.Origin, PathRadiance, Ray.Direction, Payload.WorldNormal);
+						ReflectedColor += FinalReflectionColor;
+						ReflectedWeight += 1.0;
+					}
 				}
 				else
 				{
-					float4 FinalReflectionColor = PathRadiance;
-					//float4 FinalReflectionColor = BlendWithEnvironment(BRDFThroughput, Payload, Ray.Origin, PathRadiance, Ray.Direction, Payload.WorldNormal);
-					ReflectedColor += FinalReflectionColor;
-					ReflectedWeight += 1.0;
+					if (ReflectionStruct.SkyLightParameters.y > 0)
+					{
+						float SkyAverageBrightness = 1.0f;
+
+						float3 SkyLighting = GetSkyLightReflection(Ray.Direction, Roughness, SkyAverageBrightness);
+
+						ReflectedColor += float4(SkyLighting, 1.0f);
+						ReflectedWeight += 1.0;
+					}
 				}
 
 				IncidentDirection = Ray.Direction;
diff --git a/Engine/Source/Runtime/Renderer/Private/RayTracing/RayTracingReflections.cpp b/Engine/Source/Runtime/Renderer/Private/RayTracing/RayTracingReflections.cpp
--- a/Engine/Source/Runtime/Renderer/Private/RayTracing/RayTracingReflections.cpp
+++ b/Engine/Source/Runtime/Renderer/Private/RayTracing/RayTracingReflections.cpp
@@ -144,6 +144,7 @@ class FRayTracingReflectionsRG : public FGlobalShader
 		SHADER_PARAMETER_STRUCT_REF(FViewUniformShaderParameters, ViewUniformBuffer)
 		SHADER_PARAMETER_STRUCT_REF(FSceneTexturesUniformParameters, SceneTexturesStruct)
 		SHADER_PARAMETER_STRUCT_REF(FReflectionsLightData, LightData)
+		SHADER_PARAMETER_STRUCT_REF(FReflectionUniformParameters, ReflectionStruct)
 
 		SHADER_PARAMETER_RDG_TEXTURE_UAV(RWTexture2D<float4>, ColorOutput)
 		SHADER_PARAMETER_RDG_TEXTURE_UAV(RWTexture2D<float>, RayHitDistanceOutput)
@@ -227,6 +228,11 @@ void FDeferredShadingSceneRenderer::RayTraceReflections(
 		SetupSceneTextureUniformParameters(SceneContext, FeatureLevel, ESceneTextureSetupMode::All, SceneTextures);
 		PassParameters->SceneTexturesStruct = CreateUniformBufferImmediate(SceneTextures, EUniformBufferUsage::UniformBuffer_SingleDraw);
 	}
+	{
+		FReflectionUniformParameters ReflectionStruct;
+		SetupReflectionUniformParameters(View, ReflectionStruct);
+		PassParameters->ReflectionStruct = CreateUniformBufferImmediate(ReflectionStruct, EUniformBufferUsage::UniformBuffer_SingleDraw);
+	}
 	PassParameters->ColorOutput = GraphBuilder.CreateUAV(*OutColorTexture);
 	PassParameters->RayHitDistanceOutput = GraphBuilder.CreateUAV(*OutRayHitDistanceTexture);
 
