Fixes various bugs with reflection's temporal accumulations.

diff --git a/Engine/Shaders/Private/ColorSpace.ush b/Engine/Shaders/Private/ColorSpace.ush
--- a/Engine/Shaders/Private/ColorSpace.ush
+++ b/Engine/Shaders/Private/ColorSpace.ush
@@ -131,4 +131,57 @@ float3 LAB_2_LinearRGB( float3 LAB, float3 ReferenceWhite)
 	float Z = ReferenceWhite.z * f_z;
 
 	return XYZ_2_LinearRGB(float3(X, Y, Z));
-}
\ No newline at end of file
+}
+
+//------------------------------------------------------------------
+// YCoCg and LCoCg color space for luma / chroma orthogonal processing
+//------------------------------------------------------------------
+
+// Multiplier that get applied for ALU optimisation purposes on the luma, when converting LinearRGB to YCoCg.
+#define LINEARRGB_2_YCOCG_LUMA_MULTIPLIER 4.0
+
+float3 LinearRGB_2_YCoCg(float3 RGB)
+{
+	float Y = dot(RGB, float3(1, 2, 1));
+	float Co = dot(RGB, float3(2, 0, -2));
+	float Cg = dot(RGB, float3(-1, 2, -1));
+
+	float3 YCoCg = float3(Y, Co, Cg);
+	return YCoCg;
+}
+
+float3 YCoCg_2_LinearRGB(float3 YCoCg)
+{
+	float Y = YCoCg.x * 0.25;
+	float Co = YCoCg.y * 0.25;
+	float Cg = YCoCg.z * 0.25;
+
+	float R = Y + Co - Cg;
+	float G = Y + Cg;
+	float B = Y - Co - Cg;
+
+	float3 RGB = float3(R, G, B);
+	return RGB;
+}
+
+float3 YCoCg_2_LCoCg(float3 YCoCg)
+{
+	return float3(
+		YCoCg.x,
+		YCoCg.yz * (YCoCg.x > 0 ? rcp(YCoCg.x) : 0));
+}
+
+float3 LCoCg_2_YCoCg(float3 LCoCg)
+{
+	return float3(LCoCg.x, LCoCg.x * LCoCg.yz);
+}
+
+float3 LinearRGB_2_LCoCg(float3 RGB)
+{
+	return YCoCg_2_LCoCg(LinearRGB_2_YCoCg(RGB));
+}
+
+float3 LCoCg_2_LinearRGB(float3 LCoCg)
+{
+	return YCoCg_2_LinearRGB(LCoCg_2_YCoCg(LCoCg));
+}
diff --git a/Engine/Shaders/Private/ScreenSpaceDenoise/SSDCommon.ush b/Engine/Shaders/Private/ScreenSpaceDenoise/SSDCommon.ush
--- a/Engine/Shaders/Private/ScreenSpaceDenoise/SSDCommon.ush
+++ b/Engine/Shaders/Private/ScreenSpaceDenoise/SSDCommon.ush
@@ -8,6 +8,7 @@
 
 #include "../Common.ush"
 #include "../DataStructures/SceneViewFamilyBlackboard.ush"
+#include "../ColorSpace.ush"
 
 
 //------------------------------------------------------- DEBUG CONFIG
@@ -77,6 +78,15 @@
 	#define LIGHT_TYPE_RECT 3
 
 
+//------------------------------------------------------- SHARED FUNCTION
+
+float SafeRcp(float x)
+{
+	if (x > 0)
+		return rcp(x);
+	return 0;
+}
+
 //------------------------------------------------------- COMMON SCENE INFOS
 
 /** olds commonly used information of the scene of a given sample. */
diff --git a/Engine/Shaders/Private/ScreenSpaceDenoise/SSDSignalCore.ush b/Engine/Shaders/Private/ScreenSpaceDenoise/SSDSignalCore.ush
--- a/Engine/Shaders/Private/ScreenSpaceDenoise/SSDSignalCore.ush
+++ b/Engine/Shaders/Private/ScreenSpaceDenoise/SSDSignalCore.ush
@@ -121,7 +121,7 @@ FSSFSignalSample DecodeInputSignal(const uint SignalBufferLayout, float4 Samples
 	}
 	else if (SignalBufferLayout == SIGNAL_BUFFER_LAYOUT_REFLECTIONS_INPUT)
 	{
-		OutSample.AccumulatedSampleCount = 1;
+		OutSample.AccumulatedSampleCount = Samples[0].a > 0 ? 1 : 0;
 		OutSample.LeavingRayCount = OutSample.AccumulatedSampleCount * (1 - Samples[0].a);
 		OutSample.SceneColor = Samples[0];
 		OutSample.PenumbraClosestHit = Samples[1].r;
@@ -309,7 +309,13 @@ FSSFSignalSample Max(FSSFSignalSample SampleA, FSSFSignalSample SampleB)
 
 FSSFSignalSample Clamp(FSSFSignalSample Sample, FSSFSignalSample SampleMin, FSSFSignalSample SampleMax)
 {
-	return Max(Min(Sample, SampleMax), SampleMin);
+	FSSFSignalSample OutSample;
+	OutSample.SceneColor = clamp(Sample.SceneColor, SampleMin.SceneColor, SampleMax.SceneColor);
+	OutSample.AccumulatedSampleCount = clamp(Sample.AccumulatedSampleCount, SampleMin.AccumulatedSampleCount, SampleMax.AccumulatedSampleCount);
+	OutSample.LeavingRayCount = clamp(Sample.LeavingRayCount, SampleMin.LeavingRayCount, SampleMax.LeavingRayCount);
+	OutSample.PenumbraClosestHit = clamp(Sample.PenumbraClosestHit, SampleMin.PenumbraClosestHit, SampleMax.PenumbraClosestHit);
+	OutSample.UncertainRayCount = clamp(Sample.UncertainRayCount, SampleMin.UncertainRayCount, SampleMax.UncertainRayCount);
+	return OutSample;
 }
 
 FSSFSignalSample Lerp(FSSFSignalSample Sample0, FSSFSignalSample Sample1, float Interp)
diff --git a/Engine/Shaders/Private/ScreenSpaceDenoise/SSDTemporalAccumulation.usf b/Engine/Shaders/Private/ScreenSpaceDenoise/SSDTemporalAccumulation.usf
--- a/Engine/Shaders/Private/ScreenSpaceDenoise/SSDTemporalAccumulation.usf
+++ b/Engine/Shaders/Private/ScreenSpaceDenoise/SSDTemporalAccumulation.usf
@@ -18,26 +18,37 @@
 
 
 #if CONFIG_SIGNAL_PROCESSING == SIGNAL_PROCESSING_PENUMBRA
+	// Input and output layout.
+	// TODO: Do SIGNAL_BUFFER_LAYOUT_PENUMBRA_INTERNAL that does not store certainty for the history.
+	#define CONFIG_SIGNAL_INPUT_LAYOUT   SIGNAL_BUFFER_LAYOUT_PENUMBRA_INTERNAL
+	#define CONFIG_SIGNAL_HISTORY_LAYOUT SIGNAL_BUFFER_LAYOUT_PENUMBRA_INTERNAL
+
+	// Uses a larger kernel when burteforcing to converge faster.
+	#define CONFIG_LARGER_HISTORY_KERNEL 1
+
+	// Input is from reduce pass that has been tile classified.
 	#define CONFIG_INPUT_TILE_CLASSIFICATION 1
+
+	// The work is control by tile classification to avoid doing useless mip levels.
 	#define CONFIG_USE_TILE_CLASSIFICATION 1
+
+	// Uses a rejection buffer at MIP 0 because is very unstable otherwise.
 	#define CONFIG_USE_REJECTION_BUFFER (DIM_IS_MIP_0 == 1)
+	#define CONFIG_SIGNAL_HISTORY_REJECTION_LAYOUT SIGNAL_BUFFER_LAYOUT_PENUMBRA_HISTORY_REJECTION
 
 	// Uses a larger kernel when brute forcing to increase the speed of convergence by bluring more spatially.
 	#define CONFIG_LARGE_HISTORY_KERNEL 1
 
+	// Reject based on certainty of the mip level.
 	#define CONFIG_CERTAINTY_BASED_REJECTION 1
 
-	// TODO: Do SIGNAL_BUFFER_LAYOUT_PENUMBRA_INTERNAL that does not store certainty for the history.
-	#define CONFIG_SIGNAL_INPUT_LAYOUT   SIGNAL_BUFFER_LAYOUT_PENUMBRA_INTERNAL
-	#define CONFIG_SIGNAL_HISTORY_LAYOUT SIGNAL_BUFFER_LAYOUT_PENUMBRA_INTERNAL
-
-	#define CONFIG_SIGNAL_HISTORY_REJECTION_LAYOUT SIGNAL_BUFFER_LAYOUT_PENUMBRA_HISTORY_REJECTION
-
 #elif CONFIG_SIGNAL_PROCESSING == SIGNAL_PROCESSING_REFLECTIONS
+	// Input and output layout.
 	#define CONFIG_SIGNAL_INPUT_LAYOUT   SIGNAL_BUFFER_LAYOUT_REFLECTIONS_INPUT
 	#define CONFIG_SIGNAL_HISTORY_LAYOUT SIGNAL_BUFFER_LAYOUT_REFLECTIONS_HISTORY
 
-	#define CONFIG_HISTORY_REJECTION 2
+	#define CONFIG_HISTORY_REJECTION 1
+	#define CONFIG_HISTORY_REJECTION_COLOR_SPACE 0
 
 #else
 	#error Unimplemented signal processing.
@@ -74,6 +85,16 @@
 	#define CONFIG_CERTAINTY_BASED_REJECTION 0
 #endif
 
+/** The color of the history should be rejected with. */
+#ifndef CONFIG_HISTORY_REJECTION_COLOR_SPACE
+	#define CONFIG_HISTORY_REJECTION_COLOR_SPACE 0
+#endif
+
+/** Whether should use a larger kernel to sample history. */
+#ifndef CONFIG_LARGER_HISTORY_KERNEL
+	#define CONFIG_LARGER_HISTORY_KERNEL 0
+#endif
+
 
 //------------------------------------------------------- CONFIG ENABLED DEFAULTS
 
@@ -127,6 +148,29 @@ RWTexture2D<float4>	SignalHistoryOutput1;
 
 //------------------------------------------------------- FUNCTIONS
 
+FSSFSignalSample TransformSignalForPreRejection(FSSFSignalSample Sample)
+{
+	// Unpremultiplied RGBA to avoid any opacity mismatches.
+	Sample.SceneColor.rgb *= SafeRcp(Sample.SceneColor.a);
+
+	#if CONFIG_HISTORY_REJECTION_COLOR_SPACE == 1
+		Sample.SceneColor.rgb = LinearRGB_2_LCoCg(Sample.SceneColor.rgb);
+	#endif
+	return Sample;
+}
+
+FSSFSignalSample TransformSignalForPostRejection(FSSFSignalSample Sample)
+{
+	#if CONFIG_HISTORY_REJECTION_COLOR_SPACE == 1
+		Sample.SceneColor.rgb = LCoCg_2_LinearRGB(Sample.SceneColor.rgb);
+	#endif
+
+	// Premultiplied RGBA
+	Sample.SceneColor.rgb *= Sample.SceneColor.a;
+	return Sample;
+}
+
+
 /** Sample history using a point sampler. */
 void PointSampleHistory(float2 ScreenPosition, float2 BufferUV, out FSSFSignalSample OutSample, out FSSFSampleSceneInfos OutMetadata)
 {
@@ -216,14 +260,18 @@ void TemporallyAccumulate(
 				MipLevelPow2, rcp(MipLevelPow2));
 
 			// Number of sample to use.
-			const bool bUseLargerTAAKernel = bBruteForce && (CONFIG_HISTORY_REJECTION != 0);
+			#if CONFIG_LARGER_HISTORY_KERNEL == 0
+				const bool bUseLargerTAAKernel = false;
+			#elif CONFIG_LARGER_HISTORY_KERNEL == 1
+				const bool bUseLargerTAAKernel = bBruteForce;
+			#endif
 			const uint SampleCount = bUseLargerTAAKernel ? 16 : 4;
 
 			// Disabled because it is useful to count the sample we add together when brute forcing,
 			// so can converge faster the history.
 			// TODO: kill it.
 			float KernelWeight = 1;
-			if (bUseLargerTAAKernel && 0)
+			if (bUseLargerTAAKernel)
 			{
 				KernelWeight = 0;
 
@@ -489,14 +537,14 @@ void TemporallyAccumulate(
 	
 	FSSFSignalSample HistorySample = PreviousFrame;
 
-	const bool bPostRejectionBlending = !bBruteForce;
+	const bool bPostRejectionBlending = !bBruteForce; // || CONFIG_SIGNAL_PROCESSING == SIGNAL_PROCESSING_REFLECTIONS;
 
 	if (!bPostRejectionBlending)
 	{
 		float PreviousFrameWeight = min(HistorySample.AccumulatedSampleCount, TargetedSampleCount - CurrentFrame.AccumulatedSampleCount);
 		float PreviousFrameMultiplier = HistorySample.AccumulatedSampleCount > 0 ? PreviousFrameWeight / HistorySample.AccumulatedSampleCount : 0;
 
-		HistorySample = Add(Mul(PreviousFrame, PreviousFrameMultiplier), CurrentFrame);
+		HistorySample = Add(Mul(HistorySample, PreviousFrameMultiplier), CurrentFrame);
 	}
 
 	// History rejection.
@@ -544,12 +592,14 @@ void TemporallyAccumulate(
 			NeighborMinPenumbra = GetSamplePenumbraSafe(CurrentFrame);
 			NeighborMaxPenumbra = GetSamplePenumbraSafe(CurrentFrame);
 
-			NeighborMin = NormalizeToOneSample(CurrentFrame);
-			NeighborMax = NormalizeToOneSample(CurrentFrame);
+			FSSFSignalSample TransformedCurrentFrame = TransformSignalForPreRejection(NormalizeToOneSample(CurrentFrame));
+
+			NeighborMin = TransformedCurrentFrame;
+			NeighborMax = TransformedCurrentFrame;
 
 			NeighborMomentWeight = 1;
-			NeighborMoment1 = NormalizeToOneSample(CurrentFrame);
-			NeighborMoment2 = Power(NeighborMoment1, 2.0);
+			NeighborMoment1 = TransformedCurrentFrame;
+			NeighborMoment2 = Power(TransformedCurrentFrame, 2.0);
 		}
 
 		// Look for neighborhood to build clamping box.
@@ -574,9 +624,10 @@ void TemporallyAccumulate(
 			float NeighborPenumbra = -1;
 			float NeighborUncertainty = 1.0;
 
-			FSSFTileInfos SampleTileInfos;
 			#if CONFIG_INPUT_TILE_CLASSIFICATION
 			{
+				FSSFTileInfos SampleTileInfos;
+
 				// TODO: unless being lowest mip level, can reduce this to a gather4.
 				SampleTileInfos = LoadTileInfos(uint2(floor(SampleBufferUV * View.BufferSizeAndInvSize.xy)) / 8);
 
@@ -634,15 +685,20 @@ void TemporallyAccumulate(
 							CurrentSceneMetadata,
 							NeighborSceneMetadata);
 
+						NormalizedNeighborSample = NormalizeToOneSample(Neighbor);
+
 						if (BilateralWeight > 0 && Neighbor.AccumulatedSampleCount > 0)
 						{
 							NeighborPenumbra = GetSamplePenumbra(Neighbor);
 							NeighborUncertainty = bUsesHistoryRejectionSignal
 								? 0.0
 								: Neighbor.UncertainRayCount / (Neighbor.AccumulatedSampleCount + Neighbor.UncertainRayCount);
-						}
 
-						NormalizedNeighborSample = NormalizeToOneSample(Neighbor);
+							if (CONFIG_SIGNAL_PROCESSING == SIGNAL_PROCESSING_REFLECTIONS)
+							{
+								NeighborUncertainty = 1 - NormalizedNeighborSample.SceneColor.a;
+							}
+						}
 					} // if (bUseSample)
 				}
 
@@ -663,7 +719,6 @@ void TemporallyAccumulate(
 				// Sample the meta data of the neighbor.	
 				FSSFSampleSceneInfos NeighborSceneMetadata = FetchCurrentSceneInfos(SampleScreenPosition, SampleBufferUV);
 
-				// TODO: needs to be done for TILE_MODE_FULLY_OCCLUDED and TILE_MODE_FULLY_LIGHTED
 				float BilateralWeight = ComputeBilateralWeight(
 					WorldBluringRadius,
 					CurrentSceneMetadata,
@@ -681,6 +736,8 @@ void TemporallyAccumulate(
 			}
 			#endif // !CONFIG_INPUT_TILE_CLASSIFICATION
 
+			NormalizedNeighborSample = TransformSignalForPreRejection(NormalizedNeighborSample);
+
 			if (CONFIG_CERTAINTY_BASED_REJECTION && NeighborUncertainty < 1.0)
 			{
 				NeighborMinPenumbra = min(NeighborMinPenumbra, NeighborPenumbra + NeighborUncertainty);
@@ -705,19 +762,18 @@ void TemporallyAccumulate(
 			}
 		}
 
-		// The clamping box might be uncertain, in witch case we just take the middle.
-		if (CONFIG_CERTAINTY_BASED_REJECTION && NeighborMinPenumbra > NeighborMaxPenumbra)
-		{
-			NeighborMinPenumbra = (NeighborMinPenumbra + NeighborMaxPenumbra) * 0.5;
-			NeighborMaxPenumbra = (NeighborMinPenumbra + NeighborMaxPenumbra) * 0.5;
-		}
-
 		NeighborMoment1 = Mul(NeighborMoment1, rcp(NeighborMomentWeight));
 		NeighborMoment2 = Mul(NeighborMoment2, rcp(NeighborMomentWeight));
 
 		// Clamp the history.
 		#if CONFIG_SIGNAL_PROCESSING == SIGNAL_PROCESSING_PENUMBRA
 		{
+			// The clamping box might be uncertain, in witch case we just take the middle.
+			if (CONFIG_CERTAINTY_BASED_REJECTION && NeighborMinPenumbra > NeighborMaxPenumbra)
+			{
+				NeighborMaxPenumbra = NeighborMinPenumbra = (NeighborMinPenumbra + NeighborMaxPenumbra) * 0.5;
+			}
+
 			float HistoryPenumbra = GetSamplePenumbraSafe(HistorySample);
 			float NewHistoryPenumbra = clamp(
 				HistoryPenumbra, NeighborMinPenumbra, NeighborMaxPenumbra);
@@ -741,12 +797,12 @@ void TemporallyAccumulate(
 			NeighborMax = Add(NeighborMoment1, Mul(StdDev,  StdDevMultiplier));
 
 			float RejectedSampleCount = 0;
-			HistorySample = Mul(Clamp(NormalizeToOneSample(HistorySample), NeighborMin, NeighborMax), HistorySample.AccumulatedSampleCount - RejectedSampleCount);
+			HistorySample = Mul(TransformSignalForPostRejection(Clamp(TransformSignalForPreRejection(NormalizeToOneSample(HistorySample)), NeighborMin, NeighborMax)), HistorySample.AccumulatedSampleCount - RejectedSampleCount);
 		}
 		#else // CONFIG_HISTORY_REJECTION == 1
 		{
 			float RejectedSampleCount = 0;
-			HistorySample = Mul(Clamp(NormalizeToOneSample(HistorySample), NeighborMin, NeighborMax), HistorySample.AccumulatedSampleCount - RejectedSampleCount);
+			HistorySample = Mul(TransformSignalForPostRejection(Clamp(TransformSignalForPreRejection(NormalizeToOneSample(HistorySample)), NeighborMin, NeighborMax)), HistorySample.AccumulatedSampleCount - RejectedSampleCount);
 		}
 		#endif
 	}
@@ -758,7 +814,7 @@ void TemporallyAccumulate(
 		float PreviousFrameWeight = min(HistorySample.AccumulatedSampleCount, TargetedSampleCount - CurrentFrame.AccumulatedSampleCount);
 		float PreviousFrameMultiplier = HistorySample.AccumulatedSampleCount > 0 ? PreviousFrameWeight / HistorySample.AccumulatedSampleCount : 0;
 
-		HistorySample = Add(Mul(PreviousFrame, PreviousFrameMultiplier), CurrentFrame);
+		HistorySample = Add(Mul(HistorySample, PreviousFrameMultiplier), CurrentFrame);
 	}
 
 	if (all(DispatchThreadId < uint2(View.ViewSizeAndInvSize.xy)))
