Implements spatial reconstruction pass for reflection using the spatial accumulation shader with Stackowiak

diff --git a/Engine/Shaders/Private/ScreenSpaceDenoise/SSDReflectionsSignal.ush b/Engine/Shaders/Private/ScreenSpaceDenoise/SSDReflectionsSignal.ush
--- a/Engine/Shaders/Private/ScreenSpaceDenoise/SSDReflectionsSignal.ush
+++ b/Engine/Shaders/Private/ScreenSpaceDenoise/SSDReflectionsSignal.ush
@@ -7,6 +7,19 @@
 #include "SSDSignalCore.ush"
 
 
+//------------------------------------------------------- CONFIG DISABLED DEFAULTS
+
+#ifndef CONFIG_TOKOYASHI_BILATERAL
+	#define CONFIG_TOKOYASHI_BILATERAL 0
+#endif
+
+#ifndef CONFIG_SPECULAR_MIS
+	#define CONFIG_SPECULAR_MIS 0
+#endif
+
+
+//------------------------------------------------------- FUNCTIONS
+
 FSphericalGaussian RoughnessLobe(float Roughness, float3 N, float3 V)
 {
 	float a = Pow2( max( 0.02, Roughness ) );
@@ -64,7 +77,7 @@ float ComputeBilateralWeight(
 	}
 	
 	float WorldBilateralWeight = 1;
-	if (1)
+	if (0)
 	{
 		const float WorldRadius = MaxWorldBluringDistance;
 
@@ -85,10 +98,12 @@ float ComputeBilateralWeight(
 	}
 
 	float LobeBilateralWeight = 1;
+	#if CONFIG_TOKOYASHI_BILATERAL
 	//if (View.GeneralPurposeTweak > 0.5)
 	{
 		LobeBilateralWeight = TokoyashiLobeBilateralWeight(RefSceneMetadata, NeighborSceneMetadata);
 	}
+	#endif
 
 	return WorldBilateralWeight * LobeBilateralWeight;
 }
@@ -130,6 +145,12 @@ float GetMultiImportanceSamplingWeight(
 	FSSFSampleSceneInfos RefSceneMetadata,
 	FSSFSampleSceneInfos SampleSceneMetadata,
 	FSSFSignalSample Sample)
+#if CONFIG_SPECULAR_MIS
 {
 	return 1.0;
 }
+#else // CONFIG_MIS
+{
+	return 1.0;
+}
+#endif
diff --git a/Engine/Shaders/Private/ScreenSpaceDenoise/SSDSpatialAccumulation.usf b/Engine/Shaders/Private/ScreenSpaceDenoise/SSDSpatialAccumulation.usf
--- a/Engine/Shaders/Private/ScreenSpaceDenoise/SSDSpatialAccumulation.usf
+++ b/Engine/Shaders/Private/ScreenSpaceDenoise/SSDSpatialAccumulation.usf
@@ -10,9 +10,15 @@ PostProcessTemporalAA.usf: Temporal AA
 
 //------------------------------------------------------- ENUM VALUES
 
+/** Different possible stage for spatial accumulation. Matches  */
+	#define STAGE_RECONSTRUCTION 0
+	#define STAGE_POST_FILTERING 1
+
 /** Sets of sample available for the spatial kernel. */
-	#define SAMPLE_SET_3X3 0
-	#define SAMPLE_SET_5X5_WAVELET 1
+	#define SAMPLE_SET_NONE 0
+	#define SAMPLE_SET_3X3 1
+	#define SAMPLE_SET_5X5_WAVELET 2
+	#define SAMPLE_SET_STACKOWIAK_4_SETS 3
 
 
 //------------------------------------------------------- CONFIGS
@@ -28,21 +34,41 @@ PostProcessTemporalAA.usf: Temporal AA
 
 	#define CONFIG_SAMPLE_SET SAMPLE_SET_3X3
 
+	#if DIM_STAGE == STAGE_POST_FILTERING
+	#else
+		#error Unexpected stage.
+	#endif
 
 #elif CONFIG_SIGNAL_PROCESSING == SIGNAL_PROCESSING_REFLECTIONS
-	#define CONFIG_SIGNAL_INPUT_LAYOUT  SIGNAL_BUFFER_LAYOUT_REFLECTIONS_HISTORY
-	#define CONFIG_SIGNAL_OUTPUT_LAYOUT SIGNAL_BUFFER_LAYOUT_REFLECTIONS_HISTORY
+	#if DIM_STAGE == STAGE_RECONSTRUCTION
+		#define CONFIG_SIGNAL_INPUT_LAYOUT  SIGNAL_BUFFER_LAYOUT_REFLECTIONS_INPUT
+		#define CONFIG_SIGNAL_OUTPUT_LAYOUT SIGNAL_BUFFER_LAYOUT_REFLECTIONS_HISTORY
 
-	#define CONFIG_SAMPLE_SET SAMPLE_SET_3X3
+		#define CONFIG_SAMPLE_SET SAMPLE_SET_STACKOWIAK_4_SETS
+
+		// Reconstruction has 1 ray per input pixel, therefore can use MIS.
+		#define CONFIG_SPECULAR_MIS 1
 
-	#define CONFIG_WEIGHT_KERNEL_CENTER 1
+	#elif DIM_STAGE == STAGE_POST_FILTERING
+		#define CONFIG_SIGNAL_INPUT_LAYOUT  SIGNAL_BUFFER_LAYOUT_REFLECTIONS_HISTORY
+		#define CONFIG_SIGNAL_OUTPUT_LAYOUT SIGNAL_BUFFER_LAYOUT_REFLECTIONS_HISTORY
+
+		#define CONFIG_SAMPLE_SET SAMPLE_SET_3X3
+
+		// Post filtering can only be done using tokoyashi because have more than one sample at a time.
+		#define CONFIG_TOKOYASHI_BILATERAL 1
+
+		// TODO: Not sure of the math of Tokoyashi, weight of the center kernel might be != 1.
+		#define CONFIG_WEIGHT_KERNEL_CENTER 1
+
+	#else
+		#error Unexpected stage.
+	#endif
 
 #else
 	#error Unknown signal processing.
 #endif
 
-#include "SSDSignalInterface.ush"
-
 
 //------------------------------------------------------- CONFIG DISABLED DEFAULTS
 
@@ -57,13 +83,85 @@ PostProcessTemporalAA.usf: Temporal AA
 #endif
 
 
+//------------------------------------------------------- INCLUDES DEPENDINGS ON THE CONFIG
+
+#include "SSDSignalInterface.ush"
+
+
 //------------------------------------------------------- CONSTANTS
 
 static const float kWaveletFilterWeights5x5[] = { 3.0 / 8.0, 1.0 / 4.0, 1.0 / 16.0 };
 
 
+// [ Stackowiak 2015, "Stochastic Screen-Space Reflections" ]
+static const float2 kStackowiakSamples[56 * 4] =
+{
+	float2(-0.5, -0.5), float2(+0.5, -0.5), float2(-0.5, +0.5), float2(+0.5, +0.5),
+	float2(-1.5, +0.5), float2(-1.5, -0.5), float2(-0.5, +1.5), float2(+1.5, -0.5),
+	float2(+0.5, -1.5), float2(+2.5, -0.5), float2(+1.5, +0.5), float2(-0.5, -1.5),
+	float2(-1.5, -2.5), float2(-0.5, -2.5), float2(-1.5, -1.5), float2(-0.5, +2.5),
+	float2(-1.5, +1.5), float2(+1.5, -2.5), float2(-1.5, +2.5), float2(+1.5, +2.5),
+	float2(+0.5, -2.5), float2(-2.5, -0.5), float2(-2.5, -1.5), float2(-2.5, +0.5),
+	float2(+0.5, +1.5), float2(+0.5, +2.5), float2(-3.5, +0.5), float2(+0.5, +3.5),
+	float2(+1.5, -1.5), float2(+3.5, -0.5), float2(+2.5, +1.5), float2(+3.5, +0.5),
+	float2(+1.5, +1.5), float2(-2.5, +1.5), float2(-3.5, +2.5), float2(+3.5, +1.5),
+	float2(-3.5, -0.5), float2(-1.5, -3.5), float2(-2.5, -2.5), float2(-2.5, +2.5),
+	float2(+2.5, +0.5), float2(+2.5, +2.5), float2(+1.5, +3.5), float2(+3.5, -1.5),
+	float2(-3.5, -2.5), float2(+3.5, -2.5), float2(+2.5, -1.5), float2(+0.5, -3.5),
+	float2(-0.5, +3.5), float2(-0.5, -4.5), float2(-4.5, +0.5), float2(+4.5, +0.5),
+	float2(-4.5, -1.5), float2(-3.5, +1.5), float2(-0.5, -3.5), float2(+1.5, -3.5),
+	float2(+0.5, -4.5), float2(-1.5, +3.5), float2(+0.5, +4.5), float2(-3.5, -1.5),
+	float2(-4.5, +1.5), float2(+2.5, -4.5), float2(+2.5, -2.5), float2(-1.5, +4.5),
+	float2(-2.5, -4.5), float2(+4.5, -2.5), float2(+2.5, +3.5), float2(-3.5, +3.5),
+	float2(-2.5, +3.5), float2(+0.5, -5.5), float2(-4.5, +3.5), float2(-2.5, -3.5),
+	float2(-4.5, +2.5), float2(+3.5, +3.5), float2(+2.5, -3.5), float2(+4.5, +3.5),
+	float2(+3.5, -3.5), float2(+4.5, +2.5), float2(-5.5, +1.5), float2(-4.5, -0.5),
+	float2(+3.5, +2.5), float2(-0.5, +4.5), float2(-1.5, +5.5), float2(+1.5, +5.5),
+	float2(+4.5, -0.5), float2(+5.5, +0.5), float2(+4.5, +1.5), float2(-1.5, -4.5),
+	float2(-1.5, -5.5), float2(-4.5, -2.5), float2(-2.5, +5.5), float2(+2.5, +5.5),
+	float2(+1.5, +4.5), float2(+5.5, +1.5), float2(+1.5, -4.5), float2(-3.5, -3.5),
+	float2(+3.5, -4.5), float2(-3.5, -4.5), float2(+4.5, -1.5), float2(+4.5, -3.5),
+	float2(-3.5, -5.5), float2(-2.5, -5.5), float2(-4.5, -3.5), float2(+4.5, +4.5),
+	float2(-3.5, +4.5), float2(-2.5, +4.5), float2(-5.5, -2.5), float2(-5.5, +0.5),
+	float2(+2.5, -5.5), float2(+3.5, +4.5), float2(-0.5, -5.5), float2(-0.5, +6.5),
+	float2(+2.5, +4.5), float2(-5.5, -0.5), float2(-6.5, -1.5), float2(+1.5, -5.5),
+	float2(-6.5, -0.5), float2(+0.5, +5.5), float2(+1.5, +6.5), float2(+6.5, +1.5),
+	float2(-0.5, +5.5), float2(+6.5, -0.5), float2(-4.5, -4.5), float2(-5.5, +2.5),
+	float2(+5.5, -0.5), float2(-5.5, -1.5), float2(-6.5, +3.5), float2(-1.5, +6.5),
+	float2(-6.5, +0.5), float2(+4.5, -5.5), float2(-3.5, +6.5), float2(+6.5, -1.5),
+	float2(+0.5, -6.5), float2(-5.5, -3.5), float2(+5.5, -2.5), float2(+4.5, -4.5),
+	float2(+5.5, -1.5), float2(+3.5, -6.5), float2(+5.5, +3.5), float2(+3.5, -5.5),
+	float2(-5.5, -4.5), float2(+6.5, -3.5), float2(-0.5, -6.5), float2(+3.5, +6.5),
+	float2(-5.5, +3.5), float2(+0.5, +6.5), float2(+6.5, +0.5), float2(+6.5, -2.5),
+	float2(-6.5, -3.5), float2(-4.5, +4.5), float2(-7.5, -0.5), float2(+7.5, +0.5),
+	float2(+5.5, +2.5), float2(-0.5, -7.5), float2(+0.5, +7.5), float2(-4.5, +5.5),
+	float2(+3.5, +5.5), float2(-3.5, +5.5), float2(-4.5, -5.5), float2(+4.5, +6.5),
+	float2(+5.5, -4.5), float2(+4.5, +5.5), float2(-4.5, +6.5), float2(+6.5, +4.5),
+	float2(-7.5, +1.5), float2(-6.5, +1.5), float2(+5.5, -3.5), float2(-6.5, +2.5),
+	float2(-2.5, +6.5), float2(-1.5, -7.5), float2(+5.5, +4.5), float2(-1.5, -6.5),
+	float2(-3.5, -7.5), float2(+2.5, -7.5), float2(-7.5, +2.5), float2(-6.5, -2.5),
+	float2(-5.5, +5.5), float2(+2.5, +6.5), float2(-2.5, -6.5), float2(-7.5, +0.5),
+	float2(-0.5, +7.5), float2(+7.5, -2.5), float2(-2.5, +7.5), float2(+0.5, -7.5),
+	float2(-4.5, -7.5), float2(+7.5, +1.5), float2(+1.5, -6.5), float2(-6.5, +4.5),
+	float2(-1.5, +7.5), float2(-5.5, -5.5), float2(+6.5, +2.5), float2(-3.5, -6.5),
+	float2(+3.5, -7.5), float2(-5.5, +4.5), float2(+2.5, -6.5), float2(+1.5, -7.5),
+	float2(+6.5, +3.5), float2(+5.5, -6.5), float2(-6.5, +5.5), float2(+7.5, +4.5),
+	float2(+7.5, -1.5), float2(-7.5, -1.5), float2(+3.5, +7.5), float2(-5.5, +6.5),
+	float2(+1.5, +7.5), float2(+7.5, +3.5), float2(+7.5, -0.5), float2(-7.5, -2.5),
+	float2(+5.5, +5.5), float2(+6.5, +5.5), float2(+5.5, -5.5), float2(-2.5, -7.5),
+	float2(+2.5, +7.5), float2(-7.5, -3.5), float2(-7.5, -4.5), float2(-6.5, -4.5),
+	float2(+7.5, -3.5), float2(+5.5, +6.5), float2(-5.5, -6.5), float2(-4.5, -6.5),
+	float2(+7.5, +2.5), float2(-7.5, +3.5), float2(+4.5, -6.5), float2(+7.5, -4.5),
+};
+
+static const uint kStackowiakSampleSetCount = 4;
+static const uint kStackowiakSampleCountPerSet = 56;
+
+
 //------------------------------------------------------- PARAMETERS
 
+uint MaxSampleCount;
+
 Texture2D SignalInput0;
 Texture2D SignalInput1;
 
@@ -89,43 +187,65 @@ void MainCS(
 	{
 		SceneBufferUV = clamp(SceneBufferUV, View.BufferBilinearUVMinMax.xy, View.BufferBilinearUVMinMax.zw);
 	}
-	float2 ViewportUV = BufferUVToViewportUV(SceneBufferUV);
-	float2 ScreenPosition = ViewportUVToScreenPos(ViewportUV);
-
 
 	// Read reference data.
 	FSSFSignalSample Accumulator = SampleInputSignal(SignalInput0, SignalInput1, GlobalPointClampedSampler, CONFIG_SIGNAL_INPUT_LAYOUT, SceneBufferUV);
-	FSSFSampleSceneInfos CurrentSceneMetadata = FetchCurrentSceneInfos(ScreenPosition, SceneBufferUV);
+	float AccumulatorWeight = 1;
 
-	float CurrentWorldBluringRadius = ComputePixelWorldRadius(CurrentSceneMetadata);
+	// Read reference meta data.
+	float2 ScreenPosition = ViewportUVToScreenPos(BufferUVToViewportUV(SceneBufferUV));
+	FSSFSampleSceneInfos RefSceneMetadata = FetchCurrentSceneInfos(ScreenPosition, SceneBufferUV);
+
+	float CurrentWorldBluringRadius = ComputePixelWorldRadius(RefSceneMetadata);
 
-	float AccumulatorWeight = 1;
 	float BilateralWeightAdjustement = 1;
 	#if CONFIG_WEIGHT_KERNEL_CENTER
 	{
-		BilateralWeightAdjustement = rcp(ComputeBilateralWeight(100, CurrentSceneMetadata, CurrentSceneMetadata));
+		BilateralWeightAdjustement = rcp(ComputeBilateralWeight(100, RefSceneMetadata, RefSceneMetadata));
 	}
 	#endif
 
-	#if CONFIG_SAMPLE_SET == SAMPLE_SET_3X3
-		const int KernelRadius = 1;
-	#elif CONFIG_SAMPLE_SET == SAMPLE_SET_5X5_WAVELET
-		const int KernelRadius = 2;
+	float2 KernelCenterBufferUV = SceneBufferUV;
+
+	#if SAMPLE_SET_NONE
+		if (0)
+
+	#elif CONFIG_SAMPLE_SET == SAMPLE_SET_3X3 || CONFIG_SAMPLE_SET == SAMPLE_SET_5X5_WAVELET
+		const int KernelRadius = CONFIG_SAMPLE_SET == SAMPLE_SET_3X3 ? 1 : 2;
+		
+		UNROLL for (int x = -KernelRadius; x <= KernelRadius; x++)
+		UNROLL for (int y = -KernelRadius; y <= KernelRadius; y++)
+
+	#elif CONFIG_SAMPLE_SET == SAMPLE_SET_STACKOWIAK_4_SETS
+		KernelCenterBufferUV = ((DispatchThreadId | 1) + View.ViewRectMin.xy) * View.BufferSizeAndInvSize.zw;
+		KernelCenterBufferUV = clamp(KernelCenterBufferUV, View.BufferBilinearUVMinMax.xy, View.BufferBilinearUVMinMax.zw);
+
+		uint SampleCount = 16;
+
+		// Id of the pixel in the quad.
+		uint QuadSampleId = (DispatchThreadId.x & 1) | ((DispatchThreadId.y & 1) << 1);
+
+		LOOP for (uint SampleId = 1; SampleId < clamp(SampleCount, 1, MaxSampleCount); SampleId++)
+
 	#else
-		#error Unknown set of samples.
+		#error Unknown sample set
 	#endif
-
-	// Apply the kernel
-	// TODO: per pixel stocastic important sampled one.
-	UNROLL
-	for (int x = -KernelRadius; x <= KernelRadius; x++)
-	UNROLL
-	for (int y = -KernelRadius; y <= KernelRadius; y++)
 	{
-		if (x == 0 && y == 0) continue;
+		float2 PixelOffset = 0;
+
+		#if CONFIG_SAMPLE_SET == SAMPLE_SET_3X3 || CONFIG_SAMPLE_SET == SAMPLE_SET_5X5_WAVELET
+			if (x == 0 && y == 0) continue;
+			PixelOffset = float2(x, y);
+			
+		#elif CONFIG_SAMPLE_SET == SAMPLE_SET_STACKOWIAK_4_SETS
+			PixelOffset = kStackowiakSamples[kStackowiakSampleSetCount * SampleId + QuadSampleId];
 
-		// Compue sample coordinates
-		float2 NeighborBufferUV = SceneBufferUV + float2(x, y) * View.BufferSizeAndInvSize.zw;
+		#elif CONFIG_SAMPLE_SET != SAMPLE_SET_NONE
+			#error Unknown sample set
+		#endif
+
+		// Compute sample coordinates
+		float2 NeighborBufferUV = KernelCenterBufferUV + PixelOffset * View.BufferSizeAndInvSize.zw;
 		if (1)
 		{
 			if (1 && any(NeighborBufferUV != clamp(NeighborBufferUV, View.BufferBilinearUVMinMax.xy, View.BufferBilinearUVMinMax.zw)))
@@ -147,24 +267,22 @@ void MainCS(
 		// Compute the bilateral weight.
 		float BilateralWeight = BilateralWeightAdjustement * ComputeBilateralWeight(
 			WorldBluringRadiusToBilateralWorldDistance(min(CurrentWorldBluringRadius, NeighborWorldBluringRadius)),
-			CurrentSceneMetadata,
+			RefSceneMetadata,
 			NeighborSceneMetadata);
 
 		// Multiple importance sampling weights.
-		float MISWeight = GetMultiImportanceSamplingWeight(CurrentSceneMetadata, NeighborSceneMetadata, NeighborSample);
+		float MISWeight = GetMultiImportanceSamplingWeight(RefSceneMetadata, NeighborSceneMetadata, NeighborSample);
 
 		// Compute the weight of the sample in the kernel.
-		#if CONFIG_SAMPLE_SET == SAMPLE_SET_3X3
-			float KernelWeight = 1;
-
-		#elif CONFIG_SAMPLE_SET == SAMPLE_SET_5X5_WAVELET
-			float KernelWeight = 
-				kWaveletFilterWeights[abs(x)] *
-				kWaveletFilterWeights[abs(y)] *
+		#if CONFIG_SAMPLE_SET == SAMPLE_SET_5X5_WAVELET
+			const float KernelWeight = 
+				kWaveletFilterWeights[abs(PixelOffset.x)] *
+				kWaveletFilterWeights[abs(PixelOffset.y)] *
 				rcp(kWaveletFilterWeights[0] * kWaveletFilterWeights[0]);
 
 		#else
-			#error Unknown set of samples.
+			const float KernelWeight = 1;
+
 		#endif
 
 		// Compute the final weight of that sample.
@@ -182,4 +300,4 @@ void MainCS(
 	// TODO:  LeaveRayCount = (LeaveRayCount - 1) * 9 / (9 - 2) post processing to reject when for history rejection
 
 	OutputSignal(SignalOutput0, SignalOutput1, CONFIG_SIGNAL_OUTPUT_LAYOUT, DispatchThreadId, Accumulator);
-}
+} // MainCS
diff --git a/Engine/Shaders/Private/ScreenSpaceDenoise/SSDTemporalAccumulation.usf b/Engine/Shaders/Private/ScreenSpaceDenoise/SSDTemporalAccumulation.usf
--- a/Engine/Shaders/Private/ScreenSpaceDenoise/SSDTemporalAccumulation.usf
+++ b/Engine/Shaders/Private/ScreenSpaceDenoise/SSDTemporalAccumulation.usf
@@ -54,8 +54,6 @@
 	#error Unimplemented signal processing.
 #endif
 
-#include "SSDSignalInterface.ush"
-
 
 //------------------------------------------------------- CONFIG DISABLED DEFAULTS
 
@@ -104,11 +102,9 @@
 #endif
 
 
-//------------------------------------------------------- MENDATORY CONFIG
-
-
-//------------------------------------------------------- DERIVES
+//------------------------------------------------------- INCLUDES DEPENDINGS ON THE CONFIG
 
+#include "SSDSignalInterface.ush"
 
 
 //------------------------------------------------------- PARAMETERS
diff --git a/Engine/Source/Runtime/Renderer/Private/ScreenSpaceDenoise.cpp b/Engine/Source/Runtime/Renderer/Private/ScreenSpaceDenoise.cpp
--- a/Engine/Source/Runtime/Renderer/Private/ScreenSpaceDenoise.cpp
+++ b/Engine/Source/Runtime/Renderer/Private/ScreenSpaceDenoise.cpp
@@ -33,6 +33,11 @@ static TAutoConsoleVariable<int32> CVarDoDebugPass(
 	TEXT("Adds denoiser's the debug pass."),
 	ECVF_RenderThreadSafe);
 
+static TAutoConsoleVariable<int32> CVarReflectionReconstructionSampleCount(
+	TEXT("r.Reflection.Denoise.ReconstructionSamples"), 16,
+	TEXT("Maximum number of samples for the reconstruction pass (default = 16)."),
+	ECVF_RenderThreadSafe);
+
 static TAutoConsoleVariable<int32> CVarReflectionHistoryConvolution(
 	TEXT("r.Reflection.Denoise.HistoryConvolution"), 1,
 	TEXT("Mode to use for history convolution.\n")
@@ -44,6 +49,9 @@ static TAutoConsoleVariable<int32> CVarReflectionHistoryConvolution(
 /** The maximum number of mip level supported in the denoiser. */
 static const int32 kMaxMipLevel = 4;
 
+/** Maximum number of sample per pixel supported in the stackowiak sample set. */
+static const int32 kStackowiakMaxSampleCountPerSet = 56;
+
 
 // ---------------------------------------------------- Globals
 
@@ -179,9 +187,38 @@ class FSSDSpatialAccumulationCS : public FScreenSpaceDenoisingShader
 	DECLARE_GLOBAL_SHADER(FSSDSpatialAccumulationCS);
 	SHADER_USE_PARAMETER_STRUCT(FSSDSpatialAccumulationCS, FScreenSpaceDenoisingShader);
 	
-	using FPermutationDomain = TShaderPermutationDomain<FSignalProcessingDim>;
+	enum class EStage
+	{
+		// Spatial kernel used to process raw input for the temporal accumulation.
+		ReConstruction,
+
+		// Spatial kernel used to post filter the temporal accumulation.
+		PostFiltering,
+
+		MAX
+	};
+
+	class FStageDim : SHADER_PERMUTATION_ENUM_CLASS("DIM_STAGE", EStage);
+
+	using FPermutationDomain = TShaderPermutationDomain<FSignalProcessingDim, FStageDim>;
+	
+	static bool ShouldCompilePermutation(const FGlobalShaderPermutationParameters& Parameters)
+	{
+		FPermutationDomain PermutationVector(Parameters.PermutationId);
+
+		// Only reflections have  reconstruction spatial accumulation.
+		if (PermutationVector.Get<FSignalProcessingDim>() != ESignalProcessing::Reflections && 
+			PermutationVector.Get<FStageDim>() == EStage::ReConstruction)
+		{
+			return false;
+		}
+
+		return FScreenSpaceDenoisingShader::ShouldCompilePermutation(Parameters);
+	}
 
 	BEGIN_SHADER_PARAMETER_STRUCT(FParameters, )
+		SHADER_PARAMETER(uint32, MaxSampleCount)
+
 		SHADER_PARAMETER_STRUCT_INCLUDE(FSSDCommonParameters, CommonParameters)
 
 		SHADER_PARAMETER_RDG_TEXTURE(Texture2D, SignalInput0)
@@ -413,6 +450,7 @@ static void DenoiseShadowPenumbra(
 
 		FSSDSpatialAccumulationCS::FPermutationDomain PermutationVector;
 		PermutationVector.Set<FSignalProcessingDim>(ESignalProcessing::Penumbra);
+		PermutationVector.Set<FSSDSpatialAccumulationCS::FStageDim>(FSSDSpatialAccumulationCS::EStage::PostFiltering);
 
 		TShaderMapRef<FSSDSpatialAccumulationCS> ComputeShader(View.ShaderMap, PermutationVector);
 		FComputeShaderUtils::AddPass(
@@ -609,11 +647,43 @@ IScreenSpaceDenoiser::FReflectionOutputs DenoiseReflections(
 		PrevFrameHistory = View.PrevViewInfo.ReflectionsHistory;
 	}
 
-	FRDGTextureRef SignalHistory0 = GraphBuilder.CreateTexture(SignalProcessingDesc, TEXT("SSDReflectionsHistory0"));
-	FRDGTextureRef SignalHistory1 = GraphBuilder.CreateTexture(SignalProcessingDesc, TEXT("SSDReflectionsHistory1"));
+	FRDGTextureRef SignalHistory0;
+	FRDGTextureRef SignalHistory1;
+	
+	// Spatial reconstruction with multiple important sampling to be more precise in the history rejection.
+	{
+		FRDGTextureRef SignalOutput0 = GraphBuilder.CreateTexture(SignalProcessingDesc, TEXT("SSDReflectionsReconstruction"));
+		FRDGTextureRef SignalOutput1 = GraphBuilder.CreateTexture(SignalProcessingDesc, TEXT("SSDReflectionsReconstruction"));
+
+		FSSDSpatialAccumulationCS::FParameters* PassParameters = GraphBuilder.AllocParameters<FSSDSpatialAccumulationCS::FParameters>();
+		PassParameters->MaxSampleCount = FMath::Clamp(CVarReflectionReconstructionSampleCount.GetValueOnRenderThread(), 1, kStackowiakMaxSampleCountPerSet);
+		PassParameters->CommonParameters = CommonParameters;
+		PassParameters->SignalInput0 = ReflectionInputs.Color;
+		PassParameters->SignalInput1 = ReflectionInputs.RayHitDistance;
+		PassParameters->SignalOutput0 = GraphBuilder.CreateUAV(SignalOutput0);
+		PassParameters->SignalOutput1 = GraphBuilder.CreateUAV(SignalOutput1);
+
+		FSSDSpatialAccumulationCS::FPermutationDomain PermutationVector;
+		PermutationVector.Set<FSignalProcessingDim>(ESignalProcessing::Reflections);
+		PermutationVector.Set<FSSDSpatialAccumulationCS::FStageDim>(FSSDSpatialAccumulationCS::EStage::ReConstruction);
+
+		TShaderMapRef<FSSDSpatialAccumulationCS> ComputeShader(View.ShaderMap, PermutationVector);
+		FComputeShaderUtils::AddPass(
+			GraphBuilder,
+			RDG_EVENT_NAME("SSD SpatialAccumulation(Reconstruction Samples=%i)", PassParameters->MaxSampleCount),
+			*ComputeShader,
+			PassParameters,
+			FComputeShaderUtils::GetGroupCount(DenoiseResolution, FComputeShaderUtils::kGolden2DGroupSize));
+
+		SignalHistory0 = SignalOutput0;
+		SignalHistory1 = SignalOutput1;
+	}
 
 	// Temporal pass.
 	{
+		FRDGTextureRef SignalOutput0 = GraphBuilder.CreateTexture(SignalProcessingDesc, TEXT("SSDReflectionsHistory0"));
+		FRDGTextureRef SignalOutput1 = GraphBuilder.CreateTexture(SignalProcessingDesc, TEXT("SSDReflectionsHistory1"));
+
 		const int32 MipLevel = 0;
 
 		FSSDTemporalAccumulationCS::FPermutationDomain PermutationVector;
@@ -625,8 +695,8 @@ IScreenSpaceDenoiser::FReflectionOutputs DenoiseReflections(
 		FSSDTemporalAccumulationCS::FParameters* PassParameters = GraphBuilder.AllocParameters<FSSDTemporalAccumulationCS::FParameters>();
 		PassParameters->CommonParameters = CommonParameters;
 		PassParameters->bCameraCut = View.bCameraCut || !PrevFrameHistory.RT[0].IsValid();
-		PassParameters->SignalInput0 = ReflectionInputs.Color;
-		PassParameters->SignalInput1 = ReflectionInputs.RayHitDistance;
+		PassParameters->SignalInput0 = SignalHistory0;
+		PassParameters->SignalInput1 = SignalHistory1;
 
 		PassParameters->PrevScreenToTranslatedWorld = View.PrevViewInfo.ViewMatrices.GetInvTranslatedViewProjectionMatrix();
 		PassParameters->PrevHistory0 = RegisterExternalTextureWithFallback(GraphBuilder, PrevFrameHistory.RT[0], GSystemTextures.BlackDummy);
@@ -636,8 +706,8 @@ IScreenSpaceDenoiser::FReflectionOutputs DenoiseReflections(
 		PassParameters->PrevGBufferA = RegisterExternalTextureWithFallback(GraphBuilder, View.PrevViewInfo.GBufferA, GSystemTextures.BlackDummy);
 		PassParameters->PrevGBufferB = RegisterExternalTextureWithFallback(GraphBuilder, View.PrevViewInfo.GBufferB, GSystemTextures.BlackDummy);
 
-		PassParameters->SignalHistoryOutput0 = GraphBuilder.CreateUAV(FRDGTextureUAVDesc(SignalHistory0, MipLevel));
-		PassParameters->SignalHistoryOutput1 = GraphBuilder.CreateUAV(FRDGTextureUAVDesc(SignalHistory1, MipLevel));
+		PassParameters->SignalHistoryOutput0 = GraphBuilder.CreateUAV(FRDGTextureUAVDesc(SignalOutput0, MipLevel));
+		PassParameters->SignalHistoryOutput1 = GraphBuilder.CreateUAV(FRDGTextureUAVDesc(SignalOutput1, MipLevel));
 
 		FComputeShaderUtils::AddPass(
 			GraphBuilder,
@@ -645,6 +715,9 @@ IScreenSpaceDenoiser::FReflectionOutputs DenoiseReflections(
 			*ComputeShader,
 			PassParameters,
 			FComputeShaderUtils::GetGroupCount(DenoiseResolution, FComputeShaderUtils::kGolden2DGroupSize));
+
+		SignalHistory0 = SignalOutput0;
+		SignalHistory1 = SignalOutput1;
 	}
 	
 	// Spatial filter, to converge history faster.
@@ -662,11 +735,12 @@ IScreenSpaceDenoiser::FReflectionOutputs DenoiseReflections(
 
 		FSSDSpatialAccumulationCS::FPermutationDomain PermutationVector;
 		PermutationVector.Set<FSignalProcessingDim>(ESignalProcessing::Reflections);
+		PermutationVector.Set<FSSDSpatialAccumulationCS::FStageDim>(FSSDSpatialAccumulationCS::EStage::PostFiltering);
 
 		TShaderMapRef<FSSDSpatialAccumulationCS> ComputeShader(View.ShaderMap, PermutationVector);
 		FComputeShaderUtils::AddPass(
 			GraphBuilder,
-			RDG_EVENT_NAME("SSD SpatialAccumulation"),
+			RDG_EVENT_NAME("SSD SpatialAccumulation(PostFiltering)"),
 			*ComputeShader,
 			PassParameters,
 			FComputeShaderUtils::GetGroupCount(DenoiseResolution, FComputeShaderUtils::kGolden2DGroupSize));
