Changes the reflection ray generation shader to output ray hit distance.

diff --git a/Engine/Shaders/Private/RayTracing/RayTracingReflections.usf b/Engine/Shaders/Private/RayTracing/RayTracingReflections.usf
--- a/Engine/Shaders/Private/RayTracing/RayTracingReflections.usf
+++ b/Engine/Shaders/Private/RayTracing/RayTracingReflections.usf
@@ -38,6 +38,7 @@ float ReflectionMaxRayDistance;
 RaytracingAccelerationStructure TLAS;
 
 RWTexture2D<float4> ColorOutput;
+RWTexture2D<float> RayHitDistanceOutput;
 
 uint CalcLinearIndex(uint2 PixelCoord)
 {
@@ -287,18 +288,6 @@ FMaterialClosestHitPayload TraceRayAndAccumulateResults(
 	return Payload;
 }
 
-uint ReverseNBits(uint Bitfield, const uint BitCount) // TODO: where to put that guy?
-{
-	uint Out = 0;
-	UNROLL
-	for (uint i = 0; i < BitCount; i++)
-	{
-		uint bIsSet = ((Bitfield & (0x1 << i)) >> i);
-		Out |= bIsSet << (BitCount - 1 - i);
-	}
-	return Out;
-}
-
 
 [shader("raygeneration")]
 void RayTracingReflectionsRGS()
@@ -334,7 +323,7 @@ void RayTracingReflectionsRGS()
 
 	float4 ReflectedColor = 0.0;
 	float ReflectedWeight = 0.0;
-	float HitDistance = 0.0;
+	float HitDistance = -1.0;
 
 	if (Roughness < RoughnessThreshold)
 	{
@@ -354,10 +343,10 @@ void RayTracingReflectionsRGS()
 				if (LocalSamplesPerPixel == 1)
 				{
 					const uint TemporalBits = 10;
-					uint FrameIndexMod1024 = ReverseNBits(GetPowerOfTwoModulatedFrameIndex(1 << TemporalBits), TemporalBits);
+					uint FrameIndexMod1024 = ReverseBitsN(GetPowerOfTwoModulatedFrameIndex(1 << TemporalBits), TemporalBits);
 
 					uint2 SobolBase = SobolPixel(PixelCoord);
-					uint2 SobolFrame = SobolIndex(SobolBase, FrameIndexMod1024, TemporalBits);
+					uint2 SobolFrame = ComputePixelUniqueSobolRandSample(PixelCoord);
 
 					RandSample = SobolIndexToUniformUnitSquare(SobolFrame);
 				}
@@ -365,7 +354,7 @@ void RayTracingReflectionsRGS()
 				Ray.Direction = GenerateReflectedRayDirection(IncidentDirection, WorldNormal, Roughness, RandSample);
 				if (dot(WorldNormal, Ray.Direction) <= 0.0)
 				{
-					// Refraction rays is handled separately
+					// Refraction rays are handled separately
 					break;
 				}
 
@@ -419,6 +408,7 @@ void RayTracingReflectionsRGS()
 	}
 
 	ColorOutput[PixelCoord] = float4(ReflectedColor);
+	RayHitDistanceOutput[PixelCoord] = HitDistance;
 }
 
 
diff --git a/Engine/Source/Runtime/Renderer/Private/RayTracing/RayTracingReflections.cpp b/Engine/Source/Runtime/Renderer/Private/RayTracing/RayTracingReflections.cpp
--- a/Engine/Source/Runtime/Renderer/Private/RayTracing/RayTracingReflections.cpp
+++ b/Engine/Source/Runtime/Renderer/Private/RayTracing/RayTracingReflections.cpp
@@ -197,12 +197,17 @@ void FDeferredShadingSceneRenderer::RayTraceReflections(
 #if RHI_RAYTRACING
 {
 	FSceneRenderTargets& SceneContext = FSceneRenderTargets::Get(GraphBuilder.RHICmdList);
-	FPooledRenderTargetDesc Desc = SceneContext.GetSceneColor()->GetDesc();
-	Desc.Format = PF_FloatRGBA;
-	Desc.Flags &= ~(TexCreate_FastVRAM | TexCreate_Transient);
 
-	*OutColorTexture = GraphBuilder.CreateTexture(Desc, TEXT("RayTracingReflections"));
-	*OutRayHitDistanceTexture = GraphBuilder.CreateTexture(Desc, TEXT("RayTracingReflections"));
+	{		
+		FPooledRenderTargetDesc Desc = SceneContext.GetSceneColor()->GetDesc();
+		Desc.Format = PF_FloatRGBA;
+		Desc.Flags &= ~(TexCreate_FastVRAM | TexCreate_Transient);
+
+		*OutColorTexture = GraphBuilder.CreateTexture(Desc, TEXT("RayTracingReflections"));
+		
+		Desc.Format = PF_R16F;
+		*OutRayHitDistanceTexture = GraphBuilder.CreateTexture(Desc, TEXT("RayTracingReflectionsHitDistance"));
+	}
 
 	FRayTracingReflectionsRG::FParameters* PassParameters = GraphBuilder.AllocParameters<FRayTracingReflectionsRG::FParameters>();
 
