Implements ratio estimator

diff --git a/Engine/Shaders/Private/Common.ush b/Engine/Shaders/Private/Common.ush
--- a/Engine/Shaders/Private/Common.ush
+++ b/Engine/Shaders/Private/Common.ush
@@ -1647,6 +1647,7 @@ float DecodePackedHeight(float2 PackedHeight)
 	return DecodeHeightValue(DecodePackedTwoChannelValue(PackedHeight));
 }
 
+/** Reverses all the 32 bits. */
 uint ReverseBits32( uint bits )
 {
 #if SM5_PROFILE || COMPILER_METAL
@@ -1661,6 +1662,19 @@ uint ReverseBits32( uint bits )
 #endif
 }
 
+/** Reverses all the <BitCount> lowest significant bits. */
+uint ReverseBitsN(uint Bitfield, const uint BitCount)
+{
+	uint Out = 0;
+	UNROLL
+	for (uint i = 0; i < BitCount; i++)
+	{
+		uint bIsSet = ((Bitfield & (0x1 << i)) >> i);
+		Out |= bIsSet << (BitCount - 1 - i);
+	}
+	return Out;
+}
+
 // see PixelShaderOutputCommon
 struct FPixelShaderIn
 {
diff --git a/Engine/Shaders/Private/MonteCarlo.ush b/Engine/Shaders/Private/MonteCarlo.ush
--- a/Engine/Shaders/Private/MonteCarlo.ush
+++ b/Engine/Shaders/Private/MonteCarlo.ush
@@ -123,6 +123,16 @@ uint2 SobolIndex(uint2 Base, int Index, int Bits = 10)
 	return Result;
 }
 
+/** Returns an unique sobol sample for the frame. */
+uint2 ComputePixelUniqueSobolRandSample(uint2 PixelCoord)
+{
+	const uint TemporalBits = 10;
+	uint FrameIndexMod1024 = ReverseBitsN(GetPowerOfTwoModulatedFrameIndex(1 << TemporalBits), TemporalBits);
+
+	uint2 SobolBase = SobolPixel(PixelCoord);
+	return SobolIndex(SobolBase, FrameIndexMod1024, TemporalBits);
+}
+
 
 ///////
 
diff --git a/Engine/Shaders/Private/RayTracing/RayTracingOcclusionRGS.usf b/Engine/Shaders/Private/RayTracing/RayTracingOcclusionRGS.usf
--- a/Engine/Shaders/Private/RayTracing/RayTracingOcclusionRGS.usf
+++ b/Engine/Shaders/Private/RayTracing/RayTracingOcclusionRGS.usf
@@ -161,18 +161,6 @@ bool GenerateOcclusionRay(
 	return true;
 }
 
-uint ReverseNBits(uint Bitfield, const uint BitCount)
-{
-	uint Out = 0;
-	UNROLL
-	for (uint i = 0; i < BitCount; i++)
-	{
-		uint bIsSet = ((Bitfield & (0x1 << i)) >> i);
-		Out |= bIsSet << (BitCount - 1 - i);
-	}
-	return Out;
-}
-
 [shader("raygeneration")]
 void OcclusionRGS()
 {
@@ -214,7 +202,7 @@ void OcclusionRGS()
 		if (LocalSamplesPerPixel == 1 && 1)
 		{
 			const uint TemporalBits = 10;
-			uint FrameIndexMod1024 = ReverseNBits(GetPowerOfTwoModulatedFrameIndex(1 << TemporalBits), TemporalBits);
+			uint FrameIndexMod1024 = ReverseBitsN(GetPowerOfTwoModulatedFrameIndex(1 << TemporalBits), TemporalBits);
 
 			uint2 SobolBase = SobolPixel(PixelCoord);
 			uint2 SobolFrame = SobolIndex(SobolBase, FrameIndexMod1024, TemporalBits);
diff --git a/Engine/Shaders/Private/ScreenSpaceDenoise/SSDPenumbraSignal.ush b/Engine/Shaders/Private/ScreenSpaceDenoise/SSDPenumbraSignal.ush
--- a/Engine/Shaders/Private/ScreenSpaceDenoise/SSDPenumbraSignal.ush
+++ b/Engine/Shaders/Private/ScreenSpaceDenoise/SSDPenumbraSignal.ush
@@ -155,7 +155,8 @@ float GetSamplePenumbra(FSSFSignalSample Sample)
 float GetMultiImportanceSamplingWeight(
 	FSSFSampleSceneInfos RefSceneMetadata,
 	FSSFSampleSceneInfos SampleSceneMetadata,
-	FSSFSignalSample Sample)
+	FSSFSignalSample Sample,
+	uint2 SamplePixelCoord)
 {
 	return 1.0;
 }
diff --git a/Engine/Shaders/Private/ScreenSpaceDenoise/SSDReflectionsSignal.ush b/Engine/Shaders/Private/ScreenSpaceDenoise/SSDReflectionsSignal.ush
--- a/Engine/Shaders/Private/ScreenSpaceDenoise/SSDReflectionsSignal.ush
+++ b/Engine/Shaders/Private/ScreenSpaceDenoise/SSDReflectionsSignal.ush
@@ -5,6 +5,8 @@
 
 #include "SSDCommon.ush"
 #include "SSDSignalCore.ush"
+#include "../MonteCarlo.ush"
+#include "../BRDF.ush"
 
 
 //------------------------------------------------------- CONFIG DISABLED DEFAULTS
@@ -141,16 +143,107 @@ float GetSamplePenumbra(FSSFSignalSample Sample)
 	return Sample.LeavingRayCount / Sample.AccumulatedSampleCount;
 }
 
+// TODO: duplicated...
+float3 GenerateReflectedRayDirection(
+	float3 IncidentDirection,
+	float3 WorldNormal,
+	float Roughness,
+	float2 RandSample,
+	out float3 RayDirection,
+	out float PDF)
+{
+	if (Roughness < 0.05 ) //ReflectionSmoothClamp)
+	{
+		RayDirection = reflect(IncidentDirection, WorldNormal);
+		PDF = 1.0;
+	}
+	else
+	{
+		float NoV = saturate(dot(-IncidentDirection, WorldNormal));
+		float2 E = RandSample;
+		float4 Sample = ImportanceSampleGGX(E, Pow4(Roughness));
+		float3 H = TangentToWorld(Sample.xyz, WorldNormal);
+		RayDirection = reflect(IncidentDirection, H);
+
+		PDF = Sample.w;
+	}
+
+	return RayDirection;
+}
+
+// TODO: duplicated...
+float2 SobolIndexToUniformUnitSquare(uint2 SobolRand)
+{
+	return float2(SobolRand) * rcp(65536.0) + rcp(65536.0 * 2.0);
+}
+
+// TODO: duplicated...
+void GetInputRayDirectionPixel(
+	uint2 PixelCoord,
+	FSSFSampleSceneInfos SceneMetadata,
+	out float3 RayDirection,
+	out float PDF)
+{
+	float3 IncidentDirection = normalize(SceneMetadata.TranslatedWorldPosition - View.TranslatedWorldCameraOrigin);
+
+	float2 RandSample;
+	{
+		const uint TemporalBits = 10;
+		uint FrameIndexMod1024 = ReverseBitsN(GetPowerOfTwoModulatedFrameIndex(1 << TemporalBits), TemporalBits);
+
+		uint2 SobolBase = SobolPixel(PixelCoord);
+		uint2 SobolFrame = SobolIndex(SobolBase, FrameIndexMod1024, TemporalBits);
+
+		RandSample = SobolIndexToUniformUnitSquare(SobolFrame);
+	}
+
+	GenerateReflectedRayDirection(
+		IncidentDirection,
+		SceneMetadata.WorldNormal,
+		SceneMetadata.Roughness,
+		RandSample,
+		/* out = */ RayDirection,
+		/* out = */ PDF);
+}
+
+// [ Stackowiak 2018, "Stochastic all the things: Raytracing in hybrid real-time rendering" ]
 float GetMultiImportanceSamplingWeight(
 	FSSFSampleSceneInfos RefSceneMetadata,
 	FSSFSampleSceneInfos SampleSceneMetadata,
-	FSSFSignalSample Sample)
+	FSSFSignalSample Sample,
+	uint2 SamplePixelCoord)
 #if CONFIG_SPECULAR_MIS
 {
-	return 1.0;
+	if (View.GeneralPurposeTweak > 0.5)
+	{
+		return 1;
+	}
+
+	float3 OriginalRayDirection;
+	float OriginalRayPDF;
+
+	GetInputRayDirectionPixel(
+		SamplePixelCoord,
+		SampleSceneMetadata, 
+		/* out */ OriginalRayDirection,
+		/* out */ OriginalRayPDF);
+
+	float3 HitTranslatedWorldPosition = SampleSceneMetadata.TranslatedWorldPosition + OriginalRayDirection * Sample.PenumbraClosestHit;
+
+	float a2 = Pow4(RefSceneMetadata.Roughness);
+
+	float3 N = RefSceneMetadata.WorldNormal;
+	float3 V = normalize(View.TranslatedWorldCameraOrigin - RefSceneMetadata.TranslatedWorldPosition);
+	float3 L = normalize(HitTranslatedWorldPosition - RefSceneMetadata.TranslatedWorldPosition);
+	float3 H = normalize(L + V);
+
+	float NoH = abs(dot(N, H));
+	float D = D_GGX(a2, NoH);
+
+	return D / OriginalRayPDF;
 }
-#else // CONFIG_MIS
+#else // !CONFIG_SPECULAR_MIS
 {
 	return 1.0;
 }
-#endif
+#endif // CONFIG_SPECULAR_MIS
diff --git a/Engine/Shaders/Private/ScreenSpaceDenoise/SSDSignalCore.ush b/Engine/Shaders/Private/ScreenSpaceDenoise/SSDSignalCore.ush
--- a/Engine/Shaders/Private/ScreenSpaceDenoise/SSDSignalCore.ush
+++ b/Engine/Shaders/Private/ScreenSpaceDenoise/SSDSignalCore.ush
@@ -124,7 +124,9 @@ FSSFSignalSample DecodeInputSignal(const uint SignalBufferLayout, float4 Samples
 		OutSample.AccumulatedSampleCount = Samples[0].a > 0 ? 1 : 0;
 		OutSample.LeavingRayCount = OutSample.AccumulatedSampleCount * (1 - Samples[0].a);
 		OutSample.SceneColor = Samples[0];
-		OutSample.PenumbraClosestHit = Samples[1].r;
+
+		// max() because ray hit distance is < 0 when miss / not ray has been shot.
+		OutSample.PenumbraClosestHit = max(Samples[1].r, 0);
 	}
 	else if (SignalBufferLayout == SIGNAL_BUFFER_LAYOUT_REFLECTIONS_HISTORY)
 	{
diff --git a/Engine/Shaders/Private/ScreenSpaceDenoise/SSDSpatialAccumulation.usf b/Engine/Shaders/Private/ScreenSpaceDenoise/SSDSpatialAccumulation.usf
--- a/Engine/Shaders/Private/ScreenSpaceDenoise/SSDSpatialAccumulation.usf
+++ b/Engine/Shaders/Private/ScreenSpaceDenoise/SSDSpatialAccumulation.usf
@@ -187,15 +187,17 @@ void MainCS(
 	{
 		SceneBufferUV = clamp(SceneBufferUV, View.BufferBilinearUVMinMax.xy, View.BufferBilinearUVMinMax.zw);
 	}
-
-	// Read reference data.
-	FSSFSignalSample Accumulator = SampleInputSignal(SignalInput0, SignalInput1, GlobalPointClampedSampler, CONFIG_SIGNAL_INPUT_LAYOUT, SceneBufferUV);
-	float AccumulatorWeight = 1;
-
+	
 	// Read reference meta data.
 	float2 ScreenPosition = ViewportUVToScreenPos(BufferUVToViewportUV(SceneBufferUV));
 	FSSFSampleSceneInfos RefSceneMetadata = FetchCurrentSceneInfos(ScreenPosition, SceneBufferUV);
 
+	FSSFSignalSample RefSample = SampleInputSignal(SignalInput0, SignalInput1, GlobalPointClampedSampler, CONFIG_SIGNAL_INPUT_LAYOUT, SceneBufferUV);
+
+	// Read reference data.
+	float AccumulatorWeight = 1; //GetMultiImportanceSamplingWeight(RefSceneMetadata, RefSceneMetadata, RefSample, DispatchThreadId);
+	FSSFSignalSample Accumulator = Mul(RefSample, AccumulatorWeight);
+	
 	float CurrentWorldBluringRadius = ComputePixelWorldRadius(RefSceneMetadata);
 
 	float BilateralWeightAdjustement = 1;
@@ -220,7 +222,7 @@ void MainCS(
 		KernelCenterBufferUV = ((DispatchThreadId | 1) + View.ViewRectMin.xy) * View.BufferSizeAndInvSize.zw;
 		KernelCenterBufferUV = clamp(KernelCenterBufferUV, View.BufferBilinearUVMinMax.xy, View.BufferBilinearUVMinMax.zw);
 
-		uint SampleCount = 16;
+		uint SampleCount = kStackowiakSampleCountPerSet; // TODO
 
 		// Id of the pixel in the quad.
 		uint QuadSampleId = (DispatchThreadId.x & 1) | ((DispatchThreadId.y & 1) << 1);
@@ -271,7 +273,8 @@ void MainCS(
 			NeighborSceneMetadata);
 
 		// Multiple importance sampling weights.
-		float MISWeight = GetMultiImportanceSamplingWeight(RefSceneMetadata, NeighborSceneMetadata, NeighborSample);
+		uint2 NeighborPixelCoord = floor(NeighborBufferUV * View.BufferSizeAndInvSize.xy);
+		float MISWeight = GetMultiImportanceSamplingWeight(RefSceneMetadata, NeighborSceneMetadata, NeighborSample, NeighborPixelCoord);
 
 		// Compute the weight of the sample in the kernel.
 		#if CONFIG_SAMPLE_SET == SAMPLE_SET_5X5_WAVELET
@@ -286,7 +289,7 @@ void MainCS(
 		#endif
 
 		// Compute the final weight of that sample.
-		float FinalWeight = BilateralWeight * MISWeight * KernelWeight;
+		float FinalWeight = BilateralWeight * MISWeight * KernelWeight;// * (NeighborSample > 0.0 ? 1 : 0);
 		
 		Accumulator = Add(Accumulator, Mul(NeighborSample, FinalWeight));
 	}
