Reorganises screen space denoiser's shader directory.

diff --git a/Engine/Shaders/Private/ScreenSpaceDenoise/SSDBufferLayouts.ush b/Engine/Shaders/Private/ScreenSpaceDenoise/SSDBufferLayouts.ush
deleted file mode 100644
diff --git a/Engine/Shaders/Private/ScreenSpaceDenoise/SSDBuildHistoryRejection.usf b/Engine/Shaders/Private/ScreenSpaceDenoise/SSDBuildHistoryRejection.usf
--- a/Engine/Shaders/Private/ScreenSpaceDenoise/SSDBuildHistoryRejection.usf
+++ b/Engine/Shaders/Private/ScreenSpaceDenoise/SSDBuildHistoryRejection.usf
@@ -1,21 +1,14 @@
 // Copyright 1998-2019 Epic Games, Inc. All Rights Reserved.
 
-/*=============================================================================
-SSDReduce.usf:
-=============================================================================*/
-
-#include "../MonteCarlo.ush"
-#include "../TextureSampling.ush"
-#include "SSDSignalCore.ush"
-
-#define CONFIG_SIGNAL_PROCESSING SIGNAL_PROCESSING_PENUMBRA
-#include "SSDSignalInterface.ush"
+#include "SSDDefinitions.ush"
 
 
 //------------------------------------------------------- CONFIGS
 
 #define TILE_PIXEL_SIZE 8
 
+#define CONFIG_SIGNAL_PROCESSING SIGNAL_PROCESSING_PENUMBRA
+
 #if CONFIG_SIGNAL_PROCESSING == SIGNAL_PROCESSING_PENUMBRA
 	#define CONFIG_SIGNAL_RAW_INPUT_LAYOUT	SIGNAL_BUFFER_LAYOUT_PENUMBRA_INPUT
 	#define CONFIG_SIGNAL_REDUCE_LAYOUT		SIGNAL_BUFFER_LAYOUT_PENUMBRA_INTERNAL
@@ -24,7 +17,15 @@ SSDReduce.usf:
 #else
 	#error Unknown signal processing
 
-#endif 
+#endif
+
+
+//------------------------------------------------------- INCLUDES
+
+#include "SSDSignalFramework.ush"
+
+#include "../MonteCarlo.ush"
+#include "../TextureSampling.ush"
 
 
 //------------------------------------------------------- PARAMETERS
diff --git a/Engine/Shaders/Private/ScreenSpaceDenoise/SSDCommon.ush b/Engine/Shaders/Private/ScreenSpaceDenoise/SSDCommon.ush
--- a/Engine/Shaders/Private/ScreenSpaceDenoise/SSDCommon.ush
+++ b/Engine/Shaders/Private/ScreenSpaceDenoise/SSDCommon.ush
@@ -10,74 +10,9 @@
 #include "../DataStructures/SceneViewFamilyBlackboard.ush"
 #include "../ColorSpace.ush"
 #include "../SphericalGaussian.ush"
+#include "SSDDefinitions.ush"
 
 
-//------------------------------------------------------- DEBUG CONFIG
-
-/** Turn on/off tile classfication for debugging purpose. */
-#define DEBUG_ENABLE_TILE_CLASSIFICATION 1
-
-/** Should output invalid data when know there is a need. */
-// TODO: there is an artifact in temporal reprojection.
-#define DEBUG_WASTFUL_OUTPUT 1
-
-/** Force the signal to a specific mip level. */
-#define DEBUG_MIP_LEVEL -1
-
-/** Debuging mode. */
-	#define DEBUG_MODE_DISABLED -1
-	#define DEBUG_MODE_BRUTE_FORCE_ONLY 0
-	#define DEBUG_MODE_FAST_ONLY 1
-#define DEBUG_MODE DEBUG_MODE_DISABLED
-
-/** DIsable any bilateral rejection. */
-#define DEBUG_DISABLE_BILATERAL 0
-
-/** Prefer to store normalized components in internal texture. */
-#define DEBUG_STORE_NORMALIZED_COMPONENTS 1
-
-
-//------------------------------------------------------- COMPILER CONFIG
-
-// Generate vector truncation warnings to errors.
-#pragma warning(error: 3206)
-
-
-//------------------------------------------------------- SHARED CONFIG
-
-/** Do normal based anysotropy. */
-#define CONFIG_NORMAL_ANISOTROPY 1
-
-
-//------------------------------------------------------- SHARED CONSTANTS
-
-// Maximum number of mip level supported.
-#define MAX_MIP_LEVEL_COUNT 4
-
-#define TARGETED_SAMPLE_COUNT 1024
-
-// This constant is to artificially store the signal into higher resolution that theorically should from
-// the infered bluring radius by deviding it. 1.0 means nothing happen, but in practice want this >=
-// than Nyquist's 2.0 multiplier. Uses 3.0 because this is the size of the 3x3 neighborhood kernel.
-#define BLURING_ESTIMATION_MULTIPLIER 3.0
-
-// World bluring radius for a miss.
-#define WORLD_RADIUS_INVALID -1
-#define WORLD_RADIUS_MISS 1.#INF
-
-// Mip level for sample that are invalid or missed.
-#define MIP_LEVEL_INVALID -1
-#define MIP_LEVEL_MISS 1024
-
-
-//------------------------------------------------------- SHARED ENUMS
-
-/** Light types */
-	#define LIGHT_TYPE_DIRECTIONAL 0
-	#define LIGHT_TYPE_POINT 1
-	#define LIGHT_TYPE_SPOT 2
-	#define LIGHT_TYPE_RECT 3
-
 
 //------------------------------------------------------- SHARED FUNCTION
 
diff --git a/Engine/Shaders/Private/ScreenSpaceDenoise/SSDDefinitions.ush b/Engine/Shaders/Private/ScreenSpaceDenoise/SSDDefinitions.ush
new file mode 100644
--- /dev/null
+++ b/Engine/Shaders/Private/ScreenSpaceDenoise/SSDDefinitions.ush
@@ -0,0 +1,106 @@
+// Copyright 1998-2019 Epic Games, Inc. All Rights Reserved.
+
+#pragma once
+
+//------------------------------------------------------- COMPILER CONFIG
+
+// Generate vector truncation warnings to errors.
+#pragma warning(error: 3206)
+
+
+//------------------------------------------------------- DEBUG CONFIG
+
+/** Turn on/off tile classfication for debugging purpose. */
+#define DEBUG_ENABLE_TILE_CLASSIFICATION 1
+
+/** Should output invalid data when know there is a need. */
+// TODO: there is an artifact in temporal reprojection.
+#define DEBUG_WASTFUL_OUTPUT 1
+
+/** Force the signal to a specific mip level. */
+#define DEBUG_MIP_LEVEL -1
+
+/** Debuging mode. */
+	#define DEBUG_MODE_DISABLED -1
+	#define DEBUG_MODE_BRUTE_FORCE_ONLY 0
+	#define DEBUG_MODE_FAST_ONLY 1
+#define DEBUG_MODE DEBUG_MODE_DISABLED
+
+/** DIsable any bilateral rejection. */
+#define DEBUG_DISABLE_BILATERAL 0
+
+/** Prefer to store normalized components in internal texture. */
+#define DEBUG_STORE_NORMALIZED_COMPONENTS 1
+
+
+//------------------------------------------------------- SHARED CONFIG
+
+/** Do normal based anysotropy. */
+#define CONFIG_NORMAL_ANISOTROPY 1
+
+
+//------------------------------------------------------- SHARED CONSTANTS
+
+// Maximum number of mip level supported.
+#define MAX_MIP_LEVEL_COUNT 4
+
+#define TARGETED_SAMPLE_COUNT 1024
+
+// This constant is to artificially store the signal into higher resolution that theorically should from
+// the infered bluring radius by deviding it. 1.0 means nothing happen, but in practice want this >=
+// than Nyquist's 2.0 multiplier. Uses 3.0 because this is the size of the 3x3 neighborhood kernel.
+#define BLURING_ESTIMATION_MULTIPLIER 3.0
+
+// World bluring radius for a miss.
+#define WORLD_RADIUS_INVALID -1
+#define WORLD_RADIUS_MISS 1.#INF
+
+// Mip level for sample that are invalid or missed.
+#define MIP_LEVEL_INVALID -1
+#define MIP_LEVEL_MISS 1024
+
+
+//------------------------------------------------------- ENUMS
+
+/** Layouts of the signal buffer. */
+	/** Buffer layout for the shadow penumbra given as input. */
+	#define SIGNAL_BUFFER_LAYOUT_PENUMBRA_INPUT 0
+
+	/** Internal buffer layout for the shadow penumbra. */
+	#define SIGNAL_BUFFER_LAYOUT_PENUMBRA_INTERNAL 1
+
+	/** Internal buffer layout used for history rejection for the shadow penumbra. */
+	#define SIGNAL_BUFFER_LAYOUT_PENUMBRA_HISTORY_REJECTION 2
+
+	/** Buffer layout taken as input by */ // TODO: unused, remove it.
+	// #define SIGNAL_BUFFER_LAYOUT_PENUMBRA_OUTPUT 3
+
+	/** Buffer layout for the reflection given by the ray generation shader. */
+	#define SIGNAL_BUFFER_LAYOUT_REFLECTIONS_INPUT 3
+
+	/** Buffer layout for the reflection output. */
+	#define SIGNAL_BUFFER_LAYOUT_REFLECTIONS_HISTORY 4
+
+
+/** Defines how the signal is being processed. Matches C++'s ESignalProcessing. */
+	/** Scene color anti aliasing processing. */
+	//#define SIGNAL_PROCESSING_COLORAA 0
+
+	/** Depth of field stabilization processing. */
+	//#define SIGNAL_PROCESSING_DOF 1
+
+	/** Shadow penumbra denoising. */
+	#define SIGNAL_PROCESSING_PENUMBRA 0
+
+	/** Shadow penumbra denoising. */
+	//#define SIGNAL_PROCESSING_DIFUSE_GI 3
+
+	/** Reflection denoising. */
+	#define SIGNAL_PROCESSING_REFLECTIONS 1
+
+
+/** Light types */
+	#define LIGHT_TYPE_DIRECTIONAL 0
+	#define LIGHT_TYPE_POINT 1
+	#define LIGHT_TYPE_SPOT 2
+	#define LIGHT_TYPE_RECT 3
diff --git a/Engine/Shaders/Private/ScreenSpaceDenoise/SSDOutput.usf b/Engine/Shaders/Private/ScreenSpaceDenoise/SSDOutput.usf
--- a/Engine/Shaders/Private/ScreenSpaceDenoise/SSDOutput.usf
+++ b/Engine/Shaders/Private/ScreenSpaceDenoise/SSDOutput.usf
@@ -1,11 +1,6 @@
 // Copyright 1998-2019 Epic Games, Inc. All Rights Reserved.
 
-/*=============================================================================
-SSDOutput.usf: Output the denoising signal for external consumer
-=============================================================================*/
-
-#include "SSDSignalCore.ush"
-#include "../TextureSampling.ush"
+#include "SSDDefinitions.ush"
 
 
 //------------------------------------------------------- DEBUG CONFIG
@@ -31,7 +26,11 @@ SSDOutput.usf: Output the denoising signal for external consumer
 
 #endif
 
-#include "SSDSignalInterface.ush"
+
+//------------------------------------------------------- INCLUDES
+
+#include "SSDSignalFramework.ush"
+#include "../TextureSampling.ush"
 
 
 //------------------------------------------------------- PARAMETERS
diff --git a/Engine/Shaders/Private/ScreenSpaceDenoise/SSDPenumbraSignal.ush b/Engine/Shaders/Private/ScreenSpaceDenoise/SSDPenumbraSignal.ush
deleted file mode 100644
diff --git a/Engine/Shaders/Private/ScreenSpaceDenoise/SSDReduce.usf b/Engine/Shaders/Private/ScreenSpaceDenoise/SSDReduce.usf
--- a/Engine/Shaders/Private/ScreenSpaceDenoise/SSDReduce.usf
+++ b/Engine/Shaders/Private/ScreenSpaceDenoise/SSDReduce.usf
@@ -1,20 +1,14 @@
 // Copyright 1998-2019 Epic Games, Inc. All Rights Reserved.
 
-/*=============================================================================
-SSDReduce.usf:
-=============================================================================*/
-
-#include "../MonteCarlo.ush"
-#include "SSDSignalCore.ush"
-
-#define CONFIG_SIGNAL_PROCESSING SIGNAL_PROCESSING_PENUMBRA
-#include "SSDSignalInterface.ush"
+#include "SSDDefinitions.ush"
 
 
 //------------------------------------------------------- CONFIGS
 
 #define TILE_PIXEL_SIZE 8
 
+#define CONFIG_SIGNAL_PROCESSING SIGNAL_PROCESSING_PENUMBRA
+
 #if CONFIG_SIGNAL_PROCESSING == SIGNAL_PROCESSING_PENUMBRA
 	#define CONFIG_SIGNAL_INPUT_LAYOUT		SIGNAL_BUFFER_LAYOUT_PENUMBRA_INPUT
 	#define CONFIG_SIGNAL_OUTPUT_LAYOUT		SIGNAL_BUFFER_LAYOUT_PENUMBRA_INTERNAL
@@ -25,6 +19,12 @@ SSDReduce.usf:
 #endif 
 
 
+//------------------------------------------------------- INCLUDES
+
+#include "SSDSignalFramework.ush"
+#include "../MonteCarlo.ush"
+
+
 //------------------------------------------------------- PARAMETERS
 
 Texture2D SignalInput0;
diff --git a/Engine/Shaders/Private/ScreenSpaceDenoise/SSDReflectionsSignal.ush b/Engine/Shaders/Private/ScreenSpaceDenoise/SSDReflectionsSignal.ush
deleted file mode 100644
diff --git a/Engine/Shaders/Private/ScreenSpaceDenoise/SSDSignalCore.ush b/Engine/Shaders/Private/ScreenSpaceDenoise/SSDSignalCore.ush
--- a/Engine/Shaders/Private/ScreenSpaceDenoise/SSDSignalCore.ush
+++ b/Engine/Shaders/Private/ScreenSpaceDenoise/SSDSignalCore.ush
@@ -9,45 +9,6 @@
 #include "SSDCommon.ush"
 
 
-//------------------------------------------------------- ENUMS
-
-/** Layouts of the signal buffer. */
-	/** Buffer layout for the shadow penumbra given as input. */
-	#define SIGNAL_BUFFER_LAYOUT_PENUMBRA_INPUT 0
-
-	/** Internal buffer layout for the shadow penumbra. */
-	#define SIGNAL_BUFFER_LAYOUT_PENUMBRA_INTERNAL 1
-
-	/** Internal buffer layout used for history rejection for the shadow penumbra. */
-	#define SIGNAL_BUFFER_LAYOUT_PENUMBRA_HISTORY_REJECTION 2
-
-	/** Buffer layout taken as input by */ // TODO: unused, remove it.
-	// #define SIGNAL_BUFFER_LAYOUT_PENUMBRA_OUTPUT 3
-
-	/** Buffer layout for the reflection given by the ray generation shader. */
-	#define SIGNAL_BUFFER_LAYOUT_REFLECTIONS_INPUT 3
-
-	/** Buffer layout for the reflection output. */
-	#define SIGNAL_BUFFER_LAYOUT_REFLECTIONS_HISTORY 4
-
-
-/** Defines how the signal is being processed. Matches C++'s ESignalProcessing. */
-	/** Scene color anti aliasing processing. */
-	//#define SIGNAL_PROCESSING_COLORAA 0
-
-	/** Depth of field stabilization processing. */
-	//#define SIGNAL_PROCESSING_DOF 1
-
-	/** Shadow penumbra denoising. */
-	#define SIGNAL_PROCESSING_PENUMBRA 0
-
-	/** Shadow penumbra denoising. */
-	//#define SIGNAL_PROCESSING_DIFUSE_GI 3
-
-	/** Reflection denoising. */
-	#define SIGNAL_PROCESSING_REFLECTIONS 1
-
-
 //------------------------------------------------------- STRUCTURE
 
 /** Generic data structure to manipulate any kind of signal. */
@@ -282,11 +243,6 @@ FSSFSignalSample Power(FSSFSignalSample Sample, float Exponent)
 	return Sample;
 }
 
-FSSFSignalSample Suctract(FSSFSignalSample SampleA, FSSFSignalSample SampleB)
-{
-	return Add(SampleA, Minus(SampleB));
-}
-
 FSSFSignalSample Min(FSSFSignalSample SampleA, FSSFSignalSample SampleB)
 {
 	FSSFSignalSample OutSample;
@@ -322,7 +278,7 @@ FSSFSignalSample Clamp(FSSFSignalSample Sample, FSSFSignalSample SampleMin, FSSF
 
 FSSFSignalSample Lerp(FSSFSignalSample Sample0, FSSFSignalSample Sample1, float Interp)
 {
-	return Add(Sample0, Mul(Suctract(Sample1, Sample0), Interp));
+	return Add(Sample0, Mul(Subtract(Sample1, Sample0), Interp));
 }
 
 /** Normalize the signal sample to one. */
diff --git a/Engine/Shaders/Private/ScreenSpaceDenoise/SSDSignalFramework.ush b/Engine/Shaders/Private/ScreenSpaceDenoise/SSDSignalFramework.ush
new file mode 100644
--- /dev/null
+++ b/Engine/Shaders/Private/ScreenSpaceDenoise/SSDSignalFramework.ush
@@ -0,0 +1,357 @@
+// Copyright 1998-2019 Epic Games, Inc. All Rights Reserved.
+
+#pragma once
+
+
+//------------------------------------------------------- CONFIG DISABLED DEFAULTS
+
+#ifndef CONFIG_TOKOYASHI_BILATERAL
+	#define CONFIG_TOKOYASHI_BILATERAL 0
+#endif
+
+#ifndef CONFIG_SPECULAR_MIS
+	#define CONFIG_SPECULAR_MIS 0
+#endif
+
+
+//------------------------------------------------------- INCLUDE
+
+#include "SSDCommon.ush"
+#include "SSDSignalCore.ush"
+
+#if !defined(CONFIG_SIGNAL_PROCESSING)
+	#error Missing CONFIG_SIGNAL_PROCESSING
+
+#elif CONFIG_SIGNAL_PROCESSING == SIGNAL_PROCESSING_PENUMBRA
+	#include "../LightShaderParameters.ush"
+	
+#elif CONFIG_SIGNAL_PROCESSING == SIGNAL_PROCESSING_REFLECTIONS
+	#include "../MonteCarlo.ush"
+	#include "../BRDF.ush"
+
+#else
+	#error Unknown CONFIG_SIGNAL_PROCESSING
+
+#endif
+
+
+//------------------------------------------------------- INTERNAL FOR PENUMBRA SIGNAL
+#if CONFIG_SIGNAL_PROCESSING == SIGNAL_PROCESSING_PENUMBRA
+
+float HitDistanceToWorldBluringRadius;
+
+#endif // CONFIG_SIGNAL_PROCESSING == SIGNAL_PROCESSING_PENUMBRA
+
+
+//------------------------------------------------------- INTERNAL FOR REFLECTION SIGNAL
+#if CONFIG_SIGNAL_PROCESSING == SIGNAL_PROCESSING_REFLECTIONS
+
+/** Create a spherical gaussian lob for a given roughness. */
+FSphericalGaussian RoughnessLobe(float Roughness, float3 N, float3 V)
+{
+	float a = Pow2( max( 0.02, Roughness ) );
+	float a2 = a*a;
+	
+	float NoV = saturate( abs( dot(N, V) ) + 1e-5 );
+	
+	float3 R = 2 * NoV * N - V;
+
+	FSphericalGaussian SpecularSG;
+	SpecularSG.Axis = R;
+	SpecularSG.Sharpness = 0.5 / ( a2 * max( NoV, 0.1 ) );
+	SpecularSG.Amplitude = rcp( PI * a2 );
+
+	return SpecularSG;
+}
+
+/** Create a spherical gaussian lob for a given scene metadata. */
+FSphericalGaussian RoughnessLobe(FSSFSampleSceneInfos RefSceneMetadata)
+{
+	float3 N = RefSceneMetadata.WorldNormal;
+	float3 V = normalize(View.TranslatedWorldCameraOrigin - RefSceneMetadata.TranslatedWorldPosition);
+	float Roughness = RefSceneMetadata.Roughness;
+
+	return RoughnessLobe(Roughness, N, V);
+}
+
+// [ Tokoyashi 2015, "Specular Lobe-Aware Filtering and Upsampling for Interactive Indirect Illumination" ]
+float TokoyashiLobeBilateralWeight(FSSFSampleSceneInfos RefSceneMetadata, FSSFSampleSceneInfos NeighborSceneMetadata)
+{
+	FSphericalGaussian Ref = RoughnessLobe(RefSceneMetadata);
+	FSphericalGaussian Neighbor = RoughnessLobe(NeighborSceneMetadata);
+	
+#if 1
+	const float Beta = 5.05; // TODO: to be tunned with GGX.
+
+	float CombinedSharpness = sqrt(Ref.Sharpness * Neighbor.Sharpness) / (Ref.Sharpness + Neighbor.Sharpness);
+
+	return pow(2 * CombinedSharpness, Beta) * exp((Beta * CombinedSharpness) * (dot(Ref.Axis, Neighbor.Axis) - 1));
+#else
+	return Dot(Ref, Neighbor);
+#endif
+}
+
+// TODO: duplicated...
+float3 GenerateReflectedRayDirection(
+	float3 IncidentDirection,
+	float3 WorldNormal,
+	float Roughness,
+	float2 RandSample,
+	out float3 RayDirection,
+	out float PDF)
+{
+	if (Roughness < 0.05 ) //ReflectionSmoothClamp)
+	{
+		RayDirection = reflect(IncidentDirection, WorldNormal);
+		PDF = 1.0;
+	}
+	else
+	{
+		float NoV = saturate(dot(-IncidentDirection, WorldNormal));
+		float2 E = RandSample;
+		float4 Sample = ImportanceSampleGGX(E, Pow4(Roughness));
+		float3 H = TangentToWorld(Sample.xyz, WorldNormal);
+		RayDirection = reflect(IncidentDirection, H);
+
+		PDF = Sample.w;
+	}
+
+	return RayDirection;
+}
+
+// TODO: duplicated...
+float2 SobolIndexToUniformUnitSquare(uint2 SobolRand)
+{
+	return float2(SobolRand) * rcp(65536.0) + rcp(65536.0 * 2.0);
+}
+
+// TODO: duplicated...
+void GetInputRayDirectionPixel(
+	uint2 PixelCoord,
+	FSSFSampleSceneInfos SceneMetadata,
+	out float3 RayDirection,
+	out float PDF)
+{
+	float3 IncidentDirection = normalize(SceneMetadata.TranslatedWorldPosition - View.TranslatedWorldCameraOrigin);
+
+	float2 RandSample;
+	{
+		const uint TemporalBits = 10;
+		uint FrameIndexMod1024 = ReverseBitsN(GetPowerOfTwoModulatedFrameIndex(1 << TemporalBits), TemporalBits);
+
+		uint2 SobolBase = SobolPixel(PixelCoord);
+		uint2 SobolFrame = SobolIndex(SobolBase, FrameIndexMod1024, TemporalBits);
+
+		RandSample = SobolIndexToUniformUnitSquare(SobolFrame);
+	}
+
+	GenerateReflectedRayDirection(
+		IncidentDirection,
+		SceneMetadata.WorldNormal,
+		SceneMetadata.Roughness,
+		RandSample,
+		/* out = */ RayDirection,
+		/* out = */ PDF);
+}
+
+
+#endif // CONFIG_SIGNAL_PROCESSING == SIGNAL_PROCESSING_REFLECTIONS
+
+
+//------------------------------------------------------- MAIN FUNCTIONS TO IMPLEMENT FOR THE PASSES
+
+/** Return whether this signal is occluded or not. */
+bool IsMissSample(FSSFSignalSample Sample)
+{
+	return Sample.LeavingRayCount != 0; // TODO: use PenumbraClosestHit instead.
+}
+
+/** Return whether this signal is valid or not. */
+bool IsInvalidSample(FSSFSignalSample Sample)
+{
+	return Sample.PenumbraClosestHit == -1.0; // TODO: AccumulatedSampleCount == 0
+}
+
+/** COmpute the weight of the bilateral filter to use between a reference sample and a neighbor. */
+float ComputeBilateralWeight(
+	float MaxWorldBluringDistance,
+	FSSFSampleSceneInfos RefSceneMetadata,
+	FSSFSampleSceneInfos NeighborSceneMetadata)
+#if CONFIG_SIGNAL_PROCESSING == SIGNAL_PROCESSING_PENUMBRA
+{
+	if (DEBUG_DISABLE_BILATERAL)
+	{
+		return 1;
+	}
+	
+	float WorldBilateralWeight = 1;
+	if (1) // TODO.
+	{
+		const float WorldRadius = MaxWorldBluringDistance;
+
+		float3 V = RefSceneMetadata.TranslatedWorldPosition - NeighborSceneMetadata.TranslatedWorldPosition;
+
+		float RefAnisotropyInvFactor = ComputeAnisotropyInvFactor(RefSceneMetadata);
+
+		float3 CameraToRef = View.ViewForward;
+
+		float Z = dot(CameraToRef, V);
+		float XY = length2(V - CameraToRef * Z * (1 - RefAnisotropyInvFactor));
+
+		float DistSquare = XY;
+
+		float Multiplier = rcp(WorldRadius * WorldRadius);
+
+		WorldBilateralWeight = saturate(1 - DistSquare * Multiplier);
+	}
+	
+	float LobeBilateralWeight = 1;
+	#if CONFIG_TOKOYASHI_BILATERAL
+	//if (View.GeneralPurposeTweak > 0.5)
+	{
+		LobeBilateralWeight = TokoyashiLobeBilateralWeight(RefSceneMetadata, NeighborSceneMetadata);
+	}
+	#endif
+
+	return WorldBilateralWeight;
+}
+#elif CONFIG_SIGNAL_PROCESSING == SIGNAL_PROCESSING_REFLECTIONS
+{
+	return 1;
+}
+#endif
+
+float GetSignalWorldBluringRadius(FSSFSignalSample Sample, FSSFSampleSceneInfos SceneMetadata)
+#if CONFIG_SIGNAL_PROCESSING == SIGNAL_PROCESSING_PENUMBRA
+{
+	if (IsInvalidSample(Sample))
+	{
+		return WORLD_RADIUS_INVALID;
+	}
+	else if (IsMissSample(Sample))
+	{
+		return WORLD_RADIUS_MISS;
+	}
+
+	// Compute the anysotropy.
+	float Anisotropy = ComputeAnisotropyInvFactor(SceneMetadata);
+
+	FLightShaderParameters LightParameters = GetRootLightShaderParameters();
+
+	#if !defined(DIM_LIGHT_TYPE)
+	{
+		return 0;
+	}
+	#elif DIM_LIGHT_TYPE == LIGHT_TYPE_DIRECTIONAL
+	{
+		return HitDistanceToWorldBluringRadius * Sample.PenumbraClosestHit * Anisotropy;
+	}
+	#elif DIM_LIGHT_TYPE == LIGHT_TYPE_POINT
+	{
+		return 0; // TODO: point light soft shadows are not implemented.
+	}
+	#elif DIM_LIGHT_TYPE == LIGHT_TYPE_SPOT
+	{
+		return 0; // TODO: spot light soft shadows are not implemented.
+	}
+	#elif DIM_LIGHT_TYPE == LIGHT_TYPE_RECT
+	{
+		float3 PixelToLightWorldVector = (LightParameters.Position + View.PreViewTranslation) - SceneMetadata.TranslatedWorldPosition;
+		float DistanceFromLight = length(PixelToLightWorldVector);
+
+		float Scalar = abs(dot(normalize(PixelToLightWorldVector), LightParameters.Direction));
+
+		float2 LightDimensions = GetRectLightDimensions(LightParameters);
+
+		// TODO: area light anysotropy.
+		float SmallestLightDimension = min(LightDimensions.x, LightDimensions.y);
+
+		// TODO: Sample.PenumbraClosestHit depends on the direction of the ray, but does not DistanceFromLight, witch is bad for large area light.
+		return SmallestLightDimension * Sample.PenumbraClosestHit / (DistanceFromLight - Sample.PenumbraClosestHit) * Anisotropy;
+	}
+	#endif
+}
+#elif CONFIG_SIGNAL_PROCESSING == SIGNAL_PROCESSING_REFLECTIONS
+{
+	return 1;
+}
+#endif
+
+FSSFSignalSample GetConstantSignalSample(uint TileMode)
+{
+	FSSFSignalSample OutputSample = CreateSignalSample();
+	if (TileMode == TILE_MODE_FULLY_INVALID)
+	{
+		OutputSample.AccumulatedSampleCount = 0.0;
+		OutputSample.LeavingRayCount = 0.0;
+	}
+	else if (TileMode == TILE_MODE_FULLY_LIGHTED)
+	{
+		OutputSample.AccumulatedSampleCount = 1.0;
+		OutputSample.LeavingRayCount = OutputSample.AccumulatedSampleCount;
+	}
+	else if (TileMode == TILE_MODE_FULLY_OCCLUDED)
+	{
+		OutputSample.AccumulatedSampleCount = 1.0;
+		OutputSample.LeavingRayCount = 0.0;
+	}
+	else if (TileMode == TILE_MODE_FULLY_UNSURE)
+	{
+		OutputSample.UncertainRayCount = 1.0;
+	}
+	return OutputSample;
+}
+
+// Returns the penumbra of this sample, or 1 if invalid.
+float GetSamplePenumbraSafe(FSSFSignalSample Sample)
+{
+	return (Sample.AccumulatedSampleCount > 0 ? Sample.LeavingRayCount / Sample.AccumulatedSampleCount : 1);
+}
+
+// Get penumbra, but be aware of NaN.
+float GetSamplePenumbra(FSSFSignalSample Sample)
+{
+	if (1) // TODO: For debugging purpose of division by 0.
+	{
+		return GetSamplePenumbraSafe(Sample);
+	}
+	return Sample.LeavingRayCount / Sample.AccumulatedSampleCount;
+}
+
+/** Get the weight for multi importance sampling. */
+float GetMultiImportanceSamplingWeight(
+	FSSFSampleSceneInfos RefSceneMetadata,
+	FSSFSampleSceneInfos SampleSceneMetadata,
+	FSSFSignalSample Sample,
+	uint2 SamplePixelCoord)
+#if CONFIG_SPECULAR_MIS
+// [ Stackowiak 2018, "Stochastic all the things: Raytracing in hybrid real-time rendering" ]
+{
+	float3 OriginalRayDirection;
+	float OriginalRayPDF;
+
+	GetInputRayDirectionPixel(
+		SamplePixelCoord,
+		SampleSceneMetadata, 
+		/* out */ OriginalRayDirection,
+		/* out */ OriginalRayPDF);
+
+	float3 HitTranslatedWorldPosition = SampleSceneMetadata.TranslatedWorldPosition + OriginalRayDirection * Sample.PenumbraClosestHit;
+
+	float a2 = Pow4(RefSceneMetadata.Roughness);
+
+	float3 N = RefSceneMetadata.WorldNormal;
+	float3 V = normalize(View.TranslatedWorldCameraOrigin - RefSceneMetadata.TranslatedWorldPosition);
+	float3 L = normalize(HitTranslatedWorldPosition - RefSceneMetadata.TranslatedWorldPosition);
+	float3 H = normalize(L + V);
+
+	float NoH = abs(dot(N, H));
+	float D = D_GGX(a2, NoH);
+
+	return D / OriginalRayPDF;
+}
+#else // !CONFIG_SPECULAR_MIS
+{
+	return 1.0;
+}
+#endif // CONFIG_SPECULAR_MIS
diff --git a/Engine/Shaders/Private/ScreenSpaceDenoise/SSDSignalInterface.ush b/Engine/Shaders/Private/ScreenSpaceDenoise/SSDSignalInterface.ush
deleted file mode 100644
diff --git a/Engine/Shaders/Private/ScreenSpaceDenoise/SSDSpatialAccumulation.usf b/Engine/Shaders/Private/ScreenSpaceDenoise/SSDSpatialAccumulation.usf
--- a/Engine/Shaders/Private/ScreenSpaceDenoise/SSDSpatialAccumulation.usf
+++ b/Engine/Shaders/Private/ScreenSpaceDenoise/SSDSpatialAccumulation.usf
@@ -1,11 +1,6 @@
 // Copyright 1998-2019 Epic Games, Inc. All Rights Reserved.
 
-/*=============================================================================
-PostProcessTemporalAA.usf: Temporal AA
-=============================================================================*/
-
-#include "SSDCommon.ush"
-#include "SSDSignalCore.ush"
+#include "SSDDefinitions.ush"
 
 
 //------------------------------------------------------- ENUM VALUES
@@ -83,9 +78,9 @@ PostProcessTemporalAA.usf: Temporal AA
 #endif
 
 
-//------------------------------------------------------- INCLUDES DEPENDINGS ON THE CONFIG
+//------------------------------------------------------- INCLUDES
 
-#include "SSDSignalInterface.ush"
+#include "SSDSignalFramework.ush"
 
 
 //------------------------------------------------------- CONSTANTS
diff --git a/Engine/Shaders/Private/ScreenSpaceDenoise/SSDTemporalAccumulation.usf b/Engine/Shaders/Private/ScreenSpaceDenoise/SSDTemporalAccumulation.usf
--- a/Engine/Shaders/Private/ScreenSpaceDenoise/SSDTemporalAccumulation.usf
+++ b/Engine/Shaders/Private/ScreenSpaceDenoise/SSDTemporalAccumulation.usf
@@ -1,12 +1,6 @@
 // Copyright 1998-2019 Epic Games, Inc. All Rights Reserved.
 
-/*=============================================================================
-	PostProcessTemporalAA.usf: Temporal AA
-=============================================================================*/
-
-#include "SSDCommon.ush"
-#include "SSDSignalCore.ush"
-#include "../TextureSampling.ush"
+#include "SSDDefinitions.ush"
 
 
 //------------------------------------------------------- CONFIGS
@@ -102,9 +96,10 @@
 #endif
 
 
-//------------------------------------------------------- INCLUDES DEPENDINGS ON THE CONFIG
+//------------------------------------------------------- INCLUDES
 
-#include "SSDSignalInterface.ush"
+#include "SSDSignalFramework.ush"
+#include "../TextureSampling.ush"
 
 
 //------------------------------------------------------- PARAMETERS
diff --git a/Engine/Shaders/Private/ScreenSpaceDenoise/SSDTileDilate.usf b/Engine/Shaders/Private/ScreenSpaceDenoise/SSDTileDilate.usf
--- a/Engine/Shaders/Private/ScreenSpaceDenoise/SSDTileDilate.usf
+++ b/Engine/Shaders/Private/ScreenSpaceDenoise/SSDTileDilate.usf
@@ -1,11 +1,7 @@
 // Copyright 1998-2019 Epic Games, Inc. All Rights Reserved.
 
-/*=============================================================================
-SSDReduce.usf:
-=============================================================================*/
-
-#include "../MonteCarlo.ush"
 #include "SSDSignalCore.ush"
+#include "../MonteCarlo.ush"
 
 
 //------------------------------------------------------- CONFIGS
