Fixes a denoiser divergence when reflection have not find anything and there is no skylight.

diff --git a/Engine/Shaders/Private/RayTracing/RayTracingReflections.usf b/Engine/Shaders/Private/RayTracing/RayTracingReflections.usf
--- a/Engine/Shaders/Private/RayTracing/RayTracingReflections.usf
+++ b/Engine/Shaders/Private/RayTracing/RayTracingReflections.usf
@@ -323,7 +323,7 @@ void RayTracingReflectionsRGS()
 
 	float4 ReflectedColor = 0.0;
 	float ReflectedWeight = 0.0;
-	float HitDistance = -1.0;
+	float HitDistance = -2.0;
 
 	if (Roughness < RoughnessThreshold)
 	{
@@ -342,12 +342,7 @@ void RayTracingReflectionsRGS()
 
 				if (LocalSamplesPerPixel == 1)
 				{
-					const uint TemporalBits = 10;
-					uint FrameIndexMod1024 = ReverseBitsN(GetPowerOfTwoModulatedFrameIndex(1 << TemporalBits), TemporalBits);
-
-					uint2 SobolBase = SobolPixel(PixelCoord);
 					uint2 SobolFrame = ComputePixelUniqueSobolRandSample(PixelCoord);
-
 					RandSample = SobolIndexToUniformUnitSquare(SobolFrame);
 				}
 
diff --git a/Engine/Shaders/Private/ScreenSpaceDenoise/SSDSignalCore.ush b/Engine/Shaders/Private/ScreenSpaceDenoise/SSDSignalCore.ush
--- a/Engine/Shaders/Private/ScreenSpaceDenoise/SSDSignalCore.ush
+++ b/Engine/Shaders/Private/ScreenSpaceDenoise/SSDSignalCore.ush
@@ -82,12 +82,12 @@ FSSFSignalSample DecodeInputSignal(const uint SignalBufferLayout, float4 Samples
 	}
 	else if (SignalBufferLayout == SIGNAL_BUFFER_LAYOUT_REFLECTIONS_INPUT)
 	{
-		OutSample.AccumulatedSampleCount = Samples[0].a > 0 ? 1 : 0;
+		OutSample.AccumulatedSampleCount = Samples[1].r == -2 ? 0 : 1;
 		OutSample.LeavingRayCount = OutSample.AccumulatedSampleCount * (1 - Samples[0].a);
 		OutSample.SceneColor = Samples[0];
 
 		// max() because ray hit distance is < 0 when miss / not ray has been shot.
-		OutSample.PenumbraClosestHit = max(Samples[1].r, 0);
+		OutSample.PenumbraClosestHit = Samples[1].r * OutSample.AccumulatedSampleCount;
 	}
 	else if (SignalBufferLayout == SIGNAL_BUFFER_LAYOUT_REFLECTIONS_HISTORY)
 	{
diff --git a/Engine/Shaders/Private/ScreenSpaceDenoise/SSDSignalFramework.ush b/Engine/Shaders/Private/ScreenSpaceDenoise/SSDSignalFramework.ush
--- a/Engine/Shaders/Private/ScreenSpaceDenoise/SSDSignalFramework.ush
+++ b/Engine/Shaders/Private/ScreenSpaceDenoise/SSDSignalFramework.ush
@@ -185,7 +185,7 @@ float ComputeBilateralWeight(
 	}
 	
 	float WorldBilateralWeight = 1;
-	if (1) // TODO.
+	if (1)
 	{
 		const float WorldRadius = MaxWorldBluringDistance;
 
@@ -217,7 +217,14 @@ float ComputeBilateralWeight(
 }
 #elif CONFIG_SIGNAL_PROCESSING == SIGNAL_PROCESSING_REFLECTIONS
 {
-	return 1;
+	float LobeBilateralWeight = 1;
+	#if CONFIG_TOKOYASHI_BILATERAL
+	{
+		LobeBilateralWeight = TokoyashiLobeBilateralWeight(RefSceneMetadata, NeighborSceneMetadata);
+	}
+	#endif
+
+	return LobeBilateralWeight;
 }
 #endif
 
@@ -327,6 +334,11 @@ float GetMultiImportanceSamplingWeight(
 #if CONFIG_SPECULAR_MIS
 // [ Stackowiak 2018, "Stochastic all the things: Raytracing in hybrid real-time rendering" ]
 {
+	// If invalid ray or geometric miss, just.
+	if (Sample.PenumbraClosestHit < 0.0)
+	{
+		return 1; 
+	}
 	float3 OriginalRayDirection;
 	float OriginalRayPDF;
 
