Computes temporal hysteresis on the reflection denoiser based on the roughness.

diff --git a/Engine/Shaders/Private/ScreenSpaceDenoise/SSDSignalFramework.ush b/Engine/Shaders/Private/ScreenSpaceDenoise/SSDSignalFramework.ush
--- a/Engine/Shaders/Private/ScreenSpaceDenoise/SSDSignalFramework.ush
+++ b/Engine/Shaders/Private/ScreenSpaceDenoise/SSDSignalFramework.ush
@@ -28,6 +28,8 @@
 #elif CONFIG_SIGNAL_PROCESSING == SIGNAL_PROCESSING_REFLECTIONS
 	#include "../MonteCarlo.ush"
 	#include "../BRDF.ush"
+	#include "../DeferredShadingCommon.ush"
+	#include "../CommonViewUniformBuffer.ush"
 
 #else
 	#error Unknown CONFIG_SIGNAL_PROCESSING
@@ -46,6 +48,25 @@ float HitDistanceToWorldBluringRadius;
 //------------------------------------------------------- INTERNAL FOR REFLECTION SIGNAL
 #if CONFIG_SIGNAL_PROCESSING == SIGNAL_PROCESSING_REFLECTIONS
 
+/** Return the size of the bluring radius caused by the specular lobe in ViewportUV.x space. */
+float ComputeInfinityBluringRadius(FSSFSampleSceneInfos SceneMetadata)
+{
+	// TODO: this should depend on NoV? Try up with spherical gauassian framework.
+	float RoughnessLobeAngle = ComputeAngleFromRoughness(SceneMetadata.Roughness);
+
+	return tan(0.5 * RoughnessLobeAngle) * GetCotanHalfFieldOfView().x;
+}
+
+/** Return the size of bluring radius caused by specular lobe in view space. */
+float ComputeBluringRadiusOfNormalizedSample(FSSFSignalSample Sample, FSSFSampleSceneInfos SceneMetadata)
+{
+	float DistanceCameraToPixel = length(View.TranslatedWorldCameraOrigin - SceneMetadata.TranslatedWorldPosition);
+
+	float InifinityRadius = ComputeInfinityBluringRadius(SceneMetadata);
+
+	return InifinityRadius * (1 - DistanceCameraToPixel / (DistanceCameraToPixel + Sample.PenumbraClosestHit));
+}
+
 /** Create a spherical gaussian lob for a given roughness. */
 FSphericalGaussian RoughnessLobe(float Roughness, float3 N, float3 V)
 {
diff --git a/Engine/Shaders/Private/ScreenSpaceDenoise/SSDTemporalAccumulation.usf b/Engine/Shaders/Private/ScreenSpaceDenoise/SSDTemporalAccumulation.usf
--- a/Engine/Shaders/Private/ScreenSpaceDenoise/SSDTemporalAccumulation.usf
+++ b/Engine/Shaders/Private/ScreenSpaceDenoise/SSDTemporalAccumulation.usf
@@ -137,6 +137,8 @@ Texture2D HistoryRejectionSignal1;
 RWTexture2D<float4>	SignalHistoryOutput0;
 RWTexture2D<float4>	SignalHistoryOutput1;
 
+RWTexture2D<float4>	DebugOutput;
+
 
 //------------------------------------------------------- FUNCTIONS
 
@@ -190,10 +192,6 @@ void TemporallyAccumulate(
 	FSSFTileInfos TileInfos,
 	const bool bBruteForce)
 {
-	// When brute forcing, reduce the number of samples trying to converge with to increase minimal histerisis and respond quicker to the heavy noise.
-	const float TargetedSampleCount = bBruteForce ? (0.125 * TARGETED_SAMPLE_COUNT) : (16 * MipLevelPow2 * MipLevelPow2);
-	//const float TargetedSampleCount = bBruteForce ? (0.125 * TARGETED_SAMPLE_COUNT) : TARGETED_SAMPLE_COUNT;
-
 	// Find out scene buffer UV.
 	float2 SceneBufferUV = ((DispatchThreadId + 0.5) * MipLevelPow2 + View.ViewRectMin.xy) * View.BufferSizeAndInvSize.zw;
 	if (true)
@@ -235,6 +233,27 @@ void TemporallyAccumulate(
 	// Use the diameter, because that is the distance between two pixel.
 	float PixelWorldBluringRadius = MipLevelPow2 * ComputePixelWorldRadius(CurrentSceneMetadata);
 	float WorldBluringRadius = WorldBluringRadiusToBilateralWorldDistance(PixelWorldBluringRadius);
+	
+	float TargetedSampleCount = TARGETED_SAMPLE_COUNT;
+	#if CONFIG_SIGNAL_PROCESSING == SIGNAL_PROCESSING_PENUMBRA
+	{
+		TargetedSampleCount = bBruteForce ? (0.125 * TARGETED_SAMPLE_COUNT) : (16 * MipLevelPow2 * MipLevelPow2);
+		//const float TargetedSampleCount = bBruteForce ? (0.125 * TARGETED_SAMPLE_COUNT) : TARGETED_SAMPLE_COUNT;
+	}
+	#elif CONFIG_SIGNAL_PROCESSING == SIGNAL_PROCESSING_REFLECTIONS
+	{
+		float SignalPixelBluringRadius = ComputeInfinityBluringRadius(CurrentSceneMetadata) * View.ViewSizeAndInvSize.x;
+
+		TargetedSampleCount = 2 + PI * SignalPixelBluringRadius * SignalPixelBluringRadius;
+
+		TargetedSampleCount = clamp(TargetedSampleCount, 1, TARGETED_SAMPLE_COUNT);
+
+		//DebugOutput[DispatchThreadId] = float4(
+		//	SignalPixelBluringRadius,
+		//	TargetedSampleCount,
+		//	0, 0);
+	}
+	#endif
 
 	// Sample the history buffer.
 	FSSFSignalSample PreviousFrame = CreateSignalSample();
diff --git a/Engine/Source/Runtime/Renderer/Private/ScreenSpaceDenoise.cpp b/Engine/Source/Runtime/Renderer/Private/ScreenSpaceDenoise.cpp
--- a/Engine/Source/Runtime/Renderer/Private/ScreenSpaceDenoise.cpp
+++ b/Engine/Source/Runtime/Renderer/Private/ScreenSpaceDenoise.cpp
@@ -285,6 +285,8 @@ class FSSDTemporalAccumulationCS : public FScreenSpaceDenoisingShader
 
 		SHADER_PARAMETER_RDG_TEXTURE_UAV(RWTexture2D, SignalHistoryOutput0)
 		SHADER_PARAMETER_RDG_TEXTURE_UAV(RWTexture2D, SignalHistoryOutput1)
+		
+		SHADER_PARAMETER_RDG_TEXTURE_UAV(RWTexture2D, DebugOutput) // TODO: remove
 	END_SHADER_PARAMETER_STRUCT()
 };
 
@@ -407,7 +409,7 @@ static void DenoiseShadowPenumbra(
 			PassParameters->TileClassificationOutput = GraphBuilder.CreateUAV(TileClassificationTexture);
 		}
 
-		PassParameters->DebugOutput = GraphBuilder.CreateUAV(GraphBuilder.CreateTexture(PenumbraRT->Desc, TEXT("SSDDebugReduce")));
+		PassParameters->DebugOutput = GraphBuilder.CreateUAV(GraphBuilder.CreateTexture(PenumbraRT->Desc, TEXT("SSDDebugPenumbraReduce")));
 
 		TShaderMapRef<FSSDReduceCS> ComputeShader(View.ShaderMap, PermutationVector);
 		FComputeShaderUtils::AddPass(
@@ -555,6 +557,8 @@ static void DenoiseShadowPenumbra(
 
 			PassParameters->SignalHistoryOutput0 = GraphBuilder.CreateUAV(FRDGTextureUAVDesc(SignalHistory0, MipLevel));
 			PassParameters->SignalHistoryOutput1 = GraphBuilder.CreateUAV(FRDGTextureUAVDesc(SignalHistory1, MipLevel));
+			
+			PassParameters->DebugOutput = GraphBuilder.CreateUAV(GraphBuilder.CreateTexture(PenumbraRT->Desc, TEXT("SSDDebugPenumbraTemporalAccumulation")));
 
 			FIntPoint Resolution = FIntPoint::DivideAndRoundUp(DenoiseResolution, 1 << MipLevel);
 
@@ -714,6 +718,8 @@ IScreenSpaceDenoiser::FReflectionOutputs DenoiseReflections(
 
 		PassParameters->SignalHistoryOutput0 = GraphBuilder.CreateUAV(FRDGTextureUAVDesc(SignalOutput0, MipLevel));
 		PassParameters->SignalHistoryOutput1 = GraphBuilder.CreateUAV(FRDGTextureUAVDesc(SignalOutput1, MipLevel));
+		
+		PassParameters->DebugOutput = GraphBuilder.CreateUAV(GraphBuilder.CreateTexture(SignalProcessingDesc, TEXT("SSDDebugReflectionTemporalAccumulation")));
 
 		FComputeShaderUtils::AddPass(
 			GraphBuilder,
