Fixes a regression in the reflection denoiser caused by 4714861

diff --git a/Engine/Shaders/Private/ScreenSpaceDenoise/SSDSignalFramework.ush b/Engine/Shaders/Private/ScreenSpaceDenoise/SSDSignalFramework.ush
--- a/Engine/Shaders/Private/ScreenSpaceDenoise/SSDSignalFramework.ush
+++ b/Engine/Shaders/Private/ScreenSpaceDenoise/SSDSignalFramework.ush
@@ -189,9 +189,15 @@ bool IsMissSample(FSSFSignalSample Sample)
 
 /** Return whether this signal is valid or not. */
 bool IsInvalidSample(FSSFSignalSample Sample)
+#if CONFIG_SIGNAL_PROCESSING == SIGNAL_PROCESSING_PENUMBRA
 {
 	return Sample.PenumbraClosestHit == -1.0; // TODO: AccumulatedSampleCount == 0
 }
+#else
+{
+	return Sample.AccumulatedSampleCount == 0.0;
+}
+#endif
 
 /** COmpute the weight of the bilateral filter to use between a reference sample and a neighbor. */
 float ComputeBilateralWeight(
@@ -356,10 +362,11 @@ float GetMultiImportanceSamplingWeight(
 // [ Stackowiak 2018, "Stochastic all the things: Raytracing in hybrid real-time rendering" ]
 {
 	// If invalid ray or geometric miss, just.
-	if (Sample.PenumbraClosestHit < 0.0)
+	if (IsInvalidSample(Sample))
 	{
-		return 1; 
+		return 0; 
 	}
+
 	float3 OriginalRayDirection;
 	float OriginalRayPDF;
 
@@ -369,7 +376,7 @@ float GetMultiImportanceSamplingWeight(
 		/* out */ OriginalRayDirection,
 		/* out */ OriginalRayPDF);
 
-	float3 HitTranslatedWorldPosition = SampleSceneMetadata.TranslatedWorldPosition + OriginalRayDirection * Sample.PenumbraClosestHit;
+	float3 HitTranslatedWorldPosition = SampleSceneMetadata.TranslatedWorldPosition + OriginalRayDirection * (IsMissSample(Sample) ? 1.0e27 : Sample.PenumbraClosestHit);
 
 	float a2 = Pow4(RefSceneMetadata.Roughness);
 
