Fixes a floating point underflow in Tokoyashi that could diverge the result in some cases.

diff --git a/Engine/Shaders/Private/ScreenSpaceDenoise/SSDSignalCore.ush b/Engine/Shaders/Private/ScreenSpaceDenoise/SSDSignalCore.ush
--- a/Engine/Shaders/Private/ScreenSpaceDenoise/SSDSignalCore.ush
+++ b/Engine/Shaders/Private/ScreenSpaceDenoise/SSDSignalCore.ush
@@ -88,6 +88,7 @@ FSSFSignalSample DecodeInputSignal(const uint SignalBufferLayout, float4 Samples
 
 		// max() because ray hit distance is < 0 when miss / not ray has been shot.
 		OutSample.PenumbraClosestHit = Samples[1].r * OutSample.AccumulatedSampleCount;
+		OutSample.PenumbraClosestHit = ((Samples[1].r == -1.0) ? 1.0e20 : Samples[1].r) * OutSample.AccumulatedSampleCount;
 	}
 	else if (SignalBufferLayout == SIGNAL_BUFFER_LAYOUT_REFLECTIONS_HISTORY)
 	{
diff --git a/Engine/Shaders/Private/ScreenSpaceDenoise/SSDSignalFramework.ush b/Engine/Shaders/Private/ScreenSpaceDenoise/SSDSignalFramework.ush
--- a/Engine/Shaders/Private/ScreenSpaceDenoise/SSDSignalFramework.ush
+++ b/Engine/Shaders/Private/ScreenSpaceDenoise/SSDSignalFramework.ush
@@ -182,10 +182,12 @@ void GetInputRayDirectionPixel(
 //------------------------------------------------------- MAIN FUNCTIONS TO IMPLEMENT FOR THE PASSES
 
 /** Return whether this signal is occluded or not. */
+#if CONFIG_SIGNAL_PROCESSING == SIGNAL_PROCESSING_PENUMBRA
 bool IsMissSample(FSSFSignalSample Sample)
 {
 	return Sample.LeavingRayCount != 0; // TODO: use PenumbraClosestHit instead.
 }
+#endif
 
 /** Return whether this signal is valid or not. */
 bool IsInvalidSample(FSSFSignalSample Sample)
@@ -362,7 +364,7 @@ float GetMultiImportanceSamplingWeight(
 // [ Stackowiak 2018, "Stochastic all the things: Raytracing in hybrid real-time rendering" ]
 {
 	// If invalid ray or geometric miss, just.
-	if (IsInvalidSample(Sample))
+	if (IsInvalidSample(Sample) && 0)
 	{
 		return 0; 
 	}
@@ -376,13 +378,14 @@ float GetMultiImportanceSamplingWeight(
 		/* out */ OriginalRayDirection,
 		/* out */ OriginalRayPDF);
 
-	float3 HitTranslatedWorldPosition = SampleSceneMetadata.TranslatedWorldPosition + OriginalRayDirection * (IsMissSample(Sample) ? 1.0e27 : Sample.PenumbraClosestHit);
+	// min to avoid floating point underflow that causes artifact by D anormally high.
+	float PenumbraClosestHit = min(Sample.PenumbraClosestHit, 1000000);
 
 	float a2 = Pow4(RefSceneMetadata.Roughness);
 
 	float3 N = RefSceneMetadata.WorldNormal;
 	float3 V = normalize(View.TranslatedWorldCameraOrigin - RefSceneMetadata.TranslatedWorldPosition);
-	float3 L = normalize(HitTranslatedWorldPosition - RefSceneMetadata.TranslatedWorldPosition);
+	float3 L = normalize((SampleSceneMetadata.TranslatedWorldPosition - RefSceneMetadata.TranslatedWorldPosition) + OriginalRayDirection * PenumbraClosestHit);
 	float3 H = normalize(L + V);
 
 	float NoH = abs(dot(N, H));
