Lowers the hysteresis of reflection denoiser to work arround a floating point precision that can create bandings.

diff --git a/Engine/Shaders/Private/ScreenSpaceDenoise/SSDTemporalAccumulation.usf b/Engine/Shaders/Private/ScreenSpaceDenoise/SSDTemporalAccumulation.usf
--- a/Engine/Shaders/Private/ScreenSpaceDenoise/SSDTemporalAccumulation.usf
+++ b/Engine/Shaders/Private/ScreenSpaceDenoise/SSDTemporalAccumulation.usf
@@ -246,7 +246,8 @@ void TemporallyAccumulate(
 
 		TargetedSampleCount = 2 + PI * SignalPixelBluringRadius * SignalPixelBluringRadius;
 
-		TargetedSampleCount = clamp(TargetedSampleCount, 1, TARGETED_SAMPLE_COUNT);
+		// TODO: try to make this better? * 0.125 to avoid banding caused by floating point precision.
+		TargetedSampleCount = clamp(TargetedSampleCount, 1, 0.125 * TARGETED_SAMPLE_COUNT);
 
 		//DebugOutput[DispatchThreadId] = float4(
 		//	SignalPixelBluringRadius,
