Performs history rejection of reflection's color in LCoCg color space to reduce chroma ghosting

diff --git a/Engine/Shaders/Private/ScreenSpaceDenoise/SSDTemporalAccumulation.usf b/Engine/Shaders/Private/ScreenSpaceDenoise/SSDTemporalAccumulation.usf
--- a/Engine/Shaders/Private/ScreenSpaceDenoise/SSDTemporalAccumulation.usf
+++ b/Engine/Shaders/Private/ScreenSpaceDenoise/SSDTemporalAccumulation.usf
@@ -42,7 +42,7 @@
 	#define CONFIG_SIGNAL_HISTORY_LAYOUT SIGNAL_BUFFER_LAYOUT_REFLECTIONS_HISTORY
 
 	#define CONFIG_HISTORY_REJECTION 2
-	#define CONFIG_HISTORY_REJECTION_COLOR_SPACE 0
+	#define CONFIG_HISTORY_REJECTION_COLOR_SPACE 1
 
 #else
 	#error Unimplemented signal processing.
