Fixed glossy materials in the path tracer.

diff --git a/Engine/Shaders/Private/PathTracing/Material/PathTracingGlossy.ush b/Engine/Shaders/Private/PathTracing/Material/PathTracingGlossy.ush
--- a/Engine/Shaders/Private/PathTracing/Material/PathTracingGlossy.ush
+++ b/Engine/Shaders/Private/PathTracing/Material/PathTracingGlossy.ush
@@ -3,16 +3,9 @@
 /*=============================================================================================
 Glossy.usf: Microfacet GGX BRDF sampling functions 
 ===============================================================================================*/
-#pragma once
-
-//#define USE_GGX_ONLY	
 
+#pragma once
 
-//#dxr_todo: move to a commmon file
-bool IsZero(float3 v)
-{
-	return all(v == float3(0.0, 0.0, 0.0));
-}
 
 bool IsNanOrNotFinite(float p)
 {
@@ -21,182 +14,205 @@ bool IsNanOrNotFinite(float p)
 	return (asuint(p) & (0xFF << 23)) == (0xFF << 23);
 }
 
-float ClampInf(float p)
+float ClampIfInvalid(float p)
+{
+	return IsNanOrNotFinite(p) ? 10000000 : p;
+}
+
+float GetSpecularEventProbability(float3 SpecularColor, float NoV)
+{
+	float3 F = F_Schlick(SpecularColor, NoV);
+	return (F.x + F.y + F.z) / 3.f;
+}
+
+// G_SmithV predivided by 4*NoV*NoL. #dxr-todo: move to BRDF.ush?
+float PredividedG_SmithV(float a2, float NoV, const float NoL)
 {
-	return IsNanOrNotFinite(p) ? 1000000 : p;
+	const float G_SmithV = NoV + sqrt((NoV - NoV * a2) * NoV + a2);
+	return 0.5f / (G_SmithV * NoL);
 }
 
 float Glossy_Pdf(float3 IncomingDirection, float3 OutgoingDirection, FMaterialClosestHitPayload Payload)
 {
 	float3 SpecularColor = GetSpecularColor(Payload);
 	float3 N_World = GetWorldNormal(Payload);
-	float3 Normal_Local = float3(0.0, 0.0, 1.0);
 	float Roughness = max(GetRoughness(Payload), GLOSSY_MIN_ROUGHNESS);
 
-	//#dxr_todo: check if we IncomingDirection/OutgoingDirection are already normalized
 	float3 V_World = normalize(-IncomingDirection);
-	float3 V_Local = WorldToTangent(V_World, N_World);
-	float3 L_Local = WorldToTangent(normalize(OutgoingDirection), N_World);
-
-	float NoV = abs(V_Local.z);
+	float3 V = WorldToTangent(V_World, N_World);
+	float3 L = WorldToTangent(normalize(OutgoingDirection), N_World);
 
-#ifdef	USE_GGX_ONLY
-	float ProbSpecular = 1.0; 
-#else
-	float ProbSpecular = F_Schlick(SpecularColor, NoV).g;
-#endif
+	float NoV = V.z;
+	float NoL = L.z;
 
-	float SpecularPdf, DiffusePdf;
+	const float ProbSpecular = GetSpecularEventProbability(SpecularColor, NoV);
 
 	// Specular
-	float3 Wh = normalize(V_Local + L_Local);
-	float a2 = Pow4(Roughness);
-	float G_SmithV = 2 * NoV / (NoV + sqrt(NoV * (NoV - NoV * a2) + a2));
-	float D = D_GGX(a2, abs(Wh.z));
-	SpecularPdf = D * G_SmithV * abs(dot(V_Local, Wh) / abs(V_Local.z));
-	SpecularPdf /= (4.0f * min( 0.001, dot(V_Local, Wh)));
+	float SpecularPdf = 0.0f;
+	float3 H = normalize(V + L);
+	float VoH = saturate(dot(V, H));
+	float NoH = saturate(H.z);
 
-	// Diffuse
-	DiffusePdf = abs(L_Local.z) * (1 / PI);
+	if (NoL * NoV > 0.0f && NoH > 0.0f)
+	{
+		float a2 = Pow4(Roughness);
+		float D = D_GGX(a2, NoH);
+		float G1 = PredividedG_SmithV(a2, abs(NoV), abs(NoL));
+
+		SpecularPdf = ClampIfInvalid(G1 * D);
+		SpecularPdf *= NoL;
+	}
+
+	// Diffuse 
+	float DiffusePdf = NoL / PI;
 
 	float OutPdf = ProbSpecular * SpecularPdf + (1.0f - ProbSpecular) * DiffusePdf;
 	return (OutPdf);
 }
 
-float3 Eval_Component_Specular(float3 V_Local, float3 L_Local, float Roughness, float3 SpecularColor)
-{
-	// Specular
-	float CosThetaO = abs(V_Local.z);
-	float CosThetaI = abs(L_Local.z);
-
-	float3 Wh = normalize(V_Local + L_Local);
-	if (CosThetaI == 0 || CosThetaO == 0 || IsZero(Wh))
-		return float3(0.0, 0.0, 0.0);
-
-	float3 F = F_Schlick(SpecularColor, dot(L_Local, Wh));
-	float D = D_GGX(Pow4(Roughness), abs(Wh.z));
-	float G2 = Vis_SmithJointApprox_NoPredivided(Pow4(Roughness), abs(V_Local.z), abs(L_Local.z));
-	float G1 = SmithG1V_NoPredivided(Roughness, abs(V_Local.z), abs(L_Local.z));
-	float3 SmithFactor = ClampInf(G2 / G1);
-	float3 SpecularThroughput = F * D * SmithFactor;
-	SpecularThroughput /= (4 * CosThetaI * CosThetaO);
-
-	return SpecularThroughput;
-}
-
-void Glossy_EvalMaterial(
-	float3 IncomingDirection,
-	float3 OutgoingDirection,
+void Glossy_SampleMaterial(
+	float3 RayDirection,
 	FMaterialClosestHitPayload Payload,
+	float4 RandSample,
+	out float3 OutDirection,
 	out float3 OutThroughput,
-	out float OutPdf
-)
+	out float OutPdf)
 {
+	float3 N_World = GetWorldNormal(Payload);
+
+	// Diffuse and Specular color are calculated in the payload unpacking, both depend on BaseColor and Metallic attributes
 	float3 DiffuseColor = GetDiffuseColor(Payload);
 	float3 SpecularColor = GetSpecularColor(Payload);
 	float Roughness = max(GetRoughness(Payload), GLOSSY_MIN_ROUGHNESS);
 
-	float3 N_World = GetWorldNormal(Payload);
-	float3 Normal_Local = float3(0.0, 0.0, 1.0);
-	float3 V_Local = WorldToTangent(IncomingDirection, N_World);
-	float3 L_Local = WorldToTangent(OutgoingDirection, N_World);
+	float3 V_World = normalize(-RayDirection);
+	float3 V = WorldToTangent(V_World, N_World);
+	float3 L;
 
-	float NoL = max(L_Local.z, 0.0);
+	const float NoV = V.z;
+	const float ProbSpecular = GetSpecularEventProbability(SpecularColor, NoV); 
 
-	float3 SpecularThroughput = Eval_Component_Specular(V_Local, L_Local, Roughness, SpecularColor);
+	OutPdf = 0.0f;
+	OutThroughput = float3(0, 0, 0);
 
-#ifdef	USE_GGX_ONLY
-	OutPdf = Glossy_Pdf(IncomingDirection, OutgoingDirection, Payload);
-	OutThroughput = SpecularThroughput;
-#else
-	float3 DiffuseThroughput = DiffuseColor * (1 / PI);
-	OutThroughput = SpecularThroughput + DiffuseThroughput;
+	if ((RandSample.z <= ProbSpecular)) 
+	{	
+		bool Flip = V.z < 0.0f;
+		float4 ResGGX = ImportanceSampleVisibleGGX(RandSample.xy, Pow4(Roughness), Flip ? -V : V);
+		float3 H = ResGGX.xyz;
+		if (Flip)
+		{
+			H = -H; //Use ResGGX.w for the PDF?
+		}
 
-	float NoV = abs(V_Local.z);
-	float SpecularPdf = Glossy_Pdf(IncomingDirection, OutgoingDirection, Payload);
-	float DiffusePdf = NoL / PI;
-	float ProbSpecular = F_Schlick(SpecularColor, NoV).g;
-	OutPdf = ProbSpecular * SpecularPdf + (1.0f - ProbSpecular) * DiffusePdf;
-#endif
+		L = 2 * dot(V, H) * H - V;
 
-	OutThroughput *= NoL;
+		float NoL = L.z;
+		float NoV = V.z;
+		float NoH = H.z;
+
+		if (L.z < 0.0f || NoL * NoV <= 0.0f || NoH == 0.0f)
+		{
+			return;
+		}
+
+		float a2 = Pow4(Roughness);
+		float D = D_GGX(a2, NoH);
+		float G2 = Vis_SmithJointApprox(a2, abs(NoV), abs(NoL));
+		float G1 = PredividedG_SmithV(a2, abs(NoV), abs(NoL));
+		float3 F = F_Schlick(SpecularColor, dot(V, H));
+
+		OutPdf = ClampIfInvalid(G1 * D) * abs(NoL);
+		OutThroughput = F * ClampIfInvalid(G2 / G1);
+
+		OutThroughput /= ProbSpecular;
+		OutPdf *= ProbSpecular;
+	}
+	else 
+	{
+		float4 DiffuseSampledValue = CosineSampleHemisphere(RandSample.xy);
+		L = DiffuseSampledValue.xyz;
+		L.z = abs(L.z);
+
+		OutThroughput = DiffuseColor; // Compensated below 
+		OutPdf = DiffuseSampledValue.w;
+
+		OutThroughput /= 1.0f - ProbSpecular;
+		OutPdf *= 1.0f - ProbSpecular;
+	}
+
+	// Througput is already previdided by the PDF so it has to be compensated
+	OutThroughput *= OutPdf;
+
+	OutDirection = TangentToWorld(L, N_World);
 }
 
-void Glossy_SampleMaterial(
-	float3 RayDirection,
+
+void Glossy_EvalMaterial(
+	float3 IncomingDirection,
+	float3 OutgoingDirection,
 	FMaterialClosestHitPayload Payload,
-	float4 RandSample,
-	out float3 OutDirection,
 	out float3 OutThroughput,
-	out float OutPdf)
+	out float OutPdf
+)
 {
+	float3 DiffuseColor = GetDiffuseColor(Payload);
+	float3 SpecularColor = GetSpecularColor(Payload);
+	float Roughness = max(GetRoughness(Payload), 0.001);
+
 	float3 N_World = GetWorldNormal(Payload);
+	float3 V_World = normalize(-IncomingDirection);
+	float3 L_World = normalize(OutgoingDirection);
 
-	// Diffuse and Specular color are calculated in the payload unpacking. They both depend on Base Color and Metallic attributes
-	float3 DiffuseColor		= GetDiffuseColor(Payload);
-	float3 SpecularColor	= GetSpecularColor(Payload);
+	float3 V = WorldToTangent(V_World, N_World);
+	float3 L = WorldToTangent(L_World, N_World);
 
-	float3 V_World = normalize(-RayDirection);
-	float3 V_Local = WorldToTangent(V_World, N_World);
-	float3 L_Local;
-	float Roughness = max(GetRoughness(Payload), GLOSSY_MIN_ROUGHNESS);
+	float NoV = V.z;
+	float NoL = L.z;
 
-	if (V_Local.z == 0)
+	// Specular eval
+	float3 SpecularThroughput = float3(0, 0, 0);
+	float SpecularPdf = 0.0;
+	
+	float3 H = normalize(V + L);
+	float VoH = saturate(dot(V, H));
+	float NoH = saturate(H.z);
+
+	if (NoL * NoV > 0.0f && NoH > 0.0f)
 	{
-		OutDirection = float3(0.0, 0.0, 0.0);
-		OutThroughput = float3(0.0, 0.0, 0.0);
-		OutPdf = 0.0;
-		return;
+		float a2 = Pow4(Roughness);
+		float D = D_GGX(a2, NoH);
+		float G2 = Vis_SmithJointApprox(a2, abs(NoV), abs(NoL));
+		float G1 = PredividedG_SmithV(a2, abs(NoV), abs(NoL));
+
+		SpecularThroughput = ClampIfInvalid(G2 * D);
+		SpecularPdf = ClampIfInvalid(G1 * D);
+		SpecularPdf *= NoL;
+
+		if (SpecularPdf != 0)
+		{
+			float3 F = F_Schlick(SpecularColor, VoH);
+			SpecularThroughput *= F;
+		}	
 	}
+	   
+	// Diffuse eval
+	float3 DiffuseThroughput = DiffuseColor / PI;
+	float DiffusePdf = NoL / PI;
 
-	float NoV = abs(V_Local.z);
+	// Add contributions
+	OutThroughput = (DiffuseThroughput + SpecularThroughput);
+	OutThroughput *= NoL;
 
-#ifdef	USE_GGX_ONLY
-	float ProbSpecular = 1.0;
-#else
-	float ProbSpecular = F_Schlick(SpecularColor, NoV).g;
-#endif
+	float ProbSpecular = GetSpecularEventProbability(SpecularColor, NoV);
+	OutPdf = ProbSpecular * SpecularPdf + (1.0f - ProbSpecular) * DiffusePdf;
+}
 
-	if (RandSample.z <= ProbSpecular)
-	{
-		bool Flip = V_Local.z < 0.0f;
-		float4 ResGGX = ImportanceSampleVisibleGGX(RandSample.xy, Pow4(Roughness), Flip ? -V_Local : V_Local);
-		float3 Wh = ResGGX.xyz;
 
-		if (Flip)
-			Wh = -Wh;
 
-		// Reflect from half vector
-		L_Local = 2 * dot(V_Local, Wh) * Wh - V_Local;
 
-		OutPdf = ResGGX.w;
-		OutPdf /= (4.0f * dot(V_Local, Wh));
-		OutPdf *= ProbSpecular;
 
-		if (V_Local.z * L_Local.z > 0.0)
-		{
-			OutThroughput = Eval_Component_Specular(V_Local, L_Local, Roughness, SpecularColor);
-			OutThroughput /= ProbSpecular;
-		}
-		else
-		{
-			OutThroughput = float3(0.0, 0.0, 0.0);
-		}
-	}
-	else // diffuse
-	{
-		float4 ResDiff = CosineSampleHemisphere(RandSample.xy);
-		L_Local = ResDiff.xyz;
-		OutPdf = ResDiff.w;
-	    OutPdf *= 1.0 - ProbSpecular;
-	    OutThroughput = DiffuseColor * (1 / PI);
-	    OutThroughput /= 1.0f - ProbSpecular;
-	}
 
-	float NoL = max(L_Local.z, 0.0);
-	OutThroughput *= NoL;
-	OutDirection = TangentToWorld(L_Local, N_World);
-}
 
 
 
