Removed unnecesary functions from BRDF.ush

diff --git a/Engine/Shaders/Private/BRDF.ush b/Engine/Shaders/Private/BRDF.ush
--- a/Engine/Shaders/Private/BRDF.ush
+++ b/Engine/Shaders/Private/BRDF.ush
@@ -284,24 +284,6 @@ float Vis_SmithJointApprox( float a2, float NoV, float NoL )
 	return 0.5 * rcp( Vis_SmithV + Vis_SmithL );
 }
 
-// Approximation of joint Smith term for GGX, NOT pre-divided by the BRDF denominator (4 * NoV * NoL)
-float Vis_SmithJointApprox_NoPredivided(float a2, float NoV, float NoL)
-{
-	float a = sqrt(a2);
-	float Vis_SmithV = NoL * (NoV * (1 - a) + a);
-	float Vis_SmithL = NoV * (NoL * (1 - a) + a);
-	return 2.0 * NoV * NoL * rcp(Vis_SmithV + Vis_SmithL);
-}
-
-// Smith G1(V) term, NOT pre-divided by the BRDF denominator (4 * NoV * NoL)
-float SmithG1V_NoPredivided(float Roughness, float NoV, const float NoL)
-{
-	const float Alpha = Roughness * Roughness;
-	const float a2 = Alpha * Alpha;
-	const float G_V = NoV + sqrt((NoV - NoV * a2) * NoV + a2);
-	return (2.0f * NoV / G_V);
-}
-
 float3 F_None( float3 SpecularColor )
 {
 	return SpecularColor;
