Some cleaning in the path tracer:
- Reenabled glossy materials now that they work after CL 4718864
- GetDiffuseColor returns Payload.DiffuseColor instead of Payload.BaseColor
- GetEmission renamed to GetMaterialEmission and moved to PathTracingMaterialCommon.ush
- Some cleaning in the lambert BRDF code

diff --git a/Engine/Shaders/Private/PathTracing/Material/PathTracingLambert.ush b/Engine/Shaders/Private/PathTracing/Material/PathTracingLambert.ush
--- a/Engine/Shaders/Private/PathTracing/Material/PathTracingLambert.ush
+++ b/Engine/Shaders/Private/PathTracing/Material/PathTracingLambert.ush
@@ -36,6 +36,9 @@ void Lambert_EvalMaterial(
 	out float OutPdf
 )
 {
-	OutPdf = Lambert_Pdf(OutgoingDirection, Payload);
-	OutThroughput = GetDiffuseColor(Payload) * OutPdf;
+	float3 N_World = GetWorldNormal(Payload);
+	float NoL = max(dot(N_World, OutgoingDirection), 0.0);
+
+	OutPdf = NoL / PI;
+	OutThroughput = GetDiffuseColor(Payload) * NoL / PI;
 }
\ No newline at end of file
diff --git a/Engine/Shaders/Private/PathTracing/Material/PathTracingMaterialCommon.ush b/Engine/Shaders/Private/PathTracing/Material/PathTracingMaterialCommon.ush
--- a/Engine/Shaders/Private/PathTracing/Material/PathTracingMaterialCommon.ush
+++ b/Engine/Shaders/Private/PathTracing/Material/PathTracingMaterialCommon.ush
@@ -9,9 +9,15 @@ PathTracingMaterialCommon.usf: Brdf utility functions
 #include "../PathTracingCommon.ush"
 #include "/Engine/Private/BRDF.ush"
 
+
+float3 GetMaterialEmission(FMaterialClosestHitPayload HitInfo)
+{
+	return HitInfo.Radiance;
+}
+
 float3 GetDiffuseColor(FMaterialClosestHitPayload HitInfo)
 {
-	return HitInfo.BaseColor;
+	return HitInfo.DiffuseColor;
 }
 
 float3 GetSpecularColor(FMaterialClosestHitPayload HitInfo)
@@ -48,9 +54,7 @@ bool IsPureSpecularTransmission(FMaterialClosestHitPayload HitInfo)
 
 bool IsLambert(FMaterialClosestHitPayload HitInfo)
 {
-	//#dxr_todo: force lambert material until other BRDF types are finished
-	return true;
-	//return GetRoughness(HitInfo) > GLOSSY_MAX_ROUGHNESS;
+	return GetRoughness(HitInfo) > GLOSSY_MAX_ROUGHNESS;
 }
 
 bool BSDF_Refract(float3 V_Local, float3 N_Local, float IorV, float IorL, out float3 L_Local)
diff --git a/Engine/Shaders/Private/PathTracing/PathTracing.usf b/Engine/Shaders/Private/PathTracing/PathTracing.usf
--- a/Engine/Shaders/Private/PathTracing/PathTracing.usf
+++ b/Engine/Shaders/Private/PathTracing/PathTracing.usf
@@ -29,11 +29,6 @@ float GetHitT(FMaterialClosestHitPayload HitInfo)
 	return HitInfo.HitT;
 }
 
-float3 GetEmission(FMaterialClosestHitPayload HitInfo)
-{
-	return HitInfo.Radiance;    
-}
-
 bool IsHit(RayDesc Ray, FMaterialClosestHitPayload HitInfo)
 {
 	float t = HitInfo.HitT;
@@ -134,7 +129,7 @@ void PathTracingMainRG()
 			}
 
 			// Emissive geometry contribution
-			float3 Emission = GetEmission(Payload);
+			float3 Emission = GetMaterialEmission(Payload);
 			if (any(Emission) > 0.0)
 			{
 				if (IsUnidirectionalEnabled)
