Ray tracing Payload.Radiance does not store DiffuseIndirectLighting anymore but just the result of GetMaterialEmissive because otherwise the path tracer was handling objects with not null DiffuseIndirectLighting as emitters.

DiffuseIndirectLighting is now stored in Payload.IndirectIrradiance. TODO: add pack/unpack support

diff --git a/Engine/Shaders/Private/RayTracing/RayTracingMaterialHitShaders.usf b/Engine/Shaders/Private/RayTracing/RayTracingMaterialHitShaders.usf
--- a/Engine/Shaders/Private/RayTracing/RayTracingMaterialHitShaders.usf
+++ b/Engine/Shaders/Private/RayTracing/RayTracingMaterialHitShaders.usf
@@ -302,7 +302,7 @@ void MaterialCHS(inout FPackedMaterialClosestHitPayload PackedPayload, in FDefau
 	Payload.DiffuseColor = BaseColor - BaseColor * Metallic;
 	Payload.SpecularColor = lerp(0.08 * Specular.xxx, BaseColor, Metallic);
 	
-	Payload.Radiance += Payload.DiffuseColor * DiffuseIndirectLighting; // Contains both lightmap & skylight
+	Payload.IndirectIrradiance += Payload.DiffuseColor * DiffuseIndirectLighting; // Contains both lightmap & skylight
 
 	PackedPayload = PackRayTracingPayload(Payload);
 }
