Added Payload.IndirectIrradiance to indirect contribution in ray traced reflections

diff --git a/Engine/Shaders/Private/RayTracing/RayTracingReflections.usf b/Engine/Shaders/Private/RayTracing/RayTracingReflections.usf
--- a/Engine/Shaders/Private/RayTracing/RayTracingReflections.usf
+++ b/Engine/Shaders/Private/RayTracing/RayTracingReflections.usf
@@ -283,7 +283,8 @@ FMaterialClosestHitPayload TraceRayAndAccumulateResults(
 	{
 	    // Emissive & indirect contribution
 	    Radiance += Payload.Radiance;
-    }
+		Radiance += Payload.IndirectIrradiance;
+	}
 
 	return Payload;
 }
