Changes the computation of the specular lobe using the spherical gaussian framework and adaptes number of samples in spatial accumulation.

diff --git a/Engine/Shaders/Private/ScreenSpaceDenoise/SSDSignalFramework.ush b/Engine/Shaders/Private/ScreenSpaceDenoise/SSDSignalFramework.ush
--- a/Engine/Shaders/Private/ScreenSpaceDenoise/SSDSignalFramework.ush
+++ b/Engine/Shaders/Private/ScreenSpaceDenoise/SSDSignalFramework.ush
@@ -48,27 +48,8 @@ float HitDistanceToWorldBluringRadius;
 //------------------------------------------------------- INTERNAL FOR REFLECTION SIGNAL
 #if CONFIG_SIGNAL_PROCESSING == SIGNAL_PROCESSING_REFLECTIONS
 
-/** Return the size of the bluring radius caused by the specular lobe in ViewportUV.x space. */
-float ComputeInfinityBluringRadius(FSSFSampleSceneInfos SceneMetadata)
-{
-	// TODO: this should depend on NoV? Try up with spherical gauassian framework.
-	float RoughnessLobeAngle = ComputeAngleFromRoughness(SceneMetadata.Roughness);
-
-	return tan(0.5 * RoughnessLobeAngle) * GetCotanHalfFieldOfView().x;
-}
-
-/** Return the size of bluring radius caused by specular lobe in view space. */
-float ComputeBluringRadiusOfNormalizedSample(FSSFSignalSample Sample, FSSFSampleSceneInfos SceneMetadata)
-{
-	float DistanceCameraToPixel = length(View.TranslatedWorldCameraOrigin - SceneMetadata.TranslatedWorldPosition);
-
-	float InifinityRadius = ComputeInfinityBluringRadius(SceneMetadata);
-
-	return InifinityRadius * (1 - DistanceCameraToPixel / (DistanceCameraToPixel + Sample.PenumbraClosestHit));
-}
-
 /** Create a spherical gaussian lob for a given roughness. */
-FSphericalGaussian RoughnessLobe(float Roughness, float3 N, float3 V)
+FSphericalGaussian ComputeRoughnessLobe(float Roughness, float3 N, float3 V)
 {
 	float a = Pow2( max( 0.02, Roughness ) );
 	float a2 = a*a;
@@ -86,20 +67,50 @@ FSphericalGaussian RoughnessLobe(float Roughness, float3 N, float3 V)
 }
 
 /** Create a spherical gaussian lob for a given scene metadata. */
-FSphericalGaussian RoughnessLobe(FSSFSampleSceneInfos RefSceneMetadata)
+FSphericalGaussian ComputeRoughnessLobe(FSSFSampleSceneInfos RefSceneMetadata)
 {
 	float3 N = RefSceneMetadata.WorldNormal;
 	float3 V = normalize(View.TranslatedWorldCameraOrigin - RefSceneMetadata.TranslatedWorldPosition);
 	float Roughness = RefSceneMetadata.Roughness;
 
-	return RoughnessLobe(Roughness, N, V);
+	return ComputeRoughnessLobe(Roughness, N, V);
+}
+
+/** Return the size of the bluring radius caused by the specular lobe in ViewportUV.x space. */
+float ComputeInfinityBluringRadius(FSSFSampleSceneInfos SceneMetadata)
+{
+	float RoughnessLobeAngle;
+	#if 1
+	// Compute the angle of the lobe using spherical gaussian framework.
+	{
+		FSphericalGaussian RoughnessLobe = ComputeRoughnessLobe(SceneMetadata);
+		RoughnessLobeAngle = GetConeAngle(RoughnessLobe);
+	}
+	#else
+	// Compute the roughness. Looks way to broken compared to roughness lobe.
+	{
+		RoughnessLobeAngle = ComputeAngleFromRoughness(SceneMetadata.Roughness);
+	}	
+	#endif
+
+	return tan(0.5 * RoughnessLobeAngle) * GetCotanHalfFieldOfView().x;
+}
+
+/** Return the size of bluring radius caused by specular lobe in view space. */
+float ComputeBluringRadiusOfNormalizedSample(FSSFSignalSample Sample, FSSFSampleSceneInfos SceneMetadata)
+{
+	float DistanceCameraToPixel = length(View.TranslatedWorldCameraOrigin - SceneMetadata.TranslatedWorldPosition);
+
+	float InifinityRadius = ComputeInfinityBluringRadius(SceneMetadata);
+
+	return InifinityRadius * (1 - DistanceCameraToPixel / (DistanceCameraToPixel + Sample.PenumbraClosestHit));
 }
 
 // [ Tokoyashi 2015, "Specular Lobe-Aware Filtering and Upsampling for Interactive Indirect Illumination" ]
 float TokoyashiLobeBilateralWeight(FSSFSampleSceneInfos RefSceneMetadata, FSSFSampleSceneInfos NeighborSceneMetadata)
 {
-	FSphericalGaussian Ref = RoughnessLobe(RefSceneMetadata);
-	FSphericalGaussian Neighbor = RoughnessLobe(NeighborSceneMetadata);
+	FSphericalGaussian Ref = ComputeRoughnessLobe(RefSceneMetadata);
+	FSphericalGaussian Neighbor = ComputeRoughnessLobe(NeighborSceneMetadata);
 	
 #if 1
 	const float Beta = 5.05; // TODO: to be tunned with GGX.
diff --git a/Engine/Shaders/Private/ScreenSpaceDenoise/SSDSpatialAccumulation.usf b/Engine/Shaders/Private/ScreenSpaceDenoise/SSDSpatialAccumulation.usf
--- a/Engine/Shaders/Private/ScreenSpaceDenoise/SSDSpatialAccumulation.usf
+++ b/Engine/Shaders/Private/ScreenSpaceDenoise/SSDSpatialAccumulation.usf
@@ -39,7 +39,10 @@
 		#define CONFIG_SIGNAL_INPUT_LAYOUT  SIGNAL_BUFFER_LAYOUT_REFLECTIONS_INPUT
 		#define CONFIG_SIGNAL_OUTPUT_LAYOUT SIGNAL_BUFFER_LAYOUT_REFLECTIONS_HISTORY
 
+		// Use stackowiak sample set to keep noise in the signal to avoid flicker in history rejection,
+		// and adapt number of samples based on the sample.
 		#define CONFIG_SAMPLE_SET SAMPLE_SET_STACKOWIAK_4_SETS
+		#define CONFIG_SAMPLE_COUNT_POLICY 1
 
 		// Reconstruction has 1 ray per input pixel, therefore can use MIS.
 		#define CONFIG_SPECULAR_MIS 1
@@ -77,6 +80,11 @@
 	#define CONFIG_WEIGHT_KERNEL_CENTER 0
 #endif
 
+/** Config sample set. */
+#ifndef CONFIG_SAMPLE_COUNT_POLICY
+	#define CONFIG_SAMPLE_COUNT_POLICY 0
+#endif
+
 
 //------------------------------------------------------- INCLUDES
 
@@ -163,6 +171,8 @@ Texture2D SignalInput1;
 RWTexture2D<float4>	SignalOutput0;
 RWTexture2D<float4>	SignalOutput1;
 
+RWTexture2D<float4>	DebugOutput;
+
 
 //------------------------------------------------------- FUNCTIONS
 
@@ -217,7 +227,20 @@ void MainCS(
 		KernelCenterBufferUV = ((DispatchThreadId | 1) + View.ViewRectMin.xy) * View.BufferSizeAndInvSize.zw;
 		KernelCenterBufferUV = clamp(KernelCenterBufferUV, View.BufferBilinearUVMinMax.xy, View.BufferBilinearUVMinMax.zw);
 
-		uint SampleCount = kStackowiakSampleCountPerSet; // TODO
+		uint SampleCount = kStackowiakSampleCountPerSet;
+		#if CONFIG_SAMPLE_COUNT_POLICY == 1
+		{
+			float RefBluringRadius = ComputeBluringRadiusOfNormalizedSample(RefSample, RefSceneMetadata) * View.ViewSizeAndInvSize.x;
+			float TargetedSampleCount = PI * RefBluringRadius * RefBluringRadius * rcp(kStackowiakSampleSetCount);
+
+			const float MaximumRadius = 8;
+			const float MaximumSampleCount = kStackowiakSampleCountPerSet;
+
+			SampleCount = TargetedSampleCount;
+			
+			// DebugOutput[DispatchThreadId] = float4(RefBluringRadius, SampleCount, 0, 0);
+		}
+		#endif // CONFIG_SAMPLE_COUNT_POLICY == 1
 
 		// Id of the pixel in the quad.
 		uint QuadSampleId = (DispatchThreadId.x & 1) | ((DispatchThreadId.y & 1) << 1);
diff --git a/Engine/Source/Runtime/Renderer/Private/ScreenSpaceDenoise.cpp b/Engine/Source/Runtime/Renderer/Private/ScreenSpaceDenoise.cpp
--- a/Engine/Source/Runtime/Renderer/Private/ScreenSpaceDenoise.cpp
+++ b/Engine/Source/Runtime/Renderer/Private/ScreenSpaceDenoise.cpp
@@ -231,6 +231,8 @@ class FSSDSpatialAccumulationCS : public FScreenSpaceDenoisingShader
 
 		SHADER_PARAMETER_RDG_TEXTURE_UAV(RWTexture2D, SignalOutput0)
 		SHADER_PARAMETER_RDG_TEXTURE_UAV(RWTexture2D, SignalOutput1)
+
+		SHADER_PARAMETER_RDG_TEXTURE_UAV(RWTexture2D, DebugOutput) // TODO: remove
 	END_SHADER_PARAMETER_STRUCT()
 };
 
@@ -671,6 +673,8 @@ IScreenSpaceDenoiser::FReflectionOutputs DenoiseReflections(
 		PassParameters->SignalInput1 = ReflectionInputs.RayHitDistance;
 		PassParameters->SignalOutput0 = GraphBuilder.CreateUAV(SignalOutput0);
 		PassParameters->SignalOutput1 = GraphBuilder.CreateUAV(SignalOutput1);
+		
+		PassParameters->DebugOutput = GraphBuilder.CreateUAV(GraphBuilder.CreateTexture(SignalProcessingDesc, TEXT("SSDDebugReflectionReconstruction")));
 
 		FSSDSpatialAccumulationCS::FPermutationDomain PermutationVector;
 		PermutationVector.Set<FSignalProcessingDim>(ESignalProcessing::Reflections);
