Ray tracing IndirectIrradiance is properly packed/unpacked in the payload structure.

diff --git a/Engine/Shaders/Private/RayTracing/RayTracingCommon.ush b/Engine/Shaders/Private/RayTracing/RayTracingCommon.ush
--- a/Engine/Shaders/Private/RayTracing/RayTracingCommon.ush
+++ b/Engine/Shaders/Private/RayTracing/RayTracingCommon.ush
@@ -65,7 +65,7 @@ struct FMaterialClosestHitPayload
 	uint PrimitiveLightingChannelMask;	//96	0bytes	(removed)
 	float4 CustomData;					//112	0bytes	(removed)
 	float GBufferAO;					//116	0bytes	(removed)
-	float IndirectIrradiance;			//120	0bytes	(removed)
+	float3 IndirectIrradiance;			//120	0bytes	48bits	
 
 	// Quite some code assume FRayHitInfo has a WorldPos
 	// So keep it here and force to re-construct it in Unpack call using ray information.
@@ -86,6 +86,7 @@ struct FPackedMaterialClosestHitPayload
 	int BaseColorAndOpacity[2];
 	int MetallicAndSpecularAndRoughness;
 	int IorAndShadingModelIDAndBlendingModeAndHitKind;
+	int IndirectIrradiance[2];
 };
 
 FPackedMaterialClosestHitPayload PackRayTracingPayload(FMaterialClosestHitPayload Input)
@@ -109,6 +110,9 @@ FPackedMaterialClosestHitPayload PackRayTracingPayload(FMaterialClosestHitPayloa
     Output.IorAndShadingModelIDAndBlendingModeAndHitKind |= (Input.ShadingModelID & 0xF) << 16;
 	Output.IorAndShadingModelIDAndBlendingModeAndHitKind |= (Input.BlendingMode & 0x7) << 20;
 	Output.IorAndShadingModelIDAndBlendingModeAndHitKind |= (Input.HitKind & 0x1) << 31;
+	Output.IndirectIrradiance[0] = f32tof16(Input.IndirectIrradiance.x);
+	Output.IndirectIrradiance[0] |= f32tof16(Input.IndirectIrradiance.y) << 16;
+	Output.IndirectIrradiance[1] = f32tof16(Input.IndirectIrradiance.z);
 
 	return Output;
 }
@@ -134,6 +138,9 @@ FMaterialClosestHitPayload UnpackRayTracingPayload(FPackedMaterialClosestHitPayl
 	Output.ShadingModelID = (Input.IorAndShadingModelIDAndBlendingModeAndHitKind >> 16) & 0xF;
 	Output.BlendingMode	  = (Input.IorAndShadingModelIDAndBlendingModeAndHitKind >> 20) & 0x7;
 	Output.HitKind		  = (Input.IorAndShadingModelIDAndBlendingModeAndHitKind >> 31);
+	Output.IndirectIrradiance.x	= f16tof32(Input.IndirectIrradiance[0]);
+	Output.IndirectIrradiance.y	= f16tof32(Input.IndirectIrradiance[0] >> 16);
+	Output.IndirectIrradiance.z	= f16tof32(Input.IndirectIrradiance[1]);
 
 	Output.DiffuseColor  = Output.BaseColor - Output.BaseColor * Output.Metallic;
 	Output.SpecularColor = lerp(0.08 * Output.Specular.xxx, Output.BaseColor, Output.Metallic);
