Allows separate translucency with alpha channel support in post processing.

diff --git a/Engine/Shaders/Private/PostProcessBokehDOF.usf b/Engine/Shaders/Private/PostProcessBokehDOF.usf
--- a/Engine/Shaders/Private/PostProcessBokehDOF.usf
+++ b/Engine/Shaders/Private/PostProcessBokehDOF.usf
@@ -403,13 +403,12 @@ float4 DOFRecombineCommon(float2 UV, float4 SvPosition)
 	float2 PixelPosCenter = SvPosition.xy;
 	float2 FullResUV = PixelPosCenter * PostprocessInput0Size.zw;
 
-	float3 UnfocusedSceneColor = Texture2DSample(PostprocessInput0, PostprocessInput0Sampler, FullResUV).rgb;
+	float4 UnfocusedSceneColor = Texture2DSample(PostprocessInput0, PostprocessInput0Sampler, FullResUV);
 
-	OutColor.rgb = UnfocusedSceneColor;
-	OutColor.a = 0;
+	OutColor = UnfocusedSceneColor;
 
 #if RECOMBINE_METHOD == 1 || RECOMBINE_METHOD == 3
-	float4 SceneColorAndDepth = float4(UnfocusedSceneColor, CalcSceneDepth(FullResUV));
+	float4 SceneColorAndDepth = float4(UnfocusedSceneColor.rgb, CalcSceneDepth(FullResUV));
 
 	// BokehDOF in half res
 	float2 BokehViewportUV = UV;
@@ -434,7 +433,7 @@ float4 DOFRecombineCommon(float2 UV, float4 SvPosition)
 		float4 LayerMerger = 0;
 
 		// Layer 1: half res background
-		LayerMerger.rgb = lerp(UnfocusedSceneColor, BokehDOFAccumLayer1.rgb / max(BokehDOFAccumLayer1.a, 0.001f), Layer1Mask);
+		LayerMerger.rgb = lerp(UnfocusedSceneColor.rgb, BokehDOFAccumLayer1.rgb / max(BokehDOFAccumLayer1.a, 0.001f), Layer1Mask);
 
 		// Layer 2: then we add the focused scene to fill the empty areas
 		LayerMerger = lerp(LayerMerger, float4(SceneColorAndDepth.rgb, 1), Layer2Mask);
@@ -442,13 +441,13 @@ float4 DOFRecombineCommon(float2 UV, float4 SvPosition)
 		// Layer 3: on top of that blend the front half res layer	
 		LayerMerger = lerp(LayerMerger, float4(BokehDOFAccumLayer3.rgb / max(BokehDOFAccumLayer3.a, 0.001f), 0), Layer3Mask);
 
-		UnfocusedSceneColor = LayerMerger.rgb;
+		UnfocusedSceneColor.rgb = LayerMerger.rgb;
 
 		// blend in full resolution where we are most in focus
 		FocusedWeight = LayerMerger.a;
 	}
 
-	OutColor.rgb = lerp(UnfocusedSceneColor, OutColor.rgb, FocusedWeight);
+	OutColor.rgb = lerp(UnfocusedSceneColor.rgb, OutColor.rgb, FocusedWeight);
 #endif
 
 #if RECOMBINE_METHOD == 2 || RECOMBINE_METHOD == 3
@@ -456,6 +455,11 @@ float4 DOFRecombineCommon(float2 UV, float4 SvPosition)
 
 	// add RGB, darken by A (this allows to represent translucent and additive blending)
 	OutColor.rgb = OutColor.rgb * SeparateTranslucency.a + SeparateTranslucency.rgb;
+	OutColor.a *= SeparateTranslucency.a;
+#endif
+
+#if !POST_PROCESS_ALPHA
+	OutColor.a = 0;
 #endif
 
 	return OutColor;
diff --git a/Engine/Source/Runtime/Engine/Private/SceneView.cpp b/Engine/Source/Runtime/Engine/Private/SceneView.cpp
--- a/Engine/Source/Runtime/Engine/Private/SceneView.cpp
+++ b/Engine/Source/Runtime/Engine/Private/SceneView.cpp
@@ -2398,18 +2398,14 @@ bool FSceneViewFamily::SupportsScreenPercentage() const
 
 bool FSceneViewFamily::AllowTranslucencyAfterDOF() const
 {
-	static IConsoleVariable* CVar = IConsoleManager::Get().FindConsoleVariable(TEXT("r.PostProcessing.PropagateAlpha"));
-	const bool bPostProcessAlphaChannel = CVar ? (CVar->GetInt() != 0) : false;
-
 	static IConsoleVariable* CVarMobileMSAA = IConsoleManager::Get().FindConsoleVariable(TEXT("r.MobileMSAA"));
 	const bool bMobileMSAA = CVarMobileMSAA ? (CVarMobileMSAA->GetInt() > 1) : false;
 
 	return CVarAllowTranslucencyAfterDOF.GetValueOnRenderThread() != 0
 		&& (GetFeatureLevel() > ERHIFeatureLevel::ES3_1 || (IsMobileHDR() && !bMobileMSAA)) // on <= ES3_1 separate translucency requires HDR on and MSAA off
-		&& EngineShowFlags.PostProcessing // Used for reflection captures.
-		&& !UseDebugViewPS()
-		&& EngineShowFlags.SeparateTranslucency
-		&& !bPostProcessAlphaChannel;
+	&& EngineShowFlags.PostProcessing // Used for reflection captures.
+	&& !UseDebugViewPS()
+	&& EngineShowFlags.SeparateTranslucency;
 	// If not, translucency after DOF will be rendered in standard translucency.
 }
 
