Stops the denoiser's reduce pass to produce the debug output.

diff --git a/Engine/Shaders/Private/ScreenSpaceDenoise/SSDReduce.usf b/Engine/Shaders/Private/ScreenSpaceDenoise/SSDReduce.usf
--- a/Engine/Shaders/Private/ScreenSpaceDenoise/SSDReduce.usf
+++ b/Engine/Shaders/Private/ScreenSpaceDenoise/SSDReduce.usf
@@ -162,11 +162,11 @@ void MainCS(
 	float SignalWorldBlurRadius = GetSignalWorldBluringRadius(InputSample, SceneMetadata);
 	float SignalMipLevel = ComputeSignalMipLevel(SceneMetadata, SignalWorldBlurRadius);
 
-	DebugOutput[DispatchThreadId] = float4(
-		SignalWorldBlurRadius,
-		SignalMipLevel,
-		0,
-		0);
+	//DebugOutput[DispatchThreadId] = float4(
+	//	SignalWorldBlurRadius,
+	//	SignalMipLevel,
+	//	0,
+	//	0);
 
 	float PixelWorldBluringRadius = ComputePixelWorldRadius(SceneMetadata);
 
