Adding material model to RectLight

TODO: complete unification of material invocation within shared functions (eg: DefaultLitBxDF()), that can be called with or without LTC data.

diff --git a/Engine/Shaders/Private/RayTracing/RayTracingRectLight.ush b/Engine/Shaders/Private/RayTracing/RayTracingRectLight.ush
--- a/Engine/Shaders/Private/RayTracing/RayTracingRectLight.ush
+++ b/Engine/Shaders/Private/RayTracing/RayTracingRectLight.ush
@@ -3,6 +3,7 @@
 #pragma once
 
 #include "../PathTracing/Utilities/PathTracingRandomSequence.ush"
+#include "../RectLight.ush"
 
 bool GenerateRectLightOcclusionRay(
 	float3 WorldPosition,
@@ -40,11 +41,15 @@ bool GenerateRectLightOcclusionRay(
 	RayTMin = 0.0;
 	RayTMax = length(LightPosition - WorldPosition);
 
-	// Convert PDF to solid angle measure
-	float Area = LightWidth * LightHeight;
-	float Distance2 = RayTMax * RayTMax;
-	float SolidAngle = Area * saturate(dot(-LightDirection, LightNormal)) / Distance2;
-	RayPdf = 1.0 / SolidAngle;
+	// Define rectangle
+	FRect Rect;
+	Rect.Origin = LightPosition - WorldPosition;
+	Rect.Axis[0] = LightdPdu;
+	Rect.Axis[1] = LightdPdv;
+	Rect.Axis[2] = cross(LightdPdu, LightdPdv);
+	Rect.Extent = float2(LightWidth, LightHeight);
+	FSphericalRect SphericalRect = BuildSphericalRect(Rect);
+	RayPdf = 1.0 / SphericalRect.SolidAngle;
 
 	return true;
 }
diff --git a/Engine/Shaders/Private/RayTracing/RayTracingRectLightOcclusionRGS.usf b/Engine/Shaders/Private/RayTracing/RayTracingRectLightOcclusionRGS.usf
--- a/Engine/Shaders/Private/RayTracing/RayTracingRectLightOcclusionRGS.usf
+++ b/Engine/Shaders/Private/RayTracing/RayTracingRectLightOcclusionRGS.usf
@@ -3,6 +3,8 @@
 #include "../Common.ush"
 #include "../DeferredShadingCommon.ush"
 #include "../PathTracing/Utilities/PathTracingRandomSequence.ush"
+#include "../RectLight.ush"
+#include "../ShadingModels.ush"
 #include "RayTracingCommon.ush"
 #include "RayTracingRectLight.ush"
 #include "MipTreeCommon.ush"
@@ -128,10 +130,11 @@ bool GenerateRectLightImportanceSampledRay(
 	float2 PixelOffset = RandomSequence_GenerateSample2D(RandSequence, DummyVariable);
 	float2 DimensionsAtStopLevel = DimensionsAtLevel(RectLight.MipTreeDimensions, StopLevel);
 	float2 MipPixelUV = (MipPixel + PixelOffset) / DimensionsAtStopLevel;
-	MipPdf *= DimensionsAtStopLevel;
 
 	// Map sample point to quad
 	MipPixelUV -= 0.5;
+	float2 RectLightOffset = MipPixel / DimensionsAtStopLevel;
+	float3 LightOrigin = LightPosition + LightdPdu * LightWidth * RectLightOffset.x + LightdPdv * LightHeight * RectLightOffset.y;
 	LightPosition += LightdPdu * LightWidth * MipPixelUV.x + LightdPdv * LightHeight * MipPixelUV.y;
 	float3 LightDirection = normalize(LightPosition - WorldPosition);
 
@@ -150,13 +153,17 @@ bool GenerateRectLightImportanceSampledRay(
 	RayTMin = 0.0;
 	RayTMax = length(LightPosition - WorldPosition);
 
-	// Determine area pdf
-	float Area = (RectLight.Width * RectLight.Height) / DimensionsAtStopLevel;
-	float Distance2 = RayTMax * RayTMax;
-	float SolidAngle = Area * saturate(dot(-LightDirection, LightNormal)) / Distance2;
-	float SolidAnglePdf = 1.0 / SolidAngle;
-	RayPdf = MipPdf * SolidAnglePdf;
+	// Construct solid angle pdf
+	FRect Rect;
+	Rect.Origin = LightOrigin - WorldPosition;
+	Rect.Axis[0] = LightdPdu;
+	Rect.Axis[1] = LightdPdv;
+	Rect.Axis[2] = cross(LightdPdu, LightdPdv);
+	Rect.Extent = float2(LightWidth, LightHeight) / DimensionsAtStopLevel;
+	FSphericalRect SphericalRect = BuildSphericalRect(Rect);
+	float SolidAnglePdf = 1.0 / SphericalRect.SolidAngle;
 
+	RayPdf = MipPdf * SolidAnglePdf;
 	return RayTMax > RayTMin;
 }
 
@@ -175,10 +182,13 @@ void RectLightOcclusionRGS()
 	FScreenSpaceData ScreenSpaceData = GetScreenSpaceData(UV);
 	float Depth = ScreenSpaceData.GBuffer.Depth;
 	float3 WorldPosition = ReconstructWorldPositionFromDepth(UV, Depth);
+	float3 CameraOrigin = ReconstructWorldPositionFromDepth(UV, 0.0f);
+	float3 CameraDirection = normalize(WorldPosition - CameraOrigin);
 	float3 WorldNormal = ScreenSpaceData.GBuffer.WorldNormal;
+	float3 BaseColor = ScreenSpaceData.GBuffer.BaseColor;
 	float3 DiffuseColor = ScreenSpaceData.GBuffer.DiffuseColor;
-
-	float3 LightRadiance = RectLight.Color;
+	float3 SpecularColor = ScreenSpaceData.GBuffer.SpecularColor;
+	float Roughness = ScreenSpaceData.GBuffer.Roughness;
 
 	float RayDistance = 0.0;
 	float HitCount = 0.0;
@@ -255,11 +265,25 @@ void RectLightOcclusionRGS()
 			float2 UV = float2(dot(DeltaPosition, RectLight.dPdu) / RectLight.Width, dot(DeltaPosition, RectLight.dPdv) / RectLight.Height);
 			UV = 1.0 - UV;
 			float3 TextureColor = Texture2DSample(RectLight.Texture, RectLight.TextureSampler, UV);
-			float3 LightRadiance = RectLight.Color * TextureColor;
-			// #dxr_todo: Replace with actual BRDF
-			float3 Brdf = DiffuseColor;
-			float Distance2 = Ray.TMax * Ray.TMax;
-			ExitantRadiance += LightRadiance * Brdf * saturate(NoL) / RayPdf;
+			float3 IncomingRadiance = RectLight.Color * TextureColor;
+
+			// Unlit material model
+			// #dxr_todo: use ShadingModelID dispatch instead
+			half3 N = WorldNormal;
+			half3 V = -CameraDirection;
+			half3 L = Ray.Direction;
+			float3 Diffuse = Diffuse_Lambert(DiffuseColor);
+
+			BxDFContext Context;
+			Init(Context, N, V, L);
+			float a2 = Pow4(Roughness);
+			float D = D_GGX(a2, Context.NoH);
+			float Vis = Vis_SmithJointApprox(a2, Context.NoV, NoL);
+			float3 F = F_Schlick(SpecularColor, Context.VoH);
+			float3 Specular = D * Vis * F;
+
+			float3 Brdf = Diffuse + Specular;
+			ExitantRadiance += IncomingRadiance * Brdf * NoL / RayPdf;
 #else
 			// Record exitant radiance as visibility score
 			ExitantRadiance += 1.0;
