Fixes a bug in temporal history input layout, and tries up some experiment on the post filtering.

diff --git a/Engine/Shaders/Private/ScreenSpaceDenoise/SSDSignalFramework.ush b/Engine/Shaders/Private/ScreenSpaceDenoise/SSDSignalFramework.ush
--- a/Engine/Shaders/Private/ScreenSpaceDenoise/SSDSignalFramework.ush
+++ b/Engine/Shaders/Private/ScreenSpaceDenoise/SSDSignalFramework.ush
@@ -87,7 +87,7 @@ float ComputeInfinityBluringRadius(FSSFSampleSceneInfos SceneMetadata)
 		RoughnessLobeAngle = GetConeAngle(RoughnessLobe);
 	}
 	#else
-	// Compute the roughness. Looks way to broken compared to roughness lobe.
+	// Compute the roughness. Looks way too broken compared to spherical gaussian code path.
 	{
 		RoughnessLobeAngle = ComputeAngleFromRoughness(SceneMetadata.Roughness);
 	}	
diff --git a/Engine/Shaders/Private/ScreenSpaceDenoise/SSDSpatialAccumulation.usf b/Engine/Shaders/Private/ScreenSpaceDenoise/SSDSpatialAccumulation.usf
--- a/Engine/Shaders/Private/ScreenSpaceDenoise/SSDSpatialAccumulation.usf
+++ b/Engine/Shaders/Private/ScreenSpaceDenoise/SSDSpatialAccumulation.usf
@@ -15,6 +15,11 @@
 	#define SAMPLE_SET_5X5_WAVELET 2
 	#define SAMPLE_SET_STACKOWIAK_4_SETS 3
 
+/** Policy to use to change the size of kernel. */
+	#define SAMPLE_COUNT_POLICY_DISABLED 0
+	#define SAMPLE_COUNT_POLICY_HIT_DISTANCE_BASED 1
+	#define SAMPLE_COUNT_POLICY_ACCUMULATION_BASED 2
+
 
 //------------------------------------------------------- CONFIGS
 
@@ -41,8 +46,8 @@
 
 		// Use stackowiak sample set to keep noise in the signal to avoid flicker in history rejection,
 		// and adapt number of samples based on the sample.
-		#define CONFIG_SAMPLE_SET SAMPLE_SET_STACKOWIAK_4_SETS
-		#define CONFIG_SAMPLE_COUNT_POLICY 1
+		#define CONFIG_SAMPLE_SET           SAMPLE_SET_STACKOWIAK_4_SETS
+		#define CONFIG_SAMPLE_COUNT_POLICY  SAMPLE_COUNT_POLICY_HIT_DISTANCE_BASED
 
 		// Reconstruction has 1 ray per input pixel, therefore can use MIS.
 		#define CONFIG_SPECULAR_MIS 1
@@ -51,7 +56,10 @@
 		#define CONFIG_SIGNAL_INPUT_LAYOUT  SIGNAL_BUFFER_LAYOUT_REFLECTIONS_HISTORY
 		#define CONFIG_SIGNAL_OUTPUT_LAYOUT SIGNAL_BUFFER_LAYOUT_REFLECTIONS_HISTORY
 
-		#define CONFIG_SAMPLE_SET SAMPLE_SET_3X3
+		#define CONFIG_SAMPLE_SET           SAMPLE_SET_3X3
+		//#define CONFIG_SAMPLE_SET           SAMPLE_SET_STACKOWIAK_4_SETS
+		//#define CONFIG_SAMPLE_SUBSET        1
+		#define CONFIG_SAMPLE_COUNT_POLICY  SAMPLE_COUNT_POLICY_ACCUMULATION_BASED
 
 		// Post filtering can only be done using tokoyashi because have more than one sample at a time.
 		#define CONFIG_TOKOYASHI_BILATERAL 1
@@ -80,9 +88,14 @@
 	#define CONFIG_WEIGHT_KERNEL_CENTER 0
 #endif
 
-/** Config sample set. */
+/** Changes the logic controling the number of sample to do. */
 #ifndef CONFIG_SAMPLE_COUNT_POLICY
-	#define CONFIG_SAMPLE_COUNT_POLICY 0
+	#define CONFIG_SAMPLE_COUNT_POLICY SAMPLE_COUNT_POLICY_DISABLED
+#endif
+
+/** Selects a subset of sample of a given CONFIG_SAMPLE_SET */
+#ifndef CONFIG_SAMPLE_SUBSET
+	#define CONFIG_SAMPLE_SUBSET 0
 #endif
 
 
@@ -97,7 +110,7 @@ static const float kWaveletFilterWeights5x5[] = { 3.0 / 8.0, 1.0 / 4.0, 1.0 / 16
 
 
 // [ Stackowiak 2015, "Stochastic Screen-Space Reflections" ]
-static const float2 kStackowiakSamples[56 * 4] =
+static const float2 kStackowiakSampleSet0[56 * 4] =
 {
 	float2(-0.5, -0.5), float2(+0.5, -0.5), float2(-0.5, +0.5), float2(+0.5, +0.5),
 	float2(-1.5, +0.5), float2(-1.5, -0.5), float2(-0.5, +1.5), float2(+1.5, -0.5),
@@ -157,6 +170,66 @@ static const float2 kStackowiakSamples[56 * 4] =
 	float2(+7.5, +2.5), float2(-7.5, +3.5), float2(+4.5, -6.5), float2(+7.5, -4.5),
 };
 
+static const float2 kStackowiakSampleSet1[56 * 4] =
+{
+    float2(-0.5, -0.5), float2(+0.5, -0.5), float2(-0.5, +0.5), float2(+0.5, +0.5),
+    float2(+0.5, -1.5), float2(+1.5, -1.5), float2(-1.5, -0.5), float2(+1.5, +1.5),
+    float2(-0.5, -2.5), float2(-1.5, -1.5), float2(+0.5, +1.5), float2(-1.5, +0.5),
+    float2(+1.5, -0.5), float2(-0.5, +1.5), float2(-2.5, +0.5), float2(+0.5, +2.5),
+    float2(-2.5, -1.5), float2(+2.5, +0.5), float2(+1.5, +0.5), float2(-0.5, -1.5),
+    float2(-1.5, +1.5), float2(+2.5, -2.5), float2(-3.5, -0.5), float2(-1.5, +2.5),
+    float2(-2.5, +1.5), float2(-2.5, -0.5), float2(-1.5, -2.5), float2(+2.5, -1.5),
+    float2(-3.5, +0.5), float2(-0.5, -3.5), float2(-1.5, +3.5), float2(+0.5, -2.5),
+    float2(+1.5, +2.5), float2(-0.5, +2.5), float2(+0.5, +3.5), float2(+3.5, +0.5),
+    float2(+2.5, +1.5), float2(-2.5, -2.5), float2(+2.5, -0.5), float2(+3.5, -1.5),
+    float2(-0.5, +3.5), float2(+3.5, +1.5), float2(-3.5, +2.5), float2(+3.5, +2.5),
+    float2(+3.5, -0.5), float2(+0.5, -4.5), float2(-2.5, +3.5), float2(+0.5, -3.5),
+    float2(-1.5, -4.5), float2(+1.5, +3.5), float2(+1.5, -2.5), float2(-3.5, +1.5),
+    float2(+2.5, -3.5), float2(-2.5, -3.5), float2(+2.5, +2.5), float2(+1.5, +4.5),
+    float2(-4.5, -2.5), float2(-2.5, +2.5), float2(-4.5, +1.5), float2(+4.5, +1.5),
+    float2(-2.5, -4.5), float2(+3.5, -3.5), float2(-1.5, -3.5), float2(-3.5, -1.5),
+    float2(+1.5, -4.5), float2(+4.5, -2.5), float2(+1.5, -3.5), float2(-1.5, +4.5),
+    float2(-4.5, +2.5), float2(-4.5, -0.5), float2(+2.5, +4.5), float2(-4.5, +0.5),
+    float2(-3.5, -4.5), float2(+0.5, +4.5), float2(+3.5, -2.5), float2(-3.5, -2.5),
+    float2(-3.5, +3.5), float2(+3.5, +3.5), float2(+4.5, +0.5), float2(+0.5, +5.5),
+    float2(-0.5, +4.5), float2(+4.5, -3.5), float2(-1.5, +5.5), float2(-0.5, -4.5),
+    float2(+2.5, +3.5), float2(+4.5, +2.5), float2(-2.5, +5.5), float2(+2.5, -4.5),
+    float2(+4.5, -0.5), float2(+5.5, -0.5), float2(-4.5, +4.5), float2(+5.5, -1.5),
+    float2(-5.5, -1.5), float2(-4.5, -1.5), float2(+3.5, +4.5), float2(-3.5, -3.5),
+    float2(-5.5, +0.5), float2(+1.5, -5.5), float2(-5.5, -2.5), float2(-3.5, +4.5),
+    float2(+0.5, -5.5), float2(-2.5, -5.5), float2(+2.5, +5.5), float2(+4.5, +4.5),
+    float2(+4.5, -1.5), float2(-2.5, +4.5), float2(+4.5, +3.5), float2(+0.5, +6.5),
+    float2(-0.5, -6.5), float2(+5.5, +2.5), float2(-0.5, -5.5), float2(-5.5, -0.5),
+    float2(-6.5, -1.5), float2(-0.5, +5.5), float2(-0.5, +6.5), float2(+6.5, -0.5),
+    float2(+1.5, +5.5), float2(+1.5, -6.5), float2(+5.5, +0.5), float2(-5.5, +2.5),
+    float2(+5.5, +1.5), float2(-5.5, +1.5), float2(-6.5, -0.5), float2(-1.5, -5.5),
+    float2(-5.5, -4.5), float2(-4.5, +3.5), float2(-6.5, +1.5), float2(+2.5, -5.5),
+    float2(+3.5, -5.5), float2(-5.5, -3.5), float2(+1.5, +6.5), float2(+6.5, +2.5),
+    float2(+4.5, -4.5), float2(+3.5, -6.5), float2(-4.5, -4.5), float2(-4.5, -3.5),
+    float2(-6.5, +2.5), float2(+3.5, +5.5), float2(+3.5, -4.5), float2(+5.5, -3.5),
+    float2(-5.5, +4.5), float2(+6.5, -3.5), float2(-6.5, -2.5), float2(+5.5, +4.5),
+    float2(-1.5, +6.5), float2(-0.5, -7.5), float2(-6.5, +3.5), float2(-5.5, +3.5),
+    float2(-6.5, -4.5), float2(+7.5, -1.5), float2(-3.5, -5.5), float2(+3.5, +6.5),
+    float2(+5.5, +3.5), float2(+7.5, +0.5), float2(+5.5, -2.5), float2(-6.5, +0.5),
+    float2(-7.5, +1.5), float2(-3.5, -6.5), float2(+6.5, +0.5), float2(+7.5, +1.5),
+    float2(-2.5, -7.5), float2(-3.5, +5.5), float2(-7.5, -0.5), float2(-3.5, +6.5),
+    float2(-2.5, +6.5), float2(+4.5, -6.5), float2(-5.5, +5.5), float2(+4.5, -5.5),
+    float2(+6.5, -2.5), float2(+6.5, +3.5), float2(-1.5, -6.5), float2(-1.5, +7.5),
+    float2(+6.5, +1.5), float2(-5.5, -5.5), float2(+0.5, -6.5), float2(+7.5, +3.5),
+    float2(+2.5, +6.5), float2(-4.5, +5.5), float2(-6.5, -3.5), float2(-4.5, -5.5),
+    float2(-6.5, -5.5), float2(+5.5, -6.5), float2(-2.5, -6.5), float2(+5.5, -5.5),
+    float2(+4.5, +5.5), float2(-7.5, +0.5), float2(+6.5, -1.5), float2(+0.5, -7.5),
+    float2(+7.5, -0.5), float2(-3.5, -7.5), float2(+2.5, -6.5), float2(-3.5, +7.5),
+    float2(-4.5, -7.5), float2(-0.5, +7.5), float2(-6.5, +5.5), float2(+7.5, -3.5),
+    float2(-4.5, +6.5), float2(+1.5, +7.5), float2(+5.5, -4.5), float2(+7.5, +4.5),
+    float2(+0.5, +7.5), float2(+4.5, +6.5), float2(-4.5, +7.5), float2(-7.5, -1.5),
+    float2(+3.5, -7.5), float2(+7.5, -4.5), float2(+3.5, +7.5), float2(-1.5, -7.5),
+    float2(+6.5, -4.5), float2(-7.5, -3.5), float2(+6.5, +4.5), float2(+2.5, -7.5),
+    float2(+7.5, -2.5), float2(-7.5, +2.5), float2(+1.5, -7.5), float2(-5.5, +6.5),
+    float2(+5.5, +5.5), float2(-2.5, +7.5), float2(+7.5, +2.5), float2(-7.5, -2.5),
+    float2(+2.5, +7.5), float2(-6.5, +4.5), float2(+5.5, +6.5), float2(-4.5, -6.5),
+};
+
 static const uint kStackowiakSampleSetCount = 4;
 static const uint kStackowiakSampleCountPerSet = 56;
 
@@ -212,9 +285,44 @@ void MainCS(
 	}
 	#endif
 
+	/** Find out the number of samples that should be done.
+	 */
+	float RequestedSampleCount = 1024;
+	
+	#if CONFIG_SAMPLE_SET == SAMPLE_SET_NONE
+		RequestedSampleCount = 1;
+	#elif CONFIG_SAMPLE_COUNT_POLICY == SAMPLE_COUNT_POLICY_DISABLED
+		// NOP
+	#elif CONFIG_SAMPLE_COUNT_POLICY == SAMPLE_COUNT_POLICY_HIT_DISTANCE_BASED
+	{
+		float RefBluringRadius = ComputeBluringRadiusOfNormalizedSample(RefSample, RefSceneMetadata) * View.ViewSizeAndInvSize.x;
+		RequestedSampleCount = PI * RefBluringRadius * RefBluringRadius;
+		
+		// DebugOutput[DispatchThreadId] = float4(RefBluringRadius, RequestedSampleCount, 0, 0);
+	}
+	#elif CONFIG_SAMPLE_COUNT_POLICY == SAMPLE_COUNT_POLICY_ACCUMULATION_BASED
+	{
+		// TOOD: use ComputeBluringRadiusOfNormalizedSample instead.
+		float RefBluringRadius = ComputeInfinityBluringRadius(RefSceneMetadata) * View.ViewSizeAndInvSize.x;
+		float IdealSampleCount = PI * RefBluringRadius * RefBluringRadius;
+
+		const float MaxSampleCount = 1024; // TODO.
+
+		IdealSampleCount = min(IdealSampleCount, MaxSampleCount);
+
+		RequestedSampleCount = IdealSampleCount / clamp(RefSample.AccumulatedSampleCount, 1, MaxSampleCount);
+		
+		DebugOutput[DispatchThreadId] = float4(IdealSampleCount, RequestedSampleCount, 0, 0);
+	}
+	#else
+		#error Unknown policy to control the number of samples.
+	#endif
+
+	/** Iterate over the samples of the kernel.
+	 */
 	float2 KernelCenterBufferUV = SceneBufferUV;
 
-	#if SAMPLE_SET_NONE
+	#if CONFIG_SAMPLE_SET == SAMPLE_SET_NONE
 		if (0)
 
 	#elif CONFIG_SAMPLE_SET == SAMPLE_SET_3X3 || CONFIG_SAMPLE_SET == SAMPLE_SET_5X5_WAVELET
@@ -227,25 +335,12 @@ void MainCS(
 		KernelCenterBufferUV = ((DispatchThreadId | 1) + View.ViewRectMin.xy) * View.BufferSizeAndInvSize.zw;
 		KernelCenterBufferUV = clamp(KernelCenterBufferUV, View.BufferBilinearUVMinMax.xy, View.BufferBilinearUVMinMax.zw);
 
-		uint SampleCount = kStackowiakSampleCountPerSet;
-		#if CONFIG_SAMPLE_COUNT_POLICY == 1
-		{
-			float RefBluringRadius = ComputeBluringRadiusOfNormalizedSample(RefSample, RefSceneMetadata) * View.ViewSizeAndInvSize.x;
-			float TargetedSampleCount = PI * RefBluringRadius * RefBluringRadius * rcp(kStackowiakSampleSetCount);
-
-			const float MaximumRadius = 8;
-			const float MaximumSampleCount = kStackowiakSampleCountPerSet;
-
-			SampleCount = TargetedSampleCount;
-			
-			// DebugOutput[DispatchThreadId] = float4(RefBluringRadius, SampleCount, 0, 0);
-		}
-		#endif // CONFIG_SAMPLE_COUNT_POLICY == 1
+		uint SampleCount = clamp(RequestedSampleCount * rcp(kStackowiakSampleSetCount), 1, MaxSampleCount);
 
 		// Id of the pixel in the quad.
 		uint QuadSampleId = (DispatchThreadId.x & 1) | ((DispatchThreadId.y & 1) << 1);
 
-		LOOP for (uint SampleId = 1; SampleId < clamp(SampleCount, 1, MaxSampleCount); SampleId++)
+		LOOP for (uint SampleId = 1; SampleId < SampleCount; SampleId++)
 
 	#else
 		#error Unknown sample set
@@ -258,7 +353,13 @@ void MainCS(
 			PixelOffset = float2(x, y);
 			
 		#elif CONFIG_SAMPLE_SET == SAMPLE_SET_STACKOWIAK_4_SETS
-			PixelOffset = kStackowiakSamples[kStackowiakSampleSetCount * SampleId + QuadSampleId];
+			#if CONFIG_SAMPLE_SUBSET == 0
+				PixelOffset = kStackowiakSampleSet0[kStackowiakSampleSetCount * SampleId + QuadSampleId];
+			#elif CONFIG_SAMPLE_SUBSET == 1
+				PixelOffset = kStackowiakSampleSet1[kStackowiakSampleSetCount * SampleId + QuadSampleId];
+			#else
+				#error Unknown sample subset.
+			#endif
 
 		#elif CONFIG_SAMPLE_SET != SAMPLE_SET_NONE
 			#error Unknown sample set
diff --git a/Engine/Shaders/Private/ScreenSpaceDenoise/SSDTemporalAccumulation.usf b/Engine/Shaders/Private/ScreenSpaceDenoise/SSDTemporalAccumulation.usf
--- a/Engine/Shaders/Private/ScreenSpaceDenoise/SSDTemporalAccumulation.usf
+++ b/Engine/Shaders/Private/ScreenSpaceDenoise/SSDTemporalAccumulation.usf
@@ -38,7 +38,7 @@
 
 #elif CONFIG_SIGNAL_PROCESSING == SIGNAL_PROCESSING_REFLECTIONS
 	// Input and output layout.
-	#define CONFIG_SIGNAL_INPUT_LAYOUT   SIGNAL_BUFFER_LAYOUT_REFLECTIONS_INPUT
+	#define CONFIG_SIGNAL_INPUT_LAYOUT   SIGNAL_BUFFER_LAYOUT_REFLECTIONS_HISTORY
 	#define CONFIG_SIGNAL_HISTORY_LAYOUT SIGNAL_BUFFER_LAYOUT_REFLECTIONS_HISTORY
 
 	#define CONFIG_HISTORY_REJECTION 2
@@ -142,6 +142,11 @@ RWTexture2D<float4>	DebugOutput;
 
 //------------------------------------------------------- FUNCTIONS
 
+float Luma_To_LumaLog(float x)
+{
+	return max(x > 0 ? log2(x) : -100, -100);
+}
+
 FSSFSignalSample TransformSignalForPreRejection(FSSFSignalSample Sample)
 {
 	// Unpremultiplied RGBA to avoid any opacity mismatches.
@@ -552,7 +557,7 @@ void TemporallyAccumulate(
 	
 	FSSFSignalSample HistorySample = PreviousFrame;
 
-	const bool bPostRejectionBlending = !bBruteForce; // || CONFIG_SIGNAL_PROCESSING == SIGNAL_PROCESSING_REFLECTIONS;
+	const bool bPostRejectionBlending = !bBruteForce || CONFIG_SIGNAL_PROCESSING == SIGNAL_PROCESSING_REFLECTIONS;
 
 	if (!bPostRejectionBlending)
 	{
@@ -811,8 +816,53 @@ void TemporallyAccumulate(
 			NeighborMin = Add(NeighborMoment1, Mul(StdDev, -StdDevMultiplier));
 			NeighborMax = Add(NeighborMoment1, Mul(StdDev,  StdDevMultiplier));
 
-			float RejectedSampleCount = 0;
-			HistorySample = Mul(TransformSignalForPostRejection(Clamp(TransformSignalForPreRejection(NormalizeToOneSample(HistorySample)), NeighborMin, NeighborMax)), HistorySample.AccumulatedSampleCount - RejectedSampleCount);
+			if (0)
+			{
+				FSSFSignalSample QuantizationErrorMin = Mul(NeighborMoment1, 1 - SafeRcp(HistorySample.AccumulatedSampleCount));
+				FSSFSignalSample QuantizationErrorMax = Mul(NeighborMoment1, 1 + SafeRcp(HistorySample.AccumulatedSampleCount));
+
+				NeighborMin = Min(NeighborMin, QuantizationErrorMin);
+				NeighborMax = Max(NeighborMax, QuantizationErrorMax);
+			}
+
+			FSSFSignalSample NormalizedHistorySample = NormalizeToOneSample(HistorySample);
+			FSSFSignalSample TransformedHistorySample = TransformSignalForPreRejection(NormalizedHistorySample);
+
+			FSSFSignalSample ClampedTransformedHistorySample = Clamp(TransformedHistorySample, NeighborMin, NeighborMax);
+
+			FSSFSignalSample RejectedDiff = Abs(Subtract(ClampedTransformedHistorySample, TransformedHistorySample));
+
+			float RejectionFactor = 0.0;
+			#if CONFIG_SIGNAL_PROCESSING == SIGNAL_PROCESSING_REFLECTIONS && CONFIG_HISTORY_REJECTION_COLOR_SPACE == 1
+			{
+				RejectionFactor = abs(
+					Luma_To_LumaLog(ClampedTransformedHistorySample.SceneColor.x) -
+					Luma_To_LumaLog(TransformedHistorySample.SceneColor.x));
+				
+				RejectionFactor = max(RejectionFactor, 1 * max(RejectedDiff.SceneColor.y, RejectedDiff.SceneColor.z));
+				RejectionFactor = max(RejectionFactor, 1 * RejectedDiff.SceneColor.a);
+			}
+			#else
+				#error Unsupported signal rejection.
+			#endif
+			
+			DebugOutput[DispatchThreadId] = float4(
+				RejectionFactor,
+				0,
+				0,
+				0);
+
+			float FinalHistoryWeight = HistorySample.AccumulatedSampleCount * saturate(1 - RejectionFactor);
+
+			if (!bPostRejectionBlending)
+			{
+				FinalHistoryWeight = max(FinalHistoryWeight, CurrentFrame.AccumulatedSampleCount);
+			}
+
+			FSSFSignalSample ClampedHistorySample = TransformSignalForPostRejection(ClampedTransformedHistorySample);
+			FSSFSignalSample FinalHistorySample = Mul(ClampedHistorySample, FinalHistoryWeight);
+
+			HistorySample = FinalHistorySample;
 		}
 		#else // CONFIG_HISTORY_REJECTION == 1
 		{
diff --git a/Engine/Source/Runtime/Renderer/Private/ScreenSpaceDenoise.cpp b/Engine/Source/Runtime/Renderer/Private/ScreenSpaceDenoise.cpp
--- a/Engine/Source/Runtime/Renderer/Private/ScreenSpaceDenoise.cpp
+++ b/Engine/Source/Runtime/Renderer/Private/ScreenSpaceDenoise.cpp
@@ -43,11 +43,9 @@ static TAutoConsoleVariable<int32> CVarReflectionTemporalAccumulation(
 	TEXT(""),
 	ECVF_RenderThreadSafe);
 
-static TAutoConsoleVariable<int32> CVarReflectionHistoryConvolution(
-	TEXT("r.Reflection.Denoise.HistoryConvolution"), 0, // TODO: enable by default
-	TEXT("Mode to use for history convolution.\n")
-	TEXT(" 0: disabled;\n")
-	TEXT(" 1: Spatial convolution (default)."),
+static TAutoConsoleVariable<int32> CVarReflectionHistoryConvolutionSampleCount(
+	TEXT("r.Reflection.Denoise.HistoryConvolutionSampleCount"), 1,
+	TEXT("Number of samples to use for history post filter (default = 1)."),
 	ECVF_RenderThreadSafe);
 
 
@@ -474,8 +472,8 @@ static void DenoiseShadowPenumbra(
 	if (View.ViewState && CVarShadowHaveTAA.GetValueOnRenderThread())
 	{
 		// Generate rejection signal history.
-		FRDGTextureRef HistoryRejectionSignal0;
-		FRDGTextureRef HistoryRejectionSignal1;
+		FRDGTextureRef HistoryRejectionSignal0 = nullptr;
+		FRDGTextureRef HistoryRejectionSignal1 = nullptr;
 		if (CVarHistoryRejection.GetValueOnRenderThread() == 1)
 		{
 			FRDGTextureDesc Desc = SignalProcessingDesc;
@@ -663,8 +661,8 @@ IScreenSpaceDenoiser::FReflectionOutputs DenoiseReflections(
 	
 	// Spatial reconstruction with multiple important sampling to be more precise in the history rejection.
 	{
-		FRDGTextureRef SignalOutput0 = GraphBuilder.CreateTexture(SignalProcessingDesc, TEXT("SSDReflectionsReconstruction"));
-		FRDGTextureRef SignalOutput1 = GraphBuilder.CreateTexture(SignalProcessingDesc, TEXT("SSDReflectionsReconstruction"));
+		FRDGTextureRef SignalOutput0 = GraphBuilder.CreateTexture(SignalProcessingDesc, TEXT("SSDReflectionsReconstruction0"));
+		FRDGTextureRef SignalOutput1 = GraphBuilder.CreateTexture(SignalProcessingDesc, TEXT("SSDReflectionsReconstruction1"));
 
 		FSSDSpatialAccumulationCS::FParameters* PassParameters = GraphBuilder.AllocParameters<FSSDSpatialAccumulationCS::FParameters>();
 		PassParameters->MaxSampleCount = FMath::Clamp(CVarReflectionReconstructionSampleCount.GetValueOnRenderThread(), 1, kStackowiakMaxSampleCountPerSet);
@@ -683,7 +681,7 @@ IScreenSpaceDenoiser::FReflectionOutputs DenoiseReflections(
 		TShaderMapRef<FSSDSpatialAccumulationCS> ComputeShader(View.ShaderMap, PermutationVector);
 		FComputeShaderUtils::AddPass(
 			GraphBuilder,
-			RDG_EVENT_NAME("SSD SpatialAccumulation(Reconstruction Samples=%i)", PassParameters->MaxSampleCount),
+			RDG_EVENT_NAME("SSD SpatialAccumulation(Reconstruction MaxSamples=%i)", PassParameters->MaxSampleCount),
 			*ComputeShader,
 			PassParameters,
 			FComputeShaderUtils::GetGroupCount(DenoiseResolution, FComputeShaderUtils::kGolden2DGroupSize));
@@ -693,7 +691,7 @@ IScreenSpaceDenoiser::FReflectionOutputs DenoiseReflections(
 	}
 
 	// Temporal pass.
-	if (CVarReflectionTemporalAccumulation.GetValueOnRenderThread())
+	if (View.ViewState && CVarReflectionTemporalAccumulation.GetValueOnRenderThread())
 	{
 		FRDGTextureRef SignalOutput0 = GraphBuilder.CreateTexture(SignalProcessingDesc, TEXT("SSDReflectionsHistory0"));
 		FRDGTextureRef SignalOutput1 = GraphBuilder.CreateTexture(SignalProcessingDesc, TEXT("SSDReflectionsHistory1"));
@@ -737,12 +735,14 @@ IScreenSpaceDenoiser::FReflectionOutputs DenoiseReflections(
 	}
 	
 	// Spatial filter, to converge history faster.
-	if (CVarReflectionHistoryConvolution.GetValueOnRenderThread())
+	int32 MaxPostFilterSampleCount = CVarReflectionHistoryConvolutionSampleCount.GetValueOnRenderThread();
+	if (MaxPostFilterSampleCount > 1)
 	{
 		FRDGTextureRef SignalOutput0 = GraphBuilder.CreateTexture(SignalProcessingDesc, TEXT("SSDReflectionsHistory0"));
 		FRDGTextureRef SignalOutput1 = GraphBuilder.CreateTexture(SignalProcessingDesc, TEXT("SSDReflectionsHistory1"));
 
 		FSSDSpatialAccumulationCS::FParameters* PassParameters = GraphBuilder.AllocParameters<FSSDSpatialAccumulationCS::FParameters>();
+		PassParameters->MaxSampleCount = FMath::Clamp(MaxPostFilterSampleCount, 1, kStackowiakMaxSampleCountPerSet);
 		PassParameters->CommonParameters = CommonParameters;
 		PassParameters->SignalInput0 = SignalHistory0;
 		PassParameters->SignalInput1 = SignalHistory1;
@@ -752,11 +752,13 @@ IScreenSpaceDenoiser::FReflectionOutputs DenoiseReflections(
 		FSSDSpatialAccumulationCS::FPermutationDomain PermutationVector;
 		PermutationVector.Set<FSignalProcessingDim>(ESignalProcessing::Reflections);
 		PermutationVector.Set<FSSDSpatialAccumulationCS::FStageDim>(FSSDSpatialAccumulationCS::EStage::PostFiltering);
+		
+		PassParameters->DebugOutput = GraphBuilder.CreateUAV(GraphBuilder.CreateTexture(SignalProcessingDesc, TEXT("SSDDebugReflectionPostfilter")));
 
 		TShaderMapRef<FSSDSpatialAccumulationCS> ComputeShader(View.ShaderMap, PermutationVector);
 		FComputeShaderUtils::AddPass(
 			GraphBuilder,
-			RDG_EVENT_NAME("SSD SpatialAccumulation(PostFiltering)"),
+			RDG_EVENT_NAME("SSD SpatialAccumulation(PostFiltering MaxSamples=%i)", MaxPostFilterSampleCount),
 			*ComputeShader,
 			PassParameters,
 			FComputeShaderUtils::GetGroupCount(DenoiseResolution, FComputeShaderUtils::kGolden2DGroupSize));
@@ -781,6 +783,13 @@ IScreenSpaceDenoiser::FReflectionOutputs DenoiseReflections(
 			GraphBuilder.QueueTextureExtraction(SignalHistory1, &NewHistory.RT[1]);
 		}
 	}
+	else
+	{
+		// The SignalHistory1 is always generated for temporal history, but will endup useless if there is no view state,
+		// in witch case we do not extract any textures. Don't support a shader permutation that does not produce it, because
+		// it is already a not ideal case for the denoiser.
+		//GraphBuilder.FlagUsedOnceTexture(SignalHistory1);
+	}
 
 	IScreenSpaceDenoiser::FReflectionOutputs Outputs;
 	Outputs.Color = SignalHistory0; // TODO.
