Fixed maximum shadow ray length parameter in ray traced reflection lighting

diff --git a/Engine/Shaders/Private/RayTracing/RayTracingReflections.usf b/Engine/Shaders/Private/RayTracing/RayTracingReflections.usf
--- a/Engine/Shaders/Private/RayTracing/RayTracingReflections.usf
+++ b/Engine/Shaders/Private/RayTracing/RayTracingReflections.usf
@@ -229,6 +229,8 @@ float3 ComputeDirectLighting(RayDesc Ray, FMaterialClosestHitPayload Payload, bo
 			//#dxr_todo: add support for sampling area lights (and perhaps add CVar to enable stocastic shadows?)
 
 			float3 ShadowRayDirection;
+			float3 ToLight = LightData.Position - WorldPosition;
+			float DistanceToLight = length(ToLight);
 			
 			if (LightType == LIGHT_TYPE_DIRECTIONAL)
 			{
@@ -236,10 +238,10 @@ float3 ComputeDirectLighting(RayDesc Ray, FMaterialClosestHitPayload Payload, bo
 			}
 			else
 			{
-				ShadowRayDirection = normalize(LightData.Position - WorldPosition);
+				ShadowRayDirection = ToLight / DistanceToLight;
 			}
 
-			float HitT = TraceShadowRay(WorldPosition + Payload.WorldNormal*0.001, ShadowRayDirection, ReflectionMaxRayDistance);
+			float HitT = TraceShadowRay(WorldPosition + Payload.WorldNormal*0.001, ShadowRayDirection, DistanceToLight);
 			AmbientOcclusion = HitT < 0.0f;
 		}
 
