Added simple hemisphere and distance culling to direct lighting in ray traced reflections. Added r.RayTracing.Reflections.MaxRoughness CVar to remove RT reflections from rough surfaces (default = 1, meaning no roughness culling).

diff --git a/Engine/Shaders/Private/RayTracing/RayTracingReflections.usf b/Engine/Shaders/Private/RayTracing/RayTracingReflections.usf
--- a/Engine/Shaders/Private/RayTracing/RayTracingReflections.usf
+++ b/Engine/Shaders/Private/RayTracing/RayTracingReflections.usf
@@ -34,6 +34,7 @@ uint ShouldDoDirectLighting;
 uint ShouldDoReflectedShadows;
 uint ShouldDoEmissiveAndIndirectLighting;
 float ReflectionMaxRayDistance;
+float ReflectionMaxRoughness;
 
 RaytracingAccelerationStructure TLAS;
 
@@ -220,27 +221,32 @@ float3 ComputeDirectLighting(RayDesc Ray, FMaterialClosestHitPayload Payload, bo
 		Texture2D LightSourceTexture = ReflectionLightsData.DummyRectLightTexture;
 
 		float4 LightAttenuation = 1.0f;
-		float LightProfileMultiplier = 1.0;//#dxr_todo
+		float LightProfileMultiplier = 1.0; //#dxr_todo: implement support for light profiles
+
+		float3 ShadowRayDirection;
+		float3 ToLight = LightData.Position - WorldPosition;
+		float LightMask = 1.0;
+		
+		if (LightType == LIGHT_TYPE_DIRECTIONAL)
+		{
+			ShadowRayDirection = LightData.Direction;
+		}
+		else
+		{
+			LightMask = GetLocalLightAttenuation(WorldPosition, LightData, ToLight, ShadowRayDirection);
+		}
+
+		// Skip the light sample that's below the horizon or does not contribute anything due to attenuation.
+		if (dot(Payload.WorldNormal, ToLight) <= 0 || LightMask <= 0.0)
+		{
+			continue;
+		}
 
 		//#dxr_todo: check perf improvements forcing TMin == TMax when bReflectedShadows = false, to avoid tracing rays inside if statements
 		if (bReflectedShadows)
 		{
-			//#dxr_todo: add early outs if backfacing or if light energy is not sufficient, we are inside a light volume, etc
 			//#dxr_todo: add support for sampling area lights (and perhaps add CVar to enable stocastic shadows?)
-
-			float3 ShadowRayDirection;
-			float3 ToLight = LightData.Position - WorldPosition;
 			float DistanceToLight = length(ToLight);
-			
-			if (LightType == LIGHT_TYPE_DIRECTIONAL)
-			{
-				ShadowRayDirection = LightData.Direction;
-			}
-			else
-			{
-				ShadowRayDirection = ToLight / DistanceToLight;
-			}
-
 			float HitT = TraceShadowRay(WorldPosition + Payload.WorldNormal*0.001, ShadowRayDirection, DistanceToLight);
 			AmbientOcclusion = HitT < 0.0f;
 		}
@@ -316,7 +322,6 @@ void RayTracingReflectionsRGS()
 
 	uint MaxBounces = 1; //#dxr-todo: only works with 1 bounce until hit shaders are supported
 	uint LocalSamplesPerPixel = SamplesPerPixel;
-	float RoughnessThreshold = 1.0f; //#dxr-todo: for internal testing only, remove later
 
 	// Mask out depth values at or beyond far plane
 	if (Depth >= SCENE_TEXTURES_DISABLED_SCENE_DEPTH_VALUE)
@@ -328,7 +333,7 @@ void RayTracingReflectionsRGS()
 	float ReflectedWeight = 0.0;
 	float HitDistance = -2.0;
 
-	if (Roughness < RoughnessThreshold)
+	if (Roughness < ReflectionMaxRoughness)
 	{
 		for (uint SampleIndex = 0; SampleIndex < LocalSamplesPerPixel; ++SampleIndex)
 		{
diff --git a/Engine/Source/Runtime/Renderer/Private/RayTracing/RayTracingReflections.cpp b/Engine/Source/Runtime/Renderer/Private/RayTracing/RayTracingReflections.cpp
--- a/Engine/Source/Runtime/Renderer/Private/RayTracing/RayTracingReflections.cpp
+++ b/Engine/Source/Runtime/Renderer/Private/RayTracing/RayTracingReflections.cpp
@@ -44,6 +44,13 @@ static FAutoConsoleVariableRef CVarRayTracingReflectionsMaxRayDistance(
 	TEXT("Sets the maximum ray distance for ray traced reflection rays (default = 1.0e27)")
 );
 
+static float GRayTracingReflectionsMaxRoughness = 1.0;
+static FAutoConsoleVariableRef CVarRayTracingReflectionsMaxRoughness(
+	TEXT("r.RayTracing.Reflections.MaxRoughness"),
+	GRayTracingReflectionsMaxRoughness,
+	TEXT("Surfaces with roughness larger than this value will not receive ray traced reflections (default = 1)")
+);
+
 static const int32 GReflectionLightCountMaximum = 64;
 
 BEGIN_GLOBAL_SHADER_PARAMETER_STRUCT(FReflectionsLightData, )
@@ -126,6 +133,7 @@ class FRayTracingReflectionsRG : public FGlobalShader
 		SHADER_PARAMETER(int32, ShouldDoReflectedShadows)
 		SHADER_PARAMETER(int32, ShouldDoEmissiveAndIndirectLighting)
 		SHADER_PARAMETER(float, ReflectionMaxRayDistance)
+		SHADER_PARAMETER(float, ReflectionMaxRoughness)
 
 		SHADER_PARAMETER_SRV(RaytracingAccelerationStructure, TLAS)
 
@@ -210,6 +218,7 @@ void FDeferredShadingSceneRenderer::RayTraceReflections(
 	PassParameters->ShouldDoReflectedShadows = GRayTracingReflectionsShadows;
 	PassParameters->ShouldDoEmissiveAndIndirectLighting = GRayTracingReflectionsEmissiveAndIndirectLighting;
 	PassParameters->ReflectionMaxRayDistance = GRayTracingReflectionsMaxRayDistance;
+	PassParameters->ReflectionMaxRoughness = GRayTracingReflectionsMaxRoughness;
 	PassParameters->LTCMatTexture = GSystemTextures.LTCMat->GetRenderTargetItem().ShaderResourceTexture;
 	PassParameters->LTCMatSampler = TStaticSamplerState<SF_Bilinear, AM_Clamp, AM_Clamp, AM_Clamp>::GetRHI();
 	PassParameters->LTCAmpTexture = GSystemTextures.LTCAmp->GetRenderTargetItem().ShaderResourceTexture;
