Old atmosphere common workaround for NaN in Metal output is no longer necessary from Metal 2.0 onward - Metal 1.2 on iOS/tvOS is probably still broken.

diff --git a/Engine/Shaders/Private/AtmosphereCommon.ush b/Engine/Shaders/Private/AtmosphereCommon.ush
--- a/Engine/Shaders/Private/AtmosphereCommon.ush
+++ b/Engine/Shaders/Private/AtmosphereCommon.ush
@@ -23,9 +23,10 @@
 #ifndef __ATMOSPHERE_COMMON__
 #define __ATMOSPHERE_COMMON__
 
-#if COMPILER_METAL
+#if COMPILER_METAL && MAX_SHADER_LANGUAGE_VERSION < 3
 // For inexplicable reasons until we can debug shaders and *really* see what's what we need to abs sqrt input for the functions in this shader file only.
 // As best as the vendors can tell us Metal has *slightly* different precision in some results which means we are giving sqrt negative & zero values.
+// This should now only be true on Metal 1.2 and earlier.
 #define sqrt(...) sqrt(abs(__VA_ARGS__))
 #endif
 
@@ -623,7 +624,7 @@ float4 CalculateVertexAtmosphericFog(float3 WorldPosition, float3 InCameraPositi
 	return GetAtmosphericFog(InCameraPosition, ViewVector, Distance, float3(0.f, 0.f, 0.f));
 }
 
-#if COMPILER_METAL
+#if COMPILER_METAL && MAX_SHADER_LANGUAGE_VERSION < 3
 // Do not export our workaround outside this file as it might not be desirable.
 #undef sqrt
 #endif
