Bumping ShaderVersion.ush now recompiles ray tracing shaders.  ShaderVersion.ush is the root version for all shaders.

diff --git a/Engine/Shaders/Private/RayTracing/RayTracingCommon.ush b/Engine/Shaders/Private/RayTracing/RayTracingCommon.ush
--- a/Engine/Shaders/Private/RayTracing/RayTracingCommon.ush
+++ b/Engine/Shaders/Private/RayTracing/RayTracingCommon.ush
@@ -9,6 +9,9 @@
 #ifndef RAYTRACINGCOMMON_USH_INCLUDED
 #define RAYTRACINGCOMMON_USH_INCLUDED // Workarround for UE-66460
 
+// Make sure we recompile ray tracing shaders if the main shader version changes
+#include "ShaderVersion.ush"
+
 #include "RayTracingDefinitions.ush"
 
 struct FBasicRayData
