Fixed include path

diff --git a/Engine/Shaders/Private/RayTracing/RayTracingCommon.ush b/Engine/Shaders/Private/RayTracing/RayTracingCommon.ush
--- a/Engine/Shaders/Private/RayTracing/RayTracingCommon.ush
+++ b/Engine/Shaders/Private/RayTracing/RayTracingCommon.ush
@@ -10,7 +10,7 @@
 #define RAYTRACINGCOMMON_USH_INCLUDED // Workarround for UE-66460
 
 // Make sure we recompile ray tracing shaders if the main shader version changes
-#include "ShaderVersion.ush"
+#include "/Engine/Public/ShaderVersion.ush"
 
 #include "RayTracingDefinitions.ush"
 
