Allows the reflections to be ray traced at lower resolution and upsampled by the denoiser.

diff --git a/Engine/Shaders/Private/RayTracing/RayTracingReflections.usf b/Engine/Shaders/Private/RayTracing/RayTracingReflections.usf
--- a/Engine/Shaders/Private/RayTracing/RayTracingReflections.usf
+++ b/Engine/Shaders/Private/RayTracing/RayTracingReflections.usf
@@ -33,6 +33,8 @@ uint SamplesPerPixel;
 uint ShouldDoDirectLighting;
 uint ShouldDoReflectedShadows;
 uint ShouldDoEmissiveAndIndirectLighting;
+uint UpscaleFactor;
+
 float ReflectionMaxRayDistance;
 float ReflectionMaxRoughness;
 
@@ -47,11 +49,6 @@ float GetRoughnessFade(float Roughness, float MaxRoughness)
 	return saturate(Roughness * RoughnessMaskScale + 2.0);
 }
 
-uint CalcLinearIndex(uint2 PixelCoord)
-{
-	return PixelCoord.y * View.BufferSizeAndInvSize.x + PixelCoord.x;
-}
-
 // #dxr-todo: move to RayTracingCommon.ush
 float3 ReconstructWorldPositionFromDepth(float2 UV, float Depth)
 {
@@ -303,14 +300,27 @@ FMaterialClosestHitPayload TraceRayAndAccumulateResults(
 	return Payload;
 }
 
+uint2 GetPixelCoord(uint2 DispatchThreadId)
+{
+	uint UpscaleFactorPow2 = UpscaleFactor * UpscaleFactor;
+
+	// TODO: find a way to not interfer with TAA's jittering.
+	uint SubPixelId = View.StateFrameIndex & (UpscaleFactorPow2 - 1);
+
+	return DispatchThreadId * UpscaleFactor + uint2(SubPixelId & (UpscaleFactor - 1), SubPixelId / UpscaleFactor);
+}
+
 
 [shader("raygeneration")]
 void RayTracingReflectionsRGS()
 {
-	uint2 PixelCoord = DispatchRaysIndex().xy;
+	uint2 DispatchThreadId = DispatchRaysIndex().xy;
 
+	uint2 PixelCoord = GetPixelCoord(DispatchThreadId);
+	
+	uint LinearIndex = PixelCoord.y * View.BufferSizeAndInvSize.x + PixelCoord.x; // TODO: PixelCoord or DispatchThreadId
+	
 	RandomSequence RandSequence;
-	uint LinearIndex = CalcLinearIndex(PixelCoord);
 	RandomSequence_Initialize(RandSequence, LinearIndex, View.FrameNumber);
 
 	float2 InvBufferSize = View.BufferSizeAndInvSize.zw;
@@ -358,7 +368,7 @@ void RayTracingReflectionsRGS()
 
 				if (LocalSamplesPerPixel == 1)
 				{
-					uint2 SobolFrame = ComputePixelUniqueSobolRandSample(PixelCoord);
+					uint2 SobolFrame = ComputePixelUniqueSobolRandSample(DispatchThreadId);
 					RandSample = SobolIndexToUniformUnitSquare(SobolFrame);
 				}
 
@@ -420,8 +430,8 @@ void RayTracingReflectionsRGS()
 
 	ReflectedColor *= RoughnessFade;
 
-	ColorOutput[PixelCoord] = float4(ReflectedColor);
-	RayHitDistanceOutput[PixelCoord] = HitDistance;
+	ColorOutput[DispatchThreadId] = float4(ReflectedColor);
+	RayHitDistanceOutput[DispatchThreadId] = HitDistance;
 }
 
 
diff --git a/Engine/Shaders/Private/ScreenSpaceDenoise/SSDSpatialAccumulation.usf b/Engine/Shaders/Private/ScreenSpaceDenoise/SSDSpatialAccumulation.usf
--- a/Engine/Shaders/Private/ScreenSpaceDenoise/SSDSpatialAccumulation.usf
+++ b/Engine/Shaders/Private/ScreenSpaceDenoise/SSDSpatialAccumulation.usf
@@ -17,8 +17,9 @@
 
 /** Policy to use to change the size of kernel. */
 	#define SAMPLE_COUNT_POLICY_DISABLED 0
-	#define SAMPLE_COUNT_POLICY_HIT_DISTANCE_BASED 1
-	#define SAMPLE_COUNT_POLICY_ACCUMULATION_BASED 2
+	#define SAMPLE_COUNT_POLICY_METADATA_BASED 1
+	#define SAMPLE_COUNT_POLICY_HIT_DISTANCE_BASED 2
+	#define SAMPLE_COUNT_POLICY_ACCUMULATION_BASED 3
 
 
 //------------------------------------------------------- CONFIGS
@@ -52,6 +53,9 @@
 		// Reconstruction has 1 ray per input pixel, therefore can use MIS.
 		#define CONFIG_SPECULAR_MIS 1
 
+		#define CONFIG_UPSCALE              DIM_UPSCALE
+		//#define CONFIG_SAMPLE_COUNT_POLICY  SAMPLE_COUNT_POLICY_DISABLED
+
 	#elif DIM_STAGE == STAGE_POST_FILTERING
 		#define CONFIG_SIGNAL_INPUT_LAYOUT  SIGNAL_BUFFER_LAYOUT_REFLECTIONS_HISTORY
 		#define CONFIG_SIGNAL_OUTPUT_LAYOUT SIGNAL_BUFFER_LAYOUT_REFLECTIONS_HISTORY
@@ -98,6 +102,11 @@
 	#define CONFIG_SAMPLE_SUBSET 0
 #endif
 
+/** Whether the ray tracing input may needs to be upscale to the view's resolution. */
+#ifndef CONFIG_UPSCALE
+	#define CONFIG_UPSCALE 0
+#endif
+
 
 //------------------------------------------------------- INCLUDES
 
@@ -237,6 +246,7 @@ static const uint kStackowiakSampleCountPerSet = 56;
 //------------------------------------------------------- PARAMETERS
 
 uint MaxSampleCount;
+uint UpscaleFactor;
 
 Texture2D SignalInput0;
 Texture2D SignalInput1;
@@ -249,6 +259,17 @@ RWTexture2D<float4>	DebugOutput;
 
 //------------------------------------------------------- FUNCTIONS
 
+// TODO: duplicated with reflection code.
+uint2 GetPixelCoord(uint2 DispatchThreadId)
+{
+	uint UpscaleFactorPow2 = UpscaleFactor * UpscaleFactor;
+
+	// TODO: find a way to not interfer with TAA's jittering.
+	uint SubPixelId = View.StateFrameIndex & (UpscaleFactorPow2 - 1);
+
+	return DispatchThreadId * UpscaleFactor + uint2(SubPixelId & (UpscaleFactor - 1), SubPixelId / UpscaleFactor);
+}
+
 
 //------------------------------------------------------- ENTRY POINTS
 
@@ -270,12 +291,19 @@ void MainCS(
 	float2 ScreenPosition = ViewportUVToScreenPos(BufferUVToViewportUV(SceneBufferUV));
 	FSSFSampleSceneInfos RefSceneMetadata = FetchCurrentSceneInfos(ScreenPosition, SceneBufferUV);
 
-	FSSFSignalSample RefSample = SampleInputSignal(SignalInput0, SignalInput1, GlobalPointClampedSampler, CONFIG_SIGNAL_INPUT_LAYOUT, SceneBufferUV);
+	// Sample the reference sample.
+	#if !CONFIG_UPSCALE || 1 // TODO: this sample ends up being sampled twice when upscaling.
+		FSSFSignalSample RefSample = SampleInputSignal(SignalInput0, SignalInput1, GlobalPointClampedSampler, CONFIG_SIGNAL_INPUT_LAYOUT, SceneBufferUV);
+	#endif
 
-	// Read reference data.
-	float AccumulatorWeight = 1; //GetMultiImportanceSamplingWeight(RefSceneMetadata, RefSceneMetadata, RefSample, DispatchThreadId);
-	FSSFSignalSample Accumulator = Mul(RefSample, AccumulatorWeight);
+	// Setup the spatial accumulator.
+	FSSFSignalSample Accumulator = CreateSignalSampleFromScalarValue(0.0);
+	#if !CONFIG_UPSCALE
+		float AccumulatorWeight = 1; //GetMultiImportanceSamplingWeight(RefSceneMetadata, RefSceneMetadata, RefSample, DispatchThreadId);
 	
+		Accumulator = Mul(RefSample, AccumulatorWeight);
+	#endif
+
 	float CurrentWorldBluringRadius = ComputePixelWorldRadius(RefSceneMetadata);
 
 	float BilateralWeightAdjustement = 1;
@@ -284,26 +312,34 @@ void MainCS(
 		BilateralWeightAdjustement = rcp(ComputeBilateralWeight(100, RefSceneMetadata, RefSceneMetadata));
 	}
 	#endif
+	
+	/** factor by witch should be spread out. */
+	float KernelSpreadFactor = 1;
+	#if CONFIG_UPSCALE
+		KernelSpreadFactor = UpscaleFactor;
+	#endif
 
-	/** Find out the number of samples that should be done.
-	 */
+	/** Find out the number of samples that should be done. */
 	float RequestedSampleCount = 1024;
 	
 	#if CONFIG_SAMPLE_SET == SAMPLE_SET_NONE
 		RequestedSampleCount = 1;
 	#elif CONFIG_SAMPLE_COUNT_POLICY == SAMPLE_COUNT_POLICY_DISABLED
 		// NOP
+	#elif CONFIG_SAMPLE_COUNT_POLICY == SAMPLE_COUNT_POLICY_METADATA_BASED
+	{
+		float RefBluringRadius = ComputeInfinityBluringRadius(RefSceneMetadata) * View.ViewSizeAndInvSize.x * rcp(KernelSpreadFactor);
+		RequestedSampleCount = PI * RefBluringRadius * RefBluringRadius;
+	}
 	#elif CONFIG_SAMPLE_COUNT_POLICY == SAMPLE_COUNT_POLICY_HIT_DISTANCE_BASED
 	{
-		float RefBluringRadius = ComputeBluringRadiusOfNormalizedSample(RefSample, RefSceneMetadata) * View.ViewSizeAndInvSize.x;
+		float RefBluringRadius = ComputeBluringRadiusOfNormalizedSample(RefSample, RefSceneMetadata) * View.ViewSizeAndInvSize.x * rcp(KernelSpreadFactor);
 		RequestedSampleCount = PI * RefBluringRadius * RefBluringRadius;
-		
-		// DebugOutput[DispatchThreadId] = float4(RefBluringRadius, RequestedSampleCount, 0, 0);
 	}
 	#elif CONFIG_SAMPLE_COUNT_POLICY == SAMPLE_COUNT_POLICY_ACCUMULATION_BASED
 	{
 		// TOOD: use ComputeBluringRadiusOfNormalizedSample instead.
-		float RefBluringRadius = ComputeInfinityBluringRadius(RefSceneMetadata) * View.ViewSizeAndInvSize.x;
+		float RefBluringRadius = ComputeInfinityBluringRadius(RefSceneMetadata) * View.ViewSizeAndInvSize.x * rcp(KernelSpreadFactor);
 		float IdealSampleCount = PI * RefBluringRadius * RefBluringRadius;
 
 		const float MaxSampleCount = 1024; // TODO.
@@ -318,29 +354,54 @@ void MainCS(
 		#error Unknown policy to control the number of samples.
 	#endif
 
+
 	/** Iterate over the samples of the kernel.
 	 */
 	float2 KernelCenterBufferUV = SceneBufferUV;
 
 	#if CONFIG_SAMPLE_SET == SAMPLE_SET_NONE
+		#if CONFIG_UPSCALE
+			#error Unsupported
+		#endif
 		if (0)
 
 	#elif CONFIG_SAMPLE_SET == SAMPLE_SET_3X3 || CONFIG_SAMPLE_SET == SAMPLE_SET_5X5_WAVELET
+		#if CONFIG_UPSCALE
+			#error Unsupported
+		#endif
 		const int KernelRadius = CONFIG_SAMPLE_SET == SAMPLE_SET_3X3 ? 1 : 2;
 		
 		UNROLL for (int x = -KernelRadius; x <= KernelRadius; x++)
 		UNROLL for (int y = -KernelRadius; y <= KernelRadius; y++)
 
 	#elif CONFIG_SAMPLE_SET == SAMPLE_SET_STACKOWIAK_4_SETS
-		KernelCenterBufferUV = ((DispatchThreadId | 1) + View.ViewRectMin.xy) * View.BufferSizeAndInvSize.zw;
+		#if CONFIG_UPSCALE
+			// TODO: could be optimised, but currently reusing same peace of code as reflection for maintainability.
+			uint2 RayDispatchThreadId = (DispatchThreadId - UpscaleFactor / 2) / UpscaleFactor;
+			uint2 ClosestRayPixelCoord = GetPixelCoord(RayDispatchThreadId);
+			
+			uint RaySubPixelId = View.StateFrameIndex & (UpscaleFactor * UpscaleFactor - 1);
+
+			KernelCenterBufferUV = ((ClosestRayPixelCoord + (0.5 * KernelSpreadFactor + 0.5)) + View.ViewRectMin.xy) * View.BufferSizeAndInvSize.zw;
+			
+			// Id of the pixel in the quad.
+			uint QuadSampleSubsetId = ((DispatchThreadId.x & 1) | ((DispatchThreadId.y & 1) << 1)) ^ 0x3;
+
+			bool bTakeFirstRay = RaySubPixelId == ((DispatchThreadId.x & 1) | ((DispatchThreadId.y & 1) << 1));
+
+		#else // !CONFIG_UPSCALE
+			// Put the kernel center at the center of the quad. Half pixel shift is done in the sample offsets.
+			KernelCenterBufferUV = ((DispatchThreadId | 1) + View.ViewRectMin.xy) * View.BufferSizeAndInvSize.zw;
+
+			// Id of the pixel in the quad. This is to match hard coded first samples of the sample set.
+			uint QuadSampleSubsetId = ((DispatchThreadId.x & 1) | ((DispatchThreadId.y & 1) << 1));
+		#endif
+
 		KernelCenterBufferUV = clamp(KernelCenterBufferUV, View.BufferBilinearUVMinMax.xy, View.BufferBilinearUVMinMax.zw);
 
 		uint SampleCount = clamp(RequestedSampleCount * rcp(kStackowiakSampleSetCount), 1, MaxSampleCount);
 
-		// Id of the pixel in the quad.
-		uint QuadSampleId = (DispatchThreadId.x & 1) | ((DispatchThreadId.y & 1) << 1);
-
-		LOOP for (uint SampleId = 1; SampleId < SampleCount; SampleId++)
+		LOOP for (uint SampleId = (CONFIG_UPSCALE ? 0 : 1); SampleId < SampleCount; SampleId++)
 
 	#else
 		#error Unknown sample set
@@ -349,14 +410,14 @@ void MainCS(
 		float2 PixelOffset = 0;
 
 		#if CONFIG_SAMPLE_SET == SAMPLE_SET_3X3 || CONFIG_SAMPLE_SET == SAMPLE_SET_5X5_WAVELET
-			if (x == 0 && y == 0) continue;
+			if (x == 0 && y == 0 && !CONFIG_UPSCALE) continue;
 			PixelOffset = float2(x, y);
 			
 		#elif CONFIG_SAMPLE_SET == SAMPLE_SET_STACKOWIAK_4_SETS
 			#if CONFIG_SAMPLE_SUBSET == 0
-				PixelOffset = kStackowiakSampleSet0[kStackowiakSampleSetCount * SampleId + QuadSampleId];
+				PixelOffset = kStackowiakSampleSet0[kStackowiakSampleSetCount * SampleId + QuadSampleSubsetId];
 			#elif CONFIG_SAMPLE_SUBSET == 1
-				PixelOffset = kStackowiakSampleSet1[kStackowiakSampleSetCount * SampleId + QuadSampleId];
+				PixelOffset = kStackowiakSampleSet1[kStackowiakSampleSetCount * SampleId + QuadSampleSubsetId];
 			#else
 				#error Unknown sample subset.
 			#endif
@@ -366,7 +427,7 @@ void MainCS(
 		#endif
 
 		// Compute sample coordinates
-		float2 NeighborBufferUV = KernelCenterBufferUV + PixelOffset * View.BufferSizeAndInvSize.zw;
+		float2 NeighborBufferUV = KernelCenterBufferUV + (PixelOffset * KernelSpreadFactor) * View.BufferSizeAndInvSize.zw;
 		if (1)
 		{
 			if (1 && any(NeighborBufferUV != clamp(NeighborBufferUV, View.BufferBilinearUVMinMax.xy, View.BufferBilinearUVMinMax.zw)))
@@ -410,6 +471,15 @@ void MainCS(
 		// Compute the final weight of that sample.
 		float FinalWeight = BilateralWeight * MISWeight * KernelWeight;// * (NeighborSample > 0.0 ? 1 : 0);
 		
+		#if CONFIG_UPSCALE && CONFIG_SAMPLE_SET == SAMPLE_SET_STACKOWIAK_4_SETS
+		{
+			if (SampleId == 0 && bTakeFirstRay)
+			{
+				FinalWeight = 1;
+			}
+		}
+		#endif
+
 		Accumulator = Add(Accumulator, Mul(NeighborSample, FinalWeight));
 	}
 
diff --git a/Engine/Shaders/Private/ScreenSpaceDenoise/SSDTemporalAccumulation.usf b/Engine/Shaders/Private/ScreenSpaceDenoise/SSDTemporalAccumulation.usf
--- a/Engine/Shaders/Private/ScreenSpaceDenoise/SSDTemporalAccumulation.usf
+++ b/Engine/Shaders/Private/ScreenSpaceDenoise/SSDTemporalAccumulation.usf
@@ -607,7 +607,7 @@ void TemporallyAccumulate(
 		float NeighborMaxPenumbra = 0.0;
 
 		const bool bUseCurrentSampleForCenter = bPostRejectionBlending;
-		if (bUseCurrentSampleForCenter)
+		if (bUseCurrentSampleForCenter && CurrentFrame.AccumulatedSampleCount > 0)
 		{
 			NeighborMinPenumbra = GetSamplePenumbraSafe(CurrentFrame);
 			NeighborMaxPenumbra = GetSamplePenumbraSafe(CurrentFrame);
@@ -640,7 +640,7 @@ void TemporallyAccumulate(
 				}
 			}
 
-			FSSFSignalSample NormalizedNeighborSample = CreateSignalSample();
+			FSSFSignalSample NeighborSample = CreateSignalSample();
 			float NeighborPenumbra = -1;
 			float NeighborUncertainty = 1.0;
 
@@ -689,7 +689,7 @@ void TemporallyAccumulate(
 						float2 SampleScreenPosition = ViewportUVToScreenPos(SampleViewportUV);
 
 						// Sample the neighor for history rejection.
-						FSSFSignalSample Neighbor = SampleInputSignal(
+						NeighborSample = SampleInputSignal(
 							bUsesHistoryRejectionSignal ? HistoryRejectionSignal0 : SignalInput0,
 							bUsesHistoryRejectionSignal ? HistoryRejectionSignal1 : SignalInput1,
 							GlobalPointClampedSampler,
@@ -705,19 +705,12 @@ void TemporallyAccumulate(
 							CurrentSceneMetadata,
 							NeighborSceneMetadata);
 
-						NormalizedNeighborSample = NormalizeToOneSample(Neighbor);
-
-						if (BilateralWeight > 0 && Neighbor.AccumulatedSampleCount > 0)
+						if (BilateralWeight > 0 && NeighborSample.AccumulatedSampleCount > 0)
 						{
-							NeighborPenumbra = GetSamplePenumbra(Neighbor);
+							NeighborPenumbra = GetSamplePenumbra(NeighborSample);
 							NeighborUncertainty = bUsesHistoryRejectionSignal
 								? 0.0
-								: Neighbor.UncertainRayCount / (Neighbor.AccumulatedSampleCount + Neighbor.UncertainRayCount);
-
-							if (CONFIG_SIGNAL_PROCESSING == SIGNAL_PROCESSING_REFLECTIONS)
-							{
-								NeighborUncertainty = 1 - NormalizedNeighborSample.SceneColor.a;
-							}
+								: NeighborSample.UncertainRayCount / (NeighborSample.AccumulatedSampleCount + NeighborSample.UncertainRayCount);
 						}
 					} // if (bUseSample)
 				}
@@ -729,7 +722,7 @@ void TemporallyAccumulate(
 				float2 SampleScreenPosition = ViewportUVToScreenPos(SampleViewportUV);
 
 				// Sample the neighor for history rejection.
-				FSSFSignalSample Neighbor = SampleInputSignal(
+				NeighborSample = SampleInputSignal(
 					bUsesHistoryRejectionSignal ? HistoryRejectionSignal0 : SignalInput0,
 					bUsesHistoryRejectionSignal ? HistoryRejectionSignal1 : SignalInput1,
 					GlobalPointClampedSampler,
@@ -744,32 +737,37 @@ void TemporallyAccumulate(
 					CurrentSceneMetadata,
 					NeighborSceneMetadata);
 
-				if (BilateralWeight > 0 && Neighbor.AccumulatedSampleCount > 0)
+				if (BilateralWeight > 0 && NeighborSample.AccumulatedSampleCount > 0)
 				{
-					NeighborPenumbra = GetSamplePenumbra(Neighbor);
+					NeighborPenumbra = GetSamplePenumbra(NeighborSample);
 					NeighborUncertainty = bUsesHistoryRejectionSignal
 						? 0.0
-						: Neighbor.UncertainRayCount / (Neighbor.AccumulatedSampleCount + Neighbor.UncertainRayCount);
+						: NeighborSample.UncertainRayCount / (NeighborSample.AccumulatedSampleCount + NeighborSample.UncertainRayCount);
 				}
-
-				NormalizedNeighborSample = NormalizeToOneSample(Neighbor);
 			}
 			#endif // !CONFIG_INPUT_TILE_CLASSIFICATION
-
+			
+			FSSFSignalSample NormalizedNeighborSample = NormalizeToOneSample(NeighborSample);
 			NormalizedNeighborSample = TransformSignalForPreRejection(NormalizedNeighborSample);
+			
+			#if CONFIG_SIGNAL_PROCESSING == SIGNAL_PROCESSING_PENUMBRA
+				bool bUseNeighborSample = NeighborUncertainty < 1.0;
+			#else
+				bool bUseNeighborSample = NeighborSample.AccumulatedSampleCount > 0.0;
+			#endif
 
-			if (CONFIG_CERTAINTY_BASED_REJECTION && NeighborUncertainty < 1.0)
+			if (CONFIG_CERTAINTY_BASED_REJECTION && bUseNeighborSample)
 			{
 				NeighborMinPenumbra = min(NeighborMinPenumbra, NeighborPenumbra + NeighborUncertainty);
 				NeighborMaxPenumbra = max(NeighborMaxPenumbra, NeighborPenumbra - NeighborUncertainty);
 			}
-			else if (NeighborUncertainty < 1.0)
+			else if (bUseNeighborSample)
 			{
 				NeighborMinPenumbra = min(NeighborMinPenumbra, NeighborPenumbra);
 				NeighborMaxPenumbra = max(NeighborMaxPenumbra, NeighborPenumbra);
 			}
 
-			if (NeighborUncertainty < 1.0)
+			if (bUseNeighborSample)
 			{
 				NeighborMin = Min(NeighborMin, NormalizedNeighborSample);
 				NeighborMax = Max(NeighborMax, NormalizedNeighborSample);
@@ -780,7 +778,7 @@ void TemporallyAccumulate(
 				NeighborMoment1 = Add(NeighborMoment1, Mul(NormalizedNeighborSample, Weight));
 				NeighborMoment2 = Add(NeighborMoment2, Mul(Power(NormalizedNeighborSample, 2.0), Weight));
 			}
-		}
+		} // for (neighbors)
 
 		NeighborMoment1 = Mul(NeighborMoment1, rcp(NeighborMomentWeight));
 		NeighborMoment2 = Mul(NeighborMoment2, rcp(NeighborMomentWeight));
@@ -825,44 +823,51 @@ void TemporallyAccumulate(
 				NeighborMax = Max(NeighborMax, QuantizationErrorMax);
 			}
 
+			// Transform the history so that it in correct component spaces, and normalized as the clamping box.
 			FSSFSignalSample NormalizedHistorySample = NormalizeToOneSample(HistorySample);
 			FSSFSignalSample TransformedHistorySample = TransformSignalForPreRejection(NormalizedHistorySample);
 
+			// Clamp the history.
 			FSSFSignalSample ClampedTransformedHistorySample = Clamp(TransformedHistorySample, NeighborMin, NeighborMax);
 
-			FSSFSignalSample RejectedDiff = Abs(Subtract(ClampedTransformedHistorySample, TransformedHistorySample));
+			// Transform the history back to linear component spaces.
+			FSSFSignalSample ClampedHistorySample = TransformSignalForPostRejection(ClampedTransformedHistorySample);
 
-			float RejectionFactor = 0.0;
-			#if CONFIG_SIGNAL_PROCESSING == SIGNAL_PROCESSING_REFLECTIONS && CONFIG_HISTORY_REJECTION_COLOR_SPACE == 1
+			// Reweight the anti ghosted history.
 			{
-				RejectionFactor = abs(
-					Luma_To_LumaLog(ClampedTransformedHistorySample.SceneColor.x) -
-					Luma_To_LumaLog(TransformedHistorySample.SceneColor.x));
+				FSSFSignalSample RejectedDiff = Abs(Subtract(ClampedTransformedHistorySample, TransformedHistorySample));
+
+				// Compute how much this history get changed.
+				float RejectionFactor = 0.0;
+				#if CONFIG_SIGNAL_PROCESSING == SIGNAL_PROCESSING_REFLECTIONS && CONFIG_HISTORY_REJECTION_COLOR_SPACE == 1
+				{
+					RejectionFactor = abs(
+						Luma_To_LumaLog(ClampedTransformedHistorySample.SceneColor.x) -
+						Luma_To_LumaLog(TransformedHistorySample.SceneColor.x));
 				
-				RejectionFactor = max(RejectionFactor, 1 * max(RejectedDiff.SceneColor.y, RejectedDiff.SceneColor.z));
-				RejectionFactor = max(RejectionFactor, 1 * RejectedDiff.SceneColor.a);
-			}
-			#else
-				#error Unsupported signal rejection.
-			#endif
+					RejectionFactor = max(RejectionFactor, 1 * max(RejectedDiff.SceneColor.y, RejectedDiff.SceneColor.z));
+					RejectionFactor = max(RejectionFactor, 1 * RejectedDiff.SceneColor.a);
+				}
+				#else
+					#error Unsupported signal rejection.
+				#endif
 			
-			//DebugOutput[DispatchThreadId] = float4(
-			//	RejectionFactor,
-			//	0,
-			//	0,
-			//	0);
+				// Compute an initial history weight as if we had removed sample.
+				float FinalHistoryWeight = HistorySample.AccumulatedSampleCount * saturate(1 - RejectionFactor);
 
-			float FinalHistoryWeight = HistorySample.AccumulatedSampleCount * saturate(1 - RejectionFactor);
+				// When doing the accumulation before rejection, need to make sure the input weight get through.
+				if (!bPostRejectionBlending)
+				{
+					FinalHistoryWeight = max(FinalHistoryWeight, CurrentFrame.AccumulatedSampleCount);
+				}
 
-			if (!bPostRejectionBlending)
-			{
-				FinalHistoryWeight = max(FinalHistoryWeight, CurrentFrame.AccumulatedSampleCount);
+				// When doing upscaling, may have an invalid input sample.
+				FinalHistoryWeight = max(FinalHistoryWeight, NeighborMoment1.AccumulatedSampleCount * 0.1);
+			
+				// Apply the history weight.
+				HistorySample = Mul(ClampedHistorySample, FinalHistoryWeight);
+				HistorySample.AccumulatedSampleCount = FinalHistoryWeight;
 			}
-
-			FSSFSignalSample ClampedHistorySample = TransformSignalForPostRejection(ClampedTransformedHistorySample);
-			FSSFSignalSample FinalHistorySample = Mul(ClampedHistorySample, FinalHistoryWeight);
-
-			HistorySample = FinalHistorySample;
 		}
 		#else // CONFIG_HISTORY_REJECTION == 1
 		{
diff --git a/Engine/Source/Runtime/Renderer/Private/DeferredShadingRenderer.h b/Engine/Source/Runtime/Renderer/Private/DeferredShadingRenderer.h
--- a/Engine/Source/Runtime/Renderer/Private/DeferredShadingRenderer.h
+++ b/Engine/Source/Runtime/Renderer/Private/DeferredShadingRenderer.h
@@ -436,7 +436,8 @@ private:
 		const FViewInfo& View,
 		FRDGTextureRef* OutColorTexture,
 		FRDGTextureRef* OutRayHitDistanceTexture,
-		int32 SamplePerPixel);
+		int32 SamplePerPixel,
+		float ResolutionFraction);
 
 #if RHI_RAYTRACING
 	bool ShouldRenderRayTracingAmbientOcclusion() const;
diff --git a/Engine/Source/Runtime/Renderer/Private/RayTracing/RayTracingReflections.cpp b/Engine/Source/Runtime/Renderer/Private/RayTracing/RayTracingReflections.cpp
--- a/Engine/Source/Runtime/Renderer/Private/RayTracing/RayTracingReflections.cpp
+++ b/Engine/Source/Runtime/Renderer/Private/RayTracing/RayTracingReflections.cpp
@@ -125,6 +125,7 @@ class FRayTracingReflectionsRG : public FGlobalShader
 		SHADER_PARAMETER(int32, ShouldDoDirectLighting)
 		SHADER_PARAMETER(int32, ShouldDoReflectedShadows)
 		SHADER_PARAMETER(int32, ShouldDoEmissiveAndIndirectLighting)
+		SHADER_PARAMETER(int32, UpscaleFactor)
 		SHADER_PARAMETER(float, ReflectionMaxRayDistance)
 		SHADER_PARAMETER(float, ReflectionMaxRoughness)
 
@@ -188,15 +189,22 @@ void FDeferredShadingSceneRenderer::RayTraceReflections(
 	const FViewInfo& View,
 	FRDGTextureRef* OutColorTexture,
 	FRDGTextureRef* OutRayHitDistanceTexture,
-	int32 SamplePerPixel)
+	int32 SamplePerPixel, 
+	float ResolutionFraction)
 #if RHI_RAYTRACING
 {
 	FSceneRenderTargets& SceneContext = FSceneRenderTargets::Get(GraphBuilder.RHICmdList);
 
+	int32 UpscaleFactor = int32(1.0f / ResolutionFraction);
+	ensure(ResolutionFraction == 1.0 / UpscaleFactor);
+	ensureMsgf(FComputeShaderUtils::kGolden2DGroupSize % UpscaleFactor == 0, TEXT("Reflection ray tracing will have uv misalignement."));
+	FIntPoint RayTracingResolution = FIntPoint::DivideAndRoundUp(View.ViewRect.Size(), UpscaleFactor);
+
 	{		
 		FPooledRenderTargetDesc Desc = SceneContext.GetSceneColor()->GetDesc();
 		Desc.Format = PF_FloatRGBA;
 		Desc.Flags &= ~(TexCreate_FastVRAM | TexCreate_Transient);
+		Desc.Extent /= UpscaleFactor;
 
 		*OutColorTexture = GraphBuilder.CreateTexture(Desc, TEXT("RayTracingReflections"));
 		
@@ -210,6 +218,7 @@ void FDeferredShadingSceneRenderer::RayTraceReflections(
 	PassParameters->ShouldDoDirectLighting = GRayTracingReflectionsDirectLighting;
 	PassParameters->ShouldDoReflectedShadows = GRayTracingReflectionsShadows;
 	PassParameters->ShouldDoEmissiveAndIndirectLighting = GRayTracingReflectionsEmissiveAndIndirectLighting;
+	PassParameters->UpscaleFactor = UpscaleFactor;
 	PassParameters->ReflectionMaxRayDistance = GRayTracingReflectionsMaxRayDistance;
 	PassParameters->ReflectionMaxRoughness = FMath::Clamp(View.FinalPostProcessSettings.ScreenSpaceReflectionMaxRoughness, 0.01f, 1.0f);
 	PassParameters->LTCMatTexture = GSystemTextures.LTCMat->GetRenderTargetItem().ShaderResourceTexture;
@@ -236,15 +245,15 @@ void FDeferredShadingSceneRenderer::RayTraceReflections(
 	}
 	PassParameters->ColorOutput = GraphBuilder.CreateUAV(*OutColorTexture);
 	PassParameters->RayHitDistanceOutput = GraphBuilder.CreateUAV(*OutRayHitDistanceTexture);
-
+	
 	auto RayGenShader = View.ShaderMap->GetShader<FRayTracingReflectionsRG>();
 	ClearUnusedGraphResources(RayGenShader, PassParameters);
 
 	GraphBuilder.AddPass(
-		RDG_EVENT_NAME("ReflectionRayTracing %dx%d", View.ViewRect.Width(), View.ViewRect.Height()),
+		RDG_EVENT_NAME("ReflectionRayTracing %dx%d", RayTracingResolution.X, RayTracingResolution.X),
 		PassParameters,
 		ERenderGraphPassFlags::Compute,
-		[PassParameters, this, &View, RayGenShader](FRHICommandList& RHICmdList)
+		[PassParameters, this, &View, RayGenShader, RayTracingResolution](FRHICommandList& RHICmdList)
 	{
 		auto ClosestHitShader = View.ShaderMap->GetShader<FRayTracingReflectionsCHS>();
 		auto MissShader = View.ShaderMap->GetShader<FRayTracingReflectionsMS>();
@@ -259,7 +268,7 @@ void FDeferredShadingSceneRenderer::RayTraceReflections(
 		SetShaderParameters(GlobalResources, RayGenShader, *PassParameters);
 
 		FRayTracingSceneRHIParamRef RayTracingSceneRHI = View.PerViewRayTracingScene.RayTracingSceneRHI;
-		RHICmdList.RayTraceDispatch(Pipeline, RayTracingSceneRHI, GlobalResources, View.ViewRect.Width(), View.ViewRect.Height());
+		RHICmdList.RayTraceDispatch(Pipeline, RayTracingSceneRHI, GlobalResources, RayTracingResolution.X, RayTracingResolution.Y);
 	});
 }
 #else // !RHI_RAYTRACING
diff --git a/Engine/Source/Runtime/Renderer/Private/ReflectionEnvironment.cpp b/Engine/Source/Runtime/Renderer/Private/ReflectionEnvironment.cpp
--- a/Engine/Source/Runtime/Renderer/Private/ReflectionEnvironment.cpp
+++ b/Engine/Source/Runtime/Renderer/Private/ReflectionEnvironment.cpp
@@ -102,22 +102,28 @@ static TAutoConsoleVariable<float> CVarSkySpecularOcclusionStrength(
 	ECVF_RenderThreadSafe);
 
 static int32 GRayTracingReflections = 1;
-static FAutoConsoleVariableRef CVarRayTracingReflectionss(
-	TEXT("r.RayTracing.Reflections"),
+static FAutoConsoleVariableRef CVarReflectionsMethod(
+	TEXT("r.Reflections.Method"),
 	GRayTracingReflections,
 	TEXT("0: use traditional rasterized SSR\n")
 	TEXT("1: use ray traced reflections (default when ray tracing is enabled)\n")
 );
 
+static TAutoConsoleVariable<float> CVarReflectionScreenPercentage(
+	TEXT("r.Reflections.RayTracing.ScreenPercentage"),
+	100.0f,
+	TEXT("Try to guess"),
+	ECVF_RenderThreadSafe);
+
 static int32 GRayTracingReflectionsSamplesPerPixel = 1;
 static FAutoConsoleVariableRef CVarRayTracingReflectionsSamplesPerPixel(
-	TEXT("r.RayTracing.Reflections.SamplesPerPixel"),
+	TEXT("r.Reflections.RayTracing.SamplesPerPixel"),
 	GRayTracingReflectionsSamplesPerPixel,
 	TEXT("Sets the samples-per-pixel for reflections (default = 1)")
 );
 
 static TAutoConsoleVariable<int32> CVarUseReflectionDenoiser(
-	TEXT("r.Reflection.Denoiser"),
+	TEXT("r.Reflections.Denoiser"),
 	2,
 	TEXT("Choose the denoising algorithm.\n")
 	TEXT(" 0: Disabled;\n")
@@ -745,13 +751,27 @@ void FDeferredShadingSceneRenderer::RenderDeferredReflectionsAndSkyLighting(FRHI
 			SetupSceneViewFamilyBlackboard(GraphBuilder, &SceneBlackboard);
 
 
+			IScreenSpaceDenoiser::FReflectionsRayTracingConfig RayTracingConfig;
+			RayTracingConfig.ResolutionFraction = FMath::Clamp(CVarReflectionScreenPercentage.GetValueOnRenderThread() / 100.0f, 0.25f, 1.0f);
+
+			int32 DenoiserMode = CVarUseReflectionDenoiser.GetValueOnRenderThread();
+			const bool bDenoise = DenoiserMode != 0 && GRayTracingReflectionsSamplesPerPixel == 1;
+
+			if (!bDenoise)
+			{
+				RayTracingConfig.ResolutionFraction = 1.0f;
+			}
+
 			// Ray trace the reflection.
-			IScreenSpaceDenoiser::FReflectionInputs DenoiserInputs;
-			RayTraceReflections(GraphBuilder, View, &DenoiserInputs.Color, &DenoiserInputs.RayHitDistance, GRayTracingReflectionsSamplesPerPixel);
+			IScreenSpaceDenoiser::FReflectionsInputs DenoiserInputs;
+			RayTraceReflections(
+				GraphBuilder,
+				View, &DenoiserInputs.Color, &DenoiserInputs.RayHitDistance,
+				GRayTracingReflectionsSamplesPerPixel, RayTracingConfig.ResolutionFraction);
 
 
 			// Denoise the reflections.
-			if (int32 DenoiserMode = CVarUseReflectionDenoiser.GetValueOnRenderThread() && GRayTracingReflectionsSamplesPerPixel == 1)
+			if (bDenoise)
 			{
 				const IScreenSpaceDenoiser* DefaultDenoiser = IScreenSpaceDenoiser::GetDefaultDenoiser();
 				const IScreenSpaceDenoiser* DenoiserToUse = DenoiserMode == 1 ? DefaultDenoiser : GScreenSpaceDenoiser;
@@ -762,11 +782,12 @@ void FDeferredShadingSceneRenderer::RenderDeferredReflectionsAndSkyLighting(FRHI
 					DenoiserToUse->GetDebugName(),
 					View.ViewRect.Width(), View.ViewRect.Height());
 
-				IScreenSpaceDenoiser::FReflectionOutputs DenoiserOutputs = DenoiserToUse->DenoiseReflections(
+				IScreenSpaceDenoiser::FReflectionsOutputs DenoiserOutputs = DenoiserToUse->DenoiseReflections(
 					GraphBuilder,
 					View,
 					SceneBlackboard,
-					DenoiserInputs);
+					DenoiserInputs,
+					RayTracingConfig);
 
 				GraphBuilder.QueueTextureExtraction(DenoiserOutputs.Color, &ReflectionsColor);
 			}
diff --git a/Engine/Source/Runtime/Renderer/Private/ScreenSpaceDenoise.cpp b/Engine/Source/Runtime/Renderer/Private/ScreenSpaceDenoise.cpp
--- a/Engine/Source/Runtime/Renderer/Private/ScreenSpaceDenoise.cpp
+++ b/Engine/Source/Runtime/Renderer/Private/ScreenSpaceDenoise.cpp
@@ -34,17 +34,17 @@ static TAutoConsoleVariable<int32> CVarDoDebugPass(
 	ECVF_RenderThreadSafe);
 
 static TAutoConsoleVariable<int32> CVarReflectionReconstructionSampleCount(
-	TEXT("r.Reflection.Denoise.ReconstructionSamples"), 16,
+	TEXT("r.Reflections.Denoiser.ReconstructionSamples"), 16,
 	TEXT("Maximum number of samples for the reconstruction pass (default = 16)."),
 	ECVF_RenderThreadSafe);
 
 static TAutoConsoleVariable<int32> CVarReflectionTemporalAccumulation(
-	TEXT("r.Reflection.Denoise.TemporalAccumulation"), 1,
+	TEXT("r.Reflections.Denoiser.TemporalAccumulation"), 1,
 	TEXT(""),
 	ECVF_RenderThreadSafe);
 
 static TAutoConsoleVariable<int32> CVarReflectionHistoryConvolutionSampleCount(
-	TEXT("r.Reflection.Denoise.HistoryConvolutionSampleCount"), 1,
+	TEXT("r.Reflections.Denoiser.HistoryConvolutionSampleCount"), 1,
 	TEXT("Number of samples to use for history post filter (default = 1)."),
 	ECVF_RenderThreadSafe);
 
@@ -202,8 +202,9 @@ class FSSDSpatialAccumulationCS : public FScreenSpaceDenoisingShader
 	};
 
 	class FStageDim : SHADER_PERMUTATION_ENUM_CLASS("DIM_STAGE", EStage);
+	class FUpscaleDim : SHADER_PERMUTATION_BOOL("DIM_UPSCALE");
 
-	using FPermutationDomain = TShaderPermutationDomain<FSignalProcessingDim, FStageDim>;
+	using FPermutationDomain = TShaderPermutationDomain<FSignalProcessingDim, FStageDim, FUpscaleDim>;
 	
 	static bool ShouldCompilePermutation(const FGlobalShaderPermutationParameters& Parameters)
 	{
@@ -216,11 +217,19 @@ class FSSDSpatialAccumulationCS : public FScreenSpaceDenoisingShader
 			return false;
 		}
 
+		// Only reconstruction have upscale capability for now.
+		if (PermutationVector.Get<FUpscaleDim>() && 
+			PermutationVector.Get<FStageDim>() != EStage::ReConstruction)
+		{
+			return false;
+		}
+
 		return FScreenSpaceDenoisingShader::ShouldCompilePermutation(Parameters);
 	}
 
 	BEGIN_SHADER_PARAMETER_STRUCT(FParameters, )
 		SHADER_PARAMETER(uint32, MaxSampleCount)
+		SHADER_PARAMETER(int32, UpscaleFactor)
 
 		SHADER_PARAMETER_STRUCT_INCLUDE(FSSDCommonParameters, CommonParameters)
 
@@ -630,17 +639,23 @@ static void DenoiseShadowPenumbra(
 } // DenoiseShadowPenumbra()
 
 
-IScreenSpaceDenoiser::FReflectionOutputs DenoiseReflections(
+IScreenSpaceDenoiser::FReflectionsOutputs DenoiseReflections(
 	FRDGBuilder& GraphBuilder,
 	const FViewInfo& View,
 	const FSceneViewFamilyBlackboard& SceneBlackboard,
-	const IScreenSpaceDenoiser::FReflectionInputs& ReflectionInputs)
+	const IScreenSpaceDenoiser::FReflectionsInputs& ReflectionInputs,
+	const IScreenSpaceDenoiser::FReflectionsRayTracingConfig RayTracingConfig)
 {
+	ensure(RayTracingConfig.ResolutionFraction == 1.0f || RayTracingConfig.ResolutionFraction == 0.5f);
+
 	const FIntPoint DenoiseResolution = View.ViewRect.Size();
+	
+	FSceneRenderTargets& SceneContext = FSceneRenderTargets::Get(GraphBuilder.RHICmdList);
 
 	// Descriptor to allocate internal denoising buffer.
-	FRDGTextureDesc SignalProcessingDesc = ReflectionInputs.Color->Desc;
+	FRDGTextureDesc SignalProcessingDesc = SceneContext.GetSceneColor()->GetDesc();
 	SignalProcessingDesc.Format = PF_FloatRGBA;
+	SignalProcessingDesc.Flags &= ~(TexCreate_FastVRAM | TexCreate_Transient);
 	SignalProcessingDesc.NumMips = 1;
 
 	// Setup common shader parameters.
@@ -666,6 +681,7 @@ IScreenSpaceDenoiser::FReflectionOutputs DenoiseReflections(
 
 		FSSDSpatialAccumulationCS::FParameters* PassParameters = GraphBuilder.AllocParameters<FSSDSpatialAccumulationCS::FParameters>();
 		PassParameters->MaxSampleCount = FMath::Clamp(CVarReflectionReconstructionSampleCount.GetValueOnRenderThread(), 1, kStackowiakMaxSampleCountPerSet);
+		PassParameters->UpscaleFactor = int32(1.0f / RayTracingConfig.ResolutionFraction);
 		PassParameters->CommonParameters = CommonParameters;
 		PassParameters->SignalInput0 = ReflectionInputs.Color;
 		PassParameters->SignalInput1 = ReflectionInputs.RayHitDistance;
@@ -677,11 +693,13 @@ IScreenSpaceDenoiser::FReflectionOutputs DenoiseReflections(
 		FSSDSpatialAccumulationCS::FPermutationDomain PermutationVector;
 		PermutationVector.Set<FSignalProcessingDim>(ESignalProcessing::Reflections);
 		PermutationVector.Set<FSSDSpatialAccumulationCS::FStageDim>(FSSDSpatialAccumulationCS::EStage::ReConstruction);
+		PermutationVector.Set<FSSDSpatialAccumulationCS::FUpscaleDim>(PassParameters->UpscaleFactor != 1);
 
 		TShaderMapRef<FSSDSpatialAccumulationCS> ComputeShader(View.ShaderMap, PermutationVector);
 		FComputeShaderUtils::AddPass(
 			GraphBuilder,
-			RDG_EVENT_NAME("SSD SpatialAccumulation(Reconstruction MaxSamples=%i)", PassParameters->MaxSampleCount),
+			RDG_EVENT_NAME("SSD SpatialAccumulation(Reconstruction MaxSamples=%i Upscale=%i)",
+				PassParameters->MaxSampleCount, int32(PermutationVector.Get<FSSDSpatialAccumulationCS::FUpscaleDim>())),
 			*ComputeShader,
 			PassParameters,
 			FComputeShaderUtils::GetGroupCount(DenoiseResolution, FComputeShaderUtils::kGolden2DGroupSize));
@@ -791,7 +809,7 @@ IScreenSpaceDenoiser::FReflectionOutputs DenoiseReflections(
 		GraphBuilder.RemoveUnusedTextureWarning(SignalHistory1);
 	}
 
-	IScreenSpaceDenoiser::FReflectionOutputs Outputs;
+	IScreenSpaceDenoiser::FReflectionsOutputs Outputs;
 	Outputs.Color = SignalHistory0; // TODO.
 	return Outputs;
 } // DenoiseReflections()
@@ -858,13 +876,14 @@ public:
 		return Outputs;
 	}
 
-	FReflectionOutputs DenoiseReflections(
+	FReflectionsOutputs DenoiseReflections(
 		FRDGBuilder& GraphBuilder,
 		const FViewInfo& View,
 		const FSceneViewFamilyBlackboard& SceneBlackboard,
-		const FReflectionInputs& ReflectionInputs) const override
+		const FReflectionsInputs& ReflectionInputs,
+		const FReflectionsRayTracingConfig RayTracingConfig) const override
 	{
-		return ::DenoiseReflections(GraphBuilder, View, SceneBlackboard, ReflectionInputs);
+		return ::DenoiseReflections(GraphBuilder, View, SceneBlackboard, ReflectionInputs, RayTracingConfig);
 	}
 
 }; // class FDefaultScreenSpaceDenoiser
diff --git a/Engine/Source/Runtime/Renderer/Private/ScreenSpaceDenoise.h b/Engine/Source/Runtime/Renderer/Private/ScreenSpaceDenoise.h
--- a/Engine/Source/Runtime/Renderer/Private/ScreenSpaceDenoise.h
+++ b/Engine/Source/Runtime/Renderer/Private/ScreenSpaceDenoise.h
@@ -40,6 +40,13 @@ public:
 		SHADER_PARAMETER_RDG_TEXTURE(Texture2D, SpecularPenumbra)
 	END_SHADER_PARAMETER_STRUCT()
 
+	/** The configuration of the reflection ray tracing. */
+	struct FReflectionsRayTracingConfig
+	{
+		// Resolution fraction the ray tracing is being traced at.
+		float ResolutionFraction;
+	};
+
 
 	virtual ~IScreenSpaceDenoiser() {};
 
@@ -60,22 +67,23 @@ public:
 		const FShadowPenumbraInputs& ShadowInputs) const = 0;
 	
 	/** All the inputs of the shadow denoiser. */
-	BEGIN_SHADER_PARAMETER_STRUCT(FReflectionInputs, )
+	BEGIN_SHADER_PARAMETER_STRUCT(FReflectionsInputs, )
 		SHADER_PARAMETER_RDG_TEXTURE(Texture2D, Color)
 		SHADER_PARAMETER_RDG_TEXTURE(Texture2D, RayHitDistance)
 	END_SHADER_PARAMETER_STRUCT()
 
 	/** All the outputs of the shadow denoiser may generate. */
-	BEGIN_SHADER_PARAMETER_STRUCT(FReflectionOutputs, )
+	BEGIN_SHADER_PARAMETER_STRUCT(FReflectionsOutputs, )
 		SHADER_PARAMETER_RDG_TEXTURE(Texture2D, Color)
 	END_SHADER_PARAMETER_STRUCT()
 		
 	/** Entry point to denoise reflections. */
-	virtual FReflectionOutputs DenoiseReflections(
+	virtual FReflectionsOutputs DenoiseReflections(
 		FRDGBuilder& GraphBuilder,
 		const FViewInfo& View,
 		const FSceneViewFamilyBlackboard& SceneBlackboard,
-		const FReflectionInputs& ReflectionInputs) const = 0;
+		const FReflectionsInputs& ReflectionInputs,
+		const FReflectionsRayTracingConfig RayTracingConfig) const = 0;
 
 	/** Returns the interface of the default denoiser of the renderer. */
 	static const IScreenSpaceDenoiser* GetDefaultDenoiser();
