Fixes an artifact in the denoiser when the reflection ray generation early return based on high roughness.

diff --git a/Engine/Shaders/Private/RayTracing/RayTracingReflections.usf b/Engine/Shaders/Private/RayTracing/RayTracingReflections.usf
--- a/Engine/Shaders/Private/RayTracing/RayTracingReflections.usf
+++ b/Engine/Shaders/Private/RayTracing/RayTracingReflections.usf
@@ -422,14 +422,18 @@ void RayTracingReflectionsRGS()
 			}
 		}
 	}
+	else
+	{
+		// To tell the denoiser the black input is valid even if no ray has technically been shot. Uses a high ray hit distance
+		// so that the sample may be reused by another neighbor pixel in the denoise as if RoughnessFade in ]0;1[.
+		HitDistance = 1.0e20;
+	}
 
 	if (ReflectedWeight > 0)
 	{
-		ReflectedColor *= rcp(ReflectedWeight);
+		ReflectedColor *= RoughnessFade * rcp(ReflectedWeight);
 	}
 
-	ReflectedColor *= RoughnessFade;
-
 	ColorOutput[DispatchThreadId] = float4(ReflectedColor);
 	RayHitDistanceOutput[DispatchThreadId] = HitDistance;
 }
