Changes wrongly named ratio estimator in code

diff --git a/Engine/Shaders/Private/ScreenSpaceDenoise/SSDSignalFramework.ush b/Engine/Shaders/Private/ScreenSpaceDenoise/SSDSignalFramework.ush
--- a/Engine/Shaders/Private/ScreenSpaceDenoise/SSDSignalFramework.ush
+++ b/Engine/Shaders/Private/ScreenSpaceDenoise/SSDSignalFramework.ush
@@ -9,8 +9,8 @@
 	#define CONFIG_TOKOYASHI_BILATERAL 0
 #endif
 
-#ifndef CONFIG_SPECULAR_MIS
-	#define CONFIG_SPECULAR_MIS 0
+#ifndef CONFIG_SPECULAR_RATIO_ESTIMATOR
+	#define CONFIG_SPECULAR_RATIO_ESTIMATOR 0
 #endif
 
 
@@ -366,12 +366,12 @@ float GetSamplePenumbra(FSSFSignalSample Sample)
 }
 
 /** Get the weight for multi importance sampling. */
-float GetMultiImportanceSamplingWeight(
+float GetRatioEstimatorWeight(
 	FSSFSampleSceneInfos RefSceneMetadata,
 	FSSFSampleSceneInfos SampleSceneMetadata,
 	FSSFSignalSample Sample,
 	uint2 SamplePixelCoord)
-#if CONFIG_SPECULAR_MIS
+#if CONFIG_SPECULAR_RATIO_ESTIMATOR
 // [ Stackowiak 2018, "Stochastic all the things: Raytracing in hybrid real-time rendering" ]
 {
 	// If invalid ray or geometric miss, just.
@@ -404,8 +404,8 @@ float GetMultiImportanceSamplingWeight(
 
 	return D / OriginalRayPDF;
 }
-#else // !CONFIG_SPECULAR_MIS
+#else // !CONFIG_SPECULAR_RATIO_ESTIMATOR
 {
 	return 1.0;
 }
-#endif // CONFIG_SPECULAR_MIS
+#endif // CONFIG_SPECULAR_RATIO_ESTIMATOR
diff --git a/Engine/Shaders/Private/ScreenSpaceDenoise/SSDSpatialAccumulation.usf b/Engine/Shaders/Private/ScreenSpaceDenoise/SSDSpatialAccumulation.usf
--- a/Engine/Shaders/Private/ScreenSpaceDenoise/SSDSpatialAccumulation.usf
+++ b/Engine/Shaders/Private/ScreenSpaceDenoise/SSDSpatialAccumulation.usf
@@ -51,7 +51,7 @@
 		#define CONFIG_SAMPLE_COUNT_POLICY  SAMPLE_COUNT_POLICY_HIT_DISTANCE_BASED
 
 		// Reconstruction has 1 ray per input pixel, therefore can use MIS.
-		#define CONFIG_SPECULAR_MIS 1
+		#define CONFIG_SPECULAR_RATIO_ESTIMATOR 1
 
 		#define CONFIG_UPSCALE              DIM_UPSCALE
 		//#define CONFIG_SAMPLE_COUNT_POLICY  SAMPLE_COUNT_POLICY_DISABLED
@@ -454,7 +454,7 @@ void MainCS(
 
 		// Multiple importance sampling weights.
 		uint2 NeighborPixelCoord = floor(NeighborBufferUV * View.BufferSizeAndInvSize.xy);
-		float MISWeight = GetMultiImportanceSamplingWeight(RefSceneMetadata, NeighborSceneMetadata, NeighborSample, NeighborPixelCoord);
+		float RatioEstimatorWeight = GetRatioEstimatorWeight(RefSceneMetadata, NeighborSceneMetadata, NeighborSample, NeighborPixelCoord);
 
 		// Compute the weight of the sample in the kernel.
 		#if CONFIG_SAMPLE_SET == SAMPLE_SET_5X5_WAVELET
@@ -469,7 +469,7 @@ void MainCS(
 		#endif
 
 		// Compute the final weight of that sample.
-		float FinalWeight = BilateralWeight * MISWeight * KernelWeight;// * (NeighborSample > 0.0 ? 1 : 0);
+		float FinalWeight = BilateralWeight * RatioEstimatorWeight * KernelWeight;// * (NeighborSample > 0.0 ? 1 : 0);
 		
 		#if CONFIG_UPSCALE && CONFIG_SAMPLE_SET == SAMPLE_SET_STACKOWIAK_4_SETS
 		{
