Cranks up the responsivness of sharp reflections.

diff --git a/Engine/Shaders/Private/ScreenSpaceDenoise/SSDTemporalAccumulation.usf b/Engine/Shaders/Private/ScreenSpaceDenoise/SSDTemporalAccumulation.usf
--- a/Engine/Shaders/Private/ScreenSpaceDenoise/SSDTemporalAccumulation.usf
+++ b/Engine/Shaders/Private/ScreenSpaceDenoise/SSDTemporalAccumulation.usf
@@ -249,7 +249,7 @@ void TemporallyAccumulate(
 	{
 		float SignalPixelBluringRadius = ComputeInfinityBluringRadius(CurrentSceneMetadata) * View.ViewSizeAndInvSize.x;
 
-		TargetedSampleCount = 2 + PI * SignalPixelBluringRadius * SignalPixelBluringRadius;
+		TargetedSampleCount = PI * SignalPixelBluringRadius * SignalPixelBluringRadius;
 
 		// TODO: try to make this better? * 0.125 to avoid banding caused by floating point precision.
 		TargetedSampleCount = clamp(TargetedSampleCount, 1, 0.125 * TARGETED_SAMPLE_COUNT);
@@ -257,7 +257,8 @@ void TemporallyAccumulate(
 		//DebugOutput[DispatchThreadId] = float4(
 		//	SignalPixelBluringRadius,
 		//	TargetedSampleCount,
-		//	0, 0);
+		//	max(CurrentFrame.AccumulatedSampleCount - TargetedSampleCount, 0),
+		//	0);
 	}
 	#endif
 
