Fixed missing ShaderVersion.ush include, which is required to ensure that shaders are rebuilt when version is bumped

diff --git a/Engine/Shaders/Private/PathTracing/PathTracingRayCounterComputeShader.usf b/Engine/Shaders/Private/PathTracing/PathTracingRayCounterComputeShader.usf
--- a/Engine/Shaders/Private/PathTracing/PathTracingRayCounterComputeShader.usf
+++ b/Engine/Shaders/Private/PathTracing/PathTracingRayCounterComputeShader.usf
@@ -4,6 +4,7 @@
 	RayCounterComputeShader.usf: Ray counter for the reference path tracing
 ===============================================================================================*/
 
+#include "../Common.ush"
 
 Texture2D<uint> RayCountPerPixel;
 uint2 ViewSize;
