BuildSolidAnglePdfCS was incorrectly omitting Common.ush

diff --git a/Engine/Shaders/Private/RayTracing/BuildSolidAnglePdfCS.usf b/Engine/Shaders/Private/RayTracing/BuildSolidAnglePdfCS.usf
--- a/Engine/Shaders/Private/RayTracing/BuildSolidAnglePdfCS.usf
+++ b/Engine/Shaders/Private/RayTracing/BuildSolidAnglePdfCS.usf
@@ -4,6 +4,7 @@
 BuildMipTreeCS.usf: Build MipTree PDF structure
 ===============================================================================================*/
 
+#include "../Common.ush"
 #include "SkyLightMipTreeCommon.ush"
 
 uint MipLevel;
