Adding material model to ray tracing SkyLight integrator

1) Supports DefaultLit, Cloth, and SubsurfaceProfile materials.
2) Removes SkyLight specular contribution from RenderDeferredReflectionsAndSkyLighting() when ray traced dynamic sky light is selected.
3) Fix un-normalized direction vector.
4) Uses additive blending to support subsurface profile.

diff --git a/Engine/Shaders/Private/RayTracing/CompositeSkyLightPS.usf b/Engine/Shaders/Private/RayTracing/CompositeSkyLightPS.usf
--- a/Engine/Shaders/Private/RayTracing/CompositeSkyLightPS.usf
+++ b/Engine/Shaders/Private/RayTracing/CompositeSkyLightPS.usf
@@ -15,5 +15,5 @@ void CompositeSkyLightPS(
 )
 {
 	float4 SkyLight = SkyLightTexture.Sample(SkyLightTextureSampler, UV);
-	OutColor = float4(SkyLight.rgb, 1.0);
+	OutColor = SkyLight;
 }
diff --git a/Engine/Shaders/Private/RayTracing/RayTracingSkyLightRGS.usf b/Engine/Shaders/Private/RayTracing/RayTracingSkyLightRGS.usf
--- a/Engine/Shaders/Private/RayTracing/RayTracingSkyLightRGS.usf
+++ b/Engine/Shaders/Private/RayTracing/RayTracingSkyLightRGS.usf
@@ -3,6 +3,7 @@
 #include "../Common.ush"
 #include "../MonteCarlo.ush"
 #include "../DeferredShadingCommon.ush"
+#include "../ShadingModels.ush"
 #include "SkyLightMipTreeCommon.ush"
 #include "RayTracingCommon.ush"
 #include "RayTracingDirectionalLight.ush"
@@ -325,7 +326,7 @@ bool GenerateSkyLightOcclusionRay(
 #endif
 
 	RayOrigin = WorldPosition + WorldNormal * 0.01;
-	RayDirection = GetTextureCubeVector(TexelCoord, TextureDimensions);
+	RayDirection = normalize(GetTextureCubeVector(TexelCoord, TextureDimensions));
 	RayTMin = 0.0;
 	RayTMax = 1.0e27;
 	RayPdf = FacePdf * MipPdf * AreaPdf;
@@ -385,10 +386,15 @@ void SkyLightRGS()
 	float2 InvBufferSize = View.BufferSizeAndInvSize.zw;
 	float2 UV = (float2(PixelCoord) + 0.5) * InvBufferSize;
 	FScreenSpaceData ScreenSpaceData = GetScreenSpaceData(UV);
+	FGBufferData GBuffer = ScreenSpaceData.GBuffer;
 	float Depth = ScreenSpaceData.GBuffer.Depth;
 	float3 WorldPosition = ReconstructWorldPositionFromDepth(UV, Depth);
+	float3 CameraOrigin = ReconstructWorldPositionFromDepth(UV, 0.0f);
+	float3 CameraDirection = normalize(WorldPosition - CameraOrigin);
 	float3 WorldNormal = ScreenSpaceData.GBuffer.WorldNormal;
 	float3 BaseColor = ScreenSpaceData.GBuffer.BaseColor;
+	float3 DiffuseColor = ScreenSpaceData.GBuffer.DiffuseColor;
+	float3 SpecularColor = ScreenSpaceData.GBuffer.SpecularColor;
 	float Roughness = ScreenSpaceData.GBuffer.Roughness;
 
 	float RayDistance = 0.0;
@@ -401,7 +407,8 @@ void SkyLightRGS()
 		SamplesPerPixel = 0.0;
 	}
 
-	float3 Irradiance = 0.0;
+	float3 ExitantRadiance = 0.0;
+	float3 DiffuseExitantRadiance = 0.0;
 	for (uint SampleIndex = 0; SampleIndex < SamplesPerPixel; ++SampleIndex)
 	{
 		RayDesc Ray;
@@ -425,8 +432,8 @@ void SkyLightRGS()
 		}
 		
 		// Surface-normal culling
-		float NdotL = dot(WorldNormal, Ray.Direction);
-		if (RayPdf < 0.0 || NdotL <= 0.0)
+		float NoL = dot(WorldNormal, Ray.Direction);
+		if (RayPdf <= 0.0 || NoL <= 0.0)
 		{
 			continue;
 		}
@@ -453,18 +460,161 @@ void SkyLightRGS()
 		}
 		else
 		{
-			float3 Radiance = EvalSkyLight(Ray.Direction);
-			// #dxr_todo: Apply actual material BRDF
-			float3 Brdf = BaseColor * Roughness / PI;
-			Irradiance += Radiance * Brdf * NdotL * MisWeight / RayPdf;
+			float3 IncomingRadiance = EvalSkyLight(Ray.Direction);
+
+			BxDFContext Context;
+			half3 N = WorldNormal;
+			half3 V = -CameraDirection;
+			half3 L = Ray.Direction;
+			Init(Context, N, V, L);
+
+			uint ShadingModelID = ScreenSpaceData.GBuffer.ShadingModelID;
+			float3 Diffuse = 0.0;
+			float3 Specular = 0.0;
+			float3 Transmission = 0.0;
+			switch (ShadingModelID)
+			{
+				case SHADINGMODELID_DEFAULT_LIT:
+				{
+					Diffuse = Diffuse_Lambert(DiffuseColor);
+					float a2 = Pow4(Roughness);
+					float D = D_GGX(a2, Context.NoH);
+					float Vis = Vis_SmithJointApprox(a2, Context.NoV, NoL);
+					float3 F = F_Schlick(SpecularColor, Context.VoH);
+					Specular = (D * Vis) * F;
+
+					break;
+				}
+				case SHADINGMODELID_CLOTH:
+				{
+					const float3 FuzzColor = saturate(GBuffer.CustomData.rgb);
+					const float  Cloth = saturate(GBuffer.CustomData.a);
+					float a2 = Pow4(Roughness);
+					float3 Specular0 = 0.0;
+					{
+						float D = D_GGX(a2, Context.NoH);
+						float Vis = Vis_SmithJointApprox(a2, Context.NoV, NoL);
+						float3 F = F_Schlick(SpecularColor, Context.VoH);
+						Specular0 = (D * Vis) * F;
+					}
+
+					float3 Specular1 = 0.0;
+					{
+						float D = D_InvGGX(a2, Context.NoH);
+						float Vis = Vis_Cloth(Context.NoV, NoL);
+						float3 F = F_Schlick(FuzzColor, Context.VoH);
+						Specular1 = (D * Vis) * F;
+					}
+
+					Diffuse = Diffuse_Lambert(DiffuseColor);
+					Specular = lerp(Specular0, Specular1, Cloth);
+					break;
+				}
+				case SHADINGMODELID_SUBSURFACE_PROFILE:
+				{
+					float AverageToRoughness0;
+					float AverageToRoughness1;
+					float LobeMix;
+					GetProfileDualSpecular(GBuffer, AverageToRoughness0, AverageToRoughness1, LobeMix);
+
+					float AverageRoughness = Roughness;
+					float Lobe0Roughness = max(saturate(AverageRoughness * AverageToRoughness0), 0.02f);
+					float Lobe1Roughness = saturate(AverageRoughness * AverageToRoughness1);
+
+					float Opacity = GBuffer.CustomData.a;
+					Lobe0Roughness = lerp(1.0f, Lobe0Roughness, saturate(Opacity * 10.0f));
+					Lobe1Roughness = lerp(1.0f, Lobe1Roughness, saturate(Opacity * 10.0f));
+					Diffuse = Diffuse_Burley(DiffuseColor, Roughness, Context.NoV, NoL, Context.VoH);
+
+					// DualSpecular GGX
+					float AverageAlpha2 = Pow4(AverageRoughness);
+					float Lobe0Alpha2 = Pow4(Lobe0Roughness);
+					float Lobe1Alpha2 = Pow4(Lobe1Roughness);
+					float D = lerp(D_GGX(Lobe0Alpha2, Context.NoH), D_GGX(Lobe1Alpha2, Context.NoH), LobeMix);
+					float Vis = Vis_SmithJointApprox(AverageAlpha2, Context.NoV, NoL); // Average visibility well approximates using two separate ones (one per lobe).
+					float3 F = F_Schlick(SpecularColor, Context.VoH);
+					Specular = (D * Vis) * F;
+
+					// Application of transmission profile
+					float Thickness = 1.0;
+					Thickness = DecodeThickness(Thickness);
+					Thickness *= SSSS_MAX_TRANSMISSION_PROFILE_DISTANCE;
+					FTransmissionProfileParams TransmissionParams = GetTransmissionProfileParams(GBuffer);
+					float3 Profile = GetTransmissionProfile(GBuffer, Thickness).rgb;
+
+					float3 RefractV = refract(V, -N, TransmissionParams.OneOverIOR);
+					float PhaseFunction = ApproximateHG(dot(-L, RefractV), TransmissionParams.ScatteringDistribution);
+					Transmission = Diffuse * Profile * PhaseFunction;
+
+					break;
+				}
+				case SHADINGMODELID_SUBSURFACE:
+				{
+					Diffuse = Diffuse_Lambert(DiffuseColor);
+					float a2 = Pow4(Roughness);
+					float D = D_GGX(a2, Context.NoH);
+					float Vis = Vis_SmithJointApprox(a2, Context.NoV, NoL);
+					float3 F = F_Schlick(SpecularColor, Context.VoH);
+					Specular = (D * Vis) * F;
+
+					float3 SubsurfaceColor = ExtractSubsurfaceColor(ScreenSpaceData.GBuffer);
+					float Opacity = ScreenSpaceData.GBuffer.CustomData.a;
+					float InScatter = pow(saturate(dot(L, -V)), 12) * lerp(3, .1f, Opacity); 
+					float NormalContribution = saturate(Context.NoH * Opacity + 1 - Opacity);
+					float BackScatter = ScreenSpaceData.GBuffer.GBufferAO * NormalContribution / (PI * 2.0);
+					Transmission = SubsurfaceColor * lerp(BackScatter, 1, InScatter);
+
+					break;
+				}
+				case MATERIAL_SHADINGMODEL_PREINTEGRATED_SKIN:
+				{
+					// Default GGX
+					Diffuse = Diffuse_Lambert(DiffuseColor);
+					float a2 = Pow4(Roughness);
+					float D = D_GGX(a2, Context.NoH);
+					float Vis = Vis_SmithJointApprox(a2, Context.NoV, NoL);
+					float3 F = F_Schlick(SpecularColor, Context.VoH);
+					Specular = (D * Vis) * F;
+
+					float3 SubsurfaceColor = ExtractSubsurfaceColor(GBuffer);
+					float Opacity = GBuffer.CustomData.a;
+
+					float3 PreintegratedBRDF = Texture2DSampleLevel(View.PreIntegratedBRDF, View.PreIntegratedBRDFSampler, float2(saturate(dot(N, L) * .5 + .5), 1 - Opacity), 0).rgb;
+					Transmission = PreintegratedBRDF * SubsurfaceColor;
+					break;
+				}
+				// Unsupported materials will render magenta
+				default:
+				case SHADINGMODELID_EYE:
+				case SHADINGMODELID_HAIR:
+				case SHADINGMODELID_CLEAR_COAT:
+				case SHADINGMODELID_TWOSIDED_FOLIAGE:
+				{
+					Diffuse = float3(1, 0, 1);
+				}
+				break;
+			}
+
+			// #dxr_todo: Determine best way to incorporate integrals:
+			// 1) Return combined contribution (current method)
+			// 2) Split the integral (defer specular contribution to another pass)
+			// 3) Store diffuse and specular contribution separately
+			float3 Irradiance = IncomingRadiance * NoL * MisWeight / RayPdf;
+			float3 Brdf = Diffuse + Transmission + Specular;
+			//float3 Brdf = Diffuse + Transmission;
+			ExitantRadiance += Irradiance * Brdf;
+			DiffuseExitantRadiance += Irradiance * (Diffuse + Transmission);
 		}
 	}
 
 	if (SamplesPerPixel > 0)
 	{
-		Irradiance /= SamplesPerPixel;
+		float SamplesPerPixelInv = rcp(SamplesPerPixel);
+		ExitantRadiance *= SamplesPerPixelInv;
+		DiffuseExitantRadiance *= SamplesPerPixelInv;
 	}
-	RWOcclusionMaskUAV[PixelCoord] = float4(Irradiance, 1.0);
+
+	RWOcclusionMaskUAV[PixelCoord] = float4(ExitantRadiance, Luminance(DiffuseExitantRadiance));
 
 	if (HitCount > 0.0)
 	{
diff --git a/Engine/Source/Runtime/Renderer/Private/RayTracing/RaytracingSkylight.cpp b/Engine/Source/Runtime/Renderer/Private/RayTracing/RaytracingSkylight.cpp
--- a/Engine/Source/Runtime/Renderer/Private/RayTracing/RaytracingSkylight.cpp
+++ b/Engine/Source/Runtime/Renderer/Private/RayTracing/RaytracingSkylight.cpp
@@ -120,7 +120,7 @@ void FDeferredShadingSceneRenderer::RenderRayTracingSkyLight(
 	PassParameters->RWRayDistanceUAV = GraphBuilder.CreateUAV(FRDGTextureUAVDesc(RayDistanceTexture));
 	PassParameters->SamplesPerPixel = GRayTracingSkyLightSamplesPerPixel;
 	PassParameters->SamplingStopLevel = GRayTracingSkyLightSamplingStopLevel;
-	PassParameters->SkyLightColor = FVector(1.0);
+	PassParameters->SkyLightColor = Scene->SkyLight->LightColor
 	//PassParameters->SkyLightTexture = Scene->SkyLight->ProcessedTexture->TextureRHI;
 	//PassParameters->SkyLightTextureSampler = Scene->SkyLight->ProcessedTexture->SamplerStateRHI;
 	PassParameters->TLAS = View.PerViewRayTracingScene.RayTracingSceneRHI->GetShaderResourceView();
@@ -192,6 +192,8 @@ public:
 		TLASParameter.Bind(Initializer.ParameterMap, TEXT("TLAS"));
 		SceneTexturesParameter.Bind(Initializer.ParameterMap, TEXT("SceneTexturesStruct"));
 		SkyLightParameter.Bind(Initializer.ParameterMap, TEXT("SkyLight"));
+		TransmissionProfilesTextureParameter.Bind(Initializer.ParameterMap, TEXT("SSProfilesTexture"));
+		TransmissionProfilesLinearSamplerParameter.Bind(Initializer.ParameterMap, TEXT("TransmissionProfilesLinearSampler"));
 
 		OcclusionMaskUAVParameter.Bind(Initializer.ParameterMap, TEXT("RWOcclusionMaskUAV"));
 		RayDistanceUAVParameter.Bind(Initializer.ParameterMap, TEXT("RWRayDistanceUAV"));
@@ -204,6 +206,8 @@ public:
 		Ar << TLASParameter;
 		Ar << SceneTexturesParameter;
 		Ar << SkyLightParameter;
+		Ar << TransmissionProfilesTextureParameter;
+		Ar << TransmissionProfilesLinearSamplerParameter;
 		Ar << OcclusionMaskUAVParameter;
 		Ar << RayDistanceUAVParameter;
 		return bShaderHasOutdatedParameters;
@@ -233,6 +237,22 @@ public:
 		GlobalResources.Set(OcclusionMaskUAVParameter, OcclusionMaskUAV);
 		GlobalResources.Set(RayDistanceUAVParameter, RayDistanceUAV);
 
+		if (TransmissionProfilesTextureParameter.IsBound())
+		{
+			FSceneRenderTargets& SceneContext = FSceneRenderTargets::Get(RHICmdList);
+			const IPooledRenderTarget* PooledRT = GetSubsufaceProfileTexture_RT((FRHICommandListImmediate&)RHICmdList);
+
+			if (!PooledRT)
+			{
+				// no subsurface profile was used yet
+				PooledRT = GSystemTextures.BlackDummy;
+			}
+			const FSceneRenderTargetItem& Item = PooledRT->GetRenderTargetItem();
+
+			GlobalResources.SetTexture(TransmissionProfilesTextureParameter.GetBaseIndex(), Item.ShaderResourceTexture);
+			GlobalResources.SetSampler(TransmissionProfilesLinearSamplerParameter.GetBaseIndex(), TStaticSamplerState<SF_Bilinear, AM_Clamp, AM_Clamp, AM_Clamp>::GetRHI());
+		}
+
 		RHICmdList.RayTraceDispatch(Pipeline, GlobalResources, Width, Height);
 	}
 
@@ -243,6 +263,10 @@ private:
 	FShaderUniformBufferParameter SceneTexturesParameter;
 	FShaderUniformBufferParameter SkyLightParameter;
 
+	// SSS Profile
+	FShaderResourceParameter TransmissionProfilesTextureParameter;
+	FShaderResourceParameter TransmissionProfilesLinearSamplerParameter;
+
 	// Output
 	FShaderResourceParameter OcclusionMaskUAVParameter;
 	FShaderResourceParameter RayDistanceUAVParameter;
@@ -973,7 +997,7 @@ void FDeferredShadingSceneRenderer::CompositeRayTracingSkyLight(
 	SceneContext.BeginRenderingSceneColor(RHICmdList, ESimpleRenderTargetMode::EExistingColorAndDepth, FExclusiveDepthStencil::DepthRead_StencilWrite, true);
 	RHICmdList.ApplyCachedRenderTargets(GraphicsPSOInit);
 
-	GraphicsPSOInit.BlendState = TStaticBlendState<CW_RGB, BO_Add, BF_One, BF_One>::GetRHI();
+	GraphicsPSOInit.BlendState = TStaticBlendState<CW_RGBA, BO_Add, BF_One, BF_One, BO_Add, BF_One, BF_One>::GetRHI();
 	GraphicsPSOInit.RasterizerState = TStaticRasterizerState<FM_Solid, CM_None>::GetRHI();
 	GraphicsPSOInit.DepthStencilState = TStaticDepthStencilState<false, CF_Always>::GetRHI();
 	GraphicsPSOInit.BoundShaderState.VertexDeclarationRHI = GFilterVertexDeclaration.VertexDeclarationRHI;
diff --git a/Engine/Source/Runtime/Renderer/Private/ReflectionEnvironment.cpp b/Engine/Source/Runtime/Renderer/Private/ReflectionEnvironment.cpp
--- a/Engine/Source/Runtime/Renderer/Private/ReflectionEnvironment.cpp
+++ b/Engine/Source/Runtime/Renderer/Private/ReflectionEnvironment.cpp
@@ -704,7 +704,8 @@ void FDeferredShadingSceneRenderer::RenderDeferredReflectionsAndSkyLighting(FRHI
 
 	const bool bSkyLight = Scene->SkyLight
 		&& Scene->SkyLight->ProcessedTexture
-		&& !Scene->SkyLight->bHasStaticLighting;
+		&& !Scene->SkyLight->bHasStaticLighting
+		&& !ShouldRenderRayTracingDynamicSkyLight(Scene, ViewFamily);
 
 	bool bDynamicSkyLight = ShouldRenderDeferredDynamicSkyLight(Scene, ViewFamily);
 	bool bApplySkyShadowing = false;
