Fixes floating point precision issue in ratio estimator.

diff --git a/Engine/Shaders/Private/RayTracing/RayTracingReflections.usf b/Engine/Shaders/Private/RayTracing/RayTracingReflections.usf
--- a/Engine/Shaders/Private/RayTracing/RayTracingReflections.usf
+++ b/Engine/Shaders/Private/RayTracing/RayTracingReflections.usf
@@ -67,7 +67,7 @@ float3 GenerateReflectedRayDirection(
 )
 {
 	float3 RayDirection;
-	if (Roughness < 0.05 ) //ReflectionSmoothClamp)
+	if (Roughness < 0.001 ) //ReflectionSmoothClamp)
 	{
 		RayDirection = reflect(IncidentDirection, WorldNormal);
 	}
diff --git a/Engine/Shaders/Private/ScreenSpaceDenoise/SSDSignalFramework.ush b/Engine/Shaders/Private/ScreenSpaceDenoise/SSDSignalFramework.ush
--- a/Engine/Shaders/Private/ScreenSpaceDenoise/SSDSignalFramework.ush
+++ b/Engine/Shaders/Private/ScreenSpaceDenoise/SSDSignalFramework.ush
@@ -132,7 +132,7 @@ float3 GenerateReflectedRayDirection(
 	out float3 RayDirection,
 	out float PDF)
 {
-	if (Roughness < 0.05 ) //ReflectionSmoothClamp)
+	if (Roughness < 0.001 ) //ReflectionSmoothClamp)
 	{
 		RayDirection = reflect(IncidentDirection, WorldNormal);
 		PDF = 1.0;
@@ -400,9 +400,27 @@ float GetRatioEstimatorWeight(
 	float3 H = normalize(L + V);
 
 	float NoH = abs(dot(N, H));
-	float D = D_GGX(a2, NoH);
 
-	return D / OriginalRayPDF;
+	#if 1
+	{
+		const float MaximumWeight = 10;
+
+		float d = saturate(clamp( NoH * a2 - NoH, -NoH, 0) * NoH + 1);
+	
+		if (d == 0)
+		{
+			return MaximumWeight;
+		}
+
+		return clamp(a2 / (PI * d * d * OriginalRayPDF), 0, MaximumWeight);
+	}
+	#else
+	{
+		float D = D_GGX(a2, NoH);
+	
+		return clamp(D / OriginalRayPDF, 0, 10);
+	}
+	#endif
 }
 #else // !CONFIG_SPECULAR_RATIO_ESTIMATOR
 {
diff --git a/Engine/Shaders/Private/ScreenSpaceDenoise/SSDTemporalAccumulation.usf b/Engine/Shaders/Private/ScreenSpaceDenoise/SSDTemporalAccumulation.usf
--- a/Engine/Shaders/Private/ScreenSpaceDenoise/SSDTemporalAccumulation.usf
+++ b/Engine/Shaders/Private/ScreenSpaceDenoise/SSDTemporalAccumulation.usf
@@ -888,6 +888,9 @@ void TemporallyAccumulate(
 		HistorySample = Add(Mul(HistorySample, PreviousFrameMultiplier), CurrentFrame);
 	}
 
+	// Romoves NaN and negative numbers
+	// HistorySample = Minus(Min(Minus(HistorySample), CreateSignalSampleFromScalarValue(0.0)));
+
 	if (all(DispatchThreadId < uint2(View.ViewSizeAndInvSize.xy)))
 	{
 		OutputSignal(SignalHistoryOutput0, SignalHistoryOutput1, CONFIG_SIGNAL_HISTORY_LAYOUT, DispatchThreadId, HistorySample);
