Implemented adaptive ray shortening based on roughness for ray traced reflections

diff --git a/Engine/Shaders/Private/RayTracing/RayTracingReflections.usf b/Engine/Shaders/Private/RayTracing/RayTracingReflections.usf
--- a/Engine/Shaders/Private/RayTracing/RayTracingReflections.usf
+++ b/Engine/Shaders/Private/RayTracing/RayTracingReflections.usf
@@ -35,6 +35,7 @@ uint ShouldDoReflectedShadows;
 uint ShouldDoEmissiveAndIndirectLighting;
 uint UpscaleFactor;
 
+float ReflectionMinRayDistance;
 float ReflectionMaxRayDistance;
 float ReflectionMaxRoughness;
 
@@ -351,6 +352,13 @@ void RayTracingReflectionsRGS()
 
 	const float RoughnessFade = GetRoughnessFade(Roughness, ReflectionMaxRoughness);
 
+	// Sky is only sampled when infinite reflection rays are used.
+	const bool bAllowSkySampling = ReflectionMaxRayDistance < 0;
+
+	// Shorten the rays on rougher surfaces between user-provided min and max ray lengths.
+	// When a shortened ray misses the geometry, we fall back to local reflection capture sampling (similar to SSR).
+	const float LocalMaxRayDistance = bAllowSkySampling ? 1e27f : lerp(ReflectionMaxRayDistance, ReflectionMinRayDistance, Roughness);
+
 	if (RoughnessFade > 0)
 	{
 		for (uint SampleIndex = 0; SampleIndex < LocalSamplesPerPixel; ++SampleIndex)
@@ -361,7 +369,7 @@ void RayTracingReflectionsRGS()
 				RayDesc Ray;
 				Ray.Origin = WorldPosition + WorldNormal * 0.01;
 				Ray.TMin = 0.0;
-				Ray.TMax = ReflectionMaxRayDistance;
+				Ray.TMax = LocalMaxRayDistance;
 
 				uint DummyVariable;
 				float2 RandSample = RandomSequence_GenerateSample2D(RandSequence, DummyVariable);
@@ -403,7 +411,7 @@ void RayTracingReflectionsRGS()
 						ReflectedWeight += 1.0;
 					}
 				}
-				else
+				else if (bAllowSkySampling)
 				{
 					if (ReflectionStruct.SkyLightParameters.y > 0)
 					{
@@ -417,23 +425,25 @@ void RayTracingReflectionsRGS()
 				}
 
 				IncidentDirection = Ray.Direction;
-
 				HitDistance = Payload.HitT;
 			}
 		}
 	}
+
+	if (ReflectedWeight > 0)
+	{
+		float DistanceFadeRange = LocalMaxRayDistance * 0.1; // Fade out reflections over 10% of the total maximum ray length (arbitrarily chosen)
+		float DistanceFade = bAllowSkySampling ? 1.0 : saturate((LocalMaxRayDistance - HitDistance) * rcp(DistanceFadeRange));
+
+		ReflectedColor *= DistanceFade * RoughnessFade * rcp(ReflectedWeight);
+	}
 	else
 	{
-		// To tell the denoiser the black input is valid even if no ray has technically been shot. Uses a high ray hit distance
+		// To tell the denoiser the black input is valid even if no valid ray trace result was found. Uses a high ray hit distance
 		// so that the sample may be reused by another neighbor pixel in the denoise as if RoughnessFade in ]0;1[.
 		HitDistance = 1.0e20;
 	}
 
-	if (ReflectedWeight > 0)
-	{
-		ReflectedColor *= RoughnessFade * rcp(ReflectedWeight);
-	}
-
 	ColorOutput[DispatchThreadId] = float4(ReflectedColor);
 	RayHitDistanceOutput[DispatchThreadId] = HitDistance;
 }
diff --git a/Engine/Source/Runtime/Renderer/Private/RayTracing/RayTracingReflections.cpp b/Engine/Source/Runtime/Renderer/Private/RayTracing/RayTracingReflections.cpp
--- a/Engine/Source/Runtime/Renderer/Private/RayTracing/RayTracingReflections.cpp
+++ b/Engine/Source/Runtime/Renderer/Private/RayTracing/RayTracingReflections.cpp
@@ -37,11 +37,18 @@ static FAutoConsoleVariableRef CVarRayTracingReflectionsShadows(
 	TEXT("Enables shadows in ray tracing reflections (default = 1)")
 );
 
-static float GRayTracingReflectionsMaxRayDistance = 1.0e27;
+static float GRayTracingReflectionsMinRayDistance = -1;
+static FAutoConsoleVariableRef CVarRayTracingReflectionsMinRayDistance(
+	TEXT("r.RayTracing.Reflections.MinRayDistance"),
+	GRayTracingReflectionsMinRayDistance,
+	TEXT("Sets the minimum ray distance for ray traced reflection rays. Actual reflection ray length is computed as Lerp(MaxRayDistance, MinRayDistance, Roughness), i.e. reflection rays become shorter when traced from rougher surfaces. (default = -1 (infinite rays))")
+);
+
+static float GRayTracingReflectionsMaxRayDistance = -1;
 static FAutoConsoleVariableRef CVarRayTracingReflectionsMaxRayDistance(
 	TEXT("r.RayTracing.Reflections.MaxRayDistance"),
 	GRayTracingReflectionsMaxRayDistance,
-	TEXT("Sets the maximum ray distance for ray traced reflection rays (default = 1.0e27)")
+	TEXT("Sets the maximum ray distance for ray traced reflection rays. When ray shortening is used, skybox will not be sampled in RT reflection pass and will be composited later, together with local reflection captures. Negative values turn off this optimization. (default = -1 (infinite rays))")
 );
 
 static const int32 GReflectionLightCountMaximum = 64;
@@ -126,6 +133,7 @@ class FRayTracingReflectionsRG : public FGlobalShader
 		SHADER_PARAMETER(int32, ShouldDoReflectedShadows)
 		SHADER_PARAMETER(int32, ShouldDoEmissiveAndIndirectLighting)
 		SHADER_PARAMETER(int32, UpscaleFactor)
+		SHADER_PARAMETER(float, ReflectionMinRayDistance)
 		SHADER_PARAMETER(float, ReflectionMaxRayDistance)
 		SHADER_PARAMETER(float, ReflectionMaxRoughness)
 
@@ -219,6 +227,7 @@ void FDeferredShadingSceneRenderer::RayTraceReflections(
 	PassParameters->ShouldDoReflectedShadows = GRayTracingReflectionsShadows;
 	PassParameters->ShouldDoEmissiveAndIndirectLighting = GRayTracingReflectionsEmissiveAndIndirectLighting;
 	PassParameters->UpscaleFactor = UpscaleFactor;
+	PassParameters->ReflectionMinRayDistance = FMath::Min(GRayTracingReflectionsMinRayDistance, GRayTracingReflectionsMaxRayDistance);
 	PassParameters->ReflectionMaxRayDistance = GRayTracingReflectionsMaxRayDistance;
 	PassParameters->ReflectionMaxRoughness = FMath::Clamp(View.FinalPostProcessSettings.ScreenSpaceReflectionMaxRoughness, 0.01f, 1.0f);
 	PassParameters->LTCMatTexture = GSystemTextures.LTCMat->GetRenderTargetItem().ShaderResourceTexture;
