Need to match the invariant tag with DepthOnlyVertexShader.usf

diff --git a/Engine/Shaders/Private/BasePassVertexShader.usf b/Engine/Shaders/Private/BasePassVertexShader.usf
--- a/Engine/Shaders/Private/BasePassVertexShader.usf
+++ b/Engine/Shaders/Private/BasePassVertexShader.usf
@@ -80,7 +80,7 @@ void Main(
 		float3 ODS = OffsetODS(RasterizedWorldPosition.xyz, ResolvedView.TranslatedWorldCameraOrigin.xyz, ResolvedView.StereoIPD);
 		ClipSpacePosition = mul(float4(RasterizedWorldPosition.xyz + ODS, 1.0), ResolvedView.TranslatedWorldToClip);
 	#else
-		ClipSpacePosition = mul(RasterizedWorldPosition, ResolvedView.TranslatedWorldToClip);
+		ClipSpacePosition = INVARIANT(mul(RasterizedWorldPosition, ResolvedView.TranslatedWorldToClip));
 	#endif
 		Output.Position = INVARIANT(ClipSpacePosition);
 	}
