Implements generic denoised signal sample transformation for color spaces.

diff --git a/Engine/Shaders/Private/ColorSpace.ush b/Engine/Shaders/Private/ColorSpace.ush
--- a/Engine/Shaders/Private/ColorSpace.ush
+++ b/Engine/Shaders/Private/ColorSpace.ush
@@ -1,5 +1,8 @@
 // Copyright 1998-2019 Epic Games, Inc. All Rights Reserved.
 
+#pragma once
+
+
 //------------------------------------------------------------------
 // LMS
 // https://en.wikipedia.org/wiki/LMS_color_space
diff --git a/Engine/Shaders/Private/ScreenSpaceDenoise/SSDDefinitions.ush b/Engine/Shaders/Private/ScreenSpaceDenoise/SSDDefinitions.ush
--- a/Engine/Shaders/Private/ScreenSpaceDenoise/SSDDefinitions.ush
+++ b/Engine/Shaders/Private/ScreenSpaceDenoise/SSDDefinitions.ush
@@ -33,12 +33,6 @@
 #define DEBUG_STORE_NORMALIZED_COMPONENTS 1
 
 
-//------------------------------------------------------- SHARED CONFIG
-
-/** Do normal based anysotropy. */
-#define CONFIG_NORMAL_ANISOTROPY 1
-
-
 //------------------------------------------------------- SHARED CONSTANTS
 
 // Maximum number of mip level supported.
@@ -99,8 +93,26 @@
 	#define SIGNAL_PROCESSING_REFLECTIONS 1
 
 
+/** Defines the color space bitfield. */
+	#define COLOR_SPACE_RGB 0x0
+	#define COLOR_SPACE_YCOCG 0x1
+	#define COLOR_SPACE_LCOCG 0x2
+
+	#define COLOR_SPACE_UNPREMULTIPLY 0x4
+	#define COLOR_SPACE_KARIS_WEIGHTING 0x8
+
+
 /** Light types */
 	#define LIGHT_TYPE_DIRECTIONAL 0
 	#define LIGHT_TYPE_POINT 1
 	#define LIGHT_TYPE_SPOT 2
 	#define LIGHT_TYPE_RECT 3
+
+
+//------------------------------------------------------- SHARED CONFIG
+
+/** Do normal based anysotropy. */
+#define CONFIG_NORMAL_ANISOTROPY 1
+
+/** By default, the color space stored into intermediary buffer is linear premultiplied RGBA. */
+#define CONFIG_BUFFER_COLOR_SPACE COLOR_SPACE_RGB
diff --git a/Engine/Shaders/Private/ScreenSpaceDenoise/SSDSignalFramework.ush b/Engine/Shaders/Private/ScreenSpaceDenoise/SSDSignalFramework.ush
--- a/Engine/Shaders/Private/ScreenSpaceDenoise/SSDSignalFramework.ush
+++ b/Engine/Shaders/Private/ScreenSpaceDenoise/SSDSignalFramework.ush
@@ -13,12 +13,21 @@
 	#define CONFIG_SPECULAR_RATIO_ESTIMATOR 0
 #endif
 
+/** Color space to transform the color to. */
+#ifndef CONFIG_COLOR_SPACE_TRANSFORMATION
+	#define CONFIG_COLOR_SPACE_TRANSFORMATION COLOR_SPACE_TRANSFORMATION_DISABLED
+#endif
+
 
 //------------------------------------------------------- INCLUDE
 
 #include "SSDCommon.ush"
 #include "SSDSignalCore.ush"
 
+#include "../ColorSpace.ush"
+#include "../EyeAdaptationCommon.ush"
+
+
 #if !defined(CONFIG_SIGNAL_PROCESSING)
 	#error Missing CONFIG_SIGNAL_PROCESSING
 
@@ -37,6 +46,163 @@
 #endif
 
 
+//------------------------------------------------------- INTERNAL
+
+float Luma4(float3 Color)
+{
+	return (Color.g * 2.0) + (Color.r + Color.b);
+}
+
+float KarisHdrWeightY(float Color, float Exposure) 
+{
+	return rcp(Color * Exposure + 4.0);
+}
+
+float KarisHdrWeightInvY(float Color, float Exposure) 
+{
+	return 4.0 * rcp(1.0 - Color * Exposure);
+}
+
+float Luma_To_LumaLog(float x)
+{
+	return max(x > 0 ? log2(x) : -100, -100);
+}
+
+float GetSignalLuma4(FSSFSignalSample Sample, const uint Basis)
+{
+	// TODO: should force this to be scalar load.
+	float FrameExposureScale = EyeAdaptationLookup();
+		
+	// TODO: exposes optimisation if sample is normalized.
+	const bool bIsNormalizedSample = false;
+
+	if (Basis & COLOR_SPACE_KARIS_WEIGHTING)
+	{
+		float KarisX = Luma4(Sample.SceneColor.rgb);
+		if ((Basis & 0x3))
+		{
+			KarisX = Sample.SceneColor.x;
+		}
+
+		if (!bIsNormalizedSample)
+		{
+			KarisX *= SafeRcp(Sample.AccumulatedSampleCount);
+		}
+		
+		return KarisX *= KarisHdrWeightInvY(KarisX, FrameExposureScale);
+	}
+	else
+	{
+		if ((Basis & 0x3))
+		{
+			return Sample.SceneColor.x;
+		}
+
+		return Luma4(Sample.SceneColor.rgb);
+	}
+}
+
+/** Conveniently transform one simple from a source basis to a destination basis. */
+FSSFSignalSample TransformSignal(FSSFSignalSample Sample, const uint SrcBasis, const uint DestBasis)
+{
+	// If the source and destination bases are the same, early return to not pay ALU and floating point divergence.
+	if (SrcBasis == DestBasis)
+	{
+		return Sample;
+	}
+
+	const float Tweak = 1;
+
+	// TODO: should force this to be scalar load.
+	float FrameExposureScale = EyeAdaptationLookup() * Tweak;
+		
+	// TODO: exposes optimisation if sample is normalized.
+	const bool bIsNormalizedSample = false;
+
+	//Sample.SceneColor *= SafeRcp(Sample.AccumulatedSampleCount);
+
+	// Remove Karis weighting before doing any color transformations.
+	if (SrcBasis & COLOR_SPACE_KARIS_WEIGHTING)
+	{
+		float KarisX = Luma4(Sample.SceneColor.rgb);
+		if ((SrcBasis & 0x3))
+		{
+			KarisX = Sample.SceneColor.x;
+		}
+
+		if (!bIsNormalizedSample)
+		{
+			KarisX *= SafeRcp(Sample.AccumulatedSampleCount);
+		}
+
+		Sample.SceneColor *= KarisHdrWeightInvY(KarisX, FrameExposureScale);
+	}
+	
+	if ((SrcBasis & 0x3) == COLOR_SPACE_YCOCG)
+	{
+		Sample.SceneColor.rgb = YCoCg_2_LinearRGB(Sample.SceneColor.rgb);
+	}
+	else if ((SrcBasis & 0x3) == COLOR_SPACE_LCOCG)
+	{
+		Sample.SceneColor.rgb = LCoCg_2_LinearRGB(Sample.SceneColor.rgb);
+	}
+
+	// Premultiplied RGBA
+	if ((SrcBasis & COLOR_SPACE_UNPREMULTIPLY) == (DestBasis & COLOR_SPACE_UNPREMULTIPLY))
+	{
+		// NOP
+	}
+	else if (SrcBasis & COLOR_SPACE_UNPREMULTIPLY)
+	{
+		Sample.SceneColor.rgb *= Sample.SceneColor.a;
+	}
+	else //if (DestBasis & COLOR_SPACE_UNPREMULTIPLY)
+	{
+		Sample.SceneColor.rgb *= SafeRcp(Sample.SceneColor.a);
+	}
+
+	float x = Luma4(Sample.SceneColor.rgb);
+	if ((DestBasis & 0x3) == COLOR_SPACE_YCOCG)
+	{
+		Sample.SceneColor.xyz = LinearRGB_2_YCoCg(Sample.SceneColor.rgb);
+		x = Sample.SceneColor.x;
+	}
+	else if ((DestBasis & 0x3) == COLOR_SPACE_LCOCG)
+	{
+		Sample.SceneColor.xyz = LinearRGB_2_LCoCg(Sample.SceneColor.rgb);
+		x = Sample.SceneColor.x;
+	}
+
+	if (DestBasis & COLOR_SPACE_KARIS_WEIGHTING)
+	{
+		if (bIsNormalizedSample)
+		{
+			Sample.SceneColor *= KarisHdrWeightY(x, FrameExposureScale);
+		}
+		else
+		{
+			Sample.SceneColor *= KarisHdrWeightY(x * SafeRcp(Sample.AccumulatedSampleCount), FrameExposureScale);
+		}
+	}
+	
+	//Sample.SceneColor *= Sample.AccumulatedSampleCount;
+
+	return Sample;
+}
+
+#if 0 // will probably be handy at some point.
+float4 WeightedLerp( float4 ColorA, float WeightA, float4 ColorB, float WeightB, float Blend ) 
+{
+	float BlendA = (1.0 - Blend) * WeightA;
+	float BlendB =        Blend  * WeightB;
+	float RcpBlend = rcp(BlendA + BlendB);
+	BlendA *= RcpBlend;
+	BlendB *= RcpBlend;
+	return ColorA * BlendA + ColorB * BlendB;
+}
+#endif
+
+
 //------------------------------------------------------- INTERNAL FOR PENUMBRA SIGNAL
 #if CONFIG_SIGNAL_PROCESSING == SIGNAL_PROCESSING_PENUMBRA
 
diff --git a/Engine/Shaders/Private/ScreenSpaceDenoise/SSDSpatialAccumulation.usf b/Engine/Shaders/Private/ScreenSpaceDenoise/SSDSpatialAccumulation.usf
--- a/Engine/Shaders/Private/ScreenSpaceDenoise/SSDSpatialAccumulation.usf
+++ b/Engine/Shaders/Private/ScreenSpaceDenoise/SSDSpatialAccumulation.usf
@@ -56,6 +56,10 @@
 		#define CONFIG_UPSCALE              DIM_UPSCALE
 		//#define CONFIG_SAMPLE_COUNT_POLICY  SAMPLE_COUNT_POLICY_DISABLED
 
+		// Do color accumulation with karis weighting to avoid flickering specular highlight to show up the kernel pattern.
+		//#define CONFIG_COLOR_SPACE          (COLOR_SPACE_RGB | COLOR_SPACE_KARIS_WEIGHTING)
+
+
 	#elif DIM_STAGE == STAGE_POST_FILTERING
 		#define CONFIG_SIGNAL_INPUT_LAYOUT  SIGNAL_BUFFER_LAYOUT_REFLECTIONS_HISTORY
 		#define CONFIG_SIGNAL_OUTPUT_LAYOUT SIGNAL_BUFFER_LAYOUT_REFLECTIONS_HISTORY
@@ -107,6 +111,11 @@
 	#define CONFIG_UPSCALE 0
 #endif
 
+/** Color space. */
+#ifndef CONFIG_COLOR_SPACE
+	#define CONFIG_COLOR_SPACE CONFIG_BUFFER_COLOR_SPACE
+#endif
+
 
 //------------------------------------------------------- INCLUDES
 
@@ -259,6 +268,28 @@ RWTexture2D<float4>	DebugOutput;
 
 //------------------------------------------------------- FUNCTIONS
 
+FSSFSignalSample TransformSignalForPreRejection(FSSFSignalSample Sample)
+{
+	// Unpremultiplied RGBA to avoid any opacity mismatches.
+	Sample.SceneColor.rgb *= SafeRcp(Sample.SceneColor.a);
+
+	#if CONFIG_HISTORY_REJECTION_COLOR_SPACE == 1
+		Sample.SceneColor.rgb = LinearRGB_2_LCoCg(Sample.SceneColor.rgb);
+	#endif
+	return Sample;
+}
+
+FSSFSignalSample TransformSignalForPostRejection(FSSFSignalSample Sample)
+{
+	#if CONFIG_HISTORY_REJECTION_COLOR_SPACE == 1
+		Sample.SceneColor.rgb = LCoCg_2_LinearRGB(Sample.SceneColor.rgb);
+	#endif
+
+	// Premultiplied RGBA
+	Sample.SceneColor.rgb *= Sample.SceneColor.a;
+	return Sample;
+}
+
 // TODO: duplicated with reflection code.
 uint2 GetPixelCoord(uint2 DispatchThreadId)
 {
@@ -294,15 +325,22 @@ void MainCS(
 	// Sample the reference sample.
 	#if !CONFIG_UPSCALE || 1 // TODO: this sample ends up being sampled twice when upscaling.
 		FSSFSignalSample RefSample = SampleInputSignal(SignalInput0, SignalInput1, GlobalPointClampedSampler, CONFIG_SIGNAL_INPUT_LAYOUT, SceneBufferUV);
+	
+		RefSample = TransformSignal(
+			RefSample, 
+			/* SrcBasis  = */ CONFIG_BUFFER_COLOR_SPACE, 
+			/* DestBasis = */ CONFIG_COLOR_SPACE);
 	#endif
 
 	// Setup the spatial accumulator.
-	FSSFSignalSample Accumulator = CreateSignalSampleFromScalarValue(0.0);
+	FSSFSignalSample NeighborMoment1 = CreateSignalSampleFromScalarValue(0.0);
 	#if !CONFIG_UPSCALE
 		float AccumulatorWeight = 1; //GetMultiImportanceSamplingWeight(RefSceneMetadata, RefSceneMetadata, RefSample, DispatchThreadId);
-	
-		Accumulator = Mul(RefSample, AccumulatorWeight);
+
+		NeighborMoment1 = Mul(RefSample, AccumulatorWeight);
 	#endif
+	
+	FSSFSignalSample NeighborMoment2 = Power(NeighborMoment1, 2.0);
 
 	float CurrentWorldBluringRadius = ComputePixelWorldRadius(RefSceneMetadata);
 
@@ -443,6 +481,11 @@ void MainCS(
 		// Sample the neighborhood.
 		FSSFSignalSample NeighborSample = SampleInputSignal(SignalInput0, SignalInput1, GlobalPointClampedSampler, CONFIG_SIGNAL_INPUT_LAYOUT, NeighborBufferUV);
 		FSSFSampleSceneInfos NeighborSceneMetadata = FetchCurrentSceneInfos(SampleScreenPosition, NeighborBufferUV);
+		
+		NeighborSample = TransformSignal(
+			NeighborSample, 
+			/* SrcBasis  = */ CONFIG_BUFFER_COLOR_SPACE, 
+			/* DestBasis = */ CONFIG_COLOR_SPACE);
 
 		float NeighborWorldBluringRadius = ComputePixelWorldRadius(NeighborSceneMetadata);
 
@@ -479,17 +522,23 @@ void MainCS(
 			}
 		}
 		#endif
-
-		Accumulator = Add(Accumulator, Mul(NeighborSample, FinalWeight));
+		
+		NeighborMoment1 = Add(NeighborMoment1, Mul(NeighborSample, FinalWeight));
+		NeighborMoment2 = Add(NeighborMoment2, Mul(Power(NeighborSample, 2.0), FinalWeight));
 	}
+	
+	NeighborMoment1 = TransformSignal(
+		NeighborMoment1, 
+		/* SrcBasis  = */ CONFIG_COLOR_SPACE, 
+		/* DestBasis = */ CONFIG_BUFFER_COLOR_SPACE);
 
 	#if CONFIG_NORMALIZE_OUTPUT
 	{
-		Accumulator = NormalizeToOneSample(Accumulator);
+		NeighborMoment1 = NormalizeToOneSample(NeighborMoment1);
 	}
 	#endif
 
 	// TODO:  LeaveRayCount = (LeaveRayCount - 1) * 9 / (9 - 2) post processing to reject when for history rejection
 
-	OutputSignal(SignalOutput0, SignalOutput1, CONFIG_SIGNAL_OUTPUT_LAYOUT, DispatchThreadId, Accumulator);
+	OutputSignal(SignalOutput0, SignalOutput1, CONFIG_SIGNAL_OUTPUT_LAYOUT, DispatchThreadId, NeighborMoment1);
 } // MainCS
diff --git a/Engine/Shaders/Private/ScreenSpaceDenoise/SSDTemporalAccumulation.usf b/Engine/Shaders/Private/ScreenSpaceDenoise/SSDTemporalAccumulation.usf
--- a/Engine/Shaders/Private/ScreenSpaceDenoise/SSDTemporalAccumulation.usf
+++ b/Engine/Shaders/Private/ScreenSpaceDenoise/SSDTemporalAccumulation.usf
@@ -42,7 +42,9 @@
 	#define CONFIG_SIGNAL_HISTORY_LAYOUT SIGNAL_BUFFER_LAYOUT_REFLECTIONS_HISTORY
 
 	#define CONFIG_HISTORY_REJECTION 2
-	#define CONFIG_HISTORY_REJECTION_COLOR_SPACE 1
+	#define CONFIG_HISTORY_REJECTION_COLOR_SPACE (COLOR_SPACE_LCOCG | COLOR_SPACE_UNPREMULTIPLY)
+
+	//#define CONFIG_HISTORY_BLENDING_COLOR_SPACE (COLOR_SPACE_RGB | COLOR_SPACE_UNPREMULTIPLY | COLOR_SPACE_KARIS_WEIGHTING)
 
 #else
 	#error Unimplemented signal processing.
@@ -77,16 +79,16 @@
 	#define CONFIG_CERTAINTY_BASED_REJECTION 0
 #endif
 
-/** The color of the history should be rejected with. */
-#ifndef CONFIG_HISTORY_REJECTION_COLOR_SPACE
-	#define CONFIG_HISTORY_REJECTION_COLOR_SPACE 0
-#endif
-
 /** Whether should use a larger kernel to sample history. */
 #ifndef CONFIG_LARGER_HISTORY_KERNEL
 	#define CONFIG_LARGER_HISTORY_KERNEL 0
 #endif
 
+/** Whether the color should be be weightted according to intensity. */
+#ifndef CONFIG_TEMPORAL_COLOR_WEIGHT
+	#define CONFIG_TEMPORAL_COLOR_WEIGHT 0
+#endif
+
 
 //------------------------------------------------------- CONFIG ENABLED DEFAULTS
 
@@ -95,6 +97,16 @@
 	#define CONFIG_HISTORY_REJECTION 1
 #endif
 
+/** The color space to use for history rejection. */
+#ifndef CONFIG_HISTORY_REJECTION_COLOR_SPACE
+	#define CONFIG_HISTORY_REJECTION_COLOR_SPACE CONFIG_BUFFER_COLOR_SPACE
+#endif
+
+/** The color space to use for history blending. */
+#ifndef CONFIG_HISTORY_BLENDING_COLOR_SPACE
+	#define CONFIG_HISTORY_BLENDING_COLOR_SPACE CONFIG_BUFFER_COLOR_SPACE
+#endif
+
 
 //------------------------------------------------------- INCLUDES
 
@@ -142,34 +154,22 @@ RWTexture2D<float4>	DebugOutput;
 
 //------------------------------------------------------- FUNCTIONS
 
-float Luma_To_LumaLog(float x)
-{
-	return max(x > 0 ? log2(x) : -100, -100);
-}
-
 FSSFSignalSample TransformSignalForPreRejection(FSSFSignalSample Sample)
 {
-	// Unpremultiplied RGBA to avoid any opacity mismatches.
-	Sample.SceneColor.rgb *= SafeRcp(Sample.SceneColor.a);
-
-	#if CONFIG_HISTORY_REJECTION_COLOR_SPACE == 1
-		Sample.SceneColor.rgb = LinearRGB_2_LCoCg(Sample.SceneColor.rgb);
-	#endif
-	return Sample;
+	return TransformSignal(
+		Sample,
+		/* SrcConfig = */ CONFIG_BUFFER_COLOR_SPACE,
+		/* DestConfig = */ CONFIG_HISTORY_REJECTION_COLOR_SPACE);
 }
 
 FSSFSignalSample TransformSignalForPostRejection(FSSFSignalSample Sample)
 {
-	#if CONFIG_HISTORY_REJECTION_COLOR_SPACE == 1
-		Sample.SceneColor.rgb = LCoCg_2_LinearRGB(Sample.SceneColor.rgb);
-	#endif
-
-	// Premultiplied RGBA
-	Sample.SceneColor.rgb *= Sample.SceneColor.a;
-	return Sample;
+	return TransformSignal(
+		Sample,
+		/* SrcConfig = */ CONFIG_HISTORY_REJECTION_COLOR_SPACE,
+		/* DestConfig = */ CONFIG_BUFFER_COLOR_SPACE);
 }
 
-
 /** Sample history using a point sampler. */
 void PointSampleHistory(float2 ScreenPosition, float2 BufferUV, out FSSFSignalSample OutSample, out FSSFSampleSceneInfos OutMetadata)
 {
@@ -840,7 +840,7 @@ void TemporallyAccumulate(
 
 				// Compute how much this history get changed.
 				float RejectionFactor = 0.0;
-				#if CONFIG_SIGNAL_PROCESSING == SIGNAL_PROCESSING_REFLECTIONS && CONFIG_HISTORY_REJECTION_COLOR_SPACE == 1
+				#if CONFIG_SIGNAL_PROCESSING == SIGNAL_PROCESSING_REFLECTIONS && (CONFIG_HISTORY_REJECTION_COLOR_SPACE & COLOR_SPACE_LCOCG)
 				{
 					RejectionFactor = abs(
 						Luma_To_LumaLog(ClampedTransformedHistorySample.SceneColor.x) -
@@ -885,7 +885,25 @@ void TemporallyAccumulate(
 		float PreviousFrameWeight = min(HistorySample.AccumulatedSampleCount, TargetedSampleCount - CurrentFrame.AccumulatedSampleCount);
 		float PreviousFrameMultiplier = HistorySample.AccumulatedSampleCount > 0 ? PreviousFrameWeight / HistorySample.AccumulatedSampleCount : 0;
 
+		// TODO: use WeightedLerp instead for VALU perf.
+		// Pre transformation of the signal.
+		HistorySample = TransformSignal(
+			HistorySample,
+			/* SrcBasis = */ CONFIG_BUFFER_COLOR_SPACE,
+			/* DestBasis = */ CONFIG_HISTORY_BLENDING_COLOR_SPACE);
+		CurrentFrame = TransformSignal(
+			CurrentFrame,
+			/* SrcBasis = */ CONFIG_BUFFER_COLOR_SPACE,
+			/* DestBasis = */ CONFIG_HISTORY_BLENDING_COLOR_SPACE);
+
+		// Blending.
 		HistorySample = Add(Mul(HistorySample, PreviousFrameMultiplier), CurrentFrame);
+		
+		// Post transformation of the signal.
+		HistorySample = TransformSignal(
+			HistorySample,
+			/* SrcBasis = */ CONFIG_HISTORY_BLENDING_COLOR_SPACE,
+			/* DestBasis = */ CONFIG_BUFFER_COLOR_SPACE);
 	}
 
 	// Romoves NaN and negative numbers
diff --git a/Engine/Source/Runtime/Renderer/Private/SceneViewFamilyBlackboard.cpp b/Engine/Source/Runtime/Renderer/Private/SceneViewFamilyBlackboard.cpp
--- a/Engine/Source/Runtime/Renderer/Private/SceneViewFamilyBlackboard.cpp
+++ b/Engine/Source/Runtime/Renderer/Private/SceneViewFamilyBlackboard.cpp
@@ -5,6 +5,7 @@
 =============================================================================*/
 
 #include "SceneViewFamilyBlackboard.h"
+#include "SceneRendering.h"
 #include "SceneRenderTargets.h"
 #include "SystemTextures.h"
 
@@ -24,3 +25,15 @@ void SetupSceneViewFamilyBlackboard(
 	OutBlackboard->SceneGBufferD = RegisterExternalTextureWithFallback(GraphBuilder, SceneContext.GBufferD, GSystemTextures.BlackDummy, TEXT("GBufferD"));
 	OutBlackboard->SceneGBufferE = RegisterExternalTextureWithFallback(GraphBuilder, SceneContext.GBufferE, GSystemTextures.BlackDummy, TEXT("GBufferE"));
 }
+
+FRDGTextureRef GetEyeAdaptationTexture(FRDGBuilder& GraphBuilder, const FViewInfo& View)
+{
+	if (View.HasValidEyeAdaptation())
+	{
+		return GraphBuilder.RegisterExternalTexture(View.GetEyeAdaptation(GraphBuilder.RHICmdList), TEXT("ViewEyeAdaptation"));
+	}
+	else
+	{
+		return GraphBuilder.RegisterExternalTexture(GSystemTextures.WhiteDummy, TEXT("DefaultViewEyeAdaptation"));
+	}
+}
diff --git a/Engine/Source/Runtime/Renderer/Private/SceneViewFamilyBlackboard.h b/Engine/Source/Runtime/Renderer/Private/SceneViewFamilyBlackboard.h
--- a/Engine/Source/Runtime/Renderer/Private/SceneViewFamilyBlackboard.h
+++ b/Engine/Source/Runtime/Renderer/Private/SceneViewFamilyBlackboard.h
@@ -9,6 +9,9 @@
 #include "RenderGraph.h"
 
 
+class FViewInfo;
+
+
 /** Contains reference on all available buffer for a given scene. */
 BEGIN_SHADER_PARAMETER_STRUCT(FSceneViewFamilyBlackboard, )
 	SHADER_PARAMETER_RDG_TEXTURE(Texture2D, SceneDepthBuffer)
@@ -28,3 +31,7 @@ END_SHADER_PARAMETER_STRUCT()
 void SetupSceneViewFamilyBlackboard(
 	FRDGBuilder& GraphBuilder,
 	FSceneViewFamilyBlackboard* OutBlackboard);
+
+/** Returns a render graph texture resource reference onto the eye adaptation or fallback.
+ */
+FRDGTextureRef GetEyeAdaptationTexture(FRDGBuilder& GraphBuilder, const FViewInfo& View);
diff --git a/Engine/Source/Runtime/Renderer/Private/ScreenSpaceDenoise.cpp b/Engine/Source/Runtime/Renderer/Private/ScreenSpaceDenoise.cpp
--- a/Engine/Source/Runtime/Renderer/Private/ScreenSpaceDenoise.cpp
+++ b/Engine/Source/Runtime/Renderer/Private/ScreenSpaceDenoise.cpp
@@ -107,6 +107,7 @@ public:
 /** Shader parameter structure used for all shaders. */
 BEGIN_SHADER_PARAMETER_STRUCT(FSSDCommonParameters, )
 	SHADER_PARAMETER_STRUCT_INCLUDE(FSceneViewFamilyBlackboard, SceneBlackboard)
+	SHADER_PARAMETER_RDG_TEXTURE(Texture2D, EyeAdaptation)
 	SHADER_PARAMETER_RDG_TEXTURE(Texture2D<uint>, TileClassificationTexture)
 	SHADER_PARAMETER_STRUCT_REF(FViewUniformShaderParameters, ViewUniformBuffer)
 END_SHADER_PARAMETER_STRUCT()
@@ -663,6 +664,7 @@ IScreenSpaceDenoiser::FReflectionsOutputs DenoiseReflections(
 	{
 		CommonParameters.SceneBlackboard = SceneBlackboard;
 		CommonParameters.ViewUniformBuffer = View.ViewUniformBuffer;
+		CommonParameters.EyeAdaptation = GetEyeAdaptationTexture(GraphBuilder, View);
 	}
 
 	FScreenSpaceFilteringHistory PrevFrameHistory;
@@ -724,6 +726,7 @@ IScreenSpaceDenoiser::FReflectionsOutputs DenoiseReflections(
 
 		FSSDTemporalAccumulationCS::FParameters* PassParameters = GraphBuilder.AllocParameters<FSSDTemporalAccumulationCS::FParameters>();
 		PassParameters->CommonParameters = CommonParameters;
+
 		PassParameters->bCameraCut = View.bCameraCut || !PrevFrameHistory.RT[0].IsValid();
 		PassParameters->SignalInput0 = SignalHistory0;
 		PassParameters->SignalInput1 = SignalHistory1;
