Reorder interpolants in DepthOnly and Velocity shaders to workaround shader linkage problem (triggers only in d3ddebug on DX11, but crashing on DX12)

diff --git a/Engine/Shaders/Private/DepthOnlyPixelShader.usf b/Engine/Shaders/Private/DepthOnlyPixelShader.usf
--- a/Engine/Shaders/Private/DepthOnlyPixelShader.usf
+++ b/Engine/Shaders/Private/DepthOnlyPixelShader.usf
@@ -15,12 +15,12 @@ float MobileColorValue;
 
 void Main(
 #if !MATERIALBLENDING_SOLID || OUTPUT_PIXEL_DEPTH_OFFSET
-	FVertexFactoryInterpolantsVSToPS FactoryInterpolants,
-	float4 PixelPosition : TEXCOORD6,
+	in INPUT_POSITION_QUALIFIERS float4 SvPosition : SV_Position,
 #if USE_WORLD_POSITION_EXCLUDING_SHADER_OFFSETS
 	float3 PixelPositionExcludingWPO : TEXCOORD7,
 #endif
-	in INPUT_POSITION_QUALIFIERS float4 SvPosition : SV_Position
+	FVertexFactoryInterpolantsVSToPS FactoryInterpolants
+	// float4 PixelPosition : TEXCOORD6 -- Not used currently
 	OPTIONAL_IsFrontFace
 	OPTIONAL_OutDepthConservative,
 #endif
diff --git a/Engine/Shaders/Private/DepthOnlyVertexShader.usf b/Engine/Shaders/Private/DepthOnlyVertexShader.usf
--- a/Engine/Shaders/Private/DepthOnlyVertexShader.usf
+++ b/Engine/Shaders/Private/DepthOnlyVertexShader.usf
@@ -13,15 +13,16 @@
 
 struct FDepthOnlyVSToPS
 {
+	float4 Position : SV_POSITION;
+
+	#if USE_RAW_WORLD_POSITION
+		float3 PixelPositionExcludingWPO : TEXCOORD7;
+	#endif
+	
 	#if !MATERIALBLENDING_SOLID || OUTPUT_PIXEL_DEPTH_OFFSET
 		FVertexFactoryInterpolantsVSToPS FactoryInterpolants;
 		float4 PixelPosition : TEXCOORD6;
 	#endif
-	#if USE_RAW_WORLD_POSITION
-		float3 PixelPositionExcludingWPO : TEXCOORD7;
-	#endif
-
-	float4 Position : SV_POSITION;
 };
 
 #if USING_TESSELLATION	
diff --git a/Engine/Shaders/Private/VelocityShader.usf b/Engine/Shaders/Private/VelocityShader.usf
--- a/Engine/Shaders/Private/VelocityShader.usf
+++ b/Engine/Shaders/Private/VelocityShader.usf
@@ -29,8 +29,8 @@ struct FVelocityInterpsVSToPS
 
 struct FVelocityVSToPS
 {
-	FVelocityInterpsVSToPS Interps;
 	float4 Position : SV_POSITION;
+	FVelocityInterpsVSToPS Interps;
 };
 
 #if USING_TESSELLATION
@@ -234,8 +234,8 @@ void MainVertexShader(
  *============================================================================*/
  
 void MainPixelShader(
-	FVelocityInterpsVSToPS Input,
-	in INPUT_POSITION_QUALIFIERS float4 SvPosition : SV_Position
+	in INPUT_POSITION_QUALIFIERS float4 SvPosition : SV_Position,
+	FVelocityInterpsVSToPS Input
 	OPTIONAL_IsFrontFace
 	OPTIONAL_OutDepthConservative
 	,out float4 OutColor : SV_Target0
