Sphere light sampling support

Currently selects a uniform random position on a sphere as a connection point for a
shadow ray. Enabled by adding a source radius to a point light.

TODO: Optimize via disk projection

diff --git a/Engine/Shaders/Private/RayTracing/RayTracingOcclusionRGS.usf b/Engine/Shaders/Private/RayTracing/RayTracingOcclusionRGS.usf
--- a/Engine/Shaders/Private/RayTracing/RayTracingOcclusionRGS.usf
+++ b/Engine/Shaders/Private/RayTracing/RayTracingOcclusionRGS.usf
@@ -7,6 +7,7 @@
 #include "RayTracingCommon.ush"
 #include "RayTracingDirectionalLight.ush"
 #include "RayTracingRectLight.ush"
+#include "RayTracingSphereLight.ush"
 
 /** Light types */
 #define LIGHT_TYPE_DIRECTIONAL 0
@@ -64,22 +65,31 @@ bool GenerateSpotLightOcclusionRay(
 	out float RayTMax
 )
 {
-	float3 LightDirection = LightParameters.Position - WorldPosition;
-	float RayLength = length(LightDirection);
-	LightDirection /= RayLength;
+	if (LightParameters.SourceRadius > 0.0)
+	{
+		float RayPdf = 0.0;
+		GenerateSphereLightOcclusionRay(LightParameters, WorldPosition, WorldNormal, RandSample,
+			RayOrigin, RayDirection, RayTMin, RayTMax, RayPdf);
+	}
+	else
+	{
+		float3 LightDirection = LightParameters.Position - WorldPosition;
+		float RayLength = length(LightDirection);
+		LightDirection /= RayLength;
+
+		// Apply normal perturbation when defining ray
+		RayOrigin = WorldPosition + WorldNormal * 0.01;
+		RayDirection = LightDirection;
+		RayTMin = 0.0;
+		RayTMax = RayLength;
+	}
 
 	// Apply culling
 	float CosAngle = LightParameters.SpotAngles.x;
-	if (dot(LightDirection, LightParameters.Direction) < CosAngle)
+	if (dot(RayDirection, LightParameters.Direction) < CosAngle)
 	{
 		return false;
 	}
-
-	// Apply normal perturbation when defining ray
-	RayOrigin = WorldPosition + WorldNormal * 0.01;
-	RayDirection = LightDirection;
-	RayTMin = 0.0;
-	RayTMax = RayLength;
 	return true;
 }
 
@@ -111,14 +121,30 @@ bool GenerateOcclusionRay(
 	}
 	#elif LIGHT_TYPE == LIGHT_TYPE_POINT
 	{
-		return GeneratePointLightOcclusionRay(
-			LightParameters,
-			WorldPosition, WorldNormal,
-			RandSample,
-			/* out */ RayOrigin,
-			/* out */ RayDirection,
-			/* out */ RayTMin,
-			/* out */ RayTMax);
+		if (LightParameters.SourceRadius == 0)
+		{
+			return GeneratePointLightOcclusionRay(
+				LightParameters,
+				WorldPosition, WorldNormal,
+				RandSample,
+				/* out */ RayOrigin,
+				/* out */ RayDirection,
+				/* out */ RayTMin,
+				/* out */ RayTMax);
+		}
+		else
+		{
+			float RayPdf;
+			return GenerateSphereLightOcclusionRay(
+				LightParameters,
+				WorldPosition, WorldNormal,
+				RandSample,
+				/* out */ RayOrigin,
+				/* out */ RayDirection,
+				/* out */ RayTMin,
+				/* out */ RayTMax,
+				/* out */ RayPdf);
+		}
 	}
 	#elif LIGHT_TYPE == LIGHT_TYPE_SPOT
 	{
diff --git a/Engine/Shaders/Private/RayTracing/RayTracingSphereLight.ush b/Engine/Shaders/Private/RayTracing/RayTracingSphereLight.ush
new file mode 100644
--- /dev/null
+++ b/Engine/Shaders/Private/RayTracing/RayTracingSphereLight.ush
@@ -0,0 +1,35 @@
+// Copyright 1998-2019 Epic Games, Inc. All Rights Reserved.
+
+#pragma once
+
+#include "../LightShaderParameters.ush"
+
+bool GenerateSphereLightOcclusionRay(
+	FLightShaderParameters LightParameters,
+	float3 WorldPosition,
+	float3 WorldNormal,
+	float2 RandSample,
+	out float3 RayOrigin,
+	out float3 RayDirection,
+	out float RayTMin,
+	out float RayTMax,
+	out float RayPdf
+)
+{
+	float4 Result = UniformSampleSphere(RandSample);
+	float3 LightNormal = Result.xyz;
+	float3 LightPosition = LightParameters.Position + LightNormal * LightParameters.SourceRadius;
+	float3 LightDirection = LightPosition - WorldPosition;
+	float RayLength = length(LightDirection);
+	LightDirection /= RayLength;
+
+	RayOrigin = WorldPosition + WorldNormal * 0.01;
+	RayDirection = LightDirection;
+	RayTMin = 0.0;
+	RayTMax = RayLength;
+
+	float SolidAnglePdf = Result.w * saturate(dot(LightNormal, -LightDirection)) / (RayLength * RayLength);
+	RayPdf = SolidAnglePdf;
+	return true;
+}
+ 
\ No newline at end of file
