Adding a bias hack to niagara shaders to fix UE-68456.
NV needs to look at what code they are generating since it seems like this is due to optimization/operation reordering between the depth and base passes.

diff --git a/Engine/Shaders/Private/DepthOnlyVertexShader.usf b/Engine/Shaders/Private/DepthOnlyVertexShader.usf
--- a/Engine/Shaders/Private/DepthOnlyVertexShader.usf
+++ b/Engine/Shaders/Private/DepthOnlyVertexShader.usf
@@ -174,6 +174,9 @@ void Main(
 		// XBOXONE_BIAS_HACK is defined only in a custom node in a particular material
 		// This should be removed with a future shader compiler update
 		Output.Position.z -= 0.0001 * Output.Position.w;
+	#elif NIAGARA_BIAS_HACK
+		// Similar to above, only for Niagara. (UE-68456)
+		Output.Position.z -= 0.0001 * Output.Position.w;
 	#endif
 
 #if USE_GLOBAL_CLIP_PLANE
diff --git a/Engine/Shaders/Private/NiagaraVFParticleAccess.usf b/Engine/Shaders/Private/NiagaraVFParticleAccess.usf
--- a/Engine/Shaders/Private/NiagaraVFParticleAccess.usf
+++ b/Engine/Shaders/Private/NiagaraVFParticleAccess.usf
@@ -8,6 +8,9 @@
 //uint NiagaraIntDataOffset;
 //uint NiagaraIntDataStride;
 
+// This is a hack to fix cases where depth values are invariant due to compiler optimizations/operation reordering between the depth pass and base pass. (UE-68456)
+#define NIAGARA_BIAS_HACK 1
+
 #if !RAYHITGROUPSHADER && !COMPUTESHADER
 
 Buffer<float> NiagaraParticleDataFloat;
