Implements point light and spot light shadow denoisers.

Only point lights are whitelisted in mean time the ray generation shader is done.

diff --git a/Engine/Shaders/Private/ScreenSpaceDenoise/SSDSignalFramework.ush b/Engine/Shaders/Private/ScreenSpaceDenoise/SSDSignalFramework.ush
--- a/Engine/Shaders/Private/ScreenSpaceDenoise/SSDSignalFramework.ush
+++ b/Engine/Shaders/Private/ScreenSpaceDenoise/SSDSignalFramework.ush
@@ -474,13 +474,12 @@ float GetSignalWorldBluringRadius(FSSFSignalSample Sample, FSSFSampleSceneInfos
 	{
 		return HitDistanceToWorldBluringRadius * Sample.PenumbraClosestHit * Anisotropy;
 	}
-	#elif DIM_LIGHT_TYPE == LIGHT_TYPE_POINT
+	#elif DIM_LIGHT_TYPE == LIGHT_TYPE_POINT || DIM_LIGHT_TYPE == LIGHT_TYPE_SPOT
 	{
-		return 0; // TODO: point light soft shadows are not implemented.
-	}
-	#elif DIM_LIGHT_TYPE == LIGHT_TYPE_SPOT
-	{
-		return 0; // TODO: spot light soft shadows are not implemented.
+		float3 PixelToLightWorldVector = (LightParameters.Position + View.PreViewTranslation) - SceneMetadata.TranslatedWorldPosition;
+		float DistanceFromLight = length(PixelToLightWorldVector);
+
+		return LightParameters.SourceRadius * Sample.PenumbraClosestHit / (DistanceFromLight - Sample.PenumbraClosestHit) * Anisotropy;
 	}
 	#elif DIM_LIGHT_TYPE == LIGHT_TYPE_RECT
 	{
diff --git a/Engine/Source/Runtime/Renderer/Private/ScreenSpaceDenoise.cpp b/Engine/Source/Runtime/Renderer/Private/ScreenSpaceDenoise.cpp
--- a/Engine/Source/Runtime/Renderer/Private/ScreenSpaceDenoise.cpp
+++ b/Engine/Source/Runtime/Renderer/Private/ScreenSpaceDenoise.cpp
@@ -80,8 +80,8 @@ const IScreenSpaceDenoiser* GScreenSpaceDenoiser = nullptr;
 
 static bool IsSupportedLightType(ELightComponentType LightType)
 {
-	// TODO.
-	return LightType == LightType_Directional || LightType == LightType_Rect;
+	// TODO: whitelist LightType_Spot once implemented in the ray generation shader and tested.
+	return  LightType == LightType_Point || LightType == LightType_Directional || LightType == LightType_Rect;
 }
 
 
