DR - Implemented texture LOD based on RayCones from the raytracing gems. The computation of the TextureLod itself is deviating from the proposed solution in the paper:
We still compute the Ratio between TextureUV area and Worldspace Triangle area, but than we use the area of the ConeCap (instead of it's width) directly to convert it into it's representative area in texture space.
This basically avoids bunches of log2 computations as has been done in the original paper which did all this computation in log2-space. Our method matches the Rasterizer almost perfectly.

The UvArea approximation is taken from UV Level 0 and does not take any scaling within the material graph into account.
Surface curvature for focusing mirrors or metal spheres is not yet taken into account, but can be added later for the initial bounce of RT reflections.

diff --git a/Engine/Shaders/Private/Common.ush b/Engine/Shaders/Private/Common.ush
--- a/Engine/Shaders/Private/Common.ush
+++ b/Engine/Shaders/Private/Common.ush
@@ -70,6 +70,25 @@
 // Using SV_ClipDistance has overhead (15% slower base pass in triangle bound test scene on PS4) so projects have to opt-in
 #define USE_GLOBAL_CLIP_PLANE (PLATFORM_SUPPORTS_GLOBAL_CLIP_PLANE && PROJECT_ALLOW_GLOBAL_CLIP_PLANE && !MATERIAL_DOMAIN_POSTPROCESS && !MATERIAL_DOMAIN_UI)
 
+const static MaterialFloat PI = 3.1415926535897932f;
+
+#ifndef USE_RAYTRACED_TEXTURE_RAYCONE_LOD
+#define USE_RAYTRACED_TEXTURE_RAYCONE_LOD (RAYHITGROUPSHADER)
+#endif
+
+static float GlobalRayCone_TexArea;
+float ComputeRayConeLod(Texture2D Tex)
+{
+#if USE_RAYTRACED_TEXTURE_RAYCONE_LOD
+	uint2 Dimensions;
+	Tex.GetDimensions(Dimensions.x, Dimensions.y);
+    int TexArea = Dimensions.x * Dimensions.y;
+    return 0.5f * log2(GlobalRayCone_TexArea * TexArea);
+#else
+    return 0;
+#endif
+}
+
 // This would need to be a #define in GLSL to ignore the SamplerState, however, it is currently a function call in HLSL
 // for type checking of the parameters - ironically the type checking is really only needed in GLSL!
 MaterialFloat4 Texture1DSample(Texture1D Tex, SamplerState Sampler, float UV)
@@ -83,7 +102,7 @@ MaterialFloat4 Texture1DSample(Texture1D Tex, SamplerState Sampler, float UV)
 MaterialFloat4 Texture2DSample(Texture2D Tex, SamplerState Sampler, float2 UV)
 {
 #if !PIXELSHADER
-	return Tex.SampleLevel(Sampler, UV, 0);
+	return Tex.SampleLevel(Sampler, UV, ComputeRayConeLod(Tex));
 #else
 	return Tex.Sample(Sampler, UV);
 #endif
@@ -91,7 +110,7 @@ MaterialFloat4 Texture2DSample(Texture2D Tex, SamplerState Sampler, float2 UV)
 MaterialFloat Texture2DSample_A8(Texture2D Tex, SamplerState Sampler, float2 UV)
 {
 #if !PIXELSHADER
-	return Tex.SampleLevel(Sampler, UV, 0) A8_SAMPLE_MASK;
+	return Tex.SampleLevel(Sampler, UV, ComputeRayConeLod(Tex)) A8_SAMPLE_MASK;
 #else
 	return Tex.Sample(Sampler, UV) A8_SAMPLE_MASK;
 #endif
@@ -123,7 +142,7 @@ MaterialFloat4 Texture2DSampleLevel(Texture2D Tex, SamplerState Sampler, float2
 MaterialFloat4 Texture2DSampleBias(Texture2D Tex, SamplerState Sampler, float2 UV, MaterialFloat MipBias)
 {
 #if !PIXELSHADER
-	return Tex.SampleLevel(Sampler, UV, 0);
+	return Tex.SampleLevel(Sampler, UV, ComputeRayConeLod(Tex) + MipBias);
 #else
 	return Tex.SampleBias(Sampler, UV, MipBias);
 #endif
@@ -171,7 +190,7 @@ MaterialFloat4 TextureCubeSampleGrad(TextureCube Tex, SamplerState Sampler, floa
 MaterialFloat4 TextureExternalSample(TextureExternal Tex, SamplerState Sampler, float2 UV)
 {
 #if !PIXELSHADER
-	return Tex.SampleLevel(Sampler, UV, 0);
+	return Tex.SampleLevel(Sampler, UV, ComputeRayConeLod(Tex));
 #else
 	return Tex.Sample(Sampler, UV);
 #endif
@@ -283,8 +302,6 @@ MaterialFloat Luminance( MaterialFloat3 LinearColor )
 	return dot( LinearColor, MaterialFloat3( 0.3, 0.59, 0.11 ) );
 }
 
-const static MaterialFloat PI = 3.1415926535897932f;
-
 MaterialFloat length2(MaterialFloat2 v)
 {
 	return dot(v, v);
diff --git a/Engine/Shaders/Private/LocalVertexFactory.ush b/Engine/Shaders/Private/LocalVertexFactory.ush
--- a/Engine/Shaders/Private/LocalVertexFactory.ush
+++ b/Engine/Shaders/Private/LocalVertexFactory.ush
@@ -184,7 +184,7 @@ FVertexFactoryInput LoadVertexFactoryInput(uint TriangleIndex, int VertexIndex)
 	FTriangleBaseAttributes Triangle = LoadTriangleBaseAttributes(TriangleIndex);
 
 	FVertexFactoryInput Input;
-	Input.Position = float4(WorldRayOrigin() + WorldRayDirection() * RayTCurrent(), 0.0f);
+	Input.Position = float4(Triangle.LocalPositions[VertexIndex], 1.0f);
 	Input.VertexId = Triangle.Indices[VertexIndex];
 #if VF_USE_PRIMITIVE_SCENE_DATA
 	Input.PrimitiveId = GetInstanceUserData();
diff --git a/Engine/Shaders/Private/PathTracing/PathTracing.usf b/Engine/Shaders/Private/PathTracing/PathTracing.usf
--- a/Engine/Shaders/Private/PathTracing/PathTracing.usf
+++ b/Engine/Shaders/Private/PathTracing/PathTracing.usf
@@ -76,6 +76,9 @@ void PathTracingMainRG()
 		float3 RayThroughput = 1.0;
 		float RayPdf = 1.0;
 
+		FRayCone RayCone = (FRayCone)0;
+		RayCone.SpreadAngle = View.EyeToPixelSpreadAngle;
+
 		uint RayFlags = RAY_FLAG_CULL_BACK_FACING_TRIANGLES;
 
 		// Cast ray
@@ -86,9 +89,9 @@ void PathTracingMainRG()
 			0,      // RayContributionToHitGroupIndex
 			1,      // MultiplierForGeometryContributionToShaderIndex #dxr_todo: we currently only support shader binding per mesh instance
 			0,      // MissShaderIndex
-			Ray    // RayDesc
+			Ray,    // RayDesc
+			RayCone
 		);
-
 		RayCounter++;
 
 		float MaterialPdf = 0.0;
@@ -175,6 +178,9 @@ void PathTracingMainRG()
 				if (NeePdf > 0.0)
 				{
 					RayDesc LightRay;
+
+					FRayCone LightRayCone = RayCone;
+
 					GenerateLightRay(Ray, LightId, LightUV, LightRay);
 					// #dxr_todo: Guard against tracing a back-facing ray
 
@@ -187,8 +193,9 @@ void PathTracingMainRG()
 						0,         // RayContributionToHitGroupIndex
 						1,         // MultiplierForGeometryContributionToShaderIndex #dxr_todo: we currently only support shader binding per mesh instance
 						0,         // MissShaderIndex
-						LightRay);
-
+						LightRay,
+						LightRayCone
+					);
 					RayCounter++;
 
 					if (!IsHit(LightRay, NeePayload))
@@ -282,8 +289,9 @@ void PathTracingMainRG()
 				0,    // RayContributionToHitGroupIndex
 				1,    // MultiplierForGeometryContributionToShaderIndex
 				0,    // MissShaderIndex
-				Ray);
-
+				Ray,
+				RayCone /* propagated by the material CHS */
+			);
 			RayCounter++;
 		}
 
diff --git a/Engine/Shaders/Private/PathTracing/PathTracingCHS.usf b/Engine/Shaders/Private/PathTracing/PathTracingCHS.usf
--- a/Engine/Shaders/Private/PathTracing/PathTracingCHS.usf
+++ b/Engine/Shaders/Private/PathTracing/PathTracingCHS.usf
@@ -37,5 +37,5 @@ void PathTracingMainCHS(inout FPackedMaterialClosestHitPayload PackedPayload, in
 	Payload.WorldPos = float3(0.0, 0.0, 0.0);
 	Payload.HitKind = 0;
 
-    PackedPayload = PackRayTracingPayload(Payload);
+    PackedPayload = PackRayTracingPayload(Payload, PackedPayload.RayCone);
 }
\ No newline at end of file
diff --git a/Engine/Shaders/Private/PathTracing/PathTracingMS.usf b/Engine/Shaders/Private/PathTracing/PathTracingMS.usf
--- a/Engine/Shaders/Private/PathTracing/PathTracingMS.usf
+++ b/Engine/Shaders/Private/PathTracing/PathTracingMS.usf
@@ -11,5 +11,5 @@ void PathTracingMainMS(inout FPackedMaterialClosestHitPayload PackedPayload)
 	//Payload.InstanceID = ~0;
 	//Payload.InstanceIndex = ~0;
 	//Payload.PrimitiveIndex = ~0;
-	PackedPayload.HitT = -1;
+	PackedPayload.Payload.HitT = -1;
 }
\ No newline at end of file
diff --git a/Engine/Shaders/Private/RayTracing/RayTracingCommon.ush b/Engine/Shaders/Private/RayTracing/RayTracingCommon.ush
--- a/Engine/Shaders/Private/RayTracing/RayTracingCommon.ush
+++ b/Engine/Shaders/Private/RayTracing/RayTracingCommon.ush
@@ -49,33 +49,48 @@ struct FDefaultAttributes
 	float2 Barycentrics;
 };
 
+struct FRayCone
+{
+    float Width;
+    float SpreadAngle;
+};
+
+FRayCone PropagateRayCone(in FRayCone Cone, in float SurfaceSpreadAngle, in float  HitT)
+{
+	FRayCone NewCone;
+	NewCone.Width = Cone.SpreadAngle * HitT + Cone.Width;
+	NewCone.SpreadAngle = Cone.SpreadAngle + SurfaceSpreadAngle;
+	return NewCone;
+}
+
 struct FMaterialClosestHitPayload
 {
-	float HitT;							//4		4bytes	32bits
-
-	float3 Radiance;					//16	6bytes	48bits	
-	float3 WorldNormal;					//28	6bytes	48bits
-	float3 BaseColor;					//40	6bytes	48bits
-	float3 DiffuseColor;				//52	0bytes	(derived)
-	float3 SpecularColor;				//64	0bytes	(derived)
-	float Opacity;						//68	2bytes	16bits
-	float Metallic;						//72	1bytes	8bits
-	float Specular;						//76	1bytes	8bits
-	float Roughness;					//80	2bytes	16bits
-	float Ior;							//84	2bytes	16bits
-	uint ShadingModelID;				//88	1byte	4bits
-	uint BlendingMode;					//92	0bytes	4bits (packed)
-	uint PrimitiveLightingChannelMask;	//96	0bytes	(removed)
-	float4 CustomData;					//112	0bytes	(removed)
-	float GBufferAO;					//116	0bytes	(removed)
-	float3 IndirectIrradiance;			//128	0bytes	48bits	
+	//FRayCone RayCone;                 // 8    8bytes  64bits /* the raycone is hidden in the internal payload */
+	float HitT;							//12	4bytes	32bits
+
+	float3 Radiance;					//24	6bytes	48bits	
+	float3 WorldNormal;					//36	6bytes	48bits
+	float3 BaseColor;					//48	6bytes	48bits
+	float3 DiffuseColor;				//60	0bytes	(derived)
+	float3 SpecularColor;				//72	0bytes	(derived)
+	float Opacity;						//76	2bytes	16bits
+	float Metallic;						//80	1bytes	8bits
+	float Specular;						//84	1bytes	8bits
+	float Roughness;					//88	2bytes	16bits
+	float Ior;							//92	2bytes	16bits
+	uint ShadingModelID;				//96	1byte	4bits
+	uint BlendingMode;					//100	0bytes	4bits (packed)
+	uint PrimitiveLightingChannelMask;	//104	0bytes	(removed)
+	float4 CustomData;					//120	0bytes	(removed)
+	float GBufferAO;					//124	0bytes	(removed)
+	float3 IndirectIrradiance;			//136	0bytes	48bits	/* gbuffer only has float payload and there are truncation HLSL warnings */
 
 	// Quite some code assume FRayHitInfo has a WorldPos
 	// So keep it here and force to re-construct it in Unpack call using ray information.
 	// It is not packed in FRayHitInfoPacked
-	float3 WorldPos;					//140	0bytes	(derived)
+	float3 WorldPos;					//148	0bytes	(derived)
 
-	uint HitKind;						//144	0bytes 1bit	(packed)
+	uint HitKind;						//152	0bytes 1bit	(packed)
 };
 
 #define USE_PACKED_PAYLOAD 0
@@ -85,6 +100,7 @@ struct FMaterialClosestHitPayload
 struct FPackedMaterialClosestHitPayload
 {
 	float HitT;
+	FRayCone RayCone;
 	int RadienceAndNormal[3];
 	int BaseColorAndOpacity[2];
 	int MetallicAndSpecularAndRoughness;
@@ -92,10 +108,11 @@ struct FPackedMaterialClosestHitPayload
 	int IndirectIrradiance[2];
 };
 
-FPackedMaterialClosestHitPayload PackRayTracingPayload(FMaterialClosestHitPayload Input)
+FPackedMaterialClosestHitPayload PackRayTracingPayload(FMaterialClosestHitPayload Input, in FRayCone RayCone)
 {
     FPackedMaterialClosestHitPayload Output = (FPackedMaterialClosestHitPayload)0;
 	Output.HitT = Input.HitT;
+	Output.RayCone = RayCone; 
     Output.RadienceAndNormal[0]  = f32tof16(Input.Radiance.x);
 	Output.RadienceAndNormal[0] |= f32tof16(Input.Radiance.y) << 16;
 	Output.RadienceAndNormal[1]  = f32tof16(Input.Radiance.z);
@@ -124,6 +141,7 @@ FMaterialClosestHitPayload UnpackRayTracingPayload(FPackedMaterialClosestHitPayl
 {
     FMaterialClosestHitPayload Output = (FMaterialClosestHitPayload)0;
 	Output.HitT = Input.HitT;
+	Output.RayCone = Input.RayCone; 
 	Output.Radiance.x = f16tof32(Input.RadienceAndNormal[0]);
 	Output.Radiance.y = f16tof32(Input.RadienceAndNormal[0] >> 16);
 	Output.Radiance.z = f16tof32(Input.RadienceAndNormal[1]);
@@ -152,31 +170,41 @@ FMaterialClosestHitPayload UnpackRayTracingPayload(FPackedMaterialClosestHitPayl
 }
 
 #else
-#define FPackedMaterialClosestHitPayload FMaterialClosestHitPayload
 
-FPackedMaterialClosestHitPayload PackRayTracingPayload(FMaterialClosestHitPayload Input)
+struct FPackedMaterialClosestHitPayload
+{
+	FRayCone RayCone;
+	FMaterialClosestHitPayload Payload;
+};
+
+FPackedMaterialClosestHitPayload PackRayTracingPayload(FMaterialClosestHitPayload Input, in FRayCone RayCone)
 {
-	return Input;
+	FPackedMaterialClosestHitPayload Payload;
+	Payload.RayCone = RayCone;
+	Payload.Payload = Input;
+    return Payload;
 }
 
 FMaterialClosestHitPayload UnpackRayTracingPayload(FPackedMaterialClosestHitPayload Input, RayDesc Ray)
 {
-	return Input;
+    return Input.Payload;
 }
 
 #endif
 
-FMaterialClosestHitPayload TraceRayInternal(in RaytracingAccelerationStructure TLAS, 
-											in uint RayFlags, 
-											in uint InstanceInclusionMask, 
-											in uint RayContributionToHitGroupIndex,  
-											in uint MultiplierForGeometryContributionToShaderIndex, 
-											in uint MissShaderIndex, 
-											in RayDesc Ray)
+    FMaterialClosestHitPayload TraceRayInternal(in RaytracingAccelerationStructure TLAS,
+											in uint RayFlags,
+											in uint InstanceInclusionMask,
+											in uint RayContributionToHitGroupIndex,
+											in uint MultiplierForGeometryContributionToShaderIndex,
+											in uint MissShaderIndex,
+											in RayDesc Ray,
+											inout FRayCone RayCone)
 {
 	FPackedMaterialClosestHitPayload Payload = (FPackedMaterialClosestHitPayload)0;
+	Payload.RayCone = RayCone;
 
-    TraceRay
+	TraceRay
 	(
 			TLAS,
 			RayFlags,
@@ -188,6 +216,7 @@ FMaterialClosestHitPayload TraceRayInternal(in RaytracingAccelerationStructure T
 			Payload
 	);
 
+    RayCone = Payload.RayCone;
     return UnpackRayTracingPayload(Payload, Ray);
 }
 #endif
diff --git a/Engine/Shaders/Private/RayTracing/RayTracingDebug.usf b/Engine/Shaders/Private/RayTracing/RayTracingDebug.usf
--- a/Engine/Shaders/Private/RayTracing/RayTracingDebug.usf
+++ b/Engine/Shaders/Private/RayTracing/RayTracingDebug.usf
@@ -33,6 +33,9 @@ void RayTracingDebugMainRGS()
 	Ray.TMin = 0.5f; // 1 unit in UE4 = 1 cm -0> 0.5 ==> 5 mm
 	Ray.TMax = 1e27f;
 
+    FRayCone RayCone = (FRayCone)0;
+	RayCone.SpreadAngle = View.EyeToPixelSpreadAngle;
+
 	uint RayFlags = RAY_FLAG_CULL_BACK_FACING_TRIANGLES;
 
 	FMaterialClosestHitPayload Payload = TraceRayInternal(
@@ -42,7 +45,8 @@ void RayTracingDebugMainRGS()
 		0,      // RayContributionToHitGroupIndex
 		1,      // MultiplierForGeometryContributionToShaderIndex
 		0,      // MissShaderIndex
-		Ray     // RayDesc
+		Ray,    // RayDesc
+		RayCone
 	);
 
 	float4 Result = float4(0, 0, 0, 1);
@@ -115,7 +119,7 @@ void RayTracingDebugMainRGS()
 [shader("miss")]
 void RayTracingDebugMainMS(inout FPackedMaterialClosestHitPayload PackedPayload)
 {
-	PackedPayload.HitT = -1;
+	PackedPayload.Payload.HitT = -1;
 }
 
 
@@ -134,5 +138,5 @@ void RayTracingDebugMainCHS(inout FPackedMaterialClosestHitPayload PackedPayload
 	float3x3 LocalToWorldNormal = transpose(WorldToLocal);
 	Payload.WorldNormal = normalize(mul(LocalToWorldNormal, cross(Edge0, Edge1)));
 
-    PackedPayload = PackRayTracingPayload(Payload);
+    PackedPayload = PackRayTracingPayload(Payload, PackedPayload.RayCone);
 }
diff --git a/Engine/Shaders/Private/RayTracing/RayTracingMaterialHitShaders.usf b/Engine/Shaders/Private/RayTracing/RayTracingMaterialHitShaders.usf
--- a/Engine/Shaders/Private/RayTracing/RayTracingMaterialHitShaders.usf
+++ b/Engine/Shaders/Private/RayTracing/RayTracingMaterialHitShaders.usf
@@ -176,11 +176,13 @@ float4 ClipSpacePositionToSvPosition(float4 ClipSpacePosition)
 	return SvPosition;
 }
 
-void CalcInterpolants(in FDefaultAttributes Attributes, out FVertexFactoryInterpolantsVSToPS Interpolants, out float4 FakeSvPosition)
+void CalcInterpolants(in FRayCone RayCone, in FDefaultAttributes Attributes, out FVertexFactoryInterpolantsVSToPS Interpolants, out float4 FakeSvPosition)
 {
 	FVertexFactoryInterpolantsVSToDS PerVertexInterpolants[3];
 	float4 PerVertexClipSpacePositions[3]; // Used to fake SvPosition
-	
+	float3 WorldPositions[3];
+	float2 TexCoords[3];
+
 	for (int i = 0; i < 3; i++)
 	{
 		FVertexFactoryInput Input = LoadVertexFactoryInput(PrimitiveIndex(), i);
@@ -190,13 +192,17 @@ void CalcInterpolants(in FDefaultAttributes Attributes, out FVertexFactoryInterp
 		float4 WorldPositionExcludingWPO = VertexFactoryGetWorldPosition(Input, VFIntermediates);
 		FMaterialVertexParameters VertexParameters = GetMaterialVertexParameters(Input, VFIntermediates, WorldPositionExcludingWPO, TangentToLocal);
 		
+		WorldPositions[i] = WorldPositionExcludingWPO.xyz;
 		float4 WorldPosition = WorldPositionExcludingWPO + float4(GetMaterialWorldPositionOffset(VertexParameters), 0.0f);
 		
 		PerVertexInterpolants[i] = VertexFactoryGetInterpolantsVSToDS(Input, VFIntermediates, VertexParameters);
-		
+
+#if NUM_TEX_COORD_INTERPOLATORS || USE_PARTICLE_SUBUVS
+		TexCoords[i] = PerVertexInterpolants[i].InterpolantsVSToPS.TexCoords[0].xy;
+#endif
+
 		float4 RasterizedWorldPosition = VertexFactoryGetRasterizedWorldPosition(Input, VFIntermediates, WorldPosition);
 		PerVertexClipSpacePositions[i] = mul(RasterizedWorldPosition, ResolvedView.TranslatedWorldToClip);
-
 	}
 	
 	FVertexFactoryInterpolantsVSToDS Interpolated = 
@@ -211,6 +217,23 @@ void CalcInterpolants(in FDefaultAttributes Attributes, out FVertexFactoryInterp
 	
 	float4 ClipSpacePosition = PerVertexClipSpacePositions[0] * (1 - Attributes.Barycentrics.x - Attributes.Barycentrics.y) + PerVertexClipSpacePositions[1] * Attributes.Barycentrics.x + PerVertexClipSpacePositions[2] * Attributes.Barycentrics.y;
 	FakeSvPosition = ClipSpacePosition / ClipSpacePosition.w;
+
+#if NUM_TEX_COORD_INTERPOLATORS || USE_PARTICLE_SUBUVS
+	float2 TA = TexCoords[1] - TexCoords[0];
+	float2 TB = TexCoords[2] - TexCoords[0];
+    float UvArea = abs((TA.x * TB.y) - (TA.y * TB.x));
+
+	float3 PA = WorldPositions[1] - WorldPositions[0];
+	float3 PB = WorldPositions[2] - WorldPositions[0];
+    float3 Unnormalized = cross(PB, PA);
+	float InvWorldArea = rsqrt(dot(Unnormalized, Unnormalized));
+    
+    float RayConeCapArea = RayCone.Width * RayCone.Width * (1.0f / (2.0f * PI));
+    float3 RayDir = WorldRayDirection();
+    float3 ShadingNormal = normalize(Interpolated.InterpolantsVSToPS.TangentToWorld2);
+    float RaySlope = dot(RayDir, ShadingNormal);
+    GlobalRayCone_TexArea = (UvArea * InvWorldArea * RayConeCapArea) / max(RaySlope * RaySlope, 0.00001f);
+#endif
 }
 
 [shader("closesthit")]
@@ -221,7 +244,8 @@ void MaterialCHS(inout FPackedMaterialClosestHitPayload PackedPayload, in FDefau
 	FVertexFactoryInterpolantsVSToPS Interpolants;
 	float4 SvPosition;
 	
-	CalcInterpolants(Attributes, Interpolants, SvPosition);
+	FRayCone PropagatedCone = PropagateRayCone(PackedPayload.RayCone, 0 /* surface curvature */, RayTCurrent());
+	CalcInterpolants(PropagatedCone, Attributes, Interpolants, SvPosition);
 
 	FMaterialPixelParameters MaterialParameters = GetMaterialPixelParameters(Interpolants, SvPosition);
 	FPixelMaterialInputs PixelMaterialInputs;
@@ -304,7 +328,7 @@ void MaterialCHS(inout FPackedMaterialClosestHitPayload PackedPayload, in FDefau
 	
 	Payload.IndirectIrradiance += Payload.DiffuseColor * DiffuseIndirectLighting; // Contains both lightmap & skylight
 
-	PackedPayload = PackRayTracingPayload(Payload);
+	PackedPayload = PackRayTracingPayload(Payload, PropagatedCone);
 }
 
 [shader("anyhit")] 
@@ -316,7 +340,7 @@ void MaterialAHS(inout FPackedMaterialClosestHitPayload PackedPayload, in FDefau
 	FVertexFactoryInterpolantsVSToPS Interpolants;
 	float4 SvPosition;
 	
-	CalcInterpolants(Attributes, Interpolants, SvPosition);
+	CalcInterpolants(PackedPayload.RayCone, Attributes, Interpolants, SvPosition);
 
 	FMaterialPixelParameters MaterialParameters = GetMaterialPixelParameters(Interpolants, SvPosition);
 	FPixelMaterialInputs PixelMaterialInputs;
diff --git a/Engine/Shaders/Private/RayTracing/RayTracingReflections.usf b/Engine/Shaders/Private/RayTracing/RayTracingReflections.usf
--- a/Engine/Shaders/Private/RayTracing/RayTracingReflections.usf
+++ b/Engine/Shaders/Private/RayTracing/RayTracingReflections.usf
@@ -131,7 +131,8 @@ float4 BlendWithEnvironment(inout float3 BRDFThroughput, const in FMaterialClose
 float TraceShadowRay(
 	float3 Origin,
 	float3 Direction,
-	float TMax)
+	float TMax,
+	in FRayCone ShadowRayCone)
 {
 	RayDesc Ray;
 	Ray.Origin = Origin;
@@ -148,13 +149,14 @@ float TraceShadowRay(
 		0,         // RayContributionToHitGroupIndex
 		1,         // MultiplierForGeometryContributionToShaderIndex
 		0,         // MissShaderIndex
-		Ray        // RayDesc
+		Ray,       // RayDesc
+		ShadowRayCone
 	);
 
 	return Payload.HitT;
 }
 
-float3 ComputeDirectLighting(RayDesc Ray, FMaterialClosestHitPayload Payload, bool bReflectedShadows)
+float3 ComputeDirectLighting(RayDesc Ray, in FRayCone RayCone, FMaterialClosestHitPayload Payload, bool bReflectedShadows)
 {
 	float3 DirectLighting = float3(0.0, 0.0, 0.0);
 
@@ -249,9 +251,11 @@ float3 ComputeDirectLighting(RayDesc Ray, FMaterialClosestHitPayload Payload, bo
 		//#dxr_todo: check perf improvements forcing TMin == TMax when bReflectedShadows = false, to avoid tracing rays inside if statements
 		if (bReflectedShadows)
 		{
+			FRayCone ShadowRayCone = RayCone;
+
 			//#dxr_todo: add support for sampling area lights (and perhaps add CVar to enable stocastic shadows?)
 			float DistanceToLight = length(ToLight);
-			float HitT = TraceShadowRay(WorldPosition + Payload.WorldNormal*0.001, ShadowRayDirection, DistanceToLight);
+			float HitT = TraceShadowRay(WorldPosition + Payload.WorldNormal*0.001, ShadowRayDirection, DistanceToLight, ShadowRayCone);
 			AmbientOcclusion = HitT < 0.0f;
 		}
 
@@ -263,7 +267,8 @@ float3 ComputeDirectLighting(RayDesc Ray, FMaterialClosestHitPayload Payload, bo
 }
 
 FMaterialClosestHitPayload TraceRayAndAccumulateResults(
-	RayDesc Ray,
+	in RayDesc Ray,
+	inout FRayCone RayCone,
 	inout float3 Radiance)
 {
 	// Call TraceRay iteratively until alpha is 0 and accumulate the result
@@ -279,12 +284,13 @@ FMaterialClosestHitPayload TraceRayAndAccumulateResults(
 		0,         // RayContributionToHitGroupIndex
 		1,         // MultiplierForGeometryContributionToShaderIndex
 		0,         // MissShaderIndex
-		Ray        // RayDesc
+		Ray,       // RayDesc
+		RayCone
 	);
 
 	if (ShouldDoDirectLighting)
 	{
-		float3 DirectLighting = ComputeDirectLighting(Ray, Payload, ShouldDoReflectedShadows > 0);
+		float3 DirectLighting = ComputeDirectLighting(Ray, RayCone, Payload, ShouldDoReflectedShadows > 0);
 
 		// Transform NaNs to black, transform negative colors to black.
 		DirectLighting = -min(-DirectLighting, 0.0);
@@ -364,6 +370,12 @@ void RayTracingReflectionsRGS()
 		for (uint SampleIndex = 0; SampleIndex < LocalSamplesPerPixel; ++SampleIndex)
 		{
 			float3 IncidentDirection = CameraDirection;
+					
+			FRayCone RayCone = (FRayCone)0;
+			RayCone.SpreadAngle = View.EyeToPixelSpreadAngle;
+            float SurfaceCurvature = 0.0f; /* #todo_dxr assume no curvature */
+			RayCone = PropagateRayCone(RayCone, SurfaceCurvature, Depth);
+
 			for (uint BounceIndex = 0; BounceIndex < MaxBounces; ++BounceIndex)
 			{
 				RayDesc Ray;
@@ -392,7 +404,7 @@ void RayTracingReflectionsRGS()
 				uint RayFlags = RAY_FLAG_CULL_BACK_FACING_TRIANGLES;
 
 				float3 Radiance = float3(0, 0, 0);
-				FMaterialClosestHitPayload Payload = TraceRayAndAccumulateResults(Ray, Radiance);
+				FMaterialClosestHitPayload Payload = TraceRayAndAccumulateResults(Ray, RayCone, Radiance);
 
 				float4 PathRadiance;
 				if (Payload.HitT >= 0)
@@ -452,7 +464,7 @@ void RayTracingReflectionsRGS()
 [shader("miss")]
 void RayTracingReflectionsMainMS(inout FPackedMaterialClosestHitPayload PackedPayload)
 {
-	PackedPayload.HitT = -1;
+	PackedPayload.Payload.HitT = -1;
 }
 
 
@@ -471,5 +483,5 @@ void RayTracingReflectionsMainCHS(inout FPackedMaterialClosestHitPayload PackedP
 	float3x3 LocalToWorldNormal = transpose(WorldToLocal);
 	Payload.WorldNormal = normalize(mul(LocalToWorldNormal, cross(Edge0, Edge1)));
 
-	PackedPayload = PackRayTracingPayload(Payload);
+	PackedPayload = PackRayTracingPayload(Payload, PackedPayload.RayCone);
 }
diff --git a/Engine/Source/Runtime/Engine/Private/SceneView.cpp b/Engine/Source/Runtime/Engine/Private/SceneView.cpp
--- a/Engine/Source/Runtime/Engine/Private/SceneView.cpp
+++ b/Engine/Source/Runtime/Engine/Private/SceneView.cpp
@@ -2063,6 +2063,10 @@ void FSceneView::SetupViewRectUniformBufferParameters(FViewUniformShaderParamete
 		InvBufferSizeX * (EffectiveViewRect.Max.X - 0.5),
 		InvBufferSizeY * (EffectiveViewRect.Max.Y - 0.5));
 
+	/* Texture Level-of-Detail Strategies for Real-Time Ray Tracing https://developer.nvidia.com/raytracinggems Equation 20 */
+	float RadFOV = (PI / 180.0f) * FOV;
+	ViewUniformShaderParameters.EyeToPixelSpreadAngle = FPlatformMath::Atan((2.0f * FPlatformMath::Tan(RadFOV * 0.5f)) / BufferSize.Y);
+
 	ViewUniformShaderParameters.MotionBlurNormalizedToPixel = FinalPostProcessSettings.MotionBlurMax * EffectiveViewRect.Width() / 100.0f;
 
 	{
diff --git a/Engine/Source/Runtime/Engine/Public/SceneView.h b/Engine/Source/Runtime/Engine/Public/SceneView.h
--- a/Engine/Source/Runtime/Engine/Public/SceneView.h
+++ b/Engine/Source/Runtime/Engine/Public/SceneView.h
@@ -755,7 +755,8 @@ enum ETranslucencyVolumeCascade
 	VIEW_UNIFORM_BUFFER_MEMBER(float, VolumetricLightmapBrickSize) \
 	VIEW_UNIFORM_BUFFER_MEMBER(FVector, VolumetricLightmapBrickTexelSize) \
 	VIEW_UNIFORM_BUFFER_MEMBER(float, StereoIPD) \
-	VIEW_UNIFORM_BUFFER_MEMBER(float, IndirectLightingCacheShowFlag)
+	VIEW_UNIFORM_BUFFER_MEMBER(float, IndirectLightingCacheShowFlag) \
+	VIEW_UNIFORM_BUFFER_MEMBER(float, EyeToPixelSpreadAngle)
 
 #define VIEW_UNIFORM_BUFFER_MEMBER(type, identifier) \
 	SHADER_PARAMETER(type, identifier)
diff --git a/Engine/Source/Runtime/Renderer/Private/DeferredShadingRenderer.cpp b/Engine/Source/Runtime/Renderer/Private/DeferredShadingRenderer.cpp
--- a/Engine/Source/Runtime/Renderer/Private/DeferredShadingRenderer.cpp
+++ b/Engine/Source/Runtime/Renderer/Private/DeferredShadingRenderer.cpp
@@ -912,7 +912,7 @@ FRHIRayTracingPipelineState* FDeferredShadingSceneRenderer::BindRayTracingPipeli
 	TArray<FRayTracingHitGroupInitializer> RayTracingMaterialLibrary;
 	FShaderResource::GetRayTracingMaterialLibrary(RayTracingMaterialLibrary);
 	Initializer.SetHitGroups(RayTracingMaterialLibrary);
-	Initializer.MaxPayloadSizeInBytes = 144; // #dxr_todo: set Initializer.MaxPayloadSizeInBytes based on shader requirements (it's not obvious how to compute this right now)
+	Initializer.MaxPayloadSizeInBytes = 152; // #dxr_todo: set Initializer.MaxPayloadSizeInBytes based on shader requirements (it's not obvious how to compute this right now)
 
 	FRHIRayTracingPipelineState* PipelineState = PipelineStateCache::GetAndOrCreateRayTracingPipelineState(Initializer);
 
