DR - Workaround to unblock hlslcc platforms
#rnx

diff --git a/Engine/Shaders/Private/Common.ush b/Engine/Shaders/Private/Common.ush
--- a/Engine/Shaders/Private/Common.ush
+++ b/Engine/Shaders/Private/Common.ush
@@ -76,6 +76,10 @@ const static MaterialFloat PI = 3.1415926535897932f;
 #define USE_RAYTRACED_TEXTURE_RAYCONE_LOD (RAYHITGROUPSHADER)
 #endif
 
+// Workaround
+#if COMPILER_HLSLCC
+	#define ComputeRayConeLod(Tex)	0
+#else
 static float GlobalRayCone_TexArea;
 float ComputeRayConeLod(Texture2D Tex)
 {
@@ -88,6 +92,7 @@ float ComputeRayConeLod(Texture2D Tex)
     return 0;
 #endif
 }
+#endif
 
 // This would need to be a #define in GLSL to ignore the SamplerState, however, it is currently a function call in HLSL
 // for type checking of the parameters - ironically the type checking is really only needed in GLSL!
