DR - Documented change
#rnx

diff --git a/Engine/Shaders/Private/Common.ush b/Engine/Shaders/Private/Common.ush
--- a/Engine/Shaders/Private/Common.ush
+++ b/Engine/Shaders/Private/Common.ush
@@ -76,8 +76,8 @@ const static MaterialFloat PI = 3.1415926535897932f;
 #define USE_RAYTRACED_TEXTURE_RAYCONE_LOD (RAYHITGROUPSHADER)
 #endif
 
-// Workaround
 #if COMPILER_HLSLCC
+	// Workaround until UE-69158 gets fixed
 	#define ComputeRayConeLod(Tex)	0
 #else
 static float GlobalRayCone_TexArea;
