Integrated github pull request to optimize vertex fog.

 - ExponentialFogParameters3.z contains whether the inscattering cubemap is active. It cannot be active if the shader permutation bit is not set. Therefore, predicating the check on whether the permutation define is enabled saves a comparison.

 - When additive / modulate blend modes are enabled, fog color is unused. By setting it to zero it saves some vertex ALU.

#4239

diff --git a/Engine/Shaders/Private/HeightFogCommon.ush b/Engine/Shaders/Private/HeightFogCommon.ush
--- a/Engine/Shaders/Private/HeightFogCommon.ush
+++ b/Engine/Shaders/Private/HeightFogCommon.ush
@@ -122,7 +122,12 @@ half4 GetExponentialHeightFog(float3 WorldPositionRelativeToCamera, float Exclud
 
 #if SUPPORT_FOG_DIRECTIONAL_LIGHT_INSCATTERING
 	// if InscatteringLightDirection.w is negative then it's disabled, otherwise it holds directional inscattering start distance
-	if (FogStruct.InscatteringLightDirection.w >= 0 && FogStruct.ExponentialFogParameters3.z == 0)
+	BRANCH
+	if (FogStruct.InscatteringLightDirection.w >= 0
+	#if SUPPORT_FOG_INSCATTERING_TEXTURE
+		&& FogStruct.ExponentialFogParameters3.z == 0
+	#endif
+	)
 	{
 		float DirectionalInscatteringStartDistance = FogStruct.InscatteringLightDirection.w;
 		// Setup a cosine lobe around the light direction to approximate inscattering from the directional light off of the ambient haze;
@@ -147,7 +152,14 @@ half4 GetExponentialHeightFog(float3 WorldPositionRelativeToCamera, float Exclud
 		DirectionalInscattering = 0;
 	}
 
-	return half4((InscatteringColor) * (1 - ExpFogFactor) + DirectionalInscattering, ExpFogFactor);
+	// Fog color is unused when additive / modulate blend modes are active.
+	#if (MATERIALBLENDING_ADDITIVE || MATERIALBLENDING_MODULATE)
+		half3 FogColor = 0.0;
+	#else
+		half3 FogColor = (InscatteringColor) * (1 - ExpFogFactor) + DirectionalInscattering;
+	#endif
+
+	return half4(FogColor, ExpFogFactor);
 }
 
 half4 CalculateHeightFog(float3 WorldPositionRelativeToCamera)
