Integrating github change to optimize shadow occlusion evaluation.

"Rearrange math so per pixel constant can be optimized out from per sample calculations.
(a - c)b + 1 -> a * b - (c * b - 1)

d = (c * b - 1) can be pre calculated once per pixel so there is single MAD left anymore per sample as: a * b - d
This saves 36ALU with highest quality shadows."

#5185

diff --git a/Engine/Shaders/Private/ShadowFilteringCommon.ush b/Engine/Shaders/Private/ShadowFilteringCommon.ush
--- a/Engine/Shaders/Private/ShadowFilteringCommon.ush
+++ b/Engine/Shaders/Private/ShadowFilteringCommon.ush
@@ -164,8 +164,12 @@ float4 CalculateOcclusion(float4 ShadowmapDepth, FPCFSamplerSettings Settings)
 		// The standard comparison is SceneDepth < ShadowmapDepth
 		// Using a soft transition based on depth difference
 		// Offsets shadows a bit but reduces self shadowing artifacts considerably
-		float TransitionScale = Settings.TransitionScale; //SoftTransitionScale.z;
-		float4 ShadowFactor = saturate((ShadowmapDepth - Settings.SceneDepth) * TransitionScale + 1);
+		float TransitionScale = Settings.TransitionScale;
+
+		// Unoptimized Math: saturate((ShadowmapDepth - Settings.SceneDepth) * TransitionScale + 1);
+		// Rearranged the math so that per pixel constants can be optimized from per sample constants.
+		float ConstantFactor = (Settings.SceneDepth * TransitionScale - 1);
+		float4 ShadowFactor = saturate(ShadowmapDepth * TransitionScale - ConstantFactor);
 
 		FLATTEN
 		if (Settings.bTreatMaxDepthUnshadowed)
@@ -194,8 +198,12 @@ float3 CalculateOcclusion(float3 ShadowmapDepth, FPCFSamplerSettings Settings)
 		// The standard comparison is Settings.SceneDepth < ShadowmapDepth
 		// Using a soft transition based on depth difference
 		// Offsets shadows a bit but reduces self shadowing artifacts considerably
-		float TransitionScale = Settings.TransitionScale; //SoftTransitionScale.z;
-		float3 ShadowFactor = saturate((ShadowmapDepth - Settings.SceneDepth) * TransitionScale + 1);
+		float TransitionScale = Settings.TransitionScale;
+
+		// Unoptimized Math: saturate((ShadowmapDepth - Settings.SceneDepth) * TransitionScale + 1);
+		// Rearranged the math so that per pixel constants can be optimized from per sample constants.
+		float ConstantFactor = (Settings.SceneDepth * TransitionScale - 1);
+		float3 ShadowFactor = saturate(ShadowmapDepth * TransitionScale - ConstantFactor);
 
 		FLATTEN
 		if (Settings.bTreatMaxDepthUnshadowed)
