Fixed Niagara GPU sorting.

diff --git a/Engine/Shaders/Private/NiagaraSortKeyGen.usf b/Engine/Shaders/Private/NiagaraSortKeyGen.usf
--- a/Engine/Shaders/Private/NiagaraSortKeyGen.usf
+++ b/Engine/Shaders/Private/NiagaraSortKeyGen.usf
@@ -14,6 +14,8 @@
 ------------------------------------------------------------------------------*/
 
 #include "Common.ush"
+
+#define NIAGARA_KEYGEN 1
 #include "NiagaraVFParticleAccess.usf"
 
 Buffer<uint> GPUParticleCountBuffer;
diff --git a/Engine/Shaders/Private/NiagaraVFParticleAccess.usf b/Engine/Shaders/Private/NiagaraVFParticleAccess.usf
--- a/Engine/Shaders/Private/NiagaraVFParticleAccess.usf
+++ b/Engine/Shaders/Private/NiagaraVFParticleAccess.usf
@@ -11,7 +11,11 @@
 // This is a hack to fix cases where depth values are invariant due to compiler optimizations/operation reordering between the depth pass and base pass. (UE-68456)
 #define NIAGARA_BIAS_HACK 1
 
-#if !RAYHITGROUPSHADER && !COMPUTESHADER
+#ifndef NIAGARA_KEYGEN
+#define NIAGARA_KEYGEN 0
+#endif
+
+#if (!RAYHITGROUPSHADER && !COMPUTESHADER) || NIAGARA_KEYGEN
 
 Buffer<float> NiagaraParticleDataFloat;
 uint NiagaraFloatDataOffset;
