Fixed quantization issue on coat material for low and high base color.

Add dithering on the baseColor for coat material, which mitigate the quantization issue due to low variation across.

diff --git a/Engine/Shaders/Private/BasePassPixelShader.usf b/Engine/Shaders/Private/BasePassPixelShader.usf
--- a/Engine/Shaders/Private/BasePassPixelShader.usf
+++ b/Engine/Shaders/Private/BasePassPixelShader.usf
@@ -857,6 +857,7 @@ void FPixelShaderInOut_MainPS(
 	GBuffer.Depth = MaterialParameters.ScreenPosition.w;
 	GBuffer.PrecomputedShadowFactors = GetPrecomputedShadowMasks(Interpolants, MaterialParameters.PrimitiveId, MaterialParameters.AbsoluteWorldPosition, VolumetricLightmapBrickTextureUVs);
 
+	const float Dither = InterleavedGradientNoise(MaterialParameters.SvPosition.xy, View.StateFrameIndexMod8);
 	SetGBufferForShadingModel(
 		GBuffer, 
 		MaterialParameters, 
@@ -866,7 +867,7 @@ void FPixelShaderInOut_MainPS(
 		Specular, 
 		Roughness, 
 		SubsurfaceColor,
-		SubsurfaceProfile);
+		Dither);
 
 #if USES_GBUFFER
 	GBuffer.SelectiveOutputMask = GetSelectiveOutputMask();
diff --git a/Engine/Shaders/Private/ShadingModelsMaterial.ush b/Engine/Shaders/Private/ShadingModelsMaterial.ush
--- a/Engine/Shaders/Private/ShadingModelsMaterial.ush
+++ b/Engine/Shaders/Private/ShadingModelsMaterial.ush
@@ -18,7 +18,8 @@ void SetGBufferForShadingModel(
 	const half  Specular,
 	const float Roughness,
 	const float3 SubsurfaceColor,
-	const float SubsurfaceProfile)
+	const float SubsurfaceProfile,
+	const float dither)
 {
 	GBuffer.WorldNormal = MaterialParameters.WorldNormal;
 	GBuffer.BaseColor = BaseColor;
@@ -86,6 +87,7 @@ void SetGBufferForShadingModel(
 	float LayerAttenuation = lerp( 1, (1 - F), ClearCoat );
 
 	GBuffer.BaseColor = lerp( BaseColor * LayerAttenuation, MetalSpec * Absorption * RefractionScale, Metallic * ClearCoat );
+	GBuffer.BaseColor+= dither / 255.f;
 	GBuffer.Specular *= lerp( 1, RefractionScale, ClearCoat );
 
 	GBuffer.CustomData.x = ClearCoat;
