Fix shader compilation issue introduced in CL 5336422.

CL 5336422 was wrongly resolved against 4.22

diff --git a/Engine/Shaders/Private/BasePassPixelShader.usf b/Engine/Shaders/Private/BasePassPixelShader.usf
--- a/Engine/Shaders/Private/BasePassPixelShader.usf
+++ b/Engine/Shaders/Private/BasePassPixelShader.usf
@@ -857,7 +857,7 @@ void FPixelShaderInOut_MainPS(
 	GBuffer.Depth = MaterialParameters.ScreenPosition.w;
 	GBuffer.PrecomputedShadowFactors = GetPrecomputedShadowMasks(Interpolants, MaterialParameters.PrimitiveId, MaterialParameters.AbsoluteWorldPosition, VolumetricLightmapBrickTextureUVs);
 
-	const float Dither = InterleavedGradientNoise(MaterialParameters.SvPosition.xy, View.StateFrameIndexMod8);
+	const float GBufferDither = InterleavedGradientNoise(MaterialParameters.SvPosition.xy, View.StateFrameIndexMod8);
 	SetGBufferForShadingModel(
 		GBuffer, 
 		MaterialParameters, 
@@ -867,7 +867,7 @@ void FPixelShaderInOut_MainPS(
 		Specular, 
 		Roughness, 
 		SubsurfaceColor,
-		Dither);
+		GBufferDither);
 
 #if USES_GBUFFER
 	GBuffer.SelectiveOutputMask = GetSelectiveOutputMask();
diff --git a/Engine/Shaders/Private/ShadingModelsMaterial.ush b/Engine/Shaders/Private/ShadingModelsMaterial.ush
--- a/Engine/Shaders/Private/ShadingModelsMaterial.ush
+++ b/Engine/Shaders/Private/ShadingModelsMaterial.ush
@@ -18,7 +18,6 @@ void SetGBufferForShadingModel(
 	const half  Specular,
 	const float Roughness,
 	const float3 SubsurfaceColor,
-	const float SubsurfaceProfile,
 	const float dither)
 {
 	GBuffer.WorldNormal = MaterialParameters.WorldNormal;
