Add barn door support for stochastic RectLight evaluation and path tracer.

diff --git a/Engine/Shaders/Private/DeferredLightingCommon.ush b/Engine/Shaders/Private/DeferredLightingCommon.ush
--- a/Engine/Shaders/Private/DeferredLightingCommon.ush
+++ b/Engine/Shaders/Private/DeferredLightingCommon.ush
@@ -244,95 +244,18 @@ float GetLocalLightAttenuation(
 }
 
 #define RECLIGHT_BARNDOOR 1
+// Wrapper for FDeferredLightData for computing visible rect light (i.e., unoccluded by barn doors)
 FRect GetRect(float3 ToLight, FDeferredLightData LightData)
 {
-	// Is blocked by barn doors
-	FRect Rect;
-	Rect.Origin = ToLight;
-	Rect.Axis[1] = LightData.Tangent;
-	Rect.Axis[2] = LightData.Direction;
-	Rect.Axis[0] = cross( Rect.Axis[1], Rect.Axis[2] );
-	Rect.Extent = float2(LightData.SourceRadius, LightData.SourceLength);
-	Rect.FullExtent = Rect.Extent;
-	Rect.Offset = 0;
-
-#if RECLIGHT_BARNDOOR == 1
-	// Compute the visible rectangle from the current shading point. 
-	// The new rectangle will have reduced width/height, and a shifted origin
-	//
-	// Common setup for occlusion computation
-	// Notes: Barn angle & length are identical for all sides
-	//					D_B
-	//					<-->
-	//							D_S
-	//		            <--------------->     O         +X
-	//					   -------------.--------------->------       ^
-	//					  /        .          |                \      |
-	//					 /   .                |					\     |  BarnDepth
-	//				   ./                     |					 \    v
-	//		     .      C                     v +Z
-	//		.  
-	// .
-	// S
-	//
-	// Only compute the occluded rect if the barn door has an angle less 
-	// than 88 degrees
-	if (LightData.RectLightBarnCosAngle > 0.035f)
-	{	
-		const float3 LightNormal	=-Rect.Axis[2];
-		const float3 LightdPdv		=-Rect.Axis[1];
-		const float3 LightdPdu		=-Rect.Axis[0];
-		const float2 LightExtent	= float2(LightData.SourceRadius, LightData.SourceLength);
-		const float BarnLength		= LightData.RectLightBarnLength;
-	
-		// Project shading point S into light space
-    	const float3 S_World = -ToLight;
-        float3 S_Light = float3(dot(S_World, LightdPdu), dot(S_World, LightdPdv), dot(S_World, LightNormal));
-		
-		// Compute barn door projection (D_B). Clamp the projection to the shading point if it is closer 
-		// to the light than the actual barn door
-		// Theta is the angle between the Z axis and the barn door
-		const float CosTheta = LightData.RectLightBarnCosAngle;
-        const float SinTheta = sqrt(1 - max(0, CosTheta * CosTheta));
-        const float BarnDepth = min(S_Light.z, CosTheta * BarnLength);
-        const float S_ratio = BarnDepth / (CosTheta * BarnLength);
-		const float D_B	= SinTheta * BarnLength * S_ratio;
-		
-		// Clamp shading point onto the closest edge, if it is inside the rect light
-		const float2 SignS = sign(S_Light.xy);
-		S_Light.xy = float2(SignS.x, SignS.y) * max(abs(S_Light.xy), LightExtent + D_B.xx);
-		
-		// Compute the closest rect lignt corner, offset by the barn door size
-        const float3 C = float3(SignS.x * (LightExtent.x + D_B), SignS.y * (LightExtent.y + D_B), BarnDepth);
-			
-		// Compute projected distance (D_S) of barn door onto the rect light
-		// Eta is the angle between the Z axis and the direction vector (S-C)
-		const float3 SProj  = S_Light - C;
-		const float  CosEta = max(SProj.z, 0);
-		const float2 SinEta = abs(SProj.xy);
-		const float2 TanEta	= CosEta > 0.001f ? abs(SProj.xy)/CosEta : 10000.f;
-        const float2 D_S = BarnDepth * TanEta;
-
-		// Equivalent to (e.g., X axis):
-		//  if (SignS.x < 0) MinMaxX.x += D_S - D_B;
-		//  if (SignS.x > 0) MinMaxX.y -= D_S - D_B;
-		float2 MinMaxX = float2(-LightExtent.x, LightExtent.x);
-		float2 MinMaxY = float2(-LightExtent.y, LightExtent.y);
-		MinMaxX += (D_S.x - D_B) * float2(max(0, -SignS.x), -max(0, SignS.x));
-		MinMaxY += (D_S.y - D_B) * float2(max(0, -SignS.y), -max(0, SignS.y));					
-		MinMaxX = clamp(MinMaxX, -LightExtent.x, LightExtent.x);
-		MinMaxY = clamp(MinMaxY, -LightExtent.y, LightExtent.y);
-		
-        const float2 RectOffset = 0.5f * float2(MinMaxX.x + MinMaxX.y, MinMaxY.x + MinMaxY.y);
-
-        Rect.Extent = 0.5f * float2(MinMaxX.y - MinMaxX.x, MinMaxY.y - MinMaxY.x);
-        Rect.Origin = Rect.Origin +  LightdPdu * RectOffset.x + LightdPdv * RectOffset.y;
-		Rect.Offset = -RectOffset;
-		Rect.FullExtent = LightExtent;
-    }
-#endif
-
-	return Rect;
+	return GetRect(
+		ToLight, 
+		LightData.Direction, 
+		LightData.Tangent, 
+		LightData.SourceRadius, 
+		LightData.SourceLength, 
+		LightData.RectLightBarnCosAngle, 
+		LightData.RectLightBarnLength,
+		RECLIGHT_BARNDOOR);
 }
 
 FCapsuleLight GetCapsule( float3 ToLight, FDeferredLightData LightData )
diff --git a/Engine/Shaders/Private/PathTracing/Light/PathTracingLightCommon.ush b/Engine/Shaders/Private/PathTracing/Light/PathTracingLightCommon.ush
--- a/Engine/Shaders/Private/PathTracing/Light/PathTracingLightCommon.ush
+++ b/Engine/Shaders/Private/PathTracing/Light/PathTracingLightCommon.ush
@@ -107,6 +107,16 @@ float GetAttenuation(int LightId)
 	return SceneLightsData.Attenuation[LightId];
 }
 
+float GetRectLightBarnCosAngle(int LightId)
+{
+	return SceneLightsData.RectLightBarnCosAngle[LightId];
+}
+
+float GetRectLightBarnLength(int LightId)
+{
+	return SceneLightsData.RectLightBarnLength[LightId];
+}
+
 float ComputeAttenuationFalloff(float Distance, float Attenuation)
 {
 #ifdef USE_ATTENUATION_TERM
diff --git a/Engine/Shaders/Private/PathTracing/Light/PathTracingRectLight.ush b/Engine/Shaders/Private/PathTracing/Light/PathTracingRectLight.ush
--- a/Engine/Shaders/Private/PathTracing/Light/PathTracingRectLight.ush
+++ b/Engine/Shaders/Private/PathTracing/Light/PathTracingRectLight.ush
@@ -6,6 +6,7 @@
 
 #pragma once
 
+#include "../../RectLight.ush"
 #include "PathTracingLightCommon.ush"
 
 void RectLight_GenerateLightRay(
@@ -95,6 +96,24 @@ void RectLight_SampleLight(
 	float3 LightdPdv = GetdPdv(LightId);
 	float LightWidth = GetWidth(LightId);
 	float LightHeight = GetHeight(LightId);
+	
+	// Compute the visible portion of the light from the current shading point
+	{
+		float3 ToLight = LightPosition - WorldPosition;
+		FRect VisibleRect = GetRect(
+				ToLight,
+				-LightNormal,
+				LightdPdv,
+				LightWidth*0.5f,
+				LightHeight*0.5f,
+				GetRectLightBarnCosAngle(LightId),
+				GetRectLightBarnLength(LightId),
+				true);
+
+		LightPosition = LightPosition - LightdPdu * VisibleRect.Offset.x - LightdPdv * VisibleRect.Offset.y;	
+		LightWidth	  = VisibleRect.Extent.x * 2.0f;
+		LightHeight   = VisibleRect.Extent.y * 2.0f;
+	}
 
 	// Pick a random position
 	float2 UV = RandSample.yz - 0.5;
@@ -125,7 +144,7 @@ void RectLight_SampleLight(
 	Rect.Axis[0] = LightdPdu;
 	Rect.Axis[1] = LightdPdv;
 	Rect.Axis[2] = cross(LightdPdu, LightdPdv);
-	Rect.Extent = float2(LightWidth, LightHeight);
+	Rect.Extent = float2(LightWidth, LightHeight) * 0.5f;
 	FSphericalRect SphericalRect = BuildSphericalRect(Rect);
 	OutPdf = isfinite(SphericalRect.SolidAngle) ? 1.0 / SphericalRect.SolidAngle : 0.0;
 }
\ No newline at end of file
diff --git a/Engine/Shaders/Private/RayTracing/RayTracingRectLightRGS.usf b/Engine/Shaders/Private/RayTracing/RayTracingRectLightRGS.usf
--- a/Engine/Shaders/Private/RayTracing/RayTracingRectLightRGS.usf
+++ b/Engine/Shaders/Private/RayTracing/RayTracingRectLightRGS.usf
@@ -212,16 +212,41 @@ void RectLightRGS()
 	float3 LightdPdv = RectLight.dPdv;
 	float LightWidth = RectLight.Width;
 	float LightHeight = RectLight.Height;
-	
+	float LightBarnCosAngle = RectLight.BarnCosAngle;
+	float LightBarnLength = RectLight.BarnLength;
+
+	// Compute the visible portion of the light from the current shading point
+	{
+		float3 ToLight = LightPosition - WorldPosition;
+		FRect VisibleRect = GetRect(
+				ToLight,
+				-LightNormal,
+				LightdPdv,
+				LightWidth*0.5f,
+				LightHeight*0.5f,
+				RectLight.BarnCosAngle,
+				RectLight.BarnLength,
+				true);
+
+		LightPosition = LightPosition - LightdPdu * VisibleRect.Offset.x - LightdPdv * VisibleRect.Offset.y;	
+		LightWidth	  = VisibleRect.Extent.x * 2.0f;
+		LightHeight   = VisibleRect.Extent.y * 2.0f;
+	}
+
 	// Encode the rect light  parameter into the light parameters struct
 	// todo: send direction the light shader parameter data instead of the custon RectLight data.
 	FLightShaderParameters LightParameters;
-	LightParameters.Position = RectLight.Position;
-	LightParameters.Direction = -RectLight.Normal;
-	LightParameters.Tangent = RectLight.dPdu;
-	LightParameters.SourceRadius = 0.5f * RectLight.Width;
-	LightParameters.SourceLength = 0.5f * RectLight.Height;
-	
+	LightParameters.Position = LightPosition;
+	LightParameters.Direction = -LightNormal;
+	LightParameters.Tangent = LightdPdu;
+	LightParameters.SourceRadius = 0.5f * LightWidth;
+	LightParameters.SourceLength = 0.5f * LightHeight;
+	LightParameters.RectLightBarnCosAngle =  LightBarnCosAngle;
+	LightParameters.RectLightBarnLength =  LightBarnLength;
+
+	if (LightWidth == 0 || LightHeight == 0) 
+		SamplesPerPixel = 0;
+
 	float3 ExitantRadiance = 0.0;
 	float3 DiffuseExitantRadiance = 0.0;
 	for (uint SampleIndex = 0; SampleIndex < SamplesPerPixel; ++SampleIndex)
@@ -281,10 +306,10 @@ void RectLightRGS()
 		else
 		{
 			// Sample light texture
-			float3 LightPosition = Ray.Origin + Ray.Direction * Ray.TMax;
-			float3 RectLightCorner = RectLight.Position - 0.5 * (RectLight.dPdu * RectLight.Width + RectLight.dPdv * RectLight.Height);
-			float3 DeltaPosition = LightPosition - RectLightCorner;
-			float2 SampleUV = float2(dot(DeltaPosition, RectLight.dPdu) / RectLight.Width, dot(DeltaPosition, RectLight.dPdv) / RectLight.Height);
+			float3 SampleLightPosition = Ray.Origin + Ray.Direction * Ray.TMax;
+			float3 RectLightCorner = LightPosition - 0.5 * (LightdPdu * LightWidth + LightdPdv * LightHeight);
+			float3 DeltaPosition = SampleLightPosition - RectLightCorner;
+			float2 SampleUV = float2(dot(DeltaPosition, LightdPdu) / LightWidth, dot(DeltaPosition, LightdPdv) / LightHeight);
 			SampleUV = 1.0 - SampleUV;
 			float3 TextureColor = Texture2DSample(RectLight.Texture, RectLight.TextureSampler, SampleUV).xyz;
 			float3 IncomingRadiance = RectLight.Color * TextureColor;
diff --git a/Engine/Shaders/Private/RectLight.ush b/Engine/Shaders/Private/RectLight.ush
--- a/Engine/Shaders/Private/RectLight.ush
+++ b/Engine/Shaders/Private/RectLight.ush
@@ -1,7 +1,7 @@
 // Copyright 1998-2019 Epic Games, Inc. All Rights Reserved.
 //
-
-#pragma once
+#ifndef __RECTLIGHT__
+#define __RECTLIGHT__
 
 #include "CapsuleLight.ush"
 #include "MonteCarlo.ush"
@@ -627,3 +627,102 @@ float3 RectIrradianceRef( float3 N, FRect Rect, out float BaseIrradiance, out fl
 
 	return normalize(L);
 }
+
+FRect GetRect(
+	float3 ToLight, 
+	float3 LightDataDirection, 
+	float3 LightDataTangent, 
+	float LightDataSourceRadius, 
+	float LightDataSourceLength, 
+	float LightDataRectLightBarnCosAngle, 
+	float LightDataRectLightBarnLength,
+	bool bComputeVisibleRect)
+{
+	// Is blocked by barn doors
+	FRect Rect;
+	Rect.Origin = ToLight;
+	Rect.Axis[1] = LightDataTangent;
+	Rect.Axis[2] = LightDataDirection;
+	Rect.Axis[0] = cross( Rect.Axis[1], Rect.Axis[2] );
+	Rect.Extent = float2(LightDataSourceRadius, LightDataSourceLength);
+	Rect.FullExtent = Rect.Extent;
+	Rect.Offset = 0;
+
+	// Compute the visible rectangle from the current shading point. 
+	// The new rectangle will have reduced width/height, and a shifted origin
+	//
+	// Common setup for occlusion computation
+	// Notes: Barn angle & length are identical for all sides
+	//					D_B
+	//					<-->
+	//							D_S
+	//						<--------------->     O         +X
+	//					   -------------.--------------->------       ^
+	//					  /        .          |                \      |
+	//					 /   .                |					\     |  BarnDepth
+	//				   ./                     |					 \    v
+	//		     .      C                     v +Z
+	//		.  
+	// .
+	// S
+	//
+	// Only compute the occluded rect if the barn door has an angle less 
+	// than 88 degrees
+	if (bComputeVisibleRect && LightDataRectLightBarnCosAngle > 0.035f)
+	{	
+		const float3 LightNormal	=-Rect.Axis[2];
+		const float3 LightdPdv		=-Rect.Axis[1];
+		const float3 LightdPdu		=-Rect.Axis[0];
+		const float2 LightExtent	= float2(LightDataSourceRadius, LightDataSourceLength);
+		const float BarnLength		= LightDataRectLightBarnLength;
+	
+		// Project shading point S into light space
+		const float3 S_World = -ToLight;
+		float3 S_Light = float3(dot(S_World, LightdPdu), dot(S_World, LightdPdv), dot(S_World, LightNormal));
+		
+		// Compute barn door projection (D_B). Clamp the projection to the shading point if it is closer 
+		// to the light than the actual barn door
+		// Theta is the angle between the Z axis and the barn door
+		const float CosTheta = LightDataRectLightBarnCosAngle;
+		const float SinTheta = sqrt(1 - max(0, CosTheta * CosTheta));
+		const float BarnDepth = min(S_Light.z, CosTheta * BarnLength);
+		const float S_ratio = BarnDepth / (CosTheta * BarnLength);
+		const float D_B	= SinTheta * BarnLength * S_ratio;
+		
+		// Clamp shading point onto the closest edge, if it is inside the rect light
+		const float2 SignS = sign(S_Light.xy);
+		S_Light.xy = float2(SignS.x, SignS.y) * max(abs(S_Light.xy), LightExtent + D_B.xx);
+		
+		// Compute the closest rect lignt corner, offset by the barn door size
+		const float3 C = float3(SignS.x * (LightExtent.x + D_B), SignS.y * (LightExtent.y + D_B), BarnDepth);
+			
+		// Compute projected distance (D_S) of barn door onto the rect light
+		// Eta is the angle between the Z axis and the direction vector (S-C)
+		const float3 SProj  = S_Light - C;
+		const float  CosEta = max(SProj.z, 0);
+		const float2 SinEta = abs(SProj.xy);
+		const float2 TanEta	= CosEta > 0.001f ? abs(SProj.xy)/CosEta : 10000.f;
+		const float2 D_S = BarnDepth * TanEta;
+
+		// Equivalent to (e.g., X axis):
+		//  if (SignS.x < 0) MinMaxX.x += D_S - D_B;
+		//  if (SignS.x > 0) MinMaxX.y -= D_S - D_B;
+		float2 MinMaxX = float2(-LightExtent.x, LightExtent.x);
+		float2 MinMaxY = float2(-LightExtent.y, LightExtent.y);
+		MinMaxX += (D_S.x - D_B) * float2(max(0, -SignS.x), -max(0, SignS.x));
+		MinMaxY += (D_S.y - D_B) * float2(max(0, -SignS.y), -max(0, SignS.y));
+		MinMaxX = clamp(MinMaxX, -LightExtent.x, LightExtent.x);
+		MinMaxY = clamp(MinMaxY, -LightExtent.y, LightExtent.y);
+		
+		const float2 RectOffset = 0.5f * float2(MinMaxX.x + MinMaxX.y, MinMaxY.x + MinMaxY.y);
+
+		Rect.Extent = 0.5f * float2(MinMaxX.y - MinMaxX.x, MinMaxY.y - MinMaxY.x);
+		Rect.Origin = Rect.Origin +  LightdPdu * RectOffset.x + LightdPdv * RectOffset.y;
+		Rect.Offset = -RectOffset;
+		Rect.FullExtent = LightExtent;
+	}
+
+	return Rect;
+}
+
+#endif // __RECTLIGHT__
diff --git a/Engine/Source/Runtime/Renderer/Private/PathTracing.cpp b/Engine/Source/Runtime/Renderer/Private/PathTracing.cpp
--- a/Engine/Source/Runtime/Renderer/Private/PathTracing.cpp
+++ b/Engine/Source/Runtime/Renderer/Private/PathTracing.cpp
@@ -201,6 +201,8 @@ public:
 						LightData.Color[LightData.Count] = LightParameters.Color / 4.0;
 						LightData.Dimensions[LightData.Count] = FVector(2.0f * LightParameters.SourceRadius, 2.0f * LightParameters.SourceLength, 0.0f);
 						LightData.Attenuation[LightData.Count] = 1.0 / LightParameters.InvRadius;
+						LightData.RectLightBarnCosAngle[LightData.Count] = LightParameters.RectLightBarnCosAngle;
+						LightData.RectLightBarnLength[LightData.Count] = LightParameters.RectLightBarnLength;
 						break;
 					}
 					case LightType_Spot:
diff --git a/Engine/Source/Runtime/Renderer/Private/PathTracingUniformBuffers.h b/Engine/Source/Runtime/Renderer/Private/PathTracingUniformBuffers.h
--- a/Engine/Source/Runtime/Renderer/Private/PathTracingUniformBuffers.h
+++ b/Engine/Source/Runtime/Renderer/Private/PathTracingUniformBuffers.h
@@ -29,6 +29,8 @@ BEGIN_GLOBAL_SHADER_PARAMETER_STRUCT(FPathTracingLightData, )
 	// Light-specific
 	SHADER_PARAMETER_ARRAY(FVector, Dimensions, [GLightCountMaximum])
 	SHADER_PARAMETER_ARRAY(float, Attenuation, [GLightCountMaximum])
+	SHADER_PARAMETER_ARRAY(float, RectLightBarnCosAngle, [GLightCountMaximum])
+	SHADER_PARAMETER_ARRAY(float, RectLightBarnLength, [GLightCountMaximum])
 END_GLOBAL_SHADER_PARAMETER_STRUCT()
 
 
