Fixed wrongly accessed values when building importance sampling mip hierarchy due to float precision.
What propagating living mushrooms on the screen during path tracing :)

diff --git a/Engine/Shaders/Private/PathTracing/BuildVarianceMipTreeComputeShader.usf b/Engine/Shaders/Private/PathTracing/BuildVarianceMipTreeComputeShader.usf
--- a/Engine/Shaders/Private/PathTracing/BuildVarianceMipTreeComputeShader.usf
+++ b/Engine/Shaders/Private/PathTracing/BuildVarianceMipTreeComputeShader.usf
@@ -101,7 +101,7 @@ void BuildVarianceMipTreeCS(
 {
 	uint3 GridCoordinate = DispatchThreadId;
 	uint2 Dimensions = DimensionsAtLevel(MipLevel);
-	if (GridCoordinate.x < 0 || GridCoordinate.x >= Dimensions.x || GridCoordinate.y < 0 || GridCoordinate.y >= Dimensions.y)
+	if (GridCoordinate.x >= Dimensions.x || GridCoordinate.y >= Dimensions.y)
 	{
 		return;
 	}
@@ -109,13 +109,13 @@ void BuildVarianceMipTreeCS(
 	float Sum = 0.0;
 	if (MipLevel == 0)
 	{
-		float2 UV = float2(GridCoordinate.xy) / VarianceMapDimensions;
+		float2 UV = float2(GridCoordinate.xy + 0.5f) / VarianceMapDimensions;
 
 		// Resample blurred result as the variance map
 		Gaussian5x5 Filter;
 		Filter.Initialize();
 		uint Samples = 0;
-		int2 PixelCenter = UV * ViewSize;
+		int2 PixelCenter = DispatchThreadId.xy;
 		int2 PixelMin = max(PixelCenter - Filter.Offset, 0);
 		int2 PixelMax = min(PixelCenter + Filter.Offset, ViewSize - 1);
 		for (int y = PixelMin.y; y <= PixelMax.y; ++y)
