Fix shader compile error.

diff --git a/Engine/Shaders/Private/BasePassPixelShader.usf b/Engine/Shaders/Private/BasePassPixelShader.usf
--- a/Engine/Shaders/Private/BasePassPixelShader.usf
+++ b/Engine/Shaders/Private/BasePassPixelShader.usf
@@ -867,6 +867,7 @@ void FPixelShaderInOut_MainPS(
 		Specular, 
 		Roughness, 
 		SubsurfaceColor,
+		SubsurfaceProfile,
 		GBufferDither);
 
 #if USES_GBUFFER
diff --git a/Engine/Shaders/Private/ShadingModelsMaterial.ush b/Engine/Shaders/Private/ShadingModelsMaterial.ush
--- a/Engine/Shaders/Private/ShadingModelsMaterial.ush
+++ b/Engine/Shaders/Private/ShadingModelsMaterial.ush
@@ -18,6 +18,7 @@ void SetGBufferForShadingModel(
 	const half  Specular,
 	const float Roughness,
 	const float3 SubsurfaceColor,
+	const float SubsurfaceProfile,
 	const float dither)
 {
 	GBuffer.WorldNormal = MaterialParameters.WorldNormal;
@@ -183,4 +184,4 @@ void SetGBufferForShadingModel(
 #else
 	// missing shading model, compiler should report ShadingModelID is not set
 #endif
-}
\ No newline at end of file
+}
