Fix shader compilation warning.

diff --git a/Engine/Shaders/Private/PathTracing/Light/PathTracingLightSampling.ush b/Engine/Shaders/Private/PathTracing/Light/PathTracingLightSampling.ush
--- a/Engine/Shaders/Private/PathTracing/Light/PathTracingLightSampling.ush
+++ b/Engine/Shaders/Private/PathTracing/Light/PathTracingLightSampling.ush
@@ -90,17 +90,21 @@ bool SampleLightSelection(
 	// Build irradiance estimate prefix sum
 	float IrradianceEstimate[LIGHT_COUNT_MAXIMUM];
 	EstimateLight(0, Ray, Payload, IrradianceEstimate[0]);
-	for (uint LightIndex = 1; LightIndex < SceneLightsData.Count; ++LightIndex)
 	{
-		EstimateLight(LightIndex, Ray, Payload, IrradianceEstimate[LightIndex]);
-		IrradianceEstimate[LightIndex] += IrradianceEstimate[LightIndex - 1];
+		for (uint LightIndex = 1; LightIndex < SceneLightsData.Count; ++LightIndex)
+		{
+			EstimateLight(LightIndex, Ray, Payload, IrradianceEstimate[LightIndex]);
+			IrradianceEstimate[LightIndex] += IrradianceEstimate[LightIndex - 1];
+		}
 	}
 	// dxr_todo: Consider when IrradianceEstimate[SceneLightsData.Count] == 0
 
 	// CDF normalization
-	for (uint LightIndex = 0; LightIndex < SceneLightsData.Count; ++LightIndex)
 	{
-		IrradianceEstimate[LightIndex] /= IrradianceEstimate[SceneLightsData.Count - 1];
+		for (uint LightIndex = 0; LightIndex < SceneLightsData.Count; ++LightIndex)
+		{
+			IrradianceEstimate[LightIndex] /= IrradianceEstimate[SceneLightsData.Count - 1];
+		}
 	}
 	IrradianceEstimate[SceneLightsData.Count - 1] = 1.0;
 
@@ -138,17 +142,21 @@ void PdfLightSelection(
 	// Build irradiance estimate prefix sum
 	float IrradianceEstimate[LIGHT_COUNT_MAXIMUM];
 	EstimateLight(0, Ray, Payload, IrradianceEstimate[0]);
-	for (uint LightIndex = 1; LightIndex < SceneLightsData.Count; ++LightIndex)
 	{
-		EstimateLight(LightIndex, Ray, Payload, IrradianceEstimate[LightIndex]);
-		IrradianceEstimate[LightIndex] += IrradianceEstimate[LightIndex - 1];
+		for (uint LightIndex = 1; LightIndex < SceneLightsData.Count; ++LightIndex)
+		{
+			EstimateLight(LightIndex, Ray, Payload, IrradianceEstimate[LightIndex]);
+			IrradianceEstimate[LightIndex] += IrradianceEstimate[LightIndex - 1];
+		}
 	}
 	// dxr_todo: Consider when IrradianceEstimate[SceneLightsData.Count] == 0
 
 	// CDF normalization
-	for (uint LightIndex = 0; LightIndex < SceneLightsData.Count; ++LightIndex)
 	{
-		IrradianceEstimate[LightIndex] /= IrradianceEstimate[SceneLightsData.Count - 1];
+		for (uint LightIndex = 0; LightIndex < SceneLightsData.Count; ++LightIndex)
+		{
+			IrradianceEstimate[LightIndex] /= IrradianceEstimate[SceneLightsData.Count - 1];
+		}
 	}
 	IrradianceEstimate[SceneLightsData.Count - 1] = 1.0;
 
diff --git a/Engine/Shaders/Private/PathTracing/Light/PathTracingSkyLight.ush b/Engine/Shaders/Private/PathTracing/Light/PathTracingSkyLight.ush
--- a/Engine/Shaders/Private/PathTracing/Light/PathTracingSkyLight.ush
+++ b/Engine/Shaders/Private/PathTracing/Light/PathTracingSkyLight.ush
@@ -76,7 +76,7 @@ void SkyLight_EstimateLight(
 #else
 	uint Width, Height, NumberOfLevels;
 	SkyLight.Texture.GetDimensions(0, Width, Height, NumberOfLevels);
-	float LightPower = Luminance(GetColor(LightId) * TextureCubeSampleLevel(SkyLight.Texture, SkyLight.TextureSampler, float3(0.0, 1.0, 0.0), NumberOfLevels - 1));
+	float LightPower = Luminance(GetColor(LightId) * TextureCubeSampleLevel(SkyLight.Texture, SkyLight.TextureSampler, float3(0.0, 1.0, 0.0), NumberOfLevels - 1).xyz);
 	OutIrradiance = LightPower * 2.0 * PI;
 #endif
 }
@@ -150,7 +150,7 @@ void SkyLight_SampleLight(
 	//MipPdf = PdfSkyLightOcclusionRay(WorldNormal, RayDirection);
 
 	float2 PixelOffset = RandSample.yz;
-	float2 MipPixelUV = (MipPixel + PixelOffset) / float2(DimensionsAtLevel(SkyLight.MipDimensions, StopLevel));
+	float2 MipPixelUV = (MipPixel + PixelOffset) / float2(DimensionsAtLevel(SkyLight.MipDimensions.xy, StopLevel));
 
 	// Construct Cubemap coordinate
 	uint2 TextureDimensions;
