Fix ray position bias for transmissive material.

diff --git a/Engine/Shaders/Private/PathTracing/Material/PathTracingMaterialSampling.ush b/Engine/Shaders/Private/PathTracing/Material/PathTracingMaterialSampling.ush
--- a/Engine/Shaders/Private/PathTracing/Material/PathTracingMaterialSampling.ush
+++ b/Engine/Shaders/Private/PathTracing/Material/PathTracingMaterialSampling.ush
@@ -35,8 +35,10 @@ void SampleMaterial(
 	float4 RandSample,
 	out float3 OutDirection,
 	out float3 OutThroughput,
-	out float OutPdf)
+	out float OutPdf,
+	out bool bIsTransmitted)
 {
+	bIsTransmitted = false;
 	if (IsPureSpecularReflection(Payload))
 	{
 		SpecularReflection_SampleMaterial(RayDirection, Payload, OutDirection, OutThroughput, OutPdf);
@@ -44,6 +46,7 @@ void SampleMaterial(
 	else if (IsPureSpecularTransmission(Payload))
 	{
 		SpecularRefraction_SampleMaterial(RayDirection, Payload, OutDirection, OutThroughput, OutPdf);
+		bIsTransmitted = true;
 	}
 	else if(IsLambert(Payload))
 	{
diff --git a/Engine/Shaders/Private/PathTracing/PathTracing.usf b/Engine/Shaders/Private/PathTracing/PathTracing.usf
--- a/Engine/Shaders/Private/PathTracing/PathTracing.usf
+++ b/Engine/Shaders/Private/PathTracing/PathTracing.usf
@@ -248,7 +248,8 @@ void PathTracingMainRG()
 			float3 Direction;
 			float3 Throughput;
 			float4 RandSample = RandomSequence_GenerateSample4D(RandSequence, SampleIndex);
-			SampleMaterial(Ray.Direction, RayHitInfo, RandSample, Direction, Throughput, MaterialPdf);
+			bool bIsTransmitted;
+			SampleMaterial(Ray.Direction, RayHitInfo, RandSample, Direction, Throughput, MaterialPdf, bIsTransmitted);
 			if (MaterialPdf == 0.0)
 			{
 				break;
@@ -256,7 +257,7 @@ void PathTracingMainRG()
 
 			// Update ray
 			Ray.Direction = Direction;
-			ApplyPositionBias(Ray, Payload.WorldNormal, AdaptiveSamplingData.MaxNormalBias);
+			ApplyPositionBias(Ray, bIsTransmitted ? -Payload.WorldNormal : Payload.WorldNormal, AdaptiveSamplingData.MaxNormalBias);
 			RayThroughput *= Throughput;
 			RayPdf *= MaterialPdf;
 
