Added translucency refraction option.
When r.RayTracing.Translucency.Refraction 0 : refraction ray will always go forward unperturbed. the scene back color will always be fetched.

diff --git a/Engine/Shaders/Private/RayTracing/RayTracingTranslucency.usf b/Engine/Shaders/Private/RayTracing/RayTracingTranslucency.usf
--- a/Engine/Shaders/Private/RayTracing/RayTracingTranslucency.usf
+++ b/Engine/Shaders/Private/RayTracing/RayTracingTranslucency.usf
@@ -47,6 +47,7 @@ uint UpscaleFactor;
 float TranslucencyMinRayDistance;
 float TranslucencyMaxRayDistance;
 float TranslucencyMaxRoughness;
+uint  TranslucencyRefraction;
 float MaxNormalBias;
 
 RaytracingAccelerationStructure TLAS;
@@ -239,19 +240,22 @@ void RayTracingTranslucencyRGS()
 		}
 
 		// Set refraction ray for next iteration
-		float Ior1 = DielectricF0ToIor(DielectricSpecularToF0(Payload.Specular)); // Not using Payload.Ior but parameterisation from specular.
-		float Ior2 = LastIor;
-		bHasScattered |= Ior1 != Ior2;
-		float3 RefractedDirection;
-		float ReflectionRefractionEventThroughput;
-		if (!RefractRay(Ray.Direction, Payload.WorldNormal, Ior1, Ior2, RefractedDirection, ReflectionRefractionEventThroughput))
+		float3 RefractedDirection = Ray.Direction;
+		if (TranslucencyRefraction)
 		{
-			//#dxr_todo: handle total internal reflection ray event. We could jsuit continue form this path. And not use the new IOR.
-			// We should combine that step with the one above reflection evaluation. For now we stop tracing.
-			break;
+			float Ior1 = DielectricF0ToIor(DielectricSpecularToF0(Payload.Specular)); // Not using Payload.Ior but parameterisation from specular.
+			float Ior2 = LastIor;
+			bHasScattered |= Ior1 != Ior2;
+			float ReflectionRefractionEventThroughput;
+			if (!RefractRay(Ray.Direction, Payload.WorldNormal, Ior1, Ior2, RefractedDirection, ReflectionRefractionEventThroughput))
+			{
+				//#dxr_todo: handle total internal reflection ray event. We could jsuit continue form this path. And not use the new IOR.
+				// We should combine that step with the one above reflection evaluation. For now we stop tracing.
+				break;
+			}
+			LastIor = Ior1;
+			PathThroughput *= ReflectionRefractionEventThroughput;
 		}
-		LastIor = Ior1;
-		PathThroughput *= ReflectionRefractionEventThroughput;
 
 
 		//
diff --git a/Engine/Source/Runtime/Renderer/Private/RayTracing/RayTracingTranslucency.cpp b/Engine/Source/Runtime/Renderer/Private/RayTracing/RayTracingTranslucency.cpp
--- a/Engine/Source/Runtime/Renderer/Private/RayTracing/RayTracingTranslucency.cpp
+++ b/Engine/Source/Runtime/Renderer/Private/RayTracing/RayTracingTranslucency.cpp
@@ -87,6 +87,12 @@ static FAutoConsoleVariableRef CVarRayTracingTranslucencyHeightFog(
 	GRayTracingTranslucencyHeightFog,
 	TEXT("Enables height fog in ray traced Translucency (default = 1)"));
 
+static int32 GRayTracingTranslucencyRefraction = 1;
+static FAutoConsoleVariableRef CVarRayTracingTranslucencyRefraction(
+	TEXT("r.RayTracing.Translucency.Refraction"),
+	GRayTracingTranslucencyRefraction,
+	TEXT("Enables refraction in ray traced Translucency (default = 1)"));
+
 DECLARE_GPU_STAT_NAMED(RayTracingTranslucency, TEXT("Ray Tracing Translucency"));
 
 class FRayTracingTranslucencyRGS : public FGlobalShader
@@ -108,6 +114,7 @@ class FRayTracingTranslucencyRGS : public FGlobalShader
 		SHADER_PARAMETER(float, TranslucencyMinRayDistance)
 		SHADER_PARAMETER(float, TranslucencyMaxRayDistance)
 		SHADER_PARAMETER(float, TranslucencyMaxRoughness)
+		SHADER_PARAMETER(int32, TranslucencyRefraction)
 		SHADER_PARAMETER(float, MaxNormalBias)
 
 		SHADER_PARAMETER_SRV(RaytracingAccelerationStructure, TLAS)
@@ -357,6 +364,7 @@ void FDeferredShadingSceneRenderer::RenderRayTracingTranslucencyView(
 	PassParameters->TranslucencyMinRayDistance = FMath::Min(GRayTracingTranslucencyMinRayDistance, GRayTracingTranslucencyMaxRayDistance);
 	PassParameters->TranslucencyMaxRayDistance = GRayTracingTranslucencyMaxRayDistance;
 	PassParameters->TranslucencyMaxRoughness = FMath::Clamp(GRayTracingTranslucencyMaxRoughness >= 0 ? GRayTracingTranslucencyMaxRoughness : View.FinalPostProcessSettings.RayTracingTranslucencyMaxRoughness, 0.01f, 1.0f);
+	PassParameters->TranslucencyRefraction = GRayTracingTranslucencyRefraction;
 	PassParameters->MaxNormalBias = GetRaytracingOcclusionMaxNormalBias();
 
 	PassParameters->TLAS = View.RayTracingScene.RayTracingSceneRHI->GetShaderResourceView();
