* Add reflection on transmittive material in Pathtracer
* Make position bias more robust (original fix in CL 5348390 was not enough for non-water-tight geometry)

diff --git a/Engine/Shaders/Private/PathTracing/Material/PathTracingMaterialSampling.ush b/Engine/Shaders/Private/PathTracing/Material/PathTracingMaterialSampling.ush
--- a/Engine/Shaders/Private/PathTracing/Material/PathTracingMaterialSampling.ush
+++ b/Engine/Shaders/Private/PathTracing/Material/PathTracingMaterialSampling.ush
@@ -36,17 +36,16 @@ void SampleMaterial(
 	out float3 OutDirection,
 	out float3 OutThroughput,
 	out float OutPdf,
-	out bool bIsTransmitted)
+	out float OutPositionBiasSign)
 {
-	bIsTransmitted = false;
+	OutPositionBiasSign = 1;
 	if (IsPureSpecularReflection(Payload))
 	{
 		SpecularReflection_SampleMaterial(RayDirection, Payload, OutDirection, OutThroughput, OutPdf);
 	}
 	else if (IsPureSpecularTransmission(Payload))
 	{
-		SpecularRefraction_SampleMaterial(RayDirection, Payload, OutDirection, OutThroughput, OutPdf);
-		bIsTransmitted = true;
+		SpecularRefraction_SampleMaterial(RayDirection, Payload, RandSample, OutDirection, OutThroughput, OutPdf, OutPositionBiasSign);
 	}
 	else if(IsLambert(Payload))
 	{
diff --git a/Engine/Shaders/Private/PathTracing/Material/PathTracingSpecularTransmission.ush b/Engine/Shaders/Private/PathTracing/Material/PathTracingSpecularTransmission.ush
--- a/Engine/Shaders/Private/PathTracing/Material/PathTracingSpecularTransmission.ush
+++ b/Engine/Shaders/Private/PathTracing/Material/PathTracingSpecularTransmission.ush
@@ -9,46 +9,58 @@ PathTracingSpecularTransmission.usf: Pure specular transmission BTDF sampling fu
 void SpecularRefraction_SampleMaterial(
 	float3 RayDirection,
 	FMaterialClosestHitPayload Payload,
+	float4 RandSample,
 	out float3 OutDirection,
 	out float3 OutThroughput,
-	out float OutPdf)
+	out float OutPdf,
+	out float OutPositionBiasSign)
 {
 	float3 N_World = GetWorldNormal(Payload);
-	float3 SpecularColor = Payload.BaseColor;// GetSpecularColor(Payload);
+	float3 SpecularColor = GetSpecularColor(Payload);
 
 	float3 V_World = normalize(-RayDirection);
-	float3 V_Local = WorldToTangent(V_World, N_World);
-	float3 L_Local;
+	float3 V_Local = WorldToTangent(V_World, N_World);	
+	const bool bIsEntering = V_Local.z >= 0.0;
 
-	float3 N = float3(0.0, 0.0, 1.0);
-	if (dot(V_Local, N) < 0.0)
-	{
-		N = -N;
-	}
+	// Should be an input to this function, using a stack of medium. How to handle non-water tight geometry?
+	const float IorI = 1;
+	const float IorT = Payload.Ior;
 
-	float Ior1 = Payload.Ior;
-	float Ior2 = 1.0f;
-	float Eta = Ior1 / Ior2;
-	bool bIsEntering = (dot(RayDirection, N_World) < 0.0);
-	if (bIsEntering)
-	{
-		Eta = 1.0 / Payload.Ior;
-	}
+	const float3 F = FresnelDielectric(V_Local.z, IorI, IorT);
+    const float Favg = (F.x + F.y + F.z) / 3.0f;
+	const bool bIsReflected = RandSample.x < F.x;
 	
-	if (BSDF_Refract(V_Local, N, Eta, L_Local))
+	float3 L_Local;
+	if (bIsReflected)
 	{
-		OutDirection = TangentToWorld(L_Local, N_World);
+        L_Local = float3(-V_Local.x, -V_Local.y, V_Local.z);
 
-		float SqrRatio = 1.0;// (IorI * IorI) / (IorT * IorT);	// SqrRatio is 1 for light rays	
-		OutThroughput = SqrRatio * SpecularColor * (float3(1.0f, 1.0f, 1.0f) - FresnelDielectric(L_Local.z, Ior2, Ior1)) / abs(L_Local.z);
+		OutThroughput =  F / abs(L_Local.z);
+        OutPdf = 1.0f / Favg;
+		OutPositionBiasSign = bIsEntering ? 1 : -1;
 	}
 	else
 	{
-		OutThroughput = float3(0.0, 0.0, 0.0);
+		const float Eta = bIsEntering ? IorI / IorT : IorT / IorI;
+		const float3 N = float3(0.0, 0.0, bIsEntering ? 1.0f : -1.0f);
+		const bool bIsTIR = !BSDF_Refract(V_Local, N, Eta, L_Local);
+		if (bIsTIR)
+		{
+			OutThroughput = 1;
+			OutPdf = 1;
+			OutPositionBiasSign = bIsEntering ? 1 : -1;
+		}
+		else
+		{
+			float SqrRatio = 1.0; // (IorI * IorI) / (IorT * IorT);	// SqrRatio is 1 for light rays	
+			OutThroughput = SqrRatio * (1 - F) / abs(L_Local.z);
+			OutPdf = 1.0f / (1.0f - Favg);
+			OutPositionBiasSign = bIsEntering ? -1 : 1;
+		}
 	}
 
-	OutPdf = 1.0;
-}
+	OutDirection = TangentToWorld(L_Local, N_World);
+} 
 
 float SpecularTransmission_Pdf()
 {
@@ -59,4 +71,4 @@ void SpecularTransmission_EvalMaterial(out float3 OutThroughput, out float OutPd
 {
 	OutThroughput = float3(0.0, 0.0, 0.0);
 	OutPdf = SpecularTransmission_Pdf();
-}
+}
\ No newline at end of file
diff --git a/Engine/Shaders/Private/PathTracing/PathTracing.usf b/Engine/Shaders/Private/PathTracing/PathTracing.usf
--- a/Engine/Shaders/Private/PathTracing/PathTracing.usf
+++ b/Engine/Shaders/Private/PathTracing/PathTracing.usf
@@ -248,8 +248,8 @@ void PathTracingMainRG()
 			float3 Direction;
 			float3 Throughput;
 			float4 RandSample = RandomSequence_GenerateSample4D(RandSequence, SampleIndex);
-			bool bIsTransmitted;
-			SampleMaterial(Ray.Direction, RayHitInfo, RandSample, Direction, Throughput, MaterialPdf, bIsTransmitted);
+			float SignedPositionBias;
+			SampleMaterial(Ray.Direction, RayHitInfo, RandSample, Direction, Throughput, MaterialPdf, SignedPositionBias);
 			if (MaterialPdf == 0.0)
 			{
 				break;
@@ -257,7 +257,7 @@ void PathTracingMainRG()
 
 			// Update ray
 			Ray.Direction = Direction;
-			ApplyPositionBias(Ray, bIsTransmitted ? -Payload.WorldNormal : Payload.WorldNormal, AdaptiveSamplingData.MaxNormalBias);
+			ApplyPositionBias(Ray, SignedPositionBias * Payload.WorldNormal, AdaptiveSamplingData.MaxNormalBias);
 			RayThroughput *= Throughput;
 			RayPdf *= MaterialPdf;
 
