Respect shadow geometry visibility when casting visibility rays.

diff --git a/Engine/Shaders/Private/RayTracing/RayTracingAmbientOcclusionRGS.usf b/Engine/Shaders/Private/RayTracing/RayTracingAmbientOcclusionRGS.usf
--- a/Engine/Shaders/Private/RayTracing/RayTracingAmbientOcclusionRGS.usf
+++ b/Engine/Shaders/Private/RayTracing/RayTracingAmbientOcclusionRGS.usf
@@ -105,7 +105,7 @@ void AmbientOcclusionRGS()
 		TraceRay(
 			TLAS,
 			RayFlags,
-			RAY_TRACING_MASK_OPAQUE,          // InstanceInclusionMask
+			RAY_TRACING_MASK_SHADOW,          // InstanceInclusionMask
 			RAY_TRACING_SHADER_SLOT_MATERIAL, // RayContributionToHitGroupIndex
 			RAY_TRACING_NUM_SHADER_SLOTS,     // MultiplierForGeometryContributionToShaderIndex
 			0,    // MissShaderIndex
diff --git a/Engine/Shaders/Private/RayTracing/RayTracingRectLightRGS.usf b/Engine/Shaders/Private/RayTracing/RayTracingRectLightRGS.usf
--- a/Engine/Shaders/Private/RayTracing/RayTracingRectLightRGS.usf
+++ b/Engine/Shaders/Private/RayTracing/RayTracingRectLightRGS.usf
@@ -290,7 +290,7 @@ void RectLightRGS()
 		TraceRay(
 			TLAS,      // AccelerationStructure
 			RayFlags,
-			RAY_TRACING_MASK_OPAQUE,          // InstanceInclusionMask
+			RAY_TRACING_MASK_SHADOW,          // InstanceInclusionMask
 			RAY_TRACING_SHADER_SLOT_MATERIAL, // RayContributionToHitGroupIndex
 			RAY_TRACING_NUM_SHADER_SLOTS,     // MultiplierForGeometryContributionToShaderIndex
 			0,         // MissShaderIndex
diff --git a/Engine/Shaders/Private/RayTracing/RayTracingSkyLightRGS.usf b/Engine/Shaders/Private/RayTracing/RayTracingSkyLightRGS.usf
--- a/Engine/Shaders/Private/RayTracing/RayTracingSkyLightRGS.usf
+++ b/Engine/Shaders/Private/RayTracing/RayTracingSkyLightRGS.usf
@@ -315,7 +315,7 @@ void SkyLightRGS()
 		TraceRay(
 			TLAS,      // AccelerationStructure
 			RayFlags,
-			RAY_TRACING_MASK_OPAQUE,          // InstanceInclusionMask
+			RAY_TRACING_MASK_SHADOW,          // InstanceInclusionMask
 			RAY_TRACING_SHADER_SLOT_MATERIAL, // RayContributionToHitGroupIndex
 			RAY_TRACING_NUM_SHADER_SLOTS,     // MultiplierForGeometryContributionToShaderIndex
 			0,         // MissShaderIndex
