Cherry-pick 5347750: Implements r.Shadow.Denoiser.HistoryConvolutionSamples with DRB.

diff --git a/Engine/Shaders/Private/ScreenSpaceDenoise/SSDSignalAccumulator.ush b/Engine/Shaders/Private/ScreenSpaceDenoise/SSDSignalAccumulator.ush
--- a/Engine/Shaders/Private/ScreenSpaceDenoise/SSDSignalAccumulator.ush
+++ b/Engine/Shaders/Private/ScreenSpaceDenoise/SSDSignalAccumulator.ush
@@ -12,12 +12,18 @@
  */
 struct FSSDSignalAccumulator
 {
+	// Sums of moment.
 	FSSFSignalSample Moment1;
 	FSSFSignalSample Moment2;
 
+	// Per component min and max values.
 	FSSFSignalSample Min;
 	FSSFSignalSample Max;
+
+	// Minimal inverse frequency found and intersected when sampling.
+	float MinInvFrequency;
 	
+	// Descending ring bucketing.
 	FSSFSignalSample Current;
 	float CurrentInvFrequency;
 	float CurrentTranslucency;
@@ -40,6 +46,8 @@ FSSDSignalAccumulator CreateSignalAccumulator()
 	Accumulator.Min = CreateSignalSampleFromScalarValue(1.#INF);
 	Accumulator.Max = CreateSignalSampleFromScalarValue(0.0); // TODO: -1.#INF? otherwise there is a risk to color clamp with YCoCg.
 
+	Accumulator.MinInvFrequency = WORLD_RADIUS_MISS;
+
 	// DRB initialization.
 	{
 		Accumulator.Current = CreateSignalSampleFromScalarValue(0.0);
@@ -109,6 +117,10 @@ void AccumulateSampleInDomainBoundaries(inout FSSDSignalAccumulator Accumulator,
 	{
 		Accumulator.Min = Min(Accumulator.Min, SampleInfos.Sample);
 		Accumulator.Max = Max(Accumulator.Max, SampleInfos.Sample);
+		
+		// Track the minimal frequency inverse.
+		if (SampleInfos.InvFrequency != WORLD_RADIUS_MISS)
+			Accumulator.MinInvFrequency = min(Accumulator.MinInvFrequency, SampleInfos.InvFrequency);
 	}
 }
 
@@ -121,7 +133,7 @@ void StartAccumulatingClusterInDRB(
 	inout FSSDSignalAccumulator Accumulator,
 	uint RingId)
 {
-	const uint ErrorCorrection = 1;
+	const uint ErrorCorrection = 2;
 	
 	// Compute the bluring radius of the output pixel itself.
 	float RefPixelWorldBluringRadius = AmendWorldBluringRadiusCausedByPixelSize(ComputeWorldBluringRadiusCausedByPixelSize(RefSceneMetadata));
@@ -160,20 +172,28 @@ void AccumulateSampleInDRB(inout FSSDSignalAccumulator Accumulator, FSSDSampleAc
 	{
 		float CurrentWeight = A.FinalWeight * (1.0 - bBelongsToPrevious);
 		Accumulator.Current = Add(Accumulator.Current, Mul(A.Sample, CurrentWeight));
-		Accumulator.CurrentInvFrequency += ClampedInvFrequency * CurrentWeight;
+		Accumulator.CurrentInvFrequency += A.Sample.AccumulatedSampleCount * ClampedInvFrequency * CurrentWeight;
 	}
 
 	// Accumulate current bucket translucency.
 	{
 		float SampleTranslucency = saturate(ClampedInvFrequency - Accumulator.BorderingRadius);
-		Accumulator.CurrentTranslucency += SampleTranslucency;
+		
+		// TODO: do need (A.Sample.LeavingRayCount * SafeRcp(A.Sample.AccumulatedSampleCount)) computation onen 1spp reconstruction.
+		// TODO: could pass down a normalised version of A.Sample to avoid Rcp.
+		float R = A.Sample.LeavingRayCount * SafeRcp(A.Sample.AccumulatedSampleCount);
+		
+		//Accumulator.CurrentTranslucency += SampleTranslucency;
+		Accumulator.CurrentTranslucency += lerp(R, 1, SampleTranslucency);
+		//Accumulator.CurrentTranslucency += saturate(1 - (1 - SampleTranslucency) * (1 - R));
+		//Accumulator.CurrentTranslucency += lerp(R, 1, SampleTranslucency);
 	}
 
 	// Accumulate previous bucket.
 	{
 		float PreviousWeight = A.FinalWeight * bBelongsToPrevious;
 		Accumulator.Previous = Add(Accumulator.Previous, Mul(A.Sample, PreviousWeight));
-		Accumulator.PreviousInvFrequency += ClampedInvFrequency * PreviousWeight;
+		Accumulator.PreviousInvFrequency += A.Sample.AccumulatedSampleCount * ClampedInvFrequency * PreviousWeight;
 	}
 }
 
@@ -203,7 +223,9 @@ void DijestSampleClusterInDRB(
 	// Compute final factor to use to with previous bucket.
 	float PreviousBucketFactor = (Accumulator.Previous.AccumulatedSampleCount == 0.0) ? 0.0 : (1.0 - CurrentClusterOpacity * bOccludingCluster);
 
-	//PreviousBucketFactor = 1;
+	// TODO: the big difference in sample count could lead to wrong opacity between sample clusters.
+	//if (PreviousBucketFactor < 1 && Accumulator.Previous.AccumulatedSampleCount > 0)
+	//	PreviousBucketFactor *= Accumulator.Current.AccumulatedSampleCount * rcp(Accumulator.Previous.AccumulatedSampleCount);
 
 	Accumulator.Previous = Add(Mul(Accumulator.Previous, PreviousBucketFactor), Accumulator.Current);
 	Accumulator.PreviousInvFrequency = Accumulator.PreviousInvFrequency * PreviousBucketFactor + Accumulator.CurrentInvFrequency;
diff --git a/Engine/Shaders/Private/ScreenSpaceDenoise/SSDSignalCore.ush b/Engine/Shaders/Private/ScreenSpaceDenoise/SSDSignalCore.ush
--- a/Engine/Shaders/Private/ScreenSpaceDenoise/SSDSignalCore.ush
+++ b/Engine/Shaders/Private/ScreenSpaceDenoise/SSDSignalCore.ush
@@ -254,32 +254,37 @@ void DecodeMultiplexedSignals(
 			OutSamples.Array[MultiplexId].WorldBluringRadius = OutSamples.Array[MultiplexId].AccumulatedSampleCount * (Channels.g == -1 ? WORLD_RADIUS_MISS : Channels.g);
 		}
 	}
-	else if (SignalBufferLayout == SIGNAL_BUFFER_LAYOUT_PENUMBRA_RECONSTRUCTION)
+	//else if (SignalBufferLayout == SIGNAL_BUFFER_LAYOUT_PENUMBRA_RECONSTRUCTION)
+	//{
+	//	FIX_UNROLL(MAX_SIGNAL_MULTIPLEXING)
+	//	for (uint MultiplexId = 0; MultiplexId < MAX_SIGNAL_MULTIPLEXING; MultiplexId++)
+	//	{
+	//		float2 Channels = float2(RawSample[MultiplexId / 2][2 * (MultiplexId % 2)], RawSample[MultiplexId / 2][2 * (MultiplexId % 2) + 1]);
+	//
+	//		if (bNormalizeSample)
+	//		{
+	//			OutSamples.Array[MultiplexId].AccumulatedSampleCount = Channels.g > 0 ? 1 : 0;
+	//			OutSamples.Array[MultiplexId].LeavingRayCount = Channels.g > 0 ? Channels.r : 0;
+	//		}
+	//		else
+	//		{
+	//			OutSamples.Array[MultiplexId].AccumulatedSampleCount = Channels.g;
+	//			OutSamples.Array[MultiplexId].LeavingRayCount = Channels.r * Channels.g;
+	//		}
+	//	}
+	//}
+	else if (SignalBufferLayout == SIGNAL_BUFFER_LAYOUT_PENUMBRA_HISTORY || SignalBufferLayout == SIGNAL_BUFFER_LAYOUT_PENUMBRA_RECONSTRUCTION)
 	{
 		FIX_UNROLL(MAX_SIGNAL_MULTIPLEXING)
 		for (uint MultiplexId = 0; MultiplexId < MAX_SIGNAL_MULTIPLEXING; MultiplexId++)
 		{
-			float2 Channels = float2(RawSample[MultiplexId / 2][2 * (MultiplexId % 2)], RawSample[MultiplexId / 2][2 * (MultiplexId % 2) + 1]);
-	
-			if (bNormalizeSample)
-			{
-				OutSamples.Array[MultiplexId].AccumulatedSampleCount = Channels.g > 0 ? 1 : 0;
-				OutSamples.Array[MultiplexId].LeavingRayCount = Channels.g > 0 ? Channels.r : 0;
-			}
-			else
-			{
-				OutSamples.Array[MultiplexId].AccumulatedSampleCount = Channels.g;
-				OutSamples.Array[MultiplexId].LeavingRayCount = Channels.r * Channels.g;
-			}
-		}
-	}
-	else if (SignalBufferLayout == SIGNAL_BUFFER_LAYOUT_PENUMBRA_HISTORY)
-	{
-		FIX_UNROLL(MAX_SIGNAL_MULTIPLEXING)
-		for (uint MultiplexId = 0; MultiplexId < MAX_SIGNAL_MULTIPLEXING; MultiplexId++)
-		{
-			OutSamples.Array[MultiplexId].LeavingRayCount = RawSample[MultiplexId].r * RawSample[MultiplexId].g;
-			OutSamples.Array[MultiplexId].AccumulatedSampleCount = RawSample[MultiplexId].g;
+			float3 Channels = RawSample[MultiplexId].xyz;
+			
+			float SampleCount = bNormalizeSample ? (Channels.g > 0 ? 1 : 0) : (Channels.g);
+
+			OutSamples.Array[MultiplexId].LeavingRayCount = Channels.r * SampleCount;
+			OutSamples.Array[MultiplexId].AccumulatedSampleCount = SampleCount;
+			OutSamples.Array[MultiplexId].WorldBluringRadius = (Channels.b == -1 ? WORLD_RADIUS_MISS : Channels.b) * SampleCount;
 		}
 	}
 	else if (SignalBufferLayout == SIGNAL_BUFFER_LAYOUT_REFLECTIONS_INPUT)
@@ -454,31 +459,35 @@ void EncodeMultiplexedSignals(
 		}
 		OutBufferCount = (MultiplexCount + 1) / 2;
 	}
-	else if (SignalBufferLayout == SIGNAL_BUFFER_LAYOUT_PENUMBRA_RECONSTRUCTION)
+	//else if (SignalBufferLayout == SIGNAL_BUFFER_LAYOUT_PENUMBRA_RECONSTRUCTION)
+	//{
+	//	UNROLL
+	//	for (uint MultiplexId = 0; MultiplexId < MultiplexCount; MultiplexId++)
+	//	{
+	//		float2 EncodedChannels = float2(
+	//			Sample[MultiplexId].LeavingRayCount * SafeRcp(Sample[MultiplexId].AccumulatedSampleCount),
+	//			Sample[MultiplexId].AccumulatedSampleCount);
+	//
+	//		if (MultiplexId % 2)
+	//			OutRawSample[MultiplexId / 2].zw = EncodedChannels;
+	//		else
+	//			OutRawSample[MultiplexId / 2].xy = EncodedChannels;
+	//	}
+	//	OutBufferCount = (MultiplexCount + 1) / 2;
+	//}
+	else if (SignalBufferLayout == SIGNAL_BUFFER_LAYOUT_PENUMBRA_HISTORY || SignalBufferLayout == SIGNAL_BUFFER_LAYOUT_PENUMBRA_RECONSTRUCTION)
 	{
 		UNROLL
 		for (uint MultiplexId = 0; MultiplexId < MultiplexCount; MultiplexId++)
 		{
-			float2 EncodedChannels = float2(
-				Sample[MultiplexId].LeavingRayCount * SafeRcp(Sample[MultiplexId].AccumulatedSampleCount),
-				Sample[MultiplexId].AccumulatedSampleCount);
+			float NormalizationFactor = SafeRcp(Sample[MultiplexId].AccumulatedSampleCount);
+
+			float NormalizedWorldBluringRadius = Sample[MultiplexId].WorldBluringRadius * NormalizationFactor;
 
-			if (MultiplexId % 2)
-				OutRawSample[MultiplexId / 2].zw = EncodedChannels;
-			else
-				OutRawSample[MultiplexId / 2].xy = EncodedChannels;
-		}
-		OutBufferCount = (MultiplexCount + 1) / 2;
-	}
-	else if (SignalBufferLayout == SIGNAL_BUFFER_LAYOUT_PENUMBRA_HISTORY)
-	{
-		UNROLL
-		for (uint MultiplexId = 0; MultiplexId < MultiplexCount; MultiplexId++)
-		{
 			OutRawSample[MultiplexId] = float4(
-				Sample[MultiplexId].LeavingRayCount * SafeRcp(Sample[MultiplexId].AccumulatedSampleCount),
+				Sample[MultiplexId].LeavingRayCount * NormalizationFactor,
 				Sample[MultiplexId].AccumulatedSampleCount,
-				0,
+				NormalizedWorldBluringRadius == WORLD_RADIUS_MISS ? -1 : NormalizedWorldBluringRadius,
 				0);
 		}
 		OutBufferCount = MultiplexCount;
diff --git a/Engine/Shaders/Private/ScreenSpaceDenoise/SSDSignalFramework.ush b/Engine/Shaders/Private/ScreenSpaceDenoise/SSDSignalFramework.ush
--- a/Engine/Shaders/Private/ScreenSpaceDenoise/SSDSignalFramework.ush
+++ b/Engine/Shaders/Private/ScreenSpaceDenoise/SSDSignalFramework.ush
@@ -507,6 +507,15 @@ float ComputeBilateralWeight(
 
 		PositionBilateralWeight = saturate(1 - DistFromRefPlane * rcp(MaxWorldBluringDistance));
 	}
+	#elif CONFIG_VIEW_POSITION_BASED_BILATERAL == 3 && 1
+	{
+		float3 V = RefSceneMetadata.TranslatedWorldPosition - NeighborSceneMetadata.TranslatedWorldPosition;
+
+		if (dot(V, V) > MaxWorldBluringDistance * MaxWorldBluringDistance)
+		{
+			PositionBilateralWeight = 0;
+		}
+	}
 	#elif CONFIG_VIEW_POSITION_BASED_BILATERAL == 3
 	{
 		float3 V = RefSceneMetadata.TranslatedWorldPosition - NeighborSceneMetadata.TranslatedWorldPosition;
@@ -515,6 +524,9 @@ float ComputeBilateralWeight(
 		{
 			PositionBilateralWeight = 0;
 		}
+
+		// TODO: extremely expensive with rcp()...
+		PositionBilateralWeight = saturate(1 - dot(V, V) * rcp(MaxWorldBluringDistance * MaxWorldBluringDistance));
 	}
 	#endif
 
diff --git a/Engine/Shaders/Private/ScreenSpaceDenoise/SSDSpatialAccumulation.usf b/Engine/Shaders/Private/ScreenSpaceDenoise/SSDSpatialAccumulation.usf
--- a/Engine/Shaders/Private/ScreenSpaceDenoise/SSDSpatialAccumulation.usf
+++ b/Engine/Shaders/Private/ScreenSpaceDenoise/SSDSpatialAccumulation.usf
@@ -61,7 +61,9 @@
 		#define CONFIG_SIGNAL_OUTPUT_LAYOUT SIGNAL_BUFFER_LAYOUT_PENUMBRA_HISTORY
 
 		#define CONFIG_SAMPLE_SET           SAMPLE_SET_STACKOWIAK_4_SETS
-		#define CONFIG_SAMPLE_SUBSET        1
+		//#define CONFIG_SAMPLE_SUBSET        1
+		#define CONFIG_BILATERAL_DISTANCE_COMPUTATION SIGNAL_WORLD_FREQUENCY_PRECOMPUTED_BLURING_RADIUS
+		#define CONFIG_OUTPUT_MODE          OUTPUT_MODE_DRB
 
 	#elif DIM_STAGE == STAGE_FINAL_OUTPUT
 		#define CONFIG_SIGNAL_INPUT_LAYOUT  SIGNAL_BUFFER_LAYOUT_PENUMBRA_HISTORY
@@ -430,6 +432,12 @@ void MainCS(
 				SampleInfos.Sample = CenterSample;
 				SampleInfos.FinalWeight = 1.0;
 				SampleInfos.InvFrequency = GetSignalWorldBluringRadius(CenterSample, RefSceneMetadata, BatchedSignalId);
+				
+				if (KernelConfig.BilateralDistanceComputation == SIGNAL_WORLD_FREQUENCY_PRECOMPUTED_BLURING_RADIUS)
+				{
+					// TODO: this is ineficient, could fetch the normalised WorldBluringRadius instead of SafeRcp().
+					SampleInfos.InvFrequency = CenterSample.WorldBluringRadius * SafeRcp(CenterSample.AccumulatedSampleCount);
+				}
 
 				AccumulateSample(
 					/* inout */ SignalAccumulators.Array[SignalMultiplexId],
@@ -514,29 +522,31 @@ void MainCS(
 	}
 
 	// Color processing of the signal to reduce highlight flickering.
-	FIX_UNROLL(MAX_SIGNAL_MULTIPLEXING)
-	for (uint MultiplexId = 0; MultiplexId < MAX_SIGNAL_MULTIPLEXING; MultiplexId++)
 	{
-		#if CONFIG_REJECT_HIGHEST_COLOR
-		if (Accumulator.Moment1[MultiplexId].AccumulatedSampleCount > 0)
+		FIX_UNROLL(CONFIG_SIGNAL_BATCH_SIZE)
+		for (uint MultiplexId = 0; MultiplexId < CONFIG_SIGNAL_BATCH_SIZE; MultiplexId++)
 		{
-			const float MaxNeighborWeight = saturate(SignalAccumulators.Array[MultiplexId].Moment1.AccumulatedSampleCount * rcp(10) - 1);
+			#if CONFIG_REJECT_HIGHEST_COLOR
+			if (Accumulator.Moment1[MultiplexId].AccumulatedSampleCount > 0)
+			{
+				const float MaxNeighborWeight = saturate(SignalAccumulators.Array[MultiplexId].Moment1.AccumulatedSampleCount * rcp(10) - 1);
 
-			SignalAccumulators.Array[MultiplexId].Moment1.SceneColor.rgb =
-				(SignalAccumulators.Array[MultiplexId].Moment1.SceneColor.rgb - MaxNeighbor.SceneColor.rgb * MaxNeighborWeight) *
-				(SignalAccumulators.Array[MultiplexId].Moment1.AccumulatedSampleCount / (SignalAccumulators.Array[MultiplexId].Moment1.AccumulatedSampleCount - MaxNeighborWeight));
-		}
-		#endif
+				SignalAccumulators.Array[MultiplexId].Moment1.SceneColor.rgb =
+					(SignalAccumulators.Array[MultiplexId].Moment1.SceneColor.rgb - MaxNeighbor.SceneColor.rgb * MaxNeighborWeight) *
+					(SignalAccumulators.Array[MultiplexId].Moment1.AccumulatedSampleCount / (SignalAccumulators.Array[MultiplexId].Moment1.AccumulatedSampleCount - MaxNeighborWeight));
+			}
+			#endif
 
-		SignalAccumulators.Array[MultiplexId].Moment1 = TransformSignal(
-			SignalAccumulators.Array[MultiplexId].Moment1, 
-			/* SrcBasis  = */ CONFIG_ACCUMULATION_COLOR_SPACE, 
-			/* DestBasis = */ CONFIG_OUTPUT_COLOR_SPACE);
+			SignalAccumulators.Array[MultiplexId].Moment1 = TransformSignal(
+				SignalAccumulators.Array[MultiplexId].Moment1, 
+				/* SrcBasis  = */ CONFIG_ACCUMULATION_COLOR_SPACE, 
+				/* DestBasis = */ CONFIG_OUTPUT_COLOR_SPACE);
 
-		SignalAccumulators.Array[MultiplexId].Moment2 = TransformSignal(
-			SignalAccumulators.Array[MultiplexId].Moment2, 
-			/* SrcBasis  = */ CONFIG_ACCUMULATION_COLOR_SPACE, 
-			/* DestBasis = */ CONFIG_OUTPUT_COLOR_SPACE);
+			SignalAccumulators.Array[MultiplexId].Moment2 = TransformSignal(
+				SignalAccumulators.Array[MultiplexId].Moment2, 
+				/* SrcBasis  = */ CONFIG_ACCUMULATION_COLOR_SPACE, 
+				/* DestBasis = */ CONFIG_OUTPUT_COLOR_SPACE);
+		}
 	}
 
 	// Transcode the spatial accumulation into multiplexed signal according to different modes.
@@ -571,21 +581,45 @@ void MainCS(
 			for (uint MultiplexId = 0; MultiplexId < CONFIG_SIGNAL_BATCH_SIZE; MultiplexId++)
 			{
 				OutputSamples.Array[MultiplexId] = SignalAccumulators.Array[MultiplexId].Previous;
+
+				// Output the minimal inverse frequency as new world bluring radius for subsequent passes.
+				// TODO: store this in its own signal to avoid *= N; *= rcp(N);
+				OutputSamples.Array[MultiplexId].WorldBluringRadius = 
+					OutputSamples.Array[MultiplexId].AccumulatedSampleCount * SignalAccumulators.Array[MultiplexId].MinInvFrequency;
+			
+				// No need to keep the VGPR pressure at this point for WorldBluringRadius, because no passes use it after.
+				if (DIM_STAGE == STAGE_POST_FILTERING && 0)
+				{
+					OutputSamples.Array[MultiplexId].WorldBluringRadius = 0;
+				}
 			}
 
 			if (0)
 			{
 				DebugOutput[DispatchThreadId] = float4(
 					SignalAccumulators.Array[0].Previous.AccumulatedSampleCount,
-					SignalAccumulators.Array[0].Current.AccumulatedSampleCount,
-					0,
-					0);
+					SignalAccumulators.Array[0].CurrentTranslucency,
+					SignalAccumulators.Array[0].Previous.LeavingRayCount * SafeRcp(SignalAccumulators.Array[0].Previous.AccumulatedSampleCount),
+					SignalAccumulators.Array[0].Current.LeavingRayCount * SafeRcp(SignalAccumulators.Array[0].Current.AccumulatedSampleCount));
 			}
 		}
 		#else
 			#error Unknown output mode.
 		#endif
 	}
+	
+	// Clamp the number of sample recorded.
+	if (DIM_STAGE == STAGE_POST_FILTERING)
+	{
+		FIX_UNROLL(CONFIG_SIGNAL_BATCH_SIZE)
+		for (uint MultiplexId = 0; MultiplexId < CONFIG_SIGNAL_BATCH_SIZE; MultiplexId++)
+		{
+			float CurrentSampleCount = OutputSamples.Array[MultiplexId].AccumulatedSampleCount;
+			float NewSampleCount = min(CurrentSampleCount, 128);
+
+			OutputSamples.Array[MultiplexId] = Mul(OutputSamples.Array[MultiplexId], NewSampleCount / CurrentSampleCount);
+		}
+	}
 
 	// TODO:  LeaveRayCount = (LeaveRayCount - 1) * 9 / (9 - 2) post processing to reject when for history rejection
 	
diff --git a/Engine/Shaders/Private/ScreenSpaceDenoise/SSDSpatialKernel.ush b/Engine/Shaders/Private/ScreenSpaceDenoise/SSDSpatialKernel.ush
--- a/Engine/Shaders/Private/ScreenSpaceDenoise/SSDSpatialKernel.ush
+++ b/Engine/Shaders/Private/ScreenSpaceDenoise/SSDSpatialKernel.ush
@@ -272,7 +272,8 @@ void AccumulateSampledMultiplexedSignals(
 		// But the signal's bluring radius might already be pre computed.
 		if (KernelConfig.BilateralDistanceComputation == SIGNAL_WORLD_FREQUENCY_PRECOMPUTED_BLURING_RADIUS)
 		{
-			SignalConvolutionBluringRadius = Sample.WorldBluringRadius;
+			// TODO: this is ineficient, could fetch the normalised WorldBluringRadius instead of SafeRcp().
+			SignalConvolutionBluringRadius = Sample.WorldBluringRadius * SafeRcp(Sample.AccumulatedSampleCount);
 		}
 
 		// Compute the final world distance to use for bilateral rejection.
@@ -720,7 +721,7 @@ void ConvolveStackowiakKernel(
 			{
 				const uint SampleId = (BatchCountCount - BatchId) * kSamplingBatchSize - 1 - SampleBatchId;
 
-				const bool bIsKernelCenterSample = SampleId == 0 && (SampleBatchId == 3);
+				const bool bIsKernelCenterSample = SampleId == 0 && (SampleBatchId == (kSamplingBatchSize - 1));
 				
 				float2 SampleOffset = kStackowiakSampleSet0[kStackowiakSampleSetCount * SampleId + KernelConfig.SampleTrackId];
 				if (KernelConfig.SampleSubSetId == 1)
diff --git a/Engine/Shaders/Private/ScreenSpaceDenoise/SSDTemporalAccumulation.usf b/Engine/Shaders/Private/ScreenSpaceDenoise/SSDTemporalAccumulation.usf
--- a/Engine/Shaders/Private/ScreenSpaceDenoise/SSDTemporalAccumulation.usf
+++ b/Engine/Shaders/Private/ScreenSpaceDenoise/SSDTemporalAccumulation.usf
@@ -41,7 +41,9 @@
 
 	#define CONFIG_HISTORY_REJECTION HISTORY_REJECTION_VAR_BOUNDARIES
 
-	#define CONFIG_SAMPLE_SET           SAMPLE_SET_STACKOWIAK_4_SETS
+	// Uses nearest to not leak informations on geometric edges by bilateral. This is OK because just blury greyscale details.
+	// TODO: a bit hacky, need find better solution.
+	#define CONFIG_USE_NEAREST_HISTORY 1
 
 #elif CONFIG_SIGNAL_PROCESSING == SIGNAL_PROCESSING_REFLECTIONS
 	// Input and output layout.
@@ -127,6 +129,11 @@
 	#define CONFIG_MULTIPLEXED_SIGNALS_PER_SIGNAL_DOMAIN MAX_SIGNAL_MULTIPLEXING
 #endif
 
+/** Whether can use a nearest sampler when sampling the history */
+#ifndef CONFIG_USE_NEAREST_HISTORY
+	#define CONFIG_USE_NEAREST_HISTORY 0
+#endif
+
 
 //------------------------------------------------------- CONFIG ENABLED DEFAULTS
 
@@ -291,6 +298,7 @@ void TemporallyAccumulate(
 		float2 ClampedHistoryBufferUV = clamp(HistoryBufferUV, View.BufferBilinearUVMinMax.xy, View.BufferBilinearUVMinMax.zw);
 		bool bIsPreviousFrameOffscreen = any(HistoryBufferUV != ClampedHistoryBufferUV);
 
+		// TODO: global camera cut to save perf that is going to be required for panic post filtering.
 		BRANCH
 		if (!bIsPreviousFrameOffscreen)
 		{
@@ -302,6 +310,10 @@ void TemporallyAccumulate(
 			KernelConfig.MultiplexedSignalsPerSignalDomain = CONFIG_MULTIPLEXED_SIGNALS_PER_SIGNAL_DOMAIN;
 			KernelConfig.bUnroll = true;
 			KernelConfig.bPreviousFrameMetadata = true;
+			KernelConfig.BilateralDistanceComputation = SIGNAL_WORLD_FREQUENCY_MIN_METADATA;
+			
+			if (CONFIG_USE_NEAREST_HISTORY)
+				KernelConfig.SampleSet = SAMPLE_SET_1X1;
 
 			// SGPR configuration of the kernel.
 			KernelConfig.BufferSizeAndInvSize = View.BufferSizeAndInvSize;
@@ -320,8 +332,9 @@ void TemporallyAccumulate(
 				PrevHistory_Textures_3,
 				/* inout */ SignalAccumulators);
 		
-			UNROLL
-			for (uint BatchedSignalId = 0; BatchedSignalId < MAX_SIGNAL_BATCH_SIZE; BatchedSignalId++)
+			// Exports the history sample from accumulator.
+			FIX_UNROLL(CONFIG_SIGNAL_BATCH_SIZE)
+			for (uint BatchedSignalId = 0; BatchedSignalId < CONFIG_SIGNAL_BATCH_SIZE; BatchedSignalId++)
 			{
 				HistorySamples.Array[BatchedSignalId] = SignalAccumulators.Array[BatchedSignalId].Moment1;
 				BRANCH
@@ -330,6 +343,19 @@ void TemporallyAccumulate(
 					HistorySamples.Array[BatchedSignalId] = CreateSignalSampleFromScalarValue(0.0);
 				}
 			}
+
+			// Override the previous frame's inverse frequency with the current frame's one, because don't know how much it changed.
+			FIX_UNROLL(CONFIG_SIGNAL_BATCH_SIZE)
+			for (uint BatchedSignalId = 0; BatchedSignalId < CONFIG_SIGNAL_BATCH_SIZE; BatchedSignalId++)
+			{
+				FSSFSignalSample CurrentFrameSample = CurrentFrameSamples.Array[BatchedSignalId];
+				FSSFSignalSample HistorySample = HistorySamples.Array[BatchedSignalId];
+
+				HistorySamples.Array[BatchedSignalId].WorldBluringRadius = (
+					HistorySample.AccumulatedSampleCount *
+					CurrentFrameSample.WorldBluringRadius *
+					SafeRcp(CurrentFrameSample.AccumulatedSampleCount));
+			} // for (uint BatchedSignalId = 0; BatchedSignalId < CONFIG_SIGNAL_BATCH_SIZE; BatchedSignalId++)
 		} // if (!bIsPreviousFrameOffscreen)
 	}
 	
@@ -351,6 +377,11 @@ void TemporallyAccumulate(
 			// History rejection is already fudge factor by how reprojection. Want to prioritize rejection stability with more sample
 			// than accuracy, so only take the blur distance of the reference sample that depends on depth and pixel size of the current frame.
 			KernelConfig.BilateralDistanceComputation = SIGNAL_WORLD_FREQUENCY_REF_METADATA_ONLY;
+
+			// TODO: Should be done for all signals.
+			// TODO: Should also take into account the world bluring radius du to size of the pixel.
+			if (CONFIG_SIGNAL_PROCESSING == SIGNAL_PROCESSING_MONOCHROMATIC_PENUMBRA)
+				KernelConfig.BilateralDistanceComputation = SIGNAL_WORLD_FREQUENCY_PRECOMPUTED_BLURING_RADIUS;
 			
 			KernelConfig.MultiplexedSignalsPerSignalDomain = CONFIG_MULTIPLEXED_SIGNALS_PER_SIGNAL_DOMAIN;
 
@@ -417,7 +448,7 @@ void TemporallyAccumulate(
 		FSSDSignalAccumulatorArray SignalAccumulators = CreateSignalAccumulatorArray();		
 		if (bUseCurrentSampleForCenter)
 		{
-			UNROLL
+			FIX_UNROLL(MAX_SIGNAL_MULTIPLEXING)
 			for (uint SignalMultiplexId = 0; SignalMultiplexId < MAX_SIGNAL_MULTIPLEXING; SignalMultiplexId++)
 			{
 				const uint BatchedSignalId = ComputeSignalBatchIdFromSignalMultiplexId(KernelConfig, SignalMultiplexId);
@@ -431,6 +462,12 @@ void TemporallyAccumulate(
 				SampleInfos.Sample = TransformedCurrentFrame;
 				SampleInfos.FinalWeight = 1.0;
 				SampleInfos.InvFrequency = GetSignalWorldBluringRadius(TransformedCurrentFrame, RefSceneMetadata, BatchedSignalId);
+				
+				if (KernelConfig.BilateralDistanceComputation == SIGNAL_WORLD_FREQUENCY_PRECOMPUTED_BLURING_RADIUS)
+				{
+					// TODO: this is ineficient, could fetch the normalised WorldBluringRadius instead of SafeRcp().
+					SampleInfos.InvFrequency = TransformedCurrentFrame.WorldBluringRadius * SafeRcp(TransformedCurrentFrame.AccumulatedSampleCount);
+				}
 
 				AccumulateSample(
 					/* inout */ SignalAccumulators.Array[SignalMultiplexId],
@@ -447,7 +484,7 @@ void TemporallyAccumulate(
 			/* inout */ SignalAccumulators);
 		
 		// Clamp history
-		UNROLL
+		FIX_UNROLL(CONFIG_SIGNAL_BATCH_SIZE)
 		for (uint BatchedSignalId = 0; BatchedSignalId < CONFIG_SIGNAL_BATCH_SIZE; BatchedSignalId++)
 		{
 			FSSFSignalSample NeighborMin = SignalAccumulators.Array[BatchedSignalId].Min;
@@ -560,7 +597,23 @@ void TemporallyAccumulate(
 			}
 			#endif
 
+			// Dilate the smallest inv frequency.
+			if (1)
+			{
+				float NeightborMinBluringRadius = SignalAccumulators.Array[BatchedSignalId].MinInvFrequency;
+				CurrentFrameSample.WorldBluringRadius = min(CurrentFrameSample.WorldBluringRadius, NeightborMinBluringRadius * CurrentFrameSample.AccumulatedSampleCount); 
+			
+				#if CONFIG_SIGNAL_PROCESSING == SIGNAL_PROCESSING_MONOCHROMATIC_PENUMBRA && 1
+					DebugOutput[DispatchThreadId] = float4(
+						NeightborMinBluringRadius,
+						NeightborMinBluringRadius == WORLD_RADIUS_MISS,
+						0,
+						0);
+				#endif
+			}
+			
 			HistorySamples.Array[BatchedSignalId] = HistorySample;
+			CurrentFrameSamples.Array[BatchedSignalId] = CurrentFrameSample;
 		} // for (uint BatchedSignalId = 0; BatchedSignalId < CONFIG_SIGNAL_BATCH_SIZE; BatchedSignalId++)
 	}
 	#endif // CONFIG_HISTORY_REJECTION > 0
@@ -666,6 +719,12 @@ void TemporallyAccumulate(
 				/* SrcBasis = */ CONFIG_HISTORY_BLENDING_COLOR_SPACE,
 				/* DestBasis = */ STANDARD_BUFFER_COLOR_SPACE);
 			
+			// Ignore the bluring radius of previous frame frequencies as it gets clamped anyway, just take the current frame.
+			HistorySample.WorldBluringRadius = (
+				CurrentFrameSample.WorldBluringRadius *
+				SafeRcp(CurrentFrameSample.AccumulatedSampleCount) *
+				HistorySample.AccumulatedSampleCount);
+
 			HistorySamples.Array[BatchedSignalId] = HistorySample;
 		}
 	}
diff --git a/Engine/Source/Runtime/Renderer/Private/ScreenSpaceDenoise.cpp b/Engine/Source/Runtime/Renderer/Private/ScreenSpaceDenoise.cpp
--- a/Engine/Source/Runtime/Renderer/Private/ScreenSpaceDenoise.cpp
+++ b/Engine/Source/Runtime/Renderer/Private/ScreenSpaceDenoise.cpp
@@ -28,6 +28,11 @@ static TAutoConsoleVariable<int32> CVarShadowTemporalAccumulation(
 	TEXT(""),
 	ECVF_RenderThreadSafe);
 
+static TAutoConsoleVariable<int32> CVarShadowHistoryConvolutionSampleCount(
+	TEXT("r.Shadow.Denoiser.HistoryConvolutionSamples"), 1,
+	TEXT("Number of samples to use to convolve the history over time."),
+	ECVF_RenderThreadSafe);
+
 static TAutoConsoleVariable<int32> CVarReflectionReconstructionSampleCount(
 	TEXT("r.Reflections.Denoiser.ReconstructionSamples"), 16,
 	TEXT("Maximum number of samples for the reconstruction pass (default = 16)."),
@@ -674,6 +679,14 @@ static void DenoiseSignalAtConstantPixelDensity(
 	TStaticArray<FRDGTextureDesc, kMaxBufferProcessingCount> HistoryDescs;
 	FRDGTextureDesc DebugDesc;
 	{
+		static const EPixelFormat PixelFormatPerChannel[] = {
+			PF_Unknown,
+			PF_R16F,
+			PF_G16R16F,
+			PF_FloatRGBA, // there is no 16bits float RGB
+			PF_FloatRGBA,
+		};
+
 		FRDGTextureDesc RefDesc = FRDGTextureDesc::Create2DDesc(
 			SceneBlackboard.SceneDepthBuffer->Desc.Extent,
 			PF_Unknown,
@@ -694,30 +707,26 @@ static void DenoiseSignalAtConstantPixelDensity(
 
 		if (Settings.SignalProcessing == ESignalProcessing::MonochromaticPenumbra)
 		{
-			static const EPixelFormat InjestPixelFormat[] = {
-				PF_Unknown,
-				PF_R16F,
-				PF_G16R16F,
-				PF_FloatRGBA, // there is no 16bits float RGB
-				PF_FloatRGBA,
-			};
-
 			check(Settings.SignalBatchSize >= 1 && Settings.SignalBatchSize <= IScreenSpaceDenoiser::kMaxBatchSize);
 			if (!bUseMultiInputSPPShaderPath)
-				InjestDescs[0].Format = InjestPixelFormat[Settings.SignalBatchSize];
+			{
+				InjestDescs[0].Format = PixelFormatPerChannel[Settings.SignalBatchSize];
+				InjestTextureCount = 1;
+			}
 
 			for (int32 BatchedSignalId = 0; BatchedSignalId < Settings.SignalBatchSize; BatchedSignalId++)
 			{
 				if (bUseMultiInputSPPShaderPath)
+				{
 					InjestDescs[BatchedSignalId / 2].Format = (BatchedSignalId % 2) ? PF_FloatRGBA :  PF_G16R16F;
-
-				ReconstructionDescs[BatchedSignalId / 2].Format = (BatchedSignalId % 2) ? PF_FloatRGBA :  PF_G16R16F;
-				HistoryDescs[BatchedSignalId].Format = PF_G16R16F;
+					InjestTextureCount = BatchedSignalId / 2 + 1;
+				}
+				ReconstructionDescs[BatchedSignalId].Format = PF_FloatRGBA;
+				HistoryDescs[BatchedSignalId].Format = PF_FloatRGBA;
 			}
 
-			InjestTextureCount = 1;
 			HistoryTextureCountPerSignal = 1;
-			ReconstructionTextureCount = FMath::DivideAndRoundUp(Settings.SignalBatchSize, 2);
+			ReconstructionTextureCount = Settings.SignalBatchSize;
 			bHasReconstructionLayoutDifferentFromHistory = true;
 		}
 		else if (Settings.SignalProcessing == ESignalProcessing::Reflections)
@@ -1152,7 +1161,7 @@ public:
 		Settings.InputResolutionFraction = 1.0f;
 		Settings.ReconstructionSamples = CVarShadowReconstructionSampleCount.GetValueOnRenderThread();
 		Settings.bUseTemporalAccumulation = CVarShadowTemporalAccumulation.GetValueOnRenderThread() != 0;
-		Settings.HistoryConvolutionSampleCount = 1;
+		Settings.HistoryConvolutionSampleCount = CVarShadowHistoryConvolutionSampleCount.GetValueOnRenderThread();
 		Settings.SignalBatchSize = InputParameterCount;
 
 		for (int32 BatchedSignalId = 0; BatchedSignalId < InputParameterCount; BatchedSignalId++)
