Cherry-pick 5351638: Implements history rejection preconvolution for shadows.

diff --git a/Engine/Shaders/Private/ScreenSpaceDenoise/SSDDefinitions.ush b/Engine/Shaders/Private/ScreenSpaceDenoise/SSDDefinitions.ush
--- a/Engine/Shaders/Private/ScreenSpaceDenoise/SSDDefinitions.ush
+++ b/Engine/Shaders/Private/ScreenSpaceDenoise/SSDDefinitions.ush
@@ -59,10 +59,10 @@
 #define MAX_SIGNAL_BATCH_SIZE 4
 
 // Maximum number of signal that can be multiplexed into single output.
-#define MAX_SIGNAL_MULTIPLEXING (MAX_SIGNAL_BATCH_SIZE)
+#define MAX_SIGNAL_MULTIPLEXING 8
 
 // Maximum number of buffer multiple signals can be encode to/decoded from.
-#define MAX_MULTIPLEXED_TEXTURES (MAX_SIGNAL_BATCH_SIZE)
+#define MAX_MULTIPLEXED_TEXTURES 4
 
 
 //------------------------------------------------------- ENUMS
@@ -89,6 +89,9 @@
 	 */
 	#define SIGNAL_BUFFER_LAYOUT_PENUMBRA_RECONSTRUCTION 12
 
+	/** Internal buffer to encode history rejection preconvolution. */
+	#define SIGNAL_BUFFER_LAYOUT_PENUMBRA_REJECTION 17
+
 	/** Internal buffer layout for the shadow penumbra to be stored in indivual per light histories. */
 	#define SIGNAL_BUFFER_LAYOUT_PENUMBRA_HISTORY 11
 
diff --git a/Engine/Shaders/Private/ScreenSpaceDenoise/SSDSignalCore.ush b/Engine/Shaders/Private/ScreenSpaceDenoise/SSDSignalCore.ush
--- a/Engine/Shaders/Private/ScreenSpaceDenoise/SSDSignalCore.ush
+++ b/Engine/Shaders/Private/ScreenSpaceDenoise/SSDSignalCore.ush
@@ -32,7 +32,7 @@ struct FSSFSignalSample
 /** Array of signal samples. Technically should be a typedef, but there is a bug in HLSL with the out operator of an array. */
 struct FSSDSignalArray
 {
-	FSSFSignalSample Array[MAX_SIGNAL_BATCH_SIZE];
+	FSSFSignalSample Array[MAX_SIGNAL_MULTIPLEXING];
 };
 
 
@@ -56,8 +56,8 @@ FSSFSignalSample CreateSignalSample()
 FSSDSignalArray CreateSignalArrayFromScalarValue(float Scalar)
 {
 	FSSDSignalArray OutSamples;
-	FIX_UNROLL(MAX_SIGNAL_BATCH_SIZE)
-	for (uint BatchedSignalId = 0; BatchedSignalId < MAX_SIGNAL_BATCH_SIZE; BatchedSignalId++)
+	FIX_UNROLL(MAX_SIGNAL_MULTIPLEXING)
+	for (uint BatchedSignalId = 0; BatchedSignalId < MAX_SIGNAL_MULTIPLEXING; BatchedSignalId++)
 	{
 		OutSamples.Array[BatchedSignalId] = CreateSignalSampleFromScalarValue(Scalar);
 	}
@@ -211,9 +211,10 @@ void DecodeMultiplexedSignals(
 	}
 	else if (SignalBufferLayout == SIGNAL_BUFFER_LAYOUT_PENUMBRA_INPUT_1SPP)
 	{
-		FIX_UNROLL(MAX_SIGNAL_MULTIPLEXING)
-		for (uint MultiplexId = 0; MultiplexId < MAX_SIGNAL_MULTIPLEXING; MultiplexId++)
+		FIX_UNROLL(MAX_SIGNAL_BATCH_SIZE)
+		for (uint BatchSignalId = 0; BatchSignalId < MAX_SIGNAL_BATCH_SIZE; BatchSignalId++)
 		{
+			uint MultiplexId = BatchSignalId;
 			OutSamples.Array[MultiplexId].AccumulatedSampleCount = (RawSample[MultiplexId].r == -2.0 ? 0.0 : 1.0);
 			OutSamples.Array[MultiplexId].LeavingRayCount = OutSamples.Array[MultiplexId].AccumulatedSampleCount * (RawSample[MultiplexId].r > 0.0 ? 0.0 : 1.0);
 			OutSamples.Array[MultiplexId].PenumbraClosestHit = OutSamples.Array[MultiplexId].AccumulatedSampleCount * (RawSample[MultiplexId].r == -1 ? 1000000 : RawSample[MultiplexId].r);
@@ -221,9 +222,10 @@ void DecodeMultiplexedSignals(
 	}
 	else if (SignalBufferLayout == SIGNAL_BUFFER_LAYOUT_PENUMBRA_INPUT_NSPP)
 	{
-		FIX_UNROLL(MAX_SIGNAL_MULTIPLEXING)
-		for (uint MultiplexId = 0; MultiplexId < MAX_SIGNAL_MULTIPLEXING; MultiplexId++)
+		FIX_UNROLL(MAX_SIGNAL_BATCH_SIZE)
+		for (uint BatchSignalId = 0; BatchSignalId < MAX_SIGNAL_BATCH_SIZE; BatchSignalId++)
 		{
+			uint MultiplexId = BatchSignalId;
 			// TODO: feed the actual number of sample.
 			OutSamples.Array[MultiplexId].AccumulatedSampleCount = (RawSample[MultiplexId].g == -2.0 ? 0.0 : 1.0);
 			OutSamples.Array[MultiplexId].LeavingRayCount = OutSamples.Array[MultiplexId].AccumulatedSampleCount * RawSample[MultiplexId].r;
@@ -232,9 +234,10 @@ void DecodeMultiplexedSignals(
 	}
 	else if (SignalBufferLayout == SIGNAL_BUFFER_LAYOUT_PENUMBRA_INJESTION_1SPP)
 	{
-		FIX_UNROLL(MAX_SIGNAL_MULTIPLEXING)
-		for (uint MultiplexId = 0; MultiplexId < MAX_SIGNAL_MULTIPLEXING; MultiplexId++)
+		FIX_UNROLL(MAX_SIGNAL_BATCH_SIZE)
+		for (uint BatchSignalId = 0; BatchSignalId < MAX_SIGNAL_BATCH_SIZE; BatchSignalId++)
 		{
+			uint MultiplexId = BatchSignalId;
 			float EncodedWorldBluringRadius = RawSample[MultiplexId / 4][MultiplexId % 4];
 			OutSamples.Array[MultiplexId].AccumulatedSampleCount = (EncodedWorldBluringRadius == -2.0 ? 0.0 : 1.0);
 			OutSamples.Array[MultiplexId].LeavingRayCount = OutSamples.Array[MultiplexId].AccumulatedSampleCount * (EncodedWorldBluringRadius > 0.0 ? 0.0 : 1.0);
@@ -243,9 +246,10 @@ void DecodeMultiplexedSignals(
 	}
 	else if (SignalBufferLayout == SIGNAL_BUFFER_LAYOUT_PENUMBRA_INJESTION_NSPP)
 	{
-		FIX_UNROLL(MAX_SIGNAL_MULTIPLEXING)
-		for (uint MultiplexId = 0; MultiplexId < MAX_SIGNAL_MULTIPLEXING; MultiplexId++)
+		FIX_UNROLL(MAX_SIGNAL_BATCH_SIZE)
+		for (uint BatchSignalId = 0; BatchSignalId < MAX_SIGNAL_BATCH_SIZE; BatchSignalId++)
 		{
+			uint MultiplexId = BatchSignalId;
 			float2 Channels = float2(RawSample[MultiplexId / 2][2 * (MultiplexId % 2)], RawSample[MultiplexId / 2][2 * (MultiplexId % 2) + 1]);
 	
 			// TODO: feed the actual number of sample.
@@ -275,8 +279,8 @@ void DecodeMultiplexedSignals(
 	//}
 	else if (SignalBufferLayout == SIGNAL_BUFFER_LAYOUT_PENUMBRA_HISTORY || SignalBufferLayout == SIGNAL_BUFFER_LAYOUT_PENUMBRA_RECONSTRUCTION)
 	{
-		FIX_UNROLL(MAX_SIGNAL_MULTIPLEXING)
-		for (uint MultiplexId = 0; MultiplexId < MAX_SIGNAL_MULTIPLEXING; MultiplexId++)
+		FIX_UNROLL(MAX_SIGNAL_BATCH_SIZE)
+		for (uint MultiplexId = 0; MultiplexId < MAX_SIGNAL_BATCH_SIZE; MultiplexId++)
 		{
 			float3 Channels = RawSample[MultiplexId].xyz;
 			
@@ -287,6 +291,23 @@ void DecodeMultiplexedSignals(
 			OutSamples.Array[MultiplexId].WorldBluringRadius = (Channels.b == -1 ? WORLD_RADIUS_MISS : Channels.b) * SampleCount;
 		}
 	}
+	else if (SignalBufferLayout == SIGNAL_BUFFER_LAYOUT_PENUMBRA_REJECTION)
+	{
+		FIX_UNROLL(MAX_SIGNAL_BATCH_SIZE)
+		for (uint BatchSignalId = 0; BatchSignalId < MAX_SIGNAL_BATCH_SIZE; BatchSignalId++)
+		{
+			uint MultiplexId = BatchSignalId * 2;
+			float4 Channels = RawSample[BatchSignalId];
+			float SampleCount = Channels.b;
+
+			OutSamples.Array[MultiplexId + 0].AccumulatedSampleCount = SampleCount;
+			OutSamples.Array[MultiplexId + 0].LeavingRayCount = Channels.r * SampleCount;
+			OutSamples.Array[MultiplexId + 0].WorldBluringRadius = (Channels.a == -1 ? WORLD_RADIUS_MISS : Channels.a) * SampleCount;
+			
+			OutSamples.Array[MultiplexId + 1].AccumulatedSampleCount = SampleCount;
+			OutSamples.Array[MultiplexId + 1].LeavingRayCount = Channels.g * SampleCount;
+		}
+	}
 	else if (SignalBufferLayout == SIGNAL_BUFFER_LAYOUT_REFLECTIONS_INPUT)
 	{
 		OutSamples.Array[0].AccumulatedSampleCount = RawSample[1].r == -2 ? 0 : 1;
@@ -492,6 +513,29 @@ void EncodeMultiplexedSignals(
 		}
 		OutBufferCount = MultiplexCount;
 	}
+	else if (SignalBufferLayout == SIGNAL_BUFFER_LAYOUT_PENUMBRA_REJECTION)
+	{
+		const uint BatchSignalCount = MultiplexCount / 2;
+		OutBufferCount = BatchSignalCount;
+
+		UNROLL
+		for (uint BatchSignalId = 0; BatchSignalId < BatchSignalCount; BatchSignalId++)
+		{
+			uint MultiplexId = BatchSignalId * 2;
+
+			float NormalizationFactor = SafeRcp(Sample[MultiplexId + 0].AccumulatedSampleCount);
+
+			float NormalizedWorldBluringRadius = Sample[MultiplexId + 0].WorldBluringRadius * NormalizationFactor;
+
+			// Samples are normalised when doing history preconvolution, so the number of is either 0 or 1 on both momments. Therefore
+			// Sample[MultiplexId + 0].AccumulatedSampleCount == Sample[MultiplexId + 1].AccumulatedSampleCount.
+			OutRawSample[BatchSignalId] = float4(
+				Sample[MultiplexId + 0].LeavingRayCount * NormalizationFactor,
+				Sample[MultiplexId + 1].LeavingRayCount * NormalizationFactor,
+				Sample[MultiplexId + 0].AccumulatedSampleCount,
+				NormalizedWorldBluringRadius == WORLD_RADIUS_MISS ? -1 : NormalizedWorldBluringRadius);
+		}
+	}
 	else if (SignalBufferLayout == SIGNAL_BUFFER_LAYOUT_REFLECTIONS_REJECTION)
 	{
 		// Outputs scene color to be compatible with a SSR output.
diff --git a/Engine/Shaders/Private/ScreenSpaceDenoise/SSDSpatialAccumulation.usf b/Engine/Shaders/Private/ScreenSpaceDenoise/SSDSpatialAccumulation.usf
--- a/Engine/Shaders/Private/ScreenSpaceDenoise/SSDSpatialAccumulation.usf
+++ b/Engine/Shaders/Private/ScreenSpaceDenoise/SSDSpatialAccumulation.usf
@@ -37,6 +37,7 @@
 #define CONFIG_SIGNAL_BATCH_SIZE DIM_SIGNAL_BATCH_SIZE
 
 
+// Configures all the pass for each individual signals.
 #if CONFIG_SIGNAL_PROCESSING == SIGNAL_PROCESSING_MONOCHROMATIC_PENUMBRA
 	//#define CONFIG_VIEW_POSITION_BASED_BILATERAL 2
 	#define CONFIG_VIEW_POSITION_BASED_BILATERAL 3
@@ -55,6 +56,13 @@
 		#define CONFIG_BILATERAL_DISTANCE_COMPUTATION SIGNAL_WORLD_FREQUENCY_PRECOMPUTED_BLURING_RADIUS
 		#define CONFIG_OUTPUT_MODE          OUTPUT_MODE_DRB
 
+	#elif DIM_STAGE == STAGE_REJECTION_PRE_CONVOLUTION
+		// Input and output layout.
+		#define CONFIG_SIGNAL_INPUT_LAYOUT  SIGNAL_BUFFER_LAYOUT_PENUMBRA_RECONSTRUCTION
+		#define CONFIG_SIGNAL_OUTPUT_LAYOUT SIGNAL_BUFFER_LAYOUT_PENUMBRA_REJECTION
+
+		#define CONFIG_BILATERAL_DISTANCE_COMPUTATION SIGNAL_WORLD_FREQUENCY_PRECOMPUTED_BLURING_RADIUS
+
 	#elif DIM_STAGE == STAGE_POST_FILTERING
 		// Input and output layout.
 		#define CONFIG_SIGNAL_INPUT_LAYOUT  SIGNAL_BUFFER_LAYOUT_PENUMBRA_HISTORY
@@ -101,22 +109,13 @@
 		#define CONFIG_SIGNAL_INPUT_LAYOUT  SIGNAL_BUFFER_LAYOUT_REFLECTIONS_HISTORY
 		#define CONFIG_SIGNAL_OUTPUT_LAYOUT SIGNAL_BUFFER_LAYOUT_REFLECTIONS_REJECTION
 
-		#define CONFIG_SAMPLE_SET           SAMPLE_SET_3X3_PLUS
-
 		// Can only be done using tokoyashi because have more than one sample at a time.
 		#define CONFIG_TOKOYASHI_LOBE_BILATERAL 1
 		#define CONFIG_TOKOYASHI_AXES_BILATERAL 1
 
-		// Normalize the input, because want to measure the spatial variance regardless of how many samples where used to reconstruct the signal.
-		#define CONFIG_NORMALIZE_INPUT 1
-
 		#define CONFIG_ACCUMULATION_COLOR_SPACE REFLECTIONS_REJECTION_COLOR_SPACE
 		#define CONFIG_OUTPUT_COLOR_SPACE       REFLECTIONS_REJECTION_COLOR_SPACE
 
-		// Output the 2 momment because history rejection is varaiance based, and may flicker with momment 2 loss since the pre
-		// convolution will reduce the variance of momment 1.
-		#define CONFIG_OUTPUT_MODE OUTPUT_MODE_2MOMMENT_SUM
-
 	#elif DIM_STAGE == STAGE_POST_FILTERING
 		// Input and output layout.
 		#define CONFIG_SIGNAL_INPUT_LAYOUT  SIGNAL_BUFFER_LAYOUT_REFLECTIONS_HISTORY
@@ -194,6 +193,19 @@
 	#error Unknown signal processing.
 #endif
 
+// Configures pass regardless of the signals.
+#if DIM_STAGE == STAGE_REJECTION_PRE_CONVOLUTION
+	#define CONFIG_SAMPLE_SET           SAMPLE_SET_3X3_PLUS
+
+	// Normalize the input, because want to measure the spatial variance regardless of how many samples where used to reconstruct the signal.
+	#define CONFIG_NORMALIZE_INPUT 1
+
+	// Output the 2 momment because history rejection is varaiance based, and may flicker with momment 2 loss since the pre
+	// convolution will reduce the variance of momment 1.
+	#define CONFIG_OUTPUT_MODE OUTPUT_MODE_2MOMMENT_SUM
+
+#endif
+
 
 //------------------------------------------------------- CONFIG DISABLED DEFAULTS
 
@@ -565,13 +577,14 @@ void MainCS(
 		}
 		#elif CONFIG_OUTPUT_MODE == OUTPUT_MODE_2MOMMENT_SUM
 		{
-			#if CONFIG_SIGNAL_BATCH_SIZE != 1
-				#error Unable to batch multiple signals at the same time.
-			#endif
+			MultiplexCount = 2 * CONFIG_SIGNAL_BATCH_SIZE;
 
-			OutputSamples.Array[0] = SignalAccumulators.Array[0].Moment1;
-			OutputSamples.Array[1] = SignalAccumulators.Array[0].Moment2;
-			MultiplexCount = 2;
+			FIX_UNROLL(CONFIG_SIGNAL_BATCH_SIZE)
+			for (uint BatchedSignalId = 0; BatchedSignalId < CONFIG_SIGNAL_BATCH_SIZE; BatchedSignalId++)
+			{
+				OutputSamples.Array[BatchedSignalId * 2 + 0] = SignalAccumulators.Array[BatchedSignalId * CONFIG_MULTIPLEXED_SIGNALS_PER_SIGNAL_DOMAIN + 0].Moment1;
+				OutputSamples.Array[BatchedSignalId * 2 + 1] = SignalAccumulators.Array[BatchedSignalId * CONFIG_MULTIPLEXED_SIGNALS_PER_SIGNAL_DOMAIN + 0].Moment2;
+			}
 		}
 		#elif CONFIG_OUTPUT_MODE == OUTPUT_MODE_DRB
 		{
diff --git a/Engine/Shaders/Private/ScreenSpaceDenoise/SSDTemporalAccumulation.usf b/Engine/Shaders/Private/ScreenSpaceDenoise/SSDTemporalAccumulation.usf
--- a/Engine/Shaders/Private/ScreenSpaceDenoise/SSDTemporalAccumulation.usf
+++ b/Engine/Shaders/Private/ScreenSpaceDenoise/SSDTemporalAccumulation.usf
@@ -45,6 +45,15 @@
 	// TODO: a bit hacky, need find better solution.
 	#define CONFIG_USE_NEAREST_HISTORY 1
 
+	// Use pre transformed rejection buffer that contains pre transformed momment 1 & 2.
+	#if 0
+		#define CONFIG_USE_REJECTION_BUFFER               1
+		#define CONFIG_SIGNAL_HISTORY_REJECTION_LAYOUT    SIGNAL_BUFFER_LAYOUT_PENUMBRA_REJECTION
+		
+		#define CONFIG_REJECTION_INPUT_MODE               REJECTION_INPUT_MODE_2PRETRANSFORMED_MOMMENTS
+		#define CONFIG_REJECTION_SAMPLE_SET               REJECTION_SAMPLE_SET_SOBEK2018
+	#endif
+
 #elif CONFIG_SIGNAL_PROCESSING == SIGNAL_PROCESSING_REFLECTIONS
 	// Input and output layout.
 	#define CONFIG_SIGNAL_INPUT_LAYOUT   SIGNAL_BUFFER_LAYOUT_REFLECTIONS_HISTORY
@@ -383,8 +392,6 @@ void TemporallyAccumulate(
 			if (CONFIG_SIGNAL_PROCESSING == SIGNAL_PROCESSING_MONOCHROMATIC_PENUMBRA)
 				KernelConfig.BilateralDistanceComputation = SIGNAL_WORLD_FREQUENCY_PRECOMPUTED_BLURING_RADIUS;
 			
-			KernelConfig.MultiplexedSignalsPerSignalDomain = CONFIG_MULTIPLEXED_SIGNALS_PER_SIGNAL_DOMAIN;
-
 			#if CONFIG_REJECTION_SAMPLE_SET == REJECTION_SAMPLE_SET_SOBEK2018
 			{
 				KernelConfig.SampleSet = SAMPLE_SET_3X3_SOBEK2018;
@@ -402,6 +409,9 @@ void TemporallyAccumulate(
 
 			if (bUsesHistoryRejectionSignal)
 			{
+				// History rejection have two moments of the signal being denoised.
+				KernelConfig.MultiplexedSignalsPerSignalDomain = 2;
+				
 				KernelConfig.BufferLayout = CONFIG_SIGNAL_HISTORY_REJECTION_LAYOUT;
 				KernelConfig.bNormalizeSample = false;
 			
@@ -416,6 +426,7 @@ void TemporallyAccumulate(
 			}
 			else
 			{
+				KernelConfig.MultiplexedSignalsPerSignalDomain = CONFIG_MULTIPLEXED_SIGNALS_PER_SIGNAL_DOMAIN;
 				KernelConfig.BufferLayout = CONFIG_SIGNAL_INPUT_LAYOUT;
 				KernelConfig.bNormalizeSample = true;
 			
@@ -500,12 +511,9 @@ void TemporallyAccumulate(
 			}
 			#elif CONFIG_REJECTION_INPUT_MODE == REJECTION_INPUT_MODE_2PRETRANSFORMED_MOMMENTS
 			{
-				#if CONFIG_SIGNAL_BATCH_SIZE != 1
-					#error Unable to support more than one signal.
-				#endif
-				float NormalizeFactor = SafeRcp(SignalAccumulators.Array[0].Moment1.AccumulatedSampleCount);
-				NeighborMoment1 = Mul(SignalAccumulators.Array[0].Moment1, NormalizeFactor);
-				NeighborMoment2 = Mul(SignalAccumulators.Array[1].Moment1, NormalizeFactor);
+				float NormalizeFactor = SafeRcp(SignalAccumulators.Array[BatchedSignalId * 2 + 0].Moment1.AccumulatedSampleCount);
+				NeighborMoment1 = Mul(SignalAccumulators.Array[BatchedSignalId * 2 + 0].Moment1, NormalizeFactor);
+				NeighborMoment2 = Mul(SignalAccumulators.Array[BatchedSignalId * 2 + 1].Moment1, NormalizeFactor);
 			}
 			#else
 				#error NOrmalized samples. 
diff --git a/Engine/Source/Runtime/Renderer/Private/ScreenSpaceDenoise.cpp b/Engine/Source/Runtime/Renderer/Private/ScreenSpaceDenoise.cpp
--- a/Engine/Source/Runtime/Renderer/Private/ScreenSpaceDenoise.cpp
+++ b/Engine/Source/Runtime/Renderer/Private/ScreenSpaceDenoise.cpp
@@ -148,7 +148,9 @@ static bool SignalUsesInjestion(ESignalProcessing SignalProcessing)
 /** Returns whether a signal processing uses a history rejection pre convolution pass. */
 static bool SignalUsesRejectionPreConvolution(ESignalProcessing SignalProcessing)
 {
-	return SignalProcessing == ESignalProcessing::Reflections;
+	return (
+		//SignalProcessing == ESignalProcessing::MonochromaticPenumbra ||
+		SignalProcessing == ESignalProcessing::Reflections);
 }
 
 /** Returns whether a signal processing uses a history rejection pre convolution pass. */
@@ -248,10 +250,10 @@ const TCHAR* const kReconstructionResourceNames[] = {
 
 const TCHAR* const kRejectionPreConvolutionResourceNames[] = {
 	// Penumbra
-	nullptr,
-	nullptr,
-	nullptr,
-	nullptr,
+	TEXT("ShadowRejectionPreConvolution0"),
+	TEXT("ShadowRejectionPreConvolution1"),
+	TEXT("ShadowRejectionPreConvolution2"),
+	TEXT("ShadowRejectionPreConvolution3"),
 
 	// Reflections
 	TEXT("ReflectionsRejectionPreConvolution0"),
@@ -885,7 +887,15 @@ static void DenoiseSignalAtConstantPixelDensity(
 					RejectionSignalProcessingDescs[i] = HistoryDescs[i];
 				}
 
-				if (Settings.SignalProcessing == ESignalProcessing::Reflections)
+				if (Settings.SignalProcessing == ESignalProcessing::MonochromaticPenumbra)
+				{
+					for (int32 BatchedSignalId = 0; BatchedSignalId < Settings.SignalBatchSize; BatchedSignalId++)
+					{
+						RejectionSignalProcessingDescs[BatchedSignalId].Format = PF_FloatRGBA;
+					}
+					RejectionTextureCount = Settings.SignalBatchSize;
+				}
+				else if (Settings.SignalProcessing == ESignalProcessing::Reflections)
 				{
 					RejectionSignalProcessingDescs[0].Format = PF_FloatRGBA;
 					RejectionSignalProcessingDescs[1].Format = PF_G16R16F;
