Cherry-pick 5357861: Disables nearest history reprojection for shadows that is causing a problem with spot light for some reasons.

diff --git a/Engine/Shaders/Private/ScreenSpaceDenoise/SSDTemporalAccumulation.usf b/Engine/Shaders/Private/ScreenSpaceDenoise/SSDTemporalAccumulation.usf
--- a/Engine/Shaders/Private/ScreenSpaceDenoise/SSDTemporalAccumulation.usf
+++ b/Engine/Shaders/Private/ScreenSpaceDenoise/SSDTemporalAccumulation.usf
@@ -43,7 +43,8 @@
 
 	// Uses nearest to not leak informations on geometric edges by bilateral. This is OK because just blury greyscale details.
 	// TODO: a bit hacky, need find better solution.
-	#define CONFIG_USE_NEAREST_HISTORY 1
+	// TODO: cause a regression on spot light for some reasons.
+	//#define CONFIG_USE_NEAREST_HISTORY 1
 
 	// Use pre transformed rejection buffer that contains pre transformed momment 1 & 2.
 	#if 0
