Clamp sky integration and denoiser output to be the valid half float range.

diff --git a/Engine/Shaders/Private/Common.ush b/Engine/Shaders/Private/Common.ush
--- a/Engine/Shaders/Private/Common.ush
+++ b/Engine/Shaders/Private/Common.ush
@@ -71,7 +71,7 @@
 #define USE_GLOBAL_CLIP_PLANE (PLATFORM_SUPPORTS_GLOBAL_CLIP_PLANE && PROJECT_ALLOW_GLOBAL_CLIP_PLANE && !MATERIAL_DOMAIN_POSTPROCESS && !MATERIAL_DOMAIN_UI)
 
 const static MaterialFloat PI = 3.1415926535897932f;
-
+const static float MaxHalfFloat = 65504.0f;
 #ifndef USE_RAYTRACED_TEXTURE_RAYCONE_LOD
 #define USE_RAYTRACED_TEXTURE_RAYCONE_LOD (RAYHITGROUPSHADER)
 #endif
@@ -94,6 +94,11 @@ float ComputeRayConeLod(Texture2D Tex)
 }
 #endif
 
+float  ClampToMaxHalfFloat(float X)  { return min(MaxHalfFloat, X); }
+float2 ClampToMaxHalfFloat(float2 X) { return min(MaxHalfFloat.xx, X); }
+float3 ClampToMaxHalfFloat(float3 X) { return min(MaxHalfFloat.xxx, X); }
+float4 ClampToMaxHalfFloat(float4 X) { return min(MaxHalfFloat.xxxx, X); }
+
 // This would need to be a #define in GLSL to ignore the SamplerState, however, it is currently a function call in HLSL
 // for type checking of the parameters - ironically the type checking is really only needed in GLSL!
 MaterialFloat4 Texture1DSample(Texture1D Tex, SamplerState Sampler, float UV)
diff --git a/Engine/Shaders/Private/RayTracing/RayTracingSkyLightRGS.usf b/Engine/Shaders/Private/RayTracing/RayTracingSkyLightRGS.usf
--- a/Engine/Shaders/Private/RayTracing/RayTracingSkyLightRGS.usf
+++ b/Engine/Shaders/Private/RayTracing/RayTracingSkyLightRGS.usf
@@ -362,7 +362,7 @@ void SkyLightRGS()
 	DiffuseExitantRadiance.g = Albedo.g > 0.0 ? DiffuseExitantRadiance.g / Albedo.g : DiffuseExitantRadiance.g;
 	DiffuseExitantRadiance.b = Albedo.b > 0.0 ? DiffuseExitantRadiance.b / Albedo.b : DiffuseExitantRadiance.b;
 
-	RWOcclusionMaskUAV[PixelCoord] = float4(DiffuseExitantRadiance, AmbientOcclusion);
+	RWOcclusionMaskUAV[PixelCoord] = float4(ClampToMaxHalfFloat(DiffuseExitantRadiance.xyz), AmbientOcclusion);
 
 	float HitDistance = -1.0;
 	if (HitCount > 0.0)
diff --git a/Engine/Shaders/Private/ReflectionEnvironmentShaders.usf b/Engine/Shaders/Private/ReflectionEnvironmentShaders.usf
--- a/Engine/Shaders/Private/ReflectionEnvironmentShaders.usf
+++ b/Engine/Shaders/Private/ReflectionEnvironmentShaders.usf
@@ -51,7 +51,7 @@ void CopySceneColorToCubeFaceColorPS(
 	float3 SceneColor = Texture2DSample(InTexture, InTextureSampler, Input.UV).rgb;
 
 	// Convert INF's to valid values
-	SceneColor = min(SceneColor, 65503);
+	SceneColor = ClampToMaxHalfFloat(SceneColor);
 	float3 WorldPosition = Input.ScreenVector * SceneDepth + View.WorldCameraOrigin;
 
 	float Alpha = 1;
diff --git a/Engine/Shaders/Private/ScreenSpaceDenoise/SSDSignalCore.ush b/Engine/Shaders/Private/ScreenSpaceDenoise/SSDSignalCore.ush
--- a/Engine/Shaders/Private/ScreenSpaceDenoise/SSDSignalCore.ush
+++ b/Engine/Shaders/Private/ScreenSpaceDenoise/SSDSignalCore.ush
@@ -641,11 +641,11 @@ void OutputMultiplexedSignal(
 
 	// Output the raw encoded sample according to number of RT they requires.
 	if (BufferCount >= 1)
-		OutputSignalBuffer0[PixelPosition] = RawSample[0];
+		OutputSignalBuffer0[PixelPosition] = ClampToMaxHalfFloat(RawSample[0]);
 	if (BufferCount >= 2)
-		OutputSignalBuffer1[PixelPosition] = RawSample[1];
+		OutputSignalBuffer1[PixelPosition] = ClampToMaxHalfFloat(RawSample[1]);
 	if (BufferCount >= 3)
-		OutputSignalBuffer2[PixelPosition] = RawSample[2];
+		OutputSignalBuffer2[PixelPosition] = ClampToMaxHalfFloat(RawSample[2]);
 	if (BufferCount >= 4)
-		OutputSignalBuffer3[PixelPosition] = RawSample[3];
+		OutputSignalBuffer3[PixelPosition] = ClampToMaxHalfFloat(RawSample[3]);
 }
