Add pre-exposition support for all raytracing pass to perserve precision with high lighting value.
Add half float range validation for all radiance/luminance output data in raytracing passes.

diff --git a/Engine/Shaders/Private/Common.ush b/Engine/Shaders/Private/Common.ush
--- a/Engine/Shaders/Private/Common.ush
+++ b/Engine/Shaders/Private/Common.ush
@@ -94,10 +94,10 @@ float ComputeRayConeLod(Texture2D Tex)
 }
 #endif
 
-float  ClampToMaxHalfFloat(float X)  { return min(MaxHalfFloat, X); }
-float2 ClampToMaxHalfFloat(float2 X) { return min(MaxHalfFloat.xx, X); }
-float3 ClampToMaxHalfFloat(float3 X) { return min(MaxHalfFloat.xxx, X); }
-float4 ClampToMaxHalfFloat(float4 X) { return min(MaxHalfFloat.xxxx, X); }
+float  ClampToHalfFloatRange(float  X) { return clamp(X, float(0), MaxHalfFloat); }
+float2 ClampToHalfFloatRange(float2 X) { return clamp(X, float(0).xx, MaxHalfFloat.xx); }
+float3 ClampToHalfFloatRange(float3 X) { return clamp(X, float(0).xxx, MaxHalfFloat.xxx); }
+float4 ClampToHalfFloatRange(float4 X) { return clamp(X, float(0).xxxx, MaxHalfFloat.xxxx); }
 
 // This would need to be a #define in GLSL to ignore the SamplerState, however, it is currently a function call in HLSL
 // for type checking of the parameters - ironically the type checking is really only needed in GLSL!
diff --git a/Engine/Shaders/Private/RayTracing/RayTracingGlobalIlluminationRGS.usf b/Engine/Shaders/Private/RayTracing/RayTracingGlobalIlluminationRGS.usf
--- a/Engine/Shaders/Private/RayTracing/RayTracingGlobalIlluminationRGS.usf
+++ b/Engine/Shaders/Private/RayTracing/RayTracingGlobalIlluminationRGS.usf
@@ -365,7 +365,12 @@ void GlobalIlluminationRGS()
 	}
 
 	AmbientOcclusion = saturate(AmbientOcclusion);
-	Irradiance = max(Irradiance, 0.0);
+
+#if USE_PREEXPOSURE
+	Irradiance *= View.PreExposure;
+#endif
+
+	Irradiance = ClampToHalfFloatRange(Irradiance);
 	RWGlobalIlluminationUAV[DispatchThreadId] = float4(Irradiance, AmbientOcclusion);
 	RWRayDistanceUAV[DispatchThreadId] = float2(HitDistance, SamplesPerPixel);
 	// For AO denoiser..
diff --git a/Engine/Shaders/Private/RayTracing/RayTracingRectLightRGS.usf b/Engine/Shaders/Private/RayTracing/RayTracingRectLightRGS.usf
--- a/Engine/Shaders/Private/RayTracing/RayTracingRectLightRGS.usf
+++ b/Engine/Shaders/Private/RayTracing/RayTracingRectLightRGS.usf
@@ -332,8 +332,14 @@ void RectLightRGS()
 		}
 	}
 
-	float3 Result = (SamplesPerPixel > 0) ? ExitantRadiance / SamplesPerPixel : ExitantRadiance;
-	RWLuminanceUAV[PixelCoord].rgb = Result;
+	float3 OutputRadiance = (SamplesPerPixel > 0) ? ExitantRadiance / SamplesPerPixel : ExitantRadiance;
+
+#if USE_PREEXPOSURE
+	OutputRadiance *= View.PreExposure;
+#endif
+
+	OutputRadiance = ClampToHalfFloatRange(OutputRadiance);
+	RWLuminanceUAV[PixelCoord].rgb = OutputRadiance;
 	RWLuminanceUAV[PixelCoord].a = 1.0;
 	RWRayDistanceUAV[PixelCoord] = (HitCount > 0.0) ? RayDistance / HitCount : 1.0e27;
 }
diff --git a/Engine/Shaders/Private/RayTracing/RayTracingReflections.usf b/Engine/Shaders/Private/RayTracing/RayTracingReflections.usf
--- a/Engine/Shaders/Private/RayTracing/RayTracingReflections.usf
+++ b/Engine/Shaders/Private/RayTracing/RayTracingReflections.usf
@@ -428,7 +428,11 @@ void RayTracingReflectionsRGS()
 			ReflectedColor *= DistanceFade * max(RoughnessFade, ClearCoatRoughnessFade) * rcp(ReflectedWeight);
 		}
 
-		ColorOutput[DispatchThreadId] = ReflectedColor;
+	#if USE_PREEXPOSURE
+		ReflectedColor *= View.PreExposure;
+	#endif
+
+		ColorOutput[DispatchThreadId] = ClampToHalfFloatRange(ReflectedColor);
 		RayHitDistanceOutput[DispatchThreadId] = HitDistance;
 	}
 }
diff --git a/Engine/Shaders/Private/RayTracing/RayTracingSkyLightRGS.usf b/Engine/Shaders/Private/RayTracing/RayTracingSkyLightRGS.usf
--- a/Engine/Shaders/Private/RayTracing/RayTracingSkyLightRGS.usf
+++ b/Engine/Shaders/Private/RayTracing/RayTracingSkyLightRGS.usf
@@ -362,7 +362,11 @@ void SkyLightRGS()
 	DiffuseExitantRadiance.g = Albedo.g > 0.0 ? DiffuseExitantRadiance.g / Albedo.g : DiffuseExitantRadiance.g;
 	DiffuseExitantRadiance.b = Albedo.b > 0.0 ? DiffuseExitantRadiance.b / Albedo.b : DiffuseExitantRadiance.b;
 
-	RWOcclusionMaskUAV[PixelCoord] = float4(ClampToMaxHalfFloat(DiffuseExitantRadiance.xyz), AmbientOcclusion);
+#if USE_PREEXPOSURE
+	DiffuseExitantRadiance.rgb *= View.PreExposure;
+#endif
+
+	RWOcclusionMaskUAV[PixelCoord] = float4(ClampToHalfFloatRange(DiffuseExitantRadiance.rgb), AmbientOcclusion);
 
 	float HitDistance = -1.0;
 	if (HitCount > 0.0)
diff --git a/Engine/Shaders/Private/RayTracing/RayTracingTranslucency.usf b/Engine/Shaders/Private/RayTracing/RayTracingTranslucency.usf
--- a/Engine/Shaders/Private/RayTracing/RayTracingTranslucency.usf
+++ b/Engine/Shaders/Private/RayTracing/RayTracingTranslucency.usf
@@ -276,6 +276,12 @@ void RayTracingTranslucencyRGS()
 		// Use the scene opaque radiance for pixel that has not intersected any transparent.
 		PathRadiance = Texture2DSampleLevel(SceneTexturesStruct.SceneColorTexture, SceneTexturesStruct.SceneColorTextureSampler, UV, 0).xyz;
 	}
+
+#if USE_PREEXPOSURE
+	PathRadiance.rgb *= View.PreExposure;
+#endif
+
+	PathRadiance = ClampToHalfFloatRange(PathRadiance);
 	ColorOutput[DispatchThreadId] = float4(PathRadiance, 0.0f);
 	RayHitDistanceOutput[DispatchThreadId] = HitDistance;
 }
diff --git a/Engine/Shaders/Private/ReflectionEnvironmentShaders.usf b/Engine/Shaders/Private/ReflectionEnvironmentShaders.usf
--- a/Engine/Shaders/Private/ReflectionEnvironmentShaders.usf
+++ b/Engine/Shaders/Private/ReflectionEnvironmentShaders.usf
@@ -51,7 +51,7 @@ void CopySceneColorToCubeFaceColorPS(
 	float3 SceneColor = Texture2DSample(InTexture, InTextureSampler, Input.UV).rgb;
 
 	// Convert INF's to valid values
-	SceneColor = ClampToMaxHalfFloat(SceneColor);
+	SceneColor = ClampToHalfFloatRange(SceneColor);
 	float3 WorldPosition = Input.ScreenVector * SceneDepth + View.WorldCameraOrigin;
 
 	float Alpha = 1;
diff --git a/Engine/Shaders/Private/ScreenSpaceDenoise/SSDSignalCore.ush b/Engine/Shaders/Private/ScreenSpaceDenoise/SSDSignalCore.ush
--- a/Engine/Shaders/Private/ScreenSpaceDenoise/SSDSignalCore.ush
+++ b/Engine/Shaders/Private/ScreenSpaceDenoise/SSDSignalCore.ush
@@ -641,11 +641,11 @@ void OutputMultiplexedSignal(
 
 	// Output the raw encoded sample according to number of RT they requires.
 	if (BufferCount >= 1)
-		OutputSignalBuffer0[PixelPosition] = ClampToMaxHalfFloat(RawSample[0]);
+		OutputSignalBuffer0[PixelPosition] = min(RawSample[0], MaxHalfFloat);
 	if (BufferCount >= 2)
-		OutputSignalBuffer1[PixelPosition] = ClampToMaxHalfFloat(RawSample[1]);
+		OutputSignalBuffer1[PixelPosition] = min(RawSample[1], MaxHalfFloat);
 	if (BufferCount >= 3)
-		OutputSignalBuffer2[PixelPosition] = ClampToMaxHalfFloat(RawSample[2]);
+		OutputSignalBuffer2[PixelPosition] = min(RawSample[2], MaxHalfFloat);
 	if (BufferCount >= 4)
-		OutputSignalBuffer3[PixelPosition] = ClampToMaxHalfFloat(RawSample[3]);
+		OutputSignalBuffer3[PixelPosition] = min(RawSample[3], MaxHalfFloat);
 }
