Enable textured rect light into the translucency pass.

diff --git a/Engine/Shaders/Private/RayTracing/RayTracingTranslucency.usf b/Engine/Shaders/Private/RayTracing/RayTracingTranslucency.usf
--- a/Engine/Shaders/Private/RayTracing/RayTracingTranslucency.usf
+++ b/Engine/Shaders/Private/RayTracing/RayTracingTranslucency.usf
@@ -12,8 +12,8 @@ float2 SobolIndexToUniformUnitSquare(uint2 SobolRand)
 }
 
 #define SUPPORT_CONTACT_SHADOWS		0
-#define USE_SOURCE_TEXTURE			0
-#define USE_SOURCE_TEXTURE_ARRAY	0
+#define USE_SOURCE_TEXTURE			1
+#define USE_SOURCE_TEXTURE_ARRAY	1
 
 #define LTCMatTexture				RaytracingLightsData.LTCMatTexture
 #define LTCMatSampler				RaytracingLightsData.LTCMatSampler
