* Split lighting and light culling in raytraced reflection/translucency passes.
 * Pack light data for more effiicient load/better cache utilization

This improves raytraced translucency and reflection respectively by 20% and 15% .

diff --git a/Engine/Shaders/Private/RayTracing/RayTracingCommon.ush b/Engine/Shaders/Private/RayTracing/RayTracingCommon.ush
--- a/Engine/Shaders/Private/RayTracing/RayTracingCommon.ush
+++ b/Engine/Shaders/Private/RayTracing/RayTracingCommon.ush
@@ -159,6 +159,15 @@ FPackedMaterialClosestHitPayload PackRayTracingPayload(FMaterialClosestHitPayloa
 	return Output;
 }
 
+float3 UnpackRayTracingNormal(FPackedMaterialClosestHitPayload Input)
+{
+	float3 OutWorldNormal;
+	OutWorldNormal.x = f16tof32(Input.RadianceAndNormal[1] >> 16);
+	OutWorldNormal.y = f16tof32(Input.RadianceAndNormal[2]);
+	OutWorldNormal.z = f16tof32(Input.RadianceAndNormal[2] >> 16);
+	return OutWorldNormal;
+}
+
 FMaterialClosestHitPayload UnpackRayTracingPayload(FPackedMaterialClosestHitPayload Input, RayDesc Ray)
 {
 	FMaterialClosestHitPayload Output = (FMaterialClosestHitPayload)0;
diff --git a/Engine/Shaders/Private/RayTracing/RayTracingLightingCommon.ush b/Engine/Shaders/Private/RayTracing/RayTracingLightingCommon.ush
--- a/Engine/Shaders/Private/RayTracing/RayTracingLightingCommon.ush
+++ b/Engine/Shaders/Private/RayTracing/RayTracingLightingCommon.ush
@@ -18,6 +18,12 @@ bool IsHit(RayDesc Ray, FMaterialClosestHitPayload HitInfo)
 	return t > Ray.TMin && t < Ray.TMax;
 }
 
+bool IsHit(RayDesc Ray, FPackedMaterialClosestHitPayload HitInfo)
+{
+	float t = HitInfo.HitT;
+	return t > Ray.TMin && t < Ray.TMax;
+}
+
 float GetRoughnessFade(float Roughness, float MaxRoughness)
 {
 	float RoughnessMaskScale = -2.0 / MaxRoughness;
@@ -166,12 +172,129 @@ float3 SampleAreaLightDirection(
 	return ShadowRayDirection;
 }
 
-float3 ComputeDirectLighting(
-	in RayDesc Ray, 
-	in FRayCone RayCone, 
-	in FMaterialClosestHitPayload Payload, 
+FPackedMaterialClosestHitPayload TraceRayInternalPacked(in RaytracingAccelerationStructure TLAS,
+	in uint RayFlags,
+	in uint InstanceInclusionMask,
+	in uint RayContributionToHitGroupIndex,
+	in uint MultiplierForGeometryContributionToShaderIndex,
+	in uint MissShaderIndex,
+	in RayDesc Ray,
+	inout FRayCone RayCone)
+{
+	FPackedMaterialClosestHitPayload Payload = (FPackedMaterialClosestHitPayload)0;
+	Payload.RayCone = RayCone;
+
+	TraceRay
+	(
+		TLAS,
+		RayFlags,
+		InstanceInclusionMask,
+		RayContributionToHitGroupIndex,
+		MultiplierForGeometryContributionToShaderIndex,
+		MissShaderIndex,
+		Ray,
+		Payload
+	);
+
+	RayCone = Payload.RayCone;
+	return Payload;
+}
+
+uint2 CullDirectLighting(
+	in RayDesc Ray,
+	in FRayCone RayCone,
+	in float3 WorldPosition,
+	in float3 WorldNormal,
+	in RandomSequence RandSequence,
+	in bool bReflectedShadows,
+	in float ShadowMaxNormalBias)
+{
+
+	uint2 LightCullMask = (uint2)0;
+
+	//#dxr_todo: think on the more efficient way of iterate through lights (split in light types, etc)
+	for (uint DirIndex = 0; DirIndex < RaytracingLightsDataPacked.Count; DirIndex++)
+	{
+		FDeferredLightData LightData = (FDeferredLightData)0;
+		uint Lit = 1;
+
+		uint LightType = RaytracingLightsDataPacked.Type_LightProfileIndex_RectLightTextureIndex[DirIndex].x;
+		LightData.Position = RaytracingLightsDataPacked.LightPosition_InvRadius[DirIndex].xyz;
+		LightData.InvRadius = RaytracingLightsDataPacked.LightPosition_InvRadius[DirIndex].w;
+		LightData.Color = RaytracingLightsDataPacked.LightColor_SpecularScale[DirIndex].xyz;
+		LightData.FalloffExponent = RaytracingLightsDataPacked.Direction_FalloffExponent[DirIndex].w;
+		LightData.Direction = RaytracingLightsDataPacked.Direction_FalloffExponent[DirIndex].xyz;
+		LightData.Tangent = RaytracingLightsDataPacked.Tangent_SourceRadius[DirIndex].xyz;
+		LightData.SpotAngles = RaytracingLightsDataPacked.SpotAngles_SourceLength_SoftSourceRadius[DirIndex].xy;
+		LightData.SourceRadius = RaytracingLightsDataPacked.Tangent_SourceRadius[DirIndex].w;
+		LightData.SourceLength = RaytracingLightsDataPacked.SpotAngles_SourceLength_SoftSourceRadius[DirIndex].z;
+		LightData.SoftSourceRadius = RaytracingLightsDataPacked.SpotAngles_SourceLength_SoftSourceRadius[DirIndex].w;
+		LightData.SpecularScale = RaytracingLightsDataPacked.LightColor_SpecularScale[DirIndex].w;
+		LightData.RectLightBarnCosAngle = RaytracingLightsDataPacked.DistanceFadeMAD_RectLightBarnCosAngle_RectLightBarnLength[DirIndex].z;
+		LightData.RectLightBarnLength = RaytracingLightsDataPacked.DistanceFadeMAD_RectLightBarnCosAngle_RectLightBarnLength[DirIndex].w;
+		LightData.ContactShadowLength = 0.0;
+		LightData.DistanceFadeMAD = RaytracingLightsDataPacked.DistanceFadeMAD_RectLightBarnCosAngle_RectLightBarnLength[DirIndex].xy;
+		const uint LightProfileIndex = RaytracingLightsDataPacked.Type_LightProfileIndex_RectLightTextureIndex[DirIndex].y;
+
+
+		LightData.ShadowMapChannelMask = float4(0, 0, 0, 0); // #dxr_todo;
+		LightData.ShadowedBits = 0; // Not lit dynamic shadows 
+		LightData.ContactShadowLengthInWS = false;
+
+		LightData.bRadialLight = (LightType != LIGHT_TYPE_DIRECTIONAL);
+		LightData.bSpotLight = (LightType == LIGHT_TYPE_SPOT);
+		LightData.bRectLight = (LightType == LIGHT_TYPE_RECT);
+
+		if (LightType == LIGHT_TYPE_DIRECTIONAL)
+		{
+			LightData.bInverseSquared = false;
+		}
+		else
+		{
+			LightData.bInverseSquared = LightData.FalloffExponent == 0;
+		}
+
+
+		float3 ShadowRayDirection;
+		// ToLight should not be normalized because its length is used to compute the shadow ray TMax 
+		float3 ToLight = LightData.Position - WorldPosition;
+		float LightMask = 1.0;
+
+		if (LightType == LIGHT_TYPE_DIRECTIONAL)
+		{
+			ShadowRayDirection = LightData.Direction;
+			ToLight = LightData.Direction * 100000.0f;
+		}
+		else
+		{
+			LightMask = GetLocalLightAttenuation(WorldPosition, LightData, ToLight, ShadowRayDirection);
+
+			// Skip the light sample that does not contribute anything due to attenuation.
+			if (LightMask <= 0.0)
+			{
+				Lit = 0;
+			}
+		}
+
+		// Skip the light sample pointing backwards
+		if (dot(WorldNormal, normalize(ToLight)) <= 0)
+		{
+			Lit = 0;
+		}
+
+		LightCullMask[DirIndex/32] |= Lit << (DirIndex%32);
+	}
+
+	return LightCullMask;
+}
+
+float3 ComputeDirectLightingCulled(
+	in uint2 LightCullMask,
+	in RayDesc Ray,
+	in FRayCone RayCone,
+	in FMaterialClosestHitPayload Payload,
 	in RandomSequence RandSequence,
-	in bool bReflectedShadows, 
+	in bool bReflectedShadows,
 	in float ShadowMaxNormalBias)
 {
 	float3 DirectLighting = float3(0.0, 0.0, 0.0);
@@ -197,43 +320,83 @@ float3 ComputeDirectLighting(
 	GBufferData.DiffuseColor = Payload.DiffuseColor;
 
 	FRectTexture RectTexture;
-	#if USE_SOURCE_TEXTURE_ARRAY 
-	RectTexture.SourceTexture0 = RaytracingLightsData.RectLightTexture0;
-	RectTexture.SourceTexture1 = RaytracingLightsData.RectLightTexture1;
-	RectTexture.SourceTexture2 = RaytracingLightsData.RectLightTexture2;
-	RectTexture.SourceTexture3 = RaytracingLightsData.RectLightTexture3;
-	RectTexture.SourceTexture4 = RaytracingLightsData.RectLightTexture4;
-	RectTexture.SourceTexture5 = RaytracingLightsData.RectLightTexture5;
-	RectTexture.SourceTexture6 = RaytracingLightsData.RectLightTexture6;
-	RectTexture.SourceTexture7 = RaytracingLightsData.RectLightTexture7;
+#if USE_SOURCE_TEXTURE_ARRAY 
+	RectTexture.SourceTexture0 = RaytracingLightsDataPacked.RectLightTexture0;
+	RectTexture.SourceTexture1 = RaytracingLightsDataPacked.RectLightTexture1;
+	RectTexture.SourceTexture2 = RaytracingLightsDataPacked.RectLightTexture2;
+	RectTexture.SourceTexture3 = RaytracingLightsDataPacked.RectLightTexture3;
+	RectTexture.SourceTexture4 = RaytracingLightsDataPacked.RectLightTexture4;
+	RectTexture.SourceTexture5 = RaytracingLightsDataPacked.RectLightTexture5;
+	RectTexture.SourceTexture6 = RaytracingLightsDataPacked.RectLightTexture6;
+	RectTexture.SourceTexture7 = RaytracingLightsDataPacked.RectLightTexture7;
 	RectTexture.SourceTextureIndex = 99;
-	#elif USE_SOURCE_TEXTURE
-	RectTexture = InitRectTexture(RaytracingLightsData.RectLightTexture0);
-	#else
-    RectTexture = (FRectTexture)0;
-	#endif
+#elif USE_SOURCE_TEXTURE
+	RectTexture = InitRectTexture(RaytracingLightsDataPacked.RectLightTexture0);
+#else
+	RectTexture = (FRectTexture)0;
+#endif
+
+	uint DirIndex = 0;
 
 	//#dxr_todo: think on the more efficient way of iterate through lights (split in light types, etc)
-	for (uint DirIndex = 0; DirIndex < RaytracingLightsData.Count; DirIndex++)
+	//while(LightCullMask.x || LightCullMask.y)
+	while (WaveActiveAnyTrue(LightCullMask.x || LightCullMask.y))
 	{
-		FDeferredLightData LightData = (FDeferredLightData)0;
+		const bool Active = LightCullMask.x || LightCullMask.y;
 
-		uint LightType = RaytracingLightsData.Type[DirIndex];
-		LightData.Position = RaytracingLightsData.LightPosition[DirIndex];
-		LightData.InvRadius = RaytracingLightsData.LightInvRadius[DirIndex];
-		LightData.Color = RaytracingLightsData.LightColor[DirIndex];
-		LightData.FalloffExponent = RaytracingLightsData.LightFalloffExponent[DirIndex];
-		LightData.Direction = RaytracingLightsData.Direction[DirIndex];
-		LightData.Tangent = RaytracingLightsData.Tangent[DirIndex];
-		LightData.SpotAngles = RaytracingLightsData.SpotAngles[DirIndex];
-		LightData.SourceRadius = RaytracingLightsData.SourceRadius[DirIndex];
-		LightData.SourceLength = RaytracingLightsData.SourceLength[DirIndex];
-		LightData.SoftSourceRadius = RaytracingLightsData.SoftSourceRadius[DirIndex];
-		LightData.SpecularScale = RaytracingLightsData.SpecularScale[DirIndex];
-		LightData.RectLightBarnCosAngle = RaytracingLightsData.RectLightBarnCosAngle[DirIndex];
-		LightData.RectLightBarnLength = RaytracingLightsData.RectLightBarnLength[DirIndex];
+		if (Active)
+		{
+			if ((LightCullMask.x & 0x1) == 0)
+			{
+				uint Shift = LightCullMask.x ? firstbitlow(LightCullMask.x) : 32;
+				DirIndex += Shift;
+				LightCullMask.x >>= Shift;
+
+				//fill in the top of LightCullMask.x
+				LightCullMask.x |= LightCullMask.y << (32 - Shift);
+				LightCullMask.y >>= Shift;
+
+				//may need to try once more
+				if ((LightCullMask.x & 0x1) == 0)
+				{
+					uint Shift = LightCullMask.x ? firstbitlow(LightCullMask.x) : 32;
+					DirIndex += Shift;
+					LightCullMask.x >>= Shift;
+
+					//fill in the top of LightCullMask.x
+					LightCullMask.x |= LightCullMask.y << (32 - Shift);
+					LightCullMask.y >>= Shift;
+				}
+			}
+		}
+		else
+		{
+			DirIndex = 0;
+		}
+
+		FDeferredLightData LightData = (FDeferredLightData)0;
+		uint LightType = RaytracingLightsDataPacked.Type_LightProfileIndex_RectLightTextureIndex[DirIndex].x;
+		LightData.Position = RaytracingLightsDataPacked.LightPosition_InvRadius[DirIndex].xyz;
+		LightData.InvRadius = RaytracingLightsDataPacked.LightPosition_InvRadius[DirIndex].w;
+		LightData.Color = RaytracingLightsDataPacked.LightColor_SpecularScale[DirIndex].xyz;
+		LightData.FalloffExponent = RaytracingLightsDataPacked.Direction_FalloffExponent[DirIndex].w;
+		LightData.Direction = RaytracingLightsDataPacked.Direction_FalloffExponent[DirIndex].xyz;
+		LightData.Tangent = RaytracingLightsDataPacked.Tangent_SourceRadius[DirIndex].xyz;
+		LightData.SpotAngles = RaytracingLightsDataPacked.SpotAngles_SourceLength_SoftSourceRadius[DirIndex].xy;
+		LightData.SourceRadius = RaytracingLightsDataPacked.Tangent_SourceRadius[DirIndex].w;
+		LightData.SourceLength = RaytracingLightsDataPacked.SpotAngles_SourceLength_SoftSourceRadius[DirIndex].z;
+		LightData.SoftSourceRadius = RaytracingLightsDataPacked.SpotAngles_SourceLength_SoftSourceRadius[DirIndex].w;
+		LightData.SpecularScale = RaytracingLightsDataPacked.LightColor_SpecularScale[DirIndex].w;
+		LightData.RectLightBarnCosAngle = RaytracingLightsDataPacked.DistanceFadeMAD_RectLightBarnCosAngle_RectLightBarnLength[DirIndex].z;
+		LightData.RectLightBarnLength = RaytracingLightsDataPacked.DistanceFadeMAD_RectLightBarnCosAngle_RectLightBarnLength[DirIndex].w;
 		LightData.ContactShadowLength = 0.0;
-		LightData.DistanceFadeMAD = RaytracingLightsData.DistanceFadeMAD[DirIndex];
+		LightData.DistanceFadeMAD = RaytracingLightsDataPacked.DistanceFadeMAD_RectLightBarnCosAngle_RectLightBarnLength[DirIndex].xy;
+
+	#if USE_SOURCE_TEXTURE_ARRAY 
+		RectTexture.SourceTextureIndex = RaytracingLightsDataPacked.Type_LightProfileIndex_RectLightTextureIndex[DirIndex].z;
+	#endif
+		uint LightProfileIndex = RaytracingLightsDataPacked.Type_LightProfileIndex_RectLightTextureIndex[DirIndex].y;
+
 		LightData.ShadowMapChannelMask = float4(0, 0, 0, 0); // #dxr_todo;
 		LightData.ShadowedBits = 0; // Not lit dynamic shadows 
 		LightData.ContactShadowLengthInWS = false;
@@ -242,10 +405,6 @@ float3 ComputeDirectLighting(
 		LightData.bSpotLight = (LightType == LIGHT_TYPE_SPOT);
 		LightData.bRectLight = (LightType == LIGHT_TYPE_RECT);
 
-		#if USE_SOURCE_TEXTURE_ARRAY 
-		RectTexture.SourceTextureIndex = RaytracingLightsData.RectLightTextureIndex[DirIndex];
-		#endif
-
 		if (LightType == LIGHT_TYPE_DIRECTIONAL)
 		{
 			LightData.bInverseSquared = false;
@@ -258,7 +417,7 @@ float3 ComputeDirectLighting(
 		float4 LightAttenuation = 1.0f;
 
 		float LightProfileMultiplier = 1.0;
-		uint LightProfileIndex = RaytracingLightsData.LightProfileIndex[DirIndex];
+
 		if (LightProfileIndex >= 0)
 		{
 			LightProfileMultiplier = ComputeLightProfileMultiplier(WorldPosition, LightData.Position, LightData.Direction, LightProfileIndex);
@@ -277,18 +436,6 @@ float3 ComputeDirectLighting(
 		else
 		{
 			LightMask = GetLocalLightAttenuation(WorldPosition, LightData, ToLight, ShadowRayDirection);
-			
-			// Skip the light sample that does not contribute anything due to attenuation.
-			if (LightMask <= 0.0)
-			{
-				continue;
-			}
-		}
-
-		// Skip the light sample pointing backwards
-		if (dot(Payload.WorldNormal, normalize(ToLight)) <= 0)
-		{
-			continue;
 		}
 
 		//#dxr_todo: check perf improvements forcing TMin == TMax when bReflectedShadows = false, to avoid tracing rays inside if statements
@@ -306,15 +453,24 @@ float3 ComputeDirectLighting(
 			ShadowRay.Origin = WorldPosition;
 			ShadowRay.Direction = ShadowRayDirection;
 			ShadowRay.TMin = 1e-4f;
-			ShadowRay.TMax = length(ToLight);
+			ShadowRay.TMax = Active ? length(ToLight) : ShadowRay.TMin;
 			ApplyPositionBias(ShadowRay, Payload.WorldNormal, ShadowMaxNormalBias);
 
 			float HitT = TraceShadowRay(ShadowRay, ShadowRayCone);
-			AmbientOcclusion = HitT < 0.0f;
+			AmbientOcclusion = HitT < 0.0f && Active;
 		}
 
-		float3 LightContribution = GetDynamicLighting(WorldPosition, -CameraVector, GBufferData, AmbientOcclusion, GBufferData.ShadingModelID, LightData, LightAttenuation, 0.5, uint2(0, 0), RectTexture).xyz;
-		DirectLighting += LightContribution * LightProfileMultiplier;
+		if (Active)
+		{
+			float3 LightContribution = GetDynamicLighting(WorldPosition, -CameraVector, GBufferData, AmbientOcclusion, GBufferData.ShadingModelID, LightData, LightAttenuation, 0.5, uint2(0, 0), RectTexture).xyz;
+			DirectLighting += LightContribution * LightProfileMultiplier;
+		}
+
+		DirIndex += 1;
+		LightCullMask.x >>= 1;
+
+		LightCullMask.x |= LightCullMask.y << 31;
+		LightCullMask.y >>= 1;
 	}
 
 	return DirectLighting;
@@ -356,18 +512,28 @@ FMaterialClosestHitPayload TraceRayAndAccumulateResults(
 	inout FRayCone RayCone,
 	inout float3 Radiance)
 {
-	FMaterialClosestHitPayload Payload = TraceRayInternal(
-		TLAS,      // AccelerationStructure
+	FPackedMaterialClosestHitPayload PackedPayload = TraceRayInternalPacked(
+		TLAS,								// AccelerationStructure
 		RayFlags,
 		Mask,
-		RAY_TRACING_SHADER_SLOT_MATERIAL, // RayContributionToHitGroupIndex
-		RAY_TRACING_NUM_SHADER_SLOTS,     // MultiplierForGeometryContributionToShaderIndex
-		0,         // MissShaderIndex
-		Ray,       // RayDesc
+		RAY_TRACING_SHADER_SLOT_MATERIAL,	// RayContributionToHitGroupIndex
+		RAY_TRACING_NUM_SHADER_SLOTS,		// MultiplierForGeometryContributionToShaderIndex
+		0,									// MissShaderIndex
+		Ray,								// RayDesc
 		RayCone
 	);
 
-	if (IsHit(Ray, Payload))
+	uint2 LightCullMask = (uint2)0;
+	if (IsHit(Ray, PackedPayload))
+	{
+		const float3 WorldPosition = Ray.Origin + PackedPayload.HitT * Ray.Direction;
+		const float3 WorldNormal = UnpackRayTracingNormal(PackedPayload);
+		LightCullMask = CullDirectLighting(Ray, RayCone, WorldPosition, WorldNormal, RandSequence, ReflectedShadowsType > 0, ShadowMaxNormalBias);
+	}
+
+	FMaterialClosestHitPayload Payload = UnpackRayTracingPayload(PackedPayload, Ray);
+
+	if (IsHit(Ray, PackedPayload))
 	{
 		ApplyViewDependentMaterialPayloadModifications(Ray, Payload);
 
@@ -376,7 +542,7 @@ FMaterialClosestHitPayload TraceRayAndAccumulateResults(
 
 		if (ShouldDoDirectLighting)
 		{
-			float3 DirectLighting = ComputeDirectLighting(Ray, RayCone, Payload, RandSequence, ReflectedShadowsType > 0, ShadowMaxNormalBias);
+			float3 DirectLighting = ComputeDirectLightingCulled(LightCullMask, Ray, RayCone, Payload, RandSequence, ReflectedShadowsType > 0, ShadowMaxNormalBias);
 
 			// Transform NaNs to black, transform negative colors to black.
 			DirectLighting = -min(-DirectLighting, 0.0);
diff --git a/Engine/Shaders/Private/RayTracing/RayTracingReflections.usf b/Engine/Shaders/Private/RayTracing/RayTracingReflections.usf
--- a/Engine/Shaders/Private/RayTracing/RayTracingReflections.usf
+++ b/Engine/Shaders/Private/RayTracing/RayTracingReflections.usf
@@ -24,10 +24,10 @@ float2 SobolIndexToUniformUnitSquare(uint2 SobolRand)
 #define USE_SOURCE_TEXTURE			1
 #define USE_SOURCE_TEXTURE_ARRAY	1
 
-#define LTCMatTexture				RaytracingLightsData.LTCMatTexture
-#define LTCMatSampler				RaytracingLightsData.LTCMatSampler
-#define LTCAmpTexture				RaytracingLightsData.LTCAmpTexture
-#define LTCAmpSampler				RaytracingLightsData.LTCAmpSampler
+#define LTCMatTexture				RaytracingLightsDataPacked.LTCMatTexture
+#define LTCMatSampler				RaytracingLightsDataPacked.LTCMatSampler
+#define LTCAmpTexture				RaytracingLightsDataPacked.LTCAmpTexture
+#define LTCAmpSampler				RaytracingLightsDataPacked.LTCAmpSampler
 #define PreIntegratedGF				ReflectionStruct.PreIntegratedGF			
 #define PreIntegratedGFSampler		ReflectionStruct.PreIntegratedGFSampler	
 
diff --git a/Engine/Shaders/Private/RayTracing/RayTracingTranslucency.usf b/Engine/Shaders/Private/RayTracing/RayTracingTranslucency.usf
--- a/Engine/Shaders/Private/RayTracing/RayTracingTranslucency.usf
+++ b/Engine/Shaders/Private/RayTracing/RayTracingTranslucency.usf
@@ -15,10 +15,10 @@ float2 SobolIndexToUniformUnitSquare(uint2 SobolRand)
 #define USE_SOURCE_TEXTURE			1
 #define USE_SOURCE_TEXTURE_ARRAY	1
 
-#define LTCMatTexture				RaytracingLightsData.LTCMatTexture
-#define LTCMatSampler				RaytracingLightsData.LTCMatSampler
-#define LTCAmpTexture				RaytracingLightsData.LTCAmpTexture
-#define LTCAmpSampler				RaytracingLightsData.LTCAmpSampler
+#define LTCMatTexture				RaytracingLightsDataPacked.LTCMatTexture
+#define LTCMatSampler				RaytracingLightsDataPacked.LTCMatSampler
+#define LTCAmpTexture				RaytracingLightsDataPacked.LTCAmpTexture
+#define LTCAmpSampler				RaytracingLightsDataPacked.LTCAmpSampler
 #define PreIntegratedGF				ReflectionStruct.PreIntegratedGF			
 #define PreIntegratedGFSampler		ReflectionStruct.PreIntegratedGFSampler	
 
@@ -266,7 +266,6 @@ void RayTracingTranslucencyRGS()
 		Ray.Direction = RefractedDirection;
 		float SurfaceCurvature = 0.0f; /* #todo_dxr assume no curvature */
 		RayCone = PropagateRayCone(RayCone, SurfaceCurvature, Depth);
-		const float ReflectionMaxNormalBias = 1.0;
 	}
 
 	if (!bHasScattered)
diff --git a/Engine/Source/Runtime/Renderer/Private/RayTracing/RayTracingLighting.cpp b/Engine/Source/Runtime/Renderer/Private/RayTracing/RayTracingLighting.cpp
--- a/Engine/Source/Runtime/Renderer/Private/RayTracing/RayTracingLighting.cpp
+++ b/Engine/Source/Runtime/Renderer/Private/RayTracing/RayTracingLighting.cpp
@@ -6,12 +6,12 @@
 
 #if RHI_RAYTRACING
 
-IMPLEMENT_GLOBAL_SHADER_PARAMETER_STRUCT(FRaytracingLightData, "RaytracingLightsData");
+IMPLEMENT_GLOBAL_SHADER_PARAMETER_STRUCT(FRaytracingLightDataPacked, "RaytracingLightsDataPacked");
 
-void SetupRaytracingLightData(
+void SetupRaytracingLightDataPacked(
 	const TSparseArray<FLightSceneInfoCompact>& Lights,
 	const FViewInfo& View,
-	FRaytracingLightData* LightData)
+	FRaytracingLightDataPacked* LightData)
 {
 	TMap<UTextureLightProfile*, int32> IESLightProfilesMap;
 	TMap<FTextureRHIParamRef, uint32> RectTextureMap;
@@ -66,25 +66,31 @@ void SetupRaytracingLightData(
 			}
 		}
 
-		LightData->Type[LightData->Count] = Light.LightType;
-		LightData->LightPosition[LightData->Count] = LightParameters.Position;
-		LightData->LightInvRadius[LightData->Count] = LightParameters.InvRadius;
-		LightData->LightColor[LightData->Count] = LightParameters.Color;
-		LightData->LightFalloffExponent[LightData->Count] = LightParameters.FalloffExponent;
-		LightData->Direction[LightData->Count] = LightParameters.Direction;
-		LightData->Tangent[LightData->Count] = LightParameters.Tangent;
-		LightData->SpotAngles[LightData->Count] = LightParameters.SpotAngles;
-		LightData->SpecularScale[LightData->Count] = LightParameters.SpecularScale;
-		LightData->SourceRadius[LightData->Count] = LightParameters.SourceRadius;
-		LightData->SourceLength[LightData->Count] = LightParameters.SourceLength;
-		LightData->SoftSourceRadius[LightData->Count] = LightParameters.SoftSourceRadius;
-		LightData->LightProfileIndex[LightData->Count] = IESLightProfileIndex;
-		LightData->RectLightTextureIndex[LightData->Count] = InvalidTextureIndex;
-		LightData->RectLightBarnCosAngle[LightData->Count] = LightParameters.RectLightBarnCosAngle;
-		LightData->RectLightBarnLength[LightData->Count] = LightParameters.RectLightBarnLength;
+		LightData->Type_LightProfileIndex_RectLightTextureIndex[LightData->Count].X = Light.LightType;
+		LightData->Type_LightProfileIndex_RectLightTextureIndex[LightData->Count].Y = IESLightProfileIndex;
+		LightData->Type_LightProfileIndex_RectLightTextureIndex[LightData->Count].Z = InvalidTextureIndex;
 
-		const bool bRequireTexture = Light.LightType == ELightComponentType::LightType_Rect && LightParameters.SourceTexture;
+		LightData->LightPosition_InvRadius[LightData->Count] = LightParameters.Position;
+		LightData->LightPosition_InvRadius[LightData->Count].W = LightParameters.InvRadius;
+
+		LightData->LightColor_SpecularScale[LightData->Count] = LightParameters.Color;
+		LightData->LightColor_SpecularScale[LightData->Count].W = LightParameters.SpecularScale;
+
+		LightData->Direction_FalloffExponent[LightData->Count] = LightParameters.Direction;
+		LightData->Direction_FalloffExponent[LightData->Count].W = LightParameters.FalloffExponent;
+
+		LightData->Tangent_SourceRadius[LightData->Count] = LightParameters.Tangent;
+		LightData->Tangent_SourceRadius[LightData->Count].W = LightParameters.SourceRadius;
+
+		FVector4 SpotAngles_SourceLength_SoftSourceRadius = FVector4(LightParameters.SpotAngles, FVector2D(LightParameters.SourceLength, LightParameters.SoftSourceRadius));
+		LightData->SpotAngles_SourceLength_SoftSourceRadius[LightData->Count] = SpotAngles_SourceLength_SoftSourceRadius;
+		
+		const FVector2D FadeParams = Light.LightSceneInfo->Proxy->GetDirectionalLightDistanceFadeParameters(View.GetFeatureLevel(), Light.LightSceneInfo->IsPrecomputedLightingValid(), View.MaxShadowCascades);
+
+		FVector4 DistanceFadeMAD_RectLightBarnCosAngle_RectLightBarnLength = FVector4(FadeParams.Y, -FadeParams.X * FadeParams.Y, LightParameters.RectLightBarnCosAngle, LightParameters.RectLightBarnLength);
+		LightData->DistanceFadeMAD_RectLightBarnCosAngle_RectLightBarnLength[LightData->Count] = DistanceFadeMAD_RectLightBarnCosAngle_RectLightBarnLength;
 
+		const bool bRequireTexture = Light.LightType == ELightComponentType::LightType_Rect && LightParameters.SourceTexture;
 		uint32 RectLightTextureIndex = InvalidTextureIndex;
 		if (bRequireTexture)
 		{
@@ -105,7 +111,7 @@ void SetupRaytracingLightData(
 
 		if (RectLightTextureIndex != InvalidTextureIndex)
 		{
-			LightData->RectLightTextureIndex[LightData->Count] = RectLightTextureIndex;
+			LightData->Type_LightProfileIndex_RectLightTextureIndex[LightData->Count].Z = RectLightTextureIndex;
 			switch (RectLightTextureIndex)
 			{
 			case 0: LightData->RectLightTexture0 = LightParameters.SourceTexture; break;
@@ -119,9 +125,6 @@ void SetupRaytracingLightData(
 			}
 		}
 
-		const FVector2D FadeParams = Light.LightSceneInfo->Proxy->GetDirectionalLightDistanceFadeParameters(View.GetFeatureLevel(), Light.LightSceneInfo->IsPrecomputedLightingValid(), View.MaxShadowCascades);
-		LightData->DistanceFadeMAD[LightData->Count] = FVector2D(FadeParams.Y, -FadeParams.X * FadeParams.Y);
-
 		LightData->Count++;
 
 		if (LightData->Count >= GRaytracingLightCountMaximum) break;
@@ -142,10 +145,10 @@ void SetupRaytracingLightData(
 	}
 }
 
-TUniformBufferRef<FRaytracingLightData> CreateLightDataUniformBuffer(const TSparseArray<FLightSceneInfoCompact>& Lights, const class FViewInfo& View, EUniformBufferUsage Usage)
+TUniformBufferRef<FRaytracingLightDataPacked> CreateLightDataPackedUniformBuffer(const TSparseArray<FLightSceneInfoCompact>& Lights, const class FViewInfo& View, EUniformBufferUsage Usage)
 {
-	FRaytracingLightData LightData;
-	SetupRaytracingLightData(Lights, View, &LightData);
+	FRaytracingLightDataPacked LightData;
+	SetupRaytracingLightDataPacked(Lights, View, &LightData);
 	return CreateUniformBufferImmediate(LightData, Usage);
 }
 
diff --git a/Engine/Source/Runtime/Renderer/Private/RayTracing/RayTracingLighting.h b/Engine/Source/Runtime/Renderer/Private/RayTracing/RayTracingLighting.h
--- a/Engine/Source/Runtime/Renderer/Private/RayTracing/RayTracingLighting.h
+++ b/Engine/Source/Runtime/Renderer/Private/RayTracing/RayTracingLighting.h
@@ -12,25 +12,15 @@
 
 const static uint32 GRaytracingLightCountMaximum = 64;
 
-BEGIN_GLOBAL_SHADER_PARAMETER_STRUCT(FRaytracingLightData, )
+BEGIN_GLOBAL_SHADER_PARAMETER_STRUCT(FRaytracingLightDataPacked, )
 	SHADER_PARAMETER(uint32, Count)
-	SHADER_PARAMETER_ARRAY(uint32, Type, [GRaytracingLightCountMaximum])
-	SHADER_PARAMETER_ARRAY(FVector, LightPosition, [GRaytracingLightCountMaximum])
-	SHADER_PARAMETER_ARRAY(float, LightInvRadius, [GRaytracingLightCountMaximum])
-	SHADER_PARAMETER_ARRAY(FVector, LightColor, [GRaytracingLightCountMaximum])
-	SHADER_PARAMETER_ARRAY(float, LightFalloffExponent, [GRaytracingLightCountMaximum])
-	SHADER_PARAMETER_ARRAY(FVector, Direction, [GRaytracingLightCountMaximum])
-	SHADER_PARAMETER_ARRAY(FVector, Tangent, [GRaytracingLightCountMaximum])
-	SHADER_PARAMETER_ARRAY(FVector2D, SpotAngles, [GRaytracingLightCountMaximum])
-	SHADER_PARAMETER_ARRAY(float, SpecularScale, [GRaytracingLightCountMaximum])
-	SHADER_PARAMETER_ARRAY(float, SourceRadius, [GRaytracingLightCountMaximum])
-	SHADER_PARAMETER_ARRAY(float, SourceLength, [GRaytracingLightCountMaximum])
-	SHADER_PARAMETER_ARRAY(float, SoftSourceRadius, [GRaytracingLightCountMaximum])
-	SHADER_PARAMETER_ARRAY(uint32, LightProfileIndex, [GRaytracingLightCountMaximum])
-	SHADER_PARAMETER_ARRAY(FVector2D, DistanceFadeMAD, [GRaytracingLightCountMaximum])
-	SHADER_PARAMETER_ARRAY(uint32, RectLightTextureIndex, [GRaytracingLightCountMaximum])
-	SHADER_PARAMETER_ARRAY(float, RectLightBarnCosAngle, [GRaytracingLightCountMaximum])
-	SHADER_PARAMETER_ARRAY(float, RectLightBarnLength, [GRaytracingLightCountMaximum])
+	SHADER_PARAMETER_ARRAY(FIntVector, Type_LightProfileIndex_RectLightTextureIndex, [GRaytracingLightCountMaximum])
+	SHADER_PARAMETER_ARRAY(FVector4, LightPosition_InvRadius, [GRaytracingLightCountMaximum])
+	SHADER_PARAMETER_ARRAY(FVector4, LightColor_SpecularScale, [GRaytracingLightCountMaximum])
+	SHADER_PARAMETER_ARRAY(FVector4, Direction_FalloffExponent, [GRaytracingLightCountMaximum])
+	SHADER_PARAMETER_ARRAY(FVector4, Tangent_SourceRadius, [GRaytracingLightCountMaximum])
+	SHADER_PARAMETER_ARRAY(FVector4, SpotAngles_SourceLength_SoftSourceRadius, [GRaytracingLightCountMaximum])
+	SHADER_PARAMETER_ARRAY(FVector4, DistanceFadeMAD_RectLightBarnCosAngle_RectLightBarnLength, [GRaytracingLightCountMaximum])
 	SHADER_PARAMETER_TEXTURE(Texture2D, LTCMatTexture)
 	SHADER_PARAMETER_SAMPLER(SamplerState, LTCMatSampler)
 	SHADER_PARAMETER_TEXTURE(Texture2D, LTCAmpTexture)
@@ -45,11 +35,11 @@ BEGIN_GLOBAL_SHADER_PARAMETER_STRUCT(FRaytracingLightData, )
 	SHADER_PARAMETER_TEXTURE(Texture2D, RectLightTexture7)
 END_GLOBAL_SHADER_PARAMETER_STRUCT()
 
-void SetupRaytracingLightData(
+void SetupRaytracingLightDataPacked(
 	const TSparseArray<FLightSceneInfoCompact>& Lights,
 	const FViewInfo& View,
-	FRaytracingLightData* LightData);
+	FRaytracingLightDataPacked* LightData);
 
-TUniformBufferRef<FRaytracingLightData> CreateLightDataUniformBuffer(const TSparseArray<FLightSceneInfoCompact>& Lights, const class FViewInfo& View, EUniformBufferUsage Usage);
+TUniformBufferRef<FRaytracingLightDataPacked> CreateLightDataPackedUniformBuffer(const TSparseArray<FLightSceneInfoCompact>& Lights, const class FViewInfo& View, EUniformBufferUsage Usage);
 
 #endif
\ No newline at end of file
diff --git a/Engine/Source/Runtime/Renderer/Private/RayTracing/RayTracingReflections.cpp b/Engine/Source/Runtime/Renderer/Private/RayTracing/RayTracingReflections.cpp
--- a/Engine/Source/Runtime/Renderer/Private/RayTracing/RayTracingReflections.cpp
+++ b/Engine/Source/Runtime/Renderer/Private/RayTracing/RayTracingReflections.cpp
@@ -135,7 +135,7 @@ class FRayTracingReflectionsRGS : public FGlobalShader
 
 		SHADER_PARAMETER_STRUCT_REF(FViewUniformShaderParameters, ViewUniformBuffer)
 		SHADER_PARAMETER_STRUCT_REF(FSceneTexturesUniformParameters, SceneTexturesStruct)
-		SHADER_PARAMETER_STRUCT_REF(FRaytracingLightData, LightData)
+		SHADER_PARAMETER_STRUCT_REF(FRaytracingLightDataPacked, LightDataPacked)
 		SHADER_PARAMETER_STRUCT_REF(FReflectionUniformParameters, ReflectionStruct)
 		SHADER_PARAMETER_STRUCT_REF(FFogUniformParameters, FogUniformParameters)
 		SHADER_PARAMETER_STRUCT_REF(FIESLightProfileParameters, IESLightProfileParameters)
@@ -273,7 +273,7 @@ void FDeferredShadingSceneRenderer::RenderRayTracingReflections(
 
 	CommonParameters.TLAS = View.RayTracingScene.RayTracingSceneRHI->GetShaderResourceView();
 	CommonParameters.ViewUniformBuffer = View.ViewUniformBuffer;
-	CommonParameters.LightData = CreateLightDataUniformBuffer(Scene->Lights, View, EUniformBufferUsage::UniformBuffer_SingleFrame);
+	CommonParameters.LightDataPacked = CreateLightDataPackedUniformBuffer(Scene->Lights, View, EUniformBufferUsage::UniformBuffer_SingleFrame);
 	CommonParameters.SceneTexturesStruct = CreateSceneTextureUniformBuffer( SceneContext, FeatureLevel, ESceneTextureSetupMode::All, EUniformBufferUsage::UniformBuffer_SingleFrame);
 	CommonParameters.ReflectionStruct = CreateReflectionUniformBuffer(View, EUniformBufferUsage::UniformBuffer_SingleFrame);
 	CommonParameters.FogUniformParameters = CreateFogUniformBuffer(View, EUniformBufferUsage::UniformBuffer_SingleFrame);
diff --git a/Engine/Source/Runtime/Renderer/Private/RayTracing/RayTracingTranslucency.cpp b/Engine/Source/Runtime/Renderer/Private/RayTracing/RayTracingTranslucency.cpp
--- a/Engine/Source/Runtime/Renderer/Private/RayTracing/RayTracingTranslucency.cpp
+++ b/Engine/Source/Runtime/Renderer/Private/RayTracing/RayTracingTranslucency.cpp
@@ -121,7 +121,7 @@ class FRayTracingTranslucencyRGS : public FGlobalShader
 
 		SHADER_PARAMETER_STRUCT_REF(FViewUniformShaderParameters, ViewUniformBuffer)
 		SHADER_PARAMETER_STRUCT_REF(FSceneTexturesUniformParameters, SceneTexturesStruct)
-		SHADER_PARAMETER_STRUCT_REF(FRaytracingLightData, LightData)
+		SHADER_PARAMETER_STRUCT_REF(FRaytracingLightDataPacked, LightDataPacked)
 		SHADER_PARAMETER_STRUCT_REF(FReflectionUniformParameters, ReflectionStruct)
 		SHADER_PARAMETER_STRUCT_REF(FFogUniformParameters, FogUniformParameters)
 		SHADER_PARAMETER_STRUCT_REF(FIESLightProfileParameters, IESLightProfileParameters)
@@ -369,7 +369,7 @@ void FDeferredShadingSceneRenderer::RenderRayTracingTranslucencyView(
 
 	PassParameters->TLAS = View.RayTracingScene.RayTracingSceneRHI->GetShaderResourceView();
 	PassParameters->ViewUniformBuffer = View.ViewUniformBuffer;
-	PassParameters->LightData = CreateLightDataUniformBuffer(Scene->Lights, View, EUniformBufferUsage::UniformBuffer_SingleFrame);
+	PassParameters->LightDataPacked = CreateLightDataPackedUniformBuffer(Scene->Lights, View, EUniformBufferUsage::UniformBuffer_SingleFrame);
 	PassParameters->SceneTexturesStruct = CreateSceneTextureUniformBuffer(SceneContext, FeatureLevel, ESceneTextureSetupMode::All, EUniformBufferUsage::UniformBuffer_SingleFrame);	
 	PassParameters->ReflectionStruct = CreateReflectionUniformBuffer(View, EUniformBufferUsage::UniformBuffer_SingleFrame);
 	PassParameters->FogUniformParameters = CreateFogUniformBuffer(View, EUniformBufferUsage::UniformBuffer_SingleFrame);
