Fix raytraced shadow on double sided foliage.

diff --git a/Engine/Shaders/Private/RayTracing/RayTracingOcclusionRGS.usf b/Engine/Shaders/Private/RayTracing/RayTracingOcclusionRGS.usf
--- a/Engine/Shaders/Private/RayTracing/RayTracingOcclusionRGS.usf
+++ b/Engine/Shaders/Private/RayTracing/RayTracingOcclusionRGS.usf
@@ -161,6 +161,18 @@ void OcclusionRGS()
 		LocalSamplesPerPixel = 0.0;
 	}
 
+	// Cheap solution for double sided foliage to have valid lighing through the denoiser
+	if (GBufferData.ShadingModelID == SHADINGMODELID_TWOSIDED_FOLIAGE)
+	{
+	#if LIGHT_TYPE == LIGHT_TYPE_DIRECTIONAL
+		const float3 LightDirection = LightParameters.Direction;
+	#else
+		const float3 LightDirection = normalize(LightParameters.Position - WorldPosition);
+	#endif
+		if (dot(LightDirection, WorldNormal) < 0)
+			WorldNormal = -WorldNormal;
+	}
+
 	LOOP
 	for (uint SampleIndex = 0; SampleIndex < LocalSamplesPerPixel; ++SampleIndex)
 	{
