Fixes a NaN in the GI denoiser's post filtering pass.

diff --git a/Engine/Shaders/Private/ScreenSpaceDenoise/SSDSpatialAccumulation.usf b/Engine/Shaders/Private/ScreenSpaceDenoise/SSDSpatialAccumulation.usf
--- a/Engine/Shaders/Private/ScreenSpaceDenoise/SSDSpatialAccumulation.usf
+++ b/Engine/Shaders/Private/ScreenSpaceDenoise/SSDSpatialAccumulation.usf
@@ -630,7 +630,7 @@ void MainCS(
 			float CurrentSampleCount = OutputSamples.Array[MultiplexId].AccumulatedSampleCount;
 			float NewSampleCount = min(CurrentSampleCount, 128);
 
-			OutputSamples.Array[MultiplexId] = Mul(OutputSamples.Array[MultiplexId], NewSampleCount / CurrentSampleCount);
+			OutputSamples.Array[MultiplexId] = Mul(OutputSamples.Array[MultiplexId], CurrentSampleCount > 0 ? NewSampleCount / CurrentSampleCount : 0);
 		}
 	}
 
