Fixed pure specular reflection path.

diff --git a/Engine/Shaders/Private/PathTracing/Material/PathTracingSpecularReflection.ush b/Engine/Shaders/Private/PathTracing/Material/PathTracingSpecularReflection.ush
--- a/Engine/Shaders/Private/PathTracing/Material/PathTracingSpecularReflection.ush
+++ b/Engine/Shaders/Private/PathTracing/Material/PathTracingSpecularReflection.ush
@@ -15,17 +15,14 @@ void SpecularReflection_SampleMaterial(
 	float3 N_World = GetWorldNormal(Payload);
 	float3 SpecularColor = GetSpecularColor(Payload);
 
-	//OutDirection = reflect(RayDirection, N_World); // To reflect in world coords directly
 	float3 V_World = normalize(-RayDirection);
-	float3 V_Local = WorldToTangent(V_World, N_World);
+	OutDirection = reflect(V_World, N_World);
+
+	float NoV = saturate(dot(V_World, N_World));
+	float VoH = NoV;
+	float NoL = NoV;
+	OutThroughput = F_Schlick(SpecularColor, VoH) / NoL;
 
-	float3 L_Local;
-	L_Local.x = -V_Local.x;
-	L_Local.y = -V_Local.y;
-	L_Local.z = V_Local.z;
-	OutDirection = TangentToWorld(L_Local, N_World);
-	
-    OutThroughput = SpecularColor;
 	OutPdf = 1.0;
 }
 
