Fix missing diffuse component on perfectly smooth dielectric surfaces.

Insure the pure specular path is only taken by pure conductor.

diff --git a/Engine/Shaders/Private/PathTracing/Material/PathTracingMaterialCommon.ush b/Engine/Shaders/Private/PathTracing/Material/PathTracingMaterialCommon.ush
--- a/Engine/Shaders/Private/PathTracing/Material/PathTracingMaterialCommon.ush
+++ b/Engine/Shaders/Private/PathTracing/Material/PathTracingMaterialCommon.ush
@@ -42,7 +42,7 @@ bool IsPureSpecular(FMaterialClosestHitPayload HitInfo)
 
 bool IsPureSpecularReflection(FMaterialClosestHitPayload HitInfo)
 {
-	return IsPureSpecular(HitInfo) && HitInfo.BlendingMode == 0;
+	return IsPureSpecular(HitInfo) && HitInfo.BlendingMode == 0 && HitInfo.Metallic == 1;
 }
 
 bool IsPureSpecularTransmission(FMaterialClosestHitPayload HitInfo)
