Fixed: BlueprintOffice - Assertion Failed when interacting with button: MetalStateCache.cpp

diff --git a/Engine/Shaders/Private/MobileOpacityShaders.usf b/Engine/Shaders/Private/MobileOpacityShaders.usf
--- a/Engine/Shaders/Private/MobileOpacityShaders.usf
+++ b/Engine/Shaders/Private/MobileOpacityShaders.usf
@@ -5,10 +5,13 @@ MobileOpacityShaders.usf: Renders opacity to alpha channel. Mobile only.
 =============================================================================*/
 
 #include "Common.ush"
+
+// Reroute MobileSceneTextures uniform buffer references to the base pass uniform buffer
+#define MobileSceneTextures MobileBasePass.SceneTextures
+
 #include "/Engine/Generated/Material.ush"
 #include "/Engine/Generated/VertexFactory.ush"
 
-
 #define VertexFactoryGetInterpolants VertexFactoryGetInterpolantsVSToPS
 
 void MainVS(
