Cherry-pick 5357919: Fixes shadow denoiser input encoding at grazing angles.

diff --git a/Engine/Shaders/Private/RayTracing/RayTracingOcclusionRGS.usf b/Engine/Shaders/Private/RayTracing/RayTracingOcclusionRGS.usf
--- a/Engine/Shaders/Private/RayTracing/RayTracingOcclusionRGS.usf
+++ b/Engine/Shaders/Private/RayTracing/RayTracingOcclusionRGS.usf
@@ -209,6 +209,7 @@ void OcclusionRGS()
 		}
 		else if (dot(WorldNormal, Ray.Direction) <= 0.0)
 		{
+			ClosestRayDistance = 0.001;
 			continue;
 		}
 
