Cherry-pick 5358092: Implements denoiser preconvolution pass with hexa web kernel for shadows to fight against interferances between RGS sobol and denoiser's resconstruction kernel.

diff --git a/Engine/Shaders/Private/ScreenSpaceDenoise/SSDDefinitions.ush b/Engine/Shaders/Private/ScreenSpaceDenoise/SSDDefinitions.ush
--- a/Engine/Shaders/Private/ScreenSpaceDenoise/SSDDefinitions.ush
+++ b/Engine/Shaders/Private/ScreenSpaceDenoise/SSDDefinitions.ush
@@ -174,6 +174,7 @@
 	// [ Stackowiak 2015, "Stochastic Screen-Space Reflections" ]
 	#define SAMPLE_SET_STACKOWIAK_4_SETS 8
 
+	#define SAMPLE_SET_HEXAWEB 11
 
 
 //------------------------------------------------------- GLOBAL CONFIGURATION OF ALL DENOISER SHADERS
diff --git a/Engine/Shaders/Private/ScreenSpaceDenoise/SSDSignalFramework.ush b/Engine/Shaders/Private/ScreenSpaceDenoise/SSDSignalFramework.ush
--- a/Engine/Shaders/Private/ScreenSpaceDenoise/SSDSignalFramework.ush
+++ b/Engine/Shaders/Private/ScreenSpaceDenoise/SSDSignalFramework.ush
@@ -507,7 +507,7 @@ float ComputeBilateralWeight(
 
 		PositionBilateralWeight = saturate(1 - DistFromRefPlane * rcp(MaxWorldBluringDistance));
 	}
-	#elif CONFIG_VIEW_POSITION_BASED_BILATERAL == 3 && 1
+	#elif CONFIG_VIEW_POSITION_BASED_BILATERAL == 3
 	{
 		float3 V = RefSceneMetadata.TranslatedWorldPosition - NeighborSceneMetadata.TranslatedWorldPosition;
 
@@ -516,7 +516,7 @@ float ComputeBilateralWeight(
 			PositionBilateralWeight = 0;
 		}
 	}
-	#elif CONFIG_VIEW_POSITION_BASED_BILATERAL == 3
+	#elif CONFIG_VIEW_POSITION_BASED_BILATERAL == 4
 	{
 		float3 V = RefSceneMetadata.TranslatedWorldPosition - NeighborSceneMetadata.TranslatedWorldPosition;
 
diff --git a/Engine/Shaders/Private/ScreenSpaceDenoise/SSDSpatialAccumulation.usf b/Engine/Shaders/Private/ScreenSpaceDenoise/SSDSpatialAccumulation.usf
--- a/Engine/Shaders/Private/ScreenSpaceDenoise/SSDSpatialAccumulation.usf
+++ b/Engine/Shaders/Private/ScreenSpaceDenoise/SSDSpatialAccumulation.usf
@@ -7,9 +7,10 @@
 
 /** Different possible stage for spatial accumulation. Matches  */
 	#define STAGE_RECONSTRUCTION 0
-	#define STAGE_REJECTION_PRE_CONVOLUTION 1
-	#define STAGE_POST_FILTERING 2
-	#define STAGE_FINAL_OUTPUT 3
+	#define STAGE_PRE_CONVOLUTION 1
+	#define STAGE_REJECTION_PRE_CONVOLUTION 2
+	#define STAGE_POST_FILTERING 3
+	#define STAGE_FINAL_OUTPUT 4
 
 /** Policy to use to change the size of kernel. */
 	#define SAMPLE_COUNT_POLICY_DISABLED 0
@@ -40,7 +41,7 @@
 // Configures all the pass for each individual signals.
 #if CONFIG_SIGNAL_PROCESSING == SIGNAL_PROCESSING_MONOCHROMATIC_PENUMBRA
 	//#define CONFIG_VIEW_POSITION_BASED_BILATERAL 2
-	#define CONFIG_VIEW_POSITION_BASED_BILATERAL 3
+	#define CONFIG_VIEW_POSITION_BASED_BILATERAL 4
 	#define CONFIG_MULTIPLEXED_SIGNALS_PER_SIGNAL_DOMAIN 1
 
 	#if DIM_STAGE == STAGE_RECONSTRUCTION
@@ -56,6 +57,15 @@
 		#define CONFIG_BILATERAL_DISTANCE_COMPUTATION SIGNAL_WORLD_FREQUENCY_PRECOMPUTED_BLURING_RADIUS
 		#define CONFIG_OUTPUT_MODE          OUTPUT_MODE_DRB
 
+	#elif DIM_STAGE == STAGE_PRE_CONVOLUTION
+		// Input and output layout.
+		#define CONFIG_SIGNAL_INPUT_LAYOUT  SIGNAL_BUFFER_LAYOUT_PENUMBRA_RECONSTRUCTION
+		#define CONFIG_SIGNAL_OUTPUT_LAYOUT SIGNAL_BUFFER_LAYOUT_PENUMBRA_RECONSTRUCTION
+
+		#define CONFIG_SAMPLE_SET           SAMPLE_SET_HEXAWEB
+		#define CONFIG_BILATERAL_DISTANCE_COMPUTATION SIGNAL_WORLD_FREQUENCY_PRECOMPUTED_BLURING_RADIUS
+		#define CONFIG_OUTPUT_MODE          OUTPUT_MODE_DRB
+
 	#elif DIM_STAGE == STAGE_REJECTION_PRE_CONVOLUTION
 		// Input and output layout.
 		#define CONFIG_SIGNAL_INPUT_LAYOUT  SIGNAL_BUFFER_LAYOUT_PENUMBRA_RECONSTRUCTION
@@ -407,7 +417,7 @@ void MainCS(
 		KernelConfig.BufferLayout = CONFIG_SIGNAL_INPUT_LAYOUT;
 		KernelConfig.MultiplexedSignalsPerSignalDomain = CONFIG_MULTIPLEXED_SIGNALS_PER_SIGNAL_DOMAIN;
 		KernelConfig.bUnroll = CONFIG_SAMPLE_SET != SAMPLE_SET_STACKOWIAK_4_SETS;
-		KernelConfig.bDescOrder = CONFIG_OUTPUT_MODE == OUTPUT_MODE_DRB && CONFIG_SAMPLE_SET == SAMPLE_SET_STACKOWIAK_4_SETS;
+		KernelConfig.bDescOrder = CONFIG_OUTPUT_MODE == OUTPUT_MODE_DRB;
 		KernelConfig.BilateralDistanceComputation = CONFIG_BILATERAL_DISTANCE_COMPUTATION;
 		KernelConfig.bNormalizeSample = CONFIG_NORMALIZE_INPUT != 0;
 		
@@ -457,6 +467,16 @@ void MainCS(
 			}
 		}
 
+		#if DIM_STAGE == STAGE_PRE_CONVOLUTION
+		{
+			KernelConfig.RingCount = 1;
+			KernelConfig.KernelSpreadFactor = 8;
+
+			// TODO: could be improved.
+			KernelConfig.bMinSamplePairInvFrequency = true;
+		}
+		#endif
+
 		// When doing history rejection preconvolution may have invalid ref sample, in witch case need to force take neighborhood to have a clamping box.
 		#if DIM_STAGE == STAGE_REJECTION_PRE_CONVOLUTION && CONFIG_UPSCALE
 		{
@@ -607,7 +627,7 @@ void MainCS(
 				}
 			}
 
-			if (0)
+			if (0 && DIM_STAGE == STAGE_PRE_CONVOLUTION)
 			{
 				DebugOutput[DispatchThreadId] = float4(
 					SignalAccumulators.Array[0].Previous.AccumulatedSampleCount,
diff --git a/Engine/Shaders/Private/ScreenSpaceDenoise/SSDSpatialKernel.ush b/Engine/Shaders/Private/ScreenSpaceDenoise/SSDSpatialKernel.ush
--- a/Engine/Shaders/Private/ScreenSpaceDenoise/SSDSpatialKernel.ush
+++ b/Engine/Shaders/Private/ScreenSpaceDenoise/SSDSpatialKernel.ush
@@ -65,6 +65,9 @@ struct FSSDKernelConfig
 	// Selects how the world distance should be computed for bilateral rejection at compile time.
 	uint BilateralDistanceComputation;
 
+	// Number of ring for a disk kernel.
+	uint RingCount;
+
 	// Compile time configuration whether want do LOOP or UNROLL
 	//  false by default to expose in user code when the shader byte code might potentially be big.
 	bool bUnroll;
@@ -81,6 +84,10 @@ struct FSSDKernelConfig
 	// Whether a sample should be normalised to 1 before accmulation.
 	bool bNormalizeSample;
 
+	// Whether should min sample frequency of pair of samples
+	// [ Jimenez 2014, "Next Generation Post Processing in Call of Duty: Advanced Warfare" ]
+	bool bMinSamplePairInvFrequency;
+
 	// The color space that has been encoded in the buffer.
 	uint BufferColorSpace[MAX_SIGNAL_MULTIPLEXING];
 	
@@ -132,12 +139,14 @@ FSSDKernelConfig CreateKernelConfig()
 	KernelConfig.SampleSubSetId = 0;
 	KernelConfig.BufferLayout = SIGNAL_BUFFER_LAYOUT_UNINITIALIZED;
 	KernelConfig.MultiplexedSignalsPerSignalDomain = MAX_SIGNAL_MULTIPLEXING;
+	KernelConfig.RingCount = 0;
 	KernelConfig.bUnroll = false;
 	KernelConfig.bSampleKernelCenter = false;
 	KernelConfig.bPreviousFrameMetadata = false;
 	KernelConfig.BilateralDistanceComputation = SIGNAL_WORLD_FREQUENCY_MIN_METADATA;
 	KernelConfig.bDescOrder = false;
 	KernelConfig.bNormalizeSample = false;
+	KernelConfig.bMinSamplePairInvFrequency = false;
 	
 	FIX_UNROLL(MAX_SIGNAL_MULTIPLEXING)
 	for (uint MultiplexId = 0; MultiplexId < MAX_SIGNAL_MULTIPLEXING; MultiplexId++)
@@ -373,6 +382,71 @@ void SampleAndAccumulateMultiplexedSignals(
 		bIsOutsideFrustum);
 } // SampleAndAccumulateMultiplexedSignals()
 
+void SampleAndAccumulateMultiplexedSignalsPair(
+	FSSDKernelConfig KernelConfig,
+	Texture2D SignalBuffer0,
+	Texture2D SignalBuffer1,
+	Texture2D SignalBuffer2,
+	Texture2D SignalBuffer3,
+	inout FSSDSignalAccumulatorArray Accumulators,
+	float2 SampleBufferUV[2],
+	float KernelSampleWeight)
+{
+	FSSFSampleSceneInfos SampleSceneMetadata[2];
+	FSSDSignalArray MultiplexedSamples[2];
+	bool bIsOutsideFrustum[2];
+
+	FIX_UNROLL(2)
+	for (uint PairFetchId = 0; PairFetchId < 2; PairFetchId++)
+	{
+		// Stores in SGPR whether this sample is outside the viewport, to avoid VGPR pressure to keep SampleBufferUV after texture fetches.
+		bIsOutsideFrustum[PairFetchId] = IsOutsideViewport(KernelConfig, SampleBufferUV[PairFetchId]);
+
+		SampleMultiplexedSignals(
+			KernelConfig,
+			SignalBuffer0,
+			SignalBuffer1,
+			SignalBuffer2,
+			SignalBuffer3,
+			SampleBufferUV[PairFetchId],
+			/* out */ SampleSceneMetadata[PairFetchId],
+			/* out */ MultiplexedSamples[PairFetchId]);
+	}
+	
+	// Take the min inverse frequency per signal if desired.
+	if (KernelConfig.bMinSamplePairInvFrequency)
+	{
+		FIX_UNROLL(MAX_SIGNAL_MULTIPLEXING)
+		for (uint SignalMultiplexId = 0; SignalMultiplexId < MAX_SIGNAL_MULTIPLEXING; SignalMultiplexId++)
+		{
+			float MinInvFrequency = min(
+				MultiplexedSamples[0].Array[SignalMultiplexId].WorldBluringRadius,
+				MultiplexedSamples[1].Array[SignalMultiplexId].WorldBluringRadius);
+
+			FLATTEN
+			if (MinInvFrequency > 0)
+			{
+				MultiplexedSamples[0].Array[SignalMultiplexId].WorldBluringRadius = MinInvFrequency;
+				MultiplexedSamples[1].Array[SignalMultiplexId].WorldBluringRadius = MinInvFrequency;
+			}
+		}
+	}
+
+	FIX_UNROLL(2)
+	for (uint PairAccumulateId = 0; PairAccumulateId < 2; PairAccumulateId++)
+	{
+		AccumulateSampledMultiplexedSignals(
+			KernelConfig,
+			/* inout */ Accumulators,
+			SampleBufferUV[PairAccumulateId],
+			SampleSceneMetadata[PairAccumulateId],
+			MultiplexedSamples[PairAccumulateId],
+			KernelSampleWeight,
+			/* bForceSample = */ false,
+			bIsOutsideFrustum[PairAccumulateId]);
+	}
+} // SampleAndAccumulateMultiplexedSignalsPair()
+
 void StartAccumulatingCluster(
 	FSSDKernelConfig KernelConfig,
 	inout FSSDSignalAccumulatorArray Accumulators,
@@ -399,6 +473,32 @@ void DijestAccumulatedClusterSamples(inout FSSDSignalAccumulatorArray Accumulato
 	}
 }
 
+void SampleAndAccumulateCenterSampleAsItsOwnRing(
+	FSSDKernelConfig KernelConfig,
+	Texture2D SignalBuffer0,
+	Texture2D SignalBuffer1,
+	Texture2D SignalBuffer2,
+	Texture2D SignalBuffer3,
+	inout FSSDSignalAccumulatorArray Accumulators)
+{
+	const uint RingId = 0;
+
+	StartAccumulatingCluster(KernelConfig, Accumulators, RingId);
+
+	SampleAndAccumulateMultiplexedSignals(
+		KernelConfig,
+		SignalBuffer0,
+		SignalBuffer1,
+		SignalBuffer2,
+		SignalBuffer3,
+		/* inout */ Accumulators,
+		KernelConfig.BufferUV,
+		/* KernelSampleWeight = */ 1.0,
+		/* bForceSample = */ KernelConfig.bForceKernelCenterAccumulation);
+
+	DijestAccumulatedClusterSamples(Accumulators, RingId, /* SampleCount = */ 1);
+}
+
 
 //------------------------------------------------------- EASY CONVOLUTIONS
 
@@ -806,6 +906,227 @@ void ConvolveStackowiakKernel(
 } // ConvolveStackowiakKernel()
 
 
+//------------------------------------------------------- DISK
+
+// Returns the position of the sample on the unit circle (radius = 1) for a given ring.
+float2 GetDiskSampleOnUnitCirle(uint RingId, uint RingSampleIteration, uint RingSampleId)
+{
+	RingId -= 1; // TODO.
+
+	float SampleRingPos = RingSampleId;
+
+	// Do not allign all j == 0 samples of the different ring on the X axis to increase minimal distance between all
+	// samples, that reduce variance to clean by post filtering.
+	#if 1
+		SampleRingPos += (RingId - 2 * (RingId / 2)) * 0.5;
+	#endif
+
+	#if 1
+		SampleRingPos += (RingId + 1) * 0.2;
+	#endif
+
+	float SampleAngle = PI * SampleRingPos / float(RingSampleIteration);
+			
+	return float2(cos(SampleAngle), sin(SampleAngle));
+}
+
+// Returns the rotation matrix to use between sample of the ring.
+float2x2 GetSampleRotationMatrix(uint RingSampleIteration)
+{
+	float RotationAngle = PI / float(RingSampleIteration);
+
+	float C = cos(RotationAngle);
+	float S = sin(RotationAngle);
+
+	return float2x2(
+		float2( C,  S),
+		float2(-S,  C));
+}
+
+// Returns the total number of sampling iteration for a given ring id.
+uint GetRingSamplingPairCount(const uint SampleSet, uint RingId)
+{
+	if (SampleSet == SAMPLE_SET_HEXAWEB)
+	{
+		return RingId * 3;
+	}
+
+	// This number of sample is carefully chosen to have exact number of sample a square shaped ring (SquarePos).
+	return RingId * 4;
+}
+
+// Returns the total number of sample of the kernel.
+uint GetDiskKernelSampleCount(const uint SampleSet, uint RingCount)
+{
+	if (SampleSet == SAMPLE_SET_HEXAWEB)
+	{
+		return 1 + 3 * RingCount * (RingCount + 1);
+	}
+
+	// Depends on GetRingSamplingPairCount().
+	return 1 + 4 * RingCount * (RingCount + 1);
+}
+
+// Transform at compile time a 2 dimensional batch's constant into sample pair constant, by using rotation invariance.
+float2 SampleConstFromBatchConst(const uint BatchSampleId, float2 BatchConst)
+{
+	/**
+	 *             Y
+	 *             ^
+	 *             |
+	 *        1    |
+	 *             |
+	 *             |       0
+	 *             |
+	 * - - - - - - O - - - - > X
+	 */
+	if (BatchSampleId == 1)
+		return float2(-BatchConst.y, BatchConst.x);
+	return BatchConst;
+}
+
+
+
+// Gather a ring into the accumulator.
+void GatherRingSamples(
+	FSSDKernelConfig KernelConfig,
+	Texture2D SignalBuffer0,
+	Texture2D SignalBuffer1,
+	Texture2D SignalBuffer2,
+	Texture2D SignalBuffer3,
+	inout FSSDSignalAccumulatorArray Accumulators,
+	const uint RingId)
+{
+	// Number of sample iteration for this ring.
+	const uint RingSamplePairCount = GetRingSamplingPairCount(KernelConfig.SampleSet, RingId);
+
+	// Number of sample pair to process per batch.
+	// TODO: Could potentially do 4 using symetries? Might be unpracticable because of VGPR pressure. 
+	const uint SamplePairBatchSize = (KernelConfig.SampleSet == SAMPLE_SET_HEXAWEB) ? 1 : 2;
+
+	// Number of batch to process.
+	const uint BatchCount = RingSamplePairCount / SamplePairBatchSize;
+
+	// Distance of the ring from the center of the kernel in sample count.
+	const uint RingDistance = uint(RingId + 0);
+	
+	// Generate at compile time sample rotation matrix.
+	const float2x2 SampleRotationMatrix = GetSampleRotationMatrix(RingSamplePairCount);
+
+	// Generates at compile time first sample location on circle (radius = 1).
+	const float2 FirstCircleUnitPos = GetDiskSampleOnUnitCirle(RingId, RingSamplePairCount, /* BatchId = */ 0);
+
+	// Position of the first sample on circle with radius according to KernelRadius.
+	float2 FirstCircleSamplePosOffset = (RingDistance * FirstCircleUnitPos) * KernelConfig.KernelSpreadFactor;
+
+	// Setup iteratable SGPR
+	float2 CurrentCircleUnitPos = FirstCircleUnitPos;
+	float2 CurrentCircleSamplePosOffset = FirstCircleSamplePosOffset;
+
+	#if CONFIG_SGPR_HINT_OPTIMIZATION
+	{
+		CurrentCircleUnitPos = ToScalarMemory(CurrentCircleUnitPos);
+		CurrentCircleSamplePosOffset = ToScalarMemory(CurrentCircleSamplePosOffset);
+	}
+	#endif
+
+	// Loops through all batch of samples to process.
+	LOOP
+	for (uint BatchId = 0; BatchId < BatchCount; BatchId++)
+	{
+		// Rotate the samples position along the ring.
+		CurrentCircleUnitPos = mul(CurrentCircleUnitPos, SampleRotationMatrix);
+		CurrentCircleSamplePosOffset = mul(CurrentCircleSamplePosOffset, SampleRotationMatrix);
+
+		#if CONFIG_SGPR_HINT_OPTIMIZATION
+		{
+			CurrentCircleUnitPos = ToScalarMemory(CurrentCircleUnitPos);
+			CurrentCircleSamplePosOffset = ToScalarMemory(CurrentCircleSamplePosOffset);
+		}
+		#endif
+
+		// Sample in batch of multiple pair to increase texture fetch concurency, to have better
+		// lattency hidding.
+		UNROLL
+		for (uint BatchSampleId = 0; BatchSampleId < SamplePairBatchSize; BatchSampleId++)
+		{
+			float2 CircleSamplePosOffset = SampleConstFromBatchConst(BatchSampleId, CurrentCircleSamplePosOffset);
+			
+			float2 SampleUVPair[2];
+			SampleUVPair[0] = KernelConfig.BufferUV + CircleSamplePosOffset * KernelConfig.BufferSizeAndInvSize.zw;
+			SampleUVPair[1] = KernelConfig.BufferUV - CircleSamplePosOffset * KernelConfig.BufferSizeAndInvSize.zw;
+
+			SampleAndAccumulateMultiplexedSignalsPair(
+				KernelConfig,
+				SignalBuffer0,
+				SignalBuffer1,
+				SignalBuffer2,
+				SignalBuffer3,
+				/* inout */ Accumulators,
+				SampleUVPair,
+				/* KernelWeight = */ 1.0);
+		} // for (uint BatchSampleId = 0; BatchSampleId < SamplePairBatchSize; BatchSampleId++)
+	} // for (uint BatchId = 0; BatchId < BatchCount; BatchId++)
+}
+
+void ConvolveDiskKernel(
+	FSSDKernelConfig KernelConfig,
+	Texture2D SignalBuffer0,
+	Texture2D SignalBuffer1,
+	Texture2D SignalBuffer2,
+	Texture2D SignalBuffer3,
+	inout FSSDSignalAccumulatorArray Accumulators)
+{
+	// Accumulate the center of the kernel.
+	if (KernelConfig.bSampleKernelCenter && !KernelConfig.bDescOrder)
+	{
+		SampleAndAccumulateCenterSampleAsItsOwnRing(
+			KernelConfig,
+			SignalBuffer0,
+			SignalBuffer1,
+			SignalBuffer2,
+			SignalBuffer3,
+			/* inout */ Accumulators);
+	}
+
+	// Accumulate each ring. Use LOOP, because FXC is going through its pace otherwise.
+	#if 1
+		LOOP
+	#else
+		UNROLL
+	#endif
+	for (uint RingId = 1; RingId <= KernelConfig.RingCount; RingId++)
+	{
+		const uint RingSamplePairCount = GetRingSamplingPairCount(KernelConfig.SampleSet, RingId);
+
+		StartAccumulatingCluster(KernelConfig, /* inout */ Accumulators, RingId);
+		
+		GatherRingSamples(
+			KernelConfig,
+			SignalBuffer0,
+			SignalBuffer1,
+			SignalBuffer2,
+			SignalBuffer3,
+			/* inout */ Accumulators,
+			RingId);
+			
+		DijestAccumulatedClusterSamples(/* inout */ Accumulators, RingId, RingSamplePairCount * 2);
+	} // for (uint RingId = 0; RingId < KernelConfig.RingCount; RingId++)
+		
+	// Accumulate the center of the kernel.
+	if (KernelConfig.bSampleKernelCenter && KernelConfig.bDescOrder)
+	{
+		SampleAndAccumulateCenterSampleAsItsOwnRing(
+			KernelConfig,
+			SignalBuffer0,
+			SignalBuffer1,
+			SignalBuffer2,
+			SignalBuffer3,
+			/* inout */ Accumulators);
+	}
+}
+
+
 //------------------------------------------------------- MAIN ENTRY POINTS
 
 void AccumulateKernel(
@@ -820,16 +1141,13 @@ void AccumulateKernel(
 	{
 		if (KernelConfig.bSampleKernelCenter)
 		{
-			SampleAndAccumulateMultiplexedSignals(
+			SampleAndAccumulateCenterSampleAsItsOwnRing(
 				KernelConfig,
 				SignalBuffer0,
 				SignalBuffer1,
 				SignalBuffer2,
 				SignalBuffer3,
-				/* inout */ Accumulators,
-				KernelConfig.BufferUV,
-				/* KernelSampleWeight = */ 1,
-				/* bForceSample = */ true);
+				/* inout */ Accumulators);
 		}
 	}
 	else if (KernelConfig.SampleSet == SAMPLE_SET_2X2_BILINEAR)
@@ -868,4 +1186,14 @@ void AccumulateKernel(
 			SignalBuffer3,
 			/* inout */ Accumulators);
 	}
+	else if (KernelConfig.SampleSet == SAMPLE_SET_HEXAWEB)
+	{
+		ConvolveDiskKernel(
+			KernelConfig,
+			SignalBuffer0,
+			SignalBuffer1,
+			SignalBuffer2,
+			SignalBuffer3,
+			/* inout */ Accumulators);
+	}
 } // AccumulateKernel()
diff --git a/Engine/Shaders/Private/ScreenSpaceDenoise/SSDTemporalAccumulation.usf b/Engine/Shaders/Private/ScreenSpaceDenoise/SSDTemporalAccumulation.usf
--- a/Engine/Shaders/Private/ScreenSpaceDenoise/SSDTemporalAccumulation.usf
+++ b/Engine/Shaders/Private/ScreenSpaceDenoise/SSDTemporalAccumulation.usf
@@ -612,7 +612,7 @@ void TemporallyAccumulate(
 				float NeightborMinBluringRadius = SignalAccumulators.Array[BatchedSignalId].MinInvFrequency;
 				CurrentFrameSample.WorldBluringRadius = min(CurrentFrameSample.WorldBluringRadius, NeightborMinBluringRadius * CurrentFrameSample.AccumulatedSampleCount); 
 			
-				#if CONFIG_SIGNAL_PROCESSING == SIGNAL_PROCESSING_MONOCHROMATIC_PENUMBRA && 1
+				#if CONFIG_SIGNAL_PROCESSING == SIGNAL_PROCESSING_MONOCHROMATIC_PENUMBRA && 0
 					DebugOutput[DispatchThreadId] = float4(
 						NeightborMinBluringRadius,
 						NeightborMinBluringRadius == WORLD_RADIUS_MISS,
diff --git a/Engine/Source/Runtime/Renderer/Private/ScreenSpaceDenoise.cpp b/Engine/Source/Runtime/Renderer/Private/ScreenSpaceDenoise.cpp
--- a/Engine/Source/Runtime/Renderer/Private/ScreenSpaceDenoise.cpp
+++ b/Engine/Source/Runtime/Renderer/Private/ScreenSpaceDenoise.cpp
@@ -24,10 +24,15 @@ static TAutoConsoleVariable<int32> CVarShadowUse1SPPCodePath(
 	ECVF_RenderThreadSafe);
 
 static TAutoConsoleVariable<int32> CVarShadowReconstructionSampleCount(
-	TEXT("r.Shadow.Denoiser.ReconstructionSamples"), 16,
+	TEXT("r.Shadow.Denoiser.ReconstructionSamples"), 8,
 	TEXT("Maximum number of samples for the reconstruction pass (default = 16)."),
 	ECVF_RenderThreadSafe);
 
+static TAutoConsoleVariable<int32> CVarShadowPreConvolutionCount(
+	TEXT("r.Shadow.Denoiser.PreConvolution"), 1,
+	TEXT("Number of pre-convolution passes (default = 1)."),
+	ECVF_RenderThreadSafe);
+
 static TAutoConsoleVariable<int32> CVarShadowTemporalAccumulation(
 	TEXT("r.Shadow.Denoiser.TemporalAccumulation"), 1,
 	TEXT(""),
@@ -150,6 +155,12 @@ static bool SignalUsesInjestion(ESignalProcessing SignalProcessing)
 	return SignalProcessing == ESignalProcessing::MonochromaticPenumbra;
 }
 
+/** Returns whether a signal processing uses an additional pre convolution pass. */
+static bool SignalUsesPreConvolution(ESignalProcessing SignalProcessing)
+{
+	return SignalProcessing == ESignalProcessing::MonochromaticPenumbra;
+}
+
 /** Returns whether a signal processing uses a history rejection pre convolution pass. */
 static bool SignalUsesRejectionPreConvolution(ESignalProcessing SignalProcessing)
 {
@@ -253,6 +264,32 @@ const TCHAR* const kReconstructionResourceNames[] = {
 	TEXT("GIReconstruction3"),
 };
 
+const TCHAR* const kPreConvolutionResourceNames[] = {
+	// Penumbra
+	TEXT("ShadowPreConvolution0"),
+	TEXT("ShadowPreConvolution1"),
+	TEXT("ShadowPreConvolution2"),
+	TEXT("ShadowPreConvolution3"),
+
+	// Reflections
+	nullptr,
+	nullptr,
+	nullptr,
+	nullptr,
+
+	// AmbientOcclusion
+	nullptr,
+	nullptr,
+	nullptr,
+	nullptr,
+
+	// GlobalIllumination
+	nullptr,
+	nullptr,
+	nullptr,
+	nullptr,
+};
+
 const TCHAR* const kRejectionPreConvolutionResourceNames[] = {
 	// Penumbra
 	TEXT("ShadowRejectionPreConvolution0"),
@@ -489,6 +526,9 @@ class FSSDSpatialAccumulationCS : public FScreenSpaceDenoisingShader
 		// Spatial kernel used to process raw input for the temporal accumulation.
 		ReConstruction,
 
+		// Spatial kernel to pre filter.
+		PreConvolution,
+
 		// Spatial kernel used to pre convolve history rejection.
 		RejectionPreConvolution,
 
@@ -530,6 +570,13 @@ class FSSDSpatialAccumulationCS : public FScreenSpaceDenoisingShader
 			return false;
 		}
 
+		// Only compile pre convolution for signal that uses it.
+		if (!SignalUsesPreConvolution(SignalProcessing) &&
+			PermutationVector.Get<FStageDim>() == EStage::PreConvolution)
+		{
+			return false;
+		}
+
 		// Only compile rejection pre convolution for signal that uses it.
 		if (!SignalUsesRejectionPreConvolution(SignalProcessing) &&
 			PermutationVector.Get<FStageDim>() == EStage::RejectionPreConvolution)
@@ -638,6 +685,7 @@ struct FSSDConstantPixelDensitySettings
 	int32 MaxInputSPP = 1;
 	float InputResolutionFraction = 1.0f;
 	int32 ReconstructionSamples = 1;
+	int32 PreConvolutionCount = 0;
 	bool bUseTemporalAccumulation = false;
 	int32 HistoryConvolutionSampleCount = 1;
 	float HistoryConvolutionKernelSpreadFactor = 1.0f;
@@ -875,6 +923,41 @@ static void DenoiseSignalAtConstantPixelDensity(
 		SignalHistory = NewSignalOutput;
 	}
 
+	// Spatial pre convolutions
+	for (int32 PreConvolutionId = 0; PreConvolutionId < Settings.PreConvolutionCount; PreConvolutionId++)
+	{
+		check(SignalUsesPreConvolution(Settings.SignalProcessing));
+
+		FSSDSignalTextures NewSignalOutput = CreateMultiplexedTextures(
+			GraphBuilder,
+			ReconstructionTextureCount, ReconstructionDescs,
+			GetResourceNames(kPreConvolutionResourceNames));
+
+		FSSDSpatialAccumulationCS::FParameters* PassParameters = GraphBuilder.AllocParameters<FSSDSpatialAccumulationCS::FParameters>();
+		PassParameters->CommonParameters = CommonParameters;
+		PassParameters->ConvolutionMetaData = ConvolutionMetaData;
+		PassParameters->SignalInput = SignalHistory;
+		PassParameters->SignalOutput = CreateMultiplexedUAVs(GraphBuilder, NewSignalOutput);
+
+		PassParameters->DebugOutput = GraphBuilder.CreateUAV(GraphBuilder.CreateTexture(DebugDesc, TEXT("DebugDenoiserPreConvolution")));
+
+		FSSDSpatialAccumulationCS::FPermutationDomain PermutationVector;
+		PermutationVector.Set<FSignalProcessingDim>(Settings.SignalProcessing);
+		PermutationVector.Set<FSignalBatchSizeDim>(Settings.SignalBatchSize);
+		PermutationVector.Set<FSSDSpatialAccumulationCS::FStageDim>(FSSDSpatialAccumulationCS::EStage::PreConvolution);
+		PermutationVector.Set<FMultiSPPDim>(true);
+
+		TShaderMapRef<FSSDSpatialAccumulationCS> ComputeShader(View.ShaderMap, PermutationVector);
+		FComputeShaderUtils::AddPass(
+			GraphBuilder,
+			RDG_EVENT_NAME("SSD PreConvolution(MaxSamples=7)"),
+			*ComputeShader,
+			PassParameters,
+			FComputeShaderUtils::GetGroupCount(DenoiseResolution, FSSDSpatialAccumulationCS::kGroupSize));
+
+		SignalHistory = NewSignalOutput;
+	}
+
 	// Temporal pass.
 	//
 	// Note: always done even if there is no ViewState, because it is already not an idea case for the denoiser quality, therefore not really
@@ -1177,6 +1260,7 @@ public:
 		Settings.SignalProcessing = ESignalProcessing::MonochromaticPenumbra;
 		Settings.InputResolutionFraction = 1.0f;
 		Settings.ReconstructionSamples = CVarShadowReconstructionSampleCount.GetValueOnRenderThread();
+		Settings.PreConvolutionCount = CVarShadowPreConvolutionCount.GetValueOnRenderThread();
 		Settings.bUseTemporalAccumulation = CVarShadowTemporalAccumulation.GetValueOnRenderThread() != 0;
 		Settings.HistoryConvolutionSampleCount = CVarShadowHistoryConvolutionSampleCount.GetValueOnRenderThread();
 		Settings.SignalBatchSize = InputParameterCount;
