Cherry-pick 5364734: Works arround an issue in the denoiser's signal frequency computation when an occluder is very close to an area light.

diff --git a/Engine/Shaders/Private/ScreenSpaceDenoise/SSDSignalFramework.ush b/Engine/Shaders/Private/ScreenSpaceDenoise/SSDSignalFramework.ush
--- a/Engine/Shaders/Private/ScreenSpaceDenoise/SSDSignalFramework.ush
+++ b/Engine/Shaders/Private/ScreenSpaceDenoise/SSDSignalFramework.ush
@@ -585,22 +585,23 @@ float GetSignalWorldBluringRadius(FSSFSignalSample Sample, FSSFSampleSceneInfos
 		return WORLD_RADIUS_MISS;
 	}
 
-	// Compute the anysotropy.
-	float Anisotropy = ComputeAnisotropyInvFactor(SceneMetadata);
-
 	FSSDSignalDomainKnowledge DomainKnowledge = GetSignalDomainKnowledge(BatchedSignalId);
 
+	// Sometime, artists might put occluder very close to the area light compared to they area, that may lead to negative values.
+	// TODO: the correct way to fix this is to move this world bluring radius computation into the RGS, and have DistanceFromLight = Ray's MaxT.
+	const float ClosestLightDistance = 0.001;
+
 	BRANCH
 	if (DomainKnowledge.LightType == LIGHT_TYPE_DIRECTIONAL)
 	{
-		return DomainKnowledge.HitDistanceToWorldBluringRadius * Sample.PenumbraClosestHit * Anisotropy;
+		return DomainKnowledge.HitDistanceToWorldBluringRadius * Sample.PenumbraClosestHit;
 	}
 	else if (DomainKnowledge.LightType == LIGHT_TYPE_POINT || DomainKnowledge.LightType == LIGHT_TYPE_SPOT)
 	{
 		float3 PixelToLightWorldVector = (DomainKnowledge.Light.Position + View.PreViewTranslation) - SceneMetadata.TranslatedWorldPosition;
 		float DistanceFromLight = length(PixelToLightWorldVector);
 
-		return DomainKnowledge.Light.SourceRadius * Sample.PenumbraClosestHit / (DistanceFromLight - Sample.PenumbraClosestHit) * Anisotropy;
+		return DomainKnowledge.Light.SourceRadius * Sample.PenumbraClosestHit / max(DistanceFromLight - Sample.PenumbraClosestHit, ClosestLightDistance);
 	}
 	else // if (DomainKnowledge.LightType == LIGHT_TYPE_RECT)
 	{
@@ -619,7 +620,7 @@ float GetSignalWorldBluringRadius(FSSFSignalSample Sample, FSSFSampleSceneInfos
 		float SmallestLightDimension = min(LightDimensions.x, LightDimensions.y);
 
 		// TODO: Sample.PenumbraClosestHit depends on the direction of the ray, but does not DistanceFromLight, witch is bad for large area light.
-		return SmallestLightDimension * Sample.PenumbraClosestHit / (DistanceFromLight - Sample.PenumbraClosestHit) * Anisotropy;
+		return SmallestLightDimension * Sample.PenumbraClosestHit / max(DistanceFromLight - Sample.PenumbraClosestHit, ClosestLightDistance);
 	}
 }
 #elif CONFIG_SIGNAL_PROCESSING == SIGNAL_PROCESSING_REFLECTIONS
