Cherry-pick 5385828: Fixes mid frequency shadow bluryness by tweaking pre convolution's bordering radius, and also max bilateral bluring distance with reference's.

diff --git a/Engine/Shaders/Private/ScreenSpaceDenoise/SSDSignalAccumulator.ush b/Engine/Shaders/Private/ScreenSpaceDenoise/SSDSignalAccumulator.ush
--- a/Engine/Shaders/Private/ScreenSpaceDenoise/SSDSignalAccumulator.ush
+++ b/Engine/Shaders/Private/ScreenSpaceDenoise/SSDSignalAccumulator.ush
@@ -35,6 +35,13 @@ struct FSSDSignalAccumulator
 	float HighestInvFrequency;
 };
 
+/** Informations about cluster of sample */
+struct FSSDSampleClusterInfo
+{
+	// Outter radius boundary of the cluster.
+	float OutterBoundaryRadius;
+};
+
 
 FSSDSignalAccumulator CreateSignalAccumulator()
 {
@@ -131,14 +138,14 @@ void AccumulateSampleInDomainBoundaries(inout FSSDSignalAccumulator Accumulator,
 void StartAccumulatingClusterInDRB(
 	FSSFSampleSceneInfos RefSceneMetadata,
 	inout FSSDSignalAccumulator Accumulator,
-	uint RingId)
+	FSSDSampleClusterInfo ClusterInfo)
 {
-	const uint ErrorCorrection = 2;
+	const float ErrorCorrection = 1;
 	
 	// Compute the bluring radius of the output pixel itself.
 	float RefPixelWorldBluringRadius = AmendWorldBluringRadiusCausedByPixelSize(ComputeWorldBluringRadiusCausedByPixelSize(RefSceneMetadata));
 
-	Accumulator.BorderingRadius = RefPixelWorldBluringRadius * float(RingId + ErrorCorrection);
+	Accumulator.BorderingRadius = RefPixelWorldBluringRadius * (ClusterInfo.OutterBoundaryRadius + ErrorCorrection);
 
 	// TODO: could this be constant.
 	Accumulator.HighestInvFrequency = Accumulator.BorderingRadius * 2;
@@ -150,7 +157,7 @@ void StartAccumulatingClusterInDRB(
 		Accumulator.CurrentTranslucency = 0.0;
 	}
 
-	#if CONFIG_SGPR_HINT_OPTIMIZATION
+	#if CONFIG_SGPR_HINT_OPTIMIZATION && 0
 	{
 		Accumulator.BorderingRadius = ToScalarMemory(Accumulator.BorderingRadius);
 		Accumulator.HighestInvFrequency = ToScalarMemory(Accumulator.HighestInvFrequency);
@@ -160,17 +167,15 @@ void StartAccumulatingClusterInDRB(
 
 void AccumulateSampleInDRB(inout FSSDSignalAccumulator Accumulator, FSSDSampleAccumulationInfos A)
 {
-	// Compare the sample's frequency with the bucket bordering radius.
-	float bBelongsToPrevious = saturate(A.InvFrequency - Accumulator.BorderingRadius + 0.5);
-
-	//bBelongsToPrevious = 1;
-	//bBelongsToPrevious = A.Sample.LeavingRayCount != 0.0 ? 1 : 0;
-
-	float ClampedInvFrequency  = min(A.InvFrequency, Accumulator.HighestInvFrequency);
+	float ClampedInvFrequency = min(A.InvFrequency, Accumulator.HighestInvFrequency);
 
+	// Compare the sample's frequency with the bucket bordering radius.
+	float bBelongsToPrevious = saturate(ClampedInvFrequency - Accumulator.BorderingRadius + 0.5);
+	float bBelongsToCurrent = 1.0 - bBelongsToPrevious;
+	
 	// Accumulate current bucket.
 	{
-		float CurrentWeight = A.FinalWeight * (1.0 - bBelongsToPrevious);
+		float CurrentWeight = A.FinalWeight * bBelongsToCurrent;
 		Accumulator.Current = Add(Accumulator.Current, Mul(A.Sample, CurrentWeight));
 		Accumulator.CurrentInvFrequency += A.Sample.AccumulatedSampleCount * ClampedInvFrequency * CurrentWeight;
 	}
@@ -183,10 +188,10 @@ void AccumulateSampleInDRB(inout FSSDSignalAccumulator Accumulator, FSSDSampleAc
 		// TODO: could pass down a normalised version of A.Sample to avoid Rcp.
 		float R = A.Sample.LeavingRayCount * SafeRcp(A.Sample.AccumulatedSampleCount);
 		
-		//Accumulator.CurrentTranslucency += SampleTranslucency;
-		Accumulator.CurrentTranslucency += lerp(R, 1, SampleTranslucency);
-		//Accumulator.CurrentTranslucency += saturate(1 - (1 - SampleTranslucency) * (1 - R));
-		//Accumulator.CurrentTranslucency += lerp(R, 1, SampleTranslucency);
+		float Translucency = lerp(R, 1, SampleTranslucency);
+		float TranslucencyWeight = 1;
+
+		Accumulator.CurrentTranslucency += Translucency * TranslucencyWeight;
 	}
 
 	// Accumulate previous bucket.
@@ -210,10 +215,10 @@ void DijestSampleClusterInDRB(
 	}
 	
 	// Opacity of the current bucket.
-	float CurrentClusterOpacity = saturate(1 - Accumulator.CurrentTranslucency * rcp(SampleCount));
+	float CurrentClusterOpacity = saturate(1 - Accumulator.CurrentTranslucency * rcp(float(SampleCount)));
 	
 	// Compute current and previous Coc radii.
-	float PreviousInvFrequency = Accumulator.PreviousInvFrequency * rcp(Accumulator.Previous.AccumulatedSampleCount);
+	float PreviousInvFrequency = Accumulator.PreviousInvFrequency * SafeRcp(Accumulator.Previous.AccumulatedSampleCount);
 	float CurrentInvFrequency = Accumulator.CurrentInvFrequency * rcp(Accumulator.Current.AccumulatedSampleCount);
 
 	// Whether current bucket is occluding previous bucket.
@@ -224,8 +229,8 @@ void DijestSampleClusterInDRB(
 	float PreviousBucketFactor = (Accumulator.Previous.AccumulatedSampleCount == 0.0) ? 0.0 : (1.0 - CurrentClusterOpacity * bOccludingCluster);
 
 	// TODO: the big difference in sample count could lead to wrong opacity between sample clusters.
-	//if (PreviousBucketFactor < 1 && Accumulator.Previous.AccumulatedSampleCount > 0)
-	//	PreviousBucketFactor *= Accumulator.Current.AccumulatedSampleCount * rcp(Accumulator.Previous.AccumulatedSampleCount);
+	//if (PreviousBucketFactor < 1)
+	//	PreviousBucketFactor *= Accumulator.Current.AccumulatedSampleCount * SafeRcp(Accumulator.Previous.AccumulatedSampleCount);
 
 	Accumulator.Previous = Add(Mul(Accumulator.Previous, PreviousBucketFactor), Accumulator.Current);
 	Accumulator.PreviousInvFrequency = Accumulator.PreviousInvFrequency * PreviousBucketFactor + Accumulator.CurrentInvFrequency;
@@ -238,9 +243,9 @@ void DijestSampleClusterInDRB(
 void StartAccumulatingCluster(
 	FSSFSampleSceneInfos RefSceneMetadata,
 	inout FSSDSignalAccumulator Accumulator,
-	uint RingId)
+	FSSDSampleClusterInfo ClusterInfo)
 {
-	StartAccumulatingClusterInDRB(RefSceneMetadata, /* inout */ Accumulator, RingId);
+	StartAccumulatingClusterInDRB(RefSceneMetadata, /* inout */ Accumulator, ClusterInfo);
 }
 
 /** Accumulate a sample withing the accumulator. */
diff --git a/Engine/Shaders/Private/ScreenSpaceDenoise/SSDSpatialAccumulation.usf b/Engine/Shaders/Private/ScreenSpaceDenoise/SSDSpatialAccumulation.usf
--- a/Engine/Shaders/Private/ScreenSpaceDenoise/SSDSpatialAccumulation.usf
+++ b/Engine/Shaders/Private/ScreenSpaceDenoise/SSDSpatialAccumulation.usf
@@ -55,6 +55,7 @@
 
 		#define CONFIG_SAMPLE_SET           SAMPLE_SET_STACKOWIAK_4_SETS
 		#define CONFIG_BILATERAL_DISTANCE_COMPUTATION SIGNAL_WORLD_FREQUENCY_PRECOMPUTED_BLURING_RADIUS
+		#define CONFIG_MAX_WITH_REF_DISTANCE 1
 		#define CONFIG_OUTPUT_MODE          OUTPUT_MODE_DRB
 
 	#elif DIM_STAGE == STAGE_PRE_CONVOLUTION
@@ -64,6 +65,7 @@
 
 		#define CONFIG_SAMPLE_SET           SAMPLE_SET_HEXAWEB
 		#define CONFIG_BILATERAL_DISTANCE_COMPUTATION SIGNAL_WORLD_FREQUENCY_PRECOMPUTED_BLURING_RADIUS
+		#define CONFIG_MAX_WITH_REF_DISTANCE 1
 		#define CONFIG_OUTPUT_MODE          OUTPUT_MODE_DRB
 
 	#elif DIM_STAGE == STAGE_REJECTION_PRE_CONVOLUTION
@@ -274,6 +276,11 @@
 	#define CONFIG_BILATERAL_DISTANCE_COMPUTATION SIGNAL_WORLD_FREQUENCY_MIN_METADATA
 #endif
 
+/** Whether neighbor bilateral distance should be maxed with reference one. */
+#ifndef CONFIG_MAX_WITH_REF_DISTANCE
+	#define CONFIG_MAX_WITH_REF_DISTANCE 0
+#endif
+
 
 //------------------------------------------------------- INCLUDES
 
@@ -412,6 +419,8 @@ void MainCS(
 	FSSDSignalAccumulatorArray SignalAccumulators = CreateSignalAccumulatorArray();
 	{
 		FSSDKernelConfig KernelConfig = CreateKernelConfig();
+
+		// Compile time.
 		KernelConfig.SampleSet = CONFIG_SAMPLE_SET;
 		KernelConfig.SampleSubSetId = CONFIG_SAMPLE_SUBSET;
 		KernelConfig.BufferLayout = CONFIG_SIGNAL_INPUT_LAYOUT;
@@ -420,25 +429,75 @@ void MainCS(
 		KernelConfig.bDescOrder = CONFIG_OUTPUT_MODE == OUTPUT_MODE_DRB;
 		KernelConfig.BilateralDistanceComputation = CONFIG_BILATERAL_DISTANCE_COMPUTATION;
 		KernelConfig.bNormalizeSample = CONFIG_NORMALIZE_INPUT != 0;
+		KernelConfig.bSampleKernelCenter = CONFIG_UPSCALE;
+		KernelConfig.bForceKernelCenterAccumulation = true;
 		
-		FIX_UNROLL(MAX_SIGNAL_MULTIPLEXING)
-		for (uint MultiplexId = 0; MultiplexId < MAX_SIGNAL_MULTIPLEXING; MultiplexId++)
 		{
-			KernelConfig.BufferColorSpace[MultiplexId] = CONFIG_INPUT_COLOR_SPACE;
-			KernelConfig.AccumulatorColorSpace[MultiplexId] = CONFIG_ACCUMULATION_COLOR_SPACE;
+			FIX_UNROLL(MAX_SIGNAL_MULTIPLEXING)
+			for (uint MultiplexId = 0; MultiplexId < MAX_SIGNAL_MULTIPLEXING; MultiplexId++)
+			{
+				KernelConfig.BufferColorSpace[MultiplexId] = CONFIG_INPUT_COLOR_SPACE;
+				KernelConfig.AccumulatorColorSpace[MultiplexId] = CONFIG_ACCUMULATION_COLOR_SPACE;
+			}
 		}
 
+		// SGPRs
 		KernelConfig.BufferSizeAndInvSize = View.BufferSizeAndInvSize;
 		KernelConfig.BufferBilinearUVMinMax = View.BufferBilinearUVMinMax;
 		KernelConfig.KernelSpreadFactor = KernelSpreadFactor;
-
-		KernelConfig.bSampleKernelCenter = CONFIG_UPSCALE || KernelConfig.bDescOrder;
-		KernelConfig.bForceKernelCenterAccumulation = true;
 		KernelConfig.BufferUV = SceneBufferUV;
 		KernelConfig.RefSceneMetadata = RefSceneMetadata;
+		
+		// Set up reference distance for all signals.
+		#if CONFIG_MAX_WITH_REF_DISTANCE
+		{
+			KernelConfig.bMaxWithRefBilateralDistance = true;
+
+			FIX_UNROLL(MAX_SIGNAL_MULTIPLEXING)
+			for (uint MultiplexId = 0; MultiplexId < MAX_SIGNAL_MULTIPLEXING; MultiplexId++)
+			{
+				if (KernelConfig.BilateralDistanceComputation == SIGNAL_WORLD_FREQUENCY_PRECOMPUTED_BLURING_RADIUS)
+				{
+					KernelConfig.RefBilateralDistance[MultiplexId] = RefSamples.Array[MultiplexId].WorldBluringRadius;
+				}
+				else
+				{
+					const uint BatchedSignalId = ComputeSignalBatchIdFromSignalMultiplexId(KernelConfig, MultiplexId);
+
+					KernelConfig.RefBilateralDistance[MultiplexId] = GetSignalWorldBluringRadius(RefSamples.Array[MultiplexId], RefSceneMetadata, BatchedSignalId);
+				}
+			}
+		}
+		#endif
+
+		// When doing history rejection preconvolution may have invalid ref sample, in witch case need to force take neighborhood to have a clamping box.
+		#if DIM_STAGE == STAGE_REJECTION_PRE_CONVOLUTION && CONFIG_UPSCALE
+		{
+			KernelConfig.bForceAllAccumulation = RefSamples.Array[0].AccumulatedSampleCount == 0;
+			KernelConfig.SampleSet = SAMPLE_SET_3X3_PLUS;
+		}
+		#endif
+
+		#if CONFIG_SAMPLE_SET == SAMPLE_SET_HEXAWEB
+		{
+			KernelConfig.RingCount = 1;
+			KernelConfig.KernelSpreadFactor = 8;
+
+			// TODO: could be improved.
+			//KernelConfig.bMinSamplePairInvFrequency = true;
+
+			float2 E = float2(
+				InterleavedGradientNoise(DispatchThreadId, 0),
+				InterleavedGradientNoise(DispatchThreadId, 1));
+
+			// Add a bit of jittering to hide low sample.
+			KernelConfig.bSampleKernelCenter = false;
+			KernelConfig.BufferUV += View.ViewSizeAndInvSize.zw * (E - 0.5) * (KernelConfig.KernelSpreadFactor * (0.5 * rcp(2)));
+		}
+		#endif
 
 		// When not upscaling, manually force accumulate the sample of the kernel.
-		if (!KernelConfig.bSampleKernelCenter)
+		if (!KernelConfig.bSampleKernelCenter && !KernelConfig.bDescOrder)
 		{
 			FIX_UNROLL(MAX_SIGNAL_MULTIPLEXING)
 			for (uint SignalMultiplexId = 0; SignalMultiplexId < MAX_SIGNAL_MULTIPLEXING; SignalMultiplexId++)
@@ -467,24 +526,6 @@ void MainCS(
 			}
 		}
 
-		#if DIM_STAGE == STAGE_PRE_CONVOLUTION
-		{
-			KernelConfig.RingCount = 1;
-			KernelConfig.KernelSpreadFactor = 8;
-
-			// TODO: could be improved.
-			KernelConfig.bMinSamplePairInvFrequency = true;
-		}
-		#endif
-
-		// When doing history rejection preconvolution may have invalid ref sample, in witch case need to force take neighborhood to have a clamping box.
-		#if DIM_STAGE == STAGE_REJECTION_PRE_CONVOLUTION && CONFIG_UPSCALE
-		{
-			KernelConfig.bForceAllAccumulation = RefSamples.Array[0].AccumulatedSampleCount == 0;
-			KernelConfig.SampleSet = SAMPLE_SET_3X3_PLUS;
-		}
-		#endif
-
 		if (CONFIG_SAMPLE_SET == SAMPLE_SET_STACKOWIAK_4_SETS)
 		{
 			KernelConfig.SampleCount = clamp(RequestedSampleCount * rcp(kStackowiakSampleSetCount), 1, MaxSampleCount);
@@ -551,6 +592,21 @@ void MainCS(
 				SignalInput_Textures_3,
 				/* inout */ SignalAccumulators);
 		}
+		
+		// Manually sample the center of the kernel after any accumulation when accumulating in descending order.
+		if (!KernelConfig.bSampleKernelCenter && KernelConfig.bDescOrder)
+		{
+			// Remove any jitter the kernel may have. Won't have ant VGPR cost when no jittering, because KernelConfig.BufferUV == SceneBufferUV.
+			KernelConfig.BufferUV = SceneBufferUV;
+
+			SampleAndAccumulateCenterSampleAsItsOwnCluster(
+				KernelConfig,
+				SignalInput_Textures_0,
+				SignalInput_Textures_1,
+				SignalInput_Textures_2,
+				SignalInput_Textures_3,
+				/* inout */ SignalAccumulators);
+		}
 	}
 
 	// Color processing of the signal to reduce highlight flickering.
@@ -593,6 +649,11 @@ void MainCS(
 			for (uint MultiplexId = 0; MultiplexId < CONFIG_SIGNAL_BATCH_SIZE; MultiplexId++)
 			{
 				OutputSamples.Array[MultiplexId] = SignalAccumulators.Array[MultiplexId].Moment1;
+				
+				// Output the minimal inverse frequency as new world bluring radius for subsequent passes.
+				// TODO: store this in its own signal to avoid *= N; *= rcp(N);
+				OutputSamples.Array[MultiplexId].WorldBluringRadius = 
+					OutputSamples.Array[MultiplexId].AccumulatedSampleCount * SignalAccumulators.Array[MultiplexId].MinInvFrequency;
 			}
 		}
 		#elif CONFIG_OUTPUT_MODE == OUTPUT_MODE_2MOMMENT_SUM
@@ -626,15 +687,6 @@ void MainCS(
 					OutputSamples.Array[MultiplexId].WorldBluringRadius = 0;
 				}
 			}
-
-			if (0 && DIM_STAGE == STAGE_PRE_CONVOLUTION)
-			{
-				DebugOutput[DispatchThreadId] = float4(
-					SignalAccumulators.Array[0].Previous.AccumulatedSampleCount,
-					SignalAccumulators.Array[0].CurrentTranslucency,
-					SignalAccumulators.Array[0].Previous.LeavingRayCount * SafeRcp(SignalAccumulators.Array[0].Previous.AccumulatedSampleCount),
-					SignalAccumulators.Array[0].Current.LeavingRayCount * SafeRcp(SignalAccumulators.Array[0].Current.AccumulatedSampleCount));
-			}
 		}
 		#else
 			#error Unknown output mode.
diff --git a/Engine/Shaders/Private/ScreenSpaceDenoise/SSDSpatialKernel.ush b/Engine/Shaders/Private/ScreenSpaceDenoise/SSDSpatialKernel.ush
--- a/Engine/Shaders/Private/ScreenSpaceDenoise/SSDSpatialKernel.ush
+++ b/Engine/Shaders/Private/ScreenSpaceDenoise/SSDSpatialKernel.ush
@@ -88,6 +88,9 @@ struct FSSDKernelConfig
 	// [ Jimenez 2014, "Next Generation Post Processing in Call of Duty: Advanced Warfare" ]
 	bool bMinSamplePairInvFrequency;
 
+	// Whether the bilateral distance should be maxed with reference bilateral distance.
+	bool bMaxWithRefBilateralDistance;
+
 	// The color space that has been encoded in the buffer.
 	uint BufferColorSpace[MAX_SIGNAL_MULTIPLEXING];
 	
@@ -130,6 +133,9 @@ struct FSSDKernelConfig
 
 	// Runtime selection of a track of sample.
 	uint SampleTrackId;
+
+	// Reference meta data.
+	float RefBilateralDistance[MAX_SIGNAL_MULTIPLEXING];
 };
 
 FSSDKernelConfig CreateKernelConfig()
@@ -147,19 +153,24 @@ FSSDKernelConfig CreateKernelConfig()
 	KernelConfig.bDescOrder = false;
 	KernelConfig.bNormalizeSample = false;
 	KernelConfig.bMinSamplePairInvFrequency = false;
+	KernelConfig.bMaxWithRefBilateralDistance = false;
 	
-	FIX_UNROLL(MAX_SIGNAL_MULTIPLEXING)
-	for (uint MultiplexId = 0; MultiplexId < MAX_SIGNAL_MULTIPLEXING; MultiplexId++)
 	{
-		KernelConfig.BufferColorSpace[MultiplexId] = STANDARD_BUFFER_COLOR_SPACE;
-		KernelConfig.AccumulatorColorSpace[MultiplexId] = STANDARD_BUFFER_COLOR_SPACE;
+		FIX_UNROLL(MAX_SIGNAL_MULTIPLEXING)
+		for (uint MultiplexId = 0; MultiplexId < MAX_SIGNAL_MULTIPLEXING; MultiplexId++)
+		{
+			KernelConfig.BufferColorSpace[MultiplexId] = STANDARD_BUFFER_COLOR_SPACE;
+			KernelConfig.AccumulatorColorSpace[MultiplexId] = STANDARD_BUFFER_COLOR_SPACE;
+		}
 	}
 	
+	// SGPRs.
 	KernelConfig.BufferSizeAndInvSize = float4(0, 0, 0, 0);
 	KernelConfig.BufferBilinearUVMinMax = float4(0, 0, 0, 0);
 	KernelConfig.BufferMipLevel = 0.0;
 	KernelConfig.KernelSpreadFactor = 1;
 	
+	// VGPRs.
 	KernelConfig.BoxKernelRadius = 1;
 	KernelConfig.SampleCount = 1;
 	KernelConfig.BufferUV = 0.0;
@@ -167,6 +178,15 @@ FSSDKernelConfig CreateKernelConfig()
 	KernelConfig.bForceKernelCenterAccumulation = false;
 	KernelConfig.bForceAllAccumulation = false;
 	KernelConfig.SampleTrackId = 0;
+
+	{	
+		FIX_UNROLL(MAX_SIGNAL_MULTIPLEXING)
+		for (uint MultiplexId = 0; MultiplexId < MAX_SIGNAL_MULTIPLEXING; MultiplexId++)
+		{
+			KernelConfig.RefBilateralDistance[MultiplexId] = 0.0;
+		}
+	}
+
 	return KernelConfig;
 }
 
@@ -309,6 +329,11 @@ void AccumulateSampledMultiplexedSignals(
 			FinalWorldBluringDistance = SignalConvolutionBluringRadius;
 		}
 		
+		if (KernelConfig.bMaxWithRefBilateralDistance)
+		{
+			FinalWorldBluringDistance = min(FinalWorldBluringDistance, KernelConfig.RefBilateralDistance[SignalMultiplexId]);
+		}
+	
 		float BilateralWeight = ComputeBilateralWeight(
 			FinalWorldBluringDistance,
 			KernelConfig.RefSceneMetadata,
@@ -450,7 +475,7 @@ void SampleAndAccumulateMultiplexedSignalsPair(
 void StartAccumulatingCluster(
 	FSSDKernelConfig KernelConfig,
 	inout FSSDSignalAccumulatorArray Accumulators,
-	uint RingId)
+	FSSDSampleClusterInfo ClusterInfo)
 {
 	FIX_UNROLL(MAX_SIGNAL_MULTIPLEXING)
 	for (uint SignalMultiplexId = 0; SignalMultiplexId < MAX_SIGNAL_MULTIPLEXING; SignalMultiplexId++)
@@ -458,7 +483,7 @@ void StartAccumulatingCluster(
 		StartAccumulatingCluster(
 			KernelConfig.RefSceneMetadata,
 			/* inout */ Accumulators.Array[SignalMultiplexId],
-			RingId);
+			ClusterInfo);
 	}
 }
 
@@ -473,7 +498,7 @@ void DijestAccumulatedClusterSamples(inout FSSDSignalAccumulatorArray Accumulato
 	}
 }
 
-void SampleAndAccumulateCenterSampleAsItsOwnRing(
+void SampleAndAccumulateCenterSampleAsItsOwnCluster(
 	FSSDKernelConfig KernelConfig,
 	Texture2D SignalBuffer0,
 	Texture2D SignalBuffer1,
@@ -483,7 +508,10 @@ void SampleAndAccumulateCenterSampleAsItsOwnRing(
 {
 	const uint RingId = 0;
 
-	StartAccumulatingCluster(KernelConfig, Accumulators, RingId);
+	FSSDSampleClusterInfo ClusterInfo;
+	ClusterInfo.OutterBoundaryRadius = (RingId + 1) * KernelConfig.KernelSpreadFactor;
+
+	StartAccumulatingCluster(KernelConfig, Accumulators, ClusterInfo);
 
 	SampleAndAccumulateMultiplexedSignals(
 		KernelConfig,
@@ -810,7 +838,10 @@ void ConvolveStackowiakKernel(
 			// TODO
 		}
 		
-		StartAccumulatingCluster(KernelConfig, /* inout */ Accumulators, CurrentRingId);
+		FSSDSampleClusterInfo ClusterInfo;
+		ClusterInfo.OutterBoundaryRadius = (CurrentRingId + 1) * KernelConfig.KernelSpreadFactor;
+
+		StartAccumulatingCluster(KernelConfig, /* inout */ Accumulators, ClusterInfo);
 		
 		// Processes the samples in batches so that the compiler can do lattency hidding.
 		LOOP
@@ -863,8 +894,11 @@ void ConvolveStackowiakKernel(
 					CurrentRingId -= 1;
 					NextClusterBoundary -= CurrentRingId * StocasticSamplesPerCluster;
 				
+					FSSDSampleClusterInfo ClusterInfo;
+					ClusterInfo.OutterBoundaryRadius = (CurrentRingId + 1) * KernelConfig.KernelSpreadFactor;
+
 					// Prepare the accumulators for new cluster.
-					StartAccumulatingCluster(KernelConfig, /* inout */ Accumulators, CurrentRingId);
+					StartAccumulatingCluster(KernelConfig, /* inout */ Accumulators, ClusterInfo);
 				}
 			} // for (uint SampleBatchId = 0; SampleBatchId < kSamplingBatchSize; SampleBatchId++)
 		} // for (uint BatchId = 0; BatchId < BatchCountCount; BatchId++)
@@ -1080,7 +1114,7 @@ void ConvolveDiskKernel(
 	// Accumulate the center of the kernel.
 	if (KernelConfig.bSampleKernelCenter && !KernelConfig.bDescOrder)
 	{
-		SampleAndAccumulateCenterSampleAsItsOwnRing(
+		SampleAndAccumulateCenterSampleAsItsOwnCluster(
 			KernelConfig,
 			SignalBuffer0,
 			SignalBuffer1,
@@ -1095,11 +1129,17 @@ void ConvolveDiskKernel(
 	#else
 		UNROLL
 	#endif
-	for (uint RingId = 1; RingId <= KernelConfig.RingCount; RingId++)
+	for (
+		uint RingId = (KernelConfig.bDescOrder ? KernelConfig.RingCount : 1);
+		(KernelConfig.bDescOrder ? RingId > 0 : RingId <= KernelConfig.RingCount);
+		RingId += (KernelConfig.bDescOrder ? ~0u : 1))
 	{
 		const uint RingSamplePairCount = GetRingSamplingPairCount(KernelConfig.SampleSet, RingId);
+		
+		FSSDSampleClusterInfo ClusterInfo;
+		ClusterInfo.OutterBoundaryRadius = (RingId + 1) * KernelConfig.KernelSpreadFactor;
 
-		StartAccumulatingCluster(KernelConfig, /* inout */ Accumulators, RingId);
+		StartAccumulatingCluster(KernelConfig, /* inout */ Accumulators, ClusterInfo);
 		
 		GatherRingSamples(
 			KernelConfig,
@@ -1116,7 +1156,7 @@ void ConvolveDiskKernel(
 	// Accumulate the center of the kernel.
 	if (KernelConfig.bSampleKernelCenter && KernelConfig.bDescOrder)
 	{
-		SampleAndAccumulateCenterSampleAsItsOwnRing(
+		SampleAndAccumulateCenterSampleAsItsOwnCluster(
 			KernelConfig,
 			SignalBuffer0,
 			SignalBuffer1,
@@ -1141,7 +1181,7 @@ void AccumulateKernel(
 	{
 		if (KernelConfig.bSampleKernelCenter)
 		{
-			SampleAndAccumulateCenterSampleAsItsOwnRing(
+			SampleAndAccumulateCenterSampleAsItsOwnCluster(
 				KernelConfig,
 				SignalBuffer0,
 				SignalBuffer1,
