Fixes shadow leaking artifacts in the shadow preconvolution pass.

diff --git a/Engine/Shaders/Private/ScreenSpaceDenoise/SSDSignalFramework.ush b/Engine/Shaders/Private/ScreenSpaceDenoise/SSDSignalFramework.ush
--- a/Engine/Shaders/Private/ScreenSpaceDenoise/SSDSignalFramework.ush
+++ b/Engine/Shaders/Private/ScreenSpaceDenoise/SSDSignalFramework.ush
@@ -520,14 +520,20 @@ float ComputeBilateralWeight(
 	{
 		float3 V = RefSceneMetadata.TranslatedWorldPosition - NeighborSceneMetadata.TranslatedWorldPosition;
 
-		if (dot(V, V) > MaxWorldBluringDistance * MaxWorldBluringDistance)
-		{
-			PositionBilateralWeight = 0;
-		}
-
 		// TODO: extremely expensive with rcp()...
 		PositionBilateralWeight = saturate(1 - dot(V, V) * rcp(MaxWorldBluringDistance * MaxWorldBluringDistance));
 	}
+	#elif CONFIG_VIEW_POSITION_BASED_BILATERAL == 5
+	{
+		float3 V = RefSceneMetadata.TranslatedWorldPosition - NeighborSceneMetadata.TranslatedWorldPosition;
+
+		float DistFromRefPlane = dot(RefSceneMetadata.WorldNormal, RefSceneMetadata.TranslatedWorldPosition - NeighborSceneMetadata.TranslatedWorldPosition);
+
+		// TODO: extremely expensive with rcp()...
+		PositionBilateralWeight = (
+			saturate(1 - dot(V, V) * rcp(MaxWorldBluringDistance * MaxWorldBluringDistance)) *
+			saturate(1 - 6 * abs(DistFromRefPlane) * rcp(MaxWorldBluringDistance)));
+	}
 	#endif
 
 	float NormalBilateralWeight = 1;
diff --git a/Engine/Shaders/Private/ScreenSpaceDenoise/SSDSpatialAccumulation.usf b/Engine/Shaders/Private/ScreenSpaceDenoise/SSDSpatialAccumulation.usf
--- a/Engine/Shaders/Private/ScreenSpaceDenoise/SSDSpatialAccumulation.usf
+++ b/Engine/Shaders/Private/ScreenSpaceDenoise/SSDSpatialAccumulation.usf
@@ -41,7 +41,7 @@
 // Configures all the pass for each individual signals.
 #if CONFIG_SIGNAL_PROCESSING == SIGNAL_PROCESSING_MONOCHROMATIC_PENUMBRA
 	//#define CONFIG_VIEW_POSITION_BASED_BILATERAL 2
-	#define CONFIG_VIEW_POSITION_BASED_BILATERAL 4
+	#define CONFIG_VIEW_POSITION_BASED_BILATERAL 5
 	#define CONFIG_MULTIPLEXED_SIGNALS_PER_SIGNAL_DOMAIN 1
 
 	#if DIM_STAGE == STAGE_RECONSTRUCTION
