Provide cone-angle culling behavior from the point-source location to provide mimic lighting function.

Change area source representation of spot light from a sphere to a disk.

diff --git a/Engine/Shaders/Private/RayTracing/RayTracingDiskLight.ush b/Engine/Shaders/Private/RayTracing/RayTracingDiskLight.ush
new file mode 100644
--- /dev/null
+++ b/Engine/Shaders/Private/RayTracing/RayTracingDiskLight.ush
@@ -0,0 +1,38 @@
+// Copyright 1998-2019 Epic Games, Inc. All Rights Reserved.
+
+#pragma once
+
+bool GenerateDiskLightOcclusionRay(
+	FLightShaderParameters LightParameters,
+	float3 WorldPosition,
+	float3 WorldNormal,
+	float2 RandSample,
+	out float3 RayOrigin,
+	out float3 RayDirection,
+	out float RayTMin,
+	out float RayTMax
+)
+{
+	// Sample disk of SourceRadius
+	float2 UV = UniformSampleDiskConcentric(RandSample);
+	float3 P_Local = float3(UV, 0.0) * LightParameters.SourceRadius;
+	float3 BiTangent = cross(LightParameters.Direction, LightParameters.Tangent);
+	float3 P_World = LightParameters.Position + P_Local.x * LightParameters.Tangent + P_Local.y * BiTangent;
+
+	// Construct light direction according to sample
+	float3 LightDirection = P_World - WorldPosition;
+	float RayLength = length(LightDirection);
+	LightDirection *= rcp(RayLength);
+
+	// Apply normal culling
+	float NoL = dot(LightDirection, LightParameters.Direction);
+	bool IsVisible = NoL > 0.0;
+	if (IsVisible)
+	{
+		RayOrigin = WorldPosition;
+		RayDirection = LightDirection;
+		RayTMin = 0.01;
+		RayTMax = RayLength;
+	}
+	return IsVisible;
+}
diff --git a/Engine/Shaders/Private/RayTracing/RayTracingSpotLight.ush b/Engine/Shaders/Private/RayTracing/RayTracingSpotLight.ush
--- a/Engine/Shaders/Private/RayTracing/RayTracingSpotLight.ush
+++ b/Engine/Shaders/Private/RayTracing/RayTracingSpotLight.ush
@@ -2,6 +2,8 @@
 
 #pragma once
 
+#include "RayTracingDiskLight.ush"
+
 bool GenerateSpotLightOcclusionRay(
 	FLightShaderParameters LightParameters,
 	float3 WorldPosition,
@@ -13,30 +15,29 @@ bool GenerateSpotLightOcclusionRay(
 	out float RayTMax
 )
 {
+	bool IsVisible = true;
+	float3 LightDirection = LightParameters.Position - WorldPosition;
+	float RayLength = length(LightDirection);
+	LightDirection *= rcp(RayLength);
+
 	if (LightParameters.SourceRadius > 0.0)
 	{
-		float RayPdf = 0.0;
-		GenerateSphereLightOcclusionRay(LightParameters, WorldPosition, WorldNormal, RandSample,
-			RayOrigin, RayDirection, RayTMin, RayTMax, RayPdf);
+		IsVisible = GenerateDiskLightOcclusionRay(LightParameters, WorldPosition, WorldNormal, RandSample,
+			RayOrigin, RayDirection, RayTMin, RayTMax);
 	}
 	else
 	{
-		float3 LightDirection = LightParameters.Position - WorldPosition;
-		float RayLength = length(LightDirection);
-		LightDirection /= RayLength;
-
-		// Apply normal perturbation when defining ray
 		RayOrigin = WorldPosition;
 		RayDirection = LightDirection;
-		RayTMin = 0.0;
+		RayTMin = 0.01;
 		RayTMax = RayLength;
 	}
 
 	// Apply culling
-	float CosAngle = LightParameters.SpotAngles.x;
-	if (dot(RayDirection, LightParameters.Direction) < CosAngle)
+	if (IsVisible)
 	{
-		return false;
+		float CosAngle = LightParameters.SpotAngles.x;
+		IsVisible = dot(LightDirection, LightParameters.Direction) >= CosAngle;
 	}
-	return true;
+	return IsVisible;
 }
