Get 16 more floats in Custom Primitive Data

diff --git a/Engine/Shaders/Private/SceneData.ush b/Engine/Shaders/Private/SceneData.ush
--- a/Engine/Shaders/Private/SceneData.ush
+++ b/Engine/Shaders/Private/SceneData.ush
@@ -9,7 +9,7 @@
 
 #if VF_USE_PRIMITIVE_SCENE_DATA
 
-#define NUM_CUSTOM_PRIMITIVE_DATA 4 // Num float4s used for custom data. Must match FCustomPrimitiveData::NumCustomPrimitiveDataFloat4s in SceneTypes.h
+#define NUM_CUSTOM_PRIMITIVE_DATA 8 // Num float4s used for custom data. Must match FCustomPrimitiveData::NumCustomPrimitiveDataFloat4s in SceneTypes.h
 
 // Must match FPrimitiveUniformShaderParameters in C++
 struct FPrimitiveSceneData
@@ -42,7 +42,7 @@ struct FPrimitiveSceneData
 };
 
 // Stride of a single primitive's data in float4's, must match C++
-#define PRIMITIVE_SCENE_DATA_STRIDE 31
+#define PRIMITIVE_SCENE_DATA_STRIDE 35
 
 // Fetch from scene primitive buffer
 FPrimitiveSceneData GetPrimitiveData(uint PrimitiveId) 
diff --git a/Engine/Source/Runtime/Engine/Public/PrimitiveUniformShaderParameters.h b/Engine/Source/Runtime/Engine/Public/PrimitiveUniformShaderParameters.h
--- a/Engine/Source/Runtime/Engine/Public/PrimitiveUniformShaderParameters.h
+++ b/Engine/Source/Runtime/Engine/Public/PrimitiveUniformShaderParameters.h
@@ -205,7 +205,7 @@ extern ENGINE_API TGlobalResource<FIdentityPrimitiveUniformBuffer> GIdentityPrim
 struct FPrimitiveSceneShaderData
 {
 	// Must match usf
-	enum { PrimitiveDataStrideInFloat4s = 31 };
+	enum { PrimitiveDataStrideInFloat4s = 35 };
 
 	FVector4 Data[PrimitiveDataStrideInFloat4s];
 
diff --git a/Engine/Source/Runtime/Engine/Public/SceneTypes.h b/Engine/Source/Runtime/Engine/Public/SceneTypes.h
--- a/Engine/Source/Runtime/Engine/Public/SceneTypes.h
+++ b/Engine/Source/Runtime/Engine/Public/SceneTypes.h
@@ -25,7 +25,7 @@ struct FCustomPrimitiveData
 {
 	GENERATED_USTRUCT_BODY()
 
-	static constexpr int32 NumCustomPrimitiveDataFloat4s = 4; // Must match NUM_CUSTOM_PRIMITIVE_DATA in SceneData.ush
+	static constexpr int32 NumCustomPrimitiveDataFloat4s = 8; // Must match NUM_CUSTOM_PRIMITIVE_DATA in SceneData.ush
 	static constexpr int32 NumCustomPrimitiveDataFloats = NumCustomPrimitiveDataFloat4s * 4;
 
 	UPROPERTY(EditAnywhere, Category=Rendering)
