Fixes incorrect albedo calculation when evaluating subsurface material in ray tracing SkyLight pass.

diff --git a/Engine/Shaders/Private/RayTracing/CompositeSkyLightPS.usf b/Engine/Shaders/Private/RayTracing/CompositeSkyLightPS.usf
--- a/Engine/Shaders/Private/RayTracing/CompositeSkyLightPS.usf
+++ b/Engine/Shaders/Private/RayTracing/CompositeSkyLightPS.usf
@@ -12,12 +12,7 @@ void CompositeSkyLightPS(
 )
 {
 	FScreenSpaceData ScreenSpaceData = GetScreenSpaceData(UV);
-	float3 Albedo = ScreenSpaceData.GBuffer.DiffuseColor;
-	if (UseSubsurfaceProfile(ScreenSpaceData.GBuffer.ShadingModelID))
-	{
-		Albedo = ScreenSpaceData.GBuffer.StoredBaseColor;
-	}
-
+	float3 Albedo = ScreenSpaceData.GBuffer.StoredBaseColor - ScreenSpaceData.GBuffer.StoredBaseColor * ScreenSpaceData.GBuffer.Metallic;
 	float4 SkyLight = SkyLightTexture.Sample(SkyLightTextureSampler, UV);
 	// Apply albedo after denoising
 	SkyLight.rgb *= Albedo;
diff --git a/Engine/Shaders/Private/RayTracing/RayTracingSkyLightRGS.usf b/Engine/Shaders/Private/RayTracing/RayTracingSkyLightRGS.usf
--- a/Engine/Shaders/Private/RayTracing/RayTracingSkyLightRGS.usf
+++ b/Engine/Shaders/Private/RayTracing/RayTracingSkyLightRGS.usf
@@ -246,10 +246,13 @@ void SkyLightRGS()
 	float3 CameraOrigin = ReconstructWorldPositionFromDepth(UV, 0.0f);
 	float3 CameraDirection = normalize(WorldPosition - CameraOrigin);
 	float3 WorldNormal = ScreenSpaceData.GBuffer.WorldNormal;
-	float3 BaseColor = ScreenSpaceData.GBuffer.BaseColor;
+	float3 Albedo = ScreenSpaceData.GBuffer.DiffuseColor;
+
+	// Recalculate DiffuseColor if subsurface reverted the contribution within the G-Buffer
 	if (UseSubsurfaceProfile(GBuffer.ShadingModelID))
 	{
-		ScreenSpaceData.GBuffer.DiffuseColor = ScreenSpaceData.GBuffer.StoredBaseColor;
+		Albedo = ScreenSpaceData.GBuffer.StoredBaseColor - ScreenSpaceData.GBuffer.StoredBaseColor * ScreenSpaceData.GBuffer.Metallic;
+		ScreenSpaceData.GBuffer.DiffuseColor = Albedo;
 	}
 	float3 SpecularColor = ScreenSpaceData.GBuffer.SpecularColor;
 	float Roughness = ScreenSpaceData.GBuffer.Roughness;
@@ -357,7 +360,6 @@ void SkyLightRGS()
 	}
 
 	// Pre-divide by albedo, to be recovered in compositing
-	float3 Albedo = ScreenSpaceData.GBuffer.BaseColor;
 	DiffuseExitantRadiance.r = Albedo.r > 0.0 ? DiffuseExitantRadiance.r / Albedo.r : DiffuseExitantRadiance.r;
 	DiffuseExitantRadiance.g = Albedo.g > 0.0 ? DiffuseExitantRadiance.g / Albedo.g : DiffuseExitantRadiance.g;
 	DiffuseExitantRadiance.b = Albedo.b > 0.0 ? DiffuseExitantRadiance.b / Albedo.b : DiffuseExitantRadiance.b;
