Niagara - Fix the mesh vertex factory shader to assign the LocalToWorld vector correctly.

diff --git a/Engine/Shaders/Private/NiagaraMeshVertexFactory.ush b/Engine/Shaders/Private/NiagaraMeshVertexFactory.ush
--- a/Engine/Shaders/Private/NiagaraMeshVertexFactory.ush
+++ b/Engine/Shaders/Private/NiagaraMeshVertexFactory.ush
@@ -779,7 +779,10 @@ FVertexFactoryInterpolantsVSToPS VertexFactoryGetInterpolantsVSToPS(FVertexFacto
 #endif
 
 #if USE_PARTICLE_TRANSFORM
-	Interpolants.LocalToWorld = Intermediates.LocalToWorld;
+	Interpolants.LocalToWorld[0] = Intermediates.LocalToWorld[0];
+	Interpolants.LocalToWorld[1] = Intermediates.LocalToWorld[1];
+	Interpolants.LocalToWorld[2] = Intermediates.LocalToWorld[2];
+	Interpolants.LocalToWorld[3] = Intermediates.LocalToWorld[3];
 #endif
 
 #if INSTANCED_STEREO
