Copying (fixing integration recrords)

//UE4/Main/Engine/Shaders/...

to //UE4/Dev-Core/Engine/Shaders/...

diff --git a/Engine/Shaders/Private/RayTracing/RayTracingRectLightRGS.usf b/Engine/Shaders/Private/RayTracing/RayTracingRectLightRGS.usf
new file mode 100644
--- /dev/null
+++ b/Engine/Shaders/Private/RayTracing/RayTracingRectLightRGS.usf
@@ -0,0 +1,345 @@
+// Copyright 1998-2019 Epic Games, Inc. All Rights Reserved.
+
+#include "../Common.ush"
+#include "../DeferredShadingCommon.ush"
+#include "../PathTracing/Utilities/PathTracingRandomSequence.ush"
+#include "../RectLight.ush"
+#include "../ShadingModels.ush"
+#include "RayTracingCommon.ush"
+#include "RayTracingRectLight.ush"
+#include "MipTreeCommon.ush"
+
+RaytracingAccelerationStructure TLAS;
+RWTexture2D<float4> RWLuminanceUAV;
+RWTexture2D<float> RWRayDistanceUAV;
+
+uint CalcLinearIndex(uint2 PixelCoord)
+{
+	return PixelCoord.y * View.BufferSizeAndInvSize.x + PixelCoord.x;
+}
+
+// #patrick todo: move to RayTracingCommon.ush
+float3 ReconstructWorldPositionFromDepth(float2 UV, float Depth)
+{
+	float2 ScreenPosition = (UV - View.ScreenPositionScaleBias.wz) / View.ScreenPositionScaleBias.xy;
+	float4 HomogeneousWorldPosition = mul(float4(ScreenPosition * Depth, Depth, 1), View.ScreenToWorld);
+	float3 WorldPosition = HomogeneousWorldPosition.xyz / HomogeneousWorldPosition.w;
+
+	return WorldPosition;
+}
+
+bool GenerateRectLightImportanceSampledRay(
+	float3 WorldPosition,
+	float3 WorldNormal,
+	float3 LightPosition,
+	float3 LightNormal,
+	float3 LightdPdu,
+	float3 LightdPdv,
+	float LightWidth,
+	float LightHeight,
+	inout RandomSequence RandSequence,
+	out float3 RayOrigin,
+	out float3 RayDirection,
+	out float RayTMin,
+	out float RayTMax,
+	out float RayPdf
+)
+{
+	uint MipCount = log2(RectLight.MipTreeDimensions.x);
+	uint2 MipPixel = 0;
+	float MipPdf = 1.0;
+
+	uint StopLevel = 0;
+	for (uint Index = 0; Index < MipCount - StopLevel; ++Index)
+	{
+		uint MipIndex = MipCount - Index - 1;
+		uint BufferOffset = BufferOffsetAtPixel(MipPixel, MipIndex, RectLight.MipTreeDimensions.xy);
+		uint4 BufferOffsets = BufferOffset + uint4(0, 1, 2, 3);
+		float4 PowerTerm = float4(RectLight.MipTree[BufferOffsets.x], RectLight.MipTree[BufferOffsets.y], RectLight.MipTree[BufferOffsets.z], RectLight.MipTree[BufferOffsets.w]);
+		PowerTerm = max(PowerTerm, 0.0);
+
+		// Determine sample positions
+		uint DummyVariable;
+		float2 Jitter = RandomSequence_GenerateSample2D(RandSequence, DummyVariable);
+		float2 DimInv = 1.0 / DimensionsAtLevel(RectLight.MipTreeDimensions.xy, MipIndex);
+		float2 UV00 = (MipPixel + Jitter) * DimInv - 0.5;
+		float2 UV10 = (MipPixel + Jitter + uint2(1, 0)) * DimInv - 0.5;
+		float2 UV01 = (MipPixel + Jitter + uint2(0, 1)) * DimInv - 0.5;
+		float2 UV11 = (MipPixel + Jitter + uint2(1, 1)) * DimInv - 0.5;
+
+		float3 Position00 = LightPosition + LightdPdu * LightWidth * UV00.x + LightdPdv * LightHeight * UV00.y;
+		float3 Position10 = LightPosition + LightdPdu * LightWidth * UV10.x + LightdPdv * LightHeight * UV10.y;
+		float3 Position01 = LightPosition + LightdPdu * LightWidth * UV01.x + LightdPdv * LightHeight * UV01.y;
+		float3 Position11 = LightPosition + LightdPdu * LightWidth * UV11.x + LightdPdv * LightHeight * UV11.y;
+
+		// Evaluate irradiance coefficients
+		float3 LightDirection00 = WorldPosition - Position00;
+		float3 LightDirection10 = WorldPosition - Position10;
+		float3 LightDirection01 = WorldPosition - Position01;
+		float3 LightDirection11 = WorldPosition - Position11;
+
+		float4 DistanceSquaredTerm = float4(
+			dot(LightDirection00, LightDirection00),
+			dot(LightDirection10, LightDirection10),
+			dot(LightDirection01, LightDirection01),
+			dot(LightDirection11, LightDirection11)
+		);
+
+		float4 CosineTerm = float4(
+			saturate(dot(normalize(LightDirection00), LightNormal)),
+			saturate(dot(normalize(LightDirection10), LightNormal)),
+			saturate(dot(normalize(LightDirection01), LightNormal)),
+			saturate(dot(normalize(LightDirection11), LightNormal))
+		);
+
+		// Build CDF on approximate irradiance
+		float4 Cdf = PowerTerm * CosineTerm / DistanceSquaredTerm;
+		Cdf.y += Cdf.x;
+		Cdf.z += Cdf.y;
+		Cdf.w += Cdf.z;
+		Cdf /= Cdf.w;
+
+		// Sample CDF
+		float RandSample = RandomSequence_GenerateSample1D(RandSequence, DummyVariable);
+		if (RandSample < Cdf.x)
+		{
+			MipPdf *= Cdf.x;
+		}
+		else if (RandSample < Cdf.y)
+		{
+			MipPdf *= Cdf.y - Cdf.x;
+			MipPixel.x += 1;
+		}
+		else if (RandSample < Cdf.z)
+		{
+			MipPdf *= Cdf.z - Cdf.y;
+			MipPixel.y += 1;
+		}
+		else // if (RandSample < Cdf.w)
+		{
+			MipPdf *= Cdf.w - Cdf.z;
+			MipPixel += 1;
+		}
+
+		MipPixel = uint2(MipPixel.x << 1, MipPixel.y << 1);
+	}
+	MipPixel = uint2(MipPixel.x >> 1, MipPixel.y >> 1);
+
+	// Sample within stop-level texel
+	uint DummyVariable;
+	float2 PixelOffset = RandomSequence_GenerateSample2D(RandSequence, DummyVariable);
+	float2 DimensionsAtStopLevel = DimensionsAtLevel(RectLight.MipTreeDimensions.xy, StopLevel);
+	float2 MipPixelUV = (MipPixel + PixelOffset) / DimensionsAtStopLevel;
+
+	// Map sample point to quad
+	MipPixelUV -= 0.5;
+	float2 RectLightOffset = MipPixel / DimensionsAtStopLevel;
+	float3 LightOrigin = LightPosition + LightdPdu * LightWidth * RectLightOffset.x + LightdPdv * LightHeight * RectLightOffset.y;
+	LightPosition += LightdPdu * LightWidth * MipPixelUV.x + LightdPdv * LightHeight * MipPixelUV.y;
+	float3 LightDirection = normalize(LightPosition - WorldPosition);
+
+	// Light-normal culling
+	if (dot(-LightDirection, LightNormal) <= 0.0)
+	{
+		RayPdf = 0.0;
+		return false;
+	}
+
+	// #dxr_todo: Performing direct lighting here could save some recalculation
+
+	// Define ray
+	RayOrigin = WorldPosition;
+	RayDirection = LightDirection;
+	RayTMin = 0.0;
+	RayTMax = length(LightPosition - WorldPosition);
+
+	// Construct solid angle pdf
+	FRect Rect;
+	Rect.Origin = LightOrigin - WorldPosition;
+	Rect.Axis[0] = LightdPdu;
+	Rect.Axis[1] = LightdPdv;
+	Rect.Axis[2] = cross(LightdPdu, LightdPdv);
+	Rect.Extent = float2(LightWidth, LightHeight) / DimensionsAtStopLevel;
+	FSphericalRect SphericalRect = BuildSphericalRect(Rect);
+	float SolidAnglePdf = 1.0 / SphericalRect.SolidAngle;
+
+	RayPdf = MipPdf * SolidAnglePdf;
+	return RayTMax > RayTMin;
+}
+
+[shader("raygeneration")]
+void RectLightRGS()
+{
+	uint2 PixelCoord = DispatchRaysIndex().xy;
+
+	RandomSequence RandSequence;
+	uint LinearIndex = CalcLinearIndex(PixelCoord);
+	RandomSequence_Initialize(RandSequence, LinearIndex, View.FrameNumber);
+
+	// Get G-Buffer surface data
+	float2 InvBufferSize = View.BufferSizeAndInvSize.zw;
+	float2 UV = (float2(PixelCoord) + 0.5) * InvBufferSize;
+	FScreenSpaceData ScreenSpaceData = GetScreenSpaceData(UV);
+	FGBufferData GBuffer = ScreenSpaceData.GBuffer;
+	float Depth = ScreenSpaceData.GBuffer.Depth;
+	float3 WorldPosition = ReconstructWorldPositionFromDepth(UV, Depth);
+	float3 CameraOrigin = ReconstructWorldPositionFromDepth(UV, 0.0f);
+	float3 CameraDirection = normalize(WorldPosition - CameraOrigin);
+	float3 WorldNormal = ScreenSpaceData.GBuffer.WorldNormal;
+	float3 BaseColor = ScreenSpaceData.GBuffer.BaseColor;
+	float3 DiffuseColor = ScreenSpaceData.GBuffer.DiffuseColor;
+	float3 SpecularColor = ScreenSpaceData.GBuffer.SpecularColor;
+	float Roughness = ScreenSpaceData.GBuffer.Roughness;
+
+	float RayDistance = 0.0;
+	float HitCount = 0.0;
+	uint SamplesPerPixel = RectLight.SamplesPerPixel;
+
+	// Mask out depth values that are infinitely far away
+	float DeviceZ = SceneTexturesStruct.SceneDepthTexture.Load(int3(PixelCoord, 0)).r;
+	bool IsFiniteDepth = DeviceZ > 0.0;
+	bool bTraceRay = (
+		IsFiniteDepth &&
+		GBuffer.ShadingModelID != SHADINGMODELID_UNLIT);
+	if (!bTraceRay)
+	{
+		SamplesPerPixel = 0.0;
+	}
+
+	float3 LightPosition = RectLight.Position;
+	float3 LightNormal = RectLight.Normal;
+	float3 LightdPdu = RectLight.dPdu;
+	float3 LightdPdv = RectLight.dPdv;
+	float LightWidth = RectLight.Width;
+	float LightHeight = RectLight.Height;
+	float LightBarnCosAngle = RectLight.BarnCosAngle;
+	float LightBarnLength = RectLight.BarnLength;
+
+	// Compute the visible portion of the light from the current shading point
+	{
+		float3 ToLight = LightPosition - WorldPosition;
+		FRect VisibleRect = GetRect(
+				ToLight,
+				-LightNormal,
+				LightdPdv,
+				LightWidth*0.5f,
+				LightHeight*0.5f,
+				RectLight.BarnCosAngle,
+				RectLight.BarnLength,
+				true);
+
+		LightPosition = LightPosition - LightdPdu * VisibleRect.Offset.x - LightdPdv * VisibleRect.Offset.y;	
+		LightWidth	  = VisibleRect.Extent.x * 2.0f;
+		LightHeight   = VisibleRect.Extent.y * 2.0f;
+	}
+
+	// Encode the rect light  parameter into the light parameters struct
+	// todo: send direction the light shader parameter data instead of the custon RectLight data.
+	FLightShaderParameters LightParameters;
+	LightParameters.Position = LightPosition;
+	LightParameters.Direction = -LightNormal;
+	LightParameters.Tangent = LightdPdu;
+	LightParameters.SourceRadius = 0.5f * LightWidth;
+	LightParameters.SourceLength = 0.5f * LightHeight;
+	LightParameters.RectLightBarnCosAngle =  LightBarnCosAngle;
+	LightParameters.RectLightBarnLength =  LightBarnLength;
+
+	if (LightWidth == 0 || LightHeight == 0) 
+		SamplesPerPixel = 0;
+
+	float3 ExitantRadiance = 0.0;
+	float3 DiffuseExitantRadiance = 0.0;
+	for (uint SampleIndex = 0; SampleIndex < SamplesPerPixel; ++SampleIndex)
+	{
+		RayDesc Ray;
+
+		uint DummyVariable;
+		float2 RandSample = RandomSequence_GenerateSample2D(RandSequence, DummyVariable);
+		float RayPdf = 0.0;
+
+		bool bIsValid = false;
+#if TEXTURE_IMPORTANCE_SAMPLING
+		if (RectLight.bIsTextureImportanceSampling)
+		{
+			bIsValid = GenerateRectLightImportanceSampledRay(WorldPosition, WorldNormal, LightPosition, LightNormal, LightdPdu, LightdPdv, LightWidth, LightHeight, RandSequence, Ray.Origin, Ray.Direction, Ray.TMin, Ray.TMax, RayPdf);
+		}
+		else
+#endif
+		{
+			bIsValid = GenerateRectLightOcclusionRay(LightParameters, WorldPosition, WorldNormal, RandSample, Ray.Origin, Ray.Direction, Ray.TMin, Ray.TMax, RayPdf);
+		}
+
+		if (!bIsValid)
+		{
+			continue;
+		}
+		ApplyPositionBias(Ray, WorldNormal, RectLight.MaxNormalBias);
+
+		float NoL = dot(WorldNormal, Ray.Direction);
+		if (NoL <= 0.0)
+		{
+			continue;
+		}
+
+		uint RayFlags = 0
+			| RAY_FLAG_CULL_BACK_FACING_TRIANGLES
+			| RAY_FLAG_FORCE_OPAQUE; // #dxr_todo: generate a special shader permutation for occlusion ray tracing that includes any-hit shaders for masked geometry
+
+		FDefaultPayload Payload = (FDefaultPayload)0;
+
+		TraceRay(
+			TLAS,      // AccelerationStructure
+			RayFlags,
+			RAY_TRACING_MASK_SHADOW,          // InstanceInclusionMask
+			RAY_TRACING_SHADER_SLOT_MATERIAL, // RayContributionToHitGroupIndex
+			RAY_TRACING_NUM_SHADER_SLOTS,     // MultiplierForGeometryContributionToShaderIndex
+			0,         // MissShaderIndex
+			Ray,       // RayDesc
+			Payload    // Payload
+		);
+
+		if (Payload.HitT >= 0.0)
+		{
+			RayDistance += Payload.HitT;
+			HitCount += 1.0;
+		}
+		else
+		{
+			// Sample light texture
+			float3 SampleLightPosition = Ray.Origin + Ray.Direction * Ray.TMax;
+			float3 RectLightCorner = LightPosition - 0.5 * (LightdPdu * LightWidth + LightdPdv * LightHeight);
+			float3 DeltaPosition = SampleLightPosition - RectLightCorner;
+			float2 SampleUV = float2(dot(DeltaPosition, LightdPdu) / LightWidth, dot(DeltaPosition, LightdPdv) / LightHeight);
+			SampleUV = 1.0 - SampleUV;
+			float3 TextureColor = Texture2DSample(RectLight.Texture, RectLight.TextureSampler, SampleUV).xyz;
+			float3 IncomingRadiance = RectLight.Color * TextureColor;
+
+			// Evaluate material
+			half3 N = WorldNormal;
+			half3 V = -CameraDirection;
+			half3 L = Ray.Direction;
+			FShadowTerms ShadowTerms = { 0.0, 0.0, 0.0 };
+			FDirectLighting LightingSample = EvaluateBxDF(ScreenSpaceData.GBuffer, N, V, L, NoL, ShadowTerms);
+			float3 Brdf = LightingSample.Diffuse + LightingSample.Transmission + LightingSample.Specular;
+
+			// #dxr_todo: Determine best way to incorporate integrals:
+			// 1) Return combined contribution (current method)
+			// 2) Split the integral (defer specular contribution to another pass)
+			// 3) Store diffuse and specular contribution separately
+			float3 Irradiance = IncomingRadiance * NoL / RayPdf;
+			ExitantRadiance += Irradiance * Brdf;
+			DiffuseExitantRadiance += Irradiance * (LightingSample.Diffuse + LightingSample.Transmission);
+		}
+	}
+
+	float3 OutputRadiance = (SamplesPerPixel > 0) ? ExitantRadiance / SamplesPerPixel : ExitantRadiance;
+
+#if USE_PREEXPOSURE
+	OutputRadiance *= View.PreExposure;
+#endif
+
+	OutputRadiance = ClampToHalfFloatRange(OutputRadiance);
+	RWLuminanceUAV[PixelCoord].rgb = OutputRadiance;
+	RWLuminanceUAV[PixelCoord].a = 1.0;
+	RWRayDistanceUAV[PixelCoord] = (HitCount > 0.0) ? RayDistance / HitCount : 1.0e27;
+}
diff --git a/Engine/Shaders/Private/SceneViewFamilyBlackboard.ush b/Engine/Shaders/Private/SceneViewFamilyBlackboard.ush
new file mode 100644
--- /dev/null
+++ b/Engine/Shaders/Private/SceneViewFamilyBlackboard.ush
@@ -0,0 +1,53 @@
+// Copyright 1998-2019 Epic Games, Inc. All Rights Reserved.
+
+/*=============================================================================
+	SceneViewFamilyBlackboard.ush: 
+=============================================================================*/
+
+#pragma once
+
+#include "SceneTexturesCommon.ush"
+#include "DeferredShadingCommon.ush"
+
+uint bIsSceneLightingChannelsValid;
+
+// Matches FSceneViewFamilyBlackboard
+Texture2D SceneDepthBuffer;
+Texture2D SceneVelocityBuffer;
+Texture2D SceneGBufferA;
+Texture2D SceneGBufferB;
+Texture2D SceneGBufferC;
+Texture2D SceneGBufferD;
+Texture2D SceneGBufferE;
+Texture2D<uint> SceneLightingChannels;
+
+
+// @param UV - UV space in the GBuffer textures (BufferSize resolution)
+FGBufferData GetGBufferDataFromSceneBlackboard(float2 UV, bool bGetNormalizedNormal = true)
+{
+	float4 GBufferA = SceneGBufferA.SampleLevel(GlobalPointClampedSampler, UV, 0);
+	float4 GBufferB = SceneGBufferB.SampleLevel(GlobalPointClampedSampler, UV, 0);
+	float4 GBufferC = SceneGBufferC.SampleLevel(GlobalPointClampedSampler, UV, 0);
+	float4 GBufferD = SceneGBufferD.SampleLevel(GlobalPointClampedSampler, UV, 0);
+	float4 GBufferE = SceneGBufferE.SampleLevel(GlobalPointClampedSampler, UV, 0);
+	float4 GBufferVelocity = SceneVelocityBuffer.SampleLevel(GlobalPointClampedSampler, UV, 0);
+
+	uint CustomStencil = 0;
+	float CustomNativeDepth = 0;
+
+	float DeviceZ = SceneDepthBuffer.SampleLevel(GlobalPointClampedSampler, UV, 0).r;
+	float SceneDepth = ConvertFromDeviceZ(DeviceZ);
+
+	return DecodeGBufferData(GBufferA, GBufferB, GBufferC, GBufferD, GBufferE, GBufferVelocity, CustomNativeDepth, CustomStencil, SceneDepth, bGetNormalizedNormal, CheckerFromSceneColorUV(UV));
+}
+
+/** Returns the light channel mask that should be executed for this pixel. */
+uint GetSceneLightingChannel(uint2 PixelCoord)
+{
+	BRANCH
+	if (bIsSceneLightingChannelsValid)
+	{
+		return SceneLightingChannels.Load(uint3(PixelCoord, 0)).x;
+	}
+	return ~0;
+}
