Merging //UE4/Dev-Main to Dev-Networking (//UE4/Dev-Networking)

diff --git a/Engine/Shaders/Private/RayTracing/RayTracingRectLightRGS.usf b/Engine/Shaders/Private/RayTracing/RayTracingRectLightRGS.usf
--- a/Engine/Shaders/Private/RayTracing/RayTracingRectLightRGS.usf
+++ b/Engine/Shaders/Private/RayTracing/RayTracingRectLightRGS.usf
@@ -212,16 +212,41 @@ void RectLightRGS()
 	float3 LightdPdv = RectLight.dPdv;
 	float LightWidth = RectLight.Width;
 	float LightHeight = RectLight.Height;
-	
+	float LightBarnCosAngle = RectLight.BarnCosAngle;
+	float LightBarnLength = RectLight.BarnLength;
+
+	// Compute the visible portion of the light from the current shading point
+	{
+		float3 ToLight = LightPosition - WorldPosition;
+		FRect VisibleRect = GetRect(
+				ToLight,
+				-LightNormal,
+				LightdPdv,
+				LightWidth*0.5f,
+				LightHeight*0.5f,
+				RectLight.BarnCosAngle,
+				RectLight.BarnLength,
+				true);
+
+		LightPosition = LightPosition - LightdPdu * VisibleRect.Offset.x - LightdPdv * VisibleRect.Offset.y;	
+		LightWidth	  = VisibleRect.Extent.x * 2.0f;
+		LightHeight   = VisibleRect.Extent.y * 2.0f;
+	}
+
 	// Encode the rect light  parameter into the light parameters struct
 	// todo: send direction the light shader parameter data instead of the custon RectLight data.
 	FLightShaderParameters LightParameters;
-	LightParameters.Position = RectLight.Position;
-	LightParameters.Direction = -RectLight.Normal;
-	LightParameters.Tangent = RectLight.dPdu;
-	LightParameters.SourceRadius = 0.5f * RectLight.Width;
-	LightParameters.SourceLength = 0.5f * RectLight.Height;
-	
+	LightParameters.Position = LightPosition;
+	LightParameters.Direction = -LightNormal;
+	LightParameters.Tangent = LightdPdu;
+	LightParameters.SourceRadius = 0.5f * LightWidth;
+	LightParameters.SourceLength = 0.5f * LightHeight;
+	LightParameters.RectLightBarnCosAngle =  LightBarnCosAngle;
+	LightParameters.RectLightBarnLength =  LightBarnLength;
+
+	if (LightWidth == 0 || LightHeight == 0) 
+		SamplesPerPixel = 0;
+
 	float3 ExitantRadiance = 0.0;
 	float3 DiffuseExitantRadiance = 0.0;
 	for (uint SampleIndex = 0; SampleIndex < SamplesPerPixel; ++SampleIndex)
@@ -265,7 +290,7 @@ void RectLightRGS()
 		TraceRay(
 			TLAS,      // AccelerationStructure
 			RayFlags,
-			RAY_TRACING_MASK_OPAQUE,          // InstanceInclusionMask
+			RAY_TRACING_MASK_SHADOW,          // InstanceInclusionMask
 			RAY_TRACING_SHADER_SLOT_MATERIAL, // RayContributionToHitGroupIndex
 			RAY_TRACING_NUM_SHADER_SLOTS,     // MultiplierForGeometryContributionToShaderIndex
 			0,         // MissShaderIndex
@@ -281,10 +306,10 @@ void RectLightRGS()
 		else
 		{
 			// Sample light texture
-			float3 LightPosition = Ray.Origin + Ray.Direction * Ray.TMax;
-			float3 RectLightCorner = RectLight.Position - 0.5 * (RectLight.dPdu * RectLight.Width + RectLight.dPdv * RectLight.Height);
-			float3 DeltaPosition = LightPosition - RectLightCorner;
-			float2 SampleUV = float2(dot(DeltaPosition, RectLight.dPdu) / RectLight.Width, dot(DeltaPosition, RectLight.dPdv) / RectLight.Height);
+			float3 SampleLightPosition = Ray.Origin + Ray.Direction * Ray.TMax;
+			float3 RectLightCorner = LightPosition - 0.5 * (LightdPdu * LightWidth + LightdPdv * LightHeight);
+			float3 DeltaPosition = SampleLightPosition - RectLightCorner;
+			float2 SampleUV = float2(dot(DeltaPosition, LightdPdu) / LightWidth, dot(DeltaPosition, LightdPdv) / LightHeight);
 			SampleUV = 1.0 - SampleUV;
 			float3 TextureColor = Texture2DSample(RectLight.Texture, RectLight.TextureSampler, SampleUV).xyz;
 			float3 IncomingRadiance = RectLight.Color * TextureColor;
@@ -307,8 +332,14 @@ void RectLightRGS()
 		}
 	}
 
-	float3 Result = (SamplesPerPixel > 0) ? ExitantRadiance / SamplesPerPixel : ExitantRadiance;
-	RWLuminanceUAV[PixelCoord].rgb = Result;
+	float3 OutputRadiance = (SamplesPerPixel > 0) ? ExitantRadiance / SamplesPerPixel : ExitantRadiance;
+
+#if USE_PREEXPOSURE
+	OutputRadiance *= View.PreExposure;
+#endif
+
+	OutputRadiance = ClampToHalfFloatRange(OutputRadiance);
+	RWLuminanceUAV[PixelCoord].rgb = OutputRadiance;
 	RWLuminanceUAV[PixelCoord].a = 1.0;
 	RWRayDistanceUAV[PixelCoord] = (HitCount > 0.0) ? RayDistance / HitCount : 1.0e27;
 }
