Fixes a bug where the denoiser could denoise invalid samples.

diff --git a/Engine/Shaders/Private/RayTracing/RayTracingReflections.usf b/Engine/Shaders/Private/RayTracing/RayTracingReflections.usf
--- a/Engine/Shaders/Private/RayTracing/RayTracingReflections.usf
+++ b/Engine/Shaders/Private/RayTracing/RayTracingReflections.usf
@@ -413,15 +413,14 @@ void RayTracingReflectionsRGS()
 					if (BounceIndex == 0)
 					{
 						ReflectedColor += PathRadiance;
-						ReflectedWeight += 1.0;
 					}
 					else
 					{
 						float4 FinalReflectionColor = PathRadiance;
 						//float4 FinalReflectionColor = BlendWithEnvironment(BRDFThroughput, Payload, Ray.Origin, PathRadiance, Ray.Direction, Payload.WorldNormal);
 						ReflectedColor += FinalReflectionColor;
-						ReflectedWeight += 1.0;
 					}
+					HitDistance = Payload.HitT;
 				}
 				else if (bAllowSkySampling)
 				{
@@ -432,12 +431,15 @@ void RayTracingReflectionsRGS()
 						float3 SkyLighting = GetSkyLightReflection(Ray.Direction, Roughness, SkyAverageBrightness);
 
 						ReflectedColor += float4(SkyLighting, 1.0f);
-						ReflectedWeight += 1.0;
+						
+						// To tell the denoiser the black input is valid even if no valid ray trace result was found. Uses a high ray hit distance
+						// so that the sample may be reused by another neighbor pixel in the denoise as if RoughnessFade in ]0;1[.
+						HitDistance = 1.0e20;
 					}
 				}
-
+				
+				ReflectedWeight += 1.0;
 				IncidentDirection = Ray.Direction;
-				HitDistance = Payload.HitT;
 			}
 		}
 	}
@@ -449,12 +451,6 @@ void RayTracingReflectionsRGS()
 
 		ReflectedColor *= DistanceFade * RoughnessFade * rcp(ReflectedWeight);
 	}
-	else
-	{
-		// To tell the denoiser the black input is valid even if no valid ray trace result was found. Uses a high ray hit distance
-		// so that the sample may be reused by another neighbor pixel in the denoise as if RoughnessFade in ]0;1[.
-		HitDistance = 1.0e20;
-	}
 
 	ColorOutput[DispatchThreadId] = float4(ReflectedColor);
 	RayHitDistanceOutput[DispatchThreadId] = HitDistance;
