Fixes a mistake in Denoiser's Tokoyashi implementation, and widen its use in reconstruction and temporal accumulation pass.

diff --git a/Engine/Shaders/Private/ScreenSpaceDenoise/SSDBuildHistoryRejection.usf b/Engine/Shaders/Private/ScreenSpaceDenoise/SSDBuildHistoryRejection.usf
--- a/Engine/Shaders/Private/ScreenSpaceDenoise/SSDBuildHistoryRejection.usf
+++ b/Engine/Shaders/Private/ScreenSpaceDenoise/SSDBuildHistoryRejection.usf
@@ -14,6 +14,8 @@
 	#define CONFIG_SIGNAL_REDUCE_LAYOUT		SIGNAL_BUFFER_LAYOUT_PENUMBRA_INTERNAL
 	#define CONFIG_SIGNAL_OUTPUT_LAYOUT		SIGNAL_BUFFER_LAYOUT_PENUMBRA_HISTORY_REJECTION
 
+	#define CONFIG_VIEW_POSITION_BASED_BILATERAL 1
+
 #else
 	#error Unknown signal processing
 
diff --git a/Engine/Shaders/Private/ScreenSpaceDenoise/SSDOutput.usf b/Engine/Shaders/Private/ScreenSpaceDenoise/SSDOutput.usf
--- a/Engine/Shaders/Private/ScreenSpaceDenoise/SSDOutput.usf
+++ b/Engine/Shaders/Private/ScreenSpaceDenoise/SSDOutput.usf
@@ -21,6 +21,8 @@
 #if CONFIG_SIGNAL_PROCESSING == SIGNAL_PROCESSING_PENUMBRA
 	#define CONFIG_SIGNAL_INPUT_LAYOUT SIGNAL_BUFFER_LAYOUT_PENUMBRA_INTERNAL
 
+	#define CONFIG_VIEW_POSITION_BASED_BILATERAL 1
+
 #else
 	#error Unknown signal processing
 
diff --git a/Engine/Shaders/Private/ScreenSpaceDenoise/SSDReduce.usf b/Engine/Shaders/Private/ScreenSpaceDenoise/SSDReduce.usf
--- a/Engine/Shaders/Private/ScreenSpaceDenoise/SSDReduce.usf
+++ b/Engine/Shaders/Private/ScreenSpaceDenoise/SSDReduce.usf
@@ -13,6 +13,8 @@
 	#define CONFIG_SIGNAL_INPUT_LAYOUT		SIGNAL_BUFFER_LAYOUT_PENUMBRA_INPUT
 	#define CONFIG_SIGNAL_OUTPUT_LAYOUT		SIGNAL_BUFFER_LAYOUT_PENUMBRA_INTERNAL
 
+	#define CONFIG_VIEW_POSITION_BASED_BILATERAL 1
+
 #else
 	#error Unknown signal processing
 
diff --git a/Engine/Shaders/Private/ScreenSpaceDenoise/SSDSignalFramework.ush b/Engine/Shaders/Private/ScreenSpaceDenoise/SSDSignalFramework.ush
--- a/Engine/Shaders/Private/ScreenSpaceDenoise/SSDSignalFramework.ush
+++ b/Engine/Shaders/Private/ScreenSpaceDenoise/SSDSignalFramework.ush
@@ -5,8 +5,24 @@
 
 //------------------------------------------------------- CONFIG DISABLED DEFAULTS
 
-#ifndef CONFIG_TOKOYASHI_BILATERAL
-	#define CONFIG_TOKOYASHI_BILATERAL 0
+/** Whether should use bilateral rejection for significantly different roughness. */
+#ifndef CONFIG_TOKOYASHI_LOBE_BILATERAL
+	#define CONFIG_TOKOYASHI_LOBE_BILATERAL 0
+#endif
+
+/** Whether should use bilateral rejection for significantly different lobe axis directions. */
+#ifndef CONFIG_TOKOYASHI_AXES_BILATERAL
+	#define CONFIG_TOKOYASHI_AXES_BILATERAL 0
+#endif
+
+/** Whether should use normal based rejection. */
+#ifndef CONFIG_VIEW_POSITION_BASED_BILATERAL
+	#define CONFIG_VIEW_POSITION_BASED_BILATERAL 0
+#endif
+
+/** Whether should use normal based rejection. */
+#ifndef CONFIG_NORMAL_BASED_BILATERAL
+	#define CONFIG_NORMAL_BASED_BILATERAL 0
 #endif
 
 #ifndef CONFIG_SPECULAR_RATIO_ESTIMATOR
@@ -219,16 +235,19 @@ float HitDistanceToWorldBluringRadius;
 /** Create a spherical gaussian lob for a given roughness. */
 FSphericalGaussian ComputeRoughnessLobe(float Roughness, float3 N, float3 V)
 {
-	float a = Pow2( max( 0.02, Roughness ) );
+	//float a = Pow2( max( 0.02, Roughness ) );
+	float a = Pow2( max( 0.001, Roughness ) );
 	float a2 = a*a;
 	
 	float NoV = saturate( abs( dot(N, V) ) + 1e-5 );
+	//float NoV = saturate( abs( dot(N, V) ) );
 	
 	float3 R = 2 * NoV * N - V;
 
 	FSphericalGaussian SpecularSG;
 	SpecularSG.Axis = R;
 	SpecularSG.Sharpness = 0.5 / ( a2 * max( NoV, 0.1 ) );
+	//SpecularSG.Sharpness = 0.5 / ( a2 * max( NoV, 0.1 ) );
 	SpecularSG.Amplitude = rcp( PI * a2 );
 
 	return SpecularSG;
@@ -274,23 +293,6 @@ float ComputeBluringRadiusOfNormalizedSample(FSSFSignalSample Sample, FSSFSample
 	return InifinityRadius * (1 - DistanceCameraToPixel / (DistanceCameraToPixel + Sample.PenumbraClosestHit));
 }
 
-// [ Tokoyashi 2015, "Specular Lobe-Aware Filtering and Upsampling for Interactive Indirect Illumination" ]
-float TokoyashiLobeBilateralWeight(FSSFSampleSceneInfos RefSceneMetadata, FSSFSampleSceneInfos NeighborSceneMetadata)
-{
-	FSphericalGaussian Ref = ComputeRoughnessLobe(RefSceneMetadata);
-	FSphericalGaussian Neighbor = ComputeRoughnessLobe(NeighborSceneMetadata);
-	
-#if 1
-	const float Beta = 5.05; // TODO: to be tunned with GGX.
-
-	float CombinedSharpness = sqrt(Ref.Sharpness * Neighbor.Sharpness) / (Ref.Sharpness + Neighbor.Sharpness);
-
-	return pow(2 * CombinedSharpness, Beta) * exp((Beta * CombinedSharpness) * (dot(Ref.Axis, Neighbor.Axis) - 1));
-#else
-	return Dot(Ref, Neighbor);
-#endif
-}
-
 // TODO: duplicated...
 float3 GenerateReflectedRayDirection(
 	float3 IncidentDirection,
@@ -307,12 +309,23 @@ float3 GenerateReflectedRayDirection(
 	}
 	else
 	{
-		float NoV = saturate(dot(-IncidentDirection, WorldNormal));
+		float3 N = WorldNormal;
+		float3 V = -IncidentDirection;
 		float2 E = RandSample;
+
+		float3x3 TangentBasis = GetTangentBasis(N);
+		float3 TangentV = mul(TangentBasis, V);
+
+		float NoV = saturate(dot(V, WorldNormal));
+		
 		float4 Sample = ImportanceSampleGGX(E, Pow4(Roughness));
-		float3 H = TangentToWorld(Sample.xyz, WorldNormal);
-		RayDirection = reflect(IncidentDirection, H);
+		if (View.GeneralPurposeTweak > 0.5)
+			Sample = ImportanceSampleVisibleGGX(E, Pow4(Roughness), TangentV);
+
+		float3 H = mul(Sample.xyz, TangentBasis);
+		float3 L = 2 * dot( V, H ) * H - V;
 
+		RayDirection = L;
 		PDF = Sample.w;
 	}
 
@@ -385,15 +398,14 @@ float ComputeBilateralWeight(
 	float MaxWorldBluringDistance,
 	FSSFSampleSceneInfos RefSceneMetadata,
 	FSSFSampleSceneInfos NeighborSceneMetadata)
-#if CONFIG_SIGNAL_PROCESSING == SIGNAL_PROCESSING_PENUMBRA || CONFIG_SIGNAL_PROCESSING == SIGNAL_PROCESSING_AO
 {
 	if (DEBUG_DISABLE_BILATERAL)
 	{
 		return 1;
 	}
 	
-	float WorldBilateralWeight = 1;
-	if (1)
+	float ViewPositionBilateralWeight = 1;
+	#if CONFIG_VIEW_POSITION_BASED_BILATERAL
 	{
 		const float WorldRadius = MaxWorldBluringDistance;
 
@@ -410,11 +422,12 @@ float ComputeBilateralWeight(
 
 		float Multiplier = rcp(WorldRadius * WorldRadius);
 
-		WorldBilateralWeight = saturate(1 - DistSquare * Multiplier);
+		ViewPositionBilateralWeight = saturate(1 - DistSquare * Multiplier);
 	}
+	#endif
 
 	float NormalBilateralWeight = 1;
-	if (CONFIG_SIGNAL_PROCESSING == SIGNAL_PROCESSING_AO)
+	#if CONFIG_NORMAL_BASED_BILATERAL
     {
 		float NoN = max(dot(RefSceneMetadata.WorldNormal, NeighborSceneMetadata.WorldNormal), 0);
 
@@ -431,23 +444,30 @@ float ComputeBilateralWeight(
 			NormalBilateralWeight = exp(-x * x);
 		#endif
     }
-
-	return WorldBilateralWeight * NormalBilateralWeight;
-}
-#elif CONFIG_SIGNAL_PROCESSING == SIGNAL_PROCESSING_REFLECTIONS
-{
-	float LobeBilateralWeight = 1;
-	#if CONFIG_TOKOYASHI_BILATERAL
+	#endif
+	
+	// [ Tokoyashi 2015, "Specular Lobe-Aware Filtering and Upsampling for Interactive Indirect Illumination" ]
+	float LobeSimilarity = 1;
+	float AxesSimilarity = 1;
+	#if CONFIG_TOKOYASHI_LOBE_BILATERAL || CONFIG_TOKOYASHI_AXES_BILATERAL
 	{
-		LobeBilateralWeight = TokoyashiLobeBilateralWeight(RefSceneMetadata, NeighborSceneMetadata);
+		const float Beta = 32;
+
+		FSphericalGaussian Ref = ComputeRoughnessLobe(RefSceneMetadata);
+		FSphericalGaussian Neighbor = ComputeRoughnessLobe(NeighborSceneMetadata);
+	
+		float InvSharpnessSum = rcp(Ref.Sharpness + Neighbor.Sharpness);
+
+		if (CONFIG_TOKOYASHI_LOBE_BILATERAL)
+			LobeSimilarity = pow(2 * sqrt(Ref.Sharpness * Neighbor.Sharpness) * InvSharpnessSum, Beta);
+
+		if (CONFIG_TOKOYASHI_AXES_BILATERAL)
+			AxesSimilarity = exp(-(Beta * (Ref.Sharpness * Neighbor.Sharpness) * InvSharpnessSum) * saturate(1 - dot(Ref.Axis, Neighbor.Axis)));
 	}
 	#endif
 
-	return LobeBilateralWeight;
+	return ViewPositionBilateralWeight * NormalBilateralWeight * LobeSimilarity * AxesSimilarity;
 }
-#else
-	#error Unimplemented ComputeBilateralWeight()
-#endif
 
 float GetSignalWorldBluringRadius(FSSFSignalSample Sample, FSSFSampleSceneInfos SceneMetadata)
 #if CONFIG_SIGNAL_PROCESSING == SIGNAL_PROCESSING_PENUMBRA
diff --git a/Engine/Shaders/Private/ScreenSpaceDenoise/SSDSpatialAccumulation.usf b/Engine/Shaders/Private/ScreenSpaceDenoise/SSDSpatialAccumulation.usf
--- a/Engine/Shaders/Private/ScreenSpaceDenoise/SSDSpatialAccumulation.usf
+++ b/Engine/Shaders/Private/ScreenSpaceDenoise/SSDSpatialAccumulation.usf
@@ -35,6 +35,8 @@
 	#define CONFIG_SIGNAL_INPUT_LAYOUT  SIGNAL_BUFFER_LAYOUT_PENUMBRA_INTERNAL
 	#define CONFIG_SIGNAL_OUTPUT_LAYOUT SIGNAL_BUFFER_LAYOUT_PENUMBRA_HISTORY_REJECTION
 
+	#define CONFIG_VIEW_POSITION_BASED_BILATERAL 1
+
 	#define CONFIG_SAMPLE_SET SAMPLE_SET_3X3
 
 	#if DIM_STAGE == STAGE_POST_FILTERING
@@ -53,14 +55,19 @@
 		#define CONFIG_SAMPLE_SET           SAMPLE_SET_STACKOWIAK_4_SETS
 		#define CONFIG_SAMPLE_COUNT_POLICY  SAMPLE_COUNT_POLICY_HIT_DISTANCE_BASED
 
-		// Reconstruction has 1 ray per input pixel, therefore can use MIS.
+		// Reconstruction has 1 ray per input pixel, therefore can use specular ratio estimator instead of Tokoyashi axes bilateral kernel.
 		#define CONFIG_SPECULAR_RATIO_ESTIMATOR 1
 
+		// Uses tokoyashi lobe bilateral rejection to avoid mixing up significantly different lobes that can cause leaking/strong biasing.
+		#define CONFIG_TOKOYASHI_LOBE_BILATERAL 1
+
 		//#define CONFIG_SAMPLE_COUNT_POLICY  SAMPLE_COUNT_POLICY_DISABLED
 
 		// Do color accumulation with karis weighting to avoid flickering specular highlight to show up the kernel pattern.
 		//#define CONFIG_COLOR_SPACE          (COLOR_SPACE_RGB | COLOR_SPACE_KARIS_WEIGHTING)
 
+		#define CONFIG_REJECT_HIGHEST_COLOR 0
+
 
 	#elif DIM_STAGE == STAGE_POST_FILTERING
 		// Input and output layout.
@@ -72,17 +79,18 @@
 		//#define CONFIG_SAMPLE_SUBSET        1
 		#define CONFIG_SAMPLE_COUNT_POLICY  SAMPLE_COUNT_POLICY_ACCUMULATION_BASED
 
-		// Post filtering can only be done using tokoyashi because have more than one sample at a time.
-		#define CONFIG_TOKOYASHI_BILATERAL 1
-
-		// TODO: Not sure of the math of Tokoyashi, weight of the center kernel might be != 1.
-		#define CONFIG_WEIGHT_KERNEL_CENTER 1
+		// Can only be done using tokoyashi because have more than one sample at a time.
+		#define CONFIG_TOKOYASHI_LOBE_BILATERAL 1
+		#define CONFIG_TOKOYASHI_AXES_BILATERAL 1
 
 	#else
 		#error Unexpected stage.
 	#endif
 
 #elif CONFIG_SIGNAL_PROCESSING == SIGNAL_PROCESSING_AO
+	#define CONFIG_VIEW_POSITION_BASED_BILATERAL 1
+	#define CONFIG_NORMAL_BASED_BILATERAL 1
+
 	#if DIM_STAGE == STAGE_RECONSTRUCTION
 		// Input and output layout.
 		#define CONFIG_SIGNAL_INPUT_LAYOUT  SIGNAL_BUFFER_LAYOUT_AO_INPUT
@@ -138,6 +146,11 @@
 	#define CONFIG_COLOR_SPACE CONFIG_BUFFER_COLOR_SPACE
 #endif
 
+/** Removes the highest color. */
+#ifndef CONFIG_REJECT_HIGHEST_COLOR
+	#define CONFIG_REJECT_HIGHEST_COLOR 0
+#endif
+
 
 //------------------------------------------------------- INCLUDES
 
@@ -363,6 +376,7 @@ void MainCS(
 	#endif
 	
 	FSSFSignalSample NeighborMoment2 = Power(NeighborMoment1, 2.0);
+	FSSFSignalSample MaxNeighbor = NeighborMoment1;
 
 	float CurrentWorldBluringRadius = ComputePixelWorldRadius(RefSceneMetadata);
 
@@ -547,8 +561,21 @@ void MainCS(
 		
 		NeighborMoment1 = Add(NeighborMoment1, Mul(NeighborSample, FinalWeight));
 		NeighborMoment2 = Add(NeighborMoment2, Mul(Power(NeighborSample, 2.0), FinalWeight));
-	}
+
+		MaxNeighbor = Max(MaxNeighbor, Mul(NeighborSample, FinalWeight));
+	} // for neighbors)
 	
+	#if CONFIG_REJECT_HIGHEST_COLOR
+	if (View.GeneralPurposeTweak > 0.5 && NeighborMoment1.AccumulatedSampleCount > 0)
+	{
+		const float MaxNeighborWeight = saturate(NeighborMoment1.AccumulatedSampleCount * rcp(10) - 1);
+
+		NeighborMoment1.SceneColor.rgb =
+			(NeighborMoment1.SceneColor.rgb - MaxNeighbor.SceneColor.rgb * MaxNeighborWeight) *
+			(NeighborMoment1.AccumulatedSampleCount / (NeighborMoment1.AccumulatedSampleCount - MaxNeighborWeight));
+	}
+	#endif
+
 	NeighborMoment1 = TransformSignal(
 		NeighborMoment1, 
 		/* SrcBasis  = */ CONFIG_COLOR_SPACE, 
diff --git a/Engine/Shaders/Private/ScreenSpaceDenoise/SSDTemporalAccumulation.usf b/Engine/Shaders/Private/ScreenSpaceDenoise/SSDTemporalAccumulation.usf
--- a/Engine/Shaders/Private/ScreenSpaceDenoise/SSDTemporalAccumulation.usf
+++ b/Engine/Shaders/Private/ScreenSpaceDenoise/SSDTemporalAccumulation.usf
@@ -17,6 +17,8 @@
 	#define CONFIG_SIGNAL_INPUT_LAYOUT   SIGNAL_BUFFER_LAYOUT_PENUMBRA_INTERNAL
 	#define CONFIG_SIGNAL_HISTORY_LAYOUT SIGNAL_BUFFER_LAYOUT_PENUMBRA_INTERNAL
 
+	#define CONFIG_VIEW_POSITION_BASED_BILATERAL 1
+
 	// Uses a larger kernel when burteforcing to converge faster.
 	#define CONFIG_LARGER_HISTORY_KERNEL 1
 
@@ -41,6 +43,11 @@
 	#define CONFIG_SIGNAL_INPUT_LAYOUT   SIGNAL_BUFFER_LAYOUT_REFLECTIONS_HISTORY
 	#define CONFIG_SIGNAL_HISTORY_LAYOUT SIGNAL_BUFFER_LAYOUT_REFLECTIONS_HISTORY
 
+	// Can only be done using tokoyashi because have more than one sample at a time.
+	#define CONFIG_TOKOYASHI_LOBE_BILATERAL 1
+	#define CONFIG_TOKOYASHI_AXES_BILATERAL 1
+	#define CONFIG_VIEW_POSITION_BASED_BILATERAL 1
+
 	#define CONFIG_HISTORY_REJECTION 2
 	#define CONFIG_HISTORY_REJECTION_COLOR_SPACE (COLOR_SPACE_LCOCG | COLOR_SPACE_UNPREMULTIPLY)
 
@@ -53,6 +60,9 @@
 
 	#define CONFIG_HISTORY_REJECTION 1
 
+	#define CONFIG_VIEW_POSITION_BASED_BILATERAL 1
+	#define CONFIG_NORMAL_BASED_BILATERAL 1
+
 #else
 	#error Unimplemented signal processing.
 #endif
@@ -657,6 +667,7 @@ void TemporallyAccumulate(
 			FSSFSignalSample NeighborSample = CreateSignalSample();
 			float NeighborPenumbra = -1;
 			float NeighborUncertainty = 1.0;
+			bool bUseNeighborSample = false;
 
 			#if CONFIG_INPUT_TILE_CLASSIFICATION
 			{
@@ -751,12 +762,15 @@ void TemporallyAccumulate(
 					CurrentSceneMetadata,
 					NeighborSceneMetadata);
 
-				if (BilateralWeight > 0 && NeighborSample.AccumulatedSampleCount > 0)
+				// TODO: requires further improvement for CurrentFrame.AccumulatedSampleCount == 0.
+				if (NeighborSample.AccumulatedSampleCount > 0 &&
+					(BilateralWeight > 0 || (CurrentFrame.AccumulatedSampleCount == 0 && abs(x) <= 1 && abs(y) <= 1)))
 				{
 					NeighborPenumbra = GetSamplePenumbra(NeighborSample);
 					NeighborUncertainty = bUsesHistoryRejectionSignal
 						? 0.0
 						: NeighborSample.UncertainRayCount / (NeighborSample.AccumulatedSampleCount + NeighborSample.UncertainRayCount);
+					bUseNeighborSample = true;
 				}
 			}
 			#endif // !CONFIG_INPUT_TILE_CLASSIFICATION
@@ -765,11 +779,15 @@ void TemporallyAccumulate(
 			NormalizedNeighborSample = TransformSignalForPreRejection(NormalizedNeighborSample);
 			
 			#if CONFIG_SIGNAL_PROCESSING == SIGNAL_PROCESSING_PENUMBRA
-				bool bUseNeighborSample = NeighborUncertainty < 1.0;
-			#else
-				bool bUseNeighborSample = NeighborSample.AccumulatedSampleCount > 0.0;
+				bUseNeighborSample = NeighborUncertainty < 1.0;
 			#endif
 
+			// When the denoiser is upsampling, may have no input information, and have neighbor being bilaterally rejected. In this case
+			// Force the immediate neighborhood to avoid ghosting.
+			bool bForceUseNeighbor = CurrentFrame.AccumulatedSampleCount == 0 && NeighborSample.AccumulatedSampleCount > 0 && abs(x) <= 1 && abs(y) <= 1;
+
+			//bUseNeighborSample = bUseNeighborSample || bForceUseNeighbor;
+
 			if (CONFIG_CERTAINTY_BASED_REJECTION && bUseNeighborSample)
 			{
 				NeighborMinPenumbra = min(NeighborMinPenumbra, NeighborPenumbra + NeighborUncertainty);
