Using spherical rectangles when sampling RectLight.

diff --git a/Engine/Shaders/Private/PathTracing/Light/PathTracingRectLight.ush b/Engine/Shaders/Private/PathTracing/Light/PathTracingRectLight.ush
--- a/Engine/Shaders/Private/PathTracing/Light/PathTracingRectLight.ush
+++ b/Engine/Shaders/Private/PathTracing/Light/PathTracingRectLight.ush
@@ -56,17 +56,36 @@ void RectLight_SampleLight(
 	out float OutPdf
 )
 {
-	// Pick a random position on the light
+	float3 WorldPosition = Payload.WorldPos;
+	float3 LightPosition = GetPosition(LightId);
+	float3 LightNormal = GetNormal(LightId);
+	float3 LightdPdu = GetdPdu(LightId);
+	float3 LightdPdv = GetdPdv(LightId);
+	float LightWidth = GetWidth(LightId);
+	float LightHeight = GetHeight(LightId);
+
+	// Pick a random position
 	float2 UV = RandSample.yz - 0.5;
 	OutLightUV = float3(UV, 0);
 
-	// Convert PDF to solid angle measure
-	float Area = GetWidth(LightId) * GetHeight(LightId);
-	float2 DeltaP = float2(GetWidth(LightId), GetHeight(LightId)) * OutLightUV.xy;
-	float3 LightDir = GetPosition(LightId) + DeltaP.x * GetdPdu(LightId) + DeltaP.y * GetdPdv(LightId) - Payload.WorldPos;
-	float Distance2 = dot(LightDir, LightDir);
-	LightDir /= length(LightDir);
-	float CosTheta = saturate(dot(normalize(-LightDir), GetNormal(LightId)));
-	float SolidAngle = Area * CosTheta / Distance2;
-	OutPdf = 1.0 / SolidAngle;
+	// Map sample point to quad
+	float3 LightSamplePosition = LightPosition + LightdPdu * LightWidth * UV.x + LightdPdv * LightHeight * UV.y;
+	float3 LightDirection = normalize(LightSamplePosition - WorldPosition);
+
+	// Light-normal culling
+	if (dot(-LightDirection, LightNormal) <= 0.0)
+	{
+		OutPdf = 0.0;
+		return;
+	}
+
+	// Define rectangle and compute solid angle
+	FRect Rect;
+	Rect.Origin = LightPosition - WorldPosition;
+	Rect.Axis[0] = LightdPdu;
+	Rect.Axis[1] = LightdPdv;
+	Rect.Axis[2] = cross(LightdPdu, LightdPdv);
+	Rect.Extent = float2(LightWidth, LightHeight);
+	FSphericalRect SphericalRect = BuildSphericalRect(Rect);
+	OutPdf = isfinite(SphericalRect.SolidAngle) ? 1.0 / SphericalRect.SolidAngle : 0.0;
 }
\ No newline at end of file
