Missing include

diff --git a/Engine/Shaders/Private/PathTracing/PathTracing.usf b/Engine/Shaders/Private/PathTracing/PathTracing.usf
--- a/Engine/Shaders/Private/PathTracing/PathTracing.usf
+++ b/Engine/Shaders/Private/PathTracing/PathTracing.usf
@@ -6,7 +6,8 @@ PathTracingRayGenShader.usf: Reference path tracing
 
 #include "../Common.ush"
 #include "../PostProcessCommon.ush"
-#include "../Montecarlo.ush"
+#include "../RectLight.ush"
+//#include "../Montecarlo.ush"
 
 #include "PathTracingCommon.ush"
 #include "../RayTracing/RayTracingHitGroupCommon.ush"
@@ -96,7 +97,8 @@ void PathTracingMainRG()
 
 		float MaterialPdf = 0.0;
 		uint Bounces = 0;
-		while (Bounces < 32)
+		uint BouncesMax = 32;
+		while (Bounces < BouncesMax)
 		{
 			// Environment contribution
 			if (!IsHit(Ray, Payload))
