Correctly accounting for RectLight origin when sampling RectLight integrator

diff --git a/Engine/Shaders/Private/RayTracing/RayTracingRectLight.ush b/Engine/Shaders/Private/RayTracing/RayTracingRectLight.ush
--- a/Engine/Shaders/Private/RayTracing/RayTracingRectLight.ush
+++ b/Engine/Shaders/Private/RayTracing/RayTracingRectLight.ush
@@ -26,8 +26,8 @@ bool GenerateRectLightOcclusionRay(
 	RandSample -= 0.5;
 
 	// Map sample point to quad
-	LightPosition += LightdPdu * LightWidth * RandSample.x + LightdPdv * LightHeight * RandSample.y;
-	float3 LightDirection = normalize(LightPosition - WorldPosition);
+	float3 LightSamplePosition = LightPosition + LightdPdu * LightWidth * RandSample.x + LightdPdv * LightHeight * RandSample.y;
+	float3 LightDirection = normalize(LightSamplePosition - WorldPosition);
 
 	// Light-normal culling
 	if (dot(-LightDirection, LightNormal) <= 0.0)
@@ -39,7 +39,7 @@ bool GenerateRectLightOcclusionRay(
 	RayOrigin = WorldPosition + WorldNormal * 0.01;
 	RayDirection = LightDirection;
 	RayTMin = 0.0;
-	RayTMax = length(LightPosition - WorldPosition);
+	RayTMax = length(LightSamplePosition - WorldPosition);
 
 	// Define rectangle
 	FRect Rect;
