Uses ImportantSampleVisibleGGX() for reflection ray generation shader to reduce amount of invalid reflection samples at grazing angles causing problem in the reflection denoiser.

diff --git a/Engine/Shaders/Private/RayTracing/RayTracingReflections.usf b/Engine/Shaders/Private/RayTracing/RayTracingReflections.usf
--- a/Engine/Shaders/Private/RayTracing/RayTracingReflections.usf
+++ b/Engine/Shaders/Private/RayTracing/RayTracingReflections.usf
@@ -74,11 +74,21 @@ float3 GenerateReflectedRayDirection(
 	}
 	else
 	{
-		float NoV = saturate(dot(-IncidentDirection, WorldNormal));
+		float3 N = WorldNormal;
+		float3 V = -IncidentDirection;
 		float2 E = RandSample;
-		float4 Sample = ImportanceSampleGGX(E, Pow4(Roughness));
-		float3 H = TangentToWorld(Sample.xyz, WorldNormal);
-		RayDirection = reflect(IncidentDirection, H);
+
+		float3x3 TangentBasis = GetTangentBasis(N);
+		float3 TangentV = mul(TangentBasis, V);
+
+		float NoV = saturate(dot(V, WorldNormal));
+
+		float4 Sample = ImportanceSampleVisibleGGX(E, Pow4(Roughness), TangentV);
+
+		float3 H = mul(Sample.xyz, TangentBasis);
+		float3 L = 2 * dot( V, H ) * H - V;
+
+		RayDirection = L;
 	}
 
 	return RayDirection;
diff --git a/Engine/Shaders/Private/ScreenSpaceDenoise/SSDSignalFramework.ush b/Engine/Shaders/Private/ScreenSpaceDenoise/SSDSignalFramework.ush
--- a/Engine/Shaders/Private/ScreenSpaceDenoise/SSDSignalFramework.ush
+++ b/Engine/Shaders/Private/ScreenSpaceDenoise/SSDSignalFramework.ush
@@ -317,10 +317,8 @@ float3 GenerateReflectedRayDirection(
 		float3 TangentV = mul(TangentBasis, V);
 
 		float NoV = saturate(dot(V, WorldNormal));
-		
-		float4 Sample = ImportanceSampleGGX(E, Pow4(Roughness));
-		if (View.GeneralPurposeTweak > 0.5)
-			Sample = ImportanceSampleVisibleGGX(E, Pow4(Roughness), TangentV);
+
+		float4 Sample = ImportanceSampleVisibleGGX(E, Pow4(Roughness), TangentV);
 
 		float3 H = mul(Sample.xyz, TangentBasis);
 		float3 L = 2 * dot( V, H ) * H - V;
