Fix missing reflection on pure specular surface.

The throughput was using the diffuse color instead of the specular color. In addition, the throughput was attenuated by the projected radiance, while this terms should normally be cancelled out with the denominator of the mirror BSDF.

diff --git a/Engine/Shaders/Private/PathTracing/Material/PathTracingSpecularReflection.ush b/Engine/Shaders/Private/PathTracing/Material/PathTracingSpecularReflection.ush
--- a/Engine/Shaders/Private/PathTracing/Material/PathTracingSpecularReflection.ush
+++ b/Engine/Shaders/Private/PathTracing/Material/PathTracingSpecularReflection.ush
@@ -13,7 +13,7 @@ void SpecularReflection_SampleMaterial(
 	out float OutPdf)
 {
 	float3 N_World = GetWorldNormal(Payload);
-	float3 DiffuseColor = GetDiffuseColor(Payload);
+	float3 SpecularColor = GetSpecularColor(Payload);
 
 	//OutDirection = reflect(RayDirection, N_World); // To reflect in world coords directly
 	float3 V_World = normalize(-RayDirection);
@@ -24,9 +24,8 @@ void SpecularReflection_SampleMaterial(
 	L_Local.y = -V_Local.y;
 	L_Local.z = V_Local.z;
 	OutDirection = TangentToWorld(L_Local, N_World);
-
-	float NoL = max(L_Local.z, 0.0);
-	OutThroughput = DiffuseColor * NoL;
+	
+    OutThroughput = SpecularColor;
 	OutPdf = 1.0;
 }
 
