Plugged correct ShadingModelID to RT payload.

diff --git a/Engine/Shaders/Private/DeferredShadingCommon.ush b/Engine/Shaders/Private/DeferredShadingCommon.ush
--- a/Engine/Shaders/Private/DeferredShadingCommon.ush
+++ b/Engine/Shaders/Private/DeferredShadingCommon.ush
@@ -7,6 +7,7 @@
 #ifndef __DEFERRED_SHADING_COMMON__
 #define __DEFERRED_SHADING_COMMON__
 
+#include "ShadingCommon.ush"
 #include "LightAccumulator.ush"
 #include "SceneTexturesCommon.ush"
 
@@ -266,26 +267,6 @@ float Decode71(float Scalar, out uint Mask)
 	return (Scalar - 0.5f * Mask) * 2.0f;
 }
 
-#define SHADINGMODELID_UNLIT				0
-#define SHADINGMODELID_DEFAULT_LIT			1
-#define SHADINGMODELID_SUBSURFACE			2
-#define SHADINGMODELID_PREINTEGRATED_SKIN	3
-#define SHADINGMODELID_CLEAR_COAT			4
-#define SHADINGMODELID_SUBSURFACE_PROFILE	5
-#define SHADINGMODELID_TWOSIDED_FOLIAGE		6
-#define SHADINGMODELID_HAIR					7
-#define SHADINGMODELID_CLOTH				8
-#define SHADINGMODELID_EYE					9
-#define SHADINGMODELID_NUM					10
-#define SHADINGMODELID_MASK					0xF					
-
-// The flags are defined so that 0 value has no effect!
-#define SKIP_CUSTOMDATA_MASK			(1 << 4)	// TODO remove. Can be inferred from shading model.
-#define SKIP_PRECSHADOW_MASK			(1 << 5)
-#define ZERO_PRECSHADOW_MASK			(1 << 6)
-#define SKIP_VELOCITY_MASK				(1 << 7)
-
-
 float EncodeShadingModelIdAndSelectiveOutputMask(uint ShadingModelId, uint SelectiveOutputMask)
 {
 	uint Value = (ShadingModelId & SHADINGMODELID_MASK) | SelectiveOutputMask;
@@ -726,39 +707,4 @@ FScreenSpaceData GetScreenSpaceData(float2 UV, bool bGetNormalizedNormal = true)
 	return Out;
 }
 
-// for debugging and to visualize
-float3 GetShadingModelColor(uint ShadingModelID)
-{
-	// TODO: PS4 doesn't optimize out correctly the switch(), so it thinks it needs all the Samplers even if they get compiled out
-	//	This will get fixed after launch per Sony...
-#if PS4_PROFILE
-		 if (ShadingModelID == SHADINGMODELID_UNLIT) return float3(0.1f, 0.1f, 0.2f); // Dark Blue
-	else if (ShadingModelID == SHADINGMODELID_DEFAULT_LIT) return float3(0.1f, 1.0f, 0.1f); // Green
-	else if (ShadingModelID == SHADINGMODELID_SUBSURFACE) return float3(1.0f, 0.1f, 0.1f); // Red
-	else if (ShadingModelID == SHADINGMODELID_PREINTEGRATED_SKIN) return float3(0.6f, 0.4f, 0.1f); // Brown
-	else if (ShadingModelID == SHADINGMODELID_CLEAR_COAT) return float3(0.1f, 0.4f, 0.4f); 
-	else if (ShadingModelID == SHADINGMODELID_SUBSURFACE_PROFILE) return float3(0.2f, 0.6f, 0.5f); // Cyan
-	else if (ShadingModelID == SHADINGMODELID_TWOSIDED_FOLIAGE) return float3(0.2f, 0.2f, 0.8f); // Blue
-	else if (ShadingModelID == SHADINGMODELID_HAIR) return float3(0.6f, 0.1f, 0.5f);
-	else if (ShadingModelID == SHADINGMODELID_CLOTH) return float3(0.7f, 1.0f, 1.0f); 
-	else if (ShadingModelID == SHADINGMODELID_EYE) return float3(0.3f, 1.0f, 1.0f); 
-	else return float3(1.0f, 1.0f, 1.0f); // White
-#else
-	switch(ShadingModelID)
-	{
-		case SHADINGMODELID_UNLIT: return float3(0.1f, 0.1f, 0.2f); // Dark Blue
-		case SHADINGMODELID_DEFAULT_LIT: return float3(0.1f, 1.0f, 0.1f); // Green
-		case SHADINGMODELID_SUBSURFACE: return float3(1.0f, 0.1f, 0.1f); // Red
-		case SHADINGMODELID_PREINTEGRATED_SKIN: return float3(0.6f, 0.4f, 0.1f); // Brown
-		case SHADINGMODELID_CLEAR_COAT: return float3(0.1f, 0.4f, 0.4f); // Brown
-		case SHADINGMODELID_SUBSURFACE_PROFILE: return float3(0.2f, 0.6f, 0.5f); // Cyan
-		case SHADINGMODELID_TWOSIDED_FOLIAGE: return float3(0.2f, 0.2f, 0.8f); // Cyan
-		case SHADINGMODELID_HAIR: return float3(0.6f, 0.1f, 0.5f);
-		case SHADINGMODELID_CLOTH: return float3(0.7f, 1.0f, 1.0f);
-		case SHADINGMODELID_EYE: return float3(0.3f, 1.0f, 1.0f);
-		default: return float3(1.0f, 1.0f, 1.0f); // White
-	}
-#endif
-}
-
 #endif // __DEFERRED_SHADING_COMMON__
diff --git a/Engine/Shaders/Private/RayTracing/RayTracingDebug.usf b/Engine/Shaders/Private/RayTracing/RayTracingDebug.usf
--- a/Engine/Shaders/Private/RayTracing/RayTracingDebug.usf
+++ b/Engine/Shaders/Private/RayTracing/RayTracingDebug.usf
@@ -1,6 +1,7 @@
 // Copyright 1998-2019 Epic Games, Inc. All Rights Reserved.
 
 #include "../Common.ush"
+#include "../ShadingCommon.ush"
 #include "RayTracingDebugDefinitions.ush"
 #include "RayTracingCommon.ush"
 #include "RayTracingHitGroupCommon.ush"
@@ -87,7 +88,7 @@ void RayTracingDebugMainRGS()
 			Result = float4(Payload.Ior, Payload.Ior, Payload.Ior, 1.0f);
 			break;
 		case RAY_TRACING_DEBUG_VIZ_SHADING_MODEL:
-			Result = float4(Payload.ShadingModelID, Payload.ShadingModelID, Payload.ShadingModelID, 1.0f);
+            Result = float4(GetShadingModelColor(Payload.ShadingModelID), 1.0f);
 			break;
 		case RAY_TRACING_DEBUG_VIZ_BLENDING_MODE:
 			Result = float4(Payload.BlendingMode, Payload.BlendingMode, Payload.BlendingMode, 1.0f);
diff --git a/Engine/Shaders/Private/RayTracing/RayTracingMaterialHitShaders.usf b/Engine/Shaders/Private/RayTracing/RayTracingMaterialHitShaders.usf
--- a/Engine/Shaders/Private/RayTracing/RayTracingMaterialHitShaders.usf
+++ b/Engine/Shaders/Private/RayTracing/RayTracingMaterialHitShaders.usf
@@ -22,6 +22,7 @@
 #include "/Engine/Generated/Material.ush"
 #include "/Engine/Generated/VertexFactory.ush"
 
+#include "/Engine/Private/ShadingCommon.ush"
 #include "/Engine/Private/SHCommon.ush"
 #include "/Engine/Private/ReflectionEnvironmentShared.ush"
 #include "/Engine/Private/LightmapCommon.ush"
@@ -284,13 +285,7 @@ void MaterialCHS(inout FPackedMaterialClosestHitPayload PackedPayload, in FDefau
 	Payload.Metallic = Metallic;
 	Payload.Ior = Refraction.x;
 	//Payload.Opacity = Opacity;
-
-// #if MATERIAL_SHADINGMODEL_UNLIT
-	// Payload.ShadingModelID = SHADINGMODELID_UNLIT;
-// #else
-	// Payload.ShadingModelID = SHADINGMODELID_DEFAULT_LIT;
-	// // TODO - support other ShadingModelIDs in path tracer
-// #endif
+    Payload.ShadingModelID = GetShadingModelID();
 
 	// TODO - support different blending modes in path tracer?
 	// uint MaterialBlendingMode = 0;
diff --git a/Engine/Shaders/Private/ShadingCommon.ush b/Engine/Shaders/Private/ShadingCommon.ush
new file mode 100644
--- /dev/null
+++ b/Engine/Shaders/Private/ShadingCommon.ush
@@ -0,0 +1,88 @@
+// Copyright 1998-2019 Epic Games, Inc. All Rights Reserved.
+
+//#pragma once
+#ifndef __SHADING_COMMON__
+#define __SHADING_COMMON__
+
+#define SHADINGMODELID_UNLIT				0
+#define SHADINGMODELID_DEFAULT_LIT			1
+#define SHADINGMODELID_SUBSURFACE			2
+#define SHADINGMODELID_PREINTEGRATED_SKIN	3
+#define SHADINGMODELID_CLEAR_COAT			4
+#define SHADINGMODELID_SUBSURFACE_PROFILE	5
+#define SHADINGMODELID_TWOSIDED_FOLIAGE		6
+#define SHADINGMODELID_HAIR					7
+#define SHADINGMODELID_CLOTH				8
+#define SHADINGMODELID_EYE					9
+#define SHADINGMODELID_NUM					10
+#define SHADINGMODELID_MASK					0xF					
+
+// The flags are defined so that 0 value has no effect!
+#define SKIP_CUSTOMDATA_MASK			(1 << 4)	// TODO remove. Can be inferred from shading model.
+#define SKIP_PRECSHADOW_MASK			(1 << 5)
+#define ZERO_PRECSHADOW_MASK			(1 << 6)
+#define SKIP_VELOCITY_MASK				(1 << 7)
+
+uint GetShadingModelID()
+{        
+#if MATERIAL_SHADINGMODEL_UNLIT
+	return SHADINGMODELID_UNLIT;
+#elif MATERIAL_SHADINGMODEL_DEFAULT_LIT
+	return SHADINGMODELID_DEFAULT_LIT;
+#elif MATERIAL_SHADINGMODEL_SUBSURFACE
+	return SHADINGMODELID_SUBSURFACE;
+#elif MATERIAL_SHADINGMODEL_PREINTEGRATED_SKIN
+	return SHADINGMODELID_PREINTEGRATED_SKIN;
+#elif MATERIAL_SHADINGMODEL_SUBSURFACE_PROFILE
+	return SHADINGMODELID_SUBSURFACE_PROFILE;
+#elif MATERIAL_SHADINGMODEL_CLEAR_COAT
+	return SHADINGMODELID_CLEAR_COAT;
+#elif MATERIAL_SHADINGMODEL_TWOSIDED_FOLIAGE
+	return SHADINGMODELID_TWOSIDED_FOLIAGE;
+#elif MATERIAL_SHADINGMODEL_HAIR
+	return SHADINGMODELID_HAIR;
+#elif MATERIAL_SHADINGMODEL_CLOTH
+	return SHADINGMODELID_CLOTH;
+#elif MATERIAL_SHADINGMODEL_EYE
+	return SHADINGMODELID_EYE;
+#else
+	// missing shading model, compiler should report ShadingModelID is not set
+#endif
+}
+
+// for debugging and to visualize
+float3 GetShadingModelColor(uint ShadingModelID)
+{
+	// TODO: PS4 doesn't optimize out correctly the switch(), so it thinks it needs all the Samplers even if they get compiled out
+	//	This will get fixed after launch per Sony...
+#if PS4_PROFILE
+		 if (ShadingModelID == SHADINGMODELID_UNLIT) return float3(0.1f, 0.1f, 0.2f); // Dark Blue
+	else if (ShadingModelID == SHADINGMODELID_DEFAULT_LIT) return float3(0.1f, 1.0f, 0.1f); // Green
+	else if (ShadingModelID == SHADINGMODELID_SUBSURFACE) return float3(1.0f, 0.1f, 0.1f); // Red
+	else if (ShadingModelID == SHADINGMODELID_PREINTEGRATED_SKIN) return float3(0.6f, 0.4f, 0.1f); // Brown
+	else if (ShadingModelID == SHADINGMODELID_CLEAR_COAT) return float3(0.1f, 0.4f, 0.4f); 
+	else if (ShadingModelID == SHADINGMODELID_SUBSURFACE_PROFILE) return float3(0.2f, 0.6f, 0.5f); // Cyan
+	else if (ShadingModelID == SHADINGMODELID_TWOSIDED_FOLIAGE) return float3(0.2f, 0.2f, 0.8f); // Blue
+	else if (ShadingModelID == SHADINGMODELID_HAIR) return float3(0.6f, 0.1f, 0.5f);
+	else if (ShadingModelID == SHADINGMODELID_CLOTH) return float3(0.7f, 1.0f, 1.0f); 
+	else if (ShadingModelID == SHADINGMODELID_EYE) return float3(0.3f, 1.0f, 1.0f); 
+	else return float3(1.0f, 1.0f, 1.0f); // White
+#else
+	switch(ShadingModelID)
+	{
+		case SHADINGMODELID_UNLIT: return float3(0.1f, 0.1f, 0.2f); // Dark Blue
+		case SHADINGMODELID_DEFAULT_LIT: return float3(0.1f, 1.0f, 0.1f); // Green
+		case SHADINGMODELID_SUBSURFACE: return float3(1.0f, 0.1f, 0.1f); // Red
+		case SHADINGMODELID_PREINTEGRATED_SKIN: return float3(0.6f, 0.4f, 0.1f); // Brown
+		case SHADINGMODELID_CLEAR_COAT: return float3(0.1f, 0.4f, 0.4f); // Brown
+		case SHADINGMODELID_SUBSURFACE_PROFILE: return float3(0.2f, 0.6f, 0.5f); // Cyan
+		case SHADINGMODELID_TWOSIDED_FOLIAGE: return float3(0.2f, 0.2f, 0.8f); // Cyan
+		case SHADINGMODELID_HAIR: return float3(0.6f, 0.1f, 0.5f);
+		case SHADINGMODELID_CLOTH: return float3(0.7f, 1.0f, 1.0f);
+		case SHADINGMODELID_EYE: return float3(0.3f, 1.0f, 1.0f);
+		default: return float3(1.0f, 1.0f, 1.0f); // White
+	}
+#endif
+}
+
+#endif // __SHADING_COMMON__
