Transfer Gbuffer custom data to the RT payload (to be used for advanced material like coat).

diff --git a/Engine/Shaders/Private/RayTracing/RayTracingCommon.ush b/Engine/Shaders/Private/RayTracing/RayTracingCommon.ush
--- a/Engine/Shaders/Private/RayTracing/RayTracingCommon.ush
+++ b/Engine/Shaders/Private/RayTracing/RayTracingCommon.ush
@@ -82,7 +82,7 @@ struct FMaterialClosestHitPayload
 	uint ShadingModelID;				//96	1byte	4bits
 	uint BlendingMode;					//100	0bytes	4bits (packed)
 	uint PrimitiveLightingChannelMask;	//104	0bytes	(removed)
-	float4 CustomData;					//120	0bytes	(removed)
+	float4 CustomData;					//120	4bytes	32bits	/* only xy are packed at the moment */
 	float GBufferAO;					//124	0bytes	(removed)
 	float3 IndirectIrradiance;			//136	0bytes	48bits	/* gbuffer only has float payload and there are truncation HLSL warnings */
 
@@ -107,6 +107,7 @@ struct FPackedMaterialClosestHitPayload
 	int MetallicAndSpecularAndRoughness;
 	int IorAndShadingModelIDAndBlendingModeAndHitKind;
 	int IndirectIrradiance[2];
+	int CustomData;
 };
 
 FPackedMaterialClosestHitPayload PackRayTracingPayload(FMaterialClosestHitPayload Input, in FRayCone RayCone)
@@ -134,6 +135,7 @@ FPackedMaterialClosestHitPayload PackRayTracingPayload(FMaterialClosestHitPayloa
 	Output.IndirectIrradiance[0] = f32tof16(Input.IndirectIrradiance.x);
 	Output.IndirectIrradiance[0] |= f32tof16(Input.IndirectIrradiance.y) << 16;
 	Output.IndirectIrradiance[1] = f32tof16(Input.IndirectIrradiance.z);
+	Output.CustomData = f32tof16(Input.CustomData.x) | (f32tof16(Input.CustomData.y) << 16);
 
 	return Output;
 }
@@ -163,6 +165,9 @@ FMaterialClosestHitPayload UnpackRayTracingPayload(FPackedMaterialClosestHitPayl
 	Output.IndirectIrradiance.x	= f16tof32(Input.IndirectIrradiance[0]);
 	Output.IndirectIrradiance.y	= f16tof32(Input.IndirectIrradiance[0] >> 16);
 	Output.IndirectIrradiance.z	= f16tof32(Input.IndirectIrradiance[1]);
+	Output.CustomData.x = f16tof32(Input.CustomData.x);
+	Output.CustomData.y = f16tof32(Input.CustomData.y >> 16);
+	Output.CustomData.zw = 0;
 
 	Output.DiffuseColor  = Output.BaseColor - Output.BaseColor * Output.Metallic;
 	Output.SpecularColor = lerp(0.08 * Output.Specular.xxx, Output.BaseColor, Output.Metallic);
diff --git a/Engine/Shaders/Private/RayTracing/RayTracingMaterialHitShaders.usf b/Engine/Shaders/Private/RayTracing/RayTracingMaterialHitShaders.usf
--- a/Engine/Shaders/Private/RayTracing/RayTracingMaterialHitShaders.usf
+++ b/Engine/Shaders/Private/RayTracing/RayTracingMaterialHitShaders.usf
@@ -284,6 +284,7 @@ void MaterialCHS(inout FPackedMaterialClosestHitPayload PackedPayload, in FDefau
 	Payload.Roughness = Roughness;
 	Payload.Metallic = Metallic;
 	Payload.Ior = Refraction.x;
+    Payload.CustomData = float4(GetMaterialCustomData0(MaterialParameters), GetMaterialCustomData1(MaterialParameters), 0, 0);
 	//Payload.Opacity = Opacity;
     Payload.ShadingModelID = GetShadingModelID();
 
diff --git a/Engine/Source/Runtime/Renderer/Private/DeferredShadingRenderer.cpp b/Engine/Source/Runtime/Renderer/Private/DeferredShadingRenderer.cpp
--- a/Engine/Source/Runtime/Renderer/Private/DeferredShadingRenderer.cpp
+++ b/Engine/Source/Runtime/Renderer/Private/DeferredShadingRenderer.cpp
@@ -910,7 +910,7 @@ FRHIRayTracingPipelineState* FDeferredShadingSceneRenderer::BindRayTracingPipeli
 
 	FRayTracingPipelineStateInitializer Initializer;
 
-	Initializer.MaxPayloadSizeInBytes = 152; // #dxr_todo: set Initializer.MaxPayloadSizeInBytes based on shader requirements (it's not obvious how to compute this right now)
+	Initializer.MaxPayloadSizeInBytes = 156; // #dxr_todo: set Initializer.MaxPayloadSizeInBytes based on shader requirements (it's not obvious how to compute this right now)
 
 	FRayTracingShaderRHIParamRef RayGenShaderTable[] = { RayGenShader };
 	Initializer.SetRayGenShaderTable(RayGenShaderTable);
