Added height fog sampling for reflection.
This is only done for valid samples. Otherwise, fallback will be on reflection volumes in the ReflectionAndSky pass.
Tested both maxRayDistance ,0 and >=0 on Woodchuck giving concistent result. (dev-renbder has a bug and I ahve notified Guillaume about it)
Also simplified the code (removing some unused variable/code)

diff --git a/Engine/Shaders/Private/RayTracing/RayTracingReflections.usf b/Engine/Shaders/Private/RayTracing/RayTracingReflections.usf
--- a/Engine/Shaders/Private/RayTracing/RayTracingReflections.usf
+++ b/Engine/Shaders/Private/RayTracing/RayTracingReflections.usf
@@ -26,10 +26,12 @@ float2 SobolIndexToUniformUnitSquare(uint2 SobolRand)
 #include "../ReflectionEnvironmentShared.ush"
 #include "../Montecarlo.ush"
 #include "../PathTracing/Utilities/PathTracingRandomSequence.ush" // #dxr-todo: move to a common folder
+#include "../HeightFogCommon.ush"
 #include "RayTracingCommon.ush"
 #include "RayTracingHitGroupCommon.ush"
 
 uint SamplesPerPixel;
+uint HeightFog;
 uint ShouldDoDirectLighting;
 uint ShouldDoReflectedShadows;
 uint ShouldDoEmissiveAndIndirectLighting;
@@ -416,38 +418,47 @@ void RayTracingReflectionsRGS()
 				float3 Radiance = float3(0, 0, 0);
 				FMaterialClosestHitPayload Payload = TraceRayAndAccumulateResults(Ray, RayCone, Radiance);
 
-				float4 PathRadiance;
+				float4 PathRadiance = 0.0f;
+				bool bApplyHeightFog = false;
+				float3 OriginToCollider = 0.0f;
+
 				if (Payload.HitT >= 0)
 				{
-					PathRadiance = float4(Radiance, 1.0); 
-					if (BounceIndex == 0)
-					{
-						ReflectedColor += PathRadiance;
-					}
-					else
-					{
-						float4 FinalReflectionColor = PathRadiance;
-						//float4 FinalReflectionColor = BlendWithEnvironment(BRDFThroughput, Payload, Ray.Origin, PathRadiance, Ray.Direction, Payload.WorldNormal);
-						ReflectedColor += FinalReflectionColor;
-					}
+					PathRadiance = float4(Radiance, 1.0);
 					HitDistance = Payload.HitT;
+
+					bApplyHeightFog = HeightFog > 0;
+					OriginToCollider = Payload.WorldPos - Ray.Origin;
 				}
 				else if (bAllowSkySampling)
 				{
 					if (ReflectionStruct.SkyLightParameters.y > 0)
 					{
 						float SkyAverageBrightness = 1.0f;
-
 						float3 SkyLighting = GetSkyLightReflection(Ray.Direction, Roughness, SkyAverageBrightness);
+						PathRadiance = float4(SkyLighting, 1.0);
 
-						ReflectedColor += float4(SkyLighting, 1.0f);
+						bApplyHeightFog = HeightFog > 0;
+						const float SkyFarDistance = 100000.0f;
+						OriginToCollider = Ray.Direction*SkyFarDistance;
 						
 						// To tell the denoiser the black input is valid even if no valid ray trace result was found. Uses a high ray hit distance
 						// so that the sample may be reused by another neighbor pixel in the denoise as if RoughnessFade in ]0;1[.
 						HitDistance = 1.0e20;
 					}
 				}
-				
+
+				// If there is a contribution on this path, apply the required fog on the sample. 
+				// Also accumulate the path contribution.
+				// In case no contribution is accumulated, the ReflectionEnvironmentAndSky pass will fill it up using reflection volumes.
+				//#dxr - todo: when MaxBounces>1, we will need to take into account a path throughput affected from the fog trnasmittance.
+				if (bApplyHeightFog)
+				{
+					float4 HeightFogInscatteringAndTransmittance = CalculateHeightFog(OriginToCollider);
+					PathRadiance.rgb = PathRadiance.rgb * HeightFogInscatteringAndTransmittance.a + HeightFogInscatteringAndTransmittance.rgb;
+
+				}
+				ReflectedColor += PathRadiance;
 				ReflectedWeight += 1.0;
 				IncidentDirection = Ray.Direction;
 			}
diff --git a/Engine/Source/Runtime/Renderer/Private/DeferredShadingRenderer.h b/Engine/Source/Runtime/Renderer/Private/DeferredShadingRenderer.h
--- a/Engine/Source/Runtime/Renderer/Private/DeferredShadingRenderer.h
+++ b/Engine/Source/Runtime/Renderer/Private/DeferredShadingRenderer.h
@@ -437,6 +437,7 @@ private:
 		FRDGTextureRef* OutColorTexture,
 		FRDGTextureRef* OutRayHitDistanceTexture,
 		int32 SamplePerPixel,
+		int32 HeightFog,
 		float ResolutionFraction);
 
 #if RHI_RAYTRACING
diff --git a/Engine/Source/Runtime/Renderer/Private/RayTracing/RayTracingReflections.cpp b/Engine/Source/Runtime/Renderer/Private/RayTracing/RayTracingReflections.cpp
--- a/Engine/Source/Runtime/Renderer/Private/RayTracing/RayTracingReflections.cpp
+++ b/Engine/Source/Runtime/Renderer/Private/RayTracing/RayTracingReflections.cpp
@@ -129,6 +129,7 @@ class FRayTracingReflectionsRG : public FGlobalShader
 
 	BEGIN_SHADER_PARAMETER_STRUCT(FParameters, )
 		SHADER_PARAMETER(int32, SamplesPerPixel)
+		SHADER_PARAMETER(int32, HeightFog)
 		SHADER_PARAMETER(int32, ShouldDoDirectLighting)
 		SHADER_PARAMETER(int32, ShouldDoReflectedShadows)
 		SHADER_PARAMETER(int32, ShouldDoEmissiveAndIndirectLighting)
@@ -148,6 +149,7 @@ class FRayTracingReflectionsRG : public FGlobalShader
 		SHADER_PARAMETER_STRUCT_REF(FSceneTexturesUniformParameters, SceneTexturesStruct)
 		SHADER_PARAMETER_STRUCT_REF(FReflectionsLightData, LightData)
 		SHADER_PARAMETER_STRUCT_REF(FReflectionUniformParameters, ReflectionStruct)
+		SHADER_PARAMETER_STRUCT_REF(FFogUniformParameters, FogUniformParameters)
 
 		SHADER_PARAMETER_RDG_TEXTURE_UAV(RWTexture2D<float4>, ColorOutput)
 		SHADER_PARAMETER_RDG_TEXTURE_UAV(RWTexture2D<float>, RayHitDistanceOutput)
@@ -197,7 +199,8 @@ void FDeferredShadingSceneRenderer::RayTraceReflections(
 	const FViewInfo& View,
 	FRDGTextureRef* OutColorTexture,
 	FRDGTextureRef* OutRayHitDistanceTexture,
-	int32 SamplePerPixel, 
+	int32 SamplePerPixel,
+	int32 HeightFog,
 	float ResolutionFraction)
 #if RHI_RAYTRACING
 {
@@ -223,6 +226,7 @@ void FDeferredShadingSceneRenderer::RayTraceReflections(
 	FRayTracingReflectionsRG::FParameters* PassParameters = GraphBuilder.AllocParameters<FRayTracingReflectionsRG::FParameters>();
 
 	PassParameters->SamplesPerPixel = SamplePerPixel;
+	PassParameters->HeightFog = HeightFog;
 	PassParameters->ShouldDoDirectLighting = GRayTracingReflectionsDirectLighting;
 	PassParameters->ShouldDoReflectedShadows = GRayTracingReflectionsShadows;
 	PassParameters->ShouldDoEmissiveAndIndirectLighting = GRayTracingReflectionsEmissiveAndIndirectLighting;
@@ -252,6 +256,11 @@ void FDeferredShadingSceneRenderer::RayTraceReflections(
 		SetupReflectionUniformParameters(View, ReflectionStruct);
 		PassParameters->ReflectionStruct = CreateUniformBufferImmediate(ReflectionStruct, EUniformBufferUsage::UniformBuffer_SingleDraw);
 	}
+	{
+		FFogUniformParameters FogStruct;
+		SetupFogUniformParameters(View, FogStruct);
+		PassParameters->FogUniformParameters = CreateUniformBufferImmediate(FogStruct, EUniformBufferUsage::UniformBuffer_SingleDraw);
+	}
 	PassParameters->ColorOutput = GraphBuilder.CreateUAV(*OutColorTexture);
 	PassParameters->RayHitDistanceOutput = GraphBuilder.CreateUAV(*OutRayHitDistanceTexture);
 	
diff --git a/Engine/Source/Runtime/Renderer/Private/ReflectionEnvironment.cpp b/Engine/Source/Runtime/Renderer/Private/ReflectionEnvironment.cpp
--- a/Engine/Source/Runtime/Renderer/Private/ReflectionEnvironment.cpp
+++ b/Engine/Source/Runtime/Renderer/Private/ReflectionEnvironment.cpp
@@ -121,6 +121,12 @@ static FAutoConsoleVariableRef CVarRayTracingReflectionsSamplesPerPixel(
 	GRayTracingReflectionsSamplesPerPixel,
 	TEXT("Sets the samples-per-pixel for reflections (default = 1)"));
 
+static int32 GRayTracingReflectionsHeightFog = 1;
+static FAutoConsoleVariableRef CVarRayTracingReflectionsHeightFog(
+	TEXT("r.RayTracing.Reflections.HeightFog"),
+	GRayTracingReflectionsHeightFog,
+	TEXT("Enables height fog in ray traced reflections (default = 1)"));
+
 static TAutoConsoleVariable<int32> CVarUseReflectionDenoiser(
 	TEXT("r.Reflections.Denoiser"),
 	2,
@@ -767,7 +773,7 @@ void FDeferredShadingSceneRenderer::RenderDeferredReflectionsAndSkyLighting(FRHI
 			RayTraceReflections(
 				GraphBuilder,
 				View, &DenoiserInputs.Color, &DenoiserInputs.RayHitDistance,
-				GRayTracingReflectionsSamplesPerPixel, RayTracingConfig.ResolutionFraction);
+				GRayTracingReflectionsSamplesPerPixel, GRayTracingReflectionsHeightFog, RayTracingConfig.ResolutionFraction);
 
 
 			// Denoise the reflections.
