Fixes merge collision betwen bilateral refactor and RTGI+AO in the denoiser.

diff --git a/Engine/Shaders/Private/ScreenSpaceDenoise/SSDSpatialAccumulation.usf b/Engine/Shaders/Private/ScreenSpaceDenoise/SSDSpatialAccumulation.usf
--- a/Engine/Shaders/Private/ScreenSpaceDenoise/SSDSpatialAccumulation.usf
+++ b/Engine/Shaders/Private/ScreenSpaceDenoise/SSDSpatialAccumulation.usf
@@ -110,12 +110,15 @@
 	#endif
 
 #elif CONFIG_SIGNAL_PROCESSING == SIGNAL_PROCESSING_GI_AND_AO
+	#define CONFIG_VIEW_POSITION_BASED_BILATERAL 1
+	#define CONFIG_NORMAL_BASED_BILATERAL 1
+
 	#if DIM_STAGE == STAGE_RECONSTRUCTION
 		// Input and output layout.
 		#define CONFIG_SIGNAL_INPUT_LAYOUT  SIGNAL_BUFFER_LAYOUT_GI_AND_AO_INPUT
 		#define CONFIG_SIGNAL_OUTPUT_LAYOUT SIGNAL_BUFFER_LAYOUT_GI_AND_AO_HISTORY
 
-	#define CONFIG_SAMPLE_SET           SAMPLE_SET_STACKOWIAK_4_SETS
+		#define CONFIG_SAMPLE_SET           SAMPLE_SET_STACKOWIAK_4_SETS
 
 	#elif DIM_STAGE == STAGE_POST_FILTERING
 		// Input and output layout.
diff --git a/Engine/Source/Runtime/Renderer/Private/SceneRendering.h b/Engine/Source/Runtime/Renderer/Private/SceneRendering.h
--- a/Engine/Source/Runtime/Renderer/Private/SceneRendering.h
+++ b/Engine/Source/Runtime/Renderer/Private/SceneRendering.h
@@ -760,6 +760,7 @@ struct FPreviousViewInfo
 		CustomSSRInput.SafeRelease();
 		ReflectionsHistory.SafeRelease();
 		AmbientOcclusionHistory.SafeRelease();
+		GlobalIlluminationHistory.SafeRelease();
 		ShadowHistories.Reset();
 	}
 };
