DR - Unblock hlslcc shader compilation (Vulkan, Metal)
#rnx

diff --git a/Engine/Shaders/Private/ShadingCommon.ush b/Engine/Shaders/Private/ShadingCommon.ush
--- a/Engine/Shaders/Private/ShadingCommon.ush
+++ b/Engine/Shaders/Private/ShadingCommon.ush
@@ -46,7 +46,12 @@ uint GetShadingModelID()
 #elif MATERIAL_SHADINGMODEL_EYE
 	return SHADINGMODELID_EYE;
 #else
+	#if COMPILER_HLSLCC
+		// hlslcc needs a return type
+		return SHADINGMODELID_MASK;
+	#else
 	// missing shading model, compiler should report ShadingModelID is not set
+	#endif
 #endif
 }
 
