Fixes output of the reflection ray generation shader for a ray that did not find anything within ray's T-max distance.

diff --git a/Engine/Shaders/Private/RayTracing/RayTracingReflections.usf b/Engine/Shaders/Private/RayTracing/RayTracingReflections.usf
--- a/Engine/Shaders/Private/RayTracing/RayTracingReflections.usf
+++ b/Engine/Shaders/Private/RayTracing/RayTracingReflections.usf
@@ -430,22 +430,24 @@ void RayTracingReflectionsRGS()
 					bApplyHeightFog = HeightFog > 0;
 					OriginToCollider = Payload.WorldPos - Ray.Origin;
 				}
-				else if (bAllowSkySampling)
+				else if (bAllowSkySampling && ReflectionStruct.SkyLightParameters.y > 0)
 				{
-					if (ReflectionStruct.SkyLightParameters.y > 0)
-					{
-						float SkyAverageBrightness = 1.0f;
-						float3 SkyLighting = GetSkyLightReflection(Ray.Direction, Roughness, SkyAverageBrightness);
-						PathRadiance = float4(SkyLighting, 1.0);
-
-						bApplyHeightFog = HeightFog > 0;
-						const float SkyFarDistance = 100000.0f;
-						OriginToCollider = Ray.Direction*SkyFarDistance;
+					float SkyAverageBrightness = 1.0f;
+					float3 SkyLighting = GetSkyLightReflection(Ray.Direction, Roughness, SkyAverageBrightness);
+					PathRadiance = float4(SkyLighting, 1.0);
+
+					bApplyHeightFog = HeightFog > 0;
+					const float SkyFarDistance = 100000.0f;
+					OriginToCollider = Ray.Direction*SkyFarDistance;
 						
-						// To tell the denoiser the black input is valid even if no valid ray trace result was found. Uses a high ray hit distance
-						// so that the sample may be reused by another neighbor pixel in the denoise as if RoughnessFade in ]0;1[.
-						HitDistance = 1.0e20;
-					}
+					// To tell the denoiser the black input is valid even if no valid ray trace result was found. Uses a high ray hit distance
+					// so that the sample may be reused by another neighbor pixel in the denoise as if RoughnessFade in ]0;1[.
+					HitDistance = 1.0e20;
+				}
+				else
+				{
+					// Tell to the denoiser this sample reached outer bounds of the scene without any intersection.
+					HitDistance = -1;
 				}
 
 				// If there is a contribution on this path, apply the required fog on the sample. 
