Turns on karis color weighting for reflection reconstruction pass to reduce the flickering of the reconstruction kernel at high roughness.

diff --git a/Engine/Shaders/Private/ScreenSpaceDenoise/SSDSignalFramework.ush b/Engine/Shaders/Private/ScreenSpaceDenoise/SSDSignalFramework.ush
--- a/Engine/Shaders/Private/ScreenSpaceDenoise/SSDSignalFramework.ush
+++ b/Engine/Shaders/Private/ScreenSpaceDenoise/SSDSignalFramework.ush
@@ -86,10 +86,22 @@ float Luma_To_LumaLog(float x)
 	return max(x > 0 ? log2(x) : -100, -100);
 }
 
+/** Returns the exposure multiplier that should be used for the color transformation of the signal. */
+float GetSignalColorTransformationExposure()
+{
+	// TODO: expose that guy.
+	const float Tweak = 1;
+
+	float FrameExposureScale = EyeAdaptationLookup();
+
+	return FrameExposureScale * Tweak;
+}
+
+/** Returns the linear color space Luma with a 4 multiplier for ALU perf reasons. */
 float GetSignalLuma4(FSSFSignalSample Sample, const uint Basis)
 {
 	// TODO: should force this to be scalar load.
-	float FrameExposureScale = EyeAdaptationLookup();
+	float FrameExposureScale = GetSignalColorTransformationExposure();
 		
 	// TODO: exposes optimisation if sample is normalized.
 	const bool bIsNormalizedSample = false;
@@ -129,10 +141,8 @@ FSSFSignalSample TransformSignal(FSSFSignalSample Sample, const uint SrcBasis, c
 		return Sample;
 	}
 
-	const float Tweak = 1;
-
 	// TODO: should force this to be scalar load.
-	float FrameExposureScale = EyeAdaptationLookup() * Tweak;
+	float FrameExposureScale = GetSignalColorTransformationExposure();
 		
 	// TODO: exposes optimisation if sample is normalized.
 	const bool bIsNormalizedSample = false;
diff --git a/Engine/Shaders/Private/ScreenSpaceDenoise/SSDSpatialAccumulation.usf b/Engine/Shaders/Private/ScreenSpaceDenoise/SSDSpatialAccumulation.usf
--- a/Engine/Shaders/Private/ScreenSpaceDenoise/SSDSpatialAccumulation.usf
+++ b/Engine/Shaders/Private/ScreenSpaceDenoise/SSDSpatialAccumulation.usf
@@ -65,7 +65,8 @@
 		//#define CONFIG_SAMPLE_COUNT_POLICY  SAMPLE_COUNT_POLICY_DISABLED
 
 		// Do color accumulation with karis weighting to avoid flickering specular highlight to show up the kernel pattern.
-		//#define CONFIG_COLOR_SPACE          (COLOR_SPACE_RGB | COLOR_SPACE_KARIS_WEIGHTING)
+		// TODO: This is a bit agressive.
+		#define CONFIG_COLOR_SPACE          (COLOR_SPACE_RGB | COLOR_SPACE_KARIS_WEIGHTING)
 
 		#define CONFIG_REJECT_HIGHEST_COLOR 0
 
