Fixed some issues in ray traced translucency pass: incorrect blending mode, wrong RDG name, lack of FinishingRenderingSceneColor call.
Implemented barebones ray traced translucency RGS.

diff --git a/Engine/Shaders/Private/RayTracing/RayTracingTranslucency.usf b/Engine/Shaders/Private/RayTracing/RayTracingTranslucency.usf
--- a/Engine/Shaders/Private/RayTracing/RayTracingTranslucency.usf
+++ b/Engine/Shaders/Private/RayTracing/RayTracingTranslucency.usf
@@ -24,15 +24,16 @@ float2 SobolIndexToUniformUnitSquare(uint2 SobolRand)
 #include "RayTracingHitGroupCommon.ush"
 
 uint SamplesPerPixel;
+uint MaxRefractionRays;
 uint HeightFog;
 uint ShouldDoDirectLighting;
 uint ShouldDoReflectedShadows;
 uint ShouldDoEmissiveAndIndirectLighting;
 uint UpscaleFactor;
 
-float ReflectionMinRayDistance;
-float ReflectionMaxRayDistance;
-float ReflectionMaxRoughness;
+float TranslucencyMinRayDistance;
+float TranslucencyMaxRayDistance;
+float TranslucencyMaxRoughness;
 
 RaytracingAccelerationStructure TLAS;
 
@@ -42,6 +43,17 @@ RWTexture2D<float> RayHitDistanceOutput;
 #include "RayTracingLightingCommon.ush"
 
 // #dxr-todo: move to RayTracingCommon.ush
+float GetHitT(FMaterialClosestHitPayload HitInfo)
+{
+	return HitInfo.HitT;
+}
+
+bool IsHit(RayDesc Ray, FMaterialClosestHitPayload HitInfo)
+{
+	float t = HitInfo.HitT;
+	return t > Ray.TMin && t < Ray.TMax;
+}
+
 float3 ReconstructWorldPositionFromDepth(float2 UV, float Depth)
 {
 	float2 ScreenPosition = (UV - View.ScreenPositionScaleBias.wz) / View.ScreenPositionScaleBias.xy;
@@ -62,6 +74,23 @@ uint2 GetPixelCoord(uint2 DispatchThreadId)
 }
 
 
+bool RefractRay(float3 V, float3 N, float IorV, float IorL, out float3 L)
+{
+	float NoV = dot(V, N);
+	float IorRatio = IorV / IorL;
+	float d = 1.0 - IorRatio * IorRatio * (1.0 - NoV * NoV);
+
+	if (d > 0.0)
+	{
+		L = N * (IorRatio*NoV - sqrt(d)) - V * IorRatio;
+		return true;
+	}
+
+	// TIR, force to reflect
+	L = N * (NoV + NoV) - V;
+	return false;
+}
+
 [shader("raygeneration")]
 void RayTracingTranslucencyRGS()
 {
@@ -80,26 +109,91 @@ void RayTracingTranslucencyRGS()
 	float Depth = ScreenSpaceData.GBuffer.Depth;
 	float3 CameraOrigin = ReconstructWorldPositionFromDepth(UV, 0.0f);
 	float3 WorldPosition = ReconstructWorldPositionFromDepth(UV, Depth);
-	float3 CameraDirection = normalize(WorldPosition - CameraOrigin);
-	//#dxr - todo: use also the shading normal (interpolate, or apply a more sophisticated approach)
-	float3 WorldNormal = ScreenSpaceData.GBuffer.WorldNormal;
-	float Roughness = ScreenSpaceData.GBuffer.Roughness;
-	
-	float3 BRDFThroughput = float3(1.0f, 1.0f, 1.0f);
+	float3 CameraDirection = normalize(WorldPosition - CameraOrigin);	
+	const float GBufferHitT = length(WorldPosition - CameraOrigin);
+
+	float3 TranslucencyColor = 0.0;
+	float HitDistance = 0.0f;
+
+	// Trace rays from camera origin to (Gbuffer - epsilon) to only intersect translucent objects
+	//#dxr_todo: a potential big optimization would be to include translucent objects in the GBuffer
+	// and use the GBuffer shading ID to early return for non translucent objects 
+	RayDesc Ray;
+	Ray.Origin = CameraOrigin;
+	Ray.TMin = 0.0;
+	Ray.TMax = GBufferHitT - 0.1;
+	Ray.Direction = CameraDirection;
+
+	FRayCone RayCone = (FRayCone)0;
+	RayCone.SpreadAngle = View.EyeToPixelSpreadAngle;
 
-	uint MaxBounces = 1; //#dxr-todo: only works with 1 bounce until hit shaders are supported
-	uint LocalSamplesPerPixel = SamplesPerPixel;
+	uint RayFlags = RAY_FLAG_CULL_BACK_FACING_TRIANGLES;
 
-	// Mask out depth values at or beyond far plane
-	if (Depth >= SCENE_TEXTURES_DISABLED_SCENE_DEPTH_VALUE)
+	for (uint RefractionRayIndex = 0; RefractionRayIndex < MaxRefractionRays; ++RefractionRayIndex)
 	{
-		LocalSamplesPerPixel = 0;
+		FMaterialClosestHitPayload Payload = TraceRayInternal(
+			TLAS,   // AccelerationStructure
+			RayFlags,
+			0xFF,   // InstanceInclusionMask
+			0,      // RayContributionToHitGroupIndex
+			1,      // MultiplierForGeometryContributionToShaderIndex #dxr_todo: we currently only support shader binding per mesh instance
+			0,      // MissShaderIndex
+			Ray,    // RayDesc
+			RayCone
+		);
+
+		if (!IsHit(Ray, Payload))
+		{
+			//TranslucencyColor = float3(1, 0, 0);
+			break;
+		}
+
+		//#dxr_todo: 
+		//- set stack of refraction & reflection rays
+		//- handle TIR
+		//- Handle translucency throughput depending on color and opacity (beer´s law)
+
+		float3 HitPoint = Ray.Origin + Ray.Direction * GetHitT(Payload);
+
+		// Trace reflection ray 
+		uint DummyVariable;
+		float2 RandSample = RandomSequence_GenerateSample2D(RandSequence, DummyVariable);
+
+		RayDesc ReflectionRay;
+		ReflectionRay.Origin = HitPoint;
+		ReflectionRay.TMin = 0.01;
+		ReflectionRay.TMax = TranslucencyMaxRayDistance;
+		ReflectionRay.Direction = GenerateReflectedRayDirection(Ray.Direction, Payload.WorldNormal, Payload.Roughness, RandSample);
+
+		float3 ReflectionRadiance = float3(0, 0, 0);
+		FMaterialClosestHitPayload ReflectionPayload = TraceRayAndAccumulateResults(ReflectionRay, RayCone, ReflectionRadiance);
+		TranslucencyColor += ReflectionRadiance;
+
+		// Set refraction ray for next iteration
+		float IorV = 1.0;
+		float IorL = Payload.Ior;
+		if (Payload.HitKind != HIT_KIND_TRIANGLE_FRONT_FACE)
+		{
+			IorV = Payload.Ior;
+			IorL = 1.0;
+		}
+
+		float3 RefractedDirection;
+		if (!RefractRay(Ray.Direction, Payload.WorldNormal, IorV, IorL, RefractedDirection))
+		{
+			// TIR
+			break;
+		}
+
+		Ray.Origin = HitPoint;
+		Ray.TMin = 0.001;
+		Ray.TMax = TranslucencyMaxRayDistance;
+		Ray.Direction = RefractedDirection;
+		float SurfaceCurvature = 0.0f; /* #todo_dxr assume no curvature */
+		RayCone = PropagateRayCone(RayCone, SurfaceCurvature, Depth);
 	}
 
-	float4 TranslucencyColor = 0.0;
-	float HitDistance = 0.0f;
-
-	ColorOutput[DispatchThreadId] = float4(TranslucencyColor);
+	ColorOutput[DispatchThreadId] = float4(TranslucencyColor, 1.0); //#dxr_todo: set alpha value depending on hit
 	RayHitDistanceOutput[DispatchThreadId] = HitDistance;
 }
 
diff --git a/Engine/Source/Runtime/Renderer/Private/RayTracing/RayTracingTranslucency.cpp b/Engine/Source/Runtime/Renderer/Private/RayTracing/RayTracingTranslucency.cpp
--- a/Engine/Source/Runtime/Renderer/Private/RayTracing/RayTracingTranslucency.cpp
+++ b/Engine/Source/Runtime/Renderer/Private/RayTracing/RayTracingTranslucency.cpp
@@ -21,6 +21,13 @@
 #include "PostProcess/SceneFilterRendering.h"
 #include "PipelineStateCache.h"
 
+
+static int32 GRayTracingTranslucencyMaxRefractionRays = 3;
+static FAutoConsoleVariableRef CVarRayTracingTranslucencyMaxRefractionRays(
+	TEXT("r.RayTracing.Translucency.MaxRefractionRays"),
+	GRayTracingTranslucencyMaxRefractionRays,
+	TEXT("Sets the maximum number of refraction rays for ray traced translucency (default = 3)"));
+
 static int32 GRayTracingTranslucencyEmissiveAndIndirectLighting = 1;
 static FAutoConsoleVariableRef CVarRayTracingTranslucencyEmissiveAndIndirectLighting(
 	TEXT("r.RayTracing.Translucency.EmissiveAndIndirectLighting"),
@@ -71,6 +78,7 @@ static FAutoConsoleVariableRef CVarRayTracingTranslucencyHeightFog(
 DECLARE_GPU_STAT_NAMED(RayTracingTranslucency, TEXT("Ray Tracing Translucency"));
 
 
+//#dxr_todo: factor out with the light structure in RayTracingReflections.cpp
 static const int32 GTranslucencyLightCountMaximum = 64;
 
 BEGIN_GLOBAL_SHADER_PARAMETER_STRUCT(FTranslucencyLightData, )
@@ -92,7 +100,7 @@ SHADER_PARAMETER_TEXTURE(Texture2D, DummyRectLightTexture) //#dxr_todo: replace
 END_GLOBAL_SHADER_PARAMETER_STRUCT()
 
 
-IMPLEMENT_GLOBAL_SHADER_PARAMETER_STRUCT(FTranslucencyLightData, "TranslucencyLightsData");
+IMPLEMENT_GLOBAL_SHADER_PARAMETER_STRUCT(FTranslucencyLightData, "ReflectionLightsData");
 
 
 void SetupTranslucencyLightData(
@@ -149,6 +157,7 @@ class FRayTracingTranslucencyRG : public FGlobalShader
 
 	BEGIN_SHADER_PARAMETER_STRUCT(FParameters, )
 		SHADER_PARAMETER(int32, SamplesPerPixel)
+		SHADER_PARAMETER(int32, MaxRefractionRays)
 		SHADER_PARAMETER(int32, HeightFog)
 		SHADER_PARAMETER(int32, ShouldDoDirectLighting)
 		SHADER_PARAMETER(int32, ShouldDoReflectedShadows)
@@ -280,6 +289,14 @@ IMPLEMENT_SHADER_TYPE(, FCompositeTranslucencyPS, TEXT("/Engine/Private/RayTraci
 
 void FDeferredShadingSceneRenderer::RayTraceTranslucency(FRHICommandListImmediate& RHICmdList)
 {
+	//#dxr_todo: check TPT_StandardTranslucency/TPT_TranslucencyAfterDOF support
+	ETranslucencyPass::Type TranslucencyPass(ETranslucencyPass::TPT_AllTranslucency);
+
+	if (!ShouldRenderTranslucency(TranslucencyPass))
+	{
+		return; // Early exit if nothing needs to be done.
+	}
+
 	for (int32 ViewIndex = 0, Num = Views.Num(); ViewIndex < Num; ViewIndex++)
 	{
 		FViewInfo& View = Views[ViewIndex];
@@ -320,11 +337,10 @@ void FDeferredShadingSceneRenderer::RayTraceTranslucency(FRHICommandListImmediat
 			TShaderMapRef<FCompositeTranslucencyPS> PixelShader(ShaderMap);
 
 			FGraphicsPipelineStateInitializer GraphicsPSOInit;
-			SceneContext.BeginRenderingSceneColor(RHICmdList, ESimpleRenderTargetMode::EExistingColorAndDepth, FExclusiveDepthStencil::DepthNop_StencilNop, true);
+			SceneContext.BeginRenderingSceneColor(RHICmdList, ESimpleRenderTargetMode::EExistingColorAndDepth, FExclusiveDepthStencil::DepthRead_StencilWrite, true);
 			RHICmdList.ApplyCachedRenderTargets(GraphicsPSOInit);
 
-			//#dxr_todo review render mode
-			GraphicsPSOInit.BlendState = TStaticBlendState<CW_RGBA, BO_Add, BF_One, BF_Zero>::GetRHI();
+			GraphicsPSOInit.BlendState = TStaticBlendState<CW_RGBA, BO_Add, BF_One, BF_One, BO_Add, BF_One, BF_One>::GetRHI();
 			GraphicsPSOInit.RasterizerState = TStaticRasterizerState<FM_Solid, CM_None>::GetRHI();
 			GraphicsPSOInit.DepthStencilState = TStaticDepthStencilState<false, CF_Always>::GetRHI();
 			GraphicsPSOInit.BoundShaderState.VertexDeclarationRHI = GFilterVertexDeclaration.VertexDeclarationRHI;
@@ -348,6 +364,7 @@ void FDeferredShadingSceneRenderer::RayTraceTranslucency(FRHICommandListImmediat
 		}
 
 		ResolveSceneColor(RHICmdList);
+		SceneContext.FinishRenderingSceneColor(RHICmdList);
 	}
 }
 
@@ -360,11 +377,6 @@ void FDeferredShadingSceneRenderer::RayTraceTranslucencyView(
 	int32 HeightFog,
 	float ResolutionFraction)
 {
-	if (!View.bHasTranslucentViewMeshElements)
-	{
-		return;
-	}
-
 	FSceneRenderTargets& SceneContext = FSceneRenderTargets::Get(GraphBuilder.RHICmdList);
 
 	int32 UpscaleFactor = int32(1.0f / ResolutionFraction);
@@ -387,6 +399,7 @@ void FDeferredShadingSceneRenderer::RayTraceTranslucencyView(
 	FRayTracingTranslucencyRG::FParameters* PassParameters = GraphBuilder.AllocParameters<FRayTracingTranslucencyRG::FParameters>();
 
 	PassParameters->SamplesPerPixel = SamplePerPixel;
+	PassParameters->MaxRefractionRays = GRayTracingTranslucencyMaxRefractionRays;
 	PassParameters->HeightFog = HeightFog;
 	PassParameters->ShouldDoDirectLighting = GRayTracingTranslucencyDirectLighting;
 	PassParameters->ShouldDoReflectedShadows = GRayTracingTranslucencyShadows;
@@ -430,7 +443,7 @@ void FDeferredShadingSceneRenderer::RayTraceTranslucencyView(
 	ClearUnusedGraphResources(RayGenShader, PassParameters);
 
 	GraphBuilder.AddPass(
-		RDG_EVENT_NAME("TranslucencyRayTracing %dx%d", RayTracingResolution.X, RayTracingResolution.X),
+		RDG_EVENT_NAME("TranslucencyRayTracing %dx%d", RayTracingResolution.X, RayTracingResolution.Y),
 		PassParameters,
 		ERenderGraphPassFlags::Compute,
 		[PassParameters, this, &View, RayGenShader, RayTracingResolution](FRHICommandList& RHICmdList)
