Integrate path tracer with FSkyLightData

TODO: Compute dependencies for building the deprecated CDF will be removed in a future submission.

diff --git a/Engine/Shaders/Private/PathTracing/Light/PathTracingSkyLight.ush b/Engine/Shaders/Private/PathTracing/Light/PathTracingSkyLight.ush
--- a/Engine/Shaders/Private/PathTracing/Light/PathTracingSkyLight.ush
+++ b/Engine/Shaders/Private/PathTracing/Light/PathTracingSkyLight.ush
@@ -6,141 +6,6 @@ SkyLight.usf: Light sampling functions for SkyLight implementation
 
 #pragma once
 
-#include "PathTracingLightCommon.ush"
-
-
-/**
-* Cube texel solid angle computation based on formula from Manne Öhrström's
-* thesis - "Spherical Harmonics, Precomputed Radiance Transfer and Realtime Radiosity in Computer Games".
-*/
-float AreaElement(float x, float y)
-{
-	// Equation 7.12
-	return atan2(x * y, sqrt(x * x + y * y + 1));
-}
-
-float TexelCoordSolidAngle(uint2 TextureCoord, uint2 TextureRes)
-{
-	float2 InvResolution = 1.0f / TextureRes;
-
-	// Convert texture coordinate to [-1, 1] range, offset to texel center.
-	float U = (2.0f * ((float)TextureCoord.x + 0.5f) * InvResolution.x) - 1.0f;
-	float V = (2.0f * ((float)TextureCoord.y + 0.5f) * InvResolution.y) - 1.0f;
-
-	// Projected area
-	float x0 = U - InvResolution.x;
-	float y0 = V - InvResolution.y;
-	float x1 = U + InvResolution.x;
-	float y1 = V + InvResolution.y;
-
-	// Equation 7.12
-	float SolidAngle = AreaElement(x0, y0) - AreaElement(x0, y1) - AreaElement(x1, y0) + AreaElement(x1, y1);
-	return SolidAngle;
-}
-
-uint GetMaximumComponentIndex(float3 Vector)
-{
-	float AbsX = abs(Vector.x);
-	float AbsY = abs(Vector.y);
-	float AbsZ = abs(Vector.z);
-	if (AbsX > AbsY)
-	{
-		return AbsX > AbsZ ? 0 : 2;
-	}
-	else
-	{
-		return AbsY > AbsZ ? 1 : 2;
-	}
-}
-
-float3 GetTextureCubeVector(uint3 TexelCoord, uint2 TextureRes)
-{
-	float3 Vector;
-
-	float2 UV = TexelCoord.xy / float2(TextureRes.xy);
-	UV = UV * 2.0 - 1.0;
-
-	uint Cubeface = TexelCoord.z;
-	if (Cubeface == 0)
-	{
-		Vector = float3(1, -UV.y, -UV.x);
-	}
-	else if (Cubeface == 1)
-	{
-		Vector = float3(-1, -UV.y, UV.x);
-	}
-	else if (Cubeface == 2)
-	{
-		Vector = float3(UV.x, 1, UV.y);
-	}
-	else if (Cubeface == 3)
-	{
-		Vector = float3(UV.x, -1, -UV.y);
-	}
-	else if (Cubeface == 4)
-	{
-		Vector = float3(UV.x, -UV.y, 1);
-	}
-	else if (Cubeface == 5)
-	{
-		Vector = float3(-UV.x, -UV.y, -1);
-	}
-
-	return Vector;
-}
-
-uint3 GetTextureCubeCoordinate(float3 Vector, uint2 TextureRes)
-{
-	// Strongest coordinate determines cube face
-	uint MaxComponent = GetMaximumComponentIndex(Vector);
-	Vector /= abs(Vector[MaxComponent]);
-	float3 TextureCubeCoordinate;
-	if (MaxComponent == 0)
-	{
-		if (Vector[MaxComponent] > 0)
-		{
-			TextureCubeCoordinate = float3(-Vector.z, -Vector.y, 0);
-		}
-		else
-		{
-			TextureCubeCoordinate = float3(Vector.z, -Vector.y, 1);
-		}
-	}
-	else if (MaxComponent == 1)
-	{
-		if (Vector[MaxComponent] > 0)
-		{
-			TextureCubeCoordinate = float3(Vector.x, Vector.z, 2);
-		}
-		else
-		{
-			TextureCubeCoordinate = float3(Vector.x, -Vector.z, 3);
-		}
-	}
-	else if (MaxComponent == 2)
-	{
-		if (Vector[MaxComponent] > 0)
-		{
-			TextureCubeCoordinate = float3(Vector.x, -Vector.y, 4);
-		}
-		else
-		{
-			TextureCubeCoordinate = float3(-Vector.x, -Vector.y, 5);
-		}
-	}
-
-	TextureCubeCoordinate.xy = (TextureCubeCoordinate.xy + 1.0) / 2.0;
-	TextureCubeCoordinate.xy *= TextureRes;
-	return TextureCubeCoordinate;
-}
-
-uint LinearCoordinate(uint2 Coordinate)
-{
-	uint2 Dimensions;
-	SkyLightData.SkyLightTexture.GetDimensions(Dimensions.x, Dimensions.y);
-	return Coordinate.y * Dimensions.x + Coordinate.x;
-}
-
 void SkyLight_GenerateLightRay(
 	RayDesc Ray,
 	int LightId,
@@ -148,10 +13,6 @@ void SkyLight_GenerateLightRay(
 	out RayDesc OutLightRay
 )
 {
-	OutLightRay.Origin = Ray.Origin;
-	OutLightRay.TMin = 0.001;
-	OutLightRay.Direction = normalize(LightUV);
-	OutLightRay.TMax = 1e6;
 }
 
 void SkyLight_PdfLight(
@@ -161,45 +22,7 @@ void SkyLight_PdfLight(
 	out float OutPdf
 )
 {
-#if 1
-	// Uniform spherical distribution
-	float SolidAnglePdf = 1.0 / (4.0 * PI);
-	OutPdf = SolidAnglePdf;
-#else
-	// Convert direction to cubemap coord
-	uint2 TextureRes;
-	SkyLightData.SkyLightTexture.GetDimensions(TextureRes.x, TextureRes.y);
-	uint3 TextureCubeCoordinate = GetTextureCubeCoordinate(LightUV, TextureRes);
-	uint CubeFace = TextureCubeCoordinate.z;
-
-	// Select cube face
-	float CubeFacePdf = SkyLightData.SkyLightCubeFaceCdf[CubeFace];
-	if (CubeFace > 0)
-	{
-		CubeFacePdf -= SkyLightData.SkyLightCubeFaceCdf[CubeFace - 1];
-	}
-
-	// Select column
-	uint CubeFaceIndex = CubeFace * TextureRes.y;
-	float ColumnPdf = SkyLightData.SkyLightColumnCdf[CubeFaceIndex + TextureCubeCoordinate.y];
-	if (TextureCubeCoordinate.y > 0)
-	{
-		ColumnPdf -= SkyLightData.SkyLightColumnCdf[CubeFaceIndex + TextureCubeCoordinate.y - 1];
-	}
-
-	// Select row
-	CubeFaceIndex = CubeFace * TextureRes.x * TextureRes.y;
-	float RowPdf = SkyLightData.SkyLightRowCdf[CubeFaceIndex + LinearCoordinate(TextureCubeCoordinate.xy)];
-	if (TextureCubeCoordinate.x > 0)
-	{
-		RowPdf -= SkyLightData.SkyLightRowCdf[CubeFaceIndex + LinearCoordinate(TextureCubeCoordinate.xy - uint2(1, 0))];
-	}
-
-	// Apply solid angle
-	float SolidAnglePdf = 1.0 / TexelCoordSolidAngle(TextureCubeCoordinate.xy, TextureRes);
-	OutPdf = CubeFacePdf * RowPdf * ColumnPdf * SolidAnglePdf;
-	OutPdf = max(OutPdf, 0.0);
-#endif
+	OutPdf = 0.0;
 }
 
 void SkyLight_EvalLight(
@@ -209,7 +32,7 @@ void SkyLight_EvalLight(
 	out float3 OutRadiance
 )
 {
-	OutRadiance = GetColor(LightId) * TextureCubeSampleLevel(SkyLightData.SkyLightTexture, SkyLightData.SkyLightTextureSampler, LightUV, 0).rgb;
+	OutRadiance = SkyLight.Color * TextureCubeSampleLevel(SkyLight.Texture, SkyLight.TextureSampler, LightUV, 0).rgb;
 }
 
 void SkyLight_SampleLight(
@@ -221,102 +44,6 @@ void SkyLight_SampleLight(
 	out float OutPdf
 )
 {
-#if 1
-	// Uniform spherical distribution
-	float theta = 2.0 * PI * RandSample.y;
-	float phi = acos(2.0 * RandSample.z - 1.0);
-	OutLightUV = float3(cos(theta) * sin(phi), sin(theta) * sin(phi), cos(phi));
-	OutPdf = 1.0 / (4.0 * PI);
-#else
-	uint2 TextureRes;
-	SkyLightData.SkyLightTexture.GetDimensions(TextureRes.x, TextureRes.y);
-
-	// TODO: Split random variable for reuse
-
-	// Sample from cube face Cdf
-	uint Left = 0;
-	uint Right = 5;
-	while (Left < Right)
-	{
-		uint Mid = (Left + Right) / 2;
-		if (RandSample.y <= SkyLightData.SkyLightCubeFaceCdf[Mid])
-		{
-			Right = Mid;
-		}
-		else
-		{
-			Left = Mid + 1;
-		}
-	}
-	uint CubeFace = Left;
-	uint CubeFaceIndex = CubeFace * TextureRes.y;
-
-	// Sample from column Cdf
-	Left = 0;
-	Right = TextureRes.y - 1;
-	while (Left < Right)
-	{
-		uint Mid = (Left + Right) / 2;
-		if (RandSample.z <= SkyLightData.SkyLightColumnCdf[CubeFaceIndex + Mid])
-		{
-			Right = Mid;
-		}
-		else
-		{
-			Left = Mid + 1;
-		}
-	}
-	uint RowIndex = Left;
-	CubeFaceIndex = CubeFace * TextureRes.x * TextureRes.y;
-
-	// Sample from row Cdf
-	Left = 0;
-	Right = TextureRes.x - 1;
-	while (Left < Right)
-	{
-		uint Mid = (Left + Right) / 2;
-		if (RandSample.w <= SkyLightData.SkyLightRowCdf[CubeFaceIndex + LinearCoordinate(uint2(Mid, RowIndex))])
-		{
-			Right = Mid;
-		}
-		else
-		{
-			Left = Mid + 1;
-		}
-	}
-	uint ColumnIndex = Left;
-
-	// Convert (column, row, face) to appropriate lighting direction
-	float3 TextureCubeCoordinate = float3(ColumnIndex, RowIndex, CubeFace);
-	float3 Direction = GetTextureCubeVector(TextureCubeCoordinate, TextureRes);
-	// DEBUG: The pdf is not invertible because of redundant mapping at TextureCube borders
-	//TextureCubeCoordinate = GetTextureCubeCoordinate(Direction, TextureRes);
-	OutLightUV = Direction;
-
-	// Compute cube face Pdf
-	CubeFace = TextureCubeCoordinate.z;
-	float CubeFacePdf = SkyLightData.SkyLightCubeFaceCdf[CubeFace];
-	if (CubeFace > 0)
-	{
-		CubeFacePdf -= SkyLightData.SkyLightCubeFaceCdf[CubeFace - 1];
-	}
-
-	// Compute column pdf
-	CubeFaceIndex = CubeFace * TextureRes.y;
-	float ColumnPdf = SkyLightData.SkyLightColumnCdf[CubeFaceIndex + TextureCubeCoordinate.y];
-	if (TextureCubeCoordinate.y > 0)
-	{
-		ColumnPdf -= SkyLightData.SkyLightColumnCdf[CubeFaceIndex + TextureCubeCoordinate.y - 1];
-	}
-
-	// Compute row pdf
-	CubeFaceIndex = CubeFace * TextureRes.x * TextureRes.y;
-	float RowPdf = SkyLightData.SkyLightRowCdf[CubeFaceIndex + LinearCoordinate(TextureCubeCoordinate.xy)];
-	if (TextureCubeCoordinate.x > 0)
-	{
-		RowPdf -= SkyLightData.SkyLightRowCdf[CubeFaceIndex + LinearCoordinate(TextureCubeCoordinate.xy - uint2(1, 0))];
-	}
-	float SolidAnglePdf = 1.0 / TexelCoordSolidAngle(TextureCubeCoordinate.xy, TextureRes);
-	OutPdf = CubeFacePdf * RowPdf * ColumnPdf * SolidAnglePdf;
-#endif
+	OutLightUV = 0.0;
+	OutPdf = 0.0;
 }
\ No newline at end of file
diff --git a/Engine/Shaders/Public/ShaderVersion.ush b/Engine/Shaders/Public/ShaderVersion.ush
--- a/Engine/Shaders/Public/ShaderVersion.ush
+++ b/Engine/Shaders/Public/ShaderVersion.ush
@@ -3,4 +3,4 @@
 // in Platform.ush (which should be included in any shader) it allows to invalidate the shader DDC.
 // 
 // If you are merging streams and there is a conflict with this GUID you should make a new GUID rather than taking one or the other.
-// GUID = 3ED98468-474C-4558-B573-BC368C3718D8
+// GUID = A54A0FDE-D1D0-4543-9A07-4EBA26888610
diff --git a/Engine/Source/Runtime/Renderer/Private/PathTracing.cpp b/Engine/Source/Runtime/Renderer/Private/PathTracing.cpp
--- a/Engine/Source/Runtime/Renderer/Private/PathTracing.cpp
+++ b/Engine/Source/Runtime/Renderer/Private/PathTracing.cpp
@@ -11,6 +11,7 @@
 #include "PostProcess/SceneFilterRendering.h"
 #include "PathTracingUniformBuffers.h"
 #include "RHI/Public/PipelineStateCache.h"
+#include "RayTracing/RayTracingSkyLight.h"
 
 ////static TAutoConsoleVariable<int32> CVarRayTracingPrimaryDebugMaxBounces(
 ////	TEXT("r.RayTracing.PrimaryDebug.MaxBounces"),
@@ -68,7 +69,6 @@ TAutoConsoleVariable<int32> CVarPathTracingRayCountFrequency(
 
 
 IMPLEMENT_GLOBAL_SHADER_PARAMETER_STRUCT(FPathTracingLightData, "SceneLightsData");
-IMPLEMENT_GLOBAL_SHADER_PARAMETER_STRUCT(FPathTracingSkyLightData, "SkyLightData");
 IMPLEMENT_GLOBAL_SHADER_PARAMETER_STRUCT(FPathTracingAdaptiveSamplingData, "AdaptiveSamplingData");
 
 class FPathTracingRG : public FGlobalShader
@@ -95,7 +95,8 @@ public:
 		TLASParameter.Bind(Initializer.ParameterMap, TEXT("TLAS"));
 		ViewParameter.Bind(Initializer.ParameterMap, TEXT("View"));
 		SceneLightsParameters.Bind(Initializer.ParameterMap, TEXT("SceneLightsData"));
-		SkyLightParameters.Bind(Initializer.ParameterMap, TEXT("SkyLightData"));
+		SkyLightParameters.Bind(Initializer.ParameterMap, TEXT("SkyLight"));
+		check(SkyLightParameters.IsBound());
 		AdaptiveSamplingParameters.Bind(Initializer.ParameterMap, TEXT("AdaptiveSamplingData"));
 
 		// Output
@@ -106,6 +107,7 @@ public:
 	}
 
 	void SetParameters(
+		FScene* Scene,
 		FRayTracingShaderBindingsWriter& GlobalResources,
 		const FRayTracingScene& RayTracingScene,
 		FUniformBufferRHIParamRef ViewUniformBuffer,
@@ -196,30 +198,62 @@ public:
 
 		// Sky light
 		{
-			FPathTracingSkyLightData SkyLightData;
+			FSkyLightData SkyLightData;
 			if (SkyLightTextureData)
 			{
-				SkyLightData.SkyLightRowCount = SkyLightRowCount;
-				SkyLightData.SkyLightColumnCount = SkyLightColumnCount;
-				SkyLightData.SkyLightRowCdf = SkyLightRowCdf.SRV;
-				SkyLightData.SkyLightColumnCdf = SkyLightColumnCdf.SRV;
-				SkyLightData.SkyLightCubeFaceCdf = SkyLightCubeFaceCdf.SRV;
-				SkyLightData.SkyLightTexture = SkyLightTextureData->TextureRHI;
-				SkyLightData.SkyLightTextureSampler = SkyLightTextureData->SamplerStateRHI;
+				SkyLightData.Color = FVector(Scene->SkyLight->LightColor);
+				SkyLightData.Texture = Scene->SkyLight->ProcessedTexture->TextureRHI;
+				SkyLightData.TextureSampler = Scene->SkyLight->ProcessedTexture->SamplerStateRHI;
+
+#if 0
+				// For when data is actually hooked up..
+				SkyLightData.MipDimensions = Scene->SkyLight->SkyLightMipDimensions;
+				SkyLightData.MipTreePosX = Scene->SkyLight->SkyLightMipTreePosX.SRV;
+				SkyLightData.MipTreeNegX = Scene->SkyLight->SkyLightMipTreeNegX.SRV;
+				SkyLightData.MipTreePosY = Scene->SkyLight->SkyLightMipTreePosY.SRV;
+				SkyLightData.MipTreeNegY = Scene->SkyLight->SkyLightMipTreeNegY.SRV;
+				SkyLightData.MipTreePosZ = Scene->SkyLight->SkyLightMipTreePosZ.SRV;
+				SkyLightData.MipTreeNegZ = Scene->SkyLight->SkyLightMipTreeNegZ.SRV;
+
+				SkyLightData.MipTreePdfPosX = Scene->SkyLight->SkyLightMipTreePdfPosX.SRV;
+				SkyLightData.MipTreePdfNegX = Scene->SkyLight->SkyLightMipTreePdfNegX.SRV;
+				SkyLightData.MipTreePdfPosY = Scene->SkyLight->SkyLightMipTreePdfPosY.SRV;
+				SkyLightData.MipTreePdfNegY = Scene->SkyLight->SkyLightMipTreePdfNegY.SRV;
+				SkyLightData.MipTreePdfPosZ = Scene->SkyLight->SkyLightMipTreePdfPosZ.SRV;
+				SkyLightData.MipTreePdfNegZ = Scene->SkyLight->SkyLightMipTreePdfNegZ.SRV;
+				SkyLightData.SolidAnglePdf = Scene->SkyLight->SolidAnglePdf.SRV;
+#endif
 			}
 			else
 			{
-				auto BlackTextureBuffer = RHICreateShaderResourceView(GBlackTexture->TextureRHI->GetTexture2D(), 0);
-				SkyLightData.SkyLightRowCount = 0;
-				SkyLightData.SkyLightColumnCount = 0;
-				SkyLightData.SkyLightRowCdf = BlackTextureBuffer;
-				SkyLightData.SkyLightColumnCdf = BlackTextureBuffer; 
-				SkyLightData.SkyLightCubeFaceCdf = RHICreateShaderResourceView(GBlackCubeArrayTexture->TextureRHI->GetTextureCube(), 0);
-				SkyLightData.SkyLightTexture = GBlackTextureCube->TextureRHI;
-				SkyLightData.SkyLightTextureSampler = GBlackTextureCube->SamplerStateRHI;
+				SkyLightData.Color = FVector(0.0);
+				SkyLightData.Texture = GBlackTextureCube->TextureRHI;
+				SkyLightData.TextureSampler = TStaticSamplerState<SF_Bilinear, AM_Clamp, AM_Clamp, AM_Clamp>::GetRHI();
 			}
 
-			FUniformBufferRHIRef SkyLightUniformBuffer = RHICreateUniformBuffer(&SkyLightData, FPathTracingSkyLightData::StaticStructMetadata.GetLayout(), EUniformBufferUsage::UniformBuffer_SingleDraw);
+			// dxr_todo: factor out these pass constants
+			SkyLightData.SamplesPerPixel = 0;
+			SkyLightData.SamplingStopLevel = 0;
+			// dxr_todo
+
+			SkyLightData.MipDimensions = FIntVector(0);
+			auto BlackTextureBuffer = RHICreateShaderResourceView(GBlackTexture->TextureRHI->GetTexture2D(), 0);
+			SkyLightData.MipTreePosX = BlackTextureBuffer;
+			SkyLightData.MipTreeNegX = BlackTextureBuffer;
+			SkyLightData.MipTreePosY = BlackTextureBuffer;
+			SkyLightData.MipTreeNegY = BlackTextureBuffer;
+			SkyLightData.MipTreePosZ = BlackTextureBuffer;
+			SkyLightData.MipTreeNegZ = BlackTextureBuffer;
+
+			SkyLightData.MipTreePdfPosX = BlackTextureBuffer;
+			SkyLightData.MipTreePdfNegX = BlackTextureBuffer;
+			SkyLightData.MipTreePdfPosY = BlackTextureBuffer;
+			SkyLightData.MipTreePdfNegY = BlackTextureBuffer;
+			SkyLightData.MipTreePdfPosZ = BlackTextureBuffer;
+			SkyLightData.MipTreePdfNegZ = BlackTextureBuffer;
+			SkyLightData.SolidAnglePdf = BlackTextureBuffer;
+
+			FUniformBufferRHIRef SkyLightUniformBuffer = RHICreateUniformBuffer(&SkyLightData, FSkyLightData::StaticStructMetadata.GetLayout(), EUniformBufferUsage::UniformBuffer_SingleDraw);
 			GlobalResources.Set(SkyLightParameters, SkyLightUniformBuffer);
 		}
 
@@ -434,6 +468,8 @@ void FDeferredShadingSceneRenderer::RenderPathTracing(FRHICommandListImmediate&
 	FRWBuffer SkyLightRowCdf;
 	FRWBuffer SkyLightColumnCdf;
 	FRWBuffer SkyLightCubeFaceCdf;
+#if 0
+	// #dxr_todo: Compute shader dependencies will be removed in a subsequent submission
 	if (Scene->SkyLight)
 	{
 		if (Scene->SkyLight->ShouldRebuildCdf())
@@ -450,6 +486,7 @@ void FDeferredShadingSceneRenderer::RenderPathTracing(FRHICommandListImmediate&
 		FIntVector Dimensions = FIntVector(Scene->SkyLight->ProcessedTexture->GetSizeX(), Scene->SkyLight->ProcessedTexture->GetSizeY(), 6);
 		VisualizeSkyLightCdf(RHICmdList, View, Dimensions, SkyLightRowCdf, SkyLightColumnCdf, SkyLightCubeFaceCdf);
 	}
+#endif
 
 	FSceneRenderTargets& SceneContext = FSceneRenderTargets::Get(RHICmdList);
 	auto ViewSize = View.ViewRect.Size();
@@ -518,6 +555,7 @@ void FDeferredShadingSceneRenderer::RenderPathTracing(FRHICommandListImmediate&
 	}
 
 	RayGenShader->SetParameters(
+		Scene,
 		GlobalResources,
 		View.PerViewRayTracingScene,
 		View.ViewUniformBuffer,
diff --git a/Engine/Source/Runtime/Renderer/Private/PathTracingUniformBuffers.h b/Engine/Source/Runtime/Renderer/Private/PathTracingUniformBuffers.h
--- a/Engine/Source/Runtime/Renderer/Private/PathTracingUniformBuffers.h
+++ b/Engine/Source/Runtime/Renderer/Private/PathTracingUniformBuffers.h
@@ -26,17 +26,6 @@ BEGIN_GLOBAL_SHADER_PARAMETER_STRUCT(FPathTracingLightData, )
 END_GLOBAL_SHADER_PARAMETER_STRUCT()
 
 
-BEGIN_GLOBAL_SHADER_PARAMETER_STRUCT(FPathTracingSkyLightData, )
-	SHADER_PARAMETER(uint32, SkyLightRowCount)
-	SHADER_PARAMETER(uint32, SkyLightColumnCount)
-	SHADER_PARAMETER_SRV(Buffer<float>, SkyLightRowCdf)
-	SHADER_PARAMETER_SRV(Buffer<float>, SkyLightColumnCdf)
-	SHADER_PARAMETER_SRV(Buffer<float>, SkyLightCubeFaceCdf)
-	SHADER_PARAMETER_TEXTURE(TextureCube, SkyLightTexture)
-	SHADER_PARAMETER_SAMPLER(SamplerState, SkyLightTextureSampler)
-END_GLOBAL_SHADER_PARAMETER_STRUCT()
-
-
 BEGIN_GLOBAL_SHADER_PARAMETER_STRUCT(FPathTracingAdaptiveSamplingData, )
 	SHADER_PARAMETER(uint32, UseAdaptiveSampling)
 	SHADER_PARAMETER(uint32, RandomSequence)
