Add clear coat and multiple bounces support in ray traced reflections

* Clearcoat: stochastically sample the bottom and top layer
* Multibounce: fallback to skylight at the end of the path
* Max number of samplers is increased to 16 in FRayTracingShaderBindings

diff --git a/Engine/Shaders/Private/RayTracing/RayTracingLightingCommon.ush b/Engine/Shaders/Private/RayTracing/RayTracingLightingCommon.ush
--- a/Engine/Shaders/Private/RayTracing/RayTracingLightingCommon.ush
+++ b/Engine/Shaders/Private/RayTracing/RayTracingLightingCommon.ush
@@ -99,15 +99,6 @@ float3 ComputeDirectLighting(RayDesc Ray, in FRayCone RayCone, FMaterialClosestH
 	GBufferData.IndirectIrradiance = (Payload.IndirectIrradiance.x + Payload.IndirectIrradiance.y + Payload.IndirectIrradiance.z) / 3.f; 
 	GBufferData.ShadingModelID = 1;
 
-	//#dxr_todo: calculate on the fly or store in Payload?
-#if 1
-	GBufferData.SpecularColor = lerp(0.08 * GBufferData.Specular.xxx, GBufferData.BaseColor, GBufferData.Metallic);
-	GBufferData.DiffuseColor = GBufferData.BaseColor - GBufferData.BaseColor * GBufferData.Metallic;
-#else	
-	GBufferData.SpecularColor = Payload.SpecularColor;
-	GBufferData.DiffuseColor = Payload.DiffuseColor;
-#endif
-
 	//#dxr_todo: think on the more efficient way of iterate through lights (split in light types, etc)
 	for (uint DirIndex = 0; DirIndex < ReflectionLightsData.Count; DirIndex++)
 	{
@@ -187,6 +178,33 @@ float3 ComputeDirectLighting(RayDesc Ray, in FRayCone RayCone, FMaterialClosestH
 	return DirectLighting;
 }
 
+void ApplyViewDependentMaterialPayloadModifications(RayDesc Ray, inout FMaterialClosestHitPayload Payload)
+{
+	if (Payload.ShadingModelID == SHADINGMODELID_CLEAR_COAT)
+	{
+		float ClearCoat				= Payload.CustomData.x;
+		float ClearCoatRoughness	= Payload.CustomData.y;
+		float MetalSpec = 0.9;
+
+		float NoV = saturate( dot( Payload.WorldNormal, -Ray.Direction ) );
+
+		// Approximation of refraction's effect on EnvBRDF
+		float RefractionScale = ( (NoV * 0.5 + 0.5) * NoV - 1 ) * saturate( 1.25 - 1.25 * Payload.Roughness ) + 1;
+
+		// Approximation of absorption integral, tuned for Roughness=0.4
+		float3 AbsorptionColor = Payload.BaseColor * (1 / MetalSpec);
+		float3 Absorption = AbsorptionColor * ( (NoV - 1) * 0.85 * ( 1 - lerp( AbsorptionColor, Square(AbsorptionColor), -0.78 ) ) + 1 );
+
+		float F0 = 0.04;
+		float Fc = Pow5( 1 - NoV );
+		float F = Fc + (1 - Fc) * F0;
+		float LayerAttenuation = lerp( 1, (1 - F), ClearCoat );
+
+		Payload.BaseColor = lerp( Payload.BaseColor * LayerAttenuation, MetalSpec * Absorption * RefractionScale, Payload.Metallic * ClearCoat );
+		Payload.Specular *= lerp( 1, RefractionScale, ClearCoat );
+	}
+}
+
 FMaterialClosestHitPayload TraceRayAndAccumulateResults(
 	in RayDesc Ray,
 	inout FRayCone RayCone,
@@ -209,6 +227,11 @@ FMaterialClosestHitPayload TraceRayAndAccumulateResults(
 		RayCone
 	);
 
+    ApplyViewDependentMaterialPayloadModifications(Ray, Payload);
+
+    Payload.DiffuseColor = Payload.BaseColor - Payload.BaseColor * Payload.Metallic;
+	Payload.SpecularColor = lerp(0.08 * Payload.Specular.xxx, Payload.BaseColor, Payload.Metallic);
+
 	if (ShouldDoDirectLighting)
 	{
 		float3 DirectLighting = ComputeDirectLighting(Ray, RayCone, Payload, ShouldDoReflectedShadows > 0);
@@ -222,7 +245,7 @@ FMaterialClosestHitPayload TraceRayAndAccumulateResults(
 	{
 		// Emissive & indirect contribution
 		Radiance += Payload.Radiance;
-		Radiance += Payload.IndirectIrradiance;
+		Radiance += Payload.DiffuseColor * Payload.IndirectIrradiance;
 	}
 
 	return Payload;
diff --git a/Engine/Shaders/Private/RayTracing/RayTracingMaterialHitShaders.usf b/Engine/Shaders/Private/RayTracing/RayTracingMaterialHitShaders.usf
--- a/Engine/Shaders/Private/RayTracing/RayTracingMaterialHitShaders.usf
+++ b/Engine/Shaders/Private/RayTracing/RayTracingMaterialHitShaders.usf
@@ -285,6 +285,27 @@ void MaterialCHS(inout FPackedMaterialClosestHitPayload PackedPayload, in FDefau
 	Payload.Metallic = Metallic;
 	Payload.Ior = Refraction.x;
     Payload.CustomData = float4(GetMaterialCustomData0(MaterialParameters), GetMaterialCustomData1(MaterialParameters), 0, 0);
+	#if CLEAR_COAT_BOTTOM_NORMAL
+	{
+		float2 oct2 = UnitVectorToOctahedron(Payload.WorldNormal);
+
+		#if NUM_MATERIAL_OUTPUTS_CLEARCOATBOTTOMNORMAL > 0
+			#if MATERIAL_TANGENTSPACENORMAL
+				float3 tempnormal = normalize(TransformTangentVectorToWorld( MaterialParameters.TangentToWorld, ClearCoatBottomNormal0(MaterialParameters) ));
+			#else
+				float3 tempnormal =  ClearCoatBottomNormal0(MaterialParameters);
+			#endif
+
+			float2 oct1 = UnitVectorToOctahedron(tempnormal);
+			float2 oct3 = ( (oct1 - oct2) *  0.5 ) + (128.0/255.0);
+			Payload.CustomData.a = oct3.x;
+			Payload.CustomData.z = oct3.y;
+		#else
+			Payload.CustomData.a = 128.0/255.0;
+			Payload.CustomData.z = 128.0/255.0;
+		#endif
+	}
+	#endif
 	//Payload.Opacity = Opacity;
     Payload.ShadingModelID = GetShadingModelID();
 
@@ -322,7 +343,7 @@ void MaterialCHS(inout FPackedMaterialClosestHitPayload PackedPayload, in FDefau
 	Payload.DiffuseColor = BaseColor - BaseColor * Metallic;
 	Payload.SpecularColor = lerp(0.08 * Specular.xxx, BaseColor, Metallic);
 	
-	Payload.IndirectIrradiance += Payload.DiffuseColor * DiffuseIndirectLighting; // Contains both lightmap & skylight
+	Payload.IndirectIrradiance += DiffuseIndirectLighting; // Contains both lightmap & skylight
 
 	PackedPayload = PackRayTracingPayload(Payload, PropagatedCone);
 }
diff --git a/Engine/Shaders/Private/RayTracing/RayTracingReflections.usf b/Engine/Shaders/Private/RayTracing/RayTracingReflections.usf
--- a/Engine/Shaders/Private/RayTracing/RayTracingReflections.usf
+++ b/Engine/Shaders/Private/RayTracing/RayTracingReflections.usf
@@ -78,16 +78,7 @@ void RayTracingReflectionsRGS()
 	float2 UV = (float2(PixelCoord) + 0.5) * InvBufferSize;
 	FScreenSpaceData ScreenSpaceData = GetScreenSpaceData(UV);
 	float Depth = ScreenSpaceData.GBuffer.Depth;
-	float3 CameraOrigin = ReconstructWorldPositionFromDepth(UV, 0.0f);
-	float3 WorldPosition = ReconstructWorldPositionFromDepth(UV, Depth);
-	float3 CameraDirection = normalize(WorldPosition - CameraOrigin);
-	//#dxr - todo: use also the shading normal (interpolate, or apply a more sophisticated approach)
-	float3 WorldNormal = ScreenSpaceData.GBuffer.WorldNormal;
-	float Roughness = ScreenSpaceData.GBuffer.Roughness;
-	
-	float3 BRDFThroughput = float3(1.0f, 1.0f, 1.0f);
 
-	uint MaxBounces = 1; //#dxr-todo: only works with 1 bounce until hit shaders are supported
 	uint LocalSamplesPerPixel = SamplesPerPixel;
 
 	// Mask out depth values at or beyond far plane
@@ -100,46 +91,123 @@ void RayTracingReflectionsRGS()
 	float ReflectedWeight = 0.0;
 	float HitDistance = -2.0;
 
-	const float RoughnessFade = GetRoughnessFade(Roughness, ReflectionMaxRoughness);
+	float RoughnessFade = GetRoughnessFade(ScreenSpaceData.GBuffer.Roughness, ReflectionMaxRoughness);
+	float ClearCoatRoughnessFade = ScreenSpaceData.GBuffer.ShadingModelID == SHADINGMODELID_CLEAR_COAT ? GetRoughnessFade(ScreenSpaceData.GBuffer.CustomData.y, ReflectionMaxRoughness) : 1.0f;
 
 	// Sky is only sampled when infinite reflection rays are used.
 	const bool bAllowSkySampling = ReflectionMaxRayDistance < 0;
 
 	// Shorten the rays on rougher surfaces between user-provided min and max ray lengths.
 	// When a shortened ray misses the geometry, we fall back to local reflection capture sampling (similar to SSR).
-	const float LocalMaxRayDistance = bAllowSkySampling ? 1e27f : lerp(ReflectionMaxRayDistance, ReflectionMinRayDistance, Roughness);
+	const float LocalMaxRayDistance = bAllowSkySampling ? 1e27f : lerp(ReflectionMaxRayDistance, ReflectionMinRayDistance, ScreenSpaceData.GBuffer.Roughness);
 
-	if (RoughnessFade > 0)
+	if (RoughnessFade > 0 || ClearCoatRoughnessFade > 0)
 	{
 		for (uint SampleIndex = 0; SampleIndex < LocalSamplesPerPixel; ++SampleIndex)
 		{
+			bool bIsValidSample = false;
+			
+			float3 PathRadiance = 0;
+			float3 BRDFThroughput = float3(1.0f, 1.0f, 1.0f);
+			
+			// Initialize path states according to GBuffer
+			float3 CameraOrigin = ReconstructWorldPositionFromDepth(UV, 0.0f);
+			float3 WorldPosition = ReconstructWorldPositionFromDepth(UV, Depth);
+			float3 CameraDirection = normalize(WorldPosition - CameraOrigin);
+			float3 WorldNormal = ScreenSpaceData.GBuffer.WorldNormal;
+			float Roughness = ScreenSpaceData.GBuffer.Roughness;
+			float3 SpecularColor = ScreenSpaceData.GBuffer.SpecularColor;
 			float3 IncidentDirection = CameraDirection;
-					
+			float ClearCoat;
+			float ClearCoatRoughness;
+			#if CLEAR_COAT_BOTTOM_NORMAL
+			float3 ClearCoatBottomNormal;
+			#endif
+			
+			uint ShadingModelID = ScreenSpaceData.GBuffer.ShadingModelID;
+			if (ShadingModelID == SHADINGMODELID_CLEAR_COAT)
+			{
+				ClearCoat = ScreenSpaceData.GBuffer.CustomData.x;
+				ClearCoatRoughness = ScreenSpaceData.GBuffer.CustomData.y;
+				#if CLEAR_COAT_BOTTOM_NORMAL
+					const float2 oct1 = ((float2(ScreenSpaceData.GBuffer.CustomData.a, ScreenSpaceData.GBuffer.CustomData.z) * 2) - (256.0/255.0)) + UnitVectorToOctahedron(ScreenSpaceData.GBuffer.WorldNormal);
+					ClearCoatBottomNormal = OctahedronToUnitVector(oct1);
+				#endif
+			}
+			else
+			{
+				ClearCoat = 1.0f;
+				ClearCoatRoughness = 1.0f;
+				#if CLEAR_COAT_BOTTOM_NORMAL
+					ClearCoatBottomNormal = WorldNormal;
+				#endif
+			}
+
 			FRayCone RayCone = (FRayCone)0;
 			RayCone.SpreadAngle = View.EyeToPixelSpreadAngle;
             float SurfaceCurvature = 0.0f; /* #todo_dxr assume no curvature */
 			RayCone = PropagateRayCone(RayCone, SurfaceCurvature, Depth);
 
-			for (uint BounceIndex = 0; BounceIndex < MaxBounces; ++BounceIndex)
+			uint DummyVariable;
+				
+			const uint MaxBounces = 1;
+			// Add one more 'fake bounce' and we will evaluate cubemap there
+			for (uint BounceIndex = 0; BounceIndex < MaxBounces + 1; ++BounceIndex)
 			{
+				float NoV = saturate(dot(-IncidentDirection, WorldNormal));
+				float FinalRoughness = Roughness;
+				if (ShadingModelID == SHADINGMODELID_CLEAR_COAT)
+				{
+					float F = EnvBRDF( 0.04, ClearCoatRoughness, NoV ).x;
+					F *= ClearCoat;
+					
+					float w = RandomSequence_GenerateSample1D(RandSequence, DummyVariable);
+					
+					if (w <= F)
+					{
+						SpecularColor = 0.04;
+						FinalRoughness = ClearCoatRoughness;
+					}
+					else
+					{
+						#if CLEAR_COAT_BOTTOM_NORMAL
+							WorldNormal = ClearCoatBottomNormal;
+							NoV = saturate(dot(-IncidentDirection, WorldNormal));
+						#endif
+						BRDFThroughput *= EnvBRDF( SpecularColor, Roughness, NoV );
+					}
+				}
+				else
+				{
+					// If it is not clear coat, and we're on the first bounce, the G F terms will be applied by ReflectionEnvironmentPixelShader
+					if (BounceIndex > 0)
+						BRDFThroughput *= EnvBRDF( SpecularColor, Roughness, NoV );
+				}				
+
 				RayDesc Ray;
 				Ray.Origin = WorldPosition + WorldNormal * 0.01;
 				Ray.TMin = 0.0;
 				Ray.TMax = LocalMaxRayDistance;
-
-				uint DummyVariable;
-				float2 RandSample = RandomSequence_GenerateSample2D(RandSequence, DummyVariable);
-
-				if (LocalSamplesPerPixel == 1)
+				
+				int RetryCount = 0;
+				do
 				{
-					uint2 SobolFrame = ComputePixelUniqueSobolRandSample(DispatchThreadId);
-					RandSample = SobolIndexToUniformUnitSquare(SobolFrame);
-				}
-
-				Ray.Direction = GenerateReflectedRayDirection(IncidentDirection, WorldNormal, Roughness, RandSample);
-				if (dot(WorldNormal, Ray.Direction) <= 0.0)
+					float2 RandSample = RandomSequence_GenerateSample2D(RandSequence, DummyVariable);
+
+					// Sobol can only be used on the first sample of the first bounce
+					if (BounceIndex == 0 && LocalSamplesPerPixel == 1 && RetryCount == 0)
+					{
+						uint2 SobolFrame = ComputePixelUniqueSobolRandSample(DispatchThreadId);
+						RandSample = SobolIndexToUniformUnitSquare(SobolFrame);
+					}
+				
+					Ray.Direction = GenerateReflectedRayDirection(IncidentDirection, WorldNormal, FinalRoughness, RandSample);
+				} while (dot(WorldNormal, Ray.Direction) < 0.0 && ++RetryCount < 16);
+
+				if (dot(WorldNormal, Ray.Direction) < 0.0)
 				{
 					// Refraction rays are handled separately
+					bIsValidSample = false;
 					break;
 				}
 
@@ -148,16 +216,27 @@ void RayTracingReflectionsRGS()
 				uint RayFlags = RAY_FLAG_CULL_BACK_FACING_TRIANGLES;
 
 				float3 Radiance = float3(0, 0, 0);
-				FMaterialClosestHitPayload Payload = TraceRayAndAccumulateResults(Ray, RayCone, Radiance);
+				FMaterialClosestHitPayload Payload;
+				
+				// Force the ray to miss on the MaxBounces + 1 fake bounce, so that we fallback to cubemap there
+				if (BounceIndex < MaxBounces)
+				{
+					Payload = TraceRayAndAccumulateResults(Ray, RayCone, Radiance);
+				}
 
-				float4 PathRadiance = 0.0f;
 				bool bApplyHeightFog = false;
 				float3 OriginToCollider = 0.0f;
 
-				if (Payload.HitT >= 0)
+				// Force the ray to miss on the MaxBounces + 1 fake bounce, so that we fallback to cubemap there
+				if (Payload.HitT >= 0 && BounceIndex < MaxBounces)
 				{
-					PathRadiance = float4(Radiance, 1.0);
-					HitDistance = Payload.HitT;
+					PathRadiance += BRDFThroughput * Radiance; 
+					bIsValidSample = true;
+
+					if (BounceIndex == 0)
+					{
+						HitDistance = Payload.HitT;
+					}
 
 					bApplyHeightFog = HeightFog > 0;
 					OriginToCollider = Payload.WorldPos - Ray.Origin;
@@ -165,21 +244,28 @@ void RayTracingReflectionsRGS()
 				else if (bAllowSkySampling && ReflectionStruct.SkyLightParameters.y > 0)
 				{
 					float SkyAverageBrightness = 1.0f;
-					float3 SkyLighting = GetSkyLightReflection(Ray.Direction, Roughness, SkyAverageBrightness);
-					PathRadiance = float4(SkyLighting, 1.0);
+					float3 SkyLighting = GetSkyLightReflection(Ray.Direction, FinalRoughness, SkyAverageBrightness);
+					PathRadiance += BRDFThroughput * SkyLighting;
+					bIsValidSample = true;
 
 					bApplyHeightFog = HeightFog > 0;
 					const float SkyFarDistance = 100000.0f;
-					OriginToCollider = Ray.Direction*SkyFarDistance;
-						
-					// To tell the denoiser the black input is valid even if no valid ray trace result was found. Uses a high ray hit distance
-					// so that the sample may be reused by another neighbor pixel in the denoise as if RoughnessFade in ]0;1[.
-					HitDistance = 1.0e20;
+					OriginToCollider = Ray.Direction * SkyFarDistance;
+
+					if (BounceIndex == 0)
+					{
+						// To tell the denoiser the black input is valid even if no valid ray trace result was found. Uses a high ray hit distance
+						// so that the sample may be reused by another neighbor pixel in the denoise as if RoughnessFade in ]0;1[.
+						HitDistance = 1.0e20;
+					}
 				}
 				else
 				{
-					// Tell to the denoiser this sample reached outer bounds of the scene without any intersection.
-					HitDistance = -1;
+					if (BounceIndex == 0)
+					{
+						// Tell to the denoiser this sample reached outer bounds of the scene without any intersection.
+						HitDistance = -1;
+					}
 				}
 
 				// If there is a contribution on this path, apply the required fog on the sample. 
@@ -190,11 +276,32 @@ void RayTracingReflectionsRGS()
 				{
 					float4 HeightFogInscatteringAndTransmittance = CalculateHeightFog(OriginToCollider);
 					PathRadiance.rgb = PathRadiance.rgb * HeightFogInscatteringAndTransmittance.a + HeightFogInscatteringAndTransmittance.rgb;
-
 				}
-				ReflectedColor += PathRadiance;
-				ReflectedWeight += 1.0;
+				
 				IncidentDirection = Ray.Direction;
+				WorldPosition = Ray.Origin + Payload.HitT * Ray.Direction;
+				WorldNormal = Payload.WorldNormal;
+				Roughness = Payload.Roughness;
+				SpecularColor = Payload.SpecularColor;
+				
+				ShadingModelID = Payload.ShadingModelID;
+				if (ShadingModelID == SHADINGMODELID_CLEAR_COAT)
+				{
+					ClearCoat = Payload.CustomData.x;
+					ClearCoatRoughness = Payload.CustomData.y;
+					#if CLEAR_COAT_BOTTOM_NORMAL
+						const float2 oct1 = ((float2(Payload.CustomData.a, Payload.GBuffer.CustomData.z) * 2) - (256.0/255.0)) + UnitVectorToOctahedron(Payload.WorldNormal);
+						ClearCoatBottomNormal = OctahedronToUnitVector(oct1);
+					#endif
+				}
+				
+				if (all(BRDFThroughput < 0.1)) break;
+			}
+			
+			if (bIsValidSample)
+			{
+				ReflectedColor += float4(PathRadiance, 1.0f);
+				ReflectedWeight += 1.0;
 			}
 		}
 	}
@@ -203,8 +310,8 @@ void RayTracingReflectionsRGS()
 	{
 		float DistanceFadeRange = LocalMaxRayDistance * 0.1; // Fade out reflections over 10% of the total maximum ray length (arbitrarily chosen)
 		float DistanceFade = bAllowSkySampling ? 1.0 : saturate((LocalMaxRayDistance - HitDistance) * rcp(DistanceFadeRange));
-
-		ReflectedColor *= DistanceFade * RoughnessFade * rcp(ReflectedWeight);
+		
+		ReflectedColor *= DistanceFade * max(RoughnessFade, ClearCoatRoughnessFade) * rcp(ReflectedWeight);
 	}
 
 	ColorOutput[DispatchThreadId] = float4(ReflectedColor);
diff --git a/Engine/Shaders/Private/ReflectionEnvironmentPixelShader.usf b/Engine/Shaders/Private/ReflectionEnvironmentPixelShader.usf
--- a/Engine/Shaders/Private/ReflectionEnvironmentPixelShader.usf
+++ b/Engine/Shaders/Private/ReflectionEnvironmentPixelShader.usf
@@ -235,7 +235,8 @@ float3 ReflectionEnvironment(FScreenSpaceData ScreenSpaceData, float4 UVAndScree
 	Color.a = 1 - SSR.a;
 #endif
 
-	if( GBuffer.ShadingModelID == SHADINGMODELID_CLEAR_COAT )
+#if ENABLE_CLEAR_COAT
+	if(GBuffer.ShadingModelID == SHADINGMODELID_CLEAR_COAT )
 	{
 		const float ClearCoat = GBuffer.CustomData.x;
 		Color = lerp( Color, float4(0,0,0,1), ClearCoat );
@@ -248,6 +249,7 @@ float3 ReflectionEnvironment(FScreenSpaceData ScreenSpaceData, float4 UVAndScree
 		R = 2 * dot( V, ClearCoatUnderNormal ) * ClearCoatUnderNormal - V;
 #endif
 	}
+#endif
 
 	float AO = GBuffer.GBufferAO * ScreenSpaceData.AmbientOcclusion;
 	float RoughnessSq = GBuffer.Roughness * GBuffer.Roughness;
@@ -272,6 +274,7 @@ float3 ReflectionEnvironment(FScreenSpaceData ScreenSpaceData, float4 UVAndScree
 	BRANCH
 	if( GBuffer.ShadingModelID == SHADINGMODELID_CLEAR_COAT )
 	{
+#if ENABLE_CLEAR_COAT
 		const float ClearCoat			= GBuffer.CustomData.x;
 		const float ClearCoatRoughness	= GBuffer.CustomData.y;
 
@@ -296,6 +299,9 @@ float3 ReflectionEnvironment(FScreenSpaceData ScreenSpaceData, float4 UVAndScree
 
 		float3 TopLayerR = 2 * dot( V, N ) * N - V;
 		Color.rgb += PreExposure * GatherRadiance(Color.a, WorldPosition, TopLayerR, ClearCoatRoughness, BentNormal, IndirectIrradiance, GBuffer.ShadingModelID, NumCulledReflectionCaptures, DataStartIndex);
+#else
+		// When clear coat is disabled in ReflectionEnvironmentPS (both layers evaluated in ray traced reflections), do nothing here
+#endif
 	}
 	else
 	{
diff --git a/Engine/Source/Runtime/RHI/Public/RHICommandList.h b/Engine/Source/Runtime/RHI/Public/RHICommandList.h
--- a/Engine/Source/Runtime/RHI/Public/RHICommandList.h
+++ b/Engine/Source/Runtime/RHI/Public/RHICommandList.h
@@ -113,7 +113,7 @@ struct FRayTracingShaderBindings
 	FTextureRHIParamRef Textures[32] = {};
 	FShaderResourceViewRHIParamRef SRVs[32] = {};
 	FUniformBufferRHIParamRef UniformBuffers[8] = {};
-	FSamplerStateRHIParamRef Samplers[8] = {};
+	FSamplerStateRHIParamRef Samplers[16] = {};
 	FUnorderedAccessViewRHIParamRef UAVs[8] = {};
 };
 
diff --git a/Engine/Source/Runtime/Renderer/Private/RayTracing/RayTracingReflections.cpp b/Engine/Source/Runtime/Renderer/Private/RayTracing/RayTracingReflections.cpp
--- a/Engine/Source/Runtime/Renderer/Private/RayTracing/RayTracingReflections.cpp
+++ b/Engine/Source/Runtime/Renderer/Private/RayTracing/RayTracingReflections.cpp
@@ -145,6 +145,9 @@ class FRayTracingReflectionsRG : public FGlobalShader
 		SHADER_PARAMETER_TEXTURE(Texture2D, LTCAmpTexture)
 		SHADER_PARAMETER_SAMPLER(SamplerState, LTCAmpSampler)
 
+		SHADER_PARAMETER_TEXTURE(Texture2D, PreIntegratedGF)
+		SHADER_PARAMETER_SAMPLER(SamplerState, PreIntegratedGFSampler)
+
 		SHADER_PARAMETER_STRUCT_REF(FViewUniformShaderParameters, ViewUniformBuffer)
 		SHADER_PARAMETER_STRUCT_REF(FSceneTexturesUniformParameters, SceneTexturesStruct)
 		SHADER_PARAMETER_STRUCT_REF(FReflectionsLightData, LightData)
@@ -238,6 +241,8 @@ void FDeferredShadingSceneRenderer::RayTraceReflections(
 	PassParameters->LTCMatSampler = TStaticSamplerState<SF_Bilinear, AM_Clamp, AM_Clamp, AM_Clamp>::GetRHI();
 	PassParameters->LTCAmpTexture = GSystemTextures.LTCAmp->GetRenderTargetItem().ShaderResourceTexture;
 	PassParameters->LTCAmpSampler = TStaticSamplerState<SF_Bilinear, AM_Clamp, AM_Clamp, AM_Clamp>::GetRHI();
+	PassParameters->PreIntegratedGF = GSystemTextures.PreintegratedGF->GetRenderTargetItem().ShaderResourceTexture;
+	PassParameters->PreIntegratedGFSampler = TStaticSamplerState<SF_Bilinear, AM_Clamp, AM_Clamp, AM_Clamp>::GetRHI();
 
 	PassParameters->TLAS = View.PerViewRayTracingScene.RayTracingSceneRHI->GetShaderResourceView();
 	PassParameters->ViewUniformBuffer = View.ViewUniformBuffer;
diff --git a/Engine/Source/Runtime/Renderer/Private/ReflectionEnvironment.cpp b/Engine/Source/Runtime/Renderer/Private/ReflectionEnvironment.cpp
--- a/Engine/Source/Runtime/Renderer/Private/ReflectionEnvironment.cpp
+++ b/Engine/Source/Runtime/Renderer/Private/ReflectionEnvironment.cpp
@@ -507,6 +507,7 @@ class FReflectionEnvironmentSkyLightingPS : public FGlobalShader
 	class FSkyLight					: SHADER_PERMUTATION_BOOL("ENABLE_SKY_LIGHT");
 	class FDynamicSkyLight			: SHADER_PERMUTATION_BOOL("ENABLE_DYNAMIC_SKY_LIGHT");
 	class FSkyShadowing				: SHADER_PERMUTATION_BOOL("APPLY_SKY_SHADOWING");
+	class FClearCoat				: SHADER_PERMUTATION_BOOL("ENABLE_CLEAR_COAT");
 
 	using FPermutationDomain = TShaderPermutationDomain<
 		FHasBoxCaptures,
@@ -515,7 +516,8 @@ class FReflectionEnvironmentSkyLightingPS : public FGlobalShader
 		FSpecularBounce,
 		FSkyLight,
 		FDynamicSkyLight,
-		FSkyShadowing>;
+		FSkyShadowing,
+		FClearCoat>;
 
 	static FPermutationDomain RemapPermutation(FPermutationDomain PermutationVector)
 	{
@@ -543,7 +545,7 @@ class FReflectionEnvironmentSkyLightingPS : public FGlobalShader
 		return PermutationVector;
 	}
 
-	static FPermutationDomain BuildPermutationVector(const FViewInfo& View, bool bBoxCapturesOnly, bool bSphereCapturesOnly, bool bSupportDFAOIndirectOcclusion, bool bSpecularBounce, bool bEnableSkyLight, bool bEnableDynamicSkyLight, bool bApplySkyShadowing)
+	static FPermutationDomain BuildPermutationVector(const FViewInfo& View, bool bBoxCapturesOnly, bool bSphereCapturesOnly, bool bSupportDFAOIndirectOcclusion, bool bSpecularBounce, bool bEnableSkyLight, bool bEnableDynamicSkyLight, bool bApplySkyShadowing, bool bEnableClearCoat)
 	{
 		FPermutationDomain PermutationVector;
 
@@ -554,6 +556,7 @@ class FReflectionEnvironmentSkyLightingPS : public FGlobalShader
 		PermutationVector.Set<FSkyLight>(bEnableSkyLight);
 		PermutationVector.Set<FDynamicSkyLight>(bEnableDynamicSkyLight);
 		PermutationVector.Set<FSkyShadowing>(bApplySkyShadowing);
+		PermutationVector.Set<FClearCoat>(bEnableClearCoat);
 
 		return RemapPermutation(PermutationVector);
 	}
@@ -835,7 +838,7 @@ void FDeferredShadingSceneRenderer::RenderDeferredReflectionsAndSkyLighting(FRHI
 
 			TShaderMapRef<FPostProcessVS> VertexShader(View.ShaderMap);
 
-			auto PermutationVector = FReflectionEnvironmentSkyLightingPS::BuildPermutationVector(View, bHasBoxCaptures, bHasSphereCaptures, DynamicBentNormalAO != NULL, bReflectionCapture, bSkyLight, bDynamicSkyLight, bApplySkyShadowing);
+			auto PermutationVector = FReflectionEnvironmentSkyLightingPS::BuildPermutationVector(View, bHasBoxCaptures, bHasSphereCaptures, DynamicBentNormalAO != NULL, bReflectionCapture, bSkyLight, bDynamicSkyLight, bApplySkyShadowing, !bRayTracedReflections);
 
 			TShaderMapRef<FReflectionEnvironmentSkyLightingPS> PixelShader(View.ShaderMap, PermutationVector);
 
