Fixed height fog with multibounce reflection.

diff --git a/Engine/Shaders/Private/RayTracing/RayTracingReflections.usf b/Engine/Shaders/Private/RayTracing/RayTracingReflections.usf
--- a/Engine/Shaders/Private/RayTracing/RayTracingReflections.usf
+++ b/Engine/Shaders/Private/RayTracing/RayTracingReflections.usf
@@ -109,7 +109,7 @@ void RayTracingReflectionsRGS()
 			bool bIsValidSample = false;
 			
 			float3 PathRadiance = 0;
-			float3 BRDFThroughput = float3(1.0f, 1.0f, 1.0f);
+			float3 PathThroughput = float3(1.0f, 1.0f, 1.0f);
 			
 			// Initialize path states according to GBuffer
 			float3 CameraOrigin = ReconstructWorldPositionFromDepth(UV, 0.0f);
@@ -174,14 +174,14 @@ void RayTracingReflectionsRGS()
 							WorldNormal = ClearCoatBottomNormal;
 							NoV = saturate(dot(-IncidentDirection, WorldNormal));
 						#endif
-						BRDFThroughput *= EnvBRDF( SpecularColor, Roughness, NoV );
+						PathThroughput *= EnvBRDF( SpecularColor, Roughness, NoV );
 					}
 				}
 				else
 				{
 					// If it is not clear coat, and we're on the first bounce, the G F terms will be applied by ReflectionEnvironmentPixelShader
 					if (BounceIndex > 0)
-						BRDFThroughput *= EnvBRDF( SpecularColor, Roughness, NoV );
+						PathThroughput *= EnvBRDF( SpecularColor, Roughness, NoV );
 				}				
 
 				RayDesc Ray;
@@ -228,9 +228,10 @@ void RayTracingReflectionsRGS()
 				float3 OriginToCollider = 0.0f;
 
 				// Force the ray to miss on the MaxBounces + 1 fake bounce, so that we fallback to cubemap there
+				float3 PathSegmentRadiance = 0.0f;
 				if (Payload.HitT >= 0 && BounceIndex < MaxBounces)
 				{
-					PathRadiance += BRDFThroughput * Radiance; 
+					PathSegmentRadiance += Radiance;
 					bIsValidSample = true;
 
 					if (BounceIndex == 0)
@@ -245,7 +246,7 @@ void RayTracingReflectionsRGS()
 				{
 					float SkyAverageBrightness = 1.0f;
 					float3 SkyLighting = GetSkyLightReflection(Ray.Direction, FinalRoughness, SkyAverageBrightness);
-					PathRadiance += BRDFThroughput * SkyLighting;
+					PathSegmentRadiance += SkyLighting;
 					bIsValidSample = true;
 
 					bApplyHeightFog = HeightFog > 0;
@@ -272,11 +273,15 @@ void RayTracingReflectionsRGS()
 				// Also accumulate the path contribution.
 				// In case no contribution is accumulated, the ReflectionEnvironmentAndSky pass will fill it up using reflection volumes.
 				//#dxr - todo: when MaxBounces>1, we will need to take into account a path throughput affected from the fog trnasmittance.
+				float4 HeightFogInscatteringAndTransmittance = float4(0.0f, 0.0f, 0.0f, 1.0f);	// no fog 
 				if (bApplyHeightFog)
 				{
-					float4 HeightFogInscatteringAndTransmittance = CalculateHeightFog(OriginToCollider);
-					PathRadiance.rgb = PathRadiance.rgb * HeightFogInscatteringAndTransmittance.a + HeightFogInscatteringAndTransmittance.rgb;
+					HeightFogInscatteringAndTransmittance = CalculateHeightFog(OriginToCollider);
 				}
+				// Now accumulate the radiance coming through this path segment into the full path radiance.
+				PathSegmentRadiance.rgb = PathSegmentRadiance.rgb * HeightFogInscatteringAndTransmittance.a + HeightFogInscatteringAndTransmittance.rgb;
+				PathRadiance += PathSegmentRadiance * PathThroughput;
+				PathThroughput *= HeightFogInscatteringAndTransmittance.a;
 				
 				IncidentDirection = Ray.Direction;
 				WorldPosition = Ray.Origin + Payload.HitT * Ray.Direction;
@@ -295,7 +300,7 @@ void RayTracingReflectionsRGS()
 					#endif
 				}
 				
-				if (all(BRDFThroughput < 0.1)) break;
+				if (all(PathThroughput < 0.1)) break;
 			}
 			
 			if (bIsValidSample)
