Integrate CL 4939405. Allow niagara mesh particles to access velocity & Corrected integration math to derive prev position.

diff --git a/Engine/Shaders/Private/NiagaraMeshVertexFactory.ush b/Engine/Shaders/Private/NiagaraMeshVertexFactory.ush
--- a/Engine/Shaders/Private/NiagaraMeshVertexFactory.ush
+++ b/Engine/Shaders/Private/NiagaraMeshVertexFactory.ush
@@ -386,7 +386,9 @@ float4 CalcPrevWorldPosition(FVertexFactoryInput Input, FVertexFactoryIntermedia
 #endif
 	// if we don't have a full previous transform, we can at least calculate the position to get some motion blur
 	{
-		return CalcWorldPosition(Input, Intermediates)- float4(Intermediates.ParticleVelocity.xyz*NiagaraMeshVF.DeltaSeconds, 0.0);
+		// ParticleVelocity.xyz - normalized direction vector
+		// ParticleVelocity.w - velocity magnitude
+		return CalcWorldPosition(Input, Intermediates) - float4(Intermediates.ParticleVelocity.xyz * (Intermediates.ParticleVelocity.w * NiagaraMeshVF.DeltaSeconds), 0.0);
 	}
 }
 
diff --git a/Engine/Source/Runtime/Engine/Private/Materials/HLSLMaterialTranslator.h b/Engine/Source/Runtime/Engine/Private/Materials/HLSLMaterialTranslator.h
--- a/Engine/Source/Runtime/Engine/Private/Materials/HLSLMaterialTranslator.h
+++ b/Engine/Source/Runtime/Engine/Private/Materials/HLSLMaterialTranslator.h
@@ -990,7 +990,7 @@ public:
 			OutEnvironment.SetDefine(TEXT("NEEDS_PARTICLE_POSITION"), 1);
 		}
 
-		if (bNeedsParticleVelocity)
+		if (bNeedsParticleVelocity || Material->IsUsedWithNiagaraMeshParticles())
 		{
 			OutEnvironment.SetDefine(TEXT("NEEDS_PARTICLE_VELOCITY"), 1);
 		}
