PR #6032: Fixes to AmbientOcclusion calculation in Screen Space Diffuse Indirect (Contributed by [REDACTED_PERSON])

diff --git a/Engine/Shaders/Private/SSRT/SSRTDiffuseIndirect.usf b/Engine/Shaders/Private/SSRT/SSRTDiffuseIndirect.usf
--- a/Engine/Shaders/Private/SSRT/SSRTDiffuseIndirect.usf
+++ b/Engine/Shaders/Private/SSRT/SSRTDiffuseIndirect.usf
@@ -139,12 +139,12 @@ void MainCS(
 
 			// Backface modulation of hit surface
 			float3 SampleNormal = GetGBufferDataFromSceneTextures(HitUVz.xy).WorldNormal;
-			SampleColor.rgb *= saturate( 1 - dot( SampleNormal, L ) );
+			SampleColor.rgba *= saturate( 1 - dot( SampleNormal, L ) );
 
 			SampleColor.rgb *= rcp( 1 + Luminance(SampleColor.rgb) );
 
 			DiffuseColor += SampleColor.rgb;
-			AmbientOcclusion = SampleColor.a;
+			AmbientOcclusion += SampleColor.a;
 		}
 	}
 
