Fixed: Screen goes black when looking inside a body of water on iOS

diff --git a/Engine/Shaders/Private/PostProcessMaterialShaders.usf b/Engine/Shaders/Private/PostProcessMaterialShaders.usf
--- a/Engine/Shaders/Private/PostProcessMaterialShaders.usf
+++ b/Engine/Shaders/Private/PostProcessMaterialShaders.usf
@@ -253,7 +253,11 @@ void MainPS_ES2(
 #if NEEDTOSWITCHVERTICLEAXIS
 		SvPosition.y = View.ViewSizeAndInvSize.y - SvPosition.y; 
 #endif
-	SvPosition.z = LookupDeviceZ(InViewportUV);
+	// SceneDepth texture is not accessible during post-processing as we discard it at the end of mobile BasePass
+    // instead fetch depth from SceneColor.A 
+	float Depth = Texture2DSample(MobileSceneTextures.SceneColorTexture, MobileSceneTextures.SceneColorTextureSampler, InViewportUV).a;
+	float DeviceZ = ConvertToDeviceZ(Depth);
+	SvPosition.z = DeviceZ;
 	SvPosition.z = max(SvPosition.z, 1e-18);
 	// fill out other related material parameters
 	CalcMaterialParametersPost(Parameters, PixelMaterialInputs, SvPosition, true);
