Fixed: Issues with mobile rendering and SceneDepth material node usage in post-process materials

diff --git a/Engine/Shaders/Private/SceneTexturesCommon.ush b/Engine/Shaders/Private/SceneTexturesCommon.ush
--- a/Engine/Shaders/Private/SceneTexturesCommon.ush
+++ b/Engine/Shaders/Private/SceneTexturesCommon.ush
@@ -96,6 +96,10 @@ float CalcSceneDepth(float2 ScreenUV)
 #else
 	#if MOBILE_FORCE_DEPTH_TEXTURE_READS
 		return ConvertFromDeviceZ(Texture2DSampleLevel(MobileSceneTextures.SceneDepthTexture, MobileSceneTextures.SceneDepthTextureSampler, ScreenUV, 0).r);
+	#elif POST_PROCESS_MATERIAL
+		// SceneDepth texture is not accessible during post-processing as we discard it at the end of mobile BasePass
+		// instead fetch depth from SceneColor.A 
+		return Texture2DSample(MobileSceneTextures.SceneColorTexture, MobileSceneTextures.SceneColorTextureSampler, ScreenUV).a;
 	#else
 		#if COMPILER_GLSL_ES2
 			#if IOS
@@ -138,6 +142,9 @@ float LookupDeviceZ( float2 ScreenUV )
 #if	MOBILE_FORCE_DEPTH_TEXTURE_READS
 	// native Depth buffer lookup
 	return Texture2DSampleLevel(MobileSceneTextures.SceneDepthTexture, MobileSceneTextures.SceneDepthTextureSampler, ScreenUV, 0).r;
+#elif POST_PROCESS_MATERIAL
+	// SceneDepth texture is not accessible during post-processing, CalcSceneDepth will do the right thing
+	return ConvertToDeviceZ(CalcSceneDepth(ScreenUV));
 #elif COMPILER_GLSL_ES2 || COMPILER_GLSL_ES3_1
 	// todo: can be optimized
 	return ConvertToDeviceZ(CalcSceneDepth(ScreenUV));
diff --git a/Engine/Source/Runtime/Renderer/Private/PostProcess/PostProcessMaterial.cpp b/Engine/Source/Runtime/Renderer/Private/PostProcess/PostProcessMaterial.cpp
--- a/Engine/Source/Runtime/Renderer/Private/PostProcess/PostProcessMaterial.cpp
+++ b/Engine/Source/Runtime/Renderer/Private/PostProcess/PostProcessMaterial.cpp
@@ -154,7 +154,6 @@ public:
 
 		if (MaterialTarget == EPostProcessMaterialTarget::Mobile)
 		{
-			OutEnvironment.SetDefine(TEXT("MOBILE_FORCE_DEPTH_TEXTURE_READS"), 1); // Ensure post process materials will not attempt depth buffer fetch operations.
 			OutEnvironment.SetDefine(TEXT("POST_PROCESS_MATERIAL_BEFORE_TONEMAP"), (Location != BL_AfterTonemapping) ? 1 : 0);
 		}
 	}
