Fix issues with HQ DOF on mobile devices running vulkan

diff --git a/Engine/Shaders/Private/PostProcessDOF.usf b/Engine/Shaders/Private/PostProcessDOF.usf
--- a/Engine/Shaders/Private/PostProcessDOF.usf
+++ b/Engine/Shaders/Private/PostProcessDOF.usf
@@ -99,7 +99,12 @@ void SetupPS(
 
 	float MaskDistance = View.DepthOfFieldFocalDistance + View.DepthOfFieldFocalRegion * 0.5f;
 
+#if MOBILE_SHADING
+	// Mobile path gathers Depth later in ReadFullResAndDepth
+	float4 DepthQuad = float4(0,0,0,0);
+#else
 	float4 DepthQuad = GatherSceneDepth(UV, PostprocessInput1Size.zw);
+#endif
 
 	float4 FarColor = 0;
 	float4 NearColor = 0;
