Better fix for ES3.1 external texture gamma correction

diff --git a/Engine/Shaders/Private/MaterialTemplate.ush b/Engine/Shaders/Private/MaterialTemplate.ush
--- a/Engine/Shaders/Private/MaterialTemplate.ush
+++ b/Engine/Shaders/Private/MaterialTemplate.ush
@@ -1029,13 +1029,15 @@ MaterialFloat4 ProcessMaterialColorTextureLookup(MaterialFloat4 TextureValue)
 
 MaterialFloat4 ProcessMaterialExternalTextureLookup(MaterialFloat4 TextureValue)
 {
-#if (FEATURE_LEVEL == FEATURE_LEVEL_ES2) || COMPILER_GLSL_ES3_1
+#if FEATURE_LEVEL == FEATURE_LEVEL_ES2
 	// As per https://www.khronos.org/registry/OpenGL/extensions/OES/OES_EGL_image_external.txt
 	// sampling an external texture does not change the color space, so we always need to
 	// perform manual approximate sRGB conversion
 	float3 LinRGB = TextureValue.rgb;
 	LinRGB *= LinRGB; 
 	return MaterialFloat4(LinRGB, TextureValue.a);
+#elif COMPILER_GLSL_ES3_1
+	return MaterialFloat4(pow(TextureValue.rgb, 2.2f), TextureValue.a);
 #else
 	return ProcessMaterialColorTextureLookup(TextureValue);
 #endif
