Normal blend decal fix

diff --git a/Engine/Shaders/Private/DecalCommon.ush b/Engine/Shaders/Private/DecalCommon.ush
--- a/Engine/Shaders/Private/DecalCommon.ush
+++ b/Engine/Shaders/Private/DecalCommon.ush
@@ -36,7 +36,10 @@ void DecalCommonOutput(inout FPixelShaderIn In, inout FPixelShaderOut Out, float
 	Out.MRT[0] = float4(Color, Opacity);
 
 #if USE_PREEXPOSURE
-	Out.MRT[0].xyz *= View.PreExposure;
+	// Blend mode normal should not mult with exposure because Out.MRT[0] is world space normal in that case.
+	#if DECAL_BLEND_MODE != DECALBLENDMODEID_NORMAL
+		Out.MRT[0].xyz *= View.PreExposure;
+	#endif
 #endif
 
 #if (ES2_PROFILE || ES3_1_PROFILE)
