Fix MinClearCoatLevel shader variable type in ray tracing reflections #jira UE-77950 #rb none

diff --git a/Engine/Shaders/Private/RayTracing/RayTracingReflections.usf b/Engine/Shaders/Private/RayTracing/RayTracingReflections.usf
--- a/Engine/Shaders/Private/RayTracing/RayTracingReflections.usf
+++ b/Engine/Shaders/Private/RayTracing/RayTracingReflections.usf
@@ -79,7 +79,7 @@ uint SortTileSize;
 uint2 RayTracingResolution;
 uint2 TileAlignedResolution;
 uint TestPathRoughness;
-uint MinClearCoatLevel;
+float MinClearCoatLevel;
 uint MaxUnderCoatBounces;
 
 float ReflectionMinRayDistance;
