Fix merge error that broke Multiview on mobile HMD for GLES.

diff --git a/Engine/Shaders/Private/MobileBasePassVertexShader.usf b/Engine/Shaders/Private/MobileBasePassVertexShader.usf
--- a/Engine/Shaders/Private/MobileBasePassVertexShader.usf
+++ b/Engine/Shaders/Private/MobileBasePassVertexShader.usf
@@ -42,6 +42,9 @@ void Main(
 {
 #if MOBILE_MULTI_VIEW
 #if (COMPILER_GLSL_ES2 || COMPILER_GLSL_ES3_1 || COMPILER_GLSL_ES3_1_EXT)
+	const int MultiViewId = int(ViewId);
+	ResolvedView = ResolveView(uint(MultiViewId));
+	Output.BasePassInterpolants.MultiViewId = float(MultiViewId);
 #else
 	ResolvedView = ResolveView(ViewId);
 	Output.BasePassInterpolants.MultiViewId = float(ViewId);
