Fixed Noise Node fails to compile with voronoi noise if Ray Tracing is Enabled

diff --git a/Engine/Shaders/Private/Random.ush b/Engine/Shaders/Private/Random.ush
--- a/Engine/Shaders/Private/Random.ush
+++ b/Engine/Shaders/Private/Random.ush
@@ -659,11 +659,11 @@ float4 VoronoiNoise3D_ALU(float3 v, int Quality, bool bTiling, float RepeatSize,
 	// everybody else searches a base 2x2x2 neighborhood
 	else
 	{
-		UNROLL for (offset.x = 0; offset.x <= 1; ++offset.x)
+		UNROLL_N(3) for (offset.x = 0; offset.x <= 1; ++offset.x)
 		{
-			UNROLL for (offset.y = 0; offset.y <= 1; ++offset.y)
+			UNROLL_N(3) for (offset.y = 0; offset.y <= 1; ++offset.y)
 			{
-				UNROLL for (offset.z = 0; offset.z <= 1; ++offset.z)
+				UNROLL_N(3) for (offset.z = 0; offset.z <= 1; ++offset.z)
 				{
 					p = offset + VoronoiCornerSample(NoiseTileWrap(iv + offset, bTiling, RepeatSize), Quality);
 					mindist = VoronoiCompare(mindist, iv + p, fv - p, bDistanceOnly);
@@ -683,11 +683,11 @@ float4 VoronoiNoise3D_ALU(float3 v, int Quality, bool bTiling, float RepeatSize,
 	// quality level 4: add extra sets of four cells in each direction
 	if (Quality >= 4)
 	{
-		UNROLL for (offset.x = -1; offset.x <= 2; offset.x += 3)
+		UNROLL_N(3) for (offset.x = -1; offset.x <= 2; offset.x += 3)
 		{
-			UNROLL for (offset.y = 0; offset.y <= 1; ++offset.y)
+			UNROLL_N(3) for (offset.y = 0; offset.y <= 1; ++offset.y)
 			{
-				UNROLL for (offset.z = 0; offset.z <= 1; ++offset.z)
+				UNROLL_N(3) for (offset.z = 0; offset.z <= 1; ++offset.z)
 				{
 					// along x axis
 					p = offset.xyz + VoronoiCornerSample(NoiseTileWrap(iv + offset.xyz, bTiling, RepeatSize), Quality);
