Fixed refelction capture not using correct index.

diff --git a/Engine/Shaders/Private/SceneData.ush b/Engine/Shaders/Private/SceneData.ush
--- a/Engine/Shaders/Private/SceneData.ush
+++ b/Engine/Shaders/Private/SceneData.ush
@@ -96,7 +96,7 @@ FPrimitiveSceneData GetPrimitiveData(uint PrimitiveId)
 	PrimitiveData.LightmapDataIndex = asuint(View.PrimitiveSceneData[PrimitiveBaseOffset + 24].w);
 
 	PrimitiveData.PreSkinnedLocalBoundsMin = View.PrimitiveSceneData[PrimitiveBaseOffset + 25].xyz;
-	PrimitiveData.SingleCaptureIndex = View.PrimitiveSceneData[PrimitiveBaseOffset + 25].w;
+	PrimitiveData.SingleCaptureIndex = asuint(View.PrimitiveSceneData[PrimitiveBaseOffset + 25].w);
 
 	PrimitiveData.PreSkinnedLocalBoundsMax = View.PrimitiveSceneData[PrimitiveBaseOffset + 26].xyz;
 	PrimitiveData.OutputVelocity = asuint(View.PrimitiveSceneData[PrimitiveBaseOffset + 26].w);
