Fixed: Global illumination Postprocess settings doesn't affect Reflection on Mobile

diff --git a/Engine/Shaders/Private/MobileBasePassPixelShader.usf b/Engine/Shaders/Private/MobileBasePassPixelShader.usf
--- a/Engine/Shaders/Private/MobileBasePassPixelShader.usf
+++ b/Engine/Shaders/Private/MobileBasePassPixelShader.usf
@@ -281,7 +281,9 @@ half3 GetImageBasedReflectionLighting(FMaterialPixelParameters MaterialParameter
 	}
 
 #endif
-	
+
+	SpecularIBL*= View.IndirectLightingColorScale;
+		
 #if WEBGL
 	// need a rgb swizzle instead of the existing rgba swizzle, we should add it if another use case comes up. 
 	return SpecularIBL.bgr;
