Fixed ray traced reflections on bottom layer of clearcoat materials when cubemap fallback is used.

Mirror reflection ray is now used for CompositeReflectionCapturesAndSkylight(), which looks up a prefiltered cubemap. Fix by [REDACTED_PERSON].

diff --git a/Engine/Shaders/Private/RayTracing/RayTracingReflections.usf b/Engine/Shaders/Private/RayTracing/RayTracingReflections.usf
--- a/Engine/Shaders/Private/RayTracing/RayTracingReflections.usf
+++ b/Engine/Shaders/Private/RayTracing/RayTracingReflections.usf
@@ -438,7 +438,10 @@ void RayTracingReflectionsRGS()
 						uint NumCulledReflectionCaptures = ForwardLightData.NumReflectionCaptures;
 						uint ReflectionCapturesStartIndex = 0;
 
-						float3 R = reflect(BottomLayerRay.Direction, BottomLayerWorldNormal);
+						// Compute a perfect mirror reflection direction to use for prefiltered cubemap lookup.
+						// Note that CompositeReflectionCapturesAndSkylight() uses BottomLayerRoughness to compute the cubemap mip level.
+						float3 R = reflect(IncidentDirection, BottomLayerWorldNormal);
+
 						const float NoV = saturate(dot(-BottomLayerRay.Direction, BottomLayerWorldNormal));
 						const float AO = 1.0f;
 						const float RoughnessSq = BottomLayerRoughness * BottomLayerRoughness;
