Fixed mismatch between sorted and non-sorted ray traced reflections due to different handling of back faces. Now both modes disable back face culling and look consistent.

diff --git a/Engine/Shaders/Private/RayTracing/RayTracingReflections.usf b/Engine/Shaders/Private/RayTracing/RayTracingReflections.usf
--- a/Engine/Shaders/Private/RayTracing/RayTracingReflections.usf
+++ b/Engine/Shaders/Private/RayTracing/RayTracingReflections.usf
@@ -295,7 +295,7 @@ void RayTracingReflectionsRGS()
 				{
 					if (DIM_DEFERRED_MATERIAL_MODE == DEFERRED_MATERIAL_MODE_GATHER)
 					{
-						uint RayFlags = RAY_FLAG_CULL_BACK_FACING_TRIANGLES;
+						uint RayFlags = 0;
 						DeferredMaterialPayload.SortKey = RAY_TRACING_DEFERRED_MATERIAL_KEY_RAY_MISS;
 						TraceRay(
 							TLAS,								// AccelerationStructure
