Fixed bug in bloom calculation causing mobile to render overly bright

diff --git a/Engine/Shaders/Private/PostProcessMobile.usf b/Engine/Shaders/Private/PostProcessMobile.usf
--- a/Engine/Shaders/Private/PostProcessMobile.usf
+++ b/Engine/Shaders/Private/PostProcessMobile.usf
@@ -994,7 +994,12 @@ void SunMergePS_ES2(
 
 		// Have 5 layers on mobile.
 		float Scale3 = 1.0/5.0;	
-		OutColor.rgb+= (Bloom2 * Scale3 * BloomColor);
+
+		// scale existing color first
+		OutColor.rgb *= Scale3;		
+
+		// add scaled bloom separately to prevent overflow before scaling
+		OutColor.rgb += (Bloom2 * Scale3 * BloomColor);
 
 	#else
 		OutColor.rgb = half3(0.0, 0.0, 0.0);
