Cherry-pick CL 8003124

Fix hair shading artifacts in ray tracing SkyLight pass.

diff --git a/Engine/Shaders/Private/AmbientCubemapComposite.usf b/Engine/Shaders/Private/AmbientCubemapComposite.usf
--- a/Engine/Shaders/Private/AmbientCubemapComposite.usf
+++ b/Engine/Shaders/Private/AmbientCubemapComposite.usf
@@ -310,7 +310,8 @@ float3 ImageBasedLightingHair( FGBufferData GBuffer, float3 V, float3 N, uint2 R
 			float Weight = rcp( InvWeight );
 
 			float3 Shading = 0;
-			Shading = HairShading( GBuffer, L, V, N, 1, 0, 0, Random );
+			bool bEvalMultiScatter = true;
+			Shading = HairShading( GBuffer, L, V, N, 1, 0, 0, Random, bEvalMultiScatter );
 			
 			Lighting += SampleColor * Shading * Weight;
 		}
@@ -424,7 +425,8 @@ void MainPS(in noperspective float4 UVAndScreenPos : TEXCOORD0, float4 SvPositio
 		float3 FakeNormal = normalize( V - N * dot(V,N) );
 		FakeNormal = normalize( FakeNormal + 0 * float3(0,0,1) );
 		SpecularContribution = TextureCubeSampleLevel(AmbientCubemap, AmbientCubemapSampler, FakeNormal, AbsoluteDiffuseMip).rgb;
-		SpecularContribution *= PI * HairShading( GBuffer, FakeNormal, V, N, 1, 0, 0.2, Random );
+		bool bEvalMultiScatter = true;
+		SpecularContribution *= PI * HairShading( GBuffer, FakeNormal, V, N, 1, 0, 0.2, Random, bEvalMultiScatter );
 		NonSpecularContribution = 0;
 	}
 
diff --git a/Engine/Shaders/Private/BasePassPixelShader.usf b/Engine/Shaders/Private/BasePassPixelShader.usf
--- a/Engine/Shaders/Private/BasePassPixelShader.usf
+++ b/Engine/Shaders/Private/BasePassPixelShader.usf
@@ -1030,7 +1030,8 @@ ShadingModel);
 			float3 V = MaterialParameters.CameraVector;
 			float3 L = normalize( V - N * dot(V,N) );
 			DiffuseDir = L;
-			DiffuseColorForIndirect = 2*PI * HairShading( GBuffer, L, V, N, 1, 0, 0.2, uint2(0,0) );
+			bool bEvalMultiScatter = true;
+			DiffuseColorForIndirect = 2*PI * HairShading( GBuffer, L, V, N, 1, 0, 0.2, uint2(0,0), bEvalMultiScatter );
 		}
 		#endif
 
diff --git a/Engine/Shaders/Private/RayTracing/RayTracingSkyLightRGS.usf b/Engine/Shaders/Private/RayTracing/RayTracingSkyLightRGS.usf
--- a/Engine/Shaders/Private/RayTracing/RayTracingSkyLightRGS.usf
+++ b/Engine/Shaders/Private/RayTracing/RayTracingSkyLightRGS.usf
@@ -8,6 +8,7 @@
 #include "SkyLightMipTreeCommon.ush"
 #include "RayTracingCommon.ush"
 #include "RayTracingSkyLightCommon.ush"
+#include "../ReflectionEnvironmentShared.ush"
 #include "../SceneTextureParameters.ush"
 
 struct FSkyLightVisibilityRays
@@ -82,6 +83,7 @@ void SkyLightRGS()
 	float3 BlueNoiseShift = EvalBlueNoise(QuasiRandom.BlueNoiseCoord, Dimension);
 #endif
 
+	float3 BentNormal = 0;
 	float3 ExitantRadiance = 0.0;
 	float3 DiffuseExitantRadiance = 0.0;
 	for (uint SampleIndex = 0; SampleIndex < SamplesPerPixel; ++SampleIndex)
@@ -129,6 +131,13 @@ void SkyLightRGS()
 		RayWeight = MisWeight / RayPdf;
 #endif
 
+		// Recompute WorldNormal if lighting a hair fiber
+		if (GBufferData.ShadingModelID == SHADINGMODELID_HAIR)
+		{
+			const float3 LightDirection = Ray.Direction;
+			const float3 TangentDirection = normalize(WorldNormal);
+			WorldNormal = normalize(LightDirection - TangentDirection * dot(LightDirection, TangentDirection));
+		}
 		ApplyCameraRelativeDepthBias(Ray, PixelCoord, DeviceZ, WorldNormal, SkyLight.MaxNormalBias);
 
 		uint RayFlags = 0;
@@ -159,15 +168,30 @@ void SkyLightRGS()
 		}
 		else
 		{
+			BentNormal += Ray.Direction;
 			// Evaluate material
-			half3 N = WorldNormal;
+			half3 N = GBufferData.WorldNormal;
 			half3 V = -CameraDirection;
 			half3 L = Ray.Direction;
-			float NoL = dot(WorldNormal, Ray.Direction);
-			FShadowTerms ShadowTerms = { 0.0, 0.0, 0.0 };
-			FDirectLighting LightingSample = EvaluateBxDF(GBufferData, N, V, L, NoL, ShadowTerms);
+			float NoL = saturate(dot(WorldNormal, L));
+			FDirectLighting LightingSample;
+			if (GBufferData.ShadingModelID == SHADINGMODELID_HAIR)
+			{
+				float Shadow = 0.0;
+				float Backlit = 0.0;
+				float Area = 0.0;
+				uint2 Random = 0;
+				bool bEvalMultiScatter = false;
+				LightingSample.Diffuse = HairShading(GBufferData, L, V, N, Shadow, Backlit, Area, Random, bEvalMultiScatter);
+				LightingSample.Transmission = 0;
+				LightingSample.Specular = 0;
+			}
+			else
+			{
+				FShadowTerms ShadowTerms = { 0.0, 0.0, 0.0 };
+				LightingSample = EvaluateBxDF(GBufferData, N, V, L, NoL, ShadowTerms);
+			}
 			float3 Brdf = LightingSample.Diffuse + LightingSample.Transmission + LightingSample.Specular;
-
 			float3 IncomingRadiance = EvalSkyLight(Ray.Direction);
 			ExitantRadiance += IncomingRadiance * Brdf * RayWeight;
 			DiffuseExitantRadiance += IncomingRadiance * (LightingSample.Diffuse + LightingSample.Transmission) * RayWeight;
@@ -181,6 +205,54 @@ void SkyLightRGS()
 		ExitantRadiance *= SamplesPerPixelInv;
 		DiffuseExitantRadiance *= SamplesPerPixelInv;
 		AmbientOcclusion /= SamplesPerPixel;
+		BentNormal = normalize(BentNormal);
+	}
+
+	// Add multi-scattering contribution
+	if (GBufferData.ShadingModelID == SHADINGMODELID_HAIR) {
+		half3 N = GBufferData.WorldNormal;
+		half3 V = -CameraDirection;
+		half3 L = BentNormal;
+
+		// Cast ray from light to shading point to get estimation of "shadow thickness"
+		float ShadowThickness = 0.0;
+		{
+			uint RayFlagsLocal = 0;
+#if !ENABLE_MATERIALS
+			RayFlagsLocal |= RAY_FLAG_FORCE_OPAQUE;
+#endif
+#if !ENABLE_TWO_SIDED_GEOMETRY
+			RayFlagsLocal |= RAY_FLAG_CULL_BACK_FACING_TRIANGLES;
+#endif
+			FPackedMaterialClosestHitPayload Payload = (FPackedMaterialClosestHitPayload)0;
+			Payload.SetMinimalPayloadMode(); // Only HitT will be filled by closest hit shader, skipping full material evaluation.
+
+			RayDesc RayLocal;
+			RayLocal.Origin = WorldPosition + L * SkyLight.MaxShadowThickness;
+			RayLocal.Direction = -L;
+			RayLocal.TMin = 0.0;
+			RayLocal.TMax = SkyLight.MaxShadowThickness;
+
+			FPackedMaterialClosestHitPayload PayloadLocal = (FPackedMaterialClosestHitPayload)0;
+			PayloadLocal.SetMinimalPayloadMode(); // Only HitT will be filled by closest hit shader, skipping full material evaluation.
+			TraceRay(
+				TLAS,      // AccelerationStructure
+				RayFlagsLocal,
+				RAY_TRACING_MASK_SHADOW,          // InstanceInclusionMask
+				RAY_TRACING_SHADER_SLOT_SHADOW,   // RayContributionToHitGroupIndex
+				RAY_TRACING_NUM_SHADER_SLOTS,     // MultiplierForGeometryContributionToShaderIndex
+				0,         // MissShaderIndex
+				RayLocal,       // RayDesc
+				PayloadLocal    // Payload
+			);
+
+			float3 P = RayLocal.Origin + RayLocal.Direction * PayloadLocal.HitT;
+			ShadowThickness = min(length(WorldPosition - P), 1.0);
+		}
+		float3 ScatterThroughput = PI * KajiyaKayDiffuseAttenuation(GBufferData, L, V, N, ShadowThickness);
+
+		float3 DiffuseRadiance = GetSkySHDiffuse(L) * View.SkyLightColor.rgb;
+		DiffuseExitantRadiance += DiffuseRadiance * ScatterThroughput * (1.0 - AmbientOcclusion);
 	}
 
 	// Pre-divide by albedo, to be recovered in compositing
diff --git a/Engine/Shaders/Private/ReflectionEnvironmentPixelShader.usf b/Engine/Shaders/Private/ReflectionEnvironmentPixelShader.usf
--- a/Engine/Shaders/Private/ReflectionEnvironmentPixelShader.usf
+++ b/Engine/Shaders/Private/ReflectionEnvironmentPixelShader.usf
@@ -143,7 +143,8 @@ float3 SkyLightDiffuse(FGBufferData GBuffer, float AmbientOcclusion, float2 Buff
 		float3 V = -normalize(mul(float4(ScreenPosition, 1, 0), View.ScreenToWorld).xyz);
 		float3 L = normalize(V - N * dot(V, N));
 		SkyLightingNormal = L;
-		DiffuseColor = PI * HairShading(GBuffer, L, V, N, 1, 0, 0.2, uint2(0, 0));
+		bool bEvalMultiScatter = true;
+		DiffuseColor = PI * HairShading(GBuffer, L, V, N, 1, 0, 0.2, uint2(0, 0), bEvalMultiScatter);
 	}
 
 	if (GBuffer.ShadingModelID == SHADINGMODELID_CLOTH)
diff --git a/Engine/Shaders/Private/ShadingModels.ush b/Engine/Shaders/Private/ShadingModels.ush
--- a/Engine/Shaders/Private/ShadingModels.ush
+++ b/Engine/Shaders/Private/ShadingModels.ush
@@ -661,10 +661,28 @@ float3 HairShadingRef( FGBufferData GBuffer, float3 L, float3 V, half3 N, uint2
 }
 #endif
 
+float3 KajiyaKayDiffuseAttenuation(FGBufferData GBuffer, float3 L, float3 V, half3 N, float Shadow)
+{
+	// Use soft Kajiya Kay diffuse attenuation
+	float KajiyaDiffuse = 1 - abs(dot(N, L));
+
+	float3 FakeNormal = normalize(V - N * dot(V, N));
+	//N = normalize( DiffuseN + FakeNormal * 2 );
+	N = FakeNormal;
+
+	// Hack approximation for multiple scattering.
+	float Wrap = 1;
+	float NoL = saturate((dot(N, L) + Wrap) / Square(1 + Wrap));
+	float DiffuseScatter = (1 / PI) * lerp(NoL, KajiyaDiffuse, 0.33) * GBuffer.Metallic;
+	float Luma = Luminance(GBuffer.BaseColor);
+	float3 ScatterTint = pow(GBuffer.BaseColor / Luma, 1 - Shadow);
+	return sqrt(GBuffer.BaseColor) * DiffuseScatter * ScatterTint;
+}
+
 // Approximation to HairShadingRef using concepts from the following papers:
 // [Marschner et al. 2003, "Light Scattering from Human Hair Fibers"]
 // [Pekelis et al. 2015, "A Data-Driven Light Scattering Model for Hair"]
-float3 HairShading( FGBufferData GBuffer, float3 L, float3 V, half3 N, float Shadow, float Backlit, float Area, uint2 Random )
+float3 HairShading( FGBufferData GBuffer, float3 L, float3 V, half3 N, float Shadow, float Backlit, float Area, uint2 Random, bool bEvalMultiScatter )
 {
 	// to prevent NaN with decals
 	// OR-18489 HERO: IGGY: RMB on E ability causes blinding hair effect
@@ -776,22 +794,9 @@ float3 HairShading( FGBufferData GBuffer, float3 L, float3 V, half3 N, float Sha
 	}
 #endif
 
-	if(1)
+	if (bEvalMultiScatter)
 	{
-		// Use soft Kajiya Kay diffuse attenuation
-		float KajiyaDiffuse = 1 - abs( dot(N,L) );
-
-		float3 FakeNormal = normalize( V - N * dot(V,N) );
-		//N = normalize( DiffuseN + FakeNormal * 2 );
-		N = FakeNormal;
-
-		// Hack approximation for multiple scattering.
-		float Wrap = 1;
-		float NoL = saturate( ( dot(N, L) + Wrap ) / Square( 1 + Wrap ) );
-		float DiffuseScatter = (1 / PI) * lerp( NoL, KajiyaDiffuse, 0.33 ) * GBuffer.Metallic;
-		float Luma = Luminance( GBuffer.BaseColor );
-		float3 ScatterTint = pow( GBuffer.BaseColor / Luma, 1 - Shadow );
-		S += sqrt( GBuffer.BaseColor ) * DiffuseScatter * ScatterTint;
+		S += KajiyaKayDiffuseAttenuation(GBuffer, L, V, N, Shadow);
 	}
 
 	S = -min(-S, 0.0);
@@ -804,7 +809,8 @@ FDirectLighting HairBxDF( FGBufferData GBuffer, half3 N, half3 V, half3 L, float
 	FDirectLighting Lighting;
 	Lighting.Diffuse = 0;
 	Lighting.Specular = 0;
-	Lighting.Transmission = AreaLight.FalloffColor * Falloff * HairShading( GBuffer, L, V, N, Shadow.TransmissionShadow, 1, 0, uint2(0,0) );
+	bool bEvalMultiScatter = true;
+	Lighting.Transmission = AreaLight.FalloffColor * Falloff * HairShading( GBuffer, L, V, N, Shadow.TransmissionShadow, 1, 0, uint2(0,0), bEvalMultiScatter );
 	return Lighting;
 }
 
diff --git a/Engine/Source/Runtime/Renderer/Private/RayTracing/RayTracingSkyLight.h b/Engine/Source/Runtime/Renderer/Private/RayTracing/RayTracingSkyLight.h
--- a/Engine/Source/Runtime/Renderer/Private/RayTracing/RayTracingSkyLight.h
+++ b/Engine/Source/Runtime/Renderer/Private/RayTracing/RayTracingSkyLight.h
@@ -15,6 +15,7 @@ SHADER_PARAMETER(float, MaxRayDistance)
 SHADER_PARAMETER(FVector, Color)
 SHADER_PARAMETER(FIntVector, MipDimensions)
 SHADER_PARAMETER(float, MaxNormalBias)
+SHADER_PARAMETER(float, MaxShadowThickness)
 SHADER_PARAMETER_TEXTURE(TextureCube, Texture)
 SHADER_PARAMETER_SAMPLER(SamplerState, TextureSampler)
 SHADER_PARAMETER(FIntVector, TextureDimensions)
diff --git a/Engine/Source/Runtime/Renderer/Private/RayTracing/RaytracingSkylight.cpp b/Engine/Source/Runtime/Renderer/Private/RayTracing/RaytracingSkylight.cpp
--- a/Engine/Source/Runtime/Renderer/Private/RayTracing/RaytracingSkylight.cpp
+++ b/Engine/Source/Runtime/Renderer/Private/RayTracing/RaytracingSkylight.cpp
@@ -47,7 +47,14 @@ static float GRayTracingSkyLightMaxRayDistance = 1.0e7;
 static FAutoConsoleVariableRef CVarRayTracingSkyLightMaxRayDistance(
 	TEXT("r.RayTracing.SkyLight.MaxRayDistance"),
 	GRayTracingSkyLightMaxRayDistance,
-	TEXT("Sets the samples-per-pixel for ray tracing SkyLight (default = 1.0e7)")
+	TEXT("Sets the max ray distance for ray tracing SkyLight (default = 1.0e7)")
+);
+
+static float GRayTracingSkyLightMaxShadowThickness = 1.0e3;
+static FAutoConsoleVariableRef CVarRayTracingSkyLightMaxShadowThickness(
+	TEXT("r.RayTracing.SkyLight.MaxShadowThickness"),
+	GRayTracingSkyLightMaxShadowThickness,
+	TEXT("Sets the max shadow thickness for translucent materials for ray tracing SkyLight (default = 1.0e3)")
 );
 
 static int32 GRayTracingSkyLightSamplingStopLevel = 0;
@@ -120,6 +127,7 @@ void SetupSkyLightParameters(
 	SkyLightData->SamplesPerPixel = -1;
 	SkyLightData->SamplingStopLevel = 0;
 	SkyLightData->MaxRayDistance = 1.0e27;
+	SkyLightData->MaxShadowThickness = 1.0e3;
 	SkyLightData->MaxNormalBias = GetRaytracingMaxNormalBias();
 
 	if (Scene.SkyLight && Scene.SkyLight->ProcessedTexture)
@@ -859,6 +867,7 @@ void FDeferredShadingSceneRenderer::GenerateSkyLightVisibilityRays(
 	uint32 SamplesPerPixel = GRayTracingSkyLightSamplesPerPixel >= 0 ? GRayTracingSkyLightSamplesPerPixel : Scene->SkyLight->SamplesPerPixel;
 	SkyLightData.SamplesPerPixel = SamplesPerPixel;
 	SkyLightData.MaxRayDistance = GRayTracingSkyLightMaxRayDistance;
+	SkyLightData.MaxShadowThickness = GRayTracingSkyLightMaxShadowThickness;
 	SkyLightData.SamplingStopLevel = GRayTracingSkyLightSamplingStopLevel;
 
 	// Sampling state
