Sky: always snap the camera at the planet surface to make sure sky is always validEven when the camera is under the planet ground.

diff --git a/Engine/Shaders/Private/SkyAtmosphere.usf b/Engine/Shaders/Private/SkyAtmosphere.usf
--- a/Engine/Shaders/Private/SkyAtmosphere.usf
+++ b/Engine/Shaders/Private/SkyAtmosphere.usf
@@ -55,7 +55,7 @@
 #define VIEWDATA_AVAILABLE (TRANSMITTANCE_PASS!=1 && MULTISCATT_PASS!=1 && SKYLIGHT_PASS!=1)
 
 
-
+// Float accuracy offset in Sky unit (km, so this is 1m)
 #define PLANET_RADIUS_OFFSET 0.001f
 
  // FASTSKY mapping is done based on Light0 as main light
@@ -112,7 +112,15 @@ float3 GetScreenWorldDir(in float4 SVPos)
 
 float3 GetCameraWorldPos()
 {
-	return View.WorldCameraOrigin;
+	const float Offset = PLANET_RADIUS_OFFSET * SKY_UNIT_TO_CM;
+	const float BottomRadiusWorld = Atmosphere.BottomRadius * SKY_UNIT_TO_CM;
+	const float3 PlanetCenterWorld = float3(0.0f, 0.0f, -BottomRadiusWorld);
+	const float3 PlanetCenterToCameraWorld = View.WorldCameraOrigin - PlanetCenterWorld;
+	const float DistanceToPlanetCenterWorld = length(PlanetCenterToCameraWorld);
+
+	// If the camera is below the planet surface, we snap it back onto the surface.
+	// This is to make sure the sky is always visible even if the camera is inside the virtual planet.
+	return DistanceToPlanetCenterWorld < (BottomRadiusWorld + Offset) ? PlanetCenterWorld + (BottomRadiusWorld + Offset) * (PlanetCenterToCameraWorld / DistanceToPlanetCenterWorld) : View.WorldCameraOrigin;
 }
 
 float3 GetCameraPlanetPos()
@@ -435,7 +443,9 @@ SingleScatteringResult IntegrateSingleScatteredLuminance(
 	Result.MultiScatAs1 = 0;
 
 	if (dot(WorldPos, WorldPos) <= Atmosphere.BottomRadius*Atmosphere.BottomRadius)
+	{
 		return Result;	// Camera is inside the atmosphere
+	}
 
 	float2 PixPos = SVPos.xy;
 
diff --git a/Engine/Shaders/Private/SkyAtmosphereCommon.ush b/Engine/Shaders/Private/SkyAtmosphereCommon.ush
--- a/Engine/Shaders/Private/SkyAtmosphereCommon.ush
+++ b/Engine/Shaders/Private/SkyAtmosphereCommon.ush
@@ -12,6 +12,7 @@
 
 // Kilometers as unit for computations related to the sky and its atmosphere
 #define CM_TO_SKY_UNIT 0.00001f
+#define SKY_UNIT_TO_CM (1.0f/CM_TO_SKY_UNIT)
 
 
 
