Improved CL 8123572 by doing math on CPU...
Also fixes the material nodes for the case when the camera is under the planet ground.

diff --git a/Engine/Shaders/Private/MaterialTemplate.ush b/Engine/Shaders/Private/MaterialTemplate.ush
--- a/Engine/Shaders/Private/MaterialTemplate.ush
+++ b/Engine/Shaders/Private/MaterialTemplate.ush
@@ -999,7 +999,7 @@ float3 MaterialExpressionSkyAtmosphereLightDiskLuminance(FMaterialPixelParameter
 {
 #if MATERIAL_SKY_ATMOSPHERE && PROJECT_SUPPORT_SKY_ATMOSPHERE
 	const float3 PlanetCenter = float3(0, 0, -ResolvedView.SkyAtmosphereBottomRadius);
-	const float3 PlanetCenterToWorldCameraPos = ResolvedView.WorldCameraOrigin * CM_TO_SKY_UNIT - PlanetCenter;
+	const float3 PlanetCenterToWorldCameraPos = ResolvedView.SkyWorldCameraOrigin * CM_TO_SKY_UNIT - PlanetCenter;
 	const float3 ViewDir = -Parameters.CameraVector;
 
 	return GetLightDiskLuminance(PlanetCenterToWorldCameraPos, ViewDir, ResolvedView.SkyAtmosphereBottomRadius, ResolvedView.SkyAtmosphereTopRadius,
@@ -1014,7 +1014,7 @@ float3 MaterialExpressionSkyAtmosphereViewLuminance(FMaterialPixelParameters Par
 {
 #if MATERIAL_SKY_ATMOSPHERE && PROJECT_SUPPORT_SKY_ATMOSPHERE
 	const float3 PlanetCenter = float3(0, 0, -ResolvedView.SkyAtmosphereBottomRadius);
-	const float3 PlanetCenterToWorldCameraPos = ResolvedView.WorldCameraOrigin * CM_TO_SKY_UNIT - PlanetCenter;
+	const float3 PlanetCenterToWorldCameraPos = ResolvedView.SkyWorldCameraOrigin * CM_TO_SKY_UNIT - PlanetCenter;
 	const float ViewHeight = length(PlanetCenterToWorldCameraPos);
 	const float3 ViewDir = -Parameters.CameraVector;
 
@@ -1048,7 +1048,7 @@ float4 MaterialExpressionSkyAtmosphereAerialPerspective(FMaterialPixelParameters
 
 	const float3 PlanetCenter = float3(0, 0, -ResolvedView.SkyAtmosphereBottomRadius);
 	const float3 WorldPos = Parameters.AbsoluteWorldPosition.xyz*CM_TO_SKY_UNIT;
-	const float3 CameraPos = ResolvedView.WorldCameraOrigin.xyz*CM_TO_SKY_UNIT;
+	const float3 CameraPos = ResolvedView.SkyWorldCameraOrigin.xyz*CM_TO_SKY_UNIT;
 
 	float4 NDCPosition = mul(float4(Parameters.AbsoluteWorldPosition.xyz, 1), ResolvedView.WorldToClip); // Because SvPosition seems different...
 
diff --git a/Engine/Shaders/Private/SkyAtmosphere.usf b/Engine/Shaders/Private/SkyAtmosphere.usf
--- a/Engine/Shaders/Private/SkyAtmosphere.usf
+++ b/Engine/Shaders/Private/SkyAtmosphere.usf
@@ -54,10 +54,6 @@
 // View data is not available for passes running once per scene (and not once per view).
 #define VIEWDATA_AVAILABLE (TRANSMITTANCE_PASS!=1 && MULTISCATT_PASS!=1 && SKYLIGHT_PASS!=1)
 
-
-// Float accuracy offset in Sky unit (km, so this is 1m)
-#define PLANET_RADIUS_OFFSET 0.001f
-
  // FASTSKY mapping is done based on Light0 as main light
 #define FASTSKY_LIGHT_INDEX 0
 
@@ -112,15 +108,7 @@ float3 GetScreenWorldDir(in float4 SVPos)
 
 float3 GetCameraWorldPos()
 {
-	const float Offset = PLANET_RADIUS_OFFSET * SKY_UNIT_TO_CM;
-	const float BottomRadiusWorld = Atmosphere.BottomRadius * SKY_UNIT_TO_CM;
-	const float3 PlanetCenterWorld = float3(0.0f, 0.0f, -BottomRadiusWorld);
-	const float3 PlanetCenterToCameraWorld = View.WorldCameraOrigin - PlanetCenterWorld;
-	const float DistanceToPlanetCenterWorld = length(PlanetCenterToCameraWorld);
-
-	// If the camera is below the planet surface, we snap it back onto the surface.
-	// This is to make sure the sky is always visible even if the camera is inside the virtual planet.
-	return DistanceToPlanetCenterWorld < (BottomRadiusWorld + Offset) ? PlanetCenterWorld + (BottomRadiusWorld + Offset) * (PlanetCenterToCameraWorld / DistanceToPlanetCenterWorld) : View.WorldCameraOrigin;
+	return View.SkyWorldCameraOrigin;
 }
 
 float3 GetCameraPlanetPos()
diff --git a/Engine/Shaders/Private/SkyAtmosphereCommon.ush b/Engine/Shaders/Private/SkyAtmosphereCommon.ush
--- a/Engine/Shaders/Private/SkyAtmosphereCommon.ush
+++ b/Engine/Shaders/Private/SkyAtmosphereCommon.ush
@@ -9,11 +9,12 @@
 #define D_SKY_ATMOSPHERE_COMMON
 
 
-
+// The constants below should match the one in SceneRendering.cpp
 // Kilometers as unit for computations related to the sky and its atmosphere
 #define CM_TO_SKY_UNIT 0.00001f
 #define SKY_UNIT_TO_CM (1.0f/CM_TO_SKY_UNIT)
-
+// Float accuracy offset in Sky unit (km, so this is 1m)
+#define PLANET_RADIUS_OFFSET 0.001f
 
 
 // The number of killometer per slice in the aerial pespective camera volume texture. (assuming a uniform depth distribution)
diff --git a/Engine/Source/Runtime/Engine/Public/SceneView.h b/Engine/Source/Runtime/Engine/Public/SceneView.h
--- a/Engine/Source/Runtime/Engine/Public/SceneView.h
+++ b/Engine/Source/Runtime/Engine/Public/SceneView.h
@@ -670,6 +670,7 @@ enum ETranslucencyVolumeCascade
 	VIEW_UNIFORM_BUFFER_MEMBER_ARRAY(FLinearColor, AtmosphereLightDiscLuminance, [NUM_ATMOSPHERE_LIGHTS]) \
 	VIEW_UNIFORM_BUFFER_MEMBER_ARRAY(FVector4, AtmosphereLightDiscCosHalfApexAngle, [NUM_ATMOSPHERE_LIGHTS]) \
 	VIEW_UNIFORM_BUFFER_MEMBER(FVector4, SkyViewLutSizeAndInvSize) \
+	VIEW_UNIFORM_BUFFER_MEMBER(FVector, SkyWorldCameraOrigin) \
 	VIEW_UNIFORM_BUFFER_MEMBER(float, SkyAtmosphereBottomRadius) \
 	VIEW_UNIFORM_BUFFER_MEMBER(float, SkyAtmosphereTopRadius) \
 	VIEW_UNIFORM_BUFFER_MEMBER(float, SkyAtmosphereAerialPerspectiveStartDepth) \
diff --git a/Engine/Source/Runtime/Renderer/Private/SceneRendering.cpp b/Engine/Source/Runtime/Renderer/Private/SceneRendering.cpp
--- a/Engine/Source/Runtime/Renderer/Private/SceneRendering.cpp
+++ b/Engine/Source/Runtime/Renderer/Private/SceneRendering.cpp
@@ -1256,6 +1256,20 @@ void FViewInfo::SetupUniformBufferParameters(
 				ViewUniformShaderParameters.AtmosphereLightDirection[Index] = DefaultSunDirection;
 			}
 		}
+
+		// The constants below should match the one in SkyAtmosphereCommon.ush
+		const float SkyUnitToCm = 1.0f / 0.00001f;
+		const float PlanetRadiusOffset = 0.001f;
+
+		const float Offset = PlanetRadiusOffset * SkyUnitToCm;
+		const float BottomRadiusWorld = AtmosphereSetup.BottomRadius * SkyUnitToCm;
+		const FVector PlanetCenterWorld = FVector(0.0f, 0.0f, -BottomRadiusWorld);
+		const FVector PlanetCenterToCameraWorld = ViewUniformShaderParameters.WorldCameraOrigin - PlanetCenterWorld;
+		const float DistanceToPlanetCenterWorld = PlanetCenterToCameraWorld.Size();
+
+		// If the camera is below the planet surface, we snap it back onto the surface.
+		// This is to make sure the sky is always visible even if the camera is inside the virtual planet.
+		ViewUniformShaderParameters.SkyWorldCameraOrigin = DistanceToPlanetCenterWorld < (BottomRadiusWorld + Offset) ? PlanetCenterWorld + (BottomRadiusWorld + Offset) * (PlanetCenterToCameraWorld / DistanceToPlanetCenterWorld) : ViewUniformShaderParameters.WorldCameraOrigin;
 	}
 	else
 	{
@@ -1268,6 +1282,7 @@ void FViewInfo::SetupUniformBufferParameters(
 		ViewUniformShaderParameters.SkyAtmosphereCameraAerialPerspectiveVolumeDepthSliceLength = 1.0f;
 		ViewUniformShaderParameters.SkyAtmosphereCameraAerialPerspectiveVolumeDepthSliceLengthInv = 1.0f;
 		ViewUniformShaderParameters.SkyAtmosphereApplyCameraAerialPerspectiveVolume = 1.0f;
+		ViewUniformShaderParameters.SkyWorldCameraOrigin = ViewUniformShaderParameters.WorldCameraOrigin;
 	}
 
 	ViewUniformShaderParameters.TransmittanceLutTexture = OrWhite2DIfNull(TransmittanceLutTextureFound);
