Fix runtime virtual texture normal sampling from BC5 layout
#rnx

diff --git a/Engine/Shaders/Private/VirtualTextureCommon.ush b/Engine/Shaders/Private/VirtualTextureCommon.ush
--- a/Engine/Shaders/Private/VirtualTextureCommon.ush
+++ b/Engine/Shaders/Private/VirtualTextureCommon.ush
@@ -557,7 +557,7 @@ float3 VirtualTextureUnpackNormalBC3(in float4 PackedValue)
 /** Unpack normal from BC5. */
 float3 VirtualTextureUnpackNormalBC5(in float4 PackedValue)
 {
-	return VirtualTextureUnpackNormal(PackedValue.rb);
+	return VirtualTextureUnpackNormal(PackedValue.rg);
 }
 
 /** Unpack 16 bit height (with some hardcoded range scale/bias). */
