Added possibility to link the height fog lighting to the sky. Need to be enabled per project since it takes an additional ALUs and vertex shader slot.

diff --git a/Engine/Shaders/Private/HeightFogCommon.ush b/Engine/Shaders/Private/HeightFogCommon.ush
--- a/Engine/Shaders/Private/HeightFogCommon.ush
+++ b/Engine/Shaders/Private/HeightFogCommon.ush
@@ -27,7 +27,6 @@ static const float FLT_EPSILON2 = 0.01f;
 float3 ComputeInscatteringColor(float3 CameraToReceiver, float CameraToReceiverLength)
 {
 	half3 Inscattering = FogStruct.ExponentialFogColorParameter.xyz;
-	half3 DirectionalInscattering = 0;
 
 #if SUPPORT_FOG_INSCATTERING_TEXTURE
 	BRANCH
@@ -43,6 +42,10 @@ float3 ComputeInscatteringColor(float3 CameraToReceiver, float CameraToReceiverL
 	}
 #endif
 
+#if PROJECT_SUPPORT_SKY_ATMOSPHERE_AFFECTS_HEIGHFOG
+	Inscattering += Texture2DSampleLevel(View.DistantSkyLightLutTexture, View.DistantSkyLightLutTextureSampler, float2(0.5f, 0.5f), 0.0f).rgb;
+#endif
+
 	return Inscattering;
 }
 
@@ -133,11 +136,23 @@ half4 GetExponentialHeightFog(float3 WorldPositionRelativeToCamera, float Exclud
 	#endif
 	)
 	{
-		float DirectionalInscatteringStartDistance = FogStruct.InscatteringLightDirection.w;
+#if PROJECT_SUPPORT_SKY_ATMOSPHERE_AFFECTS_HEIGHFOG
+		const float UniformPhaseFunction = 1.0f / (4.0f*PI);
+		half3 DirectionalInscatteringColor;
+		half3 DirectionalLightInscattering;
+
+		DirectionalInscatteringColor = FogStruct.DirectionalInscatteringColor.xyz + View.AtmosphereLightColorGlobalPostTransmittance[0].rgb * UniformPhaseFunction;
+		DirectionalLightInscattering = DirectionalInscatteringColor * pow(saturate(dot(CameraToReceiverNormalized, View.AtmosphereLightDirection[0].xyz)), FogStruct.DirectionalInscatteringColor.w);
+
+		DirectionalInscatteringColor = FogStruct.DirectionalInscatteringColor.xyz + View.AtmosphereLightColorGlobalPostTransmittance[1].rgb * UniformPhaseFunction;
+		DirectionalLightInscattering += DirectionalInscatteringColor * pow(saturate(dot(CameraToReceiverNormalized, View.AtmosphereLightDirection[1].xyz)), FogStruct.DirectionalInscatteringColor.w);
+#else
 		// Setup a cosine lobe around the light direction to approximate inscattering from the directional light off of the ambient haze;
 		half3 DirectionalLightInscattering = FogStruct.DirectionalInscatteringColor.xyz * pow(saturate(dot(CameraToReceiverNormalized, FogStruct.InscatteringLightDirection.xyz)), FogStruct.DirectionalInscatteringColor.w);
+#endif
 
 		// Calculate the line integral of the eye ray through the haze, using a special starting distance to limit the inscattering to the distance
+		float DirectionalInscatteringStartDistance = FogStruct.InscatteringLightDirection.w;
 		float DirExponentialHeightLineIntegral = ExponentialHeightLineIntegralShared * max(RayLength - DirectionalInscatteringStartDistance, 0.0f);
 		// Calculate the amount of light that made it through the fog using the transmission equation
 		half DirectionalInscatteringFogFactor = saturate(exp2(-DirExponentialHeightLineIntegral));
diff --git a/Engine/Source/Runtime/Engine/Classes/Engine/RendererSettings.h b/Engine/Source/Runtime/Engine/Classes/Engine/RendererSettings.h
--- a/Engine/Source/Runtime/Engine/Classes/Engine/RendererSettings.h
+++ b/Engine/Source/Runtime/Engine/Classes/Engine/RendererSettings.h
@@ -659,9 +659,20 @@ class ENGINE_API URendererSettings : public UDeveloperSettings
 	*/
 	UPROPERTY(config, EditAnywhere, Category = ShaderPermutationReduction, meta = (
 		ConsoleVariable = "r.SupportSkyAtmosphere", DisplayName = "Support Sky Atmosphere",
+		ToolTip = "The sky atmosphere component requires extra samplers/textures to be bound to apply aerial perspective on transparent surfaces (and all surfaces on mobile via per vertex evaluation).",
 		ConfigRestartRequired = true))
 		uint32 bSupportSkyAtmosphere : 1;
 
+	/**
+	"The sky atmosphere component can light up the height fog but it requires extra samplers/textures to be bound to apply aerial perspective on transparent surfaces (and all surfaces on mobile via per vertex evaluation)."
+	"It requires r.SupportSkyAtmosphere to be true."
+	*/
+	UPROPERTY(config, EditAnywhere, Category = ShaderPermutationReduction, meta = (
+		ConsoleVariable = "r.SupportSkyAtmosphereAffectsHeightFog", DisplayName = "Support Sky Atmosphere Affecting Height Fog",
+		ToolTip = "The sky atmosphere component can light up the height fog but it requires extra samplers/textures to be bound to apply aerial perspective on transparent surfaces (and all surfaces on mobile via per vertex evaluation). It requires r.SupportSkyAtmosphere to be true.",
+		ConfigRestartRequired = true))
+		uint32 bSupportSkyAtmosphereAffectsHeightFog : 1;
+
 	/**
 	"Skincache allows a compute shader to skin once each vertex, save those results into a new buffer and reuse those calculations when later running the depth, base and velocity passes. This also allows opting into the 'recompute tangents' for skinned mesh instance feature. Disabling will reduce the number of shader permutations required per material. Changing this setting requires restarting the editor."
 	*/
diff --git a/Engine/Source/Runtime/Engine/Private/RendererSettings.cpp b/Engine/Source/Runtime/Engine/Private/RendererSettings.cpp
--- a/Engine/Source/Runtime/Engine/Private/RendererSettings.cpp
+++ b/Engine/Source/Runtime/Engine/Private/RendererSettings.cpp
@@ -150,6 +150,14 @@ void URendererSettings::PostEditChangeProperty(FPropertyChangedEvent& PropertyCh
 			}
 		}
 
+		if ((PropertyChangedEvent.Property->GetFName() == GET_MEMBER_NAME_CHECKED(URendererSettings, bSupportSkyAtmosphere)))
+		{
+			if (!bSupportSkyAtmosphere)
+			{
+				bSupportSkyAtmosphereAffectsHeightFog = 0; // Always disable sky affecting height fog if sky is disabled.
+			}
+		}
+
 		ExportValuesToConsoleVariables(PropertyChangedEvent.Property);
 
 		if (PropertyChangedEvent.Property->GetFName() == GET_MEMBER_NAME_CHECKED(URendererSettings, ReflectionCaptureResolution) && 
@@ -170,6 +178,11 @@ bool URendererSettings::CanEditChange(const UProperty* InProperty) const
 		return ParentVal && (!bSupportSkinCacheShaders || !bEnableRayTracing);
 	}
 
+	if ((InProperty->GetFName() == GET_MEMBER_NAME_CHECKED(URendererSettings, bSupportSkyAtmosphereAffectsHeightFog)))
+	{
+		return ParentVal && bSupportSkyAtmosphere;
+	}
+
 	return ParentVal;
 }
 #endif // #if WITH_EDITOR
diff --git a/Engine/Source/Runtime/Engine/Private/ShaderCompiler/ShaderCompiler.cpp b/Engine/Source/Runtime/Engine/Private/ShaderCompiler/ShaderCompiler.cpp
--- a/Engine/Source/Runtime/Engine/Private/ShaderCompiler/ShaderCompiler.cpp
+++ b/Engine/Source/Runtime/Engine/Private/ShaderCompiler/ShaderCompiler.cpp
@@ -3335,9 +3335,16 @@ void GlobalBeginCompileShader(
 		Input.Environment.SetDefine(TEXT("PROJECT_MOBILE_DISABLE_VERTEX_FOG"), CVar ? (CVar->GetInt() != 0) : 0);
 	}
 
+	bool bSupportSkyAtmosphere = false;
 	{
 		static IConsoleVariable* CVar = IConsoleManager::Get().FindConsoleVariable(TEXT("r.SupportSkyAtmosphere"));
-		Input.Environment.SetDefine(TEXT("PROJECT_SUPPORT_SKY_ATMOSPHERE"), CVar ? (CVar->GetInt() != 0) : 0);
+		bSupportSkyAtmosphere = CVar && CVar->GetInt() != 0;
+		Input.Environment.SetDefine(TEXT("PROJECT_SUPPORT_SKY_ATMOSPHERE"), bSupportSkyAtmosphere ? 1 : 0);
+	}
+
+	{
+		static IConsoleVariable* CVar = IConsoleManager::Get().FindConsoleVariable(TEXT("r.SupportSkyAtmosphereAffectsHeightFog"));
+		Input.Environment.SetDefine(TEXT("PROJECT_SUPPORT_SKY_ATMOSPHERE_AFFECTS_HEIGHFOG"), (CVar && bSupportSkyAtmosphere) ? (CVar->GetInt() != 0) : 0);
 	}
 
 	{
diff --git a/Engine/Source/Runtime/Engine/Public/SceneView.h b/Engine/Source/Runtime/Engine/Public/SceneView.h
--- a/Engine/Source/Runtime/Engine/Public/SceneView.h
+++ b/Engine/Source/Runtime/Engine/Public/SceneView.h
@@ -667,6 +667,7 @@ enum ETranslucencyVolumeCascade
 	VIEW_UNIFORM_BUFFER_MEMBER_EX(float, AtmosphericFogSunDiscScale, EShaderPrecisionModifier::Half) \
 	VIEW_UNIFORM_BUFFER_MEMBER_ARRAY(FVector4, AtmosphereLightDirection, [NUM_ATMOSPHERE_LIGHTS]) \
 	VIEW_UNIFORM_BUFFER_MEMBER_ARRAY(FLinearColor, AtmosphereLightColor, [NUM_ATMOSPHERE_LIGHTS]) \
+	VIEW_UNIFORM_BUFFER_MEMBER_ARRAY(FLinearColor, AtmosphereLightColorGlobalPostTransmittance, [NUM_ATMOSPHERE_LIGHTS]) \
 	VIEW_UNIFORM_BUFFER_MEMBER_ARRAY(FLinearColor, AtmosphereLightDiscLuminance, [NUM_ATMOSPHERE_LIGHTS]) \
 	VIEW_UNIFORM_BUFFER_MEMBER_ARRAY(FVector4, AtmosphereLightDiscCosHalfApexAngle, [NUM_ATMOSPHERE_LIGHTS]) \
 	VIEW_UNIFORM_BUFFER_MEMBER(FVector4, SkyViewLutSizeAndInvSize) \
diff --git a/Engine/Source/Runtime/RenderCore/Private/Shader.cpp b/Engine/Source/Runtime/RenderCore/Private/Shader.cpp
--- a/Engine/Source/Runtime/RenderCore/Private/Shader.cpp
+++ b/Engine/Source/Runtime/RenderCore/Private/Shader.cpp
@@ -2392,6 +2392,12 @@ void ShaderMapAppendKeyString(EShaderPlatform Platform, FString& KeyString)
 		if (CVar && CVar->GetValueOnAnyThread() > 0)
 		{
 			KeyString += TEXT("_SKYATM");
+
+			static const auto CVarHeightFog = IConsoleManager::Get().FindTConsoleVariableDataInt(TEXT("r.SupportSkyAtmosphereAffectsHeightFog"));
+			if (CVarHeightFog && CVarHeightFog->GetValueOnAnyThread() > 0)
+			{
+				KeyString += TEXT("_SKYHF");
+			}
 		}
 	}
 
diff --git a/Engine/Source/Runtime/Renderer/Private/SceneRendering.cpp b/Engine/Source/Runtime/Renderer/Private/SceneRendering.cpp
--- a/Engine/Source/Runtime/Renderer/Private/SceneRendering.cpp
+++ b/Engine/Source/Runtime/Renderer/Private/SceneRendering.cpp
@@ -1204,6 +1204,16 @@ void FViewInfo::SetupUniformBufferParameters(
 		ViewUniformShaderParameters.AtmosphereLightDirection[0] = DefaultSunDirection;
 	}
 
+	auto ClearAtmosphereLightData = [&](uint32 Index)
+	{
+		check(Index < NUM_ATMOSPHERE_LIGHTS);
+		ViewUniformShaderParameters.AtmosphereLightDiscCosHalfApexAngle[Index] = FVector4(1.0f);
+		ViewUniformShaderParameters.AtmosphereLightDiscLuminance[Index] = FLinearColor::Black;
+		ViewUniformShaderParameters.AtmosphereLightColor[Index] = FLinearColor::Black;
+		ViewUniformShaderParameters.AtmosphereLightColorGlobalPostTransmittance[Index] = FLinearColor::Black;
+		ViewUniformShaderParameters.AtmosphereLightDirection[Index] = DefaultSunDirection;
+	};
+
 	FRHITexture* TransmittanceLutTextureFound = nullptr;
 	FRHITexture* SkyViewLutTextureFound = nullptr;
 	FRHITexture* CameraAerialPerspectiveVolumeFound = nullptr;
@@ -1246,14 +1256,12 @@ void FViewInfo::SetupUniformBufferParameters(
 				ViewUniformShaderParameters.AtmosphereLightDiscCosHalfApexAngle[Index] = FVector4(FMath::Cos(Light->Proxy->GetSunLightHalfApexAngleRadian()));
 				ViewUniformShaderParameters.AtmosphereLightDiscLuminance[Index] = Light->Proxy->GetOuterSpaceLuminance();
 				ViewUniformShaderParameters.AtmosphereLightColor[Index] = Light->Proxy->GetColor();
+				ViewUniformShaderParameters.AtmosphereLightColorGlobalPostTransmittance[Index] = Light->Proxy->GetColor() * Light->Proxy->GetTransmittanceFactor();
 				ViewUniformShaderParameters.AtmosphereLightDirection[Index] = SkyAtmosphere->GetAtmosphereLightDirection(Index, -Light->Proxy->GetDirection());
 			}
 			else
 			{
-				ViewUniformShaderParameters.AtmosphereLightDiscCosHalfApexAngle[Index] = FVector4(1.0f);
-				ViewUniformShaderParameters.AtmosphereLightDiscLuminance[Index] = FLinearColor::Black;
-				ViewUniformShaderParameters.AtmosphereLightColor[Index] = FLinearColor::Black;
-				ViewUniformShaderParameters.AtmosphereLightDirection[Index] = DefaultSunDirection;
+				ClearAtmosphereLightData(Index);
 			}
 		}
 
@@ -1283,6 +1291,11 @@ void FViewInfo::SetupUniformBufferParameters(
 		ViewUniformShaderParameters.SkyAtmosphereCameraAerialPerspectiveVolumeDepthSliceLengthInv = 1.0f;
 		ViewUniformShaderParameters.SkyAtmosphereApplyCameraAerialPerspectiveVolume = 1.0f;
 		ViewUniformShaderParameters.SkyWorldCameraOrigin = ViewUniformShaderParameters.WorldCameraOrigin;
+
+		for (uint8 Index = 0; Index < NUM_ATMOSPHERE_LIGHTS; ++Index)
+		{
+			ClearAtmosphereLightData(Index);
+		}
 	}
 
 	ViewUniformShaderParameters.TransmittanceLutTexture = OrWhite2DIfNull(TransmittanceLutTextureFound);
diff --git a/Engine/Source/Runtime/Renderer/Private/SkyAtmosphereRendering.cpp b/Engine/Source/Runtime/Renderer/Private/SkyAtmosphereRendering.cpp
--- a/Engine/Source/Runtime/Renderer/Private/SkyAtmosphereRendering.cpp
+++ b/Engine/Source/Runtime/Renderer/Private/SkyAtmosphereRendering.cpp
@@ -31,6 +31,13 @@ static TAutoConsoleVariable<int32> CVarSupportSkyAtmosphere(
 	TEXT("Enables SkyAtmosphere rendering and shader code."),
 	ECVF_ReadOnly | ECVF_RenderThreadSafe);
 
+// The project settings for the sky atmosphere component to affect the height fog (disable runtime and shader code)
+static TAutoConsoleVariable<int32> CVarSupportSkyAtmosphereAffectsHeightFog(
+	TEXT("r.SupportSkyAtmosphereAffectsHeightFog"),
+	0,
+	TEXT("Enables SkyAtmosphere affecting height fog. It requires r.SupportSkyAtmosphere to be true."),
+	ECVF_ReadOnly | ECVF_RenderThreadSafe);
+
 ////////////////////////////////////////////////////////////////////////// Regular sky 
 
 static TAutoConsoleVariable<float> CVarSkyAtmosphereSampleCountMin(
