Fixed secondary atmosphere light showing at Z-Up direction when not set (due to height fog default color setup)

diff --git a/Engine/Shaders/Private/HeightFogCommon.ush b/Engine/Shaders/Private/HeightFogCommon.ush
--- a/Engine/Shaders/Private/HeightFogCommon.ush
+++ b/Engine/Shaders/Private/HeightFogCommon.ush
@@ -141,11 +141,15 @@ half4 GetExponentialHeightFog(float3 WorldPositionRelativeToCamera, float Exclud
 		half3 DirectionalInscatteringColor;
 		half3 DirectionalLightInscattering;
 
+		// No need to test View.AtmosphereLightColorGlobalPostTransmittance[0].a because InscatteringLightDirection.w above is doing the same test already.
 		DirectionalInscatteringColor = FogStruct.DirectionalInscatteringColor.xyz + View.AtmosphereLightColorGlobalPostTransmittance[0].rgb * UniformPhaseFunction;
 		DirectionalLightInscattering = DirectionalInscatteringColor * pow(saturate(dot(CameraToReceiverNormalized, View.AtmosphereLightDirection[0].xyz)), FogStruct.DirectionalInscatteringColor.w);
 
-		DirectionalInscatteringColor = FogStruct.DirectionalInscatteringColor.xyz + View.AtmosphereLightColorGlobalPostTransmittance[1].rgb * UniformPhaseFunction;
-		DirectionalLightInscattering += DirectionalInscatteringColor * pow(saturate(dot(CameraToReceiverNormalized, View.AtmosphereLightDirection[1].xyz)), FogStruct.DirectionalInscatteringColor.w);
+		if (View.AtmosphereLightColorGlobalPostTransmittance[1].a > 0.0f) // Skip DirectionalInscatteringColor on the second light when disabled.
+		{
+			DirectionalInscatteringColor = FogStruct.DirectionalInscatteringColor.xyz + View.AtmosphereLightColorGlobalPostTransmittance[1].rgb * UniformPhaseFunction;
+			DirectionalLightInscattering += DirectionalInscatteringColor * pow(saturate(dot(CameraToReceiverNormalized, View.AtmosphereLightDirection[1].xyz)), FogStruct.DirectionalInscatteringColor.w);
+		}
 #else
 		// Setup a cosine lobe around the light direction to approximate inscattering from the directional light off of the ambient haze;
 		half3 DirectionalLightInscattering = FogStruct.DirectionalInscatteringColor.xyz * pow(saturate(dot(CameraToReceiverNormalized, FogStruct.InscatteringLightDirection.xyz)), FogStruct.DirectionalInscatteringColor.w);
diff --git a/Engine/Source/Runtime/Renderer/Private/SceneRendering.cpp b/Engine/Source/Runtime/Renderer/Private/SceneRendering.cpp
--- a/Engine/Source/Runtime/Renderer/Private/SceneRendering.cpp
+++ b/Engine/Source/Runtime/Renderer/Private/SceneRendering.cpp
@@ -1210,7 +1210,9 @@ void FViewInfo::SetupUniformBufferParameters(
 		ViewUniformShaderParameters.AtmosphereLightDiscCosHalfApexAngle[Index] = FVector4(1.0f);
 		ViewUniformShaderParameters.AtmosphereLightDiscLuminance[Index] = FLinearColor::Black;
 		ViewUniformShaderParameters.AtmosphereLightColor[Index] = FLinearColor::Black;
+		ViewUniformShaderParameters.AtmosphereLightColor[Index].A = 0.0f;
 		ViewUniformShaderParameters.AtmosphereLightColorGlobalPostTransmittance[Index] = FLinearColor::Black;
+		ViewUniformShaderParameters.AtmosphereLightColorGlobalPostTransmittance[Index].A = 0.0f;
 		ViewUniformShaderParameters.AtmosphereLightDirection[Index] = DefaultSunDirection;
 	};
 
@@ -1256,7 +1258,9 @@ void FViewInfo::SetupUniformBufferParameters(
 				ViewUniformShaderParameters.AtmosphereLightDiscCosHalfApexAngle[Index] = FVector4(FMath::Cos(Light->Proxy->GetSunLightHalfApexAngleRadian()));
 				ViewUniformShaderParameters.AtmosphereLightDiscLuminance[Index] = Light->Proxy->GetOuterSpaceLuminance();
 				ViewUniformShaderParameters.AtmosphereLightColor[Index] = Light->Proxy->GetColor();
+				ViewUniformShaderParameters.AtmosphereLightColor[Index].A = 1.0f;
 				ViewUniformShaderParameters.AtmosphereLightColorGlobalPostTransmittance[Index] = Light->Proxy->GetColor() * Light->Proxy->GetTransmittanceFactor();
+				ViewUniformShaderParameters.AtmosphereLightColorGlobalPostTransmittance[Index].A = 1.0f;
 				ViewUniformShaderParameters.AtmosphereLightDirection[Index] = SkyAtmosphere->GetAtmosphereLightDirection(Index, -Light->Proxy->GetDirection());
 			}
 			else
