Fixed Sky shader clip operation.

diff --git a/Engine/Shaders/Private/SkyAtmosphere.usf b/Engine/Shaders/Private/SkyAtmosphere.usf
--- a/Engine/Shaders/Private/SkyAtmosphere.usf
+++ b/Engine/Shaders/Private/SkyAtmosphere.usf
@@ -737,7 +737,7 @@ void RenderSkyAtmosphereRayMarchingPS(
 		//Now the sky pass can ignore the pixel with depth == far but it will need to alpha clip because not all RHI backend support depthbound tests.
 		// And the depthtest is already setup to avoid writing all the pixel closer than to the camera than the start distance (very good optimisation).
 		// Since this shader does not write to depth or stencil it should still benefit from EArlyZ even with the clip (See AMD depth-in-depth documentation)
-		clip(0);
+		clip(-1.0f);
 		return;
 #endif
 	}
