Prefer casting 0-length occlusion rays when area light is to be normal-culled later.

diff --git a/Engine/Shaders/Private/RayTracing/RayTracingDiskLight.ush b/Engine/Shaders/Private/RayTracing/RayTracingDiskLight.ush
--- a/Engine/Shaders/Private/RayTracing/RayTracingDiskLight.ush
+++ b/Engine/Shaders/Private/RayTracing/RayTracingDiskLight.ush
@@ -13,6 +13,11 @@ bool GenerateDiskLightOcclusionRay(
 	out float RayTMax
 )
 {
+	RayOrigin = WorldPosition;
+	RayDirection = 0.0;
+	RayTMin = 0.01;
+	RayTMax = RayTMin;
+
 	// Sample disk of SourceRadius
 	float2 UV = UniformSampleDiskConcentric(RandSample);
 	float3 P_Local = float3(UV, 0.0) * LightParameters.SourceRadius;
@@ -29,9 +34,7 @@ bool GenerateDiskLightOcclusionRay(
 	bool IsVisible = NoL > 0.0;
 	if (IsVisible)
 	{
-		RayOrigin = WorldPosition;
 		RayDirection = LightDirection;
-		RayTMin = 0.01;
 		RayTMax = RayLength;
 	}
 	return IsVisible;
diff --git a/Engine/Shaders/Private/RayTracing/RayTracingRectLight.ush b/Engine/Shaders/Private/RayTracing/RayTracingRectLight.ush
--- a/Engine/Shaders/Private/RayTracing/RayTracingRectLight.ush
+++ b/Engine/Shaders/Private/RayTracing/RayTracingRectLight.ush
@@ -18,6 +18,12 @@ bool GenerateRectLightOcclusionRay(
 	out float RayPdf
 )
 {
+	RayOrigin = WorldPosition;
+	RayDirection = 0;
+	RayTMin = 0.0;
+	RayTMax = 0.0;
+	RayPdf = 0.0;
+
 	float3 Tangent;
 	float3 BiTangent;
 	GetRectLightTangents(LightParameters, /* out */ Tangent, /* out */ BiTangent);
@@ -37,9 +43,7 @@ bool GenerateRectLightOcclusionRay(
 	}
 
 	// Apply normal perturbation when defining ray
-	RayOrigin = WorldPosition;
 	RayDirection = LightDirection;
-	RayTMin = 0.0;
 	RayTMax = length(LightSamplePosition - WorldPosition);
 
 	// Define rectangle
