Fix support for 16 bit page table values in virtual texture

diff --git a/Engine/Shaders/Private/PageTableUpdate.usf b/Engine/Shaders/Private/PageTableUpdate.usf
--- a/Engine/Shaders/Private/PageTableUpdate.usf
+++ b/Engine/Shaders/Private/PageTableUpdate.usf
@@ -65,19 +65,20 @@ void PageTableUpdateVS(
 	OutPosition.xy = ( (float2)vCorner / PageTableSize ) * float2( 2, -2 ) + float2( -1, 1 );
 
 #if USE_16BIT
-	uint Page = vLevel;
-	Page |= pPage.x << 4;
-	Page |= pPage.y << 10;
+	// We can assume pPage fits in 6 bits and pack the final output to 16 bits
+	const uint PageCoordinateBitCount = 6;
 #else
+	const uint PageCoordinateBitCount = 8;
+#endif
+
 	// We want to store the quantity vPagesWideInLevel (computed as VIRTUALTEXTURE_MAX_PAGETABLE_SIZE >> vLevel)
 	// We then want to scale by (1/VIRTUALTEXTURE_MAX_PAGETABLE_SIZE), to avoid the need to explicitly factor this into other uniforms
 	// So scale factor becomes (VIRTUALTEXTURE_MAX_PAGETABLE_SIZE >> vLevel) / VIRTUALTEXTURE_MAX_PAGETABLE_SIZE, which simplifies to (1 / exp2(vLevel)) or (1 / (1 << vLevel))
+	// We will store vLevel in the page table and calculate (1 / (1 << vLevel)) on load
 	// Also nice that this avoids need to depend on specific value of VIRTUALTEXTURE_MAX_PAGETABLE_SIZE
-
-	uint Page = f32tof16(1.0f / float(1u << vLevel));
-	Page |= pPage.x << 16;
-	Page |= pPage.y << 24;
-#endif
+	uint Page = vLevel;
+	Page |= pPage.x << 4;
+	Page |= pPage.y << (4 + PageCoordinateBitCount);
 	
 	OutColor = Page;
 }
diff --git a/Engine/Shaders/Private/VirtualTextureCommon.ush b/Engine/Shaders/Private/VirtualTextureCommon.ush
--- a/Engine/Shaders/Private/VirtualTextureCommon.ush
+++ b/Engine/Shaders/Private/VirtualTextureCommon.ush
@@ -424,6 +424,7 @@ struct VTUniform
 	float pPageSize;
 	float vPageSize;
 	float vPageBorderSize;
+	uint PageCoordinateBitCount;
 };
 
 VTUniform VTUniform_Unpack(uint4 PackedUniform)
@@ -432,6 +433,7 @@ VTUniform VTUniform_Unpack(uint4 PackedUniform)
 	result.pPageSize = asfloat(PackedUniform.w);
 	result.vPageSize = asfloat(PackedUniform.y);
 	result.vPageBorderSize = asfloat(PackedUniform.z);
+	result.PageCoordinateBitCount = PackedUniform.x == 0 ? 8 : 6;
 	return result;
 }
 
@@ -440,9 +442,11 @@ float2 VTComputePhysicalUVs(in out VTPageTableResult PageTableResult, uint Layer
 {
 	const uint PackedPageTableValue = PageTableResult.PageTableValue[LayerIndex / 4u][LayerIndex & 3u];
 
-	const float UVScale = f16tof32(PackedPageTableValue);
-	const float pPageX = UnpackByte2(PackedPageTableValue);
-	const float pPageY = UnpackByte3(PackedPageTableValue);
+	// See packing in PageTableUpdate.usf
+	const uint vLevel = PackedPageTableValue & 0xf;
+	const float UVScale = 1.0f / (float)(1 << vLevel);
+	const float pPageX = (float)((PackedPageTableValue >> 4) & ((1 << Uniform.PageCoordinateBitCount) - 1));
+	const float pPageY = (float)(PackedPageTableValue >> (4 + Uniform.PageCoordinateBitCount));
 
 	const float2 vPageFrac = frac(PageTableResult.UV * UVScale);
 	const float2 pUV = float2(pPageX, pPageY) * Uniform.pPageSize + (vPageFrac * Uniform.vPageSize + Uniform.vPageBorderSize);
diff --git a/Engine/Source/Runtime/Renderer/Private/VT/AllocatedVirtualTexture.cpp b/Engine/Source/Runtime/Renderer/Private/VT/AllocatedVirtualTexture.cpp
--- a/Engine/Source/Runtime/Renderer/Private/VT/AllocatedVirtualTexture.cpp
+++ b/Engine/Source/Runtime/Renderer/Private/VT/AllocatedVirtualTexture.cpp
@@ -269,10 +269,10 @@ void FAllocatedVirtualTexture::GetPackedUniform(FUintVector4* OutUniform, uint32
 	{
 		const uint32 vPageSize = GetVirtualTileSize();
 		const uint32 PageBorderSize = GetTileBorderSize();
-
 		const float RcpPhysicalTextureSize = 1.0f / float(PhysicalTextureSize);
 		const uint32 pPageSize = vPageSize + PageBorderSize * 2u;
-		OutUniform->X = 0u;
+
+		OutUniform->X = GetPageTableFormat() == EVTPageTableFormat::UInt16 ? 1 : 0;
 		OutUniform->Y = BitcastFloatToUInt32((float)vPageSize * RcpPhysicalTextureSize);
 		OutUniform->Z = BitcastFloatToUInt32((float)PageBorderSize * RcpPhysicalTextureSize);
 		OutUniform->W = BitcastFloatToUInt32((float)pPageSize * RcpPhysicalTextureSize);
diff --git a/Engine/Source/Runtime/Renderer/Private/VT/VirtualTexturePhysicalSpace.h b/Engine/Source/Runtime/Renderer/Private/VT/VirtualTexturePhysicalSpace.h
--- a/Engine/Source/Runtime/Renderer/Private/VT/VirtualTexturePhysicalSpace.h
+++ b/Engine/Source/Runtime/Renderer/Private/VT/VirtualTexturePhysicalSpace.h
@@ -41,7 +41,7 @@ public:
 	inline FIntVector GetPhysicalLocation(uint16 pAddress) const { return FIntVector(pAddress % TextureSizeInTiles, pAddress / TextureSizeInTiles, 0); }
 
 	// 16bit page tables allocate 6bits to address TileX/Y, so can only address tiles from 0-63
-	inline bool DoesSupport16BitPageTable() const { return false; }// TextureSizeInTiles <= 64u;
+	inline bool DoesSupport16BitPageTable() const { return TextureSizeInTiles <= 64u; }
 
 	uint32 GetSizeInBytes() const;
 
