Support sign of Normal Z in [BC3|BC3] runtime virtual texture format
This allows storage of world space normals if required

diff --git a/Engine/Shaders/Private/VirtualTextureCommon.ush b/Engine/Shaders/Private/VirtualTextureCommon.ush
--- a/Engine/Shaders/Private/VirtualTextureCommon.ush
+++ b/Engine/Shaders/Private/VirtualTextureCommon.ush
@@ -545,29 +545,30 @@ float2 VirtualTextureWorldToUV(in float3 WorldPos, in float3 Origin, in float3 U
 }
 
 /** Generic normal unpack funtion. */
-float3 VirtualTextureUnpackNormal(in float2 PackedValue)
+float3 VirtualTextureUnpackNormal(in float2 PackedXY, in float PackedSignZ)
 {
-	float2 NormalXY = PackedValue * 2.f - 1.f;
-	float NormalZ = sqrt(saturate(1.0f - dot(NormalXY, NormalXY)));
-	return float3(NormalXY.xy, NormalZ);
+	float2 NormalXY = PackedXY * 2.f - 1.f;
+	float SignZ = PackedSignZ * 2.f - 1.f;
+	float NormalZ = sqrt(saturate(1.0f - dot(NormalXY, NormalXY))) * SignZ;
+	return float3(NormalXY, NormalZ);
 }
 
 /** Unpack normal from BC3. */
 float3 VirtualTextureUnpackNormalBC3(in float4 PackedValue)
 {
-	return VirtualTextureUnpackNormal(PackedValue.ag);
+	return VirtualTextureUnpackNormal(PackedValue.ag, 1.f);
 }
 
 /** Unpack normal from two BC3 textures. */
 float3 VirtualTextureUnpackNormalBC3BC3(in float4 PackedValue0, in float4 PackedValue1)
 {
-	return VirtualTextureUnpackNormal(float2(PackedValue0.a, PackedValue1.a));
+	return VirtualTextureUnpackNormal(float2(PackedValue0.a, PackedValue1.a), PackedValue1.b);
 }
 
 /** Unpack normal from BC5. */
 float3 VirtualTextureUnpackNormalBC5(in float4 PackedValue)
 {
-	return VirtualTextureUnpackNormal(PackedValue.rg);
+	return VirtualTextureUnpackNormal(PackedValue.rg, 1.f);
 }
 
 /** Unpack 16 bit height (with some hardcoded range scale/bias). */
diff --git a/Engine/Shaders/Private/VirtualTextureCompress.usf b/Engine/Shaders/Private/VirtualTextureCompress.usf
--- a/Engine/Shaders/Private/VirtualTextureCompress.usf
+++ b/Engine/Shaders/Private/VirtualTextureCompress.usf
@@ -221,6 +221,11 @@ void CompressBaseColorNormalSpecularExCS(
 	float3 BlockSpecular[16];
 	ReadBlockRGB(RenderTexture2, TextureSampler2, SampleUV, TexelUVSize, BlockSpecular);
 
+	for (int i=0; i<16; i++)
+	{
+		BlockSpecular[i].z = round(BlockSpecular[i].z);
+	}
+
 	OutCompressTexture0_u4[ThreadId.xy + DestRect.xy] = CompressBC3Block_SRGB(BlockBaseColor, BlockNormalX);
 	OutCompressTexture1[ThreadId.xy + DestRect.xy] = CompressBC3Block(BlockSpecular, BlockNormalY);
 }
@@ -238,8 +243,10 @@ void CopyNormalSpecularExCS(
 
 	float3 BaseColor = LinearToSrgb(RenderTexture0[SamplePos].xyz);
 	float2 NormalXY = RenderTexture1[SamplePos].xy;
-	float2 RoughnessSpecular = RenderTexture2[SamplePos].xy;
+	float3 RoughnessSpecularNormalZ = RenderTexture2[SamplePos].xyz;
+	
+	RoughnessSpecularNormalZ.z = round(RoughnessSpecularNormalZ.z);
 	
 	OutCopyTexture0[ThreadId.xy + DestRect.xy] = float4(BaseColor, NormalXY.x);
-	OutCopyTexture1[ThreadId.xy + DestRect.xy] = float4(RoughnessSpecular.x, RoughnessSpecular.y, 0.0f, NormalXY.y);
+	OutCopyTexture1[ThreadId.xy + DestRect.xy] = float4(RoughnessSpecularNormalZ, NormalXY.y);
 }
diff --git a/Engine/Shaders/Private/VirtualTextureMaterial.usf b/Engine/Shaders/Private/VirtualTextureMaterial.usf
--- a/Engine/Shaders/Private/VirtualTextureMaterial.usf
+++ b/Engine/Shaders/Private/VirtualTextureMaterial.usf
@@ -39,9 +39,9 @@ void MainVS(
 }
 
 // Prepare for VirtualTextureUnpackNormal() in VirtualTextureCommon.h
-float2 PackNormal( float3 N )
+float3 PackNormal( float3 N )
 {
-	return normalize(N).xy * 0.5 + 0.5;
+	return normalize(N) * 0.5 + 0.5;
 }
 
 // Prepare for VirtualTextureUnpackHeightR16() in VirtualTextureCommon.ush
@@ -92,11 +92,12 @@ void FPixelShaderInOut_MainPS(
 	Out.MRT[0] = float4(BaseColor, 1.f) * Opacity;
 #elif defined(OUT_BASECOLOR_NORMAL)
 	Out.MRT[0] = float4(BaseColor, 1.f) * Opacity;
-	Out.MRT[1] = float4(PackNormal(Normal), 0.f, 1.f) * Opacity;
+	Out.MRT[1] = float4(PackNormal(Normal).xy, 0.f, 1.f) * Opacity;
 #elif defined(OUT_BASECOLOR_NORMAL_SPECULAR)
+	float3 PackedNormal = PackNormal(Normal);
 	Out.MRT[0] = float4(BaseColor, 1.f) * Opacity;
-	Out.MRT[1] = float4(PackNormal(Normal), 0.f, 1.f) * Opacity;
-	Out.MRT[2] = float4(Specular, Roughness, 0.f, 1.f) * Opacity;
+	Out.MRT[1] = float4(PackedNormal.xy, 0.f, 1.f) * Opacity;
+	Out.MRT[2] = float4(Specular, Roughness, PackedNormal.z, 1.f) * Opacity;
 #elif defined(OUT_WORLDHEIGHT)
 	Out.MRT[0] = float4(PackWorldHeight(WorldHeight), 0, 0, 1);
 #endif
