Fixed light contribution and shadow of RectLight in volumetric fog.

(peafour-cherrypick of 7270304 by [REDACTED_PERSON])

diff --git a/Engine/Shaders/Private/LightGridCommon.ush b/Engine/Shaders/Private/LightGridCommon.ush
--- a/Engine/Shaders/Private/LightGridCommon.ush
+++ b/Engine/Shaders/Private/LightGridCommon.ush
@@ -163,6 +163,7 @@ struct FLocalLightData
 	float4 SpotAnglesAndSourceRadiusPacked;
     float4 LightDirectionAndShadowMask;
     float4 LightTangentAndSoftSourceRadius;
+	float4 RectBarnDoor;
 	/** Flag is true if clustered deferred is supported for this light. They are always first / together.*/
 	bool bClusteredDeferredSupported; 
 	/** Flag is true if it is a simple light. They are always first / together.*/
@@ -187,6 +188,7 @@ FLocalLightData GetLocalLightData(uint GridIndex, uint EyeIndex)
 		Result.LightDirectionAndShadowMask = ForwardLightData.ForwardLocalLightBuffer[LocalLightBaseIndex + 2];
 		Result.SpotAnglesAndSourceRadiusPacked = ForwardLightData.ForwardLocalLightBuffer[LocalLightBaseIndex + 3];
 		Result.LightTangentAndSoftSourceRadius = ForwardLightData.ForwardLocalLightBuffer[LocalLightBaseIndex + 4];
+		Result.RectBarnDoor = ForwardLightData.ForwardLocalLightBuffer[LocalLightBaseIndex + 5];
 #if INSTANCED_STEREO
 	}
 	else
diff --git a/Engine/Shaders/Private/VolumetricFog.usf b/Engine/Shaders/Private/VolumetricFog.usf
--- a/Engine/Shaders/Private/VolumetricFog.usf
+++ b/Engine/Shaders/Private/VolumetricFog.usf
@@ -222,11 +222,13 @@ void InjectShadowedLocalLightPS(
 		LightData.DistanceFadeMAD = DeferredLightUniforms.DistanceFadeMAD;
 		LightData.ShadowMapChannelMask = DeferredLightUniforms.ShadowMapChannelMask;
 		LightData.ShadowedBits = DeferredLightUniforms.ShadowedBits;
+		LightData.RectLightBarnCosAngle = DeferredLightUniforms.RectLightBarnCosAngle;
+		LightData.RectLightBarnLength = DeferredLightUniforms.RectLightBarnLength;
 
-		LightData.bInverseSquared = INVERSE_SQUARED_FALLOFF;
+		LightData.bInverseSquared = INVERSE_SQUARED_FALLOFF;		
 		LightData.bRadialLight = true;
 		LightData.bSpotLight = LightData.SpotAngles.x > -2.0f;
-		LightData.bRectLight = false;
+		LightData.bRectLight = LightData.RectLightBarnLength > -2.0f;
 
 		FRectTexture RectTexture = InitRectTexture(DeferredLightUniforms.SourceTexture);
 
@@ -593,8 +595,11 @@ void LightScatteringCS(
                 LightData.Tangent = LocalLight.LightTangentAndSoftSourceRadius.xyz;
                 LightData.SoftSourceRadius = LocalLight.LightTangentAndSoftSourceRadius.w;
 				LightData.bInverseSquared = LightData.FalloffExponent == 0;
+				LightData.RectLightBarnCosAngle = LocalLight.RectBarnDoor.x;
+				LightData.RectLightBarnLength = LocalLight.RectBarnDoor.y;
 				LightData.bRadialLight = true;
 				LightData.bSpotLight = LightData.SpotAngles.x > -2.0f;
+				LightData.bRectLight = LightData.RectLightBarnLength > -2.0f;
 
 				float3 LightColor = LightData.Color;
 
diff --git a/Engine/Source/Runtime/Engine/Private/Components/PointLightComponent.cpp b/Engine/Source/Runtime/Engine/Private/Components/PointLightComponent.cpp
--- a/Engine/Source/Runtime/Engine/Private/Components/PointLightComponent.cpp
+++ b/Engine/Source/Runtime/Engine/Private/Components/PointLightComponent.cpp
@@ -44,6 +44,8 @@ void FPointLightSceneProxy::GetLightShaderParameters(FLightShaderParameters& Lig
 	LightParameters.SoftSourceRadius = SoftSourceRadius;
 	LightParameters.SourceLength = SourceLength;
 	LightParameters.SourceTexture = GWhiteTexture->TextureRHI;
+	LightParameters.RectLightBarnCosAngle = 0.0f;
+	LightParameters.RectLightBarnLength = -2.0f;
 }
 
 /**
diff --git a/Engine/Source/Runtime/Engine/Private/Components/SpotLightComponent.cpp b/Engine/Source/Runtime/Engine/Private/Components/SpotLightComponent.cpp
--- a/Engine/Source/Runtime/Engine/Private/Components/SpotLightComponent.cpp
+++ b/Engine/Source/Runtime/Engine/Private/Components/SpotLightComponent.cpp
@@ -77,6 +77,8 @@ public:
 		LightParameters.SoftSourceRadius = SoftSourceRadius;
 		LightParameters.SourceLength = SourceLength;
 		LightParameters.SourceTexture = GWhiteTexture->TextureRHI;
+		LightParameters.RectLightBarnCosAngle = 0.0f;
+		LightParameters.RectLightBarnLength = -2.0f;
 	}
 
 	// FLightSceneInfo interface.
diff --git a/Engine/Source/Runtime/Renderer/Private/BasePassRendering.h b/Engine/Source/Runtime/Renderer/Private/BasePassRendering.h
--- a/Engine/Source/Runtime/Renderer/Private/BasePassRendering.h
+++ b/Engine/Source/Runtime/Renderer/Private/BasePassRendering.h
@@ -43,6 +43,7 @@ public:
 	FVector4 LightDirectionAndShadowMapChannelMask;
 	FVector4 SpotAnglesAndSourceRadiusPacked;
 	FVector4 LightTangentAndSoftSourceRadius;
+	FVector4 RectBarnDoor;
 };
 
 BEGIN_GLOBAL_SHADER_PARAMETER_STRUCT(FSharedBasePassUniformParameters,)
diff --git a/Engine/Source/Runtime/Renderer/Private/LightGridInjection.cpp b/Engine/Source/Runtime/Renderer/Private/LightGridInjection.cpp
--- a/Engine/Source/Runtime/Renderer/Private/LightGridInjection.cpp
+++ b/Engine/Source/Runtime/Renderer/Private/LightGridInjection.cpp
@@ -420,6 +420,7 @@ void FDeferredShadingSceneRenderer::ComputeLightGrid(FRHICommandListImmediate& R
 
 						LightData.SpotAnglesAndSourceRadiusPacked = FVector4(-2, 1, 0, *(float*)&PackedWInt);
 						LightData.LightTangentAndSoftSourceRadius = FVector4(1.0f, 0.0f, 0.0f, 0.0f);
+						LightData.RectBarnDoor = FVector4(0, -2, 0, 0);
 
 #if ENABLE_LIGHT_CULLING_VIEW_SPACE_BUILD_DATA
 						FVector4 ViewSpacePosAndRadius(View.ViewMatrices.GetViewMatrix().TransformPosition(SimpleLightPerViewData.Position), SimpleLight.Radius);
@@ -501,6 +502,8 @@ void FDeferredShadingSceneRenderer::ComputeLightGrid(FRHICommandListImmediate& R
 
 							LightData.LightTangentAndSoftSourceRadius = FVector4(LightParameters.Tangent, LightParameters.SoftSourceRadius);
 
+							LightData.RectBarnDoor = FVector4(LightParameters.RectLightBarnCosAngle, LightParameters.RectLightBarnLength, 0, 0);
+
 							float VolumetricScatteringIntensity = LightProxy->GetVolumetricScatteringIntensity();
 
 							if (LightNeedsSeparateInjectionIntoVolumetricFog(LightSceneInfo, VisibleLightInfos[LightSceneInfo->Id]))
diff --git a/Engine/Source/Runtime/Renderer/Private/VolumetricFog.cpp b/Engine/Source/Runtime/Renderer/Private/VolumetricFog.cpp
--- a/Engine/Source/Runtime/Renderer/Private/VolumetricFog.cpp
+++ b/Engine/Source/Runtime/Renderer/Private/VolumetricFog.cpp
@@ -366,7 +366,7 @@ bool LightNeedsSeparateInjectionIntoVolumetricFog(const FLightSceneInfo* LightSc
 	const FLightSceneProxy* LightProxy = LightSceneInfo->Proxy;
 
 	if (GVolumetricFogInjectShadowedLightsSeparately
-		&& (LightProxy->GetLightType() == LightType_Point || LightProxy->GetLightType() == LightType_Spot)
+		&& (LightProxy->GetLightType() == LightType_Point || LightProxy->GetLightType() == LightType_Spot || LightProxy->GetLightType() == LightType_Rect)
 		&& !LightProxy->HasStaticLighting()
 		&& LightProxy->CastsDynamicShadow()
 		&& LightProxy->CastsVolumetricShadow())
