Added more fraction for small number un lux/nit meter.

diff --git a/Engine/Shaders/Private/PostProcessVisualizeHDR.usf b/Engine/Shaders/Private/PostProcessVisualizeHDR.usf
--- a/Engine/Shaders/Private/PostProcessVisualizeHDR.usf
+++ b/Engine/Shaders/Private/PostProcessVisualizeHDR.usf
@@ -174,29 +174,29 @@ void MainPS(noperspective float4 UVAndScreenPos : TEXCOORD0, float4 SvPosition :
 
 			// Luminance
 			int2 TopLeft = ViewportCenter + int2(11, 11);
-			if (MaxLuminance < 100.0f)
+			if (MaxLuminance < 1000.0f)
 			{
-				PrintSmallFloat(PixelPos, OutColor.xyz, float3(1, 1, 1), TopLeft, MaxLuminance);
+				PrintFloat(PixelPos, OutColor.xyz, float3(1, 1, 1), TopLeft, MaxLuminance);
 			}
 			else
 			{
 				PrintFloatNoFractionLarge(PixelPos, OutColor.xyz, float3(1, 1, 1), TopLeft, MaxLuminance);
 			}
-			TopLeft.x += 54;
+			TopLeft.x += 60;
 			PrintCharacter(PixelPos, OutColor.xyz, float3(1, 1, 1), TopLeft, _N_);
 			PrintCharacter(PixelPos, OutColor.xyz, float3(1, 1, 1), TopLeft, _I_);
 			PrintCharacter(PixelPos, OutColor.xyz, float3(1, 1, 1), TopLeft, _T_);
 			// Illuminance
 			TopLeft = ViewportCenter + int2(11, 22);
-			if (MaxIlluminance < 100.0f)
+			if (MaxIlluminance < 1000.0f)
 			{
-				PrintSmallFloat(PixelPos, OutColor.xyz, float3(1, 1, 1), TopLeft, MaxIlluminance);
+				PrintFloat(PixelPos, OutColor.xyz, float3(1, 1, 1), TopLeft, MaxIlluminance);
 			}
 			else
 			{
 				PrintFloatNoFractionLarge(PixelPos, OutColor.xyz, float3(1, 1, 1), TopLeft, MaxIlluminance);
 			}
-			TopLeft.x += 54;
+			TopLeft.x += 60;
 			PrintCharacter(PixelPos, OutColor.xyz, float3(1, 1, 1), TopLeft, _L_);
 			PrintCharacter(PixelPos, OutColor.xyz, float3(1, 1, 1), TopLeft, _U_);
 			PrintCharacter(PixelPos, OutColor.xyz, float3(1, 1, 1), TopLeft, _X_);
