Slate - For RHI based material shaders, this CL eliminates several permutations that we didn't need, but were using in some cases (removal of those in the editor was a previous CL).  This eliminates 3 "Disabled Effect" permutations per shaders, this is now a dyanmic branch in the shader - It's still a permutation in slate global shaders.  This also removes all the FLocalVertexFactory permutations, as well as the ConvertToUniformMesh permutation, both of which shouldn't have ever been on for UII Materials - but the editor was using the FLocalVertexFactory one (but those cases should all be fixed now).

Forgot the shader.

Matt.Kuhlenschmidt

diff --git a/Engine/Shaders/Private/SlateElementPixelShader.usf b/Engine/Shaders/Private/SlateElementPixelShader.usf
--- a/Engine/Shaders/Private/SlateElementPixelShader.usf
+++ b/Engine/Shaders/Private/SlateElementPixelShader.usf
@@ -7,9 +7,14 @@
 
 /** Display gamma x:gamma curve adjustment, y:inverse gamma (1/GEngine->DisplayGamma), z:InvertAlpha, w:Contrast */
 half4 GammaAndAlphaValues;
+/** x:DrawDisabledEffect */
 /** Per element params */
 float4 ShaderParams;
 
+#if USE_MATERIALS
+half4 DrawFlags;
+#endif
+
 /** Texture for the element */
 Texture2D ElementTexture;
 SamplerState ElementTextureSampler;
@@ -259,22 +264,36 @@ half4 Main( VertexToPixelInterpolants VIn ) : SV_Target0
 #else
 	OutColor = GetLineSegmentElementColor( VIn );
 #endif
+	
+	// For standard slate which only has a few shaders, the permutation is fine here,
+	// for Engine Materials, we don't ever define this macro, so instead we fallback
+	// to looking at the shader constants.
+#if USE_MATERIALS
+	const half DrawDisabledEffect = DrawFlags.x;
+#else
+	#if DRAW_DISABLED_EFFECT
+		const half DrawDisabledEffect = 1;
+	#else
+		const half DrawDisabledEffect = 0;
+	#endif
+#endif
 
-#if DRAW_DISABLED_EFFECT
-    #if SOURCE_IN_LINEAR_SPACE
-        half3 LumCoeffs = half3( 0.3, 0.59, .11 );
-        half3 Grayish = half3(.1, .1, .1);
-    #else
-        // These are just the linear space values above converted to sRGB space.
-        half3 LumCoeffs = float3( 0.5843, 0.7921, 0.3647 );
-        half3 Grayish = float3(0.349, 0.349, 0.349);
-    #endif
+	if (DrawDisabledEffect != 0)
+	{
+		#if SOURCE_IN_LINEAR_SPACE
+			half3 LumCoeffs = half3( 0.3, 0.59, .11 );
+			half3 Grayish = half3(.1, .1, .1);
+		#else
+			// These are just the linear space values above converted to sRGB space.
+			half3 LumCoeffs = float3( 0.5843, 0.7921, 0.3647 );
+			half3 Grayish = float3(0.349, 0.349, 0.349);
+		#endif
     
-	//desaturate
-	half Lum = dot( LumCoeffs, OutColor.rgb );
-	OutColor.rgb = lerp( OutColor.rgb, float3(Lum,Lum,Lum), .8 );
-	OutColor.rgb = lerp( OutColor.rgb, Grayish, clamp( distance( OutColor.rgb, Grayish ), 0, .8)  );
-#endif
+		//desaturate
+		half Lum = dot( LumCoeffs, OutColor.rgb );
+		OutColor.rgb = lerp( OutColor.rgb, float3(Lum,Lum,Lum), .8 );
+		OutColor.rgb = lerp( OutColor.rgb, Grayish, clamp( distance( OutColor.rgb, Grayish ), 0, .8)  );
+	}
 
     // Contrast
     OutColor = ApplyContrast(OutColor);
