Fix Ray Tracing Reflection secondary bounce bias issues

(ushell-p4-cherrypick of 8610504 by riley.alston)

diff --git a/Engine/Shaders/Private/RayTracing/RayTracingReflections.usf b/Engine/Shaders/Private/RayTracing/RayTracingReflections.usf
--- a/Engine/Shaders/Private/RayTracing/RayTracingReflections.usf
+++ b/Engine/Shaders/Private/RayTracing/RayTracingReflections.usf
@@ -134,7 +134,6 @@ void FixSampleDirectionIfNeeded(float3 SmoothSurfaceNormal, inout float3 SampleD
 [shader("raygeneration")]
 void RayTracingReflectionsRGS()
 {
-	uint DummyVariable;
 	const float ClearCoatSpecularColor = 0.04; // See ReflectionEnvironmentPixelShader.usf
 
 	FDeferredMaterialPayload DeferredMaterialPayload;
@@ -283,7 +282,14 @@ void RayTracingReflectionsRGS()
 				TopLayerRay.TMax = LocalMaxRayDistance;
 
 				TopLayerRay.Direction = GenerateReflectionDirection(RandSequence, DispatchThreadId, BounceIndex, LocalSamplesPerPixel, IncidentDirection, TopLayerWorldNormal, TopLayerRoughness);
-				ApplyCameraRelativeDepthBias(TopLayerRay, PixelCoord, DeviceZ, TopLayerWorldNormal, ReflectionMaxNormalBias);
+
+				// Apply a ray bias with respect to the camera and surface normal on the first bounce, otherwise only with respect to the surface normal for subsequent bounces
+				if (BounceIndex == 0) {
+					ApplyCameraRelativeDepthBias(TopLayerRay, PixelCoord, DeviceZ, TopLayerWorldNormal, ReflectionMaxNormalBias);
+				} else {
+					ApplyPositionBias(TopLayerRay, TopLayerWorldNormal, ReflectionMaxNormalBias);
+				}
+
 				if (BounceIndex == 0)
 				{
 					// This fixes black reflexions due to R pointing into the surface.
