Touch shader which needs recompile for tests to pass

diff --git a/Engine/Shaders/Private/VirtualTextureMaterial.usf b/Engine/Shaders/Private/VirtualTextureMaterial.usf
--- a/Engine/Shaders/Private/VirtualTextureMaterial.usf
+++ b/Engine/Shaders/Private/VirtualTextureMaterial.usf
@@ -41,7 +41,7 @@ void MainVS(
 // Prepare for VirtualTextureUnpackNormal() in VirtualTextureCommon.h
 float3 PackNormal( float3 N )
 {
-	return normalize(N) * 0.5 + 0.5;
+	return normalize(N) * 0.5f + 0.5f;
 }
 
 // Prepare for VirtualTextureUnpackHeightR16() in VirtualTextureCommon.ush
