GTAO Bugfixes

diff --git a/Engine/Shaders/Private/PostProcessAmbientOcclusion.usf b/Engine/Shaders/Private/PostProcessAmbientOcclusion.usf
--- a/Engine/Shaders/Private/PostProcessAmbientOcclusion.usf
+++ b/Engine/Shaders/Private/PostProcessAmbientOcclusion.usf
@@ -27,6 +27,10 @@
 // This leads to differences in angles and fall off distances but generally the differences are minor. For extreme FOV cases this might be a problem.
 #define GTAO_OPT 0
 
+// This define enables the thickness heuristic which assu
+#define GTAO_THICKNESS_HEURISTIC 0
+
+
 // 0: classic with weighted sample, 1: don't normalize and adjust the formula to be simpler and faster - can look better and is cheaper (Alchemy like?)
 #define OPTIMIZATION_O1 1
 
@@ -937,8 +941,8 @@ void MainSSAOSmoothCS(
 
 float4 GTAOParams[2];
 
-#define FALLOFF_START		50.0
-#define FALLOFF_END			250.0
+#define FALLOFF_START		100.0
+#define FALLOFF_END			200.0
 #define FALLOFF_START_SQ	(FALLOFF_START*FALLOFF_START)
 #define FALLOFF_END_SQ		(FALLOFF_END*FALLOFF_END)
 
@@ -953,9 +957,23 @@ RWTexture2D<float2> HorizonOutTexture;
 
 float Falloff(float distSq)
 {
-	float Falloff_Scale = 1.0f / (FALLOFF_END_SQ - FALLOFF_START_SQ);
-	float Falloff_Bias = -FALLOFF_START_SQ * Falloff_Scale; 
-	return  2.0 * saturate(distSq* Falloff_Scale + Falloff_Bias   );
+	float dist = sqrt(distSq);
+	float Falloff_Scale = 1.0f / (FALLOFF_END - FALLOFF_START);
+	float Falloff_Bias = -FALLOFF_START * Falloff_Scale; 
+	return  2.0 * saturate(dist* Falloff_Scale + Falloff_Bias   );
+
+//	float Falloff_Scale = 1.0f / (FALLOFF_END_SQ - FALLOFF_START_SQ);
+//	float Falloff_Bias = -FALLOFF_START_SQ * Falloff_Scale; 
+//	return  2.0 * saturate(distSq* Falloff_Scale + Falloff_Bias   );
+}
+
+
+float3 ScreenToViewPosX(float2 UV, float SceneDepth)
+{
+	float3 ViewPos = ScreenToViewPos(UV, SceneDepth);
+
+	//return float3( ViewPos.x, ViewPos.z, ViewPos.y);
+	return float3( ViewPos.x, ViewPos.y, ViewPos.z);
 }
 
 
@@ -969,7 +987,7 @@ float3 GetViewSpacePosFromHZB(float2 UV, float MipLevel)
 	SceneDepth = Texture2DSampleLevel( PostprocessInput1, PostprocessInput1Sampler, HZBUV, MipLevel).r;
 	SceneDepth = ConvertFromDeviceZ(SceneDepth);
 
-	return ScreenToViewPos(UV, SceneDepth);
+	return ScreenToViewPosX(UV, SceneDepth);
 }
 
 
@@ -994,8 +1012,8 @@ float3 GetNormal(float2 UV, float3 ViewSpacePosMid, bool bConstructFromDepth)
 		float ZRight		= ConvertFromDeviceZ(DeviceZ + DeviceZDdx);
 		float ZDown			= ConvertFromDeviceZ(DeviceZ + DeviceZDdy);
 
-		float3 Right		= ScreenToViewPos(UV+ float2( DeltaUV.x,      0.0f) , ZRight)-ViewSpacePosMid;
-		float3 Down			= ScreenToViewPos(UV+ float2(      0.0f, DeltaUV.y) , ZDown) -ViewSpacePosMid;
+		float3 Right		= ScreenToViewPosX(UV+ float2( DeltaUV.x,      0.0f) , ZRight)-ViewSpacePosMid;
+		float3 Down			= ScreenToViewPosX(UV+ float2(      0.0f, DeltaUV.y) , ZDown) -ViewSpacePosMid;
 		
 		ViewSpaceNormal = normalize(cross(Right, Down));
 
@@ -1018,21 +1036,38 @@ float GetLinearDepthProj(float2 ScreenUV)
 	return  1.0f / (DeviceZ * View.InvDeviceZToWorldZTransform[2] - View.InvDeviceZToWorldZTransform[3]);
 }
 
+float UpdateMaxAngle(float CurrAng, float PrevAng)
+{
+#ifdef GTAO_THICKNESS_HEURISTIC
+
+	if(CurrAng > PrevAng)
+	{
+		return CurrAng;
+	}
+	else
+	{
+		const float BLEND_FACTOR = 0.9;
+		return 	PrevAng ;//(CurrAng * BLEND_FACTOR) + (PrevAng * (1-BLEND_FACTOR));
+	}
+#else
+	return max(CurrAng, PrevAng);
+#endif
+}
+
 
 float SearchForLargestAngleRef(float2 inUV, float2 deltaUV, float3 ViewSpacePosition, float3 ViewDir)
 {
 	float BiggestAngle = -1.0f;
+	float PrevAngle    = -1.0f;
 
-	deltaUV.y	  *=-1;
-	deltaUV.y	  *=.3;
+	deltaUV.y	  *= -1;
 	float2 UV	   = inUV ;
 
 #if GTAO_OPT
-
 	// Precompute as much stuff upfront
 	float CurrDepth			= GetLinearDepthProj(UV);
 	UV						+= deltaUV;
-	float2 DeltaVS_XY	    = ScreenToViewPos(UV, CurrDepth).xy - ViewSpacePosition.xy;
+	float2 DeltaVS_XY	    = ScreenToViewPosX(UV, CurrDepth).xy - ViewSpacePosition.xy;
 	float  DeltaVS_XYLen    = sqrt(dot( DeltaVS_XY.xy, DeltaVS_XY.xy));
 	float  CurrVS_XYLen		= DeltaVS_XYLen;
 
@@ -1056,28 +1091,37 @@ float SearchForLargestAngleRef(float2 inUV, float2 deltaUV, float3 ViewSpacePosi
 			float falloff	= Falloff(distSQ);
 			BiggestAngle	= max(BiggestAngle, cosh - falloff);
 		}
-
-	}	
+	}
+	
 #else
+
+	deltaUV *=1;
+
+	// Do first sample so we have a previous to work with 
 	for(int i=1; i<=GTAO_NUMTAPS; i++)
 	{
 		UV					+= deltaUV;
 	
 		float SceneDepth	= GetDepthFromAOInput(UV);
-		float3 SampVS		= ScreenToViewPos(UV, SceneDepth);
+		float3 SampVS		= ScreenToViewPosX(UV, SceneDepth);
 	
 		float3 DiffV		= (SampVS - ViewSpacePosition) ;
 		float distSQ		= dot(DiffV, DiffV);
 	
 		// At this point we could early out if DistSQ is greater than falloff
-		BRANCH
-		if(distSQ < FALLOFF_END_SQ)
+	//	BRANCH
+	//	if(distSQ < FALLOFF_END_SQ)
 		{
 			float ooDist    = rsqrt(distSQ);
 			float cosh		= ooDist * dot(DiffV, ViewDir );
 	
 			float falloff	= Falloff(distSQ);
-			BiggestAngle	= max(BiggestAngle, cosh - falloff);
+
+			cosh			= cosh - falloff;
+			BiggestAngle	= max(BiggestAngle, cosh);
+
+			//BiggestAngle	= UpdateMaxAngle(cosh, PrevAngle);
+			//PrevAngle		= cosh;
 		}
 	}	
 #endif
@@ -1085,6 +1129,8 @@ float SearchForLargestAngleRef(float2 inUV, float2 deltaUV, float3 ViewSpacePosi
 	return BiggestAngle;
 }
 
+
+
 float SearchForLargestAngle(float2 inUV, float2 deltaUV, float3 ViewSpacePosition, float3 ViewDir)
 {
 #if GTAO_REF
@@ -1092,6 +1138,7 @@ float SearchForLargestAngle(float2 inUV, float2 deltaUV, float3 ViewSpacePositio
 #endif
 
 	float BiggestAngle = -1.0f;
+	float PrevAngle	   = -1.0f;
 
 	deltaUV.y	  *=-1;
 	float2 UV	= inUV;
@@ -1103,7 +1150,7 @@ float SearchForLargestAngle(float2 inUV, float2 deltaUV, float3 ViewSpacePositio
 	// Precompute as much stuff upfront
 	float CurrDepth			= GetLinearDepthProj(UV);
 	UV						+= deltaUV;
-	float2 DeltaVS_XY	    = ScreenToViewPos(UV, CurrDepth).xy - ViewSpacePosition.xy;
+	float2 DeltaVS_XY	    = ScreenToViewPosX(UV, CurrDepth).xy - ViewSpacePosition.xy;
 	float  DeltaVS_XYLen    = sqrt(dot( DeltaVS_XY.xy, DeltaVS_XY.xy));
 	float  CurrVS_XYLen		= DeltaVS_XYLen;
 
@@ -1119,7 +1166,7 @@ float SearchForLargestAngle(float2 inUV, float2 deltaUV, float3 ViewSpacePositio
 			DeltaVS_XYLen	+= DeltaVS_XYLen;
 			MipLevel++;
 		}
-
+		
 		if(i==GTAO_NUMTAPS/4)
 		{
 			HZB_dUV += HZB_dUV;
@@ -1150,30 +1197,36 @@ float SearchForLargestAngle(float2 inUV, float2 deltaUV, float3 ViewSpacePositio
 			float cosh		= ooDist * DiffDepth;
 
 			float falloff	= Falloff(distSQ);
-			BiggestAngle	= max(BiggestAngle, cosh - falloff);
+			cosh			= cosh - falloff;
+
+			BiggestAngle	= max(BiggestAngle, cosh);
+
+		//	BiggestAngle	= UpdateMaxAngle(cosh, PrevAngle);
+		//	PrevAngle		= cosh;
 		}
 	}	
 
 
 #else
+
 	for(int i=1; i<=GTAO_NUMTAPS; i++)
 	{
 		if(i==3)
 		{
-			deltaUV += deltaUV;
-			MipLevel++;
-		}
-
-		if(i==GTAO_NUMTAPS/4)
-		{
-			deltaUV += deltaUV;
+			deltaUV *=1.5;
 			MipLevel++;
 		}
-		else if(i==GTAO_NUMTAPS/2)
+		
+		if(i==GTAO_NUMTAPS/2)
 		{
-			deltaUV += deltaUV;
+			deltaUV *=1.5;
 			MipLevel++;
 		}
+		//else if(i==GTAO_NUMTAPS/2)
+		//{
+		//	deltaUV += deltaUV;
+		//	MipLevel++;
+		//}
 
 		UV				+= deltaUV;
 		
@@ -1189,7 +1242,14 @@ float SearchForLargestAngle(float2 inUV, float2 deltaUV, float3 ViewSpacePositio
 			float ooDist    = rsqrt(distSQ);
 			float cosh		= ooDist * dot(DiffV,ViewDir );
 			float falloff	= Falloff(distSQ);
-			BiggestAngle	= max(BiggestAngle, cosh - falloff);
+			
+			cosh			= cosh - falloff;
+
+
+			BiggestAngle	= max(BiggestAngle, cosh);
+
+			//BiggestAngle	= UpdateMaxAngle(cosh, PrevAngle);
+			//PrevAngle		= cosh;
 		}
 	}
 #endif	
@@ -1200,11 +1260,11 @@ float SearchForLargestAngle(float2 inUV, float2 deltaUV, float3 ViewSpacePositio
 
 
 // Given a screen space Axis this will search for the min and max angles
-float2 SearchAxisForAngles(float2 UV, float3 Dir, float3 ViewDir, float3 ViewSpacePos, float Scale)
+float2 SearchAxisForAngles(float2 UV, float3 ScreenDir, float3 ViewDir, float3 ViewSpacePos, float Scale)
 {
 	// TODO - Fix aspect ratio
-	float2 DeltaUV	= Scale * float2(View.BufferSizeAndInvSize.z * Dir.x, View.BufferSizeAndInvSize.w * Dir.y);
-
+	float2 DeltaUV	= Scale * float2(View.BufferSizeAndInvSize.z * ScreenDir.x, View.BufferSizeAndInvSize.w * ScreenDir.y);
+	
 	float O =  0.0f;//GTAOParams[0].z; 
 
 	float2 StartUV = UV + (DeltaUV*O);
@@ -1237,7 +1297,8 @@ float ComputeInnerIntegral(float2 Angles, float3 ScreenDir, float3 ViewDir, floa
 	Angles.x = Gamma + max(-Angles.x - Gamma, -(PI_HALF-ANGLEBIAS) );
 	Angles.y = Gamma + min( Angles.y - Gamma,  (PI_HALF-ANGLEBIAS) );
 
-	float AO = ( 0.25 * ( (Angles.x * SinGamma + CosGamma - cos((2.0 * Angles.x) - Gamma)) +
+	float AO = (/* LenProjNormal * */   0.25 * 
+					    ( (Angles.x * SinGamma + CosGamma - cos((2.0 * Angles.x) - Gamma)) +
 				  	      (Angles.y * SinGamma + CosGamma - cos((2.0 * Angles.y) - Gamma)) ));
 
 	AO = saturate(AO);
@@ -1250,6 +1311,7 @@ float ComputeInnerIntegral(float2 Angles, float3 ScreenDir, float3 ViewDir, floa
 	float Add = ScreenSpaceAOParams[4].y;
 	AO = lerp(AO, 1, saturate(SceneDepth * Mul + Add));
 
+
 	return AO;
 }
 
@@ -1283,33 +1345,33 @@ void HorizonSearchAndInnerIntegratePSandCS(in float2 UV, out float4 OutColor)
 	RandomVec.x = dot(RandomTexVec.xy, float2(TemporalCos, -TemporalSin ));
 	RandomVec.y = dot(RandomTexVec.xy, float2(TemporalSin, TemporalCos ));
 
-	float3 ViewSpacePos  = ScreenToViewPos(UV,SceneDepth);
+	float3 ViewSpacePos  = ScreenToViewPosX(UV,SceneDepth);
 
 	// Get Normal and transform into viewspace
 	float3 ViewSpaceNormal = GetNormal(UV, ViewSpacePos, false);
 #if GTAO_OPT
-	float3 ViewDir			= -normalize(ViewSpacePos.xyz);
-#else
 	float3 ViewDir			= float3(0,0,-1);
+#else
+	float3 ViewDir			= -normalize(ViewSpacePos.xyz);
 #endif
 
 	float Sum=0.0;
 
 	// Given the depth determine how many pixels away we are guranteed to be outside of the falloff distance.
 	float InvTanHalfFov			= ScreenSpaceAOParams[3].w;
-	float ScreenPixelsToCutoff  = View.BufferSizeAndInvSize.x * ((FALLOFF_END*(1.0f/InvTanHalfFov)) / ViewSpacePos.z) ; 
+	float ScreenPixelsToCutoff  = View.BufferSizeAndInvSize.x * ((FALLOFF_END*0.5) / ViewSpacePos.z) ; 
 
 	// Get radius based on sampling 
 	// Get radius based on distance - TODO - Need to use FOV and aspect ratio
 	float Scale = ScreenPixelsToCutoff /  GTAO_NUMTAPS;
 
 	// Clamp the scale. Too small and it won't be jumping by more than 1 pixel. Too much and it'll thrash the cache
-	Scale = max(Scale, 2.0f);
+	Scale = clamp(Scale, 2.0, 8.0);
 
 #if GTAO_REF
 	for(int i=0; i<GTAO_NUMROTS; i++)
 	{
-		float phi			= ((float)i) *  (PI / GTAO_NUMROTS);
+		float phi			= ((float)i) *  float(PI / GTAO_NUMROTS);
 		float3 ScreenDir	= float3(cos(phi), sin(phi), 0.0);
 #else
 		float3 ScreenDir	= float3(RandomVec.x, RandomVec.y, 0.0);
@@ -1325,7 +1387,6 @@ void HorizonSearchAndInnerIntegratePSandCS(in float2 UV, out float4 OutColor)
 	}
 #endif
 
-
 	float AmbientOcclusionIntensity = ScreenSpaceAOParams[0].w;
 	float AmbientOcclusionPower = ScreenSpaceAOParams[0].x;
 	// user adjust AO
@@ -1398,7 +1459,7 @@ void GTAOInnerIntegratePSandCS(in float2 UV, out float4 OutColor)
 	float3 ScreenDir			= float3(RandomVec.x, RandomVec.y, 0.0);
 
 	// ViewspacePos and Normal
-	float3 ViewSpacePos		= ScreenToViewPos(UV, SceneDepth);
+	float3 ViewSpacePos		= ScreenToViewPosX(UV, SceneDepth);
 	float3 WorldNormal		= GetGBufferData(UV, false).WorldNormal;
 	float3 ViewSpaceNormal	= normalize(mul(WorldNormal, (float3x3)View.TranslatedWorldToView));
 	float3 ViewDir			= -normalize(ViewSpacePos.xyz);	// TODO - This is a function of UV only.
@@ -1471,7 +1532,7 @@ float2 HorizonSearchPSandCS(in float2 UV)
 	float2 RandomVec			= (Texture2DSample(RandomNormalTexture, RandomNormalTextureSampler, TemporalOffset + UV * ViewportUVToRandomUV).rg * 2 - 1) ;
 	float3 ScreenDir			= float3(RandomVec.x, RandomVec.y, 0.0);
 	
-	float3 ViewSpacePos			= ScreenToViewPos(UV, SceneDepth);
+	float3 ViewSpacePos			= ScreenToViewPosX(UV, SceneDepth);
 	float3 ViewDir				= -normalize(ViewSpacePos.xyz);
 
 	// Get radius based on distance - TODO - Need to use FOV and aspect ratio
