Separated water lighting evaluation from BasePass shader code.

diff --git a/Engine/Shaders/Private/BasePassPixelShader.usf b/Engine/Shaders/Private/BasePassPixelShader.usf
--- a/Engine/Shaders/Private/BasePassPixelShader.usf
+++ b/Engine/Shaders/Private/BasePassPixelShader.usf
@@ -17,7 +17,6 @@
 #if MATERIAL_IS_SINGLE_LAYER_WATER
 #define PreIntegratedGF			OpaqueBasePass.PreIntegratedGFTexture
 #define PreIntegratedGFSampler	OpaqueBasePass.PreIntegratedGFSampler
-#include "ParticipatingMediaCommon.ush" 
 #endif
 
 #include "SHCommon.ush"
@@ -85,6 +84,10 @@ float NormalCurvatureToRoughness(float3 WorldNormal)
 #include "ForwardLightingCommon.ush"
 #endif
 
+#if  MATERIAL_IS_SINGLE_LAYER_WATER
+#include "SingleLayerWaterCommon.ush"
+#endif
+
 #if !FORWARD_SHADING
 
 void GetVolumeLightingNonDirectional(float4 AmbientLightingVector, float3 DiffuseColor, inout float3 InterpolatedLighting, out float4 VolumeLighting)
@@ -1068,10 +1071,9 @@ ShadingModel);
 		DiffuseColor += (DiffuseIndirectLighting * DiffuseColorForIndirect + SubsurfaceIndirectLighting * SubsurfaceColor) * AOMultiBounce( GBuffer.BaseColor, DiffOcclusion );
 
 #if MATERIAL_IS_SINGLE_LAYER_WATER
-		// Fade out diffuse as this will be hanlded by the single scattering lighting. when over the water surface.
-		// We do this after the just above GetPrecomputedIndirectLightingAndSkyLight to keep ambient lighting avialable.
-		// We also keep the SpecularColor for sun/water interactions
-		const float3 WaterIndirectDiffuse = DiffuseIndirectLighting;	// Raw unaffected illuminance
+		// Fade out diffuse as this will be handled by the single scattering lighting in water material.
+		// We do this after the just above GetPrecomputedIndirectLightingAndSkyLight to keep ambiant lighting avialable.
+		// We also keep the SpecularColor for sun/water interactions.
 		GBuffer.DiffuseColor *= BaseMaterialCoverageOverWater;
 		DiffuseColor *= BaseMaterialCoverageOverWater;
 #endif
@@ -1210,143 +1212,27 @@ ShadingModel);
 #endif
 	Color += Emissive;
 
-
-
 #if MATERIAL_IS_SINGLE_LAYER_WATER
-	// Fade out the material contribution over to water contribution according to alpha.
-	const bool CameraIsUnderWater = false;
-	
-	float2 WaterNDC = MaterialParameters.ScreenPosition.xy / MaterialParameters.ScreenPosition.w;
-	float2 ViewportUV = WaterNDC * ResolvedView.ScreenPositionScaleBias.xy + ResolvedView.ScreenPositionScaleBias.wz;
-	float SceneCustomDepth = OpaqueBasePass.SceneCustomDepthTexture.SampleLevel(OpaqueBasePass.SceneCustomDepthSampler, ViewportUV, 0).r;
-	if (SceneCustomDepth > MaterialParameters.SvPosition.z)
-	{
-		clip(-1); // Custom depth Z is used to clip water to avoid it being visible inside boats for instance.
-	}
-
-	float WaterDepth = ConvertFromDeviceZ(MaterialParameters.SvPosition.z);
-	if (WaterVisibility > 0.0f) // DeltaDepth > 0.0f should not be needed as we ARE on a water surface TODO simplify that
 	{
+		const bool CameraIsUnderWater = false;	// Fade out the material contribution over to water contribution according to material opacity.
+		const float3 SunIlluminance = ResolvedView.DirectionalLightColor.rgb * PI;	// times PI because it is divided by PI on CPU (=luminance) and we want illuminance here. 
+		const float3 WaterDiffuseIndirectIlluminance = DiffuseIndirectLighting * PI;// DiffuseIndirectLighting is luminance. So we need to multiply by PI to get illuminance.
 
-		float PixelSceneDepth = ConvertFromDeviceZ(OpaqueBasePass.SceneDepthWithoutSingleLayerWaterTexture.SampleLevel(OpaqueBasePass.SceneDepthWithoutSingleLayerWaterSampler, ViewportUV, 0).r);
-		float DeltaDepth = CameraIsUnderWater ? WaterDepth : PixelSceneDepth - WaterDepth;
-
-#if 1
-		//////////////////////////////////////////////////////////////////////////////////////////////////// BEGIN Refration TODO make common with DistoAccumulatePS
-		float2 Refraction = GetMaterialRefraction(PixelMaterialInputs);// X = IOR, Y = Refraction Bias
-		float3 ViewNormal = normalize(TransformWorldVectorToView(MaterialParameters.WorldNormal));
-#if REFRACTION_USE_PIXEL_NORMAL_OFFSET
-		half3 ViewVertexNormal = TransformTangentVectorToView(MaterialParameters, float3(0, 0, 1));
-		// Treat Refraction of 1 as no refraction, to be consistent with IOR mode
-		float2 ViewportUVDistortion = (ViewVertexNormal.xy - ViewNormal.xy) * (Refraction.x - 1);
-#else
-		// we assume the camera is in air
-		float AirIOR = 1.0f;
-		float2 ViewportUVDistortion = ViewNormal.xy * (Refraction.x - AirIOR);
-#endif
-		float2 BufferUVDistortion = ViewportUVDistortion * ResolvedView.ViewSizeAndInvSize.xy * ResolvedView.BufferSizeAndInvSize.zw;
-		// When ISR is enabled we store two FOVs in the DistortionParams and compute the aspect ratio here
-		const float4 DistortionParams = 1.0f;
-		const float OffsetFudgeFactor = 1.0f;
-		float InvTanHalfFov = DistortionParams.x;
-		float Ratio = ResolvedView.ViewSizeAndInvSize.y * ResolvedView.ViewSizeAndInvSize.w;
-		// Fix for Fov and aspect.
-		float2 FovFix = float2(InvTanHalfFov, Ratio*InvTanHalfFov);
-		BufferUVDistortion *= 0.2 * DistortionParams.zw * float2(OffsetFudgeFactor, -OffsetFudgeFactor) * FovFix;
-		////////////////////////////////////////////////////////////////////////////////////////////////////  END Refration
-
-		// TODO make distortion fading common with DistoAccumulatePS (right now 30cm to deltadepth fade out and 10m to distance fade out)
-		BufferUVDistortion *= saturate(DeltaDepth * 1.0 / 30.0f) * (1 - saturate(WaterDepth * 1.0 / 2000.0f)); 
-
-		float SceneDeviceZ = OpaqueBasePass.SceneDepthWithoutSingleLayerWaterTexture.SampleLevel(OpaqueBasePass.SceneDepthWithoutSingleLayerWaterSampler, saturate(ViewportUV + BufferUVDistortion), 0).r;
-		float SceneDepth = ConvertFromDeviceZ(SceneDeviceZ);
-		if (SceneDepth < WaterDepth)
-		{
-			SceneDepth = PixelSceneDepth;
-		}
-		else
-		{
-			ViewportUV = saturate(ViewportUV + BufferUVDistortion);
-		}
-		DeltaDepth = CameraIsUnderWater ? WaterDepth : SceneDepth - WaterDepth;
-#endif
-
-		const float3 ScatteringCoeff = max(0.0f, GetSingleLayerWaterMaterialOutput0(MaterialParameters));
-		const float3 AbsorptionCoeff = max(0.0f, GetSingleLayerWaterMaterialOutput1(MaterialParameters));
-		const float PhaseG = clamp(GetSingleLayerWaterMaterialOutput2(MaterialParameters), -1.0f, 1.0f);
-		const float PhaseValue = SchlickPhase(PhaseG, dot(ResolvedView.DirectionalLightDirection.xyz, MaterialParameters.CameraVector));
-		const float3 ExtinctionCoeff = ScatteringCoeff + AbsorptionCoeff;
-
-		// We also apply transmittance from light to under water surface. However, the scene has been lit by many sources already.
-		// So the transmittance toabove surface is simply approximated using the travel distance from the scene pixel to the water top, assuming a flat water surface.
-		// We also use the pixel SvPosition instead of the scene one pre refraction/distortion to avoid those extra ALUs.
-		// We cannot combine this transmittance with the transmittance from view because this would change the behavior of the analytical integration of light scattering integration.
-		const float3 SceneWorldPos = SvPositionToWorld(float4(MaterialParameters.SvPosition.xy, SceneDeviceZ, 1.0));
-		const float DistanceFromScenePixelToWaterTop = max(0.0, MaterialParameters.AbsoluteWorldPosition.z - SceneWorldPos.z);
-		const float3 SceneColor = exp(-DistanceFromScenePixelToWaterTop * ExtinctionCoeff) * (USE_PREEXPOSURE ? ResolvedView.OneOverPreExposure : 1.0f)
-			* OpaqueBasePass.SceneLuminanceWithoutSingleLayerWaterTexture.SampleLevel(OpaqueBasePass.SceneLuminanceWithoutSingleLayerWaterSampler, ViewportUV, 0).rgb;
-
-		float3 SunEi = ResolvedView.DirectionalLightColor.rgb * PI;	// times PI because it is divided by PI on CPU (=luminance) and we want illuminance here.  
-		float3 SunScattLi = PhaseValue * SunEi;
-		float3 AmbScattLi = PhaseValue * WaterIndirectDiffuse * PI;	// WaterIndirectDiffuse is luminance? So we need to multiply by PI to get illuminance. That we can then get through the phase function.
-#if 1
-		float DirectionalLightShadow = ComputeDirectionalLightDynamicShadowing(MaterialParameters.AbsoluteWorldPosition, GBuffer.Depth);
-
-		const float3 OpticalDepth = ExtinctionCoeff * DeltaDepth;
-		float3 Transmittance = exp(-OpticalDepth);
-		float3 ScatteredLuminance = ScatteringCoeff * (AmbScattLi + SunScattLi * DirectionalLightShadow);
-		ScatteredLuminance = (ScatteredLuminance - ScatteredLuminance * Transmittance) / ExtinctionCoeff;
-#else
-		float3 Transmittance = 1.0f;
-		float3 ScatteredLuminance = 0.0f;
-		const float RayMarchMaxDistance = min(DeltaDepth, 200.0f);  // 20 meters
-		const float RayMarchStepSize = RayMarchMaxDistance / 10.0f; // Less samples wil lresult in a bit brighter look due to TransmittanceToLightThroughWater being 1 on a longer first sample. Would need it part of analiytical integration
-		const float ShadowDither = RayMarchStepSize * GBufferDither;
-		for (float s = 0.0f; s < RayMarchMaxDistance; s += RayMarchStepSize)
-		{
-			// Only jitter shadow map sampling to not lose energy on first sample
-			float Shadow = ComputeDirectionalLightDynamicShadowing(MaterialParameters.AbsoluteWorldPosition - (s + ShadowDither)*MaterialParameters.CameraVector, GBuffer.Depth);
-
-			float3 WP = MaterialParameters.AbsoluteWorldPosition - s * MaterialParameters.CameraVector;
-			float WaterHeightAboveSample = max(0.0, MaterialParameters.AbsoluteWorldPosition.z - WP.z);
-			float3 TransmittanceToLightThroughWater = 1.0;												// no self shadow, same energy as above analytical solution
-			//float3 TransmittanceToLightThroughWater = exp(-ExtinctionCoeff * WaterHeightAboveSample);	// self shadow as transmittance to water level, close to reference, depends a bit on sample count due to first sample being critical for dense medium
-
-			float3 SampleTransmittance = exp(-ExtinctionCoeff * RayMarchStepSize); // Constant
-			float3 SS = (ScatteringCoeff * TransmittanceToLightThroughWater * (SunScattLi * Shadow + AmbScattLi));
-			ScatteredLuminance += Transmittance * (SS - SS * SampleTransmittance) / ExtinctionCoeff;
-			Transmittance *= SampleTransmittance;
-		}
-
-		// The rest of the medium
-		const float3 OpticalDepth2 = ExtinctionCoeff * max(0.0, DeltaDepth - RayMarchMaxDistance);
-		if (any(OpticalDepth2 > 0.0f))
-		{
-			float3 Transmittance2 = exp(-OpticalDepth2);
-			float3 ScatteredLuminance2 = ScatteringCoeff * (SunScattLi + AmbScattLi);
-			ScatteredLuminance += Transmittance * (ScatteredLuminance2 - ScatteredLuminance2 * Transmittance2) / ExtinctionCoeff;
-			Transmittance *= Transmittance2;
-		}
-#endif
-
-		// Apply Fresnel effect
+		// Evaluate Fresnel effect
 		const float3 N = MaterialParameters.WorldNormal;
 		const float3 V = MaterialParameters.CameraVector;
-		const float NoV = max(0.0, dot(N, V));
-		float3 EnvBrdf = EnvBRDF(GBuffer.SpecularColor, GBuffer.Roughness, NoV);
-		ScatteredLuminance *= CameraIsUnderWater ? 1.0 : (1.0 - EnvBrdf); // Under water is less visibule due to Frensel effect
-		Transmittance      *= CameraIsUnderWater ? (1.0 - EnvBrdf) : 1.0; // Above	"		"		"		"		"
-
-		// Add single in-scattering apply colored transmittance to scene color
-		Color += ScatteredLuminance * WaterVisibility;
-		Color += SceneColor * Transmittance * WaterVisibility ;
+		const float3 EnvBrdf = EnvBRDF(GBuffer.SpecularColor, GBuffer.Roughness, max(0.0, dot(N, V)));
+
+		Color += EvaluateWaterVolumeLighting(
+			MaterialParameters, PixelMaterialInputs, ResolvedView,
+			OpaqueBasePass.SceneDepthWithoutSingleLayerWaterTexture, OpaqueBasePass.SceneDepthWithoutSingleLayerWaterSampler,
+			OpaqueBasePass.SceneLuminanceWithoutSingleLayerWaterTexture, OpaqueBasePass.SceneLuminanceWithoutSingleLayerWaterSampler,
+			OpaqueBasePass.SceneCustomDepthTexture, OpaqueBasePass.SceneCustomDepthSampler,
+			SunIlluminance, WaterDiffuseIndirectIlluminance, EnvBrdf,
+			CameraIsUnderWater, WaterVisibility).Luminance;
 	}
-
-
 #endif // MATERIAL_IS_SINGLE_LAYER_WATER
 
-
-
 	#if MATERIAL_DOMAIN_POSTPROCESS
 		#if MATERIAL_OUTPUT_OPACITY_AS_ALPHA
 			Out.MRT[0] = half4(Color, Opacity);
diff --git a/Engine/Shaders/Private/MobileBasePassPixelShader.usf b/Engine/Shaders/Private/MobileBasePassPixelShader.usf
--- a/Engine/Shaders/Private/MobileBasePassPixelShader.usf
+++ b/Engine/Shaders/Private/MobileBasePassPixelShader.usf
@@ -23,7 +23,8 @@
 #include "MobileGGX.ush"
 
 #if MATERIAL_IS_SINGLE_LAYER_WATER
-#include "ParticipatingMediaCommon.ush"
+#define SIMPLE_SINGLE_LAYER_WATER 1
+#include "SingleLayerWaterCommon.ush"
 #endif
 
 // Add fading CSM plane:
@@ -440,7 +441,7 @@ void Main(
 #if MATERIAL_IS_SINGLE_LAYER_WATER
 	// Fade out diffuse as this will be hanlded by the single scattering lighting. when over the water surface.
 	// We keep the SpecularColor for sun/water interactions
-	float3 WaterIndirectDiffuse = 0;
+	float3 WaterDiffuseIndirectLuminance = 0;
 	DiffuseColor *= BaseMaterialCoverageOverWater;
 #endif
 	
@@ -586,7 +587,7 @@ void Main(
 			// TODO: Also need to do specular.
 			float3 SkyDiffuseLighting = GetSkySHDiffuseSimple(MaterialParameters.WorldNormal);
 		#if MATERIAL_IS_SINGLE_LAYER_WATER
-			WaterIndirectDiffuse += SkyDiffuseLighting;
+			WaterDiffuseIndirectLuminance += SkyDiffuseLighting;
 		#endif
 			Color += SkyDiffuseLighting * ResolvedView.SkyLightColor.rgb * DiffuseColor * MaterialAO;
 		#endif
@@ -654,68 +655,27 @@ void Main(
 		 
 	half3 Emissive = GetMaterialEmissive(PixelMaterialInputs);
 
-	Color += Emissive; 
-
-
+	Color += Emissive;
 
 #if MATERIAL_IS_SINGLE_LAYER_WATER
-	// Fade out the material contribution over to water contribution according to alpha.
-	const bool CameraIsUnderWater = false;
-
-	float WaterDepth = ConvertFromDeviceZ(MaterialParameters.SvPosition.z);
-	if (WaterVisibility > 0.0f) // DeltaDepth > 0.0f should not be needed as we ARE on a water surface TODO simplify that
 	{
-		float2 NDC = MaterialParameters.ScreenPosition.xy / MaterialParameters.ScreenPosition.w;
-		float2 ViewportUV = NDC * ResolvedView.ScreenPositionScaleBias.xy + ResolvedView.ScreenPositionScaleBias.wz;
-
-		float PixelSceneDeviceZ = MobileSceneTextures.SceneDepthTexture.SampleLevel(MobileSceneTextures.SceneDepthTextureSampler, ViewportUV, 0).r;
-		float PixelSceneDepth = ConvertFromDeviceZ(PixelSceneDeviceZ);
-		float DeltaDepth = CameraIsUnderWater ? WaterDepth : PixelSceneDepth - WaterDepth;
-
-
-		const float3 ScatteringCoeff = max(0.0f, GetSingleLayerWaterMaterialOutput0(MaterialParameters));
-		const float3 AbsorptionCoeff = max(0.0f, GetSingleLayerWaterMaterialOutput1(MaterialParameters));
-		const float PhaseG = clamp(GetSingleLayerWaterMaterialOutput2(MaterialParameters), -1.0f, 1.0f);
-		const float PhaseValue = SchlickPhase(PhaseG, dot(ResolvedView.DirectionalLightDirection.xyz, MaterialParameters.CameraVector));
-		const float3 ExtinctionCoeff = ScatteringCoeff + AbsorptionCoeff;
-		const float3 OpticalDepth = ExtinctionCoeff * DeltaDepth;
-		float3 Transmittance = exp(-OpticalDepth);
-
-		// We also apply transmittance from light to under water surface. However, the scene has been lit by many sources already.
-		// So the transmittance toabove surface is simply approximated using the travel distance from the scene pixel to the water top, assuming a flat water surface.
-		// We also use the pixel SvPosition instead of the scene one pre refraction/distortion to avoid those extra ALUs.
-		// We cannot combine this transmittance with the transmittance from view because this would change the behavior of the analytical integration of light scattering integration.
-		const float3 SceneWorldPos = SvPositionToWorld(float4(MaterialParameters.SvPosition.xy, PixelSceneDeviceZ, 1.0));
-		const float DistanceFromScenePixelToWaterTop = max(0.0, MaterialParameters.AbsoluteWorldPosition.z - SceneWorldPos.z);
-		const float FromSurfaceToLightTransmittance = exp(-DistanceFromScenePixelToWaterTop * ExtinctionCoeff);
-
-		const float DirectionalLightShadow = Shadow;
-
-		float3 SunEi = ResolvedView.DirectionalLightColor.rgb * PI;	// times PI because it is divided by PI on CPU (=luminance) and we want illuminance here.  
-		float3 SunScattLi = PhaseValue * SunEi;
-		float3 AmbScattLi = PhaseValue * WaterIndirectDiffuse * PI;	// WaterIndirectDiffuse is luminance? So we need to multiply by PI to get illuminance. That we can then get through the phase function.
-
-		float3 ScatteredLuminance = ScatteringCoeff * (AmbScattLi + SunScattLi * DirectionalLightShadow);
-		ScatteredLuminance = (ScatteredLuminance - ScatteredLuminance * Transmittance) / ExtinctionCoeff;
-
-		// Apply Fresnel effect
-		float3 EnvBrdf = SpecularColor; // SpecularColor is not F0 as in BasePassPixelShader but EnvBRDF.
-		ScatteredLuminance *= CameraIsUnderWater ? 1.0 : (1.0 - EnvBrdf); // Under water is less visibule due to Frensel effect
-		Transmittance *= CameraIsUnderWater ? (1.0 - EnvBrdf) : 1.0; // Above	"		"		"		"		"
-
-		// Add single in-scattering apply colored transmittance to scene color
-		Color += ScatteredLuminance * WaterVisibility;
-//		Color += SceneColor * Transmittance * WaterVisibility;
-
-		Opacity = 1.0 - ((1.0 - Opacity) * dot(Transmittance*FromSurfaceToLightTransmittance, float3(1.0 / 3.0, 1.0 / 3.0, 1.0 / 3.0)));
+		const bool CameraIsUnderWater = false;	// Fade out the material contribution over to water contribution according to material opacity.
+		const float3 SunIlluminance = ResolvedView.DirectionalLightColor.rgb * PI;			// times PI because it is divided by PI on CPU (=luminance) and we want illuminance here. 
+		const float3 WaterDiffuseIndirectIlluminance = WaterDiffuseIndirectLuminance * PI;	// DiffuseIndirectLighting is luminance. So we need to multiply by PI to get illuminance.
+		const float3 EnvBrdf = SpecularColor; // SpecularColor is not F0 as in BasePassPixelShader, it is EnvBRDF.
+
+		WaterVolumeLightingOutput WaterLighting = EvaluateWaterVolumeLighting(
+			MaterialParameters, PixelMaterialInputs, ResolvedView,
+			MobileSceneTextures.SceneDepthTexture, MobileSceneTextures.SceneDepthTextureSampler,
+			Shadow,
+			SunIlluminance, WaterDiffuseIndirectIlluminance, EnvBrdf,
+			CameraIsUnderWater, WaterVisibility);
+
+		// Accumulate luminance and occlude the background according to transmittance to view and mean transmittance to lights.
+		Color += WaterLighting.Luminance;
+		Opacity = 1.0 - ((1.0 - Opacity) * dot(WaterLighting.WaterToSceneToLightTransmittance, float3(1.0 / 3.0, 1.0 / 3.0, 1.0 / 3.0)));
 	}
-#endif
-
-
-
-
-
-
+#endif // MATERIAL_IS_SINGLE_LAYER_WATER
 
 	half4 VertexFog = half4(0, 0, 0, 1);
 
diff --git a/Engine/Shaders/Private/SingleLayerWaterCommon.ush b/Engine/Shaders/Private/SingleLayerWaterCommon.ush
new file mode 100644
--- /dev/null
+++ b/Engine/Shaders/Private/SingleLayerWaterCommon.ush
@@ -0,0 +1,185 @@
+/*=============================================================================
+	SingleLayerWaterCommon.ush
+=============================================================================*/
+
+#include "Common.ush"
+#include "ParticipatingMediaCommon.ush" 
+
+
+#ifndef SIMPLE_SINGLE_LAYER_WATER
+#define SIMPLE_SINGLE_LAYER_WATER 0
+#endif
+
+struct WaterVolumeLightingOutput
+{
+	float3 Luminance;
+	float3 WaterToSceneTransmittance;
+	float3 WaterToSceneToLightTransmittance;
+};
+
+// Returns as luminance, the result of lighting scattering in water and the under water scene color contribution.
+WaterVolumeLightingOutput EvaluateWaterVolumeLighting(
+	FMaterialPixelParameters MaterialParameters, FPixelMaterialInputs PixelMaterialInputs, ViewState ResolvedView,
+	Texture2D SceneDepthWithoutSingleLayerWaterTexture, SamplerState SceneDepthWithoutSingleLayerWaterSampler,
+#if SIMPLE_SINGLE_LAYER_WATER
+	float DirectionalLightShadow,
+#else
+	Texture2D SceneLuminanceWithoutSingleLayerWaterTexture, SamplerState SceneLuminanceWithoutSingleLayerWaterSampler,
+	Texture2D SceneCustomDepthTexture, SamplerState SceneCustomDepthSampler,
+#endif // SIMPLE_SINGLE_LAYER_WATER
+	float3 SunIlluminance, float3 AmbiantIlluminance, float3 EnvBrdf,
+	bool CameraIsUnderWater, float WaterVisibility)
+{
+	WaterVolumeLightingOutput Output;
+	Output.Luminance = 0.0f;
+	Output.WaterToSceneTransmittance = 1.0f;
+	Output.WaterToSceneToLightTransmittance = 1.0f;
+
+	float2 WaterNDC = MaterialParameters.ScreenPosition.xy / MaterialParameters.ScreenPosition.w;
+	float2 ViewportUV = WaterNDC * ResolvedView.ScreenPositionScaleBias.xy + ResolvedView.ScreenPositionScaleBias.wz;
+
+#if SIMPLE_SINGLE_LAYER_WATER==0
+	float SceneCustomDepth = SceneCustomDepthTexture.SampleLevel(SceneCustomDepthSampler, ViewportUV, 0).r;
+	if (SceneCustomDepth > MaterialParameters.SvPosition.z)
+	{
+		clip(-1); // Custom depth Z is used to clip, this in order to avoid water being visible inside boats for instance.
+	}
+#endif // !SIMPLE_SINGLE_LAYER_WATER
+
+	float WaterDepth = ConvertFromDeviceZ(MaterialParameters.SvPosition.z);
+	if (WaterVisibility > 0.0f)
+	{
+		float PixelSceneDeviceZ = SceneDepthWithoutSingleLayerWaterTexture.SampleLevel(SceneDepthWithoutSingleLayerWaterSampler, ViewportUV, 0).r;
+		float PixelSceneDepth = ConvertFromDeviceZ(PixelSceneDeviceZ);
+		
+#if SIMPLE_SINGLE_LAYER_WATER
+
+		float SceneDeviceZ = PixelSceneDeviceZ;
+		float SceneDepth = PixelSceneDepth;
+
+#else  // !SIMPLE_SINGLE_LAYER_WATER
+
+		//////////////////////////////////////////////////////////////////////////////////////////////////// BEGIN Refration TODO make common with DistoAccumulatePS
+		float2 Refraction = GetMaterialRefraction(PixelMaterialInputs);// X = IOR, Y = Refraction Bias
+		float3 ViewNormal = normalize(TransformWorldVectorToView(MaterialParameters.WorldNormal));
+#if REFRACTION_USE_PIXEL_NORMAL_OFFSET
+		half3 ViewVertexNormal = TransformTangentVectorToView(MaterialParameters, float3(0, 0, 1));
+		// Treat Refraction of 1 as no refraction, to be consistent with IOR mode
+		float2 ViewportUVDistortion = (ViewVertexNormal.xy - ViewNormal.xy) * (Refraction.x - 1);
+#else
+		// we assume the camera is in air
+		float AirIOR = 1.0f;
+		float2 ViewportUVDistortion = ViewNormal.xy * (Refraction.x - AirIOR);
+#endif
+		float2 BufferUVDistortion = ViewportUVDistortion * ResolvedView.ViewSizeAndInvSize.xy * ResolvedView.BufferSizeAndInvSize.zw;
+		// When ISR is enabled we store two FOVs in the DistortionParams and compute the aspect ratio here
+		const float4 DistortionParams = 1.0f;
+		const float OffsetFudgeFactor = 1.0f;
+		float InvTanHalfFov = DistortionParams.x;
+		float Ratio = ResolvedView.ViewSizeAndInvSize.y * ResolvedView.ViewSizeAndInvSize.w;
+		// Fix for Fov and aspect.
+		float2 FovFix = float2(InvTanHalfFov, Ratio*InvTanHalfFov);
+		BufferUVDistortion *= 0.2 * DistortionParams.zw * float2(OffsetFudgeFactor, -OffsetFudgeFactor) * FovFix;
+		////////////////////////////////////////////////////////////////////////////////////////////////////  END Refration
+
+		// TODO make distortion fading common with DistoAccumulatePS (right now 30cm to BehindWaterDeltaDepth fade out and 10m to distance fade out)
+		BufferUVDistortion *= saturate((PixelSceneDepth - WaterDepth) * 1.0 / 30.0f) * (1 - saturate(WaterDepth * 1.0 / 2000.0f));
+
+
+
+
+		float SceneDeviceZ = SceneDepthWithoutSingleLayerWaterTexture.SampleLevel(SceneDepthWithoutSingleLayerWaterSampler, saturate(ViewportUV + BufferUVDistortion), 0).r;
+		float SceneDepth = ConvertFromDeviceZ(SceneDeviceZ);
+		if (SceneDepth < WaterDepth)
+		{
+			SceneDepth = PixelSceneDepth; // The pixel we try to sample is closer than the water pixel: falback to no distortion
+		}
+		else
+		{
+			ViewportUV = saturate(ViewportUV + BufferUVDistortion);
+		}
+
+#endif // !SIMPLE_SINGLE_LAYER_WATER
+
+		const float BehindWaterDeltaDepth = CameraIsUnderWater ? WaterDepth : SceneDepth - WaterDepth;
+
+		const float3 ScatteringCoeff = max(0.0f, GetSingleLayerWaterMaterialOutput0(MaterialParameters));
+		const float3 AbsorptionCoeff = max(0.0f, GetSingleLayerWaterMaterialOutput1(MaterialParameters));
+		const float PhaseG = clamp(GetSingleLayerWaterMaterialOutput2(MaterialParameters), -1.0f, 1.0f);
+		const float PhaseValue = SchlickPhase(PhaseG, dot(ResolvedView.DirectionalLightDirection.xyz, MaterialParameters.CameraVector));
+		const float3 ExtinctionCoeff = ScatteringCoeff + AbsorptionCoeff;
+
+		// We also apply transmittance from light to under water surface. However, the scene has been lit by many sources already.
+		// So the transmittance toabove surface is simply approximated using the travel distance from the scene pixel to the water top, assuming a flat water surface.
+		// We cannot combine this transmittance with the transmittance from view because this would change the behavior of the analytical integration of light scattering integration.
+		const float3 BehindWaterSceneWorldPos = SvPositionToWorld(float4(MaterialParameters.SvPosition.xy, SceneDeviceZ, 1.0));
+		const float DistanceFromScenePixelToWaterTop = max(0.0, MaterialParameters.AbsoluteWorldPosition.z - BehindWaterSceneWorldPos.z);
+		const float3 MeanTransmittanceToLightSources = exp(-DistanceFromScenePixelToWaterTop * ExtinctionCoeff);
+
+#if SIMPLE_SINGLE_LAYER_WATER
+		const float3 BehindWaterSceneLuminance = 0.0f; // Cannot read back the scene color in this case
+#else
+		// We use the pixel SvPosition instead of the scene one pre refraction/distortion to avoid those extra ALUs.
+		float3 BehindWaterSceneLuminance = SceneLuminanceWithoutSingleLayerWaterTexture.SampleLevel(SceneLuminanceWithoutSingleLayerWaterSampler, ViewportUV, 0).rgb;
+		BehindWaterSceneLuminance = MeanTransmittanceToLightSources * (USE_PREEXPOSURE ? ResolvedView.OneOverPreExposure : 1.0f) * BehindWaterSceneLuminance;
+#endif
+
+		float3 SunScattLuminance = PhaseValue * SunIlluminance;
+		float3 AmbScattLuminance = PhaseValue * AmbiantIlluminance;
+
+#if 1 || SIMPLE_SINGLE_LAYER_WATER
+#if SIMPLE_SINGLE_LAYER_WATER==0
+		float DirectionalLightShadow = ComputeDirectionalLightDynamicShadowing(MaterialParameters.AbsoluteWorldPosition, /*GBuffer.Depth*/ MaterialParameters.ScreenPosition.w);
+#endif
+		const float3 OpticalDepth = ExtinctionCoeff * BehindWaterDeltaDepth;
+		float3 Transmittance = exp(-OpticalDepth);
+		float3 ScatteredLuminance = ScatteringCoeff * (AmbScattLuminance + SunScattLuminance * DirectionalLightShadow);
+		ScatteredLuminance = (ScatteredLuminance - ScatteredLuminance * Transmittance) / ExtinctionCoeff;
+#else
+		// TODO Make the volumetric shadow part work again
+		float3 Transmittance = 1.0f;
+		float3 ScatteredLuminance = 0.0f;
+		const float RayMarchMaxDistance = min(BehindWaterDeltaDepth, 200.0f);  // 20 meters
+		const float RayMarchStepSize = RayMarchMaxDistance / 10.0f; // Less samples wil lresult in a bit brighter look due to TransmittanceToLightThroughWater being 1 on a longer first sample. Would need it part of analiytical integration
+		const float ShadowDither = RayMarchStepSize * GBufferDither;
+		for (float s = 0.0f; s < RayMarchMaxDistance; s += RayMarchStepSize)
+		{
+			// Only jitter shadow map sampling to not lose energy on first sample
+			float Shadow = ComputeDirectionalLightDynamicShadowing(MaterialParameters.AbsoluteWorldPosition - (s + ShadowDither)*MaterialParameters.CameraVector, GBuffer.Depth);
+
+			float3 WP = MaterialParameters.AbsoluteWorldPosition - s * MaterialParameters.CameraVector;
+			float WaterHeightAboveSample = max(0.0, MaterialParameters.AbsoluteWorldPosition.z - WP.z);
+			float3 TransmittanceToLightThroughWater = 1.0;												// no self shadow, same energy as above analytical solution
+			//float3 TransmittanceToLightThroughWater = exp(-ExtinctionCoeff * WaterHeightAboveSample);	// self shadow as transmittance to water level, close to reference, depends a bit on sample count due to first sample being critical for dense medium
+
+			float3 SampleTransmittance = exp(-ExtinctionCoeff * RayMarchStepSize); // Constant
+			float3 SS = (ScatteringCoeff * TransmittanceToLightThroughWater * (SunScattLuminance * Shadow + AmbScattLuminance));
+			ScatteredLuminance += Transmittance * (SS - SS * SampleTransmittance) / ExtinctionCoeff;
+			Transmittance *= SampleTransmittance;
+		}
+
+		// The rest of the medium
+		const float3 OpticalDepth2 = ExtinctionCoeff * max(0.0, BehindWaterDeltaDepth - RayMarchMaxDistance);
+		if (any(OpticalDepth2 > 0.0f))
+		{
+			float3 Transmittance2 = exp(-OpticalDepth2);
+			float3 ScatteredLuminance2 = ScatteringCoeff * (SunScattLuminance + AmbScattLuminance);
+			ScatteredLuminance += Transmittance * (ScatteredLuminance2 - ScatteredLuminance2 * Transmittance2) / ExtinctionCoeff;
+			Transmittance *= Transmittance2;
+		}
+#endif
+
+		// Apply Fresnel effect to out-scattering towards the view
+		ScatteredLuminance *= CameraIsUnderWater ? 1.0 : (1.0 - EnvBrdf);	// Under water is less visible due to Fresnel effect
+		Transmittance *= CameraIsUnderWater ? (1.0 - EnvBrdf) : 1.0;		// Above	"		"		"		"		"
+
+		// Add single in-scattering apply colored transmittance to scene color
+		Output.Luminance = WaterVisibility * (ScatteredLuminance + Transmittance * BehindWaterSceneLuminance);
+		Output.WaterToSceneTransmittance = Transmittance;
+		Output.WaterToSceneToLightTransmittance = Transmittance * MeanTransmittanceToLightSources;
+	}
+
+	return Output;
+}
+
+
