Encode cloth's fuzz component with gamma 2.0 in order to increse its precision and remove some color banding there

diff --git a/Engine/Shaders/Private/ReflectionEnvironmentPixelShader.usf b/Engine/Shaders/Private/ReflectionEnvironmentPixelShader.usf
--- a/Engine/Shaders/Private/ReflectionEnvironmentPixelShader.usf
+++ b/Engine/Shaders/Private/ReflectionEnvironmentPixelShader.usf
@@ -149,7 +149,8 @@ float3 SkyLightDiffuse(FGBufferData GBuffer, float AmbientOcclusion, float2 Buff
 
 	if (GBuffer.ShadingModelID == SHADINGMODELID_CLOTH)
 	{
-		DiffuseColor += GBuffer.CustomData.rgb * GBuffer.CustomData.a;
+		float3 ClothFuzz = ExtractSubsurfaceColor(GBuffer);
+		DiffuseColor += ClothFuzz * GBuffer.CustomData.a;
 	}
 
 	// Compute the preconvolved incoming lighting with the bent normal direction
diff --git a/Engine/Shaders/Private/ShadingModels.ush b/Engine/Shaders/Private/ShadingModels.ush
--- a/Engine/Shaders/Private/ShadingModels.ush
+++ b/Engine/Shaders/Private/ShadingModels.ush
@@ -414,7 +414,7 @@ FDirectLighting SubsurfaceProfileBxDF( FGBufferData GBuffer, half3 N, half3 V, h
 
 FDirectLighting ClothBxDF( FGBufferData GBuffer, half3 N, half3 V, half3 L, float Falloff, float NoL, FAreaLight AreaLight, FShadowTerms Shadow )
 {
-	const float3 FuzzColor	= saturate(GBuffer.CustomData.rgb);
+	const float3 FuzzColor	= ExtractSubsurfaceColor(GBuffer);
 	const float  Cloth		= saturate(GBuffer.CustomData.a);
 
 	BxDFContext Context;
diff --git a/Engine/Shaders/Private/ShadingModelsMaterial.ush b/Engine/Shaders/Private/ShadingModelsMaterial.ush
--- a/Engine/Shaders/Private/ShadingModelsMaterial.ush
+++ b/Engine/Shaders/Private/ShadingModelsMaterial.ush
@@ -145,7 +145,7 @@ void SetGBufferForShadingModel(
 #if MATERIAL_SHADINGMODEL_CLOTH
 	else if (ShadingModel == SHADINGMODELID_CLOTH)
 	{
-		GBuffer.CustomData.rgb = SubsurfaceColor;
+		GBuffer.CustomData.rgb = EncodeSubsurfaceColor(SubsurfaceColor);
 		GBuffer.CustomData.a = saturate( GetMaterialCustomData0(MaterialParameters) );	// Cloth
 		GBuffer.IndirectIrradiance *= 1 - GBuffer.CustomData.a;
 	}
