Using uniform phase function for the water ambiant light source.

diff --git a/Engine/Shaders/Private/SingleLayerWaterCommon.ush b/Engine/Shaders/Private/SingleLayerWaterCommon.ush
--- a/Engine/Shaders/Private/SingleLayerWaterCommon.ush
+++ b/Engine/Shaders/Private/SingleLayerWaterCommon.ush
@@ -106,7 +106,7 @@ WaterVolumeLightingOutput EvaluateWaterVolumeLighting(
 		const float3 ScatteringCoeff = max(0.0f, GetSingleLayerWaterMaterialOutput0(MaterialParameters));
 		const float3 AbsorptionCoeff = max(0.0f, GetSingleLayerWaterMaterialOutput1(MaterialParameters));
 		const float PhaseG = clamp(GetSingleLayerWaterMaterialOutput2(MaterialParameters), -1.0f, 1.0f);
-		const float PhaseValue = SchlickPhase(PhaseG, dot(ResolvedView.DirectionalLightDirection.xyz, MaterialParameters.CameraVector));
+		const float DirLightPhaseValue = SchlickPhase(PhaseG, dot(ResolvedView.DirectionalLightDirection.xyz, MaterialParameters.CameraVector));
 		const float3 ExtinctionCoeff = ScatteringCoeff + AbsorptionCoeff;
 
 		// We also apply transmittance from light to under water surface. However, the scene has been lit by many sources already.
@@ -124,8 +124,8 @@ WaterVolumeLightingOutput EvaluateWaterVolumeLighting(
 		BehindWaterSceneLuminance = MeanTransmittanceToLightSources * (USE_PREEXPOSURE ? ResolvedView.OneOverPreExposure : 1.0f) * BehindWaterSceneLuminance;
 #endif
 
-		float3 SunScattLuminance = PhaseValue * SunIlluminance;
-		float3 AmbScattLuminance = PhaseValue * AmbiantIlluminance;
+		float3 SunScattLuminance = DirLightPhaseValue * SunIlluminance;
+		float3 AmbScattLuminance = IsotropicPhase()   * AmbiantIlluminance;
 
 #if 1 || SIMPLE_SINGLE_LAYER_WATER
 #if SIMPLE_SINGLE_LAYER_WATER==0
