New water shading model.
Also fix expenssive pipeline issue.
Water no longer depends on IsUsedWithWater

diff --git a/Engine/Shaders/Private/BasePassPixelShader.usf b/Engine/Shaders/Private/BasePassPixelShader.usf
--- a/Engine/Shaders/Private/BasePassPixelShader.usf
+++ b/Engine/Shaders/Private/BasePassPixelShader.usf
@@ -14,7 +14,7 @@
 	#define EyeAdaptationStruct OpaqueBasePass
 #endif
 
-#if MATERIAL_IS_SINGLE_LAYER_WATER
+#if MATERIAL_SHADINGMODEL_SINGLELAYERWATER
 #define PreIntegratedGF			OpaqueBasePass.PreIntegratedGFTexture
 #define PreIntegratedGFSampler	OpaqueBasePass.PreIntegratedGFSampler
 #endif
@@ -68,7 +68,7 @@ float NormalCurvatureToRoughness(float3 WorldNormal)
 #endif
 
 #include "ShadingModelsMaterial.ush"
-#if MATERIAL_SHADINGMODEL_HAIR || SIMPLE_FORWARD_DIRECTIONAL_LIGHT || MATERIAL_IS_SINGLE_LAYER_WATER
+#if MATERIAL_SHADINGMODEL_HAIR || SIMPLE_FORWARD_DIRECTIONAL_LIGHT || MATERIAL_SHADINGMODEL_SINGLELAYERWATER
 #include "ShadingModels.ush"
 #endif
 
@@ -80,11 +80,11 @@ float NormalCurvatureToRoughness(float3 WorldNormal)
 #define EDITOR_ALPHA2COVERAGE (USE_EDITOR_COMPOSITING && SUPPORTS_PIXEL_COVERAGE)
 #define POST_PROCESS_SUBSURFACE ((MATERIAL_SHADINGMODEL_SUBSURFACE_PROFILE || MATERIAL_SHADINGMODEL_EYE) && USES_GBUFFER)
 
-#if TRANSLUCENCY_LIGHTING_SURFACE_LIGHTINGVOLUME || TRANSLUCENCY_LIGHTING_SURFACE_FORWARDSHADING || FORWARD_SHADING || MATERIAL_IS_SINGLE_LAYER_WATER
+#if TRANSLUCENCY_LIGHTING_SURFACE_LIGHTINGVOLUME || TRANSLUCENCY_LIGHTING_SURFACE_FORWARDSHADING || FORWARD_SHADING || MATERIAL_SHADINGMODEL_SINGLELAYERWATER
 #include "ForwardLightingCommon.ush"
 #endif
 
-#if  MATERIAL_IS_SINGLE_LAYER_WATER
+#if  MATERIAL_SHADINGMODEL_SINGLELAYERWATER
 #include "SingleLayerWaterCommon.ush"
 #endif
 
@@ -528,7 +528,7 @@ void GetPrecomputedIndirectLightingAndSkyLight(
 	#endif
 }
 
-#if SIMPLE_FORWARD_DIRECTIONAL_LIGHT || MATERIAL_IS_SINGLE_LAYER_WATER
+#if SIMPLE_FORWARD_DIRECTIONAL_LIGHT || MATERIAL_SHADINGMODEL_SINGLELAYERWATER
 
 float3 GetSimpleForwardLightingDirectionalLight(FGBufferData GBuffer, float3 DiffuseColor, float3 SpecularColor, float Roughness, float3 WorldNormal, float3 CameraVector)
 {
@@ -1070,7 +1070,7 @@ ShadingModel);
 
 		DiffuseColor += (DiffuseIndirectLighting * DiffuseColorForIndirect + SubsurfaceIndirectLighting * SubsurfaceColor) * AOMultiBounce( GBuffer.BaseColor, DiffOcclusion );
 
-#if MATERIAL_IS_SINGLE_LAYER_WATER
+#if MATERIAL_SHADINGMODEL_SINGLELAYERWATER
 		// Fade out diffuse as this will be handled by the single scattering lighting in water material.
 		// We do this after the just above GetPrecomputedIndirectLightingAndSkyLight to keep ambiant lighting avialable.
 		// We also keep the SpecularColor for sun/water interactions.
@@ -1082,7 +1082,7 @@ ShadingModel);
 
 			Color += BasePassInterpolants.VertexDiffuseLighting * GBuffer.DiffuseColor;
 
-		#elif FORWARD_SHADING || TRANSLUCENCY_LIGHTING_SURFACE_FORWARDSHADING || TRANSLUCENCY_LIGHTING_SURFACE_LIGHTINGVOLUME || MATERIAL_IS_SINGLE_LAYER_WATER
+		#elif FORWARD_SHADING || TRANSLUCENCY_LIGHTING_SURFACE_FORWARDSHADING || TRANSLUCENCY_LIGHTING_SURFACE_LIGHTINGVOLUME || MATERIAL_SHADINGMODEL_SINGLELAYERWATER
 			uint GridIndex = 0;
 
 			#if INSTANCED_STEREO
@@ -1094,14 +1094,14 @@ ShadingModel);
 			#if FEATURE_LEVEL >= FEATURE_LEVEL_SM5
 				GridIndex = ComputeLightGridCellIndex((uint2)(MaterialParameters.SvPosition.xy - ResolvedView.ViewRectMin.xy), MaterialParameters.SvPosition.w, EyeIndex);
 
-				#if FORWARD_SHADING || TRANSLUCENCY_LIGHTING_SURFACE_FORWARDSHADING || MATERIAL_IS_SINGLE_LAYER_WATER
+				#if FORWARD_SHADING || TRANSLUCENCY_LIGHTING_SURFACE_FORWARDSHADING || MATERIAL_SHADINGMODEL_SINGLELAYERWATER
 					const float Dither = InterleavedGradientNoise(MaterialParameters.SvPosition.xy, View.StateFrameIndexMod8);
 					float3 ForwardDirectLighting = GetForwardDirectLighting(GridIndex, MaterialParameters.AbsoluteWorldPosition, MaterialParameters.CameraVector, GBuffer, NearestResolvedDepthScreenUV, MaterialParameters.PrimitiveId, EyeIndex, Dither);
 					Color += ForwardDirectLighting;
 				#endif
 			#endif
 				// No IBL for water in deferred: that is skipped because it is done in the water composite pass. It should however be applied when using forward shading in order to get reflection without the water composite pass.
-			#if !(MATERIAL_SINGLE_SHADINGMODEL && MATERIAL_SHADINGMODEL_HAIR) && (!MATERIAL_IS_SINGLE_LAYER_WATER || FORWARD_SHADING) 
+			#if !(MATERIAL_SINGLE_SHADINGMODEL && MATERIAL_SHADINGMODEL_HAIR) && (!MATERIAL_SHADINGMODEL_SINGLELAYERWATER || FORWARD_SHADING) 
 				if (GBuffer.ShadingModelID != SHADINGMODELID_HAIR)
 				{
 					int SingleCaptureIndex = GetPrimitiveData(MaterialParameters.PrimitiveId).SingleCaptureIndex;
@@ -1212,7 +1212,7 @@ ShadingModel);
 #endif
 	Color += Emissive;
 
-#if MATERIAL_IS_SINGLE_LAYER_WATER
+#if MATERIAL_SHADINGMODEL_SINGLELAYERWATER
 	{
 		const bool CameraIsUnderWater = false;	// Fade out the material contribution over to water contribution according to material opacity.
 		const float3 SunIlluminance = ResolvedView.DirectionalLightColor.rgb * PI;	// times PI because it is divided by PI on CPU (=luminance) and we want illuminance here. 
@@ -1231,7 +1231,7 @@ ShadingModel);
 			SunIlluminance, WaterDiffuseIndirectIlluminance, EnvBrdf,
 			CameraIsUnderWater, WaterVisibility).Luminance;
 	}
-#endif // MATERIAL_IS_SINGLE_LAYER_WATER
+#endif // MATERIAL_SHADINGMODEL_SINGLELAYERWATER
 
 	#if MATERIAL_DOMAIN_POSTPROCESS
 		#if MATERIAL_OUTPUT_OPACITY_AS_ALPHA
diff --git a/Engine/Shaders/Private/ForwardLightingCommon.ush b/Engine/Shaders/Private/ForwardLightingCommon.ush
--- a/Engine/Shaders/Private/ForwardLightingCommon.ush
+++ b/Engine/Shaders/Private/ForwardLightingCommon.ush
@@ -133,7 +133,7 @@ float3 GetForwardDirectLighting(uint GridIndex, float3 WorldPosition, float3 Cam
 		LightData.ShadowMapChannelMask = UnpackShadowMapChannelMask(DirectionalLightData.DirectionalLightShadowMapChannelMask);
 
 		// We want to force the directional light shadow when using water material to see shadow on the water. This could be an option later.
-		#if ((MATERIALBLENDING_SOLID || MATERIALBLENDING_MASKED) && !MATERIAL_IS_SINGLE_LAYER_WATER)
+		#if ((MATERIALBLENDING_SOLID || MATERIALBLENDING_MASKED) && !MATERIAL_SHADINGMODEL_SINGLELAYERWATER)
 			float4 PreviewShadowMapChannelMask = UnpackShadowMapChannelMask(DirectionalLightData.DirectionalLightShadowMapChannelMask >> 4);
 			float DynamicShadowing = dot(PreviewShadowMapChannelMask, DynamicShadowFactors);
 
diff --git a/Engine/Shaders/Private/MobileBasePassPixelShader.usf b/Engine/Shaders/Private/MobileBasePassPixelShader.usf
--- a/Engine/Shaders/Private/MobileBasePassPixelShader.usf
+++ b/Engine/Shaders/Private/MobileBasePassPixelShader.usf
@@ -22,7 +22,7 @@
 #include "PlanarReflectionShared.ush"
 #include "MobileGGX.ush"
 
-#if MATERIAL_IS_SINGLE_LAYER_WATER
+#if MATERIAL_SHADINGMODEL_SINGLELAYERWATER
 #define SIMPLE_SINGLE_LAYER_WATER 1
 #include "SingleLayerWaterCommon.ush"
 #endif
@@ -438,7 +438,7 @@ void Main(
 	}
 #endif
 
-#if MATERIAL_IS_SINGLE_LAYER_WATER
+#if MATERIAL_SHADINGMODEL_SINGLELAYERWATER
 	// Fade out diffuse as this will be hanlded by the single scattering lighting. when over the water surface.
 	// We keep the SpecularColor for sun/water interactions
 	float3 WaterDiffuseIndirectLuminance = 0;
@@ -586,7 +586,7 @@ void Main(
 			//@mw todo
 			// TODO: Also need to do specular.
 			float3 SkyDiffuseLighting = GetSkySHDiffuseSimple(MaterialParameters.WorldNormal);
-		#if MATERIAL_IS_SINGLE_LAYER_WATER
+		#if MATERIAL_SHADINGMODEL_SINGLELAYERWATER
 			WaterDiffuseIndirectLuminance += SkyDiffuseLighting;
 		#endif
 			Color += SkyDiffuseLighting * ResolvedView.SkyLightColor.rgb * DiffuseColor * MaterialAO;
@@ -657,7 +657,7 @@ void Main(
 
 	Color += Emissive;
 
-#if MATERIAL_IS_SINGLE_LAYER_WATER
+#if MATERIAL_SHADINGMODEL_SINGLELAYERWATER
 	{
 		const bool CameraIsUnderWater = false;	// Fade out the material contribution over to water contribution according to material opacity.
 		const float3 SunIlluminance = ResolvedView.DirectionalLightColor.rgb * PI;			// times PI because it is divided by PI on CPU (=luminance) and we want illuminance here. 
@@ -675,7 +675,7 @@ void Main(
 		Color += WaterLighting.Luminance;
 		Opacity = 1.0 - ((1.0 - Opacity) * dot(WaterLighting.WaterToSceneToLightTransmittance, float3(1.0 / 3.0, 1.0 / 3.0, 1.0 / 3.0)));
 	}
-#endif // MATERIAL_IS_SINGLE_LAYER_WATER
+#endif // MATERIAL_SHADINGMODEL_SINGLELAYERWATER
 
 	half4 VertexFog = half4(0, 0, 0, 1);
 
@@ -693,7 +693,7 @@ void Main(
 #endif
 
 	// On mobile, water (an opaque material) is rendered as trnaslucent with forced premultiplied alpha blending (see MobileBasePass::SetTranslucentRenderState)
-	#if MATERIALBLENDING_ALPHACOMPOSITE || MATERIAL_IS_SINGLE_LAYER_WATER
+	#if MATERIALBLENDING_ALPHACOMPOSITE || MATERIAL_SHADINGMODEL_SINGLELAYERWATER
 		OutColor = half4(Color * VertexFog.a + VertexFog.rgb * Opacity, Opacity);
 	#elif MATERIALBLENDING_ALPHAHOLDOUT
 		// not implemented for holdout
diff --git a/Engine/Shaders/Private/ShadingCommon.ush b/Engine/Shaders/Private/ShadingCommon.ush
--- a/Engine/Shaders/Private/ShadingCommon.ush
+++ b/Engine/Shaders/Private/ShadingCommon.ush
@@ -13,7 +13,8 @@
 #define SHADINGMODELID_HAIR					7
 #define SHADINGMODELID_CLOTH				8
 #define SHADINGMODELID_EYE					9
-#define SHADINGMODELID_NUM					10
+#define SHADINGMODELID_SINGLELAYERWATER		10
+#define SHADINGMODELID_NUM					11
 #define SHADINGMODELID_MASK					0xF		// 4 bits reserved for ShadingModelID			
 
 // The flags are defined so that 0 value has no effect!
@@ -45,6 +46,8 @@ uint GetShadingModelID()
 	return SHADINGMODELID_CLOTH;
 #elif MATERIAL_SHADINGMODEL_EYE
 	return SHADINGMODELID_EYE;
+#elif MATERIAL_SHADINGMODEL_SINGLELAYERWATER
+	return SHADINGMODELID_SINGLELAYERWATER;
 #else
 	#if COMPILER_HLSLCC
 		// hlslcc needs a return type
diff --git a/Engine/Source/Editor/PixelInspector/Private/PixelInspectorDetailsCustomization.cpp b/Engine/Source/Editor/PixelInspector/Private/PixelInspectorDetailsCustomization.cpp
--- a/Engine/Source/Editor/PixelInspector/Private/PixelInspectorDetailsCustomization.cpp
+++ b/Engine/Source/Editor/PixelInspector/Private/PixelInspectorDetailsCustomization.cpp
@@ -140,6 +140,7 @@ void FPixelInspectorDetailsCustomization::CustomizeDetails(IDetailLayoutBuilder&
 	switch (MaterialShadingModel)
 	{
 		case MSM_DefaultLit:
+		case MSM_SingleLayerWater:
 		case MSM_Unlit:
 		{
 			DetailBuilder.HideProperty(SubSurfaceColorProp);
diff --git a/Engine/Source/Editor/PixelInspector/Private/PixelInspectorResult.cpp b/Engine/Source/Editor/PixelInspector/Private/PixelInspectorResult.cpp
--- a/Engine/Source/Editor/PixelInspector/Private/PixelInspectorResult.cpp
+++ b/Engine/Source/Editor/PixelInspector/Private/PixelInspectorResult.cpp
@@ -233,6 +233,8 @@ namespace PixelInspector
 			return EMaterialShadingModel::MSM_Cloth;
 		case PIXEL_INSPECTOR_SHADINGMODELID_EYE:
 			return EMaterialShadingModel::MSM_Eye;
+		case PIXEL_INSPECTOR_SHADINGMODELID_SINGLELAYERWATER:
+			return EMaterialShadingModel::MSM_SingleLayerWater;
 		};
 		return EMaterialShadingModel::MSM_DefaultLit;
 	}
@@ -276,6 +278,7 @@ namespace PixelInspector
 		{
 		case EMaterialShadingModel::MSM_Unlit:
 		case EMaterialShadingModel::MSM_DefaultLit:
+		case EMaterialShadingModel::MSM_SingleLayerWater:
 		{
 			SubSurfaceColor = FVector(0.0f);
 			Opacity = 0.0f;
diff --git a/Engine/Source/Editor/PixelInspector/Private/PixelInspectorResult.h b/Engine/Source/Editor/PixelInspector/Private/PixelInspectorResult.h
--- a/Engine/Source/Editor/PixelInspector/Private/PixelInspectorResult.h
+++ b/Engine/Source/Editor/PixelInspector/Private/PixelInspectorResult.h
@@ -15,6 +15,7 @@
 #define PIXEL_INSPECTOR_SHADINGMODELID_HAIR 7
 #define PIXEL_INSPECTOR_SHADINGMODELID_CLOTH 8
 #define PIXEL_INSPECTOR_SHADINGMODELID_EYE 9
+#define PIXEL_INSPECTOR_SHADINGMODELID_SINGLELAYERWATER 10
 #define PIXEL_INSPECTOR_SHADINGMODELID_MASK 0xF
 
 namespace PixelInspector
diff --git a/Engine/Source/Runtime/Engine/Classes/Engine/EngineTypes.h b/Engine/Source/Runtime/Engine/Classes/Engine/EngineTypes.h
--- a/Engine/Source/Runtime/Engine/Classes/Engine/EngineTypes.h
+++ b/Engine/Source/Runtime/Engine/Classes/Engine/EngineTypes.h
@@ -472,6 +472,7 @@ enum EMaterialShadingModel
 	MSM_Hair					UMETA(DisplayName="Hair"),
 	MSM_Cloth					UMETA(DisplayName="Cloth"),
 	MSM_Eye						UMETA(DisplayName="Eye"),
+	MSM_SingleLayerWater		UMETA(DisplayName="SingleLayerWater"),
 	/** Number of unique shading models. */
 	MSM_NUM						UMETA(Hidden),
 	/** Shading model will be determined by the Material Expression Graph,
diff --git a/Engine/Source/Runtime/Engine/Private/Materials/HLSLMaterialTranslator.h b/Engine/Source/Runtime/Engine/Private/Materials/HLSLMaterialTranslator.h
--- a/Engine/Source/Runtime/Engine/Private/Materials/HLSLMaterialTranslator.h
+++ b/Engine/Source/Runtime/Engine/Private/Materials/HLSLMaterialTranslator.h
@@ -949,11 +949,7 @@ public:
 			Chunk[MP_CustomData1]					= Material->CompilePropertyAndSetMaterialProperty(MP_CustomData1		,this);
 			Chunk[MP_AmbientOcclusion]				= Material->CompilePropertyAndSetMaterialProperty(MP_AmbientOcclusion	,this);
 
-			// We parse the node graph to see if a water material is used. 
-			// We cannot use output node registration from custom output node for water because CompileCustomOutputs is called below and we need to enabled some input before hand.
-			const bool bNodeGraphIsUsingSingleLayerWaterMaterialOutput = Material->IsUsingSingleLayerWaterMaterialOutput();
-
-			if (IsTranslucentBlendMode(BlendMode) || bNodeGraphIsUsingSingleLayerWaterMaterialOutput)
+			if (IsTranslucentBlendMode(BlendMode) || MaterialShadingModels.HasShadingModel(MSM_SingleLayerWater))
 			{
 				int32 UserRefraction = ForceCast(Material->CompilePropertyAndSetMaterialProperty(MP_Refraction, this), MCT_Float1);
 				int32 RefractionDepthBias = ForceCast(ScalarParameter(FName(TEXT("RefractionDepthBias")), Material->GetRefractionDepthBiasValue()), MCT_Float1);
@@ -1078,7 +1074,6 @@ ResourcesString = TEXT("");
 			MaterialCompilationOutput.bModifiesMeshPosition = bUsesPixelDepthOffset || bUsesWorldPositionOffset;
 			MaterialCompilationOutput.bUsesWorldPositionOffset = bUsesWorldPositionOffset;
 			MaterialCompilationOutput.bUsesPixelDepthOffset = bUsesPixelDepthOffset;
-			MaterialCompilationOutput.bUsesSingleLayerWaterMaterialOutput = bNodeGraphIsUsingSingleLayerWaterMaterialOutput;
 
 			// Fully rough if we have a roughness code chunk and it's constant and evaluates to 1.
 			bIsFullyRough = Chunk[MP_Roughness] != INDEX_NONE && IsMaterialPropertyUsed(MP_Roughness, Chunk[MP_Roughness], FLinearColor(1, 0, 0, 0), 1) == false;
@@ -1156,13 +1151,20 @@ ResourcesString = TEXT("");
 				Errorf(TEXT("Sky materials must be opaque and unlit. They are expected to completely replace the background."));
 			}
 
-			if (bNodeGraphIsUsingSingleLayerWaterMaterialOutput && (!MaterialShadingModels.IsLit() || BlendMode != BLEND_Opaque || !MaterialShadingModels.HasOnlyShadingModel(MSM_DefaultLit)))
-			{
-				Errorf(TEXT("Water materials must be opaque and lit and they only support the DefaultLit shading model."));
-			}
-			if (bNodeGraphIsUsingSingleLayerWaterMaterialOutput && Material->IsDistorted())
+			if (MaterialShadingModels.HasShadingModel(MSM_SingleLayerWater))
 			{
-				Errorf(TEXT("Water materials must not use Unreal built-in distortion."));
+				if (BlendMode != BLEND_Opaque)
+				{
+					Errorf(TEXT("SingleLayerWater materials must be opaque."));
+				}
+				if (!MaterialShadingModels.HasOnlyShadingModel(MSM_SingleLayerWater))
+				{
+					Errorf(TEXT("SingleLayerWater materials cannot be combined with other shading models.")); // Simply untested for now
+				}
+				if (Material->GetMaterialInterface() && !Material->GetMaterialInterface()->GetMaterial()->HasAnyExpressionsInMaterialAndFunctionsOfType<UMaterialExpressionSingleLayerWaterMaterialOutput>())
+				{
+					Errorf(TEXT("SingleLayerWater materials requires the use of SingleLayerWaterMaterial output node."));
+				}
 			}
 
 			bool bDBufferAllowed = IsUsingDBuffers(Platform);
@@ -1415,7 +1417,6 @@ ResourcesString = TEXT("");
 		OutEnvironment.SetDefine(TEXT("NEEDS_PARTICLE_TRANSFORM"), bUsesParticleTransform);
 		OutEnvironment.SetDefine(TEXT("USES_TRANSFORM_VECTOR"), bUsesTransformVector);
 		OutEnvironment.SetDefine(TEXT("WANT_PIXEL_DEPTH_OFFSET"), bUsesPixelDepthOffset);
-		OutEnvironment.SetDefine(TEXT("MATERIAL_IS_SINGLE_LAYER_WATER"), Material->IsUsingSingleLayerWaterMaterialOutput());
 		if (IsMetalPlatform(InPlatform))
 		{
 			OutEnvironment.SetDefine(TEXT("USES_WORLD_POSITION_OFFSET"), bUsesWorldPositionOffset);
@@ -1516,6 +1517,11 @@ ResourcesString = TEXT("");
 				OutEnvironment.SetDefine(TEXT("MATERIAL_SHADINGMODEL_EYE"), TEXT("1"));
 				NumSetMaterials++;
 			}
+			if (ShadingModels.HasShadingModel(MSM_SingleLayerWater))
+			{
+				OutEnvironment.SetDefine(TEXT("MATERIAL_SHADINGMODEL_SINGLELAYERWATER"), TEXT("1"));
+				NumSetMaterials++;
+			}
 
 			if (NumSetMaterials == 1)
 			{
diff --git a/Engine/Source/Runtime/Engine/Private/Materials/Material.cpp b/Engine/Source/Runtime/Engine/Private/Materials/Material.cpp
--- a/Engine/Source/Runtime/Engine/Private/Materials/Material.cpp
+++ b/Engine/Source/Runtime/Engine/Private/Materials/Material.cpp
@@ -158,7 +158,7 @@ int32 FMaterialResource::CompilePropertyAndSetMaterialProperty(EMaterialProperty
 			// Force basic opaque surfaces to skip masked/translucent-only attributes.
 			// Some features can force the material to create a masked variant which unintentionally runs this dormant code
 			if (GetMaterialDomain() != MD_Surface || GetBlendMode() != BLEND_Opaque || (GetShadingModels().IsLit() && !GetShadingModels().HasShadingModel(MSM_DefaultLit))
-				|| MaterialInterface->GetMaterial()->bUsedWithWater )//MaterialInterface->GetMaterial()->HasAnyExpressionsInMaterialAndFunctionsOfType<UMaterialExpressionSingleLayerWaterMaterialOutput>())
+				|| GetShadingModels().HasShadingModel(MSM_SingleLayerWater))
 			{
 				Ret = MaterialInterface->CompileProperty(Compiler, Property);
 			}
diff --git a/Engine/Source/Runtime/Engine/Private/Materials/MaterialInterface.cpp b/Engine/Source/Runtime/Engine/Private/Materials/MaterialInterface.cpp
--- a/Engine/Source/Runtime/Engine/Private/Materials/MaterialInterface.cpp
+++ b/Engine/Source/Runtime/Engine/Private/Materials/MaterialInterface.cpp
@@ -124,7 +124,7 @@ FMaterialRelevance UMaterialInterface::GetRelevance_Internal(const UMaterial* Ma
 		}
 
 		const bool bIsMobile = InFeatureLevel <= ERHIFeatureLevel::ES3_1;
-		const bool bUsesSingleLayerWaterMaterial = MaterialResource->MaterialUsesSingleLayerWater_GameThread();
+		const bool bUsesSingleLayerWaterMaterial = MaterialResource->GetShadingModels().HasShadingModel(MSM_SingleLayerWater);
 		const bool IsSinglePassWaterTranslucent = bIsMobile && bUsesSingleLayerWaterMaterial;
 
 		const EBlendMode BlendMode = (EBlendMode)GetBlendMode();
diff --git a/Engine/Source/Runtime/Engine/Private/Materials/MaterialShader.cpp b/Engine/Source/Runtime/Engine/Private/Materials/MaterialShader.cpp
--- a/Engine/Source/Runtime/Engine/Private/Materials/MaterialShader.cpp
+++ b/Engine/Source/Runtime/Engine/Private/Materials/MaterialShader.cpp
@@ -69,6 +69,7 @@ FString GetShadingModelString(EMaterialShadingModel ShadingModel)
 		case MSM_TwoSidedFoliage:	ShadingModelName = TEXT("MSM_TwoSidedFoliage"); break;
 		case MSM_Cloth:				ShadingModelName = TEXT("MSM_Cloth"); break;
 		case MSM_Eye:				ShadingModelName = TEXT("MSM_Eye"); break;
+		case MSM_SingleLayerWater:	ShadingModelName = TEXT("MSM_SingleLayerWater"); break;
 		default: ShadingModelName = TEXT("Unknown"); break;
 	}
 	return ShadingModelName;
@@ -152,7 +153,7 @@ void UpdateMaterialShaderCompilingStats(const FMaterial* Material)
 	{
 		INC_DWORD_STAT_BY(STAT_ShaderCompiling_NumUnlitMaterialShaders, 1);
 	}
-	else if (ShadingModels.HasAnyShadingModel({ MSM_DefaultLit, MSM_Subsurface, MSM_PreintegratedSkin, MSM_ClearCoat, MSM_Cloth, MSM_SubsurfaceProfile, MSM_TwoSidedFoliage }))
+	else if (ShadingModels.HasAnyShadingModel({ MSM_DefaultLit, MSM_Subsurface, MSM_PreintegratedSkin, MSM_ClearCoat, MSM_Cloth, MSM_SubsurfaceProfile, MSM_TwoSidedFoliage, MSM_SingleLayerWater }))
 	{
 		INC_DWORD_STAT_BY(STAT_ShaderCompiling_NumLitMaterialShaders, 1);
 	}
diff --git a/Engine/Source/Runtime/Engine/Private/Materials/MaterialShared.cpp b/Engine/Source/Runtime/Engine/Private/Materials/MaterialShared.cpp
--- a/Engine/Source/Runtime/Engine/Private/Materials/MaterialShared.cpp
+++ b/Engine/Source/Runtime/Engine/Private/Materials/MaterialShared.cpp
@@ -517,7 +517,6 @@ void FMaterialCompilationOutput::Serialize(FArchive& Ar)
 	PackedFlags |= (bUsesPixelDepthOffset			<< 4);
 	PackedFlags |= (bUsesDistanceCullFade			<< 5);
 	PackedFlags |= (bHasRuntimeVirtualTextureOutput << 6);
-	PackedFlags |= (bUsesSingleLayerWaterMaterialOutput << 7);
 
 	Ar << PackedFlags;
 
@@ -528,7 +527,6 @@ void FMaterialCompilationOutput::Serialize(FArchive& Ar)
 	bUsesPixelDepthOffset			= (PackedFlags >> 4) & 1;
 	bUsesDistanceCullFade			= (PackedFlags >> 5) & 1;
 	bHasRuntimeVirtualTextureOutput = (PackedFlags >> 6) & 1;
-	bUsesSingleLayerWaterMaterialOutput = (PackedFlags >> 7) & 1;
 }
 
 void FMaterial::GetShaderMapId(EShaderPlatform Platform, FMaterialShaderMapId& OutId) const
@@ -848,19 +846,6 @@ bool FMaterial::MaterialModifiesMeshPosition_GameThread() const
 	return bUsesWPO || GetTessellationMode() != MTM_NoTessellation;
 }
 
-bool FMaterial::MaterialUsesSingleLayerWater_RenderThread() const
-{
-	check(IsInParallelRenderingThread());
-	return RenderingThreadShaderMap ? RenderingThreadShaderMap->UsesSingleLayerWaterMaterialOutput() : false;
-}
-
-bool FMaterial::MaterialUsesSingleLayerWater_GameThread() const
-{
-	//check(IsInGameThread());
-	FMaterialShaderMap* ShaderMap = GameThreadShaderMap.GetReference();
-	return ShaderMap ? ShaderMap->UsesSingleLayerWaterMaterialOutput() : false;
-}
-
 bool FMaterial::MaterialMayModifyMeshPosition() const
 {
 	// Conservative estimate when called before material translation has occurred. 
@@ -1465,13 +1450,6 @@ UMaterialInterface* FMaterialResource::GetMaterialInterface() const
 	return MaterialInstance ? (UMaterialInterface*)MaterialInstance : (UMaterialInterface*)Material;
 }
 
-bool FMaterialResource::IsUsingSingleLayerWaterMaterialOutput() const
-{
-	// Apparently HasAnyExpressionsInMaterialAndFunctionsOfType is too heavy. So checking bUsedWithWater instead.
-	// return GetMaterialInterface() ? GetMaterialInterface()->GetMaterial()->HasAnyExpressionsInMaterialAndFunctionsOfType<UMaterialExpressionSingleLayerWaterMaterialOutput>() : false;
-	return GetMaterialInterface() ? GetMaterialInterface()->GetMaterial()->bUsedWithWater : false;
-}
-
 #if WITH_EDITOR
 void FMaterialResource::NotifyCompilationFinished()
 {
diff --git a/Engine/Source/Runtime/Engine/Private/StaticMeshRender.cpp b/Engine/Source/Runtime/Engine/Private/StaticMeshRender.cpp
--- a/Engine/Source/Runtime/Engine/Private/StaticMeshRender.cpp
+++ b/Engine/Source/Runtime/Engine/Private/StaticMeshRender.cpp
@@ -1011,7 +1011,7 @@ void FStaticMeshSceneProxy::DrawStaticElements(FStaticPrimitiveDrawInterface* PD
 							&& !IsTranslucentBlendMode(Material->GetBlendMode())
 							&& !Material->MaterialModifiesMeshPosition_RenderThread()
 							&& Material->GetMaterialDomain() == MD_Surface
-							&& !Material->MaterialUsesSingleLayerWater_RenderThread();
+							&& !Material->GetShadingModels().HasShadingModel(MSM_SingleLayerWater);
 
 						bAllSectionsCastShadow &= Section.bCastShadow;
 					}
diff --git a/Engine/Source/Runtime/Engine/Public/MaterialShared.h b/Engine/Source/Runtime/Engine/Public/MaterialShared.h
--- a/Engine/Source/Runtime/Engine/Public/MaterialShared.h
+++ b/Engine/Source/Runtime/Engine/Public/MaterialShared.h
@@ -554,7 +554,6 @@ public:
 		bUsesWorldPositionOffset(false),
 		bUsesGlobalDistanceField(false),
 		bUsesPixelDepthOffset(false),
-		bUsesSingleLayerWaterMaterialOutput(false),
 		bUsesDistanceCullFade(false),
 		bHasRuntimeVirtualTextureOutput(false)
 	{}
@@ -602,9 +601,6 @@ public:
 	/** true if the material writes a pixel depth offset */
 	uint8 bUsesPixelDepthOffset : 1;
 
-	/** true if the material writes a pixel depth offset */
-	uint8 bUsesSingleLayerWaterMaterialOutput : 1;
-
 	/** true if the material uses distance cull fade */
 	uint8 bUsesDistanceCullFade : 1;
 
@@ -1112,7 +1108,6 @@ public:
 	bool UsesVelocitySceneTexture() const { return MaterialCompilationOutput.UsesVelocitySceneTexture(); }
 	bool UsesDistanceCullFade() const { return MaterialCompilationOutput.bUsesDistanceCullFade; }
 	bool HasRuntimeVirtualTextureOutput() const { return MaterialCompilationOutput.bHasRuntimeVirtualTextureOutput; }
-	bool UsesSingleLayerWaterMaterialOutput() const { return MaterialCompilationOutput.bUsesSingleLayerWaterMaterialOutput; }
 #if WITH_EDITOR
 	uint32 GetNumUsedUVScalars() const { return MaterialCompilationOutput.NumUsedUVScalars; }
 	uint32 GetNumUsedCustomInterpolatorScalars() const { return MaterialCompilationOutput.NumUsedCustomInterpolatorScalars; }
@@ -1524,7 +1519,6 @@ public:
 	virtual uint32 GetStencilCompare() const { return 0; }
 	virtual bool HasRuntimeVirtualTextureOutput() const { return false; }
 	virtual bool CastsRayTracedShadows() const { return true; }
-	virtual bool IsUsingSingleLayerWaterMaterialOutput() const { return false; };
 
 	/**
 	 * Should shaders compiled for this material be saved to disk?
@@ -1611,10 +1605,6 @@ public:
 	ENGINE_API bool MaterialModifiesMeshPosition_RenderThread() const;
 	ENGINE_API bool MaterialModifiesMeshPosition_GameThread() const;
 
-	/** Does the material use a single layer water. */
-	ENGINE_API bool MaterialUsesSingleLayerWater_RenderThread() const;
-	ENGINE_API bool MaterialUsesSingleLayerWater_GameThread() const;
-
 	/** Does the material use a pixel depth offset. */
 	ENGINE_API bool MaterialUsesPixelDepthOffset() const;
 
@@ -2282,7 +2272,6 @@ public:
 	ENGINE_API virtual bool ComputeFogPerPixel() const override;
 	ENGINE_API virtual bool HasRuntimeVirtualTextureOutput() const override;
 	ENGINE_API virtual bool CastsRayTracedShadows() const override;
-	ENGINE_API virtual bool IsUsingSingleLayerWaterMaterialOutput() const override;
 	ENGINE_API  virtual UMaterialInterface* GetMaterialInterface() const override;
 	/**
 	 * Should shaders compiled for this material be saved to disk?
@@ -2752,6 +2741,6 @@ bool ReloadMaterialResource(
 inline bool ShouldIncludeMaterialInDefaultOpaquePass(const FMaterial& Material)
 {
 	return !Material.IsSky()
-		&& !Material.MaterialUsesSingleLayerWater_RenderThread();
+		&& !Material.GetShadingModels().HasShadingModel(MSM_SingleLayerWater);
 }
 
diff --git a/Engine/Source/Runtime/Renderer/Private/BasePassRendering.h b/Engine/Source/Runtime/Renderer/Private/BasePassRendering.h
--- a/Engine/Source/Runtime/Renderer/Private/BasePassRendering.h
+++ b/Engine/Source/Runtime/Renderer/Private/BasePassRendering.h
@@ -499,7 +499,7 @@ public:
 		static const auto SupportStationarySkylight = IConsoleManager::Get().FindTConsoleVariableDataInt(TEXT("r.SupportStationarySkylight"));
 		static const auto SupportAllShaderPermutations = IConsoleManager::Get().FindTConsoleVariableDataInt(TEXT("r.SupportAllShaderPermutations"));
 
-		const bool IsSingleLayerWater = Parameters.Material->IsUsingSingleLayerWaterMaterialOutput();
+		const bool IsSingleLayerWater = Parameters.Material->GetShadingModels().HasShadingModel(MSM_SingleLayerWater);
 
 		const bool bTranslucent = IsTranslucentBlendMode(Parameters.Material->GetBlendMode());
 		const bool bForceAllPermutations = SupportAllShaderPermutations && SupportAllShaderPermutations->GetValueOnAnyThread() != 0;
diff --git a/Engine/Source/Runtime/Renderer/Private/LightMapDensityRendering.cpp b/Engine/Source/Runtime/Renderer/Private/LightMapDensityRendering.cpp
--- a/Engine/Source/Runtime/Renderer/Private/LightMapDensityRendering.cpp
+++ b/Engine/Source/Runtime/Renderer/Private/LightMapDensityRendering.cpp
@@ -254,8 +254,7 @@ void FLightmapDensityMeshProcessor::AddMeshBatch(const FMeshBatch& RESTRICT Mesh
 		const bool bAllowLowQualityLightMaps = (!SupportLowQualityLightmapsVar) || (SupportLowQualityLightmapsVar->GetValueOnAnyThread() != 0);
 
 		if ((!bTranslucentBlendMode || ViewIfDynamicMeshCommand->Family->EngineShowFlags.Wireframe)
-			&& !Material->IsSky()
-			&& !Material->MaterialUsesSingleLayerWater_RenderThread())
+			&& ShouldIncludeMaterialInDefaultOpaquePass(*Material))
 		{
 			if (!bMaterialMasked && !Material->MaterialModifiesMeshPosition_RenderThread())
 			{
diff --git a/Engine/Source/Runtime/Renderer/Private/MobileBasePass.cpp b/Engine/Source/Runtime/Renderer/Private/MobileBasePass.cpp
--- a/Engine/Source/Runtime/Renderer/Private/MobileBasePass.cpp
+++ b/Engine/Source/Runtime/Renderer/Private/MobileBasePass.cpp
@@ -407,7 +407,7 @@ void MobileBasePass::SetTranslucentRenderState(FMeshPassProcessorRenderState& Dr
 			DrawRenderState.SetBlendState(TStaticBlendState<CW_RGBA, BO_Add, BF_Zero, BF_InverseSourceAlpha, BO_Add, BF_One, BF_InverseSourceAlpha>::GetRHI());
 			break;
 		default:
-			if (Material.MaterialUsesSingleLayerWater_RenderThread())
+			if (Material.GetShadingModels().HasShadingModel(MSM_SingleLayerWater))
 			{
 				// Single layer water is an opaque marerial rendered as translucent on Mobile. We force pre-multiplied alpha to achieve water depth based transmittance.
 				DrawRenderState.SetBlendState(TStaticBlendState<CW_RGB, BO_Add, BF_One, BF_InverseSourceAlpha, BO_Add, BF_Zero, BF_InverseSourceAlpha>::GetRHI());
@@ -570,7 +570,7 @@ void FMobileBasePassMeshProcessor::AddMeshBatch(const FMeshBatch& RESTRICT MeshB
 	const EBlendMode BlendMode = Material.GetBlendMode();
 	const FMaterialShadingModelField ShadingModels = Material.GetShadingModels();
 	const bool bIsTranslucent = IsTranslucentBlendMode(BlendMode);
-	const bool bUsesWaterMaterial = Material.MaterialUsesSingleLayerWater_RenderThread(); // Water goes into the translucent pass
+	const bool bUsesWaterMaterial = ShadingModels.HasShadingModel(MSM_SingleLayerWater); // Water goes into the translucent pass
 	
 	if (bTranslucentBasePass)
 	{
diff --git a/Engine/Source/Runtime/Renderer/Private/PrimitiveSceneInfo.cpp b/Engine/Source/Runtime/Renderer/Private/PrimitiveSceneInfo.cpp
--- a/Engine/Source/Runtime/Renderer/Private/PrimitiveSceneInfo.cpp
+++ b/Engine/Source/Runtime/Renderer/Private/PrimitiveSceneInfo.cpp
@@ -79,7 +79,7 @@ public:
 			const bool bSupportsCachingMeshDrawCommands = SupportsCachingMeshDrawCommands(PrimitiveSceneProxy, *StaticMesh, FeatureLevel);
 
 			bool bUseSkyMaterial = Mesh.MaterialRenderProxy->GetMaterial(FeatureLevel)->IsSky();
-			bool bUseSingleLayerWaterMaterial = Mesh.MaterialRenderProxy->GetMaterial(FeatureLevel)->MaterialUsesSingleLayerWater_RenderThread();
+			bool bUseSingleLayerWaterMaterial = Mesh.MaterialRenderProxy->GetMaterial(FeatureLevel)->GetShadingModels().HasShadingModel(MSM_SingleLayerWater);
 			FStaticMeshBatchRelevance* StaticMeshRelevance = new(PrimitiveSceneInfo->StaticMeshRelevances) FStaticMeshBatchRelevance(
 				*StaticMesh, 
 				ScreenSize, 
diff --git a/Engine/Source/Runtime/Renderer/Private/SingleLayerWaterRendering.cpp b/Engine/Source/Runtime/Renderer/Private/SingleLayerWaterRendering.cpp
--- a/Engine/Source/Runtime/Renderer/Private/SingleLayerWaterRendering.cpp
+++ b/Engine/Source/Runtime/Renderer/Private/SingleLayerWaterRendering.cpp
@@ -103,7 +103,7 @@ void FSingleLayerWaterPassMeshProcessor::AddMeshBatch(const FMeshBatch& RESTRICT
 	const FMaterialRenderProxy* FallbackMaterialRenderProxyPtr = nullptr;
 	const FMaterial& Material = MeshBatch.MaterialRenderProxy->GetMaterialWithFallback(FeatureLevel, FallbackMaterialRenderProxyPtr);
 
-	if (Material.MaterialUsesSingleLayerWater_RenderThread())
+	if (Material.GetShadingModels().HasShadingModel(MSM_SingleLayerWater))
 	{
 		const ERasterizerFillMode MeshFillMode = ComputeMeshFillMode(MeshBatch, Material);
 		const ERasterizerCullMode MeshCullMode = ComputeMeshCullMode(MeshBatch, Material);
diff --git a/Engine/Source/Runtime/Renderer/Private/VelocityRendering.cpp b/Engine/Source/Runtime/Renderer/Private/VelocityRendering.cpp
--- a/Engine/Source/Runtime/Renderer/Private/VelocityRendering.cpp
+++ b/Engine/Source/Runtime/Renderer/Private/VelocityRendering.cpp
@@ -573,8 +573,7 @@ void FVelocityMeshProcessor::AddMeshBatch(const FMeshBatch& RESTRICT MeshBatch,
 		const ERasterizerCullMode MeshCullMode = ComputeMeshCullMode(MeshBatch, *Material);
 
 		if ((BlendMode == BLEND_Opaque || BlendMode == BLEND_Masked)
-			&& !Material->IsSky()
-			&& !Material->MaterialUsesSingleLayerWater_RenderThread()) // No velocity for water as of now
+			&& ShouldIncludeMaterialInDefaultOpaquePass(*Material)) // No velocity for water as of now
 		{
 			if (Material->WritesEveryPixel() && !Material->IsTwoSided() && !Material->MaterialModifiesMeshPosition_RenderThread())
 			{
