Water ray refraction before taking into account phase function.
Done with Ryan B

diff --git a/Engine/Shaders/Private/BasePassPixelShader.usf b/Engine/Shaders/Private/BasePassPixelShader.usf
--- a/Engine/Shaders/Private/BasePassPixelShader.usf
+++ b/Engine/Shaders/Private/BasePassPixelShader.usf
@@ -1228,6 +1228,7 @@ ShadingModel);
 			OpaqueBasePass.SceneDepthWithoutSingleLayerWaterTexture, OpaqueBasePass.SceneDepthWithoutSingleLayerWaterSampler,
 			OpaqueBasePass.SceneLuminanceWithoutSingleLayerWaterTexture, OpaqueBasePass.SceneLuminanceWithoutSingleLayerWaterSampler,
 			OpaqueBasePass.SceneCustomDepthTexture, OpaqueBasePass.SceneCustomDepthSampler,
+			Specular,
 			SunIlluminance, WaterDiffuseIndirectIlluminance, EnvBrdf,
 			CameraIsUnderWater, WaterVisibility).Luminance;
 	}
diff --git a/Engine/Shaders/Private/MobileBasePassPixelShader.usf b/Engine/Shaders/Private/MobileBasePassPixelShader.usf
--- a/Engine/Shaders/Private/MobileBasePassPixelShader.usf
+++ b/Engine/Shaders/Private/MobileBasePassPixelShader.usf
@@ -667,7 +667,7 @@ void Main(
 		WaterVolumeLightingOutput WaterLighting = EvaluateWaterVolumeLighting(
 			MaterialParameters, PixelMaterialInputs, ResolvedView,
 			MobileSceneTextures.SceneDepthTexture, MobileSceneTextures.SceneDepthTextureSampler,
-			Shadow,
+			Shadow, Specular,
 			SunIlluminance, WaterDiffuseIndirectIlluminance, EnvBrdf,
 			CameraIsUnderWater, WaterVisibility);
 
diff --git a/Engine/Shaders/Private/SingleLayerWaterCommon.ush b/Engine/Shaders/Private/SingleLayerWaterCommon.ush
--- a/Engine/Shaders/Private/SingleLayerWaterCommon.ush
+++ b/Engine/Shaders/Private/SingleLayerWaterCommon.ush
@@ -4,6 +4,7 @@
 
 #include "Common.ush"
 #include "ParticipatingMediaCommon.ush" 
+#include "ShadingCommon.ush" 
 
 
 #ifndef SIMPLE_SINGLE_LAYER_WATER
@@ -27,6 +28,7 @@ WaterVolumeLightingOutput EvaluateWaterVolumeLighting(
 	Texture2D SceneLuminanceWithoutSingleLayerWaterTexture, SamplerState SceneLuminanceWithoutSingleLayerWaterSampler,
 	Texture2D SceneCustomDepthTexture, SamplerState SceneCustomDepthSampler,
 #endif // SIMPLE_SINGLE_LAYER_WATER
+	float Specular,
 	float3 SunIlluminance, float3 AmbiantIlluminance, float3 EnvBrdf,
 	bool CameraIsUnderWater, float WaterVisibility)
 {
@@ -106,8 +108,11 @@ WaterVolumeLightingOutput EvaluateWaterVolumeLighting(
 		const float3 ScatteringCoeff = max(0.0f, GetSingleLayerWaterMaterialOutput0(MaterialParameters));
 		const float3 AbsorptionCoeff = max(0.0f, GetSingleLayerWaterMaterialOutput1(MaterialParameters));
 		const float PhaseG = clamp(GetSingleLayerWaterMaterialOutput2(MaterialParameters), -1.0f, 1.0f);
-		const float DirLightPhaseValue = SchlickPhase(PhaseG, dot(ResolvedView.DirectionalLightDirection.xyz, MaterialParameters.CameraVector));
 		const float3 ExtinctionCoeff = ScatteringCoeff + AbsorptionCoeff;
+		
+		float IorFrom = 1.0f; // assumes we come from air
+		float IorTo   = DielectricF0ToIor(DielectricSpecularToF0(Specular)); // Wrong if metal is set to >1. But we still keep refraction on the water surface nonetheless.
+		const float DirLightPhaseValue = SchlickPhase(PhaseG, dot(ResolvedView.DirectionalLightDirection.xyz, refract(-MaterialParameters.CameraVector, MaterialParameters.WorldNormal, IorFrom/IorTo)));
 
 		// We also apply transmittance from light to under water surface. However, the scene has been lit by many sources already.
 		// So the transmittance toabove surface is simply approximated using the travel distance from the scene pixel to the water top, assuming a flat water surface.
