Removed unused GetShadingModelID() shader function

diff --git a/Engine/Shaders/Private/ShadingCommon.ush b/Engine/Shaders/Private/ShadingCommon.ush
--- a/Engine/Shaders/Private/ShadingCommon.ush
+++ b/Engine/Shaders/Private/ShadingCommon.ush
@@ -24,40 +24,6 @@
 #define ZERO_PRECSHADOW_MASK			(1 << 6)
 #define SKIP_VELOCITY_MASK				(1 << 7)
 
-uint GetShadingModelID()
-{        
-#if MATERIAL_SHADINGMODEL_UNLIT
-	return SHADINGMODELID_UNLIT;
-#elif MATERIAL_SHADINGMODEL_DEFAULT_LIT
-	return SHADINGMODELID_DEFAULT_LIT;
-#elif MATERIAL_SHADINGMODEL_SUBSURFACE
-	return SHADINGMODELID_SUBSURFACE;
-#elif MATERIAL_SHADINGMODEL_PREINTEGRATED_SKIN
-	return SHADINGMODELID_PREINTEGRATED_SKIN;
-#elif MATERIAL_SHADINGMODEL_SUBSURFACE_PROFILE
-	return SHADINGMODELID_SUBSURFACE_PROFILE;
-#elif MATERIAL_SHADINGMODEL_CLEAR_COAT
-	return SHADINGMODELID_CLEAR_COAT;
-#elif MATERIAL_SHADINGMODEL_TWOSIDED_FOLIAGE
-	return SHADINGMODELID_TWOSIDED_FOLIAGE;
-#elif MATERIAL_SHADINGMODEL_HAIR
-	return SHADINGMODELID_HAIR;
-#elif MATERIAL_SHADINGMODEL_CLOTH
-	return SHADINGMODELID_CLOTH;
-#elif MATERIAL_SHADINGMODEL_EYE
-	return SHADINGMODELID_EYE;
-#elif MATERIAL_SHADINGMODEL_SINGLELAYERWATER
-	return SHADINGMODELID_SINGLELAYERWATER;
-#else
-	#if COMPILER_HLSLCC
-		// hlslcc needs a return type
-		return SHADINGMODELID_MASK;
-	#else
-	// missing shading model, compiler should report ShadingModelID is not set
-	#endif
-#endif
-}
-
 // for debugging and to visualize
 float3 GetShadingModelColor(uint ShadingModelID)
 {
