Fix for new water shading model to actually evaluate lighting on opaque material parts.

diff --git a/Engine/Shaders/Private/ClusteredDeferredShadingPixelShader.usf b/Engine/Shaders/Private/ClusteredDeferredShadingPixelShader.usf
--- a/Engine/Shaders/Private/ClusteredDeferredShadingPixelShader.usf
+++ b/Engine/Shaders/Private/ClusteredDeferredShadingPixelShader.usf
@@ -260,6 +260,7 @@ void ClusteredShadingPixelShader(
 	GET_LIGHT_GRID_LOCAL_LIGHTING_SINGLE_SM(SHADINGMODELID_HAIR, PixelShadingModelID, CompositedLighting, ScreenUV, CulledLightGridData, Dither, FirstNonSimpleLightIndex);
 	GET_LIGHT_GRID_LOCAL_LIGHTING_SINGLE_SM(SHADINGMODELID_CLOTH, PixelShadingModelID, CompositedLighting, ScreenUV, CulledLightGridData, Dither, FirstNonSimpleLightIndex);
 	GET_LIGHT_GRID_LOCAL_LIGHTING_SINGLE_SM(SHADINGMODELID_EYE, PixelShadingModelID, CompositedLighting, ScreenUV, CulledLightGridData, Dither, FirstNonSimpleLightIndex);
+	GET_LIGHT_GRID_LOCAL_LIGHTING_SINGLE_SM(SHADINGMODELID_SINGLELAYERWATER, PixelShadingModelID, CompositedLighting, ScreenUV, CulledLightGridData, Dither, FirstNonSimpleLightIndex);
 #else // !USE_PASS_PER_SHADING_MODEL
 
 	float2 ScreenPosition = UVAndScreenPos.zw;
diff --git a/Engine/Shaders/Private/ShadingCommon.ush b/Engine/Shaders/Private/ShadingCommon.ush
--- a/Engine/Shaders/Private/ShadingCommon.ush
+++ b/Engine/Shaders/Private/ShadingCommon.ush
@@ -40,6 +40,7 @@ float3 GetShadingModelColor(uint ShadingModelID)
 	else if (ShadingModelID == SHADINGMODELID_HAIR) return float3(0.6f, 0.1f, 0.5f);
 	else if (ShadingModelID == SHADINGMODELID_CLOTH) return float3(0.7f, 1.0f, 1.0f); 
 	else if (ShadingModelID == SHADINGMODELID_EYE) return float3(0.3f, 1.0f, 1.0f); 
+	else if (ShadingModelID == SHADINGMODELID_SINGLELAYERWATER) return float3(0.5f, 0.5f, 1.0f);
 	else return float3(1.0f, 1.0f, 1.0f); // White
 #else
 	switch(ShadingModelID)
@@ -54,6 +55,7 @@ float3 GetShadingModelColor(uint ShadingModelID)
 		case SHADINGMODELID_HAIR: return float3(0.6f, 0.1f, 0.5f);
 		case SHADINGMODELID_CLOTH: return float3(0.7f, 1.0f, 1.0f);
 		case SHADINGMODELID_EYE: return float3(0.3f, 1.0f, 1.0f);
+		case SHADINGMODELID_SINGLELAYERWATER: return float3(0.5f, 0.5f, 1.0f);
 		default: return float3(1.0f, 1.0f, 1.0f); // White
 	}
 #endif
diff --git a/Engine/Shaders/Private/ShadingModels.ush b/Engine/Shaders/Private/ShadingModels.ush
--- a/Engine/Shaders/Private/ShadingModels.ush
+++ b/Engine/Shaders/Private/ShadingModels.ush
@@ -936,6 +936,7 @@ FDirectLighting IntegrateBxDF( FGBufferData GBuffer, half3 N, half3 V, half3 L,
 	switch( GBuffer.ShadingModelID )
 	{
 		case SHADINGMODELID_DEFAULT_LIT:
+		case SHADINGMODELID_SINGLELAYERWATER:
 			return DefaultLitBxDF( GBuffer, N, V, L, Falloff, NoL, AreaLight, Shadow );
 		case SHADINGMODELID_SUBSURFACE:
 			return SubsurfaceBxDF( GBuffer, N, V, L, Falloff, NoL, AreaLight, Shadow );
diff --git a/Engine/Shaders/Private/TiledDeferredLightShaders.usf b/Engine/Shaders/Private/TiledDeferredLightShaders.usf
--- a/Engine/Shaders/Private/TiledDeferredLightShaders.usf
+++ b/Engine/Shaders/Private/TiledDeferredLightShaders.usf
@@ -345,6 +345,7 @@ void TiledDeferredLightingMain(
 		EXECUTE_SHADING_LOOPS_SINGLE_SM(SHADINGMODELID_HAIR, PixelShadingModelID, CompositedLighting, PixelPos, NumLightsAffectingTile, NumSimpleLightsAffectingTile, CameraVector, WorldPosition);
 		EXECUTE_SHADING_LOOPS_SINGLE_SM(SHADINGMODELID_CLOTH, PixelShadingModelID, CompositedLighting, PixelPos, NumLightsAffectingTile, NumSimpleLightsAffectingTile, CameraVector, WorldPosition);
 		EXECUTE_SHADING_LOOPS_SINGLE_SM(SHADINGMODELID_EYE, PixelShadingModelID, CompositedLighting, PixelPos, NumLightsAffectingTile, NumSimpleLightsAffectingTile, CameraVector, WorldPosition);
+		EXECUTE_SHADING_LOOPS_SINGLE_SM(SHADINGMODELID_SINGLELAYERWATER, PixelShadingModelID, CompositedLighting, PixelPos, NumLightsAffectingTile, NumSimpleLightsAffectingTile, CameraVector, WorldPosition);
 #else // !USE_PASS_PER_SHADING_MODEL
 		ExecuteShadingLoops(CompositedLighting, ScreenSpaceData, NumLightsAffectingTile, NumSimpleLightsAffectingTile, CameraVector, WorldPosition);
 #endif // USE_PASS_PER_SHADING_MODEL
