Disable earlydepthstencil semantic when also writing pixel depth

diff --git a/Engine/Shaders/Private/BasePassPixelShader.usf b/Engine/Shaders/Private/BasePassPixelShader.usf
--- a/Engine/Shaders/Private/BasePassPixelShader.usf
+++ b/Engine/Shaders/Private/BasePassPixelShader.usf
@@ -1356,6 +1356,18 @@ ShadingModel);
 #endif
 }
 
+// If virtual texture is enabled then use early depth test so that UAV feedback buffer writes respect the depth test
+#if NUM_VIRTUALTEXTURE_SAMPLES || LIGHTMAP_VT_ENABLED
+	#if COMPILER_SUPPORT_DEPTHSTENCIL_EARLYTEST_LATEWRITE
+		// If we support early depth test with late write behaviour then use it since we may be using discard, or modifying depth
+		#define PIXELSHADER_EARLYDEPTHSTENCIL DEPTHSTENCIL_EARLYTEST_LATEWRITE
+	#elif !OUTPUT_PIXEL_DEPTH_OFFSET
+		// Otherwise we can only use early depth test if not modifying depth
+		// Modifying depth will trigger the slow path where we write feedback to UAV even where depth occluded!
+		#define PIXELSHADER_EARLYDEPTHSTENCIL EARLYDEPTHSTENCIL
+	#endif
+#endif
+
 
 // the following needs to match to the code in FSceneRenderTargets::GetGBufferRenderTargets()
 #define PIXELSHADEROUTPUT_BASEPASS 1
diff --git a/Engine/Shaders/Private/PixelShaderOutputCommon.ush b/Engine/Shaders/Private/PixelShaderOutputCommon.ush
--- a/Engine/Shaders/Private/PixelShaderOutputCommon.ush
+++ b/Engine/Shaders/Private/PixelShaderOutputCommon.ush
@@ -43,15 +43,13 @@ PixelShaderOutputCommon.usf: To allow PS input/output passed int functions throu
 #ifndef PIXELSHADEROUTPUT_A2C
 	#define PIXELSHADEROUTPUT_A2C 0
 #endif
-
+#ifndef PIXELSHADER_EARLYDEPTHSTENCIL
+#define PIXELSHADER_EARLYDEPTHSTENCIL
+#endif
 
 // ---------------------------------------------------------------------------------
 
-#if LIGHTMAP_VT_ENABLED || NUM_VIRTUALTEXTURE_SAMPLES
-	// If virtual texture is enabled need to make sure we're using early depth test, in order for UAV feedback buffer writes to respect depth test
-	// Also need to wait until after pixel shader to write depth, since we may be using discard, or modifying depth
-	DEPTHSTENCIL_EARLYTEST_LATEWRITE
-#endif // LIGHTMAP_VT_ENABLED || NUM_VIRTUALTEXTURE_SAMPLES
+PIXELSHADER_EARLYDEPTHSTENCIL
 void MainPS
 	(
 #if PIXELSHADEROUTPUT_INTERPOLANTS || PIXELSHADEROUTPUT_BASEPASS
diff --git a/Engine/Shaders/Public/Platform.ush b/Engine/Shaders/Public/Platform.ush
--- a/Engine/Shaders/Public/Platform.ush
+++ b/Engine/Shaders/Public/Platform.ush
@@ -431,6 +431,11 @@
 #define PLATFORM_SUPPORTS_ATOMIC_UINT64 0
 #endif
 
+// Support for depth test running both before and after pixel shader
+#ifndef COMPILER_SUPPORTS_DEPTHSTENCIL_EARLYTEST_LATEWRITE
+#define COMPILER_SUPPORTS_DEPTHSTENCIL_EARLYTEST_LATEWRITE 0
+#endif
+
 // ---------------------------------------------------- Compiler attributes
 
 #if SM5_PROFILE || COMPILER_SUPPORTS_ATTRIBUTES
@@ -457,10 +462,6 @@
 #define EARLYDEPTHSTENCIL [earlydepthstencil]
 #endif
 
-#ifndef DEPTHSTENCIL_EARLYTEST_LATEWRITE
-#define DEPTHSTENCIL_EARLYTEST_LATEWRITE EARLYDEPTHSTENCIL
-#endif
-
 // ---------------------------------------------------- Compiler attribute fallbacks
 
 #ifndef UNROLL
@@ -499,6 +500,10 @@
 #define EARLYDEPTHSTENCIL
 #endif
 
+#ifndef DEPTHSTENCIL_EARLYTEST_LATEWRITE
+#define DEPTHSTENCIL_EARLYTEST_LATEWRITE
+#endif
+
 #ifndef STRONG_TYPE
 #define STRONG_TYPE
 #endif
