Fixed sun not showing on mobile preview. HLSL Refract  is just not working there.

diff --git a/Engine/Shaders/Private/SingleLayerWaterCommon.ush b/Engine/Shaders/Private/SingleLayerWaterCommon.ush
--- a/Engine/Shaders/Private/SingleLayerWaterCommon.ush
+++ b/Engine/Shaders/Private/SingleLayerWaterCommon.ush
@@ -19,6 +19,24 @@ struct WaterVolumeLightingOutput
 	float3 WaterToSceneToLightTransmittance;
 };
 
+// We use our own refract because the HLSL one was not working on all platofrm (ES3.1 preview for instance)
+bool WaterRefract(float3 V, float3 N, float Eta, out float3 R)
+{
+	const float CosThetaI = dot(V, N);
+	const float CosThetaT2 = 1.f - Eta * Eta * (1.0f - CosThetaI * CosThetaI);
+	if (CosThetaT2 < 0.0f)
+	{
+		// Total internal reflection
+		R = reflect(-V, N);
+		return false;
+	}
+
+	const float CosThetaT = -(CosThetaI >= 0.0f ? 1.0f : -1.0f) * sqrt(CosThetaT2);
+	R = (Eta * CosThetaI + CosThetaT) * N - Eta * V;
+	return true;
+}
+
+
 // Returns as luminance, the result of lighting scattering in water and the under water scene color contribution.
 WaterVolumeLightingOutput EvaluateWaterVolumeLighting(
 	FMaterialPixelParameters MaterialParameters, FPixelMaterialInputs PixelMaterialInputs, ViewState ResolvedView,
@@ -86,16 +104,28 @@ WaterVolumeLightingOutput EvaluateWaterVolumeLighting(
 
 #endif // !SIMPLE_SINGLE_LAYER_WATER
 
-		const float BehindWaterDeltaDepth = CameraIsUnderWater ? WaterDepth : SceneDepth - WaterDepth;
+		const float BehindWaterDeltaDepth = CameraIsUnderWater ? WaterDepth : max(0.0f, SceneDepth - WaterDepth);
 
 		const float3 ScatteringCoeff = max(0.0f, GetSingleLayerWaterMaterialOutput0(MaterialParameters));
 		const float3 AbsorptionCoeff = max(0.0f, GetSingleLayerWaterMaterialOutput1(MaterialParameters));
 		const float PhaseG = clamp(GetSingleLayerWaterMaterialOutput2(MaterialParameters), -1.0f, 1.0f);
 		const float3 ExtinctionCoeff = ScatteringCoeff + AbsorptionCoeff;
 		
-		float IorFrom = 1.0f; // assumes we come from air
-		float IorTo   = DielectricF0ToIor(DielectricSpecularToF0(Specular)); // Wrong if metal is set to >1. But we still keep refraction on the water surface nonetheless.
-		const float DirLightPhaseValue = SchlickPhase(PhaseG, dot(ResolvedView.DirectionalLightDirection.xyz, -refract(-MaterialParameters.CameraVector, MaterialParameters.WorldNormal, IorFrom/IorTo)));
+		float DirLightPhaseValue = 0.0f; // Default when Total Internal Reflection happens.
+		{
+			float IorFrom = 1.0f; // assumes we come from air
+			float IorTo   = DielectricF0ToIor(DielectricSpecularToF0(Specular)); // Wrong if metal is set to >1. But we still keep refraction on the water surface nonetheless.
+#if 0
+			DirLightPhaseValue = SchlickPhase(PhaseG, dot(-ResolvedView.DirectionalLightDirection.xyz, refract(MaterialParameters.CameraVector, MaterialParameters.WorldNormal, IorFrom / IorTo)));
+#else
+			const float relativeIOR = IorFrom / IorTo;
+			float3 UnderWaterRayDir = 0.0f;
+			if (WaterRefract(MaterialParameters.CameraVector, MaterialParameters.WorldNormal, relativeIOR, UnderWaterRayDir))
+			{
+				DirLightPhaseValue = SchlickPhase(PhaseG, dot(-ResolvedView.DirectionalLightDirection.xyz, UnderWaterRayDir));
+			}
+#endif
+		}
 
 		// We also apply transmittance from light to under water surface. However, the scene has been lit by many sources already.
 		// So the transmittance toabove surface is simply approximated using the travel distance from the scene pixel to the water top, assuming a flat water surface.
