Some Vulkan shader fixes

diff --git a/Engine/Shaders/Private/SSRT/SSRTReflections.usf b/Engine/Shaders/Private/SSRT/SSRTReflections.usf
--- a/Engine/Shaders/Private/SSRT/SSRTReflections.usf
+++ b/Engine/Shaders/Private/SSRT/SSRTReflections.usf
@@ -73,6 +73,7 @@ void ScreenSpaceReflectionsStencilPS(
 	const float Roughness = GetRoughness(GBuffer);
 	const float RoughnessFade = GetRoughnessFade(Roughness);
 	
+#if !defined(TILE_COMPUTE_SSR)
 	if (RoughnessFade > 0.0 && GBuffer.ShadingModelID != 0)
 	{
 		// we are going to compute SSR for this pixel, so we discard this
@@ -80,6 +81,7 @@ void ScreenSpaceReflectionsStencilPS(
 		// tehrefore execute ScreenSpaceReflectionsPS() for this pixel. 
 		discard;
 	}
+#endif
 	
 	// we are not going to compute SSR for this pixel, so we clear the color
 	// since ScreenSpaceReflectionsPS() won't be executed in this pixel.
diff --git a/Engine/Shaders/Private/SingleLayerWaterComposite.usf b/Engine/Shaders/Private/SingleLayerWaterComposite.usf
--- a/Engine/Shaders/Private/SingleLayerWaterComposite.usf
+++ b/Engine/Shaders/Private/SingleLayerWaterComposite.usf
@@ -137,8 +137,8 @@ SingleLayerWaterCompositeOutput SingleLayerWaterComposite(float2 BufferUV, float
 #ifdef TILE_CATERGORISATION_SHADER
 
 
-unsigned int TiledViewWidth;
-unsigned int TiledViewHeight;
+uint TiledViewWidth;
+uint TiledViewHeight;
 float TiledViewWidthInv;
 float TiledViewHeightInv;
 
@@ -241,8 +241,8 @@ void WaterTileCatergorisationCS(uint3 ThreadId : SV_DispatchThreadID, uint3 Grou
 
 
 
-unsigned int TiledViewWidth;
-unsigned int TiledViewHeight;
+uint TiledViewWidth;
+uint TiledViewHeight;
 
 Buffer<uint> WaterTileListData;
 
