Fixing shader compilation issues

diff --git a/Engine/Shaders/Private/DistortionCommon.ush b/Engine/Shaders/Private/DistortionCommon.ush
--- a/Engine/Shaders/Private/DistortionCommon.ush
+++ b/Engine/Shaders/Private/DistortionCommon.ush
@@ -7,7 +7,7 @@
 
 float2 ComputeBufferUVDistortion(
 	in FMaterialPixelParameters MaterialParameters, in FPixelMaterialInputs PixelMaterialInputs, in ViewState ResolvedView,
-	in half3 Normal, in float4 DistortionParams, in float2 ScreenUV, in float2 Refraction, in bool TryToClip)
+	in half3 Normal, in float4 DistortionParams, in float2 ScreenUV, in float2 RefractionIORAndBias, in bool TryToClip)
 {
 #if MATERIAL_TANGENTSPACENORMAL
 	half3 ViewNormal = normalize(TransformTangentVectorToView(MaterialParameters, Normal));
@@ -17,12 +17,12 @@ float2 ComputeBufferUVDistortion(
 
 #if REFRACTION_USE_PIXEL_NORMAL_OFFSET
 	half3 ViewVertexNormal = TransformTangentVectorToView(MaterialParameters, float3(0, 0, 1));
-	// Treat Refraction of 1 as no refraction, to be consistent with IOR mode
-	float2 ViewportUVDistortion = (ViewVertexNormal.xy - ViewNormal.xy) * (Refraction.x - 1);
+	// Treat refraction of 1 as no refraction, to be consistent with IOR mode
+	float2 ViewportUVDistortion = (ViewVertexNormal.xy - ViewNormal.xy) * (RefractionIORAndBias.x - 1);
 #else
 	// we assume the camera is in air
 	float AirIOR = 1.0f;
-	float2 ViewportUVDistortion = ViewNormal.xy * (Refraction.x - AirIOR);
+	float2 ViewportUVDistortion = ViewNormal.xy * (RefractionIORAndBias.x - AirIOR);
 #endif
 
 	float2 BufferUVDistortion = ViewportUVDistortion * ResolvedView.ViewSizeAndInvSize.xy * ResolvedView.BufferSizeAndInvSize.zw;
@@ -55,10 +55,10 @@ float2 ComputeBufferUVDistortion(
 
 void PostProcessUVDistortion(
 	in FMaterialPixelParameters MaterialParameters, in FPixelMaterialInputs PixelMaterialInputs, 
-	in float DistortSceneDepth,	inout float2 BufferUVDistortion, in float2 Refraction)
+	in float DistortSceneDepth,	inout float2 BufferUVDistortion, in float2 RefractionIORAndBias)
 {
 	// Soft thresholding 
-	float Bias = -Refraction.y;
+	float Bias = -RefractionIORAndBias.y;
 	float Range = clamp(abs(Bias * 0.5f), 0, 50);
 	float Z = DistortSceneDepth;
 	float ZCompare = MaterialParameters.ScreenPosition.w;
