Revert DistortAccumulatePS. This shader was failing to cook on X1, as it couldn't find DistortionPass.DistortionParams binding.

diff --git a/Engine/Shaders/Private/DistortAccumulatePS.usf b/Engine/Shaders/Private/DistortAccumulatePS.usf
--- a/Engine/Shaders/Private/DistortAccumulatePS.usf
+++ b/Engine/Shaders/Private/DistortAccumulatePS.usf
@@ -18,9 +18,11 @@
 #include "SceneTexturesCommon.ush"
 #include "/Engine/Generated/Material.ush"
 #include "/Engine/Generated/VertexFactory.ush"
-#include "DistortionCommon.ush"
 
+//A fudge factor scale to bring values close to what they would have been under usual circumstances prior to this change.
+static const float OffsetFudgeFactor = 0.00023;
 
+static const half DistortionScaleBias = 4.0f;
 
 /** output distortion offsets as color so they can be accumulated (via blending) */
 void Main(
@@ -54,17 +56,47 @@ void Main(
 	// material distortion offset
 	half3 Normal = GetMaterialNormal(MaterialParameters, PixelMaterialInputs);
 
+#if MATERIAL_TANGENTSPACENORMAL
+	half3 ViewNormal = normalize(TransformTangentVectorToView(MaterialParameters, Normal));
+#else
+	half3 ViewNormal = normalize(TransformWorldVectorToView(Normal));
+#endif
+
+	// X = IOR, Y = Refraction Bias
+	float2 Refraction = GetMaterialRefraction(PixelMaterialInputs);
+	 	
+#if REFRACTION_USE_PIXEL_NORMAL_OFFSET
+	half3 ViewVertexNormal = TransformTangentVectorToView(MaterialParameters, float3(0, 0, 1));
+	// Treat Refraction of 1 as no refraction, to be consistent with IOR mode
+	float2 ViewportUVDistortion = (ViewVertexNormal.xy - ViewNormal.xy) * (Refraction.x - 1);
+#else
+	// we assume the camera is in air
+	float AirIOR = 1.0f;
+	float2 ViewportUVDistortion = ViewNormal.xy * (Refraction.x - AirIOR);
+#endif
+
+	float2 BufferUVDistortion = ViewportUVDistortion * ResolvedView.ViewSizeAndInvSize.xy * ResolvedView.BufferSizeAndInvSize.zw;
+
+	//clip if the distortion distance (squared) is too small to be noticed
+	//this will result in a less expensive apply pass since the clipped pixels won't set stencil to pass
+	clip(dot(BufferUVDistortion, BufferUVDistortion) - .00001);
+
 	// Prevent silhouettes from geometry that is in front of distortion from being seen in the distortion 
 	float2 NDC = (MaterialParameters.ScreenPosition.xy / MaterialParameters.ScreenPosition.w);
 	float2 ScreenUV = NDC * ResolvedView.ScreenPositionScaleBias.xy + ResolvedView.ScreenPositionScaleBias.wz;
 
-	// X = IOR, Y = Refraction Bias
-	float2 Refraction = GetMaterialRefraction(PixelMaterialInputs);
+	// When ISR is enabled we store two FOVs in the DistortionParams and compute the aspect ratio here
+#if INSTANCED_STEREO
+	float InvTanHalfFov = DistortionParams[Interpolants.EyeIndex];
+    float Ratio = DistortionParams.z / DistortionParams.w;
+#else
+	float InvTanHalfFov = DistortionParams.x;
+    float Ratio = DistortionParams.y;
+#endif
 
-	// Compute UV distortion
-	float2 BufferUVDistortion = ComputeBufferUVDistortion(
-		MaterialParameters, PixelMaterialInputs, ResolvedView,
-		Normal, DistortionParams, ScreenUV, Refraction, true);
+	// Fix for Fov and aspect.
+ 	float2 FovFix = float2(InvTanHalfFov,Ratio*InvTanHalfFov);
+	BufferUVDistortion *= DistortionParams.zw * float2(OffsetFudgeFactor,-OffsetFudgeFactor) * FovFix;
 
 	// Sample depth at distortion offset
 	float2 DistortBufferUV = ScreenUV + BufferUVDistortion;
@@ -75,9 +107,17 @@ void Main(
 #else
 	float DistortSceneDepth = CalcSceneDepth(DistortBufferUV);
 #endif // (ES2_PROFILE || ES3_1_PROFILE) 
-
-	// Post process UV distortion according to depth
-	PostProcessUVDistortion(MaterialParameters, PixelMaterialInputs, DistortSceneDepth, BufferUVDistortion, Refraction);
+	
+	// Soft thresholding 
+	float Bias = -Refraction.y;
+	float Range = clamp(abs(Bias * 0.5f), 0, 50);
+	float Z = DistortSceneDepth;
+	float ZCompare = MaterialParameters.ScreenPosition.w;
+	float InvWidth = 1.0f / max(1.0f, Range );
+	BufferUVDistortion *= saturate((Z - ZCompare) * InvWidth + Bias);
+
+	//Scale up for better precision in low/subtle refractions at the expense of artefacts at higher refraction.
+	BufferUVDistortion*= DistortionScaleBias;
 
 	// store positive and negative offsets separately
 	float2 PosOffset = max(BufferUVDistortion,0);
