Fixed case where SCattering=Absorption=-0 causing Nans.

diff --git a/Engine/Shaders/Private/SingleLayerWaterCommon.ush b/Engine/Shaders/Private/SingleLayerWaterCommon.ush
--- a/Engine/Shaders/Private/SingleLayerWaterCommon.ush
+++ b/Engine/Shaders/Private/SingleLayerWaterCommon.ush
@@ -109,7 +109,7 @@ WaterVolumeLightingOutput EvaluateWaterVolumeLighting(
 		const float3 ScatteringCoeff = max(0.0f, GetSingleLayerWaterMaterialOutput0(MaterialParameters));
 		const float3 AbsorptionCoeff = max(0.0f, GetSingleLayerWaterMaterialOutput1(MaterialParameters));
 		const float PhaseG = clamp(GetSingleLayerWaterMaterialOutput2(MaterialParameters), -1.0f, 1.0f);
-		const float3 ExtinctionCoeff = ScatteringCoeff + AbsorptionCoeff;
+		const float3 ExtinctionCoeff = max(ScatteringCoeff + AbsorptionCoeff, 1e-20); // Max to avoid division by 0 with the analytical integral below.
 		
 		float DirLightPhaseValue = 0.0f; // Default when Total Internal Reflection happens.
 		{
