Restoring DistortAccumulatePS.usf refactor after finding a proper fix. The issue was that "#define DistortionParams DistortionPass.DistortionParams" was also changing "in float4 DistortionParams" argument of ComputeBufferUVDistortion.

diff --git a/Engine/Shaders/Private/DistortAccumulatePS.usf b/Engine/Shaders/Private/DistortAccumulatePS.usf
--- a/Engine/Shaders/Private/DistortAccumulatePS.usf
+++ b/Engine/Shaders/Private/DistortAccumulatePS.usf
@@ -18,11 +18,9 @@
 #include "SceneTexturesCommon.ush"
 #include "/Engine/Generated/Material.ush"
 #include "/Engine/Generated/VertexFactory.ush"
+#include "DistortionCommon.ush"
 
-//A fudge factor scale to bring values close to what they would have been under usual circumstances prior to this change.
-static const float OffsetFudgeFactor = 0.00023;
 
-static const half DistortionScaleBias = 4.0f;
 
 /** output distortion offsets as color so they can be accumulated (via blending) */
 void Main(
@@ -56,47 +54,17 @@ void Main(
 	// material distortion offset
 	half3 Normal = GetMaterialNormal(MaterialParameters, PixelMaterialInputs);
 
-#if MATERIAL_TANGENTSPACENORMAL
-	half3 ViewNormal = normalize(TransformTangentVectorToView(MaterialParameters, Normal));
-#else
-	half3 ViewNormal = normalize(TransformWorldVectorToView(Normal));
-#endif
-
-	// X = IOR, Y = Refraction Bias
-	float2 Refraction = GetMaterialRefraction(PixelMaterialInputs);
-	 	
-#if REFRACTION_USE_PIXEL_NORMAL_OFFSET
-	half3 ViewVertexNormal = TransformTangentVectorToView(MaterialParameters, float3(0, 0, 1));
-	// Treat Refraction of 1 as no refraction, to be consistent with IOR mode
-	float2 ViewportUVDistortion = (ViewVertexNormal.xy - ViewNormal.xy) * (Refraction.x - 1);
-#else
-	// we assume the camera is in air
-	float AirIOR = 1.0f;
-	float2 ViewportUVDistortion = ViewNormal.xy * (Refraction.x - AirIOR);
-#endif
-
-	float2 BufferUVDistortion = ViewportUVDistortion * ResolvedView.ViewSizeAndInvSize.xy * ResolvedView.BufferSizeAndInvSize.zw;
-
-	//clip if the distortion distance (squared) is too small to be noticed
-	//this will result in a less expensive apply pass since the clipped pixels won't set stencil to pass
-	clip(dot(BufferUVDistortion, BufferUVDistortion) - .00001);
-
 	// Prevent silhouettes from geometry that is in front of distortion from being seen in the distortion 
 	float2 NDC = (MaterialParameters.ScreenPosition.xy / MaterialParameters.ScreenPosition.w);
 	float2 ScreenUV = NDC * ResolvedView.ScreenPositionScaleBias.xy + ResolvedView.ScreenPositionScaleBias.wz;
 
-	// When ISR is enabled we store two FOVs in the DistortionParams and compute the aspect ratio here
-#if INSTANCED_STEREO
-	float InvTanHalfFov = DistortionParams[Interpolants.EyeIndex];
-    float Ratio = DistortionParams.z / DistortionParams.w;
-#else
-	float InvTanHalfFov = DistortionParams.x;
-    float Ratio = DistortionParams.y;
-#endif
+	// X = IOR, Y = Refraction Bias
+	float2 Refraction = GetMaterialRefraction(PixelMaterialInputs);
 
-	// Fix for Fov and aspect.
- 	float2 FovFix = float2(InvTanHalfFov,Ratio*InvTanHalfFov);
-	BufferUVDistortion *= DistortionParams.zw * float2(OffsetFudgeFactor,-OffsetFudgeFactor) * FovFix;
+	// Compute UV distortion
+	float2 BufferUVDistortion = ComputeBufferUVDistortion(
+		MaterialParameters, PixelMaterialInputs, ResolvedView,
+		Normal, DistortionParams, ScreenUV, Refraction, true);
 
 	// Sample depth at distortion offset
 	float2 DistortBufferUV = ScreenUV + BufferUVDistortion;
@@ -107,17 +75,9 @@ void Main(
 #else
 	float DistortSceneDepth = CalcSceneDepth(DistortBufferUV);
 #endif // (ES2_PROFILE || ES3_1_PROFILE) 
-	
-	// Soft thresholding 
-	float Bias = -Refraction.y;
-	float Range = clamp(abs(Bias * 0.5f), 0, 50);
-	float Z = DistortSceneDepth;
-	float ZCompare = MaterialParameters.ScreenPosition.w;
-	float InvWidth = 1.0f / max(1.0f, Range );
-	BufferUVDistortion *= saturate((Z - ZCompare) * InvWidth + Bias);
-
-	//Scale up for better precision in low/subtle refractions at the expense of artefacts at higher refraction.
-	BufferUVDistortion*= DistortionScaleBias;
+
+	// Post process UV distortion according to depth
+	PostProcessUVDistortion(MaterialParameters, PixelMaterialInputs, DistortSceneDepth, BufferUVDistortion, Refraction);
 
 	// store positive and negative offsets separately
 	float2 PosOffset = max(BufferUVDistortion,0);
diff --git a/Engine/Shaders/Private/DistortionCommon.ush b/Engine/Shaders/Private/DistortionCommon.ush
--- a/Engine/Shaders/Private/DistortionCommon.ush
+++ b/Engine/Shaders/Private/DistortionCommon.ush
@@ -4,10 +4,9 @@
 
 #include "Common.ush"
 
-
 float2 ComputeBufferUVDistortion(
 	in FMaterialPixelParameters MaterialParameters, in FPixelMaterialInputs PixelMaterialInputs, in ViewState ResolvedView,
-	in half3 Normal, in float4 DistortionParams, in float2 ScreenUV, in float2 RefractionIORAndBias, in bool TryToClip)
+	in half3 Normal, in float4 DistortionParameters, in float2 ScreenUV, in float2 RefractionIORAndBias, in bool TryToClip)
 {
 #if MATERIAL_TANGENTSPACENORMAL
 	half3 ViewNormal = normalize(TransformTangentVectorToView(MaterialParameters, Normal));
@@ -34,20 +33,20 @@ float2 ComputeBufferUVDistortion(
 		clip(dot(BufferUVDistortion, BufferUVDistortion) - .00001);
 	}
 
-	// When ISR is enabled we store two FOVs in the DistortionParams and compute the aspect ratio here
+	// When ISR is enabled we store two FOVs in the DistortionParameters and compute the aspect ratio here
 #if INSTANCED_STEREO
-	float InvTanHalfFov = DistortionParams[Interpolants.EyeIndex];
-	float Ratio = DistortionParams.z / DistortionParams.w;
+	float InvTanHalfFov = DistortionParameters[Interpolants.EyeIndex];
+	float Ratio = DistortionParameters.z / DistortionParameters.w;
 #else
-	float InvTanHalfFov = DistortionParams.x;
-	float Ratio = DistortionParams.y;
+	float InvTanHalfFov = DistortionParameters.x;
+	float Ratio = DistortionParameters.y;
 #endif
 
 	// Fix for Fov and aspect.
 	float2 FovFix = float2(InvTanHalfFov, Ratio*InvTanHalfFov);
 	//A fudge factor scale to bring values close to what they would have been under usual circumstances prior to this change.
 	const float OffsetFudgeFactor = 0.00023;
-	BufferUVDistortion *= DistortionParams.zw * float2(OffsetFudgeFactor, -OffsetFudgeFactor) * FovFix;
+	BufferUVDistortion *= DistortionParameters.zw * float2(OffsetFudgeFactor, -OffsetFudgeFactor) * FovFix;
 
 	return BufferUVDistortion;
 }
