Use depth fetch for mobile water materails

diff --git a/Engine/Shaders/Private/MobileBasePassPixelShader.usf b/Engine/Shaders/Private/MobileBasePassPixelShader.usf
--- a/Engine/Shaders/Private/MobileBasePassPixelShader.usf
+++ b/Engine/Shaders/Private/MobileBasePassPixelShader.usf
@@ -667,7 +667,6 @@ void Main(
 		const float4 NullDistortionParams = 1.0f;
 		WaterVolumeLightingOutput WaterLighting = EvaluateWaterVolumeLighting(
 			MaterialParameters, PixelMaterialInputs, ResolvedView,
-			MobileSceneTextures.SceneDepthTexture, MobileSceneTextures.SceneDepthTextureSampler,
 			Shadow, Specular, NullDistortionParams,
 			SunIlluminance, WaterDiffuseIndirectIlluminance, EnvBrdf,
 			CameraIsUnderWater, WaterVisibility);
diff --git a/Engine/Shaders/Private/SingleLayerWaterCommon.ush b/Engine/Shaders/Private/SingleLayerWaterCommon.ush
--- a/Engine/Shaders/Private/SingleLayerWaterCommon.ush
+++ b/Engine/Shaders/Private/SingleLayerWaterCommon.ush
@@ -40,10 +40,10 @@ bool WaterRefract(float3 V, float3 N, float Eta, out float3 R)
 // Returns as luminance, the result of lighting scattering in water and the under water scene color contribution.
 WaterVolumeLightingOutput EvaluateWaterVolumeLighting(
 	FMaterialPixelParameters MaterialParameters, FPixelMaterialInputs PixelMaterialInputs, ViewState ResolvedView,
-	Texture2D SceneDepthWithoutSingleLayerWaterTexture, SamplerState SceneDepthWithoutSingleLayerWaterSampler,
 #if SIMPLE_SINGLE_LAYER_WATER
 	float DirectionalLightShadow,
 #else
+	Texture2D SceneDepthWithoutSingleLayerWaterTexture, SamplerState SceneDepthWithoutSingleLayerWaterSampler,
 	Texture2D SceneLuminanceWithoutSingleLayerWaterTexture, SamplerState SceneLuminanceWithoutSingleLayerWaterSampler,
 	Texture2D SceneCustomDepthTexture, SamplerState SceneCustomDepthSampler,
 #endif // SIMPLE_SINGLE_LAYER_WATER
@@ -70,16 +70,18 @@ WaterVolumeLightingOutput EvaluateWaterVolumeLighting(
 	float WaterDepth = ConvertFromDeviceZ(MaterialParameters.SvPosition.z);
 	if (WaterVisibility > 0.0f)
 	{
-		float PixelSceneDeviceZ = SceneDepthWithoutSingleLayerWaterTexture.SampleLevel(SceneDepthWithoutSingleLayerWaterSampler, ViewportUV, 0).r;
+#if SIMPLE_SINGLE_LAYER_WATER
+		float PixelSceneDeviceZ = LookupDeviceZ(ViewportUV);
 		float PixelSceneDepth = ConvertFromDeviceZ(PixelSceneDeviceZ);
 		
-#if SIMPLE_SINGLE_LAYER_WATER
-
 		float SceneDeviceZ = PixelSceneDeviceZ;
 		float SceneDepth = PixelSceneDepth;
 
 #else  // !SIMPLE_SINGLE_LAYER_WATER
 
+		float PixelSceneDeviceZ = SceneDepthWithoutSingleLayerWaterTexture.SampleLevel(SceneDepthWithoutSingleLayerWaterSampler, ViewportUV, 0).r;
+		float PixelSceneDepth = ConvertFromDeviceZ(PixelSceneDeviceZ);
+		
 		float2 Refraction = GetMaterialRefraction(PixelMaterialInputs); // X = IOR, Y = Refraction Bias
 		float2 BufferUVDistortion = ComputeBufferUVDistortion(
 			MaterialParameters, PixelMaterialInputs, ResolvedView,
