Adding wave parameter index support in material shaders.

#rnx

diff --git a/Engine/Shaders/Private/MaterialTemplate.ush b/Engine/Shaders/Private/MaterialTemplate.ush
--- a/Engine/Shaders/Private/MaterialTemplate.ush
+++ b/Engine/Shaders/Private/MaterialTemplate.ush
@@ -314,6 +314,10 @@ struct FMaterialPixelParameters
 #if NUM_VIRTUALTEXTURE_SAMPLES || LIGHTMAP_VT_ENABLED
 	FVirtualTextureFeedbackParams VirtualTextureFeedback;
 #endif
+
+#if WATER_MESH_FACTORY
+	uint WaterWaveParamIndex;
+#endif
 };
 
 // @todo compat hack
@@ -389,6 +393,10 @@ struct FMaterialVertexParameters
 
 	// Index into View.PrimitiveSceneData
 	uint PrimitiveId;
+
+#if WATER_MESH_FACTORY
+	uint WaterWaveParamIndex;
+#endif
 };
 
 /**
@@ -521,6 +529,24 @@ float2 GetViewportUV(FMaterialPixelParameters Parameters)
 
 #endif
 
+float GetWaterWaveParamIndex(FMaterialPixelParameters Parameters)
+{
+#if WATER_MESH_FACTORY
+	return (float)Parameters.WaterWaveParamIndex;
+#else
+	return 0.0f;
+#endif
+}
+
+float GetWaterWaveParamIndex(FMaterialVertexParameters Parameters)
+{
+#if WATER_MESH_FACTORY
+	return (float)Parameters.WaterWaveParamIndex;
+#else
+	return 0.0f;
+#endif
+}
+
 // Returns whether a scene texture id is a for a post process input or not.
 bool IsPostProcessInputSceneTexture(const uint SceneTextureId)
 {
@@ -560,7 +586,7 @@ float4 GetSceneTextureUVMinMax(const uint SceneTextureId)
 	if (IsPostProcessInputSceneTexture(SceneTextureId))
 	{
 		switch (SceneTextureId)
-		{
+	{
 		case PPI_PostProcessInput0:
 			return float4(PostProcessInput_0_UVViewportBilinearMin, PostProcessInput_0_UVViewportBilinearMax);
 		case PPI_PostProcessInput1:
@@ -626,10 +652,10 @@ MaterialFloat2 GetDefaultSceneTextureUV(FMaterialPixelParameters Parameters, con
 {
 	#if POST_PROCESS_MATERIAL
 	if (IsPostProcessInputSceneTexture(SceneTextureId))
-	{
-		return ViewportUVToSceneTextureUV(GetViewportUV(Parameters), SceneTextureId);
-	}
-	#endif
+		{
+			return ViewportUVToSceneTextureUV(GetViewportUV(Parameters), SceneTextureId);
+		}
+		#endif
 
 	return GetSceneTextureUV(Parameters);
 }
