Support up to 18 grass weight layers

diff --git a/Engine/Shaders/Private/LandscapeGrassWeight.usf b/Engine/Shaders/Private/LandscapeGrassWeight.usf
--- a/Engine/Shaders/Private/LandscapeGrassWeight.usf
+++ b/Engine/Shaders/Private/LandscapeGrassWeight.usf
@@ -129,6 +129,46 @@ void PSMain(
 		#endif
 		break;
 #endif // NUM_MATERIAL_OUTPUTS_GETGRASSWEIGHT > 6
+#if NUM_MATERIAL_OUTPUTS_GETGRASSWEIGHT > 10
+	case 3:
+		OutColor.x = GetGrassWeight10(MaterialParameters);
+		#if NUM_MATERIAL_OUTPUTS_GETGRASSWEIGHT > 11
+			OutColor.y = GetGrassWeight11(MaterialParameters);
+		#else
+			OutColor.y = 0;
+		#endif
+		#if NUM_MATERIAL_OUTPUTS_GETGRASSWEIGHT > 12
+			OutColor.z = GetGrassWeight12(MaterialParameters);
+		#else
+			OutColor.z = 0;
+		#endif
+		#if NUM_MATERIAL_OUTPUTS_GETGRASSWEIGHT > 13
+			OutColor.w = GetGrassWeight13(MaterialParameters);
+		#else
+			OutColor.w = 0;
+		#endif
+		break;
+#endif // NUM_MATERIAL_OUTPUTS_GETGRASSWEIGHT > 10
+#if NUM_MATERIAL_OUTPUTS_GETGRASSWEIGHT > 14
+	case 4:
+		OutColor.x = GetGrassWeight14(MaterialParameters);
+		#if NUM_MATERIAL_OUTPUTS_GETGRASSWEIGHT > 15
+			OutColor.y = GetGrassWeight15(MaterialParameters);
+		#else
+			OutColor.y = 0;
+		#endif
+		#if NUM_MATERIAL_OUTPUTS_GETGRASSWEIGHT > 16
+			OutColor.z = GetGrassWeight16(MaterialParameters);
+		#else
+			OutColor.z = 0;
+		#endif
+		#if NUM_MATERIAL_OUTPUTS_GETGRASSWEIGHT > 17
+			OutColor.w = GetGrassWeight17(MaterialParameters);
+		#else
+			OutColor.w = 0;
+		#endif
+		break;
+#endif // NUM_MATERIAL_OUTPUTS_GETGRASSWEIGHT > 14
 	default:
 		OutColor = float4(0,0,0,0);
 		break;
