Copying //UE4/Dev-Physics to Dev-Main (//UE4/Dev-Main) [at] 9098278

diff --git a/Engine/Shaders/Private/HairBsdf.ush b/Engine/Shaders/Private/HairBsdf.ush
new file mode 100644
--- /dev/null
+++ b/Engine/Shaders/Private/HairBsdf.ush
@@ -0,0 +1,274 @@
+// Copyright 1998-2019 Epic Games, Inc. All Rights Reserved.
+
+#pragma once
+
+// Custom hair BSDF for prototype purpose
+// Need clean up, validation, and optimization
+float I0( float x )
+{
+	x = abs(x);
+	float a;
+	if( x < 3.75 )
+	{
+		float t = x / 3.75;
+		float t2 = t*t;
+		a =        + 0.0045813;
+		a = a * t2 + 0.0360768;	
+		a = a * t2 + 0.2659732;
+		a = a * t2 + 1.2067492;
+		a = a * t2 + 3.0899424;
+		a = a * t2 + 3.5156229;
+		a = a * t2 + 1.0;
+	}
+	else
+	{
+		float t = 3.75 / x;
+		a =       + 0.00392377;
+		a = a * t - 0.01647633;
+		a = a * t + 0.02635537;
+		a = a * t - 0.02057706;
+		a = a * t + 0.00916281;
+		a = a * t - 0.00157565;
+		a = a * t + 0.00225319;
+		a = a * t + 0.01328592;
+		a = a * t + 0.39894228;
+		a *= exp(x) * rsqrt(x);
+	}
+	return a;
+}
+
+float Hair_g_2( float B, float SinThetaL, float SinThetaV )
+{
+	float v = B*B;
+	float CosThetaL2 = 1 - SinThetaL * SinThetaL;
+	float CosThetaV2 = 1 - SinThetaV * SinThetaV;
+	float Mp = 0;
+	if( v < 0.1 )
+	{
+		float a = sqrt(CosThetaL2 * CosThetaV2) / v;
+		float b = -SinThetaL * SinThetaV / v;
+		float logI0a = a > 12 ? a + 0.5 * ( -log(2*PI) + log(1/a) + 0.125/a ) : log( I0(a) );
+		Mp = exp( logI0a + b - rcp(v) + 0.6931 + log(0.5 / v) );
+	}
+	else
+	{
+		Mp = rcp( exp(2/v) * v - v ) * exp( ( 1 - SinThetaL * SinThetaV ) / v ) * I0( sqrt(CosThetaL2 * CosThetaV2) / v );
+	}
+	
+	return Mp;
+}
+
+float Hair_g( float B, float Theta )
+{
+	return exp( -0.5 * Pow2( Theta ) / (B*B) ) / ( sqrt(2*PI) * B );
+}
+
+float Hair_F( float CosTheta )
+{
+	const float n = 1.55;
+	const float F0 = Pow2( (1 - n) / (1 + n) );
+	return F0 + (1 - F0) * Pow5( 1 - CosTheta );
+}
+
+float SigmaAFromReflectance(float c, float b) 
+{
+	const float b2 = b  * b;
+	const float b3 = b  * b2;
+	const float b4 = b2 * b2;
+	const float b5 = b  * b4;
+	const float a  = log(c);
+	const float t = (a / (5.969f - 0.215f * b + 2.532f * b2 - 10.73f * b3 + 5.574f * b4 + 0.245f * b5) );
+	return t*t;
+}
+
+float3 SigmaAFromReflectance(const float3 c, float beta_n) 
+{
+	return float3(
+		SigmaAFromReflectance(c.x, beta_n),
+		SigmaAFromReflectance(c.y, beta_n),
+		SigmaAFromReflectance(c.z, beta_n));
+}
+
+float3 HairShading( FGBufferData GBuffer, float3 L, float3 V, half3 N, float Shadow, FHairTransmittanceData HairTransmittance, float Backlit, float Area, uint2 Random )
+{
+	// Strand tangent is oriented from root to tip, while the existing shading code expect tip to root orientation
+	//N = -N;
+
+	// to prevent NaN with decals
+	// OR-18489 HERO: IGGY: RMB on E ability causes blinding hair effect
+	// OR-17578 HERO: HAMMER: E causes blinding light on heroes with hair
+	float ClampedRoughness = clamp(GBuffer.Roughness, 1/255.0f, 1.0f);
+
+	//const float3 DiffuseN	= OctahedronToUnitVector( GBuffer.CustomData.xy * 2 - 1 );
+	//const float Backlit	= GBuffer.CustomData.z;
+
+#if HAIR_REFERENCE
+	// todo: ClampedRoughness is missing for this code path
+	float3 S = HairShadingRef( GBuffer, L, V, N, Random );
+	//float3 S = HairShadingMarschner( GBuffer, L, V, N );
+#else
+	// N is the vector parallel to hair pointing toward root
+
+	const float VoL       = dot(V,L);                                                      
+	const float SinThetaL = dot(N,L);
+	const float SinThetaV = dot(N,V);
+	float CosThetaD = cos( 0.5 * abs( asinFast( SinThetaV ) - asinFast( SinThetaL ) ) );
+
+	//CosThetaD = abs( CosThetaD ) < 0.01 ? 0.01 : CosThetaD;
+
+	const float3 Lp = L - SinThetaL * N;
+	const float3 Vp = V - SinThetaV * N;
+	const float CosPhi = dot(Lp,Vp) * rsqrt( dot(Lp,Lp) * dot(Vp,Vp) + 1e-4 );
+	const float CosHalfPhi = sqrt( saturate( 0.5 + 0.5 * CosPhi ) );
+	//const float Phi = acosFast( CosPhi );
+
+	const float beta_n = 0.3;
+	const float3 Color = SigmaAFromReflectance(GBuffer.BaseColor, beta_n); // GBuffer.BaseColor
+
+	float n = 1.55;
+	//float n_prime = sqrt( n*n - 1 + Pow2( CosThetaD ) ) / CosThetaD;
+	float n_prime = 1.19 / CosThetaD + 0.36 * CosThetaD;
+
+	float Shift = 0.035;
+	float Alpha[] =
+	{
+		-Shift * 2,
+		Shift,
+		Shift * 4,
+	};
+	
+	float B[3];
+	B[0] = Area + Pow2( ClampedRoughness );
+	B[1] = Area + Pow2( ClampedRoughness ) / 2;
+	B[2] = Area + Pow2( ClampedRoughness ) * 2;
+
+	float3 S = 0;
+
+	// R
+	if(HairTransmittance.ScatteringComponent.x > 0)
+	{
+		const float sa = sin( Alpha[0] );
+		const float ca = cos( Alpha[0] );
+		float Shift = 2*sa* ( ca * CosHalfPhi * sqrt( 1 - SinThetaV * SinThetaV ) + sa * SinThetaV );
+#if 0
+		float Mp = Hair_g( B[0] * sqrt(2.0) * CosHalfPhi, SinThetaL + SinThetaV - Shift );
+#elif 0
+		float Mp = Hair_g( B[0], SinThetaL + SinThetaV - Shift );
+#else
+		float SinThetaV0 = SinThetaV;
+		SinThetaV0 -= 2*sa * ( CosHalfPhi * ca * sqrt( 1 - SinThetaV0 * SinThetaV0 ) + sa * SinThetaV0 );
+		float Mp = Hair_g_2(B[0], SinThetaL, SinThetaV0 );
+#endif
+		float Np = 0.25 * CosHalfPhi;
+
+#if 0
+		float Fp = Hair_F( sqrt( saturate( 0.5 + 0.5 * VoL ) ) );
+#else
+        const float CosThetaV = sqrt(1 - SinThetaV * SinThetaV);
+		float Fp = Hair_F( CosThetaV );
+#endif
+		
+		const float SpecVariation = ( GBuffer.Specular * 2 ) * lerp( 1, Backlit, saturate(-VoL) );
+		S += Mp * Np * Fp  * SpecVariation;
+	}
+
+	// TT
+	if(HairTransmittance.ScatteringComponent.y > 0)
+	{
+		float Mp = Hair_g( B[1], SinThetaL + SinThetaV - Alpha[1] );
+
+		float a = 1 / n_prime;
+		//float h = CosHalfPhi * rsqrt( 1 + a*a - 2*a * sqrt( 0.5 - 0.5 * CosPhi ) );
+		//float h = CosHalfPhi * ( ( 1 - Pow2( CosHalfPhi ) ) * a + 1 );
+		float h = CosHalfPhi * ( 1 + a * ( 0.6 - 0.8 * CosPhi ) );
+		//float h = 0.4;
+		//float yi = asinFast(h);
+		//float yt = asinFast(h / n_prime);
+		
+		float f = Hair_F( CosThetaD * sqrt( saturate( 1 - h*h ) ) );
+		float Fp = Pow2(1 - f);
+		//float3 Tp = pow( GBuffer.BaseColor, 0.5 * ( 1 + cos(2*yt) ) / CosThetaD );
+		//float3 Tp = pow( GBuffer.BaseColor, 0.5 * cos(yt) / CosThetaD );
+		#if 0
+		float3 Tp = pow( GBuffer.BaseColor, 0.5 * sqrt( 1 - Pow2(h * a) ) / CosThetaD );
+		#else
+		// This allows to have the same absorption than PBRT. The resulting color how ever 
+		// does not match the actual input color as this is supposed to have multiple 
+		// scattering event to reach the same saturation level
+		float yt = asinFast(h / n_prime);
+		float3 Tp = exp(-Color * 2 * abs(cos(yt) / CosThetaD) );
+		#endif
+		//float t = asin( 1 / n_prime );
+		//float d = ( sqrt(2) - t ) / ( 1 - t );
+		//float s = -0.5 * PI * (1 - 1 / n_prime) * log( 2*d - 1 - 2 * sqrt( d * (d - 1) ) );
+		//float s = 0.35;
+		//float Np = exp( (Phi - PI) / s ) / ( s * Pow2( 1 + exp( (Phi - PI) / s ) ) );
+		//float Np = 0.71 * exp( -1.65 * Pow2(Phi - PI) );
+		float Np = exp( -3.65 * CosPhi - 3.98 );
+
+		S += Mp * Np * Fp * Tp * Backlit;
+	}
+
+	// TRT
+	if(HairTransmittance.ScatteringComponent.z > 0)
+	{
+		float Mp = Hair_g( B[2], SinThetaL + SinThetaV - Alpha[2] );
+		
+		//float h = 0.75;
+		float f = Hair_F( CosThetaD * 0.5 );
+		float Fp = Pow2(1 - f) * f;
+		#if 0
+		//float3 Tp = pow( GBuffer.BaseColor, 1.6 / CosThetaD );
+		float3 Tp = pow( GBuffer.BaseColor, 0.8 / CosThetaD );
+		#else
+		// This allows to have the same absorption than PBRT. The resulting color however 
+		// does not match the actual input color as this is supposed to have multiple
+		// scattering events to reach the same saturation level
+		const float3 Color = SigmaAFromReflectance(GBuffer.BaseColor, beta_n);
+		float3 Tp = exp(-Color * 2 * 1.6 / CosThetaD );
+		#endif
+
+		//float s = 0.15;
+		//float Np = 0.75 * exp( Phi / s ) / ( s * Pow2( 1 + exp( Phi / s ) ) );
+		float Np = exp( 17 * CosPhi - 16.78 );
+
+		S += Mp * Np * Fp * Tp;
+	}
+#endif
+
+	if (!HairTransmittance.Enable)
+	{
+		// Use soft Kajiya Kay diffuse attenuation
+		float KajiyaDiffuse = 1 - abs( dot(N,L) );
+	
+		float3 FakeNormal = normalize( V - N * dot(V,N) );
+		//N = normalize( DiffuseN + FakeNormal * 2 );
+		N = FakeNormal;
+	
+		// Hack approximation for multiple scattering.
+		float Wrap = 1;
+		float NoL = saturate( ( dot(N, L) + Wrap ) / Square( 1 + Wrap ) );
+		float DiffuseScatter = (1 / PI) * lerp( NoL, KajiyaDiffuse, 0.33 ) * GBuffer.Metallic;
+		float Luma = Luminance( GBuffer.BaseColor );
+		float3 ScatterTint = pow( GBuffer.BaseColor / Luma, 1 - Shadow );
+		S += GBuffer.BaseColor * DiffuseScatter * pow(ScatterTint,.5);
+	}
+
+	S = -min(-S, 0.0);
+
+	if (HairTransmittance.Enable && HairTransmittance.MultipleScattering > 0)
+	{
+		S = EvaluateHairMultipleScattering(HairTransmittance, ClampedRoughness, S);
+	}
+
+	return S;
+}
+
+FDirectLighting HairBxDF( FGBufferData GBuffer, half3 N, half3 V, half3 L, float Falloff, float NoL, FAreaLight AreaLight, FShadowTerms Shadow )
+{
+	FDirectLighting Lighting;
+	Lighting.Diffuse = 0;
+	Lighting.Specular = 0;
+	Lighting.Transmission = AreaLight.FalloffColor * Falloff * HairShading( GBuffer, L, V, N, Shadow.TransmissionShadow, Shadow.HairTransmittance, 1, 0, uint2(0,0) );
+	return Lighting;
+}
\ No newline at end of file
diff --git a/Engine/Shaders/Private/HairStrands/HairStrandsCommon.ush b/Engine/Shaders/Private/HairStrands/HairStrandsCommon.ush
new file mode 100644
--- /dev/null
+++ b/Engine/Shaders/Private/HairStrands/HairStrandsCommon.ush
@@ -0,0 +1,46 @@
+// Copyright 1998-2019 Epic Games, Inc. All Rights Reserved.
+
+#pragma once
+
+// Use a custom BSDF which match more closely PBRT's one.
+// This allows to compare with offline/reference rendering
+#define HAIR_CUSTOM_BSDF 1
+
+struct HairAverageEnergy
+{
+	float3 A_R;
+	float3 A_TT;
+	float3 A_TRT;
+};
+
+struct HairAverageScattering
+{
+	float3 A_front;	// Front Scattering A_f
+	float3 A_back;	// Back Scattering A_b
+};
+
+
+HairAverageScattering SampleHairLUT(Texture3D<float4> LUTTexture, SamplerState LUTSampler, float3 Absorption, float Roughness, float SinViewAngle)
+{
+	const float2 LUTValue_R = LUTTexture.SampleLevel(LUTSampler, float3(saturate(abs(SinViewAngle)), saturate(Roughness), saturate(Absorption.x)), 0).xy;
+	const float2 LUTValue_G = LUTTexture.SampleLevel(LUTSampler, float3(saturate(abs(SinViewAngle)), saturate(Roughness), saturate(Absorption.y)), 0).xy;
+	const float2 LUTValue_B = LUTTexture.SampleLevel(LUTSampler, float3(saturate(abs(SinViewAngle)), saturate(Roughness), saturate(Absorption.z)), 0).xy;
+
+	HairAverageScattering Output;
+	Output.A_front = float3(LUTValue_R.x, LUTValue_G.x, LUTValue_B.x);
+	Output.A_back  = float3(LUTValue_R.y, LUTValue_G.y, LUTValue_B.y);
+	return Output;
+}
+
+HairAverageEnergy SampleHairEnergyLUT(Texture3D<float4> LUTTexture, SamplerState LUTSampler, float3 Absorption, float Roughness, float SinViewAngle)
+{
+	const float3 LUTValue_R = LUTTexture.SampleLevel(LUTSampler, float3(saturate(abs(SinViewAngle)), saturate(Roughness), saturate(Absorption.x)), 0).xyz;
+	const float3 LUTValue_G = LUTTexture.SampleLevel(LUTSampler, float3(saturate(abs(SinViewAngle)), saturate(Roughness), saturate(Absorption.y)), 0).xyz;
+	const float3 LUTValue_B = LUTTexture.SampleLevel(LUTSampler, float3(saturate(abs(SinViewAngle)), saturate(Roughness), saturate(Absorption.z)), 0).xyz;
+
+	HairAverageEnergy Output;
+	Output.A_R = float3(LUTValue_R.x, LUTValue_G.x, LUTValue_B.x);
+	Output.A_TT = float3(LUTValue_R.y, LUTValue_G.y, LUTValue_B.y);
+	Output.A_TRT = float3(LUTValue_R.z, LUTValue_G.z, LUTValue_B.z);
+	return Output;
+}
diff --git a/Engine/Shaders/Private/HairStrands/HairStrandsDebug.usf b/Engine/Shaders/Private/HairStrands/HairStrandsDebug.usf
new file mode 100644
--- /dev/null
+++ b/Engine/Shaders/Private/HairStrands/HairStrandsDebug.usf
@@ -0,0 +1,65 @@
+// Copyright 1998-2019 Epic Games, Inc. All Rights Reserved.
+
+#include "../Common.ush"
+#include "HairStrandsVisibilityCommon.ush"
+
+float2 OutputResolution;
+uint FastResolveMask;
+
+Texture2D<uint4> CategorizationTexture;
+Texture2D<uint2> DepthStencilTexture;
+SamplerState LinearSampler;
+
+#define DEBUG_MODE_SAMPLE_COUNT 0 
+#define DEBUG_MODE_COVERAGE_TYPE 1
+#define DEBUG_MODE_TAA_RESOLVE_TYPE 2
+
+float3 HUEtoRGB(in float H)
+{
+	float R = abs(H * 6 - 3) - 1;
+	float G = 2 - abs(H * 6 - 2);
+	float B = 2 - abs(H * 6 - 4);
+	return saturate(float3(R, G, B));
+}
+
+float3 HSVtoRGB(in float3 HSV)
+{
+	float3 RGB = HUEtoRGB(HSV.x);
+	return ((RGB - 1) * HSV.y + 1) * HSV.z;
+}
+
+float3 GetColorCode(float x)
+{
+	float c = (1 - saturate(x)) * 0.6; // Remap [0,1] to Blue-Red
+	return x > 0 ? HSVtoRGB(float3(c, 1, 1)) : float3(0, 0, 0);
+}
+
+void MainPS(
+	in FScreenVertexOutput Input,
+	out float4 OutColor : SV_Target0)
+{
+	const float2 UV = Input.UV;
+	const float2 PixelPos = UV * OutputResolution;
+
+	const FCategorizationData CatData = DecodeCategorizationData(CategorizationTexture.Load(uint3(PixelPos,0)));
+	const bool bIsValid = CatData.TotalCoverage8bit > 0;
+	if (!bIsValid)
+		discard;
+
+#if PERMUTATION_DEBUG_MODE == DEBUG_MODE_SAMPLE_COUNT
+	const float3 Color = GetColorCode(CatData.SampleCount/float(HAIR_VISIBILITY_MAXSAMPLECOUNT));
+#endif
+
+#if PERMUTATION_DEBUG_MODE == DEBUG_MODE_COVERAGE_TYPE
+	const bool bIsFullyCovered = CatData.TotalCoverage8bit >= 255;
+	const float3 Color = bIsFullyCovered ? float3(0.25,0.85,0.25) : float3(0.85,0.25,0.25);
+#endif
+
+#if PERMUTATION_DEBUG_MODE == DEBUG_MODE_TAA_RESOLVE_TYPE
+	const uint Stencil = DepthStencilTexture.Load(uint3(PixelPos,0)).y;
+	const bool bIsFastResolve = (Stencil & FastResolveMask) != 0;
+	const float3 Color = bIsFastResolve ? float3(1,1,0) : float3(0.15f, 0.25f, 0.75f);
+#endif
+
+	OutColor = float4(Color, 1);
+}
diff --git a/Engine/Shaders/Private/HairStrands/HairStrandsDeepShadowCommon.ush b/Engine/Shaders/Private/HairStrands/HairStrandsDeepShadowCommon.ush
new file mode 100644
--- /dev/null
+++ b/Engine/Shaders/Private/HairStrands/HairStrandsDeepShadowCommon.ush
@@ -0,0 +1,82 @@
+// Copyright 1998-2019 Epic Games, Inc. All Rights Reserved.
+
+#pragma once
+
+#define DEBUG_MODE 0
+
+// Layers' depths in clip space
+#define LAYER_DEPTH0 0.01f
+#define LAYER_DEPTH1 0.05f
+#define LAYER_DEPTH2 0.2f
+#define LAYER_DEPTH3 0.5f
+#define LAYER_DEPTH4 1.0f
+
+bool IsDepthCloser(float a, float b)
+{
+	return a < b;
+}
+
+float ComputeDOMWeight(float DistanceToFrontDepth, float LayerDepth)
+{
+    return IsDepthCloser(DistanceToFrontDepth, LayerDepth) ? 1 : 0;
+}
+
+float4 ComputeDOMWeights(float DistanceToFrontDepth)
+{
+	float4 Weigths = 0;
+	Weigths[0] = ComputeDOMWeight(DistanceToFrontDepth, LAYER_DEPTH0);
+	Weigths[1] = ComputeDOMWeight(DistanceToFrontDepth, LAYER_DEPTH1);
+	Weigths[2] = ComputeDOMWeight(DistanceToFrontDepth, LAYER_DEPTH2);
+	Weigths[3] = ComputeDOMWeight(DistanceToFrontDepth, LAYER_DEPTH3);
+	return Weigths;
+}	
+
+float3 HairDebugColor(float DistanceToFrontDepth)
+{
+	float3 color = 0;
+	if (DistanceToFrontDepth < LAYER_DEPTH3) color = float3(0,0,1);
+	if (DistanceToFrontDepth < LAYER_DEPTH2) color = float3(0,1,0);
+	if (DistanceToFrontDepth < LAYER_DEPTH1) color = float3(1,1,0);
+	if (DistanceToFrontDepth < LAYER_DEPTH0) color = float3(1,0,0);
+	return color;
+}
+
+float InterpolateCount(float DepthToFrontDepth, float Layer0Depth, float Layer1Depth)
+{
+	return saturate((DepthToFrontDepth - Layer0Depth) / (Layer1Depth - Layer0Depth));
+}
+
+float ComputeHairCount(float4 DomValue, float DistanceToFrontDepth)
+{
+	float OutCount = 0;
+	
+	if (DistanceToFrontDepth < LAYER_DEPTH0)
+		OutCount = lerp(          0, DomValue[0], InterpolateCount(DistanceToFrontDepth, 0, LAYER_DEPTH0));
+	else if (DistanceToFrontDepth < LAYER_DEPTH1)
+		OutCount = lerp(DomValue[0], DomValue[1], InterpolateCount(DistanceToFrontDepth, LAYER_DEPTH0, LAYER_DEPTH1));
+	else if (DistanceToFrontDepth < LAYER_DEPTH2)
+		OutCount = lerp(DomValue[1], DomValue[2], InterpolateCount(DistanceToFrontDepth, LAYER_DEPTH1, LAYER_DEPTH2));
+	else if (DistanceToFrontDepth < LAYER_DEPTH3)
+		OutCount = lerp(DomValue[2], DomValue[3], InterpolateCount(DistanceToFrontDepth, LAYER_DEPTH2, LAYER_DEPTH3));
+	else 
+		OutCount = DomValue[3];
+	return OutCount;
+}
+
+float GetDomDistanceToFrontDepth(float FrontDepth, float LightSpaceZ)
+{
+#if HAS_INVERTED_Z_BUFFER
+	return max(0.0f, FrontDepth - LightSpaceZ);
+#else
+	return max(0.0f, LightSpaceZ - FrontDepth);
+#endif
+}
+
+float GetDomDistanceToFrontDepthWithBias(float FrontDepth, float LightSpaceZ, float DepthBias)
+{
+#if HAS_INVERTED_Z_BUFFER
+	return max(0.0f, FrontDepth - LightSpaceZ - DepthBias);
+#else
+	return max(0.0f, LightSpaceZ - FrontDepth - DepthBias);
+#endif
+}
diff --git a/Engine/Shaders/Private/HairStrands/HairStrandsDeepShadowDebug.usf b/Engine/Shaders/Private/HairStrands/HairStrandsDeepShadowDebug.usf
new file mode 100644
--- /dev/null
+++ b/Engine/Shaders/Private/HairStrands/HairStrandsDeepShadowDebug.usf
@@ -0,0 +1,117 @@
+// Copyright 1998-2019 Epic Games, Inc. All Rights Reserved.
+
+#include "../Common.ush"
+#include "../PostProcessCommon.ush"
+#include "HairStrandsDeepShadowCommon.ush"
+
+float DomScale = 1;
+float2 DomAtlasOffset;
+float2 DomAtlasScale;
+float2 OutputResolution;
+float2 InvOutputResolution;
+
+int4 HairViewRect;
+
+Texture2D<float> DeepShadowDepthTexture;
+Texture2D<float4> DeepShadowLayerTexture;
+
+SamplerState LinearSampler;
+
+bool IsInsideRect(float4 Rect, float2 PixelPos, inout float2 LocalUV)
+{
+	LocalUV = (PixelPos - Rect.xy) / Rect.zw;
+
+	return 	PixelPos.x > Rect.x && PixelPos.x < Rect.x + Rect.z &&
+			PixelPos.y > Rect.y && PixelPos.y < Rect.y + Rect.w;
+}
+
+bool IsOnBorder(float2 LocalUV, float4 Rect)
+{
+	float2 PixelPos = LocalUV * Rect.zw;
+	float BorderThickness = 1;
+
+	return
+		PixelPos.x < BorderThickness || PixelPos.x > Rect.z - 1 - BorderThickness ||
+		PixelPos.y < BorderThickness || PixelPos.y > Rect.w - 1 - BorderThickness;
+}
+
+float3 AddBorder(float2 LocalUV, float4 Rect)
+{
+	const bool bIsOnBorder = IsOnBorder(LocalUV, Rect);
+	return bIsOnBorder ? float3(1,1,0) : float3(0,0,0);
+}
+
+// Display the DOM textures into some insert rect
+#if PERMUTATION_OUTPUT_TYPE == 0
+void VisualizeDomPS(
+	in FScreenVertexOutput Input,
+	out float4 OutColor : SV_Target0)
+{
+	const float TileSize = 48;
+	const float2 UV = Input.UV;
+	const float2 PixelPos = UV * OutputResolution;
+
+	float2 PixelPosAtTileCenter = PixelPos - (frac(PixelPos / TileSize) - 0.5f) * TileSize;
+	float2 UVAtTileCenter = PixelPosAtTileCenter * InvOutputResolution;
+
+	const float OffsetX = 50;
+	const float OffsetY = 220;
+	const float RectSize = 300;
+	const float4 DepthRect = float4(OffsetX, OffsetY, RectSize, RectSize);
+	const float4 DomRect = float4(OffsetX * 2 + RectSize, OffsetY, RectSize, RectSize);
+
+	float3 Color = 0;
+
+	float2 LocalUV = 0;
+	if (IsInsideRect(DepthRect, PixelPos, LocalUV))
+	{
+		const float2 AtlasUV = LocalUV * DomAtlasScale + DomAtlasOffset;
+		Color = DeepShadowDepthTexture.SampleLevel(LinearSampler, AtlasUV, 0).xxx;
+		Color += AddBorder(LocalUV, DepthRect);
+	}
+	else if (IsInsideRect(DomRect, PixelPos, LocalUV))
+	{
+		const float2 AtlasUV = LocalUV * DomAtlasScale + DomAtlasOffset;
+		Color = DeepShadowLayerTexture.SampleLevel(LinearSampler, AtlasUV, 0).xyz / DomScale;
+		Color += AddBorder(LocalUV, DomRect);
+#if 0 
+		float4 Weight = PostprocessInput2.SampleLevel(LinearSampler, LocalUV, 0);
+		if (Weight.w > 0) Color = float3(0, 0, 1);
+		if (Weight.z > 0) Color = float3(0, 1, 0);
+		if (Weight.y > 0) Color = float3(1, 1, 0);
+		if (Weight.x > 0) Color = float3(1, 0, 0);
+		Color *= 1000;
+#endif
+	}
+	else
+	{
+		discard;
+	}
+
+	OutColor = float4(Color, 1);
+}
+#endif
+
+// Display screen space rect for vizualizing projected viewrect of each hair cluster 
+#if PERMUTATION_OUTPUT_TYPE == 1
+void VisualizeDomPS(
+	in FScreenVertexOutput Input,
+	out float4 OutColor : SV_Target0)
+{
+	const float2 UV = Input.UV;
+	const float2 PixelPos = UV * OutputResolution;
+
+	float3 Color = 0;
+	float2 LocalUV = 0;
+	if (IsInsideRect(HairViewRect, PixelPos, LocalUV) && IsOnBorder(LocalUV, HairViewRect))
+	{
+		Color += float3(0, 1, 0);
+	}
+	else
+	{
+		discard;
+	}
+
+	OutColor = float4(Color, 1);
+}
+#endif
\ No newline at end of file
diff --git a/Engine/Shaders/Private/HairStrands/HairStrandsDeepShadowPS.usf b/Engine/Shaders/Private/HairStrands/HairStrandsDeepShadowPS.usf
new file mode 100644
--- /dev/null
+++ b/Engine/Shaders/Private/HairStrands/HairStrandsDeepShadowPS.usf
@@ -0,0 +1,133 @@
+// Copyright 1998-2019 Epic Games, Inc. All Rights Reserved.
+
+#include "../Common.ush"
+
+//#define SceneTexturesStruct DeepShadowPass.SceneTextures
+
+#include "../SceneTexturesCommon.ush"
+#include "HairStrandsDeepShadowCommon.ush"
+#include "HairStrandsVisibilityCommon.ush"
+#include "/Engine/Generated/Material.ush"
+#include "/Engine/Generated/VertexFactory.ush"
+
+static const float FixPointMaxValue = 1000;	// 100m @hair_todo: make this camera relative, and expose a CVAR fix changing this value
+static const float FixPointScale	= 100;	// 0.1mm precision
+static const float FixPointRange	= 2 * FixPointMaxValue * FixPointScale;
+
+#define SUPPORT_OPACITY_MASK 0
+
+void MainDepth(
+	in float4 SvPosition	: SV_Position)
+{
+#if SUPPORT_OPACITY_MASK
+	ResolvedView = ResolveView();
+	FMaterialPixelParameters MaterialParameters = GetMaterialPixelParameters(Interpolants, SvPosition);
+	FPixelMaterialInputs PixelMaterialInputs;
+	CalcMaterialParameters(MaterialParameters, PixelMaterialInputs, SvPosition, bIsFrontFace);
+	clip(GetMaterialMask(PixelMaterialInputs));
+#endif
+}
+
+void MainDom(
+	in float4 SvPosition	: SV_Position,
+	in float HairCoverage	: CUSTOM_COVERAGE,
+	out float4 OutColor		: SV_Target0)
+{
+	#if SUPPORT_OPACITY_MASK
+	{
+		ResolvedView = ResolveView();
+		FMaterialPixelParameters MaterialParameters = GetMaterialPixelParameters(Interpolants, SvPosition);
+		FPixelMaterialInputs PixelMaterialInputs;
+		CalcMaterialParameters(MaterialParameters, PixelMaterialInputs, SvPosition, bIsFrontFace);
+		clip(GetMaterialMask(PixelMaterialInputs));
+	}
+	#endif
+
+	const uint2 PixelCoord = uint2(floor(SvPosition.xy));
+    const float FrontDepth = DeepShadowPass.FrontDepthTexture.Load(uint3(PixelCoord, 0));
+	const float DistanceToFrontDepth = GetDomDistanceToFrontDepth(FrontDepth, SvPosition.z);
+
+    OutColor = ComputeDOMWeights(DistanceToFrontDepth) * HairCoverage;
+}
+
+#if VF_STRAND_HAIR
+  #if SM5_PROFILE && COMPILER_HLSL
+	RWTexture3D<uint> DensityTexture  : register(u1);
+
+	#if SUPPORT_TANGENT_PROPERTY == 1
+	RWTexture3D<uint> TangentXTexture : register(u2);
+	RWTexture3D<uint> TangentYTexture : register(u3);
+	RWTexture3D<uint> TangentZTexture : register(u4);
+	#endif
+
+	#if SUPPORT_MATERIAL_PROPERTY == 1
+	RWTexture3D<uint> MaterialTexture : register(u5);
+	#endif
+  #else
+	RWTexture3D<uint> DensityTexture;
+
+	#if SUPPORT_TANGENT_PROPERTY == 1
+	RWTexture3D<uint> TangentXTexture;
+	RWTexture3D<uint> TangentYTexture;
+	RWTexture3D<uint> TangentZTexture;
+	#endif
+
+	#if SUPPORT_MATERIAL_PROPERTY == 1
+	RWTexture3D<uint> MaterialTexture;
+	#endif
+  #endif
+#endif 
+
+void MainVoxel(
+#if SUPPORT_TANGENT_PROPERTY == 1 || SUPPORT_MATERIAL_PROPERTY == 1
+	FVertexFactoryInterpolantsVSToPS Interpolants,
+#endif
+	in float4 SvPosition : SV_Position,
+	in float4 PixelPosition : TEXCOORD6,
+	in float HairCoverage: CUSTOM_COVERAGE,
+	out uint OutColor : SV_Target0)
+{
+#if VF_STRAND_HAIR
+	const int3 VoxelCoord = (PixelPosition.xyz - DeepShadowPass.VoxelMinAABB) / (DeepShadowPass.VoxelMaxAABB - DeepShadowPass.VoxelMinAABB) * DeepShadowPass.VoxelResolution;
+		 
+	if (VoxelCoord.x >= 0 || VoxelCoord.y >= 0 || VoxelCoord.z >= 0 ||
+		VoxelCoord.x < int(DeepShadowPass.VoxelResolution) || 
+		VoxelCoord.y < int(DeepShadowPass.VoxelResolution) || 
+		VoxelCoord.z < int(DeepShadowPass.VoxelResolution))
+	{
+		const float HairCount = HairCoverage;
+		const float VoxelFixPointScale = GetVoxelDensityFixPointScale();
+		InterlockedAdd(DensityTexture[VoxelCoord], HairCount * VoxelFixPointScale);
+
+		#if SUPPORT_TANGENT_PROPERTY == 1
+		const half3 Tangent = GetTangentToWorld2(Interpolants).xyz;
+		const uint3 Tangent8bits = QuantizeTo8Bits(Tangent);
+
+		InterlockedAdd(TangentXTexture[VoxelCoord], Tangent8bits.x);
+		InterlockedAdd(TangentYTexture[VoxelCoord], Tangent8bits.y);
+		InterlockedAdd(TangentZTexture[VoxelCoord], Tangent8bits.z);
+		#endif	
+
+		#if SUPPORT_MATERIAL_PROPERTY == 1
+		{
+			ResolvedView = ResolveView();
+
+			FMaterialPixelParameters MaterialParameters = GetMaterialPixelParameters(Interpolants, SvPosition);
+			FPixelMaterialInputs PixelMaterialInputs;
+
+			float4 ScreenPosition = SvPositionToResolvedScreenPosition(SvPosition);
+			float3 TranslatedWorldPosition = SvPositionToResolvedTranslatedWorld(SvPosition);
+			CalcMaterialParametersEx(MaterialParameters, PixelMaterialInputs, SvPosition, ScreenPosition, true, TranslatedWorldPosition, TranslatedWorldPosition);
+
+			// For material we take the max color / roughness for efficiency resource instead of accumulated differences
+			const float3 BaseColor = GetMaterialBaseColor(PixelMaterialInputs);
+			const float Roughness = GetMaterialRoughness(PixelMaterialInputs);
+			const uint EncodedBaseColorAndRoughness = EncodeBaseColorAndRoughness(BaseColor, Roughness);
+			InterlockedMax(MaterialTexture[VoxelCoord], EncodedBaseColorAndRoughness);
+		}
+		#endif
+
+	}
+#endif
+	OutColor = 0xFF;
+}
\ No newline at end of file
diff --git a/Engine/Shaders/Private/HairStrands/HairStrandsDeepShadowVS.usf b/Engine/Shaders/Private/HairStrands/HairStrandsDeepShadowVS.usf
new file mode 100644
--- /dev/null
+++ b/Engine/Shaders/Private/HairStrands/HairStrandsDeepShadowVS.usf
@@ -0,0 +1,76 @@
+// Copyright 1998-2019 Epic Games, Inc. All Rights Reserved.
+
+#include "../Common.ush"
+#include "/Engine/Generated/Material.ush"
+#include "/Engine/Generated/VertexFactory.ush"
+
+#define MESH_RENDER_DEPTH 0
+#define MESH_RENDER_DOM   1
+#define MESH_RENDER_VOXEL 2
+
+struct FDeepShadowAccumulateVSToPS
+{
+#if SUPPORT_TANGENT_PROPERTY == 1 || SUPPORT_MATERIAL_PROPERTY == 1
+	FVertexFactoryInterpolantsVSToPS Interpolants;
+#endif
+	float4 Position : SV_POSITION;
+#if MESH_RENDER_MODE == MESH_RENDER_VOXEL
+	float4 PixelPosition : TEXCOORD6;
+#endif
+#if MESH_RENDER_MODE == MESH_RENDER_DOM || MESH_RENDER_MODE == MESH_RENDER_VOXEL
+	float HairCoverage : CUSTOM_COVERAGE;
+#endif
+};
+
+#define VS_OUTPUT_TYPE FDeepShadowAccumulateVSToPS
+#define SUPPORT_WPO 0
+
+#if VERTEXSHADER
+
+/** transform mesh as normal */
+void Main(
+	FVertexFactoryInput Input,
+	OPTIONAL_VertexID
+	out VS_OUTPUT_TYPE Output
+	)
+{
+	ResolvedView = ResolveView();
+
+	// #hair_todo
+	// Override the View direction with the light direction, since the Deep Shadow pass does not setup the View data, but passes data through a dedicated constant buffer
+	// Hack: Use 1337 as a flag for the strand vertex factory to not use MinRadiusAtDepth1 for computing the minimum radius for directional light. 
+#if SM5_PROFILE && COMPILER_HLSL
+	ResolvedView.TranslatedWorldCameraOrigin = float3(1337, DeepShadowPass.bIsOrtho, DeepShadowPass.MinStrandRadius);
+	ResolvedView.ViewForward = DeepShadowPass.LightDirection;
+#endif
+
+	FVertexFactoryIntermediates VFIntermediates = GetVertexFactoryIntermediates(Input);
+	float4 WorldPosition = VertexFactoryGetWorldPosition(Input, VFIntermediates);		
+	WorldPosition.xyz -= ResolvedView.PreViewTranslation.xyz; // Remove the camera relative position (todo move this into matrix transform)
+
+	Output.Position = mul(WorldPosition, DeepShadowPass.WorldToClipMatrix);
+#if MESH_RENDER_MODE == MESH_RENDER_VOXEL && defined(VF_STRAND_HAIR)
+	Output.PixelPosition = WorldPosition;
+  #if SUPPORT_TANGENT_PROPERTY == 1 || SUPPORT_MATERIAL_PROPERTY == 1
+	float3x3 TangentToLocal = VertexFactoryGetTangentToLocal(Input, VFIntermediates);
+	FMaterialVertexParameters VertexParameters = GetMaterialVertexParameters(Input, VFIntermediates, WorldPosition.xyz, TangentToLocal);
+	Output.Interpolants = VertexFactoryGetInterpolantsVSToPS(Input, VFIntermediates, VertexParameters);
+  #endif
+#endif
+
+#if MESH_RENDER_MODE == MESH_RENDER_DOM || MESH_RENDER_MODE == MESH_RENDER_VOXEL
+  #ifdef VF_STRAND_HAIR	
+	const float CurrentDepth = Output.Position.z / Output.Position.w;
+	const float PixelRadius = CurrentDepth * DeepShadowPass.MinStrandRadius;
+	const float StrandRealRadius = VFIntermediates.HairDimensions.y;
+	const float Coverage = StrandRealRadius / max(StrandRealRadius, PixelRadius);
+	Output.HairCoverage = Coverage * VFIntermediates.HairDensity;
+  #else
+	Output.HairCoverage = 1;
+  #endif
+#endif
+
+	OutputVertexID( Output );
+}
+
+#endif // VERTEXSHADER
diff --git a/Engine/Shaders/Private/HairStrands/HairStrandsDeepTransmittanceCommon.ush b/Engine/Shaders/Private/HairStrands/HairStrandsDeepTransmittanceCommon.ush
new file mode 100644
--- /dev/null
+++ b/Engine/Shaders/Private/HairStrands/HairStrandsDeepTransmittanceCommon.ush
@@ -0,0 +1,300 @@
+// Copyright 1998-2019 Epic Games, Inc. All Rights Reserved.
+
+//#include "HairStrandsShadowCommon.ush"
+
+float3 HUEtoRGB(in float H)
+{
+	float R = abs(H * 6 - 3) - 1;
+	float G = 2 - abs(H * 6 - 2);
+	float B = 2 - abs(H * 6 - 4);
+	return saturate(float3(R, G, B));
+}
+
+float3 HSVtoRGB(in float3 HSV)
+{
+	float3 RGB = HUEtoRGB(HSV.x);
+	return ((RGB - 1) * HSV.y + 1) * HSV.z;
+}
+
+float3 ColorCode(float x)
+{
+	float c = (1 - saturate(x)) * 0.6; // Remap [0,1] to Blue-Red
+	return x > 0 ? HSVtoRGB(float3(c, 1, 1)) : float3(0, 0, 0);
+}
+
+
+float3 ToLightPosition(float3 WorldPosition, float4x4 WorldToLightTransform)
+{
+	float4 LightPos = mul(float4(WorldPosition,1), WorldToLightTransform);
+	LightPos.xyz/= LightPos.w;
+	const float2 LightUV = (LightPos.xy + float(1).xx) * 0.5f;
+	return float3(LightUV.x, 1-LightUV.y, LightPos.z);
+}
+
+float SampleDOM_PCF2x2(
+	float3 LightSpacePosition, 
+	float DepthBias,
+	Texture2D<float> FrontDepthTexture,
+	Texture2D<float4> DomTexture)
+{
+	// Find the bottom left corner texel for bilinear interpolation
+	const float2 PBottomLeft = floor(LightSpacePosition.xy - float(0.5f).xx); // Bottom-left corner of the bottom left texel.
+	const uint2 C0 = uint2(PBottomLeft);
+	const uint2 C1 = C0 + uint2(1,0);
+	const uint2 C2 = C0 + uint2(0,1);
+	const uint2 C3 = C0 + uint2(1,1);
+	
+	// @todo_hair: gather4
+	const float FrontDepth0 = FrontDepthTexture.Load(uint3(C0, 0)).x;
+	const float FrontDepth1 = FrontDepthTexture.Load(uint3(C1, 0)).x;
+	const float FrontDepth2 = FrontDepthTexture.Load(uint3(C2, 0)).x;
+	const float FrontDepth3 = FrontDepthTexture.Load(uint3(C3, 0)).x;
+
+	const float4 DOMValue0  = DomTexture.Load(uint3(C0, 0));
+	const float4 DOMValue1  = DomTexture.Load(uint3(C1, 0));
+	const float4 DOMValue2  = DomTexture.Load(uint3(C2, 0));
+	const float4 DOMValue3  = DomTexture.Load(uint3(C3, 0));
+	
+	const float HairCount0	= ComputeHairCount(DOMValue0, GetDomDistanceToFrontDepthWithBias(FrontDepth0, LightSpacePosition.z, DepthBias));
+	const float HairCount1	= ComputeHairCount(DOMValue1, GetDomDistanceToFrontDepthWithBias(FrontDepth1, LightSpacePosition.z, DepthBias));
+	const float HairCount2	= ComputeHairCount(DOMValue2, GetDomDistanceToFrontDepthWithBias(FrontDepth2, LightSpacePosition.z, DepthBias));
+	const float HairCount3	= ComputeHairCount(DOMValue3, GetDomDistanceToFrontDepthWithBias(FrontDepth3, LightSpacePosition.z, DepthBias));
+
+	const float2 S = frac(LightSpacePosition.xy - (PBottomLeft + 0.5f));
+	const float HairCount01 = lerp(HairCount0, HairCount1, S.x);
+	const float HairCount23 = lerp(HairCount2, HairCount3, S.x);
+	return lerp(HairCount01, HairCount23, S.y);
+}
+
+float SampleDOM_PCF(
+	float3 LightSpacePosition, 
+	float DepthBias,
+	Texture2D<float> FrontDepthTexture,
+	Texture2D<float4> DomTexture)
+{	
+	// Explicit 1 ring (6x6 / 5x5) with uniform weighting
+	float HairCount = 0;
+	float w = 1;
+	HairCount += w * SampleDOM_PCF2x2(LightSpacePosition + float3(-2,-2, 0), DepthBias, FrontDepthTexture, DomTexture);
+	HairCount += w * SampleDOM_PCF2x2(LightSpacePosition + float3( 0,-2, 0), DepthBias, FrontDepthTexture, DomTexture);
+	HairCount += w * SampleDOM_PCF2x2(LightSpacePosition + float3( 2,-2, 0), DepthBias, FrontDepthTexture, DomTexture);
+
+	HairCount += w * SampleDOM_PCF2x2(LightSpacePosition + float3(-2, 0, 0), DepthBias, FrontDepthTexture, DomTexture);
+	HairCount += w * SampleDOM_PCF2x2(LightSpacePosition + float3( 0, 0, 0), DepthBias, FrontDepthTexture, DomTexture);
+	HairCount += w * SampleDOM_PCF2x2(LightSpacePosition + float3( 2, 0, 0), DepthBias, FrontDepthTexture, DomTexture);
+
+	HairCount += w * SampleDOM_PCF2x2(LightSpacePosition + float3(-2, 2, 0), DepthBias, FrontDepthTexture, DomTexture);
+	HairCount += w * SampleDOM_PCF2x2(LightSpacePosition + float3( 0, 2, 0), DepthBias, FrontDepthTexture, DomTexture);
+	HairCount += w * SampleDOM_PCF2x2(LightSpacePosition + float3( 2, 2, 0), DepthBias, FrontDepthTexture, DomTexture);
+
+	HairCount /= w * 9;
+	return HairCount;
+}
+
+// Sample [-1,1]
+// Jitter [0,1]
+float2 ComputeJitteredSample(float2 Sample, float2 Jitter)
+{
+#if 0
+	const float2 NormSample = (Sample + float2(1, 1)) * 0.5f;
+	const float2 JitteredNormSample = frac(NormSample + Jitter);
+	return JitteredNormSample *2 - float2(1, 1);
+#else
+	return Sample;
+#endif
+}
+
+float SampleDOM_PCSS(
+	float3 LightSpacePosition,
+	uint2 DeepShadowResolution,
+	float DepthBias,
+	Texture2D<float> FrontDepthTexture,
+	Texture2D<float4> DomTexture,
+	SamplerState LinearSampler,
+	float ApexAngleInDegree,
+	float2 Jitter) // [0,1]
+{
+	// Poisson disk position http://developer.download.nvidia.com/whitepapers/2008/PCSS_Integration.pdf
+	float2 PoissonDisk[16] =
+	{
+		 float2(-0.94201624, -0.39906216),
+		 float2(0.94558609, -0.76890725),
+		 float2(-0.094184101, -0.92938870),
+		 float2(0.34495938, 0.29387760),
+		 float2(-0.91588581, 0.45771432),
+		 float2(-0.81544232, -0.87912464),
+		 float2(-0.38277543, 0.27676845),
+		 float2(0.97484398, 0.75648379),
+		 float2(0.44323325, -0.97511554),
+		 float2(0.53742981, -0.47373420),
+		 float2(-0.26496911, -0.41893023),
+		 float2(0.79197514, 0.19090188),
+		 float2(-0.24188840, 0.99706507),
+		 float2(-0.81409955, 0.91437590),
+		 float2(0.19984126, 0.78641367),
+		 float2(0.14383161, -0.14100790)
+	};
+
+	// Find the closest occluder from light point of view
+	const uint OccluderCount = 5;
+	float OccluderDistance = 0; // Occluder distance in clip space
+	for (uint OccluderIt = 0; OccluderIt < OccluderCount; ++OccluderIt)
+	{
+		const float2 Offset = ComputeJitteredSample(PoissonDisk[OccluderIt], Jitter) * DeepShadowResolution;
+		const float2 SamplePosition = LightSpacePosition.xy + Offset;
+		const float2 P = floor(SamplePosition - float(0.5f).xx);
+		const uint2 C0 = P;
+		const uint2 C1 = C0 + uint2(1, 0);
+		const uint2 C2 = C0 + uint2(0, 1);
+		const uint2 C3 = C0 + uint2(1, 1);
+
+		// @todo_hair: gather4
+		const float FrontDepth0 = FrontDepthTexture.Load(uint3(C0, 0)).x;
+		const float FrontDepth1 = FrontDepthTexture.Load(uint3(C1, 0)).x;
+		const float FrontDepth2 = FrontDepthTexture.Load(uint3(C2, 0)).x;
+		const float FrontDepth3 = FrontDepthTexture.Load(uint3(C3, 0)).x;
+
+		const float Distance0 = GetDomDistanceToFrontDepth(FrontDepth0, LightSpacePosition.z);
+		const float Distance1 = GetDomDistanceToFrontDepth(FrontDepth1, LightSpacePosition.z);
+		const float Distance2 = GetDomDistanceToFrontDepth(FrontDepth2, LightSpacePosition.z);
+		const float Distance3 = GetDomDistanceToFrontDepth(FrontDepth3, LightSpacePosition.z);
+
+		OccluderDistance = max(OccluderDistance, max(Distance0, max(Distance1, max(Distance2, Distance3))));
+	}
+
+	// Sample DOM
+	float HairCount = 0;
+	const uint SampleCount = 16;
+	const float ApexAngleInRad = ApexAngleInDegree / 180.f * 3.141592653f;
+	const float ApexHalfAngleInRad = ApexAngleInRad * 0.5f;
+	const float SampleRadius = OccluderDistance * tan(ApexHalfAngleInRad); // Sample radius in clip space
+	for (uint SampleIt = 0; SampleIt < SampleCount; ++SampleIt)
+	{
+		const float2 Offset = ComputeJitteredSample(PoissonDisk[SampleIt], Jitter) * SampleRadius * DeepShadowResolution;
+		HairCount += SampleDOM_PCF2x2(LightSpacePosition + float3(Offset,0), DepthBias, FrontDepthTexture, DomTexture);
+	}
+	
+	HairCount /= SampleCount;
+	return HairCount;
+}
+
+struct FTransmittanceSettings
+{
+	int2	 DeepShadowResolution;
+	int2	 DeepShadowOffset;
+	float4x4 WorldToLightTransform;
+	float3   LightDirection;
+	float4   LightPosition;
+
+	float    DeepShadowDepthBiasScale;
+	float    DeepShadowDensityScale;
+	float	 TransmittanceKernelApertureInDegree;
+	float3   Random;
+	uint	 DebugMode;
+
+	uint TransmittanceKernelType;
+};
+
+#define TILE_PIXEL_SIZE 8
+
+float3 ComputeDomDebugColor(
+	float3 WorldPosition, 
+	float4x4 WorldToLightTransform, 
+	Texture2D<float> FrontDepthTexture, 
+	SamplerState LinearSampler)
+{
+	const float3 LightSpacePosition = ToLightPosition(WorldPosition, WorldToLightTransform); 
+	const float FrontDepth = FrontDepthTexture.SampleLevel(LinearSampler, LightSpacePosition.xy, 0).x;
+	const float DistanceToFrontDepth = GetDomDistanceToFrontDepth(FrontDepth, LightSpacePosition.z);
+	return HairDebugColor(DistanceToFrontDepth);
+}
+
+float3 ComputeHairCountDebugColor(
+	float3 WorldPosition, 
+	float4x4 WorldToLightTransform, 
+	Texture2D<float> FrontDepthTexture, 
+	Texture2D<float4> DomTexture, 
+	SamplerState LinearSampler, 
+	float MaxHairCount)
+{
+	const float3 LightSpacePosition = ToLightPosition(WorldPosition, WorldToLightTransform);
+	const float FrontDepth = FrontDepthTexture.SampleLevel(LinearSampler, LightSpacePosition.xy, 0).x;
+	const float4 DOMValue = DomTexture.SampleLevel(LinearSampler, LightSpacePosition.xy, 0);
+	const float DistanceToFrontDepth = GetDomDistanceToFrontDepth(FrontDepth, LightSpacePosition.z);
+	const float HairCount = ComputeHairCount(DOMValue, DistanceToFrontDepth);
+	return ColorCode(HairCount / MaxHairCount);
+}
+
+float4 ComputeTransmittanceFromHairCount(float HairCount, FGBufferData GBuffer, float SinLightAngle, Texture3D<float4> InHairScatteringLUTTexture, SamplerState InHairLUTSampler)
+{
+	HairAverageScattering AbsorptionData = SampleHairLUT(InHairScatteringLUTTexture, InHairLUTSampler, GBuffer.BaseColor, GBuffer.Roughness, SinLightAngle);
+
+	// Always shift the hair count by one to remove self-occlusion/shadow aliasing and have smootheer transition
+	// This insure the the pow function always starts at 0 for front facing hair
+	HairCount = max(0, HairCount - 1);
+	return  float4(pow(AbsorptionData.A_front, HairCount), HairCount);
+}
+
+float4 ComputeTransmittance(
+	float3 WorldPosition,
+	FGBufferData GBuffer, 
+	FTransmittanceSettings Settings,
+	Texture2D<float>	FrontDepthTexture,
+	Texture2D<float4>	DomTexture,
+	Texture3D<float4>	HairLUTTexture,
+	SamplerState		LinearSampler)
+{
+	const bool bIsHairPixel = GBuffer.ShadingModelID == SHADINGMODELID_HAIR;
+
+	float4 Transmittance = float4(1,1,1,0);
+	if (bIsHairPixel)
+	{
+		const bool bIsDirectionalLight = Settings.LightPosition.w == 0;
+
+		// LightSpacePosition is 'WorldPosition' in LightSpace (TexelCoord.XY / NDC Depth.Z)
+		const float3 LightSpacePosition = ToLightPosition(WorldPosition, Settings.WorldToLightTransform) * float3(Settings.DeepShadowResolution, 1) + float3(Settings.DeepShadowOffset, 0); 
+		const float3 L = bIsDirectionalLight ? Settings.LightDirection : normalize(Settings.LightPosition.xyz - WorldPosition);
+		const float3 T = GBuffer.WorldNormal;
+		const float SinLightAngle = dot(L, T);
+
+		float DepthBias = 0;
+		if (Settings.DeepShadowDepthBiasScale > 0)
+		{
+			#if 0 
+			const float DephtBiasSlope = 6;
+			const float DefaultDepthBias = LAYER_DEPTH0 * Settings.DeepShadowDepthBiasScale; 
+			DepthBias = (CosLightAngle*DephtBiasSlope + 1) * DefaultDepthBias;
+			#else
+			DepthBias = LAYER_DEPTH0 * Settings.DeepShadowDepthBiasScale;
+			#endif
+		}
+
+		// Compute the number of hair count between light & shading point
+		float HairCount = 0;
+		if (Settings.TransmittanceKernelType == 3)
+		{
+			HairCount = SampleDOM_PCSS(LightSpacePosition.xyz, Settings.DeepShadowResolution, DepthBias, FrontDepthTexture, DomTexture, LinearSampler, Settings.TransmittanceKernelApertureInDegree, Settings.Random.xy);
+		}
+		else if (Settings.TransmittanceKernelType == 2)
+		{
+			HairCount = SampleDOM_PCF(LightSpacePosition.xyz, DepthBias, FrontDepthTexture, DomTexture);
+		}
+		else if (Settings.TransmittanceKernelType == 1)
+		{
+			HairCount = SampleDOM_PCF2x2(LightSpacePosition.xyz, DepthBias, FrontDepthTexture, DomTexture);
+		}
+		else
+		{
+			HairCount = 0;
+		}
+
+		Transmittance = ComputeTransmittanceFromHairCount(HairCount * Settings.DeepShadowDensityScale, GBuffer, SinLightAngle, HairLUTTexture, LinearSampler);
+
+		if (Settings.DebugMode != 0)
+		{
+			Transmittance.xyz = ComputeDomDebugColor(WorldPosition, Settings.WorldToLightTransform, FrontDepthTexture, LinearSampler);
+			Transmittance.w = 1;
+		}
+	}
+	return Transmittance;
+}
diff --git a/Engine/Shaders/Private/HairStrands/HairStrandsDeepTransmittanceDualScattering.ush b/Engine/Shaders/Private/HairStrands/HairStrandsDeepTransmittanceDualScattering.ush
new file mode 100644
--- /dev/null
+++ b/Engine/Shaders/Private/HairStrands/HairStrandsDeepTransmittanceDualScattering.ush
@@ -0,0 +1,187 @@
+// Copyright 1998-2019 Epic Games, Inc. All Rights Reserved.
+
+float Hair_g2(float B,float Theta)
+{
+	#if 0
+	return Hair_g(B,Theta);
+	#else
+	// Unnormalized Gaussian
+	return exp( -0.5 * Pow2( Theta ) / (B*B) );
+	#endif
+}
+
+// Dual scattering computation are done here for faster iteration (i.e., does not invalidate tons of shaders)
+void ComputeDualScatteringTerms(
+	const FHairTransmittanceData TransmittanceData,
+	const float Roughness,
+	const float3 V,
+	const float3 L,
+	const float3 T,
+	const float4 HairComponents,
+	out float3 OutGlobalScattering,
+	out float3 OutLocalScattering)
+{
+	const float SinThetaL = clamp(dot(T, L), -1, 1);
+	const float SinThetaV = clamp(dot(T, V), -1, 1);
+
+	const float MaxAverageScatteringValue = 0.99f;
+
+	// Straight implementation of the dual scattering paper 
+    const float3 af	 = min(MaxAverageScatteringValue.xxx, TransmittanceData.A_front);
+	const float3 af2 = Pow2(af);
+	const float3 ab  = min(MaxAverageScatteringValue.xxx, TransmittanceData.A_back);
+	const float3 ab2 = Pow2(ab);
+    const float3 OneMinusAf2 = 1 - af2;
+
+	const float3 A1 = ab * af2 / OneMinusAf2;
+	const float3 A3 = ab * ab2 * af2 / (OneMinusAf2*Pow2(OneMinusAf2));
+	const float3 Ab = A1 + A3;
+
+	// todo_hair: Should use the average roughness after k scattering event instead of the BSDF roughness
+	const float Beta_R	 = Pow2( Roughness );
+	const float Beta_TT	 = Pow2( Roughness / 2 ); 
+	const float Beta_TRT = Pow2( Roughness * 2 ); 
+
+	const float Shift     = 0.035;
+	const float Shift_R   =-0.035*2;
+	const float Shift_TT  = 0.035;
+	const float Shift_TRT = 0.035*4;
+
+	// Average density factor (This is the constant used in the original paper)
+	const float df = 0.7f;
+	const float db = 0.7f;
+
+	// This is a coarse approximation of eq. 13. Normally, Beta_f should be weighted by the 'normalized' 
+	// R, TT, and TRT terms 
+	const float3 af_weights = af / (af.r + af.g + af.b);
+    const float3 Beta_f  = dot(float3(Beta_R, Beta_TT, Beta_TRT), af_weights);
+	const float3 Beta_f2 = Beta_f*Beta_f;
+    const float3 sigma_f = Beta_f * max(1, TransmittanceData.HairCount);
+
+    const float Theta_h = (asin(SinThetaL) + asin(SinThetaV)) * 0.5f;
+
+	// Global scattering spread 'Sf'
+	// It is unclear if this is correct to use a normalized Gaussian as in the original paper, as with an increase 
+	// number of hair count the Gaussian will become more and more spread, and have less and less contribution. 
+	// This is correct from an energy point of view, as long as the scattering is integrated over multiple scattering 
+	// event, which is not our case here.
+	// This is why there is a non-normalized Gaussian version in case this does not fly properly with regular 
+	// normalized Gaussian
+	// const float3 Sf = 1;
+	// const float3 Sf = Hair_g_unnormalized(sigma_f, Theta_h) / PI; 
+    float3 Sf = float3(	Hair_g2(sigma_f.r, Theta_h), 
+						Hair_g2(sigma_f.g, Theta_h), 
+						Hair_g2(sigma_f.b, Theta_h)) / PI;
+	const float3 Tf = TransmittanceData.Transmittance;
+
+	// Overall shift due to the various local scatteing event (all shift & roughnesss should vary with color)
+	//const float3 shift_f = 0.33f * (Shift_R + Shift_TT + Shift_TRT);
+    const float3 shift_f = dot(float3(Shift_R, Shift_TT, Shift_TRT), af_weights);
+	const float3 shift_b = shift_f;
+    const float3 delta_b = 0; //shift_b * (1 - 2*ab2 / Pow2(1 - af2)) * shift_f * (2 * Pow2(1 - af2) + 4*af2*ab2)/Pow3(1-af2);
+
+	const float3 ab_weights = ab / (ab.r + ab.g + ab.b);
+	const float3 Beta_b  = dot(float3(Beta_R, Beta_TT, Beta_TRT), ab_weights);
+	const float3 Beta_b2 = Beta_f2;
+
+    const float3 sigma_b = (1 + db*af2) * (ab*sqrt(2*Beta_f2 + Beta_b2) + ab*ab2*sqrt(2*Beta_f2 + Beta_b2)) / (ab + ab*ab2*(2*Beta_f + 3*Beta_b));
+
+	// Local scattering Spread 'Sb'
+    float3 Sb = float3(	Hair_g2(Pow2(sigma_b.r), Theta_h - delta_b.r),
+						Hair_g2(Pow2(sigma_b.g), Theta_h - delta_b.g),
+						Hair_g2(Pow2(sigma_b.b), Theta_h - delta_b.b))  / PI;
+
+	// Override the spread function to cancel them
+	//Sf = 1;
+	//Sb = 1;
+
+    // Different variant for managing sefl-occlusion issue for global scattering
+	const bool bIsOccluded = TransmittanceData.HairCount > 0;
+	const float3 GlobalScattering = bIsOccluded ? Tf * Sf * df : 1;
+	//const float3 GlobalScattering = TransmittanceData.HairCount >= 1 ? Tf * Sf * df : 1;
+    //const float3 GlobalScattering = lerp(1, Tf * Sf * df, saturate(TransmittanceData.HairCount));
+	//const float3 GlobalScattering = Tf * Sf * df;
+
+	const float3 LocalScattering  = 2 * Ab * Sb * db;
+	
+	// Decode if the Local/Global component should be enabled or not (for debug purpose)
+    const float LocalWeights  = floor(HairComponents.w / 10.0f) > 0?1:0;
+    const float GlobalWeights = frac(HairComponents.w / 10.0f)*10.f > 0?1:0;
+
+	OutGlobalScattering = GlobalWeights > 0 ? GlobalScattering : 1;
+	OutLocalScattering  = LocalScattering * LocalWeights;
+
+	// Final computation is done in ShadingModels.ush with the following formula
+	// GlobalScattering * (Fs + LocalScattering) * TransmittanceData.OpaqueVisibility;
+}
+	
+FHairTransmittanceData GetHairTransmittance(
+	const float3 V,
+	const float3 L,
+	FGBufferData GBuffer,
+	float4 HairTransmittance,
+	Texture3D<float4> HairLUTTexture,
+	SamplerState HairLUTSampler,
+	const uint HairTransmittanceEnable,
+	const float4 HairComponents)
+{
+	FHairTransmittanceData TransmittanceData = InitHairTransmittanceData();
+	TransmittanceData.Transmittance = HairTransmittance.xyz;
+	TransmittanceData.HairCount = HairTransmittance.w;
+	TransmittanceData.MultipleScattering = 1;
+	TransmittanceData.GlobalScattering = 1;
+	TransmittanceData.LocalScattering = 0;
+	TransmittanceData.ScatteringComponent = float3(HairComponents.x > 0 ? 1 : 0, HairComponents.y > 0 ? 1 : 0, HairComponents.z > 0 ? 1 : 0);
+	TransmittanceData.OpaqueVisibility = 1;
+	TransmittanceData.Enable = HairTransmittanceEnable > 0;
+
+	const float3 T = GBuffer.WorldNormal;
+	{
+		const float SinLightAngle = dot(L, T);
+
+		// Compute the transmittance based on precompute Hair transmittance LUT
+		HairAverageScattering AverageScattering = SampleHairLUT(HairLUTTexture, HairLUTSampler, GBuffer.BaseColor, GBuffer.Roughness, SinLightAngle);
+		TransmittanceData.A_front = AverageScattering.A_front;
+		TransmittanceData.A_back = AverageScattering.A_back;
+	}
+
+	ComputeDualScatteringTerms(
+		TransmittanceData,
+		GBuffer.Roughness,
+		V,
+		L,
+		T,
+		HairComponents,
+		TransmittanceData.GlobalScattering,
+		TransmittanceData.LocalScattering);
+
+	return TransmittanceData;
+}
+
+FHairTransmittanceData GetHairTransmittance(
+	const float3 V,
+	const float3 L,
+	FGBufferData GBuffer,
+	const uint SampleOffset,
+	const uint HairTransmittanceBufferMaxCount,
+	StructuredBuffer<float4> HairTransmittanceBuffer,
+	Texture3D<float4> HairLUTTexture,
+	SamplerState HairLUTSampler,
+	const uint HairTransmittanceEnable,
+	const float4 HairComponents)
+{
+	float4 HairTransmittance = float4(1, 1, 1, 0);
+	if (SampleOffset < HairTransmittanceBufferMaxCount)
+	{
+		HairTransmittance = HairTransmittanceBuffer[SampleOffset];
+	}
+	return GetHairTransmittance(
+		V,
+		L,
+		GBuffer,
+		HairTransmittance,
+		HairLUTTexture,
+		HairLUTSampler,
+		HairTransmittanceEnable,
+		HairComponents);
+}
diff --git a/Engine/Shaders/Private/HairStrands/HairStrandsDeepTransmittanceMask.usf b/Engine/Shaders/Private/HairStrands/HairStrandsDeepTransmittanceMask.usf
new file mode 100644
--- /dev/null
+++ b/Engine/Shaders/Private/HairStrands/HairStrandsDeepTransmittanceMask.usf
@@ -0,0 +1,201 @@
+// Copyright 1998-2019 Epic Games, Inc. All Rights Reserved.
+
+#include "../Common.ush"
+#include "../CommonViewUniformBuffer.ush"
+#include "../SceneTextureParameters.ush"
+#include "../DeferredShadingCommon.ush"
+#include "../ShadowFilteringCommon.ush"
+#include "HairStrandsDeepShadowCommon.ush"
+#include "HairStrandsCommon.ush"
+#include "HairStrandsDeepTransmittanceCommon.ush"
+#include "HairStrandsVisibilityCommon.ush"
+#include "HairStrandsDeepVoxelCommon.ush"
+
+float3 ToWorldPosition(float2 UV, float Depth)
+{
+	float2 ScreenPosition = (UV - View.ScreenPositionScaleBias.wz) / View.ScreenPositionScaleBias.xy;
+	float4 HomogeneousWorldPosition = mul(float4(ScreenPosition * Depth, Depth, 1), View.ScreenToWorld);
+	float3 WorldPosition = HomogeneousWorldPosition.xyz / HomogeneousWorldPosition.w;
+
+	return WorldPosition;
+}
+
+#define PERMUTATION_TRANSMITTANCE_DEEP_SHADOW 0
+#define PERMUTATION_TRANSMITTANCE_VOXEL 1
+
+// Need to match FDeepShadowData::MaxClusterCount in HairRendering.h
+#define MAX_HAIR_CLUSTER_COUNT 16 
+
+uint2				DeepShadow_Resolution;
+float4x4			DeepShadow_WorldToLightTransforms[MAX_HAIR_CLUSTER_COUNT];
+float3				LightDirection;
+float4				LightPosition;
+float				DepthBiasScale;
+float				DensityScale;
+float				DeepShadow_KernelAperture;
+uint				DeepShadow_KernelType;
+uint				DeepShadow_DebugMode;
+int2				DeepShadow_AtlasSlotOffsets[MAX_HAIR_CLUSTER_COUNT];
+Texture2D<float>	DeepShadow_FrontDepthTexture;
+Texture2D<float4>	DeepShadow_DomTexture;
+
+float3				Voxel_MinAABBs[MAX_HAIR_CLUSTER_COUNT];
+float3				Voxel_MaxAABBs[MAX_HAIR_CLUSTER_COUNT];
+uint				Voxel_Resolution;
+Texture3D<uint> 	Voxel_DensityTexture0;
+Texture3D<uint> 	Voxel_DensityTexture1;
+Texture3D<uint> 	Voxel_DensityTexture2;
+Texture3D<uint> 	Voxel_DensityTexture3;
+Texture3D<uint> 	Voxel_DensityTexture4;
+Texture3D<uint> 	Voxel_DensityTexture5;
+Texture3D<uint> 	Voxel_DensityTexture6;
+Texture3D<uint> 	Voxel_DensityTexture7;
+
+StructuredBuffer<FNodeData> HairVisibilityNodeData;
+StructuredBuffer<uint>		HairVisibilityNodeCoord;
+StructuredBuffer<uint4>		IndirectArgsBuffer;
+
+Texture3D<float4>			HairLUTTexture;
+RWStructuredBuffer<float4>	OutputColor;
+SamplerState				LinearSampler;
+
+#if PERMUTATION_GROUP_SIZE == 0
+	#define TRANSMITTANCE_GROUP_SIZE 64
+#elif PERMUTATION_GROUP_SIZE == 1
+	#define TRANSMITTANCE_GROUP_SIZE 32
+#else
+	#error Unknown group size	
+#endif
+
+#if PERMUTATION_TRANSMITTANCE_TYPE == PERMUTATION_TRANSMITTANCE_DEEP_SHADOW
+[numthreads(TRANSMITTANCE_GROUP_SIZE, 1, 1)]
+void MainCS(uint2 DispatchThreadId : SV_DispatchThreadID)
+{
+	// Note: Use a fixed group count width (HAIR_VISIBILITY_GROUP_COUNT_WIDTH) for avoiding loading the indirect args (to avoid dep. memory fetch)
+	const uint LocalOffset = DispatchThreadId.x + DispatchThreadId.y * TRANSMITTANCE_GROUP_SIZE * HAIR_VISIBILITY_GROUP_COUNT_WIDTH;
+
+	const uint PackedCoord	 = HairVisibilityNodeCoord[LocalOffset];
+	const uint2 PixelCoord	 = uint2((PackedCoord & 0xFFFF), ((PackedCoord>>16) & 0xFFFF));
+
+	const float2 UV = (PixelCoord + float2(0.5f, 0.5f)) / float2(View.BufferSizeAndInvSize.xy);
+	const float3 TransmittanceRandom = float3(InterleavedGradientNoise(PixelCoord, 1), InterleavedGradientNoise(PixelCoord, 2), InterleavedGradientNoise(PixelCoord, 3));
+
+	{
+		const FNodeData NodeData = HairVisibilityNodeData[LocalOffset];
+		float			SampleDepth;
+		FVisibilityData SampleVisData;
+		float3			SampleBaseColor;
+		float			SampleRoughness;
+		float			SampleSpecular;
+		DecodeNodeData(NodeData, SampleDepth, SampleVisData, SampleBaseColor, SampleRoughness, SampleSpecular);
+
+		FGBufferData GBuffer;
+		GBuffer.ShadingModelID = SHADINGMODELID_HAIR;
+		GBuffer.WorldNormal = SampleVisData.Tangent;
+		GBuffer.BaseColor = SampleBaseColor;
+		GBuffer.Roughness = SampleRoughness;
+
+		const uint ClusterIndex = clamp(SampleVisData.ClusterId, 0, MAX_HAIR_CLUSTER_COUNT - 1);
+		const int2 DeepShadowOffset = DeepShadow_AtlasSlotOffsets[ClusterIndex];
+		const float4x4 WorldToLightTransform = DeepShadow_WorldToLightTransforms[ClusterIndex];
+		
+		FTransmittanceSettings Settings;
+		Settings.DeepShadowOffset = DeepShadowOffset;
+		Settings.DeepShadowResolution = DeepShadow_Resolution;
+		Settings.WorldToLightTransform = WorldToLightTransform;
+		Settings.LightDirection = LightDirection;
+		Settings.LightPosition = LightPosition;
+		Settings.DeepShadowDepthBiasScale = DepthBiasScale;
+		Settings.DeepShadowDensityScale = DensityScale;
+		Settings.TransmittanceKernelApertureInDegree = DeepShadow_KernelAperture;
+		Settings.TransmittanceKernelType = DeepShadow_KernelType;
+		Settings.Random = TransmittanceRandom;
+		Settings.DebugMode = DeepShadow_DebugMode;
+
+		const float SceneDepth = ConvertFromDeviceZ(SampleDepth);
+		const float3 WorldPosition = ToWorldPosition(UV, SceneDepth);
+		OutputColor[LocalOffset] = ComputeTransmittance(WorldPosition, GBuffer, Settings, DeepShadow_FrontDepthTexture, DeepShadow_DomTexture, HairLUTTexture, LinearSampler);
+	}
+}
+#endif
+
+#if PERMUTATION_TRANSMITTANCE_TYPE == PERMUTATION_TRANSMITTANCE_VOXEL
+
+float4 ComputeTransmittanceVoxel(
+	float3 WorldPosition,
+	uint ClusterId,
+	FGBufferData GBuffer,
+	FTransmittanceSettings Settings,
+	Texture3D<float4>	HairLUTTexture,
+	SamplerState		HairLUTSampler)
+{
+	const bool bIsDirectionalLight = Settings.LightPosition.w == 0;	
+	const float3 L = bIsDirectionalLight ? -Settings.LightDirection : normalize(Settings.LightPosition.xyz - WorldPosition);
+	const float3 T = GBuffer.WorldNormal;
+	const float SinLightAngle = dot(L, T);
+
+	// Depth bias
+	// Origin is shifted voxels away towards the light + a constant bias of the size of the voxel
+	{
+		const float3 VoxelSize = (Voxel_MaxAABBs[ClusterId] - Voxel_MinAABBs[ClusterId]) / Voxel_Resolution;
+		const float3 DepthBias = max(VoxelSize.x, max(VoxelSize.y, VoxelSize.z)) * (Settings.DeepShadowDepthBiasScale*L + Settings.Random);
+		WorldPosition += DepthBias;
+	}
+
+	// Compute the number of hair count between light & shading point
+	float HairCount = 0;
+	if (ClusterId == 0) HairCount = ComputeHairCountVoxel(WorldPosition, L, Voxel_MinAABBs[0], Voxel_MaxAABBs[0], Voxel_Resolution, Voxel_DensityTexture0, Settings.DeepShadowDensityScale, 0);
+	if (ClusterId == 1) HairCount = ComputeHairCountVoxel(WorldPosition, L, Voxel_MinAABBs[1], Voxel_MaxAABBs[1], Voxel_Resolution, Voxel_DensityTexture1, Settings.DeepShadowDensityScale, 0);
+	if (ClusterId == 2) HairCount = ComputeHairCountVoxel(WorldPosition, L, Voxel_MinAABBs[2], Voxel_MaxAABBs[2], Voxel_Resolution, Voxel_DensityTexture2, Settings.DeepShadowDensityScale, 0);
+	if (ClusterId == 3) HairCount = ComputeHairCountVoxel(WorldPosition, L, Voxel_MinAABBs[3], Voxel_MaxAABBs[3], Voxel_Resolution, Voxel_DensityTexture3, Settings.DeepShadowDensityScale, 0);
+	if (ClusterId == 4) HairCount = ComputeHairCountVoxel(WorldPosition, L, Voxel_MinAABBs[4], Voxel_MaxAABBs[4], Voxel_Resolution, Voxel_DensityTexture4, Settings.DeepShadowDensityScale, 0);
+	if (ClusterId == 5) HairCount = ComputeHairCountVoxel(WorldPosition, L, Voxel_MinAABBs[5], Voxel_MaxAABBs[5], Voxel_Resolution, Voxel_DensityTexture5, Settings.DeepShadowDensityScale, 0);
+	if (ClusterId == 6) HairCount = ComputeHairCountVoxel(WorldPosition, L, Voxel_MinAABBs[6], Voxel_MaxAABBs[6], Voxel_Resolution, Voxel_DensityTexture6, Settings.DeepShadowDensityScale, 0);
+	if (ClusterId == 7) HairCount = ComputeHairCountVoxel(WorldPosition, L, Voxel_MinAABBs[7], Voxel_MaxAABBs[7], Voxel_Resolution, Voxel_DensityTexture7, Settings.DeepShadowDensityScale, 0);
+
+	return ComputeTransmittanceFromHairCount(HairCount, GBuffer, SinLightAngle, HairLUTTexture, HairLUTSampler);
+}
+
+
+[numthreads(TRANSMITTANCE_GROUP_SIZE, 1, 1)]
+void MainCS(uint2 DispatchThreadId : SV_DispatchThreadID)
+{
+	const uint LocalOffset = DispatchThreadId.x + DispatchThreadId.y * TRANSMITTANCE_GROUP_SIZE * HAIR_VISIBILITY_GROUP_COUNT_WIDTH;
+	const uint PackedCoord = HairVisibilityNodeCoord[LocalOffset];
+	const uint2 PixelCoord = uint2((PackedCoord & 0xFFFF), ((PackedCoord >> 16) & 0xFFFF));
+
+	const float2 UV = (PixelCoord.xy + float2(0.5f, 0.5f)) / float2(View.BufferSizeAndInvSize.xy);
+	{
+		const FNodeData NodeData = HairVisibilityNodeData[LocalOffset];
+		float			SampleDepth;
+		FVisibilityData SampleVisData;
+		float3			SampleBaseColor;
+		float			SampleRoughness;
+		float			SampleSpecular;
+		DecodeNodeData(NodeData, SampleDepth, SampleVisData, SampleBaseColor, SampleRoughness, SampleSpecular);
+
+		FGBufferData GBuffer;
+		GBuffer.ShadingModelID = SHADINGMODELID_HAIR;
+		GBuffer.WorldNormal = SampleVisData.Tangent;
+		GBuffer.BaseColor = SampleBaseColor;
+		GBuffer.Roughness = SampleRoughness;
+		GBuffer.Specular = SampleSpecular;
+
+		const uint ClusterIndex = clamp(SampleVisData.ClusterId, 0, MAX_HAIR_CLUSTER_COUNT - 1);
+
+		FTransmittanceSettings Settings;
+		Settings.LightDirection = LightDirection;
+		Settings.LightPosition = LightPosition;
+		Settings.DeepShadowDensityScale = DensityScale;
+		Settings.DeepShadowDepthBiasScale = DepthBiasScale;
+		Settings.Random = GetHairVoxelJitter(PixelCoord.xy, View.StateFrameIndexMod8);
+
+		const float SceneDepth = ConvertFromDeviceZ(SampleDepth);
+		const float3 WorldPosition = ToWorldPosition(UV, SceneDepth);
+		OutputColor[LocalOffset] = ComputeTransmittanceVoxel(WorldPosition, ClusterIndex, GBuffer, Settings, HairLUTTexture, LinearSampler);
+	}
+}
+#endif
+
+
+
diff --git a/Engine/Shaders/Private/HairStrands/HairStrandsDeepVoxelCommon.ush b/Engine/Shaders/Private/HairStrands/HairStrandsDeepVoxelCommon.ush
new file mode 100644
--- /dev/null
+++ b/Engine/Shaders/Private/HairStrands/HairStrandsDeepVoxelCommon.ush
@@ -0,0 +1,153 @@
+// Copyright 1998-2019 Epic Games, Inc. All Rights Reserved.
+
+float3 GetHairVoxelJitter(uint2 PixelCoord, uint Seed)
+{
+	return float3(
+		InterleavedGradientNoise(PixelCoord.xy, Seed),
+		InterleavedGradientNoise(PixelCoord.xy, Seed * 117),
+		InterleavedGradientNoise(PixelCoord.xy, Seed * 7901));
+}
+
+float GetHairVoxelDensity(uint3 VolumeCoord, Texture3D<uint> InDensityTexture, uint Mip=0)
+{
+	const float VoxelFixPointScale = GetVoxelDensityFixPointScale();
+	return InDensityTexture.Load(int4(VolumeCoord, Mip)) / VoxelFixPointScale;
+}
+
+bool IsInVoxelBounds(float3 P, float3 MinP, float3 MaxP)
+{
+	return
+		P.x >= MinP.x && P.y >= MinP.y && P.z >= MinP.z &&
+		P.x <= MaxP.x && P.y <= MaxP.y && P.z <= MaxP.z;
+}
+
+float ComputeHairCountVoxel(float3 WorldPosition, float3 LightDirection, float3 InVoxelAABBMin, float3 InVoxelAABBMax, float InVoxelResolution, Texture3D<uint> InDensityTexture, float DensityScale, float CountThreshold)
+{
+	// Reverse-Z
+	const float ConservativeAABBDiagonal = (InVoxelAABBMax.x - InVoxelAABBMin.x) * 3;
+	const float3 O = WorldPosition;
+	const float3 E = O + LightDirection * ConservativeAABBDiagonal;
+	const float3 D = E - O;
+
+	const float VoxelSize = (InVoxelAABBMax.x - InVoxelAABBMin.x) / InVoxelResolution;
+
+	float OutHairCount = 0;
+
+	// The rest of the code assumes O is within the AABB
+	const bool bIsInAABB = IsInVoxelBounds(O, InVoxelAABBMin, InVoxelAABBMax);
+	if (!bIsInAABB)
+		return OutHairCount;
+
+	const float2 HitT = LineBoxIntersect(O, E, InVoxelAABBMin, InVoxelAABBMax);
+	if (HitT.x < HitT.y)
+	{
+		// Count the number of fibers which are within a cylinder defined by the voxel size, 
+		// and the distance between the origin and the extent of the volume
+		// This assumes that the voxel volume is cubic (i.e. equal dimensions on all sides)
+		const float Distance = length(HitT.y * D);
+		const uint StepCount = min(ceil(Distance / VoxelSize), 256u);
+		const float Delta = VoxelSize;
+		for (uint StepIt = 0; StepIt < StepCount; ++StepIt)
+		{
+			const float3 HitP = O + (StepIt*Delta) * LightDirection;
+			const float3 UVW = (HitP - InVoxelAABBMin) / (InVoxelAABBMax - InVoxelAABBMin);
+			const int3 VolumeCoord = clamp(int3(UVW * InVoxelResolution), int3(0, 0, 0), int3(InVoxelResolution.xxx-1));
+			const float Density = GetHairVoxelDensity(VolumeCoord, InDensityTexture) * DensityScale;
+			OutHairCount += Density;
+			if (CountThreshold != 0 && OutHairCount > CountThreshold)
+				return OutHairCount;
+		}
+	}
+
+	return OutHairCount;
+}
+
+float ComputeHairCountVoxelCone(
+	float3 WorldPosition,
+	float3 LightDirection,
+	float3 InVoxelAABBMin,
+	float3 InVoxelAABBMax,
+	float InVoxelResolution,
+	Texture3D<uint> InDensityTexture,
+	float InVoxelDensityScale,
+	float CountThreshold,
+	float TanConeAngle)
+{
+	// Reverse-Z
+	const float ConservativeAABBDiagonal = (InVoxelAABBMax.x - InVoxelAABBMin.x) * 3;
+	const float3 O = WorldPosition;
+	const float3 E = O + LightDirection * ConservativeAABBDiagonal;
+	const float3 D = E - O;
+
+	float OutHairCount = 0;
+
+	// The rest of the code assumes O is within the AABB
+	const bool bIsInAABB = IsInVoxelBounds(O, InVoxelAABBMin, InVoxelAABBMax);
+	if (!bIsInAABB)
+		return OutHairCount;
+
+	const float2 HitT = LineBoxIntersect(O, E, InVoxelAABBMin, InVoxelAABBMax);
+	if (HitT.x < HitT.y)
+	{
+		const float VolumeSize = (InVoxelAABBMax.x - InVoxelAABBMin.x);
+		const float VoxelSize = VolumeSize / InVoxelResolution;
+
+		// Count the number of fibers which are within a cylinder defined by the voxel size, 
+		// and the distance between the origin and the extent of the volume
+		// This assumes that the voxel volume is cubic (i.e. equal dimensions on all sides)
+		const float Distance = length(HitT.y * D);
+		const uint MaxStepCount = min(ceil(Distance / VoxelSize), 256u);
+		float T = 0;
+		float MipLevel = 0;
+		float MipDensityScale = 1;
+		for (uint StepIt = 0; StepIt < MaxStepCount; ++StepIt)
+		{
+			float3 StepOffset = 0;
+
+			const float3 HitP = O + T * LightDirection + StepOffset;
+			const float3 UVW = (HitP - InVoxelAABBMin) / (InVoxelAABBMax - InVoxelAABBMin);
+			if (!IsInVoxelBounds(UVW, 0, 1))
+				return OutHairCount;
+
+			const uint MipVoxelResolution = uint(InVoxelResolution) >> uint(MipLevel);
+			const int3 VolumeCoord = clamp(int3(UVW * MipVoxelResolution), int3(0, 0, 0), int3(MipVoxelResolution.xxx - 1));
+			const float Density = GetHairVoxelDensity(VolumeCoord, InDensityTexture, MipLevel) * InVoxelDensityScale * MipDensityScale;
+
+			OutHairCount += Density;
+			if (CountThreshold != 0 && OutHairCount > CountThreshold)
+				return OutHairCount;
+
+			const float VoxelRadius = T * TanConeAngle;
+			const float NumVoxel = VolumeSize / max(2 * VoxelRadius, VoxelSize);
+
+			const float ResolutionBias = InVoxelResolution / NumVoxel;
+			float MipBias = log2(max(0.5f, ResolutionBias));
+			MipBias = floor(MipBias);
+			MipLevel = MipBias;
+
+			const float MipVoxelSize = VoxelSize * (MipLevel + 1);
+			{
+				MipDensityScale = VoxelSize * rcp(MipVoxelSize);
+				MipDensityScale = MipDensityScale * MipDensityScale * MipDensityScale;
+			}
+			T += MipVoxelSize;
+		}
+	}
+
+	return OutHairCount;
+}
+
+float3 GetHairVoxelTangent(
+	uint3 InVolumeCoord, 
+	float InHairDensity,
+	Texture3D<uint> InVoxelTangentXTexture,
+	Texture3D<uint> InVoxelTangentYTexture,
+	Texture3D<uint> InVoxelTangentZTexture)
+{
+	float3 T8Bits;
+	T8Bits.x = InVoxelTangentXTexture.Load(int4(InVolumeCoord, 0));
+	T8Bits.y = InVoxelTangentYTexture.Load(int4(InVolumeCoord, 0));
+	T8Bits.z = InVoxelTangentZTexture.Load(int4(InVolumeCoord, 0));
+	T8Bits /= InHairDensity;
+	return From8bits(T8Bits);
+}
\ No newline at end of file
diff --git a/Engine/Shaders/Private/HairStrands/HairStrandsEnvironmentLighting.usf b/Engine/Shaders/Private/HairStrands/HairStrandsEnvironmentLighting.usf
new file mode 100644
--- /dev/null
+++ b/Engine/Shaders/Private/HairStrands/HairStrandsEnvironmentLighting.usf
@@ -0,0 +1,565 @@
+// Copyright 1998-2019 Epic Games, Inc. All Rights Reserved.
+
+/*=============================================================================
+	ReflectionEnvironmentComputeShaders - functionality to apply local cubemaps.
+=============================================================================*/
+  
+#include "../Common.ush"
+#include "HairStrandsCommon.ush"
+#include "../DeferredShadingCommon.ush"  
+#include "../BRDF.ush"
+#include "../ReflectionEnvironmentShared.ush"
+#include "../SkyLightingShared.ush"
+#include "../DistanceFieldAOShared.ush"
+#include "../ShadingModels.ush"
+#include "../LightGridCommon.ush"
+#include "../SceneTextureParameters.ush"
+
+#define RENDER_MODE_LIGHTING 0
+#define RENDER_MODE_AO 1
+
+#include "../HairBsdf.ush"
+#include "HairStrandsVisibilityCommon.ush"
+#include "HairStrandsDeepVoxelCommon.ush"
+#include "HairStrandsDeepShadowCommon.ush"
+#include "HairStrandsDeepTransmittanceCommon.ush"
+#include "HairStrandsDeepTransmittanceDualScattering.ush"
+
+// Note: 
+// * This version of environement lighting is very much in WIP state, with a lot of option to tune tweak to explore 
+//   what kind of approximation works the best for certains type of grooms and lighting condition
+// * There is a few fudges factor as well (MS roughness bias, TT density for instance, reflected vector orientations, ...)
+// * The source of lighting is only sky lighting at the moment for simplicity. This needs to be revisited to support properly reflection probes and indirect lighting in various scenario (static lighting, stationnary lighting, dynamic lighting)
+
+// This is copy/pasta from ReflectionEnvonmentComposite & co. For simplicity, while still in developement, support only sky lighting. No local reflection probe yet.
+// @hair_todo
+float3 GatherRadiance(float3 RayDirection, float Roughness)
+{
+	//float SkyAverageBrightness = 1.0f;
+	float SkyAverageBrightness = ReflectionStruct.SkyLightCubemapBrightness;
+	return GetSkyLightReflection(RayDirection, Roughness, SkyAverageBrightness);
+}
+
+
+Texture2D<uint4> HairCategorizationTexture;
+Texture2D<uint> HairVisibilityNodeOffsetAndCount;
+StructuredBuffer<FNodeData> HairVisibilityNodeData;
+
+Texture3D<float4>			HairScatteringLUTTexture;
+Texture3D<float4>			HairEnergyLUTTexture;
+SamplerState 				HairLUTSampler;
+
+float Voxel_TanConeAngle;
+float Voxel_DensityScale;
+float Voxel_DepthBiasScale;
+
+float3 Voxel_MinAABB;
+float3 Voxel_MaxAABB;
+uint   Voxel_Resolution;
+Texture3D<uint> Voxel_DensityTexture;
+
+float InternalComputeHairCountVoxel(float3 WorldPosition, float3 LightDirection, float3 InVoxelAABBMin, float3 InVoxelAABBMax, float InVoxelResolution, Texture3D<uint> InDensityTexture, float DensityScale, float CountThreshold, bool bUseConeTracing)
+{
+	if (bUseConeTracing)
+		return ComputeHairCountVoxelCone(
+			WorldPosition,
+			LightDirection,
+			InVoxelAABBMin,
+			InVoxelAABBMax,
+			InVoxelResolution,
+			InDensityTexture,
+			DensityScale,
+			CountThreshold,
+			Voxel_TanConeAngle);
+	else
+		return ComputeHairCountVoxel(
+			WorldPosition,
+			LightDirection,
+			InVoxelAABBMin,
+			InVoxelAABBMax,
+			InVoxelResolution,
+			InDensityTexture,
+			DensityScale,
+			CountThreshold);
+}
+
+//////////////////////////////////////////////////////////////////////////////////////////////////////////////////////////
+//////////////////////////////////////////////////////////////////////////////////////////////////////////////////////////
+//////////////////////////////////////////////////////////////////////////////////////////////////////////////////////////
+
+float3 ReflectionEnvironment(FGBufferData GBuffer, float2 BufferUV, float2 ScreenPosition, float4 SvPosition)
+{
+#if USE_PREEXPOSURE
+	const float PreExposure = View.PreExposure;
+#else
+	const float PreExposure = 1.f;
+#endif
+	
+	const float3 WorldPosition = mul(float4(ScreenPosition * GBuffer.Depth, GBuffer.Depth, 1), View.ScreenToWorld).xyz;
+	const float3 CameraToPixel = normalize(WorldPosition - View.WorldCameraOrigin);	
+
+	const float3 VoxelSize = (Voxel_MaxAABB - Voxel_MinAABB) / Voxel_Resolution;
+	const float  VoxelMaxSize = max(VoxelSize.x, max(VoxelSize.y, VoxelSize.z));
+	const float3 VoxelOffset = VoxelMaxSize * GetHairVoxelJitter(SvPosition.xy, View.StateFrameIndexMod8);
+
+	const float3 V = -CameraToPixel;
+	const float3 T = GBuffer.WorldNormal;
+
+	////////////////////////////////////////////////////////////////////////////////////////////////////////////////////////////////////////////
+	////////////////////////////////////////////////////////////////////////////////////////////////////////////////////////////////////////////
+	////////////////////////////////////////////////////////////////////////////////////////////////////////////////////////////////////////////
+	////////////////////////////////////////////////////////////////////////////////////////////////////////////////////////////////////////////
+	// Trace/sample transmittance in various direction to have estimate of where is occlusion & computing the integrated dual scattering value
+	float3 UnoccludedN = 0;
+	float3 AvgDualScattering = 0;
+	#if 1
+	{
+		// Poisson disk position http://developer.download.nvidia.com/whitepapers/2008/PCSS_Integration.pdf
+		float2 PoissonDisk[16] =
+		{
+			float2(-0.94201624, -0.39906216),
+			float2(0.94558609, -0.76890725),
+			float2(-0.094184101, -0.92938870),
+			float2(0.34495938, 0.29387760),
+			float2(-0.91588581, 0.45771432),
+			float2(-0.81544232, -0.87912464),
+			float2(-0.38277543, 0.27676845),
+			float2(0.97484398, 0.75648379),
+			float2(0.44323325, -0.97511554),
+			float2(0.53742981, -0.47373420),
+			float2(-0.26496911, -0.41893023),
+			float2(0.79197514, 0.19090188),
+			float2(-0.24188840, 0.99706507),
+			float2(-0.81409955, 0.91437590),
+			float2(0.19984126, 0.78641367),
+			float2(0.14383161, -0.14100790)
+		};
+
+		float3 AvgR = 0;
+		const uint SampleCount = 16;
+		for (uint i = 0; i < SampleCount; ++i)
+		{
+			const float2 Rand2 = (PoissonDisk[i] + float2(1, 1)) * 0.5f;
+			const float3 SampleL = UniformSampleSphere(saturate(Rand2)).xyz;
+
+			// Depth bias
+			// Origin is shifted voxels away towards the light + a constant bias of the size of the voxel
+			const float3 SampleWorldPosition = WorldPosition + VoxelOffset /*+ VoxelMaxSize*Voxel_DepthBiasScale*SampleL*/;
+
+			// Compute the number of hair count between light & shading point
+			const float SinLightAngle = dot(SampleL, T);
+			const float HairCount = InternalComputeHairCountVoxel(SampleWorldPosition, SampleL, Voxel_MinAABB, Voxel_MaxAABB, Voxel_Resolution, Voxel_DensityTexture, Voxel_DensityScale, 0, true);
+			const float4 SampleTransmittance = ComputeTransmittanceFromHairCount(HairCount, GBuffer, SinLightAngle, HairScatteringLUTTexture, HairLUTSampler);
+
+			const float ScalarTransmittance = max(SampleTransmittance.x, max(SampleTransmittance.y, SampleTransmittance.z));
+			AvgR += SampleL * SampleTransmittance.xyz;
+
+			const float4 ScatteringComponent = float4(1, 0, 1, 11);
+			const FHairTransmittanceData TransmittanceData = GetHairTransmittance(V, SampleL, GBuffer, SampleTransmittance, HairScatteringLUTTexture, HairLUTSampler, true, ScatteringComponent);
+
+			const float Area = 0.2;
+			const float BackLit = 1;
+			//AvgDualScattering += HairShading(GBuffer, SampleL, V, T, 1, TransmittanceData, BackLit, Area, uint2(0, 0));
+			AvgDualScattering += TransmittanceData.LocalScattering * TransmittanceData.GlobalScattering;
+		}
+		const float SphereInt = 4 * PI;
+		AvgDualScattering *= SphereInt;
+		AvgDualScattering /= SampleCount;
+
+
+		const float NormR = sqrt(dot(AvgR, AvgR));
+		UnoccludedN = AvgR / NormR;
+		
+		// For debug purpose:
+		//const float UnoccludedVarianceR = NormR * (3 - NormR2) / (1 - NormR2);
+		//const float3 ColorDir = (UnoccludedN+ float3(1, 1, 1)) * 0.5f;
+		//return ColorDir.rgb;
+		//return PreExposure * GetSkySHDiffuse(UnoccludedN) * View.SkyLightColor.rgb;
+		//return PreExposure * GatherRadiance(UnoccludedN, GBuffer.Roughness);
+		//return AvgDualScattering;
+	}
+	#endif
+	
+	////////////////////////////////////////////////////////////////////////////////////////////////////////////////////////////////////////////
+	////////////////////////////////////////////////////////////////////////////////////////////////////////////////////////////////////////////
+	////////////////////////////////////////////////////////////////////////////////////////////////////////////////////////////////////////////
+	////////////////////////////////////////////////////////////////////////////////////////////////////////////////////////////////////////////
+	// MS ==============================================================================
+	float3 ColorMS = 0;
+	#if 1
+	{
+		float3 N = UnoccludedN;
+		float3 R = 2 * dot(V, N) * N - V;
+		R = UnoccludedN;
+
+		#if 1
+		{
+			N = GBuffer.WorldNormal;
+			N = normalize(V - N * dot(V, N));
+			float3 R_b = 2 * dot(V, N) * N - V;
+			const float alpha = 0.4f;
+			R = normalize(lerp(R_b, R, alpha));
+		}
+		#endif	
+
+		const float RoughnessBias = 0.1;
+		ColorMS.rgb += PreExposure * GatherRadiance(R, GBuffer.Roughness + RoughnessBias);
+		//ColorMS.rgb += PreExposure * GetSkySHDiffuse(R) * View.SkyLightColor.rgb;
+		ColorMS.rgb *= AvgDualScattering;
+
+		// Directional albedo
+		#if 0
+		{
+			const float SinViewAngle = dot(T, R);
+			HairAverageEnergy Energy = SampleHairEnergyLUT(HairEnergyLUTTexture, HairLUTSampler, GBuffer.BaseColor, GBuffer.Roughness, SinViewAngle);
+			ColorMS.rgb *= Energy.A_R + Energy.A_TT + Energy.A_TRT;
+		}
+		#endif	
+	}
+	#endif
+
+	////////////////////////////////////////////////////////////////////////////////////////////////////////////////////////////////////////////
+	////////////////////////////////////////////////////////////////////////////////////////////////////////////////////////////////////////////
+	////////////////////////////////////////////////////////////////////////////////////////////////////////////////////////////////////////////
+	////////////////////////////////////////////////////////////////////////////////////////////////////////////////////////////////////////////
+	// R ==============================================================================
+	float3 ColorR = 0;
+	#if 1
+	{
+		float3 R = 0;
+		{
+			// This is not correct for hair rendering
+			float3 N = UnoccludedN;
+			R = 2 * dot( V, N ) * N - V;
+
+			// Use a fake normal for computing the reflected vector instead
+			#if 1
+			{
+				N = GBuffer.WorldNormal;
+				N = normalize(V - N * dot(V, N));
+				float3 R_b = 2 * dot(V, N) * N - V;
+				const float alpha = 0.4f;
+				R = normalize(lerp(R_b, R, alpha));
+			}
+			#endif	
+
+		}
+
+		// Top of regular reflection or bottom layer of clear coat.
+		const float RoughnessBias = 0.0;
+		ColorR.rgb += PreExposure * GatherRadiance(R, GBuffer.Roughness + RoughnessBias);
+		//ColorR.rgb += PreExposure * GetSkySHDiffuse(R) * View.SkyLightColor.rgb * PI;
+
+		const float3 L = R;
+		const float SinLightAngle = dot(L, T);
+		const float SinViewAngle = dot(V, T);
+
+		// Compute the transmittance
+		float4 Transmittance = float4(1,1,1,0);
+		{
+			// Depth bias
+			// Origin is shifted voxels away towards the light + a constant bias of the size of the voxel
+			const float3 DepthBias = VoxelOffset /*+ VoxelMaxSize*Voxel_DepthBiasScale*L*/;
+			const float3 SampleWorldPosition = WorldPosition + DepthBias;
+
+			// Compute the number of hair count between light & shading point
+			const float HairCount = InternalComputeHairCountVoxel(SampleWorldPosition, L, Voxel_MinAABB, Voxel_MaxAABB, Voxel_Resolution, Voxel_DensityTexture, Voxel_DensityScale, 0, true);
+			Transmittance = ComputeTransmittanceFromHairCount(HairCount, GBuffer, SinLightAngle, HairScatteringLUTTexture, HairLUTSampler);
+		}
+
+		FHairTransmittanceData TransmittanceData = InitHairTransmittanceData();	
+
+		// Directional albedo
+		#if 0
+		{
+			HairAverageEnergy Energy = SampleHairEnergyLUT(HairEnergyLUTTexture, HairLUTSampler, GBuffer.BaseColor, GBuffer.Roughness, SinViewAngle);
+			ColorR.rgb *= (Energy.A_TRT + Energy.A_R) * Transmittance.xyz;
+		}
+		#endif
+
+		// BSDF
+		#if 1
+		{			
+			TransmittanceData.MultipleScattering = 1;
+			TransmittanceData.GlobalScattering = Transmittance.xyz;
+			TransmittanceData.ScatteringComponent = uint3(1, 0, 1); // Only R & TRT
+			TransmittanceData.Enable = true;
+
+			const float Area = 0.2;
+			const float BackLit = 0;
+			ColorR.rgb *= HairShading(GBuffer, L, V, T, 1, TransmittanceData, BackLit, Area, uint2(0, 0));
+		}
+		#endif
+	}
+	#endif // R TRT
+
+	////////////////////////////////////////////////////////////////////////////////////////////////////////////////////////////////////////////
+	////////////////////////////////////////////////////////////////////////////////////////////////////////////////////////////////////////////
+	////////////////////////////////////////////////////////////////////////////////////////////////////////////////////////////////////////////
+	////////////////////////////////////////////////////////////////////////////////////////////////////////////////////////////////////////////
+	// TT ==============================================================================
+	float3 ColorTT = 0;
+	#if 1
+	{
+		const float3 L = -V;
+		const float SinLightAngle = dot(L, T);
+		const float SinViewAngle = dot(V, T);
+
+		ColorTT.rgb += PreExposure * GatherRadiance(L, GBuffer.Roughness);
+		//ColorTT.rgb += PreExposure * GetSkySHDiffuse(L) * View.SkyLightColor.rgb;
+
+		// Compute the transmittance
+		float4 Transmittance = float4(1,1,1,0);
+		{
+			// Depth bias
+			// Origin is shifted voxels away towards the light + a constant bias of the size of the voxel
+			const float3 DepthBias = VoxelOffset /*+ VoxelMaxSize*Voxel_DepthBiasScale*L*/;
+			const float3 SampleWorldPosition = WorldPosition + DepthBias;
+
+			// Compute the number of hair count between light & shading point
+			const float DensityScaleFactor = 2;
+			const float HairCount = InternalComputeHairCountVoxel(SampleWorldPosition, L, Voxel_MinAABB, Voxel_MaxAABB, Voxel_Resolution, Voxel_DensityTexture, Voxel_DensityScale * DensityScaleFactor, 0, false);
+			Transmittance = ComputeTransmittanceFromHairCount(HairCount, GBuffer, SinLightAngle, HairScatteringLUTTexture, HairLUTSampler);
+		}
+
+		// BSDF with dual-scattering
+		#if 0
+		{
+			const FHairTransmittanceData TransmittanceData = GetHairTransmittance(V, L, GBuffer, Transmittance, HairScatteringLUTTexture, HairLUTSampler, 1, float4(0, 1, 0, 11));
+			const float Area = 0.2;
+			ColorTT.rgb *= HairShading(GBuffer, L, V, T, 1, TransmittanceData, 1, Area, uint2(0, 0));
+		}
+		#endif
+
+		// Average weighting
+		#if 1
+		{
+			HairAverageEnergy Energy = SampleHairEnergyLUT(HairEnergyLUTTexture, HairLUTSampler, GBuffer.BaseColor, GBuffer.Roughness, SinViewAngle);
+			ColorTT.rgb *= Energy.A_TT * Transmittance.xyz;
+		}
+		#endif	
+
+		// BSDF
+		#if 0
+		{
+			FHairTransmittanceData TransmittanceData = InitHairTransmittanceData();
+			TransmittanceData.MultipleScattering = 1;
+			TransmittanceData.GlobalScattering = Transmittance.xyz;
+			TransmittanceData.LocalScattering = 0;
+			TransmittanceData.ScatteringComponent = uint3(0, 1, 0);
+			TransmittanceData.Enable = true;
+
+			const float Area = 0.2;
+			ColorTT.rgb *= HairShading(GBuffer, L, V, T, 1, TransmittanceData, 0, Area, uint2(0, 0));
+		}
+		#endif	
+	}
+	#endif // TT
+
+	const float3 Color = ColorR + ColorTT + ColorMS;
+
+	// Transform NaNs to black, transform negative colors to black.
+	return -min(-Color.rgb, 0.0);
+}
+
+//////////////////////////////////////////////////////////////////////////////////////////////////////////////////////////
+// AO
+//////////////////////////////////////////////////////////////////////////////////////////////////////////////////////////
+
+float3 TraceAO(float3 WorldPosition, float3 WorldNormal, float4 SvPosition)
+{
+	const float3 VoxelSize = (Voxel_MaxAABB - Voxel_MinAABB) / Voxel_Resolution;
+	const float  VoxelMaxSize = max(VoxelSize.x, max(VoxelSize.y, VoxelSize.z));
+	const float3 VoxelOffset = VoxelMaxSize * GetHairVoxelJitter(SvPosition.xy, View.StateFrameIndexMod8);
+	
+	float3 UnoccludedN = 0;
+	uint TotalHairCount = 0;
+	// Poisson disk position http://developer.download.nvidia.com/whitepapers/2008/PCSS_Integration.pdf
+	float2 PoissonDisk[16] =
+	{
+		float2(-0.94201624, -0.39906216),
+		float2(0.94558609, -0.76890725),
+		float2(-0.094184101, -0.92938870),
+		float2(0.34495938, 0.29387760),
+		float2(-0.91588581, 0.45771432),
+		float2(-0.81544232, -0.87912464),
+		float2(-0.38277543, 0.27676845),
+		float2(0.97484398, 0.75648379),
+		float2(0.44323325, -0.97511554),
+		float2(0.53742981, -0.47373420),
+		float2(-0.26496911, -0.41893023),
+		float2(0.79197514, 0.19090188),
+		float2(-0.24188840, 0.99706507),
+		float2(-0.81409955, 0.91437590),
+		float2(0.19984126, 0.78641367),
+		float2(0.14383161, -0.14100790)
+	};
+	const uint SampleCount = 16;
+
+	float3 BentNormal = 0;
+	float Occlusion = 0;
+	for (uint i = 0; i < SampleCount; ++i)
+	{
+		float2 Rand2 = (PoissonDisk[i] + float2(1, 1)) * 0.5f;
+		Rand2 = frac(Rand2 + VoxelOffset.xy);
+		const float3 LocalL  = CosineSampleHemisphere(saturate(Rand2)).xyz;
+		const float3 SampleL = TangentToWorld(LocalL.xyz, WorldNormal);
+
+		// Depth bias
+		// Origin is shifted voxels away towards the light + a constant bias of the size of the voxel
+		const float3 SampleWorldPosition = WorldPosition + VoxelOffset /*+ VoxelMaxSize*Voxel_DepthBiasScale*SampleL*/;;
+
+		// Compute the number of hair count between light & shading point
+		const float HairCount = InternalComputeHairCountVoxel(SampleWorldPosition, SampleL, Voxel_MinAABB, Voxel_MaxAABB, Voxel_Resolution, Voxel_DensityTexture, Voxel_DensityScale, 0, true);
+		BentNormal += SampleL * saturate(1 - HairCount);
+		Occlusion += saturate(HairCount);
+	}
+	Occlusion /= SampleCount;
+	BentNormal /= SampleCount;
+
+	return saturate(Occlusion);
+}
+
+//////////////////////////////////////////////////////////////////////////////////////////////////////////////////////////
+//////////////////////////////////////////////////////////////////////////////////////////////////////////////////////////
+//////////////////////////////////////////////////////////////////////////////////////////////////////////////////////////
+float3 TotWorldPosition(float2 UV, float Depth)
+{
+	float2 ScreenPosition = (UV - View.ScreenPositionScaleBias.wz) / View.ScreenPositionScaleBias.xy;
+	float4 HomogeneousWorldPosition = mul(float4(ScreenPosition * Depth, Depth, 1), View.ScreenToWorld);
+	float3 WorldPosition = HomogeneousWorldPosition.xyz / HomogeneousWorldPosition.w;
+
+	return WorldPosition;
+}
+
+void MainPS(
+	in float4 SvPosition : SV_Position,
+	out float4 OutColor : SV_Target0
+#if PERMUTATION_RENDER_MODE == RENDER_MODE_LIGHTING
+	, out float4 OutColorSubPixel : SV_Target1
+#endif
+)
+{
+	OutColor = 0.0f;
+	const float2 BufferUV = SvPositionToBufferUV(SvPosition);
+	const float2 ScreenPosition = SvPositionToScreenPosition(SvPosition).xy;
+	const uint3 PixelCoord = uint3(floor(SvPosition.xy), 0);
+
+	FNodeDesc NodeDesc = DecodeNodeDesc(HairVisibilityNodeOffsetAndCount.Load(PixelCoord));
+
+#if PERMUTATION_RENDER_MODE == RENDER_MODE_LIGHTING
+	OutColorSubPixel = 0.0f;
+
+	if (NodeDesc.Count > 0)
+	{
+		// Note:
+		// * Currently: compute take the closest sample
+		// * Change this to per sample evaluation once we have better hair transmittance evaluation for sky
+		float Coverage = 1;
+		float ClosestDepth = 0;
+		FGBufferData HairGBuffer = (FGBufferData)0;
+		Coverage = 0;
+
+		LOOP
+		for (uint SampleIt = 0; SampleIt < NodeDesc.Count; SampleIt++)
+		{
+			const uint LocalOffset = NodeDesc.Offset + SampleIt;
+			const FNodeData NodeData = HairVisibilityNodeData[LocalOffset];
+			float			SampleDepth;
+			FVisibilityData SampleVisData;
+			float3			SampleBaseColor;
+			float			SampleRoughness;
+			float			SampleSpecular;
+			DecodeNodeData(NodeData, SampleDepth, SampleVisData, SampleBaseColor, SampleRoughness, SampleSpecular);
+
+			// Inject material data from the visibility/mini-gbuffer for hair
+			const float SceneDepth = ConvertFromDeviceZ(SampleDepth);
+			if (SampleDepth > ClosestDepth)
+			{
+				HairGBuffer.WorldNormal = SampleVisData.Tangent;
+				HairGBuffer.BaseColor = SampleBaseColor;
+				HairGBuffer.Roughness = SampleRoughness;
+				HairGBuffer.ShadingModelID = SHADINGMODELID_HAIR;
+				HairGBuffer.DiffuseColor = 0;
+				HairGBuffer.SpecularColor = 0;
+				HairGBuffer.Specular = SampleSpecular;
+				HairGBuffer.Metallic = 0;
+				HairGBuffer.Depth = SceneDepth;
+				HairGBuffer.GBufferAO = 1;
+				HairGBuffer.CustomData = 0;
+				HairGBuffer.IndirectIrradiance = 1000000;
+			}
+			ClosestDepth = max(ClosestDepth, SampleDepth);  // Inverse-Z
+
+			const float LocalCoverage = From8bitCoverage(SampleVisData.Coverage8bit);
+			Coverage += LocalCoverage;
+		}
+		
+		// Lit fully covered or partially covered pixel
+		// * Fully covered pixel output their result into the main color target
+		// * Partially covered pixel output their result into the sub pixel color target. the main color target will store the underlying pixel color 
+		{
+			const float3 SkyLighting = ReflectionEnvironment(HairGBuffer, BufferUV, ScreenPosition, SvPosition);
+
+			// Scene color alpha is used for ScreenSpaceSubsurfaceScattering (if that is not needed it can be disabled with SUBSURFACE_CHANNEL_MODE)
+			FLightAccumulator LightAccumulator = (FLightAccumulator)0;
+			LightAccumulator_Add(LightAccumulator, SkyLighting, SkyLighting, 1.0f, false);
+			const float4 LocalOutColor = LightAccumulator_GetResult(LightAccumulator);
+			if (Coverage < 1)
+			{
+				OutColorSubPixel = LocalOutColor * Coverage;
+				OutColorSubPixel.a = Coverage;
+			}
+			else
+			{
+				OutColor += LocalOutColor;
+				return;
+			}
+		}
+	}
+#endif
+
+
+#if PERMUTATION_RENDER_MODE == RENDER_MODE_AO
+	const FGBufferData GBuffer = GetGBufferDataFromSceneTextures(BufferUV);
+	const float3 WorldPosition = TotWorldPosition(BufferUV, GBuffer.Depth);
+	const float3 WorldNormal = GBuffer.WorldNormal;
+
+	// Trace AO for all pixel which are partially covered or not covered at all.
+	bool bIsValid = false;
+	if (NodeDesc.Count < 255)
+	{
+		float ClosestDepth = 0;
+		uint Coverage8Bits = 0;
+		LOOP
+		for (uint SampleIt = 0; SampleIt < NodeDesc.Count; SampleIt++)
+		{
+			const uint LocalOffset = NodeDesc.Offset + SampleIt;
+			const FNodeData NodeData = HairVisibilityNodeData[LocalOffset];
+			float			SampleDepth;
+			FVisibilityData SampleVisData;
+			float3			SampleBaseColor;
+			float			SampleRoughness;
+			float			SampleSpecular;
+			DecodeNodeData(NodeData, SampleDepth, SampleVisData, SampleBaseColor, SampleRoughness, SampleSpecular);
+			const float LocalCoverage = From8bitCoverage(SampleVisData.Coverage8bit);
+			Coverage8Bits += SampleVisData.Coverage8bit;
+		}
+		bIsValid = Coverage8Bits <= 255;
+	}
+	else
+	{
+		bIsValid = GBuffer.ShadingModelID != SHADINGMODELID_UNLIT;
+	}
+
+	if (bIsValid)
+	{
+		const float AO = TraceAO(WorldPosition, WorldNormal, SvPosition);
+		OutColor = float4(0, 0, 0, 1-AO);
+	}
+	else
+	{
+		discard;
+	}
+#endif
+}
diff --git a/Engine/Shaders/Private/HairStrands/HairStrandsGuideDeform.usf b/Engine/Shaders/Private/HairStrands/HairStrandsGuideDeform.usf
new file mode 100644
--- /dev/null
+++ b/Engine/Shaders/Private/HairStrands/HairStrandsGuideDeform.usf
@@ -0,0 +1,67 @@
+// Copyright 1998-2019 Epic Games, Inc. All Rights Reserved.
+
+#include "../Common.ush"
+
+float3 HairWorldOffset;
+uint VertexCount;
+uint IterationCount;
+
+Buffer<uint4> SimRestPosePositionBuffer;
+Buffer<uint>  SimRootIndexBuffer;
+RWBuffer<uint4> OutSimDeformedPositionBuffer;
+
+#if PERMUTATION_GROUP_SIZE == 0
+	#define GROUP_SIZE 64
+#elif PERMUTATION_GROUP_SIZE == 1
+	#define GROUP_SIZE 32
+#else
+	#error Unknown group size	
+#endif
+
+float3 Displacement(uint VertexIndex, float3 Pos, float Seed)
+{
+// By pass
+#if PERMUTATION_DEFORMATION == 0
+	return 0;
+#endif
+
+// Wave
+#if PERMUTATION_DEFORMATION == 1
+	const float Scale = 100.f;
+	Pos *= Scale;
+	//	return float3(0, cos(Pos.y + Seed / 16.f), cos(Pos.z + Seed / 32.f)) / Scale;
+	return float3(0, cos(Pos.y + Seed / 16.f), cos(Pos.z + Seed / 32.f));
+#endif
+
+// Align with some normal
+#if PERMUTATION_DEFORMATION == 2
+	const uint RootIndex = SimRootIndexBuffer[VertexIndex];
+	const float3 V0 = f16tof32(SimRestPosePositionBuffer[RootIndex]);
+	const float3 V1 = f16tof32(SimRestPosePositionBuffer[RootIndex+1]);
+	const float3 T = normalize(V1 - V0);
+	const float D = length(Pos-V0);
+
+	const float3 Offset = V0 + D*T - Pos;
+
+	float Time = 0.5f * cos(Seed / 64.f) + 0.5f; // [0..1]
+	return lerp(0, Offset, Time);
+#endif
+}
+
+[numthreads(GROUP_SIZE, 1, 1)]
+void MainCS(uint2 DispatchThreadId : SV_DispatchThreadID)
+{
+	const uint VertexIndex = DispatchThreadId.x;
+	if (VertexIndex < VertexCount)
+	{
+		const uint4  Data 			= SimRestPosePositionBuffer[VertexIndex];
+		const float3 ControlPoint	= f16tof32(Data.xyz);
+
+		float3 PreviousOffset = 0;
+		{
+			const float3 Offset = Displacement(VertexIndex, ControlPoint, IterationCount);
+			OutSimDeformedPositionBuffer[VertexIndex] = uint4(f32tof16(ControlPoint + Offset), Data.w);
+			PreviousOffset = Offset;
+		}
+	}
+}
\ No newline at end of file
diff --git a/Engine/Shaders/Private/HairStrands/HairStrandsInterpolation.usf b/Engine/Shaders/Private/HairStrands/HairStrandsInterpolation.usf
new file mode 100644
--- /dev/null
+++ b/Engine/Shaders/Private/HairStrands/HairStrandsInterpolation.usf
@@ -0,0 +1,70 @@
+// Copyright 1998-2019 Epic Games, Inc. All Rights Reserved.
+
+#include "../Common.ush"
+
+uint VertexCount;
+uint DispatchCountX;
+
+Buffer<uint4>  RenderRestPosePositionBuffer;
+Buffer<uint4>  SimRestPosePositionBuffer;
+Buffer<uint4>  DeformedSimPositionBuffer;
+Buffer<uint4>  Interpolation0Buffer;
+Buffer<uint4>  Interpolation1Buffer;
+
+RWBuffer<uint4>  OutRenderDeformedPositionBuffer;
+
+#if PERMUTATION_DEBUG == 1
+Buffer<float4>	 SimAttributeBuffer;
+RWBuffer<float4> OutRenderAttributeBuffer;
+#endif
+
+#if PERMUTATION_GROUP_SIZE == 0
+	#define GROUP_SIZE 64
+#elif PERMUTATION_GROUP_SIZE == 1
+	#define GROUP_SIZE 32
+#else
+	#error Unknown group size	
+#endif
+
+[numthreads(GROUP_SIZE, 1, 1)]
+void MainCS(uint2 DispatchThreadId : SV_DispatchThreadID)
+{
+	const uint VertexIndex = DispatchThreadId.x + DispatchThreadId.y * DispatchCountX * GROUP_SIZE;
+	if (VertexIndex < VertexCount)
+	{
+		const uint4 InterpolationData0	= Interpolation0Buffer[VertexIndex];
+		const uint4 InterpolationData1	= Interpolation1Buffer[VertexIndex];		
+
+		uint3 GuideVertexIndices = 0;
+		GuideVertexIndices.x = InterpolationData0.x | (InterpolationData1.x << 16);
+		GuideVertexIndices.y = InterpolationData0.y | (InterpolationData1.y << 16);
+		GuideVertexIndices.z = InterpolationData0.z | (InterpolationData1.z << 16);
+
+		float3 GuideVertexWeights = float3(saturate((InterpolationData0.w & 0xFF) / 255.f), saturate(((InterpolationData0.w >> 8) & 0xFF) / 255.f), 0);
+		GuideVertexWeights.z = saturate(1.0f - (GuideVertexWeights.x + GuideVertexWeights.y));
+
+		float3 CurrOffset = 0;
+		float3 PrevOffset = 0;
+		const float3 ControlPoint = f16tof32(RenderRestPosePositionBuffer[VertexIndex].xyz);
+		[unroll]
+		for (uint K = 0; K < 3; ++K)
+		{
+			const uint GuideIndex = GuideVertexIndices[K];
+			const float3 RestGuidePoint	= f16tof32(SimRestPosePositionBuffer[GuideIndex].xyz);
+
+			const float3 DeformedGuidePoint	= f16tof32(DeformedSimPositionBuffer[GuideIndex].xyz);
+			CurrOffset += (DeformedGuidePoint - RestGuidePoint) * GuideVertexWeights[K];
+
+		}
+
+		OutRenderDeformedPositionBuffer[VertexIndex] = float4(f32tof16(ControlPoint + CurrOffset), RenderRestPosePositionBuffer[VertexIndex].w);
+
+		#if PERMUTATION_DEBUG == 1
+		{
+			const uint GuideIndex = GuideVertexIndices[0]; // Take the closest guide
+			const float4 GuideAttribute = SimAttributeBuffer[GuideIndex];
+			OutRenderAttributeBuffer[VertexIndex] = GuideAttribute;
+		}
+		#endif
+	}
+}
\ No newline at end of file
diff --git a/Engine/Shaders/Private/HairStrands/HairStrandsRaytracingGeometry.usf b/Engine/Shaders/Private/HairStrands/HairStrandsRaytracingGeometry.usf
new file mode 100644
--- /dev/null
+++ b/Engine/Shaders/Private/HairStrands/HairStrandsRaytracingGeometry.usf
@@ -0,0 +1,124 @@
+// Copyright 1998-2019 Epic Games, Inc. All Rights Reserved.
+
+#include "../Common.ush"
+
+uint VertexCount;
+uint DispatchCountX;
+
+float3 StrandHairWorldOffset;
+float StrandHairRadius;
+
+Buffer<uint4> 	 PositionBuffer;
+RWBuffer<float4> OutputPositionBuffer;
+
+#if PERMUTATION_GROUP_SIZE == 0
+	#define GROUP_SIZE 64
+#elif PERMUTATION_GROUP_SIZE == 1
+	#define GROUP_SIZE 32
+#else
+	#error Unknown group size	
+#endif
+
+
+struct FControlPoint
+{
+	float3 Position;
+	float PointType;
+	float WorldRadius;
+};
+
+FControlPoint GetControlPoint(uint Index)
+{
+	FControlPoint Out = (FControlPoint)0;
+
+	if (Index < VertexCount)
+	{
+		const uint4 PackedData = PositionBuffer[Index];
+		const uint PackedAlpha = asuint(PackedData.a);
+
+		Out.Position 	= f16tof32(PackedData.xyz) + StrandHairWorldOffset;
+		Out.PointType 	= PackedAlpha & 0x3;
+		Out.WorldRadius	= ((PackedAlpha & 0x00FC) >> 2) / 63.f * StrandHairRadius;	
+	}
+
+	return Out;
+}
+
+
+float4 PackNormal(float3 V, uint A)
+{
+	return float4(clamp(V,float(-1).xxx,float(1).xxx), 1.0f);
+}
+
+[numthreads(GROUP_SIZE, 1, 1)]
+void MainCS(uint2 DispatchThreadId : SV_DispatchThreadID)
+{
+	const int IndexCurr = DispatchThreadId.x + DispatchThreadId.y * DispatchCountX * GROUP_SIZE;
+	const int MaxVertexCount = int(VertexCount);
+	const uint BaseOutIndex = IndexCurr*12;
+
+	FControlPoint P0 = GetControlPoint(IndexCurr);
+
+	// Point at the end of a strand is filled with invalid position. The positions/triangle are 
+	// ignored by DXR during BVH building
+	if (P0.PointType == 2)
+	{
+		const float NaN = sqrt(-1);
+		const float4 InvalidPoint = float4(NaN, NaN, NaN, NaN);
+
+		OutputPositionBuffer[BaseOutIndex  ]  = InvalidPoint;
+		OutputPositionBuffer[BaseOutIndex+1]  = InvalidPoint;
+		OutputPositionBuffer[BaseOutIndex+2]  = InvalidPoint;
+		OutputPositionBuffer[BaseOutIndex+3]  = InvalidPoint;
+		OutputPositionBuffer[BaseOutIndex+4]  = InvalidPoint;
+		OutputPositionBuffer[BaseOutIndex+5]  = InvalidPoint;
+		OutputPositionBuffer[BaseOutIndex+6]  = InvalidPoint;
+		OutputPositionBuffer[BaseOutIndex+7]  = InvalidPoint;
+		OutputPositionBuffer[BaseOutIndex+8]  = InvalidPoint;
+		OutputPositionBuffer[BaseOutIndex+9]  = InvalidPoint;
+		OutputPositionBuffer[BaseOutIndex+10] = InvalidPoint;
+		OutputPositionBuffer[BaseOutIndex+11] = InvalidPoint;
+		return;
+	}
+
+	FControlPoint P1 = GetControlPoint(IndexCurr+1);
+
+	const float3 Tangent = normalize(P1.Position - P0.Position);
+	float3 B = float3(0,0,-1);
+	if (abs(dot(Tangent,B)) > 0.95f)
+	{
+		B = float3(0,1,0);
+	}	
+	float3 N = normalize(cross(Tangent, B));
+	B = normalize(cross(N, Tangent));
+
+	{
+		const float3 P00 = P0.Position - N * P0.WorldRadius;
+		const float3 P01 = P0.Position + N * P0.WorldRadius;
+		const float3 P10 = P1.Position - N * P1.WorldRadius;
+		const float3 P11 = P1.Position + N * P1.WorldRadius;
+		
+		OutputPositionBuffer[BaseOutIndex  ] = float4(P00, 1);
+		OutputPositionBuffer[BaseOutIndex+1] = float4(P01, 1);
+		OutputPositionBuffer[BaseOutIndex+2] = float4(P11, 1);
+
+		OutputPositionBuffer[BaseOutIndex+3] = float4(P00, 1);
+		OutputPositionBuffer[BaseOutIndex+4] = float4(P11, 1);
+		OutputPositionBuffer[BaseOutIndex+5] = float4(P10, 1);
+	}
+	
+	{
+		const float3 P00 = P0.Position - B * P0.WorldRadius;
+		const float3 P01 = P0.Position + B * P0.WorldRadius;
+		const float3 P10 = P1.Position - B * P1.WorldRadius;
+		const float3 P11 = P1.Position + B * P1.WorldRadius;
+
+		OutputPositionBuffer[BaseOutIndex+6] = float4(P00, 1);
+		OutputPositionBuffer[BaseOutIndex+7] = float4(P01, 1);
+		OutputPositionBuffer[BaseOutIndex+8] = float4(P11, 1);
+
+		OutputPositionBuffer[BaseOutIndex+9]  = float4(P00, 1);
+		OutputPositionBuffer[BaseOutIndex+10] = float4(P11, 1);
+		OutputPositionBuffer[BaseOutIndex+11] = float4(P10, 1);
+	}
+}
\ No newline at end of file
diff --git a/Engine/Shaders/Private/HairStrands/HairStrandsTangent.usf b/Engine/Shaders/Private/HairStrands/HairStrandsTangent.usf
new file mode 100644
--- /dev/null
+++ b/Engine/Shaders/Private/HairStrands/HairStrandsTangent.usf
@@ -0,0 +1,81 @@
+// Copyright 1998-2019 Epic Games, Inc. All Rights Reserved.
+
+#include "../Common.ush"
+
+uint VertexCount;
+uint DispatchCountX;
+
+Buffer<uint4>  	 PositionBuffer;
+RWBuffer<float4> OutputTangentBuffer;
+
+#if PERMUTATION_GROUP_SIZE == 0
+	#define GROUP_SIZE 64
+#elif PERMUTATION_GROUP_SIZE == 1
+	#define GROUP_SIZE 32
+#else
+	#error Unknown group size	
+#endif
+
+float4 PackNormal(float3 V, uint A)
+{
+	return float4(clamp(V,float(-1).xxx,float(1).xxx), 1.0f);
+}
+
+[numthreads(GROUP_SIZE, 1, 1)]
+void MainCS(uint2 DispatchThreadId : SV_DispatchThreadID)
+{
+	const int IndexCurr = DispatchThreadId.x + DispatchThreadId.y * DispatchCountX * GROUP_SIZE;
+	const int MaxVertexCount = int(VertexCount);
+	if (IndexCurr < MaxVertexCount)
+	{
+		const int IndexPrev = clamp(IndexCurr-1, 0, MaxVertexCount-1);
+		const int IndexNext = clamp(IndexCurr+1, 0, MaxVertexCount-1);
+
+		const uint4 DataCurr = PositionBuffer[IndexCurr];
+		const uint4 DataPrev = PositionBuffer[IndexPrev];
+		const uint4 DataNext = PositionBuffer[IndexNext];
+
+		const float3 PCurr = f16tof32(DataCurr.xyz);
+		float3 PPrev = f16tof32(DataPrev.xyz);
+		float3 PNext = f16tof32(DataNext.xyz);
+
+		// For more details about the vertex type, see StrandHairVertexFactor.ush
+		const uint TypeCurr = asuint(DataCurr.a) & 0x3;
+		const uint TypePrev = asuint(DataPrev.a) & 0x3;
+		const uint TypeNext = asuint(DataNext.a) & 0x3;
+
+		float3 T0 = 0;
+		const bool bIsPrevValid = TypePrev == 0 || TypePrev == 1;
+		if (bIsPrevValid)
+		{
+			T0 = normalize(PCurr - PPrev);
+		}
+
+		float3 T1 = 0;
+		const bool bIsNextValid = TypeNext == 0 || TypeNext == 2;
+		if (bIsNextValid)
+		{
+			T1 = normalize(PNext - PCurr);
+		}
+
+		const float3 T = normalize(T0 + T1);
+
+		const float Norm2 = dot(T, float3(0,0,-1));
+		float3 Right;
+		if (Norm2 > 0e-3f)
+		{
+			Right = normalize(cross(T, float3(0,0,-1)));
+		}
+		else
+		{
+			Right = normalize(cross(T, float3(0,1,0)));
+		}
+
+		// Build ortho basis
+		float3 TangentX = normalize(cross(Right, T));
+		float3 TangentZ = T;
+
+		OutputTangentBuffer[IndexCurr*2]   = PackNormal(TangentX, 0);
+		OutputTangentBuffer[IndexCurr*2+1] = PackNormal(TangentZ, 127);
+	}
+}
\ No newline at end of file
