Copying //UE4/Dev-Physics to Dev-Main (//UE4/Dev-Main) [at] 9098278

diff --git a/Engine/Shaders/Private/HairStrands/HairStrandsVertexFactory.ush b/Engine/Shaders/Private/HairStrands/HairStrandsVertexFactory.ush
new file mode 100644
--- /dev/null
+++ b/Engine/Shaders/Private/HairStrands/HairStrandsVertexFactory.ush
@@ -0,0 +1,483 @@
+// Copyright 1998-2019 Epic Games, Inc. All Rights Reserved.
+
+/*=============================================================================
+	StrandHairFactory.usf
+=============================================================================*/
+#include "../VertexFactoryCommon.ush"
+#include "../LocalVertexFactoryCommon.ush"
+
+#include "/Engine/Generated/UniformBuffers/PrecomputedLightingBuffer.ush"
+
+#if MANUAL_VERTEX_FETCH
+	Buffer<float4> VertexFetch_InstanceOriginBuffer;
+	Buffer<float4> VertexFetch_InstanceTransformBuffer;
+#endif
+
+////////////////////////////////////////////////////////////////////////////////
+// VF_STRAND_HAIR
+// Define used in certain shader like DeepShadow for running specific code (coverage computation for instance) 
+// when vertices are produced by the strand hair vertex factory 
+// This is set by compilation enviromenent of the vertex factory
+///////////////////////////////////////////////////////////////////////////////
+
+float3 HairStrandsVF_WorldOffset; // Vertex position (3x16bits) are relative to this offset for improving precision while keeping vertex footprint low
+float  HairStrandsVF_Density;
+float  HairStrandsVF_Radius;
+float  HairStrandsVF_Length;
+float  HairStrandsVF_RadiusAtDepth1_Primary;
+float  HairStrandsVF_RadiusAtDepth1_Velocity;
+
+Buffer<uint4>  HairStrandsVF_PositionBuffer;
+Buffer<uint4>  HairStrandsVF_PreviousPositionBuffer;
+Buffer<float4> HairStrandsVF_AttributeBuffer;
+Buffer<float4> HairStrandsVF_TangentBuffer;
+
+// ## A quad is made of the following topology and indexing:
+//
+//  0__2    4 
+//  | /    /|
+//  |/    /_|
+//  1    5  3  
+//
+//  ## Control point identification
+//  Here is an example of two consecutive strands. The control point:
+//   * .a=1 starting
+//   * .a=0 within
+//   * .a=2 ending
+//
+//   _O       .a == 1  Strand 0
+//  | / /|
+//  |/ /_|
+//   _O       .a == 0
+//  | / /|
+//  |/ /_|
+//   _O       .a == 0
+//  | / /|
+//  |/ /_|
+//    O       .a == 2
+//     
+//    
+//    O       .a == 1   Strand 1
+//  | / /|
+//  |/ /_|
+//    O       .a == 2
+//
+//   ...
+
+/**
+ * Per-vertex inputs from bound vertex buffers
+ */
+#define FPositionOnlyVertexFactoryInput FVertexFactoryInput
+#define FPositionAndNormalOnlyVertexFactoryInput FVertexFactoryInput
+
+struct FVertexFactoryInput
+{
+#if VF_USE_PRIMITIVE_SCENE_DATA
+	uint PrimitiveId : ATTRIBUTE13;
+#endif	
+	uint VertexId : SV_VertexID;
+};
+
+#if RAYHITGROUPSHADER
+FVertexFactoryInput LoadVertexFactoryInput(uint TriangleIndex, int VertexIndex)
+{
+	// #hair_todo
+	FTriangleBaseAttributes Triangle = LoadTriangleBaseAttributes(TriangleIndex);
+
+	FVertexFactoryInput Input;
+	Input.Position = float4(Triangle.LocalPositions[VertexIndex], 1.0f);
+	Input.VertexId = Triangle.Indices[VertexIndex];
+#if VF_USE_PRIMITIVE_SCENE_DATA
+	Input.PrimitiveId = GetInstanceUserData();
+#endif
+	return Input;
+}
+#endif
+
+/** 
+ * Caches intermediates that would otherwise have to be computed multiple times.  Avoids relying on the compiler to optimize out redundant operations.
+ */
+struct FVertexFactoryIntermediates
+{
+	half3x3 TangentToLocal;
+	half3x3 TangentToWorld;
+	half TangentToWorldSign;
+	uint PrimitiveId;
+	uint   HairPrimitiveId;
+	float2 HairDimensions;
+	float2 HairUV;
+	float2 HairRootUV;
+	float  HairSeed;
+	float  HairDensity;
+};
+
+uint GetPrimitiveId(FVertexFactoryInput Input)
+{
+#if VF_USE_PRIMITIVE_SCENE_DATA
+	return Input.PrimitiveId;
+#else
+	return 0;
+#endif
+}
+
+uint GetHairPrimitiveId(FVertexFactoryInput Input)
+{
+	return Input.VertexId / 6;
+}
+
+struct FControlPoint
+{
+	float3 Position;
+	float  WorldRadius;
+	float  UCoord;
+	float  WorldLength;
+	float  Seed;
+	float2 RootUV;
+};
+
+FControlPoint GetVertexPosition(FVertexFactoryInput Input)
+{
+	const uint QuadIndex	= Input.VertexId % 6;
+	const uint BaseIndex	= Input.VertexId / 6;
+	const uint OffsetIndex	= QuadIndex == 0 || QuadIndex == 2 || QuadIndex == 4 ? 0 : 1;
+	const uint4 PackedData = HairStrandsVF_PositionBuffer[BaseIndex + OffsetIndex];
+	
+	FControlPoint Output = (FControlPoint)0;
+
+	const float3 ControlPoint	= f16tof32(PackedData.xyz) + HairStrandsVF_WorldOffset;
+	const uint PackedAlpha		= asuint(PackedData.a);
+	const uint ControlPointType = PackedAlpha & 0x3;
+	Output.WorldRadius			= ((PackedAlpha & 0x00FC) >> 2) / 63.f * HairStrandsVF_Radius;
+	Output.WorldLength			= ((PackedAlpha & 0xFF00) >> 8) / 255.f * HairStrandsVF_Length;
+	
+	// Create a degenerated quad the end of each strand to cut between each strands
+	// #hair_todo: This is not efficient for short strand like fur for instance. Need to revisit that at some point
+	const bool bIsInvalidQuad = (ControlPointType == 1 && OffsetIndex == 1) || (ControlPointType == 2 && OffsetIndex == 0);
+	const float4 Attributes		= HairStrandsVF_AttributeBuffer[BaseIndex + OffsetIndex];
+	Output.Position				= bIsInvalidQuad ? float3(1.#INF, 1.#INF, 1.#INF) : ControlPoint;
+	Output.RootUV				= Attributes.xy;
+	Output.UCoord				= Attributes.z;
+	Output.Seed					= Attributes.w;
+
+	return Output;
+}
+
+float3 GetVertexPreviousPosition(FVertexFactoryInput Input)
+{
+	const uint QuadIndex = Input.VertexId % 6;
+	const uint BaseIndex = Input.VertexId / 6;
+	const uint OffsetIndex = QuadIndex == 0 || QuadIndex == 2 || QuadIndex == 4 ? 0 : 1;
+	const uint4 PackedData = HairStrandsVF_PreviousPositionBuffer[BaseIndex + OffsetIndex];
+
+	const float3 ControlPoint = f16tof32(PackedData.xyz) + HairStrandsVF_WorldOffset;
+	const uint PackedAlpha = asuint(PackedData.a);
+	const uint ControlPointType = PackedAlpha & 0x3;
+
+	// Create a degenerated quad the end of each strand to cut between each strands
+	// #hair_todo: This is not efficient for short strand like fur for instance. Need to revisit that at some point
+	const bool bIsInvalidQuad = (ControlPointType == 1 && OffsetIndex == 1) || (ControlPointType == 2 && OffsetIndex == 0);
+	const float3 Position = bIsInvalidQuad ? float3(1.#INF, 1.#INF, 1.#INF) : ControlPoint;
+
+	return Position;
+}
+
+float GetWorldStrandRadius(FVertexFactoryInput Input)
+{
+	return GetVertexPosition(Input).WorldRadius;
+}
+
+float GetUCoord(FVertexFactoryInput Input)
+{
+	return GetVertexPosition(Input).UCoord;
+}
+
+float GetVCoord(FVertexFactoryInput Input)
+{
+	const uint QuadIndex = Input.VertexId % 6;
+	const float VCoord = QuadIndex == 0 || QuadIndex == 1 || QuadIndex == 5 ? 0.0f : 1.0f;
+	return VCoord;
+}
+
+float GetWorldStrandLength(FVertexFactoryInput Input)
+{
+	return GetVertexPosition(Input).WorldLength;
+}
+
+float GetHairSeed(FVertexFactoryInput Input)
+{
+	return GetVertexPosition(Input).Seed;
+}
+
+float2 GetHairRootUV(FVertexFactoryInput Input)
+{
+	return GetVertexPosition(Input).RootUV;
+}
+
+/** Converts from vertex factory specific interpolants to a FMaterialPixelParameters, which is used by material inputs. */
+FMaterialPixelParameters GetMaterialPixelParameters(FVertexFactoryInterpolantsVSToPS Interpolants, float4 SvPosition)
+{
+	// GetMaterialPixelParameters is responsible for fully initializing the result
+	FMaterialPixelParameters Result = MakeInitializedMaterialPixelParameters();
+
+	half3 TangentToWorld0 = GetTangentToWorld0(Interpolants).xyz;
+	half4 TangentToWorld2 = GetTangentToWorld2(Interpolants);
+	Result.UnMirrored = TangentToWorld2.w;
+
+	// Required for previewing materials that use ParticleColor
+	Result.Particle.Color = half4(1,1,1,1);
+
+	Result.TangentToWorld = AssembleTangentToWorld( TangentToWorld0, TangentToWorld2 );
+#if USE_WORLDVERTEXNORMAL_CENTER_INTERPOLATION
+	Result.WorldVertexNormal_Center = Interpolants.TangentToWorld2_Center.xyz;
+#endif
+
+	Result.TwoSidedSign = 1;
+	Result.PrimitiveId = GetPrimitiveId(Interpolants);
+	Result.HairPrimitiveId = Interpolants.HairPrimitiveId;
+	Result.HairDimensions = Interpolants.HairDimensions;
+	Result.HairUV = Interpolants.HairUV;
+	Result.HairRootUV = Interpolants.HairRootUV;
+	Result.HairSeed = Interpolants.HairSeed;
+
+#if COMPILER_SUPPORTS_TO_SCALAR_MEMORY
+	Result.PrimitiveId = ToScalarMemory(Result.PrimitiveId);
+#endif
+
+	return Result;
+}
+
+half3x3 CalcTangentToWorldNoScale(FVertexFactoryIntermediates Intermediates, half3x3 TangentToLocal)
+{
+	half3x3 LocalToWorld = GetLocalToWorld3x3(Intermediates.PrimitiveId);
+	half3 InvScale = GetPrimitiveData(Intermediates.PrimitiveId).InvNonUniformScaleAndDeterminantSign.xyz;
+	LocalToWorld[0] *= InvScale.x;
+	LocalToWorld[1] *= InvScale.y;
+	LocalToWorld[2] *= InvScale.z;
+	return mul(TangentToLocal, LocalToWorld); 
+}
+
+/** Converts from vertex factory specific input to a FMaterialVertexParameters, which is used by vertex shader material inputs. */
+FMaterialVertexParameters GetMaterialVertexParameters(FVertexFactoryInput Input, FVertexFactoryIntermediates Intermediates, float3 WorldPosition, half3x3 TangentToLocal)
+{
+	FMaterialVertexParameters Result = (FMaterialVertexParameters)0;
+	Result.WorldPosition = WorldPosition;
+
+	// does not handle instancing!
+	Result.TangentToWorld = Intermediates.TangentToWorld;
+	
+	Result.PrevFrameLocalToWorld = GetPrimitiveData(Intermediates.PrimitiveId).PreviousLocalToWorld;
+	Result.PreSkinnedPosition = GetVertexPreviousPosition(Input);
+	Result.PreSkinnedNormal = TangentToLocal[2];
+	Result.PrimitiveId = Intermediates.PrimitiveId;
+	return Result;
+}
+
+float4 CalcWorldPosition(float3 Position, uint PrimitiveId)
+{
+	return TransformLocalToTranslatedWorld(Position, PrimitiveId);
+}
+
+half3x3 CalcTangentToLocal(FVertexFactoryInput Input, out float TangentSign)
+{
+	const uint QuadIndex = Input.VertexId % 6;
+	const uint BaseIndex = Input.VertexId / 6;
+	const uint OffsetIndex = QuadIndex == 0 || QuadIndex == 2 || QuadIndex == 4 ? 0 : 1;
+	const half3 TangentInputX = HairStrandsVF_TangentBuffer[(BaseIndex + OffsetIndex) * 2 + 0].xyz;
+	const half4 TangentInputZ = HairStrandsVF_TangentBuffer[(BaseIndex + OffsetIndex) * 2 + 1];
+
+	half3 TangentX = TangentBias(TangentInputX);
+	half4 TangentZ = TangentBias(TangentInputZ);
+	TangentSign = TangentZ.w;
+	
+	// derive the binormal by getting the cross product of the normal and tangent
+	half3 TangentY = cross(TangentZ.xyz, TangentX) * TangentZ.w;
+	
+	// Recalculate TangentX off of the other two vectors
+	// This corrects quantization error since TangentX was passed in as a quantized vertex input
+	// The error shows up most in specular off of a mesh with a smoothed UV seam (normal is smooth, but tangents vary across the seam)
+	half3x3 Result;
+	Result[0] = cross(TangentY, TangentZ.xyz) * TangentZ.w;
+	Result[1] = TangentY;
+	Result[2] = TangentZ.xyz;
+
+	return Result;
+}
+
+half3x3 CalcTangentToWorld(FVertexFactoryIntermediates Intermediates, half3x3 TangentToLocal)
+{
+	half3x3 TangentToWorld = CalcTangentToWorldNoScale(Intermediates, TangentToLocal);
+	return TangentToWorld;
+}
+
+float3 GetTangent(FVertexFactoryInput Input)
+{
+	FVertexFactoryIntermediates Intermediates;
+	Intermediates.PrimitiveId = GetPrimitiveId(Input);
+	
+	float TangentSign;
+	Intermediates.TangentToLocal = CalcTangentToLocal(Input, TangentSign);
+	Intermediates.TangentToWorld = CalcTangentToWorld(Intermediates,Intermediates.TangentToLocal);
+	Intermediates.TangentToWorldSign = TangentSign * GetPrimitiveData(Intermediates.PrimitiveId).InvNonUniformScaleAndDeterminantSign.w;
+
+	// Hair shader expec the tangent to be stored in place of the normal. This is what the StrandHairComponent is filling in:
+	// [0]: Real normal		| Should be Tangent
+	// [1]: Bitangent		| Should be BiTangent
+	// [2]: Real Tangent	| Should be Normal
+	return Intermediates.TangentToWorld[2].xyz;
+}
+
+FVertexFactoryIntermediates GetVertexFactoryIntermediates(FVertexFactoryInput Input)
+{
+	FVertexFactoryIntermediates Intermediates;
+	Intermediates.PrimitiveId = GetPrimitiveId(Input);
+	Intermediates.HairPrimitiveId = GetHairPrimitiveId(Input);
+	Intermediates.HairDimensions = float2(GetWorldStrandLength(Input), GetWorldStrandRadius(Input));
+	Intermediates.HairUV = float2(GetUCoord(Input), GetVCoord(Input));
+	Intermediates.HairDensity = HairStrandsVF_Density;
+	Intermediates.HairSeed = GetHairSeed(Input);
+	Intermediates.HairRootUV = GetHairRootUV(Input);
+
+	float TangentSign;
+	Intermediates.TangentToLocal = CalcTangentToLocal(Input, TangentSign);
+	Intermediates.TangentToWorld = CalcTangentToWorld(Intermediates,Intermediates.TangentToLocal);
+	Intermediates.TangentToWorldSign = TangentSign * GetPrimitiveData(Intermediates.PrimitiveId).InvNonUniformScaleAndDeterminantSign.w;	
+	return Intermediates;
+}
+
+/**
+* Get the 3x3 tangent basis vectors for this vertex factory
+* this vertex factory will calculate the binormal on-the-fly
+*
+* @param Input - vertex input stream structure
+* @return 3x3 matrix
+*/
+half3x3 VertexFactoryGetTangentToLocal( FVertexFactoryInput Input, FVertexFactoryIntermediates Intermediates )
+{
+	return Intermediates.TangentToLocal;
+}
+
+float4 ComputeViewAlignedWorldPosition(FVertexFactoryInput Input, float3 WorldTangent, float4 WorldPosition, float WorldStrandRadius, float VelocityScale)
+{
+	const uint LocalVertexId = Input.VertexId % 6;
+	const bool bIsLeft = LocalVertexId == 0 || LocalVertexId == 1 || LocalVertexId == 5;
+	
+	// #hair_todo
+	// Use either the min strand radius computated within the strand factory, or a custom one, provided 
+	// (in a hacky way) throught the View variables. This is used by the DeepShadow pass, which inject 
+	// its data this way.
+	const bool  bUseCustomStrandRadius = ResolvedView.TranslatedWorldCameraOrigin.x == 1337.f;
+	const bool  bIsOrtho			   = bUseCustomStrandRadius && ResolvedView.TranslatedWorldCameraOrigin.y > 0;
+	const float CustomMinStrandRadius  = ResolvedView.TranslatedWorldCameraOrigin.z;
+	const float DistanceToCamera	   = length(ResolvedView.TranslatedWorldCameraOrigin - WorldPosition.xyz);
+
+	// Minimal radius to snap the strand to a sample/pixel center (to avoid aliasing)
+	const float RadiusAtDepth1 = bUseCustomStrandRadius ? CustomMinStrandRadius : lerp(HairStrandsVF_RadiusAtDepth1_Primary, HairStrandsVF_RadiusAtDepth1_Velocity, VelocityScale);
+	const float MinStrandHairRadius  = bIsOrtho ? RadiusAtDepth1 : (DistanceToCamera * RadiusAtDepth1);
+	const float3 ViewDir = -ResolvedView.ViewForward;
+	const float3 Right = normalize(cross(WorldTangent, ViewDir));
+	const float3 OutWorldPosition = WorldPosition.xyz + (bIsLeft ? -Right : Right) * max(WorldStrandRadius, MinStrandHairRadius);
+	
+	return float4(OutWorldPosition, 1);
+}
+
+// @return translated world position
+float4 VertexFactoryGetWorldPosition(FVertexFactoryInput Input, FVertexFactoryIntermediates Intermediates, float VelocityScale)
+{
+	FControlPoint ControlPoint = GetVertexPosition(Input);
+	const float3 VertexPosition = ControlPoint.Position;
+	const float4 WorldPosition = CalcWorldPosition(VertexPosition, Intermediates.PrimitiveId);
+
+	// Hair shader expect the WorldNormal to be the tangent vector
+	const float3 WorldTangent = GetTangent(Input);
+	return ComputeViewAlignedWorldPosition(Input, WorldTangent, WorldPosition, ControlPoint.WorldRadius, VelocityScale);
+}
+
+float4 VertexFactoryGetWorldPosition(FVertexFactoryInput Input, FVertexFactoryIntermediates Intermediates)
+{
+	return VertexFactoryGetWorldPosition(Input, Intermediates, 0);
+}
+
+float4 VertexFactoryGetRasterizedWorldPosition(FVertexFactoryInput Input, FVertexFactoryIntermediates Intermediates, float4 InWorldPosition)
+{
+	return InWorldPosition;
+}
+
+float3 VertexFactoryGetPositionForVertexLighting(FVertexFactoryInput Input, FVertexFactoryIntermediates Intermediates, float3 TranslatedWorldPosition)
+{
+	return TranslatedWorldPosition;
+}
+
+FVertexFactoryInterpolantsVSToPS VertexFactoryGetInterpolantsVSToPS(FVertexFactoryInput Input, FVertexFactoryIntermediates Intermediates, FMaterialVertexParameters VertexParameters)
+{
+	FVertexFactoryInterpolantsVSToPS Interpolants;
+
+	// Initialize the whole struct to 0
+	// Really only the last two components of the packed UVs have the opportunity to be uninitialized
+	Interpolants = (FVertexFactoryInterpolantsVSToPS)0;
+
+	SetTangents(Interpolants, Intermediates.TangentToWorld[0], Intermediates.TangentToWorld[2], Intermediates.TangentToWorldSign);
+
+#if INSTANCED_STEREO
+	Interpolants.EyeIndex = 0;
+#endif
+	
+	Interpolants.HairPrimitiveId = GetHairPrimitiveId(Input);
+	Interpolants.HairDimensions = float2(GetWorldStrandLength(Input), GetWorldStrandRadius(Input));
+	Interpolants.HairUV = float2(GetUCoord(Input), GetVCoord(Input));
+	Interpolants.HairSeed = GetHairSeed(Input);
+	Interpolants.HairRootUV = GetHairRootUV(Input);
+	SetPrimitiveId(Interpolants, Intermediates.PrimitiveId);
+
+	return Interpolants;
+}
+
+/** for depth-only pass */
+float4 VertexFactoryGetWorldPosition(FPositionOnlyVertexFactoryInput Input)
+{
+	const FControlPoint ControlPoint = GetVertexPosition(Input);
+	const float3 Position = ControlPoint.Position;
+	const uint PrimitiveId = GetPrimitiveId(Input);
+
+	const float4 WorldPosition = CalcWorldPosition(Position, PrimitiveId);
+    const float3 WorldTangent = GetTangent(Input);
+	return ComputeViewAlignedWorldPosition(Input, WorldTangent, WorldPosition, ControlPoint.WorldRadius, 0);
+}
+
+// @return translated world position (without quad extension/reorientation).This is used only for velocity computation
+float4 VertexFactoryGetWorldPositionRaw(FVertexFactoryInput Input, FVertexFactoryIntermediates Intermediates)
+{
+	return CalcWorldPosition(GetVertexPosition(Input).Position, Intermediates.PrimitiveId);
+}
+
+// @return previous translated world position
+float4 VertexFactoryGetPreviousWorldPosition(FVertexFactoryInput Input, FVertexFactoryIntermediates Intermediates)
+{
+	float4x4 PreviousLocalToWorldTranslated = GetPrimitiveData(Intermediates.PrimitiveId).PreviousLocalToWorld;
+	PreviousLocalToWorldTranslated[3][0] += ResolvedView.PrevPreViewTranslation.x;
+	PreviousLocalToWorldTranslated[3][1] += ResolvedView.PrevPreViewTranslation.y;
+	PreviousLocalToWorldTranslated[3][2] += ResolvedView.PrevPreViewTranslation.z;
+
+	const float3 VertexPosition = GetVertexPreviousPosition(Input);
+	return mul(float4(VertexPosition,1), PreviousLocalToWorldTranslated);
+}
+
+float4 VertexFactoryGetTranslatedPrimitiveVolumeBounds(FVertexFactoryInterpolantsVSToPS Interpolants)
+{
+	float4 ObjectWorldPositionAndRadius = GetPrimitiveData(GetPrimitiveId(Interpolants)).ObjectWorldPositionAndRadius;
+	return float4(ObjectWorldPositionAndRadius.xyz + ResolvedView.PreViewTranslation.xyz, ObjectWorldPositionAndRadius.w);
+}
+
+uint VertexFactoryGetPrimitiveId(FVertexFactoryInterpolantsVSToPS Interpolants)
+{
+	return GetPrimitiveId(Interpolants);
+}
+
+float3 VertexFactoryGetWorldNormal(FPositionAndNormalOnlyVertexFactoryInput Input)
+{
+	return float3(0, 0, 1); 
+}
+
+float3 VertexFactoryGetWorldNormal(FVertexFactoryInput Input, FVertexFactoryIntermediates Intermediates)
+{
+	return float3(0,0,1);
+}
\ No newline at end of file
diff --git a/Engine/Shaders/Private/HairStrands/HairStrandsVisibilityClearPS.usf b/Engine/Shaders/Private/HairStrands/HairStrandsVisibilityClearPS.usf
new file mode 100644
--- /dev/null
+++ b/Engine/Shaders/Private/HairStrands/HairStrandsVisibilityClearPS.usf
@@ -0,0 +1,7 @@
+// Copyright 1998-2019 Epic Games, Inc. All Rights Reserved.
+
+uint ClearValue;
+void ClearPS(out uint4 OutColor : SV_Target0)
+{
+	OutColor = ClearValue.xxxx;
+}
\ No newline at end of file
diff --git a/Engine/Shaders/Private/HairStrands/HairStrandsVisibilityCommon.ush b/Engine/Shaders/Private/HairStrands/HairStrandsVisibilityCommon.ush
new file mode 100644
--- /dev/null
+++ b/Engine/Shaders/Private/HairStrands/HairStrandsVisibilityCommon.ush
@@ -0,0 +1,307 @@
+// Copyright 1998-2019 Epic Games, Inc. All Rights Reserved.
+
+struct FVisibilityData
+{
+	float3 Tangent;		// 24 bits
+	uint   Coverage8bit;//  8 bits
+	uint   PrimitiveId;	// 28 bits
+	uint   ClusterId;	//  4 bits
+};
+
+// Footprint: 128bits | 16bytes
+// Option to reduce the footprint: 
+// * we don't use PrimitiveID at the moment which is encoded on 3bytes
+// * Deferred the evaluation of Coverage & material properties into some other pass?
+struct FNodeData
+{
+	float Depth;
+	uint  PrimitiveID_ClusterID;
+	uint  Tangent_Coverage8bit;
+	uint  BaseColor_Roughness;
+	float Specular;
+};
+
+float3 EncodeTangent(float3 N)
+{
+	return N * 0.5 + 0.5;
+}
+
+float3 DecodeTangent(float3 N)
+{
+	return N * 2 - 1;
+}
+
+uint float4ToUint(float4 v)
+{
+	uint4 i = v * 255;
+	return (0xFF & i.w) << 24 | (0xFF & i.z) << 16 | (0xFF & i.y) << 8 | (0xFF & i.x);
+}
+
+float4 UintToFloat4(uint In)
+{
+	uint4 Out;
+	Out.x = (0xFF &  In);
+	Out.y = (0xFF & (In >>  8));
+	Out.z = (0xFF & (In >> 16));
+	Out.w = (0xFF & (In >> 24));
+	return Out / 255.f;
+}
+
+uint4 EncodeVisibilityData(FVisibilityData Data)
+{
+	uint4 Output;
+
+	// Tangent and Coverage
+	// Encode into 2 x 16bits
+	float3 T = saturate(EncodeTangent(Data.Tangent));
+	Output.xy = uint2(
+		(0xFF & uint(T.x * 255)) | ((0xFF & uint(T.y * 255)) << 8),
+		(0xFF & uint(T.z * 255)) | ((0xFF & Data.Coverage8bit) << 8));
+
+	// PrimitiveId and ClusterId
+	// Encode into 2 x 16bits
+	// PrimitiveId is on 28bits | ClusterId is on 4bits
+	Output.zw = uint2(
+		0xFFFF & Data.PrimitiveId,
+		((Data.ClusterId & 0x0F) << 12) | (0x0FFF & (uint(Data.PrimitiveId) >> 16)));
+
+	return Output;
+}
+
+FVisibilityData DecodeVisibilityData(uint4 EncodedData)
+{
+	FVisibilityData Output;
+	Output.Tangent = DecodeTangent(float3(
+		float((EncodedData.x     ) & 0xFF) / 255.f,
+		float((EncodedData.x >> 8) & 0xFF) / 255.f,
+		float((EncodedData.y     ) & 0xFF) / 255.f));
+	Output.Coverage8bit = (EncodedData.y >> 8) & 0xFF;
+
+	Output.PrimitiveId = EncodedData.z | ((EncodedData.w & 0x0FFF) << 16);
+	Output.ClusterId = (EncodedData.w & 0xF000) >> 12;
+
+	return Output;
+}
+
+bool IsValidHairData(uint4 EncodedVisibilityData)
+{
+	const uint InvalidSample = 0xFFFF;
+	return EncodedVisibilityData.w != InvalidSample;
+}
+
+uint DecodeClusterId(uint4 EncodedVisibilityData)
+{
+	return (EncodedVisibilityData.w >> 12) & 0xF;
+}
+
+// * Max 67,108,864 total nodes (This is 32 nodes per pixel at 1080p if all pixel were covered)
+// * Max 63 nodes per pixel
+// 26bits for offset | 6 bits for count (max 63 nodes)
+struct FNodeDesc
+{
+	uint Offset;
+	uint Count;
+};
+
+uint EncodeNodeDesc(const FNodeDesc Desc)
+{
+	return (Desc.Offset & 0x03FFFFFF) | ((Desc.Count & 0x3F) << 26);
+}
+
+FNodeDesc DecodeNodeDesc(uint In)
+{
+	FNodeDesc Out;
+	Out.Offset = In & 0x03FFFFFF;
+	Out.Count  = (In >> 26) & 0x3F;
+	return Out;
+}
+
+FNodeData ConvertToNodeData(float Depth, FVisibilityData VisibilityData, float3 BaseColor, float Roughness, float Specular)
+{
+	// 4x16bits
+	const uint4 EncodedVisibilityData = EncodeVisibilityData(VisibilityData);
+
+	FNodeData Out;
+	Out.Depth					= Depth;																			// 32bits float
+	Out.Tangent_Coverage8bit	= (0xFFFF & EncodedVisibilityData.x) | ((0xFFFF & EncodedVisibilityData.y) << 16);	// 32bits uint
+	Out.PrimitiveID_ClusterID	= (0xFFFF & EncodedVisibilityData.z) | ((0xFFFF & EncodedVisibilityData.w) << 16);	// 32bits uint
+	Out.BaseColor_Roughness		= float4ToUint(float4(sqrt(BaseColor), Roughness));									// 32bits uint
+	Out.Specular				= Specular;																			// 32bits float
+
+	return Out;
+}
+
+ void DecodeNodeData(
+	 const FNodeData	 InNode, 
+	 out float			 OutDepth, 
+	 out FVisibilityData OutVisibilityData,
+	 out float3			 OutBaseColor, 
+	 out float			 OutRoughness,
+	 out float			 OutSpecular)
+{
+	OutDepth = InNode.Depth;
+
+	const uint4 EncodedVisibility = uint4(
+		(InNode.Tangent_Coverage8bit & 0xFFFF),
+		(InNode.Tangent_Coverage8bit >> 16) & 0xFFFF,
+		(InNode.PrimitiveID_ClusterID & 0xFFFF),
+		(InNode.PrimitiveID_ClusterID >> 16) & 0xFFFF);
+	OutVisibilityData = DecodeVisibilityData(EncodedVisibility);
+
+	const float4 BaseColorAndRoughness = UintToFloat4(InNode.BaseColor_Roughness);
+	OutBaseColor = BaseColorAndRoughness.xyz * BaseColorAndRoughness.xyz;
+	OutRoughness = BaseColorAndRoughness.w;
+	OutSpecular = InNode.Specular;
+}
+
+#define HAIR_VISIBILITY_MAXSAMPLECOUNT 8
+#define HAIR_VISIBILITY_GROUP_COUNT_WIDTH 64
+
+ struct FSampleDesc
+ {
+	 uint   SampleId;
+	 uint   ClusterId;
+	 float  Depths;
+	 float  Weights;
+	 float3 Tangents;
+	 uint   Coverage8bit;
+	 float4 BaseColorAndRoughness;
+	 float  Specular;
+ };
+
+ struct FSampleSetDesc
+ {
+	 uint   UniqueSampleCount;
+	 uint   ValidSampleCount;
+	 uint   HairSampleCount;
+ };
+
+ void ComputeUniqueSamples(
+	 out FSampleDesc OutSamples[HAIR_VISIBILITY_MAXSAMPLECOUNT],
+	 out FSampleSetDesc OutSet,
+	 const uint2 PixelCoord,
+	 const uint SampleCount,
+	 Texture2D<float4> InSceneDepthTexture,
+	 Texture2DMS<uint4, HAIR_VISIBILITY_MAXSAMPLECOUNT>  InMSAA_IDTexture,
+	 Texture2DMS<float4, HAIR_VISIBILITY_MAXSAMPLECOUNT> InMSAA_MaterialTexture,
+//	 Texture2DMS<float4, HAIR_VISIBILITY_MAXSAMPLECOUNT> InMSAA_AttributeTexture,
+	 Texture2DMS<float, HAIR_VISIBILITY_MAXSAMPLECOUNT>  InMSAA_DepthTexture)
+ {
+	 OutSet.UniqueSampleCount = 0;
+	 OutSet.ValidSampleCount = 0;
+
+	 OutSet.HairSampleCount = clamp(SampleCount, 0, HAIR_VISIBILITY_MAXSAMPLECOUNT);
+	 const uint InvalidSample = 0xFFFF;
+	 const float OpaqueDepth = InSceneDepthTexture.Load(uint3(PixelCoord, 0)).r;
+	 for (uint SampleIt = 0; SampleIt < OutSet.HairSampleCount; ++SampleIt)
+	 {
+		 const uint4 EncodedData = InMSAA_IDTexture.Load(PixelCoord, SampleIt);
+		 const bool bIsValid = IsValidHairData(EncodedData);
+		 if (!bIsValid)
+			 continue;
+
+		 const float SampleDepth = InMSAA_DepthTexture.Load(PixelCoord, SampleIt);
+		 if (OpaqueDepth > SampleDepth) // Reverse-Z
+			 continue;
+
+		 ++OutSet.ValidSampleCount;
+
+		 const FVisibilityData VisibilityData = DecodeVisibilityData(EncodedData);
+		 bool bAlreadyExist = false;
+		 for (uint UniqueIt = 0; UniqueIt < OutSet.UniqueSampleCount; ++UniqueIt)
+		 {
+			 if (VisibilityData.PrimitiveId == OutSamples[UniqueIt].SampleId)
+			 {
+				 OutSamples[UniqueIt].Weights += 1;
+				 bAlreadyExist = true;
+				 break;
+			 }
+		 }
+
+		 if (!bAlreadyExist)
+		 {
+			 //const float Specular = InMSAA_AttributeTexture.Load(PixelCoord, SampleIt).x;
+			 const float4 BaseColorAndRoughness = InMSAA_MaterialTexture.Load(PixelCoord, SampleIt);
+			 OutSamples[OutSet.UniqueSampleCount].SampleId = VisibilityData.PrimitiveId;
+			 OutSamples[OutSet.UniqueSampleCount].Tangents = VisibilityData.Tangent;
+			 OutSamples[OutSet.UniqueSampleCount].Coverage8bit = VisibilityData.Coverage8bit;
+			 OutSamples[OutSet.UniqueSampleCount].Depths = SampleDepth;
+			 OutSamples[OutSet.UniqueSampleCount].Weights = 1;
+			 OutSamples[OutSet.UniqueSampleCount].ClusterId = VisibilityData.ClusterId;
+			 OutSamples[OutSet.UniqueSampleCount].BaseColorAndRoughness = float4(BaseColorAndRoughness.xyz * BaseColorAndRoughness.xyz, BaseColorAndRoughness.a);
+			 OutSamples[OutSet.UniqueSampleCount].Specular = 0.5f; // Specular;
+			 ++OutSet.UniqueSampleCount;
+		 }
+	 }
+ }
+
+ uint To8bitCoverage(float Coverage)
+ {
+	 return min(uint(Coverage * 0x100), 0xFFu);
+ }
+
+ float From8bitCoverage(uint Coverage8bit)
+ {
+	 return float(Coverage8bit) / 255.f;
+ }
+
+ float GetVoxelDensityFixPointScale()
+ {
+	 // Constant for scaling the hair coverage during the voxelization (done with atomic integer)
+	 return 1000.f;
+ }
+
+ struct FCategorizationData
+ {
+	 uint  TotalCoverage8bit;
+	 uint  SampleCount;
+	 float ClosestDepth;
+ };
+
+ uint4 EncodeCategorizationData(FCategorizationData Data)
+ {
+	 const uint ClosestDepthUint = asuint(Data.ClosestDepth);
+	 const uint ClosestDepthLow  = ClosestDepthUint & 0xFFFF;
+	 const uint ClosestDepthHigh = (ClosestDepthUint >> 16) & 0xFFFF;
+
+	 return uint4(Data.TotalCoverage8bit, Data.SampleCount, ClosestDepthLow, ClosestDepthHigh);
+ }
+
+ FCategorizationData DecodeCategorizationData(uint4 Data)
+ {
+	 FCategorizationData Out;
+	 Out.TotalCoverage8bit			= Data.x;
+	 Out.SampleCount				= Data.y;
+	 const uint ClosestDepthLow		= Data.z & 0xFFFF;
+	 const uint ClosestDepthHigh	=(Data.w & 0xFFFF) << 16;
+	 Out.ClosestDepth				= asfloat(ClosestDepthHigh | ClosestDepthLow);
+
+	 return Out;
+ }
+
+ uint3 QuantizeTo8Bits(float3 T)
+ {
+	 return saturate((T + float(1).xxx) * 0.5f) * 0xFF;
+ }
+
+ float3 From8bits(float3 In8bits)
+ {
+	 return (normalize(In8bits / 255.f) * 2) - float(1).xxx;
+ }
+
+ uint EncodeBaseColorAndRoughness(float3 BaseColor, float Roughness)
+ {
+	 return
+		 (uint(BaseColor.x * 0xFF) & 0xFF) |
+		 (uint(BaseColor.y * 0xFF) & 0xFF) << 8 |
+		 (uint(BaseColor.z * 0xFF) & 0xFF) << 16 |
+		 (uint(Roughness   * 0xFF) & 0xFF) << 24;
+ }
+
+ void DecodeBaseColorAndRoughness(uint Encoded, out float3 BaseColor, out float Roughness)
+ {
+	 BaseColor.x = (float((Encoded) & 0xFF) / 255.f);
+	 BaseColor.y = (float((Encoded >> 8) & 0xFF) / 255.f);
+	 BaseColor.z = (float((Encoded >> 16) & 0xFF) / 255.f);
+	 Roughness   = (float((Encoded >> 24) & 0xFF) / 255.f);
+ }
\ No newline at end of file
diff --git a/Engine/Shaders/Private/HairStrands/HairStrandsVisibilityCompaction.usf b/Engine/Shaders/Private/HairStrands/HairStrandsVisibilityCompaction.usf
new file mode 100644
--- /dev/null
+++ b/Engine/Shaders/Private/HairStrands/HairStrandsVisibilityCompaction.usf
@@ -0,0 +1,249 @@
+// Copyright 1998-2019 Epic Games, Inc. All Rights Reserved.
+
+#pragma once
+
+#include "../Common.ush"
+#include "HairStrandsVisibilityCommon.ush"
+
+#define RENDER_MODE_MSAA 0
+#define VENDOR_AMD 0
+#define VENDOR_NVIDIA 1
+#define VENDOR_OTHER 2
+
+void ComputeUniqueSamplesWithoutDepthTest(
+	out uint3 OutSamples[HAIR_VISIBILITY_MAXSAMPLECOUNT],
+	out FSampleSetDesc OutSet,
+	const uint2 PixelCoord,
+	const uint SampleCount,
+	const float OpaqueDepth,
+	Texture2DMS<uint4, HAIR_VISIBILITY_MAXSAMPLECOUNT>  InMSAA_IDTexture,
+	Texture2DMS<float, HAIR_VISIBILITY_MAXSAMPLECOUNT>  InMSAA_DepthTexture)
+{
+	OutSet.UniqueSampleCount = 0;
+	OutSet.ValidSampleCount = 0;
+
+	OutSet.HairSampleCount = clamp(SampleCount, 0, HAIR_VISIBILITY_MAXSAMPLECOUNT);
+	const uint InvalidSample = 0xFFFF;
+	for (uint SampleIt = 0; SampleIt < OutSet.HairSampleCount; ++SampleIt)
+	{
+		const uint4 EncodedData = InMSAA_IDTexture.Load(PixelCoord, SampleIt);
+		const bool bIsValid = IsValidHairData(EncodedData);
+		if (!bIsValid)
+			continue;
+
+		const float SampleDepth = InMSAA_DepthTexture.Load(PixelCoord, SampleIt);
+		if (OpaqueDepth > SampleDepth) // Reverse-Z
+			continue;
+
+		++OutSet.ValidSampleCount;
+
+		const FVisibilityData VisibilityData = DecodeVisibilityData(EncodedData);
+		bool bAlreadyExist = false;
+		for (uint UniqueIt = 0; UniqueIt < OutSet.UniqueSampleCount; ++UniqueIt)
+		{
+			if (VisibilityData.PrimitiveId == OutSamples[UniqueIt].x)
+			{
+				OutSamples[UniqueIt].y += 1;
+				bAlreadyExist = true;
+				break;
+			}
+		}
+
+		if (!bAlreadyExist)
+		{
+			OutSamples[OutSet.UniqueSampleCount].x = VisibilityData.PrimitiveId;
+			OutSamples[OutSet.UniqueSampleCount].y = 1;
+			OutSamples[OutSet.UniqueSampleCount].z = SampleIt;
+			++OutSet.UniqueSampleCount;
+		}
+	}
+}
+
+
+//////////////////////////////////////////////////////////////////////////////////////////////////////////////
+#if PERMUTATION_VENDOR == VENDOR_AMD
+#define TILE_PIXEL_SIZE_X 8
+#define TILE_PIXEL_SIZE_Y 8
+#elif PERMUTATION_VENDOR == VENDOR_NVIDIA
+#define TILE_PIXEL_SIZE_X 8
+#define TILE_PIXEL_SIZE_Y 4
+#else //VENDOR_OTHER
+#define TILE_PIXEL_SIZE_X 8
+#define TILE_PIXEL_SIZE_Y 8
+#endif
+
+#ifndef HAIR_VISIBILITY_MAXSAMPLECOUNT
+#error HAIR_VISIBILITY_MAXSAMPLECOUNT is not defined
+#endif
+
+uint2 OutputResolution;
+uint MaxNodeCount;
+uint HairVisibilitySampleCount;
+uint2 ResolutionOffset;
+
+Texture2DMS<uint4, HAIR_VISIBILITY_MAXSAMPLECOUNT> MSAA_IDTexture;
+Texture2DMS<float, HAIR_VISIBILITY_MAXSAMPLECOUNT> MSAA_DepthTexture;
+Texture2DMS<float4,HAIR_VISIBILITY_MAXSAMPLECOUNT> MSAA_MaterialTexture;
+Texture2DMS<float4,HAIR_VISIBILITY_MAXSAMPLECOUNT> MSAA_AttributeTexture;
+
+RWTexture2D<uint> OutCompactNodeCounter;
+RWTexture2D<uint> OutCompactNodeIndex;
+RWStructuredBuffer<FNodeData> OutCompactNodeData;
+RWStructuredBuffer<uint> OutCompactNodeCoord;
+RWTexture2D<uint4> OutCategorizationTexture;
+
+#define VELOCITY_TYPE_NONE 0
+#define VELOCITY_TYPE_AVG 1  
+#define VELOCITY_TYPE_CLOSEST 2
+#define VELOCITY_TYPE_MAX 3
+#if PERMUTATION_VELOCITY != VELOCITY_TYPE_NONE
+	Texture2DMS<float2, HAIR_VISIBILITY_MAXSAMPLECOUNT> MSAA_VelocityTexture;
+	RWTexture2D<float2> OutVelocityTexture;
+#endif
+
+groupshared uint AllocationNodeCount;
+groupshared uint AllocationNodeOffset;
+
+[numthreads(TILE_PIXEL_SIZE_X, TILE_PIXEL_SIZE_Y, 1)]
+void MainCS(uint GroupIndex : SV_GroupIndex, uint3 DispatchThreadId : SV_DispatchThreadID)
+{
+	if (GroupIndex == 0)
+	{
+		AllocationNodeCount = 0;
+		AllocationNodeOffset = 0;
+	}
+
+	uint2 PixelCoord = DispatchThreadId.xy + ResolutionOffset;
+	if (HairVisibilitySampleCount == 0 || PixelCoord.x >= OutputResolution.x || PixelCoord.y >= OutputResolution.y)
+		PixelCoord = uint2(0, 0);
+
+	FSampleSetDesc SampleDesc;
+
+	const float OpaqueDepth = SceneTexturesStruct.SceneDepthTexture.Load(uint3(PixelCoord, 0)).r;
+	uint3 Samples[HAIR_VISIBILITY_MAXSAMPLECOUNT]; // x:PrimitiveId, y:Weight, z:SampleIndex
+	ComputeUniqueSamplesWithoutDepthTest(
+		Samples,
+		SampleDesc,
+		PixelCoord,
+		HairVisibilitySampleCount,
+		OpaqueDepth,
+		MSAA_IDTexture, 
+		MSAA_DepthTexture);
+
+	FNodeDesc NodeDesc;
+	NodeDesc.Count = SampleDesc.UniqueSampleCount;
+	NodeDesc.Offset = 0;
+
+	InterlockedAdd(AllocationNodeCount, NodeDesc.Count, NodeDesc.Offset);
+	GroupMemoryBarrierWithGroupSync();
+	if (GroupIndex == 0 && AllocationNodeCount > 0)
+	{
+		InterlockedAdd(OutCompactNodeCounter[uint2(0, 0)], AllocationNodeCount, AllocationNodeOffset);
+	}
+	GroupMemoryBarrierWithGroupSync();
+
+	// Allocate node space
+	FCategorizationData Categorization = (FCategorizationData)0;
+	if (NodeDesc.Count > 0)
+	{
+		NodeDesc.Offset += AllocationNodeOffset;
+
+		// Store final sort node data 
+		if (NodeDesc.Offset + NodeDesc.Count < MaxNodeCount)
+		{
+			#if PERMUTATION_VELOCITY == VELOCITY_TYPE_AVG
+			float2 AverageVelocity = 0;
+			uint ValidVelocityCount = 0;
+			#endif
+
+			#if PERMUTATION_VELOCITY == VELOCITY_TYPE_CLOSEST
+			float2 ClosestEncodedVelocity = 0;
+			float ClosestDepth = 0;
+			#endif		
+
+			#if PERMUTATION_VELOCITY == VELOCITY_TYPE_MAX
+			float2 MaxEncodedVelocity = 0;
+			float MaxVelocityMagnitude2 = 0;
+			#endif		
+
+			for (uint OutIndex = 0; OutIndex < NodeDesc.Count; ++OutIndex)
+			{
+				const uint SampleIt = Samples[OutIndex].z;
+
+				// Read all material data
+				// #hair_todo: simplify the encoding. technically we don't need to decompress the visibility data, 
+				// as we only need to patch the weight
+				// #hair_todo: merge similar sample
+				FNodeData OutNodeData;
+				{
+					const float SampleDepth = MSAA_DepthTexture.Load(PixelCoord, SampleIt);
+					const uint4 EncodedData = MSAA_IDTexture.Load(PixelCoord, SampleIt);
+					FVisibilityData VisibilityData = DecodeVisibilityData(EncodedData);
+					VisibilityData.Coverage8bit = To8bitCoverage(Samples[OutIndex].y / float(SampleDesc.HairSampleCount));
+					Categorization.TotalCoverage8bit += VisibilityData.Coverage8bit;
+					Categorization.ClosestDepth = max(Categorization.ClosestDepth, SampleDepth); // Inverse-Z
+
+					float4 BaseColorAndRoughness = MSAA_MaterialTexture.Load(PixelCoord, SampleIt);
+					BaseColorAndRoughness.xyz = BaseColorAndRoughness.xyz * BaseColorAndRoughness.xyz; // Cheap gamma encoding since this is encoded into 8bit target
+					const float Specular = MSAA_AttributeTexture.Load(PixelCoord, SampleIt).x;
+
+
+					OutNodeData = ConvertToNodeData(
+						SampleDepth,
+						VisibilityData,
+						BaseColorAndRoughness.xyz,
+						BaseColorAndRoughness.w,
+						Specular);
+
+					#if PERMUTATION_VELOCITY == VELOCITY_TYPE_AVG
+					const float2 EncodedVelocity = MSAA_VelocityTexture.Load(PixelCoord, SampleIt);
+					AverageVelocity += DecodeVelocityFromTexture(EncodedVelocity);
+					++ValidVelocityCount;
+					#endif			
+
+					#if PERMUTATION_VELOCITY == VELOCITY_TYPE_CLOSEST
+					const float2 EncodedVelocity = MSAA_VelocityTexture.Load(PixelCoord, SampleIt);
+					if (ClosestDepth < Categorization.ClosestDepth) // Inverse-Z
+					{
+						ClosestEncodedVelocity = EncodedVelocity;
+						ClosestDepth = Categorization.ClosestDepth;
+					}
+					#endif
+
+					#if PERMUTATION_VELOCITY == VELOCITY_TYPE_MAX
+					const float2 EncodedVelocity = MSAA_VelocityTexture.Load(PixelCoord, SampleIt);
+					const float2 CurrentVelocity = DecodeVelocityFromTexture(EncodedVelocity);
+					const float CurrentVelocityMagnitude2 = dot(CurrentVelocity, CurrentVelocity);
+					if (CurrentVelocityMagnitude2 > MaxVelocityMagnitude2)
+					{
+						MaxEncodedVelocity = EncodedVelocity;
+						MaxVelocityMagnitude2 = CurrentVelocityMagnitude2;
+					}
+					#endif			
+				}
+
+				const uint StoreIndex = NodeDesc.Offset + OutIndex;
+				OutCompactNodeData[StoreIndex] = OutNodeData;
+				OutCompactNodeCoord[StoreIndex] = (PixelCoord.x & 0xFFFF) | ((PixelCoord.y & 0xFFFF) << 16);
+			}
+
+			#if PERMUTATION_VELOCITY == VELOCITY_TYPE_AVG
+			OutVelocityTexture[PixelCoord] = EncodeVelocityToTexture(AverageVelocity / max(ValidVelocityCount, 1));
+			#endif		
+
+			#if PERMUTATION_VELOCITY == VELOCITY_TYPE_CLOSEST
+			OutVelocityTexture[PixelCoord] = ClosestEncodedVelocity;
+			#endif		
+
+			#if PERMUTATION_VELOCITY == VELOCITY_TYPE_MAX
+			OutVelocityTexture[PixelCoord] = MaxEncodedVelocity;
+			#endif		
+
+			NodeDesc.Count = NodeDesc.Count;
+			Categorization.SampleCount = NodeDesc.Count;
+		}
+	}
+
+	OutCompactNodeIndex[PixelCoord] = EncodeNodeDesc(NodeDesc);
+	OutCategorizationTexture[PixelCoord] = EncodeCategorizationData(Categorization);
+}
\ No newline at end of file
diff --git a/Engine/Shaders/Private/HairStrands/HairStrandsVisibilityComposeSubPixelPS.usf b/Engine/Shaders/Private/HairStrands/HairStrandsVisibilityComposeSubPixelPS.usf
new file mode 100644
--- /dev/null
+++ b/Engine/Shaders/Private/HairStrands/HairStrandsVisibilityComposeSubPixelPS.usf
@@ -0,0 +1,56 @@
+// Copyright 1998-2019 Epic Games, Inc. All Rights Reserved.
+
+#include "../Common.ush"
+
+#include "../SceneTexturesCommon.ush"
+#include "HairStrandsDeepShadowCommon.ush"
+#include "../DeferredShadingCommon.ush"
+#include "HairStrandsVisibilityCommon.ush"
+
+Texture2D<float4> SubPixelColorTexture;
+Texture2DMS<float, HAIR_VISIBILITY_MAXSAMPLECOUNT> HairVisibilityDepthTexture;
+
+void SubColorPS(
+	in FScreenVertexOutput Input,
+	out float4 OutColor  : SV_Target0,
+	out float  OutDepth : SV_Depth)
+{
+	const uint2 PixelCoord = floor(Input.Position.xy);
+	OutColor = SubPixelColorTexture.Load(uint3(PixelCoord, 0));
+
+	// Early out writing to not write stencil value for full coverred pixels (or not covered at all)
+	//if (OutColor.a == 0)
+	//	discard;
+
+	// Brute force reading of the depth value. This could be pre-computed during the categorize pass or the compaction pass
+	OutDepth = 0;
+	for (uint SampleIt = 0; SampleIt < HAIR_VISIBILITY_MAXSAMPLECOUNT; ++SampleIt)
+	{
+		const float SampleDepth = HairVisibilityDepthTexture.Load(PixelCoord, SampleIt);
+		OutDepth = max(OutDepth, SampleDepth); // Inverse-Z
+	}
+}
+
+float VelocityThreshold;
+Texture2DMS<float2, HAIR_VISIBILITY_MAXSAMPLECOUNT> HairVisibilityVelocityTexture;
+
+void FastResolvePS(in FScreenVertexOutput Input)
+{
+	const uint2 PixelCoord = floor(Input.Position.xy);
+
+	bool bNeedFastResolve = false;
+	for (uint SampleIt = 0; SampleIt < HAIR_VISIBILITY_MAXSAMPLECOUNT; ++SampleIt)
+	{
+		const float2 EncodedVelocity = HairVisibilityVelocityTexture.Load(PixelCoord, SampleIt);
+		const float2 Velocity = DecodeVelocityFromTexture(EncodedVelocity);
+		const float VelocityMagnitude = sqrt(dot(Velocity, Velocity));
+		if (EncodedVelocity.x > 0 && VelocityMagnitude > VelocityThreshold)
+		{
+			bNeedFastResolve = true;
+			break;
+		}
+	}
+
+	if (!bNeedFastResolve)
+		discard;
+}
\ No newline at end of file
diff --git a/Engine/Shaders/Private/HairStrands/HairStrandsVisibilityCopyIndirectArg.usf b/Engine/Shaders/Private/HairStrands/HairStrandsVisibilityCopyIndirectArg.usf
new file mode 100644
--- /dev/null
+++ b/Engine/Shaders/Private/HairStrands/HairStrandsVisibilityCopyIndirectArg.usf
@@ -0,0 +1,41 @@
+// Copyright 1998-2019 Epic Games, Inc. All Rights Reserved.
+
+#pragma once
+
+#include "../Common.ush"
+#include "HairStrandsVisibilityCommon.ush"
+
+Texture2D<uint> CounterTexture;
+RWBuffer<uint> OutArgBuffer;
+uint ThreadGroupSize;
+
+[numthreads(1, 1, 1)]
+void CopyCS(uint GroupIndex : SV_GroupIndex, uint3 DispatchThreadId : SV_DispatchThreadID)
+{
+	const uint NodeCount  = CounterTexture.Load(uint3(0,0,0));
+	const uint TotalGroup = (NodeCount + ThreadGroupSize-1) / ThreadGroupSize;
+#if 0
+	// Minimize unused groups
+	const uint SideSize   = uint(floor(sqrt(float(TotalGroup))));
+	const uint ExtraLines = (TotalGroup - SideSize * SideSize + (SideSize - 1)) / SideSize;
+	OutArgBuffer[0] = SideSize;
+	OutArgBuffer[1] = SideSize + ExtraLines;
+	OutArgBuffer[2] = 1;
+#else
+	// Note: Use a fixed group count width (64) for avoiding loading the indirect args (to avoid dep. memory fetch)
+	// This value is in sync with DeepTransmittanceMask.usf
+	if (TotalGroup < HAIR_VISIBILITY_GROUP_COUNT_WIDTH)
+	{
+		OutArgBuffer[0] = TotalGroup;
+		OutArgBuffer[1] = 1;
+		OutArgBuffer[2] = 1;
+	}
+	else
+	{
+		const uint DispatchY = ceil(float(TotalGroup) / float(HAIR_VISIBILITY_GROUP_COUNT_WIDTH));
+		OutArgBuffer[0] = 64;
+		OutArgBuffer[1] = DispatchY;
+		OutArgBuffer[2] = 1;
+	}
+#endif
+}
\ No newline at end of file
diff --git a/Engine/Shaders/Private/HairStrands/HairStrandsVisibilityDepthPS.usf b/Engine/Shaders/Private/HairStrands/HairStrandsVisibilityDepthPS.usf
new file mode 100644
--- /dev/null
+++ b/Engine/Shaders/Private/HairStrands/HairStrandsVisibilityDepthPS.usf
@@ -0,0 +1,40 @@
+// Copyright 1998-2019 Epic Games, Inc. All Rights Reserved.
+
+#include "../Common.ush"
+
+#include "../SceneTexturesCommon.ush"
+#include "HairStrandsDeepShadowCommon.ush"
+#include "../DeferredShadingCommon.ush"
+#include "HairStrandsVisibilityCommon.ush"
+
+Texture2DMS<float, HAIR_VISIBILITY_MAXSAMPLECOUNT> HairVisibilityDepthTexture;
+
+void MainPS(
+	in FScreenVertexOutput Input,
+	out float SvDepth		: SV_Depth,
+	out uint4  OutGBufferB	: SV_Target0,
+	out uint4  OutColor		: SV_Target1
+)
+{
+	float MaxDepth = 0;
+	float MinDepth = 1;
+	const uint2 PixelCoord = floor(Input.Position.xy);
+	uint ValidSampleCount = 0;
+	for (uint SampleIt=0; SampleIt<HAIR_VISIBILITY_MAXSAMPLECOUNT; ++SampleIt)
+	{
+		const float SampleDepth = HairVisibilityDepthTexture.Load(PixelCoord, SampleIt).x;
+		MinDepth = min(MinDepth, SampleDepth);
+		MaxDepth = max(MaxDepth, SampleDepth);
+		const bool bIsValid = SampleDepth > 0;
+		ValidSampleCount += bIsValid ? 1 : 0;
+	}
+
+	if (ValidSampleCount < HAIR_VISIBILITY_MAXSAMPLECOUNT || MaxDepth == 0)
+		discard;
+
+	SvDepth = MinDepth;
+	OutGBufferB = 0;
+	OutColor = 0;
+}
+
+
diff --git a/Engine/Shaders/Private/HairStrands/HairStrandsVisibilityFillOpaqueDepthPS.usf b/Engine/Shaders/Private/HairStrands/HairStrandsVisibilityFillOpaqueDepthPS.usf
new file mode 100644
--- /dev/null
+++ b/Engine/Shaders/Private/HairStrands/HairStrandsVisibilityFillOpaqueDepthPS.usf
@@ -0,0 +1,14 @@
+// Copyright 1998-2019 Epic Games, Inc. All Rights Reserved.
+
+#include "../Common.ush"
+#include "../SceneTexturesCommon.ush"
+
+Texture2D<float> SceneDepthTexture;
+
+void MainPS(
+	in FScreenVertexOutput Input,
+	out float OutDepth: SV_Depth)
+{
+	const uint2 PixelCoord = floor(Input.Position.xy);
+	OutDepth = SceneDepthTexture.Load(uint3(PixelCoord,0));
+}
diff --git a/Engine/Shaders/Private/HairStrands/HairStrandsVisibilityPS.usf b/Engine/Shaders/Private/HairStrands/HairStrandsVisibilityPS.usf
new file mode 100644
--- /dev/null
+++ b/Engine/Shaders/Private/HairStrands/HairStrandsVisibilityPS.usf
@@ -0,0 +1,103 @@
+// Copyright 1998-2019 Epic Games, Inc. All Rights Reserved.
+
+#include "../Common.ush"
+
+//#define SceneTexturesStruct DeepShadowPass.SceneTextures
+
+#include "../SceneTexturesCommon.ush"
+#include "HairStrandsDeepShadowCommon.ush"
+#include "../DeferredShadingCommon.ush"
+#include "HairStrandsVisibilityCommon.ush"
+#include "../VelocityCommon.ush"  
+#include "/Engine/Generated/Material.ush"
+#include "/Engine/Generated/VertexFactory.ush"
+
+#define SUPPORT_OPACITY_MASK 0
+#define SUPPORT_MATERIAL_PROPERTY 1
+
+#ifndef HAIR_RENDER_MODE
+#error Undefined hair render mode
+#endif
+
+uint HairVisibilityPass_HairClusterIndex;
+
+void MainVisibility(
+	centroid FVertexFactoryInterpolantsVSToPS Interpolants,
+	in float4 SvPosition	: SV_Position,
+	float4 PixelPosition : TEXCOORD6,
+	centroid in float WorldStrandRadius : TEXCOORD8,
+	const float4 PackedVelocityA : TEXCOORD9,
+	const float4 PackedVelocityC : TEXCOORD12
+	OPTIONAL_IsFrontFace,
+	out uint4  OutColor0	: SV_Target0
+	#if SUPPORT_MATERIAL_PROPERTY
+	, out float4 OutColor1	: SV_Target1
+	, out float4 OutColor2	: SV_Target2
+	, out float2 OutColor3	: SV_Target3
+	#endif
+)
+{
+	ResolvedView = ResolveView();
+
+	#if SUPPORT_OPACITY_MASK
+	{
+		FMaterialPixelParameters MaterialParameters = GetMaterialPixelParameters(Interpolants, SvPosition);
+		FPixelMaterialInputs PixelMaterialInputs;
+		CalcMaterialParameters(MaterialParameters, PixelMaterialInputs, SvPosition, bIsFrontFace);
+		clip(GetMaterialMask(PixelMaterialInputs));
+	}
+	#endif
+
+	float Coverage = 1;
+	{
+		const float SceneDepth = ConvertFromDeviceZ(SvPosition.z); // Linear depth in world unit
+		const float PixelRadius = SceneDepth * HairVisibilityPass.MinStrandRadius_Primary; // Not correct but the coverage is not used (we count instead the number of sub-sample covered)
+		const float StrandRealRadius = WorldStrandRadius;
+		Coverage = StrandRealRadius / max(StrandRealRadius, PixelRadius);
+	}
+
+	FVisibilityData VisibilityData;
+#if VF_STRAND_HAIR
+	VisibilityData.Tangent = normalize(Interpolants.TangentToWorld2.xyz);
+	VisibilityData.PrimitiveId = Interpolants.HairPrimitiveId;
+#else
+	VisibilityData.Tangent = float3(0, 0, 1);
+	VisibilityData.PrimitiveId = 0;
+#endif
+	VisibilityData.ClusterId = HairVisibilityPass_HairClusterIndex;
+	VisibilityData.Coverage8bit = To8bitCoverage(Coverage);
+
+
+	half3 BaseColor = 0;
+	float Roughness = 0;
+	float Specular  = 0;
+	float2 EncodedVelocity = 0;
+	#if SUPPORT_MATERIAL_PROPERTY
+	{
+		FMaterialPixelParameters MaterialParameters = GetMaterialPixelParameters(Interpolants, SvPosition);
+		FPixelMaterialInputs PixelMaterialInputs;
+	
+		float4 ScreenPosition = SvPositionToResolvedScreenPosition(SvPosition);
+		float3 TranslatedWorldPosition = SvPositionToResolvedTranslatedWorld(SvPosition);
+		CalcMaterialParametersEx(MaterialParameters, PixelMaterialInputs, SvPosition, ScreenPosition, bIsFrontFace, TranslatedWorldPosition, TranslatedWorldPosition);
+
+		BaseColor		 = GetMaterialBaseColor(PixelMaterialInputs);
+		Roughness		 = GetMaterialRoughness(PixelMaterialInputs);
+		half  Metallic	 = GetMaterialMetallic(PixelMaterialInputs);
+		Specular		 = GetMaterialSpecular(PixelMaterialInputs);
+		float MaterialAO = GetMaterialAmbientOcclusion(PixelMaterialInputs);
+
+		VisibilityData.Tangent = mul(normalize(GetMaterialNormalRaw(PixelMaterialInputs)), MaterialParameters.TangentToWorld);
+	}
+	{
+		const float2 Velocity = Calculate2DVelocity(PackedVelocityA, PackedVelocityC);
+		EncodedVelocity = EncodeVelocityToTexture(Velocity);
+	}
+	#endif
+	OutColor0 = EncodeVisibilityData(VisibilityData);
+	OutColor1 = float4(sqrt(BaseColor), Roughness);
+	OutColor2 = float4(Specular,0,0,0);
+	OutColor3 = EncodedVelocity;
+}
+
+
diff --git a/Engine/Shaders/Private/HairStrands/HairStrandsVisibilityVS.usf b/Engine/Shaders/Private/HairStrands/HairStrandsVisibilityVS.usf
new file mode 100644
--- /dev/null
+++ b/Engine/Shaders/Private/HairStrands/HairStrandsVisibilityVS.usf
@@ -0,0 +1,141 @@
+// Copyright 1998-2019 Epic Games, Inc. All Rights Reserved.
+
+#include "../Common.ush"
+#include "../VelocityCommon.ush"
+#include "/Engine/Generated/Material.ush"
+#include "/Engine/Generated/VertexFactory.ush"
+
+struct FHairVisibilityAccumulateVSToPS
+{
+	FVertexFactoryInterpolantsVSToPS FactoryInterpolants;
+	float4 Position : SV_POSITION;
+	float4 PixelPosition : TEXCOORD6;
+	float WorldStrandRadius : TEXCOORD8;
+	float4 PackedVelocityA : TEXCOORD9;
+	float4 PackedVelocityC : TEXCOORD12;
+};
+
+#if USING_TESSELLATION	
+	struct FHairVisibilityAccumulateVSToDS
+	{
+		FVertexFactoryInterpolantsVSToDS FactoryInterpolants;
+		float4 Position : VS_To_DS_Position;
+		OPTIONAL_VertexID_VS_To_DS
+	};
+	
+	#define VS_OUTPUT_TYPE FHairVisibilityAccumulateVSToDS
+#else
+	#define VS_OUTPUT_TYPE FHairVisibilityAccumulateVSToPS
+#endif
+
+#if USING_TESSELLATION	
+	#define FPassSpecificVSToDS FHairVisibilityAccumulateVSToDS
+	#define FPassSpecificVSToPS FHairVisibilityAccumulateVSToPS
+	#define INTERPOLATE_MEMBER(member) O.member = a.member * aInterp + b.member * bInterp
+
+	FHairVisibilityAccumulateVSToDS PassInterpolate(FHairVisibilityAccumulateVSToDS a, float aInterp, FHairVisibilityAccumulateVSToDS b, float bInterp)
+	{
+		FHairVisibilityAccumulateVSToDS O;
+		
+		O.FactoryInterpolants = VertexFactoryInterpolate(a.FactoryInterpolants, aInterp, b.FactoryInterpolants, bInterp);
+
+		return O;
+	}
+
+	#undef INTERPOLATE_MEMBER
+	
+	FHairVisibilityAccumulateVSToPS PassFinalizeTessellationOutput(FHairVisibilityAccumulateVSToDS Interpolants, float4 WorldPosition, FMaterialTessellationParameters MaterialParameters)
+	{
+		FHairVisibilityAccumulateVSToPS O;
+		
+		O.FactoryInterpolants = VertexFactoryAssignInterpolants(Interpolants.FactoryInterpolants);
+		
+		// Finally, transform position to clip-space
+		O.Position = mul(WorldPosition, HairVisibilityPass.WorldToClipMatrix); // mul(WorldPosition, View.TranslatedWorldToClip);
+			
+		O.PixelPosition = WorldPosition;
+		
+		return O;
+	}
+	
+#include "Tessellation.ush"
+#endif
+
+#if VERTEXSHADER
+
+/** transform mesh as normal */
+void Main(
+	FVertexFactoryInput Input,
+	OPTIONAL_VertexID
+	out VS_OUTPUT_TYPE Output
+	)
+{
+	ResolvedView = ResolveView();
+
+	// #hair_todo: This is a hack/quick term solution, in order to generate valide light/view direction in the case of hair strand factory, as this is used for alighning the strand with the view.
+	//ResolvedView.TranslatedWorldCameraOrigin = -HairVisibilityPass.LightDirection * 1000; // Hack for strand rendering which use the camera position for deriving the view vector for forcing strand view facing
+	FVertexFactoryIntermediates VFIntermediates = GetVertexFactoryIntermediates(Input);
+#if VF_STRAND_HAIR
+	// Use a raw position (rather than quad facing position) as the pre-position is not expanded into quad (for perf. reason)
+	float4 WorldPositionRaw = VertexFactoryGetWorldPositionRaw(Input, VFIntermediates);
+#else
+	float4 WorldPositionRaw = VertexFactoryGetWorldPosition(Input, VFIntermediates);
+#endif
+	float4 PrevWorldPosition = VertexFactoryGetPreviousWorldPosition(Input, VFIntermediates);
+	
+	// Velocity computation
+	// The velocity is used to adapt/increase the minimum rasterization size in order to avoid aliasing under heavy-motion.
+	// The higher the velocity, the larger a strand becomes.
+	float VelocityScale = 0;
+	{
+		const float4 ScreenPos			= mul(float4(WorldPositionRaw.xyz, 1), ResolvedView.TranslatedWorldToClip);
+		const float4 PrevScreenPosObj	= mul(float4(PrevWorldPosition.xyz, 1), ResolvedView.PrevTranslatedWorldToClip);
+		Output.PackedVelocityA			= INVARIANT(ScreenPos);
+		Output.PackedVelocityC			= INVARIANT(PrevScreenPosObj);
+
+		const float2 ScreenVelocity = Calculate2DVelocity(Output.PackedVelocityA, Output.PackedVelocityC);
+		VelocityScale = saturate(max(abs(ScreenVelocity.x), abs(ScreenVelocity.y)) / HairVisibilityPass.HairStrandsVelocityScale);
+	}
+
+#if VF_STRAND_HAIR
+	float4 WorldPosition = VertexFactoryGetWorldPosition(Input, VFIntermediates, VelocityScale);
+#else
+	float4 WorldPosition = VertexFactoryGetWorldPosition(Input, VFIntermediates);
+#endif
+
+	float3x3 TangentToLocal = VertexFactoryGetTangentToLocal(Input, VFIntermediates);	
+
+	FMaterialVertexParameters VertexParameters = GetMaterialVertexParameters(Input, VFIntermediates, WorldPosition.xyz, TangentToLocal);
+	// Isolate instructions used for world position offset on xbox 360, 
+	// As these cause the optimizer to generate different position calculating instructions in each pass, resulting in self-z-fighting.
+	// This is only necessary for shaders used in passes that have depth testing enabled.
+	//ISOLATE
+	//{
+	//	WorldPosition.xyz += GetMaterialWorldPositionOffset(VertexParameters);
+	//	WorldPositionRaw.xyz += GetMaterialWorldPositionOffset(VertexParameters);
+	//	PrevWorldPosition.xyz += GetMaterialPreviousWorldPositionOffset(VertexParameters); 
+	//}
+
+#if USING_TESSELLATION
+	Output.Position = WorldPosition;
+	Output.FactoryInterpolants = VertexFactoryGetInterpolantsVSToDS(Input, VFIntermediates, VertexParameters);
+	
+#else // !USING_TESSELLATION
+
+//	Output.Position = mul(WorldPosition, HairVisibilityPass.WorldToClipMatrix);
+	Output.Position = mul(WorldPosition, View.TranslatedWorldToClip);
+	Output.FactoryInterpolants = VertexFactoryGetInterpolantsVSToPS(Input, VFIntermediates, VertexParameters);
+	Output.PixelPosition = WorldPosition;
+
+#endif // USING_TESSELLATION
+
+#ifdef VF_STRAND_HAIR
+	Output.WorldStrandRadius = VFIntermediates.HairDimensions.y;
+#else
+	Output.WorldStrandRadius = 1;
+#endif
+	
+	OutputVertexID( Output );
+}
+
+#endif // VERTEXSHADER
\ No newline at end of file
diff --git a/Engine/Shaders/Private/HairStrands/HairStrandsVoxelMip.usf b/Engine/Shaders/Private/HairStrands/HairStrandsVoxelMip.usf
new file mode 100644
--- /dev/null
+++ b/Engine/Shaders/Private/HairStrands/HairStrandsVoxelMip.usf
@@ -0,0 +1,38 @@
+// Copyright 1998-2019 Epic Games, Inc. All Rights Reserved.
+
+uint VoxelResolution;
+uint SourceMip;
+uint TargetMip;
+
+Texture3D<uint> InDensityTexture;
+RWTexture3D<uint> OutDensityTexture;
+
+// #hair_todo: do two mips at a time through LDS?
+[numthreads(4, 4, 4)]
+void MainCS(uint3 DispatchThreadId : SV_DispatchThreadID)
+{
+	const uint3 OutVoxelCoord = DispatchThreadId;
+
+	const uint3 InVoxelCoord0 = uint3(OutVoxelCoord.x << 1, OutVoxelCoord.y << 1, OutVoxelCoord.z << 1);
+	const uint3 InVoxelCoord1 = InVoxelCoord0 + uint3(1,0,0);
+	const uint3 InVoxelCoord2 = InVoxelCoord0 + uint3(0,1,0);
+	const uint3 InVoxelCoord3 = InVoxelCoord0 + uint3(1,1,0);
+
+	const uint3 InVoxelCoord4 = InVoxelCoord0 + uint3(0,0,1);
+	const uint3 InVoxelCoord5 = InVoxelCoord0 + uint3(1,0,1);
+	const uint3 InVoxelCoord6 = InVoxelCoord0 + uint3(0,1,1);
+	const uint3 InVoxelCoord7 = InVoxelCoord0 + uint3(1,1,1);
+
+	const uint Total =
+		InDensityTexture[InVoxelCoord0] +
+		InDensityTexture[InVoxelCoord1] +
+		InDensityTexture[InVoxelCoord2] +
+		InDensityTexture[InVoxelCoord3] +
+
+		InDensityTexture[InVoxelCoord4] +
+		InDensityTexture[InVoxelCoord5] +
+		InDensityTexture[InVoxelCoord6] +
+		InDensityTexture[InVoxelCoord7];
+
+	OutDensityTexture[OutVoxelCoord] = Total;
+}
diff --git a/Engine/Shaders/Private/HairStrands/HairStrandsVoxelOpaque.usf b/Engine/Shaders/Private/HairStrands/HairStrandsVoxelOpaque.usf
new file mode 100644
--- /dev/null
+++ b/Engine/Shaders/Private/HairStrands/HairStrandsVoxelOpaque.usf
@@ -0,0 +1,75 @@
+// Copyright 1998-2019 Epic Games, Inc. All Rights Reserved.
+
+#include "../Common.ush"
+#include "../CommonViewUniformBuffer.ush"
+#include "../SceneTextureParameters.ush"
+
+
+float3 VoxelMinAABB;
+uint   VoxelResolution;
+float3 VoxelMaxAABB;
+float2 OutputResolution;
+float2 SceneDepthResolution;
+
+Texture2D<float> SceneDepthTexture;
+RWTexture3D<uint> DensityTexture;
+SamplerState LinearSampler;
+
+float3 ReconstructWorldPositionFromDepth(float2 UV, float Depth)
+{
+	float2 ScreenPosition = (UV - View.ScreenPositionScaleBias.wz) / View.ScreenPositionScaleBias.xy;
+	float4 HomogeneousWorldPosition = mul(float4(ScreenPosition * Depth, Depth, 1), View.ScreenToWorld);
+	float3 WorldPosition = HomogeneousWorldPosition.xyz / HomogeneousWorldPosition.w;
+
+	return WorldPosition;
+}
+
+// hair_todo: change to CS & optmize (naive code)
+void MainPS(
+	FScreenVertexOutput Input,
+	out float4  OutColor0 : SV_Target0)
+{
+	const float2 UV = Input.UV;
+	const float2 PixelCoord = UV * OutputResolution;
+
+	const float VoxelSize = (VoxelMaxAABB.x - VoxelMinAABB.x) / VoxelResolution;
+
+	uint3 VoxelPos = 0;
+	uint2 TileCount = OutputResolution / VoxelResolution;
+	VoxelPos.x = PixelCoord.x % VoxelResolution;
+	VoxelPos.y = PixelCoord.y % VoxelResolution;
+	VoxelPos.z = PixelCoord.x / VoxelResolution + PixelCoord.y / VoxelResolution * TileCount.x;
+
+	const float3 WorldPosition = VoxelPos / float(VoxelResolution) * (VoxelMaxAABB - VoxelMinAABB) + VoxelMinAABB;
+
+	float4 ClipPos = mul(float4(WorldPosition,1), View.WorldToClip);
+	ClipPos /= ClipPos.w;
+	const float DepthBias 			= 3 * VoxelSize;
+	const float  VoxelDepth 		= ConvertFromDeviceZ(ClipPos.z) - DepthBias;
+	float2 SceneUV 					= float2(0.5f*(ClipPos.x+1), 1-0.5f*(ClipPos.y+1));
+	SceneUV							= ViewportUVToBufferUV(SceneUV);
+	const float2 ScenePixelCoord 	= SceneUV * SceneDepthResolution;
+
+	OutColor0 = float4(VoxelPos.xy / float(VoxelResolution), 0, 1); 
+	const bool bIsOnScreen = SceneUV.x >= 0 && SceneUV.x < 1 && SceneUV.y >= 0 && SceneUV.y < 1;
+	if (!bIsOnScreen)
+		return;
+
+	const float ClosestDepth = ConvertFromDeviceZ(SceneDepthTexture.Load(uint3(ScenePixelCoord,0)));
+	const float3 SceneWorldPos = ReconstructWorldPositionFromDepth(SceneUV, ClosestDepth);
+
+	const bool bIsInVolume = 
+		SceneWorldPos.x >= VoxelMinAABB.x && SceneWorldPos.x < VoxelMaxAABB.x && 
+		SceneWorldPos.y >= VoxelMinAABB.y && SceneWorldPos.y < VoxelMaxAABB.y && 
+		SceneWorldPos.z >= VoxelMinAABB.z && SceneWorldPos.z < VoxelMaxAABB.z;
+
+	if (!bIsInVolume)
+		return;
+
+	// Inject opaque depth on a thin layer (Dist < DistThreshold) for avoiding weird projection
+	if (ClosestDepth < VoxelDepth && abs(ClosestDepth - VoxelDepth) < DepthBias)
+	{
+		OutColor0 = float4(1,0,0,1);
+		DensityTexture[VoxelPos] = 100000;
+	}
+}
diff --git a/Engine/Shaders/Private/HairStrands/HairStrandsVoxelRayMarching.usf b/Engine/Shaders/Private/HairStrands/HairStrandsVoxelRayMarching.usf
new file mode 100644
--- /dev/null
+++ b/Engine/Shaders/Private/HairStrands/HairStrandsVoxelRayMarching.usf
@@ -0,0 +1,165 @@
+// Copyright 1998-2019 Epic Games, Inc. All Rights Reserved.
+
+#include "../Common.ush"
+#include "../CommonViewUniformBuffer.ush"
+#include "../SceneTextureParameters.ush"
+#include "../DeferredShadingCommon.ush"
+#include "HairStrandsVisibilityCommon.ush"
+#include "HairStrandsDeepVoxelCommon.ush"
+
+float3 VoxelMinAABB;
+uint   VoxelResolution;
+float3 VoxelMaxAABB;
+float  VoxelDensityScale;
+float  DensityIsoline;
+float2 OutputResolution;
+
+Texture3D<uint> DensityTexture;
+Texture3D<uint> TangentXTexture;
+Texture3D<uint> TangentYTexture;
+Texture3D<uint> TangentZTexture;
+Texture3D<uint> MaterialTexture;
+
+SamplerState LinearSampler;
+
+
+float3 HUEtoRGB(in float H)
+{
+	float R = abs(H * 6 - 3) - 1;
+	float G = 2 - abs(H * 6 - 2);
+	float B = 2 - abs(H * 6 - 4);
+	return saturate(float3(R, G, B));
+}
+
+float3 HSVtoRGB(in float3 HSV)
+{
+	float3 RGB = HUEtoRGB(HSV.x);
+	return ((RGB - 1) * HSV.y + 1) * HSV.z;
+}
+
+float3 GetColorCode(float x)
+{
+	float c = (1 - saturate(x)) * 0.6; // Remap [0,1] to Blue-Red
+	return x > 0 ? HSVtoRGB(float3(c, 1, 1)) : float3(0, 0, 0);
+}
+
+
+float3 ReconstructWorldPositionFromDepth(float2 UV, float Depth)
+{
+	float2 ScreenPosition = (UV - View.ScreenPositionScaleBias.wz) / View.ScreenPositionScaleBias.xy;
+	float4 HomogeneousWorldPosition = mul(float4(ScreenPosition * Depth, Depth, 1), View.ScreenToWorld);
+	float3 WorldPosition = HomogeneousWorldPosition.xyz / HomogeneousWorldPosition.w;
+
+	return WorldPosition;
+}
+
+// Config
+#define DEBUG_MODE_DENSITY 0
+#define DEBUG_MODE_TANGENT 1
+#define DEBUG_MODE_BASECOLOR 2
+#define DEBUG_MODE_ROUGHNESS 3
+
+#define VOXEL_BLOCK 0
+#define VOXEL_SMOOTH 1
+#define VOXEL_FILTERING VOXEL_BLOCK
+
+void MainPS(
+	FScreenVertexOutput Input,
+	out float4  OutColor0 : SV_Target0)
+{
+	const float2 UV = Input.UV;
+	const float2 PixelCoord = UV * OutputResolution;
+	
+	// Reverse-Z
+	const float3 O = View.WorldCameraOrigin;
+	const float3 E = ReconstructWorldPositionFromDepth(UV, ConvertFromDeviceZ(0.0001f));
+	const float3 D = (E - O);
+
+	float  WorldDepth = 0;
+	float  ClipDepth = 0;
+
+	const float2 HitT = LineBoxIntersect(O, E, VoxelMinAABB, VoxelMaxAABB);
+	float3 HitColor = float3(0,0,0);
+	float AccDensity = 0;
+	bool bIsValid = false;
+	if (HitT.x < HitT.y)
+	{
+		// Use different offseting to see either: 
+		// * Individual voxel (for debugging mainly)
+		// * Smoothed voxel (as used for transmittance computation)
+	#if VOXEL_FILTERING == VOXEL_BLOCK
+		const float3 Offset = 0;
+		const float Dither = InterleavedGradientNoise(PixelCoord.xy, 1);
+	#endif
+	#if VOXEL_FILTERING == VOXEL_SMOOTH
+		const float3 Offset = float3(InterleavedGradientNoise(PixelCoord.xy, 1), InterleavedGradientNoise(PixelCoord.xy, 2), InterleavedGradientNoise(PixelCoord.xy, 3));
+		const float Dither = 0;
+	#endif
+
+		const uint MipLevel = 0; // clamp(View.GeneralPurposeTweak, 0, 6);
+		const uint MipResolution = VoxelResolution >> MipLevel;
+		const uint MaxStep = 256;
+		const float Delta = (HitT.y-HitT.x)/float(MaxStep);
+		for (uint StepIt=0; StepIt<MaxStep; ++StepIt)
+		{			
+			const float3 HitP = Offset + O + ((Dither+StepIt)*Delta + HitT.x) * D;
+			const float3 UVW  = (HitP - VoxelMinAABB) / (VoxelMaxAABB - VoxelMinAABB);
+			const int3 VolumeCoord = clamp(int3(UVW * MipResolution), int3(0,0,0), int3(MipResolution.xxx-1));
+			const float Density = GetHairVoxelDensity(VolumeCoord, DensityTexture, MipLevel) * VoxelDensityScale;
+		#if PERMUTATION_DEBUG_MODE == DEBUG_MODE_DENSITY
+			if (Density > 0)
+			{
+				AccDensity = max(AccDensity, float(Density));
+				bIsValid = true;
+			}
+		#endif	
+
+		#if PERMUTATION_DEBUG_MODE == DEBUG_MODE_TANGENT
+			if (Density > 0)
+			{
+				const float3 Tangent = GetHairVoxelTangent(VolumeCoord, Density, TangentXTexture, TangentYTexture, TangentZTexture);
+				HitColor = float3((Tangent + 1)*0.5f);
+				bIsValid = true;
+				break;
+			}
+		#endif
+
+		#if PERMUTATION_DEBUG_MODE == DEBUG_MODE_BASECOLOR
+			if (Density > 0)
+			{
+				float3 BaseColor = 0;
+				float Roughness = 0;
+				const uint EncodedBaseColorAndRoughness = MaterialTexture.Load(int4(VolumeCoord, 0));
+				DecodeBaseColorAndRoughness(EncodedBaseColorAndRoughness, BaseColor, Roughness);
+				HitColor = BaseColor;
+				bIsValid = true;
+				break;
+			}
+		#endif
+
+		#if PERMUTATION_DEBUG_MODE == DEBUG_MODE_ROUGHNESS
+			if (Density > 0)
+			{
+				float3 BaseColor = 0;
+				float Roughness = 0;
+				const uint EncodedBaseColorAndRoughness = MaterialTexture.Load(int4(VolumeCoord, 0));
+				DecodeBaseColorAndRoughness(EncodedBaseColorAndRoughness, BaseColor, Roughness);
+				HitColor = Roughness.xxx;
+				bIsValid = true;
+				break;
+			}
+		#endif
+		}
+	}
+
+	if (!bIsValid)
+	{ 
+		discard;
+	}
+
+	#if PERMUTATION_DEBUG_MODE == DEBUG_MODE_DENSITY
+	HitColor = GetColorCode(saturate(AccDensity));
+	#endif
+
+	OutColor0 = float4(HitColor,1);
+}
diff --git a/Engine/Shaders/Private/LocalVertexFactoryCommon.ush b/Engine/Shaders/Private/LocalVertexFactoryCommon.ush
--- a/Engine/Shaders/Private/LocalVertexFactoryCommon.ush
+++ b/Engine/Shaders/Private/LocalVertexFactoryCommon.ush
@@ -40,6 +40,13 @@ struct FVertexFactoryInterpolantsVSToPS
 		nointerpolation uint LightmapDataIndex : LIGHTMAP_ID;
 	#endif
 #endif
+#if VF_STRAND_HAIR
+	nointerpolation uint HairPrimitiveId	: HAIR_PRIMITIVE_ID;
+	float2 HairDimensions					: HAIR_DIMENSIONS;
+	float2 HairUV							: HAIR_UV;
+	float2 HairRootUV						: HAIR_ROOT_UV;
+	float  HairSeed							: HAIR_SEED;
+#endif
 };
 
 #if NUM_TEX_COORD_INTERPOLATORS || USE_PARTICLE_SUBUVS
diff --git a/Engine/Shaders/Private/MaterialTemplate.ush b/Engine/Shaders/Private/MaterialTemplate.ush
--- a/Engine/Shaders/Private/MaterialTemplate.ush
+++ b/Engine/Shaders/Private/MaterialTemplate.ush
@@ -287,6 +287,15 @@ struct FMaterialPixelParameters
 	// Index into View.PrimitiveSceneData
 	uint PrimitiveId;
 
+	// Actual primitive Id
+#if VF_STRAND_HAIR
+	uint HairPrimitiveId;
+	float2 HairUV;			// U: parametric coord along the length V: parametric coord along the width)
+	float2 HairDimensions;  // X: Length, Y: Radius
+	float  HairSeed;		// [0,1] quantized on 8bits
+	float2 HairRootUV;		// [0,1]^2	8bits
+#endif
+
 	/** Per-particle properties. Only valid for particle vertex factories. */
 	FMaterialParticleParameters Particle;
 
@@ -990,6 +999,47 @@ float3 MaterialExpressionBlackBody( float Temp )
 	return mul( XYZtoRGB, float3( X, Y, Z ) ) * pow( 0.0004 * Temp, 4 );
 }
 
+float2 MaterialExpressionGetHairRootUV(FMaterialPixelParameters Parameters)
+{
+#if VF_STRAND_HAIR
+	return Parameters.HairRootUV;
+#else
+	return float2(0, 0);
+#endif
+}
+
+float2 MaterialExpressionGetHairUV(FMaterialPixelParameters Parameters)
+{	
+#if VF_STRAND_HAIR
+	return Parameters.HairUV;
+#else
+	return float2(0,0);
+#endif
+}
+
+float2 MaterialExpressionGetHairDimensions(FMaterialPixelParameters Parameters)
+{
+#if VF_STRAND_HAIR
+	return Parameters.HairDimensions; // Length | Radius
+#else
+	return float2(0, 0);
+#endif
+}
+
+float MaterialExpressionGetHairSeed(FMaterialPixelParameters Parameters)
+{
+#if VF_STRAND_HAIR
+	return Parameters.HairSeed;
+#else
+	return 0;
+#endif
+}
+
+float3 MaterialExpressionGetHairTangent(FMaterialPixelParameters Parameters)
+{
+	return Parameters.TangentToWorld[2];
+}
+
 float4 MaterialExpressionAtmosphericFog(FMaterialPixelParameters Parameters, float3 AbsoluteWorldPosition)
 {
 #if MATERIAL_ATMOSPHERIC_FOG
diff --git a/Engine/Shaders/Private/RayTracing/RayTracingAmbientOcclusionRGS.usf b/Engine/Shaders/Private/RayTracing/RayTracingAmbientOcclusionRGS.usf
--- a/Engine/Shaders/Private/RayTracing/RayTracingAmbientOcclusionRGS.usf
+++ b/Engine/Shaders/Private/RayTracing/RayTracingAmbientOcclusionRGS.usf
@@ -95,11 +95,12 @@ void AmbientOcclusionRGS()
 
 		FPackedMaterialClosestHitPayload Payload = (FPackedMaterialClosestHitPayload)0;
 		Payload.SetMinimalPayloadMode(); // Only HitT will be filled by closest hit shader, skipping full material evaluation.
-
+		
+		const uint InstanceInclusionMask = RAY_TRACING_MASK_SHADOW | RAY_TRACING_MASK_THIN_SHADOW;
 		TraceRay(
 			TLAS,
 			RayFlags,
-			RAY_TRACING_MASK_SHADOW,          // InstanceInclusionMask
+			InstanceInclusionMask,
 			RAY_TRACING_SHADER_SLOT_SHADOW,   // RayContributionToHitGroupIndex
 			RAY_TRACING_NUM_SHADER_SLOTS,     // MultiplierForGeometryContributionToShaderIndex
 			0,    // MissShaderIndex
diff --git a/Engine/Shaders/Private/RayTracing/RayTracingCreateGatherPointsRGS.usf b/Engine/Shaders/Private/RayTracing/RayTracingCreateGatherPointsRGS.usf
--- a/Engine/Shaders/Private/RayTracing/RayTracingCreateGatherPointsRGS.usf
+++ b/Engine/Shaders/Private/RayTracing/RayTracingCreateGatherPointsRGS.usf
@@ -236,7 +236,7 @@ void RayTracingCreateGatherPointsRGS()
 						half3 V = -Ray.Direction;
 						half3 L = LightRay.Direction;
 						float NoL = saturate(dot(N, L));
-						FShadowTerms ShadowTerms = { 0.0, 0.0, 0.0 };
+						FShadowTerms ShadowTerms = { 0.0, 0.0, 0.0, InitHairTransmittanceData() };
 						FDirectLighting LightingSample = EvaluateBxDF(GBufferData, N, V, L, NoL, ShadowTerms);
 						float3 MaterialThroughput = LightingSample.Diffuse;
 
diff --git a/Engine/Shaders/Private/RayTracing/RayTracingFinalGatherRGS.usf b/Engine/Shaders/Private/RayTracing/RayTracingFinalGatherRGS.usf
--- a/Engine/Shaders/Private/RayTracing/RayTracingFinalGatherRGS.usf
+++ b/Engine/Shaders/Private/RayTracing/RayTracingFinalGatherRGS.usf
@@ -112,7 +112,7 @@ void RayTracingFinalGatherRGS()
 		half3 V = -CameraDirection;
 		half3 L = normalize(GPPosition - WorldPosition);
 		float NoL = saturate(dot(N, L));
-		FShadowTerms ShadowTerms = { 0.0, 0.0, 0.0 };
+		FShadowTerms ShadowTerms = { 0.0, 0.0, 0.0, InitHairTransmittanceData() };
 		FDirectLighting LightingSample = EvaluateBxDF(GBufferData, N, V, L, NoL, ShadowTerms);
 		float3 RayThroughput = LightingSample.Diffuse / DiffuseColor;
 		float RayPdf = NoL / PI;
diff --git a/Engine/Shaders/Private/RayTracing/RayTracingGlobalIlluminationRGS.usf b/Engine/Shaders/Private/RayTracing/RayTracingGlobalIlluminationRGS.usf
--- a/Engine/Shaders/Private/RayTracing/RayTracingGlobalIlluminationRGS.usf
+++ b/Engine/Shaders/Private/RayTracing/RayTracingGlobalIlluminationRGS.usf
@@ -129,7 +129,7 @@ void GlobalIlluminationRGS()
 		half3 V = -CameraDirection;
 		half3 L = Ray.Direction;
 		float NoL = saturate(dot(N, L));
-		FShadowTerms ShadowTerms = { 0.0, 0.0, 0.0 };
+		FShadowTerms ShadowTerms = { 0.0, 0.0, 0.0, InitHairTransmittanceData() };
 		FDirectLighting LightingSample = EvaluateBxDF(GBufferData, N, V, L, NoL, ShadowTerms);
 		RayThroughput *= LightingSample.Diffuse / DiffuseColor;
 #else
@@ -264,7 +264,7 @@ void GlobalIlluminationRGS()
 						half3 V = -Ray.Direction;
 						half3 L = LightRay.Direction;
 						float NoL = saturate(dot(N, L));
-						FShadowTerms ShadowTerms = { 0.0, 0.0, 0.0 };
+						FShadowTerms ShadowTerms = { 0.0, 0.0, 0.0, InitHairTransmittanceData() };
 						FDirectLighting LightingSample = EvaluateBxDF(GBufferData, N, V, L, NoL, ShadowTerms);
 						MaterialThroughput = LightingSample.Diffuse;
 						MaterialEvalPdf = 1.0;
@@ -304,7 +304,7 @@ void GlobalIlluminationRGS()
 				half3 V = -Ray.Direction;
 				half3 L = Direction;
 				float NoL = saturate(dot(N, L));
-				FShadowTerms ShadowTerms = { 0.0, 0.0, 0.0 };
+				FShadowTerms ShadowTerms = { 0.0, 0.0, 0.0, InitHairTransmittanceData() };
 				FDirectLighting LightingSample = EvaluateBxDF(GBufferData, N, V, L, NoL, ShadowTerms);
 				Throughput = LightingSample.Diffuse;
 #endif
diff --git a/Engine/Shaders/Private/RayTracing/RayTracingOcclusionRGS.usf b/Engine/Shaders/Private/RayTracing/RayTracingOcclusionRGS.usf
--- a/Engine/Shaders/Private/RayTracing/RayTracingOcclusionRGS.usf
+++ b/Engine/Shaders/Private/RayTracing/RayTracingOcclusionRGS.usf
@@ -7,6 +7,7 @@
 #include "../SceneTextureParameters.ush"
 #include "../ScreenSpaceDenoise/SSDPublic.ush"
 #include "../SobolRandom.ush"
+#include "../HairStrands/HairStrandsVisibilityCommon.ush"
 #include "RayTracingCommon.ush"
 #include "RayTracingDirectionalLight.ush"
 #include "RayTracingRectLight.ush"
@@ -20,10 +21,15 @@
 #define LIGHT_TYPE_SPOT 2
 #define LIGHT_TYPE_RECT 3
 
+#define ShadowMaskType_Opaque 0
+#define ShadowMaskType_Hair 1
+
 RaytracingAccelerationStructure TLAS;
 RWTexture2D<float4> RWOcclusionMaskUAV;
 RWTexture2D<float> RWRayDistanceUAV;
+Texture2D<uint4> HairCategorizationTexture;
 
+uint ShadowMaskType;
 uint LightingChannelMask;
 uint SamplesPerPixel;
 float NormalBias;
@@ -114,7 +120,7 @@ bool GenerateOcclusionRay(
 [shader("raygeneration")]
 void OcclusionRGS()
 {
-	uint2 PixelCoord = DispatchRaysIndex().xy + View.ViewRectMin + PixelOffset;
+	uint2 PixelCoord = DispatchRaysIndex().xy + View.ViewRectMin.xy + PixelOffset;
 
 	float Visibility = 0.0;
 	float SumRayDistance = 0.0;
@@ -134,27 +140,52 @@ void OcclusionRGS()
 		// Get G-Buffer surface data
 		float2 InvBufferSize = View.BufferSizeAndInvSize.zw;
 		float2 BufferUV = (float2(PixelCoord) + 0.5) * InvBufferSize;
-		FGBufferData GBufferData = GetGBufferDataFromSceneTextures(BufferUV);
-		float DeviceZ = SceneDepthBuffer.SampleLevel(SceneDepthBufferSampler, BufferUV, 0).r;
-		float3 WorldPosition = ReconstructWorldPositionFromDeviceZ(PixelCoord, DeviceZ);
-		float3 WorldNormal = GBufferData.WorldNormal;
+		uint ShadingModelID = SHADINGMODELID_UNLIT;
+		float3 WorldNormal = 0;
+		{
+			const FGBufferData GBufferData = GetGBufferDataFromSceneTextures(BufferUV);
+			ShadingModelID = GBufferData.ShadingModelID;
+			WorldNormal = GBufferData.WorldNormal;
+		}
 
-		uint LightChannel = GetSceneLightingChannel(PixelCoord);
 
 		// Mask out depth values that are infinitely far away
-		bool IsFiniteDepth = DeviceZ > 0.0;
-		bool bTraceRay = (
-			IsFiniteDepth &&
-			GBufferData.ShadingModelID != SHADINGMODELID_UNLIT &&
-			(LightChannel & LightingChannelMask) != 0);
 
+		float DeviceZ = SceneDepthBuffer.SampleLevel(SceneDepthBufferSampler, BufferUV, 0).r;
+		const bool bIsDepthValid = SceneDepthBuffer.Load(int3(PixelCoord, 0)).r > 0.0;
+		bool bIsValidPixel = ShadingModelID != SHADINGMODELID_UNLIT && bIsDepthValid;
+
+		// Check if there is hair geometry
+		uint RaytracingMask = 0;
+		if (bIsValidPixel && ShadowMaskType == ShadowMaskType_Opaque)
+		{
+			RaytracingMask = RAY_TRACING_MASK_SHADOW | RAY_TRACING_MASK_THIN_SHADOW;
+		}
+	#if USE_HAIR_LIGHTING
+		else if (ShadowMaskType == ShadowMaskType_Hair)
+		{
+			RaytracingMask = RAY_TRACING_MASK_SHADOW;
+			const FCategorizationData Categorization = DecodeCategorizationData(HairCategorizationTexture.Load(uint3(PixelCoord, 0)));
+			bIsValidPixel = Categorization.TotalCoverage8bit > 0;
+			if (bIsValidPixel)
+			{
+				DeviceZ = Categorization.ClosestDepth;
+				ShadingModelID = SHADINGMODELID_HAIR;
+			}
+		}
+	#endif
+
+		const float3 WorldPosition = ReconstructWorldPositionFromDeviceZ(PixelCoord, DeviceZ);
+
+		const uint LightChannel = GetSceneLightingChannel(PixelCoord);
+		const bool bTraceRay = bIsValidPixel && (LightChannel & LightingChannelMask) != 0;
 		if (!bTraceRay)
 		{
 			LocalSamplesPerPixel = 0.0;
 		}
 
 		// Cheap solution for double sided foliage to have valid lighing through the denoiser
-		if (GBufferData.ShadingModelID == SHADINGMODELID_TWOSIDED_FOLIAGE)
+		if (ShadingModelID == SHADINGMODELID_TWOSIDED_FOLIAGE)
 		{
 		#if LIGHT_TYPE == LIGHT_TYPE_DIRECTIONAL
 			const float3 LightDirection = LightParameters.Direction;
@@ -178,16 +209,14 @@ void OcclusionRGS()
 		// It seems having proper transmission is more important, thus the light 
 		// version is enabled by default
 		bool bApplyNormalCulling = true;
-		if (GBufferData.ShadingModelID == SHADINGMODELID_HAIR)
+		if (ShadingModelID == SHADINGMODELID_HAIR)
 		{
 #if LIGHT_TYPE == LIGHT_TYPE_DIRECTIONAL
 			const float3 LightDirection = LightParameters.Direction;
 #else
 			const float3 LightDirection = normalize(LightParameters.Position - WorldPosition);
 #endif
-
-			const float3 TangentDirection = normalize(WorldNormal);
-			WorldNormal = normalize(LightDirection - TangentDirection * dot(LightDirection, TangentDirection));
+			WorldNormal = LightDirection;
 			bApplyNormalCulling = false;
 		}
 
@@ -259,7 +288,7 @@ void OcclusionRGS()
 			TraceRay(
 				TLAS,      // AccelerationStructure
 				RayFlags,
-				RAY_TRACING_MASK_SHADOW,        // InstanceInclusionMask
+				RaytracingMask,                 // InstanceInclusionMask
 				RAY_TRACING_SHADER_SLOT_SHADOW, // RayContributionToHitGroupIndex
 				RAY_TRACING_NUM_SHADER_SLOTS,   // MultiplierForGeometryContributionToShaderIndex
 				0,         // MissShaderIndex
diff --git a/Engine/Shaders/Private/RayTracing/RayTracingSkyLightCommon.ush b/Engine/Shaders/Private/RayTracing/RayTracingSkyLightCommon.ush
--- a/Engine/Shaders/Private/RayTracing/RayTracingSkyLightCommon.ush
+++ b/Engine/Shaders/Private/RayTracing/RayTracingSkyLightCommon.ush
@@ -11,8 +11,8 @@ float PdfSkyLightOcclusionRay(
 {
 	// Determine primary direction
 	uint MipCount = log2(SkyLight.MipDimensions.x);
-	uint3 TextureCoord = GetTextureCubeCoordinate(WorldDirection, SkyLight.TextureDimensions);
-	float2 TextureUV = TextureCoord.xy / (float2)SkyLight.TextureDimensions;
+	uint3 TextureCoord = GetTextureCubeCoordinate(WorldDirection, SkyLight.TextureDimensions.xy);
+	float2 TextureUV = TextureCoord.xy / (float2)SkyLight.TextureDimensions.xy;
 	uint2 MipCoord = TextureUV * SkyLight.MipDimensions.xy;
 	float FacePdf = PdfFace(FaceCdf, TextureCoord.z);
 
@@ -35,7 +35,7 @@ float PdfSkyLightOcclusionRay(
 		// Build Cdf
 		uint MipIndex = MipCount - Index - 1;
 #if USE_MORTON_CODE
-		uint BufferOffset = BufferOffsetAtPixel(MipPixel, MipIndex, SkyLight.MipDimensions);
+		uint BufferOffset = BufferOffsetAtPixel(MipPixel, MipIndex, SkyLight.MipDimensions.xy);
 		uint4 BufferOffsets = uint4(BufferOffset, BufferOffset + 1, BufferOffset + 2, BufferOffset + 3);
 #else
 		uint4 BufferOffsets;
@@ -66,7 +66,7 @@ float PdfSkyLightOcclusionRay(
 	uint LinearIndex = BufferOffsetAtPixel(MipCoord, StopLevel, SkyLight.MipDimensions.xy);
 	float AreaPdf = SkyLight.SolidAnglePdf[LinearIndex];
 #else
-	float AreaPdf = 1.0 / TexelCoordSolidAngle(MipCoord, DimensionsAtLevel(SkyLight.MipDimensions, StopLevel));
+	float AreaPdf = 1.0 / TexelCoordSolidAngle(MipCoord, DimensionsAtLevel(SkyLight.MipDimensions.xy, StopLevel));
 #endif
 	return FacePdf * MipPdf * AreaPdf;
 #endif
@@ -135,20 +135,20 @@ bool GenerateSkyLightOcclusionRay(
 	MipPixel = uint2(MipPixel.x >> 1, MipPixel.y >> 1);
 
 	float2 PixelOffset = RandSample.yz;
-	float2 MipPixelUV = (MipPixel + PixelOffset) / float2(DimensionsAtLevel(SkyLight.MipDimensions, StopLevel));
+	float2 MipPixelUV = (MipPixel + PixelOffset) / float2(DimensionsAtLevel(SkyLight.MipDimensions.xy, StopLevel));
 
 	// Construct Cubemap coordinate
-	float3 TexelCoord = float3(MipPixelUV * SkyLight.TextureDimensions, FaceIndex);
+	float3 TexelCoord = float3(MipPixelUV * SkyLight.TextureDimensions.xy, FaceIndex);
 
 	// Determine area pdf
 #if USE_SOLID_ANGLE_TEXTURE
 	uint LinearIndex = BufferOffsetAtPixel(MipPixel, StopLevel, SkyLight.MipDimensions.xy);
 	float AreaPdf = SkyLight.SolidAnglePdf[LinearIndex];
 #else
-	float AreaPdf = 1.0 / TexelCoordSolidAngle(MipPixel, DimensionsAtLevel(SkyLight.MipDimensions, StopLevel));
+	float AreaPdf = 1.0 / TexelCoordSolidAngle(MipPixel, DimensionsAtLevel(SkyLight.MipDimensions.xy, StopLevel));
 #endif
 
-	RayDirection = normalize(GetTextureCubeVector(TexelCoord, SkyLight.TextureDimensions));
+	RayDirection = normalize(GetTextureCubeVector(TexelCoord, SkyLight.TextureDimensions.xy));
 	RayPdf = FacePdf * MipPdf * AreaPdf;
 
 	return true;
diff --git a/Engine/Shaders/Private/RayTracing/RayTracingSkyLightRGS.usf b/Engine/Shaders/Private/RayTracing/RayTracingSkyLightRGS.usf
--- a/Engine/Shaders/Private/RayTracing/RayTracingSkyLightRGS.usf
+++ b/Engine/Shaders/Private/RayTracing/RayTracingSkyLightRGS.usf
@@ -34,7 +34,7 @@ float3 EvalSkyLight(float3 Direction)
 [shader("raygeneration")]
 void SkyLightRGS()
 {
-	uint2 PixelCoord = DispatchRaysIndex().xy + View.ViewRectMin;
+	uint2 PixelCoord = DispatchRaysIndex().xy + View.ViewRectMin.xy;
 	uint LinearIndex = CalcLinearIndex(PixelCoord);
 
 	// Get G-Buffer surface data
@@ -182,13 +182,14 @@ void SkyLightRGS()
 				float Area = 0.0;
 				uint2 Random = 0;
 				bool bEvalMultiScatter = false;
-				LightingSample.Diffuse = HairShading(GBufferData, L, V, N, Shadow, Backlit, Area, Random, bEvalMultiScatter);
+				FHairTransmittanceData HairTransmittance = InitHairTransmittanceData();
+				LightingSample.Diffuse = HairShading(GBufferData, L, V, N, Shadow, HairTransmittance, Backlit, Area, Random, bEvalMultiScatter);
 				LightingSample.Transmission = 0;
 				LightingSample.Specular = 0;
 			}
 			else
 			{
-				FShadowTerms ShadowTerms = { 0.0, 0.0, 0.0 };
+				FShadowTerms ShadowTerms = { 0.0, 0.0, 0.0, InitHairTransmittanceData() };
 				LightingSample = EvaluateBxDF(GBufferData, N, V, L, NoL, ShadowTerms);
 			}
 			float3 Brdf = LightingSample.Diffuse + LightingSample.Transmission + LightingSample.Specular;
diff --git a/Engine/Shaders/Private/ReflectionEnvironmentPixelShader.usf b/Engine/Shaders/Private/ReflectionEnvironmentPixelShader.usf
--- a/Engine/Shaders/Private/ReflectionEnvironmentPixelShader.usf
+++ b/Engine/Shaders/Private/ReflectionEnvironmentPixelShader.usf
@@ -143,8 +143,9 @@ float3 SkyLightDiffuse(FGBufferData GBuffer, float AmbientOcclusion, float2 Buff
 		float3 V = -normalize(mul(float4(ScreenPosition, 1, 0), View.ScreenToWorld).xyz);
 		float3 L = normalize(V - N * dot(V, N));
 		SkyLightingNormal = L;
+		FHairTransmittanceData TransmittanceData = InitHairTransmittanceData();
 		bool bEvalMultiScatter = true;
-		DiffuseColor = PI * HairShading(GBuffer, L, V, N, 1, 0, 0.2, uint2(0, 0), bEvalMultiScatter);
+		DiffuseColor = PI * HairShading(GBuffer, L, V, N, 1, TransmittanceData, 0, 0.2, uint2(0, 0), bEvalMultiScatter);
 	}
 
 	if (GBuffer.ShadingModelID == SHADINGMODELID_CLOTH)
diff --git a/Engine/Shaders/Private/ShadingModels.ush b/Engine/Shaders/Private/ShadingModels.ush
--- a/Engine/Shaders/Private/ShadingModels.ush
+++ b/Engine/Shaders/Private/ShadingModels.ush
@@ -9,6 +9,10 @@
 #include "RectLight.ush"
 #include "TransmissionCommon.ush"
 
+#ifndef HAIR_CUSTOM_BSDF
+#define HAIR_CUSTOM_BSDF 0
+#endif
+
 #if 0
 void StandardShadingShared( float3 DiffuseColor, float3 SpecularColor, float Roughness, float3 V, half3 N )
 {
@@ -47,6 +51,45 @@ void StandardShadingPerLight( Shared, float3 L, float3 V, half3 N )
 }
 #endif
 
+struct FHairTransmittanceData
+{
+	// Average front/back scattering for a given L, V, T (tangent)
+	float3 Transmittance;
+	float3 A_front;			
+	float3 A_back;			
+
+	float OpaqueVisibility;
+	float HairCount;
+	float MultipleScattering;
+
+
+	// TEMP: for fastning iteration times
+	float3 LocalScattering;
+	float3 GlobalScattering;
+
+	uint3 ScatteringComponent; // R TT TRT
+	bool  Enable;
+};
+
+FHairTransmittanceData InitHairTransmittanceData()
+{
+	FHairTransmittanceData o;
+	o.Transmittance = 1;
+	o.A_front = 1;			
+	o.A_back = 1;			
+	o.OpaqueVisibility = 1;
+	o.HairCount = 0;
+	o.MultipleScattering = 0;
+	o.Enable = false;
+	
+	// TEMP: for fastning iteration times
+	o.LocalScattering = 0;
+	o.GlobalScattering = 1;
+	o.ScatteringComponent = 1;
+
+	return o;
+}
+
 struct FAreaLight
 {
 	float		SphereSinAlpha;
@@ -65,6 +108,7 @@ struct FShadowTerms
 	float	SurfaceShadow;
 	float	TransmissionShadow;
 	float	TransmissionThickness;
+	FHairTransmittanceData HairTransmittance;
 };
 
 float New_a2( float a2, float SinAlpha, float VoH )
@@ -441,17 +485,25 @@ FDirectLighting ClothBxDF( FGBufferData GBuffer, half3 N, half3 V, half3 L, floa
 	return Lighting;
 }
 
+#if HAIR_CUSTOM_BSDF == 0
 float Hair_g( float B, float Theta )
 {
 	return exp( -0.5 * Pow2( Theta ) / (B*B) ) / ( sqrt(2*PI) * B );
 }
+#else
+float Hair_g( float B, float Theta );
+#endif
 
+#if HAIR_CUSTOM_BSDF == 0
 float Hair_F( float CosTheta )
 {
 	const float n = 1.55;
 	const float F0 = Pow2( (1 - n) / (1 + n) );
 	return F0 + (1 - F0) * Pow5( 1 - CosTheta );
 }
+#else
+float Hair_F( float CosTheta );
+#endif
 
 #define HAIR_REFERENCE	0
 #if HAIR_REFERENCE
@@ -679,10 +731,19 @@ float3 KajiyaKayDiffuseAttenuation(FGBufferData GBuffer, float3 L, float3 V, hal
 	return sqrt(GBuffer.BaseColor) * DiffuseScatter * ScatterTint;
 }
 
+float3 EvaluateHairMultipleScattering(
+	const FHairTransmittanceData TransmittanceData,
+	const float Roughness,
+	const float3 Fs)
+{
+	return TransmittanceData.GlobalScattering * (Fs + TransmittanceData.LocalScattering) * TransmittanceData.OpaqueVisibility;
+}
+
 // Approximation to HairShadingRef using concepts from the following papers:
 // [Marschner et al. 2003, "Light Scattering from Human Hair Fibers"]
 // [Pekelis et al. 2015, "A Data-Driven Light Scattering Model for Hair"]
-float3 HairShading( FGBufferData GBuffer, float3 L, float3 V, half3 N, float Shadow, float Backlit, float Area, uint2 Random, bool bEvalMultiScatter )
+#if HAIR_CUSTOM_BSDF == 0
+float3 HairShading( FGBufferData GBuffer, float3 L, float3 V, half3 N, float Shadow, FHairTransmittanceData HairTransmittance, float Backlit, float Area, uint2 Random, bool bEvalMultiScatter)
 {
 	// to prevent NaN with decals
 	// OR-18489 HERO: IGGY: RMB on E ability causes blinding hair effect
@@ -733,7 +794,7 @@ float3 HairShading( FGBufferData GBuffer, float3 L, float3 V, half3 N, float Sha
 	float3 S = 0;
 
 	// R
-	if(1)
+	if(HairTransmittance.ScatteringComponent.x > 0)
 	{
 		const float sa = sin( Alpha[0] );
 		const float ca = cos( Alpha[0] );
@@ -746,7 +807,7 @@ float3 HairShading( FGBufferData GBuffer, float3 L, float3 V, half3 N, float Sha
 	}
 
 	// TT
-	if(1)
+	if(HairTransmittance.ScatteringComponent.y > 0)
 	{
 		float Mp = Hair_g( B[1], SinThetaL + SinThetaV - Alpha[1] );
 
@@ -776,7 +837,7 @@ float3 HairShading( FGBufferData GBuffer, float3 L, float3 V, half3 N, float Sha
 	}
 
 	// TRT
-	if(1)
+	if(HairTransmittance.ScatteringComponent.z > 0)
 	{
 		float Mp = Hair_g( B[2], SinThetaL + SinThetaV - Alpha[2] );
 		
@@ -794,13 +855,17 @@ float3 HairShading( FGBufferData GBuffer, float3 L, float3 V, half3 N, float Sha
 	}
 #endif
 
-	if (bEvalMultiScatter)
+	if (!HairTransmittance.Enable && bEvalMultiScatter)
 	{
 		S += KajiyaKayDiffuseAttenuation(GBuffer, L, V, N, Shadow);
 	}
 
 	S = -min(-S, 0.0);
 
+	if (HairTransmittance.Enable && HairTransmittance.MultipleScattering > 0)
+	{
+		S = EvaluateHairMultipleScattering(HairTransmittance, ClampedRoughness, S);
+	}
 	return S;
 }
 
@@ -810,10 +875,15 @@ FDirectLighting HairBxDF( FGBufferData GBuffer, half3 N, half3 V, half3 L, float
 	Lighting.Diffuse = 0;
 	Lighting.Specular = 0;
 	bool bEvalMultiScatter = true;
-	Lighting.Transmission = AreaLight.FalloffColor * Falloff * HairShading( GBuffer, L, V, N, Shadow.TransmissionShadow, 1, 0, uint2(0,0), bEvalMultiScatter );
+	Lighting.Transmission = AreaLight.FalloffColor * Falloff * HairShading( GBuffer, L, V, N, Shadow.TransmissionShadow, Shadow.HairTransmittance, 1, 0, uint2(0,0), bEvalMultiScatter);
 	return Lighting;
 }
 
+#else
+float3 HairShading( FGBufferData GBuffer, float3 L, float3 V, half3 N, float Shadow, FHairTransmittanceData HairTransmittance, float Backlit, float Area, uint2 Random, bool bEvalMultiScatter);
+FDirectLighting HairBxDF( FGBufferData GBuffer, half3 N, half3 V, half3 L, float Falloff, float NoL, FAreaLight AreaLight, FShadowTerms Shadow );
+#endif
+
 FDirectLighting SubsurfaceBxDF( FGBufferData GBuffer, half3 N, half3 V, half3 L, float Falloff, float NoL, FAreaLight AreaLight, FShadowTerms Shadow )
 {
 	FDirectLighting Lighting = DefaultLitBxDF( GBuffer, N, V, L, Falloff, NoL, AreaLight, Shadow );
diff --git a/Engine/Shaders/Private/ShadowProjectionPixelShader.usf b/Engine/Shaders/Private/ShadowProjectionPixelShader.usf
--- a/Engine/Shaders/Private/ShadowProjectionPixelShader.usf
+++ b/Engine/Shaders/Private/ShadowProjectionPixelShader.usf
@@ -88,19 +88,45 @@ float3 ModulatedShadowBlendOp(float3 Source)
 	return Source.rgb*Dest.rgb;
 }
 
+
+#if SUBPIXEL_SHADOW
+#include "HairStrands/HairStrandsVisibilityCommon.ush"
+Texture2D<uint4> HairCategorizationTexture;
+#endif
+
 /**
 * Entry point for uniform manual PCF that supports lights using normal shadows.
 */
+// Need to disable early depth stencil as discard operation are not supported. 
+// This is only for hair case. One way to solve this would be to run tile coverring 
+// hair pixels, rather than rendering a full screen to avoid partially this extra cost
+#if SUBPIXEL_SHADOW == 0
 EARLYDEPTHSTENCIL
+#endif
 void Main(
 	in float4 SVPos : SV_POSITION,
 	out float4 OutColor : SV_Target0
 	)
 {
 	const FPQMPContext PQMPContext = PQMPInit(SVPos.xy);
+	float2 ScreenUV = float2(SVPos.xy * View.BufferSizeAndInvSize.zw);
+	float SceneW = CalcSceneDepth(ScreenUV);
+
+	bool bIsSubsurfaceCompatible = true;
+#if SUBPIXEL_SHADOW
+	const uint3 PixelCoord = uint3(floor(SVPos.xy), 0);
+	const FCategorizationData CategorizationData = DecodeCategorizationData(HairCategorizationTexture.Load(PixelCoord));
+	if (CategorizationData.TotalCoverage8bit > 0)
+	{
+		SceneW = ConvertFromDeviceZ(CategorizationData.ClosestDepth);
+		bIsSubsurfaceCompatible = false;
+	}
+	else
+	{
+		discard;
+	}
+#endif
 
-	float2 ScreenUV = float2( SVPos.xy * View.BufferSizeAndInvSize.zw );
-	float SceneW = CalcSceneDepth( ScreenUV );
 	float4 ScreenPosition = float4(((ScreenUV.xy - View.ScreenPositionScaleBias.wz ) / View.ScreenPositionScaleBias.xy) * SceneW, SceneW, 1);
 	float4 ShadowPosition = mul(ScreenPosition, ScreenToShadowMatrix);
 	float3 WorldPosition = mul(ScreenPosition, View.ScreenToWorld).xyz;
@@ -220,7 +246,7 @@ void Main(
 		FGBufferData GBufferData = GetGBufferData(ScreenUV);
 
 		BRANCH
-		if (IsSubsurfaceModel(GBufferData.ShadingModelID))
+		if (bIsSubsurfaceCompatible && IsSubsurfaceModel(GBufferData.ShadingModelID))
 		{
 			float Opacity = GBufferData.CustomData.a;
 			// Derive density from a heuristic using opacity, tweaked for useful falloff ranges and to give a linear depth falloff with opacity
@@ -235,7 +261,7 @@ void Main(
 			// Note: Snapping based on the largest size kernel is done in FProjectedShadowInfo::FProjectedShadowInfo, MaxDownsampleFactor is set to 4
 			// Only factors of MaxDownsampleFactor will be stable
 			float SquareRootFilterScale = lerp(1.999f, 0, Opacity);
-			int SquareRootFilterScaleInt = SquareRootFilterScale + 1;
+			int SquareRootFilterScaleInt = int(SquareRootFilterScale) + 1;
 
 			#if UNFILTERED_SHADOW_PROJECTION
 				float ShadowMapDepth = Texture2DSampleLevel(ShadowDepthTexture, ShadowDepthTextureSampler, ShadowPosition.xy, 0).x;
diff --git a/Engine/Shaders/Shared/RayTracingDefinitions.h b/Engine/Shaders/Shared/RayTracingDefinitions.h
--- a/Engine/Shaders/Shared/RayTracingDefinitions.h
+++ b/Engine/Shaders/Shared/RayTracingDefinitions.h
@@ -17,6 +17,7 @@
 
 #define RAY_TRACING_MASK_OPAQUE						0x01    // Opaque and alpha tested meshes and particles (e.g. used by reflection, shadow, AO and GI tracing passes)
 #define RAY_TRACING_MASK_TRANSLUCENT				0x02    // Opaque and alpha tested meshes and particles (e.g. used by translucency tracing pass)
+#define RAY_TRACING_MASK_THIN_SHADOW				0x04    // Whether the thin geometry (e.g. hair) is visible for shadow rays
 #define RAY_TRACING_MASK_SHADOW						0x08    // Whether the geometry is visible for shadow rays
 #define RAY_TRACING_MASK_ALL						0xFF
 
diff --git a/Engine/Source/Editor/ClothPainter/Private/ClothMeshAdapter.cpp b/Engine/Source/Editor/ClothPainter/Private/ClothMeshAdapter.cpp
--- a/Engine/Source/Editor/ClothPainter/Private/ClothMeshAdapter.cpp
+++ b/Engine/Source/Editor/ClothPainter/Private/ClothMeshAdapter.cpp
@@ -2,7 +2,7 @@
 
 #include "ClothMeshAdapter.h"
 
-#include "ClothingAssetInterface.h"
+#include "ClothingAssetBase.h"
 #include "Assets/ClothingAsset.h"
 
 #include "MeshPaintTypes.h"
diff --git a/Engine/Source/Editor/ClothPainter/Private/ClothingPaintEditMode.cpp b/Engine/Source/Editor/ClothPainter/Private/ClothingPaintEditMode.cpp
--- a/Engine/Source/Editor/ClothPainter/Private/ClothingPaintEditMode.cpp
+++ b/Engine/Source/Editor/ClothPainter/Private/ClothingPaintEditMode.cpp
@@ -9,7 +9,7 @@
 
 #include "ClothPainter.h"
 #include "ComponentReregisterContext.h"
-#include "ClothingAssetInterface.h"
+#include "ClothingAssetBase.h"
 #include "ComponentRecreateRenderStateContext.h"
 #include "IPersonaToolkit.h"
 #include "Assets/ClothingAsset.h"
diff --git a/Engine/Source/Editor/ClothingSystemEditor/ClothingSystemEditor.Build.cs b/Engine/Source/Editor/ClothingSystemEditor/ClothingSystemEditor.Build.cs
--- a/Engine/Source/Editor/ClothingSystemEditor/ClothingSystemEditor.Build.cs
+++ b/Engine/Source/Editor/ClothingSystemEditor/ClothingSystemEditor.Build.cs
@@ -8,7 +8,7 @@ public class ClothingSystemEditor : ModuleRules
 	public ClothingSystemEditor(ReadOnlyTargetRules Target) : base(Target)
 	{
         PublicIncludePathModuleNames.Add("UnrealEd");
-        PublicIncludePathModuleNames.Add("ClothingSystemRuntime");
+        PublicIncludePathModuleNames.Add("ClothingSystemRuntimeNv");
         PublicIncludePathModuleNames.Add("ClothingSystemEditorInterface");
 
         PrivateDependencyModuleNames.AddRange(
@@ -23,7 +23,7 @@ public class ClothingSystemEditor : ModuleRules
 			new string[] { 
 				"Core",
 				"CoreUObject",
-                "ClothingSystemRuntime",
+                "ClothingSystemRuntimeNv",
                 "ContentBrowser",
                 "UnrealEd",
                 "SlateCore",
diff --git a/Engine/Source/Editor/ClothingSystemEditor/Private/ClothingSystemEditorModule.cpp b/Engine/Source/Editor/ClothingSystemEditor/Private/ClothingSystemEditorModule.cpp
--- a/Engine/Source/Editor/ClothingSystemEditor/Private/ClothingSystemEditorModule.cpp
+++ b/Engine/Source/Editor/ClothingSystemEditor/Private/ClothingSystemEditorModule.cpp
@@ -3,7 +3,7 @@
 #include "ClothingSystemEditorModule.h"
 #include "ClothingSystemEditorInterfaceModule.h"
 
-#include "ClothingSimulationFactory.h"
+#include "ClothingSimulationFactoryNv.h"
 #include "SimulationEditorExtenderNv.h"
 #include "ClothingAssetFactory.h"
 #include "Modules/ModuleManager.h"
diff --git a/Engine/Source/Editor/ClothingSystemEditor/Private/SimulationEditorExtenderNv.cpp b/Engine/Source/Editor/ClothingSystemEditor/Private/SimulationEditorExtenderNv.cpp
--- a/Engine/Source/Editor/ClothingSystemEditor/Private/SimulationEditorExtenderNv.cpp
+++ b/Engine/Source/Editor/ClothingSystemEditor/Private/SimulationEditorExtenderNv.cpp
@@ -6,7 +6,7 @@
 
 #include "ClothingSimulation.h"
 #include "ClothingSimulationNv.h"
-#include "ClothingSimulationFactory.h"
+#include "ClothingSimulationFactoryNv.h"
 #include "IPersonaPreviewScene.h"
 #include "Animation/DebugSkelMeshComponent.h"
 #include "Framework/Commands/UIAction.h"
