Shader shelf from Sebastien
removed custom depth to clip water.

diff --git a/Engine/Shaders/Private/BasePassPixelShader.usf b/Engine/Shaders/Private/BasePassPixelShader.usf
--- a/Engine/Shaders/Private/BasePassPixelShader.usf
+++ b/Engine/Shaders/Private/BasePassPixelShader.usf
@@ -1228,7 +1228,6 @@ ShadingModel);
 			MaterialParameters, PixelMaterialInputs, ResolvedView,
 			OpaqueBasePass.SceneColorWithoutSingleLayerWaterTexture, OpaqueBasePass.SceneColorWithoutSingleLayerWaterSampler,
 			OpaqueBasePass.SceneDepthWithoutSingleLayerWaterTexture, OpaqueBasePass.SceneDepthWithoutSingleLayerWaterSampler,
-			OpaqueBasePass.SceneCustomDepthTexture, OpaqueBasePass.SceneCustomDepthSampler,
 			OpaqueBasePass.SceneWithoutSingleLayerWaterValidUVRect,
 			Specular, OpaqueBasePass.DistortionParams,
 			SunIlluminance, WaterDiffuseIndirectIlluminance, EnvBrdf,
diff --git a/Engine/Shaders/Private/SingleLayerWaterCommon.ush b/Engine/Shaders/Private/SingleLayerWaterCommon.ush
--- a/Engine/Shaders/Private/SingleLayerWaterCommon.ush
+++ b/Engine/Shaders/Private/SingleLayerWaterCommon.ush
@@ -47,7 +47,6 @@ WaterVolumeLightingOutput EvaluateWaterVolumeLighting(
 #else
 	Texture2D SceneColorWithoutSingleLayerWaterTexture, SamplerState SceneColorWithoutSingleLayerWaterSampler,
 	Texture2D SceneDepthWithoutSingleLayerWaterTexture, SamplerState SceneDepthWithoutSingleLayerWaterSampler,
-	Texture2D SceneCustomDepthTexture, SamplerState SceneCustomDepthSampler,
 	float4 SceneWithoutSingleLayerWaterValidUVRect,
 #endif // SIMPLE_SINGLE_LAYER_WATER
 	float Specular, const float4 DistortionParams,
@@ -62,14 +61,6 @@ WaterVolumeLightingOutput EvaluateWaterVolumeLighting(
 	float2 WaterNDC = MaterialParameters.ScreenPosition.xy / MaterialParameters.ScreenPosition.w;
 	float2 ViewportUV = WaterNDC * ResolvedView.ScreenPositionScaleBias.xy + ResolvedView.ScreenPositionScaleBias.wz;
 
-#if SIMPLE_SINGLE_LAYER_WATER==0
-	float SceneCustomDepth = SceneCustomDepthTexture.SampleLevel(SceneCustomDepthSampler, ViewportUV, 0).r;
-	if (SceneCustomDepth > MaterialParameters.SvPosition.z)
-	{
-		clip(-1); // Custom depth Z is used to clip, this in order to avoid water being visible inside boats for instance.
-	}
-#endif // !SIMPLE_SINGLE_LAYER_WATER
-
 	float WaterDepth = ConvertFromDeviceZ(MaterialParameters.SvPosition.z);
 	if (WaterVisibility > 0.0f)
 	{
diff --git a/Engine/Source/Runtime/Renderer/Private/BasePassRendering.cpp b/Engine/Source/Runtime/Renderer/Private/BasePassRendering.cpp
--- a/Engine/Source/Runtime/Renderer/Private/BasePassRendering.cpp
+++ b/Engine/Source/Runtime/Renderer/Private/BasePassRendering.cpp
@@ -590,8 +590,6 @@ void CreateOpaqueBasePassUniformBuffer(
 	BasePassParameters.SceneDepthWithoutSingleLayerWaterSampler = TStaticSamplerState<SF_Point>::GetRHI();
 	BasePassParameters.PreIntegratedGFTexture = GSystemTextures.PreintegratedGF->GetRenderTargetItem().ShaderResourceTexture;
 	BasePassParameters.PreIntegratedGFSampler = TStaticSamplerState<SF_Bilinear, AM_Clamp, AM_Clamp, AM_Clamp>::GetRHI();
-	BasePassParameters.SceneCustomDepthTexture = (SceneRenderTargets.CustomDepth.IsValid() ? SceneRenderTargets.CustomDepth : GSystemTextures.BlackDummy)->GetRenderTargetItem().ShaderResourceTexture;
-	BasePassParameters.SceneCustomDepthSampler = TStaticSamplerState<SF_Point>::GetRHI();
 	SetupDistortionParams(BasePassParameters.DistortionParams, View);
 
 	// Misc
diff --git a/Engine/Source/Runtime/Renderer/Private/BasePassRendering.h b/Engine/Source/Runtime/Renderer/Private/BasePassRendering.h
--- a/Engine/Source/Runtime/Renderer/Private/BasePassRendering.h
+++ b/Engine/Source/Runtime/Renderer/Private/BasePassRendering.h
@@ -78,8 +78,6 @@ BEGIN_GLOBAL_SHADER_PARAMETER_STRUCT(FOpaqueBasePassUniformParameters,)
 	SHADER_PARAMETER_SAMPLER(SamplerState, SceneDepthWithoutSingleLayerWaterSampler)
 	SHADER_PARAMETER_TEXTURE(Texture2D, PreIntegratedGFTexture)
 	SHADER_PARAMETER_SAMPLER(SamplerState, PreIntegratedGFSampler)
-	SHADER_PARAMETER_TEXTURE(Texture2D, SceneCustomDepthTexture)
-	SHADER_PARAMETER_SAMPLER(SamplerState, SceneCustomDepthSampler)
 	SHADER_PARAMETER(FVector4, SceneWithoutSingleLayerWaterValidUVRect)
 	SHADER_PARAMETER(FVector4, DistortionParams)
 	// Misc
