From yuriy
Added support for conservative depth output from materials for switch

diff --git a/Engine/Shaders/Private/MaterialTemplate.ush b/Engine/Shaders/Private/MaterialTemplate.ush
--- a/Engine/Shaders/Private/MaterialTemplate.ush
+++ b/Engine/Shaders/Private/MaterialTemplate.ush
@@ -2410,7 +2410,7 @@ half3x3 AssembleTangentToWorld( half3 TangentToWorld0, half4 TangentToWorld2 )
 
 // Whether to use the hidden d3d11 feature that supports depth writes with ZCull by only pushing into the screen
 //@todo - use for other SM5 platforms
-#define SUPPORTS_CONSERVATIVE_DEPTH_WRITES ((COMPILER_HLSL && FEATURE_LEVEL >= FEATURE_LEVEL_SM5) || (PS4_PROFILE) || (COMPILER_METAL && FEATURE_LEVEL >= FEATURE_LEVEL_SM5))
+#define SUPPORTS_CONSERVATIVE_DEPTH_WRITES ((COMPILER_HLSL && FEATURE_LEVEL >= FEATURE_LEVEL_SM5) || (PS4_PROFILE) || (COMPILER_METAL && FEATURE_LEVEL >= FEATURE_LEVEL_SM5) || SWITCH_PROFILE || SWITCH_PROFILE_FORWARD)
 #define USE_CONSERVATIVE_DEPTH_WRITES (OUTPUT_PIXEL_DEPTH_OFFSET && SUPPORTS_CONSERVATIVE_DEPTH_WRITES) 
 
 #if USE_CONSERVATIVE_DEPTH_WRITES
@@ -2426,6 +2426,9 @@ half3x3 AssembleTangentToWorld( half3 TangentToWorld0, half4 TangentToWorld2 )
 #elif PS4_PROFILE
 	#define INPUT_POSITION_QUALIFIERS
 	#define DEPTH_WRITE_SEMANTIC S_DEPTH_LE_OUTPUT
+#elif SWITCH_PROFILE || SWITCH_PROFILE_FORWARD
+	#define INPUT_POSITION_QUALIFIERS
+	#define DEPTH_WRITE_SEMANTIC SV_DepthLessEqual
 #else
 	#error USE_CONSERVATIVE_DEPTH_WRITES enabled for unsupported platform
 #endif
