Fixed water absorption and scattering due to missing scene color writes from the single layer water material. PIXELSHADEROUTPUT_MRT0 needs to to be true when MATERIAL_SHADINGMODEL_SINGLELAYERWATER is enabled.

+ Shelve from Zabir Hoque: Disable dbuffer blending during water rendering. ~3900us -> ~3500us speedup.

diff --git a/Engine/Shaders/Private/BasePassPixelShader.usf b/Engine/Shaders/Private/BasePassPixelShader.usf
--- a/Engine/Shaders/Private/BasePassPixelShader.usf
+++ b/Engine/Shaders/Private/BasePassPixelShader.usf
@@ -862,7 +862,7 @@ void FPixelShaderInOut_MainPS(
 #endif
 
 	float DBufferOpacity = 1.0f;
-#if USE_DBUFFER && MATERIALDECALRESPONSEMASK && !MATERIALBLENDING_ANY_TRANSLUCENT
+#if USE_DBUFFER && MATERIALDECALRESPONSEMASK && !MATERIALBLENDING_ANY_TRANSLUCENT && !MATERIAL_SHADINGMODEL_SINGLELAYERWATER
 	// apply decals from the DBuffer
 #if SM5_PROFILE
 	//Temporary workaround to avoid crashes on AMD, revert back to BRANCH
@@ -1373,7 +1373,7 @@ ShadingModel);
 // the following needs to match to the code in FSceneRenderTargets::GetGBufferRenderTargets()
 #define PIXELSHADEROUTPUT_BASEPASS 1
 #if USES_GBUFFER
-#define PIXELSHADEROUTPUT_MRT0 (!SELECTIVE_BASEPASS_OUTPUTS || NEEDS_BASEPASS_VERTEX_FOGGING || USES_EMISSIVE_COLOR || ALLOW_STATIC_LIGHTING)
+#define PIXELSHADEROUTPUT_MRT0 (!SELECTIVE_BASEPASS_OUTPUTS || NEEDS_BASEPASS_VERTEX_FOGGING || USES_EMISSIVE_COLOR || ALLOW_STATIC_LIGHTING || MATERIAL_SHADINGMODEL_SINGLELAYERWATER)
 #define PIXELSHADEROUTPUT_MRT1 ((!SELECTIVE_BASEPASS_OUTPUTS || !MATERIAL_SHADINGMODEL_UNLIT))
 #define PIXELSHADEROUTPUT_MRT2 ((!SELECTIVE_BASEPASS_OUTPUTS || !MATERIAL_SHADINGMODEL_UNLIT))
 #define PIXELSHADEROUTPUT_MRT3 ((!SELECTIVE_BASEPASS_OUTPUTS || !MATERIAL_SHADINGMODEL_UNLIT))
