Shared water shader samplers.

diff --git a/Engine/Shaders/Private/BasePassPixelShader.usf b/Engine/Shaders/Private/BasePassPixelShader.usf
--- a/Engine/Shaders/Private/BasePassPixelShader.usf
+++ b/Engine/Shaders/Private/BasePassPixelShader.usf
@@ -16,8 +16,12 @@
 
 #if MATERIAL_SHADINGMODEL_SINGLELAYERWATER
 #define PreIntegratedGF			OpaqueBasePass.PreIntegratedGFTexture
+#if SUPPORTS_INDEPENDENT_SAMPLERS
+#define PreIntegratedGFSampler	View.SharedBilinearClampedSampler
+#else
 #define PreIntegratedGFSampler	OpaqueBasePass.PreIntegratedGFSampler
 #endif
+#endif
 
 #include "SHCommon.ush"
 #include "/Engine/Generated/Material.ush"
@@ -321,9 +325,15 @@ void GetSkyLighting(FMaterialPixelParameters MaterialParameters, VTPageTableResu
 #if SUPPORTS_INDEPENDENT_SAMPLERS
 	#define ILCSharedSampler1 View.SharedBilinearClampedSampler
 	#define ILCSharedSampler2 View.SharedBilinearClampedSampler
+
+	#define SingleLayerWaterSceneColorSampler View.SharedBilinearClampedSampler
+	#define SingleLayerWaterSceneDepthSampler View.SharedPointClampedSampler
 #else
 	#define ILCSharedSampler1 IndirectLightingCache.IndirectLightingCacheTextureSampler1
 	#define ILCSharedSampler2 IndirectLightingCache.IndirectLightingCacheTextureSampler2
+
+	#define SingleLayerWaterSceneColorSampler OpaqueBasePass.SceneColorWithoutSingleLayerWaterSampler
+	#define SingleLayerWaterSceneDepthSampler OpaqueBasePass.SceneDepthWithoutSingleLayerWaterSampler
 #endif
 
 /** Calculates indirect lighting contribution on this object from precomputed data. */
@@ -1226,8 +1236,8 @@ ShadingModel);
 
 		Color += EvaluateWaterVolumeLighting(
 			MaterialParameters, PixelMaterialInputs, ResolvedView,
-			OpaqueBasePass.SceneColorWithoutSingleLayerWaterTexture, OpaqueBasePass.SceneColorWithoutSingleLayerWaterSampler,
-			OpaqueBasePass.SceneDepthWithoutSingleLayerWaterTexture, OpaqueBasePass.SceneDepthWithoutSingleLayerWaterSampler,
+			OpaqueBasePass.SceneColorWithoutSingleLayerWaterTexture, SingleLayerWaterSceneColorSampler,
+			OpaqueBasePass.SceneDepthWithoutSingleLayerWaterTexture, SingleLayerWaterSceneDepthSampler,
 			OpaqueBasePass.SceneWithoutSingleLayerWaterValidUVRect,
 			Specular, OpaqueBasePass.DistortionParams,
 			SunIlluminance, WaterDiffuseIndirectIlluminance, EnvBrdf,
