[NO NON-PII COMMIT MESSAGE]

diff --git a/Engine/Shaders/Private/SceneTexturesCommon.ush b/Engine/Shaders/Private/SceneTexturesCommon.ush
--- a/Engine/Shaders/Private/SceneTexturesCommon.ush
+++ b/Engine/Shaders/Private/SceneTexturesCommon.ush
@@ -127,8 +127,7 @@ float CalcSceneDepth(float2 ScreenUV)
 				// We cannot fall back to fetching the alpha channel when MobileHDR=false because the alpha channel is only 8-bit.
 				return DepthbufferFetchES2(65000.0, View.InvDeviceZToWorldZTransform[2], View.InvDeviceZToWorldZTransform[3]);
 			#else
-				// Try framebuffer_fetch_depth_stencil to get the depth directly, and if unavailable fall back to framebuffer_fetch from the alpha channel.
-				return DepthbufferFetchES2(FramebufferFetchES2().w, View.InvDeviceZToWorldZTransform[2], View.InvDeviceZToWorldZTransform[3]);
+				return FramebufferFetchES2().w;
 			#endif
 		#elif VULKAN_SUBPASS_DEPTHFETCH
 			// Special intrinsic to read from the current depth buffer
