Sky optimisation using fastmath for acos and atan2.

With a basic sky shader:
 - Uses 4 less VGPRs on ps4. Not measurable win on ps4 but could probably be a win for switch or mobile.
 - DXBC 157 instructions down to 146
 - no visual difference seen

diff --git a/Engine/Shaders/Private/SkyAtmosphere.usf b/Engine/Shaders/Private/SkyAtmosphere.usf
--- a/Engine/Shaders/Private/SkyAtmosphere.usf
+++ b/Engine/Shaders/Private/SkyAtmosphere.usf
@@ -264,7 +264,7 @@ void UvToSkyViewLutParams(out float3 ViewDir, in float ViewHeight, in float2 UV)
 
 	float Vhorizon = sqrt(ViewHeight * ViewHeight - Atmosphere.BottomRadius * Atmosphere.BottomRadius);
 	float CosBeta = Vhorizon / ViewHeight;				// cos of zenith angle from horizon to zeniht
-	float Beta = acos(CosBeta);
+	float Beta = acosFast(CosBeta);
 	float ZenithHorizonAngle = PI - Beta;
 
 	float ViewZenithAngle;
diff --git a/Engine/Shaders/Private/SkyAtmosphereCommon.ush b/Engine/Shaders/Private/SkyAtmosphereCommon.ush
--- a/Engine/Shaders/Private/SkyAtmosphereCommon.ush
+++ b/Engine/Shaders/Private/SkyAtmosphereCommon.ush
@@ -129,12 +129,12 @@ void SkyViewLutParamsToUv(
 {
 	float Vhorizon = sqrt(ViewHeight * ViewHeight - BottomRadius * BottomRadius);
 	float CosBeta = Vhorizon / ViewHeight;				// GroundToHorizonCos
-	float Beta = acos(CosBeta);
+	float Beta = acosFast(CosBeta);
 	float ZenithHorizonAngle = PI - Beta;
 
 	if (!IntersectGround)
 	{
-		float Coord = acos(ViewZenithCosAngle) / ZenithHorizonAngle;
+		float Coord = acosFast(ViewZenithCosAngle) / ZenithHorizonAngle;
 		Coord = 1.0f - Coord;
 		Coord = sqrt(Coord);
 		Coord = 1.0f - Coord;
@@ -142,13 +142,13 @@ void SkyViewLutParamsToUv(
 	}
 	else
 	{
-		float Coord = (acos(ViewZenithCosAngle) - ZenithHorizonAngle) / Beta;
+		float Coord = (acosFast(ViewZenithCosAngle) - ZenithHorizonAngle) / Beta;
 		Coord = sqrt(Coord);
 		UV.y = Coord * 0.5f + 0.5f;
 	}
 
 	{
-		UV.x = (atan2(-ViewDir.y, -ViewDir.x) + PI) / (2.0f * PI);
+		UV.x = (atan2Fast(-ViewDir.y, -ViewDir.x) + PI) / (2.0f * PI);
 	}
 
 	// Constrain uvs to valid sub texel range (avoid zenith derivative issue making LUT usage visible)
