Added water fallback to constant under water depth when reaching out mobile depth maximum representatble value.
With a fade region. A small transition is visible and it also depends a bit on the material. That is the best I can do without evaluating material twice.

diff --git a/Engine/Shaders/Private/SingleLayerWaterCommon.ush b/Engine/Shaders/Private/SingleLayerWaterCommon.ush
--- a/Engine/Shaders/Private/SingleLayerWaterCommon.ush
+++ b/Engine/Shaders/Private/SingleLayerWaterCommon.ush
@@ -67,7 +67,22 @@ WaterVolumeLightingOutput EvaluateWaterVolumeLighting(
 #if SIMPLE_SINGLE_LAYER_WATER
 		float PixelSceneDeviceZ = LookupDeviceZ(ViewportUV);
 		float PixelSceneDepth = ConvertFromDeviceZ(PixelSceneDeviceZ);
-		
+
+#if FEATURE_LEVEL <= FEATURE_LEVEL_ES3_1
+		// Only mobile run the SIMPLE_SINGLE_LAYER_WATER path today.
+		// On mobile, depth is stored linear in alpha channel of scene color as FP16 so it is easy to reach it max value. It would then result in fully transparent water.
+		// To counter that, we assume any visible water pixel has a constant under water depth and we progressively lerp towards this simplified representation.
+		// Lighting remains correct thanks to the analytical integration of scattering.
+		const float FadeLenght = 10000.0f;								// 0.1 kilometers
+		const float FadeStartDistance = 65504.0f - FadeLenght;
+		const float FadeBehindWaterDepth = WaterDepth + FadeLenght;
+		if (WaterDepth > FadeStartDistance)
+		{
+			PixelSceneDepth = lerp(PixelSceneDepth, FadeBehindWaterDepth, saturate((WaterDepth - FadeStartDistance) / FadeLenght));
+			PixelSceneDeviceZ = ConvertToDeviceZ(PixelSceneDepth);
+		}
+#endif
+
 		float SceneDeviceZ = PixelSceneDeviceZ;
 		float SceneDepth = PixelSceneDepth;
 
