Fixed ray tracing reflections not energy conserving when SPP > 1

diff --git a/Engine/Shaders/Private/RayTracing/RayTracingReflections.usf b/Engine/Shaders/Private/RayTracing/RayTracingReflections.usf
--- a/Engine/Shaders/Private/RayTracing/RayTracingReflections.usf
+++ b/Engine/Shaders/Private/RayTracing/RayTracingReflections.usf
@@ -726,7 +726,7 @@ void RayTracingReflectionsRGS()
 			float DistanceFadeRange = LocalMaxRayDistance * 0.1; // Fade out reflections over 10% of the total maximum ray length (arbitrarily chosen)
 			float DistanceFade = bAllowSkySampling ? 1.0 : saturate((LocalMaxRayDistance - HitDistance) * rcp(DistanceFadeRange));
 
-			ReflectedColor *= DistanceFade * max(RoughnessFade, ClearCoatRoughnessFade) * rcp(ReflectedWeight);
+			ReflectedColor *= DistanceFade * max(RoughnessFade, ClearCoatRoughnessFade) * rcp(LocalSamplesPerPixel);
 		}
 
 	#if USE_PREEXPOSURE
