Shelve from Sebastien.Hillaire: Do not sample directional light shadow twice on water materials. 5% reduction in pixel shader cost on PS4.

diff --git a/Engine/Shaders/Private/BasePassPixelShader.usf b/Engine/Shaders/Private/BasePassPixelShader.usf
--- a/Engine/Shaders/Private/BasePassPixelShader.usf
+++ b/Engine/Shaders/Private/BasePassPixelShader.usf
@@ -1067,6 +1067,7 @@ ShadingModel);
 
 		float IndirectOcclusion = 1.0f;
 		float2 NearestResolvedDepthScreenUV = 0;
+		float DirectionalLightShadow = 1.0f;
 
 		#if FORWARD_SHADING && (MATERIALBLENDING_SOLID || MATERIALBLENDING_MASKED)
 			float2 NDC = MaterialParameters.ScreenPosition.xy / MaterialParameters.ScreenPosition.w;
@@ -1107,7 +1108,7 @@ ShadingModel);
 
 				#if FORWARD_SHADING || TRANSLUCENCY_LIGHTING_SURFACE_FORWARDSHADING || MATERIAL_SHADINGMODEL_SINGLELAYERWATER
 					const float Dither = InterleavedGradientNoise(MaterialParameters.SvPosition.xy, View.StateFrameIndexMod8);
-					float3 ForwardDirectLighting = GetForwardDirectLighting(GridIndex, MaterialParameters.AbsoluteWorldPosition, MaterialParameters.CameraVector, GBuffer, NearestResolvedDepthScreenUV, MaterialParameters.PrimitiveId, EyeIndex, Dither);
+					float3 ForwardDirectLighting = GetForwardDirectLighting(GridIndex, MaterialParameters.AbsoluteWorldPosition, MaterialParameters.CameraVector, GBuffer, NearestResolvedDepthScreenUV, MaterialParameters.PrimitiveId, EyeIndex, Dither, DirectionalLightShadow);
 					Color += ForwardDirectLighting;
 				#endif
 			#endif
@@ -1236,6 +1237,7 @@ ShadingModel);
 
 		Color += EvaluateWaterVolumeLighting(
 			MaterialParameters, PixelMaterialInputs, ResolvedView,
+			DirectionalLightShadow,
 			OpaqueBasePass.SceneColorWithoutSingleLayerWaterTexture, SingleLayerWaterSceneColorSampler,
 			OpaqueBasePass.SceneDepthWithoutSingleLayerWaterTexture, SingleLayerWaterSceneDepthSampler,
 			OpaqueBasePass.SceneWithoutSingleLayerWaterMaxUV,
diff --git a/Engine/Shaders/Private/ClusteredDeferredShadingPixelShader.usf b/Engine/Shaders/Private/ClusteredDeferredShadingPixelShader.usf
--- a/Engine/Shaders/Private/ClusteredDeferredShadingPixelShader.usf
+++ b/Engine/Shaders/Private/ClusteredDeferredShadingPixelShader.usf
@@ -111,7 +111,8 @@ float3 GetLightGridLocalLighting(const FCulledLightsGridData InLightGridData, fl
 		FRectTexture RectTexture = InitRectTexture(ForwardLightData.DummyRectLightSourceTexture);
 		// NOTE: uses AO data like tiled deferred (as opposed to forward path)
 		// TODO: Passing in 0,0 as SVPos, what is meaning of this? Only used if 'REFERENCE_QUALITY' is on.
-		DirectLighting += GetDynamicLighting(WorldPosition, CameraVector, ScreenSpaceData.GBuffer, ScreenSpaceData.AmbientOcclusion, ScreenSpaceData.GBuffer.ShadingModelID, LightData, LightAttenuation, Dither, uint2(0, 0), RectTexture).xyz;
+		float SurfaceShadow = 1.0f;
+		DirectLighting += GetDynamicLighting(WorldPosition, CameraVector, ScreenSpaceData.GBuffer, ScreenSpaceData.AmbientOcclusion, ScreenSpaceData.GBuffer.ShadingModelID, LightData,LightAttenuation, Dither, uint2(0, 0), RectTexture, SurfaceShadow).xyz;
 
 #else // 1
 		FSimpleDeferredLightData SimpleLightData = (FSimpleDeferredLightData)0;
diff --git a/Engine/Shaders/Private/DeferredLightPixelShaders.usf b/Engine/Shaders/Private/DeferredLightPixelShaders.usf
--- a/Engine/Shaders/Private/DeferredLightPixelShaders.usf
+++ b/Engine/Shaders/Private/DeferredLightPixelShaders.usf
@@ -207,7 +207,8 @@ void ComputeHairLighting(
 		const float4 LightAttenuation = ScreenShadowMaskSubPixelTexture.Load(PixelCoord);
 
 		FRectTexture RectTexture = InitRectTexture(DeferredLightUniforms.SourceTexture);
-		const float4 Radiance = GetDynamicLighting(DerivedParams.WorldPosition, DerivedParams.CameraVector, HairScreenSpaceData.GBuffer, HairScreenSpaceData.AmbientOcclusion, HairScreenSpaceData.GBuffer.ShadingModelID, LightData, LightAttenuation, Dither, uint2(InputParams.PixelPos), RectTexture);
+		float SurfaceShadow = 1.0f;
+		const float4 Radiance = GetDynamicLighting(DerivedParams.WorldPosition, DerivedParams.CameraVector, HairScreenSpaceData.GBuffer, HairScreenSpaceData.AmbientOcclusion, HairScreenSpaceData.GBuffer.ShadingModelID, LightData, LightAttenuation, Dither, uint2(InputParams.PixelPos), RectTexture, SurfaceShadow);
 		const float  Attenuation = ComputeLightProfileMultiplier(DerivedParams.WorldPosition, DeferredLightUniforms.Position, -DeferredLightUniforms.Direction, DeferredLightUniforms.Tangent);
 
 		const float4 LinearColor = Radiance * Attenuation;
@@ -304,7 +305,8 @@ void DeferredLightPixelMain(
 		float Dither = InterleavedGradientNoise(InputParams.PixelPos, View.StateFrameIndexMod8 );
 
 		FRectTexture RectTexture = InitRectTexture(DeferredLightUniforms.SourceTexture);
-		const float4 Radiance = GetDynamicLighting(DerivedParams.WorldPosition, DerivedParams.CameraVector, ScreenSpaceData.GBuffer, ScreenSpaceData.AmbientOcclusion, ScreenSpaceData.GBuffer.ShadingModelID, LightData, GetPerPixelLightAttenuation(InputParams.ScreenUV), Dither, uint2(InputParams.PixelPos), RectTexture);
+		float SurfaceShadow = 1.0f;
+		const float4 Radiance = GetDynamicLighting(DerivedParams.WorldPosition, DerivedParams.CameraVector, ScreenSpaceData.GBuffer, ScreenSpaceData.AmbientOcclusion, ScreenSpaceData.GBuffer.ShadingModelID, LightData, GetPerPixelLightAttenuation(InputParams.ScreenUV), Dither, uint2(InputParams.PixelPos), RectTexture, SurfaceShadow);
 		const float  Attenuation = ComputeLightProfileMultiplier(DerivedParams.WorldPosition, DeferredLightUniforms.Position, -DeferredLightUniforms.Direction, DeferredLightUniforms.Tangent);
 
 #if USE_HAIR_LIGHTING
diff --git a/Engine/Shaders/Private/DeferredLightingCommon.ush b/Engine/Shaders/Private/DeferredLightingCommon.ush
--- a/Engine/Shaders/Private/DeferredLightingCommon.ush
+++ b/Engine/Shaders/Private/DeferredLightingCommon.ush
@@ -282,7 +282,10 @@ FCapsuleLight GetCapsule( float3 ToLight, FDeferredLightData LightData )
 }
 
 /** Calculates lighting for a given position, normal, etc with a fully featured lighting model designed for quality. */
-float4 GetDynamicLighting(float3 WorldPosition, float3 CameraVector, FGBufferData GBuffer, float AmbientOcclusion, uint ShadingModelID, FDeferredLightData LightData, float4 LightAttenuation, float Dither, uint2 SVPos, FRectTexture SourceTexture)
+float4 GetDynamicLighting(
+	float3 WorldPosition, float3 CameraVector, FGBufferData GBuffer, float AmbientOcclusion, uint ShadingModelID, 
+	FDeferredLightData LightData, float4 LightAttenuation, float Dither, uint2 SVPos, FRectTexture SourceTexture,
+	inout float SurfaceShadow)
 {
 	FLightAccumulator LightAccumulator = (FLightAccumulator)0;
 
@@ -315,6 +318,7 @@ float4 GetDynamicLighting(float3 WorldPosition, float3 CameraVector, FGBufferDat
 		Shadow.HairTransmittance.Transmittance = 1;
 		Shadow.HairTransmittance.OpaqueVisibility = 1;
 		GetShadowTerms(GBuffer, LightData, WorldPosition, L, LightAttenuation, Dither, Shadow);
+		SurfaceShadow = Shadow.SurfaceShadow;
 
 		LightAccumulator.EstimatedCost += 0.3f;		// add the cost of getting the shadow terms
 
diff --git a/Engine/Shaders/Private/ForwardLightingCommon.ush b/Engine/Shaders/Private/ForwardLightingCommon.ush
--- a/Engine/Shaders/Private/ForwardLightingCommon.ush
+++ b/Engine/Shaders/Private/ForwardLightingCommon.ush
@@ -96,7 +96,7 @@ float GetIndirectOcclusion(float2 ScreenUV, FGBufferData GBufferData)
 	return IndirectOcclusion;
 }
 
-float3 GetForwardDirectLighting(uint GridIndex, float3 WorldPosition, float3 CameraVector, FGBufferData GBufferData, float2 ScreenUV, uint PrimitiveId, uint EyeIndex, float Dither)
+float3 GetForwardDirectLighting(uint GridIndex, float3 WorldPosition, float3 CameraVector, FGBufferData GBufferData, float2 ScreenUV, uint PrimitiveId, uint EyeIndex, float Dither, inout float OutDirectionalLightShadow)
 {
 	float4 DynamicShadowFactors = 1;
 
@@ -152,10 +152,11 @@ float3 GetForwardDirectLighting(uint GridIndex, float3 WorldPosition, float3 Cam
 				GBufferData.PrecomputedShadowFactors.x = 1;
 			#endif
 			float DynamicShadowFactor = ComputeDirectionalLightDynamicShadowing(WorldPosition, GBufferData.Depth);
+
 			float4 LightAttenuation = float4(DynamicShadowFactor.x, DynamicShadowFactor.x, 1, 1);
 		#endif
-		
-		float3 NewLighting = GetDynamicLighting(WorldPosition, -CameraVector, GBufferData, 1, GBufferData.ShadingModelID, LightData, LightAttenuation, Dither, uint2(0,0), RectTexture).xyz;
+
+		float3 NewLighting = GetDynamicLighting(WorldPosition, -CameraVector, GBufferData, 1, GBufferData.ShadingModelID, LightData, LightAttenuation, Dither, uint2(0,0), RectTexture, OutDirectionalLightShadow).xyz;
 
 		FLATTEN
 		if ((DirectionalLightData.DirectionalLightShadowMapChannelMask >> 8) & LightingChannelMask)
@@ -207,7 +208,8 @@ float3 GetForwardDirectLighting(uint GridIndex, float3 WorldPosition, float3 Cam
 		float4 PreviewShadowMapChannelMask = UnpackShadowMapChannelMask(LightTypeAndPackedShadowMapChannelMask >> 4);
 		float DynamicShadowing = dot(PreviewShadowMapChannelMask, DynamicShadowFactors);
 		float4 LightAttenuation = float4(1, 1, DynamicShadowing.x, DynamicShadowing.x);
-		float3 NewLighting = GetDynamicLighting(WorldPosition, -CameraVector, GBufferData, 1, GBufferData.ShadingModelID, LightData, LightAttenuation, Dither, uint2(0,0), RectTexture).xyz;
+		float SurfaceShadow = 1.0f;
+		float3 NewLighting = GetDynamicLighting(WorldPosition, -CameraVector, GBufferData, 1, GBufferData.ShadingModelID, LightData, LightAttenuation, Dither, uint2(0,0), RectTexture, SurfaceShadow).xyz;
 
 		FLATTEN
 		if ((LightTypeAndPackedShadowMapChannelMask >> 8) & LightingChannelMask)
diff --git a/Engine/Shaders/Private/RayTracing/RayTracingLightingCommon.ush b/Engine/Shaders/Private/RayTracing/RayTracingLightingCommon.ush
--- a/Engine/Shaders/Private/RayTracing/RayTracingLightingCommon.ush
+++ b/Engine/Shaders/Private/RayTracing/RayTracingLightingCommon.ush
@@ -473,7 +473,8 @@ float3 ComputeDirectLightingCulled(
 
 		if (Active)
 		{
-			float3 LightContribution = GetDynamicLighting(WorldPosition, -CameraVector, GBufferData, AmbientOcclusion, GBufferData.ShadingModelID, LightData, LightAttenuation, 0.5, uint2(0, 0), RectTexture).xyz;
+			float SurfaceShadow = 1.0f;
+			float3 LightContribution = GetDynamicLighting(WorldPosition, -CameraVector, GBufferData, AmbientOcclusion, GBufferData.ShadingModelID, LightData, LightAttenuation, 0.5, uint2(0, 0), RectTexture, SurfaceShadow).xyz;
 			DirectLighting += LightContribution * LightProfileMultiplier;
 		}
 
diff --git a/Engine/Shaders/Private/SingleLayerWaterCommon.ush b/Engine/Shaders/Private/SingleLayerWaterCommon.ush
--- a/Engine/Shaders/Private/SingleLayerWaterCommon.ush
+++ b/Engine/Shaders/Private/SingleLayerWaterCommon.ush
@@ -42,9 +42,8 @@ bool WaterRefract(float3 V, float3 N, float Eta, out float3 R)
 // Returns as luminance, the result of lighting scattering in water and the under water scene color contribution.
 WaterVolumeLightingOutput EvaluateWaterVolumeLighting(
 	FMaterialPixelParameters MaterialParameters, FPixelMaterialInputs PixelMaterialInputs, ViewState ResolvedView,
-#if SIMPLE_SINGLE_LAYER_WATER
 	float DirectionalLightShadow,
-#else
+#if !SIMPLE_SINGLE_LAYER_WATER
 	Texture2D SceneColorWithoutSingleLayerWaterTexture, SamplerState SceneColorWithoutSingleLayerWaterSampler,
 	Texture2D SceneDepthWithoutSingleLayerWaterTexture, SamplerState SceneDepthWithoutSingleLayerWaterSampler,
 	float2 SceneWithoutSingleLayerWaterMaxUV,
@@ -61,9 +60,9 @@ WaterVolumeLightingOutput EvaluateWaterVolumeLighting(
 	float2 WaterNDC = MaterialParameters.ScreenPosition.xy / MaterialParameters.ScreenPosition.w;
 	float2 ViewportUV = WaterNDC * ResolvedView.ScreenPositionScaleBias.xy + ResolvedView.ScreenPositionScaleBias.wz;
 
-	float WaterDepth = ConvertFromDeviceZ(MaterialParameters.SvPosition.z);
 	if (WaterVisibility > 0.0f)
 	{
+		float WaterDepth = ConvertFromDeviceZ(MaterialParameters.SvPosition.z);
 #if SIMPLE_SINGLE_LAYER_WATER
 		float PixelSceneDeviceZ = LookupDeviceZ(ViewportUV);
 		float PixelSceneDepth = ConvertFromDeviceZ(PixelSceneDeviceZ);
@@ -159,9 +158,6 @@ WaterVolumeLightingOutput EvaluateWaterVolumeLighting(
 		float3 AmbScattLuminance = IsotropicPhase()   * AmbiantIlluminance;
 
 #if 1 || SIMPLE_SINGLE_LAYER_WATER
-#if SIMPLE_SINGLE_LAYER_WATER==0
-		float DirectionalLightShadow = ComputeDirectionalLightDynamicShadowing(MaterialParameters.AbsoluteWorldPosition, /*GBuffer.Depth*/ MaterialParameters.ScreenPosition.w);
-#endif
 		const float3 OpticalDepth = ExtinctionCoeff * BehindWaterDeltaDepth;
 		float3 Transmittance = exp(-OpticalDepth);
 		float3 ScatteredLuminance = ScatteringCoeff * (AmbScattLuminance + SunScattLuminance * DirectionalLightShadow);
diff --git a/Engine/Shaders/Private/TiledDeferredLightShaders.usf b/Engine/Shaders/Private/TiledDeferredLightShaders.usf
--- a/Engine/Shaders/Private/TiledDeferredLightShaders.usf
+++ b/Engine/Shaders/Private/TiledDeferredLightShaders.usf
@@ -88,7 +88,8 @@ void ExecuteShadingLoops(inout float4 CompositedLighting, FScreenSpaceData Scree
 		// Lights requiring light attenuation are not supported tiled for now
 		// WARNING: DeferredLightUniforms.SourceTexture is not set here, but compiles out because of bRectLight being false and USE_SOURCE_TEXTURE == 0.
 		FRectTexture RectTexture = InitRectTexture(DeferredLightUniforms.SourceTexture);
-		CompositedLighting += GetDynamicLighting(WorldPosition, CameraVector, ScreenSpaceData.GBuffer, ScreenSpaceData.AmbientOcclusion, ScreenSpaceData.GBuffer.ShadingModelID, LightData, float4(1, 1, 1, 1), 0.5, uint2(0,0), RectTexture);
+		float SurfaceShadow = 1.0f;
+		CompositedLighting += GetDynamicLighting(WorldPosition, CameraVector, ScreenSpaceData.GBuffer, ScreenSpaceData.AmbientOcclusion, ScreenSpaceData.GBuffer.ShadingModelID, LightData, float4(1, 1, 1, 1), 0.5, uint2(0,0), RectTexture, SurfaceShadow);
 	}
 
 	// The simple shading model does not depend on ShadingModelID, so use it anytime the material is lit
