[Water Mesh] WAter vertex factory doesn't write out velocity, so make sure none of that code is present by disabling WRITES_VELOCITY_TO_GBUFFER

#rnx

diff --git a/Engine/Shaders/Private/BasePassCommon.ush b/Engine/Shaders/Private/BasePassCommon.ush
--- a/Engine/Shaders/Private/BasePassCommon.ush
+++ b/Engine/Shaders/Private/BasePassCommon.ush
@@ -51,7 +51,7 @@
 
 // If a primitive has static lighting, we assume it is not moving. If it is, it will be rerendered in an extra renderpass.
 #define SUPPORTS_WRITING_VELOCITY_TO_BASE_PASS	(FEATURE_LEVEL >= FEATURE_LEVEL_SM4 && (MATERIALBLENDING_SOLID || MATERIALBLENDING_MASKED) && !SIMPLE_FORWARD_SHADING) 
-#define WRITES_VELOCITY_TO_GBUFFER				((SUPPORTS_WRITING_VELOCITY_TO_BASE_PASS || USES_GBUFFER) && GBUFFER_HAS_VELOCITY && (!SELECTIVE_BASEPASS_OUTPUTS || !(STATICLIGHTING_TEXTUREMASK || STATICLIGHTING_SIGNEDDISTANCEFIELD || HQ_TEXTURE_LIGHTMAP || LQ_TEXTURE_LIGHTMAP)))
+#define WRITES_VELOCITY_TO_GBUFFER				((SUPPORTS_WRITING_VELOCITY_TO_BASE_PASS || USES_GBUFFER) && GBUFFER_HAS_VELOCITY && (!SELECTIVE_BASEPASS_OUTPUTS || !(STATICLIGHTING_TEXTUREMASK || STATICLIGHTING_SIGNEDDISTANCEFIELD || HQ_TEXTURE_LIGHTMAP || LQ_TEXTURE_LIGHTMAP || WATER_MESH_FACTORY)))
 
 
 // Use an extra interpolator for position
