F - Workaround for hlslcc to unblock Mac
#rnx

diff --git a/Engine/Shaders/Private/LandscapeVertexFactory.ush b/Engine/Shaders/Private/LandscapeVertexFactory.ush
--- a/Engine/Shaders/Private/LandscapeVertexFactory.ush
+++ b/Engine/Shaders/Private/LandscapeVertexFactory.ush
@@ -67,7 +67,11 @@ float4 GetNeighborSectionLod()
 	//    - - 3 - -   
 	//       +Y
 
-	int2 NeighborOffset[4] = { {0, -1}, {-1, 0}, {1, 0}, {0, 1} };
+	int2 NeighborOffset[4];
+    NeighborOffset[0] = int2(0, -1);
+    NeighborOffset[1] = int2(-1, 0);
+    NeighborOffset[2] = int2(1, 0);
+    NeighborOffset[3] = int2(0, 1);
 	float NeighborLod[4];
 	for (int i = 0; i < 4; i++)
 	{
@@ -898,4 +902,4 @@ float4 VertexFactoryGetTranslatedPrimitiveVolumeBounds(FVertexFactoryInterpolant
 uint VertexFactoryGetPrimitiveId(FVertexFactoryInterpolantsVSToPS Interpolants)
 {
 	return 0;
-}
\ No newline at end of file
+}
