Fixed black sky when looking up through water on non-mobile.  It is probably still broken on mobile

diff --git a/Engine/Shaders/Private/SingleLayerWaterCommon.ush b/Engine/Shaders/Private/SingleLayerWaterCommon.ush
--- a/Engine/Shaders/Private/SingleLayerWaterCommon.ush
+++ b/Engine/Shaders/Private/SingleLayerWaterCommon.ush
@@ -131,8 +131,12 @@ WaterVolumeLightingOutput EvaluateWaterVolumeLighting(
 		const float3 AbsorptionCoeff = max(0.0f, GetSingleLayerWaterMaterialOutput1(MaterialParameters));
 		const float PhaseG = clamp(GetSingleLayerWaterMaterialOutput2(MaterialParameters), -1.0f, 1.0f);
 		// Max to avoid division by 0 with the analytical integral below.
+#if FEATURE_LEVEL <= FEATURE_LEVEL_ES3_1
 		// 1e-5 is high enough to avoid denorms on mobile
 		const float3 ExtinctionCoeff = max(ScatteringCoeff + AbsorptionCoeff, 1e-5); 
+#else
+		const float3 ExtinctionCoeff = max(ScatteringCoeff + AbsorptionCoeff, 1e-20); 
+#endif
 		
 		float DirLightPhaseValue = 0.0f; // Default when Total Internal Reflection happens.
 		{
