From Sebastien Hillaire:

Fixed extinction clamp to only be used where it is needed.
Avoid black water while allowing opticaldepth=0 and thus transmittance=0 when absorption+scattering=0.

This is a follow up to 9425005. Transmittance was still not 0 when abs=scatt=0, it is just that the threshold is small so very low opaticalDepth and black will happen on mobile. Now All platforms have same consistent behavior.

diff --git a/Engine/Shaders/Private/SingleLayerWaterCommon.ush b/Engine/Shaders/Private/SingleLayerWaterCommon.ush
--- a/Engine/Shaders/Private/SingleLayerWaterCommon.ush
+++ b/Engine/Shaders/Private/SingleLayerWaterCommon.ush
@@ -130,13 +130,10 @@ WaterVolumeLightingOutput EvaluateWaterVolumeLighting(
 		const float3 ScatteringCoeff = max(0.0f, GetSingleLayerWaterMaterialOutput0(MaterialParameters));
 		const float3 AbsorptionCoeff = max(0.0f, GetSingleLayerWaterMaterialOutput1(MaterialParameters));
 		const float PhaseG = clamp(GetSingleLayerWaterMaterialOutput2(MaterialParameters), -1.0f, 1.0f);
+		const float3 ExtinctionCoeff = ScatteringCoeff + AbsorptionCoeff; 
 		// Max to avoid division by 0 with the analytical integral below.
-#if FEATURE_LEVEL <= FEATURE_LEVEL_ES3_1
 		// 1e-5 is high enough to avoid denorms on mobile
-		const float3 ExtinctionCoeff = max(ScatteringCoeff + AbsorptionCoeff, 1e-5); 
-#else
-		const float3 ExtinctionCoeff = max(ScatteringCoeff + AbsorptionCoeff, 1e-20); 
-#endif
+		const float3 ExtinctionCoeffSafe = max(ScatteringCoeff + AbsorptionCoeff, 1e-5); 
 		
 		float DirLightPhaseValue = 0.0f; // Default when Total Internal Reflection happens.
 		{
@@ -176,7 +173,7 @@ WaterVolumeLightingOutput EvaluateWaterVolumeLighting(
 		const float3 OpticalDepth = ExtinctionCoeff * BehindWaterDeltaDepth;
 		float3 Transmittance = exp(-OpticalDepth);
 		float3 ScatteredLuminance = ScatteringCoeff * (AmbScattLuminance + SunScattLuminance * DirectionalLightShadow);
-		ScatteredLuminance = (ScatteredLuminance - ScatteredLuminance * Transmittance) / ExtinctionCoeff;
+		ScatteredLuminance = (ScatteredLuminance - ScatteredLuminance * Transmittance) / ExtinctionCoeffSafe;
 #else
 		// TODO Make the volumetric shadow part work again
 		float3 Transmittance = 1.0f;
@@ -196,7 +193,7 @@ WaterVolumeLightingOutput EvaluateWaterVolumeLighting(
 
 			float3 SampleTransmittance = exp(-ExtinctionCoeff * RayMarchStepSize); // Constant
 			float3 SS = (ScatteringCoeff * TransmittanceToLightThroughWater * (SunScattLuminance * Shadow + AmbScattLuminance));
-			ScatteredLuminance += Transmittance * (SS - SS * SampleTransmittance) / ExtinctionCoeff;
+			ScatteredLuminance += Transmittance * (SS - SS * SampleTransmittance) / ExtinctionCoeffSafe;
 			Transmittance *= SampleTransmittance;
 		}
 
@@ -206,7 +203,7 @@ WaterVolumeLightingOutput EvaluateWaterVolumeLighting(
 		{
 			float3 Transmittance2 = exp(-OpticalDepth2);
 			float3 ScatteredLuminance2 = ScatteringCoeff * (SunScattLuminance + AmbScattLuminance);
-			ScatteredLuminance += Transmittance * (ScatteredLuminance2 - ScatteredLuminance2 * Transmittance2) / ExtinctionCoeff;
+			ScatteredLuminance += Transmittance * (ScatteredLuminance2 - ScatteredLuminance2 * Transmittance2) / ExtinctionCoeffSafe;
 			Transmittance *= Transmittance2;
 		}
 #endif
