Switch and mobile water optimisation
 - no refraction, no phase function
 - no shadow on water
In DXBC, instruction count goes from 683 to 585

TO BE TESTED ON SWITCH BY Mickael
No regression seen on:
 - Tested on PC with switch forward mode enabled.
 - PC deferred and forward mode tested and still working.
 - PS4 tested and working as without perf regression noticed.

diff --git a/Engine/Shaders/Private/BasePassPixelShader.usf b/Engine/Shaders/Private/BasePassPixelShader.usf
--- a/Engine/Shaders/Private/BasePassPixelShader.usf
+++ b/Engine/Shaders/Private/BasePassPixelShader.usf
@@ -84,10 +84,6 @@ float NormalCurvatureToRoughness(float3 WorldNormal)
 #define EDITOR_ALPHA2COVERAGE (USE_EDITOR_COMPOSITING && SUPPORTS_PIXEL_COVERAGE)
 #define POST_PROCESS_SUBSURFACE ((MATERIAL_SHADINGMODEL_SUBSURFACE_PROFILE || MATERIAL_SHADINGMODEL_EYE) && USES_GBUFFER)
 
-#if TRANSLUCENCY_LIGHTING_SURFACE_LIGHTINGVOLUME || TRANSLUCENCY_LIGHTING_SURFACE_FORWARDSHADING || FORWARD_SHADING || MATERIAL_SHADINGMODEL_SINGLELAYERWATER
-#include "ForwardLightingCommon.ush"
-#endif
-
 #if  MATERIAL_SHADINGMODEL_SINGLELAYERWATER
 
 #if SINGLE_LAYER_WATER_SIMPLE_FORWARD
@@ -95,6 +91,9 @@ float NormalCurvatureToRoughness(float3 WorldNormal)
 	// Use mobile like simple shading
 	#define SIMPLE_SINGLE_LAYER_WATER 1
 
+	// Remove forward lighting directional light shadow
+	#define DISABLE_FORWARD_DIRECTIONAL_LIGHT_SHADOW 1
+
 	// Change the blend mode here from opaque to pre-multiplied-alpha for simplicity
 	#undef MATERIALBLENDING_ALPHACOMPOSITE
 	#define MATERIALBLENDING_ALPHACOMPOSITE 1
@@ -107,6 +106,10 @@ float NormalCurvatureToRoughness(float3 WorldNormal)
 
 #endif // MATERIAL_SHADINGMODEL_SINGLELAYERWATER
 
+#if TRANSLUCENCY_LIGHTING_SURFACE_LIGHTINGVOLUME || TRANSLUCENCY_LIGHTING_SURFACE_FORWARDSHADING || FORWARD_SHADING || MATERIAL_SHADINGMODEL_SINGLELAYERWATER
+#include "ForwardLightingCommon.ush"
+#endif
+
 #if !FORWARD_SHADING
 
 void GetVolumeLightingNonDirectional(float4 AmbientLightingVector, float3 DiffuseColor, inout float3 InterpolatedLighting, out float4 VolumeLighting)
diff --git a/Engine/Shaders/Private/ForwardLightingCommon.ush b/Engine/Shaders/Private/ForwardLightingCommon.ush
--- a/Engine/Shaders/Private/ForwardLightingCommon.ush
+++ b/Engine/Shaders/Private/ForwardLightingCommon.ush
@@ -13,6 +13,10 @@
 #define FILTER_DIRECTIONAL_LIGHT_SHADOWING 1
 #include "ForwardShadowingCommon.ush"
 
+#ifndef DISABLE_FORWARD_DIRECTIONAL_LIGHT_SHADOW
+#define DISABLE_FORWARD_DIRECTIONAL_LIGHT_SHADOW 0
+#endif
+
 float PrevSceneColorPreExposureInv;
 
 void UpdateNearestSample(float Z, float2 UV, float FullResZ, inout float MinDist, inout float2 NearestUV)
@@ -133,7 +137,9 @@ float3 GetForwardDirectLighting(uint GridIndex, float3 WorldPosition, float3 Cam
 		LightData.ShadowMapChannelMask = UnpackShadowMapChannelMask(DirectionalLightData.DirectionalLightShadowMapChannelMask);
 
 		// We want to force the directional light shadow when using water material to see shadow on the water. This could be an option later.
-		#if ((MATERIALBLENDING_SOLID || MATERIALBLENDING_MASKED) && !MATERIAL_SHADINGMODEL_SINGLELAYERWATER)
+		#if DISABLE_FORWARD_DIRECTIONAL_LIGHT_SHADOW
+			float4 LightAttenuation = float4(1, 1, 1, 1);
+		#elif ((MATERIALBLENDING_SOLID || MATERIALBLENDING_MASKED) && !MATERIAL_SHADINGMODEL_SINGLELAYERWATER)
 			float4 PreviewShadowMapChannelMask = UnpackShadowMapChannelMask(DirectionalLightData.DirectionalLightShadowMapChannelMask >> 4);
 			float DynamicShadowing = dot(PreviewShadowMapChannelMask, DynamicShadowFactors);
 
diff --git a/Engine/Shaders/Private/MobileBasePassPixelShader.usf b/Engine/Shaders/Private/MobileBasePassPixelShader.usf
--- a/Engine/Shaders/Private/MobileBasePassPixelShader.usf
+++ b/Engine/Shaders/Private/MobileBasePassPixelShader.usf
@@ -500,8 +500,9 @@ void Main(
 		IndirectIrradiance *= MaterialAO;
 	#endif
 
-	#if !MATERIAL_SHADINGMODEL_UNLIT
-		half Shadow = GetPrimaryPrecomputedShadowMask(Interpolants).r;
+		half Shadow = 1.0f;
+	#if !MATERIAL_SHADINGMODEL_UNLIT && !MATERIAL_SHADINGMODEL_SINGLELAYERWATER
+		Shadow = GetPrimaryPrecomputedShadowMask(Interpolants).r;
 		#if DIRECTIONAL_LIGHT_CSM && !MOBILE_MULTI_VIEW
 		// Cascaded Shadow Map
 		{
diff --git a/Engine/Shaders/Private/SingleLayerWaterCommon.ush b/Engine/Shaders/Private/SingleLayerWaterCommon.ush
--- a/Engine/Shaders/Private/SingleLayerWaterCommon.ush
+++ b/Engine/Shaders/Private/SingleLayerWaterCommon.ush
@@ -137,11 +137,11 @@ WaterVolumeLightingOutput EvaluateWaterVolumeLighting(
 		
 		float DirLightPhaseValue = 0.0f; // Default when Total Internal Reflection happens.
 		{
+#if SIMPLE_SINGLE_LAYER_WATER
+			DirLightPhaseValue = IsotropicPhase();
+#else
 			float IorFrom = 1.0f; // assumes we come from air
 			float IorTo   = DielectricF0ToIor(DielectricSpecularToF0(Specular)); // Wrong if metal is set to >1. But we still keep refraction on the water surface nonetheless.
-#if 0
-			DirLightPhaseValue = SchlickPhase(PhaseG, dot(-ResolvedView.DirectionalLightDirection.xyz, refract(MaterialParameters.CameraVector, MaterialParameters.WorldNormal, IorFrom / IorTo)));
-#else
 			const float relativeIOR = IorFrom / IorTo;
 			float3 UnderWaterRayDir = 0.0f;
 			if (WaterRefract(MaterialParameters.CameraVector, MaterialParameters.WorldNormal, relativeIOR, UnderWaterRayDir))
@@ -169,11 +169,14 @@ WaterVolumeLightingOutput EvaluateWaterVolumeLighting(
 		float3 SunScattLuminance = DirLightPhaseValue * SunIlluminance;
 		float3 AmbScattLuminance = IsotropicPhase()   * AmbiantIlluminance;
 
-#if 1 || SIMPLE_SINGLE_LAYER_WATER
+#define VOLUMETRICSHADOW	0
+#if !VOLUMETRICSHADOW || SIMPLE_SINGLE_LAYER_WATER
+
 		const float3 OpticalDepth = ExtinctionCoeff * BehindWaterDeltaDepth;
 		float3 Transmittance = exp(-OpticalDepth);
 		float3 ScatteredLuminance = ScatteringCoeff * (AmbScattLuminance + SunScattLuminance * DirectionalLightShadow);
 		ScatteredLuminance = (ScatteredLuminance - ScatteredLuminance * Transmittance) / ExtinctionCoeffSafe;
+
 #else
 		// TODO Make the volumetric shadow part work again
 		float3 Transmittance = 1.0f;
