First pass of a fast point cloud render component (single draw, verts implied from point location, point coloring globally or per point)

diff --git a/Engine/Plugins/Experimental/AR/PointCloud/PointCloud.uplugin b/Engine/Plugins/Experimental/AR/PointCloud/PointCloud.uplugin
new file mode 100644
--- /dev/null
+++ b/Engine/Plugins/Experimental/AR/PointCloud/PointCloud.uplugin
@@ -0,0 +1,25 @@
+{
+	"FileVersion" : 3,
+	"Version" : 1,
+	"VersionName" : "1.0",
+	"FriendlyName" : "Point Cloud support for AR",
+	"Description" : "An efficient way to render point cloud data that is generated by AR devices",
+	"Category" : "Experimental",
+	"CreatedBy" : "Epic Games, Inc.",
+	"CreatedByURL" : "http://epicgames.com",
+	"DocsURL" : "",
+	"MarketplaceURL" : "",
+	"SupportURL" : "",
+	"EnabledByDefault" : false,
+	"CanContainContent" : true,
+	"IsBetaVersion" : true,
+	"Installed" : false,
+	"Modules" :
+	[
+		{
+			"Name" : "PointCloud",
+			"Type" : "Runtime",
+			"LoadingPhase" : "PostConfigInit"
+		}
+	]
+}
diff --git a/Engine/Plugins/Experimental/AR/PointCloud/Source/PointCloud/PointCloud.Build.cs b/Engine/Plugins/Experimental/AR/PointCloud/Source/PointCloud/PointCloud.Build.cs
new file mode 100644
--- /dev/null
+++ b/Engine/Plugins/Experimental/AR/PointCloud/Source/PointCloud/PointCloud.Build.cs
@@ -0,0 +1,54 @@
+// Copyright 1998-2018 Epic Games, Inc. All Rights Reserved.
+
+namespace UnrealBuildTool.Rules
+{
+	public class PointCloud : ModuleRules
+	{
+		public PointCloud(ReadOnlyTargetRules Target) : base(Target)
+		{
+            PublicIncludePaths.AddRange(
+                new string[] {
+                    // ... add public include paths required here ...
+                }
+                );
+
+
+            PrivateIncludePaths.AddRange(
+                new string[] {
+                    "PointCloud/Public",
+                    "PointCloud/Private"
+                    // ... add other private include paths required here ...
+                }
+                );
+
+
+            PublicDependencyModuleNames.AddRange(
+                new string[]
+                {
+                    "Core",
+                    "Engine",
+					"RenderCore",
+					"ShaderCore",
+					"RHI",
+                    "AugmentedReality",
+                    // ... add other public dependencies that you statically link with here ...
+                }
+                );
+
+			PrivateDependencyModuleNames.AddRange(
+				new string[]
+				{
+    				"CoreUObject"
+					// ... add private dependencies that you statically link with here ...
+				}
+				);
+
+			DynamicallyLoadedModuleNames.AddRange(
+				new string[]
+				{
+					// ... add any modules that your module loads dynamically here ...
+				}
+				);
+		}
+	}
+}
diff --git a/Engine/Plugins/Experimental/AR/PointCloud/Source/PointCloud/Private/PointCloudBuffers.h b/Engine/Plugins/Experimental/AR/PointCloud/Source/PointCloud/Private/PointCloudBuffers.h
new file mode 100644
--- /dev/null
+++ b/Engine/Plugins/Experimental/AR/PointCloud/Source/PointCloud/Private/PointCloudBuffers.h
@@ -0,0 +1,232 @@
+// Copyright 1998-2018 Epic Games, Inc. All Rights Reserved.
+
+#pragma once
+
+#include "Containers/ResourceArray.h"
+
+class FPointCloudVertexResourceArray :
+	public FResourceArrayInterface
+{
+public:
+	FPointCloudVertexResourceArray(const void* InData, uint32 InSize)
+		: Data(InData)
+		, Size(InSize)
+	{
+	}
+
+	virtual const void* GetResourceData() const override { return Data; }
+	virtual uint32 GetResourceDataSize() const override { return Size; }
+	virtual void Discard() override { }
+	virtual bool IsStatic() const override { return false; }
+	virtual bool GetAllowCPUAccess() const override { return false; }
+	virtual void SetAllowCPUAccess(bool bInNeedsCPUAccess) override { }
+
+private:
+	const void* Data;
+	uint32 Size;
+};
+
+/** Point cloud vertex buffer that can hold an arbitrary single data type (color or position) */
+class FPointCloudVertexBufferBase :
+	public FVertexBuffer
+{
+public:
+	inline uint32 GetNumVerts() const
+	{
+		return NumVerts;
+	}
+
+	inline FShaderResourceViewRHIRef GetBufferSRV() const
+	{
+		return BufferSRV;
+	}
+
+protected:
+	void InitWith(const void* InVertexData, uint32 InNumVerts, uint32 InSizeInBytes)
+	{
+		InitResource();
+
+		NumVerts = InNumVerts;
+
+		FPointCloudVertexResourceArray ResourceArray(InVertexData, InSizeInBytes);
+		FRHIResourceCreateInfo CreateInfo(&ResourceArray);
+		VertexBufferRHI = RHICreateVertexBuffer(InSizeInBytes, BUF_Static | BUF_ShaderResource, CreateInfo);
+	}
+
+	uint32 NumVerts = 0;
+	FShaderResourceViewRHIRef BufferSRV;
+};
+
+/** Point cloud color buffer that sets the SRV too */
+class FPointCloudColorVertexBuffer :
+	public FPointCloudVertexBufferBase
+{
+public:
+	void InitRHIWith(const TArray<FColor>& RawColorData)
+	{
+		InitWith(RawColorData.GetData(), (uint32)RawColorData.Num(), RawColorData.Num() * sizeof(FColor));
+		BufferSRV = RHICreateShaderResourceView(VertexBufferRHI, sizeof(FColor), PF_R8G8B8A8);
+		// We have a full stream so we can index safely
+		ColorMask = ~0;
+	}
+
+	void InitRHIWith(const FColor& RawColor)
+	{
+		InitWith((const void*)&RawColor, 1, sizeof(FColor));
+		BufferSRV = RHICreateShaderResourceView(VertexBufferRHI, sizeof(FColor), PF_R8G8B8A8);
+	}
+
+	inline uint32 GetColorMask() const
+	{
+		return ColorMask;
+	}
+
+private:
+	uint32 ColorMask = 0;
+};
+
+/** Point cloud location buffer that sets the SRV too */
+class FPointCloudLocationVertexBuffer :
+	public FPointCloudVertexBufferBase
+{
+public:
+	void InitRHIWith(const TArray<FVector>& RawLocationData)
+	{
+		InitWith(RawLocationData.GetData(), RawLocationData.Num(), RawLocationData.Num() * sizeof(FVector));
+		BufferSRV = RHICreateShaderResourceView(VertexBufferRHI, sizeof(float), PF_R32_FLOAT);
+	}
+};
+
+/**
+ * We generate a index buffer for N points in the point cloud section. Each point
+ * generates the vert positions using vertex id and point id (vertex id / 4)
+ * by fetching the points and colors from the buffers
+ */
+class FPointCloudIndexBuffer :
+	public FIndexBuffer
+{
+public:
+	FPointCloudIndexBuffer()
+		: NumPoints(0)
+		, NumPrimitives(0)
+		, MaxIndex(0)
+		, bIsQuadList(false)
+	{
+	}
+
+	FPointCloudIndexBuffer(uint32 InNumPoints)
+		: NumPoints(InNumPoints)
+		, NumPrimitives(0)
+		, MaxIndex(4 * InNumPoints)
+		, bIsQuadList(false)
+	{
+	}
+
+	/** Generates a quad list when available on the platform */
+	template <typename INDEX_TYPE>
+	void CreateQuadList()
+	{
+		bIsQuadList = true;
+		NumPrimitives = NumPoints;
+		const uint32 Size = sizeof(INDEX_TYPE) * 4 * NumPoints;
+		const uint32 Stride = sizeof(INDEX_TYPE);
+		FRHIResourceCreateInfo CreateInfo;
+		void* Buffer = nullptr;
+		IndexBufferRHI = RHICreateAndLockIndexBuffer( Stride, Size, BUF_Static, CreateInfo, Buffer );
+		INDEX_TYPE* Indices = (INDEX_TYPE*)Buffer;
+		for (uint32 SpriteIndex = 0; SpriteIndex < NumPoints; ++SpriteIndex)
+		{
+			Indices[SpriteIndex * 4 + 0] = SpriteIndex * 4 + 0;
+			Indices[SpriteIndex * 4 + 1] = SpriteIndex * 4 + 1;
+			Indices[SpriteIndex * 4 + 2] = SpriteIndex * 4 + 3;
+			Indices[SpriteIndex * 4 + 3] = SpriteIndex * 4 + 2;
+		}
+		RHIUnlockIndexBuffer( IndexBufferRHI );
+	}
+
+	/** Generates a tri list when quad lists are not available on the platform */
+	template <typename INDEX_TYPE>
+	void CreateTriList()
+	{
+		NumPrimitives = 2 * NumPoints;
+		const uint32 Size = sizeof(INDEX_TYPE) * 6 * NumPoints;
+		const uint32 Stride = sizeof(INDEX_TYPE);
+		FRHIResourceCreateInfo CreateInfo;
+		void* Buffer = nullptr;
+		IndexBufferRHI = RHICreateAndLockIndexBuffer( Stride, Size, BUF_Static, CreateInfo, Buffer );
+		INDEX_TYPE* Indices = (INDEX_TYPE*)Buffer;
+		for (uint32 SpriteIndex = 0; SpriteIndex < NumPoints; ++SpriteIndex)
+		{
+			Indices[SpriteIndex * 6 + 0] = SpriteIndex * 4 + 0;
+			Indices[SpriteIndex * 6 + 1] = SpriteIndex * 4 + 3;
+			Indices[SpriteIndex * 6 + 2] = SpriteIndex * 4 + 2;
+			Indices[SpriteIndex * 6 + 3] = SpriteIndex * 4 + 0;
+			Indices[SpriteIndex * 6 + 4] = SpriteIndex * 4 + 1;
+			Indices[SpriteIndex * 6 + 5] = SpriteIndex * 4 + 3;
+		}
+		RHIUnlockIndexBuffer( IndexBufferRHI );
+	}
+
+	virtual void InitRHI() override
+	{
+		check(NumPoints > 0 && MaxIndex > 0);
+//@todo joeg - cvar quad vs tri in addition to platform support
+		const bool bShouldUseQuadList = GRHISupportsQuadTopology;
+		// Use 32 bit indices if needed
+		if (NumPoints * 4 > 65535)
+		{
+			if (bShouldUseQuadList)
+			{
+				CreateQuadList<uint32>();
+			}
+			else
+			{
+				CreateTriList<uint32>();
+			}
+		}
+		else
+		{
+			if (bShouldUseQuadList)
+			{
+				CreateQuadList<uint16>();
+			}
+			else
+			{
+				CreateTriList<uint16>();
+			}
+		}
+	}
+
+	void InitRHIWithSize(uint32 InNumPoints)
+	{
+		NumPoints = InNumPoints;
+		MaxIndex = (4 * InNumPoints) - 1;
+		InitResource();
+	}
+
+	inline bool IsQuadList() const
+	{
+		return bIsQuadList;
+	}
+
+	inline bool IsTriList() const
+	{
+		return !IsQuadList();
+	}
+
+	inline uint32 GetNumPrimitives() const
+	{
+		return NumPrimitives;
+	}
+
+	inline uint32 GetMaxIndex() const
+	{
+		return MaxIndex;
+	}
+
+private:
+	uint32 NumPoints;
+	uint32 NumPrimitives;
+	uint32 MaxIndex;
+	bool bIsQuadList;
+};
diff --git a/Engine/Plugins/Experimental/AR/PointCloud/Source/PointCloud/Private/PointCloudComponent.cpp b/Engine/Plugins/Experimental/AR/PointCloud/Source/PointCloud/Private/PointCloudComponent.cpp
new file mode 100644
--- /dev/null
+++ b/Engine/Plugins/Experimental/AR/PointCloud/Source/PointCloud/Private/PointCloudComponent.cpp
@@ -0,0 +1,148 @@
+// Copyright 1998-2018 Epic Games, Inc. All Rights Reserved.
+
+#include "PointCloudComponent.h"
+#include "PointCloud.h"
+#include "PointCloudSceneProxy.h"
+
+DECLARE_CYCLE_STAT(TEXT("Create Point Cloud Proxy"), STAT_PointCloud_CreateSceneProxy, STATGROUP_POINTCLOUD);
+DECLARE_CYCLE_STAT(TEXT("Point Cloud Comp Update"), STAT_PointCloud_ComponentUpdateCost, STATGROUP_POINTCLOUD);
+
+UPointCloudComponent::UPointCloudComponent(const FObjectInitializer& ObjectInitializer)
+	: Super(ObjectInitializer)
+	, PointColor(1.f, 1.f, 1.f, 1.f)
+	, PointSize(1.f)
+{
+	PrimaryComponentTick.bCanEverTick = true;
+	PrimaryComponentTick.TickGroup = TG_PostPhysics;
+
+	// Load our material we use f
+	static ConstructorHelpers::FObjectFinder<UMaterialInterface> DefaultMaterial(TEXT("/PointCloud/DefaultPointCloudMaterial"));
+	PointCloudMaterial = DefaultMaterial.Object;
+	if (PointCloudMaterial == nullptr)
+	{
+		PointCloudMaterial = GEngine->WireframeMaterial;
+	}
+}
+
+void UPointCloudComponent::ClearPointCloud()
+{
+	PointCloud.Empty();
+	PointColors.Empty();
+	WorldBounds = FBoxSphereBounds(FVector(0, 0, 0), FVector(0, 0, 0), 0);
+
+	MarkRenderStateDirty();
+}
+
+FPrimitiveSceneProxy* UPointCloudComponent::CreateSceneProxy()
+{
+	SCOPE_CYCLE_COUNTER(STAT_PointCloud_CreateSceneProxy);
+
+	// Use the fast version if supported
+	if (PointCloud.Num() > 0 && RHISupportsManualVertexFetch(GMaxRHIShaderPlatform))
+	{
+		return new FPointCloudSceneProxy(this);
+	}
+	return new FNoFetchPointCloudSceneProxy(this);
+}
+
+FBoxSphereBounds UPointCloudComponent::CalcBounds(const FTransform& LocalToWorld) const
+{
+	return WorldBounds;
+}
+
+FMatrix UPointCloudComponent::GetRenderMatrix() const
+{
+	return FMatrix::Identity;
+}
+
+void UPointCloudComponent::SetPointCloud(const TArray<FVector>& Points)
+{
+	SCOPE_CYCLE_COUNTER(STAT_PointCloud_ComponentUpdateCost);
+
+	// Zero and rebuild our bounds from the points
+	FBox PointBounds(Points);
+	WorldBounds = FBoxSphereBounds(PointBounds);
+	PointCloud = Points;
+
+	// An update won't be faster than a rebuild, so force a rebuild
+	MarkRenderStateDirty();
+}
+
+void UPointCloudComponent::SetPointCloudWithColors(const TArray<FVector>& Points, const TArray<FColor>& Colors)
+{
+	{
+		SCOPE_CYCLE_COUNTER(STAT_PointCloud_ComponentUpdateCost);
+		PointColors = Colors;
+	}
+	SetPointCloud(Points);
+}
+
+void UPointCloudComponent::SetPointColor(const FLinearColor& Color)
+{
+	PointColor = Color;
+//@todo joeg - Change to a render thread call on the proxy
+	MarkRenderStateDirty();
+}
+
+void UPointCloudComponent::SetPointSize(float Size)
+{
+	PointSize = Size;
+	if (PointSize <= 0.f)
+	{
+		PointSize = 1.f;
+	}
+//@todo joeg - Change to a render thread call on the proxy
+	MarkRenderStateDirty();
+}
+
+TArray<FVector> UPointCloudComponent::GetPointsInBox(const FBox& WorldSpaceBox) const
+{
+	TArray<FVector> OutPoints;
+	// Miminize the number of allocations by presizing the max
+	OutPoints.Empty(PointCloud.Num());
+
+	for (const FVector& Point : PointCloud)
+	{
+		if (WorldSpaceBox.IsInsideOrOn(Point))
+		{
+			OutPoints.Add(Point);
+		}
+	}
+
+	return OutPoints;
+}
+
+TArray<FVector> UPointCloudComponent::GetPointsOutsideBox(const FBox& WorldSpaceBox) const
+{
+	TArray<FVector> OutPoints;
+	// Miminize the number of allocations by presizing the max
+	OutPoints.Empty(PointCloud.Num());
+	
+	for (const FVector& Point : PointCloud)
+	{
+		if (!WorldSpaceBox.IsInsideOrOn(Point))
+		{
+			OutPoints.Add(Point);
+		}
+	}
+	
+	return OutPoints;
+}
+
+void UPointCloudComponent::SetIsVisible(bool bNewVisibity)
+{
+	if (bNewVisibity != bIsVisible)
+	{
+		bIsVisible = bNewVisibity;
+		
+		MarkRenderStateDirty();
+	}
+}
+
+void UPointCloudComponent::GetUsedMaterials(TArray<UMaterialInterface*>& OutMaterials, bool) const
+{
+	if (PointCloudMaterial != nullptr)
+	{
+		OutMaterials.Add(PointCloudMaterial);
+	}
+}
diff --git a/Engine/Plugins/Experimental/AR/PointCloud/Source/PointCloud/Private/PointCloudModule.cpp b/Engine/Plugins/Experimental/AR/PointCloud/Source/PointCloud/Private/PointCloudModule.cpp
new file mode 100644
--- /dev/null
+++ b/Engine/Plugins/Experimental/AR/PointCloud/Source/PointCloud/Private/PointCloudModule.cpp
@@ -0,0 +1,23 @@
+// Copyright 1998-2018 Epic Games, Inc. All Rights Reserved.
+
+#include "PointCloud.h"
+
+DEFINE_LOG_CATEGORY(LogPointCloud);
+
+class FPointCloudModule :
+	public IModuleInterface
+{
+	/** IModuleInterface implementation */
+	virtual void StartupModule() override;
+	virtual void ShutdownModule() override;
+};
+
+IMPLEMENT_MODULE(FPointCloudModule, PointCloud);
+
+void FPointCloudModule::StartupModule()
+{
+}
+
+void FPointCloudModule::ShutdownModule()
+{
+}
diff --git a/Engine/Plugins/Experimental/AR/PointCloud/Source/PointCloud/Private/PointCloudSceneProxy.cpp b/Engine/Plugins/Experimental/AR/PointCloud/Source/PointCloud/Private/PointCloudSceneProxy.cpp
new file mode 100644
--- /dev/null
+++ b/Engine/Plugins/Experimental/AR/PointCloud/Source/PointCloud/Private/PointCloudSceneProxy.cpp
@@ -0,0 +1,235 @@
+// Copyright 1998-2018 Epic Games, Inc. All Rights Reserved.
+
+#include "PointCloudSceneProxy.h"
+#include "PointCloud.h"
+#include "PointCloudComponent.h"
+
+#include "PrimitiveViewRelevance.h"
+#include "RHI.h"
+#include "RenderResource.h"
+#include "RenderingThread.h"
+#include "Containers/ResourceArray.h"
+#include "EngineGlobals.h"
+#include "VertexFactory.h"
+#include "MaterialShared.h"
+#include "Materials/Material.h"
+#include "LocalVertexFactory.h"
+#include "Engine/Engine.h"
+#include "SceneManagement.h"
+#include "DynamicMeshBuilder.h"
+
+DECLARE_CYCLE_STAT(TEXT("Update Point Cloud GT"), STAT_PointCloud_UpdatePointCloud, STATGROUP_POINTCLOUD);
+DECLARE_CYCLE_STAT(TEXT("Get Mesh Elements"), STAT_PointCloud_GetMeshElements, STATGROUP_POINTCLOUD);
+DECLARE_CYCLE_STAT(TEXT("Create RT Resources"), STAT_PointCloud_CreateRenderThreadResources, STATGROUP_POINTCLOUD);
+
+
+FNoFetchPointCloudSceneProxy::FNoFetchPointCloudSceneProxy(UPointCloudComponent* Component)
+	: FPrimitiveSceneProxy(Component)
+	, Points(Component->PointCloud)
+	, Color(Component->PointColor)
+	, Size(Component->PointSize)
+	, bIsVisible(Component->bIsVisible)
+	, MaterialRelevance(Component->GetMaterialRelevance(GetScene().GetFeatureLevel()))
+{
+	if (Size <= 0.f)
+	{
+		Size = 1.f;
+	}
+}
+
+FNoFetchPointCloudSceneProxy::~FNoFetchPointCloudSceneProxy()
+{
+}
+
+SIZE_T FNoFetchPointCloudSceneProxy::GetTypeHash() const
+{
+	static size_t UniquePointer;
+	return reinterpret_cast<size_t>(&UniquePointer);
+}
+
+void FNoFetchPointCloudSceneProxy::GetDynamicMeshElements(const TArray<const FSceneView*>& Views, const FSceneViewFamily& ViewFamily, uint32 VisibilityMap, FMeshElementCollector& Collector) const
+{
+	SCOPE_CYCLE_COUNTER(STAT_PointCloud_GetMeshElements);
+	
+	if (!bIsVisible)
+	{
+		return;
+	}
+
+	const bool bIsWireframe = AllowDebugViewmodes() && ViewFamily.EngineShowFlags.Wireframe;
+
+	for (int32 ViewIndex = 0; ViewIndex < Views.Num(); ViewIndex++)
+	{
+		if (VisibilityMap & (1 << ViewIndex))
+		{
+			if (bIsWireframe)
+			{
+				// Draw bounds around the points
+				RenderBounds(Collector.GetPDI(ViewIndex), ViewFamily.EngineShowFlags, GetBounds(), IsSelected());
+			}
+			FPrimitiveDrawInterface* PDI = Collector.GetPDI(ViewIndex);
+			// Loop through manually drawing the points since PT_PointList isn't really a supported primitive type
+			for (const FVector& Point : Points)
+			{
+				PDI->DrawPoint(Point, Color, Size, SDPG_World);
+			}
+		}
+	}
+}
+
+FPrimitiveViewRelevance FNoFetchPointCloudSceneProxy::GetViewRelevance(const FSceneView* View) const
+{
+	FPrimitiveViewRelevance Result;
+	Result.bDrawRelevance = IsShown(View) && bIsVisible;
+	Result.bShadowRelevance = false;
+	Result.bDynamicRelevance = true;
+	Result.bRenderInMainPass = true;
+	Result.bUsesLightingChannels = false;
+	Result.bRenderCustomDepth = false;
+	MaterialRelevance.SetPrimitiveViewRelevance(Result);
+	return Result;
+}
+
+uint32 FNoFetchPointCloudSceneProxy::GetMemoryFootprint(void) const
+{
+	return(sizeof(*this) + GetAllocatedSize());
+}
+
+uint32 FNoFetchPointCloudSceneProxy::GetAllocatedSize(void) const
+{
+	return FPrimitiveSceneProxy::GetAllocatedSize();
+}
+
+
+FPointCloudSceneProxy::FPointCloudSceneProxy(UPointCloudComponent* Component)
+	: FPrimitiveSceneProxy(Component)
+	, bIsVisible(Component->bIsVisible)
+	, Points(Component->PointCloud)
+	, Colors(Component->PointColors)
+	, PointColor(Component->PointColor)
+	, PointSize(Component->PointSize)
+	, Material(Component->PointCloudMaterial)
+	, MaterialRelevance(Component->GetMaterialRelevance(GetScene().GetFeatureLevel()))
+	, PointCloudVertexFactory(GetScene().GetFeatureLevel())
+{
+	if (PointSize <= 0.f)
+	{
+		PointSize = 1.f;
+	}
+}
+
+FPointCloudSceneProxy::~FPointCloudSceneProxy()
+{
+	PointCloudVertexFactory.ReleaseResource();
+	PointCloudIndexBuffer.ReleaseResource();
+	PointCloudColorVertexBuffer.ReleaseResource();
+	PointCloudLocationVertexBuffer.ReleaseResource();
+}
+
+void FPointCloudSceneProxy::CreateRenderThreadResources()
+{
+	SCOPE_CYCLE_COUNTER(STAT_PointCloud_CreateRenderThreadResources);
+
+	PointCloudVertexFactory.InitResource();
+	PointCloudIndexBuffer.InitRHIWithSize(Points.Num());
+
+	// We either use a single color or the color array
+	if (Colors.Num())
+	{
+		PointCloudColorVertexBuffer.InitRHIWith(Colors);
+		Colors.Empty();
+	}
+	else
+	{
+		PointCloudColorVertexBuffer.InitRHIWith(PointColor.ToFColor(false));
+	}
+
+	PointCloudLocationVertexBuffer.InitRHIWith(Points);
+	Points.Empty();
+
+	// Setup the vertex factory shader parameters
+	FPointCloudVertexFactoryParameters UniformParameters;
+	UniformParameters.VertexFetch_PointLocationBuffer = PointCloudLocationVertexBuffer.GetBufferSRV();
+	UniformParameters.VertexFetch_PointColorBuffer = PointCloudColorVertexBuffer.GetBufferSRV();
+	PointCloudVertexFactory.SetParameters(UniformParameters, PointCloudColorVertexBuffer.GetColorMask(), PointSize);
+}
+
+SIZE_T FPointCloudSceneProxy::GetTypeHash() const
+{
+	static size_t UniquePointer;
+	return reinterpret_cast<size_t>(&UniquePointer);
+}
+
+void FPointCloudSceneProxy::GetDynamicMeshElements(const TArray<const FSceneView*>& Views, const FSceneViewFamily& ViewFamily, uint32 VisibilityMap, FMeshElementCollector& Collector) const
+{
+	SCOPE_CYCLE_COUNTER(STAT_PointCloud_GetMeshElements);
+
+	if (!bIsVisible)
+	{
+		return;
+	}
+
+	const bool bIsWireframe = AllowDebugViewmodes() && ViewFamily.EngineShowFlags.Wireframe;
+	FMaterialRenderProxy* MaterialProxy = Material->GetRenderProxy(IsSelected());
+	if (bIsWireframe)
+	{
+		FMaterialRenderProxy* WireframeMaterialInstance = new FColoredMaterialRenderProxy(
+				GEngine->WireframeMaterial ? GEngine->WireframeMaterial->GetRenderProxy(IsSelected()) : nullptr,
+				FLinearColor(0, 0.5f, 1.f)
+		);
+		Collector.RegisterOneFrameMaterialProxy(WireframeMaterialInstance);
+		MaterialProxy = WireframeMaterialInstance;
+	}
+	// Nothing to render with
+	if (MaterialProxy == nullptr)
+	{
+		return;
+	}
+
+	for (int32 ViewIndex = 0; ViewIndex < Views.Num(); ViewIndex++)
+	{
+		if (VisibilityMap & (1 << ViewIndex))
+		{
+			if (bIsWireframe)
+			{
+				// Draw bounds around the points
+				RenderBounds(Collector.GetPDI(ViewIndex), ViewFamily.EngineShowFlags, GetBounds(), IsSelected());
+			}
+
+			// Create a mesh batch for this chunk of point cloud
+			FMeshBatch& MeshBatch = Collector.AllocateMesh();
+			MeshBatch.CastShadow = false;
+			MeshBatch.bUseAsOccluder = false;
+			MeshBatch.VertexFactory = &PointCloudVertexFactory;
+			MeshBatch.MaterialRenderProxy = MaterialProxy;
+			MeshBatch.ReverseCulling = IsLocalToWorldDeterminantNegative();
+			MeshBatch.DepthPriorityGroup = SDPG_World;
+
+			// Set up index buffer
+			MeshBatch.Type = PointCloudIndexBuffer.IsTriList() ? PT_TriangleList : PT_QuadList;
+			FMeshBatchElement& BatchElement = MeshBatch.Elements[0];
+			BatchElement.FirstIndex = 0;
+			BatchElement.NumPrimitives = PointCloudIndexBuffer.GetNumPrimitives();
+			BatchElement.MinVertexIndex = 0;
+			BatchElement.MaxVertexIndex = PointCloudIndexBuffer.GetMaxIndex();
+			BatchElement.IndexBuffer = &PointCloudIndexBuffer;
+			BatchElement.PrimitiveUniformBufferResource = &GIdentityPrimitiveUniformBuffer;
+
+			MeshBatch.bCanApplyViewModeOverrides = false;
+			Collector.AddMesh(ViewIndex, MeshBatch);
+		}
+	}
+}
+
+FPrimitiveViewRelevance FPointCloudSceneProxy::GetViewRelevance(const FSceneView* View) const
+{
+	FPrimitiveViewRelevance Result;
+	Result.bDrawRelevance = IsShown(View) && bIsVisible;
+	Result.bShadowRelevance = false;
+	Result.bDynamicRelevance = true;
+	Result.bRenderInMainPass = true;
+	Result.bUsesLightingChannels = false;
+	Result.bRenderCustomDepth = false;
+	MaterialRelevance.SetPrimitiveViewRelevance(Result);
+	return Result;
+}
diff --git a/Engine/Plugins/Experimental/AR/PointCloud/Source/PointCloud/Private/PointCloudSceneProxy.h b/Engine/Plugins/Experimental/AR/PointCloud/Source/PointCloud/Private/PointCloudSceneProxy.h
new file mode 100644
--- /dev/null
+++ b/Engine/Plugins/Experimental/AR/PointCloud/Source/PointCloud/Private/PointCloudSceneProxy.h
@@ -0,0 +1,88 @@
+// Copyright 1998-2018 Epic Games, Inc. All Rights Reserved.
+
+#pragma once
+
+#include "CoreMinimal.h"
+#include "PrimitiveSceneProxy.h"
+#include "PointCloudVertexFactory.h"
+#include "PointCloudBuffers.h"
+
+class UPointCloudComponent;
+
+/** Point cloud rendering when vertex fetch is not supported, aka the slow way */
+class FNoFetchPointCloudSceneProxy final :
+	public FPrimitiveSceneProxy
+{
+public:
+	FNoFetchPointCloudSceneProxy(UPointCloudComponent* Component);
+	virtual ~FNoFetchPointCloudSceneProxy();
+	
+	SIZE_T GetTypeHash() const override;
+	
+	virtual void GetDynamicMeshElements(const TArray<const FSceneView*>& Views, const FSceneViewFamily& ViewFamily, uint32 VisibilityMap, FMeshElementCollector& Collector) const override;
+	virtual FPrimitiveViewRelevance GetViewRelevance(const FSceneView* View) const;
+	
+	virtual uint32 GetMemoryFootprint(void) const;
+	uint32 GetAllocatedSize(void) const;
+	
+private:
+	TArray<FVector> Points;
+	FLinearColor Color;
+	float Size;
+	bool bIsVisible;
+	FMaterialRelevance MaterialRelevance;
+};
+
+
+/** Point cloud rendering when manual vertex fetch is supported */
+class FPointCloudSceneProxy :
+	public FPrimitiveSceneProxy
+{
+public:
+	FPointCloudSceneProxy(UPointCloudComponent* Component);
+	virtual ~FPointCloudSceneProxy();
+
+	SIZE_T GetTypeHash() const override;
+
+	virtual void GetDynamicMeshElements(const TArray<const FSceneView*>& Views, const FSceneViewFamily& ViewFamily, uint32 VisibilityMap, FMeshElementCollector& Collector) const override;
+	virtual FPrimitiveViewRelevance GetViewRelevance(const FSceneView* View) const;
+
+	virtual uint32 GetMemoryFootprint(void) const override
+	{
+		return(sizeof(*this) + GetAllocatedSize());
+	}
+
+	uint32 GetAllocatedSize(void) const
+	{
+		return FPrimitiveSceneProxy::GetAllocatedSize();
+	}
+
+	/**
+	 *	Called when the rendering thread adds the proxy to the scene.
+	 *	This function allows for generating renderer-side resources.
+	 *	Called in the rendering thread.
+	 */
+	virtual void CreateRenderThreadResources() override;
+
+private:
+	bool bIsVisible;
+	/** Temp array until the render resource gets created */
+	TArray<FVector> Points;
+	/** Temp array until the render resource gets created */
+	TArray<FColor> Colors;
+	/** The linear color to use to draw all points */
+	FLinearColor PointColor;
+	/** The point size to render at */
+	float PointSize;
+	/** The material from the component to render with */
+	UMaterialInterface* Material;
+	/** The index buffer to use when drawing */
+	FPointCloudIndexBuffer PointCloudIndexBuffer;
+	/** The vertex buffer of colors for each point */
+	FPointCloudColorVertexBuffer PointCloudColorVertexBuffer;
+	/** The vertex buffer of locations for each point */
+	FPointCloudLocationVertexBuffer PointCloudLocationVertexBuffer;
+
+	FMaterialRelevance MaterialRelevance;
+	FPointCloudVertexFactory PointCloudVertexFactory;
+};
diff --git a/Engine/Plugins/Experimental/AR/PointCloud/Source/PointCloud/Private/PointCloudVertexFactory.cpp b/Engine/Plugins/Experimental/AR/PointCloud/Source/PointCloud/Private/PointCloudVertexFactory.cpp
new file mode 100644
--- /dev/null
+++ b/Engine/Plugins/Experimental/AR/PointCloud/Source/PointCloud/Private/PointCloudVertexFactory.cpp
@@ -0,0 +1,140 @@
+// Copyright 1998-2018 Epic Games, Inc. All Rights Reserved.
+
+#include "PointCloudVertexFactory.h"
+#include "RHIStaticStates.h"
+#include "SceneManagement.h"
+#include "EngineGlobals.h"
+#include "Engine/Engine.h"
+#include "Materials/Material.h"
+#include "ShaderParameterUtils.h"
+
+IMPLEMENT_UNIFORM_BUFFER_STRUCT(FPointCloudVertexFactoryParameters, TEXT("PointCloudVF"));
+
+/**
+ * Shader parameters for the vector field visualization vertex factory.
+ */
+class FPointCloudVertexFactoryShaderParameters :
+	public FVertexFactoryShaderParameters
+{
+public:
+	virtual void Bind(const FShaderParameterMap& ParameterMap) override
+	{
+		ColorMask.Bind(ParameterMap, TEXT("ColorMask"));
+		PointSize.Bind(ParameterMap, TEXT("PointSize"));
+	}
+
+	virtual void Serialize(FArchive& Ar) override
+	{
+		Ar << ColorMask;
+		Ar << PointSize;
+	}
+
+	virtual void SetMesh(FRHICommandList& RHICmdList, FShader* Shader, const FVertexFactory* InVertexFactory, const FSceneView& View, const FMeshBatchElement& BatchElement, uint32 DataFlags) const override
+	{
+		FPointCloudVertexFactory* VertexFactory = (FPointCloudVertexFactory*)InVertexFactory;
+		FVertexShaderRHIParamRef VertexShaderRHI = Shader->GetVertexShader();
+
+		SetUniformBufferParameter(RHICmdList, VertexShaderRHI, Shader->GetUniformBufferParameter<FPointCloudVertexFactoryParameters>(), VertexFactory->GetPointCloudVertexFactoryUniformBuffer());
+
+		SetShaderValue(RHICmdList, VertexShaderRHI, ColorMask, VertexFactory->GetColorMask());
+		SetShaderValue(RHICmdList, VertexShaderRHI, PointSize, VertexFactory->GetPointSize());
+	}
+
+	virtual uint32 GetSize() const override { return sizeof(*this); }
+
+private:
+	FShaderParameter ColorMask;
+	FShaderParameter PointSize;
+};
+
+/**
+ * Vertex declaration for point clouds. Verts aren't used so this is to make complaints go away
+ */
+class FPointCloudVertexDeclaration :
+	public FRenderResource
+{
+public:
+	FVertexDeclarationRHIRef VertexDeclarationRHI;
+
+	virtual void InitRHI() override
+	{
+		FVertexDeclarationElementList Elements;
+		Elements.Add(FVertexElement(0, 0, VET_Float4, 0, sizeof(FVector4)));
+		VertexDeclarationRHI = RHICreateVertexDeclaration(Elements);
+	}
+
+	virtual void ReleaseRHI() override
+	{
+		VertexDeclarationRHI.SafeRelease();
+	}
+};
+TGlobalResource<FPointCloudVertexDeclaration> GPointCloudVertexDeclaration;
+
+/**
+ * A dummy vertex buffer to bind when rendering point clouds. This prevents
+ * some D3D debug warnings about zero-element input layouts but is not strictly
+ * required.
+ */
+class FDummyVertexBuffer :
+	public FVertexBuffer
+{
+public:
+
+	virtual void InitRHI() override
+	{
+		FRHIResourceCreateInfo CreateInfo;
+		void* BufferData = nullptr;
+		VertexBufferRHI = RHICreateAndLockVertexBuffer(sizeof(FVector4) * 4, BUF_Static, CreateInfo, BufferData);
+		FVector4* DummyContents = (FVector4*)BufferData;
+//@todo - joeg do I need a quad's worth?
+		DummyContents[0] = FVector4(0.0f, 0.0f, 0.0f, 0.0f);
+		DummyContents[1] = FVector4(1.0f, 0.0f, 0.0f, 0.0f);
+		DummyContents[2] = FVector4(0.0f, 1.0f, 0.0f, 0.0f);
+		DummyContents[3] = FVector4(1.0f, 1.0f, 0.0f, 0.0f);
+		RHIUnlockVertexBuffer(VertexBufferRHI);
+	}
+};
+TGlobalResource<FDummyVertexBuffer> GDummyPointCloudVertexBuffer;
+
+void FPointCloudVertexFactory::InitRHI()
+{
+	FVertexStream Stream;
+
+	// No streams should currently exist.
+	check( Streams.Num() == 0 );
+
+	Stream.VertexBuffer = &GDummyPointCloudVertexBuffer;
+	Stream.Stride = sizeof(FVector4);
+	Stream.Offset = 0;
+	Streams.Add(Stream);
+
+	// Set the declaration.
+	check(IsValidRef(GPointCloudVertexDeclaration.VertexDeclarationRHI));
+	SetDeclaration(GPointCloudVertexDeclaration.VertexDeclarationRHI);
+}
+
+void FPointCloudVertexFactory::ReleaseRHI()
+{
+	UniformBuffer.SafeRelease();
+	FVertexFactory::ReleaseRHI();
+}
+
+bool FPointCloudVertexFactory::ShouldCompilePermutation(EShaderPlatform Platform, const class FMaterial* Material, const class FShaderType* ShaderType)
+{
+	// We exclusively use manual fetch, so we need that supported
+	return RHISupportsManualVertexFetch(Platform);
+}
+
+FVertexFactoryShaderParameters* FPointCloudVertexFactory::ConstructShaderParameters(EShaderFrequency ShaderFrequency)
+{
+	return ShaderFrequency == SF_Vertex ? new FPointCloudVertexFactoryShaderParameters() : nullptr;
+}
+
+void FPointCloudVertexFactory::SetParameters(const FPointCloudVertexFactoryParameters& InUniformParameters, const uint32 InMask, const float InSize)
+{
+	UniformBuffer = FPointCloudVertexFactoryBufferRef::CreateUniformBufferImmediate(InUniformParameters, UniformBuffer_MultiFrame);
+	ColorMask = InMask;
+	PointSize = InSize;
+}
+
+IMPLEMENT_VERTEX_FACTORY_TYPE(FPointCloudVertexFactory, "/Engine/Private/PointCloudVertexFactory.ush", true, false, false, false, false);
diff --git a/Engine/Plugins/Experimental/AR/PointCloud/Source/PointCloud/Private/PointCloudVertexFactory.h b/Engine/Plugins/Experimental/AR/PointCloud/Source/PointCloud/Private/PointCloudVertexFactory.h
new file mode 100644
--- /dev/null
+++ b/Engine/Plugins/Experimental/AR/PointCloud/Source/PointCloud/Private/PointCloudVertexFactory.h
@@ -0,0 +1,78 @@
+// Copyright 1998-2018 Epic Games, Inc. All Rights Reserved.
+#pragma once
+
+#include "CoreMinimal.h"
+#include "UniformBuffer.h"
+#include "VertexFactory.h"
+
+/**
+ * Uniform buffer to hold parameters for point cloud rendering
+ */
+BEGIN_UNIFORM_BUFFER_STRUCT( FPointCloudVertexFactoryParameters, )
+	UNIFORM_MEMBER_SRV(Buffer<float>, VertexFetch_PointLocationBuffer)
+	UNIFORM_MEMBER_SRV(Buffer<float4>, VertexFetch_PointColorBuffer)
+END_UNIFORM_BUFFER_STRUCT( FPointCloudVertexFactoryParameters )
+typedef TUniformBufferRef<FPointCloudVertexFactoryParameters> FPointCloudVertexFactoryBufferRef;
+
+/**
+ * Vertex factory for point cloud rendering. This base version uses the dummy color buffer
+ */
+class FPointCloudVertexFactory :
+	public FVertexFactory
+{
+	DECLARE_VERTEX_FACTORY_TYPE(FPointCloudVertexFactory);
+
+public:
+	FPointCloudVertexFactory(ERHIFeatureLevel::Type InFeatureLevel)
+		: FVertexFactory(InFeatureLevel)
+	{
+	}
+
+	/**
+	 * Constructs render resources for this vertex factory.
+	 */
+	virtual void InitRHI() override;
+
+	/**
+	 * Release render resources for this vertex factory.
+	 */
+	virtual void ReleaseRHI() override;
+
+	/**
+	 * Should we cache the material's shadertype on this platform with this vertex factory?
+	 */
+	static bool ShouldCompilePermutation(EShaderPlatform Platform, const class FMaterial* Material, const class FShaderType* ShaderType);
+
+	/**
+	 * Construct shader parameters for this type of vertex factory.
+	 */
+	static FVertexFactoryShaderParameters* ConstructShaderParameters(EShaderFrequency ShaderFrequency);
+
+	/**
+	 * Set parameters for this vertex factory instance.
+	 */
+	void SetParameters(const FPointCloudVertexFactoryParameters& InUniformParameters, const uint32 InMask, float InSize);
+
+	inline const FUniformBufferRHIRef GetPointCloudVertexFactoryUniformBuffer() const
+	{
+		return UniformBuffer;
+	}
+
+	inline const uint32 GetColorMask() const
+	{
+		return ColorMask;
+	}
+
+	inline const float GetPointSize() const
+	{
+		return PointSize;
+	}
+
+private:
+	/** Buffers to read from */
+	FUniformBufferRHIRef UniformBuffer;
+	/** Mask of zero when using a global color or all bits when using a stream */
+	uint32 ColorMask;
+	/** The point size to use when rendering separate from the component scale for zooming operations, etc. */
+	float PointSize;
+};
diff --git a/Engine/Plugins/Experimental/AR/PointCloud/Source/PointCloud/Public/PointCloud.h b/Engine/Plugins/Experimental/AR/PointCloud/Source/PointCloud/Public/PointCloud.h
new file mode 100644
--- /dev/null
+++ b/Engine/Plugins/Experimental/AR/PointCloud/Source/PointCloud/Public/PointCloud.h
@@ -0,0 +1,9 @@
+// Copyright 1998-2018 Epic Games, Inc. All Rights Reserved.
+
+#pragma once
+
+#include "CoreMinimal.h"
+
+DECLARE_LOG_CATEGORY_EXTERN(LogPointCloud, Log, All);
+
+DECLARE_STATS_GROUP(TEXT("Point Cloud"), STATGROUP_POINTCLOUD, STATCAT_Advanced);
diff --git a/Engine/Plugins/Experimental/AR/PointCloud/Source/PointCloud/Public/PointCloudComponent.h b/Engine/Plugins/Experimental/AR/PointCloud/Source/PointCloud/Public/PointCloudComponent.h
new file mode 100644
--- /dev/null
+++ b/Engine/Plugins/Experimental/AR/PointCloud/Source/PointCloud/Public/PointCloudComponent.h
@@ -0,0 +1,103 @@
+// Copyright 1998-2018 Epic Games, Inc. All Rights Reserved.
+
+#pragma once
+
+#include "CoreMinimal.h"
+#include "UObject/ObjectMacros.h"
+#include "Components/MeshComponent.h"
+#include "PointCloudComponent.generated.h"
+
+class FPrimitiveSceneProxy;
+
+/**
+ * Component for rendering a point cloud
+ */
+UCLASS(hidecategories = (Object, LOD), meta = (BlueprintSpawnableComponent), ClassGroup = PointCloud)
+class POINTCLOUD_API UPointCloudComponent :
+	public UMeshComponent
+{
+	GENERATED_UCLASS_BODY()
+
+	/**	If true, each tick the component will render its point cloud */
+	UPROPERTY(EditAnywhere, BlueprintReadWrite, Category = "Point Cloud")
+	bool bIsVisible;
+	
+	/** Changes the visibility setting */
+	UFUNCTION(BlueprintCallable, Category = "Point Cloud")
+	void SetIsVisible(bool bNewVisibility);
+
+	/** Point cloud data that will be used for rendering, assumes each point is in world space */
+	UPROPERTY(BlueprintReadOnly, Category = "Point Cloud")
+	TArray<FVector> PointCloud;
+
+	/** Point cloud color data that will be used for rendering */
+	UPROPERTY(BlueprintReadOnly, Category = "Point Cloud")
+	TArray<FColor> PointColors;
+
+	/**
+	 * Updates the point cloud data with the new set of points
+	 *
+	 * @param Points the new set of points to use
+	 */
+	UFUNCTION(BlueprintCallable, Category = "Point Cloud")
+	void SetPointCloud(const TArray<FVector>& Points);
+
+	/**
+	 * Updates the point cloud data with the new set of points and colors
+	 *
+	 * @param Points the new set of points to use
+	 * @param Colors the new set of colors to use, must match points size
+	 */
+	UFUNCTION(BlueprintCallable, Category = "Point Cloud")
+	void SetPointCloudWithColors(const TArray<FVector>& Points, const TArray<FColor>& Colors);
+
+	/**
+	 * Empties the point cloud
+	 */
+	UFUNCTION(BlueprintCallable, Category = "Point Cloud")
+	void ClearPointCloud();
+	
+	/** The color to render the points with */
+	UPROPERTY(EditAnywhere, BlueprintReadOnly, Category = "Point Cloud")
+	FLinearColor PointColor;
+
+	/** Allows you to change the color of the points being rendered */
+	UFUNCTION(BlueprintCallable, Category = "Point Cloud")
+	void SetPointColor(const FLinearColor& Color);
+
+	/** The size of the point when rendering */
+	UPROPERTY(EditAnywhere, BlueprintReadOnly, Category = "Point Cloud")
+	float PointSize;
+	
+	/** Allows you to change the size of the points being rendered */
+	UFUNCTION(BlueprintCallable, Category = "Point Cloud")
+	void SetPointSize(float Size);
+
+	/** Determines which points are within the box and returns those to the caller */
+	UFUNCTION(BlueprintPure, Category = "Point Cloud")
+	TArray<FVector> GetPointsInBox(const FBox& WorldSpaceBox) const;
+
+	/** Determines which points are outside the box and returns those to the caller */
+	UFUNCTION(BlueprintPure, Category = "Point Cloud")
+	TArray<FVector> GetPointsOutsideBox(const FBox& WorldSpaceBox) const;
+
+	/** The material to render with */
+	UPROPERTY()
+	UMaterialInterface* PointCloudMaterial;
+
+private:
+	virtual FMatrix GetRenderMatrix() const override;
+	virtual FPrimitiveSceneProxy* CreateSceneProxy() override;
+	virtual int32 GetNumMaterials() const override { return 1; }
+	virtual void GetUsedMaterials(TArray<UMaterialInterface*>& OutMaterials, bool bGetDebugMaterials = false) const override;
+	virtual FBoxSphereBounds CalcBounds(const FTransform& LocalToWorld) const override;
+
+	/** World space bounds of the point cloud */
+	UPROPERTY()
+	FBoxSphereBounds WorldBounds;
+
+	/** Used to know whether to update the point cloud or not */
+	float LastUpdateTimestamp;
+};
+
+
diff --git a/Engine/Shaders/Private/PointCloudVertexFactory.ush b/Engine/Shaders/Private/PointCloudVertexFactory.ush
new file mode 100644
--- /dev/null
+++ b/Engine/Shaders/Private/PointCloudVertexFactory.ush
@@ -0,0 +1,326 @@
+// Copyright 1998-2018 Epic Games, Inc. All Rights Reserved.
+
+/*====================================================================================================================
+	PointCloudVertexFactory.usf: Vertex factory for point clouds. Based upon vector field visualization vertex factory.
+ 		There are no vertices used, just indices. Vertex positions are calculated based on vertex id and the point
+ 		cloud location buffer. Optionally, supports per point colors using another manual fetch stream. Otherwise, it
+ 		used the component's color setting for all points
+=====================================================================================================================*/
+
+#include "VertexFactoryCommon.ush"
+
+/** Buffer of point locations to manually fetch */
+//Buffer<float3> VertexFetch_PointLocationBuffer;
+
+/** Buffer of point colors to manually fetch. This is a dummy buffer if HAS_POINT_COLOR_BUFFER is 0 */
+//Buffer<float4> VertexFetch_PointColorBuffer;
+
+/** Used to specify point size independently from scale, in case both want to be used for */
+float PointSize;
+
+/** Will either be 0 or all bits to select between always zero index and a real index */
+uint ColorMask;
+
+/**
+ * Vertex attributes to fetch.
+ */
+struct FVertexFactoryInput
+{
+	/** Unique vertex ID. */
+	uint VertexId : SV_VertexID;
+	/** This doesn't use instancing */
+	uint InstanceId : SV_InstanceID;
+};
+
+/**
+ * Attributes to interpolate from the vertex shader to the pixel shader.
+ */
+struct FVertexFactoryInterpolantsVSToPS
+{
+	/** Dummy value to interpolate. */
+	float2 DummyTexCoord : TEXCOORD0;
+
+	half4 Color : COLOR0;
+
+#if INSTANCED_STEREO
+	nointerpolation uint PackedEyeIndex : PACKED_EYE_INDEX;
+#endif
+};
+
+/**
+ * Intermediate values computed in the vertex shader.
+ */
+struct FVertexFactoryIntermediates
+{
+	/** The position of the vertex in translated world space */
+	float3 Position;
+	/** The position of the point in translated world space */
+	float3 PointTranslatedWorldPosition;
+	/** The point's tangent in world space (+V) */
+	float3 TangentUp;
+	/** The point's tangent in world space (+U) */
+	float3 TangentRight;
+	/** The color of the sprite. */
+	half4 Color;
+};
+
+/**
+ * Fetches the point's location from the point location buffer
+ *
+ * @param Input the vertex factory input
+ *
+ * @return the position of the point
+ */
+float3 GetPointLocation(FVertexFactoryInput Input)
+{
+	// Load the point's location
+	uint InstanceId = (Input.VertexId / 4) * 3;
+	float X = PointCloudVF.VertexFetch_PointLocationBuffer[InstanceId + 0];
+	float Y = PointCloudVF.VertexFetch_PointLocationBuffer[InstanceId + 1];
+	float Z = PointCloudVF.VertexFetch_PointLocationBuffer[InstanceId + 2];
+	return float3(X, Y, Z);
+}
+
+/**
+ * Fetches the point's color from the point color buffer
+ *
+ * @param Input the vertex factory input
+ *
+ * @return the color of the point
+ */
+half4 GetPointColor(FVertexFactoryInput Input)
+{
+	// InstanceId will always be zero if we're using a single color for all points, otherwise it will be the buffer index
+	uint InstanceId = (Input.VertexId / 4) & ColorMask;
+
+	return half4(PointCloudVF.VertexFetch_PointColorBuffer[InstanceId] FMANUALFETCH_COLOR_COMPONENT_SWIZZLE);
+}
+
+/**
+ * Determines the X & Y offsets for this point based off of vertex id and point size
+ *
+ * @param Input the vertex factory input
+ *
+ * @return the color of the point
+ */
+float2 GetOffsetForVertex(FVertexFactoryInput Input)
+{
+	const float2 PositionOffsets[4] =
+	{
+		// Upper left
+		float2(0.5f, -0.5f),
+		// Upper right
+		float2(0.5f, 0.5f),
+		// Lower left
+		float2(-0.5f, -0.5f),
+		// Lower right
+		float2(-0.5f, 0.5f)
+	};
+	// Choose the offset of the vert
+	uint OffsetIndex = Input.VertexId % 4;
+	return PositionOffsets[OffsetIndex] * PointSize;
+}
+
+float3 SafeNormalize(float3 V)
+{
+	return V * rsqrt(max(dot(V,V),0.01));
+}
+
+/**
+ * Compute the tangents for a point
+ *
+ * @param OutUp - The tangent vector pointing up in screen space.
+ * @param OutRight - The tangent vector pointing right in screen space.
+ * @param PointTranslatedWorldPosition - Translated world position of the point
+ */
+void GetPointTangents(out float3 OutUp, out float3 OutRight, float3 PointTranslatedWorldPosition)
+{
+	// Select camera up/right vectors
+	float3 CameraRight = ResolvedView.ViewRight;
+	float3 CameraUp = -ResolvedView.ViewUp;
+
+	// Determine the vector from the point to the camera
+	float3 CameraDirection	= SafeNormalize(ResolvedView.TranslatedWorldCameraOrigin - PointTranslatedWorldPosition.xyz);
+	float3 RightVector = CameraRight.xyz;
+	float3 UpVector = CameraUp.xyz;
+
+	// Rotate toward the camera
+	RightVector = SafeNormalize(cross(CameraDirection, float3(0, 0, 1)));
+	UpVector = cross(CameraDirection, RightVector);
+
+	// Rotate the sprite to determine final tangents.
+	OutRight = RightVector;
+	OutUp = UpVector;
+}
+
+/**
+ * Computes intermediates for the given vertex.
+ *
+ * @param Input - Vertex attributes.
+ *
+ * @returns the computed intermediate values.
+ */
+FVertexFactoryIntermediates GetVertexFactoryIntermediates(FVertexFactoryInput Input)
+{
+	// Load the position from the float stream
+	float3 LoadedPosition = GetPointLocation(Input);
+	// Transform position to (post-view-translation) world space
+	float3 PointTranslatedWorldPosition = TransformLocalToTranslatedWorld(LoadedPosition.xyz).xyz;
+
+	// Compute tangents for the point
+	float3 TangentUp, TangentRight;
+	GetPointTangents(TangentUp, TangentRight, PointTranslatedWorldPosition);
+
+	// Offset of the vertex from the point's location
+	float2 Offset = GetOffsetForVertex(Input);
+	float3 VertexWorldOffset = Offset.x * TangentRight + Offset.y * TangentUp;
+	float3 VertexWorldPosition = PointTranslatedWorldPosition + VertexWorldOffset;
+
+	// Build and return the set of intermediates
+	FVertexFactoryIntermediates Intermediates;
+	Intermediates.Position = VertexWorldPosition;
+	Intermediates.PointTranslatedWorldPosition = PointTranslatedWorldPosition;
+	Intermediates.TangentUp = TangentUp;
+	Intermediates.TangentRight = TangentRight;
+
+	Intermediates.Color = GetPointColor(Input);
+
+	return Intermediates;
+}
+
+/**
+ * Computes material parameterss for a given pixel.
+ *
+ * @param Interpolants - Attributes interpolated from the vertex shader.
+ *
+ * @returns per-pixel material parameters.
+ */
+FMaterialPixelParameters GetMaterialPixelParameters(FVertexFactoryInterpolantsVSToPS Interpolants, float4 SvPosition)
+{
+	// GetMaterialPixelParameters is responsible for fully initializing the result
+	FMaterialPixelParameters Result = MakeInitializedMaterialPixelParameters();
+
+	// Note that a non-zero texture coordinate is used to prevent the compiler
+	// from optimizing out texture lookups that can cause mandatory requirements
+	// to not be bound.
+
+#if NUM_MATERIAL_TEXCOORDS
+	UNROLL
+	for (int CoordinateIndex = 0; CoordinateIndex < NUM_MATERIAL_TEXCOORDS; CoordinateIndex++)
+	{
+		Result.TexCoords[CoordinateIndex] = Interpolants.DummyTexCoord;
+	}
+#endif	//NUM_MATERIAL_TEXCOORDS
+
+	Result.VertexColor = Interpolants.Color;
+
+	Result.TwoSidedSign = 1;
+	return Result;
+}
+
+/**
+ * Computes material parameters for a given vertex.
+ *
+ * @param Input - Attributes for this vertex.
+ * @param Intermediates - Intermediates computed for this vertex.
+ * @param WorldPosition - The position of this vertex in world space.
+ * @param TangentToLocal - The tangent basis for this vertex.
+ *
+ * @returns per-vertex material parameters.
+ */
+FMaterialVertexParameters GetMaterialVertexParameters(FVertexFactoryInput Input, FVertexFactoryIntermediates Intermediates, float3 WorldPosition, float3x3 TangentToLocal)
+{
+	FMaterialVertexParameters Result = (FMaterialVertexParameters)0;
+	Result.WorldPosition = WorldPosition;
+	Result.VertexColor = Intermediates.Color;
+	Result.TangentToWorld = mul(TangentToLocal, GetLocalToWorld3x3());
+	Result.PrevFrameLocalToWorld = Primitive.LocalToWorld;
+	return Result;
+}
+
+/**
+ * Computes the world space position of this vertex.
+ *
+ * @param Input - Vertex attributes.
+ * @param Intermediates - Intermediates computed for this vertex.
+ *
+ * @returns the position of this vertex in world space.
+ */
+float4 VertexFactoryGetWorldPosition(FVertexFactoryInput Input, FVertexFactoryIntermediates Intermediates)
+{
+	return TransformLocalToTranslatedWorld(Intermediates.Position.xyz);
+}
+
+float4 VertexFactoryGetRasterizedWorldPosition(FVertexFactoryInput Input, FVertexFactoryIntermediates Intermediates, float4 TranslatedWorldPosition)
+{
+	return TranslatedWorldPosition;
+}
+
+float3 VertexFactoryGetPositionForVertexLighting(FVertexFactoryInput Input, FVertexFactoryIntermediates Intermediates, float3 TranslatedWorldPosition)
+{
+	return TranslatedWorldPosition;
+}
+
+/**
+ * Constructs values that need to be interpolated from the vertex shader to the pixel shader.
+ *
+ * @param Input - Vertex attributes.
+ * @param Intermediates - Intermediates computed for this vertex.
+ *
+ * @returns interpolants.
+ */
+FVertexFactoryInterpolantsVSToPS VertexFactoryGetInterpolantsVSToPS(FVertexFactoryInput Input, FVertexFactoryIntermediates Intermediates, FMaterialVertexParameters VertexParameters)
+{
+	FVertexFactoryInterpolantsVSToPS Interpolants;
+	Interpolants.DummyTexCoord = float2(0,0);
+
+	Interpolants.Color = Intermediates.Color;
+
+#if INSTANCED_STEREO
+	Interpolants.PackedEyeIndex = 0;
+#endif
+
+	return Interpolants;
+}
+
+/**
+ * Computes the position of this vertex last frame in world space.
+ *
+ * @param Input - Vertex attributes.
+ * @param Intermediates - Intermediates computed for this vertex.
+ *
+ * @returns the previous position of this vertex in world space.
+ */
+float4 VertexFactoryGetPreviousWorldPosition(FVertexFactoryInput Input, FVertexFactoryIntermediates Intermediates)
+{
+	return VertexFactoryGetWorldPosition(Input, Intermediates);
+}
+
+/**
+ * Computes the tangent basis for this vertex in world space.
+ *
+ * @param Input - Vertex attributes.
+ * @param Intermediates - Intermediates computed for this vertex.
+ *
+ * @returns the tangent basis for this vertex in world space.
+ */
+float3x3 VertexFactoryGetTangentToLocal(FVertexFactoryInput Input, FVertexFactoryIntermediates Intermediates)
+{
+	float3x3 TangentToLocal;
+	TangentToLocal[0] = Intermediates.TangentRight;
+	TangentToLocal[1] = Intermediates.TangentUp;
+	TangentToLocal[2] = normalize(cross(Intermediates.TangentRight.xyz , Intermediates.TangentUp.xyz));
+	return TangentToLocal;
+}
+
+#if INSTANCED_STEREO
+uint VertexFactoryGetEyeIndex(uint InstanceId)
+{
+	return 0;
+}
+#endif
+
+float4 VertexFactoryGetTranslatedPrimitiveVolumeBounds(FVertexFactoryInterpolantsVSToPS Interpolants)
+{
+	return 0;
+}
\ No newline at end of file
