Cleaned up some comments

diff --git a/Engine/Shaders/Private/PointCloudVertexFactory.ush b/Engine/Shaders/Private/PointCloudVertexFactory.ush
--- a/Engine/Shaders/Private/PointCloudVertexFactory.ush
+++ b/Engine/Shaders/Private/PointCloudVertexFactory.ush
@@ -9,12 +9,6 @@
 
 #include "VertexFactoryCommon.ush"
 
-/** Buffer of point locations to manually fetch */
-//Buffer<float3> VertexFetch_PointLocationBuffer;
-
-/** Buffer of point colors to manually fetch. This is a dummy buffer if HAS_POINT_COLOR_BUFFER is 0 */
-//Buffer<float4> VertexFetch_PointColorBuffer;
-
 /** Used to specify point size independently from scale, in case both want to be used for */
 float PointSize;
 
