Fix lighting/volume fog being broken with stereo-instancing enabled, by ensuring the correct fog volume texture is referenced when sampling fog.

diff --git a/Engine/Shaders/Private/BasePassPixelShader.usf b/Engine/Shaders/Private/BasePassPixelShader.usf
--- a/Engine/Shaders/Private/BasePassPixelShader.usf
+++ b/Engine/Shaders/Private/BasePassPixelShader.usf
@@ -1128,7 +1128,7 @@ ShadingModel);
 	if (FogStruct.ApplyVolumetricFog > 0) 
 	{
 		float3 VolumeUV = ComputeVolumeUV(MaterialParameters.AbsoluteWorldPosition, ResolvedView.WorldToClip);
-		Fogging = CombineVolumetricFog(HeightFogging, VolumeUV);
+		Fogging = CombineVolumetricFog(HeightFogging, VolumeUV, Interpolants.EyeIndex);
 	}
 #endif
 
diff --git a/Engine/Shaders/Private/HeightFogCommon.ush b/Engine/Shaders/Private/HeightFogCommon.ush
--- a/Engine/Shaders/Private/HeightFogCommon.ush
+++ b/Engine/Shaders/Private/HeightFogCommon.ush
@@ -16,6 +16,15 @@
 #define SUPPORT_VOLUMETRIC_FOG (FEATURE_LEVEL >= FEATURE_LEVEL_SM4)
 #endif
 
+#if INSTANCED_STEREO
+	#if MATERIALBLENDING_ANY_TRANSLUCENT
+		#define FogStructISR TranslucentBasePass.Shared.FogISR
+	#else
+		#define FogStructISR OpaqueBasePass.Shared.FogISR
+	#endif
+#endif
+
+
 static const float FLT_EPSILON = 0.001f;
 static const float FLT_EPSILON2 = 0.01f;
 
@@ -207,14 +216,25 @@ half4 CalculateHeightFog(float3 WorldPositionRelativeToCamera)
 	#define SharedIntegratedLightScatteringSampler FogStruct.IntegratedLightScatteringSampler
 #endif
 
-float4 CombineVolumetricFog(float4 GlobalFog, float3 VolumeUV)
+float4 CombineVolumetricFog(float4 GlobalFog, float3 VolumeUV, uint EyeIndex)
 {
 	float4 VolumetricFogLookup = float4(0, 0, 0, 1);
 
 #if SUPPORT_VOLUMETRIC_FOG
 	if (FogStruct.ApplyVolumetricFog > 0)
 	{
+	#if INSTANCED_STEREO
+		if (EyeIndex == 0)
+		{
+			VolumetricFogLookup = Texture3DSampleLevel(FogStruct.IntegratedLightScattering, SharedIntegratedLightScatteringSampler, VolumeUV, 0);
+		}
+		else
+		{
+			VolumetricFogLookup = Texture3DSampleLevel(FogStructISR.IntegratedLightScattering, SharedIntegratedLightScatteringSampler, VolumeUV, 0);
+		}
+	#else
 		VolumetricFogLookup = Texture3DSampleLevel(FogStruct.IntegratedLightScattering, SharedIntegratedLightScatteringSampler, VolumeUV, 0);
+	#endif
 	}
 #endif
 
diff --git a/Engine/Shaders/Private/HeightFogPixelShader.usf b/Engine/Shaders/Private/HeightFogPixelShader.usf
--- a/Engine/Shaders/Private/HeightFogPixelShader.usf
+++ b/Engine/Shaders/Private/HeightFogPixelShader.usf
@@ -38,8 +38,8 @@ void ExponentialPixelMain(
 	float3 VolumeUV = float3((SVPos.xy - View.ViewRectMin.xy) * View.VolumetricFogSVPosToVolumeUV, ZSlice);
 
 	float4 HeightFogInscatteringAndOpacity = CalculateHeightFog(WorldPositionRelativeToCamera);
-	float4 FogInscatteringAndOpacity = CombineVolumetricFog(HeightFogInscatteringAndOpacity, VolumeUV);
 
+	float4 FogInscatteringAndOpacity = CombineVolumetricFog(HeightFogInscatteringAndOpacity, VolumeUV, 0);
 	float LightShaftMask = Texture2DSample(OcclusionTexture, OcclusionSampler, TexCoord).x;
 	FogInscatteringAndOpacity.rgb *= LightShaftMask;
 
diff --git a/Engine/Source/Runtime/Renderer/Private/BasePassRendering.cpp b/Engine/Source/Runtime/Renderer/Private/BasePassRendering.cpp
--- a/Engine/Source/Runtime/Renderer/Private/BasePassRendering.cpp
+++ b/Engine/Source/Runtime/Renderer/Private/BasePassRendering.cpp
@@ -490,21 +490,24 @@ void SetupSharedBasePassParameters(
 {
 	SharedParameters.Forward = View.ForwardLightingResources->ForwardLightData;
 
+	SetupFogUniformParameters(View, SharedParameters.Fog);
+
 	if (View.bIsInstancedStereoEnabled && View.StereoPass == EStereoscopicPass::eSSP_LEFT_EYE)
 	{
 		const FSceneView& RightEye = *View.Family->Views[1];
 		SharedParameters.ForwardISR = RightEye.ForwardLightingResources->ForwardLightData;
+		SetupFogUniformParameters((FViewInfo&)RightEye, SharedParameters.FogISR);
 	}
 	else
 	{
 		SharedParameters.ForwardISR = View.ForwardLightingResources->ForwardLightData;
+		SharedParameters.FogISR = SharedParameters.Fog;
 	}
 
 	const FScene* Scene = View.Family->Scene ? View.Family->Scene->GetRenderScene() : nullptr;
 	const FPlanarReflectionSceneProxy* ReflectionSceneProxy = Scene ? Scene->GetForwardPassGlobalPlanarReflection() : nullptr;
 
 	SetupReflectionUniformParameters(View, SharedParameters.Reflection);
-	SetupFogUniformParameters(View, SharedParameters.Fog);
 	SetupPlanarReflectionUniformParameters(View, ReflectionSceneProxy, SharedParameters.PlanarReflection);
 
 	const IPooledRenderTarget* PooledRT = GetSubsufaceProfileTexture_RT(RHICmdList);
diff --git a/Engine/Source/Runtime/Renderer/Private/BasePassRendering.h b/Engine/Source/Runtime/Renderer/Private/BasePassRendering.h
--- a/Engine/Source/Runtime/Renderer/Private/BasePassRendering.h
+++ b/Engine/Source/Runtime/Renderer/Private/BasePassRendering.h
@@ -51,6 +51,7 @@ BEGIN_GLOBAL_SHADER_PARAMETER_STRUCT(FSharedBasePassUniformParameters,)
 	SHADER_PARAMETER_STRUCT(FReflectionUniformParameters, Reflection)
 	SHADER_PARAMETER_STRUCT(FPlanarReflectionUniformParameters, PlanarReflection) // Single global planar reflection for the forward pass.
 	SHADER_PARAMETER_STRUCT(FFogUniformParameters, Fog)
+	SHADER_PARAMETER_STRUCT(FFogUniformParameters, FogISR)
 	SHADER_PARAMETER_TEXTURE(Texture2D, SSProfilesTexture)
 END_GLOBAL_SHADER_PARAMETER_STRUCT()
 
