Fix build error with invalid EyeIndex subscript with BasePassPixelShader.usf when no instanced stereo defined.

diff --git a/Engine/Shaders/Private/BasePassPixelShader.usf b/Engine/Shaders/Private/BasePassPixelShader.usf
--- a/Engine/Shaders/Private/BasePassPixelShader.usf
+++ b/Engine/Shaders/Private/BasePassPixelShader.usf
@@ -1128,7 +1128,11 @@ ShadingModel);
 	if (FogStruct.ApplyVolumetricFog > 0) 
 	{
 		float3 VolumeUV = ComputeVolumeUV(MaterialParameters.AbsoluteWorldPosition, ResolvedView.WorldToClip);
+#if INSTANCED_STEREO
 		Fogging = CombineVolumetricFog(HeightFogging, VolumeUV, Interpolants.EyeIndex);
+#else
+		Fogging = CombineVolumetricFog(HeightFogging, VolumeUV, 0);
+#endif
 	}
 #endif
 
