SSGI first version
Improved SSR. Added binary search. Disabled line segment intersection.

diff --git a/Engine/Shaders/Private/ForwardLightingCommon.ush b/Engine/Shaders/Private/ForwardLightingCommon.ush
--- a/Engine/Shaders/Private/ForwardLightingCommon.ush
+++ b/Engine/Shaders/Private/ForwardLightingCommon.ush
@@ -348,25 +348,24 @@ half3 GetImageBasedReflectionLighting(FMaterialPixelParameters MaterialParameter
 	{
 		float StepOffset = InterleavedGradientNoise( MaterialParameters.SvPosition.xy, View.StateFrameIndexMod8 );
 		StepOffset -= 0.5;
-		
-		float4 HitUVzTime;
-		float HCBLevel;
 
-		RayCast(
-			TranslucentBasePass.HZBTexture, TranslucentBasePass.HZBSampler, float2(1, 1),
-			MaterialParameters.WorldPosition_CamRelative, RayDirection, 0, 0, MaterialParameters.ScreenPosition.w,
+		float3 HitUVz;
+		float Level = 0;
+
+		bool bHit = RayCast(
+			TranslucentBasePass.HZBTexture, TranslucentBasePass.HZBSampler,
+			MaterialParameters.WorldPosition_CamRelative, RayDirection, Roughness, MaterialParameters.ScreenPosition.w,
 			12, StepOffset,
 			TranslucentBasePass.HZBUvFactorAndInvFactor,
-			TranslucentBasePass.PrevScreenPositionScaleBias,
-			HitUVzTime, HCBLevel
+			HitUVz,
+			Level
 		);
 
-		// if there was a hit
-		BRANCH if( HitUVzTime.w < 1 )
+		BRANCH if( bHit )
 		{
 			float2 SampleUV;
 			float Vignette;
-			ReprojectHit(TranslucentBasePass.PrevScreenPositionScaleBias, HitUVzTime.xyz, SampleUV, Vignette);
+			ReprojectHit(TranslucentBasePass.PrevScreenPositionScaleBias, HitUVz, SampleUV, Vignette);
 
 			float4 SSR = SampleScreenColor( 
 				TranslucentBasePass.PrevSceneColor, 
diff --git a/Engine/Shaders/Private/ScreenSpaceDiffuseIndirect.usf b/Engine/Shaders/Private/ScreenSpaceDiffuseIndirect.usf
new file mode 100644
--- /dev/null
+++ b/Engine/Shaders/Private/ScreenSpaceDiffuseIndirect.usf
@@ -0,0 +1,149 @@
+// Copyright 1998-2018 Epic Games, Inc. All Rights Reserved.
+
+#include "Common.ush"	
+#include "PostProcessCommon.ush"				
+#include "DeferredShadingCommon.ush"
+#include "Random.ush"
+#include "BRDF.ush"
+#include "MonteCarlo.ush"
+#include "ScreenSpaceRayCast.ush"
+
+Texture2D		HZBTexture;
+SamplerState	HZBSampler;
+
+Texture2D		VelocityTexture;
+SamplerState	VelocitySampler;
+
+Texture2D		ColorTexture;
+SamplerState	ColorSampler;
+
+float4 HZBUvFactorAndInvFactor;
+
+float4 PrevScreenPositionScaleBias;
+float  PrevSceneColorPreExposureCorrection;
+
+void ScreenSpaceDiffuseIndirectPS(
+	in noperspective float4 UVAndScreenPos : TEXCOORD0,
+	float4 SvPosition : SV_POSITION,
+	out float4 OutColor : SV_Target0,
+	out float OutAO : SV_Target1 )
+{
+	float2 UV = UVAndScreenPos.xy;
+	float2 ScreenPos = UVAndScreenPos.zw;
+	uint2 PixelPos = (uint2)SvPosition.xy;
+
+	OutColor = 0;
+
+	FScreenSpaceData ScreenSpaceData = GetScreenSpaceData(UV);
+	FGBufferData GBuffer = ScreenSpaceData.GBuffer;
+
+	// No SSS because diffuse luminance is in alpha channel and we don't have dual source blending.
+	BRANCH if( GBuffer.ShadingModelID == 0 || GBuffer.ShadingModelID == SHADINGMODELID_SUBSURFACE_PROFILE )
+	{
+		discard;
+	}
+
+	const float3 N = GBuffer.WorldNormal;
+	const float SceneDepth = CalcSceneDepth(UV);
+	const float3 PositionTranslatedWorld = mul( float4( ScreenPos * SceneDepth, SceneDepth, 1 ), View.ScreenToTranslatedWorld ).xyz;
+
+#if QUALITY == 1
+	uint NumSteps = 8;
+	uint NumRays = 4;
+#elif QUALITY == 2
+	uint NumSteps = 8;
+	uint NumRays = 8;
+#elif QUALITY == 3
+	uint NumSteps = 8;
+	uint NumRays = 12;
+#else // QUALITY == 4
+	uint NumSteps = 12;
+	uint NumRays = 16;
+#endif
+
+	float2 Noise;
+	Noise.x = InterleavedGradientNoise( SvPosition.xy, View.StateFrameIndexMod8 );
+	Noise.y = InterleavedGradientNoise( SvPosition.xy, View.StateFrameIndexMod8 * 117 );
+		
+	//uint2 Random = 0x10000 * Noise;
+	uint2 Random = Rand3DPCG16( int3( PixelPos, View.StateFrameIndexMod8 ) ).xy;
+			
+	float3x3 TangentBasis = GetTangentBasis( N );
+
+	// Shoot multiple rays
+	LOOP for( int i = 0; i < NumRays; i++ )
+	{
+		float StepOffset = Noise.x;
+		StepOffset -= 0.9;
+			
+		float2 E = Hammersley16( i, NumRays, Random );
+
+#if 0
+		float3 L = CosineSampleHemisphere( E ).xyz;
+		L = mul( L, TangentBasis );
+#else
+		float3 L;
+		L.xy = UniformSampleDiskConcentric( E );
+		L.z = sqrt( 1 - dot( L.xy, L.xy ) );
+		L = mul( L, TangentBasis );
+#endif
+				
+		float3 HitUVz;
+		float Level = 0;
+
+		bool bHit = RayCast(
+			HZBTexture, HZBSampler,
+			PositionTranslatedWorld, L, 1, SceneDepth, 
+			NumSteps, StepOffset,
+			HZBUvFactorAndInvFactor,
+			HitUVz,
+			Level
+		);
+
+#if 0	// Backface check
+		if( bHit )
+		{
+			float3 SampleNormal = GetGBufferData( HitUVz.xy ).WorldNormal;
+			bHit = dot( SampleNormal, L ) < 0;
+		}
+#endif
+
+		// if there was a hit
+		BRANCH if( bHit )
+		{
+			float2 SampleUV;
+			float Vignette;
+			ReprojectHit(PrevScreenPositionScaleBias, VelocityTexture, VelocitySampler, HitUVz, SampleUV, Vignette);
+
+#if 0
+			float4 SampleColor = SampleHCBLevel( ColorTexture, ColorSampler, SampleUV, Level, HZBUvFactorAndInvFactor );
+#else
+			float4 SampleColor = SampleScreenColor( ColorTexture, ColorSampler, SampleUV );
+#endif
+
+			// Backface modulation of hit surface
+			float3 SampleNormal = GetGBufferData( HitUVz.xy ).WorldNormal;
+			SampleColor.rgb *= saturate( 1 - dot( SampleNormal, L ) );
+
+			SampleColor.rgb *= rcp( 1 + Luminance(SampleColor.rgb) );
+
+			OutColor += SampleColor;
+		}
+	}
+
+	OutColor /= NumRays;
+	OutColor.rgb *= rcp( 1 - Luminance(OutColor.rgb) );
+	OutColor.rgb *= View.IndirectLightingColorScale;
+
+	OutColor.rgb *= GBuffer.DiffuseColor;
+	OutColor.a = 1 - OutColor.a;
+
+	//OutColor.rgb = OutColor.a;
+	//OutColor.a = 0;
+
+#if USE_PREEXPOSURE
+	OutColor.rgb *= PrevSceneColorPreExposureCorrection;
+#endif
+
+	OutAO = OutColor.a;
+}
\ No newline at end of file
diff --git a/Engine/Shaders/Private/ScreenSpaceRayCast.ush b/Engine/Shaders/Private/ScreenSpaceRayCast.ush
--- a/Engine/Shaders/Private/ScreenSpaceRayCast.ush
+++ b/Engine/Shaders/Private/ScreenSpaceRayCast.ush
@@ -4,206 +4,57 @@
 
 #include "Common.ush"	
 
-// 0:reference (slower, potentially higher quality) 1:use mips of the HZB depth (better performance)
-#define USE_HZB		1
-
-#define SCALAR_BRANCHLESS		0
-#define VECTORIZED_BRANCHLESS	0
-#define VECTORIZED_EARLY_OUT	1
-
-#if USE_HZB
-#define HZB_LEVEL_OFFSET 1.0
-#else
-#define HZB_LEVEL_OFFSET 0.0
-#endif
-
-float4 SampleDepthTexture( Texture2D Texture, SamplerState Sampler, float Level, float4 SampleUV0, float4 SampleUV1 )
-{
-	float4 SampleDepth;
-#if USE_HZB
-	// SampleUV{0,1}.{xy,zw} should already be in the HZB UV frame using HZBUvFactorAndInvFactor.xy
-	SampleDepth.x = Texture2DSampleLevel( Texture, Sampler, SampleUV0.xy, Level ).r;
-	SampleDepth.y = Texture2DSampleLevel( Texture, Sampler, SampleUV0.zw, Level ).r;
-	SampleDepth.z = Texture2DSampleLevel( Texture, Sampler, SampleUV1.xy, Level ).r;
-	SampleDepth.w = Texture2DSampleLevel( Texture, Sampler, SampleUV1.zw, Level ).r;
-#else
-	SampleDepth.x = Texture2DSampleLevel( SceneTexturesStruct.SceneDepthTexture, SceneTexturesStruct.SceneDepthTextureSampler, SampleUV0.xy, 0 ).r;
-	SampleDepth.y = Texture2DSampleLevel( SceneTexturesStruct.SceneDepthTexture, SceneTexturesStruct.SceneDepthTextureSampler, SampleUV0.zw, 0 ).r;
-	SampleDepth.z = Texture2DSampleLevel( SceneTexturesStruct.SceneDepthTexture, SceneTexturesStruct.SceneDepthTextureSampler, SampleUV1.xy, 0 ).r;
-	SampleDepth.w = Texture2DSampleLevel( SceneTexturesStruct.SceneDepthTexture, SceneTexturesStruct.SceneDepthTextureSampler, SampleUV1.zw, 0 ).r;
-#endif
-	return SampleDepth;
-}
-
-void RayCast(
-	Texture2D Texture, SamplerState Sampler, float2 TextureSize,
+bool RayCast(
+	Texture2D Texture, SamplerState Sampler,
 	float3 RayOriginTranslatedWorld, float3 RayDirection,
-	float Roughness, float ConeAngleWorld, float SceneDepth,
-	int NumSteps, float StepOffset,
+	float Roughness, float SceneDepth,
+	uint NumSteps, float StepOffset,
 	float4 HZBUvFactorAndInvFactor, 
-	float4 PrevScreenPositionScaleBias,
-	out float4 OutHitUVzTime,
-	out float OutLevelHCB
+	out float3 OutHitUVz,
+	out float Level
 )
 {
 	// TODO provide RayStartUVz
-	const float4 RayStartV = mul( float4( RayOriginTranslatedWorld, 1 ), View.TranslatedWorldToView );
-	const float4 RayDirV = mul(float4(RayDirection * SceneDepth, 0), View.TranslatedWorldToView);
-	const float4 RayEndV = RayStartV + RayDirV;
-
-	const float RayEndRadiusV = length(RayDirV.xyz) * tan(ConeAngleWorld * 0.4);
-	const float2 RayEndBorderV = RayEndV.xy + RayEndRadiusV / sqrt(2.0);
-
-	const float4 RayStartClip	= mul( RayStartV, View.ViewToClip );
-	const float4 RayEndClip		= mul( RayEndV, View.ViewToClip );
-	const float2 RayEndBorderClip = mul(float4(RayEndBorderV, RayEndV.zw), View.ViewToClip).xy;
+	const float4 RayStartClip	= mul( float4( RayOriginTranslatedWorld, 1 ), View.TranslatedWorldToClip );
+	const float4 RayEndClip		= mul( float4( RayOriginTranslatedWorld + RayDirection * SceneDepth, 1 ), View.TranslatedWorldToClip );
 
 	const float3 RayStartScreen = RayStartClip.xyz / RayStartClip.w;
 	const float3 RayEndScreen = RayEndClip.xyz / RayEndClip.w;
-	const float2 RayEndBorderScreen = RayEndBorderClip.xy / RayEndClip.w;
 
-	float4 RayDepthClip	= RayStartClip + mul( float4( 0, 0, SceneDepth, 0 ), View.ViewToClip );
-	float3 RayDepthScreen = RayDepthClip.xyz / RayDepthClip.w;
+	const float4 RayDepthClip	= RayStartClip + mul( float4( 0, 0, SceneDepth, 0 ), View.ViewToClip );
+	const float3 RayDepthScreen = RayDepthClip.xyz / RayDepthClip.w;
 
 	float3 RayStepScreen = RayEndScreen - RayStartScreen;
-	float2 RayStepRadiusScreen = RayEndBorderScreen - RayEndScreen.xy;
 
 	{
 		// Computes the scale down factor for RayStepScreen required to fit on the X and Y axis in order to clip it in the viewport
 		const float RayStepScreenInvFactor = 0.5 * length( RayStepScreen.xy );
-		const float2 AbsRayStepScreen = abs(RayStepScreen.xy);
-		const float2 S = (AbsRayStepScreen - max(abs(RayStepScreen.xy + RayStartScreen.xy * RayStepScreenInvFactor) - RayStepScreenInvFactor, 0.0f)) / AbsRayStepScreen;
+		const float2 S = 1 - max(abs(RayStepScreen.xy + RayStartScreen.xy * RayStepScreenInvFactor) - RayStepScreenInvFactor, 0.0f) / abs(RayStepScreen.xy);
 
 		// Rescales RayStepScreen accordingly
 		const float RayStepFactor = min(S.x, S.y) / RayStepScreenInvFactor;
-
 		RayStepScreen *= RayStepFactor;
-		RayStepRadiusScreen *= RayStepFactor;
 	}
 
-	const float2 HzbSize = float2(1024, 512);
-#if USE_HZB
 	float3 RayStartUVz = float3( (RayStartScreen.xy * float2( 0.5, -0.5 ) + 0.5) * HZBUvFactorAndInvFactor.xy, RayStartScreen.z );
-	float3 RayStepUVz = float3( RayStepScreen.xy * float2( 0.5, -0.5 ) * HZBUvFactorAndInvFactor.xy, RayStepScreen.z );
-	float2 RayStepRadiusUV = RayStepRadiusScreen.xy * float2( 0.5, -0.5 ) * HZBUvFactorAndInvFactor.xy;
-#else
-	float3 RayStartUVz = float3( (RayStartScreen.xy * float2( 0.5, -0.5 ) + 0.5), RayStartScreen.z );
-	float3 RayStepUVz = float3( RayStepScreen.xy * float2( 0.5, -0.5 ), RayStepScreen.z );
-	float2 RayStepRadiusUV = RayStepRadiusScreen.xy * float2( 0.5, -0.5 );
-#endif
-	const float RayStepRadiusFactor = length(RayStepRadiusUV * TextureSize);
+	float3 RayStepUVz  = float3(  RayStepScreen.xy  * float2( 0.5, -0.5 )		 * HZBUvFactorAndInvFactor.xy, RayStepScreen.z );
 	
 	const float Step = 1.0 / NumSteps;
-
-	const float CompareTolerance = max(abs( RayStepUVz.z ), ( RayStartScreen.z - RayDepthScreen.z ) * 4) * Step;
+	float CompareTolerance = max( abs( RayStepUVz.z ), ( RayStartScreen.z - RayDepthScreen.z ) * 2 ) * Step;
+	//const float CompareTolerance = abs( RayDepthScreen.z - RayStartScreen.z ) * Step * 2;
 
 	// avoid bugs with early returns inside of loops on certain platform compilers.
-	float4 Result = float4( 0, 0, 0, 1 );
-
-#if SCALAR_BRANCHLESS
-
-	float MinHitTime = 1;
-	float LastDiff = 0;
-
-	float SampleTime = StepOffset * Step + Step;
-
-	UNROLL
-	for( int i = 0; i < NumSteps; i++ )
-	{
-		float3 SampleUVz = RayStartUVz + RayStepUVz * SampleTime;
-		
-		// Use lower res for farther samples
-		float Level = Roughness * (i * 4.0 / NumSteps) + HZB_LEVEL_OFFSET;
-		float SampleDepth = Texture.SampleLevel( Sampler, SampleUVz.xy, Level ).r;
-
-		float DepthDiff = SampleUVz.z - SampleDepth;
-		bool Hit = abs( DepthDiff + CompareTolerance ) < CompareTolerance;
-
-		// Find more accurate hit using line segment intersection
-		float TimeLerp = saturate( LastDiff / (LastDiff - DepthDiff) );
-		float IntersectTime = SampleTime + TimeLerp * Step - Step;
-		float HitTime = Hit ? IntersectTime : 1;
-		MinHitTime = min( MinHitTime, HitTime );
-
-		LastDiff = DepthDiff;
-
-		SampleTime += Step;
-	}
-
-	float3 HitUVz = RayStartUVz + RayStepUVz * MinHitTime;
-
-	Result = float4( HitUVz, MinHitTime );
-
-#elif VECTORIZED_BRANCHLESS
-
-	float MinHitTime = 1;
-	float LastDiff = 0;
-	float Level = HZB_LEVEL_OFFSET;
-	
-	// Vectorized to group fetches
-	float4 SampleTime = ( StepOffset + float4( 1, 2, 3, 4 ) ) * Step;
-	float4 SampleUV0 = RayStartUVz.xyxy + RayStepUVz.xyxy * SampleTime.xxyy;
-	float4 SampleUV1 = RayStartUVz.xyxy + RayStepUVz.xyxy * SampleTime.zzww;
-	float4 SampleZ   = RayStartUVz.zzzz + RayStepUVz.zzzz * SampleTime;
-
-	LOOP
-	for( int i = 0; i < NumSteps; i += 4 )
-	{
-		// Use lower res for farther samples
-		float4 SampleDepth = SampleDepthTexture( Texture, Sampler, Level, SampleUV0, SampleUV1 );
-
-		#if 1
-			float4 DepthDiff = SampleZ - SampleDepth;
-			float4 IntersectTime = ( SampleDepth - RayStartUVz.zzzz ) / RayStepUVz.zzzz;
-		
-			bool4 Hit = abs( DepthDiff + CompareTolerance ) < CompareTolerance;
-	
-			// If hit set to intersect time. If missed set to 1, beyond end of ray
-			float4 HitTime = Hit ? IntersectTime : 1;
-
-			// Take closest hit
-			HitTime.xy = min( HitTime.xy, HitTime.zw );
-			MinHitTime = min( HitTime.x, HitTime.y );
-		#else
-			// Line segment intersection
-			float4 DepthDiff1 = SampleZ - SampleDepth;
-			float4 DepthDiff0 = float4( LastDiff, DepthDiff1.xyz );
-			float4 TimeLerp = saturate( DepthDiff0 / (DepthDiff0 - DepthDiff1) );
-			float4 IntersectTime = SampleTime + (TimeLerp - 1) / (NumSteps + 1);
-		
-			bool4 Hit = abs( DepthDiff1 + CompareTolerance ) < CompareTolerance;
-	
-			// If hit set to intersect time. If missed set to 1, beyond end of ray
-			float4 HitTime = Hit ? IntersectTime : 1;
-
-			// Take closest hit
-			HitTime.xy = min( HitTime.xy, HitTime.zw );
-			MinHitTime = min( MinHitTime, min( HitTime.x, HitTime.y ) );
-
-			LastDiff = DepthDiff1.w;
-		#endif
-		Level += Roughness * (16.0 / NumSteps);
-
-		SampleTime	+= 4.0 / (NumSteps + 1);
-		SampleUV0	+= RayStepUVz.xyxy * 4.0 / (NumSteps + 1);
-		SampleUV1	+= RayStepUVz.xyxy * 4.0 / (NumSteps + 1);
-		SampleZ		+= RayStepUVz.zzzz * 4.0 / (NumSteps + 1);
-	}
-
-	float3 HitUVz = RayStartUVz + RayStepUVz * MinHitTime;
-
-	Result = float4( HitUVz, MinHitTime );
-
-#elif 1
+	bool bHit = false;
+	OutHitUVz = float3( 0, 0, 0 );
 
 	float LastDiff = 0;
-	float Level = HZB_LEVEL_OFFSET;
+	Level = 1;
 
 	RayStepUVz *= Step;
 	float3 RayUVz = RayStartUVz + RayStepUVz * StepOffset;
 
 	LOOP
-	for( int i = 0; i < NumSteps; i += 4 )
+	for( uint i = 0; i < NumSteps; i += 4 )
 	{
 		// Vectorized to group fetches
 		float4 SampleUV0 = RayUVz.xyxy + RayStepUVz.xyxy * float4( 1, 1, 2, 2 );
@@ -211,242 +62,125 @@ void RayCast(
 		float4 SampleZ   = RayUVz.zzzz + RayStepUVz.zzzz * float4( 1, 2, 3, 4 );
 		
 		// Use lower res for farther samples
-		float4 SampleDepth = SampleDepthTexture( Texture, Sampler, Level, SampleUV0, SampleUV1 );
+		float4 SampleDepth;
+		SampleDepth.x = Texture.SampleLevel( Sampler, SampleUV0.xy, Level ).r;
+		SampleDepth.y = Texture.SampleLevel( Sampler, SampleUV0.zw, Level ).r;
+		Level += (8.0 / NumSteps) * Roughness;
+		
+		SampleDepth.z = Texture.SampleLevel( Sampler, SampleUV1.xy, Level ).r;
+		SampleDepth.w = Texture.SampleLevel( Sampler, SampleUV1.zw, Level ).r;
+		Level += (8.0 / NumSteps) * Roughness;
 
 		float4 DepthDiff = SampleZ - SampleDepth;
-		bool4 Hit = abs( -DepthDiff - CompareTolerance ) < CompareTolerance;
+		bool4 Hit = abs( DepthDiff + CompareTolerance ) < CompareTolerance;
 
 		BRANCH if( any( Hit ) )
 		{
-			float DepthDiff0 = DepthDiff[2];
-			float DepthDiff1 = DepthDiff[3];
-			float MinTime = 3;
-
-			FLATTEN if( Hit[2] )
+			if(1)
 			{
-				DepthDiff0 = DepthDiff[1];
-				DepthDiff1 = DepthDiff[2];
-				MinTime = 2;
-			}
-			FLATTEN if( Hit[1] )
-			{
-				DepthDiff0 = DepthDiff[0];
-				DepthDiff1 = DepthDiff[1];
-				MinTime = 1;
+				// If hit set to intersect time. If missed set to beyond end of ray
+				float4 HitTime = Hit ? float4( 1, 2, 3, 4 ) : 5;
+
+				// Take closest hit
+				float Time1 = min( min3( HitTime.x, HitTime.y, HitTime.z ), HitTime.w );
+				float Time0 = Time1 - 1;
+
+				const uint NumBinarySteps = Roughness < 0.2 ? 4 : 0;
+
+				// Binary search
+				for( uint j = 0; j < NumBinarySteps; j++ )
+				{
+					CompareTolerance *= 0.5;
+
+					float  MidTime = 0.5 * ( Time0 + Time1 );
+					float3 MidUVz = RayUVz + RayStepUVz * MidTime;
+					float  MidDepth = Texture.SampleLevel( Sampler, MidUVz.xy, Level ).r;
+					float  MidDepthDiff = MidUVz.z - MidDepth;
+
+					if( abs( MidDepthDiff + CompareTolerance ) < CompareTolerance )
+					{
+						Time1 = MidTime;
+					}
+					else
+					{
+						Time0 = MidTime;
+					}
+				}
+			
+				OutHitUVz = RayUVz + RayStepUVz * Time1;
 			}
-			FLATTEN if( Hit[0] )
+			else
 			{
-				DepthDiff0 = LastDiff;
-				DepthDiff1 = DepthDiff[0];
-				MinTime = 0;
-			}
+				float DepthDiff0 = DepthDiff[2];
+				float DepthDiff1 = DepthDiff[3];
+				float Time0 = 3;
+
+				FLATTEN if( Hit[2] )
+				{
+					DepthDiff0 = DepthDiff[1];
+					DepthDiff1 = DepthDiff[2];
+					Time0 = 2;
+				}
+				FLATTEN if( Hit[1] )
+				{
+					DepthDiff0 = DepthDiff[0];
+					DepthDiff1 = DepthDiff[1];
+					Time0 = 1;
+				}
+				FLATTEN if( Hit[0] )
+				{
+					DepthDiff0 = LastDiff;
+					DepthDiff1 = DepthDiff[0];
+					Time0 = 0;
+				}
+
+				float Time1 = Time0 + 1;
+	#if 0
+				// Binary search
+				for( uint j = 0; j < 4; j++ )
+				{
+					CompareTolerance *= 0.5;
+
+					float  MidTime = 0.5 * ( Time0 + Time1 );
+					float3 MidUVz = RayUVz + RayStepUVz * MidTime;
+					float  MidDepth = Texture.SampleLevel( Sampler, MidUVz.xy, Level ).r;
+					float  MidDepthDiff = MidUVz.z - MidDepth;
+
+					if( abs( MidDepthDiff + CompareTolerance ) < CompareTolerance )
+					{
+						DepthDiff1	= MidDepthDiff;
+						Time1		= MidTime;
+					}
+					else
+					{
+						DepthDiff0	= MidDepthDiff;
+						Time0		= MidTime;
+					}
+				}
+	#endif
 
-			// Find more accurate hit using line segment intersection
-			float TimeLerp = saturate( DepthDiff0 / (DepthDiff0 - DepthDiff1) );
-			float IntersectTime = MinTime + TimeLerp;
-			float3 HitUVz = RayUVz + RayStepUVz * IntersectTime;
+				// Find more accurate hit using line segment intersection
+				float TimeLerp = saturate( DepthDiff0 / (DepthDiff0 - DepthDiff1) );
+				float IntersectTime = Time0 + TimeLerp;
+				//float IntersectTime = lerp( Time0, Time1, TimeLerp );
+				
+				OutHitUVz = RayUVz + RayStepUVz * IntersectTime;
+			}
 
-			Result = float4( HitUVz, 0.5 );
+			bHit = true;
 			break;
 		}
 
 		LastDiff = DepthDiff.w;
-		Level += Roughness * (16.0 / NumSteps);
-
-		RayUVz += 4 * RayStepUVz;
-	}
-
-#elif 1
-
-	float Level = HZB_LEVEL_OFFSET;
-
-	RayStepUVz *= Step;
-	float3 RayUVz = RayStartUVz + RayStepUVz * StepOffset;
-
-	LOOP
-	for( uint i = 0; i < NumSteps; i += 4 )
-	{
-		// Vectorized to group fetches
-		float4 SampleUV0 = RayUVz.xyxy + RayStepUVz.xyxy * float4( 1, 1, 2, 2 );
-		float4 SampleUV1 = RayUVz.xyxy + RayStepUVz.xyxy * float4( 3, 3, 4, 4 );
-		float4 SampleZ   = RayUVz.zzzz + RayStepUVz.zzzz * float4( 1, 2, 3, 4 );
-		
-		// Use lower res for farther samples
-		float4 SampleDepth = SampleDepthTexture( Texture, Sampler, Level, SampleUV0, SampleUV1 );
-
-		float4 DepthDiff = SampleZ - SampleDepth;
-		bool4 Hit = abs( -DepthDiff - CompareTolerance ) < CompareTolerance;
-
-		BRANCH if( any( Hit ) )
-		{
-			float MinTime = 3;
-			MinTime = Hit[2] ? 2 : MinTime;
-			MinTime = Hit[1] ? 1 : MinTime;
-			MinTime = Hit[0] ? 0 : MinTime;
-
-			float3 HitUVz = RayUVz + RayStepUVz * MinTime;
 
-			Result = float4( HitUVz, 0.5 );
-			break;
-		}
-
-		Level += Roughness * (16.0 / NumSteps);
 		RayUVz += 4 * RayStepUVz;
 	}
-
-#else // VECTORIZED_EARLY_OUT
-
-	float LastDiff = 0;
-	float Level = HZB_LEVEL_OFFSET;
-
-	float4 SampleTime = ( StepOffset + float4( 1, 2, 3, 4 ) ) * Step;
-
-	LOOP
-	for( int i = 0; i < NumSteps; i += 4 )
-	{
-		// Vectorized to group fetches
-		float4 SampleUV0 = RayStartUVz.xyxy + RayStepUVz.xyxy * SampleTime.xxyy;
-		float4 SampleUV1 = RayStartUVz.xyxy + RayStepUVz.xyxy * SampleTime.zzww;
-		float4 SampleZ   = RayStartUVz.zzzz + RayStepUVz.zzzz * SampleTime;
-		
-		// Use lower res for farther samples
-		float4 SampleDepth = SampleDepthTexture( Texture, Sampler, Level, SampleUV0, SampleUV1 );
-
-		float4 DepthDiff1 = SampleZ - SampleDepth;
-		bool4 Hit = abs( -DepthDiff1 - CompareTolerance ) < CompareTolerance;
-
-		BRANCH if( any( Hit ) )
-		{
-			// Find more accurate hit using line segment intersection
-			float4 DepthDiff0 = float4( LastDiff, DepthDiff1.xyz );
-			float4 TimeLerp = saturate( DepthDiff0 / (DepthDiff0 - DepthDiff1) );
-			float4 IntersectTime = SampleTime + (TimeLerp - 1) / (NumSteps + 1);
-			float4 HitTime = Hit ? IntersectTime : 1;
-
-			// Take closest hit
-			HitTime.xy = min( HitTime.xy, HitTime.zw );
-			float MinHitTime = min( HitTime.x, HitTime.y );
-			
-			float3 HitUVz = RayStartUVz + RayStepUVz * MinHitTime;
-
-			Result = float4( HitUVz, MinHitTime );
-			break;
-		}
-
-		LastDiff = DepthDiff1.w;
-		Level += Roughness * (16.0 / NumSteps);
-
-		SampleTime += 4.0 / (NumSteps + 1);
-	}
-
-#endif
-	
-#if USE_HZB
-	Result.xy *= HZBUvFactorAndInvFactor.zw;
-	Result.xy = Result.xy * float2( 2, -2 ) + float2( -1, 1 );
-	Result.xy = Result.xy * View.ScreenPositionScaleBias.xy + View.ScreenPositionScaleBias.wz;
-#endif
-
-	OutHitUVzTime = Result;
-	OutLevelHCB = log2(OutHitUVzTime.w * RayStepRadiusFactor) + 1;
-}
-
-float4 RayCastCheap( Texture2D Texture, SamplerState Sampler, float3 R, float SceneDepth, float3 PositionTranslatedWorld, float StepOffset )
-{
-	// TODO provide RayStartUVz
-
-	// NOTE could clip ray against frustum planes
-
-	// TODO use screen position and skip matrix mul
-	float4 RayStartClip	= mul( float4( PositionTranslatedWorld, 1 ), View.TranslatedWorldToClip );
-	float4 RayEndClip	= mul( float4( PositionTranslatedWorld + R * SceneDepth, 1 ), View.TranslatedWorldToClip );
-
-	float3 RayStartScreen = RayStartClip.xyz / RayStartClip.w;
-	float3 RayEndScreen = RayEndClip.xyz / RayEndClip.w;
-
-	// Normalize 2D length
-	float3 RayStepScreen = ( RayEndScreen - RayStartScreen ) / length( RayEndScreen.xy - RayStartScreen.xy );
-	RayStepScreen *= 0.75;
-
-#if USE_HZB
-	float3 RayStartUVz = float3( RayStartScreen.xy * float2( 0.5, -0.5 ) + 0.5, RayStartScreen.z );
-	float3 RayStepUVz = float3( RayStepScreen.xy * float2( 0.5, -0.5 ), RayStepScreen.z );
-#else
-	float3 RayStartUVz = float3( RayStartScreen.xy * View.ScreenPositionScaleBias.xy + View.ScreenPositionScaleBias.wz, RayStartScreen.z );
-	float3 RayStepUVz = float3( RayStepScreen.xy * View.ScreenPositionScaleBias.xy, RayStepScreen.z );
-#endif
 	
-	const float Step = 1.0 / (4 + 1);
-
-	// *2 to get less morie pattern in extreme cases, larger values make object appear not grounded in reflections
-	const float CompareTolerance = abs( RayStepUVz.z ) * Step * 2;
-
-	// Vectorized to group fetches
-	float4 SampleTime = ( StepOffset + float4( 1, 2, 3, 4 ) ) * Step;
-	float4 SampleUV0 = RayStartUVz.xyxy + RayStepUVz.xyxy * SampleTime.xxyy;
-	float4 SampleUV1 = RayStartUVz.xyxy + RayStepUVz.xyxy * SampleTime.zzww;
-	float4 SampleZ   = RayStartUVz.zzzz + RayStepUVz.zzzz * SampleTime;
-
-	float4 SampleDepth;
-#if USE_HZB
-	SampleDepth.x = Texture2DSampleLevel( Texture, Sampler, SampleUV0.xy, 0 ).r;
-	SampleDepth.y = Texture2DSampleLevel( Texture, Sampler, SampleUV0.zw, 0 ).r;
-	SampleDepth.z = Texture2DSampleLevel( Texture, Sampler, SampleUV1.xy, 1 ).r;
-	SampleDepth.w = Texture2DSampleLevel( Texture, Sampler, SampleUV1.zw, 1 ).r;
-#else
-	SampleDepth.x = Texture2DSampleLevel( SceneTexturesStruct.SceneDepthTexture, SceneTexturesStruct.SceneDepthTextureSampler, SampleUV0.xy, 0 ).r;
-	SampleDepth.y = Texture2DSampleLevel( SceneTexturesStruct.SceneDepthTexture, SceneTexturesStruct.SceneDepthTextureSampler, SampleUV0.zw, 0 ).r;
-	SampleDepth.z = Texture2DSampleLevel( SceneTexturesStruct.SceneDepthTexture, SceneTexturesStruct.SceneDepthTextureSampler, SampleUV1.xy, 0 ).r;
-	SampleDepth.w = Texture2DSampleLevel( SceneTexturesStruct.SceneDepthTexture, SceneTexturesStruct.SceneDepthTextureSampler, SampleUV1.zw, 0 ).r;
-#endif
-
-	#if 1
-		float4 DepthDiff = SampleZ - SampleDepth;
-		float4 IntersectTime = ( SampleDepth - RayStartUVz.zzzz ) / RayStepUVz.zzzz;
-		
-		bool4 Hit = abs( DepthDiff + CompareTolerance ) < CompareTolerance;
-	
-		// If hit set to intersect time. If missed set to 1, beyond end of ray
-		float4 HitTime = Hit ? IntersectTime : 1;
-
-		// Take closest hit
-		HitTime.xy = min( HitTime.xy, HitTime.zw );
-		float MinHitTime = min( HitTime.x, HitTime.y );
-	#else
-		// Line segment intersection
-		float4 DepthDiff1 = SampleZ - SampleDepth;
-		float4 DepthDiff0 = float4( 0, DepthDiff1.xyz );
-		float4 TimeLerp = saturate( DepthDiff0 / (DepthDiff0 - DepthDiff1) );
-		float4 IntersectTime = SampleTime + (TimeLerp - 1) * Step;
-		
-		bool4 Hit = abs( DepthDiff1 + CompareTolerance ) < CompareTolerance;
-	
-		// If hit set to intersect time. If missed set to 1, beyond end of ray
-		float4 HitTime = Hit ? IntersectTime : 1;
-
-		// Take closest hit
-		HitTime.xy = min( HitTime.xy, HitTime.zw );
-		float MinHitTime = min( HitTime.x, HitTime.y );
-	#endif
-
-	float3 HitUVz = RayStartUVz + RayStepUVz * MinHitTime;
+	OutHitUVz.xy *= HZBUvFactorAndInvFactor.zw;
+	OutHitUVz.xy = OutHitUVz.xy * float2( 2, -2 ) + float2( -1, 1 );
+	OutHitUVz.xy = OutHitUVz.xy * View.ScreenPositionScaleBias.xy + View.ScreenPositionScaleBias.wz;
 
-	BRANCH
-	if( MinHitTime == 1 && RayStepUVz.z < 0 )
-	{
-		float4 RayInfClip	= mul( float4( PositionTranslatedWorld + R * 2000000, 1 ), View.TranslatedWorldToClip );
-		float3 RayInfScreen = RayInfClip.xyz / RayInfClip.w;
-		float3 RayInfUVz = float3( RayInfScreen.xy * float2( 0.5, -0.5 ) + 0.5, RayInfScreen.z );
-
-		HitUVz = RayInfUVz;
-		MinHitTime = 0.5;
-	}
-
-#if USE_HZB
-	HitUVz.xy = HitUVz.xy = HitUVz.xy * float2( 2, -2 ) + float2( -1, 1 );
-	HitUVz.xy = HitUVz.xy * View.ScreenPositionScaleBias.xy + View.ScreenPositionScaleBias.wz;
-#endif
-	return float4( HitUVz, MinHitTime );
+	return bHit;
 }
 
 float ComputeHitVignetteFromScreenPos(float2 ScreenPos)
@@ -460,7 +194,6 @@ float ComputeHitVignetteFromScreenPos(float2 ScreenPos)
 
 void ReprojectHit(float4 PrevScreenPositionScaleBias, float3 HitUVz, out float2 OutPrevUV, out float OutVignette)
 {
-	//return HitUVz.xy;
 	// Camera motion for pixel (in ScreenPos space).
 	float2 ThisScreen = (HitUVz.xy - View.ScreenPositionScaleBias.wz) / View.ScreenPositionScaleBias.xy;
 	float4 ThisClip = float4( ThisScreen, HitUVz.z, 1 );
@@ -474,7 +207,6 @@ void ReprojectHit(float4 PrevScreenPositionScaleBias, float3 HitUVz, out float2
 
 void ReprojectHit(float4 PrevScreenPositionScaleBias, Texture2D Texture, SamplerState Sampler, float3 HitUVz, out float2 OutPrevUV, out float OutVignette)
 {
-	//return HitUVz.xy;
 	// Camera motion for pixel (in ScreenPos space).
 	float2 ThisScreen = (HitUVz.xy - View.ScreenPositionScaleBias.wz) / View.ScreenPositionScaleBias.xy;
 	float4 ThisClip = float4( ThisScreen, HitUVz.z, 1 );
@@ -500,48 +232,19 @@ float4 SampleScreenColor(Texture2D Texture, SamplerState Sampler, float2 UV)
 	OutColor.rgb = Texture.SampleLevel( Sampler, UV, 0 ).rgb;
 	// Transform NaNs to black, transform negative colors to black.
 	OutColor.rgb = -min(-OutColor.rgb, 0.0);
-	
 	OutColor.a = 1;
 	
 	return OutColor;
 }
 
-float4 SampleHCBLevel( Texture2D Texture, SamplerState Sampler, float3 HitUVz, float Level, float4 HZBUvFactorAndInvFactor )
+float4 SampleHCBLevel( Texture2D Texture, SamplerState Sampler, float2 UV, float Level, float4 HZBUvFactorAndInvFactor )
 {
 	float4 OutColor;
 
-#if PREV_FRAME_COLOR
-	// Find previous screen position for hit since color buffer is from last frame
-	// TODO combine to single matrix
-	float4 HitClip = float4( (HitUVz.xy - View.ScreenPositionScaleBias.wz) / View.ScreenPositionScaleBias.xy, HitUVz.z, 1 );
-	float4 HitTranslatedWorld = mul( HitClip, View.ClipToTranslatedWorld );
-	HitTranslatedWorld /= HitTranslatedWorld.w;
-		
-	float3 PrevTranslatedWorld = HitTranslatedWorld.xyz + (View.PrevPreViewTranslation - View.PreViewTranslation);
-	float4 PrevClip = mul( float4( PrevTranslatedWorld, 1 ), View.PrevTranslatedWorldToClip );
-	float2 PrevScreen = PrevClip.xy / PrevClip.w;
-	float2 PrevUV = PrevScreen.xy * View.ScreenPositionScaleBias.xy + View.ScreenPositionScaleBias.wz;
-
-	OutColor.rgb = Texture.SampleLevel( Sampler, PrevUV * HZBUvFactorAndInvFactor.xy, Level ).rgb;
-	OutColor.a = 1;
-
-	// Off screen masking
-	float2 Vignette = saturate( abs( PrevScreen ) * 5 - 4 );
-#else
-	OutColor.rgb = Texture.SampleLevel( Sampler, HitUVz.xy * HZBUvFactorAndInvFactor.xy, Level ).rgb;
-	OutColor.a = 1;
-
-	// Off screen masking
-	float2 HitScreenPos = ( HitUVz.xy - View.ScreenPositionScaleBias.wz ) / View.ScreenPositionScaleBias.xy;
-	float2 Vignette = saturate( abs( HitScreenPos ) * 5 - 4 );
-#endif
-		
-	//PrevScreen sometimes has NaNs or Infs.  DX11 is protected because saturate turns NaNs -> 0.
-	//Do a SafeSaturate so other platforms get the same protection.
-	OutColor *= SafeSaturate( 1.0 - dot( Vignette, Vignette ) );
-
+	OutColor.rgb = Texture.SampleLevel( Sampler, UV * HZBUvFactorAndInvFactor.xy, Level ).rgb;
 	// Transform NaNs to black, transform negative colors to black.
 	OutColor.rgb = -min(-OutColor.rgb, 0.0);
+	OutColor.a = 1;
 
 	return OutColor;
 }
\ No newline at end of file
diff --git a/Engine/Shaders/Private/ScreenSpaceReflections.usf b/Engine/Shaders/Private/ScreenSpaceReflections.usf
--- a/Engine/Shaders/Private/ScreenSpaceReflections.usf
+++ b/Engine/Shaders/Private/ScreenSpaceReflections.usf
@@ -12,10 +12,6 @@
 #include "MonteCarlo.ush"
 #include "ScreenSpaceRayCast.ush"
 
-// 0:old for reference / 1:new optimized (fewer ALU instructions, minor performance difference on NV Titan)
-// define can be removed or used to judge the quality difference or profile on different platforms
-#define RANDOM_OPTIMZATION 1
-
 #define SSR_CONE_QUALITY 5
 
 // .r:Intensity in 0..1 range .g:RoughnessMaskMul, b:EnableDiscard for FPostProcessScreenSpaceReflectionsStencilPS, a:(bTemporalAAIsOn?TemporalAAIndex:StateFrameIndexMod8)*1551
@@ -28,25 +24,6 @@ float4 HZBUvFactorAndInvFactor;
 // Equivalent of View.ScreenPositionScaleBias, but for SSR.
 float4 PrevScreenPositionScaleBias;
 
-uint MortonCode( uint x )
-{
-	//x = (x ^ (x <<  8)) & 0x00ff00ff;
-	//x = (x ^ (x <<  4)) & 0x0f0f0f0f;
-	x = (x ^ (x <<  2)) & 0x33333333;
-	x = (x ^ (x <<  1)) & 0x55555555;
-	return x;
-}
-
-uint ReverseUIntBits( uint bits )
-{
-	//bits = ( bits << 16) | ( bits >> 16);
-	//bits = ( (bits & 0x00ff00ff) << 8 ) | ( (bits & 0xff00ff00) >> 8 );
-	//bits = ( (bits & 0x0f0f0f0f) << 4 ) | ( (bits & 0xf0f0f0f0) >> 4 );
-	bits = ( (bits & 0x33333333) << 2 ) | ( (bits & 0xcccccccc) >> 2 );
-	bits = ( (bits & 0x55555555) << 1 ) | ( (bits & 0xaaaaaaaa) >> 1 );
-	return bits;
-}
-
 float GetRoughness(const in FGBufferData GBuffer)
 {
 	float Roughness = GBuffer.Roughness;
@@ -105,7 +82,7 @@ void ScreenSpaceReflectionsPS(in noperspective float4 UVAndScreenPos : TEXCOORD0
 {
 	float2 UV = UVAndScreenPos.xy;
 	float2 ScreenPos = UVAndScreenPos.zw;
-	int2 PixelPos = (int2)SvPosition.xy;
+	uint2 PixelPos = (uint2)SvPosition.xy;
 
 	OutColor = 0;
 
@@ -121,188 +98,178 @@ void ScreenSpaceReflectionsPS(in noperspective float4 UVAndScreenPos : TEXCOORD0
 		return;
 	}
 
-#if SSR_QUALITY == 0
-	// visualize SSR
-
-	float PatternMask = ((PixelPos.x / 2 + PixelPos.y / 2) % 2) * 0.7f;
-
-	OutColor = lerp(float4(1, 0, 0, 1), float4(1, 1 ,0, 1), PatternMask) * 0.3f;
-	return;
-#endif
-
 	const float3 N = GBuffer.WorldNormal;
 	const float SceneDepth = CalcSceneDepth(UV);
 	const float3 PositionTranslatedWorld = mul( float4( ScreenPos * SceneDepth, SceneDepth, 1 ), View.ScreenToTranslatedWorld ).xyz;
 	const float3 V = normalize(View.TranslatedWorldCameraOrigin - PositionTranslatedWorld);
 
-	float FrameRandom = SSRParams.a;
-
-#if SSR_QUALITY == SSR_CONE_QUALITY
-	const int NumSteps = 16;
+	float a = Roughness * Roughness;
+	float a2 = a * a;
 	
-	{
-		const float2 PixelPos = SvPosition.xy;
-		const float2 E = float2(PseudoRandom(PixelPos + float2(FrameRandom, 0)), PseudoRandom(PixelPos + float2(0, FrameRandom)));
-		const float StepOffset = PseudoRandom(PixelPos + FrameRandom);
-
-		const float4 HNAndPDF = ImportanceSampleBlinn( E, Pow4(Roughness) );
-		const float3 H = TangentToWorld( HNAndPDF.xyz, N );
-		const float3 R = 2 * dot( V, H ) * H - V;
+	float NoV = saturate( dot( N, V ) );
+	float G_SmithV = 2 * NoV / (NoV + sqrt(NoV * (NoV - NoV * a2) + a2));
 
-		const float ReflectPDF = HNAndPDF.w / (4 * saturate(dot(V, H)));
-		const float TotalSamples = View.TemporalAAParams.g;
-		const float SolidAngleSample = 1.0 / ( TotalSamples * ReflectPDF );
-		const float ConeAngleWorld = acos( 1 - SolidAngleSample / (2*PI) );
-
-		float4 HitUVzTime;
-		float HCBLevel;
-		
-		RayCast(
-			PostprocessInput1, PostprocessInput1Sampler, PostprocessInput1Size.xy,
-			PositionTranslatedWorld, R, Roughness, ConeAngleWorld, SceneDepth,
-			NumSteps, StepOffset,
-			HZBUvFactorAndInvFactor,
-			PrevScreenPositionScaleBias,
-			HitUVzTime, HCBLevel
-		);
-
-		// if there was a hit
-		BRANCH if( HitUVzTime.w < 1 )
-		{
-			OutColor = SampleHCBLevel( PostprocessInput2, PostprocessInput2Sampler, HitUVzTime.xyz, HCBLevel, HZBUvFactorAndInvFactor );
-		}
-	}
-#else //SSR_QUALITY != SSR_CONE_QUALITY
 #if SSR_QUALITY == 1
-	const int NumSteps = 8;
-	const int NumRays = 1;
+	uint NumSteps = 8;
+	uint NumRays = 1;
 #elif SSR_QUALITY == 2
-	const int NumSteps = 16;
-	const int NumRays = 1;
+	uint NumSteps = 16;
+	uint NumRays = 1;
 #elif SSR_QUALITY == 3
-	const int NumSteps = 8;
-	const int NumRays = 4;
+	uint NumSteps = 8;
+	uint NumRays = 4;
 #else // SSR_QUALITY == 4
-	const int NumSteps = 12;
-	const int NumRays = 12;
+	uint NumSteps = 12;
+	uint NumRays = 12;
 #endif
 	
 	if( NumRays > 1 )
 	{
-#if RANDOM_OPTIMZATION
-		// magic numbers are found by experimentation
-//		uint2 Random = PseudoRandom(SvPosition.xy + View.StateFrameIndexMod8 * float2(97, 71)) * uint2(12, 25);	// SSR_QUALITY=3 instructions=253 but looks bad
-		uint2 Random = PseudoRandom(SvPosition.xy + View.StateFrameIndexMod8 * float2(97, 71)) * uint2(0x3127352, 0x11229256);		// doesn't compile with external HLSL with obscure error message but the line above does
-#else
-		uint2 Random = Rand3DPCG16(uint3(PixelPos, View.StateFrameIndexMod8)).xy;	// reference, SSR_QUALITY=3
-#endif
-		if( Roughness > 0.1 )
+		float2 Noise;
+		Noise.x = InterleavedGradientNoise( SvPosition.xy, View.StateFrameIndexMod8 );
+		Noise.y = InterleavedGradientNoise( SvPosition.xy, View.StateFrameIndexMod8 * 117 );
+		
+		//uint2 Random = 0x10000 * Noise;
+		uint2 Random = Rand3DPCG16( int3( PixelPos, View.StateFrameIndexMod8 ) ).xy;
+			
+		float3x3 TangentBasis = GetTangentBasis( N );
+		float3 TangentV = mul( TangentBasis, V );
+
+		float Count = 0;
+
+		if( Roughness < 0.1 )
 		{
-			// Sample set dithered over 4x4 pixels
-			uint Morton = MortonCode( PixelPos.x & 3 ) | ( MortonCode( PixelPos.y & 3 ) * 2 );
-			uint PixelIndex = ReverseUIntBits( Morton );
+			NumSteps = min( NumSteps * NumRays, 24 );
+			NumRays = 1;
+		}
 
-			// Shoot multiple rays
-			LOOP for( int i = 0; i < NumRays; i++ )
-			{
-				// TODO better per ray offset?
-				uint Offset = ( PixelIndex + ReverseUIntBits( FrameRandom + i * 117 ) ) & 15;
-				float StepOffset = Offset / 15.0;
-				StepOffset -= 0.5;
+		// Shoot multiple rays
+		LOOP for( uint i = 0; i < NumRays; i++ )
+		{
+			float StepOffset = Noise.x;
+			StepOffset -= 0.9;
 			
-				float2 E = Hammersley( i, NumRays, Random );
-				float3 H = TangentToWorld( ImportanceSampleBlinn( E, Pow4(Roughness) ).xyz, N );
-				float3 R = 2 * dot( V, H ) * H - V;
+			float2 E = Hammersley16( i, NumRays, Random );
+
+#if 1
+			// TODO EnvBRDF isn't a perfect match anymore.
+			float3 H = mul( ImportanceSampleVisibleGGX( E, a2, TangentV ).xyz, TangentBasis );
+			float3 L = 2 * dot( V, H ) * H - V;
+#elif 0
+			float3 H = mul( ImportanceSampleGGX( E, a2 ).xyz, TangentBasis );
+			float3 L = 2 * dot( V, H ) * H - V;
+#elif 0
+			float3 L = CosineSampleHemisphere( E, N ).xyz;
+#elif 0
+			float3 L = CosineSampleHemisphere( E ).xyz;
+			L = mul( L, TangentBasis );
+#else
+			float3 L;
+			L.xy = UniformSampleDiskConcentric( E );
+			L.z = sqrt( 1 - dot( L.xy, L.xy ) );
+			L = mul( L, TangentBasis );
+			float3 H = normalize(V + L);
+#endif
 
-				float4 HitUVzTime;
-				float HCBLevel;
+			float NoL = saturate( dot(N, L) );
+			float NoH = saturate( dot(N, H) );
+			float VoH = saturate( dot(V, H) );
+		
+			if( NoL > 0 && VoH > 0 )
+			{
+				float D = D_GGX( a2, NoH );
+				float G_SmithL = 2 * NoL / (NoL + sqrt(NoL * (NoL - NoL * a2) + a2));
+				float Vis = Vis_Smith( a2, NoV,  NoL );
+				float3 F = F_Schlick( GBuffer.SpecularColor, VoH );
+				
+				float3 HitUVz;
+				float Level = 0;
 
-				RayCast(
-					PostprocessInput1, PostprocessInput1Sampler, PostprocessInput1Size.xy,
-					PositionTranslatedWorld, R, Roughness, 0.001, SceneDepth,
+				bool bHit = RayCast(
+					PostprocessInput1, PostprocessInput1Sampler,
+					PositionTranslatedWorld, L, Roughness, SceneDepth, 
 					NumSteps, StepOffset,
 					HZBUvFactorAndInvFactor,
-					PrevScreenPositionScaleBias,
-					HitUVzTime, HCBLevel
+					HitUVz,
+					Level
 				);
 
+#if 0	// Backface check
+				if( bHit )
+				{
+					float3 SampleNormal = GetGBufferData( HitUVz.xy ).WorldNormal;
+					bHit = dot( SampleNormal, L ) < 0;
+				}
+#endif
+
 				// if there was a hit
-				BRANCH if( HitUVzTime.w < 1 )
+				BRANCH if( bHit )
 				{
 					float2 SampleUV;
 					float Vignette;
-					ReprojectHit(PrevScreenPositionScaleBias, PostprocessInput3, PostprocessInput3Sampler, HitUVzTime.xyz, SampleUV, Vignette);
+					ReprojectHit(PrevScreenPositionScaleBias, PostprocessInput3, PostprocessInput3Sampler, HitUVz, SampleUV, Vignette);
+
+#if SSR_QUALITY == SSR_CONE_QUALITY
+					float4 SampleColor = SampleHCBLevel( PostprocessInput2, PostprocessInput2Sampler, SampleUV, Level, HZBUvFactorAndInvFactor ) * Vignette;
+#else
+					float4 SampleColor = SampleScreenColor( PostprocessInput0, PostprocessInput0Sampler, SampleUV ) * Vignette;
+#endif
+					SampleColor.rgb *= rcp( 1 + Luminance(SampleColor.rgb) );
+
+					// Correct integration applies DGF below but for speed we apply EnvBRDF later when compositing
+#if 0
+					// CosineSampleHemisphere
+					// PDF = NoL / PI,
+					SampleColor.rgb *= F * (D * Vis * PI);
+#elif 0
+					// ImportanceSampleGGX
+					// PDF = D * NoH / (4 * VoH),
+					SampleColor.rgb *= F * ( NoL * Vis * (4 * VoH / NoH) );
+#elif 0
+					// ImportanceSampleVisibleGGX
+					// PDF = G_SmithV * VoH * D / NoV / (4 * VoH)
+					// PDF = G_SmithV * D / (4 * NoV);
+					SampleColor.rgb *= F * G_SmithL;
+#endif
 
-					float4 SampleColor = SampleScreenColor(PostprocessInput0, PostprocessInput0Sampler, SampleUV) * Vignette;
-				
-					SampleColor.rgb /= 1 + Luminance(SampleColor.rgb);
 					OutColor += SampleColor;
 				}
-			}
-
-			OutColor /= NumRays;
-			OutColor.rgb /= 1 - Luminance(OutColor.rgb);
-		}
-		else
-		{
-			float StepOffset = InterleavedGradientNoise(SvPosition.xy, View.StateFrameIndexMod8);
-			StepOffset -= 0.5;
-		
-			float3 R = reflect( -V, N );
-		
-			float4 HitUVzTime;
-			float HCBLevel;
-
-			RayCast(
-				PostprocessInput1, PostprocessInput1Sampler, PostprocessInput1Size.xy,
-				PositionTranslatedWorld, R, Roughness, 0.001, SceneDepth,
-				16, StepOffset,
-				HZBUvFactorAndInvFactor,
-				PrevScreenPositionScaleBias,
-				HitUVzTime, HCBLevel
-			);
-
-			// if there was a hit
-			BRANCH if( HitUVzTime.w < 1 )
-			{
-				float2 SampleUV;
-				float Vignette;
-				ReprojectHit(PrevScreenPositionScaleBias, PostprocessInput3, PostprocessInput3Sampler, HitUVzTime.xyz, SampleUV, Vignette);
 
-				OutColor = SampleScreenColor(PostprocessInput0, PostprocessInput0Sampler, SampleUV) * Vignette;
+				Count++;
 			}
 		}
+
+		//OutColor /= NumRays;
+		OutColor /= max( Count, 0.0001 );
+		OutColor.rgb *= rcp( 1 - Luminance(OutColor.rgb) );
 	}
 	else
 	{
 		float StepOffset = InterleavedGradientNoise(SvPosition.xy, View.StateFrameIndexMod8);
-		StepOffset -= 0.5;
+		StepOffset -= 0.9;
 		
 		float3 R = reflect( -V, N );
-		
-		float4 HitUVzTime;
-		float HCBLevel;
+		float3 HitUVz;
+		float Level = 0;
 
-		RayCast(
-			PostprocessInput1, PostprocessInput1Sampler, PostprocessInput1Size.xy,
-			PositionTranslatedWorld, R, Roughness, 0.001, SceneDepth,
+		bool bHit = RayCast(
+			PostprocessInput1, PostprocessInput1Sampler,
+			PositionTranslatedWorld, R, Roughness, SceneDepth,
 			NumSteps, StepOffset,
 			HZBUvFactorAndInvFactor,
-			PrevScreenPositionScaleBias,
-			HitUVzTime, HCBLevel
+			HitUVz,
+			Level
 		);
 
-		// if there was a hit
-		BRANCH if( HitUVzTime.w < 1 )
+		BRANCH if( bHit )
 		{
 			float2 SampleUV;
 			float Vignette;
-			ReprojectHit(PrevScreenPositionScaleBias, PostprocessInput3, PostprocessInput3Sampler, HitUVzTime.xyz, SampleUV, Vignette);
+			ReprojectHit(PrevScreenPositionScaleBias, PostprocessInput3, PostprocessInput3Sampler, HitUVz, SampleUV, Vignette);
 
 			OutColor = SampleScreenColor(PostprocessInput0, PostprocessInput0Sampler, SampleUV) * Vignette;
 		}
 	}
-#endif //SSR_QUALITY  != SSR_CONE_QUALITY
 	
 	OutColor *= RoughnessFade;
 	OutColor *= SSRParams.r;
diff --git a/Engine/Source/Runtime/Renderer/Private/DeferredShadingRenderer.cpp b/Engine/Source/Runtime/Renderer/Private/DeferredShadingRenderer.cpp
--- a/Engine/Source/Runtime/Renderer/Private/DeferredShadingRenderer.cpp
+++ b/Engine/Source/Runtime/Renderer/Private/DeferredShadingRenderer.cpp
@@ -479,6 +479,8 @@ static FORCEINLINE bool NeedsPrePass(const FDeferredShadingSceneRenderer* Render
 		(Renderer->EarlyZPassMode != DDM_None || Renderer->bEarlyZPassMovable != 0);
 }
 
+bool ShouldRenderScreenSpaceDiffuseIndirect( const FViewInfo& View );
+
 bool FDeferredShadingSceneRenderer::RenderHzb(FRHICommandListImmediate& RHICmdList)
 {
 	FSceneRenderTargets& SceneContext = FSceneRenderTargets::Get(RHICmdList);
@@ -494,10 +496,11 @@ bool FDeferredShadingSceneRenderer::RenderHzb(FRHICommandListImmediate& RHICmdLi
 		FViewInfo& View = Views[ViewIndex];
 		FSceneViewState* ViewState = (FSceneViewState*)View.State;
 
-		const uint32 bSSR = ShouldRenderScreenSpaceReflections(View);
+		const bool bSSR  = ShouldRenderScreenSpaceReflections(View);
 		const bool bSSAO = ShouldRenderScreenSpaceAmbientOcclusion(View);
+		const bool bSSGI = ShouldRenderScreenSpaceDiffuseIndirect(View);
 
-		if (bSSAO || bHZBOcclusion || bSSR)
+		if (bSSAO || bHZBOcclusion || bSSR || bSSGI)
 		{
 			FRDGBuilder GraphBuilder(RHICmdList);
 			{
@@ -977,6 +980,7 @@ FRHIRayTracingPipelineState* FDeferredShadingSceneRenderer::BindRayTracingPipeli
 #endif // RHI_RAYTRACING
 
 extern bool IsLpvIndirectPassRequired(const FViewInfo& View);
+void RenderScreenSpaceDiffuseIndirect( FRHICommandListImmediate& RHICmdList, FViewInfo& View, TRefCountPtr<IPooledRenderTarget>& VelocityRT );
 
 static TAutoConsoleVariable<float> CVarStallInitViews(
 	TEXT("CriticalPathStall.AfterInitViews"),
@@ -1853,6 +1857,8 @@ void FDeferredShadingSceneRenderer::Render(FRHICommandListImmediate& RHICmdList)
 			Scene->UniformBuffers.UpdateViewUniformBuffer(Views[ViewIndex]);
 
 			GCompositionLighting.ProcessAfterBasePass(RHICmdList, Views[ViewIndex]);
+			
+			RenderScreenSpaceDiffuseIndirect( RHICmdList, Views[ ViewIndex ], VelocityRT );
 		}
 		ServiceLocalQueue();
 	}
diff --git a/Engine/Source/Runtime/Renderer/Private/ScreenSpaceDiffuseIndirect.cpp b/Engine/Source/Runtime/Renderer/Private/ScreenSpaceDiffuseIndirect.cpp
new file mode 100644
--- /dev/null
+++ b/Engine/Source/Runtime/Renderer/Private/ScreenSpaceDiffuseIndirect.cpp
@@ -0,0 +1,196 @@
+// Copyright 1998-2018 Epic Games, Inc. All Rights Reserved.
+
+#include "StaticBoundShaderState.h"
+#include "SceneUtils.h"
+#include "PostProcess/SceneRenderTargets.h"
+#include "SceneRenderTargetParameters.h"
+#include "ScenePrivate.h"
+#include "PostProcessing.h"
+#include "PostProcess/SceneFilterRendering.h"
+#include "ClearQuad.h"
+#include "PipelineStateCache.h"
+
+static TAutoConsoleVariable<int32> CVarSSGIQuality(
+	TEXT("r.SSGI.Quality"),
+	0,
+	TEXT("Whether to use screen space diffuse indirect and at what quality setting.\n"),
+	ECVF_Scalability | ECVF_RenderThreadSafe);
+
+DECLARE_GPU_STAT_NAMED(ScreenSpaceDiffuseIndirect, TEXT("Screen Space Diffuse Indirect"));
+
+bool ShouldRenderScreenSpaceDiffuseIndirect( const FViewInfo& View )
+{
+	int Quality = CVarSSGIQuality.GetValueOnRenderThread();
+
+	if( Quality <= 0 )
+	{
+		return false;
+	}
+
+	if (IsAnyForwardShadingEnabled(View.GetShaderPlatform()))
+	{
+		return false;
+	}
+
+	return true;
+}
+
+
+class FScreenSpaceDiffuseIndirectPS : public FGlobalShader
+{
+	DECLARE_GLOBAL_SHADER(FScreenSpaceDiffuseIndirectPS);
+	SHADER_USE_PARAMETER_STRUCT(FScreenSpaceDiffuseIndirectPS, FGlobalShader)
+
+	class FQualityDim : SHADER_PERMUTATION_INT( "QUALITY", 5 );
+	using FPermutationDomain = TShaderPermutationDomain< FQualityDim >;
+	
+	BEGIN_SHADER_PARAMETER_STRUCT( FParameters, )
+		SHADER_PARAMETER( FVector4,		HZBUvFactorAndInvFactor )
+		SHADER_PARAMETER( FVector4,		PrevScreenPositionScaleBias )
+		SHADER_PARAMETER( float,		PrevSceneColorPreExposureCorrection )
+		
+		RENDER_TARGET_BINDING_SLOTS()
+
+		SHADER_PARAMETER_RDG_TEXTURE( Texture2D,	HZBTexture )
+		SHADER_PARAMETER_SAMPLER( SamplerState,		HZBSampler )
+
+		SHADER_PARAMETER_RDG_TEXTURE( Texture2D,	VelocityTexture )
+		SHADER_PARAMETER_SAMPLER( SamplerState,		VelocitySampler )
+
+		SHADER_PARAMETER_RDG_TEXTURE( Texture2D,	ColorTexture )
+		SHADER_PARAMETER_SAMPLER( SamplerState,		ColorSampler )
+
+		SHADER_PARAMETER_STRUCT_REF(	FViewUniformShaderParameters,	View )
+		SHADER_PARAMETER_STRUCT_REF(	FSceneTexturesUniformParameters,SceneTextures )
+	END_SHADER_PARAMETER_STRUCT()
+
+	static bool ShouldCompilePermutation(const FGlobalShaderPermutationParameters& Parameters)
+	{
+		return IsFeatureLevelSupported(Parameters.Platform, ERHIFeatureLevel::SM4);
+	}
+};
+
+IMPLEMENT_GLOBAL_SHADER( FScreenSpaceDiffuseIndirectPS, "/Engine/Private/ScreenSpaceDiffuseIndirect.usf", "ScreenSpaceDiffuseIndirectPS", SF_Pixel );
+
+// --------------------------------------------------------
+
+void RenderScreenSpaceDiffuseIndirect( FRHICommandListImmediate& RHICmdList, FViewInfo& View, TRefCountPtr<IPooledRenderTarget>& VelocityRT )
+{
+	FTemporalAAHistory& TemporalAAHistory = View.PrevViewInfo.TemporalAAHistory;
+
+	if( !ShouldRenderScreenSpaceDiffuseIndirect( View ) || !TemporalAAHistory.IsValid() )
+	{
+		return;
+	}
+	
+	const int32 Quality = FMath::Clamp( CVarSSGIQuality.GetValueOnRenderThread(), 1, 4 );
+
+	FRDGBuilder GraphBuilder( RHICmdList );
+
+	FSceneRenderTargets& SceneContext = FSceneRenderTargets::Get( RHICmdList );
+
+	const FRDGTexture* SceneColor	= GraphBuilder.RegisterExternalTexture( SceneContext.GetSceneColor() );
+	const FRDGTexture* ScreenSpaceAO= GraphBuilder.RegisterExternalTexture( SceneContext.ScreenSpaceAO );
+	const FRDGTexture* HZBTexture	= GraphBuilder.RegisterExternalTexture( View.HZB );
+	const FRDGTexture* ColorTexture	= GraphBuilder.RegisterExternalTexture( TemporalAAHistory.RT[0] );
+
+	const FRDGTexture* VelocityTexture;
+	if( VelocityRT && !View.bCameraCut )
+	{
+		VelocityTexture = GraphBuilder.RegisterExternalTexture( VelocityRT );
+	}
+	else
+	{
+		// No velocity, use black
+		VelocityTexture = GraphBuilder.RegisterExternalTexture( GSystemTextures.BlackDummy );
+	}
+
+	FScreenSpaceDiffuseIndirectPS::FParameters* PassParameters = GraphBuilder.AllocParameters< FScreenSpaceDiffuseIndirectPS::FParameters >();
+
+	PassParameters->RenderTargets[0] = FRenderTargetBinding( SceneColor, ERenderTargetLoadAction::ELoad, ERenderTargetStoreAction::EStore );
+	PassParameters->RenderTargets[1] = FRenderTargetBinding( ScreenSpaceAO, ERenderTargetLoadAction::ENoAction, ERenderTargetStoreAction::EStore );
+
+	PassParameters->HZBTexture = HZBTexture;
+	PassParameters->HZBSampler = TStaticSamplerState< SF_Point, AM_Clamp, AM_Clamp, AM_Clamp >::GetRHI();
+
+	PassParameters->VelocityTexture = VelocityTexture;
+	PassParameters->VelocitySampler = TStaticSamplerState< SF_Point, AM_Clamp, AM_Clamp, AM_Clamp >::GetRHI();
+
+	PassParameters->ColorTexture = ColorTexture;
+	PassParameters->ColorSampler = TStaticSamplerState< SF_Bilinear, AM_Clamp, AM_Clamp, AM_Clamp >::GetRHI();
+
+	const FVector2D HZBUvFactor(
+		float( View.ViewRect.Width() )  / float( 2 * View.HZBMipmap0Size.X ),
+		float( View.ViewRect.Height() ) / float( 2 * View.HZBMipmap0Size.Y )
+		);
+			
+	PassParameters->HZBUvFactorAndInvFactor = FVector4(
+		HZBUvFactor.X,
+		HZBUvFactor.Y,
+		1.0f / HZBUvFactor.X,
+		1.0f / HZBUvFactor.Y );
+
+	FIntPoint ViewportOffset	= TemporalAAHistory.ViewportRect.Min;
+	FIntPoint ViewportExtent	= TemporalAAHistory.ViewportRect.Size();
+	FIntPoint BufferSize		= TemporalAAHistory.ReferenceBufferSize;
+
+	PassParameters->PrevScreenPositionScaleBias = FVector4(
+		 ViewportExtent.X * 0.5f / BufferSize.X,
+		-ViewportExtent.Y * 0.5f / BufferSize.Y,
+		(ViewportExtent.X * 0.5f + ViewportOffset.X) / BufferSize.X,
+		(ViewportExtent.Y * 0.5f + ViewportOffset.Y) / BufferSize.Y );
+
+	PassParameters->PrevSceneColorPreExposureCorrection = View.PreExposure / TemporalAAHistory.SceneColorPreExposure;
+
+	PassParameters->View = View.ViewUniformBuffer;
+	PassParameters->SceneTextures = CreateSceneTextureUniformBufferSingleDraw( RHICmdList, ESceneTextureSetupMode::SceneDepth | ESceneTextureSetupMode::GBuffers, View.FeatureLevel );
+
+	GraphBuilder.AddPass(
+		RDG_EVENT_NAME( "ScreenSpaceDiffuseIndirect %dx%d", View.ViewRect.Width(), View.ViewRect.Height() ),
+		PassParameters,
+		ERenderGraphPassFlags::None,
+		[ PassParameters, &View, Quality ]( FRHICommandListImmediate& RHICmdList )
+		{
+			SCOPED_GPU_STAT(RHICmdList, ScreenSpaceDiffuseIndirect);
+
+			RHICmdList.SetViewport( View.ViewRect.Min.X, View.ViewRect.Min.Y, 0.0f, View.ViewRect.Max.X, View.ViewRect.Max.Y, 1.0f );
+
+			FScreenSpaceDiffuseIndirectPS::FPermutationDomain PermutationVector;
+			PermutationVector.Set< FScreenSpaceDiffuseIndirectPS::FQualityDim >( Quality );
+
+			auto VertexShader	= View.ShaderMap->GetShader< FPostProcessVS >();
+			auto PixelShader	= View.ShaderMap->GetShader< FScreenSpaceDiffuseIndirectPS >( PermutationVector );
+
+			{
+				FGraphicsPipelineStateInitializer GraphicsPSOInit;
+				RHICmdList.ApplyCachedRenderTargets( GraphicsPSOInit );
+
+				GraphicsPSOInit.BlendState = TStaticBlendState< CW_RGB, BO_Add, BF_One, BF_SourceAlpha >::GetRHI();
+				GraphicsPSOInit.DepthStencilState = TStaticDepthStencilState<false, CF_Always>::GetRHI();
+				GraphicsPSOInit.RasterizerState = TStaticRasterizerState<>::GetRHI();
+				GraphicsPSOInit.PrimitiveType = PT_TriangleList;
+				GraphicsPSOInit.BoundShaderState.VertexDeclarationRHI = GFilterVertexDeclaration.VertexDeclarationRHI;
+				GraphicsPSOInit.BoundShaderState.VertexShaderRHI = GETSAFERHISHADER_VERTEX( VertexShader );
+				GraphicsPSOInit.BoundShaderState.PixelShaderRHI = GETSAFERHISHADER_PIXEL( PixelShader );
+
+				SetGraphicsPipelineState( RHICmdList, GraphicsPSOInit );
+			}
+
+			SetShaderParameters( RHICmdList, PixelShader, PixelShader->GetPixelShader(), *PassParameters );
+
+			DrawPostProcessPass(
+				RHICmdList,
+				0, 0,
+				View.ViewRect.Width(), View.ViewRect.Height(),
+				View.ViewRect.Min.X, View.ViewRect.Min.Y,
+				View.ViewRect.Width(), View.ViewRect.Height(),
+				View.ViewRect.Size(),
+				FSceneRenderTargets::Get( RHICmdList ).GetBufferSizeXY(),
+				VertexShader,
+				View.StereoPass,
+				false,
+				EDRF_UseTriangleOptimization);
+		});
+
+	GraphBuilder.Execute();
+}
