DR - Make hlslcc happy
#rnx

diff --git a/Engine/Shaders/Private/ScreenSpaceReflections.usf b/Engine/Shaders/Private/ScreenSpaceReflections.usf
--- a/Engine/Shaders/Private/ScreenSpaceReflections.usf
+++ b/Engine/Shaders/Private/ScreenSpaceReflections.usf
@@ -139,7 +139,7 @@ void ScreenSpaceReflectionsPS(in noperspective float4 UVAndScreenPos : TEXCOORD0
 
 		if( Roughness < 0.1 )
 		{
-			NumSteps = min( NumSteps * NumRays, 24 );
+			NumSteps = min( NumSteps * NumRays, 24u );
 			NumRays = 1;
 		}
 
