- Niagara particles should take into account the PerViewTranslation.
- Check for 0 vector when normalizing velocity vector via hlsl normalize()

diff --git a/Engine/Shaders/Private/NiagaraMeshVertexFactory.ush b/Engine/Shaders/Private/NiagaraMeshVertexFactory.ush
--- a/Engine/Shaders/Private/NiagaraMeshVertexFactory.ush
+++ b/Engine/Shaders/Private/NiagaraMeshVertexFactory.ush
@@ -386,9 +386,16 @@ float4 CalcPrevWorldPosition(FVertexFactoryInput Input, FVertexFactoryIntermedia
 #endif
 	// if we don't have a full previous transform, we can at least calculate the position to get some motion blur
 	{
+		float3 WorldPosition = CalcWorldPosition(Input, Intermediates).xyz;
+
 		// ParticleVelocity.xyz - normalized direction vector
 		// ParticleVelocity.w - velocity magnitude
-		return CalcWorldPosition(Input, Intermediates) - float4(Intermediates.ParticleVelocity.xyz * (Intermediates.ParticleVelocity.w * NiagaraMeshVF.DeltaSeconds), 0.0);
+		float3 WorldDisplacement = -(Intermediates.ParticleVelocity.xyz * (Intermediates.ParticleVelocity.w * NiagaraMeshVF.DeltaSeconds));
+
+		float3 PrevWorldPosition = WorldPosition - ResolvedView.PreViewTranslation;
+		PrevWorldPosition += WorldDisplacement + ResolvedView.PrevPreViewTranslation;
+
+		return float4(PrevWorldPosition.xyz, Input.Position.w);
 	}
 }
 
@@ -549,8 +556,10 @@ FVertexFactoryIntermediates GetVertexFactoryIntermediates(FVertexFactoryInput In
 #endif
 
 #if USE_PARTICLE_VELOCITY
-	Intermediates.ParticleVelocity.xyz = normalize(ParticleVelocity);
-	Intermediates.ParticleVelocity.w = length(ParticleVelocity);
+	float ParticleVelocityMagnitude = length(ParticleVelocity);
+
+	Intermediates.ParticleVelocity.xyz = (ParticleVelocityMagnitude > 0) ? normalize(ParticleVelocity) : float3(0.0f, 0.0f, 0.0f);
+	Intermediates.ParticleVelocity.w = ParticleVelocityMagnitude;
 #endif
 
 	float4 Transform1 = float4(1.0f, 0.0f, 0.0f, Intermediates.ParticlePosition.x);
