Fixed shader compilation error after wrong mergedown 5208164 over 4863219.

diff --git a/Engine/Shaders/Private/RayTracing/RayTracingOcclusionRGS.usf b/Engine/Shaders/Private/RayTracing/RayTracingOcclusionRGS.usf
--- a/Engine/Shaders/Private/RayTracing/RayTracingOcclusionRGS.usf
+++ b/Engine/Shaders/Private/RayTracing/RayTracingOcclusionRGS.usf
@@ -112,7 +112,7 @@ bool GenerateOcclusionRay(
 		float LightRadius = LightParameters.SourceRadius;
 
 		GenerateDirectionalLightOcclusionRay(
-			WorldPosition, WorldNormal, NormalBias,
+			WorldPosition, WorldNormal,
 			LightDirection, LightRadius,
 			RandSample,
 			/* out */ RayOrigin,
