Fixed overbright specular with decals in forward.
Clamp of roughness helps but only normal normalisation helps to match deferred (GBuffer normal is normalised in DecodeGBufferData)

diff --git a/Engine/Shaders/Private/DBufferDecalShared.ush b/Engine/Shaders/Private/DBufferDecalShared.ush
--- a/Engine/Shaders/Private/DBufferDecalShared.ush
+++ b/Engine/Shaders/Private/DBufferDecalShared.ush
@@ -125,8 +125,11 @@ void ApplyDBufferData(
 	FDBufferData DBufferData, inout float3 WorldNormal, inout float3 SubsurfaceColor, inout float Roughness, 
 	inout float3 BaseColor, inout float Metallic, inout float Specular )
 {
-	WorldNormal = WorldNormal * DBufferData.NormalOpacity + DBufferData.PreMulWorldNormal;
-	Roughness = Roughness * DBufferData.RoughnessOpacity + DBufferData.PreMulRoughness;
+	// We normalise the normal to get smoother visual result (it helps to avoid having D_GGX explodes toward infinity, and matches DecodeGBufferData)
+	WorldNormal = normalize(WorldNormal * DBufferData.NormalOpacity + DBufferData.PreMulWorldNormal);
+	// We clamp roughness to a minimum value to avoid (also helps with D_GGX)
+	Roughness = max(0.04, Roughness * DBufferData.RoughnessOpacity + DBufferData.PreMulRoughness);
+
 	Metallic = Metallic * DBufferData.RoughnessOpacity + DBufferData.PreMulMetallic;
 	Specular = Specular * DBufferData.RoughnessOpacity + DBufferData.PreMulSpecular;
 
