Fix shader compilation issue (ambiguous function call) in Draw Text.

diff --git a/Engine/Shaders/Private/ShaderPrintCommon.ush b/Engine/Shaders/Private/ShaderPrintCommon.ush
--- a/Engine/Shaders/Private/ShaderPrintCommon.ush
+++ b/Engine/Shaders/Private/ShaderPrintCommon.ush
@@ -170,14 +170,17 @@ float2 ShaderPrint(in float2 ScreenPos, in bool Value)
 {
 	if (ShaderPrintFilterEnable)
 	{
-		ScreenPos = ShaderPrint(ScreenPos, Value ? 1U : 0U);
+		ScreenPos = ShaderPrint(ScreenPos, Value ? 1u : 0u, SHADER_PRINT_TYPE_UINT);
 	}
 	return ScreenPos;
 }
 
 void ShaderPrint(in bool Value)
 {
-	ShaderPrintCursorPos = ShaderPrint(ShaderPrintCursorPos, Value);
+	if (ShaderPrintFilterEnable)
+	{
+		ShaderPrintCursorPos = ShaderPrint(ShaderPrintCursorPos, Value ? 1u : 0u, SHADER_PRINT_TYPE_UINT);
+	}
 }
 
 float2 ShaderPrint(in float2 ScreenPos, in float2 Value)
