Add position and normal stream support to (certain) vertex factory.

diff --git a/Engine/Shaders/Private/GeometryCacheVertexFactory.ush b/Engine/Shaders/Private/GeometryCacheVertexFactory.ush
--- a/Engine/Shaders/Private/GeometryCacheVertexFactory.ush
+++ b/Engine/Shaders/Private/GeometryCacheVertexFactory.ush
@@ -95,12 +95,30 @@ struct FPositionOnlyVertexFactoryInput
 	float4	Position	: ATTRIBUTE0;
 };
 
+struct FPositionAndNormalOnlyVertexFactoryInput
+{
+	float4	Position	: ATTRIBUTE0;
+	float4	Normal		: ATTRIBUTE1;
+};
+
 /** for depth-only pass */
 float4 VertexFactoryGetWorldPosition(FPositionOnlyVertexFactoryInput Input)
 {
 	float4 Position = Input.Position;
-	
-	return TransformLocalToTranslatedWorld(Position);
+	return TransformLocalToTranslatedWorld(Position.xyz);
+}
+
+/** for shadow pass (used for slope bias) */
+float4 VertexFactoryGetWorldPosition(FPositionAndNormalOnlyVertexFactoryInput Input)
+{
+	float4 Position = Input.Position;
+	return TransformLocalToTranslatedWorld(Position.xyz);
+}
+
+float3 VertexFactoryGetWorldNormal(FPositionAndNormalOnlyVertexFactoryInput Input)
+{
+	float4 Normal = Input.Normal;
+	return RotateLocalToWorld(Normal.xyz);
 }
 
 struct FVertexFactoryInterpolantsVSToPS
@@ -247,6 +265,11 @@ FVertexFactoryIntermediates GetVertexFactoryIntermediates(FVertexFactoryInput In
 
 	// Fill TangentToLocal
 	Intermediates.TangentToLocal = CalcTangentToLocal(Input);
+	
+	float3 TangentToWorld0;
+	float4 TangentToWorld2;
+	CalcTangentToWorld(Input, Intermediates, TangentToWorld0, TangentToWorld2);
+	Intermediates.TangentToWorld = AssembleTangentToWorld(TangentToWorld0, TangentToWorld2);
 
 	// Swizzle vertex color.
 	Intermediates.Color = Input.Color FCOLOR_COMPONENT_SWIZZLE;
diff --git a/Engine/Shaders/Private/GpuSkinVertexFactory.ush b/Engine/Shaders/Private/GpuSkinVertexFactory.ush
--- a/Engine/Shaders/Private/GpuSkinVertexFactory.ush
+++ b/Engine/Shaders/Private/GpuSkinVertexFactory.ush
@@ -587,6 +587,15 @@ void CalcTangentToWorld(FVertexFactoryInput Input, FVertexFactoryIntermediates I
 	TangentToWorld2 = float4(TangentToWorld[2], Input.TangentZ.w * Primitive.InvNonUniformScaleAndDeterminantSign.w);
 }
 
+float3 VertexFactoryGetWorldNormal(FVertexFactoryInput Input, FVertexFactoryIntermediates Intermediates)
+{
+	float3 TangentToWorld0; 
+	float4 TangentToWorld2;
+	CalcTangentToWorld(Input, Intermediates, TangentToWorld0, TangentToWorld2);
+
+	return TangentToWorld2.xyz;
+}
+
 FVertexFactoryInterpolantsVSToPS VertexFactoryGetInterpolantsVSToPS(FVertexFactoryInput Input, FVertexFactoryIntermediates Intermediates, FMaterialVertexParameters VertexParameters)
 {
 	FVertexFactoryInterpolantsVSToPS Interpolants;
diff --git a/Engine/Shaders/Private/LandscapeVertexFactory.ush b/Engine/Shaders/Private/LandscapeVertexFactory.ush
--- a/Engine/Shaders/Private/LandscapeVertexFactory.ush
+++ b/Engine/Shaders/Private/LandscapeVertexFactory.ush
@@ -183,6 +183,11 @@ float4 VertexFactoryGetWorldPosition(FVertexFactoryInput Input, FVertexFactoryIn
 	return INVARIANT(TransformLocalToTranslatedWorld(GetLocalPosition(Intermediates)));
 }
 
+float3 VertexFactoryGetWorldNormal(FVertexFactoryIntermediates Intermediates)
+{
+	return Intermediates.WorldNormal;
+}
+
 float4 VertexFactoryGetPreviousWorldPosition(FVertexFactoryInput Input, FVertexFactoryIntermediates Intermediates)
 {
 	// Note we still use LocalToWorld. Landscape's transform never changes.
diff --git a/Engine/Shaders/Private/LocalVertexFactory.ush b/Engine/Shaders/Private/LocalVertexFactory.ush
--- a/Engine/Shaders/Private/LocalVertexFactory.ush
+++ b/Engine/Shaders/Private/LocalVertexFactory.ush
@@ -221,6 +221,35 @@ struct FPositionOnlyVertexFactoryInput
 #endif
 };
 
+/** 
+ * Per-vertex inputs from bound vertex buffers.  Used by passes with a trimmed down position-and-normal-only shader.
+ */
+struct FPositionAndNormalOnlyVertexFactoryInput
+{
+	float4	Position	: ATTRIBUTE0;
+	float4	Normal		: ATTRIBUTE2;
+
+#if USE_INSTANCING && !USE_INSTANCING_EMULATED && !MANUAL_VERTEX_FETCH
+	float4 InstanceOrigin : ATTRIBUTE8;  // per-instance random in w 
+	half4 InstanceTransform1 : ATTRIBUTE9;  // hitproxy.r + 256 * selected in .w
+	half4 InstanceTransform2 : ATTRIBUTE10; // hitproxy.g in .w
+	half4 InstanceTransform3 : ATTRIBUTE11; // hitproxy.b in .w
+#endif	// USE_INSTANCING && !USE_INSTANCING_EMULATED
+
+#if VF_USE_PRIMITIVE_SCENE_DATA
+	uint PrimitiveId : ATTRIBUTE1;
+#endif
+
+/** Optional instance ID for vertex layered rendering */
+#if (FEATURE_LEVEL >= FEATURE_LEVEL_ES3_1) && (MANUAL_VERTEX_FETCH && (USE_INSTANCING && !USE_INSTANCING_EMULATED))
+	uint InstanceId	: SV_InstanceID;
+#endif
+
+#if MANUAL_VERTEX_FETCH
+	uint VertexId : SV_VertexID;
+#endif
+};
+
 /** 
  * Caches intermediates that would otherwise have to be computed multiple times.  Avoids relying on the compiler to optimize out redundant operations.
  */
@@ -287,6 +316,30 @@ float4x4 GetInstanceTransform(FPositionOnlyVertexFactoryInput Input)
 	#endif	// !USE_INSTANCING_EMULATED
 }
 
+float4x4 GetInstanceTransform(FPositionAndNormalOnlyVertexFactoryInput Input)
+{
+	#if !USE_INSTANCING_EMULATED && MANUAL_VERTEX_FETCH
+	uint InstanceId = GetInstanceId(Input.InstanceId);
+	return float4x4(
+		float4(VertexFetch_InstanceTransformBuffer[3 * (InstanceId + InstanceOffset) + 0].xyz, 0.0f),
+		float4(VertexFetch_InstanceTransformBuffer[3 * (InstanceId + InstanceOffset) + 1].xyz, 0.0f),
+		float4(VertexFetch_InstanceTransformBuffer[3 * (InstanceId + InstanceOffset) + 2].xyz, 0.0f),
+		float4(VertexFetch_InstanceOriginBuffer[(InstanceId + InstanceOffset)].xyz, 1.0f));
+	#elif !USE_INSTANCING_EMULATED
+	return float4x4(
+		float4(Input.InstanceTransform1.xyz, 0.0f),
+		float4(Input.InstanceTransform2.xyz, 0.0f),
+		float4(Input.InstanceTransform3.xyz, 0.0f),
+		float4(Input.InstanceOrigin.xyz, 1.0f));
+	#else
+	return float4x4(
+		float4(CPUInstanceTransform[0].xyz, 0.0f), 
+		float4(CPUInstanceTransform[1].xyz, 0.0f), 
+		float4(CPUInstanceTransform[2].xyz, 0.0f), 
+		float4(CPUInstanceOrigin.xyz, 1.0f));
+	#endif	// !USE_INSTANCING_EMULATED
+}
+
 half3x3 GetInstanceToLocal3x3(FVertexFactoryIntermediates Intermediates)
 {
 	#if !USE_INSTANCING_EMULATED
@@ -853,6 +906,48 @@ float4 VertexFactoryGetWorldPosition(FPositionOnlyVertexFactoryInput Input)
 #endif	// USE_INSTANCING
 }
 
+/** for depth-only pass (slope depth bias) */
+float4 VertexFactoryGetWorldPosition(FPositionAndNormalOnlyVertexFactoryInput Input)
+{
+	float4 Position = Input.Position;
+	
+#if VF_USE_PRIMITIVE_SCENE_DATA
+	uint PrimitiveId = Input.PrimitiveId;
+#else
+	uint PrimitiveId = 0;
+#endif
+
+#if USE_INSTANCING
+	return CalcWorldPosition(Position, GetInstanceTransform(Input), PrimitiveId);
+#else
+	return CalcWorldPosition(Position, PrimitiveId);
+#endif	// USE_INSTANCING
+}
+
+float3 VertexFactoryGetWorldNormal(FPositionAndNormalOnlyVertexFactoryInput Input)
+{
+	float3 Normal = Input.Normal.xyz;
+	
+#if VF_USE_PRIMITIVE_SCENE_DATA
+	uint PrimitiveId = Input.PrimitiveId;
+#else
+	uint PrimitiveId = 0;
+#endif
+
+#if USE_INSTANCING
+	const float3 InstanceTransformedNormal = mul(float4(Normal,0), GetInstanceTransform(Input)).xyz;
+	return RotateLocalToWorld(InstanceTransformedNormal, PrimitiveId);
+#else
+	return RotateLocalToWorld(Normal, PrimitiveId);
+#endif	// USE_INSTANCING
+}
+
+
+float3 VertexFactoryGetWorldNormal(FVertexFactoryIntermediates Intermediates)
+{
+	return Intermediates.TangentToWorld[2];
+}
+
 // Information when using pass through to get the previous position, as no matrix is available/needed
 #if GPUSKIN_PASS_THROUGH
 Buffer<float> GPUSkinCachePreviousPositionBuffer;
diff --git a/Engine/Shaders/Private/VertexFactoryCommon.ush b/Engine/Shaders/Private/VertexFactoryCommon.ush
--- a/Engine/Shaders/Private/VertexFactoryCommon.ush
+++ b/Engine/Shaders/Private/VertexFactoryCommon.ush
@@ -24,6 +24,17 @@ float4 TransformLocalToTranslatedWorld(float3 LocalPosition)
 	return float4(RotatedPosition + (Primitive.LocalToWorld[3].xyz + ResolvedView.PreViewTranslation.xyz),1);
 }
 
+float3 RotateLocalToWorld(float3 LocalDirection, uint PrimitiveId)
+{
+	float4x4 LocalToWorld = GetPrimitiveData(PrimitiveId).LocalToWorld;
+	return LocalToWorld[0].xyz * LocalDirection.xxx + LocalToWorld[1].xyz * LocalDirection.yyy + LocalToWorld[2].xyz * LocalDirection.zzz;	
+}
+
+float3 RotateLocalToWorld(float3 LocalDirection)
+{
+	return Primitive.LocalToWorld[0].xyz * LocalDirection.xxx + Primitive.LocalToWorld[1].xyz * LocalDirection.yyy + Primitive.LocalToWorld[2].xyz * LocalDirection.zzz;
+}
+
 // Octahedron Normal Vectors
 // [Cigolle 2014, "A Survey of Efficient Representations for Independent Unit Vectors"]
 //						Mean	Max
diff --git a/Engine/Source/Runtime/Engine/Private/GeometryCacheVertexFactory.cpp b/Engine/Source/Runtime/Engine/Private/GeometryCacheVertexFactory.cpp
--- a/Engine/Source/Runtime/Engine/Private/GeometryCacheVertexFactory.cpp
+++ b/Engine/Source/Runtime/Engine/Private/GeometryCacheVertexFactory.cpp
@@ -191,6 +191,13 @@ void FGeometryCacheVertexVertexFactory::InitRHI()
 			PositionOnlyStreamElements.Add(AccessStreamComponent(Data.PositionComponent, 0, EVertexInputStreamType::PositionOnly));
 			InitDeclaration(PositionOnlyStreamElements, EVertexInputStreamType::PositionOnly);
 		}
+
+		{
+			FVertexDeclarationElementList PositionAndNormalOnlyStreamElements;
+			PositionAndNormalOnlyStreamElements.Add(AccessStreamComponent(Data.PositionComponent, 0, EVertexInputStreamType::PositionAndNormalOnly));
+			PositionAndNormalOnlyStreamElements.Add(AccessStreamComponent(Data.TangentBasisComponents[1], 1, EVertexInputStreamType::PositionAndNormalOnly));
+			InitDeclaration(PositionAndNormalOnlyStreamElements, EVertexInputStreamType::PositionAndNormalOnly);
+		}
 	}
 
 	FVertexDeclarationElementList Elements;
diff --git a/Engine/Source/Runtime/Engine/Private/InstancedStaticMesh.cpp b/Engine/Source/Runtime/Engine/Private/InstancedStaticMesh.cpp
--- a/Engine/Source/Runtime/Engine/Private/InstancedStaticMesh.cpp
+++ b/Engine/Source/Runtime/Engine/Private/InstancedStaticMesh.cpp
@@ -482,7 +482,7 @@ void FInstancedStaticMeshVertexFactory::InitRHI()
 	check(HasValidFeatureLevel());
 	const bool bInstanced = GRHISupportsInstancing;
 
-#if !ALLOW_DITHERED_LOD_FOR_INSTANCED_STATIC_MESHES // position only shaders cannot work with dithered LOD
+#if !ALLOW_DITHERED_LOD_FOR_INSTANCED_STATIC_MESHES // position(and normal) only shaders cannot work with dithered LOD
 	// If the vertex buffer containing position is not the same vertex buffer containing the rest of the data,
 	// then initialize PositionStream and PositionDeclaration.
 	if(Data.PositionComponent.VertexBuffer != Data.TangentBasisComponents[0].VertexBuffer)
@@ -507,6 +507,7 @@ void FInstancedStaticMeshVertexFactory::InitRHI()
 			InitDeclaration(StreamElements, InputStreamType);
 		};
 		AddDeclaration(EVertexInputStreamType::PositionOnly, bInstanced, false);
+		AddDeclaration(EVertexInputStreamType::PositionAndNormalOnly, bInstanced, true);
 	}
 #endif
 
diff --git a/Engine/Source/Runtime/Engine/Private/LocalVertexFactory.cpp b/Engine/Source/Runtime/Engine/Private/LocalVertexFactory.cpp
--- a/Engine/Source/Runtime/Engine/Private/LocalVertexFactory.cpp
+++ b/Engine/Source/Runtime/Engine/Private/LocalVertexFactory.cpp
@@ -282,6 +282,7 @@ void FLocalVertexFactory::InitRHI()
 			InitDeclaration(StreamElements, InputStreamType);
 		};
 		AddDeclaration(EVertexInputStreamType::PositionOnly, false);
+		AddDeclaration(EVertexInputStreamType::PositionAndNormalOnly, true);
 	}
 
 	FVertexDeclarationElementList Elements;
diff --git a/Engine/Source/Runtime/RenderCore/Private/VertexFactory.cpp b/Engine/Source/Runtime/RenderCore/Private/VertexFactory.cpp
--- a/Engine/Source/Runtime/RenderCore/Private/VertexFactory.cpp
+++ b/Engine/Source/Runtime/RenderCore/Private/VertexFactory.cpp
@@ -221,6 +221,16 @@ void FVertexFactory::GetStreams(ERHIFeatureLevel::Type InFeatureLevel, EVertexIn
 			OutVertexStreams.Add(FVertexInputStream(StreamIndex, Stream.Offset, Stream.VertexBuffer->VertexBufferRHI));
 		}
 	}
+	else if (VertexStreamType == EVertexInputStreamType::PositionAndNormalOnly)
+	{
+		// Set the predefined vertex streams.
+		for (int32 StreamIndex = 0; StreamIndex < PositionAndNormalStream.Num(); StreamIndex++)
+		{
+			const FVertexStream& Stream = PositionAndNormalStream[StreamIndex];
+			check(Stream.VertexBuffer->IsInitialized());
+			OutVertexStreams.Add(FVertexInputStream(StreamIndex, Stream.Offset, Stream.VertexBuffer->VertexBufferRHI));
+		}
+	}
 	else
 	{
 		// NOT_IMPLEMENTED
@@ -231,7 +241,8 @@ void FVertexFactory::OffsetInstanceStreams(uint32 InstanceOffset, EVertexInputSt
 {
 	const TArrayView<const FVertexStream>& StreamArray = 
 		  VertexStreamType == EVertexInputStreamType::PositionOnly ?			MakeArrayView(PositionStream) : 
-		/*VertexStreamType == EVertexInputStreamType::Default*/					MakeArrayView(Streams);
+		( VertexStreamType == EVertexInputStreamType::PositionAndNormalOnly ?	MakeArrayView(PositionAndNormalStream) : 
+		/*VertexStreamType == EVertexInputStreamType::Default*/					MakeArrayView(Streams));
 
 	for (int32 StreamIndex = 0; StreamIndex < StreamArray.Num(); StreamIndex++)
 	{
@@ -254,8 +265,10 @@ void FVertexFactory::ReleaseRHI()
 {
 	Declaration.SafeRelease();
 	PositionDeclaration.SafeRelease();
+	PositionAndNormalDeclaration.SafeRelease();
 	Streams.Empty();
 	PositionStream.Empty();
+	PositionAndNormalStream.Empty();
 }
 
 FVertexElement FVertexFactory::AccessStreamComponent(const FVertexStreamComponent& Component, uint8 AttributeIndex)
@@ -279,6 +292,8 @@ FVertexElement FVertexFactory::AccessStreamComponent(const FVertexStreamComponen
 
 	if (InputStreamType == EVertexInputStreamType::PositionOnly)
 		return FVertexElement(PositionStream.AddUnique(VertexStream), Component.Offset, Component.Type, AttributeIndex, VertexStream.Stride, EnumHasAnyFlags(EVertexStreamUsage::Instancing, VertexStream.VertexStreamUsage));
+	else if (InputStreamType == EVertexInputStreamType::PositionAndNormalOnly)
+		return FVertexElement(PositionAndNormalStream.AddUnique(VertexStream), Component.Offset, Component.Type, AttributeIndex, VertexStream.Stride, EnumHasAnyFlags(EVertexStreamUsage::Instancing, VertexStream.VertexStreamUsage));
 	else /* (InputStreamType == EVertexInputStreamType::Default) */
 		return FVertexElement(Streams.AddUnique(VertexStream), Component.Offset, Component.Type, AttributeIndex, VertexStream.Stride, EnumHasAnyFlags(EVertexStreamUsage::Instancing, VertexStream.VertexStreamUsage));
 }
@@ -290,6 +305,10 @@ void FVertexFactory::InitDeclaration(const FVertexDeclarationElementList& Elemen
 	{
 		PositionDeclaration = PipelineStateCache::GetOrCreateVertexDeclaration(Elements);
 	}
+	else if (StreamType == EVertexInputStreamType::PositionAndNormalOnly)
+	{
+		PositionAndNormalDeclaration = PipelineStateCache::GetOrCreateVertexDeclaration(Elements);
+	}
 	else // (StreamType == EVertexInputStreamType::Default)
 	{
 		// Create the vertex declaration for rendering the factory normally.
diff --git a/Engine/Source/Runtime/RenderCore/Public/VertexFactory.h b/Engine/Source/Runtime/RenderCore/Public/VertexFactory.h
--- a/Engine/Source/Runtime/RenderCore/Public/VertexFactory.h
+++ b/Engine/Source/Runtime/RenderCore/Public/VertexFactory.h
@@ -81,7 +81,8 @@ enum class EVertexStreamUsage : uint8
 enum class EVertexInputStreamType : uint8
 {
 	Default = 0,
-	PositionOnly
+	PositionOnly,
+	PositionAndNormalOnly
 };
 
 ENUM_CLASS_FLAGS(EVertexStreamUsage);
@@ -518,6 +519,7 @@ public:
 		{
 		case EVertexInputStreamType::Default:				return Declaration;
 		case EVertexInputStreamType::PositionOnly:			return PositionDeclaration;
+		case EVertexInputStreamType::PositionAndNormalOnly:	return PositionAndNormalDeclaration;
 		}
 		return Declaration;
 	}
@@ -527,6 +529,9 @@ public:
 	/** Indicates whether the vertex factory supports a position-only stream. */
 	virtual bool SupportsPositionOnlyStream() const { return !!PositionStream.Num(); }
 
+	/** Indicates whether the vertex factory supports a position-and-normal-only stream. */
+	virtual bool SupportsPositionAndNormalOnlyStream() const { return !!PositionAndNormalStream.Num(); }
+
 	/** Indicates whether the vertex factory supports a null pixel shader. */
 	virtual bool SupportsNullPixelShader() const { return true; }
 
@@ -604,18 +609,20 @@ protected:
 	
 	bool bSupportsManualVertexFetch = false;
 
-	int8 PrimitiveIdStreamIndex[3] = { -1, -1 }; // Need to match entry count of EVertexInputStreamType
+	int8 PrimitiveIdStreamIndex[3] = { -1, -1, -1 }; // Need to match entry count of EVertexInputStreamType
 
 private:
 
 	/** The position only vertex stream used to render the factory during depth only passes. */
 	TArray<FVertexStream,TInlineAllocator<2> > PositionStream;
+	TArray<FVertexStream, TInlineAllocator<3> > PositionAndNormalStream;
 
 	/** The RHI vertex declaration used to render the factory normally. */
 	FVertexDeclarationRHIRef Declaration;
 
 	/** The RHI vertex declaration used to render the factory during depth only passes. */
 	FVertexDeclarationRHIRef PositionDeclaration;
+	FVertexDeclarationRHIRef PositionAndNormalDeclaration;
 };
 
 /**
diff --git a/Engine/Source/Runtime/Renderer/Public/MeshPassProcessor.h b/Engine/Source/Runtime/Renderer/Public/MeshPassProcessor.h
--- a/Engine/Source/Runtime/Renderer/Public/MeshPassProcessor.h
+++ b/Engine/Source/Runtime/Renderer/Public/MeshPassProcessor.h
@@ -863,7 +863,8 @@ struct TMeshProcessorShaders
 enum class EMeshPassFeatures
 {
 	Default = 0,
-	PositionOnly = 1 << 0
+	PositionOnly = 1 << 0,
+	PositionAndNormalOnly = 1 << 1
 };
 ENUM_CLASS_FLAGS(EMeshPassFeatures);
 
diff --git a/Engine/Source/Runtime/Renderer/Public/MeshPassProcessor.inl b/Engine/Source/Runtime/Renderer/Public/MeshPassProcessor.inl
--- a/Engine/Source/Runtime/Renderer/Public/MeshPassProcessor.inl
+++ b/Engine/Source/Runtime/Renderer/Public/MeshPassProcessor.inl
@@ -35,6 +35,7 @@ void FMeshPassProcessor::BuildMeshDrawCommands(
 
 	EVertexInputStreamType InputStreamType = EVertexInputStreamType::Default;
 	if ((MeshPassFeatures & EMeshPassFeatures::PositionOnly) != EMeshPassFeatures::Default)				InputStreamType = EVertexInputStreamType::PositionOnly;
+	if ((MeshPassFeatures & EMeshPassFeatures::PositionAndNormalOnly) != EMeshPassFeatures::Default)	InputStreamType = EVertexInputStreamType::PositionAndNormalOnly;
 
 	check(VertexFactory && VertexFactory->IsInitialized());
 	FVertexDeclarationRHIParamRef VertexDeclaration = VertexFactory->GetDeclaration(InputStreamType);
