Preprocess out access to UAVs that aren't always used in the DOFReduce shader permutations.

Only necessary because SPIRV-Tools can't unroll loops to perform branch simplification and dead-code elimination so it retains the unused UAVs and the code that references them. This means the unreachable UAVs remain in the shader parameter map. Fixing SPIRV-Tools to handle this will take a lot longer.

diff --git a/Engine/Shaders/Private/DiaphragmDOF/DOFReduce.usf b/Engine/Shaders/Private/DiaphragmDOF/DOFReduce.usf
--- a/Engine/Shaders/Private/DiaphragmDOF/DOFReduce.usf
+++ b/Engine/Shaders/Private/DiaphragmDOF/DOFReduce.usf
@@ -219,22 +219,30 @@ void OutputMipLevel(const uint MipLevel, uint2 MipPixelPos, float4 Color, float
 		{
 			Output0Mip0[MipPixelPos] = OutputValue0;
 		}
+#if DIM_REDUCE_MIP_COUNT > 1
 		else if (MipLevel == 1)
 		{
 			Output0Mip1[MipPixelPos] = OutputValue0;
 		}
+#if DIM_REDUCE_MIP_COUNT > 2
 		else if (MipLevel == 2)
 		{
 			Output0Mip2[MipPixelPos] = OutputValue0;
 		}
+#if DIM_REDUCE_MIP_COUNT > 3
 		else if (MipLevel == 3)
 		{
 			Output0Mip3[MipPixelPos] = OutputValue0;
 		}
+#if DIM_REDUCE_MIP_COUNT > 4
 		else if (MipLevel == 4)
 		{
 			Output0Mip4[MipPixelPos] = OutputValue0;
 		}
+#endif
+#endif
+#endif
+#endif
 	}
 	#elif CONFIG_GATHER_INPUT_LAYOUT == GATHER_INPUT_LAYOUT_RGB_ALPHA_COC || CONFIG_GATHER_INPUT_LAYOUT == GATHER_INPUT_LAYOUT_RGB_SEPARATE_COC
 	{
@@ -243,26 +251,34 @@ void OutputMipLevel(const uint MipLevel, uint2 MipPixelPos, float4 Color, float
 			Output0Mip0[MipPixelPos] = OutputValue0;
 			Output1Mip0[MipPixelPos] = OutputValue1;
 		}
+#if DIM_REDUCE_MIP_COUNT > 1
 		else if (MipLevel == 1)
 		{
 			Output0Mip1[MipPixelPos] = OutputValue0;
 			Output1Mip1[MipPixelPos] = OutputValue1;
 		}
+#if DIM_REDUCE_MIP_COUNT > 2
 		else if (MipLevel == 2)
 		{
 			Output0Mip2[MipPixelPos] = OutputValue0;
 			Output1Mip2[MipPixelPos] = OutputValue1;
 		}
+#if DIM_REDUCE_MIP_COUNT > 3
 		else if (MipLevel == 3)
 		{
 			Output0Mip3[MipPixelPos] = OutputValue0;
 			Output1Mip3[MipPixelPos] = OutputValue1;
 		}
+#if DIM_REDUCE_MIP_COUNT > 4
 		else if (MipLevel == 4)
 		{
 			Output0Mip4[MipPixelPos] = OutputValue0;
 			Output1Mip4[MipPixelPos] = OutputValue1;
 		}
+#endif
+#endif
+#endif
+#endif
 	}
 	#else
 		#error Unknown gather input layout.
