Shader workarounds to push alpha to be opaque when it is very close (but not quite).

diff --git a/Engine/Shaders/Private/DiaphragmDOF/DOFRecombine.usf b/Engine/Shaders/Private/DiaphragmDOF/DOFRecombine.usf
--- a/Engine/Shaders/Private/DiaphragmDOF/DOFRecombine.usf
+++ b/Engine/Shaders/Private/DiaphragmDOF/DOFRecombine.usf
@@ -1049,6 +1049,13 @@ void RecombineMainCS(
 	// Convert alpha channel from opacity back to translucency.
 	OutputFinalSceneColor.a = 1 - OutputFinalSceneColor.a;
 
+	// Ensure that alpha values that are expected to be opaque (but are only close to opaque) are forced to be opaque.
+	// (0.005 chosen to accomodate handling of 1/255)
+	if (OutputFinalSceneColor.a < 0.005)
+	{
+		OutputFinalSceneColor.a = 0;
+	}
+	
 	// Debug optimisation colors.
 	#if DEBUG_FAST_PATHS
 	{
diff --git a/Engine/Shaders/Private/PostProcessTemporalAA.usf b/Engine/Shaders/Private/PostProcessTemporalAA.usf
--- a/Engine/Shaders/Private/PostProcessTemporalAA.usf
+++ b/Engine/Shaders/Private/PostProcessTemporalAA.usf
@@ -2213,6 +2213,13 @@ FTAAHistoryPayload TemporalAASample(uint2 GroupId, uint2 GroupThreadId, uint Gro
 		OutputPayload = AddPayload(MulPayload(History, Weights.x), MulPayload(IntermediaryResult.Filtered, Weights.y));
 	}
 
+	// Ensure that alpha values that are expected to be opaque (but are only close to opaque) are forced to be opaque.
+	// (0.995 chosen to accomodate handling of 254/255)
+	if (OutputPayload.Color.a > 0.995)
+	{
+		OutputPayload.Color.a = 1;
+	}
+
 	OutputPayload.Color = TransformBackToRawLinearSceneColor(OutputPayload.Color);
 	
 	#if USE_PREEXPOSURE
