Fix PreExposure for TileDeferred

diff --git a/Engine/Shaders/Private/TiledDeferredLightShaders.usf b/Engine/Shaders/Private/TiledDeferredLightShaders.usf
--- a/Engine/Shaders/Private/TiledDeferredLightShaders.usf
+++ b/Engine/Shaders/Private/TiledDeferredLightShaders.usf
@@ -297,7 +297,11 @@ void TiledDeferredLightingMain(
 			CompositedLighting.rgb += GetSimpleDynamicLighting(WorldPosition, CameraVector, ScreenSpaceData.GBuffer.WorldNormal, ScreenSpaceData.AmbientOcclusion, ScreenSpaceData.GBuffer.DiffuseColor, ScreenSpaceData.GBuffer.SpecularColor, ScreenSpaceData.GBuffer.Roughness, LightData);
 		}
 	}
-		
+
+#if USE_PREEXPOSURE
+	CompositedLighting *= View.PreExposure;
+#endif
+
 #endif
 
 	// Only write to the buffer for threads inside the view
