DR - Workaround hlslcc issue
#rnx

diff --git a/Engine/Shaders/Private/ClusteredDeferredShadingPixelShader.usf b/Engine/Shaders/Private/ClusteredDeferredShadingPixelShader.usf
--- a/Engine/Shaders/Private/ClusteredDeferredShadingPixelShader.usf
+++ b/Engine/Shaders/Private/ClusteredDeferredShadingPixelShader.usf
@@ -34,7 +34,7 @@
  * Adds local lighting using the light grid, does not apply directional lights, as they are done elsewhere.
  * Does not support dynamic shadows, as these require the per-light shadow mask.
  */
-float3 GetLightGridLocalLighting(const FCulledLightsGridData LightGridData, float3 WorldPosition, float3 CameraVector, FScreenSpaceData ScreenSpaceData, float2 ScreenUV, uint EyeIndex, float Dither, uint FirstNonSimpleLightIndex)
+float3 GetLightGridLocalLighting(const FCulledLightsGridData InLightGridData, float3 WorldPosition, float3 CameraVector, FScreenSpaceData ScreenSpaceData, float2 ScreenUV, uint EyeIndex, float Dither, uint FirstNonSimpleLightIndex)
 {
 	float3 DirectLighting = 0;
 
@@ -43,12 +43,12 @@ float3 GetLightGridLocalLighting(const FCulledLightsGridData LightGridData, floa
 
 	// Limit max to ForwardLightData.NumLocalLights.
 	// This prevents GPU hangs when the PS tries to read from uninitialized NumCulledLightsGrid buffer
-	const uint NumLocalLights = min(LightGridData.NumLocalLights, GetNumLocalLights(EyeIndex));
+	const uint NumLocalLights = min(InLightGridData.NumLocalLights, GetNumLocalLights(EyeIndex));
 
 	LOOP
 	for (uint LocalLightListIndex = FirstNonSimpleLightIndex; LocalLightListIndex < NumLocalLights; LocalLightListIndex++)
 	{
-		const FLocalLightData LocalLight = GetLocalLightData(LightGridData.DataStartIndex + LocalLightListIndex, EyeIndex);
+		const FLocalLightData LocalLight = GetLocalLightData(InLightGridData.DataStartIndex + LocalLightListIndex, EyeIndex);
 
 		// The lights are sorted such that all that support clustered deferred are at the beginning, there might be others
 		// (e.g., lights with dynamic shadows) so we break out when the condition fails.
@@ -126,7 +126,7 @@ float3 GetLightGridLocalLighting(const FCulledLightsGridData LightGridData, floa
 	}
 
 	// For debugging
-	//DirectLighting = LightGridData.NumLocalLights / (float)ForwardLightData.MaxCulledLightsPerCell;
+	//DirectLighting = InLightGridData.NumLocalLights / (float)ForwardLightData.MaxCulledLightsPerCell;
 	return DirectLighting;
 }
 
