Fixed "LightGridData" identifier prefix regarding coding convention.

#rnx

diff --git a/Engine/Shaders/Private/LightGridCommon.ush b/Engine/Shaders/Private/LightGridCommon.ush
--- a/Engine/Shaders/Private/LightGridCommon.ush
+++ b/Engine/Shaders/Private/LightGridCommon.ush
@@ -19,16 +19,16 @@
 	#endif
 #endif
 
-struct LightGridData
+struct FLightGridData
 {
 	uint LightGridPixelSizeShift;
 	float3 LightGridZParams;
 	int3 CulledGridSize;
 };
 
-LightGridData GetLightGridData(uint EyeIndex)
+FLightGridData GetLightGridData(uint EyeIndex)
 {
-	LightGridData Result;
+	FLightGridData Result;
 
 #if INSTANCED_STEREO
 	BRANCH
@@ -55,7 +55,7 @@ LightGridData GetLightGridData(uint EyeIndex)
 
 uint ComputeLightGridCellIndex(uint2 PixelPos, float SceneDepth, uint EyeIndex)
 {
-	const LightGridData GridData = GetLightGridData(EyeIndex);
+	const FLightGridData GridData = GetLightGridData(EyeIndex);
 	uint ZSlice = (uint)(max(0, log2(SceneDepth * GridData.LightGridZParams.x + GridData.LightGridZParams.y) * GridData.LightGridZParams.z));
 	ZSlice = min(ZSlice, (uint)(GridData.CulledGridSize.z - 1));
 	uint3 GridCoordinate = uint3(PixelPos >> GridData.LightGridPixelSizeShift, ZSlice);
