Linked most of the remaining dxr_todo-s to JIRAs and removed obsolete ones. No functional changes.

diff --git a/Engine/Shaders/Private/PathTracing/BuildVarianceMipTreeComputeShader.usf b/Engine/Shaders/Private/PathTracing/BuildVarianceMipTreeComputeShader.usf
--- a/Engine/Shaders/Private/PathTracing/BuildVarianceMipTreeComputeShader.usf
+++ b/Engine/Shaders/Private/PathTracing/BuildVarianceMipTreeComputeShader.usf
@@ -146,7 +146,6 @@ void BuildVarianceMipTreeCS(
 		{
 			uint2 Pixel = PixelOffset + uint2(i >> 1, i & 1);
 			uint InputIndex = BufferOffsetAtPixel(Pixel, MipLevel - 1);
-			// #dxr_todo: Experiment between average and max
   			Sum += RWVarianceMipTree[InputIndex];
 		}
 		Sum /= 4;
diff --git a/Engine/Shaders/Private/PathTracing/Light/PathTracingLightCommon.ush b/Engine/Shaders/Private/PathTracing/Light/PathTracingLightCommon.ush
--- a/Engine/Shaders/Private/PathTracing/Light/PathTracingLightCommon.ush
+++ b/Engine/Shaders/Private/PathTracing/Light/PathTracingLightCommon.ush
@@ -121,7 +121,7 @@ float ComputeAttenuationFalloff(float Distance, float Attenuation)
 {
 #ifdef USE_ATTENUATION_TERM
 	// Mirrors GetLocalLightAttenuation() custom attenuation controls
-	// #dxr_todo: encapsulate this function in a shared space
+	// #dxr_todo: UE-72508: encapsulate this function in a shared space
 	return Square(saturate(1.0 - Square(Square(Distance) / Square(Attenuation))));
 #else
 	return 1.0;
diff --git a/Engine/Shaders/Private/PathTracing/Light/PathTracingLightSampling.ush b/Engine/Shaders/Private/PathTracing/Light/PathTracingLightSampling.ush
--- a/Engine/Shaders/Private/PathTracing/Light/PathTracingLightSampling.ush
+++ b/Engine/Shaders/Private/PathTracing/Light/PathTracingLightSampling.ush
@@ -13,7 +13,7 @@ LightSampling.usf: Light sampling for next-event estimation
 #include "PathTracingSkyLight.ush"
 #include "PathTracingSpotLight.ush"
 
-// #dxr_todo: unify with GLightCountMaximum in c++ code
+// #dxr_todo:UE-72532 unify with GLightCountMaximum in c++ code
 static const uint LIGHT_COUNT_MAXIMUM = 64;
 
 void GenerateLightRay(RayDesc Ray, int LightId, float3 LightUV, out RayDesc OutLightRay)
diff --git a/Engine/Shaders/Private/PathTracing/Light/PathTracingRectLight.ush b/Engine/Shaders/Private/PathTracing/Light/PathTracingRectLight.ush
--- a/Engine/Shaders/Private/PathTracing/Light/PathTracingRectLight.ush
+++ b/Engine/Shaders/Private/PathTracing/Light/PathTracingRectLight.ush
@@ -57,7 +57,7 @@ void RectLight_EstimateLight(
 )
 {
 	// Distance to centroid
-	// #dxr_todo: Use closest point, instead
+	// #dxr_todo: UE-72533 Use closest point, instead
 	float3 LightDirection = GetPosition(LightId) - Payload.WorldPos;
 	float LightDistanceSquared = dot(LightDirection, LightDirection);
 	LightDirection = normalize(LightDirection);
diff --git a/Engine/Shaders/Private/PathTracing/Light/PathTracingSkyLight.ush b/Engine/Shaders/Private/PathTracing/Light/PathTracingSkyLight.ush
--- a/Engine/Shaders/Private/PathTracing/Light/PathTracingSkyLight.ush
+++ b/Engine/Shaders/Private/PathTracing/Light/PathTracingSkyLight.ush
@@ -44,7 +44,6 @@ void SkyLight_PdfLight(
 	uint2 MipCoord = TextureUV * SkyLight.MipDimensions.xy;
 	float FacePdf = PdfFace(MipCount, TextureCoord.z);
 
-	// #dxr_todo: StopLevel can be customized for run-time performance
 	uint StopLevel = min(SkyLight.SamplingStopLevel, MipCount);
 	float MipPdf = PdfMipTree(uint3(MipCoord, TextureCoord.z), StopLevel);
 
diff --git a/Engine/Shaders/Private/PathTracing/Light/PathTracingSpotLight.ush b/Engine/Shaders/Private/PathTracing/Light/PathTracingSpotLight.ush
--- a/Engine/Shaders/Private/PathTracing/Light/PathTracingSpotLight.ush
+++ b/Engine/Shaders/Private/PathTracing/Light/PathTracingSpotLight.ush
@@ -143,10 +143,9 @@ void SpotLight_EstimateLight(
 	{
 		float Falloff = NoL - CosConeAngle;
 		Falloff *= Falloff;
-		// #dxr_todo: change falloff to take distance^2
 		Falloff *= ComputeAttenuationFalloff(sqrt(LightDistanceSquared), GetAttenuation(LightId));
 		OutIrradiance = LightPower * Falloff / LightDistanceSquared;
-		// #dxr_todo: apply inner core
+		// #dxr_todo: UE-72534 apply inner core
 	}
 }
 
@@ -171,7 +170,7 @@ void SpotLight_SampleLight(
 		float SinConeAngle = sqrtFast(1 - CosConeAngle * CosConeAngle);
 		OutLightUV = normalize(float3(SinConeAngle * DiskPosition, CosConeAngle));
 		float SphericalCapArea = 2.0 * PI * Radius * Radius * (1.0 - CosConeAngle);
-		// #dxr_todo: apply the correct Jacobian
+		// #dxr_todo: UE-72535 apply the correct Jacobian
 		OutPdf = 1.0 / SphericalCapArea;
 	}
 	else
diff --git a/Engine/Shaders/Private/PathTracing/Material/PathTracingSpecularReflection.ush b/Engine/Shaders/Private/PathTracing/Material/PathTracingSpecularReflection.ush
--- a/Engine/Shaders/Private/PathTracing/Material/PathTracingSpecularReflection.ush
+++ b/Engine/Shaders/Private/PathTracing/Material/PathTracingSpecularReflection.ush
@@ -24,6 +24,8 @@ void SpecularReflection_SampleMaterial(
 	OutThroughput = F_Schlick(SpecularColor, VoH) / NoL;
 
 	OutPdf = 1.0;
+
+	//OutThroughput = float3(1.0,0.0,0.0);
 }
 
 float SpecularReflection_Pdf()
diff --git a/Engine/Shaders/Private/PathTracing/Material/PathTracingTestBrdfs.usf b/Engine/Shaders/Private/PathTracing/Material/PathTracingTestBrdfs.usf
--- a/Engine/Shaders/Private/PathTracing/Material/PathTracingTestBrdfs.usf
+++ b/Engine/Shaders/Private/PathTracing/Material/PathTracingTestBrdfs.usf
@@ -90,7 +90,6 @@ void TestBrdfIntegrityCS(
 		PDF = -1.0;
 	}
 
-	//#dxr_todo accumulate relative errors instead of flags?
 	float BrdfThreshold = 0.05;
 	float PdfThreshold = 0.05;
 	bool SampleOK = true;
diff --git a/Engine/Shaders/Private/PathTracing/PathTracing.usf b/Engine/Shaders/Private/PathTracing/PathTracing.usf
--- a/Engine/Shaders/Private/PathTracing/PathTracing.usf
+++ b/Engine/Shaders/Private/PathTracing/PathTracing.usf
@@ -54,7 +54,6 @@ void PathTracingMainRG()
 	float2 BufferSize = View.BufferSizeAndInvSize.xy; 
 	uint LinearIndex = LaunchIndex.x + (LaunchIndex.y * BufferSize.x); 
 	RandomSequence RandSequence;
-	// #dxr_todo: unify initialization between ray tracing and path tracing raygen shaders
 	RandomSequence_Initialize(RandSequence, LinearIndex, AdaptiveSamplingData.Iteration);
 	RandSequence.Type = AdaptiveSamplingData.RandomSequence;
 	uint SampleIndex = 0;
diff --git a/Engine/Shaders/Private/RayTracing/BuildMipTreeCS.usf b/Engine/Shaders/Private/RayTracing/BuildMipTreeCS.usf
--- a/Engine/Shaders/Private/RayTracing/BuildMipTreeCS.usf
+++ b/Engine/Shaders/Private/RayTracing/BuildMipTreeCS.usf
@@ -31,7 +31,6 @@ void BuildMipTreeCS(
 	float Sum = 0.0;
 	if (MipLevel == 0)
 	{
-		// #dxr_todo: TextureCubeSampleLevel() may be enough.. could do live during query??
 		float3 Direction = GetTextureCubeVector(uint3(GridCoordinate.xy, FaceIndex), Dimensions);
 		float3 Value = TextureCubeSample(Texture, TextureSampler, Direction).rgb;
 		Sum += Luminance(Value.rgb);
@@ -44,7 +43,6 @@ void BuildMipTreeCS(
 		{
 			uint2 Pixel = PixelOffset + uint2(i >> 1, i & 1);
 			uint InputIndex = BufferOffsetAtPixel(Pixel, MipLevel - 1, Dimensions);
-			// #dxr_todo: Experiment between average and max
 			Sum += RWMipTree[InputIndex];
 		}
 		Sum /= 4;
@@ -68,7 +66,6 @@ void BuildRectLightMipTreeCS(
 	float Sum = 0.0;
 	if (MipLevel == 0)
 	{
-		// #dxr_todo: TextureCubeSampleLevel() may be enough.. could do live during query??
 		uint2 TextureDimensions;
 		RectLightTexture.GetDimensions(TextureDimensions.x, TextureDimensions.y);
 		float2 UV = GridCoordinate.xy / float2(TextureDimensions);
@@ -84,8 +81,6 @@ void BuildRectLightMipTreeCS(
 		{
 			uint2 Pixel = PixelOffset + uint2(i >> 1, i & 1);
 			uint InputIndex = BufferOffsetAtPixel(Pixel, MipLevel - 1, Dimensions);
-			// #dxr_todo: Experiment between average and max
-			//Sum = max(Sum, RWMipTree[InputIndex]);
 			Sum += RWMipTree[InputIndex];
 		}
 		Sum /= 4;
diff --git a/Engine/Shaders/Private/RayTracing/BuildMipTreePdfCS.usf b/Engine/Shaders/Private/RayTracing/BuildMipTreePdfCS.usf
--- a/Engine/Shaders/Private/RayTracing/BuildMipTreePdfCS.usf
+++ b/Engine/Shaders/Private/RayTracing/BuildMipTreePdfCS.usf
@@ -29,7 +29,6 @@ float CalcSkyLightPdf(uint2 Pixel)
 	{
 		uint MipIndex = MipCount - Index - 1;
 
-		// #dxr_todo: Use float4?
 		float Cdf[4];
 		Cdf[0] = SampleMipTree(MipPixel, MipIndex);
 		Cdf[1] = Cdf[0] + SampleMipTree(MipPixel + uint2(1, 0), MipIndex);
diff --git a/Engine/Shaders/Private/RayTracing/MipTreeCommon.ush b/Engine/Shaders/Private/RayTracing/MipTreeCommon.ush
--- a/Engine/Shaders/Private/RayTracing/MipTreeCommon.ush
+++ b/Engine/Shaders/Private/RayTracing/MipTreeCommon.ush
@@ -11,7 +11,7 @@ uint2 DimensionsAtLevel(uint2 Dimensions, uint Level)
 	return uint2(Dimensions.x >> Level, Dimensions.y >> Level);
 }
 
-// #dxr_todo: Discussing with Brian where these utility functions should live.
+// #dxr_todo: UE-72536 unifying morton code
 uint MortonEncode(uint2 Pixel)
 {
 	Pixel.x &= 0x0000ffff;
diff --git a/Engine/Shaders/Private/RayTracing/RayTracingBuiltInShaders.usf b/Engine/Shaders/Private/RayTracing/RayTracingBuiltInShaders.usf
--- a/Engine/Shaders/Private/RayTracing/RayTracingBuiltInShaders.usf
+++ b/Engine/Shaders/Private/RayTracing/RayTracingBuiltInShaders.usf
@@ -2,7 +2,6 @@
 
 #include "RayTracingCommon.ush"
 
-// #dxr_todo: move RT_REGISTER macro declaration into some a shared header, as it's now used in a few places
 #define RT_CONCATENATE2(a, b) a##b
 #define RT_CONCATENATE(a, b) RT_CONCATENATE2(a, b)
 #define RT_REGISTER(InType, InIndex, InSpace) register(RT_CONCATENATE(InType, InIndex), RT_CONCATENATE(space, InSpace))
diff --git a/Engine/Shaders/Private/RayTracing/RayTracingGlobalIlluminationRGS.usf b/Engine/Shaders/Private/RayTracing/RayTracingGlobalIlluminationRGS.usf
--- a/Engine/Shaders/Private/RayTracing/RayTracingGlobalIlluminationRGS.usf
+++ b/Engine/Shaders/Private/RayTracing/RayTracingGlobalIlluminationRGS.usf
@@ -29,7 +29,7 @@ bool EvalSkyLight;
 bool UseRussianRoulette;
 float MaxNormalBias;
 
-// #dxr_todo: Unify with reflections and translucency in RayTracingCommon.ush
+// #dxr_todo: UE-72508 Unify with reflections and translucency in RayTracingCommon.ush
 uint2 GetPixelCoord(uint2 DispatchThreadId)
 {
 	uint UpscaleFactorPow2 = UpscaleFactor * UpscaleFactor;
@@ -205,7 +205,6 @@ void GlobalIlluminationRGS()
 				}
 				break;
 			}
-			// #dxr_todo: Allow for material emission?
 
 			if (Bounce == 0)
 			{
@@ -328,7 +327,6 @@ void GlobalIlluminationRGS()
 				FDirectLighting LightingSample = EvaluateBxDF(GBufferData, N, V, L, NoL, ShadowTerms);
 				Throughput = LightingSample.Diffuse;
 #endif
-				// #dxr_todo: Degenerate guard?
 				if (MaterialPdf <= 0.0)
 				{
 					break;
@@ -339,10 +337,6 @@ void GlobalIlluminationRGS()
 				ApplyPositionBias(Ray, WorldNormal, MaxNormalBias);
 				RayThroughput *= Throughput;
 				RayPdf *= MaterialPdf;
-
-				// #dxr_todo: Russian roulette?
-
-				// #dxr_todo: Firefly rejection?
 			}
 
 			Bounce++;
diff --git a/Engine/Shaders/Private/RayTracing/RayTracingHitGroupCommon.ush b/Engine/Shaders/Private/RayTracing/RayTracingHitGroupCommon.ush
--- a/Engine/Shaders/Private/RayTracing/RayTracingHitGroupCommon.ush
+++ b/Engine/Shaders/Private/RayTracing/RayTracingHitGroupCommon.ush
@@ -84,8 +84,7 @@ FTriangleBaseAttributes LoadTriangleBaseAttributes(uint PrimitiveId)
 	// Fetch vertex positions (in local space)
 	const uint VertexStride = HitGroupSystemRootConstants.GetVertexStride();
 
-	// #dxr_todo: handle various vertex formats, for now only supporting float3
-	// #dxr_todo: handle BaseVertexIndex
+	// #dxr_todo: UE-72160 handle various vertex formats, for now only supporting float3
 	Result.LocalPositions[0] = LoadVertexPositionFloat3(HitGroupSystemVertexBuffer, Result.Indices[0], VertexStride);
 	Result.LocalPositions[1] = LoadVertexPositionFloat3(HitGroupSystemVertexBuffer, Result.Indices[1], VertexStride);
 	Result.LocalPositions[2] = LoadVertexPositionFloat3(HitGroupSystemVertexBuffer, Result.Indices[2], VertexStride);
diff --git a/Engine/Shaders/Private/RayTracing/RayTracingLightingCommon.ush b/Engine/Shaders/Private/RayTracing/RayTracingLightingCommon.ush
--- a/Engine/Shaders/Private/RayTracing/RayTracingLightingCommon.ush
+++ b/Engine/Shaders/Private/RayTracing/RayTracingLightingCommon.ush
@@ -212,7 +212,6 @@ uint2 CullDirectLighting(
 
 	uint2 LightCullMask = (uint2)0;
 
-	//#dxr_todo: think on the more efficient way of iterate through lights (split in light types, etc)
 	for (uint DirIndex = 0; DirIndex < RaytracingLightsDataPacked.Count; DirIndex++)
 	{
 		FDeferredLightData LightData = (FDeferredLightData)0;
@@ -237,7 +236,7 @@ uint2 CullDirectLighting(
 		const uint LightProfileIndex = RaytracingLightsDataPacked.Type_LightProfileIndex_RectLightTextureIndex[DirIndex].y;
 
 
-		LightData.ShadowMapChannelMask = float4(0, 0, 0, 0); // #dxr_todo;
+		LightData.ShadowMapChannelMask = float4(0, 0, 0, 0); 
 		LightData.ShadowedBits = 0; // Not lit dynamic shadows 
 		LightData.ContactShadowLengthInWS = false;
 
@@ -303,7 +302,6 @@ float3 ComputeDirectLightingCulled(
 	float3 CameraVector = -Ray.Direction;
 	float AmbientOcclusion = 1.0;
 
-	// #dxr_todo: review when interpolants fill the Payload with all the data we need
 	FGBufferData GBufferData = (FGBufferData)0;
 	GBufferData.Depth = 1.f; // Do not use depth
 	GBufferData.WorldNormal = Payload.WorldNormal;
@@ -313,7 +311,6 @@ float3 ComputeDirectLightingCulled(
 	GBufferData.Roughness = Payload.Roughness;
 	GBufferData.CustomData = Payload.CustomData;
 	GBufferData.GBufferAO = Payload.GBufferAO;
-	//#dxr_todo: change float3 to float in payload to reduce space?;
 	GBufferData.IndirectIrradiance = (Payload.IndirectIrradiance.x + Payload.IndirectIrradiance.y + Payload.IndirectIrradiance.z) / 3.f; 
 	GBufferData.ShadingModelID = SHADINGMODELID_DEFAULT_LIT;
 	GBufferData.SpecularColor = Payload.SpecularColor;
@@ -338,8 +335,6 @@ float3 ComputeDirectLightingCulled(
 
 	uint DirIndex = 0;
 
-	//#dxr_todo: think on the more efficient way of iterate through lights (split in light types, etc)
-	//while(LightCullMask.x || LightCullMask.y)
 	while (WaveActiveAnyTrue(LightCullMask.x || LightCullMask.y))
 	{
 		const bool Active = LightCullMask.x || LightCullMask.y;
@@ -397,7 +392,7 @@ float3 ComputeDirectLightingCulled(
 	#endif
 		uint LightProfileIndex = RaytracingLightsDataPacked.Type_LightProfileIndex_RectLightTextureIndex[DirIndex].y;
 
-		LightData.ShadowMapChannelMask = float4(0, 0, 0, 0); // #dxr_todo;
+		LightData.ShadowMapChannelMask = float4(0, 0, 0, 0); 
 		LightData.ShadowedBits = 0; // Not lit dynamic shadows 
 		LightData.ContactShadowLengthInWS = false;
 
@@ -438,7 +433,6 @@ float3 ComputeDirectLightingCulled(
 			LightMask = GetLocalLightAttenuation(WorldPosition, LightData, ToLight, ShadowRayDirection);
 		}
 
-		//#dxr_todo: check perf improvements forcing TMin == TMax when bReflectedShadows = false, to avoid tracing rays inside if statements
 		if (bReflectedShadows)
 		{
 			if (ReflectedShadowsType == 2)
@@ -448,7 +442,6 @@ float3 ComputeDirectLightingCulled(
 
 			FRayCone ShadowRayCone = RayCone;
 
-			//#dxr_todo: add support for sampling area lights (and perhaps add CVar to enable stocastic shadows?)
 			RayDesc ShadowRay;
 			ShadowRay.Origin = WorldPosition;
 			ShadowRay.Direction = ShadowRayDirection;
diff --git a/Engine/Shaders/Private/RayTracing/RayTracingMaterialHitShaders.usf b/Engine/Shaders/Private/RayTracing/RayTracingMaterialHitShaders.usf
--- a/Engine/Shaders/Private/RayTracing/RayTracingMaterialHitShaders.usf
+++ b/Engine/Shaders/Private/RayTracing/RayTracingMaterialHitShaders.usf
@@ -10,7 +10,7 @@
 	#define USING_TESSELLATION 1
 #endif
 
-// #dxr_todo: Temporarily hacking clip() away to prevent driver crash. Should be integrated into MaterialTemplate.ush in the future.
+// #dxr_todo: UE-72537 move to common.ush
 #define clip(x) (x)
 #define ddx(x) 0
 #define ddy(x) 0
@@ -262,7 +262,7 @@ void MaterialCHS(inout FPackedMaterialClosestHitPayload PackedPayload, in FDefau
 
 		MaterialParameters.CameraVector = -WorldRayDirection();
 
-		// #dxr_todo: support world position offset at some point
+		// #dxr_todo: UE-72130 support world position offset 
 		// #if USE_WORLD_POSITION_EXCLUDING_SHADER_OFFSETS
 		// CalcMaterialParametersEx(MaterialParameters, PixelMaterialInputs, SvPosition, ScreenPosition, bIsFrontFace, TranslatedWorldPosition, BasePassInterpolants.PixelPositionExcludingWPO);
 		// #else
@@ -374,7 +374,7 @@ void MaterialAHS(inout FPackedMaterialClosestHitPayload PackedPayload, in FDefau
 
 		MaterialParameters.CameraVector = -WorldRayDirection();
 		
-		// #dxr_todo: support world position offset at some point
+		// #dxr_todo: UE-72130 support world position offset 
 		// #if USE_WORLD_POSITION_EXCLUDING_SHADER_OFFSETS
 		// CalcMaterialParametersEx(MaterialParameters, PixelMaterialInputs, SvPosition, ScreenPosition, bIsFrontFace, TranslatedWorldPosition, BasePassInterpolants.PixelPositionExcludingWPO);
 		// #else
diff --git a/Engine/Shaders/Private/RayTracing/RayTracingRectLightRGS.usf b/Engine/Shaders/Private/RayTracing/RayTracingRectLightRGS.usf
--- a/Engine/Shaders/Private/RayTracing/RayTracingRectLightRGS.usf
+++ b/Engine/Shaders/Private/RayTracing/RayTracingRectLightRGS.usf
@@ -145,8 +145,6 @@ bool GenerateRectLightImportanceSampledRay(
 		return false;
 	}
 
-	// #dxr_todo: Performing direct lighting here could save some recalculation
-
 	// Define ray
 	RayOrigin = WorldPosition;
 	RayDirection = LightDirection;
@@ -283,7 +281,7 @@ void RectLightRGS()
 
 		uint RayFlags = 0
 			| RAY_FLAG_CULL_BACK_FACING_TRIANGLES
-			| RAY_FLAG_FORCE_OPAQUE; // #dxr_todo: generate a special shader permutation for occlusion ray tracing that includes any-hit shaders for masked geometry
+			| RAY_FLAG_FORCE_OPAQUE; // #dxr_todo: UE-72563 generate a special shader permutation for occlusion ray tracing that includes any-hit shaders for masked geometry
 
 		FDefaultPayload Payload = (FDefaultPayload)0;
 
@@ -322,10 +320,6 @@ void RectLightRGS()
 			FDirectLighting LightingSample = EvaluateBxDF(ScreenSpaceData.GBuffer, N, V, L, NoL, ShadowTerms);
 			float3 Brdf = LightingSample.Diffuse + LightingSample.Transmission + LightingSample.Specular;
 
-			// #dxr_todo: Determine best way to incorporate integrals:
-			// 1) Return combined contribution (current method)
-			// 2) Split the integral (defer specular contribution to another pass)
-			// 3) Store diffuse and specular contribution separately
 			float3 Irradiance = IncomingRadiance * NoL / RayPdf;
 			ExitantRadiance += Irradiance * Brdf;
 			DiffuseExitantRadiance += Irradiance * (LightingSample.Diffuse + LightingSample.Transmission);
diff --git a/Engine/Shaders/Private/RayTracing/RayTracingReflections.usf b/Engine/Shaders/Private/RayTracing/RayTracingReflections.usf
--- a/Engine/Shaders/Private/RayTracing/RayTracingReflections.usf
+++ b/Engine/Shaders/Private/RayTracing/RayTracingReflections.usf
@@ -63,7 +63,7 @@ float ReflectionMaxNormalBias;
 
 RaytracingAccelerationStructure TLAS;
 
-#include "RayTracingLightingCommon.ush" // #dxr_todo: helper functions must take TLAS explicitly rather than rely on a global declared above
+#include "RayTracingLightingCommon.ush" // #dxr_todo: UE-72548 helper functions must take TLAS explicitly rather than rely on a global declared above
 
 RWTexture2D<float4> ColorOutput;
 RWTexture2D<float> RayHitDistanceOutput;
diff --git a/Engine/Shaders/Private/RayTracing/RayTracingSkyLightRGS.usf b/Engine/Shaders/Private/RayTracing/RayTracingSkyLightRGS.usf
--- a/Engine/Shaders/Private/RayTracing/RayTracingSkyLightRGS.usf
+++ b/Engine/Shaders/Private/RayTracing/RayTracingSkyLightRGS.usf
@@ -21,7 +21,7 @@ uint CalcLinearIndex(uint2 PixelCoord)
 	return PixelCoord.y * View.BufferSizeAndInvSize.x + PixelCoord.x;
 }
 
-// #dxr_todo: Copied from PathTracingSkyLight.ush
+// #dxr_todo: UE-72508: encapsulate this function in a shared space
 float3 EvalSkyLight(float3 Direction)
 {
 	return SkyLight.Color * TextureCubeSampleLevel(SkyLight.Texture, SkyLight.TextureSampler, Direction, 0).rgb;
@@ -41,7 +41,6 @@ float PdfSkyLightOcclusionRay(
 	uint2 MipCoord = TextureUV * SkyLight.MipDimensions.xy;
 	float FacePdf = PdfFace(MipCount, TextureCoord.z);
 
-	// #dxr_todo: StopLevel can be customized for run-time performance
 	uint StopLevel = min(SkyLight.SamplingStopLevel, MipCount);
 #if USE_MIP_PDF_TEXTURE
 	float MipPdf = PdfMipTree(uint3(MipCoord, TextureCoord.z), StopLevel);
@@ -284,7 +283,7 @@ void SkyLightRGS()
 		// Determine sky light or lambert ray
 		uint DummyVariable;
 		float4 RandSample = RandomSequence_GenerateSample4D(RandSequence, DummyVariable);
-		// #dxr_todo: Should strategy selection be driven by CVar?
+
 		float SkyLightSamplingStrategyPdf = 0.5;
 		float MisWeight = 1.0;
 		if (RandSample.x < SkyLightSamplingStrategyPdf)
@@ -309,7 +308,7 @@ void SkyLightRGS()
 		}
 
 		uint RayFlags = 0
-			| RAY_FLAG_FORCE_OPAQUE; // #dxr_todo: generate a special shader permutation for occlusion ray tracing that includes any-hit shaders for masked geometry
+			| RAY_FLAG_FORCE_OPAQUE; // #dxr_todo: UE-72563 generate a special shader permutation for occlusion ray tracing that includes any-hit shaders for masked geometry
 #if !ENABLE_TWO_SIDED_GEOMETRY
 		RayFlags |= RAY_FLAG_CULL_BACK_FACING_TRIANGLES;
 #endif
diff --git a/Engine/Shaders/Private/RayTracing/RayTracingTranslucency.usf b/Engine/Shaders/Private/RayTracing/RayTracingTranslucency.usf
--- a/Engine/Shaders/Private/RayTracing/RayTracingTranslucency.usf
+++ b/Engine/Shaders/Private/RayTracing/RayTracingTranslucency.usf
@@ -131,8 +131,6 @@ void RayTracingTranslucencyRGS()
 	const float GBufferHitT = length(WorldPosition - CameraOrigin);
 
 	// Trace rays from camera origin to (Gbuffer - epsilon) to only intersect translucent objects
-	//#dxr_todo: a potential big optimization would be to include translucent objects in the GBuffer
-	// and use the GBuffer shading ID to early return for non translucent objects 
 	RayDesc Ray;
 	Ray.Origin = CameraOrigin;
 	Ray.TMin = 0.0;
@@ -247,7 +245,7 @@ void RayTracingTranslucencyRGS()
 			float ReflectionRefractionEventThroughput;
 			if (!RefractRay(Ray.Direction, Payload.WorldNormal, Ior1, Ior2, RefractedDirection, ReflectionRefractionEventThroughput))
 			{
-				//#dxr_todo: handle total internal reflection ray event. We could jsuit continue form this path. And not use the new IOR.
+				//#dxr_todo: handle total internal reflection ray event. We could just continue form this path. And not use the new IOR.
 				// We should combine that step with the one above reflection evaluation. For now we stop tracing.
 				break;
 			}
diff --git a/Engine/Shaders/Shared/RayTracingDefinitions.h b/Engine/Shaders/Shared/RayTracingDefinitions.h
--- a/Engine/Shaders/Shared/RayTracingDefinitions.h
+++ b/Engine/Shaders/Shared/RayTracingDefinitions.h
@@ -12,7 +12,7 @@
 #define RAY_TRACING_SHADER_VERSION 0x510cacb6
 
 #define RAY_TRACING_REGISTER_SPACE_LOCAL  0 // default register space for hit group (closest hit, any hit, intersection) shader resources
-#define RAY_TRACING_REGISTER_SPACE_GLOBAL 1 // register space for ray generation and miss shaders (#dxr_todo: make global resources also available to hit shaders)
+#define RAY_TRACING_REGISTER_SPACE_GLOBAL 1 // register space for ray generation and miss shaders 
 #define RAY_TRACING_REGISTER_SPACE_SYSTEM 2 // register space for "system" parameters (index buffer, vertex buffer, fetch parameters)
 
 #define RAY_TRACING_MASK_OPAQUE						0x01    // Opaque and alpha tested meshes and particles (e.g. used by reflection, shadow, AO and GI tracing passes)
diff --git a/Engine/Source/Editor/UnrealEd/Private/SEditorViewportViewMenu.cpp b/Engine/Source/Editor/UnrealEd/Private/SEditorViewportViewMenu.cpp
--- a/Engine/Source/Editor/UnrealEd/Private/SEditorViewportViewMenu.cpp
+++ b/Engine/Source/Editor/UnrealEd/Private/SEditorViewportViewMenu.cpp
@@ -186,7 +186,7 @@ const FSlateBrush* SEditorViewportViewMenu::GetViewMenuLabelIcon() const
 				break;
 
 #if RHI_RAYTRACING
-			// #dxr_todo: use special icons for ray tracing view modes
+			// #dxr_todo: UE-72550 use special icons for ray tracing view modes
 			case VMI_RayTracingDebug:
 				Icon = LitIcon; 
 				break;
diff --git a/Engine/Source/Runtime/D3D12RHI/Private/D3D12RootSignature.cpp b/Engine/Source/Runtime/D3D12RHI/Private/D3D12RootSignature.cpp
--- a/Engine/Source/Runtime/D3D12RHI/Private/D3D12RootSignature.cpp
+++ b/Engine/Source/Runtime/D3D12RHI/Private/D3D12RootSignature.cpp
@@ -110,8 +110,6 @@ FD3D12RootSignatureDesc::FD3D12RootSignatureDesc(const FD3D12QuantizedBoundShade
 		// -----------
 		// 32 bytes
 
-		// #dxr_todo: Root parameters for hit shaders should be added in the shader pipeline, so that regular root parameter generation code can be used without hard-coding anything.
-
 		check(RootParameterCount == 0 && RootParametersSize == 0); // We expect system RT parameters to come first
 
 		// Index buffer descriptor
diff --git a/Engine/Source/Runtime/Engine/Private/SkeletalRenderCPUSkin.cpp b/Engine/Source/Runtime/Engine/Private/SkeletalRenderCPUSkin.cpp
--- a/Engine/Source/Runtime/Engine/Private/SkeletalRenderCPUSkin.cpp
+++ b/Engine/Source/Runtime/Engine/Private/SkeletalRenderCPUSkin.cpp
@@ -341,7 +341,6 @@ void FSkeletalMeshObjectCPUSkin::FSkeletalMeshObjectLOD::InitResources(FSkelMesh
 		FIndexBufferRHIRef IndexBufferRHI = LODModel.MultiSizeIndexContainer.GetIndexBuffer()->IndexBufferRHI;
 		uint32 VertexBufferStride = LODModel.StaticVertexBuffers.PositionVertexBuffer.GetStride();
 
-		//#dxr_todo: do we need support for separate sections in FRayTracingGeometryData?
 		uint32 TrianglesCount = 0;
 		for (int32 SectionIndex = 0; SectionIndex < LODModel.RenderSections.Num(); SectionIndex++)
 		{
diff --git a/Engine/Source/Runtime/Engine/Private/SkeletalRenderStatic.cpp b/Engine/Source/Runtime/Engine/Private/SkeletalRenderStatic.cpp
--- a/Engine/Source/Runtime/Engine/Private/SkeletalRenderStatic.cpp
+++ b/Engine/Source/Runtime/Engine/Private/SkeletalRenderStatic.cpp
@@ -131,7 +131,6 @@ void FSkeletalMeshObjectStatic::FSkeletalMeshObjectLOD::InitResources(FSkelMeshC
 		FIndexBufferRHIRef IndexBufferRHI = LODModel.MultiSizeIndexContainer.GetIndexBuffer()->IndexBufferRHI;
 		uint32 VertexBufferStride = LODModel.StaticVertexBuffers.PositionVertexBuffer.GetStride();
 
-		//#dxr_todo: do we need support for separate sections in FRayTracingGeometryData?
 		uint32 TrianglesCount = 0;
 		for (int32 SectionIndex = 0; SectionIndex < LODModel.RenderSections.Num(); SectionIndex++)
 		{
diff --git a/Engine/Source/Runtime/Engine/Private/StaticMeshRender.cpp b/Engine/Source/Runtime/Engine/Private/StaticMeshRender.cpp
--- a/Engine/Source/Runtime/Engine/Private/StaticMeshRender.cpp
+++ b/Engine/Source/Runtime/Engine/Private/StaticMeshRender.cpp
@@ -206,7 +206,8 @@ FStaticMeshSceneProxy::FStaticMeshSceneProxy(UStaticMeshComponent* InComponent,
 				MaterialRelevance |= UMaterial::GetDefaultMaterial(MD_Surface)->GetRelevance(FeatureLevel);
 			}
 
-			// #dxr_todo warning this is a hack:
+			// #dxr_todo UE-72552: (skybox should be excluded from RT by default) warning this is a hack:
+			// #dxr_todo UE-72554 (handling of unlit geometries):
 			//  - The only reason we exclude unlit from any RT acceleration structure right now is to not hit the sky dome. 
 			//  - It works but this would need consideration (e.g. infinite RTAO would not work if the sky dome is hit, double reflection due sky light cubemap, reflection testing no intersection, etc.).
 			if (SectionInfo.Material->GetShadingModel() != MSM_Unlit)
diff --git a/Engine/Source/Runtime/Engine/Public/SceneView.h b/Engine/Source/Runtime/Engine/Public/SceneView.h
--- a/Engine/Source/Runtime/Engine/Public/SceneView.h
+++ b/Engine/Source/Runtime/Engine/Public/SceneView.h
@@ -28,7 +28,6 @@ class FSceneViewFamily;
 class FVolumetricFogViewResources;
 class FIESLightProfileResource;
 
-// #dxr_todo: share this enum with ray tracing shader code via RayTracingDefinitions.ush
 enum class ERayTracingRenderMode
 {
 	Disabled			= 0,
diff --git a/Engine/Source/Runtime/RHI/Private/PipelineStateCache.cpp b/Engine/Source/Runtime/RHI/Private/PipelineStateCache.cpp
--- a/Engine/Source/Runtime/RHI/Private/PipelineStateCache.cpp
+++ b/Engine/Source/Runtime/RHI/Private/PipelineStateCache.cpp
@@ -595,8 +595,7 @@ FAutoConsoleTaskPriority CPrio_FCompilePipelineStateTask(
 class FRayTracingPipelineCache
 {
 	// #dxr_todo UE-68234: This needs to support fully asynchronous, non-blocking pipeline creation with explicit completion query mechanism.
-	// #dxr_todo: Could move this to a separate cpp file.
-	// #dxr_todo: Should support eviction of stale pipelines.
+	// #dxr_todo: UE-72566 Should support eviction of stale pipelines.
 	// #dxr_todo: we will likely also need an explicit ray tracing pipeline sub-object cache to hold closest hit shaders
 
 public:
diff --git a/Engine/Source/Runtime/RHI/Public/RHICommandList.h b/Engine/Source/Runtime/RHI/Public/RHICommandList.h
--- a/Engine/Source/Runtime/RHI/Public/RHICommandList.h
+++ b/Engine/Source/Runtime/RHI/Public/RHICommandList.h
@@ -110,7 +110,6 @@ extern RHI_API TAutoConsoleVariable<int32> CVarRHICmdWidth;
 extern RHI_API TAutoConsoleVariable<int32> CVarRHICmdFlushRenderThreadTasks;
 
 #if RHI_RAYTRACING
-// #dxr_todo: This is a placeholder. FMeshDrawSingleShaderBindings should be generalized and used instead.
 struct FRayTracingShaderBindings
 {
 	FTextureRHIParamRef Textures[32] = {};
diff --git a/Engine/Source/Runtime/Renderer/Private/DeferredShadingRenderer.cpp b/Engine/Source/Runtime/Renderer/Private/DeferredShadingRenderer.cpp
--- a/Engine/Source/Runtime/Renderer/Private/DeferredShadingRenderer.cpp
+++ b/Engine/Source/Runtime/Renderer/Private/DeferredShadingRenderer.cpp
@@ -620,7 +620,7 @@ bool FDeferredShadingSceneRenderer::GatherRayTracingWorldInstances(FRHICommandLi
 					continue;
 				}
 
-				//#dxr_todo The Raytracing codepath does not support Showflags since data moved to the SceneInfo. 
+				//#dxr_todo UE-68621  The Raytracing codepath does not support Showflags since data moved to the SceneInfo. 
 				//Touching the SceneProxy to determine this would simply cost too much
 				if (SceneInfo->bShouldRenderInMainPass && SceneInfo->bDrawInGame)
 				{
@@ -766,7 +766,7 @@ bool FDeferredShadingSceneRenderer::DispatchRayTracingWorldUpdates(FRHICommandLi
 		View.RayTracingScene.RayTracingSceneRHI = RHICreateRayTracingScene(Initializer);
 		RHICmdList.BuildAccelerationStructure(View.RayTracingScene.RayTracingSceneRHI);
 
-		// #dxr_todo: register each effect at startup and just loop over them automatically to gather all required shaders
+		// #dxr_todo: UE-72565: refactor ray tracing effects to not be member functions of DeferredShadingRenderer. register each effect at startup and just loop over them automatically to gather all required shaders
 		TArray<FRayTracingShaderRHIParamRef> RayGenShaders;
 		PrepareRayTracingReflections(View, RayGenShaders);
 		PrepareRayTracingShadows(View, RayGenShaders);
@@ -948,7 +948,7 @@ void FDeferredShadingSceneRenderer::Render(FRHICommandListImmediate& RHICmdList)
 		&& ViewFamily.EngineShowFlags.DeferredLighting
 		&& bUseGBuffer
 #if RHI_RAYTRACING
-		&& (Views[0].RayTracingRenderMode != ERayTracingRenderMode::PathTracing) // #dxr_todo: what about multi-view case?
+		&& (Views[0].RayTracingRenderMode != ERayTracingRenderMode::PathTracing) // #dxr_todo: UE-72557 multi-view case
 		&& (Views[0].RayTracingRenderMode != ERayTracingRenderMode::RayTracingDebug)
 #endif
 		;
@@ -1933,7 +1933,7 @@ void FDeferredShadingSceneRenderer::Render(FRHICommandListImmediate& RHICmdList)
 		for (int32 ViewIndex = 0, Num = Views.Num(); ViewIndex < Num; ViewIndex++)
 		{
 			const FViewInfo& View = Views[ViewIndex];
-			//#dxr_todo: multiview case
+			//#dxr_todo: UE-72557 multiview case
 			bAnyViewWithRaytracingTranslucency = bAnyViewWithRaytracingTranslucency || (View.FinalPostProcessSettings.TranslucencyType == ETranslucencyType::RayTracing);
 		}
 
@@ -2353,7 +2353,7 @@ void FDeferredShadingSceneRenderer::CopyStencilToLightingChannelTexture(FRHIComm
 bool FDeferredShadingSceneRenderer::CanOverlayRayTracingOutput(void) const
 {
 #if RHI_RAYTRACING
-	// #dxr_todo: what about multi-view case?
+	// #dxr_todo:UE-72557 multi-view case
 	return  (Views[0].RayTracingRenderMode != ERayTracingRenderMode::PathTracing)
 		&& (Views[0].RayTracingRenderMode != ERayTracingRenderMode::RayTracingDebug);
 #else // RHI_RAYTRACING
diff --git a/Engine/Source/Runtime/Renderer/Private/DeferredShadingRenderer.h b/Engine/Source/Runtime/Renderer/Private/DeferredShadingRenderer.h
--- a/Engine/Source/Runtime/Renderer/Private/DeferredShadingRenderer.h
+++ b/Engine/Source/Runtime/Renderer/Private/DeferredShadingRenderer.h
@@ -553,7 +553,7 @@ private:
 	FRHIRayTracingPipelineState* BindRayTracingMaterialPipeline(FRHICommandList& RHICmdList, const FViewInfo& View, const TArrayView<const FRayTracingShaderRHIParamRef>& RayGenShaderTable, FRayTracingShaderRHIParamRef MissShader, FRayTracingShaderRHIParamRef DefaultClosestHitShader);
 	FRHIRayTracingPipelineState* BindRayTracingDeferredMaterialGatherPipeline(FRHICommandList& RHICmdList, const FViewInfo& View, FRayTracingShaderRHIParamRef RayGenShader);
 
-	// #dxr_todo: register each effect at startup and just loop over them automatically
+	// #dxr_todo: UE-72565: refactor ray tracing effects to not be member functions of DeferredShadingRenderer. Register each effect at startup and just loop over them automatically
 	static void PrepareRayTracingReflections(const FViewInfo& View, TArray<FRayTracingShaderRHIParamRef>& OutRayGenShaders);
 	static void PrepareRayTracingShadows(const FViewInfo& View, TArray<FRayTracingShaderRHIParamRef>& OutRayGenShaders);
 	static void PrepareRayTracingRectLight(const FViewInfo& View, TArray<FRayTracingShaderRHIParamRef>& OutRayGenShaders);
diff --git a/Engine/Source/Runtime/Renderer/Private/PathTracing.cpp b/Engine/Source/Runtime/Renderer/Private/PathTracing.cpp
--- a/Engine/Source/Runtime/Renderer/Private/PathTracing.cpp
+++ b/Engine/Source/Runtime/Renderer/Private/PathTracing.cpp
@@ -197,7 +197,7 @@ public:
 						LightData.Normal[LightData.Count] = -LightParameters.Direction;
 						LightData.dPdu[LightData.Count] = FVector::CrossProduct(LightParameters.Tangent, LightParameters.Direction);
 						LightData.dPdv[LightData.Count] = LightParameters.Tangent;
-						// #dxr_todo: define these differences from Lit..
+						// #dxr_todo: UE-72556  define these differences from Lit..
 						LightData.Color[LightData.Count] = LightParameters.Color / 4.0;
 						LightData.Dimensions[LightData.Count] = FVector(2.0f * LightParameters.SourceRadius, 2.0f * LightParameters.SourceLength, 0.0f);
 						LightData.Attenuation[LightData.Count] = 1.0 / LightParameters.InvRadius;
@@ -210,7 +210,7 @@ public:
 						LightData.Type[LightData.Count] = 4;
 						LightData.Position[LightData.Count] = LightParameters.Position;
 						LightData.Normal[LightData.Count] = -LightParameters.Direction;
-						// #dxr_todo: define these differences from Lit..
+						// #dxr_todo: UE-72556  define these differences from Lit..
 						LightData.Color[LightData.Count] = 4.0 * PI * LightParameters.Color;
 						LightData.Dimensions[LightData.Count] = FVector(LightParameters.SpotAngles, LightParameters.SourceRadius);
 						LightData.Attenuation[LightData.Count] = 1.0 / LightParameters.InvRadius;
@@ -221,7 +221,7 @@ public:
 					{
 						LightData.Type[LightData.Count] = 1;
 						LightData.Position[LightData.Count] = LightParameters.Position;
-						// #dxr_todo: define these differences from Lit..
+						// #dxr_todo: UE-72556  define these differences from Lit..
 						LightData.Color[LightData.Count] = LightParameters.Color / (4.0 * PI);
 						LightData.Dimensions[LightData.Count] = FVector(0.0, 0.0, LightParameters.SourceRadius);
 						LightData.Attenuation[LightData.Count] = 1.0 / LightParameters.InvRadius;
@@ -361,8 +361,7 @@ class FPathTracingCompositorPS : public FGlobalShader
 
 	static bool ShouldCompilePermutation(const FGlobalShaderPermutationParameters& Parameters)
 	{
-		// #dxr_todo: this should also check if ray tracing is enabled for the target platform & project
-		return IsFeatureLevelSupported(Parameters.Platform, ERHIFeatureLevel::SM5);
+		return ShouldCompileRayTracingShadersForProject(Parameters.Platform);
 	}
 
 	static void ModifyCompilationEnvironment(const FGlobalShaderPermutationParameters& Parameters, FShaderCompilerEnvironment& OutEnvironment)
diff --git a/Engine/Source/Runtime/Renderer/Private/PathTracing/PathTracingTestbed.cpp b/Engine/Source/Runtime/Renderer/Private/PathTracing/PathTracingTestbed.cpp
--- a/Engine/Source/Runtime/Renderer/Private/PathTracing/PathTracingTestbed.cpp
+++ b/Engine/Source/Runtime/Renderer/Private/PathTracing/PathTracingTestbed.cpp
@@ -361,7 +361,6 @@ IMPLEMENT_SHADER_TYPE(, FTestBrdfIntegrateHemispherePatchCS, TEXT("/Engine/Priva
 TArray<FString> GetBrdfNames(void)
 {
 	// Order should match BrdfType parameter in TestBrdfIntegrityCS until there is a common definition
-	//#dxr_todo: add support for testing different parameters for each BRDF
 	TArray<FString> BrdfNames;
 	BrdfNames.Add("Lambert");
 	BrdfNames.Add("Glossy GGX");
@@ -408,7 +407,6 @@ void TestBRDFsIntegrity(void)
 
 	for (uint32 BrdfType = 0; BrdfType < (uint32) BrdfNames.Num(); ++BrdfType)
 	{
-		//#dxr_todo: check thread safety of UE_LOG() when tests are moved to the EngineTest framework (safety depends because FOutputDevice can overload CanBeUsedOnAnyThread() )
 		UE_LOG(LogShaders, Display, TEXT("Executed validation test for BRDF: %s"), *BrdfNames[BrdfType]);
 
 		FRWBufferStructured BrdsResultsBuffer;
@@ -579,7 +577,7 @@ void TestBRDFandPDFConsistency(void)
 	FRandomStream RandomStream(0);
 	for (uint32 i = 0; i < WosCount; ++i)
 	{
-		//#dxr_todo: use cos weight instead of uniform for this test
+		//#dxr_todo: UE-72555 use cos weight instead of uniform for this test
 		const float U1 = RandomStream.GetFraction();
 		const float U2 = RandomStream.GetFraction();
 
@@ -703,7 +701,7 @@ void TestBRDFandPDFConsistency(void)
 				{
 					uint32 diff = fabsf(SampledPercent - ExpectedPercent);
 					float Err = (float)diff;
-					if (Err > 20.0) //#dxr_todo: adjust threshold
+					if (Err > 20.0)
 					{
 						BinMismatches++;
 					}
diff --git a/Engine/Source/Runtime/Renderer/Private/RayTracing/RayTracingAmbientOcclusion.cpp b/Engine/Source/Runtime/Renderer/Private/RayTracing/RayTracingAmbientOcclusion.cpp
--- a/Engine/Source/Runtime/Renderer/Private/RayTracing/RayTracingAmbientOcclusion.cpp
+++ b/Engine/Source/Runtime/Renderer/Private/RayTracing/RayTracingAmbientOcclusion.cpp
@@ -114,7 +114,7 @@ public:
 		FRayTracingShaderRHIParamRef RayGenShaderTable[] = { GetRayTracingShader() };
 		Initializer.SetRayGenShaderTable(RayGenShaderTable);
 
-		FRHIRayTracingPipelineState* Pipeline = PipelineStateCache::GetAndOrCreateRayTracingPipelineState(Initializer); // #dxr_todo: this should be done once at load-time and cached
+		FRHIRayTracingPipelineState* Pipeline = PipelineStateCache::GetAndOrCreateRayTracingPipelineState(Initializer); 
 
 		FRayTracingShaderBindingsWriter GlobalResources;
 		GlobalResources.Set(TLASParameter, RayTracingScene.RayTracingSceneRHI->GetShaderResourceView());
diff --git a/Engine/Source/Runtime/Renderer/Private/RayTracing/RayTracingBarycentrics.cpp b/Engine/Source/Runtime/Renderer/Private/RayTracing/RayTracingBarycentrics.cpp
--- a/Engine/Source/Runtime/Renderer/Private/RayTracing/RayTracingBarycentrics.cpp
+++ b/Engine/Source/Runtime/Renderer/Private/RayTracing/RayTracingBarycentrics.cpp
@@ -96,7 +96,7 @@ void FDeferredShadingSceneRenderer::RenderRayTracingBarycentrics(FRHICommandList
 	Initializer.SetHitGroupTable(HitGroupTable);
 	Initializer.bAllowHitGroupIndexing = false; // Use the same hit shader for all geometry in the scene by disabling SBT indexing.
 
-	FRHIRayTracingPipelineState* Pipeline = PipelineStateCache::GetAndOrCreateRayTracingPipelineState(Initializer); // #dxr_todo: this should be done once at load-time and cached
+	FRHIRayTracingPipelineState* Pipeline = PipelineStateCache::GetAndOrCreateRayTracingPipelineState(Initializer); 
 
 	FRayTracingSceneRHIParamRef RayTracingSceneRHI = View.RayTracingScene.RayTracingSceneRHI;
 
diff --git a/Engine/Source/Runtime/Renderer/Private/RayTracing/RayTracingDynamicGeometry.cpp b/Engine/Source/Runtime/Renderer/Private/RayTracing/RayTracingDynamicGeometry.cpp
--- a/Engine/Source/Runtime/Renderer/Private/RayTracing/RayTracingDynamicGeometry.cpp
+++ b/Engine/Source/Runtime/Renderer/Private/RayTracing/RayTracingDynamicGeometry.cpp
@@ -30,7 +30,6 @@ public:
 
 	static bool ShouldCompilePermutation(EShaderPlatform Platform, const FMaterial* Material, const FVertexFactoryType* VertexFactoryType)
 	{
-		// #dxr_todo: this should also check if ray tracing is enabled for the target platform & project
 		return IsSupportedDynamicVertexFactoryType(VertexFactoryType) && ShouldCompileRayTracingShadersForProject(Platform);
 	}
 
diff --git a/Engine/Source/Runtime/Renderer/Private/RayTracing/RayTracingGlobalIllumination.cpp b/Engine/Source/Runtime/Renderer/Private/RayTracing/RayTracingGlobalIllumination.cpp
--- a/Engine/Source/Runtime/Renderer/Private/RayTracing/RayTracingGlobalIllumination.cpp
+++ b/Engine/Source/Runtime/Renderer/Private/RayTracing/RayTracingGlobalIllumination.cpp
@@ -141,7 +141,7 @@ void SetupLightParameters(
 			LightParameters->Normal[LightParameters->Count] = -LightShaderParameters.Direction;
 			LightParameters->dPdu[LightParameters->Count] = FVector::CrossProduct(LightShaderParameters.Direction, LightShaderParameters.Tangent);
 			LightParameters->dPdv[LightParameters->Count] = LightShaderParameters.Tangent;
-			// #dxr_todo: define these differences from Lit..
+			// #dxr_todo: UE-72556 define these differences from Lit..
 			LightParameters->Color[LightParameters->Count] = LightShaderParameters.Color / 4.0;
 			LightParameters->Dimensions[LightParameters->Count] = FVector(2.0f * LightShaderParameters.SourceRadius, 2.0f * LightShaderParameters.SourceLength, 0.0f);
 			LightParameters->Attenuation[LightParameters->Count] = 1.0 / LightShaderParameters.InvRadius;
@@ -152,7 +152,7 @@ void SetupLightParameters(
 		{
 			LightParameters->Type[LightParameters->Count] = 1;
 			LightParameters->Position[LightParameters->Count] = LightShaderParameters.Position;
-			// #dxr_todo: define these differences from Lit..
+			// #dxr_todo: UE-72556 define these differences from Lit..
 			LightParameters->Color[LightParameters->Count] = LightShaderParameters.Color / (4.0 * PI);
 			float SourceRadius = 0.0; // LightShaderParameters.SourceRadius causes too much noise for little pay off at this time
 			LightParameters->Dimensions[LightParameters->Count] = FVector(0.0, 0.0, SourceRadius);
@@ -164,7 +164,7 @@ void SetupLightParameters(
 			LightParameters->Type[LightParameters->Count] = 4;
 			LightParameters->Position[LightParameters->Count] = LightShaderParameters.Position;
 			LightParameters->Normal[LightParameters->Count] = -LightShaderParameters.Direction;
-			// #dxr_todo: define these differences from Lit..
+			// #dxr_todo: UE-72556 define these differences from Lit..
 			LightParameters->Color[LightParameters->Count] = 4.0 * PI * LightShaderParameters.Color;
 			float SourceRadius = 0.0; // LightShaderParameters.SourceRadius causes too much noise for little pay off at this time
 			LightParameters->Dimensions[LightParameters->Count] = FVector(LightShaderParameters.SpotAngles, SourceRadius);
@@ -188,7 +188,7 @@ bool ShouldRenderRayTracingGlobalIllumination(const TArray<FViewInfo>& Views)
 		for (int32 ViewIndex = 0, Num = Views.Num(); ViewIndex < Num; ViewIndex++)
 		{
 			const FViewInfo& View = Views[ViewIndex];
-			//#dxr_todo: multiview case
+			//#dxr_todo: UE-72557 multiview case
 			if (View.FinalPostProcessSettings.RayTracingGI > 0)
 			{
 				return true;
diff --git a/Engine/Source/Runtime/Renderer/Private/RayTracing/RayTracingIESLightProfiles.h b/Engine/Source/Runtime/Renderer/Private/RayTracing/RayTracingIESLightProfiles.h
--- a/Engine/Source/Runtime/Renderer/Private/RayTracing/RayTracingIESLightProfiles.h
+++ b/Engine/Source/Runtime/Renderer/Private/RayTracing/RayTracingIESLightProfiles.h
@@ -54,7 +54,7 @@ private:
 			&& Texture->PlatformData->PixelFormat == PF_FloatRGBA
 			&& Texture->PlatformData->Mips.Num() == 1
 			&& Texture->PlatformData->Mips[0].SizeX == AllowedIESProfileWidth() 
-			//#dxr_todo: anisotropy in IES files is ignored so far (to support that, we should not store one IES profile per row but use more than one row per profile in that case)
+			//#dxr_todo: UE-70840 anisotropy in IES files is ignored so far (to support that, we should not store one IES profile per row but use more than one row per profile in that case)
 			&& Texture->PlatformData->Mips[0].SizeY == 1
 			)
 		{
diff --git a/Engine/Source/Runtime/Renderer/Private/RayTracing/RayTracingMaterialHitShaders.cpp b/Engine/Source/Runtime/Renderer/Private/RayTracing/RayTracingMaterialHitShaders.cpp
--- a/Engine/Source/Runtime/Renderer/Private/RayTracing/RayTracingMaterialHitShaders.cpp
+++ b/Engine/Source/Runtime/Renderer/Private/RayTracing/RayTracingMaterialHitShaders.cpp
@@ -108,7 +108,6 @@ public:
 
 	static bool ShouldCompilePermutation(EShaderPlatform Platform, const FMaterial* Material, const FVertexFactoryType* VertexFactoryType)
 	{
-		// #dxr_todo: this should also check if ray tracing is enabled for the target platform & project
 		return IsSupportedVertexFactoryType(VertexFactoryType)
 			&& (Material->IsMasked() == UseAnyHitShader)
 			&& LightMapPolicyType::ShouldCompilePermutation(Platform, Material, VertexFactoryType)
diff --git a/Engine/Source/Runtime/Renderer/Private/RayTracing/RayTracingRectLight.cpp b/Engine/Source/Runtime/Renderer/Private/RayTracing/RayTracingRectLight.cpp
--- a/Engine/Source/Runtime/Renderer/Private/RayTracing/RayTracingRectLight.cpp
+++ b/Engine/Source/Runtime/Renderer/Private/RayTracing/RayTracingRectLight.cpp
@@ -270,7 +270,7 @@ public:
 		FRayTracingShaderRHIParamRef RayGenShaderTable[] = { GetRayTracingShader() };
 		Initializer.SetRayGenShaderTable(RayGenShaderTable);
 
-		FRHIRayTracingPipelineState* Pipeline = PipelineStateCache::GetAndOrCreateRayTracingPipelineState(Initializer); // #dxr_todo: this should be done once at load-time and cached
+		FRHIRayTracingPipelineState* Pipeline = PipelineStateCache::GetAndOrCreateRayTracingPipelineState(Initializer);
 
 		FRayTracingShaderBindingsWriter GlobalResources;
 		GlobalResources.Set(TLASParameter, RayTracingScene.RayTracingSceneRHI->GetShaderResourceView());
@@ -497,7 +497,7 @@ void RenderRayTracingRectLightInternal(
 	RectLightData.dPdv = FVector(WorldToLight.M[0][2], WorldToLight.M[1][2], WorldToLight.M[2][2]);
 	RectLightData.Color = LightShaderParameters.Color / 2.0;
 
-	// #dxr_todo: Ray traced textured area lights are 1.5X brighter than those in lit mode.
+	// #dxr_todo: JIRA Ray traced textured area lights are 1.5X brighter than those in lit mode.
 	if (RectLightSceneProxy->HasSourceTexture())
 	{
 		RectLightData.Color *= 2.0 / 3.0;
diff --git a/Engine/Source/Runtime/Renderer/Private/RayTracing/RayTracingTranslucency.cpp b/Engine/Source/Runtime/Renderer/Private/RayTracing/RayTracingTranslucency.cpp
--- a/Engine/Source/Runtime/Renderer/Private/RayTracing/RayTracingTranslucency.cpp
+++ b/Engine/Source/Runtime/Renderer/Private/RayTracing/RayTracingTranslucency.cpp
@@ -258,7 +258,7 @@ void FDeferredShadingSceneRenderer::RenderRayTracingTranslucency(FRHICommandList
 		SCOPED_DRAW_EVENT(RHICmdList, RayTracingTranslucency);
 		SCOPED_GPU_STAT(RHICmdList, RayTracingTranslucency);
 
-		FRDGBuilder GraphBuilder(RHICmdList); //#dxr_todo: convert the entire translucency pass to render graph.
+		FRDGBuilder GraphBuilder(RHICmdList); 
 		FSceneRenderTargets& SceneContext = FSceneRenderTargets::Get(GraphBuilder.RHICmdList);
 
 		FSceneViewFamilyBlackboard SceneBlackboard;
diff --git a/Engine/Source/Runtime/Renderer/Private/RayTracing/RaytracingSkylight.cpp b/Engine/Source/Runtime/Renderer/Private/RayTracing/RaytracingSkylight.cpp
--- a/Engine/Source/Runtime/Renderer/Private/RayTracing/RaytracingSkylight.cpp
+++ b/Engine/Source/Runtime/Renderer/Private/RayTracing/RaytracingSkylight.cpp
@@ -204,7 +204,7 @@ public:
 		FRayTracingShaderRHIParamRef RayGenShaderTable[] = { GetRayTracingShader() };
 		Initializer.SetRayGenShaderTable(RayGenShaderTable);
 
-		FRHIRayTracingPipelineState* Pipeline = PipelineStateCache::GetAndOrCreateRayTracingPipelineState(Initializer); // #dxr_todo: this should be done once at load-time and cached
+		FRHIRayTracingPipelineState* Pipeline = PipelineStateCache::GetAndOrCreateRayTracingPipelineState(Initializer); 
 
 		FRayTracingShaderBindingsWriter GlobalResources;
 		GlobalResources.Set(TLASParameter, RayTracingScene.RayTracingSceneRHI->GetShaderResourceView());
diff --git a/Engine/Source/Runtime/Renderer/Private/ReflectionEnvironment.cpp b/Engine/Source/Runtime/Renderer/Private/ReflectionEnvironment.cpp
--- a/Engine/Source/Runtime/Renderer/Private/ReflectionEnvironment.cpp
+++ b/Engine/Source/Runtime/Renderer/Private/ReflectionEnvironment.cpp
@@ -135,7 +135,7 @@ static TAutoConsoleVariable<int32> CVarUseReflectionDenoiser(
 	TEXT("Choose the denoising algorithm.\n")
 	TEXT(" 0: Disabled;\n")
 	TEXT(" 1: Forces the default denoiser of the renderer;\n")
-	TEXT(" 2: GScreenSpaceDenoiser witch may be overriden by a third party plugin (default)."),
+	TEXT(" 2: GScreenSpaceDenoiser which may be overriden by a third party plugin (default)."),
 	ECVF_RenderThreadSafe);
 
 
@@ -708,7 +708,7 @@ void FDeferredShadingSceneRenderer::RenderDeferredReflectionsAndSkyLighting(FRHI
 	{
 		const FViewInfo& View = Views[ViewIndex];
 		bReflectionCapture = bReflectionCapture || View.bIsReflectionCapture;
-		//#dxr_todo: multiview case
+		//#dxr_todo: UE-72557 multiview case
 		bAnyViewWithRaytracingReflections = bAnyViewWithRaytracingReflections || (View.FinalPostProcessSettings.ReflectionsType == EReflectionsType::RayTracing);
 	}
 
