Reflection captures can now be sampled for the last refelction event in PT Reflection.
Disabled by default

diff --git a/Engine/Shaders/Private/RayTracing/RayTracingReflections.usf b/Engine/Shaders/Private/RayTracing/RayTracingReflections.usf
--- a/Engine/Shaders/Private/RayTracing/RayTracingReflections.usf
+++ b/Engine/Shaders/Private/RayTracing/RayTracingReflections.usf
@@ -45,9 +45,32 @@ float2 SobolIndexToUniformUnitSquare(uint2 SobolRand)
 #include "RayTracingSpotLight.ush"
 #include "RayTracingPointLight.ush"
 
+
+// Setup the reflection composite to sample the sky and reflection captures.
+#define REFLECTION_COMPOSITE_HAS_BOX_CAPTURES 1
+#define REFLECTION_COMPOSITE_HAS_SPHERE_CAPTURES 1
+#define REFLECTION_COMPOSITE_USE_BLENDED_REFLECTION_CAPTURES 1
+#define REFLECTION_COMPOSITE_SUPPORT_SKYLIGHT_BLEND 1
+// Enable skly light sampling as last fallback after reflection captures
+#ifdef ENABLE_SKY_LIGHT
+#undef ENABLE_SKY_LIGHT
+#endif
+#define ENABLE_SKY_LIGHT 1
+// Disable static lighting as this is not needed in this case.
+#ifdef ALLOW_STATIC_LIGHTING
+#undef ALLOW_STATIC_LIGHTING
+#endif
+#define ALLOW_STATIC_LIGHTING 0
+// Also specify that we are going to go over a range of reflection captures without culling available.
+// All reflection captures available are always available here even if completely outside the view.
+#define REFLECTION_COMPOSITE_NO_CULLING_DATA
+#include "../ReflectionEnvironmentComposite.ush"
+
+
 uint SamplesPerPixel;
 uint MaxBounces;
 uint HeightFog;
+uint UseReflectionCaptures;
 uint ShouldDoDirectLighting;
 uint ReflectedShadowsType;
 uint ShouldDoEmissiveAndIndirectLighting;
@@ -234,10 +257,12 @@ void RayTracingReflectionsRGS()
 			// We fix normal pointing away towards the geometry by assuming coated layer has smooth normal (assuming it is the mesh smooth normal).
 			// Fog is applied as a function of the top layer path, for both top and botto; refelctions combined (as an optimisation approximation).
 			//
-			for (uint BounceIndex = 0; BounceIndex < MaxBounces; ++BounceIndex)
+			uint BounceIndex = 0;
+			RayDesc TopLayerRay;
+			bool bSkyWasHit = false;
+			for (; BounceIndex < MaxBounces; ++BounceIndex)
 			{
 				// Prepare a ray to trace the top layer reflection
-				RayDesc TopLayerRay;
 				TopLayerRay.Origin = WorldPosition;
 				TopLayerRay.TMin = 0.0;
 				TopLayerRay.TMax = LocalMaxRayDistance;
@@ -300,7 +325,6 @@ void RayTracingReflectionsRGS()
 
 				// Compute some flags and values
 				bool bApplyHeightFog = false;
-				bool bSkyWasHit = false;
 				float3 OriginToCollider = 0.0f;
 				if (Payload.HitT >= 0)
 				{
@@ -464,6 +488,41 @@ void RayTracingReflectionsRGS()
 			
 			if (bIsValidSample)
 			{
+				if (UseReflectionCaptures && BounceIndex == MaxBounces && !(bSkyWasHit || HitDistance < 0.0f))
+				{
+					// We have reached the end of the bounce count on a surface. 
+					// So accumulate reflection from captures for the top layer for a better visual result.
+
+					// Ignoring any form of AO
+					float IndirectSpecularOcclusion = 1.0f;
+					float3 ExtraIndirectSpecular = 0;
+					float IndirectIrradiance = 0; 
+
+					// Not possible to use the screen culled grid. So going over every capture.
+					// dxr_todo This would need a spatial data structure.
+					uint NumCulledReflectionCaptures = ForwardLightData.NumReflectionCaptures;
+					uint ReflectionCapturesStartIndex = 0;
+
+					float3 R = reflect(TopLayerRay.Direction, TopLayerWorldNormal);
+					const float NoV = saturate(dot(-TopLayerRay.Direction, TopLayerWorldNormal));
+					const float AO = 1.0f;
+					const float RoughnessSq = TopLayerRoughness * TopLayerRoughness;
+					const float SpecularOcclusion = GetSpecularOcclusion(NoV, RoughnessSq, AO);
+
+					PathRadiance.rgb += EnvBRDF(TopLayerSpecularColor, TopLayerRoughness, NoV) * SpecularOcclusion * PathThroughput *
+						CompositeReflectionCapturesAndSkylight(
+							1.0,	// coverage of 1
+							WorldPosition,
+							R,
+							TopLayerRoughness,
+							IndirectIrradiance,
+							IndirectSpecularOcclusion,
+							ExtraIndirectSpecular,
+							NumCulledReflectionCaptures,
+							ReflectionCapturesStartIndex,
+							0);
+				}
+
 				ReflectedColor += float4(PathRadiance, 1.0f);
 				ReflectedWeight += 1.0;
 			}
diff --git a/Engine/Shaders/Private/ReflectionEnvironmentComposite.ush b/Engine/Shaders/Private/ReflectionEnvironmentComposite.ush
--- a/Engine/Shaders/Private/ReflectionEnvironmentComposite.ush
+++ b/Engine/Shaders/Private/ReflectionEnvironmentComposite.ush
@@ -33,7 +33,10 @@ float3 CompositeReflectionCapturesAndSkylight(
 			break;
 		}
 
-	uint CaptureIndex = 0;
+		uint CaptureIndex = 0;
+#ifdef REFLECTION_COMPOSITE_NO_CULLING_DATA
+		CaptureIndex = TileCaptureIndex;	// Go from 0 to NumCapturesAffectingTile as absolute index in capture array
+#else
 	#if INSTANCED_STEREO
 		BRANCH
 		if (EyeIndex == 0)
@@ -49,6 +52,7 @@ float3 CompositeReflectionCapturesAndSkylight(
 			CaptureIndex = ForwardLightDataISR.CulledLightDataGrid[CaptureDataStartIndex + TileCaptureIndex];
 		}
 	#endif
+#endif
 
 			
 		float4 CapturePositionAndRadius = ReflectionCapture.PositionAndRadius[CaptureIndex];
diff --git a/Engine/Source/Runtime/Renderer/Private/RayTracing/RayTracingReflections.cpp b/Engine/Source/Runtime/Renderer/Private/RayTracing/RayTracingReflections.cpp
--- a/Engine/Source/Runtime/Renderer/Private/RayTracing/RayTracingReflections.cpp
+++ b/Engine/Source/Runtime/Renderer/Private/RayTracing/RayTracingReflections.cpp
@@ -61,6 +61,13 @@ static FAutoConsoleVariableRef CVarRayTracingReflectionsShadows(
 	TEXT(" 2: Soft area shadows")
 );
 
+static int32 GRayTracingReflectionsCaptures = 0;
+static FAutoConsoleVariableRef CVarRayTracingReflectionsCaptures(
+	TEXT("r.RayTracing.Reflections.ReflectionCaptures"),
+	GRayTracingReflectionsCaptures,
+	TEXT("Enables ray tracing reflections to use reflection captures as the last bounce reflection. Particularly usefull for metals in reflection. (default = 0)")
+);
+
 static float GRayTracingReflectionsMinRayDistance = -1;
 static FAutoConsoleVariableRef CVarRayTracingReflectionsMinRayDistance(
 	TEXT("r.RayTracing.Reflections.MinRayDistance"),
@@ -119,6 +126,7 @@ class FRayTracingReflectionsRGS : public FGlobalShader
 		SHADER_PARAMETER(int32, SamplesPerPixel)
 		SHADER_PARAMETER(int32, MaxBounces)
 		SHADER_PARAMETER(int32, HeightFog)
+		SHADER_PARAMETER(int32, UseReflectionCaptures)
 		SHADER_PARAMETER(int32, ShouldDoDirectLighting)
 		SHADER_PARAMETER(int32, ReflectedShadowsType)
 		SHADER_PARAMETER(int32, ShouldDoEmissiveAndIndirectLighting)
@@ -138,6 +146,8 @@ class FRayTracingReflectionsRGS : public FGlobalShader
 		SHADER_PARAMETER_STRUCT_REF(FRaytracingLightDataPacked, LightDataPacked)
 		SHADER_PARAMETER_STRUCT_REF(FReflectionUniformParameters, ReflectionStruct)
 		SHADER_PARAMETER_STRUCT_REF(FFogUniformParameters, FogUniformParameters)
+		SHADER_PARAMETER_STRUCT_REF(FReflectionCaptureShaderData, ReflectionCapture)
+		SHADER_PARAMETER_STRUCT_REF(FForwardLightData, Forward)
 		SHADER_PARAMETER_STRUCT_REF(FIESLightProfileParameters, IESLightProfileParameters)
 
 		// Optional indirection buffer used for sorted materials
@@ -263,6 +273,7 @@ void FDeferredShadingSceneRenderer::RenderRayTracingReflections(
 	CommonParameters.SamplesPerPixel = SamplePerPixel;
 	CommonParameters.MaxBounces = GRayTracingReflectionsMaxBounces > -1? GRayTracingReflectionsMaxBounces : View.FinalPostProcessSettings.RayTracingReflectionsMaxBounces;
 	CommonParameters.HeightFog = HeightFog;
+	CommonParameters.UseReflectionCaptures = GRayTracingReflectionsCaptures;
 	CommonParameters.ShouldDoDirectLighting = GRayTracingReflectionsDirectLighting;
 	CommonParameters.ReflectedShadowsType = GRayTracingReflectionsShadows > -1 ? GRayTracingReflectionsShadows : (int32)View.FinalPostProcessSettings.RayTracingReflectionsShadows;
 	CommonParameters.ShouldDoEmissiveAndIndirectLighting = GRayTracingReflectionsEmissiveAndIndirectLighting;
@@ -285,6 +296,8 @@ void FDeferredShadingSceneRenderer::RenderRayTracingReflections(
 	CommonParameters.RayHitDistanceOutput = GraphBuilder.CreateUAV(*OutRayHitDistanceTexture);
 	CommonParameters.RayImaginaryDepthOutput = GraphBuilder.CreateUAV(*OutRayImaginaryDepthTexture);
 	CommonParameters.SortTileSize = SortTileSize;
+	CommonParameters.ReflectionCapture = View.ReflectionCaptureUniformBuffer;
+	CommonParameters.Forward = View.ForwardLightingResources->ForwardLightDataUniformBuffer;
 
 	for (uint32 PassIndex = 0; PassIndex < NumPasses; ++PassIndex)
 	{
