DR - Temp fix for SM4 failing to boot
#rnx

diff --git a/Engine/Shaders/Private/EyeAdaptationCommon.ush b/Engine/Shaders/Private/EyeAdaptationCommon.ush
--- a/Engine/Shaders/Private/EyeAdaptationCommon.ush
+++ b/Engine/Shaders/Private/EyeAdaptationCommon.ush
@@ -17,6 +17,11 @@ float EyeAdaptationLookup()
 {
 	return EyeAdaptation.Load(int3(0, 0, 0)).x;
 }
+#else
+float EyeAdaptationLookup()
+{
+	return 1;
+}
 #endif
 
 #endif // __EYE_ADAPTATION_COMMON__
