Fixed shader compilation errors when using DXC. All shaders should now compile successfully, however the output does not always match FXC.

diff --git a/Engine/Shaders/Private/DiaphragmDOF/DOFCocTileDilate.usf b/Engine/Shaders/Private/DiaphragmDOF/DOFCocTileDilate.usf
--- a/Engine/Shaders/Private/DiaphragmDOF/DOFCocTileDilate.usf
+++ b/Engine/Shaders/Private/DiaphragmDOF/DOFCocTileDilate.usf
@@ -91,7 +91,7 @@ void CocDilateMainCS(uint2 DispatchThreadId : SV_DispatchThreadID)
 		OutSampleBuckets[BucketId].BgdMinIntersectableCoc = EXTREMELY_LARGE_COC_RADIUS;
 		OutSampleBuckets[BucketId].BgdMinConservativeCloserGoemetryDistance = EXTREMELY_LARGE_COC_RADIUS;
 
-		UNROLL
+		UNROLL_N(DIM_DILATE_RADIUS*4)
 		for (uint BatchId = 0; BatchId < 4 * RingDistance; BatchId++)
 		{
 			// Generates at compile time sample location if square shaped ring in sample count distance from the kernel's center.
diff --git a/Engine/Shaders/Private/DiaphragmDOF/DOFGatherKernel.ush b/Engine/Shaders/Private/DiaphragmDOF/DOFGatherKernel.ush
--- a/Engine/Shaders/Private/DiaphragmDOF/DOFGatherKernel.ush
+++ b/Engine/Shaders/Private/DiaphragmDOF/DOFGatherKernel.ush
@@ -792,7 +792,7 @@ uint GatherSamplesAndResolve(
 			}
 		}
 
-		UNROLL
+		UNROLL_N(5)
 		for (; RingId >= 0; RingId--)
 		{
 			GatherRing(
diff --git a/Engine/Shaders/Private/DiaphragmDOF/DOFGatherPass.usf b/Engine/Shaders/Private/DiaphragmDOF/DOFGatherPass.usf
--- a/Engine/Shaders/Private/DiaphragmDOF/DOFGatherPass.usf
+++ b/Engine/Shaders/Private/DiaphragmDOF/DOFGatherPass.usf
@@ -393,7 +393,7 @@ void GatherMainCS(
 		#if DEBUG_FAST_ACCUMULATOR
 		{
 			#if CONFIG_KERNEL_SAMPLING_DENSITY_DRIVER > 0
-			UNROLL
+			UNROLL_N(CONFIG_MAX_SAMPLING_DENSITY_CHANGES)
 			for (uint i = 0; i < CONFIG_MAX_SAMPLING_DENSITY_CHANGES; i++)
 			{
 				if (i < DebugGatherDensityChanges)
diff --git a/Engine/Shaders/Private/LandscapeGrassWeight.usf b/Engine/Shaders/Private/LandscapeGrassWeight.usf
--- a/Engine/Shaders/Private/LandscapeGrassWeight.usf
+++ b/Engine/Shaders/Private/LandscapeGrassWeight.usf
@@ -55,14 +55,15 @@ int OutputPass;
 
 // Pixel Shader
 void PSMain( 
-	FLandscapeGrassWeightInterpolantsVSToPS Interpolants,
-	in float4 SvPosition : SV_Position
+	FLandscapeGrassWeightInterpolantsVSToPS Interpolants
 	OPTIONAL_IsFrontFace,
 	out half4 OutColor	: SV_Target0
 	)
 {
 	ResolvedView = ResolveView();
 
+	float4 SvPosition = Interpolants.Position;
+
 	FMaterialPixelParameters MaterialParameters = GetMaterialPixelParameters(Interpolants.FactoryInterpolants, SvPosition);
 	FPixelMaterialInputs PixelMaterialInputs;
 	CalcMaterialParameters(MaterialParameters, PixelMaterialInputs, SvPosition, bIsFrontFace); 
diff --git a/Engine/Shaders/Private/RayTracing/MaterialSort.usf b/Engine/Shaders/Private/RayTracing/MaterialSort.usf
--- a/Engine/Shaders/Private/RayTracing/MaterialSort.usf
+++ b/Engine/Shaders/Private/RayTracing/MaterialSort.usf
@@ -53,7 +53,7 @@ void MaterialSortLocal(uint3 DispatchThread : SV_DispatchThreadID, uint3 GroupTh
 	uint Slot[NUM_ELEMENTS / NUM_THREADS];
 
 	//get counts
-	[unroll]
+	[unroll(NUM_ELEMENTS/NUM_THREADS)]
 	for (int i = GroupThread.x; i < NUM_ELEMENTS ; i += NUM_THREADS)
 	{
 		uint Index = GroupBase + i;
@@ -83,7 +83,7 @@ void MaterialSortLocal(uint3 DispatchThread : SV_DispatchThreadID, uint3 GroupTh
 	GroupMemoryBarrierWithGroupSync();
 	
 	//write data
-	[unroll]
+	[unroll(NUM_ELEMENTS/NUM_THREADS)]
 	for (int i = GroupThread.x; i < NUM_ELEMENTS; i += NUM_THREADS)
 	{
 		uint Index = GroupBase + i;
diff --git a/Engine/Shaders/Private/VolumetricFog.usf b/Engine/Shaders/Private/VolumetricFog.usf
--- a/Engine/Shaders/Private/VolumetricFog.usf
+++ b/Engine/Shaders/Private/VolumetricFog.usf
@@ -192,10 +192,11 @@ float InverseSquaredLightDistanceBiasScale;
 
 void InjectShadowedLocalLightPS(
 	FWriteToSliceGeometryOutput Input,
-	in float4 SvPosition : SV_Position,
 	out float4 OutScattering : SV_Target0
 	)
 {
+	float4 SvPosition = Input.Vertex.Position;
+
 	OutScattering = 0;
 
 	uint3 GridCoordinate = uint3(SvPosition.xy, Input.LayerIndex);
diff --git a/Engine/Shaders/Public/Platform.ush b/Engine/Shaders/Public/Platform.ush
--- a/Engine/Shaders/Public/Platform.ush
+++ b/Engine/Shaders/Public/Platform.ush
@@ -439,6 +439,7 @@
 
 /** Avoids flow control constructs. */
 #define UNROLL [unroll]
+#define UNROLL_N(N) [unroll(N)]
 
 /** Gives preference to flow control constructs. */
 #define LOOP [loop]
@@ -465,6 +466,10 @@
 #define UNROLL
 #endif
 
+#ifndef UNROLL_N
+#define UNROLL_N(N)
+#endif
+
 #ifndef LOOP
 #define LOOP
 #endif
