Removed duplicated functions in several ray tracing files (GetHitT, IsHit, ReconstructWorldPositionFromDepth)
The only functional change is that RTGI now takes into account TMin/TMax properly in IsHit().

diff --git a/Engine/Shaders/Private/PathTracing/PathTracing.usf b/Engine/Shaders/Private/PathTracing/PathTracing.usf
--- a/Engine/Shaders/Private/PathTracing/PathTracing.usf
+++ b/Engine/Shaders/Private/PathTracing/PathTracing.usf
@@ -7,7 +7,7 @@ PathTracingRayGenShader.usf: Reference path tracing
 #include "../Common.ush"
 #include "../PostProcessCommon.ush"
 #include "../RectLight.ush"
-//#include "../Montecarlo.ush"
+#include "../RayTracing/RayTracingCommon.ush"
 
 #include "PathTracingCommon.ush"
 #include "../RayTracing/RayTracingHitGroupCommon.ush"
@@ -25,17 +25,6 @@ RWTexture2D<uint>	RayCountPerPixelRT;
 RaytracingAccelerationStructure TLAS;
 
 
-float GetHitT(FMaterialClosestHitPayload HitInfo)
-{
-	return HitInfo.HitT;
-}
-
-bool IsHit(RayDesc Ray, FMaterialClosestHitPayload HitInfo)
-{
-	float t = HitInfo.HitT;
-	return t > Ray.TMin && t < Ray.TMax;
-}
-
 [shader("raygeneration")]
 void PathTracingMainRG()
 {
diff --git a/Engine/Shaders/Private/PathTracing/PathTracingCHS.usf b/Engine/Shaders/Private/PathTracing/PathTracingCHS.usf
--- a/Engine/Shaders/Private/PathTracing/PathTracingCHS.usf
+++ b/Engine/Shaders/Private/PathTracing/PathTracingCHS.usf
@@ -1,5 +1,6 @@
 // Copyright 1998-2019 Epic Games, Inc. All Rights Reserved.
 
+#include "../Common.ush"
 #include "PathTracingDefinitions.ush"
 #include "../RayTracing/RayTracingCommon.ush"
 #include "../RayTracing/RayTracingHitGroupCommon.ush"
diff --git a/Engine/Shaders/Private/PathTracing/PathTracingCommon.ush b/Engine/Shaders/Private/PathTracing/PathTracingCommon.ush
--- a/Engine/Shaders/Private/PathTracing/PathTracingCommon.ush
+++ b/Engine/Shaders/Private/PathTracing/PathTracingCommon.ush
@@ -12,14 +12,6 @@
 #include "../RayTracing/RayTracingCommon.ush"
 #include "PathTracingDefinitions.ush"
 
-float3 ReconstructWorldPositionFromDepth(float2 UV, float Depth)
-{
-	float2 ScreenPosition = (UV - View.ScreenPositionScaleBias.wz) / View.ScreenPositionScaleBias.xy;
-	float4 HomogeneousWorldPosition = mul(float4(ScreenPosition * Depth, Depth, 1), View.ScreenToWorld);
-	float3 WorldPosition = HomogeneousWorldPosition.xyz / HomogeneousWorldPosition.w;
-
-	return WorldPosition;
-}
 
 #endif // PATHTRACINGCOMMON_USH_INCLUDED // Workarround for UE-66460
 
diff --git a/Engine/Shaders/Private/PathTracing/PathTracingMS.usf b/Engine/Shaders/Private/PathTracing/PathTracingMS.usf
--- a/Engine/Shaders/Private/PathTracing/PathTracingMS.usf
+++ b/Engine/Shaders/Private/PathTracing/PathTracingMS.usf
@@ -1,5 +1,6 @@
 // Copyright 1998-2019 Epic Games, Inc. All Rights Reserved.
 
+#include "../Common.ush"
 #include "PathTracingDefinitions.ush"
 #include "../RayTracing/RayTracingCommon.ush"
 #include "../RayTracing/RayTracingHitGroupCommon.ush"
diff --git a/Engine/Shaders/Private/RayTracing/MaterialSort.usf b/Engine/Shaders/Private/RayTracing/MaterialSort.usf
--- a/Engine/Shaders/Private/RayTracing/MaterialSort.usf
+++ b/Engine/Shaders/Private/RayTracing/MaterialSort.usf
@@ -1,3 +1,4 @@
+#include "../Common.ush"
 #include "RayTracingCommon.ush"
 
 #if DIM_SORT_SIZE == 0
diff --git a/Engine/Shaders/Private/RayTracing/RayTracingAmbientOcclusionRGS.usf b/Engine/Shaders/Private/RayTracing/RayTracingAmbientOcclusionRGS.usf
--- a/Engine/Shaders/Private/RayTracing/RayTracingAmbientOcclusionRGS.usf
+++ b/Engine/Shaders/Private/RayTracing/RayTracingAmbientOcclusionRGS.usf
@@ -15,16 +15,6 @@ uint CalcLinearIndex(uint2 PixelCoord)
 	return PixelCoord.y * View.BufferSizeAndInvSize.x + PixelCoord.x;
 }
 
-// #patrick todo: move to RayTracingCommon.ush
-float3 ReconstructWorldPositionFromDepth(float2 UV, float Depth)
-{
-	float2 ScreenPosition = (UV - View.ScreenPositionScaleBias.wz) / View.ScreenPositionScaleBias.xy;
-	float4 HomogeneousWorldPosition = mul(float4(ScreenPosition * Depth, Depth, 1), View.ScreenToWorld);
-	float3 WorldPosition = HomogeneousWorldPosition.xyz / HomogeneousWorldPosition.w;
-
-	return WorldPosition;
-}
-
 void GenerateCosineNormalRay(
 	float3 WorldPosition,
 	float3 WorldNormal,
diff --git a/Engine/Shaders/Private/RayTracing/RayTracingBarycentrics.usf b/Engine/Shaders/Private/RayTracing/RayTracingBarycentrics.usf
--- a/Engine/Shaders/Private/RayTracing/RayTracingBarycentrics.usf
+++ b/Engine/Shaders/Private/RayTracing/RayTracingBarycentrics.usf
@@ -6,15 +6,6 @@
 RWTexture2D<float4> Output;
 RaytracingAccelerationStructure TLAS;
 
-float3 ReconstructWorldPositionFromDepth(float2 UV, float Depth)
-{
-	float2 ScreenPosition = (UV - View.ScreenPositionScaleBias.wz) / View.ScreenPositionScaleBias.xy;
-	float4 HomogeneousWorldPosition = mul(float4(ScreenPosition * Depth, Depth, 1), View.ScreenToWorld);
-	float3 WorldPosition = HomogeneousWorldPosition.xyz / HomogeneousWorldPosition.w;
-
-	return WorldPosition;
-}
-
 struct FBarycentricsPayload
 {
 	float2 Barycentrics;
diff --git a/Engine/Shaders/Private/RayTracing/RayTracingBuiltInShaders.usf b/Engine/Shaders/Private/RayTracing/RayTracingBuiltInShaders.usf
--- a/Engine/Shaders/Private/RayTracing/RayTracingBuiltInShaders.usf
+++ b/Engine/Shaders/Private/RayTracing/RayTracingBuiltInShaders.usf
@@ -1,5 +1,6 @@
 // Copyright 1998-2019 Epic Games, Inc. All Rights Reserved.
 
+#include "../Common.ush"
 #include "RayTracingCommon.ush"
 
 #define RT_CONCATENATE2(a, b) a##b
diff --git a/Engine/Shaders/Private/RayTracing/RayTracingCommon.ush b/Engine/Shaders/Private/RayTracing/RayTracingCommon.ush
--- a/Engine/Shaders/Private/RayTracing/RayTracingCommon.ush
+++ b/Engine/Shaders/Private/RayTracing/RayTracingCommon.ush
@@ -75,6 +75,15 @@ FRayCone PropagateRayCone(in FRayCone Cone, in float SurfaceSpreadAngle, in floa
 	return NewCone;
 }
 
+float3 ReconstructWorldPositionFromDepth(float2 UV, float Depth)
+{
+	float2 ScreenPosition = (UV - View.ScreenPositionScaleBias.wz) / View.ScreenPositionScaleBias.xy;
+	float4 HomogeneousWorldPosition = mul(float4(ScreenPosition * Depth, Depth, 1), View.ScreenToWorld);
+	float3 WorldPosition = HomogeneousWorldPosition.xyz / HomogeneousWorldPosition.w;
+
+	return WorldPosition;
+}
+
 #define RAY_TRACING_BLEND_MODE_OPAQUE				0
 #define RAY_TRACING_BLEND_MODE_ALPHA_COMPOSITE		1
 #define RAY_TRACING_BLEND_MODE_TRANSLUCENT			2
@@ -263,6 +272,24 @@ void ApplyPositionBias(inout RayDesc Ray, const float3 WorldNormal, const float
 
 	Ray.Origin += WorldNormal * NormalBias;
 }
+
+float GetHitT(FMaterialClosestHitPayload HitInfo)
+{
+	return HitInfo.HitT;
+}
+
+bool IsHit(RayDesc Ray, FMaterialClosestHitPayload HitInfo)
+{
+	float t = HitInfo.HitT;
+	return t > Ray.TMin && t < Ray.TMax;
+}
+
+bool IsHit(RayDesc Ray, FPackedMaterialClosestHitPayload HitInfo)
+{
+	float t = HitInfo.HitT;
+	return t > Ray.TMin && t < Ray.TMax;
+}
+
 #endif // !COMPUTESHADER
 
 #endif // RAYTRACINGCOMMON_USH_INCLUDED // Workarround for UE-66460
diff --git a/Engine/Shaders/Private/RayTracing/RayTracingDebug.usf b/Engine/Shaders/Private/RayTracing/RayTracingDebug.usf
--- a/Engine/Shaders/Private/RayTracing/RayTracingDebug.usf
+++ b/Engine/Shaders/Private/RayTracing/RayTracingDebug.usf
@@ -12,15 +12,6 @@ RaytracingAccelerationStructure TLAS;
 
 uint VisualizationMode;
 
-float3 ReconstructWorldPositionFromDepth(float2 UV, float Depth)
-{
-	float2 ScreenPosition = (UV - View.ScreenPositionScaleBias.wz) / View.ScreenPositionScaleBias.xy;
-	float4 HomogeneousWorldPosition = mul(float4(ScreenPosition * Depth, Depth, 1), View.ScreenToWorld);
-	float3 WorldPosition = HomogeneousWorldPosition.xyz / HomogeneousWorldPosition.w;
-
-	return WorldPosition;
-}
-
 [shader("raygeneration")]
 void RayTracingDebugMainRGS()
 {
diff --git a/Engine/Shaders/Private/RayTracing/RayTracingDeferredMaterials.usf b/Engine/Shaders/Private/RayTracing/RayTracingDeferredMaterials.usf
--- a/Engine/Shaders/Private/RayTracing/RayTracingDeferredMaterials.usf
+++ b/Engine/Shaders/Private/RayTracing/RayTracingDeferredMaterials.usf
@@ -1,5 +1,6 @@
 // Copyright 1998-2019 Epic Games, Inc. All Rights Reserved.
 
+#include "../Common.ush"
 #include "RayTracingCommon.ush"
 #include "RayTracingHitGroupCommon.ush"
 
diff --git a/Engine/Shaders/Private/RayTracing/RayTracingDirectionalLight.ush b/Engine/Shaders/Private/RayTracing/RayTracingDirectionalLight.ush
--- a/Engine/Shaders/Private/RayTracing/RayTracingDirectionalLight.ush
+++ b/Engine/Shaders/Private/RayTracing/RayTracingDirectionalLight.ush
@@ -7,16 +7,6 @@
 #include "../LightShaderParameters.ush"
 #include "RayTracingCommon.ush"
 
-// #patrick todo: move to RayTracingCommon.ush
-float3 ReconstructWorldPositionFromDepth(float2 UV, float Depth)
-{
-	float2 ScreenPosition = (UV - View.ScreenPositionScaleBias.wz) / View.ScreenPositionScaleBias.xy;
-	float4 HomogeneousWorldPosition = mul(float4(ScreenPosition * Depth, Depth, 1), View.ScreenToWorld);
-	float3 WorldPosition = HomogeneousWorldPosition.xyz / HomogeneousWorldPosition.w;
-
-	return WorldPosition;
-}
-
 // Adapted from "A Low Distortion Map Between Disk and Square." Improved based
 // on comments from Peter Shirley's blog post: "Improved code for concentric map."
 float2 ToConcentricMap(float2 RectangularCoords)
diff --git a/Engine/Shaders/Private/RayTracing/RayTracingGlobalIlluminationRGS.usf b/Engine/Shaders/Private/RayTracing/RayTracingGlobalIlluminationRGS.usf
--- a/Engine/Shaders/Private/RayTracing/RayTracingGlobalIlluminationRGS.usf
+++ b/Engine/Shaders/Private/RayTracing/RayTracingGlobalIlluminationRGS.usf
@@ -72,16 +72,6 @@ void GenerateCosineNormalRay(
 	RayPdf = Direction_Tangent.w;
 }
 
-float GetHitT(FMaterialClosestHitPayload HitInfo)
-{
-	return HitInfo.HitT;
-}
-
-bool IsHit(RayDesc Ray, FMaterialClosestHitPayload HitInfo)
-{
-	return HitInfo.HitT >= 0.0;
-}
-
 [shader("raygeneration")]
 void GlobalIlluminationRGS()
 {
diff --git a/Engine/Shaders/Private/RayTracing/RayTracingLightingCommon.ush b/Engine/Shaders/Private/RayTracing/RayTracingLightingCommon.ush
--- a/Engine/Shaders/Private/RayTracing/RayTracingLightingCommon.ush
+++ b/Engine/Shaders/Private/RayTracing/RayTracingLightingCommon.ush
@@ -7,22 +7,6 @@
 #define LIGHT_TYPE_RECT				3 
 #define LIGHT_TYPE_MAX				4 
 
-float GetHitT(FMaterialClosestHitPayload HitInfo)
-{
-	return HitInfo.HitT;
-}
-
-bool IsHit(RayDesc Ray, FMaterialClosestHitPayload HitInfo)
-{
-	float t = HitInfo.HitT;
-	return t > Ray.TMin && t < Ray.TMax;
-}
-
-bool IsHit(RayDesc Ray, FPackedMaterialClosestHitPayload HitInfo)
-{
-	float t = HitInfo.HitT;
-	return t > Ray.TMin && t < Ray.TMax;
-}
 
 float GetRoughnessFade(float Roughness, float MaxRoughness)
 {
diff --git a/Engine/Shaders/Private/RayTracing/RayTracingMaterialDefaultHitShaders.usf b/Engine/Shaders/Private/RayTracing/RayTracingMaterialDefaultHitShaders.usf
--- a/Engine/Shaders/Private/RayTracing/RayTracingMaterialDefaultHitShaders.usf
+++ b/Engine/Shaders/Private/RayTracing/RayTracingMaterialDefaultHitShaders.usf
@@ -1,5 +1,6 @@
 // Copyright 1998-2019 Epic Games, Inc. All Rights Reserved.
 
+#include "../Common.ush"
 #include "RayTracingCommon.ush"
 
 [shader("closesthit")]
diff --git a/Engine/Shaders/Private/RayTracing/RayTracingRectLightRGS.usf b/Engine/Shaders/Private/RayTracing/RayTracingRectLightRGS.usf
--- a/Engine/Shaders/Private/RayTracing/RayTracingRectLightRGS.usf
+++ b/Engine/Shaders/Private/RayTracing/RayTracingRectLightRGS.usf
@@ -18,16 +18,6 @@ uint CalcLinearIndex(uint2 PixelCoord)
 	return PixelCoord.y * View.BufferSizeAndInvSize.x + PixelCoord.x;
 }
 
-// #patrick todo: move to RayTracingCommon.ush
-float3 ReconstructWorldPositionFromDepth(float2 UV, float Depth)
-{
-	float2 ScreenPosition = (UV - View.ScreenPositionScaleBias.wz) / View.ScreenPositionScaleBias.xy;
-	float4 HomogeneousWorldPosition = mul(float4(ScreenPosition * Depth, Depth, 1), View.ScreenToWorld);
-	float3 WorldPosition = HomogeneousWorldPosition.xyz / HomogeneousWorldPosition.w;
-
-	return WorldPosition;
-}
-
 bool GenerateRectLightImportanceSampledRay(
 	float3 WorldPosition,
 	float3 WorldNormal,
diff --git a/Engine/Shaders/Private/RayTracing/RayTracingTest.usf b/Engine/Shaders/Private/RayTracing/RayTracingTest.usf
--- a/Engine/Shaders/Private/RayTracing/RayTracingTest.usf
+++ b/Engine/Shaders/Private/RayTracing/RayTracingTest.usf
@@ -1,5 +1,6 @@
 // Copyright 1998-2019 Epic Games, Inc. All Rights Reserved.
 
+#include "../Common.ush"
 #include "RayTracingCommon.ush"
 
 RaytracingAccelerationStructure  TLAS;
