ReconstructWorldPositionFromDepth() replaced with View.WorldCameraOrigin when appropriate.

diff --git a/Engine/Shaders/Private/RayTracing/RayTracingBarycentrics.usf b/Engine/Shaders/Private/RayTracing/RayTracingBarycentrics.usf
--- a/Engine/Shaders/Private/RayTracing/RayTracingBarycentrics.usf
+++ b/Engine/Shaders/Private/RayTracing/RayTracingBarycentrics.usf
@@ -22,7 +22,7 @@ void RayTracingBarycentricsMainRGS()
 
 	RayDesc Ray;
 
-	Ray.Origin = ReconstructWorldPositionFromDepth(RenderTargetUV, 0.0f);
+	Ray.Origin = View.WorldCameraOrigin;
 	Ray.Direction = normalize(ReconstructWorldPositionFromDepth(RenderTargetUV, 10000.f) - Ray.Origin);
 	Ray.TMin = 0.5f; // 1 unit in UE4 = 1 cm -0> 0.5 ==> 5 mm
 	Ray.TMax = 1e27f;
diff --git a/Engine/Shaders/Private/RayTracing/RayTracingDebug.usf b/Engine/Shaders/Private/RayTracing/RayTracingDebug.usf
--- a/Engine/Shaders/Private/RayTracing/RayTracingDebug.usf
+++ b/Engine/Shaders/Private/RayTracing/RayTracingDebug.usf
@@ -21,7 +21,7 @@ void RayTracingDebugMainRGS()
 
 	RayDesc Ray;
 
-	Ray.Origin = ReconstructWorldPositionFromDepth(RenderTargetUV, 0.0f);
+	Ray.Origin = View.WorldCameraOrigin;
 	Ray.Direction = normalize(ReconstructWorldPositionFromDepth(RenderTargetUV, 10000.f) - Ray.Origin);
 	Ray.TMin = 0.5f; // 1 unit in UE4 = 1 cm -0> 0.5 ==> 5 mm
 	Ray.TMax = 1e27f;
diff --git a/Engine/Shaders/Private/RayTracing/RayTracingGlobalIlluminationRGS.usf b/Engine/Shaders/Private/RayTracing/RayTracingGlobalIlluminationRGS.usf
--- a/Engine/Shaders/Private/RayTracing/RayTracingGlobalIlluminationRGS.usf
+++ b/Engine/Shaders/Private/RayTracing/RayTracingGlobalIlluminationRGS.usf
@@ -97,7 +97,7 @@ void GlobalIlluminationRGS()
 	}
 	float Depth = ScreenSpaceData.GBuffer.Depth;
 	float3 WorldPosition = ReconstructWorldPositionFromDepth(UV, Depth);
-	float3 CameraOrigin = ReconstructWorldPositionFromDepth(UV, 0.0);
+	float3 CameraOrigin = View.WorldCameraOrigin;
 	float3 CameraDirection = normalize(WorldPosition - CameraOrigin);
 	float3 WorldNormal = ScreenSpaceData.GBuffer.WorldNormal;
 	uint ShadingModelID = ScreenSpaceData.GBuffer.ShadingModelID;
diff --git a/Engine/Shaders/Private/RayTracing/RayTracingRectLightRGS.usf b/Engine/Shaders/Private/RayTracing/RayTracingRectLightRGS.usf
--- a/Engine/Shaders/Private/RayTracing/RayTracingRectLightRGS.usf
+++ b/Engine/Shaders/Private/RayTracing/RayTracingRectLightRGS.usf
@@ -171,7 +171,7 @@ void RectLightRGS()
 	FGBufferData GBuffer = ScreenSpaceData.GBuffer;
 	float Depth = ScreenSpaceData.GBuffer.Depth;
 	float3 WorldPosition = ReconstructWorldPositionFromDepth(UV, Depth);
-	float3 CameraOrigin = ReconstructWorldPositionFromDepth(UV, 0.0f);
+	float3 CameraOrigin = View.WorldCameraOrigin;
 	float3 CameraDirection = normalize(WorldPosition - CameraOrigin);
 	float3 WorldNormal = ScreenSpaceData.GBuffer.WorldNormal;
 	float3 BaseColor = ScreenSpaceData.GBuffer.BaseColor;
diff --git a/Engine/Shaders/Private/RayTracing/RayTracingReflections.usf b/Engine/Shaders/Private/RayTracing/RayTracingReflections.usf
--- a/Engine/Shaders/Private/RayTracing/RayTracingReflections.usf
+++ b/Engine/Shaders/Private/RayTracing/RayTracingReflections.usf
@@ -220,7 +220,7 @@ void RayTracingReflectionsRGS()
 			float3 PathThroughput = float3(1.0f, 1.0f, 1.0f);
 			
 			// Initialize path states according to GBuffer
-			float3 CameraOrigin = ReconstructWorldPositionFromDepth(UV, 0.0f);
+			float3 CameraOrigin = View.WorldCameraOrigin;
 			float3 WorldPosition = ReconstructWorldPositionFromDepth(UV, Depth);
 			float3 CameraDirection = normalize(WorldPosition - CameraOrigin);
 			float3 IncidentDirection = CameraDirection;
diff --git a/Engine/Shaders/Private/RayTracing/RayTracingSkyLightRGS.usf b/Engine/Shaders/Private/RayTracing/RayTracingSkyLightRGS.usf
--- a/Engine/Shaders/Private/RayTracing/RayTracingSkyLightRGS.usf
+++ b/Engine/Shaders/Private/RayTracing/RayTracingSkyLightRGS.usf
@@ -242,7 +242,7 @@ void SkyLightRGS()
 	FGBufferData GBuffer = ScreenSpaceData.GBuffer;
 	float Depth = ScreenSpaceData.GBuffer.Depth;
 	float3 WorldPosition = ReconstructWorldPositionFromDepth(UV, Depth);
-	float3 CameraOrigin = ReconstructWorldPositionFromDepth(UV, 0.0f);
+	float3 CameraOrigin = View.WorldCameraOrigin;
 	float3 CameraDirection = normalize(WorldPosition - CameraOrigin);
 	float3 WorldNormal = ScreenSpaceData.GBuffer.WorldNormal;
 	float3 BaseColor = ScreenSpaceData.GBuffer.BaseColor;
diff --git a/Engine/Shaders/Private/RayTracing/RayTracingTranslucency.usf b/Engine/Shaders/Private/RayTracing/RayTracingTranslucency.usf
--- a/Engine/Shaders/Private/RayTracing/RayTracingTranslucency.usf
+++ b/Engine/Shaders/Private/RayTracing/RayTracingTranslucency.usf
@@ -125,7 +125,7 @@ void RayTracingTranslucencyRGS()
 	float2 UV = (float2(PixelCoord) + 0.5) * InvBufferSize;
 	FScreenSpaceData ScreenSpaceData = GetScreenSpaceData(UV);
 	float Depth = ScreenSpaceData.GBuffer.Depth;
-	float3 CameraOrigin = ReconstructWorldPositionFromDepth(UV, 0.0f);
+	float3 CameraOrigin = View.WorldCameraOrigin;
 	float3 WorldPosition = ReconstructWorldPositionFromDepth(UV, Depth);
 	float3 CameraDirection = normalize(WorldPosition - CameraOrigin);	
 	const float GBufferHitT = length(WorldPosition - CameraOrigin);
