Fixed path tracing does not respect the object cast shadows flag

diff --git a/Engine/Shaders/Private/PathTracing/PathTracing.usf b/Engine/Shaders/Private/PathTracing/PathTracing.usf
--- a/Engine/Shaders/Private/PathTracing/PathTracing.usf
+++ b/Engine/Shaders/Private/PathTracing/PathTracing.usf
@@ -178,7 +178,7 @@ void PathTracingMainRG()
 					FMaterialClosestHitPayload NeePayload = TraceRayInternal(
 						TLAS,      // AccelerationStructure
 						RayFlags,
-						RAY_TRACING_MASK_ALL,
+						RAY_TRACING_MASK_SHADOW,
 						RAY_TRACING_SHADER_SLOT_MATERIAL, // RayContributionToHitGroupIndex
 						RAY_TRACING_NUM_SHADER_SLOTS,     // MultiplierForGeometryContributionToShaderIndex
 						0,         // MissShaderIndex
@@ -254,7 +254,6 @@ void PathTracingMainRG()
 			float EnergyThreshold = 1.0; 
 			if (PathEnergy < EnergyThreshold && Bounces > 4)
 			{
-
 				float Q = 1.0 - PathEnergy;
 				if (RandomSequence_GenerateSample1D(RandSequence, SampleIndex) < Q)
 				{
