Ray tracing helper functions now take parameters (TLAS, shadows and lighting flags) explicitly rather than rely on a global declared

diff --git a/Engine/Shaders/Private/RayTracing/RayTracingLightingCommon.ush b/Engine/Shaders/Private/RayTracing/RayTracingLightingCommon.ush
--- a/Engine/Shaders/Private/RayTracing/RayTracingLightingCommon.ush
+++ b/Engine/Shaders/Private/RayTracing/RayTracingLightingCommon.ush
@@ -50,7 +50,8 @@ float3 GenerateReflectedRayDirection(
 
 float TraceShadowRay(
 	in RayDesc Ray,
-	in FRayCone ShadowRayCone)
+	in FRayCone ShadowRayCone,
+	in RaytracingAccelerationStructure TLAS)
 {
 	uint RayFlags = RAY_FLAG_CULL_BACK_FACING_TRIANGLES | RAY_FLAG_ACCEPT_FIRST_HIT_AND_END_SEARCH | RAY_FLAG_SKIP_CLOSEST_HIT_SHADER;
 
@@ -193,7 +194,6 @@ uint2 CullDirectLighting(
 	in bool bReflectedShadows,
 	in float ShadowMaxNormalBias)
 {
-
 	uint2 LightCullMask = (uint2)0;
 
 	for (uint DirIndex = 0; DirIndex < RaytracingLightsDataPacked.Count; DirIndex++)
@@ -275,9 +275,10 @@ float3 ComputeDirectLightingCulled(
 	in uint2 LightCullMask,
 	in RayDesc Ray,
 	in FRayCone RayCone,
+	in RaytracingAccelerationStructure TLAS,
 	in FMaterialClosestHitPayload Payload,
 	in RandomSequence RandSequence,
-	in bool bReflectedShadows,
+	in uint ReflectedShadowsType,
 	in float ShadowMaxNormalBias)
 {
 	float3 DirectLighting = float3(0.0, 0.0, 0.0);
@@ -417,7 +418,7 @@ float3 ComputeDirectLightingCulled(
 			LightMask = GetLocalLightAttenuation(WorldPosition, LightData, ToLight, ShadowRayDirection);
 		}
 
-		if (bReflectedShadows)
+		if (ReflectedShadowsType > 0)
 		{
 			if (ReflectedShadowsType == 2)
 			{
@@ -433,7 +434,7 @@ float3 ComputeDirectLightingCulled(
 			ShadowRay.TMax = Active ? length(ToLight) : ShadowRay.TMin;
 			ApplyPositionBias(ShadowRay, Payload.WorldNormal, ShadowMaxNormalBias);
 
-			float HitT = TraceShadowRay(ShadowRay, ShadowRayCone);
+			float HitT = TraceShadowRay(ShadowRay, ShadowRayCone, TLAS);
 			AmbientOcclusion = HitT < 0.0f && Active;
 		}
 
@@ -482,6 +483,7 @@ void ComputeBottomLayerMaterialProperties(RayDesc Ray, inout FMaterialClosestHit
 
 FMaterialClosestHitPayload TraceRayResults(
 	in RayDesc Ray,
+	in RaytracingAccelerationStructure TLAS,
 	in RandomSequence RandSequence,
 	in uint RayFlags,
 	in float ShadowMaxNormalBias,
@@ -506,10 +508,14 @@ FMaterialClosestHitPayload TraceRayResults(
 void AccumulateResults(
 	in FMaterialClosestHitPayload Payload,
 	in RayDesc Ray,
+	in RaytracingAccelerationStructure TLAS,
 	in RandomSequence RandSequence,
 	in uint RayFlags,
 	in float ShadowMaxNormalBias,
 	in uint Mask,
+	in uint ReflectedShadowsType,
+	in uint ShouldDoDirectLighting,
+	in uint ShouldDoEmissiveAndIndirectLighting,
 	inout FRayCone RayCone,
 	inout float3 Radiance)
 {
@@ -524,7 +530,7 @@ void AccumulateResults(
 
 		if (ShouldDoDirectLighting)
 		{
-			float3 DirectLighting = ComputeDirectLightingCulled(LightCullMask, Ray, RayCone, Payload, RandSequence, ReflectedShadowsType > 0, ShadowMaxNormalBias);
+			float3 DirectLighting = ComputeDirectLightingCulled(LightCullMask, Ray, RayCone, TLAS, Payload, RandSequence, ReflectedShadowsType, ShadowMaxNormalBias);
 
 			// Transform NaNs to black, transform negative colors to black.
 			DirectLighting = -min(-DirectLighting, 0.0);
@@ -542,15 +548,20 @@ void AccumulateResults(
 
 FMaterialClosestHitPayload TraceRayAndAccumulateResults(
 	in RayDesc Ray,
+	in RaytracingAccelerationStructure TLAS, 
 	in RandomSequence RandSequence,
 	in uint RayFlags,
 	in float ShadowMaxNormalBias,
 	in uint Mask,
+	in uint ReflectedShadowsType,
+	in uint ShouldDoDirectLighting,
+	in uint ShouldDoEmissiveAndIndirectLighting,
 	inout FRayCone RayCone,
 	inout float3 Radiance)
 {
 	FMaterialClosestHitPayload Payload = TraceRayResults(
 		Ray,
+		TLAS,
 		RandSequence,
 		RayFlags,
 		ShadowMaxNormalBias,
@@ -561,10 +572,14 @@ FMaterialClosestHitPayload TraceRayAndAccumulateResults(
 	AccumulateResults(
 		Payload,
 		Ray,
+		TLAS,
 		RandSequence,
 		RayFlags,
 		ShadowMaxNormalBias,
 		Mask,
+		ReflectedShadowsType,
+		ShouldDoDirectLighting,
+		ShouldDoEmissiveAndIndirectLighting,
 		RayCone,
 		Radiance);
 
@@ -573,15 +588,20 @@ FMaterialClosestHitPayload TraceRayAndAccumulateResults(
 
 FMaterialClosestHitPayload TraceRayAndAccumulateBottomLayerResults(
 	in RayDesc Ray,
+	in RaytracingAccelerationStructure TLAS,
 	in RandomSequence RandSequence,
 	in uint RayFlags,
 	in float ShadowMaxNormalBias,
 	in uint Mask,
+	in uint ReflectedShadowsType,
+	in uint ShouldDoDirectLighting,
+	in uint ShouldDoEmissiveAndIndirectLighting,
 	inout FRayCone RayCone,
 	inout float3 Radiance)
 {
 	FMaterialClosestHitPayload BottomLayerPayload = TraceRayResults(
 		Ray,
+		TLAS,
 		RandSequence,
 		RayFlags,
 		ShadowMaxNormalBias,
@@ -594,10 +614,14 @@ FMaterialClosestHitPayload TraceRayAndAccumulateBottomLayerResults(
 	AccumulateResults(
 		BottomLayerPayload,
 		Ray,
+		TLAS,
 		RandSequence,
 		RayFlags,
 		ShadowMaxNormalBias,
 		Mask,
+		ReflectedShadowsType,
+		ShouldDoDirectLighting,
+		ShouldDoEmissiveAndIndirectLighting,
 		RayCone,
 		Radiance);
 
diff --git a/Engine/Shaders/Private/RayTracing/RayTracingReflections.usf b/Engine/Shaders/Private/RayTracing/RayTracingReflections.usf
--- a/Engine/Shaders/Private/RayTracing/RayTracingReflections.usf
+++ b/Engine/Shaders/Private/RayTracing/RayTracingReflections.usf
@@ -35,15 +35,17 @@ float2 SobolIndexToUniformUnitSquare(uint2 SobolRand)
 #include "../DeferredLightingCommon.ush"
 #include "../ReflectionEnvironmentShared.ush"
 #include "../Montecarlo.ush"
-#include "../PathTracing/Utilities/PathTracingRandomSequence.ush" // #dxr-todo: move to a common folder
+#include "../PathTracing/Utilities/PathTracingRandomSequence.ush" 
 #include "../HeightFogCommon.ush"
-#include "RayTracingCommon.ush"
+#include "../LightShaderParameters.ush"
 #include "RayTracingHitGroupCommon.ush"
 #include "RayTracingDirectionalLight.ush"
 #include "RayTracingRectLight.ush"
 #include "RayTracingSphereLight.ush"
 #include "RayTracingSpotLight.ush"
 #include "RayTracingPointLight.ush"
+#include "RayTracingCommon.ush"
+#include "RayTracingLightingCommon.ush" 
 
 
 // Setup the reflection composite to sample the sky and reflection captures.
@@ -71,8 +73,8 @@ uint SamplesPerPixel;
 uint MaxBounces;
 uint HeightFog;
 uint UseReflectionCaptures;
-uint ShouldDoDirectLighting;
 uint ReflectedShadowsType;
+uint ShouldDoDirectLighting;
 uint ShouldDoEmissiveAndIndirectLighting;
 uint UpscaleFactor;
 uint SortTileSize;
@@ -86,8 +88,6 @@ float ReflectionMaxNormalBias;
 
 RaytracingAccelerationStructure TLAS;
 
-#include "RayTracingLightingCommon.ush" // #dxr_todo: UE-72548 helper functions must take TLAS explicitly rather than rely on a global declared above
-
 RWTexture2D<float4> ColorOutput;
 RWTexture2D<float> RayHitDistanceOutput;
 RWTexture2D<float> RayImaginaryDepthOutput;
@@ -311,15 +311,14 @@ void RayTracingReflectionsRGS()
 							TopLayerRay.TMax = 0;
 						}
 					}
-				}
-
+				}		
 
 
 				// Trace the top layer reflection
 				float3 TopLayerRadiance = float3(0, 0, 0);
 				FMaterialClosestHitPayload Payload;
 				const uint RayFlags = 0;					
-				Payload = TraceRayAndAccumulateResults(TopLayerRay, RandSequence, RayFlags, ReflectionMaxNormalBias, RAY_TRACING_MASK_OPAQUE, RayCone, TopLayerRadiance);
+				Payload = TraceRayAndAccumulateResults(TopLayerRay, TLAS, RandSequence, RayFlags, ReflectionMaxNormalBias, RAY_TRACING_MASK_OPAQUE, ReflectedShadowsType, ShouldDoDirectLighting, ShouldDoEmissiveAndIndirectLighting, RayCone, TopLayerRadiance);
 				bool isTopLayerRayValid = Payload.HitT >= 0 || (bAllowSkySampling && ReflectionStruct.SkyLightParameters.y > 0);
 
 
@@ -386,7 +385,7 @@ void RayTracingReflectionsRGS()
 
 					
 					FRayCone BottomLayerRayCone = (FRayCone)0;
-					FMaterialClosestHitPayload BottomLayerPayload = TraceRayAndAccumulateBottomLayerResults(BottomLayerRay, RandSequence, RayFlags, ReflectionMaxNormalBias, RAY_TRACING_MASK_OPAQUE, RayCone, BottomLayerRadiance);
+					FMaterialClosestHitPayload BottomLayerPayload = TraceRayAndAccumulateBottomLayerResults(BottomLayerRay, TLAS, RandSequence, RayFlags, ReflectionMaxNormalBias, RAY_TRACING_MASK_OPAQUE, ReflectedShadowsType, ShouldDoDirectLighting, ShouldDoEmissiveAndIndirectLighting, RayCone, BottomLayerRadiance);
 
 					// Distance and sky: similar processing to main ray.
 					float BottomLayerHitDistance = -1;
diff --git a/Engine/Shaders/Private/RayTracing/RayTracingTranslucency.usf b/Engine/Shaders/Private/RayTracing/RayTracingTranslucency.usf
--- a/Engine/Shaders/Private/RayTracing/RayTracingTranslucency.usf
+++ b/Engine/Shaders/Private/RayTracing/RayTracingTranslucency.usf
@@ -39,8 +39,8 @@ float2 SobolIndexToUniformUnitSquare(uint2 SobolRand)
 uint SamplesPerPixel;
 uint MaxRefractionRays;
 uint HeightFog;
-uint ShouldDoDirectLighting;
 uint ReflectedShadowsType;
+uint ShouldDoDirectLighting;
 uint ShouldDoEmissiveAndIndirectLighting;
 uint UpscaleFactor;
 
@@ -157,7 +157,7 @@ void RayTracingTranslucencyRGS()
 		float3 PathVertexRadiance = float3(0, 0, 0);
 		//#dxr-todo: RayFlags should be 0 for refraction rays
 		const uint RayFlags = bRefractionBackFaceCulling ? RAY_FLAG_CULL_BACK_FACING_TRIANGLES : 0;
-		FMaterialClosestHitPayload Payload = TraceRayAndAccumulateResults(Ray, RandSequence, RayFlags, MaxNormalBias, RAY_TRACING_MASK_ALL, RayCone, PathVertexRadiance);
+		FMaterialClosestHitPayload Payload = TraceRayAndAccumulateResults(Ray, TLAS, RandSequence, RayFlags, MaxNormalBias, RAY_TRACING_MASK_ALL, ReflectedShadowsType, ShouldDoDirectLighting, ShouldDoEmissiveAndIndirectLighting, RayCone, PathVertexRadiance);
 		LastRoughness = Payload.Roughness;
 
 		//
@@ -207,7 +207,7 @@ void RayTracingTranslucencyRGS()
 
 			float3 ReflectionRadiance = float3(0, 0, 0);
 			const uint RayFlags = bReflectionBackFaceCulling ? RAY_FLAG_CULL_BACK_FACING_TRIANGLES : 0;
-			FMaterialClosestHitPayload ReflectionPayload = TraceRayAndAccumulateResults(ReflectionRay, RandSequence, RayFlags, MaxNormalBias, RAY_TRACING_MASK_ALL, RayCone, ReflectionRadiance);
+			FMaterialClosestHitPayload ReflectionPayload = TraceRayAndAccumulateResults(ReflectionRay, TLAS, RandSequence, RayFlags, MaxNormalBias, RAY_TRACING_MASK_ALL, ReflectedShadowsType, ShouldDoDirectLighting, ShouldDoEmissiveAndIndirectLighting, RayCone, ReflectionRadiance);
 
 			// If we have not hit anything, sample the distance sky radiance.
 			if (!IsHit(ReflectionRay, ReflectionPayload))
