Fix an invalid vector truncation in the mobile emulation shader path.

diff --git a/Engine/Shaders/Private/MobileBasePassPixelShader.usf b/Engine/Shaders/Private/MobileBasePassPixelShader.usf
--- a/Engine/Shaders/Private/MobileBasePassPixelShader.usf
+++ b/Engine/Shaders/Private/MobileBasePassPixelShader.usf
@@ -418,7 +418,11 @@ void Main(
 	{
 		// this feature is only needed for development/editor - we can compile it out for a shipping build (see r.CompileShadersForDevelopment cvar help)
 		DiffuseColor = DiffuseColor * ResolvedView.DiffuseOverrideParameter.w + ResolvedView.DiffuseOverrideParameter.xyz;
-		SpecularColor = SpecularColor * ResolvedView.SpecularOverrideParameter.w + ResolvedView.SpecularOverrideParameter.xyz;
+#if NONMETAL
+    SpecularColor = SpecularColor * ResolvedView.SpecularOverrideParameter.w + ResolvedView.SpecularOverrideParameter.x;
+#else
+    SpecularColor = SpecularColor * ResolvedView.SpecularOverrideParameter.w + ResolvedView.SpecularOverrideParameter.xyz;
+#endif
 	}
 #endif
 	
