DR - Compile fix for hlslcc platforms
#rnx

diff --git a/Engine/Shaders/Private/PostProcessVelocityFlatten.usf b/Engine/Shaders/Private/PostProcessVelocityFlatten.usf
--- a/Engine/Shaders/Private/PostProcessVelocityFlatten.usf
+++ b/Engine/Shaders/Private/PostProcessVelocityFlatten.usf
@@ -23,8 +23,8 @@ RWTexture2D<float4>	OutVelocityFlatTexture;
 RWTexture2D<float4>	OutVelocityTileTexture;
 RWTexture2D<float4>	OutDilatedVelocityTileTexture;
 
-SCREEN_PASS_TEXTURE_VIEWPORT(Velocity);
-SCREEN_PASS_TEXTURE_VIEWPORT(VelocityTile);
+SCREEN_PASS_TEXTURE_VIEWPORT(Velocity)
+SCREEN_PASS_TEXTURE_VIEWPORT(VelocityTile)
 
 float4 MinMaxLength(float4 v0, float4 v1)
 {
