DR - Compile fixes
#rnx

diff --git a/Engine/Shaders/Private/PostProcessMotionBlur.usf b/Engine/Shaders/Private/PostProcessMotionBlur.usf
--- a/Engine/Shaders/Private/PostProcessMotionBlur.usf
+++ b/Engine/Shaders/Private/PostProcessMotionBlur.usf
@@ -25,10 +25,10 @@ float MotionBlur_VelocityScale;
 float MotionBlur_VelocityScaleForTiles;
 float MotionBlur_VelocityMax;
 
-SCREEN_PASS_TEXTURE_VIEWPORT(Color);
-SCREEN_PASS_TEXTURE_VIEWPORT(Velocity);
-SCREEN_PASS_TEXTURE_VIEWPORT(VelocityTile);
-SCREEN_PASS_TEXTURE_VIEWPORT_TRANSFORM(ColorToVelocity);
+SCREEN_PASS_TEXTURE_VIEWPORT(Color)
+SCREEN_PASS_TEXTURE_VIEWPORT(Velocity)
+SCREEN_PASS_TEXTURE_VIEWPORT(VelocityTile)
+SCREEN_PASS_TEXTURE_VIEWPORT_TRANSFORM(ColorToVelocity)
 
 Texture2D ColorTexture;
 Texture2D VelocityTexture;
diff --git a/Engine/Shaders/Private/PostProcessSubsurface.usf b/Engine/Shaders/Private/PostProcessSubsurface.usf
--- a/Engine/Shaders/Private/PostProcessSubsurface.usf
+++ b/Engine/Shaders/Private/PostProcessSubsurface.usf
@@ -22,8 +22,8 @@ Texture2D SubsurfaceInput1_Texture;
 
 SamplerState SubsurfaceSampler; 
 
-SCREEN_PASS_TEXTURE_VIEWPORT(SubsurfaceInput0);
-SCREEN_PASS_TEXTURE_VIEWPORT(SubsurfaceInput1);
+SCREEN_PASS_TEXTURE_VIEWPORT(SubsurfaceInput0)
+SCREEN_PASS_TEXTURE_VIEWPORT(SubsurfaceInput1)
 
 #define SUBSURFACE_DIRECTION_HORIZONTAL 0
 #define SUBSURFACE_DIRECTION_VERTICAL 1
