DR - ES2 compile fix
#rnx

diff --git a/Engine/Shaders/Private/MobileSceneCapture.usf b/Engine/Shaders/Private/MobileSceneCapture.usf
--- a/Engine/Shaders/Private/MobileSceneCapture.usf
+++ b/Engine/Shaders/Private/MobileSceneCapture.usf
@@ -18,7 +18,10 @@ void MainCopyVS(
 	out FScreenVertexOutput Output
 	)
 {
-	DrawRectangle( float4( InPosition, 0, 1 ), InUV, Output.Position, Output.UV);
+	float2 UVs = 0;
+	DrawRectangle( float4( InPosition, 0, 1 ), InUV, Output.Position, UVs);
+	// float2 -> half2
+	Output.UV = UVs;
 
 #if DECODING_MOSAIC
 	OutTexCoords[0] = Output.UV + InvTexSize.xy * float2( 0,-1);
