Fixes wrong #if in MobileBasePassPixelShader.usf

diff --git a/Engine/Shaders/Private/MobileBasePassPixelShader.usf b/Engine/Shaders/Private/MobileBasePassPixelShader.usf
--- a/Engine/Shaders/Private/MobileBasePassPixelShader.usf
+++ b/Engine/Shaders/Private/MobileBasePassPixelShader.usf
@@ -654,7 +654,7 @@ void Main(
 
 	#if MATERIALBLENDING_ALPHACOMPOSITE
 		OutColor = half4(Color * VertexFog.a + VertexFog.rgb * Opacity, Opacity);
-	#if MATERIALBLENDING_ALPHAHOLDOUT
+	#elif MATERIALBLENDING_ALPHAHOLDOUT
 		// not implemented for holdout
 		OutColor = half4(Color * VertexFog.a + VertexFog.rgb * Opacity, Opacity);
 	#elif MATERIALBLENDING_TRANSLUCENT
