Removed redundant functions in ray tracing code

diff --git a/Engine/Shaders/Private/RayTracing/RayTracingAmbientOcclusionRGS.usf b/Engine/Shaders/Private/RayTracing/RayTracingAmbientOcclusionRGS.usf
--- a/Engine/Shaders/Private/RayTracing/RayTracingAmbientOcclusionRGS.usf
+++ b/Engine/Shaders/Private/RayTracing/RayTracingAmbientOcclusionRGS.usf
@@ -10,10 +10,6 @@ RaytracingAccelerationStructure TLAS;
 RWTexture2D<float> RWOcclusionMaskUAV;
 RWTexture2D<float> RWHitDistanceUAV;
 
-uint CalcLinearIndex(uint2 PixelCoord)
-{
-	return PixelCoord.y * View.BufferSizeAndInvSize.x + PixelCoord.x;
-}
 
 void GenerateCosineNormalRay(
 	float3 WorldPosition,
diff --git a/Engine/Shaders/Private/RayTracing/RayTracingCommon.ush b/Engine/Shaders/Private/RayTracing/RayTracingCommon.ush
--- a/Engine/Shaders/Private/RayTracing/RayTracingCommon.ush
+++ b/Engine/Shaders/Private/RayTracing/RayTracingCommon.ush
@@ -84,6 +84,21 @@ float3 ReconstructWorldPositionFromDepth(float2 UV, float Depth)
 	return WorldPosition;
 }
 
+uint CalcLinearIndex(uint2 PixelCoord)
+{
+	return PixelCoord.y * View.BufferSizeAndInvSize.x + PixelCoord.x;
+}
+
+uint2 GetPixelCoord(uint2 DispatchThreadId, uint UpscaleFactor)
+{
+	uint UpscaleFactorPow2 = UpscaleFactor * UpscaleFactor;
+
+	// TODO: find a way to not interfer with TAA's jittering.
+	uint SubPixelId = View.StateFrameIndex & (UpscaleFactorPow2 - 1);
+
+	return DispatchThreadId * UpscaleFactor + uint2(SubPixelId & (UpscaleFactor - 1), SubPixelId / UpscaleFactor);
+}
+
 #define RAY_TRACING_BLEND_MODE_OPAQUE				0
 #define RAY_TRACING_BLEND_MODE_ALPHA_COMPOSITE		1
 #define RAY_TRACING_BLEND_MODE_TRANSLUCENT			2
diff --git a/Engine/Shaders/Private/RayTracing/RayTracingGlobalIlluminationRGS.usf b/Engine/Shaders/Private/RayTracing/RayTracingGlobalIlluminationRGS.usf
--- a/Engine/Shaders/Private/RayTracing/RayTracingGlobalIlluminationRGS.usf
+++ b/Engine/Shaders/Private/RayTracing/RayTracingGlobalIlluminationRGS.usf
@@ -29,22 +29,6 @@ bool EvalSkyLight;
 bool UseRussianRoulette;
 float MaxNormalBias;
 
-// #dxr_todo: UE-72508 Unify with reflections and translucency in RayTracingCommon.ush
-uint2 GetPixelCoord(uint2 DispatchThreadId)
-{
-	uint UpscaleFactorPow2 = UpscaleFactor * UpscaleFactor;
-
-	// TODO: find a way to not interfer with TAA's jittering.
-	uint SubPixelId = View.StateFrameIndex & (UpscaleFactorPow2 - 1);
-
-	return DispatchThreadId * UpscaleFactor + uint2(SubPixelId & (UpscaleFactor - 1), SubPixelId / UpscaleFactor);
-}
-
-uint CalcLinearIndex(uint2 PixelCoord)
-{
-	return PixelCoord.y * View.BufferSizeAndInvSize.x + PixelCoord.x;
-}
-
 void GenerateCosineNormalRay(
 	float3 WorldPosition,
 	float3 WorldNormal,
@@ -79,7 +63,7 @@ void GlobalIlluminationRGS()
 	RWGlobalIlluminationUAV[DispatchThreadId] = 0.0;
 	RWRayDistanceUAV[DispatchThreadId] = float2(-1.0, 0.0);
 
-	uint2 PixelCoord = GetPixelCoord(DispatchThreadId);
+	uint2 PixelCoord = GetPixelCoord(DispatchThreadId, UpscaleFactor);
 	RandomSequence RandSequence;
 	uint LinearIndex = CalcLinearIndex(PixelCoord);
 	RandomSequence_Initialize(RandSequence, LinearIndex, View.FrameNumber);
diff --git a/Engine/Shaders/Private/RayTracing/RayTracingOcclusionRGS.usf b/Engine/Shaders/Private/RayTracing/RayTracingOcclusionRGS.usf
--- a/Engine/Shaders/Private/RayTracing/RayTracingOcclusionRGS.usf
+++ b/Engine/Shaders/Private/RayTracing/RayTracingOcclusionRGS.usf
@@ -32,11 +32,6 @@ float2 SobolIndexToUniformUnitSquare(uint2 SobolRand)
 	return float2(SobolRand) * rcp(65536.0) + rcp(65536.0 * 2.0);
 }
 
-uint CalcLinearIndex(uint2 PixelCoord)
-{
-	return PixelCoord.y * View.BufferSizeAndInvSize.x + PixelCoord.x;
-}
-
 bool GenerateOcclusionRay(
 	FLightShaderParameters LightParameters,
 	float3 WorldPosition,
diff --git a/Engine/Shaders/Private/RayTracing/RayTracingRectLightRGS.usf b/Engine/Shaders/Private/RayTracing/RayTracingRectLightRGS.usf
--- a/Engine/Shaders/Private/RayTracing/RayTracingRectLightRGS.usf
+++ b/Engine/Shaders/Private/RayTracing/RayTracingRectLightRGS.usf
@@ -13,11 +13,6 @@ RaytracingAccelerationStructure TLAS;
 RWTexture2D<float4> RWLuminanceUAV;
 RWTexture2D<float> RWRayDistanceUAV;
 
-uint CalcLinearIndex(uint2 PixelCoord)
-{
-	return PixelCoord.y * View.BufferSizeAndInvSize.x + PixelCoord.x;
-}
-
 bool GenerateRectLightImportanceSampledRay(
 	float3 WorldPosition,
 	float3 WorldNormal,
diff --git a/Engine/Shaders/Private/RayTracing/RayTracingReflections.usf b/Engine/Shaders/Private/RayTracing/RayTracingReflections.usf
--- a/Engine/Shaders/Private/RayTracing/RayTracingReflections.usf
+++ b/Engine/Shaders/Private/RayTracing/RayTracingReflections.usf
@@ -96,16 +96,6 @@ RWTexture2D<float> RayImaginaryDepthOutput;
 // Material buffer and tile size are used for sorted materials.
 RWStructuredBuffer<FDeferredMaterialPayload> MaterialBuffer;
 
-uint2 GetPixelCoord(uint2 DispatchThreadId)
-{
-	uint UpscaleFactorPow2 = UpscaleFactor * UpscaleFactor;
-
-	// TODO: find a way to not interfer with TAA's jittering.
-	uint SubPixelId = View.StateFrameIndex & (UpscaleFactorPow2 - 1);
-
-	return DispatchThreadId * UpscaleFactor + uint2(SubPixelId & (UpscaleFactor - 1), SubPixelId / UpscaleFactor);
-}
-
 // Generate a random direction to sample according to world normal and roughness.
 float3 GenerateReflectionDirection(inout RandomSequence RandSequence, uint2 DispatchThreadId, uint BounceIndex, uint SamplePerPixel, float3 IncidentDirection, float3 WorldNormal, float Roughness)
 {
@@ -163,7 +153,7 @@ void RayTracingReflectionsRGS()
 	uint2 DispatchThreadId = DispatchRaysIndex().xy;
 #endif
 
-	uint2 PixelCoord = GetPixelCoord(DispatchThreadId); // GBuffer pixel coordinates (full resolution)
+	uint2 PixelCoord = GetPixelCoord(DispatchThreadId, UpscaleFactor); // GBuffer pixel coordinates (full resolution)
 
 	uint LinearIndex = PixelCoord.y * View.BufferSizeAndInvSize.x + PixelCoord.x; // TODO: PixelCoord or DispatchThreadId
 	
diff --git a/Engine/Shaders/Private/RayTracing/RayTracingSkyLightRGS.usf b/Engine/Shaders/Private/RayTracing/RayTracingSkyLightRGS.usf
--- a/Engine/Shaders/Private/RayTracing/RayTracingSkyLightRGS.usf
+++ b/Engine/Shaders/Private/RayTracing/RayTracingSkyLightRGS.usf
@@ -16,12 +16,7 @@ RWTexture2D<float2> RWRayDistanceUAV;
 #define USE_MIP_PDF_TEXTURE 1
 #define USE_SOLID_ANGLE_TEXTURE 1
 
-uint CalcLinearIndex(uint2 PixelCoord)
-{
-	return PixelCoord.y * View.BufferSizeAndInvSize.x + PixelCoord.x;
-}
 
-// #dxr_todo: UE-72508: encapsulate this function in a shared space
 float3 EvalSkyLight(float3 Direction)
 {
 	return SkyLight.Color * TextureCubeSampleLevel(SkyLight.Texture, SkyLight.TextureSampler, Direction, 0).rgb;
diff --git a/Engine/Shaders/Private/RayTracing/RayTracingTranslucency.usf b/Engine/Shaders/Private/RayTracing/RayTracingTranslucency.usf
--- a/Engine/Shaders/Private/RayTracing/RayTracingTranslucency.usf
+++ b/Engine/Shaders/Private/RayTracing/RayTracingTranslucency.usf
@@ -57,16 +57,6 @@ RWTexture2D<float> RayHitDistanceOutput;
 
 #include "RayTracingLightingCommon.ush"
 
-uint2 GetPixelCoord(uint2 DispatchThreadId)
-{
-	uint UpscaleFactorPow2 = UpscaleFactor * UpscaleFactor;
-
-	// TODO: find a way to not interfer with TAA's jittering.
-	uint SubPixelId = View.StateFrameIndex & (UpscaleFactorPow2 - 1);
-
-	return DispatchThreadId * UpscaleFactor + uint2(SubPixelId & (UpscaleFactor - 1), SubPixelId / UpscaleFactor);
-}
-
 float FresnelDielectric(float Eta, float IoH, float ToH)
 {
 	float Rs = Square((Eta * IoH - ToH) / (Eta * IoH + ToH));
@@ -115,7 +105,7 @@ float3 GetSkyRadiance(float3 Direction, float Roughness)
 void RayTracingTranslucencyRGS()
 {
 	uint2 DispatchThreadId = DispatchRaysIndex().xy;
-	uint2 PixelCoord = GetPixelCoord(DispatchThreadId);
+	uint2 PixelCoord = GetPixelCoord(DispatchThreadId, UpscaleFactor);
 	uint LinearIndex = PixelCoord.y * View.BufferSizeAndInvSize.x + PixelCoord.x; // TODO: PixelCoord or DispatchThreadId
 	
 	RandomSequence RandSequence;
