Fixed pre-exposure issues with decal emissive and planar reflection.

diff --git a/Engine/Shaders/Private/DecalCommon.ush b/Engine/Shaders/Private/DecalCommon.ush
--- a/Engine/Shaders/Private/DecalCommon.ush
+++ b/Engine/Shaders/Private/DecalCommon.ush
@@ -35,6 +35,10 @@ void DecalCommonOutput(inout FPixelShaderIn In, inout FPixelShaderOut Out, float
 	// RETURN_COLOR not needed unless writing to SceneColor
 	Out.MRT[0] = float4(Color, Opacity);
 
+#if USE_PREEXPOSURE
+	Out.MRT[0].xyz *= View.PreExposure;
+#endif
+
 #if (ES2_PROFILE || ES3_1_PROFILE)
 #if (COMPILER_GLSL_ES2 || COMPILER_GLSL_ES3_1) || (MOBILE_EMULATION)
 	// Do decal blending if encoding requires it.
diff --git a/Engine/Shaders/Private/PlanarReflectionShaders.usf b/Engine/Shaders/Private/PlanarReflectionShaders.usf
--- a/Engine/Shaders/Private/PlanarReflectionShaders.usf
+++ b/Engine/Shaders/Private/PlanarReflectionShaders.usf
@@ -142,5 +142,9 @@ void PlanarReflectionPS(
 	if (ScreenSpaceData.GBuffer.ShadingModelID > 0)
 	{
 		OutColor = ComputePlanarReflections(WorldPosition, ScreenSpaceData.GBuffer.WorldNormal, ScreenSpaceData.GBuffer.Roughness, PlanarReflectionStruct.PlanarReflectionSampler);
+
+	#if USE_PREEXPOSURE
+		OutColor.xyz *= View.PreExposure;
+	#endif
 	}
 }
