Fix ShaderPrint on XB1
Integrating fix from Dev-VT
XB1 doesn't like function with same name as uniform variable

diff --git a/Engine/Shaders/Private/ShaderPrintCommon.ush b/Engine/Shaders/Private/ShaderPrintCommon.ush
--- a/Engine/Shaders/Private/ShaderPrintCommon.ush
+++ b/Engine/Shaders/Private/ShaderPrintCommon.ush
@@ -11,7 +11,7 @@
 // --------------------------------------------------------------------------
 // Uniforms and structures
 
-// Needs to match C++ code in ShaderPrint.cpp
+// Needs to match C++ code in ShaderPrintUniform.cpp
 struct ShaderPrintItem
 {
 	float2 ScreenPos; // Position in normalized coordinates
@@ -28,7 +28,7 @@ RWStructuredBuffer<ShaderPrintItem> RWValuesBuffer;
 #define SHADER_PRINT_TYPE_UINT   3
 #define SHADER_PRINT_TYPE_HEX    4
 
-#include "/Engine/Generated/UniformBuffers/ShaderPrint.ush"
+#include "/Engine/Generated/UniformBuffers/ShaderPrintUniform.ush"
 //float4 FontSize;      // .xy character size, .zw character spacing (in normalized coordinates)
 //int MaxValueCount;    // maximum number of values that we can print in one frame
 //int MaxSymbolCount;   // maximum number of symbols (after conversion from characters) that we can print in one frame
@@ -46,7 +46,7 @@ static bool ShaderPrintFilterEnable = true;
 void ShaderPrint_Internal(in ShaderPrintItem Item)
 {
 	// If MaxValueCount is 0 then we don't reset the buffer counter so early out here
-	if (ShaderPrint.MaxValueCount == 0)
+	if (ShaderPrintUniform.MaxValueCount == 0)
 	{
 		return;
 	}
@@ -57,7 +57,7 @@ void ShaderPrint_Internal(in ShaderPrintItem Item)
 
 	// Prevent writing off the buffer
 	// Note that counter still increases so need clamp when reading it in later passes
-	if (IndexToStore >= ShaderPrint.MaxValueCount)
+	if (IndexToStore >= ShaderPrintUniform.MaxValueCount)
 	{
 		return;
 	}
@@ -83,7 +83,7 @@ float2 ShaderPrintSymbol(in float2 ScreenPos, in int Symbol)
 	if (ShaderPrintFilterEnable)
 	{
 		ShaderPrint_Internal(ScreenPos, Symbol, SHADER_PRINT_TYPE_SYMBOL);
-		ScreenPos.x += ShaderPrint.FontSize.z;
+		ScreenPos.x += ShaderPrintUniform.FontSize.z;
 	}
 	return ScreenPos;
 }
@@ -126,7 +126,7 @@ float2 ShaderPrint(in float2 ScreenPos, in float Value, in int Type = SHADER_PRI
 	if (ShaderPrintFilterEnable)
 	{
 		ShaderPrint_Internal(ScreenPos, asint(Value), Type);
-		ScreenPos.x += ShaderPrint.FontSize.z * ShaderPrintTabCount;
+		ScreenPos.x += ShaderPrintUniform.FontSize.z * ShaderPrintTabCount;
 	}
 	return ScreenPos;
 }
@@ -141,7 +141,7 @@ float2 ShaderPrint(in float2 ScreenPos, in int Value, in int Type = SHADER_PRINT
 	if (ShaderPrintFilterEnable)
 	{
 		ShaderPrint_Internal(ScreenPos, Value, Type);
-		ScreenPos.x += ShaderPrint.FontSize.z * ShaderPrintTabCount;
+		ScreenPos.x += ShaderPrintUniform.FontSize.z * ShaderPrintTabCount;
 	}
 	return ScreenPos;
 }
@@ -156,7 +156,7 @@ float2 ShaderPrint(in float2 ScreenPos, in uint Value, in int Type = SHADER_PRIN
 	if (ShaderPrintFilterEnable)
 	{
 		ShaderPrint_Internal(ScreenPos, asint(Value), Type);
-		ScreenPos.x += ShaderPrint.FontSize.z * ShaderPrintTabCount;
+		ScreenPos.x += ShaderPrintUniform.FontSize.z * ShaderPrintTabCount;
 	}
 	return ScreenPos;
 }
@@ -332,11 +332,11 @@ float2 ShaderPrint(in float2 ScreenPos, in float3x3 Value)
 	if (ShaderPrintFilterEnable)
 	{
 		ShaderPrint(ScreenPos, Value[0]);
-		ScreenPos += float2(0.f, ShaderPrint.FontSize.w);
+		ScreenPos += float2(0.f, ShaderPrintUniform.FontSize.w);
 		ShaderPrint(ScreenPos, Value[1]);
-		ScreenPos += float2(0.f, ShaderPrint.FontSize.w);
+		ScreenPos += float2(0.f, ShaderPrintUniform.FontSize.w);
 		ShaderPrint(ScreenPos, Value[2]);
-		ScreenPos += float2(0.f, ShaderPrint.FontSize.w);
+		ScreenPos += float2(0.f, ShaderPrintUniform.FontSize.w);
 	}
 	return ScreenPos;
 }
@@ -351,13 +351,13 @@ float2 ShaderPrint(in float2 ScreenPos, in float4x3 Value)
 	if (ShaderPrintFilterEnable)
 	{
 		ShaderPrint(ScreenPos, float4(Value[0], 0.f));
-		ScreenPos += float2(0.f, ShaderPrint.FontSize.w);
+		ScreenPos += float2(0.f, ShaderPrintUniform.FontSize.w);
 		ShaderPrint(ScreenPos, float4(Value[1], 0.f));
-		ScreenPos += float2(0.f, ShaderPrint.FontSize.w);
+		ScreenPos += float2(0.f, ShaderPrintUniform.FontSize.w);
 		ShaderPrint(ScreenPos, float4(Value[2], 0.f));
-		ScreenPos += float2(0.f, ShaderPrint.FontSize.w);
+		ScreenPos += float2(0.f, ShaderPrintUniform.FontSize.w);
 		ShaderPrint(ScreenPos, float4(Value[3], 1.f));
-		ScreenPos += float2(0.f, ShaderPrint.FontSize.w);
+		ScreenPos += float2(0.f, ShaderPrintUniform.FontSize.w);
 	}
 	return ScreenPos;
 }
@@ -372,13 +372,13 @@ float2 ShaderPrint(in float2 ScreenPos, in float4x4 Value)
 	if (ShaderPrintFilterEnable)
 	{
 		ShaderPrint(ScreenPos, Value[0]);
-		ScreenPos += float2(0.f, ShaderPrint.FontSize.w);
+		ScreenPos += float2(0.f, ShaderPrintUniform.FontSize.w);
 		ShaderPrint(ScreenPos, Value[1]);
-		ScreenPos += float2(0.f, ShaderPrint.FontSize.w);
+		ScreenPos += float2(0.f, ShaderPrintUniform.FontSize.w);
 		ShaderPrint(ScreenPos, Value[2]);
-		ScreenPos += float2(0.f, ShaderPrint.FontSize.w);
+		ScreenPos += float2(0.f, ShaderPrintUniform.FontSize.w);
 		ShaderPrint(ScreenPos, Value[3]);
-		ScreenPos += float2(0.f, ShaderPrint.FontSize.w);
+		ScreenPos += float2(0.f, ShaderPrintUniform.FontSize.w);
 	}
 	return ScreenPos;
 }
@@ -395,7 +395,7 @@ float2 ShaderPrintSpace(in float2 ScreenPos, in float Count = 1)
 {
 	if (ShaderPrintFilterEnable)
 	{
-		ScreenPos.x += ShaderPrint.FontSize.z * Count;
+		ScreenPos.x += ShaderPrintUniform.FontSize.z * Count;
 	}
 	return ScreenPos;
 
@@ -411,7 +411,7 @@ float2 ShaderPrintNewline(in float2 ScreenPos)
 	if (ShaderPrintFilterEnable)
 	{
 		ScreenPos.x = ShaderPrintCursorStart.x;
-		ScreenPos.y += ShaderPrint.FontSize.w;
+		ScreenPos.y += ShaderPrintUniform.FontSize.w;
 	}
 	return ScreenPos;
 }
diff --git a/Engine/Shaders/Private/ShaderPrintDraw.usf b/Engine/Shaders/Private/ShaderPrintDraw.usf
--- a/Engine/Shaders/Private/ShaderPrintDraw.usf
+++ b/Engine/Shaders/Private/ShaderPrintDraw.usf
@@ -36,7 +36,7 @@ void InitValueBufferCS()
 [numthreads(1, 1, 1)]
 void BuildIndirectDispatchArgsCS()
 {
-	uint ValueCount = min(ShaderPrint.MaxValueCount, ValuesBuffer[0].Value);
+	uint ValueCount = min(ShaderPrintUniform.MaxValueCount, ValuesBuffer[0].Value);
 
 	RWIndirectDispatchArgsBuffer[0] = (ValueCount + 63)/64;
 	RWIndirectDispatchArgsBuffer[1] = 1;
@@ -59,7 +59,7 @@ void AddSymbol(in ShaderPrintItem Item)
 	
 	// Prevent writing off the buffer
 	// Note that counter still increases so need to clamp when reading it in later passes
-	if (IndexToStore >= ShaderPrint.MaxSymbolCount)
+	if (IndexToStore >= ShaderPrintUniform.MaxSymbolCount)
 	{
 		return;
 	}
@@ -76,7 +76,7 @@ void AddSymbol(inout float2 ScreenPos, in int Symbol)
 	Item.Type = 0;
 	AddSymbol(Item);
 
-	ScreenPos.x += ShaderPrint.FontSize.z;
+	ScreenPos.x += ShaderPrintUniform.FontSize.z;
 }
 
 int GetDecimalSymbol(in int Digit)
@@ -182,7 +182,7 @@ void BuildSymbolBufferCS(
 	uint3 DispatchThreadId : SV_DispatchThreadID
 	)
 {
-	uint ValueCount = min(ShaderPrint.MaxValueCount, ValuesBuffer[0].Value);
+	uint ValueCount = min(ShaderPrintUniform.MaxValueCount, ValuesBuffer[0].Value);
 	if (DispatchThreadId.x >= ValueCount)
 	{
 		return;
@@ -209,7 +209,7 @@ void BuildSymbolBufferCS(
 [numthreads(1, 1, 1)]
 void BuildIndirectDrawArgsCS()
 {
-	uint SymbolCount = min(ShaderPrint.MaxSymbolCount, SymbolsBuffer[0].Value);
+	uint SymbolCount = min(ShaderPrintUniform.MaxSymbolCount, SymbolsBuffer[0].Value);
 
 	RWIndirectDrawArgsBuffer[0] = 6;
 	RWIndirectDrawArgsBuffer[1] = SymbolCount * 2; // Main glyph + drop shadow
@@ -232,7 +232,7 @@ void DrawSymbolsVS(
 {
 	bool bDropShadow = (InstanceId & 1) == 0;
 	OutShadow = bDropShadow ? 1 : 0;
-	float2 DropShadowOffset = bDropShadow ? ShaderPrint.FontSize.xy * float2(0.2f, -0.2f) : 0;
+	float2 DropShadowOffset = bDropShadow ? ShaderPrintUniform.FontSize.xy * float2(0.2f, -0.2f) : 0;
 	
 	InstanceId = InstanceId >> 1;
 	ShaderPrintItem Symbol = SymbolsBuffer[InstanceId + 1];
@@ -242,7 +242,7 @@ void DrawSymbolsVS(
 	
 	float2 VertPos = float2(OutUV.x, 1.f - OutUV.y) * 2.f - 1.f;
 	float2 InstancePos = float2(Symbol.ScreenPos.x, 1.f - Symbol.ScreenPos.y) * 2.f - 1.f;
-	OutPosition.xy = VertPos.xy * ShaderPrint.FontSize.xy + InstancePos + DropShadowOffset;
+	OutPosition.xy = VertPos.xy * ShaderPrintUniform.FontSize.xy + InstancePos + DropShadowOffset;
 	OutPosition.zw = float2(0, 1);
 }
 
diff --git a/Engine/Source/Runtime/Renderer/Private/ShaderPrint.cpp b/Engine/Source/Runtime/Renderer/Private/ShaderPrint.cpp
--- a/Engine/Source/Runtime/Renderer/Private/ShaderPrint.cpp
+++ b/Engine/Source/Runtime/Renderer/Private/ShaderPrint.cpp
@@ -71,7 +71,7 @@ namespace ShaderPrint
 	}
 
 	// ShaderPrint uniform buffer
-	IMPLEMENT_GLOBAL_SHADER_PARAMETER_STRUCT(FUniformBufferParameters, "ShaderPrint");
+	IMPLEMENT_GLOBAL_SHADER_PARAMETER_STRUCT(FUniformBufferParameters, "ShaderPrintUniform");
 	typedef TUniformBufferRef<FUniformBufferParameters> FUniformBufferRef;
 
 	// Fill the uniform buffer parameters
@@ -316,7 +316,7 @@ namespace ShaderPrint
 		SCOPED_DRAW_EVENT(RHICmdList, ShaderPrintBeginView);
 
 		// Clear the output buffer internal counter ready for use
-		const ERHIFeatureLevel::Type FeatureLevel = GMaxRHIFeatureLevel;
+		const ERHIFeatureLevel::Type FeatureLevel = View.GetFeatureLevel();
 		TShaderMap<FGlobalShaderType>* GlobalShaderMap = GetGlobalShaderMap(FeatureLevel);
 
 		TShaderMapRef< FShaderInitValueBufferCS > ComputeShader(GlobalShaderMap);
@@ -352,7 +352,7 @@ namespace ShaderPrint
 		FShaderResourceViewRHIRef ValuesBuffer = View.ShaderPrintValueBuffer.SRV;
 		FTextureRHIRef FontTexture = GEngine->MiniFontTexture != nullptr ? GEngine->MiniFontTexture->Resource->TextureRHI : GSystemTextures.BlackDummy->GetRenderTargetItem().ShaderResourceTexture;;
 
-		const ERHIFeatureLevel::Type FeatureLevel = GMaxRHIFeatureLevel;
+		const ERHIFeatureLevel::Type FeatureLevel = View.GetFeatureLevel();
 		TShaderMap<FGlobalShaderType>* GlobalShaderMap = GetGlobalShaderMap(FeatureLevel);
 
 		// BuildIndirectDispatchArgs
