Changes the computation of bluring radius for specular denoising by measuring the planar deviation of the specular lobe that truncate only 10% of the tale.

diff --git a/Engine/Shaders/Private/ScreenSpaceDenoise/SSDSignalFramework.ush b/Engine/Shaders/Private/ScreenSpaceDenoise/SSDSignalFramework.ush
--- a/Engine/Shaders/Private/ScreenSpaceDenoise/SSDSignalFramework.ush
+++ b/Engine/Shaders/Private/ScreenSpaceDenoise/SSDSignalFramework.ush
@@ -317,6 +317,7 @@ float4 WeightedLerp( float4 ColorA, float WeightA, float4 ColorB, float WeightB,
 //------------------------------------------------------- INTERNAL FOR REFLECTION SIGNAL
 #if CONFIG_SIGNAL_PROCESSING == SIGNAL_PROCESSING_REFLECTIONS
 
+
 /** Create a spherical gaussian lob for a given roughness. */
 FSphericalGaussian ComputeRoughnessLobe(float Roughness, float3 N, float3 V)
 {
@@ -353,6 +354,28 @@ float ComputeInfinityBluringRadius(FSSFSampleSceneInfos SceneMetadata)
 {
 	float RoughnessLobeAngle;
 	#if 1
+	{
+		// Measure the angle that only truncate 10% of D_GGX's tale.
+		const float E = (1.0 - 0.1);
+	
+		float a2 = Pow4(SceneMetadata.Roughness);
+	
+		float3x3 TangentBasis = GetTangentBasis(SceneMetadata.WorldNormal);
+		float3 V = mul(TangentBasis, normalize(View.TranslatedWorldCameraOrigin - SceneMetadata.TranslatedWorldPosition));
+
+		float3 L[2];
+		UNROLL
+		for (uint i = 0; i < 2; i++)
+		{
+			float P = i * 0.5 + 0.0;
+			float3 H = ImportanceSampleVisibleGGX(float2(E, P), a2, V).xyz;
+
+			L[i] = 2 * dot( V, H ) * H - V;
+		}
+
+		RoughnessLobeAngle = acos(dot(L[0], L[1]));
+	}
+	#elif 0
 	// Compute the angle of the lobe using spherical gaussian framework.
 	{
 		FSphericalGaussian RoughnessLobe = ComputeRoughnessLobe(SceneMetadata);
