Pseudo texture sampling improvement (sampling was not correct with respect to voxel center, now similar to hardware)
Also floor is the way to go here.

Bringing this in as it seems it was not done after I did it in the Niagara branch.

Merging
//UE4/Dev-Niagara/Engine/Shaders/Private/Common.ush
to //UE4/Dev-Rendering/Engine/Shaders/Private/Common.ush

diff --git a/Engine/Shaders/Private/Common.ush b/Engine/Shaders/Private/Common.ush
--- a/Engine/Shaders/Private/Common.ush
+++ b/Engine/Shaders/Private/Common.ush
@@ -223,18 +223,20 @@ float2 Tile1Dto2D(float xsize, float idx)
 // return a pseudovolume texture sample.
 // useful for simulating 3D texturing with a 2D texture or as a texture flipbook with lerped transitions
 // treats 2d layout of frames a 3d texture and performs bilinear filtering by blending with an offset Z frame.
+// Wrap repeat mode along XY is not seamless. This is however enough for current sampling use cases all in [0,1].
 // @param Tex          = Input Texture Object storing Volume Data
 // @param inPos        = Input float3 for Position, 0-1
-// @param xsize        = Input float for num frames in x,y directions
-// @param numFrames    = Input float for num total frames
+// @param xysize       = Input float for num frames in x,y directions
+// @param numframes    = Input float for num total frames
 // @param mipmode      = Sampling mode: 0 = use miplevel, 1 = use UV computed gradients, 2 = Use gradients (default=0)
 // @param miplevel     = MIP level to use in mipmode=0 (default 0)
 // @param InDDX, InDDY = Texture gradients in mipmode=2
 float4 PseudoVolumeTexture(Texture2D Tex, SamplerState TexSampler, float3 inPos, float2 xysize, float numframes,
 	uint mipmode = 0, float miplevel = 0, float2 InDDX = 0, float2 InDDY = 0)
 {
-	float zframe = ceil(inPos.z * numframes);
-	float zphase = frac(inPos.z * numframes);
+	float z = inPos.z - 0.5f / numframes;	// This offset is needed to have a behavior consistent with hardware sampling (voxel value is at their center)
+	float zframe = floor(z * numframes);
+	float zphase = frac(z * numframes);
 
 	float2 uv = frac(inPos.xy) / xysize;
 
