Fixed macro condition for tessellation-dependent vertex attribute in BasePassCommon.ush.

diff --git a/Engine/Shaders/Private/BasePassCommon.ush b/Engine/Shaders/Private/BasePassCommon.ush
--- a/Engine/Shaders/Private/BasePassCommon.ush
+++ b/Engine/Shaders/Private/BasePassCommon.ush
@@ -60,6 +60,8 @@
 #define TRANSLUCENCY_PERVERTEX_LIGHTING_VOLUME (!FORWARD_SHADING && (TRANSLUCENCY_LIGHTING_VOLUMETRIC_PERVERTEX_DIRECTIONAL || TRANSLUCENCY_LIGHTING_VOLUMETRIC_PERVERTEX_NONDIRECTIONAL))
 #define TRANSLUCENCY_PERVERTEX_FORWARD_SHADING (FORWARD_SHADING && (TRANSLUCENCY_LIGHTING_VOLUMETRIC_PERVERTEX_DIRECTIONAL || TRANSLUCENCY_LIGHTING_VOLUMETRIC_PERVERTEX_NONDIRECTIONAL))
 
+#define TESSELLATION_SUPPORTED (PC_D3D || GL4_PROFILE || SWITCH_PROFILE || SWITCH_PROFILE_FORWARD || METAL_TESSELLATION_PROFILE || VULKAN_ENABLE_TESSELLATION)
+
 struct FSharedBasePassInterpolants
 {
 	//for texture-lightmapped translucency we can pass the vertex fog in its own interpolator
@@ -67,7 +69,7 @@ struct FSharedBasePassInterpolants
 	float4 VertexFog		: TEXCOORD7;
 #endif
 
-#if !PC_D3D && !GL4_PROFILE && !SWITCH_PROFILE && !SWITCH_PROFILE_FORWARD && !METAL_TESSELLATION_PROFILE
+#if !TESSELLATION_SUPPORTED
 	// Note: TEXCOORD8 is unused
 
 	#if USE_WORLD_POSITION_EXCLUDING_SHADER_OFFSETS
@@ -104,7 +106,7 @@ struct FSharedBasePassInterpolants
 #endif
 };
 
-#if PC_D3D || GL4_PROFILE || SWITCH_PROFILE || SWITCH_PROFILE_FORWARD || METAL_TESSELLATION_PROFILE || VULKAN_ENABLE_TESSELLATION
+#if TESSELLATION_SUPPORTED
 
 /** Interpolants passed from the vertex shader to the pixel shader. */
 struct FBasePassInterpolantsVSToPS : FSharedBasePassInterpolants
