More post merge compilation errors. Hopefully the last one.

diff --git a/Engine/Shaders/Private/BasePassPixelShader.usf b/Engine/Shaders/Private/BasePassPixelShader.usf
--- a/Engine/Shaders/Private/BasePassPixelShader.usf
+++ b/Engine/Shaders/Private/BasePassPixelShader.usf
@@ -887,7 +887,7 @@ void FPixelShaderInOut_MainPS(
 	#endif
 
 		// Make sure not to touch 0,0 which is clear color
-		float4 EncodedVelocity = float4(EncodeVelocityToTexture(Velocity) * GetPrimitiveData(MaterialParameters.PrimitiveId).DrawsVelocit, 0, 0);
+		float4 EncodedVelocity = float4(EncodeVelocityToTexture(Velocity) * GetPrimitiveData(MaterialParameters.PrimitiveId).DrawsVelocity, 0, 0);
 	#if USES_GBUFFER
 		GBuffer.Velocity = EncodedVelocity;
 	#else
