Fixed issue with PreExposure not working correctly with SSS.

diff --git a/Engine/Shaders/Private/DeferredLightPixelShaders.usf b/Engine/Shaders/Private/DeferredLightPixelShaders.usf
--- a/Engine/Shaders/Private/DeferredLightPixelShaders.usf
+++ b/Engine/Shaders/Private/DeferredLightPixelShaders.usf
@@ -102,7 +102,7 @@ void DeferredLightPixelMain(
 		OutColor *= ComputeLightProfileMultiplier(WorldPosition, DeferredLightUniforms.Position, -DeferredLightUniforms.Direction, DeferredLightUniforms.Tangent);
 
 #if USE_PREEXPOSURE
-		OutColor *= View.PreExposure;
+		OutColor.rgb *= View.PreExposure;
 #endif
 	}
 }
\ No newline at end of file
diff --git a/Engine/Shaders/Private/PostProcessSubsurface.usf b/Engine/Shaders/Private/PostProcessSubsurface.usf
--- a/Engine/Shaders/Private/PostProcessSubsurface.usf
+++ b/Engine/Shaders/Private/PostProcessSubsurface.usf
@@ -344,7 +344,7 @@ SDiffuseAndSpecular ReconstructLighting(float2 UVSceneColor, uint ReconstructMet
 	  // If we're not doing checkerboard encoding, we just need to read a single pixel and decode (combined diffuse/spec in RGB)
 	  float4 CenterSample = SubsurfaceInput0_Texture.SampleLevel(SharedSubsurfaceSampler0, UVSceneColor, 0);
 	  float3 CombinedColor = CenterSample.rgb;
-	  float DiffuseLuminance = CenterSample.a;
+	  float DiffuseLuminance = CenterSample.a * View.PreExposure;
   
 	  float CombinedLuminance = Luminance(CombinedColor);
 	  float DiffuseFactor = saturate(DiffuseLuminance / CombinedLuminance);
