Replaced remaining header guards with "#pragma once" directives in global shaders.

#rnx

diff --git a/Engine/Shaders/Private/AtmosphereCommon.ush b/Engine/Shaders/Private/AtmosphereCommon.ush
--- a/Engine/Shaders/Private/AtmosphereCommon.ush
+++ b/Engine/Shaders/Private/AtmosphereCommon.ush
@@ -20,8 +20,7 @@
 
 =============================================================================*/
 
-#ifndef __ATMOSPHERE_COMMON__
-#define __ATMOSPHERE_COMMON__
+#pragma once
 
 #if COMPILER_METAL && MAX_SHADER_LANGUAGE_VERSION <= 3
 // For inexplicable reasons until we can debug shaders and *really* see what's what we need to abs sqrt input for the functions in this shader file only.
@@ -629,5 +628,3 @@ float4 CalculateVertexAtmosphericFog(float3 WorldPosition, float3 InCameraPositi
 #undef sqrt
 #endif
 
-#endif
-
diff --git a/Engine/Shaders/Private/AtmospherePrecomputeCommon.ush b/Engine/Shaders/Private/AtmospherePrecomputeCommon.ush
--- a/Engine/Shaders/Private/AtmospherePrecomputeCommon.ush
+++ b/Engine/Shaders/Private/AtmospherePrecomputeCommon.ush
@@ -20,8 +20,7 @@
 
 =============================================================================*/
 
-#ifndef __ATMOSPHERE_PRECOMPUTE_COMMON__
-#define __ATMOSPHERE_PRECOMPUTE_COMMON__
+#pragma once
 
 /** Integration Parameters */
 const static int TransmittanceIntegralSamples = 500;
@@ -43,5 +42,3 @@ Texture3D AtmosphereDeltaSMTexture;
 SamplerState AtmosphereDeltaSMTextureSampler;
 Texture3D AtmosphereDeltaJTexture;
 SamplerState AtmosphereDeltaJTextureSampler;
-
-#endif
diff --git a/Engine/Shaders/Private/BRDF.ush b/Engine/Shaders/Private/BRDF.ush
--- a/Engine/Shaders/Private/BRDF.ush
+++ b/Engine/Shaders/Private/BRDF.ush
@@ -4,8 +4,7 @@
 	BRDF.usf: Bidirectional reflectance distribution functions.
 =============================================================================*/
 
-#ifndef __BRDF_COMMON__
-#define __BRDF_COMMON__
+#pragma once
 
 struct BxDFContext
 {
@@ -394,4 +393,3 @@ float Vis_Cloth( float NoV, float NoL )
 	return rcp( 4 * ( NoL + NoV - NoL * NoV ) );
 }
 
-#endif // __BRDF_COMMON__
\ No newline at end of file
diff --git a/Engine/Shaders/Private/Common.ush b/Engine/Shaders/Private/Common.ush
--- a/Engine/Shaders/Private/Common.ush
+++ b/Engine/Shaders/Private/Common.ush
@@ -6,10 +6,6 @@
 
 #pragma once
 
-// Workarround for UE-66460
-#ifndef __USH_COMMON__
-#define __USH_COMMON__
-
 #include "/Engine/Public/Platform.ush"
 // These types are used for material translator generated code, or any functions the translated code can call
 #if PIXELSHADER
@@ -1785,4 +1781,3 @@ struct FPixelShaderOut
 
 #endif // SUPPORTS_INDEPENDENT_SAMPLERS
 
-#endif // __USH_COMMON__
diff --git a/Engine/Shaders/Private/DeferredLightingCommon.ush b/Engine/Shaders/Private/DeferredLightingCommon.ush
--- a/Engine/Shaders/Private/DeferredLightingCommon.ush
+++ b/Engine/Shaders/Private/DeferredLightingCommon.ush
@@ -4,8 +4,7 @@
 	DeferredLightingCommon.usf: Common definitions for deferred lighting.
 =============================================================================*/
 
-#ifndef __DEFERRED_LIGHTING_COMMON__
-#define __DEFERRED_LIGHTING_COMMON__
+#pragma once
 
 #include "DeferredShadingCommon.ush"
 #include "DynamicLightingCommon.ush"
@@ -416,4 +415,3 @@ float3 GetSimpleDynamicLighting(float3 WorldPosition, float3 CameraVector, float
 	return OutLighting;
 }
 
-#endif
\ No newline at end of file
diff --git a/Engine/Shaders/Private/DeferredShadingCommon.ush b/Engine/Shaders/Private/DeferredShadingCommon.ush
--- a/Engine/Shaders/Private/DeferredShadingCommon.ush
+++ b/Engine/Shaders/Private/DeferredShadingCommon.ush
@@ -4,8 +4,7 @@
 	DeferredShadingCommon.usf: Common definitions for deferred shading.
 =============================================================================*/
 
-#ifndef __DEFERRED_SHADING_COMMON__
-#define __DEFERRED_SHADING_COMMON__
+#pragma once
 
 #include "ShadingCommon.ush"
 #include "LightAccumulator.ush"
@@ -707,4 +706,3 @@ FScreenSpaceData GetScreenSpaceData(float2 UV, bool bGetNormalizedNormal = true)
 	return Out;
 }
 
-#endif // __DEFERRED_SHADING_COMMON__
diff --git a/Engine/Shaders/Private/EyeAdaptationCommon.ush b/Engine/Shaders/Private/EyeAdaptationCommon.ush
--- a/Engine/Shaders/Private/EyeAdaptationCommon.ush
+++ b/Engine/Shaders/Private/EyeAdaptationCommon.ush
@@ -4,8 +4,7 @@
 	EyeAdaptationCommon.usf: PostProcessing shared functions and structures.
 =============================================================================*/
 
-#ifndef __EYE_ADAPTATION_COMMON__
-#define __EYE_ADAPTATION_COMMON__
+#pragma once
 
 
 // Single pixel RT used for EyeAdaptation, use function EyeAdaptationLookup() to access.
@@ -24,4 +23,3 @@ float EyeAdaptationLookup()
 }
 #endif
 
-#endif // __EYE_ADAPTATION_COMMON__
diff --git a/Engine/Shaders/Private/LightAccumulator.ush b/Engine/Shaders/Private/LightAccumulator.ush
--- a/Engine/Shaders/Private/LightAccumulator.ush
+++ b/Engine/Shaders/Private/LightAccumulator.ush
@@ -4,8 +4,7 @@
 	LightAccumulator.usf: FLightAccumulator "class" and it's methods, useful for screen space subsurface scattering
 =============================================================================*/
 
-#ifndef __LIGHTACCUMULATOR_COMMON__
-#define __LIGHTACCUMULATOR_COMMON__
+#pragma once
 
 // set by c++, not set for LPV
 // 0 / 1
@@ -100,4 +99,3 @@ float4 LightAccumulator_GetResult(FLightAccumulator In)
 	return Ret;
 }
 
-#endif // __LIGHTACCUMULATOR_COMMON__
\ No newline at end of file
diff --git a/Engine/Shaders/Private/PostProcessCommon.ush b/Engine/Shaders/Private/PostProcessCommon.ush
--- a/Engine/Shaders/Private/PostProcessCommon.ush
+++ b/Engine/Shaders/Private/PostProcessCommon.ush
@@ -4,8 +4,7 @@
 	PostProcessCommon.usf: PostProcessing shared functions and structures.
 =============================================================================*/
 
-#ifndef __POST_PROCESS_COMMON__
-#define __POST_PROCESS_COMMON__
+#pragma once
 
 #include "SceneTexturesCommon.ush"
 
@@ -232,4 +231,3 @@ bool IsComputeUVOutOfBounds(in float2 UV)
 	return (max(CenterDist.x, CenterDist.y) >= 0.5f);
 }
 
-#endif // __POST_PROCESS_COMMON__
diff --git a/Engine/Shaders/Private/Random.ush b/Engine/Shaders/Private/Random.ush
--- a/Engine/Shaders/Private/Random.ush
+++ b/Engine/Shaders/Private/Random.ush
@@ -4,8 +4,7 @@
 	Random.ush: A pseudo-random number generator.
 =============================================================================*/
 
-#ifndef __Random_ush__
-#define __Random_ush__
+#pragma once
 
 // @param xy should be a integer position (e.g. pixel position on the screen), repeats each 128x128 pixels
 // similar to a texture lookup but is only ALU
@@ -891,7 +890,3 @@ float VolumeRaymarch(float3 posPixelWS, float3 posCameraWS)
 }
 
 
-
-
-
-#endif
diff --git a/Engine/Shaders/Private/RectLight.ush b/Engine/Shaders/Private/RectLight.ush
--- a/Engine/Shaders/Private/RectLight.ush
+++ b/Engine/Shaders/Private/RectLight.ush
@@ -1,7 +1,7 @@
 // Copyright 1998-2019 Epic Games, Inc. All Rights Reserved.
 //
-#ifndef __RECTLIGHT__
-#define __RECTLIGHT__
+
+#pragma once
 
 #include "CapsuleLight.ush"
 #include "MonteCarlo.ush"
@@ -725,4 +725,3 @@ FRect GetRect(
 	return Rect;
 }
 
-#endif // __RECTLIGHT__
diff --git a/Engine/Shaders/Private/ShadingCommon.ush b/Engine/Shaders/Private/ShadingCommon.ush
--- a/Engine/Shaders/Private/ShadingCommon.ush
+++ b/Engine/Shaders/Private/ShadingCommon.ush
@@ -1,8 +1,6 @@
 // Copyright 1998-2019 Epic Games, Inc. All Rights Reserved.
 
-//#pragma once
-#ifndef __SHADING_COMMON__
-#define __SHADING_COMMON__
+#pragma once
 
 #define SHADINGMODELID_UNLIT				0
 #define SHADINGMODELID_DEFAULT_LIT			1
@@ -108,4 +106,3 @@ float3 ComputeF0(float Specular, float3 BaseColor, float Metallic)
 }
 
 
-#endif // __SHADING_COMMON__
diff --git a/Engine/Shaders/Private/SubsurfaceProfileCommon.ush b/Engine/Shaders/Private/SubsurfaceProfileCommon.ush
--- a/Engine/Shaders/Private/SubsurfaceProfileCommon.ush
+++ b/Engine/Shaders/Private/SubsurfaceProfileCommon.ush
@@ -4,8 +4,7 @@
 	SubsurfaceProfileCommon.ush: Subsurface scattering parameter lookup / decoding constants.
 =============================================================================================*/
 
-#ifndef __SUBSURFACE_PROFILE_COMMON__
-#define __SUBSURFACE_PROFILE_COMMON__
+#pragma once
 
 // NOTE: Changing offsets below requires updating all instances of #SSSS_CONSTANTS
 // TODO: This needs to be defined in a single place and shared between C++ and shaders!
@@ -35,4 +34,3 @@ Texture2D SSProfilesTexture;
 #define ActualSSProfilesTexture SSProfilesTexture
 #endif
 
-#endif // __SUBSURFACE_PROFILE_COMMON__
