Fixes pre exposure that was multiplying the alpha channel in TAA.

diff --git a/Engine/Shaders/Private/PostProcessTemporalAA.usf b/Engine/Shaders/Private/PostProcessTemporalAA.usf
--- a/Engine/Shaders/Private/PostProcessTemporalAA.usf
+++ b/Engine/Shaders/Private/PostProcessTemporalAA.usf
@@ -801,10 +801,9 @@ float4 TransformBackToRawLinearSceneColor(float4 SceneColor)
 // Transform current frame's RAW scene color RGB to TAA's working color space.
 float4 TransformCurrentFrameSceneColor(float4 RawSceneColorRGBA)
 {
+	float4 RawLinearSceneColorRGBA = RawSceneColorRGBA;
 	#if USE_PREEXPOSURE
-		float4 RawLinearSceneColorRGBA = RawSceneColorRGBA *  PreExposureSettings.y;
-	#else
-		float4 RawLinearSceneColorRGBA = RawSceneColorRGBA;
+		RawLinearSceneColorRGBA.rgb *= PreExposureSettings.y;
 	#endif
 	return TransformSceneColor(RawLinearSceneColorRGBA);
 }
