Removes TAA permutation for gaussian DOF that is no longer used because mobile renderer only.

diff --git a/Engine/Shaders/Private/PostProcessTemporalAA.usf b/Engine/Shaders/Private/PostProcessTemporalAA.usf
--- a/Engine/Shaders/Private/PostProcessTemporalAA.usf
+++ b/Engine/Shaders/Private/PostProcessTemporalAA.usf
@@ -34,9 +34,6 @@
 	// Have RGB, opacity in alpha and DOF's CocRadius in separate's Red.
 	#define HISTORY_PAYLOAD_RGB_OPACITY_COC 3
 	
-	// Have RGB, opacity in alpha and old DOF's CocRadius in alpha.
-	#define HISTORY_PAYLOAD_RGB_COC_LEGACY 4
-
 
 /** Caching method for scene color. */
 	// Disable any in code cache.
@@ -67,18 +64,7 @@
 	#define AA_SINGLE_PASS_RESPONSIVE 1
 #endif
 
-#if TAA_PASS_CONFIG == 0 // Used for DOF
-	#define AA_HISTORY_PAYLOAD (HISTORY_PAYLOAD_RGB_COC_LEGACY)
-	#define AA_BORDER 1
-	#define AA_CROSS 4 // because running at half res.
-	#define AA_DYNAMIC 1
-	#define AA_FILTERED 0 // TODO Turn back on when Catmull-Rom weights can be disabled
-	#define AA_FORCE_ALPHA_CLAMP 1
-	#define AA_MANUALLY_CLAMP_HISTORY_UV 1
-	#define AA_YCOCG 1
-	#define AA_LOWER_RESOLUTION 1
-
-#elif TAA_PASS_CONFIG == 1 // Main
+#if TAA_PASS_CONFIG == 0 // Main
 	#if POST_PROCESS_ALPHA
 		#define AA_HISTORY_PAYLOAD (HISTORY_PAYLOAD_RGB_TRANSLUCENCY)
 	#else
@@ -101,7 +87,7 @@
 		#define AA_SAMPLE_CACHE_METHOD (AA_SAMPLE_CACHE_METHOD_LDS)
 	#endif
 
-#elif TAA_PASS_CONFIG == 2 // Used for SSR
+#elif TAA_PASS_CONFIG == 3 // Used for SSR
 	#define AA_HISTORY_PAYLOAD (HISTORY_PAYLOAD_RGB_OPACITY)
 	#define AA_DYNAMIC 1
 	#define AA_FILTERED 1
@@ -110,7 +96,7 @@
 	#define AA_YCOCG 1
 	#define AA_BORDER 1
 
-#elif TAA_PASS_CONFIG == 3 // Used for LightShaft
+#elif TAA_PASS_CONFIG == 4 // Used for LightShaft
 	#define AA_HISTORY_PAYLOAD (HISTORY_PAYLOAD_RGB_OPACITY)
 	#define AA_FILTERED 1
 	#define AA_LERP 64
@@ -119,7 +105,7 @@
 	#define AA_LOWER_RESOLUTION 1
 	#define AA_BORDER 1
 
-#elif TAA_PASS_CONFIG == 4 || TAA_PASS_CONFIG == 7 // MainUpsampling & MainSuperSampling
+#elif TAA_PASS_CONFIG == 1 || TAA_PASS_CONFIG == 2 // MainUpsampling & MainSuperSampling
 	#if POST_PROCESS_ALPHA
 		#define AA_HISTORY_PAYLOAD (HISTORY_PAYLOAD_RGB_TRANSLUCENCY)
 	#else
@@ -874,7 +860,7 @@ float GetSceneColorHdrWeight(
 		Sample.Color = InputParams.CachedNeighbors0[NeighborsId];
 		Sample.CocRadius = 0;
 
-		#if AA_HISTORY_PAYLOAD == HISTORY_PAYLOAD_RGB_COC || AA_HISTORY_PAYLOAD == HISTORY_PAYLOAD_RGB_COC_LEGACY
+		#if AA_HISTORY_PAYLOAD == HISTORY_PAYLOAD_RGB_COC
 			Sample.CocRadius = Sample.Color.a;
 		#elif AA_HISTORY_PAYLOAD == HISTORY_PAYLOAD_RGB_OPACITY_COC
 			Sample.CocRadius = InputParams.CachedNeighbors1[NeighborsId];
@@ -1353,7 +1339,7 @@ FTAASceneColorSample SampleCachedSceneColorTexture(
 
 	Sample.CocRadius = 0;
 
-	#if AA_HISTORY_PAYLOAD == HISTORY_PAYLOAD_RGB_COC || AA_HISTORY_PAYLOAD == HISTORY_PAYLOAD_RGB_COC_LEGACY
+	#if AA_HISTORY_PAYLOAD == HISTORY_PAYLOAD_RGB_COC
 		Sample.CocRadius = Sample.Color.a;
 	#elif AA_HISTORY_PAYLOAD == HISTORY_PAYLOAD_RGB_OPACITY_COC
 		Sample.CocRadius = InputSceneMetadata.SampleLevel(
@@ -2285,7 +2271,7 @@ void MainCS(
 	float4 OutColor0 = 0;
 	float4 OutColor1 = 0;
 
-	#if AA_HISTORY_PAYLOAD == HISTORY_PAYLOAD_RGB_COC || AA_HISTORY_PAYLOAD == HISTORY_PAYLOAD_RGB_COC_LEGACY
+	#if AA_HISTORY_PAYLOAD == HISTORY_PAYLOAD_RGB_COC
 	{
 		OutColor0.rgb = OutputPayload.Color.rgb;
 		OutColor0.a = OutputPayload.CocRadius;
@@ -2369,7 +2355,7 @@ void MainPS(
 
 	FTAAHistoryPayload OutputPayload = TemporalAASample(GroupId, GroupThreadId, GroupThreadIndex, ViewportUV, FrameExposureScale);
 
-	#if AA_HISTORY_PAYLOAD == HISTORY_PAYLOAD_RGB_COC || AA_HISTORY_PAYLOAD == HISTORY_PAYLOAD_RGB_COC_LEGACY
+	#if AA_HISTORY_PAYLOAD == HISTORY_PAYLOAD_RGB_COC
 	{
 		OutColor0.rgb = OutputPayload.Color.rgb;
 		OutColor0.a = OutputPayload.CocRadius;
diff --git a/Engine/Source/Runtime/Renderer/Private/PostProcess/PostProcessTemporalAA.cpp b/Engine/Source/Runtime/Renderer/Private/PostProcess/PostProcessTemporalAA.cpp
--- a/Engine/Source/Runtime/Renderer/Private/PostProcess/PostProcessTemporalAA.cpp
+++ b/Engine/Source/Runtime/Renderer/Private/PostProcess/PostProcessTemporalAA.cpp
@@ -369,25 +369,23 @@ DECLARE_GPU_STAT(TAA)
 
 
 const TCHAR* const kTAAOutputNames[] = {
-	TEXT("DOFTemporalAA"),
+	TEXT("TemporalAA"),
+	TEXT("TemporalAA"),
 	TEXT("TemporalAA"),
 	TEXT("SSRTemporalAA"),
 	TEXT("LightShaftTemporalAA"),
-	TEXT("TemporalAA"),
 	TEXT("DOFTemporalAA"),
 	TEXT("DOFTemporalAA"),
-	TEXT("TemporalAA"),
 };
 
 const TCHAR* const kTAAPassNames[] = {
-	TEXT("LegacyDOF"),
 	TEXT("Main"),
+	TEXT("MainUpsampling"),
+	TEXT("MainSuperSampling"),
 	TEXT("ScreenSpaceReflections"),
 	TEXT("LightShaft"),
-	TEXT("MainUpsampling"),
 	TEXT("DiaphragmDOF"),
 	TEXT("DiaphragmDOFUpsampling"),
-	TEXT("MainSuperSampling"),
 };
 
 
diff --git a/Engine/Source/Runtime/Renderer/Private/PostProcess/PostProcessTemporalAA.h b/Engine/Source/Runtime/Renderer/Private/PostProcess/PostProcessTemporalAA.h
--- a/Engine/Source/Runtime/Renderer/Private/PostProcess/PostProcessTemporalAA.h
+++ b/Engine/Source/Runtime/Renderer/Private/PostProcess/PostProcessTemporalAA.h
@@ -22,14 +22,20 @@ struct FTemporalAAHistory;
 /** Lists of TAA configurations. */
 enum class ETAAPassConfig
 {
-	LegacyDepthOfField,
+	// Permutations for main scene color TAA.
 	Main,
+	MainUpsampling,
+	MainSuperSampling,
+
+	// Permutation for SSR noise accumulation.
 	ScreenSpaceReflections,
+	
+	// Permutation for light shaft noise accumulation.
 	LightShaft,
-	MainUpsampling,
+
+	// Permutation for DOF that handle Coc.
 	DiaphragmDOF,
 	DiaphragmDOFUpsampling,
-	MainSuperSampling,
 
 	MAX
 };
diff --git a/Engine/Source/Runtime/Renderer/Private/PostProcess/PostProcessing.cpp b/Engine/Source/Runtime/Renderer/Private/PostProcess/PostProcessing.cpp
--- a/Engine/Source/Runtime/Renderer/Private/PostProcess/PostProcessing.cpp
+++ b/Engine/Source/Runtime/Renderer/Private/PostProcess/PostProcessing.cpp
@@ -490,36 +490,10 @@ static bool AddPostProcessDepthOfFieldGaussian(FPostprocessContext& Context, FDe
 		{
 			FSceneViewState* ViewState = (FSceneViewState*)Context.View.State;
 
-			FRenderingCompositeOutputRef DOFInputPass = DOFSetup;
-			const bool bMobileQuality = (Context.View.GetFeatureLevel() <= ERHIFeatureLevel::ES3_1);
-			if (Context.View.AntiAliasingMethod == AAM_TemporalAA && ViewState && !bMobileQuality)
-			{
-				// If no history use current as history
-				FTemporalAAHistory* DOFHistory = bFarPass ? &ViewState->DOFHistory : &ViewState->DOFHistory2;
-				bool& bDOFHistory = bFarPass ? ViewState->bDOFHistory : ViewState->bDOFHistory2;
-
-				const bool bIsComputePass = ShouldDoComputePostProcessing(Context.View);
-
-
-				FTAAPassParameters Parameters(Context.View);
-				Parameters.Pass = ETAAPassConfig::LegacyDepthOfField;
-				Parameters.bIsComputePass = bIsComputePass;
-				Parameters.SetupViewRect(Context.View, /* ResolutionDivisor = */ 2);
-
-				FRenderingCompositePass* NodeTemporalAA = Context.Graph.RegisterPass(new (FMemStack::Get()) FRCPassPostProcessTemporalAA(
-					Context, Parameters, *DOFHistory, DOFHistory));
-
-				NodeTemporalAA->SetInput(ePId_Input0, DOFSetup);
-				NodeTemporalAA->SetInput(ePId_Input2, VelocityInput);
-
-				DOFInputPass = FRenderingCompositeOutputRef(NodeTemporalAA);
-				bDOFHistory = false;
-			}
-
 			const TCHAR* BlurDebugX = bFarPass ? TEXT("FarDOFBlurX") : TEXT("NearDOFBlurX");
 			const TCHAR* BlurDebugY = bFarPass ? TEXT("FarDOFBlurY") : TEXT("NearDOFBlurY");
 
-			return RenderGaussianBlur(Context, BlurDebugX, BlurDebugY, DOFInputPass, BlurSize);
+			return RenderGaussianBlur(Context, BlurDebugX, BlurDebugY, DOFSetup, BlurSize);
 		};
 
 		const bool bFar = FarSize > 0.0f;
