Fixes color pre exposure in screen space denoiser's history reprojection.

diff --git a/Engine/Shaders/Private/ScreenSpaceDenoise/SSDTemporalAccumulation.usf b/Engine/Shaders/Private/ScreenSpaceDenoise/SSDTemporalAccumulation.usf
--- a/Engine/Shaders/Private/ScreenSpaceDenoise/SSDTemporalAccumulation.usf
+++ b/Engine/Shaders/Private/ScreenSpaceDenoise/SSDTemporalAccumulation.usf
@@ -174,6 +174,8 @@
 
 int bCameraCut[MAX_SIGNAL_BATCH_SIZE];
 
+float HistoryPreExposureCorrection;
+
 static const uint iMipLevel = 0;
 static const uint iMipLevelPow2 = 1;
 static const float MipLevel = 0;
@@ -354,17 +356,20 @@ void TemporallyAccumulate(
 				}
 			}
 
-			// Override the previous frame's inverse frequency with the current frame's one, because don't know how much it changed.
 			FIX_UNROLL(CONFIG_SIGNAL_BATCH_SIZE)
 			for (uint BatchedSignalId = 0; BatchedSignalId < CONFIG_SIGNAL_BATCH_SIZE; BatchedSignalId++)
 			{
 				FSSFSignalSample CurrentFrameSample = CurrentFrameSamples.Array[BatchedSignalId];
 				FSSFSignalSample HistorySample = HistorySamples.Array[BatchedSignalId];
 
+				// Override the previous frame's inverse frequency with the current frame's one, because don't know how much it changed.
 				HistorySamples.Array[BatchedSignalId].WorldBluringRadius = (
 					HistorySample.AccumulatedSampleCount *
 					CurrentFrameSample.WorldBluringRadius *
 					SafeRcp(CurrentFrameSample.AccumulatedSampleCount));
+
+				// Applies correction when pre exposure change happen.
+				HistorySamples.Array[BatchedSignalId].SceneColor.rgb *= HistoryPreExposureCorrection;
 			} // for (uint BatchedSignalId = 0; BatchedSignalId < CONFIG_SIGNAL_BATCH_SIZE; BatchedSignalId++)
 		} // if (!bIsPreviousFrameOffscreen)
 	}
diff --git a/Engine/Source/Runtime/Renderer/Private/PostProcess/PostProcessEyeAdaptation.cpp b/Engine/Source/Runtime/Renderer/Private/PostProcess/PostProcessEyeAdaptation.cpp
--- a/Engine/Source/Runtime/Renderer/Private/PostProcess/PostProcessEyeAdaptation.cpp
+++ b/Engine/Source/Runtime/Renderer/Private/PostProcess/PostProcessEyeAdaptation.cpp
@@ -495,6 +495,12 @@ void FSceneViewState::UpdatePreExposure(FViewInfo& View)
 
 	// Update the pre-exposure value on the actual view
 	View.PreExposure = PreExposure;
+
+	// Update the pre exposure of all temporal histories.
+	if (!View.bViewStateIsReadOnly)
+	{
+		PrevFrameViewInfo.SceneColorPreExposure = PreExposure;
+	}
 }
 
 FPooledRenderTargetDesc FRCPassPostProcessEyeAdaptation::ComputeOutputDesc(EPassOutputId InPassOutputId) const
diff --git a/Engine/Source/Runtime/Renderer/Private/PostProcess/PostProcessTemporalAA.cpp b/Engine/Source/Runtime/Renderer/Private/PostProcess/PostProcessTemporalAA.cpp
--- a/Engine/Source/Runtime/Renderer/Private/PostProcess/PostProcessTemporalAA.cpp
+++ b/Engine/Source/Runtime/Renderer/Private/PostProcess/PostProcessTemporalAA.cpp
@@ -557,7 +557,7 @@ FTAAOutputs FTAAPassParameters::AddTemporalAAPass(
 		{
 			CommonShaderParameters.PreExposureSettings.X = View.PreExposure;
 			CommonShaderParameters.PreExposureSettings.Y = 1.f / FMath::Max<float>(SMALL_NUMBER, View.PreExposure);
-			CommonShaderParameters.PreExposureSettings.Z = InputHistory.IsValid() ? InputHistory.SceneColorPreExposure : View.PreExposure;
+			CommonShaderParameters.PreExposureSettings.Z = View.PrevViewInfo.SceneColorPreExposure;
 			CommonShaderParameters.PreExposureSettings.W  = 1.f / FMath::Max<float>(SMALL_NUMBER, CommonShaderParameters.PreExposureSettings.Z);
 		}
 
@@ -799,7 +799,6 @@ FTAAOutputs FTAAPassParameters::AddTemporalAAPass(
 
 		OutputHistory->ViewportRect = DestRect;
 		OutputHistory->ReferenceBufferSize = OutputExtent * ResolutionDivisor;
-		OutputHistory->SceneColorPreExposure = View.PreExposure;
 	}
 
 	return Outputs;
diff --git a/Engine/Source/Runtime/Renderer/Private/PostProcess/ScreenSpaceReflections.cpp b/Engine/Source/Runtime/Renderer/Private/PostProcess/ScreenSpaceReflections.cpp
--- a/Engine/Source/Runtime/Renderer/Private/PostProcess/ScreenSpaceReflections.cpp
+++ b/Engine/Source/Runtime/Renderer/Private/PostProcess/ScreenSpaceReflections.cpp
@@ -322,7 +322,7 @@ class FPostProcessScreenSpaceReflectionsPS : public FGlobalShader
 
 			if (Context.View.PrevViewInfo.TemporalAAHistory.IsValid())
 			{
-				PrevSceneColorPreExposureCorrectionValue = Context.View.PreExposure / Context.View.PrevViewInfo.TemporalAAHistory.SceneColorPreExposure;
+				PrevSceneColorPreExposureCorrectionValue = Context.View.PreExposure / Context.View.PrevViewInfo.SceneColorPreExposure;
 			}
 
 			SetShaderValue(RHICmdList, ShaderRHI, PrevSceneColorPreExposureCorrection, PrevSceneColorPreExposureCorrectionValue);
diff --git a/Engine/Source/Runtime/Renderer/Private/SceneRendering.h b/Engine/Source/Runtime/Renderer/Private/SceneRendering.h
--- a/Engine/Source/Runtime/Renderer/Private/SceneRendering.h
+++ b/Engine/Source/Runtime/Renderer/Private/SceneRendering.h
@@ -593,9 +593,6 @@ struct FTemporalAAHistory
 	// Viewport coordinate of the history in RT according to ReferenceBufferSize.
 	FIntRect ViewportRect;
 
-	// Scene color's PreExposure.
-	float SceneColorPreExposure;
-
 
 	void SafeRelease()
 	{
@@ -643,6 +640,9 @@ struct FPreviousViewInfo
 	// View matrices.
 	FViewMatrices ViewMatrices;
 
+	// Scene color's PreExposure.
+	float SceneColorPreExposure = 1.0f;
+
 	// Depth buffer and Normals of the previous frame generating this history entry for bilateral kernel rejection.
 	TRefCountPtr<IPooledRenderTarget> DepthBuffer;
 	TRefCountPtr<IPooledRenderTarget> GBufferA;
diff --git a/Engine/Source/Runtime/Renderer/Private/ScreenSpaceDiffuseIndirect.cpp b/Engine/Source/Runtime/Renderer/Private/ScreenSpaceDiffuseIndirect.cpp
--- a/Engine/Source/Runtime/Renderer/Private/ScreenSpaceDiffuseIndirect.cpp
+++ b/Engine/Source/Runtime/Renderer/Private/ScreenSpaceDiffuseIndirect.cpp
@@ -140,7 +140,7 @@ void RenderScreenSpaceDiffuseIndirect( FRHICommandListImmediate& RHICmdList, FVi
 		(ViewportExtent.X * 0.5f + ViewportOffset.X) / BufferSize.X,
 		(ViewportExtent.Y * 0.5f + ViewportOffset.Y) / BufferSize.Y );
 
-	PassParameters->PrevSceneColorPreExposureCorrection = View.PreExposure / TemporalAAHistory.SceneColorPreExposure;
+	PassParameters->PrevSceneColorPreExposureCorrection = View.PreExposure / View.PrevViewInfo.SceneColorPreExposure;
 
 	PassParameters->View = View.ViewUniformBuffer;
 	PassParameters->SceneTextures = CreateSceneTextureUniformBufferSingleDraw( RHICmdList, ESceneTextureSetupMode::SceneDepth | ESceneTextureSetupMode::GBuffers, View.FeatureLevel );
diff --git a/Engine/Source/Runtime/Renderer/Private/TranslucentRendering.cpp b/Engine/Source/Runtime/Renderer/Private/TranslucentRendering.cpp
--- a/Engine/Source/Runtime/Renderer/Private/TranslucentRendering.cpp
+++ b/Engine/Source/Runtime/Renderer/Private/TranslucentRendering.cpp
@@ -592,12 +592,12 @@ void CreateTranslucentBasePassUniformBuffer(
 			if (View.PrevViewInfo.CustomSSRInput.IsValid())
 			{
 				PrevSceneColorRT = &View.PrevViewInfo.CustomSSRInput;
-				PrevSceneColorPreExposureInvValue = 1.0f / View.PrevViewInfo.TemporalAAHistory.SceneColorPreExposure;
+				PrevSceneColorPreExposureInvValue = 1.0f / View.PrevViewInfo.SceneColorPreExposure;
 			}
 			else if (View.PrevViewInfo.TemporalAAHistory.IsValid())
 			{
 				PrevSceneColorRT = &View.PrevViewInfo.TemporalAAHistory.RT[0];
-				PrevSceneColorPreExposureInvValue = 1.0f / View.PrevViewInfo.TemporalAAHistory.SceneColorPreExposure;
+				PrevSceneColorPreExposureInvValue = 1.0f / View.PrevViewInfo.SceneColorPreExposure;
 			}
 
 			BasePassParameters.PrevSceneColor = (*PrevSceneColorRT)->GetRenderTargetItem().ShaderResourceTexture;
