Fixes exposure based color weighting in DOF, TAA, Denoiser when pre exposure is used.

diff --git a/Engine/Shaders/Private/DiaphragmDOF/DOFDownsample.ush b/Engine/Shaders/Private/DiaphragmDOF/DOFDownsample.ush
--- a/Engine/Shaders/Private/DiaphragmDOF/DOFDownsample.ush
+++ b/Engine/Shaders/Private/DiaphragmDOF/DOFDownsample.ush
@@ -113,6 +113,10 @@ float DownsampleCoc(float CocRadii[4])
 // Optimized HDR weighting function.
 float HdrWeight4(float3 Color, float Exposure) 
 {
+#if USE_PREEXPOSURE
+	Exposure *= View.PreExposure;
+#endif
+
 	return rcp(Luma4(Color) * Exposure + 4.0);
 }
 
diff --git a/Engine/Shaders/Private/DiaphragmDOF/DOFReduce.usf b/Engine/Shaders/Private/DiaphragmDOF/DOFReduce.usf
--- a/Engine/Shaders/Private/DiaphragmDOF/DOFReduce.usf
+++ b/Engine/Shaders/Private/DiaphragmDOF/DOFReduce.usf
@@ -371,7 +371,11 @@ void ReduceCS(
 	#endif
 
 	// Sample the frame exposure to SGPR.
-	float FrameExposureScale = ToScalarMemory(EyeAdaptationLookup());
+	#if USE_PREEXPOSURE
+		float FrameExposureScale = ToScalarMemory(EyeAdaptationLookup() * View.OneOverPreExposure);
+	#else
+		float FrameExposureScale = ToScalarMemory(EyeAdaptationLookup());
+	#endif
 
 	// Indirect dispatch 
 	#if DIM_HYBRID_SCATTER_FGD || DIM_HYBRID_SCATTER_BGD
diff --git a/Engine/Shaders/Private/PostProcessMitchellNetravali.usf b/Engine/Shaders/Private/PostProcessMitchellNetravali.usf
--- a/Engine/Shaders/Private/PostProcessMitchellNetravali.usf
+++ b/Engine/Shaders/Private/PostProcessMitchellNetravali.usf
@@ -98,8 +98,11 @@ RWTexture2D<float4>	Output0;
 [numthreads(THREADGROUP_SIZEX, THREADGROUP_SIZEY, 1)]
 void DownsampleMainCS(uint2 DispatchThreadId : SV_DispatchThreadID)
 {
-	// Frame exposure
-	const float FrameExposureScale = EyeAdaptationLookup();
+	#if USE_PREEXPOSURE
+		const float FrameExposureScale = ToScalarMemory(EyeAdaptationLookup() * View.OneOverPreExposure);
+	#else
+		const float FrameExposureScale = ToScalarMemory(EyeAdaptationLookup());
+	#endif
 
 	const float B = rcp(3.0);
 	const float C = rcp(3.0);
diff --git a/Engine/Shaders/Private/PostProcessTemporalAA.usf b/Engine/Shaders/Private/PostProcessTemporalAA.usf
--- a/Engine/Shaders/Private/PostProcessTemporalAA.usf
+++ b/Engine/Shaders/Private/PostProcessTemporalAA.usf
@@ -441,9 +441,7 @@ static const uint kNeighborsCount = 9;
 
 //------------------------------------------------------- PARAMETERS
 
-// (Pre-Exposure, 1 / Pre-exposure, LastFramePreExposure, 1 / LastFramePreExposure)
-float4 PreExposureSettings;
-
+float HistoryPreExposureCorrection;
 float CurrentFrameWeight;
 int bCameraCut;
 
@@ -789,11 +787,7 @@ float4 TransformBackToRawLinearSceneColor(float4 SceneColor)
 // Transform current frame's RAW scene color RGB to TAA's working color space.
 float4 TransformCurrentFrameSceneColor(float4 RawSceneColorRGBA)
 {
-	float4 RawLinearSceneColorRGBA = RawSceneColorRGBA;
-	#if USE_PREEXPOSURE
-		RawLinearSceneColorRGBA.rgb *= PreExposureSettings.y;
-	#endif
-	return TransformSceneColor(RawLinearSceneColorRGBA);
+	return TransformSceneColor(RawSceneColorRGBA);
 }
 
 // Get the Luma4 of the sceneColor
@@ -1836,7 +1830,7 @@ FTAAHistoryPayload SampleHistory(in float2 HistoryScreenPosition)
 	#endif
 
 	#if USE_PREEXPOSURE
-		HistoryPayload.Color.rgb *= PreExposureSettings.w;
+		HistoryPayload.Color.rgb *= HistoryPreExposureCorrection;
 	#endif
 
 	HistoryPayload.Color = TransformSceneColor(HistoryPayload.Color);
@@ -1887,7 +1881,12 @@ FTAAHistoryPayload TemporalAASample(uint2 GroupId, uint2 GroupThreadId, uint Gro
 	FTAAInputParameters InputParams;
 
 	// Per frame setup.
-	InputParams.FrameExposureScale = FrameExposureScale;
+	#if USE_PREEXPOSURE
+		InputParams.FrameExposureScale = ToScalarMemory(FrameExposureScale * View.OneOverPreExposure);
+	#else
+		InputParams.FrameExposureScale = ToScalarMemory(FrameExposureScale);
+	#endif
+
 
 	// Per pixel setup.
 	{
@@ -2217,10 +2216,6 @@ FTAAHistoryPayload TemporalAASample(uint2 GroupId, uint2 GroupThreadId, uint Gro
 
 	OutputPayload.Color = TransformBackToRawLinearSceneColor(OutputPayload.Color);
 	
-	#if USE_PREEXPOSURE
-		OutputPayload.Color.rgb *= PreExposureSettings.x;
-	#endif
-
 	#if AA_NAN
 		// Transform NaNs to black, transform negative colors to black.
 		OutputPayload.Color = -min(-OutputPayload.Color, 0.0);
@@ -2348,7 +2343,7 @@ void MainPS(
 	uint2 GroupThreadId = int2(0, 0);
 	uint GroupThreadIndex = 0;
 	float FrameExposureScale = EyeAdaptationLookup();
-
+	
 	FTAAHistoryPayload OutputPayload = TemporalAASample(GroupId, GroupThreadId, GroupThreadIndex, ViewportUV, FrameExposureScale);
 
 	#if AA_HISTORY_PAYLOAD == HISTORY_PAYLOAD_RGB_COC
diff --git a/Engine/Shaders/Private/ScreenSpaceDenoise/SSDSignalFramework.ush b/Engine/Shaders/Private/ScreenSpaceDenoise/SSDSignalFramework.ush
--- a/Engine/Shaders/Private/ScreenSpaceDenoise/SSDSignalFramework.ush
+++ b/Engine/Shaders/Private/ScreenSpaceDenoise/SSDSignalFramework.ush
@@ -150,8 +150,12 @@ float GetSignalColorTransformationExposure()
 	const float Tweak = 1;
 
 	float FrameExposureScale = EyeAdaptationLookup();
+	
+	#if USE_PREEXPOSURE
+		FrameExposureScale *= View.OneOverPreExposure;
+	#endif
 
-	return FrameExposureScale * Tweak;
+	return ToScalarMemory(FrameExposureScale * Tweak);
 }
 
 /** Returns the linear color space Luma with a 4 multiplier for ALU perf reasons. */
diff --git a/Engine/Source/Runtime/Renderer/Private/PostProcess/PostProcessTemporalAA.cpp b/Engine/Source/Runtime/Renderer/Private/PostProcess/PostProcessTemporalAA.cpp
--- a/Engine/Source/Runtime/Renderer/Private/PostProcess/PostProcessTemporalAA.cpp
+++ b/Engine/Source/Runtime/Renderer/Private/PostProcess/PostProcessTemporalAA.cpp
@@ -69,10 +69,10 @@ static float CatmullRom( float x )
 
 
 BEGIN_SHADER_PARAMETER_STRUCT(FTAAShaderParameters,)
-	SHADER_PARAMETER(FVector4, PreExposureSettings)
 	SHADER_PARAMETER(FVector4, ViewportUVToInputBufferUV)
 	SHADER_PARAMETER(FVector4, MaxViewportUVAndSvPositionToViewportUV)
 	SHADER_PARAMETER(FVector2D, ScreenPosAbsMax)
+	SHADER_PARAMETER(float, HistoryPreExposureCorrection)
 	SHADER_PARAMETER(float, CurrentFrameWeight)
 	SHADER_PARAMETER(int32, bCameraCut)
 	
@@ -550,16 +550,7 @@ FTAAOutputs FTAAPassParameters::AddTemporalAAPass(
 			ResDivisor / float(DestRect.Width()),
 			ResDivisor / float(DestRect.Height()));
 
-		// Pre-exposure, One over Pre-exposure, History pre-exposure, History one over pre-exposure.
-		// DOF settings must preserve scene color range.
-		CommonShaderParameters.PreExposureSettings = FVector4(1.f, 1.f, 1.f, 1.f);
-		if (Pass == ETAAPassConfig::Main)
-		{
-			CommonShaderParameters.PreExposureSettings.X = View.PreExposure;
-			CommonShaderParameters.PreExposureSettings.Y = 1.f / FMath::Max<float>(SMALL_NUMBER, View.PreExposure);
-			CommonShaderParameters.PreExposureSettings.Z = View.PrevViewInfo.SceneColorPreExposure;
-			CommonShaderParameters.PreExposureSettings.W  = 1.f / FMath::Max<float>(SMALL_NUMBER, CommonShaderParameters.PreExposureSettings.Z);
-		}
+		CommonShaderParameters.HistoryPreExposureCorrection = View.PreExposure / View.PrevViewInfo.SceneColorPreExposure;
 
 		{
 			float InvSizeX = 1.0f / float(SrcSize.X);
