Removed some obsolete dxr_todos. No functionality changed

diff --git a/Engine/Shaders/Private/PathTracing/PathTracing.usf b/Engine/Shaders/Private/PathTracing/PathTracing.usf
--- a/Engine/Shaders/Private/PathTracing/PathTracing.usf
+++ b/Engine/Shaders/Private/PathTracing/PathTracing.usf
@@ -272,7 +272,6 @@ void PathTracingMainRG()
 			}
 
 			// Trace again
-			//#dxr-todo: check ray flags when hit shaders are supported
 			Payload = TraceRayInternal(TLAS,
 				RAY_FLAG_CULL_BACK_FACING_TRIANGLES,
 				RAY_TRACING_MASK_ALL,
diff --git a/Engine/Shaders/Private/RayTracing/RayTracingReflections.usf b/Engine/Shaders/Private/RayTracing/RayTracingReflections.usf
--- a/Engine/Shaders/Private/RayTracing/RayTracingReflections.usf
+++ b/Engine/Shaders/Private/RayTracing/RayTracingReflections.usf
@@ -242,10 +242,10 @@ void RayTracingReflectionsRGS()
 			RayCone = PropagateRayCone(RayCone, SurfaceCurvature, Depth);
 
 			//
-			// Reflections are trace by always following the top layer reflection vector. This is recursive and up to MaxBounces recursion.
+			// Reflections are traced by always following the top layer reflection vector. This is recursive and up to MaxBounces recursion.
 			// In the case of coated material, we also trace a reflection ray for the bottom layer which can have different normal and roughness. This is not recursive.
 			// We fix normal pointing away towards the geometry by assuming coated layer has smooth normal (assuming it is the mesh smooth normal).
-			// Fog is applied as a function of the top layer path, for both top and botto; refelctions combined (as an optimisation approximation).
+			// Fog is applied as a function of the top layer path, for both top and bottom; reflections combined (as an optimization approximation).
 			//
 			uint BounceIndex = 0;
 			RayDesc TopLayerRay;
@@ -266,8 +266,6 @@ void RayTracingReflectionsRGS()
 					FixSampleDirectionIfNeeded(ScreenSpaceData.GBuffer.WorldNormal, TopLayerRay.Direction);
 				}
 
-				//#dxr-todo: apply heuristics to do different shading computations depending on BounceIndex
-
 				if (BounceIndex == 0 && SampleIndex == 0)
 				{
 					if (DIM_DEFERRED_MATERIAL_MODE == DEFERRED_MATERIAL_MODE_GATHER)
diff --git a/Engine/Shaders/Private/RayTracing/RayTracingTranslucency.usf b/Engine/Shaders/Private/RayTracing/RayTracingTranslucency.usf
--- a/Engine/Shaders/Private/RayTracing/RayTracingTranslucency.usf
+++ b/Engine/Shaders/Private/RayTracing/RayTracingTranslucency.usf
@@ -26,7 +26,7 @@ float2 SobolIndexToUniformUnitSquare(uint2 SobolRand)
 #include "../DeferredLightingCommon.ush"
 #include "../ReflectionEnvironmentShared.ush"
 #include "../Montecarlo.ush"
-#include "../PathTracing/Utilities/PathTracingRandomSequence.ush" // #dxr-todo: move to a common folder
+#include "../PathTracing/Utilities/PathTracingRandomSequence.ush" 
 #include "../HeightFogCommon.ush"
 #include "RayTracingCommon.ush"
 #include "RayTracingHitGroupCommon.ush"
