Fix shader warnings.

diff --git a/Engine/Shaders/Private/NiagaraMeshVertexFactory.ush b/Engine/Shaders/Private/NiagaraMeshVertexFactory.ush
--- a/Engine/Shaders/Private/NiagaraMeshVertexFactory.ush
+++ b/Engine/Shaders/Private/NiagaraMeshVertexFactory.ush
@@ -593,7 +593,7 @@ FVertexFactoryIntermediates GetVertexFactoryIntermediates(FVertexFactoryInput In
 	// World position and localtoworld transform.
 	float4x4 LocalToWorldMat = float4x4(Transform1, Transform2, Transform3, float4(0.0f, 0.0f, 0.0f, 1.0f));
 
-	Intermediates.ParticleTranslatedWorldPosition = Intermediates.ParticlePosition + ResolvedView.PreViewTranslation.xyz;
+	Intermediates.ParticleTranslatedWorldPosition = Intermediates.ParticlePosition.xyz + ResolvedView.PreViewTranslation.xyz;
 
 	//LocalToWorldMat = mul(ScaleMatrix, LocalToWorldMat);
 
diff --git a/Engine/Shaders/Private/VectorFieldVisualizationVertexFactory.ush b/Engine/Shaders/Private/VectorFieldVisualizationVertexFactory.ush
--- a/Engine/Shaders/Private/VectorFieldVisualizationVertexFactory.ush
+++ b/Engine/Shaders/Private/VectorFieldVisualizationVertexFactory.ush
@@ -56,7 +56,7 @@ FVertexFactoryIntermediates GetVertexFactoryIntermediates( FVertexFactoryInput I
 	const float3 VoxelSize = VectorFieldVis.VoxelSize.xyz;
 	const float3 HalfVoxelOffset = VoxelSize.xyz * 0.5f;
 	const float3 VoxelMultipliers = float3(1, VoxelSize.x, VoxelSize.x * VoxelSize.y);
-	const float3 VoxelIndex = floor((float3)GetInstanceId(Input.InstanceId.xxx) * VoxelMultipliers.xyz);
+	const float3 VoxelIndex = floor((float3)GetInstanceId(Input.InstanceId).xxx * VoxelMultipliers.xyz);
 	const float3 VoxelUV = frac(VoxelIndex.xyz * VoxelSize.xyz) + HalfVoxelOffset;
 
 	// Read vector for this voxel.
