Fix missing scale factor when transforming from local to world space.

These functions RotateLocalToWorld and RotateLocalToWorld were introduced in CL5554937 for position and normal only stream support. They are not used anywhere else.

diff --git a/Engine/Shaders/Private/VertexFactoryCommon.ush b/Engine/Shaders/Private/VertexFactoryCommon.ush
--- a/Engine/Shaders/Private/VertexFactoryCommon.ush
+++ b/Engine/Shaders/Private/VertexFactoryCommon.ush
@@ -26,13 +26,21 @@ float4 TransformLocalToTranslatedWorld(float3 LocalPosition)
 
 float3 RotateLocalToWorld(float3 LocalDirection, uint PrimitiveId)
 {
-	float4x4 LocalToWorld = GetPrimitiveData(PrimitiveId).LocalToWorld;
-	return LocalToWorld[0].xyz * LocalDirection.xxx + LocalToWorld[1].xyz * LocalDirection.yyy + LocalToWorld[2].xyz * LocalDirection.zzz;	
+	const float4x4 LocalToWorld = GetPrimitiveData(PrimitiveId).LocalToWorld;
+	const float3 InvScale = GetPrimitiveData(PrimitiveId).InvNonUniformScaleAndDeterminantSign.xyz;
+	return 
+		InvScale.x * LocalToWorld[0].xyz * LocalDirection.xxx + 
+		InvScale.y * LocalToWorld[1].xyz * LocalDirection.yyy + 
+		InvScale.z * LocalToWorld[2].xyz * LocalDirection.zzz;	
 }
 
 float3 RotateLocalToWorld(float3 LocalDirection)
 {
-	return Primitive.LocalToWorld[0].xyz * LocalDirection.xxx + Primitive.LocalToWorld[1].xyz * LocalDirection.yyy + Primitive.LocalToWorld[2].xyz * LocalDirection.zzz;
+	const float3 InvScale = Primitive.InvNonUniformScaleAndDeterminantSign.xyz;
+	return 
+		InvScale.x * Primitive.LocalToWorld[0].xyz * LocalDirection.xxx + 
+		InvScale.y * Primitive.LocalToWorld[1].xyz * LocalDirection.yyy + 
+		InvScale.z * Primitive.LocalToWorld[2].xyz * LocalDirection.zzz;
 }
 
 // Octahedron Normal Vectors
