Add support for slope bias.

Current default value is set to 0 for all type of light to disable it.

diff --git a/Engine/Shaders/Private/ShadowDepthCommon.ush b/Engine/Shaders/Private/ShadowDepthCommon.ush
--- a/Engine/Shaders/Private/ShadowDepthCommon.ush
+++ b/Engine/Shaders/Private/ShadowDepthCommon.ush
@@ -24,7 +24,7 @@ struct FShadowDepthVSToPS
 
 #if PERSPECTIVE_CORRECT_DEPTH
 	float ShadowDepth : TEXCOORD6;
-
+	float DepthBias : TEXCOORD8;
 #elif !COMPILER_SUPPORTS_EMPTY_STRUCTS && !ONEPASS_POINTLIGHT_SHADOW
 	float Dummy : TEXCOORD6;
 #endif
@@ -40,6 +40,7 @@ struct FShadowDepthVSToPS
 struct FShadowDepthVSToGS : FShadowDepthVSToPS
 {
 	float4 GSPosition : TEXCOORD6;
+	float3 WorldNormal : TEXCOORD8;
 };
 #endif
 
diff --git a/Engine/Shaders/Private/ShadowDepthPixelShader.usf b/Engine/Shaders/Private/ShadowDepthPixelShader.usf
--- a/Engine/Shaders/Private/ShadowDepthPixelShader.usf
+++ b/Engine/Shaders/Private/ShadowDepthPixelShader.usf
@@ -83,11 +83,9 @@ void Main(
 #endif
 
 #if PERSPECTIVE_CORRECT_DEPTH
-	float DepthBias = PassStruct.ShadowParams.x;
-	float InvMaxSubjectDepth = PassStruct.ShadowParams.y;
-
+	const float InvMaxSubjectDepth = PassStruct.ShadowParams.w;
 	Inputs.ShadowDepth *= InvMaxSubjectDepth;
-	Inputs.ShadowDepth += DepthBias;
+	Inputs.ShadowDepth += Inputs.DepthBias;
 
 	OutDepth = saturate(Inputs.ShadowDepth);
 #endif
diff --git a/Engine/Shaders/Private/ShadowDepthVertexShader.usf b/Engine/Shaders/Private/ShadowDepthVertexShader.usf
--- a/Engine/Shaders/Private/ShadowDepthVertexShader.usf
+++ b/Engine/Shaders/Private/ShadowDepthVertexShader.usf
@@ -23,9 +23,19 @@
 #include "/Engine/Generated/VertexFactory.ush"
 #include "ShadowDepthCommon.ush"
 
-void SetShadowDepthOutputs(float4 WorldPosition, out float4 OutPosition, out float ShadowDepth)
+void SetShadowDepthOutputs(
+	float4x4 WorldToClipMatrix, 
+	float4x4 WorldToShadowMatrix, 
+	float4 WorldPosition, 
+	float3 WorldVertexNormal, 
+	out float4 OutPosition, 
+	out float ShadowDepth
+#if PERSPECTIVE_CORRECT_DEPTH
+	, out float OutDepthBias
+#endif
+)
 {
-	OutPosition = mul(WorldPosition, PassStruct.ProjectionMatrix);
+	OutPosition = mul(WorldPosition, WorldToClipMatrix);
 
 	// Clamp the vertex to the near plane if it is in front of the near plane
 	// This has problems if some vertices of a triangle get clamped and others do not, also causes artifacts with non-ortho projections
@@ -35,13 +45,34 @@ void SetShadowDepthOutputs(float4 WorldPosition, out float4 OutPosition, out flo
 		OutPosition.w = 1.0f;
 	}
 
+	#if ONEPASS_POINTLIGHT_SHADOW
+	const float3 ViewDirection = -normalize(mul(WorldPosition, WorldToShadowMatrix).xyz);
+	const float3 ViewNormal = mul(float4(WorldVertexNormal,0), WorldToShadowMatrix).xyz;
+	const float NoL = abs(dot(ViewDirection, ViewNormal));
+	#else
+	const float NoL = abs(dot(
+		float3(WorldToShadowMatrix[0].z, WorldToShadowMatrix[1].z, WorldToShadowMatrix[2].z),
+		WorldVertexNormal));
+	#endif
+
+	const float MaxSlopeDepthBias = PassStruct.ShadowParams.z;
+	const float Slope = min(abs(NoL) > 0 ? sqrt(saturate(1 - NoL*NoL)) / NoL : MaxSlopeDepthBias, MaxSlopeDepthBias);
+	
+	const float SlopeDepthBias = PassStruct.ShadowParams.y;
+	const float SlopeBias = SlopeDepthBias * Slope;
+
+	const float ConstantDepthBias = PassStruct.ShadowParams.x;
+	const float DepthBias = SlopeBias + ConstantDepthBias;
+
 	#if PERSPECTIVE_CORRECT_DEPTH
 		ShadowDepth = OutPosition.z;
+		OutDepthBias = DepthBias;
+	#elif ONEPASS_POINTLIGHT_SHADOW
+		ShadowDepth = 0;
+		//OutPosition.z += DepthBias;
 	#else
-		float DepthBias = PassStruct.ShadowParams.x;
-		float InvMaxSubjectDepth = PassStruct.ShadowParams.y;
-
 		// Output linear, normalized depth
+		const float InvMaxSubjectDepth = PassStruct.ShadowParams.w;
 		ShadowDepth = OutPosition.z * InvMaxSubjectDepth + DepthBias;
 		OutPosition.z = ShadowDepth * OutPosition.w;
 	#endif
@@ -92,7 +123,26 @@ void MainOnePassPointLightGS(triangle FShadowDepthVSToGS Input[3], inout Triangl
 		for (int VertexIndex = 0; VertexIndex < 3; VertexIndex++)
 		{
 			// Calculate the clip space position for each cube face
+			// Note: We don't apply any constant/slop bias here. A constant bias is applied in the shadow projection 
+			// code (see CubemapHardwarePCF() in ShadowProjectionCommon.ush)
+			// If one want to use this code path, don't forget to remove the bias in CubemapHardwarePCF().
+			#if 0
+			const float4 WorldPosition = Input[VertexIndex].GSPosition;
+			const float3 WorldVertexNormal = Input[VertexIndex].WorldNormal;
+			
+			float4 ClipSpacePosition;
+			float ShadowDepth;
+			SetShadowDepthOutputs(
+				PassStruct.ShadowViewProjectionMatrices[CubeFaceIndex],
+				PassStruct.ShadowViewMatrices[CubeFaceIndex], 
+				WorldPosition, 
+				WorldVertexNormal, 
+				ClipSpacePosition, 
+				ShadowDepth);
+			ClipSpacePositions[VertexIndex] = ClipSpacePosition;
+			#else
 			ClipSpacePositions[VertexIndex] = mul(Input[VertexIndex].GSPosition, PassStruct.ShadowViewProjectionMatrices[CubeFaceIndex]);
+			#endif
 		}
 		// Frustum culling, saves GPU time with high poly meshes
 		BRANCH
@@ -106,13 +156,13 @@ void MainOnePassPointLightGS(triangle FShadowDepthVSToGS Input[3], inout Triangl
 			{
 				Output.OutPosition = ClipSpacePositions[VertexIndex];
 
-				#if INTERPOLATE_VF_ATTRIBUTES
-					Output.PSInputs.FactoryInterpolants = Input[VertexIndex].FactoryInterpolants;
-				#endif
+			#if INTERPOLATE_VF_ATTRIBUTES
+				Output.PSInputs.FactoryInterpolants = Input[VertexIndex].FactoryInterpolants;
+			#endif
 
-#if INTERPOLATE_POSITION
+			#if INTERPOLATE_POSITION
 				Output.PSInputs.PixelPosition = Input[VertexIndex].GSPosition.xyz;
-#endif
+			#endif
 				OutStream.Append(Output);
 			}
 			OutStream.RestartStrip();
@@ -178,13 +228,18 @@ FShadowDepthDSToPS PassFinalizeTessellationOutput(FShadowDepthVSToDS Interpolant
 	O.Position = WorldPosition = mul(WorldPosition, PassStruct.ShadowViewProjectionMatrices[Interpolants.LayerIndex]);
 #endif
 #else
+	const float3 WorldNormal = MaterialParameters.TangentToWorld[2];
 
 	float ShadowDepth;
 	SetShadowDepthOutputs(
+		PassStruct.ProjectionMatrix,
+		PassStruct.ViewMatrix,
 		WorldPosition, 
+		WorldNormal,
 		O.Position, 
 #if PERSPECTIVE_CORRECT_DEPTH
 		O.ShadowDepth
+		O.DepthBias,
 #else
 		ShadowDepth
 #endif
@@ -230,6 +285,7 @@ void Main(
 	float3x3 TangentToLocal = VertexFactoryGetTangentToLocal(Input, VFIntermediates);
 
 	FMaterialVertexParameters VertexParameters = GetMaterialVertexParameters(Input, VFIntermediates, WorldPos.xyz, TangentToLocal);
+	const float3 WorldNormal = VertexParameters.TangentToWorld[2];
 	WorldPos.xyz += GetMaterialWorldPositionOffset(VertexParameters);
 
 #if USING_TESSELLATION
@@ -245,6 +301,7 @@ void Main(
 #elif ONEPASS_POINTLIGHT_SHADOW
 
 	OutPosition = OutParameters.GSPosition = WorldPos;
+	OutParameters.WorldNormal = WorldNormal;
 
 	#if INTERPOLATE_VF_ATTRIBUTES
 		// Masked materials need texture coords to clip
@@ -267,12 +324,16 @@ void Main(
 	float Dummy;
 
 	SetShadowDepthOutputs(
+		PassStruct.ProjectionMatrix,
+		PassStruct.ViewMatrix,
 		WorldPos, 
+		WorldNormal,
 		OutPosition, 
 #if !PERSPECTIVE_CORRECT_DEPTH
 		Dummy
 #else
-		OutParameters.ShadowDepth
+		OutParameters.ShadowDepth,
+		OutParameters.DepthBias
 #endif
 		);
 	
@@ -322,7 +383,7 @@ void MainForGS(
 
 #if POSITION_ONLY
 void PositionOnlyMain(
-	in FPositionOnlyVertexFactoryInput Input,
+	in FPositionAndNormalOnlyVertexFactoryInput Input,
 #if USING_ONEPASS_POINT_VERTEX_LAYER
 	out uint LayerIndex : SV_RenderTargetArrayIndex,
 #endif
@@ -342,9 +403,11 @@ void PositionOnlyMain(
 	OutParameters.FactoryInterpolants = (FVertexFactoryInterpolantsVSToPS)0;
 #endif
 
+	float3 WorldNormal = VertexFactoryGetWorldNormal(Input);
+
 #if ONEPASS_POINTLIGHT_SHADOW
 	OutPosition = OutParameters.GSPosition = WorldPos;
-
+	OutParameters.WorldNormal = WorldNormal;
 	#if USING_ONEPASS_POINT_VERTEX_LAYER
 		LayerIndex = LayerId;
 		OutPosition = mul(WorldPos, PassStruct.ShadowViewProjectionMatrices[LayerIndex]);
@@ -355,10 +418,21 @@ void PositionOnlyMain(
 
 #else // #if ONEPASS_POINTLIGHT_SHADOW
 	float ShadowDepth;
-	SetShadowDepthOutputs(WorldPos,OutPosition,ShadowDepth);
+	SetShadowDepthOutputs(
+		PassStruct.ProjectionMatrix, 
+		PassStruct.ViewMatrix, 
+		WorldPos, 
+		WorldNormal, 
+		OutPosition,
 	#if PERSPECTIVE_CORRECT_DEPTH
-		OutParameters.ShadowDepth = ShadowDepth;
-	#elif !COMPILER_SUPPORTS_EMPTY_STRUCTS
+		OutParameters.ShadowDepth,
+		OutParameters.DepthBias
+	#else
+		ShadowDepth
+	#endif
+	);
+
+	#if !PERSPECTIVE_CORRECT_DEPTH && !COMPILER_SUPPORTS_EMPTY_STRUCTS
 		OutParameters.Dummy = 0;
 	#endif
 #endif // #if ONEPASS_POINTLIGHT_SHADOW
@@ -370,7 +444,7 @@ void PositionOnlyMain(
 
 
 void PositionOnlyMainForGS(
-	in FPositionOnlyVertexFactoryInput Input,
+	in FPositionAndNormalOnlyVertexFactoryInput Input,
 #if USING_ONEPASS_POINT_VERTEX_LAYER
 	out uint LayerIndex : SV_RenderTargetArrayIndex,
 #endif
diff --git a/Engine/Shaders/Private/ShadowProjectionCommon.ush b/Engine/Shaders/Private/ShadowProjectionCommon.ush
--- a/Engine/Shaders/Private/ShadowProjectionCommon.ush
+++ b/Engine/Shaders/Private/ShadowProjectionCommon.ush
@@ -149,7 +149,7 @@ void CalcCubemapVectorAndFaceIndex(float3 NormalizedLightVector, float ShadowFil
 }
 
 /** Computes shadowing for a given world position from a cubemap shadowmap used on a point light. */
-float CubemapHardwarePCF(float3 WorldPosition, float3 LightPosition, float LightInvRadius, float DepthBias)
+float CubemapHardwarePCF(float3 WorldPosition, float3 LightPosition, float LightInvRadius, float DepthBias, float SlopeDepthBias, float MaxSlopeDepthBias)
 {
 	float Shadow = 1;
 
diff --git a/Engine/Shaders/Private/ShadowProjectionPixelShader.usf b/Engine/Shaders/Private/ShadowProjectionPixelShader.usf
--- a/Engine/Shaders/Private/ShadowProjectionPixelShader.usf
+++ b/Engine/Shaders/Private/ShadowProjectionPixelShader.usf
@@ -55,8 +55,8 @@ float4 LightPositionAndInvRadius;
 	#include "TransmissionThickness.ush"
 #else
 	float4x4 ScreenToShadowMatrix;
-	// .x:DepthBias, y: MaxSubjectZ - MinSubjectZ
-	float2 ProjectionDepthBiasParameters;
+	// .x:DepthBias, .y:SlopeDepthBias, .z:MaxSlopeDepthBias, .w: MaxSubjectZ - MinSubjectZ
+	float4 ProjectionDepthBiasParameters;
 #endif
 
 #if USE_FADE_PLANE
@@ -171,7 +171,7 @@ void Main(
 		Settings.SceneDepth = LightSpacePixelDepthForOpaque;
 		Settings.TransitionScale = SoftTransitionScale.z;
 		Settings.MaxKernelSize = PCSSParameters.y;
-		Settings.SvPosition = SVPos;
+		Settings.SvPosition = SVPos.xy;
 		Settings.PQMPContext = PQMPContext;
 		Settings.DebugViewportUV = ScreenUV;
 		
@@ -237,16 +237,16 @@ void Main(
 				Settings.SceneDepth = LightSpacePixelDepthForSSS + ProjectionDepthBiasParameters.x;
 				Settings.bSubsurface = true;
 				Settings.bTreatMaxDepthUnshadowed = false;
-				Settings.DensityMulConstant = Density * ProjectionDepthBiasParameters.y;
-				Settings.ProjectionDepthBiasParameters = ProjectionDepthBiasParameters;
+				Settings.DensityMulConstant = Density * ProjectionDepthBiasParameters.w;
+				Settings.ProjectionDepthBiasParameters = ProjectionDepthBiasParameters.xw;
 
-#if USE_TRANSMISSION
+			  #if USE_TRANSMISSION
 				if (GBufferData.ShadingModelID == SHADINGMODELID_SUBSURFACE_PROFILE)
 				{
 					SSSTransmission = CalcTransmissionThickness(ScreenPosition.xyz, LightPositionAndInvRadius.xyz, GBufferData, Settings);
 				}
 				else
-#endif
+			  #endif
 				{
 					// ideally we use a larger filter kernel for SSSbut as Gather4 makes that harder
 					SSSTransmission = ManualPCF(ShadowPosition.xy, Settings);
@@ -302,8 +302,11 @@ void Main(
 
 #if FEATURE_LEVEL >= FEATURE_LEVEL_SM4
 
-// .x:DepthBias, y: unused, zw: depth projection parameters
-float4 PointLightDepthBiasAndProjParameters;
+// .x:DepthBias, y: SlopeDepthBias, z: MaxSlopeDepthBias, 
+float3 PointLightDepthBias;
+
+// xy: depth projection parameters
+float2 PointLightProjParameters;
 
 /** Pixel shader for projecting a one pass point light shadow from a cube map. */
 void MainOnePassPointLightPS(
@@ -327,7 +330,7 @@ void MainOnePassPointLightPS(
 	float Shadow = TextureCubeSampleLevel(ShadowDepthCubeTexture, ShadowDepthTextureSampler, LightVector, 0);
 	#endif
 #else
- 	float Shadow = CubemapHardwarePCF(WorldPosition, LightPositionAndInvRadius.xyz, LightPositionAndInvRadius.w, PointLightDepthBiasAndProjParameters.x);
+ 	float Shadow = CubemapHardwarePCF(WorldPosition, LightPositionAndInvRadius.xyz, LightPositionAndInvRadius.w, PointLightDepthBias.x, PointLightDepthBias.y, PointLightDepthBias.z);
 
 	Shadow = saturate( (Shadow - 0.5) * ShadowSharpen + 0.5 );
 #endif
@@ -355,7 +358,7 @@ void MainOnePassPointLightPS(
 	bool IsSubsurfaceProfile = GBufferData.ShadingModelID == SHADINGMODELID_SUBSURFACE_PROFILE;
 	if (IsSubsurfaceProfile)
 	{
-		Thickness = CalcTransmissionThickness(GBufferData, WorldPosition, LightVector, PointLightDepthBiasAndProjParameters.zw, LightPositionAndInvRadius.w, 0);
+		Thickness = CalcTransmissionThickness(GBufferData, WorldPosition, LightVector, PointLightProjParameters, LightPositionAndInvRadius.w, 0);
 	}
 	OutColor.a = IsSubsurfaceProfile ? EncodeLightAttenuation(Thickness) : OutColor.b;
 #endif // USE_TRANSMISSION
diff --git a/Engine/Shaders/Private/TransmissionThickness.ush b/Engine/Shaders/Private/TransmissionThickness.ush
--- a/Engine/Shaders/Private/TransmissionThickness.ush
+++ b/Engine/Shaders/Private/TransmissionThickness.ush
@@ -9,9 +9,9 @@ Functions for Transmission Thickness
 -------------------------------------------------------------------------------*/
 float4x4 ScreenToShadowMatrix;
 
-// x is start,y is end,From 1 to 32 sample
-// .x:DepthBias, y: MaxSubjectZ - MinSubjectZ
-float2 ProjectionDepthBiasParameters;
+// x is start, z is end,From 1 to 32 sample
+// .x:DepthBias, .y:SlopeDepthBias, .z:MaxSlopeDepthBias, .w: MaxSubjectZ - MinSubjectZ
+float4 ProjectionDepthBiasParameters;
 
 // Poisson disks generated with http://www.coderhaus.com/?p=11
 // Combine all Poisson disk togeather, borrow idea from Jorge's multiple sample single scattering 
@@ -150,7 +150,7 @@ float CalcTransmissionThickness(float3 ViewPos, float3 LightPos, FGBufferData GB
 
 
 	ShrinkedScreenPos.xy = ShrinkedScreenPos.xy / ShrinkedScreenPos.w;
-	Settings.DensityMulConstant = DepthScale * ProjectionDepthBiasParameters.y;
+	Settings.DensityMulConstant = DepthScale * ProjectionDepthBiasParameters.w;
 
 
 	Settings.SceneDepth = ShrinkedScreenPos.z + ProjectionDepthBiasParameters.x;
@@ -215,7 +215,7 @@ float CalcTransmissionThickness(FGBufferData GBufferData, float3 WorldPosition,
 		Settings.bSubsurface = true;
 		//Shadow Depth range Near to far, so no need mul far here.
 		Settings.DensityMulConstant = TransmissionParams.ExtinctionScale * (10.0f / LightInvRadius);
-		Settings.ProjectionDepthBiasParameters = ProjectionDepthBiasParameters;
+		Settings.ProjectionDepthBiasParameters = ProjectionDepthBiasParameters.xw;
 		Settings.SceneDepth = 0.0f;
 		Settings.SceneDepth = CalcShadowLinearDepth(CompareDistance + ShadowDepthBias, ProjParam) * LightInvRadius;
 
diff --git a/Engine/Shaders/Private/VolumeLightingCommon.ush b/Engine/Shaders/Private/VolumeLightingCommon.ush
--- a/Engine/Shaders/Private/VolumeLightingCommon.ush
+++ b/Engine/Shaders/Private/VolumeLightingCommon.ush
@@ -15,7 +15,8 @@ float4x4 WorldToShadowMatrix;
 float4 ShadowmapMinMax;
 // .x:1/SplitNearFadeRegion, .y:1/SplitFarFadeRegion .zw:DistanceFadeMAD
 float4 ShadowInjectParams;
-float2 DepthBiasParameters;
+// .x:DepthBias, .y:SlopeDepthBias, .z:MaxSlopeDepthBias, .w:Projection parameters
+float4 DepthBiasParameters;
 
 /** Whether to compute static shadowing. */
 uint bStaticallyShadowed;
@@ -94,7 +95,10 @@ float ComputeVolumeShadowing(float3 WorldPositionForLighting, bool bPointLight,
 
 		if (bUseCubemapShadowing)
 		{
-			DynamicShadowFactor = CubemapHardwarePCF(WorldPositionForLighting, DeferredLightUniforms.Position, DeferredLightUniforms.InvRadius, 0.03f * 512 * InvShadowmapResolution);
+			const float DepthBias = 0.03f * 512 * InvShadowmapResolution;
+			const float SlopeDepthBias = 0;
+			const float MaxSlopeDepthBias = 0;
+			DynamicShadowFactor = CubemapHardwarePCF(WorldPositionForLighting, DeferredLightUniforms.Position, DeferredLightUniforms.InvRadius, DepthBias, SlopeDepthBias, MaxSlopeDepthBias);
 		}
 		else
 		{
diff --git a/Engine/Source/Runtime/Engine/Classes/Components/LightComponent.h b/Engine/Source/Runtime/Engine/Classes/Components/LightComponent.h
--- a/Engine/Source/Runtime/Engine/Classes/Components/LightComponent.h
+++ b/Engine/Source/Runtime/Engine/Classes/Components/LightComponent.h
@@ -94,6 +94,16 @@ class ENGINE_API ULightComponent : public ULightComponentBase
 	UPROPERTY(EditAnywhere, BlueprintReadOnly, Category=Light, AdvancedDisplay, meta=(UIMin = "0", UIMax = "1"))
 	float ShadowBias;
 
+	/**
+	 * Controls how accurate self shadowing of whole scene shadows from this light are. This works in addition to shadow bias, by increasing the 
+	 * amount of bias depending on the slope of a surface.
+	 * At 0, shadows will start at the their caster surface, but there will be many self shadowing artifacts.
+	 * larger values, shadows will start further from their caster, and there won't be self shadowing artifacts but object might appear to fly.
+	 * around 0.5 seems to be a good tradeoff. This also affects the soft transition of shadows
+	 */
+	UPROPERTY(EditAnywhere, BlueprintReadOnly, Category = Light, AdvancedDisplay, meta = (UIMin = "0", UIMax = "1"))
+	float ShadowSlopeBias;
+
 	/** Amount to sharpen shadow filtering */
 	UPROPERTY(EditAnywhere, BlueprintReadOnly, Category=Light, AdvancedDisplay, meta=(UIMin = "0.0", UIMax = "1.0", DisplayName = "Shadow Filter Sharpen"))
 	float ShadowSharpen;
@@ -275,6 +285,9 @@ public:
 
 	UFUNCTION(BlueprintCallable, Category="Rendering|Components|Light")
 	void SetShadowBias(float NewValue);
+	
+	UFUNCTION(BlueprintCallable, Category = "Rendering|Components|Light")
+	void SetShadowSlopeBias(float NewValue);
 
 	UFUNCTION(BlueprintCallable, Category = "Rendering|Components|Light")
 	void SetSpecularScale(float NewValue);
diff --git a/Engine/Source/Runtime/Engine/Private/Components/DirectionalLightComponent.cpp b/Engine/Source/Runtime/Engine/Private/Components/DirectionalLightComponent.cpp
--- a/Engine/Source/Runtime/Engine/Private/Components/DirectionalLightComponent.cpp
+++ b/Engine/Source/Runtime/Engine/Private/Components/DirectionalLightComponent.cpp
@@ -808,6 +808,7 @@ void UDirectionalLightComponent::PostEditChangeProperty(FPropertyChangedEvent& P
 	ShadowDistanceFadeoutFraction = FMath::Clamp(ShadowDistanceFadeoutFraction, 0.0f, 1.0f);
 	// max range is larger than UI
 	ShadowBias = FMath::Clamp(ShadowBias, 0.0f, 10.0f);
+	ShadowSlopeBias = FMath::Clamp(ShadowSlopeBias, 0.0f, 10.0f);
 
 	Super::PostEditChangeProperty(PropertyChangedEvent);
 }
diff --git a/Engine/Source/Runtime/Engine/Private/Components/LightComponent.cpp b/Engine/Source/Runtime/Engine/Private/Components/LightComponent.cpp
--- a/Engine/Source/Runtime/Engine/Private/Components/LightComponent.cpp
+++ b/Engine/Source/Runtime/Engine/Private/Components/LightComponent.cpp
@@ -249,6 +249,7 @@ FLightSceneProxy::FLightSceneProxy(const ULightComponent* InLightComponent)
 	, VolumetricScatteringIntensity(FMath::Max(InLightComponent->VolumetricScatteringIntensity, 0.0f))
 	, ShadowResolutionScale(InLightComponent->ShadowResolutionScale)
 	, ShadowBias(InLightComponent->ShadowBias)
+	, ShadowSlopeBias(InLightComponent->ShadowSlopeBias)
 	, ShadowSharpen(InLightComponent->ShadowSharpen)
 	, ContactShadowLength(InLightComponent->ContactShadowLength)
 	, SpecularScale(InLightComponent->SpecularScale)
@@ -395,6 +396,7 @@ ULightComponent::ULightComponent(const FObjectInitializer& ObjectInitializer)
 	IndirectLightingIntensity = 1.0f;
 	ShadowResolutionScale = 1.0f;
 	ShadowBias = 0.5f;
+	ShadowSlopeBias = 0.0f;
 	ShadowSharpen = 0.0f;
 	ContactShadowLength = 0.0f;
 	ContactShadowLengthInWS = false;
@@ -642,6 +644,7 @@ void ULightComponent::PostEditChangeProperty(FPropertyChangedEvent& PropertyChan
 		PropertyName != GET_MEMBER_NAME_STRING_CHECKED(ULightComponent, DisabledBrightness) &&
 		PropertyName != GET_MEMBER_NAME_STRING_CHECKED(ULightComponent, ShadowResolutionScale) &&
 		PropertyName != GET_MEMBER_NAME_STRING_CHECKED(ULightComponent, ShadowBias) &&
+		PropertyName != GET_MEMBER_NAME_STRING_CHECKED(ULightComponent, ShadowSlopeBias) &&
 		PropertyName != GET_MEMBER_NAME_STRING_CHECKED(ULightComponent, ShadowSharpen) &&
 		PropertyName != GET_MEMBER_NAME_STRING_CHECKED(ULightComponent, ContactShadowLength) &&
 		PropertyName != GET_MEMBER_NAME_STRING_CHECKED(ULightComponent, ContactShadowLengthInWS) &&
@@ -1028,6 +1031,16 @@ void ULightComponent::SetShadowBias(float NewValue)
 	}
 }
 
+void ULightComponent::SetShadowSlopeBias(float NewValue)
+{
+	if (AreDynamicDataChangesAllowed()
+		&& ShadowSlopeBias != NewValue)
+	{
+		ShadowSlopeBias = NewValue;
+		MarkRenderStateDirty();
+	}
+}
+
 void ULightComponent::SetSpecularScale(float NewValue)
 {
 	if (AreDynamicDataChangesAllowed()
diff --git a/Engine/Source/Runtime/Engine/Public/GPUSkinVertexFactory.h b/Engine/Source/Runtime/Engine/Public/GPUSkinVertexFactory.h
--- a/Engine/Source/Runtime/Engine/Public/GPUSkinVertexFactory.h
+++ b/Engine/Source/Runtime/Engine/Public/GPUSkinVertexFactory.h
@@ -487,6 +487,7 @@ public:
 
 	//TODO should be supported
 	bool SupportsPositionOnlyStream() const override { return false; }
+	bool SupportsPositionAndNormalOnlyStream() const override { return false; }
 
 	inline void InvalidateStreams()
 	{
diff --git a/Engine/Source/Runtime/Engine/Public/SceneManagement.h b/Engine/Source/Runtime/Engine/Public/SceneManagement.h
--- a/Engine/Source/Runtime/Engine/Public/SceneManagement.h
+++ b/Engine/Source/Runtime/Engine/Public/SceneManagement.h
@@ -1236,6 +1236,7 @@ public:
 
 	// Accessors.
 	float GetUserShadowBias() const { return ShadowBias; }
+	float GetUserShadowSlopeBias() const { return ShadowSlopeBias; }
 
 	/** 
 	 * Note: The Rendering thread must not dereference UObjects!  
@@ -1362,6 +1363,9 @@ protected:
 	/** User setting from light component, 0:no bias, 0.5:reasonable, larger object might appear to float */
 	float ShadowBias;
 
+	/** User setting from light component, 0:no bias, 0.5:reasonable, larger object might appear to float */
+	float ShadowSlopeBias;
+
 	/** Sharpen shadow filtering */
 	float ShadowSharpen;
 
diff --git a/Engine/Source/Runtime/Renderer/Private/ShadowDepthRendering.cpp b/Engine/Source/Runtime/Renderer/Private/ShadowDepthRendering.cpp
--- a/Engine/Source/Runtime/Renderer/Private/ShadowDepthRendering.cpp
+++ b/Engine/Source/Runtime/Renderer/Private/ShadowDepthRendering.cpp
@@ -54,9 +54,10 @@ void SetupShadowDepthPassUniformBuffer(
 	FSceneRenderTargets& SceneRenderTargets = FSceneRenderTargets::Get(RHICmdList);
 	SetupSceneTextureUniformParameters(SceneRenderTargets, View.FeatureLevel, ESceneTextureSetupMode::None, ShadowDepthPassParameters.SceneTextures);
 
-	ShadowDepthPassParameters.ProjectionMatrix = FTranslationMatrix(ShadowInfo->PreShadowTranslation - View.ViewMatrices.GetPreViewTranslation()) * ShadowInfo->SubjectAndReceiverMatrix;
+	ShadowDepthPassParameters.ProjectionMatrix = FTranslationMatrix(ShadowInfo->PreShadowTranslation - View.ViewMatrices.GetPreViewTranslation()) * ShadowInfo->SubjectAndReceiverMatrix;	
+	ShadowDepthPassParameters.ViewMatrix = ShadowInfo->ShadowViewMatrix;
 
-	ShadowDepthPassParameters.ShadowParams = FVector2D(ShadowInfo->GetShaderDepthBias(), ShadowInfo->InvMaxSubjectDepth);
+	ShadowDepthPassParameters.ShadowParams = FVector4(ShadowInfo->GetShaderDepthBias(), ShadowInfo->GetShaderSlopeDepthBias(), ShadowInfo->GetShaderMaxSlopeDepthBias(), ShadowInfo->bOnePassPointLightShadow ? 1 : ShadowInfo->InvMaxSubjectDepth);
 	// Only clamp vertices to the near plane when rendering whole scene directional light shadow depths or preshadows from directional lights
 	const bool bClampToNearPlaneValue = ShadowInfo->IsWholeSceneDirectionalShadow() || (ShadowInfo->bPreShadow && ShadowInfo->bDirectionalLight);
 	ShadowDepthPassParameters.bClampToNearPlane = bClampToNearPlaneValue ? 1.0f : 0.0f;
@@ -70,6 +71,7 @@ void SetupShadowDepthPassUniformBuffer(
 			// Have to apply the pre-view translation to the view - projection matrices
 			FMatrix TranslatedShadowViewProjectionMatrix = Translation * ShadowInfo->OnePassShadowViewProjectionMatrices[FaceIndex];
 			ShadowDepthPassParameters.ShadowViewProjectionMatrices[FaceIndex] = TranslatedShadowViewProjectionMatrix;
+			ShadowDepthPassParameters.ShadowViewMatrices[FaceIndex] = ShadowInfo->OnePassShadowViewMatrices[FaceIndex];
 		}
 	}
 
@@ -99,8 +101,9 @@ void SetupShadowDepthPassUniformBuffer(
 	SetupMobileSceneTextureUniformParameters(SceneRenderTargets, View.FeatureLevel, false, ShadowDepthPassParameters.SceneTextures);
 
 	ShadowDepthPassParameters.ProjectionMatrix = FTranslationMatrix(ShadowInfo->PreShadowTranslation - View.ViewMatrices.GetPreViewTranslation()) * ShadowInfo->SubjectAndReceiverMatrix;
+	ShadowDepthPassParameters.ViewMatrix = ShadowInfo->ShadowViewMatrix;
 
-	ShadowDepthPassParameters.ShadowParams = FVector2D(ShadowInfo->GetShaderDepthBias(), ShadowInfo->InvMaxSubjectDepth);
+	ShadowDepthPassParameters.ShadowParams = FVector4(ShadowInfo->GetShaderDepthBias(), ShadowInfo->GetShaderSlopeDepthBias(), ShadowInfo->GetShaderMaxSlopeDepthBias(), ShadowInfo->InvMaxSubjectDepth);
 	// Only clamp vertices to the near plane when rendering whole scene directional light shadow depths or preshadows from directional lights
 	const bool bClampToNearPlaneValue = ShadowInfo->IsWholeSceneDirectionalShadow() || (ShadowInfo->bPreShadow && ShadowInfo->bDirectionalLight);
 	ShadowDepthPassParameters.bClampToNearPlane = bClampToNearPlaneValue ? 1.0f : 0.0f;
@@ -231,7 +234,7 @@ public:
 				|| ((!Material->WritesEveryPixel(true) || Material->MaterialMayModifyMeshPosition()) && !bUsePositionOnlyStream))
 				// Only compile one pass point light shaders for feature levels >= SM4
 				&& (ShaderMode != VertexShadowDepth_OnePassPointLight || IsFeatureLevelSupported(Platform, ERHIFeatureLevel::SM4))
-				// Only compile position-only shaders for vertex factories that support it.
+				// Only compile position-only shaders for vertex factories that support it. (Note: this assumes that a vertex factor which supports PositionOnly, supports also PositionAndNormalOnly)
 				&& (!bUsePositionOnlyStream || VertexFactoryType->SupportsPositionOnly())
 				// Don't render ShadowDepth for translucent unlit materials
 				&& Material->ShouldCastDynamicShadows()
@@ -342,6 +345,7 @@ public:
 	static void ModifyCompilationEnvironment(EShaderPlatform Platform, const FMaterial* Material, FShaderCompilerEnvironment& OutEnvironment)
 	{
 		FMeshMaterialShader::ModifyCompilationEnvironment(Platform, Material, OutEnvironment);
+		OutEnvironment.SetDefine(TEXT("ONEPASS_POINTLIGHT_SHADOW"), 1);
 		TShadowDepthVS<VertexShadowDepth_OnePassPointLight, false, false, true>::ModifyCompilationEnvironment(Platform, Material, OutEnvironment);
 	}
 
@@ -1740,8 +1744,9 @@ void FShadowDepthPassMeshProcessor::Process(
 		TShadowDepthBasePS<bRenderReflectiveShadowMap>,
 		FOnePassPointShadowDepthGS> ShadowDepthPassShaders;
 
-	const bool bUsePositionOnlyVS = !bRenderReflectiveShadowMap
-		&& VertexFactory->SupportsPositionOnlyStream()
+	const bool bUsePositionOnlyVS = 
+		!bRenderReflectiveShadowMap
+		&& VertexFactory->SupportsPositionAndNormalOnlyStream()
 		&& MaterialResource.WritesEveryPixel(true)
 		&& !MaterialResource.MaterialModifiesMeshPosition_RenderThread();
 
@@ -1779,7 +1784,7 @@ void FShadowDepthPassMeshProcessor::Process(
 			MeshFillMode,
 			MeshCullMode,
 			SortKey,
-			bUsePositionOnlyVS ? EMeshPassFeatures::PositionOnly : EMeshPassFeatures::Default,
+			bUsePositionOnlyVS ? EMeshPassFeatures::PositionAndNormalOnly : EMeshPassFeatures::Default,
 			ShaderElementData);
 	}
 }
diff --git a/Engine/Source/Runtime/Renderer/Private/ShadowRendering.cpp b/Engine/Source/Runtime/Renderer/Private/ShadowRendering.cpp
--- a/Engine/Source/Runtime/Renderer/Private/ShadowRendering.cpp
+++ b/Engine/Source/Runtime/Renderer/Private/ShadowRendering.cpp
@@ -20,13 +20,25 @@ static TAutoConsoleVariable<float> CVarCSMShadowDepthBias(
 	TEXT("Constant depth bias used by CSM"),
 	ECVF_RenderThreadSafe);
 
+static TAutoConsoleVariable<float> CVarCSMShadowSlopeScaleDepthBias(
+	TEXT("r.Shadow.CSMSlopeScaleDepthBias"),
+	5.0f,
+	TEXT("Slope scale depth bias used by CSM"),
+	ECVF_RenderThreadSafe);
+
 static TAutoConsoleVariable<float> CVarPerObjectDirectionalShadowDepthBias(
 	TEXT("r.Shadow.PerObjectDirectionalDepthBias"),
 	20.0f,
 	TEXT("Constant depth bias used by per-object shadows from directional lights\n")
-	TEXT("Lower values give better self-shadowing, but increase self-shadowing artifacts"),
+	TEXT("Lower values give better shadow contact, but increase self-shadowing artifacts"),
 	ECVF_RenderThreadSafe);
 
+static TAutoConsoleVariable<float> CVarPerObjectDirectionalShadowSlopeScaleDepthBias(
+	TEXT("r.Shadow.PerObjectDirectionalSlopeDepthBias"),
+	5.0f,
+	TEXT("Slope scale depth bias used by per-object shadows from directional lights\n")
+	TEXT("Lower values give better shadow contact, but increase self-shadowing artifacts"),
+	ECVF_RenderThreadSafe);
 
 static TAutoConsoleVariable<float> CVarCSMSplitPenumbraScale(
 	TEXT("r.Shadow.CSMSplitPenumbraScale"),
@@ -58,12 +70,24 @@ static TAutoConsoleVariable<float> CVarPointLightShadowDepthBias(
 	TEXT("Depth bias that is applied in the depth pass for shadows from point lights. (0.03 avoids peter paning but has some shadow acne)"),
 	ECVF_RenderThreadSafe);
 
+static TAutoConsoleVariable<float> CVarPointLightShadowSlopeScaleDepthBias(
+	TEXT("r.Shadow.PointLightSlopeScaleDepthBias"),
+	5.0f,
+	TEXT("Slope scale depth bias that is applied in the depth pass for shadows from point lights"),
+	ECVF_RenderThreadSafe);
+
 static TAutoConsoleVariable<float> CVarSpotLightShadowDepthBias(
 	TEXT("r.Shadow.SpotLightDepthBias"),
 	5.0f,
 	TEXT("Depth bias that is applied in the depth pass for per object projected shadows from spot lights"),
 	ECVF_RenderThreadSafe);
 
+static TAutoConsoleVariable<float> CVarSpotLightShadowSlopeScaleDepthBias(
+	TEXT("r.Shadow.SpotLightSlopeDepthBias"),
+	5.0f,
+	TEXT("Slope scale depth bias that is applied in the depth pass for per object projected shadows from spot lights"),
+	ECVF_RenderThreadSafe);
+
 static TAutoConsoleVariable<int32> CVarEnableModulatedSelfShadow(
 	TEXT("r.Shadow.EnableModulatedSelfShadow"),
 	0,
@@ -92,6 +116,13 @@ static TAutoConsoleVariable<int32> CVarMaxSoftKernelSize(
 	TEXT("Mazimum size of the softening kernels in pixels."),
 	ECVF_RenderThreadSafe);
 
+static TAutoConsoleVariable<float> CVarShadowMaxSlopeScaleDepthBias(
+	TEXT("r.Shadow.ShadowMaxSlopeScaleDepthBias"),
+	2.0f,
+	TEXT("Max Slope depth bias used for shadows for all lights\n")
+	TEXT("Higher values give better self-shadowing, but increase self-shadowing artifacts"),
+	ECVF_RenderThreadSafe);
+
 DEFINE_GPU_STAT(ShadowProjection);
 
 // 0:off, 1:low, 2:med, 3:high, 4:very high, 5:max
@@ -1131,12 +1162,16 @@ FMatrix FProjectedShadowInfo::GetWorldToShadowMatrix(FVector4& ShadowmapMinMax,
 void FProjectedShadowInfo::UpdateShaderDepthBias()
 {
 	float DepthBias = 0;
+	float SlopeScaleDepthBias = 1;
 
 	if (IsWholeScenePointLightShadow())
 	{
 		DepthBias = CVarPointLightShadowDepthBias.GetValueOnRenderThread() * 512.0f / FMath::Max(ResolutionX, ResolutionY);
 		// * 2.0f to be compatible with the system we had before ShadowBias
 		DepthBias *= 2.0f * LightSceneInfo->Proxy->GetUserShadowBias();
+
+		SlopeScaleDepthBias = CVarPointLightShadowSlopeScaleDepthBias.GetValueOnRenderThread();
+		SlopeScaleDepthBias *= LightSceneInfo->Proxy->GetUserShadowSlopeBias();
 	}
 	else if (IsWholeSceneDirectionalShadow())
 	{
@@ -1146,14 +1181,17 @@ void FProjectedShadowInfo::UpdateShaderDepthBias()
 		DepthBias = CVarCSMShadowDepthBias.GetValueOnRenderThread() / (MaxSubjectZ - MinSubjectZ);
 
 		float WorldSpaceTexelScale = ShadowBounds.W / ResolutionX;
-		
 		DepthBias *= WorldSpaceTexelScale;
 		DepthBias *= LightSceneInfo->Proxy->GetUserShadowBias();
+
+		SlopeScaleDepthBias = CVarCSMShadowSlopeScaleDepthBias.GetValueOnRenderThread();
+		SlopeScaleDepthBias *= LightSceneInfo->Proxy->GetUserShadowSlopeBias();
 	}
 	else if (bPreShadow)
 	{
 		// Preshadows don't need a depth bias since there is no self shadowing
 		DepthBias = 0;
+		SlopeScaleDepthBias = 0;
 	}
 	else
 	{
@@ -1169,6 +1207,9 @@ void FProjectedShadowInfo::UpdateShaderDepthBias()
 		
 			DepthBias *= WorldSpaceTexelScale;
 			DepthBias *= 0.5f;	// avg GetUserShadowBias, in that case we don't want this adjustable
+
+			SlopeScaleDepthBias = CVarPerObjectDirectionalShadowSlopeScaleDepthBias.GetValueOnRenderThread();
+			SlopeScaleDepthBias *= LightSceneInfo->Proxy->GetUserShadowSlopeBias();
 		}
 		else
 		{
@@ -1177,6 +1218,9 @@ void FProjectedShadowInfo::UpdateShaderDepthBias()
 			DepthBias = LightTypeDepthBias * 512.0f / ((MaxSubjectZ - MinSubjectZ) * FMath::Max(ResolutionX, ResolutionY));
 			// * 2.0f to be compatible with the system we had before ShadowBias
 			DepthBias *= 2.0f * LightSceneInfo->Proxy->GetUserShadowBias();
+
+			SlopeScaleDepthBias = CVarSpotLightShadowSlopeScaleDepthBias.GetValueOnRenderThread();
+			SlopeScaleDepthBias *= LightSceneInfo->Proxy->GetUserShadowSlopeBias();
 		}
 
 		// Prevent a large depth bias due to low resolution from causing near plane clipping
@@ -1184,6 +1228,8 @@ void FProjectedShadowInfo::UpdateShaderDepthBias()
 	}
 
 	ShaderDepthBias = FMath::Max(DepthBias, 0.0f);
+	ShaderSlopeDepthBias = FMath::Max(DepthBias * SlopeScaleDepthBias, 0.0f);
+	ShaderMaxSlopeDepthBias = CVarShadowMaxSlopeScaleDepthBias.GetValueOnRenderThread();
 }
 
 float FProjectedShadowInfo::ComputeTransitionSize() const
@@ -1368,7 +1414,7 @@ bool FSceneRenderer::RenderShadowProjections(FRHICommandListImmediate& RHICmdLis
 			// Set the light's scissor rectangle.
 			LightSceneInfo->Proxy->SetScissorRect(RHICmdList, View, View.ViewRect);
 
-		    Scene->UniformBuffers.UpdateViewUniformBuffer(View);
+			Scene->UniformBuffers.UpdateViewUniformBuffer(View);
 
 			// Project the shadow depth buffers onto the scene.
 			for (int32 ShadowIndex = 0; ShadowIndex < NormalShadows.Num(); ShadowIndex++)
diff --git a/Engine/Source/Runtime/Renderer/Private/ShadowRendering.h b/Engine/Source/Runtime/Renderer/Private/ShadowRendering.h
--- a/Engine/Source/Runtime/Renderer/Private/ShadowRendering.h
+++ b/Engine/Source/Runtime/Renderer/Private/ShadowRendering.h
@@ -172,15 +172,18 @@ BEGIN_GLOBAL_SHADER_PARAMETER_STRUCT(FShadowDepthPassUniformParameters,)
 	SHADER_PARAMETER_STRUCT(FSceneTexturesUniformParameters, SceneTextures)
 	SHADER_PARAMETER_STRUCT(FLpvWriteUniformBufferParameters, LPV)
 	SHADER_PARAMETER(FMatrix, ProjectionMatrix)
-	SHADER_PARAMETER(FVector2D, ShadowParams)
+	SHADER_PARAMETER(FMatrix, ViewMatrix)
+	SHADER_PARAMETER(FVector4, ShadowParams)
 	SHADER_PARAMETER(float, bClampToNearPlane)
 	SHADER_PARAMETER_ARRAY(FMatrix, ShadowViewProjectionMatrices, [6])
+	SHADER_PARAMETER_ARRAY(FMatrix, ShadowViewMatrices, [6])
 END_GLOBAL_SHADER_PARAMETER_STRUCT()
 
 BEGIN_GLOBAL_SHADER_PARAMETER_STRUCT(FMobileShadowDepthPassUniformParameters,)
 	SHADER_PARAMETER_STRUCT(FMobileSceneTextureUniformParameters, SceneTextures)
 	SHADER_PARAMETER(FMatrix, ProjectionMatrix)
-	SHADER_PARAMETER(FVector2D, ShadowParams)
+	SHADER_PARAMETER(FMatrix, ViewMatrix)
+	SHADER_PARAMETER(FVector4, ShadowParams)
 	SHADER_PARAMETER(float, bClampToNearPlane)
 	SHADER_PARAMETER_ARRAY(FMatrix, ShadowViewProjectionMatrices, [6])
 END_GLOBAL_SHADER_PARAMETER_STRUCT()
@@ -322,6 +325,9 @@ public:
 	/** View projection matrices for each cubemap face, used by one pass point light shadows. */
 	TArray<FMatrix> OnePassShadowViewProjectionMatrices;
 
+	/** View matrices for each cubemap face, used by one pass point light shadows. */
+	TArray<FMatrix> OnePassShadowViewMatrices;
+	
 	/** Frustums for each cubemap face, used for object culling one pass point light shadows. */
 	TArray<FConvexVolume> OnePassShadowFrustums;
 
@@ -363,6 +369,9 @@ public:
 		);
 
 	float GetShaderDepthBias() const { return ShaderDepthBias; }
+	float GetShaderSlopeDepthBias() const { return ShaderSlopeDepthBias; }
+	float GetShaderMaxSlopeDepthBias() const { return ShaderMaxSlopeDepthBias; }
+	FVector GetShaderShadowReceiverBias() const { return ShaderShadowReceiverDepthBias; }
 
 	/**
 	 * Renders the shadow subject depth.
@@ -468,8 +477,9 @@ public:
 		return RenderTargets.GetSize();
 	}
 
-	/** Computes and updates ShaderDepthBias */
+	/** Computes and updates ShaderDepthBias and ShaderSlopeDepthBias */
 	void UpdateShaderDepthBias();
+
 	/** How large the soft PCF comparison should be, similar to DepthBias, before this was called TransitionScale and 1/Size */
 	float ComputeTransitionSize() const;
 
@@ -543,6 +553,9 @@ private:
 	 * Set by UpdateShaderDepthBias(), get with GetShaderDepthBias(), -1 if not set
 	 */
 	float ShaderDepthBias;
+	float ShaderSlopeDepthBias;
+	float ShaderMaxSlopeDepthBias;
+	FVector ShaderShadowReceiverDepthBias;
 
 	void CopyCachedShadowMap(FRHICommandList& RHICmdList, const FMeshPassProcessorRenderState& DrawRenderState, FSceneRenderer* SceneRenderer, const FViewInfo& View);
 
@@ -756,6 +769,7 @@ public:
 		FadePlaneOffset.Bind(ParameterMap,TEXT("FadePlaneOffset"));
 		InvFadePlaneLength.Bind(ParameterMap,TEXT("InvFadePlaneLength"));
 		ShadowTileOffsetAndSizeParam.Bind(ParameterMap, TEXT("ShadowTileOffsetAndSize"));
+		ShadowReceiverParams.Bind(ParameterMap, TEXT("ShadowReceiverParams"));
 	}
 
 	void Set(FRHICommandList& RHICmdList, FShader* Shader, const FSceneView& View, const FProjectedShadowInfo* ShadowInfo)
@@ -830,7 +844,7 @@ public:
 				);
 		}
 
-		SetShaderValue(RHICmdList, ShaderRHI, ProjectionDepthBias, FVector2D(ShadowInfo->GetShaderDepthBias(), ShadowInfo->MaxSubjectZ - ShadowInfo->MinSubjectZ));
+		SetShaderValue(RHICmdList, ShaderRHI, ProjectionDepthBias, FVector4(ShadowInfo->GetShaderDepthBias(), ShadowInfo->GetShaderSlopeDepthBias(), ShadowInfo->GetShaderMaxSlopeDepthBias(), ShadowInfo->MaxSubjectZ - ShadowInfo->MinSubjectZ));
 		SetShaderValue(RHICmdList, ShaderRHI, FadePlaneOffset, ShadowInfo->CascadeSettings.FadePlaneOffset);
 
 		if(InvFadePlaneLength.IsBound())
@@ -838,6 +852,11 @@ public:
 			check(ShadowInfo->CascadeSettings.FadePlaneLength > 0);
 			SetShaderValue(RHICmdList, ShaderRHI, InvFadePlaneLength, 1.0f / ShadowInfo->CascadeSettings.FadePlaneLength);
 		}
+
+		if (ShadowReceiverParams.IsBound())
+		{
+			SetShaderValue(RHICmdList, ShaderRHI, ShadowReceiverParams, ShadowInfo->GetShaderShadowReceiverBias());
+		}
 	}
 
 	/** Serializer. */
@@ -853,6 +872,7 @@ public:
 		Ar << P.FadePlaneOffset;
 		Ar << P.InvFadePlaneLength;
 		Ar << P.ShadowTileOffsetAndSizeParam;
+		Ar << P.ShadowReceiverParams;
 		return Ar;
 	}
 
@@ -868,6 +888,7 @@ private:
 	FShaderParameter FadePlaneOffset;
 	FShaderParameter InvFadePlaneLength;
 	FShaderParameter ShadowTileOffsetAndSizeParam;
+	FShaderParameter ShadowReceiverParams;
 };
 
 /**
@@ -1242,7 +1263,8 @@ public:
 		LightPosition.Bind(Initializer.ParameterMap,TEXT("LightPositionAndInvRadius"));
 		ShadowFadeFraction.Bind(Initializer.ParameterMap,TEXT("ShadowFadeFraction"));
 		ShadowSharpen.Bind(Initializer.ParameterMap,TEXT("ShadowSharpen"));
-		PointLightDepthBiasAndProjParameters.Bind(Initializer.ParameterMap,TEXT("PointLightDepthBiasAndProjParameters"));
+		PointLightDepthBias.Bind(Initializer.ParameterMap,TEXT("PointLightDepthBias"));
+		PointLightProjParameters.Bind(Initializer.ParameterMap, TEXT("PointLightProjParameters"));
 		TransmissionProfilesTexture.Bind(Initializer.ParameterMap, TEXT("SSProfilesTexture"));
 	}
 
@@ -1283,7 +1305,8 @@ public:
 		float Far = LightProxy.GetRadius();
 		FVector2D param = FVector2D(Far / (Far - Near), -Near * Far / (Far - Near));
 		FVector2D projParam = FVector2D(1.0f / param.Y, param.X / param.Y);
-		SetShaderValue(RHICmdList, ShaderRHI, PointLightDepthBiasAndProjParameters, FVector4(ShadowInfo->GetShaderDepthBias(), 0.0f, projParam.X, projParam.Y));
+		SetShaderValue(RHICmdList, ShaderRHI, PointLightDepthBias, FVector(ShadowInfo->GetShaderDepthBias(), ShadowInfo->GetShaderSlopeDepthBias(), ShadowInfo->GetShaderMaxSlopeDepthBias()));
+		SetShaderValue(RHICmdList, ShaderRHI, PointLightProjParameters, FVector2D(projParam.X, projParam.Y));
 
 		FSceneRenderTargets& SceneContext = FSceneRenderTargets::Get(RHICmdList);
 		{
@@ -1326,7 +1349,8 @@ public:
 		Ar << LightPosition;
 		Ar << ShadowFadeFraction;
 		Ar << ShadowSharpen;
-		Ar << PointLightDepthBiasAndProjParameters;
+		Ar << PointLightDepthBias;
+		Ar << PointLightProjParameters;
 		Ar << TransmissionProfilesTexture;
 		return bShaderHasOutdatedParameters;
 	}
@@ -1338,7 +1362,8 @@ private:
 	FShaderParameter LightPosition;
 	FShaderParameter ShadowFadeFraction;
 	FShaderParameter ShadowSharpen;
-	FShaderParameter PointLightDepthBiasAndProjParameters;
+	FShaderParameter PointLightDepthBias;
+	FShaderParameter PointLightProjParameters;
 	FShaderResourceParameter TransmissionProfilesTexture;
 
 
diff --git a/Engine/Source/Runtime/Renderer/Private/ShadowSetup.cpp b/Engine/Source/Runtime/Renderer/Private/ShadowSetup.cpp
--- a/Engine/Source/Runtime/Renderer/Private/ShadowSetup.cpp
+++ b/Engine/Source/Runtime/Renderer/Private/ShadowSetup.cpp
@@ -519,6 +519,7 @@ FProjectedShadowInfo::FProjectedShadowInfo()
 	, NumDynamicSubjectMeshElements(0)
 	, NumSubjectMeshCommandBuildRequestElements(0)
 	, ShaderDepthBias(0.0f)
+	, ShaderSlopeDepthBias(0.0f)
 {
 }
 
@@ -2625,6 +2626,7 @@ void FSceneRenderer::CreateWholeSceneProjectedShadow(
 						const FMatrix FaceProjection = FPerspectiveMatrix(PI / 4.0f, 1, 1, 1, LightProxy.GetRadius());
 						const FVector LightPosition = LightProxy.GetPosition();
 
+						ProjectedShadowInfo->OnePassShadowViewMatrices.Empty(6);
 						ProjectedShadowInfo->OnePassShadowViewProjectionMatrices.Empty(6);
 						ProjectedShadowInfo->OnePassShadowFrustums.Empty(6);
 						ProjectedShadowInfo->OnePassShadowFrustums.AddZeroed(6);
@@ -2635,7 +2637,9 @@ void FSceneRenderer::CreateWholeSceneProjectedShadow(
 						for (int32 FaceIndex = 0; FaceIndex < 6; FaceIndex++)
 						{
 							// Create a view projection matrix for each cube face
-							const FMatrix ShadowViewProjectionMatrix = FLookAtMatrix(LightPosition, LightPosition + CubeDirections[FaceIndex], UpVectors[FaceIndex]) * ScaleMatrix * FaceProjection;
+							const FMatrix WorldToLightMatrix = FLookAtMatrix(LightPosition, LightPosition + CubeDirections[FaceIndex], UpVectors[FaceIndex]) * ScaleMatrix;
+							ProjectedShadowInfo->OnePassShadowViewMatrices.Add(WorldToLightMatrix);
+							const FMatrix ShadowViewProjectionMatrix = WorldToLightMatrix * FaceProjection;
 							ProjectedShadowInfo->OnePassShadowViewProjectionMatrices.Add(ShadowViewProjectionMatrix);
 							// Create a convex volume out of the frustum so it can be used for object culling
 							GetViewFrustumBounds(ProjectedShadowInfo->OnePassShadowFrustums[FaceIndex], ShadowViewProjectionMatrix, false);
diff --git a/Engine/Source/Runtime/Renderer/Private/TranslucentLighting.cpp b/Engine/Source/Runtime/Renderer/Private/TranslucentLighting.cpp
--- a/Engine/Source/Runtime/Renderer/Private/TranslucentLighting.cpp
+++ b/Engine/Source/Runtime/Renderer/Private/TranslucentLighting.cpp
@@ -206,7 +206,7 @@ public:
 			FTranslationMatrix(ShadowInfo->PreShadowTranslation - View.ViewMatrices.GetPreViewTranslation()) * ShadowInfo->SubjectAndReceiverMatrix
 			);
 
-		SetShaderValue(RHICmdList, ShaderRHI, ShadowParams, FVector2D(ShadowInfo->GetShaderDepthBias(), ShadowInfo->InvMaxSubjectDepth));
+		SetShaderValue(RHICmdList, ShaderRHI, ShadowParams, FVector4(ShadowInfo->GetShaderDepthBias(), ShadowInfo->GetShaderSlopeDepthBias(), ShadowInfo->GetShaderMaxSlopeDepthBias(), ShadowInfo->InvMaxSubjectDepth));
 		// Only clamp vertices to the near plane when rendering whole scene directional light shadow depths or preshadows from directional lights
 		const bool bClampToNearPlaneValue = ShadowInfo->IsWholeSceneDirectionalShadow() || (ShadowInfo->bPreShadow && ShadowInfo->bDirectionalLight);
 		SetShaderValue(RHICmdList, ShaderRHI,ClampToNearPlane,bClampToNearPlaneValue ? 1.0f : 0.0f);
diff --git a/Engine/Source/Runtime/Renderer/Private/VolumeLighting.h b/Engine/Source/Runtime/Renderer/Private/VolumeLighting.h
--- a/Engine/Source/Runtime/Renderer/Private/VolumeLighting.h
+++ b/Engine/Source/Runtime/Renderer/Private/VolumeLighting.h
@@ -93,7 +93,7 @@ public:
 
 		if (bDynamicallyShadowed)
 		{
-			SetShaderValue(RHICmdList, ShaderRHI, DepthBiasParameters, FVector2D(ShadowMap->GetShaderDepthBias(), 1.0f / (ShadowMap->MaxSubjectZ - ShadowMap->MinSubjectZ)));
+			SetShaderValue(RHICmdList, ShaderRHI, DepthBiasParameters, FVector4(ShadowMap->GetShaderDepthBias(), ShadowMap->GetShaderSlopeDepthBias(), ShadowMap->GetShaderMaxSlopeDepthBias(), 1.0f / (ShadowMap->MaxSubjectZ - ShadowMap->MinSubjectZ)));
 
 			FTextureRHIParamRef ShadowDepthTextureResource = nullptr;
 			if (LightType == LightType_Point || LightType == LightType_Rect)
