Fixes mistakenly modified E in SSR because visual studio doesn't reload file when shelving files...

diff --git a/Engine/Shaders/Private/ScreenSpaceReflections.usf b/Engine/Shaders/Private/ScreenSpaceReflections.usf
--- a/Engine/Shaders/Private/ScreenSpaceReflections.usf
+++ b/Engine/Shaders/Private/ScreenSpaceReflections.usf
@@ -156,7 +156,7 @@ void ScreenSpaceReflectionsPS(
 
 #if 1
 			// TODO EnvBRDF isn't a perfect match anymore.
-			float3 H = mul( ImportanceSampleVisibleGGX(UniformSampleDisk(E), a2, TangentV ).xyz, TangentBasis );
+			float3 H = mul( ImportanceSampleVisibleGGX(E, a2, TangentV ).xyz, TangentBasis );
 			float3 L = 2 * dot( V, H ) * H - V;
 #elif 0
 			float3 H = mul( ImportanceSampleGGX( E, a2 ).xyz, TangentBasis );
