Fix ProjectionDepthBiasParameters being applied twice when evaluating subsurface scattering with the float3 CalculateOcclusion version (used in ShadowQuality 1 or 3/4 with non-Gather4 support).

diff --git a/Engine/Shaders/Private/ShadowFilteringCommon.ush b/Engine/Shaders/Private/ShadowFilteringCommon.ush
--- a/Engine/Shaders/Private/ShadowFilteringCommon.ush
+++ b/Engine/Shaders/Private/ShadowFilteringCommon.ush
@@ -188,7 +188,7 @@ float3 CalculateOcclusion(float3 ShadowmapDepth, FPCFSamplerSettings Settings)
 		// Determine the distance that the light traveled through the subsurface object
 		// This assumes that anything between this subsurface pixel and the light was also a subsurface material,
 		// As a result, subsurface materials receive leaked light based on how transparent they are
-		float3 Thickness = max(Settings.SceneDepth - (ShadowmapDepth - Settings.ProjectionDepthBiasParameters.x), 0);
+		float3 Thickness = max(Settings.SceneDepth - ShadowmapDepth, 0);
 		float3 Occlusion = saturate(exp(-Thickness * Settings.DensityMulConstant));
 		// Never shadow from depths that were never written to (max depth value)
 		return ShadowmapDepth > .99f ? 1 : Occlusion;
