Change subsurface occlusion to use a fast exp, instead exp.

diff --git a/Engine/Shaders/Private/ShadowFilteringCommon.ush b/Engine/Shaders/Private/ShadowFilteringCommon.ush
--- a/Engine/Shaders/Private/ShadowFilteringCommon.ush
+++ b/Engine/Shaders/Private/ShadowFilteringCommon.ush
@@ -147,6 +147,13 @@ float PCF1x1(float2 Fraction, float4 Values00)
 	return lerp(HorizontalLerp00.x, HorizontalLerp00.y, Fraction.y);
 }
 
+// Use fast approximation of exp(x)
+// http://www.humus.name/Articles/Persson_LowLevelThinking.pdf p.27
+float4 FastExp(float4 x)
+{
+	return exp2(1.442695f * x);
+}
+
 float4 CalculateOcclusion(float4 ShadowmapDepth, FPCFSamplerSettings Settings)
 {
 	if (Settings.bSubsurface)
@@ -155,7 +162,7 @@ float4 CalculateOcclusion(float4 ShadowmapDepth, FPCFSamplerSettings Settings)
 		// This assumes that anything between this subsurface pixel and the light was also a subsurface material,
 		// As a result, subsurface materials receive leaked light based on how transparent they are
 		float4 Thickness = max(Settings.SceneDepth - ShadowmapDepth, 0);
-		float4 Occlusion = saturate(exp(-Thickness * Settings.DensityMulConstant));
+		float4 Occlusion = saturate(FastExp(-Thickness * Settings.DensityMulConstant));
 		// Never shadow from depths that were never written to (max depth value)
 		return ShadowmapDepth > .99f ? 1 : Occlusion;
 	}
@@ -189,7 +196,7 @@ float3 CalculateOcclusion(float3 ShadowmapDepth, FPCFSamplerSettings Settings)
 		// This assumes that anything between this subsurface pixel and the light was also a subsurface material,
 		// As a result, subsurface materials receive leaked light based on how transparent they are
 		float3 Thickness = max(Settings.SceneDepth - ShadowmapDepth, 0);
-		float3 Occlusion = saturate(exp(-Thickness * Settings.DensityMulConstant));
+		float3 Occlusion = saturate(FastExp(-Thickness * Settings.DensityMulConstant));
 		// Never shadow from depths that were never written to (max depth value)
 		return ShadowmapDepth > .99f ? 1 : Occlusion;
 	}
