Add float3 version of FastExp.

diff --git a/Engine/Shaders/Private/ShadowFilteringCommon.ush b/Engine/Shaders/Private/ShadowFilteringCommon.ush
--- a/Engine/Shaders/Private/ShadowFilteringCommon.ush
+++ b/Engine/Shaders/Private/ShadowFilteringCommon.ush
@@ -154,6 +154,11 @@ float4 FastExp(float4 x)
 	return exp2(1.442695f * x);
 }
 
+float3 FastExp(float3 x)
+{
+	return exp2(1.442695f * x);
+}
+
 float4 CalculateOcclusion(float4 ShadowmapDepth, FPCFSamplerSettings Settings)
 {
 	if (Settings.bSubsurface)
