Fixes all RGS that were using a FrameNumber external from the renderer instead of ViewState's frame index.

diff --git a/Engine/Shaders/Private/RayTracing/RayTracingAmbientOcclusionRGS.usf b/Engine/Shaders/Private/RayTracing/RayTracingAmbientOcclusionRGS.usf
--- a/Engine/Shaders/Private/RayTracing/RayTracingAmbientOcclusionRGS.usf
+++ b/Engine/Shaders/Private/RayTracing/RayTracingAmbientOcclusionRGS.usf
@@ -43,7 +43,7 @@ void AmbientOcclusionRGS()
 
 	RandomSequence RandSequence;
 	uint LinearIndex = CalcLinearIndex(PixelCoord);
-	RandomSequence_Initialize(RandSequence, LinearIndex, View.FrameNumber);
+	RandomSequence_Initialize(RandSequence, LinearIndex, View.StateFrameIndex);
 
 	// Get G-Buffer surface data
 	float2 InvBufferSize = View.BufferSizeAndInvSize.zw;
diff --git a/Engine/Shaders/Private/RayTracing/RayTracingGlobalIlluminationRGS.usf b/Engine/Shaders/Private/RayTracing/RayTracingGlobalIlluminationRGS.usf
--- a/Engine/Shaders/Private/RayTracing/RayTracingGlobalIlluminationRGS.usf
+++ b/Engine/Shaders/Private/RayTracing/RayTracingGlobalIlluminationRGS.usf
@@ -106,7 +106,7 @@ void GlobalIlluminationRGS()
 	for (uint SampleIndex = 0; SampleIndex < SamplesPerPixel; ++SampleIndex)
 	{
 		// Use Scrambled Halton sequence
-		uint FrameIndex = View.FrameNumber % 1024;
+		uint FrameIndex = View.StateFrameIndex % 1024;
 		RandomSequence_Initialize(RandSequence, LinearIndex, FrameIndex * SamplesPerPixel + SampleIndex);
 		RandSequence.Type = 2;
 
diff --git a/Engine/Shaders/Private/RayTracing/RayTracingOcclusionRGS.usf b/Engine/Shaders/Private/RayTracing/RayTracingOcclusionRGS.usf
--- a/Engine/Shaders/Private/RayTracing/RayTracingOcclusionRGS.usf
+++ b/Engine/Shaders/Private/RayTracing/RayTracingOcclusionRGS.usf
@@ -129,7 +129,7 @@ void OcclusionRGS()
 	{
 	RandomSequence RandSequence;
 	uint LinearIndex = CalcLinearIndex(PixelCoord);
-	RandomSequence_Initialize(RandSequence, LinearIndex, View.FrameNumber);
+	RandomSequence_Initialize(RandSequence, LinearIndex, View.StateFrameIndex);
 
 	FLightShaderParameters LightParameters = GetRootLightShaderParameters();
 
diff --git a/Engine/Shaders/Private/RayTracing/RayTracingRectLightRGS.usf b/Engine/Shaders/Private/RayTracing/RayTracingRectLightRGS.usf
--- a/Engine/Shaders/Private/RayTracing/RayTracingRectLightRGS.usf
+++ b/Engine/Shaders/Private/RayTracing/RayTracingRectLightRGS.usf
@@ -157,7 +157,7 @@ void RectLightRGS()
 
 	RandomSequence RandSequence;
 	uint LinearIndex = CalcLinearIndex(PixelCoord);
-	RandomSequence_Initialize(RandSequence, LinearIndex, View.FrameNumber);
+	RandomSequence_Initialize(RandSequence, LinearIndex, View.StateFrameIndex);
 
 	// Get G-Buffer surface data
 	float2 InvBufferSize = View.BufferSizeAndInvSize.zw;
diff --git a/Engine/Shaders/Private/RayTracing/RayTracingReflections.usf b/Engine/Shaders/Private/RayTracing/RayTracingReflections.usf
--- a/Engine/Shaders/Private/RayTracing/RayTracingReflections.usf
+++ b/Engine/Shaders/Private/RayTracing/RayTracingReflections.usf
@@ -158,7 +158,7 @@ void RayTracingReflectionsRGS()
 	uint LinearIndex = PixelCoord.y * View.BufferSizeAndInvSize.x + PixelCoord.x; // TODO: PixelCoord or DispatchThreadId
 	
 	RandomSequence RandSequence;
-	RandomSequence_Initialize(RandSequence, LinearIndex, View.FrameNumber);
+	RandomSequence_Initialize(RandSequence, LinearIndex, View.StateFrameIndex);
 
 	float2 InvBufferSize = View.BufferSizeAndInvSize.zw;
 	float2 UV = (float2(PixelCoord) + 0.5) * InvBufferSize;
diff --git a/Engine/Shaders/Private/RayTracing/RayTracingSkyLightRGS.usf b/Engine/Shaders/Private/RayTracing/RayTracingSkyLightRGS.usf
--- a/Engine/Shaders/Private/RayTracing/RayTracingSkyLightRGS.usf
+++ b/Engine/Shaders/Private/RayTracing/RayTracingSkyLightRGS.usf
@@ -251,7 +251,7 @@ void SkyLightRGS()
 	float RayDistance = 0.0;
 	float HitCount = 0.0;
 	uint SamplesPerPixel = SkyLight.SamplesPerPixel;
-	uint FrameIndex = View.FrameNumber % SamplesPerPixel;
+	uint FrameIndex = View.StateFrameIndex % SamplesPerPixel;
 
 	// Mask out depth values that are infinitely far away
 	float DeviceZ = SceneTexturesStruct.SceneDepthTexture.Load(int3(PixelCoord, 0)).r;
diff --git a/Engine/Shaders/Private/RayTracing/RayTracingTranslucency.usf b/Engine/Shaders/Private/RayTracing/RayTracingTranslucency.usf
--- a/Engine/Shaders/Private/RayTracing/RayTracingTranslucency.usf
+++ b/Engine/Shaders/Private/RayTracing/RayTracingTranslucency.usf
@@ -109,7 +109,7 @@ void RayTracingTranslucencyRGS()
 	uint LinearIndex = PixelCoord.y * View.BufferSizeAndInvSize.x + PixelCoord.x; // TODO: PixelCoord or DispatchThreadId
 	
 	RandomSequence RandSequence;
-	RandomSequence_Initialize(RandSequence, LinearIndex, View.FrameNumber);
+	RandomSequence_Initialize(RandSequence, LinearIndex, View.StateFrameIndex);
 
 	float2 InvBufferSize = View.BufferSizeAndInvSize.zw;
 	float2 UV = (float2(PixelCoord) + 0.5) * InvBufferSize;
