Move MortonEncode/MortonDecode into a separate HLSL header

TODO: Unify functions and call site as described in UE-72536

diff --git a/Engine/Shaders/Private/MortonCode.ush b/Engine/Shaders/Private/MortonCode.ush
new file mode 100644
--- /dev/null
+++ b/Engine/Shaders/Private/MortonCode.ush
@@ -0,0 +1,40 @@
+// Copyright 1998-2019 Epic Games, Inc. All Rights Reserved.
+
+/*=============================================================================
+	MortonCode.ush: Morton code lookup functions
+=============================================================================*/
+
+// Encodes two 16-bit integers into one 32-bit morton code
+uint MortonEncode(uint2 Pixel)
+{
+	Pixel.x &= 0x0000ffff;
+	Pixel.x = (Pixel.x ^ (Pixel.x << 8)) & 0x00ff00ff;
+	Pixel.x = (Pixel.x ^ (Pixel.x << 4)) & 0x0f0f0f0f;
+	Pixel.x = (Pixel.x ^ (Pixel.x << 2)) & 0x33333333;
+	Pixel.x = (Pixel.x ^ (Pixel.x << 1)) & 0x55555555;
+
+	Pixel.y &= 0x0000ffff;
+	Pixel.y = (Pixel.y ^ (Pixel.y << 8)) & 0x00ff00ff;
+	Pixel.y = (Pixel.y ^ (Pixel.y << 4)) & 0x0f0f0f0f;
+	Pixel.y = (Pixel.y ^ (Pixel.y << 2)) & 0x33333333;
+	Pixel.y = (Pixel.y ^ (Pixel.y << 1)) & 0x55555555;
+	return Pixel.x | (Pixel.y << 1);
+}
+
+// Decodes one 32-bit morton code into two 16-bit integers
+uint2 MortonDecode(uint Morton)
+{
+	uint2 Pixel = uint2(Morton, Morton >> 1);
+	Pixel.x &= 0x55555555;
+	Pixel.x = (Pixel.x ^ (Pixel.x >> 1)) & 0x33333333;
+	Pixel.x = (Pixel.x ^ (Pixel.x >> 2)) & 0x0f0f0f0f;
+	Pixel.x = (Pixel.x ^ (Pixel.x >> 4)) & 0x00ff00ff;
+	Pixel.x = (Pixel.x ^ (Pixel.x >> 8)) & 0x0000ffff;
+
+	Pixel.y &= 0x55555555;
+	Pixel.y = (Pixel.y ^ (Pixel.y >> 1)) & 0x33333333;
+	Pixel.y = (Pixel.y ^ (Pixel.y >> 2)) & 0x0f0f0f0f;
+	Pixel.y = (Pixel.y ^ (Pixel.y >> 4)) & 0x00ff00ff;
+	Pixel.y = (Pixel.y ^ (Pixel.y >> 8)) & 0x0000ffff;
+	return Pixel;
+}
diff --git a/Engine/Shaders/Private/RayTracing/MipTreeCommon.ush b/Engine/Shaders/Private/RayTracing/MipTreeCommon.ush
--- a/Engine/Shaders/Private/RayTracing/MipTreeCommon.ush
+++ b/Engine/Shaders/Private/RayTracing/MipTreeCommon.ush
@@ -4,6 +4,8 @@
 	MipTreeCommon.ush: Common utilities for MipTree-based sampling
 ===============================================================================================*/
 
+#include "../MortonCode.ush"
+
 #define USE_MORTON_CODE 1
 
 uint2 DimensionsAtLevel(uint2 Dimensions, uint Level)
@@ -11,40 +13,6 @@ uint2 DimensionsAtLevel(uint2 Dimensions, uint Level)
 	return uint2(Dimensions.x >> Level, Dimensions.y >> Level);
 }
 
-// #dxr_todo: UE-72536 unifying morton code
-uint MortonEncode(uint2 Pixel)
-{
-	Pixel.x &= 0x0000ffff;
-	Pixel.x = (Pixel.x ^ (Pixel.x << 8)) & 0x00ff00ff;
-	Pixel.x = (Pixel.x ^ (Pixel.x << 4)) & 0x0f0f0f0f;
-	Pixel.x = (Pixel.x ^ (Pixel.x << 2)) & 0x33333333;
-	Pixel.x = (Pixel.x ^ (Pixel.x << 1)) & 0x55555555;
-
-	Pixel.y &= 0x0000ffff;
-	Pixel.y = (Pixel.y ^ (Pixel.y << 8)) & 0x00ff00ff;
-	Pixel.y = (Pixel.y ^ (Pixel.y << 4)) & 0x0f0f0f0f;
-	Pixel.y = (Pixel.y ^ (Pixel.y << 2)) & 0x33333333;
-	Pixel.y = (Pixel.y ^ (Pixel.y << 1)) & 0x55555555;
-	return Pixel.x | (Pixel.y << 1);
-}
-
-uint2 MortonDecode(uint Morton)
-{
-	uint2 Pixel = uint2(Morton, Morton >> 1);
-	Pixel.x &= 0x55555555;
-	Pixel.x = (Pixel.x ^ (Pixel.x >> 1)) & 0x33333333;
-	Pixel.x = (Pixel.x ^ (Pixel.x >> 2)) & 0x0f0f0f0f;
-	Pixel.x = (Pixel.x ^ (Pixel.x >> 4)) & 0x00ff00ff;
-	Pixel.x = (Pixel.x ^ (Pixel.x >> 8)) & 0x0000ffff;
-
-	Pixel.y &= 0x55555555;
-	Pixel.y = (Pixel.y ^ (Pixel.y >> 1)) & 0x33333333;
-	Pixel.y = (Pixel.y ^ (Pixel.y >> 2)) & 0x0f0f0f0f;
-	Pixel.y = (Pixel.y ^ (Pixel.y >> 4)) & 0x00ff00ff;
-	Pixel.y = (Pixel.y ^ (Pixel.y >> 8)) & 0x0000ffff;
-	return Pixel;
-}
-
 uint BufferOffsetAtPixel(uint2 Pixel, uint Level, uint2 Dimensions)
 {
 	uint Offset = 0;
