QuasiRandom implementation

diff --git a/Engine/Shaders/Private/Halton.ush b/Engine/Shaders/Private/Halton.ush
new file mode 100644
--- /dev/null
+++ b/Engine/Shaders/Private/Halton.ush
@@ -0,0 +1,38 @@
+// Copyright 1998-2019 Epic Games, Inc. All Rights Reserved.
+
+/*=============================================================================
+Halton.ush: Halton sampling interface and utilities
+=============================================================================*/
+
+#pragma once
+
+uint HaltonPrime(uint Dimension)
+{
+	return HaltonPrimes.Primes[Dimension % HaltonPrimes.Dimensions];
+}
+
+// https://en.wikipedia.org/wiki/Halton_sequence
+// Procedural definition for the Halton sequence, given a base and index
+float HaltonProcedural(uint Index, uint Base)
+{
+	float r = 0.0;
+	float f = 1.0;
+
+	float BaseInv = 1.0 / Base;
+	while (Index > 0)
+	{
+		f *= BaseInv;
+		r += f * (Index % Base);
+		Index /= Base;
+	}
+
+	return r;
+}
+
+float HaltonSequence(uint SequenceIndex, int IterationIndex, int DimensionIndex)
+{
+	int SequenceOffset = SequenceIndex * HaltonIteration.IterationCount * HaltonIteration.Dimensions;
+	int IterationOffset = IterationIndex * HaltonIteration.Dimensions;
+	float RandSample = HaltonIteration.SequenceIteration[SequenceOffset + IterationOffset + DimensionIndex];
+	return RandSample;
+}
diff --git a/Engine/Shaders/Private/QuasiRandom.ush b/Engine/Shaders/Private/QuasiRandom.ush
new file mode 100644
--- /dev/null
+++ b/Engine/Shaders/Private/QuasiRandom.ush
@@ -0,0 +1,116 @@
+// Copyright 1998-2019 Epic Games, Inc. All Rights Reserved.
+
+/*=============================================================================
+	QuasiRandom.ush: Quasi-random sampling interface and utilities
+=============================================================================*/
+
+#pragma once
+
+#include "BlueNoise.ush"
+#include "Halton.ush"
+
+struct FQuasiRandom
+{
+	uint PositionSeed;
+	uint TimeSeed;
+
+	uint HaltonSequence;
+	uint HaltonIteration;
+	uint Dimension;
+
+	uint2 BlueNoiseCoord;
+};
+
+// http://burtleburtle.net/bob/hash/integer.html
+// Bob Jenkins integer hashing function in 6 shifts.
+uint IntegerHash(uint a)
+{
+	a = (a + 0x7ed55d16) + (a << 12);
+	a = (a ^ 0xc761c23c) ^ (a >> 19);
+	a = (a + 0x165667b1) + (a << 5);
+	a = (a + 0xd3a2646c) ^ (a << 9);
+	a = (a + 0xfd7046c5) + (a << 3);
+	a = (a ^ 0xb55a4f09) ^ (a >> 16);
+	return a;
+}
+
+void QuasiRandom_Update(inout FQuasiRandom QuasiRandom, uint Iteration, uint IterationStride);
+
+void QuasiRandom_Initialize(inout FQuasiRandom QuasiRandom, float2 PixelIndex, uint Iteration, uint IterationStride)
+{
+	// Translate PixelIndex to SequenceIndex if scrambling is desired
+	QuasiRandom.PositionSeed = IntegerHash(MortonEncode(PixelIndex));
+
+	uint2 TileIndex = PixelIndex % uint2(HaltonIteration.SequenceColumnCount, HaltonIteration.SequenceRowCount);
+	QuasiRandom.HaltonSequence = MortonEncode(TileIndex);
+	QuasiRandom.BlueNoiseCoord = PixelIndex % BlueNoise.Dimensions.xy;
+
+	QuasiRandom_Update(QuasiRandom, Iteration, IterationStride);
+}
+
+void QuasiRandom_Update(inout FQuasiRandom QuasiRandom, uint Iteration, uint IterationStride)
+{
+	QuasiRandom.TimeSeed = View.FrameNumber * IterationStride + Iteration;
+	QuasiRandom.HaltonIteration = Iteration;
+	QuasiRandom.Dimension = 0;
+}
+
+float QuasiRandom_Sample(inout FQuasiRandom QuasiRandom)
+{
+	float RandomSample = 0.0;
+	BRANCH
+	if (QuasiRandom.Dimension < BlueNoise.Dimensions.z * 3)
+	{
+		// Evaluate shift vector
+		float3 BlueNoiseShiftVector = EvalBlueNoise(QuasiRandom.BlueNoiseCoord, QuasiRandom.Dimension / 3);
+		float BlueNoiseShift = BlueNoiseShiftVector[QuasiRandom.Dimension % 3];
+
+		// Evaluate Halton
+		float HaltonSample = HaltonSequence(QuasiRandom.HaltonSequence, QuasiRandom.HaltonIteration, QuasiRandom.Dimension);
+
+		// Wrap sum
+		RandomSample = fmod(HaltonSample + BlueNoiseShift, 1.0);
+	}
+	else
+	{
+		RandomSample = HaltonProcedural(QuasiRandom.PositionSeed + QuasiRandom.TimeSeed, HaltonPrime(QuasiRandom.Dimension));
+	}
+
+	// Increment dimension
+	QuasiRandom.Dimension++;
+
+	// Return shifted Halton sequence value
+	return RandomSample;
+}
+
+float3 QuasiRandom_Sample3(inout FQuasiRandom QuasiRandom)
+{
+	float3 RandomSample = 0.0;
+	BRANCH
+	if (QuasiRandom.Dimension + 2 < BlueNoise.Dimensions.z * 3)
+	{
+		// Evaluate shift vector
+		float3 BlueNoiseShift = EvalBlueNoise(QuasiRandom.BlueNoiseCoord, QuasiRandom.Dimension / 3);
+
+		// Evaluate Halton
+		float3 HaltonSample;
+		HaltonSample.x = HaltonSequence(QuasiRandom.HaltonSequence, QuasiRandom.HaltonIteration, QuasiRandom.Dimension);
+		HaltonSample.y = HaltonSequence(QuasiRandom.HaltonSequence, QuasiRandom.HaltonIteration, QuasiRandom.Dimension + 1);
+		HaltonSample.z = HaltonSequence(QuasiRandom.HaltonSequence, QuasiRandom.HaltonIteration, QuasiRandom.Dimension + 2);
+
+		// Wrap sum
+		RandomSample = fmod(HaltonSample + BlueNoiseShift, 1.0);
+	}
+	else
+	{
+		RandomSample.x = HaltonProcedural(QuasiRandom.PositionSeed + QuasiRandom.TimeSeed, HaltonPrime(QuasiRandom.Dimension));
+		RandomSample.y = HaltonProcedural(QuasiRandom.PositionSeed + QuasiRandom.TimeSeed, HaltonPrime(QuasiRandom.Dimension + 1));
+		RandomSample.z = HaltonProcedural(QuasiRandom.PositionSeed + QuasiRandom.TimeSeed, HaltonPrime(QuasiRandom.Dimension + 2));
+	}
+
+	// Increment dimension
+	QuasiRandom.Dimension += 3;
+
+	// Return shifted Halton sequence value
+	return RandomSample;
+}
