Integrating RTGI with QuasiRandom sequence

diff --git a/Engine/Shaders/Private/RayTracing/RayTracingGlobalIlluminationRGS.usf b/Engine/Shaders/Private/RayTracing/RayTracingGlobalIlluminationRGS.usf
--- a/Engine/Shaders/Private/RayTracing/RayTracingGlobalIlluminationRGS.usf
+++ b/Engine/Shaders/Private/RayTracing/RayTracingGlobalIlluminationRGS.usf
@@ -8,10 +8,11 @@
 #include "RayTracingCommon.ush"
 #include "RayTracingHitGroupCommon.ush"
 
-#include "../PathTracing/Utilities/PathTracingRandomSequence.ush" 
 #include "../PathTracing/Light/PathTracingLightSampling.ush"
 #include "../PathTracing/Material/PathTracingMaterialSampling.ush"
 
+#include "../QuasiRandom.ush"
+
 #define USE_PATHTRACING_MATERIALS 0
 
 RaytracingAccelerationStructure TLAS;
@@ -32,7 +33,7 @@ float MaxNormalBias;
 void GenerateCosineNormalRay(
 	float3 WorldPosition,
 	float3 WorldNormal,
-	inout RandomSequence RandSequence,
+	float2 RandSample,
 	out float3 RayOrigin,
 	out float3 RayDirection,
 	out float RayTMin,
@@ -40,11 +41,6 @@ void GenerateCosineNormalRay(
 	out float RayPdf
 )
 {
-	// Draw random variable
-	float2 BufferSize = View.BufferSizeAndInvSize.xy;
-	uint DummyVariable;
-	float2 RandSample = RandomSequence_GenerateSample2D(RandSequence, DummyVariable);
-
 	// Perform cosine-hemispherical sampling and convert to world-space
 	float4 Direction_Tangent = CosineSampleHemisphere(RandSample);
 	float3 Direction_World = TangentToWorld(Direction_Tangent.xyz, WorldNormal);
@@ -64,9 +60,7 @@ void GlobalIlluminationRGS()
 	RWRayDistanceUAV[DispatchThreadId] = float2(-1.0, 0.0);
 
 	uint2 PixelCoord = GetPixelCoord(DispatchThreadId, UpscaleFactor);
-	RandomSequence RandSequence;
 	uint LinearIndex = CalcLinearIndex(PixelCoord);
-	RandomSequence_Initialize(RandSequence, LinearIndex, View.FrameNumber);
 
 	bool IsUnidirectionalEnabled = false;
 
@@ -99,24 +93,24 @@ void GlobalIlluminationRGS()
 		return;
 	}
 
+	FQuasiRandom QuasiRandom;
+	QuasiRandom_Initialize(QuasiRandom, DispatchThreadId, 0, SamplesPerPixel);
+
 	float3 Irradiance = 0;
 	float HitDistance = 0.0;
 	float HitCount = 0.0;
 	float AmbientOcclusion = 0.0;
 	for (uint SampleIndex = 0; SampleIndex < SamplesPerPixel; ++SampleIndex)
 	{
-		// Use Scrambled Halton sequence
-		uint FrameIndex = View.StateFrameIndex % 1024;
-		RandomSequence_Initialize(RandSequence, LinearIndex, FrameIndex * SamplesPerPixel + SampleIndex);
-		RandSequence.Type = 2;
+		QuasiRandom_Update(QuasiRandom, SampleIndex, SamplesPerPixel);
 
 		float3 RayThroughput = 1.0;
+		float3 RandSample = QuasiRandom_Sample3(QuasiRandom);
 
 		// Russian roulette based on DiffuseColor
 		if (UseRussianRoulette)
 		{
-			uint DummyVariable;
-			float RRSample = RandomSequence_GenerateSample1D(RandSequence, DummyVariable);
+			float RRSample = RandSample.x;
 			float ProbabilityOfSuccess = Luminance(DiffuseColor);
 			float ProbabilityOfTermination = 1.0 - ProbabilityOfSuccess;
 			if (RRSample < ProbabilityOfTermination) continue;
@@ -127,7 +121,7 @@ void GlobalIlluminationRGS()
 		RayDesc Ray;
 		float RayPdf = 1.0;
 #if 1
-		GenerateCosineNormalRay(WorldPosition, WorldNormal, RandSequence, Ray.Origin, Ray.Direction, Ray.TMin, Ray.TMax, RayPdf);
+		GenerateCosineNormalRay(WorldPosition, WorldNormal, RandSample.yz, Ray.Origin, Ray.Direction, Ray.TMin, Ray.TMax, RayPdf);
 		half3 N = WorldNormal;
 		half3 V = -CameraDirection;
 		half3 L = Ray.Direction;
@@ -218,8 +212,10 @@ void GlobalIlluminationRGS()
 				int LightId;
 				float3 LightUV;
 				float NeePdf = 0.0;
-				uint DummyVariable;
-				float4 RandSample4 = RandomSequence_GenerateSample4D(RandSequence, DummyVariable);
+
+				float4 RandSample4;
+				RandSample4.xyz = QuasiRandom_Sample3(QuasiRandom);
+				RandSample4.w = QuasiRandom_Sample3(QuasiRandom).x;
 				SampleLight(Ray, Payload, RandSample4, LightId, LightUV, NeePdf);
 
 				if (NeePdf > 0.0)
@@ -280,18 +276,19 @@ void GlobalIlluminationRGS()
 			}
 
 			// Perform material sampling
-			// dxr_todo: only worth doing when Bounce + 1 < MaxBounces
 			if (Bounce + 1 < MaxBounces)
 			{
 				float3 Direction;
 				float3 Throughput = 1.0;
 #if USE_PATHTRACING_MATERIALS
-				uint DummyVariable;
-				float4 RandSample = RandomSequence_GenerateSample4D(RandSequence, DummyVariable);
+				float4 RandSample;
+				RandSample.xyz = QuasiRandom_Sample3(QuasiRandom);
+				RandSample.w = QuasiRandom_Sample3(QuasiRandom).x;
 				SampleMaterial(Ray.Direction, Payload, RandSample, Direction, Throughput, MaterialPdf);
 #else
+				float3 RandSample3 = QuasiRandom_Sample3(QuasiRandom);
 				float3 RayOrigin = Ray.Origin;
-				GenerateCosineNormalRay(RayOrigin, Payload.WorldNormal, RandSequence, Ray.Origin, Direction, Ray.TMin, Ray.TMax, MaterialPdf);
+				GenerateCosineNormalRay(RayOrigin, Payload.WorldNormal, RandSample3.xy, Ray.Origin, Direction, Ray.TMin, Ray.TMax, MaterialPdf);
 				
 				half3 N = Payload.WorldNormal;
 				half3 V = -Ray.Direction;
diff --git a/Engine/Source/Runtime/Renderer/Private/RayTracing/RayTracingGlobalIllumination.cpp b/Engine/Source/Runtime/Renderer/Private/RayTracing/RayTracingGlobalIllumination.cpp
--- a/Engine/Source/Runtime/Renderer/Private/RayTracing/RayTracingGlobalIllumination.cpp
+++ b/Engine/Source/Runtime/Renderer/Private/RayTracing/RayTracingGlobalIllumination.cpp
@@ -19,6 +19,7 @@
 #include "PostProcess/PostProcessing.h"
 #include "PostProcess/SceneFilterRendering.h"
 #include "Raytracing/RaytracingOptions.h"
+#include "BlueNoise.h"
 
 static int32 GRayTracingGlobalIllumination = -1;
 static FAutoConsoleVariableRef CVarRayTracingGlobalIllumination(
@@ -235,6 +236,9 @@ class FGlobalIlluminationRGS : public FGlobalShader
 		SHADER_PARAMETER_RDG_TEXTURE_UAV(RWTexture2D<float>, RWRayDistanceUAV)
 		SHADER_PARAMETER_STRUCT_REF(FViewUniformShaderParameters, ViewUniformBuffer)
 		SHADER_PARAMETER_STRUCT_REF(FSceneTexturesUniformParameters, SceneTexturesStruct)
+		SHADER_PARAMETER_STRUCT_REF(FHaltonIteration, HaltonIteration)
+		SHADER_PARAMETER_STRUCT_REF(FHaltonPrimes, HaltonPrimes)
+		SHADER_PARAMETER_STRUCT_REF(FBlueNoise, BlueNoise)
 		SHADER_PARAMETER_STRUCT_REF(FPathTracingLightData, LightParameters)
 		SHADER_PARAMETER_STRUCT_REF(FSkyLightData, SkyLight)
 
@@ -377,6 +381,22 @@ void FDeferredShadingSceneRenderer::RenderRayTracingGlobalIllumination(
 	SetupSceneTextureUniformParameters(SceneContext, FeatureLevel, ESceneTextureSetupMode::All, SceneTextures);
 	// Ray generation
 	{
+		static uint32 Iteration = 0;
+		uint32 IterationCount = RayTracingGISamplesPerPixel;
+		uint32 SequenceCount = 1;
+		uint32 DimensionCount = 24;
+		FHaltonSequenceIteration HaltonSequenceIteration(Scene->HaltonSequence, IterationCount, SequenceCount, DimensionCount, Iteration);
+
+		FHaltonIteration HaltonIteration;
+		InitializeHaltonSequenceIteration(HaltonSequenceIteration, HaltonIteration);
+		Iteration = (Iteration + 1) % 1024;
+
+		FHaltonPrimes HaltonPrimes;
+		InitializeHaltonPrimes(Scene->HaltonPrimesResource, HaltonPrimes);
+
+		FBlueNoise BlueNoise;
+		InitializeBlueNoise(BlueNoise);
+
 		FPathTracingLightData LightParameters;
 		SetupLightParameters(Scene->Lights, View, &LightParameters);
 
@@ -406,6 +426,9 @@ void FDeferredShadingSceneRenderer::RenderRayTracingGlobalIllumination(
 		PassParameters->TLAS = View.RayTracingScene.RayTracingSceneRHI->GetShaderResourceView();
 		PassParameters->ViewUniformBuffer = View.ViewUniformBuffer;
 		PassParameters->SceneTexturesStruct = CreateUniformBufferImmediate(SceneTextures, EUniformBufferUsage::UniformBuffer_SingleDraw);
+		PassParameters->HaltonIteration = CreateUniformBufferImmediate(HaltonIteration, EUniformBufferUsage::UniformBuffer_SingleDraw);
+		PassParameters->HaltonPrimes = CreateUniformBufferImmediate(HaltonPrimes, EUniformBufferUsage::UniformBuffer_SingleDraw);
+		PassParameters->BlueNoise = CreateUniformBufferImmediate(BlueNoise, EUniformBufferUsage::UniformBuffer_SingleDraw);
 		PassParameters->LightParameters = CreateUniformBufferImmediate(LightParameters, EUniformBufferUsage::UniformBuffer_SingleDraw);
 		PassParameters->SkyLight = CreateUniformBufferImmediate(SkyLightParameters, EUniformBufferUsage::UniformBuffer_SingleDraw);
 		TRefCountPtr<IPooledRenderTarget> SubsurfaceProfileRT((IPooledRenderTarget*) GetSubsufaceProfileTexture_RT(RHICmdList));
