Fixed mobile compilation issue
Verified locally on OGLES

diff --git a/Engine/Shaders/Private/BRDF.ush b/Engine/Shaders/Private/BRDF.ush
--- a/Engine/Shaders/Private/BRDF.ush
+++ b/Engine/Shaders/Private/BRDF.ush
@@ -361,6 +361,11 @@ void EnvBRDFApproxFullyRough(inout half3 DiffuseColor, inout half3 SpecularColor
 	SpecularColor = 0;
 	// We do not modify Roughness here as this is done differently at different places.
 }
+void EnvBRDFApproxFullyRough(inout half3 DiffuseColor, inout half SpecularColor)
+{
+	DiffuseColor += SpecularColor * 0.45;
+	SpecularColor = 0;
+}
 
 
 float D_InvBlinn( float a2, float NoH )
