Remove receiver bias temporally.

This is removed temporally to reduce the number of screenshot to update. It will be reactivated once we slope bias is activated.

diff --git a/Engine/Shaders/Private/ShadowProjectionPixelShader.usf b/Engine/Shaders/Private/ShadowProjectionPixelShader.usf
--- a/Engine/Shaders/Private/ShadowProjectionPixelShader.usf
+++ b/Engine/Shaders/Private/ShadowProjectionPixelShader.usf
@@ -192,7 +192,7 @@ void Main(
 		Settings.ShadowDepthTexture = ShadowDepthTexture;
 		Settings.ShadowDepthTextureSampler = ShadowDepthTextureSampler;
 		Settings.ShadowBufferSize = ShadowBufferSize;
-		Settings.TransitionScale = SoftTransitionScale.z * lerp(0.1, 1.0, NoL);
+		Settings.TransitionScale = SoftTransitionScale.z;// * lerp(0.1, 1.0, NoL); Temporally disabled until slope bias is activated at the same time, to reduce the number of screenshots to update.
 		Settings.SceneDepth = LightSpacePixelDepthForOpaque;
 		Settings.bSubsurface = false;
 		Settings.bTreatMaxDepthUnshadowed = false;
