Fixes shader compilation failure due to multiple include of MortonCode.ush

diff --git a/Engine/Shaders/Private/MortonCode.ush b/Engine/Shaders/Private/MortonCode.ush
--- a/Engine/Shaders/Private/MortonCode.ush
+++ b/Engine/Shaders/Private/MortonCode.ush
@@ -4,6 +4,9 @@
 	MortonCode.ush: Morton code lookup functions
 =============================================================================*/
 
+#pragma once
+
+
 // Encodes two 16-bit integers into one 32-bit morton code
 uint MortonEncode(uint2 Pixel)
 {
