Fixed Vulkan tessellation after merge with Main branch.

#rnx

diff --git a/Engine/Shaders/Private/BasePassCommon.ush b/Engine/Shaders/Private/BasePassCommon.ush
--- a/Engine/Shaders/Private/BasePassCommon.ush
+++ b/Engine/Shaders/Private/BasePassCommon.ush
@@ -62,7 +62,7 @@
 #define TRANSLUCENCY_PERVERTEX_LIGHTING_VOLUME (!FORWARD_SHADING && (TRANSLUCENCY_LIGHTING_VOLUMETRIC_PERVERTEX_DIRECTIONAL || TRANSLUCENCY_LIGHTING_VOLUMETRIC_PERVERTEX_NONDIRECTIONAL))
 #define TRANSLUCENCY_PERVERTEX_FORWARD_SHADING (FORWARD_SHADING && (TRANSLUCENCY_LIGHTING_VOLUMETRIC_PERVERTEX_DIRECTIONAL || TRANSLUCENCY_LIGHTING_VOLUMETRIC_PERVERTEX_NONDIRECTIONAL))
 
-#define TESSELLATION_SUPPORTED (PC_D3D || GL4_PROFILE || SWITCH_PROFILE || SWITCH_PROFILE_FORWARD || METAL_TESSELLATION_PROFILE || VULKAN_ENABLE_TESSELLATION)
+#define TESSELLATION_SUPPORTED (PC_D3D || GL4_PROFILE || SWITCH_PROFILE || SWITCH_PROFILE_FORWARD || METAL_TESSELLATION_PROFILE || VULKAN_PROFILE_SM5)
 
 struct FSharedBasePassInterpolants
 {
diff --git a/Engine/Shaders/Public/Platform/Vulkan/VulkanCommon.ush b/Engine/Shaders/Public/Platform/Vulkan/VulkanCommon.ush
--- a/Engine/Shaders/Public/Platform/Vulkan/VulkanCommon.ush
+++ b/Engine/Shaders/Public/Platform/Vulkan/VulkanCommon.ush
@@ -5,7 +5,7 @@
 =============================================================================*/
 
 // Update this GUID to improve shader recompilation for Vulkan only shaders
-// GUID = 4671795C-DD86-4E61-BFDD-52511D0D353F
+// GUID = 8F54F78F-6560-4404-B58B-D11D22128DFC
 // Glslang 1.1.82.1
 
 #pragma once
diff --git a/Engine/Source/Runtime/Engine/Private/ShaderCompiler/ShaderCompiler.cpp b/Engine/Source/Runtime/Engine/Private/ShaderCompiler/ShaderCompiler.cpp
--- a/Engine/Source/Runtime/Engine/Private/ShaderCompiler/ShaderCompiler.cpp
+++ b/Engine/Source/Runtime/Engine/Private/ShaderCompiler/ShaderCompiler.cpp
@@ -3141,15 +3141,6 @@ void GlobalBeginCompileShader(
 		}
 	}
 
-	if (IsVulkanPlatform((EShaderPlatform)Target.Platform))
-	{
-		static const auto* CVar = IConsoleManager::Get().FindConsoleVariable(TEXT("r.Vulkan.EnableTessellation"));
-		if (CVar && CVar->GetInt() != 0)
-		{
-			Input.Environment.SetDefine(TEXT("VULKAN_ENABLE_TESSELLATION"), 1);
-		}
-	}
-	
 	if (IsMetalPlatform((EShaderPlatform)Target.Platform))
 	{
 		{
diff --git a/Engine/Source/Runtime/RHI/Private/RHI.cpp b/Engine/Source/Runtime/RHI/Private/RHI.cpp
--- a/Engine/Source/Runtime/RHI/Private/RHI.cpp
+++ b/Engine/Source/Runtime/RHI/Private/RHI.cpp
@@ -796,7 +796,7 @@ RHI_API bool RHISupportsTessellation(const EShaderPlatform Platform)
 {
 	if (IsFeatureLevelSupported(Platform, ERHIFeatureLevel::SM5))
 	{
-		return (Platform == SP_PCD3D_SM5) || (Platform == SP_XBOXONE_D3D12) || (Platform == SP_OPENGL_SM5) || (Platform == SP_OPENGL_ES31_EXT) || (Platform == SP_METAL_SM5) /*|| (IsVulkanSM5Platform(Platform))*/;
+		return (Platform == SP_PCD3D_SM5) || (Platform == SP_XBOXONE_D3D12) || (Platform == SP_OPENGL_SM5) || (Platform == SP_OPENGL_ES31_EXT) || (Platform == SP_METAL_SM5) || (IsVulkanSM5Platform(Platform));
 	}
 	return false;
 }
