Convert properly to nits for PQ lut indexing

diff --git a/Engine/Shaders/Private/PostProcessCombineLUTs.usf b/Engine/Shaders/Private/PostProcessCombineLUTs.usf
--- a/Engine/Shaders/Private/PostProcessCombineLUTs.usf
+++ b/Engine/Shaders/Private/PostProcessCombineLUTs.usf
@@ -254,8 +254,13 @@ float4 CombineLUTsCommon(float2 InUV, uint InLayerIndex)
 	float3 LUTEncodedColor = Neutral.rgb;
 	float3 LinearColor;
 	// Decode texture values as ST-2084 (Dolby PQ)
-	if( GetOutputDevice() >= 3 )
-		LinearColor = ST2084ToLinear( LUTEncodedColor );
+	if (GetOutputDevice() >= 3)
+	{
+		// Since ST2084 returns linear values in nits, divide by a scale factor to convert
+		// the reference nit result to be 1.0 in linear.
+		// (for efficiency multiply by precomputed inverse)
+		LinearColor = ST2084ToLinear(LUTEncodedColor) * LinearToNitsScaleInverse;
+	}
 	// Decode log values
 	else
 		LinearColor = LogToLin( LUTEncodedColor ) - LogToLin( 0 );
diff --git a/Engine/Shaders/Private/PostProcessTonemap.usf b/Engine/Shaders/Private/PostProcessTonemap.usf
--- a/Engine/Shaders/Private/PostProcessTonemap.usf
+++ b/Engine/Shaders/Private/PostProcessTonemap.usf
@@ -77,8 +77,9 @@ half3 ColorLookupTable( half3 LinearColor )
 	float3 LUTEncodedColor;
 	// Encode as ST-2084 (Dolby PQ) values
 	#if (DIM_OUTPUT_DEVICE == TONEMAPPER_OUTOUT_ACES1000nitST2084 || DIM_OUTPUT_DEVICE == TONEMAPPER_OUTOUT_ACES2000nitST2084 || DIM_OUTPUT_DEVICE == TONEMAPPER_OUTOUT_ACES1000nitScRGB || DIM_OUTPUT_DEVICE == TONEMAPPER_OUTOUT_ACES2000nitScRGB || DIM_OUTPUT_DEVICE == TONEMAPPER_OUTOUT_LinearEXR)
-		LUTEncodedColor = LinearToST2084( LinearColor );
-
+		// ST2084 expects to receive linear values 0-10000 in nits.
+		// So the linear value must be multiplied by a scale factor to convert to nits.
+		LUTEncodedColor = LinearToST2084(LinearColor * LinearToNitsScale);
 	#else
 		LUTEncodedColor = LinToLog( LinearColor + LogToLin( 0 ) );
 
diff --git a/Engine/Shaders/Private/TonemapCommon.ush b/Engine/Shaders/Private/TonemapCommon.ush
--- a/Engine/Shaders/Private/TonemapCommon.ush
+++ b/Engine/Shaders/Private/TonemapCommon.ush
@@ -10,6 +10,14 @@
 // usually 1/2.2, the .y is used for inverse gamma when "gamma only" mode is not used
 half3 InverseGamma;
 
+// Scale factor for converting pixel values to nits. 
+// This value is required for PQ (ST2084) conversions, because PQ linear values are in nits. 
+// The purpose is to make good use of PQ lut entries. A scale factor of 100 conveniently places 
+// about half of the PQ lut indexing below 1.0, with the other half for input values over 1.0.
+// Also, 100nits is the expected monitor brightness for a 1.0 pixel value without a tone curve.
+float LinearToNitsScale = 100.0;
+float LinearToNitsScaleInverse = 1.0 / 100.0;
+
 // Film tonal and color control.
 half4 ColorMatrixR_ColorCurveCd1;
 half4 ColorMatrixG_ColorCurveCd3Cm3;
