Uses view state's frame index instead of global frame index to properly handle r.Test.FreezeTemporalSequences and related debuging functionalities.

diff --git a/Engine/Shaders/Private/QuasiRandom.ush b/Engine/Shaders/Private/QuasiRandom.ush
--- a/Engine/Shaders/Private/QuasiRandom.ush
+++ b/Engine/Shaders/Private/QuasiRandom.ush
@@ -50,7 +50,7 @@ void QuasiRandom_Initialize(inout FQuasiRandom QuasiRandom, float2 PixelIndex, u
 
 void QuasiRandom_Update(inout FQuasiRandom QuasiRandom, uint Iteration, uint IterationStride)
 {
-	QuasiRandom.TimeSeed = View.FrameNumber * IterationStride + Iteration;
+	QuasiRandom.TimeSeed = View.StateFrameIndex * IterationStride + Iteration;
 	QuasiRandom.HaltonIteration = Iteration;
 	QuasiRandom.Dimension = 0;
 }
