Fixes shader parameter binding failure of the tonemapper and combine LUT shaders.

diff --git a/Engine/Shaders/Private/TonemapCommon.ush b/Engine/Shaders/Private/TonemapCommon.ush
--- a/Engine/Shaders/Private/TonemapCommon.ush
+++ b/Engine/Shaders/Private/TonemapCommon.ush
@@ -15,8 +15,8 @@ half3 InverseGamma;
 // The purpose is to make good use of PQ lut entries. A scale factor of 100 conveniently places 
 // about half of the PQ lut indexing below 1.0, with the other half for input values over 1.0.
 // Also, 100nits is the expected monitor brightness for a 1.0 pixel value without a tone curve.
-float LinearToNitsScale = 100.0;
-float LinearToNitsScaleInverse = 1.0 / 100.0;
+static const float LinearToNitsScale = 100.0;
+static const float LinearToNitsScaleInverse = 1.0 / 100.0;
 
 // Film tonal and color control.
 half4 ColorMatrixR_ColorCurveCd1;
