DR - Workaround hlslcc bug to unblock platforms
#rnx

diff --git a/Engine/Shaders/Private/ScreenSpaceReflections.usf b/Engine/Shaders/Private/ScreenSpaceReflections.usf
--- a/Engine/Shaders/Private/ScreenSpaceReflections.usf
+++ b/Engine/Shaders/Private/ScreenSpaceReflections.usf
@@ -32,6 +32,13 @@ SamplerState SceneColorSampler;
 Texture2D HZB;
 SamplerState HZBSampler;
 
+#if COMPILER_HLSLCC
+#define INFINITE_WORKAROUND		3.402823e+38
+#else
+#define INFINITE_WORKAROUND		1.#INF
+#endif
+
+
 float GetRoughness(const in FGBufferData GBuffer)
 {
 	float Roughness = GBuffer.Roughness;
@@ -133,7 +140,7 @@ void ScreenSpaceReflectionsPS(
 	float NoV = saturate( dot( N, V ) );
 	float G_SmithV = 2 * NoV / (NoV + sqrt(NoV * (NoV - NoV * a2) + a2));
 
-	float ClosestHitDistanceSqr = 1.#INF;
+	float ClosestHitDistanceSqr = INFINITE_WORKAROUND;
 
 #if SSR_QUALITY == 1
 	uint NumSteps = 8;
@@ -270,7 +277,7 @@ void ScreenSpaceReflectionsPS(
 		OutColor /= max( NumRays, 0.0001 );
 		OutColor.rgb *= rcp( 1 - Luminance(OutColor.rgb) );
 
-		if (ClosestHitDistanceSqr == 1.#INF)
+		if (ClosestHitDistanceSqr == INFINITE_WORKAROUND)
 		{
 			ClosestHitDistanceSqr = DENOISER_MISS_HIT_DISTANCE;
 		}
