Uses the same random generator between SSR and RTR for reflection denoiser 1spp input.

diff --git a/Engine/Shaders/Private/CapsuleLightIntegrate.ush b/Engine/Shaders/Private/CapsuleLightIntegrate.ush
--- a/Engine/Shaders/Private/CapsuleLightIntegrate.ush
+++ b/Engine/Shaders/Private/CapsuleLightIntegrate.ush
@@ -6,6 +6,7 @@
 #include "MonteCarlo.ush"
 #include "ShadingModels.ush"
 #include "CapsuleLight.ush"
+#include "SobolRandom.ush"
 
 Texture2D DummyRectLightTextureForCapsuleCompilerWarning;
 
diff --git a/Engine/Shaders/Private/DiaphragmDOF/DOFRecombine.usf b/Engine/Shaders/Private/DiaphragmDOF/DOFRecombine.usf
--- a/Engine/Shaders/Private/DiaphragmDOF/DOFRecombine.usf
+++ b/Engine/Shaders/Private/DiaphragmDOF/DOFRecombine.usf
@@ -10,6 +10,8 @@
 #include "../SceneTexturesCommon.ush"
 #include "../Random.ush"
 #include "../MonteCarlo.ush"
+#include "../SobolRandom.ush"
+
 
 
 //------------------------------------------------------- DEBUG COMPILE TIME CONFIG
diff --git a/Engine/Shaders/Private/MaterialTemplate.ush b/Engine/Shaders/Private/MaterialTemplate.ush
--- a/Engine/Shaders/Private/MaterialTemplate.ush
+++ b/Engine/Shaders/Private/MaterialTemplate.ush
@@ -6,6 +6,7 @@
 
 #include "/Engine/Private/SceneTexturesCommon.ush"
 #include "/Engine/Private/Random.ush"
+#include "/Engine/Private/SobolRandom.ush"
 #include "/Engine/Private/MonteCarlo.ush"
 #include "/Engine/Generated/UniformBuffers/Material.ush"
 #include "/Engine/Private/DepthOfFieldCommon.ush"
diff --git a/Engine/Shaders/Private/MonteCarlo.ush b/Engine/Shaders/Private/MonteCarlo.ush
--- a/Engine/Shaders/Private/MonteCarlo.ush
+++ b/Engine/Shaders/Private/MonteCarlo.ush
@@ -76,63 +76,6 @@ float2 JitteredR2( uint Index, float2 Jitter, float JitterAmount = 0.5 )
 	return frac( a * Index + ( JitterAmount * 0.25 * d0 * sqrt(PI) * rsqrt( Index - i0 ) ) * Jitter );
 }
 
-///////
-// Typical usage of the Sobol functions for one or more points per pixel
-//   uint2 SobolBase = SobolPixel(uint2(SvPosition.xy));                    // First sample for this pixel
-//   for(int i = 0; i < N; ++i)
-//     float2 Point = float2(SobolIndex(SobolBase, i)) / 0x10000;           // Points as [0,1) x [0,1)
-//
-// Typical usage for one or more points per frame
-//   uint2 SobolBase = SobolPixel(uint2(SvPosition.xy));                    // Frame 0, point 0 for this pixel
-//   uint2 SobolFrame = SobolIndex(SobolBase, View.StateFrameIndexMod8, 3); // Least significant bits for frame
-//   for(int i = 0; i < N; ++i)
-//     float2 Point = float2(SobolIndex(SobolFrame, i << 3)) / 0x10000;     // Higher-order bits for point within frame
-//
-// For additional independent point sets
-//   uint2 SobolBase2 = SobolBase ^ RandSeed; // where RandSeed is a uint2 with values in 0 to 0xffff
-
-// Compute a Sobol-distributed point from a 256x256 pixel grid one pixel in that grid
-// @param Pixel Pixel/cell position in the 256x256 grid
-// @return Sobol position relative to the pixel corner, with components in the range 0 to 0xffff
-uint2 SobolPixel(uint2 Pixel)
-{
-	// look up for pixel
-	int3 SobolLo = int3(Pixel & 0xfu, 0);
-	int3 SobolHi = int3((Pixel >> 4u) & 0xfu, 0) + int3(16, 0, 0);
-	uint Packed = View.SobolSamplingTexture.Load(SobolLo) ^ View.SobolSamplingTexture.Load(SobolHi);
-	return uint2(Packed, Packed << 8u) & 0xff00u;
-}
-
-// Evaluate additional Sobol points within a pixel
-// @param Base  Base Sobol point for this pixel
-// @param Index Which 2D Sobol point to return
-// @param Bits  Optional max bits in index (to avoid extra calculation)
-// @return Sobol position in the range 0 to 0xffff
-uint2 SobolIndex(uint2 Base, int Index, int Bits = 10)
-{
-	uint2 SobolNumbers[10] = {
-		uint2(0x8680u, 0x4c80u), uint2(0xf240u, 0x9240u), uint2(0x8220u, 0x0e20u), uint2(0x4110u, 0x1610u), uint2(0xa608u, 0x7608u),
-		uint2(0x8a02u, 0x280au), uint2(0xe204u, 0x9e04u), uint2(0xa400u, 0x4682u), uint2(0xe300u, 0xa74du), uint2(0xb700u, 0x9817u),
-	};
-
-	uint2 Result = Base;
-	UNROLL for (int b = 0; b < 10 && b < Bits; ++b)
-	{
-		Result ^= (Index & (1 << b)) ? SobolNumbers[b] : 0;
-	}
-	return Result;
-}
-
-/** Returns an unique sobol sample for the frame. */
-uint2 ComputePixelUniqueSobolRandSample(uint2 PixelCoord)
-{
-	const uint TemporalBits = 10;
-	uint FrameIndexMod1024 = ReverseBitsN(GetPowerOfTwoModulatedFrameIndex(1 << TemporalBits), TemporalBits);
-
-	uint2 SobolBase = SobolPixel(PixelCoord);
-	return SobolIndex(SobolBase, FrameIndexMod1024, TemporalBits);
-}
-
 
 ///////
 
diff --git a/Engine/Shaders/Private/RayTracing/RayTracingOcclusionRGS.usf b/Engine/Shaders/Private/RayTracing/RayTracingOcclusionRGS.usf
--- a/Engine/Shaders/Private/RayTracing/RayTracingOcclusionRGS.usf
+++ b/Engine/Shaders/Private/RayTracing/RayTracingOcclusionRGS.usf
@@ -6,6 +6,7 @@
 #include "../LightShaderParameters.ush"
 #include "../SceneViewFamilyBlackboard.ush"
 #include "../ScreenSpaceDenoise/SSDPublic.ush"
+#include "../SobolRandom.ush"
 #include "RayTracingCommon.ush"
 #include "RayTracingDirectionalLight.ush"
 #include "RayTracingRectLight.ush"
@@ -28,11 +29,6 @@ uint SamplesPerPixel;
 float NormalBias;
 int4 LightScissor;
 
-float2 SobolIndexToUniformUnitSquare(uint2 SobolRand)
-{
-	return float2(SobolRand) * rcp(65536.0) + rcp(65536.0 * 2.0);
-}
-
 bool GenerateOcclusionRay(
 	FLightShaderParameters LightParameters,
 	float3 WorldPosition,
diff --git a/Engine/Shaders/Private/RayTracing/RayTracingReflections.usf b/Engine/Shaders/Private/RayTracing/RayTracingReflections.usf
--- a/Engine/Shaders/Private/RayTracing/RayTracingReflections.usf
+++ b/Engine/Shaders/Private/RayTracing/RayTracingReflections.usf
@@ -15,11 +15,6 @@
 
 #include "../Common.ush"
 
-float2 SobolIndexToUniformUnitSquare(uint2 SobolRand)
-{
-	return float2(SobolRand) * rcp(65536.0) + rcp(65536.0 * 2.0);
-}
-
 #define SUPPORT_CONTACT_SHADOWS		0
 #define USE_SOURCE_TEXTURE			1
 #define USE_SOURCE_TEXTURE_ARRAY	1
@@ -39,6 +34,7 @@ float2 SobolIndexToUniformUnitSquare(uint2 SobolRand)
 #include "../HeightFogCommon.ush"
 #include "../LightShaderParameters.ush"
 #include "../ScreenSpaceDenoise/SSDPublic.ush"
+#include "../SobolRandom.ush"
 #include "RayTracingHitGroupCommon.ush"
 #include "RayTracingDirectionalLight.ush"
 #include "RayTracingRectLight.ush"
@@ -105,16 +101,16 @@ float3 GenerateReflectionDirection(inout RandomSequence RandSequence, uint2 Disp
 	float3 outDirection;
 	do
 	{
-		float2 RandSample = RandomSequence_GenerateSample2D(RandSequence, DummyVariable);
+		float2 E = RandomSequence_GenerateSample2D(RandSequence, DummyVariable);
 
 		// Sobol can only be used on the first sample of the first bounce
 		if (BounceIndex == 0 && SamplePerPixel == 1 && RetryCount == 0)
 		{
-			uint2 SobolFrame = ComputePixelUniqueSobolRandSample(DispatchThreadId);
-			RandSample = SobolIndexToUniformUnitSquare(SobolFrame);
+			E = Rand1SPPDenoiserInput(DispatchThreadId);
+			E.x *= 1 - GGX_IMPORTANT_SAMPLE_BIAS;
 		}
 
-		outDirection = GenerateReflectedRayDirection(IncidentDirection, WorldNormal, Roughness, RandSample);
+		outDirection = GenerateReflectedRayDirection(IncidentDirection, WorldNormal, Roughness, E);
 	} while (dot(WorldNormal, outDirection) < 0.0 && ++RetryCount < 16);
 
 	return outDirection;
diff --git a/Engine/Shaders/Private/RayTracing/RayTracingTranslucency.usf b/Engine/Shaders/Private/RayTracing/RayTracingTranslucency.usf
--- a/Engine/Shaders/Private/RayTracing/RayTracingTranslucency.usf
+++ b/Engine/Shaders/Private/RayTracing/RayTracingTranslucency.usf
@@ -6,11 +6,6 @@
 
 #include "../Common.ush"
 
-float2 SobolIndexToUniformUnitSquare(uint2 SobolRand)
-{
-	return float2(SobolRand) * rcp(65536.0) + rcp(65536.0 * 2.0);
-}
-
 #define SUPPORT_CONTACT_SHADOWS		0
 #define USE_SOURCE_TEXTURE			1
 #define USE_SOURCE_TEXTURE_ARRAY	1
@@ -28,6 +23,7 @@ float2 SobolIndexToUniformUnitSquare(uint2 SobolRand)
 #include "../Montecarlo.ush"
 #include "../PathTracing/Utilities/PathTracingRandomSequence.ush" 
 #include "../HeightFogCommon.ush"
+#include "../SobolRandom.ush"
 #include "RayTracingCommon.ush"
 #include "RayTracingHitGroupCommon.ush"
 #include "RayTracingDirectionalLight.ush"
diff --git a/Engine/Shaders/Private/ScreenSpaceDenoise/SSDPublic.ush b/Engine/Shaders/Private/ScreenSpaceDenoise/SSDPublic.ush
--- a/Engine/Shaders/Private/ScreenSpaceDenoise/SSDPublic.ush
+++ b/Engine/Shaders/Private/ScreenSpaceDenoise/SSDPublic.ush
@@ -2,6 +2,10 @@
 
 #pragma once
 
+#include "../Random.ush"
+#include "../MonteCarlo.ush"
+#include "../SobolRandom.ush"
+
 
 /** Bias used on GGX important sample when denoising, to remove part of the tale that create a lot more noise. */
 #define GGX_IMPORTANT_SAMPLE_BIAS 0.1
@@ -11,3 +15,26 @@
 
 /** Special value to be encoded in denoiser input when no ray intersection has been found. */
 #define DENOISER_MISS_HIT_DISTANCE -1.0
+
+
+/** Random generator used for 1 sample per pixel for denoiser input. */
+float2 Rand1SPPDenoiserInput(uint2 PixelPos)
+{
+	float2 E;
+
+	#if 1
+	{
+		uint2 Random = Rand3DPCG16( int3( PixelPos, View.StateFrameIndexMod8 ) ).xy;
+		E = float2(Random) * rcp(65536.0); // equivalent to Hammersley16(0, 1, Random).
+	}
+	#elif 0
+	{
+		uint2 SobolFrame = ComputePixelUniqueSobolRandSample(PixelPos);
+		E = SobolIndexToUniformUnitSquare(SobolFrame);
+	}
+	#else
+		#error Miss-configured random generator.
+	#endif
+
+	return E;
+}
diff --git a/Engine/Shaders/Private/ScreenSpaceDenoise/SSDSignalFramework.ush b/Engine/Shaders/Private/ScreenSpaceDenoise/SSDSignalFramework.ush
--- a/Engine/Shaders/Private/ScreenSpaceDenoise/SSDSignalFramework.ush
+++ b/Engine/Shaders/Private/ScreenSpaceDenoise/SSDSignalFramework.ush
@@ -442,12 +442,6 @@ float3 GenerateReflectedRayDirection(
 	return RayDirection;
 }
 
-// TODO: duplicated...
-float2 SobolIndexToUniformUnitSquare(uint2 SobolRand)
-{
-	return float2(SobolRand) * rcp(65536.0) + rcp(65536.0 * 2.0);
-}
-
 // TODO: duplicated...
 void GetInputRayDirectionPixel(
 	uint2 PixelCoord,
@@ -457,22 +451,17 @@ void GetInputRayDirectionPixel(
 {
 	float3 IncidentDirection = normalize(SceneMetadata.TranslatedWorldPosition - View.TranslatedWorldCameraOrigin);
 
-	float2 RandSample;
+	float2 E;
 	{
-		const uint TemporalBits = 10;
-		uint FrameIndexMod1024 = ReverseBitsN(GetPowerOfTwoModulatedFrameIndex(1 << TemporalBits), TemporalBits);
-
-		uint2 SobolBase = SobolPixel(PixelCoord);
-		uint2 SobolFrame = SobolIndex(SobolBase, FrameIndexMod1024, TemporalBits);
-
-		RandSample = SobolIndexToUniformUnitSquare(SobolFrame);
+		E = Rand1SPPDenoiserInput(PixelCoord);
+		E.x *= 1 - GGX_IMPORTANT_SAMPLE_BIAS;
 	}
 
 	GenerateReflectedRayDirection(
 		IncidentDirection,
 		SceneMetadata.WorldNormal,
 		SceneMetadata.Roughness,
-		RandSample,
+		E,
 		/* out = */ RayDirection,
 		/* out = */ PDF);
 }
diff --git a/Engine/Shaders/Private/ScreenSpaceReflections.usf b/Engine/Shaders/Private/ScreenSpaceReflections.usf
--- a/Engine/Shaders/Private/ScreenSpaceReflections.usf
+++ b/Engine/Shaders/Private/ScreenSpaceReflections.usf
@@ -290,17 +290,16 @@ void ScreenSpaceReflectionsPS(
 		float3 L;
 		if (bGlossy)
 		{
-			uint2 Random = Rand3DPCG16( int3( PixelPos, View.StateFrameIndexMod8 ) ).xy;
-			float2 E = float2(Random) * rcp(65536.0);
-			
-			float3x3 TangentBasis = GetTangentBasis( N );
-			float3 TangentV = mul( TangentBasis, V );
+			float2 E = Rand1SPPDenoiserInput(PixelPos);
 			
 			#if SSR_OUTPUT_FOR_DENOISER
 			{
 				E.x *= 1 - GGX_IMPORTANT_SAMPLE_BIAS;
 			}
 			#endif
+			
+			float3x3 TangentBasis = GetTangentBasis( N );
+			float3 TangentV = mul( TangentBasis, V );
 
 			float3 H = mul( ImportanceSampleVisibleGGX(E, a2, TangentV ).xyz, TangentBasis );
 			L = 2 * dot( V, H ) * H - V;
diff --git a/Engine/Shaders/Private/ShadowPercentageCloserFiltering.ush b/Engine/Shaders/Private/ShadowPercentageCloserFiltering.ush
--- a/Engine/Shaders/Private/ShadowPercentageCloserFiltering.ush
+++ b/Engine/Shaders/Private/ShadowPercentageCloserFiltering.ush
@@ -6,6 +6,8 @@
 
 #include "ShadowFilteringCommon.ush"
 #include "MonteCarlo.ush"
+#include "SobolRandom.ush"
+
 
 // CONFIGURATION OF PCSS
 // -----------------------------------------------------------------------
diff --git a/Engine/Shaders/Private/SobolRandom.ush b/Engine/Shaders/Private/SobolRandom.ush
new file mode 100644
--- /dev/null
+++ b/Engine/Shaders/Private/SobolRandom.ush
@@ -0,0 +1,67 @@
+// Copyright 1998-2019 Epic Games, Inc. All Rights Reserved.
+
+#pragma once
+
+
+///////
+// Typical usage of the Sobol functions for one or more points per pixel
+//   uint2 SobolBase = SobolPixel(uint2(SvPosition.xy));                    // First sample for this pixel
+//   for(int i = 0; i < N; ++i)
+//     float2 Point = float2(SobolIndex(SobolBase, i)) / 0x10000;           // Points as [0,1) x [0,1)
+//
+// Typical usage for one or more points per frame
+//   uint2 SobolBase = SobolPixel(uint2(SvPosition.xy));                    // Frame 0, point 0 for this pixel
+//   uint2 SobolFrame = SobolIndex(SobolBase, View.StateFrameIndexMod8, 3); // Least significant bits for frame
+//   for(int i = 0; i < N; ++i)
+//     float2 Point = float2(SobolIndex(SobolFrame, i << 3)) / 0x10000;     // Higher-order bits for point within frame
+//
+// For additional independent point sets
+//   uint2 SobolBase2 = SobolBase ^ RandSeed; // where RandSeed is a uint2 with values in 0 to 0xffff
+
+// Compute a Sobol-distributed point from a 256x256 pixel grid one pixel in that grid
+// @param Pixel Pixel/cell position in the 256x256 grid
+// @return Sobol position relative to the pixel corner, with components in the range 0 to 0xffff
+uint2 SobolPixel(uint2 Pixel)
+{
+	// look up for pixel
+	int3 SobolLo = int3(Pixel & 0xfu, 0);
+	int3 SobolHi = int3((Pixel >> 4u) & 0xfu, 0) + int3(16, 0, 0);
+	uint Packed = View.SobolSamplingTexture.Load(SobolLo) ^ View.SobolSamplingTexture.Load(SobolHi);
+	return uint2(Packed, Packed << 8u) & 0xff00u;
+}
+
+// Evaluate additional Sobol points within a pixel
+// @param Base  Base Sobol point for this pixel
+// @param Index Which 2D Sobol point to return
+// @param Bits  Optional max bits in index (to avoid extra calculation)
+// @return Sobol position in the range 0 to 0xffff
+uint2 SobolIndex(uint2 Base, int Index, int Bits = 10)
+{
+	uint2 SobolNumbers[10] = {
+		uint2(0x8680u, 0x4c80u), uint2(0xf240u, 0x9240u), uint2(0x8220u, 0x0e20u), uint2(0x4110u, 0x1610u), uint2(0xa608u, 0x7608u),
+		uint2(0x8a02u, 0x280au), uint2(0xe204u, 0x9e04u), uint2(0xa400u, 0x4682u), uint2(0xe300u, 0xa74du), uint2(0xb700u, 0x9817u),
+	};
+
+	uint2 Result = Base;
+	UNROLL for (int b = 0; b < 10 && b < Bits; ++b)
+	{
+		Result ^= (Index & (1 << b)) ? SobolNumbers[b] : 0;
+	}
+	return Result;
+}
+
+/** Returns an unique sobol sample for the frame. */
+uint2 ComputePixelUniqueSobolRandSample(uint2 PixelCoord)
+{
+	const uint TemporalBits = 10;
+	uint FrameIndexMod1024 = ReverseBitsN(GetPowerOfTwoModulatedFrameIndex(1 << TemporalBits), TemporalBits);
+
+	uint2 SobolBase = SobolPixel(PixelCoord);
+	return SobolIndex(SobolBase, FrameIndexMod1024, TemporalBits);
+}
+
+/** Maps a sobol random generated sample to 2D unit square. */
+float2 SobolIndexToUniformUnitSquare(uint2 SobolRand)
+{
+	return float2(SobolRand) * rcp(65536.0) + rcp(65536.0 * 2.0);
+}
