Fixing ray tracing global illumination energy explosion

The energy explosion is caused by two problems:
1) Neglecting to include tmin and tmax in IsHit() comparisions
2) Adding normal-based biasing on subsequent path chains

Verified correctness with path tracing after comparing results under FurnaceTest.

diff --git a/Engine/Shaders/Private/RayTracing/RayTracingCommon.ush b/Engine/Shaders/Private/RayTracing/RayTracingCommon.ush
--- a/Engine/Shaders/Private/RayTracing/RayTracingCommon.ush
+++ b/Engine/Shaders/Private/RayTracing/RayTracingCommon.ush
@@ -298,13 +298,13 @@ float GetHitT(FMaterialClosestHitPayload HitInfo)
 bool IsHit(RayDesc Ray, FMaterialClosestHitPayload HitInfo)
 {
 	float t = HitInfo.HitT;
-	return t > Ray.TMin && t < Ray.TMax;
+	return t >= Ray.TMin && t <= Ray.TMax;
 }
 
 bool IsHit(RayDesc Ray, FPackedMaterialClosestHitPayload HitInfo)
 {
 	float t = HitInfo.HitT;
-	return t > Ray.TMin && t < Ray.TMax;
+	return t >= Ray.TMin && t <= Ray.TMax;
 }
 
 #endif // !COMPUTESHADER
diff --git a/Engine/Shaders/Private/RayTracing/RayTracingGlobalIlluminationRGS.usf b/Engine/Shaders/Private/RayTracing/RayTracingGlobalIlluminationRGS.usf
--- a/Engine/Shaders/Private/RayTracing/RayTracingGlobalIlluminationRGS.usf
+++ b/Engine/Shaders/Private/RayTracing/RayTracingGlobalIlluminationRGS.usf
@@ -47,7 +47,7 @@ void GenerateCosineNormalRay(
 
 	RayOrigin = WorldPosition;
 	RayDirection = Direction_World;
-	RayTMin = 0.0;
+	RayTMin = 0.01;
 	RayTMax = max(MaxRayDistanceForGI, MaxRayDistanceForAO);
 	RayPdf = Direction_Tangent.w;
 }
@@ -284,7 +284,8 @@ void GlobalIlluminationRGS()
 				float4 RandSample;
 				RandSample.xyz = QuasiRandom_Sample3(QuasiRandom);
 				RandSample.w = QuasiRandom_Sample3(QuasiRandom).x;
-				SampleMaterial(Ray.Direction, Payload, RandSample, Direction, Throughput, MaterialPdf);
+				float PositionBiasSign;
+				SampleMaterial(Ray.Direction, Payload, RandSample, Direction, Throughput, MaterialPdf, PositionBiasSign);
 #else
 				float3 RandSample3 = QuasiRandom_Sample3(QuasiRandom);
 				float3 RayOrigin = Ray.Origin;
@@ -305,7 +306,6 @@ void GlobalIlluminationRGS()
 
 				// Update ray
 				Ray.Direction = Direction;
-				ApplyPositionBias(Ray, WorldNormal, MaxNormalBias);
 				RayThroughput *= Throughput;
 				RayPdf *= MaterialPdf;
 			}
