Attempt blind fix of VT related compile errors in ray tracing shaders

diff --git a/Engine/Shaders/Private/RayTracing/RayTracingMaterialHitShaders.usf b/Engine/Shaders/Private/RayTracing/RayTracingMaterialHitShaders.usf
--- a/Engine/Shaders/Private/RayTracing/RayTracingMaterialHitShaders.usf
+++ b/Engine/Shaders/Private/RayTracing/RayTracingMaterialHitShaders.usf
@@ -53,7 +53,8 @@ void GetSkyLighting(FMaterialPixelParameters MaterialParameters, float3 WorldNor
 			#elif HQ_TEXTURE_LIGHTMAP
 
 				// Bent normal from precomputed texture
-				float4 WorldSkyBentNormalAndOcclusion = GetSkyBentNormalAndOcclusion(LightmapUV * float2(1, 2), MaterialParameters.SvPosition.xy);
+				VTPageTableResult DummyVTPageTableResult = (VTPageTableResult)0.0f;
+				float4 WorldSkyBentNormalAndOcclusion = GetSkyBentNormalAndOcclusion(DummyVTPageTableResult, LightmapUV * float2(1, 2), 0u, MaterialParameters.SvPosition.xy);
 				// Renormalize as vector was quantized and compressed
 				float3 NormalizedBentNormal = normalize(WorldSkyBentNormalAndOcclusion.xyz);
 				SkyVisibility = WorldSkyBentNormalAndOcclusion.w;
@@ -142,8 +143,8 @@ void GetPrecomputedIndirectLightingAndSkyLight(
 		GetLightMapCoordinates(Interpolants, LightmapUV0, LightmapUV1, LightmapDataIndex);
 		SkyOcclusionUV = LightmapUV0;
 
-		uint VTRequest = 0;
-		GetLightMapColorHQ(LightmapUV0, LightmapUV1, LightmapDataIndex, DiffuseDir, MaterialParameters.SvPosition.xy, ShadingModel, OutDiffuseLighting, OutSubsurfaceLighting, VTRequest);
+		VTPageTableResult DummyVTPageTableResult = (VTPageTableResult)0.0f;
+		GetLightMapColorHQ(DummyVTPageTableResult, LightmapUV0, LightmapUV1, LightmapDataIndex, DiffuseDir, MaterialParameters.SvPosition.xy, ShadingModel, OutDiffuseLighting, OutSubsurfaceLighting);
 
 	// Low quality texture lightmaps
 	#elif LQ_TEXTURE_LIGHTMAP
