Shader compilation fix (based on another earlier fix for HLSLCC)
#rnx

diff --git a/Engine/Shaders/Private/ScreenSpaceReflections.usf b/Engine/Shaders/Private/ScreenSpaceReflections.usf
--- a/Engine/Shaders/Private/ScreenSpaceReflections.usf
+++ b/Engine/Shaders/Private/ScreenSpaceReflections.usf
@@ -32,7 +32,7 @@ SamplerState SceneColorSampler;
 Texture2D HZB;
 SamplerState HZBSampler;
 
-#if COMPILER_HLSLCC
+#if COMPILER_HLSLCC || PS4_PROFILE
 #define INFINITE_WORKAROUND		3.402823e+38
 #else
 #define INFINITE_WORKAROUND		1.#INF
