Enable VT lightmap sampling in VT hit shaders to properly fix UE-74416

diff --git a/Engine/Shaders/Private/RayTracing/RayTracingMaterialHitShaders.usf b/Engine/Shaders/Private/RayTracing/RayTracingMaterialHitShaders.usf
--- a/Engine/Shaders/Private/RayTracing/RayTracingMaterialHitShaders.usf
+++ b/Engine/Shaders/Private/RayTracing/RayTracingMaterialHitShaders.usf
@@ -26,10 +26,11 @@
 #include "/Engine/Private/DeferredShadingCommon.ush"
 #include "/Engine/Private/SHCommon.ush"
 #include "/Engine/Private/ReflectionEnvironmentShared.ush"
+#include "/Engine/Private/VirtualTextureCommon.ush"
 #include "/Engine/Private/LightmapCommon.ush"
 
 /** Computes sky diffuse lighting, including precomputed shadowing. */
-void GetSkyLighting(FMaterialPixelParameters MaterialParameters, float3 WorldNormal, float2 LightmapUV, float3 SkyOcclusionUV3D, uint ShadingModel, out float3 OutDiffuseLighting, out float3 OutSubsurfaceLighting)
+void GetSkyLighting(FMaterialPixelParameters MaterialParameters, VTPageTableResult LightmapVTPageTableResult, float3 WorldNormal, float2 LightmapUV, uint LightmapDataIndex, float3 SkyOcclusionUV3D, uint ShadingModel, out float3 OutDiffuseLighting, out float3 OutSubsurfaceLighting)
 {
 	OutDiffuseLighting = 0;
 	OutSubsurfaceLighting = 0;
@@ -53,8 +54,7 @@ void GetSkyLighting(FMaterialPixelParameters MaterialParameters, float3 WorldNor
 			#elif HQ_TEXTURE_LIGHTMAP
 
 				// Bent normal from precomputed texture
-				VTPageTableResult DummyVTPageTableResult = (VTPageTableResult)0.0f;
-				float4 WorldSkyBentNormalAndOcclusion = GetSkyBentNormalAndOcclusion(DummyVTPageTableResult, LightmapUV * float2(1, 2), 0u, MaterialParameters.SvPosition.xy);
+				float4 WorldSkyBentNormalAndOcclusion = GetSkyBentNormalAndOcclusion(LightmapVTPageTableResult, LightmapUV * float2(1, 2), LightmapDataIndex, MaterialParameters.SvPosition.xy);
 				// Renormalize as vector was quantized and compressed
 				float3 NormalizedBentNormal = normalize(WorldSkyBentNormalAndOcclusion.xyz);
 				SkyVisibility = WorldSkyBentNormalAndOcclusion.w;
@@ -109,6 +109,7 @@ void GetSkyLighting(FMaterialPixelParameters MaterialParameters, float3 WorldNor
 void GetPrecomputedIndirectLightingAndSkyLight(
 	FMaterialPixelParameters MaterialParameters, 
 	FVertexFactoryInterpolantsVSToPS Interpolants,
+	VTPageTableResult LightmapVTPageTableResult,
 	float3 DiffuseDir,
 	float3 VolumetricLightmapBrickTextureUVs,
 	uint ShadingModel,
@@ -120,6 +121,7 @@ void GetPrecomputedIndirectLightingAndSkyLight(
 	OutDiffuseLighting = 0;
 	OutSubsurfaceLighting = 0;
 	float2 SkyOcclusionUV = 0;
+	uint SkyOcclusionDataIndex = 0u;
 
 	#if PRECOMPUTED_IRRADIANCE_VOLUME_LIGHTING
 		FThreeBandSHVectorRGB IrradianceSH = GetVolumetricLightmapSH3(VolumetricLightmapBrickTextureUVs);
@@ -141,10 +143,12 @@ void GetPrecomputedIndirectLightingAndSkyLight(
 		float2 LightmapUV0, LightmapUV1;
 		uint LightmapDataIndex;
 		GetLightMapCoordinates(Interpolants, LightmapUV0, LightmapUV1, LightmapDataIndex);
+		#if LIGHTMAP_VT_ENABLED
+			LightmapVTPageTableResult = LightmapGetVTSampleInfo(LightmapUV0, LightmapDataIndex, MaterialParameters.SvPosition.xy);
+		#endif
 		SkyOcclusionUV = LightmapUV0;
-
-		VTPageTableResult DummyVTPageTableResult = (VTPageTableResult)0.0f;
-		GetLightMapColorHQ(DummyVTPageTableResult, LightmapUV0, LightmapUV1, LightmapDataIndex, DiffuseDir, MaterialParameters.SvPosition.xy, ShadingModel, OutDiffuseLighting, OutSubsurfaceLighting);
+		SkyOcclusionDataIndex = LightmapDataIndex;
+		GetLightMapColorHQ(LightmapVTPageTableResult, LightmapUV0, LightmapUV1, LightmapDataIndex, DiffuseDir, MaterialParameters.SvPosition.xy, ShadingModel, OutDiffuseLighting, OutSubsurfaceLighting);
 
 	// Low quality texture lightmaps
 	#elif LQ_TEXTURE_LIGHTMAP
@@ -162,7 +166,7 @@ void GetPrecomputedIndirectLightingAndSkyLight(
 
 	float3 SkyDiffuseLighting;
 	float3 SkySubsurfaceLighting;
-	GetSkyLighting(MaterialParameters, DiffuseDir, SkyOcclusionUV, VolumetricLightmapBrickTextureUVs, ShadingModel, SkyDiffuseLighting, SkySubsurfaceLighting);
+	GetSkyLighting(MaterialParameters, LightmapVTPageTableResult, DiffuseDir, SkyOcclusionUV, SkyOcclusionDataIndex, VolumetricLightmapBrickTextureUVs, ShadingModel, SkyDiffuseLighting, SkySubsurfaceLighting);
 
 	OutSubsurfaceLighting += SkySubsurfaceLighting;
 
@@ -339,7 +343,9 @@ void MaterialCHS(inout FPackedMaterialClosestHitPayload PackedPayload, in FDefau
 #if PRECOMPUTED_IRRADIANCE_VOLUME_LIGHTING
 	VolumetricLightmapBrickTextureUVs = ComputeVolumetricLightmapBrickTextureUVs(MaterialParameters.AbsoluteWorldPosition);
 #endif
-	GetPrecomputedIndirectLightingAndSkyLight(MaterialParameters, Interpolants, Payload.WorldNormal, VolumetricLightmapBrickTextureUVs, Payload.ShadingModelID, DiffuseIndirectLighting, SubsurfaceIndirectLighting, IndirectIrradiance);
+
+	VTPageTableResult LightmapVTPageTableResult = (VTPageTableResult)0.0f;
+	GetPrecomputedIndirectLightingAndSkyLight(MaterialParameters, Interpolants, LightmapVTPageTableResult, Payload.WorldNormal, VolumetricLightmapBrickTextureUVs, Payload.ShadingModelID, DiffuseIndirectLighting, SubsurfaceIndirectLighting, IndirectIrradiance);
 	
 	Payload.DiffuseColor = BaseColor - BaseColor * Metallic;
 	Payload.SpecularColor = ComputeF0(Specular, BaseColor, Metallic);
