Hardcodes the number of input and output texture in the denoiser's shaders to work arround shader compiler

diff --git a/Engine/Shaders/Private/ScreenSpaceDenoise/SSDInjest.usf b/Engine/Shaders/Private/ScreenSpaceDenoise/SSDInjest.usf
--- a/Engine/Shaders/Private/ScreenSpaceDenoise/SSDInjest.usf
+++ b/Engine/Shaders/Private/ScreenSpaceDenoise/SSDInjest.usf
@@ -21,6 +21,9 @@
 	#endif
 	#define CONFIG_MULTIPLEXED_SIGNALS_PER_SIGNAL_DOMAIN 1
 
+	#define CONFIG_INPUT_TEXTURE_COUNT DIM_SIGNAL_BATCH_SIZE
+	#define CONFIG_OUTPUT_TEXTURE_COUNT ((DIM_SIGNAL_BATCH_SIZE + 1) / 2)
+
 #else
 	#error Unknown signal processing.
 #endif
@@ -34,15 +37,54 @@
 
 //------------------------------------------------------- PARAMETERS
 
+#if !defined(CONFIG_INPUT_TEXTURE_COUNT)
+	#error Missing CONFIG_INPUT_TEXTURE_COUNT
+#endif
+
 Texture2D SignalInput_Textures_0;
+
+#if CONFIG_INPUT_TEXTURE_COUNT > 1
 Texture2D SignalInput_Textures_1;
+#else
+#define SignalInput_Textures_1 SignalInput_Textures_0
+#endif
+
+#if CONFIG_INPUT_TEXTURE_COUNT > 2
 Texture2D SignalInput_Textures_2;
+#else
+#define SignalInput_Textures_2 SignalInput_Textures_0
+#endif
+
+#if CONFIG_INPUT_TEXTURE_COUNT > 3
 Texture2D SignalInput_Textures_3;
+#else
+#define SignalInput_Textures_3 SignalInput_Textures_0
+#endif
+
+
+#if !defined(CONFIG_OUTPUT_TEXTURE_COUNT)
+	#error Missing CONFIG_OUTPUT_TEXTURE_COUNT
+#endif
 
 RWTexture2D<float4>	SignalOutput_UAVs_0;
+
+#if CONFIG_OUTPUT_TEXTURE_COUNT > 1
 RWTexture2D<float4>	SignalOutput_UAVs_1;
+#else
+#define SignalOutput_UAVs_1 SignalOutput_UAVs_0
+#endif
+
+#if CONFIG_OUTPUT_TEXTURE_COUNT > 2
 RWTexture2D<float4>	SignalOutput_UAVs_2;
+#else
+#define SignalOutput_UAVs_2 SignalOutput_UAVs_0
+#endif
+
+#if CONFIG_OUTPUT_TEXTURE_COUNT > 3
 RWTexture2D<float4>	SignalOutput_UAVs_3;
+#else
+#define SignalOutput_UAVs_3 SignalOutput_UAVs_0
+#endif
 
 
 //------------------------------------------------------- ENTRY POINT
diff --git a/Engine/Shaders/Private/ScreenSpaceDenoise/SSDSpatialAccumulation.usf b/Engine/Shaders/Private/ScreenSpaceDenoise/SSDSpatialAccumulation.usf
--- a/Engine/Shaders/Private/ScreenSpaceDenoise/SSDSpatialAccumulation.usf
+++ b/Engine/Shaders/Private/ScreenSpaceDenoise/SSDSpatialAccumulation.usf
@@ -53,6 +53,9 @@
 		#endif
 		#define CONFIG_SIGNAL_OUTPUT_LAYOUT SIGNAL_BUFFER_LAYOUT_PENUMBRA_RECONSTRUCTION
 
+		#define CONFIG_INPUT_TEXTURE_COUNT ((CONFIG_SIGNAL_BATCH_SIZE + 1) / 2)
+		#define CONFIG_OUTPUT_TEXTURE_COUNT CONFIG_SIGNAL_BATCH_SIZE
+
 		#define CONFIG_SAMPLE_SET           SAMPLE_SET_STACKOWIAK_4_SETS
 		#define CONFIG_BILATERAL_DISTANCE_COMPUTATION SIGNAL_WORLD_FREQUENCY_PRECOMPUTED_BLURING_RADIUS
 		#define CONFIG_MAX_WITH_REF_DISTANCE 1
@@ -63,6 +66,9 @@
 		#define CONFIG_SIGNAL_INPUT_LAYOUT  SIGNAL_BUFFER_LAYOUT_PENUMBRA_RECONSTRUCTION
 		#define CONFIG_SIGNAL_OUTPUT_LAYOUT SIGNAL_BUFFER_LAYOUT_PENUMBRA_RECONSTRUCTION
 
+		#define CONFIG_INPUT_TEXTURE_COUNT CONFIG_SIGNAL_BATCH_SIZE
+		#define CONFIG_OUTPUT_TEXTURE_COUNT CONFIG_SIGNAL_BATCH_SIZE
+
 		#define CONFIG_SAMPLE_SET           SAMPLE_SET_HEXAWEB
 		#define CONFIG_BILATERAL_DISTANCE_COMPUTATION SIGNAL_WORLD_FREQUENCY_PRECOMPUTED_BLURING_RADIUS
 		#define CONFIG_MAX_WITH_REF_DISTANCE 1
@@ -73,6 +79,9 @@
 		#define CONFIG_SIGNAL_INPUT_LAYOUT  SIGNAL_BUFFER_LAYOUT_PENUMBRA_RECONSTRUCTION
 		#define CONFIG_SIGNAL_OUTPUT_LAYOUT SIGNAL_BUFFER_LAYOUT_PENUMBRA_REJECTION
 
+		#define CONFIG_INPUT_TEXTURE_COUNT CONFIG_SIGNAL_BATCH_SIZE
+		#define CONFIG_OUTPUT_TEXTURE_COUNT CONFIG_SIGNAL_BATCH_SIZE
+
 		#define CONFIG_BILATERAL_DISTANCE_COMPUTATION SIGNAL_WORLD_FREQUENCY_PRECOMPUTED_BLURING_RADIUS
 
 	#elif DIM_STAGE == STAGE_POST_FILTERING
@@ -80,6 +89,9 @@
 		#define CONFIG_SIGNAL_INPUT_LAYOUT  SIGNAL_BUFFER_LAYOUT_PENUMBRA_HISTORY
 		#define CONFIG_SIGNAL_OUTPUT_LAYOUT SIGNAL_BUFFER_LAYOUT_PENUMBRA_HISTORY
 
+		#define CONFIG_INPUT_TEXTURE_COUNT CONFIG_SIGNAL_BATCH_SIZE
+		#define CONFIG_OUTPUT_TEXTURE_COUNT CONFIG_SIGNAL_BATCH_SIZE
+
 		#define CONFIG_SAMPLE_SET           SAMPLE_SET_STACKOWIAK_4_SETS
 		//#define CONFIG_SAMPLE_SUBSET        1
 		#define CONFIG_BILATERAL_DISTANCE_COMPUTATION SIGNAL_WORLD_FREQUENCY_PRECOMPUTED_BLURING_RADIUS
@@ -87,6 +99,10 @@
 
 	#elif DIM_STAGE == STAGE_FINAL_OUTPUT
 		#define CONFIG_SIGNAL_INPUT_LAYOUT  SIGNAL_BUFFER_LAYOUT_PENUMBRA_HISTORY
+
+		#define CONFIG_INPUT_TEXTURE_COUNT CONFIG_SIGNAL_BATCH_SIZE
+		#define CONFIG_OUTPUT_TEXTURE_COUNT CONFIG_SIGNAL_BATCH_SIZE
+
 		#define CONFIG_SAMPLE_SET           SAMPLE_SET_1X1
 
 	#else
@@ -122,6 +138,9 @@
 		// TODO: This is a bit agressive.
 		#define CONFIG_ACCUMULATION_COLOR_SPACE (COLOR_SPACE_RGB | COLOR_SPACE_KARIS_WEIGHTING)
 
+		#define CONFIG_INPUT_TEXTURE_COUNT 2
+		#define CONFIG_OUTPUT_TEXTURE_COUNT 2
+
 	#elif DIM_STAGE == STAGE_REJECTION_PRE_CONVOLUTION
 		// Input and output layout.
 		#define CONFIG_SIGNAL_INPUT_LAYOUT  SIGNAL_BUFFER_LAYOUT_REFLECTIONS_HISTORY
@@ -134,6 +153,9 @@
 		#define CONFIG_ACCUMULATION_COLOR_SPACE REFLECTIONS_REJECTION_COLOR_SPACE
 		#define CONFIG_OUTPUT_COLOR_SPACE       REFLECTIONS_REJECTION_COLOR_SPACE
 
+		#define CONFIG_INPUT_TEXTURE_COUNT 2
+		#define CONFIG_OUTPUT_TEXTURE_COUNT 3
+
 	#elif DIM_STAGE == STAGE_POST_FILTERING
 		// Input and output layout.
 		#define CONFIG_SIGNAL_INPUT_LAYOUT  SIGNAL_BUFFER_LAYOUT_REFLECTIONS_HISTORY
@@ -149,6 +171,9 @@
 		#define CONFIG_TOKOYASHI_LOBE_BILATERAL 1
 		#define CONFIG_TOKOYASHI_AXES_BILATERAL 1
 
+		#define CONFIG_INPUT_TEXTURE_COUNT 2
+		#define CONFIG_OUTPUT_TEXTURE_COUNT 2
+
 	#else
 		#error Unexpected stage.
 	#endif
@@ -157,6 +182,9 @@
 	#define CONFIG_VIEW_POSITION_BASED_BILATERAL 2
 	#define CONFIG_NORMAL_BASED_BILATERAL 1
 
+	#define CONFIG_INPUT_TEXTURE_COUNT 1
+	#define CONFIG_OUTPUT_TEXTURE_COUNT 1
+
 	#if DIM_STAGE == STAGE_RECONSTRUCTION
 		// Input and output layout.
 		#define CONFIG_SIGNAL_INPUT_LAYOUT  SIGNAL_BUFFER_LAYOUT_AO_INPUT
@@ -184,6 +212,9 @@
 	#define CONFIG_VIEW_POSITION_BASED_BILATERAL 2
 	#define CONFIG_NORMAL_BASED_BILATERAL 1
 
+	#define CONFIG_INPUT_TEXTURE_COUNT 2
+	#define CONFIG_OUTPUT_TEXTURE_COUNT 2
+
 	#if DIM_STAGE == STAGE_RECONSTRUCTION
 		// Input and output layout.
 		#define CONFIG_SIGNAL_INPUT_LAYOUT  SIGNAL_BUFFER_LAYOUT_GI_AND_AO_INPUT
@@ -304,15 +335,55 @@ uint UpscaleFactor;
 	float KernelSpreadFactor;
 #endif
 
+
+#if !defined(CONFIG_INPUT_TEXTURE_COUNT)
+	#error Missing CONFIG_INPUT_TEXTURE_COUNT
+#endif
+
 Texture2D SignalInput_Textures_0;
+
+#if CONFIG_INPUT_TEXTURE_COUNT > 1
 Texture2D SignalInput_Textures_1;
+#else
+#define SignalInput_Textures_1 SignalInput_Textures_0
+#endif
+
+#if CONFIG_INPUT_TEXTURE_COUNT > 2
 Texture2D SignalInput_Textures_2;
+#else
+#define SignalInput_Textures_2 SignalInput_Textures_0
+#endif
+
+#if CONFIG_INPUT_TEXTURE_COUNT > 3
 Texture2D SignalInput_Textures_3;
+#else
+#define SignalInput_Textures_3 SignalInput_Textures_0
+#endif
+
+
+#if !defined(CONFIG_OUTPUT_TEXTURE_COUNT)
+	#error Missing CONFIG_OUTPUT_TEXTURE_COUNT
+#endif
 
 RWTexture2D<float4>	SignalOutput_UAVs_0;
+
+#if CONFIG_OUTPUT_TEXTURE_COUNT > 1
 RWTexture2D<float4>	SignalOutput_UAVs_1;
+#else
+#define SignalOutput_UAVs_1 SignalOutput_UAVs_0
+#endif
+
+#if CONFIG_OUTPUT_TEXTURE_COUNT > 2
 RWTexture2D<float4>	SignalOutput_UAVs_2;
+#else
+#define SignalOutput_UAVs_2 SignalOutput_UAVs_0
+#endif
+
+#if CONFIG_OUTPUT_TEXTURE_COUNT > 3
 RWTexture2D<float4>	SignalOutput_UAVs_3;
+#else
+#define SignalOutput_UAVs_3 SignalOutput_UAVs_0
+#endif
 
 RWTexture2D<float4>	DebugOutput;
 
diff --git a/Engine/Shaders/Private/ScreenSpaceDenoise/SSDTemporalAccumulation.usf b/Engine/Shaders/Private/ScreenSpaceDenoise/SSDTemporalAccumulation.usf
--- a/Engine/Shaders/Private/ScreenSpaceDenoise/SSDTemporalAccumulation.usf
+++ b/Engine/Shaders/Private/ScreenSpaceDenoise/SSDTemporalAccumulation.usf
@@ -36,6 +36,9 @@
 	#define CONFIG_VIEW_POSITION_BASED_BILATERAL 2
 	#define CONFIG_MULTIPLEXED_SIGNALS_PER_SIGNAL_DOMAIN 1
 
+	#define CONFIG_INPUT_TEXTURE_COUNT CONFIG_SIGNAL_BATCH_SIZE
+	#define CONFIG_HISTORY_TEXTURE_COUNT CONFIG_SIGNAL_BATCH_SIZE
+
 	#define CONFIG_SIGNAL_INPUT_LAYOUT   SIGNAL_BUFFER_LAYOUT_PENUMBRA_RECONSTRUCTION
 	#define CONFIG_SIGNAL_HISTORY_LAYOUT SIGNAL_BUFFER_LAYOUT_PENUMBRA_HISTORY
 
@@ -51,6 +54,8 @@
 		#define CONFIG_USE_REJECTION_BUFFER               1
 		#define CONFIG_SIGNAL_HISTORY_REJECTION_LAYOUT    SIGNAL_BUFFER_LAYOUT_PENUMBRA_REJECTION
 		
+		#define CONFIG_REJECTION_TEXTURE_COUNT CONFIG_SIGNAL_BATCH_SIZE
+
 		#define CONFIG_REJECTION_INPUT_MODE               REJECTION_INPUT_MODE_2PRETRANSFORMED_MOMMENTS
 		#define CONFIG_REJECTION_SAMPLE_SET               REJECTION_SAMPLE_SET_SOBEK2018
 	#endif
@@ -60,6 +65,9 @@
 	#define CONFIG_SIGNAL_INPUT_LAYOUT   SIGNAL_BUFFER_LAYOUT_REFLECTIONS_HISTORY
 	#define CONFIG_SIGNAL_HISTORY_LAYOUT SIGNAL_BUFFER_LAYOUT_REFLECTIONS_HISTORY
 
+	#define CONFIG_INPUT_TEXTURE_COUNT 2
+	#define CONFIG_HISTORY_TEXTURE_COUNT 2
+
 	// Can only be done using tokoyashi because have more than one sample at a time.
 	#define CONFIG_TOKOYASHI_LOBE_BILATERAL 1
 	#define CONFIG_TOKOYASHI_AXES_BILATERAL 1
@@ -74,6 +82,8 @@
 		#define CONFIG_USE_REJECTION_BUFFER               1
 		#define CONFIG_SIGNAL_HISTORY_REJECTION_LAYOUT    SIGNAL_BUFFER_LAYOUT_REFLECTIONS_REJECTION
 		
+		#define CONFIG_REJECTION_TEXTURE_COUNT 3 
+
 		#define CONFIG_REJECTION_BUFFER_COLOR_SPACE       REFLECTIONS_REJECTION_COLOR_SPACE
 		#define CONFIG_REJECTION_INPUT_MODE               REJECTION_INPUT_MODE_2PRETRANSFORMED_MOMMENTS
 		#define CONFIG_REJECTION_SAMPLE_SET               REJECTION_SAMPLE_SET_SOBEK2018
@@ -84,6 +94,9 @@
 	#define CONFIG_SIGNAL_INPUT_LAYOUT   SIGNAL_BUFFER_LAYOUT_AO_HISTORY
 	#define CONFIG_SIGNAL_HISTORY_LAYOUT SIGNAL_BUFFER_LAYOUT_AO_HISTORY
 
+	#define CONFIG_INPUT_TEXTURE_COUNT 1
+	#define CONFIG_HISTORY_TEXTURE_COUNT 1
+
 	#define CONFIG_HISTORY_REJECTION HISTORY_REJECTION_VAR_BOUNDARIES
 
 	#define CONFIG_VIEW_POSITION_BASED_BILATERAL 2
@@ -94,6 +107,9 @@
 	#define CONFIG_SIGNAL_INPUT_LAYOUT   SIGNAL_BUFFER_LAYOUT_GI_AND_AO_RECONSTRUCTION
 	#define CONFIG_SIGNAL_HISTORY_LAYOUT SIGNAL_BUFFER_LAYOUT_GI_AND_AO_HISTORY
 
+	#define CONFIG_INPUT_TEXTURE_COUNT 2
+	#define CONFIG_HISTORY_TEXTURE_COUNT 2
+
 	#define CONFIG_HISTORY_REJECTION HISTORY_REJECTION_MINMAX_BOUNDARIES
 	//#define CONFIG_HISTORY_REJECTION HISTORY_REJECTION_ADAPTIVE_INTEGRATION
 	#define CONFIG_MULTIPLEXED_SIGNALS_PER_SIGNAL_DOMAIN 3
@@ -182,25 +198,90 @@ static const float MipLevel = 0;
 static const float MipLevelPow2 = 1;
 static const float InvMipLevelPow2 = 1;
 
+
+#if !defined(CONFIG_INPUT_TEXTURE_COUNT)
+	#error Missing CONFIG_INPUT_TEXTURE_COUNT
+#endif
+
 Texture2D SignalInput_Textures_0;
+
+#if CONFIG_INPUT_TEXTURE_COUNT > 1
 Texture2D SignalInput_Textures_1;
+#else
+#define SignalInput_Textures_1 SignalInput_Textures_0
+#endif
+
+#if CONFIG_INPUT_TEXTURE_COUNT > 2
 Texture2D SignalInput_Textures_2;
+#else
+#define SignalInput_Textures_2 SignalInput_Textures_0
+#endif
+
+#if CONFIG_INPUT_TEXTURE_COUNT > 3
 Texture2D SignalInput_Textures_3;
+#else
+#define SignalInput_Textures_3 SignalInput_Textures_0
+#endif
 
-Texture2D HistoryRejectionSignal_Textures_0;
-Texture2D HistoryRejectionSignal_Textures_1;
-Texture2D HistoryRejectionSignal_Textures_2;
-Texture2D HistoryRejectionSignal_Textures_3;
 
-Texture2D PrevHistory_Textures_0;
-Texture2D PrevHistory_Textures_1;
-Texture2D PrevHistory_Textures_2;
-Texture2D PrevHistory_Textures_3;
+#if CONFIG_USE_REJECTION_BUFFER
+	#if !defined(CONFIG_REJECTION_TEXTURE_COUNT)
+		#error Missing CONFIG_REJECTION_TEXTURE_COUNT
+	#endif
+
+	Texture2D HistoryRejectionSignal_Textures_0;
+
+	#if CONFIG_REJECTION_TEXTURE_COUNT > 1
+	Texture2D HistoryRejectionSignal_Textures_1;
+	#else
+	#define HistoryRejectionSignal_Textures_1 HistoryRejectionSignal_Textures_0
+	#endif
+
+	#if CONFIG_REJECTION_TEXTURE_COUNT > 2
+	Texture2D HistoryRejectionSignal_Textures_2;
+	#else
+	#define HistoryRejectionSignal_Textures_2 HistoryRejectionSignal_Textures_0
+	#endif
 
+	#if CONFIG_REJECTION_TEXTURE_COUNT > 3
+	Texture2D HistoryRejectionSignal_Textures_3;
+	#else
+	#define HistoryRejectionSignal_Textures_3 HistoryRejectionSignal_Textures_0
+	#endif
+#endif // CONFIG_USE_REJECTION_BUFFER
+
+
+#if !defined(CONFIG_HISTORY_TEXTURE_COUNT)
+	#error Missing CONFIG_HISTORY_TEXTURE_COUNT
+#endif
+
+Texture2D PrevHistory_Textures_0;
 RWTexture2D<float4>	SignalHistoryOutput_UAVs_0;
+
+#if CONFIG_HISTORY_TEXTURE_COUNT > 1
+Texture2D PrevHistory_Textures_1;
 RWTexture2D<float4>	SignalHistoryOutput_UAVs_1;
+#else
+#define PrevHistory_Textures_1 PrevHistory_Textures_0
+#define SignalHistoryOutput_UAVs_1 SignalHistoryOutput_UAVs_0
+#endif
+
+#if CONFIG_HISTORY_TEXTURE_COUNT > 2
+Texture2D PrevHistory_Textures_2;
 RWTexture2D<float4>	SignalHistoryOutput_UAVs_2;
+#else
+#define PrevHistory_Textures_2 PrevHistory_Textures_0
+#define SignalHistoryOutput_UAVs_2 SignalHistoryOutput_UAVs_0
+#endif
+
+#if CONFIG_HISTORY_TEXTURE_COUNT > 3
+Texture2D PrevHistory_Textures_3;
 RWTexture2D<float4>	SignalHistoryOutput_UAVs_3;
+#else
+#define PrevHistory_Textures_3 PrevHistory_Textures_0
+#define SignalHistoryOutput_UAVs_3 SignalHistoryOutput_UAVs_0
+#endif
+
 
 RWTexture2D<float4>	DebugOutput;
 
@@ -380,7 +461,6 @@ void TemporallyAccumulate(
 	#if (CONFIG_HISTORY_REJECTION == HISTORY_REJECTION_MINMAX_BOUNDARIES || CONFIG_HISTORY_REJECTION == HISTORY_REJECTION_VAR_BOUNDARIES)
 	{
 		// Whether the history is rejected using rejection history signal.
-		const bool bUsesHistoryRejectionSignal = CONFIG_USE_REJECTION_BUFFER;
 		const bool bUseCurrentSampleForCenter = !CONFIG_USE_REJECTION_BUFFER;
 
 		FSSDKernelConfig KernelConfig = CreateKernelConfig();
@@ -413,7 +493,7 @@ void TemporallyAccumulate(
 			}
 			#endif
 
-			if (bUsesHistoryRejectionSignal)
+			if (CONFIG_USE_REJECTION_BUFFER)
 			{
 				// History rejection have two moments of the signal being denoised.
 				KernelConfig.MultiplexedSignalsPerSignalDomain = 2;
@@ -492,14 +572,24 @@ void TemporallyAccumulate(
 			}
 		}
 
-		AccumulateKernel(
-			KernelConfig,
-			bUsesHistoryRejectionSignal ? HistoryRejectionSignal_Textures_0 : SignalInput_Textures_0,
-			bUsesHistoryRejectionSignal ? HistoryRejectionSignal_Textures_1 : SignalInput_Textures_1,
-			bUsesHistoryRejectionSignal ? HistoryRejectionSignal_Textures_2 : SignalInput_Textures_2,
-			bUsesHistoryRejectionSignal ? HistoryRejectionSignal_Textures_3 : SignalInput_Textures_3,
-			/* inout */ SignalAccumulators);
-		
+		#if CONFIG_USE_REJECTION_BUFFER
+			AccumulateKernel(
+				KernelConfig,
+				HistoryRejectionSignal_Textures_0,
+				HistoryRejectionSignal_Textures_1,
+				HistoryRejectionSignal_Textures_2,
+				HistoryRejectionSignal_Textures_3,
+				/* inout */ SignalAccumulators);
+		#else
+			AccumulateKernel(
+				KernelConfig,
+				SignalInput_Textures_0,
+				SignalInput_Textures_1,
+				SignalInput_Textures_2,
+				SignalInput_Textures_3,
+				/* inout */ SignalAccumulators);
+		#endif
+
 		// Clamp history
 		FIX_UNROLL(CONFIG_SIGNAL_BATCH_SIZE)
 		for (uint BatchedSignalId = 0; BatchedSignalId < CONFIG_SIGNAL_BATCH_SIZE; BatchedSignalId++)
diff --git a/Engine/Source/Runtime/Renderer/Private/ScreenSpaceDenoise.cpp b/Engine/Source/Runtime/Renderer/Private/ScreenSpaceDenoise.cpp
--- a/Engine/Source/Runtime/Renderer/Private/ScreenSpaceDenoise.cpp
+++ b/Engine/Source/Runtime/Renderer/Private/ScreenSpaceDenoise.cpp
@@ -223,8 +223,8 @@ const TCHAR* const kInjestResourceNames[] = {
 	// Penumbra
 	TEXT("ShadowDenoiserInjest0"),
 	TEXT("ShadowDenoiserInjest1"),
-	TEXT("ShadowDenoiserInjest2"),
-	TEXT("ShadowDenoiserInjest3"),
+	nullptr,
+	nullptr,
 
 	// Reflections
 	nullptr,
@@ -255,20 +255,20 @@ const TCHAR* const kReconstructionResourceNames[] = {
 	// Reflections
 	TEXT("ReflectionsReconstruction0"),
 	TEXT("ReflectionsReconstruction1"),
-	TEXT("ReflectionsReconstruction2"),
-	TEXT("ReflectionsReconstruction3"),
+	nullptr,
+	nullptr,
 
 	// AmbientOcclusion
 	TEXT("AOReconstruction0"),
-	TEXT("AOReconstruction1"),
-	TEXT("AOReconstruction2"),
-	TEXT("AOReconstruction3"),
+	nullptr,
+	nullptr,
+	nullptr,
 
 	// GlobalIllumination
 	TEXT("GIReconstruction0"),
 	TEXT("GIReconstruction1"),
-	TEXT("GIReconstruction2"),
-	TEXT("GIReconstruction3"),
+	nullptr,
+	nullptr,
 };
 
 const TCHAR* const kPreConvolutionResourceNames[] = {
@@ -308,7 +308,7 @@ const TCHAR* const kRejectionPreConvolutionResourceNames[] = {
 	TEXT("ReflectionsRejectionPreConvolution0"),
 	TEXT("ReflectionsRejectionPreConvolution1"),
 	TEXT("ReflectionsRejectionPreConvolution2"),
-	TEXT("ReflectionsRejectionPreConvolution3"),
+	nullptr,
 
 	// AmbientOcclusion
 	nullptr,
@@ -333,20 +333,20 @@ const TCHAR* const kTemporalAccumulationResourceNames[] = {
 	// Reflections
 	TEXT("ReflectionsTemporalAccumulation0"),
 	TEXT("ReflectionsTemporalAccumulation1"),
-	TEXT("ReflectionsTemporalAccumulation2"),
-	TEXT("ReflectionsTemporalAccumulation3"),
+	nullptr,
+	nullptr,
 
 	// AmbientOcclusion
 	TEXT("AOTemporalAccumulation0"),
-	TEXT("AOTemporalAccumulation1"),
-	TEXT("AOTemporalAccumulation2"),
-	TEXT("AOTemporalAccumulation3"),
+	nullptr,
+	nullptr,
+	nullptr,
 
 	// GlobalIllumination
 	TEXT("GITemporalAccumulation0"),
 	TEXT("GITemporalAccumulation1"),
-	TEXT("GITemporalAccumulation2"),
-	TEXT("GITemporalAccumulation3"),
+	nullptr,
+	nullptr,
 };
 
 const TCHAR* const kHistoryConvolutionResourceNames[] = {
@@ -359,20 +359,20 @@ const TCHAR* const kHistoryConvolutionResourceNames[] = {
 	// Reflections
 	TEXT("ReflectionsHistoryConvolution0"),
 	TEXT("ReflectionsHistoryConvolution1"),
-	TEXT("ReflectionsHistoryConvolution2"),
-	TEXT("ReflectionsHistoryConvolution3"),
+	nullptr,
+	nullptr,
 
 	// AmbientOcclusion
 	TEXT("AOHistoryConvolution0"),
-	TEXT("AOHistoryConvolution1"),
-	TEXT("AOHistoryConvolution2"),
-	TEXT("AOHistoryConvolution3"),
+	nullptr,
+	nullptr,
+	nullptr,
 
 	// GlobalIllumination
 	TEXT("GIHistoryConvolution0"),
 	TEXT("GIHistoryConvolution1"),
-	TEXT("GIHistoryConvolution2"),
-	TEXT("GIHistoryConvolution3"),
+	nullptr,
+	nullptr,
 };
 
 const TCHAR* const kDenoiserOutputResourceNames[] = {
@@ -715,6 +715,7 @@ static FSSDSignalTextures CopyAndBackfillSignalInput(const FSSDSignalTextures& S
 {
 	FSSDSignalTextures Result = SignalInput;
 
+	#if 0
 	for (int32 i = 0; i < kMaxBufferProcessingCount; ++i)
 	{
 		if (!Result.Textures[i])
@@ -722,6 +723,7 @@ static FSSDSignalTextures CopyAndBackfillSignalInput(const FSSDSignalTextures& S
 			Result.Textures[i] = Result.Textures[i - 1];
 		}
 	}
+	#endif
 
 	return Result;
 }
@@ -852,6 +854,12 @@ static void DenoiseSignalAtConstantPixelDensity(
 		check(ReconstructionTextureCount > 0);
 	}
 
+	// Create a UAV use to output debugging information from the shader.
+	auto CreateDebugUAV = [&](const TCHAR* DebugTextureName)
+	{
+		return GraphBuilder.CreateUAV(GraphBuilder.CreateTexture(DebugDesc, DebugTextureName));
+	};
+
 	int32 HistoryTextureCount = HistoryTextureCountPerSignal * Settings.SignalBatchSize;
 
 	check(HistoryTextureCount <= kMaxBufferProcessingCount);
@@ -933,7 +941,7 @@ static void DenoiseSignalAtConstantPixelDensity(
 		PassParameters->SignalInput = CopyAndBackfillSignalInput(SignalHistory);
 		PassParameters->SignalOutput = CreateMultiplexedUAVs(GraphBuilder, NewSignalOutput);
 		
-		PassParameters->DebugOutput = GraphBuilder.CreateUAV(GraphBuilder.CreateTexture(DebugDesc, TEXT("SSDDebugReflectionReconstruction")));
+		PassParameters->DebugOutput = CreateDebugUAV(TEXT("DebugDenoiserReconstruction"));
 
 		FSSDSpatialAccumulationCS::FPermutationDomain PermutationVector;
 		PermutationVector.Set<FSignalProcessingDim>(Settings.SignalProcessing);
@@ -971,8 +979,8 @@ static void DenoiseSignalAtConstantPixelDensity(
 		PassParameters->ConvolutionMetaData = ConvolutionMetaData;
 		PassParameters->SignalInput = CopyAndBackfillSignalInput(SignalHistory);
 		PassParameters->SignalOutput = CreateMultiplexedUAVs(GraphBuilder, NewSignalOutput);
-
-		PassParameters->DebugOutput = GraphBuilder.CreateUAV(GraphBuilder.CreateTexture(DebugDesc, TEXT("DebugDenoiserPreConvolution")));
+		
+		PassParameters->DebugOutput = CreateDebugUAV(TEXT("DebugDenoiserPreConvolution"));
 
 		FSSDSpatialAccumulationCS::FPermutationDomain PermutationVector;
 		PermutationVector.Set<FSignalProcessingDim>(Settings.SignalProcessing);
@@ -1047,8 +1055,8 @@ static void DenoiseSignalAtConstantPixelDensity(
 			PermutationVector.Set<FSignalBatchSizeDim>(Settings.SignalBatchSize);
 			PermutationVector.Set<FSSDSpatialAccumulationCS::FStageDim>(FSSDSpatialAccumulationCS::EStage::RejectionPreConvolution);
 			PermutationVector.Set<FMultiSPPDim>(true);
-
-			PassParameters->DebugOutput = GraphBuilder.CreateUAV(GraphBuilder.CreateTexture(DebugDesc, TEXT("DebugRejectionPreConvolution")));
+			
+			PassParameters->DebugOutput = CreateDebugUAV(TEXT("DebugDenoiserRejectionPreConvolution"));
 
 			TShaderMapRef<FSSDSpatialAccumulationCS> ComputeShader(View.ShaderMap, PermutationVector);
 			FComputeShaderUtils::AddPass(
@@ -1111,8 +1119,8 @@ static void DenoiseSignalAtConstantPixelDensity(
 		} // for (uint32 BatchedSignalId = 0; BatchedSignalId < Settings.SignalBatchSize; BatchedSignalId++)
 
 		PassParameters->PrevHistory = CopyAndBackfillSignalInput(PassParameters->PrevHistory);
-
-		PassParameters->DebugOutput = GraphBuilder.CreateUAV(GraphBuilder.CreateTexture(DebugDesc, TEXT("SSDDebugReflectionTemporalAccumulation")));
+		
+		PassParameters->DebugOutput = CreateDebugUAV(TEXT("DebugDenoiserTemporalAccumulation"));
 
 		FComputeShaderUtils::AddPass(
 			GraphBuilder,
@@ -1147,7 +1155,7 @@ static void DenoiseSignalAtConstantPixelDensity(
 		PermutationVector.Set<FSSDSpatialAccumulationCS::FStageDim>(FSSDSpatialAccumulationCS::EStage::PostFiltering);
 		PermutationVector.Set<FMultiSPPDim>(true);
 		
-		PassParameters->DebugOutput = GraphBuilder.CreateUAV(GraphBuilder.CreateTexture(DebugDesc, TEXT("SSDDebugReflectionPostfilter")));
+		PassParameters->DebugOutput = CreateDebugUAV(TEXT("DebugDenoiserPostfilter"));
 
 		TShaderMapRef<FSSDSpatialAccumulationCS> ComputeShader(View.ShaderMap, PermutationVector);
 		FComputeShaderUtils::AddPass(
