Renames some of the structures and members of the denoiser

diff --git a/Engine/Shaders/Private/ScreenSpaceDenoise/SSDCommon.ush b/Engine/Shaders/Private/ScreenSpaceDenoise/SSDCommon.ush
--- a/Engine/Shaders/Private/ScreenSpaceDenoise/SSDCommon.ush
+++ b/Engine/Shaders/Private/ScreenSpaceDenoise/SSDCommon.ush
@@ -37,7 +37,7 @@ float SafeRcp(float x)
 //------------------------------------------------------- COMMON SCENE INFOS
 
 /** olds commonly used information of the scene of a given sample. */
-struct FSSFSampleSceneInfos
+struct FSSDSampleSceneInfos
 {
 	/** Normal of the pixel in world space. */
 	float3 WorldNormal;
@@ -52,9 +52,9 @@ struct FSSFSampleSceneInfos
 	float Roughness;
 };
 
-FSSFSampleSceneInfos CreateSampleSceneInfos()
+FSSDSampleSceneInfos CreateSampleSceneInfos()
 {
-	FSSFSampleSceneInfos Infos;
+	FSSDSampleSceneInfos Infos;
 	Infos.WorldNormal             = 0;
 	Infos.TranslatedWorldPosition = 0;
 	Infos.WorldDepth              = 0;
@@ -62,11 +62,11 @@ FSSFSampleSceneInfos CreateSampleSceneInfos()
 	return Infos;
 }
 
-FSSFSampleSceneInfos InitSampleSceneInfos(
+FSSDSampleSceneInfos InitSampleSceneInfos(
 	float4x4 ScreenToTranslatedWorld,
 	float2 ScreenPosition, float WorldDepth, float3 WorldNormal, float Roughness)
 {
-	FSSFSampleSceneInfos Infos;
+	FSSDSampleSceneInfos Infos;
 
 	float2 ClipPosition = ScreenPosition * (View.ViewToClip[3][3] < 1.0f ? WorldDepth : 1.0f);
 	Infos.TranslatedWorldPosition = mul(float4(ClipPosition, WorldDepth, 1), ScreenToTranslatedWorld).xyz;
@@ -76,7 +76,7 @@ FSSFSampleSceneInfos InitSampleSceneInfos(
 	return Infos;
 }
 
-FSSFSampleSceneInfos FetchCurrentSceneInfos(float2 ScreenPosition, float2 BufferUV)
+FSSDSampleSceneInfos FetchCurrentSceneInfos(float2 ScreenPosition, float2 BufferUV)
 {
 	FGBufferData GBufferData = GetGBufferDataFromSceneBlackboard(BufferUV);
 	return InitSampleSceneInfos(
@@ -84,14 +84,14 @@ FSSFSampleSceneInfos FetchCurrentSceneInfos(float2 ScreenPosition, float2 Buffer
 		ScreenPosition, GBufferData.Depth, GBufferData.WorldNormal, GBufferData.Roughness);
 }
 
-FSSFSampleSceneInfos FetchPrevSceneInfos(float2 ScreenPosition, float2 BufferUV)
+FSSDSampleSceneInfos FetchPrevSceneInfos(float2 ScreenPosition, float2 BufferUV)
 {
 	// TODO: so what to do with mip level of these guys?
 	float WorldSceneDepth = ConvertFromDeviceZ(PrevDepthBuffer.SampleLevel(GlobalPointClampedSampler, BufferUV, 0).r);
 	float3 WorldNormal = DecodeNormal(PrevGBufferA.SampleLevel(GlobalPointClampedSampler, BufferUV, 0).rgb);
 	float Roughness = PrevGBufferB.SampleLevel(GlobalPointClampedSampler, BufferUV, 0).b;
 
-	FSSFSampleSceneInfos OutMetadata;
+	FSSDSampleSceneInfos OutMetadata;
 	OutMetadata = InitSampleSceneInfos(
 		View.PrevScreenToTranslatedWorld,
 		ScreenPosition, WorldSceneDepth, WorldNormal, Roughness);
@@ -101,19 +101,19 @@ FSSFSampleSceneInfos FetchPrevSceneInfos(float2 ScreenPosition, float2 BufferUV)
 	return OutMetadata;
 }
 
-float ProjectWorldDistanceToScreenDistance(FSSFSampleSceneInfos SceneMetadata, float WorldDistance)
+float ProjectWorldDistanceToScreenDistance(FSSDSampleSceneInfos SceneMetadata, float WorldDistance)
 {
 	float4 ClipDistance = mul(View.ViewToClipNoAA, float4(WorldDistance, 0, SceneMetadata.WorldDepth, 1));
 	return ClipDistance.x / ClipDistance.w;
 }
 
-float ProjectWorldDistanceToPixelDistance(FSSFSampleSceneInfos SceneMetadata, float WorldDistance)
+float ProjectWorldDistanceToPixelDistance(FSSDSampleSceneInfos SceneMetadata, float WorldDistance)
 {
 	return WorldDistance * View.ViewSizeAndInvSize.x * 0.5 / (SceneMetadata.WorldDepth * GetTanHalfFieldOfView().x);
 }
 
 /** Compute the anisotropy value. The lower the value is, the more anysotropic the sample is. */
-float ComputeAnisotropyInvFactor(FSSFSampleSceneInfos SceneMetadata)
+float ComputeAnisotropyInvFactor(FSSDSampleSceneInfos SceneMetadata)
 {
 	if (!CONFIG_NORMAL_ANISOTROPY)
 	{
@@ -126,7 +126,7 @@ float ComputeAnisotropyInvFactor(FSSFSampleSceneInfos SceneMetadata)
 }
 
 /** Return whether this signal should store in given mip level. */
-float ComputeSignalMipLevel(FSSFSampleSceneInfos SceneMetadata, float WorldBluringRadius)
+float ComputeSignalMipLevel(FSSDSampleSceneInfos SceneMetadata, float WorldBluringRadius)
 {
 	if (WorldBluringRadius == WORLD_RADIUS_INVALID)
 	{
@@ -154,7 +154,7 @@ float ComputeSignalMipLevel(FSSFSampleSceneInfos SceneMetadata, float WorldBluri
 
 /** Returns the radius of a pixel in world space. */
 // TODO: move that to CPU
-float ComputeWorldBluringRadiusCausedByPixelSize(FSSFSampleSceneInfos SceneMetadata)
+float ComputeWorldBluringRadiusCausedByPixelSize(FSSDSampleSceneInfos SceneMetadata)
 {
 	// Should be multiplied 0.5* for the diameter to radius, and by 2.0 because GetTanHalfFieldOfView() cover only half of the pixels.
 	return (View.ViewSizeAndInvSize.z * GetTanHalfFieldOfView().x) * SceneMetadata.WorldDepth;
diff --git a/Engine/Shaders/Private/ScreenSpaceDenoise/SSDInjest.usf b/Engine/Shaders/Private/ScreenSpaceDenoise/SSDInjest.usf
--- a/Engine/Shaders/Private/ScreenSpaceDenoise/SSDInjest.usf
+++ b/Engine/Shaders/Private/ScreenSpaceDenoise/SSDInjest.usf
@@ -105,7 +105,7 @@ void MainCS(
 	
 	// Read reference meta data.
 	float2 ScreenPosition = ViewportUVToScreenPos(BufferUVToViewportUV(SceneBufferUV));
-	FSSFSampleSceneInfos RefSceneMetadata = FetchCurrentSceneInfos(ScreenPosition, SceneBufferUV);
+	FSSDSampleSceneInfos RefSceneMetadata = FetchCurrentSceneInfos(ScreenPosition, SceneBufferUV);
 
 	// Sample the reference sample.
 	FSSDSignalArray MultiplexedSamples = SampleMultiplexedSignals(
diff --git a/Engine/Shaders/Private/ScreenSpaceDenoise/SSDSignalAccumulator.ush b/Engine/Shaders/Private/ScreenSpaceDenoise/SSDSignalAccumulator.ush
--- a/Engine/Shaders/Private/ScreenSpaceDenoise/SSDSignalAccumulator.ush
+++ b/Engine/Shaders/Private/ScreenSpaceDenoise/SSDSignalAccumulator.ush
@@ -13,22 +13,22 @@
 struct FSSDSignalAccumulator
 {
 	// Sums of moment.
-	FSSFSignalSample Moment1;
-	FSSFSignalSample Moment2;
+	FSSDSignalSample Moment1;
+	FSSDSignalSample Moment2;
 
 	// Per component min and max values.
-	FSSFSignalSample Min;
-	FSSFSignalSample Max;
+	FSSDSignalSample Min;
+	FSSDSignalSample Max;
 
 	// Minimal inverse frequency found and intersected when sampling.
 	float MinInvFrequency;
 	
 	// Descending ring bucketing.
-	FSSFSignalSample Current;
+	FSSDSignalSample Current;
 	float CurrentInvFrequency;
 	float CurrentTranslucency;
 
-	FSSFSignalSample Previous;
+	FSSDSignalSample Previous;
 	float PreviousInvFrequency;
 
 	float BorderingRadius;
@@ -77,7 +77,7 @@ FSSDSignalAccumulator CreateSignalAccumulator()
 struct FSSDSampleAccumulationInfos
 {
 	// Sample to accumulate.
-	FSSFSignalSample Sample;
+	FSSDSignalSample Sample;
 
 	// Weight of the sample in the kernel.
 	float FinalWeight;
@@ -136,7 +136,7 @@ void AccumulateSampleInDomainBoundaries(inout FSSDSignalAccumulator Accumulator,
 // [ SIGGRAPH 2018, "A life of bokeh" ]
 
 void StartAccumulatingClusterInDRB(
-	FSSFSampleSceneInfos RefSceneMetadata,
+	FSSDSampleSceneInfos RefSceneMetadata,
 	inout FSSDSignalAccumulator Accumulator,
 	FSSDSampleClusterInfo ClusterInfo)
 {
@@ -177,16 +177,16 @@ void AccumulateSampleInDRB(inout FSSDSignalAccumulator Accumulator, FSSDSampleAc
 	{
 		float CurrentWeight = A.FinalWeight * bBelongsToCurrent;
 		Accumulator.Current = Add(Accumulator.Current, Mul(A.Sample, CurrentWeight));
-		Accumulator.CurrentInvFrequency += A.Sample.AccumulatedSampleCount * ClampedInvFrequency * CurrentWeight;
+		Accumulator.CurrentInvFrequency += A.Sample.SampleCount * ClampedInvFrequency * CurrentWeight;
 	}
 
 	// Accumulate current bucket translucency.
 	{
 		float SampleTranslucency = saturate(ClampedInvFrequency - Accumulator.BorderingRadius);
 		
-		// TODO: do need (A.Sample.LeavingRayCount * SafeRcp(A.Sample.AccumulatedSampleCount)) computation onen 1spp reconstruction.
+		// TODO: do need (A.Sample.MissCount * SafeRcp(A.Sample.SampleCount)) computation onen 1spp reconstruction.
 		// TODO: could pass down a normalised version of A.Sample to avoid Rcp.
-		float R = A.Sample.LeavingRayCount * SafeRcp(A.Sample.AccumulatedSampleCount);
+		float R = A.Sample.MissCount * SafeRcp(A.Sample.SampleCount);
 		
 		float Translucency = lerp(R, 1, SampleTranslucency);
 		float TranslucencyWeight = 1;
@@ -198,7 +198,7 @@ void AccumulateSampleInDRB(inout FSSDSignalAccumulator Accumulator, FSSDSampleAc
 	{
 		float PreviousWeight = A.FinalWeight * bBelongsToPrevious;
 		Accumulator.Previous = Add(Accumulator.Previous, Mul(A.Sample, PreviousWeight));
-		Accumulator.PreviousInvFrequency += A.Sample.AccumulatedSampleCount * ClampedInvFrequency * PreviousWeight;
+		Accumulator.PreviousInvFrequency += A.Sample.SampleCount * ClampedInvFrequency * PreviousWeight;
 	}
 }
 
@@ -209,7 +209,7 @@ void DijestSampleClusterInDRB(
 	//if (RingId != 0)
 	//	return;
 
-	if (Accumulator.Current.AccumulatedSampleCount == 0.0)
+	if (Accumulator.Current.SampleCount == 0.0)
 	{
 		return;
 	}
@@ -218,19 +218,19 @@ void DijestSampleClusterInDRB(
 	float CurrentClusterOpacity = saturate(1 - Accumulator.CurrentTranslucency * rcp(float(SampleCount)));
 	
 	// Compute current and previous Coc radii.
-	float PreviousInvFrequency = Accumulator.PreviousInvFrequency * SafeRcp(Accumulator.Previous.AccumulatedSampleCount);
-	float CurrentInvFrequency = Accumulator.CurrentInvFrequency * rcp(Accumulator.Current.AccumulatedSampleCount);
+	float PreviousInvFrequency = Accumulator.PreviousInvFrequency * SafeRcp(Accumulator.Previous.SampleCount);
+	float CurrentInvFrequency = Accumulator.CurrentInvFrequency * rcp(Accumulator.Current.SampleCount);
 
 	// Whether current bucket is occluding previous bucket.
 	// TODO: put a contrast.
 	float bOccludingCluster = saturate((PreviousInvFrequency - CurrentInvFrequency) * rcp(Accumulator.BorderingRadius));
 
 	// Compute final factor to use to with previous bucket.
-	float PreviousBucketFactor = (Accumulator.Previous.AccumulatedSampleCount == 0.0) ? 0.0 : (1.0 - CurrentClusterOpacity * bOccludingCluster);
+	float PreviousBucketFactor = (Accumulator.Previous.SampleCount == 0.0) ? 0.0 : (1.0 - CurrentClusterOpacity * bOccludingCluster);
 
 	// TODO: the big difference in sample count could lead to wrong opacity between sample clusters.
 	//if (PreviousBucketFactor < 1)
-	//	PreviousBucketFactor *= Accumulator.Current.AccumulatedSampleCount * SafeRcp(Accumulator.Previous.AccumulatedSampleCount);
+	//	PreviousBucketFactor *= Accumulator.Current.SampleCount * SafeRcp(Accumulator.Previous.SampleCount);
 
 	Accumulator.Previous = Add(Mul(Accumulator.Previous, PreviousBucketFactor), Accumulator.Current);
 	Accumulator.PreviousInvFrequency = Accumulator.PreviousInvFrequency * PreviousBucketFactor + Accumulator.CurrentInvFrequency;
@@ -241,7 +241,7 @@ void DijestSampleClusterInDRB(
 
 /** Start accumulating a cluster of samples. */
 void StartAccumulatingCluster(
-	FSSFSampleSceneInfos RefSceneMetadata,
+	FSSDSampleSceneInfos RefSceneMetadata,
 	inout FSSDSignalAccumulator Accumulator,
 	FSSDSampleClusterInfo ClusterInfo)
 {
diff --git a/Engine/Shaders/Private/ScreenSpaceDenoise/SSDSignalCore.ush b/Engine/Shaders/Private/ScreenSpaceDenoise/SSDSignalCore.ush
--- a/Engine/Shaders/Private/ScreenSpaceDenoise/SSDSignalCore.ush
+++ b/Engine/Shaders/Private/ScreenSpaceDenoise/SSDSignalCore.ush
@@ -8,47 +8,47 @@
 //------------------------------------------------------- STRUCTURE
 
 /** Generic data structure to manipulate any kind of signal. */
-struct FSSFSignalSample
+struct FSSDSignalSample
 {
 	// Number of sample accumulated.
-	float AccumulatedSampleCount;
+	float SampleCount;
 
 	// Scene color and alpha.
 	float4 SceneColor;
 
 	// Number of ray that did not find any intersections.
-	float LeavingRayCount;
+	float MissCount;
 
 	// Number of ray that are uncertain due to occlusion with higher frequency.
 	float UncertainRayCount;
 
-	// Hit distance of a this sample, valid only if AccumulatedSampleCount == 1.
-	float PenumbraClosestHit;
+	// Hit distance of a this sample, valid only if SampleCount == 1.
+	float ClosestHitDistance;
 	
-	// Blur radius allowed for this sample, valid only if AccumulatedSampleCount == 1.
+	// Blur radius allowed for this sample, valid only if SampleCount == 1.
 	float WorldBluringRadius;
 };
 
 /** Array of signal samples. Technically should be a typedef, but there is a bug in HLSL with the out operator of an array. */
 struct FSSDSignalArray
 {
-	FSSFSignalSample Array[MAX_SIGNAL_MULTIPLEXING];
+	FSSDSignalSample Array[MAX_SIGNAL_MULTIPLEXING];
 };
 
 
-FSSFSignalSample CreateSignalSampleFromScalarValue(float Scalar)
+FSSDSignalSample CreateSignalSampleFromScalarValue(float Scalar)
 {
-	FSSFSignalSample OutSample;
+	FSSDSignalSample OutSample;
 	OutSample.SceneColor = Scalar;
-	OutSample.AccumulatedSampleCount = Scalar;
-	OutSample.LeavingRayCount = Scalar;
-	OutSample.PenumbraClosestHit = Scalar;
+	OutSample.SampleCount = Scalar;
+	OutSample.MissCount = Scalar;
+	OutSample.ClosestHitDistance = Scalar;
 	OutSample.UncertainRayCount = Scalar;
 	OutSample.WorldBluringRadius = Scalar;
 	return OutSample;
 }
 
-FSSFSignalSample CreateSignalSample()
+FSSDSignalSample CreateSignalSample()
 {
 	return CreateSignalSampleFromScalarValue(0.0);
 }
@@ -67,127 +67,127 @@ FSSDSignalArray CreateSignalArrayFromScalarValue(float Scalar)
 
 //------------------------------------------------------- SIGNAL OPERATORS
 
-FSSFSignalSample Mul(FSSFSignalSample Sample, float Scalar)
+FSSDSignalSample Mul(FSSDSignalSample Sample, float Scalar)
 {
-	FSSFSignalSample OutSample;
+	FSSDSignalSample OutSample;
 	OutSample.SceneColor             = Sample.SceneColor * Scalar;
-	OutSample.AccumulatedSampleCount = Sample.AccumulatedSampleCount * Scalar;
-	OutSample.LeavingRayCount        = Sample.LeavingRayCount * Scalar;
-	OutSample.PenumbraClosestHit     = Sample.PenumbraClosestHit * Scalar;
+	OutSample.SampleCount = Sample.SampleCount * Scalar;
+	OutSample.MissCount        = Sample.MissCount * Scalar;
+	OutSample.ClosestHitDistance     = Sample.ClosestHitDistance * Scalar;
 	OutSample.UncertainRayCount      = Sample.UncertainRayCount * Scalar;
 	OutSample.WorldBluringRadius     = Sample.WorldBluringRadius * Scalar;
 	return OutSample;
 }
 
-FSSFSignalSample Add(FSSFSignalSample SampleA, FSSFSignalSample SampleB)
+FSSDSignalSample Add(FSSDSignalSample SampleA, FSSDSignalSample SampleB)
 {
 	SampleA.SceneColor             += SampleB.SceneColor;
-	SampleA.AccumulatedSampleCount += SampleB.AccumulatedSampleCount;
-	SampleA.LeavingRayCount        += SampleB.LeavingRayCount;
-	SampleA.PenumbraClosestHit     += SampleB.PenumbraClosestHit;
+	SampleA.SampleCount += SampleB.SampleCount;
+	SampleA.MissCount        += SampleB.MissCount;
+	SampleA.ClosestHitDistance     += SampleB.ClosestHitDistance;
 	SampleA.UncertainRayCount      += SampleB.UncertainRayCount;
 	SampleA.WorldBluringRadius     += SampleB.WorldBluringRadius;
 	return SampleA;
 }
 
-FSSFSignalSample Minus(FSSFSignalSample Sample)
+FSSDSignalSample Minus(FSSDSignalSample Sample)
 {
 	Sample.SceneColor             = -Sample.SceneColor;
-	Sample.AccumulatedSampleCount = -Sample.AccumulatedSampleCount;
-	Sample.LeavingRayCount        = -Sample.LeavingRayCount;
-	Sample.PenumbraClosestHit     = -Sample.PenumbraClosestHit;
+	Sample.SampleCount = -Sample.SampleCount;
+	Sample.MissCount        = -Sample.MissCount;
+	Sample.ClosestHitDistance     = -Sample.ClosestHitDistance;
 	Sample.UncertainRayCount      = -Sample.UncertainRayCount;
 	Sample.WorldBluringRadius     = -Sample.WorldBluringRadius;
 	return Sample;
 }
 
-FSSFSignalSample Subtract(FSSFSignalSample SampleA, FSSFSignalSample SampleB)
+FSSDSignalSample Subtract(FSSDSignalSample SampleA, FSSDSignalSample SampleB)
 {
 	return Add(SampleA, Minus(SampleB));
 }
 
-FSSFSignalSample Abs(FSSFSignalSample Sample)
+FSSDSignalSample Abs(FSSDSignalSample Sample)
 {
 	Sample.SceneColor             = abs(Sample.SceneColor);
-	Sample.AccumulatedSampleCount = abs(Sample.AccumulatedSampleCount);
-	Sample.LeavingRayCount        = abs(Sample.LeavingRayCount);
-	Sample.PenumbraClosestHit     = abs(Sample.PenumbraClosestHit);
+	Sample.SampleCount = abs(Sample.SampleCount);
+	Sample.MissCount        = abs(Sample.MissCount);
+	Sample.ClosestHitDistance     = abs(Sample.ClosestHitDistance);
 	Sample.UncertainRayCount      = abs(Sample.UncertainRayCount);
 	Sample.WorldBluringRadius     = abs(Sample.WorldBluringRadius);
 	return Sample;
 }
 
-FSSFSignalSample Sqrt(FSSFSignalSample Sample)
+FSSDSignalSample Sqrt(FSSDSignalSample Sample)
 {
 	Sample.SceneColor             = sqrt(Sample.SceneColor);
-	Sample.AccumulatedSampleCount = sqrt(Sample.AccumulatedSampleCount);
-	Sample.LeavingRayCount        = sqrt(Sample.LeavingRayCount);
-	Sample.PenumbraClosestHit     = sqrt(Sample.PenumbraClosestHit);
+	Sample.SampleCount = sqrt(Sample.SampleCount);
+	Sample.MissCount        = sqrt(Sample.MissCount);
+	Sample.ClosestHitDistance     = sqrt(Sample.ClosestHitDistance);
 	Sample.UncertainRayCount      = sqrt(Sample.UncertainRayCount);
 	Sample.WorldBluringRadius     = sqrt(Sample.WorldBluringRadius);
 	return Sample;
 }
 
-FSSFSignalSample Power(FSSFSignalSample Sample, float Exponent)
+FSSDSignalSample Power(FSSDSignalSample Sample, float Exponent)
 {
 	Sample.SceneColor.r           = pow(Sample.SceneColor.r, Exponent);
 	Sample.SceneColor.g           = pow(Sample.SceneColor.g, Exponent);
 	Sample.SceneColor.b           = pow(Sample.SceneColor.b, Exponent);
 	Sample.SceneColor.a           = pow(Sample.SceneColor.a, Exponent);
-	Sample.AccumulatedSampleCount = pow(Sample.AccumulatedSampleCount, Exponent);
-	Sample.LeavingRayCount        = pow(Sample.LeavingRayCount, Exponent);
-	Sample.PenumbraClosestHit     = pow(Sample.PenumbraClosestHit, Exponent);
+	Sample.SampleCount = pow(Sample.SampleCount, Exponent);
+	Sample.MissCount        = pow(Sample.MissCount, Exponent);
+	Sample.ClosestHitDistance     = pow(Sample.ClosestHitDistance, Exponent);
 	Sample.UncertainRayCount      = pow(Sample.UncertainRayCount, Exponent);
 	Sample.WorldBluringRadius     = pow(Sample.WorldBluringRadius, Exponent);
 	return Sample;
 }
 
-FSSFSignalSample Min(FSSFSignalSample SampleA, FSSFSignalSample SampleB)
+FSSDSignalSample Min(FSSDSignalSample SampleA, FSSDSignalSample SampleB)
 {
-	FSSFSignalSample OutSample;
+	FSSDSignalSample OutSample;
 	OutSample.SceneColor = min(SampleA.SceneColor, SampleB.SceneColor);
-	OutSample.AccumulatedSampleCount = min(SampleA.AccumulatedSampleCount, SampleB.AccumulatedSampleCount);
-	OutSample.LeavingRayCount = min(SampleA.LeavingRayCount, SampleB.LeavingRayCount);
-	OutSample.PenumbraClosestHit = min(SampleA.PenumbraClosestHit, SampleB.PenumbraClosestHit);
+	OutSample.SampleCount = min(SampleA.SampleCount, SampleB.SampleCount);
+	OutSample.MissCount = min(SampleA.MissCount, SampleB.MissCount);
+	OutSample.ClosestHitDistance = min(SampleA.ClosestHitDistance, SampleB.ClosestHitDistance);
 	OutSample.UncertainRayCount = min(SampleA.UncertainRayCount, SampleB.UncertainRayCount);
 	OutSample.WorldBluringRadius = min(SampleA.WorldBluringRadius, SampleB.WorldBluringRadius);
 	return OutSample;
 }
 
-FSSFSignalSample Max(FSSFSignalSample SampleA, FSSFSignalSample SampleB)
+FSSDSignalSample Max(FSSDSignalSample SampleA, FSSDSignalSample SampleB)
 {
-	FSSFSignalSample OutSample;
+	FSSDSignalSample OutSample;
 	OutSample.SceneColor = max(SampleA.SceneColor, SampleB.SceneColor);
-	OutSample.AccumulatedSampleCount = max(SampleA.AccumulatedSampleCount, SampleB.AccumulatedSampleCount);
-	OutSample.LeavingRayCount = max(SampleA.LeavingRayCount, SampleB.LeavingRayCount);
-	OutSample.PenumbraClosestHit = max(SampleA.PenumbraClosestHit, SampleB.PenumbraClosestHit);
+	OutSample.SampleCount = max(SampleA.SampleCount, SampleB.SampleCount);
+	OutSample.MissCount = max(SampleA.MissCount, SampleB.MissCount);
+	OutSample.ClosestHitDistance = max(SampleA.ClosestHitDistance, SampleB.ClosestHitDistance);
 	OutSample.UncertainRayCount = max(SampleA.UncertainRayCount, SampleB.UncertainRayCount);
 	OutSample.WorldBluringRadius = max(SampleA.WorldBluringRadius, SampleB.WorldBluringRadius);
 	return OutSample;
 }
 
-FSSFSignalSample Clamp(FSSFSignalSample Sample, FSSFSignalSample SampleMin, FSSFSignalSample SampleMax)
+FSSDSignalSample Clamp(FSSDSignalSample Sample, FSSDSignalSample SampleMin, FSSDSignalSample SampleMax)
 {
-	FSSFSignalSample OutSample;
+	FSSDSignalSample OutSample;
 	OutSample.SceneColor = clamp(Sample.SceneColor, SampleMin.SceneColor, SampleMax.SceneColor);
-	OutSample.AccumulatedSampleCount = clamp(Sample.AccumulatedSampleCount, SampleMin.AccumulatedSampleCount, SampleMax.AccumulatedSampleCount);
-	OutSample.LeavingRayCount = clamp(Sample.LeavingRayCount, SampleMin.LeavingRayCount, SampleMax.LeavingRayCount);
-	OutSample.PenumbraClosestHit = clamp(Sample.PenumbraClosestHit, SampleMin.PenumbraClosestHit, SampleMax.PenumbraClosestHit);
+	OutSample.SampleCount = clamp(Sample.SampleCount, SampleMin.SampleCount, SampleMax.SampleCount);
+	OutSample.MissCount = clamp(Sample.MissCount, SampleMin.MissCount, SampleMax.MissCount);
+	OutSample.ClosestHitDistance = clamp(Sample.ClosestHitDistance, SampleMin.ClosestHitDistance, SampleMax.ClosestHitDistance);
 	OutSample.UncertainRayCount = clamp(Sample.UncertainRayCount, SampleMin.UncertainRayCount, SampleMax.UncertainRayCount);
 	OutSample.WorldBluringRadius = clamp(Sample.WorldBluringRadius, SampleMin.WorldBluringRadius, SampleMax.WorldBluringRadius);
 	return OutSample;
 }
 
-FSSFSignalSample Lerp(FSSFSignalSample Sample0, FSSFSignalSample Sample1, float Interp)
+FSSDSignalSample Lerp(FSSDSignalSample Sample0, FSSDSignalSample Sample1, float Interp)
 {
 	return Add(Sample0, Mul(Subtract(Sample1, Sample0), Interp));
 }
 
 /** Normalize the signal sample to one. */
-FSSFSignalSample NormalizeToOneSample(FSSFSignalSample Sample)
+FSSDSignalSample NormalizeToOneSample(FSSDSignalSample Sample)
 {
-	FSSFSignalSample OutSample = Mul(Sample, Sample.AccumulatedSampleCount > 0 ? rcp(Sample.AccumulatedSampleCount) : 0);
-	OutSample.AccumulatedSampleCount = Sample.AccumulatedSampleCount > 0 ? 1 : 0;
+	FSSDSignalSample OutSample = Mul(Sample, Sample.SampleCount > 0 ? rcp(Sample.SampleCount) : 0);
+	OutSample.SampleCount = Sample.SampleCount > 0 ? 1 : 0;
 	return OutSample;
 }
 
@@ -215,9 +215,9 @@ void DecodeMultiplexedSignals(
 		for (uint BatchSignalId = 0; BatchSignalId < MAX_SIGNAL_BATCH_SIZE; BatchSignalId++)
 		{
 			uint MultiplexId = BatchSignalId;
-			OutSamples.Array[MultiplexId].AccumulatedSampleCount = (RawSample[MultiplexId].r == DENOISER_INVALID_HIT_DISTANCE ? 0.0 : 1.0);
-			OutSamples.Array[MultiplexId].LeavingRayCount = OutSamples.Array[MultiplexId].AccumulatedSampleCount * (RawSample[MultiplexId].r > 0.0 ? 0.0 : 1.0);
-			OutSamples.Array[MultiplexId].PenumbraClosestHit = OutSamples.Array[MultiplexId].AccumulatedSampleCount * (RawSample[MultiplexId].r == DENOISER_MISS_HIT_DISTANCE ? 1000000 : RawSample[MultiplexId].r);
+			OutSamples.Array[MultiplexId].SampleCount = (RawSample[MultiplexId].r == DENOISER_INVALID_HIT_DISTANCE ? 0.0 : 1.0);
+			OutSamples.Array[MultiplexId].MissCount = OutSamples.Array[MultiplexId].SampleCount * (RawSample[MultiplexId].r > 0.0 ? 0.0 : 1.0);
+			OutSamples.Array[MultiplexId].ClosestHitDistance = OutSamples.Array[MultiplexId].SampleCount * (RawSample[MultiplexId].r == DENOISER_MISS_HIT_DISTANCE ? 1000000 : RawSample[MultiplexId].r);
 		}
 	}
 	else if (SignalBufferLayout == SIGNAL_BUFFER_LAYOUT_PENUMBRA_INPUT_NSPP)
@@ -227,9 +227,9 @@ void DecodeMultiplexedSignals(
 		{
 			uint MultiplexId = BatchSignalId;
 			// TODO: feed the actual number of sample.
-			OutSamples.Array[MultiplexId].AccumulatedSampleCount = (RawSample[MultiplexId].g == DENOISER_INVALID_HIT_DISTANCE ? 0.0 : 1.0);
-			OutSamples.Array[MultiplexId].LeavingRayCount = OutSamples.Array[MultiplexId].AccumulatedSampleCount * RawSample[MultiplexId].r;
-			OutSamples.Array[MultiplexId].PenumbraClosestHit = OutSamples.Array[MultiplexId].AccumulatedSampleCount * (RawSample[MultiplexId].g == DENOISER_MISS_HIT_DISTANCE ? 1000000 : RawSample[MultiplexId].g);
+			OutSamples.Array[MultiplexId].SampleCount = (RawSample[MultiplexId].g == DENOISER_INVALID_HIT_DISTANCE ? 0.0 : 1.0);
+			OutSamples.Array[MultiplexId].MissCount = OutSamples.Array[MultiplexId].SampleCount * RawSample[MultiplexId].r;
+			OutSamples.Array[MultiplexId].ClosestHitDistance = OutSamples.Array[MultiplexId].SampleCount * (RawSample[MultiplexId].g == DENOISER_MISS_HIT_DISTANCE ? 1000000 : RawSample[MultiplexId].g);
 		}
 	}
 	else if (SignalBufferLayout == SIGNAL_BUFFER_LAYOUT_PENUMBRA_INJESTION_1SPP)
@@ -239,9 +239,9 @@ void DecodeMultiplexedSignals(
 		{
 			uint MultiplexId = BatchSignalId;
 			float EncodedWorldBluringRadius = RawSample[MultiplexId / 4][MultiplexId % 4];
-			OutSamples.Array[MultiplexId].AccumulatedSampleCount = (EncodedWorldBluringRadius == -2.0 ? 0.0 : 1.0);
-			OutSamples.Array[MultiplexId].LeavingRayCount = OutSamples.Array[MultiplexId].AccumulatedSampleCount * (EncodedWorldBluringRadius > 0.0 ? 0.0 : 1.0);
-			OutSamples.Array[MultiplexId].WorldBluringRadius = OutSamples.Array[MultiplexId].AccumulatedSampleCount * (EncodedWorldBluringRadius == -1 ? WORLD_RADIUS_MISS : EncodedWorldBluringRadius);
+			OutSamples.Array[MultiplexId].SampleCount = (EncodedWorldBluringRadius == -2.0 ? 0.0 : 1.0);
+			OutSamples.Array[MultiplexId].MissCount = OutSamples.Array[MultiplexId].SampleCount * (EncodedWorldBluringRadius > 0.0 ? 0.0 : 1.0);
+			OutSamples.Array[MultiplexId].WorldBluringRadius = OutSamples.Array[MultiplexId].SampleCount * (EncodedWorldBluringRadius == -1 ? WORLD_RADIUS_MISS : EncodedWorldBluringRadius);
 		}
 	}
 	else if (SignalBufferLayout == SIGNAL_BUFFER_LAYOUT_PENUMBRA_INJESTION_NSPP)
@@ -253,9 +253,9 @@ void DecodeMultiplexedSignals(
 			float2 Channels = float2(RawSample[MultiplexId / 2][2 * (MultiplexId % 2)], RawSample[MultiplexId / 2][2 * (MultiplexId % 2) + 1]);
 	
 			// TODO: feed the actual number of sample.
-			OutSamples.Array[MultiplexId].AccumulatedSampleCount = (Channels.g == -2.0 ? 0.0 : 1.0);
-			OutSamples.Array[MultiplexId].LeavingRayCount = OutSamples.Array[MultiplexId].AccumulatedSampleCount * Channels.r;
-			OutSamples.Array[MultiplexId].WorldBluringRadius = OutSamples.Array[MultiplexId].AccumulatedSampleCount * (Channels.g == -1 ? WORLD_RADIUS_MISS : Channels.g);
+			OutSamples.Array[MultiplexId].SampleCount = (Channels.g == -2.0 ? 0.0 : 1.0);
+			OutSamples.Array[MultiplexId].MissCount = OutSamples.Array[MultiplexId].SampleCount * Channels.r;
+			OutSamples.Array[MultiplexId].WorldBluringRadius = OutSamples.Array[MultiplexId].SampleCount * (Channels.g == -1 ? WORLD_RADIUS_MISS : Channels.g);
 		}
 	}
 	//else if (SignalBufferLayout == SIGNAL_BUFFER_LAYOUT_PENUMBRA_RECONSTRUCTION)
@@ -267,13 +267,13 @@ void DecodeMultiplexedSignals(
 	//
 	//		if (bNormalizeSample)
 	//		{
-	//			OutSamples.Array[MultiplexId].AccumulatedSampleCount = Channels.g > 0 ? 1 : 0;
-	//			OutSamples.Array[MultiplexId].LeavingRayCount = Channels.g > 0 ? Channels.r : 0;
+	//			OutSamples.Array[MultiplexId].SampleCount = Channels.g > 0 ? 1 : 0;
+	//			OutSamples.Array[MultiplexId].MissCount = Channels.g > 0 ? Channels.r : 0;
 	//		}
 	//		else
 	//		{
-	//			OutSamples.Array[MultiplexId].AccumulatedSampleCount = Channels.g;
-	//			OutSamples.Array[MultiplexId].LeavingRayCount = Channels.r * Channels.g;
+	//			OutSamples.Array[MultiplexId].SampleCount = Channels.g;
+	//			OutSamples.Array[MultiplexId].MissCount = Channels.r * Channels.g;
 	//		}
 	//	}
 	//}
@@ -286,8 +286,8 @@ void DecodeMultiplexedSignals(
 			
 			float SampleCount = bNormalizeSample ? (Channels.g > 0 ? 1 : 0) : (Channels.g);
 
-			OutSamples.Array[MultiplexId].LeavingRayCount = Channels.r * SampleCount;
-			OutSamples.Array[MultiplexId].AccumulatedSampleCount = SampleCount;
+			OutSamples.Array[MultiplexId].MissCount = Channels.r * SampleCount;
+			OutSamples.Array[MultiplexId].SampleCount = SampleCount;
 			OutSamples.Array[MultiplexId].WorldBluringRadius = (Channels.b == DENOISER_MISS_HIT_DISTANCE ? WORLD_RADIUS_MISS : Channels.b) * SampleCount;
 		}
 	}
@@ -300,131 +300,131 @@ void DecodeMultiplexedSignals(
 			float4 Channels = RawSample[BatchSignalId];
 			float SampleCount = Channels.b;
 
-			OutSamples.Array[MultiplexId + 0].AccumulatedSampleCount = SampleCount;
-			OutSamples.Array[MultiplexId + 0].LeavingRayCount = Channels.r * SampleCount;
+			OutSamples.Array[MultiplexId + 0].SampleCount = SampleCount;
+			OutSamples.Array[MultiplexId + 0].MissCount = Channels.r * SampleCount;
 			OutSamples.Array[MultiplexId + 0].WorldBluringRadius = (Channels.a == -1 ? WORLD_RADIUS_MISS : Channels.a) * SampleCount;
 			
-			OutSamples.Array[MultiplexId + 1].AccumulatedSampleCount = SampleCount;
-			OutSamples.Array[MultiplexId + 1].LeavingRayCount = Channels.g * SampleCount;
+			OutSamples.Array[MultiplexId + 1].SampleCount = SampleCount;
+			OutSamples.Array[MultiplexId + 1].MissCount = Channels.g * SampleCount;
 		}
 	}
 	else if (SignalBufferLayout == SIGNAL_BUFFER_LAYOUT_REFLECTIONS_INPUT)
 	{
-		OutSamples.Array[0].AccumulatedSampleCount = RawSample[1].r == DENOISER_INVALID_HIT_DISTANCE ? 0 : 1;
-		OutSamples.Array[0].LeavingRayCount = OutSamples.Array[0].AccumulatedSampleCount * (1 - RawSample[0].a);
+		OutSamples.Array[0].SampleCount = RawSample[1].r == DENOISER_INVALID_HIT_DISTANCE ? 0 : 1;
+		OutSamples.Array[0].MissCount = OutSamples.Array[0].SampleCount * (1 - RawSample[0].a);
 		OutSamples.Array[0].SceneColor = RawSample[0];
 
 		// max() because ray hit distance is < 0 when miss / not ray has been shot.
-		//OutSamples.Array[0].PenumbraClosestHit = RawSample[1].r * OutSamples.Array[0].AccumulatedSampleCount;
-		OutSamples.Array[0].PenumbraClosestHit = ((RawSample[1].r == DENOISER_MISS_HIT_DISTANCE) ? 1.0e20 : RawSample[1].r) * OutSamples.Array[0].AccumulatedSampleCount;
+		//OutSamples.Array[0].ClosestHitDistance = RawSample[1].r * OutSamples.Array[0].SampleCount;
+		OutSamples.Array[0].ClosestHitDistance = ((RawSample[1].r == DENOISER_MISS_HIT_DISTANCE) ? 1.0e20 : RawSample[1].r) * OutSamples.Array[0].SampleCount;
 	}
 	else if (SignalBufferLayout == SIGNAL_BUFFER_LAYOUT_REFLECTIONS_REJECTION)
 	{
 		if (bNormalizeSample)
 		{
-			OutSamples.Array[0].AccumulatedSampleCount = RawSample[1].r > 0 ? 1 : 0;
+			OutSamples.Array[0].SampleCount = RawSample[1].r > 0 ? 1 : 0;
 		}
 		else
 		{
-			OutSamples.Array[0].AccumulatedSampleCount = RawSample[1].r;
+			OutSamples.Array[0].SampleCount = RawSample[1].r;
 		}
 		
-		OutSamples.Array[0].LeavingRayCount = OutSamples.Array[0].AccumulatedSampleCount * (1 - RawSample[0].a);
-		OutSamples.Array[0].SceneColor = OutSamples.Array[0].AccumulatedSampleCount * RawSample[0];
+		OutSamples.Array[0].MissCount = OutSamples.Array[0].SampleCount * (1 - RawSample[0].a);
+		OutSamples.Array[0].SceneColor = OutSamples.Array[0].SampleCount * RawSample[0];
 
-		OutSamples.Array[1].AccumulatedSampleCount = RawSample[1].g;
-		OutSamples.Array[1].LeavingRayCount = OutSamples.Array[0].AccumulatedSampleCount * (1 - RawSample[2].a);
-		OutSamples.Array[1].SceneColor = OutSamples.Array[0].AccumulatedSampleCount * RawSample[2];
+		OutSamples.Array[1].SampleCount = RawSample[1].g;
+		OutSamples.Array[1].MissCount = OutSamples.Array[0].SampleCount * (1 - RawSample[2].a);
+		OutSamples.Array[1].SceneColor = OutSamples.Array[0].SampleCount * RawSample[2];
 	}
 	else if (SignalBufferLayout == SIGNAL_BUFFER_LAYOUT_REFLECTIONS_HISTORY)
 	{
 		if (bNormalizeSample)
 		{
-			OutSamples.Array[0].AccumulatedSampleCount = RawSample[1].r > 0 ? 1 : 0;
+			OutSamples.Array[0].SampleCount = RawSample[1].r > 0 ? 1 : 0;
 		}
 		else
 		{
-			OutSamples.Array[0].AccumulatedSampleCount = RawSample[1].r;
+			OutSamples.Array[0].SampleCount = RawSample[1].r;
 		}
 		
-		OutSamples.Array[0].LeavingRayCount = OutSamples.Array[0].AccumulatedSampleCount * (1 - RawSample[0].a);
-		OutSamples.Array[0].SceneColor = OutSamples.Array[0].AccumulatedSampleCount * RawSample[0];
+		OutSamples.Array[0].MissCount = OutSamples.Array[0].SampleCount * (1 - RawSample[0].a);
+		OutSamples.Array[0].SceneColor = OutSamples.Array[0].SampleCount * RawSample[0];
 
-		OutSamples.Array[1].AccumulatedSampleCount = OutSamples.Array[0].AccumulatedSampleCount;
-		OutSamples.Array[1].SceneColor.x = OutSamples.Array[1].AccumulatedSampleCount * RawSample[1].g;
+		OutSamples.Array[1].SampleCount = OutSamples.Array[0].SampleCount;
+		OutSamples.Array[1].SceneColor.x = OutSamples.Array[1].SampleCount * RawSample[1].g;
 		
-		OutSamples.Array[2].AccumulatedSampleCount = OutSamples.Array[0].AccumulatedSampleCount;
-		OutSamples.Array[2].SceneColor.x = OutSamples.Array[2].AccumulatedSampleCount * RawSample[1].b;
+		OutSamples.Array[2].SampleCount = OutSamples.Array[0].SampleCount;
+		OutSamples.Array[2].SceneColor.x = OutSamples.Array[2].SampleCount * RawSample[1].b;
 	}
 	else if (SignalBufferLayout == SIGNAL_BUFFER_LAYOUT_AO_INPUT)
 	{
-		OutSamples.Array[0].AccumulatedSampleCount = RawSample[1].r < 0 ? 0 : 1;
-		OutSamples.Array[0].LeavingRayCount = OutSamples.Array[0].AccumulatedSampleCount * RawSample[0].r;
-		OutSamples.Array[0].PenumbraClosestHit = OutSamples.Array[0].AccumulatedSampleCount * RawSample[1].r;
+		OutSamples.Array[0].SampleCount = RawSample[1].r < 0 ? 0 : 1;
+		OutSamples.Array[0].MissCount = OutSamples.Array[0].SampleCount * RawSample[0].r;
+		OutSamples.Array[0].ClosestHitDistance = OutSamples.Array[0].SampleCount * RawSample[1].r;
 	}
 	else if (SignalBufferLayout == SIGNAL_BUFFER_LAYOUT_AO_HISTORY)
 	{
 		if (bNormalizeSample)
 		{
-			OutSamples.Array[0].AccumulatedSampleCount = RawSample[0].g > 0 ? 1 : 0;
-			OutSamples.Array[0].LeavingRayCount = RawSample[0].g > 0 ? RawSample[0].r : 0;
+			OutSamples.Array[0].SampleCount = RawSample[0].g > 0 ? 1 : 0;
+			OutSamples.Array[0].MissCount = RawSample[0].g > 0 ? RawSample[0].r : 0;
 		}
 		else
 		{
-			OutSamples.Array[0].AccumulatedSampleCount = RawSample[0].g;
-			OutSamples.Array[0].LeavingRayCount = RawSample[0].g * RawSample[0].r;
+			OutSamples.Array[0].SampleCount = RawSample[0].g;
+			OutSamples.Array[0].MissCount = RawSample[0].g * RawSample[0].r;
 		}
 	}
 	else if (SignalBufferLayout == SIGNAL_BUFFER_LAYOUT_GI_AND_AO_INPUT)
 	{
-		OutSamples.Array[0].AccumulatedSampleCount = RawSample[1].g;
-		OutSamples.Array[0].LeavingRayCount = OutSamples.Array[0].AccumulatedSampleCount * (1 - RawSample[0].a);
-		OutSamples.Array[0].SceneColor = OutSamples.Array[0].AccumulatedSampleCount * RawSample[0];
+		OutSamples.Array[0].SampleCount = RawSample[1].g;
+		OutSamples.Array[0].MissCount = OutSamples.Array[0].SampleCount * (1 - RawSample[0].a);
+		OutSamples.Array[0].SceneColor = OutSamples.Array[0].SampleCount * RawSample[0];
 
 		// max() because ray hit distance is < 0 when miss / not ray has been shot.
-		OutSamples.Array[0].PenumbraClosestHit = RawSample[1].r * OutSamples.Array[0].AccumulatedSampleCount;
-		OutSamples.Array[0].PenumbraClosestHit = ((RawSample[1].r == -1.0) ? 1.0e20 : RawSample[1].r) * OutSamples.Array[0].AccumulatedSampleCount;
+		OutSamples.Array[0].ClosestHitDistance = RawSample[1].r * OutSamples.Array[0].SampleCount;
+		OutSamples.Array[0].ClosestHitDistance = ((RawSample[1].r == -1.0) ? 1.0e20 : RawSample[1].r) * OutSamples.Array[0].SampleCount;
 	}
 	else if (SignalBufferLayout == SIGNAL_BUFFER_LAYOUT_GI_AND_AO_RECONSTRUCTION)
 	{
 		if (bNormalizeSample)
 		{
-			OutSamples.Array[0].AccumulatedSampleCount = RawSample[1].r > 0 ? 1 : 0;
+			OutSamples.Array[0].SampleCount = RawSample[1].r > 0 ? 1 : 0;
 		}
 		else
 		{
-			OutSamples.Array[0].AccumulatedSampleCount = RawSample[1].r;
+			OutSamples.Array[0].SampleCount = RawSample[1].r;
 		}
 		
-		OutSamples.Array[0].LeavingRayCount = OutSamples.Array[0].AccumulatedSampleCount * (1 - RawSample[0].a);
-		OutSamples.Array[0].SceneColor = OutSamples.Array[0].AccumulatedSampleCount * RawSample[0];
+		OutSamples.Array[0].MissCount = OutSamples.Array[0].SampleCount * (1 - RawSample[0].a);
+		OutSamples.Array[0].SceneColor = OutSamples.Array[0].SampleCount * RawSample[0];
 	}
 	else if (SignalBufferLayout == SIGNAL_BUFFER_LAYOUT_GI_AND_AO_HISTORY)
 	{
 		if (bNormalizeSample)
 		{
-			OutSamples.Array[0].AccumulatedSampleCount = RawSample[1].r > 0 ? 1 : 0;
+			OutSamples.Array[0].SampleCount = RawSample[1].r > 0 ? 1 : 0;
 		}
 		else
 		{
-			OutSamples.Array[0].AccumulatedSampleCount = RawSample[1].r;
+			OutSamples.Array[0].SampleCount = RawSample[1].r;
 		}
 		
-		OutSamples.Array[0].LeavingRayCount = OutSamples.Array[0].AccumulatedSampleCount * (1 - RawSample[0].a);
-		OutSamples.Array[0].SceneColor = OutSamples.Array[0].AccumulatedSampleCount * RawSample[0];
+		OutSamples.Array[0].MissCount = OutSamples.Array[0].SampleCount * (1 - RawSample[0].a);
+		OutSamples.Array[0].SceneColor = OutSamples.Array[0].SampleCount * RawSample[0];
 		
-		OutSamples.Array[1].AccumulatedSampleCount = OutSamples.Array[0].AccumulatedSampleCount;
-		OutSamples.Array[1].SceneColor.x = OutSamples.Array[1].AccumulatedSampleCount * RawSample[1].g;
+		OutSamples.Array[1].SampleCount = OutSamples.Array[0].SampleCount;
+		OutSamples.Array[1].SceneColor.x = OutSamples.Array[1].SampleCount * RawSample[1].g;
 		
-		OutSamples.Array[2].AccumulatedSampleCount = OutSamples.Array[0].AccumulatedSampleCount;
-		OutSamples.Array[2].SceneColor.x = OutSamples.Array[2].AccumulatedSampleCount * RawSample[1].b;
+		OutSamples.Array[2].SampleCount = OutSamples.Array[0].SampleCount;
+		OutSamples.Array[2].SceneColor.x = OutSamples.Array[2].SampleCount * RawSample[1].b;
 	}
 } // DecodeMultiplexedSignals()
 
 /** Encode output signal sample. */
 void EncodeMultiplexedSignals(
 	const uint SignalBufferLayout, const uint MultiplexCount,
-	FSSFSignalSample Sample[MAX_SIGNAL_MULTIPLEXING],
+	FSSDSignalSample Sample[MAX_SIGNAL_MULTIPLEXING],
 	out float4 OutRawSample[MAX_MULTIPLEXED_TEXTURES],
 	out uint OutBufferCount)
 {
@@ -442,7 +442,7 @@ void EncodeMultiplexedSignals(
 		for (uint MultiplexId = 0; MultiplexId < MultiplexCount; MultiplexId++)
 		{
 			float EncodedChannel = Sample[MultiplexId].WorldBluringRadius;
-			if (Sample[MultiplexId].AccumulatedSampleCount == 0)
+			if (Sample[MultiplexId].SampleCount == 0)
 			{
 				EncodedChannel = -2;
 			}
@@ -461,10 +461,10 @@ void EncodeMultiplexedSignals(
 		for (uint MultiplexId = 0; MultiplexId < MultiplexCount; MultiplexId++)
 		{
 			float2 EncodedChannels = float2(
-				Sample[MultiplexId].LeavingRayCount,
-				Sample[MultiplexId].WorldBluringRadius) * SafeRcp(Sample[MultiplexId].AccumulatedSampleCount);
+				Sample[MultiplexId].MissCount,
+				Sample[MultiplexId].WorldBluringRadius) * SafeRcp(Sample[MultiplexId].SampleCount);
 			
-			if (Sample[MultiplexId].AccumulatedSampleCount == 0)
+			if (Sample[MultiplexId].SampleCount == 0)
 			{
 				EncodedChannels.y = -2;
 			}
@@ -486,8 +486,8 @@ void EncodeMultiplexedSignals(
 	//	for (uint MultiplexId = 0; MultiplexId < MultiplexCount; MultiplexId++)
 	//	{
 	//		float2 EncodedChannels = float2(
-	//			Sample[MultiplexId].LeavingRayCount * SafeRcp(Sample[MultiplexId].AccumulatedSampleCount),
-	//			Sample[MultiplexId].AccumulatedSampleCount);
+	//			Sample[MultiplexId].MissCount * SafeRcp(Sample[MultiplexId].SampleCount),
+	//			Sample[MultiplexId].SampleCount);
 	//
 	//		if (MultiplexId % 2)
 	//			OutRawSample[MultiplexId / 2].zw = EncodedChannels;
@@ -501,13 +501,13 @@ void EncodeMultiplexedSignals(
 		UNROLL
 		for (uint MultiplexId = 0; MultiplexId < MultiplexCount; MultiplexId++)
 		{
-			float NormalizationFactor = SafeRcp(Sample[MultiplexId].AccumulatedSampleCount);
+			float NormalizationFactor = SafeRcp(Sample[MultiplexId].SampleCount);
 
 			float NormalizedWorldBluringRadius = Sample[MultiplexId].WorldBluringRadius * NormalizationFactor;
 
 			OutRawSample[MultiplexId] = float4(
-				Sample[MultiplexId].LeavingRayCount * NormalizationFactor,
-				Sample[MultiplexId].AccumulatedSampleCount,
+				Sample[MultiplexId].MissCount * NormalizationFactor,
+				Sample[MultiplexId].SampleCount,
 				NormalizedWorldBluringRadius == WORLD_RADIUS_MISS ? -1 : NormalizedWorldBluringRadius,
 				0);
 		}
@@ -523,40 +523,40 @@ void EncodeMultiplexedSignals(
 		{
 			uint MultiplexId = BatchSignalId * 2;
 
-			float NormalizationFactor = SafeRcp(Sample[MultiplexId + 0].AccumulatedSampleCount);
+			float NormalizationFactor = SafeRcp(Sample[MultiplexId + 0].SampleCount);
 
 			float NormalizedWorldBluringRadius = Sample[MultiplexId + 0].WorldBluringRadius * NormalizationFactor;
 
 			// Samples are normalised when doing history preconvolution, so the number of is either 0 or 1 on both momments. Therefore
-			// Sample[MultiplexId + 0].AccumulatedSampleCount == Sample[MultiplexId + 1].AccumulatedSampleCount.
+			// Sample[MultiplexId + 0].SampleCount == Sample[MultiplexId + 1].SampleCount.
 			OutRawSample[BatchSignalId] = float4(
-				Sample[MultiplexId + 0].LeavingRayCount * NormalizationFactor,
-				Sample[MultiplexId + 1].LeavingRayCount * NormalizationFactor,
-				Sample[MultiplexId + 0].AccumulatedSampleCount,
+				Sample[MultiplexId + 0].MissCount * NormalizationFactor,
+				Sample[MultiplexId + 1].MissCount * NormalizationFactor,
+				Sample[MultiplexId + 0].SampleCount,
 				NormalizedWorldBluringRadius == WORLD_RADIUS_MISS ? -1 : NormalizedWorldBluringRadius);
 		}
 	}
 	else if (SignalBufferLayout == SIGNAL_BUFFER_LAYOUT_REFLECTIONS_REJECTION)
 	{
 		// Outputs scene color to be compatible with a SSR output.
-		OutRawSample[0] = Sample[0].SceneColor * (Sample[0].AccumulatedSampleCount > 0 ? rcp(Sample[0].AccumulatedSampleCount) : 0);
-		OutRawSample[1].r = Sample[0].AccumulatedSampleCount;
+		OutRawSample[0] = Sample[0].SceneColor * (Sample[0].SampleCount > 0 ? rcp(Sample[0].SampleCount) : 0);
+		OutRawSample[1].r = Sample[0].SampleCount;
 		OutBufferCount = 2;
 
 		if (MultiplexCount == 2)
 		{
-			OutRawSample[2] = Sample[1].SceneColor * (Sample[1].AccumulatedSampleCount > 0 ? rcp(Sample[1].AccumulatedSampleCount) : 0);
-			OutRawSample[1].g = Sample[1].AccumulatedSampleCount;
+			OutRawSample[2] = Sample[1].SceneColor * (Sample[1].SampleCount > 0 ? rcp(Sample[1].SampleCount) : 0);
+			OutRawSample[1].g = Sample[1].SampleCount;
 			OutBufferCount = 3;
 		}
 	}
 	else if (SignalBufferLayout == SIGNAL_BUFFER_LAYOUT_REFLECTIONS_HISTORY)
 	{
-		float NormalizationFactor = SafeRcp(Sample[0].AccumulatedSampleCount);
+		float NormalizationFactor = SafeRcp(Sample[0].SampleCount);
 
 		// Outputs scene color to be compatible with a SSR output.
 		OutRawSample[0] = Sample[0].SceneColor * NormalizationFactor;
-		OutRawSample[1].r = Sample[0].AccumulatedSampleCount;
+		OutRawSample[1].r = Sample[0].SampleCount;
 		OutBufferCount = 2;
 
 		if (MultiplexCount == 3)
@@ -569,23 +569,23 @@ void EncodeMultiplexedSignals(
 	else if (SignalBufferLayout == SIGNAL_BUFFER_LAYOUT_AO_HISTORY)
 	{
 		// Outputs number of ray miss in red to be compatible as a SSAO output.
-		OutRawSample[0].r = Sample[0].AccumulatedSampleCount > 0 ? Sample[0].LeavingRayCount / Sample[0].AccumulatedSampleCount : 1.0;
-		OutRawSample[0].g = Sample[0].AccumulatedSampleCount;
+		OutRawSample[0].r = Sample[0].SampleCount > 0 ? Sample[0].MissCount / Sample[0].SampleCount : 1.0;
+		OutRawSample[0].g = Sample[0].SampleCount;
 	}
 	else if (SignalBufferLayout == SIGNAL_BUFFER_LAYOUT_GI_AND_AO_RECONSTRUCTION)
 	{
-		float NormalizationFactor = SafeRcp(Sample[0].AccumulatedSampleCount);
+		float NormalizationFactor = SafeRcp(Sample[0].SampleCount);
 
 		OutRawSample[0] = Sample[0].SceneColor * NormalizationFactor;
-		OutRawSample[1].r = Sample[0].AccumulatedSampleCount;
+		OutRawSample[1].r = Sample[0].SampleCount;
 		OutBufferCount = 2;
 	}
 	else if (SignalBufferLayout == SIGNAL_BUFFER_LAYOUT_GI_AND_AO_HISTORY)
 	{
-		float NormalizationFactor = SafeRcp(Sample[0].AccumulatedSampleCount);
+		float NormalizationFactor = SafeRcp(Sample[0].SampleCount);
 
 		OutRawSample[0] = Sample[0].SceneColor * NormalizationFactor;
-		OutRawSample[1].r = Sample[0].AccumulatedSampleCount;
+		OutRawSample[1].r = Sample[0].SampleCount;
 		OutBufferCount = 2;
 
 		if (MultiplexCount == 3)
@@ -629,7 +629,7 @@ void OutputMultiplexedSignal(
 	RWTexture2D<float4> OutputSignalBuffer2,
 	RWTexture2D<float4> OutputSignalBuffer3,
 	const uint SignalBufferLayout, const uint MultiplexCount,
-	const uint2 PixelPosition, FSSFSignalSample MultiplexedSamples[MAX_SIGNAL_MULTIPLEXING])
+	const uint2 PixelPosition, FSSDSignalSample MultiplexedSamples[MAX_SIGNAL_MULTIPLEXING])
 {
 	// Encode the output signal.
 	float4 RawSample[MAX_MULTIPLEXED_TEXTURES];
diff --git a/Engine/Shaders/Private/ScreenSpaceDenoise/SSDSignalFramework.ush b/Engine/Shaders/Private/ScreenSpaceDenoise/SSDSignalFramework.ush
--- a/Engine/Shaders/Private/ScreenSpaceDenoise/SSDSignalFramework.ush
+++ b/Engine/Shaders/Private/ScreenSpaceDenoise/SSDSignalFramework.ush
@@ -159,7 +159,7 @@ float GetSignalColorTransformationExposure()
 }
 
 /** Returns the linear color space Luma with a 4 multiplier for ALU perf reasons. */
-float GetSignalLuma4(FSSFSignalSample Sample, const uint Basis)
+float GetSignalLuma4(FSSDSignalSample Sample, const uint Basis)
 {
 	// TODO: should force this to be scalar load.
 	float FrameExposureScale = GetSignalColorTransformationExposure();
@@ -177,7 +177,7 @@ float GetSignalLuma4(FSSFSignalSample Sample, const uint Basis)
 
 		if (!bIsNormalizedSample)
 		{
-			KarisX *= SafeRcp(Sample.AccumulatedSampleCount);
+			KarisX *= SafeRcp(Sample.SampleCount);
 		}
 		
 		return KarisX * KarisHdrWeightInvY(KarisX, FrameExposureScale);
@@ -194,7 +194,7 @@ float GetSignalLuma4(FSSFSignalSample Sample, const uint Basis)
 }
 
 /** Conveniently transform one simple from a source basis to a destination basis. */
-FSSFSignalSample TransformSignal(FSSFSignalSample Sample, const uint SrcBasis, const uint DestBasis)
+FSSDSignalSample TransformSignal(FSSDSignalSample Sample, const uint SrcBasis, const uint DestBasis)
 {
 	// If the source and destination bases are the same, early return to not pay ALU and floating point divergence.
 	if (SrcBasis == DestBasis)
@@ -216,7 +216,7 @@ FSSFSignalSample TransformSignal(FSSFSignalSample Sample, const uint SrcBasis, c
 
 	if (bDebugForceNormalizeColor)
 	{
-		Sample.SceneColor *= SafeRcp(Sample.AccumulatedSampleCount);
+		Sample.SceneColor *= SafeRcp(Sample.SampleCount);
 	}
 
 	// Remove Karis weighting before doing any color transformations.
@@ -230,7 +230,7 @@ FSSFSignalSample TransformSignal(FSSFSignalSample Sample, const uint SrcBasis, c
 
 		if (!bIsNormalizedColor)
 		{
-			KarisX *= SafeRcp(Sample.AccumulatedSampleCount);
+			KarisX *= SafeRcp(Sample.SampleCount);
 		}
 
 		Sample.SceneColor *= KarisHdrWeightInvY(KarisX, FrameExposureScale);
@@ -249,7 +249,7 @@ FSSFSignalSample TransformSignal(FSSFSignalSample Sample, const uint SrcBasis, c
 		Sample.SceneColor.rgb = LCoCg_2_LinearRGB(Sample.SceneColor.rgb);
 	}
 
-	float Alpha = Sample.SceneColor.a * SafeRcp(Sample.AccumulatedSampleCount);
+	float Alpha = Sample.SceneColor.a * SafeRcp(Sample.SampleCount);
 	if (bIsNormalizedColor)
 	{
 		Alpha = Sample.SceneColor.a;
@@ -293,13 +293,13 @@ FSSFSignalSample TransformSignal(FSSFSignalSample Sample, const uint SrcBasis, c
 		}
 		else
 		{
-			Sample.SceneColor *= KarisHdrWeightY(x * SafeRcp(Sample.AccumulatedSampleCount), FrameExposureScale);
+			Sample.SceneColor *= KarisHdrWeightY(x * SafeRcp(Sample.SampleCount), FrameExposureScale);
 		}
 	}
 	
 	if (bDebugForceNormalizeColor)
 	{
-		Sample.SceneColor *= Sample.AccumulatedSampleCount;
+		Sample.SceneColor *= Sample.SampleCount;
 	}
 
 	return Sample;
@@ -344,7 +344,7 @@ FSphericalGaussian ComputeRoughnessLobe(float Roughness, float3 N, float3 V)
 }
 
 /** Create a spherical gaussian lob for a given scene metadata. */
-FSphericalGaussian ComputeRoughnessLobe(FSSFSampleSceneInfos RefSceneMetadata)
+FSphericalGaussian ComputeRoughnessLobe(FSSDSampleSceneInfos RefSceneMetadata)
 {
 	float3 N = RefSceneMetadata.WorldNormal;
 	float3 V = normalize(View.TranslatedWorldCameraOrigin - RefSceneMetadata.TranslatedWorldPosition);
@@ -354,7 +354,7 @@ FSphericalGaussian ComputeRoughnessLobe(FSSFSampleSceneInfos RefSceneMetadata)
 }
 
 /** Return the size of the bluring radius caused by the specular lobe in ViewportUV.x space. */
-float ComputeInfinityBluringRadius(FSSFSampleSceneInfos SceneMetadata)
+float ComputeInfinityBluringRadius(FSSDSampleSceneInfos SceneMetadata)
 {
 	float RoughnessLobeAngle;
 	#if 1
@@ -396,13 +396,13 @@ float ComputeInfinityBluringRadius(FSSFSampleSceneInfos SceneMetadata)
 }
 
 /** Return the size of bluring radius caused by specular lobe in view space. */
-float ComputeBluringRadiusOfNormalizedSample(FSSFSignalSample Sample, FSSFSampleSceneInfos SceneMetadata)
+float ComputeBluringRadiusOfNormalizedSample(FSSDSignalSample Sample, FSSDSampleSceneInfos SceneMetadata)
 {
 	float DistanceCameraToPixel = length(View.TranslatedWorldCameraOrigin - SceneMetadata.TranslatedWorldPosition);
 
 	float InifinityRadius = ComputeInfinityBluringRadius(SceneMetadata);
 
-	return InifinityRadius * (1 - DistanceCameraToPixel / (DistanceCameraToPixel + Sample.PenumbraClosestHit));
+	return InifinityRadius * (1 - DistanceCameraToPixel / (DistanceCameraToPixel + Sample.ClosestHitDistance));
 }
 
 // TODO: duplicated...
@@ -445,7 +445,7 @@ float3 GenerateReflectedRayDirection(
 // TODO: duplicated...
 void GetInputRayDirectionPixel(
 	uint2 PixelCoord,
-	FSSFSampleSceneInfos SceneMetadata,
+	FSSDSampleSceneInfos SceneMetadata,
 	out float3 RayDirection,
 	out float PDF)
 {
@@ -474,23 +474,23 @@ void GetInputRayDirectionPixel(
 
 /** Return whether this signal is occluded or not. */
 #if CONFIG_SIGNAL_PROCESSING == SIGNAL_PROCESSING_MONOCHROMATIC_PENUMBRA
-bool IsMissSample(FSSFSignalSample Sample)
+bool IsMissSample(FSSDSignalSample Sample)
 {
-	return Sample.LeavingRayCount != 0.0;
+	return Sample.MissCount != 0.0;
 }
 #endif
 
 /** Return whether this signal is valid or not. */
-bool IsInvalidSample(FSSFSignalSample Sample)
+bool IsInvalidSample(FSSDSignalSample Sample)
 {
-	return Sample.AccumulatedSampleCount == 0.0;
+	return Sample.SampleCount == 0.0;
 }
 
 /** COmpute the weight of the bilateral filter to use between a reference sample and a neighbor. */
 float ComputeBilateralWeight(
 	float MaxWorldBluringDistance,
-	FSSFSampleSceneInfos RefSceneMetadata,
-	FSSFSampleSceneInfos NeighborSceneMetadata)
+	FSSDSampleSceneInfos RefSceneMetadata,
+	FSSDSampleSceneInfos NeighborSceneMetadata)
 {
 	if (DEBUG_DISABLE_BILATERAL)
 	{
@@ -595,7 +595,7 @@ float ComputeBilateralWeight(
 	return PositionBilateralWeight * NormalBilateralWeight * LobeSimilarity * AxesSimilarity;
 }
 
-float GetSignalWorldBluringRadius(FSSFSignalSample Sample, FSSFSampleSceneInfos SceneMetadata, uint BatchedSignalId)
+float GetSignalWorldBluringRadius(FSSDSignalSample Sample, FSSDSampleSceneInfos SceneMetadata, uint BatchedSignalId)
 #if CONFIG_SIGNAL_PROCESSING == SIGNAL_PROCESSING_MONOCHROMATIC_PENUMBRA
 {
 	if (IsInvalidSample(Sample))
@@ -616,14 +616,14 @@ float GetSignalWorldBluringRadius(FSSFSignalSample Sample, FSSFSampleSceneInfos
 	BRANCH
 	if (DomainKnowledge.LightType == LIGHT_TYPE_DIRECTIONAL)
 	{
-		return DomainKnowledge.HitDistanceToWorldBluringRadius * Sample.PenumbraClosestHit;
+		return DomainKnowledge.HitDistanceToWorldBluringRadius * Sample.ClosestHitDistance;
 	}
 	else if (DomainKnowledge.LightType == LIGHT_TYPE_POINT || DomainKnowledge.LightType == LIGHT_TYPE_SPOT)
 	{
 		float3 PixelToLightWorldVector = (DomainKnowledge.Light.Position + View.PreViewTranslation) - SceneMetadata.TranslatedWorldPosition;
 		float DistanceFromLight = length(PixelToLightWorldVector);
 
-		return DomainKnowledge.Light.SourceRadius * Sample.PenumbraClosestHit / max(DistanceFromLight - Sample.PenumbraClosestHit, ClosestLightDistance);
+		return DomainKnowledge.Light.SourceRadius * Sample.ClosestHitDistance / max(DistanceFromLight - Sample.ClosestHitDistance, ClosestLightDistance);
 	}
 	else // if (DomainKnowledge.LightType == LIGHT_TYPE_RECT)
 	{
@@ -641,8 +641,8 @@ float GetSignalWorldBluringRadius(FSSFSignalSample Sample, FSSFSampleSceneInfos
 		// TODO: area light anysotropy.
 		float SmallestLightDimension = min(LightDimensions.x, LightDimensions.y);
 
-		// TODO: Sample.PenumbraClosestHit depends on the direction of the ray, but does not DistanceFromLight, witch is bad for large area light.
-		return SmallestLightDimension * Sample.PenumbraClosestHit / max(DistanceFromLight - Sample.PenumbraClosestHit, ClosestLightDistance);
+		// TODO: Sample.ClosestHitDistance depends on the direction of the ray, but does not DistanceFromLight, witch is bad for large area light.
+		return SmallestLightDimension * Sample.ClosestHitDistance / max(DistanceFromLight - Sample.ClosestHitDistance, ClosestLightDistance);
 	}
 }
 #elif CONFIG_SIGNAL_PROCESSING == SIGNAL_PROCESSING_REFLECTIONS
@@ -651,37 +651,37 @@ float GetSignalWorldBluringRadius(FSSFSignalSample Sample, FSSFSampleSceneInfos
 }
 #elif CONFIG_SIGNAL_PROCESSING == SIGNAL_PROCESSING_AO
 {
-	return Sample.PenumbraClosestHit;
+	return Sample.ClosestHitDistance;
 }
 #elif CONFIG_SIGNAL_PROCESSING == SIGNAL_PROCESSING_GI_AND_AO
 {
-	return Sample.PenumbraClosestHit;
+	return Sample.ClosestHitDistance;
 }
 #else
 	#error Unimplemented GetSignalWorldBluringRadius()
 #endif
 
 // Returns the penumbra of this sample, or 1 if invalid.
-float GetSamplePenumbraSafe(FSSFSignalSample Sample)
+float GetSamplePenumbraSafe(FSSDSignalSample Sample)
 {
-	return (Sample.AccumulatedSampleCount > 0 ? Sample.LeavingRayCount / Sample.AccumulatedSampleCount : 1);
+	return (Sample.SampleCount > 0 ? Sample.MissCount / Sample.SampleCount : 1);
 }
 
 // Get penumbra, but be aware of NaN.
-float GetSamplePenumbra(FSSFSignalSample Sample)
+float GetSamplePenumbra(FSSDSignalSample Sample)
 {
 	if (1) // TODO: For debugging purpose of division by 0.
 	{
 		return GetSamplePenumbraSafe(Sample);
 	}
-	return Sample.LeavingRayCount / Sample.AccumulatedSampleCount;
+	return Sample.MissCount / Sample.SampleCount;
 }
 
 /** Get the weight for multi importance sampling. */
 float GetRatioEstimatorWeight(
-	FSSFSampleSceneInfos RefSceneMetadata,
-	FSSFSampleSceneInfos SampleSceneMetadata,
-	FSSFSignalSample Sample,
+	FSSDSampleSceneInfos RefSceneMetadata,
+	FSSDSampleSceneInfos SampleSceneMetadata,
+	FSSDSignalSample Sample,
 	uint2 SamplePixelCoord)
 #if CONFIG_SPECULAR_RATIO_ESTIMATOR
 // [ Stackowiak 2018, "Stochastic all the things: Raytracing in hybrid real-time rendering" ]
@@ -702,7 +702,7 @@ float GetRatioEstimatorWeight(
 		/* out */ OriginalRayPDF);
 
 	// min to avoid floating point underflow that causes artifact by D anormally high.
-	float PenumbraClosestHit = min(Sample.PenumbraClosestHit, 1000000);
+	float ClosestHitDistance = min(Sample.ClosestHitDistance, 1000000);
 
 	float a2 = Pow4(RefSceneMetadata.Roughness);
 
@@ -712,7 +712,7 @@ float GetRatioEstimatorWeight(
 	#if CONFIG_VALU_OPTIMIZATIONS
 	{
 		float3 V = (View.TranslatedWorldCameraOrigin - RefSceneMetadata.TranslatedWorldPosition);
-		float3 L = ((SampleSceneMetadata.TranslatedWorldPosition - RefSceneMetadata.TranslatedWorldPosition) + OriginalRayDirection * PenumbraClosestHit);
+		float3 L = ((SampleSceneMetadata.TranslatedWorldPosition - RefSceneMetadata.TranslatedWorldPosition) + OriginalRayDirection * ClosestHitDistance);
 		float3 H = (L + V * (sqrt(length2(L)) * rsqrt(length2(V))));
 
 		NoH = saturate(dot(N, H) * rsqrt(length2(H)));
@@ -720,7 +720,7 @@ float GetRatioEstimatorWeight(
 	#else
 	{
 		float3 V = normalize(View.TranslatedWorldCameraOrigin - RefSceneMetadata.TranslatedWorldPosition);
-		float3 L = normalize((SampleSceneMetadata.TranslatedWorldPosition - RefSceneMetadata.TranslatedWorldPosition) + OriginalRayDirection * PenumbraClosestHit);
+		float3 L = normalize((SampleSceneMetadata.TranslatedWorldPosition - RefSceneMetadata.TranslatedWorldPosition) + OriginalRayDirection * ClosestHitDistance);
 		float3 H = normalize(L + V);
 
 		NoH = abs(dot(N, H));
diff --git a/Engine/Shaders/Private/ScreenSpaceDenoise/SSDSpatialAccumulation.usf b/Engine/Shaders/Private/ScreenSpaceDenoise/SSDSpatialAccumulation.usf
--- a/Engine/Shaders/Private/ScreenSpaceDenoise/SSDSpatialAccumulation.usf
+++ b/Engine/Shaders/Private/ScreenSpaceDenoise/SSDSpatialAccumulation.usf
@@ -420,7 +420,7 @@ void MainCS(
 	
 	// Read reference meta data.
 	float2 ScreenPosition = ViewportUVToScreenPos(BufferUVToViewportUV(SceneBufferUV));
-	FSSFSampleSceneInfos RefSceneMetadata = FetchCurrentSceneInfos(ScreenPosition, SceneBufferUV);
+	FSSDSampleSceneInfos RefSceneMetadata = FetchCurrentSceneInfos(ScreenPosition, SceneBufferUV);
 
 	// Sample the reference sample.
 	#if !CONFIG_UPSCALE || 1
@@ -482,7 +482,7 @@ void MainCS(
 
 		IdealSampleCount = min(IdealSampleCount, MaxSampleCount);
 
-		RequestedSampleCount = IdealSampleCount / clamp(RefSamples.Array[0].AccumulatedSampleCount, 1, MaxSampleCount);
+		RequestedSampleCount = IdealSampleCount / clamp(RefSamples.Array[0].SampleCount, 1, MaxSampleCount);
 	}
 	#elif CONFIG_SAMPLE_COUNT_POLICY == SAMPLE_COUNT_POLICY_SAMPLE_ACCUMULATION_BASED
 	{
@@ -491,7 +491,7 @@ void MainCS(
 		#endif
 		const float IdealSampleCount = TARGETED_SAMPLE_COUNT;
 
-		RequestedSampleCount = clamp(IdealSampleCount / RefSamples.Array[0].AccumulatedSampleCount, 1, MaxSampleCount);
+		RequestedSampleCount = clamp(IdealSampleCount / RefSamples.Array[0].SampleCount, 1, MaxSampleCount);
 	}
 	#else
 		#error Unknown policy to control the number of samples.
@@ -556,7 +556,7 @@ void MainCS(
 		// When doing history rejection preconvolution may have invalid ref sample, in witch case need to force take neighborhood to have a clamping box.
 		#if DIM_STAGE == STAGE_REJECTION_PRE_CONVOLUTION && CONFIG_UPSCALE
 		{
-			KernelConfig.bForceAllAccumulation = RefSamples.Array[0].AccumulatedSampleCount == 0;
+			KernelConfig.bForceAllAccumulation = RefSamples.Array[0].SampleCount == 0;
 			KernelConfig.SampleSet = SAMPLE_SET_3X3_PLUS;
 		}
 		#endif
@@ -587,7 +587,7 @@ void MainCS(
 			{
 				const uint BatchedSignalId = ComputeSignalBatchIdFromSignalMultiplexId(KernelConfig, SignalMultiplexId);
 
-				FSSFSignalSample CenterSample = TransformSignalSampleForAccumulation(
+				FSSDSignalSample CenterSample = TransformSignalSampleForAccumulation(
 					KernelConfig,
 					SignalMultiplexId,
 					RefSamples.Array[SignalMultiplexId]);
@@ -600,7 +600,7 @@ void MainCS(
 				if (KernelConfig.BilateralDistanceComputation == SIGNAL_WORLD_FREQUENCY_PRECOMPUTED_BLURING_RADIUS)
 				{
 					// TODO: this is ineficient, could fetch the normalised WorldBluringRadius instead of SafeRcp().
-					SampleInfos.InvFrequency = CenterSample.WorldBluringRadius * SafeRcp(CenterSample.AccumulatedSampleCount);
+					SampleInfos.InvFrequency = CenterSample.WorldBluringRadius * SafeRcp(CenterSample.SampleCount);
 				}
 
 				AccumulateSample(
@@ -698,13 +698,13 @@ void MainCS(
 		for (uint MultiplexId = 0; MultiplexId < CONFIG_SIGNAL_BATCH_SIZE; MultiplexId++)
 		{
 			#if CONFIG_REJECT_HIGHEST_COLOR
-			if (Accumulator.Moment1[MultiplexId].AccumulatedSampleCount > 0)
+			if (Accumulator.Moment1[MultiplexId].SampleCount > 0)
 			{
-				const float MaxNeighborWeight = saturate(SignalAccumulators.Array[MultiplexId].Moment1.AccumulatedSampleCount * rcp(10) - 1);
+				const float MaxNeighborWeight = saturate(SignalAccumulators.Array[MultiplexId].Moment1.SampleCount * rcp(10) - 1);
 
 				SignalAccumulators.Array[MultiplexId].Moment1.SceneColor.rgb =
 					(SignalAccumulators.Array[MultiplexId].Moment1.SceneColor.rgb - MaxNeighbor.SceneColor.rgb * MaxNeighborWeight) *
-					(SignalAccumulators.Array[MultiplexId].Moment1.AccumulatedSampleCount / (SignalAccumulators.Array[MultiplexId].Moment1.AccumulatedSampleCount - MaxNeighborWeight));
+					(SignalAccumulators.Array[MultiplexId].Moment1.SampleCount / (SignalAccumulators.Array[MultiplexId].Moment1.SampleCount - MaxNeighborWeight));
 			}
 			#endif
 
@@ -736,7 +736,7 @@ void MainCS(
 				// Output the minimal inverse frequency as new world bluring radius for subsequent passes.
 				// TODO: store this in its own signal to avoid *= N; *= rcp(N);
 				OutputSamples.Array[MultiplexId].WorldBluringRadius = 
-					OutputSamples.Array[MultiplexId].AccumulatedSampleCount * SignalAccumulators.Array[MultiplexId].MinInvFrequency;
+					OutputSamples.Array[MultiplexId].SampleCount * SignalAccumulators.Array[MultiplexId].MinInvFrequency;
 			}
 		}
 		#elif CONFIG_OUTPUT_MODE == OUTPUT_MODE_2MOMMENT_SUM
@@ -762,7 +762,7 @@ void MainCS(
 				// Output the minimal inverse frequency as new world bluring radius for subsequent passes.
 				// TODO: store this in its own signal to avoid *= N; *= rcp(N);
 				OutputSamples.Array[MultiplexId].WorldBluringRadius = 
-					OutputSamples.Array[MultiplexId].AccumulatedSampleCount * SignalAccumulators.Array[MultiplexId].MinInvFrequency;
+					OutputSamples.Array[MultiplexId].SampleCount * SignalAccumulators.Array[MultiplexId].MinInvFrequency;
 			
 				// No need to keep the VGPR pressure at this point for WorldBluringRadius, because no passes use it after.
 				if (DIM_STAGE == STAGE_POST_FILTERING && 0)
@@ -782,7 +782,7 @@ void MainCS(
 		FIX_UNROLL(CONFIG_SIGNAL_BATCH_SIZE)
 		for (uint MultiplexId = 0; MultiplexId < CONFIG_SIGNAL_BATCH_SIZE; MultiplexId++)
 		{
-			float CurrentSampleCount = OutputSamples.Array[MultiplexId].AccumulatedSampleCount;
+			float CurrentSampleCount = OutputSamples.Array[MultiplexId].SampleCount;
 			float NewSampleCount = min(CurrentSampleCount, 128);
 
 			OutputSamples.Array[MultiplexId] = Mul(OutputSamples.Array[MultiplexId], CurrentSampleCount > 0 ? NewSampleCount / CurrentSampleCount : 0);
diff --git a/Engine/Shaders/Private/ScreenSpaceDenoise/SSDSpatialKernel.ush b/Engine/Shaders/Private/ScreenSpaceDenoise/SSDSpatialKernel.ush
--- a/Engine/Shaders/Private/ScreenSpaceDenoise/SSDSpatialKernel.ush
+++ b/Engine/Shaders/Private/ScreenSpaceDenoise/SSDSpatialKernel.ush
@@ -19,7 +19,7 @@
 	// Depends only based of the sample's hit distance and metadata.
 	#define SIGNAL_WORLD_FREQUENCY_HIT_DISTANCE 3
 
-	// Uses FSSFSignalSample::WorldBluringRadius precomputed in the sample.
+	// Uses FSSDSignalSample::WorldBluringRadius precomputed in the sample.
 	#define SIGNAL_WORLD_FREQUENCY_PRECOMPUTED_BLURING_RADIUS 4
 
 
@@ -123,7 +123,7 @@ struct FSSDKernelConfig
 	float2 BufferUV;
 
 	// Metadata of the scene for the bilateral therm.
-	FSSFSampleSceneInfos RefSceneMetadata;
+	FSSDSampleSceneInfos RefSceneMetadata;
 
 	// Runtime to force the first sample of the kernel to be accumulated.
 	bool bForceKernelCenterAccumulation;
@@ -198,10 +198,10 @@ static const float kWaveletFilterWeights5x5[] = { 3.0 / 8.0, 1.0 / 4.0, 1.0 / 16
 
 //------------------------------------------------------- SHARED SAMPLING
 
-FSSFSignalSample TransformSignalSampleForAccumulation(
+FSSDSignalSample TransformSignalSampleForAccumulation(
 	FSSDKernelConfig KernelConfig,
 	uint MultiplexId,
-	FSSFSignalSample Sample)
+	FSSDSignalSample Sample)
 {
 	// TODO: could pass down information that this sample may be normalized.
 	Sample = TransformSignal(
@@ -233,14 +233,14 @@ void SampleMultiplexedSignals(
 	Texture2D SignalBuffer2,
 	Texture2D SignalBuffer3,
 	float2 SampleBufferUV,
-	out FSSFSampleSceneInfos OutSampleSceneMetadata,
+	out FSSDSampleSceneInfos OutSampleSceneMetadata,
 	out FSSDSignalArray OutMultiplexedSamples)
 {
 	// Sample the meta data of the sample.	
 	float2 SampleViewportUV = BufferUVToViewportUV(SampleBufferUV);
 	float2 SampleScreenPosition = ViewportUVToScreenPos(SampleViewportUV);
 	
-	FSSFSampleSceneInfos SampleSceneMetadata;
+	FSSDSampleSceneInfos SampleSceneMetadata;
 	if (KernelConfig.bPreviousFrameMetadata)
 	{
 		OutSampleSceneMetadata = FetchPrevSceneInfos(SampleScreenPosition, SampleBufferUV);
@@ -269,7 +269,7 @@ void AccumulateSampledMultiplexedSignals(
 	FSSDKernelConfig KernelConfig,
 	inout FSSDSignalAccumulatorArray Accumulators,
 	float2 SampleBufferUV,
-	FSSFSampleSceneInfos SampleSceneMetadata,
+	FSSDSampleSceneInfos SampleSceneMetadata,
 	FSSDSignalArray MultiplexedSamples,
 	float KernelSampleWeight,
 	const bool bForceSample,
@@ -288,7 +288,7 @@ void AccumulateSampledMultiplexedSignals(
 		const uint BatchedSignalId = ComputeSignalBatchIdFromSignalMultiplexId(KernelConfig, SignalMultiplexId);
 
 		// Fetch and pre process the sample for accumulation.
-		FSSFSignalSample Sample = MultiplexedSamples.Array[SignalMultiplexId];
+		FSSDSignalSample Sample = MultiplexedSamples.Array[SignalMultiplexId];
 		
 		Sample = TransformSignalSampleForAccumulation(KernelConfig, SignalMultiplexId, Sample);
 		
@@ -302,7 +302,7 @@ void AccumulateSampledMultiplexedSignals(
 		if (KernelConfig.BilateralDistanceComputation == SIGNAL_WORLD_FREQUENCY_PRECOMPUTED_BLURING_RADIUS)
 		{
 			// TODO: this is ineficient, could fetch the normalised WorldBluringRadius instead of SafeRcp().
-			SignalConvolutionBluringRadius = Sample.WorldBluringRadius * SafeRcp(Sample.AccumulatedSampleCount);
+			SignalConvolutionBluringRadius = Sample.WorldBluringRadius * SafeRcp(Sample.SampleCount);
 		}
 
 		// Compute the final world distance to use for bilateral rejection.
@@ -356,7 +356,7 @@ void AccumulateSampledMultiplexedSignals(
 
 		// TODO: bIsOutsideFrustum could afect number of samples for DRB.
 		FLATTEN
-		if (SampleInfos.Sample.AccumulatedSampleCount != 0 && !bIsOutsideFrustum)
+		if (SampleInfos.Sample.SampleCount != 0 && !bIsOutsideFrustum)
 		{
 			AccumulateSample(
 				/* inout */ Accumulators.Array[SignalMultiplexId],
@@ -380,7 +380,7 @@ void SampleAndAccumulateMultiplexedSignals(
 	float KernelSampleWeight,
 	const bool bForceSample)
 {
-	FSSFSampleSceneInfos SampleSceneMetadata;
+	FSSDSampleSceneInfos SampleSceneMetadata;
 	FSSDSignalArray MultiplexedSamples;
 	
 	// Stores in SGPR whether this sample is outside the viewport, to avoid VGPR pressure to keep SampleBufferUV after texture fetches.
@@ -417,7 +417,7 @@ void SampleAndAccumulateMultiplexedSignalsPair(
 	float2 SampleBufferUV[2],
 	float KernelSampleWeight)
 {
-	FSSFSampleSceneInfos SampleSceneMetadata[2];
+	FSSDSampleSceneInfos SampleSceneMetadata[2];
 	FSSDSignalArray MultiplexedSamples[2];
 	bool bIsOutsideFrustum[2];
 
@@ -879,10 +879,10 @@ void ConvolveStackowiakKernel(
 				if (SampleId == NextClusterBoundary && (SampleBatchId % 2) == 0)
 				{
 					// Compute the number samples that have been accumulated for this sample.
-					uint AccumulatedSampleCountForCluster = min(CurrentRingId * StocasticSamplesPerCluster, SampleCount - SampleId);
+					uint SampleCountForCluster = min(CurrentRingId * StocasticSamplesPerCluster, SampleCount - SampleId);
 
 					// Digest all acumulators.
-					DijestAccumulatedClusterSamples(/* inout */ Accumulators, CurrentRingId, AccumulatedSampleCountForCluster);
+					DijestAccumulatedClusterSamples(/* inout */ Accumulators, CurrentRingId, SampleCountForCluster);
 				
 					BRANCH
 					if (!KernelConfig.bSampleKernelCenter && SampleId == 1)
diff --git a/Engine/Shaders/Private/ScreenSpaceDenoise/SSDTemporalAccumulation.usf b/Engine/Shaders/Private/ScreenSpaceDenoise/SSDTemporalAccumulation.usf
--- a/Engine/Shaders/Private/ScreenSpaceDenoise/SSDTemporalAccumulation.usf
+++ b/Engine/Shaders/Private/ScreenSpaceDenoise/SSDTemporalAccumulation.usf
@@ -288,7 +288,7 @@ RWTexture2D<float4>	DebugOutput;
 
 //------------------------------------------------------- FUNCTIONS
 
-FSSFSignalSample TransformInputBufferForPreRejection(FSSFSignalSample Sample)
+FSSDSignalSample TransformInputBufferForPreRejection(FSSDSignalSample Sample)
 {
 	return TransformSignal(
 		Sample,
@@ -296,7 +296,7 @@ FSSFSignalSample TransformInputBufferForPreRejection(FSSFSignalSample Sample)
 		/* DestConfig = */ CONFIG_HISTORY_REJECTION_COLOR_SPACE);
 }
 
-FSSFSignalSample TransformSignalForPostRejection(FSSFSignalSample Sample)
+FSSDSignalSample TransformSignalForPostRejection(FSSDSignalSample Sample)
 {
 	return TransformSignal(
 		Sample,
@@ -346,7 +346,7 @@ void TemporallyAccumulate(
 
 
 	// Sample current frame data.
-	FSSFSampleSceneInfos RefSceneMetadata = FetchCurrentSceneInfos(ScreenPosition, SceneBufferUV);
+	FSSDSampleSceneInfos RefSceneMetadata = FetchCurrentSceneInfos(ScreenPosition, SceneBufferUV);
 	
 	FSSDSignalArray CurrentFrameSamples = SampleMultiplexedSignals(
 		SignalInput_Textures_0,
@@ -440,14 +440,14 @@ void TemporallyAccumulate(
 			FIX_UNROLL(CONFIG_SIGNAL_BATCH_SIZE)
 			for (uint BatchedSignalId = 0; BatchedSignalId < CONFIG_SIGNAL_BATCH_SIZE; BatchedSignalId++)
 			{
-				FSSFSignalSample CurrentFrameSample = CurrentFrameSamples.Array[BatchedSignalId];
-				FSSFSignalSample HistorySample = HistorySamples.Array[BatchedSignalId];
+				FSSDSignalSample CurrentFrameSample = CurrentFrameSamples.Array[BatchedSignalId];
+				FSSDSignalSample HistorySample = HistorySamples.Array[BatchedSignalId];
 
 				// Override the previous frame's inverse frequency with the current frame's one, because don't know how much it changed.
 				HistorySamples.Array[BatchedSignalId].WorldBluringRadius = (
-					HistorySample.AccumulatedSampleCount *
+					HistorySample.SampleCount *
 					CurrentFrameSample.WorldBluringRadius *
-					SafeRcp(CurrentFrameSample.AccumulatedSampleCount));
+					SafeRcp(CurrentFrameSample.SampleCount));
 
 				// Applies correction when pre exposure change happen.
 				HistorySamples.Array[BatchedSignalId].SceneColor.rgb *= HistoryPreExposureCorrection;
@@ -524,7 +524,7 @@ void TemporallyAccumulate(
 				if (MAX_SIGNAL_BATCH_SIZE == 1)
 				{
 					// TODO: missing && abs(PixelOffset.x) <= 1 && abs(PixelOffset.y) <= 1;
-					KernelConfig.bForceAllAccumulation = CurrentFrameSamples.Array[0].AccumulatedSampleCount == 0;
+					KernelConfig.bForceAllAccumulation = CurrentFrameSamples.Array[0].SampleCount == 0;
 				}
 			}
 		}
@@ -550,7 +550,7 @@ void TemporallyAccumulate(
 			{
 				const uint BatchedSignalId = ComputeSignalBatchIdFromSignalMultiplexId(KernelConfig, SignalMultiplexId);
 
-				FSSFSignalSample TransformedCurrentFrame = TransformSignalSampleForAccumulation(
+				FSSDSignalSample TransformedCurrentFrame = TransformSignalSampleForAccumulation(
 					KernelConfig,
 					SignalMultiplexId,
 					NormalizeToOneSample(CurrentFrameSamples.Array[SignalMultiplexId]));
@@ -563,7 +563,7 @@ void TemporallyAccumulate(
 				if (KernelConfig.BilateralDistanceComputation == SIGNAL_WORLD_FREQUENCY_PRECOMPUTED_BLURING_RADIUS)
 				{
 					// TODO: this is ineficient, could fetch the normalised WorldBluringRadius instead of SafeRcp().
-					SampleInfos.InvFrequency = TransformedCurrentFrame.WorldBluringRadius * SafeRcp(TransformedCurrentFrame.AccumulatedSampleCount);
+					SampleInfos.InvFrequency = TransformedCurrentFrame.WorldBluringRadius * SafeRcp(TransformedCurrentFrame.SampleCount);
 				}
 
 				AccumulateSample(
@@ -594,20 +594,20 @@ void TemporallyAccumulate(
 		FIX_UNROLL(CONFIG_SIGNAL_BATCH_SIZE)
 		for (uint BatchedSignalId = 0; BatchedSignalId < CONFIG_SIGNAL_BATCH_SIZE; BatchedSignalId++)
 		{
-			FSSFSignalSample NeighborMin = SignalAccumulators.Array[BatchedSignalId].Min;
-			FSSFSignalSample NeighborMax = SignalAccumulators.Array[BatchedSignalId].Max;
+			FSSDSignalSample NeighborMin = SignalAccumulators.Array[BatchedSignalId].Min;
+			FSSDSignalSample NeighborMax = SignalAccumulators.Array[BatchedSignalId].Max;
 		
-			FSSFSignalSample NeighborMoment1 = CreateSignalSampleFromScalarValue(0.0);
-			FSSFSignalSample NeighborMoment2 = CreateSignalSampleFromScalarValue(0.0);
+			FSSDSignalSample NeighborMoment1 = CreateSignalSampleFromScalarValue(0.0);
+			FSSDSignalSample NeighborMoment2 = CreateSignalSampleFromScalarValue(0.0);
 			#if CONFIG_REJECTION_INPUT_MODE == REJECTION_INPUT_MODE_1UNNORMALIZED
 			{
-				float NormalizeFactor = SafeRcp(SignalAccumulators.Array[BatchedSignalId].Moment1.AccumulatedSampleCount);
+				float NormalizeFactor = SafeRcp(SignalAccumulators.Array[BatchedSignalId].Moment1.SampleCount);
 				NeighborMoment1 = Mul(SignalAccumulators.Array[BatchedSignalId].Moment1, NormalizeFactor);
 				NeighborMoment2 = Mul(SignalAccumulators.Array[BatchedSignalId].Moment2, NormalizeFactor);
 			}
 			#elif CONFIG_REJECTION_INPUT_MODE == REJECTION_INPUT_MODE_2PRETRANSFORMED_MOMMENTS
 			{
-				float NormalizeFactor = SafeRcp(SignalAccumulators.Array[BatchedSignalId * 2 + 0].Moment1.AccumulatedSampleCount);
+				float NormalizeFactor = SafeRcp(SignalAccumulators.Array[BatchedSignalId * 2 + 0].Moment1.SampleCount);
 				NeighborMoment1 = Mul(SignalAccumulators.Array[BatchedSignalId * 2 + 0].Moment1, NormalizeFactor);
 				NeighborMoment2 = Mul(SignalAccumulators.Array[BatchedSignalId * 2 + 1].Moment1, NormalizeFactor);
 			}
@@ -616,14 +616,14 @@ void TemporallyAccumulate(
 			#endif
 		
 			#if CONFIG_SIGNAL_PROCESSING == SIGNAL_PROCESSING_REFLECTIONS && 0
-				FSSFSignalSample Temp = TransformSignalForPostRejection(NeighborMoment1);
+				FSSDSignalSample Temp = TransformSignalForPostRejection(NeighborMoment1);
 				DebugOutput[DispatchThreadId] = float4(
 					Temp.SceneColor.rgb,
 					0);
 			#endif
 				
-			FSSFSignalSample CurrentFrameSample = CurrentFrameSamples.Array[BatchedSignalId];
-			FSSFSignalSample HistorySample = HistorySamples.Array[BatchedSignalId];
+			FSSDSignalSample CurrentFrameSample = CurrentFrameSamples.Array[BatchedSignalId];
+			FSSDSignalSample HistorySample = HistorySamples.Array[BatchedSignalId];
 
 			// Clamp the history.
 			#if CONFIG_HISTORY_REJECTION == HISTORY_REJECTION_VAR_BOUNDARIES
@@ -631,32 +631,32 @@ void TemporallyAccumulate(
 				const float StdDevMultiplier = 1.25;
 
 				// StdDev = sqrt(abs(NeighborSquareAvg - NeighborAvg * NeighborAvg));
-				FSSFSignalSample StdDev = Sqrt(Abs(Subtract(NeighborMoment2, Power(NeighborMoment1, 2))));
+				FSSDSignalSample StdDev = Sqrt(Abs(Subtract(NeighborMoment2, Power(NeighborMoment1, 2))));
 				NeighborMin = Add(NeighborMoment1, Mul(StdDev, -StdDevMultiplier));
 				NeighborMax = Add(NeighborMoment1, Mul(StdDev,  StdDevMultiplier));
 
 				if (0)
 				{
-					FSSFSignalSample QuantizationErrorMin = Mul(NeighborMoment1, 1 - SafeRcp(HistorySample.AccumulatedSampleCount));
-					FSSFSignalSample QuantizationErrorMax = Mul(NeighborMoment1, 1 + SafeRcp(HistorySample.AccumulatedSampleCount));
+					FSSDSignalSample QuantizationErrorMin = Mul(NeighborMoment1, 1 - SafeRcp(HistorySample.SampleCount));
+					FSSDSignalSample QuantizationErrorMax = Mul(NeighborMoment1, 1 + SafeRcp(HistorySample.SampleCount));
 
 					NeighborMin = Min(NeighborMin, QuantizationErrorMin);
 					NeighborMax = Max(NeighborMax, QuantizationErrorMax);
 				}
 
 				// Transform the history so that it in correct component spaces, and normalized as the clamping box.
-				FSSFSignalSample NormalizedHistorySample = NormalizeToOneSample(HistorySample);
-				FSSFSignalSample TransformedHistorySample = TransformInputBufferForPreRejection(NormalizedHistorySample);
+				FSSDSignalSample NormalizedHistorySample = NormalizeToOneSample(HistorySample);
+				FSSDSignalSample TransformedHistorySample = TransformInputBufferForPreRejection(NormalizedHistorySample);
 
 				// Clamp the history.
-				FSSFSignalSample ClampedTransformedHistorySample = Clamp(TransformedHistorySample, NeighborMin, NeighborMax);
+				FSSDSignalSample ClampedTransformedHistorySample = Clamp(TransformedHistorySample, NeighborMin, NeighborMax);
 
 				// Transform the history back to linear component spaces.
-				FSSFSignalSample ClampedHistorySample = TransformSignalForPostRejection(ClampedTransformedHistorySample);
+				FSSDSignalSample ClampedHistorySample = TransformSignalForPostRejection(ClampedTransformedHistorySample);
 
 				// Reweight the anti ghosted history.
 				{
-					FSSFSignalSample RejectedDiff = Abs(Subtract(ClampedTransformedHistorySample, TransformedHistorySample));
+					FSSDSignalSample RejectedDiff = Abs(Subtract(ClampedTransformedHistorySample, TransformedHistorySample));
 
 					// Compute how much this history get changed.
 					float RejectionFactor = 0.0;
@@ -671,33 +671,33 @@ void TemporallyAccumulate(
 					}
 					#elif CONFIG_SIGNAL_PROCESSING == SIGNAL_PROCESSING_MONOCHROMATIC_PENUMBRA || CONFIG_SIGNAL_PROCESSING == SIGNAL_PROCESSING_AO || CONFIG_SIGNAL_PROCESSING == SIGNAL_PROCESSING_GI_AND_AO
 					{
-						RejectionFactor = abs(ClampedTransformedHistorySample.LeavingRayCount - TransformedHistorySample.LeavingRayCount);
+						RejectionFactor = abs(ClampedTransformedHistorySample.MissCount - TransformedHistorySample.MissCount);
 					}
 					#else
 						#error Unsupported signal rejection.
 					#endif
 			
 					// Compute an initial history weight as if we had removed sample.
-					float FinalHistoryWeight = HistorySample.AccumulatedSampleCount * saturate(1 - RejectionFactor);
+					float FinalHistoryWeight = HistorySample.SampleCount * saturate(1 - RejectionFactor);
 
 					// When doing the accumulation before rejection, need to make sure the input weight get through.
 					if (!bPostRejectionBlending)
 					{
-						FinalHistoryWeight = max(FinalHistoryWeight, CurrentFrameSample.AccumulatedSampleCount);
+						FinalHistoryWeight = max(FinalHistoryWeight, CurrentFrameSample.SampleCount);
 					}
 
 					// When doing upscaling, may have an invalid input sample.
-					FinalHistoryWeight = max(FinalHistoryWeight, NeighborMoment1.AccumulatedSampleCount * 0.1);
+					FinalHistoryWeight = max(FinalHistoryWeight, NeighborMoment1.SampleCount * 0.1);
 			
 					// Apply the history weight.
 					HistorySample = Mul(ClampedHistorySample, FinalHistoryWeight);
-					HistorySample.AccumulatedSampleCount = FinalHistoryWeight;
+					HistorySample.SampleCount = FinalHistoryWeight;
 				}
 			}
 			#elif CONFIG_HISTORY_REJECTION == HISTORY_REJECTION_MINMAX_BOUNDARIES
 			{
 				float RejectedSampleCount = 0;
-				HistorySample = Mul(TransformSignalForPostRejection(Clamp(TransformInputBufferForPreRejection(NormalizeToOneSample(HistorySample)), NeighborMin, NeighborMax)), HistorySample.AccumulatedSampleCount - RejectedSampleCount);
+				HistorySample = Mul(TransformSignalForPostRejection(Clamp(TransformInputBufferForPreRejection(NormalizeToOneSample(HistorySample)), NeighborMin, NeighborMax)), HistorySample.SampleCount - RejectedSampleCount);
 			}
 			#endif
 
@@ -705,7 +705,7 @@ void TemporallyAccumulate(
 			if (1)
 			{
 				float NeightborMinBluringRadius = SignalAccumulators.Array[BatchedSignalId].MinInvFrequency;
-				CurrentFrameSample.WorldBluringRadius = min(CurrentFrameSample.WorldBluringRadius, NeightborMinBluringRadius * CurrentFrameSample.AccumulatedSampleCount); 
+				CurrentFrameSample.WorldBluringRadius = min(CurrentFrameSample.WorldBluringRadius, NeightborMinBluringRadius * CurrentFrameSample.SampleCount); 
 			
 				#if CONFIG_SIGNAL_PROCESSING == SIGNAL_PROCESSING_MONOCHROMATIC_PENUMBRA && 0
 					DebugOutput[DispatchThreadId] = float4(
@@ -731,18 +731,18 @@ void TemporallyAccumulate(
 
 		const uint BatchedSignalId = 0;
 	
-		FSSFSignalSample CurrentFrameSample = CurrentFrameSamples.Array[BatchedSignalId];
-		FSSFSignalSample HistorySample = HistorySamples.Array[BatchedSignalId];
+		FSSDSignalSample CurrentFrameSample = CurrentFrameSamples.Array[BatchedSignalId];
+		FSSDSignalSample HistorySample = HistorySamples.Array[BatchedSignalId];
 
 		// Construct the current frame sample of the temporal analysis.
 		FSSDSignalArray CurrentTemporalAnalysis = CreateSignalArrayFromScalarValue(0.0);
 		{
 			float Luma = GetSignalLuma4(CurrentFrameSample, STANDARD_BUFFER_COLOR_SPACE);
-			float SampleCount = CurrentFrameSample.AccumulatedSampleCount; // TODO: 1 instead?
+			float SampleCount = CurrentFrameSample.SampleCount; // TODO: 1 instead?
 
-			CurrentTemporalAnalysis.Array[1].AccumulatedSampleCount = SampleCount;
+			CurrentTemporalAnalysis.Array[1].SampleCount = SampleCount;
 			CurrentTemporalAnalysis.Array[1].SceneColor.x = Luma * SampleCount;
-			CurrentTemporalAnalysis.Array[2].AccumulatedSampleCount = SampleCount;
+			CurrentTemporalAnalysis.Array[2].SampleCount = SampleCount;
 			CurrentTemporalAnalysis.Array[2].SceneColor.x = (Luma * Luma) * SampleCount;
 		}
 
@@ -751,8 +751,8 @@ void TemporallyAccumulate(
 		{
 			const float MinimumAnalasisHisteresis = 0.3; // TODO: expose cvar
 
-			float PreviousFrameWeight = min(HistorySample.AccumulatedSampleCount, HistorySamples.Array[1].AccumulatedSampleCount * rcp(MinimumAnalasisHisteresis) - CurrentFrameSample.AccumulatedSampleCount);
-			float PreviousFrameMultiplier = HistorySample.AccumulatedSampleCount > 0 ? PreviousFrameWeight / HistorySample.AccumulatedSampleCount : 0;
+			float PreviousFrameWeight = min(HistorySample.SampleCount, HistorySamples.Array[1].SampleCount * rcp(MinimumAnalasisHisteresis) - CurrentFrameSample.SampleCount);
+			float PreviousFrameMultiplier = HistorySample.SampleCount > 0 ? PreviousFrameWeight / HistorySample.SampleCount : 0;
 			
 			UNROLL
 			for (uint SignalMultiplexId = 1; SignalMultiplexId < 3; SignalMultiplexId++)
@@ -764,7 +764,7 @@ void TemporallyAccumulate(
 		// Compute the histeresis of the main signal.
 		float MainHisteresiss = 0;
 		{
-			float InvSampleCount = rcp(TemporalAnalysis.Array[1].AccumulatedSampleCount);
+			float InvSampleCount = rcp(TemporalAnalysis.Array[1].SampleCount);
 
 			float AvgLuma = TemporalAnalysis.Array[2].SceneColor.x * InvSampleCount;
 			float VarianceLuma = TemporalAnalysis.Array[2].SceneColor.x * InvSampleCount - AvgLuma * AvgLuma;
@@ -777,8 +777,8 @@ void TemporallyAccumulate(
 
 		// Accumulate the actual history.
 		{
-			float PreviousFrameWeight = min(HistorySample.AccumulatedSampleCount, CurrentFrameSample.AccumulatedSampleCount * rcp(MainHisteresiss) - CurrentFrameSample.AccumulatedSampleCount);
-			float PreviousFrameMultiplier = HistorySample.AccumulatedSampleCount > 0 ? PreviousFrameWeight / HistorySample.AccumulatedSampleCount : 0;
+			float PreviousFrameWeight = min(HistorySample.SampleCount, CurrentFrameSample.SampleCount * rcp(MainHisteresiss) - CurrentFrameSample.SampleCount);
+			float PreviousFrameMultiplier = HistorySample.SampleCount > 0 ? PreviousFrameWeight / HistorySample.SampleCount : 0;
 			
 			HistorySample = Add(Mul(HistorySample, PreviousFrameMultiplier), CurrentFrameSample);
 		}
@@ -797,11 +797,11 @@ void TemporallyAccumulate(
 		UNROLL
 		for (uint BatchedSignalId = 0; BatchedSignalId < MAX_SIGNAL_BATCH_SIZE; BatchedSignalId++)
 		{
-			FSSFSignalSample CurrentFrameSample = CurrentFrameSamples.Array[BatchedSignalId];
-			FSSFSignalSample HistorySample = HistorySamples.Array[BatchedSignalId];
+			FSSDSignalSample CurrentFrameSample = CurrentFrameSamples.Array[BatchedSignalId];
+			FSSDSignalSample HistorySample = HistorySamples.Array[BatchedSignalId];
 
-			float PreviousFrameWeight = min(HistorySample.AccumulatedSampleCount, TargetedSampleCount - CurrentFrameSample.AccumulatedSampleCount);
-			float PreviousFrameMultiplier = HistorySample.AccumulatedSampleCount > 0 ? PreviousFrameWeight / HistorySample.AccumulatedSampleCount : 0;
+			float PreviousFrameWeight = min(HistorySample.SampleCount, TargetedSampleCount - CurrentFrameSample.SampleCount);
+			float PreviousFrameMultiplier = HistorySample.SampleCount > 0 ? PreviousFrameWeight / HistorySample.SampleCount : 0;
 
 			// TODO: use WeightedLerp instead for VALU perf.
 			// Pre transformation of the signal.
@@ -826,8 +826,8 @@ void TemporallyAccumulate(
 			// Ignore the bluring radius of previous frame frequencies as it gets clamped anyway, just take the current frame.
 			HistorySample.WorldBluringRadius = (
 				CurrentFrameSample.WorldBluringRadius *
-				SafeRcp(CurrentFrameSample.AccumulatedSampleCount) *
-				HistorySample.AccumulatedSampleCount);
+				SafeRcp(CurrentFrameSample.SampleCount) *
+				HistorySample.SampleCount);
 
 			HistorySamples.Array[BatchedSignalId] = HistorySample;
 		}
