Speeds up the AO denoiser by using a history rejection pre-convolution

diff --git a/Engine/Shaders/Private/ScreenSpaceDenoise/SSDDefinitions.ush b/Engine/Shaders/Private/ScreenSpaceDenoise/SSDDefinitions.ush
--- a/Engine/Shaders/Private/ScreenSpaceDenoise/SSDDefinitions.ush
+++ b/Engine/Shaders/Private/ScreenSpaceDenoise/SSDDefinitions.ush
@@ -107,10 +107,13 @@
 
 
 	/** Buffer layout given by the ray generation shader for ambient occlusion. */
-	#define SIGNAL_BUFFER_LAYOUT_AO_INPUT 5
+	#define SIGNAL_BUFFER_LAYOUT_AO_INPUT 0xA000
 
 	/** Internal buffer layout for ambient occlusion. */
-	#define SIGNAL_BUFFER_LAYOUT_AO_HISTORY 6
+	#define SIGNAL_BUFFER_LAYOUT_AO_REJECTION 0xA001
+
+	/** Internal buffer layout for ambient occlusion. */
+	#define SIGNAL_BUFFER_LAYOUT_AO_HISTORY 0xA002
 
 
 	/** Buffer layout given by the ray generation shader for global illumination. */
diff --git a/Engine/Shaders/Private/ScreenSpaceDenoise/SSDSignalCore.ush b/Engine/Shaders/Private/ScreenSpaceDenoise/SSDSignalCore.ush
--- a/Engine/Shaders/Private/ScreenSpaceDenoise/SSDSignalCore.ush
+++ b/Engine/Shaders/Private/ScreenSpaceDenoise/SSDSignalCore.ush
@@ -362,6 +362,25 @@ void DecodeMultiplexedSignals(
 		OutSamples.Array[0].MissCount = OutSamples.Array[0].SampleCount * RawSample[0].r;
 		OutSamples.Array[0].ClosestHitDistance = OutSamples.Array[0].SampleCount * RawSample[1].r;
 	}
+	else if (SignalBufferLayout == SIGNAL_BUFFER_LAYOUT_AO_REJECTION)
+	{
+		if (bNormalizeSample)
+		{
+			OutSamples.Array[0].SampleCount = RawSample[0].g > 0 ? 1 : 0;
+			OutSamples.Array[0].MissCount = RawSample[0].g > 0 ? RawSample[0].r : 0;
+			
+			OutSamples.Array[1].SampleCount = RawSample[0].a > 0 ? 1 : 0;
+			OutSamples.Array[1].MissCount = RawSample[0].a > 0 ? RawSample[0].b : 0;
+		}
+		else
+		{
+			OutSamples.Array[0].SampleCount = RawSample[0].g;
+			OutSamples.Array[0].MissCount = RawSample[0].g * RawSample[0].r;
+			
+			OutSamples.Array[1].SampleCount = RawSample[0].a;
+			OutSamples.Array[1].MissCount = RawSample[0].a * RawSample[0].b;
+		}
+	}
 	else if (SignalBufferLayout == SIGNAL_BUFFER_LAYOUT_AO_HISTORY)
 	{
 		if (bNormalizeSample)
@@ -566,6 +585,18 @@ void EncodeMultiplexedSignals(
 			OutBufferCount = 3;
 		}
 	}
+	else if (SignalBufferLayout == SIGNAL_BUFFER_LAYOUT_AO_REJECTION)
+	{
+		// Outputs number of ray miss in red to be compatible as a SSAO output.
+		OutRawSample[0].r = Sample[0].SampleCount > 0 ? Sample[0].MissCount / Sample[0].SampleCount : 1.0;
+		OutRawSample[0].g = Sample[0].SampleCount;
+		
+		if (MultiplexCount == 2)
+		{
+			OutRawSample[0].b = Sample[1].SampleCount > 0 ? Sample[1].MissCount / Sample[1].SampleCount : 1.0;
+			OutRawSample[0].a = Sample[1].SampleCount;
+		}
+	}
 	else if (SignalBufferLayout == SIGNAL_BUFFER_LAYOUT_AO_HISTORY)
 	{
 		// Outputs number of ray miss in red to be compatible as a SSAO output.
diff --git a/Engine/Shaders/Private/ScreenSpaceDenoise/SSDSpatialAccumulation.usf b/Engine/Shaders/Private/ScreenSpaceDenoise/SSDSpatialAccumulation.usf
--- a/Engine/Shaders/Private/ScreenSpaceDenoise/SSDSpatialAccumulation.usf
+++ b/Engine/Shaders/Private/ScreenSpaceDenoise/SSDSpatialAccumulation.usf
@@ -194,6 +194,13 @@
 		#define CONFIG_BILATERAL_DISTANCE_COMPUTATION SIGNAL_WORLD_FREQUENCY_HIT_DISTANCE
 		#define CONFIG_OUTPUT_MODE          OUTPUT_MODE_DRB
 
+	#elif DIM_STAGE == STAGE_REJECTION_PRE_CONVOLUTION
+		// Input and output layout.
+		#define CONFIG_SIGNAL_INPUT_LAYOUT  SIGNAL_BUFFER_LAYOUT_AO_HISTORY
+		#define CONFIG_SIGNAL_OUTPUT_LAYOUT SIGNAL_BUFFER_LAYOUT_AO_REJECTION
+
+		#define CONFIG_BILATERAL_DISTANCE_COMPUTATION SIGNAL_WORLD_FREQUENCY_HIT_DISTANCE
+
 	#elif DIM_STAGE == STAGE_POST_FILTERING
 		// Input and output layout.
 		#define CONFIG_SIGNAL_INPUT_LAYOUT  SIGNAL_BUFFER_LAYOUT_AO_HISTORY
diff --git a/Engine/Shaders/Private/ScreenSpaceDenoise/SSDTemporalAccumulation.usf b/Engine/Shaders/Private/ScreenSpaceDenoise/SSDTemporalAccumulation.usf
--- a/Engine/Shaders/Private/ScreenSpaceDenoise/SSDTemporalAccumulation.usf
+++ b/Engine/Shaders/Private/ScreenSpaceDenoise/SSDTemporalAccumulation.usf
@@ -102,6 +102,16 @@
 	#define CONFIG_VIEW_POSITION_BASED_BILATERAL 2
 	#define CONFIG_NORMAL_BASED_BILATERAL 1
 
+	// Use pre transformed rejection buffer that contains pre transformed momment 1 & 2.
+	#if 1
+		#define CONFIG_USE_REJECTION_BUFFER               1
+		#define CONFIG_SIGNAL_HISTORY_REJECTION_LAYOUT    SIGNAL_BUFFER_LAYOUT_AO_REJECTION
+		#define CONFIG_REJECTION_TEXTURE_COUNT            1
+
+		#define CONFIG_REJECTION_INPUT_MODE               REJECTION_INPUT_MODE_2PRETRANSFORMED_MOMMENTS
+		#define CONFIG_REJECTION_SAMPLE_SET               REJECTION_SAMPLE_SET_SOBEK2018
+	#endif
+
 #elif CONFIG_SIGNAL_PROCESSING == SIGNAL_PROCESSING_GI_AND_AO
 	// Input and output layout.
 	#define CONFIG_SIGNAL_INPUT_LAYOUT   SIGNAL_BUFFER_LAYOUT_GI_AND_AO_RECONSTRUCTION
diff --git a/Engine/Source/Runtime/Renderer/Private/ScreenSpaceDenoise.cpp b/Engine/Source/Runtime/Renderer/Private/ScreenSpaceDenoise.cpp
--- a/Engine/Source/Runtime/Renderer/Private/ScreenSpaceDenoise.cpp
+++ b/Engine/Source/Runtime/Renderer/Private/ScreenSpaceDenoise.cpp
@@ -171,7 +171,8 @@ static bool SignalUsesRejectionPreConvolution(ESignalProcessing SignalProcessing
 {
 	return (
 		//SignalProcessing == ESignalProcessing::MonochromaticPenumbra ||
-		SignalProcessing == ESignalProcessing::Reflections);
+		SignalProcessing == ESignalProcessing::Reflections ||
+		SignalProcessing == ESignalProcessing::AmbientOcclusion);
 }
 
 /** Returns whether a signal processing uses a history rejection pre convolution pass. */
@@ -311,7 +312,7 @@ const TCHAR* const kRejectionPreConvolutionResourceNames[] = {
 	nullptr,
 
 	// AmbientOcclusion
-	nullptr,
+	TEXT("AORejectionPreConvolution0"),
 	nullptr,
 	nullptr,
 	nullptr,
@@ -1033,6 +1034,10 @@ static void DenoiseSignalAtConstantPixelDensity(
 					RejectionSignalProcessingDescs[2].Format = PF_FloatRGBA;
 					RejectionTextureCount = 3;
 				}
+				else if (Settings.SignalProcessing == ESignalProcessing::AmbientOcclusion)
+				{
+					RejectionSignalProcessingDescs[0].Format = PF_FloatRGBA;
+				}
 				else
 				{
 					check(0);
