Converts the denoiser to UNROLL_N

diff --git a/Engine/Shaders/Private/ScreenSpaceDenoise/SSDCommon.ush b/Engine/Shaders/Private/ScreenSpaceDenoise/SSDCommon.ush
--- a/Engine/Shaders/Private/ScreenSpaceDenoise/SSDCommon.ush
+++ b/Engine/Shaders/Private/ScreenSpaceDenoise/SSDCommon.ush
@@ -14,9 +14,6 @@
 #include "SSDDefinitions.ush"
 
 
-#define FIX_UNROLL(N) [unroll(N)]
-
-
 //------------------------------------------------------- SHARED SHADER PARAMETERS
 
 float4x4  PrevScreenToTranslatedWorld;
diff --git a/Engine/Shaders/Private/ScreenSpaceDenoise/SSDSignalArray.ush b/Engine/Shaders/Private/ScreenSpaceDenoise/SSDSignalArray.ush
--- a/Engine/Shaders/Private/ScreenSpaceDenoise/SSDSignalArray.ush
+++ b/Engine/Shaders/Private/ScreenSpaceDenoise/SSDSignalArray.ush
@@ -8,7 +8,7 @@
 FSSDSignalArray Mul(FSSDSignalArray Samples, float Scalar[MAX_SIGNAL_BATCH_SIZE])
 {
 	FSSDSignalArray OutSamples;
-	FIX_UNROLL(MAX_SIGNAL_BATCH_SIZE)
+	UNROLL_N(MAX_SIGNAL_BATCH_SIZE)
 	for (uint BatchedSignalId = 0; BatchedSignalId < MAX_SIGNAL_BATCH_SIZE; BatchedSignalId++)
 	{
 		OutSamples.Array[BatchedSignalId] = Mul(Samples.Array[BatchedSignalId], Scalar[BatchedSignalId]);
@@ -19,7 +19,7 @@ FSSDSignalArray Mul(FSSDSignalArray Samples, float Scalar[MAX_SIGNAL_BATCH_SIZE]
 FSSDSignalArray Add(FSSDSignalArray SamplesA, FSSDSignalArray SamplesB)
 {
 	FSSDSignalArray OutSamples;
-	FIX_UNROLL(MAX_SIGNAL_BATCH_SIZE)
+	UNROLL_N(MAX_SIGNAL_BATCH_SIZE)
 	for (uint BatchedSignalId = 0; BatchedSignalId < MAX_SIGNAL_BATCH_SIZE; BatchedSignalId++)
 	{
 		OutSamples.Array[BatchedSignalId] = Add(SamplesA.Array[BatchedSignalId], SamplesB.Array[BatchedSignalId]);
@@ -30,7 +30,7 @@ FSSDSignalArray Add(FSSDSignalArray SamplesA, FSSDSignalArray SamplesB)
 FSSDSignalArray Minus(FSSDSignalArray Samples)
 {
 	FSSDSignalArray OutSamples;
-	FIX_UNROLL(MAX_SIGNAL_BATCH_SIZE)
+	UNROLL_N(MAX_SIGNAL_BATCH_SIZE)
 	for (uint BatchedSignalId = 0; BatchedSignalId < MAX_SIGNAL_BATCH_SIZE; BatchedSignalId++)
 	{
 		OutSamples.Array[BatchedSignalId] = Minus(Samples.Array[BatchedSignalId]);
@@ -41,7 +41,7 @@ FSSDSignalArray Minus(FSSDSignalArray Samples)
 FSSDSignalArray Subtract(FSSDSignalArray SamplesA, FSSDSignalArray SamplesB)
 {
 	FSSDSignalArray OutSamples;
-	FIX_UNROLL(MAX_SIGNAL_BATCH_SIZE)
+	UNROLL_N(MAX_SIGNAL_BATCH_SIZE)
 	for (uint BatchedSignalId = 0; BatchedSignalId < MAX_SIGNAL_BATCH_SIZE; BatchedSignalId++)
 	{
 		OutSamples.Array[BatchedSignalId] = Subtract(SamplesA.Array[BatchedSignalId], SamplesB.Array[BatchedSignalId]);
@@ -52,7 +52,7 @@ FSSDSignalArray Subtract(FSSDSignalArray SamplesA, FSSDSignalArray SamplesB)
 FSSDSignalArray Abs(FSSDSignalArray Samples)
 {
 	FSSDSignalArray OutSamples;
-	FIX_UNROLL(MAX_SIGNAL_BATCH_SIZE)
+	UNROLL_N(MAX_SIGNAL_BATCH_SIZE)
 	for (uint BatchedSignalId = 0; BatchedSignalId < MAX_SIGNAL_BATCH_SIZE; BatchedSignalId++)
 	{
 		OutSamples.Array[BatchedSignalId] = Abs(Samples.Array[BatchedSignalId]);
@@ -63,7 +63,7 @@ FSSDSignalArray Abs(FSSDSignalArray Samples)
 FSSDSignalArray Sqrt(FSSDSignalArray Samples)
 {
 	FSSDSignalArray OutSamples;
-	FIX_UNROLL(MAX_SIGNAL_BATCH_SIZE)
+	UNROLL_N(MAX_SIGNAL_BATCH_SIZE)
 	for (uint BatchedSignalId = 0; BatchedSignalId < MAX_SIGNAL_BATCH_SIZE; BatchedSignalId++)
 	{
 		OutSamples.Array[BatchedSignalId] = Sqrt(Samples.Array[BatchedSignalId]);
@@ -74,7 +74,7 @@ FSSDSignalArray Sqrt(FSSDSignalArray Samples)
 FSSDSignalArray Power(FSSDSignalArray Samples, float Exponent)
 {
 	FSSDSignalArray OutSamples;
-	FIX_UNROLL(MAX_SIGNAL_BATCH_SIZE)
+	UNROLL_N(MAX_SIGNAL_BATCH_SIZE)
 	for (uint BatchedSignalId = 0; BatchedSignalId < MAX_SIGNAL_BATCH_SIZE; BatchedSignalId++)
 	{
 		OutSamples.Array[BatchedSignalId] = Power(Samples.Array[BatchedSignalId], Exponent);
@@ -85,7 +85,7 @@ FSSDSignalArray Power(FSSDSignalArray Samples, float Exponent)
 FSSDSignalArray Min(FSSDSignalArray SamplesA, FSSDSignalArray SamplesB)
 {
 	FSSDSignalArray OutSamples;
-	FIX_UNROLL(MAX_SIGNAL_BATCH_SIZE)
+	UNROLL_N(MAX_SIGNAL_BATCH_SIZE)
 	for (uint BatchedSignalId = 0; BatchedSignalId < MAX_SIGNAL_BATCH_SIZE; BatchedSignalId++)
 	{
 		OutSamples.Array[BatchedSignalId] = Min(SamplesA.Array[BatchedSignalId], SamplesB.Array[BatchedSignalId]);
@@ -96,7 +96,7 @@ FSSDSignalArray Min(FSSDSignalArray SamplesA, FSSDSignalArray SamplesB)
 FSSDSignalArray Max(FSSDSignalArray SamplesA, FSSDSignalArray SamplesB)
 {
 	FSSDSignalArray OutSamples;
-	FIX_UNROLL(MAX_SIGNAL_BATCH_SIZE)
+	UNROLL_N(MAX_SIGNAL_BATCH_SIZE)
 	for (uint BatchedSignalId = 0; BatchedSignalId < MAX_SIGNAL_BATCH_SIZE; BatchedSignalId++)
 	{
 		OutSamples.Array[BatchedSignalId] = Max(SamplesA.Array[BatchedSignalId], SamplesB.Array[BatchedSignalId]);
@@ -107,7 +107,7 @@ FSSDSignalArray Max(FSSDSignalArray SamplesA, FSSDSignalArray SamplesB)
 FSSDSignalArray Clamp(FSSDSignalArray Samples, FSSDSignalArray SamplesMin, FSSDSignalArray SamplesMax)
 {
 	FSSDSignalArray OutSamples;
-	FIX_UNROLL(MAX_SIGNAL_BATCH_SIZE)
+	UNROLL_N(MAX_SIGNAL_BATCH_SIZE)
 	for (uint BatchedSignalId = 0; BatchedSignalId < MAX_SIGNAL_BATCH_SIZE; BatchedSignalId++)
 	{
 		OutSamples.Array[BatchedSignalId] = Clamp(Samples.Array[BatchedSignalId], SamplesMin.Array[BatchedSignalId], SamplesMax.Array[BatchedSignalId]);
@@ -118,7 +118,7 @@ FSSDSignalArray Clamp(FSSDSignalArray Samples, FSSDSignalArray SamplesMin, FSSDS
 FSSDSignalArray Lerp(FSSDSignalArray Samples0, FSSDSignalArray Samples1, float Interp[MAX_SIGNAL_BATCH_SIZE])
 {
 	FSSDSignalArray OutSamples;
-	FIX_UNROLL(MAX_SIGNAL_BATCH_SIZE)
+	UNROLL_N(MAX_SIGNAL_BATCH_SIZE)
 	for (uint BatchedSignalId = 0; BatchedSignalId < MAX_SIGNAL_BATCH_SIZE; BatchedSignalId++)
 	{
 		OutSamples.Array[BatchedSignalId] = Lerp(Samples0.Array[BatchedSignalId], Samples1.Array[BatchedSignalId], Interp[BatchedSignalId]);
@@ -129,7 +129,7 @@ FSSDSignalArray Lerp(FSSDSignalArray Samples0, FSSDSignalArray Samples1, float I
 FSSDSignalArray NormalizeToOneSample(FSSDSignalArray Samples)
 {
 	FSSDSignalArray OutSamples;
-	FIX_UNROLL(MAX_SIGNAL_BATCH_SIZE)
+	UNROLL_N(MAX_SIGNAL_BATCH_SIZE)
 	for (uint BatchedSignalId = 0; BatchedSignalId < MAX_SIGNAL_BATCH_SIZE; BatchedSignalId++)
 	{
 		OutSamples.Array[BatchedSignalId] = NormalizeToOneSample(Samples.Array[BatchedSignalId]);
diff --git a/Engine/Shaders/Private/ScreenSpaceDenoise/SSDSignalCore.ush b/Engine/Shaders/Private/ScreenSpaceDenoise/SSDSignalCore.ush
--- a/Engine/Shaders/Private/ScreenSpaceDenoise/SSDSignalCore.ush
+++ b/Engine/Shaders/Private/ScreenSpaceDenoise/SSDSignalCore.ush
@@ -56,7 +56,7 @@ FSSDSignalSample CreateSignalSample()
 FSSDSignalArray CreateSignalArrayFromScalarValue(float Scalar)
 {
 	FSSDSignalArray OutSamples;
-	FIX_UNROLL(MAX_SIGNAL_MULTIPLEXING)
+	UNROLL_N(MAX_SIGNAL_MULTIPLEXING)
 	for (uint BatchedSignalId = 0; BatchedSignalId < MAX_SIGNAL_MULTIPLEXING; BatchedSignalId++)
 	{
 		OutSamples.Array[BatchedSignalId] = CreateSignalSampleFromScalarValue(Scalar);
@@ -211,7 +211,7 @@ void DecodeMultiplexedSignals(
 	}
 	else if (SignalBufferLayout == SIGNAL_BUFFER_LAYOUT_PENUMBRA_INPUT_1SPP)
 	{
-		FIX_UNROLL(MAX_SIGNAL_BATCH_SIZE)
+		UNROLL_N(MAX_SIGNAL_BATCH_SIZE)
 		for (uint BatchSignalId = 0; BatchSignalId < MAX_SIGNAL_BATCH_SIZE; BatchSignalId++)
 		{
 			uint MultiplexId = BatchSignalId;
@@ -222,7 +222,7 @@ void DecodeMultiplexedSignals(
 	}
 	else if (SignalBufferLayout == SIGNAL_BUFFER_LAYOUT_PENUMBRA_INPUT_NSPP)
 	{
-		FIX_UNROLL(MAX_SIGNAL_BATCH_SIZE)
+		UNROLL_N(MAX_SIGNAL_BATCH_SIZE)
 		for (uint BatchSignalId = 0; BatchSignalId < MAX_SIGNAL_BATCH_SIZE; BatchSignalId++)
 		{
 			uint MultiplexId = BatchSignalId;
@@ -234,7 +234,7 @@ void DecodeMultiplexedSignals(
 	}
 	else if (SignalBufferLayout == SIGNAL_BUFFER_LAYOUT_PENUMBRA_INJESTION_1SPP)
 	{
-		FIX_UNROLL(MAX_SIGNAL_BATCH_SIZE)
+		UNROLL_N(MAX_SIGNAL_BATCH_SIZE)
 		for (uint BatchSignalId = 0; BatchSignalId < MAX_SIGNAL_BATCH_SIZE; BatchSignalId++)
 		{
 			uint MultiplexId = BatchSignalId;
@@ -246,7 +246,7 @@ void DecodeMultiplexedSignals(
 	}
 	else if (SignalBufferLayout == SIGNAL_BUFFER_LAYOUT_PENUMBRA_INJESTION_NSPP)
 	{
-		FIX_UNROLL(MAX_SIGNAL_BATCH_SIZE)
+		UNROLL_N(MAX_SIGNAL_BATCH_SIZE)
 		for (uint BatchSignalId = 0; BatchSignalId < MAX_SIGNAL_BATCH_SIZE; BatchSignalId++)
 		{
 			uint MultiplexId = BatchSignalId;
@@ -260,7 +260,7 @@ void DecodeMultiplexedSignals(
 	}
 	//else if (SignalBufferLayout == SIGNAL_BUFFER_LAYOUT_PENUMBRA_RECONSTRUCTION)
 	//{
-	//	FIX_UNROLL(MAX_SIGNAL_MULTIPLEXING)
+	//	UNROLL_N(MAX_SIGNAL_MULTIPLEXING)
 	//	for (uint MultiplexId = 0; MultiplexId < MAX_SIGNAL_MULTIPLEXING; MultiplexId++)
 	//	{
 	//		float2 Channels = float2(RawSample[MultiplexId / 2][2 * (MultiplexId % 2)], RawSample[MultiplexId / 2][2 * (MultiplexId % 2) + 1]);
@@ -279,7 +279,7 @@ void DecodeMultiplexedSignals(
 	//}
 	else if (SignalBufferLayout == SIGNAL_BUFFER_LAYOUT_PENUMBRA_HISTORY || SignalBufferLayout == SIGNAL_BUFFER_LAYOUT_PENUMBRA_RECONSTRUCTION)
 	{
-		FIX_UNROLL(MAX_SIGNAL_BATCH_SIZE)
+		UNROLL_N(MAX_SIGNAL_BATCH_SIZE)
 		for (uint MultiplexId = 0; MultiplexId < MAX_SIGNAL_BATCH_SIZE; MultiplexId++)
 		{
 			float3 Channels = RawSample[MultiplexId].xyz;
@@ -293,7 +293,7 @@ void DecodeMultiplexedSignals(
 	}
 	else if (SignalBufferLayout == SIGNAL_BUFFER_LAYOUT_PENUMBRA_REJECTION)
 	{
-		FIX_UNROLL(MAX_SIGNAL_BATCH_SIZE)
+		UNROLL_N(MAX_SIGNAL_BATCH_SIZE)
 		for (uint BatchSignalId = 0; BatchSignalId < MAX_SIGNAL_BATCH_SIZE; BatchSignalId++)
 		{
 			uint MultiplexId = BatchSignalId * 2;
@@ -448,7 +448,7 @@ void EncodeMultiplexedSignals(
 	out uint OutBufferCount)
 {
 	// Init all raw samples.
-	FIX_UNROLL(MAX_MULTIPLEXED_TEXTURES)
+	UNROLL_N(MAX_MULTIPLEXED_TEXTURES)
 	for (uint i = 0; i < MAX_MULTIPLEXED_TEXTURES; i++)
 		OutRawSample[i] = 0;
 
diff --git a/Engine/Shaders/Private/ScreenSpaceDenoise/SSDSpatialAccumulation.usf b/Engine/Shaders/Private/ScreenSpaceDenoise/SSDSpatialAccumulation.usf
--- a/Engine/Shaders/Private/ScreenSpaceDenoise/SSDSpatialAccumulation.usf
+++ b/Engine/Shaders/Private/ScreenSpaceDenoise/SSDSpatialAccumulation.usf
@@ -523,7 +523,7 @@ void MainCS(
 		KernelConfig.bForceKernelCenterAccumulation = true;
 		
 		{
-			FIX_UNROLL(MAX_SIGNAL_MULTIPLEXING)
+			UNROLL_N(MAX_SIGNAL_MULTIPLEXING)
 			for (uint MultiplexId = 0; MultiplexId < MAX_SIGNAL_MULTIPLEXING; MultiplexId++)
 			{
 				KernelConfig.BufferColorSpace[MultiplexId] = CONFIG_INPUT_COLOR_SPACE;
@@ -543,7 +543,7 @@ void MainCS(
 		{
 			KernelConfig.bMaxWithRefBilateralDistance = true;
 
-			FIX_UNROLL(MAX_SIGNAL_MULTIPLEXING)
+			UNROLL_N(MAX_SIGNAL_MULTIPLEXING)
 			for (uint MultiplexId = 0; MultiplexId < MAX_SIGNAL_MULTIPLEXING; MultiplexId++)
 			{
 				if (KernelConfig.BilateralDistanceComputation == SIGNAL_WORLD_FREQUENCY_PRECOMPUTED_BLURING_RADIUS)
@@ -589,7 +589,7 @@ void MainCS(
 		// When not upscaling, manually force accumulate the sample of the kernel.
 		if (!KernelConfig.bSampleKernelCenter && !KernelConfig.bDescOrder)
 		{
-			FIX_UNROLL(MAX_SIGNAL_MULTIPLEXING)
+			UNROLL_N(MAX_SIGNAL_MULTIPLEXING)
 			for (uint SignalMultiplexId = 0; SignalMultiplexId < MAX_SIGNAL_MULTIPLEXING; SignalMultiplexId++)
 			{
 				const uint BatchedSignalId = ComputeSignalBatchIdFromSignalMultiplexId(KernelConfig, SignalMultiplexId);
@@ -701,7 +701,7 @@ void MainCS(
 
 	// Color processing of the signal to reduce highlight flickering.
 	{
-		FIX_UNROLL(CONFIG_SIGNAL_BATCH_SIZE)
+		UNROLL_N(CONFIG_SIGNAL_BATCH_SIZE)
 		for (uint MultiplexId = 0; MultiplexId < CONFIG_SIGNAL_BATCH_SIZE; MultiplexId++)
 		{
 			#if CONFIG_REJECT_HIGHEST_COLOR
@@ -735,7 +735,7 @@ void MainCS(
 		{
 			MultiplexCount = CONFIG_SIGNAL_BATCH_SIZE;
 			
-			FIX_UNROLL(CONFIG_SIGNAL_BATCH_SIZE)
+			UNROLL_N(CONFIG_SIGNAL_BATCH_SIZE)
 			for (uint MultiplexId = 0; MultiplexId < CONFIG_SIGNAL_BATCH_SIZE; MultiplexId++)
 			{
 				OutputSamples.Array[MultiplexId] = SignalAccumulators.Array[MultiplexId].Moment1;
@@ -750,7 +750,7 @@ void MainCS(
 		{
 			MultiplexCount = 2 * CONFIG_SIGNAL_BATCH_SIZE;
 
-			FIX_UNROLL(CONFIG_SIGNAL_BATCH_SIZE)
+			UNROLL_N(CONFIG_SIGNAL_BATCH_SIZE)
 			for (uint BatchedSignalId = 0; BatchedSignalId < CONFIG_SIGNAL_BATCH_SIZE; BatchedSignalId++)
 			{
 				OutputSamples.Array[BatchedSignalId * 2 + 0] = SignalAccumulators.Array[BatchedSignalId * CONFIG_MULTIPLEXED_SIGNALS_PER_SIGNAL_DOMAIN + 0].Moment1;
@@ -761,7 +761,7 @@ void MainCS(
 		{
 			MultiplexCount = CONFIG_SIGNAL_BATCH_SIZE;
 			
-			FIX_UNROLL(CONFIG_SIGNAL_BATCH_SIZE)
+			UNROLL_N(CONFIG_SIGNAL_BATCH_SIZE)
 			for (uint MultiplexId = 0; MultiplexId < CONFIG_SIGNAL_BATCH_SIZE; MultiplexId++)
 			{
 				OutputSamples.Array[MultiplexId] = SignalAccumulators.Array[MultiplexId].Previous;
@@ -786,7 +786,7 @@ void MainCS(
 	// Clamp the number of sample recorded.
 	if (DIM_STAGE == STAGE_POST_FILTERING)
 	{
-		FIX_UNROLL(CONFIG_SIGNAL_BATCH_SIZE)
+		UNROLL_N(CONFIG_SIGNAL_BATCH_SIZE)
 		for (uint MultiplexId = 0; MultiplexId < CONFIG_SIGNAL_BATCH_SIZE; MultiplexId++)
 		{
 			float CurrentSampleCount = OutputSamples.Array[MultiplexId].SampleCount;
diff --git a/Engine/Shaders/Private/ScreenSpaceDenoise/SSDSpatialKernel.ush b/Engine/Shaders/Private/ScreenSpaceDenoise/SSDSpatialKernel.ush
--- a/Engine/Shaders/Private/ScreenSpaceDenoise/SSDSpatialKernel.ush
+++ b/Engine/Shaders/Private/ScreenSpaceDenoise/SSDSpatialKernel.ush
@@ -36,7 +36,7 @@ FSSDSignalAccumulatorArray CreateSignalAccumulatorArray()
 {
 	FSSDSignalAccumulatorArray Accumulators;
 	
-	FIX_UNROLL(MAX_SIGNAL_MULTIPLEXING)
+	UNROLL_N(MAX_SIGNAL_MULTIPLEXING)
 	for (uint i = 0; i < MAX_SIGNAL_MULTIPLEXING; i++)
 	{
 		Accumulators.Array[i] = CreateSignalAccumulator();
@@ -156,7 +156,7 @@ FSSDKernelConfig CreateKernelConfig()
 	KernelConfig.bMaxWithRefBilateralDistance = false;
 	
 	{
-		FIX_UNROLL(MAX_SIGNAL_MULTIPLEXING)
+		UNROLL_N(MAX_SIGNAL_MULTIPLEXING)
 		for (uint MultiplexId = 0; MultiplexId < MAX_SIGNAL_MULTIPLEXING; MultiplexId++)
 		{
 			KernelConfig.BufferColorSpace[MultiplexId] = STANDARD_BUFFER_COLOR_SPACE;
@@ -180,7 +180,7 @@ FSSDKernelConfig CreateKernelConfig()
 	KernelConfig.SampleTrackId = 0;
 
 	{	
-		FIX_UNROLL(MAX_SIGNAL_MULTIPLEXING)
+		UNROLL_N(MAX_SIGNAL_MULTIPLEXING)
 		for (uint MultiplexId = 0; MultiplexId < MAX_SIGNAL_MULTIPLEXING; MultiplexId++)
 		{
 			KernelConfig.RefBilateralDistance[MultiplexId] = 0.0;
@@ -281,7 +281,7 @@ void AccumulateSampledMultiplexedSignals(
 	// When sampling the previous frame, need to be forgiving in the world based bilateral rejection because TAA jitter.
 	const float WorldBluringDistanceForgiveness = KernelConfig.bPreviousFrameMetadata ? 3 : 1.0;
 	
-	FIX_UNROLL(MAX_SIGNAL_MULTIPLEXING)
+	UNROLL_N(MAX_SIGNAL_MULTIPLEXING)
 	for (uint SignalMultiplexId = 0; SignalMultiplexId < MAX_SIGNAL_MULTIPLEXING; SignalMultiplexId++)
 	{
 		// Compute at compile time the id of the signal being processed.
@@ -421,7 +421,7 @@ void SampleAndAccumulateMultiplexedSignalsPair(
 	FSSDSignalArray MultiplexedSamples[2];
 	bool bIsOutsideFrustum[2];
 
-	FIX_UNROLL(2)
+	UNROLL_N(2)
 	for (uint PairFetchId = 0; PairFetchId < 2; PairFetchId++)
 	{
 		// Stores in SGPR whether this sample is outside the viewport, to avoid VGPR pressure to keep SampleBufferUV after texture fetches.
@@ -441,7 +441,7 @@ void SampleAndAccumulateMultiplexedSignalsPair(
 	// Take the min inverse frequency per signal if desired.
 	if (KernelConfig.bMinSamplePairInvFrequency)
 	{
-		FIX_UNROLL(MAX_SIGNAL_MULTIPLEXING)
+		UNROLL_N(MAX_SIGNAL_MULTIPLEXING)
 		for (uint SignalMultiplexId = 0; SignalMultiplexId < MAX_SIGNAL_MULTIPLEXING; SignalMultiplexId++)
 		{
 			float MinInvFrequency = min(
@@ -457,7 +457,7 @@ void SampleAndAccumulateMultiplexedSignalsPair(
 		}
 	}
 
-	FIX_UNROLL(2)
+	UNROLL_N(2)
 	for (uint PairAccumulateId = 0; PairAccumulateId < 2; PairAccumulateId++)
 	{
 		AccumulateSampledMultiplexedSignals(
@@ -477,7 +477,7 @@ void StartAccumulatingCluster(
 	inout FSSDSignalAccumulatorArray Accumulators,
 	FSSDSampleClusterInfo ClusterInfo)
 {
-	FIX_UNROLL(MAX_SIGNAL_MULTIPLEXING)
+	UNROLL_N(MAX_SIGNAL_MULTIPLEXING)
 	for (uint SignalMultiplexId = 0; SignalMultiplexId < MAX_SIGNAL_MULTIPLEXING; SignalMultiplexId++)
 	{
 		StartAccumulatingCluster(
@@ -489,7 +489,7 @@ void StartAccumulatingCluster(
 
 void DijestAccumulatedClusterSamples(inout FSSDSignalAccumulatorArray Accumulators, uint RingId, uint SampleCount)
 {
-	FIX_UNROLL(MAX_SIGNAL_MULTIPLEXING)
+	UNROLL_N(MAX_SIGNAL_MULTIPLEXING)
 	for (uint SignalMultiplexId = 0; SignalMultiplexId < MAX_SIGNAL_MULTIPLEXING; SignalMultiplexId++)
 	{
 		DijestAccumulatedClusterSamples(
@@ -546,7 +546,7 @@ void AccumulateBilinear(
 		KernelConfig.BufferSizeAndInvSize.zw,
 		MipLevelPow2, rcp(MipLevelPow2));
 	
-	FIX_UNROLL(4)
+	UNROLL_N(4)
 	for (uint i = 0; i < 4; i++)
 	{
 		float2 SampleOffset = BilinearSamplingOffsets2x2[i];
@@ -847,7 +847,7 @@ void ConvolveStackowiakKernel(
 		LOOP
 		for (uint BatchId = 0; BatchId < BatchCountCount; BatchId++)
 		{
-			FIX_UNROLL(2)
+			UNROLL_N(2)
 			for (uint SampleBatchId = 0; SampleBatchId < 2; SampleBatchId++)
 			{
 				const uint SampleId = (BatchCountCount - BatchId) * kSamplingBatchSize - 1 - SampleBatchId;
diff --git a/Engine/Shaders/Private/ScreenSpaceDenoise/SSDTemporalAccumulation.usf b/Engine/Shaders/Private/ScreenSpaceDenoise/SSDTemporalAccumulation.usf
--- a/Engine/Shaders/Private/ScreenSpaceDenoise/SSDTemporalAccumulation.usf
+++ b/Engine/Shaders/Private/ScreenSpaceDenoise/SSDTemporalAccumulation.usf
@@ -436,7 +436,7 @@ void TemporallyAccumulate(
 				/* inout */ SignalAccumulators);
 		
 			// Exports the history sample from accumulator.
-			FIX_UNROLL(CONFIG_SIGNAL_BATCH_SIZE)
+			UNROLL_N(CONFIG_SIGNAL_BATCH_SIZE)
 			for (uint BatchedSignalId = 0; BatchedSignalId < CONFIG_SIGNAL_BATCH_SIZE; BatchedSignalId++)
 			{
 				HistorySamples.Array[BatchedSignalId] = SignalAccumulators.Array[BatchedSignalId].Moment1;
@@ -447,7 +447,7 @@ void TemporallyAccumulate(
 				}
 			}
 
-			FIX_UNROLL(CONFIG_SIGNAL_BATCH_SIZE)
+			UNROLL_N(CONFIG_SIGNAL_BATCH_SIZE)
 			for (uint BatchedSignalId = 0; BatchedSignalId < CONFIG_SIGNAL_BATCH_SIZE; BatchedSignalId++)
 			{
 				FSSDSignalSample CurrentFrameSample = CurrentFrameSamples.Array[BatchedSignalId];
@@ -555,7 +555,7 @@ void TemporallyAccumulate(
 		FSSDSignalAccumulatorArray SignalAccumulators = CreateSignalAccumulatorArray();		
 		if (bUseCurrentSampleForCenter)
 		{
-			FIX_UNROLL(MAX_SIGNAL_MULTIPLEXING)
+			UNROLL_N(MAX_SIGNAL_MULTIPLEXING)
 			for (uint SignalMultiplexId = 0; SignalMultiplexId < MAX_SIGNAL_MULTIPLEXING; SignalMultiplexId++)
 			{
 				const uint BatchedSignalId = ComputeSignalBatchIdFromSignalMultiplexId(KernelConfig, SignalMultiplexId);
@@ -601,7 +601,7 @@ void TemporallyAccumulate(
 		#endif
 
 		// Clamp history
-		FIX_UNROLL(CONFIG_SIGNAL_BATCH_SIZE)
+		UNROLL_N(CONFIG_SIGNAL_BATCH_SIZE)
 		for (uint BatchedSignalId = 0; BatchedSignalId < CONFIG_SIGNAL_BATCH_SIZE; BatchedSignalId++)
 		{
 			FSSDSignalSample NeighborMin = SignalAccumulators.Array[BatchedSignalId].Min;
