Compiles out the SceneColor component of the denoising signal on shadow and AO denoiser.

diff --git a/Engine/Shaders/Private/ScreenSpaceDenoise/SSDInjest.usf b/Engine/Shaders/Private/ScreenSpaceDenoise/SSDInjest.usf
--- a/Engine/Shaders/Private/ScreenSpaceDenoise/SSDInjest.usf
+++ b/Engine/Shaders/Private/ScreenSpaceDenoise/SSDInjest.usf
@@ -9,7 +9,7 @@
 
 #define CONFIG_SIGNAL_PROCESSING DIM_SIGNAL_PROCESSING
 #define CONFIG_SIGNAL_BATCH_SIZE DIM_SIGNAL_BATCH_SIZE
-
+#define CONFIG_SIGNAL_COLOR 0
 
 #if CONFIG_SIGNAL_PROCESSING == SIGNAL_PROCESSING_MONOCHROMATIC_PENUMBRA
 	#if DIM_MULTI_SPP
diff --git a/Engine/Shaders/Private/ScreenSpaceDenoise/SSDSignalCore.ush b/Engine/Shaders/Private/ScreenSpaceDenoise/SSDSignalCore.ush
--- a/Engine/Shaders/Private/ScreenSpaceDenoise/SSDSignalCore.ush
+++ b/Engine/Shaders/Private/ScreenSpaceDenoise/SSDSignalCore.ush
@@ -5,6 +5,17 @@
 #include "SSDCommon.ush"
 
 
+// Work arround FXC shader compiler that confuses FSSDSignalSample with a FSphericalGaussian when
+// CONFIG_SIGNAL_COLOR == 0 on the overloaded functions
+#undef CONFIG_SIGNAL_COLOR
+#define CONFIG_SIGNAL_COLOR 1
+
+
+#if !defined(CONFIG_SIGNAL_COLOR)
+	#error CONFIG_SIGNAL_COLOR has not been configured.
+#endif
+
+
 //------------------------------------------------------- STRUCTURE
 
 /** Generic data structure to manipulate any kind of signal. */
@@ -12,9 +23,11 @@ struct FSSDSignalSample
 {
 	// Number of sample accumulated.
 	float SampleCount;
-
+	
 	// Scene color and alpha.
-	float4 SceneColor;
+	#if CONFIG_SIGNAL_COLOR
+		float4 SceneColor;
+	#endif
 
 	// Number of ray that did not find any intersections.
 	float MissCount;
@@ -39,7 +52,9 @@ struct FSSDSignalArray
 FSSDSignalSample CreateSignalSampleFromScalarValue(float Scalar)
 {
 	FSSDSignalSample OutSample;
-	OutSample.SceneColor = Scalar;
+	#if CONFIG_SIGNAL_COLOR
+		OutSample.SceneColor = Scalar;
+	#endif
 	OutSample.SampleCount = Scalar;
 	OutSample.MissCount = Scalar;
 	OutSample.ClosestHitDistance = Scalar;
@@ -70,9 +85,11 @@ FSSDSignalArray CreateSignalArrayFromScalarValue(float Scalar)
 FSSDSignalSample Mul(FSSDSignalSample Sample, float Scalar)
 {
 	FSSDSignalSample OutSample;
-	OutSample.SceneColor             = Sample.SceneColor * Scalar;
-	OutSample.SampleCount = Sample.SampleCount * Scalar;
-	OutSample.MissCount        = Sample.MissCount * Scalar;
+	#if CONFIG_SIGNAL_COLOR
+		OutSample.SceneColor         = Sample.SceneColor * Scalar;
+	#endif
+	OutSample.SampleCount            = Sample.SampleCount * Scalar;
+	OutSample.MissCount              = Sample.MissCount * Scalar;
 	OutSample.ClosestHitDistance     = Sample.ClosestHitDistance * Scalar;
 	OutSample.UncertainRayCount      = Sample.UncertainRayCount * Scalar;
 	OutSample.WorldBluringRadius     = Sample.WorldBluringRadius * Scalar;
@@ -81,9 +98,11 @@ FSSDSignalSample Mul(FSSDSignalSample Sample, float Scalar)
 
 FSSDSignalSample Add(FSSDSignalSample SampleA, FSSDSignalSample SampleB)
 {
-	SampleA.SceneColor             += SampleB.SceneColor;
-	SampleA.SampleCount += SampleB.SampleCount;
-	SampleA.MissCount        += SampleB.MissCount;
+	#if CONFIG_SIGNAL_COLOR
+		SampleA.SceneColor         += SampleB.SceneColor;
+	#endif
+	SampleA.SampleCount            += SampleB.SampleCount;
+	SampleA.MissCount              += SampleB.MissCount;
 	SampleA.ClosestHitDistance     += SampleB.ClosestHitDistance;
 	SampleA.UncertainRayCount      += SampleB.UncertainRayCount;
 	SampleA.WorldBluringRadius     += SampleB.WorldBluringRadius;
@@ -92,9 +111,11 @@ FSSDSignalSample Add(FSSDSignalSample SampleA, FSSDSignalSample SampleB)
 
 FSSDSignalSample Minus(FSSDSignalSample Sample)
 {
-	Sample.SceneColor             = -Sample.SceneColor;
-	Sample.SampleCount = -Sample.SampleCount;
-	Sample.MissCount        = -Sample.MissCount;
+	#if CONFIG_SIGNAL_COLOR
+		Sample.SceneColor         = -Sample.SceneColor;
+	#endif
+	Sample.SampleCount            = -Sample.SampleCount;
+	Sample.MissCount              = -Sample.MissCount;
 	Sample.ClosestHitDistance     = -Sample.ClosestHitDistance;
 	Sample.UncertainRayCount      = -Sample.UncertainRayCount;
 	Sample.WorldBluringRadius     = -Sample.WorldBluringRadius;
@@ -103,14 +124,24 @@ FSSDSignalSample Minus(FSSDSignalSample Sample)
 
 FSSDSignalSample Subtract(FSSDSignalSample SampleA, FSSDSignalSample SampleB)
 {
-	return Add(SampleA, Minus(SampleB));
+	#if CONFIG_SIGNAL_COLOR
+		SampleA.SceneColor     -= SampleB.SceneColor;
+	#endif
+	SampleA.SampleCount        -= SampleB.SampleCount;
+	SampleA.MissCount          -= SampleB.MissCount;
+	SampleA.ClosestHitDistance -= SampleB.ClosestHitDistance;
+	SampleA.UncertainRayCount  -= SampleB.UncertainRayCount;
+	SampleA.WorldBluringRadius -= SampleB.WorldBluringRadius;
+	return SampleA;
 }
 
 FSSDSignalSample Abs(FSSDSignalSample Sample)
 {
-	Sample.SceneColor             = abs(Sample.SceneColor);
-	Sample.SampleCount = abs(Sample.SampleCount);
-	Sample.MissCount        = abs(Sample.MissCount);
+	#if CONFIG_SIGNAL_COLOR
+		Sample.SceneColor         = abs(Sample.SceneColor);
+	#endif
+	Sample.SampleCount            = abs(Sample.SampleCount);
+	Sample.MissCount              = abs(Sample.MissCount);
 	Sample.ClosestHitDistance     = abs(Sample.ClosestHitDistance);
 	Sample.UncertainRayCount      = abs(Sample.UncertainRayCount);
 	Sample.WorldBluringRadius     = abs(Sample.WorldBluringRadius);
@@ -119,9 +150,11 @@ FSSDSignalSample Abs(FSSDSignalSample Sample)
 
 FSSDSignalSample Sqrt(FSSDSignalSample Sample)
 {
-	Sample.SceneColor             = sqrt(Sample.SceneColor);
-	Sample.SampleCount = sqrt(Sample.SampleCount);
-	Sample.MissCount        = sqrt(Sample.MissCount);
+	#if CONFIG_SIGNAL_COLOR
+		Sample.SceneColor         = sqrt(Sample.SceneColor);
+	#endif
+	Sample.SampleCount            = sqrt(Sample.SampleCount);
+	Sample.MissCount              = sqrt(Sample.MissCount);
 	Sample.ClosestHitDistance     = sqrt(Sample.ClosestHitDistance);
 	Sample.UncertainRayCount      = sqrt(Sample.UncertainRayCount);
 	Sample.WorldBluringRadius     = sqrt(Sample.WorldBluringRadius);
@@ -130,12 +163,14 @@ FSSDSignalSample Sqrt(FSSDSignalSample Sample)
 
 FSSDSignalSample Power(FSSDSignalSample Sample, float Exponent)
 {
-	Sample.SceneColor.r           = pow(Sample.SceneColor.r, Exponent);
-	Sample.SceneColor.g           = pow(Sample.SceneColor.g, Exponent);
-	Sample.SceneColor.b           = pow(Sample.SceneColor.b, Exponent);
-	Sample.SceneColor.a           = pow(Sample.SceneColor.a, Exponent);
-	Sample.SampleCount = pow(Sample.SampleCount, Exponent);
-	Sample.MissCount        = pow(Sample.MissCount, Exponent);
+	#if CONFIG_SIGNAL_COLOR
+		Sample.SceneColor.r       = pow(Sample.SceneColor.r, Exponent);
+		Sample.SceneColor.g       = pow(Sample.SceneColor.g, Exponent);
+		Sample.SceneColor.b       = pow(Sample.SceneColor.b, Exponent);
+		Sample.SceneColor.a       = pow(Sample.SceneColor.a, Exponent);
+	#endif
+	Sample.SampleCount            = pow(Sample.SampleCount, Exponent);
+	Sample.MissCount              = pow(Sample.MissCount, Exponent);
 	Sample.ClosestHitDistance     = pow(Sample.ClosestHitDistance, Exponent);
 	Sample.UncertainRayCount      = pow(Sample.UncertainRayCount, Exponent);
 	Sample.WorldBluringRadius     = pow(Sample.WorldBluringRadius, Exponent);
@@ -145,7 +180,9 @@ FSSDSignalSample Power(FSSDSignalSample Sample, float Exponent)
 FSSDSignalSample Min(FSSDSignalSample SampleA, FSSDSignalSample SampleB)
 {
 	FSSDSignalSample OutSample;
-	OutSample.SceneColor = min(SampleA.SceneColor, SampleB.SceneColor);
+	#if CONFIG_SIGNAL_COLOR
+		OutSample.SceneColor = min(SampleA.SceneColor, SampleB.SceneColor);
+	#endif
 	OutSample.SampleCount = min(SampleA.SampleCount, SampleB.SampleCount);
 	OutSample.MissCount = min(SampleA.MissCount, SampleB.MissCount);
 	OutSample.ClosestHitDistance = min(SampleA.ClosestHitDistance, SampleB.ClosestHitDistance);
@@ -157,7 +194,9 @@ FSSDSignalSample Min(FSSDSignalSample SampleA, FSSDSignalSample SampleB)
 FSSDSignalSample Max(FSSDSignalSample SampleA, FSSDSignalSample SampleB)
 {
 	FSSDSignalSample OutSample;
-	OutSample.SceneColor = max(SampleA.SceneColor, SampleB.SceneColor);
+	#if CONFIG_SIGNAL_COLOR
+		OutSample.SceneColor = max(SampleA.SceneColor, SampleB.SceneColor);
+	#endif
 	OutSample.SampleCount = max(SampleA.SampleCount, SampleB.SampleCount);
 	OutSample.MissCount = max(SampleA.MissCount, SampleB.MissCount);
 	OutSample.ClosestHitDistance = max(SampleA.ClosestHitDistance, SampleB.ClosestHitDistance);
@@ -169,7 +208,9 @@ FSSDSignalSample Max(FSSDSignalSample SampleA, FSSDSignalSample SampleB)
 FSSDSignalSample Clamp(FSSDSignalSample Sample, FSSDSignalSample SampleMin, FSSDSignalSample SampleMax)
 {
 	FSSDSignalSample OutSample;
-	OutSample.SceneColor = clamp(Sample.SceneColor, SampleMin.SceneColor, SampleMax.SceneColor);
+	#if CONFIG_SIGNAL_COLOR
+		OutSample.SceneColor = clamp(Sample.SceneColor, SampleMin.SceneColor, SampleMax.SceneColor);
+	#endif
 	OutSample.SampleCount = clamp(Sample.SampleCount, SampleMin.SampleCount, SampleMax.SampleCount);
 	OutSample.MissCount = clamp(Sample.MissCount, SampleMin.MissCount, SampleMax.MissCount);
 	OutSample.ClosestHitDistance = clamp(Sample.ClosestHitDistance, SampleMin.ClosestHitDistance, SampleMax.ClosestHitDistance);
@@ -308,6 +349,7 @@ void DecodeMultiplexedSignals(
 			OutSamples.Array[MultiplexId + 1].MissCount = Channels.g * SampleCount;
 		}
 	}
+	#if CONFIG_SIGNAL_COLOR
 	else if (SignalBufferLayout == SIGNAL_BUFFER_LAYOUT_REFLECTIONS_INPUT)
 	{
 		OutSamples.Array[0].SampleCount = RawSample[1].r == DENOISER_INVALID_HIT_DISTANCE ? 0 : 1;
@@ -356,6 +398,7 @@ void DecodeMultiplexedSignals(
 		OutSamples.Array[2].SampleCount = OutSamples.Array[0].SampleCount;
 		OutSamples.Array[2].SceneColor.x = OutSamples.Array[2].SampleCount * RawSample[1].b;
 	}
+	#endif
 	else if (SignalBufferLayout == SIGNAL_BUFFER_LAYOUT_AO_INPUT)
 	{
 		OutSamples.Array[0].SampleCount = RawSample[1].r < 0 ? 0 : 1;
@@ -394,6 +437,7 @@ void DecodeMultiplexedSignals(
 			OutSamples.Array[0].MissCount = RawSample[0].g * RawSample[0].r;
 		}
 	}
+	#if CONFIG_SIGNAL_COLOR
 	else if (SignalBufferLayout == SIGNAL_BUFFER_LAYOUT_GI_AND_AO_INPUT)
 	{
 		OutSamples.Array[0].SampleCount = RawSample[1].g;
@@ -438,6 +482,7 @@ void DecodeMultiplexedSignals(
 		OutSamples.Array[2].SampleCount = OutSamples.Array[0].SampleCount;
 		OutSamples.Array[2].SceneColor.x = OutSamples.Array[2].SampleCount * RawSample[1].b;
 	}
+	#endif
 } // DecodeMultiplexedSignals()
 
 /** Encode output signal sample. */
@@ -555,6 +600,7 @@ void EncodeMultiplexedSignals(
 				NormalizedWorldBluringRadius == WORLD_RADIUS_MISS ? -1 : NormalizedWorldBluringRadius);
 		}
 	}
+	#if CONFIG_SIGNAL_COLOR
 	else if (SignalBufferLayout == SIGNAL_BUFFER_LAYOUT_REFLECTIONS_REJECTION)
 	{
 		// Outputs scene color to be compatible with a SSR output.
@@ -585,6 +631,7 @@ void EncodeMultiplexedSignals(
 			OutBufferCount = 3;
 		}
 	}
+	#endif
 	else if (SignalBufferLayout == SIGNAL_BUFFER_LAYOUT_AO_REJECTION)
 	{
 		// Outputs number of ray miss in red to be compatible as a SSAO output.
@@ -603,6 +650,7 @@ void EncodeMultiplexedSignals(
 		OutRawSample[0].r = Sample[0].SampleCount > 0 ? Sample[0].MissCount / Sample[0].SampleCount : 1.0;
 		OutRawSample[0].g = Sample[0].SampleCount;
 	}
+	#if CONFIG_SIGNAL_COLOR
 	else if (SignalBufferLayout == SIGNAL_BUFFER_LAYOUT_GI_AND_AO_RECONSTRUCTION)
 	{
 		float NormalizationFactor = SafeRcp(Sample[0].SampleCount);
@@ -625,6 +673,7 @@ void EncodeMultiplexedSignals(
 			OutRawSample[1].b = Sample[2].SceneColor.x * NormalizationFactor;
 		}
 	}
+	#endif
 } // EncodeMultiplexedSignals()
 
 /** Sample multiple input signals that have been multiplexed. */
diff --git a/Engine/Shaders/Private/ScreenSpaceDenoise/SSDSignalFramework.ush b/Engine/Shaders/Private/ScreenSpaceDenoise/SSDSignalFramework.ush
--- a/Engine/Shaders/Private/ScreenSpaceDenoise/SSDSignalFramework.ush
+++ b/Engine/Shaders/Private/ScreenSpaceDenoise/SSDSignalFramework.ush
@@ -160,6 +160,7 @@ float GetSignalColorTransformationExposure()
 
 /** Returns the linear color space Luma with a 4 multiplier for ALU perf reasons. */
 float GetSignalLuma4(FSSDSignalSample Sample, const uint Basis)
+#if CONFIG_SIGNAL_COLOR
 {
 	// TODO: should force this to be scalar load.
 	float FrameExposureScale = GetSignalColorTransformationExposure();
@@ -192,9 +193,15 @@ float GetSignalLuma4(FSSDSignalSample Sample, const uint Basis)
 		return Luma4(Sample.SceneColor.rgb);
 	}
 }
+#else
+{
+	return 0;
+}
+#endif
 
 /** Conveniently transform one simple from a source basis to a destination basis. */
 FSSDSignalSample TransformSignal(FSSDSignalSample Sample, const uint SrcBasis, const uint DestBasis)
+#if CONFIG_SIGNAL_COLOR
 {
 	// If the source and destination bases are the same, early return to not pay ALU and floating point divergence.
 	if (SrcBasis == DestBasis)
@@ -304,6 +311,11 @@ FSSDSignalSample TransformSignal(FSSDSignalSample Sample, const uint SrcBasis, c
 
 	return Sample;
 }
+#else
+{
+	return Sample;
+}
+#endif
 
 #if 0 // will probably be handy at some point.
 float4 WeightedLerp( float4 ColorA, float WeightA, float4 ColorB, float WeightB, float Blend ) 
diff --git a/Engine/Shaders/Private/ScreenSpaceDenoise/SSDSpatialAccumulation.usf b/Engine/Shaders/Private/ScreenSpaceDenoise/SSDSpatialAccumulation.usf
--- a/Engine/Shaders/Private/ScreenSpaceDenoise/SSDSpatialAccumulation.usf
+++ b/Engine/Shaders/Private/ScreenSpaceDenoise/SSDSpatialAccumulation.usf
@@ -40,6 +40,8 @@
 
 // Configures all the pass for each individual signals.
 #if CONFIG_SIGNAL_PROCESSING == SIGNAL_PROCESSING_MONOCHROMATIC_PENUMBRA
+	#define CONFIG_SIGNAL_COLOR 0
+
 	//#define CONFIG_VIEW_POSITION_BASED_BILATERAL 2
 	#define CONFIG_VIEW_POSITION_BASED_BILATERAL 5
 	#define CONFIG_MULTIPLEXED_SIGNALS_PER_SIGNAL_DOMAIN 1
@@ -110,6 +112,8 @@
 	#endif
 
 #elif CONFIG_SIGNAL_PROCESSING == SIGNAL_PROCESSING_REFLECTIONS
+	#define CONFIG_SIGNAL_COLOR 1
+
 	#if DIM_STAGE == STAGE_RECONSTRUCTION
 		// Input and output layout.
 		#define CONFIG_SIGNAL_INPUT_LAYOUT  SIGNAL_BUFFER_LAYOUT_REFLECTIONS_INPUT
@@ -179,6 +183,8 @@
 	#endif
 
 #elif CONFIG_SIGNAL_PROCESSING == SIGNAL_PROCESSING_AO
+	#define CONFIG_SIGNAL_COLOR 0
+
 	#define CONFIG_VIEW_POSITION_BASED_BILATERAL 2
 	#define CONFIG_NORMAL_BASED_BILATERAL 1
 
@@ -216,6 +222,8 @@
 	#endif
 
 #elif CONFIG_SIGNAL_PROCESSING == SIGNAL_PROCESSING_GI_AND_AO
+	#define CONFIG_SIGNAL_COLOR 1
+
 	#define CONFIG_VIEW_POSITION_BASED_BILATERAL 2
 	#define CONFIG_NORMAL_BASED_BILATERAL 1
 
@@ -705,13 +713,18 @@ void MainCS(
 		for (uint MultiplexId = 0; MultiplexId < CONFIG_SIGNAL_BATCH_SIZE; MultiplexId++)
 		{
 			#if CONFIG_REJECT_HIGHEST_COLOR
-			if (Accumulator.Moment1[MultiplexId].SampleCount > 0)
 			{
-				const float MaxNeighborWeight = saturate(SignalAccumulators.Array[MultiplexId].Moment1.SampleCount * rcp(10) - 1);
+				#if !CONFIG_SIGNAL_COLOR
+					#error Need to compile signal color.
+				#endif
+				if (Accumulator.Moment1[MultiplexId].SampleCount > 0)
+				{
+					const float MaxNeighborWeight = saturate(SignalAccumulators.Array[MultiplexId].Moment1.SampleCount * rcp(10) - 1);
 
-				SignalAccumulators.Array[MultiplexId].Moment1.SceneColor.rgb =
-					(SignalAccumulators.Array[MultiplexId].Moment1.SceneColor.rgb - MaxNeighbor.SceneColor.rgb * MaxNeighborWeight) *
-					(SignalAccumulators.Array[MultiplexId].Moment1.SampleCount / (SignalAccumulators.Array[MultiplexId].Moment1.SampleCount - MaxNeighborWeight));
+					SignalAccumulators.Array[MultiplexId].Moment1.SceneColor.rgb =
+						(SignalAccumulators.Array[MultiplexId].Moment1.SceneColor.rgb - MaxNeighbor.SceneColor.rgb * MaxNeighborWeight) *
+						(SignalAccumulators.Array[MultiplexId].Moment1.SampleCount / (SignalAccumulators.Array[MultiplexId].Moment1.SampleCount - MaxNeighborWeight));
+				}
 			}
 			#endif
 
diff --git a/Engine/Shaders/Private/ScreenSpaceDenoise/SSDTemporalAccumulation.usf b/Engine/Shaders/Private/ScreenSpaceDenoise/SSDTemporalAccumulation.usf
--- a/Engine/Shaders/Private/ScreenSpaceDenoise/SSDTemporalAccumulation.usf
+++ b/Engine/Shaders/Private/ScreenSpaceDenoise/SSDTemporalAccumulation.usf
@@ -33,6 +33,8 @@
 
 
 #if CONFIG_SIGNAL_PROCESSING == SIGNAL_PROCESSING_MONOCHROMATIC_PENUMBRA
+	#define CONFIG_SIGNAL_COLOR 0
+
 	#define CONFIG_VIEW_POSITION_BASED_BILATERAL 2
 	#define CONFIG_MULTIPLEXED_SIGNALS_PER_SIGNAL_DOMAIN 1
 
@@ -61,6 +63,8 @@
 	#endif
 
 #elif CONFIG_SIGNAL_PROCESSING == SIGNAL_PROCESSING_REFLECTIONS
+	#define CONFIG_SIGNAL_COLOR 1
+
 	// Input and output layout.
 	#define CONFIG_SIGNAL_INPUT_LAYOUT   SIGNAL_BUFFER_LAYOUT_REFLECTIONS_HISTORY
 	#define CONFIG_SIGNAL_HISTORY_LAYOUT SIGNAL_BUFFER_LAYOUT_REFLECTIONS_HISTORY
@@ -90,6 +94,8 @@
 	#endif
 
 #elif CONFIG_SIGNAL_PROCESSING == SIGNAL_PROCESSING_AO
+	#define CONFIG_SIGNAL_COLOR 0
+
 	// Input and output layout.
 	#define CONFIG_SIGNAL_INPUT_LAYOUT   SIGNAL_BUFFER_LAYOUT_AO_HISTORY
 	#define CONFIG_SIGNAL_HISTORY_LAYOUT SIGNAL_BUFFER_LAYOUT_AO_HISTORY
@@ -113,6 +119,8 @@
 	#endif
 
 #elif CONFIG_SIGNAL_PROCESSING == SIGNAL_PROCESSING_GI_AND_AO
+	#define CONFIG_SIGNAL_COLOR 1
+
 	// Input and output layout.
 	#define CONFIG_SIGNAL_INPUT_LAYOUT   SIGNAL_BUFFER_LAYOUT_GI_AND_AO_RECONSTRUCTION
 	#define CONFIG_SIGNAL_HISTORY_LAYOUT SIGNAL_BUFFER_LAYOUT_GI_AND_AO_HISTORY
@@ -460,7 +468,9 @@ void TemporallyAccumulate(
 					SafeRcp(CurrentFrameSample.SampleCount));
 
 				// Applies correction when pre exposure change happen.
-				HistorySamples.Array[BatchedSignalId].SceneColor.rgb *= HistoryPreExposureCorrection;
+				#if CONFIG_SIGNAL_COLOR
+					HistorySamples.Array[BatchedSignalId].SceneColor.rgb *= HistoryPreExposureCorrection;
+				#endif
 			} // for (uint BatchedSignalId = 0; BatchedSignalId < CONFIG_SIGNAL_BATCH_SIZE; BatchedSignalId++)
 		} // if (!bIsPreviousFrameOffscreen)
 	}
@@ -672,6 +682,9 @@ void TemporallyAccumulate(
 					float RejectionFactor = 0.0;
 					#if CONFIG_SIGNAL_PROCESSING == SIGNAL_PROCESSING_REFLECTIONS && (CONFIG_HISTORY_REJECTION_COLOR_SPACE & COLOR_SPACE_LCOCG)
 					{
+						#if !CONFIG_SIGNAL_COLOR
+							#error Need to compile signal color.
+						#endif
 						RejectionFactor = abs(
 							Luma_To_LumaLog(ClampedTransformedHistorySample.SceneColor.x) -
 							Luma_To_LumaLog(TransformedHistorySample.SceneColor.x));
@@ -738,6 +751,10 @@ void TemporallyAccumulate(
 		#if CONFIG_SIGNAL_BATCH_SIZE != 1
 			#error Unsupported.
 		#endif
+	
+		#if !CONFIG_SIGNAL_COLOR
+			#error Need to compile signal color.
+		#endif
 
 		const uint BatchedSignalId = 0;
 	
