Includes all sub dependencies of QuasiRandom.ush so it can be used straight with a single #include

diff --git a/Engine/Shaders/Private/BlueNoise.ush b/Engine/Shaders/Private/BlueNoise.ush
--- a/Engine/Shaders/Private/BlueNoise.ush
+++ b/Engine/Shaders/Private/BlueNoise.ush
@@ -3,6 +3,6 @@
 float3 EvalBlueNoise(uint2 TileCoord, uint Slice)
 {
 	uint3 BlueNoiseCoord = uint3(TileCoord.x, Slice * BlueNoise.Dimensions.y + TileCoord.y, 0);
-	float3 Noise = BlueNoise.Texture.Load(BlueNoiseCoord, 0);
+	float3 Noise = BlueNoise.Texture.Load(BlueNoiseCoord, 0).rgb;
 	return Noise;
 }
diff --git a/Engine/Shaders/Private/PathTracing/Utilities/PathTracingRandomSequence.ush b/Engine/Shaders/Private/PathTracing/Utilities/PathTracingRandomSequence.ush
--- a/Engine/Shaders/Private/PathTracing/Utilities/PathTracingRandomSequence.ush
+++ b/Engine/Shaders/Private/PathTracing/Utilities/PathTracingRandomSequence.ush
@@ -6,6 +6,10 @@
 
 #pragma once
 
+
+#include "../../QuasiRandom.ush" 
+
+
 struct RandomSequence
 {
 	uint Type;
@@ -134,19 +138,6 @@ uint Prime512(uint Dimension)
 }
 
 
-// http://burtleburtle.net/bob/hash/integer.html
-// Bob Jenkins integer hashing function in 6 shifts.
-uint IntegerHash(uint a)
-{
-	a = (a + 0x7ed55d16) + (a << 12);
-	a = (a ^ 0xc761c23c) ^ (a >> 19);
-	a = (a + 0x165667b1) + (a << 5);
-	a = (a + 0xd3a2646c) ^ (a << 9);
-	a = (a + 0xfd7046c5) + (a << 3);
-	a = (a ^ 0xb55a4f09) ^ (a >> 16);
-	return a;
-}
-
 float Halton(uint Index, uint Base)
 {
 	float r = 0.0;
diff --git a/Engine/Shaders/Private/QuasiRandom.ush b/Engine/Shaders/Private/QuasiRandom.ush
--- a/Engine/Shaders/Private/QuasiRandom.ush
+++ b/Engine/Shaders/Private/QuasiRandom.ush
@@ -8,6 +8,7 @@
 
 #include "BlueNoise.ush"
 #include "Halton.ush"
+#include "MortonCode.ush"
 
 struct FQuasiRandom
 {
@@ -59,7 +60,7 @@ float QuasiRandom_Sample(inout FQuasiRandom QuasiRandom)
 {
 	float RandomSample = 0.0;
 	BRANCH
-	if (QuasiRandom.Dimension < BlueNoise.Dimensions.z * 3)
+	if (QuasiRandom.Dimension < uint(BlueNoise.Dimensions.z * 3))
 	{
 		// Evaluate shift vector
 		float3 BlueNoiseShiftVector = EvalBlueNoise(QuasiRandom.BlueNoiseCoord, QuasiRandom.Dimension / 3);
@@ -87,7 +88,7 @@ float3 QuasiRandom_Sample3(inout FQuasiRandom QuasiRandom)
 {
 	float3 RandomSample = 0.0;
 	BRANCH
-	if (QuasiRandom.Dimension + 2 < BlueNoise.Dimensions.z * 3)
+	if (QuasiRandom.Dimension + 2 < uint(BlueNoise.Dimensions.z * 3))
 	{
 		// Evaluate shift vector
 		float3 BlueNoiseShift = EvalBlueNoise(QuasiRandom.BlueNoiseCoord, QuasiRandom.Dimension / 3);
