Improves diffuse indirect denoiser.

Notably have prototype code of spherical harmonic in the core signal framework. To be continued...

diff --git a/Engine/Shaders/Private/SHCommon.ush b/Engine/Shaders/Private/SHCommon.ush
--- a/Engine/Shaders/Private/SHCommon.ush
+++ b/Engine/Shaders/Private/SHCommon.ush
@@ -49,6 +49,15 @@ struct FThreeBandSHVectorRGB
 	FThreeBandSHVector B;
 };
 
+FTwoBandSHVectorRGB MulSH(FTwoBandSHVectorRGB A, half Scalar)
+{
+	FTwoBandSHVectorRGB Result;
+	Result.R.V = A.R.V * Scalar;
+	Result.G.V = A.G.V * Scalar;
+	Result.B.V = A.B.V * Scalar;
+	return Result;
+}
+
 FTwoBandSHVectorRGB MulSH(FTwoBandSHVector A, half3 Color)
 {
 	FTwoBandSHVectorRGB Result;
diff --git a/Engine/Shaders/Private/ScreenSpaceDenoise/SSDCommon.ush b/Engine/Shaders/Private/ScreenSpaceDenoise/SSDCommon.ush
--- a/Engine/Shaders/Private/ScreenSpaceDenoise/SSDCommon.ush
+++ b/Engine/Shaders/Private/ScreenSpaceDenoise/SSDCommon.ush
@@ -6,14 +6,42 @@
 
 #pragma once
 
+#if !defined(SPHERICAL_HARMONIC_ORDER)
+	#define SPHERICAL_HARMONIC_ORDER 0
+#endif
+
 #include "../Common.ush"
 #include "../SceneViewFamilyBlackboard.ush"
 #include "../ColorSpace.ush"
 #include "../SphericalGaussian.ush"
+
+#if SPHERICAL_HARMONIC_ORDER > 0
+	#include "../SHCommon.ush"
+#endif
+
 #include "SSDPublic.ush"
 #include "SSDDefinitions.ush"
 
 
+//------------------------------------------------------- SHARED CONSTANTS DEPENDING ON CONFIG OF SHADER
+
+#if CONFIG_SIGNAL_PROCESSING == SIGNAL_PROCESSING_REFLECTIONS
+	#define TARGETED_SAMPLE_COUNT 128
+
+#elif CONFIG_SIGNAL_PROCESSING == SIGNAL_PROCESSING_AO || CONFIG_SIGNAL_PROCESSING == SIGNAL_PROCESSING_MONOCHROMATIC_PENUMBRA
+	// A gray scale valued encoded into a 16bit float only have 10bits mantissa.
+	#define TARGETED_SAMPLE_COUNT 1024
+
+#elif CONFIG_SIGNAL_PROCESSING == SIGNAL_PROCESSING_GI_AND_AO
+	// Indirect diffuse requires really large spatial kernel to clean the noise.
+	#define TARGETED_SAMPLE_COUNT 4096
+
+#else
+	#error Unconfigured TARGETED_SAMPLE_COUNT
+
+#endif
+
+
 //------------------------------------------------------- SHARED SHADER PARAMETERS
 
 float4x4  PrevScreenToTranslatedWorld;
@@ -31,6 +59,25 @@ float SafeRcp(float x)
 	return 0;
 }
 
+//------------------------------------------------------- SPHERICAL HARMONIC INTERFACE
+
+#if SPHERICAL_HARMONIC_ORDER == 0
+	// NOP
+
+#elif SPHERICAL_HARMONIC_ORDER == 2
+typedef FTwoBandSHVector FSSDSphericalHarmonic;
+typedef FTwoBandSHVectorRGB FSSDSphericalHarmonicRGB;
+
+#elif SPHERICAL_HARMONIC_ORDER == 3
+typedef FThreeBandSHVector FSSDSphericalHarmonic;
+typedef FThreeBandSHVectorRGB FSSDSphericalHarmonicRGB;
+
+#else
+	#error Unknown spherical harmonic configuration.
+
+#endif
+
+
 //------------------------------------------------------- COMMON SCENE INFOS
 
 /** olds commonly used information of the scene of a given sample. */
diff --git a/Engine/Shaders/Private/ScreenSpaceDenoise/SSDDefinitions.ush b/Engine/Shaders/Private/ScreenSpaceDenoise/SSDDefinitions.ush
--- a/Engine/Shaders/Private/ScreenSpaceDenoise/SSDDefinitions.ush
+++ b/Engine/Shaders/Private/ScreenSpaceDenoise/SSDDefinitions.ush
@@ -40,8 +40,6 @@
 // Maximum number of mip level supported.
 #define MAX_MIP_LEVEL_COUNT 4
 
-#define TARGETED_SAMPLE_COUNT 1024
-
 // This constant is to artificially store the signal into higher resolution that theorically should from
 // the infered bluring radius by deviding it. 1.0 means nothing happen, but in practice want this >=
 // than Nyquist's 2.0 multiplier. Uses 3.0 because this is the size of the 3x3 neighborhood kernel.
@@ -201,3 +199,7 @@
 
 /** Color space used for reflection history rejection. */
 #define REFLECTIONS_REJECTION_COLOR_SPACE (COLOR_SPACE_LCOCG | COLOR_SPACE_UNPREMULTIPLY)
+
+/** Defines the order of the spherical harmonic to denoise. */
+// TODO: quality permutation instead?
+#define SPHERICAL_HARMONIC_ORDER 2
diff --git a/Engine/Shaders/Private/ScreenSpaceDenoise/SSDSignalCore.ush b/Engine/Shaders/Private/ScreenSpaceDenoise/SSDSignalCore.ush
--- a/Engine/Shaders/Private/ScreenSpaceDenoise/SSDSignalCore.ush
+++ b/Engine/Shaders/Private/ScreenSpaceDenoise/SSDSignalCore.ush
@@ -15,6 +15,10 @@
 	#error CONFIG_SIGNAL_COLOR has not been configured.
 #endif
 
+#ifndef CONFIG_SIGNAL_COLOR_SH
+	#define CONFIG_SIGNAL_COLOR_SH 0
+#endif
+
 
 //------------------------------------------------------- STRUCTURE
 
@@ -29,6 +33,11 @@ struct FSSDSignalSample
 		float4 SceneColor;
 	#endif
 
+	// Spherical harmonic of the color.
+	#if CONFIG_SIGNAL_COLOR_SH
+		FSSDSphericalHarmonicRGB ColorSH;
+	#endif
+
 	// Number of ray that did not find any intersections.
 	float MissCount;
 
@@ -55,6 +64,21 @@ FSSDSignalSample CreateSignalSampleFromScalarValue(float Scalar)
 	#if CONFIG_SIGNAL_COLOR
 		OutSample.SceneColor = Scalar;
 	#endif
+	#if CONFIG_SIGNAL_COLOR_SH && SPHERICAL_HARMONIC_ORDER == 2
+		OutSample.ColorSH.R.V = Scalar;
+		OutSample.ColorSH.G.V = Scalar;
+		OutSample.ColorSH.B.V = Scalar;
+	#elif CONFIG_SIGNAL_COLOR_SH && SPHERICAL_HARMONIC_ORDER == 3
+		OutSample.ColorSH.R.V0 = Scalar;
+		OutSample.ColorSH.R.V1 = Scalar;
+		OutSample.ColorSH.R.V2 = Scalar;
+		OutSample.ColorSH.G.V0 = Scalar;
+		OutSample.ColorSH.G.V1 = Scalar;
+		OutSample.ColorSH.G.V2 = Scalar;
+		OutSample.ColorSH.B.V0 = Scalar;
+		OutSample.ColorSH.B.V1 = Scalar;
+		OutSample.ColorSH.B.V2 = Scalar;
+	#endif
 	OutSample.SampleCount = Scalar;
 	OutSample.MissCount = Scalar;
 	OutSample.ClosestHitDistance = Scalar;
@@ -88,6 +112,9 @@ FSSDSignalSample Mul(FSSDSignalSample Sample, float Scalar)
 	#if CONFIG_SIGNAL_COLOR
 		OutSample.SceneColor         = Sample.SceneColor * Scalar;
 	#endif
+	#if CONFIG_SIGNAL_COLOR_SH
+		OutSample.ColorSH            = MulSH(Sample.ColorSH, Scalar);
+	#endif
 	OutSample.SampleCount            = Sample.SampleCount * Scalar;
 	OutSample.MissCount              = Sample.MissCount * Scalar;
 	OutSample.ClosestHitDistance     = Sample.ClosestHitDistance * Scalar;
@@ -101,6 +128,9 @@ FSSDSignalSample Add(FSSDSignalSample SampleA, FSSDSignalSample SampleB)
 	#if CONFIG_SIGNAL_COLOR
 		SampleA.SceneColor         += SampleB.SceneColor;
 	#endif
+	#if CONFIG_SIGNAL_COLOR_SH
+		SampleA.ColorSH            = AddSH(SampleA.ColorSH, SampleB.ColorSH);
+	#endif
 	SampleA.SampleCount            += SampleB.SampleCount;
 	SampleA.MissCount              += SampleB.MissCount;
 	SampleA.ClosestHitDistance     += SampleB.ClosestHitDistance;
@@ -114,6 +144,9 @@ FSSDSignalSample Minus(FSSDSignalSample Sample)
 	#if CONFIG_SIGNAL_COLOR
 		Sample.SceneColor         = -Sample.SceneColor;
 	#endif
+	#if CONFIG_SIGNAL_COLOR_SH
+		Sample.ColorSH         = MulSH(Sample.ColorSH, -1.0);
+	#endif
 	Sample.SampleCount            = -Sample.SampleCount;
 	Sample.MissCount              = -Sample.MissCount;
 	Sample.ClosestHitDistance     = -Sample.ClosestHitDistance;
@@ -124,15 +157,7 @@ FSSDSignalSample Minus(FSSDSignalSample Sample)
 
 FSSDSignalSample Subtract(FSSDSignalSample SampleA, FSSDSignalSample SampleB)
 {
-	#if CONFIG_SIGNAL_COLOR
-		SampleA.SceneColor     -= SampleB.SceneColor;
-	#endif
-	SampleA.SampleCount        -= SampleB.SampleCount;
-	SampleA.MissCount          -= SampleB.MissCount;
-	SampleA.ClosestHitDistance -= SampleB.ClosestHitDistance;
-	SampleA.UncertainRayCount  -= SampleB.UncertainRayCount;
-	SampleA.WorldBluringRadius -= SampleB.WorldBluringRadius;
-	return SampleA;
+	return Add(SampleA, Minus(SampleB));
 }
 
 FSSDSignalSample Abs(FSSDSignalSample Sample)
@@ -140,6 +165,21 @@ FSSDSignalSample Abs(FSSDSignalSample Sample)
 	#if CONFIG_SIGNAL_COLOR
 		Sample.SceneColor         = abs(Sample.SceneColor);
 	#endif
+	#if CONFIG_SIGNAL_COLOR_SH && SPHERICAL_HARMONIC_ORDER == 2
+		Sample.ColorSH.R.V = abs(Sample.ColorSH.R.V);
+		Sample.ColorSH.G.V = abs(Sample.ColorSH.G.V);
+		Sample.ColorSH.B.V = abs(Sample.ColorSH.B.V);
+	#elif CONFIG_SIGNAL_COLOR_SH && SPHERICAL_HARMONIC_ORDER == 3
+		Sample.ColorSH.R.V0 = abs(Sample.ColorSH.R.V0);
+		Sample.ColorSH.R.V1 = abs(Sample.ColorSH.R.V1);
+		Sample.ColorSH.R.V2 = abs(Sample.ColorSH.R.V2);
+		Sample.ColorSH.G.V0 = abs(Sample.ColorSH.G.V0);
+		Sample.ColorSH.G.V1 = abs(Sample.ColorSH.G.V1);
+		Sample.ColorSH.G.V2 = abs(Sample.ColorSH.G.V2);
+		Sample.ColorSH.B.V0 = abs(Sample.ColorSH.B.V0);
+		Sample.ColorSH.B.V1 = abs(Sample.ColorSH.B.V1);
+		Sample.ColorSH.B.V2 = abs(Sample.ColorSH.B.V2);
+	#endif
 	Sample.SampleCount            = abs(Sample.SampleCount);
 	Sample.MissCount              = abs(Sample.MissCount);
 	Sample.ClosestHitDistance     = abs(Sample.ClosestHitDistance);
@@ -153,6 +193,21 @@ FSSDSignalSample Sqrt(FSSDSignalSample Sample)
 	#if CONFIG_SIGNAL_COLOR
 		Sample.SceneColor         = sqrt(Sample.SceneColor);
 	#endif
+	#if CONFIG_SIGNAL_COLOR_SH && SPHERICAL_HARMONIC_ORDER == 2
+		Sample.ColorSH.R.V = sqrt(Sample.ColorSH.R.V);
+		Sample.ColorSH.G.V = sqrt(Sample.ColorSH.G.V);
+		Sample.ColorSH.B.V = sqrt(Sample.ColorSH.B.V);
+	#elif CONFIG_SIGNAL_COLOR_SH && SPHERICAL_HARMONIC_ORDER == 3
+		Sample.ColorSH.R.V0 = sqrt(Sample.ColorSH.R.V0);
+		Sample.ColorSH.R.V1 = sqrt(Sample.ColorSH.R.V1);
+		Sample.ColorSH.R.V2 = sqrt(Sample.ColorSH.R.V2);
+		Sample.ColorSH.G.V0 = sqrt(Sample.ColorSH.G.V0);
+		Sample.ColorSH.G.V1 = sqrt(Sample.ColorSH.G.V1);
+		Sample.ColorSH.G.V2 = sqrt(Sample.ColorSH.G.V2);
+		Sample.ColorSH.B.V0 = sqrt(Sample.ColorSH.B.V0);
+		Sample.ColorSH.B.V1 = sqrt(Sample.ColorSH.B.V1);
+		Sample.ColorSH.B.V2 = sqrt(Sample.ColorSH.B.V2);
+	#endif
 	Sample.SampleCount            = sqrt(Sample.SampleCount);
 	Sample.MissCount              = sqrt(Sample.MissCount);
 	Sample.ClosestHitDistance     = sqrt(Sample.ClosestHitDistance);
@@ -169,6 +224,21 @@ FSSDSignalSample Power(FSSDSignalSample Sample, float Exponent)
 		Sample.SceneColor.b       = pow(Sample.SceneColor.b, Exponent);
 		Sample.SceneColor.a       = pow(Sample.SceneColor.a, Exponent);
 	#endif
+	#if CONFIG_SIGNAL_COLOR_SH && SPHERICAL_HARMONIC_ORDER == 2
+		Sample.ColorSH.R.V = pow(Sample.ColorSH.R.V, Exponent);
+		Sample.ColorSH.G.V = pow(Sample.ColorSH.G.V, Exponent);
+		Sample.ColorSH.B.V = pow(Sample.ColorSH.B.V, Exponent);
+	#elif CONFIG_SIGNAL_COLOR_SH && SPHERICAL_HARMONIC_ORDER == 3
+		Sample.ColorSH.R.V0 = pow(Sample.ColorSH.R.V0, Exponent);
+		Sample.ColorSH.R.V1 = pow(Sample.ColorSH.R.V1, Exponent);
+		Sample.ColorSH.R.V2 = pow(Sample.ColorSH.R.V2, Exponent);
+		Sample.ColorSH.G.V0 = pow(Sample.ColorSH.G.V0, Exponent);
+		Sample.ColorSH.G.V1 = pow(Sample.ColorSH.G.V1, Exponent);
+		Sample.ColorSH.G.V2 = pow(Sample.ColorSH.G.V2, Exponent);
+		Sample.ColorSH.B.V0 = pow(Sample.ColorSH.B.V0, Exponent);
+		Sample.ColorSH.B.V1 = pow(Sample.ColorSH.B.V1, Exponent);
+		Sample.ColorSH.B.V2 = pow(Sample.ColorSH.B.V2, Exponent);
+	#endif
 	Sample.SampleCount            = pow(Sample.SampleCount, Exponent);
 	Sample.MissCount              = pow(Sample.MissCount, Exponent);
 	Sample.ClosestHitDistance     = pow(Sample.ClosestHitDistance, Exponent);
@@ -183,6 +253,21 @@ FSSDSignalSample Min(FSSDSignalSample SampleA, FSSDSignalSample SampleB)
 	#if CONFIG_SIGNAL_COLOR
 		OutSample.SceneColor = min(SampleA.SceneColor, SampleB.SceneColor);
 	#endif
+	#if CONFIG_SIGNAL_COLOR_SH && SPHERICAL_HARMONIC_ORDER == 2
+		OutSample.ColorSH.R.V = min(SampleA.ColorSH.R.V, SampleB.ColorSH.R.V);
+		OutSample.ColorSH.G.V = min(SampleA.ColorSH.G.V, SampleB.ColorSH.G.V);
+		OutSample.ColorSH.B.V = min(SampleA.ColorSH.B.V, SampleB.ColorSH.B.V);
+	#elif CONFIG_SIGNAL_COLOR_SH && SPHERICAL_HARMONIC_ORDER == 3
+		OutSample.ColorSH.R.V0 = min(SampleA.ColorSH.R.V0, SampleB.ColorSH.R.V0);
+		OutSample.ColorSH.R.V1 = min(SampleA.ColorSH.R.V1, SampleB.ColorSH.R.V1);
+		OutSample.ColorSH.R.V2 = min(SampleA.ColorSH.R.V2, SampleB.ColorSH.R.V2);
+		OutSample.ColorSH.G.V0 = min(SampleA.ColorSH.G.V0, SampleB.ColorSH.G.V0);
+		OutSample.ColorSH.G.V1 = min(SampleA.ColorSH.G.V1, SampleB.ColorSH.G.V1);
+		OutSample.ColorSH.G.V2 = min(SampleA.ColorSH.G.V2, SampleB.ColorSH.G.V2);
+		OutSample.ColorSH.B.V0 = min(SampleA.ColorSH.B.V0, SampleB.ColorSH.B.V0);
+		OutSample.ColorSH.B.V1 = min(SampleA.ColorSH.B.V1, SampleB.ColorSH.B.V1);
+		OutSample.ColorSH.B.V2 = min(SampleA.ColorSH.B.V2, SampleB.ColorSH.B.V2);
+	#endif
 	OutSample.SampleCount = min(SampleA.SampleCount, SampleB.SampleCount);
 	OutSample.MissCount = min(SampleA.MissCount, SampleB.MissCount);
 	OutSample.ClosestHitDistance = min(SampleA.ClosestHitDistance, SampleB.ClosestHitDistance);
@@ -197,6 +282,21 @@ FSSDSignalSample Max(FSSDSignalSample SampleA, FSSDSignalSample SampleB)
 	#if CONFIG_SIGNAL_COLOR
 		OutSample.SceneColor = max(SampleA.SceneColor, SampleB.SceneColor);
 	#endif
+	#if CONFIG_SIGNAL_COLOR_SH && SPHERICAL_HARMONIC_ORDER == 2
+		OutSample.ColorSH.R.V = max(SampleA.ColorSH.R.V, SampleB.ColorSH.R.V);
+		OutSample.ColorSH.G.V = max(SampleA.ColorSH.G.V, SampleB.ColorSH.G.V);
+		OutSample.ColorSH.B.V = max(SampleA.ColorSH.B.V, SampleB.ColorSH.B.V);
+	#elif CONFIG_SIGNAL_COLOR_SH && SPHERICAL_HARMONIC_ORDER == 3
+		OutSample.ColorSH.R.V0 = max(SampleA.ColorSH.R.V0, SampleB.ColorSH.R.V0);
+		OutSample.ColorSH.R.V1 = max(SampleA.ColorSH.R.V1, SampleB.ColorSH.R.V1);
+		OutSample.ColorSH.R.V2 = max(SampleA.ColorSH.R.V2, SampleB.ColorSH.R.V2);
+		OutSample.ColorSH.G.V0 = max(SampleA.ColorSH.G.V0, SampleB.ColorSH.G.V0);
+		OutSample.ColorSH.G.V1 = max(SampleA.ColorSH.G.V1, SampleB.ColorSH.G.V1);
+		OutSample.ColorSH.G.V2 = max(SampleA.ColorSH.G.V2, SampleB.ColorSH.G.V2);
+		OutSample.ColorSH.B.V0 = max(SampleA.ColorSH.B.V0, SampleB.ColorSH.B.V0);
+		OutSample.ColorSH.B.V1 = max(SampleA.ColorSH.B.V1, SampleB.ColorSH.B.V1);
+		OutSample.ColorSH.B.V2 = max(SampleA.ColorSH.B.V2, SampleB.ColorSH.B.V2);
+	#endif
 	OutSample.SampleCount = max(SampleA.SampleCount, SampleB.SampleCount);
 	OutSample.MissCount = max(SampleA.MissCount, SampleB.MissCount);
 	OutSample.ClosestHitDistance = max(SampleA.ClosestHitDistance, SampleB.ClosestHitDistance);
@@ -211,6 +311,21 @@ FSSDSignalSample Clamp(FSSDSignalSample Sample, FSSDSignalSample SampleMin, FSSD
 	#if CONFIG_SIGNAL_COLOR
 		OutSample.SceneColor = clamp(Sample.SceneColor, SampleMin.SceneColor, SampleMax.SceneColor);
 	#endif
+	#if CONFIG_SIGNAL_COLOR_SH && SPHERICAL_HARMONIC_ORDER == 2
+		OutSample.ColorSH.R.V = clamp(Sample.ColorSH.R.V, SampleMin.ColorSH.R.V, SampleMax.ColorSH.R.V);
+		OutSample.ColorSH.G.V = clamp(Sample.ColorSH.G.V, SampleMin.ColorSH.G.V, SampleMax.ColorSH.G.V);
+		OutSample.ColorSH.B.V = clamp(Sample.ColorSH.B.V, SampleMin.ColorSH.B.V, SampleMax.ColorSH.B.V);
+	#elif CONFIG_SIGNAL_COLOR_SH && SPHERICAL_HARMONIC_ORDER == 3
+		OutSample.ColorSH.R.V0 = clamp(Sample.ColorSH.R.V0, SampleMin.ColorSH.R.V0, SampleMax.ColorSH.R.V0);
+		OutSample.ColorSH.R.V1 = clamp(Sample.ColorSH.R.V1, SampleMin.ColorSH.R.V1, SampleMax.ColorSH.R.V1);
+		OutSample.ColorSH.R.V2 = clamp(Sample.ColorSH.R.V2, SampleMin.ColorSH.R.V2, SampleMax.ColorSH.R.V2);
+		OutSample.ColorSH.G.V0 = clamp(Sample.ColorSH.G.V0, SampleMin.ColorSH.G.V0, SampleMax.ColorSH.G.V0);
+		OutSample.ColorSH.G.V1 = clamp(Sample.ColorSH.G.V1, SampleMin.ColorSH.G.V1, SampleMax.ColorSH.G.V1);
+		OutSample.ColorSH.G.V2 = clamp(Sample.ColorSH.G.V2, SampleMin.ColorSH.G.V2, SampleMax.ColorSH.G.V2);
+		OutSample.ColorSH.B.V0 = clamp(Sample.ColorSH.B.V0, SampleMin.ColorSH.B.V0, SampleMax.ColorSH.B.V0);
+		OutSample.ColorSH.B.V1 = clamp(Sample.ColorSH.B.V1, SampleMin.ColorSH.B.V1, SampleMax.ColorSH.B.V1);
+		OutSample.ColorSH.B.V2 = clamp(Sample.ColorSH.B.V2, SampleMin.ColorSH.B.V2, SampleMax.ColorSH.B.V2);
+	#endif
 	OutSample.SampleCount = clamp(Sample.SampleCount, SampleMin.SampleCount, SampleMax.SampleCount);
 	OutSample.MissCount = clamp(Sample.MissCount, SampleMin.MissCount, SampleMax.MissCount);
 	OutSample.ClosestHitDistance = clamp(Sample.ClosestHitDistance, SampleMin.ClosestHitDistance, SampleMax.ClosestHitDistance);
diff --git a/Engine/Shaders/Private/ScreenSpaceDenoise/SSDSignalFramework.ush b/Engine/Shaders/Private/ScreenSpaceDenoise/SSDSignalFramework.ush
--- a/Engine/Shaders/Private/ScreenSpaceDenoise/SSDSignalFramework.ush
+++ b/Engine/Shaders/Private/ScreenSpaceDenoise/SSDSignalFramework.ush
@@ -44,6 +44,10 @@
 #include "../EyeAdaptationCommon.ush"
 #include "../LightShaderParameters.ush"
 
+#if CONFIG_SIGNAL_COLOR_SH
+	#include "../QuasiRandom.ush"
+#endif
+
 
 #if !defined(CONFIG_SIGNAL_PROCESSING)
 	#error Missing CONFIG_SIGNAL_PROCESSING
@@ -531,7 +535,7 @@ float ComputeBilateralWeight(
 	}
 	#elif CONFIG_VIEW_POSITION_BASED_BILATERAL == 2
 	{
-		float DistFromRefPlane = dot(RefSceneMetadata.WorldNormal, RefSceneMetadata.TranslatedWorldPosition - NeighborSceneMetadata.TranslatedWorldPosition);
+		float DistFromRefPlane = abs(dot(RefSceneMetadata.WorldNormal, RefSceneMetadata.TranslatedWorldPosition - NeighborSceneMetadata.TranslatedWorldPosition));
 
 		PositionBilateralWeight = saturate(1 - DistFromRefPlane * rcp(MaxWorldBluringDistance));
 	}
@@ -766,3 +770,31 @@ float GetRatioEstimatorWeight(
 	return 1.0;
 }
 #endif // CONFIG_SPECULAR_RATIO_ESTIMATOR
+
+
+// Matches GenerateCosineNormalRay()
+// TODO: expose a Weight that multiply directly onto the Sample.SceneColor to reduce ALU
+#if CONFIG_SIGNAL_COLOR_SH
+FSSDSphericalHarmonicRGB ComputeSampleColorSH(FSSDSampleSceneInfos SampleSceneMetadata, FSSDSignalSample Sample, uint2 NeighborPixelCoord)
+{
+	const uint SamplesPerPixel = 1;
+	FQuasiRandom QuasiRandom;
+	QuasiRandom_Initialize(QuasiRandom, NeighborPixelCoord, 0, SamplesPerPixel);
+	QuasiRandom_Update(QuasiRandom, /* SampleIndex = */ 0, SamplesPerPixel);
+
+	float3 RandSample = QuasiRandom_Sample3(QuasiRandom);
+
+	float4 TangentDirection = CosineSampleHemisphere(RandSample.yz);
+	float3 WorldDirection = TangentToWorld(TangentDirection.xyz, SampleSceneMetadata.WorldNormal);
+
+	#if SPHERICAL_HARMONIC_ORDER == 2
+		FSSDSphericalHarmonic SampleBasis = SHBasisFunction(WorldDirection);
+		return MulSH(SampleBasis, Sample.SceneColor.rgb);
+
+	#elif SPHERICAL_HARMONIC_ORDER == 3
+		FSSDSphericalHarmonic SampleBasis = SHBasisFunction3(WorldDirection);
+		return MulSH3(SampleBasis, Sample.SceneColor.rgb);
+
+	#endif
+}
+#endif
diff --git a/Engine/Shaders/Private/ScreenSpaceDenoise/SSDSpatialAccumulation.usf b/Engine/Shaders/Private/ScreenSpaceDenoise/SSDSpatialAccumulation.usf
--- a/Engine/Shaders/Private/ScreenSpaceDenoise/SSDSpatialAccumulation.usf
+++ b/Engine/Shaders/Private/ScreenSpaceDenoise/SSDSpatialAccumulation.usf
@@ -72,6 +72,8 @@
 		#define CONFIG_OUTPUT_TEXTURE_COUNT CONFIG_SIGNAL_BATCH_SIZE
 
 		#define CONFIG_SAMPLE_SET           SAMPLE_SET_HEXAWEB
+		#define CONFIG_CUSTOM_SPREAD_FACTOR 1
+
 		#define CONFIG_BILATERAL_DISTANCE_COMPUTATION SIGNAL_WORLD_FREQUENCY_PRECOMPUTED_BLURING_RADIUS
 		#define CONFIG_MAX_WITH_REF_DISTANCE 1
 		#define CONFIG_OUTPUT_MODE          OUTPUT_MODE_DRB
@@ -223,6 +225,7 @@
 
 #elif CONFIG_SIGNAL_PROCESSING == SIGNAL_PROCESSING_GI_AND_AO
 	#define CONFIG_SIGNAL_COLOR 1
+	#define CONFIG_SIGNAL_COLOR_SH 1
 
 	#define CONFIG_VIEW_POSITION_BASED_BILATERAL 2
 	#define CONFIG_NORMAL_BASED_BILATERAL 1
@@ -239,6 +242,15 @@
 
 		#define CONFIG_BILATERAL_DISTANCE_COMPUTATION SIGNAL_WORLD_FREQUENCY_HIT_DISTANCE
 
+	#elif DIM_STAGE == STAGE_PRE_CONVOLUTION
+		// Input and output layout.
+		#define CONFIG_SIGNAL_INPUT_LAYOUT  SIGNAL_BUFFER_LAYOUT_GI_AND_AO_RECONSTRUCTION
+		#define CONFIG_SIGNAL_OUTPUT_LAYOUT SIGNAL_BUFFER_LAYOUT_GI_AND_AO_RECONSTRUCTION
+
+		#define CONFIG_SAMPLE_SET           SAMPLE_SET_HEXAWEB
+		#define CONFIG_CUSTOM_SPREAD_FACTOR 1
+		//#define CONFIG_MAX_WITH_REF_DISTANCE 1
+
 	#elif DIM_STAGE == STAGE_POST_FILTERING
 		// Input and output layout.
 		#define CONFIG_SIGNAL_INPUT_LAYOUT  SIGNAL_BUFFER_LAYOUT_GI_AND_AO_HISTORY
@@ -493,9 +505,7 @@ void MainCS(
 		float RefBluringRadius = ComputeInfinityBluringRadius(RefSceneMetadata) * View.ViewSizeAndInvSize.x * rcp(KernelSpreadFactor);
 		float IdealSampleCount = PI * RefBluringRadius * RefBluringRadius;
 
-		const float MaxSampleCount = 1024; // TODO.
-
-		IdealSampleCount = min(IdealSampleCount, MaxSampleCount);
+		IdealSampleCount = min(IdealSampleCount, TARGETED_SAMPLE_COUNT);
 
 		RequestedSampleCount = IdealSampleCount / clamp(RefSamples.Array[0].SampleCount, 1, MaxSampleCount);
 	}
@@ -504,9 +514,7 @@ void MainCS(
 		#if CONFIG_SIGNAL_BATCH_SIZE != 1
 			#error Unable to support more than one signal.
 		#endif
-		const float IdealSampleCount = TARGETED_SAMPLE_COUNT;
-
-		RequestedSampleCount = clamp(IdealSampleCount / RefSamples.Array[0].SampleCount, 1, MaxSampleCount);
+		RequestedSampleCount = clamp(TARGETED_SAMPLE_COUNT / RefSamples.Array[0].SampleCount, 1, MaxSampleCount);
 	}
 	#else
 		#error Unknown policy to control the number of samples.
@@ -529,6 +537,9 @@ void MainCS(
 		KernelConfig.bNormalizeSample = CONFIG_NORMALIZE_INPUT != 0;
 		KernelConfig.bSampleKernelCenter = CONFIG_UPSCALE;
 		KernelConfig.bForceKernelCenterAccumulation = true;
+
+		// Reconstruct the spherical harmonic when reconstructing from 1spp.
+		KernelConfig.bComputeSampleColorSH = DIM_STAGE == STAGE_RECONSTRUCTION;
 		
 		{
 			UNROLL_N(MAX_SIGNAL_MULTIPLEXING)
@@ -579,7 +590,6 @@ void MainCS(
 		#if CONFIG_SAMPLE_SET == SAMPLE_SET_HEXAWEB
 		{
 			KernelConfig.RingCount = 1;
-			KernelConfig.KernelSpreadFactor = 8;
 
 			// TODO: could be improved.
 			//KernelConfig.bMinSamplePairInvFrequency = true;
@@ -601,11 +611,14 @@ void MainCS(
 			for (uint SignalMultiplexId = 0; SignalMultiplexId < MAX_SIGNAL_MULTIPLEXING; SignalMultiplexId++)
 			{
 				const uint BatchedSignalId = ComputeSignalBatchIdFromSignalMultiplexId(KernelConfig, SignalMultiplexId);
-
+				
+				uint2 RefPixelCoord = floor(KernelConfig.BufferUV * KernelConfig.BufferSizeAndInvSize.xy);
 				FSSDSignalSample CenterSample = TransformSignalSampleForAccumulation(
 					KernelConfig,
 					SignalMultiplexId,
-					RefSamples.Array[SignalMultiplexId]);
+					RefSceneMetadata,
+					RefSamples.Array[SignalMultiplexId],
+					RefPixelCoord);
 				
 				FSSDSampleAccumulationInfos SampleInfos;
 				SampleInfos.Sample = CenterSample;
@@ -803,7 +816,7 @@ void MainCS(
 		for (uint MultiplexId = 0; MultiplexId < CONFIG_SIGNAL_BATCH_SIZE; MultiplexId++)
 		{
 			float CurrentSampleCount = OutputSamples.Array[MultiplexId].SampleCount;
-			float NewSampleCount = min(CurrentSampleCount, 128);
+			float NewSampleCount = min(CurrentSampleCount, TARGETED_SAMPLE_COUNT);
 
 			OutputSamples.Array[MultiplexId] = Mul(OutputSamples.Array[MultiplexId], CurrentSampleCount > 0 ? NewSampleCount / CurrentSampleCount : 0);
 		}
diff --git a/Engine/Shaders/Private/ScreenSpaceDenoise/SSDSpatialKernel.ush b/Engine/Shaders/Private/ScreenSpaceDenoise/SSDSpatialKernel.ush
--- a/Engine/Shaders/Private/ScreenSpaceDenoise/SSDSpatialKernel.ush
+++ b/Engine/Shaders/Private/ScreenSpaceDenoise/SSDSpatialKernel.ush
@@ -91,6 +91,9 @@ struct FSSDKernelConfig
 	// Whether the bilateral distance should be maxed with reference bilateral distance.
 	bool bMaxWithRefBilateralDistance;
 
+	// Whether the spherical harmonic of a sample should be computed before accumulation.
+	bool bComputeSampleColorSH;
+
 	// The color space that has been encoded in the buffer.
 	uint BufferColorSpace[MAX_SIGNAL_MULTIPLEXING];
 	
@@ -154,6 +157,7 @@ FSSDKernelConfig CreateKernelConfig()
 	KernelConfig.bNormalizeSample = false;
 	KernelConfig.bMinSamplePairInvFrequency = false;
 	KernelConfig.bMaxWithRefBilateralDistance = false;
+	KernelConfig.bComputeSampleColorSH = false;
 	
 	{
 		UNROLL_N(MAX_SIGNAL_MULTIPLEXING)
@@ -201,13 +205,24 @@ static const float kWaveletFilterWeights5x5[] = { 3.0 / 8.0, 1.0 / 4.0, 1.0 / 16
 FSSDSignalSample TransformSignalSampleForAccumulation(
 	FSSDKernelConfig KernelConfig,
 	uint MultiplexId,
-	FSSDSignalSample Sample)
+	FSSDSampleSceneInfos SampleSceneMetadata,
+	FSSDSignalSample Sample,
+	uint2 SamplePixelCoord)
 {
+	// Transform the color space.
 	// TODO: could pass down information that this sample may be normalized.
 	Sample = TransformSignal(
 		Sample, 
 		/* SrcBasis  = */ KernelConfig.BufferColorSpace[MultiplexId], 
 		/* DestBasis = */ KernelConfig.AccumulatorColorSpace[MultiplexId]);
+	
+	// Compute the spherical harmonic of the sample.
+	#if CONFIG_SIGNAL_COLOR_SH
+	if (KernelConfig.bComputeSampleColorSH)
+	{
+		Sample.ColorSH = ComputeSampleColorSH(SampleSceneMetadata, Sample, SamplePixelCoord);
+	}
+	#endif
 
 	return Sample;
 }
@@ -286,11 +301,13 @@ void AccumulateSampledMultiplexedSignals(
 	{
 		// Compute at compile time the id of the signal being processed.
 		const uint BatchedSignalId = ComputeSignalBatchIdFromSignalMultiplexId(KernelConfig, SignalMultiplexId);
+		
+		// TODO: direction of the ray should be cached by injest or output by RGS, otherwise ends up with VGPR pressure because of SampleBufferUV.
+		uint2 NeighborPixelCoord = floor(SampleBufferUV * KernelConfig.BufferSizeAndInvSize.xy);
 
 		// Fetch and pre process the sample for accumulation.
 		FSSDSignalSample Sample = MultiplexedSamples.Array[SignalMultiplexId];
-		
-		Sample = TransformSignalSampleForAccumulation(KernelConfig, SignalMultiplexId, Sample);
+		Sample = TransformSignalSampleForAccumulation(KernelConfig, SignalMultiplexId, SampleSceneMetadata, Sample, NeighborPixelCoord);
 		
 		// Compute the bluring radius of pixel itself.
 		float SamplePixelWorldBluringRadius = ComputeWorldBluringRadiusCausedByPixelSize(SampleSceneMetadata);
@@ -339,8 +356,6 @@ void AccumulateSampledMultiplexedSignals(
 			KernelConfig.RefSceneMetadata,
 			SampleSceneMetadata);
 		
-		// TODO: direction of the ray should be cached by injest or output by RGS, otherwise ends up with VGPR pressure because of SampleBufferUV.
-		uint2 NeighborPixelCoord = floor(SampleBufferUV * KernelConfig.BufferSizeAndInvSize.xy);
 		float RatioEstimatorWeight = GetRatioEstimatorWeight(
 			KernelConfig.RefSceneMetadata, SampleSceneMetadata, Sample, NeighborPixelCoord);
 
diff --git a/Engine/Shaders/Private/ScreenSpaceDenoise/SSDTemporalAccumulation.usf b/Engine/Shaders/Private/ScreenSpaceDenoise/SSDTemporalAccumulation.usf
--- a/Engine/Shaders/Private/ScreenSpaceDenoise/SSDTemporalAccumulation.usf
+++ b/Engine/Shaders/Private/ScreenSpaceDenoise/SSDTemporalAccumulation.usf
@@ -392,14 +392,12 @@ void TemporallyAccumulate(
 		TargetedSampleCount = PI * SignalPixelBluringRadius * SignalPixelBluringRadius;
 
 		// TODO: try to make this better? * 0.125 to avoid banding caused by floating point precision.
-		TargetedSampleCount = clamp(TargetedSampleCount, 1, 0.125 * TARGETED_SAMPLE_COUNT);
+		TargetedSampleCount = clamp(TargetedSampleCount, 1, TARGETED_SAMPLE_COUNT);
 	}
-	#elif CONFIG_SIGNAL_PROCESSING == SIGNAL_PROCESSING_AO || CONFIG_SIGNAL_PROCESSING == SIGNAL_PROCESSING_GI_AND_AO || CONFIG_SIGNAL_PROCESSING == SIGNAL_PROCESSING_MONOCHROMATIC_PENUMBRA
+	#else
 	{
 		TargetedSampleCount = TARGETED_SAMPLE_COUNT;
 	}
-	#else
-		#error
 	#endif
 
 	// Sample the history buffer.
@@ -475,6 +473,14 @@ void TemporallyAccumulate(
 		} // if (!bIsPreviousFrameOffscreen)
 	}
 	
+	#if CONFIG_SIGNAL_PROCESSING == SIGNAL_PROCESSING_GI_AND_AO && 0
+		DebugOutput[DispatchThreadId] = float4(
+			HistorySamples.Array[0].SampleCount / 4096,
+			0,
+			0,
+			0);
+	#endif
+	
 	const bool bPostRejectionBlending = true;
 
 	// History rejection.
@@ -569,11 +575,14 @@ void TemporallyAccumulate(
 			for (uint SignalMultiplexId = 0; SignalMultiplexId < MAX_SIGNAL_MULTIPLEXING; SignalMultiplexId++)
 			{
 				const uint BatchedSignalId = ComputeSignalBatchIdFromSignalMultiplexId(KernelConfig, SignalMultiplexId);
-
+				
+				uint2 RefPixelCoord = floor(KernelConfig.BufferUV * KernelConfig.BufferSizeAndInvSize.xy);
 				FSSDSignalSample TransformedCurrentFrame = TransformSignalSampleForAccumulation(
 					KernelConfig,
 					SignalMultiplexId,
-					NormalizeToOneSample(CurrentFrameSamples.Array[SignalMultiplexId]));
+					RefSceneMetadata,
+					NormalizeToOneSample(CurrentFrameSamples.Array[SignalMultiplexId]),
+					RefPixelCoord);
 
 				FSSDSampleAccumulationInfos SampleInfos;
 				SampleInfos.Sample = TransformedCurrentFrame;
@@ -856,6 +865,7 @@ void TemporallyAccumulate(
 				SafeRcp(CurrentFrameSample.SampleCount) *
 				HistorySample.SampleCount);
 
+			if (View.GeneralPurposeTweak > 0.5)
 			HistorySamples.Array[BatchedSignalId] = HistorySample;
 		}
 	}
diff --git a/Engine/Source/Runtime/Renderer/Private/ScreenSpaceDenoise.cpp b/Engine/Source/Runtime/Renderer/Private/ScreenSpaceDenoise.cpp
--- a/Engine/Source/Runtime/Renderer/Private/ScreenSpaceDenoise.cpp
+++ b/Engine/Source/Runtime/Renderer/Private/ScreenSpaceDenoise.cpp
@@ -13,6 +13,7 @@
 #include "ClearQuad.h"
 #include "PipelineStateCache.h"
 #include "SceneViewFamilyBlackboard.h"
+#include "BlueNoise.h"
 
 
 // ---------------------------------------------------- Cvars
@@ -83,13 +84,18 @@ static TAutoConsoleVariable<int32> CVarGIReconstructionSampleCount(
 	TEXT("Maximum number of samples for the reconstruction pass (default = 16)."),
 	ECVF_RenderThreadSafe);
 
+static TAutoConsoleVariable<int32> CVarGIPreConvolutionCount(
+	TEXT("r.GlobalIllumination.Denoiser.PreConvolution"), 1,
+	TEXT("Number of pre-convolution passes (default = 1)."),
+	ECVF_RenderThreadSafe);
+
 static TAutoConsoleVariable<int32> CVarGITemporalAccumulation(
 	TEXT("r.GlobalIllumination.Denoiser.TemporalAccumulation"), 1,
 	TEXT("Accumulates the samples over multiple frames."),
 	ECVF_RenderThreadSafe);
 
 static TAutoConsoleVariable<int32> CVarGIHistoryConvolutionSampleCount(
-	TEXT("r.GlobalIllumination.Denoiser.HistoryConvolution.SampleCount"), 16,
+	TEXT("r.GlobalIllumination.Denoiser.HistoryConvolution.SampleCount"), 1,
 	TEXT("Number of samples to use for history post filter (default = 1)."),
 	ECVF_RenderThreadSafe);
 
@@ -163,7 +169,9 @@ static bool SignalUsesInjestion(ESignalProcessing SignalProcessing)
 /** Returns whether a signal processing uses an additional pre convolution pass. */
 static bool SignalUsesPreConvolution(ESignalProcessing SignalProcessing)
 {
-	return SignalProcessing == ESignalProcessing::MonochromaticPenumbra;
+	return
+		SignalProcessing == ESignalProcessing::MonochromaticPenumbra ||
+		SignalProcessing == ESignalProcessing::GlobalIllumination;
 }
 
 /** Returns whether a signal processing uses a history rejection pre convolution pass. */
@@ -292,8 +300,8 @@ const TCHAR* const kPreConvolutionResourceNames[] = {
 	nullptr,
 
 	// GlobalIllumination
-	nullptr,
-	nullptr,
+	TEXT("GIPreConvolution0"),
+	TEXT("GIPreConvolution1"),
 	nullptr,
 	nullptr,
 };
@@ -436,6 +444,10 @@ BEGIN_SHADER_PARAMETER_STRUCT(FSSDCommonParameters, )
 	SHADER_PARAMETER_RDG_TEXTURE(Texture2D, EyeAdaptation)
 	SHADER_PARAMETER_RDG_TEXTURE(Texture2D<uint>, TileClassificationTexture)
 	SHADER_PARAMETER_STRUCT_REF(FViewUniformShaderParameters, ViewUniformBuffer)
+	
+	SHADER_PARAMETER_STRUCT_REF(FHaltonIteration, HaltonIteration)
+	SHADER_PARAMETER_STRUCT_REF(FHaltonPrimes, HaltonPrimes)
+	SHADER_PARAMETER_STRUCT_REF(FBlueNoise, BlueNoise)
 END_SHADER_PARAMETER_STRUCT()
 
 /** Shader parameter structure use to bind all signal generically. */
@@ -625,7 +637,7 @@ class FSSDSpatialAccumulationCS : public FGlobalShader
 		SHADER_PARAMETER(uint32, MaxSampleCount)
 		SHADER_PARAMETER(int32, UpscaleFactor)
 		SHADER_PARAMETER(float, KernelSpreadFactor)
-
+		
 		SHADER_PARAMETER_STRUCT_INCLUDE(FSSDCommonParameters, CommonParameters)
 		SHADER_PARAMETER_STRUCT_INCLUDE(FSSDConvolutionMetaData, ConvolutionMetaData)
 
@@ -873,6 +885,31 @@ static void DenoiseSignalAtConstantPixelDensity(
 		CommonParameters.EyeAdaptation = GetEyeAdaptationTexture(GraphBuilder, View);
 	}
 
+	#if RHI_RAYTRACING
+	if (Settings.SignalProcessing == ESignalProcessing::GlobalIllumination)
+	{
+		uint32 IterationCount = Settings.MaxInputSPP;
+		uint32 SequenceCount = 1;
+		uint32 DimensionCount = 24;
+
+		FScene* Scene = static_cast<FScene*>(View.Family->Scene);
+
+		FHaltonSequenceIteration HaltonSequenceIteration(Scene->HaltonSequence, IterationCount, SequenceCount, DimensionCount, View.ViewState ? (View.ViewState->FrameIndex % 1024) : 0);
+		FHaltonIteration HaltonIteration;
+		InitializeHaltonSequenceIteration(HaltonSequenceIteration, HaltonIteration);
+
+		FHaltonPrimes HaltonPrimes;
+		InitializeHaltonPrimes(Scene->HaltonPrimesResource, HaltonPrimes);
+
+		FBlueNoise BlueNoise;
+		InitializeBlueNoise(BlueNoise);
+		
+		CommonParameters.HaltonIteration = CreateUniformBufferImmediate(HaltonIteration, EUniformBufferUsage::UniformBuffer_SingleFrame);
+		CommonParameters.HaltonPrimes = CreateUniformBufferImmediate(HaltonPrimes, EUniformBufferUsage::UniformBuffer_SingleFrame);
+		CommonParameters.BlueNoise = CreateUniformBufferImmediate(BlueNoise, EUniformBufferUsage::UniformBuffer_SingleFrame);
+	}
+	#endif // RHI_RAYTRACING
+
 	// Setup all the metadata to do spatial convolution.
 	FSSDConvolutionMetaData ConvolutionMetaData;
 	if (Settings.SignalProcessing == ESignalProcessing::MonochromaticPenumbra)
@@ -978,6 +1015,7 @@ static void DenoiseSignalAtConstantPixelDensity(
 		FSSDSpatialAccumulationCS::FParameters* PassParameters = GraphBuilder.AllocParameters<FSSDSpatialAccumulationCS::FParameters>();
 		PassParameters->CommonParameters = CommonParameters;
 		PassParameters->ConvolutionMetaData = ConvolutionMetaData;
+		PassParameters->KernelSpreadFactor = 8 * (1 << PreConvolutionId);
 		PassParameters->SignalInput = CopyAndBackfillSignalInput(SignalHistory);
 		PassParameters->SignalOutput = CreateMultiplexedUAVs(GraphBuilder, NewSignalOutput);
 		
@@ -992,7 +1030,7 @@ static void DenoiseSignalAtConstantPixelDensity(
 		TShaderMapRef<FSSDSpatialAccumulationCS> ComputeShader(View.ShaderMap, PermutationVector);
 		FComputeShaderUtils::AddPass(
 			GraphBuilder,
-			RDG_EVENT_NAME("SSD PreConvolution(MaxSamples=7)"),
+			RDG_EVENT_NAME("SSD PreConvolution(MaxSamples=7 Spread=%d)", PassParameters->KernelSpreadFactor),
 			*ComputeShader,
 			PassParameters,
 			FComputeShaderUtils::GetGroupCount(DenoiseResolution, FSSDSpatialAccumulationCS::kGroupSize));
@@ -1425,6 +1463,7 @@ public:
 		Settings.bUseTemporalAccumulation = CVarAOTemporalAccumulation.GetValueOnRenderThread() != 0;
 		Settings.HistoryConvolutionSampleCount = CVarAOHistoryConvolutionSampleCount.GetValueOnRenderThread();
 		Settings.HistoryConvolutionKernelSpreadFactor = CVarAOHistoryConvolutionKernelSpreadFactor.GetValueOnRenderThread();
+		Settings.MaxInputSPP = RayTracingConfig.RayCountPerPixel;
 
 		TStaticArray<FScreenSpaceFilteringHistory*, IScreenSpaceDenoiser::kMaxBatchSize> PrevHistories;
 		TStaticArray<FScreenSpaceFilteringHistory*, IScreenSpaceDenoiser::kMaxBatchSize> NewHistories;
@@ -1462,9 +1501,11 @@ public:
 		Settings.SignalProcessing = ESignalProcessing::GlobalIllumination;
 		Settings.InputResolutionFraction = Config.ResolutionFraction;
 		Settings.ReconstructionSamples = CVarGIReconstructionSampleCount.GetValueOnRenderThread();
+		Settings.PreConvolutionCount = CVarGIPreConvolutionCount.GetValueOnRenderThread();
 		Settings.bUseTemporalAccumulation = CVarGITemporalAccumulation.GetValueOnRenderThread() != 0;
 		Settings.HistoryConvolutionSampleCount = CVarGIHistoryConvolutionSampleCount.GetValueOnRenderThread();
 		Settings.HistoryConvolutionKernelSpreadFactor = CVarGIHistoryConvolutionKernelSpreadFactor.GetValueOnRenderThread();
+		Settings.MaxInputSPP = Config.RayCountPerPixel;
 
 		TStaticArray<FScreenSpaceFilteringHistory*, IScreenSpaceDenoiser::kMaxBatchSize> PrevHistories;
 		TStaticArray<FScreenSpaceFilteringHistory*, IScreenSpaceDenoiser::kMaxBatchSize> NewHistories;
@@ -1502,9 +1543,11 @@ public:
 		Settings.SignalProcessing = ESignalProcessing::GlobalIllumination;
 		Settings.InputResolutionFraction = Config.ResolutionFraction;
 		Settings.ReconstructionSamples = CVarGIReconstructionSampleCount.GetValueOnRenderThread();
+		Settings.PreConvolutionCount = CVarGIPreConvolutionCount.GetValueOnRenderThread();
 		Settings.bUseTemporalAccumulation = CVarGITemporalAccumulation.GetValueOnRenderThread() != 0;
 		Settings.HistoryConvolutionSampleCount = CVarGIHistoryConvolutionSampleCount.GetValueOnRenderThread();
 		Settings.HistoryConvolutionKernelSpreadFactor = CVarGIHistoryConvolutionKernelSpreadFactor.GetValueOnRenderThread();
+		Settings.MaxInputSPP = Config.RayCountPerPixel;
 
 		TStaticArray<FScreenSpaceFilteringHistory*, IScreenSpaceDenoiser::kMaxBatchSize> PrevHistories;
 		TStaticArray<FScreenSpaceFilteringHistory*, IScreenSpaceDenoiser::kMaxBatchSize> NewHistories;
