Add world normall accessor to vertex factories
Fix slope bias on terrain (Landscape)

diff --git a/Engine/Shaders/Private/GeometryCacheVertexFactory.ush b/Engine/Shaders/Private/GeometryCacheVertexFactory.ush
--- a/Engine/Shaders/Private/GeometryCacheVertexFactory.ush
+++ b/Engine/Shaders/Private/GeometryCacheVertexFactory.ush
@@ -277,6 +277,11 @@ FVertexFactoryIntermediates GetVertexFactoryIntermediates(FVertexFactoryInput In
 	return Intermediates;
 }
 
+float3 VertexFactoryGetWorldNormal(FVertexFactoryInput Input, FVertexFactoryIntermediates Intermediates)
+{
+	return Intermediates.TangentToWorld[2];
+}
+
 /**
 * Get the 3x3 tangent basis vectors for this vertex factory
 *
diff --git a/Engine/Shaders/Private/LandscapeVertexFactory.ush b/Engine/Shaders/Private/LandscapeVertexFactory.ush
--- a/Engine/Shaders/Private/LandscapeVertexFactory.ush
+++ b/Engine/Shaders/Private/LandscapeVertexFactory.ush
@@ -187,7 +187,7 @@ float4 VertexFactoryGetWorldPosition(FVertexFactoryInput Input, FVertexFactoryIn
 	return INVARIANT(TransformLocalToTranslatedWorld(GetLocalPosition(Intermediates)));
 }
 
-float3 VertexFactoryGetWorldNormal(FVertexFactoryIntermediates Intermediates)
+float3 VertexFactoryGetWorldNormal(FVertexFactoryInput Input, FVertexFactoryIntermediates Intermediates)
 {
 	return Intermediates.WorldNormal;
 }
diff --git a/Engine/Shaders/Private/LocalVertexFactory.ush b/Engine/Shaders/Private/LocalVertexFactory.ush
--- a/Engine/Shaders/Private/LocalVertexFactory.ush
+++ b/Engine/Shaders/Private/LocalVertexFactory.ush
@@ -939,7 +939,7 @@ float3 VertexFactoryGetWorldNormal(FPositionAndNormalOnlyVertexFactoryInput Inpu
 }
 
 
-float3 VertexFactoryGetWorldNormal(FVertexFactoryIntermediates Intermediates)
+float3 VertexFactoryGetWorldNormal(FVertexFactoryInput Input, FVertexFactoryIntermediates Intermediates)
 {
 	return Intermediates.TangentToWorld[2];
 }
diff --git a/Engine/Shaders/Private/MeshParticleVertexFactory.ush b/Engine/Shaders/Private/MeshParticleVertexFactory.ush
--- a/Engine/Shaders/Private/MeshParticleVertexFactory.ush
+++ b/Engine/Shaders/Private/MeshParticleVertexFactory.ush
@@ -463,6 +463,11 @@ float4 VertexFactoryGetWorldPosition(FVertexFactoryInput Input, FVertexFactoryIn
 	return CalcWorldPosition(Input);
 }
 
+float3 VertexFactoryGetWorldNormal(FVertexFactoryInput Input, FVertexFactoryIntermediates Intermediates)
+{
+	return TangentBias(Input.TangentZ);
+}
+
 float4 VertexFactoryGetRasterizedWorldPosition(FVertexFactoryInput Input, FVertexFactoryIntermediates Intermediates, float4 InWorldPosition)
 {
 	return InWorldPosition;
diff --git a/Engine/Shaders/Private/NiagaraMeshVertexFactory.ush b/Engine/Shaders/Private/NiagaraMeshVertexFactory.ush
--- a/Engine/Shaders/Private/NiagaraMeshVertexFactory.ush
+++ b/Engine/Shaders/Private/NiagaraMeshVertexFactory.ush
@@ -689,6 +689,11 @@ float4 VertexFactoryGetRasterizedWorldPosition(FVertexFactoryInput Input, FVerte
 	return InWorldPosition;
 }
 
+float3 VertexFactoryGetWorldNormal(FVertexFactoryInput Input, FVertexFactoryIntermediates Intermediates)
+{
+	return Intermediates.LocalToWorld[2].xyz;
+}
+
 float3 VertexFactoryGetPositionForVertexLighting(FVertexFactoryInput Input, FVertexFactoryIntermediates Intermediates, float3 TranslatedWorldPosition)
 {
 	return TranslatedWorldPosition;
diff --git a/Engine/Shaders/Private/NiagaraRibbonVertexFactory.ush b/Engine/Shaders/Private/NiagaraRibbonVertexFactory.ush
--- a/Engine/Shaders/Private/NiagaraRibbonVertexFactory.ush
+++ b/Engine/Shaders/Private/NiagaraRibbonVertexFactory.ush
@@ -642,6 +642,12 @@ float4 VertexFactoryGetWorldPosition(FVertexFactoryInput Input, FVertexFactoryIn
 	return CalcWorldPosition(Input, Intermediates);
 }
 
+float3 VertexFactoryGetWorldNormal(FVertexFactoryInput Input, FVertexFactoryIntermediates Intermediates)
+{
+	return Intermediates.Facing;
+}
+
+
 float4 VertexFactoryGetRasterizedWorldPosition(FVertexFactoryInput Input, FVertexFactoryIntermediates Intermediates, float4 InWorldPosition)
 {
 	return InWorldPosition;
diff --git a/Engine/Shaders/Private/NiagaraSpriteVertexFactory.ush b/Engine/Shaders/Private/NiagaraSpriteVertexFactory.ush
--- a/Engine/Shaders/Private/NiagaraSpriteVertexFactory.ush
+++ b/Engine/Shaders/Private/NiagaraSpriteVertexFactory.ush
@@ -872,6 +872,11 @@ float4 VertexFactoryGetWorldPosition(FVertexFactoryInput Input, FVertexFactoryIn
 	return float4(Intermediates.VertexWorldPosition,1);
 }
 
+float3 VertexFactoryGetWorldNormal(FVertexFactoryInput Input, FVertexFactoryIntermediates Intermediates)
+{
+	return CalcTangentBasis(Intermediates)[2].xyz;
+}
+
 float4 VertexFactoryGetRasterizedWorldPosition(FVertexFactoryInput Input, FVertexFactoryIntermediates Intermediates, float4 TranslatedWorldPosition)
 {
 #if SPHERICAL_PARTICLE_OPACITY
diff --git a/Engine/Shaders/Private/ParticleBeamTrailVertexFactory.ush b/Engine/Shaders/Private/ParticleBeamTrailVertexFactory.ush
--- a/Engine/Shaders/Private/ParticleBeamTrailVertexFactory.ush
+++ b/Engine/Shaders/Private/ParticleBeamTrailVertexFactory.ush
@@ -177,6 +177,11 @@ float4 VertexFactoryGetWorldPosition(FVertexFactoryInput Input, FVertexFactoryIn
 	return CalcWorldPosition(Input);
 }
 
+float3 VertexFactoryGetWorldNormal(FVertexFactoryInput Input, FVertexFactoryIntermediates Intermediates)
+{
+	return CalcTangentBasis(Input)[2];
+}
+
 float4 VertexFactoryGetRasterizedWorldPosition(FVertexFactoryInput Input, FVertexFactoryIntermediates Intermediates, float4 InWorldPosition)
 {
 	return InWorldPosition;
diff --git a/Engine/Shaders/Private/ParticleGPUSpriteVertexFactory.ush b/Engine/Shaders/Private/ParticleGPUSpriteVertexFactory.ush
--- a/Engine/Shaders/Private/ParticleGPUSpriteVertexFactory.ush
+++ b/Engine/Shaders/Private/ParticleGPUSpriteVertexFactory.ush
@@ -589,6 +589,11 @@ float3x3 VertexFactoryGetTangentToLocal( FVertexFactoryInput Input, FVertexFacto
 	return TangentToLocal;
 }
 
+float3 VertexFactoryGetWorldNormal(FVertexFactoryInput Input, FVertexFactoryIntermediates Intermediates)
+{
+	return VertexFactoryGetTangentToLocal(Input, Intermediates)[2].xyz;
+}
+
 /**
  * Constructs values that need to be interpolated from the vertex shader to the pixel shader.
  * @param Input - Vertex attributes.
diff --git a/Engine/Shaders/Private/ParticleSpriteVertexFactory.ush b/Engine/Shaders/Private/ParticleSpriteVertexFactory.ush
--- a/Engine/Shaders/Private/ParticleSpriteVertexFactory.ush
+++ b/Engine/Shaders/Private/ParticleSpriteVertexFactory.ush
@@ -557,6 +557,11 @@ float4 VertexFactoryGetWorldPosition(FVertexFactoryInput Input, FVertexFactoryIn
 	return float4(Intermediates.VertexWorldPosition,1);
 }
 
+float3 VertexFactoryGetWorldNormal(FVertexFactoryInput Input, FVertexFactoryIntermediates Intermediates)
+{
+	return CalcTangentBasis(Intermediates)[2];
+}
+
 float4 VertexFactoryGetRasterizedWorldPosition(FVertexFactoryInput Input, FVertexFactoryIntermediates Intermediates, float4 TranslatedWorldPosition)
 {
 #if SPHERICAL_PARTICLE_OPACITY
diff --git a/Engine/Shaders/Private/ShadowDepthVertexShader.usf b/Engine/Shaders/Private/ShadowDepthVertexShader.usf
--- a/Engine/Shaders/Private/ShadowDepthVertexShader.usf
+++ b/Engine/Shaders/Private/ShadowDepthVertexShader.usf
@@ -285,7 +285,8 @@ void Main(
 	float3x3 TangentToLocal = VertexFactoryGetTangentToLocal(Input, VFIntermediates);
 
 	FMaterialVertexParameters VertexParameters = GetMaterialVertexParameters(Input, VFIntermediates, WorldPos.xyz, TangentToLocal);
-	const float3 WorldNormal = VertexParameters.TangentToWorld[2];
+	const float3 WorldNormal = VertexFactoryGetWorldNormal(Input, VFIntermediates);
+
 	WorldPos.xyz += GetMaterialWorldPositionOffset(VertexParameters);
 
 #if USING_TESSELLATION
diff --git a/Engine/Shaders/Private/VectorFieldVisualizationVertexFactory.ush b/Engine/Shaders/Private/VectorFieldVisualizationVertexFactory.ush
--- a/Engine/Shaders/Private/VectorFieldVisualizationVertexFactory.ush
+++ b/Engine/Shaders/Private/VectorFieldVisualizationVertexFactory.ush
@@ -126,11 +126,16 @@ FMaterialVertexParameters GetMaterialVertexParameters(FVertexFactoryInput Input,
  * @param Intermediates - Intermediates computed for this vertex.
  * @returns the position of this vertex in world space.
  */
-float4 VertexFactoryGetWorldPosition( FVertexFactoryInput Input, FVertexFactoryIntermediates Intermediates )
+float4 VertexFactoryGetWorldPosition(FVertexFactoryInput Input, FVertexFactoryIntermediates Intermediates )
 {
 	return Intermediates.WorldPosition;
 }
 
+float3 VertexFactoryGetWorldNormal(FVertexFactoryInput Input, FVertexFactoryIntermediates Intermediates)
+{
+	return float3(0,0,1);
+}
+
 float4 VertexFactoryGetRasterizedWorldPosition(FVertexFactoryInput Input, FVertexFactoryIntermediates Intermediates, float4 TranslatedWorldPosition)
 {
 	return TranslatedWorldPosition;
