Working around for NVIDIA driver failures in 430.39 and 430.64.

diff --git a/Engine/Shaders/Private/RayTracing/RayTracingAmbientOcclusionRGS.usf b/Engine/Shaders/Private/RayTracing/RayTracingAmbientOcclusionRGS.usf
--- a/Engine/Shaders/Private/RayTracing/RayTracingAmbientOcclusionRGS.usf
+++ b/Engine/Shaders/Private/RayTracing/RayTracingAmbientOcclusionRGS.usf
@@ -60,6 +60,8 @@ void AmbientOcclusionRGS()
 	float RayCount = 0.0;
 	float Visibility = 0.0;
 	float ClosestRayHitDistance = 1.#INF;
+	// Declaring intensity as a local variable is a workaround for a shader compiler bug with NV drivers 430.39 and 430.64.
+	const float IntensityLocal = Intensity;
 	uint SamplesPerPixelLocal = SamplesPerPixel;
 	
 	// Mask out depth values beyond far plane
@@ -106,7 +108,7 @@ void AmbientOcclusionRGS()
 		);
 		
 		RayCount += 1;
-		Visibility += Payload.HitT < 0 ? 1.0 : 1.0 - Intensity;
+		Visibility += Payload.HitT < 0 ? 1.0 : 1.0 - IntensityLocal;
 		if (Payload.HitT >= 0.0)
 		{
 			ClosestRayHitDistance = min(ClosestRayHitDistance, Payload.HitT);
