Adding uniform hemispherical sampling debugging option to SkyLight importance sampling.

diff --git a/Engine/Shaders/Private/PathTracing/Light/PathTracingSkyLight.ush b/Engine/Shaders/Private/PathTracing/Light/PathTracingSkyLight.ush
--- a/Engine/Shaders/Private/PathTracing/Light/PathTracingSkyLight.ush
+++ b/Engine/Shaders/Private/PathTracing/Light/PathTracingSkyLight.ush
@@ -8,7 +8,7 @@ SkyLight.usf: Light sampling functions for SkyLight implementation
 
 #include "../../RayTracing/SkyLightMipTreeCommon.ush"
 
-#define USE_IMPORTANCE_SAMPLING
+#define USE_HIERARCHICAL_IMPORTANCE_SAMPLING
 
 void SkyLight_GenerateLightRay(
 	RayDesc Ray,
@@ -30,8 +30,8 @@ void SkyLight_PdfLight(
 	out float OutPdf
 )
 {
-#ifndef USE_IMPORTANCE_SAMPLING
-	OutPdf = 0.0; return;
+#ifndef USE_HIERARCHICAL_IMPORTANCE_SAMPLING
+	OutPdf = 1.0 / (2.0 * PI);
 #else
 	// Determine primary direction
 	//uint MipCount = Log2(SkyLight.MipDimensions.x);
@@ -70,14 +70,10 @@ void SkyLight_EstimateLight(
 	out float OutIrradiance
 )
 {
-#ifndef USE_IMPORTANCE_SAMPLING
-	OutIrradiance = 0.0; return;
-#else
 	uint Width, Height, NumberOfLevels;
 	SkyLight.Texture.GetDimensions(0, Width, Height, NumberOfLevels);
 	float LightPower = Luminance(GetColor(LightId) * TextureCubeSampleLevel(SkyLight.Texture, SkyLight.TextureSampler, float3(0.0, 1.0, 0.0), NumberOfLevels - 1).xyz);
 	OutIrradiance = LightPower * 2.0 * PI;
-#endif
 }
 
 void SkyLight_SampleLight(
@@ -89,8 +85,10 @@ void SkyLight_SampleLight(
 	out float OutPdf
 )
 {
-#ifndef USE_IMPORTANCE_SAMPLING
-	OutPdf = 0.0; return;
+#ifndef USE_HIERARCHICAL_IMPORTANCE_SAMPLING
+	float4 Direction_Tangent = UniformSampleHemisphere(RandSample.yz);
+	OutLightUV = TangentToWorld(Direction_Tangent.xyz, Payload.WorldNormal);
+	OutPdf = Direction_Tangent.w;
 #else
 	float3 WorldPosition = Ray.Origin;
 	float3 WorldNormal = Payload.WorldNormal;
