Allow virtual texture stacks made from UDIMs with different block dimensions
- UDIM textures forced to wrap address mode, out-of-bounds vAddresses are automatically wrapped during VT system update
- All VTs create mips down to 1x1 pixel (clamped at max number of mips based on max page table size), mips smaller than 1 tile are needed when stacking UDIMs of differing sizes
- Low resolution VT tiles are locked by allocated VT rather than producer, since lowest mip required may be different for different allocated VT stacks
- This requires mechanism for ref-counting locks, since the same tile may be locked by multiple sources
- UDIM scale is now uniform, rather than baked into material, since we now allow it to change in material instances

diff --git a/Engine/Shaders/Private/PageTableUpdate.usf b/Engine/Shaders/Private/PageTableUpdate.usf
--- a/Engine/Shaders/Private/PageTableUpdate.usf
+++ b/Engine/Shaders/Private/PageTableUpdate.usf
@@ -2,8 +2,6 @@
 
 #include "Common.ush"
 
-#define VIRTUALTEXTURE_MAX_PAGETABLE_SIZE 2048u
-
 uint PageTableSize;
 uint FirstUpdate;
 uint NumUpdates;
@@ -71,10 +69,12 @@ void PageTableUpdateVS(
 	Page |= pPage.x << 4;
 	Page |= pPage.y << 10;
 #else
-	// vPagesWideInLevel is an integer quantity, but we store as a half-float, since VT sampling code needs value as floating point
-	// we can extract half-float from lower 16 bits to a full float vgpr with a single instruction
-	const float vPagesWideInLevel = float(VIRTUALTEXTURE_MAX_PAGETABLE_SIZE >> vLevel);
-	uint Page = f32tof16(vPagesWideInLevel);
+	// We want to store the quantity vPagesWideInLevel (computed as VIRTUALTEXTURE_MAX_PAGETABLE_SIZE >> vLevel)
+	// We then want to scale by (1/VIRTUALTEXTURE_MAX_PAGETABLE_SIZE), to avoid the need to explicitly factor this into other uniforms
+	// So scale factor becomes (VIRTUALTEXTURE_MAX_PAGETABLE_SIZE >> vLevel) / VIRTUALTEXTURE_MAX_PAGETABLE_SIZE, which simplifies to (1 / exp2(vLevel)) or (1 / (1 << vLevel))
+	// Also nice that this avoids need to depend on specific value of VIRTUALTEXTURE_MAX_PAGETABLE_SIZE
+
+	uint Page = f32tof16(1.0f / float(1u << vLevel));
 	Page |= pPage.x << 16;
 	Page |= pPage.y << 24;
 #endif
diff --git a/Engine/Shaders/Private/VirtualTextureCommon.ush b/Engine/Shaders/Private/VirtualTextureCommon.ush
--- a/Engine/Shaders/Private/VirtualTextureCommon.ush
+++ b/Engine/Shaders/Private/VirtualTextureCommon.ush
@@ -185,33 +185,72 @@ struct VTPageTableResult
 	uint PackedRequest;
 };
 
-VTPageTableResult TextureLoadVirtualPageTableLevel(Texture2D<uint4> PageTable0,
-	float2 UV, float Level,
+uint TextureComputeVirtualMipLevel(inout VTPageTableResult OutResult,
+	float2 dUVdx, float2 dUVdy, float MipBias,
 	float2 SvPositionXY,
+	VTPageTableUniform PageTableUniform)
+{
+	OutResult.dUVdx = dUVdx * PageTableUniform.SizeInPages;
+	OutResult.dUVdy = dUVdy * PageTableUniform.SizeInPages;
+
+	const float Noise = InterleavedGradientNoise(SvPositionXY, View.StateFrameIndexMod8);
+	const float ComputedLevel = MipLevelAniso2D(OutResult.dUVdx, OutResult.dUVdy, PageTableUniform.MaxAnisoLog2) + MipBias + Noise * 0.5f - 0.25f;
+	return clamp(int(ComputedLevel) + int(PageTableUniform.vPageTableMipBias), 0, int(PageTableUniform.MaxLevel));
+}
+
+void TextureLoadVirtualPageTableInternal(inout VTPageTableResult OutResult,
+	Texture2D<uint4> PageTable0,
+	float2 UV, uint vLevel,
 	VTPageTableUniform PageTableUniform,
 	uint AddressU, uint AddressV)
 {
-	VTPageTableResult Result;
+	UV.x = ApplyAddressMode(UV.x, AddressU);
+	UV.y = ApplyAddressMode(UV.y, AddressV);
+	OutResult.UV = UV * PageTableUniform.SizeInPages;
 
-	Result.dUVdx = float2(0.0f, 0.0f);
-	Result.dUVdy = float2(0.0f, 0.0f);
-	const uint vLevel = clamp(int(Level), 0, int(PageTableUniform.MaxLevel));
+	const uint vPageX = (uint(OutResult.UV.x) + PageTableUniform.XOffsetInPages) >> vLevel;
+	const uint vPageY = (uint(OutResult.UV.y) + PageTableUniform.YOffsetInPages) >> vLevel;
+	OutResult.PageTableValue[0] = PageTable0.Load(int3(vPageX, vPageY, vLevel));
+	OutResult.PageTableValue[1] = uint4(0u, 0u, 0u, 0u);
+
+	// PageTableID packed in upper 4 bits of 'PackedPageTableUniform', which is the bit position we want it in for PackedRequest as well, just need to mask off extra bits
+	OutResult.PackedRequest = PageTableUniform.ShiftedPageTableID;
+	OutResult.PackedRequest |= vPageX;
+	OutResult.PackedRequest |= vPageY << 12;
+	OutResult.PackedRequest |= vLevel << 24;
+}
 
+void TextureLoadVirtualPageTableInternal(inout VTPageTableResult OutResult,
+	Texture2D<uint4> PageTable0, Texture2D<uint4> PageTable1,
+	float2 UV, uint vLevel,
+	VTPageTableUniform PageTableUniform,
+	uint AddressU, uint AddressV)
+{
 	UV.x = ApplyAddressMode(UV.x, AddressU);
 	UV.y = ApplyAddressMode(UV.y, AddressV);
-	Result.UV = UV * PageTableUniform.UVScale;
+	OutResult.UV = UV * PageTableUniform.SizeInPages;
 
-	const uint vPageX = (uint(UV.x * PageTableUniform.SizeInPages.x) + PageTableUniform.XOffsetInPages) >> vLevel;
-	const uint vPageY = (uint(UV.y * PageTableUniform.SizeInPages.y) + PageTableUniform.YOffsetInPages) >> vLevel;
-	Result.PageTableValue[0] = PageTable0.Load(int3(vPageX, vPageY, vLevel));
-	Result.PageTableValue[1] = uint4(0u, 0u, 0u, 0u);
+	const uint vPageX = (uint(OutResult.UV.x) + PageTableUniform.XOffsetInPages) >> vLevel;
+	const uint vPageY = (uint(OutResult.UV.y) + PageTableUniform.YOffsetInPages) >> vLevel;
+	OutResult.PageTableValue[0] = PageTable0.Load(int3(vPageX, vPageY, vLevel));
+	OutResult.PageTableValue[1] = PageTable1.Load(int3(vPageX, vPageY, vLevel));
 
 	// PageTableID packed in upper 4 bits of 'PackedPageTableUniform', which is the bit position we want it in for PackedRequest as well, just need to mask off extra bits
-	Result.PackedRequest = PageTableUniform.ShiftedPageTableID;
-	Result.PackedRequest |= vPageX;
-	Result.PackedRequest |= vPageY << 12;
-	Result.PackedRequest |= vLevel << 24;
+	OutResult.PackedRequest = PageTableUniform.ShiftedPageTableID;
+	OutResult.PackedRequest |= vPageX;
+	OutResult.PackedRequest |= vPageY << 12;
+	OutResult.PackedRequest |= vLevel << 24;
+}
 
+VTPageTableResult TextureLoadVirtualPageTableLevel(Texture2D<uint4> PageTable0,
+	float2 UV, float Level,
+	float2 SvPositionXY,
+	VTPageTableUniform PageTableUniform,
+	uint AddressU, uint AddressV)
+{
+	VTPageTableResult Result = (VTPageTableResult)0.0f;
+	const uint vLevel = clamp(int(Level), 0, int(PageTableUniform.MaxLevel));
+	TextureLoadVirtualPageTableInternal(Result, PageTable0, UV * PageTableUniform.UVScale, vLevel, PageTableUniform, AddressU, AddressV);
 	return Result;
 }
 
@@ -221,27 +260,9 @@ VTPageTableResult TextureLoadVirtualPageTableLevel(Texture2D<uint4> PageTable0,
 	VTPageTableUniform PageTableUniform,
 	uint AddressU, uint AddressV)
 {
-	VTPageTableResult Result;
-
-	Result.dUVdx = float2(0.0f, 0.0f);
-	Result.dUVdy = float2(0.0f, 0.0f);
+	VTPageTableResult Result = (VTPageTableResult)0.0f;
 	const uint vLevel = clamp(int(Level), 0, int(PageTableUniform.MaxLevel));
-
-	UV.x = ApplyAddressMode(UV.x, AddressU);
-	UV.y = ApplyAddressMode(UV.y, AddressV);
-	Result.UV = UV * PageTableUniform.UVScale;
-
-	const uint vPageX = (uint(UV.x * PageTableUniform.SizeInPages.x) + PageTableUniform.XOffsetInPages) >> vLevel;
-	const uint vPageY = (uint(UV.y * PageTableUniform.SizeInPages.y) + PageTableUniform.YOffsetInPages) >> vLevel;
-	Result.PageTableValue[0] = PageTable0.Load(int3(vPageX, vPageY, vLevel));
-	Result.PageTableValue[1] = PageTable1.Load(int3(vPageX, vPageY, vLevel));
-
-	// PageTableID packed in upper 4 bits of 'PackedPageTableUniform', which is the bit position we want it in for PackedRequest as well, just need to mask off extra bits
-	Result.PackedRequest = PageTableUniform.ShiftedPageTableID;
-	Result.PackedRequest |= vPageX;
-	Result.PackedRequest |= vPageY << 12;
-	Result.PackedRequest |= vLevel << 24;
-
+	TextureLoadVirtualPageTableInternal(Result, PageTable0, PageTable1, UV * PageTableUniform.UVScale, vLevel, PageTableUniform, AddressU, AddressV);
 	return Result;
 }
 
@@ -251,30 +272,9 @@ VTPageTableResult TextureLoadVirtualPageTableGrad(Texture2D<uint4> PageTable0,
 	VTPageTableUniform PageTableUniform,
 	uint AddressU, uint AddressV)
 {
-	VTPageTableResult Result;
-
-	Result.dUVdx = dUVdx * PageTableUniform.UVScale;
-	Result.dUVdy = dUVdy * PageTableUniform.UVScale;
-
-	const float Noise = InterleavedGradientNoise(SvPositionXY, View.StateFrameIndexMod8);
-	const float ComputedLevel = MipLevelAniso2D(Result.dUVdx, Result.dUVdy, PageTableUniform.MaxAnisoLog2) + MipBias + Noise * 0.5f - 0.25f;
-	const uint vLevel = clamp(int(ComputedLevel) + int(PageTableUniform.vPageTableMipBias), 0, int(PageTableUniform.MaxLevel));
-
-	UV.x = ApplyAddressMode(UV.x, AddressU);
-	UV.y = ApplyAddressMode(UV.y, AddressV);
-	Result.UV = UV * PageTableUniform.UVScale;
-
-	const uint vPageX = (uint(UV.x * PageTableUniform.SizeInPages.x) + PageTableUniform.XOffsetInPages) >> vLevel;
-	const uint vPageY = (uint(UV.y * PageTableUniform.SizeInPages.y) + PageTableUniform.YOffsetInPages) >> vLevel;
-	Result.PageTableValue[0] = PageTable0.Load(int3(vPageX, vPageY, vLevel));
-	Result.PageTableValue[1] = uint4(0u, 0u, 0u, 0u);
-
-	// PageTableID packed in upper 4 bits of 'PackedPageTableUniform', which is the bit position we want it in for PackedRequest as well, just need to mask off extra bits
-	Result.PackedRequest = PageTableUniform.ShiftedPageTableID;
-	Result.PackedRequest |= vPageX;
-	Result.PackedRequest |= vPageY << 12;
-	Result.PackedRequest |= vLevel << 24;
-
+	VTPageTableResult Result = (VTPageTableResult)0.0f;
+	const uint vLevel = TextureComputeVirtualMipLevel(Result, dUVdx * PageTableUniform.UVScale, dUVdy * PageTableUniform.UVScale, MipBias, SvPositionXY, PageTableUniform);
+	TextureLoadVirtualPageTableInternal(Result, PageTable0, UV * PageTableUniform.UVScale, vLevel, PageTableUniform, AddressU, AddressV);
 	return Result;
 }
 
@@ -284,30 +284,9 @@ VTPageTableResult TextureLoadVirtualPageTableGrad(Texture2D<uint4> PageTable0, T
 	VTPageTableUniform PageTableUniform,
 	uint AddressU, uint AddressV)
 {
-	VTPageTableResult Result;
-
-	Result.dUVdx = dUVdx * PageTableUniform.UVScale;
-	Result.dUVdy = dUVdy * PageTableUniform.UVScale;
-
-	const float Noise = InterleavedGradientNoise(SvPositionXY, View.StateFrameIndexMod8);
-	const float ComputedLevel = MipLevelAniso2D(Result.dUVdx, Result.dUVdy, PageTableUniform.MaxAnisoLog2) + MipBias + Noise * 0.5f - 0.25f;
-	const uint vLevel = clamp(int(ComputedLevel) + int(PageTableUniform.vPageTableMipBias), 0, int(PageTableUniform.MaxLevel));
-
-	UV.x = ApplyAddressMode(UV.x, AddressU);
-	UV.y = ApplyAddressMode(UV.y, AddressV);
-	Result.UV = UV * PageTableUniform.UVScale;
-
-	const uint vPageX = (uint(UV.x * PageTableUniform.SizeInPages.x) + PageTableUniform.XOffsetInPages) >> vLevel;
-	const uint vPageY = (uint(UV.y * PageTableUniform.SizeInPages.y) + PageTableUniform.YOffsetInPages) >> vLevel;
-	Result.PageTableValue[0] = PageTable0.Load(int3(vPageX, vPageY, vLevel));
-	Result.PageTableValue[1] = PageTable1.Load(int3(vPageX, vPageY, vLevel));
-
-	// PageTableID packed in upper 4 bits of 'PackedPageTableUniform', which is the bit position we want it in for PackedRequest as well, just need to mask off extra bits
-	Result.PackedRequest = PageTableUniform.ShiftedPageTableID;
-	Result.PackedRequest |= vPageX;
-	Result.PackedRequest |= vPageY << 12;
-	Result.PackedRequest |= vLevel << 24;
-
+	VTPageTableResult Result = (VTPageTableResult)0.0f;
+	const uint vLevel = TextureComputeVirtualMipLevel(Result, dUVdx * PageTableUniform.UVScale, dUVdy * PageTableUniform.UVScale, MipBias, SvPositionXY, PageTableUniform);
+	TextureLoadVirtualPageTableInternal(Result, PageTable0, PageTable1, UV * PageTableUniform.UVScale, vLevel, PageTableUniform, AddressU, AddressV);
 	return Result;
 }
 
@@ -317,11 +296,14 @@ VTPageTableResult TextureLoadVirtualPageTable(Texture2D<uint4> PageTable0,
 	VTPageTableUniform PageTableUniform,
 	uint AddressU, uint AddressV)
 {
+	VTPageTableResult Result = (VTPageTableResult)0.0f;
+	const float2 ScaledUV = UV * PageTableUniform.UVScale;
+	uint vLevel = 0u;
 #if PIXELSHADER
-	return TextureLoadVirtualPageTableGrad(PageTable0, UV, ddx(UV), ddy(UV), MipBias, SvPositionXY, PageTableUniform, AddressU, AddressV);
-#else
-	return TextureLoadVirtualPageTableLevel(PageTable0, UV, 0.0f, SvPositionXY, PageTableUniform, AddressU, AddressV);
-#endif
+	vLevel = TextureComputeVirtualMipLevel(Result, ddx(ScaledUV), ddy(ScaledUV), MipBias, SvPositionXY, PageTableUniform);
+#endif // PIXELSHADER
+	TextureLoadVirtualPageTableInternal(Result, PageTable0, ScaledUV, vLevel, PageTableUniform, AddressU, AddressV);
+	return Result;
 }
 
 VTPageTableResult TextureLoadVirtualPageTable(Texture2D<uint4> PageTable0, Texture2D<uint4> PageTable1,
@@ -330,11 +312,14 @@ VTPageTableResult TextureLoadVirtualPageTable(Texture2D<uint4> PageTable0, Textu
 	VTPageTableUniform PageTableUniform,
 	uint AddressU, uint AddressV)
 {
+	VTPageTableResult Result = (VTPageTableResult)0.0f;
+	const float2 ScaledUV = UV * PageTableUniform.UVScale;
+	uint vLevel = 0u;
 #if PIXELSHADER
-	return TextureLoadVirtualPageTableGrad(PageTable0, PageTable1, UV, ddx(UV), ddy(UV), MipBias, SvPositionXY, PageTableUniform, AddressU, AddressV);
-#else
-	return TextureLoadVirtualPageTableLevel(PageTable0, PageTable1, UV, 0.0f, SvPositionXY, PageTableUniform, AddressU, AddressV);
-#endif
+	vLevel = TextureComputeVirtualMipLevel(Result, ddx(ScaledUV), ddy(ScaledUV), MipBias, SvPositionXY, PageTableUniform);
+#endif // PIXELSHADER
+	TextureLoadVirtualPageTableInternal(Result, PageTable0, PageTable1, ScaledUV, vLevel, PageTableUniform, AddressU, AddressV);
+	return Result;
 }
 
 // Overloads generated by material translator...makes request and stores feedback
@@ -423,8 +408,6 @@ struct VTUniform
 	float pPageSize;
 	float vPageSize;
 	float vPageBorderSize;
-	//uint PageTableSizeInPages;
-	//uint PageTableFormat;
 };
 
 VTUniform VTUniform_Unpack(uint4 PackedUniform)
@@ -433,8 +416,6 @@ VTUniform VTUniform_Unpack(uint4 PackedUniform)
 	result.pPageSize = asfloat(PackedUniform.w);
 	result.vPageSize = asfloat(PackedUniform.y);
 	result.vPageBorderSize = asfloat(PackedUniform.z);
-	//result.PageTableSizeInPages = PackedUniform.x & 0xfff;
-	//result.PageTableFormat = (PackedUniform.x >> 12) & 0x1;
 	return result;
 }
 
@@ -450,14 +431,14 @@ float2 VTComputePhysicalUVs(in out VTPageTableResult PageTableResult, uint Layer
 {
 	const uint PackedPageTableValue = PageTableResult.PageTableValue[LayerIndex / 4u][LayerIndex & 3u];
 
-	const float vPagesWideInLevel = f16tof32(PackedPageTableValue);
+	const float UVScale = f16tof32(PackedPageTableValue);
 	const float pPageX = UnpackByte2(PackedPageTableValue);
 	const float pPageY = UnpackByte3(PackedPageTableValue);
 
-	const float2 vPageFrac = frac(PageTableResult.UV * vPagesWideInLevel);
+	const float2 vPageFrac = frac(PageTableResult.UV * UVScale);
 	const float2 pUV = float2(pPageX, pPageY) * Uniform.pPageSize + (vPageFrac * Uniform.vPageSize + Uniform.vPageBorderSize);
 
-	const float ddxyScale = vPagesWideInLevel * Uniform.vPageSize;
+	const float ddxyScale = UVScale * Uniform.vPageSize;
 	PageTableResult.dUVdx *= ddxyScale;
 	PageTableResult.dUVdy *= ddxyScale;
 	return pUV;
@@ -483,90 +464,6 @@ MaterialFloat4 TextureVirtualSampleLevel(
 	return Physical.SampleLevel(PhysicalSampler, pUV, 0.0f);	
 }
 
-/**
-	Sample a virtual texture with automatic miplevel determination.
-	
-	@param Physical: Physical cache texture
-	@param PhysicalSampler: Sampler to use for the physical texture. This sampler will affect the final filtered result (nearest, bilinear, ...)
-	@param PageTable: Page table texture
-	@param PageTableSampler: Sampler to use when sampling the page table. This should be a nearest sampler.
-	@param VirtualTextureUniformData: Info members describing the VT
-	@param SvPositionXY: Screen space pixel position
-	@param FeedbackParams: An initialized FVirtualTextureFeedbackParams structure. If you do not want to generate feedback for this request a dummy can be passed in.
-	@param SampleIndex: The index of the VT sample in the shader (E.g. the 5th vt sample).
-	@param OriginalUV: The original UV's. Without Clamping or wrapping applied. Can be in the full float range.
-	@param UV: The address-mode-corrected UV to sample the VT at. This Should always be in the [0...1[ range. I.E. any clamping/wrapping behavior should be done before calling this function.
-	
-	Notes:
-	- SvPositionXY is used to seed the stochastic trilinear filtering. 0,0 can be safely passed in for non-pixel shader based shader types.
-	- SampleIndex is only used for the FeedbackParams so a dummy can be passed in if the output of the FeedbackParams is not used,
-	
-*/
-MaterialFloat4 TextureVirtualSample(
-	Texture2D Physical, SamplerState PhysicalSampler,
-	Texture2D<uint4> PageTable,
-	VTPageTableUniform PageTableUniform, VTUniform Uniform, float2 SvPositionXY,
-	in out FVirtualTextureFeedbackParams Feedback, uint SampleIndex,
-	float2 UV, uint AddressU, uint AddressV)
-{
-	
-	VTPageTableResult PageTableResult = TextureLoadVirtualPageTable(PageTable, UV, 0.0f, SvPositionXY, PageTableUniform, AddressU, AddressV);
-	StoreVirtualTextureFeedback(Feedback, SampleIndex, PageTableResult.PackedRequest);
-	return TextureVirtualSample( Physical, PhysicalSampler, PageTableResult, 0u, Uniform);
-}
-
-/**
-	Sample a virtual texture with a fixed mip-bias.
-	See TextureVirtualSample for a description of the arguments.
-*/
-MaterialFloat4 TextureVirtualSampleBias(
-	Texture2D Physical, SamplerState PhysicalSampler,
-	Texture2D<uint4> PageTable,
-	VTPageTableUniform PageTableUniform, VTUniform Uniform, float2 SvPositionXY,
-	in out FVirtualTextureFeedbackParams Feedback, uint SampleIndex,
-	float2 UV, uint AddressU, uint AddressV,
-	MaterialFloat MipBias)
-{
-	VTPageTableResult PageTableResult = TextureLoadVirtualPageTable(PageTable, UV, MipBias, SvPositionXY, PageTableUniform, AddressU, AddressV);
-	StoreVirtualTextureFeedback(Feedback, SampleIndex, PageTableResult.PackedRequest);
-	return TextureVirtualSample( Physical, PhysicalSampler, PageTableResult, 0u, Uniform);
-}
-
-/**
-	Sample a virtual texture on a fixed mipmap level.
-	See TextureVirtualSample for a description of the arguments.
-*/
-MaterialFloat4 TextureVirtualSampleLevel(
-	Texture2D Physical, SamplerState PhysicalSampler,
-	Texture2D<uint4> PageTable,
-	VTPageTableUniform PageTableUniform, VTUniform Uniform, float2 SvPositionXY,
-	in out FVirtualTextureFeedbackParams Feedback, uint SampleIndex,
-	float2 UV, uint AddressU, uint AddressV,
-	MaterialFloat Level)
-{
-	VTPageTableResult PageTableResult = TextureLoadVirtualPageTableLevel(PageTable, UV, Level, SvPositionXY, PageTableUniform, AddressU, AddressV);
-	StoreVirtualTextureFeedback(Feedback, SampleIndex, PageTableResult.PackedRequest);
-	return TextureVirtualSampleLevel(Physical, PhysicalSampler, PageTableResult, 0u, Uniform);
-}
-
-
-/**
-	Sample a virtual texture with automatic mipmap selection based on the passed in texture derivatives.
-	See TextureVirtualSample for a description of the arguments.
-*/
-MaterialFloat4 TextureVirtualSampleGrad(
-	Texture2D Physical, SamplerState PhysicalSampler,
-	Texture2D<uint4> PageTable,
-	VTPageTableUniform PageTableUniform, VTUniform Uniform, float2 SvPositionXY,
-	in out FVirtualTextureFeedbackParams Feedback, uint SampleIndex,
-	float2 UV, uint AddressU, uint AddressV,
-	MaterialFloat2 dUVdx, MaterialFloat2 dUVdy)
-{
-	VTPageTableResult PageTableResult = TextureLoadVirtualPageTableGrad(PageTable, UV, dUVdx, dUVdy, 0.0f, SvPositionXY, PageTableUniform, AddressU, AddressV);
-	StoreVirtualTextureFeedback(Feedback, SampleIndex, PageTableResult.PackedRequest);
-	return TextureVirtualSample( Physical, PhysicalSampler, PageTableResult, 0u, Uniform);
-}
-
 /**
 	Get the the index of the miplevel that would be sampled when doing a TextureVirtualSample
 	This may be lower than the mipmap the shader wanted to sample if data is currently not available at the desired miplevel.
diff --git a/Engine/Source/Editor/UnrealEd/Private/Texture2DPreview.cpp b/Engine/Source/Editor/UnrealEd/Private/Texture2DPreview.cpp
--- a/Engine/Source/Editor/UnrealEd/Private/Texture2DPreview.cpp
+++ b/Engine/Source/Editor/UnrealEd/Private/Texture2DPreview.cpp
@@ -83,8 +83,9 @@ public:
 
 			FUintVector4 PageTableUniform[2];
 			FUintVector4 Uniform;
-			VirtualTextureValue->GetPackedPageTableUniform(PageTableUniform);
-			VirtualTextureValue->GetPackedPhysicalTextureUniform(&Uniform, (uint32)LayerIndex);
+
+			AllocatedVT->GetPackedPageTableUniform(PageTableUniform, false);
+			AllocatedVT->GetPackedUniform(&Uniform, (uint32)LayerIndex);
 
 			SetShaderValueArray(RHICmdList, GetPixelShader(), VTPackedPageTableUniform, PageTableUniform, ARRAY_COUNT(PageTableUniform));
 			SetShaderValue(RHICmdList, GetPixelShader(), VTPackedUniform, Uniform);
diff --git a/Engine/Source/Runtime/Engine/Private/LightmapUniformShaderParameters.cpp b/Engine/Source/Runtime/Engine/Private/LightmapUniformShaderParameters.cpp
--- a/Engine/Source/Runtime/Engine/Private/LightmapUniformShaderParameters.cpp
+++ b/Engine/Source/Runtime/Engine/Private/LightmapUniformShaderParameters.cpp
@@ -4,7 +4,6 @@
 #include "SceneManagement.h"
 #include "LightMap.h"
 #include "VT/VirtualTexture.h"
-#include "VT/VirtualTextureUtility.h"
 #include "UnrealEngine.h"
 
 IMPLEMENT_GLOBAL_SHADER_PARAMETER_STRUCT(FPrecomputedLightingUniformParameters, "PrecomputedLightingBuffer");
@@ -103,11 +102,11 @@ void GetPrecomputedLightingParameters(
 			IAllocatedVirtualTexture* AllocatedVT = ResourceCluster->AcquireAllocatedVT();
 			check(AllocatedVT);
 
-			VTGetPackedPageTableUniform(&Parameters.LightmapVTPackedPageTableUniform[0], AllocatedVT);
+			AllocatedVT->GetPackedPageTableUniform(&Parameters.LightmapVTPackedPageTableUniform[0], true);
 			NumLightmapVTLayers = AllocatedVT->GetNumLayers();
 			for (uint32 LayerIndex = 0u; LayerIndex < NumLightmapVTLayers; ++LayerIndex)
 			{
-				VTGetPackedUniform(&Parameters.LightmapVTPackedUniform[LayerIndex], AllocatedVT, LayerIndex);
+				AllocatedVT->GetPackedUniform(&Parameters.LightmapVTPackedUniform[LayerIndex], LayerIndex);
 			}
 		}
 		else
diff --git a/Engine/Source/Runtime/Engine/Private/Materials/HLSLMaterialTranslator.h b/Engine/Source/Runtime/Engine/Private/Materials/HLSLMaterialTranslator.h
--- a/Engine/Source/Runtime/Engine/Private/Materials/HLSLMaterialTranslator.h
+++ b/Engine/Source/Runtime/Engine/Private/Materials/HLSLMaterialTranslator.h
@@ -4279,24 +4279,34 @@ protected:
 
 			TextureAddress AddressU = TA_Wrap;
 			TextureAddress AddressV = TA_Wrap;
-			switch (SamplerSource)
+			if (Tex2D && Tex2D->Source.GetNumBlocks() > 1)
 			{
-			case SSM_FromTextureAsset:
-				check(Tex2D);
-				AddressU = Tex2D->AddressX;
-				AddressV = Tex2D->AddressY;
-				break;
-			case SSM_Wrap_WorldGroupSettings:
+				// UDIM (multi-block) texture are forced to use wrap address mode
+				// This is important for supporting VT stacks made from UDIMs with differing number of blocks, as this requires wrapping vAddress for certain layers
 				AddressU = TA_Wrap;
 				AddressV = TA_Wrap;
-				break;
-			case SSM_Clamp_WorldGroupSettings:
-				AddressU = TA_Clamp;
-				AddressV = TA_Clamp;
-				break;
-			default:
-				checkNoEntry();
-				break;
+			}
+			else
+			{
+				switch (SamplerSource)
+				{
+				case SSM_FromTextureAsset:
+					check(Tex2D);
+					AddressU = Tex2D->AddressX;
+					AddressV = Tex2D->AddressY;
+					break;
+				case SSM_Wrap_WorldGroupSettings:
+					AddressU = TA_Wrap;
+					AddressV = TA_Wrap;
+					break;
+				case SSM_Clamp_WorldGroupSettings:
+					AddressU = TA_Clamp;
+					AddressV = TA_Clamp;
+					break;
+				default:
+					checkNoEntry();
+					break;
+				}
 			}
 
 			VTLayerIndex = MaterialCompilationOutput.UniformExpressionSet.UniformVirtualTextureExpressions[VirtualTextureIndex]->GetLayerIndex();
@@ -4312,19 +4322,11 @@ protected:
 				// This also means that any texture parameters set in material instances for VTs must match the aspect ratio of the texture in the parent material
 				// (Otherwise could potentially break stacks)
 				check(Tex2D);
-				const float TextureAspectRatio = (float)Tex2D->GetSizeX() / (float)Tex2D->GetSizeY();
 
-				const FIntPoint SizeInBlocks = Tex2D->Source.GetSizeInBlocks();
-				if (SizeInBlocks.X > 1 || SizeInBlocks.Y > 1)
-				{
-					// Transform the UVs to UDIM space if needed
-					// This will bake the proper UDIM dimensions into the material shader, which means these dimensions must match for any textures set in material instances
-					// This doesn't create any additional restrictions however, since it's also a requirement that UDIM dimensions must match for all layers of a VT stack
-					// That means that if we let overridden textures change UDIM dimensions, that could potentially break any stacks that are already baked into material
-					const float UVScaleX = (SizeInBlocks.X == 0) ? 1.0f : (1.0f / (float)SizeInBlocks.X);
-					const float UVScaleY = (SizeInBlocks.Y == 0) ? 1.0f : (1.0f / (float)SizeInBlocks.Y);
-					CoordinateIndex = Mul(CoordinateIndex, Constant2(UVScaleX, UVScaleY));
-				}
+				// Using Source size because we care about the aspect ratio of each block (each block of multi-block texture must have same aspect ratio)
+				// We can still combine multi-block textures of different block aspect ratios, as long as each block has the same ratio
+				// This is because we only need to overlay VT pages from within a given block
+				const float TextureAspectRatio = (float)Tex2D->Source.GetSizeX() / (float)Tex2D->Source.GetSizeY();
 
 				// Create a page table sample for each new set of sample parameters
 				VTStackIndex = AcquireVTStackIndex(MipValueMode, AddressU, AddressV, TextureAspectRatio, CoordinateIndex, MipValue0Index, MipValue1Index, INDEX_NONE);
diff --git a/Engine/Source/Runtime/Engine/Private/Materials/MaterialUniformExpressions.cpp b/Engine/Source/Runtime/Engine/Private/Materials/MaterialUniformExpressions.cpp
--- a/Engine/Source/Runtime/Engine/Private/Materials/MaterialUniformExpressions.cpp
+++ b/Engine/Source/Runtime/Engine/Private/Materials/MaterialUniformExpressions.cpp
@@ -16,7 +16,6 @@
 #include "RenderCore.h"
 #include "VirtualTexturing.h"
 #include "VT/RuntimeVirtualTexture.h"
-#include "VT/VirtualTextureUtility.h"
 
 static TAutoConsoleVariable<int32> CVarSupportMaterialLayers(
 	TEXT("r.SupportMaterialLayers"),
@@ -451,7 +450,15 @@ void FUniformExpressionSet::FillUniformBuffer(const FMaterialRenderContext& Mate
 		{
 			const IAllocatedVirtualTexture* AllocatedVT = UniformExpressionCache.AllocatedVTs[VTStackIndex];
 			FUintVector4* VTPackedPageTableUniform = (FUintVector4*)BufferCursor;
-			VTGetPackedPageTableUniform(VTPackedPageTableUniform, AllocatedVT);
+			if (AllocatedVT)
+			{
+				AllocatedVT->GetPackedPageTableUniform(VTPackedPageTableUniform, true);
+			}
+			else
+			{
+				VTPackedPageTableUniform[0] = FUintVector4(ForceInitToZero);
+				VTPackedPageTableUniform[1] = FUintVector4(ForceInitToZero);
+			}
 			BufferCursor = VTPackedPageTableUniform + 2;
 		}
 		
@@ -472,8 +479,10 @@ void FUniformExpressionSet::FillUniformBuffer(const FMaterialRenderContext& Mate
 				{
 					const FVTPackedStackAndLayerIndex StackAndLayerIndex = GetVTStackAndLayerIndex(ExpressionIndex);
 					const IAllocatedVirtualTexture* AllocatedVT = UniformExpressionCache.AllocatedVTs[StackAndLayerIndex.StackIndex];
-					VTGetPackedUniform(VTPackedUniform, AllocatedVT, StackAndLayerIndex.LayerIndex);
-
+					if (AllocatedVT)
+					{
+						AllocatedVT->GetPackedUniform(VTPackedUniform, StackAndLayerIndex.LayerIndex);
+					}
 					bFoundTexture = true;
 				}
 			}
@@ -487,7 +496,10 @@ void FUniformExpressionSet::FillUniformBuffer(const FMaterialRenderContext& Mate
 				{
 					int32 LayerIndex = UniformVirtualTextureExpressions[ExpressionIndex]->GetLayerIndex();
 					IAllocatedVirtualTexture const* AllocatedVT = Texture->GetAllocatedVirtualTexture();
-					VTGetPackedUniform(VTPackedUniform, AllocatedVT, LayerIndex);
+					if (AllocatedVT)
+					{
+						AllocatedVT->GetPackedUniform(VTPackedUniform, LayerIndex);
+					}
 				}
 			}
 		}
diff --git a/Engine/Source/Runtime/Engine/Private/Texture2D.cpp b/Engine/Source/Runtime/Engine/Private/Texture2D.cpp
--- a/Engine/Source/Runtime/Engine/Private/Texture2D.cpp
+++ b/Engine/Source/Runtime/Engine/Private/Texture2D.cpp
@@ -37,7 +37,6 @@
 
 #include "VT/UploadingVirtualTexture.h"
 #include "VT/VirtualTexturePoolConfig.h"
-#include "VT/VirtualTextureUtility.h"
 
 #if WITH_EDITOR
 #include "Settings/EditorExperimentalSettings.h"
@@ -1830,19 +1829,20 @@ void FTexture2DResource::GetData( uint32 MipIndex, void* Dest, uint32 DestPitch
 
 
 
-FVirtualTexture2DResource::FVirtualTexture2DResource(const UTexture2D* InOwner, FVirtualTextureBuiltData* InVTData, int32 SetFirstMipToUse)
+FVirtualTexture2DResource::FVirtualTexture2DResource(const UTexture2D* InOwner, FVirtualTextureBuiltData* InVTData, int32 InFirstMipToUse)
 	: AllocatedVT(nullptr)
 	, VTData(InVTData)
 	, TextureOwner(InOwner)
 {
 	check(InOwner);
 	check(InVTData);
-	FirstMipToUse = FMath::Min((int32)VTData->GetNumMips() - 1, SetFirstMipToUse);
-	if (InOwner)
-	{
-		// InOwner may be NULL in the case of virtual lightmap
-		bSRGB = InOwner->SRGB;
-	}
+
+	// Don't allow input mip bias to drop size below a single tile
+	const uint32 SizeInTiles = FMath::Max(VTData->GetWidthInTiles(), VTData->GetHeightInTiles());
+	const uint32 MaxMip = FMath::CeilLogTwo(SizeInTiles);
+	FirstMipToUse = FMath::Min((int32)MaxMip, InFirstMipToUse);
+
+	bSRGB = InOwner->SRGB;
 }
 
 FVirtualTexture2DResource::~FVirtualTexture2DResource()
@@ -1882,8 +1882,10 @@ void FVirtualTexture2DResource::InitRHI()
 	ProducerDesc.Dimensions = 2;
 	ProducerDesc.TileSize = VTData->TileSize;
 	ProducerDesc.TileBorderSize = VTData->TileBorderSize;
-	ProducerDesc.WidthInTiles = GetNumTilesX();
-	ProducerDesc.HeightInTiles = GetNumTilesY();
+	ProducerDesc.BlockWidthInTiles = FMath::DivideAndRoundUp<uint32>(GetNumTilesX(), VTData->WidthInBlocks);
+	ProducerDesc.BlockHeightInTiles = FMath::DivideAndRoundUp<uint32>(GetNumTilesY(), VTData->HeightInBlocks);
+	ProducerDesc.WidthInBlocks = VTData->WidthInBlocks;
+	ProducerDesc.HeightInBlocks = VTData->HeightInBlocks;
 	ProducerDesc.DepthInTiles = 1u;
 	ProducerDesc.MaxLevel = MaxLevel;
 	ProducerDesc.NumLayers = VTData->GetNumLayers();
@@ -1915,12 +1917,21 @@ void FVirtualTexture2DResource::InitializeEditorResources(IVirtualTexture* InVir
 			uint32 TileY;
 		};
 
+		// Choose a mip level for the thumbnail texture to ensure proper size
+		const uint32 MaxMipLevel = VTData->GetNumMips() - 1u;
+		uint32 MipLevel = 0;
+		uint32 MipWidth = GetSizeX();
+		uint32 MipHeight = GetSizeY();
+		while ((MipWidth > 128u && MipHeight > 128u) && MipLevel < MaxMipLevel)
+		{
+			++MipLevel;
+			MipWidth = FMath::Max(MipWidth / 2u, 1u);
+			MipHeight = FMath::Max(MipHeight / 2u, 1u);
+		}
+
 		const EPixelFormat PixelFormat = VTData->LayerTypes[0];
-		const uint32 MaxLevel = GetNumMips() - 1u; // Could potentially use a more detailed mip level here, if thumbnails are too low resolution
-		const uint32 MipWidth = FMath::Max(GetSizeX() >> MaxLevel, 1u);
-		const uint32 MipHeight = FMath::Max(GetSizeY() >> MaxLevel, 1u);
-		const uint32 MipWidthInTiles = FMath::Max(GetNumTilesX() >> MaxLevel, 1u);
-		const uint32 MipHeightInTiles = FMath::Max(GetNumTilesY() >> MaxLevel, 1u);
+		const uint32 MipWidthInTiles = FMath::Max(GetNumTilesX() >> MipLevel, 1u);
+		const uint32 MipHeightInTiles = FMath::Max(GetNumTilesY() >> MipLevel, 1u);
 		const uint32 TileSizeInPixels = GetTileSize();
 		const uint32 LayerMask = 1u; // FVirtualTexture2DResource should only have a single layer
 
@@ -1931,7 +1942,7 @@ void FVirtualTexture2DResource::InitializeEditorResources(IVirtualTexture* InVir
 			for (uint32 TileX = 0u; TileX < MipWidthInTiles; ++TileX)
 			{
 				const uint32 vAddress = FMath::MortonCode2(TileX) | (FMath::MortonCode2(TileY) << 1);
-				const FVTRequestPageResult RequestResult = InVirtualTexture->RequestPageData(ProducerHandle, LayerMask, MaxLevel, vAddress, EVTRequestPagePriority::High);
+				const FVTRequestPageResult RequestResult = InVirtualTexture->RequestPageData(ProducerHandle, LayerMask, MipLevel, vAddress, EVTRequestPagePriority::High);
 				// High priority request should always generate data
 				if (ensure(VTRequestPageStatus_HasData(RequestResult.Status)))
 				{
@@ -1962,7 +1973,7 @@ void FVirtualTexture2DResource::InitializeEditorResources(IVirtualTexture* InVir
 			IVirtualTextureFinalizer* Finalizer = InVirtualTexture->ProducePageData(RHICommandList,
 				GMaxRHIFeatureLevel,
 				EVTProducePageFlags::SkipPageBorders, // don't want to produce page borders, since we're laying out tiles in a regular texture
-				ProducerHandle, LayerMask, MaxLevel, vAddress,
+				ProducerHandle, LayerMask, MipLevel, vAddress,
 				Page.Handle,
 				&TargetLayer);
 			if (Finalizer)
@@ -2041,12 +2052,12 @@ void FVirtualTexture2DResource::ReleaseAllocatedVT()
 
 uint32 FVirtualTexture2DResource::GetSizeX() const
 {
-	return VTData->Width >> FirstMipToUse;
+	return FMath::Max(VTData->Width >> FirstMipToUse, 1u);
 }
 
 uint32 FVirtualTexture2DResource::GetSizeY() const
 {
-	return VTData->Height >> FirstMipToUse;
+	return FMath::Max(VTData->Height >> FirstMipToUse, 1u);
 }
 
 EPixelFormat FVirtualTexture2DResource::GetFormat(uint32 LayerIndex) const
@@ -2061,12 +2072,12 @@ FIntPoint FVirtualTexture2DResource::GetSizeInBlocks() const
 
 uint32 FVirtualTexture2DResource::GetNumTilesX() const
 {
-	return VTData->GetWidthInTiles() >> FirstMipToUse;
+	return FMath::Max(VTData->GetWidthInTiles() >> FirstMipToUse, 1u);
 }
 
 uint32 FVirtualTexture2DResource::GetNumTilesY() const
 {
-	return VTData->GetHeightInTiles() >> FirstMipToUse;
+	return FMath::Max(VTData->GetHeightInTiles() >> FirstMipToUse, 1u);
 }
 
 uint32 FVirtualTexture2DResource::GetBorderSize() const
@@ -2099,19 +2110,6 @@ uint32 FVirtualTexture2DResource::GetAllocatedvAddress() const
 	return ~0;
 }
 
-void FVirtualTexture2DResource::GetPackedPageTableUniform(FUintVector4* Uniform) const
-{
-	check(AllocatedVT);
-	VTGetPackedPageTableUniform(Uniform, AllocatedVT);
-	
-}
-
-void FVirtualTexture2DResource::GetPackedPhysicalTextureUniform(FUintVector4* Uniform, uint32 LayerIndex) const
-{
-	check(AllocatedVT);
-	VTGetPackedUniform(Uniform, AllocatedVT, LayerIndex);
-}
-
 FIntPoint FVirtualTexture2DResource::GetPhysicalTextureSize(uint32 LayerIndex) const
 {
 	if (AllocatedVT)
diff --git a/Engine/Source/Runtime/Engine/Private/TextureDerivedData.cpp b/Engine/Source/Runtime/Engine/Private/TextureDerivedData.cpp
--- a/Engine/Source/Runtime/Engine/Private/TextureDerivedData.cpp
+++ b/Engine/Source/Runtime/Engine/Private/TextureDerivedData.cpp
@@ -73,7 +73,7 @@ static TAutoConsoleVariable<int32> CVarVTTileBorderSize(
 
 // This GUID is mixed into DDC version for virtual textures only, this allows updating DDC version for VT without invalidating DDC for all textures
 // This is useful during development, but once large numbers of VT are present in shipped content, it will have the same problem as TEXTURE_DERIVEDDATA_VER
-#define TEXTURE_VT_DERIVEDDATA_VER	TEXT("186764627EE1418BBAD6ADBA2C2CAE63")
+#define TEXTURE_VT_DERIVEDDATA_VER	TEXT("91E2F0C570CD44AAB33CFA9D000BEFAC")
 
 #if ENABLE_COOK_STATS
 namespace TextureCookStats
diff --git a/Engine/Source/Runtime/Engine/Private/VT/RuntimeVirtualTexture.cpp b/Engine/Source/Runtime/Engine/Private/VT/RuntimeVirtualTexture.cpp
--- a/Engine/Source/Runtime/Engine/Private/VT/RuntimeVirtualTexture.cpp
+++ b/Engine/Source/Runtime/Engine/Private/VT/RuntimeVirtualTexture.cpp
@@ -81,10 +81,12 @@ void URuntimeVirtualTexture::GetProducerDescription(FVTProducerDescription& OutD
 	OutDesc.Dimensions = 2;
 	OutDesc.TileSize = GetTileSize();
 	OutDesc.TileBorderSize = GetTileBorderSize();
-	OutDesc.WidthInTiles = GetWidth() / GetTileSize();
-	OutDesc.HeightInTiles = GetHeight() / GetTileSize();
-	OutDesc.MaxLevel = FMath::Max(FMath::CeilLogTwo(FMath::Max(OutDesc.WidthInTiles, OutDesc.HeightInTiles)) - RemoveLowMips, 1u);
+	OutDesc.BlockWidthInTiles = GetWidth() / GetTileSize();
+	OutDesc.BlockHeightInTiles = GetHeight() / GetTileSize();
+	OutDesc.MaxLevel = FMath::Max(FMath::CeilLogTwo(FMath::Max(OutDesc.BlockWidthInTiles, OutDesc.BlockHeightInTiles)) - RemoveLowMips, 1u);
 	OutDesc.DepthInTiles = 1;
+	OutDesc.WidthInBlocks = 1;
+	OutDesc.HeightInBlocks = 1;
 
 	switch (MaterialType)
 	{
diff --git a/Engine/Source/Runtime/Engine/Private/VT/UploadingVirtualTexture.cpp b/Engine/Source/Runtime/Engine/Private/VT/UploadingVirtualTexture.cpp
--- a/Engine/Source/Runtime/Engine/Private/VT/UploadingVirtualTexture.cpp
+++ b/Engine/Source/Runtime/Engine/Private/VT/UploadingVirtualTexture.cpp
@@ -70,6 +70,13 @@ uint32 FUploadingVirtualTexture::GetLocalMipBias(uint8 vLevel, uint32 vAddress)
 	while (vLevel < NumMips)
 	{
 		const uint32 TileIndex = Data->GetTileIndex(vLevel, vAddress);
+		if (TileIndex == ~0u)
+		{
+			// vAddress is out-of-bounds for the given producer, this is a 
+			NumNonResidentLevels += (NumMips - vLevel);
+			break;
+		}
+
 		const int32 ChunkIndex = Data->GetChunkIndex(TileIndex);
 		if (ChunkIndex >= 0)
 		{
diff --git a/Engine/Source/Runtime/Engine/Private/VT/VirtualTextureBuiltData.h b/Engine/Source/Runtime/Engine/Private/VT/VirtualTextureBuiltData.h
--- a/Engine/Source/Runtime/Engine/Private/VT/VirtualTextureBuiltData.h
+++ b/Engine/Source/Runtime/Engine/Private/VT/VirtualTextureBuiltData.h
@@ -135,7 +135,11 @@ struct FVirtualTextureBuiltData
 	{
 		check(vLevel < NumMips);
 		const uint32 TileIndex = TileIndexPerMip[vLevel] + vAddress * NumLayers;
-		check(TileIndex < TileIndexPerMip[vLevel + 1]);
+		if (TileIndex >= TileIndexPerMip[vLevel + 1])
+		{
+			// vAddress is out of bounds for this texture/mip level
+			return ~0u;
+		}
 		return TileIndex;
 	}
 
diff --git a/Engine/Source/Runtime/Engine/Private/VT/VirtualTextureDataBuilder.cpp b/Engine/Source/Runtime/Engine/Private/VT/VirtualTextureDataBuilder.cpp
--- a/Engine/Source/Runtime/Engine/Private/VT/VirtualTextureDataBuilder.cpp
+++ b/Engine/Source/Runtime/Engine/Private/VT/VirtualTextureDataBuilder.cpp
@@ -333,7 +333,9 @@ void FVirtualTextureDataBuilder::Build(const FTextureSourceData& InSourceData, c
 	const uint32 SizeInTiles = FMath::DivideAndRoundUp<uint32>(Size, TileSize);
 	const uint32 BlockSize = FMath::Max(BlockSizeX, BlockSizeY);
 	const uint32 BlockSizeInTiles = FMath::DivideAndRoundUp<uint32>(BlockSize, TileSize);
-	OutData.NumMips = FMath::CeilLogTwo(SizeInTiles) + 1;
+
+	// Mip down to 1x1 pixels, but don't create more mips than can fit in max page table texture
+	OutData.NumMips = FMath::Min<uint32>(FMath::CeilLogTwo(Size) + 1, VIRTUALTEXTURE_LOG2_MAX_PAGETABLE_SIZE);
 
 	BuildSourcePixels(InSourceData, InCompositeSourceData);
 
@@ -884,10 +886,17 @@ void FVirtualTextureDataBuilder::BuildSourcePixels(const FTextureSourceData& Sou
 			BlockData.SizeY = CompressedMips[0].SizeY;
 
 			const uint32 BlockSize = FMath::Max(BlockData.SizeX, BlockData.SizeY);
-			const uint32 BlockSizeInTiles = FMath::DivideAndRoundUp<uint32>(BlockSize, TileSize);
-			const uint32 MaxMipInBlock = FMath::CeilLogTwo(BlockSizeInTiles);
-
-			BlockData.NumMips = FMath::Min<int32>(CompressedMips.Num(), MaxMipInBlock + 1);
+			if (NumBlocks == 1u)
+			{
+				const uint32 MaxMipInBlock = FMath::CeilLogTwo(BlockSize);
+				BlockData.NumMips = FMath::Min<int32>(CompressedMips.Num(), MaxMipInBlock + 1);
+			}
+			else
+			{
+				const uint32 BlockSizeInTiles = FMath::DivideAndRoundUp<uint32>(BlockSize, TileSize);
+				const uint32 MaxMipInBlock = FMath::CeilLogTwo(BlockSizeInTiles);
+				BlockData.NumMips = FMath::Min<int32>(CompressedMips.Num(), MaxMipInBlock + 1);
+			}
 			BlockData.MipsPerLayer[LayerIndex].Reserve(BlockData.NumMips);
 			for (int32 MipIndex = 0; MipIndex < BlockData.NumMips; ++MipIndex)
 			{
@@ -920,7 +929,7 @@ void FVirtualTextureDataBuilder::BuildSourcePixels(const FTextureSourceData& Sou
 		const uint32 MipInputSizeX = FMath::RoundUpToPowerOfTwo(SizeInBlocksX * MipWidthInBlock);
 		const uint32 MipInputSizeY = FMath::RoundUpToPowerOfTwo(SizeInBlocksY * MipHeightInBlock);
 		const uint32 MipInputSize = FMath::Max(MipInputSizeX, MipInputSizeY);
-		const uint32 MipInputSizeInTiles = FMath::DivideAndRoundUp<uint32>(MipInputSize, TileSize);
+		//const uint32 MipInputSizeInTiles = FMath::DivideAndRoundUp<uint32>(MipInputSize, TileSize);
 
 		FTextureSourceBlockData& SourceMiptailBlock = SourceBlocks.AddDefaulted_GetRef();
 		SourceMiptailBlock.BlockX = 0;
@@ -929,7 +938,7 @@ void FVirtualTextureDataBuilder::BuildSourcePixels(const FTextureSourceData& Sou
 		SourceMiptailBlock.SizeInBlocksY = SizeInBlocksY;
 		SourceMiptailBlock.SizeX = FMath::Max(MipInputSizeX >> 1, 1u);
 		SourceMiptailBlock.SizeY = FMath::Max(MipInputSizeY >> 1, 1u);
-		SourceMiptailBlock.NumMips = FMath::CeilLogTwo(MipInputSizeInTiles); // Don't add 1, since 'MipInputSizeInTiles' is one mip larger
+		SourceMiptailBlock.NumMips = OutData.NumMips - MaxMipInBlock - 1;//   FMath::CeilLogTwo(MipInputSize); // Don't add 1, since 'MipInputSize' is one mip larger
 		SourceMiptailBlock.NumSlices = 1; // TODO?
 		SourceMiptailBlock.MipBias = MaxMipInBlock + 1;
 		SourceMiptailBlock.MipsPerLayer.AddDefaulted(NumLayers);
diff --git a/Engine/Source/Runtime/Engine/Private/VT/VirtualTextureUtility.cpp b/Engine/Source/Runtime/Engine/Private/VT/VirtualTextureUtility.cpp
deleted file mode 100644
diff --git a/Engine/Source/Runtime/Engine/Private/VT/VirtualTextureUtility.h b/Engine/Source/Runtime/Engine/Private/VT/VirtualTextureUtility.h
deleted file mode 100644
diff --git a/Engine/Source/Runtime/Engine/Public/TextureResource.h b/Engine/Source/Runtime/Engine/Public/TextureResource.h
--- a/Engine/Source/Runtime/Engine/Public/TextureResource.h
+++ b/Engine/Source/Runtime/Engine/Public/TextureResource.h
@@ -296,8 +296,6 @@ public:
 	ENGINE_API uint32 GetBorderSize() const;
 	uint32 GetAllocatedvAddress() const;
 
-	ENGINE_API void GetPackedPageTableUniform(FUintVector4* Uniform) const;
-	ENGINE_API void GetPackedPhysicalTextureUniform(FUintVector4* Uniform, uint32 LayerIndex) const;
 	ENGINE_API FIntPoint GetPhysicalTextureSize(uint32 LayerIndex) const;
 
 private:
diff --git a/Engine/Source/Runtime/RenderCore/Public/VirtualTexturing.h b/Engine/Source/Runtime/RenderCore/Public/VirtualTexturing.h
--- a/Engine/Source/Runtime/RenderCore/Public/VirtualTexturing.h
+++ b/Engine/Source/Runtime/RenderCore/Public/VirtualTexturing.h
@@ -32,8 +32,8 @@ inline bool operator!=(const FVirtualTextureProducerHandle& Lhs, const FVirtualT
 /** Maximum number of layers that can be allocated in a single VT page table */
 #define VIRTUALTEXTURE_SPACE_MAXLAYERS 8
 
-/** Maximum dimension of VT page table texture (Duplicated in PageTableUpdate.usf, keep in sync) */
-#define VIRTUALTEXTURE_LOG2_MAX_PAGETABLE_SIZE 11u
+/** Maximum dimension of VT page table texture */
+#define VIRTUALTEXTURE_LOG2_MAX_PAGETABLE_SIZE 12u
 #define VIRTUALTEXTURE_MAX_PAGETABLE_SIZE (1u << VIRTUALTEXTURE_LOG2_MAX_PAGETABLE_SIZE)
 
 /**
@@ -98,9 +98,18 @@ struct FVTProducerDescription
 	bool bCreateRenderTarget = false;
 	uint32 TileSize = 0u;
 	uint32 TileBorderSize = 0u;
-	uint32 WidthInTiles = 0u;
-	uint32 HeightInTiles = 0u;
+
+	/**
+	 * Producers are made up of a number of block, each block has uniform size, and blocks are arranged in a larger grid
+	 * 'Normal' VTs will typically be a single block, for UDIM textures, blocks will map to individual UDIM texture sheets
+	 * When multiple producers are allocated together, they will be aligned such that blocks of each layer overlay on top of each other
+	 * Number of blocks for each layer may be different in this case, this is handled by wrapping blocks for layers with fewer blocks
+	 */
+	uint32 BlockWidthInTiles = 0u;
+	uint32 BlockHeightInTiles = 0u;
 	uint32 DepthInTiles = 0u;
+	uint16 WidthInBlocks = 1u;
+	uint16 HeightInBlocks = 1u;
 	uint8 Dimensions = 0u;
 	uint8 NumLayers = 0u;
 	uint8 MaxLevel = 0u;
@@ -255,13 +264,17 @@ public:
 	inline IAllocatedVirtualTexture(const FAllocatedVTDescription& InDesc,
 		uint32 InSpaceID,
 		EVTPageTableFormat InPageTableFormat,
-		uint32 InWidthInTiles,
-		uint32 InHeightInTiles,
+		uint32 InBlockWidthInTiles,
+		uint32 InBlockHeightInTiles,
+		uint32 InWidthInBlocks,
+		uint32 InHeightInBlocks,
 		uint32 InDepthInTiles)
 		: Description(InDesc)
 		, SpaceID(InSpaceID)
-		, WidthInTiles(InWidthInTiles)
-		, HeightInTiles(InHeightInTiles)
+		, BlockWidthInTiles(InBlockWidthInTiles)
+		, BlockHeightInTiles(InBlockHeightInTiles)
+		, WidthInBlocks(InWidthInBlocks)
+		, HeightInBlocks(InHeightInBlocks)
 		, DepthInTiles(InDepthInTiles)
 		, PageTableFormat(InPageTableFormat)
 		, MaxLevel(0u)
@@ -272,6 +285,13 @@ public:
 	virtual FRHITexture* GetPhysicalTexture(uint32 InLayerIndex) const = 0;
 	virtual FRHIShaderResourceView* GetPhysicalTextureView(uint32 InLayerIndex, bool bSRGB) const = 0;
 	virtual uint32 GetPhysicalTextureSize(uint32 InLayerIndex) const = 0;
+
+	/** Writes 2x FUintVector4 */
+	virtual void GetPackedPageTableUniform(FUintVector4* OutUniform, bool bApplyBlockScale) const = 0;
+
+	/** Writes 1x FUintVector4 */
+	virtual void GetPackedUniform(FUintVector4* OutUniform, uint32 LayerIndex) const = 0;
+
 	virtual void DumpToConsole(bool bVerbose) const {}
 
 	inline const FAllocatedVTDescription& GetDescription() const { return Description; }
@@ -283,11 +303,13 @@ public:
 	inline uint32 GetPhysicalTileSize() const { return Description.TileSize + Description.TileBorderSize * 2u; }
 	inline uint32 GetNumLayers() const { return Description.NumLayers; }
 	inline uint8 GetDimensions() const { return Description.Dimensions; }
-	inline uint32 GetWidthInTiles() const { return WidthInTiles; }
-	inline uint32 GetHeightInTiles() const { return HeightInTiles; }
+	inline uint32 GetWidthInBlocks() const { return WidthInBlocks; }
+	inline uint32 GetHeightInBlocks() const { return HeightInBlocks; }
+	inline uint32 GetWidthInTiles() const { return BlockWidthInTiles * WidthInBlocks; }
+	inline uint32 GetHeightInTiles() const { return BlockHeightInTiles * HeightInBlocks; }
 	inline uint32 GetDepthInTiles() const { return DepthInTiles; }
-	inline uint32 GetWidthInPixels() const { return WidthInTiles * Description.TileSize; }
-	inline uint32 GetHeightInPixels() const { return HeightInTiles * Description.TileSize; }
+	inline uint32 GetWidthInPixels() const { return GetWidthInTiles() * Description.TileSize; }
+	inline uint32 GetHeightInPixels() const { return GetHeightInTiles() * Description.TileSize; }
 	inline uint32 GetDepthInPixels() const { return DepthInTiles * Description.TileSize; }
 	inline uint32 GetNumPageTableTextures() const { return (Description.NumLayers + LayersPerPageTableTexture - 1u) / LayersPerPageTableTexture; }
 	inline uint32 GetSpaceID() const { return SpaceID; }
@@ -302,8 +324,10 @@ protected:
 
 	FAllocatedVTDescription Description;
 	uint32 SpaceID;
-	uint32 WidthInTiles;
-	uint32 HeightInTiles;
+	uint32 BlockWidthInTiles;
+	uint32 BlockHeightInTiles;
+	uint32 WidthInBlocks;
+	uint32 HeightInBlocks;
 	uint32 DepthInTiles;
 	EVTPageTableFormat PageTableFormat;
 
diff --git a/Engine/Source/Runtime/Renderer/Private/VT/AllocatedVirtualTexture.cpp b/Engine/Source/Runtime/Renderer/Private/VT/AllocatedVirtualTexture.cpp
--- a/Engine/Source/Runtime/Renderer/Private/VT/AllocatedVirtualTexture.cpp
+++ b/Engine/Source/Runtime/Renderer/Private/VT/AllocatedVirtualTexture.cpp
@@ -5,14 +5,17 @@
 #include "VirtualTextureSpace.h"
 #include "VirtualTexturePhysicalSpace.h"
 
-FAllocatedVirtualTexture::FAllocatedVirtualTexture(uint32 InFrame,
+FAllocatedVirtualTexture::FAllocatedVirtualTexture(FVirtualTextureSystem* InSystem,
+	uint32 InFrame,
 	const FAllocatedVTDescription& InDesc,
 	FVirtualTextureSpace* InSpace,
 	FVirtualTextureProducer* const* InProducers,
-	uint32 InWidthInTiles,
-	uint32 InHeightInTiles,
+	uint32 InBlockWidthInTiles,
+	uint32 InBlockHeightInTiles,
+	uint32 InWidthInBlocks,
+	uint32 InHeightInBlocks,
 	uint32 InDepthInTiles)
-	: IAllocatedVirtualTexture(InDesc, InSpace->GetID(), InSpace->GetDescription().Format, InWidthInTiles, InHeightInTiles, InDepthInTiles)
+	: IAllocatedVirtualTexture(InDesc, InSpace->GetID(), InSpace->GetDescription().Format, InBlockWidthInTiles, InBlockHeightInTiles, InWidthInBlocks, InHeightInBlocks, InDepthInTiles)
 	, Space(InSpace)
 	, RefCount(1)
 	, FrameAllocated(InFrame)
@@ -37,11 +40,40 @@ FAllocatedVirtualTexture::FAllocatedVirtualTexture(uint32 InFrame,
 		}
 	}
 
-	// must have at least 1 valid layer/producer
+	// Must have at least 1 valid layer/producer
 	check(NumUniqueProducers > 0u);
 
-	// make sure max level calculation is valid (don't end up with more mips than supported by size)
-	check(MaxLevel <= FMath::CeilLogTwo(FMath::Max(InWidthInTiles, InHeightInTiles)));
+	// Max level of overall allocated VT is limited by size in tiles
+	// With multiple layers of different sizes, some layers may have mips smaller than a single tile
+	MaxLevel = FMath::Min(MaxLevel, FMath::CeilLogTwo(FMath::Max(GetWidthInTiles(), GetHeightInTiles())));
+
+	// Lock lowest resolution mip from each producer
+	// Depending on the block dimensions of the producers that make up this allocated VT, different allocated VTs may need to lock different low resolution mips from the same producer
+	// In the common case where block dimensions match, same mip will be locked by all allocated VTs that make use of the same producer
+	for (uint32 LayerIndex = 0u; LayerIndex < Description.NumLayers; ++LayerIndex)
+	{
+		FVirtualTextureProducer* Producer = InProducers[LayerIndex];
+		if (Producer && Producer->GetDescription().bPersistentHighestMip)
+		{
+			const uint32 ProducerIndex = UniqueProducerIndexForLayer[LayerIndex];
+			const uint32 MipBias = UniqueProducerMipBias[ProducerIndex];
+			check(MipBias <= MaxLevel);
+			const uint32 Local_vLevel = MaxLevel - MipBias;
+			check(Local_vLevel <= Producer->GetMaxLevel());
+
+			const uint32 RootWidthInTiles = FMath::Max(Producer->GetWidthInTiles() >> Local_vLevel, 1u);
+			const uint32 RootHeightInTiles = FMath::Max(Producer->GetHeightInTiles() >> Local_vLevel, 1u);
+			for (uint32 TileY = 0u; TileY < RootHeightInTiles; ++TileY)
+			{
+				for (uint32 TileX = 0u; TileX < RootWidthInTiles; ++TileX)
+				{
+					const uint32 Local_vAddress = FMath::MortonCode2(TileX) | (FMath::MortonCode2(TileY) << 1);
+					const FVirtualTextureLocalTile TileToUnlock(InDesc.ProducerHandle[LayerIndex], Local_vAddress, Local_vLevel);
+					InSystem->LockTile(TileToUnlock);
+				}
+			}
+		}
+	}
 
 	VirtualAddress = Space->AllocateVirtualTexture(this);
 }
@@ -65,8 +97,32 @@ void FAllocatedVirtualTexture::Release(FVirtualTextureSystem* System)
 	check(IsInRenderingThread());
 	check(RefCount.GetValue() == 0);
 
+	// Unlock any locked tiles
 	for (uint32 LayerIndex = 0u; LayerIndex < Description.NumLayers; ++LayerIndex)
 	{
+		const FVirtualTextureProducerHandle ProducerHandle = GetProducerHandle(LayerIndex);
+		FVirtualTextureProducer* Producer = System->FindProducer(ProducerHandle);
+		if (Producer && Producer->GetDescription().bPersistentHighestMip)
+		{
+			const uint32 ProducerIndex = UniqueProducerIndexForLayer[LayerIndex];
+			const uint32 MipBias = UniqueProducerMipBias[ProducerIndex];
+			check(MipBias <= MaxLevel);
+			const uint32 Local_vLevel = MaxLevel - MipBias;
+			check(Local_vLevel <= Producer->GetMaxLevel());
+
+			const uint32 RootWidthInTiles = FMath::Max(Producer->GetWidthInTiles() >> Local_vLevel, 1u);
+			const uint32 RootHeightInTiles = FMath::Max(Producer->GetHeightInTiles() >> Local_vLevel, 1u);
+			for (uint32 TileY = 0u; TileY < RootHeightInTiles; ++TileY)
+			{
+				for (uint32 TileX = 0u; TileX < RootWidthInTiles; ++TileX)
+				{
+					const uint32 Local_vAddress = FMath::MortonCode2(TileX) | (FMath::MortonCode2(TileY) << 1);
+					const FVirtualTextureLocalTile TileToUnlock(ProducerHandle, Local_vAddress, Local_vLevel);
+					System->UnlockTile(TileToUnlock, Producer);
+				}
+			}
+		}
+
 		// Physical pool needs to evict all pages that belong to this VT's space
 		// TODO - could improve this to only evict pages belonging to this VT
 		if (PhysicalSpace[LayerIndex])
@@ -103,13 +159,13 @@ uint32 FAllocatedVirtualTexture::AddUniqueProducer(const FVirtualTextureProducer
 	check(ProducerDesc.TileSize == Description.TileSize);
 	check(ProducerDesc.TileBorderSize == Description.TileBorderSize);
 
-	const uint32 SizeInTiles = FMath::Max(WidthInTiles, HeightInTiles);
-	const uint32 ProducerSizeInTiles = FMath::Max(ProducerDesc.WidthInTiles, ProducerDesc.HeightInTiles);
-	const uint32 MipBias = FMath::CeilLogTwo(SizeInTiles / ProducerSizeInTiles);
+	const uint32 BlockSizeInTiles = FMath::Max(BlockWidthInTiles, BlockHeightInTiles);
+	const uint32 ProducerBlockSizeInTiles = FMath::Max(ProducerDesc.BlockWidthInTiles, ProducerDesc.BlockHeightInTiles);
+	const uint32 MipBias = FMath::CeilLogTwo(BlockSizeInTiles / ProducerBlockSizeInTiles);
 
-	check((SizeInTiles / ProducerSizeInTiles) * ProducerSizeInTiles == SizeInTiles);
-	check(ProducerDesc.WidthInTiles << MipBias == WidthInTiles);
-	check(ProducerDesc.HeightInTiles << MipBias == HeightInTiles);
+	check((BlockSizeInTiles / ProducerBlockSizeInTiles) * ProducerBlockSizeInTiles == BlockSizeInTiles);
+	check(ProducerDesc.BlockWidthInTiles << MipBias == BlockWidthInTiles);
+	check(ProducerDesc.BlockHeightInTiles << MipBias == BlockHeightInTiles);
 
 	MaxLevel = FMath::Max<uint32>(MaxLevel, ProducerDesc.MaxLevel + MipBias);
 
@@ -153,3 +209,79 @@ uint32 FAllocatedVirtualTexture::GetPhysicalTextureSize(uint32 InLayerIndex) con
 	}
 	return 0u;
 }
+
+static inline uint32 BitcastFloatToUInt32(float v)
+{
+	const union
+	{
+		float FloatValue;
+		uint32 UIntValue;
+	} u = { v };
+	return u.UIntValue;
+}
+
+void FAllocatedVirtualTexture::GetPackedPageTableUniform(FUintVector4* OutUniform, bool bApplyBlockScale) const
+{
+	static const auto MaxAnisotropyCVar = IConsoleManager::Get().FindTConsoleVariableDataInt(TEXT("r.MaxAnisotropy"));
+
+	const uint32 vPageX = FMath::ReverseMortonCode2(VirtualAddress);
+	const uint32 vPageY = FMath::ReverseMortonCode2(VirtualAddress >> 1);
+	const uint32 vPageSize = GetVirtualTileSize();
+	const uint32 PageBorderSize = GetTileBorderSize();
+	const uint32 WidthInPages = GetWidthInTiles();
+	const uint32 HeightInPages = GetHeightInTiles();
+	const uint32 vPageTableMipBias = FMath::FloorLog2(vPageSize);
+
+	const uint32 MaxAnisotropy = FMath::Clamp<int32>(MaxAnisotropyCVar->GetValueOnRenderThread(), 1, PageBorderSize);
+	const uint32 MaxAnisotropyLog2 = FMath::FloorLog2(MaxAnisotropy);
+
+	// make sure everything fits in the allocated number of bits
+	checkSlow(vPageX < 4096u);
+	checkSlow(vPageY < 4096u);
+	checkSlow(vPageTableMipBias < 16u);
+	checkSlow(MaxLevel < 16u);
+	checkSlow(SpaceID < 16u);
+
+	if (bApplyBlockScale)
+	{
+		OutUniform[0].X = BitcastFloatToUInt32(1.0f / (float)WidthInBlocks);
+		OutUniform[0].Y = BitcastFloatToUInt32(1.0f / (float)HeightInBlocks);
+	}
+	else
+	{
+		OutUniform[0].X = BitcastFloatToUInt32(1.0f);
+		OutUniform[0].Y = BitcastFloatToUInt32(1.0f);
+	}
+
+	OutUniform[0].Z = BitcastFloatToUInt32((float)WidthInPages);
+	OutUniform[0].W = BitcastFloatToUInt32((float)HeightInPages);
+
+	OutUniform[1].X = BitcastFloatToUInt32((float)MaxAnisotropyLog2);
+	OutUniform[1].Y = vPageX | (vPageY << 12) | (vPageTableMipBias << 24);
+	OutUniform[1].Z = MaxLevel;
+	OutUniform[1].W = (SpaceID << 28);
+}
+
+void FAllocatedVirtualTexture::GetPackedUniform(FUintVector4* OutUniform, uint32 LayerIndex) const
+{
+	const uint32 PhysicalTextureSize = GetPhysicalTextureSize(LayerIndex);
+	if (PhysicalTextureSize > 0u)
+	{
+		const uint32 vPageSize = GetVirtualTileSize();
+		const uint32 PageBorderSize = GetTileBorderSize();
+
+		const float RcpPhysicalTextureSize = 1.0f / float(PhysicalTextureSize);
+		const uint32 pPageSize = vPageSize + PageBorderSize * 2u;
+		OutUniform->X = 0u;
+		OutUniform->Y = BitcastFloatToUInt32((float)vPageSize * RcpPhysicalTextureSize);
+		OutUniform->Z = BitcastFloatToUInt32((float)PageBorderSize * RcpPhysicalTextureSize);
+		OutUniform->W = BitcastFloatToUInt32((float)pPageSize * RcpPhysicalTextureSize);
+	}
+	else
+	{
+		OutUniform->X = 0u;
+		OutUniform->Y = 0u;
+		OutUniform->Z = 0u;
+		OutUniform->W = 0u;
+	}
+}
diff --git a/Engine/Source/Runtime/Renderer/Private/VT/AllocatedVirtualTexture.h b/Engine/Source/Runtime/Renderer/Private/VT/AllocatedVirtualTexture.h
--- a/Engine/Source/Runtime/Renderer/Private/VT/AllocatedVirtualTexture.h
+++ b/Engine/Source/Runtime/Renderer/Private/VT/AllocatedVirtualTexture.h
@@ -17,12 +17,15 @@ class FTexturePageLocks;
 class FAllocatedVirtualTexture final : public IAllocatedVirtualTexture
 {
 public:
-	FAllocatedVirtualTexture(uint32 InFrame,
+	FAllocatedVirtualTexture(FVirtualTextureSystem* InSystem,
+		uint32 InFrame,
 		const FAllocatedVTDescription& InDesc,
 		FVirtualTextureSpace* InSpace,
 		FVirtualTextureProducer* const* InProducers,
-		uint32 InWidthInTiles,
-		uint32 InHeightInTiles,
+		uint32 InBlockWidthInTiles,
+		uint32 InBlockHeightInTiles,
+		uint32 InWidthInBlocks,
+		uint32 InHeightInBlocks,
 		uint32 InDepthInTiles);
 
 	virtual ~FAllocatedVirtualTexture();
@@ -34,6 +37,8 @@ public:
 	virtual FRHITexture* GetPhysicalTexture(uint32 InLayerIndex) const override;
 	virtual FRHIShaderResourceView* GetPhysicalTextureView(uint32 InLayerIndex, bool bSRGB) const override;
 	virtual uint32 GetPhysicalTextureSize(uint32 InLayerIndex) const override;
+	virtual void GetPackedPageTableUniform(FUintVector4* OutUniform, bool bApplyBlockScale) const override;
+	virtual void GetPackedUniform(FUintVector4* OutUniform, uint32 LayerIndex) const override;
 	virtual void Destroy(FVirtualTextureSystem* System) override;
 	// end IAllocatedVirtualTexture
 
diff --git a/Engine/Source/Runtime/Renderer/Private/VT/RuntimeVirtualTextureProducer.cpp b/Engine/Source/Runtime/Renderer/Private/VT/RuntimeVirtualTextureProducer.cpp
--- a/Engine/Source/Runtime/Renderer/Private/VT/RuntimeVirtualTextureProducer.cpp
+++ b/Engine/Source/Runtime/Renderer/Private/VT/RuntimeVirtualTextureProducer.cpp
@@ -41,8 +41,8 @@ void FRuntimeVirtualTextureFinalizer::Finalize(FRHICommandListImmediate& RHICmdL
 
 		const uint32 X = FMath::ReverseMortonCode2(Entry.vAddress);
 		const uint32 Y = FMath::ReverseMortonCode2(Entry.vAddress >> 1);
-		const uint32 DivisorX = Desc.WidthInTiles >> Entry.vLevel;
-		const uint32 DivisorY = Desc.HeightInTiles >> Entry.vLevel;
+		const uint32 DivisorX = Desc.BlockWidthInTiles >> Entry.vLevel;
+		const uint32 DivisorY = Desc.BlockHeightInTiles >> Entry.vLevel;
 
 		const FVector2D UV((float)X / (float)DivisorX, (float)Y / (float)DivisorY);
 		const FVector2D UVSize(1.f / (float)DivisorX, 1.f / (float)DivisorY);
diff --git a/Engine/Source/Runtime/Renderer/Private/VT/TexturePageLocks.cpp b/Engine/Source/Runtime/Renderer/Private/VT/TexturePageLocks.cpp
new file mode 100644
--- /dev/null
+++ b/Engine/Source/Runtime/Renderer/Private/VT/TexturePageLocks.cpp
@@ -0,0 +1,105 @@
+// Copyright 1998-2019 Epic Games, Inc. All Rights Reserved.
+
+#include "TexturePageLocks.h"
+
+FTexturePageLocks::FTexturePageLocks()
+{
+}
+
+uint32 FTexturePageLocks::AcquireIndex()
+{
+	if (FreeIndices.Num() > 0)
+	{
+		return FreeIndices.Pop();
+	}
+	LockedTiles.AddDefaulted();
+	const uint32 Index = LockCounts.AddZeroed();
+	check(LockedTiles.Num() == LockCounts.Num());
+	return Index;
+}
+
+bool FTexturePageLocks::Lock(const FVirtualTextureLocalTile& Tile)
+{
+	const uint16 Hash = (uint16)MurmurFinalize64(Tile.PackedValue);
+	for (uint32 Index = TileHash.First(Hash); TileHash.IsValid(Index); Index = TileHash.Next(Index))
+	{
+		if (LockedTiles[Index] == Tile)
+		{
+			const uint32 PrevLockCount = LockCounts[Index];
+			check(PrevLockCount > 0u);
+			check(PrevLockCount < 0xffff);
+			LockCounts[Index] = PrevLockCount + 1u;
+			return false;
+		}
+	}
+
+	const uint32 Index = AcquireIndex();
+	LockedTiles[Index] = Tile;
+	LockCounts[Index] = 1u;
+	TileHash.Add(Hash, Index);
+	return true;
+}
+
+bool FTexturePageLocks::Unlock(const FVirtualTextureLocalTile& Tile)
+{
+	const uint16 Hash = (uint16)MurmurFinalize64(Tile.PackedValue);
+	for (uint32 Index = TileHash.First(Hash); TileHash.IsValid(Index); Index = TileHash.Next(Index))
+	{
+		if (LockedTiles[Index] == Tile)
+		{
+			const uint32 PrevLockCount = LockCounts[Index];
+			check(PrevLockCount > 0u);
+			if (PrevLockCount == 1u)
+			{
+				// no longer locked
+				FreeIndices.Add(Index);
+				TileHash.Remove(Hash, Index);
+				LockCounts[Index] = 0u;
+				LockedTiles[Index].PackedValue = 0u;
+				return true;
+			}
+
+			// still locked
+			LockCounts[Index] = PrevLockCount - 1u;
+			return false;
+		}
+	}
+
+	// Possible that we're trying to unlock a tile that's already been force-unlocked because the producer was destroyed
+	return false;
+}
+
+void FTexturePageLocks::ForceUnlockAll(const FVirtualTextureProducerHandle& ProducerHandle, TArray<FVirtualTextureLocalTile>& OutUnlockedTiles)
+{
+	for (int32 Index = 0u; Index < LockedTiles.Num(); ++Index)
+	{
+		if (LockCounts[Index] > 0u)
+		{
+			const FVirtualTextureLocalTile& Tile = LockedTiles[Index];
+			if (Tile.GetProducerHandle() == ProducerHandle)
+			{
+				OutUnlockedTiles.Add(Tile);
+				FreeIndices.Add(Index);
+				TileHash.Remove(MurmurFinalize64(Tile.PackedValue), Index);
+				LockCounts[Index] = 0u;
+				LockedTiles[Index].PackedValue = 0u;
+			}
+		}
+	}
+}
+
+bool FTexturePageLocks::IsLocked(const FVirtualTextureLocalTile& Tile) const
+{
+	const uint16 Hash = (uint16)MurmurFinalize64(Tile.PackedValue);
+	for (uint32 Index = TileHash.First(Hash); TileHash.IsValid(Index); Index = TileHash.Next(Index))
+	{
+		if (LockedTiles[Index] == Tile)
+		{
+			// lock count needs to be at least 1 if it's in the hash table
+			check(LockCounts[Index] > 0u);
+			return true;
+		}
+	}
+
+	return false;
+}
diff --git a/Engine/Source/Runtime/Renderer/Private/VT/TexturePageLocks.h b/Engine/Source/Runtime/Renderer/Private/VT/TexturePageLocks.h
new file mode 100644
--- /dev/null
+++ b/Engine/Source/Runtime/Renderer/Private/VT/TexturePageLocks.h
@@ -0,0 +1,35 @@
+// Copyright 1998-2019 Epic Games, Inc. All Rights Reserved.
+
+#pragma once
+
+#include "CoreMinimal.h"
+#include "Containers/HashTable.h"
+#include "VirtualTexturing.h"
+
+/**
+ * This tracks lock counts for individual VT tiles
+ * We need to track lock counts for all tiles, since different allocated VTs may try to lock/unlock the same pages,
+ * and we need a way to know when a given tile is no longer locked by anything
+ */
+class FTexturePageLocks
+{
+public:
+	FTexturePageLocks();
+
+	/** Lock/Unlock return true if the lock state was changed */
+	bool Lock(const FVirtualTextureLocalTile& Tile);
+	bool Unlock(const FVirtualTextureLocalTile& Tile);
+
+	/** Force-unlock all tiles using the given producer, returns list of all tiles that were unlocked */
+	void ForceUnlockAll(const FVirtualTextureProducerHandle& ProducerHandle, TArray<FVirtualTextureLocalTile>& OutUnlockedTiles);
+
+	bool IsLocked(const FVirtualTextureLocalTile& Tile) const;
+
+private:
+	uint32 AcquireIndex();
+
+	FHashTable TileHash;
+	TArray<FVirtualTextureLocalTile> LockedTiles;
+	TArray<uint16> LockCounts;
+	TArray<uint32> FreeIndices;
+};
diff --git a/Engine/Source/Runtime/Renderer/Private/VT/VirtualTextureProducer.cpp b/Engine/Source/Runtime/Renderer/Private/VT/VirtualTextureProducer.cpp
--- a/Engine/Source/Runtime/Renderer/Private/VT/VirtualTextureProducer.cpp
+++ b/Engine/Source/Runtime/Renderer/Private/VT/VirtualTextureProducer.cpp
@@ -8,17 +8,7 @@ void FVirtualTextureProducer::Release(FVirtualTextureSystem* System, const FVirt
 {
 	if (Description.bPersistentHighestMip)
 	{
-		const uint32 RootWidthInTiles = FMath::Max(Description.WidthInTiles >> Description.MaxLevel, 1u);
-		const uint32 RootHeightInTiles = FMath::Max(Description.HeightInTiles >> Description.MaxLevel, 1u);
-		for (uint32 TileY = 0u; TileY < RootHeightInTiles; ++TileY)
-		{
-			for (uint32 TileX = 0u; TileX < RootWidthInTiles; ++TileX)
-			{
-				const uint32 Local_vAddress = FMath::MortonCode2(TileX) | (FMath::MortonCode2(TileY) << 1);
-				const FVirtualTextureLocalTile TileToUnlock(HandleToSelf, Local_vAddress, Description.MaxLevel);
-				System->UnlockTile(TileToUnlock);
-			}
-		}
+		System->ForceUnlockAllTiles(HandleToSelf, this);
 	}
 
 	for (uint32 LayerIndex = 0u; LayerIndex < Description.NumLayers; ++LayerIndex)
@@ -42,7 +32,11 @@ FVirtualTextureProducerCollection::FVirtualTextureProducerCollection()
 FVirtualTextureProducerHandle FVirtualTextureProducerCollection::RegisterProducer(FVirtualTextureSystem* System, const FVTProducerDescription& InDesc, IVirtualTexture* InProducer)
 {
 	check(IsInRenderingThread());
-	check(InDesc.MaxLevel <= FMath::CeilLogTwo(FMath::Max(InDesc.WidthInTiles, InDesc.HeightInTiles)));
+	const uint32 ProducerWidth = InDesc.BlockWidthInTiles * InDesc.WidthInBlocks * InDesc.TileSize;
+	const uint32 ProducerHeight = InDesc.BlockHeightInTiles * InDesc.HeightInBlocks * InDesc.TileSize;
+	check(ProducerWidth > 0u);
+	check(ProducerHeight > 0u);
+	check(InDesc.MaxLevel <= FMath::CeilLogTwo(FMath::Max(ProducerWidth, ProducerHeight)));
 	check(InProducer);
 
 	const uint32 Index = AcquireEntry();
@@ -63,23 +57,6 @@ FVirtualTextureProducerHandle FVirtualTextureProducerCollection::RegisterProduce
 	}
 
 	const FVirtualTextureProducerHandle Handle(Index, Entry.Magic);
-
-	if (InDesc.bPersistentHighestMip)
-	{
-		const uint32 MaxLevel = InDesc.MaxLevel;
-		const uint32 RootWidthInTiles = FMath::Max(InDesc.WidthInTiles >> MaxLevel, 1u);
-		const uint32 RootHeightInTiles = FMath::Max(InDesc.HeightInTiles >> MaxLevel, 1u);
-		for (uint32 TileY = 0u; TileY < RootHeightInTiles; ++TileY)
-		{
-			for (uint32 TileX = 0u; TileX < RootWidthInTiles; ++TileX)
-			{
-				const uint32 Local_vAddress = FMath::MortonCode2(TileX) | (FMath::MortonCode2(TileY) << 1);
-				const FVirtualTextureLocalTile TileToLock(Handle, Local_vAddress, MaxLevel);
-				System->LockTile(TileToLock);
-			}
-		}
-	}
-
 	return Handle;
 }
 
diff --git a/Engine/Source/Runtime/Renderer/Private/VT/VirtualTextureProducer.h b/Engine/Source/Runtime/Renderer/Private/VT/VirtualTextureProducer.h
--- a/Engine/Source/Runtime/Renderer/Private/VT/VirtualTextureProducer.h
+++ b/Engine/Source/Runtime/Renderer/Private/VT/VirtualTextureProducer.h
@@ -20,8 +20,8 @@ public:
 	inline const FVTProducerDescription& GetDescription() const { return Description; }
 	inline IVirtualTexture* GetVirtualTexture() const { return VirtualTexture; }
 	inline uint32 GetNumLayers() const { return Description.NumLayers; }
-	inline uint32 GetWidthInTiles() const { return Description.WidthInTiles; }
-	inline uint32 GetHeightInTiles() const { return Description.HeightInTiles; }
+	inline uint32 GetWidthInTiles() const { return Description.BlockWidthInTiles * Description.WidthInBlocks; }
+	inline uint32 GetHeightInTiles() const { return Description.BlockHeightInTiles * Description.HeightInBlocks; }
 	inline uint32 GetDepthInTiles() const { return Description.DepthInTiles; }
 	inline EPixelFormat GetLayerFormat(uint32 LayerIndex) const { check(LayerIndex < Description.NumLayers); return Description.LayerFormat[LayerIndex]; }
 	inline FVirtualTexturePhysicalSpace* GetPhysicalSpace(uint32 LayerIndex) const { check(LayerIndex < Description.NumLayers); return PhysicalSpace[LayerIndex]; }
diff --git a/Engine/Source/Runtime/Renderer/Private/VT/VirtualTextureSystem.cpp b/Engine/Source/Runtime/Renderer/Private/VT/VirtualTextureSystem.cpp
--- a/Engine/Source/Runtime/Renderer/Private/VT/VirtualTextureSystem.cpp
+++ b/Engine/Source/Runtime/Renderer/Private/VT/VirtualTextureSystem.cpp
@@ -326,8 +326,10 @@ IAllocatedVirtualTexture* FVirtualTextureSystem::AllocateVirtualTexture(const FA
 		return AllocatedVT;
 	}
 
-	uint32 WidthInTiles = 0u;
-	uint32 HeightInTiles = 0u;
+	uint32 BlockWidthInTiles = 0u;
+	uint32 BlockHeightInTiles = 0u;
+	uint32 WidthInBlocks = 0u;
+	uint32 HeightInBlocks = 0u;
 	uint32 DepthInTiles = 0u;
 	bool bSupport16BitPageTable = true;
 	FVirtualTextureProducer* ProducerForLayer[VIRTUALTEXTURE_SPACE_MAXLAYERS] = { nullptr };
@@ -338,9 +340,13 @@ IAllocatedVirtualTexture* FVirtualTextureSystem::AllocateVirtualTexture(const FA
 		ProducerForLayer[LayerIndex] = Producer;
 		if (Producer)
 		{
-			WidthInTiles = FMath::Max(WidthInTiles, Producer->GetWidthInTiles());
-			HeightInTiles = FMath::Max(HeightInTiles, Producer->GetHeightInTiles());
-			DepthInTiles = FMath::Max(DepthInTiles, Producer->GetDepthInTiles());
+			const FVTProducerDescription& ProducerDesc = Producer->GetDescription();
+			BlockWidthInTiles = FMath::Max(BlockWidthInTiles, ProducerDesc.BlockWidthInTiles);
+			BlockHeightInTiles = FMath::Max(BlockHeightInTiles, ProducerDesc.BlockHeightInTiles);
+			WidthInBlocks = FMath::Max<uint32>(WidthInBlocks, ProducerDesc.WidthInBlocks);
+			HeightInBlocks = FMath::Max<uint32>(HeightInBlocks, ProducerDesc.HeightInBlocks);
+			DepthInTiles = FMath::Max(DepthInTiles, ProducerDesc.DepthInTiles);
+
 			FVirtualTexturePhysicalSpace* PhysicalSpace = Producer->GetPhysicalSpace(Desc.LocalLayerToProduce[LayerIndex]);
 			if (!PhysicalSpace->DoesSupport16BitPageTable())
 			{
@@ -350,8 +356,10 @@ IAllocatedVirtualTexture* FVirtualTextureSystem::AllocateVirtualTexture(const FA
 		}
 	}
 
-	check(WidthInTiles > 0u);
-	check(HeightInTiles > 0u);
+	check(BlockWidthInTiles > 0u);
+	check(BlockHeightInTiles > 0u);
+	check(WidthInBlocks > 0u);
+	check(HeightInBlocks > 0u);
 	check(DepthInTiles > 0u);
 
 	FVTSpaceDescription SpaceDesc;
@@ -361,9 +369,9 @@ IAllocatedVirtualTexture* FVirtualTextureSystem::AllocateVirtualTexture(const FA
 	SpaceDesc.TileBorderSize = Desc.TileBorderSize;
 	SpaceDesc.bPrivateSpace = Desc.bPrivateSpace;
 	SpaceDesc.Format = bSupport16BitPageTable ? EVTPageTableFormat::UInt16 : EVTPageTableFormat::UInt32;
-	FVirtualTextureSpace* Space = AcquireSpace(SpaceDesc, FMath::Max(WidthInTiles, HeightInTiles));
+	FVirtualTextureSpace* Space = AcquireSpace(SpaceDesc, FMath::Max(BlockWidthInTiles * WidthInBlocks, BlockHeightInTiles * HeightInBlocks));
 
-	AllocatedVT = new FAllocatedVirtualTexture(Frame, Desc, Space, ProducerForLayer, WidthInTiles, HeightInTiles, DepthInTiles);
+	AllocatedVT = new FAllocatedVirtualTexture(this, Frame, Desc, Space, ProducerForLayer, BlockWidthInTiles, BlockHeightInTiles, WidthInBlocks, HeightInBlocks, DepthInTiles);
 	if (bAnyLayerProducerWantsPersistentHighestMip)
 	{
 		AllocatedVTsToMap.Add(AllocatedVT);
@@ -421,6 +429,11 @@ void FVirtualTextureSystem::ReleaseProducer(const FVirtualTextureProducerHandle&
 	Producers.ReleaseProducer(this, Handle);
 }
 
+FVirtualTextureProducer* FVirtualTextureSystem::FindProducer(const FVirtualTextureProducerHandle& Handle)
+{
+	return Producers.FindProducer(Handle);
+}
+
 FVirtualTextureSpace* FVirtualTextureSystem::AcquireSpace(const FVTSpaceDescription& InDesc, uint32 InSizeNeeded)
 {
 	// If InDesc requests a private space, don't reuse any existing spaces
@@ -501,30 +514,51 @@ void FVirtualTextureSystem::ReleasePhysicalSpace(FVirtualTexturePhysicalSpace* S
 void FVirtualTextureSystem::LockTile(const FVirtualTextureLocalTile& Tile)
 {
 	check(IsInRenderingThread());
-	TilesToLock.Add(Tile);
+
+	if (TileLocks.Lock(Tile))
+	{
+		checkSlow(!TilesToLock.Contains(Tile));
+		TilesToLock.Add(Tile);
+	}
+}
+
+static void UnlockTileInternal(const FVirtualTextureProducerHandle& ProducerHandle, const FVirtualTextureProducer* Producer, const FVirtualTextureLocalTile& Tile, uint32 Frame)
+{
+	for (uint32 LocalLayerIndex = 0u; LocalLayerIndex < Producer->GetNumLayers(); ++LocalLayerIndex)
+	{
+		FVirtualTexturePhysicalSpace* PhysicalSpace = Producer->GetPhysicalSpace(LocalLayerIndex);
+		FTexturePagePool& PagePool = PhysicalSpace->GetPagePool();
+		const uint32 pAddress = PagePool.FindPageAddress(ProducerHandle, LocalLayerIndex, Tile.Local_vAddress, Tile.Local_vLevel);
+		if (pAddress != ~0u)
+		{
+			PagePool.Unlock(Frame, pAddress);
+		}
+	}
+}
+
+void FVirtualTextureSystem::UnlockTile(const FVirtualTextureLocalTile& Tile, const FVirtualTextureProducer* Producer)
+{
+	check(IsInRenderingThread());
+
+	if (TileLocks.Unlock(Tile))
+	{
+		// Tile is no longer locked
+		TilesToLock.Remove(Tile);
+		UnlockTileInternal(Tile.GetProducerHandle(), Producer, Tile, Frame);
+	}
 }
 
-void FVirtualTextureSystem::UnlockTile(const FVirtualTextureLocalTile& Tile)
+void FVirtualTextureSystem::ForceUnlockAllTiles(const FVirtualTextureProducerHandle& ProducerHandle, const FVirtualTextureProducer* Producer)
 {
 	check(IsInRenderingThread());
 
-	// Tile is no longer locked
-	TilesToLock.Remove(Tile);
+	TArray<FVirtualTextureLocalTile> TilesToUnlock;
+	TileLocks.ForceUnlockAll(ProducerHandle, TilesToUnlock);
 
-	const FVirtualTextureProducerHandle ProducerHandle = Tile.GetProducerHandle();
-	const FVirtualTextureProducer* Producer = Producers.FindProducer(ProducerHandle);
-	if (Producer)
+	for (const FVirtualTextureLocalTile& Tile : TilesToUnlock)
 	{
-		for (uint32 LocalLayerIndex = 0u; LocalLayerIndex < Producer->GetNumLayers(); ++LocalLayerIndex)
-		{
-			FVirtualTexturePhysicalSpace* PhysicalSpace = Producer->GetPhysicalSpace(LocalLayerIndex);
-			FTexturePagePool& PagePool = PhysicalSpace->GetPagePool();
-			const uint32 pAddress = PagePool.FindPageAddress(ProducerHandle, LocalLayerIndex, Tile.Local_vAddress, Tile.Local_vLevel);
-			if (pAddress != ~0u)
-			{
-				PagePool.Unlock(Frame, pAddress);
-			}
-		}
+		TilesToLock.Remove(Tile);
+		UnlockTileInternal(ProducerHandle, Producer, Tile, Frame);
 	}
 }
 
@@ -805,6 +839,7 @@ void FVirtualTextureSystem::Update(FRHICommandListImmediate& RHICmdList, ERHIFea
 		{
 			const FVirtualTextureProducerHandle ProducerHandle = Tile.GetProducerHandle();
 			const FVirtualTextureProducer* Producer = Producers.FindProducer(ProducerHandle);
+			checkSlow(TileLocks.IsLocked(Tile));
 			if (Producer)
 			{
 				uint8 LocalLayerMaskToLoad = 0u;
@@ -1153,6 +1188,7 @@ void FVirtualTextureSystem::GatherRequestsTask(const FGatherRequestsParameters&
 				continue;
 			}
 
+			const uint32 MaxLevel = FMath::Min(Producer->GetMaxLevel(), AllocatedVT->GetMaxLevel());
 			const uint32 ProducerMipBias = AllocatedVT->GetUniqueProducerMipBias(ProducerIndex);
 			uint32 Mapping_vLevel = FMath::Max(vLevel, ProducerMipBias);
 
@@ -1163,16 +1199,25 @@ void FVirtualTextureSystem::GatherRequestsTask(const FGatherRequestsParameters&
 			// Local_vLevel is the level within the producer that we want to allocate/map
 			// here we subtract ProducerMipBias (clamped to ensure we don't fall below 0),
 			// which effectively matches more detailed mips of lower resolution producers with less detailed mips of higher resolution producers
-			uint8 Local_vLevel = vLevel - FMath::Min(vLevel, ProducerMipBias);
+			uint32 Local_vLevel = vLevel - FMath::Min(vLevel, ProducerMipBias);
+
+			// Wrap Local_vAddress for the given producer
+			// For square textures, this could simply be (Local_vAddress % NumTilesInMip), but that doesn't work for non-square
+			// Possible there is a more clever approach to take here
+			{
+				const uint32 ProducerMipWidthInTiles = FMath::Max(Producer->GetWidthInTiles() >> Local_vLevel, 1u);
+				const uint32 ProducerMipHeightInTiles = FMath::Max(Producer->GetWidthInTiles() >> Local_vLevel, 1u);
+				uint32 Local_vTileX = FMath::ReverseMortonCode2(Local_vAddress);
+				uint32 Local_vTileY = FMath::ReverseMortonCode2(Local_vAddress >> 1);
+				Local_vTileX &= (ProducerMipWidthInTiles - 1u);
+				Local_vTileY &= (ProducerMipHeightInTiles - 1u);
+				Local_vAddress = FMath::MortonCode2(Local_vTileX) | (FMath::MortonCode2(Local_vTileY) << 1);
+			}
 
 			const uint32 LocalMipBias = Producer->GetVirtualTexture()->GetLocalMipBias(Local_vLevel, Local_vAddress);
 			if (LocalMipBias > 0u)
 			{
 				Local_vLevel += LocalMipBias;
-				if (Local_vLevel > Producer->GetMaxLevel())
-				{
-					continue;
-				}
 				Local_vAddress >>= (LocalMipBias * vDimensions);
 				Mapping_vLevel = FMath::Max(vLevel, LocalMipBias + ProducerMipBias);
 			}
@@ -1191,15 +1236,15 @@ void FVirtualTextureSystem::GatherRequestsTask(const FGatherRequestsParameters&
 				const FTexturePagePool& RESTRICT PagePool = PhysicalSpace->GetPagePool();
 
 				// Find the highest resolution tile that's currently loaded
-				const uint32 pAddress = PagePool.FindNearestPageAddress(ProducerHandle, LocalLayerIndex, Local_vAddress, Local_vLevel, Producer->GetMaxLevel());
-				uint32 AllocatedLocal_vLevel = Producer->GetMaxLevel() + 1u;
+				const uint32 pAddress = PagePool.FindNearestPageAddress(ProducerHandle, LocalLayerIndex, Local_vAddress, Local_vLevel, MaxLevel);
+				uint32 AllocatedLocal_vLevel = MaxLevel + 1u;
 				if (pAddress != ~0u)
 				{
 					AllocatedLocal_vLevel = PagePool.GetLocalLevelForAddress(pAddress);
 					check(AllocatedLocal_vLevel >= Local_vLevel);
 
 					const uint32 Allocated_vLevel = AllocatedLocal_vLevel + ProducerMipBias;
-					check(Allocated_vLevel <= AllocatedVT->GetMaxLevel());
+					ensure(Allocated_vLevel <= AllocatedVT->GetMaxLevel());
 
 					const uint32 AllocatedMapping_vLevel = FMath::Max(Allocated_vLevel, ProducerMipBias);
 					const uint32 Allocated_vAddress = vAddress & (0xffffffff << (Allocated_vLevel * vDimensions));
@@ -1302,7 +1347,7 @@ void FVirtualTextureSystem::GatherRequestsTask(const FGatherRequestsParameters&
 					if (LoadRequestIndex != 0xffff)
 					{
 						const uint32 Prefetch_vLevel = PrefetchLocal_vLevel + ProducerMipBias;
-						check(Prefetch_vLevel <= AllocatedVT->GetMaxLevel());
+						ensure(Prefetch_vLevel <= AllocatedVT->GetMaxLevel());
 						const uint32 PrefetchMapping_vLevel = FMath::Max(Prefetch_vLevel, ProducerMipBias);
 						const uint32 Prefetch_vAddress = vAddress & (0xffffffff << (Prefetch_vLevel * vDimensions));
 						for (uint32 LoadLayerIndex = 0u; LoadLayerIndex < NumLayersToLoad; ++LoadLayerIndex)
@@ -1321,7 +1366,6 @@ void FVirtualTextureSystem::GatherRequestsTask(const FGatherRequestsParameters&
 				}
 			}
 
-			// it's possible that 'LocalLayerMaskToLoad' is now 0, if all the required pages were already resident and simply needed to be mapped
 			if (LocalLayerMaskToLoad != 0u)
 			{
 				const uint16 LoadRequestIndex = RequestList->AddLoadRequest(FVirtualTextureLocalTile(ProducerHandle, Local_vAddress, Local_vLevel), LocalLayerMaskToLoad, PageCount);
@@ -1631,8 +1675,6 @@ void FVirtualTextureSystem::SubmitRequests(FRHICommandListImmediate& RHICmdList,
 		{
 			const FAllocatedVirtualTexture* AllocatedVT = AllocatedVTsToMap[Index];
 			const uint32 vDimensions = AllocatedVT->GetDimensions();
-			const uint32 WidthInTiles = AllocatedVT->GetWidthInTiles();
-			const uint32 HeightInTiles = AllocatedVT->GetHeightInTiles();
 			const uint32 BaseTileX = FMath::ReverseMortonCode2(AllocatedVT->GetVirtualAddress());
 			const uint32 BaseTileY = FMath::ReverseMortonCode2(AllocatedVT->GetVirtualAddress() >> 1);
 			FVirtualTextureSpace* Space = AllocatedVT->GetSpace();
@@ -1645,12 +1687,16 @@ void FVirtualTextureSystem::SubmitRequests(FRHICommandListImmediate& RHICmdList,
 				const uint32 LocalLayerIndex = AllocatedVT->GetLocalLayerToProduce(LayerIndex);
 				const FVirtualTextureProducerHandle ProducerHandle = AllocatedVT->GetUniqueProducerHandle(ProducerIndex);
 				const FVirtualTextureProducer* Producer = Producers.FindProducer(ProducerHandle);
+				const uint32 WidthInTiles = Producer->GetWidthInTiles();
+				const uint32 HeightInTiles = Producer->GetHeightInTiles();
+				const uint32 MaxLevel = FMath::Min(Producer->GetMaxLevel(), AllocatedVT->GetMaxLevel());
+
 				FVirtualTexturePhysicalSpace* PhysicalSpace = AllocatedVT->GetPhysicalSpace(LayerIndex);
 				FTexturePagePool& PagePool = PhysicalSpace->GetPagePool();
 				FTexturePageMap& PageMap = Space->GetPageMap(LayerIndex);
 				
 				bool bIsLayerFullyMapped = false;
-				for (uint32 Local_vLevel = 0u; Local_vLevel <= Producer->GetMaxLevel(); ++Local_vLevel)
+				for (uint32 Local_vLevel = 0u; Local_vLevel <= MaxLevel; ++Local_vLevel)
 				{
 					const uint32 vLevel = Local_vLevel + ProducerMipBias;
 					const uint32 LevelWidthInTiles = FMath::Max(WidthInTiles >> vLevel, 1u);
diff --git a/Engine/Source/Runtime/Renderer/Private/VT/VirtualTextureSystem.h b/Engine/Source/Runtime/Renderer/Private/VT/VirtualTextureSystem.h
--- a/Engine/Source/Runtime/Renderer/Private/VT/VirtualTextureSystem.h
+++ b/Engine/Source/Runtime/Renderer/Private/VT/VirtualTextureSystem.h
@@ -7,6 +7,7 @@
 #include "RendererInterface.h"
 #include "Templates/UniquePtr.h"
 #include "VT/VirtualTextureProducer.h"
+#include "VT/TexturePageLocks.h"
 #include "VirtualTexturing.h"
 
 class FAllocatedVirtualTexture;
@@ -42,6 +43,7 @@ public:
 
 	FVirtualTextureProducerHandle RegisterProducer(const FVTProducerDescription& InDesc, IVirtualTexture* InProducer);
 	void ReleaseProducer(const FVirtualTextureProducerHandle& Handle);
+	FVirtualTextureProducer* FindProducer(const FVirtualTextureProducerHandle& Handle);
 
 	FVirtualTextureSpace* AcquireSpace(const FVTSpaceDescription& InDesc, uint32 InSizeNeeded);
 	void ReleaseSpace(FVirtualTextureSpace* Space);
@@ -53,7 +55,8 @@ public:
 	FVirtualTexturePhysicalSpace* GetPhysicalSpace(uint16 ID) const { return PhysicalSpaces[ID].Get(); }
 
 	void LockTile(const FVirtualTextureLocalTile& Tile);
-	void UnlockTile(const FVirtualTextureLocalTile& Tile);
+	void UnlockTile(const FVirtualTextureLocalTile& Tile, const FVirtualTextureProducer* Producer);
+	void ForceUnlockAllTiles(const FVirtualTextureProducerHandle& ProducerHandle, const FVirtualTextureProducer* Producer);
 	void RequestTilesForRegion(const IAllocatedVirtualTexture* AllocatedVT, const FVector2D& InScreenSpaceSize, const FIntRect& InTextureRegion, int32 InMipLevel = -1);
 	void LoadPendingTiles(FRHICommandListImmediate& RHICmdList, ERHIFeatureLevel::Type FeatureLevel);
 	void FlushCache();
@@ -108,6 +111,8 @@ private:
 	TArray<uint32> RequestedPackedTiles;
 
 	TArray<FVirtualTextureLocalTile> TilesToLock;
+	FTexturePageLocks TileLocks;
+
 	FCriticalSection ContinuousUpdateTilesToProduceCS;
 	TSet<FVirtualTextureLocalTile> ContinuousUpdateTilesToProduce;
 	TSet<FVirtualTextureLocalTile> MappedTilesToProduce;
