Fixed exposure was not taken into account in raytracing barycentrics and view mode

diff --git a/Engine/Shaders/Private/RayTracing/RayTracingBarycentrics.usf b/Engine/Shaders/Private/RayTracing/RayTracingBarycentrics.usf
--- a/Engine/Shaders/Private/RayTracing/RayTracingBarycentrics.usf
+++ b/Engine/Shaders/Private/RayTracing/RayTracingBarycentrics.usf
@@ -49,7 +49,7 @@ void RayTracingBarycentricsMainRGS()
 	{
 		Result = float4(Payload.Barycentrics, 1.0f - Payload.Barycentrics.x - Payload.Barycentrics.y, 1.0f);
 	}
-	Output[PixelCoord] = Result;
+	Output[PixelCoord] = Result * View.PreExposure;
 }
 
 [shader("miss")]
diff --git a/Engine/Shaders/Private/RayTracing/RayTracingDebug.usf b/Engine/Shaders/Private/RayTracing/RayTracingDebug.usf
--- a/Engine/Shaders/Private/RayTracing/RayTracingDebug.usf
+++ b/Engine/Shaders/Private/RayTracing/RayTracingDebug.usf
@@ -106,7 +106,7 @@ void RayTracingDebugMainRGS()
 		}
 	}
 
-	Output[PixelCoord] = Result;
+	Output[PixelCoord] = Result * View.PreExposure;
 }
 
 [shader("miss")]
