Restores some of the previous behavior of SSR in mean time this get investigated further.

diff --git a/Engine/Shaders/Private/ScreenSpaceDiffuseIndirect.usf b/Engine/Shaders/Private/ScreenSpaceDiffuseIndirect.usf
--- a/Engine/Shaders/Private/ScreenSpaceDiffuseIndirect.usf
+++ b/Engine/Shaders/Private/ScreenSpaceDiffuseIndirect.usf
@@ -1,5 +1,7 @@
 // Copyright 1998-2019 Epic Games, Inc. All Rights Reserved.
 
+#define IS_SSGI_SHADER 1
+
 #include "Common.ush"	
 #include "PostProcessCommon.ush"				
 #include "DeferredShadingCommon.ush"
@@ -68,11 +70,14 @@ void ScreenSpaceDiffuseIndirectPS(
 	//uint2 Random = 0x10000 * Noise;
 	uint2 Random = Rand3DPCG16( int3( PixelPos, View.StateFrameIndexMod8 ) ).xy;
 			
+	// TODO: lot of VGPR, worth recomputing every ray instead?
 	float3x3 TangentBasis = GetTangentBasis( N );
 
 	// Shoot multiple rays
 	LOOP for( int i = 0; i < NumRays; i++ )
 	{
+		// TOOD: VGPR, might be worth recomputing instead?
+		// SvPosition: is SvPosition really needed? Can recompute from PositionTranslatedWorld?
 		float StepOffset = Noise.x;
 		StepOffset -= 0.9;
 			
diff --git a/Engine/Shaders/Private/ScreenSpaceRayCast.ush b/Engine/Shaders/Private/ScreenSpaceRayCast.ush
--- a/Engine/Shaders/Private/ScreenSpaceRayCast.ush
+++ b/Engine/Shaders/Private/ScreenSpaceRayCast.ush
@@ -2,6 +2,10 @@
 
 #pragma once
 
+#ifndef IS_SSGI_SHADER
+#define IS_SSGI_SHADER 0
+#endif
+
 #include "Common.ush"	
 
 bool RayCast(
@@ -40,8 +44,11 @@ bool RayCast(
 	float3 RayStepUVz  = float3(  RayStepScreen.xy  * float2( 0.5, -0.5 )		 * HZBUvFactorAndInvFactor.xy, RayStepScreen.z );
 	
 	const float Step = 1.0 / NumSteps;
-	float CompareTolerance = max( abs( RayStepUVz.z ), ( RayStartScreen.z - RayDepthScreen.z ) * 2 ) * Step;
-	//const float CompareTolerance = abs( RayDepthScreen.z - RayStartScreen.z ) * Step * 2;
+	#if IS_SSGI_SHADER
+		float CompareTolerance = max( abs( RayStepUVz.z ), ( RayStartScreen.z - RayDepthScreen.z ) * 2 ) * Step;
+	#else
+		float CompareTolerance = abs( RayDepthScreen.z - RayStartScreen.z ) * Step * 4;
+	#endif
 
 	// avoid bugs with early returns inside of loops on certain platform compilers.
 	bool bHit = false;
@@ -76,7 +83,7 @@ bool RayCast(
 
 		BRANCH if( any( Hit ) )
 		{
-			if(1)
+			if(IS_SSGI_SHADER)
 			{
 				// If hit set to intersect time. If missed set to beyond end of ray
 				float4 HitTime = Hit ? float4( 1, 2, 3, 4 ) : 5;
diff --git a/Engine/Shaders/Private/ScreenSpaceReflections.usf b/Engine/Shaders/Private/ScreenSpaceReflections.usf
--- a/Engine/Shaders/Private/ScreenSpaceReflections.usf
+++ b/Engine/Shaders/Private/ScreenSpaceReflections.usf
@@ -188,7 +188,11 @@ void ScreenSpaceReflectionsPS(
 		LOOP for( uint i = 0; i < NumRays; i++ )
 		{
 			float StepOffset = Noise.x;
-			StepOffset -= 0.9;
+			#if 0 // TODO
+				StepOffset -= 0.9;
+			#else
+				StepOffset -= 0.5;
+			#endif
 			
 			float2 E = Hammersley16( i, NumRays, Random );
 
@@ -285,7 +289,11 @@ void ScreenSpaceReflectionsPS(
 	else
 	{
 		float StepOffset = InterleavedGradientNoise(SvPosition.xy, View.StateFrameIndexMod8);
-		StepOffset -= 0.9;
+		#if 0 // TODO
+			StepOffset -= 0.9;
+		#else
+			StepOffset -= 0.5;
+		#endif
 		
 		float3 L;
 		if (bGlossy)
