Cuts ties of SSGI from post process graph shader header.

diff --git a/Engine/Shaders/Private/SSRT/SSRTDiffuseIndirect.usf b/Engine/Shaders/Private/SSRT/SSRTDiffuseIndirect.usf
--- a/Engine/Shaders/Private/SSRT/SSRTDiffuseIndirect.usf
+++ b/Engine/Shaders/Private/SSRT/SSRTDiffuseIndirect.usf
@@ -2,8 +2,7 @@
 
 #define IS_SSGI_SHADER 1
 
-#include "SSRTRayCast.ush"	
-#include "../PostProcessCommon.ush"				
+#include "SSRTRayCast.ush"
 #include "../DeferredShadingCommon.ush"
 #include "../Random.ush"
 #include "../BRDF.ush"
