Cuts ties of DOF from post process graph shader header.

diff --git a/Engine/Shaders/Private/DiaphragmDOF/DOFBokehLUT.usf b/Engine/Shaders/Private/DiaphragmDOF/DOFBokehLUT.usf
--- a/Engine/Shaders/Private/DiaphragmDOF/DOFBokehLUT.usf
+++ b/Engine/Shaders/Private/DiaphragmDOF/DOFBokehLUT.usf
@@ -153,10 +153,6 @@ void RingCoordinateFromCartesianCoordinate(const float2 C, out float OutRingRadi
 
 //------------------------------------------------------- OUTPUTS
 
-#if !defined(NEW_CODE) // TODO: remove me
-#define BokehLUTOutput Output0
-#endif
-
 RWTexture2D<float4>	BokehLUTOutput;
 
 
diff --git a/Engine/Shaders/Private/DiaphragmDOF/DOFCocTileDilate.usf b/Engine/Shaders/Private/DiaphragmDOF/DOFCocTileDilate.usf
--- a/Engine/Shaders/Private/DiaphragmDOF/DOFCocTileDilate.usf
+++ b/Engine/Shaders/Private/DiaphragmDOF/DOFCocTileDilate.usf
@@ -18,6 +18,7 @@
 
 //------------------------------------------------------- PASSES.
 
+uint4 ViewportRect;
 uint SampleOffsetMultipler = 1;
 float fSampleOffsetMultipler = 1;
 
@@ -31,22 +32,11 @@ Texture2D TileInput_Background;
 Texture2D DilatedTileMinMax_Foreground;
 Texture2D DilatedTileMinMax_Background;
 
-#if defined(NEW_CODE) // TODO: remove me
-#define Output0 TileOutput_Foreground
-#define Output1 TileOutput_Background
-#else
-#define TileInput_Foreground PostprocessInput0
-#define TileInput_Background PostprocessInput1
-
-#define DilatedTileMinMax_Foreground PostprocessInput2
-#define DilatedTileMinMax_Background PostprocessInput3
-#endif
-
 
 //------------------------------------------------------- OUTPUTS.
 
-RWTexture2D<float4>	Output0;
-RWTexture2D<float4>	Output1;
+RWTexture2D<float4>	TileOutput_Foreground;
+RWTexture2D<float4>	TileOutput_Background;
 
 
 //------------------------------------------------------- ENTRY POINT.
@@ -257,10 +247,10 @@ void CocDilateMainCS(uint2 DispatchThreadId : SV_DispatchThreadID)
 		float4 RawOutput[COC_TILE_MAX_BUFFER_COUNT];
 		EncodeCocTileSample(OutSample, CocTileLayout, RawOutput);
 
-		Output0[PixelPos] = RawOutput[0];
+		TileOutput_Foreground[PixelPos] = RawOutput[0];
 		if (static_condition(GetBufferCountForCocTileLayour(CocTileLayout) > 1))
 		{
-			Output1[PixelPos] = RawOutput[1];
+			TileOutput_Background[PixelPos] = RawOutput[1];
 		}
 	}
 }
diff --git a/Engine/Shaders/Private/DiaphragmDOF/DOFCocTileFlatten.usf b/Engine/Shaders/Private/DiaphragmDOF/DOFCocTileFlatten.usf
--- a/Engine/Shaders/Private/DiaphragmDOF/DOFCocTileFlatten.usf
+++ b/Engine/Shaders/Private/DiaphragmDOF/DOFCocTileFlatten.usf
@@ -24,27 +24,18 @@
 
 //------------------------------------------------------- PARAMETERS.
 
+uint4 ViewportRect;
 float2 ThreadIdToBufferUV;
 float2 MaxBufferUV;
 
 Texture2D GatherInput_SceneColor;
 Texture2D GatherInput_SeparateCoc;
 
-#if defined(NEW_CODE) // TODO: remove me
-#define Output0 TileOutput_Foreground
-#define Output1 TileOutput_Background
-#else
-#define GatherInput_SceneColor PostprocessInput0
-#define GatherInput_SeparateCoc PostprocessInput0
-#endif
-
 
 //------------------------------------------------------- OUTPUTS.
 
-RWTexture2D<float4>	Output0;
-RWTexture2D<float4>	Output1;
-//RWTexture2D<float4>	TileOutput_Foreground;
-//RWTexture2D<float4>	TileOutput_Background;
+RWTexture2D<float4>	TileOutput_Foreground;
+RWTexture2D<float4>	TileOutput_Background;
 
 
 //------------------------------------------------------- LDS.
@@ -217,10 +208,10 @@ void CocFlattenMainCS(
 			EncodeCocTileSample(Local, LDS_LAYOUT, RawOutput);
 	
 			// Output.
-			Output0[GroupId] = RawOutput[0];
+			TileOutput_Foreground[GroupId] = RawOutput[0];
 			if (static_condition(GetBufferCountForCocTileLayour(LDS_LAYOUT) > 1))
 			{
-				Output1[GroupId] = RawOutput[1];
+				TileOutput_Background[GroupId] = RawOutput[1];
 			}
 		}
 	}
@@ -304,10 +295,10 @@ void CocFlattenMainCS(
 			EncodeCocTileSample(OutputTileSample, LDS_LAYOUT, RawOutput);
 	
 			// Output.
-			Output0[GroupId] = RawOutput[0];
+			TileOutput_Foreground[GroupId] = RawOutput[0];
 			if (static_condition(GetBufferCountForCocTileLayour(LDS_LAYOUT) > 1))
 			{
-				Output1[GroupId] = RawOutput[1];
+				TileOutput_Background[GroupId] = RawOutput[1];
 			}
 		}
 	}
diff --git a/Engine/Shaders/Private/DiaphragmDOF/DOFCommon.ush b/Engine/Shaders/Private/DiaphragmDOF/DOFCommon.ush
--- a/Engine/Shaders/Private/DiaphragmDOF/DOFCommon.ush
+++ b/Engine/Shaders/Private/DiaphragmDOF/DOFCommon.ush
@@ -7,7 +7,6 @@
 #pragma once
 
 #include "../Common.ush"
-#include "../PostProcessCommon.ush"
 
 // Piggy back on existing Circle DOF code.
 #include "../CircleDOFCommon.ush"
diff --git a/Engine/Shaders/Private/DiaphragmDOF/DOFDownsample.usf b/Engine/Shaders/Private/DiaphragmDOF/DOFDownsample.usf
--- a/Engine/Shaders/Private/DiaphragmDOF/DOFDownsample.usf
+++ b/Engine/Shaders/Private/DiaphragmDOF/DOFDownsample.usf
@@ -9,6 +9,7 @@
 
 //------------------------------------------------------- PARAMETERS
 
+uint4 ViewportRect;
 float2 MaxBufferUV;
 float OutputCocRadiusMultiplier;
 
@@ -16,24 +17,13 @@ float4 GatherInputSize;
 Texture2D GatherInput_SceneColor;
 Texture2D GatherInput_SeparateCoc;
 
-#if defined(NEW_CODE) // TODO: remove me
-#define Output0 OutDownsampledGatherInput_SceneColor
-#define Output1 OutDownsampledGatherInput_SeparateCoc
-
-#else
-#define GatherInputSize PostprocessInput0Size
-#define GatherInput_SceneColor PostprocessInput0
-#define GatherInput_SeparateCoc PostprocessInput1
-
-#endif
-
 
 //------------------------------------------------------- OUTPUTS
 
-RWTexture2D<float4>	Output0;
+RWTexture2D<float4>	OutDownsampledGatherInput_SceneColor;
 
 #if CONFIG_DOF_ALPHA
-RWTexture2D<float>	Output1;
+RWTexture2D<float>	OutDownsampledGatherInput_SeparateCoc;
 #endif
 
 
@@ -87,12 +77,12 @@ void DownsampleCS(uint2 DispatchThreadId : SV_DispatchThreadID)
 	{
 		#if CONFIG_DOF_ALPHA
 		{
-			Output0[DispatchThreadId] = OutColor;
-			Output1[DispatchThreadId] = OutCocRadius * OutputCocRadiusMultiplier;
+			OutDownsampledGatherInput_SceneColor[DispatchThreadId] = OutColor;
+			OutDownsampledGatherInput_SeparateCoc[DispatchThreadId] = OutCocRadius * OutputCocRadiusMultiplier;
 		}
 		#else
 		{
-			Output0[DispatchThreadId] = float4(OutColor.rgb, OutCocRadius * OutputCocRadiusMultiplier);
+			OutDownsampledGatherInput_SceneColor[DispatchThreadId] = float4(OutColor.rgb, OutCocRadius * OutputCocRadiusMultiplier);
 		}
 		#endif
 	}
diff --git a/Engine/Shaders/Private/DiaphragmDOF/DOFGatherKernel.ush b/Engine/Shaders/Private/DiaphragmDOF/DOFGatherKernel.ush
--- a/Engine/Shaders/Private/DiaphragmDOF/DOFGatherKernel.ush
+++ b/Engine/Shaders/Private/DiaphragmDOF/DOFGatherKernel.ush
@@ -12,6 +12,8 @@
 
 //------------------------------------------------------- PARAMETERS
 
+float4 ViewportSize;
+
 float4 GatherInputSize;
 Texture2D GatherInput_SceneColor;
 Texture2D GatherInput_SeparateCoc;
@@ -19,14 +21,6 @@ Texture2D GatherInput_SeparateCoc;
 Texture2D BokehLUT;
 
 
-#if !defined(NEW_CODE) // TODO: remove me
-#define GatherInputSize PostprocessInput0Size
-#define GatherInput_SceneColor PostprocessInput0
-#define GatherInput_SeparateCoc PostprocessInput1
-#define BokehLUT PostprocessInput4
-#endif
-
-
 //------------------------------------------------------- GATHERING KERNEL
 
 // LARGEST_RINGS_FIRST controles the order the rings must be proceded.
diff --git a/Engine/Shaders/Private/DiaphragmDOF/DOFGatherPass.usf b/Engine/Shaders/Private/DiaphragmDOF/DOFGatherPass.usf
--- a/Engine/Shaders/Private/DiaphragmDOF/DOFGatherPass.usf
+++ b/Engine/Shaders/Private/DiaphragmDOF/DOFGatherPass.usf
@@ -25,6 +25,7 @@ float2 TemporalJitterPixels;
 // Mip bias for the gathering kernel.
 float MipBias;
 
+uint4 ViewportRect;
 float2 DispatchThreadIdToInputBufferUV;
 float MaxRecombineAbsCocRadius;
 
@@ -41,15 +42,6 @@ Texture2D TileClassification_Background;
 
 //------------------------------------------------------- OUTPUTS
 
-#if !defined(NEW_CODE) // TODO: remove me
-#define ConvolutionOutput_SceneColor Output0
-#define ConvolutionOutput_SeparateAlpha Output1
-#define ScatterOcclusionOutput Output2
-
-#define TileClassification_Foreground PostprocessInput2
-#define TileClassification_Background PostprocessInput3
-#endif
-
 RWTexture2D<float4>	ConvolutionOutput_SceneColor;
 RWTexture2D<float4>	ConvolutionOutput_SeparateAlpha;
 
@@ -101,7 +93,6 @@ void StoreAccumulatorOutput(in uint2 OutputPixelPosition, in FAccumulatorOutput
 		#elif DIM_LAYER_PROCESSING == LAYER_PROCESSING_BACKGROUND_ONLY
 			ConvolutionOutput_SceneColor[OutputPixelPosition] = AccumulatorOutput.BackgroundColor;
 			ConvolutionOutput_SeparateAlpha[OutputPixelPosition] = AccumulatorOutput.BackgroundWeight > 0 ? 1 : 0;
-			//Output1[OutputPixelPosition] = float4(1, 0, 0, 0);
 		
 		#elif DIM_LAYER_PROCESSING == LAYER_PROCESSING_SLIGHT_OUT_OF_FOCUS
 			ConvolutionOutput_SceneColor[OutputPixelPosition] = float4(
@@ -282,19 +273,6 @@ void GatherMainCS(
 	}
 	#endif
 
-	#if DIM_LAYER_PROCESSING == LAYER_PROCESSING_FOREGROUND_HOLE_FILLING && 0
-		if (all(DispatchThreadId < ViewportRect.zw) && bSkipGathering)
-		{
-			//Output0[DispatchThreadId] = float4(
-			//	GatherParameters.KernelRadius,
-			//	MaxCocRadiusAbs,
-			//	MinCocRadiusAbs,
-			//	1);
-			Output0[DispatchThreadId] = float4(1, 0, 0, 0);
-		}
-		return;
-	#endif
-
 	// Early return if there is no convolution to do.
 	BRANCH
 	if (bSkipGathering)
diff --git a/Engine/Shaders/Private/DiaphragmDOF/DOFHybridScatterPixelShader.usf b/Engine/Shaders/Private/DiaphragmDOF/DOFHybridScatterPixelShader.usf
--- a/Engine/Shaders/Private/DiaphragmDOF/DOFHybridScatterPixelShader.usf
+++ b/Engine/Shaders/Private/DiaphragmDOF/DOFHybridScatterPixelShader.usf
@@ -28,14 +28,6 @@ float4 ScatterOcclusionSize;
 Texture2D ScatterOcclusion;
 
 
-#if !defined(NEW_CODE) // TODO: remove me
-#define BokehLUT PostprocessInput2
-#define ScatterOcclusionSize PostprocessInput3Size
-#define ScatterOcclusion PostprocessInput3
-#endif
-
-
-
 //------------------------------------------------------- ENTRY POINT
 
 void ScatterMainPS(
diff --git a/Engine/Shaders/Private/DiaphragmDOF/DOFHybridScatterVertexShader.usf b/Engine/Shaders/Private/DiaphragmDOF/DOFHybridScatterVertexShader.usf
--- a/Engine/Shaders/Private/DiaphragmDOF/DOFHybridScatterVertexShader.usf
+++ b/Engine/Shaders/Private/DiaphragmDOF/DOFHybridScatterVertexShader.usf
@@ -21,6 +21,8 @@ static const float2 kSpriteVertices[6] = {
 
 //------------------------------------------------------- PARAMETERS
 
+float4 ViewportSize;
+
 float ScatteringScaling;
 float CocRadiusToCircumscribedRadius;
 StructuredBuffer<float4> ScatterDrawList;
diff --git a/Engine/Shaders/Private/DiaphragmDOF/DOFPostfiltering.usf b/Engine/Shaders/Private/DiaphragmDOF/DOFPostfiltering.usf
--- a/Engine/Shaders/Private/DiaphragmDOF/DOFPostfiltering.usf
+++ b/Engine/Shaders/Private/DiaphragmDOF/DOFPostfiltering.usf
@@ -13,6 +13,7 @@
 #define THREADGROUP_SIZEY		(THREADGROUP_SIZEX)
 #define THREADGROUP_TOTALSIZE	(THREADGROUP_SIZEX * THREADGROUP_SIZEY)
 
+uint4 ViewportRect;
 float2 MaxInputBufferUV;
 float MinGatherRadius;
 
@@ -24,18 +25,6 @@ Texture2D ConvolutionInput_SeparateAlpha;
 Texture2D TileClassification_Foreground;
 Texture2D TileClassification_Background;
 
-#if !defined(NEW_CODE) // TODO: remove me
-#define ConvolutionOutput_SceneColor Output0
-#define ConvolutionOutput_SeparateAlpha Output1
-
-#define ConvolutionInputSize PostprocessInput0Size
-#define ConvolutionInput_SceneColor PostprocessInput0
-#define ConvolutionInput_SeparateAlpha PostprocessInput1
-
-#define TileClassification_Foreground PostprocessInput2
-#define TileClassification_Background PostprocessInput3
-#endif
-
 
 RWTexture2D<float4>	ConvolutionOutput_SceneColor;
 RWTexture2D<float4>	ConvolutionOutput_SeparateAlpha;
diff --git a/Engine/Shaders/Private/DiaphragmDOF/DOFRecombine.usf b/Engine/Shaders/Private/DiaphragmDOF/DOFRecombine.usf
--- a/Engine/Shaders/Private/DiaphragmDOF/DOFRecombine.usf
+++ b/Engine/Shaders/Private/DiaphragmDOF/DOFRecombine.usf
@@ -131,6 +131,9 @@
 
 //------------------------------------------------------- PARAMETERS
 
+float4 ViewportSize;
+uint4 ViewportRect;
+
 float2 TemporalJitterPixels;
 float4 CocModelParameters;
 float2 DepthBlurParameters;
@@ -158,30 +161,6 @@ Texture2D BackgroundConvolution_SceneColor;
 Texture2D BackgroundConvolution_SeparateAlpha;
 
 
-#if !defined(NEW_CODE) // TODO: remove me
-
-#define SceneColorInput PostprocessInput0
-#define SceneSeparateCoc PostprocessInput1
-#define SceneSeparateTranslucency PostprocessInput2
-
-#define BokehLUT PostprocessInput10
-
-#define ForegroundConvolution_SceneColor PostprocessInput3
-#define ForegroundConvolution_SeparateAlpha PostprocessInput4
-
-#define ForegroundHoleFillingConvolution_SceneColor PostprocessInput5
-#define ForegroundHoleFillingConvolution_SeparateAlpha PostprocessInput6
-
-#define SlightOutOfFocusConvolution_SceneColor PostprocessInput9
-
-#define BackgroundConvolution_SceneColor PostprocessInput7
-#define BackgroundConvolution_SeparateAlpha PostprocessInput8
-
-#define SceneColorOutput Output0
-
-#endif
-
-
 RWTexture2D<float4>	SceneColorOutput;
 
 
@@ -211,11 +190,7 @@ struct FRecombineInputParameters
 
 float SampleWorldDepth(float2 BufferUV)
 {
-#if defined(NEW_CODE)
 	return ConvertFromDeviceZ(SceneDepthTexture.SampleLevel(GlobalPointClampedSampler, BufferUV, 0).x);
-#else
-	return CalcSceneDepth(BufferUV);
-#endif
 }
 
 // {0 to 1} output.
@@ -714,11 +689,7 @@ void RecombineMainCS(
 	// Setup input parameters.
 	FRecombineInputParameters InputParameters;
 	{
-#if defined(NEW_CODE) // TODO: remove me
 		InputParameters.DOFBufferSize = ConvolutionInputSize;
-#else
-		InputParameters.DOFBufferSize = (DIM_LAYER_PROCESSING == LAYER_PROCESSING_BACKGROUND_ONLY) ? PostprocessInput7Size : PostprocessInput3Size;
-#endif
 
 		InputParameters.ViewportUV = (DispatchThreadId + 0.5) * ViewportSize.zw;
 		InputParameters.SceneBufferUV = ViewportUVToBufferUV(InputParameters.ViewportUV);
diff --git a/Engine/Shaders/Private/DiaphragmDOF/DOFReduce.usf b/Engine/Shaders/Private/DiaphragmDOF/DOFReduce.usf
--- a/Engine/Shaders/Private/DiaphragmDOF/DOFReduce.usf
+++ b/Engine/Shaders/Private/DiaphragmDOF/DOFReduce.usf
@@ -59,6 +59,7 @@
 
 //------------------------------------------------------- PARAMETERS
 
+uint4 ViewportRect;
 uint MaxScatteringGroupCount;
 float2 MaxInputBufferUV;
 float MinScatteringCocRadius;
@@ -71,41 +72,20 @@ Texture2D GatherInput_SeparateCoc;
 float4 QuarterResGatherInputSize;
 Texture2D QuarterResGatherInput_SceneColor;
 
-#if defined(NEW_CODE) // TODO: remove me
-#define Output0Mip0 OutputMips_0_SceneColor
-#define Output1Mip0 OutputMips_0_SeparateCoc
-#define Output0Mip1 OutputMips_1_SceneColor
-#define Output1Mip1 OutputMips_1_SeparateCoc
-#define Output0Mip2 OutputMips_2_SceneColor
-#define Output1Mip2 OutputMips_2_SeparateCoc
-#define Output0Mip3 OutputMips_3_SceneColor
-#define Output1Mip3 OutputMips_3_SeparateCoc
-
-#else
-
-#define GatherInputSize PostprocessInput0Size
-#define GatherInput_SceneColor PostprocessInput0
-#define GatherInput_SeparateCoc PostprocessInput1
-
-#define QuarterResGatherInputSize PostprocessInput2Size
-#define QuarterResGatherInput_SceneColor PostprocessInput2
-
-#endif
-
 
 //------------------------------------------------------- OUTPUTS
 
-RWTexture2D<float4>	Output0Mip0;
-RWTexture2D<float4>	Output0Mip1;
-RWTexture2D<float4>	Output0Mip2;
-RWTexture2D<float4>	Output0Mip3;
+RWTexture2D<float4>	OutputMips_0_SceneColor;
+RWTexture2D<float4>	OutputMips_1_SceneColor;
+RWTexture2D<float4>	OutputMips_2_SceneColor;
+RWTexture2D<float4>	OutputMips_3_SceneColor;
 
 #if CONFIG_GATHER_INPUT_LAYOUT == GATHER_INPUT_LAYOUT_RGB_ALPHA_COC || CONFIG_GATHER_INPUT_LAYOUT == GATHER_INPUT_LAYOUT_RGB_SEPARATE_COC
 
-RWTexture2D<float>	Output1Mip0;
-RWTexture2D<float>	Output1Mip1;
-RWTexture2D<float>	Output1Mip2;
-RWTexture2D<float>	Output1Mip3;
+RWTexture2D<float>	OutputMips_0_SeparateCoc;
+RWTexture2D<float>	OutputMips_1_SeparateCoc;
+RWTexture2D<float>	OutputMips_2_SeparateCoc;
+RWTexture2D<float>	OutputMips_3_SeparateCoc;
 
 #endif
 
@@ -227,13 +207,13 @@ void ReduceOperator(FCocDownsampleParams DownsampleParams, uint OutId, uint Redu
 void OutputMipLevel(const uint MipLevel, uint2 MipPixelPos, float4 Color, float CocRadius)
 {
 	#if CONFIG_GATHER_INPUT_LAYOUT == GATHER_INPUT_LAYOUT_RGB_COC
-		float4 OutputValue0 = float4(Color.rgb, CocRadius);
+		float4 OutputSceneColor = float4(Color.rgb, CocRadius);
 	#elif CONFIG_GATHER_INPUT_LAYOUT == GATHER_INPUT_LAYOUT_RGB_SEPARATE_COC
-		float4 OutputValue0 = float4(Color.rgb, 0);
-		float OutputValue1 = CocRadius;
+		float4 OutputSceneColor = float4(Color.rgb, 0);
+		float OutputSeparateCoc = CocRadius;
 	#else
-		float4 OutputValue0 = Color;
-		float OutputValue1 = CocRadius;
+		float4 OutputSceneColor = Color;
+		float OutputSeparateCoc = CocRadius;
 	#endif
 	
 	// This is hugly, but compile time.
@@ -241,29 +221,23 @@ void OutputMipLevel(const uint MipLevel, uint2 MipPixelPos, float4 Color, float
 	{
 		if (MipLevel == 0)
 		{
-			Output0Mip0[MipPixelPos] = OutputValue0;
+			OutputMips_0_SceneColor[MipPixelPos] = OutputSceneColor;
 		}
 #if DIM_REDUCE_MIP_COUNT > 1
 		else if (MipLevel == 1)
 		{
-			Output0Mip1[MipPixelPos] = OutputValue0;
+			OutputMips_1_SceneColor[MipPixelPos] = OutputSceneColor;
 		}
 #if DIM_REDUCE_MIP_COUNT > 2
 		else if (MipLevel == 2)
 		{
-			Output0Mip2[MipPixelPos] = OutputValue0;
+			OutputMips_2_SceneColor[MipPixelPos] = OutputSceneColor;
 		}
 #if DIM_REDUCE_MIP_COUNT > 3
 		else if (MipLevel == 3)
 		{
-			Output0Mip3[MipPixelPos] = OutputValue0;
+			OutputMips_3_SceneColor[MipPixelPos] = OutputSceneColor;
 		}
-#if DIM_REDUCE_MIP_COUNT > 4
-		else if (MipLevel == 4)
-		{
-			Output0Mip4[MipPixelPos] = OutputValue0;
-		}
-#endif
 #endif
 #endif
 #endif
@@ -272,34 +246,27 @@ void OutputMipLevel(const uint MipLevel, uint2 MipPixelPos, float4 Color, float
 	{
 		if (MipLevel == 0)
 		{
-			Output0Mip0[MipPixelPos] = OutputValue0;
-			Output1Mip0[MipPixelPos] = OutputValue1;
+			OutputMips_0_SceneColor[MipPixelPos] = OutputSceneColor;
+			OutputMips_0_SeparateCoc[MipPixelPos] = OutputSeparateCoc;
 		}
 #if DIM_REDUCE_MIP_COUNT > 1
 		else if (MipLevel == 1)
 		{
-			Output0Mip1[MipPixelPos] = OutputValue0;
-			Output1Mip1[MipPixelPos] = OutputValue1;
+			OutputMips_1_SceneColor[MipPixelPos] = OutputSceneColor;
+			OutputMips_1_SeparateCoc[MipPixelPos] = OutputSeparateCoc;
 		}
 #if DIM_REDUCE_MIP_COUNT > 2
 		else if (MipLevel == 2)
 		{
-			Output0Mip2[MipPixelPos] = OutputValue0;
-			Output1Mip2[MipPixelPos] = OutputValue1;
+			OutputMips_2_SceneColor[MipPixelPos] = OutputSceneColor;
+			OutputMips_2_SeparateCoc[MipPixelPos] = OutputSeparateCoc;
 		}
 #if DIM_REDUCE_MIP_COUNT > 3
 		else if (MipLevel == 3)
 		{
-			Output0Mip3[MipPixelPos] = OutputValue0;
-			Output1Mip3[MipPixelPos] = OutputValue1;
+			OutputMips_3_SceneColor[MipPixelPos] = OutputSceneColor;
+			OutputMips_3_SeparateCoc[MipPixelPos] = OutputSeparateCoc;
 		}
-#if DIM_REDUCE_MIP_COUNT > 4
-		else if (MipLevel == 4)
-		{
-			Output0Mip4[MipPixelPos] = OutputValue0;
-			Output1Mip4[MipPixelPos] = OutputValue1;
-		}
-#endif
 #endif
 #endif
 #endif
diff --git a/Engine/Shaders/Private/DiaphragmDOF/DOFSetup.usf b/Engine/Shaders/Private/DiaphragmDOF/DOFSetup.usf
--- a/Engine/Shaders/Private/DiaphragmDOF/DOFSetup.usf
+++ b/Engine/Shaders/Private/DiaphragmDOF/DOFSetup.usf
@@ -20,6 +20,8 @@
 
 //------------------------------------------------------- PARAMETERS
 
+uint4 ViewportRect;
+
 float4 CocModelParameters;
 float2 DepthBlurParameters;
 
@@ -29,11 +31,6 @@ float2 CocRadiusBasis;
 Texture2D SceneColorTexture;
 Texture2D SceneDepthTexture;
 
-#if defined(NEW_CODE) // TODO: remove me
-#define PostprocessInput0 SceneColorTexture
-#define PostprocessInput1 SceneDepthTexture
-#endif
-
 
 //------------------------------------------------------- OUTPUTS
 
@@ -84,8 +81,8 @@ void SetupCS(
 		SceneBufferUV = clamp(SceneBufferUV, View.BufferBilinearUVMinMax.xy, View.BufferBilinearUVMinMax.zw);
 	}
 
-	float4 SceneColor = PostprocessInput0.SampleLevel(GlobalPointClampedSampler, SceneBufferUV, 0);
-	float DeviceZ = PostprocessInput1.SampleLevel(GlobalPointClampedSampler, SceneBufferUV, 0).r;
+	float4 SceneColor = SceneColorTexture.SampleLevel(GlobalPointClampedSampler, SceneBufferUV, 0);
+	float DeviceZ = SceneDepthTexture.SampleLevel(GlobalPointClampedSampler, SceneBufferUV, 0).r;
 	float CocRadius = DepthToHalfResCocRadius(CocModelParameters, DepthBlurParameters, ConvertFromDeviceZ(DeviceZ));
 	
 	// Convert from translucency to opacity.
@@ -184,9 +181,9 @@ void SetupCS(uint2 DispatchThreadId : SV_DispatchThreadID)
 {
 	float2 SceneBufferUV = (float2(View.ViewRectMin.xy) + (DispatchThreadId + 0.5)) * View.BufferSizeAndInvSize.zw;
 	
-	float4 OutColor = PostprocessInput0.SampleLevel(GlobalPointClampedSampler, SceneBufferUV, 0);
+	float4 OutColor = SceneColorTexture.SampleLevel(GlobalPointClampedSampler, SceneBufferUV, 0);
 	
-	float DeviceZ = PostprocessInput1.SampleLevel(GlobalPointClampedSampler, SceneBufferUV, 0).r;
+	float DeviceZ = SceneDepthTexture.SampleLevel(GlobalPointClampedSampler, SceneBufferUV, 0).r;
 	float OutCocRadius = CocRadiusBasis.x * DepthToHalfResCocRadius(CocModelParameters, DepthBlurParameters, ConvertFromDeviceZ(DeviceZ));
 
 	if (all(DispatchThreadId < ViewportRect.zw))
@@ -225,7 +222,7 @@ void SetupCS(
 	if (all(GroupId * GROUP_BORDER_SIZE + (GROUP_BORDER_SIZE + 1) <= ViewportRect.zw))
 	{
 		// Fetch Coc
-		float4 SceneDepthGather = PostprocessInput1.Gather(GlobalPointClampedSampler, SceneBufferUV);
+		float4 SceneDepthGather = SceneDepthTexture.Gather(GlobalPointClampedSampler, SceneBufferUV);
 		float CocRadii[4];
 		CocRadii[0] = DepthToHalfResCocRadius(CocModelParameters, DepthBlurParameters, ConvertFromDeviceZ(SceneDepthGather.w));
 		CocRadii[1] = DepthToHalfResCocRadius(CocModelParameters, DepthBlurParameters, ConvertFromDeviceZ(SceneDepthGather.z));
@@ -237,7 +234,7 @@ void SetupCS(
 		for (uint i = 0; i < 4; i++)
 		{
 			float2 SampleUV = SceneBufferUV + SceneBufferUVOffset * float2(kOffsetsCross3x3[i]);
-			SceneColors[i] = PostprocessInput0.SampleLevel(GlobalPointClampedSampler, SampleUV, 0);
+			SceneColors[i] = SceneColorTexture.SampleLevel(GlobalPointClampedSampler, SampleUV, 0);
 
 			// Convert from translucency to opacity.
 			SceneColors[i].a = 1 - SceneColors[i].a;
@@ -273,12 +270,12 @@ void SetupCS(
 				SampleUV = clamp(SampleUV, View.BufferBilinearUVMinMax.xy, View.BufferBilinearUVMinMax.zw);
 			}
 
-			SceneColors[i] = PostprocessInput0.SampleLevel(GlobalPointClampedSampler, SampleUV, 0);
+			SceneColors[i] = SceneColorTexture.SampleLevel(GlobalPointClampedSampler, SampleUV, 0);
 			
 			// Convert from translucency to opacity.
 			SceneColors[i].a = 1 - SceneColors[i].a;
 
-			float SampleDeviceZ = PostprocessInput1.SampleLevel(GlobalPointClampedSampler, SampleUV, 0).r;
+			float SampleDeviceZ = SceneDepthTexture.SampleLevel(GlobalPointClampedSampler, SampleUV, 0).r;
 			CocRadii[i] = DepthToHalfResCocRadius(CocModelParameters, DepthBlurParameters, ConvertFromDeviceZ(SampleDeviceZ));
 		}
 
diff --git a/Engine/Source/Runtime/Renderer/Private/PostProcess/DiaphragmDOF.cpp b/Engine/Source/Runtime/Renderer/Private/PostProcess/DiaphragmDOF.cpp
--- a/Engine/Source/Runtime/Renderer/Private/PostProcess/DiaphragmDOF.cpp
+++ b/Engine/Source/Runtime/Renderer/Private/PostProcess/DiaphragmDOF.cpp
@@ -420,7 +420,6 @@ public:
 	static void ModifyCompilationEnvironment(const FGlobalShaderPermutationParameters& Parameters, FShaderCompilerEnvironment& OutEnvironment)
 	{
 		OutEnvironment.SetDefine(TEXT("COC_TILE_SIZE"), kCocTileSize);
-		OutEnvironment.SetDefine(TEXT("NEW_CODE"), kCocTileSize); // TODO: kill that
 	}
 
 	FDiaphragmDOFShader() {}
