Fix CameraVector not being calculated correctly in ray tracing hit shaders

diff --git a/Engine/Shaders/Private/MaterialTemplate.ush b/Engine/Shaders/Private/MaterialTemplate.ush
--- a/Engine/Shaders/Private/MaterialTemplate.ush
+++ b/Engine/Shaders/Private/MaterialTemplate.ush
@@ -2140,8 +2140,12 @@ void CalcMaterialParametersEx(
 	// Originally this was normalize(normalize(TranslatedWorldPosition.xyz)) but tegra4 appears to optimize that out.
 	Parameters.CameraVector = normalize(-0.01 * Parameters.WorldPosition_CamRelative.xyz);
 #else
-	// TranslatedWorldPosition is the world position translated to the camera position, which is just -CameraVector
-	Parameters.CameraVector = normalize(-Parameters.WorldPosition_CamRelative.xyz);
+	#if !RAYHITGROUPSHADER
+		// TranslatedWorldPosition is the world position translated to the camera position, which is just -CameraVector
+		Parameters.CameraVector = normalize(-Parameters.WorldPosition_CamRelative.xyz);
+	#else
+		Parameters.CameraVector = -WorldRayDirection();
+	#endif
 #endif
 
 	Parameters.LightVector = 0;
diff --git a/Engine/Shaders/Private/RayTracing/RayTracingMaterialHitShaders.usf b/Engine/Shaders/Private/RayTracing/RayTracingMaterialHitShaders.usf
--- a/Engine/Shaders/Private/RayTracing/RayTracingMaterialHitShaders.usf
+++ b/Engine/Shaders/Private/RayTracing/RayTracingMaterialHitShaders.usf
@@ -14,6 +14,7 @@
 #define clip(x) (x)
 #define ddx(x) 0
 #define ddy(x) 0
+#define fwidth(x) 0
 
 #include "/Engine/Private/Common.ush"
 #include "RayTracingCommon.ush"
@@ -192,15 +193,15 @@ void CalcInterpolants(in FRayCone RayCone, in FDefaultAttributes Attributes, out
 		FVertexFactoryIntermediates VFIntermediates = GetVertexFactoryIntermediates(Input);
 		float3x3 TangentToLocal = VertexFactoryGetTangentToLocal(Input, VFIntermediates);	
 		float4 WorldPositionExcludingWPO = VertexFactoryGetWorldPosition(Input, VFIntermediates);
-		FMaterialVertexParameters VertexParameters = GetMaterialVertexParameters(Input, VFIntermediates, WorldPositionExcludingWPO, TangentToLocal);
+		FMaterialVertexParameters VertexParameters = GetMaterialVertexParameters(Input, VFIntermediates, WorldPositionExcludingWPO.xyz, TangentToLocal);
 		
 		WorldPositions[i] = WorldPositionExcludingWPO.xyz;
 		float4 WorldPosition = WorldPositionExcludingWPO + float4(GetMaterialWorldPositionOffset(VertexParameters), 0.0f);
 		
 		PerVertexInterpolants[i] = VertexFactoryGetInterpolantsVSToDS(Input, VFIntermediates, VertexParameters);
 
-#if NUM_TEX_COORD_INTERPOLATORS || USE_PARTICLE_SUBUVS
-		TexCoords[i] = PerVertexInterpolants[i].InterpolantsVSToPS.TexCoords[0].xy;
+#if (NUM_MATERIAL_TEXCOORDS || USE_PARTICLE_SUBUVS) && !VERTEX_FACTORY_MODIFIES_TESSELLATION
+		TexCoords[i] = VertexParameters.TexCoords[0].xy;
 #endif
 
 		float4 RasterizedWorldPosition = VertexFactoryGetRasterizedWorldPosition(Input, VFIntermediates, WorldPosition);
@@ -219,8 +220,8 @@ void CalcInterpolants(in FRayCone RayCone, in FDefaultAttributes Attributes, out
 	
 	float4 ClipSpacePosition = PerVertexClipSpacePositions[0] * (1 - Attributes.Barycentrics.x - Attributes.Barycentrics.y) + PerVertexClipSpacePositions[1] * Attributes.Barycentrics.x + PerVertexClipSpacePositions[2] * Attributes.Barycentrics.y;
 
-#if NUM_TEX_COORD_INTERPOLATORS || USE_PARTICLE_SUBUVS
-	float2 TA = TexCoords[1] - TexCoords[0];
+#if (NUM_MATERIAL_TEXCOORDS || USE_PARTICLE_SUBUVS) && !VERTEX_FACTORY_MODIFIES_TESSELLATION
+	float2 TA = TexCoords[1] - TexCoords[0]; 
 	float2 TB = TexCoords[2] - TexCoords[0];
     float UvArea = abs((TA.x * TB.y) - (TA.y * TB.x));
 
@@ -259,8 +260,6 @@ void MaterialCHS(inout FPackedMaterialClosestHitPayload PackedPayload, in FDefau
 		
 		bool bIsFrontFace = HitKind() == HIT_KIND_TRIANGLE_FRONT_FACE;
 
-		MaterialParameters.CameraVector = -WorldRayDirection();
-
 		// #dxr_todo: UE-72130 support world position offset 
 		// #if USE_WORLD_POSITION_EXCLUDING_SHADER_OFFSETS
 		// CalcMaterialParametersEx(MaterialParameters, PixelMaterialInputs, SvPosition, ScreenPosition, bIsFrontFace, TranslatedWorldPosition, BasePassInterpolants.PixelPositionExcludingWPO);
