Fix TexCoords in ray tracing hit shaders

diff --git a/Engine/Shaders/Private/RayTracing/RayTracingMaterialHitShaders.usf b/Engine/Shaders/Private/RayTracing/RayTracingMaterialHitShaders.usf
--- a/Engine/Shaders/Private/RayTracing/RayTracingMaterialHitShaders.usf
+++ b/Engine/Shaders/Private/RayTracing/RayTracingMaterialHitShaders.usf
@@ -200,8 +200,8 @@ void CalcInterpolants(in FRayCone RayCone, in FDefaultAttributes Attributes, out
 		
 		PerVertexInterpolants[i] = VertexFactoryGetInterpolantsVSToDS(Input, VFIntermediates, VertexParameters);
 
-#if (NUM_MATERIAL_TEXCOORDS || USE_PARTICLE_SUBUVS) && !VERTEX_FACTORY_MODIFIES_TESSELLATION
-		TexCoords[i] = VertexParameters.TexCoords[0].xy;
+#if (NUM_TEX_COORD_INTERPOLATORS || USE_PARTICLE_SUBUVS) && !VERTEX_FACTORY_MODIFIES_TESSELLATION
+		TexCoords[i] = PerVertexInterpolants[i].InterpolantsVSToPS.TexCoords[0].xy;
 #endif
 
 		float4 RasterizedWorldPosition = VertexFactoryGetRasterizedWorldPosition(Input, VFIntermediates, WorldPosition);
@@ -220,7 +220,7 @@ void CalcInterpolants(in FRayCone RayCone, in FDefaultAttributes Attributes, out
 	
 	float4 ClipSpacePosition = PerVertexClipSpacePositions[0] * (1 - Attributes.Barycentrics.x - Attributes.Barycentrics.y) + PerVertexClipSpacePositions[1] * Attributes.Barycentrics.x + PerVertexClipSpacePositions[2] * Attributes.Barycentrics.y;
 
-#if (NUM_MATERIAL_TEXCOORDS || USE_PARTICLE_SUBUVS) && !VERTEX_FACTORY_MODIFIES_TESSELLATION
+#if (NUM_TEX_COORD_INTERPOLATORS || USE_PARTICLE_SUBUVS) && !VERTEX_FACTORY_MODIFIES_TESSELLATION
 	float2 TA = TexCoords[1] - TexCoords[0]; 
 	float2 TB = TexCoords[2] - TexCoords[0];
     float UvArea = abs((TA.x * TB.y) - (TA.y * TB.x));
