Abstract some platform specific details used by VT shaders

diff --git a/Engine/Shaders/Private/PixelShaderOutputCommon.ush b/Engine/Shaders/Private/PixelShaderOutputCommon.ush
--- a/Engine/Shaders/Private/PixelShaderOutputCommon.ush
+++ b/Engine/Shaders/Private/PixelShaderOutputCommon.ush
@@ -49,14 +49,8 @@ PixelShaderOutputCommon.usf: To allow PS input/output passed int functions throu
 
 #if LIGHTMAP_VT_ENABLED || NUM_VIRTUALTEXTURE_SAMPLES
 	// If virtual texture is enabled need to make sure we're using early depth test, in order for UAV feedback buffer writes to respect depth test
-	#if PS4_PROFILE
-		// PS4 HW has both early depth test and RE_Z mode
-		// On PS4, early depth stencil test also implies early depth write, which is not compatible with pixel shaders that discard or modify output depth
-		// To handle these cases, RE_Z mode performs depth test once before pixel shader (to handle early rejection) and again after pixel shader along with depth write
-		ENABLE_RE_Z
-	#else
-		EARLYDEPTHSTENCIL
-	#endif
+	// Also need to wait until after pixel shader to write depth, since we may be using discard, or modifying depth
+	DEPTHSTENCIL_EARLYTEST_LATEWRITE
 #endif // LIGHTMAP_VT_ENABLED || NUM_VIRTUALTEXTURE_SAMPLES
 void MainPS
 	(
diff --git a/Engine/Shaders/Private/VirtualTextureCommon.ush b/Engine/Shaders/Private/VirtualTextureCommon.ush
--- a/Engine/Shaders/Private/VirtualTextureCommon.ush
+++ b/Engine/Shaders/Private/VirtualTextureCommon.ush
@@ -419,13 +419,6 @@ VTUniform VTUniform_Unpack(uint4 PackedUniform)
 	return result;
 }
 
-#if !PS4_PROFILE
-float UnpackByte0(uint v) { return float(v & 0xff); }
-float UnpackByte1(uint v) { return float((v >> 8) & 0xff); }
-float UnpackByte2(uint v) { return float((v >> 16) & 0xff); }
-float UnpackByte3(uint v) { return float(v >> 24); }
-#endif
-
 // Applies proper scaling to dUVdx/dUVdy in PageTableResult
 float2 VTComputePhysicalUVs(in out VTPageTableResult PageTableResult, uint LayerIndex, VTUniform Uniform)
 {
diff --git a/Engine/Shaders/Public/Platform.ush b/Engine/Shaders/Public/Platform.ush
--- a/Engine/Shaders/Public/Platform.ush
+++ b/Engine/Shaders/Public/Platform.ush
@@ -462,6 +462,9 @@
 #define EARLYDEPTHSTENCIL [earlydepthstencil]
 #endif
 
+#ifndef DEPTHSTENCIL_EARLYTEST_LATEWRITE
+#define DEPTHSTENCIL_EARLYTEST_LATEWRITE EARLYDEPTHSTENCIL
+#endif
 
 // ---------------------------------------------------- Compiler attribute fallbacks
 
@@ -648,6 +651,12 @@ float4 max3( float4 a, float4 b, float4 c )
 
 #endif
 
+#if !defined(COMPILER_SUPPORTS_UNPACKBYTEN)
+float UnpackByte0(uint v) { return float(v & 0xff); }
+float UnpackByte1(uint v) { return float((v >> 8) & 0xff); }
+float UnpackByte2(uint v) { return float((v >> 16) & 0xff); }
+float UnpackByte3(uint v) { return float(v >> 24); }
+#endif // !COMPILER_SUPPORTS_UNPACKBYTEN
 
 #if COMPILER_HLSLCC
 #define ddx_fine(x) ddx(x)
