Fixes PDF of CosineSampleHemisphere()

diff --git a/Engine/Shaders/Private/MonteCarlo.ush b/Engine/Shaders/Private/MonteCarlo.ush
--- a/Engine/Shaders/Private/MonteCarlo.ush
+++ b/Engine/Shaders/Private/MonteCarlo.ush
@@ -175,7 +175,7 @@ float4 CosineSampleHemisphere( float2 E, float3 N )
 	float3 H = UniformSampleSphere( E ).xyz;
 	H = normalize( N + H );
 
-	float PDF = H.z * (1.0 /  PI);
+	float PDF = dot(H, N) * (1.0 /  PI);
 
 	return float4( H, PDF );
 }
@@ -245,6 +245,7 @@ float4 ImportanceSampleVisibleGGX( float2 E, float a2, float3 V )
 	float3 Vh = normalize( float3( a * V.xy, V.z ) );
 
 	// Orthonormal basis
+	// Tangent0 is orthogonal to N.
 	float3 Tangent0 = (Vh.z < 0.9999) ? normalize( cross( float3(0, 0, 1), Vh ) ) : float3(1, 0, 0);
 	float3 Tangent1 = cross( Vh, Tangent0 );
 
