DR - Android shader compile fix
#rnx

diff --git a/Engine/Shaders/Private/Common.ush b/Engine/Shaders/Private/Common.ush
--- a/Engine/Shaders/Private/Common.ush
+++ b/Engine/Shaders/Private/Common.ush
@@ -72,18 +72,18 @@ const static float MaxHalfFloat = 65504.0f;
 #define USE_RAYTRACED_TEXTURE_RAYCONE_LOD (RAYHITGROUPSHADER)
 #endif
 
+#if USE_RAYTRACED_TEXTURE_RAYCONE_LOD
 static float GlobalRayCone_TexArea;
 float ComputeRayConeLod(Texture2D Tex)
 {
-#if USE_RAYTRACED_TEXTURE_RAYCONE_LOD
 	uint2 Dimensions;
 	Tex.GetDimensions(Dimensions.x, Dimensions.y);
-    int TexArea = Dimensions.x * Dimensions.y;
-    return 0.5f * log2(GlobalRayCone_TexArea * TexArea);
+	int TexArea = Dimensions.x * Dimensions.y;
+	return 0.5f * log2(GlobalRayCone_TexArea * TexArea);
+}
 #else
-    return 0;
+#define ComputeRayConeLod(Tex)		0
 #endif
-}
 
 float  ClampToHalfFloatRange(float  X) { return clamp(X, float(0), MaxHalfFloat); }
 float2 ClampToHalfFloatRange(float2 X) { return clamp(X, float(0).xx, MaxHalfFloat.xx); }
