DR - Proper fix for android
#rnx

diff --git a/Engine/Shaders/Private/Common.ush b/Engine/Shaders/Private/Common.ush
--- a/Engine/Shaders/Private/Common.ush
+++ b/Engine/Shaders/Private/Common.ush
@@ -38,7 +38,9 @@
 #define SetSceneColorLayoutToFloat4(dest,value) dest.rgb = (value).rgb
 #endif
 
-#if (!COMPILER_GLSL_ES2 && !COMPILER_GLSL_ES3_1 && !COMPILER_VULKAN)
+#define SUPPORTS_TEXTURE_EXTERNAL		(COMPILER_GLSL_ES2 || COMPILER_GLSL_ES3_1 || COMPILER_VULKAN)
+
+#if !SUPPORTS_TEXTURE_EXTERNAL
 #define TextureExternal Texture2D
 #endif
 
@@ -72,18 +74,26 @@ const static float MaxHalfFloat = 65504.0f;
 #define USE_RAYTRACED_TEXTURE_RAYCONE_LOD (RAYHITGROUPSHADER)
 #endif
 
-#if USE_RAYTRACED_TEXTURE_RAYCONE_LOD
 static float GlobalRayCone_TexArea;
 float ComputeRayConeLod(Texture2D Tex)
 {
+#if USE_RAYTRACED_TEXTURE_RAYCONE_LOD
 	uint2 Dimensions;
 	Tex.GetDimensions(Dimensions.x, Dimensions.y);
 	int TexArea = Dimensions.x * Dimensions.y;
 	return 0.5f * log2(GlobalRayCone_TexArea * TexArea);
-}
 #else
-#define ComputeRayConeLod(Tex)		0
+    return 0;
 #endif
+}
+
+#if SUPPORTS_TEXTURE_EXTERNAL
+float ComputeRayConeLod(TextureExternal Tex)
+{
+	return 0;
+}
+#endif
+
 
 float  ClampToHalfFloatRange(float  X) { return clamp(X, float(0), MaxHalfFloat); }
 float2 ClampToHalfFloatRange(float2 X) { return clamp(X, float(0).xx, MaxHalfFloat.xx); }
