PR #5875: Increase Material Bins (Contributed by [REDACTED_PERSON])

#github https://github.com/EpicGames/UnrealEngine/pull/5875

diff --git a/Engine/Shaders/Private/RayTracing/MaterialSort.usf b/Engine/Shaders/Private/RayTracing/MaterialSort.usf
--- a/Engine/Shaders/Private/RayTracing/MaterialSort.usf
+++ b/Engine/Shaders/Private/RayTracing/MaterialSort.usf
@@ -21,7 +21,7 @@
 	#error Unknown sorting size
 #endif
 
-#define NUM_BINS  64
+#define NUM_BINS  128
 
 groupshared uint Bins[NUM_BINS];
 groupshared uint Offsets[NUM_BINS];
