PR #5873: Use structured buffer for ray tracing light constants (Contributed by [REDACTED_PERSON])

#github https://github.com/EpicGames/UnrealEngine/pull/5873

diff --git a/Engine/Shaders/Private/RayTracing/RayTracingLightingCommon.ush b/Engine/Shaders/Private/RayTracing/RayTracingLightingCommon.ush
--- a/Engine/Shaders/Private/RayTracing/RayTracingLightingCommon.ush
+++ b/Engine/Shaders/Private/RayTracing/RayTracingLightingCommon.ush
@@ -7,6 +7,31 @@
 #define LIGHT_TYPE_RECT				3 
 #define LIGHT_TYPE_MAX				4 
 
+//Must match definition in RayTracingLighting.h
+struct FRTLightingData
+{
+	uint Type;
+	uint LightProfileIndex;
+	uint RectLightTextureIndex;
+	uint Pad;
+	float3 LightPosition;
+	float InvRadius;
+	float3 Direction;
+	float FalloffExponent;
+	float3 LightColor;
+	float SpecularScale;
+	float3 Tangent;
+	float SourceRadius;
+	float2 SpotAngles;
+	float SourceLength;
+	float SoftSourceRadius;
+	float2 DistanceFadeMAD;
+	float RectLightBarnCosAngle;
+	float RectLightBarnLength;
+	float4 Dummy;
+};
+
+StructuredBuffer<FRTLightingData> LightDataBuffer;
 
 float GetRoughnessFade(float Roughness, float MaxRoughness)
 {
@@ -200,24 +225,25 @@ uint2 CullDirectLighting(
 	{
 		FDeferredLightData LightData = (FDeferredLightData)0;
 		uint Lit = 1;
-
-		uint LightType = RaytracingLightsDataPacked.Type_LightProfileIndex_RectLightTextureIndex[DirIndex].x;
-		LightData.Position = RaytracingLightsDataPacked.LightPosition_InvRadius[DirIndex].xyz;
-		LightData.InvRadius = RaytracingLightsDataPacked.LightPosition_InvRadius[DirIndex].w;
-		LightData.Color = RaytracingLightsDataPacked.LightColor_SpecularScale[DirIndex].xyz;
-		LightData.FalloffExponent = RaytracingLightsDataPacked.Direction_FalloffExponent[DirIndex].w;
-		LightData.Direction = RaytracingLightsDataPacked.Direction_FalloffExponent[DirIndex].xyz;
-		LightData.Tangent = RaytracingLightsDataPacked.Tangent_SourceRadius[DirIndex].xyz;
-		LightData.SpotAngles = RaytracingLightsDataPacked.SpotAngles_SourceLength_SoftSourceRadius[DirIndex].xy;
-		LightData.SourceRadius = RaytracingLightsDataPacked.Tangent_SourceRadius[DirIndex].w;
-		LightData.SourceLength = RaytracingLightsDataPacked.SpotAngles_SourceLength_SoftSourceRadius[DirIndex].z;
-		LightData.SoftSourceRadius = RaytracingLightsDataPacked.SpotAngles_SourceLength_SoftSourceRadius[DirIndex].w;
-		LightData.SpecularScale = RaytracingLightsDataPacked.LightColor_SpecularScale[DirIndex].w;
-		LightData.RectLightBarnCosAngle = RaytracingLightsDataPacked.DistanceFadeMAD_RectLightBarnCosAngle_RectLightBarnLength[DirIndex].z;
-		LightData.RectLightBarnLength = RaytracingLightsDataPacked.DistanceFadeMAD_RectLightBarnCosAngle_RectLightBarnLength[DirIndex].w;
+		FRTLightingData LightingData = LightDataBuffer[DirIndex];
+
+		uint LightType = LightingData.Type;
+		LightData.Position = LightingData.LightPosition;
+		LightData.InvRadius = LightingData.InvRadius;
+		LightData.Color = LightingData.LightColor;
+		LightData.FalloffExponent = LightingData.FalloffExponent;
+		LightData.Direction = LightingData.Direction;
+		LightData.Tangent = LightingData.Tangent;
+		LightData.SpotAngles = LightingData.SpotAngles;
+		LightData.SourceRadius = LightingData.SourceRadius;
+		LightData.SourceLength = LightingData.SourceLength;
+		LightData.SoftSourceRadius = LightingData.SoftSourceRadius;
+		LightData.SpecularScale = LightingData.SpecularScale;
+		LightData.RectLightBarnCosAngle = LightingData.RectLightBarnCosAngle;
+		LightData.RectLightBarnLength = LightingData.RectLightBarnLength;
 		LightData.ContactShadowLength = 0.0;
-		LightData.DistanceFadeMAD = RaytracingLightsDataPacked.DistanceFadeMAD_RectLightBarnCosAngle_RectLightBarnLength[DirIndex].xy;
-		const uint LightProfileIndex = RaytracingLightsDataPacked.Type_LightProfileIndex_RectLightTextureIndex[DirIndex].y;
+		LightData.DistanceFadeMAD = LightingData.DistanceFadeMAD;
+		const uint LightProfileIndex = LightingData.LightProfileIndex;
 
 
 		LightData.ShadowMapChannelMask = float4(0, 0, 0, 0); 
@@ -355,29 +381,32 @@ float3 ComputeDirectLightingCulled(
 		}
 
 		FDeferredLightData LightData = (FDeferredLightData)0;
-		uint LightType = RaytracingLightsDataPacked.Type_LightProfileIndex_RectLightTextureIndex[DirIndex].x;
-		LightData.Position = RaytracingLightsDataPacked.LightPosition_InvRadius[DirIndex].xyz;
-		LightData.InvRadius = RaytracingLightsDataPacked.LightPosition_InvRadius[DirIndex].w;
-		LightData.Color = RaytracingLightsDataPacked.LightColor_SpecularScale[DirIndex].xyz;
-		LightData.FalloffExponent = RaytracingLightsDataPacked.Direction_FalloffExponent[DirIndex].w;
-		LightData.Direction = RaytracingLightsDataPacked.Direction_FalloffExponent[DirIndex].xyz;
-		LightData.Tangent = RaytracingLightsDataPacked.Tangent_SourceRadius[DirIndex].xyz;
-		LightData.SpotAngles = RaytracingLightsDataPacked.SpotAngles_SourceLength_SoftSourceRadius[DirIndex].xy;
-		LightData.SourceRadius = RaytracingLightsDataPacked.Tangent_SourceRadius[DirIndex].w;
-		LightData.SourceLength = RaytracingLightsDataPacked.SpotAngles_SourceLength_SoftSourceRadius[DirIndex].z;
-		LightData.SoftSourceRadius = RaytracingLightsDataPacked.SpotAngles_SourceLength_SoftSourceRadius[DirIndex].w;
-		LightData.SpecularScale = RaytracingLightsDataPacked.LightColor_SpecularScale[DirIndex].w;
-		LightData.RectLightBarnCosAngle = RaytracingLightsDataPacked.DistanceFadeMAD_RectLightBarnCosAngle_RectLightBarnLength[DirIndex].z;
-		LightData.RectLightBarnLength = RaytracingLightsDataPacked.DistanceFadeMAD_RectLightBarnCosAngle_RectLightBarnLength[DirIndex].w;
+		FRTLightingData LightingData = LightDataBuffer[DirIndex];
+
+		uint LightType = LightingData.Type;
+		LightData.Position = LightingData.LightPosition;
+		LightData.InvRadius = LightingData.InvRadius;
+		LightData.Color = LightingData.LightColor;
+		LightData.FalloffExponent = LightingData.FalloffExponent;
+		LightData.Direction = LightingData.Direction;
+		LightData.Tangent = LightingData.Tangent;
+		LightData.SpotAngles = LightingData.SpotAngles;
+		LightData.SourceRadius = LightingData.SourceRadius;
+		LightData.SourceLength = LightingData.SourceLength;
+		LightData.SoftSourceRadius = LightingData.SoftSourceRadius;
+		LightData.SpecularScale = LightingData.SpecularScale;
+		LightData.RectLightBarnCosAngle = LightingData.RectLightBarnCosAngle;
+		LightData.RectLightBarnLength = LightingData.RectLightBarnLength;
 		LightData.ContactShadowLength = 0.0;
-		LightData.DistanceFadeMAD = RaytracingLightsDataPacked.DistanceFadeMAD_RectLightBarnCosAngle_RectLightBarnLength[DirIndex].xy;
+		LightData.DistanceFadeMAD = LightingData.DistanceFadeMAD;
+
+#if USE_SOURCE_TEXTURE_ARRAY 
+		RectTexture.SourceTextureIndex = LightingData.RectLightTextureIndex;
+#endif
 
-	#if USE_SOURCE_TEXTURE_ARRAY 
-		RectTexture.SourceTextureIndex = RaytracingLightsDataPacked.Type_LightProfileIndex_RectLightTextureIndex[DirIndex].z;
-	#endif
-		uint LightProfileIndex = RaytracingLightsDataPacked.Type_LightProfileIndex_RectLightTextureIndex[DirIndex].y;
+		const uint LightProfileIndex = LightingData.LightProfileIndex;
 
-		LightData.ShadowMapChannelMask = float4(0, 0, 0, 0); 
+		LightData.ShadowMapChannelMask = float4(0, 0, 0, 0);
 		LightData.ShadowedBits = 0; // Not lit dynamic shadows 
 		LightData.ContactShadowLengthInWS = false;
 
diff --git a/Engine/Source/Runtime/Renderer/Private/RayTracing/RayTracingLighting.cpp b/Engine/Source/Runtime/Renderer/Private/RayTracing/RayTracingLighting.cpp
--- a/Engine/Source/Runtime/Renderer/Private/RayTracing/RayTracingLighting.cpp
+++ b/Engine/Source/Runtime/Renderer/Private/RayTracing/RayTracingLighting.cpp
@@ -11,7 +11,8 @@ IMPLEMENT_GLOBAL_SHADER_PARAMETER_STRUCT(FRaytracingLightDataPacked, "Raytracing
 void SetupRaytracingLightDataPacked(
 	const TSparseArray<FLightSceneInfoCompact>& Lights,
 	const FViewInfo& View,
-	FRaytracingLightDataPacked* LightData)
+	FRaytracingLightDataPacked* LightData,
+	TResourceArray<FRTLightingData>& LightDataArray)
 {
 	TMap<UTextureLightProfile*, int32> IESLightProfilesMap;
 	TMap<FTextureRHIParamRef, uint32> RectTextureMap;
@@ -87,29 +88,40 @@ void SetupRaytracingLightDataPacked(
 			}
 		}
 
-		LightData->Type_LightProfileIndex_RectLightTextureIndex[LightData->Count].X = Light.LightType;
-		LightData->Type_LightProfileIndex_RectLightTextureIndex[LightData->Count].Y = IESLightProfileIndex;
-		LightData->Type_LightProfileIndex_RectLightTextureIndex[LightData->Count].Z = InvalidTextureIndex;
+		FRTLightingData LightDataElement;
 
-		LightData->LightPosition_InvRadius[LightData->Count] = LightParameters.Position;
-		LightData->LightPosition_InvRadius[LightData->Count].W = LightParameters.InvRadius;
+		LightDataElement.Type = Light.LightType;
+		LightDataElement.LightProfileIndex = IESLightProfileIndex;
+		LightDataElement.RectLightTextureIndex = InvalidTextureIndex;
 
-		LightData->LightColor_SpecularScale[LightData->Count] = LightParameters.Color;
-		LightData->LightColor_SpecularScale[LightData->Count].W = LightParameters.SpecularScale;
+		for (int32 Element = 0; Element < 3; Element++)
+		{
+			LightDataElement.Direction[Element] = LightParameters.Direction[Element];
+			LightDataElement.LightPosition[Element] = LightParameters.Position[Element];
+			LightDataElement.LightColor[Element] = LightParameters.Color[Element];
+			LightDataElement.Tangent[Element] = LightParameters.Tangent[Element];
+		}
 
-		LightData->Direction_FalloffExponent[LightData->Count] = LightParameters.Direction;
-		LightData->Direction_FalloffExponent[LightData->Count].W = LightParameters.FalloffExponent;
+		const FVector2D FadeParams = Light.LightSceneInfo->Proxy->GetDirectionalLightDistanceFadeParameters(View.GetFeatureLevel(), Light.LightSceneInfo->IsPrecomputedLightingValid(), View.MaxShadowCascades);
+		const FVector2D DistanceFadeMAD = { FadeParams.Y, -FadeParams.X * FadeParams.Y };
 
-		LightData->Tangent_SourceRadius[LightData->Count] = LightParameters.Tangent;
-		LightData->Tangent_SourceRadius[LightData->Count].W = LightParameters.SourceRadius;
+		for (int32 Element = 0; Element < 2; Element++)
+		{
+			LightDataElement.SpotAngles[Element] = LightParameters.SpotAngles[Element];
+			LightDataElement.DistanceFadeMAD[Element] = DistanceFadeMAD[Element];
+		}
 
-		FVector4 SpotAngles_SourceLength_SoftSourceRadius = FVector4(LightParameters.SpotAngles, FVector2D(LightParameters.SourceLength, LightParameters.SoftSourceRadius));
-		LightData->SpotAngles_SourceLength_SoftSourceRadius[LightData->Count] = SpotAngles_SourceLength_SoftSourceRadius;
-		
-		const FVector2D FadeParams = Light.LightSceneInfo->Proxy->GetDirectionalLightDistanceFadeParameters(View.GetFeatureLevel(), Light.LightSceneInfo->IsPrecomputedLightingValid(), View.MaxShadowCascades);
+		LightDataElement.InvRadius = LightParameters.InvRadius;
+		LightDataElement.SpecularScale = LightParameters.SpecularScale;
+		LightDataElement.FalloffExponent = LightParameters.FalloffExponent;
+		LightDataElement.SourceRadius = LightParameters.SourceRadius;
+		LightDataElement.SourceLength = LightParameters.SourceLength;
+		LightDataElement.SoftSourceRadius = LightParameters.SoftSourceRadius;
+		LightDataElement.RectLightBarnCosAngle = LightParameters.RectLightBarnCosAngle;
+		LightDataElement.RectLightBarnLength = LightParameters.RectLightBarnLength;
+		LightDataElement.Pad = 0;
 
-		FVector4 DistanceFadeMAD_RectLightBarnCosAngle_RectLightBarnLength = FVector4(FadeParams.Y, -FadeParams.X * FadeParams.Y, LightParameters.RectLightBarnCosAngle, LightParameters.RectLightBarnLength);
-		LightData->DistanceFadeMAD_RectLightBarnCosAngle_RectLightBarnLength[LightData->Count] = DistanceFadeMAD_RectLightBarnCosAngle_RectLightBarnLength;
+		LightDataArray.Add(LightDataElement);
 
 		const bool bRequireTexture = Light.LightType == ELightComponentType::LightType_Rect && LightParameters.SourceTexture;
 		uint32 RectLightTextureIndex = InvalidTextureIndex;
@@ -132,7 +144,7 @@ void SetupRaytracingLightDataPacked(
 
 		if (RectLightTextureIndex != InvalidTextureIndex)
 		{
-			LightData->Type_LightProfileIndex_RectLightTextureIndex[LightData->Count].Z = RectLightTextureIndex;
+			LightDataArray[LightData->Count].RectLightTextureIndex = RectLightTextureIndex;
 			switch (RectLightTextureIndex)
 			{
 			case 0: LightData->RectLightTexture0 = LightParameters.SourceTexture; break;
@@ -166,10 +178,27 @@ void SetupRaytracingLightDataPacked(
 	}
 }
 
-TUniformBufferRef<FRaytracingLightDataPacked> CreateLightDataPackedUniformBuffer(const TSparseArray<FLightSceneInfoCompact>& Lights, const class FViewInfo& View, EUniformBufferUsage Usage)
+TUniformBufferRef<FRaytracingLightDataPacked> CreateLightDataPackedUniformBuffer(const TSparseArray<FLightSceneInfoCompact>& Lights, const class FViewInfo& View, EUniformBufferUsage Usage, FStructuredBufferRHIRef& OutLightDataArray)
 {
 	FRaytracingLightDataPacked LightData;
-	SetupRaytracingLightDataPacked(Lights, View, &LightData);
+	TResourceArray<FRTLightingData> LightDataArray;
+
+	SetupRaytracingLightDataPacked(Lights, View, &LightData, LightDataArray);
+
+	check(LightData.Count == LightDataArray.Num());
+
+	// need at least one element
+	if (LightDataArray.Num() == 0)
+	{
+		LightDataArray.AddZeroed(1);
+	}
+
+	// This buffer might be best placed as an element of the LightData uniform buffer
+	FRHIResourceCreateInfo CreateInfo;
+	CreateInfo.ResourceArray = &LightDataArray;
+
+	OutLightDataArray = RHICreateStructuredBuffer(sizeof(FRTLightingData), LightDataArray.GetResourceDataSize(), BUF_Static | BUF_ShaderResource, CreateInfo);
+
 	return CreateUniformBufferImmediate(LightData, Usage);
 }
 
diff --git a/Engine/Source/Runtime/Renderer/Private/RayTracing/RayTracingLighting.h b/Engine/Source/Runtime/Renderer/Private/RayTracing/RayTracingLighting.h
--- a/Engine/Source/Runtime/Renderer/Private/RayTracing/RayTracingLighting.h
+++ b/Engine/Source/Runtime/Renderer/Private/RayTracing/RayTracingLighting.h
@@ -1,6 +1,6 @@
 // Copyright 1998-2019 Epic Games, Inc. All Rights Reserved.
 
-#pragma once 
+#pragma once
 
 #include "CoreMinimal.h"
 #include "ShaderParameters.h"
@@ -14,13 +14,6 @@
 BEGIN_GLOBAL_SHADER_PARAMETER_STRUCT(FRaytracingLightDataPacked, )
 	SHADER_PARAMETER(uint32, Count)
 	SHADER_PARAMETER(float, IESLightProfileInvCount)
-	SHADER_PARAMETER_ARRAY(FIntVector, Type_LightProfileIndex_RectLightTextureIndex, [RAY_TRACING_LIGHT_COUNT_MAXIMUM])
-	SHADER_PARAMETER_ARRAY(FVector4, LightPosition_InvRadius, [RAY_TRACING_LIGHT_COUNT_MAXIMUM])
-	SHADER_PARAMETER_ARRAY(FVector4, LightColor_SpecularScale, [RAY_TRACING_LIGHT_COUNT_MAXIMUM])
-	SHADER_PARAMETER_ARRAY(FVector4, Direction_FalloffExponent, [RAY_TRACING_LIGHT_COUNT_MAXIMUM])
-	SHADER_PARAMETER_ARRAY(FVector4, Tangent_SourceRadius, [RAY_TRACING_LIGHT_COUNT_MAXIMUM])
-	SHADER_PARAMETER_ARRAY(FVector4, SpotAngles_SourceLength_SoftSourceRadius, [RAY_TRACING_LIGHT_COUNT_MAXIMUM])
-	SHADER_PARAMETER_ARRAY(FVector4, DistanceFadeMAD_RectLightBarnCosAngle_RectLightBarnLength, [RAY_TRACING_LIGHT_COUNT_MAXIMUM])
 	SHADER_PARAMETER_TEXTURE(Texture2D, LTCMatTexture)
 	SHADER_PARAMETER_SAMPLER(SamplerState, LTCMatSampler)
 	SHADER_PARAMETER_TEXTURE(Texture2D, LTCAmpTexture)
@@ -37,11 +30,41 @@ BEGIN_GLOBAL_SHADER_PARAMETER_STRUCT(FRaytracingLightDataPacked, )
 	SHADER_PARAMETER_TEXTURE(Texture2D, IESLightProfileTexture)
 END_GLOBAL_SHADER_PARAMETER_STRUCT()
 
+// Must match struct definition in RayTacedLightingCommon.ush
+struct FRTLightingData
+{
+	int32 Type;
+	int32 LightProfileIndex;
+	int32 RectLightTextureIndex;
+
+	// Force alignment before next vector
+	int32 Pad;
+
+	float LightPosition[3];
+	float InvRadius;
+	float Direction[3];
+	float FalloffExponent;
+	float LightColor[3];
+	float SpecularScale;
+	float Tangent[3];
+	float SourceRadius;
+	float SpotAngles[2];
+	float SourceLength;
+	float SoftSourceRadius;
+	float DistanceFadeMAD[2];
+	float RectLightBarnCosAngle;
+	float RectLightBarnLength;
+
+	// Align struct to 128 bytes to better match cache lines
+	float Dummy[4];
+};
+
 void SetupRaytracingLightDataPacked(
 	const TSparseArray<FLightSceneInfoCompact>& Lights,
 	const FViewInfo& View,
-	FRaytracingLightDataPacked* LightData);
+	FRaytracingLightDataPacked* LightData,
+	TResourceArray<FRTLightingData>& LightDataArray);
 
-TUniformBufferRef<FRaytracingLightDataPacked> CreateLightDataPackedUniformBuffer(const TSparseArray<FLightSceneInfoCompact>& Lights, const class FViewInfo& View, EUniformBufferUsage Usage);
+TUniformBufferRef<FRaytracingLightDataPacked> CreateLightDataPackedUniformBuffer(const TSparseArray<FLightSceneInfoCompact>& Lights, const class FViewInfo& View, EUniformBufferUsage Usage, FStructuredBufferRHIRef& LightDataArray);
 
-#endif
\ No newline at end of file
+#endif
diff --git a/Engine/Source/Runtime/Renderer/Private/RayTracing/RayTracingReflections.cpp b/Engine/Source/Runtime/Renderer/Private/RayTracing/RayTracingReflections.cpp
--- a/Engine/Source/Runtime/Renderer/Private/RayTracing/RayTracingReflections.cpp
+++ b/Engine/Source/Runtime/Renderer/Private/RayTracing/RayTracingReflections.cpp
@@ -139,6 +139,7 @@ class FRayTracingReflectionsRGS : public FGlobalShader
 		SHADER_PARAMETER(float, ReflectionMaxNormalBias)
 
 		SHADER_PARAMETER_SRV(RaytracingAccelerationStructure, TLAS)
+		SHADER_PARAMETER_SRV(StructuredBuffer<FRTLightingData>, LightDataBuffer)
 
 		SHADER_PARAMETER_STRUCT_REF(FViewUniformShaderParameters, ViewUniformBuffer)
 		SHADER_PARAMETER_STRUCT_REF(FSceneTexturesUniformParameters, SceneTexturesStruct)
@@ -283,7 +284,11 @@ void FDeferredShadingSceneRenderer::RenderRayTracingReflections(
 
 	CommonParameters.TLAS = View.RayTracingScene.RayTracingSceneRHI->GetShaderResourceView();
 	CommonParameters.ViewUniformBuffer = View.ViewUniformBuffer;
-	CommonParameters.LightDataPacked = CreateLightDataPackedUniformBuffer(Scene->Lights, View, EUniformBufferUsage::UniformBuffer_SingleFrame);
+
+	FStructuredBufferRHIRef LightingDataBuffer;
+	CommonParameters.LightDataPacked = CreateLightDataPackedUniformBuffer(Scene->Lights, View, EUniformBufferUsage::UniformBuffer_SingleFrame, LightingDataBuffer);
+	CommonParameters.LightDataBuffer = RHICreateShaderResourceView(LightingDataBuffer);
+
 	CommonParameters.SceneTexturesStruct = CreateSceneTextureUniformBuffer( SceneContext, FeatureLevel, ESceneTextureSetupMode::All, EUniformBufferUsage::UniformBuffer_SingleFrame);
 	CommonParameters.ReflectionStruct = CreateReflectionUniformBuffer(View, EUniformBufferUsage::UniformBuffer_SingleFrame);
 	CommonParameters.FogUniformParameters = CreateFogUniformBuffer(View, EUniformBufferUsage::UniformBuffer_SingleFrame);
diff --git a/Engine/Source/Runtime/Renderer/Private/RayTracing/RayTracingTranslucency.cpp b/Engine/Source/Runtime/Renderer/Private/RayTracing/RayTracingTranslucency.cpp
--- a/Engine/Source/Runtime/Renderer/Private/RayTracing/RayTracingTranslucency.cpp
+++ b/Engine/Source/Runtime/Renderer/Private/RayTracing/RayTracingTranslucency.cpp
@@ -141,6 +141,7 @@ class FRayTracingTranslucencyRGS : public FGlobalShader
 		SHADER_PARAMETER(float, MaxNormalBias)
 
 		SHADER_PARAMETER_SRV(RaytracingAccelerationStructure, TLAS)
+		SHADER_PARAMETER_SRV(StructuredBuffer<FRTLightingData>, LightDataBuffer)
 
 		SHADER_PARAMETER_STRUCT_REF(FViewUniformShaderParameters, ViewUniformBuffer)
 		SHADER_PARAMETER_STRUCT_REF(FSceneTexturesUniformParameters, SceneTexturesStruct)
@@ -390,7 +391,11 @@ void FDeferredShadingSceneRenderer::RenderRayTracingTranslucencyView(
 
 	PassParameters->TLAS = View.RayTracingScene.RayTracingSceneRHI->GetShaderResourceView();
 	PassParameters->ViewUniformBuffer = View.ViewUniformBuffer;
-	PassParameters->LightDataPacked = CreateLightDataPackedUniformBuffer(Scene->Lights, View, EUniformBufferUsage::UniformBuffer_SingleFrame);
+
+	FStructuredBufferRHIRef LightingDataBuffer;
+	PassParameters->LightDataPacked = CreateLightDataPackedUniformBuffer(Scene->Lights, View, EUniformBufferUsage::UniformBuffer_SingleFrame, LightingDataBuffer);
+	PassParameters->LightDataBuffer = RHICreateShaderResourceView(LightingDataBuffer);
+
 	PassParameters->SceneTexturesStruct = CreateSceneTextureUniformBuffer(SceneContext, FeatureLevel, ESceneTextureSetupMode::All, EUniformBufferUsage::UniformBuffer_SingleFrame);	
 	PassParameters->ReflectionStruct = CreateReflectionUniformBuffer(View, EUniformBufferUsage::UniformBuffer_SingleFrame);
 	PassParameters->FogUniformParameters = CreateFogUniformBuffer(View, EUniformBufferUsage::UniformBuffer_SingleFrame);
