DR - Fix merge issue
#rnx

diff --git a/Engine/Shaders/Private/ScreenPass.usf b/Engine/Shaders/Private/ScreenPass.usf
--- a/Engine/Shaders/Private/ScreenPass.usf
+++ b/Engine/Shaders/Private/ScreenPass.usf
@@ -5,6 +5,15 @@
 Texture2D SrcTexture;
 SamplerState SrcSampler;
 
+void ScreenPassVS(
+	in float4 InPosition : ATTRIBUTE0,
+	in float2 InTexCoord : ATTRIBUTE1,
+	out noperspective float4 OutUVAndScreenPos : TEXCOORD0,
+	out float4 OutPosition : SV_POSITION)
+{
+	DrawRectangle(InPosition, InTexCoord, OutPosition, OutUVAndScreenPos);
+}
+
 float4 CopyRectPS(
 	noperspective float4 UVAndScreenPos : TEXCOORD0
     ) : SV_Target0
