Caching SkyLight cube face CDF in SkyLight pass.

Provides measurable benefit when casting many sky light importance sampling rays. Diminishing returns for global illumination and path tracing passes as doing so will require extra plumbing and register pressure.

diff --git a/Engine/Shaders/Private/RayTracing/RayTracingSkyLightRGS.usf b/Engine/Shaders/Private/RayTracing/RayTracingSkyLightRGS.usf
--- a/Engine/Shaders/Private/RayTracing/RayTracingSkyLightRGS.usf
+++ b/Engine/Shaders/Private/RayTracing/RayTracingSkyLightRGS.usf
@@ -6,8 +6,7 @@
 #include "../ShadingModels.ush"
 #include "SkyLightMipTreeCommon.ush"
 #include "RayTracingCommon.ush"
-#include "RayTracingDirectionalLight.ush"
-#include "RayTracingRectLight.ush"
+#include "../PathTracing/Utilities/PathTracingRandomSequence.ush"
 
 RaytracingAccelerationStructure TLAS;
 RWTexture2D<float4> RWOcclusionMaskUAV;
@@ -23,6 +22,7 @@ float3 EvalSkyLight(float3 Direction)
 }
 
 float PdfSkyLightOcclusionRay(
+	float FaceCdf[7],
 	float3 WorldNormal,
 	float3 WorldDirection
 )
@@ -34,7 +34,7 @@ float PdfSkyLightOcclusionRay(
 	uint3 TextureCoord = GetTextureCubeCoordinate(WorldDirection, TextureRes);
 	float2 TextureUV = TextureCoord.xy / (float2)TextureRes;
 	uint2 MipCoord = TextureUV * SkyLight.MipDimensions.xy;
-	float FacePdf = PdfFace(MipCount, TextureCoord.z);
+	float FacePdf = PdfFace(FaceCdf, TextureCoord.z);
 
 	uint StopLevel = min(SkyLight.SamplingStopLevel, MipCount);
 #if USE_MIP_PDF_TEXTURE
@@ -93,6 +93,7 @@ float PdfSkyLightOcclusionRay(
 }
 
 bool GenerateSkyLightOcclusionRay(
+	float FaceCdf[7],
 	float3 WorldPosition,
 	float3 WorldNormal,
 	float4 RandSample,
@@ -106,7 +107,7 @@ bool GenerateSkyLightOcclusionRay(
 	// Determine the appropriate face to use
 	uint MipCount = log2(SkyLight.MipDimensions.x);
 	float FacePdf = 0.0;
-	uint FaceIndex = SampleFace(MipCount, RandSample.x, FacePdf);
+	uint FaceIndex = SampleFace(FaceCdf, RandSample.x, FacePdf);
 
 	uint2 MipPixel = 0;
 	float MipPdf = 1.0;
@@ -267,6 +268,11 @@ void SkyLightRGS()
 		SamplesPerPixel = 0.0;
 	}
 
+	// Cache Cubeface CDF
+	uint MipCount = log2(SkyLight.MipDimensions.x);
+	float FaceCdf[7];
+	BuildFaceCdf(MipCount, FaceCdf);
+
 	float3 ExitantRadiance = 0.0;
 	float3 DiffuseExitantRadiance = 0.0;
 	for (uint SampleIndex = 0; SampleIndex < SamplesPerPixel; ++SampleIndex)
@@ -287,13 +293,13 @@ void SkyLightRGS()
 		if (RandSample.x < SkyLightSamplingStrategyPdf)
 		{
 			RandSample.x /= SkyLightSamplingStrategyPdf;
-			GenerateSkyLightOcclusionRay(WorldPosition, WorldNormal, RandSample, Ray.Origin, Ray.Direction, Ray.TMin, Ray.TMax, RayPdf);
+			GenerateSkyLightOcclusionRay(FaceCdf, WorldPosition, WorldNormal, RandSample, Ray.Origin, Ray.Direction, Ray.TMin, Ray.TMax, RayPdf);
 			MisWeight = RayPdf / (SkyLightSamplingStrategyPdf * RayPdf + (1.0 - SkyLightSamplingStrategyPdf) * PdfLambertOcclusionRay(WorldNormal, Ray.Direction));
 		}
 		else
 		{
 			GenerateLambertOcclusionRay(WorldPosition, WorldNormal, RandSample.yz, Ray.Origin, Ray.Direction, Ray.TMin, Ray.TMax, RayPdf);
-			MisWeight = RayPdf / ((1.0 - SkyLightSamplingStrategyPdf) * RayPdf + SkyLightSamplingStrategyPdf * PdfSkyLightOcclusionRay(WorldNormal, Ray.Direction));
+			MisWeight = RayPdf / ((1.0 - SkyLightSamplingStrategyPdf) * RayPdf + SkyLightSamplingStrategyPdf * PdfSkyLightOcclusionRay(FaceCdf, WorldNormal, Ray.Direction));
 		}
 
 		ApplyPositionBias(Ray, WorldNormal, SkyLight.MaxNormalBias);
diff --git a/Engine/Shaders/Private/RayTracing/SkyLightMipTreeCommon.ush b/Engine/Shaders/Private/RayTracing/SkyLightMipTreeCommon.ush
--- a/Engine/Shaders/Private/RayTracing/SkyLightMipTreeCommon.ush
+++ b/Engine/Shaders/Private/RayTracing/SkyLightMipTreeCommon.ush
@@ -131,65 +131,64 @@ float TexelCoordSolidAngle(uint2 TextureCoord, uint2 TextureRes)
 	return SolidAngle;
 }
 
-uint SampleFace(uint MipCount, float RandSample, out float FacePdf)
+void BuildFaceCdf(uint MipCount, out float FaceCdf[7])
 {
 	uint BufferOffset = BufferOffsetAtPixel(uint2(0, 0), MipCount, SkyLight.MipDimensions.xy);
 
-	float FaceCdf[6];
-	FaceCdf[0] = SkyLight.MipTreePosX[BufferOffset];
-	FaceCdf[1] = FaceCdf[0] + SkyLight.MipTreeNegX[BufferOffset];
-	FaceCdf[2] = FaceCdf[1] + SkyLight.MipTreePosY[BufferOffset];
-	FaceCdf[3] = FaceCdf[2] + SkyLight.MipTreeNegY[BufferOffset];
-	FaceCdf[4] = FaceCdf[3] + SkyLight.MipTreePosZ[BufferOffset];
-	FaceCdf[5] = FaceCdf[4] + SkyLight.MipTreeNegZ[BufferOffset];
-	float Sum = FaceCdf[5];
+	FaceCdf[0] = 0.0;
+	FaceCdf[1] = SkyLight.MipTreePosX[BufferOffset];
+	FaceCdf[2] = FaceCdf[1] + SkyLight.MipTreeNegX[BufferOffset];
+	FaceCdf[3] = FaceCdf[2] + SkyLight.MipTreePosY[BufferOffset];
+	FaceCdf[4] = FaceCdf[3] + SkyLight.MipTreeNegY[BufferOffset];
+	FaceCdf[5] = FaceCdf[4] + SkyLight.MipTreePosZ[BufferOffset];
+	FaceCdf[6] = FaceCdf[5] + SkyLight.MipTreeNegZ[BufferOffset];
+	float Sum = FaceCdf[6];
 
-	for (uint Index = 0; Index < 5; ++Index)
+	for (uint Index = 1; Index < 6; ++Index)
 	{
 		FaceCdf[Index] /= Sum;
 	}
-	FaceCdf[5] = 1.0;
+	FaceCdf[6] = 1.0;
+}
 
+uint SampleFace(float FaceCdf[7], float RandSample, out float FacePdf)
+{
 	// Determine CDF entry
-	uint FaceIndex = 0;
+	uint FaceIndex = 1;
 	while (RandSample > FaceCdf[FaceIndex])
 	{
 		FaceIndex++;
 	}
 
 	// Calculate PDF and return entry
-	FacePdf = FaceCdf[FaceIndex];
-	if (FaceIndex > 0)
-		FacePdf -= FaceCdf[FaceIndex - 1];
-	return FaceIndex;
+	FacePdf = FaceCdf[FaceIndex] - FaceCdf[FaceIndex - 1];
+
+	// Cdf is incremented by 1
+	return FaceIndex - 1;
 }
 
-float PdfFace(uint MipCount, uint FaceIndex)
+uint SampleFace(int MipCount, float RandSample, out float FacePdf)
 {
-	// Build Cdf
-	uint BufferOffset = BufferOffsetAtPixel(uint2(0, 0), MipCount, SkyLight.MipDimensions.xy);
-
-	float FaceCdf[6];
-	FaceCdf[0] = SkyLight.MipTreePosX[BufferOffset];
-	FaceCdf[1] = FaceCdf[0] + SkyLight.MipTreeNegX[BufferOffset];
-	FaceCdf[2] = FaceCdf[1] + SkyLight.MipTreePosY[BufferOffset];
-	FaceCdf[3] = FaceCdf[2] + SkyLight.MipTreeNegY[BufferOffset];
-	FaceCdf[4] = FaceCdf[3] + SkyLight.MipTreePosZ[BufferOffset];
-	FaceCdf[5] = FaceCdf[4] + SkyLight.MipTreeNegZ[BufferOffset];
-	float Sum = FaceCdf[5];
-
-	for (uint Index = 0; Index < 5; ++Index)
-	{
-		FaceCdf[Index] /= Sum;
-	}
-	FaceCdf[5] = 1.0;
+	float FaceCdf[7];
+	BuildFaceCdf(MipCount, FaceCdf);
+	return SampleFace(FaceCdf, RandSample, FacePdf);
+}
 
-	float Pdf = FaceCdf[FaceIndex];
-	if (FaceIndex > 0)
-		Pdf -= FaceCdf[FaceIndex - 1];
+float PdfFace(float FaceCdf[7], uint FaceIndex)
+{
+	// Cdf is offset by 1
+	FaceIndex++;
+	float Pdf = FaceCdf[FaceIndex] - FaceCdf[FaceIndex - 1];
 	return Pdf;
 }
 
+float PdfFace(int MipCount, uint FaceIndex)
+{
+	float FaceCdf[7];
+	BuildFaceCdf(MipCount, FaceCdf);
+	return PdfFace(FaceCdf, FaceIndex);
+}
+
 float4 SampleMipTree(uint FaceIndex, uint4 BufferOffset)
 {
 	float4 Values;
