Refactored ray tracing payload structures and removed redundant miss/hit shaders.

All ray tracing payloads now inherit from a common base that contains ray intersection distance (HitT).
Most ray tracing effects now use a single common miss shader that simply sets the HitT of the common base payload to a negative value.

diff --git a/Engine/Shaders/Private/PathTracing/PathTracing.usf b/Engine/Shaders/Private/PathTracing/PathTracing.usf
--- a/Engine/Shaders/Private/PathTracing/PathTracing.usf
+++ b/Engine/Shaders/Private/PathTracing/PathTracing.usf
@@ -88,7 +88,7 @@ void PathTracingMainRG()
 		while (Bounces < PathTracingData.MaxBounces)
 		{
 			// Environment contribution
-			if (!IsHit(Ray, Payload))
+			if (Payload.IsMiss())
 			{
 				// WARNING: Until ray payload encodes NEE LightId, the execution depends on this ordering!
 				uint SkyLightId = 0;
@@ -150,7 +150,7 @@ void PathTracingMainRG()
 
 			Bounces = Bounces + 1;
 			// Update origin
-			Ray.Origin = Ray.Origin + Ray.Direction * GetHitT(Payload);
+			Ray.Origin = Ray.Origin + Ray.Direction * Payload.HitT;
 			FMaterialClosestHitPayload RayHitInfo = Payload; 
 
 			// Perform next-event estimation
@@ -187,7 +187,7 @@ void PathTracingMainRG()
 					);
 					RayCounter++;
 
-					if (!IsHit(LightRay, NeePayload))
+					if (NeePayload.IsMiss())
 					{
 						FMaterialClosestHitPayload RayHitInfo = Payload; 
 #if 0
diff --git a/Engine/Shaders/Private/PathTracing/PathTracingCHS.usf b/Engine/Shaders/Private/PathTracing/PathTracingCHS.usf
deleted file mode 100644
diff --git a/Engine/Shaders/Private/PathTracing/PathTracingMS.usf b/Engine/Shaders/Private/PathTracing/PathTracingMS.usf
deleted file mode 100644
diff --git a/Engine/Shaders/Private/RayTracing/RayTracingAmbientOcclusionRGS.usf b/Engine/Shaders/Private/RayTracing/RayTracingAmbientOcclusionRGS.usf
--- a/Engine/Shaders/Private/RayTracing/RayTracingAmbientOcclusionRGS.usf
+++ b/Engine/Shaders/Private/RayTracing/RayTracingAmbientOcclusionRGS.usf
@@ -108,8 +108,8 @@ void AmbientOcclusionRGS()
 		);
 		
 		RayCount += 1;
-		Visibility += Payload.HitT < 0 ? 1.0 : 1.0 - IntensityLocal;
-		if (Payload.HitT >= 0.0)
+		Visibility += Payload.IsMiss() ? 1.0 : 1.0 - IntensityLocal;
+		if (Payload.IsHit())
 		{
 			ClosestRayHitDistance = min(ClosestRayHitDistance, Payload.HitT);
 		}
diff --git a/Engine/Shaders/Private/RayTracing/RayTracingBarycentrics.usf b/Engine/Shaders/Private/RayTracing/RayTracingBarycentrics.usf
--- a/Engine/Shaders/Private/RayTracing/RayTracingBarycentrics.usf
+++ b/Engine/Shaders/Private/RayTracing/RayTracingBarycentrics.usf
@@ -6,10 +6,9 @@
 RWTexture2D<float4> Output;
 RaytracingAccelerationStructure TLAS;
 
-struct FBarycentricsPayload
+struct FBarycentricsPayload : FMinimalPayload
 {
 	float2 Barycentrics;
-	bool Hit;
 };
 
 
@@ -45,22 +44,16 @@ void RayTracingBarycentricsMainRGS()
 	);
 
 	float4 Result = float4(0, 0, 0, 1);
-	if (Payload.Hit)
+	if (Payload.IsHit())
 	{
 		Result = float4(Payload.Barycentrics, 1.0f - Payload.Barycentrics.x - Payload.Barycentrics.y, 1.0f);
 	}
 	Output[PixelCoord] = Result * View.PreExposure;
 }
 
-[shader("miss")]
-void RayTracingBarycentricsMainMS(inout FBarycentricsPayload Payload)
-{
-	Payload.Hit = false;
-}
-
 [shader("closesthit")]
 void RayTracingBarycentricsMainCHS(inout FBarycentricsPayload Payload, in FDefaultAttributes Attributes)
 {
 	Payload.Barycentrics = Attributes.Barycentrics;
-	Payload.Hit = true;
+	Payload.HitT = RayTCurrent();
 }
diff --git a/Engine/Shaders/Private/RayTracing/RayTracingBuiltInShaders.usf b/Engine/Shaders/Private/RayTracing/RayTracingBuiltInShaders.usf
--- a/Engine/Shaders/Private/RayTracing/RayTracingBuiltInShaders.usf
+++ b/Engine/Shaders/Private/RayTracing/RayTracingBuiltInShaders.usf
@@ -10,7 +10,7 @@
 RaytracingAccelerationStructure               TLAS               : RT_REGISTER(t, 0, RAY_TRACING_REGISTER_SPACE_GLOBAL);
 StructuredBuffer<FBasicRayData>               Rays               : RT_REGISTER(t, 1, RAY_TRACING_REGISTER_SPACE_GLOBAL);
 RWStructuredBuffer<uint>                      OcclusionOutput    : RT_REGISTER(u, 0, RAY_TRACING_REGISTER_SPACE_GLOBAL);
-RWStructuredBuffer<FBasicRayIntersectionData> IntersectionOutput : RT_REGISTER(u, 1, RAY_TRACING_REGISTER_SPACE_GLOBAL);
+RWStructuredBuffer<FIntersectionPayload>      IntersectionOutput : RT_REGISTER(u, 1, RAY_TRACING_REGISTER_SPACE_GLOBAL);
 
 #undef RT_REGISTER
 #undef RT_CONCATENATE
@@ -34,8 +34,7 @@ void OcclusionMainRG()
 	              | RAY_FLAG_FORCE_OPAQUE                    // don't run anyhit shader
 	              | RAY_FLAG_SKIP_CLOSEST_HIT_SHADER;        // don't run closesthit shader
 
-	FOcclusionPayload Payload;
-	Payload.Hit = true;
+	FMinimalPayload Payload = (FMinimalPayload)0;
 
 	TraceRay(
 		TLAS,   // AccelerationStructure
@@ -48,13 +47,7 @@ void OcclusionMainRG()
 		Payload // Payload
 	);
 
-	OcclusionOutput[RayIndex] = Payload.Hit ? ~0 : 0;
-}
-
-[shader("miss")]
-void OcclusionMainMS(inout FOcclusionPayload Payload)
-{
-	Payload.Hit = false;
+	OcclusionOutput[RayIndex] = Payload.IsHit() ? ~0 : 0;
 }
 
 [shader("raygeneration")]
@@ -73,7 +66,7 @@ void IntersectionMainRG()
 	              | RAY_FLAG_CULL_BACK_FACING_TRIANGLES      // use back face culling
 	              | RAY_FLAG_FORCE_OPAQUE;                   // don't run anyhit shader
 
-	FBasicRayIntersectionData Payload = (FBasicRayIntersectionData)0;
+	FIntersectionPayload Payload = (FIntersectionPayload)0;
 
 	TraceRay(
 		TLAS,   // AccelerationStructure
@@ -89,18 +82,10 @@ void IntersectionMainRG()
 	IntersectionOutput[RayIndex] = Payload;
 }
 
-[shader("miss")]
-void IntersectionMainMS(inout FBasicRayIntersectionData Payload)
-{
-	Payload.Barycentrics = float2(0, 0);
-	Payload.InstanceIndex = ~0;
-	Payload.PrimitiveIndex = ~0;
-	
-}
-
 [shader("closesthit")]
-void IntersectionMainCHS(inout FBasicRayIntersectionData Payload, in FDefaultAttributes Attributes)
+void IntersectionMainCHS(inout FIntersectionPayload Payload, in FDefaultAttributes Attributes)
 {
+	Payload.HitT = RayTCurrent();
 	Payload.Barycentrics = Attributes.Barycentrics;
 	Payload.InstanceIndex = InstanceIndex();
 	Payload.PrimitiveIndex = PrimitiveIndex();
@@ -116,13 +101,10 @@ void DefaultMainCHS(inout FDefaultPayload Payload, in FDefaultAttributes Attribu
 	Payload.HitT = RayTCurrent();
 }
 
+// Default miss shader that can be used in all ray tracing pipelines that use payloads derived from FMinimalPayload.
 [shader("miss")]
-void DefaultMainMS(inout FDefaultPayload Payload)
+void DefaultMainMS(inout FMinimalPayload Payload)
 {
-	Payload.Barycentrics = float2(0, 0);
-	Payload.InstanceID = ~0;
-	Payload.InstanceIndex = ~0;
-	Payload.PrimitiveIndex = ~0;
-	Payload.HitT = -1;
+	Payload.SetMiss();
 }
 
diff --git a/Engine/Shaders/Private/RayTracing/RayTracingCommon.ush b/Engine/Shaders/Private/RayTracing/RayTracingCommon.ush
--- a/Engine/Shaders/Private/RayTracing/RayTracingCommon.ush
+++ b/Engine/Shaders/Private/RayTracing/RayTracingCommon.ush
@@ -23,37 +23,35 @@ struct FBasicRayData
 	float TFar;
 };
 
-struct FBasicRayIntersectionData
+struct FMinimalPayload
 {
-	float  Distance;
-	uint   PrimitiveIndex; // 0xFFFFFFFF if no intersection if found
-	uint   InstanceIndex;  // 0xFFFFFFFF if no intersection if found
-	float2 Barycentrics;
+	float HitT; // Distance from ray origin to the intersection point in the ray direction. Negative on miss.
+
+	bool IsMiss() { return HitT < 0; }
+	bool IsHit() { return !IsMiss(); }
+
+	void SetMiss() { HitT = -1; }
 };
 
-struct FOcclusionPayload
+struct FIntersectionPayload : FMinimalPayload
 {
-	bool Hit;
+	uint   PrimitiveIndex; // Index of the primitive within the geometry inside the bottom-level acceleration structure instance. Undefined on miss.
+	uint   InstanceIndex;  // Index of the current instance in the top-level structure. Undefined on miss.
+	float2 Barycentrics;   // Primitive barycentric coordinates of the intersection point. Undefined on miss.
 };
 
-struct FDefaultPayload
+struct FDefaultPayload : FIntersectionPayload
 {
-	float2 Barycentrics;
-	uint   InstanceID;
-	uint   InstanceIndex;
-	uint   PrimitiveIndex;
-	float  HitT;
-	float2 Padding;
+	uint   InstanceID; // Value of FRayTracingGeometryInstance::UserData. Undefined on miss.
 };
 
 #define RAY_TRACING_DEFERRED_MATERIAL_KEY_INVALID   0xFFFFFFFF
 #define RAY_TRACING_DEFERRED_MATERIAL_KEY_RAY_MISS  0xFFFFFFFE
 
-struct FDeferredMaterialPayload
+struct FDeferredMaterialPayload : FMinimalPayload
 {
 	uint   SortKey;          // MaterialID by default
 	uint   PixelCoordinates; // X in low 16 bits, Y in high 16 bits
-	float  HitT;
 };
 
 struct FDefaultAttributes
@@ -107,53 +105,53 @@ uint2 GetPixelCoord(uint2 DispatchThreadId, uint UpscaleFactor)
 #define RAY_TRACING_BLEND_MODE_ALPHA_HOLDOUT		5
 
 
-struct FMaterialClosestHitPayload
+struct FMaterialClosestHitPayload : FMinimalPayload
 {
 										// Unpacked  Packed
 										// offset    bytes
-	FRayCone RayCone;					// 8         8       64bits
-	float HitT;							// 12        4       32bits
-
-	float3 Radiance;					// 24        6       48bits
-	float3 WorldNormal;					// 36        6       48bits
-	float3 BaseColor;					// 48        6       48bits
-	float3 DiffuseColor;				// 60        0       (derived)
-	float3 SpecularColor;				// 72        0       (derived)
-	float Opacity;						// 76        2       16bits
-	float Metallic;						// 80        1       8bits
-	float Specular;						// 84        1       8bits
-	float Roughness;					// 88        2       16bits
-	float Ior;							// 92        2       16bits
-	uint ShadingModelID;				// 96        1       4bits
-	uint BlendingMode;					// 100       0       4bits (packed with ShadingModelID)
-	uint PrimitiveLightingChannelMask;	// 104       0       (removed)
-	float4 CustomData;					// 120       4       32bits -- only xy are packed at the moment
-	float GBufferAO;					// 124       0       (removed)
-	float3 IndirectIrradiance;			// 136       0       48bits -- gbuffer only has float payload and there are truncation HLSL warnings
+	// float FMinimalPayload::HitT		// 0         4       32bits
+	FRayCone RayCone;					// 4         8       64bits
+	float3 Radiance;					// 8         6       48bits
+	float3 WorldNormal;					// 24        6       48bits
+	float3 BaseColor;					// 36        6       48bits
+	float3 DiffuseColor;				// 48        0       (derived)
+	float3 SpecularColor;				// 60        0       (derived)
+	float Opacity;						// 72        2       16bits
+	float Metallic;						// 76        1       8bits
+	float Specular;						// 80        1       8bits
+	float Roughness;					// 84        2       16bits
+	float Ior;							// 88        2       16bits
+	uint ShadingModelID;				// 92        1       4bits
+	uint BlendingMode;					// 96        0       4bits (packed with ShadingModelID)
+	uint PrimitiveLightingChannelMask;	// 100       0       (removed)
+	float4 CustomData;					// 104       4       32bits -- only xy are packed at the moment
+	float GBufferAO;					// 120       0       (removed)
+	float3 IndirectIrradiance;			// 124       0       48bits -- gbuffer only has float payload and there are truncation HLSL warnings
 
 	// Quite some code assume FRayHitInfo has a WorldPos
 	// So keep it here and force to re-construct it in Unpack call using ray information.
 	// It is not packed in FRayHitInfoPacked
-	float3 WorldPos;					// 148       0       (derived)
+	float3 WorldPos;					// 136       0       (derived)
 
-	uint HitKind;						// 152       0       1bit (packed)
+	uint HitKind;						// 148       0       1bit (packed)
+										// 152 total
 };
 
 #define USE_PACKED_PAYLOAD 1
 
 #if !COMPUTESHADER
 #if USE_PACKED_PAYLOAD
-struct FPackedMaterialClosestHitPayload
+struct FPackedMaterialClosestHitPayload : FMinimalPayload
 {
-	float HitT;                                         // 4  bytes
-	FRayCone RayCone;                                   // 8  bytes
+	// float FMinimalPayload::HitT                       // 4  bytes
+	FRayCone RayCone;                                    // 8  bytes
 	uint RadianceAndNormal[3];                           // 12 bytes
 	uint BaseColorAndOpacity[2];                         // 8  bytes
 	uint MetallicAndSpecularAndRoughness;                // 4  bytes
 	uint IorAndShadingModelIDAndBlendingModeAndHitKind;  // 4  bytes
 	uint IndirectIrradiance[2];                          // 8  bytes
 	uint CustomData;                                     // 4  bytes
-};                                                      // 52 bytes total
+};                                                       // 52 bytes total
 
 FPackedMaterialClosestHitPayload PackRayTracingPayload(FMaterialClosestHitPayload Input, in FRayCone RayCone)
 {
@@ -290,23 +288,6 @@ void ApplyPositionBias(inout RayDesc Ray, const float3 WorldNormal, const float
 	Ray.Origin += WorldNormal * NormalBias;
 }
 
-float GetHitT(FMaterialClosestHitPayload HitInfo)
-{
-	return HitInfo.HitT;
-}
-
-bool IsHit(RayDesc Ray, FMaterialClosestHitPayload HitInfo)
-{
-	float t = HitInfo.HitT;
-	return t >= Ray.TMin && t <= Ray.TMax;
-}
-
-bool IsHit(RayDesc Ray, FPackedMaterialClosestHitPayload HitInfo)
-{
-	float t = HitInfo.HitT;
-	return t >= Ray.TMin && t <= Ray.TMax;
-}
-
 #endif // !COMPUTESHADER
 
 #endif // RAYTRACINGCOMMON_USH_INCLUDED // Workarround for UE-66460
diff --git a/Engine/Shaders/Private/RayTracing/RayTracingDebug.usf b/Engine/Shaders/Private/RayTracing/RayTracingDebug.usf
--- a/Engine/Shaders/Private/RayTracing/RayTracingDebug.usf
+++ b/Engine/Shaders/Private/RayTracing/RayTracingDebug.usf
@@ -44,7 +44,7 @@ void RayTracingDebugMainRGS()
 
 	float4 Result = float4(0, 0, 0, 1);
 
-	if (Payload.HitT > 0)
+	if (Payload.IsHit())
 	{
 		switch (VisualizationMode)
 		{
@@ -109,13 +109,6 @@ void RayTracingDebugMainRGS()
 	Output[PixelCoord] = Result * View.PreExposure;
 }
 
-[shader("miss")]
-void RayTracingDebugMainMS(inout FPackedMaterialClosestHitPayload PackedPayload)
-{
-	PackedPayload.HitT = -1;
-}
-
-
 [shader("closesthit")]
 void RayTracingDebugMainCHS(inout FPackedMaterialClosestHitPayload PackedPayload, in FDefaultAttributes Attributes)
 {
diff --git a/Engine/Shaders/Private/RayTracing/RayTracingDeferredMaterials.usf b/Engine/Shaders/Private/RayTracing/RayTracingDeferredMaterials.usf
--- a/Engine/Shaders/Private/RayTracing/RayTracingDeferredMaterials.usf
+++ b/Engine/Shaders/Private/RayTracing/RayTracingDeferredMaterials.usf
@@ -10,10 +10,3 @@ void DeferredMaterialCHS(inout FDeferredMaterialPayload Payload, in FDefaultAttr
 	Payload.SortKey = GetHitGroupUserData();
 	Payload.HitT    = RayTCurrent();
 }
-
-[shader("miss")]
-void DeferredMaterialMS(inout FDeferredMaterialPayload Payload)
-{
-	Payload.SortKey = RAY_TRACING_DEFERRED_MATERIAL_KEY_RAY_MISS;
-	Payload.HitT    = -1;
-}
diff --git a/Engine/Shaders/Private/RayTracing/RayTracingGlobalIlluminationRGS.usf b/Engine/Shaders/Private/RayTracing/RayTracingGlobalIlluminationRGS.usf
--- a/Engine/Shaders/Private/RayTracing/RayTracingGlobalIlluminationRGS.usf
+++ b/Engine/Shaders/Private/RayTracing/RayTracingGlobalIlluminationRGS.usf
@@ -164,7 +164,7 @@ void GlobalIlluminationRGS()
 			);
 
 			// Environment hit
-			if (!IsHit(Ray, Payload))
+			if (Payload.IsMiss())
 			{
 				// Optional multi-bounce SkyLight contribution
 				if (EvalSkyLight && Bounce > 0)
@@ -245,7 +245,7 @@ void GlobalIlluminationRGS()
 					);
 
 					// No hit indicates successful next-event connection
-					if (!IsHit(LightRay, NeePayload))
+					if (NeePayload.IsMiss())
 					{
 						// Evaluate radiance
 						float3 Radiance;
@@ -348,13 +348,6 @@ void GlobalIlluminationRGS()
 	//RWRayDistanceUAV[DispatchThreadId] = float2(Luminance(Irradiance), HitDistance);
 }
 
-[shader("miss")]
-void RayTracingGlobalIlluminationMS(inout FPackedMaterialClosestHitPayload PackedPayload)
-{
-	PackedPayload.HitT = -1;
-}
-
-
 [shader("closesthit")]
 void RayTracingGlobalIlluminationCHS(inout FPackedMaterialClosestHitPayload PackedPayload, in FDefaultAttributes Attributes)
 {
diff --git a/Engine/Shaders/Private/RayTracing/RayTracingLightingCommon.ush b/Engine/Shaders/Private/RayTracing/RayTracingLightingCommon.ush
--- a/Engine/Shaders/Private/RayTracing/RayTracingLightingCommon.ush
+++ b/Engine/Shaders/Private/RayTracing/RayTracingLightingCommon.ush
@@ -548,9 +548,8 @@ void AccumulateResults(
 	inout FRayCone RayCone,
 	inout float3 Radiance)
 {
-	if (IsHit(Ray, Payload))
+	if (Payload.IsHit())
 	{
-
 		Payload.DiffuseColor = Payload.BaseColor - Payload.BaseColor * Payload.Metallic;
 		Payload.SpecularColor = ComputeF0(Payload.Specular, Payload.BaseColor, Payload.Metallic);
 
diff --git a/Engine/Shaders/Private/RayTracing/RayTracingMaterialDefaultHitShaders.usf b/Engine/Shaders/Private/RayTracing/RayTracingMaterialDefaultHitShaders.usf
--- a/Engine/Shaders/Private/RayTracing/RayTracingMaterialDefaultHitShaders.usf
+++ b/Engine/Shaders/Private/RayTracing/RayTracingMaterialDefaultHitShaders.usf
@@ -19,9 +19,3 @@ void HiddenMaterialAHS(inout FPackedMaterialClosestHitPayload PackedPayload, in
 {
 	IgnoreHit();
 }  
-
-[shader("miss")]
-void DefaultMaterialMS(inout FPackedMaterialClosestHitPayload PackedPayload)
-{
-	PackedPayload.HitT = -1;
-}
diff --git a/Engine/Shaders/Private/RayTracing/RayTracingOcclusionRGS.usf b/Engine/Shaders/Private/RayTracing/RayTracingOcclusionRGS.usf
--- a/Engine/Shaders/Private/RayTracing/RayTracingOcclusionRGS.usf
+++ b/Engine/Shaders/Private/RayTracing/RayTracingOcclusionRGS.usf
@@ -235,8 +235,8 @@ void OcclusionRGS()
 			);
 
 			RayCount += 1.0;
-			Visibility += Payload.HitT < 0 ? 1.0 : 0.0;
-			if (Payload.HitT >= 0.0)
+			Visibility += Payload.IsMiss() ? 1.0 : 0.0;
+			if (Payload.IsHit())
 			{
 				ClosestRayDistance = ClosestRayDistance == DENOISER_INVALID_HIT_DISTANCE || Payload.HitT < ClosestRayDistance ? Payload.HitT : ClosestRayDistance;
 				SumRayDistance += Payload.HitT;
diff --git a/Engine/Shaders/Private/RayTracing/RayTracingRectLightRGS.usf b/Engine/Shaders/Private/RayTracing/RayTracingRectLightRGS.usf
--- a/Engine/Shaders/Private/RayTracing/RayTracingRectLightRGS.usf
+++ b/Engine/Shaders/Private/RayTracing/RayTracingRectLightRGS.usf
@@ -281,7 +281,7 @@ void RectLightRGS()
 			Payload    // Payload
 		);
 
-		if (Payload.HitT >= 0.0)
+		if (Payload.IsHit())
 		{
 			RayDistance += Payload.HitT;
 			HitCount += 1.0;
diff --git a/Engine/Shaders/Private/RayTracing/RayTracingReflections.usf b/Engine/Shaders/Private/RayTracing/RayTracingReflections.usf
--- a/Engine/Shaders/Private/RayTracing/RayTracingReflections.usf
+++ b/Engine/Shaders/Private/RayTracing/RayTracingReflections.usf
@@ -270,6 +270,7 @@ void RayTracingReflectionsRGS()
 					if (DIM_DEFERRED_MATERIAL_MODE == DEFERRED_MATERIAL_MODE_GATHER)
 					{
 						uint RayFlags = RAY_FLAG_CULL_BACK_FACING_TRIANGLES;
+						DeferredMaterialPayload.SortKey = RAY_TRACING_DEFERRED_MATERIAL_KEY_RAY_MISS;
 						TraceRay(
 							TLAS,								// AccelerationStructure
 							RayFlags,
@@ -306,13 +307,13 @@ void RayTracingReflectionsRGS()
 				FMaterialClosestHitPayload Payload;
 				const uint RayFlags = 0;					
 				Payload = TraceRayAndAccumulateResults(TopLayerRay, TLAS, RandSequence, RayFlags, ReflectionMaxNormalBias, RAY_TRACING_MASK_OPAQUE, ReflectedShadowsType, ShouldDoDirectLighting, ShouldDoEmissiveAndIndirectLighting, RayCone, TopLayerRadiance);
-				bool isTopLayerRayValid = Payload.HitT >= 0 || (bAllowSkySampling && ReflectionStruct.SkyLightParameters.y > 0);
+				bool isTopLayerRayValid = Payload.IsHit() || (bAllowSkySampling && ReflectionStruct.SkyLightParameters.y > 0);
 
 
 				// Compute some flags and values
 				bool bApplyHeightFog = false;
 				float3 OriginToCollider = 0.0f;
-				if (Payload.HitT >= 0)
+				if (Payload.IsHit())
 				{
 					bIsValidSample = true;
 					if (BounceIndex == 0)
@@ -377,7 +378,7 @@ void RayTracingReflectionsRGS()
 					// Distance and sky: similar processing to main ray.
 					float BottomLayerHitDistance = -1;
 					bool IsBottomSampleValid = false;
-					if (BottomLayerPayload.HitT >= 0)
+					if (BottomLayerPayload.IsHit())
 					{
 						IsBottomSampleValid = true;
 						BottomLayerHitDistance = BottomLayerPayload.HitT;
@@ -577,29 +578,3 @@ void RayTracingReflectionsRGS()
 	}
 }
 
-
-[shader("miss")]
-void RayTracingReflectionsMainMS(inout FPackedMaterialClosestHitPayload PackedPayload)
-{
-	PackedPayload.HitT = -1;
-}
-
-
-[shader("closesthit")]
-void RayTracingReflectionsMainCHS(inout FPackedMaterialClosestHitPayload PackedPayload, in FDefaultAttributes Attributes)
-{
-	FMaterialClosestHitPayload Payload = (FMaterialClosestHitPayload)0;
-
-	Payload.HitT = RayTCurrent();
-	Payload.BaseColor = float3(0.2, 0.3, 0.4);
-
-	FTriangleBaseAttributes Triangle = LoadTriangleBaseAttributes(PrimitiveIndex());
-	float3 Edge0 = Triangle.LocalPositions[2] - Triangle.LocalPositions[0];
-	float3 Edge1 = Triangle.LocalPositions[1] - Triangle.LocalPositions[0];
-	float3x3 WorldToLocal = (float3x3)WorldToObject();
-	float3x3 LocalToWorldNormal = transpose(WorldToLocal);
-	Payload.WorldNormal = normalize(mul(LocalToWorldNormal, cross(Edge0, Edge1)));
-
-	PackedPayload = PackRayTracingPayload(Payload, PackedPayload.RayCone);
-}
-
diff --git a/Engine/Shaders/Private/RayTracing/RayTracingSkyLightRGS.usf b/Engine/Shaders/Private/RayTracing/RayTracingSkyLightRGS.usf
--- a/Engine/Shaders/Private/RayTracing/RayTracingSkyLightRGS.usf
+++ b/Engine/Shaders/Private/RayTracing/RayTracingSkyLightRGS.usf
@@ -331,7 +331,7 @@ void SkyLightRGS()
 			Payload    // Payload
 		);
 
-		if (Payload.HitT >= 0.0)
+		if (Payload.IsHit())
 		{
 			RayDistance += Payload.HitT;
 			HitCount += 1.0;
diff --git a/Engine/Shaders/Private/RayTracing/RayTracingTest.usf b/Engine/Shaders/Private/RayTracing/RayTracingTest.usf
--- a/Engine/Shaders/Private/RayTracing/RayTracingTest.usf
+++ b/Engine/Shaders/Private/RayTracing/RayTracingTest.usf
@@ -25,8 +25,7 @@ void TestMainRGS()
 	              | RAY_FLAG_FORCE_OPAQUE                    // don't run anyhit shader
 	              | RAY_FLAG_SKIP_CLOSEST_HIT_SHADER;        // don't run closesthit shader
 
-	FOcclusionPayload Payload;
-	Payload.Hit = true;
+	FMinimalPayload Payload = (FMinimalPayload)0;
 
 	TraceRay(
 		TLAS,   // AccelerationStructure
@@ -39,5 +38,5 @@ void TestMainRGS()
 		Payload // Payload
 	);
 
-	Output[RayIndex] = Payload.Hit ? ~0 : 0;
+	Output[RayIndex] = Payload.IsHit() ? ~0 : 0;
 }
diff --git a/Engine/Shaders/Private/RayTracing/RayTracingTranslucency.usf b/Engine/Shaders/Private/RayTracing/RayTracingTranslucency.usf
--- a/Engine/Shaders/Private/RayTracing/RayTracingTranslucency.usf
+++ b/Engine/Shaders/Private/RayTracing/RayTracingTranslucency.usf
@@ -149,7 +149,7 @@ void RayTracingTranslucencyRGS()
 		//
 		// Handle no hit condition
 		//
-		if (!IsHit(Ray, Payload))
+		if (Payload.IsMiss())
 		{
 			if (bHasScattered && bAllowSkySampling)
 			{
@@ -158,7 +158,7 @@ void RayTracingTranslucencyRGS()
 			}
 			break;
 		}
-		float3 HitPoint = Ray.Origin + Ray.Direction * GetHitT(Payload);
+		float3 HitPoint = Ray.Origin + Ray.Direction * Payload.HitT;
 
 		//
 		// Handle surface lighting
@@ -196,7 +196,7 @@ void RayTracingTranslucencyRGS()
 			FMaterialClosestHitPayload ReflectionPayload = TraceRayAndAccumulateResults(ReflectionRay, TLAS, RandSequence, RayFlags, MaxNormalBias, RAY_TRACING_MASK_ALL, ReflectedShadowsType, ShouldDoDirectLighting, ShouldDoEmissiveAndIndirectLighting, RayCone, ReflectionRadiance);
 
 			// If we have not hit anything, sample the distance sky radiance.
-			if (!IsHit(ReflectionRay, ReflectionPayload))
+			if (ReflectionPayload.IsMiss())
 			{
 				ReflectionRadiance = GetSkyRadiance(ReflectionRay.Direction, LastRoughness);
 			}
@@ -268,13 +268,6 @@ void RayTracingTranslucencyRGS()
 }
 
 
-[shader("miss")]
-void RayTracingTranslucencyMainMS(inout FPackedMaterialClosestHitPayload PackedPayload)
-{
-	PackedPayload.HitT = -1;
-}
-
-
 [shader("closesthit")]
 void RayTracingTranslucencyMainCHS(inout FPackedMaterialClosestHitPayload PackedPayload, in FDefaultAttributes Attributes)
 {
diff --git a/Engine/Source/Runtime/D3D12RHI/Private/D3D12RayTracing.cpp b/Engine/Source/Runtime/D3D12RHI/Private/D3D12RayTracing.cpp
--- a/Engine/Source/Runtime/D3D12RHI/Private/D3D12RayTracing.cpp
+++ b/Engine/Source/Runtime/D3D12RHI/Private/D3D12RayTracing.cpp
@@ -1710,7 +1710,7 @@ public:
 			FRHIRayTracingShader* OcclusionRGSTable[] = { GetBuildInRayTracingShader<FOcclusionMainRG>() };
 			OcclusionInitializer.SetRayGenShaderTable(OcclusionRGSTable);
 
-			FRHIRayTracingShader* OcclusionMSTable[] = { GetBuildInRayTracingShader<FOcclusionMainMS>() };
+			FRHIRayTracingShader* OcclusionMSTable[] = { GetBuildInRayTracingShader<FDefaultMainMS>() };
 			OcclusionInitializer.SetMissShaderTable(OcclusionMSTable);
 
 			OcclusionInitializer.bAllowHitGroupIndexing = false;
@@ -1725,7 +1725,7 @@ public:
 			FRHIRayTracingShader* IntersectionRGSTable[] = { GetBuildInRayTracingShader<FIntersectionMainRG>() };
 			IntersectionInitializer.SetRayGenShaderTable(IntersectionRGSTable);
 
-			FRHIRayTracingShader* IntersectionMSTable[] = { GetBuildInRayTracingShader<FIntersectionMainMS>() };
+			FRHIRayTracingShader* IntersectionMSTable[] = { GetBuildInRayTracingShader<FDefaultMainMS>() };
 			IntersectionInitializer.SetMissShaderTable(IntersectionMSTable);
 
 			FRHIRayTracingShader* IntersectionHitTable[] = { GetBuildInRayTracingShader<FIntersectionMainCHS>() };
diff --git a/Engine/Source/Runtime/RHI/Public/RHICommandList.h b/Engine/Source/Runtime/RHI/Public/RHICommandList.h
--- a/Engine/Source/Runtime/RHI/Public/RHICommandList.h
+++ b/Engine/Source/Runtime/RHI/Public/RHICommandList.h
@@ -128,13 +128,13 @@ struct FBasicRayData
 	float TFar;
 };
 
-// C++ counter-part of FBasicRayIntersectionData declared in RayTracingCommon.ush
-struct FBasicRayIntersectionData
+// C++ counter-part of FIntersectionPayload declared in RayTracingCommon.ush
+struct FIntersectionPayload
 {
-	float  Distance;
-	uint32 PrimitiveIndex; // 0xFFFFFFFF if no intersection if found
-	uint32 InstanceIndex;  // 0xFFFFFFFF if no intersection if found
-	float  Barycentrics[2];
+	float  HitT;            // Distance from ray origin to the intersection point in the ray direction. Negative on miss.
+	uint32 PrimitiveIndex;  // Index of the primitive within the geometry inside the bottom-level acceleration structure instance. Undefined on miss.
+	uint32 InstanceIndex;   // Index of the current instance in the top-level structure. Undefined on miss.
+	float  Barycentrics[2]; // Primitive barycentric coordinates of the intersection point. Undefined on miss.
 };
 #endif // RHI_RAYTRACING
 
@@ -3055,7 +3055,7 @@ public:
 
 	/**
 	 * Trace rays from an input buffer of FBasicRayData.
-	 * Primitive intersection results are written to output buffer as FBasicRayIntersectionData.
+	 * Primitive intersection results are written to output buffer as FIntersectionPayload.
 	 */
 	FORCEINLINE_DEBUGGABLE void RayTraceIntersection(FRHIRayTracingScene* Scene,
 		FRHIShaderResourceView* Rays,
diff --git a/Engine/Source/Runtime/RHI/Public/RHIResources.h b/Engine/Source/Runtime/RHI/Public/RHIResources.h
--- a/Engine/Source/Runtime/RHI/Public/RHIResources.h
+++ b/Engine/Source/Runtime/RHI/Public/RHIResources.h
@@ -2204,7 +2204,7 @@ public:
 			GetTypeHash(Initializer.GetCallableHash());
 	}
 
-	uint32 MaxPayloadSizeInBytes = 32; // sizeof FDefaultPayload declared in RayTracingCommon.ush
+	uint32 MaxPayloadSizeInBytes = 24; // sizeof FDefaultPayload declared in RayTracingCommon.ush
 
 	bool bAllowHitGroupIndexing = true;
 
@@ -2221,7 +2221,7 @@ public:
 	}
 
 	// Shaders that will be invoked if a ray misses all geometry.
-	// If this table is empty, then a built-in default miss shader will be used that sets HitT member of FDefaultPayload to -1.
+	// If this table is empty, then a built-in default miss shader will be used that sets HitT member of FMinimalPayload to -1.
 	// Desired miss shader can be selected by providing MissShaderIndex to TraceRay() function.
 	void SetMissShaderTable(const TArrayView<FRHIRayTracingShader*>& InMissShaders, uint64 Hash = 0)
 	{
diff --git a/Engine/Source/Runtime/Renderer/Private/DeferredShadingRenderer.cpp b/Engine/Source/Runtime/Renderer/Private/DeferredShadingRenderer.cpp
--- a/Engine/Source/Runtime/Renderer/Private/DeferredShadingRenderer.cpp
+++ b/Engine/Source/Runtime/Renderer/Private/DeferredShadingRenderer.cpp
@@ -792,11 +792,9 @@ bool FDeferredShadingSceneRenderer::DispatchRayTracingWorldUpdates(FRHICommandLi
 		if (RayGenShaders.Num())
 		{
 			auto DefaultHitShader = View.ShaderMap->GetShader<FOpaqueShadowHitGroup>()->GetRayTracingShader();
-			auto DefaultMissShader = View.ShaderMap->GetShader<FDefaultMaterialMS>()->GetRayTracingShader();
 
 			View.RayTracingMaterialPipeline = BindRayTracingMaterialPipeline(RHICmdList, View,
 				RayGenShaders,
-				DefaultMissShader,
 				DefaultHitShader
 			);
 		}
diff --git a/Engine/Source/Runtime/Renderer/Private/DeferredShadingRenderer.h b/Engine/Source/Runtime/Renderer/Private/DeferredShadingRenderer.h
--- a/Engine/Source/Runtime/Renderer/Private/DeferredShadingRenderer.h
+++ b/Engine/Source/Runtime/Renderer/Private/DeferredShadingRenderer.h
@@ -561,7 +561,7 @@ private:
 
 	bool GatherRayTracingWorldInstances(FRHICommandListImmediate& RHICmdList);
 	bool DispatchRayTracingWorldUpdates(FRHICommandListImmediate& RHICmdList);
-	FRayTracingPipelineState* BindRayTracingMaterialPipeline(FRHICommandList& RHICmdList, const FViewInfo& View, const TArrayView<FRHIRayTracingShader*>& RayGenShaderTable, FRHIRayTracingShader* MissShader, FRHIRayTracingShader* DefaultClosestHitShader);
+	FRayTracingPipelineState* BindRayTracingMaterialPipeline(FRHICommandList& RHICmdList, const FViewInfo& View, const TArrayView<FRHIRayTracingShader*>& RayGenShaderTable, FRHIRayTracingShader* DefaultClosestHitShader);
 	FRayTracingPipelineState* BindRayTracingDeferredMaterialGatherPipeline(FRHICommandList& RHICmdList, const FViewInfo& View, FRHIRayTracingShader* RayGenShader);
 
 	// #dxr_todo: UE-72565: refactor ray tracing effects to not be member functions of DeferredShadingRenderer. Register each effect at startup and just loop over them automatically
diff --git a/Engine/Source/Runtime/Renderer/Private/PathTracing.cpp b/Engine/Source/Runtime/Renderer/Private/PathTracing.cpp
--- a/Engine/Source/Runtime/Renderer/Private/PathTracing.cpp
+++ b/Engine/Source/Runtime/Renderer/Private/PathTracing.cpp
@@ -304,48 +304,6 @@ public:
 };
 IMPLEMENT_SHADER_TYPE(, FPathTracingRG, TEXT("/Engine/Private/PathTracing/PathTracing.usf"), TEXT("PathTracingMainRG"), SF_RayGen);
 
-class FPathTracingCHS : public FGlobalShader
-{
-public:
-	DECLARE_SHADER_TYPE(FPathTracingCHS, Global);
-
-public:
-	static bool ShouldCompilePermutation(const FGlobalShaderPermutationParameters& Parameters)
-	{
-		return ShouldCompileRayTracingShadersForProject(Parameters.Platform);
-	}
-
-	FPathTracingCHS() {}
-	virtual ~FPathTracingCHS() {}
-
-	FPathTracingCHS(const ShaderMetaType::CompiledShaderInitializerType& Initializer)
-		: FGlobalShader(Initializer)
-	{}
-};
-
-IMPLEMENT_SHADER_TYPE(, FPathTracingCHS, TEXT("/Engine/Private/PathTracing/PathTracingCHS.usf"), TEXT("PathTracingMainCHS"), SF_RayHitGroup);
-
-class FPathTracingMS : public FGlobalShader
-{
-public:
-	DECLARE_SHADER_TYPE(FPathTracingMS, Global);
-
-public:
-	static bool ShouldCompilePermutation(const FGlobalShaderPermutationParameters& Parameters)
-	{
-		return ShouldCompileRayTracingShadersForProject(Parameters.Platform);
-	}
-
-	FPathTracingMS() {}
-	virtual ~FPathTracingMS() {}
-
-	FPathTracingMS(const ShaderMetaType::CompiledShaderInitializerType& Initializer)
-		: FGlobalShader(Initializer)
-	{}
-};
-
-IMPLEMENT_SHADER_TYPE(, FPathTracingMS, TEXT("/Engine/Private/PathTracing/PathTracingMS.usf"), TEXT("PathTracingMainMS"), SF_RayMiss);
-
 DECLARE_GPU_STAT_NAMED(Stat_GPU_PathTracing, TEXT("Reference Path Tracing"));
 DECLARE_GPU_STAT_NAMED(Stat_GPU_PathTracingBuildSkyLightCDF, TEXT("Path Tracing: Build Sky Light CDF"));
 DECLARE_GPU_STAT_NAMED(Stat_GPU_PathTracingBuildVarianceMipTree, TEXT("Path Tracing: Build Variance Map Tree"));
@@ -485,8 +443,6 @@ void FDeferredShadingSceneRenderer::RenderPathTracing(FRHICommandListImmediate&
 	ClearUAV(RHICmdList, RayCountPerPixelRT->GetRenderTargetItem(), FLinearColor::Black);
 
 	auto RayGenShader = GetGlobalShaderMap(FeatureLevel)->GetShader<FPathTracingRG>();
-	auto MissShader = GetGlobalShaderMap(FeatureLevel)->GetShader<FPathTracingMS>();
-	auto ClosestHitShader = GetGlobalShaderMap(FeatureLevel)->GetShader<FPathTracingCHS>();
 
 	FRayTracingShaderBindingsWriter GlobalResources;
 
diff --git a/Engine/Source/Runtime/Renderer/Private/RayTracing/RayTracingAmbientOcclusion.cpp b/Engine/Source/Runtime/Renderer/Private/RayTracing/RayTracingAmbientOcclusion.cpp
--- a/Engine/Source/Runtime/Renderer/Private/RayTracing/RayTracingAmbientOcclusion.cpp
+++ b/Engine/Source/Runtime/Renderer/Private/RayTracing/RayTracingAmbientOcclusion.cpp
@@ -217,9 +217,6 @@ void FDeferredShadingSceneRenderer::RenderRayTracingAmbientOcclusion(
 			FRHIRayTracingShader* RayGenShaderTable[] = { RayGenerationShader->GetRayTracingShader() };
 			Initializer.SetRayGenShaderTable(RayGenShaderTable);
 
-			FRHIRayTracingShader* MissShaderTable[] = { View.ShaderMap->GetShader<FDefaultMaterialMS>()->GetRayTracingShader() };
-			Initializer.SetMissShaderTable(MissShaderTable);
-
 			FRHIRayTracingShader* HitGroupTable[] = { View.ShaderMap->GetShader<FOpaqueShadowHitGroup>()->GetRayTracingShader() };
 			Initializer.SetHitGroupTable(HitGroupTable);
 			Initializer.bAllowHitGroupIndexing = false; // Use the same hit shader for all geometry in the scene by disabling SBT indexing.
diff --git a/Engine/Source/Runtime/Renderer/Private/RayTracing/RayTracingBarycentrics.cpp b/Engine/Source/Runtime/Renderer/Private/RayTracing/RayTracingBarycentrics.cpp
--- a/Engine/Source/Runtime/Renderer/Private/RayTracing/RayTracingBarycentrics.cpp
+++ b/Engine/Source/Runtime/Renderer/Private/RayTracing/RayTracingBarycentrics.cpp
@@ -28,26 +28,6 @@ class FRayTracingBarycentricsRGS : public FGlobalShader
 };
 IMPLEMENT_GLOBAL_SHADER(FRayTracingBarycentricsRGS, "/Engine/Private/RayTracing/RayTracingBarycentrics.usf", "RayTracingBarycentricsMainRGS", SF_RayGen);
 
-// Example ray miss shader
-class FRayTracingBarycentricsMS : public FGlobalShader
-{
-	DECLARE_SHADER_TYPE(FRayTracingBarycentricsMS, Global);
-
-public:
-
-	static bool ShouldCompilePermutation(const FGlobalShaderPermutationParameters& Parameters)
-	{
-		return ShouldCompileRayTracingShadersForProject(Parameters.Platform);
-	}
-
-	FRayTracingBarycentricsMS() = default;
-	FRayTracingBarycentricsMS(const ShaderMetaType::CompiledShaderInitializerType& Initializer)
-		: FGlobalShader(Initializer)
-	{}
-};
-IMPLEMENT_SHADER_TYPE(, FRayTracingBarycentricsMS, TEXT("/Engine/Private/RayTracing/RayTracingBarycentrics.usf"), TEXT("RayTracingBarycentricsMainMS"), SF_RayMiss);
-
-
 // Example closest hit shader
 class FRayTracingBarycentricsCHS : public FGlobalShader
 {
@@ -82,16 +62,12 @@ void FDeferredShadingSceneRenderer::RenderRayTracingBarycentrics(FRHICommandList
 
 	auto RayGenShader = ShaderMap->GetShader<FRayTracingBarycentricsRGS>();
 	auto ClosestHitShader = ShaderMap->GetShader<FRayTracingBarycentricsCHS>();
-	auto MissShader = ShaderMap->GetShader<FRayTracingBarycentricsMS>();
 
 	FRayTracingPipelineStateInitializer Initializer;
 
 	FRHIRayTracingShader* RayGenShaderTable[] = { RayGenShader->GetRayTracingShader() };
 	Initializer.SetRayGenShaderTable(RayGenShaderTable);
 
-	FRHIRayTracingShader* MissShaderTable[] = { MissShader->GetRayTracingShader() };
-	Initializer.SetMissShaderTable(MissShaderTable);
-
 	FRHIRayTracingShader* HitGroupTable[] = { ClosestHitShader->GetRayTracingShader() };
 	Initializer.SetHitGroupTable(HitGroupTable);
 	Initializer.bAllowHitGroupIndexing = false; // Use the same hit shader for all geometry in the scene by disabling SBT indexing.
diff --git a/Engine/Source/Runtime/Renderer/Private/RayTracing/RayTracingDebug.cpp b/Engine/Source/Runtime/Renderer/Private/RayTracing/RayTracingDebug.cpp
--- a/Engine/Source/Runtime/Renderer/Private/RayTracing/RayTracingDebug.cpp
+++ b/Engine/Source/Runtime/Renderer/Private/RayTracing/RayTracingDebug.cpp
@@ -34,25 +34,6 @@ class FRayTracingDebugRGS : public FGlobalShader
 };
 IMPLEMENT_GLOBAL_SHADER(FRayTracingDebugRGS, "/Engine/Private/RayTracing/RayTracingDebug.usf", "RayTracingDebugMainRGS", SF_RayGen);
 
-class FRayTracingDebugMS : public FGlobalShader
-{
-	DECLARE_SHADER_TYPE(FRayTracingDebugMS, Global);
-
-public:
-
-	static bool ShouldCompilePermutation(const FGlobalShaderPermutationParameters& Parameters)
-	{
-		return ShouldCompileRayTracingShadersForProject(Parameters.Platform);
-	}
-
-	FRayTracingDebugMS() = default;
-	FRayTracingDebugMS(const ShaderMetaType::CompiledShaderInitializerType& Initializer)
-		: FGlobalShader(Initializer)
-	{}
-};
-IMPLEMENT_SHADER_TYPE(, FRayTracingDebugMS, TEXT("/Engine/Private/RayTracing/RayTracingDebug.usf"), TEXT("RayTracingDebugMainMS"), SF_RayMiss);
-
-
 class FRayTracingDebugCHS : public FGlobalShader
 {
 	DECLARE_SHADER_TYPE(FRayTracingDebugCHS, Global);
@@ -122,7 +103,6 @@ void FDeferredShadingSceneRenderer::RenderRayTracingDebug(FRHICommandListImmedia
 
 	auto RayGenShader = ShaderMap->GetShader<FRayTracingDebugRGS>();
 	auto ClosestHitShader = ShaderMap->GetShader<FRayTracingDebugCHS>();
-	auto MissShader = ShaderMap->GetShader<FRayTracingDebugMS>();
 
 	FRayTracingPipelineState* Pipeline = View.RayTracingMaterialPipeline;
 
diff --git a/Engine/Source/Runtime/Renderer/Private/RayTracing/RayTracingDeferredMaterials.cpp b/Engine/Source/Runtime/Renderer/Private/RayTracing/RayTracingDeferredMaterials.cpp
--- a/Engine/Source/Runtime/Renderer/Private/RayTracing/RayTracingDeferredMaterials.cpp
+++ b/Engine/Source/Runtime/Renderer/Private/RayTracing/RayTracingDeferredMaterials.cpp
@@ -24,21 +24,7 @@ class FRayTracingDeferredMaterialCHS : public FGlobalShader
 	using FParameters = FEmptyShaderParameters;
 };
 
-class FRayTracingDeferredMaterialMS : public FGlobalShader
-{
-	DECLARE_GLOBAL_SHADER(FRayTracingDeferredMaterialMS)
-	SHADER_USE_ROOT_PARAMETER_STRUCT(FRayTracingDeferredMaterialMS, FGlobalShader)
-
-		static bool ShouldCompilePermutation(const FGlobalShaderPermutationParameters& Parameters)
-	{
-		return ShouldCompileRayTracingShadersForProject(Parameters.Platform);
-	}
-
-	using FParameters = FEmptyShaderParameters;
-};
-
 IMPLEMENT_GLOBAL_SHADER(FRayTracingDeferredMaterialCHS, "/Engine/Private/RayTracing/RayTracingDeferredMaterials.usf", "DeferredMaterialCHS", SF_RayHitGroup);
-IMPLEMENT_GLOBAL_SHADER(FRayTracingDeferredMaterialMS,  "/Engine/Private/RayTracing/RayTracingDeferredMaterials.usf", "DeferredMaterialMS",  SF_RayMiss);
 
 FRayTracingPipelineState* FDeferredShadingSceneRenderer::BindRayTracingDeferredMaterialGatherPipeline(FRHICommandList& RHICmdList, const FViewInfo& View, FRHIRayTracingShader* RayGenShader)
 {
@@ -49,10 +35,6 @@ FRayTracingPipelineState* FDeferredShadingSceneRenderer::BindRayTracingDeferredM
 	FRHIRayTracingShader* RayGenShaderTable[] = { RayGenShader };
 	Initializer.SetRayGenShaderTable(RayGenShaderTable);
 
-	auto* MissShader = View.ShaderMap->GetShader<FRayTracingDeferredMaterialMS>();
-	FRHIRayTracingShader* MissShaderTable[] = { MissShader->GetRayTracingShader() };
-	Initializer.SetMissShaderTable(MissShaderTable);
-
 	Initializer.MaxPayloadSizeInBytes = 12; // sizeof FDeferredMaterialPayload
 
 	// Get the ray tracing materials
diff --git a/Engine/Source/Runtime/Renderer/Private/RayTracing/RayTracingGlobalIllumination.cpp b/Engine/Source/Runtime/Renderer/Private/RayTracing/RayTracingGlobalIllumination.cpp
--- a/Engine/Source/Runtime/Renderer/Private/RayTracing/RayTracingGlobalIllumination.cpp
+++ b/Engine/Source/Runtime/Renderer/Private/RayTracing/RayTracingGlobalIllumination.cpp
@@ -294,22 +294,8 @@ class FRayTracingGlobalIlluminationCHS : public FGlobalShader
 	using FParameters = FEmptyShaderParameters;
 };
 
-class FRayTracingGlobalIlluminationMS : public FGlobalShader
-{
-	DECLARE_GLOBAL_SHADER(FRayTracingGlobalIlluminationMS)
-	SHADER_USE_ROOT_PARAMETER_STRUCT(FRayTracingGlobalIlluminationMS, FGlobalShader)
-
-	static bool ShouldCompilePermutation(const FGlobalShaderPermutationParameters& Parameters)
-	{
-		return ShouldCompileRayTracingShadersForProject(Parameters.Platform);
-	}
-
-	using FParameters = FEmptyShaderParameters;
-};
-
 IMPLEMENT_GLOBAL_SHADER(FGlobalIlluminationRGS, "/Engine/Private/RayTracing/RayTracingGlobalIlluminationRGS.usf", "GlobalIlluminationRGS", SF_RayGen);
 IMPLEMENT_GLOBAL_SHADER(FRayTracingGlobalIlluminationCHS, "/Engine/Private/RayTracing/RayTracingGlobalIlluminationRGS.usf", "RayTracingGlobalIlluminationCHS", SF_RayHitGroup);
-IMPLEMENT_GLOBAL_SHADER(FRayTracingGlobalIlluminationMS, "/Engine/Private/RayTracing/RayTracingGlobalIlluminationRGS.usf", "RayTracingGlobalIlluminationMS", SF_RayMiss);
 IMPLEMENT_GLOBAL_SHADER(FRayTracingGlobalIlluminationCompositePS, "/Engine/Private/RayTracing/RayTracingGlobalIlluminationCompositePS.usf", "GlobalIlluminationCompositePS", SF_Pixel);
 IMPLEMENT_GLOBAL_SHADER(FRayTracingGlobalIlluminationSceneColorCompositePS, "/Engine/Private/RayTracing/RayTracingGlobalIlluminationCompositePS.usf", "GlobalIlluminationSceneColorCompositePS", SF_Pixel);
 
diff --git a/Engine/Source/Runtime/Renderer/Private/RayTracing/RayTracingMaterialHitShaders.cpp b/Engine/Source/Runtime/Renderer/Private/RayTracing/RayTracingMaterialHitShaders.cpp
--- a/Engine/Source/Runtime/Renderer/Private/RayTracing/RayTracingMaterialHitShaders.cpp
+++ b/Engine/Source/Runtime/Renderer/Private/RayTracing/RayTracingMaterialHitShaders.cpp
@@ -157,7 +157,6 @@ IMPLEMENT_MATERIALCHS_TYPE(TUniformLightMapPolicy<LMP_DISTANCE_FIELD_SHADOWS_AND
 
 IMPLEMENT_GLOBAL_SHADER(FHiddenMaterialHitGroup, "/Engine/Private/RayTracing/RayTracingMaterialDefaultHitShaders.usf", "closesthit=HiddenMaterialCHS anyhit=HiddenMaterialAHS", SF_RayHitGroup);
 IMPLEMENT_GLOBAL_SHADER(FOpaqueShadowHitGroup, "/Engine/Private/RayTracing/RayTracingMaterialDefaultHitShaders.usf", "closesthit=OpaqueShadowCHS", SF_RayHitGroup);
-IMPLEMENT_GLOBAL_SHADER(FDefaultMaterialMS, "/Engine/Private/RayTracing/RayTracingMaterialDefaultHitShaders.usf", "DefaultMaterialMS", SF_RayMiss);
 
 template<typename LightMapPolicyType>
 static FMaterialCHS* GetMaterialHitShader(const FMaterial& RESTRICT MaterialResource, const FVertexFactory* VertexFactory, bool UseTextureLod)
@@ -441,7 +440,6 @@ FRayTracingPipelineState* FDeferredShadingSceneRenderer::BindRayTracingMaterialP
 	FRHICommandList& RHICmdList,
 	const FViewInfo& View,
 	const TArrayView<FRHIRayTracingShader*>& RayGenShaderTable,
-	FRHIRayTracingShader* MissShader,
 	FRHIRayTracingShader* DefaultClosestHitShader
 )
 {
@@ -456,9 +454,6 @@ FRayTracingPipelineState* FDeferredShadingSceneRenderer::BindRayTracingMaterialP
 
 	Initializer.SetRayGenShaderTable(RayGenShaderTable);
 
-	FRHIRayTracingShader* MissShaderTable[] = { MissShader };
-	Initializer.SetMissShaderTable(MissShaderTable);
-
 	const bool bEnableMaterials = GEnableRayTracingMaterials != 0;
 
 	TArray<FRHIRayTracingShader*> RayTracingMaterialLibrary;
diff --git a/Engine/Source/Runtime/Renderer/Private/RayTracing/RayTracingMaterialHitShaders.h b/Engine/Source/Runtime/Renderer/Private/RayTracing/RayTracingMaterialHitShaders.h
--- a/Engine/Source/Runtime/Renderer/Private/RayTracing/RayTracingMaterialHitShaders.h
+++ b/Engine/Source/Runtime/Renderer/Private/RayTracing/RayTracingMaterialHitShaders.h
@@ -76,17 +76,4 @@ class FOpaqueShadowHitGroup : public FGlobalShader
 	using FParameters = FEmptyShaderParameters;
 };
 
-class FDefaultMaterialMS : public FGlobalShader
-{
-	DECLARE_GLOBAL_SHADER(FDefaultMaterialMS)
-	SHADER_USE_ROOT_PARAMETER_STRUCT(FDefaultMaterialMS, FGlobalShader)
-
-	static bool ShouldCompilePermutation(const FGlobalShaderPermutationParameters& Parameters)
-	{
-		return ShouldCompileRayTracingShadersForProject(Parameters.Platform);
-	}
-
-	using FParameters = FEmptyShaderParameters;
-};
-
 #endif // RHI_RAYTRACING
diff --git a/Engine/Source/Runtime/Renderer/Private/RayTracing/RayTracingReflections.cpp b/Engine/Source/Runtime/Renderer/Private/RayTracing/RayTracingReflections.cpp
--- a/Engine/Source/Runtime/Renderer/Private/RayTracing/RayTracingReflections.cpp
+++ b/Engine/Source/Runtime/Renderer/Private/RayTracing/RayTracingReflections.cpp
@@ -163,35 +163,7 @@ class FRayTracingReflectionsRGS : public FGlobalShader
 	}
 };
 
-class FRayTracingReflectionsCHS : public FGlobalShader
-{
-	DECLARE_GLOBAL_SHADER(FRayTracingReflectionsCHS);
-	SHADER_USE_ROOT_PARAMETER_STRUCT(FRayTracingReflectionsCHS, FGlobalShader)
-		
-	static bool ShouldCompilePermutation(const FGlobalShaderPermutationParameters& Parameters)
-	{
-		return ShouldCompileRayTracingShadersForProject(Parameters.Platform);
-	}
-	
-	using FParameters = FEmptyShaderParameters;
-};
-
-class FRayTracingReflectionsMS : public FGlobalShader
-{
-	DECLARE_GLOBAL_SHADER(FRayTracingReflectionsMS);
-	SHADER_USE_ROOT_PARAMETER_STRUCT(FRayTracingReflectionsMS, FGlobalShader)
-		
-	static bool ShouldCompilePermutation(const FGlobalShaderPermutationParameters& Parameters)
-	{
-		return ShouldCompileRayTracingShadersForProject(Parameters.Platform);
-	}
-
-	using FParameters = FEmptyShaderParameters;
-};
-
 IMPLEMENT_GLOBAL_SHADER(FRayTracingReflectionsRGS, "/Engine/Private/RayTracing/RayTracingReflections.usf", "RayTracingReflectionsRGS", SF_RayGen);
-IMPLEMENT_GLOBAL_SHADER(FRayTracingReflectionsCHS, "/Engine/Private/RayTracing/RayTracingReflections.usf", "RayTracingReflectionsMainCHS", SF_RayHitGroup);
-IMPLEMENT_GLOBAL_SHADER(FRayTracingReflectionsMS, "/Engine/Private/RayTracing/RayTracingReflections.usf", "RayTracingReflectionsMainMS", SF_RayMiss);
 
 void FDeferredShadingSceneRenderer::PrepareRayTracingReflections(const FViewInfo& View, TArray<FRHIRayTracingShader*>& OutRayGenShaders)
 {
diff --git a/Engine/Source/Runtime/Renderer/Private/RayTracing/RayTracingShadows.cpp b/Engine/Source/Runtime/Renderer/Private/RayTracing/RayTracingShadows.cpp
--- a/Engine/Source/Runtime/Renderer/Private/RayTracing/RayTracingShadows.cpp
+++ b/Engine/Source/Runtime/Renderer/Private/RayTracing/RayTracingShadows.cpp
@@ -226,9 +226,6 @@ void FDeferredShadingSceneRenderer::RenderRayTracingShadows(
 				FRHIRayTracingShader* RayGenShaderTable[] = { RayGenerationShader->GetRayTracingShader() };
 				Initializer.SetRayGenShaderTable(RayGenShaderTable);
 
-				FRHIRayTracingShader* MissShaderTable[] = { View.ShaderMap->GetShader<FDefaultMaterialMS>()->GetRayTracingShader() };
-				Initializer.SetMissShaderTable(MissShaderTable);
-
 				FRHIRayTracingShader* HitGroupTable[] = { View.ShaderMap->GetShader<FOpaqueShadowHitGroup>()->GetRayTracingShader() };
 				Initializer.SetHitGroupTable(HitGroupTable);
 				Initializer.bAllowHitGroupIndexing = false; // Use the same hit shader for all geometry in the scene by disabling SBT indexing.
diff --git a/Engine/Source/Runtime/Renderer/Private/RayTracing/RayTracingTranslucency.cpp b/Engine/Source/Runtime/Renderer/Private/RayTracing/RayTracingTranslucency.cpp
--- a/Engine/Source/Runtime/Renderer/Private/RayTracing/RayTracingTranslucency.cpp
+++ b/Engine/Source/Runtime/Renderer/Private/RayTracing/RayTracingTranslucency.cpp
@@ -172,23 +172,8 @@ class FRayTracingTranslucencyCHS : public FGlobalShader
 	using FParameters = FEmptyShaderParameters;
 };
 
-class FRayTracingTranslucencyMS : public FGlobalShader
-{
-	DECLARE_GLOBAL_SHADER(FRayTracingTranslucencyMS);
-	SHADER_USE_ROOT_PARAMETER_STRUCT(FRayTracingTranslucencyMS, FGlobalShader)
-
-		static bool ShouldCompilePermutation(const FGlobalShaderPermutationParameters& Parameters)
-	{
-		return ShouldCompileRayTracingShadersForProject(Parameters.Platform);
-	}
-
-	using FParameters = FEmptyShaderParameters;
-};
-
 IMPLEMENT_GLOBAL_SHADER(FRayTracingTranslucencyRGS, "/Engine/Private/RayTracing/RayTracingTranslucency.usf", "RayTracingTranslucencyRGS", SF_RayGen);
 IMPLEMENT_GLOBAL_SHADER(FRayTracingTranslucencyCHS, "/Engine/Private/RayTracing/RayTracingTranslucency.usf", "RayTracingTranslucencyMainCHS", SF_RayHitGroup);
-IMPLEMENT_GLOBAL_SHADER(FRayTracingTranslucencyMS, "/Engine/Private/RayTracing/RayTracingTranslucency.usf", "RayTracingTranslucencyMainMS", SF_RayMiss);
-
 
 //#dxr-todo: should we unify it with the composition happening in the non raytraced translucency pass? In that case it should use FCopySceneColorPS
 // Probably, but the architecture depends on the denoiser -> discuss
diff --git a/Engine/Source/Runtime/Renderer/Private/RayTracing/RaytracingSkylight.cpp b/Engine/Source/Runtime/Renderer/Private/RayTracing/RaytracingSkylight.cpp
--- a/Engine/Source/Runtime/Renderer/Private/RayTracing/RaytracingSkylight.cpp
+++ b/Engine/Source/Runtime/Renderer/Private/RayTracing/RaytracingSkylight.cpp
@@ -886,9 +886,6 @@ void FDeferredShadingSceneRenderer::RenderRayTracingSkyLight(
 				FRHIRayTracingShader* RayGenShaderTable[] = { RayGenerationShader->GetRayTracingShader() };
 				Initializer.SetRayGenShaderTable(RayGenShaderTable);
 
-				FRHIRayTracingShader* MissShaderTable[] = { View.ShaderMap->GetShader<FDefaultMaterialMS>()->GetRayTracingShader() };
-				Initializer.SetMissShaderTable(MissShaderTable);
-
 				FRHIRayTracingShader* HitGroupTable[] = { View.ShaderMap->GetShader<FOpaqueShadowHitGroup>()->GetRayTracingShader() };
 				Initializer.SetHitGroupTable(HitGroupTable);
 				Initializer.bAllowHitGroupIndexing = false; // Use the same hit shader for all geometry in the scene by disabling SBT indexing.
diff --git a/Engine/Source/Runtime/Renderer/Private/Tests/RayTracingTestbed.cpp b/Engine/Source/Runtime/Renderer/Private/Tests/RayTracingTestbed.cpp
--- a/Engine/Source/Runtime/Renderer/Private/Tests/RayTracingTestbed.cpp
+++ b/Engine/Source/Runtime/Renderer/Private/Tests/RayTracingTestbed.cpp
@@ -89,7 +89,7 @@ bool RunRayTracingTestbed_RenderThread(const FString& Parameters)
 
 	{
 		FRHIResourceCreateInfo CreateInfo;
-		IntersectionResultBuffer = RHICreateVertexBuffer(sizeof(FBasicRayIntersectionData)*NumRays, BUF_Static | BUF_UnorderedAccess, CreateInfo);
+		IntersectionResultBuffer = RHICreateVertexBuffer(sizeof(FIntersectionPayload)*NumRays, BUF_Static | BUF_UnorderedAccess, CreateInfo);
 		IntersectionResultBufferView = RHICreateUnorderedAccessView(IntersectionResultBuffer, PF_R32_UINT);
 	}
 
@@ -151,23 +151,24 @@ bool RunRayTracingTestbed_RenderThread(const FString& Parameters)
 		// Read back and validate intersection trace results
 
 		{
-			auto MappedResults = (const FBasicRayIntersectionData*)RHILockVertexBuffer(IntersectionResultBuffer, 0, sizeof(FBasicRayIntersectionData)*NumRays, RLM_ReadOnly);
+			auto MappedResults = (const FIntersectionPayload*)RHILockVertexBuffer(IntersectionResultBuffer, 0, sizeof(FIntersectionPayload)*NumRays, RLM_ReadOnly);
 
 			check(MappedResults);
 
 			// expect hit primitive 0, instance 0, barycentrics {0.5, 0.125}
+			check(MappedResults[0].HitT >= 0);
 			check(MappedResults[0].PrimitiveIndex == 0);
 			check(MappedResults[0].InstanceIndex == 0);
 			check(FMath::IsNearlyEqual(MappedResults[0].Barycentrics[0], 0.5f));
 			check(FMath::IsNearlyEqual(MappedResults[0].Barycentrics[1], 0.125f));
 
-			check(MappedResults[1].PrimitiveIndex == ~0u); // expect miss
-			check(MappedResults[2].PrimitiveIndex == ~0u); // expect miss
-			check(MappedResults[3].PrimitiveIndex == ~0u); // expect miss
+			check(MappedResults[1].HitT < 0); // expect miss
+			check(MappedResults[2].HitT < 0); // expect miss
+			check(MappedResults[3].HitT < 0); // expect miss
 
 			RHIUnlockVertexBuffer(IntersectionResultBuffer);
 
-			bIntersectionTestOK = (MappedResults[1].PrimitiveIndex == ~0u) && (MappedResults[2].PrimitiveIndex == ~0u) && (MappedResults[3].PrimitiveIndex == ~0u);
+			bIntersectionTestOK = (MappedResults[0].HitT >= 0) && (MappedResults[1].HitT < 0) && (MappedResults[2].HitT < 0) && (MappedResults[3].HitT < 0);
 		}
 	}
 
diff --git a/Engine/Source/Runtime/UtilityShaders/Private/BuiltInRayTracingShaders.cpp b/Engine/Source/Runtime/UtilityShaders/Private/BuiltInRayTracingShaders.cpp
--- a/Engine/Source/Runtime/UtilityShaders/Private/BuiltInRayTracingShaders.cpp
+++ b/Engine/Source/Runtime/UtilityShaders/Private/BuiltInRayTracingShaders.cpp
@@ -5,9 +5,7 @@
 #if RHI_RAYTRACING
 
 IMPLEMENT_GLOBAL_SHADER( FOcclusionMainRG,		"/Engine/Private/RayTracing/RayTracingBuiltInShaders.usf",		"OcclusionMainRG",				SF_RayGen);
-IMPLEMENT_SHADER_TYPE(, FOcclusionMainMS,		TEXT("/Engine/Private/RayTracing/RayTracingBuiltInShaders.usf"), TEXT("OcclusionMainMS"),		SF_RayMiss);
 IMPLEMENT_GLOBAL_SHADER( FIntersectionMainRG,	"/Engine/Private/RayTracing/RayTracingBuiltInShaders.usf",		"IntersectionMainRG",			SF_RayGen);
-IMPLEMENT_SHADER_TYPE(, FIntersectionMainMS,	TEXT("/Engine/Private/RayTracing/RayTracingBuiltInShaders.usf"), TEXT("IntersectionMainMS"),	SF_RayMiss);
 IMPLEMENT_SHADER_TYPE(, FIntersectionMainCHS,	TEXT("/Engine/Private/RayTracing/RayTracingBuiltInShaders.usf"), TEXT("IntersectionMainCHS"),	SF_RayHitGroup);
 IMPLEMENT_SHADER_TYPE(, FDefaultMainCHS,		TEXT("/Engine/Private/RayTracing/RayTracingBuiltInShaders.usf"), TEXT("DefaultMainCHS"),		SF_RayHitGroup);
 IMPLEMENT_SHADER_TYPE(, FDefaultMainMS,			TEXT("/Engine/Private/RayTracing/RayTracingBuiltInShaders.usf"), TEXT("DefaultMainMS"),			SF_RayMiss);
diff --git a/Engine/Source/Runtime/UtilityShaders/Public/BuiltInRayTracingShaders.h b/Engine/Source/Runtime/UtilityShaders/Public/BuiltInRayTracingShaders.h
--- a/Engine/Source/Runtime/UtilityShaders/Public/BuiltInRayTracingShaders.h
+++ b/Engine/Source/Runtime/UtilityShaders/Public/BuiltInRayTracingShaders.h
@@ -39,17 +39,6 @@ class UTILITYSHADERS_API FOcclusionMainRG : public FBuiltInRayTracingShader
 	END_SHADER_PARAMETER_STRUCT()
 };
 
-class FOcclusionMainMS : public FBuiltInRayTracingShader
-{
-	DECLARE_EXPORTED_SHADER_TYPE(FOcclusionMainMS, Global, UTILITYSHADERS_API);
-public:
-
-	FOcclusionMainMS() = default;
-	FOcclusionMainMS(const ShaderMetaType::CompiledShaderInitializerType& Initializer)
-		: FBuiltInRayTracingShader(Initializer)
-	{}
-};
-
 class UTILITYSHADERS_API FIntersectionMainRG : public FBuiltInRayTracingShader
 {
 	DECLARE_GLOBAL_SHADER(FIntersectionMainRG);
@@ -58,21 +47,10 @@ class UTILITYSHADERS_API FIntersectionMainRG : public FBuiltInRayTracingShader
 	BEGIN_SHADER_PARAMETER_STRUCT(FParameters, )
 		SHADER_PARAMETER_RDG_BUFFER_SRV(RaytracingAccelerationStructure, TLAS)
 		SHADER_PARAMETER_RDG_BUFFER_SRV(StructuredBuffer<FBasicRayData>, Rays)
-		SHADER_PARAMETER_RDG_BUFFER_UAV(RWStructuredBuffer<FBasicRayIntersectionData>, IntersectionOutput)
+		SHADER_PARAMETER_RDG_BUFFER_UAV(RWStructuredBuffer<FIntersectionPayload>, IntersectionOutput)
 	END_SHADER_PARAMETER_STRUCT()
 };
 
-class FIntersectionMainMS : public FBuiltInRayTracingShader
-{
-	DECLARE_EXPORTED_SHADER_TYPE(FIntersectionMainMS, Global, UTILITYSHADERS_API);
-public:
-
-	FIntersectionMainMS() = default;
-	FIntersectionMainMS(const ShaderMetaType::CompiledShaderInitializerType& Initializer)
-		: FBuiltInRayTracingShader(Initializer)
-	{}
-};
-
 class FIntersectionMainCHS : public FBuiltInRayTracingShader
 {
 	DECLARE_EXPORTED_SHADER_TYPE(FIntersectionMainCHS, Global, UTILITYSHADERS_API);
