Ray tracing translucency now relies on Payload.Hitkind to refract the ray

diff --git a/Engine/Shaders/Private/RayTracing/RayTracingTranslucency.usf b/Engine/Shaders/Private/RayTracing/RayTracingTranslucency.usf
--- a/Engine/Shaders/Private/RayTracing/RayTracingTranslucency.usf
+++ b/Engine/Shaders/Private/RayTracing/RayTracingTranslucency.usf
@@ -60,12 +60,9 @@ float FresnelDielectric(float Eta, float IoH, float ToH)
 	return (Rs + Rp) / 2;
 }
 
-bool RefractRay(float3 V, float3 N, float Ior1, float Ior2, out float3 L, out float EventThroughput)
+bool RefractRay(float3 V, float3 N, float Ior1, float Ior2, bool bIsEntering, out float3 L, out float EventThroughput)
 {
 	float Eta = Ior1 / Ior2;
-
-	//#dxr_todo: check HitKind functionality as bool bIsEntering = Payload.HitKind == HIT_KIND_TRIANGLE_FRONT_FACE;
-	bool bIsEntering = (dot(V, N) < 0.0);
 	if (bIsEntering)
 	{
 		Eta = 1.0f / Ior1;
@@ -229,7 +226,9 @@ void RayTracingTranslucencyRGS()
 			float Ior2 = LastIor;
 			bHasScattered |= Ior1 != Ior2;
 			float ReflectionRefractionEventThroughput;
-			if (!RefractRay(Ray.Direction, Payload.WorldNormal, Ior1, Ior2, RefractedDirection, ReflectionRefractionEventThroughput))
+			bool bIsEntering = Payload.HitKind == HIT_KIND_TRIANGLE_FRONT_FACE;
+
+			if (!RefractRay(Ray.Direction, Payload.WorldNormal, Ior1, Ior2, bIsEntering, RefractedDirection, ReflectionRefractionEventThroughput))
 			{
 				//#dxr_todo: handle total internal reflection ray event. We could just continue form this path. And not use the new IOR.
 				// We should combine that step with the one above reflection evaluation. For now we stop tracing.
