Caching texture resolution to remove GetDimensions() call when evaluating SkyLight texture.

diff --git a/Engine/Shaders/Private/RayTracing/RayTracingSkyLightRGS.usf b/Engine/Shaders/Private/RayTracing/RayTracingSkyLightRGS.usf
--- a/Engine/Shaders/Private/RayTracing/RayTracingSkyLightRGS.usf
+++ b/Engine/Shaders/Private/RayTracing/RayTracingSkyLightRGS.usf
@@ -29,10 +29,8 @@ float PdfSkyLightOcclusionRay(
 {
 	// Determine primary direction
 	uint MipCount = log2(SkyLight.MipDimensions.x);
-	uint2 TextureRes;
-	SkyLight.Texture.GetDimensions(TextureRes.x, TextureRes.y);
-	uint3 TextureCoord = GetTextureCubeCoordinate(WorldDirection, TextureRes);
-	float2 TextureUV = TextureCoord.xy / (float2)TextureRes;
+	uint3 TextureCoord = GetTextureCubeCoordinate(WorldDirection, SkyLight.TextureDimensions);
+	float2 TextureUV = TextureCoord.xy / (float2)SkyLight.TextureDimensions;
 	uint2 MipCoord = TextureUV * SkyLight.MipDimensions.xy;
 	float FacePdf = PdfFace(FaceCdf, TextureCoord.z);
 
@@ -166,9 +164,7 @@ bool GenerateSkyLightOcclusionRay(
 	float2 MipPixelUV = (MipPixel + PixelOffset) / float2(DimensionsAtLevel(SkyLight.MipDimensions, StopLevel));
 
 	// Construct Cubemap coordinate
-	uint2 TextureDimensions;
-	SkyLight.Texture.GetDimensions(TextureDimensions.x, TextureDimensions.y);
-	float3 TexelCoord = float3(MipPixelUV * TextureDimensions, FaceIndex);
+	float3 TexelCoord = float3(MipPixelUV * SkyLight.TextureDimensions, FaceIndex);
 
 	// Determine area pdf
 #if USE_SOLID_ANGLE_TEXTURE
@@ -179,7 +175,7 @@ bool GenerateSkyLightOcclusionRay(
 #endif
 
 	RayOrigin = WorldPosition;
-	RayDirection = normalize(GetTextureCubeVector(TexelCoord, TextureDimensions));
+	RayDirection = normalize(GetTextureCubeVector(TexelCoord, SkyLight.TextureDimensions));
 	RayTMin = 0.0;
 	RayTMax = SkyLight.MaxRayDistance;
 	RayPdf = FacePdf * MipPdf * AreaPdf;
@@ -255,7 +251,8 @@ void SkyLightRGS()
 	float RayDistance = 0.0;
 	float HitCount = 0.0;
 	uint SamplesPerPixel = SkyLight.SamplesPerPixel;
-	uint FrameIndex = View.StateFrameIndex % SamplesPerPixel;
+	//uint FrameIndex = View.StateFrameIndex % SamplesPerPixel;
+	uint FrameIndex = 0;
 
 	// Mask out depth values that are infinitely far away
 	float DeviceZ = SceneTexturesStruct.SceneDepthTexture.Load(int3(PixelCoord, 0)).r;
diff --git a/Engine/Source/Runtime/Renderer/Private/RayTracing/RayTracingSkyLight.h b/Engine/Source/Runtime/Renderer/Private/RayTracing/RayTracingSkyLight.h
--- a/Engine/Source/Runtime/Renderer/Private/RayTracing/RayTracingSkyLight.h
+++ b/Engine/Source/Runtime/Renderer/Private/RayTracing/RayTracingSkyLight.h
@@ -17,6 +17,7 @@ SHADER_PARAMETER(FIntVector, MipDimensions)
 SHADER_PARAMETER(float, MaxNormalBias)
 SHADER_PARAMETER_TEXTURE(TextureCube, Texture)
 SHADER_PARAMETER_SAMPLER(SamplerState, TextureSampler)
+SHADER_PARAMETER(FIntVector, TextureDimensions)
 SHADER_PARAMETER_SRV(Buffer<float>, MipTreePosX)
 SHADER_PARAMETER_SRV(Buffer<float>, MipTreeNegX)
 SHADER_PARAMETER_SRV(Buffer<float>, MipTreePosY)
diff --git a/Engine/Source/Runtime/Renderer/Private/RayTracing/RaytracingSkylight.cpp b/Engine/Source/Runtime/Renderer/Private/RayTracing/RaytracingSkylight.cpp
--- a/Engine/Source/Runtime/Renderer/Private/RayTracing/RaytracingSkylight.cpp
+++ b/Engine/Source/Runtime/Renderer/Private/RayTracing/RaytracingSkylight.cpp
@@ -108,6 +108,7 @@ void SetupSkyLightParameters(
 	{
 		SkyLightData->Color = FVector(Scene.SkyLight->GetEffectiveLightColor());
 		SkyLightData->Texture = Scene.SkyLight->ProcessedTexture->TextureRHI;
+		SkyLightData->TextureDimensions = FIntVector(Scene.SkyLight->ProcessedTexture->GetSizeX(), Scene.SkyLight->ProcessedTexture->GetSizeY(), 0);
 		SkyLightData->TextureSampler = Scene.SkyLight->ProcessedTexture->SamplerStateRHI;
 		SkyLightData->MipDimensions = Scene.SkyLight->SkyLightMipDimensions;
 
