Revert accidental test code

diff --git a/Engine/Shaders/Private/RayTracing/RayTracingSkyLightRGS.usf b/Engine/Shaders/Private/RayTracing/RayTracingSkyLightRGS.usf
--- a/Engine/Shaders/Private/RayTracing/RayTracingSkyLightRGS.usf
+++ b/Engine/Shaders/Private/RayTracing/RayTracingSkyLightRGS.usf
@@ -251,8 +251,7 @@ void SkyLightRGS()
 	float RayDistance = 0.0;
 	float HitCount = 0.0;
 	uint SamplesPerPixel = SkyLight.SamplesPerPixel;
-	//uint FrameIndex = View.StateFrameIndex % SamplesPerPixel;
-	uint FrameIndex = 0;
+	uint FrameIndex = View.StateFrameIndex % SamplesPerPixel;
 
 	// Mask out depth values that are infinitely far away
 	float DeviceZ = SceneTexturesStruct.SceneDepthTexture.Load(int3(PixelCoord, 0)).r;
