PR #5872: Use frame buffer as shading cache when not refracting (Contributed by [REDACTED_PERSON])

diff --git a/Engine/Shaders/Private/RayTracing/RayTracingTranslucency.usf b/Engine/Shaders/Private/RayTracing/RayTracingTranslucency.usf
--- a/Engine/Shaders/Private/RayTracing/RayTracingTranslucency.usf
+++ b/Engine/Shaders/Private/RayTracing/RayTracingTranslucency.usf
@@ -156,6 +156,7 @@ void RayTracingTranslucencyRGS()
 			break;
 		}
 		float3 HitPoint = Ray.Origin + Ray.Direction * Payload.HitT;
+		float NextMaxRayDistance = Ray.TMax - Payload.HitT;
 
 		//
 		// Handle surface lighting
@@ -236,6 +237,9 @@ void RayTracingTranslucencyRGS()
 			}
 			LastIor = Ior1;
 			PathThroughput *= ReflectionRefractionEventThroughput;
+
+			// ray has bent, so it may need to go arbitrarily far
+			NextMaxRayDistance = LocalMaxRayDistance;
 		}
 
 
@@ -245,7 +249,7 @@ void RayTracingTranslucencyRGS()
 
 		Ray.Origin = HitPoint;
 		Ray.TMin = 0.01;
-		Ray.TMax = LocalMaxRayDistance;
+		Ray.TMax = NextMaxRayDistance;
 		Ray.Direction = RefractedDirection;
 		float SurfaceCurvature = 0.0f; /* #todo_dxr assume no curvature */
 		RayCone = PropagateRayCone(RayCone, SurfaceCurvature, Depth);
