PR #5874: Terminate reflection bounces on accumulated roughness (Contributed by [REDACTED_PERSON])

diff --git a/Engine/Shaders/Private/RayTracing/RayTracingReflections.usf b/Engine/Shaders/Private/RayTracing/RayTracingReflections.usf
--- a/Engine/Shaders/Private/RayTracing/RayTracingReflections.usf
+++ b/Engine/Shaders/Private/RayTracing/RayTracingReflections.usf
@@ -77,6 +77,7 @@ uint UpscaleFactor;
 uint SortTileSize;
 uint2 RayTracingResolution;
 uint2 TileAlignedResolution;
+uint TestPathRoughness;
 
 float ReflectionMinRayDistance;
 float ReflectionMaxRayDistance;
@@ -204,7 +205,8 @@ void RayTracingReflectionsRGS()
 		for (uint SampleIndex = 0; SampleIndex < LocalSamplesPerPixel; ++SampleIndex)
 		{
 			bool bIsValidSample = false;
-			
+			bool bNeedsCapture = true;
+
 			float3 PathRadiance = 0;
 			float3 PathThroughput = float3(1.0f, 1.0f, 1.0f);
 			
@@ -219,6 +221,9 @@ void RayTracingReflectionsRGS()
 			float  TopLayerRoughness = ScreenSpaceData.GBuffer.Roughness;
 			float3 TopLayerSpecularColor = ScreenSpaceData.GBuffer.SpecularColor;
 
+			// Track Roughness to terminate bounces early
+			float AccumRoughness = TopLayerRoughness;
+
 			float  ClearCoat = 1.0f;
 			float  BottomLayerRoughness = TopLayerRoughness;
 			float3 BottomLayerSpecularColor = TopLayerSpecularColor;
@@ -471,11 +476,26 @@ void RayTracingReflectionsRGS()
 				}
 
 				if (all(PathThroughput < 0.001) || bSkyWasHit || HitDistance<0.0f) break;
+
+				if (TestPathRoughness)
+				{
+					//Accumulate roughness for this bounce
+					AccumRoughness = 1.0f - (1.0f - AccumRoughness)*(1.0f - TopLayerRoughness);
+
+					float RoughnessFade = GetRoughnessFade(AccumRoughness, ReflectionMaxRoughness);
+
+					if (RoughnessFade <= 0.0f)
+					{
+						bNeedsCapture = true;
+						break;
+					}
+				}
 			}
 			
 			if (bIsValidSample)
 			{
-				if (UseReflectionCaptures && BounceIndex == MaxBounces && !(bSkyWasHit || HitDistance < 0.0f))
+				bNeedsCapture |= BounceIndex == MaxBounces && !(bSkyWasHit || HitDistance < 0.0f);
+				if (UseReflectionCaptures && bNeedsCapture)
 				{
 					// We have reached the end of the bounce count on a surface. 
 					// So accumulate reflection from captures for the top layer for a better visual result.
diff --git a/Engine/Source/Runtime/Renderer/Private/RayTracing/RayTracingReflections.cpp b/Engine/Source/Runtime/Renderer/Private/RayTracing/RayTracingReflections.cpp
--- a/Engine/Source/Runtime/Renderer/Private/RayTracing/RayTracingReflections.cpp
+++ b/Engine/Source/Runtime/Renderer/Private/RayTracing/RayTracingReflections.cpp
@@ -108,6 +108,12 @@ static TAutoConsoleVariable<int32> CVarRayTracingReflectionsSortSize(
 	TEXT("5: 4096 Elements (Default)\n"),
 	ECVF_RenderThreadSafe);
 
+static TAutoConsoleVariable<int32> CVarRayTracingReflectionsTestPathRoughness(
+	TEXT("r.RayTracing.Reflections.TestPathRoughness"),
+	1,
+	TEXT("Accumulate roughness along path and test accumulated roughness against MaxRoughness before launching the next bounce (default 1)"),
+	ECVF_RenderThreadSafe);
+
 static const int32 GReflectionLightCountMaximum = 64;
 
 class FRayTracingReflectionsRGS : public FGlobalShader
@@ -137,6 +143,7 @@ class FRayTracingReflectionsRGS : public FGlobalShader
 		SHADER_PARAMETER(float, ReflectionMaxRayDistance)
 		SHADER_PARAMETER(float, ReflectionMaxRoughness)
 		SHADER_PARAMETER(float, ReflectionMaxNormalBias)
+		SHADER_PARAMETER(int32, TestPathRoughness)
 
 		SHADER_PARAMETER_SRV(RaytracingAccelerationStructure, TLAS)
 		SHADER_PARAMETER_SRV(StructuredBuffer<FRTLightingData>, LightDataBuffer)
@@ -253,6 +260,7 @@ void FDeferredShadingSceneRenderer::RenderRayTracingReflections(
 	CommonParameters.ReflectionMaxNormalBias = GetRaytracingMaxNormalBias();
 	CommonParameters.RayTracingResolution = RayTracingResolution;
 	CommonParameters.TileAlignedResolution = TileAlignedResolution;
+	CommonParameters.TestPathRoughness = CVarRayTracingReflectionsTestPathRoughness.GetValueOnRenderThread();
 
 	CommonParameters.TLAS = View.RayTracingScene.RayTracingSceneRHI->GetShaderResourceView();
 	CommonParameters.ViewUniformBuffer = View.ViewUniformBuffer;
