Fixed clear coat materials  in ray traced reflections did not match raster.

PR #5888: Rtx clearcoat fixes (Contributed by [REDACTED_PERSON])

diff --git a/Engine/Shaders/Private/RayTracing/RayTracingReflections.usf b/Engine/Shaders/Private/RayTracing/RayTracingReflections.usf
--- a/Engine/Shaders/Private/RayTracing/RayTracingReflections.usf
+++ b/Engine/Shaders/Private/RayTracing/RayTracingReflections.usf
@@ -78,6 +78,8 @@ uint SortTileSize;
 uint2 RayTracingResolution;
 uint2 TileAlignedResolution;
 uint TestPathRoughness;
+uint MinClearCoatLevel;
+uint MaxUnderCoatBounces;
 
 float ReflectionMinRayDistance;
 float ReflectionMaxRayDistance;
@@ -256,6 +258,24 @@ void RayTracingReflectionsRGS()
 			bool bSkyWasHit = false;
 			for (; BounceIndex < MaxBounces; ++BounceIndex)
 			{
+				// Composite pass applies BRDF, except
+				bool bApplyBRDF = BounceIndex != 0;
+
+				BRANCH
+				if (ShadingModelID == SHADINGMODELID_CLEAR_COAT && ClearCoat < MinClearCoatLevel)
+				{
+					// top layer isn't really contributing, just trace bottom
+					TopLayerRoughness = BottomLayerRoughness;
+					TopLayerSpecularColor = BottomLayerSpecularColor;
+					TopLayerWorldNormal = BottomLayerWorldNormal;
+
+					// override to default material
+					ShadingModelID = SHADINGMODELID_DEFAULT_LIT;
+
+					// Foce BRDF application, as ClearCoat materials are assumed fully resolved
+					bApplyBRDF = true;
+				}
+
 				// Prepare a ray to trace the top layer reflection
 				TopLayerRay.Origin = WorldPosition;
 				TopLayerRay.TMin = 0.0;
@@ -376,57 +396,90 @@ void RayTracingReflectionsRGS()
 						FixSampleDirectionIfNeeded(ScreenSpaceData.GBuffer.WorldNormal, BottomLayerRay.Direction);
 					}
 
-					
-					FRayCone BottomLayerRayCone = (FRayCone)0;
-					FMaterialClosestHitPayload BottomLayerPayload = TraceRayAndAccumulateBottomLayerResults(BottomLayerRay, TLAS, RandSequence, RayFlags, ReflectionMaxNormalBias, RAY_TRACING_MASK_OPAQUE, ReflectedShadowsType, ShouldDoDirectLighting, ShouldDoEmissiveAndIndirectLighting, RayCone, BottomLayerRadiance);
+					bool TraceBottomRay = (BounceIndex < MaxUnderCoatBounces) && (GetRoughnessFade(BottomLayerRoughness, ReflectionMaxRoughness) > 0.0);
 
-					// Distance and sky: similar processing to main ray.
-					float BottomLayerHitDistance = -1;
-					bool IsBottomSampleValid = false;
-					if (BottomLayerPayload.IsHit())
+					if (TraceBottomRay)
 					{
-						IsBottomSampleValid = true;
-						BottomLayerHitDistance = BottomLayerPayload.HitT;
-					}
-					else if (bAllowSkySampling && ReflectionStruct.SkyLightParameters.y > 0)
-					{
-						IsBottomSampleValid = true;
-						float SkyAverageBrightness = 1.0f;
-						float3 SkyLighting = GetSkyLightReflection(BottomLayerRay.Direction, BottomLayerRoughness, SkyAverageBrightness);
-						BottomLayerRadiance += SkyLighting;
-						BottomLayerHitDistance = 1.0e20;
-					}
+						FRayCone BottomLayerRayCone = (FRayCone)0;
+						FMaterialClosestHitPayload BottomLayerPayload = TraceRayAndAccumulateBottomLayerResults(BottomLayerRay, TLAS, RandSequence, RayFlags, ReflectionMaxNormalBias, RAY_TRACING_MASK_OPAQUE, ReflectedShadowsType, ShouldDoDirectLighting, ShouldDoEmissiveAndIndirectLighting, RayCone, BottomLayerRadiance);
+
+						// Distance and sky: similar processing to main ray.
+						float BottomLayerHitDistance = -1;
+						bool IsBottomSampleValid = false;
+						if (BottomLayerPayload.IsHit())
+						{
+							IsBottomSampleValid = true;
+							BottomLayerHitDistance = BottomLayerPayload.HitT;
+						}
+						else if (bAllowSkySampling && ReflectionStruct.SkyLightParameters.y > 0)
+						{
+							IsBottomSampleValid = true;
+							float SkyAverageBrightness = 1.0f;
+							float3 SkyLighting = GetSkyLightReflection(BottomLayerRay.Direction, BottomLayerRoughness, SkyAverageBrightness);
+							BottomLayerRadiance += SkyLighting;
+							BottomLayerHitDistance = 1.0e20;
+						}
 					
-					if (BounceIndex == 0)
+						if (BounceIndex == 0)
+						{
+							HitDistance = IsBottomSampleValid ? min(bIsValidSample ? HitDistance : 1.0e27f, BottomLayerHitDistance) : HitDistance;
+						}
+					}
+					else
 					{
-						HitDistance = IsBottomSampleValid ? min(bIsValidSample ? HitDistance : 1.0e27f, BottomLayerHitDistance) : HitDistance;
+						// Ignoring any form of AO
+						float IndirectSpecularOcclusion = 1.0f;
+						float3 ExtraIndirectSpecular = 0;
+						float IndirectIrradiance = 0;
+
+						// Not possible to use the screen culled grid. So going over every capture.
+						// dxr_todo This would need a spatial data structure.
+						uint NumCulledReflectionCaptures = ForwardLightData.NumReflectionCaptures;
+						uint ReflectionCapturesStartIndex = 0;
+
+						float3 R = reflect(BottomLayerRay.Direction, BottomLayerWorldNormal);
+						const float NoV = saturate(dot(-BottomLayerRay.Direction, BottomLayerWorldNormal));
+						const float AO = 1.0f;
+						const float RoughnessSq = BottomLayerRoughness * BottomLayerRoughness;
+						const float SpecularOcclusion = GetSpecularOcclusion(NoV, RoughnessSq, AO);
+
+						BottomLayerRadiance = SpecularOcclusion *
+							CompositeReflectionCapturesAndSkylight(
+								1.0,	// coverage of 1
+								WorldPosition,
+								R,
+								BottomLayerRoughness,
+								IndirectIrradiance,
+								IndirectSpecularOcclusion,
+								ExtraIndirectSpecular,
+								NumCulledReflectionCaptures,
+								ReflectionCapturesStartIndex,
+								0);
 					}
 
-
 					// Apply correct weights to the bottom and top layers (See ReflectionEnvironmentPixelShader)
 					{
 						float NoVBottom = saturate(dot(-IncidentDirection, BottomLayerWorldNormal));
 
 						float TopLayerFSChlick = EnvBRDF(TopLayerSpecularColor, TopLayerRoughness, NoV).x;
-						float BottomLayerFSChlick = EnvBRDF(BottomLayerSpecularColor, BottomLayerRoughness, NoVBottom).x;
 
-						float TopLayerRefractionTransmittance = (1 - TopLayerFSChlick);
+						float TopLayerRefractionTransmittance = (1 - TopLayerFSChlick * ClearCoat);
 
-						float2 AB = PreIntegratedGF.SampleLevel(PreIntegratedGFSampler, float2(NoV, TopLayerRoughness), 0).rg;
-						float3 TopLayerSpecularLobIntegral = TopLayerSpecularColor * AB.x + AB.y * saturate(50.0 * TopLayerSpecularColor.g) * (1 - ClearCoat);
+						float2 AB = PreIntegratedGF.SampleLevel(PreIntegratedGFSampler, float2(NoVBottom, BottomLayerRoughness), 0).rg;
+						float3 BottomLayerFSChlick = BottomLayerSpecularColor * AB.x + AB.y * saturate(50.0 * BottomLayerSpecularColor.g) * (1 - ClearCoat);
 
 						// dxr-todo: double check the math when we can compare with path tracer.
 						// Matches what is found in ReflectionEnvironmentPixelShader.usf.
 						//PathSegmentRadiance.rgb = ClearCoat * (TopLayerRadiance * TopLayerSpecularLobIntegral * TopLayerRefractionTransmittance + TopLayerRadiance * TopLayerFSChlick);
-						TopLayerReflectionEventThroughput = ClearCoat * (TopLayerSpecularLobIntegral * TopLayerRefractionTransmittance + TopLayerFSChlick);
-						float BottomToTopRefractionThroughput = 1.0; // This is already evaluated in TraceRayAndAccumulateBottomLayerResults
+						TopLayerReflectionEventThroughput = ClearCoat * TopLayerFSChlick;
+						float BottomToTopRefractionThroughput = TopLayerRefractionTransmittance;
 						BottomLayerRadiance = BottomLayerRadiance * BottomLayerFSChlick * BottomToTopRefractionThroughput; // Bottom layer refelciton accounts for its throughput.
 					}
 				}
 				else
 				{
 					// If it is not clear coat, and we are on the first bounce, the G F terms will be applied by ReflectionEnvironmentPixelShader
-					if (BounceIndex > 0)
+					if (bApplyBRDF)
 						TopLayerReflectionEventThroughput *= EnvBRDF(TopLayerSpecularColor, TopLayerRoughness, NoV);
 				}
 
diff --git a/Engine/Shaders/Private/ReflectionEnvironmentPixelShader.usf b/Engine/Shaders/Private/ReflectionEnvironmentPixelShader.usf
--- a/Engine/Shaders/Private/ReflectionEnvironmentPixelShader.usf
+++ b/Engine/Shaders/Private/ReflectionEnvironmentPixelShader.usf
@@ -220,7 +220,7 @@ float3 ReflectionEnvironment(FGBufferData GBuffer, float AmbientOcclusion, float
 	if(GBuffer.ShadingModelID == SHADINGMODELID_CLEAR_COAT )
 	{
 #if RAY_TRACED_REFLECTIONS
-		Color = float4(0, 0, 0, 0); // Clear coat reflection is expected to be computed on a black background
+		Color = float4(0, 0, 0, 1); // Clear coat reflection is expected to be computed on a black background
 #endif
 		const float ClearCoat = GBuffer.CustomData.x;
 		Color = lerp( Color, float4(0,0,0,1), ClearCoat );
@@ -273,9 +273,10 @@ float3 ReflectionEnvironment(FGBufferData GBuffer, float AmbientOcclusion, float
 		Color.rgb *= LayerAttenuation;
 		Color.a = F;
 		
-			
+#if !RAY_TRACED_REFLECTIONS	
 		Color.rgb += SSR.rgb * F;
 		Color.a *= 1 - SSR.a;
+#endif
 			
 		Color.a *= SpecularOcclusion;
 
diff --git a/Engine/Source/Runtime/Renderer/Private/RayTracing/RayTracingReflections.cpp b/Engine/Source/Runtime/Renderer/Private/RayTracing/RayTracingReflections.cpp
--- a/Engine/Source/Runtime/Renderer/Private/RayTracing/RayTracingReflections.cpp
+++ b/Engine/Source/Runtime/Renderer/Private/RayTracing/RayTracingReflections.cpp
@@ -114,6 +114,21 @@ static TAutoConsoleVariable<int32> CVarRayTracingReflectionsTestPathRoughness(
 	TEXT("Accumulate roughness along path and test accumulated roughness against MaxRoughness before launching the next bounce (default 1)"),
 	ECVF_RenderThreadSafe);
 
+static TAutoConsoleVariable<float> CVarRayTracingReflectionsMinClearCoatLevel(
+	TEXT("r.RayTracing.Reflections.MinClearCoatLevel"),
+	0.01f,
+	TEXT("Minimum level at which to apply clear coat shading (default 0.01)\n")
+	TEXT(" Note: causes some variation in height fog due to using the bottom layer path"),
+	ECVF_RenderThreadSafe);
+
+static TAutoConsoleVariable<int32> CVarRayTracingReflectionsMaxUnderCoatBounces(
+	TEXT("r.RayTracing.Reflections.MaxUnderCoatBounces"),
+	0,
+	TEXT("How many bounces to apply ray traced reflections to the undercoat layer. Extra bounces will use reflection probes. (default 0, always use probes)"),
+	ECVF_RenderThreadSafe);
+
+
+
 static const int32 GReflectionLightCountMaximum = 64;
 
 class FRayTracingReflectionsRGS : public FGlobalShader
@@ -144,6 +159,8 @@ class FRayTracingReflectionsRGS : public FGlobalShader
 		SHADER_PARAMETER(float, ReflectionMaxRoughness)
 		SHADER_PARAMETER(float, ReflectionMaxNormalBias)
 		SHADER_PARAMETER(int32, TestPathRoughness)
+		SHADER_PARAMETER(float, MinClearCoatLevel)
+		SHADER_PARAMETER(int32, MaxUnderCoatBounces)
 
 		SHADER_PARAMETER_SRV(RaytracingAccelerationStructure, TLAS)
 		SHADER_PARAMETER_SRV(StructuredBuffer<FRTLightingData>, LightDataBuffer)
@@ -261,6 +278,8 @@ void FDeferredShadingSceneRenderer::RenderRayTracingReflections(
 	CommonParameters.RayTracingResolution = RayTracingResolution;
 	CommonParameters.TileAlignedResolution = TileAlignedResolution;
 	CommonParameters.TestPathRoughness = CVarRayTracingReflectionsTestPathRoughness.GetValueOnRenderThread();
+	CommonParameters.MinClearCoatLevel = CVarRayTracingReflectionsMinClearCoatLevel.GetValueOnRenderThread();
+	CommonParameters.MaxUnderCoatBounces = CVarRayTracingReflectionsMaxUnderCoatBounces.GetValueOnRenderThread();
 
 	CommonParameters.TLAS = View.RayTracingScene.RayTracingSceneRHI->GetShaderResourceView();
 	CommonParameters.ViewUniformBuffer = View.ViewUniformBuffer;
