Fix ShadowDepthVertexShader for graphics pipelines that couple tessellation with geometry shaders.

#rnx

diff --git a/Engine/Shaders/Private/LocalVertexFactory.ush b/Engine/Shaders/Private/LocalVertexFactory.ush
--- a/Engine/Shaders/Private/LocalVertexFactory.ush
+++ b/Engine/Shaders/Private/LocalVertexFactory.ush
@@ -1041,7 +1041,10 @@ float4 VertexFactoryGetPreviousWorldPosition(FVertexFactoryInput Input, FVertexF
 
 	FVertexFactoryInterpolantsVSToDS VertexFactoryInterpolate(FVertexFactoryInterpolantsVSToDS a, float aInterp, FVertexFactoryInterpolantsVSToDS b, float bInterp)
 	{
-		FVertexFactoryInterpolantsVSToDS O;
+		// Default initialize. Otherwise, some graphics pipelines that
+		// couple tessellation with geometry shaders won't write to all TEXCOORD semantics,
+		// but read from them when <FVertexFactoryInterpolantsVSToDS> is being copied as a whole.
+		FVertexFactoryInterpolantsVSToDS O = (FVertexFactoryInterpolantsVSToDS)0;
 
 #if VF_USE_PRIMITIVE_SCENE_DATA
 		O.InterpolantsVSToPS.PrimitiveId = a.InterpolantsVSToPS.PrimitiveId;
diff --git a/Engine/Shaders/Private/ShadowDepthVertexShader.usf b/Engine/Shaders/Private/ShadowDepthVertexShader.usf
--- a/Engine/Shaders/Private/ShadowDepthVertexShader.usf
+++ b/Engine/Shaders/Private/ShadowDepthVertexShader.usf
@@ -194,6 +194,7 @@ struct FShadowDepthDSToPS : FShadowDepthVSToPS
 #else
 	float4 Position : TEXCOORD6;
 #endif
+	float3 WorldNormal : TEXCOORD8;
 #else
 	float4 Position : SV_POSITION;
 #endif
@@ -220,15 +221,18 @@ FShadowDepthDSToPS PassFinalizeTessellationOutput(FShadowDepthVSToDS Interpolant
 	O.FactoryInterpolants = VertexFactoryAssignInterpolants(Interpolants.FactoryInterpolants);
 #endif
 
+	const float3 WorldNormal = MaterialParameters.TangentToWorld[2];
 
 #if ONEPASS_POINTLIGHT_SHADOW
-	O.Position = WorldPosition;
 #if USING_VERTEX_SHADER_LAYER
 	O.LayerIndex = Interpolants.LayerIndex;
 	O.Position = WorldPosition = mul(WorldPosition, PassStruct.ShadowViewProjectionMatrices[Interpolants.LayerIndex]);
-#endif
+	O.InterpPosition = WorldPosition;
+#else
+	O.Position = WorldPosition;
+#endif // USING_VERTEX_SHADER_LAYER
+	O.WorldNormal = WorldNormal;
 #else
-	const float3 WorldNormal = MaterialParameters.TangentToWorld[2];
 
 	float ShadowDepth;
 	SetShadowDepthOutputs(
