Fix TexCoords calculation with SubUV

diff --git a/Engine/Shaders/Private/NiagaraSpriteVertexFactory.ush b/Engine/Shaders/Private/NiagaraSpriteVertexFactory.ush
--- a/Engine/Shaders/Private/NiagaraSpriteVertexFactory.ush
+++ b/Engine/Shaders/Private/NiagaraSpriteVertexFactory.ush
@@ -1001,6 +1001,15 @@ FVertexFactoryInterpolantsVSToDS VertexFactoryGetInterpolantsVSToDS(FVertexFacto
 	return Interpolants;
 }
 
+float2 VertexFactoryGetTextureCoordinateDS( FVertexFactoryInterpolantsVSToDS Interpolants )
+{
+#if NUM_MATERIAL_TEXCOORDS
+	return Interpolants.InterpolantsVSToPS.TexCoords[0].xy;
+#else // #if NUM_MATERIAL_TEXCOORDS
+	return float2(0,0);
+#endif // #if NUM_MATERIAL_TEXCOORDS
+}
+
 FVertexFactoryInterpolantsVSToDS VertexFactoryInterpolate(FVertexFactoryInterpolantsVSToDS a, float aInterp, FVertexFactoryInterpolantsVSToDS b, float bInterp)
 {
 	FVertexFactoryInterpolantsVSToDS O;
@@ -1018,7 +1027,7 @@ FVertexFactoryInterpolantsVSToDS VertexFactoryInterpolate(FVertexFactoryInterpol
 #endif
 
 #if USE_PARTICLE_SUBUVS
-	TESSELLATION_INTERPOLATE_MEMBER(InterpolantsVSToPS.Interpolants.ParticleSubUVs);
+	TESSELLATION_INTERPOLATE_MEMBER(InterpolantsVSToPS.ParticleSubUVs);
 #endif
 
 	TESSELLATION_INTERPOLATE_MEMBER(InterpolantsVSToPS.TangentToWorld0AndInterp_Sizer);
diff --git a/Engine/Shaders/Private/RayTracing/RayTracingMaterialHitShaders.usf b/Engine/Shaders/Private/RayTracing/RayTracingMaterialHitShaders.usf
--- a/Engine/Shaders/Private/RayTracing/RayTracingMaterialHitShaders.usf
+++ b/Engine/Shaders/Private/RayTracing/RayTracingMaterialHitShaders.usf
@@ -2,7 +2,7 @@
 
 #pragma once
 
-#define ENABLE_SKY_LIGHT 1
+#define ENABLE_SKY_LIGHT 1      
 #define NEEDS_LIGHTMAP_COORDINATE (HQ_TEXTURE_LIGHTMAP || LQ_TEXTURE_LIGHTMAP)
 #ifdef USING_TESSELLATION
 	#undef USING_TESSELLATION
@@ -20,15 +20,15 @@
 #include "RayTracingCommon.ush"
 #include "RayTracingHitGroupCommon.ush"
 
-#include "/Engine/Generated/Material.ush"
+#include "/Engine/Generated/Material.ush"   
 #include "/Engine/Generated/VertexFactory.ush"
 
 #include "/Engine/Private/ShadingCommon.ush"
 #include "/Engine/Private/DeferredShadingCommon.ush"
 #include "/Engine/Private/SHCommon.ush"
 #include "/Engine/Private/ReflectionEnvironmentShared.ush"
-#include "/Engine/Private/VirtualTextureCommon.ush"
-#include "/Engine/Private/LightmapCommon.ush"
+#include "/Engine/Private/VirtualTextureCommon.ush" 
+#include "/Engine/Private/LightmapCommon.ush"     
 
 /** Computes sky diffuse lighting, including precomputed shadowing. */
 void GetSkyLighting(FMaterialPixelParameters MaterialParameters, VTPageTableResult LightmapVTPageTableResult, float3 WorldNormal, float2 LightmapUV, uint LightmapDataIndex, float3 SkyOcclusionUV3D, uint ShadingModel, out float3 OutDiffuseLighting, out float3 OutSubsurfaceLighting)
@@ -200,9 +200,7 @@ void CalcInterpolants(in FRayCone RayCone, in FDefaultAttributes Attributes, out
 		
 		PerVertexInterpolants[i] = VertexFactoryGetInterpolantsVSToDS(Input, VFIntermediates, VertexParameters);
 
-#if (NUM_TEX_COORD_INTERPOLATORS || USE_PARTICLE_SUBUVS) && !VERTEX_FACTORY_MODIFIES_TESSELLATION
-		TexCoords[i] = PerVertexInterpolants[i].InterpolantsVSToPS.TexCoords[0].xy;
-#endif
+		TexCoords[i] = VertexFactoryGetTextureCoordinateDS(PerVertexInterpolants[i]);
 
 		float4 RasterizedWorldPosition = VertexFactoryGetRasterizedWorldPosition(Input, VFIntermediates, WorldPosition);
 		PerVertexClipSpacePositions[i] = mul(RasterizedWorldPosition, ResolvedView.TranslatedWorldToClip);
