Integrating QuasiRandom with ray tracing SkyLight pass. Remove normal-culling after NSight profiling analysis.

diff --git a/Engine/Shaders/Private/RayTracing/RayTracingSkyLightRGS.usf b/Engine/Shaders/Private/RayTracing/RayTracingSkyLightRGS.usf
--- a/Engine/Shaders/Private/RayTracing/RayTracingSkyLightRGS.usf
+++ b/Engine/Shaders/Private/RayTracing/RayTracingSkyLightRGS.usf
@@ -4,9 +4,9 @@
 #include "../MonteCarlo.ush"
 #include "../DeferredShadingCommon.ush"
 #include "../ShadingModels.ush"
+#include "../QuasiRandom.ush"
 #include "SkyLightMipTreeCommon.ush"
 #include "RayTracingCommon.ush"
-#include "../PathTracing/Utilities/PathTracingRandomSequence.ush"
 
 RaytracingAccelerationStructure TLAS;
 RWTexture2D<float4> RWOcclusionMaskUAV;
@@ -38,12 +38,12 @@ float PdfSkyLightOcclusionRay(
 #if USE_MIP_PDF_TEXTURE
 	float MipPdf = PdfMipTree(uint3(MipCoord, TextureCoord.z), StopLevel);
 
-#if USE_SOLID_ANGLE_TEXTURE
-	uint LinearIndex = BufferOffsetAtPixel(MipCoord, StopLevel, SkyLight.MipDimensions.xy);
-	float AreaPdf = SkyLight.SolidAnglePdf[LinearIndex];
-#else
-	float AreaPdf = 1.0 / TexelCoordSolidAngle(MipCoord, DimensionsAtLevel(SkyLight.MipDimensions, StopLevel));
-#endif
+	#if USE_SOLID_ANGLE_TEXTURE
+		uint LinearIndex = BufferOffsetAtPixel(MipCoord, StopLevel, SkyLight.MipDimensions.xy);
+		float AreaPdf = SkyLight.SolidAnglePdf[LinearIndex];
+	#else
+		float AreaPdf = 1.0 / TexelCoordSolidAngle(MipCoord, DimensionsAtLevel(SkyLight.MipDimensions, StopLevel));
+	#endif
 	return FacePdf * MipPdf * AreaPdf;
 #else
 	uint2 MipPixel = 0;
@@ -80,12 +80,12 @@ float PdfSkyLightOcclusionRay(
 	// DEBUG: assert that PdfMipTree is equivalent for evaluating the decision tree..
 	//MipPdf = PdfMipTree(uint3(MipCoord, TextureCoord.z), StopLevel);
 
-#if USE_SOLID_ANGLE_TEXTURE
-	uint LinearIndex = BufferOffsetAtPixel(MipCoord, StopLevel, SkyLight.MipDimensions.xy);
-	float AreaPdf = SkyLight.SolidAnglePdf[LinearIndex];
-#else
-	float AreaPdf = 1.0 / TexelCoordSolidAngle(MipCoord, DimensionsAtLevel(SkyLight.MipDimensions, StopLevel));
-#endif
+	#if USE_SOLID_ANGLE_TEXTURE
+		uint LinearIndex = BufferOffsetAtPixel(MipCoord, StopLevel, SkyLight.MipDimensions.xy);
+		float AreaPdf = SkyLight.SolidAnglePdf[LinearIndex];
+	#else
+		float AreaPdf = 1.0 / TexelCoordSolidAngle(MipCoord, DimensionsAtLevel(SkyLight.MipDimensions, StopLevel));
+	#endif
 	return FacePdf * MipPdf * AreaPdf;
 #endif
 }
@@ -94,7 +94,7 @@ bool GenerateSkyLightOcclusionRay(
 	float FaceCdf[7],
 	float3 WorldPosition,
 	float3 WorldNormal,
-	float4 RandSample,
+	float3 RandSample,
 	out float3 RayOrigin,
 	out float3 RayDirection,
 	out float RayTMin,
@@ -223,8 +223,6 @@ bool GenerateLambertOcclusionRay(
 void SkyLightRGS()
 {
 	uint2 PixelCoord = DispatchRaysIndex().xy + View.ViewRectMin;
-
-	RandomSequence RandSequence;
 	uint LinearIndex = CalcLinearIndex(PixelCoord);
 
 	// Get G-Buffer surface data
@@ -251,7 +249,6 @@ void SkyLightRGS()
 	float RayDistance = 0.0;
 	float HitCount = 0.0;
 	uint SamplesPerPixel = SkyLight.SamplesPerPixel;
-	uint FrameIndex = View.StateFrameIndex % SamplesPerPixel;
 
 	// Mask out depth values that are infinitely far away
 	float DeviceZ = SceneTexturesStruct.SceneDepthTexture.Load(int3(PixelCoord, 0)).r;
@@ -269,23 +266,28 @@ void SkyLightRGS()
 	float FaceCdf[7];
 	BuildFaceCdf(MipCount, FaceCdf);
 
+	FQuasiRandom QuasiRandom;
+	const uint Dimension = 0;
+	QuasiRandom_Initialize(QuasiRandom, PixelCoord, 0, SamplesPerPixel);
+	float3 BlueNoiseShift = EvalBlueNoise(QuasiRandom.BlueNoiseCoord, Dimension);
+
 	float3 ExitantRadiance = 0.0;
 	float3 DiffuseExitantRadiance = 0.0;
 	for (uint SampleIndex = 0; SampleIndex < SamplesPerPixel; ++SampleIndex)
 	{
-		// Use Scrambled Halton sequence
-		RandomSequence_Initialize(RandSequence, LinearIndex, FrameIndex * SamplesPerPixel + SampleIndex);
-		RandSequence.Type = 2;
-
+		QuasiRandom_Update(QuasiRandom, SampleIndex, SamplesPerPixel);
 		RayDesc Ray;
 		float RayPdf = 0.0;
 
 		// Determine sky light or lambert ray
-		uint DummyVariable;
-		float4 RandSample = RandomSequence_GenerateSample4D(RandSequence, DummyVariable);
+		float3 RandSample = QuasiRandom_Sample3(QuasiRandom);
+		RandSample = HaltonSequence3(QuasiRandom.HaltonSequence, QuasiRandom.HaltonIteration, Dimension);
+		RandSample += BlueNoiseShift;
+		RandSample = fmod(RandSample, 1.0);
 
 		float SkyLightSamplingStrategyPdf = 0.5;
 		float MisWeight = 1.0;
+		BRANCH
 		if (RandSample.x < SkyLightSamplingStrategyPdf)
 		{
 			RandSample.x /= SkyLightSamplingStrategyPdf;
@@ -300,13 +302,6 @@ void SkyLightRGS()
 
 		ApplyPositionBias(Ray, WorldNormal, SkyLight.MaxNormalBias);
 
-		// Surface-normal culling
-		float NoL = dot(WorldNormal, Ray.Direction);
-		if (RayPdf <= 0.0 || NoL <= 0.0)
-		{
-			continue;
-		}
-
 		uint RayFlags = 0;
 #if !ENABLE_MATERIALS
 		RayFlags |= RAY_FLAG_FORCE_OPAQUE;
@@ -338,13 +333,14 @@ void SkyLightRGS()
 			half3 N = WorldNormal;
 			half3 V = -CameraDirection;
 			half3 L = Ray.Direction;
+			float NoL = dot(WorldNormal, Ray.Direction);
 			FShadowTerms ShadowTerms = { 0.0, 0.0, 0.0 };
 			FDirectLighting LightingSample = EvaluateBxDF(ScreenSpaceData.GBuffer, N, V, L, NoL, ShadowTerms);
 			float3 Brdf = LightingSample.Diffuse + LightingSample.Transmission + LightingSample.Specular;
 
 			float3 IncomingRadiance = EvalSkyLight(Ray.Direction);
 			ExitantRadiance += IncomingRadiance * Brdf * MisWeight / RayPdf;
-			DiffuseExitantRadiance += IncomingRadiance * (LightingSample.Diffuse + LightingSample.Transmission) * MisWeight / RayPdf;
+			DiffuseExitantRadiance += IncomingRadiance * (LightingSample.Diffuse + LightingSample.Transmission) * (MisWeight / RayPdf);
 		}
 	}
 
diff --git a/Engine/Source/Runtime/Renderer/Private/RayTracing/RaytracingSkylight.cpp b/Engine/Source/Runtime/Renderer/Private/RayTracing/RaytracingSkylight.cpp
--- a/Engine/Source/Runtime/Renderer/Private/RayTracing/RaytracingSkylight.cpp
+++ b/Engine/Source/Runtime/Renderer/Private/RayTracing/RaytracingSkylight.cpp
@@ -23,6 +23,7 @@ static int32 GRayTracingSkyLight = 0;
 #include "RayGenShaderUtils.h"
 #include "SceneTextureParameters.h"
 #include "ScreenSpaceDenoise.h"
+#include "BlueNoise.h"
 
 #include "Raytracing/RaytracingOptions.h"
 #include "PostProcess/PostProcessing.h"
@@ -178,6 +179,9 @@ class FRayTracingSkyLightRGS : public FGlobalShader
 		SHADER_PARAMETER_STRUCT_REF(FViewUniformShaderParameters, ViewUniformBuffer)
 		SHADER_PARAMETER_STRUCT_REF(FSceneTexturesUniformParameters, SceneTexturesStruct)
 		SHADER_PARAMETER_STRUCT_REF(FSkyLightData, SkyLightData)
+		SHADER_PARAMETER_STRUCT_REF(FHaltonIteration, HaltonIteration)
+		SHADER_PARAMETER_STRUCT_REF(FHaltonPrimes, HaltonPrimes)
+		SHADER_PARAMETER_STRUCT_REF(FBlueNoise, BlueNoise)
 
 		SHADER_PARAMETER_RDG_TEXTURE(Texture2D, SSProfilesTexture)
 		SHADER_PARAMETER_SAMPLER(SamplerState, TransmissionProfilesLinearSampler)
@@ -829,6 +833,13 @@ void FDeferredShadingSceneRenderer::RenderRayTracingSkyLight(
 		RayDistanceTexture = GraphBuilder.CreateTexture(Desc, TEXT("RayTracingSkyLightHitDistance"));
 	}
 
+	// Sampling state
+	FHaltonPrimes HaltonPrimes;
+	InitializeHaltonPrimes(Scene->HaltonPrimesResource, HaltonPrimes);
+
+	FBlueNoise BlueNoise;
+	InitializeBlueNoise(BlueNoise);
+
 	// Fill SkyLight parameters
 	FSkyLightData SkyLightData;
 	SetupSkyLightParameters(*Scene, &SkyLightData);
@@ -851,6 +862,8 @@ void FDeferredShadingSceneRenderer::RenderRayTracingSkyLight(
 	PassParameters->RWRayDistanceUAV = GraphBuilder.CreateUAV(RayDistanceTexture);
 	PassParameters->SceneTexturesStruct = CreateUniformBufferImmediate(SceneUniformTextures, EUniformBufferUsage::UniformBuffer_SingleDraw);
 	PassParameters->SkyLightData = CreateUniformBufferImmediate(SkyLightData, EUniformBufferUsage::UniformBuffer_SingleDraw);
+	PassParameters->HaltonPrimes = CreateUniformBufferImmediate(HaltonPrimes, EUniformBufferUsage::UniformBuffer_SingleDraw);
+	PassParameters->BlueNoise = CreateUniformBufferImmediate(BlueNoise, EUniformBufferUsage::UniformBuffer_SingleDraw);
 	PassParameters->SSProfilesTexture = GraphBuilder.RegisterExternalTexture(SubsurfaceProfileRT);
 	PassParameters->TransmissionProfilesLinearSampler = TStaticSamplerState<SF_Bilinear, AM_Clamp, AM_Clamp, AM_Clamp>::GetRHI();
 
@@ -858,9 +871,20 @@ void FDeferredShadingSceneRenderer::RenderRayTracingSkyLight(
 	{
 		FViewInfo& View = Views[ViewIndex];
 
+
+		// Halton iteration
+		uint32 IterationCount = GRayTracingSkyLightSamplesPerPixel >= 0 ? GRayTracingSkyLightSamplesPerPixel : Scene->SkyLight->SamplesPerPixel;
+		uint32 SequenceCount = 1;
+		uint32 DimensionCount = 3;
+		FHaltonSequenceIteration HaltonSequenceIteration(Scene->HaltonSequence, IterationCount, SequenceCount, DimensionCount, View.ViewState ? (View.ViewState->FrameIndex % 1024) : 0);
+
+		FHaltonIteration HaltonIteration;
+		InitializeHaltonSequenceIteration(HaltonSequenceIteration, HaltonIteration);
+
 		//  View-dependent parameters
 		PassParameters->TLAS = View.RayTracingScene.RayTracingSceneRHI->GetShaderResourceView();
 		PassParameters->ViewUniformBuffer = View.ViewUniformBuffer;
+		PassParameters->HaltonIteration = CreateUniformBufferImmediate(HaltonIteration, EUniformBufferUsage::UniformBuffer_SingleDraw);
 
 		FRayTracingSkyLightRGS::FPermutationDomain PermutationVector;
 		PermutationVector.Set<FRayTracingSkyLightRGS::FEnableTwoSidedGeometryDim>(CVarRayTracingSkyLightEnableTwoSidedGeometry.GetValueOnRenderThread() != 0);
