Rewrites indirect lighting application to the scene color using RDG, to have area where different GI methods can hook themselves.

diff --git a/Engine/Shaders/Private/AmbientCubemapComposite.usf b/Engine/Shaders/Private/AmbientCubemapComposite.usf
new file mode 100644
--- /dev/null
+++ b/Engine/Shaders/Private/AmbientCubemapComposite.usf
@@ -0,0 +1,473 @@
+// Copyright 1998-2019 Epic Games, Inc. All Rights Reserved.
+
+#include "Common.ush"
+#include "PostProcessCommon.ush"
+#include "DeferredShadingCommon.ush"
+#include "CubemapCommon.ush"
+#include "Random.ush"
+#include "BRDF.ush"
+#include "MonteCarlo.ush"
+#include "ShadingModels.ush"
+#include "SceneTextureParameters.ush"
+
+#define IMPORTANCE_SAMPLE		0
+
+Texture2D AmbientOcclusionTexture;
+SamplerState AmbientOcclusionSampler;
+
+
+float3 DiffuseIBL( uint2 Random, float3 DiffuseColor, float Roughness, float3 N, float3 V )
+{
+	N = normalize( N );
+	V = normalize( V );
+
+	float3 DiffuseLighting = 0;
+	
+	float NoV = saturate( dot( N, V ) );
+
+	const uint NumSamples = 32;
+	for( uint i = 0; i < NumSamples; i++ )
+	{
+		float2 E = Hammersley( i, NumSamples, Random );
+		float3 L = TangentToWorld( CosineSampleHemisphere( E ).xyz, N );
+		float3 H = normalize(V + L);
+
+		float NoL = saturate( dot( N, L ) );
+		float NoH = saturate( dot( N, H ) );
+		float VoH = saturate( dot( V, H ) );
+
+		if( NoL > 0 )
+		{
+			float3 SampleColor = AmbientCubemap.SampleLevel( AmbientCubemapSampler, L, 0 ).rgb;
+
+			float FD90 = ( 0.5 + 2 * VoH * VoH ) * Roughness;
+			//float FD90 = 0.5 + 2 * VoH * VoH * Roughness;
+			float FdV = 1 + (FD90 - 1) * pow( 1 - NoV, 5 );
+			float FdL = 1 + (FD90 - 1) * pow( 1 - NoL, 5 );
+
+#if 1
+			// lambert = DiffuseColor * NoL / PI
+			// pdf = NoL / PI
+			DiffuseLighting += SampleColor * DiffuseColor * FdV * FdL * ( 1 - 0.3333 * Roughness );
+#else
+			DiffuseLighting += SampleColor * DiffuseColor;
+#endif
+		}
+	}
+
+	return DiffuseLighting / NumSamples;
+}
+
+float3 SpecularIBL( uint2 Random, float3 SpecularColor, float Roughness, float3 N, float3 V )
+{
+	float3 SpecularLighting = 0;
+
+	const uint NumSamples = 32;
+	for( uint i = 0; i < NumSamples; i++ )
+	{
+		float2 E = Hammersley( i, NumSamples, Random );
+		float3 H = TangentToWorld( ImportanceSampleGGX( E, Pow4(Roughness) ).xyz, N );
+		float3 L = 2 * dot( V, H ) * H - V;
+
+		float NoV = saturate( dot( N, V ) );
+		float NoL = saturate( dot( N, L ) );
+		float NoH = saturate( dot( N, H ) );
+		float VoH = saturate( dot( V, H ) );
+		
+		if( NoL > 0 )
+		{
+			float3 SampleColor = AmbientCubemap.SampleLevel( AmbientCubemapSampler, L, 0 ).rgb;
+
+			float Vis = Vis_SmithJointApprox( Pow4(Roughness), NoV, NoL );
+			float Fc = pow( 1 - VoH, 5 );
+			float3 F = (1 - Fc) * SpecularColor + Fc;
+
+			// Incident light = SampleColor * NoL
+			// Microfacet specular = D*G*F / (4*NoL*NoV) = D*Vis*F
+			// pdf = D * NoH / (4 * VoH)
+			SpecularLighting += SampleColor * F * ( NoL * Vis * (4 * VoH / NoH) );
+		}
+	}
+
+	return SpecularLighting / NumSamples;
+}
+
+
+float3 ImageBasedLightingMIS( FGBufferData GBuffer, float3 V, float3 N, uint2 Random )
+{
+	float3 Lighting = 0;
+
+	float Roughness1 = GBuffer.Roughness;
+	float Roughness2 = 0.1;
+	
+	uint NumSamples[] =
+	{
+		16,
+		16,
+		0,
+	};
+
+	UNROLL
+	for( uint Set = 0; Set < 3; Set++ )
+	{
+		LOOP
+		for( uint i = 0; i < NumSamples[ Set ]; i++ )
+		{
+			float2 E = Hammersley( i, NumSamples[ Set ], Random );
+			
+			float3 L, H;
+			if( Set == 0 )
+			{
+				L = TangentToWorld( CosineSampleHemisphere( E ).xyz, N );
+				H = normalize(V + L);
+			}
+			else if( Set == 1 )
+			{
+				H = TangentToWorld( ImportanceSampleGGX( E, Pow4(Roughness1) ).xyz, N );
+				L = 2 * dot( V, H ) * H - V;
+			}
+			else
+			{
+				H = TangentToWorld( ImportanceSampleGGX( E, Pow4(Roughness2) ).xyz, N );
+				L = 2 * dot( V, H ) * H - V;
+			}
+
+			float NoL = saturate( dot(N, L) );
+			float NoH = saturate( dot(N, H) );
+			float VoH = saturate( dot(V, H) );
+		
+			if( NoL > 0 && VoH > 0 )
+			{
+				float3 SampleColor = AmbientCubemap.SampleLevel( AmbientCubemapSampler, L, 0 ).rgb;
+
+				float PDF[] =
+				{
+					NoL / PI,
+					D_GGX( Pow4(Roughness1), NoH ) * NoH / (4 * VoH),
+					D_GGX( Pow4(Roughness2), NoH ) * NoH / (4 * VoH),
+				};
+
+				// MIS balance heuristic
+				float InvWeight = 0;
+				UNROLL for( uint j = 0; j < 3; j++ )
+				{
+					InvWeight += PDF[j] * NumSamples[j];
+				}
+				float Weight = rcp( InvWeight );
+
+				FShadowTerms Shadow = { 1, 1, 1 };
+				FDirectLighting LightingSample = EvaluateBxDF( GBuffer, N, V, L, NoL, Shadow );
+
+				Lighting += SampleColor * Weight * ( LightingSample.Diffuse + LightingSample.Specular + LightingSample.Transmission );
+			}
+		}
+	}
+
+	return Lighting;
+}
+
+float3 FilterEnvMap( uint2 Random, float Roughness, float3 N, float3 V )
+{
+	float3 FilteredColor = 0;
+	float Weight = 0;
+
+	const uint NumSamples = 64;
+	for( uint i = 0; i < NumSamples; i++ )
+	{
+		float2 E = Hammersley( i, NumSamples, Random );
+		float3 H = TangentToWorld( ImportanceSampleGGX( E, Pow4(Roughness) ).xyz, N );
+		float3 L = 2 * dot( V, H ) * H - V;
+
+		float NoL = saturate( dot( N, L ) );
+		if( NoL > 0 )
+		{
+			FilteredColor += AmbientCubemap.SampleLevel( AmbientCubemapSampler, L, 0 ).rgb * NoL;
+			Weight += NoL;
+		}
+	}
+
+	return FilteredColor / max( Weight, 0.001 );
+}
+
+float3 PrefilterEnvMap( uint2 Random, float Roughness, float3 R )
+{
+	float3 FilteredColor = 0;
+	float Weight = 0;
+		
+	const uint NumSamples = 64;
+	for( uint i = 0; i < NumSamples; i++ )
+	{
+		float2 E = Hammersley( i, NumSamples, Random );
+		float3 H = TangentToWorld( ImportanceSampleGGX( E, Pow4(Roughness) ).xyz, R );
+		float3 L = 2 * dot( R, H ) * H - R;
+
+		float NoL = saturate( dot( R, L ) );
+		if( NoL > 0 )
+		{
+			FilteredColor += AmbientCubemap.SampleLevel( AmbientCubemapSampler, L, 0 ).rgb * NoL;
+			Weight += NoL;
+		}
+	}
+
+	return FilteredColor / max( Weight, 0.001 );
+}
+
+float3 IntegrateBRDF( uint2 Random, float Roughness, float NoV )
+{
+	float3 V;
+	V.x = sqrt( 1.0f - NoV * NoV );	// sin
+	V.y = 0;
+	V.z = NoV;						// cos
+
+	float A = 0;
+	float B = 0;
+	float C = 0;
+
+	const uint NumSamples = 64;
+	for( uint i = 0; i < NumSamples; i++ )
+	{
+		float2 E = Hammersley( i, NumSamples, Random );
+
+		{
+			float3 H = ImportanceSampleGGX( E, Pow4(Roughness) ).xyz;
+			float3 L = 2 * dot( V, H ) * H - V;
+
+			float NoL = saturate( L.z );
+			float NoH = saturate( H.z );
+			float VoH = saturate( dot( V, H ) );
+
+			if( NoL > 0 )
+			{
+				float a = Square( Roughness );
+				float a2 = a*a;
+				float Vis = Vis_SmithJointApprox( a2, NoV, NoL );
+				float Vis_SmithV = NoL * sqrt( NoV * (NoV - NoV * a2) + a2 );
+				float Vis_SmithL = NoV * sqrt( NoL * (NoL - NoL * a2) + a2 );
+				//float Vis = 0.5 * rcp( Vis_SmithV + Vis_SmithL );
+
+				// Incident light = NoL
+				// pdf = D * NoH / (4 * VoH)
+				// NoL * Vis / pdf
+				float NoL_Vis_PDF = NoL * Vis * (4 * VoH / NoH);
+
+				float Fc = pow( 1 - VoH, 5 );
+				A += (1 - Fc) * NoL_Vis_PDF;
+				B += Fc * NoL_Vis_PDF;
+			}
+		}
+
+		{
+			float3 L = CosineSampleHemisphere( E ).xyz;
+			float3 H = normalize(V + L);
+
+			float NoL = saturate( L.z );
+			float NoH = saturate( H.z );
+			float VoH = saturate( dot( V, H ) );
+
+			float FD90 = ( 0.5 + 2 * VoH * VoH ) * Roughness;
+			float FdV = 1 + (FD90 - 1) * pow( 1 - NoV, 5 );
+			float FdL = 1 + (FD90 - 1) * pow( 1 - NoL, 5 );
+			C += FdV * FdL * ( 1 - 0.3333 * Roughness );
+		}
+	}
+
+	return float3( A, B, C ) / NumSamples;
+}
+
+float3 ApproximateSpecularIBL( uint2 Random, float3 SpecularColor, float Roughness, float3 N, float3 V )
+{
+	// Function replaced with prefiltered environment map sample
+	float3 R = 2 * dot( V, N ) * N - V;
+	float3 PrefilteredColor = PrefilterEnvMap( Random, Roughness, R );
+	//float3 PrefilteredColor = FilterEnvMap( Random, Roughness, N, V );
+
+	// Function replaced with 2D texture sample
+	float NoV = saturate( dot( N, V ) );
+	float2 AB = IntegrateBRDF( Random, Roughness, NoV ).xy;
+
+	return PrefilteredColor * ( SpecularColor * AB.x + AB.y );
+}
+
+
+float3 ImageBasedLightingHair( FGBufferData GBuffer, float3 V, float3 N, uint2 Random )
+{
+	float3 Lighting = 0;
+
+	uint NumSamples = 32;
+	
+	LOOP
+	for( uint i = 0; i < NumSamples; i++ )
+	{
+		float2 E = Hammersley( i, NumSamples, Random );
+		float3 L = UniformSampleSphere( E ).xyz;
+		
+		{
+			float3 SampleColor = AmbientCubemap.SampleLevel( AmbientCubemapSampler, L, 0 ).rgb;
+
+			float PDF = 1.0 / (4 * PI);
+
+			float InvWeight = PDF * NumSamples;
+			float Weight = rcp( InvWeight );
+
+			float3 Shading = 0;
+			Shading = HairShading( GBuffer, L, V, N, 1, 0, 0, Random );
+			
+			Lighting += SampleColor * Shading * Weight;
+		}
+	}
+
+	return Lighting;
+}
+
+void MainPS(in noperspective float4 UVAndScreenPos : TEXCOORD0, float4 SvPosition : SV_POSITION, out float4 OutColor : SV_Target0)
+{
+	float2 BufferUV = UVAndScreenPos.xy;
+	float2 ScreenSpacePos = UVAndScreenPos.zw;
+	int2 PixelPos = int2(SvPosition.xy);
+
+	FGBufferData GBuffer = GetGBufferDataFromSceneTextures(BufferUV);
+
+	float AbsoluteDiffuseMip = AmbientCubemapMipAdjust.z;
+
+
+	float3 ScreenVector = normalize(mul(float4(ScreenSpacePos, 1, 0), View.ScreenToWorld).xyz);
+
+	float3 N = GBuffer.WorldNormal;
+	float3 V = -ScreenVector;
+	float3 R0 = 2 * dot( V, N ) * N - V;
+
+	//float NoV = max( dot(N, V), 1e-5 );
+	float NoV = saturate( dot(N, V) );
+
+	// Point lobe in off-specular peak direction
+	float a = Square( GBuffer.Roughness );
+	float3 R = lerp( N, R0, (1 - a) * ( sqrt(1 - a) + a ) );
+	
+	uint2 Random = Rand3DPCG16( uint3( PixelPos, View.StateFrameIndexMod8) ).xy;
+
+	float3 NonSpecularContribution = 0;
+	float3 SpecularContribution = 0;
+	
+	BRANCH if( GBuffer.ShadingModelID == SHADINGMODELID_CLEAR_COAT && CLEAR_COAT_BOTTOM_NORMAL)
+	{
+		const float2 oct1 = ((float2(GBuffer.CustomData.a, GBuffer.CustomData.z) * 2) - (256.0/255.0)) + UnitVectorToOctahedron(GBuffer.WorldNormal);
+		N = OctahedronToUnitVector(oct1);			
+	}
+	float3 DiffuseLookup =  TextureCubeSampleLevel(AmbientCubemap, AmbientCubemapSampler, N, AbsoluteDiffuseMip).rgb;
+	NonSpecularContribution += GBuffer.DiffuseColor * DiffuseLookup;
+
+	// Diffuse
+	{
+		// we want to access the mip with the preconvolved diffuse lighting (coneangle=90 degree)
+		//NonSpecularContribution += GBuffer.DiffuseColor * DiffuseLookup;
+	}
+
+	// Specular
+	{
+
+		BRANCH if( GBuffer.ShadingModelID == SHADINGMODELID_CLEAR_COAT )
+		{
+			float ClearCoat			  = GBuffer.CustomData.x;
+
+			#if CLEAR_COAT_BOTTOM_NORMAL
+				//const float3 ClearCoatUnderNormal = OctahedronToUnitVector((float2(GBuffer.CustomData.a, GBuffer.CustomData.z) * 2) - (256.0/255.0));
+				const float2 oct1 = ((float2(GBuffer.CustomData.a, GBuffer.CustomData.z) * 2) - (256.0/255.0)) + UnitVectorToOctahedron(GBuffer.WorldNormal);
+ 				const float3 ClearCoatUnderNormal = OctahedronToUnitVector(oct1);
+				float3 R2 = 2 * dot( V, ClearCoatUnderNormal ) * ClearCoatUnderNormal - V;
+			
+				float Mip = ComputeCubemapMipFromRoughness( GBuffer.Roughness, AmbientCubemapMipAdjust.w );
+				float3 SampleColor = TextureCubeSampleLevel( AmbientCubemap, AmbientCubemapSampler, R2, Mip ).rgb;
+			
+				float2 AB = PreIntegratedGF.SampleLevel( PreIntegratedGFSampler, float2( NoV, GBuffer.Roughness ), 0 ).rg;
+				SpecularContribution += SampleColor * ( GBuffer.SpecularColor * AB.x + AB.y * (1 - ClearCoat) );
+			#else
+				float Mip = ComputeCubemapMipFromRoughness( GBuffer.Roughness, AmbientCubemapMipAdjust.w );
+				float3 SampleColor = TextureCubeSampleLevel( AmbientCubemap, AmbientCubemapSampler, R, Mip ).rgb;
+
+				float2 AB = PreIntegratedGF.SampleLevel( PreIntegratedGFSampler, float2( NoV, GBuffer.Roughness ), 0 ).rg;
+				SpecularContribution += SampleColor * ( GBuffer.SpecularColor * AB.x + AB.y * (1 - ClearCoat) );
+			#endif
+		}
+		else
+		{
+			float Mip = ComputeCubemapMipFromRoughness( GBuffer.Roughness, AmbientCubemapMipAdjust.w );
+			float3 SampleColor = TextureCubeSampleLevel( AmbientCubemap, AmbientCubemapSampler, R, Mip ).rgb;
+
+			SpecularContribution += SampleColor * EnvBRDF( GBuffer.SpecularColor, GBuffer.Roughness, NoV );
+			//SpecularContribution += ApproximateSpecularIBL( Random, GBuffer.SpecularColor, GBuffer.Roughness, GBuffer.WorldNormal, -ScreenVector );
+		}
+	}
+
+	BRANCH if( GBuffer.ShadingModelID == SHADINGMODELID_CLEAR_COAT )
+	{
+		const float ClearCoat			= GBuffer.CustomData.x;
+		const float ClearCoatRoughness	= GBuffer.CustomData.y;
+
+		float Mip = ComputeCubemapMipFromRoughness( ClearCoatRoughness, AmbientCubemapMipAdjust.w );
+		float3 SampleColor = TextureCubeSampleLevel( AmbientCubemap, AmbientCubemapSampler, R0, Mip ).rgb;
+
+		// F_Schlick
+		float F0 = 0.04;
+		float Fc = pow( 1 - NoV, 5 );
+		float F = Fc + (1 - Fc) * F0;
+		F *= ClearCoat;
+
+		float LayerAttenuation = (1 - F);
+
+		NonSpecularContribution *= LayerAttenuation;
+		SpecularContribution *= LayerAttenuation;
+		SpecularContribution += SampleColor * F;
+	}
+	
+	BRANCH if( GBuffer.ShadingModelID == SHADINGMODELID_HAIR )
+	{
+		float3 FakeNormal = normalize( V - N * dot(V,N) );
+		FakeNormal = normalize( FakeNormal + 0 * float3(0,0,1) );
+		SpecularContribution = TextureCubeSampleLevel(AmbientCubemap, AmbientCubemapSampler, FakeNormal, AbsoluteDiffuseMip).rgb;
+		SpecularContribution *= PI * HairShading( GBuffer, FakeNormal, V, N, 1, 0, 0.2, Random );
+		NonSpecularContribution = 0;
+	}
+
+#if IMPORTANCE_SAMPLE
+	if( GBuffer.ShadingModelID > 0 && BufferUV.x > 0.5 )
+	{
+		NonSpecularContribution = 0;
+
+		BRANCH if( GBuffer.ShadingModelID == SHADINGMODELID_HAIR )
+		{
+			SpecularContribution = ImageBasedLightingHair( GBuffer, -ScreenVector, GBuffer.WorldNormal, Random );
+		}
+		else
+		{
+			SpecularContribution = ImageBasedLightingMIS( GBuffer, -ScreenVector, GBuffer.WorldNormal, Random );
+		}
+	}
+#endif
+
+	// apply darkening from ambient occlusion (does not use PostprocessInput1 to set white texture if SSAO is off)
+	float AmbientOcclusion = GBuffer.GBufferAO * AmbientOcclusionTexture.SampleLevel(AmbientOcclusionSampler, BufferUV, 0).r;
+
+	// Subsurface
+	BRANCH if(GBuffer.ShadingModelID == SHADINGMODELID_SUBSURFACE || GBuffer.ShadingModelID == SHADINGMODELID_PREINTEGRATED_SKIN)
+	{
+		// some view dependent and some non view dependent (hard coded)
+		float DependentSplit = 0.5f;
+
+		float3 SubsurfaceColor = ExtractSubsurfaceColor(GBuffer);
+
+		// view independent (shared lookup for diffuse for better performance
+		NonSpecularContribution += DiffuseLookup * SubsurfaceColor * (DependentSplit);
+		// view dependent (blurriness is hard coded)
+		SpecularContribution += TextureCubeSampleLevel(AmbientCubemap, AmbientCubemapSampler, ScreenVector, AbsoluteDiffuseMip - 2.5f).rgb * SubsurfaceColor * (AmbientOcclusion * (1.0f - DependentSplit));
+	}
+
+	FLightAccumulator LightAccumulator = (FLightAccumulator)0;
+
+	const bool bNeedsSeparateSubsurfaceLightAccumulation = UseSubsurfaceProfile(GBuffer.ShadingModelID);
+
+	LightAccumulator_Add(LightAccumulator, NonSpecularContribution + SpecularContribution, NonSpecularContribution, AmbientCubemapColor.rgb, bNeedsSeparateSubsurfaceLightAccumulation);
+
+	OutColor = LightAccumulator_GetResult(LightAccumulator);
+
+	OutColor *= AmbientOcclusion;
+}
diff --git a/Engine/Shaders/Private/DiffuseIndirectComposite.usf b/Engine/Shaders/Private/DiffuseIndirectComposite.usf
new file mode 100644
--- /dev/null
+++ b/Engine/Shaders/Private/DiffuseIndirectComposite.usf
@@ -0,0 +1,54 @@
+// Copyright 1998-2019 Epic Games, Inc. All Rights Reserved.
+
+#include "Common.ush"
+#include "SceneTextureParameters.ush"
+
+
+float AmbientOcclusionStaticFraction;
+
+Texture2D AmbientOcclusionTexture;
+SamplerState AmbientOcclusionSampler;
+
+Texture2D DiffuseIndirectTexture;
+SamplerState DiffuseIndirectSampler;
+
+
+void MainPS(float4 SvPosition : SV_POSITION, out float4 OutColor : SV_Target0)
+{
+	float2 BufferUV = SvPositionToBufferUV(SvPosition);
+	float2 ScreenPosition = SvPositionToScreenPosition(SvPosition).xy;
+
+	// Sample scene textures.
+	FGBufferData GBuffer = GetGBufferDataFromSceneTextures(BufferUV);
+	
+	// Sample the ambient occlusion that is dynamically generated every frame.
+	float DynamicAmbientOcclusion = 1.0f;
+	#if DIM_APPLY_AMBIENT_OCCLUSION
+		DynamicAmbientOcclusion = AmbientOcclusionTexture.SampleLevel(AmbientOcclusionSampler, BufferUV, 0).r;
+	#endif
+
+	// Compute the final ambient occlusion to be applied.
+	float FinalAmbientOcclusion = GBuffer.GBufferAO * DynamicAmbientOcclusion;
+	
+	OutColor.rgb = 0.0f;
+	OutColor.a = 1.0f;
+
+	// Apply diffuse indirect.
+	#if DIM_APPLY_DIFFUSE_INDIRECT
+	{
+		float3 DiffuseColor = GBuffer.DiffuseColor;
+		if (UseSubsurfaceProfile(GBuffer.ShadingModelID))
+		{
+			DiffuseColor = GBuffer.StoredBaseColor;
+		}
+
+		OutColor.rgb += DiffuseColor * DiffuseIndirectTexture.SampleLevel(DiffuseIndirectSampler, BufferUV, 0).rgb;
+	}
+	#endif
+
+	// Apply AmbientOcclusion to the scene color. since before deferred direct lighting, we assume all lighting in SceneColor to be indirect ligthing
+	{
+		float AOMask = (GBuffer.ShadingModelID != SHADINGMODELID_UNLIT);
+		OutColor.a = lerp(1.0f, FinalAmbientOcclusion, AOMask * AmbientOcclusionStaticFraction);
+	}
+}
diff --git a/Engine/Shaders/Private/PostProcessAmbient.usf b/Engine/Shaders/Private/PostProcessAmbient.usf
deleted file mode 100644
diff --git a/Engine/Shaders/Private/PostProcessAmbientOcclusion.usf b/Engine/Shaders/Private/PostProcessAmbientOcclusion.usf
--- a/Engine/Shaders/Private/PostProcessAmbientOcclusion.usf
+++ b/Engine/Shaders/Private/PostProcessAmbientOcclusion.usf
@@ -897,19 +897,3 @@ void MainSSAOSmoothCS(
 	SSAOSmoothResult[DestPixelPos] = Result * 0.25;
 }
 #endif
-
-
-// used if StaticFraction > small number
-void BasePassAOPS(in noperspective float4 UVAndScreenPos : TEXCOORD0, out float4 OutColor : SV_Target0)
-{
-	float2 UV = UVAndScreenPos.xy;
-
-	FScreenSpaceData ScreenSpaceData = GetScreenSpaceData(UV);
-
-	float StaticFraction = ScreenSpaceAOParams[3].z;
-
-	// can be optimized
-	float AOMask = (ScreenSpaceData.GBuffer.ShadingModelID != SHADINGMODELID_UNLIT);
-
-	OutColor = lerp(1.0f, ScreenSpaceData.AmbientOcclusion * ScreenSpaceData.GBuffer.GBufferAO, AOMask * StaticFraction);
-}
diff --git a/Engine/Shaders/Private/SSRT/SSRTDiffuseIndirect.usf b/Engine/Shaders/Private/SSRT/SSRTDiffuseIndirect.usf
--- a/Engine/Shaders/Private/SSRT/SSRTDiffuseIndirect.usf
+++ b/Engine/Shaders/Private/SSRT/SSRTDiffuseIndirect.usf
@@ -7,6 +7,11 @@
 #include "../Random.ush"
 #include "../BRDF.ush"
 #include "../MonteCarlo.ush"
+#include "../SceneTextureParameters.ush"
+
+
+#define TILE_PIXEL_SIZE 8
+
 
 Texture2D		HZBTexture;
 SamplerState	HZBSampler;
@@ -22,29 +27,33 @@ float4 HZBUvFactorAndInvFactor;
 float4 PrevScreenPositionScaleBias;
 float  PrevSceneColorPreExposureCorrection;
 
-void ScreenSpaceDiffuseIndirectPS(
-	in noperspective float4 UVAndScreenPos : TEXCOORD0,
-	float4 SvPosition : SV_POSITION,
-	out float4 OutColor : SV_Target0,
-	out float OutAO : SV_Target1 )
-{
-	float2 UV = UVAndScreenPos.xy;
-	float2 ScreenPos = UVAndScreenPos.zw;
-	uint2 PixelPos = (uint2)SvPosition.xy;
 
-	OutColor = 0;
+RWTexture2D<float4>	IndirectDiffuseOutput;
+RWTexture2D<float>	AmbientOcclusionOutput;
 
-	FScreenSpaceData ScreenSpaceData = GetScreenSpaceData(UV);
-	FGBufferData GBuffer = ScreenSpaceData.GBuffer;
+
+[numthreads(TILE_PIXEL_SIZE, TILE_PIXEL_SIZE, 1)]
+void MainCS(
+	uint2 DispatchThreadId : SV_DispatchThreadID,
+	uint2 GroupId : SV_GroupID,
+	uint2 GroupThreadId : SV_GroupThreadID,
+	uint GroupThreadIndex : SV_GroupIndex)
+{
+	uint2 PixelPos = DispatchThreadId + View.ViewRectMin.xy;
+	float2 BufferUV = (PixelPos + 0.5) * View.BufferSizeAndInvSize.zw;
+	float2 ScreenPos = (DispatchThreadId + 0.5) * View.ViewSizeAndInvSize.zw;
+
+	FGBufferData GBuffer = GetGBufferDataFromSceneTextures(BufferUV);
 
 	// No SSS because diffuse luminance is in alpha channel and we don't have dual source blending.
 	BRANCH if( GBuffer.ShadingModelID == 0 || GBuffer.ShadingModelID == SHADINGMODELID_SUBSURFACE_PROFILE )
 	{
-		discard;
+		IndirectDiffuseOutput[PixelPos] = 0.0;
+		AmbientOcclusionOutput[PixelPos] = 0.0;
 	}
 
 	const float3 N = GBuffer.WorldNormal;
-	const float SceneDepth = CalcSceneDepth(UV);
+	const float SceneDepth = GBuffer.Depth;
 	const float3 PositionTranslatedWorld = mul( float4( ScreenPos * SceneDepth, SceneDepth, 1 ), View.ScreenToTranslatedWorld ).xyz;
 
 #if QUALITY == 1
@@ -62,17 +71,21 @@ void ScreenSpaceDiffuseIndirectPS(
 #endif
 
 	float2 Noise;
-	Noise.x = InterleavedGradientNoise( SvPosition.xy, View.StateFrameIndexMod8 );
-	Noise.y = InterleavedGradientNoise( SvPosition.xy, View.StateFrameIndexMod8 * 117 );
+	Noise.x = InterleavedGradientNoise(PixelPos + 0.5, View.StateFrameIndexMod8);
+	Noise.y = InterleavedGradientNoise(PixelPos + 0.5, View.StateFrameIndexMod8 * 117);
 		
 	//uint2 Random = 0x10000 * Noise;
 	uint2 Random = Rand3DPCG16( int3( PixelPos, View.StateFrameIndexMod8 ) ).xy;
 			
 	// TODO: lot of VGPR, worth recomputing every ray instead?
 	float3x3 TangentBasis = GetTangentBasis( N );
+	
+	float3 DiffuseColor = 0;
+	float AmbientOcclusion = 0;
 
 	// Shoot multiple rays
-	LOOP for( int i = 0; i < NumRays; i++ )
+	LOOP
+	for (uint i = 0; i < NumRays; i++)
 	{
 		// TOOD: VGPR, might be worth recomputing instead?
 		// SvPosition: is SvPosition really needed? Can recompute from PositionTranslatedWorld?
@@ -106,7 +119,7 @@ void ScreenSpaceDiffuseIndirectPS(
 #if 0	// Backface check
 		if( bHit )
 		{
-			float3 SampleNormal = GetGBufferData( HitUVz.xy ).WorldNormal;
+			float3 SampleNormal = GetGBufferDataFromSceneTextures(HitUVz.xy).WorldNormal;
 			bHit = dot( SampleNormal, L ) < 0;
 		}
 #endif
@@ -125,28 +138,28 @@ void ScreenSpaceDiffuseIndirectPS(
 #endif
 
 			// Backface modulation of hit surface
-			float3 SampleNormal = GetGBufferData( HitUVz.xy ).WorldNormal;
+			float3 SampleNormal = GetGBufferDataFromSceneTextures(HitUVz.xy).WorldNormal;
 			SampleColor.rgb *= saturate( 1 - dot( SampleNormal, L ) );
 
 			SampleColor.rgb *= rcp( 1 + Luminance(SampleColor.rgb) );
 
-			OutColor += SampleColor;
+			DiffuseColor += SampleColor.rgb;
+			AmbientOcclusion = SampleColor.a;
 		}
 	}
 
-	OutColor /= NumRays;
-	OutColor.rgb *= rcp( 1 - Luminance(OutColor.rgb) );
-	OutColor.rgb *= View.IndirectLightingColorScale;
+	DiffuseColor *= rcp(float(NumRays));
+	AmbientOcclusion *= rcp(float(NumRays));
 
-	OutColor.rgb *= GBuffer.DiffuseColor;
-	OutColor.a = 1 - OutColor.a;
+	DiffuseColor *= rcp( 1 - Luminance(DiffuseColor) );
+	DiffuseColor *= View.IndirectLightingColorScale;
 
-	//OutColor.rgb = OutColor.a;
-	//OutColor.a = 0;
-
-#if USE_PREEXPOSURE
-	OutColor.rgb *= PrevSceneColorPreExposureCorrection;
-#endif
+	AmbientOcclusion = 1 - AmbientOcclusion;
 
-	OutAO = OutColor.a;
-}
\ No newline at end of file
+	#if USE_PREEXPOSURE
+		DiffuseColor *= PrevSceneColorPreExposureCorrection;
+	#endif
+	
+	IndirectDiffuseOutput[PixelPos] = float4(DiffuseColor, 0.0);
+	AmbientOcclusionOutput[PixelPos] = AmbientOcclusion;
+} // MainCS()
diff --git a/Engine/Source/Runtime/Renderer/Private/AmbientCubemapParameters.h b/Engine/Source/Runtime/Renderer/Private/AmbientCubemapParameters.h
--- a/Engine/Source/Runtime/Renderer/Private/AmbientCubemapParameters.h
+++ b/Engine/Source/Runtime/Renderer/Private/AmbientCubemapParameters.h
@@ -12,7 +12,7 @@
 
 class FShaderParameterMap;
 
-/** Pixel shader parameters needed for deferred passes. */
+// DEPRECATED: use FAmbientCubemapParameters instead.
 class FCubemapShaderParameters
 {
 public:
@@ -34,3 +34,12 @@ private:
 	void SetParametersTemplate(FRHICommandList& RHICmdList, TRHIShader* ShaderRHI, const FFinalPostProcessSettings::FCubemapEntry& Entry) const;
 };
 
+/** Shader parameters needed for deferred passes sampling the ambient cube map. */
+BEGIN_SHADER_PARAMETER_STRUCT(FAmbientCubemapParameters, )
+	SHADER_PARAMETER(FLinearColor, AmbientCubemapColor)
+	SHADER_PARAMETER(FVector4, AmbientCubemapMipAdjust)
+	SHADER_PARAMETER_TEXTURE(TextureCube, AmbientCubemap)
+	SHADER_PARAMETER_SAMPLER(SamplerState, AmbientCubemapSampler)
+END_SHADER_PARAMETER_STRUCT()
+
+void SetupAmbientCubemapParameters(const FFinalPostProcessSettings::FCubemapEntry& Entry, FAmbientCubemapParameters* OutParameters);
diff --git a/Engine/Source/Runtime/Renderer/Private/CompositionLighting/CompositionLighting.cpp b/Engine/Source/Runtime/Renderer/Private/CompositionLighting/CompositionLighting.cpp
--- a/Engine/Source/Runtime/Renderer/Private/CompositionLighting/CompositionLighting.cpp
+++ b/Engine/Source/Runtime/Renderer/Private/CompositionLighting/CompositionLighting.cpp
@@ -9,7 +9,6 @@
 #include "PostProcess/RenderingCompositionGraph.h"
 #include "PostProcess/PostProcessInput.h"
 #include "PostProcess/PostProcessing.h"
-#include "CompositionLighting/PostProcessAmbient.h"
 #include "CompositionLighting/PostProcessLpvIndirect.h"
 #include "CompositionLighting/PostProcessAmbientOcclusion.h"
 #include "CompositionLighting/PostProcessDeferredDecals.h"
@@ -107,15 +106,6 @@ bool ShouldRenderScreenSpaceAmbientOcclusion(const FViewInfo& View)
 	return bEnabled;
 }
 
-static void AddPostProcessingAmbientCubemap(FPostprocessContext& Context, FRenderingCompositeOutputRef AmbientOcclusion)
-{
-	FRenderingCompositePass* Pass = Context.Graph.RegisterPass(new(FMemStack::Get()) FRCPassPostProcessAmbient());
-	Pass->SetInput(ePId_Input0, Context.FinalOutput);
-	Pass->SetInput(ePId_Input1, AmbientOcclusion);
-
-	Context.FinalOutput = FRenderingCompositeOutputRef(Pass);
-}
-
 // @param Levels 0..3, how many different resolution levels we want to render
 static FRenderingCompositeOutputRef AddPostProcessingAmbientOcclusion(FRHICommandListImmediate& RHICmdList, FPostprocessContext& Context, uint32 Levels)
 {
@@ -351,19 +341,6 @@ void FCompositionLighting::ProcessAfterBasePass(FRHICommandListImmediate& RHICmd
 				}
 
 			}
-
-			if (SceneContext.bScreenSpaceAOIsValid && FSSAOHelper::IsBasePassAmbientOcclusionRequired(Context.View))
-			{
-				FRenderingCompositePass* Pass = Context.Graph.RegisterPass(new(FMemStack::Get()) FRCPassPostProcessBasePassAO());
-				Pass->AddDependency(Context.FinalOutput);
-
-				Context.FinalOutput = FRenderingCompositeOutputRef(Pass);
-			}
-
-			if (IsAmbientCubemapPassRequired(Context.View))
-			{
-				AddPostProcessingAmbientCubemap(Context, AmbientOcclusion);
-			}
 		}
 
 		// The graph setup should be finished before this line ----------------------------------------
diff --git a/Engine/Source/Runtime/Renderer/Private/CompositionLighting/PostProcessAmbient.cpp b/Engine/Source/Runtime/Renderer/Private/CompositionLighting/PostProcessAmbient.cpp
--- a/Engine/Source/Runtime/Renderer/Private/CompositionLighting/PostProcessAmbient.cpp
+++ b/Engine/Source/Runtime/Renderer/Private/CompositionLighting/PostProcessAmbient.cpp
@@ -4,7 +4,6 @@
 	PostProcessAmbient.cpp: Post processing ambient implementation.
 =============================================================================*/
 
-#include "CompositionLighting/PostProcessAmbient.h"
 #include "StaticBoundShaderState.h"
 #include "SceneUtils.h"
 #include "PostProcess/SceneRenderTargets.h"
@@ -15,146 +14,6 @@
 #include "AmbientCubemapParameters.h"
 #include "PipelineStateCache.h"
 
-/** Encapsulates the post processing ambient pixel shader. */
-class FPostProcessAmbientPS : public FGlobalShader
-{
-	DECLARE_SHADER_TYPE(FPostProcessAmbientPS, Global);
-
-	static bool ShouldCompilePermutation(const FGlobalShaderPermutationParameters& Parameters)
-	{
-		return IsFeatureLevelSupported(Parameters.Platform, ERHIFeatureLevel::SM4);
-	}
-
-	/** Default constructor. */
-	FPostProcessAmbientPS() {}
-
-public:
-	FPostProcessPassParameters PostprocessParameter;
-	FCubemapShaderParameters CubemapShaderParameters;
-	FSceneTextureShaderParameters SceneTextureParameters;
-	FShaderResourceParameter PreIntegratedGF;
-	FShaderResourceParameter PreIntegratedGFSampler;
-
-	/** Initialization constructor. */
-	FPostProcessAmbientPS(const ShaderMetaType::CompiledShaderInitializerType& Initializer)
-		: FGlobalShader(Initializer)
-	{
-		PostprocessParameter.Bind(Initializer.ParameterMap);
-		SceneTextureParameters.Bind(Initializer);
-		CubemapShaderParameters.Bind(Initializer.ParameterMap);
-		PreIntegratedGF.Bind(Initializer.ParameterMap, TEXT("PreIntegratedGF"));
-		PreIntegratedGFSampler.Bind(Initializer.ParameterMap, TEXT("PreIntegratedGFSampler"));
-	}
-
-	template <typename TRHICmdList>
-	void SetParameters(TRHICmdList& RHICmdList, const FRenderingCompositePassContext& Context, const FFinalPostProcessSettings::FCubemapEntry& Entry)
-	{
-		FRHIPixelShader* ShaderRHI = GetPixelShader();
-
-		FGlobalShader::SetParameters<FViewUniformShaderParameters>(RHICmdList, ShaderRHI, Context.View.ViewUniformBuffer);
-		PostprocessParameter.SetPS(RHICmdList, ShaderRHI, Context, TStaticSamplerState<SF_Bilinear, AM_Clamp, AM_Clamp, AM_Clamp>::GetRHI());
-		SceneTextureParameters.Set(RHICmdList, ShaderRHI, Context.View.FeatureLevel, ESceneTextureSetupMode::All);
-		CubemapShaderParameters.SetParameters(RHICmdList, ShaderRHI, Entry);
-		SetTextureParameter(RHICmdList, ShaderRHI, PreIntegratedGF, PreIntegratedGFSampler, TStaticSamplerState<SF_Bilinear,AM_Clamp,AM_Clamp,AM_Clamp>::GetRHI(), GSystemTextures.PreintegratedGF->GetRenderTargetItem().ShaderResourceTexture );
-	}
-	
-	// FShader interface.
-	virtual bool Serialize(FArchive& Ar) override
-	{
-		bool bShaderHasOutdatedParameters = FGlobalShader::Serialize(Ar);
-		Ar << PostprocessParameter << CubemapShaderParameters << SceneTextureParameters << PreIntegratedGF << PreIntegratedGFSampler;
-		return bShaderHasOutdatedParameters;
-	}
-};
-
-IMPLEMENT_SHADER_TYPE(, FPostProcessAmbientPS, TEXT("/Engine/Private/PostProcessAmbient.usf"), TEXT("MainPS"), SF_Pixel);
-
-void FRCPassPostProcessAmbient::Render(FRenderingCompositePassContext& Context, FGraphicsPipelineStateInitializer& GraphicsPSOInit)
-{
-	TShaderMapRef<FPostProcessVS> VertexShader(Context.GetShaderMap());
-	TShaderMapRef<FPostProcessAmbientPS> PixelShader(Context.GetShaderMap());
-	const FViewInfo& View = Context.View;
-
-	GraphicsPSOInit.BoundShaderState.VertexDeclarationRHI = GFilterVertexDeclaration.VertexDeclarationRHI;
-	GraphicsPSOInit.BoundShaderState.VertexShaderRHI = GETSAFERHISHADER_VERTEX(*VertexShader);
-	GraphicsPSOInit.BoundShaderState.PixelShaderRHI = GETSAFERHISHADER_PIXEL(*PixelShader);
-	GraphicsPSOInit.PrimitiveType = PT_TriangleList;
-
-	SetGraphicsPipelineState(Context.RHICmdList, GraphicsPSOInit);
-
-	uint32 Count = Context.View.FinalPostProcessSettings.ContributingCubemaps.Num();
-	for (uint32 i = 0; i < Count; ++i)
-	{
-		PixelShader->SetParameters(Context.RHICmdList, Context, Context.View.FinalPostProcessSettings.ContributingCubemaps[i]);
-
-		DrawPostProcessPass(
-			Context.RHICmdList,
-			0, 0,
-			View.ViewRect.Width(), View.ViewRect.Height(),
-			View.ViewRect.Min.X, View.ViewRect.Min.Y,
-			View.ViewRect.Width(), View.ViewRect.Height(),
-			View.ViewRect.Size(),
-			FSceneRenderTargets::Get(Context.RHICmdList).GetBufferSizeXY(),
-			*VertexShader,
-			View.StereoPass, 
-			Context.HasHmdMesh(),
-			EDRF_UseTriangleOptimization);
-	}
-}
-
-void FRCPassPostProcessAmbient::Process(FRenderingCompositePassContext& Context)
-{
-	SCOPED_DRAW_EVENT(Context.RHICmdList, PostProcessAmbient);
-
-	const FViewInfo& View = Context.View;
-	const FSceneViewFamily& ViewFamily = *(View.Family);
-
-	FIntRect SrcRect = View.ViewRect;
-	// todo: view size should scale with input texture size so we can do SSAO in half resolution as well
-	FIntRect DestRect = View.ViewRect;
-	FIntPoint DestSize = DestRect.Size();
-
-	const FSceneRenderTargetItem& DestRenderTarget = FSceneRenderTargets::Get(Context.RHICmdList).GetSceneColor()->GetRenderTargetItem();
-
-	// Set the view family's render target/viewport.
-	FRHIRenderPassInfo RPInfo(DestRenderTarget.TargetableTexture, ERenderTargetActions::Load_Store);
-	TransitionRenderPassTargets(Context.RHICmdList, RPInfo);
-	Context.RHICmdList.BeginRenderPass(RPInfo, TEXT("PostProcessAmbient"));
-	{
-		Context.SetViewportAndCallRHI(View.ViewRect);
-
-		FGraphicsPipelineStateInitializer GraphicsPSOInit;
-		Context.RHICmdList.ApplyCachedRenderTargets(GraphicsPSOInit);
-
-		// set the state
-		GraphicsPSOInit.BlendState = TStaticBlendState<CW_RGB, BO_Add, BF_One, BF_One, BO_Add, BF_One, BF_One>::GetRHI();
-		GraphicsPSOInit.RasterizerState = TStaticRasterizerState<>::GetRHI();
-		GraphicsPSOInit.DepthStencilState = TStaticDepthStencilState<false, CF_Always>::GetRHI();
-
-		FScene* Scene = ViewFamily.Scene->GetRenderScene();
-		check(Scene);
-		uint32 NumReflectionCaptures = Scene->ReflectionSceneData.RegisteredReflectionCaptures.Num();
-
-		// Ambient cubemap specular will be applied in the reflection environment pass if it is enabled
-		//const bool bApplySpecular = View.Family->EngineShowFlags.ReflectionEnvironment == 0 || NumReflectionCaptures == 0;
-		//bool bClearCoatNeeded = (View.ShadingModelMaskInView & (1 << MSM_ClearCoat)) != 0;
-
-		Render(Context, GraphicsPSOInit);
-	}
-	Context.RHICmdList.EndRenderPass();
-	Context.RHICmdList.CopyToResolveTarget(DestRenderTarget.TargetableTexture, DestRenderTarget.ShaderResourceTexture, FResolveParams());
-}
-
-FPooledRenderTargetDesc FRCPassPostProcessAmbient::ComputeOutputDesc(EPassOutputId InPassOutputId) const
-{
-	// we assume this pass is additively blended with the scene color so this data is not needed
-	FPooledRenderTargetDesc Ret;
-
-	Ret.DebugName = TEXT("AmbientCubeMap");
-
-	return Ret;
-}
-
 /*-----------------------------------------------------------------------------
 FCubemapShaderParameters
 -----------------------------------------------------------------------------*/
@@ -180,32 +39,13 @@ void FCubemapShaderParameters::SetParameters(FRHICommandList& RHICmdList, FRHICo
 template<typename TRHIShader>
 void FCubemapShaderParameters::SetParametersTemplate(FRHICommandList& RHICmdList, TRHIShader* ShaderRHI, const FFinalPostProcessSettings::FCubemapEntry& Entry) const
 {
+	FAmbientCubemapParameters ShaderParameters;
+	SetupAmbientCubemapParameters(Entry, &ShaderParameters);
+
 	// floats to render the cubemap
 	{
-		float MipCount = 0;
-
-		if(Entry.AmbientCubemap)
-		{
-			int32 CubemapWidth = Entry.AmbientCubemap->GetSurfaceWidth();
-			MipCount = FMath::Log2(CubemapWidth) + 1.0f;
-		}
-
-		{
-			FLinearColor AmbientCubemapTintMulScaleValue = Entry.AmbientCubemapTintMulScaleValue;
-
-			SetShaderValue(RHICmdList, ShaderRHI, AmbientCubemapColor, AmbientCubemapTintMulScaleValue);
-		}
-
-		{
-			FVector4 AmbientCubemapMipAdjustValue(0, 0, 0, 0);
-
-			AmbientCubemapMipAdjustValue.X =  1.0f - GDiffuseConvolveMipLevel / MipCount;
-			AmbientCubemapMipAdjustValue.Y = (MipCount - 1.0f) * AmbientCubemapMipAdjustValue.X;
-			AmbientCubemapMipAdjustValue.Z = MipCount - GDiffuseConvolveMipLevel;
-			AmbientCubemapMipAdjustValue.W = MipCount;
-
-			SetShaderValue(RHICmdList, ShaderRHI, AmbientCubemapMipAdjust, AmbientCubemapMipAdjustValue);
-		}
+		SetShaderValue(RHICmdList, ShaderRHI, AmbientCubemapColor, ShaderParameters.AmbientCubemapColor);
+		SetShaderValue(RHICmdList, ShaderRHI, AmbientCubemapMipAdjust, ShaderParameters.AmbientCubemapMipAdjust);
 	}
 
 	// cubemap texture
@@ -222,3 +62,31 @@ FArchive& operator<<(FArchive& Ar, FCubemapShaderParameters& P)
 
 	return Ar;
 }
+
+void SetupAmbientCubemapParameters(const FFinalPostProcessSettings::FCubemapEntry& Entry, FAmbientCubemapParameters* OutParameters)
+{
+	// floats to render the cubemap
+	{
+		float MipCount = 0;
+
+		if(Entry.AmbientCubemap)
+		{
+			int32 CubemapWidth = Entry.AmbientCubemap->GetSurfaceWidth();
+			MipCount = FMath::Log2(CubemapWidth) + 1.0f;
+		}
+
+		OutParameters->AmbientCubemapColor = Entry.AmbientCubemapTintMulScaleValue;
+
+		OutParameters->AmbientCubemapMipAdjust.X =  1.0f - GDiffuseConvolveMipLevel / MipCount;
+		OutParameters->AmbientCubemapMipAdjust.Y = (MipCount - 1.0f) * OutParameters->AmbientCubemapMipAdjust.X;
+		OutParameters->AmbientCubemapMipAdjust.Z = MipCount - GDiffuseConvolveMipLevel;
+		OutParameters->AmbientCubemapMipAdjust.W = MipCount;
+	}
+
+	// cubemap texture
+	{
+		FTexture* AmbientCubemapTexture = Entry.AmbientCubemap ? Entry.AmbientCubemap->Resource : GBlackTextureCube;
+		OutParameters->AmbientCubemap = AmbientCubemapTexture->TextureRHI;
+		OutParameters->AmbientCubemapSampler = AmbientCubemapTexture->SamplerStateRHI;
+	}
+}
diff --git a/Engine/Source/Runtime/Renderer/Private/CompositionLighting/PostProcessAmbient.h b/Engine/Source/Runtime/Renderer/Private/CompositionLighting/PostProcessAmbient.h
deleted file mode 100644
diff --git a/Engine/Source/Runtime/Renderer/Private/CompositionLighting/PostProcessAmbientOcclusion.cpp b/Engine/Source/Runtime/Renderer/Private/CompositionLighting/PostProcessAmbientOcclusion.cpp
--- a/Engine/Source/Runtime/Renderer/Private/CompositionLighting/PostProcessAmbientOcclusion.cpp
+++ b/Engine/Source/Runtime/Renderer/Private/CompositionLighting/PostProcessAmbientOcclusion.cpp
@@ -1149,126 +1149,3 @@ FPooledRenderTargetDesc FRCPassPostProcessAmbientOcclusion::ComputeOutputDesc(EP
 
 	return Ret;
 }
-
-// --------------------------------------------------------
-
-
-/** Encapsulates the post processing ambient occlusion pixel shader. */
-class FPostProcessBasePassAOPS : public FGlobalShader
-{
-	DECLARE_SHADER_TYPE(FPostProcessBasePassAOPS, Global);
-
-	static bool ShouldCompilePermutation(const FGlobalShaderPermutationParameters& Parameters)
-	{
-		return IsFeatureLevelSupported(Parameters.Platform, ERHIFeatureLevel::SM4);
-	}
-
-	/** Default constructor. */
-	FPostProcessBasePassAOPS() {}
-
-public:
-	FPostProcessPassParameters PostprocessParameter;
-	FSceneTextureShaderParameters SceneTextureParameters;
-	FScreenSpaceAOParameters ScreenSpaceAOParams;
-
-	/** Initialization constructor. */
-	FPostProcessBasePassAOPS(const ShaderMetaType::CompiledShaderInitializerType& Initializer)
-		: FGlobalShader(Initializer)
-	{
-		PostprocessParameter.Bind(Initializer.ParameterMap);
-		SceneTextureParameters.Bind(Initializer);
-		ScreenSpaceAOParams.Bind(Initializer.ParameterMap);
-	}
-
-	template <typename TRHICmdList>
-	void SetParameters(TRHICmdList& RHICmdList, const FRenderingCompositePassContext& Context, FIntPoint InputTextureSize)
-	{
-		const FFinalPostProcessSettings& Settings = Context.View.FinalPostProcessSettings;
-		FRHIPixelShader* ShaderRHI = GetPixelShader();
-
-		FGlobalShader::SetParameters<FViewUniformShaderParameters>(RHICmdList, ShaderRHI, Context.View.ViewUniformBuffer);
-		PostprocessParameter.SetPS(RHICmdList, ShaderRHI, Context, TStaticSamplerState<SF_Point,AM_Clamp,AM_Clamp,AM_Clamp>::GetRHI());
-		SceneTextureParameters.Set(RHICmdList, ShaderRHI, Context.View.FeatureLevel, ESceneTextureSetupMode::All);
-		ScreenSpaceAOParams.Set(RHICmdList, Context.View, ShaderRHI, InputTextureSize);
-	}
-
-	// FShader interface.
-	virtual bool Serialize(FArchive& Ar) override
-	{
-		bool bShaderHasOutdatedParameters = FGlobalShader::Serialize(Ar);
-		Ar << PostprocessParameter << SceneTextureParameters << ScreenSpaceAOParams;
-		return bShaderHasOutdatedParameters;
-	}
-};
-
-IMPLEMENT_SHADER_TYPE(,FPostProcessBasePassAOPS,TEXT("/Engine/Private/PostProcessAmbientOcclusion.usf"),TEXT("BasePassAOPS"),SF_Pixel);
-
-// --------------------------------------------------------
-
-void FRCPassPostProcessBasePassAO::Process(FRenderingCompositePassContext& Context)
-{
-	SCOPED_GPU_STAT(Context.RHICmdList, BasePassAO);
-
-	const FViewInfo& View = Context.View;
-
-	SCOPED_DRAW_EVENTF(Context.RHICmdList, ApplyAOToBasePassSceneColor, TEXT("ApplyAOToBasePassSceneColor %dx%d"), View.ViewRect.Width(), View.ViewRect.Height());
-
-	FSceneRenderTargets& SceneContext = FSceneRenderTargets::Get(Context.RHICmdList);
-
-	const FSceneRenderTargetItem& DestRenderTarget = SceneContext.GetSceneColor()->GetRenderTargetItem();
-
-	// Set the view family's render target/viewport.
-	Context.RHICmdList.TransitionResource(EResourceTransitionAccess::EWritable, DestRenderTarget.TargetableTexture);
-
-	FRHIRenderPassInfo RPInfo(DestRenderTarget.TargetableTexture, ERenderTargetActions::Load_Store);
-	Context.RHICmdList.BeginRenderPass(RPInfo, TEXT("BasePassAmbientOcclusion"));
-	{
-		Context.SetViewportAndCallRHI(View.ViewRect);
-
-		FGraphicsPipelineStateInitializer GraphicsPSOInit;
-		Context.RHICmdList.ApplyCachedRenderTargets(GraphicsPSOInit);
-
-		// set the state
-		GraphicsPSOInit.BlendState = TStaticBlendState<CW_RGBA, BO_Add, BF_DestColor, BF_Zero, BO_Add, BF_DestAlpha, BF_Zero>::GetRHI();
-		GraphicsPSOInit.RasterizerState = TStaticRasterizerState<>::GetRHI();
-		GraphicsPSOInit.DepthStencilState = TStaticDepthStencilState<false, CF_Always>::GetRHI();
-
-		TShaderMapRef<FPostProcessVS> VertexShader(Context.GetShaderMap());
-		TShaderMapRef<FPostProcessBasePassAOPS> PixelShader(Context.GetShaderMap());
-
-		GraphicsPSOInit.BoundShaderState.VertexDeclarationRHI = GFilterVertexDeclaration.VertexDeclarationRHI;
-		GraphicsPSOInit.BoundShaderState.VertexShaderRHI = GETSAFERHISHADER_VERTEX(*VertexShader);
-		GraphicsPSOInit.BoundShaderState.PixelShaderRHI = GETSAFERHISHADER_PIXEL(*PixelShader);
-		GraphicsPSOInit.PrimitiveType = PT_TriangleList;
-
-		SetGraphicsPipelineState(Context.RHICmdList, GraphicsPSOInit);
-
-		VertexShader->SetParameters(Context);
-		PixelShader->SetParameters(Context.RHICmdList, Context, SceneContext.GetBufferSizeXY());
-
-		DrawPostProcessPass(
-			Context.RHICmdList,
-			0, 0,
-			View.ViewRect.Width(), View.ViewRect.Height(),
-			View.ViewRect.Min.X, View.ViewRect.Min.Y,
-			View.ViewRect.Width(), View.ViewRect.Height(),
-			View.ViewRect.Size(),
-			SceneContext.GetBufferSizeXY(),
-			*VertexShader,
-			View.StereoPass,
-			Context.HasHmdMesh(),
-			EDRF_UseTriangleOptimization);
-	}
-	Context.RHICmdList.EndRenderPass();
-	Context.RHICmdList.CopyToResolveTarget(DestRenderTarget.TargetableTexture, DestRenderTarget.ShaderResourceTexture, FResolveParams());
-}
-
-FPooledRenderTargetDesc FRCPassPostProcessBasePassAO::ComputeOutputDesc(EPassOutputId InPassOutputId) const
-{
-	// we assume this pass is additively blended with the scene color so this data is not needed
-	FPooledRenderTargetDesc Ret;
-
-	Ret.DebugName = TEXT("SceneColorWithAO");
-
-	return Ret;
-}
diff --git a/Engine/Source/Runtime/Renderer/Private/CompositionLighting/PostProcessAmbientOcclusion.h b/Engine/Source/Runtime/Renderer/Private/CompositionLighting/PostProcessAmbientOcclusion.h
--- a/Engine/Source/Runtime/Renderer/Private/CompositionLighting/PostProcessAmbientOcclusion.h
+++ b/Engine/Source/Runtime/Renderer/Private/CompositionLighting/PostProcessAmbientOcclusion.h
@@ -123,14 +123,3 @@ private:
 	const bool bAOSetupAsInput;
 	const bool bForceIntermediateOutput;
 };
-
-// apply the AO to the SceneColor (lightmapped object), extra pas that is not always needed
-// derives from TRenderingCompositePassBase<InputCount, OutputCount> 
-class FRCPassPostProcessBasePassAO : public TRenderingCompositePassBase<0, 1>
-{
-public:
-	// interface FRenderingCompositePass ---------
-	virtual void Process(FRenderingCompositePassContext& Context) override;
-	virtual void Release() override { delete this; }
-	virtual FPooledRenderTargetDesc ComputeOutputDesc(EPassOutputId InPassOutputId) const override;
-};
\ No newline at end of file
diff --git a/Engine/Source/Runtime/Renderer/Private/DeferredShadingRenderer.cpp b/Engine/Source/Runtime/Renderer/Private/DeferredShadingRenderer.cpp
--- a/Engine/Source/Runtime/Renderer/Private/DeferredShadingRenderer.cpp
+++ b/Engine/Source/Runtime/Renderer/Private/DeferredShadingRenderer.cpp
@@ -393,8 +393,6 @@ static FORCEINLINE bool NeedsPrePass(const FDeferredShadingSceneRenderer* Render
 		(Renderer->EarlyZPassMode != DDM_None || Renderer->bEarlyZPassMovable != 0);
 }
 
-bool ShouldRenderScreenSpaceDiffuseIndirect( const FViewInfo& View );
-
 bool FDeferredShadingSceneRenderer::RenderHzb(FRHICommandListImmediate& RHICmdList)
 {
 	FSceneRenderTargets& SceneContext = FSceneRenderTargets::Get(RHICmdList);
@@ -1596,8 +1594,6 @@ void FDeferredShadingSceneRenderer::Render(FRHICommandListImmediate& RHICmdList)
 			Scene->UniformBuffers.UpdateViewUniformBuffer(Views[ViewIndex]);
 
 			GCompositionLighting.ProcessAfterBasePass(RHICmdList, Views[ViewIndex]);
-			
-			RenderScreenSpaceDiffuseIndirect( RHICmdList, Views[ ViewIndex ], SceneContext.SceneVelocity );
 		}
 		ServiceLocalQueue();
 	}
@@ -1626,6 +1622,8 @@ void FDeferredShadingSceneRenderer::Render(FRHICommandListImmediate& RHICmdList)
 
 		GRenderTargetPool.AddPhaseEvent(TEXT("Lighting"));
 
+		RenderDiffuseIndirectAndAmbientOcclusion(RHICmdList);
+
 		// These modulate the scenecolor output from the basepass, which is assumed to be indirect lighting
 		RenderIndirectCapsuleShadows(
 			RHICmdList, 
diff --git a/Engine/Source/Runtime/Renderer/Private/DeferredShadingRenderer.h b/Engine/Source/Runtime/Renderer/Private/DeferredShadingRenderer.h
--- a/Engine/Source/Runtime/Renderer/Private/DeferredShadingRenderer.h
+++ b/Engine/Source/Runtime/Renderer/Private/DeferredShadingRenderer.h
@@ -192,6 +192,8 @@ private:
 	/** Renders the scene's atmosphere. */
 	void RenderAtmosphere(FRHICommandListImmediate& RHICmdList, const FLightShaftsOutput& LightShaftsOutput);
 
+	void RenderDiffuseIndirectAndAmbientOcclusion(FRHICommandListImmediate& RHICmdList);
+
 	/** Renders sky lighting and reflections that can be done in a deferred pass. */
 	void RenderDeferredReflectionsAndSkyLighting(FRHICommandListImmediate& RHICmdList, TRefCountPtr<IPooledRenderTarget>& DynamicBentNormalAO, TRefCountPtr<IPooledRenderTarget>& VelocityRT);
 
diff --git a/Engine/Source/Runtime/Renderer/Private/IndirectLightRendering.cpp b/Engine/Source/Runtime/Renderer/Private/IndirectLightRendering.cpp
new file mode 100644
--- /dev/null
+++ b/Engine/Source/Runtime/Renderer/Private/IndirectLightRendering.cpp
@@ -0,0 +1,312 @@
+// Copyright 1998-2019 Epic Games, Inc. All Rights Reserved.
+
+#include "IndirectLightRendering.h"
+#include "RenderGraph.h"
+#include "PixelShaderUtils.h"
+#include "AmbientCubemapParameters.h"
+#include "SceneTextureParameters.h"
+#include "ScreenSpaceDenoise.h"
+#include "ScreenSpaceRayTracing.h"
+#include "DeferredShadingRenderer.h"
+#include "PostProcessing.h" // for FPostProcessVS
+
+
+static TAutoConsoleVariable<int32> CVarDiffuseIndirectDenoiser(
+	TEXT("r.DiffuseIndirect.Denoiser"), 1,
+	TEXT("Denoising options (default = 1)"),
+	ECVF_RenderThreadSafe);
+
+
+bool IsAmbientCubemapPassRequired(const FSceneView& View);
+
+
+class FDiffuseIndirectCompositePS : public FGlobalShader
+{
+	DECLARE_GLOBAL_SHADER(FDiffuseIndirectCompositePS)
+	SHADER_USE_PARAMETER_STRUCT(FDiffuseIndirectCompositePS, FGlobalShader)
+
+	class FApplyDiffuseIndirectDim : SHADER_PERMUTATION_BOOL("DIM_APPLY_DIFFUSE_INDIRECT");
+	class FApplyAmbientOcclusionDim : SHADER_PERMUTATION_BOOL("DIM_APPLY_AMBIENT_OCCLUSION");
+
+	using FPermutationDomain = TShaderPermutationDomain<
+		FApplyDiffuseIndirectDim,
+		FApplyAmbientOcclusionDim>;
+
+	static bool ShouldCompilePermutation(const FGlobalShaderPermutationParameters& Parameters)
+	{
+		FPermutationDomain PermutationVector(Parameters.PermutationId);
+
+		// Do not compile a shader that does not apply anything.
+		if (!PermutationVector.Get<FApplyDiffuseIndirectDim>() &&
+			!PermutationVector.Get<FApplyAmbientOcclusionDim>())
+		{
+			return false;
+		}
+
+		// Diffuse indirect generation is SM5 only.
+		if (PermutationVector.Get<FApplyDiffuseIndirectDim>())
+		{
+			return IsFeatureLevelSupported(Parameters.Platform, ERHIFeatureLevel::SM5);
+		}
+
+		return IsFeatureLevelSupported(Parameters.Platform, ERHIFeatureLevel::SM4);
+	}
+
+	BEGIN_SHADER_PARAMETER_STRUCT(FParameters, )
+		SHADER_PARAMETER(float, AmbientOcclusionStaticFraction)
+
+		SHADER_PARAMETER_RDG_TEXTURE(Texture2D, DiffuseIndirectTexture)
+		SHADER_PARAMETER_SAMPLER(SamplerState, DiffuseIndirectSampler)
+		
+		SHADER_PARAMETER_RDG_TEXTURE(Texture2D, AmbientOcclusionTexture)
+		SHADER_PARAMETER_SAMPLER(SamplerState,  AmbientOcclusionSampler)
+		
+		SHADER_PARAMETER_STRUCT_INCLUDE(FSceneTextureParameters, SceneTextures)
+		SHADER_PARAMETER_STRUCT_INCLUDE(FSceneTextureSamplerParameters, SceneTextureSamplers)
+		SHADER_PARAMETER_STRUCT_REF(FViewUniformShaderParameters, ViewUniformBuffer)
+		
+		RENDER_TARGET_BINDING_SLOTS()
+	END_SHADER_PARAMETER_STRUCT()
+};
+
+class FAmbientCubemapCompositePS : public FGlobalShader
+{
+	DECLARE_GLOBAL_SHADER(FAmbientCubemapCompositePS)
+	SHADER_USE_PARAMETER_STRUCT(FAmbientCubemapCompositePS, FGlobalShader)
+
+	static bool ShouldCompilePermutation(const FGlobalShaderPermutationParameters& Parameters)
+	{
+		return IsFeatureLevelSupported(Parameters.Platform, ERHIFeatureLevel::SM4);
+	}
+
+	BEGIN_SHADER_PARAMETER_STRUCT(FParameters, )
+		SHADER_PARAMETER_TEXTURE(Texture2D, PreIntegratedGF)
+		SHADER_PARAMETER_SAMPLER(SamplerState, PreIntegratedGFSampler)
+
+		SHADER_PARAMETER_RDG_TEXTURE(Texture2D, AmbientOcclusionTexture)
+		SHADER_PARAMETER_SAMPLER(SamplerState,  AmbientOcclusionSampler)
+		
+		SHADER_PARAMETER_STRUCT_INCLUDE(FAmbientCubemapParameters, AmbientCubemap)
+		SHADER_PARAMETER_STRUCT_INCLUDE(FSceneTextureParameters, SceneTextures)
+		SHADER_PARAMETER_STRUCT_INCLUDE(FSceneTextureSamplerParameters, SceneTextureSamplers)
+		SHADER_PARAMETER_STRUCT_REF(FViewUniformShaderParameters, ViewUniformBuffer)
+		
+		RENDER_TARGET_BINDING_SLOTS()
+	END_SHADER_PARAMETER_STRUCT()
+};
+
+IMPLEMENT_GLOBAL_SHADER(FDiffuseIndirectCompositePS, "/Engine/Private/DiffuseIndirectComposite.usf", "MainPS", SF_Pixel);
+IMPLEMENT_GLOBAL_SHADER(FAmbientCubemapCompositePS, "/Engine/Private/AmbientCubemapComposite.usf", "MainPS", SF_Pixel);
+
+
+void FDeferredShadingSceneRenderer::RenderDiffuseIndirectAndAmbientOcclusion(FRHICommandListImmediate& RHICmdListImmediate)
+{
+	SCOPED_DRAW_EVENT(RHICmdListImmediate, DiffuseIndirectAndAO)
+
+	FSceneRenderTargets& SceneContext = FSceneRenderTargets::Get(RHICmdListImmediate);
+
+	// Forwared shading SSAO is applied before the basepass using only the depth buffer.
+	if (IsForwardShadingEnabled(ViewFamily.GetShaderPlatform()))
+	{
+		return;
+	}
+
+	FRDGBuilder GraphBuilder(RHICmdListImmediate);
+	
+	FSceneTextureParameters SceneTextures;
+	SetupSceneTextureParameters(GraphBuilder, &SceneTextures);
+	
+	FRDGTextureRef SceneColor = GraphBuilder.RegisterExternalTexture(SceneContext.GetSceneColor());
+
+	for (FViewInfo& View : Views)
+	{
+		const bool bApplySSGI = ShouldRenderScreenSpaceDiffuseIndirect(View);
+		const bool bApplySSAO = SceneContext.bScreenSpaceAOIsValid;
+
+		int32 DenoiseMode = CVarDiffuseIndirectDenoiser.GetValueOnRenderThread();
+
+		
+		IScreenSpaceDenoiser::FAmbientOcclusionRayTracingConfig RayTracingConfig;
+
+		// TODO: integrate RTGI.
+		// TODO: hybrid SSGI / RTGI
+		IScreenSpaceDenoiser::FDiffuseIndirectInputs DenoiserInputs;
+		if (bApplySSGI)
+		{
+			RenderScreenSpaceDiffuseIndirect(GraphBuilder, SceneTextures, SceneColor, View, /* out */ &DenoiserInputs);
+			
+			// TODO: Denoise.
+			DenoiseMode = 0;
+		}
+		else
+		{
+			// No need for denoising.
+			DenoiseMode = 0;
+		}
+		
+		IScreenSpaceDenoiser::FDiffuseIndirectOutputs DenoiserOutputs;
+		if (DenoiseMode != 0)
+		{
+			const IScreenSpaceDenoiser* DefaultDenoiser = IScreenSpaceDenoiser::GetDefaultDenoiser();
+			const IScreenSpaceDenoiser* DenoiserToUse = DenoiseMode == 1 ? DefaultDenoiser : GScreenSpaceDenoiser;
+
+			RDG_EVENT_SCOPE(GraphBuilder, "%s%s(DiffuseIndirect) %dx%d",
+				DenoiserToUse != DefaultDenoiser ? TEXT("ThirdParty ") : TEXT(""),
+				DenoiserToUse->GetDebugName(),
+				View.ViewRect.Width(), View.ViewRect.Height());
+
+			DenoiserOutputs = DenoiserToUse->DenoiseDiffuseIndirect(
+				GraphBuilder,
+				View,
+				&View.PrevViewInfo,
+				SceneTextures,
+				DenoiserInputs,
+				RayTracingConfig);
+		}
+		else
+		{
+			DenoiserOutputs.Color = DenoiserInputs.Color;
+			DenoiserOutputs.AmbientOcclusionMask = DenoiserInputs.AmbientOcclusionMask;
+		}
+
+		// Extract the dynamic AO for further down lide
+		if (DenoiserOutputs.AmbientOcclusionMask)
+		{
+			ensureMsgf(!bApplySSAO, TEXT("Looks like SSAO has been computed for this view but is being overridden."));
+			ensureMsgf(Views.Num() == 1, TEXT("SSGI can only output AO for 1 view")); // TODO.
+			GraphBuilder.QueueTextureExtraction(DenoiserOutputs.AmbientOcclusionMask, &SceneContext.ScreenSpaceAO);
+			SceneContext.bScreenSpaceAOIsValid = true;
+		}
+		else if (bApplySSAO)
+		{
+			DenoiserOutputs.AmbientOcclusionMask = GraphBuilder.RegisterExternalTexture(SceneContext.ScreenSpaceAO);
+		}
+
+		// Applies diffuse indirect and ambient occlusion to the scene color.
+		if (DenoiserOutputs.Color || bApplySSAO)
+		{
+			FDiffuseIndirectCompositePS::FParameters* PassParameters = GraphBuilder.AllocParameters<FDiffuseIndirectCompositePS::FParameters>();
+			
+			PassParameters->AmbientOcclusionStaticFraction = FMath::Clamp(View.FinalPostProcessSettings.AmbientOcclusionStaticFraction, 0.0f, 1.0f);
+
+			PassParameters->DiffuseIndirectTexture = DenoiserOutputs.Color;
+			PassParameters->DiffuseIndirectSampler = TStaticSamplerState<SF_Point>::GetRHI();
+
+			PassParameters->AmbientOcclusionTexture = DenoiserOutputs.AmbientOcclusionMask;
+			PassParameters->AmbientOcclusionSampler = TStaticSamplerState<SF_Point>::GetRHI();
+			
+			PassParameters->SceneTextures = SceneTextures;
+			SetupSceneTextureSamplers(&PassParameters->SceneTextureSamplers);
+			PassParameters->ViewUniformBuffer = View.ViewUniformBuffer;
+
+			PassParameters->RenderTargets[0] = FRenderTargetBinding(
+				SceneColor, ERenderTargetLoadAction::ELoad, ERenderTargetStoreAction::EStore);
+		
+			FDiffuseIndirectCompositePS::FPermutationDomain PermutationVector;
+			PermutationVector.Set<FDiffuseIndirectCompositePS::FApplyDiffuseIndirectDim>(PassParameters->DiffuseIndirectTexture != nullptr);
+			PermutationVector.Set<FDiffuseIndirectCompositePS::FApplyAmbientOcclusionDim>(PassParameters->AmbientOcclusionTexture != nullptr);
+
+			TShaderMapRef<FDiffuseIndirectCompositePS> PixelShader(View.ShaderMap, PermutationVector);
+			ClearUnusedGraphResources(*PixelShader, PassParameters);
+
+			GraphBuilder.AddPass(
+				RDG_EVENT_NAME(
+					"DiffuseIndirectComposite(ApplyAO=%s ApplyDiffuseIndirect=%s) %dx%d",
+					PermutationVector.Get<FDiffuseIndirectCompositePS::FApplyAmbientOcclusionDim>() ? TEXT("Yes") : TEXT("No"),
+					PermutationVector.Get<FDiffuseIndirectCompositePS::FApplyDiffuseIndirectDim>() ? TEXT("Yes") : TEXT("No"),
+					View.ViewRect.Width(), View.ViewRect.Height()),
+				PassParameters,
+				ERenderGraphPassFlags::None,
+				[PassParameters, &View, PixelShader, PermutationVector](FRHICommandList& RHICmdList)
+			{
+				RHICmdList.SetViewport(View.ViewRect.Min.X, View.ViewRect.Min.Y, 0.0f, View.ViewRect.Max.X, View.ViewRect.Max.Y, 0.0);
+				
+				FGraphicsPipelineStateInitializer GraphicsPSOInit;
+				FPixelShaderUtils::InitFullscreenPipelineState(RHICmdList, View.ShaderMap, *PixelShader, /* out */ GraphicsPSOInit);
+				
+				if (PermutationVector.Get<FDiffuseIndirectCompositePS::FApplyAmbientOcclusionDim>())
+				{
+					GraphicsPSOInit.BlendState = TStaticBlendState<CW_RGBA, BO_Add, BF_DestAlpha, BF_One, BO_Add, BF_DestAlpha, BF_Zero>::GetRHI();
+				}
+				else
+				{
+					GraphicsPSOInit.BlendState = TStaticBlendState<CW_RGB, BO_Add, BF_One, BF_One, BO_Add, BF_One, BF_Zero>::GetRHI();
+				}
+				SetGraphicsPipelineState(RHICmdList, GraphicsPSOInit);
+
+				SetShaderParameters(RHICmdList, *PixelShader, PixelShader->GetPixelShader(), *PassParameters);
+
+				FPixelShaderUtils::DrawFullscreenTriangle(RHICmdList);
+			});
+		} // if (DenoiserOutputs.Color || bApplySSAO)
+
+		// Apply the ambient cubemaps
+		if (IsAmbientCubemapPassRequired(View))
+		{
+			FAmbientCubemapCompositePS::FParameters* PassParameters = GraphBuilder.AllocParameters<FAmbientCubemapCompositePS::FParameters>();
+			
+			PassParameters->PreIntegratedGF = GSystemTextures.PreintegratedGF->GetRenderTargetItem().ShaderResourceTexture;
+			PassParameters->PreIntegratedGFSampler = TStaticSamplerState<SF_Bilinear, AM_Clamp, AM_Clamp, AM_Clamp>::GetRHI();
+			
+			PassParameters->AmbientOcclusionTexture = DenoiserOutputs.AmbientOcclusionMask;
+			PassParameters->AmbientOcclusionSampler = TStaticSamplerState<SF_Point>::GetRHI();
+			
+			PassParameters->SceneTextures = SceneTextures;
+			SetupSceneTextureSamplers(&PassParameters->SceneTextureSamplers);
+			PassParameters->ViewUniformBuffer = View.ViewUniformBuffer;
+
+			PassParameters->RenderTargets[0] = FRenderTargetBinding(
+				SceneColor, ERenderTargetLoadAction::ELoad, ERenderTargetStoreAction::EStore);
+		
+			TShaderMapRef<FAmbientCubemapCompositePS> PixelShader(View.ShaderMap);
+			GraphBuilder.AddPass(
+				RDG_EVENT_NAME("AmbientCubemapComposite %dx%d", View.ViewRect.Width(), View.ViewRect.Height()),
+				PassParameters,
+				ERDGPassFlags::None,
+				[PassParameters, &View, PixelShader](FRHICommandList& RHICmdList)
+			{
+				TShaderMapRef<FPostProcessVS> VertexShader(View.ShaderMap);
+				
+				RHICmdList.SetViewport(View.ViewRect.Min.X, View.ViewRect.Min.Y, 0.0f, View.ViewRect.Max.X, View.ViewRect.Max.Y, 0.0);
+
+				FGraphicsPipelineStateInitializer GraphicsPSOInit;
+				RHICmdList.ApplyCachedRenderTargets(GraphicsPSOInit);
+
+				// set the state
+				GraphicsPSOInit.BlendState = TStaticBlendState<CW_RGB, BO_Add, BF_One, BF_One, BO_Add, BF_One, BF_One>::GetRHI();
+				GraphicsPSOInit.RasterizerState = TStaticRasterizerState<>::GetRHI();
+				GraphicsPSOInit.DepthStencilState = TStaticDepthStencilState<false, CF_Always>::GetRHI();
+
+				GraphicsPSOInit.BoundShaderState.VertexDeclarationRHI = GFilterVertexDeclaration.VertexDeclarationRHI;
+				GraphicsPSOInit.BoundShaderState.VertexShaderRHI = GETSAFERHISHADER_VERTEX(*VertexShader);
+				GraphicsPSOInit.BoundShaderState.PixelShaderRHI = GETSAFERHISHADER_PIXEL(*PixelShader);
+				GraphicsPSOInit.PrimitiveType = PT_TriangleList;
+
+				SetGraphicsPipelineState(RHICmdList, GraphicsPSOInit);
+
+				uint32 Count = View.FinalPostProcessSettings.ContributingCubemaps.Num();
+				for (const FFinalPostProcessSettings::FCubemapEntry& CubemapEntry : View.FinalPostProcessSettings.ContributingCubemaps)
+				{
+					FAmbientCubemapCompositePS::FParameters ShaderParameters = *PassParameters;
+					SetupAmbientCubemapParameters(CubemapEntry, &ShaderParameters.AmbientCubemap);
+					SetShaderParameters(RHICmdList, *PixelShader, PixelShader->GetPixelShader(), ShaderParameters);
+					
+					DrawPostProcessPass(
+						RHICmdList,
+						0, 0,
+						View.ViewRect.Width(), View.ViewRect.Height(),
+						View.ViewRect.Min.X, View.ViewRect.Min.Y,
+						View.ViewRect.Width(), View.ViewRect.Height(),
+						View.ViewRect.Size(),
+						FSceneRenderTargets::Get(RHICmdList).GetBufferSizeXY(),
+						*VertexShader,
+						View.StereoPass, 
+						false, // TODO.
+						EDRF_UseTriangleOptimization);
+				}
+			});
+		} // if (IsAmbientCubemapPassRequired(View))
+	} // for (FViewInfo& View : Views)
+
+	GraphBuilder.Execute();
+}
diff --git a/Engine/Source/Runtime/Renderer/Private/IndirectLightRendering.h b/Engine/Source/Runtime/Renderer/Private/IndirectLightRendering.h
new file mode 100644
--- /dev/null
+++ b/Engine/Source/Runtime/Renderer/Private/IndirectLightRendering.h
@@ -0,0 +1,4 @@
+// Copyright 1998-2019 Epic Games, Inc. All Rights Reserved.
+
+#include "ScenePrivate.h"
+
diff --git a/Engine/Source/Runtime/Renderer/Private/ScreenSpaceDenoise.h b/Engine/Source/Runtime/Renderer/Private/ScreenSpaceDenoise.h
--- a/Engine/Source/Runtime/Renderer/Private/ScreenSpaceDenoise.h
+++ b/Engine/Source/Runtime/Renderer/Private/ScreenSpaceDenoise.h
@@ -160,6 +160,11 @@ public:
 	BEGIN_SHADER_PARAMETER_STRUCT(FDiffuseIndirectInputs, )
 		// Irradiance in RGB, AO mask in alpha.
 		SHADER_PARAMETER_RDG_TEXTURE(Texture2D, Color)
+
+		// Ambient occlusion mask stored in the red channel as [0; 1].
+		SHADER_PARAMETER_RDG_TEXTURE(Texture2D, AmbientOcclusionMask)
+
+		// Hit distance in world space.
 		SHADER_PARAMETER_RDG_TEXTURE(Texture2D, RayHitDistance)
 	END_SHADER_PARAMETER_STRUCT()
 
@@ -167,6 +172,9 @@ public:
 	BEGIN_SHADER_PARAMETER_STRUCT(FDiffuseIndirectOutputs, )
 		// Irradiance in RGB, AO mask in alpha.
 		SHADER_PARAMETER_RDG_TEXTURE(Texture2D, Color)
+		
+		// Ambient occlusion mask stored in the red channel as [0; 1].
+		SHADER_PARAMETER_RDG_TEXTURE(Texture2D, AmbientOcclusionMask)
 	END_SHADER_PARAMETER_STRUCT()
 
 	/** Entry point to denoise reflections. */
diff --git a/Engine/Source/Runtime/Renderer/Private/ScreenSpaceRayTracing.cpp b/Engine/Source/Runtime/Renderer/Private/ScreenSpaceRayTracing.cpp
--- a/Engine/Source/Runtime/Renderer/Private/ScreenSpaceRayTracing.cpp
+++ b/Engine/Source/Runtime/Renderer/Private/ScreenSpaceRayTracing.cpp
@@ -79,11 +79,11 @@ bool ShouldRenderScreenSpaceReflections(const FViewInfo& View)
 	return true;
 }
 
-bool ShouldRenderScreenSpaceDiffuseIndirect( const FViewInfo& View )
+bool ShouldRenderScreenSpaceDiffuseIndirect(const FViewInfo& View)
 {
 	int Quality = CVarSSGIQuality.GetValueOnRenderThread();
 
-	if( Quality <= 0 )
+	if (Quality <= 0)
 	{
 		return false;
 	}
@@ -93,7 +93,7 @@ bool ShouldRenderScreenSpaceDiffuseIndirect( const FViewInfo& View )
 		return false;
 	}
 
-	return true;
+	return View.PrevViewInfo.TemporalAAHistory.IsValid();
 }
 
 bool IsSSRTemporalPassRequired(const FViewInfo& View)
@@ -219,10 +219,10 @@ class FScreenSpaceReflectionsPS : public FGlobalShader
 	END_SHADER_PARAMETER_STRUCT()
 };
 
-class FScreenSpaceDiffuseIndirectPS : public FGlobalShader
+class FScreenSpaceDiffuseIndirectCS : public FGlobalShader
 {
-	DECLARE_GLOBAL_SHADER(FScreenSpaceDiffuseIndirectPS);
-	SHADER_USE_PARAMETER_STRUCT(FScreenSpaceDiffuseIndirectPS, FGlobalShader)
+	DECLARE_GLOBAL_SHADER(FScreenSpaceDiffuseIndirectCS);
+	SHADER_USE_PARAMETER_STRUCT(FScreenSpaceDiffuseIndirectCS, FGlobalShader)
 
 	class FQualityDim : SHADER_PERMUTATION_INT( "QUALITY", 5 );
 	using FPermutationDomain = TShaderPermutationDomain< FQualityDim >;
@@ -232,8 +232,6 @@ class FScreenSpaceDiffuseIndirectPS : public FGlobalShader
 		SHADER_PARAMETER( FVector4,		PrevScreenPositionScaleBias )
 		SHADER_PARAMETER( float,		PrevSceneColorPreExposureCorrection )
 		
-		RENDER_TARGET_BINDING_SLOTS()
-
 		SHADER_PARAMETER_RDG_TEXTURE( Texture2D,	HZBTexture )
 		SHADER_PARAMETER_SAMPLER( SamplerState,		HZBSampler )
 
@@ -242,9 +240,13 @@ class FScreenSpaceDiffuseIndirectPS : public FGlobalShader
 
 		SHADER_PARAMETER_RDG_TEXTURE( Texture2D,	ColorTexture )
 		SHADER_PARAMETER_SAMPLER( SamplerState,		ColorSampler )
+		
+		SHADER_PARAMETER_STRUCT_INCLUDE(FSceneTextureParameters, SceneTextures)
+		SHADER_PARAMETER_STRUCT_INCLUDE(FSceneTextureSamplerParameters, SceneTextureSamplers)
+		SHADER_PARAMETER_STRUCT_REF(FViewUniformShaderParameters, View)
 
-		SHADER_PARAMETER_STRUCT_REF(	FViewUniformShaderParameters,	View )
-		SHADER_PARAMETER_STRUCT_REF(	FSceneTexturesUniformParameters,SceneTextures )
+		SHADER_PARAMETER_RDG_TEXTURE_UAV(RWTexture<float4>, IndirectDiffuseOutput)
+		SHADER_PARAMETER_RDG_TEXTURE_UAV(RWTexture<float>,  AmbientOcclusionOutput)
 	END_SHADER_PARAMETER_STRUCT()
 
 	static bool ShouldCompilePermutation(const FGlobalShaderPermutationParameters& Parameters)
@@ -256,7 +258,7 @@ class FScreenSpaceDiffuseIndirectPS : public FGlobalShader
 
 IMPLEMENT_GLOBAL_SHADER(FScreenSpaceReflectionsPS,        "/Engine/Private/SSRT/SSRTReflections.usf", "ScreenSpaceReflectionsPS", SF_Pixel);
 IMPLEMENT_GLOBAL_SHADER(FScreenSpaceReflectionsStencilPS, "/Engine/Private/SSRT/SSRTReflections.usf", "ScreenSpaceReflectionsStencilPS", SF_Pixel);
-IMPLEMENT_GLOBAL_SHADER(FScreenSpaceDiffuseIndirectPS,    "/Engine/Private/SSRT/SSRTDiffuseIndirect.usf", "ScreenSpaceDiffuseIndirectPS", SF_Pixel);
+IMPLEMENT_GLOBAL_SHADER(FScreenSpaceDiffuseIndirectCS,    "/Engine/Private/SSRT/SSRTDiffuseIndirect.usf", "MainCS", SF_Compute);
 
 
 void GetSSRShaderOptionsForQuality(ESSRQuality Quality, IScreenSpaceDenoiser::FReflectionsRayTracingConfig* OutRayTracingConfigs)
@@ -349,7 +351,7 @@ void RenderScreenSpaceReflections(
 		FRDGTextureDesc Desc = FPooledRenderTargetDesc::Create2DDesc(
 			FSceneRenderTargets::Get_FrameConstantsOnly().GetBufferSizeXY(),
 			PF_FloatRGBA, FClearValueBinding(FLinearColor(0, 0, 0, 0)),
-			TexCreate_None, TexCreate_RenderTargetable,
+			TexCreate_None, TexCreate_RenderTargetable | TexCreate_ShaderResource | TexCreate_UAV,
 			false);
 
 		Desc.AutoWritable = false;
@@ -509,46 +511,47 @@ void RenderScreenSpaceReflections(
 	}
 } // RenderScreenSpaceReflections()
 
-void RenderScreenSpaceDiffuseIndirect( FRHICommandListImmediate& RHICmdList, FViewInfo& View, TRefCountPtr<IPooledRenderTarget>& VelocityRT )
+void RenderScreenSpaceDiffuseIndirect(
+	FRDGBuilder& GraphBuilder, 
+	const FSceneTextureParameters& SceneTextures,
+	const FRDGTextureRef CurrentSceneColor,
+	const FViewInfo& View,
+	IScreenSpaceDenoiser::FDiffuseIndirectInputs* OutDenoiserInputs)
 {
-	FTemporalAAHistory& TemporalAAHistory = View.PrevViewInfo.TemporalAAHistory;
+	check(ShouldRenderScreenSpaceDiffuseIndirect(View));
 
-	if( !ShouldRenderScreenSpaceDiffuseIndirect( View ) || !TemporalAAHistory.IsValid() )
-	{
-		return;
-	}
+	const FTemporalAAHistory& TemporalAAHistory = View.PrevViewInfo.TemporalAAHistory;
+	check(TemporalAAHistory.IsValid()); // TODO.
 	
 	const int32 Quality = FMath::Clamp( CVarSSGIQuality.GetValueOnRenderThread(), 1, 4 );
 
-	FRDGBuilder GraphBuilder( RHICmdList );
+	FSceneRenderTargets& SceneContext = FSceneRenderTargets::Get(GraphBuilder.RHICmdList);
 
-	FSceneRenderTargets& SceneContext = FSceneRenderTargets::Get( RHICmdList );
+	// Allocate outputs.
+	{
+		FRDGTextureDesc Desc = FRDGTextureDesc::Create2DDesc(
+			SceneTextures.SceneDepthBuffer->Desc.Extent,
+			PF_FloatRGBA,
+			FClearValueBinding::None,
+			/* InFlags = */ TexCreate_None,
+			/* InTargetableFlags = */ TexCreate_ShaderResource | TexCreate_RenderTargetable | TexCreate_UAV,
+			/* bInForceSeparateTargetAndShaderResource = */ false);
 
-	FRDGTexture* SceneColor	= GraphBuilder.RegisterExternalTexture( SceneContext.GetSceneColor() );
-	FRDGTexture* ScreenSpaceAO = GraphBuilder.CreateTexture(SceneContext.ScreenSpaceAO->GetDesc(), TEXT("SSRTAO"));
-	FRDGTexture* HZBTexture	= GraphBuilder.RegisterExternalTexture( View.HZB );
-	FRDGTexture* ColorTexture	= GraphBuilder.RegisterExternalTexture( TemporalAAHistory.RT[0] );
+		OutDenoiserInputs->Color = GraphBuilder.CreateTexture(Desc, TEXT("SSRTDiffuseIndirect"));
 
-	FRDGTexture* VelocityTexture;
-	if( VelocityRT && !View.bCameraCut )
-	{
-		VelocityTexture = GraphBuilder.RegisterExternalTexture( VelocityRT );
-	}
-	else
-	{
-		// No velocity, use black
-		VelocityTexture = GraphBuilder.RegisterExternalTexture( GSystemTextures.BlackDummy );
+		Desc.Format = PF_R16F;
+		OutDenoiserInputs->AmbientOcclusionMask = GraphBuilder.CreateTexture(Desc, TEXT("SSRTAmbientOcclusion"));
 	}
 
-	FScreenSpaceDiffuseIndirectPS::FParameters* PassParameters = GraphBuilder.AllocParameters< FScreenSpaceDiffuseIndirectPS::FParameters >();
+	FRDGTexture* HZBTexture	= GraphBuilder.RegisterExternalTexture( View.HZB );
+	FRDGTexture* ColorTexture	= GraphBuilder.RegisterExternalTexture( TemporalAAHistory.RT[0] );
 
-	PassParameters->RenderTargets[0] = FRenderTargetBinding( SceneColor, ERenderTargetLoadAction::ELoad, ERenderTargetStoreAction::EStore );
-	PassParameters->RenderTargets[1] = FRenderTargetBinding( ScreenSpaceAO, ERenderTargetLoadAction::ENoAction, ERenderTargetStoreAction::EStore );
+	FScreenSpaceDiffuseIndirectCS::FParameters* PassParameters = GraphBuilder.AllocParameters<FScreenSpaceDiffuseIndirectCS::FParameters>();
 
 	PassParameters->HZBTexture = HZBTexture;
 	PassParameters->HZBSampler = TStaticSamplerState< SF_Point, AM_Clamp, AM_Clamp, AM_Clamp >::GetRHI();
 
-	PassParameters->VelocityTexture = VelocityTexture;
+	PassParameters->VelocityTexture = SceneTextures.SceneVelocityBuffer;
 	PassParameters->VelocitySampler = TStaticSamplerState< SF_Point, AM_Clamp, AM_Clamp, AM_Clamp >::GetRHI();
 
 	PassParameters->ColorTexture = ColorTexture;
@@ -577,58 +580,21 @@ void RenderScreenSpaceDiffuseIndirect( FRHICommandListImmediate& RHICmdList, FVi
 
 	PassParameters->PrevSceneColorPreExposureCorrection = View.PreExposure / View.PrevViewInfo.SceneColorPreExposure;
 
+	PassParameters->SceneTextures = SceneTextures;
+	SetupSceneTextureSamplers(&PassParameters->SceneTextureSamplers);
 	PassParameters->View = View.ViewUniformBuffer;
-	PassParameters->SceneTextures = CreateSceneTextureUniformBufferSingleDraw( RHICmdList, ESceneTextureSetupMode::SceneDepth | ESceneTextureSetupMode::GBuffers, View.FeatureLevel );
-
-	GraphBuilder.AddPass(
-		RDG_EVENT_NAME( "ScreenSpaceDiffuseIndirect(Quality=%d) %dx%d", Quality, View.ViewRect.Width(), View.ViewRect.Height() ),
-		PassParameters,
-		ERenderGraphPassFlags::None,
-		[ PassParameters, &View, Quality ]( FRHICommandListImmediate& InRHICmdList )
-	{
-		SCOPED_GPU_STAT(InRHICmdList, ScreenSpaceDiffuseIndirect);
-
-		InRHICmdList.SetViewport( View.ViewRect.Min.X, View.ViewRect.Min.Y, 0.0f, View.ViewRect.Max.X, View.ViewRect.Max.Y, 1.0f );
-
-		FScreenSpaceDiffuseIndirectPS::FPermutationDomain PermutationVector;
-		PermutationVector.Set< FScreenSpaceDiffuseIndirectPS::FQualityDim >( Quality );
-
-		auto VertexShader	= View.ShaderMap->GetShader< FScreenPassVS >();
-		auto PixelShader	= View.ShaderMap->GetShader< FScreenSpaceDiffuseIndirectPS >( PermutationVector );
-		CA_ASSUME(PixelShader);
-
-		{
-			FGraphicsPipelineStateInitializer GraphicsPSOInit;
-			InRHICmdList.ApplyCachedRenderTargets( GraphicsPSOInit );
-
-			GraphicsPSOInit.BlendState = TStaticBlendState< CW_RGB, BO_Add, BF_One, BF_SourceAlpha >::GetRHI();
-			GraphicsPSOInit.DepthStencilState = TStaticDepthStencilState<false, CF_Always>::GetRHI();
-			GraphicsPSOInit.RasterizerState = TStaticRasterizerState<>::GetRHI();
-			GraphicsPSOInit.PrimitiveType = PT_TriangleList;
-			GraphicsPSOInit.BoundShaderState.VertexDeclarationRHI = GFilterVertexDeclaration.VertexDeclarationRHI;
-			GraphicsPSOInit.BoundShaderState.VertexShaderRHI = GETSAFERHISHADER_VERTEX( VertexShader );
-			GraphicsPSOInit.BoundShaderState.PixelShaderRHI = GETSAFERHISHADER_PIXEL( PixelShader );
+	
+	PassParameters->IndirectDiffuseOutput = GraphBuilder.CreateUAV(OutDenoiserInputs->Color);
+	PassParameters->AmbientOcclusionOutput = GraphBuilder.CreateUAV(OutDenoiserInputs->AmbientOcclusionMask);
 
-			SetGraphicsPipelineState( InRHICmdList, GraphicsPSOInit );
-		}
+	FScreenSpaceDiffuseIndirectCS::FPermutationDomain PermutationVector;
+	PermutationVector.Set<FScreenSpaceDiffuseIndirectCS::FQualityDim>(Quality);
 
-		SetShaderParameters( InRHICmdList, PixelShader, PixelShader->GetPixelShader(), *PassParameters );
-
-		DrawPostProcessPass(
-			InRHICmdList,
-			0, 0,
-			View.ViewRect.Width(), View.ViewRect.Height(),
-			View.ViewRect.Min.X, View.ViewRect.Min.Y,
-			View.ViewRect.Width(), View.ViewRect.Height(),
-			View.ViewRect.Size(),
-			FSceneRenderTargets::Get( InRHICmdList ).GetBufferSizeXY(),
-			VertexShader,
-			View.StereoPass,
-			false,
-			EDRF_UseTriangleOptimization);
-	});
-
-	GraphBuilder.QueueTextureExtraction(ScreenSpaceAO, &SceneContext.ScreenSpaceAO);
-
-	GraphBuilder.Execute();
+	TShaderMapRef<FScreenSpaceDiffuseIndirectCS> ComputeShader(View.ShaderMap, PermutationVector);
+	FComputeShaderUtils::AddPass(
+		GraphBuilder,
+		RDG_EVENT_NAME("ScreenSpaceDiffuseIndirect(Quality=%d) %dx%d", Quality, View.ViewRect.Width(), View.ViewRect.Height()),
+		*ComputeShader,
+		PassParameters,
+		FComputeShaderUtils::GetGroupCount(View.ViewRect.Size(), 8));
 } // RenderScreenSpaceDiffuseIndirect()
diff --git a/Engine/Source/Runtime/Renderer/Private/ScreenSpaceRayTracing.h b/Engine/Source/Runtime/Renderer/Private/ScreenSpaceRayTracing.h
--- a/Engine/Source/Runtime/Renderer/Private/ScreenSpaceRayTracing.h
+++ b/Engine/Source/Runtime/Renderer/Private/ScreenSpaceRayTracing.h
@@ -22,6 +22,8 @@ enum class ESSRQuality
 
 bool ShouldRenderScreenSpaceReflections(const FViewInfo& View);
 
+bool ShouldRenderScreenSpaceDiffuseIndirect(const FViewInfo& View);
+
 void GetSSRQualityForView(const FViewInfo& View, ESSRQuality* OutQuality, IScreenSpaceDenoiser::FReflectionsRayTracingConfig* OutRayTracingConfigs);
 
 bool IsSSRTemporalPassRequired(const FViewInfo& View);
@@ -35,4 +37,9 @@ void RenderScreenSpaceReflections(
 	bool bDenoiser,
 	IScreenSpaceDenoiser::FReflectionsInputs* DenoiserInputs);
 
-void RenderScreenSpaceDiffuseIndirect( FRHICommandListImmediate& RHICmdList, FViewInfo& View, TRefCountPtr<IPooledRenderTarget>& VelocityRT );
+void RenderScreenSpaceDiffuseIndirect(
+	FRDGBuilder& GraphBuilder, 
+	const FSceneTextureParameters& SceneTextures,
+	const FRDGTextureRef SceneColor,
+	const FViewInfo& View,
+	IScreenSpaceDenoiser::FDiffuseIndirectInputs* OutDenoiserInputs);
