Fixed unbound parameter in ray tracing translucency

diff --git a/Engine/Shaders/Private/RayTracing/RayTracingTranslucency.usf b/Engine/Shaders/Private/RayTracing/RayTracingTranslucency.usf
--- a/Engine/Shaders/Private/RayTracing/RayTracingTranslucency.usf
+++ b/Engine/Shaders/Private/RayTracing/RayTracingTranslucency.usf
@@ -24,6 +24,7 @@
 #include "../PathTracing/Utilities/PathTracingRandomSequence.ush" 
 #include "../HeightFogCommon.ush"
 #include "../SobolRandom.ush"
+#include "../SceneTextureParameters.ush"
 #include "RayTracingCommon.ush"
 #include "RayTracingHitGroupCommon.ush"
 #include "RayTracingDirectionalLight.ush"
@@ -46,6 +47,8 @@ float TranslucencyMaxRoughness;
 uint  TranslucencyRefraction;
 float MaxNormalBias;
 
+Texture2D SceneColorTexture;
+
 RaytracingAccelerationStructure TLAS;
 
 RWTexture2D<float4> ColorOutput;
@@ -106,8 +109,8 @@ void RayTracingTranslucencyRGS()
 
 	float2 InvBufferSize = View.BufferSizeAndInvSize.zw;
 	float2 UV = (float2(PixelCoord) + 0.5) * InvBufferSize;
-	FScreenSpaceData ScreenSpaceData = GetScreenSpaceData(UV);
-	float Depth = ScreenSpaceData.GBuffer.Depth;
+	FGBufferData GBufferData = GetGBufferDataFromSceneTextures(UV);
+	float Depth = GBufferData.Depth;
 	float3 CameraOrigin = View.WorldCameraOrigin;
 	float3 WorldPosition = ReconstructWorldPositionFromDepth(UV, Depth);
 	float3 CameraDirection = normalize(WorldPosition - CameraOrigin);	
@@ -258,7 +261,7 @@ void RayTracingTranslucencyRGS()
 	if (!bHasScattered)
 	{
 		// Use the scene radiance for ray that has not been scattered/refracted (no surface or IORin=IORout). Still apply the throughtput in case we have traversed surfaces with opacity>0.
-		PathRadiance += PathThroughput * Texture2DSampleLevel(SceneTexturesStruct.SceneColorTexture, SceneTexturesStruct.SceneColorTextureSampler, UV, 0).xyz;
+		PathRadiance += PathThroughput * SceneColorTexture.SampleLevel(GlobalPointClampedSampler, UV, 0).xyz;
 	}
 
 #if USE_PREEXPOSURE
diff --git a/Engine/Source/Runtime/Renderer/Private/RayTracing/RayTracingTranslucency.cpp b/Engine/Source/Runtime/Renderer/Private/RayTracing/RayTracingTranslucency.cpp
--- a/Engine/Source/Runtime/Renderer/Private/RayTracing/RayTracingTranslucency.cpp
+++ b/Engine/Source/Runtime/Renderer/Private/RayTracing/RayTracingTranslucency.cpp
@@ -152,6 +152,8 @@ class FRayTracingTranslucencyRGS : public FGlobalShader
 		SHADER_PARAMETER_STRUCT_INCLUDE(FSceneTextureParameters, SceneTextures)
 		SHADER_PARAMETER_STRUCT_INCLUDE(FSceneTextureSamplerParameters, SceneTextureSamplers)
 
+		SHADER_PARAMETER_RDG_TEXTURE(Texture2D, SceneColorTexture)
+
 		SHADER_PARAMETER_RDG_TEXTURE_UAV(RWTexture2D<float4>, ColorOutput)
 		SHADER_PARAMETER_RDG_TEXTURE_UAV(RWTexture2D<float>, RayHitDistanceOutput)
 	END_SHADER_PARAMETER_STRUCT()
@@ -392,6 +394,9 @@ void FDeferredShadingSceneRenderer::RenderRayTracingTranslucencyView(
 	PassParameters->SceneTextures = SceneTextures;
 	PassParameters->SceneTextureSamplers = SceneTextureSamplers;
 
+	FRDGTextureRef SceneColorTexture = GraphBuilder.RegisterExternalTexture(SceneContext.GetSceneColor(), TEXT("SceneColor"));
+	PassParameters->SceneColorTexture = SceneColorTexture;
+
 	PassParameters->ReflectionStruct = CreateReflectionUniformBuffer(View, EUniformBufferUsage::UniformBuffer_SingleFrame);
 	PassParameters->FogUniformParameters = CreateFogUniformBuffer(View, EUniformBufferUsage::UniformBuffer_SingleFrame);
 
