Fixes shader compilation failure when r.Shaders.Optimize=0

diff --git a/Engine/Shaders/Private/RayTracing/RayTracingGlobalIlluminationCompositePS.usf b/Engine/Shaders/Private/RayTracing/RayTracingGlobalIlluminationCompositePS.usf
--- a/Engine/Shaders/Private/RayTracing/RayTracingGlobalIlluminationCompositePS.usf
+++ b/Engine/Shaders/Private/RayTracing/RayTracingGlobalIlluminationCompositePS.usf
@@ -31,7 +31,6 @@ void GlobalIlluminationSceneColorCompositePS(
 	out float4 OutColor : SV_Target0
 )
 {
-	FScreenSpaceData ScreenSpaceData = GetScreenSpaceData(UV);
 	float4 GlobalIllumination = GlobalIlluminationTexture.Sample(GlobalIlluminationSampler, UV);
 	OutColor = float4(GlobalIllumination.rgb, 0.0);
 }
\ No newline at end of file
