Fixes ambient occlusion reconstruction pass having wrong hard coded number of input texture.

diff --git a/Engine/Shaders/Private/ScreenSpaceDenoise/SSDSpatialAccumulation.usf b/Engine/Shaders/Private/ScreenSpaceDenoise/SSDSpatialAccumulation.usf
--- a/Engine/Shaders/Private/ScreenSpaceDenoise/SSDSpatialAccumulation.usf
+++ b/Engine/Shaders/Private/ScreenSpaceDenoise/SSDSpatialAccumulation.usf
@@ -190,14 +190,14 @@
 	#define CONFIG_VIEW_POSITION_BASED_BILATERAL 2
 	#define CONFIG_NORMAL_BASED_BILATERAL 1
 
-	#define CONFIG_INPUT_TEXTURE_COUNT 1
-	#define CONFIG_OUTPUT_TEXTURE_COUNT 1
-
 	#if DIM_STAGE == STAGE_RECONSTRUCTION
 		// Input and output layout.
 		#define CONFIG_SIGNAL_INPUT_LAYOUT  SIGNAL_BUFFER_LAYOUT_AO_INPUT
 		#define CONFIG_SIGNAL_OUTPUT_LAYOUT SIGNAL_BUFFER_LAYOUT_AO_HISTORY
 
+		#define CONFIG_INPUT_TEXTURE_COUNT 2
+		#define CONFIG_OUTPUT_TEXTURE_COUNT 1
+
 		#define CONFIG_SAMPLE_SET           SAMPLE_SET_STACKOWIAK_4_SETS
 		#define CONFIG_BILATERAL_DISTANCE_COMPUTATION SIGNAL_WORLD_FREQUENCY_HIT_DISTANCE
 		#define CONFIG_OUTPUT_MODE          OUTPUT_MODE_DRB
@@ -207,6 +207,9 @@
 		#define CONFIG_SIGNAL_INPUT_LAYOUT  SIGNAL_BUFFER_LAYOUT_AO_HISTORY
 		#define CONFIG_SIGNAL_OUTPUT_LAYOUT SIGNAL_BUFFER_LAYOUT_AO_REJECTION
 
+		#define CONFIG_INPUT_TEXTURE_COUNT 1
+		#define CONFIG_OUTPUT_TEXTURE_COUNT 1
+
 		#define CONFIG_BILATERAL_DISTANCE_COMPUTATION SIGNAL_WORLD_FREQUENCY_HIT_DISTANCE
 
 	#elif DIM_STAGE == STAGE_POST_FILTERING
@@ -214,6 +217,9 @@
 		#define CONFIG_SIGNAL_INPUT_LAYOUT  SIGNAL_BUFFER_LAYOUT_AO_HISTORY
 		#define CONFIG_SIGNAL_OUTPUT_LAYOUT SIGNAL_BUFFER_LAYOUT_AO_HISTORY
 
+		#define CONFIG_INPUT_TEXTURE_COUNT 1
+		#define CONFIG_OUTPUT_TEXTURE_COUNT 1
+
 		#define CONFIG_SAMPLE_SET           SAMPLE_SET_STACKOWIAK_4_SETS
 		#define CONFIG_SAMPLE_SUBSET        1
 
