Fix forward rendering path where GBuffer isn't available.

diff --git a/Engine/Shaders/Private/ShadowProjectionPixelShader.usf b/Engine/Shaders/Private/ShadowProjectionPixelShader.usf
--- a/Engine/Shaders/Private/ShadowProjectionPixelShader.usf
+++ b/Engine/Shaders/Private/ShadowProjectionPixelShader.usf
@@ -182,18 +182,24 @@ void Main(
 	}
 	#else // !USE_PCSS
 	{
-		// Attenuate soft transition based on the angle with the light and the shading normal (acts as a receiver bias)
-		FGBufferData GBufferData = GetGBufferData(ScreenUV);
-		const bool bIsDirectional = LightPositionOrDirection.w == 0;
-		const float3 LightDirection = bIsDirectional ? -LightPositionOrDirection.xyz : normalize(LightPositionOrDirection.xyz - WorldPosition);
-		const float NoL = saturate(dot(GBufferData.WorldNormal, LightDirection));
+		#ifdef SHADING_PATH_DEFERRED
+			// Attenuate soft transition based on the angle with the light and the shading normal (acts as a receiver bias)
+			FGBufferData GBufferData = GetGBufferData(ScreenUV);
+			const bool bIsDirectional = LightPositionOrDirection.w == 0;
+			const float3 LightDirection = bIsDirectional ? -LightPositionOrDirection.xyz : normalize(LightPositionOrDirection.xyz - WorldPosition);
+			const float NoL = saturate(dot(GBufferData.WorldNormal, LightDirection));
+		#endif
 		
 		FPCFSamplerSettings Settings;
 		
 		Settings.ShadowDepthTexture = ShadowDepthTexture;
 		Settings.ShadowDepthTextureSampler = ShadowDepthTextureSampler;
 		Settings.ShadowBufferSize = ShadowBufferSize;
+		#if SHADING_PATH_DEFERRED
 		Settings.TransitionScale = SoftTransitionScale.z * lerp(ProjectionDepthBiasParameters.z, 1.0, NoL);
+		#else
+		Settings.TransitionScale = SoftTransitionScale.z;
+		#endif
 		Settings.SceneDepth = LightSpacePixelDepthForOpaque;
 		Settings.bSubsurface = false;
 		Settings.bTreatMaxDepthUnshadowed = false;
