Fix multiple definitions of ComputeRayConeLod() with the exact same signature. This cause GLSL to fail to compile on ES2/ES3/ES3.1

diff --git a/Engine/Shaders/Private/Common.ush b/Engine/Shaders/Private/Common.ush
--- a/Engine/Shaders/Private/Common.ush
+++ b/Engine/Shaders/Private/Common.ush
@@ -83,14 +83,14 @@ float ComputeRayConeLod(Texture2D Tex)
 	int TexArea = Dimensions.x * Dimensions.y;
 	return 0.5f * log2(GlobalRayCone_TexArea * TexArea);
 #else
-    return 0;
+    return 0.0f;
 #endif
 }
 
 #if SUPPORTS_TEXTURE_EXTERNAL
-float ComputeRayConeLod(TextureExternal Tex)
+float ComputeRayConeLodExternal(TextureExternal Tex)
 {
-	return 0;
+	return 0.0f;
 }
 #endif
 
@@ -201,7 +201,11 @@ MaterialFloat4 TextureCubeSampleGrad(TextureCube Tex, SamplerState Sampler, floa
 MaterialFloat4 TextureExternalSample(TextureExternal Tex, SamplerState Sampler, float2 UV)
 {
 #if !PIXELSHADER
-	return Tex.SampleLevel(Sampler, UV, ComputeRayConeLod(Tex));
+	#if SUPPORTS_TEXTURE_EXTERNAL
+		return Tex.SampleLevel(Sampler, UV, ComputeRayConeLodExternal(Tex));
+	#else
+		return Tex.SampleLevel(Sampler, UV, 0);
+	#endif
 #else
 	return Tex.Sample(Sampler, UV);
 #endif
