Incorporate feedback to changes form mobile hlsl changes.
  - Still allow external textures sample to use ComputeRayConeLod() if external texture is not defined.

diff --git a/Engine/Shaders/Private/Common.ush b/Engine/Shaders/Private/Common.ush
--- a/Engine/Shaders/Private/Common.ush
+++ b/Engine/Shaders/Private/Common.ush
@@ -87,14 +87,6 @@ float ComputeRayConeLod(Texture2D Tex)
 #endif
 }
 
-#if SUPPORTS_TEXTURE_EXTERNAL
-float ComputeRayConeLodExternal(TextureExternal Tex)
-{
-	return 0.0f;
-}
-#endif
-
-
 float  ClampToHalfFloatRange(float  X) { return clamp(X, float(0), MaxHalfFloat); }
 float2 ClampToHalfFloatRange(float2 X) { return clamp(X, float(0).xx, MaxHalfFloat.xx); }
 float3 ClampToHalfFloatRange(float3 X) { return clamp(X, float(0).xxx, MaxHalfFloat.xxx); }
@@ -202,9 +194,9 @@ MaterialFloat4 TextureExternalSample(TextureExternal Tex, SamplerState Sampler,
 {
 #if !PIXELSHADER
 	#if SUPPORTS_TEXTURE_EXTERNAL
-		return Tex.SampleLevel(Sampler, UV, ComputeRayConeLodExternal(Tex));
-	#else
 		return Tex.SampleLevel(Sampler, UV, 0);
+	#else
+		return Tex.SampleLevel(Sampler, UV, ComputeRayConeLod(Tex));
 	#endif
 #else
 	return Tex.Sample(Sampler, UV);
