Fix for when we have SHADING_PATH_DEFERRED AND FORWARD_SHADING defined!

diff --git a/Engine/Shaders/Private/ShadowProjectionPixelShader.usf b/Engine/Shaders/Private/ShadowProjectionPixelShader.usf
--- a/Engine/Shaders/Private/ShadowProjectionPixelShader.usf
+++ b/Engine/Shaders/Private/ShadowProjectionPixelShader.usf
@@ -182,7 +182,7 @@ void Main(
 	}
 	#else // !USE_PCSS
 	{
-		#ifdef SHADING_PATH_DEFERRED
+		#ifdef SHADING_PATH_DEFERRED && !FORWARD_SHADING
 			// Attenuate soft transition based on the angle with the light and the shading normal (acts as a receiver bias)
 			FGBufferData GBufferData = GetGBufferData(ScreenUV);
 			const bool bIsDirectional = LightPositionOrDirection.w == 0;
@@ -195,7 +195,7 @@ void Main(
 		Settings.ShadowDepthTexture = ShadowDepthTexture;
 		Settings.ShadowDepthTextureSampler = ShadowDepthTextureSampler;
 		Settings.ShadowBufferSize = ShadowBufferSize;
-		#if SHADING_PATH_DEFERRED
+		#if SHADING_PATH_DEFERRED && !FORWARD_SHADING
 		Settings.TransitionScale = SoftTransitionScale.z * lerp(ProjectionDepthBiasParameters.z, 1.0, NoL);
 		#else
 		Settings.TransitionScale = SoftTransitionScale.z;
