Change #ifdef to #if

diff --git a/Engine/Shaders/Private/ShadowProjectionPixelShader.usf b/Engine/Shaders/Private/ShadowProjectionPixelShader.usf
--- a/Engine/Shaders/Private/ShadowProjectionPixelShader.usf
+++ b/Engine/Shaders/Private/ShadowProjectionPixelShader.usf
@@ -182,7 +182,7 @@ void Main(
 	}
 	#else // !USE_PCSS
 	{
-		#ifdef SHADING_PATH_DEFERRED && !FORWARD_SHADING
+		#if SHADING_PATH_DEFERRED && !FORWARD_SHADING
 			// Attenuate soft transition based on the angle with the light and the shading normal (acts as a receiver bias)
 			FGBufferData GBufferData = GetGBufferData(ScreenUV);
 			const bool bIsDirectional = LightPositionOrDirection.w == 0;
