Fixes DFAO skylight shadowing

diff --git a/Engine/Shaders/Private/ReflectionEnvironmentPixelShader.usf b/Engine/Shaders/Private/ReflectionEnvironmentPixelShader.usf
--- a/Engine/Shaders/Private/ReflectionEnvironmentPixelShader.usf
+++ b/Engine/Shaders/Private/ReflectionEnvironmentPixelShader.usf
@@ -73,7 +73,7 @@ float3 SkyLightDiffuse(FGBufferData GBuffer, float AmbientOcclusion, float2 Buff
 	
 	float3 DiffuseIrradiance = 0;
 
-#if APPLY_SKY_SHADOWING && 0
+#if APPLY_SKY_SHADOWING
 #define USE_DIRECTIONAL_OCCLUSION_ON_SKY_DIFFUSE 1
 
 #if USE_DIRECTIONAL_OCCLUSION_ON_SKY_DIFFUSE
