Fixes incorrect depth compare tolerance on SSR

diff --git a/Engine/Shaders/Private/SSRT/SSRTRayCast.ush b/Engine/Shaders/Private/SSRT/SSRTRayCast.ush
--- a/Engine/Shaders/Private/SSRT/SSRTRayCast.ush
+++ b/Engine/Shaders/Private/SSRT/SSRTRayCast.ush
@@ -47,7 +47,7 @@ bool RayCast(
 	#if IS_SSGI_SHADER
 		float CompareTolerance = max( abs( RayStepUVz.z ), ( RayStartScreen.z - RayDepthScreen.z ) * 2 ) * Step;
 	#else
-		float CompareTolerance = abs( RayDepthScreen.z - RayStartScreen.z ) * Step * 4;
+		float CompareTolerance = max( abs( RayStepUVz.z ), ( RayStartScreen.z - RayDepthScreen.z ) * 4 ) * Step;
 	#endif
 
 	// avoid bugs with early returns inside of loops on certain platform compilers.
