Brings back force sharps SSR with Roughness < 0.1 that was removed by 5100295

diff --git a/Engine/Shaders/Private/SSRT/SSRTReflections.usf b/Engine/Shaders/Private/SSRT/SSRTReflections.usf
--- a/Engine/Shaders/Private/SSRT/SSRTReflections.usf
+++ b/Engine/Shaders/Private/SSRT/SSRTReflections.usf
@@ -164,7 +164,7 @@ void ScreenSpaceReflectionsPS(
 		float2 Noise;
 		Noise.x = InterleavedGradientNoise( SvPosition.xy, View.StateFrameIndexMod8 );
 		Noise.y = InterleavedGradientNoise( SvPosition.xy, View.StateFrameIndexMod8 * 117 );
-		
+	
 		//uint2 Random = 0x10000 * Noise;
 		uint2 Random = Rand3DPCG16( int3( PixelPos, View.StateFrameIndexMod8 ) ).xy;
 			
@@ -221,6 +221,11 @@ void ScreenSpaceReflectionsPS(
 				
 			float3 HitUVz;
 			float Level = 0;
+			
+			if( Roughness < 0.1 )
+			{
+				L = reflect(-V, N);
+			}
 
 			bool bHit = RayCast(
 				HZB, HZBSampler,
