Fixed post merge shader compilation error

diff --git a/Engine/Shaders/Private/LocalVertexFactory.ush b/Engine/Shaders/Private/LocalVertexFactory.ush
--- a/Engine/Shaders/Private/LocalVertexFactory.ush
+++ b/Engine/Shaders/Private/LocalVertexFactory.ush
@@ -349,13 +349,20 @@ float4x4 GetInstanceTransform(FPositionOnlyVertexFactoryInput Input)
 
 float4x4 GetInstanceTransform(FPositionAndNormalOnlyVertexFactoryInput Input)
 {
-	#if !USE_INSTANCING_EMULATED && MANUAL_VERTEX_FETCH
+	#if !USE_INSTANCING_EMULATED && MANUAL_VERTEX_FETCH && !USE_INSTANCING_BONEMAP
 	uint InstanceId = GetInstanceId(Input.InstanceId);
 	return float4x4(
 		float4(VertexFetch_InstanceTransformBuffer[3 * (InstanceId + InstanceOffset) + 0].xyz, 0.0f),
 		float4(VertexFetch_InstanceTransformBuffer[3 * (InstanceId + InstanceOffset) + 1].xyz, 0.0f),
 		float4(VertexFetch_InstanceTransformBuffer[3 * (InstanceId + InstanceOffset) + 2].xyz, 0.0f),
 		float4(VertexFetch_InstanceOriginBuffer[(InstanceId + InstanceOffset)].xyz, 1.0f));
+	#elif !USE_INSTANCING_EMULATED && MANUAL_VERTEX_FETCH && USE_INSTANCING_BONEMAP
+		uint InstanceIndex = VertexFetch_InstanceBoneMapBuffer[LocalVF.VertexFetch_Parameters[VF_VertexOffset] + Input.VertexId];
+		return float4x4(
+			float4(VertexFetch_InstanceTransformBuffer[4 * InstanceIndex + 0].xyz, 0.0f),
+			float4(VertexFetch_InstanceTransformBuffer[4 * InstanceIndex + 1].xyz, 0.0f),
+			float4(VertexFetch_InstanceTransformBuffer[4 * InstanceIndex + 2].xyz, 0.0f),
+			float4(VertexFetch_InstanceTransformBuffer[4 * InstanceIndex + 3].xyz, 1.0f));
 	#elif !USE_INSTANCING_EMULATED
 	return float4x4(
 		float4(Input.InstanceTransform1.xyz, 0.0f),
