Dispatches exactly one lane per pixel for SSRTGI.

diff --git a/Engine/Shaders/Private/ForwardLightingCommon.ush b/Engine/Shaders/Private/ForwardLightingCommon.ush
--- a/Engine/Shaders/Private/ForwardLightingCommon.ush
+++ b/Engine/Shaders/Private/ForwardLightingCommon.ush
@@ -346,6 +346,8 @@ half3 GetImageBasedReflectionLighting(FMaterialPixelParameters MaterialParameter
 		float StepOffset = InterleavedGradientNoise( MaterialParameters.SvPosition.xy, View.StateFrameIndexMod8 );
 		StepOffset -= 0.5;
 
+		bool bDebugPrint = false;
+
 		float3 HitUVz;
 		float Level = 0;
 
@@ -354,6 +356,7 @@ half3 GetImageBasedReflectionLighting(FMaterialPixelParameters MaterialParameter
 			MaterialParameters.WorldPosition_CamRelative, RayDirection, Roughness, MaterialParameters.ScreenPosition.w,
 			12, StepOffset,
 			TranslucentBasePass.HZBUvFactorAndInvFactor,
+			bDebugPrint,
 			HitUVz,
 			Level
 		);
diff --git a/Engine/Shaders/Private/SSRT/SSRTDiffuseIndirect.usf b/Engine/Shaders/Private/SSRT/SSRTDiffuseIndirect.usf
--- a/Engine/Shaders/Private/SSRT/SSRTDiffuseIndirect.usf
+++ b/Engine/Shaders/Private/SSRT/SSRTDiffuseIndirect.usf
@@ -1,5 +1,59 @@
 // Copyright 1998-2019 Epic Games, Inc. All Rights Reserved.
 
+#define CONFIG_SORT_RAYS 0
+#define CONFIG_SORT_BIN_COUNT 16
+
+#define DEBUG_SSRT 0
+
+
+#if QUALITY == 1
+	#define CONFIG_RAY_STEPS 8
+	#define CONFIG_RAY_COUNT 4
+
+	#define TILE_PIXEL_SIZE_X 8
+	#define TILE_PIXEL_SIZE_X_LOG 3
+
+	#define TILE_PIXEL_SIZE_Y 8
+	#define TILE_PIXEL_SIZE_Y_LOG 3
+
+#elif QUALITY == 2
+	#define CONFIG_RAY_STEPS 8
+	#define CONFIG_RAY_COUNT 8
+
+	#define TILE_PIXEL_SIZE_X 8
+	#define TILE_PIXEL_SIZE_X_LOG 3
+
+	#define TILE_PIXEL_SIZE_Y 4
+	#define TILE_PIXEL_SIZE_Y_LOG 2
+
+#elif QUALITY == 3
+	#define CONFIG_RAY_STEPS 8
+	#define CONFIG_RAY_COUNT 16
+
+	#define TILE_PIXEL_SIZE_X 4
+	#define TILE_PIXEL_SIZE_X_LOG 2
+
+	#define TILE_PIXEL_SIZE_Y 4
+	#define TILE_PIXEL_SIZE_Y_LOG 2
+
+#elif QUALITY == 4
+	#define CONFIG_RAY_STEPS 12
+	#define CONFIG_RAY_COUNT 32
+
+	#define TILE_PIXEL_SIZE_X 4
+	#define TILE_PIXEL_SIZE_X_LOG 2
+
+	#define TILE_PIXEL_SIZE_Y 2
+	#define TILE_PIXEL_SIZE_Y_LOG 1
+
+#else
+	#error Unknown Quality.
+
+#endif
+
+#define TILE_PIXEL_COUNT (TILE_PIXEL_SIZE_X * TILE_PIXEL_SIZE_Y)
+#define LANE_PER_GROUPS (TILE_PIXEL_COUNT * CONFIG_RAY_COUNT)
+		
 #define IS_SSGI_SHADER 1
 
 #include "SSRTRayCast.ush"
@@ -9,157 +63,439 @@
 #include "../MonteCarlo.ush"
 #include "../SceneTextureParameters.ush"
 
+#if LANE_PER_GROUPS == 64 && 0
+	#undef GroupMemoryBarrierWithGroupSync
+	#define GroupMemoryBarrierWithGroupSync()
+#endif
 
-#define TILE_PIXEL_SIZE 8
-
-
-Texture2D		HZBTexture;
-SamplerState	HZBSampler;
-
-Texture2D		VelocityTexture;
-SamplerState	VelocitySampler;
-
-Texture2D		ColorTexture;
-SamplerState	ColorSampler;
+Texture2D HZBTexture;
+Texture2D VelocityTexture;
+Texture2D ColorTexture;
 
 float4 HZBUvFactorAndInvFactor;
 
 float4 PrevScreenPositionScaleBias;
 float  PrevSceneColorPreExposureCorrection;
 
-
 RWTexture2D<float4>	IndirectDiffuseOutput;
 RWTexture2D<float>	AmbientOcclusionOutput;
 
+groupshared uint SharedMemory[LANE_PER_GROUPS * 2];
 
-[numthreads(TILE_PIXEL_SIZE, TILE_PIXEL_SIZE, 1)]
-void MainCS(
-	uint2 DispatchThreadId : SV_DispatchThreadID,
-	uint2 GroupId : SV_GroupID,
-	uint2 GroupThreadId : SV_GroupThreadID,
-	uint GroupThreadIndex : SV_GroupIndex)
+
+uint CompressN(float3 N)
 {
-	uint2 PixelPos = DispatchThreadId + View.ViewRectMin.xy;
-	float2 BufferUV = (PixelPos + 0.5) * View.BufferSizeAndInvSize.zw;
-    float2 ScreenPos = ViewportUVToScreenPos((DispatchThreadId + 0.5) * View.ViewSizeAndInvSize.zw);
+	// matches 8bits GBuffer A to be lossless.
+	uint3 K = uint3(saturate(N * 0.5 + 0.5) * 255.0);
+	return uint(K.x << 0 | K.y << 8 | K.z << 16);
+}
 
-	FGBufferData GBuffer = GetGBufferDataFromSceneTextures(BufferUV);
+float3 DecompressN(uint EncodedN)
+{
+	uint3 K;
+	K.x = (EncodedN >>  0) & 0xFF;
+	K.y = (EncodedN >>  8) & 0xFF;
+	K.z = (EncodedN >> 16) & 0xFF;
 
-	// No SSS because diffuse luminance is in alpha channel and we don't have dual source blending.
-	BRANCH if( GBuffer.ShadingModelID == 0 || GBuffer.ShadingModelID == SHADINGMODELID_SUBSURFACE_PROFILE )
-	{
-		IndirectDiffuseOutput[PixelPos] = 0.0;
-		AmbientOcclusionOutput[PixelPos] = 0.0;
-	}
+	return float3(K) * (2.0 / 255.0) - 1.0;
+}
 
-	const float3 N = GBuffer.WorldNormal;
-	const float SceneDepth = GBuffer.Depth;
-	const float3 PositionTranslatedWorld = mul( float4( ScreenPos * SceneDepth, SceneDepth, 1 ), View.ScreenToTranslatedWorld ).xyz;
+uint2 DecodeGroupPixelOffset(uint GroupPixelId)
+{
+    return uint2(GroupPixelId % TILE_PIXEL_SIZE_X, (GroupPixelId >> TILE_PIXEL_SIZE_X_LOG) % TILE_PIXEL_SIZE_Y);
+}
 
-#if QUALITY == 1
-	uint NumSteps = 8;
-	uint NumRays = 4;
-#elif QUALITY == 2
-	uint NumSteps = 8;
-	uint NumRays = 8;
-#elif QUALITY == 3
-	uint NumSteps = 8;
-	uint NumRays = 12;
-#else // QUALITY == 4
-	uint NumSteps = 12;
-	uint NumRays = 16;
-#endif
+uint EncodeGroupPixelOffset(uint2 GroupPixelOffset)
+{
+	return GroupPixelOffset.x | (GroupPixelOffset.y << TILE_PIXEL_SIZE_X_LOG);
+}
+
+uint2 ComputePixelPosition(uint2 GroupId, uint2 GroupPixelOffset)
+{
+    return GroupId * uint2(TILE_PIXEL_SIZE_X, TILE_PIXEL_SIZE_Y) + GroupPixelOffset;
+}
+
+void UpdateLane2DCoordinateInformations(
+	uint2 PixelPosition,
+	out float2 BufferUV,
+	out float2 ScreenPos)
+{
+	// TODO: split screen
+	BufferUV = (PixelPosition + 0.5) * View.BufferSizeAndInvSize.zw;
+    ScreenPos = ViewportUVToScreenPos((PixelPosition + 0.5) * View.ViewSizeAndInvSize.zw);
+}
+
+float3 ComputeTranslatedWorldPositions(float2 ScreenPos, float SceneDepth)
+{
+	return mul(float4(ScreenPos * SceneDepth, SceneDepth, 1), View.ScreenToTranslatedWorld).xyz;
+}
+
+uint2 ComputePixelPosition(float3 TranslatedWorldPosition)
+{
+	float4 ClipPosition = mul(float4(TranslatedWorldPosition, 1), View.TranslatedWorldToClip);
+	float2 ScreenPos = ClipPosition.xy * rcp(ClipPosition.w);
+	float2 ViewportUV = ScreenPosToViewportUV(ScreenPos);
+	return uint2(ViewportUV * View.ViewSizeAndInvSize.xy);
+}
+
+float ComputeSceneDepth(float3 TranslatedWorldPosition)
+{
+	// TODO: do everything in view space instead of world space?
+	return mul(float4(TranslatedWorldPosition, 1.0), View.TranslatedWorldToView).z;
+}
+
+uint2 ComputeRandomSeed(uint2 PixelPosition)
+{
+	return Rand3DPCG16(int3(PixelPosition, View.StateFrameIndexMod8)).xy;
+}
 
-	float2 Noise;
-	Noise.x = InterleavedGradientNoise(PixelPos + 0.5, View.StateFrameIndexMod8);
-	Noise.y = InterleavedGradientNoise(PixelPos + 0.5, View.StateFrameIndexMod8 * 117);
+float3 ComputeL(float3 N, float2 E)
+{
+	float3x3 TangentBasis = GetTangentBasis(N);
+	
+	float3 L;
+	L.xy = UniformSampleDiskConcentric( E );
+	L.z = sqrt( 1 - dot( L.xy, L.xy ) );
+	L = mul( L, TangentBasis );
+
+	return L;
+
+	float3 TangentL = CosineSampleHemisphere(E).xyz;
+	return mul(TangentL, TangentBasis);
+}
+
+#if CONFIG_SORT_RAYS
+
+uint SelectDirectionBin(
+	uint2 GroupId,
+	uint GroupPixelId,
+	uint RaySequenceId,
+	float3 N,
+	float SceneDepth)
+{
+	uint2 GroupPixelOffset = DecodeGroupPixelOffset(GroupPixelId);
+	uint2 PixelPosition = ComputePixelPosition(GroupId, GroupPixelOffset);
+
+	uint2 RandomSeed = ComputeRandomSeed(PixelPosition);
+
+	float2 E = Hammersley16(RaySequenceId, CONFIG_RAY_COUNT, RandomSeed);
 		
-	//uint2 Random = 0x10000 * Noise;
-	uint2 Random = Rand3DPCG16( int3( PixelPos, View.StateFrameIndexMod8 ) ).xy;
+	float3 L = ComputeL(N, E);
 			
-	// TODO: lot of VGPR, worth recomputing every ray instead?
-	float3x3 TangentBasis = GetTangentBasis( N );
+	float2 BufferUV;
+	float2 ScreenPos;
+	UpdateLane2DCoordinateInformations(PixelPosition, /* out */ BufferUV, /* out */ ScreenPos);
 	
-	float3 DiffuseColor = 0;
-	float AmbientOcclusion = 0;
+	float3 TranslatedWorldPosition = ComputeTranslatedWorldPositions(ScreenPos, SceneDepth);
 
-	// Shoot multiple rays
-	LOOP
-	for (uint i = 0; i < NumRays; i++)
-	{
-		// TOOD: VGPR, might be worth recomputing instead?
-		// SvPosition: is SvPosition really needed? Can recompute from PositionTranslatedWorld?
-		float StepOffset = Noise.x;
-		StepOffset -= 0.9;
+	// TODO provide RayStartUVz
+	const float4 RayStartClip = mul(float4(TranslatedWorldPosition, 1), View.TranslatedWorldToClip);
+	const float4 RayEndClip   = mul(float4(TranslatedWorldPosition + L * SceneDepth, 1), View.TranslatedWorldToClip);
+
+    const float3 RayStartScreen = RayStartClip.xyz * rcp(RayStartClip.w);
+    const float3 RayEndScreen = RayEndClip.xyz * rcp(RayEndClip.w);
+
+	float2 RayPixelDirection = (RayEndScreen - RayStartScreen).xy * View.ViewSizeAndInvSize.xy;
+
+	float DirectionAngle = atan2(RayPixelDirection.x, RayPixelDirection.y);
 			
-		float2 E = Hammersley16( i, NumRays, Random );
+	uint SelectedBinId = floor(DirectionAngle * (CONFIG_SORT_BIN_COUNT * rcp(2 * PI)) + (CONFIG_SORT_BIN_COUNT / 2 + 0.5)); 
+	SelectedBinId = clamp(SelectedBinId, 0, CONFIG_SORT_BIN_COUNT - 1);
 
-#if 0
-		float3 L = CosineSampleHemisphere( E ).xyz;
-		L = mul( L, TangentBasis );
-#else
-		float3 L;
-		L.xy = UniformSampleDiskConcentric( E );
-		L.z = sqrt( 1 - dot( L.xy, L.xy ) );
-		L = mul( L, TangentBasis );
-#endif
+	return SelectedBinId;
+} // SelectDirectionBin()
+
+#endif // CONFIG_SORT_RAYS
+
+[numthreads(TILE_PIXEL_SIZE_X, TILE_PIXEL_SIZE_Y, CONFIG_RAY_COUNT)]
+void MainCS(
+	uint2 GroupId : SV_GroupID,
+	uint GroupThreadIndex : SV_GroupIndex)
+{
+	// Id of the wave in the group.
+	uint GroupWaveIndex = GroupThreadIndex / 64;
+	
+	{
+		const uint BinsAddress = TILE_PIXEL_COUNT * 2;
+
+		uint GroupPixelId = GroupThreadIndex % TILE_PIXEL_COUNT;
+		uint RaySequenceId = GroupThreadIndex / TILE_PIXEL_COUNT;
+		
+		// Store GBuffer into LDS
+		{
+			BRANCH
+			if (RaySequenceId == 0)
+			{
+				uint2 GroupPixelOffset = DecodeGroupPixelOffset(GroupPixelId);
+                uint2 PixelPosition = ComputePixelPosition(GroupId, GroupPixelOffset);
 				
-		float3 HitUVz;
-		float Level = 0;
-
-		bool bHit = RayCast(
-			HZBTexture, HZBSampler,
-			PositionTranslatedWorld, L, 1, SceneDepth, 
-			NumSteps, StepOffset,
-			HZBUvFactorAndInvFactor,
-			HitUVz,
-			Level
-		);
-
-#if 0	// Backface check
-		if( bHit )
+				float2 BufferUV = (PixelPosition + 0.5) * View.BufferSizeAndInvSize.zw;
+		
+				FGBufferData GBuffer = GetGBufferDataFromSceneTextures(BufferUV);
+
+				bool bEarlyReturn = GBuffer.ShadingModelID == SHADINGMODELID_UNLIT || GBuffer.ShadingModelID == SHADINGMODELID_SUBSURFACE_PROFILE;
+				
+				SharedMemory[TILE_PIXEL_COUNT * 0 | GroupPixelId] = CompressN(GBuffer.WorldNormal);
+				SharedMemory[TILE_PIXEL_COUNT * 1 | GroupPixelId] = asuint(bEarlyReturn ? -1.0 : GBuffer.Depth);
+			}
+			else if (GroupWaveIndex == 1) // TODO.
+			{
+				// Clears the bins
+				SharedMemory[BinsAddress | GroupPixelId] = 0;
+			}
+		}
+		
+		#if CONFIG_SORT_RAYS
+	
+		// Sorts rays.
 		{
-			float3 SampleNormal = GetGBufferDataFromSceneTextures(HitUVz.xy).WorldNormal;
-			bHit = dot( SampleNormal, L ) < 0;
+			// Group sync to acces gbuffer data.
+			GroupMemoryBarrierWithGroupSync();
+
+			// Fetch information from LDS.
+			float3 N;
+			float SceneDepth;
+			{
+				uint Raw0 = SharedMemory[TILE_PIXEL_COUNT * 0 | GroupPixelId];
+				uint Raw1 = SharedMemory[TILE_PIXEL_COUNT * 1 | GroupPixelId];
+
+				N = DecompressN(Raw0);
+				SceneDepth = asfloat(Raw1);
+			}
+
+			// Compute the ray's bin.
+			uint SelectedBinId = SelectDirectionBin(
+				GroupId,
+				GroupPixelId,
+				RaySequenceId,
+				N,
+				SceneDepth);
+		
+			// Count ray in bins
+			uint OffsetInBin = 0;
+			InterlockedAdd(SharedMemory[BinsAddress | SelectedBinId], 1, OffsetInBin);
+			
+			GroupMemoryBarrierWithGroupSync();
+
+			// Compute the starting position of the bin in LDS.
+			// TODO: faster way to do it?
+			uint SelectedBinStartingPos = 0;
+			UNROLL
+			for (uint BinId = 0; BinId < CONFIG_SORT_BIN_COUNT; BinId++)
+			{
+				FLATTEN
+				if (BinId < SelectedBinId)
+				{
+					SelectedBinStartingPos += SharedMemory[BinsAddress | BinId];
+				}
+			}
+			
+			GroupMemoryBarrierWithGroupSync();
+
+			// Store ray informations in LDS.
+			{
+				uint DestPos = SelectedBinStartingPos + OffsetInBin;
+
+				#if 0 // Uncomment me to measure performance improvement.
+					if (View.GeneralPurposeTweak > 0.5)
+						DestPos = GroupPixelId + RaySequenceId * TILE_PIXEL_COUNT;
+				#endif
+		
+				#if LANE_PER_GROUPS > 256
+					#error "LANE_PER_GROUPS needs to be encodable in 8bits."
+				#endif
+
+				// TODO: Encoded N and 24bits and pack EncodedGroupPixelOffset.
+                SharedMemory[LANE_PER_GROUPS * 0 | DestPos] = asuint(CompressN(N) | (GroupPixelId << 24) | (RaySequenceId << (24 + TILE_PIXEL_SIZE_X_LOG + TILE_PIXEL_SIZE_Y_LOG)));
+				SharedMemory[LANE_PER_GROUPS * 1 | DestPos] = asuint(SceneDepth);
+			}
 		}
-#endif
+		#endif // CONFIG_SORT_RAYS
+	}
+	
+	GroupMemoryBarrierWithGroupSync();
+	
+	// Shoot ray
+	{
+		uint GroupPixelId;
+		uint RaySequenceId;
+		uint CompressedN;
+		float SceneDepth;
+		#if CONFIG_SORT_RAYS
+		{
+			uint Row0 = SharedMemory[LANE_PER_GROUPS * 0 | GroupThreadIndex];
+			uint Row1 = SharedMemory[LANE_PER_GROUPS * 1 | GroupThreadIndex];
+
+			// Unpack
+			RaySequenceId = Row0 >> (24 + TILE_PIXEL_SIZE_X_LOG + TILE_PIXEL_SIZE_Y_LOG);
+            GroupPixelId = (Row0 >> 24) % TILE_PIXEL_COUNT;
+			CompressedN = Row0;
+			SceneDepth = asfloat(Row1);
+		}
+		#else // !CONFIG_SORT_RAYS
+		{
+			GroupPixelId = GroupThreadIndex % TILE_PIXEL_COUNT;
+			RaySequenceId = GroupThreadIndex / TILE_PIXEL_COUNT;
+
+			uint Raw0 = SharedMemory[TILE_PIXEL_COUNT * 0 | GroupPixelId];
+			uint Raw1 = SharedMemory[TILE_PIXEL_COUNT * 1 | GroupPixelId];
+
+			CompressedN = Raw0;
+			SceneDepth = asfloat(Raw1);
+		}
+		#endif // !CONFIG_SORT_RAYS
+		
+		GroupMemoryBarrierWithGroupSync();
+		
+		const bool bTraceRay = SceneDepth > 0;
+		
+		uint2 CompressedColor;
 
-		// if there was a hit
-		BRANCH if( bHit )
+		BRANCH
+		if (bTraceRay)
 		{
-			float2 SampleUV;
-			float Vignette;
-			ReprojectHit(PrevScreenPositionScaleBias, VelocityTexture, VelocitySampler, HitUVz, SampleUV, Vignette);
+			uint2 GroupPixelOffset = DecodeGroupPixelOffset(GroupPixelId);
+			uint2 PixelPosition = ComputePixelPosition(GroupId, GroupPixelOffset);
+	
+			float2 BufferUV;
+			float2 ScreenPos;
+			UpdateLane2DCoordinateInformations(PixelPosition, /* out */ BufferUV, /* out */ ScreenPos);
+	
+			float3 TranslatedWorldPosition = ComputeTranslatedWorldPositions(ScreenPos, SceneDepth);
 
-#if 0
-			float4 SampleColor = SampleHCBLevel( ColorTexture, ColorSampler, SampleUV, Level, HZBUvFactorAndInvFactor );
-#else
-			float4 SampleColor = SampleScreenColor( ColorTexture, ColorSampler, SampleUV );
-#endif
+			uint2 RandomSeed = ComputeRandomSeed(PixelPosition);
+			float2 E = Hammersley16(RaySequenceId, CONFIG_RAY_COUNT, RandomSeed);		
+			float3 L = ComputeL(DecompressN(CompressedN), E);
+			
+			float StepOffset = InterleavedGradientNoise(PixelPosition + 0.5, View.StateFrameIndexMod8) - 0.9;
+	
+            bool bDebugPrint = all(PixelPosition == uint2(View.ViewSizeAndInvSize.xy) / 2);
+
+			float Level = 0; // TODO: remove
+			float3 HitUVz;
+			bool bHit = RayCast(
+				HZBTexture, GlobalPointClampedSampler,
+				TranslatedWorldPosition, L, 1, SceneDepth, 
+				CONFIG_RAY_COUNT, StepOffset,
+				HZBUvFactorAndInvFactor,
+				bDebugPrint,
+				/* out */ HitUVz,
+				/* out */ Level);
+
+			#if 0 // Backface check
+			if (bHit)
+			{
+				float3 SampleNormal = GetGBufferDataFromSceneTextures(HitUVz.xy).WorldNormal;
+				bHit = dot(SampleNormal, L) < 0;
+			}
+			#endif
+			
+			// if there was a hit
+			//BRANCH // TODO
+			if (bHit)
+			{
+				float2 SampleUV;
+				float Vignette;
+				ReprojectHit(PrevScreenPositionScaleBias, VelocityTexture, GlobalPointClampedSampler, HitUVz, SampleUV, Vignette);
+
+				float4 SampleColor = SampleScreenColor(ColorTexture, GlobalBilinearClampedSampler, SampleUV);
+
+				float SampleColorWeight = 1.0;
 
-			// Backface modulation of hit surface
-			float3 SampleNormal = GetGBufferDataFromSceneTextures(HitUVz.xy).WorldNormal;
-			SampleColor.rgb *= saturate( 1 - dot( SampleNormal, L ) );
+				// Backface modulation of hit surface
+				{
+					float3 SampleNormal = GetGBufferDataFromSceneTextures(HitUVz.xy).WorldNormal;
+					
+					uint2 GroupPixelOffset = DecodeGroupPixelOffset(GroupPixelId);
+					uint2 PixelPosition = ComputePixelPosition(GroupId, GroupPixelOffset);
 
-			SampleColor.rgb *= rcp( 1 + Luminance(SampleColor.rgb) );
+					uint2 RandomSeed = ComputeRandomSeed(PixelPosition);
 
-			DiffuseColor += SampleColor.rgb;
-			AmbientOcclusion = SampleColor.a;
+					float2 E = Hammersley16(RaySequenceId, CONFIG_RAY_COUNT, RandomSeed);
+		
+					float3 L = ComputeL(DecompressN(CompressedN), E);
+		
+					SampleColorWeight *= saturate( 1 - dot( SampleNormal, L ) );
+				}
+
+				#if CONFIG_RAY_COUNT > 1
+					SampleColorWeight *= rcp( 1 + Luminance(SampleColor.rgb) );
+				#endif
+			
+				float3 DiffuseColor = SampleColor.rgb * SampleColorWeight;
+				float AmbientOcclusion = SampleColor.a;
+				
+				CompressedColor.x = asuint(f32tof16(DiffuseColor.r) << 16 | f32tof16(DiffuseColor.g));
+				CompressedColor.y = asuint(f32tof16(DiffuseColor.b) << 16 | f32tof16(AmbientOcclusion));
+			}
+			else
+			{
+				CompressedColor = uint2(0, 0);
+			}
+			
 		}
-	}
+		else if (!bTraceRay)
+        {
+            CompressedColor = uint2(0, 0);
+        }
+		
+		uint DestPos = GroupPixelId + RaySequenceId * TILE_PIXEL_COUNT;
+			
+		SharedMemory[LANE_PER_GROUPS * 0 | DestPos] = CompressedColor.x;
+		SharedMemory[LANE_PER_GROUPS * 1 | DestPos] = CompressedColor.y;
+    }
+	
+	GroupMemoryBarrierWithGroupSync();
+	
+	// Store ray to UAV
+	BRANCH
+	if (GroupThreadIndex < TILE_PIXEL_COUNT)
+	{
+		const uint GroupPixelId = GroupThreadIndex;
+	
+		float3 DiffuseColor = 0;
+		float AmbientOcclusion = 0;
 
-	DiffuseColor *= rcp(float(NumRays));
-	AmbientOcclusion *= rcp(float(NumRays));
+		UNROLL
+		for (uint RaySequenceId = 0; RaySequenceId < CONFIG_RAY_COUNT; RaySequenceId++)
+		{
+			uint SrcPos = GroupPixelId + RaySequenceId * TILE_PIXEL_COUNT;
 
-	DiffuseColor *= rcp( 1 - Luminance(DiffuseColor) );
-	DiffuseColor *= View.IndirectLightingColorScale;
+			uint Row0 = SharedMemory[LANE_PER_GROUPS * 0 | SrcPos];
+			uint Row1 = SharedMemory[LANE_PER_GROUPS * 1 | SrcPos];
 
-	AmbientOcclusion = 1 - AmbientOcclusion;
+			DiffuseColor.r += f16tof32(Row0 >> 16);
+			DiffuseColor.g += f16tof32(Row0 >>  0);
+			DiffuseColor.b += f16tof32(Row1 >> 16);
+			AmbientOcclusion += f16tof32(Row1 >> 0);
+		}
 
-	#if USE_PREEXPOSURE
-		DiffuseColor *= PrevSceneColorPreExposureCorrection;
-	#endif
-	
-	IndirectDiffuseOutput[PixelPos] = float4(DiffuseColor, 0.0);
-	AmbientOcclusionOutput[PixelPos] = AmbientOcclusion;
+		#if CONFIG_RAY_COUNT > 1
+		{
+			DiffuseColor *= rcp(float(CONFIG_RAY_COUNT));
+			AmbientOcclusion *= rcp(float(CONFIG_RAY_COUNT));
+
+			DiffuseColor *= rcp( 1 - Luminance(DiffuseColor) );
+		}	
+		#endif
+
+		DiffuseColor *= View.IndirectLightingColorScale;
+
+		AmbientOcclusion = 1 - AmbientOcclusion;
+
+		#if USE_PREEXPOSURE
+			DiffuseColor *= PrevSceneColorPreExposureCorrection;
+		#endif
+
+		//AmbientOcclusion = 1;
+
+		// Output.
+		{
+			uint2 GroupPixelOffset = DecodeGroupPixelOffset(GroupPixelId);
+			uint2 OutputPixelCoordinate = ComputePixelPosition(GroupId, GroupPixelOffset);
+
+			IndirectDiffuseOutput[OutputPixelCoordinate] = float4(DiffuseColor, 0.0);
+			AmbientOcclusionOutput[OutputPixelCoordinate] = AmbientOcclusion;
+		}
+	} // if (GroupThreadIndex < TILE_PIXEL_COUNT)
 } // MainCS()
diff --git a/Engine/Shaders/Private/SSRT/SSRTRayCast.ush b/Engine/Shaders/Private/SSRT/SSRTRayCast.ush
--- a/Engine/Shaders/Private/SSRT/SSRTRayCast.ush
+++ b/Engine/Shaders/Private/SSRT/SSRTRayCast.ush
@@ -8,12 +8,40 @@
 
 #include "../Common.ush"	
 
+
+#ifndef DEBUG_SSRT
+	#define DEBUG_SSRT 0
+#endif
+
+
+#if DEBUG_SSRT
+
+RWTexture2D<float4> ScreenSpaceRayTracingDebugOutput;
+
+
+void PrintSample(
+	float4 HZBUvFactorAndInvFactor, 
+	float2 HZBUV,
+	float4 DebugColor)
+{	
+	HZBUV.xy *= HZBUvFactorAndInvFactor.zw;
+
+	uint2 Position = HZBUV.xy * View.ViewSizeAndInvSize.xy;
+
+	ScreenSpaceRayTracingDebugOutput[Position] = DebugColor;
+}
+
+#endif
+
+
+
 bool RayCast(
 	Texture2D Texture, SamplerState Sampler,
 	float3 RayOriginTranslatedWorld, float3 RayDirection,
 	float Roughness, float SceneDepth,
 	uint NumSteps, float StepOffset,
 	float4 HZBUvFactorAndInvFactor, 
+	bool bDebugPrint,
 	out float3 OutHitUVz,
 	out float Level
 )
@@ -59,6 +87,13 @@ bool RayCast(
 
 	RayStepUVz *= Step;
 	float3 RayUVz = RayStartUVz + RayStepUVz * StepOffset;
+	
+	#if DEBUG_SSRT
+	{
+		if (bDebugPrint)
+			PrintSample(HZBUvFactorAndInvFactor, RayStartUVz.xy, float4(1, 0, 0, 1));
+	}
+	#endif
 
 	LOOP
 	for( uint i = 0; i < NumSteps; i += 4 )
@@ -68,6 +103,18 @@ bool RayCast(
 		float4 SampleUV1 = RayUVz.xyxy + RayStepUVz.xyxy * float4( 3, 3, 4, 4 );
 		float4 SampleZ   = RayUVz.zzzz + RayStepUVz.zzzz * float4( 1, 2, 3, 4 );
 		
+		#if DEBUG_SSRT
+		{
+			if (bDebugPrint)
+			{
+				PrintSample(HZBUvFactorAndInvFactor, SampleUV0.xy, float4(0, 1, 0, 1));
+				PrintSample(HZBUvFactorAndInvFactor, SampleUV0.zw, float4(0, 1, 0, 1));
+				PrintSample(HZBUvFactorAndInvFactor, SampleUV1.xy, float4(0, 1, 0, 1));
+				PrintSample(HZBUvFactorAndInvFactor, SampleUV1.zw, float4(0, 1, 0, 1));
+			}
+		}
+		#endif
+
 		// Use lower res for farther samples
 		float4 SampleDepth;
 		SampleDepth.x = Texture.SampleLevel( Sampler, SampleUV0.xy, Level ).r;
@@ -183,6 +230,13 @@ bool RayCast(
 		RayUVz += 4 * RayStepUVz;
 	}
 	
+	#if DEBUG_SSRT
+	{
+		if (bDebugPrint)
+			PrintSample(HZBUvFactorAndInvFactor, OutHitUVz.xy, float4(0, 0, 1, 1));
+	}
+	#endif
+
 	OutHitUVz.xy *= HZBUvFactorAndInvFactor.zw;
 	OutHitUVz.xy = OutHitUVz.xy * float2( 2, -2 ) + float2( -1, 1 );
 	OutHitUVz.xy = OutHitUVz.xy * View.ScreenPositionScaleBias.xy + View.ScreenPositionScaleBias.wz;
diff --git a/Engine/Shaders/Private/SSRT/SSRTReflections.usf b/Engine/Shaders/Private/SSRT/SSRTReflections.usf
--- a/Engine/Shaders/Private/SSRT/SSRTReflections.usf
+++ b/Engine/Shaders/Private/SSRT/SSRTReflections.usf
@@ -106,6 +106,8 @@ void ScreenSpaceReflectionsPS(
 	float2 UV = SvPosition.xy * View.BufferSizeAndInvSize.zw;
 	float2 ScreenPos = ViewportUVToScreenPos((SvPosition.xy - View.ViewRectMin.xy) * View.ViewSizeAndInvSize.zw);
 	uint2 PixelPos = (uint2)SvPosition.xy;
+	
+    bool bDebugPrint = all(PixelPos == uint2(View.ViewSizeAndInvSize.xy) / 2);
 
 	OutColor = 0;
 	
@@ -226,12 +228,13 @@ void ScreenSpaceReflectionsPS(
 			{
 				L = reflect(-V, N);
 			}
-
+			
 			bool bHit = RayCast(
 				HZB, HZBSampler,
 				PositionTranslatedWorld, L, Roughness, SceneDepth, 
 				NumSteps, StepOffset,
 				HZBUvFactorAndInvFactor,
+				bDebugPrint,
 				HitUVz,
 				Level
 			);
@@ -324,6 +327,7 @@ void ScreenSpaceReflectionsPS(
 			PositionTranslatedWorld, L, Roughness, SceneDepth,
 			NumSteps, StepOffset,
 			HZBUvFactorAndInvFactor,
+			bDebugPrint,
 			HitUVz,
 			Level
 		);
diff --git a/Engine/Source/Runtime/Renderer/Private/IndirectLightRendering.cpp b/Engine/Source/Runtime/Renderer/Private/IndirectLightRendering.cpp
--- a/Engine/Source/Runtime/Renderer/Private/IndirectLightRendering.cpp
+++ b/Engine/Source/Runtime/Renderer/Private/IndirectLightRendering.cpp
@@ -192,7 +192,7 @@ void FDeferredShadingSceneRenderer::RenderDiffuseIndirectAndAmbientOcclusion(FRH
 		// Extract the dynamic AO for application of AO beyond RenderDiffuseIndirectAndAmbientOcclusion()
 		if (DenoiserOutputs.AmbientOcclusionMask)
 		{
-			ensureMsgf(!bApplySSAO, TEXT("Looks like SSAO has been computed for this view but is being overridden."));
+			//ensureMsgf(!bApplySSAO, TEXT("Looks like SSAO has been computed for this view but is being overridden."));
 			ensureMsgf(Views.Num() == 1, TEXT("Need to add support for one AO texture per view in FSceneRenderTargets")); // TODO.
 			GraphBuilder.QueueTextureExtraction(DenoiserOutputs.AmbientOcclusionMask, &SceneContext.ScreenSpaceAO);
 			SceneContext.bScreenSpaceAOIsValid = true;
diff --git a/Engine/Source/Runtime/Renderer/Private/ScreenSpaceRayTracing.cpp b/Engine/Source/Runtime/Renderer/Private/ScreenSpaceRayTracing.cpp
--- a/Engine/Source/Runtime/Renderer/Private/ScreenSpaceRayTracing.cpp
+++ b/Engine/Source/Runtime/Renderer/Private/ScreenSpaceRayTracing.cpp
@@ -224,7 +224,7 @@ class FScreenSpaceDiffuseIndirectCS : public FGlobalShader
 	DECLARE_GLOBAL_SHADER(FScreenSpaceDiffuseIndirectCS);
 	SHADER_USE_PARAMETER_STRUCT(FScreenSpaceDiffuseIndirectCS, FGlobalShader)
 
-	class FQualityDim : SHADER_PERMUTATION_INT( "QUALITY", 5 );
+	class FQualityDim : SHADER_PERMUTATION_RANGE_INT("QUALITY", 1, 4);
 	using FPermutationDomain = TShaderPermutationDomain< FQualityDim >;
 	
 	BEGIN_SHADER_PARAMETER_STRUCT( FParameters, )
@@ -247,6 +247,8 @@ class FScreenSpaceDiffuseIndirectCS : public FGlobalShader
 
 		SHADER_PARAMETER_RDG_TEXTURE_UAV(RWTexture<float4>, IndirectDiffuseOutput)
 		SHADER_PARAMETER_RDG_TEXTURE_UAV(RWTexture<float>,  AmbientOcclusionOutput)
+
+		SHADER_PARAMETER_RDG_TEXTURE_UAV(RWTexture<float4>, ScreenSpaceRayTracingDebugOutput)
 	END_SHADER_PARAMETER_STRUCT()
 
 	static bool ShouldCompilePermutation(const FGlobalShaderPermutationParameters& Parameters)
@@ -289,6 +291,40 @@ void GetSSRShaderOptionsForQuality(ESSRQuality Quality, IScreenSpaceDenoiser::FR
 	}
 }
 
+void GetSSRTGIShaderOptionsForQuality(int32 Quality, FIntPoint* OutGroupSize, int32* OutRayCountPerPixel)
+{
+	if (Quality == 1)
+	{
+		OutGroupSize->X = 8;
+		OutGroupSize->Y = 8;
+		*OutRayCountPerPixel = 4;
+	}
+	else if (Quality == 2)
+	{
+		OutGroupSize->X = 8;
+		OutGroupSize->Y = 4;
+		*OutRayCountPerPixel = 8;
+	}
+	else if (Quality == 3)
+	{
+		OutGroupSize->X = 4;
+		OutGroupSize->Y = 4;
+		*OutRayCountPerPixel = 16;
+	}
+	else if (Quality == 4)
+	{
+		OutGroupSize->X = 4;
+		OutGroupSize->Y = 2;
+		*OutRayCountPerPixel = 32;
+	}
+	else
+	{
+		check(0);
+	}
+
+	check(OutGroupSize->X * OutGroupSize->Y * (*OutRayCountPerPixel) == 256);
+}
+
 } // namespace
 
 void GetSSRQualityForView(const FViewInfo& View, ESSRQuality* OutQuality, IScreenSpaceDenoiser::FReflectionsRayTracingConfig* OutRayTracingConfigs)
@@ -525,6 +561,10 @@ void RenderScreenSpaceDiffuseIndirect(
 	
 	const int32 Quality = FMath::Clamp( CVarSSGIQuality.GetValueOnRenderThread(), 1, 4 );
 
+	FIntPoint GroupSize;
+	int32 RayCountPerPixel;
+	GetSSRTGIShaderOptionsForQuality(Quality, &GroupSize, &RayCountPerPixel);
+
 	FSceneRenderTargets& SceneContext = FSceneRenderTargets::Get(GraphBuilder.RHICmdList);
 
 	// Allocate outputs.
@@ -532,7 +572,7 @@ void RenderScreenSpaceDiffuseIndirect(
 		FRDGTextureDesc Desc = FRDGTextureDesc::Create2DDesc(
 			SceneTextures.SceneDepthBuffer->Desc.Extent,
 			PF_FloatRGBA,
-			FClearValueBinding::None,
+			FClearValueBinding::Transparent,
 			/* InFlags = */ TexCreate_None,
 			/* InTargetableFlags = */ TexCreate_ShaderResource | TexCreate_RenderTargetable | TexCreate_UAV,
 			/* bInForceSeparateTargetAndShaderResource = */ false);
@@ -587,14 +627,34 @@ void RenderScreenSpaceDiffuseIndirect(
 	PassParameters->IndirectDiffuseOutput = GraphBuilder.CreateUAV(OutDenoiserInputs->Color);
 	PassParameters->AmbientOcclusionOutput = GraphBuilder.CreateUAV(OutDenoiserInputs->AmbientOcclusionMask);
 
+	if (0) // DEBUG output clear
+	{
+		FRDGTexture* DebugTexture = GraphBuilder.CreateTexture(OutDenoiserInputs->Color->Desc, TEXT("SSRTDebug"));
+
+		FRenderTargetParameters* ClearPassParameters = GraphBuilder.AllocParameters<FRenderTargetParameters>();
+		ClearPassParameters->RenderTargets[0] = FRenderTargetBinding(DebugTexture, ERenderTargetLoadAction::EClear, ERenderTargetStoreAction::EStore);
+
+		GraphBuilder.AddPass(
+			RDG_EVENT_NAME("Clear SSRTDebug"),
+			ClearPassParameters,
+			ERenderGraphPassFlags::None,
+			[](FRHICommandList& RHICmdList)
+		{
+			// NOP
+		});
+
+		PassParameters->ScreenSpaceRayTracingDebugOutput = GraphBuilder.CreateUAV(DebugTexture);
+	}
+
 	FScreenSpaceDiffuseIndirectCS::FPermutationDomain PermutationVector;
 	PermutationVector.Set<FScreenSpaceDiffuseIndirectCS::FQualityDim>(Quality);
 
 	TShaderMapRef<FScreenSpaceDiffuseIndirectCS> ComputeShader(View.ShaderMap, PermutationVector);
 	FComputeShaderUtils::AddPass(
 		GraphBuilder,
-		RDG_EVENT_NAME("ScreenSpaceDiffuseIndirect(Quality=%d) %dx%d", Quality, View.ViewRect.Width(), View.ViewRect.Height()),
+		RDG_EVENT_NAME("ScreenSpaceDiffuseIndirect(Quality=%d RayPerPixel=%d) %dx%d",
+			Quality, RayCountPerPixel, View.ViewRect.Width(), View.ViewRect.Height()),
 		*ComputeShader,
 		PassParameters,
-		FComputeShaderUtils::GetGroupCount(View.ViewRect.Size(), 8));
+		FComputeShaderUtils::GetGroupCount(View.ViewRect.Size(), GroupSize));
 } // RenderScreenSpaceDiffuseIndirect()
