Optimises occupancy of RayCast()

diff --git a/Engine/Shaders/Private/SSRT/SSRTDiffuseIndirect.usf b/Engine/Shaders/Private/SSRT/SSRTDiffuseIndirect.usf
--- a/Engine/Shaders/Private/SSRT/SSRTDiffuseIndirect.usf
+++ b/Engine/Shaders/Private/SSRT/SSRTDiffuseIndirect.usf
@@ -376,7 +376,7 @@ void MainCS(
 			bool bHit = RayCast(
 				HZBTexture, GlobalPointClampedSampler,
 				TranslatedWorldPosition, L, 1, SceneDepth, 
-				CONFIG_RAY_COUNT, StepOffset,
+				CONFIG_RAY_STEPS, StepOffset,
 				HZBUvFactorAndInvFactor,
 				bDebugPrint,
 				/* out */ HitUVz,
diff --git a/Engine/Shaders/Private/SSRT/SSRTRayCast.ush b/Engine/Shaders/Private/SSRT/SSRTRayCast.ush
--- a/Engine/Shaders/Private/SSRT/SSRTRayCast.ush
+++ b/Engine/Shaders/Private/SSRT/SSRTRayCast.ush
@@ -77,11 +77,7 @@ bool RayCast(
 	#else
 		float CompareTolerance = max( abs( RayStepUVz.z ), ( RayStartScreen.z - RayDepthScreen.z ) * 4 ) * Step;
 	#endif
-
-	// avoid bugs with early returns inside of loops on certain platform compilers.
-	bool bHit = false;
-	OutHitUVz = float3( 0, 0, 0 );
-
+	
 	float LastDiff = 0;
 	Level = 1;
 
@@ -94,7 +90,11 @@ bool RayCast(
 			PrintSample(HZBUvFactorAndInvFactor, RayStartUVz.xy, float4(1, 0, 0, 1));
 	}
 	#endif
-
+	
+	float4 MultipleSampleDepthDiff;
+	bool4 bMultipleSampleHit;
+	bool bFoundAnyHit = false;
+	
 	LOOP
 	for( uint i = 0; i < NumSteps; i += 4 )
 	{
@@ -125,71 +125,89 @@ bool RayCast(
 		SampleDepth.w = Texture.SampleLevel( Sampler, SampleUV1.zw, Level ).r;
 		Level += (8.0 / NumSteps) * Roughness;
 
-		float4 DepthDiff = SampleZ - SampleDepth;
-		bool4 Hit = abs( DepthDiff + CompareTolerance ) < CompareTolerance;
+		MultipleSampleDepthDiff = SampleZ - SampleDepth;
+		bMultipleSampleHit = abs(MultipleSampleDepthDiff + CompareTolerance) < CompareTolerance;
+		
+		bFoundAnyHit = any(bMultipleSampleHit);
 
-		BRANCH if( any( Hit ) )
+		BRANCH
+		if (bFoundAnyHit)
 		{
-			if(IS_SSGI_SHADER)
-			{
-				// If hit set to intersect time. If missed set to beyond end of ray
-				float4 HitTime = Hit ? float4( 1, 2, 3, 4 ) : 5;
-
-				// Take closest hit
-				float Time1 = min( min3( HitTime.x, HitTime.y, HitTime.z ), HitTime.w );
-				float Time0 = Time1 - 1;
-
-				const uint NumBinarySteps = Roughness < 0.2 ? 4 : 0;
-
-				// Binary search
-				for( uint j = 0; j < NumBinarySteps; j++ )
-				{
-					CompareTolerance *= 0.5;
+			break;
+		}
 
-					float  MidTime = 0.5 * ( Time0 + Time1 );
-					float3 MidUVz = RayUVz + RayStepUVz * MidTime;
-					float  MidDepth = Texture.SampleLevel( Sampler, MidUVz.xy, Level ).r;
-					float  MidDepthDiff = MidUVz.z - MidDepth;
+		LastDiff = MultipleSampleDepthDiff.w;
 
-					if( abs( MidDepthDiff + CompareTolerance ) < CompareTolerance )
-					{
-						Time1 = MidTime;
-					}
-					else
-					{
-						Time0 = MidTime;
-					}
-				}
+		RayUVz += 4 * RayStepUVz;
+	} // for( uint i = 0; i < NumSteps; i += 4 )
+	
+	BRANCH
+	if (bFoundAnyHit)
+    {
+        if (IS_SSGI_SHADER)
+        {
+			// If hit set to intersect time. If missed set to beyond end of ray
+            float4 HitTime = bMultipleSampleHit ? float4(1, 2, 3, 4) : 5;
+
+			// Take closest hit
+            float Time1 = min(min3(HitTime.x, HitTime.y, HitTime.z), HitTime.w);
+            float Time0 = Time1 - 1;
+
+            const uint NumBinarySteps = Roughness < 0.2 ? 4 : 0;
+
+			// Binary search
+            for (uint j = 0; j < NumBinarySteps; j++)
+            {
+                CompareTolerance *= 0.5;
+
+                float MidTime = 0.5 * (Time0 + Time1);
+                float3 MidUVz = RayUVz + RayStepUVz * MidTime;
+                float MidDepth = Texture.SampleLevel(Sampler, MidUVz.xy, Level).r;
+                float MidDepthDiff = MidUVz.z - MidDepth;
+
+                if (abs(MidDepthDiff + CompareTolerance) < CompareTolerance)
+                {
+                    Time1 = MidTime;
+                }
+                else
+                {
+                    Time0 = MidTime;
+                }
+            }
 			
-				OutHitUVz = RayUVz + RayStepUVz * Time1;
-			}
-			else
+            OutHitUVz = RayUVz + RayStepUVz * Time1;
+        }
+        else // if (!IS_SSGI_SHADER)
+        {
+            float DepthDiff0 = MultipleSampleDepthDiff[2];
+            float DepthDiff1 = MultipleSampleDepthDiff[3];
+            float Time0 = 3;
+
+            FLATTEN
+            if (bMultipleSampleHit[2])
+            {
+                DepthDiff0 = MultipleSampleDepthDiff[1];
+                DepthDiff1 = MultipleSampleDepthDiff[2];
+                Time0 = 2;
+            }
+            FLATTEN
+            if (bMultipleSampleHit[1])
+            {
+                DepthDiff0 = MultipleSampleDepthDiff[0];
+                DepthDiff1 = MultipleSampleDepthDiff[1];
+                Time0 = 1;
+            }
+            FLATTEN
+            if (bMultipleSampleHit[0])
+            {
+                DepthDiff0 = LastDiff;
+                DepthDiff1 = MultipleSampleDepthDiff[0];
+                Time0 = 0;
+            }
+
+            float Time1 = Time0 + 1;
+			#if 0
 			{
-				float DepthDiff0 = DepthDiff[2];
-				float DepthDiff1 = DepthDiff[3];
-				float Time0 = 3;
-
-				FLATTEN if( Hit[2] )
-				{
-					DepthDiff0 = DepthDiff[1];
-					DepthDiff1 = DepthDiff[2];
-					Time0 = 2;
-				}
-				FLATTEN if( Hit[1] )
-				{
-					DepthDiff0 = DepthDiff[0];
-					DepthDiff1 = DepthDiff[1];
-					Time0 = 1;
-				}
-				FLATTEN if( Hit[0] )
-				{
-					DepthDiff0 = LastDiff;
-					DepthDiff1 = DepthDiff[0];
-					Time0 = 0;
-				}
-
-				float Time1 = Time0 + 1;
-	#if 0
 				// Binary search
 				for( uint j = 0; j < 4; j++ )
 				{
@@ -211,37 +229,34 @@ bool RayCast(
 						Time0		= MidTime;
 					}
 				}
-	#endif
-
-				// Find more accurate hit using line segment intersection
-				float TimeLerp = saturate( DepthDiff0 / (DepthDiff0 - DepthDiff1) );
-				float IntersectTime = Time0 + TimeLerp;
-				//float IntersectTime = lerp( Time0, Time1, TimeLerp );
-				
-				OutHitUVz = RayUVz + RayStepUVz * IntersectTime;
 			}
+			#endif
 
-			bHit = true;
-			break;
+			// Find more accurate hit using line segment intersection
+            float TimeLerp = saturate(DepthDiff0 / (DepthDiff0 - DepthDiff1));
+            float IntersectTime = Time0 + TimeLerp;
+			//float IntersectTime = lerp( Time0, Time1, TimeLerp );
+				
+            OutHitUVz = RayUVz + RayStepUVz * IntersectTime;
+        } // if (!IS_SSGI_SHADER)
+		
+		#if DEBUG_SSRT
+		{
+			if (bDebugPrint)
+				PrintSample(HZBUvFactorAndInvFactor, OutHitUVz.xy, float4(0, 0, 1, 1));
 		}
+		#endif
 
-		LastDiff = DepthDiff.w;
-
-		RayUVz += 4 * RayStepUVz;
-	}
+		OutHitUVz.xy *= HZBUvFactorAndInvFactor.zw;
+		OutHitUVz.xy = OutHitUVz.xy * float2( 2, -2 ) + float2( -1, 1 );
+		OutHitUVz.xy = OutHitUVz.xy * View.ScreenPositionScaleBias.xy + View.ScreenPositionScaleBias.wz;
+    }
+	else
+    {
+		OutHitUVz = float3(0, 0, 0);
+    }
 	
-	#if DEBUG_SSRT
-	{
-		if (bDebugPrint)
-			PrintSample(HZBUvFactorAndInvFactor, OutHitUVz.xy, float4(0, 0, 1, 1));
-	}
-	#endif
-
-	OutHitUVz.xy *= HZBUvFactorAndInvFactor.zw;
-	OutHitUVz.xy = OutHitUVz.xy * float2( 2, -2 ) + float2( -1, 1 );
-	OutHitUVz.xy = OutHitUVz.xy * View.ScreenPositionScaleBias.xy + View.ScreenPositionScaleBias.wz;
-
-	return bHit;
+	return bFoundAnyHit;
 }
 
 float ComputeHitVignetteFromScreenPos(float2 ScreenPos)
diff --git a/Engine/Source/Runtime/Renderer/Private/ScreenSpaceRayTracing.cpp b/Engine/Source/Runtime/Renderer/Private/ScreenSpaceRayTracing.cpp
--- a/Engine/Source/Runtime/Renderer/Private/ScreenSpaceRayTracing.cpp
+++ b/Engine/Source/Runtime/Renderer/Private/ScreenSpaceRayTracing.cpp
@@ -214,6 +214,8 @@ class FScreenSpaceReflectionsPS : public FGlobalShader
 		
 		SHADER_PARAMETER_RDG_TEXTURE(Texture2D, HZB)
 		SHADER_PARAMETER_SAMPLER(SamplerState, HZBSampler)
+
+		SHADER_PARAMETER_RDG_TEXTURE_UAV(RWTexture<float4>, ScreenSpaceRayTracingDebugOutput)
 		
 		RENDER_TARGET_BINDING_SLOTS()
 	END_SHADER_PARAMETER_STRUCT()
@@ -325,6 +327,31 @@ void GetSSRTGIShaderOptionsForQuality(int32 Quality, FIntPoint* OutGroupSize, in
 	check(OutGroupSize->X * OutGroupSize->Y * (*OutRayCountPerPixel) == 256);
 }
 
+FRDGTextureUAV* CreateScreenSpaceRayTracingDebugUAV(FRDGBuilder& GraphBuilder, const FRDGTextureDesc& Desc)
+#if 0
+{
+	FRDGTexture* DebugTexture = GraphBuilder.CreateTexture(Desc, TEXT("SSRTDebug"));
+
+	FRenderTargetParameters* PassParameters = GraphBuilder.AllocParameters<FRenderTargetParameters>();
+	PassParameters->RenderTargets[0] = FRenderTargetBinding(DebugTexture, ERenderTargetLoadAction::EClear, ERenderTargetStoreAction::EStore);
+
+	GraphBuilder.AddPass(
+		RDG_EVENT_NAME("Clear SSRTDebug"),
+		PassParameters,
+		ERenderGraphPassFlags::None,
+		[](FRHICommandList& RHICmdList)
+	{
+		// NOP
+	});
+
+	return GraphBuilder.CreateUAV(DebugTexture);
+}
+#else
+{
+	return nullptr;
+}
+#endif
+
 } // namespace
 
 void GetSSRQualityForView(const FViewInfo& View, ESSRQuality* OutQuality, IScreenSpaceDenoiser::FReflectionsRayTracingConfig* OutRayTracingConfigs)
@@ -513,7 +540,8 @@ void RenderScreenSpaceReflections(
 		
 		PassParameters->HZB = GraphBuilder.RegisterExternalTexture(View.HZB);
 		PassParameters->HZBSampler = TStaticSamplerState<SF_Point>::GetRHI();
-		
+
+		PassParameters->ScreenSpaceRayTracingDebugOutput = CreateScreenSpaceRayTracingDebugUAV(GraphBuilder, DenoiserInputs->Color->Desc);
 		PassParameters->RenderTargets = RenderTargets;
 
 		TShaderMapRef<FScreenSpaceReflectionsPS> PixelShader(View.ShaderMap, PermutationVector);
@@ -626,25 +654,7 @@ void RenderScreenSpaceDiffuseIndirect(
 	
 	PassParameters->IndirectDiffuseOutput = GraphBuilder.CreateUAV(OutDenoiserInputs->Color);
 	PassParameters->AmbientOcclusionOutput = GraphBuilder.CreateUAV(OutDenoiserInputs->AmbientOcclusionMask);
-
-	if (0) // DEBUG output clear
-	{
-		FRDGTexture* DebugTexture = GraphBuilder.CreateTexture(OutDenoiserInputs->Color->Desc, TEXT("SSRTDebug"));
-
-		FRenderTargetParameters* ClearPassParameters = GraphBuilder.AllocParameters<FRenderTargetParameters>();
-		ClearPassParameters->RenderTargets[0] = FRenderTargetBinding(DebugTexture, ERenderTargetLoadAction::EClear, ERenderTargetStoreAction::EStore);
-
-		GraphBuilder.AddPass(
-			RDG_EVENT_NAME("Clear SSRTDebug"),
-			ClearPassParameters,
-			ERenderGraphPassFlags::None,
-			[](FRHICommandList& RHICmdList)
-		{
-			// NOP
-		});
-
-		PassParameters->ScreenSpaceRayTracingDebugOutput = GraphBuilder.CreateUAV(DebugTexture);
-	}
+	PassParameters->ScreenSpaceRayTracingDebugOutput = CreateScreenSpaceRayTracingDebugUAV(GraphBuilder, OutDenoiserInputs->Color->Desc);
 
 	FScreenSpaceDiffuseIndirectCS::FPermutationDomain PermutationVector;
 	PermutationVector.Set<FScreenSpaceDiffuseIndirectCS::FQualityDim>(Quality);
