Standardizes a platform agnostic INFINITE_FLOAT for shaders

diff --git a/Engine/Shaders/Private/RayTracing/RayTracingAmbientOcclusionRGS.usf b/Engine/Shaders/Private/RayTracing/RayTracingAmbientOcclusionRGS.usf
--- a/Engine/Shaders/Private/RayTracing/RayTracingAmbientOcclusionRGS.usf
+++ b/Engine/Shaders/Private/RayTracing/RayTracingAmbientOcclusionRGS.usf
@@ -61,7 +61,7 @@ void AmbientOcclusionRGS()
 
 	float RayCount = 0.0;
 	float Visibility = 0.0;
-	float ClosestRayHitDistance = 1.#INF;
+	float ClosestRayHitDistance = INFINITE_FLOAT;
 	// Declaring intensity as a local variable is a workaround for a shader compiler bug with NV drivers 430.39 and 430.64.
 	const float IntensityLocal = Intensity;
 	uint SamplesPerPixelLocal = SamplesPerPixel;
diff --git a/Engine/Shaders/Private/SSRT/SSRTReflections.usf b/Engine/Shaders/Private/SSRT/SSRTReflections.usf
--- a/Engine/Shaders/Private/SSRT/SSRTReflections.usf
+++ b/Engine/Shaders/Private/SSRT/SSRTReflections.usf
@@ -27,12 +27,6 @@ SamplerState SceneColorSampler;
 Texture2D HZB;
 SamplerState HZBSampler;
 
-#if COMPILER_HLSLCC || PS4_PROFILE
-#define INFINITE_WORKAROUND		3.402823e+38
-#else
-#define INFINITE_WORKAROUND		1.#INF
-#endif
-
 
 float GetRoughness(const in FGBufferData GBuffer)
 {
@@ -137,7 +131,7 @@ void ScreenSpaceReflectionsPS(
 	float NoV = saturate( dot( N, V ) );
 	float G_SmithV = 2 * NoV / (NoV + sqrt(NoV * (NoV - NoV * a2) + a2));
 
-	float ClosestHitDistanceSqr = INFINITE_WORKAROUND;
+	float ClosestHitDistanceSqr = INFINITE_FLOAT;
 
 #if SSR_QUALITY == 1
 	uint NumSteps = 8;
@@ -284,7 +278,7 @@ void ScreenSpaceReflectionsPS(
 		OutColor /= max( NumRays, 0.0001 );
 		OutColor.rgb *= rcp( 1 - Luminance(OutColor.rgb) );
 
-		if (ClosestHitDistanceSqr == INFINITE_WORKAROUND)
+		if (ClosestHitDistanceSqr == INFINITE_FLOAT)
 		{
 			ClosestHitDistanceSqr = DENOISER_MISS_HIT_DISTANCE;
 		}
diff --git a/Engine/Shaders/Private/ScreenSpaceDenoise/SSDSignalAccumulator.ush b/Engine/Shaders/Private/ScreenSpaceDenoise/SSDSignalAccumulator.ush
--- a/Engine/Shaders/Private/ScreenSpaceDenoise/SSDSignalAccumulator.ush
+++ b/Engine/Shaders/Private/ScreenSpaceDenoise/SSDSignalAccumulator.ush
@@ -50,8 +50,8 @@ FSSDSignalAccumulator CreateSignalAccumulator()
 	Accumulator.Moment1 = CreateSignalSampleFromScalarValue(0.0);
 	Accumulator.Moment2 = CreateSignalSampleFromScalarValue(0.0);
 		
-	Accumulator.Min = CreateSignalSampleFromScalarValue(1.#INF);
-	Accumulator.Max = CreateSignalSampleFromScalarValue(0.0); // TODO: -1.#INF? otherwise there is a risk to color clamp with YCoCg.
+	Accumulator.Min = CreateSignalSampleFromScalarValue(INFINITE_FLOAT);
+	Accumulator.Max = CreateSignalSampleFromScalarValue(0.0); // TODO: -INFINITE_FLOAT? otherwise there is a risk to color clamp with YCoCg.
 
 	Accumulator.MinInvFrequency = WORLD_RADIUS_MISS;
 
diff --git a/Engine/Shaders/Public/Platform.ush b/Engine/Shaders/Public/Platform.ush
--- a/Engine/Shaders/Public/Platform.ush
+++ b/Engine/Shaders/Public/Platform.ush
@@ -690,3 +690,9 @@ float UnpackByte3(uint v) { return float(v >> 24); }
 	#define UNORM unorm
 #endif
 #endif
+
+#if COMPILER_HLSLCC || PS4_PROFILE
+	#define INFINITE_FLOAT 3.402823e+38
+#else
+	#define INFINITE_FLOAT 1.#INF
+#endif
