Stops keeping arround DispatchThreadId until the end of the spatial reconstruction, saving 2 VGPRs.

diff --git a/Engine/Shaders/Private/ScreenSpaceDenoise/SSDDefinitions.ush b/Engine/Shaders/Private/ScreenSpaceDenoise/SSDDefinitions.ush
--- a/Engine/Shaders/Private/ScreenSpaceDenoise/SSDDefinitions.ush
+++ b/Engine/Shaders/Private/ScreenSpaceDenoise/SSDDefinitions.ush
@@ -192,7 +192,7 @@
 #define CONFIG_NORMAL_ANISOTROPY 0
 
 /** Whether should enable VGPR optimizations. */
-#define CONFIG_VGPR_OPTIMIZATION 0
+#define CONFIG_VGPR_OPTIMIZATION 1
 
 /** Whether should do VALU optimizations. */
 #define CONFIG_VALU_OPTIMIZATIONS 1
diff --git a/Engine/Shaders/Private/ScreenSpaceDenoise/SSDSignalAccumulator.ush b/Engine/Shaders/Private/ScreenSpaceDenoise/SSDSignalAccumulator.ush
--- a/Engine/Shaders/Private/ScreenSpaceDenoise/SSDSignalAccumulator.ush
+++ b/Engine/Shaders/Private/ScreenSpaceDenoise/SSDSignalAccumulator.ush
@@ -253,6 +253,8 @@ void AccumulateSample(
 	inout FSSDSignalAccumulator Accumulator,
 	FSSDSampleAccumulationInfos SampleInfos)
 {
+	// TODO: Decompress and recompress the accumulator to reduce VGPR pressure.
+
 	AccumulateSampleInOneBin(/* inout */ Accumulator, SampleInfos);
 	AccumulateSampleInDomainBoundaries(/* inout */ Accumulator, SampleInfos);
 	AccumulateSampleInDRB(/* inout */ Accumulator, SampleInfos);
diff --git a/Engine/Shaders/Private/ScreenSpaceDenoise/SSDSignalCore.ush b/Engine/Shaders/Private/ScreenSpaceDenoise/SSDSignalCore.ush
--- a/Engine/Shaders/Private/ScreenSpaceDenoise/SSDSignalCore.ush
+++ b/Engine/Shaders/Private/ScreenSpaceDenoise/SSDSignalCore.ush
@@ -617,6 +617,8 @@ void DecodeMultiplexedSignals(
 		OutSamples.Array[0].ColorSH.R.V = RawSample[0] * SampleCount;
 		OutSamples.Array[0].ColorSH.G.V = RawSample[1] * SampleCount;
 		OutSamples.Array[0].ColorSH.B.V = RawSample[2] * SampleCount;
+
+		// TODO: SRV on RG F32 and manually decode with f16tof32(), going to reduce VGPR pressure between fetch and result.
 	}
 	#endif
 } // DecodeMultiplexedSignals()
diff --git a/Engine/Shaders/Private/ScreenSpaceDenoise/SSDSignalFramework.ush b/Engine/Shaders/Private/ScreenSpaceDenoise/SSDSignalFramework.ush
--- a/Engine/Shaders/Private/ScreenSpaceDenoise/SSDSignalFramework.ush
+++ b/Engine/Shaders/Private/ScreenSpaceDenoise/SSDSignalFramework.ush
@@ -537,6 +537,12 @@ float ComputeBilateralWeight(
 	}
 	#elif CONFIG_VIEW_POSITION_BASED_BILATERAL == 2
 	{
+		// TODO: save 2 VGPR by encoding RefSceneMetadata.WorldNormal
+
+		// TODO: save 2 VGPR by using SceneDepth instead of RefSceneMetadata.TranslatedWorldPosition (not sure it can work in math yet).
+		// float2 ClipPosition = ScreenPosition * (View.ViewToClip[3][3] < 1.0f ? WorldDepth : 1.0f);
+		// Infos.TranslatedWorldPosition = mul(float4(ClipPosition, WorldDepth, 1), ScreenToTranslatedWorld).xyz;
+
 		float DistFromRefPlane = abs(dot(RefSceneMetadata.WorldNormal, RefSceneMetadata.TranslatedWorldPosition - NeighborSceneMetadata.TranslatedWorldPosition));
 
 		PositionBilateralWeight = saturate(1 - DistFromRefPlane * rcp(MaxWorldBluringDistance));
diff --git a/Engine/Shaders/Private/ScreenSpaceDenoise/SSDSpatialAccumulation.usf b/Engine/Shaders/Private/ScreenSpaceDenoise/SSDSpatialAccumulation.usf
--- a/Engine/Shaders/Private/ScreenSpaceDenoise/SSDSpatialAccumulation.usf
+++ b/Engine/Shaders/Private/ScreenSpaceDenoise/SSDSpatialAccumulation.usf
@@ -550,6 +550,11 @@ void MainCS(
 		#error Unknown policy to control the number of samples.
 	#endif
 
+	// Register renaming of members of FSSDKernelConfig to survive until the output to UAV
+	#if CONFIG_SAMPLE_SET == SAMPLE_SET_STACKOWIAK_4_SETS && CONFIG_VGPR_OPTIMIZATION
+		float2 KernelBufferUV;
+		uint SampleTrackId;
+	#endif
 
 	// Accumulate spatially the input.
 	FSSDSignalAccumulatorArray SignalAccumulators = CreateSignalAccumulatorArray();
@@ -584,6 +589,8 @@ void MainCS(
 		KernelConfig.BufferSizeAndInvSize = View.BufferSizeAndInvSize;
 		KernelConfig.BufferBilinearUVMinMax = View.BufferBilinearUVMinMax;
 		KernelConfig.KernelSpreadFactor = KernelSpreadFactor;
+
+		// VGPRs
 		KernelConfig.BufferUV = SceneBufferUV;
 		KernelConfig.RefSceneMetadata = RefSceneMetadata;
 		
@@ -667,7 +674,7 @@ void MainCS(
 			}
 		}
 
-		if (CONFIG_SAMPLE_SET == SAMPLE_SET_STACKOWIAK_4_SETS)
+		#if CONFIG_SAMPLE_SET == SAMPLE_SET_STACKOWIAK_4_SETS
 		{
 			KernelConfig.SampleCount = clamp(RequestedSampleCount * rcp(kStackowiakSampleSetCount), 1, MaxSampleCount);
 
@@ -698,7 +705,16 @@ void MainCS(
 				// Id of the pixel in the quad. This is to match hard coded first samples of the sample set.
 				KernelConfig.SampleTrackId = ((DispatchThreadId.x & 1) | ((DispatchThreadId.y & 1) << 1));
 			}
+
+			#if CONFIG_VGPR_OPTIMIZATION
+				// Keek sample SampleTrackId & SceneBufferUV arround for computation of pixel output coordinate.
+				// Should be VGPR free given it's curernt is being used in accumulation has well that is highest VGPR pressure of the shader.
+				// TODO: could save 1 VGPR by using 2 SGPR instead of SampleTrackId.
+				SampleTrackId = KernelConfig.SampleTrackId;
+				KernelBufferUV = KernelConfig.BufferUV;
+			#endif
 		}
+		#endif // CONFIG_SAMPLE_SET == SAMPLE_SET_STACKOWIAK_4_SETS
 
 		AccumulateKernel(
 			KernelConfig,
@@ -707,37 +723,34 @@ void MainCS(
 			SignalInput_Textures_2,
 			SignalInput_Textures_3,
 			/* inout */ SignalAccumulators);
-	
-		//if (DIM_STAGE == STAGE_REJECTION_PRE_CONVOLUTION)
-		//	DebugOutput[DispatchThreadId] = float4(
-		//		KernelConfig.bForceAllAccumulation ? 1 : 0,
-		//		0,
-		//		0,
-		//		0);
 
 		// When doing history rejection pre convolution, could still have no information found with the 3x3 + kernel,
 		// therefore dynamically complete to form an entire 3x3 convolution.
-		BRANCH
-		if (DIM_STAGE == STAGE_REJECTION_PRE_CONVOLUTION &&
-			KernelConfig.SampleSet == SAMPLE_SET_3X3_PLUS &&
-			KernelConfig.bForceAllAccumulation)
+		#if DIM_STAGE == STAGE_REJECTION_PRE_CONVOLUTION
 		{
-			KernelConfig.SampleSet = SAMPLE_SET_3X3_CROSS;
-			KernelConfig.bSampleKernelCenter = false;
+			BRANCH
+			if (KernelConfig.SampleSet == SAMPLE_SET_3X3_PLUS &&
+				KernelConfig.bForceAllAccumulation)
+			{
+				KernelConfig.SampleSet = SAMPLE_SET_3X3_CROSS;
+				KernelConfig.bSampleKernelCenter = false;
 			
-			AccumulateKernel(
-				KernelConfig,
-				SignalInput_Textures_0,
-				SignalInput_Textures_1,
-				SignalInput_Textures_2,
-				SignalInput_Textures_3,
-				/* inout */ SignalAccumulators);
+				AccumulateKernel(
+					KernelConfig,
+					SignalInput_Textures_0,
+					SignalInput_Textures_1,
+					SignalInput_Textures_2,
+					SignalInput_Textures_3,
+					/* inout */ SignalAccumulators);
+			}
 		}
+		#endif // DIM_STAGE == STAGE_REJECTION_PRE_CONVOLUTION
 		
 		// Manually sample the center of the kernel after any accumulation when accumulating in descending order.
 		if (!KernelConfig.bSampleKernelCenter && KernelConfig.bDescOrder)
 		{
 			// Remove any jitter the kernel may have. Won't have ant VGPR cost when no jittering, because KernelConfig.BufferUV == SceneBufferUV.
+			// TODO: This is costly for VGPR pressure if using KernelConfig.BufferUV was != SceneBufferUV.
 			KernelConfig.BufferUV = SceneBufferUV;
 
 			SampleAndAccumulateCenterSampleAsItsOwnCluster(
@@ -840,7 +853,7 @@ void MainCS(
 	}
 	
 	// Clamp the number of sample recorded.
-	if (DIM_STAGE == STAGE_POST_FILTERING)
+	#if DIM_STAGE == STAGE_POST_FILTERING
 	{
 		UNROLL_N(CONFIG_SIGNAL_BATCH_SIZE)
 		for (uint MultiplexId = 0; MultiplexId < CONFIG_SIGNAL_BATCH_SIZE; MultiplexId++)
@@ -851,12 +864,20 @@ void MainCS(
 			OutputSamples.Array[MultiplexId] = Mul(OutputSamples.Array[MultiplexId], CurrentSampleCount > 0 ? NewSampleCount / CurrentSampleCount : 0);
 		}
 	}
+	#endif // DIM_STAGE == STAGE_POST_FILTERING
 
 	// TODO:  LeaveRayCount = (LeaveRayCount - 1) * 9 / (9 - 2) post processing to reject when for history rejection
 	
-	#if CONFIG_VGPR_OPTIMIZATION && CONFIG_SAMPLE_SET != SAMPLE_SET_STACKOWIAK_4_SETS
-		// No need to keep DispatchThreadId, while SceneBufferUV is arround at highest VGPR peak because center of the kernel.
-		uint2 OutputPixelPostion = uint2(SceneBufferUV * View.BufferSizeAndInvSize.xy);
+	#if CONFIG_VGPR_OPTIMIZATION
+		// No need to keep DispatchThreadId, can recompute the output pixel position based on information stored in VGPRs for spatial kernel.
+		#if CONFIG_SAMPLE_SET == SAMPLE_SET_STACKOWIAK_4_SETS
+			#if CONFIG_UPSCALE
+				SampleTrackId ^= 0x3;
+			#endif
+			uint2 OutputPixelPostion = (uint2(KernelBufferUV * View.BufferSizeAndInvSize.xy) & ~0x1) | (uint2(SampleTrackId, SampleTrackId >> 1) & 0x1);
+		#else
+			uint2 OutputPixelPostion = uint2(SceneBufferUV * View.BufferSizeAndInvSize.xy);
+		#endif
 	#else
 		uint2 OutputPixelPostion = DispatchThreadId;
 	#endif 
diff --git a/Engine/Shaders/Private/ScreenSpaceDenoise/SSDSpatialKernel.ush b/Engine/Shaders/Private/ScreenSpaceDenoise/SSDSpatialKernel.ush
--- a/Engine/Shaders/Private/ScreenSpaceDenoise/SSDSpatialKernel.ush
+++ b/Engine/Shaders/Private/ScreenSpaceDenoise/SSDSpatialKernel.ush
@@ -926,6 +926,7 @@ void ConvolveStackowiakKernel(
 	}
 	else // if (!KernelConfig.bDescOrder)
 	{
+		// TODO: kSamplingBatchSize for lattency hidding
 		LOOP
 		for (uint SampleId = (KernelConfig.bSampleKernelCenter ? 0 : 1); SampleId < KernelConfig.SampleCount; SampleId++)
 		{
