Moves buffer encodings for denoiser's signals into its own file.

diff --git a/Engine/Shaders/Private/ScreenSpaceDenoise/SSDInjest.usf b/Engine/Shaders/Private/ScreenSpaceDenoise/SSDInjest.usf
--- a/Engine/Shaders/Private/ScreenSpaceDenoise/SSDInjest.usf
+++ b/Engine/Shaders/Private/ScreenSpaceDenoise/SSDInjest.usf
@@ -33,6 +33,7 @@
 
 #include "SSDSignalFramework.ush"
 #include "SSDSignalArray.ush"
+#include "SSDSignalBufferEncoding.ush"
 
 
 //------------------------------------------------------- PARAMETERS
diff --git a/Engine/Shaders/Private/ScreenSpaceDenoise/SSDSignalBufferEncoding.ush b/Engine/Shaders/Private/ScreenSpaceDenoise/SSDSignalBufferEncoding.ush
new file mode 100644
--- /dev/null
+++ b/Engine/Shaders/Private/ScreenSpaceDenoise/SSDSignalBufferEncoding.ush
@@ -0,0 +1,537 @@
+// Copyright 1998-2019 Epic Games, Inc. All Rights Reserved.
+
+#pragma once
+
+#include "SSDSignalCore.ush"
+
+
+/** Decode input signal sample. */
+void DecodeMultiplexedSignals(
+	const uint SignalBufferLayout,
+	const uint MultiplexedSampleId,
+	const bool bNormalizeSample,
+	float4 RawSample[MAX_MULTIPLEXED_TEXTURES],
+	out FSSDSignalArray OutSamples)
+{
+	OutSamples = CreateSignalArrayFromScalarValue(0.0);
+
+	if (SignalBufferLayout == SIGNAL_BUFFER_LAYOUT_UNINITIALIZED)
+	{
+		// trash out the output sample to identify quickly a miss configured buffer layout.
+		OutSamples = CreateSignalArrayFromScalarValue(-INFINITE_FLOAT);
+	}
+	else if (SignalBufferLayout == SIGNAL_BUFFER_LAYOUT_PENUMBRA_INPUT_1SPP)
+	{
+		UNROLL_N(MAX_SIGNAL_BATCH_SIZE)
+		for (uint BatchSignalId = 0; BatchSignalId < MAX_SIGNAL_BATCH_SIZE; BatchSignalId++)
+		{
+			uint MultiplexId = BatchSignalId;
+			OutSamples.Array[MultiplexId].SampleCount = (RawSample[MultiplexId].r == DENOISER_INVALID_HIT_DISTANCE ? 0.0 : 1.0);
+			OutSamples.Array[MultiplexId].MissCount = OutSamples.Array[MultiplexId].SampleCount * (RawSample[MultiplexId].r > 0.0 ? 0.0 : 1.0);
+			OutSamples.Array[MultiplexId].ClosestHitDistance = OutSamples.Array[MultiplexId].SampleCount * (RawSample[MultiplexId].r == DENOISER_MISS_HIT_DISTANCE ? 1000000 : RawSample[MultiplexId].r);
+		}
+	}
+	else if (SignalBufferLayout == SIGNAL_BUFFER_LAYOUT_PENUMBRA_INPUT_NSPP)
+	{
+		UNROLL_N(MAX_SIGNAL_BATCH_SIZE)
+		for (uint BatchSignalId = 0; BatchSignalId < MAX_SIGNAL_BATCH_SIZE; BatchSignalId++)
+		{
+			uint MultiplexId = BatchSignalId;
+			// TODO: feed the actual number of sample.
+			OutSamples.Array[MultiplexId].SampleCount = (RawSample[MultiplexId].g == DENOISER_INVALID_HIT_DISTANCE ? 0.0 : 1.0);
+			OutSamples.Array[MultiplexId].MissCount = OutSamples.Array[MultiplexId].SampleCount * RawSample[MultiplexId].r;
+			OutSamples.Array[MultiplexId].ClosestHitDistance = OutSamples.Array[MultiplexId].SampleCount * (RawSample[MultiplexId].g == DENOISER_MISS_HIT_DISTANCE ? 1000000 : RawSample[MultiplexId].g);
+		}
+	}
+	else if (SignalBufferLayout == SIGNAL_BUFFER_LAYOUT_PENUMBRA_INJESTION_1SPP)
+	{
+		UNROLL_N(MAX_SIGNAL_BATCH_SIZE)
+		for (uint BatchSignalId = 0; BatchSignalId < MAX_SIGNAL_BATCH_SIZE; BatchSignalId++)
+		{
+			uint MultiplexId = BatchSignalId;
+			float EncodedWorldBluringRadius = RawSample[MultiplexId / 4][MultiplexId % 4];
+			OutSamples.Array[MultiplexId].SampleCount = (EncodedWorldBluringRadius == -2.0 ? 0.0 : 1.0);
+			OutSamples.Array[MultiplexId].MissCount = OutSamples.Array[MultiplexId].SampleCount * (EncodedWorldBluringRadius > 0.0 ? 0.0 : 1.0);
+			OutSamples.Array[MultiplexId].WorldBluringRadius = OutSamples.Array[MultiplexId].SampleCount * (EncodedWorldBluringRadius == -1 ? WORLD_RADIUS_MISS : EncodedWorldBluringRadius);
+		}
+	}
+	else if (SignalBufferLayout == SIGNAL_BUFFER_LAYOUT_PENUMBRA_INJESTION_NSPP)
+	{
+		UNROLL_N(MAX_SIGNAL_BATCH_SIZE)
+		for (uint BatchSignalId = 0; BatchSignalId < MAX_SIGNAL_BATCH_SIZE; BatchSignalId++)
+		{
+			uint MultiplexId = BatchSignalId;
+			float2 Channels = float2(RawSample[MultiplexId / 2][2 * (MultiplexId % 2)], RawSample[MultiplexId / 2][2 * (MultiplexId % 2) + 1]);
+	
+			// TODO: feed the actual number of sample.
+			OutSamples.Array[MultiplexId].SampleCount = (Channels.g == -2.0 ? 0.0 : 1.0);
+			OutSamples.Array[MultiplexId].MissCount = OutSamples.Array[MultiplexId].SampleCount * Channels.r;
+			OutSamples.Array[MultiplexId].WorldBluringRadius = OutSamples.Array[MultiplexId].SampleCount * (Channels.g == -1 ? WORLD_RADIUS_MISS : Channels.g);
+		}
+	}
+	//else if (SignalBufferLayout == SIGNAL_BUFFER_LAYOUT_PENUMBRA_RECONSTRUCTION)
+	//{
+	//	UNROLL_N(MAX_SIGNAL_MULTIPLEXING)
+	//	for (uint MultiplexId = 0; MultiplexId < MAX_SIGNAL_MULTIPLEXING; MultiplexId++)
+	//	{
+	//		float2 Channels = float2(RawSample[MultiplexId / 2][2 * (MultiplexId % 2)], RawSample[MultiplexId / 2][2 * (MultiplexId % 2) + 1]);
+	//
+	//		if (bNormalizeSample)
+	//		{
+	//			OutSamples.Array[MultiplexId].SampleCount = Channels.g > 0 ? 1 : 0;
+	//			OutSamples.Array[MultiplexId].MissCount = Channels.g > 0 ? Channels.r : 0;
+	//		}
+	//		else
+	//		{
+	//			OutSamples.Array[MultiplexId].SampleCount = Channels.g;
+	//			OutSamples.Array[MultiplexId].MissCount = Channels.r * Channels.g;
+	//		}
+	//	}
+	//}
+	else if (SignalBufferLayout == SIGNAL_BUFFER_LAYOUT_PENUMBRA_HISTORY || SignalBufferLayout == SIGNAL_BUFFER_LAYOUT_PENUMBRA_RECONSTRUCTION)
+	{
+		UNROLL_N(MAX_SIGNAL_BATCH_SIZE)
+		for (uint MultiplexId = 0; MultiplexId < MAX_SIGNAL_BATCH_SIZE; MultiplexId++)
+		{
+			float3 Channels = RawSample[MultiplexId].xyz;
+			
+			float SampleCount = bNormalizeSample ? (Channels.g > 0 ? 1 : 0) : (Channels.g);
+
+			OutSamples.Array[MultiplexId].MissCount = Channels.r * SampleCount;
+			OutSamples.Array[MultiplexId].SampleCount = SampleCount;
+			OutSamples.Array[MultiplexId].WorldBluringRadius = (Channels.b == DENOISER_MISS_HIT_DISTANCE ? WORLD_RADIUS_MISS : Channels.b) * SampleCount;
+		}
+	}
+	else if (SignalBufferLayout == SIGNAL_BUFFER_LAYOUT_PENUMBRA_REJECTION)
+	{
+		UNROLL_N(MAX_SIGNAL_BATCH_SIZE)
+		for (uint BatchSignalId = 0; BatchSignalId < MAX_SIGNAL_BATCH_SIZE; BatchSignalId++)
+		{
+			uint MultiplexId = BatchSignalId * 2;
+			float4 Channels = RawSample[BatchSignalId];
+			float SampleCount = Channels.b;
+
+			OutSamples.Array[MultiplexId + 0].SampleCount = SampleCount;
+			OutSamples.Array[MultiplexId + 0].MissCount = Channels.r * SampleCount;
+			OutSamples.Array[MultiplexId + 0].WorldBluringRadius = (Channels.a == -1 ? WORLD_RADIUS_MISS : Channels.a) * SampleCount;
+			
+			OutSamples.Array[MultiplexId + 1].SampleCount = SampleCount;
+			OutSamples.Array[MultiplexId + 1].MissCount = Channels.g * SampleCount;
+		}
+	}
+	#if CONFIG_SIGNAL_COLOR
+	else if (SignalBufferLayout == SIGNAL_BUFFER_LAYOUT_REFLECTIONS_INPUT)
+	{
+		OutSamples.Array[0].SampleCount = RawSample[1].r == DENOISER_INVALID_HIT_DISTANCE ? 0 : 1;
+		OutSamples.Array[0].MissCount = OutSamples.Array[0].SampleCount * (1 - RawSample[0].a);
+		OutSamples.Array[0].SceneColor = RawSample[0];
+
+		// max() because ray hit distance is < 0 when miss / not ray has been shot.
+		//OutSamples.Array[0].ClosestHitDistance = RawSample[1].r * OutSamples.Array[0].SampleCount;
+		OutSamples.Array[0].ClosestHitDistance = ((RawSample[1].r == DENOISER_MISS_HIT_DISTANCE) ? 1.0e20 : RawSample[1].r) * OutSamples.Array[0].SampleCount;
+	}
+	else if (SignalBufferLayout == SIGNAL_BUFFER_LAYOUT_REFLECTIONS_REJECTION)
+	{
+		if (bNormalizeSample)
+		{
+			OutSamples.Array[0].SampleCount = RawSample[1].r > 0 ? 1 : 0;
+		}
+		else
+		{
+			OutSamples.Array[0].SampleCount = RawSample[1].r;
+		}
+		
+		OutSamples.Array[0].MissCount = OutSamples.Array[0].SampleCount * (1 - RawSample[0].a);
+		OutSamples.Array[0].SceneColor = OutSamples.Array[0].SampleCount * RawSample[0];
+
+		OutSamples.Array[1].SampleCount = RawSample[1].g;
+		OutSamples.Array[1].MissCount = OutSamples.Array[0].SampleCount * (1 - RawSample[2].a);
+		OutSamples.Array[1].SceneColor = OutSamples.Array[0].SampleCount * RawSample[2];
+	}
+	else if (SignalBufferLayout == SIGNAL_BUFFER_LAYOUT_REFLECTIONS_HISTORY)
+	{
+		if (bNormalizeSample)
+		{
+			OutSamples.Array[0].SampleCount = RawSample[1].r > 0 ? 1 : 0;
+		}
+		else
+		{
+			OutSamples.Array[0].SampleCount = RawSample[1].r;
+		}
+		
+		OutSamples.Array[0].MissCount = OutSamples.Array[0].SampleCount * (1 - RawSample[0].a);
+		OutSamples.Array[0].SceneColor = OutSamples.Array[0].SampleCount * RawSample[0];
+
+		OutSamples.Array[1].SampleCount = OutSamples.Array[0].SampleCount;
+		OutSamples.Array[1].SceneColor.x = OutSamples.Array[1].SampleCount * RawSample[1].g;
+		
+		OutSamples.Array[2].SampleCount = OutSamples.Array[0].SampleCount;
+		OutSamples.Array[2].SceneColor.x = OutSamples.Array[2].SampleCount * RawSample[1].b;
+	}
+	#endif
+	else if (SignalBufferLayout == SIGNAL_BUFFER_LAYOUT_AO_INPUT)
+	{
+		OutSamples.Array[0].SampleCount = RawSample[1].r < 0 ? 0 : 1;
+		OutSamples.Array[0].MissCount = OutSamples.Array[0].SampleCount * RawSample[0].r;
+		OutSamples.Array[0].ClosestHitDistance = OutSamples.Array[0].SampleCount * RawSample[1].r;
+	}
+	else if (SignalBufferLayout == SIGNAL_BUFFER_LAYOUT_AO_REJECTION)
+	{
+		if (bNormalizeSample)
+		{
+			OutSamples.Array[0].SampleCount = RawSample[0].g > 0 ? 1 : 0;
+			OutSamples.Array[0].MissCount = RawSample[0].g > 0 ? RawSample[0].r : 0;
+			
+			OutSamples.Array[1].SampleCount = RawSample[0].a > 0 ? 1 : 0;
+			OutSamples.Array[1].MissCount = RawSample[0].a > 0 ? RawSample[0].b : 0;
+		}
+		else
+		{
+			OutSamples.Array[0].SampleCount = RawSample[0].g;
+			OutSamples.Array[0].MissCount = RawSample[0].g * RawSample[0].r;
+			
+			OutSamples.Array[1].SampleCount = RawSample[0].a;
+			OutSamples.Array[1].MissCount = RawSample[0].a * RawSample[0].b;
+		}
+	}
+	else if (SignalBufferLayout == SIGNAL_BUFFER_LAYOUT_AO_HISTORY)
+	{
+		if (bNormalizeSample)
+		{
+			OutSamples.Array[0].SampleCount = RawSample[0].g > 0 ? 1 : 0;
+			OutSamples.Array[0].MissCount = RawSample[0].g > 0 ? RawSample[0].r : 0;
+		}
+		else
+		{
+			OutSamples.Array[0].SampleCount = RawSample[0].g;
+			OutSamples.Array[0].MissCount = RawSample[0].g * RawSample[0].r;
+		}
+	}
+	#if CONFIG_SIGNAL_COLOR
+	else if (SignalBufferLayout == SIGNAL_BUFFER_LAYOUT_DIFFUSE_INDIRECT_AND_AO_INPUT_NSPP)
+	{
+		OutSamples.Array[0].SampleCount = RawSample[1].g;
+		OutSamples.Array[0].MissCount = OutSamples.Array[0].SampleCount * (1 - RawSample[0].a);
+		OutSamples.Array[0].SceneColor = OutSamples.Array[0].SampleCount * RawSample[0];
+
+		// max() because ray hit distance is < 0 when miss / not ray has been shot.
+		OutSamples.Array[0].ClosestHitDistance = RawSample[1].r * OutSamples.Array[0].SampleCount;
+		OutSamples.Array[0].ClosestHitDistance = ((RawSample[1].r == -1.0) ? 1.0e20 : RawSample[1].r) * OutSamples.Array[0].SampleCount;
+	}
+	else if (SignalBufferLayout == SIGNAL_BUFFER_LAYOUT_DIFFUSE_INDIRECT_AND_AO_RECONSTRUCTION)
+	{
+		if (bNormalizeSample)
+		{
+			OutSamples.Array[0].SampleCount = RawSample[1].r > 0 ? 1 : 0;
+		}
+		else
+		{
+			OutSamples.Array[0].SampleCount = RawSample[1].r;
+		}
+		
+		OutSamples.Array[0].MissCount = OutSamples.Array[0].SampleCount * (1 - RawSample[0].a);
+		OutSamples.Array[0].SceneColor = OutSamples.Array[0].SampleCount * RawSample[0];
+	}
+	else if (SignalBufferLayout == SIGNAL_BUFFER_LAYOUT_DIFFUSE_INDIRECT_AND_AO_HISTORY)
+	{
+		if (bNormalizeSample)
+		{
+			OutSamples.Array[0].SampleCount = RawSample[1].r > 0 ? 1 : 0;
+		}
+		else
+		{
+			OutSamples.Array[0].SampleCount = RawSample[1].r;
+		}
+		
+		OutSamples.Array[0].MissCount = OutSamples.Array[0].SampleCount * (1 - RawSample[0].a);
+		OutSamples.Array[0].SceneColor = OutSamples.Array[0].SampleCount * RawSample[0];
+		
+		OutSamples.Array[1].SampleCount = OutSamples.Array[0].SampleCount;
+		OutSamples.Array[1].SceneColor.x = OutSamples.Array[1].SampleCount * RawSample[1].g;
+		
+		OutSamples.Array[2].SampleCount = OutSamples.Array[0].SampleCount;
+		OutSamples.Array[2].SceneColor.x = OutSamples.Array[2].SampleCount * RawSample[1].b;
+	}
+	#endif
+	#if CONFIG_SIGNAL_COLOR_SH
+	else if (SignalBufferLayout == SIGNAL_BUFFER_LAYOUT_DIFFUSE_INDIRECT_HARMONIC)
+	{
+		float SampleCount;
+		if (bNormalizeSample)
+		{
+			SampleCount = 1.0;
+			OutSamples.Array[0].SampleCount = RawSample[3].x > 0.0 ? 1.0 : 0.0;
+		}
+		else
+		{
+			SampleCount = RawSample[3].x;
+			OutSamples.Array[0].SampleCount = SampleCount;
+		}
+
+		OutSamples.Array[0].MissCount = RawSample[3].y * SampleCount;
+		OutSamples.Array[0].ColorSH.R.V = RawSample[0] * SampleCount;
+		OutSamples.Array[0].ColorSH.G.V = RawSample[1] * SampleCount;
+		OutSamples.Array[0].ColorSH.B.V = RawSample[2] * SampleCount;
+
+		// TODO: SRV on RG F32 and manually decode with f16tof32(), going to reduce VGPR pressure between fetch and result.
+	}
+	#endif
+} // DecodeMultiplexedSignals()
+
+/** Encode output signal sample. */
+void EncodeMultiplexedSignals(
+	const uint SignalBufferLayout, const uint MultiplexCount,
+	FSSDSignalSample Sample[MAX_SIGNAL_MULTIPLEXING],
+	out float4 OutRawSample[MAX_MULTIPLEXED_TEXTURES],
+	out uint OutBufferCount)
+{
+	// Init all raw samples.
+	UNROLL_N(MAX_MULTIPLEXED_TEXTURES)
+	for (uint i = 0; i < MAX_MULTIPLEXED_TEXTURES; i++)
+		OutRawSample[i] = 0;
+
+	// Number of buffer the signal get encoded onto <= MAX_MULTIPLEXED_TEXTURES.
+	OutBufferCount = 1;
+
+	if (SignalBufferLayout == SIGNAL_BUFFER_LAYOUT_PENUMBRA_INJESTION_1SPP)
+	{
+		UNROLL
+		for (uint MultiplexId = 0; MultiplexId < MultiplexCount; MultiplexId++)
+		{
+			float EncodedChannel = Sample[MultiplexId].WorldBluringRadius;
+			if (Sample[MultiplexId].SampleCount == 0)
+			{
+				EncodedChannel = -2;
+			}
+			else if (Sample[MultiplexId].WorldBluringRadius == WORLD_RADIUS_MISS)
+			{
+				EncodedChannel = -1;
+			}
+
+			OutRawSample[MultiplexId / 4][MultiplexId % 4] = EncodedChannel;
+		}
+		OutBufferCount = (MultiplexCount + 3) / 4;
+	}
+	else if (SignalBufferLayout == SIGNAL_BUFFER_LAYOUT_PENUMBRA_INJESTION_NSPP)
+	{
+		UNROLL
+		for (uint MultiplexId = 0; MultiplexId < MultiplexCount; MultiplexId++)
+		{
+			float2 EncodedChannels = float2(
+				Sample[MultiplexId].MissCount,
+				Sample[MultiplexId].WorldBluringRadius) * SafeRcp(Sample[MultiplexId].SampleCount);
+			
+			if (Sample[MultiplexId].SampleCount == 0)
+			{
+				EncodedChannels.y = -2;
+			}
+			else if (Sample[MultiplexId].WorldBluringRadius == WORLD_RADIUS_MISS)
+			{
+				EncodedChannels.y = -1;
+			}
+
+			if (MultiplexId % 2)
+				OutRawSample[MultiplexId / 2].zw = EncodedChannels;
+			else
+				OutRawSample[MultiplexId / 2].xy = EncodedChannels;
+		}
+		OutBufferCount = (MultiplexCount + 1) / 2;
+	}
+	//else if (SignalBufferLayout == SIGNAL_BUFFER_LAYOUT_PENUMBRA_RECONSTRUCTION)
+	//{
+	//	UNROLL
+	//	for (uint MultiplexId = 0; MultiplexId < MultiplexCount; MultiplexId++)
+	//	{
+	//		float2 EncodedChannels = float2(
+	//			Sample[MultiplexId].MissCount * SafeRcp(Sample[MultiplexId].SampleCount),
+	//			Sample[MultiplexId].SampleCount);
+	//
+	//		if (MultiplexId % 2)
+	//			OutRawSample[MultiplexId / 2].zw = EncodedChannels;
+	//		else
+	//			OutRawSample[MultiplexId / 2].xy = EncodedChannels;
+	//	}
+	//	OutBufferCount = (MultiplexCount + 1) / 2;
+	//}
+	else if (SignalBufferLayout == SIGNAL_BUFFER_LAYOUT_PENUMBRA_HISTORY || SignalBufferLayout == SIGNAL_BUFFER_LAYOUT_PENUMBRA_RECONSTRUCTION)
+	{
+		UNROLL
+		for (uint MultiplexId = 0; MultiplexId < MultiplexCount; MultiplexId++)
+		{
+			float NormalizationFactor = SafeRcp(Sample[MultiplexId].SampleCount);
+
+			float NormalizedWorldBluringRadius = Sample[MultiplexId].WorldBluringRadius * NormalizationFactor;
+
+			OutRawSample[MultiplexId] = float4(
+				Sample[MultiplexId].MissCount * NormalizationFactor,
+				Sample[MultiplexId].SampleCount,
+				NormalizedWorldBluringRadius == WORLD_RADIUS_MISS ? -1 : NormalizedWorldBluringRadius,
+				0);
+		}
+		OutBufferCount = MultiplexCount;
+	}
+	else if (SignalBufferLayout == SIGNAL_BUFFER_LAYOUT_PENUMBRA_REJECTION)
+	{
+		const uint BatchSignalCount = MultiplexCount / 2;
+		OutBufferCount = BatchSignalCount;
+
+		UNROLL
+		for (uint BatchSignalId = 0; BatchSignalId < BatchSignalCount; BatchSignalId++)
+		{
+			uint MultiplexId = BatchSignalId * 2;
+
+			float NormalizationFactor = SafeRcp(Sample[MultiplexId + 0].SampleCount);
+
+			float NormalizedWorldBluringRadius = Sample[MultiplexId + 0].WorldBluringRadius * NormalizationFactor;
+
+			// Samples are normalised when doing history preconvolution, so the number of is either 0 or 1 on both momments. Therefore
+			// Sample[MultiplexId + 0].SampleCount == Sample[MultiplexId + 1].SampleCount.
+			OutRawSample[BatchSignalId] = float4(
+				Sample[MultiplexId + 0].MissCount * NormalizationFactor,
+				Sample[MultiplexId + 1].MissCount * NormalizationFactor,
+				Sample[MultiplexId + 0].SampleCount,
+				NormalizedWorldBluringRadius == WORLD_RADIUS_MISS ? -1 : NormalizedWorldBluringRadius);
+		}
+	}
+	#if CONFIG_SIGNAL_COLOR
+	else if (SignalBufferLayout == SIGNAL_BUFFER_LAYOUT_REFLECTIONS_REJECTION)
+	{
+		// Outputs scene color to be compatible with a SSR output.
+		OutRawSample[0] = Sample[0].SceneColor * (Sample[0].SampleCount > 0 ? rcp(Sample[0].SampleCount) : 0);
+		OutRawSample[1].r = Sample[0].SampleCount;
+		OutBufferCount = 2;
+
+		if (MultiplexCount == 2)
+		{
+			OutRawSample[2] = Sample[1].SceneColor * (Sample[1].SampleCount > 0 ? rcp(Sample[1].SampleCount) : 0);
+			OutRawSample[1].g = Sample[1].SampleCount;
+			OutBufferCount = 3;
+		}
+	}
+	else if (SignalBufferLayout == SIGNAL_BUFFER_LAYOUT_REFLECTIONS_HISTORY)
+	{
+		float NormalizationFactor = SafeRcp(Sample[0].SampleCount);
+
+		// Outputs scene color to be compatible with a SSR output.
+		OutRawSample[0] = Sample[0].SceneColor * NormalizationFactor;
+		OutRawSample[1].r = Sample[0].SampleCount;
+		OutBufferCount = 2;
+
+		if (MultiplexCount == 3)
+		{
+			OutRawSample[1].g = Sample[1].SceneColor.x * NormalizationFactor;
+			OutRawSample[1].b = Sample[2].SceneColor.x * NormalizationFactor;
+			OutBufferCount = 3;
+		}
+	}
+	#endif
+	else if (SignalBufferLayout == SIGNAL_BUFFER_LAYOUT_AO_REJECTION)
+	{
+		// Outputs number of ray miss in red to be compatible as a SSAO output.
+		OutRawSample[0].r = Sample[0].SampleCount > 0 ? Sample[0].MissCount / Sample[0].SampleCount : 1.0;
+		OutRawSample[0].g = Sample[0].SampleCount;
+		
+		if (MultiplexCount == 2)
+		{
+			OutRawSample[0].b = Sample[1].SampleCount > 0 ? Sample[1].MissCount / Sample[1].SampleCount : 1.0;
+			OutRawSample[0].a = Sample[1].SampleCount;
+		}
+	}
+	else if (SignalBufferLayout == SIGNAL_BUFFER_LAYOUT_AO_HISTORY)
+	{
+		// Outputs number of ray miss in red to be compatible as a SSAO output.
+		OutRawSample[0].r = Sample[0].SampleCount > 0 ? Sample[0].MissCount / Sample[0].SampleCount : 1.0;
+		OutRawSample[0].g = Sample[0].SampleCount;
+	}
+	#if CONFIG_SIGNAL_COLOR
+	else if (SignalBufferLayout == SIGNAL_BUFFER_LAYOUT_DIFFUSE_INDIRECT_AND_AO_RECONSTRUCTION)
+	{
+		float NormalizationFactor = SafeRcp(Sample[0].SampleCount);
+
+		OutRawSample[0] = Sample[0].SceneColor * NormalizationFactor;
+		OutRawSample[1].r = Sample[0].SampleCount;
+		OutBufferCount = 2;
+	}
+	else if (SignalBufferLayout == SIGNAL_BUFFER_LAYOUT_DIFFUSE_INDIRECT_AND_AO_HISTORY)
+	{
+		float NormalizationFactor = SafeRcp(Sample[0].SampleCount);
+
+		OutRawSample[0] = Sample[0].SceneColor * NormalizationFactor;
+		OutRawSample[1].r = Sample[0].SampleCount;
+		OutBufferCount = 2;
+
+		if (MultiplexCount == 3)
+		{
+			OutRawSample[1].g = Sample[1].SceneColor.x * NormalizationFactor;
+			OutRawSample[1].b = Sample[2].SceneColor.x * NormalizationFactor;
+		}
+	}
+	#endif
+	#if CONFIG_SIGNAL_COLOR_SH
+	else if (SignalBufferLayout == SIGNAL_BUFFER_LAYOUT_DIFFUSE_INDIRECT_HARMONIC)
+	{
+		float NormalizationFactor = SafeRcp(Sample[0].SampleCount);
+
+		OutRawSample[3].x = Sample[0].SampleCount;
+		OutRawSample[3].y = Sample[0].MissCount * NormalizationFactor;
+		OutRawSample[0] = Sample[0].ColorSH.R.V * NormalizationFactor;
+		OutRawSample[1] = Sample[0].ColorSH.G.V * NormalizationFactor;
+		OutRawSample[2] = Sample[0].ColorSH.B.V * NormalizationFactor;
+		OutBufferCount = 4;
+	}
+	#endif
+} // EncodeMultiplexedSignals()
+
+/** Sample multiple input signals that have been multiplexed. */
+FSSDSignalArray SampleMultiplexedSignals(
+	Texture2D SignalBuffer0, Texture2D SignalBuffer1, Texture2D SignalBuffer2, Texture2D SignalBuffer3,
+	SamplerState Sampler,
+	const uint SignalBufferLayout, const uint MultiplexedSampleId,
+	const bool bNormalizeSample,
+	float2 UV, float Level = 0)
+{
+	float4 RawSample[MAX_MULTIPLEXED_TEXTURES];
+	RawSample[0] = SignalBuffer0.SampleLevel(Sampler, UV, Level);
+	RawSample[1] = SignalBuffer1.SampleLevel(Sampler, UV, Level);
+	RawSample[2] = SignalBuffer2.SampleLevel(Sampler, UV, Level);
+	RawSample[3] = SignalBuffer3.SampleLevel(Sampler, UV, Level);
+
+	FSSDSignalArray MultiplexedSamples;
+	DecodeMultiplexedSignals(
+		SignalBufferLayout,
+		MultiplexedSampleId,
+		bNormalizeSample,
+		RawSample,
+		/* out */ MultiplexedSamples);
+
+	return MultiplexedSamples;
+}
+
+
+/** Outputs multiplexed signal. */
+void OutputMultiplexedSignal(
+	RWTexture2D<float4> OutputSignalBuffer0,
+	RWTexture2D<float4> OutputSignalBuffer1,
+	RWTexture2D<float4> OutputSignalBuffer2,
+	RWTexture2D<float4> OutputSignalBuffer3,
+	const uint SignalBufferLayout, const uint MultiplexCount,
+	const uint2 PixelPosition, FSSDSignalSample MultiplexedSamples[MAX_SIGNAL_MULTIPLEXING])
+{
+	// Encode the output signal.
+	float4 RawSample[MAX_MULTIPLEXED_TEXTURES];
+	uint BufferCount;
+	EncodeMultiplexedSignals(
+		SignalBufferLayout, MultiplexCount,
+		MultiplexedSamples,
+		/* out */ RawSample, /* out */ BufferCount);
+
+	// Output the raw encoded sample according to number of RT they requires.
+	if (BufferCount >= 1)
+		OutputSignalBuffer0[PixelPosition] = min(RawSample[0], MaxHalfFloat);
+	if (BufferCount >= 2)
+		OutputSignalBuffer1[PixelPosition] = min(RawSample[1], MaxHalfFloat);
+	if (BufferCount >= 3)
+		OutputSignalBuffer2[PixelPosition] = min(RawSample[2], MaxHalfFloat);
+	if (BufferCount >= 4)
+		OutputSignalBuffer3[PixelPosition] = min(RawSample[3], MaxHalfFloat);
+}
diff --git a/Engine/Shaders/Private/ScreenSpaceDenoise/SSDSignalCore.ush b/Engine/Shaders/Private/ScreenSpaceDenoise/SSDSignalCore.ush
--- a/Engine/Shaders/Private/ScreenSpaceDenoise/SSDSignalCore.ush
+++ b/Engine/Shaders/Private/ScreenSpaceDenoise/SSDSignalCore.ush
@@ -346,537 +346,3 @@ FSSDSignalSample NormalizeToOneSample(FSSDSignalSample Sample)
 	OutSample.SampleCount = Sample.SampleCount > 0 ? 1 : 0;
 	return OutSample;
 }
-
-
-//------------------------------------------------------- INPUT/OUTPUT
-
-/** Decode input signal sample. */
-void DecodeMultiplexedSignals(
-	const uint SignalBufferLayout,
-	const uint MultiplexedSampleId,
-	const bool bNormalizeSample,
-	float4 RawSample[MAX_MULTIPLEXED_TEXTURES],
-	out FSSDSignalArray OutSamples)
-{
-	OutSamples = CreateSignalArrayFromScalarValue(0.0);
-
-	if (SignalBufferLayout == SIGNAL_BUFFER_LAYOUT_UNINITIALIZED)
-	{
-		// trash out the output sample to identify quickly a miss configured buffer layout.
-		OutSamples = CreateSignalArrayFromScalarValue(-INFINITE_FLOAT);
-	}
-	else if (SignalBufferLayout == SIGNAL_BUFFER_LAYOUT_PENUMBRA_INPUT_1SPP)
-	{
-		UNROLL_N(MAX_SIGNAL_BATCH_SIZE)
-		for (uint BatchSignalId = 0; BatchSignalId < MAX_SIGNAL_BATCH_SIZE; BatchSignalId++)
-		{
-			uint MultiplexId = BatchSignalId;
-			OutSamples.Array[MultiplexId].SampleCount = (RawSample[MultiplexId].r == DENOISER_INVALID_HIT_DISTANCE ? 0.0 : 1.0);
-			OutSamples.Array[MultiplexId].MissCount = OutSamples.Array[MultiplexId].SampleCount * (RawSample[MultiplexId].r > 0.0 ? 0.0 : 1.0);
-			OutSamples.Array[MultiplexId].ClosestHitDistance = OutSamples.Array[MultiplexId].SampleCount * (RawSample[MultiplexId].r == DENOISER_MISS_HIT_DISTANCE ? 1000000 : RawSample[MultiplexId].r);
-		}
-	}
-	else if (SignalBufferLayout == SIGNAL_BUFFER_LAYOUT_PENUMBRA_INPUT_NSPP)
-	{
-		UNROLL_N(MAX_SIGNAL_BATCH_SIZE)
-		for (uint BatchSignalId = 0; BatchSignalId < MAX_SIGNAL_BATCH_SIZE; BatchSignalId++)
-		{
-			uint MultiplexId = BatchSignalId;
-			// TODO: feed the actual number of sample.
-			OutSamples.Array[MultiplexId].SampleCount = (RawSample[MultiplexId].g == DENOISER_INVALID_HIT_DISTANCE ? 0.0 : 1.0);
-			OutSamples.Array[MultiplexId].MissCount = OutSamples.Array[MultiplexId].SampleCount * RawSample[MultiplexId].r;
-			OutSamples.Array[MultiplexId].ClosestHitDistance = OutSamples.Array[MultiplexId].SampleCount * (RawSample[MultiplexId].g == DENOISER_MISS_HIT_DISTANCE ? 1000000 : RawSample[MultiplexId].g);
-		}
-	}
-	else if (SignalBufferLayout == SIGNAL_BUFFER_LAYOUT_PENUMBRA_INJESTION_1SPP)
-	{
-		UNROLL_N(MAX_SIGNAL_BATCH_SIZE)
-		for (uint BatchSignalId = 0; BatchSignalId < MAX_SIGNAL_BATCH_SIZE; BatchSignalId++)
-		{
-			uint MultiplexId = BatchSignalId;
-			float EncodedWorldBluringRadius = RawSample[MultiplexId / 4][MultiplexId % 4];
-			OutSamples.Array[MultiplexId].SampleCount = (EncodedWorldBluringRadius == -2.0 ? 0.0 : 1.0);
-			OutSamples.Array[MultiplexId].MissCount = OutSamples.Array[MultiplexId].SampleCount * (EncodedWorldBluringRadius > 0.0 ? 0.0 : 1.0);
-			OutSamples.Array[MultiplexId].WorldBluringRadius = OutSamples.Array[MultiplexId].SampleCount * (EncodedWorldBluringRadius == -1 ? WORLD_RADIUS_MISS : EncodedWorldBluringRadius);
-		}
-	}
-	else if (SignalBufferLayout == SIGNAL_BUFFER_LAYOUT_PENUMBRA_INJESTION_NSPP)
-	{
-		UNROLL_N(MAX_SIGNAL_BATCH_SIZE)
-		for (uint BatchSignalId = 0; BatchSignalId < MAX_SIGNAL_BATCH_SIZE; BatchSignalId++)
-		{
-			uint MultiplexId = BatchSignalId;
-			float2 Channels = float2(RawSample[MultiplexId / 2][2 * (MultiplexId % 2)], RawSample[MultiplexId / 2][2 * (MultiplexId % 2) + 1]);
-	
-			// TODO: feed the actual number of sample.
-			OutSamples.Array[MultiplexId].SampleCount = (Channels.g == -2.0 ? 0.0 : 1.0);
-			OutSamples.Array[MultiplexId].MissCount = OutSamples.Array[MultiplexId].SampleCount * Channels.r;
-			OutSamples.Array[MultiplexId].WorldBluringRadius = OutSamples.Array[MultiplexId].SampleCount * (Channels.g == -1 ? WORLD_RADIUS_MISS : Channels.g);
-		}
-	}
-	//else if (SignalBufferLayout == SIGNAL_BUFFER_LAYOUT_PENUMBRA_RECONSTRUCTION)
-	//{
-	//	UNROLL_N(MAX_SIGNAL_MULTIPLEXING)
-	//	for (uint MultiplexId = 0; MultiplexId < MAX_SIGNAL_MULTIPLEXING; MultiplexId++)
-	//	{
-	//		float2 Channels = float2(RawSample[MultiplexId / 2][2 * (MultiplexId % 2)], RawSample[MultiplexId / 2][2 * (MultiplexId % 2) + 1]);
-	//
-	//		if (bNormalizeSample)
-	//		{
-	//			OutSamples.Array[MultiplexId].SampleCount = Channels.g > 0 ? 1 : 0;
-	//			OutSamples.Array[MultiplexId].MissCount = Channels.g > 0 ? Channels.r : 0;
-	//		}
-	//		else
-	//		{
-	//			OutSamples.Array[MultiplexId].SampleCount = Channels.g;
-	//			OutSamples.Array[MultiplexId].MissCount = Channels.r * Channels.g;
-	//		}
-	//	}
-	//}
-	else if (SignalBufferLayout == SIGNAL_BUFFER_LAYOUT_PENUMBRA_HISTORY || SignalBufferLayout == SIGNAL_BUFFER_LAYOUT_PENUMBRA_RECONSTRUCTION)
-	{
-		UNROLL_N(MAX_SIGNAL_BATCH_SIZE)
-		for (uint MultiplexId = 0; MultiplexId < MAX_SIGNAL_BATCH_SIZE; MultiplexId++)
-		{
-			float3 Channels = RawSample[MultiplexId].xyz;
-			
-			float SampleCount = bNormalizeSample ? (Channels.g > 0 ? 1 : 0) : (Channels.g);
-
-			OutSamples.Array[MultiplexId].MissCount = Channels.r * SampleCount;
-			OutSamples.Array[MultiplexId].SampleCount = SampleCount;
-			OutSamples.Array[MultiplexId].WorldBluringRadius = (Channels.b == DENOISER_MISS_HIT_DISTANCE ? WORLD_RADIUS_MISS : Channels.b) * SampleCount;
-		}
-	}
-	else if (SignalBufferLayout == SIGNAL_BUFFER_LAYOUT_PENUMBRA_REJECTION)
-	{
-		UNROLL_N(MAX_SIGNAL_BATCH_SIZE)
-		for (uint BatchSignalId = 0; BatchSignalId < MAX_SIGNAL_BATCH_SIZE; BatchSignalId++)
-		{
-			uint MultiplexId = BatchSignalId * 2;
-			float4 Channels = RawSample[BatchSignalId];
-			float SampleCount = Channels.b;
-
-			OutSamples.Array[MultiplexId + 0].SampleCount = SampleCount;
-			OutSamples.Array[MultiplexId + 0].MissCount = Channels.r * SampleCount;
-			OutSamples.Array[MultiplexId + 0].WorldBluringRadius = (Channels.a == -1 ? WORLD_RADIUS_MISS : Channels.a) * SampleCount;
-			
-			OutSamples.Array[MultiplexId + 1].SampleCount = SampleCount;
-			OutSamples.Array[MultiplexId + 1].MissCount = Channels.g * SampleCount;
-		}
-	}
-	#if CONFIG_SIGNAL_COLOR
-	else if (SignalBufferLayout == SIGNAL_BUFFER_LAYOUT_REFLECTIONS_INPUT)
-	{
-		OutSamples.Array[0].SampleCount = RawSample[1].r == DENOISER_INVALID_HIT_DISTANCE ? 0 : 1;
-		OutSamples.Array[0].MissCount = OutSamples.Array[0].SampleCount * (1 - RawSample[0].a);
-		OutSamples.Array[0].SceneColor = RawSample[0];
-
-		// max() because ray hit distance is < 0 when miss / not ray has been shot.
-		//OutSamples.Array[0].ClosestHitDistance = RawSample[1].r * OutSamples.Array[0].SampleCount;
-		OutSamples.Array[0].ClosestHitDistance = ((RawSample[1].r == DENOISER_MISS_HIT_DISTANCE) ? 1.0e20 : RawSample[1].r) * OutSamples.Array[0].SampleCount;
-	}
-	else if (SignalBufferLayout == SIGNAL_BUFFER_LAYOUT_REFLECTIONS_REJECTION)
-	{
-		if (bNormalizeSample)
-		{
-			OutSamples.Array[0].SampleCount = RawSample[1].r > 0 ? 1 : 0;
-		}
-		else
-		{
-			OutSamples.Array[0].SampleCount = RawSample[1].r;
-		}
-		
-		OutSamples.Array[0].MissCount = OutSamples.Array[0].SampleCount * (1 - RawSample[0].a);
-		OutSamples.Array[0].SceneColor = OutSamples.Array[0].SampleCount * RawSample[0];
-
-		OutSamples.Array[1].SampleCount = RawSample[1].g;
-		OutSamples.Array[1].MissCount = OutSamples.Array[0].SampleCount * (1 - RawSample[2].a);
-		OutSamples.Array[1].SceneColor = OutSamples.Array[0].SampleCount * RawSample[2];
-	}
-	else if (SignalBufferLayout == SIGNAL_BUFFER_LAYOUT_REFLECTIONS_HISTORY)
-	{
-		if (bNormalizeSample)
-		{
-			OutSamples.Array[0].SampleCount = RawSample[1].r > 0 ? 1 : 0;
-		}
-		else
-		{
-			OutSamples.Array[0].SampleCount = RawSample[1].r;
-		}
-		
-		OutSamples.Array[0].MissCount = OutSamples.Array[0].SampleCount * (1 - RawSample[0].a);
-		OutSamples.Array[0].SceneColor = OutSamples.Array[0].SampleCount * RawSample[0];
-
-		OutSamples.Array[1].SampleCount = OutSamples.Array[0].SampleCount;
-		OutSamples.Array[1].SceneColor.x = OutSamples.Array[1].SampleCount * RawSample[1].g;
-		
-		OutSamples.Array[2].SampleCount = OutSamples.Array[0].SampleCount;
-		OutSamples.Array[2].SceneColor.x = OutSamples.Array[2].SampleCount * RawSample[1].b;
-	}
-	#endif
-	else if (SignalBufferLayout == SIGNAL_BUFFER_LAYOUT_AO_INPUT)
-	{
-		OutSamples.Array[0].SampleCount = RawSample[1].r < 0 ? 0 : 1;
-		OutSamples.Array[0].MissCount = OutSamples.Array[0].SampleCount * RawSample[0].r;
-		OutSamples.Array[0].ClosestHitDistance = OutSamples.Array[0].SampleCount * RawSample[1].r;
-	}
-	else if (SignalBufferLayout == SIGNAL_BUFFER_LAYOUT_AO_REJECTION)
-	{
-		if (bNormalizeSample)
-		{
-			OutSamples.Array[0].SampleCount = RawSample[0].g > 0 ? 1 : 0;
-			OutSamples.Array[0].MissCount = RawSample[0].g > 0 ? RawSample[0].r : 0;
-			
-			OutSamples.Array[1].SampleCount = RawSample[0].a > 0 ? 1 : 0;
-			OutSamples.Array[1].MissCount = RawSample[0].a > 0 ? RawSample[0].b : 0;
-		}
-		else
-		{
-			OutSamples.Array[0].SampleCount = RawSample[0].g;
-			OutSamples.Array[0].MissCount = RawSample[0].g * RawSample[0].r;
-			
-			OutSamples.Array[1].SampleCount = RawSample[0].a;
-			OutSamples.Array[1].MissCount = RawSample[0].a * RawSample[0].b;
-		}
-	}
-	else if (SignalBufferLayout == SIGNAL_BUFFER_LAYOUT_AO_HISTORY)
-	{
-		if (bNormalizeSample)
-		{
-			OutSamples.Array[0].SampleCount = RawSample[0].g > 0 ? 1 : 0;
-			OutSamples.Array[0].MissCount = RawSample[0].g > 0 ? RawSample[0].r : 0;
-		}
-		else
-		{
-			OutSamples.Array[0].SampleCount = RawSample[0].g;
-			OutSamples.Array[0].MissCount = RawSample[0].g * RawSample[0].r;
-		}
-	}
-	#if CONFIG_SIGNAL_COLOR
-	else if (SignalBufferLayout == SIGNAL_BUFFER_LAYOUT_DIFFUSE_INDIRECT_AND_AO_INPUT_NSPP)
-	{
-		OutSamples.Array[0].SampleCount = RawSample[1].g;
-		OutSamples.Array[0].MissCount = OutSamples.Array[0].SampleCount * (1 - RawSample[0].a);
-		OutSamples.Array[0].SceneColor = OutSamples.Array[0].SampleCount * RawSample[0];
-
-		// max() because ray hit distance is < 0 when miss / not ray has been shot.
-		OutSamples.Array[0].ClosestHitDistance = RawSample[1].r * OutSamples.Array[0].SampleCount;
-		OutSamples.Array[0].ClosestHitDistance = ((RawSample[1].r == -1.0) ? 1.0e20 : RawSample[1].r) * OutSamples.Array[0].SampleCount;
-	}
-	else if (SignalBufferLayout == SIGNAL_BUFFER_LAYOUT_DIFFUSE_INDIRECT_AND_AO_RECONSTRUCTION)
-	{
-		if (bNormalizeSample)
-		{
-			OutSamples.Array[0].SampleCount = RawSample[1].r > 0 ? 1 : 0;
-		}
-		else
-		{
-			OutSamples.Array[0].SampleCount = RawSample[1].r;
-		}
-		
-		OutSamples.Array[0].MissCount = OutSamples.Array[0].SampleCount * (1 - RawSample[0].a);
-		OutSamples.Array[0].SceneColor = OutSamples.Array[0].SampleCount * RawSample[0];
-	}
-	else if (SignalBufferLayout == SIGNAL_BUFFER_LAYOUT_DIFFUSE_INDIRECT_AND_AO_HISTORY)
-	{
-		if (bNormalizeSample)
-		{
-			OutSamples.Array[0].SampleCount = RawSample[1].r > 0 ? 1 : 0;
-		}
-		else
-		{
-			OutSamples.Array[0].SampleCount = RawSample[1].r;
-		}
-		
-		OutSamples.Array[0].MissCount = OutSamples.Array[0].SampleCount * (1 - RawSample[0].a);
-		OutSamples.Array[0].SceneColor = OutSamples.Array[0].SampleCount * RawSample[0];
-		
-		OutSamples.Array[1].SampleCount = OutSamples.Array[0].SampleCount;
-		OutSamples.Array[1].SceneColor.x = OutSamples.Array[1].SampleCount * RawSample[1].g;
-		
-		OutSamples.Array[2].SampleCount = OutSamples.Array[0].SampleCount;
-		OutSamples.Array[2].SceneColor.x = OutSamples.Array[2].SampleCount * RawSample[1].b;
-	}
-	#endif
-	#if CONFIG_SIGNAL_COLOR_SH
-	else if (SignalBufferLayout == SIGNAL_BUFFER_LAYOUT_DIFFUSE_INDIRECT_HARMONIC)
-	{
-		float SampleCount;
-		if (bNormalizeSample)
-		{
-			SampleCount = 1.0;
-			OutSamples.Array[0].SampleCount = RawSample[3].x > 0.0 ? 1.0 : 0.0;
-		}
-		else
-		{
-			SampleCount = RawSample[3].x;
-			OutSamples.Array[0].SampleCount = SampleCount;
-		}
-
-		OutSamples.Array[0].MissCount = RawSample[3].y * SampleCount;
-		OutSamples.Array[0].ColorSH.R.V = RawSample[0] * SampleCount;
-		OutSamples.Array[0].ColorSH.G.V = RawSample[1] * SampleCount;
-		OutSamples.Array[0].ColorSH.B.V = RawSample[2] * SampleCount;
-
-		// TODO: SRV on RG F32 and manually decode with f16tof32(), going to reduce VGPR pressure between fetch and result.
-	}
-	#endif
-} // DecodeMultiplexedSignals()
-
-/** Encode output signal sample. */
-void EncodeMultiplexedSignals(
-	const uint SignalBufferLayout, const uint MultiplexCount,
-	FSSDSignalSample Sample[MAX_SIGNAL_MULTIPLEXING],
-	out float4 OutRawSample[MAX_MULTIPLEXED_TEXTURES],
-	out uint OutBufferCount)
-{
-	// Init all raw samples.
-	UNROLL_N(MAX_MULTIPLEXED_TEXTURES)
-	for (uint i = 0; i < MAX_MULTIPLEXED_TEXTURES; i++)
-		OutRawSample[i] = 0;
-
-	// Number of buffer the signal get encoded onto <= MAX_MULTIPLEXED_TEXTURES.
-	OutBufferCount = 1;
-
-	if (SignalBufferLayout == SIGNAL_BUFFER_LAYOUT_PENUMBRA_INJESTION_1SPP)
-	{
-		UNROLL
-		for (uint MultiplexId = 0; MultiplexId < MultiplexCount; MultiplexId++)
-		{
-			float EncodedChannel = Sample[MultiplexId].WorldBluringRadius;
-			if (Sample[MultiplexId].SampleCount == 0)
-			{
-				EncodedChannel = -2;
-			}
-			else if (Sample[MultiplexId].WorldBluringRadius == WORLD_RADIUS_MISS)
-			{
-				EncodedChannel = -1;
-			}
-
-			OutRawSample[MultiplexId / 4][MultiplexId % 4] = EncodedChannel;
-		}
-		OutBufferCount = (MultiplexCount + 3) / 4;
-	}
-	else if (SignalBufferLayout == SIGNAL_BUFFER_LAYOUT_PENUMBRA_INJESTION_NSPP)
-	{
-		UNROLL
-		for (uint MultiplexId = 0; MultiplexId < MultiplexCount; MultiplexId++)
-		{
-			float2 EncodedChannels = float2(
-				Sample[MultiplexId].MissCount,
-				Sample[MultiplexId].WorldBluringRadius) * SafeRcp(Sample[MultiplexId].SampleCount);
-			
-			if (Sample[MultiplexId].SampleCount == 0)
-			{
-				EncodedChannels.y = -2;
-			}
-			else if (Sample[MultiplexId].WorldBluringRadius == WORLD_RADIUS_MISS)
-			{
-				EncodedChannels.y = -1;
-			}
-
-			if (MultiplexId % 2)
-				OutRawSample[MultiplexId / 2].zw = EncodedChannels;
-			else
-				OutRawSample[MultiplexId / 2].xy = EncodedChannels;
-		}
-		OutBufferCount = (MultiplexCount + 1) / 2;
-	}
-	//else if (SignalBufferLayout == SIGNAL_BUFFER_LAYOUT_PENUMBRA_RECONSTRUCTION)
-	//{
-	//	UNROLL
-	//	for (uint MultiplexId = 0; MultiplexId < MultiplexCount; MultiplexId++)
-	//	{
-	//		float2 EncodedChannels = float2(
-	//			Sample[MultiplexId].MissCount * SafeRcp(Sample[MultiplexId].SampleCount),
-	//			Sample[MultiplexId].SampleCount);
-	//
-	//		if (MultiplexId % 2)
-	//			OutRawSample[MultiplexId / 2].zw = EncodedChannels;
-	//		else
-	//			OutRawSample[MultiplexId / 2].xy = EncodedChannels;
-	//	}
-	//	OutBufferCount = (MultiplexCount + 1) / 2;
-	//}
-	else if (SignalBufferLayout == SIGNAL_BUFFER_LAYOUT_PENUMBRA_HISTORY || SignalBufferLayout == SIGNAL_BUFFER_LAYOUT_PENUMBRA_RECONSTRUCTION)
-	{
-		UNROLL
-		for (uint MultiplexId = 0; MultiplexId < MultiplexCount; MultiplexId++)
-		{
-			float NormalizationFactor = SafeRcp(Sample[MultiplexId].SampleCount);
-
-			float NormalizedWorldBluringRadius = Sample[MultiplexId].WorldBluringRadius * NormalizationFactor;
-
-			OutRawSample[MultiplexId] = float4(
-				Sample[MultiplexId].MissCount * NormalizationFactor,
-				Sample[MultiplexId].SampleCount,
-				NormalizedWorldBluringRadius == WORLD_RADIUS_MISS ? -1 : NormalizedWorldBluringRadius,
-				0);
-		}
-		OutBufferCount = MultiplexCount;
-	}
-	else if (SignalBufferLayout == SIGNAL_BUFFER_LAYOUT_PENUMBRA_REJECTION)
-	{
-		const uint BatchSignalCount = MultiplexCount / 2;
-		OutBufferCount = BatchSignalCount;
-
-		UNROLL
-		for (uint BatchSignalId = 0; BatchSignalId < BatchSignalCount; BatchSignalId++)
-		{
-			uint MultiplexId = BatchSignalId * 2;
-
-			float NormalizationFactor = SafeRcp(Sample[MultiplexId + 0].SampleCount);
-
-			float NormalizedWorldBluringRadius = Sample[MultiplexId + 0].WorldBluringRadius * NormalizationFactor;
-
-			// Samples are normalised when doing history preconvolution, so the number of is either 0 or 1 on both momments. Therefore
-			// Sample[MultiplexId + 0].SampleCount == Sample[MultiplexId + 1].SampleCount.
-			OutRawSample[BatchSignalId] = float4(
-				Sample[MultiplexId + 0].MissCount * NormalizationFactor,
-				Sample[MultiplexId + 1].MissCount * NormalizationFactor,
-				Sample[MultiplexId + 0].SampleCount,
-				NormalizedWorldBluringRadius == WORLD_RADIUS_MISS ? -1 : NormalizedWorldBluringRadius);
-		}
-	}
-	#if CONFIG_SIGNAL_COLOR
-	else if (SignalBufferLayout == SIGNAL_BUFFER_LAYOUT_REFLECTIONS_REJECTION)
-	{
-		// Outputs scene color to be compatible with a SSR output.
-		OutRawSample[0] = Sample[0].SceneColor * (Sample[0].SampleCount > 0 ? rcp(Sample[0].SampleCount) : 0);
-		OutRawSample[1].r = Sample[0].SampleCount;
-		OutBufferCount = 2;
-
-		if (MultiplexCount == 2)
-		{
-			OutRawSample[2] = Sample[1].SceneColor * (Sample[1].SampleCount > 0 ? rcp(Sample[1].SampleCount) : 0);
-			OutRawSample[1].g = Sample[1].SampleCount;
-			OutBufferCount = 3;
-		}
-	}
-	else if (SignalBufferLayout == SIGNAL_BUFFER_LAYOUT_REFLECTIONS_HISTORY)
-	{
-		float NormalizationFactor = SafeRcp(Sample[0].SampleCount);
-
-		// Outputs scene color to be compatible with a SSR output.
-		OutRawSample[0] = Sample[0].SceneColor * NormalizationFactor;
-		OutRawSample[1].r = Sample[0].SampleCount;
-		OutBufferCount = 2;
-
-		if (MultiplexCount == 3)
-		{
-			OutRawSample[1].g = Sample[1].SceneColor.x * NormalizationFactor;
-			OutRawSample[1].b = Sample[2].SceneColor.x * NormalizationFactor;
-			OutBufferCount = 3;
-		}
-	}
-	#endif
-	else if (SignalBufferLayout == SIGNAL_BUFFER_LAYOUT_AO_REJECTION)
-	{
-		// Outputs number of ray miss in red to be compatible as a SSAO output.
-		OutRawSample[0].r = Sample[0].SampleCount > 0 ? Sample[0].MissCount / Sample[0].SampleCount : 1.0;
-		OutRawSample[0].g = Sample[0].SampleCount;
-		
-		if (MultiplexCount == 2)
-		{
-			OutRawSample[0].b = Sample[1].SampleCount > 0 ? Sample[1].MissCount / Sample[1].SampleCount : 1.0;
-			OutRawSample[0].a = Sample[1].SampleCount;
-		}
-	}
-	else if (SignalBufferLayout == SIGNAL_BUFFER_LAYOUT_AO_HISTORY)
-	{
-		// Outputs number of ray miss in red to be compatible as a SSAO output.
-		OutRawSample[0].r = Sample[0].SampleCount > 0 ? Sample[0].MissCount / Sample[0].SampleCount : 1.0;
-		OutRawSample[0].g = Sample[0].SampleCount;
-	}
-	#if CONFIG_SIGNAL_COLOR
-	else if (SignalBufferLayout == SIGNAL_BUFFER_LAYOUT_DIFFUSE_INDIRECT_AND_AO_RECONSTRUCTION)
-	{
-		float NormalizationFactor = SafeRcp(Sample[0].SampleCount);
-
-		OutRawSample[0] = Sample[0].SceneColor * NormalizationFactor;
-		OutRawSample[1].r = Sample[0].SampleCount;
-		OutBufferCount = 2;
-	}
-	else if (SignalBufferLayout == SIGNAL_BUFFER_LAYOUT_DIFFUSE_INDIRECT_AND_AO_HISTORY)
-	{
-		float NormalizationFactor = SafeRcp(Sample[0].SampleCount);
-
-		OutRawSample[0] = Sample[0].SceneColor * NormalizationFactor;
-		OutRawSample[1].r = Sample[0].SampleCount;
-		OutBufferCount = 2;
-
-		if (MultiplexCount == 3)
-		{
-			OutRawSample[1].g = Sample[1].SceneColor.x * NormalizationFactor;
-			OutRawSample[1].b = Sample[2].SceneColor.x * NormalizationFactor;
-		}
-	}
-	#endif
-	#if CONFIG_SIGNAL_COLOR_SH
-	else if (SignalBufferLayout == SIGNAL_BUFFER_LAYOUT_DIFFUSE_INDIRECT_HARMONIC)
-	{
-		float NormalizationFactor = SafeRcp(Sample[0].SampleCount);
-
-		OutRawSample[3].x = Sample[0].SampleCount;
-		OutRawSample[3].y = Sample[0].MissCount * NormalizationFactor;
-		OutRawSample[0] = Sample[0].ColorSH.R.V * NormalizationFactor;
-		OutRawSample[1] = Sample[0].ColorSH.G.V * NormalizationFactor;
-		OutRawSample[2] = Sample[0].ColorSH.B.V * NormalizationFactor;
-		OutBufferCount = 4;
-	}
-	#endif
-} // EncodeMultiplexedSignals()
-
-/** Sample multiple input signals that have been multiplexed. */
-FSSDSignalArray SampleMultiplexedSignals(
-	Texture2D SignalBuffer0, Texture2D SignalBuffer1, Texture2D SignalBuffer2, Texture2D SignalBuffer3,
-	SamplerState Sampler,
-	const uint SignalBufferLayout, const uint MultiplexedSampleId,
-	const bool bNormalizeSample,
-	float2 UV, float Level = 0)
-{
-	float4 RawSample[MAX_MULTIPLEXED_TEXTURES];
-	RawSample[0] = SignalBuffer0.SampleLevel(Sampler, UV, Level);
-	RawSample[1] = SignalBuffer1.SampleLevel(Sampler, UV, Level);
-	RawSample[2] = SignalBuffer2.SampleLevel(Sampler, UV, Level);
-	RawSample[3] = SignalBuffer3.SampleLevel(Sampler, UV, Level);
-
-	FSSDSignalArray MultiplexedSamples;
-	DecodeMultiplexedSignals(
-		SignalBufferLayout,
-		MultiplexedSampleId,
-		bNormalizeSample,
-		RawSample,
-		/* out */ MultiplexedSamples);
-
-	return MultiplexedSamples;
-}
-
-
-/** Outputs multiplexed signal. */
-void OutputMultiplexedSignal(
-	RWTexture2D<float4> OutputSignalBuffer0,
-	RWTexture2D<float4> OutputSignalBuffer1,
-	RWTexture2D<float4> OutputSignalBuffer2,
-	RWTexture2D<float4> OutputSignalBuffer3,
-	const uint SignalBufferLayout, const uint MultiplexCount,
-	const uint2 PixelPosition, FSSDSignalSample MultiplexedSamples[MAX_SIGNAL_MULTIPLEXING])
-{
-	// Encode the output signal.
-	float4 RawSample[MAX_MULTIPLEXED_TEXTURES];
-	uint BufferCount;
-	EncodeMultiplexedSignals(
-		SignalBufferLayout, MultiplexCount,
-		MultiplexedSamples,
-		/* out */ RawSample, /* out */ BufferCount);
-
-	// Output the raw encoded sample according to number of RT they requires.
-	if (BufferCount >= 1)
-		OutputSignalBuffer0[PixelPosition] = min(RawSample[0], MaxHalfFloat);
-	if (BufferCount >= 2)
-		OutputSignalBuffer1[PixelPosition] = min(RawSample[1], MaxHalfFloat);
-	if (BufferCount >= 3)
-		OutputSignalBuffer2[PixelPosition] = min(RawSample[2], MaxHalfFloat);
-	if (BufferCount >= 4)
-		OutputSignalBuffer3[PixelPosition] = min(RawSample[3], MaxHalfFloat);
-}
diff --git a/Engine/Shaders/Private/ScreenSpaceDenoise/SSDSpatialKernel.ush b/Engine/Shaders/Private/ScreenSpaceDenoise/SSDSpatialKernel.ush
--- a/Engine/Shaders/Private/ScreenSpaceDenoise/SSDSpatialKernel.ush
+++ b/Engine/Shaders/Private/ScreenSpaceDenoise/SSDSpatialKernel.ush
@@ -1,6 +1,7 @@
 // Copyright 1998-2019 Epic Games, Inc. All Rights Reserved.
 
 #include "SSDSignalAccumulator.ush"
+#include "SSDSignalBufferEncoding.ush"
 #include "../TextureSampling.ush"
 
 
