Implements signal VGPR compression to save another 4 VGPR of the spherical harmonic denoiser.

diff --git a/Engine/Shaders/Private/ScreenSpaceDenoise/SSDDefinitions.ush b/Engine/Shaders/Private/ScreenSpaceDenoise/SSDDefinitions.ush
--- a/Engine/Shaders/Private/ScreenSpaceDenoise/SSDDefinitions.ush
+++ b/Engine/Shaders/Private/ScreenSpaceDenoise/SSDDefinitions.ush
@@ -34,6 +34,9 @@
 /** Prefer to store normalized components in internal texture. */
 #define DEBUG_STORE_NORMALIZED_COMPONENTS 1
 
+/** Disable any VGPR compression of the signal. */
+#define DEBUG_DISABLE_SIGNAL_COMPRESSION 0
+
 
 //------------------------------------------------------- SHARED CONSTANTS
 
@@ -185,13 +188,21 @@
 	#define SAMPLE_SET_HEXAWEB 11
 
 
+/** Compression of signal for VGPR occupency. */
+	/** The signal is not being compressed. */
+	#define SIGNAL_COMPRESSION_DISABLED 0
+
+	/** VGPR compression for spherical harmonic. */
+	#define SIGNAL_COMPRESSION_DIFFUSE_INDIRECT_HARMONIC 0xD200
+
+
 //------------------------------------------------------- GLOBAL CONFIGURATION OF ALL DENOISER SHADERS
 
 /** Do normal based anysotropy. */
 // TODO: the hierarchical denoiser needs that.
 #define CONFIG_NORMAL_ANISOTROPY 0
 
-/** Whether should enable VGPR optimizations. */
+/** Whether should enable VGPR optimizations in the entry point of the shaders. */
 #define CONFIG_VGPR_OPTIMIZATION 1
 
 /** Whether should do VALU optimizations. */
diff --git a/Engine/Shaders/Private/ScreenSpaceDenoise/SSDSignalAccumulator.ush b/Engine/Shaders/Private/ScreenSpaceDenoise/SSDSignalAccumulator.ush
--- a/Engine/Shaders/Private/ScreenSpaceDenoise/SSDSignalAccumulator.ush
+++ b/Engine/Shaders/Private/ScreenSpaceDenoise/SSDSignalAccumulator.ush
@@ -3,6 +3,12 @@
 #pragma once
 
 #include "SSDSignalCore.ush"
+#include "SSDSignalCompression.ush"
+
+
+#ifndef CONFIG_SIGNAL_VGPR_COMPRESSION
+	#define CONFIG_SIGNAL_VGPR_COMPRESSION SIGNAL_COMPRESSION_DISABLED
+#endif
 
 
 //------------------------------------------------------- SIGNAL ACCUMULATOR
@@ -16,6 +22,10 @@ struct FSSDSignalAccumulator
 	FSSDSignalSample Moment1;
 	FSSDSignalSample Moment2;
 
+	// Compressed form of moment signals.
+	FSSDCompressedSignalSample CompressedMoment1;
+	FSSDCompressedSignalSample CompressedMoment2;
+
 	// Per component min and max values.
 	FSSDSignalSample Min;
 	FSSDSignalSample Max;
@@ -50,6 +60,9 @@ FSSDSignalAccumulator CreateSignalAccumulator()
 	Accumulator.Moment1 = CreateSignalSampleFromScalarValue(0.0);
 	Accumulator.Moment2 = CreateSignalSampleFromScalarValue(0.0);
 		
+	Accumulator.CompressedMoment1 = CreateCompressedSignalSampleFromScalarValue(0.0, CONFIG_SIGNAL_VGPR_COMPRESSION);
+	Accumulator.CompressedMoment2 = CreateCompressedSignalSampleFromScalarValue(0.0, CONFIG_SIGNAL_VGPR_COMPRESSION);
+
 	Accumulator.Min = CreateSignalSampleFromScalarValue(INFINITE_FLOAT);
 	Accumulator.Max = CreateSignalSampleFromScalarValue(0.0); // TODO: -INFINITE_FLOAT? otherwise there is a risk to color clamp with YCoCg.
 
@@ -237,6 +250,21 @@ void DijestSampleClusterInDRB(
 }
 
 
+//------------------------------------------------------- COMPRESSION & DECOMPRESSION
+
+void CompressSignalAccumulator(inout FSSDSignalAccumulator Accumulator)
+{
+	CompressSignalSample(Accumulator.Moment1, CONFIG_SIGNAL_VGPR_COMPRESSION, /* out */ Accumulator.CompressedMoment1);
+	CompressSignalSample(Accumulator.Moment2, CONFIG_SIGNAL_VGPR_COMPRESSION, /* out */ Accumulator.CompressedMoment2);
+}
+
+void UncompressSignalAccumulator(inout FSSDSignalAccumulator Accumulator)
+{
+	UncompressSignalSample(Accumulator.CompressedMoment1, CONFIG_SIGNAL_VGPR_COMPRESSION, /* inout */ Accumulator.Moment1);
+	UncompressSignalSample(Accumulator.CompressedMoment2, CONFIG_SIGNAL_VGPR_COMPRESSION, /* inout */ Accumulator.Moment2);
+}
+
+
 //------------------------------------------------------- SPATIAL KERNEL ENTRY POINT.
 
 /** Start accumulating a cluster of samples. */
@@ -253,11 +281,13 @@ void AccumulateSample(
 	inout FSSDSignalAccumulator Accumulator,
 	FSSDSampleAccumulationInfos SampleInfos)
 {
-	// TODO: Decompress and recompress the accumulator to reduce VGPR pressure.
+	UncompressSignalAccumulator(/* inout */ Accumulator);
 
 	AccumulateSampleInOneBin(/* inout */ Accumulator, SampleInfos);
 	AccumulateSampleInDomainBoundaries(/* inout */ Accumulator, SampleInfos);
 	AccumulateSampleInDRB(/* inout */ Accumulator, SampleInfos);
+
+	CompressSignalAccumulator(/* inout */ Accumulator);
 }
 
 /** Digest the accumulated cluster of samples. */
diff --git a/Engine/Shaders/Private/ScreenSpaceDenoise/SSDSignalCompression.ush b/Engine/Shaders/Private/ScreenSpaceDenoise/SSDSignalCompression.ush
new file mode 100644
--- /dev/null
+++ b/Engine/Shaders/Private/ScreenSpaceDenoise/SSDSignalCompression.ush
@@ -0,0 +1,99 @@
+// Copyright 1998-2019 Epic Games, Inc. All Rights Reserved.
+
+#pragma once
+
+#include "SSDSignalCore.ush"
+
+
+//------------------------------------------------------- STRUCTURE
+
+#if CONFIG_SIGNAL_COLOR_SH
+	#define MAX_SIGNAL_COMPRESSED_VGPRS 7
+#else
+	#define MAX_SIGNAL_COMPRESSED_VGPRS 1
+#endif
+
+/** Generic data structure to store a signal sample in a compressed form. */
+struct FSSDCompressedSignalSample
+{
+	// Raw array of vector general purpose registers used to store signal.
+	uint VGPR[MAX_SIGNAL_COMPRESSED_VGPRS];
+};
+
+
+uint Compress2Floats(float2 v)
+{
+	// f32tof16() is full rate on GCN.
+	return f32tof16(v.x) | (f32tof16(v.y) << 16);
+}
+
+float2 Uncompress2Floats(uint v)
+{
+	// f16tof32() is full rate on GCN.
+	return float2(f16tof32(v), f16tof32(v >> 16));
+}
+
+/** Compress part of signal into smaller VGPR foorprint. */
+void CompressSignalSample(in FSSDSignalSample Sample, uint CompressionLayout, out FSSDCompressedSignalSample OutCompressedSample)
+{
+	for (uint i = 0; i < MAX_SIGNAL_COMPRESSED_VGPRS; i++)
+	{
+		OutCompressedSample.VGPR[i] = 0;
+	}
+
+	if (CompressionLayout == SIGNAL_COMPRESSION_DISABLED || DEBUG_DISABLE_SIGNAL_COMPRESSION)
+	{
+		// NOP
+	}
+	#if CONFIG_SIGNAL_COLOR_SH
+	else if (CompressionLayout == SIGNAL_COMPRESSION_DIFFUSE_INDIRECT_HARMONIC)
+	{
+		OutCompressedSample.VGPR[0] = Compress2Floats(float2(Sample.SampleCount, Sample.MissCount));
+		#if SPHERICAL_HARMONIC_ORDER == 2
+			OutCompressedSample.VGPR[1] = Compress2Floats(Sample.ColorSH.R.V.xy);
+			OutCompressedSample.VGPR[2] = Compress2Floats(Sample.ColorSH.R.V.zw);
+			OutCompressedSample.VGPR[3] = Compress2Floats(Sample.ColorSH.G.V.xy);
+			OutCompressedSample.VGPR[4] = Compress2Floats(Sample.ColorSH.G.V.zw);
+			OutCompressedSample.VGPR[5] = Compress2Floats(Sample.ColorSH.B.V.xy);
+			OutCompressedSample.VGPR[6] = Compress2Floats(Sample.ColorSH.B.V.zw);
+		#else
+			#error Unimplemented.
+		#endif
+	}
+	#endif
+}
+
+/** Decompress signal from tight VGPR layout into a sample. */
+void UncompressSignalSample(in FSSDCompressedSignalSample CompressedSample, uint CompressionLayout, inout FSSDSignalSample OutSample)
+{
+	if (CompressionLayout == SIGNAL_COMPRESSION_DISABLED || DEBUG_DISABLE_SIGNAL_COMPRESSION)
+	{
+		// NOP to not clobber information in OutSample that remained uncompressed.
+	}
+	#if CONFIG_SIGNAL_COLOR_SH
+	else if (CompressionLayout == SIGNAL_COMPRESSION_DIFFUSE_INDIRECT_HARMONIC)
+	{
+		OutSample.SampleCount    = Uncompress2Floats(CompressedSample.VGPR[0]).x;
+		OutSample.MissCount      = Uncompress2Floats(CompressedSample.VGPR[0]).y;
+		#if SPHERICAL_HARMONIC_ORDER == 2
+			OutSample.ColorSH.R.V.xy = Uncompress2Floats(CompressedSample.VGPR[1]);
+			OutSample.ColorSH.R.V.zw = Uncompress2Floats(CompressedSample.VGPR[2]);
+			OutSample.ColorSH.G.V.xy = Uncompress2Floats(CompressedSample.VGPR[3]);
+			OutSample.ColorSH.G.V.zw = Uncompress2Floats(CompressedSample.VGPR[4]);
+			OutSample.ColorSH.B.V.xy = Uncompress2Floats(CompressedSample.VGPR[5]);
+			OutSample.ColorSH.B.V.zw = Uncompress2Floats(CompressedSample.VGPR[6]);
+		#else
+			#error Unimplemented.
+		#endif
+	}
+	#endif
+}
+
+FSSDCompressedSignalSample CreateCompressedSignalSampleFromScalarValue(float Scalar, uint CompressionLayout)
+{
+	FSSDSignalSample Sample = CreateSignalSampleFromScalarValue(Scalar);
+
+	FSSDCompressedSignalSample CompressedSample;
+	CompressSignalSample(Sample, CompressionLayout, /* out */ CompressedSample);
+	return CompressedSample;
+}
diff --git a/Engine/Shaders/Private/ScreenSpaceDenoise/SSDSpatialAccumulation.usf b/Engine/Shaders/Private/ScreenSpaceDenoise/SSDSpatialAccumulation.usf
--- a/Engine/Shaders/Private/ScreenSpaceDenoise/SSDSpatialAccumulation.usf
+++ b/Engine/Shaders/Private/ScreenSpaceDenoise/SSDSpatialAccumulation.usf
@@ -286,6 +286,10 @@
 	#define CONFIG_SIGNAL_INPUT_LAYOUT  SIGNAL_BUFFER_LAYOUT_DIFFUSE_INDIRECT_HARMONIC
 	#define CONFIG_SIGNAL_OUTPUT_LAYOUT SIGNAL_BUFFER_LAYOUT_DIFFUSE_INDIRECT_HARMONIC
 
+	// Spherical harmonics are a lot of data, need to shrink VGPR pressure to improve lattency hidding when fetching the buffer.
+	// TODO: some shader compiler completly falls apparts with the current implementation of this and actually drops in occupency.
+	#define CONFIG_SIGNAL_VGPR_COMPRESSION SIGNAL_COMPRESSION_DIFFUSE_INDIRECT_HARMONIC
+
 	#if DIM_STAGE == STAGE_RECONSTRUCTION
 		#define CONFIG_SAMPLE_SET           SAMPLE_SET_STACKOWIAK_4_SETS
 
@@ -764,10 +768,13 @@ void MainCS(
 	}
 
 	// Color processing of the signal to reduce highlight flickering.
+	#if CONFIG_ACCUMULATION_COLOR_SPACE != CONFIG_OUTPUT_COLOR_SPACE || CONFIG_REJECT_HIGHEST_COLOR
 	{
 		UNROLL_N(CONFIG_SIGNAL_BATCH_SIZE)
 		for (uint MultiplexId = 0; MultiplexId < CONFIG_SIGNAL_BATCH_SIZE; MultiplexId++)
 		{
+			UncompressSignalAccumulator(/* inout */ SignalAccumulators.Array[MultiplexId]);
+
 			#if CONFIG_REJECT_HIGHEST_COLOR
 			{
 				#if !CONFIG_SIGNAL_COLOR
@@ -782,19 +789,26 @@ void MainCS(
 						(SignalAccumulators.Array[MultiplexId].Moment1.SampleCount / (SignalAccumulators.Array[MultiplexId].Moment1.SampleCount - MaxNeighborWeight));
 				}
 			}
-			#endif
-
-			SignalAccumulators.Array[MultiplexId].Moment1 = TransformSignal(
-				SignalAccumulators.Array[MultiplexId].Moment1, 
-				/* SrcBasis  = */ CONFIG_ACCUMULATION_COLOR_SPACE, 
-				/* DestBasis = */ CONFIG_OUTPUT_COLOR_SPACE);
-
-			SignalAccumulators.Array[MultiplexId].Moment2 = TransformSignal(
-				SignalAccumulators.Array[MultiplexId].Moment2, 
-				/* SrcBasis  = */ CONFIG_ACCUMULATION_COLOR_SPACE, 
-				/* DestBasis = */ CONFIG_OUTPUT_COLOR_SPACE);
+			#endif // CONFIG_REJECT_HIGHEST_COLOR
+
+			#if CONFIG_ACCUMULATION_COLOR_SPACE != CONFIG_OUTPUT_COLOR_SPACE
+				SignalAccumulators.Array[MultiplexId].Moment1 = TransformSignal(
+					SignalAccumulators.Array[MultiplexId].Moment1, 
+					/* SrcBasis  = */ CONFIG_ACCUMULATION_COLOR_SPACE, 
+					/* DestBasis = */ CONFIG_OUTPUT_COLOR_SPACE);
+
+				SignalAccumulators.Array[MultiplexId].Moment2 = TransformSignal(
+					SignalAccumulators.Array[MultiplexId].Moment2, 
+					/* SrcBasis  = */ CONFIG_ACCUMULATION_COLOR_SPACE, 
+					/* DestBasis = */ CONFIG_OUTPUT_COLOR_SPACE);
+			#endif // CONFIG_ACCUMULATION_COLOR_SPACE != CONFIG_OUTPUT_COLOR_SPACE
+
+			// TODO: it might be better to just uncompress before this for loop and remain uncompressed,
+			// so the color operation get done in practice during the output sample transcoding.
+			CompressSignalAccumulator(/* inout */ SignalAccumulators.Array[MultiplexId]);
 		}
 	}
+	#endif // CONFIG_ACCUMULATION_COLOR_SPACE != CONFIG_OUTPUT_COLOR_SPACE || CONFIG_REJECT_HIGHEST_COLOR
 
 	// Transcode the spatial accumulation into multiplexed signal according to different modes.
 	uint MultiplexCount = 1;
@@ -807,6 +821,8 @@ void MainCS(
 			UNROLL_N(CONFIG_SIGNAL_BATCH_SIZE)
 			for (uint MultiplexId = 0; MultiplexId < CONFIG_SIGNAL_BATCH_SIZE; MultiplexId++)
 			{
+				UncompressSignalAccumulator(/* inout */ SignalAccumulators.Array[MultiplexId]);
+
 				OutputSamples.Array[MultiplexId] = SignalAccumulators.Array[MultiplexId].Moment1;
 				
 				// Output the minimal inverse frequency as new world bluring radius for subsequent passes.
@@ -822,6 +838,8 @@ void MainCS(
 			UNROLL_N(CONFIG_SIGNAL_BATCH_SIZE)
 			for (uint BatchedSignalId = 0; BatchedSignalId < CONFIG_SIGNAL_BATCH_SIZE; BatchedSignalId++)
 			{
+				UncompressSignalAccumulator(/* inout */ SignalAccumulators.Array[BatchedSignalId * CONFIG_MULTIPLEXED_SIGNALS_PER_SIGNAL_DOMAIN + 0]);
+
 				OutputSamples.Array[BatchedSignalId * 2 + 0] = SignalAccumulators.Array[BatchedSignalId * CONFIG_MULTIPLEXED_SIGNALS_PER_SIGNAL_DOMAIN + 0].Moment1;
 				OutputSamples.Array[BatchedSignalId * 2 + 1] = SignalAccumulators.Array[BatchedSignalId * CONFIG_MULTIPLEXED_SIGNALS_PER_SIGNAL_DOMAIN + 0].Moment2;
 			}
@@ -833,6 +851,8 @@ void MainCS(
 			UNROLL_N(CONFIG_SIGNAL_BATCH_SIZE)
 			for (uint MultiplexId = 0; MultiplexId < CONFIG_SIGNAL_BATCH_SIZE; MultiplexId++)
 			{
+				UncompressSignalAccumulator(/* inout */ SignalAccumulators.Array[MultiplexId]);
+
 				OutputSamples.Array[MultiplexId] = SignalAccumulators.Array[MultiplexId].Previous;
 
 				// Output the minimal inverse frequency as new world bluring radius for subsequent passes.
