Implements manual texel formatting for the spherical harmonic denoiser, saving another 5 VGPRs.

diff --git a/Engine/Shaders/Private/ScreenSpaceDenoise/SSDCommon.ush b/Engine/Shaders/Private/ScreenSpaceDenoise/SSDCommon.ush
--- a/Engine/Shaders/Private/ScreenSpaceDenoise/SSDCommon.ush
+++ b/Engine/Shaders/Private/ScreenSpaceDenoise/SSDCommon.ush
@@ -62,6 +62,29 @@ float SafeRcp(float x)
 	return 0;
 }
 
+uint Compress2Floats(float2 v)
+{
+	// f32tof16() is full rate on GCN.
+	return f32tof16(v.x) | (f32tof16(v.y) << 16);
+}
+
+float2 Uncompress2Floats(uint v)
+{
+	// f16tof32() is full rate on GCN.
+	return float2(f16tof32(v), f16tof32(v >> 16));
+}
+
+uint2 Compress4Floats(float4 v)
+{
+	return uint2(Compress2Floats(v.xy), Compress2Floats(v.zw));
+}
+
+float4 Uncompress4Floats(uint2 v)
+{
+	return float4(Uncompress2Floats(v.x), Uncompress2Floats(v.y));
+}
+
+
 //------------------------------------------------------- SPHERICAL HARMONIC INTERFACE
 
 #if SPHERICAL_HARMONIC_ORDER == 0
diff --git a/Engine/Shaders/Private/ScreenSpaceDenoise/SSDDefinitions.ush b/Engine/Shaders/Private/ScreenSpaceDenoise/SSDDefinitions.ush
--- a/Engine/Shaders/Private/ScreenSpaceDenoise/SSDDefinitions.ush
+++ b/Engine/Shaders/Private/ScreenSpaceDenoise/SSDDefinitions.ush
@@ -212,6 +212,9 @@
 // TODO: the shader compiler is not doing the right thing.
 #define CONFIG_SGPR_HINT_OPTIMIZATION 0
 
+/** Manually format texel in the shader to reduce VGPR pressure with overlapped texture fetched. */
+#define CONFIG_MANUAL_TEXEL_FORMATTING 1
+
 /** By default, the color space stored into intermediary buffer is linear premultiplied RGBA. */
 #define STANDARD_BUFFER_COLOR_SPACE COLOR_SPACE_RGB
 
diff --git a/Engine/Shaders/Private/ScreenSpaceDenoise/SSDSignalBufferEncoding.ush b/Engine/Shaders/Private/ScreenSpaceDenoise/SSDSignalBufferEncoding.ush
--- a/Engine/Shaders/Private/ScreenSpaceDenoise/SSDSignalBufferEncoding.ush
+++ b/Engine/Shaders/Private/ScreenSpaceDenoise/SSDSignalBufferEncoding.ush
@@ -256,6 +256,15 @@ void DecodeMultiplexedSignals(
 	#if CONFIG_SIGNAL_COLOR_SH
 	else if (SignalBufferLayout == SIGNAL_BUFFER_LAYOUT_DIFFUSE_INDIRECT_HARMONIC)
 	{
+		#if CONFIG_MANUAL_TEXEL_FORMATTING
+		{
+			 RawSample[0] = Uncompress4Floats(uint2(asuint(RawSample[0].x), asuint(RawSample[0].y)));
+			 RawSample[1] = Uncompress4Floats(uint2(asuint(RawSample[1].x), asuint(RawSample[1].y)));
+			 RawSample[2] = Uncompress4Floats(uint2(asuint(RawSample[2].x), asuint(RawSample[2].y)));
+			 RawSample[3].xy = Uncompress2Floats(asuint(RawSample[3].x));
+		}
+		#endif
+
 		float SampleCount;
 		if (bNormalizeSample)
 		{
@@ -478,6 +487,16 @@ void EncodeMultiplexedSignals(
 		OutRawSample[1] = Sample[0].ColorSH.G.V * NormalizationFactor;
 		OutRawSample[2] = Sample[0].ColorSH.B.V * NormalizationFactor;
 		OutBufferCount = 4;
+
+		#if CONFIG_MANUAL_TEXEL_FORMATTING
+		{
+			 OutRawSample[0] = float4(asfloat(Compress2Floats(OutRawSample[0].xy)), asfloat(Compress2Floats(OutRawSample[0].zw)), 0, 0);
+			 OutRawSample[1] = float4(asfloat(Compress2Floats(OutRawSample[1].xy)), asfloat(Compress2Floats(OutRawSample[1].zw)), 0, 0);
+			 OutRawSample[2] = float4(asfloat(Compress2Floats(OutRawSample[2].xy)), asfloat(Compress2Floats(OutRawSample[2].zw)), 0, 0);
+			 OutRawSample[3] = float4(asfloat(Compress2Floats(OutRawSample[3].xy)), 0, 0, 0);
+		}
+
+		#endif
 	}
 	#endif
 } // EncodeMultiplexedSignals()
diff --git a/Engine/Shaders/Private/ScreenSpaceDenoise/SSDSignalCompression.ush b/Engine/Shaders/Private/ScreenSpaceDenoise/SSDSignalCompression.ush
--- a/Engine/Shaders/Private/ScreenSpaceDenoise/SSDSignalCompression.ush
+++ b/Engine/Shaders/Private/ScreenSpaceDenoise/SSDSignalCompression.ush
@@ -21,18 +21,6 @@ struct FSSDCompressedSignalSample
 };
 
 
-uint Compress2Floats(float2 v)
-{
-	// f32tof16() is full rate on GCN.
-	return f32tof16(v.x) | (f32tof16(v.y) << 16);
-}
-
-float2 Uncompress2Floats(uint v)
-{
-	// f16tof32() is full rate on GCN.
-	return float2(f16tof32(v), f16tof32(v >> 16));
-}
-
 /** Compress part of signal into smaller VGPR foorprint. */
 void CompressSignalSample(in FSSDSignalSample Sample, uint CompressionLayout, out FSSDCompressedSignalSample OutCompressedSample)
 {
diff --git a/Engine/Source/Runtime/Renderer/Private/ScreenSpaceDenoise.cpp b/Engine/Source/Runtime/Renderer/Private/ScreenSpaceDenoise.cpp
--- a/Engine/Source/Runtime/Renderer/Private/ScreenSpaceDenoise.cpp
+++ b/Engine/Source/Runtime/Renderer/Private/ScreenSpaceDenoise.cpp
@@ -865,6 +865,9 @@ static void DenoiseSignalAtConstantPixelDensity(
 	TStaticArray<FRDGTextureDesc, kMaxBufferProcessingCount> HistoryDescs;
 	FRDGTextureDesc DebugDesc;
 	{
+		// Manually format texel in the shader to reduce VGPR pressure with overlapped texture fetched.
+		const bool bManualTexelFormatting = true;
+
 		static const EPixelFormat PixelFormatPerChannel[] = {
 			PF_Unknown,
 			PF_R16F,
@@ -943,12 +946,12 @@ static void DenoiseSignalAtConstantPixelDensity(
 		{
 			for (int32 i = 0; i < 3; i++)
 			{
-				ReconstructionDescs[i].Format = PF_FloatRGBA;
-				HistoryDescs[i].Format = PF_FloatRGBA;
+				ReconstructionDescs[i].Format = bManualTexelFormatting ? PF_G32R32F : PF_FloatRGBA;
+				HistoryDescs[i].Format = bManualTexelFormatting ? PF_G32R32F : PF_FloatRGBA;
 			}
 
-			ReconstructionDescs[3].Format = PF_G16R16F;
-			HistoryDescs[3].Format = PF_G16R16F;
+			ReconstructionDescs[3].Format = bManualTexelFormatting ? PF_R32_FLOAT : PF_G16R16F;
+			HistoryDescs[3].Format = bManualTexelFormatting ? PF_R32_FLOAT : PF_G16R16F;
 
 			ReconstructionTextureCount = IScreenSpaceDenoiser::kSphericalHarmonicTextureCount;
 			HistoryTextureCountPerSignal = IScreenSpaceDenoiser::kSphericalHarmonicTextureCount; // TODO: only 3 textures for history
