Saves another 8 VGPRS in spherical harmonic, by compressing normals, computing SampleTrackId from lane, and working arround shader compiler bugs.

ISOLATE to be reimplemented.

diff --git a/Engine/Shaders/Private/ScreenSpaceDenoise/SSDCommon.ush b/Engine/Shaders/Private/ScreenSpaceDenoise/SSDCommon.ush
--- a/Engine/Shaders/Private/ScreenSpaceDenoise/SSDCommon.ush
+++ b/Engine/Shaders/Private/ScreenSpaceDenoise/SSDCommon.ush
@@ -84,6 +84,25 @@ float4 Uncompress4Floats(uint2 v)
 	return float4(Uncompress2Floats(v.x), Uncompress2Floats(v.y));
 }
 
+uint CompressN(float3 N)
+{
+	// matches 8bits GBuffer A to be lossless.
+	uint3 K = uint3(saturate(N * 0.5 + 0.5) * 255.0);
+
+	// TODO: K.x is stored in << 24 instead of << 0 to work arround a shader compiler bug.
+	return uint(K.x << 24 | K.y << 8 | K.z << 16);
+}
+
+float3 DecompressN(uint EncodedN)
+{
+	uint3 K;
+	K.x = (EncodedN >> 24) & 0xFF;
+	K.y = (EncodedN >>  8) & 0xFF;
+	K.z = (EncodedN >> 16) & 0xFF;
+
+	return float3(K) * (2.0 / 255.0) - 1.0;
+}
+
 
 //------------------------------------------------------- SPHERICAL HARMONIC INTERFACE
 
@@ -112,6 +131,9 @@ struct FSSDSampleSceneInfos
 	/** Normal of the pixel in world space. */
 	float3 WorldNormal;
 
+	/** Compressed form of the world normal. */
+	uint CompressedWorldNormal;
+
 	/** The position of the pixel in the translated world frame. */
 	float3 TranslatedWorldPosition;
 
@@ -126,6 +148,7 @@ FSSDSampleSceneInfos CreateSampleSceneInfos()
 {
 	FSSDSampleSceneInfos Infos;
 	Infos.WorldNormal             = 0;
+	Infos.CompressedWorldNormal   = 0;
 	Infos.TranslatedWorldPosition = 0;
 	Infos.WorldDepth              = 0;
 	Infos.Roughness               = 0;
@@ -142,6 +165,7 @@ FSSDSampleSceneInfos InitSampleSceneInfos(
 	Infos.TranslatedWorldPosition = mul(float4(ClipPosition, WorldDepth, 1), ScreenToTranslatedWorld).xyz;
 	Infos.WorldDepth = WorldDepth;
 	Infos.WorldNormal = WorldNormal;
+	Infos.CompressedWorldNormal = CompressN(WorldNormal); // TODO: directly fetch CompressedWorldNormal from the sample 
 	Infos.Roughness = Roughness;
 	return Infos;
 }
@@ -171,6 +195,18 @@ FSSDSampleSceneInfos FetchPrevSceneInfos(float2 ScreenPosition, float2 BufferUV)
 	return OutMetadata;
 }
 
+// ISOLATE
+float3 GetWorldNormal(FSSDSampleSceneInfos SceneMetadata)
+#if CONFIG_COMPRESS_WORLD_NORMAL
+{
+	return DecompressN(SceneMetadata.CompressedWorldNormal);
+}
+#else
+{
+	return SceneMetadata.WorldNormal;
+}
+#endif
+
 float ProjectWorldDistanceToScreenDistance(FSSDSampleSceneInfos SceneMetadata, float WorldDistance)
 {
 	float4 ClipDistance = mul(View.ViewToClipNoAA, float4(WorldDistance, 0, SceneMetadata.WorldDepth, 1));
@@ -190,7 +226,7 @@ float ComputeAnisotropyInvFactor(FSSDSampleSceneInfos SceneMetadata)
 		return 1;
 	}
 
-	float Anysotropy = abs(dot(SceneMetadata.WorldNormal, View.ViewForward));
+	float Anysotropy = abs(dot(GetWorldNormal(SceneMetadata), View.ViewForward));
 
 	return max(Anysotropy, rcp(8));
 }
@@ -227,7 +263,8 @@ float ComputeSignalMipLevel(FSSDSampleSceneInfos SceneMetadata, float WorldBluri
 float ComputeWorldBluringRadiusCausedByPixelSize(FSSDSampleSceneInfos SceneMetadata)
 {
 	// Should be multiplied 0.5* for the diameter to radius, and by 2.0 because GetTanHalfFieldOfView() cover only half of the pixels.
-	return (View.ViewSizeAndInvSize.z * GetTanHalfFieldOfView().x) * SceneMetadata.WorldDepth;
+	// TODO: should do shader parameter.
+	return ToScalarMemory(View.ViewSizeAndInvSize.z * GetTanHalfFieldOfView().x) * SceneMetadata.WorldDepth;
 }
 
 /** Convert to bilateral world distance. */
diff --git a/Engine/Shaders/Private/ScreenSpaceDenoise/SSDDefinitions.ush b/Engine/Shaders/Private/ScreenSpaceDenoise/SSDDefinitions.ush
--- a/Engine/Shaders/Private/ScreenSpaceDenoise/SSDDefinitions.ush
+++ b/Engine/Shaders/Private/ScreenSpaceDenoise/SSDDefinitions.ush
@@ -196,6 +196,15 @@
 	#define SIGNAL_COMPRESSION_DIFFUSE_INDIRECT_HARMONIC 0xD200
 
 
+/** Technic used to compute the world vector betweem neighbor and reference. */
+	// Directly use FSSDKernelConfig::RefSceneMetadata.TranslatedWorldPosition. Cost 3 VGPR over entire kernel inner loop.
+	#define NEIGHBOR_TO_REF_CACHE_WORLD_POSITION 0
+
+	// Compute by deriving the information of the reference pixel last minute from kernel configuration, trading VALU to save VGPR.
+	// TODO: does not work with previous frame reprojection.
+	#define NEIGHBOR_TO_REF_LOWEST_VGPR_PRESSURE 1
+
+
 //------------------------------------------------------- GLOBAL CONFIGURATION OF ALL DENOISER SHADERS
 
 /** Do normal based anysotropy. */
@@ -215,6 +224,9 @@
 /** Manually format texel in the shader to reduce VGPR pressure with overlapped texture fetched. */
 #define CONFIG_MANUAL_TEXEL_FORMATTING 1
 
+/** Configures whether the world normal should be compressed in FSSDSampleSceneInfos. */
+#define CONFIG_COMPRESS_WORLD_NORMAL 1
+
 /** By default, the color space stored into intermediary buffer is linear premultiplied RGBA. */
 #define STANDARD_BUFFER_COLOR_SPACE COLOR_SPACE_RGB
 
diff --git a/Engine/Shaders/Private/ScreenSpaceDenoise/SSDSignalFramework.ush b/Engine/Shaders/Private/ScreenSpaceDenoise/SSDSignalFramework.ush
--- a/Engine/Shaders/Private/ScreenSpaceDenoise/SSDSignalFramework.ush
+++ b/Engine/Shaders/Private/ScreenSpaceDenoise/SSDSignalFramework.ush
@@ -508,7 +508,8 @@ bool IsInvalidSample(FSSDSignalSample Sample)
 float ComputeBilateralWeight(
 	float MaxWorldBluringDistance,
 	FSSDSampleSceneInfos RefSceneMetadata,
-	FSSDSampleSceneInfos NeighborSceneMetadata)
+	FSSDSampleSceneInfos NeighborSceneMetadata,
+	float3 NeighborToRefWorldVector)
 {
 	if (DEBUG_DISABLE_BILATERAL)
 	{
@@ -520,14 +521,12 @@ float ComputeBilateralWeight(
 	{
 		const float WorldRadius = MaxWorldBluringDistance;
 
-		float3 V = RefSceneMetadata.TranslatedWorldPosition - NeighborSceneMetadata.TranslatedWorldPosition;
-
 		float RefAnisotropyInvFactor = ComputeAnisotropyInvFactor(RefSceneMetadata);
 
 		float3 CameraToRef = View.ViewForward;
 
-		float Z = dot(CameraToRef, V);
-		float XY = length2(V - CameraToRef * Z * (1 - RefAnisotropyInvFactor));
+		float Z = dot(CameraToRef, NeighborToRefWorldVector);
+		float XY = length2(NeighborToRefWorldVector - CameraToRef * Z * (1 - RefAnisotropyInvFactor));
 
 		float DistSquare = XY;
 
@@ -537,48 +536,41 @@ float ComputeBilateralWeight(
 	}
 	#elif CONFIG_VIEW_POSITION_BASED_BILATERAL == 2
 	{
-		// TODO: save 2 VGPR by encoding RefSceneMetadata.WorldNormal
-
 		// TODO: save 2 VGPR by using SceneDepth instead of RefSceneMetadata.TranslatedWorldPosition (not sure it can work in math yet).
 		// float2 ClipPosition = ScreenPosition * (View.ViewToClip[3][3] < 1.0f ? WorldDepth : 1.0f);
 		// Infos.TranslatedWorldPosition = mul(float4(ClipPosition, WorldDepth, 1), ScreenToTranslatedWorld).xyz;
 
-		float DistFromRefPlane = abs(dot(RefSceneMetadata.WorldNormal, RefSceneMetadata.TranslatedWorldPosition - NeighborSceneMetadata.TranslatedWorldPosition));
+		float DistFromRefPlane = abs(dot(GetWorldNormal(RefSceneMetadata), NeighborToRefWorldVector));
 
 		PositionBilateralWeight = saturate(1 - DistFromRefPlane * rcp(MaxWorldBluringDistance));
 	}
 	#elif CONFIG_VIEW_POSITION_BASED_BILATERAL == 3
 	{
-		float3 V = RefSceneMetadata.TranslatedWorldPosition - NeighborSceneMetadata.TranslatedWorldPosition;
-
-		if (dot(V, V) > MaxWorldBluringDistance * MaxWorldBluringDistance)
+		if (dot(NeighborToRefWorldVector, NeighborToRefWorldVector) > MaxWorldBluringDistance * MaxWorldBluringDistance)
 		{
 			PositionBilateralWeight = 0;
 		}
 	}
 	#elif CONFIG_VIEW_POSITION_BASED_BILATERAL == 4
 	{
-		float3 V = RefSceneMetadata.TranslatedWorldPosition - NeighborSceneMetadata.TranslatedWorldPosition;
-
 		// TODO: extremely expensive with rcp()...
-		PositionBilateralWeight = saturate(1 - dot(V, V) * rcp(MaxWorldBluringDistance * MaxWorldBluringDistance));
+		PositionBilateralWeight = saturate(1 - dot(NeighborToRefWorldVector, NeighborToRefWorldVector) * rcp(MaxWorldBluringDistance * MaxWorldBluringDistance));
 	}
 	#elif CONFIG_VIEW_POSITION_BASED_BILATERAL == 5
 	{
-		float3 V = RefSceneMetadata.TranslatedWorldPosition - NeighborSceneMetadata.TranslatedWorldPosition;
-
-		float DistFromRefPlane = dot(RefSceneMetadata.WorldNormal, RefSceneMetadata.TranslatedWorldPosition - NeighborSceneMetadata.TranslatedWorldPosition);
+		float DistFromRefPlane = dot(GetWorldNormal(RefSceneMetadata), RefSceneMetadata.TranslatedWorldPosition - NeighborSceneMetadata.TranslatedWorldPosition);
 
 		// TODO: extremely expensive with rcp()...
 		PositionBilateralWeight = (
-			saturate(1 - dot(V, V) * rcp(MaxWorldBluringDistance * MaxWorldBluringDistance)) *
+			saturate(1 - dot(NeighborToRefWorldVector, NeighborToRefWorldVector) * rcp(MaxWorldBluringDistance * MaxWorldBluringDistance)) *
 			saturate(1 - 6 * abs(DistFromRefPlane) * rcp(MaxWorldBluringDistance)));
 	}
 	#endif
 
 	float NormalBilateralWeight = 1;
 	#if CONFIG_NORMAL_BASED_BILATERAL
-    {
+	{
+		// TODO: use GetWorldNormal() instead to reduce VGPR?
 		float NoN = max(dot(RefSceneMetadata.WorldNormal, NeighborSceneMetadata.WorldNormal), 0);
 
 		#if 1 // TODO
@@ -593,7 +585,7 @@ float ComputeBilateralWeight(
 
 			NormalBilateralWeight = exp(-x * x);
 		#endif
-    }
+	}
 	#endif
 	
 	// [ Tokoyashi 2015, "Specular Lobe-Aware Filtering and Upsampling for Interactive Indirect Illumination" ]
@@ -734,6 +726,7 @@ float GetRatioEstimatorWeight(
 
 	float a2 = Pow4(RefSceneMetadata.Roughness);
 
+	// TODO: Improve quantization in CONFIG_COMPRESS_WORLD_NORMAL and use it.
 	float3 N = RefSceneMetadata.WorldNormal;
 
 	float NoH;
@@ -797,7 +790,7 @@ FSSDSphericalHarmonicRGB ComputeSampleColorSH(FSSDSampleSceneInfos SampleSceneMe
 	float3 RandSample = QuasiRandom_Sample3(QuasiRandom);
 
 	float4 TangentDirection = CosineSampleHemisphere(RandSample.yz);
-	float3 WorldDirection = TangentToWorld(TangentDirection.xyz, SampleSceneMetadata.WorldNormal);
+	float3 WorldDirection = TangentToWorld(TangentDirection.xyz, SampleSceneMetadata.WorldNormal); // TODO: use GetWorldNormal()
 
 	#if SPHERICAL_HARMONIC_ORDER == 2
 		FSSDSphericalHarmonic SampleBasis = SHBasisFunction(WorldDirection);
diff --git a/Engine/Shaders/Private/ScreenSpaceDenoise/SSDSpatialAccumulation.usf b/Engine/Shaders/Private/ScreenSpaceDenoise/SSDSpatialAccumulation.usf
--- a/Engine/Shaders/Private/ScreenSpaceDenoise/SSDSpatialAccumulation.usf
+++ b/Engine/Shaders/Private/ScreenSpaceDenoise/SSDSpatialAccumulation.usf
@@ -287,8 +287,10 @@
 	#define CONFIG_SIGNAL_OUTPUT_LAYOUT SIGNAL_BUFFER_LAYOUT_DIFFUSE_INDIRECT_HARMONIC
 
 	// Spherical harmonics are a lot of data, need to shrink VGPR pressure to improve lattency hidding when fetching the buffer.
-	// TODO: some shader compiler completly falls apparts with the current implementation of this and actually drops in occupency.
+	// TODO: some shader compiler completly falls apparts with the current implementation of
+	//       CONFIG_SIGNAL_VGPR_COMPRESSION and actually drops in occupency.
 	#define CONFIG_SIGNAL_VGPR_COMPRESSION SIGNAL_COMPRESSION_DIFFUSE_INDIRECT_HARMONIC
+	#define CONFIG_NEIGHBOR_TO_REF_COMPUTATION NEIGHBOR_TO_REF_LOWEST_VGPR_PRESSURE
 
 	#if DIM_STAGE == STAGE_RECONSTRUCTION
 		#define CONFIG_SAMPLE_SET           SAMPLE_SET_STACKOWIAK_4_SETS
@@ -380,6 +382,17 @@
 #endif
 
 
+//------------------------------------------------------- LATE SETUP
+
+#if COMPILER_SUPPORTS_WAVE_ONCE || 1
+	#if CONFIG_UPSCALE
+		#define CONFIG_VGPR_FREE_SAMPLE_TRACK_ID 2
+	#else
+		#define CONFIG_VGPR_FREE_SAMPLE_TRACK_ID 1
+	#endif
+#endif
+
+
 //------------------------------------------------------- INCLUDES
 
 #include "SSDSignalFramework.ush"
@@ -682,7 +695,7 @@ void MainCS(
 		{
 			KernelConfig.SampleCount = clamp(RequestedSampleCount * rcp(kStackowiakSampleSetCount), 1, MaxSampleCount);
 
-			if (CONFIG_UPSCALE)
+			#if CONFIG_UPSCALE
 			{
 				// TODO: could be optimised, but currently reusing same peace of code as reflection for maintainability.
 				uint2 RayDispatchThreadId = (DispatchThreadId - UpscaleFactor / 2) / UpscaleFactor;
@@ -701,7 +714,7 @@ void MainCS(
 				// To avoid precision problem when comparing potentially identicall 
 				KernelConfig.bForceKernelCenterAccumulation = RaySubPixelId == ((DispatchThreadId.x & 1) | ((DispatchThreadId.y & 1) << 1));
 			}
-			else
+			#else
 			{
 				// Put the kernel center at the center of the quad. Half pixel shift is done in the sample offsets.
 				KernelConfig.BufferUV = (DispatchThreadId | 1) * View.BufferSizeAndInvSize.zw;
@@ -709,6 +722,7 @@ void MainCS(
 				// Id of the pixel in the quad. This is to match hard coded first samples of the sample set.
 				KernelConfig.SampleTrackId = ((DispatchThreadId.x & 1) | ((DispatchThreadId.y & 1) << 1));
 			}
+			#endif
 
 			#if CONFIG_VGPR_OPTIMIZATION
 				// Keek sample SampleTrackId & SceneBufferUV arround for computation of pixel output coordinate.
@@ -887,19 +901,25 @@ void MainCS(
 	#endif // DIM_STAGE == STAGE_POST_FILTERING
 
 	// TODO:  LeaveRayCount = (LeaveRayCount - 1) * 9 / (9 - 2) post processing to reject when for history rejection
-	
+	uint2 OutputPixelPostion;
 	#if CONFIG_VGPR_OPTIMIZATION
+	{
 		// No need to keep DispatchThreadId, can recompute the output pixel position based on information stored in VGPRs for spatial kernel.
 		#if CONFIG_SAMPLE_SET == SAMPLE_SET_STACKOWIAK_4_SETS
+			#if CONFIG_VGPR_FREE_SAMPLE_TRACK_ID
+				SampleTrackId = GetSampleTrackIdFromLaneIndex();
+			#endif
+
 			#if CONFIG_UPSCALE
 				SampleTrackId ^= 0x3;
 			#endif
-			uint2 OutputPixelPostion = (uint2(KernelBufferUV * View.BufferSizeAndInvSize.xy) & ~0x1) | (uint2(SampleTrackId, SampleTrackId >> 1) & 0x1);
+			OutputPixelPostion = (uint2(KernelBufferUV * View.BufferSizeAndInvSize.xy) & ~0x1) | (uint2(SampleTrackId, SampleTrackId >> 1) & 0x1);
 		#else
-			uint2 OutputPixelPostion = uint2(SceneBufferUV * View.BufferSizeAndInvSize.xy);
+			OutputPixelPostion = uint2(SceneBufferUV * View.BufferSizeAndInvSize.xy);
 		#endif
+	}
 	#else
-		uint2 OutputPixelPostion = DispatchThreadId;
+		OutputPixelPostion = DispatchThreadId;
 	#endif 
 
 	// Output the multiplexed signal.
diff --git a/Engine/Shaders/Private/ScreenSpaceDenoise/SSDSpatialKernel.ush b/Engine/Shaders/Private/ScreenSpaceDenoise/SSDSpatialKernel.ush
--- a/Engine/Shaders/Private/ScreenSpaceDenoise/SSDSpatialKernel.ush
+++ b/Engine/Shaders/Private/ScreenSpaceDenoise/SSDSpatialKernel.ush
@@ -24,6 +24,18 @@
 	#define SIGNAL_WORLD_FREQUENCY_PRECOMPUTED_BLURING_RADIUS 4
 
 
+//------------------------------------------------------- CONFIG DISABLED DEFAULTS
+
+/** Selects how the computation of world vector between the reference and neighbor should be computed. */
+#ifndef CONFIG_NEIGHBOR_TO_REF_COMPUTATION
+	#define CONFIG_NEIGHBOR_TO_REF_COMPUTATION NEIGHBOR_TO_REF_CACHE_WORLD_POSITION
+#endif
+
+/** Compute FSSDKernelConfig::SampleTrackId directly based on lane id. */
+#ifndef CONFIG_VGPR_FREE_SAMPLE_TRACK_ID
+	#define CONFIG_VGPR_FREE_SAMPLE_TRACK_ID 0
+#endif
+
 
 //------------------------------------------------------- STRUCTURES
 
@@ -201,6 +213,87 @@ FSSDKernelConfig CreateKernelConfig()
 static const float kWaveletFilterWeights5x5[] = { 3.0 / 8.0, 1.0 / 4.0, 1.0 / 16.0 };
 
 
+//------------------------------------------------------- REDERIVE INFORMATION FOR LOWER VGPR OCCUPENCY
+
+#if CONFIG_VGPR_FREE_SAMPLE_TRACK_ID
+
+// ISOLATE
+uint GetSampleTrackIdFromLaneIndex()
+{
+	uint LaneId = WaveGetLaneIndex();
+
+	return ((LaneId & 0x1) | ((LaneId >> 2) & 0x2)) ^ (CONFIG_VGPR_FREE_SAMPLE_TRACK_ID == 2 ? 0x3 : 0x0);
+}
+
+#endif // CONFIG_VGPR_FREE_SAMPLE_TRACK_ID
+
+
+/** Deduce the buffer UV of the output pixel this kernel has been configured for. */
+float2 ComputeRefBufferUV(FSSDKernelConfig KernelConfig)
+{
+	if (KernelConfig.SampleSet == SAMPLE_SET_STACKOWIAK_4_SETS)
+	{
+		uint SampleTrackId = KernelConfig.SampleTrackId;
+
+		#if CONFIG_VGPR_FREE_SAMPLE_TRACK_ID
+			SampleTrackId = GetSampleTrackIdFromLaneIndex();
+		#endif
+
+		// Matches first line of kStackowiakSampleSet0
+		// TODO: could be optimised further by just setting sign bit on 0.5.
+		float2 SampleOffset = float2(
+			SampleTrackId & 0x1 ? 0.5 : -0.5,
+			SampleTrackId & 0x2 ? 0.5 : -0.5);
+
+		return KernelConfig.BufferUV + SampleOffset * KernelConfig.BufferSizeAndInvSize.xy;
+	}
+	else if (SAMPLE_SET_HEXAWEB)
+	{
+		return float2(0.0, 0.0); // Impossible because of random offset certainely needed using this..
+	}
+
+	return KernelConfig.BufferUV;
+}
+
+float2 ComputeRefScreenPos(FSSDKernelConfig KernelConfig)
+{
+	float2 RefBufferUV = ComputeRefBufferUV(KernelConfig);
+
+	// TODO: reduce VALU here.
+	float2 RefViewportUV = BufferUVToViewportUV(RefBufferUV);
+	float2 RefScreenPos = ViewportUVToScreenPos(RefViewportUV);
+
+	return RefScreenPos;
+}
+
+float3 ComputeWorldVectorFromNeighborToRef(
+	FSSDKernelConfig KernelConfig,
+	FSSDSampleSceneInfos NeighborSceneMetadata,
+	float2 NeighborBufferUV)
+#if CONFIG_NEIGHBOR_TO_REF_COMPUTATION == NEIGHBOR_TO_REF_CACHE_WORLD_POSITION
+{
+	float3 NeighborToRefWorldVector = KernelConfig.RefSceneMetadata.TranslatedWorldPosition - NeighborSceneMetadata.TranslatedWorldPosition;
+	return NeighborToRefWorldVector;
+}
+#elif CONFIG_NEIGHBOR_TO_REF_COMPUTATION == NEIGHBOR_TO_REF_LOWEST_VGPR_PRESSURE
+{
+	float2 RefScreenPos = ComputeRefScreenPos(KernelConfig);
+	float3 RefClipPosition = float3(RefScreenPos * (View.ViewToClip[3][3] < 1.0f ? KernelConfig.RefSceneMetadata.WorldDepth : 1.0f), KernelConfig.RefSceneMetadata.WorldDepth);
+
+	// TODO: reduce VALU here.
+	float2 NeighborViewportUV = BufferUVToViewportUV(NeighborBufferUV);
+	float2 NeighborScreenPos = ViewportUVToScreenPos(NeighborViewportUV);
+
+	float3 NeighborClipPosition = float3(NeighborScreenPos * (View.ViewToClip[3][3] < 1.0f ? NeighborSceneMetadata.WorldDepth : 1.0f), NeighborSceneMetadata.WorldDepth);
+
+	float3 NeighborToRefWorldVector = mul(float4(RefClipPosition - NeighborClipPosition, 0), View.ScreenToTranslatedWorld).xyz;
+	return NeighborToRefWorldVector;
+}
+#else
+	#error Unknown settings
+#endif
+
+
 //------------------------------------------------------- SHARED SAMPLING
 
 FSSDSignalSample TransformSignalSampleForAccumulation(
@@ -352,10 +445,18 @@ void AccumulateSampledMultiplexedSignals(
 			FinalWorldBluringDistance = min(FinalWorldBluringDistance, KernelConfig.RefBilateralDistance[SignalMultiplexId]);
 		}
 	
+		// Compute the vector from neighbor to reference in the most optimal way.
+		float3 NeighborToRefWorldVector = ComputeWorldVectorFromNeighborToRef(
+			KernelConfig,
+			SampleSceneMetadata,
+			SampleBufferUV);
+
+		// Compute the weight to be applied to do bilateral rejection.
 		float BilateralWeight = ComputeBilateralWeight(
 			FinalWorldBluringDistance,
 			KernelConfig.RefSceneMetadata,
-			SampleSceneMetadata);
+			SampleSceneMetadata,
+			NeighborToRefWorldVector);
 		
 		float RatioEstimatorWeight = GetRatioEstimatorWeight(
 			KernelConfig.RefSceneMetadata, SampleSceneMetadata, Sample, NeighborPixelCoord);
@@ -870,10 +971,15 @@ void ConvolveStackowiakKernel(
 
 				const bool bIsKernelCenterSample = SampleId == 0 && (SampleBatchId == (kSamplingBatchSize - 1));
 				
-				float2 SampleOffset = kStackowiakSampleSet0[kStackowiakSampleSetCount * SampleId + KernelConfig.SampleTrackId];
+				uint SampleTrackId = KernelConfig.SampleTrackId;
+				#if CONFIG_VGPR_FREE_SAMPLE_TRACK_ID
+					SampleTrackId = GetSampleTrackIdFromLaneIndex();
+				#endif
+
+				float2 SampleOffset = kStackowiakSampleSet0[kStackowiakSampleSetCount * SampleId + SampleTrackId];
 				if (KernelConfig.SampleSubSetId == 1)
 				{
-					SampleOffset = kStackowiakSampleSet1[kStackowiakSampleSetCount * SampleId + KernelConfig.SampleTrackId];
+					SampleOffset = kStackowiakSampleSet1[kStackowiakSampleSetCount * SampleId + SampleTrackId];
 				}
 				
 				float2 SampleBufferUV = KernelConfig.BufferUV + (SampleOffset * KernelConfig.KernelSpreadFactor) * KernelConfig.BufferSizeAndInvSize.zw;
@@ -933,10 +1039,15 @@ void ConvolveStackowiakKernel(
 		{
 			const bool bIsKernelCenterSample = SampleId == 0;
 
-			float2 SampleOffset = kStackowiakSampleSet0[kStackowiakSampleSetCount * SampleId + KernelConfig.SampleTrackId];
+			uint SampleTrackId = KernelConfig.SampleTrackId;
+			#if CONFIG_VGPR_FREE_SAMPLE_TRACK_ID
+				SampleTrackId = GetSampleTrackIdFromLaneIndex();
+			#endif
+
+			float2 SampleOffset = kStackowiakSampleSet0[kStackowiakSampleSetCount * SampleId + SampleTrackId];
 			if (KernelConfig.SampleSubSetId == 1)
 			{
-				SampleOffset = kStackowiakSampleSet1[kStackowiakSampleSetCount * SampleId + KernelConfig.SampleTrackId];
+				SampleOffset = kStackowiakSampleSet1[kStackowiakSampleSetCount * SampleId + SampleTrackId];
 			}
 			
 			float2 SampleBufferUV = KernelConfig.BufferUV + (SampleOffset * KernelConfig.KernelSpreadFactor) * KernelConfig.BufferSizeAndInvSize.zw;
