Deletes old ISOLATE attribute that was no longer used.

diff --git a/Engine/Plugins/FX/Niagara/Shaders/Private/NiagaraRibbonVertexFactory.ush b/Engine/Plugins/FX/Niagara/Shaders/Private/NiagaraRibbonVertexFactory.ush
--- a/Engine/Plugins/FX/Niagara/Shaders/Private/NiagaraRibbonVertexFactory.ush
+++ b/Engine/Plugins/FX/Niagara/Shaders/Private/NiagaraRibbonVertexFactory.ush
@@ -397,7 +397,6 @@ FInterpVtxData GetInterpVtxData(FVertexFactoryInput Input)
 	}
 
 #if USE_PARTICLE_INTERPOLATION
-	ISOLATE
 	{
 		InterpVtxData.Alpha = (float)((Input.InterpVtxId / 2) % NiagaraRibbonVF.InterpCount) * NiagaraRibbonVF.OneOverInterpCount;
 	}
@@ -584,7 +583,6 @@ void GetNiagaraParticleTangents(in FInterpVtxData InterpVtxData, in out FVertexF
 		Intermediates.SideVector = Intermediates.Right;
 	}
 
-	ISOLATE
 	{
 		Intermediates.Up = cross(Intermediates.Right, Intermediates.Facing);
 	}
@@ -606,7 +604,6 @@ void GetNiagaraParticleTangents(in FInterpVtxData InterpVtxData, in out FVertexF
 
 float4 CalcWorldPosition(FVertexFactoryInput Input, FVertexFactoryIntermediates Intermediates)
 {
-	ISOLATE
 	{
 		// expand and align
 		int Sgn = (((int)Input.InterpVtxId & 0x1) * 2) - 1;
diff --git a/Engine/Shaders/Private/BasePassTessellationShaders.usf b/Engine/Shaders/Private/BasePassTessellationShaders.usf
--- a/Engine/Shaders/Private/BasePassTessellationShaders.usf
+++ b/Engine/Shaders/Private/BasePassTessellationShaders.usf
@@ -60,7 +60,6 @@ FBasePassVSToPS PassFinalizeTessellationOutput(FBasePassVSToDS Interpolants, flo
 	(FSharedBasePassInterpolants)O.BasePassInterpolants = (FSharedBasePassInterpolants)Interpolants.BasePassInterpolants;
 
 	// Transform position to clip-space
-	ISOLATE
 	{
 		O.Position = mul(WorldPosition, ResolvedView.TranslatedWorldToClip);
 #if WRITES_VELOCITY_TO_GBUFFER
diff --git a/Engine/Shaders/Private/BasePassVertexShader.usf b/Engine/Shaders/Private/BasePassVertexShader.usf
--- a/Engine/Shaders/Private/BasePassVertexShader.usf
+++ b/Engine/Shaders/Private/BasePassVertexShader.usf
@@ -60,7 +60,6 @@ void Main(
 	// Isolate instructions used for world position offset
 	// As these cause the optimizer to generate different position calculating instructions in each pass, resulting in self-z-fighting.
 	// This is only necessary for shaders used in passes that have depth testing enabled.
-	ISOLATE
 	{
 		WorldPosition.xyz += GetMaterialWorldPositionOffset(VertexParameters);
 	}
@@ -73,7 +72,6 @@ void Main(
 		Output.BasePassInterpolants.WorldPositionExcludingWPO = WorldPositionExcludingWPO.xyz;
 	#endif
 #else
-	ISOLATE
 	{
 		float4 RasterizedWorldPosition = VertexFactoryGetRasterizedWorldPosition(Input, VFIntermediates, WorldPosition);
 	#if ODS_CAPTURE
diff --git a/Engine/Shaders/Private/DebugViewModeVertexShader.usf b/Engine/Shaders/Private/DebugViewModeVertexShader.usf
--- a/Engine/Shaders/Private/DebugViewModeVertexShader.usf
+++ b/Engine/Shaders/Private/DebugViewModeVertexShader.usf
@@ -62,7 +62,6 @@ MaterialFloat StoreTexSample(in out FTexCoordScalesParams Params, float4 C, int
 		FDebugVSToPS Output;
 		
 		// Finally, transform position to clip-space
-		ISOLATE
 		{
 
 			Output.Position = mul(WorldPosition, ResolvedView.TranslatedWorldToClip);
@@ -136,7 +135,6 @@ void Main(
 	// Isolate instructions used for world position offset
 	// As these cause the optimizer to generate different position calculating instructions in each pass, resulting in self-z-fighting.
 	// This is only necessary for shaders used in passes that have depth testing enabled.
-	ISOLATE
 	{
 		WorldPos.xyz += GetMaterialWorldPositionOffset(VertexParameters);
 	}
@@ -146,7 +144,6 @@ void Main(
 	Output.Position = WorldPos;
 	Output.FactoryInterpolants = VertexFactoryGetInterpolantsVSToDS(Input, VFIntermediates, VertexParameters);
 #else
-	ISOLATE
 	{
 		float4 RasterizedWorldPosition = VertexFactoryGetRasterizedWorldPosition(Input, VFIntermediates, WorldPos);
 		Output.Position = mul(RasterizedWorldPosition, ResolvedView.TranslatedWorldToClip);
diff --git a/Engine/Shaders/Private/Definitions.usf b/Engine/Shaders/Private/Definitions.usf
--- a/Engine/Shaders/Private/Definitions.usf
+++ b/Engine/Shaders/Private/Definitions.usf
@@ -140,7 +140,6 @@
 #define LIGHTMAP_VT_ENABLED								0
 #endif
 
-#define ISOLATE
 #define NOEXPRESSIONOPTIMIZATIONS						
 
 #ifndef HAS_INVERTED_Z_BUFFER
diff --git a/Engine/Shaders/Private/DepthOnlyVertexShader.usf b/Engine/Shaders/Private/DepthOnlyVertexShader.usf
--- a/Engine/Shaders/Private/DepthOnlyVertexShader.usf
+++ b/Engine/Shaders/Private/DepthOnlyVertexShader.usf
@@ -68,7 +68,6 @@ struct FDepthOnlyVSToPS
 		#endif
 
 		// Finally, transform position to clip-space
-		ISOLATE
 		{
 			O.Position = mul(WorldPosition, ResolvedView.TranslatedWorldToClip);
 		}
@@ -128,7 +127,6 @@ void Main(
 	// Isolate instructions used for world position offset
 	// As these cause the optimizer to generate different position calculating instructions in each pass, resulting in self-z-fighting.
 	// This is only necessary for shaders used in passes that have depth testing enabled.
-	ISOLATE
 	{
 		WorldPos.xyz += GetMaterialWorldPositionOffset(VertexParameters);
 	}
@@ -137,7 +135,6 @@ void Main(
 	// Transformation is done in Domain shader when tessellating
 	Output.Position = WorldPos; 
 #else
-	ISOLATE
 	{
 		float4 RasterizedWorldPosition = VertexFactoryGetRasterizedWorldPosition(Input, VFIntermediates, WorldPos);
 	#if ODS_CAPTURE
diff --git a/Engine/Shaders/Private/DistortAccumulateVS.usf b/Engine/Shaders/Private/DistortAccumulateVS.usf
--- a/Engine/Shaders/Private/DistortAccumulateVS.usf
+++ b/Engine/Shaders/Private/DistortAccumulateVS.usf
@@ -112,7 +112,6 @@ void Main(
 	// Isolate instructions used for world position offset on xbox 360, 
 	// As these cause the optimizer to generate different position calculating instructions in each pass, resulting in self-z-fighting.
 	// This is only necessary for shaders used in passes that have depth testing enabled.
-	ISOLATE
 	{
 		WorldPosition.xyz += GetMaterialWorldPositionOffset(VertexParameters);
 	}
diff --git a/Engine/Shaders/Private/HitProxyVertexShader.usf b/Engine/Shaders/Private/HitProxyVertexShader.usf
--- a/Engine/Shaders/Private/HitProxyVertexShader.usf
+++ b/Engine/Shaders/Private/HitProxyVertexShader.usf
@@ -79,7 +79,6 @@ struct FHitProxyVSToPS
 		#endif
 
 		// Finally, transform position to clip-space
-		ISOLATE
 		{
 			O.Position = mul(WorldPosition, View.TranslatedWorldToClip);
 		}
@@ -132,7 +131,6 @@ void Main(
 	Output.FactoryInterpolants = VertexFactoryGetInterpolantsVSToDS(Input, VFIntermediates, VertexParameters);
 #else
 
-	ISOLATE
 	{
 		float4 RasterizedWorldPosition = VertexFactoryGetRasterizedWorldPosition(Input, VFIntermediates, WorldPosition);
 		Output.Position = mul(RasterizedWorldPosition, View.TranslatedWorldToClip);
diff --git a/Engine/Shaders/Private/LightMapDensityShader.usf b/Engine/Shaders/Private/LightMapDensityShader.usf
--- a/Engine/Shaders/Private/LightMapDensityShader.usf
+++ b/Engine/Shaders/Private/LightMapDensityShader.usf
@@ -60,7 +60,6 @@ struct FLightMapDensityVSToPS
 		O.FactoryInterpolants = VertexFactoryAssignInterpolants(Interpolants.FactoryInterpolants);
 		
 		// Finally, transform position to clip-space
-		ISOLATE
 		{
 			O.Position = mul(WorldPosition, ResolvedView.TranslatedWorldToClip);
 		}
@@ -107,7 +106,6 @@ void MainVertexShader(
 #else // !USING_TESSELLATION
 
 	Output.WorldPosition = WorldPosition;
-	ISOLATE
 	{
 		float4 RasterizedWorldPosition = VertexFactoryGetRasterizedWorldPosition(Input, VFIntermediates, Output.WorldPosition);
 		Output.Position = mul(RasterizedWorldPosition, ResolvedView.TranslatedWorldToClip);
diff --git a/Engine/Shaders/Private/MeshDecals.usf b/Engine/Shaders/Private/MeshDecals.usf
--- a/Engine/Shaders/Private/MeshDecals.usf
+++ b/Engine/Shaders/Private/MeshDecals.usf
@@ -112,7 +112,6 @@ void MainVS(
 	// Isolate instructions used for world position offset on xbox 360, 
 	// As these cause the optimizer to generate different position calculating instructions in each pass, resulting in self-z-fighting.
 	// This is only necessary for shaders used in passes that have depth testing enabled.
-	ISOLATE
 	{
 		WorldPosition.xyz += GetMaterialWorldPositionOffset(VertexParameters);
 	}
diff --git a/Engine/Shaders/Private/PositionOnlyDepthVertexShader.usf b/Engine/Shaders/Private/PositionOnlyDepthVertexShader.usf
--- a/Engine/Shaders/Private/PositionOnlyDepthVertexShader.usf
+++ b/Engine/Shaders/Private/PositionOnlyDepthVertexShader.usf
@@ -39,7 +39,6 @@ void Main(
 
 	float4 WorldPos = VertexFactoryGetWorldPosition(Input);
 
-	ISOLATE
 	{
 	#if ODS_CAPTURE
 		float3 ODS = OffsetODS(WorldPos.xyz, ResolvedView.TranslatedWorldCameraOrigin.xyz, ResolvedView.StereoIPD);
diff --git a/Engine/Shaders/Private/VelocityShader.usf b/Engine/Shaders/Private/VelocityShader.usf
--- a/Engine/Shaders/Private/VelocityShader.usf
+++ b/Engine/Shaders/Private/VelocityShader.usf
@@ -84,7 +84,6 @@ struct FVelocityVSToPS
 		
 		O.Interps.FactoryInterpolants = VertexFactoryAssignInterpolants(Interpolants.FactoryInterpolants);
 
-		ISOLATE
 		{
 			O.Position = mul(WorldPosition, ResolvedView.TranslatedWorldToClip);
 		}
@@ -142,7 +141,6 @@ void MainVertexShader(
 	// Isolate instructions used for world position offset on xbox 360, 
 	// As these cause the optimizer to generate different position calculating instructions in each pass, resulting in self-z-fighting.
 	// This is only necessary for shaders used in passes that have depth testing enabled.
-	ISOLATE
 	{
 		TranslatedWorldPosition.xyz += GetMaterialWorldPositionOffset(VertexParameters);
 	}
@@ -172,7 +170,6 @@ void MainVertexShader(
 	#endif
 
 #if !USING_TESSELLATION
-		ISOLATE
 		{
 #endif
 
diff --git a/Engine/Shaders/Private/VirtualTextureMaterial.usf b/Engine/Shaders/Private/VirtualTextureMaterial.usf
--- a/Engine/Shaders/Private/VirtualTextureMaterial.usf
+++ b/Engine/Shaders/Private/VirtualTextureMaterial.usf
@@ -28,7 +28,6 @@ void MainVS(
 	float3x3 TangentToLocal = VertexFactoryGetTangentToLocal(Input, VFIntermediates);	
 	FMaterialVertexParameters VertexParameters = GetMaterialVertexParameters(Input, VFIntermediates, WorldPosition.xyz, TangentToLocal);
 
-	ISOLATE
 	{
 		float4 RasterizedWorldPosition = VertexFactoryGetRasterizedWorldPosition(Input, VFIntermediates, WorldPosition);
 		float4 ClipSpacePosition = INVARIANT(mul(RasterizedWorldPosition, ResolvedView.TranslatedWorldToClip));
