Reimplements ISOLATE shader attribute and use it for denoiser.

diff --git a/Engine/Shaders/Private/ScreenSpaceDenoise/SSDCommon.ush b/Engine/Shaders/Private/ScreenSpaceDenoise/SSDCommon.ush
--- a/Engine/Shaders/Private/ScreenSpaceDenoise/SSDCommon.ush
+++ b/Engine/Shaders/Private/ScreenSpaceDenoise/SSDCommon.ush
@@ -195,7 +195,7 @@ FSSDSampleSceneInfos FetchPrevSceneInfos(float2 ScreenPosition, float2 BufferUV)
 	return OutMetadata;
 }
 
-// ISOLATE
+ISOLATE
 float3 GetWorldNormal(FSSDSampleSceneInfos SceneMetadata)
 #if CONFIG_COMPRESS_WORLD_NORMAL
 {
diff --git a/Engine/Shaders/Private/ScreenSpaceDenoise/SSDSpatialKernel.ush b/Engine/Shaders/Private/ScreenSpaceDenoise/SSDSpatialKernel.ush
--- a/Engine/Shaders/Private/ScreenSpaceDenoise/SSDSpatialKernel.ush
+++ b/Engine/Shaders/Private/ScreenSpaceDenoise/SSDSpatialKernel.ush
@@ -217,7 +217,7 @@ static const float kWaveletFilterWeights5x5[] = { 3.0 / 8.0, 1.0 / 4.0, 1.0 / 16
 
 #if CONFIG_VGPR_FREE_SAMPLE_TRACK_ID
 
-// ISOLATE
+ISOLATE
 uint GetSampleTrackIdFromLaneIndex()
 {
 	uint LaneId = WaveGetLaneIndex();
diff --git a/Engine/Shaders/Public/Platform.ush b/Engine/Shaders/Public/Platform.ush
--- a/Engine/Shaders/Public/Platform.ush
+++ b/Engine/Shaders/Public/Platform.ush
@@ -510,6 +510,11 @@
 #define StrongTypedBuffer Buffer
 #endif
 
+#ifndef ISOLATE
+#define ISOLATE
+#endif
+
+
 // ---------------------------------------------------- Interpolator attribute fallbacks
 
 #ifndef COMPRESSED_16_FLOAT
