Fixes shader compilation failure on D3D in the denoiser

diff --git a/Engine/Shaders/Private/ScreenSpaceDenoise/SSDDefinitions.ush b/Engine/Shaders/Private/ScreenSpaceDenoise/SSDDefinitions.ush
--- a/Engine/Shaders/Private/ScreenSpaceDenoise/SSDDefinitions.ush
+++ b/Engine/Shaders/Private/ScreenSpaceDenoise/SSDDefinitions.ush
@@ -2,6 +2,9 @@
 
 #pragma once
 
+#include "/Engine/Public/Platform.ush"
+
+
 //------------------------------------------------------- COMPILER CONFIG
 
 // Generate vector truncation warnings to errors.
diff --git a/Engine/Shaders/Private/ScreenSpaceDenoise/SSDSpatialAccumulation.usf b/Engine/Shaders/Private/ScreenSpaceDenoise/SSDSpatialAccumulation.usf
--- a/Engine/Shaders/Private/ScreenSpaceDenoise/SSDSpatialAccumulation.usf
+++ b/Engine/Shaders/Private/ScreenSpaceDenoise/SSDSpatialAccumulation.usf
@@ -384,7 +384,7 @@
 
 //------------------------------------------------------- LATE SETUP
 
-#if COMPILER_SUPPORTS_WAVE_ONCE || 1
+#if COMPILER_SUPPORTS_WAVE_ONCE
 	#if CONFIG_UPSCALE
 		#define CONFIG_VGPR_FREE_SAMPLE_TRACK_ID 2
 	#else
