Implements generic public buffer encoding and decoding of the spherical harmonic denoiser

diff --git a/Engine/Shaders/Private/ScreenSpaceDenoise/SSDCommon.ush b/Engine/Shaders/Private/ScreenSpaceDenoise/SSDCommon.ush
--- a/Engine/Shaders/Private/ScreenSpaceDenoise/SSDCommon.ush
+++ b/Engine/Shaders/Private/ScreenSpaceDenoise/SSDCommon.ush
@@ -62,28 +62,6 @@ float SafeRcp(float x)
 	return 0;
 }
 
-uint Compress2Floats(float2 v)
-{
-	// f32tof16() is full rate on GCN.
-	return f32tof16(v.x) | (f32tof16(v.y) << 16);
-}
-
-float2 Uncompress2Floats(uint v)
-{
-	// f16tof32() is full rate on GCN.
-	return float2(f16tof32(v), f16tof32(v >> 16));
-}
-
-uint2 Compress4Floats(float4 v)
-{
-	return uint2(Compress2Floats(v.xy), Compress2Floats(v.zw));
-}
-
-float4 Uncompress4Floats(uint2 v)
-{
-	return float4(Uncompress2Floats(v.x), Uncompress2Floats(v.y));
-}
-
 uint CompressN(float3 N)
 {
 	// matches 8bits GBuffer A to be lossless.
diff --git a/Engine/Shaders/Private/ScreenSpaceDenoise/SSDDefinitions.ush b/Engine/Shaders/Private/ScreenSpaceDenoise/SSDDefinitions.ush
--- a/Engine/Shaders/Private/ScreenSpaceDenoise/SSDDefinitions.ush
+++ b/Engine/Shaders/Private/ScreenSpaceDenoise/SSDDefinitions.ush
@@ -208,6 +208,11 @@
 	#define NEIGHBOR_TO_REF_LOWEST_VGPR_PRESSURE 1
 
 
+/** Type to use for texture sampling. */
+	#define SIGNAL_TEXTURE_TYPE_FLOAT4 0
+	#define SIGNAL_TEXTURE_TYPE_UINT2 1
+
+
 //------------------------------------------------------- GLOBAL CONFIGURATION OF ALL DENOISER SHADERS
 
 /** Do normal based anysotropy. */
@@ -224,9 +229,6 @@
 // TODO: the shader compiler is not doing the right thing.
 #define CONFIG_SGPR_HINT_OPTIMIZATION 0
 
-/** Manually format texel in the shader to reduce VGPR pressure with overlapped texture fetched. */
-#define CONFIG_MANUAL_TEXEL_FORMATTING 1
-
 /** Configures whether the world normal should be compressed in FSSDSampleSceneInfos. */
 #define CONFIG_COMPRESS_WORLD_NORMAL 1
 
diff --git a/Engine/Shaders/Private/ScreenSpaceDenoise/SSDInjest.usf b/Engine/Shaders/Private/ScreenSpaceDenoise/SSDInjest.usf
--- a/Engine/Shaders/Private/ScreenSpaceDenoise/SSDInjest.usf
+++ b/Engine/Shaders/Private/ScreenSpaceDenoise/SSDInjest.usf
@@ -42,22 +42,22 @@
 	#error Missing CONFIG_INPUT_TEXTURE_COUNT
 #endif
 
-Texture2D SignalInput_Textures_0;
+FSSDTexture2D SignalInput_Textures_0;
 
 #if CONFIG_INPUT_TEXTURE_COUNT > 1
-Texture2D SignalInput_Textures_1;
+FSSDTexture2D SignalInput_Textures_1;
 #else
 #define SignalInput_Textures_1 SignalInput_Textures_0
 #endif
 
 #if CONFIG_INPUT_TEXTURE_COUNT > 2
-Texture2D SignalInput_Textures_2;
+FSSDTexture2D SignalInput_Textures_2;
 #else
 #define SignalInput_Textures_2 SignalInput_Textures_0
 #endif
 
 #if CONFIG_INPUT_TEXTURE_COUNT > 3
-Texture2D SignalInput_Textures_3;
+FSSDTexture2D SignalInput_Textures_3;
 #else
 #define SignalInput_Textures_3 SignalInput_Textures_0
 #endif
@@ -67,22 +67,22 @@ Texture2D SignalInput_Textures_3;
 	#error Missing CONFIG_OUTPUT_TEXTURE_COUNT
 #endif
 
-RWTexture2D<float4>	SignalOutput_UAVs_0;
+FSSDRWTexture2D SignalOutput_UAVs_0;
 
 #if CONFIG_OUTPUT_TEXTURE_COUNT > 1
-RWTexture2D<float4>	SignalOutput_UAVs_1;
+FSSDRWTexture2D	SignalOutput_UAVs_1;
 #else
 #define SignalOutput_UAVs_1 SignalOutput_UAVs_0
 #endif
 
 #if CONFIG_OUTPUT_TEXTURE_COUNT > 2
-RWTexture2D<float4>	SignalOutput_UAVs_2;
+FSSDRWTexture2D	SignalOutput_UAVs_2;
 #else
 #define SignalOutput_UAVs_2 SignalOutput_UAVs_0
 #endif
 
 #if CONFIG_OUTPUT_TEXTURE_COUNT > 3
-RWTexture2D<float4>	SignalOutput_UAVs_3;
+FSSDRWTexture2D	SignalOutput_UAVs_3;
 #else
 #define SignalOutput_UAVs_3 SignalOutput_UAVs_0
 #endif
diff --git a/Engine/Shaders/Private/ScreenSpaceDenoise/SSDPublicBufferEncoding.ush b/Engine/Shaders/Private/ScreenSpaceDenoise/SSDPublicBufferEncoding.ush
new file mode 100644
--- /dev/null
+++ b/Engine/Shaders/Private/ScreenSpaceDenoise/SSDPublicBufferEncoding.ush
@@ -0,0 +1,117 @@
+// Copyright 1998-2019 Epic Games, Inc. All Rights Reserved.
+
+#pragma once
+
+#include "../Common.ush"
+#include "../SHCommon.ush"
+
+
+/** Compile time settings to encode buffer to the denoiser. */
+	#define SSD_ENCODE_CLAMP_COLOR 0x1
+	#define SSD_ENCODE_NORMALIZE_MISS 0x2
+	#define SSD_ENCODE_NORMALIZE_COLOR 0x4
+	#define SSD_ENCODE_NORMALIZE (SSD_ENCODE_NORMALIZE_MISS | SSD_ENCODE_NORMALIZE_COLOR)
+
+/** Compile time settings to decode buffer from the denoiser. */
+	#define SSD_DECODE_NORMALIZE 0x1
+
+
+uint Compress2Floats(float2 v)
+{
+	// f32tof16() is full rate on GCN.
+	return f32tof16(v.x) | (f32tof16(v.y) << 16);
+}
+
+float2 Uncompress2Floats(uint v)
+{
+	// f16tof32() is full rate on GCN.
+	return float2(f16tof32(v), f16tof32(v >> 16));
+}
+
+uint2 Compress4Floats(float4 v)
+{
+	return uint2(Compress2Floats(v.xy), Compress2Floats(v.zw));
+}
+
+float4 Uncompress4Floats(uint2 v)
+{
+	return float4(Uncompress2Floats(v.x), Uncompress2Floats(v.y));
+}
+
+float4 f32x4tof16(float4 OutRawSample)
+{
+	return float4(asfloat(Compress2Floats(OutRawSample.xy)), asfloat(Compress2Floats(OutRawSample.zw)), 0, 0);
+}
+
+
+// Encode a SIGNAL_BUFFER_LAYOUT_DIFFUSE_INDIRECT_HARMONIC formated buffer.
+void EncodeDiffuseSphericalHarmonicTexel(
+	float SampleCount,
+	float MissCount,
+	FTwoBandSHVectorRGB ColorSH,
+	const uint Options,
+	out uint2 OutRawSample[4])
+{
+	float2 UncompressedSample3;
+	UncompressedSample3.x = SampleCount;
+
+	float NormalizationFactor = SampleCount > 0.0 ? rcp(SampleCount) : 0.0;
+	if (Options & SSD_ENCODE_NORMALIZE_MISS)
+	{
+		UncompressedSample3.y = MissCount * NormalizationFactor;
+	}
+	else
+	{
+		UncompressedSample3.y = MissCount;
+	}
+
+	float ColorNormalizationFactor = 1.0;
+	if (Options & SSD_ENCODE_NORMALIZE_COLOR)
+	{
+		ColorNormalizationFactor = NormalizationFactor;
+	}
+
+	float4 UncompressedSample0 = ColorSH.R.V * ColorNormalizationFactor;
+	float4 UncompressedSample1 = ColorSH.G.V * ColorNormalizationFactor;
+	float4 UncompressedSample2 = ColorSH.B.V * ColorNormalizationFactor;
+
+	if (Options & SSD_ENCODE_CLAMP_COLOR)
+	{
+		UncompressedSample0 = min(UncompressedSample0, MaxHalfFloat);
+		UncompressedSample1 = min(UncompressedSample1, MaxHalfFloat);
+		UncompressedSample2 = min(UncompressedSample2, MaxHalfFloat);
+	}
+
+	OutRawSample[0] = Compress4Floats(UncompressedSample0);
+	OutRawSample[1] = Compress4Floats(UncompressedSample1);
+	OutRawSample[2] = Compress4Floats(UncompressedSample2);
+	OutRawSample[3] = uint2(Compress2Floats(UncompressedSample3), 0);
+}
+
+// Decodes a SIGNAL_BUFFER_LAYOUT_DIFFUSE_INDIRECT_HARMONIC formated buffer.
+void DecodeDiffuseSphericalHarmonicTexel(
+	uint2 RawSample[4],
+	const uint Options,
+	out float OutSampleCount,
+	out float OutMissCount,
+	out FTwoBandSHVectorRGB OutColorSH)
+{
+	float2 UncompressedSample3 = Uncompress2Floats(RawSample[3].x);
+
+	float Multiplier;
+	if (Options & SSD_DECODE_NORMALIZE)
+	{
+		Multiplier = 1.0;
+		OutSampleCount = UncompressedSample3.x > 0.0 ? 1.0 : 0.0;
+	}
+	else
+	{
+		Multiplier = UncompressedSample3.x;
+		OutSampleCount = Multiplier;
+	}
+
+	OutMissCount = UncompressedSample3.y * Multiplier;
+	OutColorSH.R.V = Uncompress4Floats(RawSample[0]) * Multiplier;
+	OutColorSH.G.V = Uncompress4Floats(RawSample[1]) * Multiplier;
+	OutColorSH.B.V = Uncompress4Floats(RawSample[2]) * Multiplier;
+}
diff --git a/Engine/Shaders/Private/ScreenSpaceDenoise/SSDSignalBufferEncoding.ush b/Engine/Shaders/Private/ScreenSpaceDenoise/SSDSignalBufferEncoding.ush
--- a/Engine/Shaders/Private/ScreenSpaceDenoise/SSDSignalBufferEncoding.ush
+++ b/Engine/Shaders/Private/ScreenSpaceDenoise/SSDSignalBufferEncoding.ush
@@ -3,6 +3,45 @@
 #pragma once
 
 #include "SSDSignalCore.ush"
+#include "SSDPublicBufferEncoding.ush"
+
+/** Whether the color should be clamped when encoding signal. */
+#define CONFIG_ENCODING_CLAMP_COLOR 1
+
+/** Selects the type that should be used when sampling a buffer */
+#ifndef CONFIG_SIGNAL_TEXTURE_TYPE
+	#define CONFIG_SIGNAL_TEXTURE_TYPE SIGNAL_TEXTURE_TYPE_FLOAT4
+#endif
+
+
+#if CONFIG_SIGNAL_TEXTURE_TYPE == SIGNAL_TEXTURE_TYPE_FLOAT4
+	#define FSSDRawSample float4
+	#define FSSDTexture2D Texture2D
+	#define FSSDRWTexture2D RWTexture2D<float4>
+
+#elif CONFIG_SIGNAL_TEXTURE_TYPE == SIGNAL_TEXTURE_TYPE_UINT2
+	#define FSSDRawSample uint2
+	#define FSSDTexture2D Texture2D<uint2>
+	#define FSSDRWTexture2D RWTexture2D<uint2>
+
+#else
+	#error Unknown type for a signal texture.
+#endif
+
+
+
+float3 ClampColorForEncoding(float3 Color)
+{
+	#if CONFIG_ENCODING_CLAMP_COLOR
+		Color = min(Color, MaxHalfFloat);
+	#endif
+	return Color;
+}
+
+float4 ClampColorForEncoding(float4 Color)
+{
+	return float4(ClampColorForEncoding(Color.rgb), Color.a);
+}
 
 
 /** Decode input signal sample. */
@@ -10,8 +49,9 @@ void DecodeMultiplexedSignals(
 	const uint SignalBufferLayout,
 	const uint MultiplexedSampleId,
 	const bool bNormalizeSample,
-	float4 RawSample[MAX_MULTIPLEXED_TEXTURES],
+	FSSDRawSample RawSample[MAX_MULTIPLEXED_TEXTURES],
 	out FSSDSignalArray OutSamples)
+#if CONFIG_SIGNAL_TEXTURE_TYPE == SIGNAL_TEXTURE_TYPE_FLOAT4
 {
 	OutSamples = CreateSignalArrayFromScalarValue(0.0);
 
@@ -253,46 +293,36 @@ void DecodeMultiplexedSignals(
 		OutSamples.Array[2].SceneColor.x = OutSamples.Array[2].SampleCount * RawSample[1].b;
 	}
 	#endif
+} // DecodeMultiplexedSignals()
+#elif CONFIG_SIGNAL_TEXTURE_TYPE == SIGNAL_TEXTURE_TYPE_UINT2
+{
+	OutSamples = CreateSignalArrayFromScalarValue(0.0);
+
 	#if CONFIG_SIGNAL_COLOR_SH
-	else if (SignalBufferLayout == SIGNAL_BUFFER_LAYOUT_DIFFUSE_INDIRECT_HARMONIC)
+	if (SignalBufferLayout == SIGNAL_BUFFER_LAYOUT_DIFFUSE_INDIRECT_HARMONIC)
 	{
-		#if CONFIG_MANUAL_TEXEL_FORMATTING
-		{
-			 RawSample[0] = Uncompress4Floats(uint2(asuint(RawSample[0].x), asuint(RawSample[0].y)));
-			 RawSample[1] = Uncompress4Floats(uint2(asuint(RawSample[1].x), asuint(RawSample[1].y)));
-			 RawSample[2] = Uncompress4Floats(uint2(asuint(RawSample[2].x), asuint(RawSample[2].y)));
-			 RawSample[3].xy = Uncompress2Floats(asuint(RawSample[3].x));
-		}
-		#endif
+		const uint DecodeOptions = bNormalizeSample ? SSD_DECODE_NORMALIZE : 0x0;
 
-		float SampleCount;
-		if (bNormalizeSample)
-		{
-			SampleCount = 1.0;
-			OutSamples.Array[0].SampleCount = RawSample[3].x > 0.0 ? 1.0 : 0.0;
-		}
-		else
-		{
-			SampleCount = RawSample[3].x;
-			OutSamples.Array[0].SampleCount = SampleCount;
-		}
-
-		OutSamples.Array[0].MissCount = RawSample[3].y * SampleCount;
-		OutSamples.Array[0].ColorSH.R.V = RawSample[0] * SampleCount;
-		OutSamples.Array[0].ColorSH.G.V = RawSample[1] * SampleCount;
-		OutSamples.Array[0].ColorSH.B.V = RawSample[2] * SampleCount;
+		DecodeDiffuseSphericalHarmonicTexel(
+			RawSample,
+			DecodeOptions,
+			/* out */ OutSamples.Array[0].SampleCount,
+			/* out */ OutSamples.Array[0].MissCount,
+			/* out */ OutSamples.Array[0].ColorSH);
 
 		// TODO: SRV on RG F32 and manually decode with f16tof32(), going to reduce VGPR pressure between fetch and result.
 	}
 	#endif
 } // DecodeMultiplexedSignals()
+#endif
 
 /** Encode output signal sample. */
 void EncodeMultiplexedSignals(
 	const uint SignalBufferLayout, const uint MultiplexCount,
 	FSSDSignalSample Sample[MAX_SIGNAL_MULTIPLEXING],
-	out float4 OutRawSample[MAX_MULTIPLEXED_TEXTURES],
+	out FSSDRawSample OutRawSample[MAX_MULTIPLEXED_TEXTURES],
 	out uint OutBufferCount)
+#if CONFIG_SIGNAL_TEXTURE_TYPE == SIGNAL_TEXTURE_TYPE_FLOAT4
 {
 	// Init all raw samples.
 	UNROLL_N(MAX_MULTIPLEXED_TEXTURES)
@@ -406,13 +436,13 @@ void EncodeMultiplexedSignals(
 	else if (SignalBufferLayout == SIGNAL_BUFFER_LAYOUT_REFLECTIONS_REJECTION)
 	{
 		// Outputs scene color to be compatible with a SSR output.
-		OutRawSample[0] = Sample[0].SceneColor * (Sample[0].SampleCount > 0 ? rcp(Sample[0].SampleCount) : 0);
+		OutRawSample[0] = ClampColorForEncoding(Sample[0].SceneColor * (Sample[0].SampleCount > 0 ? rcp(Sample[0].SampleCount) : 0));
 		OutRawSample[1].r = Sample[0].SampleCount;
 		OutBufferCount = 2;
 
 		if (MultiplexCount == 2)
 		{
-			OutRawSample[2] = Sample[1].SceneColor * (Sample[1].SampleCount > 0 ? rcp(Sample[1].SampleCount) : 0);
+			OutRawSample[2] = ClampColorForEncoding(Sample[1].SceneColor * (Sample[1].SampleCount > 0 ? rcp(Sample[1].SampleCount) : 0));
 			OutRawSample[1].g = Sample[1].SampleCount;
 			OutBufferCount = 3;
 		}
@@ -422,7 +452,7 @@ void EncodeMultiplexedSignals(
 		float NormalizationFactor = SafeRcp(Sample[0].SampleCount);
 
 		// Outputs scene color to be compatible with a SSR output.
-		OutRawSample[0] = Sample[0].SceneColor * NormalizationFactor;
+		OutRawSample[0] = ClampColorForEncoding(Sample[0].SceneColor * NormalizationFactor);
 		OutRawSample[1].r = Sample[0].SampleCount;
 		OutBufferCount = 2;
 
@@ -457,7 +487,7 @@ void EncodeMultiplexedSignals(
 	{
 		float NormalizationFactor = SafeRcp(Sample[0].SampleCount);
 
-		OutRawSample[0] = Sample[0].SceneColor * NormalizationFactor;
+		OutRawSample[0] = ClampColorForEncoding(Sample[0].SceneColor * NormalizationFactor);
 		OutRawSample[1].r = Sample[0].SampleCount;
 		OutBufferCount = 2;
 	}
@@ -465,7 +495,7 @@ void EncodeMultiplexedSignals(
 	{
 		float NormalizationFactor = SafeRcp(Sample[0].SampleCount);
 
-		OutRawSample[0] = Sample[0].SceneColor * NormalizationFactor;
+		OutRawSample[0] = ClampColorForEncoding(Sample[0].SceneColor * NormalizationFactor);
 		OutRawSample[1].r = Sample[0].SampleCount;
 		OutBufferCount = 2;
 
@@ -476,44 +506,67 @@ void EncodeMultiplexedSignals(
 		}
 	}
 	#endif
+} // EncodeMultiplexedSignals()
+#elif CONFIG_SIGNAL_TEXTURE_TYPE == SIGNAL_TEXTURE_TYPE_UINT2
+{
+	// Init all raw samples.
+	UNROLL_N(MAX_MULTIPLEXED_TEXTURES)
+	for (uint i = 0; i < MAX_MULTIPLEXED_TEXTURES; i++)
+		OutRawSample[i] = 0;
+
+	// Number of buffer the signal get encoded onto <= MAX_MULTIPLEXED_TEXTURES.
+	OutBufferCount = 1;
+
 	#if CONFIG_SIGNAL_COLOR_SH
-	else if (SignalBufferLayout == SIGNAL_BUFFER_LAYOUT_DIFFUSE_INDIRECT_HARMONIC)
+	if (SignalBufferLayout == SIGNAL_BUFFER_LAYOUT_DIFFUSE_INDIRECT_HARMONIC)
 	{
-		float NormalizationFactor = SafeRcp(Sample[0].SampleCount);
+		const uint EncodeOptions = (CONFIG_ENCODING_CLAMP_COLOR ? SSD_ENCODE_CLAMP_COLOR : 0) | SSD_ENCODE_NORMALIZE;
 
-		OutRawSample[3].x = Sample[0].SampleCount;
-		OutRawSample[3].y = Sample[0].MissCount * NormalizationFactor;
-		OutRawSample[0] = Sample[0].ColorSH.R.V * NormalizationFactor;
-		OutRawSample[1] = Sample[0].ColorSH.G.V * NormalizationFactor;
-		OutRawSample[2] = Sample[0].ColorSH.B.V * NormalizationFactor;
-		OutBufferCount = 4;
-
-		#if CONFIG_MANUAL_TEXEL_FORMATTING
-		{
-			 OutRawSample[0] = float4(asfloat(Compress2Floats(OutRawSample[0].xy)), asfloat(Compress2Floats(OutRawSample[0].zw)), 0, 0);
-			 OutRawSample[1] = float4(asfloat(Compress2Floats(OutRawSample[1].xy)), asfloat(Compress2Floats(OutRawSample[1].zw)), 0, 0);
-			 OutRawSample[2] = float4(asfloat(Compress2Floats(OutRawSample[2].xy)), asfloat(Compress2Floats(OutRawSample[2].zw)), 0, 0);
-			 OutRawSample[3] = float4(asfloat(Compress2Floats(OutRawSample[3].xy)), 0, 0, 0);
-		}
+		EncodeDiffuseSphericalHarmonicTexel(
+			Sample[0].SampleCount,
+			Sample[0].MissCount,
+			Sample[0].ColorSH,
+			EncodeOptions,
+			/* out */ OutRawSample);
 
-		#endif
+		OutBufferCount = 4;
 	}
 	#endif
 } // EncodeMultiplexedSignals()
+#else
+	#error Unimplemented.
+#endif
+
 
 /** Sample multiple input signals that have been multiplexed. */
 FSSDSignalArray SampleMultiplexedSignals(
-	Texture2D SignalBuffer0, Texture2D SignalBuffer1, Texture2D SignalBuffer2, Texture2D SignalBuffer3,
+	FSSDTexture2D SignalBuffer0, FSSDTexture2D SignalBuffer1, FSSDTexture2D SignalBuffer2, FSSDTexture2D SignalBuffer3,
 	SamplerState Sampler,
 	const uint SignalBufferLayout, const uint MultiplexedSampleId,
 	const bool bNormalizeSample,
 	float2 UV, float Level = 0)
 {
-	float4 RawSample[MAX_MULTIPLEXED_TEXTURES];
-	RawSample[0] = SignalBuffer0.SampleLevel(Sampler, UV, Level);
-	RawSample[1] = SignalBuffer1.SampleLevel(Sampler, UV, Level);
-	RawSample[2] = SignalBuffer2.SampleLevel(Sampler, UV, Level);
-	RawSample[3] = SignalBuffer3.SampleLevel(Sampler, UV, Level);
+	FSSDRawSample RawSample[MAX_MULTIPLEXED_TEXTURES];
+
+	#if CONFIG_SIGNAL_TEXTURE_TYPE == SIGNAL_TEXTURE_TYPE_FLOAT4
+	{
+		RawSample[0] = SignalBuffer0.SampleLevel(Sampler, UV, Level);
+		RawSample[1] = SignalBuffer1.SampleLevel(Sampler, UV, Level);
+		RawSample[2] = SignalBuffer2.SampleLevel(Sampler, UV, Level);
+		RawSample[3] = SignalBuffer3.SampleLevel(Sampler, UV, Level);
+	}
+	#elif CONFIG_SIGNAL_TEXTURE_TYPE == SIGNAL_TEXTURE_TYPE_UINT2
+	{
+		// TODO: UV * View.BufferSizeAndInvSize.xy could regress in VGPR in inner loops
+		int3 Coord = int3(int2(UV * View.BufferSizeAndInvSize.xy), int(Level));
+		RawSample[0] = SignalBuffer0.Load(Coord);
+		RawSample[1] = SignalBuffer1.Load(Coord);
+		RawSample[2] = SignalBuffer2.Load(Coord);
+		RawSample[3] = SignalBuffer3.Load(Coord);
+	}
+	#else
+		#error Unimplemented.
+	#endif
 
 	FSSDSignalArray MultiplexedSamples;
 	DecodeMultiplexedSignals(
@@ -529,15 +582,15 @@ FSSDSignalArray SampleMultiplexedSignals(
 
 /** Outputs multiplexed signal. */
 void OutputMultiplexedSignal(
-	RWTexture2D<float4> OutputSignalBuffer0,
-	RWTexture2D<float4> OutputSignalBuffer1,
-	RWTexture2D<float4> OutputSignalBuffer2,
-	RWTexture2D<float4> OutputSignalBuffer3,
+	FSSDRWTexture2D OutputSignalBuffer0,
+	FSSDRWTexture2D OutputSignalBuffer1,
+	FSSDRWTexture2D OutputSignalBuffer2,
+	FSSDRWTexture2D OutputSignalBuffer3,
 	const uint SignalBufferLayout, const uint MultiplexCount,
 	const uint2 PixelPosition, FSSDSignalSample MultiplexedSamples[MAX_SIGNAL_MULTIPLEXING])
 {
 	// Encode the output signal.
-	float4 RawSample[MAX_MULTIPLEXED_TEXTURES];
+	FSSDRawSample RawSample[MAX_MULTIPLEXED_TEXTURES];
 	uint BufferCount;
 	EncodeMultiplexedSignals(
 		SignalBufferLayout, MultiplexCount,
@@ -546,11 +599,11 @@ void OutputMultiplexedSignal(
 
 	// Output the raw encoded sample according to number of RT they requires.
 	if (BufferCount >= 1)
-		OutputSignalBuffer0[PixelPosition] = min(RawSample[0], MaxHalfFloat);
+		OutputSignalBuffer0[PixelPosition] = RawSample[0];
 	if (BufferCount >= 2)
-		OutputSignalBuffer1[PixelPosition] = min(RawSample[1], MaxHalfFloat);
+		OutputSignalBuffer1[PixelPosition] = RawSample[1];
 	if (BufferCount >= 3)
-		OutputSignalBuffer2[PixelPosition] = min(RawSample[2], MaxHalfFloat);
+		OutputSignalBuffer2[PixelPosition] = RawSample[2];
 	if (BufferCount >= 4)
-		OutputSignalBuffer3[PixelPosition] = min(RawSample[3], MaxHalfFloat);
+		OutputSignalBuffer3[PixelPosition] = RawSample[3];
 }
diff --git a/Engine/Shaders/Private/ScreenSpaceDenoise/SSDSignalCompression.ush b/Engine/Shaders/Private/ScreenSpaceDenoise/SSDSignalCompression.ush
--- a/Engine/Shaders/Private/ScreenSpaceDenoise/SSDSignalCompression.ush
+++ b/Engine/Shaders/Private/ScreenSpaceDenoise/SSDSignalCompression.ush
@@ -3,6 +3,7 @@
 #pragma once
 
 #include "SSDSignalCore.ush"
+#include "SSDPublicBufferEncoding.ush"
 
 
 //------------------------------------------------------- STRUCTURE
diff --git a/Engine/Shaders/Private/ScreenSpaceDenoise/SSDSpatialAccumulation.usf b/Engine/Shaders/Private/ScreenSpaceDenoise/SSDSpatialAccumulation.usf
--- a/Engine/Shaders/Private/ScreenSpaceDenoise/SSDSpatialAccumulation.usf
+++ b/Engine/Shaders/Private/ScreenSpaceDenoise/SSDSpatialAccumulation.usf
@@ -285,6 +285,7 @@
 	// Input and output layout.
 	#define CONFIG_SIGNAL_INPUT_LAYOUT  SIGNAL_BUFFER_LAYOUT_DIFFUSE_INDIRECT_HARMONIC
 	#define CONFIG_SIGNAL_OUTPUT_LAYOUT SIGNAL_BUFFER_LAYOUT_DIFFUSE_INDIRECT_HARMONIC
+	#define CONFIG_SIGNAL_TEXTURE_TYPE   SIGNAL_TEXTURE_TYPE_UINT2
 
 	// Spherical harmonics are a lot of data, need to shrink VGPR pressure to improve lattency hidding when fetching the buffer.
 	// TODO: some shader compiler completly falls apparts with the current implementation of
@@ -414,22 +415,22 @@ uint UpscaleFactor;
 	#error Missing CONFIG_INPUT_TEXTURE_COUNT
 #endif
 
-Texture2D SignalInput_Textures_0;
+FSSDTexture2D SignalInput_Textures_0;
 
 #if CONFIG_INPUT_TEXTURE_COUNT > 1
-Texture2D SignalInput_Textures_1;
+FSSDTexture2D SignalInput_Textures_1;
 #else
 #define SignalInput_Textures_1 SignalInput_Textures_0
 #endif
 
 #if CONFIG_INPUT_TEXTURE_COUNT > 2
-Texture2D SignalInput_Textures_2;
+FSSDTexture2D SignalInput_Textures_2;
 #else
 #define SignalInput_Textures_2 SignalInput_Textures_0
 #endif
 
 #if CONFIG_INPUT_TEXTURE_COUNT > 3
-Texture2D SignalInput_Textures_3;
+FSSDTexture2D SignalInput_Textures_3;
 #else
 #define SignalInput_Textures_3 SignalInput_Textures_0
 #endif
@@ -439,22 +440,22 @@ Texture2D SignalInput_Textures_3;
 	#error Missing CONFIG_OUTPUT_TEXTURE_COUNT
 #endif
 
-RWTexture2D<float4>	SignalOutput_UAVs_0;
+FSSDRWTexture2D	SignalOutput_UAVs_0;
 
 #if CONFIG_OUTPUT_TEXTURE_COUNT > 1
-RWTexture2D<float4>	SignalOutput_UAVs_1;
+FSSDRWTexture2D	SignalOutput_UAVs_1;
 #else
 #define SignalOutput_UAVs_1 SignalOutput_UAVs_0
 #endif
 
 #if CONFIG_OUTPUT_TEXTURE_COUNT > 2
-RWTexture2D<float4>	SignalOutput_UAVs_2;
+FSSDRWTexture2D	SignalOutput_UAVs_2;
 #else
 #define SignalOutput_UAVs_2 SignalOutput_UAVs_0
 #endif
 
 #if CONFIG_OUTPUT_TEXTURE_COUNT > 3
-RWTexture2D<float4>	SignalOutput_UAVs_3;
+FSSDRWTexture2D	SignalOutput_UAVs_3;
 #else
 #define SignalOutput_UAVs_3 SignalOutput_UAVs_0
 #endif
diff --git a/Engine/Shaders/Private/ScreenSpaceDenoise/SSDSpatialKernel.ush b/Engine/Shaders/Private/ScreenSpaceDenoise/SSDSpatialKernel.ush
--- a/Engine/Shaders/Private/ScreenSpaceDenoise/SSDSpatialKernel.ush
+++ b/Engine/Shaders/Private/ScreenSpaceDenoise/SSDSpatialKernel.ush
@@ -337,10 +337,10 @@ bool IsOutsideViewport(FSSDKernelConfig KernelConfig, float2 SampleBufferUV)
 /** Sample multiplexed samples and their metadata for kernel use. */
 void SampleMultiplexedSignals(
 	FSSDKernelConfig KernelConfig,
-	Texture2D SignalBuffer0,
-	Texture2D SignalBuffer1,
-	Texture2D SignalBuffer2,
-	Texture2D SignalBuffer3,
+	FSSDTexture2D SignalBuffer0,
+	FSSDTexture2D SignalBuffer1,
+	FSSDTexture2D SignalBuffer2,
+	FSSDTexture2D SignalBuffer3,
 	float2 SampleBufferUV,
 	out FSSDSampleSceneInfos OutSampleSceneMetadata,
 	out FSSDSignalArray OutMultiplexedSamples)
@@ -488,10 +488,10 @@ void AccumulateSampledMultiplexedSignals(
  */
 void SampleAndAccumulateMultiplexedSignals(
 	FSSDKernelConfig KernelConfig,
-	Texture2D SignalBuffer0,
-	Texture2D SignalBuffer1,
-	Texture2D SignalBuffer2,
-	Texture2D SignalBuffer3,
+	FSSDTexture2D SignalBuffer0,
+	FSSDTexture2D SignalBuffer1,
+	FSSDTexture2D SignalBuffer2,
+	FSSDTexture2D SignalBuffer3,
 	inout FSSDSignalAccumulatorArray Accumulators,
 	float2 SampleBufferUV,
 	float KernelSampleWeight,
@@ -526,10 +526,10 @@ void SampleAndAccumulateMultiplexedSignals(
 
 void SampleAndAccumulateMultiplexedSignalsPair(
 	FSSDKernelConfig KernelConfig,
-	Texture2D SignalBuffer0,
-	Texture2D SignalBuffer1,
-	Texture2D SignalBuffer2,
-	Texture2D SignalBuffer3,
+	FSSDTexture2D SignalBuffer0,
+	FSSDTexture2D SignalBuffer1,
+	FSSDTexture2D SignalBuffer2,
+	FSSDTexture2D SignalBuffer3,
 	inout FSSDSignalAccumulatorArray Accumulators,
 	float2 SampleBufferUV[2],
 	float KernelSampleWeight)
@@ -617,10 +617,10 @@ void DijestAccumulatedClusterSamples(inout FSSDSignalAccumulatorArray Accumulato
 
 void SampleAndAccumulateCenterSampleAsItsOwnCluster(
 	FSSDKernelConfig KernelConfig,
-	Texture2D SignalBuffer0,
-	Texture2D SignalBuffer1,
-	Texture2D SignalBuffer2,
-	Texture2D SignalBuffer3,
+	FSSDTexture2D SignalBuffer0,
+	FSSDTexture2D SignalBuffer1,
+	FSSDTexture2D SignalBuffer2,
+	FSSDTexture2D SignalBuffer3,
 	inout FSSDSignalAccumulatorArray Accumulators)
 {
 	const uint RingId = 0;
@@ -649,10 +649,10 @@ void SampleAndAccumulateCenterSampleAsItsOwnCluster(
 
 void AccumulateBilinear(
 	FSSDKernelConfig KernelConfig,
-	Texture2D SignalBuffer0,
-	Texture2D SignalBuffer1,
-	Texture2D SignalBuffer2,
-	Texture2D SignalBuffer3,
+	FSSDTexture2D SignalBuffer0,
+	FSSDTexture2D SignalBuffer1,
+	FSSDTexture2D SignalBuffer2,
+	FSSDTexture2D SignalBuffer3,
 	inout FSSDSignalAccumulatorArray Accumulators)
 {
 	const float MipLevelPow2 = 1;
@@ -689,10 +689,10 @@ void AccumulateBilinear(
 
 void AccumulateSquareKernel(
 	FSSDKernelConfig KernelConfig,
-	Texture2D SignalBuffer0,
-	Texture2D SignalBuffer1,
-	Texture2D SignalBuffer2,
-	Texture2D SignalBuffer3,
+	FSSDTexture2D SignalBuffer0,
+	FSSDTexture2D SignalBuffer1,
+	FSSDTexture2D SignalBuffer2,
+	FSSDTexture2D SignalBuffer3,
 	inout FSSDSignalAccumulatorArray Accumulators)
 {
 	int KernelRadius = 1;
@@ -918,10 +918,10 @@ static const uint kStackowiakSampleCountPerSet = 56;
 
 void ConvolveStackowiakKernel(
 	FSSDKernelConfig KernelConfig,
-	Texture2D SignalBuffer0,
-	Texture2D SignalBuffer1,
-	Texture2D SignalBuffer2,
-	Texture2D SignalBuffer3,
+	FSSDTexture2D SignalBuffer0,
+	FSSDTexture2D SignalBuffer1,
+	FSSDTexture2D SignalBuffer2,
+	FSSDTexture2D SignalBuffer3,
 	inout FSSDSignalAccumulatorArray Accumulators)
 {
 	if (KernelConfig.bDescOrder)
@@ -1152,10 +1152,10 @@ float2 SampleConstFromBatchConst(const uint BatchSampleId, float2 BatchConst)
 // Gather a ring into the accumulator.
 void GatherRingSamples(
 	FSSDKernelConfig KernelConfig,
-	Texture2D SignalBuffer0,
-	Texture2D SignalBuffer1,
-	Texture2D SignalBuffer2,
-	Texture2D SignalBuffer3,
+	FSSDTexture2D SignalBuffer0,
+	FSSDTexture2D SignalBuffer1,
+	FSSDTexture2D SignalBuffer2,
+	FSSDTexture2D SignalBuffer3,
 	inout FSSDSignalAccumulatorArray Accumulators,
 	const uint RingId)
 {
@@ -1233,10 +1233,10 @@ void GatherRingSamples(
 
 void ConvolveDiskKernel(
 	FSSDKernelConfig KernelConfig,
-	Texture2D SignalBuffer0,
-	Texture2D SignalBuffer1,
-	Texture2D SignalBuffer2,
-	Texture2D SignalBuffer3,
+	FSSDTexture2D SignalBuffer0,
+	FSSDTexture2D SignalBuffer1,
+	FSSDTexture2D SignalBuffer2,
+	FSSDTexture2D SignalBuffer3,
 	inout FSSDSignalAccumulatorArray Accumulators)
 {
 	// Accumulate the center of the kernel.
@@ -1299,10 +1299,10 @@ void ConvolveDiskKernel(
 
 void AccumulateKernel(
 	FSSDKernelConfig KernelConfig,
-	Texture2D SignalBuffer0,
-	Texture2D SignalBuffer1,
-	Texture2D SignalBuffer2,
-	Texture2D SignalBuffer3,
+	FSSDTexture2D SignalBuffer0,
+	FSSDTexture2D SignalBuffer1,
+	FSSDTexture2D SignalBuffer2,
+	FSSDTexture2D SignalBuffer3,
 	inout FSSDSignalAccumulatorArray Accumulators)
 {
 	if (KernelConfig.SampleSet == SAMPLE_SET_1X1)
diff --git a/Engine/Shaders/Private/ScreenSpaceDenoise/SSDTemporalAccumulation.usf b/Engine/Shaders/Private/ScreenSpaceDenoise/SSDTemporalAccumulation.usf
--- a/Engine/Shaders/Private/ScreenSpaceDenoise/SSDTemporalAccumulation.usf
+++ b/Engine/Shaders/Private/ScreenSpaceDenoise/SSDTemporalAccumulation.usf
@@ -142,6 +142,7 @@
 	// Input and output layout.
 	#define CONFIG_SIGNAL_INPUT_LAYOUT   SIGNAL_BUFFER_LAYOUT_DIFFUSE_INDIRECT_HARMONIC
 	#define CONFIG_SIGNAL_HISTORY_LAYOUT SIGNAL_BUFFER_LAYOUT_DIFFUSE_INDIRECT_HARMONIC
+	#define CONFIG_SIGNAL_TEXTURE_TYPE   SIGNAL_TEXTURE_TYPE_UINT2
 
 	#define CONFIG_INPUT_TEXTURE_COUNT 4
 	#define CONFIG_HISTORY_TEXTURE_COUNT 4
@@ -250,22 +251,22 @@ static const float InvMipLevelPow2 = 1;
 	#error Missing CONFIG_INPUT_TEXTURE_COUNT
 #endif
 
-Texture2D SignalInput_Textures_0;
+FSSDTexture2D SignalInput_Textures_0;
 
 #if CONFIG_INPUT_TEXTURE_COUNT > 1
-Texture2D SignalInput_Textures_1;
+FSSDTexture2D SignalInput_Textures_1;
 #else
 #define SignalInput_Textures_1 SignalInput_Textures_0
 #endif
 
 #if CONFIG_INPUT_TEXTURE_COUNT > 2
-Texture2D SignalInput_Textures_2;
+FSSDTexture2D SignalInput_Textures_2;
 #else
 #define SignalInput_Textures_2 SignalInput_Textures_0
 #endif
 
 #if CONFIG_INPUT_TEXTURE_COUNT > 3
-Texture2D SignalInput_Textures_3;
+FSSDTexture2D SignalInput_Textures_3;
 #else
 #define SignalInput_Textures_3 SignalInput_Textures_0
 #endif
@@ -276,22 +277,22 @@ Texture2D SignalInput_Textures_3;
 		#error Missing CONFIG_REJECTION_TEXTURE_COUNT
 	#endif
 
-	Texture2D HistoryRejectionSignal_Textures_0;
+	FSSDTexture2D HistoryRejectionSignal_Textures_0;
 
 	#if CONFIG_REJECTION_TEXTURE_COUNT > 1
-	Texture2D HistoryRejectionSignal_Textures_1;
+	FSSDTexture2D HistoryRejectionSignal_Textures_1;
 	#else
 	#define HistoryRejectionSignal_Textures_1 HistoryRejectionSignal_Textures_0
 	#endif
 
 	#if CONFIG_REJECTION_TEXTURE_COUNT > 2
-	Texture2D HistoryRejectionSignal_Textures_2;
+	FSSDTexture2D HistoryRejectionSignal_Textures_2;
 	#else
 	#define HistoryRejectionSignal_Textures_2 HistoryRejectionSignal_Textures_0
 	#endif
 
 	#if CONFIG_REJECTION_TEXTURE_COUNT > 3
-	Texture2D HistoryRejectionSignal_Textures_3;
+	FSSDTexture2D HistoryRejectionSignal_Textures_3;
 	#else
 	#define HistoryRejectionSignal_Textures_3 HistoryRejectionSignal_Textures_0
 	#endif
@@ -302,28 +303,28 @@ Texture2D SignalInput_Textures_3;
 	#error Missing CONFIG_HISTORY_TEXTURE_COUNT
 #endif
 
-Texture2D PrevHistory_Textures_0;
-RWTexture2D<float4>	SignalHistoryOutput_UAVs_0;
+FSSDTexture2D PrevHistory_Textures_0;
+FSSDRWTexture2D	SignalHistoryOutput_UAVs_0;
 
 #if CONFIG_HISTORY_TEXTURE_COUNT > 1
-Texture2D PrevHistory_Textures_1;
-RWTexture2D<float4>	SignalHistoryOutput_UAVs_1;
+FSSDTexture2D PrevHistory_Textures_1;
+FSSDRWTexture2D	SignalHistoryOutput_UAVs_1;
 #else
 #define PrevHistory_Textures_1 PrevHistory_Textures_0
 #define SignalHistoryOutput_UAVs_1 SignalHistoryOutput_UAVs_0
 #endif
 
 #if CONFIG_HISTORY_TEXTURE_COUNT > 2
-Texture2D PrevHistory_Textures_2;
-RWTexture2D<float4>	SignalHistoryOutput_UAVs_2;
+FSSDTexture2D PrevHistory_Textures_2;
+FSSDRWTexture2D	SignalHistoryOutput_UAVs_2;
 #else
 #define PrevHistory_Textures_2 PrevHistory_Textures_0
 #define SignalHistoryOutput_UAVs_2 SignalHistoryOutput_UAVs_0
 #endif
 
 #if CONFIG_HISTORY_TEXTURE_COUNT > 3
-Texture2D PrevHistory_Textures_3;
-RWTexture2D<float4>	SignalHistoryOutput_UAVs_3;
+FSSDTexture2D PrevHistory_Textures_3;
+FSSDRWTexture2D	SignalHistoryOutput_UAVs_3;
 #else
 #define PrevHistory_Textures_3 PrevHistory_Textures_0
 #define SignalHistoryOutput_UAVs_3 SignalHistoryOutput_UAVs_0
