Fixes a bug when sampling mip levels when using uint texture loads.

diff --git a/Engine/Shaders/Private/ScreenSpaceDenoise/SSDSignalBufferEncoding.ush b/Engine/Shaders/Private/ScreenSpaceDenoise/SSDSignalBufferEncoding.ush
--- a/Engine/Shaders/Private/ScreenSpaceDenoise/SSDSignalBufferEncoding.ush
+++ b/Engine/Shaders/Private/ScreenSpaceDenoise/SSDSignalBufferEncoding.ush
@@ -558,7 +558,7 @@ FSSDSignalArray SampleMultiplexedSignals(
 	#elif CONFIG_SIGNAL_TEXTURE_TYPE == SIGNAL_TEXTURE_TYPE_UINT2
 	{
 		// TODO: UV * View.BufferSizeAndInvSize.xy could regress in VGPR in inner loops
-		int3 Coord = int3(int2(UV * View.BufferSizeAndInvSize.xy), int(Level));
+		int3 Coord = int3(uint2(UV * View.BufferSizeAndInvSize.xy) >> uint(Level), int(Level));
 		RawSample[0] = SignalBuffer0.Load(Coord);
 		RawSample[1] = SignalBuffer1.Load(Coord);
 		RawSample[2] = SignalBuffer2.Load(Coord);
