Fixes a bug in the denoiser's NEIGHBOR_TO_REF_LOWEST_VGPR_PRESSURE code path

diff --git a/Engine/Shaders/Private/ScreenSpaceDenoise/SSDSpatialAccumulation.usf b/Engine/Shaders/Private/ScreenSpaceDenoise/SSDSpatialAccumulation.usf
--- a/Engine/Shaders/Private/ScreenSpaceDenoise/SSDSpatialAccumulation.usf
+++ b/Engine/Shaders/Private/ScreenSpaceDenoise/SSDSpatialAccumulation.usf
@@ -294,8 +294,8 @@
 	#define CONFIG_NEIGHBOR_TO_REF_COMPUTATION NEIGHBOR_TO_REF_LOWEST_VGPR_PRESSURE
 
 	#if DIM_STAGE == STAGE_RECONSTRUCTION
-		#define CONFIG_SAMPLE_SET           SAMPLE_SET_STOCASTIC_HIERARCHY
-		//#define CONFIG_SAMPLE_SET           SAMPLE_SET_STACKOWIAK_4_SETS
+		//#define CONFIG_SAMPLE_SET           SAMPLE_SET_STOCASTIC_HIERARCHY
+		#define CONFIG_SAMPLE_SET           SAMPLE_SET_STACKOWIAK_4_SETS
 
 	#elif DIM_STAGE == STAGE_REJECTION_PRE_CONVOLUTION
 
diff --git a/Engine/Shaders/Private/ScreenSpaceDenoise/SSDSpatialKernel.ush b/Engine/Shaders/Private/ScreenSpaceDenoise/SSDSpatialKernel.ush
--- a/Engine/Shaders/Private/ScreenSpaceDenoise/SSDSpatialKernel.ush
+++ b/Engine/Shaders/Private/ScreenSpaceDenoise/SSDSpatialKernel.ush
@@ -245,7 +245,7 @@ float2 ComputeRefBufferUV(FSSDKernelConfig KernelConfig)
 			SampleTrackId & 0x1 ? 0.5 : -0.5,
 			SampleTrackId & 0x2 ? 0.5 : -0.5);
 
-		return KernelConfig.BufferUV + SampleOffset * KernelConfig.BufferSizeAndInvSize.xy;
+		return KernelConfig.BufferUV + SampleOffset * KernelConfig.BufferSizeAndInvSize.zw;
 	}
 	else if (SAMPLE_SET_HEXAWEB)
 	{
@@ -277,6 +277,7 @@ float3 ComputeWorldVectorFromNeighborToRef(
 }
 #elif CONFIG_NEIGHBOR_TO_REF_COMPUTATION == NEIGHBOR_TO_REF_LOWEST_VGPR_PRESSURE
 {
+	// Recompute the the screen position of the reference, from the most minimal VGPR footprint.
 	float2 RefScreenPos = ComputeRefScreenPos(KernelConfig);
 	float3 RefClipPosition = float3(RefScreenPos * (View.ViewToClip[3][3] < 1.0f ? KernelConfig.RefSceneMetadata.WorldDepth : 1.0f), KernelConfig.RefSceneMetadata.WorldDepth);
 
