Fix contact shadows missing for the eye material type

diff --git a/Engine/Shaders/Private/DeferredLightingCommon.ush b/Engine/Shaders/Private/DeferredLightingCommon.ush
--- a/Engine/Shaders/Private/DeferredLightingCommon.ush
+++ b/Engine/Shaders/Private/DeferredLightingCommon.ush
@@ -199,9 +199,9 @@ void GetShadowTerms(FGBufferData GBuffer, FDeferredLightData LightData, float3 W
 		Shadow.SurfaceShadow *= ContactShadow;
 
 		FLATTEN
-		if( GBuffer.ShadingModelID == SHADINGMODELID_HAIR )
+		if( GBuffer.ShadingModelID == SHADINGMODELID_HAIR ||  GBuffer.ShadingModelID == SHADINGMODELID_EYE  )
 			Shadow.TransmissionShadow *= ContactShadow;
-		else if( GBuffer.ShadingModelID != SHADINGMODELID_EYE )
+		else
 			Shadow.TransmissionShadow *= ContactShadow * 0.5 + 0.5;
 	}
 #endif
