Added experimental support for Burley Normalized Subsurface Scattering.

1. Burley is disabled by default. To preview it, you need to 1) set r.SubsurfaceScattering.Type to 1 and 2) check the enable burley box in the profile.
2. Burley parameters always update from Separable. To disable it, set r.SSS.Burley.AlwaysUpdateParametersFromSeparable to 0.

diff --git a/Engine/Shaders/Private/BurleyNormalizedSSSCommon.ush b/Engine/Shaders/Private/BurleyNormalizedSSSCommon.ush
new file mode 100644
--- /dev/null
+++ b/Engine/Shaders/Private/BurleyNormalizedSSSCommon.ush
@@ -0,0 +1,165 @@
+// Copyright 1998-2019 Epic Games, Inc. All Rights Reserved.
+
+/*=========================================================================
+	BurleyNormalizedSSSCommon.ush: Burley common functions.
+=========================================================================*/
+#pragma once
+
+#include "SubsurfaceProfileCommon.ush"
+
+#define BURLEY_MM_2_CM		0.1f
+#define BURLEY_CM_2_MM      10.0f
+
+//------------------------------------------------------------------------------------------
+// Consistent in BurleyNormalizedSSSCommon.ush and SubsurfaceProfile.cpp
+
+#define SSS_TYPE_SSSS		1.0f
+#define SSS_TYPE_BURLEY	    0.0f
+
+//make sure UIMax|ClampMax of WorldScale * ENC_WORLDSCALE_TO_UNIT <= 1
+#define ENC_WORLDSCALE_IN_CM_TO_UNIT 0.02f
+#define DEC_UNIT_TO_WORLDSCALE_IN_CM 1/ENC_WORLDSCALE_IN_CM_TO_UNIT
+
+//make sure UIMax|ClampMax of DiffuseMeanFreePath * ENC_DIFFUSEMEANFREEPATH_IN_MM_TO_UNIT <= 1
+#define ENC_DIFFUSEMEANFREEPATH_IN_MM_TO_UNIT 0.01f
+#define DEC_UNIT_TO_DIFFUSEMEANFREEPATH_IN_MM 1/ENC_DIFFUSEMEANFREEPATH_IN_MM_TO_UNIT
+//------------------------------------------------------------------------------------------
+
+float DecodeWorldScale(float EncodedWorldScale)
+{
+	return EncodedWorldScale * DEC_UNIT_TO_WORLDSCALE_IN_CM;
+}
+
+float4 DecodeDiffuseMeanFreePath(float4 EncodedDiffuseMeanFreePath)
+{
+	return EncodedDiffuseMeanFreePath * DEC_UNIT_TO_DIFFUSEMEANFREEPATH_IN_MM;
+}
+
+bool GetSubsurfaceTransmittanceProfileUseBurley(uint SubsurfaceProfileInt)
+{
+	// 0..255, which SubSurface profile to pick
+	float Type = ActualSSProfilesTexture.Load(int3(SSSS_BOUNDARY_COLOR_BLEED_OFFSET, SubsurfaceProfileInt, 0)).a;
+
+	return abs(Type - SSS_TYPE_BURLEY) < 0.01f;
+}
+
+bool GetSubsurfaceProfileUseBurley(uint SubsurfaceProfileInt)
+{
+	// 0..255, which SubSurface profile to pick
+	float Type = ActualSSProfilesTexture.Load(int3(SSSS_BOUNDARY_COLOR_BLEED_OFFSET, SubsurfaceProfileInt, 0)).a;
+
+	return abs(Type - SSS_TYPE_BURLEY) < 0.01f;
+}
+
+inline float3 Burley_Profile(float Radius, float3 Albedo, float3 S3D, float L)
+{   //R(r)r
+	float3 D = 1 / S3D;
+	float R = Radius / L;
+	const float Inv8Pi = 1.0 / (8 * PI);
+	float3 NegRbyD = -R / D;
+	return Albedo * max((exp(NegRbyD) + exp(NegRbyD / 3.0)) / (D*L)*Inv8Pi, 0);
+}
+
+//Diffuse profile basic formula
+//    D : shapes the height and width of the profile curve.
+//	  Radius : the distance between the entering surface point and the exit surface point.
+//Assume: r and d >0
+float GetDiffuseReflectProfile(float D, float Radius)
+{
+	//The diffuse reflectance profile:
+	//R(d,r) = \frac{e^{-r/d}+e^{-r/(3d)}}{8*pi*d*r}
+	const float Inv8Pi = 1.0 / (8 * PI);
+	float NegRbyD = -Radius / D;
+	return (exp(NegRbyD) + exp(NegRbyD / 3.0)) / (D*Radius)*Inv8Pi;
+}
+
+float3 GetDiffuseReflectProfileWithDiffuseMeanFreePath(float3 SurfaceAlbedo, float3 L, float3 S3D, float Radius)
+{
+	//rR(r)
+	float3 D = 1 / S3D;
+	float3 R = Radius / L;
+	const float Inv8Pi = 1.0 / (8 * PI);
+	float3 NegRbyD = -R / D;
+	return SurfaceAlbedo * max((exp(NegRbyD) + exp(NegRbyD / 3.0)) / (D*L)*Inv8Pi, 0);
+}
+
+float GetComponentForScalingFactorEstimation(float4 SurfaceAlbedo)
+{
+	return SurfaceAlbedo.a;
+}
+
+float GetDiffuseMeanFreePathForSampling(float4 DiffuseMeanFreePath)
+{
+	return DiffuseMeanFreePath.a;
+}
+
+//-------------------------------------------------------------------------
+// Three scaling factor function
+// Method 1: The light directly goes into the volume in a direction perpendicular to the surface.
+// Average relative error: 5.5% (reference to MC)
+float GetPerpendicularScalingFactor(float SurfaceAlbedo)
+{
+	return 1.85 - SurfaceAlbedo + 7 * UnClampedPow(SurfaceAlbedo - 0.8, 3);
+}
+
+float3 GetPerpendicularScalingFactor(float3 SurfaceAlbedo)
+{
+	return 1.85 - SurfaceAlbedo + 7 * UnClampedPow(SurfaceAlbedo - 0.8, 3);
+}
+
+// Method 2: Ideal diffuse transmission at the surface. More appropriate for rough surface.
+// Average relative error: 3.9% (reference to MC)
+float GetDiffuseSurfaceScalingFactor(float SurfaceAlbedo)
+{
+	return 1.9 - SurfaceAlbedo + 3.5*(SurfaceAlbedo - 0.8)*(SurfaceAlbedo - 0.8);
+}
+
+float3 GetDiffuseSurfaceScalingFactor(float3 SurfaceAlbedo)
+{
+	return 1.9 - SurfaceAlbedo + 3.5*(SurfaceAlbedo - 0.8)*(SurfaceAlbedo - 0.8);
+}
+
+// Method 3: The spectral of diffuse mean free path on the surface.
+// Avergate relative error: 7.7% (reference to MC)
+float GetSearchLightDiffuseScalingFactor(float SurfaceAlbedo)
+{
+	return 3.5 + 100 * UnClampedPow(SurfaceAlbedo - 0.33, 4);
+}
+
+float3 GetSearchLightDiffuseScalingFactor(float3 SurfaceAlbedo)
+{
+	return 3.5 + 100 * UnClampedPow(SurfaceAlbedo - 0.33, 4);
+}
+
+//@TODO Revisit here after offline comparison.
+//call this function in TransmissionCommon.ush
+float4 GetBurleyTransmissionProfile(FGBufferData GBufferData, float Thickness)
+{
+	// 0..255, which SubSurface profile to pick
+#if 0
+	uint SubsurfaceProfileInt = ExtractSubsurfaceProfileInt(GBufferData);
+	float WorldScale = DecodeWorldScale(ActualSSProfilesTexture.Load(int3(SSSS_SUBSURFACE_COLOR_OFFSET, SubsurfaceProfileInt,0)).a)*10.0f; //in cm. we do not multiply by 100.0f as we did in the subsurface shader.
+	float3 MeanFreePath = DecodeDiffuseMeanFreePath(ActualSSProfilesTexture.Load(int3(SSSS_SUBSURFACE_COLOR_OFFSET, SubsurfaceProfileInt,0)));
+	float3 SurfaceAlbedo = ActualSSProfilesTexture.Load(int3(SSS__SURFACEALBEDO_OFFSET, SubsurfaceProfileInt, 0)).rgb;
+	bool  UseBurleyTransmittance = GetSubsurfaceTransmittanceProfileUseBurley(SubsurfaceProfileInt);
+
+	//this is slow.
+	BRANCH
+	if (UseBurleyTransmittance)
+	{
+		//See the burley Normalized SSS paper for the formula between equation 3 and 4. by now we still need to clamp?
+		float3 ScalingFactor = GetSearchLightDiffuseScalingFactor(SurfaceAlbedo);
+		float4 Output=float4(0,0,0,1);
+		Output.xyz=GetDiffuseReflectProfileWithDiffuseMeanFreePath(SurfaceAlbedo, MeanFreePath, ScalingFactor,Thickness*WorldScale);
+		return Output;
+	}
+	else
+	{
+		float3 TextureSize;
+		ActualSSProfilesTexture.GetDimensions(0, TextureSize.x, TextureSize.y, TextureSize.z); // TODO: it's slightly faster to pass texture dimensions through a constant buffer
+		float2 uv = float2((SSSS_TRANSMISSION_PROFILE_OFFSET + Thickness / SSSS_MAX_TRANSMISSION_PROFILE_DISTANCE * (SSSS_TRANSMISSION_PROFILE_SIZE - 1) + 0.5f), SubsurfaceProfileInt + 0.5f) / (TextureSize.xy);
+		return Texture2DSampleLevel(ActualSSProfilesTexture, TransmissionProfilesLinearSampler, uv, 0);
+	}
+#endif
+	return 0;
+}
\ No newline at end of file
diff --git a/Engine/Shaders/Private/SubsurfaceBurleyNormalized.usf b/Engine/Shaders/Private/SubsurfaceBurleyNormalized.usf
new file mode 100644
--- /dev/null
+++ b/Engine/Shaders/Private/SubsurfaceBurleyNormalized.usf
@@ -0,0 +1,796 @@
+// Copyright 1998-2019 Epic Games, Inc. All Rights Reserved.
+
+/*=============================================================================
+	SubsurfaceBurleyNormalized.usf: Screenspace Burley subsurface scattering implementation.
+=============================================================================*/
+
+#include "Common.ush"
+#include "Random.ush"
+#include "PostProcessCommon.ush"
+#include "DeferredShadingCommon.ush"
+#include "SubsurfaceProfileCommon.ush"
+#include "BurleyNormalizedSSSCommon.ush"
+
+#include "ScreenPass.ush"
+
+// x:Radius*DistanceToProjectionWindow/KernelSize*0.5, y:DistanceToProjectionWindow, zw: unused
+float4 SubsurfaceParams;
+
+Texture2D SubsurfaceInput0_Texture;
+Texture2D SubsurfaceInput1_Texture;
+
+SamplerState SubsurfaceSampler0;
+SamplerState SubsurfaceSampler1;
+
+#if SUPPORTS_INDEPENDENT_SAMPLERS
+	#define SharedSubsurfaceSampler0 SubsurfaceSampler0
+	#define SharedSubsurfaceSampler1 SubsurfaceSampler0
+#else
+	#define SharedSubsurfaceSampler0 SubsurfaceSampler0
+	#define SharedSubsurfaceSampler1 SubsurfaceSampler1
+#endif
+
+SCREEN_PASS_TEXTURE_VIEWPORT(SubsurfaceInput0)
+SCREEN_PASS_TEXTURE_VIEWPORT(SubsurfaceInput1)
+
+#define SUBSURFACE_PASS_ONE 0
+#define SUBSURFACE_PASS_TWO 1
+
+#define SUBSURFACE_DIRECTION_HORIZONTAL SUBSURFACE_PASS_ONE
+#define SUBSURFACE_DIRECTION_VERTICAL SUBSURFACE_PASS_TWO
+
+// Controls the quality (number of samples) of the blur kernel.
+#define SUBSURFACE_QUALITY_LOW 0
+#define SUBSURFACE_QUALITY_MEDIUM 1
+#define SUBSURFACE_QUALITY_HIGH 2
+
+// Full resolution recombine.
+#define SUBSURFACE_RECOMBINE_MODE_FULLRES 0
+
+// Half resolution recombine.
+#define SUBSURFACE_RECOMBINE_MODE_HALFRES 1
+
+// Just reconstruct the lighting (needed for scalability).
+#define SUBSURFACE_RECOMBINE_MODE_NO_SCATTERING 2
+
+// Controls the quality of lighting reconstruction.
+#define SUBSURFACE_RECOMBINE_QUALITY_LOW 0
+#define SUBSURFACE_RECOMBINE_QUALITY_HIGH 1
+
+#ifndef SUBSURFACE_RECOMBINE_QUALITY
+	#define SUBSURFACE_RECOMBINE_QUALITY SUBSURFACE_RECOMBINE_QUALITY_LOW
+#endif
+
+#if SUBSURFACE_RECOMBINE_QUALITY == SUBSURFACE_RECOMBINE_QUALITY_HALFRES
+	#define SUBSURFACE_HALFRES 1
+#endif
+
+//=============================================================================
+// Setup for "Burley"
+//=============================================================================
+
+#define BURLEY_NUM_SAMPLES	64
+#define BURLEY_INV_NUM_SAMPLES (1.0f/BURLEY_NUM_SAMPLES)
+
+struct FBurleyParameter
+{
+	float4 SurfaceAlbedo;
+	float4 DiffuseMeanFreePath;
+	float  WorldScale;
+};
+
+//volume albedo and mean free path length
+float4 GetSubsurfaceProfileSurfaceAlbedo(FGBufferData GBufferData)
+{
+	// 0..255, which SubSurface profile to pick
+	uint SubsurfaceProfileInt = ExtractSubsurfaceProfileInt(GBufferData);
+	//SSSS_KERNEL0_OFFSET
+	return ActualSSProfilesTexture.Load(int3(SSS__SURFACEALBEDO_OFFSET, SubsurfaceProfileInt, 0));
+}
+
+float4 GetSubsurfaceProfileDiffuseMeanFreePath(FGBufferData GBufferData)
+{
+	// 0..255, which SubSurface profile to pick
+	uint SubsurfaceProfileInt = ExtractSubsurfaceProfileInt(GBufferData);
+	//SSSS_KERNEL0_OFFSET
+	return ActualSSProfilesTexture.Load(int3(SSS__DMFP_OFFSET, SubsurfaceProfileInt, 0));
+}
+
+float GetSubsurfaceProfileWorldScale(FGBufferData GBufferData)
+{
+	// 0..255, which SubSurface profile to pick
+	uint SubsurfaceProfileInt = ExtractSubsurfaceProfileInt(GBufferData);
+
+	return ActualSSProfilesTexture.Load(int3(SSSS_SUBSURFACE_COLOR_OFFSET, SubsurfaceProfileInt,0)).a;
+}
+
+// fetch volume albedo and diffuse mean free path
+// The diffuse mean free path is modulated by ScatterRadius and a custom term passed in
+// which is the opacity texture.
+FBurleyParameter GetBurleyParameters(FGBufferData GBuffer)
+{
+	FBurleyParameter BurleyParameter;
+
+	BurleyParameter.SurfaceAlbedo = GetSubsurfaceProfileSurfaceAlbedo(GBuffer);
+	float4 DiffuseMeanFreePath = GetSubsurfaceProfileDiffuseMeanFreePath(GBuffer);
+	BurleyParameter.DiffuseMeanFreePath = DecodeDiffuseMeanFreePath(DiffuseMeanFreePath);
+	float WorldScale = DecodeWorldScale(GetSubsurfaceProfileWorldScale(GBuffer))*BURLEY_CM_2_MM;
+	BurleyParameter.WorldScale = WorldScale;
+	BRANCH
+	if (GBuffer.ShadingModelID == SHADINGMODELID_SUBSURFACE_PROFILE)
+	{
+		// scale with opacity stored in CustomData.a
+		BurleyParameter.DiffuseMeanFreePath *= GBuffer.CustomData.a;
+	}
+	else if (GBuffer.ShadingModelID == SHADINGMODELID_EYE)
+	{
+		//eye 
+	}
+	return BurleyParameter;
+}
+
+float2 Generate2DRandomNumber(int3 Seed)
+{
+	return float2(Rand3DPCG16(Seed).xy) / 0x10000;
+}
+
+float GetCDF(float D, float X, float XI)
+{
+	return 1 - 0.25*exp(-X / D) - 0.75*exp(-X / (3 * D)) - XI;
+}
+float GetCDFDeriv1(float D, float X)
+{
+	return 0.25 / D * (exp(-X / D) + exp(-X / (3 * D)));
+}
+
+float GetCDFDeriv2(float D, float X)
+{
+	return exp(-X / D)*(-0.0833333*exp(2 * X / (3 * D)) - 0.25) / (D*D);
+}
+
+// Without clamp, the result will be zero which will lead to infinity for R/pdf. and Nan for the weighted sum.
+// We need to clamp this to a very small pdf. this clamp is based on value below
+#define CLAMP_PDF 0.00001
+
+// Given D and a random number, use root finding method to calculate the radius
+// in meters
+float RadiusRootFinding(float D, float RandomNumber, float X0)
+{
+	// Make the distribution correct.
+	// r=cdf^(-1)(p)     if p<1-CLAMP_PDF
+	//  =cdf^(-1)(1-ClAMP_PDF) P>=1-CLAMP_PDF
+	// RandomNumber = clamp(RandomNumber,0,1-CLAMP_PDF);
+ 
+	const int NumOfIteration = 3;
+	float Xn = X0;
+
+	UNROLL
+		for (int i = 0; i < NumOfIteration; ++i)
+		{
+			float Fx = GetCDF(D, Xn, RandomNumber);
+			float DFx = GetCDFDeriv1(D, Xn);
+			float DFxx = GetCDFDeriv2(D, Xn);
+			Xn = Xn - 2 * Fx*DFx / (2 * DFx*DFx - Fx * DFxx);
+		}
+	return Xn;
+}
+
+//Brian's approximation.
+float RadiusRootFindByApproximation(float D, float RandomNumber)
+{
+	return D * ((2 - 2.6)*RandomNumber - 2)*log(1-RandomNumber);
+}
+
+// get the probability to sample a disk.
+float GetPdf(float Radius, float L, float S)
+{
+	//without clamp, the result will be zero which will lead to infinity for R/pdf. and Nan for the weighted sum.
+	//we need to clamp this to a very small pdf.
+
+	float Pdf = GetCDFDeriv1(L / S, Radius);
+	return clamp(Pdf, CLAMP_PDF, 1);
+}
+
+struct FBurleySampleInfo
+{
+	float Radius;
+	float Theta;
+	float Pdf;
+};
+
+FBurleySampleInfo GenerateSampleInfo(float2 Rand0T1,float DiffuseMeanFreePathForSample, float SpectralForSample,uint SequenceId)
+{
+	FBurleySampleInfo BurleySampleInfo;
+
+#define RADIUS_SAMPLE_PDF
+#ifdef RADIUS_SAMPLE_UNIFORM_DISK
+
+	float RadiusSampled = max(DiffuseMeanFreePathForSample * sqrtFast(Rand0T1.x), 0.0001f);//should not be zero
+	float Pdf = 1 / (2 * PI);	// use actual pdf
+	float Theta = Rand0T1.y * 2 * PI;
+#elif defined(RADIUS_SAMPLE_PDF)
+
+#define ROOT_APROXIMATE
+
+#ifdef ROOT_FINDING
+	float FoundRoot = RadiusRootFinding(DiffuseMeanFreePathForSample / SpectralForSample, Rand0T1.x, DiffuseMeanFreePathForSample);
+#elif defined(ROOT_APROXIMATE)
+	//Approximation
+	float FoundRoot = RadiusRootFindByApproximation(DiffuseMeanFreePathForSample / SpectralForSample,Rand0T1.x);
+#endif
+	BurleySampleInfo.Radius = max(FoundRoot, 0.00001f);//in mm
+
+#define SAMPLE_ANGLE_FIBONACCI	
+#ifdef SAMPLE_ANGLE_RANDOM
+	BurleySampleInfo.Theta = Rand0T1.y * 2 * PI;
+#elif defined(SAMPLE_ANGLE_FIBONACCI)
+	//Less variance
+	BurleySampleInfo.Theta = (((SequenceId + 0.5 + Rand0T1.y))*(1 + sqrt(5))*0.5) * 2 * PI;//Fibonacci sequence for angle.
+#endif
+	BurleySampleInfo.Pdf = GetPdf(BurleySampleInfo.Radius, DiffuseMeanFreePathForSample, SpectralForSample); //
+#endif
+
+	return BurleySampleInfo;
+}
+
+//=============================================================================
+// setup for "SeparableSSS.ush"
+//=============================================================================
+
+#if SUBSURFACE_QUALITY == SUBSURFACE_QUALITY_LOW
+	#define	SSSS_N_KERNELWEIGHTCOUNT SSSS_KERNEL2_SIZE
+	#define	SSSS_N_KERNELWEIGHTOFFSET SSSS_KERNEL2_OFFSET
+#elif SUBSURFACE_QUALITY == SUBSURFACE_QUALITY_MEDIUM
+	#define	SSSS_N_KERNELWEIGHTCOUNT SSSS_KERNEL1_SIZE
+	#define	SSSS_N_KERNELWEIGHTOFFSET SSSS_KERNEL1_OFFSET
+#else // SUBSURFACE_QUALITY == SUBSURFACE_QUALITY_HIGH
+	#define	SSSS_N_KERNELWEIGHTCOUNT SSSS_KERNEL0_SIZE
+	#define	SSSS_N_KERNELWEIGHTOFFSET SSSS_KERNEL0_OFFSET
+#endif
+
+// 0: faster
+// 1: no color bleeding in z direction
+#define SSSS_FOLLOW_SURFACE 1
+
+float GetMaskFromDepthInAlpha(float Alpha)
+{
+	return Alpha > 0;
+}
+
+// can be optimized
+float GetSubsurfaceStrength(float2 UV)
+{
+	FScreenSpaceData ScreenSpaceData = GetScreenSpaceData(UV);
+
+	float Mask = UseSubsurfaceProfile(ScreenSpaceData.GBuffer.ShadingModelID) ? 1.0f : 0.0f;
+	return Mask * ScreenSpaceData.GBuffer.CustomData.a;
+}
+
+// @return .rgb is the weight for color channel, .a is the sample location
+float4 GetSubsurfaceProfileKernel(uint SampleIndex, uint SubsurfaceProfileInt)
+{
+	const float4 TableMax = float4(1, 1, 1, SUBSURFACE_KERNEL_SIZE);
+
+	return ActualSSProfilesTexture.Load(int3(SampleIndex, SubsurfaceProfileInt, 0)) * TableMax;
+}
+
+float GetSubsurfaceProfileRadiusScale(FGBufferData GBufferData)
+{
+	// 0..255, which SubSurface profile to pick
+	uint SubsurfaceProfileInt = ExtractSubsurfaceProfileInt(GBufferData);
+
+	return GetSubsurfaceProfileKernel(SSSS_N_KERNELWEIGHTOFFSET + SSSS_N_KERNELWEIGHTCOUNT - 1, SubsurfaceProfileInt).a;
+}
+
+float3 GetSubsurfaceProfileColor(FGBufferData GBufferData)
+{
+	// 0..255, which SubSurface profile to pick
+	uint SubsurfaceProfileInt = ExtractSubsurfaceProfileInt(GBufferData);
+
+	return GetSubsurfaceProfileKernel(SSSS_SUBSURFACE_COLOR_OFFSET, SubsurfaceProfileInt).rgb;
+}
+
+uint GetSubsurfaceProfileId(float2 BufferUV)
+{
+	BufferUV = clamp(BufferUV, SubsurfaceInput0_UVViewportBilinearMin, SubsurfaceInput0_UVViewportBilinearMax);
+
+	FGBufferData LocalGBufferData = GetGBufferData(BufferUV);
+	return ExtractSubsurfaceProfileInt(LocalGBufferData);
+}
+
+float3 GetSubsurfaceProfileBoundaryColorBleed(FGBufferData GBufferData)
+{
+	// 0..255, which SubSurface profile to pick
+	uint SubsurfaceProfileInt = ExtractSubsurfaceProfileInt(GBufferData);
+
+	return GetSubsurfaceProfileKernel(SSSS_BOUNDARY_COLOR_BLEED_OFFSET, SubsurfaceProfileInt).rgb;
+}
+
+float4 GetSceneColor(float2 BufferUV)
+{
+	BufferUV = clamp(BufferUV, SubsurfaceInput0_UVViewportBilinearMin, SubsurfaceInput0_UVViewportBilinearMax);
+
+	return Texture2DSampleLevel(SubsurfaceInput0_Texture, SharedSubsurfaceSampler0, BufferUV, 0);
+}
+
+// from [REDACTED_PROFILE_URL]/separable-sss/tree/master/Demo
+// [REDACTED_PERSON] [REDACTED_PROFILE_URL]
+// [REDACTED_PROFILE_URL]
+#include "SeparableSSS.ush"
+
+//=============================================================================
+
+bool InUnitBox(float2 UV)
+{
+	return UV.x >= 0 && UV.y >= 0 && UV.y < 1 && UV.y < 1;
+}
+
+// @return 0=don't blend in, 1:fully blend in
+float ComputeFullResLerp(FScreenSpaceData ScreenSpaceData, float2 UVSceneColor, float2 FullResInputSizeInverse)
+{
+	float SSSScaleX = SubsurfaceParams.x;
+
+	float scale = SSSScaleX / CalcSceneDepth(UVSceneColor);
+
+	float HorizontalScaler = SUBSURFACE_RADIUS_SCALE;
+
+	// Calculate the final step to fetch the surrounding pixels:
+	float finalStep = scale * HorizontalScaler;
+
+	finalStep *= GetSubsurfaceProfileRadiusScale(ScreenSpaceData.GBuffer);
+
+	float PixelSizeRadius = finalStep / (FullResInputSizeInverse.x * 0.5f);
+
+	// tweaked for skin, a more flat kernel might need a smaller value, around 2 seems reasonable because we do half res
+	const float PixelSize = 4.0f;
+
+	float Ret = saturate(PixelSizeRadius - PixelSize);
+
+	// opacity allows to scale the radius - at some point we should fade in the full resolution, we don't have a masking other than that.
+	Ret *= saturate(ScreenSpaceData.GBuffer.CustomData.a * 10);
+
+	// todo: Subsurface has some non scatter contribution - all that should come from the Full res
+	return Ret;
+}
+
+struct SDiffuseAndSpecular
+{
+	float3 Diffuse;
+	float3 Specular;
+};
+
+// can be moved/shared
+half3 LookupSceneColor(float2 SceneUV, int2 PixelOffset)
+{
+#if ES2_PROFILE && COMPILER_GLSL_ES2
+	// slower but always works
+	// to prevent "error: Texture offset not supported on GLSL ES"
+	return Texture2DSample(SubsurfaceInput0_Texture, SharedSubsurfaceSampler0, SceneUV + PixelOffset * SubsurfaceInput0_ExtentInverse).rgb;
+#else
+	// faster
+	return SubsurfaceInput0_Texture.SampleLevel(SharedSubsurfaceSampler0, SceneUV, 0, PixelOffset).rgb;
+#endif
+}
+ 
+// @param UVSceneColor for the full res rendertarget (BufferSize) e.g. SceneColor or GBuffers
+// @param ReconstructMethod 0/1/2/3 (should be a literal constant to allow compiler optimizations)
+SDiffuseAndSpecular ReconstructLighting(float2 UVSceneColor, uint ReconstructMethod)
+{
+	SDiffuseAndSpecular Ret;
+
+	// If SUBSURFACE_CHANNEL_MODE is 0, checkerboard is forced on
+#if SUBSURFACE_PROFILE_CHECKERBOARD || SUBSURFACE_CHANNEL_MODE == 0
+	{
+	  bool bChecker = CheckerFromSceneColorUV(UVSceneColor);
+
+	  // todo: We could alternate the diagonal with TemporalAA or even only only 1 sample for low spec or 4 for high spec
+
+	  float3 Quant0 = SubsurfaceInput0_Texture.SampleLevel(SharedSubsurfaceSampler0, UVSceneColor, 0).rgb;
+
+	  // todo: expose as scalability setting (can be evaluate best without TemporalAA)
+	  // 0:fast but pattern can appear, 1:better, 2: even better, 3: best but expensive
+	  float3 Quant1;
+
+	  if(ReconstructMethod == 0)
+	  {
+		  // cheap, crappy
+		  Quant1 = LookupSceneColor(UVSceneColor, int2(1, 0));
+	  }
+	  else if(ReconstructMethod == 1)
+	  {
+		  // acceptable but not perfect
+		  Quant1 = 0.5f * (
+			  LookupSceneColor(UVSceneColor, int2( 1, 0)) +
+			  LookupSceneColor(UVSceneColor, int2(-1, 0)));
+	  }
+	  else if(ReconstructMethod == 2)
+	  {
+		  // almost same as 1?
+		  Quant1 = 0.25f * (
+			  LookupSceneColor(UVSceneColor, int2( 1,  0)) +
+			  LookupSceneColor(UVSceneColor, int2( 0,  1)) +
+			  LookupSceneColor(UVSceneColor, int2(-1,  0)) +
+			  LookupSceneColor(UVSceneColor, int2( 0, -1)));
+	  }
+	  else if(ReconstructMethod == 3)
+	  {
+		  // very good 
+		  float3 A = LookupSceneColor(UVSceneColor, int2( 1,  0));
+		  float3 B = LookupSceneColor(UVSceneColor, int2(-1,  0));
+		  float3 C = LookupSceneColor(UVSceneColor, int2( 0,  1));
+		  float3 D = LookupSceneColor(UVSceneColor, int2( 0, -1));
+
+		  // Luminance could be green channel only
+		  float a = Luminance(A);
+		  float b = Luminance(B);
+		  float c = Luminance(C);
+		  float d = Luminance(D);
+
+		  float ab = abs(a - b);
+		  float cd = abs(c - d);
+
+		  // take the average in the direction that avoids dither pattern
+		  Quant1 = 0.5f * lerp(A + B, C + D, ab > cd);
+	  }
+
+	  Ret.Diffuse = lerp(Quant1, Quant0, bChecker);
+	  Ret.Specular = lerp(Quant0, Quant1, bChecker);
+	}
+#else // SUBSURFACE_PROFILE_CHECKERBOARD
+	{
+	  // If we're not doing checkerboard encoding, we just need to read a single pixel and decode (combined diffuse/spec in RGB)
+	  float4 CenterSample = SubsurfaceInput0_Texture.SampleLevel(SharedSubsurfaceSampler0, UVSceneColor, 0);
+	  float3 CombinedColor = CenterSample.rgb;
+	  float DiffuseLuminance = CenterSample.a * View.PreExposure;
+  
+	  float CombinedLuminance = Luminance(CombinedColor);
+	  float DiffuseFactor = saturate(DiffuseLuminance / CombinedLuminance);
+	  float SpecularFactor = 1.0f - DiffuseFactor;
+  
+	  Ret.Diffuse = CombinedColor * DiffuseFactor;
+	  Ret.Specular = CombinedColor * SpecularFactor;
+	}
+#endif // !SUBSURFACE_PROFILE_CHECKERBOARD
+
+	return Ret;
+}
+
+// @param UVSceneColor for the full res rendertarget (BufferSize) e.g. SceneColor or GBuffers
+// @return .RGB Color that should be scattared, .A:1 for subsurface scattering material, 0 for not
+float4 SetupSubsurfaceForOnePixel(float2 UVSceneColor)
+{
+	float4 Ret = 0;
+
+	FScreenSpaceData ScreenSpaceData = GetScreenSpaceData(UVSceneColor);
+
+	BRANCH if (UseSubsurfaceProfile(ScreenSpaceData.GBuffer.ShadingModelID))
+	{
+		// '1' is lower quality but that is acceptable here
+		SDiffuseAndSpecular DiffuseAndSpecular = ReconstructLighting(UVSceneColor, 1);
+
+		Ret.rgb = DiffuseAndSpecular.Diffuse;
+
+		// it's a valid sample
+		Ret.a = 1;
+	}
+
+	return Ret;
+}
+
+void SetupPS(in noperspective float4 UVAndScreenPos : TEXCOORD0, out float4 OutColor : SV_Target0)
+{
+	float2 BufferUV = UVAndScreenPos.xy;
+
+#if SUBSURFACE_HALF_RES
+	// order aligned with Gather() hardware implementation
+	// RGB: color*A, A:weight 0 if no subsurface scattering
+
+	float4 A = SetupSubsurfaceForOnePixel(min(BufferUV + float2(-0.5,  0.5f) * SubsurfaceInput0_ExtentInverse, SubsurfaceInput0_UVViewportBilinearMax));
+	float4 B = SetupSubsurfaceForOnePixel(min(BufferUV + float2( 0.5,  0.5f) * SubsurfaceInput0_ExtentInverse, SubsurfaceInput0_UVViewportBilinearMax));
+	float4 C = SetupSubsurfaceForOnePixel(min(BufferUV + float2( 0.5, -0.5f) * SubsurfaceInput0_ExtentInverse, SubsurfaceInput0_UVViewportBilinearMax));
+	float4 D = SetupSubsurfaceForOnePixel(min(BufferUV + float2(-0.5, -0.5f) * SubsurfaceInput0_ExtentInverse, SubsurfaceInput0_UVViewportBilinearMax));
+
+	float4 Sum = (A + B) + (C + D);
+
+	float Div = 1.0f / max(Sum.a, 0.00001f);
+
+	OutColor.rgb = Sum.rgb * Div;
+
+	float4 FourDepth = GatherSceneDepth(BufferUV, SubsurfaceInput0_ExtentInverse);
+
+	// average all valid depth values to a single one
+	float SingleDepth = dot(FourDepth, float4(A.a, B.a, C.a, D.a)) * Div;
+
+	OutColor.a = SingleDepth;
+#else // SUBSURFACE_HALF_RES
+	OutColor = SetupSubsurfaceForOnePixel(BufferUV);
+	if(OutColor.a > 0)
+	{
+		OutColor.a = CalcSceneDepth(BufferUV);
+	}
+#endif // SUBSURFACE_HALF_RES
+}
+
+float RadiusRootFindingCM(float D, float RandomNumber, float X0)
+{
+	return RadiusRootFinding(D*0.01,RandomNumber,X0*0.01)*100.0f;
+}
+
+float GetPdfInCM(float Radius, float L, float S)
+{
+	return GetPdf(Radius *0.01f, L*0.01f, S);
+}
+
+float4 BurleyNormalizedSS(float2 BufferUV)
+{
+	FScreenSpaceData ScreenSpaceData = GetScreenSpaceData(BufferUV);
+
+#ifdef SUBSURFACE_SINGLE_PASS
+	float DepthAtDiscCenter = CalcSceneDepth(BufferUV);
+	uint ReconstructMethod = (SUBSURFACE_RECOMBINE_QUALITY == SUBSURFACE_RECOMBINE_QUALITY_HIGH) ? 3 : 1;
+#else
+	float DepthAtDiscCenter = Texture2DSample(SubsurfaceInput0_Texture, SubsurfaceSampler0, BufferUV).w;
+#endif
+
+	float4 OutColor = 0;
+	BRANCH if (DepthAtDiscCenter <= 0)
+	{
+		return OutColor;
+	}
+
+	FBurleyParameter BurleyParameter = GetBurleyParameters(ScreenSpaceData.GBuffer);
+	float DiffuseMeanFreePathForSampling = GetDiffuseMeanFreePathForSampling(BurleyParameter.DiffuseMeanFreePath);
+	float A = GetComponentForScalingFactorEstimation(BurleyParameter.SurfaceAlbedo);
+	uint SubsurfaceProfileInt = ExtractSubsurfaceProfileInt(ScreenSpaceData.GBuffer);
+	float3 BoundaryColorBleed = GetSubsurfaceProfileBoundaryColorBleed(ScreenSpaceData.GBuffer).xyz;
+
+#define LIGHT_PERPENDICULAR_DIFFUSE_SURFACE
+
+#ifdef LIGHT_PERPENDICULAR
+	float S = GetPerpendicularScalingFactor(A);
+	float3 S3D = GetPerpendicularScalingFactor3D(BurleyParameter.SurfaceAlbedo);
+#elif defined(LIGHT_DIFFUSESURFACE)
+	float S = GetDiffuseSurfaceScalingFactor(A);
+	float3 S3D = GetDiffuseSurfaceScalingFactor3D(BurleyParameter.SurfaceAlbedo);
+#elif defined(LIGHT_PERPENDICULAR_DIFFUSE_SURFACE)
+	float S = GetSearchLightDiffuseScalingFactor(A);
+	float3 S3D = GetSearchLightDiffuseScalingFactor(BurleyParameter.SurfaceAlbedo.rgb);
+#endif
+
+	int SeedStart = View.FrameNumber;
+	float3 WeightingFactor = 0.0f;
+	float4 RadianceAccumulated = float4(0.0f, 0.0f, 0.0f, 1.0f);
+	float Mask = 1.0f;
+
+	LOOP
+		for (int i = 0; i < BURLEY_NUM_SAMPLES; ++i)
+		{
+			//step 1: sample generation
+			//create an 2d disk sampling pattern (we can load from the disk as a texture or buffer).
+			int3  Seed3D = int3(BufferUV*SubsurfaceInput0_Extent, SeedStart++);
+			float2 Random0T1 = Generate2DRandomNumber(Seed3D);
+			FBurleySampleInfo BurleySampleInfo = GenerateSampleInfo(Random0T1, DiffuseMeanFreePathForSampling, S, i);
+
+			//step 2: get the light radiance and depth at the offset
+			//and estimate the scale from the random disk sampling space to sceen space.
+			float YScale = BURLEY_MM_2_CM / DepthAtDiscCenter; //from mm to cm, and from world to screen
+			float XScale = YScale * SubsurfaceInput0_ExtentInverse.x*SubsurfaceInput0_Extent.y;
+
+			//world unit to screen space unit, we use world scale to control
+			//the scaling of the radius unit when mean free path is fixed.
+			float2 UVOffset = float2(XScale, YScale)*BurleySampleInfo.Radius*BurleyParameter.WorldScale;
+			UVOffset.x *= cos(BurleySampleInfo.Theta);
+			UVOffset.y *= sin(BurleySampleInfo.Theta);
+
+			//Sampling
+			{
+				
+				float2 SampledDiscUV = BufferUV + UVOffset;
+
+#ifdef SUBSURFACE_SINGLE_PASS
+				SDiffuseAndSpecular SampledDiffuseAndSpecular = ReconstructLighting(SampledDiscUV, ReconstructMethod);
+				float4 SampledRadianceAndDepth = float4(SampledDiffuseAndSpecular.Diffuse, CalcSceneDepth(SampledDiscUV));
+#else
+				float4 SampledRadianceAndDepth = Texture2DSample(SubsurfaceInput0_Texture, SubsurfaceSampler0, SampledDiscUV);
+#endif
+				//Determine the tint color, if the sampling pixel is not subsurface, we use tint color
+				//to mask out the sampling.
+				uint LocalProfileInt = GetSubsurfaceProfileId(SampledDiscUV);
+				float3 TintColor = (LocalProfileInt == SubsurfaceProfileInt) ? 1.0f : BoundaryColorBleed;
+				Mask = UseSubsurfaceProfile(GetScreenSpaceData(SampledDiscUV).GBuffer.ShadingModelID)? 1 : 0;
+
+				//step 3: create the bilateral filtering weighted Distance between entry and exit.
+#define USE_BILATERAL_FILTERING
+
+#ifdef USE_BILATERAL_FILTERING
+				// 10.0/depth is the depth in world unit (cm)
+				float DeltaDepth = min(abs(10.0f / SampledRadianceAndDepth.w - 10.0f / DepthAtDiscCenter), 60000)*BURLEY_CM_2_MM;//in mm
+				float RadiusSampled = sqrt(BurleySampleInfo.Radius * BurleySampleInfo.Radius + DeltaDepth * DeltaDepth); //mm.
+//#define RESAMPLE_PDF
+#ifdef RESAMPLE_PDF
+				BurleySampleInfo.Pdf = GetPdf(RadiusSampled, DiffuseMeanFreePathForSample, S);
+#endif
+#else
+				float RadiusSampled = BurleySampleInfo.Radius;
+#endif				
+
+				//step 4: accumulate radiance from the diffusion profile rR(r)
+				//make sure the DiffuseMeanFreePath is not zero and in m.
+				float3 DiffusionProfile = GetDiffuseReflectProfileWithDiffuseMeanFreePath(BurleyParameter.SurfaceAlbedo, BurleyParameter.DiffuseMeanFreePath, S3D, RadiusSampled);
+				RadianceAccumulated.xyz += (DiffusionProfile / BurleySampleInfo.Pdf) * (SampledRadianceAndDepth.xyz*TintColor*Mask);
+
+#ifdef USE_BILATERAL_FILTERING
+				WeightingFactor += (DiffusionProfile / BurleySampleInfo.Pdf) * Mask;
+#endif
+			}
+
+		}
+
+	//0.99995f is a compensitation to make it energe conservation.
+	const float EnergyNormalization = 1.0f / 0.99995f;
+
+#if defined(RADIUS_SAMPLE_UNIFORM_DISK)
+	RadianceAccumulated.xyz *= (InvSampleCount*0.5 * 2 * PI);
+#elif !defined(USE_BILATERAL_FILTERING)
+	RadianceAccumulated.xyz *= (BURLEY_INV_NUM_SAMPLES*0.5 * 2 * PI) * EnergyNormalization;
+#else
+	//The added epsilon is used to avoid divid by zero.
+	RadianceAccumulated.xyz *= 1.0f / (WeightingFactor + 0.000001f)*EnergyNormalization;
+#endif
+
+	OutColor.xyz = RadianceAccumulated.xyz;
+	OutColor.w = 1.0f;// 1.0f;
+	return OutColor;
+}
+
+// input0 is created by the SetupPS shader
+void MainPS(noperspective float4 UVAndScreenPos : TEXCOORD0, out float4 OutColor : SV_Target0)
+{
+	
+	float2 BufferUV = UVAndScreenPos.xy;
+
+	//determine if we will use burley
+	FScreenSpaceData ScreenSpaceData = GetScreenSpaceData(BufferUV);
+	uint SubsurfaceProfileInt = ExtractSubsurfaceProfileInt(ScreenSpaceData.GBuffer);
+	bool bUseBurley = GetSubsurfaceProfileUseBurley(SubsurfaceProfileInt);
+
+	
+#if SUBSURFACE_PASS == 0 || SUBSURFACE_PASS == SUBSURFACE_PASS_ONE
+
+	BRANCH
+	if (bUseBurley)
+	{	//needs to perform burley normalized subsurface scattering
+		OutColor = BurleyNormalizedSS(BufferUV);
+		return;
+	}
+
+	// horizontal
+	float2 ViewportDirectionUV = float2(1, 0) * SUBSURFACE_RADIUS_SCALE;
+#else
+
+	BRANCH
+	if (bUseBurley)
+	{
+		//direct sample the color
+
+		OutColor = Texture2DSample(SubsurfaceInput0_Texture, SubsurfaceSampler0, BufferUV);
+		return;
+	}
+	// vertical
+	float2 ViewportDirectionUV = float2(0, 1) * SUBSURFACE_RADIUS_SCALE * (SubsurfaceInput0_Extent.x * SubsurfaceInput0_ExtentInverse.y);
+#endif
+
+	ViewportDirectionUV *= (SubsurfaceInput0_ViewportSize.x * SubsurfaceInput0_ExtentInverse.x);
+
+	OutColor = SSSSBlurPS(BufferUV, ViewportDirectionUV, false);
+
+#if SUBSURFACE_PASS == SUBSURFACE_DIRECTION_VERTICAL || SUBSURFACE_PASS == SUBSURFACE_PASS_TWO
+	// second pass prepares the setup from the recombine pass which doesn't need depth but wants to reconstruct the color
+	OutColor.a = GetMaskFromDepthInAlpha(OutColor.a);
+#endif
+}
+
+// Recombines the half res Subsurface filtered lighting contribution (upsampled and renormalized with the alpha)
+// with the SceneColor.
+void SubsurfaceRecombinePS(noperspective float4 UVAndScreenPos : TEXCOORD0, float4 SvPosition : SV_POSITION, out float4 OutColor : SV_Target0)
+{
+
+	float2 BufferUV = UVAndScreenPos.xy;
+
+	FScreenSpaceData ScreenSpaceData = GetScreenSpaceData(BufferUV);
+	
+	if (!UseSubsurfaceProfile(ScreenSpaceData.GBuffer.ShadingModelID))
+	{
+		OutColor = Texture2DSample(SubsurfaceInput0_Texture, SharedSubsurfaceSampler0, BufferUV);
+		return;
+	}
+
+	float3 SSSColor = float3(0, 0, 0);
+	float LerpFactor = 1;
+
+#if SUBSURFACE_RECOMBINE_MODE != SUBSURFACE_RECOMBINE_MODE_NO_SCATTERING
+#if SUBSURFACE_HALF_RES
+	// fade out subsurface scattering if radius is too small to be more crips (not blend with half resolution)
+	// minor quality improvement (faces are more detailed in distance)
+	LerpFactor = ComputeFullResLerp(ScreenSpaceData, BufferUV, SubsurfaceInput1_ExtentInverse);
+#endif // SUBSURFACE_HALF_RES
+
+	{
+		float4 SSSColorWithAlpha = Texture2DSample(SubsurfaceInput1_Texture, SharedSubsurfaceSampler1, BufferUV);
+
+		// renormalize to dilate RGB to fix half res upsampling artifacts
+		SSSColor = SSSColorWithAlpha.rgb / max(SSSColorWithAlpha.a, 0.00001f);
+	}
+#else // SUBSURFACE_RECOMBINE_MODE == SUBSURFACE_RECOMBINE_MODE_NO_SCATTERING
+	// Scalability requests no Scatter, but we still need to reconstruct a color
+	LerpFactor = 0;
+#endif // SUBSURFACE_RECOMBINE_MODE
+
+	// we multiply the base color later in to get more crips human skin textures (scanned data always has Subsurface included)
+	float3 StoredBaseColor = ScreenSpaceData.GBuffer.StoredBaseColor;
+	float StoredSpecular = ScreenSpaceData.GBuffer.StoredSpecular;
+
+	uint ReconstructMethod = (SUBSURFACE_RECOMBINE_QUALITY == SUBSURFACE_RECOMBINE_QUALITY_HIGH) ? 3 : 1;
+
+	SDiffuseAndSpecular DiffuseAndSpecular = ReconstructLighting(BufferUV, ReconstructMethod);
+	
+	float3 ExtractedNonSubsurface = DiffuseAndSpecular.Specular;
+
+	// asset specific color
+	float3 SubsurfaceColor = GetSubsurfaceProfileColor(ScreenSpaceData.GBuffer);
+	float3 FadedSubsurfaceColor = SubsurfaceColor * LerpFactor;
+
+	//@TODO: clean up the separableSSS color
+	// combine potentially half res with full res
+	float3 SubsurfaceLighting =  lerp(DiffuseAndSpecular.Diffuse, SSSColor, FadedSubsurfaceColor);
+
+	OutColor = float4(SubsurfaceLighting * StoredBaseColor + ExtractedNonSubsurface, 0); 
+}
+
+void SubsurfaceViewportCopyPS(noperspective float4 InUV : TEXCOORD0, out float4 OutColor : SV_Target0)
+{
+	OutColor = Texture2DSample(SubsurfaceInput0_Texture, SharedSubsurfaceSampler0, InUV.xy);
+}
+
+void SubsurfaceSinglePassPS(noperspective float4 InUV : TEXCOORD0, out float4 OutColor : SV_Target0)
+{
+	float2 BufferUV = InUV.xy;
+	FScreenSpaceData ScreenSpaceData = GetScreenSpaceData(BufferUV);
+	
+	float DepthAtDiscCenter = CalcSceneDepth(BufferUV);
+	
+	bool bCopyScene = DepthAtDiscCenter <= 0 || (!UseSubsurfaceProfile(ScreenSpaceData.GBuffer.ShadingModelID));
+#if SUBSURFACE_RECOMBINE_MODE == SUBSURFACE_RECOMBINE_MODE_NO_SCATTERING
+	bCopyScene = true;
+#endif
+
+	//Setup data we are going to use.
+	uint ReconstructMethod = (SUBSURFACE_RECOMBINE_QUALITY == SUBSURFACE_RECOMBINE_QUALITY_HIGH) ? 3 : 1;
+	SDiffuseAndSpecular DiffuseAndSpecular = ReconstructLighting(BufferUV, ReconstructMethod);
+	float4 StoredSceneColor = float4(DiffuseAndSpecular.Diffuse * ScreenSpaceData.GBuffer.StoredBaseColor + DiffuseAndSpecular.Specular,0);
+	float LerpFactor = 1.0f;
+	float4 MeanRadiance = float4(0.0f, 0.0f, 0.0f, 1.0f);
+
+	BRANCH
+		if (bCopyScene) {
+			OutColor = StoredSceneColor;
+			return;
+		}
+
+#if SUBSURFACE_RECOMBINE_MODE != SUBSURFACE_RECOMBINE_MODE_NO_SCATTERING
+	MeanRadiance = BurleyNormalizedSS(BufferUV);
+
+	//---------------------------------------------------------------------------
+	// use the same logic as combine
+	float3 SSSColor = MeanRadiance.rgb / max(MeanRadiance.a, 0.00001f);
+	
+	float3 SubsurfaceColor = GetSubsurfaceProfileColor(ScreenSpaceData.GBuffer);
+	float3 FadedSubsurfaceColor = SubsurfaceColor * LerpFactor;
+	//recombine with the normal and surface texture color
+	float3 SubsurfaceLighting = lerp(DiffuseAndSpecular.Diffuse, SSSColor, FadedSubsurfaceColor);
+	float3 ExtractedNonSubsurface = DiffuseAndSpecular.Specular;
+	float3 StoredBaseColor = ScreenSpaceData.GBuffer.StoredBaseColor;
+
+	OutColor = float4(SubsurfaceLighting * StoredBaseColor + ExtractedNonSubsurface,0);
+#endif
+}
\ No newline at end of file
diff --git a/Engine/Shaders/Private/SubsurfaceProfileCommon.ush b/Engine/Shaders/Private/SubsurfaceProfileCommon.ush
--- a/Engine/Shaders/Private/SubsurfaceProfileCommon.ush
+++ b/Engine/Shaders/Private/SubsurfaceProfileCommon.ush
@@ -9,7 +9,9 @@
 // NOTE: Changing offsets below requires updating all instances of #SSSS_CONSTANTS
 // TODO: This needs to be defined in a single place and shared between C++ and shaders!
 #define SSSS_SUBSURFACE_COLOR_OFFSET			0
-#define SSSS_TRANSMISSION_OFFSET				(SSSS_SUBSURFACE_COLOR_OFFSET+1)
+#define SSS__SURFACEALBEDO_OFFSET                (SSSS_SUBSURFACE_COLOR_OFFSET+1)
+#define SSS__DMFP_OFFSET                        (SSS__SURFACEALBEDO_OFFSET+1)
+#define SSSS_TRANSMISSION_OFFSET				(SSS__DMFP_OFFSET+1)
 #define SSSS_BOUNDARY_COLOR_BLEED_OFFSET		(SSSS_TRANSMISSION_OFFSET+1)
 #define SSSS_DUAL_SPECULAR_OFFSET				(SSSS_BOUNDARY_COLOR_BLEED_OFFSET+1)
 #define SSSS_KERNEL0_OFFSET						(SSSS_DUAL_SPECULAR_OFFSET+1)
diff --git a/Engine/Source/Runtime/Engine/Classes/Engine/SubsurfaceProfile.h b/Engine/Source/Runtime/Engine/Classes/Engine/SubsurfaceProfile.h
--- a/Engine/Source/Runtime/Engine/Classes/Engine/SubsurfaceProfile.h
+++ b/Engine/Source/Runtime/Engine/Classes/Engine/SubsurfaceProfile.h
@@ -13,6 +13,36 @@ struct FSubsurfaceProfileStruct
 {
 	GENERATED_USTRUCT_BODY()
 	
+	/**
+	* It should match The base color of the corresponding material as much as possible.
+	*/
+	UPROPERTY(Category = "Burley Normalized", EditAnywhere, BlueprintReadOnly, meta = (ClampMin = "0.01", UIMax = "1.0", ClampMax = "1.0", HideAlphaChannel))
+	FLinearColor SurfaceAlbedo;
+
+	/**
+	* Controls how far light goes into the subsurface in the Red, Green and Blue channel. It is scaled by Mean Free path distance.
+	*/
+	UPROPERTY(Category = "Burley Normalized", EditAnywhere, BlueprintReadOnly, meta = (ClampMin = "0.001", UIMax = "1.0", ClampMax = "1.0", HideAlphaChannel))
+		FLinearColor MeanFreePathColor; //MeanFreePathLength;
+	
+	/**
+	* Subsurface mean free path distance in world/unreal units (cm)
+	*/
+	UPROPERTY(Category = "Burley Normalized", EditAnywhere, BluePrintReadOnly, meta = (ClampMin = "0.1", UIMax = "10.0", ClampMax = "10.0"))
+		float MeanFreePathDistance;
+
+	/**
+	* Control the scale of world/unreal units (cm)
+	*/
+	UPROPERTY(Category = "Burley Normalized", EditAnywhere, BlueprintReadOnly, meta = (ClampMin = "0.1", UIMax = "50.0", ClampMax = "50.0"))
+		float WorldScale;
+
+	/**
+	* Effective only when Burley subsurface scattering is enabled in cmd.
+	*/
+	UPROPERTY(Category = "Burley Normalized", EditAnywhere, BlueprintReadOnly)
+	bool  bEnableBurley;
+
 	/** in world/unreal units (cm) */
 	UPROPERTY(Category = "SubsurfaceProfileStruct", EditAnywhere, BlueprintReadOnly, meta = (ClampMin = "0.1", UIMax = "50.0", ClampMax = "1000.0"))
 	float ScatterRadius;
@@ -72,6 +102,13 @@ struct FSubsurfaceProfileStruct
 		Roughness0 = 0.75f;
 		Roughness1 = 1.30f;
 		LobeMix = 0.85f;
+		bEnableBurley = false;
+		//match FalloffColor exactly for the default setting
+		SurfaceAlbedo = FLinearColor(0.1447f,0.0535f,0.0432f);
+		MeanFreePathColor = FLinearColor(0.5963f,0.2679f,0.2253f);
+		MeanFreePathDistance = 0.1f;
+		//match ScatterRadius
+		WorldScale = 1.2f;
 	}
 
 	void Invalidate()
@@ -88,6 +125,11 @@ struct FSubsurfaceProfileStruct
 		Roughness0 = 0.75f;
 		Roughness1 = 1.30f;
 		LobeMix = 0.85f;
+		bEnableBurley = false;
+		SurfaceAlbedo = FLinearColor(0, 0, 0);
+		MeanFreePathColor = FLinearColor(0, 0, 0);
+		MeanFreePathDistance = 0.1f;
+		WorldScale = 0.0f;
 	}
 };
 
@@ -106,6 +148,9 @@ class USubsurfaceProfile : public UObject
 	//~ Begin UObject Interface
 	virtual void BeginDestroy();
 	virtual void PostEditChangeProperty(struct FPropertyChangedEvent& PropertyChangedEvent);
+
+	// Upgrade parameters from Separable to Burley.
+	virtual void PostLoad();
 	//~ End UObject Interface
 };
 
diff --git a/Engine/Source/Runtime/Engine/Private/Rendering/BurleyNormalizedSSS.cpp b/Engine/Source/Runtime/Engine/Private/Rendering/BurleyNormalizedSSS.cpp
new file mode 100644
--- /dev/null
+++ b/Engine/Source/Runtime/Engine/Private/Rendering/BurleyNormalizedSSS.cpp
@@ -0,0 +1,137 @@
+// Copyright 1998-2019 Epic Games, Inc. All Rights Reserved.
+
+/*=============================================================================
+BurleyNormalizedSSS.cpp: Compute the transmition profile for Burley normalized SSS
+=============================================================================*/
+
+
+#include "Rendering/BurleyNormalizedSSS.h"
+
+inline float Burley_Profile(float r, float A,float S, float L)
+{   //2PIR(r)r
+	float D = 1 / S;
+	float R = r / L;
+	const float Inv8Pi = 1.0 / (8 * PI);
+	float NegRbyD = -R / D;
+	float RrDotR = A*FMath::Max((exp(NegRbyD) + exp(NegRbyD / 3.0)) / (D*L)*Inv8Pi, 0.0);
+	return 2 * PI* RrDotR;
+}
+
+inline FVector Burley_Profile(float r, FLinearColor SurfaceAlbedo, FVector ScalingFactor, FLinearColor DiffuseMeanFreePath)
+{  
+	return FVector( Burley_Profile(r, SurfaceAlbedo.R, ScalingFactor.X, DiffuseMeanFreePath.R) ,
+					Burley_Profile(r, SurfaceAlbedo.G, ScalingFactor.Y, DiffuseMeanFreePath.G) ,
+					Burley_Profile(r, SurfaceAlbedo.B, ScalingFactor.Z, DiffuseMeanFreePath.B));
+}
+
+//--------------------------------------------------------------------------
+//Map burley ColorFallOff to Burley SurfaceAlbedo and diffuse mean free path.
+void MapFallOffColor2SurfaceAlbedoAndDiffuseMeanFreePath(float FalloffColor, float& SurfaceAlbedo, float& DiffuseMeanFreePath)
+{
+	//@TODO, use picewise function to separate Falloffcolor to around (0,0.2) and (0.2, 1) to make it more correct
+	//map Falloffcolor to SurfaceAlbedo with 4 polynomial, error < 2e-3.
+	float X = FalloffColor;
+	float X2 = X * X;
+	float X4 = X2 * X2;
+	SurfaceAlbedo = 0.06698*X4 - 0.1405*X2*X + 0.09174*X2 + 0.1278*X - 0.0007324;
+
+	//map FalOffColor to L with 6 polynomial, error < 6e-3  
+	DiffuseMeanFreePath = 7.674*X4*X2 - 25.13*X4*X + 31.6*X4 - 18.77*X2*X + 4.954*X2 + 0.2557*X + 0.01;
+}
+
+//map SurfaceAlbedo and  DiffuseMeanFreePath to Falloff, Actually, we only need to use SurfaceAlbedo, as Diffuse mean free path 
+//is coupled with SurfaceAlbedo. @TODO: map arbitrary SurfaceAlbedo and DiffuseMeanFreePath to FallOffColor and lerp factor.
+float MapAlbedoAndDiffuseMeanFreePath2FallOffColor(float SurfaceAlbedo, float&FallOffColor)
+{
+	float X = FallOffColor;
+	float X2 = X * X;
+	float X4 = X2 * X2;
+	//split 0.04,
+	if (SurfaceAlbedo < 0.04)
+	{
+		//6 polynomial, error <3e-3
+		return 4.261e+09*X4*X2 + -4.932e+08*X4*X + 2.195e+07*X4 - 4.624e+05*X2*X + 4448 * X2 - 7.119*X + 0.00699;
+	}
+	else
+	{
+		//error <1e-3
+		return 0.8498*X2 + 6.738*X + 0.007304;
+	}
+}
+
+//-----------------------------------------------------------------
+// Functions should be identical on both cpu and gpu
+// Method 1: The light directly goes into the volume in a direction perpendicular to the surface.
+// Average relative error: 5.5% (reference to MC)
+float GetPerpendicularScalingFactor(float SurfaceAlbedo)
+{
+	return 1.85 - SurfaceAlbedo + 7 * FMath::Pow(SurfaceAlbedo - 0.8, 3);
+}
+
+FVector GetPerpendicularScalingFactor(FLinearColor SurfaceAlbedo)
+{
+	return FVector(GetPerpendicularScalingFactor(SurfaceAlbedo.R),
+		GetPerpendicularScalingFactor(SurfaceAlbedo.G),
+		GetPerpendicularScalingFactor(SurfaceAlbedo.B)
+	);
+}
+
+// Method 2: Ideal diffuse transmission at the surface. More appropriate for rough surface.
+// Average relative error: 3.9% (reference to MC)
+float GetDiffuseSurfaceScalingFactor(float SurfaceAlbedo)
+{
+	return 1.9 - SurfaceAlbedo + 3.5 * FMath::Pow(SurfaceAlbedo - 0.8, 2);
+}
+
+FVector GetDiffuseSurfaceScalingFactor(FLinearColor SurfaceAlbedo)
+{
+	return FVector(GetDiffuseSurfaceScalingFactor(SurfaceAlbedo.R),
+		GetDiffuseSurfaceScalingFactor(SurfaceAlbedo.G),
+		GetDiffuseSurfaceScalingFactor(SurfaceAlbedo.B)
+	);
+}
+
+// Method 3: The spectral of diffuse mean free path on the surface.
+// Avergate relative error: 7.7% (reference to MC)
+float GetSearchLightDiffuseScalingFactor(float SurfaceAlbedo)
+{
+	return 3.5 + 100 * FMath::Pow(SurfaceAlbedo - 0.33, 4);
+}
+
+FVector GetSearchLightDiffuseScalingFactor(FLinearColor SurfaceAlbedo)
+{
+	return FVector(GetSearchLightDiffuseScalingFactor(SurfaceAlbedo.R),
+		GetSearchLightDiffuseScalingFactor(SurfaceAlbedo.G),
+		GetSearchLightDiffuseScalingFactor(SurfaceAlbedo.B)
+	);
+}
+
+void ComputeTransmissionProfileBurley(FLinearColor* TargetBuffer, uint32 TargetBufferSize, FLinearColor SubsurfaceColor, 
+																		FLinearColor FalloffColor, float ExtinctionScale, 
+																		FLinearColor SurfaceAlbedo, FLinearColor DiffuseMeanFreePath)
+{
+	check(TargetBuffer);
+	check(TargetBufferSize > 0);
+
+	static float MaxTransmissionProfileDistance = 5.0f; // See MAX_TRANSMISSION_PROFILE_DISTANCE in TransmissionCommon.ush
+	FVector ScalingFactor = GetSearchLightDiffuseScalingFactor(SurfaceAlbedo);
+	for (uint32 i = 0; i < TargetBufferSize; ++i)
+	{
+		//10 mm
+		const float InvSize = 1.0f / TargetBufferSize;
+		float Distance = i * InvSize * MaxTransmissionProfileDistance;
+
+		FVector TransmissionProfile = Burley_Profile(Distance, SurfaceAlbedo,ScalingFactor, DiffuseMeanFreePath);
+		TargetBuffer[i] = TransmissionProfile;
+		//Use Luminance of scattering as SSSS shadow.
+		TargetBuffer[i].A = exp(-Distance * ExtinctionScale);
+	}
+
+	// Do this is because 5mm is not enough cool down the scattering to zero, although which is small number but after tone mapping still noticeable
+	// so just Let last pixel be 0 which make sure thickness great than MaxRadius have no scattering
+	static bool bMakeLastPixelBlack = true;
+	if (bMakeLastPixelBlack)
+	{
+		TargetBuffer[TargetBufferSize - 1] = FLinearColor::Black;
+	}
+}
diff --git a/Engine/Source/Runtime/Engine/Private/Rendering/BurleyNormalizedSSS.h b/Engine/Source/Runtime/Engine/Private/Rendering/BurleyNormalizedSSS.h
new file mode 100644
--- /dev/null
+++ b/Engine/Source/Runtime/Engine/Private/Rendering/BurleyNormalizedSSS.h
@@ -0,0 +1,24 @@
+// Copyright 1998-2019 Epic Games, Inc. All Rights Reserved.
+
+/*=============================================================================
+BurleyNormalizedSSS.h: Compute the transmission profile and convert parameters
+=============================================================================*/
+
+#pragma once
+
+#include "CoreMinimal.h"
+
+// @param TargetBuffer, needs to be preallocated with TargetBufferSize (RGB is the sample weight, A is the offset), [0] is the center samples, following elements need to be mirrored with A, -A
+// @param TargetBufferSize >0
+// @parma SubsurfaceColor see SubsurfaceProfile.h
+// @parma FalloffColor see SubsurfaceProfile.h
+// @parma SurfaceAlbedo see SubsurfaceProfile.h
+// @parma DiffuseMeanFreePath mean free path color scaled by mean free path distance.
+void ComputeTransmissionProfileBurley(FLinearColor* TargetBuffer, uint32 TargetBufferSize, FLinearColor SubsurfaceColor, 
+									FLinearColor FalloffColor, float ExtinctionScale,
+									FLinearColor SurfaceAlbedo, FLinearColor DiffuseMeanFreePath);
+
+//@param FalloffColor from Separable SSS
+//@param SurfaceAlbedo mapped from FallofColor 
+//@param DiffuseMeanFreePath mapped from FallofColor
+void MapFallOffColor2SurfaceAlbedoAndDiffuseMeanFreePath(float FalloffColor, float& SurfaceAlbedo, float& DiffuseMeanFreePath);
\ No newline at end of file
diff --git a/Engine/Source/Runtime/Engine/Private/Rendering/SubsurfaceProfile.cpp b/Engine/Source/Runtime/Engine/Private/Rendering/SubsurfaceProfile.cpp
--- a/Engine/Source/Runtime/Engine/Private/Rendering/SubsurfaceProfile.cpp
+++ b/Engine/Source/Runtime/Engine/Private/Rendering/SubsurfaceProfile.cpp
@@ -4,6 +4,7 @@
 #include "RenderingThread.h"
 #include "RendererInterface.h"
 #include "Rendering/SeparableSSS.h"
+#include "Rendering/BurleyNormalizedSSS.h"
 #include "EngineModule.h"
 #include "RenderTargetPool.h"
 
@@ -126,7 +127,9 @@ static float GetNextSmallerPositiveFloat(float x)
 // NOTE: Changing offsets below requires updating all instances of #SSSS_CONSTANTS
 // TODO: This needs to be defined in a single place and shared between C++ and shaders!
 #define SSSS_SUBSURFACE_COLOR_OFFSET			0
-#define SSSS_TRANSMISSION_OFFSET				(SSSS_SUBSURFACE_COLOR_OFFSET+1)
+#define SSS__SURFACEALBEDO_OFFSET                (SSSS_SUBSURFACE_COLOR_OFFSET+1)
+#define SSS__DMFP_OFFSET                        (SSS__SURFACEALBEDO_OFFSET+1)
+#define SSSS_TRANSMISSION_OFFSET				(SSS__DMFP_OFFSET+1)
 #define SSSS_BOUNDARY_COLOR_BLEED_OFFSET		(SSSS_TRANSMISSION_OFFSET+1)
 #define SSSS_DUAL_SPECULAR_OFFSET				(SSSS_BOUNDARY_COLOR_BLEED_OFFSET+1)
 #define SSSS_KERNEL0_OFFSET						(SSSS_DUAL_SPECULAR_OFFSET+1)
@@ -141,6 +144,55 @@ static float GetNextSmallerPositiveFloat(float x)
 #define	SSSS_MAX_TRANSMISSION_PROFILE_DISTANCE	5.0f // See MaxTransmissionProfileDistance in ComputeTransmissionProfile(), SeparableSSS.cpp
 #define	SSSS_MAX_DUAL_SPECULAR_ROUGHNESS		2.0f
 
+//------------------------------------------------------------------------------------------
+// Consistent in BurleyNormalizedSSSCommon.ush and SubsurfaceProfile.cpp
+
+#define SSS_TYPE_SSSS		1.0f
+#define SSS_TYPE_BURLEY	    0.0f
+
+//make sure UIMax|ClampMax of WorldScale * ENC_WORLDSCALE_TO_UNIT <= 1
+#define ENC_WORLDSCALE_IN_CM_TO_UNIT 0.02f
+#define DEC_UNIT_TO_WORLDSCALE_IN_CM 1/ENC_WORLDSCALE_IN_CM_TO_UNIT
+
+//make sure UIMax|ClampMax of DiffuseMeanFreePath * ENC_DIFFUSEMEANFREEPATH_IN_MM_TO_UNIT <= 1
+#define ENC_DIFFUSEMEANFREEPATH_IN_MM_TO_UNIT 0.01f
+#define DEC_UNIT_TO_DIFFUSEMEANFREEPATH_IN_MM 1/ENC_DIFFUSEMEANFREEPATH_IN_MM_TO_UNIT
+//------------------------------------------------------------------------------------------
+
+//in [0,1]
+float EncodeWorldScale(float WorldScale)
+{
+	return WorldScale * ENC_WORLDSCALE_IN_CM_TO_UNIT;
+}
+
+float DecodeWorldScale(float EncodedWorldScale)
+{
+	return EncodedWorldScale * DEC_UNIT_TO_WORLDSCALE_IN_CM;
+}
+
+//in [0,1]
+FLinearColor EncodeDiffuseMeanFreePath(FLinearColor DiffuseMeanFreePath)
+{
+	return DiffuseMeanFreePath * ENC_DIFFUSEMEANFREEPATH_IN_MM_TO_UNIT;
+}
+
+FLinearColor DecodeDiffuseMeanFreePath(FLinearColor EncodedDiffuseMeanFreePath)
+{
+	return EncodedDiffuseMeanFreePath * DEC_UNIT_TO_DIFFUSEMEANFREEPATH_IN_MM;
+}
+
+void SetupSurfaceAlbedoAndDiffuseMeanFreePath(FLinearColor& SurfaceAlbedo, FLinearColor& Dmfp)
+{
+	//Store the value that corresponds to the largest Dmfp (diffuse mean free path) channel to A channel.
+	//This is an optimization to shift finding the max correspondence workload
+	//to CPU.
+	const float MaxDmfpComp = FMath::Max3(Dmfp.R, Dmfp.G, Dmfp.B);
+	const uint32 IndexOfMaxDmfp = (Dmfp.R == MaxDmfpComp) ? 0 : ((Dmfp.B == MaxDmfpComp) ? 1 : 2);
+
+	SurfaceAlbedo.A = SurfaceAlbedo.Component(IndexOfMaxDmfp);
+	Dmfp.A = MaxDmfpComp;
+}
+
 float Sqrt2(float X)
 {
 	return sqrtf(sqrtf(X));
@@ -182,6 +234,14 @@ void FSubsurfaceProfileTexture::CreateTexture(FRHICommandListImmediate& RHICmdLi
 	const float FloatScale = GetNextSmallerPositiveFloat(0x10000);
 	check((int32)GetNextSmallerPositiveFloat(0x10000) == 0xffff);
 
+	const auto CVarSSSType = IConsoleManager::Get().FindTConsoleVariableDataInt(TEXT("r.SubsurfaceScattering.Type"));
+	check(CVarSSSType);
+	const bool bBurleyIsTurnedOn = CVarSSSType->GetValueOnAnyThread() == 1;
+
+	const auto CVarPassType = IConsoleManager::Get().FindTConsoleVariableDataInt(TEXT("r.SSS.Burley.PassType"));
+	check(CVarPassType);
+	const bool bSinglePass = CVarPassType->GetValueOnAnyThread() == 1;
+
 	for (uint32 y = 0; y < Height; ++y)
 	{
 		FSubsurfaceProfileStruct Data = SubsurfaceProfileEntries[y].Settings;
@@ -193,13 +253,22 @@ void FSubsurfaceProfileTexture::CreateTexture(FRHICommandListImmediate& RHICmdLi
 
 		Data.SubsurfaceColor = Data.SubsurfaceColor.GetClamped();
 		Data.FalloffColor = Data.FalloffColor.GetClamped(Bias);
-
+		Data.MeanFreePathColor = Data.MeanFreePathColor.GetClamped(Bias);
+		
 		// to allow blending of the Subsurface with fullres in the shader
 		TextureRow[SSSS_SUBSURFACE_COLOR_OFFSET] = Data.SubsurfaceColor;
-		TextureRow[SSSS_SUBSURFACE_COLOR_OFFSET].A = 0; // unused
+		TextureRow[SSSS_SUBSURFACE_COLOR_OFFSET].A = EncodeWorldScale(Data.WorldScale);
+		
+		FLinearColor DifffuseMeanFreePath = Data.MeanFreePathColor*Data.MeanFreePathDistance*10.0f; // convert cm to mm.
+		SetupSurfaceAlbedoAndDiffuseMeanFreePath(Data.SurfaceAlbedo, DifffuseMeanFreePath);
+		TextureRow[SSS__SURFACEALBEDO_OFFSET] = Data.SurfaceAlbedo;
+		TextureRow[SSS__DMFP_OFFSET] = EncodeDiffuseMeanFreePath(DifffuseMeanFreePath);
 
 		TextureRow[SSSS_BOUNDARY_COLOR_BLEED_OFFSET] = Data.BoundaryColorBleed;
 
+		bool bUseBurley = bSinglePass ? bBurleyIsTurnedOn : (bBurleyIsTurnedOn && Data.bEnableBurley);
+		TextureRow[SSSS_BOUNDARY_COLOR_BLEED_OFFSET].A = bUseBurley ? SSS_TYPE_BURLEY : SSS_TYPE_SSSS;
+
 		float MaterialRoughnessToAverage = Data.Roughness0 * (1.0f - Data.LobeMix) + Data.Roughness1 * Data.LobeMix;
 		float AverageToRoughness0 = Data.Roughness0 / MaterialRoughnessToAverage;
 		float AverageToRoughness1 = Data.Roughness1 / MaterialRoughnessToAverage;
@@ -219,8 +288,15 @@ void FSubsurfaceProfileTexture::CreateTexture(FRHICommandListImmediate& RHICmdLi
 		ComputeMirroredSSSKernel(&TextureRow[SSSS_KERNEL1_OFFSET], SSSS_KERNEL1_SIZE, Data.SubsurfaceColor, Data.FalloffColor);
 		ComputeMirroredSSSKernel(&TextureRow[SSSS_KERNEL2_OFFSET], SSSS_KERNEL2_SIZE, Data.SubsurfaceColor, Data.FalloffColor);
 
-		ComputeTransmissionProfile(&TextureRow[SSSS_TRANSMISSION_PROFILE_OFFSET], SSSS_TRANSMISSION_PROFILE_SIZE, Data.SubsurfaceColor, Data.FalloffColor, Data.ExtinctionScale);
-
+		if (bUseBurley)
+		{
+			ComputeTransmissionProfileBurley(&TextureRow[SSSS_TRANSMISSION_PROFILE_OFFSET], SSSS_TRANSMISSION_PROFILE_SIZE, 
+				Data.SubsurfaceColor, Data.FalloffColor, Data.ExtinctionScale,Data.SurfaceAlbedo, DifffuseMeanFreePath);
+		}
+		else
+		{
+			ComputeTransmissionProfile(&TextureRow[SSSS_TRANSMISSION_PROFILE_OFFSET], SSSS_TRANSMISSION_PROFILE_SIZE, Data.SubsurfaceColor, Data.FalloffColor, Data.ExtinctionScale);
+		}
 		// could be lower than 1 (but higher than 0) to range compress for better quality (for 8 bit)
 		const float TableMaxRGB = 1.0f;
 		const float TableMaxA = 3.0f;
@@ -306,12 +382,17 @@ bool FSubsurfaceProfileTexture::GetEntryString(uint32 Index, FString& Out) const
 	const FSubsurfaceProfileStruct& ref = SubsurfaceProfileEntries[Index].Settings;
 
 
-	Out = FString::Printf(TEXT(" %c. %p ScatterRadius=%.1f, SubsurfaceColor=%.1f %.1f %.1f, FalloffColor=%.1f %.1f %.1f"), 
+	Out = FString::Printf(TEXT(" %c. %p ScatterRadius=%.1f, SubsurfaceColor=%.1f %.1f %.1f, FalloffColor=%.1f %.1f %.1f, \
+								SurfaceAlbedo=%.1f %.1f %.1f, MeanFreePathColor=%.1f %.1f %.1f, MeanFreePathDistance=%.1f, WorldScale=%.1f"), 
 		MiniFontCharFromIndex(Index), 
 		SubsurfaceProfileEntries[Index].Profile,
 		ref.ScatterRadius,
 		ref.SubsurfaceColor.R, ref.SubsurfaceColor.G, ref.SubsurfaceColor.B,
-		ref.FalloffColor.R, ref.FalloffColor.G, ref.FalloffColor.B);
+		ref.FalloffColor.R, ref.FalloffColor.G, ref.FalloffColor.B,
+		ref.SurfaceAlbedo.R, ref.SurfaceAlbedo.G, ref.SurfaceAlbedo.B,
+		ref.MeanFreePathColor.R, ref.MeanFreePathColor.G, ref.MeanFreePathColor.B,
+		ref.MeanFreePathDistance,
+		ref.WorldScale);
 
 	return true;
 }
@@ -347,6 +428,15 @@ void FSubsurfaceProfileTexture::Dump()
 				SubsurfaceProfileEntries[i].Settings.SubsurfaceColor.R, SubsurfaceProfileEntries[i].Settings.SubsurfaceColor.G, SubsurfaceProfileEntries[i].Settings.SubsurfaceColor.B);
 			UE_LOG(LogSubsurfaceProfile, Log, TEXT("     FalloffColor=%f %f %f"),
 				SubsurfaceProfileEntries[i].Settings.FalloffColor.R, SubsurfaceProfileEntries[i].Settings.FalloffColor.G, SubsurfaceProfileEntries[i].Settings.FalloffColor.B);
+
+			UE_LOG(LogSubsurfaceProfile, Log, TEXT("     SurfaceAlbedo=%f %f %f"),
+				SubsurfaceProfileEntries[i].Settings.SurfaceAlbedo.R, SubsurfaceProfileEntries[i].Settings.SurfaceAlbedo.G, SubsurfaceProfileEntries[i].Settings.SurfaceAlbedo.B);
+			UE_LOG(LogSubsurfaceProfile, Log, TEXT("     MeanFreePathColor=%f %f %f"),
+				SubsurfaceProfileEntries[i].Settings.MeanFreePathColor.R, SubsurfaceProfileEntries[i].Settings.MeanFreePathColor.G, SubsurfaceProfileEntries[i].Settings.MeanFreePathColor.B);
+			UE_LOG(LogSubsurfaceProfile, Log, TEXT("     MeanFreePathDistance=%f"),
+				SubsurfaceProfileEntries[i].Settings.MeanFreePathDistance);
+			UE_LOG(LogSubsurfaceProfile, Log, TEXT("     WorldScale=%f"),
+				SubsurfaceProfileEntries[i].Settings.WorldScale);
 		}
 	}
 
@@ -393,3 +483,22 @@ void USubsurfaceProfile::PostEditChangeProperty(struct FPropertyChangedEvent& Pr
 			GSubsurfaceProfileTextureObject.UpdateProfile(SettingsLocal, Profile);
 		});
 }
+
+void USubsurfaceProfile::PostLoad()
+{
+	Super::PostLoad();
+
+	const auto CVar = IConsoleManager::Get().
+		FindTConsoleVariableDataInt(TEXT("r.SSS.Burley.AlwaysUpdateParametersFromSeparable"));
+	check(CVar);
+	
+	const bool bUpdateBurleyParametersFromSeparable = CVar->GetValueOnAnyThread() == 1;
+
+	if (bUpdateBurleyParametersFromSeparable)
+	{
+		MapFallOffColor2SurfaceAlbedoAndDiffuseMeanFreePath(Settings.FalloffColor.R, Settings.SurfaceAlbedo.R, Settings.MeanFreePathColor.R);
+		MapFallOffColor2SurfaceAlbedoAndDiffuseMeanFreePath(Settings.FalloffColor.G, Settings.SurfaceAlbedo.G, Settings.MeanFreePathColor.G);
+		MapFallOffColor2SurfaceAlbedoAndDiffuseMeanFreePath(Settings.FalloffColor.B, Settings.SurfaceAlbedo.B, Settings.MeanFreePathColor.B);
+		Settings.WorldScale = Settings.ScatterRadius;
+	}
+}
\ No newline at end of file
diff --git a/Engine/Source/Runtime/Renderer/Private/PostProcess/PostProcessSubsurface.cpp b/Engine/Source/Runtime/Renderer/Private/PostProcess/PostProcessSubsurface.cpp
--- a/Engine/Source/Runtime/Renderer/Private/PostProcess/PostProcessSubsurface.cpp
+++ b/Engine/Source/Runtime/Renderer/Private/PostProcess/PostProcessSubsurface.cpp
@@ -6,128 +6,35 @@
 
 #include "PostProcess/PostProcessSubsurface.h"
 #include "PostProcess/SceneRenderTargets.h"
+#include "PostProcess/SubsurfaceCommon.h"
+#include "PostProcess/SubsurfaceBurleyNormalized.h"
+#include "PostProcess/SubsurfaceSeparable.h"
 #include "Engine/SubsurfaceProfile.h"
 #include "CanvasTypes.h"
 #include "RenderTargetTemp.h"
 
-ENGINE_API IPooledRenderTarget* GetSubsufaceProfileTexture_RT(FRHICommandListImmediate& RHICmdList);
-
 namespace
 {
-TAutoConsoleVariable<int32> CVarSubsurfaceScattering(
-	TEXT("r.SubsurfaceScattering"),
-	1,
-	TEXT(" 0: disabled\n")
-	TEXT(" 1: enabled (default)"),
-	ECVF_RenderThreadSafe | ECVF_Scalability);
-
-TAutoConsoleVariable<float> CVarSSSScale(
-	TEXT("r.SSS.Scale"),
-	1.0f,
-	TEXT("Affects the Screen space subsurface scattering pass")
-	TEXT("(use shadingmodel SubsurfaceProfile, get near to the object as the default)\n")
-	TEXT("is human skin which only scatters about 1.2cm)\n")
-	TEXT(" 0: off (if there is no object on the screen using this pass it should automatically disable the post process pass)\n")
-	TEXT("<1: scale scatter radius down (for testing)\n")
-	TEXT(" 1: use given radius form the Subsurface scattering asset (default)\n")
-	TEXT(">1: scale scatter radius up (for testing)"),
-	ECVF_Scalability | ECVF_RenderThreadSafe);
-
-TAutoConsoleVariable<int32> CVarSSSHalfRes(
-	TEXT("r.SSS.HalfRes"),
-	1,
-	TEXT(" 0: full quality (not optimized, as reference)\n")
-	TEXT(" 1: parts of the algorithm runs in half resolution which is lower quality but faster (default)"),
-	ECVF_RenderThreadSafe | ECVF_Scalability);
-
-TAutoConsoleVariable<int32> CVarSSSQuality(
-	TEXT("r.SSS.Quality"),
-	0,
-	TEXT("Defines the quality of the recombine pass when using the SubsurfaceScatteringProfile shading model\n")
-	TEXT(" 0: low (faster, default)\n")
-	TEXT(" 1: high (sharper details but slower)\n")
-	TEXT("-1: auto, 1 if TemporalAA is disabled (without TemporalAA the quality is more noticable)"),
-	ECVF_RenderThreadSafe | ECVF_Scalability);
-
-TAutoConsoleVariable<int32> CVarSSSFilter(
-	TEXT("r.SSS.Filter"),
-	1,
-	TEXT("Defines the filter method for Screenspace Subsurface Scattering feature.\n")
-	TEXT(" 0: point filter (useful for testing, could be cleaner)\n")
-	TEXT(" 1: bilinear filter"),
-	ECVF_RenderThreadSafe | ECVF_Scalability);
-
-TAutoConsoleVariable<int32> CVarSSSSampleSet(
-	TEXT("r.SSS.SampleSet"),
-	2,
-	TEXT("Defines how many samples we use for Screenspace Subsurface Scattering feature.\n")
-	TEXT(" 0: lowest quality (6*2+1)\n")
-	TEXT(" 1: medium quality (9*2+1)\n")
-	TEXT(" 2: high quality (13*2+1) (default)"),
-	ECVF_RenderThreadSafe | ECVF_Scalability);
-
-TAutoConsoleVariable<int32> CVarSSSCheckerboard(
-	TEXT("r.SSS.Checkerboard"),
-	2,
-	TEXT("Enables or disables checkerboard rendering for subsurface profile rendering.\n")
-	TEXT("This is necessary if SceneColor does not include a floating point alpha channel (e.g 32-bit formats)\n")
-	TEXT(" 0: Disabled (high quality) \n")
-	TEXT(" 1: Enabled (low quality). Surface lighting will be at reduced resolution.\n")
-	TEXT(" 2: Automatic. Non-checkerboard lighting will be applied if we have a suitable rendertarget format\n"),
-	ECVF_RenderThreadSafe);
-} //! namespace
-
-enum class ESubsurfaceMode : uint32
-{
-	// Performs a full resolution scattering filter.
-	FullRes,
-
-	// Performs a half resolution scattering filter.
-	HalfRes,
-
-	// Reconstructs lighting, but does not perform scattering.
-	Bypass,
-
-	MAX
-};
-
-// Returns the [0, N] clamped value of the 'r.SSS.Scale' CVar.
-float GetSubsurfaceRadiusScale()
-{
-	static const auto CVar = IConsoleManager::Get().FindTConsoleVariableDataFloat(TEXT("r.SSS.Scale"));
-	check(CVar);
-
-	return FMath::Max(0.0f, CVar->GetValueOnRenderThread());
-}
-
-// Returns the current subsurface mode required by the current view.
-ESubsurfaceMode GetSubsurfaceModeForView(const FViewInfo& View)
-{
-	const float Radius = GetSubsurfaceRadiusScale();
-	const bool bShowSubsurfaceScattering = Radius > 0 && View.Family->EngineShowFlags.SubsurfaceScattering;
-
-	if (bShowSubsurfaceScattering)
-	{
-		const bool bHalfRes = CVarSSSHalfRes.GetValueOnRenderThread() != 0;
-		if (bHalfRes)
-		{
-			return ESubsurfaceMode::HalfRes;
-		}
-		else
-		{
-			return ESubsurfaceMode::FullRes;
-		}
-	}
-	else
-	{
-		return ESubsurfaceMode::Bypass;
-	}
+	//only add ubber subsurface console variables here
+	static TAutoConsoleVariable<int32> CVarSubsurfaceScatteringType(
+		TEXT("r.SubsurfaceScattering.Type"),
+		0,
+		TEXT("Subsurface type:\n")
+		TEXT(" 0. Separable SSS (default).\n")
+		TEXT(" 1. Burley Normalized.")
+		TEXT("Else. Separable SSS. Make sure we always have SSS"),
+		ECVF_Scalability | ECVF_RenderThreadSafe
+	);
 }
 
 bool IsSubsurfaceEnabled()
 {
-	const bool bEnabled = CVarSubsurfaceScattering.GetValueOnAnyThread() != 0;
-	const bool bHasScale = CVarSSSScale.GetValueOnAnyThread() > 0.0f;
+	static const auto CVarSubsurfaceScattering = IConsoleManager::Get().FindTConsoleVariableDataInt(TEXT("r.SubsurfaceScattering"));
+	static const auto CVarSSSScale = IConsoleManager::Get().FindTConsoleVariableDataFloat(TEXT("r.SSS.Scale"));
+	check(CVarSubsurfaceScattering);
+	check(CVarSSSScale);
+	const bool bEnabled = CVarSubsurfaceScattering->GetValueOnAnyThread() != 0;
+	const bool bHasScale = CVarSSSScale->GetValueOnAnyThread() > 0.0f;
 	return (bEnabled && bHasScale);
 }
 
@@ -162,107 +69,13 @@ uint32 GetSubsurfaceRequiredViewMask(const TArray<FViewInfo>& Views)
 
 bool IsSubsurfaceCheckerboardFormat(EPixelFormat SceneColorFormat)
 {
-	int CVarValue = CVarSSSCheckerboard.GetValueOnRenderThread();
-	if (CVarValue == 0)
-	{
+	// No checkerboard for Burley
+	if (CVarSubsurfaceScatteringType.GetValueOnRenderThread() == 1)
 		return false;
-	}
-	else if (CVarValue == 1)
-	{
-		return true;
-	}
-	else if (CVarValue == 2)
-	{
-		switch (SceneColorFormat)
-		{
-		case PF_A32B32G32R32F:
-		case PF_FloatRGBA:
-			return false;
-		default:
-			return true;
-		}
-	}
-	return true;
+	return IsSeparableSubsurfaceCheckerboardFormat(SceneColorFormat);
 }
 
-// Returns the SS profile texture with a black fallback texture if none exists yet.
-FRHITexture* GetSubsurfaceProfileTexture(FRHICommandListImmediate& RHICmdList)
-{
-	const IPooledRenderTarget* ProfileTextureTarget = GetSubsufaceProfileTexture_RT(RHICmdList);
-
-	if (!ProfileTextureTarget)
-	{
-		// No subsurface profile was used yet
-		ProfileTextureTarget = GSystemTextures.BlackDummy;
-	}
-
-	return ProfileTextureTarget->GetRenderTargetItem().ShaderResourceTexture;
-}
-
-// Set of common shader parameters shared by all subsurface shaders.
-BEGIN_SHADER_PARAMETER_STRUCT(FSubsurfaceParameters, )
-SHADER_PARAMETER(FVector4, SubsurfaceParams)
-	SHADER_PARAMETER_STRUCT_REF(FSceneTexturesUniformParameters, SceneUniformBuffer)
-	SHADER_PARAMETER_STRUCT_REF(FViewUniformShaderParameters, ViewUniformBuffer)
-	SHADER_PARAMETER_SAMPLER(SamplerState, BilinearTextureSampler)
-	SHADER_PARAMETER_TEXTURE(Texture2D, SSProfilesTexture)
-END_SHADER_PARAMETER_STRUCT()
-
-FSubsurfaceParameters GetSubsurfaceCommonParameters(FRHICommandListImmediate& RHICmdList, const FViewInfo& View)
-{
-	FSceneRenderTargets& SceneContext = FSceneRenderTargets::Get(RHICmdList);
-
-	const float DistanceToProjectionWindow = View.ViewMatrices.GetProjectionMatrix().M[0][0];
-	const float SSSScaleZ = DistanceToProjectionWindow * GetSubsurfaceRadiusScale();
-	const float SSSScaleX = SSSScaleZ / SUBSURFACE_KERNEL_SIZE * 0.5f;
-
-	FSubsurfaceParameters Parameters;
-	Parameters.SubsurfaceParams = FVector4(SSSScaleX, SSSScaleZ, 0, 0);
-	Parameters.ViewUniformBuffer = View.ViewUniformBuffer;
-	Parameters.SceneUniformBuffer = CreateSceneTextureUniformBuffer(
-		SceneContext, View.FeatureLevel, ESceneTextureSetupMode::All, EUniformBufferUsage::UniformBuffer_SingleFrame);
-	Parameters.BilinearTextureSampler = TStaticSamplerState<SF_Bilinear>::GetRHI();
-	Parameters.SSProfilesTexture = GetSubsurfaceProfileTexture(RHICmdList);
-	return Parameters;
-}
-
-// A shader parameter struct for a single subsurface input texture.
-BEGIN_SHADER_PARAMETER_STRUCT(FSubsurfaceInput, )
-	SHADER_PARAMETER_STRUCT_INCLUDE(FScreenPassTextureViewportParameters, Viewport)
-	SHADER_PARAMETER_RDG_TEXTURE(Texture2D, Texture)
-END_SHADER_PARAMETER_STRUCT()
-
-FSubsurfaceInput GetSubsurfaceInput(FRDGTextureRef Texture, const FScreenPassTextureViewportParameters& ViewportParameters)
-{
-	FSubsurfaceInput Input;
-	Input.Texture = Texture;
-	Input.Viewport = ViewportParameters;
-	return Input;
-}
-
-// Base class for a subsurface shader.
-class FSubsurfaceShader : public FGlobalShader
-{
-public:
-	static bool ShouldCompilePermutation(const FGlobalShaderPermutationParameters& Parameters)
-	{
-		return IsFeatureLevelSupported(Parameters.Platform, ERHIFeatureLevel::SM4);
-	}
-
-	static void ModifyCompilationEnvironment(const FGlobalShaderPermutationParameters& Parameters, FShaderCompilerEnvironment& OutEnvironment)
-	{
-		FGlobalShader::ModifyCompilationEnvironment(Parameters, OutEnvironment);
-		OutEnvironment.SetDefine(TEXT("SUBSURFACE_RADIUS_SCALE"), SUBSURFACE_RADIUS_SCALE);
-		OutEnvironment.SetDefine(TEXT("SUBSURFACE_KERNEL_SIZE"), SUBSURFACE_KERNEL_SIZE);
-	}
-
-	FSubsurfaceShader() = default;
-	FSubsurfaceShader(const ShaderMetaType::CompiledShaderInitializerType& Initializer)
-		: FGlobalShader(Initializer)
-	{}
-};
-
-// Encapsulates the post processing subsurface scattering pixel shader.
+// Encapsulates the post processing subsurface scattering common pixel shader.
 class FSubsurfaceVisualizePS : public FSubsurfaceShader
 {
 public:
@@ -280,139 +93,6 @@ public:
 
 IMPLEMENT_GLOBAL_SHADER(FSubsurfaceVisualizePS, "/Engine/Private/PostProcessSubsurface.usf", "VisualizePS", SF_Pixel);
 
-// Encapsulates the post processing subsurface scattering pixel shader.
-class FSubsurfaceSetupPS : public FSubsurfaceShader
-{
-public:
-	DECLARE_GLOBAL_SHADER(FSubsurfaceSetupPS);
-	SHADER_USE_PARAMETER_STRUCT(FSubsurfaceSetupPS, FSubsurfaceShader)
-
-	BEGIN_SHADER_PARAMETER_STRUCT(FParameters, )
-		SHADER_PARAMETER_STRUCT_INCLUDE(FSubsurfaceParameters, Subsurface)
-		SHADER_PARAMETER_STRUCT(FSubsurfaceInput, SubsurfaceInput0)
-		SHADER_PARAMETER_SAMPLER(SamplerState, SubsurfaceSampler0)
-		RENDER_TARGET_BINDING_SLOTS()
-	END_SHADER_PARAMETER_STRUCT()
-
-	class FDimensionHalfRes : SHADER_PERMUTATION_BOOL("SUBSURFACE_HALF_RES");
-	class FDimensionCheckerboard : SHADER_PERMUTATION_BOOL("SUBSURFACE_PROFILE_CHECKERBOARD");
-	using FPermutationDomain = TShaderPermutationDomain<FDimensionHalfRes, FDimensionCheckerboard>;
-};
-
-IMPLEMENT_GLOBAL_SHADER(FSubsurfaceSetupPS, "/Engine/Private/PostProcessSubsurface.usf", "SetupPS", SF_Pixel);
-
-// Shader for the SSS separable blur.
-class FSubsurfacePS : public FSubsurfaceShader
-{
-public:
-	DECLARE_GLOBAL_SHADER(FSubsurfacePS);
-	SHADER_USE_PARAMETER_STRUCT(FSubsurfacePS, FSubsurfaceShader);
-
-	BEGIN_SHADER_PARAMETER_STRUCT(FParameters, )
-		SHADER_PARAMETER_STRUCT_INCLUDE(FSubsurfaceParameters, Subsurface)
-		SHADER_PARAMETER_STRUCT(FSubsurfaceInput, SubsurfaceInput0)
-		SHADER_PARAMETER_SAMPLER(SamplerState, SubsurfaceSampler0)
-		RENDER_TARGET_BINDING_SLOTS()
-	END_SHADER_PARAMETER_STRUCT()
-
-	// Direction of the 1D separable filter.
-	enum class EDirection : uint32
-	{
-		Horizontal,
-		Vertical,
-		MAX
-	};
-
-	// Controls the quality (number of samples) of the blur kernel.
-	enum class EQuality : uint32
-	{
-		Low,
-		Medium,
-		High,
-		MAX
-	};
-
-	class FDimensionDirection : SHADER_PERMUTATION_ENUM_CLASS("SUBSURFACE_DIRECTION", EDirection);
-	class FDimensionQuality : SHADER_PERMUTATION_ENUM_CLASS("SUBSURFACE_QUALITY", EQuality);
-	using FPermutationDomain = TShaderPermutationDomain<FDimensionDirection, FDimensionQuality>;
-
-	// Returns the sampler state based on the requested SSS filter CVar setting.
-	static FRHISamplerState* GetSamplerState()
-	{
-		if (CVarSSSFilter.GetValueOnRenderThread())
-		{
-			return TStaticSamplerState<SF_Bilinear, AM_Border, AM_Border, AM_Border>::GetRHI();
-		}
-		else
-		{
-			return TStaticSamplerState<SF_Point, AM_Border, AM_Border, AM_Border>::GetRHI();
-		}
-	}
-
-	// Returns the SSS quality level requested by the SSS SampleSet CVar setting.
-	static EQuality GetQuality()
-	{
-		return static_cast<FSubsurfacePS::EQuality>(
-			FMath::Clamp(
-				CVarSSSSampleSet.GetValueOnRenderThread(),
-				static_cast<int32>(FSubsurfacePS::EQuality::Low),
-				static_cast<int32>(FSubsurfacePS::EQuality::High)));
-	}
-};
-
-IMPLEMENT_GLOBAL_SHADER(FSubsurfacePS, "/Engine/Private/PostProcessSubsurface.usf", "MainPS", SF_Pixel);
-
-// Encapsulates the post processing subsurface recombine pixel shader.
-class FSubsurfaceRecombinePS : public FSubsurfaceShader
-{
-	DECLARE_GLOBAL_SHADER(FSubsurfaceRecombinePS);
-	SHADER_USE_PARAMETER_STRUCT(FSubsurfaceRecombinePS, FSubsurfaceShader);
-
-	BEGIN_SHADER_PARAMETER_STRUCT(FParameters, )
-		SHADER_PARAMETER_STRUCT_INCLUDE(FSubsurfaceParameters, Subsurface)
-		SHADER_PARAMETER_STRUCT(FSubsurfaceInput, SubsurfaceInput0)
-		SHADER_PARAMETER_STRUCT(FSubsurfaceInput, SubsurfaceInput1)
-		SHADER_PARAMETER_SAMPLER(SamplerState, SubsurfaceSampler0)
-		SHADER_PARAMETER_SAMPLER(SamplerState, SubsurfaceSampler1)
-		RENDER_TARGET_BINDING_SLOTS()
-	END_SHADER_PARAMETER_STRUCT();
-
-	// Controls the quality of lighting reconstruction.
-	enum class EQuality : uint32
-	{
-		Low,
-		High,
-		MAX
-	};
-
-	class FDimensionMode : SHADER_PERMUTATION_ENUM_CLASS("SUBSURFACE_RECOMBINE_MODE", ESubsurfaceMode);
-	class FDimensionQuality : SHADER_PERMUTATION_ENUM_CLASS("SUBSURFACE_RECOMBINE_QUALITY", EQuality);
-	class FDimensionCheckerboard : SHADER_PERMUTATION_BOOL("SUBSURFACE_PROFILE_CHECKERBOARD");
-	using FPermutationDomain = TShaderPermutationDomain<FDimensionMode, FDimensionQuality, FDimensionCheckerboard>;
-
-	// Returns the Recombine quality level requested by the SSS Quality CVar setting.
-	static EQuality GetQuality(const FViewInfo& View)
-	{
-		const uint32 QualityCVar = CVarSSSQuality.GetValueOnRenderThread();
-
-		// Quality is forced to high when the CVar is set to 'auto' and TAA is NOT enabled.
-		// TAA improves quality through temporal filtering, making it less necessary to use
-		// high quality mode.
-		const bool bUseHighQuality = (QualityCVar == -1 && View.AntiAliasingMethod != AAM_TemporalAA);
-
-		if (QualityCVar == 1 || bUseHighQuality)
-		{
-			return EQuality::High;
-		}
-		else
-		{
-			return EQuality::Low;
-		}
-	}
-};
-
-IMPLEMENT_GLOBAL_SHADER(FSubsurfaceRecombinePS, "/Engine/Private/PostProcessSubsurface.usf", "SubsurfaceRecombinePS", SF_Pixel);
-
 // Encapsulates a simple copy pixel shader.
 class FSubsurfaceViewportCopyPS : public FSubsurfaceShader
 {
@@ -441,156 +121,17 @@ void ComputeSubsurfaceForView(
 	FRDGTextureRef SceneTextureOutput,
 	ERenderTargetLoadAction SceneTextureLoadAction)
 {
-	check(SceneTexture);
-	check(SceneTextureOutput);
-	check(SceneViewport.Extent == SceneTexture->Desc.Extent);
-
-	const FViewInfo& View = ScreenPassView.View;
-
-	const FSceneViewFamily* ViewFamily = View.Family;
-
-	const FRDGTextureDesc& SceneTextureDesc = SceneTexture->Desc;
-
-	const ESubsurfaceMode SubsurfaceMode = GetSubsurfaceModeForView(View);
-
-	const bool bHalfRes = (SubsurfaceMode == ESubsurfaceMode::HalfRes);
-
-	const bool bCheckerboard = IsSubsurfaceCheckerboardFormat(SceneTextureDesc.Format);
-
-	const uint32 ScaleFactor = bHalfRes ? 2 : 1;
-
-	/**
-	 * All subsurface passes within the screen-space subsurface effect can operate at half or full resolution,
-	 * depending on the subsurface mode. The values are precomputed and shared among all Subsurface textures.
-	 */
-	const FScreenPassTextureViewport SubsurfaceViewport = FScreenPassTextureViewport::CreateDownscaled(SceneViewport, ScaleFactor);
-
-	const FRDGTextureDesc SubsurfaceTextureDescriptor = FRDGTextureDesc::Create2DDesc(
-		SubsurfaceViewport.Extent,
-		PF_FloatRGBA,
-		FClearValueBinding(),
-		TexCreate_None,
-		TexCreate_RenderTargetable | TexCreate_ShaderResource,
-		false);
-	
-	const FSubsurfaceParameters SubsurfaceCommonParameters = GetSubsurfaceCommonParameters(GraphBuilder.RHICmdList, View);
-	const FScreenPassTextureViewportParameters SubsurfaceViewportParameters = GetScreenPassTextureViewportParameters(SubsurfaceViewport);
-	const FScreenPassTextureViewportParameters SceneViewportParameters = GetScreenPassTextureViewportParameters(SceneViewport);
-
-	FRDGTextureRef SetupTexture = SceneTexture;
-	FRDGTextureRef SubsurfaceTextureX = nullptr;
-	FRDGTextureRef SubsurfaceTextureY = nullptr;
-
-	FRHISamplerState* PointClampSampler = TStaticSamplerState<SF_Point, AM_Clamp, AM_Clamp, AM_Clamp>::GetRHI();
-	FRHISamplerState* BilinearBorderSampler = TStaticSamplerState<SF_Bilinear, AM_Border, AM_Border, AM_Border>::GetRHI();
+	// Compute subsurface based on cmd configuration
+	static const auto CVar = IConsoleManager::Get().FindTConsoleVariableDataInt(TEXT("r.SubsurfaceScattering.Type"));
+	check(CVar);
 
-	/**
-	 * When in bypass mode, the setup and convolution passes are skipped, but lighting
-	 * reconstruction is still performed in the recombine pass.
-	 */
-	if (SubsurfaceMode != ESubsurfaceMode::Bypass)
+	if (CVar->GetValueOnRenderThread() == 1)	// Use Burley normalized SSS
 	{
-		SetupTexture = GraphBuilder.CreateTexture(SubsurfaceTextureDescriptor, TEXT("SubsurfaceSetupTexture"));
-
-		// Setup pass outputs the diffuse scene color and depth in preparation for the scatter passes.
-		{
-			FSubsurfaceSetupPS::FParameters* PassParameters = GraphBuilder.AllocParameters<FSubsurfaceSetupPS::FParameters>();
-			PassParameters->Subsurface = SubsurfaceCommonParameters;
-			PassParameters->RenderTargets[0] = FRenderTargetBinding(SetupTexture, ERenderTargetLoadAction::ENoAction, ERenderTargetStoreAction::EStore);
-			PassParameters->SubsurfaceInput0 = GetSubsurfaceInput(SceneTexture, SceneViewportParameters);
-			PassParameters->SubsurfaceSampler0 = PointClampSampler;
-
-			FSubsurfaceSetupPS::FPermutationDomain PixelShaderPermutationVector;
-			PixelShaderPermutationVector.Set<FSubsurfaceSetupPS::FDimensionHalfRes>(bHalfRes);
-			PixelShaderPermutationVector.Set<FSubsurfaceSetupPS::FDimensionCheckerboard>(bCheckerboard);
-			TShaderMapRef<FSubsurfaceSetupPS> PixelShader(View.ShaderMap, PixelShaderPermutationVector);
-
-			/**
-			 * The subsurface viewport is intentionally used as both the target and texture viewport, even though the texture
-			 * is potentially double the size. This is to ensure that the source UVs map 1-to-1 with pixel centers of the target,
-			 * in order to ensure that the checkerboard pattern selects the correct pixels from the scene texture. This still works
-			 * because the texture viewport is normalized into UV space, so it doesn't matter that the dimensions are twice as large.
-			 */
-			AddDrawScreenPass(GraphBuilder, RDG_EVENT_NAME("SubsurfaceSetup"), ScreenPassView, SubsurfaceViewport, SubsurfaceViewport, *PixelShader, PassParameters);
-		}
-
-		SubsurfaceTextureX = GraphBuilder.CreateTexture(SubsurfaceTextureDescriptor, TEXT("SubsurfaceTextureX"));
-		SubsurfaceTextureY = GraphBuilder.CreateTexture(SubsurfaceTextureDescriptor, TEXT("SubsurfaceTextureY"));
-
-		FRHISamplerState* SubsurfaceSamplerState = FSubsurfacePS::GetSamplerState();
-		const FSubsurfacePS::EQuality SubsurfaceQuality = FSubsurfacePS::GetQuality();
-
-		struct FSubsurfacePassInfo
-		{
-			FSubsurfacePassInfo(const TCHAR* InName, FRDGTextureRef InInput, FRDGTextureRef InOutput)
-				: Name(InName)
-				, Input(InInput)
-				, Output(InOutput)
-			{}
-
-			const TCHAR* Name;
-			FRDGTextureRef Input;
-			FRDGTextureRef Output;
-		};
-
-		const FSubsurfacePassInfo SubsurfacePassInfoByDirection[] =
-		{
-			{ TEXT("SubsurfaceX"), SetupTexture, SubsurfaceTextureX },
-			{ TEXT("SubsurfaceY"), SubsurfaceTextureX, SubsurfaceTextureY },
-		};
-
-		// Horizontal / Vertical scattering passes using a separable filter.
-		for (uint32 DirectionIndex = 0; DirectionIndex < static_cast<uint32>(FSubsurfacePS::EDirection::MAX); ++DirectionIndex)
-		{
-			const auto Direction = static_cast<FSubsurfacePS::EDirection>(DirectionIndex);
-
-			const FSubsurfacePassInfo& PassInfo = SubsurfacePassInfoByDirection[DirectionIndex];
-			FRDGTextureRef TextureInput = PassInfo.Input;
-			FRDGTextureRef TextureOutput = PassInfo.Output;
-
-			FSubsurfacePS::FParameters* PassParameters = GraphBuilder.AllocParameters<FSubsurfacePS::FParameters>();
-			PassParameters->Subsurface = SubsurfaceCommonParameters;
-			PassParameters->RenderTargets[0] = FRenderTargetBinding(TextureOutput, ERenderTargetLoadAction::ENoAction, ERenderTargetStoreAction::EStore);
-			PassParameters->SubsurfaceInput0 = GetSubsurfaceInput(TextureInput, SubsurfaceViewportParameters);
-			PassParameters->SubsurfaceSampler0 = SubsurfaceSamplerState;
-
-			FSubsurfacePS::FPermutationDomain PixelShaderPermutationVector;
-			PixelShaderPermutationVector.Set<FSubsurfacePS::FDimensionDirection>(Direction);
-			PixelShaderPermutationVector.Set<FSubsurfacePS::FDimensionQuality>(SubsurfaceQuality);
-			TShaderMapRef<FSubsurfacePS> PixelShader(View.ShaderMap, PixelShaderPermutationVector);
-
-			AddDrawScreenPass(GraphBuilder, FRDGEventName(PassInfo.Name), ScreenPassView, SubsurfaceViewport, SubsurfaceViewport, *PixelShader, PassParameters);
-		}
+		ComputeBurleySubsurfaceForView(GraphBuilder, ScreenPassView, SceneViewport, SceneTexture, SceneTextureOutput, SceneTextureLoadAction);
 	}
-
-	// Recombines scattering result with scene color.
+	else										// Use separable SSS.
 	{
-		FSubsurfaceRecombinePS::FParameters* PassParameters = GraphBuilder.AllocParameters<FSubsurfaceRecombinePS::FParameters>();
-		PassParameters->Subsurface = SubsurfaceCommonParameters;
-		PassParameters->RenderTargets[0] = FRenderTargetBinding(SceneTextureOutput, SceneTextureLoadAction, ERenderTargetStoreAction::EStore);
-		PassParameters->SubsurfaceInput0 = GetSubsurfaceInput(SceneTexture, SceneViewportParameters);
-		PassParameters->SubsurfaceSampler0 = BilinearBorderSampler;
-
-		// Scattering output target is only used when scattering is enabled.
-		if (SubsurfaceMode != ESubsurfaceMode::Bypass)
-		{
-			PassParameters->SubsurfaceInput1 = GetSubsurfaceInput(SubsurfaceTextureY, SubsurfaceViewportParameters);
-			PassParameters->SubsurfaceSampler1 = BilinearBorderSampler;
-		}
-
-		const FSubsurfaceRecombinePS::EQuality RecombineQuality = FSubsurfaceRecombinePS::GetQuality(View);
-
-		FSubsurfaceRecombinePS::FPermutationDomain PixelShaderPermutationVector;
-		PixelShaderPermutationVector.Set<FSubsurfaceRecombinePS::FDimensionMode>(SubsurfaceMode);
-		PixelShaderPermutationVector.Set<FSubsurfaceRecombinePS::FDimensionQuality>(RecombineQuality);
-		PixelShaderPermutationVector.Set<FSubsurfaceRecombinePS::FDimensionCheckerboard>(bCheckerboard);
-		TShaderMapRef<FSubsurfaceRecombinePS> PixelShader(View.ShaderMap, PixelShaderPermutationVector);
-
-		/**
-		 * See the related comment above in the prepare pass. The scene viewport is used as both the target and
-		 * texture viewport in order to ensure that the correct pixel is sampled for checkerboard rendering.
-		 */
-		AddDrawScreenPass(GraphBuilder, RDG_EVENT_NAME("SubsurfaceRecombine"), ScreenPassView, SceneViewport, SceneViewport, *PixelShader, PassParameters);
+		ComputeSeparableSubsurfaceForView(GraphBuilder, ScreenPassView, SceneViewport, SceneTexture, SceneTextureOutput, SceneTextureLoadAction);
 	}
 }
 
diff --git a/Engine/Source/Runtime/Renderer/Private/PostProcess/SubsurfaceBurleyNormalized.cpp b/Engine/Source/Runtime/Renderer/Private/PostProcess/SubsurfaceBurleyNormalized.cpp
new file mode 100644
--- /dev/null
+++ b/Engine/Source/Runtime/Renderer/Private/PostProcess/SubsurfaceBurleyNormalized.cpp
@@ -0,0 +1,463 @@
+
+// Copyright 1998-2019 Epic Games, Inc. All Rights Reserved.
+
+/*=============================================================================
+	SubsurfaceBurleyNormalized.cpp: Screenspace Burley subsurface scattering implementation.
+=============================================================================*/
+
+#include "PostProcess/SubsurfaceBurleyNormalized.h"
+#include "PostProcess/SceneRenderTargets.h"
+#include "PostProcess/SubsurfaceCommon.h"
+#include "Engine/SubsurfaceProfile.h"
+#include "CanvasTypes.h"
+#include "RenderTargetTemp.h"
+
+namespace
+{
+	TAutoConsoleVariable<int32> CVarSSSBurleyPassType(
+		TEXT("r.SSS.Burley.PassType"),
+		0,
+		TEXT("Select pass type to use.\n")
+		TEXT(" 0: Multipath for performance\n")
+		TEXT(" 1: Single pass for quality"),
+		ECVF_RenderThreadSafe | ECVF_Scalability
+	);
+
+	TAutoConsoleVariable<int32> CVarSSSBurleyUpdateParameter(
+		TEXT("r.SSS.Burley.AlwaysUpdateParametersFromSeparable"),
+		1,
+		TEXT("0: Will not update parameters when the program loads.")
+		TEXT("1: Always update from the separable when the program loads. (default)"),
+		ECVF_RenderThreadSafe | ECVF_Scalability
+	);
+}
+
+enum class ESubsurfaceBurleyPassType : uint32
+{
+	// Performs setup -> subsurface -> recombine three pass
+	Multiple,
+
+	// Performs a single pass
+	Single,
+
+	MAX
+};
+
+// Return the current subsurface pass type
+ESubsurfaceBurleyPassType GetSuburfaceBurleyPassType()
+{
+	const auto PassType = CVarSSSBurleyPassType.GetValueOnRenderThread();
+	if (PassType == 0)
+	{
+		return ESubsurfaceBurleyPassType::Multiple;
+	}
+	else
+	{
+		return ESubsurfaceBurleyPassType::Single;
+	}
+}
+
+// Set of common shader parameters shared by all subsurface shaders.
+BEGIN_SHADER_PARAMETER_STRUCT(FSubsurfaceBurleyParameters, )
+	SHADER_PARAMETER(FVector4, SubsurfaceParams)
+	SHADER_PARAMETER_STRUCT_REF(FSceneTexturesUniformParameters, SceneUniformBuffer)
+	SHADER_PARAMETER_STRUCT_REF(FViewUniformShaderParameters, ViewUniformBuffer)
+	SHADER_PARAMETER_SAMPLER(SamplerState, BilinearTextureSampler)
+	SHADER_PARAMETER_TEXTURE(Texture2D, SSProfilesTexture)
+END_SHADER_PARAMETER_STRUCT()
+
+FSubsurfaceBurleyParameters GetSubsurfaceBurleyCommonParameters(FRHICommandListImmediate& RHICmdList, const FViewInfo& View)
+{
+	FSceneRenderTargets& SceneContext = FSceneRenderTargets::Get(RHICmdList);
+
+	const float DistanceToProjectionWindow = View.ViewMatrices.GetProjectionMatrix().M[0][0];
+	const float SSSScaleZ = DistanceToProjectionWindow * GetSubsurfaceRadiusScale();
+	const float SSSScaleX = SSSScaleZ / SUBSURFACE_KERNEL_SIZE * 0.5f;
+
+	FSubsurfaceBurleyParameters Parameters;
+	Parameters.SubsurfaceParams = FVector4(SSSScaleX, SSSScaleZ, 0, 0);
+	Parameters.ViewUniformBuffer = View.ViewUniformBuffer;
+	Parameters.SceneUniformBuffer = CreateSceneTextureUniformBuffer(
+		SceneContext, View.FeatureLevel, ESceneTextureSetupMode::All, EUniformBufferUsage::UniformBuffer_SingleFrame);
+	Parameters.BilinearTextureSampler = TStaticSamplerState<SF_Bilinear>::GetRHI();
+	Parameters.SSProfilesTexture = GetSubsurfaceProfileTexture(RHICmdList);
+	return Parameters;
+}
+
+// A shader parameter struct for a single subsurface input texture.
+BEGIN_SHADER_PARAMETER_STRUCT(FSubsurfaceBurleyInput, )
+	SHADER_PARAMETER_STRUCT_INCLUDE(FScreenPassTextureViewportParameters, Viewport)
+	SHADER_PARAMETER_RDG_TEXTURE(Texture2D, Texture)
+END_SHADER_PARAMETER_STRUCT()
+
+FSubsurfaceBurleyInput GetSubsurfaceBurleyInput(FRDGTextureRef Texture, const FScreenPassTextureViewportParameters& ViewportParameters)
+{
+	FSubsurfaceBurleyInput Input;
+	Input.Texture = Texture;
+	Input.Viewport = ViewportParameters;
+	return Input;
+}
+
+// Encapsulates the post processing subsurface scattering pixel shader.
+class FSubsurfaceBurleySetupPS : public FSubsurfaceShader
+{
+public:
+	DECLARE_GLOBAL_SHADER(FSubsurfaceBurleySetupPS);
+	SHADER_USE_PARAMETER_STRUCT(FSubsurfaceBurleySetupPS, FSubsurfaceShader)
+
+	BEGIN_SHADER_PARAMETER_STRUCT(FParameters, )
+		SHADER_PARAMETER_STRUCT_INCLUDE(FSubsurfaceBurleyParameters, Subsurface)
+		SHADER_PARAMETER_STRUCT(FSubsurfaceBurleyInput, SubsurfaceInput0)
+		SHADER_PARAMETER_SAMPLER(SamplerState, SubsurfaceSampler0)
+		RENDER_TARGET_BINDING_SLOTS()
+	END_SHADER_PARAMETER_STRUCT()
+
+	class FDimensionHalfRes : SHADER_PERMUTATION_BOOL("SUBSURFACE_HALF_RES");
+	class FDimensionCheckerboard : SHADER_PERMUTATION_BOOL("SUBSURFACE_PROFILE_CHECKERBOARD");
+	using FPermutationDomain = TShaderPermutationDomain<FDimensionHalfRes, FDimensionCheckerboard>;
+};
+
+IMPLEMENT_GLOBAL_SHADER(FSubsurfaceBurleySetupPS, "/Engine/Private/SubsurfaceBurleyNormalized.usf", "SetupPS", SF_Pixel);
+
+// Shader for the SSS separable blur.
+class FSubsurfaceBurleyPS : public FSubsurfaceShader
+{
+public:
+	DECLARE_GLOBAL_SHADER(FSubsurfaceBurleyPS);
+	SHADER_USE_PARAMETER_STRUCT(FSubsurfaceBurleyPS, FSubsurfaceShader);
+
+	BEGIN_SHADER_PARAMETER_STRUCT(FParameters, )
+		SHADER_PARAMETER_STRUCT_INCLUDE(FSubsurfaceBurleyParameters, Subsurface)
+		SHADER_PARAMETER_STRUCT(FSubsurfaceBurleyInput, SubsurfaceInput0)
+		SHADER_PARAMETER_SAMPLER(SamplerState, SubsurfaceSampler0)
+		RENDER_TARGET_BINDING_SLOTS()
+	END_SHADER_PARAMETER_STRUCT()
+
+	// Direction of the 1D separable filter.
+	enum class EDirection : uint32
+	{
+		Horizontal,
+		Vertical,
+		MAX
+	};
+
+	enum class ESubsurfacePass : uint32
+	{
+		BURLEY,		// Horizontal
+		VARIANCE,	// Vertical
+		MAX
+	};
+
+	// Controls the quality (number of samples) of the blur kernel.
+	enum class EQuality : uint32
+	{
+		Low,
+		Medium,
+		High,
+		MAX
+	};
+
+	class FSubsurfacePassFunction : SHADER_PERMUTATION_ENUM_CLASS("SUBSURFACE_PASS", ESubsurfacePass);
+	class FDimensionQuality : SHADER_PERMUTATION_ENUM_CLASS("SUBSURFACE_QUALITY", EQuality);
+	using FPermutationDomain = TShaderPermutationDomain<FSubsurfacePassFunction, FDimensionQuality>;
+
+	// Returns the sampler state based on the requested SSS filter CVar setting.
+	static FRHISamplerState* GetSamplerState()
+	{
+		if (GetSSSFilter())
+		{
+			return TStaticSamplerState<SF_Bilinear, AM_Border, AM_Border, AM_Border>::GetRHI();
+		}
+		else
+		{
+			return TStaticSamplerState<SF_Point, AM_Border, AM_Border, AM_Border>::GetRHI();
+		}
+	}
+
+	// Returns the SSS quality level requested by the SSS SampleSet CVar setting.
+	static EQuality GetQuality()
+	{
+		return static_cast<FSubsurfaceBurleyPS::EQuality>(
+			FMath::Clamp(
+				GetSSSSampleSet(),
+				static_cast<int32>(FSubsurfaceBurleyPS::EQuality::Low),
+				static_cast<int32>(FSubsurfaceBurleyPS::EQuality::High)));
+	}
+};
+
+IMPLEMENT_GLOBAL_SHADER(FSubsurfaceBurleyPS, "/Engine/Private/SubsurfaceBurleyNormalized.usf", "MainPS", SF_Pixel);
+
+// Encapsulates the post processing subsurface recombine pixel shader.
+class FSubsurfaceBurleyRecombinePS : public FSubsurfaceShader
+{
+	DECLARE_GLOBAL_SHADER(FSubsurfaceBurleyRecombinePS);
+	SHADER_USE_PARAMETER_STRUCT(FSubsurfaceBurleyRecombinePS, FSubsurfaceShader);
+
+	BEGIN_SHADER_PARAMETER_STRUCT(FParameters, )
+		SHADER_PARAMETER_STRUCT_INCLUDE(FSubsurfaceBurleyParameters, Subsurface)
+		SHADER_PARAMETER_STRUCT(FSubsurfaceBurleyInput, SubsurfaceInput0)
+		SHADER_PARAMETER_STRUCT(FSubsurfaceBurleyInput, SubsurfaceInput1)
+		SHADER_PARAMETER_SAMPLER(SamplerState, SubsurfaceSampler0)
+		SHADER_PARAMETER_SAMPLER(SamplerState, SubsurfaceSampler1)
+		RENDER_TARGET_BINDING_SLOTS()
+	END_SHADER_PARAMETER_STRUCT();
+
+	// Controls the quality of lighting reconstruction.
+	enum class EQuality : uint32
+	{
+		Low,
+		High,
+		MAX
+	};
+
+	class FDimensionMode : SHADER_PERMUTATION_ENUM_CLASS("SUBSURFACE_RECOMBINE_MODE", ESubsurfaceMode);
+	class FDimensionQuality : SHADER_PERMUTATION_ENUM_CLASS("SUBSURFACE_RECOMBINE_QUALITY", EQuality);
+	class FDimensionCheckerboard : SHADER_PERMUTATION_BOOL("SUBSURFACE_PROFILE_CHECKERBOARD");
+	using FPermutationDomain = TShaderPermutationDomain<FDimensionMode, FDimensionQuality, FDimensionCheckerboard>;
+
+	// Returns the Recombine quality level requested by the SSS Quality CVar setting.
+	static EQuality GetQuality(const FViewInfo& View)
+	{
+		const uint32 QualityCVar = GetSSSQuality();
+
+		// Quality is forced to high when the CVar is set to 'auto' and TAA is NOT enabled.
+		// TAA improves quality through temporal filtering, making it less necessary to use
+		// high quality mode.
+		const bool bUseHighQuality = (QualityCVar == -1 && View.AntiAliasingMethod != AAM_TemporalAA);
+
+		if (QualityCVar == 1 || bUseHighQuality)
+		{
+			return EQuality::High;
+		}
+		else
+		{
+			return EQuality::Low;
+		}
+	}
+};
+
+IMPLEMENT_GLOBAL_SHADER(FSubsurfaceBurleyRecombinePS, "/Engine/Private/SubsurfaceBurleyNormalized.usf", "SubsurfaceRecombinePS", SF_Pixel);
+
+class FSubsurfaceBurleySinglePassPS : public FSubsurfaceShader
+{
+public:
+	static void ModifyCompilationEnvironment(const FGlobalShaderPermutationParameters& Parameters, FShaderCompilerEnvironment& OutEnvironment)
+	{
+		FGlobalShader::ModifyCompilationEnvironment(Parameters, OutEnvironment);
+		OutEnvironment.SetDefine(TEXT("SUBSURFACE_RADIUS_SCALE"), SUBSURFACE_RADIUS_SCALE);
+		OutEnvironment.SetDefine(TEXT("SUBSURFACE_KERNEL_SIZE"), SUBSURFACE_KERNEL_SIZE);
+		OutEnvironment.SetDefine(TEXT("SUBSURFACE_SINGLE_PASS"),1);
+	}
+
+	DECLARE_GLOBAL_SHADER(FSubsurfaceBurleySinglePassPS);
+	SHADER_USE_PARAMETER_STRUCT(FSubsurfaceBurleySinglePassPS, FSubsurfaceShader)
+    
+		BEGIN_SHADER_PARAMETER_STRUCT(FParameters, )
+		SHADER_PARAMETER_STRUCT_INCLUDE(FSubsurfaceBurleyParameters, Subsurface)
+		SHADER_PARAMETER_STRUCT(FSubsurfaceBurleyInput, SubsurfaceInput0)
+		//SHADER_PARAMETER_STRUCT(FSubsurfaceBurleyInput, SubsurfaceInput1)
+		SHADER_PARAMETER_SAMPLER(SamplerState, SubsurfaceSampler0)
+		//SHADER_PARAMETER_SAMPLER(SamplerState, SubsurfaceSampler1)
+		RENDER_TARGET_BINDING_SLOTS()
+		END_SHADER_PARAMETER_STRUCT()
+	
+	class FDimensionMode : SHADER_PERMUTATION_ENUM_CLASS("SUBSURFACE_RECOMBINE_MODE", ESubsurfaceMode);
+	class FDimensionHalfRes : SHADER_PERMUTATION_BOOL("SUBSURFACE_HALF_RES");
+	class FDimensionCheckerboard : SHADER_PERMUTATION_BOOL("SUBSURFACE_PROFILE_CHECKERBOARD");
+
+	using FPermutationDomain = TShaderPermutationDomain<FDimensionMode,FDimensionHalfRes, FDimensionCheckerboard>;
+};
+
+IMPLEMENT_GLOBAL_SHADER(FSubsurfaceBurleySinglePassPS, "/Engine/Private/SubsurfaceBurleyNormalized.usf", "SubsurfaceSinglePassPS", SF_Pixel);
+
+void ComputeBurleySubsurfaceForView(
+	FRDGBuilder& GraphBuilder,
+	const FScreenPassViewInfo& ScreenPassView,
+	const FScreenPassTextureViewport& SceneViewport,
+	FRDGTextureRef SceneTexture,
+	FRDGTextureRef SceneTextureOutput,
+	ERenderTargetLoadAction SceneTextureLoadAction)
+{
+	check(SceneTexture);
+	check(SceneTextureOutput);
+	check(SceneViewport.Extent == SceneTexture->Desc.Extent);
+
+	const FViewInfo& View = ScreenPassView.View;
+
+	const FSceneViewFamily* ViewFamily = View.Family;
+
+	const FRDGTextureDesc& SceneTextureDesc = SceneTexture->Desc;
+
+	const ESubsurfaceMode SubsurfaceMode = GetSubsurfaceModeForView(View);
+
+	const bool bHalfRes = (SubsurfaceMode == ESubsurfaceMode::HalfRes);
+
+	const bool bCheckerboard = false;// IsSubsurfaceBurleyNormalizedCheckerboardFormat(SceneTextureDesc.Format);
+
+	const uint32 ScaleFactor = bHalfRes ? 2 : 1;
+
+	const ESubsurfaceBurleyPassType PassType = GetSuburfaceBurleyPassType();
+
+	/**
+	 * All subsurface passes within the screen-space subsurface effect can operate at half or full resolution,
+	 * depending on the subsurface mode. The values are precomputed and shared among all Subsurface textures.
+	 */
+	const FScreenPassTextureViewport SubsurfaceViewport = FScreenPassTextureViewport::CreateDownscaled(SceneViewport, ScaleFactor);
+
+	const FRDGTextureDesc SubsurfaceTextureDescriptor = FRDGTextureDesc::Create2DDesc(
+		SubsurfaceViewport.Extent,
+		PF_FloatRGBA,
+		FClearValueBinding(),
+		TexCreate_None,
+		TexCreate_RenderTargetable | TexCreate_ShaderResource,
+		false);
+	
+	const FSubsurfaceBurleyParameters SubsurfaceCommonParameters = GetSubsurfaceBurleyCommonParameters(GraphBuilder.RHICmdList, View);
+	const FScreenPassTextureViewportParameters SubsurfaceViewportParameters = GetScreenPassTextureViewportParameters(SubsurfaceViewport);
+	const FScreenPassTextureViewportParameters SceneViewportParameters = GetScreenPassTextureViewportParameters(SceneViewport);
+
+	FRDGTextureRef SetupTexture = SceneTexture;
+	FRDGTextureRef SubsurfaceSubpassOneTex = nullptr;
+	FRDGTextureRef SubsurfaceSubpassTwoTex = nullptr;
+
+	FRHISamplerState* PointClampSampler = TStaticSamplerState<SF_Point, AM_Clamp, AM_Clamp, AM_Clamp>::GetRHI();
+	FRHISamplerState* BilinearBorderSampler = TStaticSamplerState<SF_Bilinear, AM_Border, AM_Border, AM_Border>::GetRHI();
+
+	if(PassType == ESubsurfaceBurleyPassType::Single)
+	{
+		FRDGTextureRef SubsurfaceSinglePassTexture = SceneTexture;
+		FRHISamplerState* SubsurfaceSamplerState = FSubsurfaceBurleyPS::GetSamplerState();
+		const FSubsurfaceBurleyPS::EQuality SubsurfaceQuality = FSubsurfaceBurleyPS::GetQuality();
+		
+		//single pass
+		FSubsurfaceBurleySinglePassPS::FParameters* PassParameters = GraphBuilder.AllocParameters<FSubsurfaceBurleySinglePassPS::FParameters>();
+		PassParameters->Subsurface = SubsurfaceCommonParameters;
+		PassParameters->RenderTargets[0] = FRenderTargetBinding(SceneTextureOutput, SceneTextureLoadAction, ERenderTargetStoreAction::EStore);
+		PassParameters->SubsurfaceInput0 = GetSubsurfaceBurleyInput(SceneTexture, SceneViewportParameters);
+		PassParameters->SubsurfaceSampler0 = SubsurfaceSamplerState;
+
+		FSubsurfaceBurleySinglePassPS::FPermutationDomain PixelShaderPermutationVector;
+		PixelShaderPermutationVector.Set<FSubsurfaceBurleySinglePassPS::FDimensionMode>(ESubsurfaceMode::FullRes/*SubsurfaceMode*/);
+		PixelShaderPermutationVector.Set<FSubsurfaceBurleySinglePassPS::FDimensionHalfRes>(bHalfRes);
+		PixelShaderPermutationVector.Set<FSubsurfaceBurleySinglePassPS::FDimensionCheckerboard>(bCheckerboard);
+		TShaderMapRef<FSubsurfaceBurleySinglePassPS> PixelShader(View.ShaderMap, PixelShaderPermutationVector);
+
+		//AddDrawScreenPass(GraphBuilder, RDG_EVENT_NAME("SubsurfaceBurley"), ScreenPassView, SubsurfaceViewport, SubsurfaceViewport, *PixelShader, PassParameters);
+		AddDrawScreenPass(GraphBuilder, RDG_EVENT_NAME("SubsurfaceBurleySinglePass"), ScreenPassView, SceneViewport, SceneViewport, *PixelShader, PassParameters);
+		return;
+	}
+
+	//-----------------------------------------------------------------------------------------------------------------------------------
+	//
+	/**
+	 * When in bypass mode, the setup and convolution passes are skipped, but lighting
+	 * reconstruction is still performed in the recombine pass.
+	 */
+	if (SubsurfaceMode != ESubsurfaceMode::Bypass)
+	{
+		SetupTexture = GraphBuilder.CreateTexture(SubsurfaceTextureDescriptor, TEXT("SubsurfaceSetupTexture"));
+
+		// Setup pass outputs the diffuse scene color and depth in preparation for the scatter passes.
+		{
+			FSubsurfaceBurleySetupPS::FParameters* PassParameters = GraphBuilder.AllocParameters<FSubsurfaceBurleySetupPS::FParameters>();
+			PassParameters->Subsurface = SubsurfaceCommonParameters;
+			PassParameters->RenderTargets[0] = FRenderTargetBinding(SetupTexture, ERenderTargetLoadAction::ENoAction, ERenderTargetStoreAction::EStore);
+			PassParameters->SubsurfaceInput0 = GetSubsurfaceBurleyInput(SceneTexture, SceneViewportParameters);
+			PassParameters->SubsurfaceSampler0 = PointClampSampler;
+
+			FSubsurfaceBurleySetupPS::FPermutationDomain PixelShaderPermutationVector;
+			PixelShaderPermutationVector.Set<FSubsurfaceBurleySetupPS::FDimensionHalfRes>(bHalfRes);
+			PixelShaderPermutationVector.Set<FSubsurfaceBurleySetupPS::FDimensionCheckerboard>(bCheckerboard);
+			TShaderMapRef<FSubsurfaceBurleySetupPS> PixelShader(View.ShaderMap, PixelShaderPermutationVector);
+
+			/**
+			 * The subsurface viewport is intentionally used as both the target and texture viewport, even though the texture
+			 * is potentially double the size. This is to ensure that the source UVs map 1-to-1 with pixel centers of the target,
+			 * in order to ensure that the checkerboard pattern selects the correct pixels from the scene texture. This still works
+			 * because the texture viewport is normalized into UV space, so it doesn't matter that the dimensions are twice as large.
+			 */
+			AddDrawScreenPass(GraphBuilder, RDG_EVENT_NAME("SubsurfaceSetup"), ScreenPassView, SubsurfaceViewport, SubsurfaceViewport, *PixelShader, PassParameters);
+		}
+
+		SubsurfaceSubpassOneTex = GraphBuilder.CreateTexture(SubsurfaceTextureDescriptor, TEXT("SubsurfaceSubpassOneTex"));
+		SubsurfaceSubpassTwoTex = GraphBuilder.CreateTexture(SubsurfaceTextureDescriptor, TEXT("SubsurfaceSubpassTwoTex"));
+
+		FRHISamplerState* SubsurfaceSamplerState = FSubsurfaceBurleyPS::GetSamplerState();
+		const FSubsurfaceBurleyPS::EQuality SubsurfaceQuality = FSubsurfaceBurleyPS::GetQuality();
+
+		struct FSubsurfacePassInfo
+		{
+			FSubsurfacePassInfo(const TCHAR* InName, FRDGTextureRef InInput, FRDGTextureRef InOutput)
+				: Name(InName)
+				, Input(InInput)
+				, Output(InOutput)
+			{}
+
+			const TCHAR* Name;
+			FRDGTextureRef Input;
+			FRDGTextureRef Output;
+		};
+
+		//we will determine burley or SSSS at runtime.
+		const FSubsurfacePassInfo SubsurfacePassInfos[] =
+		{
+			{ TEXT("SubsurfacePassOneTex"), SetupTexture, SubsurfaceSubpassOneTex },
+			{ TEXT("SubsurfacePassTwoTex"), SubsurfaceSubpassOneTex, SubsurfaceSubpassTwoTex },
+		};
+
+		const uint32 NumOfPasses = static_cast<uint32>(FSubsurfaceBurleyPS::ESubsurfacePass::MAX);
+
+		// Horizontal / Vertical scattering passes using a separable filter.
+		for (uint32 PassIndex = 0; PassIndex < NumOfPasses; ++PassIndex)
+		{
+			const auto SubsurfacePassFunction = static_cast<FSubsurfaceBurleyPS::ESubsurfacePass>(PassIndex);
+
+			const FSubsurfacePassInfo& PassInfo = SubsurfacePassInfos[PassIndex];
+			FRDGTextureRef TextureInput = PassInfo.Input;
+			FRDGTextureRef TextureOutput = PassInfo.Output;
+
+			FSubsurfaceBurleyPS::FParameters* PassParameters = GraphBuilder.AllocParameters<FSubsurfaceBurleyPS::FParameters>();
+			PassParameters->Subsurface = SubsurfaceCommonParameters;
+			PassParameters->RenderTargets[0] = FRenderTargetBinding(TextureOutput, ERenderTargetLoadAction::ENoAction, ERenderTargetStoreAction::EStore);
+			PassParameters->SubsurfaceInput0 = GetSubsurfaceBurleyInput(TextureInput, SubsurfaceViewportParameters);
+			PassParameters->SubsurfaceSampler0 = SubsurfaceSamplerState;
+
+			FSubsurfaceBurleyPS::FPermutationDomain PixelShaderPermutationVector;
+			PixelShaderPermutationVector.Set<FSubsurfaceBurleyPS::FSubsurfacePassFunction>(SubsurfacePassFunction);
+			PixelShaderPermutationVector.Set<FSubsurfaceBurleyPS::FDimensionQuality>(SubsurfaceQuality);
+			TShaderMapRef<FSubsurfaceBurleyPS> PixelShader(View.ShaderMap, PixelShaderPermutationVector);
+
+			AddDrawScreenPass(GraphBuilder, FRDGEventName(PassInfo.Name), ScreenPassView, SubsurfaceViewport, SubsurfaceViewport, *PixelShader, PassParameters);
+		}
+	}
+
+	// Recombines scattering result with scene color.
+	{
+		FSubsurfaceBurleyRecombinePS::FParameters* PassParameters = GraphBuilder.AllocParameters<FSubsurfaceBurleyRecombinePS::FParameters>();
+		PassParameters->Subsurface = SubsurfaceCommonParameters;
+		PassParameters->RenderTargets[0] = FRenderTargetBinding(SceneTextureOutput, SceneTextureLoadAction, ERenderTargetStoreAction::EStore);
+		PassParameters->SubsurfaceInput0 = GetSubsurfaceBurleyInput(SceneTexture, SceneViewportParameters);
+		PassParameters->SubsurfaceSampler0 = BilinearBorderSampler;
+
+		// Scattering output target is only used when scattering is enabled.
+		if (SubsurfaceMode != ESubsurfaceMode::Bypass)
+		{
+			PassParameters->SubsurfaceInput1 = GetSubsurfaceBurleyInput(SubsurfaceSubpassTwoTex, SubsurfaceViewportParameters);
+			PassParameters->SubsurfaceSampler1 = BilinearBorderSampler;
+		}
+
+		const FSubsurfaceBurleyRecombinePS::EQuality RecombineQuality = FSubsurfaceBurleyRecombinePS::GetQuality(View);
+
+		FSubsurfaceBurleyRecombinePS::FPermutationDomain PixelShaderPermutationVector;
+		PixelShaderPermutationVector.Set<FSubsurfaceBurleyRecombinePS::FDimensionMode>(SubsurfaceMode);
+		PixelShaderPermutationVector.Set<FSubsurfaceBurleyRecombinePS::FDimensionQuality>(RecombineQuality);
+		PixelShaderPermutationVector.Set<FSubsurfaceBurleyRecombinePS::FDimensionCheckerboard>(bCheckerboard);
+		TShaderMapRef<FSubsurfaceBurleyRecombinePS> PixelShader(View.ShaderMap, PixelShaderPermutationVector);
+
+		/**
+		 * See the related comment above in the prepare pass. The scene viewport is used as both the target and
+		 * texture viewport in order to ensure that the correct pixel is sampled for checkerboard rendering.
+		 */
+		AddDrawScreenPass(GraphBuilder, RDG_EVENT_NAME("SubsurfaceRecombine"), ScreenPassView, SceneViewport, SceneViewport, *PixelShader, PassParameters);
+	}
+}
\ No newline at end of file
diff --git a/Engine/Source/Runtime/Renderer/Private/PostProcess/SubsurfaceBurleyNormalized.h b/Engine/Source/Runtime/Renderer/Private/PostProcess/SubsurfaceBurleyNormalized.h
new file mode 100644
--- /dev/null
+++ b/Engine/Source/Runtime/Renderer/Private/PostProcess/SubsurfaceBurleyNormalized.h
@@ -0,0 +1,18 @@
+// Copyright 1998-2019 Epic Games, Inc. All Rights Reserved.
+
+/*=============================================================================
+	SubsurfaceBurleyNormalized.h: Screenspace Burley subsurface scattering implementation.
+=============================================================================*/
+
+#pragma once
+
+#include "ScreenPass.h"
+#include "RenderingCompositionGraph.h"
+
+void ComputeBurleySubsurfaceForView(
+	FRDGBuilder& GraphBuilder,
+	const FScreenPassViewInfo& ScreenPassView,
+	const FScreenPassTextureViewport& SceneViewport,
+	FRDGTextureRef SceneTexture,
+	FRDGTextureRef SceneTextureOutput,
+	ERenderTargetLoadAction SceneTextureLoadAction);
\ No newline at end of file
diff --git a/Engine/Source/Runtime/Renderer/Private/PostProcess/SubsurfaceCommon.cpp b/Engine/Source/Runtime/Renderer/Private/PostProcess/SubsurfaceCommon.cpp
new file mode 100644
--- /dev/null
+++ b/Engine/Source/Runtime/Renderer/Private/PostProcess/SubsurfaceCommon.cpp
@@ -0,0 +1,156 @@
+// Copyright 1998-2019 Epic Games, Inc. All Rights Reserved.
+
+/*=============================================================================
+	PostProcessSubsurfaceCommon.cpp: Screenspace subsurface common functions
+=============================================================================*/
+
+#include "PostProcess/SubsurfaceCommon.h"
+#include "RenderTargetTemp.h"
+#include "SystemTextures.h"
+#include "SceneRenderTargets.h"
+
+ENGINE_API IPooledRenderTarget* GetSubsufaceProfileTexture_RT(FRHICommandListImmediate& RHICmdList);
+
+namespace
+{
+	TAutoConsoleVariable<int32> CVarSubsurfaceScattering(
+		TEXT("r.SubsurfaceScattering"),
+		1,
+		TEXT(" 0: disabled\n")
+		TEXT(" 1: enabled (default)"),
+		ECVF_RenderThreadSafe | ECVF_Scalability);
+	
+	TAutoConsoleVariable<float> CVarSSSScale(
+		TEXT("r.SSS.Scale"),
+		1.0f,
+		TEXT("Affects the Screen space subsurface scattering pass")
+		TEXT("(use shadingmodel SubsurfaceProfile, get near to the object as the default)\n")
+		TEXT("is human skin which only scatters about 1.2cm)\n")
+		TEXT(" 0: off (if there is no object on the screen using this pass it should automatically disable the post process pass)\n")
+		TEXT("<1: scale scatter radius down (for testing)\n")
+		TEXT(" 1: use given radius form the Subsurface scattering asset (default)\n")
+		TEXT(">1: scale scatter radius up (for testing)"),
+		ECVF_Scalability | ECVF_RenderThreadSafe);
+		
+	TAutoConsoleVariable<int32> CVarSSSHalfRes(
+		TEXT("r.SSS.HalfRes"),
+		1,
+		TEXT(" 0: full quality (not optimized, as reference)\n")
+		TEXT(" 1: parts of the algorithm runs in half resolution which is lower quality but faster (default)"),
+		ECVF_RenderThreadSafe | ECVF_Scalability);
+
+	TAutoConsoleVariable<int32> CVarSSSQuality(
+		TEXT("r.SSS.Quality"),
+		0,
+		TEXT("Defines the quality of the recombine pass when using the SubsurfaceScatteringProfile shading model\n")
+		TEXT(" 0: low (faster, default)\n")
+		TEXT(" 1: high (sharper details but slower)\n")
+		TEXT("-1: auto, 1 if TemporalAA is disabled (without TemporalAA the quality is more noticable)"),
+		ECVF_RenderThreadSafe | ECVF_Scalability);
+
+	TAutoConsoleVariable<int32> CVarSSSFilter(
+		TEXT("r.SSS.Filter"),
+		1,
+		TEXT("Defines the filter method for Screenspace Subsurface Scattering feature.\n")
+		TEXT(" 0: point filter (useful for testing, could be cleaner)\n")
+		TEXT(" 1: bilinear filter"),
+		ECVF_RenderThreadSafe | ECVF_Scalability);
+
+	TAutoConsoleVariable<int32> CVarSSSSampleSet(
+		TEXT("r.SSS.SampleSet"),
+		2,
+		TEXT("Defines how many samples we use for Screenspace Subsurface Scattering feature.\n")
+		TEXT(" 0: lowest quality (6*2+1)\n")
+		TEXT(" 1: medium quality (9*2+1)\n")
+		TEXT(" 2: high quality (13*2+1) (default)"),
+		ECVF_RenderThreadSafe | ECVF_Scalability);
+}
+
+// Returns the [0, N] clamped value of the 'r.SSS.Scale' CVar.
+float GetSubsurfaceRadiusScale()
+{
+	static const auto CVar = IConsoleManager::Get().FindTConsoleVariableDataFloat(TEXT("r.SSS.Scale"));
+	check(CVar);
+
+	return FMath::Max(0.0f, CVar->GetValueOnRenderThread());
+}
+
+int32 GetSSSFilter()
+{
+	return CVarSSSFilter.GetValueOnRenderThread();
+}
+
+int32 GetSSSSampleSet()
+{
+	return CVarSSSSampleSet.GetValueOnRenderThread();
+}
+
+int32 GetSSSQuality()
+{
+	return CVarSSSQuality.GetValueOnRenderThread();
+}
+
+// Returns the SS profile texture with a black fallback texture if none exists yet.
+// Actually we do not need this for the burley normalized SSS.
+FRHITexture* GetSubsurfaceProfileTexture(FRHICommandListImmediate& RHICmdList)
+{
+	const IPooledRenderTarget* ProfileTextureTarget = GetSubsufaceProfileTexture_RT(RHICmdList);
+
+	if (!ProfileTextureTarget)
+	{
+		// No subsurface profile was used yet
+		ProfileTextureTarget = GSystemTextures.BlackDummy;
+	}
+
+	return ProfileTextureTarget->GetRenderTargetItem().ShaderResourceTexture;
+}
+
+// Returns the current subsurface mode required by the current view.
+ESubsurfaceMode GetSubsurfaceModeForView(const FViewInfo& View)
+{
+	const float Radius = GetSubsurfaceRadiusScale();
+	const bool bShowSubsurfaceScattering = Radius > 0 && View.Family->EngineShowFlags.SubsurfaceScattering;
+
+	if (bShowSubsurfaceScattering)
+	{
+		const bool bHalfRes = CVarSSSHalfRes.GetValueOnRenderThread() != 0;
+		if (bHalfRes)
+		{
+			return ESubsurfaceMode::HalfRes;
+		}
+		else
+		{
+			return ESubsurfaceMode::FullRes;
+		}
+	}
+	else
+	{
+		return ESubsurfaceMode::Bypass;
+	}
+}
+
+FSubsurfaceParameters GetSubsurfaceCommonParameters(FRHICommandListImmediate& RHICmdList, const FViewInfo& View)
+{
+	FSceneRenderTargets& SceneContext = FSceneRenderTargets::Get(RHICmdList);
+
+	const float DistanceToProjectionWindow = View.ViewMatrices.GetProjectionMatrix().M[0][0];
+	const float SSSScaleZ = DistanceToProjectionWindow * GetSubsurfaceRadiusScale();
+	const float SSSScaleX = SSSScaleZ / SUBSURFACE_KERNEL_SIZE * 0.5f;
+
+	FSubsurfaceParameters Parameters;
+	Parameters.SubsurfaceParams = FVector4(SSSScaleX, SSSScaleZ, 0, 0);
+	Parameters.ViewUniformBuffer = View.ViewUniformBuffer;
+	Parameters.SceneUniformBuffer = CreateSceneTextureUniformBuffer(
+		SceneContext, View.FeatureLevel, ESceneTextureSetupMode::All, EUniformBufferUsage::UniformBuffer_SingleFrame);
+	Parameters.BilinearTextureSampler = TStaticSamplerState<SF_Bilinear>::GetRHI();
+	Parameters.SSProfilesTexture = GetSubsurfaceProfileTexture(RHICmdList);
+	return Parameters;
+}
+
+FSubsurfaceInput GetSubsurfaceInput(FRDGTextureRef Texture, const FScreenPassTextureViewportParameters& ViewportParameters)
+{
+	FSubsurfaceInput Input;
+	Input.Texture = Texture;
+	Input.Viewport = ViewportParameters;
+	return Input;
+}
\ No newline at end of file
diff --git a/Engine/Source/Runtime/Renderer/Private/PostProcess/SubsurfaceCommon.h b/Engine/Source/Runtime/Renderer/Private/PostProcess/SubsurfaceCommon.h
new file mode 100644
--- /dev/null
+++ b/Engine/Source/Runtime/Renderer/Private/PostProcess/SubsurfaceCommon.h
@@ -0,0 +1,79 @@
+// Copyright 1998-2019 Epic Games, Inc. All Rights Reserved.
+
+/*=============================================================================
+	PostProcessSubsurfaceCommon.h: Screenspace subsurface common functions
+=============================================================================*/
+
+#pragma once
+
+#include "Engine/SubsurfaceProfile.h"
+#include "ScreenPass.h"
+
+enum class ESubsurfaceMode : uint32
+{
+	// Performs a full resolution scattering filter.
+	FullRes,
+
+	// Performs a half resolution scattering filter.
+	HalfRes,
+
+	// Reconstructs lighting, but does not perform scattering.
+	Bypass,
+
+	MAX
+};
+
+// Returns the [0, N] clamped value of the 'r.SSS.Scale' CVar.
+float GetSubsurfaceRadiusScale();
+
+int32 GetSSSFilter();
+int32 GetSSSSampleSet();
+int32 GetSSSQuality();
+
+// Returns the SS profile texture with a black fallback texture if none exists yet.
+// Actually we do not need this for the burley normalized SSS.
+FRHITexture* GetSubsurfaceProfileTexture(FRHICommandListImmediate& RHICmdList);
+
+// Returns the current subsurface mode required by the current view.
+ESubsurfaceMode GetSubsurfaceModeForView(const FViewInfo& View);
+
+// Set of common shader parameters shared by all subsurface shaders.
+BEGIN_SHADER_PARAMETER_STRUCT(FSubsurfaceParameters, )
+SHADER_PARAMETER(FVector4, SubsurfaceParams)
+	SHADER_PARAMETER_STRUCT_REF(FSceneTexturesUniformParameters, SceneUniformBuffer)
+	SHADER_PARAMETER_STRUCT_REF(FViewUniformShaderParameters, ViewUniformBuffer)
+	SHADER_PARAMETER_SAMPLER(SamplerState, BilinearTextureSampler)
+	SHADER_PARAMETER_TEXTURE(Texture2D, SSProfilesTexture)
+END_SHADER_PARAMETER_STRUCT()
+
+FSubsurfaceParameters GetSubsurfaceCommonParameters(FRHICommandListImmediate& RHICmdList, const FViewInfo& View);
+
+// Base class for a subsurface shader.
+class FSubsurfaceShader : public FGlobalShader
+{
+public:
+	static bool ShouldCompilePermutation(const FGlobalShaderPermutationParameters& Parameters)
+	{
+		return IsFeatureLevelSupported(Parameters.Platform, ERHIFeatureLevel::SM4);
+	}
+
+	static void ModifyCompilationEnvironment(const FGlobalShaderPermutationParameters& Parameters, FShaderCompilerEnvironment& OutEnvironment)
+	{
+		FGlobalShader::ModifyCompilationEnvironment(Parameters, OutEnvironment);
+		OutEnvironment.SetDefine(TEXT("SUBSURFACE_RADIUS_SCALE"), SUBSURFACE_RADIUS_SCALE);
+		OutEnvironment.SetDefine(TEXT("SUBSURFACE_KERNEL_SIZE"), SUBSURFACE_KERNEL_SIZE);
+	}
+
+	FSubsurfaceShader() = default;
+	FSubsurfaceShader(const ShaderMetaType::CompiledShaderInitializerType& Initializer)
+		: FGlobalShader(Initializer)
+	{}
+};
+
+// A shader parameter struct for a single subsurface input texture.
+BEGIN_SHADER_PARAMETER_STRUCT(FSubsurfaceInput, )
+SHADER_PARAMETER_STRUCT_INCLUDE(FScreenPassTextureViewportParameters, Viewport)
+SHADER_PARAMETER_RDG_TEXTURE(Texture2D, Texture)
+END_SHADER_PARAMETER_STRUCT()
+
+FSubsurfaceInput GetSubsurfaceInput(FRDGTextureRef Texture, const FScreenPassTextureViewportParameters& ViewportParameters);
diff --git a/Engine/Source/Runtime/Renderer/Private/PostProcess/SubsurfaceSeparable.cpp b/Engine/Source/Runtime/Renderer/Private/PostProcess/SubsurfaceSeparable.cpp
new file mode 100644
--- /dev/null
+++ b/Engine/Source/Runtime/Renderer/Private/PostProcess/SubsurfaceSeparable.cpp
@@ -0,0 +1,342 @@
+// Copyright 1998-2019 Epic Games, Inc. All Rights Reserved.
+
+/*=============================================================================
+	SubsurfaceSeparable.h: Screenspace Burley subsurface scattering implementation.
+=============================================================================*/
+
+#include "PostProcess/SubsurfaceSeparable.h"
+#include "PostProcess/SubsurfaceCommon.h"
+
+namespace
+{
+	TAutoConsoleVariable<int32> CVarSSSCheckerboard(
+		TEXT("r.SSS.Checkerboard"),
+		2,
+		TEXT("Enables or disables checkerboard rendering for subsurface profile rendering.\n")
+		TEXT("This is necessary if SceneColor does not include a floating point alpha channel (e.g 32-bit formats)\n")
+		TEXT(" 0: Disabled (high quality) \n")
+		TEXT(" 1: Enabled (low quality). Surface lighting will be at reduced resolution.\n")
+		TEXT(" 2: Automatic. Non-checkerboard lighting will be applied if we have a suitable rendertarget format\n"),
+		ECVF_RenderThreadSafe);
+		
+}
+
+bool IsSeparableSubsurfaceCheckerboardFormat(EPixelFormat SceneColorFormat)
+{
+	int CVarValue = CVarSSSCheckerboard.GetValueOnRenderThread();
+	if (CVarValue == 0)
+	{
+		return false;
+	}
+	else if (CVarValue == 1)
+	{
+		return true;
+	}
+	else if (CVarValue == 2)
+	{
+		switch (SceneColorFormat)
+		{
+		case PF_A32B32G32R32F:
+		case PF_FloatRGBA:
+			return false;
+		default:
+			return true;
+		}
+	}
+	return true;
+}
+
+// Encapsulates the post processing subsurface scattering pixel shader.
+class FSubsurfaceSetupPS : public FSubsurfaceShader
+{
+public:
+	DECLARE_GLOBAL_SHADER(FSubsurfaceSetupPS);
+	SHADER_USE_PARAMETER_STRUCT(FSubsurfaceSetupPS, FSubsurfaceShader)
+
+	BEGIN_SHADER_PARAMETER_STRUCT(FParameters, )
+		SHADER_PARAMETER_STRUCT_INCLUDE(FSubsurfaceParameters, Subsurface)
+		SHADER_PARAMETER_STRUCT(FSubsurfaceInput, SubsurfaceInput0)
+		SHADER_PARAMETER_SAMPLER(SamplerState, SubsurfaceSampler0)
+		RENDER_TARGET_BINDING_SLOTS()
+	END_SHADER_PARAMETER_STRUCT()
+
+	class FDimensionHalfRes : SHADER_PERMUTATION_BOOL("SUBSURFACE_HALF_RES");
+	class FDimensionCheckerboard : SHADER_PERMUTATION_BOOL("SUBSURFACE_PROFILE_CHECKERBOARD");
+	using FPermutationDomain = TShaderPermutationDomain<FDimensionHalfRes, FDimensionCheckerboard>;
+};
+
+IMPLEMENT_GLOBAL_SHADER(FSubsurfaceSetupPS, "/Engine/Private/PostProcessSubsurface.usf", "SetupPS", SF_Pixel);
+
+// Shader for the SSS separable blur.
+class FSubsurfacePS : public FSubsurfaceShader
+{
+public:
+	DECLARE_GLOBAL_SHADER(FSubsurfacePS);
+	SHADER_USE_PARAMETER_STRUCT(FSubsurfacePS, FSubsurfaceShader);
+
+	BEGIN_SHADER_PARAMETER_STRUCT(FParameters, )
+		SHADER_PARAMETER_STRUCT_INCLUDE(FSubsurfaceParameters, Subsurface)
+		SHADER_PARAMETER_STRUCT(FSubsurfaceInput, SubsurfaceInput0)
+		SHADER_PARAMETER_SAMPLER(SamplerState, SubsurfaceSampler0)
+		RENDER_TARGET_BINDING_SLOTS()
+	END_SHADER_PARAMETER_STRUCT()
+
+	// Direction of the 1D separable filter.
+	enum class EDirection : uint32
+	{
+		Horizontal,
+		Vertical,
+		MAX
+	};
+
+	// Controls the quality (number of samples) of the blur kernel.
+	enum class EQuality : uint32
+	{
+		Low,
+		Medium,
+		High,
+		MAX
+	};
+
+	class FDimensionDirection : SHADER_PERMUTATION_ENUM_CLASS("SUBSURFACE_DIRECTION", EDirection);
+	class FDimensionQuality : SHADER_PERMUTATION_ENUM_CLASS("SUBSURFACE_QUALITY", EQuality);
+	using FPermutationDomain = TShaderPermutationDomain<FDimensionDirection, FDimensionQuality>;
+
+	// Returns the sampler state based on the requested SSS filter CVar setting.
+	static FRHISamplerState* GetSamplerState()
+	{
+		if (GetSSSFilter())
+		{
+			return TStaticSamplerState<SF_Bilinear, AM_Border, AM_Border, AM_Border>::GetRHI();
+		}
+		else
+		{
+			return TStaticSamplerState<SF_Point, AM_Border, AM_Border, AM_Border>::GetRHI();
+		}
+	}
+
+	// Returns the SSS quality level requested by the SSS SampleSet CVar setting.
+	static EQuality GetQuality()
+	{
+		return static_cast<FSubsurfacePS::EQuality>(
+			FMath::Clamp(
+				GetSSSSampleSet(),
+				static_cast<int32>(FSubsurfacePS::EQuality::Low),
+				static_cast<int32>(FSubsurfacePS::EQuality::High)));
+	}
+};
+
+IMPLEMENT_GLOBAL_SHADER(FSubsurfacePS, "/Engine/Private/PostProcessSubsurface.usf", "MainPS", SF_Pixel);
+
+// Encapsulates the post processing subsurface recombine pixel shader.
+class FSubsurfaceRecombinePS : public FSubsurfaceShader
+{
+	DECLARE_GLOBAL_SHADER(FSubsurfaceRecombinePS);
+	SHADER_USE_PARAMETER_STRUCT(FSubsurfaceRecombinePS, FSubsurfaceShader);
+
+	BEGIN_SHADER_PARAMETER_STRUCT(FParameters, )
+		SHADER_PARAMETER_STRUCT_INCLUDE(FSubsurfaceParameters, Subsurface)
+		SHADER_PARAMETER_STRUCT(FSubsurfaceInput, SubsurfaceInput0)
+		SHADER_PARAMETER_STRUCT(FSubsurfaceInput, SubsurfaceInput1)
+		SHADER_PARAMETER_SAMPLER(SamplerState, SubsurfaceSampler0)
+		SHADER_PARAMETER_SAMPLER(SamplerState, SubsurfaceSampler1)
+		RENDER_TARGET_BINDING_SLOTS()
+	END_SHADER_PARAMETER_STRUCT();
+
+	// Controls the quality of lighting reconstruction.
+	enum class EQuality : uint32
+	{
+		Low,
+		High,
+		MAX
+	};
+
+	class FDimensionMode : SHADER_PERMUTATION_ENUM_CLASS("SUBSURFACE_RECOMBINE_MODE", ESubsurfaceMode);
+	class FDimensionQuality : SHADER_PERMUTATION_ENUM_CLASS("SUBSURFACE_RECOMBINE_QUALITY", EQuality);
+	class FDimensionCheckerboard : SHADER_PERMUTATION_BOOL("SUBSURFACE_PROFILE_CHECKERBOARD");
+	using FPermutationDomain = TShaderPermutationDomain<FDimensionMode, FDimensionQuality, FDimensionCheckerboard>;
+
+	// Returns the Recombine quality level requested by the SSS Quality CVar setting.
+	static EQuality GetQuality(const FViewInfo& View)
+	{
+		const uint32 QualityCVar = GetSSSQuality();
+
+		// Quality is forced to high when the CVar is set to 'auto' and TAA is NOT enabled.
+		// TAA improves quality through temporal filtering, making it less necessary to use
+		// high quality mode.
+		const bool bUseHighQuality = (QualityCVar == -1 && View.AntiAliasingMethod != AAM_TemporalAA);
+
+		if (QualityCVar == 1 || bUseHighQuality)
+		{
+			return EQuality::High;
+		}
+		else
+		{
+			return EQuality::Low;
+		}
+	}
+};
+
+IMPLEMENT_GLOBAL_SHADER(FSubsurfaceRecombinePS, "/Engine/Private/PostProcessSubsurface.usf", "SubsurfaceRecombinePS", SF_Pixel);
+
+
+void ComputeSeparableSubsurfaceForView(
+	FRDGBuilder& GraphBuilder,
+	const FScreenPassViewInfo& ScreenPassView,
+	const FScreenPassTextureViewport& SceneViewport,
+	FRDGTextureRef SceneTexture,
+	FRDGTextureRef SceneTextureOutput,
+	ERenderTargetLoadAction SceneTextureLoadAction)
+{
+	check(SceneTexture);
+	check(SceneTextureOutput);
+	check(SceneViewport.Extent == SceneTexture->Desc.Extent);
+
+	const FViewInfo& View = ScreenPassView.View;
+
+	const FSceneViewFamily* ViewFamily = View.Family;
+
+	const FRDGTextureDesc& SceneTextureDesc = SceneTexture->Desc;
+
+	const ESubsurfaceMode SubsurfaceMode = GetSubsurfaceModeForView(View);
+
+	const bool bHalfRes = (SubsurfaceMode == ESubsurfaceMode::HalfRes);
+
+	const bool bCheckerboard = IsSeparableSubsurfaceCheckerboardFormat(SceneTextureDesc.Format);
+
+	const uint32 ScaleFactor = bHalfRes ? 2 : 1;
+
+	/**
+	 * All subsurface passes within the screen-space subsurface effect can operate at half or full resolution,
+	 * depending on the subsurface mode. The values are precomputed and shared among all Subsurface textures.
+	 */
+	const FScreenPassTextureViewport SubsurfaceViewport = FScreenPassTextureViewport::CreateDownscaled(SceneViewport, ScaleFactor);
+
+	const FRDGTextureDesc SubsurfaceTextureDescriptor = FRDGTextureDesc::Create2DDesc(
+		SubsurfaceViewport.Extent,
+		PF_FloatRGBA,
+		FClearValueBinding(),
+		TexCreate_None,
+		TexCreate_RenderTargetable | TexCreate_ShaderResource,
+		false);
+	
+	const FSubsurfaceParameters SubsurfaceCommonParameters = GetSubsurfaceCommonParameters(GraphBuilder.RHICmdList, View);
+	const FScreenPassTextureViewportParameters SubsurfaceViewportParameters = GetScreenPassTextureViewportParameters(SubsurfaceViewport);
+	const FScreenPassTextureViewportParameters SceneViewportParameters = GetScreenPassTextureViewportParameters(SceneViewport);
+
+	FRDGTextureRef SetupTexture = SceneTexture;
+	FRDGTextureRef SubsurfaceTextureX = nullptr;
+	FRDGTextureRef SubsurfaceTextureY = nullptr;
+
+	FRHISamplerState* PointClampSampler = TStaticSamplerState<SF_Point, AM_Clamp, AM_Clamp, AM_Clamp>::GetRHI();
+	FRHISamplerState* BilinearBorderSampler = TStaticSamplerState<SF_Bilinear, AM_Border, AM_Border, AM_Border>::GetRHI();
+
+	/**
+	 * When in bypass mode, the setup and convolution passes are skipped, but lighting
+	 * reconstruction is still performed in the recombine pass.
+	 */
+	if (SubsurfaceMode != ESubsurfaceMode::Bypass)
+	{
+		SetupTexture = GraphBuilder.CreateTexture(SubsurfaceTextureDescriptor, TEXT("SubsurfaceSetupTexture"));
+
+		// Setup pass outputs the diffuse scene color and depth in preparation for the scatter passes.
+		{
+			FSubsurfaceSetupPS::FParameters* PassParameters = GraphBuilder.AllocParameters<FSubsurfaceSetupPS::FParameters>();
+			PassParameters->Subsurface = SubsurfaceCommonParameters;
+			PassParameters->RenderTargets[0] = FRenderTargetBinding(SetupTexture, ERenderTargetLoadAction::ENoAction, ERenderTargetStoreAction::EStore);
+			PassParameters->SubsurfaceInput0 = GetSubsurfaceInput(SceneTexture, SceneViewportParameters);
+			PassParameters->SubsurfaceSampler0 = PointClampSampler;
+
+			FSubsurfaceSetupPS::FPermutationDomain PixelShaderPermutationVector;
+			PixelShaderPermutationVector.Set<FSubsurfaceSetupPS::FDimensionHalfRes>(bHalfRes);
+			PixelShaderPermutationVector.Set<FSubsurfaceSetupPS::FDimensionCheckerboard>(bCheckerboard);
+			TShaderMapRef<FSubsurfaceSetupPS> PixelShader(View.ShaderMap, PixelShaderPermutationVector);
+
+			/**
+			 * The subsurface viewport is intentionally used as both the target and texture viewport, even though the texture
+			 * is potentially double the size. This is to ensure that the source UVs map 1-to-1 with pixel centers of the target,
+			 * in order to ensure that the checkerboard pattern selects the correct pixels from the scene texture. This still works
+			 * because the texture viewport is normalized into UV space, so it doesn't matter that the dimensions are twice as large.
+			 */
+			AddDrawScreenPass(GraphBuilder, RDG_EVENT_NAME("SubsurfaceSetup"), ScreenPassView, SubsurfaceViewport, SubsurfaceViewport, *PixelShader, PassParameters);
+		}
+
+		SubsurfaceTextureX = GraphBuilder.CreateTexture(SubsurfaceTextureDescriptor, TEXT("SubsurfaceTextureX"));
+		SubsurfaceTextureY = GraphBuilder.CreateTexture(SubsurfaceTextureDescriptor, TEXT("SubsurfaceTextureY"));
+
+		FRHISamplerState* SubsurfaceSamplerState = FSubsurfacePS::GetSamplerState();
+		const FSubsurfacePS::EQuality SubsurfaceQuality = FSubsurfacePS::GetQuality();
+
+		struct FSubsurfacePassInfo
+		{
+			FSubsurfacePassInfo(const TCHAR* InName, FRDGTextureRef InInput, FRDGTextureRef InOutput)
+				: Name(InName)
+				, Input(InInput)
+				, Output(InOutput)
+			{}
+
+			const TCHAR* Name;
+			FRDGTextureRef Input;
+			FRDGTextureRef Output;
+		};
+
+		const FSubsurfacePassInfo SubsurfacePassInfoByDirection[] =
+		{
+			{ TEXT("SubsurfaceX"), SetupTexture, SubsurfaceTextureX },
+			{ TEXT("SubsurfaceY"), SubsurfaceTextureX, SubsurfaceTextureY },
+		};
+
+		// Horizontal / Vertical scattering passes using a separable filter.
+		for (uint32 DirectionIndex = 0; DirectionIndex < static_cast<uint32>(FSubsurfacePS::EDirection::MAX); ++DirectionIndex)
+		{
+			const auto Direction = static_cast<FSubsurfacePS::EDirection>(DirectionIndex);
+
+			const FSubsurfacePassInfo& PassInfo = SubsurfacePassInfoByDirection[DirectionIndex];
+			FRDGTextureRef TextureInput = PassInfo.Input;
+			FRDGTextureRef TextureOutput = PassInfo.Output;
+
+			FSubsurfacePS::FParameters* PassParameters = GraphBuilder.AllocParameters<FSubsurfacePS::FParameters>();
+			PassParameters->Subsurface = SubsurfaceCommonParameters;
+			PassParameters->RenderTargets[0] = FRenderTargetBinding(TextureOutput, ERenderTargetLoadAction::ENoAction, ERenderTargetStoreAction::EStore);
+			PassParameters->SubsurfaceInput0 = GetSubsurfaceInput(TextureInput, SubsurfaceViewportParameters);
+			PassParameters->SubsurfaceSampler0 = SubsurfaceSamplerState;
+
+			FSubsurfacePS::FPermutationDomain PixelShaderPermutationVector;
+			PixelShaderPermutationVector.Set<FSubsurfacePS::FDimensionDirection>(Direction);
+			PixelShaderPermutationVector.Set<FSubsurfacePS::FDimensionQuality>(SubsurfaceQuality);
+			TShaderMapRef<FSubsurfacePS> PixelShader(View.ShaderMap, PixelShaderPermutationVector);
+
+			AddDrawScreenPass(GraphBuilder, FRDGEventName(PassInfo.Name), ScreenPassView, SubsurfaceViewport, SubsurfaceViewport, *PixelShader, PassParameters);
+		}
+	}
+
+	// Recombines scattering result with scene color.
+	{
+		FSubsurfaceRecombinePS::FParameters* PassParameters = GraphBuilder.AllocParameters<FSubsurfaceRecombinePS::FParameters>();
+		PassParameters->Subsurface = SubsurfaceCommonParameters;
+		PassParameters->RenderTargets[0] = FRenderTargetBinding(SceneTextureOutput, SceneTextureLoadAction, ERenderTargetStoreAction::EStore);
+		PassParameters->SubsurfaceInput0 = GetSubsurfaceInput(SceneTexture, SceneViewportParameters);
+		PassParameters->SubsurfaceSampler0 = BilinearBorderSampler;
+
+		// Scattering output target is only used when scattering is enabled.
+		if (SubsurfaceMode != ESubsurfaceMode::Bypass)
+		{
+			PassParameters->SubsurfaceInput1 = GetSubsurfaceInput(SubsurfaceTextureY, SubsurfaceViewportParameters);
+			PassParameters->SubsurfaceSampler1 = BilinearBorderSampler;
+		}
+
+		const FSubsurfaceRecombinePS::EQuality RecombineQuality = FSubsurfaceRecombinePS::GetQuality(View);
+
+		FSubsurfaceRecombinePS::FPermutationDomain PixelShaderPermutationVector;
+		PixelShaderPermutationVector.Set<FSubsurfaceRecombinePS::FDimensionMode>(SubsurfaceMode);
+		PixelShaderPermutationVector.Set<FSubsurfaceRecombinePS::FDimensionQuality>(RecombineQuality);
+		PixelShaderPermutationVector.Set<FSubsurfaceRecombinePS::FDimensionCheckerboard>(bCheckerboard);
+		TShaderMapRef<FSubsurfaceRecombinePS> PixelShader(View.ShaderMap, PixelShaderPermutationVector);
+
+		/**
+		 * See the related comment above in the prepare pass. The scene viewport is used as both the target and
+		 * texture viewport in order to ensure that the correct pixel is sampled for checkerboard rendering.
+		 */
+		AddDrawScreenPass(GraphBuilder, RDG_EVENT_NAME("SubsurfaceRecombine"), ScreenPassView, SceneViewport, SceneViewport, *PixelShader, PassParameters);
+	}
+}
diff --git a/Engine/Source/Runtime/Renderer/Private/PostProcess/SubsurfaceSeparable.h b/Engine/Source/Runtime/Renderer/Private/PostProcess/SubsurfaceSeparable.h
new file mode 100644
--- /dev/null
+++ b/Engine/Source/Runtime/Renderer/Private/PostProcess/SubsurfaceSeparable.h
@@ -0,0 +1,21 @@
+// Copyright 1998-2019 Epic Games, Inc. All Rights Reserved.
+
+/*=============================================================================
+	SubsurfaceSeparable.h: Screenspace Burley subsurface scattering implementation.
+=============================================================================*/
+
+#pragma once
+
+#include "ScreenPass.h"
+#include "RenderingCompositionGraph.h"
+#include "PostProcess/SubsurfaceCommon.h"
+
+bool IsSeparableSubsurfaceCheckerboardFormat(EPixelFormat SceneColorFormat);
+
+void ComputeSeparableSubsurfaceForView(
+	FRDGBuilder& GraphBuilder,
+	const FScreenPassViewInfo& ScreenPassView,
+	const FScreenPassTextureViewport& SceneViewport,
+	FRDGTextureRef SceneTexture,
+	FRDGTextureRef SceneTextureOutput,
+	ERenderTargetLoadAction SceneTextureLoadAction);
\ No newline at end of file
