Fix RT Translucency primary ray MaxT

diff --git a/Engine/Shaders/Private/RayTracing/RayTracingTranslucency.usf b/Engine/Shaders/Private/RayTracing/RayTracingTranslucency.usf
--- a/Engine/Shaders/Private/RayTracing/RayTracingTranslucency.usf
+++ b/Engine/Shaders/Private/RayTracing/RayTracingTranslucency.usf
@@ -111,10 +111,12 @@ void RayTracingTranslucencyRGS()
 	float2 UV = (float2(PixelCoord) + 0.5) * InvBufferSize;
 	FGBufferData GBufferData = GetGBufferDataFromSceneTextures(UV);
 	float Depth = GBufferData.Depth;
+	float3 WorldPosition = ReconstructWorldPositionFromDepth(UV, Depth);
+	float WorldSpaceDistance = length(WorldPosition - View.WorldViewOrigin);
 
 	// Trace rays from camera origin to (Gbuffer - epsilon) to only intersect translucent objects
 	RayDesc Ray = CreatePrimaryRay(UV);
-	Ray.TMax = Depth - 0.1;
+	Ray.TMax = WorldSpaceDistance - 0.1;
 	FRayCone RayCone = (FRayCone)0;
 	RayCone.SpreadAngle = View.EyeToPixelSpreadAngle;
 
