Implements accumulator compression for denoiser

diff --git a/Engine/Shaders/Private/ScreenSpaceDenoise/SSDDefinitions.ush b/Engine/Shaders/Private/ScreenSpaceDenoise/SSDDefinitions.ush
--- a/Engine/Shaders/Private/ScreenSpaceDenoise/SSDDefinitions.ush
+++ b/Engine/Shaders/Private/ScreenSpaceDenoise/SSDDefinitions.ush
@@ -201,6 +201,14 @@
 	#define SIGNAL_COMPRESSION_DIFFUSE_INDIRECT_HARMONIC 0xD200
 
 
+/** Compression of signal for VGPR occupency. */
+	/** The signal is not being compressed. */
+	#define ACCUMULATOR_COMPRESSION_DISABLED 0
+
+	/** VGPR compression for spherical harmonic. */
+	#define ACCUMULATOR_COMPRESSION_PENUMBRA_DRB 0x0001
+
+
 /** Technic used to compute the world vector betweem neighbor and reference. */
 	// Directly use FSSDKernelConfig::RefSceneMetadata.TranslatedWorldPosition. Cost 3 VGPR over entire kernel inner loop.
 	#define NEIGHBOR_TO_REF_CACHE_WORLD_POSITION 0
diff --git a/Engine/Shaders/Private/ScreenSpaceDenoise/SSDSignalAccumulatorArray.ush b/Engine/Shaders/Private/ScreenSpaceDenoise/SSDSignalAccumulatorArray.ush
new file mode 100644
--- /dev/null
+++ b/Engine/Shaders/Private/ScreenSpaceDenoise/SSDSignalAccumulatorArray.ush
@@ -0,0 +1,143 @@
+// Copyright 1998-2019 Epic Games, Inc. All Rights Reserved.
+
+#include "SSDSignalAccumulator.ush"
+
+
+// Maximum number of VGPR that can be used to compress accumulator array. Currently bound by ACCUMULATOR_COMPRESSION_PENUMBRA_DRB.
+#define MAX_ACCUMULATOR_COMPRESSED_VGPRS 24
+
+
+/** Structure that is just an array of accumulators, to work arround shader compiler bug on "out" of arrays. */
+struct FSSDSignalAccumulatorArray
+{
+	FSSDSignalAccumulator Array[MAX_SIGNAL_MULTIPLEXING];
+};
+
+/** Accumulator array that is compressed to have lower VGPR footprint. */
+struct FSSDCompressedSignalAccumulatorArray
+{
+	uint VGPR[MAX_ACCUMULATOR_COMPRESSED_VGPRS];
+};
+
+/** Created an initialised accumulator array. */
+FSSDSignalAccumulatorArray CreateSignalAccumulatorArray()
+{
+	FSSDSignalAccumulatorArray Accumulators;
+	
+	UNROLL_N(MAX_SIGNAL_MULTIPLEXING)
+	for (uint i = 0; i < MAX_SIGNAL_MULTIPLEXING; i++)
+	{
+		Accumulators.Array[i] = CreateSignalAccumulator();
+	}
+	return Accumulators;
+}
+
+/** Created an uninitialized accumulator array that is compressed. */
+FSSDCompressedSignalAccumulatorArray CreateUninitialisedCompressedAccumulatorArray()
+{
+	FSSDCompressedSignalAccumulatorArray CompressedAccumulators;
+	UNROLL_N(MAX_ACCUMULATOR_COMPRESSED_VGPRS)
+	for (uint i = 0; i < MAX_ACCUMULATOR_COMPRESSED_VGPRS; i++)
+	{
+		CompressedAccumulators.VGPR[i] = 0;
+	}
+	return CompressedAccumulators;
+}
+
+/** Compress the accumulator array. */
+FSSDCompressedSignalAccumulatorArray CompressAccumulatorArray(in FSSDSignalAccumulatorArray Accumulators, uint CompressionLayout)
+{
+	// Initialize the float2 arrays to be compressed at half in VGPRs.
+	float2 RawArray[MAX_ACCUMULATOR_COMPRESSED_VGPRS];
+	{
+		UNROLL_N(MAX_ACCUMULATOR_COMPRESSED_VGPRS)
+		for (uint i = 0; i < 2 * MAX_ACCUMULATOR_COMPRESSED_VGPRS; i++)
+		{
+			RawArray[i].x = 0;
+			RawArray[i].y = 0;
+		}
+	}
+	
+	if (CompressionLayout == ACCUMULATOR_COMPRESSION_DISABLED)
+	{
+		// NOP
+	}
+	else if (CompressionLayout == ACCUMULATOR_COMPRESSION_PENUMBRA_DRB)
+	{
+		UNROLL_N(MAX_SIGNAL_BATCH_SIZE)
+		for (uint SignalBatchId = 0; SignalBatchId < MAX_SIGNAL_BATCH_SIZE; SignalBatchId++)
+		{
+			const FSSDSignalAccumulator Accum = Accumulators.Array[SignalBatchId];
+
+			RawArray[SignalBatchId * 6 + 0].x = Accum.Current.SampleCount;
+			RawArray[SignalBatchId * 6 + 0].y = Accum.Current.MissCount;
+			RawArray[SignalBatchId * 6 + 1].x = Accum.Current.WorldBluringRadius;
+			RawArray[SignalBatchId * 6 + 1].y = Accum.CurrentInvFrequency;
+			RawArray[SignalBatchId * 6 + 2].x = Accum.CurrentTranslucency;
+			RawArray[SignalBatchId * 6 + 2].y = Accum.Previous.SampleCount;
+			RawArray[SignalBatchId * 6 + 3].x = Accum.Previous.MissCount;
+			RawArray[SignalBatchId * 6 + 3].y = Accum.Previous.WorldBluringRadius;
+			RawArray[SignalBatchId * 6 + 4].x = Accum.PreviousInvFrequency;
+			RawArray[SignalBatchId * 6 + 4].y = Accum.MinInvFrequency;
+			RawArray[SignalBatchId * 6 + 5].x = Accum.HighestInvFrequency;
+			RawArray[SignalBatchId * 6 + 5].y = Accum.BorderingRadius;
+		} // for (uint SignalMultiplexId = 0; SignalMultiplexId < MAX_SIGNAL_BATCH_SIZE; SignalMultiplexId++)
+	}
+
+	// Compress.
+	FSSDCompressedSignalAccumulatorArray CompressedAccumulators;
+	UNROLL_N(MAX_ACCUMULATOR_COMPRESSED_VGPRS)
+	for (uint i = 0; i < MAX_ACCUMULATOR_COMPRESSED_VGPRS; i++)
+	{
+		CompressedAccumulators.VGPR[i] = Compress2Floats(RawArray[i]);
+	}
+
+	return CompressedAccumulators;
+} // CompressAccumulatorArray()
+
+/** Uncompress the accumulator array. */
+FSSDSignalAccumulatorArray UncompressAccumulatorArray(in FSSDCompressedSignalAccumulatorArray CompressedAccumulators, uint CompressionLayout)
+{
+	// Uncompress VGPRs to floats.
+	float2 RawArray[MAX_ACCUMULATOR_COMPRESSED_VGPRS];
+	{
+		UNROLL_N(MAX_ACCUMULATOR_COMPRESSED_VGPRS)
+		for (uint i = 0; i < MAX_ACCUMULATOR_COMPRESSED_VGPRS; i++)
+		{
+			RawArray[i] = Uncompress2Floats(CompressedAccumulators.VGPR[i]);
+		}
+	}
+
+	FSSDSignalAccumulatorArray Accumulators = CreateSignalAccumulatorArray();
+
+	if (CompressionLayout == ACCUMULATOR_COMPRESSION_DISABLED)
+	{
+		// NOP to not clobber information in OutSample that remained uncompressed.
+	}
+	else if (CompressionLayout == ACCUMULATOR_COMPRESSION_PENUMBRA_DRB)
+	{
+		// TODO: CONFIG_SIGNAL_BATCH_SIZE?
+		UNROLL_N(MAX_SIGNAL_BATCH_SIZE)
+		for (uint SignalBatchId = 0; SignalBatchId < MAX_SIGNAL_BATCH_SIZE; SignalBatchId++)
+		{
+			FSSDSignalAccumulator Accum = CreateSignalAccumulator();
+
+			Accum.Current.SampleCount         = RawArray[SignalBatchId * 6 + 0].x;
+			Accum.Current.MissCount           = RawArray[SignalBatchId * 6 + 0].y;
+			Accum.Current.WorldBluringRadius  = RawArray[SignalBatchId * 6 + 1].x;
+			Accum.CurrentInvFrequency         = RawArray[SignalBatchId * 6 + 1].y;
+			Accum.CurrentTranslucency         = RawArray[SignalBatchId * 6 + 2].x;
+			Accum.Previous.SampleCount        = RawArray[SignalBatchId * 6 + 2].y;
+			Accum.Previous.MissCount          = RawArray[SignalBatchId * 6 + 3].x;
+			Accum.Previous.WorldBluringRadius = RawArray[SignalBatchId * 6 + 3].y;
+			Accum.PreviousInvFrequency        = RawArray[SignalBatchId * 6 + 4].x;
+			Accum.MinInvFrequency             = RawArray[SignalBatchId * 6 + 4].y;
+			Accum.HighestInvFrequency         = RawArray[SignalBatchId * 6 + 5].x;
+			Accum.BorderingRadius             = RawArray[SignalBatchId * 6 + 5].y;
+
+			Accumulators.Array[SignalBatchId] = Accum;
+		} // for (uint SignalBatchId = 0; SignalBatchId < MAX_SIGNAL_BATCH_SIZE; SignalBatchId++)
+	}
+
+	return Accumulators;
+} // UncompressAccumulatorArray()
diff --git a/Engine/Shaders/Private/ScreenSpaceDenoise/SSDSignalCompression.ush b/Engine/Shaders/Private/ScreenSpaceDenoise/SSDSignalCompression.ush
--- a/Engine/Shaders/Private/ScreenSpaceDenoise/SSDSignalCompression.ush
+++ b/Engine/Shaders/Private/ScreenSpaceDenoise/SSDSignalCompression.ush
@@ -25,6 +25,7 @@ struct FSSDCompressedSignalSample
 /** Compress part of signal into smaller VGPR foorprint. */
 void CompressSignalSample(in FSSDSignalSample Sample, uint CompressionLayout, out FSSDCompressedSignalSample OutCompressedSample)
 {
+	UNROLL_N(MAX_SIGNAL_COMPRESSED_VGPRS)
 	for (uint i = 0; i < MAX_SIGNAL_COMPRESSED_VGPRS; i++)
 	{
 		OutCompressedSample.VGPR[i] = 0;
diff --git a/Engine/Shaders/Private/ScreenSpaceDenoise/SSDSpatialAccumulation.usf b/Engine/Shaders/Private/ScreenSpaceDenoise/SSDSpatialAccumulation.usf
--- a/Engine/Shaders/Private/ScreenSpaceDenoise/SSDSpatialAccumulation.usf
+++ b/Engine/Shaders/Private/ScreenSpaceDenoise/SSDSpatialAccumulation.usf
@@ -113,6 +113,12 @@
 		#error Unexpected stage.
 	#endif
 
+	// Compress the DRB accumulator to have lower VGPR footprint.
+	#if defined(CONFIG_OUTPUT_MODE) && CONFIG_OUTPUT_MODE == OUTPUT_MODE_DRB
+		// Looks like shader compilers completly give up.
+		// #define CONFIG_ACCUMULATOR_VGPR_COMPRESSION ACCUMULATOR_COMPRESSION_PENUMBRA_DRB
+	#endif
+
 #elif CONFIG_SIGNAL_PROCESSING == SIGNAL_PROCESSING_REFLECTIONS
 	#define CONFIG_SIGNAL_COLOR 1
 
@@ -576,7 +582,7 @@ void MainCS(
 	#endif
 
 	// Accumulate spatially the input.
-	FSSDSignalAccumulatorArray SignalAccumulators = CreateSignalAccumulatorArray();
+	FSSDSignalAccumulatorArray SignalAccumulators;
 	{
 		FSSDKernelConfig KernelConfig = CreateKernelConfig();
 
@@ -660,6 +666,8 @@ void MainCS(
 		}
 		#endif
 
+		FSSDSignalAccumulatorArray UncompressedAccumulators = CreateSignalAccumulatorArray();
+
 		// When not upscaling, manually force accumulate the sample of the kernel.
 		if (!KernelConfig.bSampleKernelCenter && !KernelConfig.bDescOrder)
 		{
@@ -688,7 +696,7 @@ void MainCS(
 				}
 
 				AccumulateSample(
-					/* inout */ SignalAccumulators.Array[SignalMultiplexId],
+					/* inout */ UncompressedAccumulators.Array[SignalMultiplexId],
 					SampleInfos);
 			}
 		}
@@ -755,6 +763,8 @@ void MainCS(
 			#endif
 		}
 		#endif // CONFIG_SAMPLE_SET == SAMPLE_SET_STOCASTIC_HIERARCHY
+		
+		FSSDCompressedSignalAccumulatorArray CompressedAccumulators = CompressAccumulatorArray(UncompressedAccumulators, CONFIG_ACCUMULATOR_VGPR_COMPRESSION);
 
 		AccumulateKernel(
 			KernelConfig,
@@ -762,7 +772,8 @@ void MainCS(
 			SignalInput_Textures_1,
 			SignalInput_Textures_2,
 			SignalInput_Textures_3,
-			/* inout */ SignalAccumulators);
+			/* inout */ UncompressedAccumulators,
+			/* inout */ CompressedAccumulators);
 
 		// When doing history rejection pre convolution, could still have no information found with the 3x3 + kernel,
 		// therefore dynamically complete to form an entire 3x3 convolution.
@@ -781,7 +792,8 @@ void MainCS(
 					SignalInput_Textures_1,
 					SignalInput_Textures_2,
 					SignalInput_Textures_3,
-					/* inout */ SignalAccumulators);
+					/* inout */ UncompressedAccumulators,
+					/* inout */ CompressedAccumulators);
 			}
 		}
 		#endif // DIM_STAGE == STAGE_REJECTION_PRE_CONVOLUTION
@@ -799,8 +811,15 @@ void MainCS(
 				SignalInput_Textures_1,
 				SignalInput_Textures_2,
 				SignalInput_Textures_3,
-				/* inout */ SignalAccumulators);
+				/* inout */ UncompressedAccumulators,
+				/* inout */ CompressedAccumulators);
 		}
+		
+		#if CONFIG_ACCUMULATOR_VGPR_COMPRESSION == ACCUMULATOR_COMPRESSION_DISABLED
+			SignalAccumulators = UncompressedAccumulators;
+		#else
+			SignalAccumulators = UncompressAccumulatorArray(CompressedAccumulators, CONFIG_ACCUMULATOR_VGPR_COMPRESSION);
+		#endif
 	}
 
 	// Color processing of the signal to reduce highlight flickering.
diff --git a/Engine/Shaders/Private/ScreenSpaceDenoise/SSDSpatialKernel.ush b/Engine/Shaders/Private/ScreenSpaceDenoise/SSDSpatialKernel.ush
--- a/Engine/Shaders/Private/ScreenSpaceDenoise/SSDSpatialKernel.ush
+++ b/Engine/Shaders/Private/ScreenSpaceDenoise/SSDSpatialKernel.ush
@@ -1,6 +1,6 @@
 // Copyright 1998-2019 Epic Games, Inc. All Rights Reserved.
 
-#include "SSDSignalAccumulator.ush"
+#include "SSDSignalAccumulatorArray.ush"
 #include "SSDSignalBufferEncoding.ush"
 #include "../TextureSampling.ush"
 
@@ -36,27 +36,13 @@
 	#define CONFIG_VGPR_FREE_SAMPLE_TRACK_ID 0
 #endif
 
+#ifndef CONFIG_ACCUMULATOR_VGPR_COMPRESSION
+	#define CONFIG_ACCUMULATOR_VGPR_COMPRESSION ACCUMULATOR_COMPRESSION_DISABLED
+#endif
 
-//------------------------------------------------------- STRUCTURES
 
-/** Structure that is just an array of accumulators, to work arround shader compiler bug on "out" of arrays. */
-struct FSSDSignalAccumulatorArray
-{
-	FSSDSignalAccumulator Array[MAX_SIGNAL_MULTIPLEXING];
-};
-
-FSSDSignalAccumulatorArray CreateSignalAccumulatorArray()
-{
-	FSSDSignalAccumulatorArray Accumulators;
-	
-	UNROLL_N(MAX_SIGNAL_MULTIPLEXING)
-	for (uint i = 0; i < MAX_SIGNAL_MULTIPLEXING; i++)
-	{
-		Accumulators.Array[i] = CreateSignalAccumulator();
-	}
-	return Accumulators;
-}
 
+//------------------------------------------------------- STRUCTURES
 
 /** Configures the spatial kernel. */
 struct FSSDKernelConfig
@@ -378,7 +364,8 @@ void SampleMultiplexedSignals(
 /** Accumulate multiplexed samples and their metadata to an accumulator. */
 void AccumulateSampledMultiplexedSignals(
 	FSSDKernelConfig KernelConfig,
-	inout FSSDSignalAccumulatorArray Accumulators,
+	inout FSSDSignalAccumulatorArray UncompressedAccumulators,
+	inout FSSDCompressedSignalAccumulatorArray CompressedAccumulators,
 	float2 SampleBufferUV,
 	FSSDSampleSceneInfos SampleSceneMetadata,
 	FSSDSignalArray MultiplexedSamples,
@@ -392,6 +379,10 @@ void AccumulateSampledMultiplexedSignals(
 	// When sampling the previous frame, need to be forgiving in the world based bilateral rejection because TAA jitter.
 	const float WorldBluringDistanceForgiveness = KernelConfig.bPreviousFrameMetadata ? 3 : 1.0;
 	
+	#if CONFIG_ACCUMULATOR_VGPR_COMPRESSION != ACCUMULATOR_COMPRESSION_DISABLED
+		FSSDSignalAccumulatorArray Accumulators = UncompressAccumulatorArray(CompressedAccumulators, CONFIG_ACCUMULATOR_VGPR_COMPRESSION);
+	#endif
+
 	UNROLL_N(MAX_SIGNAL_MULTIPLEXING)
 	for (uint SignalMultiplexId = 0; SignalMultiplexId < MAX_SIGNAL_MULTIPLEXING; SignalMultiplexId++)
 	{
@@ -477,11 +468,25 @@ void AccumulateSampledMultiplexedSignals(
 		FLATTEN
 		if (SampleInfos.Sample.SampleCount != 0 && !bIsOutsideFrustum)
 		{
-			AccumulateSample(
-				/* inout */ Accumulators.Array[SignalMultiplexId],
-				SampleInfos);
+			#if CONFIG_ACCUMULATOR_VGPR_COMPRESSION == ACCUMULATOR_COMPRESSION_DISABLED
+			{
+				AccumulateSample(
+					/* inout */ UncompressedAccumulators.Array[SignalMultiplexId],
+					SampleInfos);
+			}
+			#else
+			{
+				AccumulateSample(
+					/* inout */ Accumulators.Array[SignalMultiplexId],
+					SampleInfos);
+			}
+			#endif
 		}
 	} // for (uint SignalMultiplexId = 0; SignalMultiplexId < MAX_SIGNAL_MULTIPLEXING; SignalMultiplexId++)
+	
+	#if CONFIG_ACCUMULATOR_VGPR_COMPRESSION != ACCUMULATOR_COMPRESSION_DISABLED
+		CompressedAccumulators = CompressAccumulatorArray(Accumulators, CONFIG_ACCUMULATOR_VGPR_COMPRESSION);
+	#endif
 } // AccumulateSampledMultiplexedSignals().
 
 /** Sample and accumulate to accumulatore array.
@@ -494,7 +499,8 @@ void SampleAndAccumulateMultiplexedSignals(
 	FSSDTexture2D SignalBuffer1,
 	FSSDTexture2D SignalBuffer2,
 	FSSDTexture2D SignalBuffer3,
-	inout FSSDSignalAccumulatorArray Accumulators,
+	inout FSSDSignalAccumulatorArray UncompressedAccumulators,
+	inout FSSDCompressedSignalAccumulatorArray CompressedAccumulators,
 	float2 SampleBufferUV,
 	float MipLevel,
 	float KernelSampleWeight,
@@ -519,7 +525,8 @@ void SampleAndAccumulateMultiplexedSignals(
 
 	AccumulateSampledMultiplexedSignals(
 		KernelConfig,
-		/* inout */ Accumulators,
+		/* inout */ UncompressedAccumulators,
+		/* inout */ CompressedAccumulators,
 		SampleBufferUV,
 		SampleSceneMetadata,
 		MultiplexedSamples,
@@ -534,7 +541,8 @@ void SampleAndAccumulateMultiplexedSignalsPair(
 	FSSDTexture2D SignalBuffer1,
 	FSSDTexture2D SignalBuffer2,
 	FSSDTexture2D SignalBuffer3,
-	inout FSSDSignalAccumulatorArray Accumulators,
+	inout FSSDSignalAccumulatorArray UncompressedAccumulators,
+	inout FSSDCompressedSignalAccumulatorArray CompressedAccumulators,
 	float2 SampleBufferUV[2],
 	float KernelSampleWeight)
 {
@@ -584,7 +592,8 @@ void SampleAndAccumulateMultiplexedSignalsPair(
 	{
 		AccumulateSampledMultiplexedSignals(
 			KernelConfig,
-			/* inout */ Accumulators,
+			/* inout */ UncompressedAccumulators,
+			/* inout */ CompressedAccumulators,
 			SampleBufferUV[PairAccumulateId],
 			SampleSceneMetadata[PairAccumulateId],
 			MultiplexedSamples[PairAccumulateId],
@@ -596,28 +605,69 @@ void SampleAndAccumulateMultiplexedSignalsPair(
 
 void StartAccumulatingCluster(
 	FSSDKernelConfig KernelConfig,
-	inout FSSDSignalAccumulatorArray Accumulators,
+	inout FSSDSignalAccumulatorArray UncompressedAccumulators,
+	inout FSSDCompressedSignalAccumulatorArray CompressedAccumulators,
 	FSSDSampleClusterInfo ClusterInfo)
 {
+	#if CONFIG_ACCUMULATOR_VGPR_COMPRESSION != ACCUMULATOR_COMPRESSION_DISABLED
+		FSSDSignalAccumulatorArray Accumulators = UncompressAccumulatorArray(CompressedAccumulators, CONFIG_ACCUMULATOR_VGPR_COMPRESSION);
+	#endif
+
 	UNROLL_N(MAX_SIGNAL_MULTIPLEXING)
 	for (uint SignalMultiplexId = 0; SignalMultiplexId < MAX_SIGNAL_MULTIPLEXING; SignalMultiplexId++)
 	{
-		StartAccumulatingCluster(
-			KernelConfig.RefSceneMetadata,
-			/* inout */ Accumulators.Array[SignalMultiplexId],
-			ClusterInfo);
+		#if CONFIG_ACCUMULATOR_VGPR_COMPRESSION == ACCUMULATOR_COMPRESSION_DISABLED
+		{
+			StartAccumulatingCluster(
+				KernelConfig.RefSceneMetadata,
+				/* inout */ UncompressedAccumulators.Array[SignalMultiplexId],
+				ClusterInfo);
+		}
+		#else
+		{
+			StartAccumulatingCluster(
+				KernelConfig.RefSceneMetadata,
+				/* inout */ Accumulators.Array[SignalMultiplexId],
+				ClusterInfo);
+		}
+		#endif
 	}
+	
+	#if CONFIG_ACCUMULATOR_VGPR_COMPRESSION != ACCUMULATOR_COMPRESSION_DISABLED
+		CompressedAccumulators = CompressAccumulatorArray(Accumulators, CONFIG_ACCUMULATOR_VGPR_COMPRESSION);
+	#endif
 }
 
-void DijestAccumulatedClusterSamples(inout FSSDSignalAccumulatorArray Accumulators, uint RingId, uint SampleCount)
+void DijestAccumulatedClusterSamples(
+	inout FSSDSignalAccumulatorArray UncompressedAccumulators,
+	inout FSSDCompressedSignalAccumulatorArray CompressedAccumulators,
+	uint RingId, uint SampleCount)
 {
+	#if CONFIG_ACCUMULATOR_VGPR_COMPRESSION != ACCUMULATOR_COMPRESSION_DISABLED
+		FSSDSignalAccumulatorArray Accumulators = UncompressAccumulatorArray(CompressedAccumulators, CONFIG_ACCUMULATOR_VGPR_COMPRESSION);
+	#endif
+		
 	UNROLL_N(MAX_SIGNAL_MULTIPLEXING)
 	for (uint SignalMultiplexId = 0; SignalMultiplexId < MAX_SIGNAL_MULTIPLEXING; SignalMultiplexId++)
 	{
-		DijestAccumulatedClusterSamples(
-			/* inout */ Accumulators.Array[SignalMultiplexId],
-			RingId, SampleCount);
+		#if CONFIG_ACCUMULATOR_VGPR_COMPRESSION == ACCUMULATOR_COMPRESSION_DISABLED
+		{
+			DijestAccumulatedClusterSamples(
+				/* inout */ UncompressedAccumulators.Array[SignalMultiplexId],
+				RingId, SampleCount);
+		}
+		#else
+		{
+			DijestAccumulatedClusterSamples(
+				/* inout */ Accumulators.Array[SignalMultiplexId],
+				RingId, SampleCount);
+		}
+		#endif
 	}
+	
+	#if CONFIG_ACCUMULATOR_VGPR_COMPRESSION != ACCUMULATOR_COMPRESSION_DISABLED
+		CompressedAccumulators = CompressAccumulatorArray(Accumulators, CONFIG_ACCUMULATOR_VGPR_COMPRESSION);
+	#endif
 }
 
 void SampleAndAccumulateCenterSampleAsItsOwnCluster(
@@ -626,14 +676,19 @@ void SampleAndAccumulateCenterSampleAsItsOwnCluster(
 	FSSDTexture2D SignalBuffer1,
 	FSSDTexture2D SignalBuffer2,
 	FSSDTexture2D SignalBuffer3,
-	inout FSSDSignalAccumulatorArray Accumulators)
+	inout FSSDSignalAccumulatorArray UncompressedAccumulators,
+	inout FSSDCompressedSignalAccumulatorArray CompressedAccumulators)
 {
 	const uint RingId = 0;
 
 	FSSDSampleClusterInfo ClusterInfo;
 	ClusterInfo.OutterBoundaryRadius = (RingId + 1) * KernelConfig.KernelSpreadFactor;
 
-	StartAccumulatingCluster(KernelConfig, Accumulators, ClusterInfo);
+	StartAccumulatingCluster(
+		KernelConfig, 
+		/* inout */ UncompressedAccumulators,
+		/* inout */ CompressedAccumulators,
+		ClusterInfo);
 
 	SampleAndAccumulateMultiplexedSignals(
 		KernelConfig,
@@ -641,13 +696,17 @@ void SampleAndAccumulateCenterSampleAsItsOwnCluster(
 		SignalBuffer1,
 		SignalBuffer2,
 		SignalBuffer3,
-		/* inout */ Accumulators,
+		/* inout */ UncompressedAccumulators,
+		/* inout */ CompressedAccumulators,
 		KernelConfig.BufferUV,
 		/* MipLevel = */ 0.0,
 		/* KernelSampleWeight = */ 1.0,
 		/* bForceSample = */ KernelConfig.bForceKernelCenterAccumulation);
 
-	DijestAccumulatedClusterSamples(Accumulators, RingId, /* SampleCount = */ 1);
+	DijestAccumulatedClusterSamples(
+		/* inout */ UncompressedAccumulators,
+		/* inout */ CompressedAccumulators,
+		RingId, /* SampleCount = */ 1);
 }
 
 
@@ -659,7 +718,8 @@ void AccumulateBilinear(
 	FSSDTexture2D SignalBuffer1,
 	FSSDTexture2D SignalBuffer2,
 	FSSDTexture2D SignalBuffer3,
-	inout FSSDSignalAccumulatorArray Accumulators)
+	inout FSSDSignalAccumulatorArray UncompressedAccumulators,
+	inout FSSDCompressedSignalAccumulatorArray CompressedAccumulators)
 {
 	const float MipLevelPow2 = 1;
 
@@ -686,7 +746,8 @@ void AccumulateBilinear(
 			SignalBuffer1,
 			SignalBuffer2,
 			SignalBuffer3,
-			/* inout */ Accumulators,
+			/* inout */ UncompressedAccumulators,
+			/* inout */ CompressedAccumulators,
 			SampleBufferUV,
 			/* MipLevel = */ 0.0,
 			BilinearWeight,
@@ -700,7 +761,8 @@ void AccumulateSquareKernel(
 	FSSDTexture2D SignalBuffer1,
 	FSSDTexture2D SignalBuffer2,
 	FSSDTexture2D SignalBuffer3,
-	inout FSSDSignalAccumulatorArray Accumulators)
+	inout FSSDSignalAccumulatorArray UncompressedAccumulators,
+	inout FSSDCompressedSignalAccumulatorArray CompressedAccumulators)
 {
 	int KernelRadius = 1;
 	if (KernelConfig.SampleSet == SAMPLE_SET_5X5_WAVELET)
@@ -747,7 +809,8 @@ void AccumulateSquareKernel(
 				SignalBuffer1,
 				SignalBuffer2,
 				SignalBuffer3,
-				/* inout */ Accumulators,
+				/* inout */ UncompressedAccumulators,
+				/* inout */ CompressedAccumulators,
 				SampleBufferUV,
 				/* MipLevel = */ 0.0,
 				KernelWeight,
@@ -790,7 +853,8 @@ void AccumulateSquareKernel(
 				SignalBuffer1,
 				SignalBuffer2,
 				SignalBuffer3,
-				/* inout */ Accumulators,
+				/* inout */ UncompressedAccumulators,
+				/* inout */ CompressedAccumulators,
 				SampleBufferUV,
 				/* MipLevel = */ 0.0,
 				KernelWeight,
@@ -931,7 +995,8 @@ void ConvolveStackowiakKernel(
 	FSSDTexture2D SignalBuffer1,
 	FSSDTexture2D SignalBuffer2,
 	FSSDTexture2D SignalBuffer3,
-	inout FSSDSignalAccumulatorArray Accumulators)
+	inout FSSDSignalAccumulatorArray UncompressedAccumulators,
+	inout FSSDCompressedSignalAccumulatorArray CompressedAccumulators)
 {
 	if (KernelConfig.bDescOrder)
 	{
@@ -967,7 +1032,11 @@ void ConvolveStackowiakKernel(
 		FSSDSampleClusterInfo ClusterInfo;
 		ClusterInfo.OutterBoundaryRadius = (CurrentRingId + 1) * KernelConfig.KernelSpreadFactor;
 
-		StartAccumulatingCluster(KernelConfig, /* inout */ Accumulators, ClusterInfo);
+		StartAccumulatingCluster(
+			KernelConfig,
+			/* inout */ UncompressedAccumulators,
+			/* inout */ CompressedAccumulators,
+			ClusterInfo);
 		
 		// Processes the samples in batches so that the compiler can do lattency hidding.
 		LOOP
@@ -1000,7 +1069,8 @@ void ConvolveStackowiakKernel(
 					SignalBuffer1,
 					SignalBuffer2,
 					SignalBuffer3,
-					/* inout */ Accumulators,
+					/* inout */ UncompressedAccumulators,
+					/* inout */ CompressedAccumulators,
 					SampleBufferUV,
 					/* MipLevel = */ 0.0,
 					KernelWeight,
@@ -1014,7 +1084,10 @@ void ConvolveStackowiakKernel(
 					uint SampleCountForCluster = min(CurrentRingId * StocasticSamplesPerCluster, SampleCount - SampleId);
 
 					// Digest all acumulators.
-					DijestAccumulatedClusterSamples(/* inout */ Accumulators, CurrentRingId, SampleCountForCluster);
+					DijestAccumulatedClusterSamples(
+						/* inout */ UncompressedAccumulators,
+						/* inout */ CompressedAccumulators,
+						CurrentRingId, SampleCountForCluster);
 				
 					BRANCH
 					if (!KernelConfig.bSampleKernelCenter && SampleId == 1)
@@ -1030,7 +1103,11 @@ void ConvolveStackowiakKernel(
 					ClusterInfo.OutterBoundaryRadius = (CurrentRingId + 1) * KernelConfig.KernelSpreadFactor;
 
 					// Prepare the accumulators for new cluster.
-					StartAccumulatingCluster(KernelConfig, /* inout */ Accumulators, ClusterInfo);
+					StartAccumulatingCluster(
+						KernelConfig, 
+						/* inout */ UncompressedAccumulators,
+						/* inout */ CompressedAccumulators,
+						ClusterInfo);
 				}
 			} // for (uint SampleBatchId = 0; SampleBatchId < kSamplingBatchSize; SampleBatchId++)
 		} // for (uint BatchId = 0; BatchId < BatchCountCount; BatchId++)
@@ -1038,7 +1115,10 @@ void ConvolveStackowiakKernel(
 		// NextClusterBoundary is not capable to reach 0, therefore need to manually digest the center sample.
 		if (KernelConfig.bSampleKernelCenter)
 		{
-			DijestAccumulatedClusterSamples(/* inout */ Accumulators, /* RingId = */ 0, /* SampleCount = */ 1);
+			DijestAccumulatedClusterSamples(
+				/* inout */ UncompressedAccumulators,
+				/* inout */ CompressedAccumulators,
+				/* RingId = */ 0, /* SampleCount = */ 1);
 		}
 	}
 	else // if (!KernelConfig.bDescOrder)
@@ -1069,7 +1149,8 @@ void ConvolveStackowiakKernel(
 				SignalBuffer1,
 				SignalBuffer2,
 				SignalBuffer3,
-				/* inout */ Accumulators,
+				/* inout */ UncompressedAccumulators,
+				/* inout */ CompressedAccumulators,
 				SampleBufferUV,
 				/* MipLevel = */ 0.0,
 				KernelWeight,
@@ -1167,7 +1248,8 @@ void GatherRingSamples(
 	FSSDTexture2D SignalBuffer1,
 	FSSDTexture2D SignalBuffer2,
 	FSSDTexture2D SignalBuffer3,
-	inout FSSDSignalAccumulatorArray Accumulators,
+	inout FSSDSignalAccumulatorArray UncompressedAccumulators,
+	inout FSSDCompressedSignalAccumulatorArray CompressedAccumulators,
 	const uint RingId)
 {
 	// Number of sample iteration for this ring.
@@ -1235,7 +1317,8 @@ void GatherRingSamples(
 				SignalBuffer1,
 				SignalBuffer2,
 				SignalBuffer3,
-				/* inout */ Accumulators,
+				/* inout */ UncompressedAccumulators,
+				/* inout */ CompressedAccumulators,
 				SampleUVPair,
 				/* KernelWeight = */ 1.0);
 		} // for (uint BatchSampleId = 0; BatchSampleId < SamplePairBatchSize; BatchSampleId++)
@@ -1248,7 +1331,8 @@ void ConvolveDiskKernel(
 	FSSDTexture2D SignalBuffer1,
 	FSSDTexture2D SignalBuffer2,
 	FSSDTexture2D SignalBuffer3,
-	inout FSSDSignalAccumulatorArray Accumulators)
+	inout FSSDSignalAccumulatorArray UncompressedAccumulators,
+	inout FSSDCompressedSignalAccumulatorArray CompressedAccumulators)
 {
 	// Accumulate the center of the kernel.
 	if (KernelConfig.bSampleKernelCenter && !KernelConfig.bDescOrder)
@@ -1259,7 +1343,8 @@ void ConvolveDiskKernel(
 			SignalBuffer1,
 			SignalBuffer2,
 			SignalBuffer3,
-			/* inout */ Accumulators);
+			/* inout */ UncompressedAccumulators,
+			/* inout */ CompressedAccumulators);
 	}
 
 	// Accumulate each ring. Use LOOP, because FXC is going through its pace otherwise.
@@ -1278,7 +1363,11 @@ void ConvolveDiskKernel(
 		FSSDSampleClusterInfo ClusterInfo;
 		ClusterInfo.OutterBoundaryRadius = (RingId + 1) * KernelConfig.KernelSpreadFactor;
 
-		StartAccumulatingCluster(KernelConfig, /* inout */ Accumulators, ClusterInfo);
+		StartAccumulatingCluster(
+			KernelConfig,
+			/* inout */ UncompressedAccumulators,
+			/* inout */ CompressedAccumulators,
+			ClusterInfo);
 		
 		GatherRingSamples(
 			KernelConfig,
@@ -1286,10 +1375,14 @@ void ConvolveDiskKernel(
 			SignalBuffer1,
 			SignalBuffer2,
 			SignalBuffer3,
-			/* inout */ Accumulators,
+			/* inout */ UncompressedAccumulators,
+			/* inout */ CompressedAccumulators,
 			RingId);
 			
-		DijestAccumulatedClusterSamples(/* inout */ Accumulators, RingId, RingSamplePairCount * 2);
+		DijestAccumulatedClusterSamples(
+			/* inout */ UncompressedAccumulators,
+			/* inout */ CompressedAccumulators,
+			RingId, RingSamplePairCount * 2);
 	} // for (uint RingId = 0; RingId < KernelConfig.RingCount; RingId++)
 		
 	// Accumulate the center of the kernel.
@@ -1301,7 +1394,8 @@ void ConvolveDiskKernel(
 			SignalBuffer1,
 			SignalBuffer2,
 			SignalBuffer3,
-			/* inout */ Accumulators);
+			/* inout */ UncompressedAccumulators,
+			/* inout */ CompressedAccumulators);
 	}
 }
 
@@ -1314,7 +1408,8 @@ void ConvolveStocasticHierarchy(
 	FSSDTexture2D SignalBuffer1,
 	FSSDTexture2D SignalBuffer2,
 	FSSDTexture2D SignalBuffer3,
-	inout FSSDSignalAccumulatorArray Accumulators)
+	inout FSSDSignalAccumulatorArray UncompressedAccumulators,
+	inout FSSDCompressedSignalAccumulatorArray CompressedAccumulators)
 {
 	// Accumulate the center of the kernel.
 	if (KernelConfig.bSampleKernelCenter && !KernelConfig.bDescOrder)
@@ -1325,7 +1420,8 @@ void ConvolveStocasticHierarchy(
 			SignalBuffer1,
 			SignalBuffer2,
 			SignalBuffer3,
-			/* inout */ Accumulators);
+			/* inout */ UncompressedAccumulators,
+			/* inout */ CompressedAccumulators);
 	}
 
 	UNROLL_N(1)
@@ -1362,7 +1458,8 @@ void ConvolveStocasticHierarchy(
 				SignalBuffer1,
 				SignalBuffer2,
 				SignalBuffer3,
-				/* inout */ Accumulators,
+				/* inout */ UncompressedAccumulators,
+				/* inout */ CompressedAccumulators,
 				SampleBufferUV,
 				/* MipLevel = */ float(MipLevel),
 				/* KernelSampleWeight = */ 1.0,
@@ -1379,7 +1476,8 @@ void ConvolveStocasticHierarchy(
 			SignalBuffer1,
 			SignalBuffer2,
 			SignalBuffer3,
-			/* inout */ Accumulators);
+			/* inout */ UncompressedAccumulators,
+			/* inout */ CompressedAccumulators);
 	}
 }
 
@@ -1391,7 +1489,8 @@ void AccumulateKernel(
 	FSSDTexture2D SignalBuffer1,
 	FSSDTexture2D SignalBuffer2,
 	FSSDTexture2D SignalBuffer3,
-	inout FSSDSignalAccumulatorArray Accumulators)
+	inout FSSDSignalAccumulatorArray UncompressedAccumulators,
+	inout FSSDCompressedSignalAccumulatorArray CompressedAccumulators)
 {
 	if (KernelConfig.SampleSet == SAMPLE_SET_1X1)
 	{
@@ -1403,7 +1502,8 @@ void AccumulateKernel(
 				SignalBuffer1,
 				SignalBuffer2,
 				SignalBuffer3,
-				/* inout */ Accumulators);
+				/* inout */ UncompressedAccumulators,
+				/* inout */ CompressedAccumulators);
 		}
 	}
 	else if (KernelConfig.SampleSet == SAMPLE_SET_2X2_BILINEAR)
@@ -1414,7 +1514,8 @@ void AccumulateKernel(
 			SignalBuffer1,
 			SignalBuffer2,
 			SignalBuffer3,
-			/* inout */ Accumulators);
+			/* inout */ UncompressedAccumulators,
+			/* inout */ CompressedAccumulators);
 	}
 	else if (
 		KernelConfig.SampleSet == SAMPLE_SET_3X3 ||
@@ -1430,7 +1531,8 @@ void AccumulateKernel(
 			SignalBuffer1,
 			SignalBuffer2,
 			SignalBuffer3,
-			/* inout */ Accumulators);
+			/* inout */ UncompressedAccumulators,
+			/* inout */ CompressedAccumulators);
 	}
 	else if (KernelConfig.SampleSet == SAMPLE_SET_STACKOWIAK_4_SETS)
 	{
@@ -1440,7 +1542,8 @@ void AccumulateKernel(
 			SignalBuffer1,
 			SignalBuffer2,
 			SignalBuffer3,
-			/* inout */ Accumulators);
+			/* inout */ UncompressedAccumulators,
+			/* inout */ CompressedAccumulators);
 	}
 	else if (KernelConfig.SampleSet == SAMPLE_SET_HEXAWEB)
 	{
@@ -1450,7 +1553,8 @@ void AccumulateKernel(
 			SignalBuffer1,
 			SignalBuffer2,
 			SignalBuffer3,
-			/* inout */ Accumulators);
+			/* inout */ UncompressedAccumulators,
+			/* inout */ CompressedAccumulators);
 	}
 	else if (KernelConfig.SampleSet == SAMPLE_SET_STOCASTIC_HIERARCHY)
 	{
@@ -1460,6 +1564,7 @@ void AccumulateKernel(
 			SignalBuffer1,
 			SignalBuffer2,
 			SignalBuffer3,
-			/* inout */ Accumulators);
+			/* inout */ UncompressedAccumulators,
+			/* inout */ CompressedAccumulators);
 	}
 } // AccumulateKernel()
diff --git a/Engine/Shaders/Private/ScreenSpaceDenoise/SSDTemporalAccumulation.usf b/Engine/Shaders/Private/ScreenSpaceDenoise/SSDTemporalAccumulation.usf
--- a/Engine/Shaders/Private/ScreenSpaceDenoise/SSDTemporalAccumulation.usf
+++ b/Engine/Shaders/Private/ScreenSpaceDenoise/SSDTemporalAccumulation.usf
@@ -452,13 +452,16 @@ void TemporallyAccumulate(
 			KernelConfig.RefSceneMetadata = RefSceneMetadata;
 			
 			FSSDSignalAccumulatorArray SignalAccumulators = CreateSignalAccumulatorArray();
+			FSSDCompressedSignalAccumulatorArray UnusedCompressedAccumulators = CreateUninitialisedCompressedAccumulatorArray();
+
 			AccumulateKernel(
 				KernelConfig,
 				PrevHistory_Textures_0,
 				PrevHistory_Textures_1,
 				PrevHistory_Textures_2,
 				PrevHistory_Textures_3,
-				/* inout */ SignalAccumulators);
+				/* inout */ SignalAccumulators,
+				/* inout */ UnusedCompressedAccumulators);
 		
 			// Exports the history sample from accumulator.
 			UNROLL_N(CONFIG_SIGNAL_BATCH_SIZE)
@@ -620,23 +623,29 @@ void TemporallyAccumulate(
 			}
 		}
 
-		#if CONFIG_USE_REJECTION_BUFFER
-			AccumulateKernel(
-				KernelConfig,
-				HistoryRejectionSignal_Textures_0,
-				HistoryRejectionSignal_Textures_1,
-				HistoryRejectionSignal_Textures_2,
-				HistoryRejectionSignal_Textures_3,
-				/* inout */ SignalAccumulators);
-		#else
-			AccumulateKernel(
-				KernelConfig,
-				SignalInput_Textures_0,
-				SignalInput_Textures_1,
-				SignalInput_Textures_2,
-				SignalInput_Textures_3,
-				/* inout */ SignalAccumulators);
-		#endif
+		{
+			FSSDCompressedSignalAccumulatorArray UnusedCompressedAccumulators = CreateUninitialisedCompressedAccumulatorArray();
+
+			#if CONFIG_USE_REJECTION_BUFFER
+				AccumulateKernel(
+					KernelConfig,
+					HistoryRejectionSignal_Textures_0,
+					HistoryRejectionSignal_Textures_1,
+					HistoryRejectionSignal_Textures_2,
+					HistoryRejectionSignal_Textures_3,
+					/* inout */ SignalAccumulators,
+					/* inout */ UnusedCompressedAccumulators);
+			#else
+				AccumulateKernel(
+					KernelConfig,
+					SignalInput_Textures_0,
+					SignalInput_Textures_1,
+					SignalInput_Textures_2,
+					SignalInput_Textures_3,
+					/* inout */ SignalAccumulators,
+					/* inout */ UnusedCompressedAccumulators);
+			#endif
+		}
 
 		// Clamp history
 		UNROLL_N(CONFIG_SIGNAL_BATCH_SIZE)
