Adds options to compile out chuncks of the denoiser code base to ease the compilers

diff --git a/Engine/Shaders/Private/ScreenSpaceDenoise/SSDSignalAccumulator.ush b/Engine/Shaders/Private/ScreenSpaceDenoise/SSDSignalAccumulator.ush
--- a/Engine/Shaders/Private/ScreenSpaceDenoise/SSDSignalAccumulator.ush
+++ b/Engine/Shaders/Private/ScreenSpaceDenoise/SSDSignalAccumulator.ush
@@ -10,6 +10,26 @@
 	#define CONFIG_SIGNAL_VGPR_COMPRESSION SIGNAL_COMPRESSION_DISABLED
 #endif
 
+#ifndef COMPILE_MOMENT1_ACCUMULATOR
+	#define COMPILE_MOMENT1_ACCUMULATOR 0
+#endif
+
+#ifndef COMPILE_MOMENT2_ACCUMULATOR
+	#define COMPILE_MOMENT2_ACCUMULATOR 0
+#endif
+
+#ifndef COMPILE_MINMAX_ACCUMULATOR
+	#define COMPILE_MINMAX_ACCUMULATOR 0
+#endif
+
+#ifndef COMPILE_MININVFREQ_ACCUMULATOR
+	#define COMPILE_MININVFREQ_ACCUMULATOR 0
+#endif
+
+#ifndef COMPILE_DRB_ACCUMULATOR
+	#define COMPILE_DRB_ACCUMULATOR 0
+#endif
+
 
 //------------------------------------------------------- SIGNAL ACCUMULATOR
 
@@ -19,30 +39,46 @@
 struct FSSDSignalAccumulator
 {
 	// Sums of moment.
-	FSSDSignalSample Moment1;
-	FSSDSignalSample Moment2;
+	#if COMPILE_MOMENT1_ACCUMULATOR
+		FSSDSignalSample Moment1;
+	#endif
+
+	#if COMPILE_MOMENT2_ACCUMULATOR
+		FSSDSignalSample Moment2;
+	#endif
 
 	// Compressed form of moment signals.
-	FSSDCompressedSignalSample CompressedMoment1;
-	FSSDCompressedSignalSample CompressedMoment2;
+	#if COMPILE_MOMENT1_ACCUMULATOR
+		FSSDCompressedSignalSample CompressedMoment1;
+	#endif
+
+	#if COMPILE_MOMENT2_ACCUMULATOR
+		FSSDCompressedSignalSample CompressedMoment2;
+	#endif
 
 	// Per component min and max values.
-	FSSDSignalSample Min;
-	FSSDSignalSample Max;
+	#if COMPILE_MINMAX_ACCUMULATOR
+		FSSDSignalSample Min;
+		FSSDSignalSample Max;
+	#endif
 
 	// Minimal inverse frequency found and intersected when sampling.
-	float MinInvFrequency;
-	
+	#if COMPILE_MININVFREQ_ACCUMULATOR
+		float MinInvFrequency;
+	#endif
+
 	// Descending ring bucketing.
-	FSSDSignalSample Current;
-	float CurrentInvFrequency;
-	float CurrentTranslucency;
+	#if COMPILE_DRB_ACCUMULATOR
+		FSSDSignalSample Current;
+		float CurrentInvFrequency;
+		float CurrentTranslucency;
 
-	FSSDSignalSample Previous;
-	float PreviousInvFrequency;
+		FSSDSignalSample Previous;
+		float PreviousInvFrequency;
 
-	float BorderingRadius;
-	float HighestInvFrequency;
+		float BorderingRadius;
+		float HighestInvFrequency;
+	#endif // COMPILE_DRB_ACCUMULATOR
 };
 
 /** Informations about cluster of sample */
@@ -57,18 +93,33 @@ FSSDSignalAccumulator CreateSignalAccumulator()
 {
 	FSSDSignalAccumulator Accumulator;
 	
-	Accumulator.Moment1 = CreateSignalSampleFromScalarValue(0.0);
-	Accumulator.Moment2 = CreateSignalSampleFromScalarValue(0.0);
-		
-	Accumulator.CompressedMoment1 = CreateCompressedSignalSampleFromScalarValue(0.0, CONFIG_SIGNAL_VGPR_COMPRESSION);
-	Accumulator.CompressedMoment2 = CreateCompressedSignalSampleFromScalarValue(0.0, CONFIG_SIGNAL_VGPR_COMPRESSION);
+	#if COMPILE_MOMENT1_ACCUMULATOR
+		Accumulator.Moment1 = CreateSignalSampleFromScalarValue(0.0);
+	#endif
+
+	#if COMPILE_MOMENT2_ACCUMULATOR
+		Accumulator.Moment2 = CreateSignalSampleFromScalarValue(0.0);
+	#endif	
 
-	Accumulator.Min = CreateSignalSampleFromScalarValue(INFINITE_FLOAT);
-	Accumulator.Max = CreateSignalSampleFromScalarValue(0.0); // TODO: -INFINITE_FLOAT? otherwise there is a risk to color clamp with YCoCg.
+	#if COMPILE_MOMENT1_ACCUMULATOR
+		Accumulator.CompressedMoment1 = CreateCompressedSignalSampleFromScalarValue(0.0, CONFIG_SIGNAL_VGPR_COMPRESSION);
+	#endif
 
-	Accumulator.MinInvFrequency = WORLD_RADIUS_MISS;
+	#if COMPILE_MOMENT2_ACCUMULATOR
+		Accumulator.CompressedMoment2 = CreateCompressedSignalSampleFromScalarValue(0.0, CONFIG_SIGNAL_VGPR_COMPRESSION);
+	#endif
+	
+	#if COMPILE_MINMAX_ACCUMULATOR
+		Accumulator.Min = CreateSignalSampleFromScalarValue(INFINITE_FLOAT);
+		Accumulator.Max = CreateSignalSampleFromScalarValue(0.0); // TODO: -INFINITE_FLOAT? otherwise there is a risk to color clamp with YCoCg.
+	#endif
+		
+	#if COMPILE_MININVFREQ_ACCUMULATOR
+		Accumulator.MinInvFrequency = WORLD_RADIUS_MISS;
+	#endif
 
 	// DRB initialization.
+	#if COMPILE_DRB_ACCUMULATOR
 	{
 		Accumulator.Current = CreateSignalSampleFromScalarValue(0.0);
 		Accumulator.CurrentInvFrequency = 0.0;
@@ -80,6 +131,7 @@ FSSDSignalAccumulator CreateSignalAccumulator()
 		Accumulator.BorderingRadius = 0.0;
 		Accumulator.HighestInvFrequency = 0.0;
 	}
+	#endif
 
 	return Accumulator;
 }
@@ -122,8 +174,14 @@ float AmendWorldBluringRadiusCausedByPixelSize(float WorldBluringDistance)
 void AccumulateSampleInOneBin(inout FSSDSignalAccumulator Accumulator, FSSDSampleAccumulationInfos SampleInfos)
 {
 	float SampleWeight = SampleInfos.FinalWeight;
-	Accumulator.Moment1 = Add(Accumulator.Moment1, Mul(SampleInfos.Sample, SampleWeight));
-	Accumulator.Moment2 = Add(Accumulator.Moment2, Mul(Power(SampleInfos.Sample, 2), SampleWeight));
+	
+	#if COMPILE_MOMENT1_ACCUMULATOR
+		Accumulator.Moment1 = Add(Accumulator.Moment1, Mul(SampleInfos.Sample, SampleWeight));
+	#endif
+	
+	#if COMPILE_MOMENT2_ACCUMULATOR
+		Accumulator.Moment2 = Add(Accumulator.Moment2, Mul(Power(SampleInfos.Sample, 2), SampleWeight));
+	#endif
 }
 
 
@@ -135,12 +193,19 @@ void AccumulateSampleInDomainBoundaries(inout FSSDSignalAccumulator Accumulator,
 	FLATTEN
 	if (SampleInfos.FinalWeight > 0)
 	{
-		Accumulator.Min = Min(Accumulator.Min, SampleInfos.Sample);
-		Accumulator.Max = Max(Accumulator.Max, SampleInfos.Sample);
+		#if COMPILE_MINMAX_ACCUMULATOR
+		{
+			Accumulator.Min = Min(Accumulator.Min, SampleInfos.Sample);
+			Accumulator.Max = Max(Accumulator.Max, SampleInfos.Sample);
+		}
+		#endif
 		
 		// Track the minimal frequency inverse.
-		if (SampleInfos.InvFrequency != WORLD_RADIUS_MISS)
-			Accumulator.MinInvFrequency = min(Accumulator.MinInvFrequency, SampleInfos.InvFrequency);
+		
+		#if COMPILE_MININVFREQ_ACCUMULATOR
+			if (SampleInfos.InvFrequency != WORLD_RADIUS_MISS)
+				Accumulator.MinInvFrequency = min(Accumulator.MinInvFrequency, SampleInfos.InvFrequency);
+		#endif
 	}
 }
 
@@ -148,6 +213,8 @@ void AccumulateSampleInDomainBoundaries(inout FSSDSignalAccumulator Accumulator,
 //------------------------------------------------------- DESCENDING RING BUCKETING
 // [ SIGGRAPH 2018, "A life of bokeh" ]
 
+#if COMPILE_DRB_ACCUMULATOR
+
 void StartAccumulatingClusterInDRB(
 	FSSDSampleSceneInfos RefSceneMetadata,
 	inout FSSDSignalAccumulator Accumulator,
@@ -249,19 +316,31 @@ void DijestSampleClusterInDRB(
 	Accumulator.PreviousInvFrequency = Accumulator.PreviousInvFrequency * PreviousBucketFactor + Accumulator.CurrentInvFrequency;
 }
 
+#endif // COMPILE_DRB_ACCUMULATOR
+
 
 //------------------------------------------------------- COMPRESSION & DECOMPRESSION
 
 void CompressSignalAccumulator(inout FSSDSignalAccumulator Accumulator)
 {
-	CompressSignalSample(Accumulator.Moment1, CONFIG_SIGNAL_VGPR_COMPRESSION, /* out */ Accumulator.CompressedMoment1);
-	CompressSignalSample(Accumulator.Moment2, CONFIG_SIGNAL_VGPR_COMPRESSION, /* out */ Accumulator.CompressedMoment2);
+	#if COMPILE_MOMENT1_ACCUMULATOR
+		CompressSignalSample(Accumulator.Moment1, CONFIG_SIGNAL_VGPR_COMPRESSION, /* out */ Accumulator.CompressedMoment1);
+	#endif
+	
+	#if COMPILE_MOMENT2_ACCUMULATOR
+		CompressSignalSample(Accumulator.Moment2, CONFIG_SIGNAL_VGPR_COMPRESSION, /* out */ Accumulator.CompressedMoment2);
+	#endif
 }
 
 void UncompressSignalAccumulator(inout FSSDSignalAccumulator Accumulator)
 {
-	UncompressSignalSample(Accumulator.CompressedMoment1, CONFIG_SIGNAL_VGPR_COMPRESSION, /* inout */ Accumulator.Moment1);
-	UncompressSignalSample(Accumulator.CompressedMoment2, CONFIG_SIGNAL_VGPR_COMPRESSION, /* inout */ Accumulator.Moment2);
+	#if COMPILE_MOMENT1_ACCUMULATOR
+		UncompressSignalSample(Accumulator.CompressedMoment1, CONFIG_SIGNAL_VGPR_COMPRESSION, /* inout */ Accumulator.Moment1);
+	#endif
+	
+	#if COMPILE_MOMENT2_ACCUMULATOR
+		UncompressSignalSample(Accumulator.CompressedMoment2, CONFIG_SIGNAL_VGPR_COMPRESSION, /* inout */ Accumulator.Moment2);
+	#endif
 }
 
 
@@ -273,7 +352,9 @@ void StartAccumulatingCluster(
 	inout FSSDSignalAccumulator Accumulator,
 	FSSDSampleClusterInfo ClusterInfo)
 {
-	StartAccumulatingClusterInDRB(RefSceneMetadata, /* inout */ Accumulator, ClusterInfo);
+	#if COMPILE_DRB_ACCUMULATOR
+		StartAccumulatingClusterInDRB(RefSceneMetadata, /* inout */ Accumulator, ClusterInfo);
+	#endif
 }
 
 /** Accumulate a sample withing the accumulator. */
@@ -285,7 +366,10 @@ void AccumulateSample(
 
 	AccumulateSampleInOneBin(/* inout */ Accumulator, SampleInfos);
 	AccumulateSampleInDomainBoundaries(/* inout */ Accumulator, SampleInfos);
-	AccumulateSampleInDRB(/* inout */ Accumulator, SampleInfos);
+
+	#if COMPILE_DRB_ACCUMULATOR
+		AccumulateSampleInDRB(/* inout */ Accumulator, SampleInfos);
+	#endif // COMPILE_DRB_ACCUMULATOR
 
 	CompressSignalAccumulator(/* inout */ Accumulator);
 }
@@ -295,5 +379,7 @@ void DijestAccumulatedClusterSamples(
 	inout FSSDSignalAccumulator Accumulator,
 	uint RingId, uint SampleCount)
 {
-	DijestSampleClusterInDRB(/* inout */ Accumulator, RingId, SampleCount);
+	#if COMPILE_DRB_ACCUMULATOR
+		DijestSampleClusterInDRB(/* inout */ Accumulator, RingId, SampleCount);
+	#endif
 }
diff --git a/Engine/Shaders/Private/ScreenSpaceDenoise/SSDSignalAccumulatorArray.ush b/Engine/Shaders/Private/ScreenSpaceDenoise/SSDSignalAccumulatorArray.ush
--- a/Engine/Shaders/Private/ScreenSpaceDenoise/SSDSignalAccumulatorArray.ush
+++ b/Engine/Shaders/Private/ScreenSpaceDenoise/SSDSignalAccumulatorArray.ush
@@ -62,6 +62,7 @@ FSSDCompressedSignalAccumulatorArray CompressAccumulatorArray(in FSSDSignalAccum
 	{
 		// NOP
 	}
+	#if COMPILE_DRB_ACCUMULATOR && COMPILE_MININVFREQ_ACCUMULATOR
 	else if (CompressionLayout == ACCUMULATOR_COMPRESSION_PENUMBRA_DRB)
 	{
 		UNROLL_N(MAX_SIGNAL_BATCH_SIZE)
@@ -83,6 +84,7 @@ FSSDCompressedSignalAccumulatorArray CompressAccumulatorArray(in FSSDSignalAccum
 			RawArray[SignalBatchId * 6 + 5].y = Accum.BorderingRadius;
 		} // for (uint SignalMultiplexId = 0; SignalMultiplexId < MAX_SIGNAL_BATCH_SIZE; SignalMultiplexId++)
 	}
+	#endif // COMPILE_DRB_ACCUMULATOR && COMPILE_MININVFREQ_ACCUMULATOR
 
 	// Compress.
 	FSSDCompressedSignalAccumulatorArray CompressedAccumulators;
@@ -114,6 +116,7 @@ FSSDSignalAccumulatorArray UncompressAccumulatorArray(in FSSDCompressedSignalAcc
 	{
 		// NOP to not clobber information in OutSample that remained uncompressed.
 	}
+	#if COMPILE_DRB_ACCUMULATOR && COMPILE_MININVFREQ_ACCUMULATOR
 	else if (CompressionLayout == ACCUMULATOR_COMPRESSION_PENUMBRA_DRB)
 	{
 		// TODO: CONFIG_SIGNAL_BATCH_SIZE?
@@ -138,6 +141,7 @@ FSSDSignalAccumulatorArray UncompressAccumulatorArray(in FSSDCompressedSignalAcc
 			Accumulators.Array[SignalBatchId] = Accum;
 		} // for (uint SignalBatchId = 0; SignalBatchId < MAX_SIGNAL_BATCH_SIZE; SignalBatchId++)
 	}
+	#endif // COMPILE_DRB_ACCUMULATOR && COMPILE_MININVFREQ_ACCUMULATOR
 
 	return Accumulators;
 } // UncompressAccumulatorArray()
diff --git a/Engine/Shaders/Private/ScreenSpaceDenoise/SSDSpatialAccumulation.usf b/Engine/Shaders/Private/ScreenSpaceDenoise/SSDSpatialAccumulation.usf
--- a/Engine/Shaders/Private/ScreenSpaceDenoise/SSDSpatialAccumulation.usf
+++ b/Engine/Shaders/Private/ScreenSpaceDenoise/SSDSpatialAccumulation.usf
@@ -390,6 +390,35 @@
 #endif
 
 
+//------------------------------------------------------- COMPILATION WHITE LISTING
+
+// Whitelist kernel to compile.
+#if CONFIG_SAMPLE_SET == SAMPLE_SET_STACKOWIAK_4_SETS
+	#define COMPILE_STACKOWIAK_KERNEL 1
+#elif CONFIG_SAMPLE_SET == SAMPLE_SET_HEXAWEB
+	#define COMPILE_DISK_KERNEL 1
+#else
+	#define COMPILE_BOX_KERNEL 1
+#endif
+
+// White list accumulators to compile.
+#if CONFIG_OUTPUT_MODE == OUTPUT_MODE_DRB
+	#define COMPILE_DRB_ACCUMULATOR 1
+	#define COMPILE_MININVFREQ_ACCUMULATOR 1
+
+#elif CONFIG_OUTPUT_MODE == OUTPUT_MODE_2MOMMENT_SUM
+	#define COMPILE_MOMENT1_ACCUMULATOR 1
+	#define COMPILE_MOMENT2_ACCUMULATOR 1
+
+#elif CONFIG_OUTPUT_MODE == OUTPUT_MODE_SUM
+	#define COMPILE_MOMENT1_ACCUMULATOR 1
+	#define COMPILE_MININVFREQ_ACCUMULATOR 1
+
+#else
+	#error Unknown output mode.
+#endif
+
+
 //------------------------------------------------------- LATE SETUP
 
 #if COMPILER_SUPPORTS_WAVE_ONCE
@@ -847,15 +876,21 @@ void MainCS(
 			#endif // CONFIG_REJECT_HIGHEST_COLOR
 
 			#if CONFIG_ACCUMULATION_COLOR_SPACE != CONFIG_OUTPUT_COLOR_SPACE
+			{
+				#if COMPILE_MOMENT1_ACCUMULATOR
 				SignalAccumulators.Array[MultiplexId].Moment1 = TransformSignal(
 					SignalAccumulators.Array[MultiplexId].Moment1, 
 					/* SrcBasis  = */ CONFIG_ACCUMULATION_COLOR_SPACE, 
 					/* DestBasis = */ CONFIG_OUTPUT_COLOR_SPACE);
-
+				#endif
+				
+				#if COMPILE_MOMENT2_ACCUMULATOR
 				SignalAccumulators.Array[MultiplexId].Moment2 = TransformSignal(
 					SignalAccumulators.Array[MultiplexId].Moment2, 
 					/* SrcBasis  = */ CONFIG_ACCUMULATION_COLOR_SPACE, 
 					/* DestBasis = */ CONFIG_OUTPUT_COLOR_SPACE);
+				#endif
+			}
 			#endif // CONFIG_ACCUMULATION_COLOR_SPACE != CONFIG_OUTPUT_COLOR_SPACE
 
 			// TODO: it might be better to just uncompress before this for loop and remain uncompressed,
diff --git a/Engine/Shaders/Private/ScreenSpaceDenoise/SSDSpatialKernel.ush b/Engine/Shaders/Private/ScreenSpaceDenoise/SSDSpatialKernel.ush
--- a/Engine/Shaders/Private/ScreenSpaceDenoise/SSDSpatialKernel.ush
+++ b/Engine/Shaders/Private/ScreenSpaceDenoise/SSDSpatialKernel.ush
@@ -41,6 +41,22 @@
 #endif
 
 
+#ifndef COMPILE_BOX_KERNEL
+	#define COMPILE_BOX_KERNEL 0
+#endif
+
+#ifndef COMPILE_STACKOWIAK_KERNEL
+	#define COMPILE_STACKOWIAK_KERNEL 0
+#endif
+
+#ifndef COMPILE_DISK_KERNEL
+	#define COMPILE_DISK_KERNEL 0
+#endif
+
+#ifndef COMPILE_HEIRARCHY_KERNEL
+	#define COMPILE_HEIRARCHY_KERNEL 0
+#endif
+
 
 //------------------------------------------------------- STRUCTURES
 
@@ -712,6 +728,8 @@ void SampleAndAccumulateCenterSampleAsItsOwnCluster(
 
 //------------------------------------------------------- EASY CONVOLUTIONS
 
+#if COMPILE_BOX_KERNEL
+
 void AccumulateBilinear(
 	FSSDKernelConfig KernelConfig,
 	FSSDTexture2D SignalBuffer0,
@@ -862,10 +880,14 @@ void AccumulateSquareKernel(
 		}
 	}
 } // AccumulateSquareKernel()
+	
+#endif // COMPILE_BOX_KERNEL
 
 
 //------------------------------------------------------- STACKOWIAK 2018
 
+#if COMPILE_STACKOWIAK_KERNEL
+
 static const float2 kStackowiakSampleSet0[56 * 4] =
 {
 	float2(-0.5, -0.5), float2(+0.5, -0.5), float2(-0.5, +0.5), float2(+0.5, +0.5),
@@ -1159,9 +1181,13 @@ void ConvolveStackowiakKernel(
 	} // if (!KernelConfig.bDescOrder)
 } // ConvolveStackowiakKernel()
 
+#endif // COMPILE_STACKOWIAK_KERNEL
+
 
 //------------------------------------------------------- DISK
 
+#if COMPILE_DISK_KERNEL
+
 // Returns the position of the sample on the unit circle (radius = 1) for a given ring.
 float2 GetDiskSampleOnUnitCirle(uint RingId, uint RingSampleIteration, uint RingSampleId)
 {
@@ -1399,9 +1425,13 @@ void ConvolveDiskKernel(
 	}
 }
 
+#endif // COMPILE_DISK_KERNEL
+
 
 //------------------------------------------------------- STOCASTIC HIERARCHY
 
+#if COMPILE_HEIRARCHY_KERNEL
+
 void ConvolveStocasticHierarchy(
 	FSSDKernelConfig KernelConfig,
 	FSSDTexture2D SignalBuffer0,
@@ -1481,6 +1511,8 @@ void ConvolveStocasticHierarchy(
 	}
 }
 
+#endif // COMPILE_HEIRARCHY_KERNEL
+
 //------------------------------------------------------- MAIN ENTRY POINTS
 
 void AccumulateKernel(
@@ -1492,7 +1524,11 @@ void AccumulateKernel(
 	inout FSSDSignalAccumulatorArray UncompressedAccumulators,
 	inout FSSDCompressedSignalAccumulatorArray CompressedAccumulators)
 {
-	if (KernelConfig.SampleSet == SAMPLE_SET_1X1)
+	if (KernelConfig.SampleSet == 0xDEADDEAD)
+	{
+	}
+#if COMPILE_BOX_KERNEL
+	else if (KernelConfig.SampleSet == SAMPLE_SET_1X1)
 	{
 		if (KernelConfig.bSampleKernelCenter)
 		{
@@ -1534,6 +1570,8 @@ void AccumulateKernel(
 			/* inout */ UncompressedAccumulators,
 			/* inout */ CompressedAccumulators);
 	}
+#endif//  COMPILE_BOX_KERNEL
+#if COMPILE_STACKOWIAK_KERNEL
 	else if (KernelConfig.SampleSet == SAMPLE_SET_STACKOWIAK_4_SETS)
 	{
 		ConvolveStackowiakKernel(
@@ -1545,6 +1583,8 @@ void AccumulateKernel(
 			/* inout */ UncompressedAccumulators,
 			/* inout */ CompressedAccumulators);
 	}
+#endif // COMPILE_STACKOWIAK_KERNEL
+#if COMPILE_DISK_KERNEL
 	else if (KernelConfig.SampleSet == SAMPLE_SET_HEXAWEB)
 	{
 		ConvolveDiskKernel(
@@ -1556,6 +1596,8 @@ void AccumulateKernel(
 			/* inout */ UncompressedAccumulators,
 			/* inout */ CompressedAccumulators);
 	}
+#endif // COMPILE_DISK_KERNEL
+#if COMPILE_HEIRARCHY_KERNEL
 	else if (KernelConfig.SampleSet == SAMPLE_SET_STOCASTIC_HIERARCHY)
 	{
 		ConvolveStocasticHierarchy(
@@ -1567,4 +1609,5 @@ void AccumulateKernel(
 			/* inout */ UncompressedAccumulators,
 			/* inout */ CompressedAccumulators);
 	}
+#endif // COMPILE_HEIRARCHY_KERNEL
 } // AccumulateKernel()
diff --git a/Engine/Shaders/Private/ScreenSpaceDenoise/SSDTemporalAccumulation.usf b/Engine/Shaders/Private/ScreenSpaceDenoise/SSDTemporalAccumulation.usf
--- a/Engine/Shaders/Private/ScreenSpaceDenoise/SSDTemporalAccumulation.usf
+++ b/Engine/Shaders/Private/ScreenSpaceDenoise/SSDTemporalAccumulation.usf
@@ -216,6 +216,24 @@
 #endif
 
 
+//------------------------------------------------------- COMPILATION WHITE LISTING
+
+#define COMPILE_BOX_KERNEL 1
+
+// Moment 1 needed for previous frame sampling.
+#define COMPILE_MOMENT1_ACCUMULATOR 1
+
+#if CONFIG_HISTORY_REJECTION != HISTORY_REJECTION_DISABLED
+	#define COMPILE_MININVFREQ_ACCUMULATOR 1
+#endif
+
+#if CONFIG_HISTORY_REJECTION == HISTORY_REJECTION_MINMAX_BOUNDARIES
+	#define COMPILE_MINMAX_ACCUMULATOR 1
+#elif CONFIG_HISTORY_REJECTION == HISTORY_REJECTION_VAR_BOUNDARIES && CONFIG_REJECTION_INPUT_MODE == REJECTION_INPUT_MODE_1UNNORMALIZED
+	#define COMPILE_MOMENT2_ACCUMULATOR 1
+#endif
+
+
 //------------------------------------------------------- INCLUDES
 
 #include "SSDSignalFramework.ush"
@@ -651,16 +669,16 @@ void TemporallyAccumulate(
 		UNROLL_N(CONFIG_SIGNAL_BATCH_SIZE)
 		for (uint BatchedSignalId = 0; BatchedSignalId < CONFIG_SIGNAL_BATCH_SIZE; BatchedSignalId++)
 		{
-			FSSDSignalSample NeighborMin = SignalAccumulators.Array[BatchedSignalId].Min;
-			FSSDSignalSample NeighborMax = SignalAccumulators.Array[BatchedSignalId].Max;
-		
 			FSSDSignalSample NeighborMoment1 = CreateSignalSampleFromScalarValue(0.0);
 			FSSDSignalSample NeighborMoment2 = CreateSignalSampleFromScalarValue(0.0);
 			#if CONFIG_REJECTION_INPUT_MODE == REJECTION_INPUT_MODE_1UNNORMALIZED
 			{
 				float NormalizeFactor = SafeRcp(SignalAccumulators.Array[BatchedSignalId].Moment1.SampleCount);
 				NeighborMoment1 = Mul(SignalAccumulators.Array[BatchedSignalId].Moment1, NormalizeFactor);
-				NeighborMoment2 = Mul(SignalAccumulators.Array[BatchedSignalId].Moment2, NormalizeFactor);
+				
+				#if COMPILE_MOMENT2_ACCUMULATOR
+					NeighborMoment2 = Mul(SignalAccumulators.Array[BatchedSignalId].Moment2, NormalizeFactor);
+				#endif
 			}
 			#elif CONFIG_REJECTION_INPUT_MODE == REJECTION_INPUT_MODE_2PRETRANSFORMED_MOMMENTS
 			{
@@ -689,8 +707,8 @@ void TemporallyAccumulate(
 
 				// StdDev = sqrt(abs(NeighborSquareAvg - NeighborAvg * NeighborAvg));
 				FSSDSignalSample StdDev = Sqrt(Abs(Subtract(NeighborMoment2, Power(NeighborMoment1, 2))));
-				NeighborMin = Add(NeighborMoment1, Mul(StdDev, -StdDevMultiplier));
-				NeighborMax = Add(NeighborMoment1, Mul(StdDev,  StdDevMultiplier));
+				FSSDSignalSample NeighborMin = Add(NeighborMoment1, Mul(StdDev, -StdDevMultiplier));
+				FSSDSignalSample NeighborMax = Add(NeighborMoment1, Mul(StdDev,  StdDevMultiplier));
 
 				if (0)
 				{
@@ -756,6 +774,9 @@ void TemporallyAccumulate(
 			}
 			#elif CONFIG_HISTORY_REJECTION == HISTORY_REJECTION_MINMAX_BOUNDARIES
 			{
+				FSSDSignalSample NeighborMin = SignalAccumulators.Array[BatchedSignalId].Min;
+				FSSDSignalSample NeighborMax = SignalAccumulators.Array[BatchedSignalId].Max;
+		
 				float RejectedSampleCount = 0;
 				HistorySample = Mul(TransformSignalForPostRejection(Clamp(TransformInputBufferForPreRejection(NormalizeToOneSample(HistorySample)), NeighborMin, NeighborMax)), HistorySample.SampleCount - RejectedSampleCount);
 			}
