Shrinks denoiser's MAX_SIGNAL_BATCH_SIZE to the bare minimum

diff --git a/Engine/Shaders/Private/ScreenSpaceDenoise/SSDCommon.ush b/Engine/Shaders/Private/ScreenSpaceDenoise/SSDCommon.ush
--- a/Engine/Shaders/Private/ScreenSpaceDenoise/SSDCommon.ush
+++ b/Engine/Shaders/Private/ScreenSpaceDenoise/SSDCommon.ush
@@ -25,6 +25,15 @@
 
 //------------------------------------------------------- SHARED CONSTANTS DEPENDING ON CONFIG OF SHADER
 
+// Maximum number of signal that can be denoised at same time.
+#ifndef MAX_SIGNAL_BATCH_SIZE
+	#define MAX_SIGNAL_BATCH_SIZE 1
+#endif
+
+// Maximum number of signal that can be multiplexed into single output.
+#define MAX_SIGNAL_MULTIPLEXING (MAX_SIGNAL_BATCH_SIZE * 2)
+
+
 #if CONFIG_SIGNAL_PROCESSING == SIGNAL_PROCESSING_REFLECTIONS
 	#define TARGETED_SAMPLE_COUNT 128
 
diff --git a/Engine/Shaders/Private/ScreenSpaceDenoise/SSDDefinitions.ush b/Engine/Shaders/Private/ScreenSpaceDenoise/SSDDefinitions.ush
--- a/Engine/Shaders/Private/ScreenSpaceDenoise/SSDDefinitions.ush
+++ b/Engine/Shaders/Private/ScreenSpaceDenoise/SSDDefinitions.ush
@@ -59,12 +59,6 @@
 #define MIP_LEVEL_INVALID -1
 #define MIP_LEVEL_MISS 1024
 
-// Maximum number of signal that can be denoised at same time.
-#define MAX_SIGNAL_BATCH_SIZE 4
-
-// Maximum number of signal that can be multiplexed into single output.
-#define MAX_SIGNAL_MULTIPLEXING 8
-
 // Maximum number of buffer multiple signals can be encode to/decoded from.
 #define MAX_MULTIPLEXED_TEXTURES 4
 
diff --git a/Engine/Shaders/Private/ScreenSpaceDenoise/SSDInjest.usf b/Engine/Shaders/Private/ScreenSpaceDenoise/SSDInjest.usf
--- a/Engine/Shaders/Private/ScreenSpaceDenoise/SSDInjest.usf
+++ b/Engine/Shaders/Private/ScreenSpaceDenoise/SSDInjest.usf
@@ -24,6 +24,8 @@
 	#define CONFIG_INPUT_TEXTURE_COUNT DIM_SIGNAL_BATCH_SIZE
 	#define CONFIG_OUTPUT_TEXTURE_COUNT ((DIM_SIGNAL_BATCH_SIZE + 1) / 2)
 
+	#define MAX_SIGNAL_BATCH_SIZE CONFIG_SIGNAL_BATCH_SIZE
+
 #else
 	#error Unknown signal processing.
 #endif
diff --git a/Engine/Shaders/Private/ScreenSpaceDenoise/SSDSignalBufferEncoding.ush b/Engine/Shaders/Private/ScreenSpaceDenoise/SSDSignalBufferEncoding.ush
--- a/Engine/Shaders/Private/ScreenSpaceDenoise/SSDSignalBufferEncoding.ush
+++ b/Engine/Shaders/Private/ScreenSpaceDenoise/SSDSignalBufferEncoding.ush
@@ -202,11 +202,13 @@ void DecodeMultiplexedSignals(
 		OutSamples.Array[0].MissCount = OutSamples.Array[0].SampleCount * (1 - RawSample[0].a);
 		OutSamples.Array[0].SceneColor = OutSamples.Array[0].SampleCount * RawSample[0];
 
-		OutSamples.Array[1].SampleCount = OutSamples.Array[0].SampleCount;
-		OutSamples.Array[1].SceneColor.x = OutSamples.Array[1].SampleCount * RawSample[1].g;
-		
-		OutSamples.Array[2].SampleCount = OutSamples.Array[0].SampleCount;
-		OutSamples.Array[2].SceneColor.x = OutSamples.Array[2].SampleCount * RawSample[1].b;
+		// Temporal analysis of the history.
+		#if MAX_SIGNAL_MULTIPLEXING >= 3
+			OutSamples.Array[1].SampleCount = OutSamples.Array[0].SampleCount;
+			OutSamples.Array[1].SceneColor.x = OutSamples.Array[1].SampleCount * RawSample[1].g;
+			OutSamples.Array[2].SampleCount = OutSamples.Array[0].SampleCount;
+			OutSamples.Array[2].SceneColor.x = OutSamples.Array[2].SampleCount * RawSample[1].b;
+		#endif
 	}
 	#endif
 	else if (SignalBufferLayout == SIGNAL_BUFFER_LAYOUT_AO_INPUT)
@@ -286,11 +288,13 @@ void DecodeMultiplexedSignals(
 		OutSamples.Array[0].MissCount = OutSamples.Array[0].SampleCount * (1 - RawSample[0].a);
 		OutSamples.Array[0].SceneColor = OutSamples.Array[0].SampleCount * RawSample[0];
 		
-		OutSamples.Array[1].SampleCount = OutSamples.Array[0].SampleCount;
-		OutSamples.Array[1].SceneColor.x = OutSamples.Array[1].SampleCount * RawSample[1].g;
-		
-		OutSamples.Array[2].SampleCount = OutSamples.Array[0].SampleCount;
-		OutSamples.Array[2].SceneColor.x = OutSamples.Array[2].SampleCount * RawSample[1].b;
+		// Temporal analysis.
+		#if MAX_SIGNAL_MULTIPLEXING >= 3
+			OutSamples.Array[1].SampleCount = OutSamples.Array[0].SampleCount;
+			OutSamples.Array[1].SceneColor.x = OutSamples.Array[1].SampleCount * RawSample[1].g;
+			OutSamples.Array[2].SampleCount = OutSamples.Array[0].SampleCount;
+			OutSamples.Array[2].SceneColor.x = OutSamples.Array[2].SampleCount * RawSample[1].b;
+		#endif
 	}
 	#endif
 } // DecodeMultiplexedSignals()
@@ -456,12 +460,15 @@ void EncodeMultiplexedSignals(
 		OutRawSample[1].r = Sample[0].SampleCount;
 		OutBufferCount = 2;
 
+		// Temporal analysis.
+		#if MAX_SIGNAL_MULTIPLEXING >= 3
 		if (MultiplexCount == 3)
 		{
 			OutRawSample[1].g = Sample[1].SceneColor.x * NormalizationFactor;
 			OutRawSample[1].b = Sample[2].SceneColor.x * NormalizationFactor;
 			OutBufferCount = 3;
 		}
+		#endif
 	}
 	#endif
 	else if (SignalBufferLayout == SIGNAL_BUFFER_LAYOUT_AO_REJECTION)
@@ -499,11 +506,14 @@ void EncodeMultiplexedSignals(
 		OutRawSample[1].r = Sample[0].SampleCount;
 		OutBufferCount = 2;
 
+		// Temporal analysis.
+		#if MAX_SIGNAL_MULTIPLEXING >= 3
 		if (MultiplexCount == 3)
 		{
 			OutRawSample[1].g = Sample[1].SceneColor.x * NormalizationFactor;
 			OutRawSample[1].b = Sample[2].SceneColor.x * NormalizationFactor;
 		}
+		#endif
 	}
 	#endif
 } // EncodeMultiplexedSignals()
diff --git a/Engine/Shaders/Private/ScreenSpaceDenoise/SSDSpatialAccumulation.usf b/Engine/Shaders/Private/ScreenSpaceDenoise/SSDSpatialAccumulation.usf
--- a/Engine/Shaders/Private/ScreenSpaceDenoise/SSDSpatialAccumulation.usf
+++ b/Engine/Shaders/Private/ScreenSpaceDenoise/SSDSpatialAccumulation.usf
@@ -46,6 +46,8 @@
 	#define CONFIG_VIEW_POSITION_BASED_BILATERAL 5
 	#define CONFIG_MULTIPLEXED_SIGNALS_PER_SIGNAL_DOMAIN 1
 
+	#define MAX_SIGNAL_BATCH_SIZE CONFIG_SIGNAL_BATCH_SIZE
+
 	#if DIM_STAGE == STAGE_RECONSTRUCTION
 		// Input and output layout.
 		#if DIM_MULTI_SPP
@@ -297,7 +299,6 @@
 	// TODO: some shader compiler completly falls apparts with the current implementation of
 	//       CONFIG_SIGNAL_VGPR_COMPRESSION and actually drops in occupency.
 	#define CONFIG_SIGNAL_VGPR_COMPRESSION SIGNAL_COMPRESSION_DIFFUSE_INDIRECT_HARMONIC
-	#define CONFIG_NEIGHBOR_TO_REF_COMPUTATION NEIGHBOR_TO_REF_LOWEST_VGPR_PRESSURE
 
 	#if DIM_STAGE == STAGE_RECONSTRUCTION
 		//#define CONFIG_SAMPLE_SET           SAMPLE_SET_STOCASTIC_HIERARCHY
@@ -327,6 +328,10 @@
 #endif
 
 
+// No previous frame reprojection, save VGPR.
+//#define CONFIG_NEIGHBOR_TO_REF_COMPUTATION NEIGHBOR_TO_REF_LOWEST_VGPR_PRESSURE
+
+
 //------------------------------------------------------- CONFIG DISABLED DEFAULTS
 
 /** Changes the logic controling the number of sample to do. */
diff --git a/Engine/Shaders/Private/ScreenSpaceDenoise/SSDTemporalAccumulation.usf b/Engine/Shaders/Private/ScreenSpaceDenoise/SSDTemporalAccumulation.usf
--- a/Engine/Shaders/Private/ScreenSpaceDenoise/SSDTemporalAccumulation.usf
+++ b/Engine/Shaders/Private/ScreenSpaceDenoise/SSDTemporalAccumulation.usf
@@ -62,6 +62,8 @@
 		#define CONFIG_REJECTION_SAMPLE_SET               REJECTION_SAMPLE_SET_SOBEK2018
 	#endif
 
+	#define MAX_SIGNAL_BATCH_SIZE CONFIG_SIGNAL_BATCH_SIZE
+
 #elif CONFIG_SIGNAL_PROCESSING == SIGNAL_PROCESSING_REFLECTIONS
 	#define CONFIG_SIGNAL_COLOR 1
 
