Renames all signal operator to allow compiling out color in denoiser without conflicting with spherical harmonics.

diff --git a/Engine/Shaders/Private/ScreenSpaceDenoise/SSDInjest.usf b/Engine/Shaders/Private/ScreenSpaceDenoise/SSDInjest.usf
--- a/Engine/Shaders/Private/ScreenSpaceDenoise/SSDInjest.usf
+++ b/Engine/Shaders/Private/ScreenSpaceDenoise/SSDInjest.usf
@@ -9,7 +9,6 @@
 
 #define CONFIG_SIGNAL_PROCESSING DIM_SIGNAL_PROCESSING
 #define CONFIG_SIGNAL_BATCH_SIZE DIM_SIGNAL_BATCH_SIZE
-#define CONFIG_SIGNAL_COLOR 0
 
 #if CONFIG_SIGNAL_PROCESSING == SIGNAL_PROCESSING_MONOCHROMATIC_PENUMBRA
 	#if DIM_MULTI_SPP
diff --git a/Engine/Shaders/Private/ScreenSpaceDenoise/SSDReduce.usf b/Engine/Shaders/Private/ScreenSpaceDenoise/SSDReduce.usf
--- a/Engine/Shaders/Private/ScreenSpaceDenoise/SSDReduce.usf
+++ b/Engine/Shaders/Private/ScreenSpaceDenoise/SSDReduce.usf
@@ -15,8 +15,7 @@
 
 // Configures all the pass for each individual signals.
 #if CONFIG_SIGNAL_PROCESSING == SIGNAL_PROCESSING_DIFFUSE_SPHERICAL_HARMONIC
-	// #define CONFIG_SIGNAL_COLOR 0 // TODO
-	#define CONFIG_SIGNAL_COLOR_SH 1
+	#define COMPILE_SIGNAL_COLOR_SH 1
 
 	// Given it's a spherical harmonic that store directionality, only need position based rejection.
 	#define CONFIG_VIEW_POSITION_BASED_BILATERAL 2
diff --git a/Engine/Shaders/Private/ScreenSpaceDenoise/SSDSignalAccumulator.ush b/Engine/Shaders/Private/ScreenSpaceDenoise/SSDSignalAccumulator.ush
--- a/Engine/Shaders/Private/ScreenSpaceDenoise/SSDSignalAccumulator.ush
+++ b/Engine/Shaders/Private/ScreenSpaceDenoise/SSDSignalAccumulator.ush
@@ -176,11 +176,11 @@ void AccumulateSampleInOneBin(inout FSSDSignalAccumulator Accumulator, FSSDSampl
 	float SampleWeight = SampleInfos.FinalWeight;
 	
 	#if COMPILE_MOMENT1_ACCUMULATOR
-		Accumulator.Moment1 = Add(Accumulator.Moment1, Mul(SampleInfos.Sample, SampleWeight));
+		Accumulator.Moment1 = AddSignal(Accumulator.Moment1, MulSignal(SampleInfos.Sample, SampleWeight));
 	#endif
 	
 	#if COMPILE_MOMENT2_ACCUMULATOR
-		Accumulator.Moment2 = Add(Accumulator.Moment2, Mul(Power(SampleInfos.Sample, 2), SampleWeight));
+		Accumulator.Moment2 = AddSignal(Accumulator.Moment2, MulSignal(PowerSignal(SampleInfos.Sample, 2), SampleWeight));
 	#endif
 }
 
@@ -195,8 +195,8 @@ void AccumulateSampleInDomainBoundaries(inout FSSDSignalAccumulator Accumulator,
 	{
 		#if COMPILE_MINMAX_ACCUMULATOR
 		{
-			Accumulator.Min = Min(Accumulator.Min, SampleInfos.Sample);
-			Accumulator.Max = Max(Accumulator.Max, SampleInfos.Sample);
+			Accumulator.Min = MinSignal(Accumulator.Min, SampleInfos.Sample);
+			Accumulator.Max = MaxSignal(Accumulator.Max, SampleInfos.Sample);
 		}
 		#endif
 		
@@ -256,7 +256,7 @@ void AccumulateSampleInDRB(inout FSSDSignalAccumulator Accumulator, FSSDSampleAc
 	// Accumulate current bucket.
 	{
 		float CurrentWeight = A.FinalWeight * bBelongsToCurrent;
-		Accumulator.Current = Add(Accumulator.Current, Mul(A.Sample, CurrentWeight));
+		Accumulator.Current = AddSignal(Accumulator.Current, MulSignal(A.Sample, CurrentWeight));
 		Accumulator.CurrentInvFrequency += A.Sample.SampleCount * ClampedInvFrequency * CurrentWeight;
 	}
 
@@ -277,7 +277,7 @@ void AccumulateSampleInDRB(inout FSSDSignalAccumulator Accumulator, FSSDSampleAc
 	// Accumulate previous bucket.
 	{
 		float PreviousWeight = A.FinalWeight * bBelongsToPrevious;
-		Accumulator.Previous = Add(Accumulator.Previous, Mul(A.Sample, PreviousWeight));
+		Accumulator.Previous = AddSignal(Accumulator.Previous, MulSignal(A.Sample, PreviousWeight));
 		Accumulator.PreviousInvFrequency += A.Sample.SampleCount * ClampedInvFrequency * PreviousWeight;
 	}
 }
@@ -312,7 +312,7 @@ void DijestSampleClusterInDRB(
 	//if (PreviousBucketFactor < 1)
 	//	PreviousBucketFactor *= Accumulator.Current.SampleCount * SafeRcp(Accumulator.Previous.SampleCount);
 
-	Accumulator.Previous = Add(Mul(Accumulator.Previous, PreviousBucketFactor), Accumulator.Current);
+	Accumulator.Previous = AddSignal(MulSignal(Accumulator.Previous, PreviousBucketFactor), Accumulator.Current);
 	Accumulator.PreviousInvFrequency = Accumulator.PreviousInvFrequency * PreviousBucketFactor + Accumulator.CurrentInvFrequency;
 }
 
diff --git a/Engine/Shaders/Private/ScreenSpaceDenoise/SSDSignalArray.ush b/Engine/Shaders/Private/ScreenSpaceDenoise/SSDSignalArray.ush
--- a/Engine/Shaders/Private/ScreenSpaceDenoise/SSDSignalArray.ush
+++ b/Engine/Shaders/Private/ScreenSpaceDenoise/SSDSignalArray.ush
@@ -5,123 +5,123 @@
 #include "SSDSignalCore.ush"
 
 
-FSSDSignalArray Mul(FSSDSignalArray Samples, float Scalar[MAX_SIGNAL_BATCH_SIZE])
+FSSDSignalArray MulSignal(FSSDSignalArray Samples, float Scalar[MAX_SIGNAL_BATCH_SIZE])
 {
 	FSSDSignalArray OutSamples;
 	UNROLL_N(MAX_SIGNAL_BATCH_SIZE)
 	for (uint BatchedSignalId = 0; BatchedSignalId < MAX_SIGNAL_BATCH_SIZE; BatchedSignalId++)
 	{
-		OutSamples.Array[BatchedSignalId] = Mul(Samples.Array[BatchedSignalId], Scalar[BatchedSignalId]);
+		OutSamples.Array[BatchedSignalId] = MulSignal(Samples.Array[BatchedSignalId], Scalar[BatchedSignalId]);
 	}
 	return OutSamples;
 }
 
-FSSDSignalArray Add(FSSDSignalArray SamplesA, FSSDSignalArray SamplesB)
+FSSDSignalArray AddSignal(FSSDSignalArray SamplesA, FSSDSignalArray SamplesB)
 {
 	FSSDSignalArray OutSamples;
 	UNROLL_N(MAX_SIGNAL_BATCH_SIZE)
 	for (uint BatchedSignalId = 0; BatchedSignalId < MAX_SIGNAL_BATCH_SIZE; BatchedSignalId++)
 	{
-		OutSamples.Array[BatchedSignalId] = Add(SamplesA.Array[BatchedSignalId], SamplesB.Array[BatchedSignalId]);
+		OutSamples.Array[BatchedSignalId] = AddSignal(SamplesA.Array[BatchedSignalId], SamplesB.Array[BatchedSignalId]);
 	}
 	return OutSamples;
 }
 
-FSSDSignalArray Minus(FSSDSignalArray Samples)
+FSSDSignalArray MinusSignal(FSSDSignalArray Samples)
 {
 	FSSDSignalArray OutSamples;
 	UNROLL_N(MAX_SIGNAL_BATCH_SIZE)
 	for (uint BatchedSignalId = 0; BatchedSignalId < MAX_SIGNAL_BATCH_SIZE; BatchedSignalId++)
 	{
-		OutSamples.Array[BatchedSignalId] = Minus(Samples.Array[BatchedSignalId]);
+		OutSamples.Array[BatchedSignalId] = MinusSignal(Samples.Array[BatchedSignalId]);
 	}
 	return OutSamples;
 }
 
-FSSDSignalArray Subtract(FSSDSignalArray SamplesA, FSSDSignalArray SamplesB)
+FSSDSignalArray SubtractSignal(FSSDSignalArray SamplesA, FSSDSignalArray SamplesB)
 {
 	FSSDSignalArray OutSamples;
 	UNROLL_N(MAX_SIGNAL_BATCH_SIZE)
 	for (uint BatchedSignalId = 0; BatchedSignalId < MAX_SIGNAL_BATCH_SIZE; BatchedSignalId++)
 	{
-		OutSamples.Array[BatchedSignalId] = Subtract(SamplesA.Array[BatchedSignalId], SamplesB.Array[BatchedSignalId]);
+		OutSamples.Array[BatchedSignalId] = SubtractSignal(SamplesA.Array[BatchedSignalId], SamplesB.Array[BatchedSignalId]);
 	}
 	return OutSamples;
 }
 
-FSSDSignalArray Abs(FSSDSignalArray Samples)
+FSSDSignalArray AbsSignal(FSSDSignalArray Samples)
 {
 	FSSDSignalArray OutSamples;
 	UNROLL_N(MAX_SIGNAL_BATCH_SIZE)
 	for (uint BatchedSignalId = 0; BatchedSignalId < MAX_SIGNAL_BATCH_SIZE; BatchedSignalId++)
 	{
-		OutSamples.Array[BatchedSignalId] = Abs(Samples.Array[BatchedSignalId]);
+		OutSamples.Array[BatchedSignalId] = AbsSignal(Samples.Array[BatchedSignalId]);
 	}
 	return OutSamples;
 }
 
-FSSDSignalArray Sqrt(FSSDSignalArray Samples)
+FSSDSignalArray SqrtSignal(FSSDSignalArray Samples)
 {
 	FSSDSignalArray OutSamples;
 	UNROLL_N(MAX_SIGNAL_BATCH_SIZE)
 	for (uint BatchedSignalId = 0; BatchedSignalId < MAX_SIGNAL_BATCH_SIZE; BatchedSignalId++)
 	{
-		OutSamples.Array[BatchedSignalId] = Sqrt(Samples.Array[BatchedSignalId]);
+		OutSamples.Array[BatchedSignalId] = SqrtSignal(Samples.Array[BatchedSignalId]);
 	}
 	return OutSamples;
 }
 
-FSSDSignalArray Power(FSSDSignalArray Samples, float Exponent)
+FSSDSignalArray PowerSignal(FSSDSignalArray Samples, float Exponent)
 {
 	FSSDSignalArray OutSamples;
 	UNROLL_N(MAX_SIGNAL_BATCH_SIZE)
 	for (uint BatchedSignalId = 0; BatchedSignalId < MAX_SIGNAL_BATCH_SIZE; BatchedSignalId++)
 	{
-		OutSamples.Array[BatchedSignalId] = Power(Samples.Array[BatchedSignalId], Exponent);
+		OutSamples.Array[BatchedSignalId] = PowerSignal(Samples.Array[BatchedSignalId], Exponent);
 	}
 	return OutSamples;
 }
 
-FSSDSignalArray Min(FSSDSignalArray SamplesA, FSSDSignalArray SamplesB)
+FSSDSignalArray MinSignal(FSSDSignalArray SamplesA, FSSDSignalArray SamplesB)
 {
 	FSSDSignalArray OutSamples;
 	UNROLL_N(MAX_SIGNAL_BATCH_SIZE)
 	for (uint BatchedSignalId = 0; BatchedSignalId < MAX_SIGNAL_BATCH_SIZE; BatchedSignalId++)
 	{
-		OutSamples.Array[BatchedSignalId] = Min(SamplesA.Array[BatchedSignalId], SamplesB.Array[BatchedSignalId]);
+		OutSamples.Array[BatchedSignalId] = MinSignal(SamplesA.Array[BatchedSignalId], SamplesB.Array[BatchedSignalId]);
 	}
 	return OutSamples;
 }
 
-FSSDSignalArray Max(FSSDSignalArray SamplesA, FSSDSignalArray SamplesB)
+FSSDSignalArray MaxSignal(FSSDSignalArray SamplesA, FSSDSignalArray SamplesB)
 {
 	FSSDSignalArray OutSamples;
 	UNROLL_N(MAX_SIGNAL_BATCH_SIZE)
 	for (uint BatchedSignalId = 0; BatchedSignalId < MAX_SIGNAL_BATCH_SIZE; BatchedSignalId++)
 	{
-		OutSamples.Array[BatchedSignalId] = Max(SamplesA.Array[BatchedSignalId], SamplesB.Array[BatchedSignalId]);
+		OutSamples.Array[BatchedSignalId] = MaxSignal(SamplesA.Array[BatchedSignalId], SamplesB.Array[BatchedSignalId]);
 	}
 	return OutSamples;
 }
 
-FSSDSignalArray Clamp(FSSDSignalArray Samples, FSSDSignalArray SamplesMin, FSSDSignalArray SamplesMax)
+FSSDSignalArray ClampSignal(FSSDSignalArray Samples, FSSDSignalArray SamplesMin, FSSDSignalArray SamplesMax)
 {
 	FSSDSignalArray OutSamples;
 	UNROLL_N(MAX_SIGNAL_BATCH_SIZE)
 	for (uint BatchedSignalId = 0; BatchedSignalId < MAX_SIGNAL_BATCH_SIZE; BatchedSignalId++)
 	{
-		OutSamples.Array[BatchedSignalId] = Clamp(Samples.Array[BatchedSignalId], SamplesMin.Array[BatchedSignalId], SamplesMax.Array[BatchedSignalId]);
+		OutSamples.Array[BatchedSignalId] = ClampSignal(Samples.Array[BatchedSignalId], SamplesMin.Array[BatchedSignalId], SamplesMax.Array[BatchedSignalId]);
 	}
 	return OutSamples;
 }
 
-FSSDSignalArray Lerp(FSSDSignalArray Samples0, FSSDSignalArray Samples1, float Interp[MAX_SIGNAL_BATCH_SIZE])
+FSSDSignalArray LerpSignal(FSSDSignalArray Samples0, FSSDSignalArray Samples1, float Interp[MAX_SIGNAL_BATCH_SIZE])
 {
 	FSSDSignalArray OutSamples;
 	UNROLL_N(MAX_SIGNAL_BATCH_SIZE)
 	for (uint BatchedSignalId = 0; BatchedSignalId < MAX_SIGNAL_BATCH_SIZE; BatchedSignalId++)
 	{
-		OutSamples.Array[BatchedSignalId] = Lerp(Samples0.Array[BatchedSignalId], Samples1.Array[BatchedSignalId], Interp[BatchedSignalId]);
+		OutSamples.Array[BatchedSignalId] = LerpSignal(Samples0.Array[BatchedSignalId], Samples1.Array[BatchedSignalId], Interp[BatchedSignalId]);
 	}
 	return OutSamples;
 }
diff --git a/Engine/Shaders/Private/ScreenSpaceDenoise/SSDSignalBufferEncoding.ush b/Engine/Shaders/Private/ScreenSpaceDenoise/SSDSignalBufferEncoding.ush
--- a/Engine/Shaders/Private/ScreenSpaceDenoise/SSDSignalBufferEncoding.ush
+++ b/Engine/Shaders/Private/ScreenSpaceDenoise/SSDSignalBufferEncoding.ush
@@ -159,7 +159,7 @@ void DecodeMultiplexedSignals(
 			OutSamples.Array[MultiplexId + 1].MissCount = Channels.g * SampleCount;
 		}
 	}
-	#if CONFIG_SIGNAL_COLOR
+	#if COMPILE_SIGNAL_COLOR
 	else if (SignalBufferLayout == SIGNAL_BUFFER_LAYOUT_REFLECTIONS_INPUT)
 	{
 		OutSamples.Array[0].SampleCount = RawSample[1].r == DENOISER_INVALID_HIT_DISTANCE ? 0 : 1;
@@ -249,7 +249,7 @@ void DecodeMultiplexedSignals(
 			OutSamples.Array[0].MissCount = RawSample[0].g * RawSample[0].r;
 		}
 	}
-	#if CONFIG_SIGNAL_COLOR
+	#if COMPILE_SIGNAL_COLOR
 	else if (SignalBufferLayout == SIGNAL_BUFFER_LAYOUT_DIFFUSE_INDIRECT_AND_AO_INPUT_NSPP)
 	{
 		OutSamples.Array[0].SampleCount = RawSample[1].g;
@@ -302,7 +302,7 @@ void DecodeMultiplexedSignals(
 {
 	OutSamples = CreateSignalArrayFromScalarValue(0.0);
 
-	#if CONFIG_SIGNAL_COLOR_SH
+	#if COMPILE_SIGNAL_COLOR_SH
 	if (SignalBufferLayout == SIGNAL_BUFFER_LAYOUT_DIFFUSE_INDIRECT_HARMONIC)
 	{
 		const uint DecodeOptions = bNormalizeSample ? SSD_DECODE_NORMALIZE : 0x0;
@@ -436,7 +436,7 @@ void EncodeMultiplexedSignals(
 				NormalizedWorldBluringRadius == WORLD_RADIUS_MISS ? -1 : NormalizedWorldBluringRadius);
 		}
 	}
-	#if CONFIG_SIGNAL_COLOR
+	#if COMPILE_SIGNAL_COLOR
 	else if (SignalBufferLayout == SIGNAL_BUFFER_LAYOUT_REFLECTIONS_REJECTION)
 	{
 		// Outputs scene color to be compatible with a SSR output.
@@ -489,7 +489,7 @@ void EncodeMultiplexedSignals(
 		OutRawSample[0].r = Sample[0].SampleCount > 0 ? Sample[0].MissCount / Sample[0].SampleCount : 1.0;
 		OutRawSample[0].g = Sample[0].SampleCount;
 	}
-	#if CONFIG_SIGNAL_COLOR
+	#if COMPILE_SIGNAL_COLOR
 	else if (SignalBufferLayout == SIGNAL_BUFFER_LAYOUT_DIFFUSE_INDIRECT_AND_AO_RECONSTRUCTION)
 	{
 		float NormalizationFactor = SafeRcp(Sample[0].SampleCount);
@@ -527,7 +527,7 @@ void EncodeMultiplexedSignals(
 	// Number of buffer the signal get encoded onto <= MAX_MULTIPLEXED_TEXTURES.
 	OutBufferCount = 1;
 
-	#if CONFIG_SIGNAL_COLOR_SH
+	#if COMPILE_SIGNAL_COLOR_SH
 	if (SignalBufferLayout == SIGNAL_BUFFER_LAYOUT_DIFFUSE_INDIRECT_HARMONIC)
 	{
 		const uint EncodeOptions = (CONFIG_ENCODING_CLAMP_COLOR ? SSD_ENCODE_CLAMP_COLOR : 0) | SSD_ENCODE_NORMALIZE;
diff --git a/Engine/Shaders/Private/ScreenSpaceDenoise/SSDSignalCompression.ush b/Engine/Shaders/Private/ScreenSpaceDenoise/SSDSignalCompression.ush
--- a/Engine/Shaders/Private/ScreenSpaceDenoise/SSDSignalCompression.ush
+++ b/Engine/Shaders/Private/ScreenSpaceDenoise/SSDSignalCompression.ush
@@ -8,7 +8,7 @@
 
 //------------------------------------------------------- STRUCTURE
 
-#if CONFIG_SIGNAL_COLOR_SH
+#if COMPILE_SIGNAL_COLOR_SH
 	#define MAX_SIGNAL_COMPRESSED_VGPRS 7
 #else
 	#define MAX_SIGNAL_COMPRESSED_VGPRS 1
@@ -35,7 +35,7 @@ void CompressSignalSample(in FSSDSignalSample Sample, uint CompressionLayout, ou
 	{
 		// NOP
 	}
-	#if CONFIG_SIGNAL_COLOR_SH
+	#if COMPILE_SIGNAL_COLOR_SH
 	else if (CompressionLayout == SIGNAL_COMPRESSION_DIFFUSE_INDIRECT_HARMONIC)
 	{
 		OutCompressedSample.VGPR[0] = Compress2Floats(float2(Sample.SampleCount, Sample.MissCount));
@@ -60,7 +60,7 @@ void UncompressSignalSample(in FSSDCompressedSignalSample CompressedSample, uint
 	{
 		// NOP to not clobber information in OutSample that remained uncompressed.
 	}
-	#if CONFIG_SIGNAL_COLOR_SH
+	#if COMPILE_SIGNAL_COLOR_SH
 	else if (CompressionLayout == SIGNAL_COMPRESSION_DIFFUSE_INDIRECT_HARMONIC)
 	{
 		OutSample.SampleCount    = Uncompress2Floats(CompressedSample.VGPR[0]).x;
diff --git a/Engine/Shaders/Private/ScreenSpaceDenoise/SSDSignalCore.ush b/Engine/Shaders/Private/ScreenSpaceDenoise/SSDSignalCore.ush
--- a/Engine/Shaders/Private/ScreenSpaceDenoise/SSDSignalCore.ush
+++ b/Engine/Shaders/Private/ScreenSpaceDenoise/SSDSignalCore.ush
@@ -5,18 +5,13 @@
 #include "SSDCommon.ush"
 
 
-// Work arround FXC shader compiler that confuses FSSDSignalSample with a FSphericalGaussian when
-// CONFIG_SIGNAL_COLOR == 0 on the overloaded functions
-#undef CONFIG_SIGNAL_COLOR
-#define CONFIG_SIGNAL_COLOR 1
-
-
-#if !defined(CONFIG_SIGNAL_COLOR)
-	#error CONFIG_SIGNAL_COLOR has not been configured.
+#ifndef COMPILE_SIGNAL_COLOR
+	#define COMPILE_SIGNAL_COLOR 0
 #endif
 
-#ifndef CONFIG_SIGNAL_COLOR_SH
-	#define CONFIG_SIGNAL_COLOR_SH 0
+
+#ifndef COMPILE_SIGNAL_COLOR_SH
+	#define COMPILE_SIGNAL_COLOR_SH 0
 #endif
 
 
@@ -29,12 +24,12 @@ struct FSSDSignalSample
 	float SampleCount;
 	
 	// Scene color and alpha.
-	#if CONFIG_SIGNAL_COLOR
+	#if COMPILE_SIGNAL_COLOR
 		float4 SceneColor;
 	#endif
 
 	// Spherical harmonic of the color.
-	#if CONFIG_SIGNAL_COLOR_SH
+	#if COMPILE_SIGNAL_COLOR_SH
 		FSSDSphericalHarmonicRGB ColorSH;
 	#endif
 
@@ -61,14 +56,14 @@ struct FSSDSignalArray
 FSSDSignalSample CreateSignalSampleFromScalarValue(float Scalar)
 {
 	FSSDSignalSample OutSample;
-	#if CONFIG_SIGNAL_COLOR
+	#if COMPILE_SIGNAL_COLOR
 		OutSample.SceneColor = Scalar;
 	#endif
-	#if CONFIG_SIGNAL_COLOR_SH && SPHERICAL_HARMONIC_ORDER == 2
+	#if COMPILE_SIGNAL_COLOR_SH && SPHERICAL_HARMONIC_ORDER == 2
 		OutSample.ColorSH.R.V = Scalar;
 		OutSample.ColorSH.G.V = Scalar;
 		OutSample.ColorSH.B.V = Scalar;
-	#elif CONFIG_SIGNAL_COLOR_SH && SPHERICAL_HARMONIC_ORDER == 3
+	#elif COMPILE_SIGNAL_COLOR_SH && SPHERICAL_HARMONIC_ORDER == 3
 		OutSample.ColorSH.R.V0 = Scalar;
 		OutSample.ColorSH.R.V1 = Scalar;
 		OutSample.ColorSH.R.V2 = Scalar;
@@ -106,13 +101,13 @@ FSSDSignalArray CreateSignalArrayFromScalarValue(float Scalar)
 
 //------------------------------------------------------- SIGNAL OPERATORS
 
-FSSDSignalSample Mul(FSSDSignalSample Sample, float Scalar)
+FSSDSignalSample MulSignal(FSSDSignalSample Sample, float Scalar)
 {
 	FSSDSignalSample OutSample;
-	#if CONFIG_SIGNAL_COLOR
+	#if COMPILE_SIGNAL_COLOR
 		OutSample.SceneColor         = Sample.SceneColor * Scalar;
 	#endif
-	#if CONFIG_SIGNAL_COLOR_SH
+	#if COMPILE_SIGNAL_COLOR_SH
 		OutSample.ColorSH            = MulSH(Sample.ColorSH, Scalar);
 	#endif
 	OutSample.SampleCount            = Sample.SampleCount * Scalar;
@@ -123,12 +118,12 @@ FSSDSignalSample Mul(FSSDSignalSample Sample, float Scalar)
 	return OutSample;
 }
 
-FSSDSignalSample Add(FSSDSignalSample SampleA, FSSDSignalSample SampleB)
+FSSDSignalSample AddSignal(FSSDSignalSample SampleA, FSSDSignalSample SampleB)
 {
-	#if CONFIG_SIGNAL_COLOR
+	#if COMPILE_SIGNAL_COLOR
 		SampleA.SceneColor         += SampleB.SceneColor;
 	#endif
-	#if CONFIG_SIGNAL_COLOR_SH
+	#if COMPILE_SIGNAL_COLOR_SH
 		SampleA.ColorSH            = AddSH(SampleA.ColorSH, SampleB.ColorSH);
 	#endif
 	SampleA.SampleCount            += SampleB.SampleCount;
@@ -139,12 +134,12 @@ FSSDSignalSample Add(FSSDSignalSample SampleA, FSSDSignalSample SampleB)
 	return SampleA;
 }
 
-FSSDSignalSample Minus(FSSDSignalSample Sample)
+FSSDSignalSample MinusSignal(FSSDSignalSample Sample)
 {
-	#if CONFIG_SIGNAL_COLOR
+	#if COMPILE_SIGNAL_COLOR
 		Sample.SceneColor         = -Sample.SceneColor;
 	#endif
-	#if CONFIG_SIGNAL_COLOR_SH
+	#if COMPILE_SIGNAL_COLOR_SH
 		Sample.ColorSH         = MulSH(Sample.ColorSH, -1.0);
 	#endif
 	Sample.SampleCount            = -Sample.SampleCount;
@@ -155,21 +150,21 @@ FSSDSignalSample Minus(FSSDSignalSample Sample)
 	return Sample;
 }
 
-FSSDSignalSample Subtract(FSSDSignalSample SampleA, FSSDSignalSample SampleB)
+FSSDSignalSample SubtractSignal(FSSDSignalSample SampleA, FSSDSignalSample SampleB)
 {
-	return Add(SampleA, Minus(SampleB));
+	return AddSignal(SampleA, MinusSignal(SampleB));
 }
 
-FSSDSignalSample Abs(FSSDSignalSample Sample)
+FSSDSignalSample AbsSignal(FSSDSignalSample Sample)
 {
-	#if CONFIG_SIGNAL_COLOR
+	#if COMPILE_SIGNAL_COLOR
 		Sample.SceneColor         = abs(Sample.SceneColor);
 	#endif
-	#if CONFIG_SIGNAL_COLOR_SH && SPHERICAL_HARMONIC_ORDER == 2
+	#if COMPILE_SIGNAL_COLOR_SH && SPHERICAL_HARMONIC_ORDER == 2
 		Sample.ColorSH.R.V = abs(Sample.ColorSH.R.V);
 		Sample.ColorSH.G.V = abs(Sample.ColorSH.G.V);
 		Sample.ColorSH.B.V = abs(Sample.ColorSH.B.V);
-	#elif CONFIG_SIGNAL_COLOR_SH && SPHERICAL_HARMONIC_ORDER == 3
+	#elif COMPILE_SIGNAL_COLOR_SH && SPHERICAL_HARMONIC_ORDER == 3
 		Sample.ColorSH.R.V0 = abs(Sample.ColorSH.R.V0);
 		Sample.ColorSH.R.V1 = abs(Sample.ColorSH.R.V1);
 		Sample.ColorSH.R.V2 = abs(Sample.ColorSH.R.V2);
@@ -188,16 +183,16 @@ FSSDSignalSample Abs(FSSDSignalSample Sample)
 	return Sample;
 }
 
-FSSDSignalSample Sqrt(FSSDSignalSample Sample)
+FSSDSignalSample SqrtSignal(FSSDSignalSample Sample)
 {
-	#if CONFIG_SIGNAL_COLOR
+	#if COMPILE_SIGNAL_COLOR
 		Sample.SceneColor         = sqrt(Sample.SceneColor);
 	#endif
-	#if CONFIG_SIGNAL_COLOR_SH && SPHERICAL_HARMONIC_ORDER == 2
+	#if COMPILE_SIGNAL_COLOR_SH && SPHERICAL_HARMONIC_ORDER == 2
 		Sample.ColorSH.R.V = sqrt(Sample.ColorSH.R.V);
 		Sample.ColorSH.G.V = sqrt(Sample.ColorSH.G.V);
 		Sample.ColorSH.B.V = sqrt(Sample.ColorSH.B.V);
-	#elif CONFIG_SIGNAL_COLOR_SH && SPHERICAL_HARMONIC_ORDER == 3
+	#elif COMPILE_SIGNAL_COLOR_SH && SPHERICAL_HARMONIC_ORDER == 3
 		Sample.ColorSH.R.V0 = sqrt(Sample.ColorSH.R.V0);
 		Sample.ColorSH.R.V1 = sqrt(Sample.ColorSH.R.V1);
 		Sample.ColorSH.R.V2 = sqrt(Sample.ColorSH.R.V2);
@@ -216,19 +211,19 @@ FSSDSignalSample Sqrt(FSSDSignalSample Sample)
 	return Sample;
 }
 
-FSSDSignalSample Power(FSSDSignalSample Sample, float Exponent)
+FSSDSignalSample PowerSignal(FSSDSignalSample Sample, float Exponent)
 {
-	#if CONFIG_SIGNAL_COLOR
+	#if COMPILE_SIGNAL_COLOR
 		Sample.SceneColor.r       = pow(Sample.SceneColor.r, Exponent);
 		Sample.SceneColor.g       = pow(Sample.SceneColor.g, Exponent);
 		Sample.SceneColor.b       = pow(Sample.SceneColor.b, Exponent);
 		Sample.SceneColor.a       = pow(Sample.SceneColor.a, Exponent);
 	#endif
-	#if CONFIG_SIGNAL_COLOR_SH && SPHERICAL_HARMONIC_ORDER == 2
+	#if COMPILE_SIGNAL_COLOR_SH && SPHERICAL_HARMONIC_ORDER == 2
 		Sample.ColorSH.R.V = pow(Sample.ColorSH.R.V, Exponent);
 		Sample.ColorSH.G.V = pow(Sample.ColorSH.G.V, Exponent);
 		Sample.ColorSH.B.V = pow(Sample.ColorSH.B.V, Exponent);
-	#elif CONFIG_SIGNAL_COLOR_SH && SPHERICAL_HARMONIC_ORDER == 3
+	#elif COMPILE_SIGNAL_COLOR_SH && SPHERICAL_HARMONIC_ORDER == 3
 		Sample.ColorSH.R.V0 = pow(Sample.ColorSH.R.V0, Exponent);
 		Sample.ColorSH.R.V1 = pow(Sample.ColorSH.R.V1, Exponent);
 		Sample.ColorSH.R.V2 = pow(Sample.ColorSH.R.V2, Exponent);
@@ -247,17 +242,17 @@ FSSDSignalSample Power(FSSDSignalSample Sample, float Exponent)
 	return Sample;
 }
 
-FSSDSignalSample Min(FSSDSignalSample SampleA, FSSDSignalSample SampleB)
+FSSDSignalSample MinSignal(FSSDSignalSample SampleA, FSSDSignalSample SampleB)
 {
 	FSSDSignalSample OutSample;
-	#if CONFIG_SIGNAL_COLOR
+	#if COMPILE_SIGNAL_COLOR
 		OutSample.SceneColor = min(SampleA.SceneColor, SampleB.SceneColor);
 	#endif
-	#if CONFIG_SIGNAL_COLOR_SH && SPHERICAL_HARMONIC_ORDER == 2
+	#if COMPILE_SIGNAL_COLOR_SH && SPHERICAL_HARMONIC_ORDER == 2
 		OutSample.ColorSH.R.V = min(SampleA.ColorSH.R.V, SampleB.ColorSH.R.V);
 		OutSample.ColorSH.G.V = min(SampleA.ColorSH.G.V, SampleB.ColorSH.G.V);
 		OutSample.ColorSH.B.V = min(SampleA.ColorSH.B.V, SampleB.ColorSH.B.V);
-	#elif CONFIG_SIGNAL_COLOR_SH && SPHERICAL_HARMONIC_ORDER == 3
+	#elif COMPILE_SIGNAL_COLOR_SH && SPHERICAL_HARMONIC_ORDER == 3
 		OutSample.ColorSH.R.V0 = min(SampleA.ColorSH.R.V0, SampleB.ColorSH.R.V0);
 		OutSample.ColorSH.R.V1 = min(SampleA.ColorSH.R.V1, SampleB.ColorSH.R.V1);
 		OutSample.ColorSH.R.V2 = min(SampleA.ColorSH.R.V2, SampleB.ColorSH.R.V2);
@@ -276,17 +271,17 @@ FSSDSignalSample Min(FSSDSignalSample SampleA, FSSDSignalSample SampleB)
 	return OutSample;
 }
 
-FSSDSignalSample Max(FSSDSignalSample SampleA, FSSDSignalSample SampleB)
+FSSDSignalSample MaxSignal(FSSDSignalSample SampleA, FSSDSignalSample SampleB)
 {
 	FSSDSignalSample OutSample;
-	#if CONFIG_SIGNAL_COLOR
+	#if COMPILE_SIGNAL_COLOR
 		OutSample.SceneColor = max(SampleA.SceneColor, SampleB.SceneColor);
 	#endif
-	#if CONFIG_SIGNAL_COLOR_SH && SPHERICAL_HARMONIC_ORDER == 2
+	#if COMPILE_SIGNAL_COLOR_SH && SPHERICAL_HARMONIC_ORDER == 2
 		OutSample.ColorSH.R.V = max(SampleA.ColorSH.R.V, SampleB.ColorSH.R.V);
 		OutSample.ColorSH.G.V = max(SampleA.ColorSH.G.V, SampleB.ColorSH.G.V);
 		OutSample.ColorSH.B.V = max(SampleA.ColorSH.B.V, SampleB.ColorSH.B.V);
-	#elif CONFIG_SIGNAL_COLOR_SH && SPHERICAL_HARMONIC_ORDER == 3
+	#elif COMPILE_SIGNAL_COLOR_SH && SPHERICAL_HARMONIC_ORDER == 3
 		OutSample.ColorSH.R.V0 = max(SampleA.ColorSH.R.V0, SampleB.ColorSH.R.V0);
 		OutSample.ColorSH.R.V1 = max(SampleA.ColorSH.R.V1, SampleB.ColorSH.R.V1);
 		OutSample.ColorSH.R.V2 = max(SampleA.ColorSH.R.V2, SampleB.ColorSH.R.V2);
@@ -305,17 +300,17 @@ FSSDSignalSample Max(FSSDSignalSample SampleA, FSSDSignalSample SampleB)
 	return OutSample;
 }
 
-FSSDSignalSample Clamp(FSSDSignalSample Sample, FSSDSignalSample SampleMin, FSSDSignalSample SampleMax)
+FSSDSignalSample ClampSignal(FSSDSignalSample Sample, FSSDSignalSample SampleMin, FSSDSignalSample SampleMax)
 {
 	FSSDSignalSample OutSample;
-	#if CONFIG_SIGNAL_COLOR
+	#if COMPILE_SIGNAL_COLOR
 		OutSample.SceneColor = clamp(Sample.SceneColor, SampleMin.SceneColor, SampleMax.SceneColor);
 	#endif
-	#if CONFIG_SIGNAL_COLOR_SH && SPHERICAL_HARMONIC_ORDER == 2
+	#if COMPILE_SIGNAL_COLOR_SH && SPHERICAL_HARMONIC_ORDER == 2
 		OutSample.ColorSH.R.V = clamp(Sample.ColorSH.R.V, SampleMin.ColorSH.R.V, SampleMax.ColorSH.R.V);
 		OutSample.ColorSH.G.V = clamp(Sample.ColorSH.G.V, SampleMin.ColorSH.G.V, SampleMax.ColorSH.G.V);
 		OutSample.ColorSH.B.V = clamp(Sample.ColorSH.B.V, SampleMin.ColorSH.B.V, SampleMax.ColorSH.B.V);
-	#elif CONFIG_SIGNAL_COLOR_SH && SPHERICAL_HARMONIC_ORDER == 3
+	#elif COMPILE_SIGNAL_COLOR_SH && SPHERICAL_HARMONIC_ORDER == 3
 		OutSample.ColorSH.R.V0 = clamp(Sample.ColorSH.R.V0, SampleMin.ColorSH.R.V0, SampleMax.ColorSH.R.V0);
 		OutSample.ColorSH.R.V1 = clamp(Sample.ColorSH.R.V1, SampleMin.ColorSH.R.V1, SampleMax.ColorSH.R.V1);
 		OutSample.ColorSH.R.V2 = clamp(Sample.ColorSH.R.V2, SampleMin.ColorSH.R.V2, SampleMax.ColorSH.R.V2);
@@ -334,15 +329,15 @@ FSSDSignalSample Clamp(FSSDSignalSample Sample, FSSDSignalSample SampleMin, FSSD
 	return OutSample;
 }
 
-FSSDSignalSample Lerp(FSSDSignalSample Sample0, FSSDSignalSample Sample1, float Interp)
+FSSDSignalSample LerpSignal(FSSDSignalSample Sample0, FSSDSignalSample Sample1, float Interp)
 {
-	return Add(Sample0, Mul(Subtract(Sample1, Sample0), Interp));
+	return AddSignal(Sample0, MulSignal(SubtractSignal(Sample1, Sample0), Interp));
 }
 
 /** Normalize the signal sample to one. */
 FSSDSignalSample NormalizeToOneSample(FSSDSignalSample Sample)
 {
-	FSSDSignalSample OutSample = Mul(Sample, Sample.SampleCount > 0 ? rcp(Sample.SampleCount) : 0);
+	FSSDSignalSample OutSample = MulSignal(Sample, Sample.SampleCount > 0 ? rcp(Sample.SampleCount) : 0);
 	OutSample.SampleCount = Sample.SampleCount > 0 ? 1 : 0;
 	return OutSample;
 }
diff --git a/Engine/Shaders/Private/ScreenSpaceDenoise/SSDSignalFramework.ush b/Engine/Shaders/Private/ScreenSpaceDenoise/SSDSignalFramework.ush
--- a/Engine/Shaders/Private/ScreenSpaceDenoise/SSDSignalFramework.ush
+++ b/Engine/Shaders/Private/ScreenSpaceDenoise/SSDSignalFramework.ush
@@ -44,7 +44,7 @@
 #include "../EyeAdaptationCommon.ush"
 #include "../LightShaderParameters.ush"
 
-#if CONFIG_SIGNAL_COLOR_SH
+#if COMPILE_SIGNAL_COLOR_SH
 	#include "../QuasiRandom.ush"
 #endif
 
@@ -166,7 +166,7 @@ float GetSignalColorTransformationExposure()
 
 /** Returns the linear color space Luma with a 4 multiplier for ALU perf reasons. */
 float GetSignalLuma4(FSSDSignalSample Sample, const uint Basis)
-#if CONFIG_SIGNAL_COLOR
+#if COMPILE_SIGNAL_COLOR
 {
 	// TODO: should force this to be scalar load.
 	float FrameExposureScale = GetSignalColorTransformationExposure();
@@ -207,7 +207,7 @@ float GetSignalLuma4(FSSDSignalSample Sample, const uint Basis)
 
 /** Conveniently transform one simple from a source basis to a destination basis. */
 FSSDSignalSample TransformSignal(FSSDSignalSample Sample, const uint SrcBasis, const uint DestBasis)
-#if CONFIG_SIGNAL_COLOR
+#if COMPILE_SIGNAL_COLOR
 {
 	// If the source and destination bases are the same, early return to not pay ALU and floating point divergence.
 	if (SrcBasis == DestBasis)
@@ -779,7 +779,7 @@ float GetRatioEstimatorWeight(
 
 // Matches GenerateCosineNormalRay()
 // TODO: expose a Weight that multiply directly onto the Sample.SceneColor to reduce ALU
-#if CONFIG_SIGNAL_COLOR_SH
+#if COMPILE_SIGNAL_COLOR_SH && COMPILE_SIGNAL_COLOR
 FSSDSphericalHarmonicRGB ComputeSampleColorSH(FSSDSampleSceneInfos SampleSceneMetadata, FSSDSignalSample Sample, uint2 NeighborPixelCoord)
 {
 	const uint SamplesPerPixel = 1;
diff --git a/Engine/Shaders/Private/ScreenSpaceDenoise/SSDSpatialAccumulation.usf b/Engine/Shaders/Private/ScreenSpaceDenoise/SSDSpatialAccumulation.usf
--- a/Engine/Shaders/Private/ScreenSpaceDenoise/SSDSpatialAccumulation.usf
+++ b/Engine/Shaders/Private/ScreenSpaceDenoise/SSDSpatialAccumulation.usf
@@ -40,7 +40,7 @@
 
 // Configures all the pass for each individual signals.
 #if CONFIG_SIGNAL_PROCESSING == SIGNAL_PROCESSING_MONOCHROMATIC_PENUMBRA
-	#define CONFIG_SIGNAL_COLOR 0
+	#define COMPILE_SIGNAL_COLOR 0
 
 	//#define CONFIG_VIEW_POSITION_BASED_BILATERAL 2
 	#define CONFIG_VIEW_POSITION_BASED_BILATERAL 5
@@ -122,7 +122,7 @@
 	#endif
 
 #elif CONFIG_SIGNAL_PROCESSING == SIGNAL_PROCESSING_REFLECTIONS
-	#define CONFIG_SIGNAL_COLOR 1
+	#define COMPILE_SIGNAL_COLOR 1
 
 	#if DIM_STAGE == STAGE_RECONSTRUCTION
 		// Input and output layout.
@@ -193,8 +193,6 @@
 	#endif
 
 #elif CONFIG_SIGNAL_PROCESSING == SIGNAL_PROCESSING_AO
-	#define CONFIG_SIGNAL_COLOR 0
-
 	#define CONFIG_VIEW_POSITION_BASED_BILATERAL 2
 	#define CONFIG_NORMAL_BASED_BILATERAL 1
 
@@ -238,7 +236,7 @@
 	#endif
 
 #elif CONFIG_SIGNAL_PROCESSING == SIGNAL_PROCESSING_DIFFUSE_INDIRECT_AND_AO
-	#define CONFIG_SIGNAL_COLOR 1
+	#define COMPILE_SIGNAL_COLOR 1
 
 	#define CONFIG_VIEW_POSITION_BASED_BILATERAL 2
 	#define CONFIG_NORMAL_BASED_BILATERAL 1
@@ -279,8 +277,7 @@
 	#endif
 
 #elif CONFIG_SIGNAL_PROCESSING == SIGNAL_PROCESSING_DIFFUSE_SPHERICAL_HARMONIC
-	// #define CONFIG_SIGNAL_COLOR 0 // TODO
-	#define CONFIG_SIGNAL_COLOR_SH 1
+	#define COMPILE_SIGNAL_COLOR_SH 1
 
 	// Given it's a spherical harmonic that store directionality, only need position based rejection.
 	#define CONFIG_VIEW_POSITION_BASED_BILATERAL 2
@@ -866,7 +863,7 @@ void MainCS(
 
 			#if CONFIG_REJECT_HIGHEST_COLOR
 			{
-				#if !CONFIG_SIGNAL_COLOR
+				#if !COMPILE_SIGNAL_COLOR
 					#error Need to compile signal color.
 				#endif
 				if (Accumulator.Moment1[MultiplexId].SampleCount > 0)
@@ -976,7 +973,7 @@ void MainCS(
 			float CurrentSampleCount = OutputSamples.Array[MultiplexId].SampleCount;
 			float NewSampleCount = min(CurrentSampleCount, TARGETED_SAMPLE_COUNT);
 
-			OutputSamples.Array[MultiplexId] = Mul(OutputSamples.Array[MultiplexId], CurrentSampleCount > 0 ? NewSampleCount / CurrentSampleCount : 0);
+			OutputSamples.Array[MultiplexId] = MulSignal(OutputSamples.Array[MultiplexId], CurrentSampleCount > 0 ? NewSampleCount / CurrentSampleCount : 0);
 		}
 	}
 	#endif // DIM_STAGE == STAGE_POST_FILTERING
diff --git a/Engine/Shaders/Private/ScreenSpaceDenoise/SSDSpatialKernel.ush b/Engine/Shaders/Private/ScreenSpaceDenoise/SSDSpatialKernel.ush
--- a/Engine/Shaders/Private/ScreenSpaceDenoise/SSDSpatialKernel.ush
+++ b/Engine/Shaders/Private/ScreenSpaceDenoise/SSDSpatialKernel.ush
@@ -314,7 +314,7 @@ FSSDSignalSample TransformSignalSampleForAccumulation(
 		/* DestBasis = */ KernelConfig.AccumulatorColorSpace[MultiplexId]);
 	
 	// Compute the spherical harmonic of the sample.
-	#if CONFIG_SIGNAL_COLOR_SH
+	#if COMPILE_SIGNAL_COLOR_SH && COMPILE_SIGNAL_COLOR
 	if (KernelConfig.bComputeSampleColorSH)
 	{
 		Sample.ColorSH = ComputeSampleColorSH(SampleSceneMetadata, Sample, SamplePixelCoord);
diff --git a/Engine/Shaders/Private/ScreenSpaceDenoise/SSDTemporalAccumulation.usf b/Engine/Shaders/Private/ScreenSpaceDenoise/SSDTemporalAccumulation.usf
--- a/Engine/Shaders/Private/ScreenSpaceDenoise/SSDTemporalAccumulation.usf
+++ b/Engine/Shaders/Private/ScreenSpaceDenoise/SSDTemporalAccumulation.usf
@@ -33,8 +33,6 @@
 
 
 #if CONFIG_SIGNAL_PROCESSING == SIGNAL_PROCESSING_MONOCHROMATIC_PENUMBRA
-	#define CONFIG_SIGNAL_COLOR 0
-
 	#define CONFIG_VIEW_POSITION_BASED_BILATERAL 2
 	#define CONFIG_MULTIPLEXED_SIGNALS_PER_SIGNAL_DOMAIN 1
 
@@ -65,7 +63,7 @@
 	#define MAX_SIGNAL_BATCH_SIZE CONFIG_SIGNAL_BATCH_SIZE
 
 #elif CONFIG_SIGNAL_PROCESSING == SIGNAL_PROCESSING_REFLECTIONS
-	#define CONFIG_SIGNAL_COLOR 1
+	#define COMPILE_SIGNAL_COLOR 1
 
 	// Input and output layout.
 	#define CONFIG_SIGNAL_INPUT_LAYOUT   SIGNAL_BUFFER_LAYOUT_REFLECTIONS_HISTORY
@@ -96,8 +94,6 @@
 	#endif
 
 #elif CONFIG_SIGNAL_PROCESSING == SIGNAL_PROCESSING_AO
-	#define CONFIG_SIGNAL_COLOR 0
-
 	// Input and output layout.
 	#define CONFIG_SIGNAL_INPUT_LAYOUT   SIGNAL_BUFFER_LAYOUT_AO_HISTORY
 	#define CONFIG_SIGNAL_HISTORY_LAYOUT SIGNAL_BUFFER_LAYOUT_AO_HISTORY
@@ -121,7 +117,7 @@
 	#endif
 
 #elif CONFIG_SIGNAL_PROCESSING == SIGNAL_PROCESSING_DIFFUSE_INDIRECT_AND_AO
-	#define CONFIG_SIGNAL_COLOR 1
+	#define COMPILE_SIGNAL_COLOR 1
 
 	// Input and output layout.
 	#define CONFIG_SIGNAL_INPUT_LAYOUT   SIGNAL_BUFFER_LAYOUT_DIFFUSE_INDIRECT_AND_AO_RECONSTRUCTION
@@ -138,8 +134,7 @@
 	#define CONFIG_NORMAL_BASED_BILATERAL 1
 
 #elif CONFIG_SIGNAL_PROCESSING == SIGNAL_PROCESSING_DIFFUSE_SPHERICAL_HARMONIC
-	// #define CONFIG_SIGNAL_COLOR 0 // TODO
-	#define CONFIG_SIGNAL_COLOR_SH 1
+	#define COMPILE_SIGNAL_COLOR_SH 1
 
 	// Input and output layout.
 	#define CONFIG_SIGNAL_INPUT_LAYOUT   SIGNAL_BUFFER_LAYOUT_DIFFUSE_INDIRECT_HARMONIC
@@ -508,7 +503,7 @@ void TemporallyAccumulate(
 					SafeRcp(CurrentFrameSample.SampleCount));
 
 				// Applies correction when pre exposure change happen.
-				#if CONFIG_SIGNAL_COLOR
+				#if COMPILE_SIGNAL_COLOR
 					HistorySamples.Array[BatchedSignalId].SceneColor.rgb *= HistoryPreExposureCorrection;
 				#endif
 			} // for (uint BatchedSignalId = 0; BatchedSignalId < CONFIG_SIGNAL_BATCH_SIZE; BatchedSignalId++)
@@ -676,17 +671,17 @@ void TemporallyAccumulate(
 			#if CONFIG_REJECTION_INPUT_MODE == REJECTION_INPUT_MODE_1UNNORMALIZED
 			{
 				float NormalizeFactor = SafeRcp(SignalAccumulators.Array[BatchedSignalId].Moment1.SampleCount);
-				NeighborMoment1 = Mul(SignalAccumulators.Array[BatchedSignalId].Moment1, NormalizeFactor);
+				NeighborMoment1 = MulSignal(SignalAccumulators.Array[BatchedSignalId].Moment1, NormalizeFactor);
 				
 				#if COMPILE_MOMENT2_ACCUMULATOR
-					NeighborMoment2 = Mul(SignalAccumulators.Array[BatchedSignalId].Moment2, NormalizeFactor);
+					NeighborMoment2 = MulSignal(SignalAccumulators.Array[BatchedSignalId].Moment2, NormalizeFactor);
 				#endif
 			}
 			#elif CONFIG_REJECTION_INPUT_MODE == REJECTION_INPUT_MODE_2PRETRANSFORMED_MOMMENTS
 			{
 				float NormalizeFactor = SafeRcp(SignalAccumulators.Array[BatchedSignalId * 2 + 0].Moment1.SampleCount);
-				NeighborMoment1 = Mul(SignalAccumulators.Array[BatchedSignalId * 2 + 0].Moment1, NormalizeFactor);
-				NeighborMoment2 = Mul(SignalAccumulators.Array[BatchedSignalId * 2 + 1].Moment1, NormalizeFactor);
+				NeighborMoment1 = MulSignal(SignalAccumulators.Array[BatchedSignalId * 2 + 0].Moment1, NormalizeFactor);
+				NeighborMoment2 = MulSignal(SignalAccumulators.Array[BatchedSignalId * 2 + 1].Moment1, NormalizeFactor);
 			}
 			#else
 				#error NOrmalized samples. 
@@ -708,17 +703,17 @@ void TemporallyAccumulate(
 				const float StdDevMultiplier = 1.25;
 
 				// StdDev = sqrt(abs(NeighborSquareAvg - NeighborAvg * NeighborAvg));
-				FSSDSignalSample StdDev = Sqrt(Abs(Subtract(NeighborMoment2, Power(NeighborMoment1, 2))));
-				FSSDSignalSample NeighborMin = Add(NeighborMoment1, Mul(StdDev, -StdDevMultiplier));
-				FSSDSignalSample NeighborMax = Add(NeighborMoment1, Mul(StdDev,  StdDevMultiplier));
+				FSSDSignalSample StdDev = SqrtSignal(AbsSignal(SubtractSignal(NeighborMoment2, PowerSignal(NeighborMoment1, 2))));
+				FSSDSignalSample NeighborMin = AddSignal(NeighborMoment1, MulSignal(StdDev, -StdDevMultiplier));
+				FSSDSignalSample NeighborMax = AddSignal(NeighborMoment1, MulSignal(StdDev,  StdDevMultiplier));
 
 				if (0)
 				{
-					FSSDSignalSample QuantizationErrorMin = Mul(NeighborMoment1, 1 - SafeRcp(HistorySample.SampleCount));
-					FSSDSignalSample QuantizationErrorMax = Mul(NeighborMoment1, 1 + SafeRcp(HistorySample.SampleCount));
+					FSSDSignalSample QuantizationErrorMin = MulSignal(NeighborMoment1, 1 - SafeRcp(HistorySample.SampleCount));
+					FSSDSignalSample QuantizationErrorMax = MulSignal(NeighborMoment1, 1 + SafeRcp(HistorySample.SampleCount));
 
-					NeighborMin = Min(NeighborMin, QuantizationErrorMin);
-					NeighborMax = Max(NeighborMax, QuantizationErrorMax);
+					NeighborMin = MinSignal(NeighborMin, QuantizationErrorMin);
+					NeighborMax = MaxSignal(NeighborMax, QuantizationErrorMax);
 				}
 
 				// Transform the history so that it in correct component spaces, and normalized as the clamping box.
@@ -726,20 +721,20 @@ void TemporallyAccumulate(
 				FSSDSignalSample TransformedHistorySample = TransformInputBufferForPreRejection(NormalizedHistorySample);
 
 				// Clamp the history.
-				FSSDSignalSample ClampedTransformedHistorySample = Clamp(TransformedHistorySample, NeighborMin, NeighborMax);
+				FSSDSignalSample ClampedTransformedHistorySample = ClampSignal(TransformedHistorySample, NeighborMin, NeighborMax);
 
 				// Transform the history back to linear component spaces.
 				FSSDSignalSample ClampedHistorySample = TransformSignalForPostRejection(ClampedTransformedHistorySample);
 
 				// Reweight the anti ghosted history.
 				{
-					FSSDSignalSample RejectedDiff = Abs(Subtract(ClampedTransformedHistorySample, TransformedHistorySample));
+					FSSDSignalSample RejectedDiff = AbsSignal(SubtractSignal(ClampedTransformedHistorySample, TransformedHistorySample));
 
 					// Compute how much this history get changed.
 					float RejectionFactor = 0.0;
 					#if CONFIG_SIGNAL_PROCESSING == SIGNAL_PROCESSING_REFLECTIONS && (CONFIG_HISTORY_REJECTION_COLOR_SPACE & COLOR_SPACE_LCOCG)
 					{
-						#if !CONFIG_SIGNAL_COLOR
+						#if !COMPILE_SIGNAL_COLOR
 							#error Need to compile signal color.
 						#endif
 						RejectionFactor = abs(
@@ -770,7 +765,7 @@ void TemporallyAccumulate(
 					FinalHistoryWeight = max(FinalHistoryWeight, NeighborMoment1.SampleCount * 0.1);
 			
 					// Apply the history weight.
-					HistorySample = Mul(ClampedHistorySample, FinalHistoryWeight);
+					HistorySample = MulSignal(ClampedHistorySample, FinalHistoryWeight);
 					HistorySample.SampleCount = FinalHistoryWeight;
 				}
 			}
@@ -780,7 +775,7 @@ void TemporallyAccumulate(
 				FSSDSignalSample NeighborMax = SignalAccumulators.Array[BatchedSignalId].Max;
 		
 				float RejectedSampleCount = 0;
-				HistorySample = Mul(TransformSignalForPostRejection(Clamp(TransformInputBufferForPreRejection(NormalizeToOneSample(HistorySample)), NeighborMin, NeighborMax)), HistorySample.SampleCount - RejectedSampleCount);
+				HistorySample = MulSignal(TransformSignalForPostRejection(ClampSignal(TransformInputBufferForPreRejection(NormalizeToOneSample(HistorySample)), NeighborMin, NeighborMax)), HistorySample.SampleCount - RejectedSampleCount);
 			}
 			#endif
 
@@ -812,7 +807,7 @@ void TemporallyAccumulate(
 			#error Unsupported.
 		#endif
 	
-		#if !CONFIG_SIGNAL_COLOR
+		#if !COMPILE_SIGNAL_COLOR
 			#error Need to compile signal color.
 		#endif
 
@@ -867,7 +862,7 @@ void TemporallyAccumulate(
 			float PreviousFrameWeight = min(HistorySample.SampleCount, CurrentFrameSample.SampleCount * rcp(MainHisteresiss) - CurrentFrameSample.SampleCount);
 			float PreviousFrameMultiplier = HistorySample.SampleCount > 0 ? PreviousFrameWeight / HistorySample.SampleCount : 0;
 			
-			HistorySample = Add(Mul(HistorySample, PreviousFrameMultiplier), CurrentFrameSample);
+			HistorySample = AddSignal(MulSignal(HistorySample, PreviousFrameMultiplier), CurrentFrameSample);
 		}
 		
 		HistorySamples.Array[BatchedSignalId] = HistorySample;
@@ -902,7 +897,7 @@ void TemporallyAccumulate(
 				/* DestBasis = */ CONFIG_HISTORY_BLENDING_COLOR_SPACE);
 
 			// Blending.
-			HistorySample = Add(Mul(HistorySample, PreviousFrameMultiplier), CurrentFrameSample);
+			HistorySample = AddSignal(MulSignal(HistorySample, PreviousFrameMultiplier), CurrentFrameSample);
 		
 			// Post transformation of the signal.
 			HistorySample = TransformSignal(
