Copying //UE4/Dev-Volumetrics to Dev-Rendering (//UE4/Dev-Rendering)
New SkyAtmosphere component

diff --git a/Engine/Config/BaseScalability.ini b/Engine/Config/BaseScalability.ini
--- a/Engine/Config/BaseScalability.ini
+++ b/Engine/Config/BaseScalability.ini
@@ -390,6 +390,10 @@ r.SSS.HalfRes=1
 r.EmitterSpawnRateScale=0.125
 r.ParticleLightQuality=0
 r.DetailLevel=0
+r.SkyAtmosphere.FastSkyLUT=1
+r.SkyAtmosphere.AerialPerspectiveLUT.FastApplyOnOpaque=1 ; Always have FastSkyLUT 1 in third case to avoid wrong sky
+r.SkyAtmosphere.FastSkyLUT.SampleCountMin=2.0
+r.SkyAtmosphere.FastSkyLUT.SampleCountMin=16.0
 
 [EffectsQuality@1]
 r.TranslucencyLightingVolumeDim=32
@@ -406,6 +410,8 @@ r.SSS.HalfRes=1
 r.EmitterSpawnRateScale=0.25
 r.ParticleLightQuality=0
 r.DetailLevel=1
+r.SkyAtmosphere.FastSkyLUT=1
+r.SkyAtmosphere.AerialPerspectiveLUT.FastApplyOnOpaque=1 ; Always have FastSkyLUT 1 in third case to avoid wrong sky
 
 [EffectsQuality@2]
 r.TranslucencyLightingVolumeDim=48
@@ -422,6 +428,8 @@ r.SSS.HalfRes=1
 r.EmitterSpawnRateScale=0.5
 r.ParticleLightQuality=1
 r.DetailLevel=2
+r.SkyAtmosphere.FastSkyLUT=1
+r.SkyAtmosphere.AerialPerspectiveLUT.FastApplyOnOpaque=0
 
 [EffectsQuality@3]
 r.TranslucencyLightingVolumeDim=64
@@ -438,6 +446,8 @@ r.SSS.HalfRes=1
 r.EmitterSpawnRateScale=1.0
 r.ParticleLightQuality=2
 r.DetailLevel=3
+r.SkyAtmosphere.FastSkyLUT=0
+r.SkyAtmosphere.AerialPerspectiveLUT.FastApplyOnOpaque=0
 
 [EffectsQuality@Cine]
 r.TranslucencyLightingVolumeDim=64
@@ -454,6 +464,14 @@ r.SSS.HalfRes=0
 r.EmitterSpawnRateScale=1.0
 r.ParticleLightQuality=2
 r.DetailLevel=4
+r.SkyAtmosphere.FastSkyLUT=0
+r.r.SkyAtmosphere.AerialPerspectiveLUT.FastApplyOnOpaque=0
+r.SkyAtmosphere.SampleCountMin=16.0
+r.SkyAtmosphere.SampleCountMax=32.0
+r.SkyAtmosphere.AerialPerspectiveLUT.DepthResolution=64.0
+r.SkyAtmosphere.AerialPerspectiveLUT.SampleCountPerSlice=4.0
+r.SkyAtmosphere.TransmittanceLUT.SampleCount=20.0
+r.SkyAtmosphere.MultiScatteringLUT.SampleCount=20.0
 
 ;-----------------------------------------------------------------------------------------------------------------
 
diff --git a/Engine/Shaders/Private/AtmosphereCommon.ush b/Engine/Shaders/Private/AtmosphereCommon.ush
--- a/Engine/Shaders/Private/AtmosphereCommon.ush
+++ b/Engine/Shaders/Private/AtmosphereCommon.ush
@@ -543,8 +543,8 @@ float3 GetSunColor(float3 X, float T, float3 V, float3 S, float Radius, float Mu
     }
 	else 
 	{
-		float SunDiskMask = step(cos(View.AtmosphericFogSunDiscHalfApexAngleRadian), dot(V, S));
-        return TransmittanceValue * SunDiskMask * View.AtmosphericFogSunDiscLuminance; // Eq (9)
+		float SunDiskMask = step(View.AtmosphereLightDiscCosHalfApexAngle[0], dot(V, S));
+        return TransmittanceValue * SunDiskMask * View.AtmosphereLightDiscLuminance[0].rgb; // Eq (9)
     }
 }
 
@@ -578,19 +578,19 @@ float4 GetAtmosphericFog(float3 ViewPosition, float3 ViewVector, float SceneDept
 
 	float3 Attenuation;
 
-	float3 InscatterColor = GetInscatterColor(FogDepth, ViewPosition, T, V, View.AtmosphericFogSunDirection, Radius, Mu, Attenuation, bIsSceneGeometry); //S[L]-T(ViewPosition,xs)S[l]|xs
+	float3 InscatterColor = GetInscatterColor(FogDepth, ViewPosition, T, V, View.AtmosphereLightDirection[0].xyz, Radius, Mu, Attenuation, bIsSceneGeometry); //S[L]-T(ViewPosition,xs)S[l]|xs
 #if BASEPASS_ATMOSPHERIC_FOG // Transluceny rendering just follows normal render flag
 	float3 GroundColor = 0.f;
 	BRANCH
 	if ((View.AtmosphericFogRenderMask & RENDERFLAG_DISABLE_GROUND_SCATTERING) == 0)
 	{
-		GroundColor = GetGroundColor(float4(SceneColor.xyz, 1.f), ViewPosition, T, V, View.AtmosphericFogSunDirection, Radius, Attenuation, bIsSceneGeometry); //R[L0]+R[L*]
+		GroundColor = GetGroundColor(float4(SceneColor.xyz, 1.f), ViewPosition, T, V, View.AtmosphereLightDirection[0].xyz, Radius, Attenuation, bIsSceneGeometry); //R[L0]+R[L*]
 	}
 #else
 	#if MATERIAL_ATMOSPHERIC_FOG || ATMOSPHERIC_NO_GROUND_SCATTERING // Material doesn't need this ground scattering for now 
 		float3 GroundColor = 0.f;
 	#else
-		float3 GroundColor = GetGroundColor(float4(SceneColor.xyz, 1.f), ViewPosition, T, V, View.AtmosphericFogSunDirection, Radius, Attenuation, bIsSceneGeometry); //R[L0]+R[L*]
+		float3 GroundColor = GetGroundColor(float4(SceneColor.xyz, 1.f), ViewPosition, T, V, View.AtmosphereLightDirection[0].xyz, Radius, Attenuation, bIsSceneGeometry); //R[L0]+R[L*]
 	#endif
 #endif
 
@@ -599,19 +599,19 @@ float4 GetAtmosphericFog(float3 ViewPosition, float3 ViewVector, float SceneDept
 	BRANCH
 	if ((View.AtmosphericFogRenderMask & RENDERFLAG_DISABLE_SUN_DISK) == 0)
 	{
-		Sun = GetSunColor(ViewPosition, T, V, View.AtmosphericFogSunDirection, Radius, Mu); //L0
+		Sun = GetSunColor(ViewPosition, T, V, View.AtmosphereLightDirection[0].xyz, Radius, Mu); //L0
 	}
 #else
 	#if MATERIAL_ATMOSPHERIC_FOG || ATMOSPHERIC_NO_SUN_DISK // Material doesn't need to render the sun disk
 		float3 Sun = 0.f;
 	#else
-		float3 Sun = GetSunColor(ViewPosition, T, V, View.AtmosphericFogSunDirection, Radius, Mu); //L0
+		float3 Sun = GetSunColor(ViewPosition, T, V, View.AtmosphereLightDirection[0].xyz, Radius, Mu); //L0
 	#endif
 #endif
 
 	// Decay color
 	float3 OriginalColor = Sun + GroundColor + InscatterColor;
-	OriginalColor = View.AtmosphericFogSunPower * ShadowFactor * View.AtmosphericFogSunColor.rgb * OriginalColor;
+	OriginalColor = View.AtmosphericFogSunPower * ShadowFactor * View.AtmosphereLightColor[0].rgb * OriginalColor;
 	float4 OutColor = float4(OriginalColor, lerp(saturate(Attenuation.r * View.AtmosphericFogDensityScale - View.AtmosphericFogDensityOffset), 1.f, (1.f - DistanceRatio)) );
 	return OutColor;
 }
diff --git a/Engine/Shaders/Private/BasePassCommon.ush b/Engine/Shaders/Private/BasePassCommon.ush
--- a/Engine/Shaders/Private/BasePassCommon.ush
+++ b/Engine/Shaders/Private/BasePassCommon.ush
@@ -11,12 +11,28 @@
 	#define PlanarReflectionStruct TranslucentBasePass.Shared.PlanarReflection
 	#define FogStruct TranslucentBasePass.Shared.Fog
 	#define ActualSSProfilesTexture TranslucentBasePass.Shared.SSProfilesTexture
+	#define ApplyCameraAerialPerspectiveVolume TranslucentBasePass.Shared.SkyAtmosphere.ApplyCameraAerialPerspectiveVolume
+	#define CameraAerialPerspectiveVolume TranslucentBasePass.Shared.SkyAtmosphere.CameraAerialPerspectiveVolume
+	#define CameraAerialPerspectiveVolumeSampler TranslucentBasePass.Shared.SkyAtmosphere.CameraAerialPerspectiveVolumeSampler
+	#define CameraAerialPerspectiveVolumeStartDepth TranslucentBasePass.Shared.SkyAtmosphere.CameraAerialPerspectiveVolumeStartDepth
+	#define CameraAerialPerspectiveVolumeDepthResolution TranslucentBasePass.Shared.SkyAtmosphere.CameraAerialPerspectiveVolumeDepthResolution
+	#define CameraAerialPerspectiveVolumeDepthResolutionInv TranslucentBasePass.Shared.SkyAtmosphere.CameraAerialPerspectiveVolumeDepthResolutionInv
+	#define CameraAerialPerspectiveVolumeDepthSliceLength TranslucentBasePass.Shared.SkyAtmosphere.CameraAerialPerspectiveVolumeDepthSliceLength
+	#define CameraAerialPerspectiveVolumeDepthSliceLengthInv TranslucentBasePass.Shared.SkyAtmosphere.CameraAerialPerspectiveVolumeDepthSliceLengthInv
 #else
 	#define ForwardLightData OpaqueBasePass.Shared.Forward
 	#define ReflectionStruct OpaqueBasePass.Shared.Reflection
 	#define PlanarReflectionStruct OpaqueBasePass.Shared.PlanarReflection
 	#define FogStruct OpaqueBasePass.Shared.Fog
 	#define ActualSSProfilesTexture OpaqueBasePass.Shared.SSProfilesTexture
+	#define ApplyCameraAerialPerspectiveVolume OpaqueBasePass.Shared.SkyAtmosphere.ApplyCameraAerialPerspectiveVolume
+	#define CameraAerialPerspectiveVolume OpaqueBasePass.Shared.SkyAtmosphere.CameraAerialPerspectiveVolume
+	#define CameraAerialPerspectiveVolumeSampler OpaqueBasePass.Shared.SkyAtmosphere.CameraAerialPerspectiveVolumeSampler
+	#define CameraAerialPerspectiveVolumeStartDepth OpaqueBasePass.Shared.SkyAtmosphere.CameraAerialPerspectiveVolumeStartDepth
+	#define CameraAerialPerspectiveVolumeDepthResolution OpaqueBasePass.Shared.SkyAtmosphere.CameraAerialPerspectiveVolumeDepthResolution
+	#define CameraAerialPerspectiveVolumeDepthResolutionInv OpaqueBasePass.Shared.SkyAtmosphere.CameraAerialPerspectiveVolumeDepthResolutionInv
+	#define CameraAerialPerspectiveVolumeDepthSliceLength OpaqueBasePass.Shared.SkyAtmosphere.CameraAerialPerspectiveVolumeDepthSliceLength
+	#define CameraAerialPerspectiveVolumeDepthSliceLengthInv OpaqueBasePass.Shared.SkyAtmosphere.CameraAerialPerspectiveVolumeDepthSliceLengthInv
 #endif 
 
 #undef NEEDS_LIGHTMAP_COORDINATE
diff --git a/Engine/Shaders/Private/BasePassPixelShader.usf b/Engine/Shaders/Private/BasePassPixelShader.usf
--- a/Engine/Shaders/Private/BasePassPixelShader.usf
+++ b/Engine/Shaders/Private/BasePassPixelShader.usf
@@ -33,6 +33,7 @@
 
 #if NEEDS_BASEPASS_PIXEL_FOGGING || NEEDS_BASEPASS_PIXEL_VOLUMETRIC_FOGGING
 	#include "HeightFogCommon.ush"
+	#include "SkyAtmosphereCommon.ush"
 #endif
 
 #include "ReflectionEnvironmentShared.ush"
@@ -1128,6 +1129,26 @@ ShadingModel);
 	}
 #endif
 
+#if NEEDS_BASEPASS_PIXEL_FOGGING
+	if (ApplyCameraAerialPerspectiveVolume > 0.0f)
+	{
+		const float OneOverPreExposure = USE_PREEXPOSURE ? ResolvedView.OneOverPreExposure : 1.0f;
+
+		float4 NDCPosition = mul(float4(MaterialParameters.AbsoluteWorldPosition.xyz, 1), ResolvedView.WorldToClip);
+
+		// Sample the aerial perspective (AP).
+		float4 AP = GetAerialPerspectiveLuminanceTransmittance(
+			NDCPosition, MaterialParameters.AbsoluteWorldPosition.xyz*CM_TO_SKY_UNIT, ResolvedView.WorldCameraOrigin.xyz*CM_TO_SKY_UNIT,
+			CameraAerialPerspectiveVolume, CameraAerialPerspectiveVolumeSampler,
+			CameraAerialPerspectiveVolumeDepthResolutionInv,
+			CameraAerialPerspectiveVolumeDepthResolution,
+			CameraAerialPerspectiveVolumeStartDepth,
+			CameraAerialPerspectiveVolumeDepthSliceLength,
+			CameraAerialPerspectiveVolumeDepthSliceLengthInv, 
+			OneOverPreExposure, Fogging);
+	}
+#endif
+
 	// Volume lighting for lit translucency
 #if (MATERIAL_SHADINGMODEL_DEFAULT_LIT || MATERIAL_SHADINGMODEL_SUBSURFACE) && (MATERIALBLENDING_TRANSLUCENT || MATERIALBLENDING_ADDITIVE) && !SIMPLE_FORWARD_SHADING && !FORWARD_SHADING
 	if (GBuffer.ShadingModelID == SHADINGMODELID_DEFAULT_LIT || GBuffer.ShadingModelID == SHADINGMODELID_SUBSURFACE)
diff --git a/Engine/Shaders/Private/BasePassVertexShader.usf b/Engine/Shaders/Private/BasePassVertexShader.usf
--- a/Engine/Shaders/Private/BasePassVertexShader.usf
+++ b/Engine/Shaders/Private/BasePassVertexShader.usf
@@ -18,6 +18,10 @@
 #include "ForwardLightingCommon.ush"
 #endif
 
+#if NEEDS_BASEPASS_VERTEX_FOGGING
+#include "SkyAtmosphereCommon.ush"
+#endif
+
 /** Entry point for the base pass vertex shader. */
 void Main(
 	FVertexFactoryInput Input,
@@ -123,6 +127,21 @@ void Main(
 	#else
 	Output.BasePassInterpolants.VertexFog = CalculateHeightFog(WorldPosition.xyz - ResolvedView.TranslatedWorldCameraOrigin);
 	#endif
+
+	if (ApplyCameraAerialPerspectiveVolume > 0.0f)
+	{
+		const float OneOverPreExposure = USE_PREEXPOSURE ? ResolvedView.OneOverPreExposure : 1.0f;
+		// Sample the aerial perspective (AP). It is also blended under the VertexFog parameter.
+		float4 AP = GetAerialPerspectiveLuminanceTransmittance(
+			Output.Position, WorldPosition.xyz*CM_TO_SKY_UNIT, ResolvedView.TranslatedWorldCameraOrigin*CM_TO_SKY_UNIT,
+			CameraAerialPerspectiveVolume, CameraAerialPerspectiveVolumeSampler,
+			CameraAerialPerspectiveVolumeDepthResolutionInv,
+			CameraAerialPerspectiveVolumeDepthResolution,
+			CameraAerialPerspectiveVolumeStartDepth,
+			CameraAerialPerspectiveVolumeDepthSliceLength,
+			CameraAerialPerspectiveVolumeDepthSliceLengthInv,
+			OneOverPreExposure, Output.BasePassInterpolants.VertexFog);
+	}
 #endif
 
 #if TRANSLUCENCY_ANY_PERVERTEX_LIGHTING
diff --git a/Engine/Shaders/Private/Definitions.usf b/Engine/Shaders/Private/Definitions.usf
--- a/Engine/Shaders/Private/Definitions.usf
+++ b/Engine/Shaders/Private/Definitions.usf
@@ -104,6 +104,10 @@
 #define	MATERIAL_ATMOSPHERIC_FOG						0
 #endif
 
+#ifndef MATERIAL_SKY_ATMOSPHERE
+#define	MATERIAL_SKY_ATMOSPHERE							0
+#endif
+
 #ifndef BASEPASS_ATMOSPHERIC_FOG
 #define	BASEPASS_ATMOSPHERIC_FOG						0
 #endif
diff --git a/Engine/Shaders/Private/MaterialTemplate.ush b/Engine/Shaders/Private/MaterialTemplate.ush
--- a/Engine/Shaders/Private/MaterialTemplate.ush
+++ b/Engine/Shaders/Private/MaterialTemplate.ush
@@ -73,6 +73,10 @@
 	#include "/Engine/Private/AtmosphereCommon.ush"
 #endif
 
+#if MATERIAL_SKY_ATMOSPHERE
+#include "/Engine/Private/SkyAtmosphereCommon.ush"
+#endif
+
 #include "/Engine/Private/PaniniProjection.ush"
 
 #ifndef USE_DITHERED_LOD_TRANSITION
@@ -975,7 +979,7 @@ float4 MaterialExpressionAtmosphericFog(FMaterialPixelParameters Parameters, flo
 float3 MaterialExpressionAtmosphericLightVector(FMaterialPixelParameters Parameters)
 {
 #if MATERIAL_ATMOSPHERIC_FOG
-	return ResolvedView.AtmosphericFogSunDirection;
+	return ResolvedView.AtmosphereLightDirection[0].xyz;
 #else
 	return float3(0.f, 0.f, 0.f);
 #endif
@@ -984,12 +988,116 @@ float3 MaterialExpressionAtmosphericLightVector(FMaterialPixelParameters Paramet
 float3 MaterialExpressionAtmosphericLightColor(FMaterialPixelParameters Parameters)
 {
 #if MATERIAL_ATMOSPHERIC_FOG
-	return ResolvedView.AtmosphericFogSunColor;
+	return ResolvedView.AtmosphereLightColor[0].rgb;
 #else
 	return float3(0.f, 0.f, 0.f);
 #endif
 }
 
+float3 MaterialExpressionSkyAtmosphereLightIlluminance(FMaterialPixelParameters Parameters, uint LightIndex)
+{
+#if MATERIAL_SKY_ATMOSPHERE
+	const float3 PlanetCenter = float3(0, 0, -ResolvedView.SkyAtmosphereBottomRadius);
+	const float3 WorldPos = GetWorldPosition(Parameters)*CM_TO_SKY_UNIT;
+	const float3 PlanetCenterToWorldPos = WorldPos - PlanetCenter;
+	const float WorldPosHeight = length(PlanetCenterToWorldPos);
+	const float3 UpVector = PlanetCenterToWorldPos / WorldPosHeight;
+	const float SunZenithCosAngle = dot(ResolvedView.AtmosphereLightDirection[LightIndex].xyz, UpVector);
+
+	float2 TransmittanceLutUv;
+	getTransmittanceLutUvs(WorldPosHeight, SunZenithCosAngle, ResolvedView.SkyAtmosphereBottomRadius, ResolvedView.SkyAtmosphereTopRadius, TransmittanceLutUv);
+	return ResolvedView.AtmosphereLightColor[LightIndex].rgb * Texture2DSampleLevel(View.TransmittanceLutTexture, View.TransmittanceLutTextureSampler, TransmittanceLutUv, 0.0f).rgb;
+#else
+	return float3(0.0f, 0.0f, 0.0f);
+#endif
+}
+
+float3 MaterialExpressionSkyAtmosphereLightDirection(FMaterialPixelParameters Parameters, uint LightIndex)
+{
+#if MATERIAL_SKY_ATMOSPHERE
+	return ResolvedView.AtmosphereLightDirection[LightIndex].xyz;
+#else
+	return float3(0.0f, 0.0f, 0.0f);
+#endif
+}
+
+float3 MaterialExpressionSkyAtmosphereLightDiskLuminance(FMaterialPixelParameters Parameters, uint LightIndex)
+{
+#if MATERIAL_SKY_ATMOSPHERE
+	const float3 PlanetCenter = float3(0, 0, -ResolvedView.SkyAtmosphereBottomRadius);
+	const float3 PlanetCenterToWorldCameraPos = ResolvedView.WorldCameraOrigin * CM_TO_SKY_UNIT - PlanetCenter;
+	const float3 ViewDir = -Parameters.CameraVector;
+
+	return GetLightDiskLuminance(PlanetCenterToWorldCameraPos, ViewDir, ResolvedView.SkyAtmosphereBottomRadius, ResolvedView.SkyAtmosphereTopRadius,
+		View.TransmittanceLutTexture, View.TransmittanceLutTextureSampler,
+		ResolvedView.AtmosphereLightDirection[LightIndex].xyz, ResolvedView.AtmosphereLightDiscCosHalfApexAngle[LightIndex], ResolvedView.AtmosphereLightDiscLuminance[LightIndex].xyz);
+#else
+	return float3(0.0f, 0.0f, 0.0f);
+#endif
+}
+
+float3 MaterialExpressionSkyAtmosphereViewLuminance(FMaterialPixelParameters Parameters)
+{
+#if MATERIAL_SKY_ATMOSPHERE
+	const float3 PlanetCenter = float3(0, 0, -ResolvedView.SkyAtmosphereBottomRadius);
+	const float3 PlanetCenterToWorldCameraPos = ResolvedView.WorldCameraOrigin * CM_TO_SKY_UNIT - PlanetCenter;
+	const float ViewHeight = length(PlanetCenterToWorldCameraPos);
+	const float3 UpVector = PlanetCenterToWorldCameraPos / ViewHeight;
+	const float3 ViewDir = -Parameters.CameraVector;
+	const float ViewZenithCosAngle = dot(ViewDir, UpVector);
+
+	float2 Sol = RayIntersectSphere(PlanetCenterToWorldCameraPos, ViewDir, float4(0.0f, 0.0f, 0.0f, ResolvedView.SkyAtmosphereBottomRadius));
+	const bool IntersectGround = any(Sol > 0.0f);
+
+	float2 SkyViewLutUv;
+	SkyViewLutParamsToUv(IntersectGround, ViewZenithCosAngle, ViewDir, ViewHeight, ResolvedView.SkyAtmosphereBottomRadius, ResolvedView.SkyViewLutSizeAndInvSize, SkyViewLutUv);
+	float3 SkyAtmosphereViewLuminance = Texture2DSampleLevel(View.SkyViewLutTexture, View.SkyViewLutTextureSampler, SkyViewLutUv, 0.0f).rgb;
+#if USE_PREEXPOSURE
+	SkyAtmosphereViewLuminance *= ResolvedView.OneOverPreExposure;
+#endif
+	return SkyAtmosphereViewLuminance;
+#else
+	return float3(0.0f, 0.0f, 0.0f);
+#endif
+}
+
+float4 MaterialExpressionSkyAtmosphereAerialPerspective(FMaterialPixelParameters Parameters)
+{
+#if MATERIAL_SKY_ATMOSPHERE
+	const float OneOverPreExposure = USE_PREEXPOSURE ? ResolvedView.OneOverPreExposure : 1.0f;
+
+	const float3 PlanetCenter = float3(0, 0, -ResolvedView.SkyAtmosphereBottomRadius);
+	const float3 WorldPos = Parameters.AbsoluteWorldPosition.xyz*CM_TO_SKY_UNIT;
+	const float3 CameraPos = ResolvedView.WorldCameraOrigin.xyz*CM_TO_SKY_UNIT;
+
+	float4 NDCPosition = mul(float4(Parameters.AbsoluteWorldPosition.xyz, 1), ResolvedView.WorldToClip); // Because SvPosition seems different...
+
+	float4 FogToApplyOver = 1.0f;
+	float4 AerialPerspective = GetAerialPerspectiveLuminanceTransmittance(
+		NDCPosition, WorldPos, CameraPos,
+		View.CameraAerialPerspectiveVolume, View.CameraAerialPerspectiveVolumeSampler,
+		ResolvedView.SkyAtmosphereCameraAerialPerspectiveVolumeDepthResolutionInv,
+		ResolvedView.SkyAtmosphereCameraAerialPerspectiveVolumeDepthResolution,
+		ResolvedView.SkyAtmosphereCameraAerialPerspectiveVolumeStartDepth,
+		ResolvedView.SkyAtmosphereCameraAerialPerspectiveVolumeDepthSliceLength,
+		ResolvedView.SkyAtmosphereCameraAerialPerspectiveVolumeDepthSliceLengthInv,
+		OneOverPreExposure, FogToApplyOver);
+	return AerialPerspective;
+#else
+	return float4(0.0f, 0.0f, 0.0f, 1.0f); // RGB= null scattering, A= null transmittance
+#endif
+}
+
+float3 MaterialExpressionSkyAtmosphereDistantLightScatteredLuminance(FMaterialPixelParameters Parameters)
+{
+#if MATERIAL_SKY_ATMOSPHERE
+	// TODO load on platforms supporting it
+	return Texture2DSampleLevel(View.DistantSkyLightLutTexture, View.DistantSkyLightLutTextureSampler, float2(0.5f, 0.5f), 0.0f).rgb;
+#else
+	return float3(0.0f, 0.0f, 0.0f);
+#endif
+}
+
 /**
  * Utility function to unmirror one coordinate value to the other side
  * UnMirrored == 1 if normal
diff --git a/Engine/Shaders/Private/MobileBasePassCommon.ush b/Engine/Shaders/Private/MobileBasePassCommon.ush
--- a/Engine/Shaders/Private/MobileBasePassCommon.ush
+++ b/Engine/Shaders/Private/MobileBasePassCommon.ush
@@ -6,6 +6,14 @@
 
 #define FogStruct MobileBasePass.Fog
 #define PlanarReflectionStruct MobileBasePass.PlanarReflection
+#define ApplyCameraAerialPerspectiveVolume MobileBasePass.SkyAtmosphere.ApplyCameraAerialPerspectiveVolume
+#define CameraAerialPerspectiveVolume MobileBasePass.SkyAtmosphere.CameraAerialPerspectiveVolume
+#define CameraAerialPerspectiveVolumeSampler MobileBasePass.SkyAtmosphere.CameraAerialPerspectiveVolumeSampler
+#define CameraAerialPerspectiveVolumeStartDepth MobileBasePass.SkyAtmosphere.CameraAerialPerspectiveVolumeStartDepth
+#define CameraAerialPerspectiveVolumeDepthResolution MobileBasePass.SkyAtmosphere.CameraAerialPerspectiveVolumeDepthResolution
+#define CameraAerialPerspectiveVolumeDepthResolutionInv MobileBasePass.SkyAtmosphere.CameraAerialPerspectiveVolumeDepthResolutionInv
+#define CameraAerialPerspectiveVolumeDepthSliceLength MobileBasePass.SkyAtmosphere.CameraAerialPerspectiveVolumeDepthSliceLength
+#define CameraAerialPerspectiveVolumeDepthSliceLengthInv MobileBasePass.SkyAtmosphere.CameraAerialPerspectiveVolumeDepthSliceLengthInv
 
 #define USE_PS_CLIP_PLANE PROJECT_ALLOW_GLOBAL_CLIP_PLANE
 
diff --git a/Engine/Shaders/Private/MobileBasePassPixelShader.usf b/Engine/Shaders/Private/MobileBasePassPixelShader.usf
--- a/Engine/Shaders/Private/MobileBasePassPixelShader.usf
+++ b/Engine/Shaders/Private/MobileBasePassPixelShader.usf
@@ -646,6 +646,7 @@ void Main(
 	VertexFog = BasePassInterpolants.VertexFog;
 #endif
 #endif
+	// NEEDS_BASEPASS_PIXEL_FOGGING is not allowed on mobile for the sake of performance.
 
 #if !MATERIAL_SHADINGMODEL_UNLIT && MOBILE_EMULATION
 	Color = lerp(Color, DiffuseColor + SpecularColor, ResolvedView.UnlitViewmodeMask);
diff --git a/Engine/Shaders/Private/MobileBasePassVertexShader.usf b/Engine/Shaders/Private/MobileBasePassVertexShader.usf
--- a/Engine/Shaders/Private/MobileBasePassVertexShader.usf
+++ b/Engine/Shaders/Private/MobileBasePassVertexShader.usf
@@ -21,6 +21,10 @@
 #include "/Engine/Generated/VertexFactory.ush"
 #include "HeightFogCommon.ush"
 
+#if USE_VERTEX_FOG
+#include "SkyAtmosphereCommon.ush"
+#endif
+
 struct FMobileShadingBasePassVSToPS
 {
 	FVertexFactoryInterpolantsVSToPS FactoryInterpolants;
@@ -83,10 +87,26 @@ void Main(
 #endif
 
 #if USE_VERTEX_FOG
+	float4 VertexFog = CalculateHeightFog(WorldPosition.xyz - ResolvedView.TranslatedWorldCameraOrigin);
+
+	if (ApplyCameraAerialPerspectiveVolume > 0.0f)
+	{
+		const float OneOverPreExposure = USE_PREEXPOSURE ? ResolvedView.OneOverPreExposure : 1.0f;
+		// Sample the aerial perspective (AP). It is also blended under the VertexFog parameter.
+		float4 AP = GetAerialPerspectiveLuminanceTransmittance(
+			Output.Position, WorldPosition.xyz*CM_TO_SKY_UNIT, ResolvedView.TranslatedWorldCameraOrigin*CM_TO_SKY_UNIT,
+			CameraAerialPerspectiveVolume, CameraAerialPerspectiveVolumeSampler,
+			CameraAerialPerspectiveVolumeDepthResolutionInv,
+			CameraAerialPerspectiveVolumeDepthResolution,
+			CameraAerialPerspectiveVolumeStartDepth,
+			CameraAerialPerspectiveVolumeDepthSliceLength,
+			CameraAerialPerspectiveVolumeDepthSliceLengthInv,
+			OneOverPreExposure, VertexFog);
+	}
 #if PACK_INTERPOLANTS
-	PackedInterps[0] = CalculateHeightFog(WorldPosition.xyz - ResolvedView.TranslatedWorldCameraOrigin);
+	PackedInterps[0] = VertexFog;
 #else
-	Output.BasePassInterpolants.VertexFog = CalculateHeightFog(WorldPosition.xyz - ResolvedView.TranslatedWorldCameraOrigin);
+	Output.BasePassInterpolants.VertexFog = VertexFog;
 #endif // PACK_INTERPOLANTS
 #endif
 
diff --git a/Engine/Shaders/Private/SkyAtmosphere.usf b/Engine/Shaders/Private/SkyAtmosphere.usf
new file mode 100644
--- /dev/null
+++ b/Engine/Shaders/Private/SkyAtmosphere.usf
@@ -0,0 +1,1478 @@
+// Copyright 1998-2019 Epic Games, Inc. All Rights Reserved.
+
+/*=============================================================================
+	SkyAtmosphere.usf: Sky and atmosphere rendering functions.
+=============================================================================*/
+
+
+#include "Common.ush"
+
+#if SHADING_PATH_MOBILE
+#define MobileSceneTextures SceneTexturesStruct
+#endif
+
+#include "SceneTexturesCommon.ush"
+#include "SkyAtmosphereCommon.ush"
+
+
+
+#ifndef GROUND_GI_ENABLED
+#define GROUND_GI_ENABLED 0
+#endif
+#ifndef TRANSMITANCE_METHOD
+#define TRANSMITANCE_METHOD 2
+#endif
+#ifndef MULTISCATTERING_APPROX_ENABLED 
+#define MULTISCATTERING_APPROX_ENABLED 0 
+#endif
+#ifndef COLORED_TRANSMITTANCE_ENABLED	// Never used, UE4 does not supports dual blending
+#define COLORED_TRANSMITTANCE_ENABLED 0 
+#endif
+#ifndef FASTSKY_ENABLED 
+#define FASTSKY_ENABLED 0 
+#endif
+#ifndef FASTAERIALPERSPECTIVE_ENABLED 
+#define FASTAERIALPERSPECTIVE_ENABLED 0 
+#endif
+#ifndef SOURCE_DISK_ENABLED
+#define SOURCE_DISK_ENABLED 0
+#endif
+#ifndef ONLY_SKY_PIXEL_ENABLED
+#define ONLY_SKY_PIXEL_ENABLED 0
+#endif
+#ifndef PER_PIXEL_NOISE
+#define PER_PIXEL_NOISE 0
+#endif
+#ifndef SECOND_ATMOSPHERE_LIGHT_ENABLED
+#define SECOND_ATMOSPHERE_LIGHT_ENABLED 0
+#endif
+
+#ifndef TRANSMITTANCE_PASS
+#define TRANSMITTANCE_PASS 0
+#endif
+#ifndef MULTISCATT_PASS
+#define MULTISCATT_PASS 0
+#endif
+#ifndef SKYLIGHT_PASS
+#define SKYLIGHT_PASS 0
+#endif
+
+// View data is not available for passes running once per scene (and not once per view).
+#define VIEWDATA_AVAILABLE (TRANSMITTANCE_PASS!=1 && MULTISCATT_PASS!=1 && SKYLIGHT_PASS!=1)
+
+
+
+#define PLANET_RADIUS_OFFSET 0.001f
+
+ // FASTSKY mapping is done based on Light0 as main light
+#define FASTSKY_LIGHT_INDEX 0
+
+// This (arbitrary) multi scattering storage scale is used to make sure the LUT values are in an acceptable range when using r.SkyAtmosphere.MultiScatteringLUT.UseSmallFormat is 1.
+#define MULTISCATTERING_STORE_SCALE (4.0f * PI * 5.0f)
+#define MULTISCATTERING_STORE_SCALE_INV 1.0f/MULTISCATTERING_STORE_SCALE
+
+
+
+// - RayOrigin: ray origin
+// - RayDir: normalized ray direction
+// - SphereCenter: sphere center
+// - SphereRadius: sphere radius
+// - Returns distance from RayOrigin to closest intersecion with sphere,
+//   or -1.0 if no intersection.
+float RaySphereIntersectNearest(float3 RayOrigin, float3 RayDir, float3 SphereCenter, float SphereRadius)
+{
+	float2 Sol = RayIntersectSphere(RayOrigin, RayDir, float4(SphereCenter, SphereRadius));
+	float Sol0 = Sol.x;
+	float Sol1 = Sol.y;
+	if (Sol0 < 0.0f && Sol1 < 0.0f)
+	{
+		return -1.0f;
+	}
+	if (Sol0 < 0.0f)
+	{
+		return max(0.0f, Sol1);
+	}
+	else if (Sol1 < 0.0f)
+	{
+		return max(0.0f, Sol0);
+	}
+	return max(0.0f, min(Sol0, Sol1));
+}
+
+
+float4 GetScreenWorldPos(float4 SVPos, float DeviceZ)
+{
+#if HAS_INVERTED_Z_BUFFER
+	DeviceZ = max(0.000000000001, DeviceZ);	// TODO: investigate why SvPositionToWorld returns bad values when DeviceZ is far=0 when using inverted z
+#endif
+	return float4(SvPositionToWorld(float4(SVPos.xy, DeviceZ, 1.0)), 1.0);
+}
+
+float3 GetScreenWorldDir(in float4 SVPos)
+{
+	float2 ScreenPosition = SvPositionToScreenPosition(SVPos).xy;
+	const float Depth = 1000000.0f;
+	float4 WorldPos = mul(float4(ScreenPosition * Depth, Depth, 1), View.ScreenToWorld);
+	return normalize(WorldPos.xyz - View.WorldCameraOrigin);
+}
+
+float3 GetCameraWorldPos()
+{
+	return View.WorldCameraOrigin;
+}
+
+float3 GetCameraPlanetPos()
+{
+	return GetCameraWorldPos()*CM_TO_SKY_UNIT + float3(0, 0, Atmosphere.BottomRadius);
+}
+
+
+
+
+
+bool MoveToTopAtmosphere(inout float3 WorldPos, in float3 WorldDir, in float AtmosphereTopRadius)
+{
+	float ViewHeight = length(WorldPos);
+	if (ViewHeight > AtmosphereTopRadius)
+	{
+		float TTop = RaySphereIntersectNearest(WorldPos, WorldDir, float3(0.0f, 0.0f, 0.0f), AtmosphereTopRadius);
+		if (TTop >= 0.0f)
+		{
+			float3 UpVector = WorldPos / ViewHeight;
+			float3 UpOffset = UpVector * -PLANET_RADIUS_OFFSET;
+			WorldPos = WorldPos + WorldDir * TTop + UpOffset;
+		}
+		else
+		{
+			// Ray is not intersecting the atmosphere
+			return false;
+		}
+	}
+	return true; // ok to start tracing
+}
+
+
+
+
+
+
+
+
+
+
+
+float RayleighPhase(float CosTheta)
+{
+	float Factor = 3.0f / (16.0f * PI);
+	return Factor * (1.0f + CosTheta * CosTheta);
+}
+
+float HgPhase(float G, float CosTheta)
+{
+	// Reference implementation (i.e. not schlick approximation). 
+	// See http://www.pbr-book.org/3ed-2018/Volume_Scattering/Phase_Functions.html
+	float Numer = 1.0f - G * G;
+	float Denom = 1.0f + G * G + 2.0f * G * CosTheta;
+	return Numer / (4.0f * PI * Denom * sqrt(Denom));
+}
+
+
+
+
+
+
+
+
+
+
+float3 GetAlbedo(float3 Scattering, float3 Extinction)
+{
+	return Scattering / max(0.001f, Extinction);
+}
+
+struct MediumSampleRGB
+{
+	float3 Scattering;
+	float3 Absorption;
+	float3 Extinction;
+
+	float3 ScatteringMie;
+	float3 AbsorptionMie;
+	float3 ExtinctionMie;
+
+	float3 ScatteringRay;
+	float3 AbsorptionRay;
+	float3 ExtinctionRay;
+
+	float3 ScatteringOzo;
+	float3 AbsorptionOzo;
+	float3 ExtinctionOzo;
+
+	float3 Albedo;
+};
+
+// If this is changed, please also update USkyAtmosphereComponent::GetTransmittance 
+MediumSampleRGB SampleMediumRGB(in float3 WorldPos)
+{
+	const float SampleHeight = max(0.0, (length(WorldPos) - Atmosphere.BottomRadius));
+
+	const float DensityMie = exp(Atmosphere.MieDensityExpScale * SampleHeight);
+
+	const float DensityRay = exp(Atmosphere.RayleighDensityExpScale * SampleHeight);
+
+	const float DensityOzo = SampleHeight < Atmosphere.AbsorptionDensity0LayerWidth ?
+		saturate(Atmosphere.AbsorptionDensity0LinearTerm * SampleHeight + Atmosphere.AbsorptionDensity0ConstantTerm) :	// We use saturate to allow the user to create plateau, and it is free on GCN.
+		saturate(Atmosphere.AbsorptionDensity1LinearTerm * SampleHeight + Atmosphere.AbsorptionDensity1ConstantTerm);
+
+	MediumSampleRGB s;
+
+	s.ScatteringMie = DensityMie * Atmosphere.MieScattering.rgb;
+	s.AbsorptionMie = DensityMie * Atmosphere.MieAbsorption.rgb;
+	s.ExtinctionMie = DensityMie * Atmosphere.MieExtinction.rgb;
+
+	s.ScatteringRay = DensityRay * Atmosphere.RayleighScattering.rgb;
+	s.AbsorptionRay = 0.0f;
+	s.ExtinctionRay = s.ScatteringRay + s.AbsorptionRay;
+
+	s.ScatteringOzo = 0.0f;
+	s.AbsorptionOzo = DensityOzo * Atmosphere.AbsorptionExtinction.rgb;
+	s.ExtinctionOzo = s.ScatteringOzo + s.AbsorptionOzo;
+
+	s.Scattering = s.ScatteringMie + s.ScatteringRay + s.ScatteringOzo;
+	s.Absorption = s.AbsorptionMie + s.AbsorptionRay + s.AbsorptionOzo;
+	s.Extinction = s.ExtinctionMie + s.ExtinctionRay + s.ExtinctionOzo;
+	s.Albedo = GetAlbedo(s.Scattering, s.Extinction);
+
+	return s;
+}
+
+
+
+
+#define DEFAULT_SAMPLE_OFFSET 0.3f
+float SkyAtmosphereNoise(float2 UV)
+{
+//	return DEFAULT_SAMPLE_OFFSET;
+//	return float(Rand3DPCG32(int3(UV.x, UV.y, S)).x) / 4294967296.0f;
+#if VIEWDATA_AVAILABLE && PER_PIXEL_NOISE
+	return InterleavedGradientNoise(UV.xy, float(View.StateFrameIndexMod8));
+#else
+	return DEFAULT_SAMPLE_OFFSET;
+#endif
+}
+
+
+
+
+
+
+
+
+
+
+
+
+
+
+////////////////////////////////////////////////////////////
+// LUT functions
+////////////////////////////////////////////////////////////
+
+// Transmittance LUT function parameterisation from Bruneton 2017 [REDACTED_PROFILE_URL]/precomputed_atmospheric_scattering
+// uv in [0,1]
+// ViewZenithCosAngle in [-1,1]
+// ViewHeight in [bottomRAdius, topRadius]
+
+void UvToLutTransmittanceParams(out float ViewHeight, out float ViewZenithCosAngle, in float2 UV)
+{
+	//UV = FromSubUvsToUnit(UV, SkyAtmosphere.TransmittanceLutSizeAndInvSize); // No real impact so off
+	float Xmu = UV.x;
+	float Xr = UV.y;
+
+	float H = sqrt(Atmosphere.TopRadius * Atmosphere.TopRadius - Atmosphere.BottomRadius * Atmosphere.BottomRadius);
+	float Rho = H * Xr;
+	ViewHeight = sqrt(Rho * Rho + Atmosphere.BottomRadius * Atmosphere.BottomRadius);
+
+	float Dmin = Atmosphere.TopRadius - ViewHeight;
+	float Dmax = Rho + H;
+	float D = Dmin + Xmu * (Dmax - Dmin);
+	ViewZenithCosAngle = D == 0.0f ? 1.0f : (H * H - Rho * Rho - D * D) / (2.0f * ViewHeight * D);
+	ViewZenithCosAngle = clamp(ViewZenithCosAngle, -1.0f, 1.0f);
+}
+
+void LutTransmittanceParamsToUv(in float ViewHeight, in float ViewZenithCosAngle, out float2 UV)
+{
+	getTransmittanceLutUvs(ViewHeight, ViewZenithCosAngle, Atmosphere.BottomRadius, Atmosphere.TopRadius, UV);
+}
+
+// SkyViewLut is a new texture used for fast sky rendering.
+// 
+
+void UvToSkyViewLutParams(out float3 ViewDir, in float ViewHeight, in float2 UV)
+{
+	// Constrain uvs to valid sub texel range (avoid zenith derivative issue making LUT usage visible)
+	UV = FromSubUvsToUnit(UV, SkyAtmosphere.SkyViewLutSizeAndInvSize);
+
+	float Vhorizon = sqrt(ViewHeight * ViewHeight - Atmosphere.BottomRadius * Atmosphere.BottomRadius);
+	float CosBeta = Vhorizon / ViewHeight;				// cos of zenith angle from horizon to zeniht
+	float Beta = acos(CosBeta);
+	float ZenithHorizonAngle = PI - Beta;
+
+	float ViewZenithAngle;
+	if (UV.y < 0.5f)
+	{
+		float Coord = 2.0f * UV.y;
+		Coord = 1.0f - Coord;
+		Coord *= Coord;
+		Coord = 1.0f - Coord;
+		ViewZenithAngle = ZenithHorizonAngle * Coord;
+	}
+	else
+	{
+		float Coord = UV.y * 2.0f - 1.0f;
+		Coord *= Coord;
+		ViewZenithAngle = ZenithHorizonAngle + Beta * Coord;
+	}
+	float CosViewZenithAngle = cos(ViewZenithAngle);
+	float SinViewZenithAngle = sqrt(1.0 - CosViewZenithAngle * CosViewZenithAngle) * (ViewZenithAngle > 0.0f ? 1.0f : -1.0f); // Equivalent to sin(ViewZenithAngle)
+
+	float LongitudeViewCosAngle = UV.x * 2.0f * PI;
+
+	// Make sure those values are in range as it could disrupt other math done later such as sqrt(1.0-c*c)
+	float CosLongitudeViewCosAngle = cos(LongitudeViewCosAngle);
+	float SinLongitudeViewCosAngle = sqrt(1.0 - CosLongitudeViewCosAngle * CosLongitudeViewCosAngle) * (LongitudeViewCosAngle <= PI ? 1.0f : -1.0f); // Equivalent to sin(LongitudeViewCosAngle)
+	ViewDir = float3(
+		SinViewZenithAngle * CosLongitudeViewCosAngle,
+		SinViewZenithAngle * SinLongitudeViewCosAngle,
+		CosViewZenithAngle
+		);
+}
+
+
+
+
+
+
+
+
+
+
+
+
+
+
+
+
+#if USE_PREEXPOSURE && VIEWDATA_AVAILABLE
+#define ViewPreExposure View.PreExposure
+#define ViewOneOverPreExposure View.OneOverPreExposure
+#else
+#define ViewPreExposure 1.0f
+#define ViewOneOverPreExposure 1.0f
+#endif
+
+#define NearDepthValue (HAS_INVERTED_Z_BUFFER ? 1.0f : 0.0f)
+#define FarDepthValue  (HAS_INVERTED_Z_BUFFER ? 0.0f : 1.0f)
+
+
+
+SamplerState TransmittanceLutTextureSampler;
+SamplerState MultiScatteredLuminanceLutTextureSampler;
+SamplerState SkyViewLutTextureSampler;
+SamplerState CameraAerialPerspectiveVolumeTextureSampler;
+
+Texture2D<float4> TransmittanceLutTexture;
+Texture2D<float4> MultiScatteredLuminanceLutTexture;
+Texture2D<float4> SkyViewLutTexture;
+Texture3D<float4> CameraAerialPerspectiveVolumeTexture;
+
+
+
+
+
+
+float3 GetLightDiskLuminance(float3 WorldPos, float3 WorldDir, uint LightIndex)
+{
+#if SOURCE_DISK_ENABLED
+	float t = RaySphereIntersectNearest(WorldPos, WorldDir, float3(0.0f, 0.0f, 0.0f), Atmosphere.BottomRadius);
+	if (t < 0.0f) // no intersection with the planet
+	{
+		float3 LightDiskLuminance = GetLightDiskLuminance(
+			WorldPos, WorldDir, Atmosphere.BottomRadius, Atmosphere.TopRadius,
+			TransmittanceLutTexture, TransmittanceLutTextureSampler,
+			View.AtmosphereLightDirection[LightIndex].xyz, View.AtmosphereLightDiscCosHalfApexAngle[LightIndex], View.AtmosphereLightDiscLuminance[LightIndex].xyz);
+
+		// Clamp to avoid crazy high values (and exposed 64000.0f luminance is already crazy high, solar system sun is 1.6x10^9). Also this removes +inf float and helps TAA.
+		const float3 MaxLightLuminance = 64000.0f;
+		float3 ExposedLightLuminance = LightDiskLuminance * ViewPreExposure;
+		ExposedLightLuminance = min(ExposedLightLuminance, MaxLightLuminance);
+
+#if 1
+		const float ViewDotLight = dot(WorldDir, View.AtmosphereLightDirection[LightIndex].xyz);
+		const float CosHalfApex = View.AtmosphereLightDiscCosHalfApexAngle[LightIndex];
+		const float HalfCosHalfApex = CosHalfApex + (1.0f - CosHalfApex) * 0.25; // Start fading when at 75% distance from light disk center (in cosine space)
+
+		// Apply smooth fading at edge. This is currently an eye balled fade out that works well in many cases.
+		const float Weight = 1.0-saturate((HalfCosHalfApex - ViewDotLight) / (HalfCosHalfApex - CosHalfApex));
+		ExposedLightLuminance = ExposedLightLuminance * Weight;
+#endif 
+
+		return ExposedLightLuminance;
+	}
+#endif
+	return 0.0f;
+}
+
+
+float3 GetMultipleScattering(float3 WorlPos, float ViewZenithCosAngle)
+{
+	float2 UV = saturate(float2(ViewZenithCosAngle*0.5f + 0.5f, (length(WorlPos) - Atmosphere.BottomRadius) / (Atmosphere.TopRadius - Atmosphere.BottomRadius)));
+	float3 MultiScatteredLuminance = MultiScatteredLuminanceLutTexture.SampleLevel(MultiScatteredLuminanceLutTextureSampler, UV, 0).rgb;
+	return MultiScatteredLuminance * MULTISCATTERING_STORE_SCALE_INV;
+}
+
+
+float3 GetTransmittance(in float LightZenithCosAngle, in float PHeight)
+{
+	float2 UV;
+	LutTransmittanceParamsToUv(PHeight, LightZenithCosAngle, UV);
+#ifdef WHITE_TRANSMITTANCE
+	float3 TransmittanceToLight = 1.0f;
+#else
+	float3 TransmittanceToLight = TransmittanceLutTexture.SampleLevel(TransmittanceLutTextureSampler, UV, 0).rgb;
+#endif
+	return TransmittanceToLight;
+}
+
+
+
+
+
+
+
+struct SingleScatteringResult
+{
+	float3 L;						// Scattered light (luminance)
+	float3 OpticalDepth;			// Optical depth (1/m)
+	float3 Transmittance;			// Transmittance in [0,1] (unitless)
+	float3 MultiScatAs1;
+};
+
+struct SamplingSetup
+{
+	bool VariableSampleCount;
+	float SampleCountIni;			// Used when VariableSampleCount is false
+	float MinSampleCount;
+	float MaxSampleCount;
+	float DistanceToSampleCountMaxInv;
+};
+
+SingleScatteringResult IntegrateSingleScatteredLuminance(
+	in float4 SVPos, in float3 WorldPos, in float3 WorldDir,
+	in bool Ground, in SamplingSetup Sampling, in float DeviceZ, in bool MieRayPhase,
+	in float3 Light0Dir, in float3 Light1Dir, in float3 Light0Illuminance, in float3 Light1Illuminance,
+	in float AerialPespectiveViewDistanceScale,
+	in float tMaxMax = 9000000.0f)
+{
+	SingleScatteringResult Result;
+	Result.L = 0;
+	Result.OpticalDepth = 0;
+	Result.Transmittance = 1.0f;
+	Result.MultiScatAs1 = 0;
+
+	if (dot(WorldPos, WorldPos) <= Atmosphere.BottomRadius*Atmosphere.BottomRadius)
+		return Result;	// Camera is inside the atmosphere
+
+	float2 PixPos = SVPos.xy;
+
+	// Compute next intersection with atmosphere or ground 
+	float3 PlanetO = float3(0.0f, 0.0f, 0.0f);
+	float tBottom = RaySphereIntersectNearest(WorldPos, WorldDir, PlanetO, Atmosphere.BottomRadius);
+	float tTop = RaySphereIntersectNearest(WorldPos, WorldDir, PlanetO, Atmosphere.TopRadius);
+	float tMax = 0.0f;
+	if (tBottom < 0.0f)
+	{
+		if (tTop < 0.0f)
+		{
+			tMax = 0.0f; // No intersection with planet nor its atmosphere: stop right away  
+			return Result;
+		}
+		else
+		{
+			tMax = tTop;
+		}
+	}
+	else
+	{
+		if (tTop > 0.0f)
+		{
+			tMax = min(tTop, tBottom);
+		}
+	}
+
+	float PlanetOnOpaque = 1.0f;	// This is used to hide opaque meshes under the planet ground
+	float3 DepthBufferWorldPos = 0.0f;
+#if VIEWDATA_AVAILABLE
+	if (DeviceZ != FarDepthValue)
+	{
+		DepthBufferWorldPos = GetScreenWorldPos(SVPos, DeviceZ).xyz*CM_TO_SKY_UNIT;
+		float tDepth = length(DepthBufferWorldPos.xyz - (WorldPos + float3(0.0, 0.0, -Atmosphere.BottomRadius))); // apply planet offset to go back to origin as top of planet mode. 
+		if (tDepth < tMax)
+		{
+			tMax = tDepth;
+		}
+		else
+		{
+			// Artists did not like that we handle automatic hiding of opaque element behind the planet.
+			// Now, pixel under the surface of earht will receive aerial perspective as if they were  on the ground.
+			//PlanetOnOpaque = 0.0;
+		}
+	}
+#endif
+	tMax = min(tMax, tMaxMax);
+
+	// Sample count 
+	float SampleCount = Sampling.SampleCountIni;
+	float SampleCountFloor = Sampling.SampleCountIni;
+	float tMaxFloor = tMax;
+	if (Sampling.VariableSampleCount)
+	{
+		SampleCount = lerp(Sampling.MinSampleCount, Sampling.MaxSampleCount, saturate(tMax*Sampling.DistanceToSampleCountMaxInv));
+		SampleCountFloor = floor(SampleCount);
+		tMaxFloor = tMax * SampleCountFloor / SampleCount;	// rescale tMax to map to the last entire step segment.
+	}
+	float dt = tMax / SampleCount;
+
+	// Phase functions
+	const float uniformPhase = 1.0f / (4.0f * PI);
+	const float3 wi = Light0Dir;
+	const float3 wo = WorldDir;
+	float cosTheta = dot(wi, wo);
+	float MiePhaseValueLight0 = HgPhase(Atmosphere.MiePhaseG, -cosTheta);	// negate cosTheta because due to WorldDir being a "in" direction. 
+	float RayleighPhaseValueLight0 = RayleighPhase(cosTheta);
+#if SECOND_ATMOSPHERE_LIGHT_ENABLED
+	cosTheta = dot(Light1Dir, wo);
+	float MiePhaseValueLight1 = HgPhase(Atmosphere.MiePhaseG, -cosTheta);	// negate cosTheta because due to WorldDir being a "in" direction. 
+	float RayleighPhaseValueLight1 = RayleighPhase(cosTheta);
+#endif
+
+
+	// Ray march the atmosphere to integrate optical depth
+	float3 L = 0.0f;
+	float3 Throughput = 1.0f;
+	float3 OpticalDepth = 0.0f;
+	float t = 0.0f;
+	float tPrev = 0.0f;
+
+	float3 ExposedLight0Illuminance = Light0Illuminance * ViewPreExposure;
+#if SECOND_ATMOSPHERE_LIGHT_ENABLED
+	float3 ExposedLight1Illuminance = Light1Illuminance * ViewPreExposure;
+#endif
+
+	float PixelNoise = PER_PIXEL_NOISE ? SkyAtmosphereNoise(PixPos.xy) : DEFAULT_SAMPLE_OFFSET;
+	for (float SampleI = 0.0f; SampleI < SampleCount; SampleI += 1.0f)
+	{
+		if (Sampling.VariableSampleCount)
+		{
+			// More expenssive but artefact free
+			float t0 = (SampleI) / SampleCountFloor;
+			float t1 = (SampleI + 1.0f) / SampleCountFloor;;
+			// Non linear distribution of samples within the range.
+			t0 = t0 * t0;
+			t1 = t1 * t1;
+			// Make t0 and t1 world space distances.
+			t0 = tMaxFloor * t0;
+			if (t1 > 1.0f)
+			{
+				t1 = tMax;
+				//t1 = tMaxFloor;	// this reveal depth slices
+			}
+			else
+			{
+				t1 = tMaxFloor * t1;
+			}
+			t = t0 + (t1 - t0) * PixelNoise;
+			dt = t1 - t0;
+		}
+		else
+		{
+			t = tMax * SampleI / SampleCount;
+		}
+		float3 P = WorldPos + t * WorldDir;
+
+		MediumSampleRGB Medium = SampleMediumRGB(P);
+		const float3 SampleOpticalDepth = Medium.Extinction * dt * AerialPespectiveViewDistanceScale;
+		const float3 SampleTransmittance = exp(-SampleOpticalDepth);
+		OpticalDepth += SampleOpticalDepth;
+
+		float PHeight = length(P);
+
+		// Phase and trnasmittance for light 0
+		const float3 UpVector = P / PHeight;
+		float Light0ZenithCosAngle = dot(Light0Dir, UpVector);
+		float3 TransmittanceToLight0 = GetTransmittance(Light0ZenithCosAngle, PHeight);
+		float3 PhaseTimesScattering0;
+		if (MieRayPhase)
+		{
+			PhaseTimesScattering0 = Medium.ScatteringMie * MiePhaseValueLight0 + Medium.ScatteringRay * RayleighPhaseValueLight0;
+		}
+		else
+		{
+			PhaseTimesScattering0 = Medium.Scattering * uniformPhase;
+		}
+
+#if SECOND_ATMOSPHERE_LIGHT_ENABLED
+		// Phase and transmittance for light 1
+		float Light1ZenithCosAngle = dot(Light1Dir, UpVector);
+		float3 TransmittanceToLight1 = GetTransmittance(Light1ZenithCosAngle, PHeight);
+		float3 PhaseTimesScattering1;
+		if (MieRayPhase)
+		{
+			PhaseTimesScattering1 = Medium.ScatteringMie * MiePhaseValueLight1 + Medium.ScatteringRay * RayleighPhaseValueLight1;
+		}
+		else
+		{
+			PhaseTimesScattering1 = Medium.Scattering * uniformPhase;
+		}
+#endif
+
+		// Multiple scattering approximation
+		float3 MultiScatteredLuminance0 = 0.0f;
+#if MULTISCATTERING_APPROX_ENABLED
+		MultiScatteredLuminance0 = GetMultipleScattering(P, Light0ZenithCosAngle);
+#endif
+
+		// Planet shadow
+		float tPlanet0 = RaySphereIntersectNearest(P, Light0Dir, PlanetO + PLANET_RADIUS_OFFSET * UpVector, Atmosphere.BottomRadius);
+		float PlanetShadow0 = tPlanet0 >= 0.0f ? 0.0f : 1.0f;
+		// MultiScatteredLuminance is already pre-exposed, atmospheric light contribution needs to be pre exposed
+		// Multi-scattering is also not affected by PlanetShadow or TransmittanceToLight because it contains diffuse light after single scattering.
+		float3 S = ExposedLight0Illuminance * (PlanetShadow0 * TransmittanceToLight0 * PhaseTimesScattering0 + MultiScatteredLuminance0 * Medium.Scattering);
+
+#if SECOND_ATMOSPHERE_LIGHT_ENABLED
+		float tPlanet1 = RaySphereIntersectNearest(P, Light1Dir, PlanetO + PLANET_RADIUS_OFFSET * UpVector, Atmosphere.BottomRadius);
+		float PlanetShadow1 = tPlanet1 >= 0.0f ? 0.0f : 1.0f;
+		//  Multi-scattering can work for the second light but it is disabled for the sake of performance.
+		S += ExposedLight1Illuminance * PlanetShadow1 * TransmittanceToLight1 * PhaseTimesScattering1;// +MultiScatteredLuminance * Medium.Scattering);
+#endif
+
+		// When using the power serie to accumulate all sattering order, serie r must be <1 for a serie to converge. 
+		// Under extreme coefficient, MultiScatAs1 can grow larger and thus results in broken visuals. 
+		// The way to fix that is to use a proper analytical integration as porposed in slide 28 of http://www.frostbite.com/2015/08/physically-based-unified-volumetric-rendering-in-frostbite/ 
+		// However, it is possible to disable as it can also work using simple power serie sum unroll up to 5th order. The rest of the orders has a really low contribution. 
+#define MULTI_SCATTERING_POWER_SERIE 0 
+
+#if MULTI_SCATTERING_POWER_SERIE==0 
+		// 1 is the integration of luminance over the 4pi of a sphere, and assuming an isotropic phase function of 1.0/(4*PI) 
+		Result.MultiScatAs1 += Throughput * Medium.Scattering * 1.0f * dt;
+#else 
+		float3 MS = Medium.Scattering * 1;
+		float3 MSint = (MS - MS * SampleTransmittance) / Medium.Extinction;
+		Result.MultiScatAs1 += Throughput * MSint;
+#endif 
+
+
+#if 0
+		L += Throughput * S * dt;
+		Throughput *= SampleTransmittance;
+#else
+		// See slide 28 at http://www.frostbite.com/2015/08/physically-based-unified-volumetric-rendering-in-frostbite/ 
+		float3 Sint = (S - S * SampleTransmittance) / Medium.Extinction;	// integrate along the current step segment 
+		L += Throughput * Sint;														// accumulate and also take into account the transmittance from previous steps
+		Throughput *= SampleTransmittance;
+#endif
+
+		tPrev = t;
+	}
+
+	if (Ground && tMax == tBottom)
+	{
+		// Account for bounced light off the planet
+		float3 P = WorldPos + tBottom * WorldDir;
+		float PHeight = length(P);
+
+		const float3 UpVector = P / PHeight;
+		float Light0ZenithCosAngle = dot(Light0Dir, UpVector);
+		float3 TransmittanceToLight0 = GetTransmittance(Light0ZenithCosAngle, PHeight);
+
+		const float NdotL0 = saturate(dot(UpVector, Light0Dir));
+		L += Light0Illuminance * TransmittanceToLight0 * Throughput * NdotL0 * Atmosphere.GroundAlbedo.rgb / PI;
+#if SECOND_ATMOSPHERE_LIGHT_ENABLED
+		{
+			const float NdotL1 = saturate(dot(UpVector, Light1Dir));
+			float Light1ZenithCosAngle = dot(UpVector, Light1Dir);
+			float3 TransmittanceToLight1 = GetTransmittance(Light1ZenithCosAngle, PHeight);
+			L += Light1Illuminance * TransmittanceToLight1 * Throughput * NdotL1 * Atmosphere.GroundAlbedo.rgb / PI;
+		}
+#endif
+	}
+
+	Result.L = L;
+	Result.OpticalDepth = OpticalDepth;
+	Result.Transmittance = Throughput * PlanetOnOpaque;
+
+	return Result;
+}
+
+
+
+
+
+
+
+
+
+
+
+
+
+
+
+
+
+
+
+
+
+
+
+
+
+
+
+
+float StartDepthZ;
+
+
+void SkyAtmosphereVS(
+	in uint VertexId : SV_VertexID,
+	out float4 Position : SV_POSITION)
+{
+	float2 UV = -1.0f;
+	UV = VertexId == 1 ? float2(-1.0f, 3.0f) : UV;
+	UV = VertexId == 2 ? float2(3.0f, -1.0f) : UV;
+
+#if ONLY_SKY_PIXEL_ENABLED
+	const float Depth = FarDepthValue;
+#else
+	const float Depth = StartDepthZ;
+#endif
+
+	Position = float4(UV, Depth, 1.0f);
+}
+
+
+
+
+
+
+
+
+
+float4 PrepareOutput(float3 Luminance, float3 Transmittance = float3(1.0f, 1.0f, 1.0f))
+{
+	const float GreyScaleTransmittance = dot(Transmittance, float3(1.0f / 3.0f, 1.0f / 3.0f, 1.0f / 3.0f));
+	return float4(Luminance, GreyScaleTransmittance);
+}
+
+void RenderSkyAtmosphereRayMarchingPS(
+	in float4 SVPos : SV_POSITION,
+	out float4 OutLuminance : SV_Target0
+#if COLORED_TRANSMITTANCE_ENABLED
+	, float4 Transmittance : SV_TARGET1;	// For dual source blending when available
+#endif
+)
+{
+	OutLuminance = 0;
+
+	float2 PixPos = SVPos.xy;
+	float2 UvBuffer = PixPos * View.BufferSizeAndInvSize.zw;	// Uv for depth buffer read (size can be larger than viewport)
+
+	//const float displaySize = 256.0f;
+	//if(all(PixPos<displaySize))
+	//{
+	//	OutLuminance = PrepareOutput(SkyViewLutTexture.SampleLevel(SkyViewLutTextureSampler, PixPos / displaySize, 0).rgb);
+	//	//OutLuminance = PrepareOutput(CameraAerialPerspectiveVolumeTexture.SampleLevel(CameraAerialPerspectiveVolumeTextureSampler, float3(PixPos / displaySize, PixPos.x/ displaySize), 0).rgb * ViewOneOverPreExposure);
+	//	return;
+	//}
+
+	float3 WorldPos = GetCameraPlanetPos();
+	float3 WorldDir = GetScreenWorldDir(SVPos);
+
+	float3 PreExposedL = 0;
+#if  ONLY_SKY_PIXEL_ENABLED
+	// In this case we assume all pixel have a FarDepthValue to avoid reading the depth buffer. 
+	// Others pixels will be skipped by depth test.
+	float DeviceZ = FarDepthValue;
+#else
+	float DeviceZ = LookupDeviceZ(UvBuffer);
+#endif
+
+	// Get the light disk luminance to draw 
+	float3 LuminanceScale = 1.0f;
+	if (DeviceZ == FarDepthValue)
+	{
+		LuminanceScale = SkyAtmosphere.SkyLuminanceFactor;
+		PreExposedL += GetLightDiskLuminance(WorldPos, WorldDir, 0);
+#if SECOND_ATMOSPHERE_LIGHT_ENABLED
+		PreExposedL += GetLightDiskLuminance(WorldPos, WorldDir, 1);
+#endif
+	}
+
+	float3 UpVector = normalize(WorldPos);
+	float ViewHeight = length(WorldPos);
+#if FASTSKY_ENABLED
+	if (ViewHeight < Atmosphere.TopRadius && DeviceZ == FarDepthValue)
+	{
+		float2 UV;
+		float ViewZenithCosAngle = dot(WorldDir, UpVector);
+
+		bool IntersectGround = RaySphereIntersectNearest(WorldPos, WorldDir, float3(0, 0, 0), Atmosphere.BottomRadius) >= 0.0f;
+
+		SkyViewLutParamsToUv(IntersectGround, ViewZenithCosAngle, WorldDir, ViewHeight, Atmosphere.BottomRadius, SkyAtmosphere.SkyViewLutSizeAndInvSize, UV);
+		float3 SkyLuminance = SkyViewLutTexture.SampleLevel(SkyViewLutTextureSampler, UV, 0).rgb;
+
+		PreExposedL += SkyLuminance * LuminanceScale;
+		OutLuminance = PrepareOutput(PreExposedL);
+		return;
+	}
+#endif
+
+#if FASTAERIALPERSPECTIVE_ENABLED
+
+#if  ONLY_SKY_PIXEL_ENABLED==0
+	// We only compute the aerial perspective only when "not only rendering the sky" pixels. We simply remove that code section to simplify the shader.
+
+#if COLORED_TRANSMITTANCE_ENABLED
+#error The FASTAERIALPERSPECTIVE_ENABLED path does not support COLORED_TRANSMITTANCE_ENABLED.
+#else
+
+	const float OneOverPreExposure = 1.0f;
+	float4 NullFog = float4(0.0f, 0.0f, 0.0f, 1.0f);
+
+	float3 DepthBufferWorldPos = GetScreenWorldPos(SVPos, DeviceZ).xyz;
+	float4 NDCPosition = mul(float4(DepthBufferWorldPos.xyz, 1), View.WorldToClip);
+
+	float4 AP = GetAerialPerspectiveLuminanceTransmittance(
+		NDCPosition, DepthBufferWorldPos*CM_TO_SKY_UNIT, GetCameraWorldPos()*CM_TO_SKY_UNIT,
+		CameraAerialPerspectiveVolumeTexture, CameraAerialPerspectiveVolumeTextureSampler,
+		SkyAtmosphere.CameraAerialPerspectiveVolumeDepthResolutionInv,
+		SkyAtmosphere.CameraAerialPerspectiveVolumeDepthResolution,
+		SkyAtmosphere.CameraAerialPerspectiveVolumeStartDepth,
+		SkyAtmosphere.CameraAerialPerspectiveVolumeDepthSliceLength,
+		SkyAtmosphere.CameraAerialPerspectiveVolumeDepthSliceLengthInv,
+		OneOverPreExposure, NullFog);
+
+	PreExposedL += AP.rgb * LuminanceScale;
+	float Transmittance = AP.a;
+
+	OutLuminance = PrepareOutput(PreExposedL, float3(Transmittance, Transmittance, Transmittance));
+	return;
+#endif
+
+#endif // ONLY_SKY_PIXEL_ENABLED==0
+
+#else // FASTAERIALPERSPECTIVE_ENABLED
+
+	WorldPos += WorldDir * SkyAtmosphere.CameraAerialPerspectiveVolumeStartDepth; // Offset according to start depth
+
+	// Move to top atmosphere as the starting point for ray marching.
+	// This is critical to be after the above to not disrupt above atmosphere tests and voxel selection.
+	if (!MoveToTopAtmosphere(WorldPos, WorldDir, Atmosphere.TopRadius))
+	{
+		// Ray is not intersecting the atmosphere
+		OutLuminance = PrepareOutput(PreExposedL);
+		return;
+	}
+
+	SamplingSetup Sampling;
+	{
+		Sampling.VariableSampleCount = true;
+		Sampling.MinSampleCount = SkyAtmosphere.SampleCountMin;
+		Sampling.MaxSampleCount = SkyAtmosphere.SampleCountMax;
+		Sampling.DistanceToSampleCountMaxInv = SkyAtmosphere.DistanceToSampleCountMaxInv;
+	}
+	const bool Ground = false;
+	const bool MieRayPhase = true;
+	const float AerialPespectiveViewDistanceScale = DeviceZ == FarDepthValue ? 1.0f : SkyAtmosphere.AerialPespectiveViewDistanceScale;
+	SingleScatteringResult ss = IntegrateSingleScatteredLuminance(
+		SVPos, WorldPos, WorldDir,
+		Ground, Sampling, DeviceZ, MieRayPhase,
+		View.AtmosphereLightDirection[0].xyz, View.AtmosphereLightDirection[1].xyz, View.AtmosphereLightColor[0].rgb, View.AtmosphereLightColor[1].rgb,
+		AerialPespectiveViewDistanceScale);
+
+	PreExposedL += ss.L * LuminanceScale;
+
+#if COLORED_TRANSMITTANCE_ENABLED
+#error Requires support for dual source blending.
+	output.Luminance = float4(L, 1.0f);
+	output.Transmittance = float4(ss.Transmittance, 1.0f);
+#else
+	OutLuminance = PrepareOutput(PreExposedL, ss.Transmittance);
+	return;
+#endif
+
+#endif // FASTAERIALPERSPECTIVE_ENABLED
+}
+
+
+
+
+
+#ifndef THREADGROUP_SIZE
+#define THREADGROUP_SIZE 1
+#endif
+
+RWTexture2D<float3> TransmittanceLutUAV;
+
+[numthreads(THREADGROUP_SIZE, THREADGROUP_SIZE, 1)]
+void RenderTransmittanceLutCS(uint3 ThreadId : SV_DispatchThreadID)
+{
+	float2 PixPos = float2(ThreadId.xy) + 0.5f;
+
+	// Compute camera position from LUT coords
+	float2 UV = (PixPos) * SkyAtmosphere.TransmittanceLutSizeAndInvSize.zw;
+	float ViewHeight;
+	float ViewZenithCosAngle;
+	UvToLutTransmittanceParams(ViewHeight, ViewZenithCosAngle, UV);
+
+	//  A few extra needed constants
+	float3 WorldPos = float3(0.0f, 0.0f, ViewHeight);
+	float3 WorldDir = float3(0.0f, sqrt(1.0f - ViewZenithCosAngle * ViewZenithCosAngle), ViewZenithCosAngle);
+
+	SamplingSetup Sampling;
+	{
+		Sampling.VariableSampleCount = false;
+		Sampling.SampleCountIni = SkyAtmosphere.TransmittanceSampleCount;
+	}
+	const bool Ground = false;
+	const float DeviceZ = FarDepthValue;
+	const bool MieRayPhase = false;
+	const float3 NullLightDirection = float3(0.0f, 0.0f, 1.0f);
+	const float3 NullLightIlluminance = float3(0.0f, 0.0f, 0.0f);
+	const float AerialPespectiveViewDistanceScale = 1.0f;
+	SingleScatteringResult ss = IntegrateSingleScatteredLuminance(
+		float4(PixPos,0.0f,1.0f), WorldPos, WorldDir,
+		Ground, Sampling, DeviceZ, MieRayPhase,
+		NullLightDirection, NullLightDirection, NullLightIlluminance, NullLightIlluminance,
+		AerialPespectiveViewDistanceScale);
+
+	float3 transmittance = exp(-ss.OpticalDepth);
+
+	TransmittanceLutUAV[int2(PixPos)] = transmittance;
+}
+
+
+
+RWTexture2D<float3> MultiScatteredLuminanceLutUAV;
+
+[numthreads(THREADGROUP_SIZE, THREADGROUP_SIZE, 1)]
+void RenderMultiScatteredLuminanceLutCS(uint3 ThreadId : SV_DispatchThreadID)
+{
+	float2 PixPos = float2(ThreadId.xy) + 0.5f;
+
+	float CosLightZenithAngle = (PixPos.x * SkyAtmosphere.MultiScatteredLuminanceLutSizeAndInvSize.z) * 2.0f - 1.0f;
+	float3 LightDir = float3(0.0f, sqrt(saturate(1.0f - CosLightZenithAngle * CosLightZenithAngle)), CosLightZenithAngle);
+	const float3 NullLightDirection = float3(0.0f, 0.0f, 1.0f);
+	const float3 NullLightIlluminance = float3(0.0f, 0.0f, 0.0f);
+	const float3 OneIlluminance = float3(1.0f, 1.0f, 1.0f);  // Assume a pure white light illuminance for the LUT to act as a transfer (be independent of the light, only dependent on the earth)
+	float ViewHeight = Atmosphere.BottomRadius + (PixPos.y * SkyAtmosphere.MultiScatteredLuminanceLutSizeAndInvSize.w) * (Atmosphere.TopRadius - Atmosphere.BottomRadius);
+
+	float3 WorldPos = float3(0.0f, 0.0f, ViewHeight);
+	float3 WorldDir = float3(0.0f, 0.0f, 1.0f);
+
+	SamplingSetup Sampling;
+	{
+		Sampling.VariableSampleCount = false;
+		Sampling.SampleCountIni = SkyAtmosphere.MultiScatteringSampleCount;
+	}
+	const bool Ground = true;
+	const float DeviceZ = FarDepthValue;
+	const bool MieRayPhase = false;
+	const float AerialPespectiveViewDistanceScale = 1.0f;
+
+	const float SphereSolidAngle = 4.0f * PI;
+	const float IsotropicPhase = 1.0f / SphereSolidAngle;
+#if 1
+
+	// Cheap and good enough approximation (but lose energy) 
+	SingleScatteringResult r0 = IntegrateSingleScatteredLuminance(float4(PixPos, 0.0f, 1.0f), WorldPos, WorldDir, Ground, Sampling, DeviceZ, MieRayPhase,
+		LightDir, NullLightDirection, OneIlluminance, NullLightIlluminance, AerialPespectiveViewDistanceScale);
+	SingleScatteringResult r1 = IntegrateSingleScatteredLuminance(float4(PixPos, 0.0f, 1.0f), WorldPos, -WorldDir, Ground, Sampling, DeviceZ, MieRayPhase,
+		LightDir, NullLightDirection, OneIlluminance, NullLightIlluminance, AerialPespectiveViewDistanceScale);
+
+	float3 IntegratedIlluminance = (SphereSolidAngle / 2.0f) * (r0.L + r1.L);
+	float3 MultiScatAs1 = (1.0f / 2.0f)*(r0.MultiScatAs1 + r1.MultiScatAs1);
+	float3 InScatteredLuminance = IntegratedIlluminance * IsotropicPhase;
+
+#else
+
+	// Less cheap but approximation closer to ground truth
+	SingleScatteringResult r0 = IntegrateSingleScatteredLuminance(float4(PixPos, 0.0f, 1.0f), WorldPos, float3(0.70710678118f 0.0f, 0.70710678118f), LightDir, Ground, Sampling, DeviceZ, MieRayPhase, AerialPespectiveViewDistanceScale);
+	SingleScatteringResult r1 = IntegrateSingleScatteredLuminance(float4(PixPos, 0.0f, 1.0f), WorldPos, float3(-0.70710678118f, 0.0f, 0.70710678118f), LightDir, Ground, Sampling, DeviceZ, MieRayPhase, AerialPespectiveViewDistanceScale);
+	SingleScatteringResult r2 = IntegrateSingleScatteredLuminance(float4(PixPos, 0.0f, 1.0f), WorldPos, float3(0.0f, 0.70710678118f, 0.70710678118f), LightDir, Ground, Sampling, DeviceZ, MieRayPhase, AerialPespectiveViewDistanceScale);
+	SingleScatteringResult r3 = IntegrateSingleScatteredLuminance(float4(PixPos, 0.0f, 1.0f), WorldPos, float3(0.0f, -0.70710678118f, 0.70710678118f), LightDir, Ground, Sampling, DeviceZ, MieRayPhase, AerialPespectiveViewDistanceScale);
+	SingleScatteringResult r4 = IntegrateSingleScatteredLuminance(float4(PixPos, 0.0f, 1.0f), WorldPos, float3(0.70710678118f, 0.0f, -0.70710678118f), LightDir, Ground, Sampling, DeviceZ, MieRayPhase, AerialPespectiveViewDistanceScale);
+	SingleScatteringResult r5 = IntegrateSingleScatteredLuminance(float4(PixPos, 0.0f, 1.0f), WorldPos, float3(-0.70710678118f, 0.0f, -0.70710678118f), LightDir, Ground, Sampling, DeviceZ, MieRayPhase, AerialPespectiveViewDistanceScale);
+	SingleScatteringResult r6 = IntegrateSingleScatteredLuminance(float4(PixPos, 0.0f, 1.0f), WorldPos, float3(0.0f, 0.70710678118f, -0.70710678118f), LightDir, Ground, Sampling, DeviceZ, MieRayPhase, AerialPespectiveViewDistanceScale);
+	SingleScatteringResult r7 = IntegrateSingleScatteredLuminance(float4(PixPos, 0.0f, 1.0f), WorldPos, float3(0.0f, -0.70710678118f, -0.70710678118f), LightDir, Ground, Sampling, DeviceZ, MieRayPhase, AerialPespectiveViewDistanceScale);
+
+	// Integral of in-scattered Luminance (Lumen/(m2.sr)) over the sphere gives illuminance (Lumen/m2). 
+	// This is done with equal importance for each samples over the sphere. 
+	float3 IntegratedIlluminance = (SphereSolidAngle / 8.0f) * (r0.L + r1.L + r2.L + r3.L + r4.L + r5.L + r6.L + r7.L);
+
+	// MultiScatAs1 represents the contribution of a uniform environment light over a sphere of luminance 1 and assuming an isotropic phase function 
+	float3 MultiScatAs1 = (1.0f / 8.0f)*(r0.MultiScatAs1 + r1.MultiScatAs1 + r2.MultiScatAs1 + r3.MultiScatAs1 + r4.MultiScatAs1 + r5.MultiScatAs1 + r6.MultiScatAs1 + r7.MultiScatAs1);
+
+	// Compute the InScatteredLuminance (Lumen/(m2.sr)) assuming a uniform IntegratedIlluminance, isotropic phase function (1.0/sr)
+	// and the fact that this illumiance would be used for each path/raymarch samples of each path
+	float3 InScatteredLuminance = IntegratedIlluminance * IsotropicPhase;
+
+#endif 
+
+	// MultiScatAs1 represents the amount of luminance scattered as if the integral of scattered luminance over the sphere would be 1.
+	//  - 1st order of scattering: one can ray-march a straight path as usual over the sphere. That is InScatteredLuminance.
+	//  - 2nd order of scattering: the inscattered luminance is InScatteredLuminance at each of samples of fist order integration. Assuming a uniform phase function that is represented by MultiScatAs1,
+	//  - 3nd order of scattering: the inscattered luminance is (InScatteredLuminance * MultiScatAs1 * MultiScatAs1)
+	//  - etc.
+#if	MULTI_SCATTERING_POWER_SERIE==0
+	float3 MultiScatAs1SQR = MultiScatAs1 * MultiScatAs1;
+	float3 L = InScatteredLuminance * (1.0f + MultiScatAs1 + MultiScatAs1SQR + MultiScatAs1 * MultiScatAs1SQR + MultiScatAs1SQR * MultiScatAs1SQR);
+#else
+	// For a serie, sum_{n=0}^{n=+inf} = 1 + r + r^2 + r^3 + ... + r^n = 1 / (1.0 - r), see https://en.wikipedia.org/wiki/Geometric_series  
+	const float3 R = MultiScatAs1;
+	const float3 SumOfAllMultiScatteringEventsContribution = 1.0f / (1.0f - R);
+	float3 L = InScatteredLuminance * SumOfAllMultiScatteringEventsContribution;
+#endif
+
+	// MultipleScatteringFactor can be applied here because the LUT is compute every frame
+	// L is pre-exposed since InScatteredLuminance is computed from pre-exposed sun light. So multi-scattering contribution is pre-exposed.
+	MultiScatteredLuminanceLutUAV[int2(PixPos)] = L * Atmosphere.MultiScatteringFactor * MULTISCATTERING_STORE_SCALE;
+}
+
+
+
+
+RWTexture2D<float3> SkyViewLutUAV;
+
+[numthreads(THREADGROUP_SIZE, THREADGROUP_SIZE, 1)]
+void RenderSkyViewLutCS(uint3 ThreadId : SV_DispatchThreadID)
+{
+	float2 PixPos = float2(ThreadId.xy) + 0.5f;
+	float2 UV = PixPos * SkyAtmosphere.SkyViewLutSizeAndInvSize.zw;
+
+
+	float3 WorldPos = GetCameraPlanetPos();
+	float ViewHeight = length(WorldPos);
+
+
+
+	// TODO explain the mapping that is happening here...
+	// TODO we need a valid local basis for lights and we need to transform them in it. 
+	// ===> Without that, the Sky lut is only valid for camera is close to the origin (at the Z-top of the planet)
+	float3 WorldDir;
+	UvToSkyViewLutParams(WorldDir, ViewHeight, UV);
+
+
+	// Move to top atmospehre
+	if (!MoveToTopAtmosphere(WorldPos, WorldDir, Atmosphere.TopRadius))
+	{
+		// Ray is not intersecting the atmosphere
+		SkyViewLutUAV[int2(PixPos)] = 0.0f;
+		return;
+	}
+
+
+	SamplingSetup Sampling;
+	{
+		Sampling.VariableSampleCount = true;
+		Sampling.MinSampleCount = SkyAtmosphere.FastSkySampleCountMin;
+		Sampling.MaxSampleCount = SkyAtmosphere.FastSkySampleCountMax;
+		Sampling.DistanceToSampleCountMaxInv = SkyAtmosphere.FastSkyDistanceToSampleCountMaxInv;
+	}
+	const bool Ground = false;
+	const float DeviceZ = FarDepthValue;
+	const bool MieRayPhase = true;
+	const float AerialPespectiveViewDistanceScale = 1.0f;
+	SingleScatteringResult ss = IntegrateSingleScatteredLuminance(
+		float4(PixPos, 0.0f, 1.0f), WorldPos, WorldDir,
+		Ground, Sampling, DeviceZ, MieRayPhase,
+		View.AtmosphereLightDirection[0].xyz, View.AtmosphereLightDirection[1].xyz, View.AtmosphereLightColor[0].rgb, View.AtmosphereLightColor[1].rgb,
+		AerialPespectiveViewDistanceScale);
+
+	SkyViewLutUAV[int2(PixPos)] = ss.L;
+}
+
+
+
+
+RWTexture2D<float3> DistantSkyLightLutUAV;
+Buffer<float4> UniformSphereSamplesBuffer;
+float4 AtmosphereLightDirection0;
+float4 AtmosphereLightDirection1;
+float4 AtmosphereLightColor0;
+float4 AtmosphereLightColor1;
+float DistantSkyLightSampleAltitude;
+
+groupshared float3 GroupSkyLuminanceSamples[THREADGROUP_SIZE*THREADGROUP_SIZE];
+
+[numthreads(THREADGROUP_SIZE, THREADGROUP_SIZE, 1)]
+void RenderDistantSkyLightLutCS(uint3 ThreadId : SV_DispatchThreadID)
+{
+	const int LinearIndex = ThreadId.y*THREADGROUP_SIZE + ThreadId.x;
+	float2 PixPos = float2(ThreadId.xy) + 0.5f;
+	float2 UV = PixPos * SkyAtmosphere.SkyViewLutSizeAndInvSize.zw;
+
+	// As of today, we assume the world is alway at the top of the planet along Z.
+	// If it needs to change, we can transform all AtmosphereLightDirection into local basis onder the camera (it would then be view dependent).
+	// Overall this is fine because this sky lighting ambient is used for clouds using a dome and it won't be used in such sky views.
+	// IF needed later, we could compute illuminance for multiple position on earth in a lat/long texture
+	float3 SamplePos = float3(0.0f, 0.0f, DistantSkyLightSampleAltitude) + float3(0, 0, Atmosphere.BottomRadius);
+	float ViewHeight = length(SamplePos);
+
+	// We are going to trace 64 times using 64 parallel threads.
+	// Result are written in shared memory and prefix sum is applied to integrate the lighting in a single RGB value
+	// that can then be used to lit clouds in the sky mesh shader graph.
+
+	// Select a direction for this thread
+	const float3 SampleDir = UniformSphereSamplesBuffer[LinearIndex].xyz;
+
+	SamplingSetup Sampling;
+	{
+		Sampling.VariableSampleCount = false;
+		Sampling.SampleCountIni = 10.0f;
+	}
+	const bool Ground = false;
+	const float DeviceZ = FarDepthValue;
+	const bool MieRayPhase = false;
+	const float AerialPespectiveViewDistanceScale = 1.0f;
+	SingleScatteringResult ss = IntegrateSingleScatteredLuminance(
+		float4(PixPos, 0.0f, 1.0f), SamplePos, SampleDir,
+		Ground, Sampling, DeviceZ, MieRayPhase,
+		AtmosphereLightDirection0.xyz, AtmosphereLightDirection1.xyz, AtmosphereLightColor0.rgb, AtmosphereLightColor1.rgb,
+		AerialPespectiveViewDistanceScale);
+
+	GroupSkyLuminanceSamples[LinearIndex] = ss.L;
+
+	// Wait for all group threads to be done
+	GroupMemoryBarrierWithGroupSync();
+
+	// Now we manually apply prefix sum for a thread group size of 64
+#if SKYLIGHT_PASS==1 && THREADGROUP_SIZE!=8
+#error This shader only works for a thread group size of 8x8
+#endif
+
+	if (LinearIndex < 32)
+	{
+		GroupSkyLuminanceSamples[LinearIndex] += GroupSkyLuminanceSamples[LinearIndex + 32];
+	}
+	GroupMemoryBarrierWithGroupSync();
+	if (LinearIndex < 16)
+	{
+		GroupSkyLuminanceSamples[LinearIndex] += GroupSkyLuminanceSamples[LinearIndex + 16];
+	}
+	GroupMemoryBarrierWithGroupSync();
+
+	// The smallest wave size is 16 on Intel hardware. So now we can do simple math operations without group sync.
+
+	if (LinearIndex < 8)
+	{
+		GroupSkyLuminanceSamples[LinearIndex] += GroupSkyLuminanceSamples[LinearIndex + 8];
+	}
+	if (LinearIndex < 4)
+	{
+		GroupSkyLuminanceSamples[LinearIndex] += GroupSkyLuminanceSamples[LinearIndex + 4];
+	}
+	if (LinearIndex < 2)
+	{
+		GroupSkyLuminanceSamples[LinearIndex] += GroupSkyLuminanceSamples[LinearIndex + 2];
+	}
+	if (LinearIndex < 1)
+	{
+		const float3 AccumulatedLuminanceSamples = GroupSkyLuminanceSamples[LinearIndex] + GroupSkyLuminanceSamples[LinearIndex + 1];
+		const float SamplerSolidAngle = 4.0f * PI / float(THREADGROUP_SIZE * THREADGROUP_SIZE);
+		const float3 Illuminance = AccumulatedLuminanceSamples * SamplerSolidAngle;
+		const float3 UniformPhaseFunction = 1.0f / (4.0f * PI);
+		DistantSkyLightLutUAV[int2(PixPos)] = Illuminance * UniformPhaseFunction; // Luminance assuming scattering in a medium with a uniform phase function.
+		// Since this is ran once per scene, we do not have access to view data (VIEWDATA_AVAILABLE==0). 
+		// So this buffer is not pre-exposed today.
+	}
+}
+
+
+
+
+
+RWTexture3D<float4> CameraAerialPerspectiveVolumeUAV;
+
+[numthreads(THREADGROUP_SIZE, THREADGROUP_SIZE, THREADGROUP_SIZE)]
+void RenderCameraAerialPerspectiveVolumeCS(uint3 ThreadId : SV_DispatchThreadID)
+{
+	float2 PixPos = float2(ThreadId.xy) + 0.5f;
+	float2 UV = PixPos  * SkyAtmosphere.CameraAerialPerspectiveVolumeSizeAndInvSize.zw;
+
+	float4 SVPos = float4(UV*View.ViewSizeAndInvSize.xy, 0.0f, 1.0f);// SV_POS as if resolution was the one from the scene view.
+	float3 WorldDir = GetScreenWorldDir(SVPos);
+	float3 CamPos = GetCameraPlanetPos();
+
+	float Slice = ((float(ThreadId.z) + 0.5f) * SkyAtmosphere.CameraAerialPerspectiveVolumeDepthResolutionInv); // +0.5 to always have a distance to integrate over
+	Slice *= Slice;	// squared distribution
+	Slice *= SkyAtmosphere.CameraAerialPerspectiveVolumeDepthResolution;
+
+	float3 RayStartWorldPos = CamPos + SkyAtmosphere.CameraAerialPerspectiveVolumeStartDepth * WorldDir; // Offset according to start depth
+	float ViewHeight;
+
+
+	// Compute position from froxel information
+	float tMax = Slice * SkyAtmosphere.CameraAerialPerspectiveVolumeDepthSliceLength;
+	float3 VoxelWorldPos = RayStartWorldPos + tMax * WorldDir;
+	float VoxelHeight = length(VoxelWorldPos);
+
+	// Check if the voxel is under the horizon.
+	const float UnderGround = VoxelHeight < Atmosphere.BottomRadius;
+
+	// Check if the voxel is beind the planet (to next check for below the horizon case)
+	float3 CameraToVoxel = VoxelWorldPos - CamPos;
+	float CameraToVoxelLen = length(CameraToVoxel);
+	float3 CameraToVoxelDir = CameraToVoxel / CameraToVoxelLen;
+	float PlanetNearT = RaySphereIntersectNearest(CamPos, CameraToVoxelDir, float3(0, 0, 0), Atmosphere.BottomRadius);
+	bool BelowHorizon = PlanetNearT > 0.0f && CameraToVoxelLen > PlanetNearT;
+
+	if (BelowHorizon || UnderGround)
+	{
+		CamPos += normalize(CamPos) * 0.02f;	// TODO: investigate why we need this workaround. Without it, we get some bad color and flickering on the ground only (floating point issue with sphere intersection code?).
+
+		float3 VoxelWorldPosNorm = normalize(VoxelWorldPos);
+		float3 CamProjOnGround = normalize(CamPos) * Atmosphere.BottomRadius;
+		float3 VoxProjOnGround = VoxelWorldPosNorm * Atmosphere.BottomRadius;
+		float3 VoxelGroundToRayStart = CamPos - VoxProjOnGround;
+		if (BelowHorizon && dot(normalize(VoxelGroundToRayStart), VoxelWorldPosNorm) < 0.0001f)
+		{
+			// We are behing the sphere and the psehre normal is pointing away from V: we are below the horizon.
+			float3 MiddlePoint = 0.5f * (CamProjOnGround + VoxProjOnGround);
+			float MiddlePointHeight = length(MiddlePoint);
+
+			// Compute the new position to evaluate and store the value in the voxel.
+			// the position is the oposite side of the horizon point from the view point,
+			// The offset of 1.001f is needed to get matching colors and for the ray to not hit the earth again later due to floating point accuracy
+			float3 MiddlePointOnGround = normalize(MiddlePoint) * Atmosphere.BottomRadius;// *1.001f;
+			VoxelWorldPos = CamPos + 2.0f * (MiddlePointOnGround - CamPos);
+
+
+			//CameraAerialPerspectiveVolumeUAV[ThreadId] = float4(1, 0, 0, 0);
+			//return; // debug
+		}
+		else if (UnderGround)
+		{
+			//No obstruction from the planet, so use the point on the ground
+			VoxelWorldPos = normalize(VoxelWorldPos) * (Atmosphere.BottomRadius);
+			//VoxelWorldPos = CamPos + CameraToVoxelDir * PlanetNearT;		// better but gives visuall artefact as visible voxels on a simple plane at altitude 0
+
+			//CameraAerialPerspectiveVolumeUAV[ThreadId] = float4(0, 1, 0, 0);
+			//return; // debug
+		}
+		 
+		WorldDir = normalize(VoxelWorldPos - CamPos);
+		RayStartWorldPos = CamPos + SkyAtmosphere.CameraAerialPerspectiveVolumeStartDepth * WorldDir; // Offset according to start depth
+		tMax = length(VoxelWorldPos - RayStartWorldPos);
+	}
+	float tMaxMax = tMax;
+
+	// Move ray marching start up to top atmosphere.
+	ViewHeight = length(RayStartWorldPos);
+	if (ViewHeight >= Atmosphere.TopRadius)
+	{
+		float3 prevWorlPos = RayStartWorldPos;
+		if (!MoveToTopAtmosphere(RayStartWorldPos, WorldDir, Atmosphere.TopRadius))
+		{
+			// Ray is not intersecting the atmosphere
+			CameraAerialPerspectiveVolumeUAV[ThreadId] = float4(0.0f, 0.0f, 0.0f, 1.0f);
+			return;
+		}
+		float LengthToAtmosphere = length(prevWorlPos - RayStartWorldPos);
+		if (tMaxMax < LengthToAtmosphere)
+		{
+			// tMaxMax for this voxel is not within the planet atmosphere
+			CameraAerialPerspectiveVolumeUAV[ThreadId] = float4(0.0f, 0.0f, 0.0f, 1.0f);
+			return;
+		}
+		// Now world position has been moved to the atmosphere boundary: we need to reduce tMaxMax accordingly. 
+		tMaxMax = max(0.0, tMaxMax - LengthToAtmosphere);
+	}
+
+
+	SamplingSetup Sampling;
+	{
+		Sampling.VariableSampleCount = false;
+		Sampling.SampleCountIni = max(1.0f, (float(ThreadId.z) + 1.0f) * SkyAtmosphere.CameraAerialPerspectiveSampleCountPerSlice);
+	}
+	const bool Ground = false;
+	const float DeviceZ = FarDepthValue;
+	const bool MieRayPhase = true;
+	const float AerialPespectiveViewDistanceScale = SkyAtmosphere.AerialPespectiveViewDistanceScale;
+	SingleScatteringResult ss = IntegrateSingleScatteredLuminance(
+		float4(PixPos, 0.0f, 1.0f), RayStartWorldPos, WorldDir,
+		Ground, Sampling, DeviceZ, MieRayPhase,
+		View.AtmosphereLightDirection[0].xyz, View.AtmosphereLightDirection[1].xyz, View.AtmosphereLightColor[0].rgb, View.AtmosphereLightColor[1].rgb,
+		AerialPespectiveViewDistanceScale,
+		tMaxMax);
+
+	const float Transmittance = dot(ss.Transmittance, float3(1.0f / 3.0f, 1.0f / 3.0f, 1.0f / 3.0f));
+	CameraAerialPerspectiveVolumeUAV[ThreadId] = float4(ss.L, Transmittance);
+}
+
+
+
+
+//////////////////////////////////////////////////////////////////////////
+// Debug shaders
+//////////////////////////////////////////////////////////////////////////
+
+float ViewPortWidth;
+float ViewPortHeight;
+
+float Mean3(float3 v)
+{
+	return dot(v, float3(1.0f / 3.0f, 1.0f / 3.0f, 1.0f / 3.0f));
+}
+
+void RenderSkyAtmosphereDebugPS(
+	in float4 SVPos : SV_POSITION,
+	out float4 OutLuminance : SV_Target0,
+	out float OutDepth : SV_DEPTH)
+{
+	float2 PixPos = SVPos.xy;
+
+	// Common ray marching input
+	SamplingSetup Sampling;
+	{
+		Sampling.VariableSampleCount = false;
+		Sampling.SampleCountIni = 128.0f;
+	}
+	const bool Ground = false;
+	const bool MieRayPhase = true;
+	float DeviceZ = FarDepthValue;
+	const float3 NullLightDirection = float3(0.0f, 0.0f, 1.0f);
+	const float3 NullLightIlluminance = float3(0.0f, 0.0f, 0.0f);
+	const float AerialPespectiveViewDistanceScale = 1.0f;
+
+	// We position the camera at the same height as the view port camera, along z-up vector
+	float3 WorldPos = float3(0.0f, 0.0f, length(View.WorldCameraOrigin*CM_TO_SKY_UNIT - float3(0.0f, 0.0f, -Atmosphere.BottomRadius)));
+
+	// Some organisation constants
+	const float Margin = 2.0f;
+	const float TimeOfDayHeight = 64.0f;
+	const float TimeOfDayTop = TimeOfDayHeight + Margin*2.0f;
+
+	// Hemisphere view, also show a bit of ground.
+	float2 DebugSize = ViewPortWidth * 0.25f;
+	float2 DebugPos = float2(2.0f, ViewPortHeight - DebugSize.y - TimeOfDayTop);
+	if(all(PixPos < (DebugPos+DebugSize)) && all(PixPos > DebugPos))
+	{
+		OutDepth = NearDepthValue;
+		float2 UV = clamp(PixPos - DebugPos, 0.0f, DebugSize) / DebugSize;
+
+		UV = UV * 2.0f - 1.0f;
+		float UVLen = length(UV);
+		if (UVLen < 1.0f)
+		{
+			float SinToHorizon = -2.0f * UVLen*UVLen + 1.0f;
+			float CosXY = cos(asin(SinToHorizon));
+			float3 WorldDir = normalize(float3(CosXY*normalize(UV), SinToHorizon));
+
+			SingleScatteringResult SS = IntegrateSingleScatteredLuminance(
+				SVPos, WorldPos, WorldDir,
+				Ground, Sampling, DeviceZ, MieRayPhase,
+				View.AtmosphereLightDirection[0].xyz, NullLightDirection, View.AtmosphereLightColor[0].rgb, NullLightIlluminance,
+				AerialPespectiveViewDistanceScale);
+			OutLuminance = float4(SS.L, 0.0f);
+			return;
+		}
+		OutLuminance = 0.0f;
+		return;
+	}
+
+	// All-at-once time of day visualization.
+	DebugSize = float2(1024.0f, TimeOfDayHeight);
+	DebugPos = float2(70.0f, ViewPortHeight - DebugSize.y - Margin);
+	// Adaptive TimeOfDay count
+	const float DesiredTimeOfDayWidth = 128.0f;
+	float NumTimeOfDay = floor((ViewPortWidth - DebugPos.x) / DesiredTimeOfDayWidth);
+	DebugSize.x = NumTimeOfDay * DesiredTimeOfDayWidth;
+	if (all(PixPos < (DebugPos + DebugSize)) && all(PixPos > DebugPos))
+	{
+		OutDepth = NearDepthValue;
+		float2 UV = clamp(PixPos - DebugPos, 0.0f, DebugSize) / DebugSize;
+
+		float CurrentTimeOfDay = floor(UV.x*NumTimeOfDay) / NumTimeOfDay;
+		UV.x = frac(UV.x*NumTimeOfDay);
+
+		float LightHorizonAngle = -0.05f + CurrentTimeOfDay * CurrentTimeOfDay * PI * 0.5f;
+		float3 LightDir = float3(cos(LightHorizonAngle), 0.0f, sin(LightHorizonAngle));
+
+		float ViewLightAngleCos = cos(PI * UV.x);
+		float ViewLightAngleSin = sqrt(1.0f - ViewLightAngleCos * ViewLightAngleCos);
+		float ViewHorizonCos = cos(0.5f * PI * (1.0f-UV.y));
+		float ViewHorizonSin = sqrt(1.0f - ViewHorizonCos * ViewHorizonCos);
+
+		float3 WorldDir = float3( 
+			ViewHorizonCos * ViewLightAngleCos,
+			ViewHorizonCos * ViewLightAngleSin,
+			ViewHorizonSin );
+
+		SingleScatteringResult SS = IntegrateSingleScatteredLuminance(
+			SVPos, WorldPos, WorldDir,
+			Ground, Sampling, DeviceZ, MieRayPhase,
+			LightDir, NullLightDirection, View.AtmosphereLightColor[0].rgb, NullLightIlluminance,
+			AerialPespectiveViewDistanceScale);
+
+		OutLuminance = float4(SS.L, 0.0f);
+		return;
+	}
+
+	// Atmosphere density distribution.
+	DebugSize = float2(ViewPortWidth * 0.2f, ViewPortHeight * 0.7f);
+	DebugPos = float2(ViewPortWidth - DebugSize.x - 2.0f, ViewPortHeight * 0.1f);
+	if (all(PixPos < (DebugPos + DebugSize)) && all(PixPos > DebugPos))
+	{
+		OutDepth = NearDepthValue;
+		float2 UV = clamp(PixPos - DebugPos, 0.0f, DebugSize) / DebugSize;
+		UV.y = 1.0f - UV.y;
+		UV = UV.yx;
+
+		float3 Color = float3(pow(1.0f-UV, 200.0f), 0.0f);
+
+		const float TotMax = 0.01f;	// Cosntant max extinction is less confusing
+
+
+
+		float3 SampleWorldPos = float3(0.0f, 0.0f, Atmosphere.BottomRadius + UV.x * (Atmosphere.TopRadius - Atmosphere.BottomRadius));
+		MediumSampleRGB Sample = SampleMediumRGB(SampleWorldPos);
+
+		// Curves represent the maximum contribution as extinction (no differenciation between scattering or absorption)
+		float3 Curves = float3(
+			Mean3(Sample.ExtinctionRay) / TotMax,
+			Mean3(Sample.ExtinctionMie) / TotMax,
+			Mean3(Sample.ExtinctionOzo) / TotMax
+			);
+
+		Curves = saturate((Curves - UV.y)*100000.0f);
+
+		OutLuminance = float4(Curves, 0.0f); // No exposure for constant visual
+		return;
+	}
+
+	// Brown ground
+	DebugSize.y += 15.0f;
+	if (all(PixPos < (DebugPos + DebugSize)) && all(PixPos > DebugPos))
+	{
+		OutDepth = NearDepthValue;
+		OutLuminance = float4(float3(0.125f, 0.05f, 0.005f), 0.0f); // No exposure for constant visual
+		return;
+	}
+
+	clip(-1.0f);
+	OutLuminance = 0.0f;
+	return;
+}
+
diff --git a/Engine/Shaders/Public/ShaderVersion.ush b/Engine/Shaders/Public/ShaderVersion.ush
--- a/Engine/Shaders/Public/ShaderVersion.ush
+++ b/Engine/Shaders/Public/ShaderVersion.ush
@@ -3,6 +3,6 @@
 // in Platform.ush (which should be included in any shader) it allows to invalidate the shader DDC.
 // 
 // If you are merging streams and there is a conflict with this GUID you should make a new GUID rather than taking one or the other.
-// GUID = 420AAFE740634C17B1A9AB0E616A8105
+// GUID = 69f74d6AA9DA4A1DA3BFFA82DBA76DBA
 
 
diff --git a/Engine/Source/Editor/UnrealEd/Private/LevelEditorViewport.cpp b/Engine/Source/Editor/UnrealEd/Private/LevelEditorViewport.cpp
--- a/Engine/Source/Editor/UnrealEd/Private/LevelEditorViewport.cpp
+++ b/Engine/Source/Editor/UnrealEd/Private/LevelEditorViewport.cpp
@@ -2683,13 +2683,21 @@ bool FLevelEditorViewportClient::InputKey(FViewport* InViewport, int32 Controlle
 		return true;
 	}
 
-	if (InputState.IsCtrlButtonPressed() && Key == EKeys::L)
+	bool bCtrlLPressed = InputState.IsCtrlButtonPressed() && Key == EKeys::L;
+	if (bCtrlLPressed && InputState.IsShiftButtonPressed())
 	{
-		bUserIsControllingSunLight = true;
+		bUserIsControllingSunLight1 = true;
 		UserIsControllingSunLightTimer = 3.0f; // Keep the widget open for a few seconds even when not tweaking the sun light
 		return true;
 	}
-	bUserIsControllingSunLight = false;
+	if (bCtrlLPressed)
+	{
+		bUserIsControllingSunLight0 = true;
+		UserIsControllingSunLightTimer = 3.0f; // Keep the widget open for a few seconds even when not tweaking the sun light
+		return true;
+	}
+	bUserIsControllingSunLight0 = false;
+	bUserIsControllingSunLight1 = false;
 
 	bool bHandled = FEditorViewportClient::InputKey(InViewport,ControllerId,Key,Event,AmountDepressed,bGamepad);
 
@@ -4020,37 +4028,52 @@ EMouseCursor::Type FLevelEditorViewportClient::GetCursor(FViewport* InViewport,i
 
 void FLevelEditorViewportClient::MouseMove(FViewport* InViewport, int32 x, int32 y)
 {
-	if (bUserIsControllingSunLight)
+	if (bUserIsControllingSunLight0 || bUserIsControllingSunLight1)
 	{
 		UWorld* ViewportWorld = GetWorld();
+
+		UDirectionalLightComponent* SelectedSunLight = nullptr;
+		float SelectedLightLuminance = 0.0f;
+		const uint8 DesiredLightIndex = bUserIsControllingSunLight0 ? 0 : 1;
 		for (TObjectIterator<UDirectionalLightComponent> ComponentIt; ComponentIt; ++ComponentIt)
 		{
 			if (ComponentIt->GetWorld() == ViewportWorld)
 			{
 				UDirectionalLightComponent* SunLight = *ComponentIt;
 
-				int32 mouseDeltaX = x - CachedLastMouseX;
-				int32 mouseDeltaY = y - CachedLastMouseY;
-				if (!SunLight->IsUsedAsAtmosphereSunLight() || !SunLight->bVisible)
+				if (!SunLight->IsUsedAsAtmosphereSunLight() || SunLight->GetAtmosphereSunLightIndex()!= DesiredLightIndex || !SunLight->bVisible)
 					continue;
 
-				FTransform ComponentTransform = SunLight->GetComponentTransform();
-				FQuat SunRotation = ComponentTransform.GetRotation();
-				// Rotate around up axis (yaw)
-				FVector UpVector = FVector(0, 0, 1);
-				SunRotation = FQuat(UpVector, float(mouseDeltaX)*0.01f) * SunRotation;
-				// Sun Zenith rotation (pitch)
-				FVector PitchRotationAxis = FVector::CrossProduct(SunRotation.GetForwardVector(), UpVector);
-				PitchRotationAxis.Normalize();
-				SunRotation = FQuat(PitchRotationAxis, float(mouseDeltaY)*0.01f) * SunRotation;
-
-				ComponentTransform.SetRotation(SunRotation);
-				SunLight->SetWorldTransform(ComponentTransform);
-
-				// Only manipulate a single light, the first one.
-				UserControlledSunLightMatrix = ComponentTransform;
+				float LightLuminance = SunLight->GetColoredLightBrightness().ComputeLuminance();
+				if (!SelectedSunLight ||							// Set it if null
+					SelectedLightLuminance < LightLuminance)		// Or choose the brightest sun light
+				{
+					SelectedSunLight = SunLight;
+				}
 			}
 		}
+
+		if (SelectedSunLight)
+		{
+			int32 mouseDeltaX = x - CachedLastMouseX;
+			int32 mouseDeltaY = y - CachedLastMouseY;
+
+			FTransform ComponentTransform = SelectedSunLight->GetComponentTransform();
+			FQuat SunRotation = ComponentTransform.GetRotation();
+			// Rotate around up axis (yaw)
+			FVector UpVector = FVector(0, 0, 1);
+			SunRotation = FQuat(UpVector, float(mouseDeltaX)*0.01f) * SunRotation;
+			// Sun Zenith rotation (pitch)
+			FVector PitchRotationAxis = FVector::CrossProduct(SunRotation.GetForwardVector(), UpVector);
+			PitchRotationAxis.Normalize();
+			SunRotation = FQuat(PitchRotationAxis, float(mouseDeltaY)*0.01f) * SunRotation;
+
+			ComponentTransform.SetRotation(SunRotation);
+			SelectedSunLight->SetWorldTransform(ComponentTransform);
+
+			// Only manipulate a single light, the first one.
+			UserControlledSunLightMatrix = ComponentTransform;
+		}
 	}
 
 	FEditorViewportClient::MouseMove(InViewport, x, y);
diff --git a/Engine/Source/Editor/UnrealEd/Private/Lightmass/Lightmass.cpp b/Engine/Source/Editor/UnrealEd/Private/Lightmass/Lightmass.cpp
--- a/Engine/Source/Editor/UnrealEd/Private/Lightmass/Lightmass.cpp
+++ b/Engine/Source/Editor/UnrealEd/Private/Lightmass/Lightmass.cpp
@@ -29,6 +29,7 @@
 #include "Components/DirectionalLightComponent.h"
 #include "Renderer/Private/AtmosphereRendering.h"
 #include "Components/SkyLightComponent.h"
+#include "Rendering/SkyAtmosphereCommonData.h"
 #include "Components/ModelComponent.h"
 #include "Materials/MaterialInstanceConstant.h"
 #include "EngineUtils.h"
@@ -488,6 +489,7 @@ void FLightmassProcessor::SwarmCallback( NSwarm::FMessage* CallbackMessage, void
 FLightmassExporter::FLightmassExporter( UWorld* InWorld )
 	: Swarm( NSwarm::FSwarmInterface::Get() ) 
 	, AtmosphericFogComponent(nullptr)
+	, SkyAtmosphereComponent(nullptr)
 	, ExportStage(NotRunning)
 	, CurrentAmortizationIndex(0)
 	, OpenedMaterialExportChannels()
@@ -1025,24 +1027,45 @@ void FLightmassExporter::WriteLights( int32 Channel )
 {
 	// Search for the sun light component the same way as in FScene::RemoveLightSceneInfo_RenderThread.
 	// If found, we keep the pointer to apply atmosphere transmittance to it in the light loop.
-	const UDirectionalLightComponent* SunLight = nullptr;
-	FLinearColor SunLightAtmosphereTransmittance(FLinearColor::White);
-	if (AtmosphericFogComponent && AtmosphericFogComponent->bAtmosphereAffectsSunIlluminance)
+	const UDirectionalLightComponent* AtmosphereLights[NUM_ATMOSPHERE_LIGHTS];
+	float AtmosphereLightsBrightness[NUM_ATMOSPHERE_LIGHTS];
+	FLinearColor SunLightAtmosphereTransmittance[NUM_ATMOSPHERE_LIGHTS];
+	for (uint8 Index = 0; Index < NUM_ATMOSPHERE_LIGHTS; ++Index)
+	{
+		AtmosphereLights[Index] = nullptr;
+		AtmosphereLightsBrightness[Index] = 0.0f;
+		SunLightAtmosphereTransmittance[Index] = FLinearColor::White;
+	}
+
+	// Compute a mapping between directional light and trnasmittance to apply. For each AtmosphereSunLightIndex, the brightest lights is kept.
+	if (AtmosphericFogComponent && AtmosphericFogComponent->bAtmosphereAffectsSunIlluminance || SkyAtmosphereComponent)
 	{
-		float SunLightEnergy = 0.0f;
 		for (int32 LightIndex = 0; LightIndex < DirectionalLights.Num(); ++LightIndex)
 		{
 			const UDirectionalLightComponent* Light = DirectionalLights[LightIndex];
-			if (Light->IsUsedAsAtmosphereSunLight() && Light->GetColoredLightBrightness().ComputeLuminance() > SunLightEnergy)
+			if (Light->IsUsedAsAtmosphereSunLight() && Light->GetColoredLightBrightness().ComputeLuminance() > AtmosphereLightsBrightness[Light->GetAtmosphereSunLightIndex()])
 			{
-				SunLight = Light;
-				SunLightEnergy = SunLight->GetColoredLightBrightness().ComputeLuminance();
+				AtmosphereLights[Light->GetAtmosphereSunLightIndex()] = Light;
+				AtmosphereLightsBrightness[Light->GetAtmosphereSunLightIndex()] = Light->GetColoredLightBrightness().ComputeLuminance();
 			}
 		}
 
-		if (SunLight)
+		for (uint8 Index = 0; Index < NUM_ATMOSPHERE_LIGHTS; ++Index)
 		{
-			SunLightAtmosphereTransmittance = AtmosphericFogComponent->GetTransmittance(-SunLight->GetDirection()); 
+			if (AtmosphereLights[Index])
+			{
+				const FVector SunDirection = -AtmosphereLights[Index]->GetDirection();
+				if (SkyAtmosphereComponent)
+				{
+					FAtmosphereSetup AtmosphereSetup(*SkyAtmosphereComponent);
+					SunLightAtmosphereTransmittance[Index] = AtmosphereSetup.GetTransmittanceAtGroundLevel(SunDirection);
+				}
+				else if (AtmosphericFogComponent && Index==0)	// Legacy AtmosphericFogComponent only takes into account Index 0
+				{
+					SunLightAtmosphereTransmittance[Index] = AtmosphericFogComponent->GetTransmittance(SunDirection);
+				}
+				break;
+			}
 		}
 	}
 
@@ -1059,9 +1082,14 @@ void FLightmassExporter::WriteLights( int32 Channel )
 		LightData.LightSourceRadius = 0;
 		LightData.LightSourceLength = 0;
 
-		if (Light == SunLight)
+		// Apply atmosphere transmittance if needed.
+		for (uint8 Index = 0; Index < NUM_ATMOSPHERE_LIGHTS; ++Index)
 		{
-			LightData.Color *= SunLightAtmosphereTransmittance;
+			if (AtmosphereLights[Index] && AtmosphereLights[Index] == Light)
+			{
+				LightData.Color *= SunLightAtmosphereTransmittance[Index];
+				break;
+			}
 		}
 
 		TArray< uint8 > LightProfileTextureData;
diff --git a/Engine/Source/Editor/UnrealEd/Private/Lightmass/Lightmass.h b/Engine/Source/Editor/UnrealEd/Private/Lightmass/Lightmass.h
--- a/Engine/Source/Editor/UnrealEd/Private/Lightmass/Lightmass.h
+++ b/Engine/Source/Editor/UnrealEd/Private/Lightmass/Lightmass.h
@@ -15,6 +15,7 @@
 #include "Lightmass/LightmassImportanceVolume.h"
 #include "Components/LightmassPortalComponent.h"
 #include "Atmosphere/AtmosphericFogComponent.h"
+#include "Components/SkyAtmosphereComponent.h"
 #if PLATFORM_WINDOWS
 #include "Windows/AllowWindowsPlatformTypes.h"
 #endif
@@ -120,6 +121,11 @@ public:
 		AtmosphericFogComponent = AtmosphericFog;
 	}
 
+	void SetSkyAtmosphereComponent(const USkyAtmosphereComponent* SkyAtmosphere)
+	{
+		SkyAtmosphereComponent = SkyAtmosphere;
+	}
+
 	// if provided, InStaticLightingMesh is used to UV unwrap the material into the static lighting textures
 	void AddMaterial(UMaterialInterface* InMaterialInterface, const FStaticLightingMesh* InStaticLightingMesh = nullptr);
 
@@ -206,6 +212,7 @@ private:
 	TArray<FMatrix>				Portals;
 
 	const UAtmosphericFogComponent*	AtmosphericFogComponent;
+	const USkyAtmosphereComponent*	SkyAtmosphereComponent;
 
 	FLightmassWorldInfoSettings LevelSettings;
 	/** The number of local cores to leave unused */
diff --git a/Engine/Source/Editor/UnrealEd/Private/StaticLightingSystem/StaticLightingSystem.cpp b/Engine/Source/Editor/UnrealEd/Private/StaticLightingSystem/StaticLightingSystem.cpp
--- a/Engine/Source/Editor/UnrealEd/Private/StaticLightingSystem/StaticLightingSystem.cpp
+++ b/Engine/Source/Editor/UnrealEd/Private/StaticLightingSystem/StaticLightingSystem.cpp
@@ -41,6 +41,7 @@
 #include "Engine/GeneratedMeshAreaLight.h"
 #include "Components/SkyLightComponent.h"
 #include "Atmosphere/AtmosphericFogComponent.h"
+#include "Components/SkyAtmosphereComponent.h"
 #include "Components/ModelComponent.h"
 #include "Engine/LightMapTexture2D.h"
 #include "Editor.h"
@@ -773,6 +774,9 @@ void FStaticLightingSystem::InvalidateStaticLighting()
 			continue;
 		}
 
+		// Clear all the atmosphere guids from the MapBuildData when starting a new build.
+		Level->GetOrCreateMapBuildData()->ClearSkyAtmosphereBuildData();
+
 		const bool bBuildLightingForLevel = Options.ShouldBuildLightingForLevel( Level );
 		
 		if (bBuildLightingForLevel)
@@ -1260,12 +1264,12 @@ void FStaticLightingSystem::ApplyNewLightingData(bool bLightingSuccessful)
 
 				if (Actor && bLightingSuccessful && !Options.bOnlyBuildSelected)
 				{
-					TInlineComponentArray<ULightComponent*> Components;
-					Actor->GetComponents(Components);
+					TInlineComponentArray<ULightComponent*> LightComponents;
+					Actor->GetComponents(LightComponents);
 
-					for (int32 ComponentIndex = 0; ComponentIndex < Components.Num(); ComponentIndex++)
+					for (int32 ComponentIndex = 0; ComponentIndex < LightComponents.Num(); ComponentIndex++)
 					{
-						ULightComponent* LightComponent = Components[ComponentIndex];
+						ULightComponent* LightComponent = LightComponents[ComponentIndex];
 						if (LightComponent && (LightComponent->HasStaticShadowing() || LightComponent->HasStaticLighting()))
 						{
 							if (!Registry->GetLightBuildData(LightComponent->LightGuid))
@@ -1275,6 +1279,21 @@ void FStaticLightingSystem::ApplyNewLightingData(bool bLightingSuccessful)
 							}
 						}
 					}
+
+					// For each SkyAtmosphere which is a dependency of the build, add its guid to MapBuildData to track that it now has been built.
+					TInlineComponentArray<USkyAtmosphereComponent*> SkyAtmosphereComponents;
+					Actor->GetComponents(SkyAtmosphereComponents);
+					for (int32 ComponentIndex = 0; ComponentIndex < SkyAtmosphereComponents.Num(); ComponentIndex++)
+					{
+						USkyAtmosphereComponent* SkyAtmosphereComponent = SkyAtmosphereComponents[ComponentIndex];
+						if (SkyAtmosphereComponent)
+						{
+							if (!Registry->GetSkyAtmosphereBuildData(SkyAtmosphereComponent->GetStaticLightingBuiltGuid()))
+							{
+								Registry->FindOrAllocateSkyAtmosphereBuildData(SkyAtmosphereComponent->GetStaticLightingBuiltGuid());
+							}
+						}
+					}
 				}
 			}
 
@@ -2026,12 +2045,29 @@ void FStaticLightingSystem::GatherScene()
 		}
 	}
 
+	bool LegacyAtmosphericFogRegistered = false;
 	for (TObjectIterator<UAtmosphericFogComponent> It; It; ++It)
 	{
 		UAtmosphericFogComponent* AtmosphericFog = *It;
 		if (AtmosphericFog->GetOwner() && World->ContainsActor(AtmosphericFog->GetOwner()) && !AtmosphericFog->IsPendingKill() && ShouldOperateOnLevel(AtmosphericFog->GetOwner()->GetLevel()))
 		{
 			LightmassExporter->SetAtmosphericComponent(AtmosphericFog);
+			LegacyAtmosphericFogRegistered = true;
+			break;	// We only register the first we find
+		}
+	}
+
+	for (TObjectIterator<USkyAtmosphereComponent> It; It; ++It)
+	{
+		USkyAtmosphereComponent* SkyAtmosphere = *It;
+		if (SkyAtmosphere->GetOwner() && World->ContainsActor(SkyAtmosphere->GetOwner()) && !SkyAtmosphere->IsPendingKill() && ShouldOperateOnLevel(SkyAtmosphere->GetOwner()->GetLevel()))
+		{
+			if (LegacyAtmosphericFogRegistered)
+			{
+				FMessageLog("LightingResults").Warning(LOCTEXT("LightmassError_BuildSelected", "Both a legacy AtmosphericFog and a new SkyAtmosphere wants to register. Lightmass will not consider the legacy component."));
+			}
+			LightmassExporter->SetAtmosphericComponent(nullptr);
+			LightmassExporter->SetSkyAtmosphereComponent(SkyAtmosphere);
 			break;	// We only register the first we find
 		}
 	}
diff --git a/Engine/Source/Editor/UnrealEd/Public/EditorViewportClient.h b/Engine/Source/Editor/UnrealEd/Public/EditorViewportClient.h
--- a/Engine/Source/Editor/UnrealEd/Public/EditorViewportClient.h
+++ b/Engine/Source/Editor/UnrealEd/Public/EditorViewportClient.h
@@ -1541,8 +1541,9 @@ protected:
 	int32 CachedLastMouseX = 0;
 	int32 CachedLastMouseY = 0;
 
-	/** True is the use is controling the light via a*/
-	bool bUserIsControllingSunLight = false;
+	/** True is the use is controling the light via a shorcut*/
+	bool bUserIsControllingSunLight0 = false;
+	bool bUserIsControllingSunLight1 = false;
 	float UserIsControllingSunLightTimer = 0.0f;
 	FTransform UserControlledSunLightMatrix;
 
diff --git a/Engine/Source/Runtime/Core/Private/Misc/MapErrors.cpp b/Engine/Source/Runtime/Core/Private/Misc/MapErrors.cpp
--- a/Engine/Source/Runtime/Core/Private/Misc/MapErrors.cpp
+++ b/Engine/Source/Runtime/Core/Private/Misc/MapErrors.cpp
@@ -16,6 +16,8 @@ FName FMapErrors::VertexColorsNotMatchOriginalMesh(TEXT("VertexColorsNotMatchOri
 FName FMapErrors::CollisionEnabledNoCollisionGeom(TEXT("CollisionEnabledNoCollisionGeom"));
 FName FMapErrors::ShadowCasterUsingBoundsScale(TEXT("ShadowCasterUsingBoundsScale"));
 FName FMapErrors::MultipleSkyLights(TEXT("MultipleSkyLights"));
+FName FMapErrors::MultipleSkyAtmospheres(TEXT("MultipleSkyAtmospheres"));
+FName FMapErrors::MultipleSkyAtmosphereTypes(TEXT("MultipleSkyAtmosphereTypes"));
 FName FMapErrors::InvalidTrace(TEXT("InvalidTrace"));
 FName FMapErrors::BrushZeroPolygons(TEXT("BrushZeroPolygons"));
 FName FMapErrors::CleanBSPMaterials(TEXT("CleanBSPMaterials"));
diff --git a/Engine/Source/Runtime/Core/Public/Misc/MapErrors.h b/Engine/Source/Runtime/Core/Public/Misc/MapErrors.h
--- a/Engine/Source/Runtime/Core/Public/Misc/MapErrors.h
+++ b/Engine/Source/Runtime/Core/Public/Misc/MapErrors.h
@@ -58,6 +58,9 @@ struct CORE_API FMapErrors
     
 	static FName MultipleSkyLights;
 
+	static FName MultipleSkyAtmospheres;
+	static FName MultipleSkyAtmosphereTypes;
+
 	/** {ActorName} has WorldTrace blocked. It will be considered to be world geometry */
 	static FName InvalidTrace;
 
diff --git a/Engine/Source/Runtime/Core/Public/UObject/RenderingObjectVersion.h b/Engine/Source/Runtime/Core/Public/UObject/RenderingObjectVersion.h
--- a/Engine/Source/Runtime/Core/Public/UObject/RenderingObjectVersion.h
+++ b/Engine/Source/Runtime/Core/Public/UObject/RenderingObjectVersion.h
@@ -90,6 +90,8 @@ struct CORE_API FRenderingObjectVersion
 		// UTextureRenderTarget2D now explicitly allows users to create sRGB or non-sRGB type targets
 		ExplicitSRGBSetting,
 
+		SkyAtmosphereStaticLightingVersioning,
+
 		// -----<new versions can be added above this line>-------------------------------------------------
 		VersionPlusOne,
 		LatestVersion = VersionPlusOne - 1
diff --git a/Engine/Source/Runtime/Engine/Classes/Components/DirectionalLightComponent.h b/Engine/Source/Runtime/Engine/Classes/Components/DirectionalLightComponent.h
--- a/Engine/Source/Runtime/Engine/Classes/Components/DirectionalLightComponent.h
+++ b/Engine/Source/Runtime/Engine/Classes/Components/DirectionalLightComponent.h
@@ -159,6 +159,9 @@ class ENGINE_API UDirectionalLightComponent : public ULightComponent
 	UPROPERTY(EditAnywhere, AdvancedDisplay, BlueprintReadOnly, Category=Light, meta=(DisplayName = "Atmosphere / Fog Sun Light"))
 	uint32 bUsedAsAtmosphereSunLight : 1;
 
+	UPROPERTY(EditAnywhere, AdvancedDisplay, BlueprintReadOnly, Category = Light, meta = (DisplayName = "Atmosphere Sun Light Index", UIMin = "0", UIMax = "1", ClampMin = "0", ClampMax= "1"))
+	int32 AtmosphereSunLightIndex;
+
 	UFUNCTION(BlueprintCallable, Category="Rendering|Lighting")
 	void SetDynamicShadowDistanceMovableLight(float NewValue);
 
@@ -201,6 +204,10 @@ class ENGINE_API UDirectionalLightComponent : public ULightComponent
 	{
 		return bUsedAsAtmosphereSunLight;
 	}
+	virtual uint8 GetAtmosphereSunLightIndex() const override
+	{
+		return AtmosphereSunLightIndex;
+	}
 	//~ End ULightComponent Interface
 
 	//~ Begin UObject Interface
diff --git a/Engine/Source/Runtime/Engine/Classes/Components/LightComponent.h b/Engine/Source/Runtime/Engine/Classes/Components/LightComponent.h
--- a/Engine/Source/Runtime/Engine/Classes/Components/LightComponent.h
+++ b/Engine/Source/Runtime/Engine/Classes/Components/LightComponent.h
@@ -360,6 +360,10 @@ public:
 	{
 		return false;
 	}
+	virtual uint8 GetAtmosphereSunLightIndex() const
+	{
+		return 0;
+	}
 
 	/** Compute current light brightness based on whether there is a valid IES profile texture attached, and whether IES brightness is enabled */
 	virtual float ComputeLightBrightness() const;
