Fix uniform buffer with same name.
One of them had to go anyway but that never asserted before...
Tested on win64 and mobile es3.1 using editor settings.

diff --git a/Engine/Shaders/Private/BasePassCommon.ush b/Engine/Shaders/Private/BasePassCommon.ush
--- a/Engine/Shaders/Private/BasePassCommon.ush
+++ b/Engine/Shaders/Private/BasePassCommon.ush
@@ -11,28 +11,12 @@
 	#define PlanarReflectionStruct TranslucentBasePass.Shared.PlanarReflection
 	#define FogStruct TranslucentBasePass.Shared.Fog
 	#define ActualSSProfilesTexture TranslucentBasePass.Shared.SSProfilesTexture
-	#define ApplyCameraAerialPerspectiveVolume TranslucentBasePass.Shared.SkyAtmosphere.ApplyCameraAerialPerspectiveVolume
-	#define CameraAerialPerspectiveVolume TranslucentBasePass.Shared.SkyAtmosphere.CameraAerialPerspectiveVolume
-	#define CameraAerialPerspectiveVolumeSampler TranslucentBasePass.Shared.SkyAtmosphere.CameraAerialPerspectiveVolumeSampler
-	#define CameraAerialPerspectiveVolumeStartDepth TranslucentBasePass.Shared.SkyAtmosphere.CameraAerialPerspectiveVolumeStartDepth
-	#define CameraAerialPerspectiveVolumeDepthResolution TranslucentBasePass.Shared.SkyAtmosphere.CameraAerialPerspectiveVolumeDepthResolution
-	#define CameraAerialPerspectiveVolumeDepthResolutionInv TranslucentBasePass.Shared.SkyAtmosphere.CameraAerialPerspectiveVolumeDepthResolutionInv
-	#define CameraAerialPerspectiveVolumeDepthSliceLength TranslucentBasePass.Shared.SkyAtmosphere.CameraAerialPerspectiveVolumeDepthSliceLength
-	#define CameraAerialPerspectiveVolumeDepthSliceLengthInv TranslucentBasePass.Shared.SkyAtmosphere.CameraAerialPerspectiveVolumeDepthSliceLengthInv
 #else
 	#define ForwardLightData OpaqueBasePass.Shared.Forward
 	#define ReflectionStruct OpaqueBasePass.Shared.Reflection
 	#define PlanarReflectionStruct OpaqueBasePass.Shared.PlanarReflection
 	#define FogStruct OpaqueBasePass.Shared.Fog
 	#define ActualSSProfilesTexture OpaqueBasePass.Shared.SSProfilesTexture
-	#define ApplyCameraAerialPerspectiveVolume OpaqueBasePass.Shared.SkyAtmosphere.ApplyCameraAerialPerspectiveVolume
-	#define CameraAerialPerspectiveVolume OpaqueBasePass.Shared.SkyAtmosphere.CameraAerialPerspectiveVolume
-	#define CameraAerialPerspectiveVolumeSampler OpaqueBasePass.Shared.SkyAtmosphere.CameraAerialPerspectiveVolumeSampler
-	#define CameraAerialPerspectiveVolumeStartDepth OpaqueBasePass.Shared.SkyAtmosphere.CameraAerialPerspectiveVolumeStartDepth
-	#define CameraAerialPerspectiveVolumeDepthResolution OpaqueBasePass.Shared.SkyAtmosphere.CameraAerialPerspectiveVolumeDepthResolution
-	#define CameraAerialPerspectiveVolumeDepthResolutionInv OpaqueBasePass.Shared.SkyAtmosphere.CameraAerialPerspectiveVolumeDepthResolutionInv
-	#define CameraAerialPerspectiveVolumeDepthSliceLength OpaqueBasePass.Shared.SkyAtmosphere.CameraAerialPerspectiveVolumeDepthSliceLength
-	#define CameraAerialPerspectiveVolumeDepthSliceLengthInv OpaqueBasePass.Shared.SkyAtmosphere.CameraAerialPerspectiveVolumeDepthSliceLengthInv
 #endif 
 
 #undef NEEDS_LIGHTMAP_COORDINATE
diff --git a/Engine/Shaders/Private/BasePassPixelShader.usf b/Engine/Shaders/Private/BasePassPixelShader.usf
--- a/Engine/Shaders/Private/BasePassPixelShader.usf
+++ b/Engine/Shaders/Private/BasePassPixelShader.usf
@@ -1130,7 +1130,7 @@ ShadingModel);
 #endif
 
 #if NEEDS_BASEPASS_PIXEL_FOGGING
-	if (ApplyCameraAerialPerspectiveVolume > 0.0f)
+	if (ResolvedView.SkyAtmosphereApplyCameraAerialPerspectiveVolume > 0.0f)
 	{
 		const float OneOverPreExposure = USE_PREEXPOSURE ? ResolvedView.OneOverPreExposure : 1.0f;
 
@@ -1139,12 +1139,13 @@ ShadingModel);
 		// Sample the aerial perspective (AP).
 		float4 AP = GetAerialPerspectiveLuminanceTransmittance(
 			NDCPosition, MaterialParameters.AbsoluteWorldPosition.xyz*CM_TO_SKY_UNIT, ResolvedView.WorldCameraOrigin.xyz*CM_TO_SKY_UNIT,
-			CameraAerialPerspectiveVolume, CameraAerialPerspectiveVolumeSampler,
-			CameraAerialPerspectiveVolumeDepthResolutionInv,
-			CameraAerialPerspectiveVolumeDepthResolution,
-			CameraAerialPerspectiveVolumeStartDepth,
-			CameraAerialPerspectiveVolumeDepthSliceLength,
-			CameraAerialPerspectiveVolumeDepthSliceLengthInv, 
+			View.CameraAerialPerspectiveVolume, View.CameraAerialPerspectiveVolumeSampler,
+			ResolvedView.SkyAtmosphereCameraAerialPerspectiveVolumeDepthResolutionInv,
+			ResolvedView.SkyAtmosphereCameraAerialPerspectiveVolumeDepthResolutionInv,
+			ResolvedView.SkyAtmosphereCameraAerialPerspectiveVolumeDepthResolution,
+			ResolvedView.SkyAtmosphereCameraAerialPerspectiveVolumeStartDepth,
+			ResolvedView.SkyAtmosphereCameraAerialPerspectiveVolumeDepthSliceLength,
+			ResolvedView.SkyAtmosphereCameraAerialPerspectiveVolumeDepthSliceLengthInv, 
 			OneOverPreExposure, Fogging);
 	}
 #endif
diff --git a/Engine/Shaders/Private/BasePassVertexShader.usf b/Engine/Shaders/Private/BasePassVertexShader.usf
--- a/Engine/Shaders/Private/BasePassVertexShader.usf
+++ b/Engine/Shaders/Private/BasePassVertexShader.usf
@@ -128,18 +128,18 @@ void Main(
 	Output.BasePassInterpolants.VertexFog = CalculateHeightFog(WorldPosition.xyz - ResolvedView.TranslatedWorldCameraOrigin);
 	#endif
 
-	if (ApplyCameraAerialPerspectiveVolume > 0.0f)
+	if (ResolvedView.SkyAtmosphereApplyCameraAerialPerspectiveVolume > 0.0f)
 	{
 		const float OneOverPreExposure = USE_PREEXPOSURE ? ResolvedView.OneOverPreExposure : 1.0f;
 		// Sample the aerial perspective (AP). It is also blended under the VertexFog parameter.
 		float4 AP = GetAerialPerspectiveLuminanceTransmittance(
 			Output.Position, WorldPosition.xyz*CM_TO_SKY_UNIT, ResolvedView.TranslatedWorldCameraOrigin*CM_TO_SKY_UNIT,
-			CameraAerialPerspectiveVolume, CameraAerialPerspectiveVolumeSampler,
-			CameraAerialPerspectiveVolumeDepthResolutionInv,
-			CameraAerialPerspectiveVolumeDepthResolution,
-			CameraAerialPerspectiveVolumeStartDepth,
-			CameraAerialPerspectiveVolumeDepthSliceLength,
-			CameraAerialPerspectiveVolumeDepthSliceLengthInv,
+			View.CameraAerialPerspectiveVolume, View.CameraAerialPerspectiveVolumeSampler,
+			ResolvedView.SkyAtmosphereCameraAerialPerspectiveVolumeDepthResolutionInv,
+			ResolvedView.SkyAtmosphereCameraAerialPerspectiveVolumeDepthResolution,
+			ResolvedView.SkyAtmosphereCameraAerialPerspectiveVolumeStartDepth,
+			ResolvedView.SkyAtmosphereCameraAerialPerspectiveVolumeDepthSliceLength,
+			ResolvedView.SkyAtmosphereCameraAerialPerspectiveVolumeDepthSliceLengthInv,
 			OneOverPreExposure, Output.BasePassInterpolants.VertexFog);
 	}
 #endif
diff --git a/Engine/Shaders/Private/MobileBasePassCommon.ush b/Engine/Shaders/Private/MobileBasePassCommon.ush
--- a/Engine/Shaders/Private/MobileBasePassCommon.ush
+++ b/Engine/Shaders/Private/MobileBasePassCommon.ush
@@ -6,14 +6,6 @@
 
 #define FogStruct MobileBasePass.Fog
 #define PlanarReflectionStruct MobileBasePass.PlanarReflection
-#define ApplyCameraAerialPerspectiveVolume MobileBasePass.SkyAtmosphere.ApplyCameraAerialPerspectiveVolume
-#define CameraAerialPerspectiveVolume MobileBasePass.SkyAtmosphere.CameraAerialPerspectiveVolume
-#define CameraAerialPerspectiveVolumeSampler MobileBasePass.SkyAtmosphere.CameraAerialPerspectiveVolumeSampler
-#define CameraAerialPerspectiveVolumeStartDepth MobileBasePass.SkyAtmosphere.CameraAerialPerspectiveVolumeStartDepth
-#define CameraAerialPerspectiveVolumeDepthResolution MobileBasePass.SkyAtmosphere.CameraAerialPerspectiveVolumeDepthResolution
-#define CameraAerialPerspectiveVolumeDepthResolutionInv MobileBasePass.SkyAtmosphere.CameraAerialPerspectiveVolumeDepthResolutionInv
-#define CameraAerialPerspectiveVolumeDepthSliceLength MobileBasePass.SkyAtmosphere.CameraAerialPerspectiveVolumeDepthSliceLength
-#define CameraAerialPerspectiveVolumeDepthSliceLengthInv MobileBasePass.SkyAtmosphere.CameraAerialPerspectiveVolumeDepthSliceLengthInv
 
 #define USE_PS_CLIP_PLANE PROJECT_ALLOW_GLOBAL_CLIP_PLANE
 
diff --git a/Engine/Shaders/Private/MobileBasePassVertexShader.usf b/Engine/Shaders/Private/MobileBasePassVertexShader.usf
--- a/Engine/Shaders/Private/MobileBasePassVertexShader.usf
+++ b/Engine/Shaders/Private/MobileBasePassVertexShader.usf
@@ -89,18 +89,18 @@ void Main(
 #if USE_VERTEX_FOG
 	float4 VertexFog = CalculateHeightFog(WorldPosition.xyz - ResolvedView.TranslatedWorldCameraOrigin);
 
-	if (ApplyCameraAerialPerspectiveVolume > 0.0f)
+	if (ResolvedView.SkyAtmosphereApplyCameraAerialPerspectiveVolume > 0.0f)
 	{
 		const float OneOverPreExposure = USE_PREEXPOSURE ? ResolvedView.OneOverPreExposure : 1.0f;
 		// Sample the aerial perspective (AP). It is also blended under the VertexFog parameter.
 		float4 AP = GetAerialPerspectiveLuminanceTransmittance(
 			Output.Position, WorldPosition.xyz*CM_TO_SKY_UNIT, ResolvedView.TranslatedWorldCameraOrigin*CM_TO_SKY_UNIT,
-			CameraAerialPerspectiveVolume, CameraAerialPerspectiveVolumeSampler,
-			CameraAerialPerspectiveVolumeDepthResolutionInv,
-			CameraAerialPerspectiveVolumeDepthResolution,
-			CameraAerialPerspectiveVolumeStartDepth,
-			CameraAerialPerspectiveVolumeDepthSliceLength,
-			CameraAerialPerspectiveVolumeDepthSliceLengthInv,
+			View.CameraAerialPerspectiveVolume, View.CameraAerialPerspectiveVolumeSampler,
+			ResolvedView.SkyAtmosphereCameraAerialPerspectiveVolumeDepthResolutionInv,
+			ResolvedView.SkyAtmosphereCameraAerialPerspectiveVolumeDepthResolution,
+			ResolvedView.SkyAtmosphereCameraAerialPerspectiveVolumeStartDepth,
+			ResolvedView.SkyAtmosphereCameraAerialPerspectiveVolumeDepthSliceLength,
+			ResolvedView.SkyAtmosphereCameraAerialPerspectiveVolumeDepthSliceLengthInv,
 			OneOverPreExposure, VertexFog);
 	}
 #if PACK_INTERPOLANTS
diff --git a/Engine/Shaders/Public/ShaderVersion.ush b/Engine/Shaders/Public/ShaderVersion.ush
--- a/Engine/Shaders/Public/ShaderVersion.ush
+++ b/Engine/Shaders/Public/ShaderVersion.ush
@@ -3,6 +3,6 @@
 // in Platform.ush (which should be included in any shader) it allows to invalidate the shader DDC.
 // 
 // If you are merging streams and there is a conflict with this GUID you should make a new GUID rather than taking one or the other.
-// GUID = 69f74d6AA9DA4A1DA3BFFA82DBA76DBA
+// GUID = 5F4B9C07A98C448285D3008A25E3A32A
 
 
diff --git a/Engine/Source/Runtime/Renderer/Private/BasePassRendering.cpp b/Engine/Source/Runtime/Renderer/Private/BasePassRendering.cpp
--- a/Engine/Source/Runtime/Renderer/Private/BasePassRendering.cpp
+++ b/Engine/Source/Runtime/Renderer/Private/BasePassRendering.cpp
@@ -506,7 +506,6 @@ void SetupSharedBasePassParameters(
 	SetupReflectionUniformParameters(View, SharedParameters.Reflection);
 	SetupFogUniformParameters(View, SharedParameters.Fog);
 	SetupPlanarReflectionUniformParameters(View, ReflectionSceneProxy, SharedParameters.PlanarReflection);
-	SetupSkyAtmosphereBasePassSharedUniformShaderParameters(View, SharedParameters.SkyAtmosphere);
 
 	const IPooledRenderTarget* PooledRT = GetSubsufaceProfileTexture_RT(RHICmdList);
 
diff --git a/Engine/Source/Runtime/Renderer/Private/BasePassRendering.h b/Engine/Source/Runtime/Renderer/Private/BasePassRendering.h
--- a/Engine/Source/Runtime/Renderer/Private/BasePassRendering.h
+++ b/Engine/Source/Runtime/Renderer/Private/BasePassRendering.h
@@ -52,7 +52,6 @@ BEGIN_GLOBAL_SHADER_PARAMETER_STRUCT(FSharedBasePassUniformParameters,)
 	SHADER_PARAMETER_STRUCT(FReflectionUniformParameters, Reflection)
 	SHADER_PARAMETER_STRUCT(FPlanarReflectionUniformParameters, PlanarReflection) // Single global planar reflection for the forward pass.
 	SHADER_PARAMETER_STRUCT(FFogUniformParameters, Fog)
-	SHADER_PARAMETER_STRUCT(FSkyAtmosphereBasePassSharedUniformShaderParameters, SkyAtmosphere)
 	SHADER_PARAMETER_TEXTURE(Texture2D, SSProfilesTexture)
 END_GLOBAL_SHADER_PARAMETER_STRUCT()
 
diff --git a/Engine/Source/Runtime/Renderer/Private/MobileBasePassRendering.cpp b/Engine/Source/Runtime/Renderer/Private/MobileBasePassRendering.cpp
--- a/Engine/Source/Runtime/Renderer/Private/MobileBasePassRendering.cpp
+++ b/Engine/Source/Runtime/Renderer/Private/MobileBasePassRendering.cpp
@@ -157,8 +157,6 @@ void SetupMobileBasePassUniformParameters(
 {
 	SetupFogUniformParameters(View, BasePassParameters.Fog);
 
-	SetupSkyAtmosphereBasePassSharedUniformShaderParameters(View, BasePassParameters.SkyAtmosphere);
-
 	const FScene* Scene = View.Family->Scene ? View.Family->Scene->GetRenderScene() : nullptr;
 	const FPlanarReflectionSceneProxy* ReflectionSceneProxy = Scene ? Scene->GetForwardPassGlobalPlanarReflection() : nullptr;
 	SetupPlanarReflectionUniformParameters(View, ReflectionSceneProxy, BasePassParameters.PlanarReflection);
diff --git a/Engine/Source/Runtime/Renderer/Private/MobileBasePassRendering.h b/Engine/Source/Runtime/Renderer/Private/MobileBasePassRendering.h
--- a/Engine/Source/Runtime/Renderer/Private/MobileBasePassRendering.h
+++ b/Engine/Source/Runtime/Renderer/Private/MobileBasePassRendering.h
@@ -28,7 +28,6 @@ BEGIN_GLOBAL_SHADER_PARAMETER_STRUCT(FMobileBasePassUniformParameters, )
 	SHADER_PARAMETER_STRUCT(FFogUniformParameters, Fog)
 	SHADER_PARAMETER_STRUCT(FPlanarReflectionUniformParameters, PlanarReflection) // Single global planar reflection for the forward pass.
 	SHADER_PARAMETER_STRUCT(FMobileSceneTextureUniformParameters, SceneTextures)
-	SHADER_PARAMETER_STRUCT(FSkyAtmosphereBasePassSharedUniformShaderParameters, SkyAtmosphere)
 END_GLOBAL_SHADER_PARAMETER_STRUCT()
 
 extern void SetupMobileBasePassUniformParameters(
diff --git a/Engine/Source/Runtime/Renderer/Private/SceneRendering.cpp b/Engine/Source/Runtime/Renderer/Private/SceneRendering.cpp
--- a/Engine/Source/Runtime/Renderer/Private/SceneRendering.cpp
+++ b/Engine/Source/Runtime/Renderer/Private/SceneRendering.cpp
@@ -1227,8 +1227,8 @@ void FViewInfo::SetupUniformBufferParameters(
 		ViewUniformShaderParameters.SkyAtmosphereBottomRadius = AtmosphereSetup.BottomRadius;
 		ViewUniformShaderParameters.SkyAtmosphereTopRadius = AtmosphereSetup.TopRadius;
 
-		FSkyAtmosphereBasePassSharedUniformShaderParameters OutParameters;
-		SetupSkyAtmosphereBasePassSharedUniformShaderParameters(*this, OutParameters);
+		FSkyAtmosphereViewSharedUniformShaderParameters OutParameters;
+		SetupSkyAtmosphereViewSharedUniformShaderParameters(*this, OutParameters);
 		ViewUniformShaderParameters.SkyAtmosphereCameraAerialPerspectiveVolumeStartDepth = OutParameters.CameraAerialPerspectiveVolumeStartDepth;
 		ViewUniformShaderParameters.SkyAtmosphereCameraAerialPerspectiveVolumeDepthResolution = OutParameters.CameraAerialPerspectiveVolumeDepthResolution;
 		ViewUniformShaderParameters.SkyAtmosphereCameraAerialPerspectiveVolumeDepthResolutionInv = OutParameters.CameraAerialPerspectiveVolumeDepthResolutionInv;
diff --git a/Engine/Source/Runtime/Renderer/Private/SkyAtmosphereRendering.cpp b/Engine/Source/Runtime/Renderer/Private/SkyAtmosphereRendering.cpp
--- a/Engine/Source/Runtime/Renderer/Private/SkyAtmosphereRendering.cpp
+++ b/Engine/Source/Runtime/Renderer/Private/SkyAtmosphereRendering.cpp
@@ -218,7 +218,6 @@ BEGIN_GLOBAL_SHADER_PARAMETER_STRUCT(FSkyAtmosphereInternalCommonParameters, )
 END_GLOBAL_SHADER_PARAMETER_STRUCT()
 
 IMPLEMENT_GLOBAL_SHADER_PARAMETER_STRUCT(FAtmosphereUniformShaderParameters, "Atmosphere");
-IMPLEMENT_GLOBAL_SHADER_PARAMETER_STRUCT(FSkyAtmosphereBasePassSharedUniformShaderParameters, "SkyAtmosphere");
 IMPLEMENT_GLOBAL_SHADER_PARAMETER_STRUCT(FSkyAtmosphereInternalCommonParameters, "SkyAtmosphere");
 
 
@@ -255,7 +254,7 @@ bool ShouldRenderSkyAtmosphere(const FSkyAtmosphereRenderSceneInfo* SkyAtmospher
 	// TODO: Add new or reuse EngineShowFlags.AtmosphericFog? ALso take into account EngineShowFlags.Fog as previously?
 }
 
-void SetupSkyAtmosphereBasePassSharedUniformShaderParameters(const class FViewInfo& View, FSkyAtmosphereBasePassSharedUniformShaderParameters& OutParameters)
+void SetupSkyAtmosphereViewSharedUniformShaderParameters(const class FViewInfo& View, FSkyAtmosphereViewSharedUniformShaderParameters& OutParameters)
 {
 	GET_VALID_DATA_FROM_CVAR;
 
@@ -264,7 +263,6 @@ void SetupSkyAtmosphereBasePassSharedUniformShaderParameters(const class FViewIn
 	{
 		SkyAtmosphereCameraAerialPerspectiveVolume = View.SkyAtmosphereCameraAerialPerspectiveVolume->GetRenderTargetItem().ShaderResourceTexture;
 	}
-	SetBlackAlpha13DIfNull(SkyAtmosphereCameraAerialPerspectiveVolume);
 
 	OutParameters.ApplyCameraAerialPerspectiveVolume = View.SkyAtmosphereCameraAerialPerspectiveVolume == nullptr ? 0.0f : 1.0f;
 	OutParameters.CameraAerialPerspectiveVolumeStartDepth = CameraAerialPerspectiveVolumeStartDepth;
@@ -272,8 +270,8 @@ void SetupSkyAtmosphereBasePassSharedUniformShaderParameters(const class FViewIn
 	OutParameters.CameraAerialPerspectiveVolumeDepthResolutionInv = 1.0f / OutParameters.CameraAerialPerspectiveVolumeDepthResolution;
 	OutParameters.CameraAerialPerspectiveVolumeDepthSliceLength = CameraAerialPerspectiveVolumeDepthSliceLength;
 	OutParameters.CameraAerialPerspectiveVolumeDepthSliceLengthInv = 1.0f / OutParameters.CameraAerialPerspectiveVolumeDepthSliceLength;
-	OutParameters.CameraAerialPerspectiveVolume = SkyAtmosphereCameraAerialPerspectiveVolume;
-	OutParameters.CameraAerialPerspectiveVolumeSampler = TStaticSamplerState<SF_Trilinear>::GetRHI();
+
+	SetBlackAlpha13DIfNull(SkyAtmosphereCameraAerialPerspectiveVolume); // Needs to be after we set ApplyCameraAerialPerspectiveVolume
 }
 
 static void CopyAtmosphereSetupToUniformShaderParameters(FAtmosphereUniformShaderParameters& out, const FAtmosphereSetup& Atmosphere)
diff --git a/Engine/Source/Runtime/Renderer/Private/SkyAtmosphereRendering.h b/Engine/Source/Runtime/Renderer/Private/SkyAtmosphereRendering.h
--- a/Engine/Source/Runtime/Renderer/Private/SkyAtmosphereRendering.h
+++ b/Engine/Source/Runtime/Renderer/Private/SkyAtmosphereRendering.h
@@ -18,6 +18,7 @@ class USkyAtmosphereComponent;
 
 
 // Use as a global shader parameter struct and also the CPU structure representing the atmosphere it self.
+// This is static for a version of a component. When a component is changed/tweaked, it is recreated.
 BEGIN_GLOBAL_SHADER_PARAMETER_STRUCT(FAtmosphereUniformShaderParameters, )
 	SHADER_PARAMETER(float, MultiScatteringFactor)
 	SHADER_PARAMETER(float, BottomRadius)
@@ -38,17 +39,16 @@ BEGIN_GLOBAL_SHADER_PARAMETER_STRUCT(FAtmosphereUniformShaderParameters, )
 	SHADER_PARAMETER(FLinearColor, GroundAlbedo)
 END_GLOBAL_SHADER_PARAMETER_STRUCT()
 
-// Extra data shared by the Sky/Atmosphere system with Base passes.
-BEGIN_GLOBAL_SHADER_PARAMETER_STRUCT(FSkyAtmosphereBasePassSharedUniformShaderParameters,)
-	SHADER_PARAMETER_TEXTURE(Texture3D, CameraAerialPerspectiveVolume)
-	SHADER_PARAMETER_SAMPLER(SamplerState, CameraAerialPerspectiveVolumeSampler)
-	SHADER_PARAMETER(float, CameraAerialPerspectiveVolumeStartDepth)
-	SHADER_PARAMETER(float, CameraAerialPerspectiveVolumeDepthResolution)
-	SHADER_PARAMETER(float, CameraAerialPerspectiveVolumeDepthResolutionInv)
-	SHADER_PARAMETER(float, CameraAerialPerspectiveVolumeDepthSliceLength)
-	SHADER_PARAMETER(float, CameraAerialPerspectiveVolumeDepthSliceLengthInv)
-	SHADER_PARAMETER(float, ApplyCameraAerialPerspectiveVolume)
-END_GLOBAL_SHADER_PARAMETER_STRUCT()
+// These parameters are shared on the view global uniform buffer and are dynamically changed with cvars.
+struct FSkyAtmosphereViewSharedUniformShaderParameters
+{
+	float CameraAerialPerspectiveVolumeStartDepth;
+	float CameraAerialPerspectiveVolumeDepthResolution;
+	float CameraAerialPerspectiveVolumeDepthResolutionInv;
+	float CameraAerialPerspectiveVolumeDepthSliceLength;
+	float CameraAerialPerspectiveVolumeDepthSliceLengthInv;
+	float ApplyCameraAerialPerspectiveVolume;
+};
 
 
 
@@ -95,6 +95,6 @@ private:
 
 bool ShouldRenderSkyAtmosphere(const FSkyAtmosphereRenderSceneInfo* SkyAtmosphere, EShaderPlatform ShaderPlatform);
 
-extern void SetupSkyAtmosphereBasePassSharedUniformShaderParameters(const class FViewInfo& View, FSkyAtmosphereBasePassSharedUniformShaderParameters& OutParameters);
+extern void SetupSkyAtmosphereViewSharedUniformShaderParameters(const class FViewInfo& View, FSkyAtmosphereViewSharedUniformShaderParameters& OutParameters);
 
 
