Implements tile classification prototype for SSGI

diff --git a/Engine/Shaders/Private/SSRT/SSRTDiffuseIndirect.usf b/Engine/Shaders/Private/SSRT/SSRTDiffuseIndirect.usf
--- a/Engine/Shaders/Private/SSRT/SSRTDiffuseIndirect.usf
+++ b/Engine/Shaders/Private/SSRT/SSRTDiffuseIndirect.usf
@@ -1,10 +1,16 @@
 // Copyright 1998-2019 Epic Games, Inc. All Rights Reserved.
 
+// Generate vector truncation warnings to errors.
+#pragma warning(error: 3206)
+
 #define CONFIG_SORT_RAYS 0
 #define CONFIG_SORT_BIN_COUNT 16
+#define CONFIG_EARLY_RETURN 0
 
 #define DEBUG_SSRT 0
 
+#define DEBUG_RAY_COUNT 0
+
 
 #if QUALITY == 1
 	#define CONFIG_RAY_STEPS 8
@@ -64,6 +70,7 @@
 #define IS_SSGI_SHADER 1
 
 #include "SSRTRayCast.ush"
+#include "SSRTTileClassificationBuffer.ush"
 #include "../DeferredShadingCommon.ush"
 #include "../Random.ush"
 #include "../BRDF.ush"
@@ -182,14 +189,90 @@ float3 ComputeL(float3 N, float2 E)
 	return mul(TangentL, TangentBasis);
 }
 
+/** Test whether a screen space ray can be discarded or not based on tile information. */
+bool TestRayEarlyReturn(FSSRTTileInfos TileInfos, FSSRTRay Ray)
+#if CONFIG_EARLY_RETURN
+{
+	float2 RayStepPixel = Ray.RayStepScreen.xy * View.ViewSizeAndInvSize.xy;
+
+	// Screen position is vertically flipped from 
+	RayStepPixel.y = -RayStepPixel.y;
+
+	// Compute screen ray direction in [0; 2 PI[
+	float RayDirection = atan2(-RayStepPixel.y, -RayStepPixel.x) + PI;
+	
+	float DeltaU = length(RayStepPixel) * (0.5 * View.ViewSizeAndInvSize.z);
+	float DeltaZ = Ray.RayStepScreen.z;
+	float RayTheta = atan2(DeltaU, -DeltaZ);
+	
+	// Compare RayTheta according to RayDirection against the TileInfos.Directionality stored into SGPR.
+	bool bEarlyReturn;
+	#if DIRECTIONALITY_DIVISION == 8
+	FLATTEN
+	if (RayDirection < (0.5 * 2 * PI / DIRECTIONALITY_DIVISION) || RayDirection > (7.5 * 2 * PI / DIRECTIONALITY_DIVISION))
+	{
+		bEarlyReturn = RayTheta > TileInfos.Directionality[0];
+		//bEarlyReturn = true;
+	}
+	else if (RayDirection < (1.5 * 2 * PI / DIRECTIONALITY_DIVISION))
+	{
+		bEarlyReturn = RayTheta > TileInfos.Directionality[1];
+		//bEarlyReturn = true;
+	}
+	else if (RayDirection < (2.5 * 2 * PI / DIRECTIONALITY_DIVISION))
+	{
+		bEarlyReturn = RayTheta > TileInfos.Directionality[2];
+		//bEarlyReturn = true;
+		//bEarlyReturn = false;
+	}
+	else if (RayDirection < (3.5 * 2 * PI / DIRECTIONALITY_DIVISION))
+	{
+		bEarlyReturn = RayTheta > TileInfos.Directionality[3];
+		//bEarlyReturn = true;
+		//bEarlyReturn = false;
+	}
+	else if (RayDirection < (4.5 * 2 * PI / DIRECTIONALITY_DIVISION))
+	{
+		bEarlyReturn = RayTheta > TileInfos.Directionality[4];
+		//bEarlyReturn = true;
+	}
+	else if (RayDirection < (5.5 * 2 * PI / DIRECTIONALITY_DIVISION))
+	{
+		bEarlyReturn = RayTheta > TileInfos.Directionality[5];
+		//bEarlyReturn = true;
+	}
+	else if (RayDirection < (6.5 * 2 * PI / DIRECTIONALITY_DIVISION))
+	{
+		bEarlyReturn = RayTheta > TileInfos.Directionality[6];
+		//bEarlyReturn = true;
+	}
+	else
+	{
+		bEarlyReturn = RayTheta > TileInfos.Directionality[7];
+		//bEarlyReturn = true;
+	}
+	#else
+		#error Need to update this code.
+	#endif
+
+	return bEarlyReturn;
+}
+#else
+{
+	return false;
+}
+#endif
+
 #if CONFIG_SORT_RAYS
 
 uint SelectDirectionBin(
+	FSSRTTileInfos TileInfos,
 	uint2 GroupId,
 	uint GroupPixelId,
 	uint RaySequenceId,
 	float3 N,
-	float SceneDepth)
+	float SceneDepth,
+	inout bool bTraceRay)
 {
 	uint2 GroupPixelOffset = DecodeGroupPixelOffset(GroupPixelId);
 	uint2 PixelPosition = ComputePixelPosition(GroupId, GroupPixelOffset);
@@ -205,6 +288,14 @@ uint SelectDirectionBin(
 	UpdateLane2DCoordinateInformations(PixelPosition, /* out */ BufferUV, /* out */ ScreenPos);
 	
 	float3 TranslatedWorldPosition = ComputeTranslatedWorldPositions(ScreenPos, SceneDepth);
+	
+	FSSRTRay Ray = InitScreenSpaceRay(TranslatedWorldPosition, L, SceneDepth);
+	
+	// Early out if the tile classification can detect this ray direction is not going to intersect anything.
+	bool bEarlyOut = TestRayEarlyReturn(TileInfos, Ray);
+	
+	// Don't trace a ray if it can be early out by tile classification.
+	bTraceRay = bTraceRay && !bEarlyOut;
 
 	// TODO provide RayStartUVz
 	const float4 RayStartClip = mul(float4(TranslatedWorldPosition, 1), View.TranslatedWorldToClip);
@@ -216,11 +307,11 @@ uint SelectDirectionBin(
 	float2 RayPixelDirection = (RayEndScreen - RayStartScreen).xy * View.ViewSizeAndInvSize.xy;
 
 	float DirectionAngle = atan2(RayPixelDirection.x, RayPixelDirection.y);
-			
-	uint SelectedBinId = floor(DirectionAngle * (CONFIG_SORT_BIN_COUNT * rcp(2 * PI)) + (CONFIG_SORT_BIN_COUNT / 2 + 0.5)); 
-	SelectedBinId = clamp(SelectedBinId, 0, CONFIG_SORT_BIN_COUNT - 1);
+	
+	uint SelectedBinId = floor(DirectionAngle * ((CONFIG_SORT_BIN_COUNT - 1) * rcp(2 * PI)) + ((CONFIG_SORT_BIN_COUNT - 1) / 2 + 0.5)); 
+	SelectedBinId = clamp(SelectedBinId, 0, CONFIG_SORT_BIN_COUNT - 2);
 
-	return SelectedBinId;
+	return bTraceRay ? SelectedBinId : (CONFIG_SORT_BIN_COUNT - 1);
 } // SelectDirectionBin()
 
 #endif // CONFIG_SORT_RAYS
@@ -233,12 +324,18 @@ void MainCS(
 	// Id of the wave in the group.
 	uint GroupWaveIndex = GroupThreadIndex / 64;
 	
+	FSSRTTileInfos TileInfos;
 	{
 		const uint BinsAddress = TILE_PIXEL_COUNT * 2;
 
 		uint GroupPixelId = GroupThreadIndex % TILE_PIXEL_COUNT;
 		uint RaySequenceId = GroupThreadIndex / TILE_PIXEL_COUNT;
-		
+
+		// Compute TileCoord from GroupId to ensure the compiler understand it is group invariant to use scalar load.
+		uint2 TileCoord = GroupId / uint2(TILE_RES_DIVISOR / TILE_PIXEL_SIZE_X, TILE_RES_DIVISOR / TILE_PIXEL_SIZE_Y);
+		TileInfos = LoadTileInfos(TileCoord);
+
+
 		// Store GBuffer into LDS
 		{
 			BRANCH
@@ -251,10 +348,10 @@ void MainCS(
 		
 				FGBufferData GBuffer = GetGBufferDataFromSceneTextures(BufferUV);
 
-				bool bEarlyReturn = GBuffer.ShadingModelID == SHADINGMODELID_UNLIT || GBuffer.ShadingModelID == SHADINGMODELID_SUBSURFACE_PROFILE;
+				bool bTraceRay = GBuffer.ShadingModelID != SHADINGMODELID_UNLIT && GBuffer.ShadingModelID != SHADINGMODELID_SUBSURFACE_PROFILE;
 				
 				SharedMemory[TILE_PIXEL_COUNT * 0 | GroupPixelId] = CompressN(GBuffer.WorldNormal);
-				SharedMemory[TILE_PIXEL_COUNT * 1 | GroupPixelId] = asuint(bEarlyReturn ? -1.0 : GBuffer.Depth);
+				SharedMemory[TILE_PIXEL_COUNT * 1 | GroupPixelId] = asuint(bTraceRay ? GBuffer.Depth : -1.0);
 			}
 			else if (GroupWaveIndex == 1) // TODO.
 			{
@@ -271,6 +368,7 @@ void MainCS(
 			GroupMemoryBarrierWithGroupSync();
 
 			// Fetch information from LDS.
+			bool bTraceRay;
 			float3 N;
 			float SceneDepth;
 			{
@@ -278,16 +376,19 @@ void MainCS(
 				uint Raw1 = SharedMemory[TILE_PIXEL_COUNT * 1 | GroupPixelId];
 
 				N = DecompressN(Raw0);
-				SceneDepth = asfloat(Raw1);
+				SceneDepth = abs(asfloat(Raw1));
+				bTraceRay = asfloat(Raw1) > 0.0;
 			}
 
 			// Compute the ray's bin.
 			uint SelectedBinId = SelectDirectionBin(
+				TileInfos,
 				GroupId,
 				GroupPixelId,
 				RaySequenceId,
 				N,
-				SceneDepth);
+				SceneDepth,
+				/* inout */ bTraceRay);
 		
 			// Count ray in bins
 			uint OffsetInBin = 0;
@@ -314,7 +415,7 @@ void MainCS(
 			{
 				uint DestPos = SelectedBinStartingPos + OffsetInBin;
 
-				#if 0 // Uncomment me to measure performance improvement.
+				#if 1 // Uncomment me to measure performance improvement.
 					if (View.GeneralPurposeTweak > 0.5)
 						DestPos = GroupPixelId + RaySequenceId * TILE_PIXEL_COUNT;
 				#endif
@@ -325,7 +426,7 @@ void MainCS(
 
 				// TODO: Encoded N and 24bits and pack EncodedGroupPixelOffset.
                 SharedMemory[LANE_PER_GROUPS * 0 | DestPos] = asuint(CompressN(N) | (GroupPixelId << 24) | (RaySequenceId << (24 + TILE_PIXEL_SIZE_X_LOG + TILE_PIXEL_SIZE_Y_LOG)));
-				SharedMemory[LANE_PER_GROUPS * 1 | DestPos] = asuint(SceneDepth);
+				SharedMemory[LANE_PER_GROUPS * 1 | DestPos] = asuint(bTraceRay ? SceneDepth : -1.0);
 			}
 		}
 		#endif // CONFIG_SORT_RAYS
@@ -339,16 +440,18 @@ void MainCS(
 		uint RaySequenceId;
 		uint CompressedN;
 		float SceneDepth;
+		bool bTraceRay;
 		#if CONFIG_SORT_RAYS
 		{
-			uint Row0 = SharedMemory[LANE_PER_GROUPS * 0 | GroupThreadIndex];
-			uint Row1 = SharedMemory[LANE_PER_GROUPS * 1 | GroupThreadIndex];
+			uint Raw0 = SharedMemory[LANE_PER_GROUPS * 0 | GroupThreadIndex];
+			uint Raw1 = SharedMemory[LANE_PER_GROUPS * 1 | GroupThreadIndex];
 
 			// Unpack
-			RaySequenceId = Row0 >> (24 + TILE_PIXEL_SIZE_X_LOG + TILE_PIXEL_SIZE_Y_LOG);
-            GroupPixelId = (Row0 >> 24) % TILE_PIXEL_COUNT;
-			CompressedN = Row0;
-			SceneDepth = asfloat(Row1);
+			RaySequenceId = Raw0 >> (24 + TILE_PIXEL_SIZE_X_LOG + TILE_PIXEL_SIZE_Y_LOG);
+            GroupPixelId = (Raw0 >> 24) % TILE_PIXEL_COUNT;
+			CompressedN = Raw0;
+			SceneDepth = asfloat(Raw1);
+			bTraceRay = asfloat(Raw1) > 0;
 		}
 		#else // !CONFIG_SORT_RAYS
 		{
@@ -360,13 +463,15 @@ void MainCS(
 
 			CompressedN = Raw0;
 			SceneDepth = asfloat(Raw1);
+			bTraceRay = asfloat(Raw1) > 0;
 		}
 		#endif // !CONFIG_SORT_RAYS
 		
 		GroupMemoryBarrierWithGroupSync();
-		
-		const bool bTraceRay = SceneDepth > 0;
-		
+
+		#if DEBUG_RAY_COUNT
+			float DebugRayCount = 0.0;
+		#endif
 		uint2 CompressedColor;
 
 		BRANCH
@@ -389,16 +494,37 @@ void MainCS(
 	
             bool bDebugPrint = all(PixelPosition == uint2(View.ViewSizeAndInvSize.xy) / 2);
 
+			FSSRTRay Ray = InitScreenSpaceRay(TranslatedWorldPosition, L, SceneDepth);
+			
 			float Level;
 			float3 HitUVz;
-			bool bHit = RayCast(
-				HZBTexture, HZBTextureSampler,
-				TranslatedWorldPosition, L, 1, SceneDepth, 
-				CONFIG_RAY_STEPS, StepOffset,
-				HZBUvFactorAndInvFactor,
-				bDebugPrint,
-				/* out */ HitUVz,
-				/* out */ Level);
+			bool bHit;
+
+			#if !CONFIG_SORT_RAYS
+			// Early out if the tile classification can detect this ray direction is not going to intersect anything.
+			bool bEarlyOut = TestRayEarlyReturn(TileInfos, Ray);
+			
+			BRANCH
+			if (bEarlyOut)
+			{
+				bHit = false;
+				Level = 0;
+				HitUVz = 0;
+			}
+			else
+			#endif
+			{
+				bHit = CastScreenSpaceRay(
+					HZBTexture, HZBTextureSampler,
+					Ray, 1, CONFIG_RAY_STEPS, StepOffset,
+					HZBUvFactorAndInvFactor, bDebugPrint,
+					/* out */ HitUVz,
+					/* out */ Level);
+				
+				#if DEBUG_RAY_COUNT
+					DebugRayCount += 1.0;
+				#endif
+			}
 
 			#if 0 // Backface check
 			if (bHit)
@@ -409,10 +535,10 @@ void MainCS(
 			#endif
 			
 			// if there was a hit
-			//BRANCH // TODO
+			BRANCH
 			if (bHit)
 			{
-				float4 SampleColor = ColorTexture.SampleLevel(ColorTextureSampler, HitUVz, Level);
+				float4 SampleColor = ColorTexture.SampleLevel(ColorTextureSampler, HitUVz.xy, Level);
 				
 				float SampleColorWeight = 1.0;
 
@@ -455,6 +581,14 @@ void MainCS(
             CompressedColor = uint2(0, 0);
         }
 		
+		// Output debugging info instead of actual intersection.
+		#if DEBUG_RAY_COUNT
+		{
+			CompressedColor.x = asuint(f32tof16(DebugRayCount) << 16);
+			CompressedColor.y = 0;
+		}
+		#endif
+		
 		uint DestPos = GroupPixelId + RaySequenceId * TILE_PIXEL_COUNT;
 			
 		SharedMemory[LANE_PER_GROUPS * 0 | DestPos] = CompressedColor.x;
@@ -512,6 +646,8 @@ void MainCS(
 
 			IndirectDiffuseOutput[OutputPixelCoordinate] = float4(DiffuseColor, 0.0);
 			AmbientOcclusionOutput[OutputPixelCoordinate] = AmbientOcclusion;
+			
+			//IndirectDiffuseOutput[OutputPixelCoordinate] = float4(TileInfos.Directionality[2] / PI, 0.0, 0.0, 0.0);
 		}
 	} // if (GroupThreadIndex < TILE_PIXEL_COUNT)
 } // MainCS()
diff --git a/Engine/Shaders/Private/SSRT/SSRTRayCast.ush b/Engine/Shaders/Private/SSRT/SSRTRayCast.ush
--- a/Engine/Shaders/Private/SSRT/SSRTRayCast.ush
+++ b/Engine/Shaders/Private/SSRT/SSRTRayCast.ush
@@ -34,17 +34,34 @@ void PrintSample(
 #endif
 
 
+/** Return float multiplier to scale RayStepScreen such that it clip it right at the edge of the screen. */
+float GetStepScreenFactorToClipAtScreenEdge(float2 RayStartScreen, float2 RayStepScreen)
+{
+	// Computes the scale down factor for RayStepScreen required to fit on the X and Y axis in order to clip it in the viewport
+	const float RayStepScreenInvFactor = 0.5 * length(RayStepScreen);
+	const float2 S = 1 - max(abs(RayStepScreen + RayStartScreen * RayStepScreenInvFactor) - RayStepScreenInvFactor, 0.0f) / abs(RayStepScreen);
 
-bool RayCast(
-	Texture2D Texture, SamplerState Sampler,
-	float3 RayOriginTranslatedWorld, float3 RayDirection,
-	float Roughness, float SceneDepth,
-	uint NumSteps, float StepOffset,
-	float4 HZBUvFactorAndInvFactor, 
-	bool bDebugPrint,
-	out float3 OutHitUVz,
-	out float Level
-)
+	// Rescales RayStepScreen accordingly
+	const float RayStepFactor = min(S.x, S.y) / RayStepScreenInvFactor;
+
+	return RayStepFactor;
+}
+
+
+/** Structure that represent a ray to be shot in screen space. */
+struct FSSRTRay
+{
+	float3 RayStartScreen;
+	float3 RayStepScreen;
+
+	float CompareTolerance;
+};
+
+/** Compile a ray for screen space ray casting. */
+FSSRTRay InitScreenSpaceRay(
+	float3 RayOriginTranslatedWorld,
+	float3 RayDirection,
+	float SceneDepth)
 {
 	// TODO provide RayStartUVz
 	const float4 RayStartClip	= mul( float4( RayOriginTranslatedWorld, 1 ), View.TranslatedWorldToClip );
@@ -55,28 +72,42 @@ bool RayCast(
 
 	const float4 RayDepthClip	= RayStartClip + mul( float4( 0, 0, SceneDepth, 0 ), View.ViewToClip );
 	const float3 RayDepthScreen = RayDepthClip.xyz / RayDepthClip.w;
+	
+	FSSRTRay Ray;
+	Ray.RayStartScreen = RayStartScreen;
+	Ray.RayStepScreen = RayEndScreen - RayStartScreen;
 
-	float3 RayStepScreen = RayEndScreen - RayStartScreen;
+	// TODO
+	#if IS_SSGI_SHADER
+		Ray.CompareTolerance = max(abs(Ray.RayStepScreen.z), (RayStartScreen.z - RayDepthScreen.z) * 2);
+	#else
+		Ray.CompareTolerance = max(abs(Ray.RayStepScreen.z), (RayStartScreen.z - RayDepthScreen.z) * 4);
+	#endif
 
-	{
-		// Computes the scale down factor for RayStepScreen required to fit on the X and Y axis in order to clip it in the viewport
-		const float RayStepScreenInvFactor = 0.5 * length( RayStepScreen.xy );
-		const float2 S = 1 - max(abs(RayStepScreen.xy + RayStartScreen.xy * RayStepScreenInvFactor) - RayStepScreenInvFactor, 0.0f) / abs(RayStepScreen.xy);
+	return Ray;
+} // InitScreenSpaceRay()
 
-		// Rescales RayStepScreen accordingly
-		const float RayStepFactor = min(S.x, S.y) / RayStepScreenInvFactor;
-		RayStepScreen *= RayStepFactor;
-	}
+/** Cast a screen space ray. */
+bool CastScreenSpaceRay(
+	Texture2D Texture, SamplerState Sampler,
+	FSSRTRay Ray,
+	float Roughness,
+	uint NumSteps, float StepOffset,
+	float4 HZBUvFactorAndInvFactor, 
+	bool bDebugPrint,
+	out float3 OutHitUVz,
+	out float Level)
+{
+	const float3 RayStartScreen = Ray.RayStartScreen;
+	float3 RayStepScreen = Ray.RayStepScreen;
+
+	RayStepScreen *= GetStepScreenFactorToClipAtScreenEdge(RayStartScreen.xy, RayStepScreen.xy);
 
 	float3 RayStartUVz = float3( (RayStartScreen.xy * float2( 0.5, -0.5 ) + 0.5) * HZBUvFactorAndInvFactor.xy, RayStartScreen.z );
 	float3 RayStepUVz  = float3(  RayStepScreen.xy  * float2( 0.5, -0.5 )		 * HZBUvFactorAndInvFactor.xy, RayStepScreen.z );
 	
 	const float Step = 1.0 / NumSteps;
-	#if IS_SSGI_SHADER
-		float CompareTolerance = max( abs( RayStepUVz.z ), ( RayStartScreen.z - RayDepthScreen.z ) * 2 ) * Step;
-	#else
-		float CompareTolerance = max( abs( RayStepUVz.z ), ( RayStartScreen.z - RayDepthScreen.z ) * 4 ) * Step;
-	#endif
+	float CompareTolerance = Ray.CompareTolerance * Step;
 	
 	float LastDiff = 0;
 	Level = 1;
@@ -257,7 +288,29 @@ bool RayCast(
     }
 	
 	return bFoundAnyHit;
-}
+} // CastScreenSpaceRay()
+
+
+bool RayCast(
+	Texture2D Texture, SamplerState Sampler,
+	float3 RayOriginTranslatedWorld, float3 RayDirection,
+	float Roughness, float SceneDepth,
+	uint NumSteps, float StepOffset,
+	float4 HZBUvFactorAndInvFactor, 
+	bool bDebugPrint,
+	out float3 OutHitUVz,
+	out float Level)
+{
+	FSSRTRay Ray = InitScreenSpaceRay(RayOriginTranslatedWorld, RayDirection, SceneDepth);
+
+	return CastScreenSpaceRay(
+		Texture, Sampler,
+		Ray,
+		Roughness, NumSteps, StepOffset,
+		HZBUvFactorAndInvFactor, bDebugPrint,
+		/* out */ OutHitUVz,
+		/* out */ Level);
+} // RayCast()
 
 float ComputeHitVignetteFromScreenPos(float2 ScreenPos)
 {
diff --git a/Engine/Shaders/Private/SSRT/SSRTTileClassification.usf b/Engine/Shaders/Private/SSRT/SSRTTileClassification.usf
new file mode 100644
--- /dev/null
+++ b/Engine/Shaders/Private/SSRT/SSRTTileClassification.usf
@@ -0,0 +1,114 @@
+// Copyright 1998-2019 Epic Games, Inc. All Rights Reserved.
+
+#include "SSRTTileClassificationBuffer.ush"
+#include "SSRTRaycast.ush"
+#include "../SceneTextureParameters.ush"
+
+
+#define GROUP_TILE_SIZE 8
+
+
+float2 SamplePixelToHZBUV;
+
+Texture2D HZBTexture;
+
+#if SUPPORTS_INDEPENDENT_SAMPLERS
+#define HZBTextureSampler GlobalPointClampedSampler
+
+#else
+SamplerState HZBTextureSampler;
+
+#endif
+
+
+RWBuffer<float> TileClassificationBufferOutput;
+RWTexture2D<float4> DebugOutput;
+
+//void StoreTileInfos(uint2 TileCoord, FSSRTTileInfos TileInfos)
+//{
+//	uint TileIndex = TileCoord.x + TileCoord.y * TileBufferExtent.x;
+//
+//	UNROLL_N(DIRECTIONALITY_COUNT)
+//	for (uint i = 0; i < DIRECTIONALITY_COUNT; i++)
+//	{
+//		TileClassificationBufferOutput[TileIndex + ViewTileCount * i] = TileInfos.Directionality[i];
+//	}
+//}
+
+
+[numthreads(GROUP_TILE_SIZE, GROUP_TILE_SIZE, DIRECTIONALITY_DIVISION)]
+void MainCS(
+	uint2 DispatchThreadId : SV_DispatchThreadID,
+	uint2 GroupId : SV_GroupID,
+	uint3 GroupThreadId : SV_GroupThreadID,
+	uint GroupThreadIndex : SV_GroupIndex)
+{
+	uint2 TileCoord = DispatchThreadId;
+	uint DirectionId = GroupThreadId.z;
+
+	float2 ViewportUV = (TileCoord * GROUP_TILE_SIZE + GROUP_TILE_SIZE / 2) * View.ViewSizeAndInvSize.zw;
+
+	float DirectionAndle = float(DirectionId) * (2 * PI * rcp(float(DIRECTIONALITY_DIVISION)));
+
+	float2 RayPixelDirection = float2(cos(DirectionAndle), sin(DirectionAndle));
+
+	float2 RayStartScreen = ViewportUVToScreenPos(ViewportUV);
+	float2 RayStepScreen = normalize(View.ViewSizeAndInvSize.zw * RayPixelDirection);
+	RayStepScreen *= GetStepScreenFactorToClipAtScreenEdge(RayStartScreen, RayStepScreen);
+	
+	float2 RayEndViewportUV = ScreenPosToViewportUV(RayStartScreen + RayStepScreen);
+
+	float2 RayStartPixel = ViewportUV * View.ViewSizeAndInvSize.xy;
+	float2 RayEndPixel = RayEndViewportUV * View.ViewSizeAndInvSize.xy;
+	float MaxSampleDistance = length(RayStartPixel - RayEndPixel);
+
+	// HZB mip0 is already half resolution.
+	const float StartLevel = log2(GROUP_TILE_SIZE) - 1;
+
+	// TODO: should be furtherest depth.
+	float TileDeviceZ = HZBTexture.SampleLevel(HZBTextureSampler, RayStartPixel * SamplePixelToHZBUV, StartLevel).r;
+	
+	float MaxTheta = 0;
+
+	const float SampleDistanceExponent = 3;
+	const float MipLevelMultiplier = 0.5 * log2(SampleDistanceExponent - 1);
+	float CurrentSampleDistance = GROUP_TILE_SIZE;
+
+	for (uint i = 0; i < 16; i++)
+	{
+		float SamplePixelDistance = CurrentSampleDistance;
+
+		//if (SamplePixelDistance > 1.5 * MaxSampleDistance)
+		//	break;
+
+		float SampleLevel = StartLevel + i * MipLevelMultiplier;
+
+		float2 SamplePixel = RayStartPixel + RayPixelDirection * SamplePixelDistance;
+
+		// TODO: should be closest depth.
+		float SampleDiviceZ = HZBTexture.SampleLevel(HZBTextureSampler, SamplePixel * SamplePixelToHZBUV, SampleLevel).r;
+
+		// Screen space distance in viewport UV from origin of the ray.
+		float DeltaU = SamplePixelDistance * View.ViewSizeAndInvSize.z;
+
+		// Assumes inverted Z buffer.
+		float DeltaZ = SampleDiviceZ - TileDeviceZ;
+
+		float Theta = atan2(DeltaU, -DeltaZ);
+
+		MaxTheta = max(MaxTheta, Theta);
+		CurrentSampleDistance *= SampleDistanceExponent;
+	}
+
+	BRANCH
+	if (TileCoord.x < TileBufferExtent.x && TileCoord.y < TileBufferExtent.y)
+	{
+		uint TileIndex = TileCoord.x + TileCoord.y * TileBufferExtent.x;
+		TileClassificationBufferOutput[TileIndex + ViewTileCount * DirectionId] = MaxTheta;
+	}
+
+	//if (DirectionId == View.GeneralPurposeTweak)
+	//{
+	//	DebugOutput[TileCoord] = float4(MaxTheta / PI, TileDeviceZ, 0, 0);
+	//}
+} // MainCS()
diff --git a/Engine/Shaders/Private/SSRT/SSRTTileClassificationBuffer.ush b/Engine/Shaders/Private/SSRT/SSRTTileClassificationBuffer.ush
new file mode 100644
--- /dev/null
+++ b/Engine/Shaders/Private/SSRT/SSRTTileClassificationBuffer.ush
@@ -0,0 +1,40 @@
+// Copyright 1998-2019 Epic Games, Inc. All Rights Reserved.
+
+#pragma once
+
+#include "../Common.ush"
+
+
+#define TILE_RES_DIVISOR_LOG2 3
+#define TILE_RES_DIVISOR 8
+#define DIRECTIONALITY_DIVISION 8
+
+
+uint2 TileBufferExtent;
+uint ViewTileCount;
+uint MaxTileCount;
+
+Buffer<float> TileClassificationBuffer;
+
+
+/** Information about tiles in SGPRs. */
+struct FSSRTTileInfos
+{
+	float Directionality[DIRECTIONALITY_DIVISION];
+};
+
+
+/** Load tile informations using scalar load. */
+FSSRTTileInfos LoadTileInfos(uint2 TileCoord)
+{
+	uint TileIndex = TileCoord.x + TileCoord.y * TileBufferExtent.x;
+
+	FSSRTTileInfos TileInfos;
+
+	UNROLL_N(DIRECTIONALITY_DIVISION)
+	for (uint i = 0; i < DIRECTIONALITY_DIVISION; i++)
+	{
+		TileInfos.Directionality[i] = TileClassificationBuffer.Load(TileIndex + ViewTileCount * i);
+	}
+	return TileInfos;
+}
diff --git a/Engine/Source/Runtime/Renderer/Private/ScreenSpaceRayTracing.cpp b/Engine/Source/Runtime/Renderer/Private/ScreenSpaceRayTracing.cpp
--- a/Engine/Source/Runtime/Renderer/Private/ScreenSpaceRayTracing.cpp
+++ b/Engine/Source/Runtime/Renderer/Private/ScreenSpaceRayTracing.cpp
@@ -152,6 +152,54 @@ FLinearColor ComputeSSRParams(const FViewInfo& View, ESSRQuality SSRQuality, boo
 }
 
 
+
+BEGIN_SHADER_PARAMETER_STRUCT(FSSRTTileClassificationParameters, )
+	SHADER_PARAMETER(FIntPoint, TileBufferExtent)
+	SHADER_PARAMETER(int32, ViewTileCount)
+	SHADER_PARAMETER(int32, MaxTileCount)
+END_SHADER_PARAMETER_STRUCT()
+
+BEGIN_SHADER_PARAMETER_STRUCT(FSSRTTileClassificationResources, )
+	SHADER_PARAMETER_RDG_BUFFER(Buffer<float>, TileClassificationBuffer)
+END_SHADER_PARAMETER_STRUCT()
+
+BEGIN_SHADER_PARAMETER_STRUCT(FSSRTTileClassificationSRVs, )
+	SHADER_PARAMETER_RDG_BUFFER_SRV(Buffer<float>, TileClassificationBuffer)
+END_SHADER_PARAMETER_STRUCT()
+
+BEGIN_SHADER_PARAMETER_STRUCT(FSSRTTileClassificationUAVs, )
+	SHADER_PARAMETER_RDG_BUFFER_UAV(RWBuffer<float>, TileClassificationBufferOutput)
+END_SHADER_PARAMETER_STRUCT()
+
+FSSRTTileClassificationResources CreateTileClassificationResources(FRDGBuilder& GraphBuilder, const FViewInfo& View, FIntPoint MaxRenderTargetSize, FSSRTTileClassificationParameters* OutParameters)
+{
+	FIntPoint MaxTileBufferExtent = FIntPoint::DivideAndRoundUp(MaxRenderTargetSize, 8);
+	int32 MaxTileCount = MaxTileBufferExtent.X * MaxTileBufferExtent.Y;
+
+	OutParameters->TileBufferExtent = FIntPoint::DivideAndRoundUp(View.ViewRect.Size(), 8);
+	OutParameters->ViewTileCount = OutParameters->TileBufferExtent.X * OutParameters->TileBufferExtent.Y;
+
+	FSSRTTileClassificationResources Resources;
+	Resources.TileClassificationBuffer = GraphBuilder.CreateBuffer(FRDGBufferDesc::CreateBufferDesc(sizeof(float), MaxTileCount * 8), TEXT("SSRTTileClassification"));
+	return Resources;
+}
+
+FSSRTTileClassificationSRVs CreateSRVs(FRDGBuilder& GraphBuilder, const FSSRTTileClassificationResources& ClassificationResources)
+{
+	FSSRTTileClassificationSRVs SRVs;
+	SRVs.TileClassificationBuffer = GraphBuilder.CreateSRV(FRDGBufferSRVDesc(ClassificationResources.TileClassificationBuffer, PF_R32_FLOAT));
+	return SRVs;
+}
+
+FSSRTTileClassificationUAVs CreateUAVs(FRDGBuilder& GraphBuilder, const FSSRTTileClassificationResources& ClassificationResources)
+{
+	FSSRTTileClassificationUAVs UAVs;
+	UAVs.TileClassificationBufferOutput = GraphBuilder.CreateUAV(FRDGBufferUAVDesc(ClassificationResources.TileClassificationBuffer, PF_R32_FLOAT));
+	return UAVs;
+}
+
+
+
 BEGIN_SHADER_PARAMETER_STRUCT(FSSRCommonParameters, )
 	SHADER_PARAMETER(FLinearColor, SSRParams)
 	SHADER_PARAMETER_STRUCT_INCLUDE(FSceneTextureParameters, SceneTextures)
@@ -160,7 +208,6 @@ BEGIN_SHADER_PARAMETER_STRUCT(FSSRCommonParameters, )
 END_SHADER_PARAMETER_STRUCT()
 
 
-
 class FSSRQualityDim : SHADER_PERMUTATION_ENUM_CLASS("SSR_QUALITY", ESSRQuality);
 class FSSROutputForDenoiser : SHADER_PERMUTATION_BOOL("SSR_OUTPUT_FOR_DENOISER");
 
@@ -191,6 +238,33 @@ class FSSRTPrevFrameReductionCS : public FGlobalShader
 	END_SHADER_PARAMETER_STRUCT()
 };
 
+class FSSRTDiffuseTileClassificationCS : public FGlobalShader
+{
+	DECLARE_GLOBAL_SHADER(FSSRTDiffuseTileClassificationCS);
+	SHADER_USE_PARAMETER_STRUCT(FSSRTDiffuseTileClassificationCS, FGlobalShader);
+
+	static bool ShouldCompilePermutation(const FGlobalShaderPermutationParameters& Parameters)
+	{
+		return IsFeatureLevelSupported(Parameters.Platform, ERHIFeatureLevel::SM5);
+	}
+
+	BEGIN_SHADER_PARAMETER_STRUCT(FParameters, )
+		SHADER_PARAMETER(FVector2D, SamplePixelToHZBUV)
+
+		SHADER_PARAMETER_RDG_TEXTURE(Texture2D, HZBTexture)
+		SHADER_PARAMETER_SAMPLER(SamplerState, HZBTextureSampler)
+
+		SHADER_PARAMETER_STRUCT_REF(FViewUniformShaderParameters, View)
+		SHADER_PARAMETER_STRUCT_INCLUDE(FSceneTextureParameters, SceneTextures)
+		SHADER_PARAMETER_STRUCT_INCLUDE(FSceneTextureSamplerParameters, SceneTextureSamplers)
+
+		SHADER_PARAMETER_STRUCT_INCLUDE(FSSRTTileClassificationParameters, TileClassificationParameters)
+		SHADER_PARAMETER_STRUCT_INCLUDE(FSSRTTileClassificationUAVs, TileClassificationUAVs)
+
+		SHADER_PARAMETER_RDG_TEXTURE_UAV(RWTexture2D<float4>, DebugOutput)
+	END_SHADER_PARAMETER_STRUCT()
+};
+
 class FScreenSpaceReflectionsStencilPS : public FGlobalShader
 {
 	DECLARE_GLOBAL_SHADER(FScreenSpaceReflectionsStencilPS);
@@ -270,6 +344,9 @@ class FScreenSpaceDiffuseIndirectCS : public FGlobalShader
 		SHADER_PARAMETER_STRUCT_INCLUDE(FSceneTextureSamplerParameters, SceneTextureSamplers)
 		SHADER_PARAMETER_STRUCT_REF(FViewUniformShaderParameters, View)
 
+		SHADER_PARAMETER_STRUCT_INCLUDE(FSSRTTileClassificationParameters, ClassificationParameters)
+		SHADER_PARAMETER_STRUCT_INCLUDE(FSSRTTileClassificationSRVs, ClassificationSRVs)
+
 		SHADER_PARAMETER_RDG_TEXTURE_UAV(RWTexture<float4>, IndirectDiffuseOutput)
 		SHADER_PARAMETER_RDG_TEXTURE_UAV(RWTexture<float>,  AmbientOcclusionOutput)
 
@@ -283,7 +360,8 @@ class FScreenSpaceDiffuseIndirectCS : public FGlobalShader
 };
 
 
-IMPLEMENT_GLOBAL_SHADER(FSSRTPrevFrameReductionCS,        "/Engine/Private/SSRT/SSRTPrevFrameReduction.usf", "MainCS", SF_Compute);
+IMPLEMENT_GLOBAL_SHADER(FSSRTPrevFrameReductionCS, "/Engine/Private/SSRT/SSRTPrevFrameReduction.usf", "MainCS", SF_Compute);
+IMPLEMENT_GLOBAL_SHADER(FSSRTDiffuseTileClassificationCS, "/Engine/Private/SSRT/SSRTTileClassification.usf", "MainCS", SF_Compute);
 IMPLEMENT_GLOBAL_SHADER(FScreenSpaceReflectionsPS,        "/Engine/Private/SSRT/SSRTReflections.usf", "ScreenSpaceReflectionsPS", SF_Pixel);
 IMPLEMENT_GLOBAL_SHADER(FScreenSpaceReflectionsStencilPS, "/Engine/Private/SSRT/SSRTReflections.usf", "ScreenSpaceReflectionsStencilPS", SF_Pixel);
 IMPLEMENT_GLOBAL_SHADER(FScreenSpaceDiffuseIndirectCS,    "/Engine/Private/SSRT/SSRTDiffuseIndirect.usf", "MainCS", SF_Compute);
@@ -617,7 +695,7 @@ void RenderScreenSpaceDiffuseIndirect(
 	int32 RayCountPerPixel;
 	GetSSRTGIShaderOptionsForQuality(Quality, &GroupSize, &RayCountPerPixel);
 
-	FRDGTexture* HZBTexture = GraphBuilder.RegisterExternalTexture(View.HZB);
+	FRDGTexture* HZBTexture = GraphBuilder.RegisterExternalTexture(View.ClosestHZB);
 	FRDGTexture* ColorTexture = GraphBuilder.RegisterExternalTexture(TemporalAAHistory.RT[0]);
 
 	// Reproject and reduce previous frame color.
@@ -673,6 +751,49 @@ void RenderScreenSpaceDiffuseIndirect(
 			FComputeShaderUtils::GetGroupCount(View.ViewRect.Size(), GroupSize));
 	}
 
+	// Tile classify.
+	FSSRTTileClassificationParameters ClassificationParameters;
+	FSSRTTileClassificationResources ClassificationResources;
+	{
+		ClassificationResources = CreateTileClassificationResources(GraphBuilder, View, SceneTextures.SceneDepthBuffer->Desc.Extent, &ClassificationParameters);
+
+		FIntPoint ThreadCount = ClassificationParameters.TileBufferExtent;
+
+		FSSRTDiffuseTileClassificationCS::FParameters* PassParameters = GraphBuilder.AllocParameters<FSSRTDiffuseTileClassificationCS::FParameters>();
+		PassParameters->SamplePixelToHZBUV = FVector2D(
+			0.5f / float(HZBTexture->Desc.Extent.X),
+			0.5f / float(HZBTexture->Desc.Extent.Y));
+
+		PassParameters->HZBTexture = HZBTexture;
+		PassParameters->HZBTextureSampler = TStaticSamplerState<SF_Point>::GetRHI();
+
+		PassParameters->SceneTextures = SceneTextures;
+		SetupSceneTextureSamplers(&PassParameters->SceneTextureSamplers);
+		PassParameters->View = View.ViewUniformBuffer;
+
+		PassParameters->TileClassificationParameters = ClassificationParameters;
+		PassParameters->TileClassificationUAVs = CreateUAVs(GraphBuilder, ClassificationResources);
+
+		{
+			FRDGTextureDesc DebugDesc = FRDGTextureDesc::Create2DDesc(
+				SceneTextures.SceneDepthBuffer->Desc.Extent / 8,
+				PF_FloatRGBA,
+				FClearValueBinding::Transparent,
+				/* InFlags = */ TexCreate_None,
+				/* InTargetableFlags = */ TexCreate_ShaderResource | TexCreate_UAV,
+				/* bInForceSeparateTargetAndShaderResource = */ false);
+
+			PassParameters->DebugOutput = GraphBuilder.CreateUAV(GraphBuilder.CreateTexture(DebugDesc, TEXT("DebugSSRTTiles")));
+		}
+
+		TShaderMapRef<FSSRTDiffuseTileClassificationCS> ComputeShader(View.ShaderMap);
+		FComputeShaderUtils::AddPass(
+			GraphBuilder,
+			RDG_EVENT_NAME("ScreenSpaceDiffuseClassification %dx%d", ThreadCount.X, ThreadCount.Y),
+			*ComputeShader,
+			PassParameters,
+			FComputeShaderUtils::GetGroupCount(ThreadCount, 8));
+	}
 
 	{
 		// Allocate outputs.
@@ -725,6 +846,9 @@ void RenderScreenSpaceDiffuseIndirect(
 		SetupSceneTextureSamplers(&PassParameters->SceneTextureSamplers);
 		PassParameters->View = View.ViewUniformBuffer;
 	
+		PassParameters->ClassificationParameters = ClassificationParameters;
+		PassParameters->ClassificationSRVs = CreateSRVs(GraphBuilder, ClassificationResources);
+
 		PassParameters->IndirectDiffuseOutput = GraphBuilder.CreateUAV(OutDenoiserInputs->Color);
 		PassParameters->AmbientOcclusionOutput = GraphBuilder.CreateUAV(OutDenoiserInputs->AmbientOcclusionMask);
 		PassParameters->ScreenSpaceRayTracingDebugOutput = CreateScreenSpaceRayTracingDebugUAV(GraphBuilder, OutDenoiserInputs->Color->Desc);
