Copying Sky //UE4/Dev-Volumetrics to Dev-Rendering (//UE4/Dev-Rendering)
Last copy up before SIGGRAPH contains a bunch of clean up and auto toggle to raymarching in space.

diff --git a/Engine/Config/BaseScalability.ini b/Engine/Config/BaseScalability.ini
--- a/Engine/Config/BaseScalability.ini
+++ b/Engine/Config/BaseScalability.ini
@@ -390,10 +390,14 @@ r.SSS.HalfRes=1
 r.EmitterSpawnRateScale=0.125
 r.ParticleLightQuality=0
 r.DetailLevel=0
+r.SkyAtmosphere.AerialPerspectiveLUT.FastApplyOnOpaque=1 ; Always have FastSkyLUT 1 in this case to avoid wrong sky
+r.SkyAtmosphere.AerialPerspectiveLUT.SampleCountPerSlice=1
+r.SkyAtmosphere.AerialPerspectiveLUT.DepthResolution=8
 r.SkyAtmosphere.FastSkyLUT=1
-r.SkyAtmosphere.AerialPerspectiveLUT.FastApplyOnOpaque=1 ; Always have FastSkyLUT 1 in third case to avoid wrong sky
 r.SkyAtmosphere.FastSkyLUT.SampleCountMin=2.0
 r.SkyAtmosphere.FastSkyLUT.SampleCountMin=16.0
+r.SkyAtmosphere.TransmittanceLUT.UseSmallFormat=1
+r.SkyAtmosphere.MultiScatteringLUT.UseSmallFormat=1
 
 [EffectsQuality@1]
 r.TranslucencyLightingVolumeDim=32
@@ -411,7 +415,7 @@ r.EmitterSpawnRateScale=0.25
 r.ParticleLightQuality=0
 r.DetailLevel=1
 r.SkyAtmosphere.FastSkyLUT=1
-r.SkyAtmosphere.AerialPerspectiveLUT.FastApplyOnOpaque=1 ; Always have FastSkyLUT 1 in third case to avoid wrong sky
+r.SkyAtmosphere.AerialPerspectiveLUT.FastApplyOnOpaque=1 ; Always have FastSkyLUT 1 in this case to avoid wrong sky
 
 [EffectsQuality@2]
 r.TranslucencyLightingVolumeDim=48
@@ -429,7 +433,7 @@ r.EmitterSpawnRateScale=0.5
 r.ParticleLightQuality=1
 r.DetailLevel=2
 r.SkyAtmosphere.FastSkyLUT=1
-r.SkyAtmosphere.AerialPerspectiveLUT.FastApplyOnOpaque=0
+r.SkyAtmosphere.AerialPerspectiveLUT.FastApplyOnOpaque=1 ; Always have FastSkyLUT 1 in this case to avoid wrong sky
 
 [EffectsQuality@3]
 r.TranslucencyLightingVolumeDim=64
@@ -446,8 +450,8 @@ r.SSS.HalfRes=1
 r.EmitterSpawnRateScale=1.0
 r.ParticleLightQuality=2
 r.DetailLevel=3
-r.SkyAtmosphere.FastSkyLUT=0
-r.SkyAtmosphere.AerialPerspectiveLUT.FastApplyOnOpaque=0
+r.SkyAtmosphere.FastSkyLUT=1
+r.SkyAtmosphere.AerialPerspectiveLUT.FastApplyOnOpaque=1 ; Always have FastSkyLUT 1 in this case to avoid wrong sky
 
 [EffectsQuality@Cine]
 r.TranslucencyLightingVolumeDim=64
@@ -466,11 +470,11 @@ r.ParticleLightQuality=2
 r.DetailLevel=4
 r.SkyAtmosphere.FastSkyLUT=0
 r.r.SkyAtmosphere.AerialPerspectiveLUT.FastApplyOnOpaque=0
-r.SkyAtmosphere.SampleCountMin=16.0
+r.SkyAtmosphere.SampleCountMin=8.0
 r.SkyAtmosphere.SampleCountMax=32.0
-r.SkyAtmosphere.AerialPerspectiveLUT.DepthResolution=64.0
-r.SkyAtmosphere.AerialPerspectiveLUT.SampleCountPerSlice=4.0
-r.SkyAtmosphere.TransmittanceLUT.SampleCount=20.0
+r.SkyAtmosphere.AerialPerspectiveLUT.DepthResolution=32.0
+r.SkyAtmosphere.AerialPerspectiveLUT.SampleCountPerSlice=2.0
+r.SkyAtmosphere.TransmittanceLUT.SampleCount=30.0
 r.SkyAtmosphere.MultiScatteringLUT.SampleCount=20.0
 
 ;-----------------------------------------------------------------------------------------------------------------
diff --git a/Engine/Shaders/Private/SkyAtmosphere.usf b/Engine/Shaders/Private/SkyAtmosphere.usf
--- a/Engine/Shaders/Private/SkyAtmosphere.usf
+++ b/Engine/Shaders/Private/SkyAtmosphere.usf
@@ -126,10 +126,6 @@ float3 GetCameraPlanetPos()
 	return GetCameraWorldPos()*CM_TO_SKY_UNIT + float3(0, 0, Atmosphere.BottomRadius);
 }
 
-
-
-
-
 bool MoveToTopAtmosphere(inout float3 WorldPos, in float3 WorldDir, in float AtmosphereTopRadius)
 {
 	float ViewHeight = length(WorldPos);
@@ -161,6 +157,11 @@ bool MoveToTopAtmosphere(inout float3 WorldPos, in float3 WorldDir, in float Atm
 
 
 
+
+////////////////////////////////////////////////////////////
+// Participating medium properties
+////////////////////////////////////////////////////////////
+
 float RayleighPhase(float CosTheta)
 {
 	float Factor = 3.0f / (16.0f * PI);
@@ -176,15 +177,6 @@ float HgPhase(float G, float CosTheta)
 	return Numer / (4.0f * PI * Denom * sqrt(Denom));
 }
 
-
-
-
-
-
-
-
-
-
 float3 GetAlbedo(float3 Scattering, float3 Extinction)
 {
 	return Scattering / max(0.001f, Extinction);
@@ -248,32 +240,6 @@ MediumSampleRGB SampleMediumRGB(in float3 WorldPos)
 
 
 
-
-#define DEFAULT_SAMPLE_OFFSET 0.3f
-float SkyAtmosphereNoise(float2 UV)
-{
-//	return DEFAULT_SAMPLE_OFFSET;
-//	return float(Rand3DPCG32(int3(UV.x, UV.y, S)).x) / 4294967296.0f;
-#if VIEWDATA_AVAILABLE && PER_PIXEL_NOISE
-	return InterleavedGradientNoise(UV.xy, float(View.StateFrameIndexMod8));
-#else
-	return DEFAULT_SAMPLE_OFFSET;
-#endif
-}
-
-
-
-
-
-
-
-
-
-
-
-
-
-
 ////////////////////////////////////////////////////////////
 // LUT functions
 ////////////////////////////////////////////////////////////
@@ -306,7 +272,8 @@ void LutTransmittanceParamsToUv(in float ViewHeight, in float ViewZenithCosAngle
 }
 
 // SkyViewLut is a new texture used for fast sky rendering.
-// 
+// It is low resolution of the sky rendering around the camera,
+// basically a lat/long parameterisation with more texel close to the horizon for more accuracy during sun set.
 
 void UvToSkyViewLutParams(out float3 ViewDir, in float ViewHeight, in float2 UV)
 {
@@ -350,18 +317,9 @@ void UvToSkyViewLutParams(out float3 ViewDir, in float ViewHeight, in float2 UV)
 
 
 
-
-
-
-
-
-
-
-
-
-
-
-
+////////////////////////////////////////////////////////////
+// Utilities
+////////////////////////////////////////////////////////////
 
 #if USE_PREEXPOSURE && VIEWDATA_AVAILABLE
 #define ViewPreExposure View.PreExposure
@@ -388,9 +346,6 @@ Texture3D<float4> CameraAerialPerspectiveVolumeTexture;
 
 
 
-
-
-
 float3 GetLightDiskLuminance(float3 WorldPos, float3 WorldDir, uint LightIndex)
 {
 #if SOURCE_DISK_ENABLED
@@ -445,11 +400,29 @@ float3 GetTransmittance(in float LightZenithCosAngle, in float PHeight)
 }
 
 
+#define DEFAULT_SAMPLE_OFFSET 0.3f
+float SkyAtmosphereNoise(float2 UV)
+{
+	//	return DEFAULT_SAMPLE_OFFSET;
+	//	return float(Rand3DPCG32(int3(UV.x, UV.y, S)).x) / 4294967296.0f;
+#if VIEWDATA_AVAILABLE && PER_PIXEL_NOISE
+	return InterleavedGradientNoise(UV.xy, float(View.StateFrameIndexMod8));
+#else
+	return DEFAULT_SAMPLE_OFFSET;
+#endif
+}
+
+
+
 
 
 
 
 
+////////////////////////////////////////////////////////////
+// Main scattering/transmitance integration fucntion
+////////////////////////////////////////////////////////////
+
 struct SingleScatteringResult
 {
 	float3 L;						// Scattered light (luminance)
@@ -556,7 +529,6 @@ SingleScatteringResult IntegrateSingleScatteredLuminance(
 	float RayleighPhaseValueLight1 = RayleighPhase(cosTheta);
 #endif
 
-
 	// Ray march the atmosphere to integrate optical depth
 	float3 L = 0.0f;
 	float3 Throughput = 1.0f;
@@ -572,6 +544,7 @@ SingleScatteringResult IntegrateSingleScatteredLuminance(
 	float PixelNoise = PER_PIXEL_NOISE ? SkyAtmosphereNoise(PixPos.xy) : DEFAULT_SAMPLE_OFFSET;
 	for (float SampleI = 0.0f; SampleI < SampleCount; SampleI += 1.0f)
 	{
+		// Compute current ray t and sample point P
 		if (Sampling.VariableSampleCount)
 		{
 			// More expenssive but artefact free
@@ -599,15 +572,15 @@ SingleScatteringResult IntegrateSingleScatteredLuminance(
 			t = tMax * SampleI / SampleCount;
 		}
 		float3 P = WorldPos + t * WorldDir;
+		float PHeight = length(P);
 
+		// Sample the medium
 		MediumSampleRGB Medium = SampleMediumRGB(P);
 		const float3 SampleOpticalDepth = Medium.Extinction * dt * AerialPespectiveViewDistanceScale;
 		const float3 SampleTransmittance = exp(-SampleOpticalDepth);
 		OpticalDepth += SampleOpticalDepth;
 
-		float PHeight = length(P);
-
-		// Phase and trnasmittance for light 0
+		// Phase and transmittance for light 0
 		const float3 UpVector = P / PHeight;
 		float Light0ZenithCosAngle = dot(Light0Dir, UpVector);
 		float3 TransmittanceToLight0 = GetTransmittance(Light0ZenithCosAngle, PHeight);
@@ -620,7 +593,6 @@ SingleScatteringResult IntegrateSingleScatteredLuminance(
 		{
 			PhaseTimesScattering0 = Medium.Scattering * uniformPhase;
 		}
-
 #if SECOND_ATMOSPHERE_LIGHT_ENABLED
 		// Phase and transmittance for light 1
 		float Light1ZenithCosAngle = dot(Light1Dir, UpVector);
@@ -661,7 +633,6 @@ SingleScatteringResult IntegrateSingleScatteredLuminance(
 		// The way to fix that is to use a proper analytical integration as porposed in slide 28 of http://www.frostbite.com/2015/08/physically-based-unified-volumetric-rendering-in-frostbite/ 
 		// However, it is possible to disable as it can also work using simple power serie sum unroll up to 5th order. The rest of the orders has a really low contribution. 
 #define MULTI_SCATTERING_POWER_SERIE 0 
-
 #if MULTI_SCATTERING_POWER_SERIE==0 
 		// 1 is the integration of luminance over the 4pi of a sphere, and assuming an isotropic phase function of 1.0/(4*PI) 
 		Result.MultiScatAs1 += Throughput * Medium.Scattering * 1.0f * dt;
@@ -671,7 +642,6 @@ SingleScatteringResult IntegrateSingleScatteredLuminance(
 		Result.MultiScatAs1 += Throughput * MSint;
 #endif 
 
-
 #if 0
 		L += Throughput * S * dt;
 		Throughput *= SampleTransmittance;
@@ -720,30 +690,11 @@ SingleScatteringResult IntegrateSingleScatteredLuminance(
 
 
 
-
-
-
-
-
-
-
-
-
-
-
-
-
-
-
-
-
-
 
 
 
 float StartDepthZ;
 
-
 void SkyAtmosphereVS(
 	in uint VertexId : SV_VertexID,
 	out float4 Position : SV_POSITION)
@@ -763,12 +714,6 @@ void SkyAtmosphereVS(
 
 
 
-
-
-
-
-
-
 float4 PrepareOutput(float3 Luminance, float3 Transmittance = float3(1.0f, 1.0f, 1.0f))
 {
 	const float GreyScaleTransmittance = dot(Transmittance, float3(1.0f / 3.0f, 1.0f / 3.0f, 1.0f / 3.0f));
@@ -788,6 +733,7 @@ void RenderSkyAtmosphereRayMarchingPS(
 	float2 PixPos = SVPos.xy;
 	float2 UvBuffer = PixPos * View.BufferSizeAndInvSize.zw;	// Uv for depth buffer read (size can be larger than viewport)
 
+	// Debug print a texture
 	//const float displaySize = 256.0f;
 	//if(all(PixPos<displaySize))
 	//{
@@ -917,8 +863,6 @@ void RenderSkyAtmosphereRayMarchingPS(
 
 
 
-
-
 #ifndef THREADGROUP_SIZE
 #define THREADGROUP_SIZE 1
 #endif
@@ -1052,7 +996,6 @@ void RenderMultiScatteredLuminanceLutCS(uint3 ThreadId : SV_DispatchThreadID)
 
 
 
-
 RWTexture2D<float3> SkyViewLutUAV;
 
 [numthreads(THREADGROUP_SIZE, THREADGROUP_SIZE, 1)]
@@ -1066,8 +1009,6 @@ void RenderSkyViewLutCS(uint3 ThreadId : SV_DispatchThreadID)
 	float ViewHeight = length(WorldPos);
 
 
-
-	// TODO explain the mapping that is happening here...
 	// TODO we need a valid local basis for lights and we need to transform them in it. 
 	// ===> Without that, the Sky lut is only valid for camera is close to the origin (at the Z-top of the planet)
 	float3 WorldDir;
@@ -1105,7 +1046,6 @@ void RenderSkyViewLutCS(uint3 ThreadId : SV_DispatchThreadID)
 
 
 
-
 RWTexture2D<float3> DistantSkyLightLutUAV;
 Buffer<float4> UniformSphereSamplesBuffer;
 float4 AtmosphereLightDirection0;
@@ -1201,8 +1141,6 @@ void RenderDistantSkyLightLutCS(uint3 ThreadId : SV_DispatchThreadID)
 
 
 
-
-
 RWTexture3D<float4> CameraAerialPerspectiveVolumeUAV;
 
 [numthreads(THREADGROUP_SIZE, THREADGROUP_SIZE, THREADGROUP_SIZE)]
@@ -1258,7 +1196,6 @@ void RenderCameraAerialPerspectiveVolumeCS(uint3 ThreadId : SV_DispatchThreadID)
 			float3 MiddlePointOnGround = normalize(MiddlePoint) * Atmosphere.BottomRadius;// *1.001f;
 			VoxelWorldPos = CamPos + 2.0f * (MiddlePointOnGround - CamPos);
 
-
 			//CameraAerialPerspectiveVolumeUAV[ThreadId] = float4(1, 0, 0, 0);
 			//return; // debug
 		}
@@ -1266,7 +1203,7 @@ void RenderCameraAerialPerspectiveVolumeCS(uint3 ThreadId : SV_DispatchThreadID)
 		{
 			//No obstruction from the planet, so use the point on the ground
 			VoxelWorldPos = normalize(VoxelWorldPos) * (Atmosphere.BottomRadius);
-			//VoxelWorldPos = CamPos + CameraToVoxelDir * PlanetNearT;		// better but gives visuall artefact as visible voxels on a simple plane at altitude 0
+			//VoxelWorldPos = CamPos + CameraToVoxelDir * PlanetNearT;		// better match but gives visual artefact as visible voxels on a simple plane at altitude 0
 
 			//CameraAerialPerspectiveVolumeUAV[ThreadId] = float4(0, 1, 0, 0);
 			//return; // debug
@@ -1310,7 +1247,7 @@ void RenderCameraAerialPerspectiveVolumeCS(uint3 ThreadId : SV_DispatchThreadID)
 	const float DeviceZ = FarDepthValue;
 	const bool MieRayPhase = true;
 	const float AerialPespectiveViewDistanceScale = SkyAtmosphere.AerialPespectiveViewDistanceScale;
-	SingleScatteringResult ss = IntegrateSingleScatteredLuminance(
+	SingleScatteringResult ss = IntegrateSingleScatteredLuminance( 
 		float4(PixPos, 0.0f, 1.0f), RayStartWorldPos, WorldDir,
 		Ground, Sampling, DeviceZ, MieRayPhase,
 		View.AtmosphereLightDirection[0].xyz, View.AtmosphereLightDirection[1].xyz, View.AtmosphereLightColor[0].rgb, View.AtmosphereLightColor[1].rgb,
@@ -1323,7 +1260,6 @@ void RenderCameraAerialPerspectiveVolumeCS(uint3 ThreadId : SV_DispatchThreadID)
 
 
 
-
 //////////////////////////////////////////////////////////////////////////
 // Debug shaders
 //////////////////////////////////////////////////////////////////////////
diff --git a/Engine/Source/Runtime/Renderer/Private/SkyAtmosphereRendering.cpp b/Engine/Source/Runtime/Renderer/Private/SkyAtmosphereRendering.cpp
--- a/Engine/Source/Runtime/Renderer/Private/SkyAtmosphereRendering.cpp
+++ b/Engine/Source/Runtime/Renderer/Private/SkyAtmosphereRendering.cpp
@@ -75,7 +75,7 @@ static TAutoConsoleVariable<float> CVarSkyAtmosphereFastSkyLUTWidth(
 	ECVF_RenderThreadSafe | ECVF_Scalability);
 
 static TAutoConsoleVariable<float> CVarSkyAtmosphereFastSkyLUTHeight(
-	TEXT("r.SkyAtmosphere.FastSkyLUT.Height"), 108,
+	TEXT("r.SkyAtmosphere.FastSkyLUT.Height"), 104,
 	TEXT(""),
 	ECVF_RenderThreadSafe | ECVF_Scalability);
 
@@ -94,7 +94,7 @@ static TAutoConsoleVariable<int32> CVarSkyAtmosphereAerialPerspectiveDepthTest(
 ////////////////////////////////////////////////////////////////////////// Aerial perspective LUT
 
 static TAutoConsoleVariable<float> CVarSkyAtmosphereAerialPerspectiveLUTDepthResolution(
-	TEXT("r.SkyAtmosphere.AerialPerspectiveLUT.DepthResolution"), 32.0f,
+	TEXT("r.SkyAtmosphere.AerialPerspectiveLUT.DepthResolution"), 16.0f,
 	TEXT("The number of depth slice to use for the aerial perspective volume texture."),
 	ECVF_RenderThreadSafe | ECVF_Scalability);
 
@@ -1102,6 +1102,7 @@ void FSceneRenderer::RenderSkyAtmosphere(FRHICommandListImmediate& RHICmdList)
 	SCOPED_GPU_STAT(RHICmdList, SkyAtmosphere);
 
 	FSkyAtmosphereRenderSceneInfo& SkyInfo = *Scene->GetSkyAtmosphereSceneInfo();
+	const FAtmosphereSetup& Atmosphere = SkyInfo.GetAtmosphereSetup();
 	const bool bMultiScattering = SkyInfo.IsMultiScatteringEnabled();
 	const bool bFastSky = CVarSkyAtmosphereFastSkyLUT.GetValueOnRenderThread() > 0;
 	const bool bFastAerialPespective = CVarSkyAtmosphereAerialPerspectiveApplyOnOpaque.GetValueOnRenderThread() > 0;
@@ -1127,6 +1128,13 @@ void FSceneRenderer::RenderSkyAtmosphere(FRHICommandListImmediate& RHICmdList)
 
 		const bool bLightDiskEnabled = !View.bIsReflectionCapture;
 
+		const FVector ViewOrigin = View.ViewMatrices.GetViewOrigin();
+		const float KmToCm = 100000.0f;
+		const FVector PlanetOrigin = FVector(0.0f, 0.0f, -Atmosphere.BottomRadius * KmToCm);
+		const float TopOfAtmosphere = Atmosphere.TopRadius * KmToCm;
+		const float SafeEdge = 1000.0f;	// 10 meters
+		const bool ForceRayMarching = (FVector::Distance(ViewOrigin, PlanetOrigin) - TopOfAtmosphere - SafeEdge) > 0.0f;
+
 		// Render the sky, and optionally the atmosphere aerial perspective, on the scene luminance buffer
 		if (!bIsMobilePlatformAndOnlySkyPixel)
 		{
@@ -1138,8 +1146,8 @@ void FSceneRenderer::RenderSkyAtmosphere(FRHICommandListImmediate& RHICmdList)
 		{
 			FRenderSkyAtmospherePS::FPermutationDomain PermutationVector;
 			PermutationVector.Set<FSkyPermutationMultiScatteringApprox>(bMultiScattering);
-			PermutationVector.Set<FFastSky>(bFastSky);
-			PermutationVector.Set<FFastAerialPespective>(bFastAerialPespective);
+			PermutationVector.Set<FFastSky>(bFastSky && !ForceRayMarching);
+			PermutationVector.Set<FFastAerialPespective>(bFastAerialPespective && !ForceRayMarching);
 			PermutationVector.Set<FSourceDiskEnabled>(bLightDiskEnabled);
 			PermutationVector.Set<FOnlySkyPixel>(bIsMobilePlatformAndOnlySkyPixel);
 			PermutationVector.Set<FSecondAtmosphereLight>(bSecondAtmosphereLightEnabled);
@@ -1174,7 +1182,6 @@ void FSceneRenderer::RenderSkyAtmosphere(FRHICommandListImmediate& RHICmdList)
 				float HalfHorizontalFOV = FMath::Atan(1.0f / ViewProjectionMatrix.M[0][0]);
 				float HalfVerticalFOV = FMath::Atan(1.0f / ViewProjectionMatrix.M[1][1]);
 				const float HalfFOV = View.FOV * 0.5f * PI / 180.0f;
-				const float KmToCm = 100000.0f;
 				const float StartDepthView = FMath::Cos(FMath::Max(HalfHorizontalFOV, HalfVerticalFOV)) * InternalCommonParameters.CameraAerialPerspectiveVolumeStartDepth * KmToCm;
 				const FVector4 Projected = ViewProjectionMatrix.TransformFVector4(FVector4(0.0f, 0.0f, StartDepthView, 1.0f));
 				VsPassParameters.StartDepthZ = Projected.Z / Projected.W;
diff --git a/Engine/Source/Runtime/Renderer/Private/SkyAtmosphereRendering.h b/Engine/Source/Runtime/Renderer/Private/SkyAtmosphereRendering.h
--- a/Engine/Source/Runtime/Renderer/Private/SkyAtmosphereRendering.h
+++ b/Engine/Source/Runtime/Renderer/Private/SkyAtmosphereRendering.h
@@ -73,7 +73,7 @@ public:
 	TRefCountPtr<IPooledRenderTarget>& GetDistantSkyLightLutTexture() { return DistantSkyLightLutTexture; }
 	FTextureRHIRef GetDistantSkyLightLutTextureRHI();
 
-	const FAtmosphereSetup GetAtmosphereSetup() const { return AtmosphereSetup; }
+	const FAtmosphereSetup& GetAtmosphereSetup() const { return AtmosphereSetup; }
 	const FAtmosphereUniformShaderParameters* GetAtmosphereShaderParameters() const { return &AtmosphereUniformShaderParameters; }
 
 	bool bStaticLightingBuilt;
