Fixed few ray tracing shader compile warnings

* Vector truncation in LoadVertexFactoryInput()
* Unused expression result due to wrong clip(x) macro implementation

diff --git a/Engine/Shaders/Private/GeometryCacheVertexFactory.ush b/Engine/Shaders/Private/GeometryCacheVertexFactory.ush
--- a/Engine/Shaders/Private/GeometryCacheVertexFactory.ush
+++ b/Engine/Shaders/Private/GeometryCacheVertexFactory.ush
@@ -74,7 +74,7 @@ FVertexFactoryInput LoadVertexFactoryInput(uint TriangleIndex, int VertexIndex)
 	Input.Position.x = GeomCacheMVF.Position[VertexId * 3 + 0];
 	Input.Position.y = GeomCacheMVF.Position[VertexId * 3 + 1];
 	Input.Position.z = GeomCacheMVF.Position[VertexId * 3 + 2];
-	Input.TangentX = GeomCacheMVF.TangentX[VertexId];
+	Input.TangentX = GeomCacheMVF.TangentX[VertexId].xyz;
 	Input.TangentZ = GeomCacheMVF.TangentZ[VertexId];
 	Input.Color = GeomCacheMVF.Color[VertexId];
     Input.MotionBlurData.x = GeomCacheMVF.MotionBlurData[VertexId * 3 + 0];
diff --git a/Engine/Shaders/Private/RayTracing/RayTracingMaterialHitShaders.usf b/Engine/Shaders/Private/RayTracing/RayTracingMaterialHitShaders.usf
--- a/Engine/Shaders/Private/RayTracing/RayTracingMaterialHitShaders.usf
+++ b/Engine/Shaders/Private/RayTracing/RayTracingMaterialHitShaders.usf
@@ -11,7 +11,7 @@
 #endif
 
 // #dxr_todo: UE-72537 move to common.ush
-#define clip(x) (x)
+#define clip(x)
 #define ddx(x) 0
 #define ddy(x) 0
 #define fwidth(x) 0
