Fixes polychromatic penumbra harmonics denoiser's history rejection

diff --git a/Engine/Shaders/Private/ScreenSpaceDenoise/SSDSignalBufferEncoding.ush b/Engine/Shaders/Private/ScreenSpaceDenoise/SSDSignalBufferEncoding.ush
--- a/Engine/Shaders/Private/ScreenSpaceDenoise/SSDSignalBufferEncoding.ush
+++ b/Engine/Shaders/Private/ScreenSpaceDenoise/SSDSignalBufferEncoding.ush
@@ -186,6 +186,11 @@ void DecodeMultiplexedSignals(
 		{
 			float Weight = RawSample[2].r;
 
+			if (bNormalizeSample)
+			{
+				Weight = Weight > 0 ? 1 : 0;
+			}
+
 			OutSamples.Array[0].SampleCount = Weight;
 			OutSamples.Array[0].MissCount = RawSample[0].a * Weight;
 			OutSamples.Array[0].SceneColor.rgb = RawSample[0].rgb * Weight;
@@ -194,6 +199,11 @@ void DecodeMultiplexedSignals(
 		// Decode specular harmonic
 		{
 			float Weight = RawSample[2].g;
+			
+			if (bNormalizeSample)
+			{
+				Weight = Weight > 0 ? 1 : 0;
+			}
 
 			OutSamples.Array[1].SampleCount = Weight;
 			OutSamples.Array[1].MissCount = RawSample[1].a * Weight;
