Copying //UE4/Dev-Volumetrics to Dev-Rendering (//UE4/Dev-Rendering)

diff --git a/Engine/Config/BaseScalability.ini b/Engine/Config/BaseScalability.ini
--- a/Engine/Config/BaseScalability.ini
+++ b/Engine/Config/BaseScalability.ini
@@ -469,7 +469,7 @@ r.EmitterSpawnRateScale=1.0
 r.ParticleLightQuality=2
 r.DetailLevel=4
 r.SkyAtmosphere.FastSkyLUT=0
-r.r.SkyAtmosphere.AerialPerspectiveLUT.FastApplyOnOpaque=0
+r.SkyAtmosphere.AerialPerspectiveLUT.FastApplyOnOpaque=0
 r.SkyAtmosphere.SampleCountMin=8.0
 r.SkyAtmosphere.SampleCountMax=32.0
 r.SkyAtmosphere.AerialPerspectiveLUT.DepthResolution=32.0
diff --git a/Engine/Shaders/Private/BasePassPixelShader.usf b/Engine/Shaders/Private/BasePassPixelShader.usf
--- a/Engine/Shaders/Private/BasePassPixelShader.usf
+++ b/Engine/Shaders/Private/BasePassPixelShader.usf
@@ -33,8 +33,10 @@
 
 #if NEEDS_BASEPASS_PIXEL_FOGGING || NEEDS_BASEPASS_PIXEL_VOLUMETRIC_FOGGING
 	#include "HeightFogCommon.ush"
+#if PROJECT_SUPPORT_SKY_ATMOSPHERE
 	#include "SkyAtmosphereCommon.ush"
 #endif
+#endif
 
 #include "ReflectionEnvironmentShared.ush"
 
@@ -1129,7 +1131,7 @@ ShadingModel);
 	}
 #endif
 
-#if NEEDS_BASEPASS_PIXEL_FOGGING
+#if NEEDS_BASEPASS_PIXEL_FOGGING && PROJECT_SUPPORT_SKY_ATMOSPHERE
 	if (ResolvedView.SkyAtmosphereApplyCameraAerialPerspectiveVolume > 0.0f)
 	{
 		const float OneOverPreExposure = USE_PREEXPOSURE ? ResolvedView.OneOverPreExposure : 1.0f;
@@ -1137,11 +1139,10 @@ ShadingModel);
 		float4 NDCPosition = mul(float4(MaterialParameters.AbsoluteWorldPosition.xyz, 1), ResolvedView.WorldToClip);
 
 		// Sample the aerial perspective (AP).
-		float4 AP = GetAerialPerspectiveLuminanceTransmittance(
+		Fogging = GetAerialPerspectiveLuminanceTransmittanceWithFogOver(
 			NDCPosition, MaterialParameters.AbsoluteWorldPosition.xyz*CM_TO_SKY_UNIT, ResolvedView.WorldCameraOrigin.xyz*CM_TO_SKY_UNIT,
 			View.CameraAerialPerspectiveVolume, View.CameraAerialPerspectiveVolumeSampler,
 			ResolvedView.SkyAtmosphereCameraAerialPerspectiveVolumeDepthResolutionInv,
-			ResolvedView.SkyAtmosphereCameraAerialPerspectiveVolumeDepthResolutionInv,
 			ResolvedView.SkyAtmosphereCameraAerialPerspectiveVolumeDepthResolution,
 			ResolvedView.SkyAtmosphereCameraAerialPerspectiveVolumeStartDepth,
 			ResolvedView.SkyAtmosphereCameraAerialPerspectiveVolumeDepthSliceLength,
diff --git a/Engine/Shaders/Private/BasePassVertexShader.usf b/Engine/Shaders/Private/BasePassVertexShader.usf
--- a/Engine/Shaders/Private/BasePassVertexShader.usf
+++ b/Engine/Shaders/Private/BasePassVertexShader.usf
@@ -18,7 +18,7 @@
 #include "ForwardLightingCommon.ush"
 #endif
 
-#if NEEDS_BASEPASS_VERTEX_FOGGING
+#if NEEDS_BASEPASS_VERTEX_FOGGING && PROJECT_SUPPORT_SKY_ATMOSPHERE
 #include "SkyAtmosphereCommon.ush"
 #endif
 
@@ -128,11 +128,12 @@ void Main(
 	Output.BasePassInterpolants.VertexFog = CalculateHeightFog(WorldPosition.xyz - ResolvedView.TranslatedWorldCameraOrigin);
 	#endif
 
+#if PROJECT_SUPPORT_SKY_ATMOSPHERE
 	if (ResolvedView.SkyAtmosphereApplyCameraAerialPerspectiveVolume > 0.0f)
 	{
 		const float OneOverPreExposure = USE_PREEXPOSURE ? ResolvedView.OneOverPreExposure : 1.0f;
 		// Sample the aerial perspective (AP). It is also blended under the VertexFog parameter.
-		float4 AP = GetAerialPerspectiveLuminanceTransmittance(
+		Output.BasePassInterpolants.VertexFog = GetAerialPerspectiveLuminanceTransmittanceWithFogOver(
 			Output.Position, WorldPosition.xyz*CM_TO_SKY_UNIT, ResolvedView.TranslatedWorldCameraOrigin*CM_TO_SKY_UNIT,
 			View.CameraAerialPerspectiveVolume, View.CameraAerialPerspectiveVolumeSampler,
 			ResolvedView.SkyAtmosphereCameraAerialPerspectiveVolumeDepthResolutionInv,
@@ -143,6 +144,7 @@ void Main(
 			OneOverPreExposure, Output.BasePassInterpolants.VertexFog);
 	}
 #endif
+#endif
 
 #if TRANSLUCENCY_ANY_PERVERTEX_LIGHTING
 	float3 WorldPositionForVertexLightingTranslated = VertexFactoryGetPositionForVertexLighting(Input, VFIntermediates, WorldPosition.xyz);
diff --git a/Engine/Shaders/Private/MaterialTemplate.ush b/Engine/Shaders/Private/MaterialTemplate.ush
--- a/Engine/Shaders/Private/MaterialTemplate.ush
+++ b/Engine/Shaders/Private/MaterialTemplate.ush
@@ -73,7 +73,7 @@
 	#include "/Engine/Private/AtmosphereCommon.ush"
 #endif
 
-#if MATERIAL_SKY_ATMOSPHERE
+#if MATERIAL_SKY_ATMOSPHERE && PROJECT_SUPPORT_SKY_ATMOSPHERE
 #include "/Engine/Private/SkyAtmosphereCommon.ush"
 #endif
 
@@ -996,7 +996,7 @@ float3 MaterialExpressionAtmosphericLightColor(FMaterialPixelParameters Paramete
 
 float3 MaterialExpressionSkyAtmosphereLightIlluminance(FMaterialPixelParameters Parameters, uint LightIndex)
 {
-#if MATERIAL_SKY_ATMOSPHERE
+#if MATERIAL_SKY_ATMOSPHERE && PROJECT_SUPPORT_SKY_ATMOSPHERE
 	const float3 PlanetCenter = float3(0, 0, -ResolvedView.SkyAtmosphereBottomRadius);
 	const float3 WorldPos = GetWorldPosition(Parameters)*CM_TO_SKY_UNIT;
 	const float3 PlanetCenterToWorldPos = WorldPos - PlanetCenter;
@@ -1014,7 +1014,7 @@ float3 MaterialExpressionSkyAtmosphereLightIlluminance(FMaterialPixelParameters
 
 float3 MaterialExpressionSkyAtmosphereLightDirection(FMaterialPixelParameters Parameters, uint LightIndex)
 {
-#if MATERIAL_SKY_ATMOSPHERE
+#if MATERIAL_SKY_ATMOSPHERE && PROJECT_SUPPORT_SKY_ATMOSPHERE
 	return ResolvedView.AtmosphereLightDirection[LightIndex].xyz;
 #else
 	return float3(0.0f, 0.0f, 0.0f);
@@ -1023,7 +1023,7 @@ float3 MaterialExpressionSkyAtmosphereLightDirection(FMaterialPixelParameters Pa
 
 float3 MaterialExpressionSkyAtmosphereLightDiskLuminance(FMaterialPixelParameters Parameters, uint LightIndex)
 {
-#if MATERIAL_SKY_ATMOSPHERE
+#if MATERIAL_SKY_ATMOSPHERE && PROJECT_SUPPORT_SKY_ATMOSPHERE
 	const float3 PlanetCenter = float3(0, 0, -ResolvedView.SkyAtmosphereBottomRadius);
 	const float3 PlanetCenterToWorldCameraPos = ResolvedView.WorldCameraOrigin * CM_TO_SKY_UNIT - PlanetCenter;
 	const float3 ViewDir = -Parameters.CameraVector;
@@ -1038,19 +1038,25 @@ float3 MaterialExpressionSkyAtmosphereLightDiskLuminance(FMaterialPixelParameter
 
 float3 MaterialExpressionSkyAtmosphereViewLuminance(FMaterialPixelParameters Parameters)
 {
-#if MATERIAL_SKY_ATMOSPHERE
+#if MATERIAL_SKY_ATMOSPHERE && PROJECT_SUPPORT_SKY_ATMOSPHERE
 	const float3 PlanetCenter = float3(0, 0, -ResolvedView.SkyAtmosphereBottomRadius);
 	const float3 PlanetCenterToWorldCameraPos = ResolvedView.WorldCameraOrigin * CM_TO_SKY_UNIT - PlanetCenter;
 	const float ViewHeight = length(PlanetCenterToWorldCameraPos);
-	const float3 UpVector = PlanetCenterToWorldCameraPos / ViewHeight;
 	const float3 ViewDir = -Parameters.CameraVector;
-	const float ViewZenithCosAngle = dot(ViewDir, UpVector);
 
-	float2 Sol = RayIntersectSphere(PlanetCenterToWorldCameraPos, ViewDir, float4(0.0f, 0.0f, 0.0f, ResolvedView.SkyAtmosphereBottomRadius));
+	// The referencial used to build the Sky View lut
+	float3x3 LocalReferencial = GetSkyViewLutReferential(PlanetCenterToWorldCameraPos, ResolvedView.ViewForward);
+	// Compute inputs in this referential
+	float3 WorldPosLocal = float3(0.0, 0.0, ViewHeight);
+	float3 UpVectorLocal = float3(0.0, 0.0, 1.0);
+	float3 WorldDirLocal = mul(ViewDir, LocalReferencial);
+	float ViewZenithCosAngle = dot(WorldDirLocal, UpVectorLocal);
+
+	float2 Sol = RayIntersectSphere(WorldPosLocal, WorldDirLocal, float4(0.0f, 0.0f, 0.0f, ResolvedView.SkyAtmosphereBottomRadius));
 	const bool IntersectGround = any(Sol > 0.0f);
 
 	float2 SkyViewLutUv;
-	SkyViewLutParamsToUv(IntersectGround, ViewZenithCosAngle, ViewDir, ViewHeight, ResolvedView.SkyAtmosphereBottomRadius, ResolvedView.SkyViewLutSizeAndInvSize, SkyViewLutUv);
+	SkyViewLutParamsToUv(IntersectGround, ViewZenithCosAngle, WorldDirLocal, ViewHeight, ResolvedView.SkyAtmosphereBottomRadius, ResolvedView.SkyViewLutSizeAndInvSize, SkyViewLutUv);
 	float3 SkyAtmosphereViewLuminance = Texture2DSampleLevel(View.SkyViewLutTexture, View.SkyViewLutTextureSampler, SkyViewLutUv, 0.0f).rgb;
 #if USE_PREEXPOSURE
 	SkyAtmosphereViewLuminance *= ResolvedView.OneOverPreExposure;
@@ -1063,7 +1069,7 @@ float3 MaterialExpressionSkyAtmosphereViewLuminance(FMaterialPixelParameters Par
 
 float4 MaterialExpressionSkyAtmosphereAerialPerspective(FMaterialPixelParameters Parameters)
 {
-#if MATERIAL_SKY_ATMOSPHERE
+#if MATERIAL_SKY_ATMOSPHERE && PROJECT_SUPPORT_SKY_ATMOSPHERE
 	const float OneOverPreExposure = USE_PREEXPOSURE ? ResolvedView.OneOverPreExposure : 1.0f;
 
 	const float3 PlanetCenter = float3(0, 0, -ResolvedView.SkyAtmosphereBottomRadius);
@@ -1081,7 +1087,7 @@ float4 MaterialExpressionSkyAtmosphereAerialPerspective(FMaterialPixelParameters
 		ResolvedView.SkyAtmosphereCameraAerialPerspectiveVolumeStartDepth,
 		ResolvedView.SkyAtmosphereCameraAerialPerspectiveVolumeDepthSliceLength,
 		ResolvedView.SkyAtmosphereCameraAerialPerspectiveVolumeDepthSliceLengthInv,
-		OneOverPreExposure, FogToApplyOver);
+		OneOverPreExposure);
 	return AerialPerspective;
 #else
 	return float4(0.0f, 0.0f, 0.0f, 1.0f); // RGB= null scattering, A= null transmittance
@@ -1090,7 +1096,7 @@ float4 MaterialExpressionSkyAtmosphereAerialPerspective(FMaterialPixelParameters
 
 float3 MaterialExpressionSkyAtmosphereDistantLightScatteredLuminance(FMaterialPixelParameters Parameters)
 {
-#if MATERIAL_SKY_ATMOSPHERE
+#if MATERIAL_SKY_ATMOSPHERE && PROJECT_SUPPORT_SKY_ATMOSPHERE
 	// TODO load on platforms supporting it
 	return Texture2DSampleLevel(View.DistantSkyLightLutTexture, View.DistantSkyLightLutTextureSampler, float2(0.5f, 0.5f), 0.0f).rgb;
 #else
diff --git a/Engine/Shaders/Private/MobileBasePassVertexShader.usf b/Engine/Shaders/Private/MobileBasePassVertexShader.usf
--- a/Engine/Shaders/Private/MobileBasePassVertexShader.usf
+++ b/Engine/Shaders/Private/MobileBasePassVertexShader.usf
@@ -21,7 +21,7 @@
 #include "/Engine/Generated/VertexFactory.ush"
 #include "HeightFogCommon.ush"
 
-#if USE_VERTEX_FOG
+#if USE_VERTEX_FOG && PROJECT_SUPPORT_SKY_ATMOSPHERE
 #include "SkyAtmosphereCommon.ush"
 #endif
 
@@ -89,11 +89,12 @@ void Main(
 #if USE_VERTEX_FOG
 	float4 VertexFog = CalculateHeightFog(WorldPosition.xyz - ResolvedView.TranslatedWorldCameraOrigin);
 
+#if PROJECT_SUPPORT_SKY_ATMOSPHERE
 	if (ResolvedView.SkyAtmosphereApplyCameraAerialPerspectiveVolume > 0.0f)
 	{
 		const float OneOverPreExposure = USE_PREEXPOSURE ? ResolvedView.OneOverPreExposure : 1.0f;
 		// Sample the aerial perspective (AP). It is also blended under the VertexFog parameter.
-		float4 AP = GetAerialPerspectiveLuminanceTransmittance(
+		VertexFog = GetAerialPerspectiveLuminanceTransmittanceWithFogOver(
 			Output.Position, WorldPosition.xyz*CM_TO_SKY_UNIT, ResolvedView.TranslatedWorldCameraOrigin*CM_TO_SKY_UNIT,
 			View.CameraAerialPerspectiveVolume, View.CameraAerialPerspectiveVolumeSampler,
 			ResolvedView.SkyAtmosphereCameraAerialPerspectiveVolumeDepthResolutionInv,
@@ -103,6 +104,8 @@ void Main(
 			ResolvedView.SkyAtmosphereCameraAerialPerspectiveVolumeDepthSliceLengthInv,
 			OneOverPreExposure, VertexFog);
 	}
+#endif
+
 #if PACK_INTERPOLANTS
 	PackedInterps[0] = VertexFog;
 #else
diff --git a/Engine/Shaders/Private/SkyAtmosphere.usf b/Engine/Shaders/Private/SkyAtmosphere.usf
--- a/Engine/Shaders/Private/SkyAtmosphere.usf
+++ b/Engine/Shaders/Private/SkyAtmosphere.usf
@@ -382,6 +382,7 @@ float3 GetLightDiskLuminance(float3 WorldPos, float3 WorldDir, uint LightIndex)
 float3 GetMultipleScattering(float3 WorlPos, float ViewZenithCosAngle)
 {
 	float2 UV = saturate(float2(ViewZenithCosAngle*0.5f + 0.5f, (length(WorlPos) - Atmosphere.BottomRadius) / (Atmosphere.TopRadius - Atmosphere.BottomRadius)));
+	// We do no apply UV transform to sub range here as it has minimal impact.
 	float3 MultiScatteredLuminance = MultiScatteredLuminanceLutTexture.SampleLevel(MultiScatteredLuminanceLutTextureSampler, UV, 0).rgb;
 	return MultiScatteredLuminance * MULTISCATTERING_STORE_SCALE_INV;
 }
@@ -569,7 +570,7 @@ SingleScatteringResult IntegrateSingleScatteredLuminance(
 		}
 		else
 		{
-			t = tMax * SampleI / SampleCount;
+			t = tMax * (SampleI + PixelNoise) / SampleCount;
 		}
 		float3 P = WorldPos + t * WorldDir;
 		float PHeight = length(P);
@@ -765,17 +766,25 @@ void RenderSkyAtmosphereRayMarchingPS(
 #endif
 	}
 
-	float3 UpVector = normalize(WorldPos);
 	float ViewHeight = length(WorldPos);
 #if FASTSKY_ENABLED
 	if (ViewHeight < Atmosphere.TopRadius && DeviceZ == FarDepthValue)
 	{
 		float2 UV;
-		float ViewZenithCosAngle = dot(WorldDir, UpVector);
 
-		bool IntersectGround = RaySphereIntersectNearest(WorldPos, WorldDir, float3(0, 0, 0), Atmosphere.BottomRadius) >= 0.0f;
+		// The referencial used to build the Sky View lut
+		float3x3 LocalReferencial = GetSkyViewLutReferential(WorldPos, View.ViewForward);
 
-		SkyViewLutParamsToUv(IntersectGround, ViewZenithCosAngle, WorldDir, ViewHeight, Atmosphere.BottomRadius, SkyAtmosphere.SkyViewLutSizeAndInvSize, UV);
+		// Input vectors expressed in this referencial: Up is always Z. Also note that ViewHeight is unchanged in this referencial.
+		float3 WorldPosLocal = float3(0.0, 0.0, ViewHeight);
+		float3 UpVectorLocal = float3(0.0, 0.0, 1.0);
+		float3 WorldDirLocal = mul(WorldDir, LocalReferencial);
+
+		// Now evaluate inputs in the referential.
+		float ViewZenithCosAngle = dot(WorldDirLocal, UpVectorLocal);
+		bool IntersectGround = RaySphereIntersectNearest(WorldPosLocal, WorldDirLocal, float3(0, 0, 0), Atmosphere.BottomRadius) >= 0.0f;
+
+		SkyViewLutParamsToUv(IntersectGround, ViewZenithCosAngle, WorldDirLocal, ViewHeight, Atmosphere.BottomRadius, SkyAtmosphere.SkyViewLutSizeAndInvSize, UV);
 		float3 SkyLuminance = SkyViewLutTexture.SampleLevel(SkyViewLutTextureSampler, UV, 0).rgb;
 
 		PreExposedL += SkyLuminance * LuminanceScale;
@@ -794,7 +803,6 @@ void RenderSkyAtmosphereRayMarchingPS(
 #else
 
 	const float OneOverPreExposure = 1.0f;
-	float4 NullFog = float4(0.0f, 0.0f, 0.0f, 1.0f);
 
 	float3 DepthBufferWorldPos = GetScreenWorldPos(SVPos, DeviceZ).xyz;
 	float4 NDCPosition = mul(float4(DepthBufferWorldPos.xyz, 1), View.WorldToClip);
@@ -807,7 +815,7 @@ void RenderSkyAtmosphereRayMarchingPS(
 		SkyAtmosphere.CameraAerialPerspectiveVolumeStartDepth,
 		SkyAtmosphere.CameraAerialPerspectiveVolumeDepthSliceLength,
 		SkyAtmosphere.CameraAerialPerspectiveVolumeDepthSliceLengthInv,
-		OneOverPreExposure, NullFog);
+		OneOverPreExposure);
 
 	PreExposedL += AP.rgb * LuminanceScale;
 	float Transmittance = AP.a;
@@ -914,6 +922,7 @@ RWTexture2D<float3> MultiScatteredLuminanceLutUAV;
 void RenderMultiScatteredLuminanceLutCS(uint3 ThreadId : SV_DispatchThreadID)
 {
 	float2 PixPos = float2(ThreadId.xy) + 0.5f;
+	// We do no apply UV transform from sub range here as it has minimal impact.
 
 	float CosLightZenithAngle = (PixPos.x * SkyAtmosphere.MultiScatteredLuminanceLutSizeAndInvSize.z) * 2.0f - 1.0f;
 	float3 LightDir = float3(0.0f, sqrt(saturate(1.0f - CosLightZenithAngle * CosLightZenithAngle)), CosLightZenithAngle);
@@ -1004,15 +1013,26 @@ void RenderSkyViewLutCS(uint3 ThreadId : SV_DispatchThreadID)
 	float2 PixPos = float2(ThreadId.xy) + 0.5f;
 	float2 UV = PixPos * SkyAtmosphere.SkyViewLutSizeAndInvSize.zw;
 
-
 	float3 WorldPos = GetCameraPlanetPos();
-	float ViewHeight = length(WorldPos);
 
+	// For the sky view lut to work, and not be distorted, we need to transform the view and light directions 
+	// into a referential with UP being perpendicular to the ground. And with origin at the planet center.
 
-	// TODO we need a valid local basis for lights and we need to transform them in it. 
-	// ===> Without that, the Sky lut is only valid for camera is close to the origin (at the Z-top of the planet)
+	// This is the local referencial
+	float3x3 LocalReferencial = GetSkyViewLutReferential(WorldPos, View.ViewForward);
+
+	// This is the LUT camera height and position in the local referential
+	float ViewHeight = length(WorldPos);
+	WorldPos = float3(0.0, 0.0, ViewHeight);
+
+	// Get the view direction in this local referential
 	float3 WorldDir;
 	UvToSkyViewLutParams(WorldDir, ViewHeight, UV);
+	// And also both light source direction
+	float3 AtmosphereLightDirection0 = View.AtmosphereLightDirection[0].xyz;
+	AtmosphereLightDirection0 = mul(AtmosphereLightDirection0, LocalReferencial);
+	float3 AtmosphereLightDirection1 = View.AtmosphereLightDirection[1].xyz;
+	AtmosphereLightDirection1 = mul(AtmosphereLightDirection1, LocalReferencial);
 
 
 	// Move to top atmospehre
@@ -1038,7 +1058,7 @@ void RenderSkyViewLutCS(uint3 ThreadId : SV_DispatchThreadID)
 	SingleScatteringResult ss = IntegrateSingleScatteredLuminance(
 		float4(PixPos, 0.0f, 1.0f), WorldPos, WorldDir,
 		Ground, Sampling, DeviceZ, MieRayPhase,
-		View.AtmosphereLightDirection[0].xyz, View.AtmosphereLightDirection[1].xyz, View.AtmosphereLightColor[0].rgb, View.AtmosphereLightColor[1].rgb,
+		AtmosphereLightDirection0, AtmosphereLightDirection1, View.AtmosphereLightColor[0].rgb, View.AtmosphereLightColor[1].rgb,
 		AerialPespectiveViewDistanceScale);
 
 	SkyViewLutUAV[int2(PixPos)] = ss.L;
diff --git a/Engine/Shaders/Private/SkyAtmosphereCommon.ush b/Engine/Shaders/Private/SkyAtmosphereCommon.ush
--- a/Engine/Shaders/Private/SkyAtmosphereCommon.ush
+++ b/Engine/Shaders/Private/SkyAtmosphereCommon.ush
@@ -27,7 +27,7 @@ float4 GetAerialPerspectiveLuminanceTransmittance(
 	float AerialPerspectiveVolumeStartDepth,
 	float AerialPerspectiveVolumeDepthSliceLength,
 	float AerialPerspectiveVolumeDepthSliceLengthInv,
-	float OneOverExposure, inout float4 FogToApplyOver)
+	float OneOverExposure)
 {
 	float2 ScreenUv = (NDC.xy / NDC.ww) * float2(0.5f, -0.5f) + 0.5f;
 
@@ -64,17 +64,41 @@ float4 GetAerialPerspectiveLuminanceTransmittance(
 	AP.rgb *= OneOverExposure;
 #endif
 
+	return AP;
+}
+
+
+float4 GetAerialPerspectiveLuminanceTransmittanceWithFogOver(
+	float4 NDC, float3 SampledWorldPos, float3 CameraWorldPos,
+	Texture3D<float4> AerialPerspectiveVolumeTexture, SamplerState AerialPerspectiveVolumeTextureSampler,
+	float AerialPerspectiveVolumeDepthResolutionInv,
+	float AerialPerspectiveVolumeDepthResolution,
+	float AerialPerspectiveVolumeStartDepth,
+	float AerialPerspectiveVolumeDepthSliceLength,
+	float AerialPerspectiveVolumeDepthSliceLengthInv,
+	float OneOverExposure, float4 FogToApplyOver)
+{
+	float4 AP = GetAerialPerspectiveLuminanceTransmittance(
+		NDC, SampledWorldPos, CameraWorldPos,
+		AerialPerspectiveVolumeTexture, AerialPerspectiveVolumeTextureSampler,
+		AerialPerspectiveVolumeDepthResolutionInv,
+		AerialPerspectiveVolumeDepthResolution,
+		AerialPerspectiveVolumeStartDepth,
+		AerialPerspectiveVolumeDepthSliceLength,
+		AerialPerspectiveVolumeDepthSliceLengthInv,
+		OneOverExposure);
+
+	float4 FinalFog;
 	// Apply any other fog OVER aerial perspective because AP is usually optically thiner.
-	FogToApplyOver.rgb = FogToApplyOver.rgb + AP.rgb * FogToApplyOver.a;
+	FinalFog.rgb = FogToApplyOver.rgb + AP.rgb * FogToApplyOver.a;
 	// And combine both transmittance.
-	FogToApplyOver.a *= AP.a;
+	FinalFog.a   = FogToApplyOver.a * AP.a;
 
-	return AP;
+	return FinalFog;
 }
 
 
 
-// TODO precompute those terms from resolutions (Or set resolution as #defined constants)
 float2 FromUnitToSubUvs(float2 uv, float4 SizeAndInvSize) { return (uv + 0.5f * SizeAndInvSize.zw) * (SizeAndInvSize.xy / (SizeAndInvSize.xy + 1.0f)); }
 float2 FromSubUvsToUnit(float2 uv, float4 SizeAndInvSize) { return (uv - 0.5f * SizeAndInvSize.zw) * (SizeAndInvSize.xy / (SizeAndInvSize.xy - 1.0f)); }
 
@@ -129,6 +153,16 @@ void SkyViewLutParamsToUv(
 	UV = FromUnitToSubUvs(UV, SkyViewLutSizeAndInvSize);
 }
 
+float3x3 GetSkyViewLutReferential(in float3 WorldPos, in float3 ViewForward)
+{
+	float3	Up = normalize(WorldPos);
+	float3	Forward = ViewForward;		// This can make texel visible when the camera is rotating. Use constant worl direction instead?
+	float3	Left = normalize(cross(Forward, Up));
+	Forward = normalize(cross(Up, Left));
+	float3x3 LocalReferencial = transpose(float3x3(Forward, Left, Up));
+	return LocalReferencial;
+}
+
 
 
 float3 GetLightDiskLuminance(
diff --git a/Engine/Source/Runtime/Engine/Classes/Engine/RendererSettings.h b/Engine/Source/Runtime/Engine/Classes/Engine/RendererSettings.h
--- a/Engine/Source/Runtime/Engine/Classes/Engine/RendererSettings.h
+++ b/Engine/Source/Runtime/Engine/Classes/Engine/RendererSettings.h
@@ -648,6 +648,14 @@ class ENGINE_API URendererSettings : public UDeveloperSettings
 		ConfigRestartRequired = true))
 		uint32 bSupportAtmosphericFog : 1;
 
+	/**
+	"The sky atmosphere component requires extra samplers/textures to be bound to apply aerial perspective on transparent surfaces (and all surfaces on mobile via per vertex evaluation)." 
+	*/
+	UPROPERTY(config, EditAnywhere, Category = ShaderPermutationReduction, meta = (
+		ConsoleVariable = "r.SupportSkyAtmosphere", DisplayName = "Support Sky Atmosphere",
+		ConfigRestartRequired = true))
+		uint32 bSupportSkyAtmosphere : 1;
+
 	/**
 	"Skincache allows a compute shader to skin once each vertex, save those results into a new buffer and reuse those calculations when later running the depth, base and velocity passes. This also allows opting into the 'recompute tangents' for skinned mesh instance feature. Disabling will reduce the number of shader permutations required per material. Changing this setting requires restarting the editor."
 	*/
diff --git a/Engine/Source/Runtime/Engine/Private/RendererSettings.cpp b/Engine/Source/Runtime/Engine/Private/RendererSettings.cpp
--- a/Engine/Source/Runtime/Engine/Private/RendererSettings.cpp
+++ b/Engine/Source/Runtime/Engine/Private/RendererSettings.cpp
@@ -67,6 +67,7 @@ URendererSettings::URendererSettings(const FObjectInitializer& ObjectInitializer
 	bSupportStationarySkylight = true;
 	bSupportPointLightWholeSceneShadows = true;
 	bSupportAtmosphericFog = true;
+	bSupportSkyAtmosphere = true;
 	bSupportSkinCacheShaders = false;
 	bSupportMaterialLayers = false;
 	GPUSimulationTextureSizeX = 1024;
diff --git a/Engine/Source/Runtime/Engine/Private/SceneManagement.cpp b/Engine/Source/Runtime/Engine/Private/SceneManagement.cpp
--- a/Engine/Source/Runtime/Engine/Private/SceneManagement.cpp
+++ b/Engine/Source/Runtime/Engine/Private/SceneManagement.cpp
@@ -1027,6 +1027,7 @@ void FReadOnlyCVARCache::Init()
 	static const auto CVarSupportAllShaderPermutations = IConsoleManager::Get().FindTConsoleVariableDataInt(TEXT("r.SupportAllShaderPermutations"));	
 	static const auto CVarVertexFoggingForOpaque = IConsoleManager::Get().FindTConsoleVariableDataInt(TEXT("r.VertexFoggingForOpaque"));	
 	static const auto CVarAllowStaticLighting = IConsoleManager::Get().FindTConsoleVariableDataInt(TEXT("r.AllowStaticLighting"));
+	static const auto CVarSupportSkyAtmosphere = IConsoleManager::Get().FindTConsoleVariableDataInt(TEXT("r.SupportSkyAtmosphere"));
 
 	static const auto CVarMobileAllowMovableDirectionalLights = IConsoleManager::Get().FindTConsoleVariableDataInt(TEXT("r.Mobile.AllowMovableDirectionalLights"));
 	static const auto CVarMobileEnableStaticAndCSMShadowReceivers = IConsoleManager::Get().FindTConsoleVariableDataInt(TEXT("r.Mobile.EnableStaticAndCSMShadowReceivers"));
@@ -1042,6 +1043,7 @@ void FReadOnlyCVARCache::Init()
 	bEnablePointLightShadows = !CVarSupportPointLightWholeSceneShadows || CVarSupportPointLightWholeSceneShadows->GetValueOnAnyThread() != 0 || bForceAllPermutations;
 	bEnableLowQualityLightmaps = !CVarSupportLowQualityLightmaps || CVarSupportLowQualityLightmaps->GetValueOnAnyThread() != 0 || bForceAllPermutations;
 	bAllowStaticLighting = CVarAllowStaticLighting->GetValueOnAnyThread() != 0;
+	bSupportSkyAtmosphere = !CVarSupportSkyAtmosphere || CVarSupportSkyAtmosphere->GetValueOnAnyThread() != 0;
 
 	// mobile
 	bMobileAllowMovableDirectionalLights = CVarMobileAllowMovableDirectionalLights->GetValueOnAnyThread() != 0;
diff --git a/Engine/Source/Runtime/Engine/Private/ShaderCompiler/ShaderCompiler.cpp b/Engine/Source/Runtime/Engine/Private/ShaderCompiler/ShaderCompiler.cpp
--- a/Engine/Source/Runtime/Engine/Private/ShaderCompiler/ShaderCompiler.cpp
+++ b/Engine/Source/Runtime/Engine/Private/ShaderCompiler/ShaderCompiler.cpp
@@ -3346,6 +3346,11 @@ void GlobalBeginCompileShader(
 		Input.Environment.SetDefine(TEXT("PROJECT_MOBILE_DISABLE_VERTEX_FOG"), CVar ? (CVar->GetInt() != 0) : 0);
 	}
 
+	{
+		static IConsoleVariable* CVar = IConsoleManager::Get().FindConsoleVariable(TEXT("r.SupportSkyAtmosphere"));
+		Input.Environment.SetDefine(TEXT("PROJECT_SUPPORT_SKY_ATMOSPHERE"), CVar ? (CVar->GetInt() != 0) : 0);
+	}
+
 	{
 		static IConsoleVariable* CVar = IConsoleManager::Get().FindConsoleVariable(TEXT("r.Mobile.UseLegacyShadingModel"));
 		Input.Environment.SetDefine(TEXT("PROJECT_MOBILE_USE_LEGACY_SHADING"), CVar ? (CVar->GetInt() != 0) : 0);
diff --git a/Engine/Source/Runtime/Engine/Public/SceneManagement.h b/Engine/Source/Runtime/Engine/Public/SceneManagement.h
--- a/Engine/Source/Runtime/Engine/Public/SceneManagement.h
+++ b/Engine/Source/Runtime/Engine/Public/SceneManagement.h
@@ -2913,6 +2913,7 @@ struct FReadOnlyCVARCache
 	bool bEnableAtmosphericFog;
 	bool bEnableLowQualityLightmaps;
 	bool bAllowStaticLighting;
+	bool bSupportSkyAtmosphere;
 
 	// Mobile specific
 	bool bMobileAllowMovableDirectionalLights;
diff --git a/Engine/Source/Runtime/RenderCore/Private/Shader.cpp b/Engine/Source/Runtime/RenderCore/Private/Shader.cpp
--- a/Engine/Source/Runtime/RenderCore/Private/Shader.cpp
+++ b/Engine/Source/Runtime/RenderCore/Private/Shader.cpp
@@ -2472,6 +2472,14 @@ void ShaderMapAppendKeyString(EShaderPlatform Platform, FString& KeyString)
 		}
 	}
 
+	{
+		static const auto CVar = IConsoleManager::Get().FindTConsoleVariableDataInt(TEXT("r.SupportSkyAtmosphere"));
+		if (CVar && CVar->GetValueOnAnyThread() > 0)
+		{
+			KeyString += TEXT("_SKYATM");
+		}
+	}
+
 	{
 		static const auto CVar = IConsoleManager::Get().FindTConsoleVariableDataInt(TEXT("r.EarlyZPassOnlyMaterialMasking"));
 		if (CVar && CVar->GetValueOnAnyThread() > 0)
diff --git a/Engine/Source/Runtime/Renderer/Private/SkyAtmosphereRendering.cpp b/Engine/Source/Runtime/Renderer/Private/SkyAtmosphereRendering.cpp
--- a/Engine/Source/Runtime/Renderer/Private/SkyAtmosphereRendering.cpp
+++ b/Engine/Source/Runtime/Renderer/Private/SkyAtmosphereRendering.cpp
@@ -18,12 +18,19 @@
 //#pragma optimize( "", off )
 
 
-
+// The runtime ON/OFF toggle
 static TAutoConsoleVariable<int32> CVarSkyAtmosphere(
 	TEXT("r.SkyAtmosphere"), 1,
 	TEXT("SkyAtmosphere components are rendered when this is not 0, otherwise ignored.\n"),
 	ECVF_RenderThreadSafe);
 
+// The project settings (disable runtime and shader code)
+static TAutoConsoleVariable<int32> CVarSupportAtmosphericFog(
+	TEXT("r.SupportSkyAtmosphere"),
+	1,
+	TEXT("Enables SkyAtmosphere rendering and shader code."),
+	ECVF_ReadOnly | ECVF_RenderThreadSafe);
+
 ////////////////////////////////////////////////////////////////////////// Regular sky 
 
 static TAutoConsoleVariable<float> CVarSkyAtmosphereSampleCountMin(
@@ -244,13 +251,13 @@ IMPLEMENT_GLOBAL_SHADER_PARAMETER_STRUCT(FSkyAtmosphereInternalCommonParameters,
 static bool ShouldPipelineCompileSkyAtmosphereShader(EShaderPlatform ShaderPlatform)
 {
 	// Requires SM5 or ES3_1 (GL/Vulkan) for compute shaders and volume textures support.
-	return RHISupportsComputeShaders(ShaderPlatform); // RHISupportsVolumeTextures(GetMaxSupportedFeatureLevel(ShaderPlatform));
+	return RHISupportsComputeShaders(ShaderPlatform);
 }
 
 bool ShouldRenderSkyAtmosphere(const FSkyAtmosphereRenderSceneInfo* SkyAtmosphere, EShaderPlatform ShaderPlatform)
 {
 	const bool ShadersCompiled = ShouldPipelineCompileSkyAtmosphereShader(ShaderPlatform);
-	return SkyAtmosphere && ShadersCompiled && CVarSkyAtmosphere.GetValueOnRenderThread() > 0;
+	return FReadOnlyCVARCache::Get().bSupportSkyAtmosphere && SkyAtmosphere && ShadersCompiled && CVarSkyAtmosphere.GetValueOnRenderThread() > 0;
 	// TODO: Add new or reuse EngineShowFlags.AtmosphericFog? ALso take into account EngineShowFlags.Fog as previously?
 }
 
