Fixed specular lighting discontinuity in hybrid ray traced reflection mode by always using eye vector that would've been computed for screen space lighting.

diff --git a/Engine/Shaders/Private/RayTracing/RayTracingReflections.usf b/Engine/Shaders/Private/RayTracing/RayTracingReflections.usf
--- a/Engine/Shaders/Private/RayTracing/RayTracingReflections.usf
+++ b/Engine/Shaders/Private/RayTracing/RayTracingReflections.usf
@@ -386,13 +386,32 @@ void RayTracingReflectionsRGS()
 				}
 				else
 				{
-					Payload = TraceRayAndAccumulateResults(
+					Payload = TraceRayResults(
 						TopLayerRay,
 						TLAS,
 						RandSequence,
 						RayFlags,
 						ReflectionMaxNormalBias,
 						RAY_TRACING_MASK_OPAQUE,
+						RayCone,
+						TopLayerRadiance);
+
+					RayDesc LightingRay = TopLayerRay;
+					if (DIM_HYBRID)
+					{
+						// Enforce consistent eye vector for computing lighting between screen space and ray traced reflections when using hybrid mode.
+						// This ensures that specular lighting matches and there is no visual discontinuity, even though it is wrong.
+						LightingRay.Direction = normalize(Payload.WorldPos - View.WorldCameraOrigin);
+					}
+
+					AccumulateResults(
+						Payload,
+						LightingRay,
+						TLAS,
+						RandSequence,
+						RayFlags,
+						ReflectionMaxNormalBias,
+						RAY_TRACING_MASK_OPAQUE,
 						ReflectedShadowsType,
 						ShouldDoDirectLighting,
 						ShouldDoEmissiveAndIndirectLighting,
