Refactors bilateral rejection to be configurable on per multiplexed signal.

diff --git a/Engine/Shaders/Private/ScreenSpaceDenoise/SSDDefinitions.ush b/Engine/Shaders/Private/ScreenSpaceDenoise/SSDDefinitions.ush
--- a/Engine/Shaders/Private/ScreenSpaceDenoise/SSDDefinitions.ush
+++ b/Engine/Shaders/Private/ScreenSpaceDenoise/SSDDefinitions.ush
@@ -227,6 +227,32 @@
 	#define SIGNAL_TEXTURE_TYPE_UINT2 1
 
 
+/** Bilateral settings bitfield. */
+	#define BILATERAL_POSITION_BASED(n) (0x0000 + (n) & 0xF)
+	#define BILATERAL_NORMAL 0x0010
+	#define BILATERAL_TOKOYASHI_LOBE 0x0020
+	#define BILATERAL_TOKOYASHI_AXES 0x0040
+	#define BILATERAL_TOKOYASHI (BILATERAL_TOKOYASHI_LOBE | BILATERAL_TOKOYASHI_AXES)
+
+/** Bilateral settings presets for spatial kernels. */
+	// No bilateral settings.
+	#define BILATERAL_PRESET_DISABLED 0x0000
+
+	#define BILATERAL_PRESET_MONOCHROMATIC_PENUMBRA 0x0011
+
+	#define BILATERAL_PRESET_POLYCHROMATIC_PENUMBRA 0x0021
+
+	#define BILATERAL_PRESET_REFLECTIONS 0x1001
+	#define BILATERAL_PRESET_REFLECTIONS_1SPP 0x1002
+	#define BILATERAL_PRESET_REFLECTIONS_TAA 0x1003
+
+	#define BILATERAL_PRESET_DIFFUSE 0x2001
+
+	#define BILATERAL_PRESET_SPHERICAL_HARMONIC 0x3001
+
+
+
+
 //------------------------------------------------------- GLOBAL CONFIGURATION OF ALL DENOISER SHADERS
 
 /** Do normal based anysotropy. */
diff --git a/Engine/Shaders/Private/ScreenSpaceDenoise/SSDSignalFramework.ush b/Engine/Shaders/Private/ScreenSpaceDenoise/SSDSignalFramework.ush
--- a/Engine/Shaders/Private/ScreenSpaceDenoise/SSDSignalFramework.ush
+++ b/Engine/Shaders/Private/ScreenSpaceDenoise/SSDSignalFramework.ush
@@ -5,26 +5,6 @@
 
 //------------------------------------------------------- CONFIG DISABLED DEFAULTS
 
-/** Whether should use bilateral rejection for significantly different roughness. */
-#ifndef CONFIG_TOKOYASHI_LOBE_BILATERAL
-	#define CONFIG_TOKOYASHI_LOBE_BILATERAL 0
-#endif
-
-/** Whether should use bilateral rejection for significantly different lobe axis directions. */
-#ifndef CONFIG_TOKOYASHI_AXES_BILATERAL
-	#define CONFIG_TOKOYASHI_AXES_BILATERAL 0
-#endif
-
-/** Whether should use normal based rejection. */
-#ifndef CONFIG_VIEW_POSITION_BASED_BILATERAL
-	#define CONFIG_VIEW_POSITION_BASED_BILATERAL 0
-#endif
-
-/** Whether should use normal based rejection. */
-#ifndef CONFIG_NORMAL_BASED_BILATERAL
-	#define CONFIG_NORMAL_BASED_BILATERAL 0
-#endif
-
 #ifndef CONFIG_SPECULAR_RATIO_ESTIMATOR
 	#define CONFIG_SPECULAR_RATIO_ESTIMATOR 0
 #endif
@@ -339,10 +319,6 @@ float4 WeightedLerp( float4 ColorA, float WeightA, float4 ColorB, float WeightB,
 #endif
 
 
-//------------------------------------------------------- INTERNAL FOR REFLECTION SIGNAL
-#if CONFIG_SIGNAL_PROCESSING == SIGNAL_PROCESSING_REFLECTIONS
-
-
 /** Create a spherical gaussian lob for a given roughness. */
 FSphericalGaussian ComputeRoughnessLobe(float Roughness, float3 N, float3 V)
 {
@@ -488,9 +464,6 @@ void GetInputRayDirectionPixel(
 }
 
 
-#endif // CONFIG_SIGNAL_PROCESSING == SIGNAL_PROCESSING_REFLECTIONS
-
-
 //------------------------------------------------------- MAIN FUNCTIONS TO IMPLEMENT FOR THE PASSES
 
 /** Return whether this signal is occluded or not. */
@@ -509,6 +482,7 @@ bool IsInvalidSample(FSSDSignalSample Sample)
 
 /** COmpute the weight of the bilateral filter to use between a reference sample and a neighbor. */
 float ComputeBilateralWeight(
+	const uint BilateralSettings,
 	float MaxWorldBluringDistance,
 	FSSDSampleSceneInfos RefSceneMetadata,
 	FSSDSampleSceneInfos NeighborSceneMetadata,
@@ -520,7 +494,7 @@ float ComputeBilateralWeight(
 	}
 	
 	float PositionBilateralWeight = 1;
-	#if CONFIG_VIEW_POSITION_BASED_BILATERAL == 1
+	if ((BilateralSettings & 0xF) == 1)
 	{
 		const float WorldRadius = MaxWorldBluringDistance;
 
@@ -537,7 +511,7 @@ float ComputeBilateralWeight(
 
 		PositionBilateralWeight = saturate(1 - DistSquare * Multiplier);
 	}
-	#elif CONFIG_VIEW_POSITION_BASED_BILATERAL == 2
+	else if ((BilateralSettings & 0xF) == 2)
 	{
 		// TODO: save 2 VGPR by using SceneDepth instead of RefSceneMetadata.TranslatedWorldPosition (not sure it can work in math yet).
 		// float2 ClipPosition = ScreenPosition * (View.ViewToClip[3][3] < 1.0f ? WorldDepth : 1.0f);
@@ -547,19 +521,19 @@ float ComputeBilateralWeight(
 
 		PositionBilateralWeight = saturate(1 - DistFromRefPlane * rcp(MaxWorldBluringDistance));
 	}
-	#elif CONFIG_VIEW_POSITION_BASED_BILATERAL == 3
+	else if ((BilateralSettings & 0xF) == 3)
 	{
 		if (dot(NeighborToRefWorldVector, NeighborToRefWorldVector) > MaxWorldBluringDistance * MaxWorldBluringDistance)
 		{
 			PositionBilateralWeight = 0;
 		}
 	}
-	#elif CONFIG_VIEW_POSITION_BASED_BILATERAL == 4
+	else if ((BilateralSettings & 0xF) == 4)
 	{
 		// TODO: extremely expensive with rcp()...
 		PositionBilateralWeight = saturate(1 - dot(NeighborToRefWorldVector, NeighborToRefWorldVector) * rcp(MaxWorldBluringDistance * MaxWorldBluringDistance));
 	}
-	#elif CONFIG_VIEW_POSITION_BASED_BILATERAL == 5
+	else if ((BilateralSettings & 0xF) == 5)
 	{
 		float DistFromRefPlane = dot(GetWorldNormal(RefSceneMetadata), RefSceneMetadata.TranslatedWorldPosition - NeighborSceneMetadata.TranslatedWorldPosition);
 
@@ -568,10 +542,9 @@ float ComputeBilateralWeight(
 			saturate(1 - dot(NeighborToRefWorldVector, NeighborToRefWorldVector) * rcp(MaxWorldBluringDistance * MaxWorldBluringDistance)) *
 			saturate(1 - 6 * abs(DistFromRefPlane) * rcp(MaxWorldBluringDistance)));
 	}
-	#endif
 
 	float NormalBilateralWeight = 1;
-	#if CONFIG_NORMAL_BASED_BILATERAL
+	if (BilateralSettings & BILATERAL_NORMAL)
 	{
 		// TODO: use GetWorldNormal() instead to reduce VGPR?
 		float NoN = max(dot(RefSceneMetadata.WorldNormal, NeighborSceneMetadata.WorldNormal), 0);
@@ -589,12 +562,11 @@ float ComputeBilateralWeight(
 			NormalBilateralWeight = exp(-x * x);
 		#endif
 	}
-	#endif
 	
 	// [ Tokoyashi 2015, "Specular Lobe-Aware Filtering and Upsampling for Interactive Indirect Illumination" ]
 	float LobeSimilarity = 1;
 	float AxesSimilarity = 1;
-	#if CONFIG_TOKOYASHI_LOBE_BILATERAL || CONFIG_TOKOYASHI_AXES_BILATERAL
+	if (BilateralSettings & (BILATERAL_TOKOYASHI_LOBE | BILATERAL_TOKOYASHI_AXES))
 	{
 		const float Beta = 32;
 
@@ -603,13 +575,12 @@ float ComputeBilateralWeight(
 	
 		float InvSharpnessSum = rcp(Ref.Sharpness + Neighbor.Sharpness);
 
-		if (CONFIG_TOKOYASHI_LOBE_BILATERAL)
+		if (BilateralSettings & BILATERAL_TOKOYASHI_LOBE)
 			LobeSimilarity = pow(2 * sqrt(Ref.Sharpness * Neighbor.Sharpness) * InvSharpnessSum, Beta);
-
-		if (CONFIG_TOKOYASHI_AXES_BILATERAL)
+		
+		if (BilateralSettings & BILATERAL_TOKOYASHI_AXES)
 			AxesSimilarity = exp(-(Beta * (Ref.Sharpness * Neighbor.Sharpness) * InvSharpnessSum) * saturate(1 - dot(Ref.Axis, Neighbor.Axis)));
 	}
-	#endif
 
 	return PositionBilateralWeight * NormalBilateralWeight * LobeSimilarity * AxesSimilarity;
 }
@@ -652,19 +623,18 @@ float GetSignalWorldBluringRadius(FSSDSignalSample Sample, FSSDSampleSceneInfos
 }
 #elif CONFIG_SIGNAL_PROCESSING == SIGNAL_PROCESSING_REFLECTIONS
 {
-	return 1;
-}
-#elif CONFIG_SIGNAL_PROCESSING == SIGNAL_PROCESSING_AO
-{
-	return Sample.ClosestHitDistance;
+	// Not used by any pass.
+	return 0;
 }
-#elif CONFIG_SIGNAL_PROCESSING == SIGNAL_PROCESSING_DIFFUSE_INDIRECT_AND_AO
+#elif CONFIG_SIGNAL_PROCESSING == SIGNAL_PROCESSING_AO || CONFIG_SIGNAL_PROCESSING == SIGNAL_PROCESSING_DIFFUSE_INDIRECT_AND_AO
 {
+	// How much can be blured is about the the length of hit distance assuming some amount of visiiblity.
 	return Sample.ClosestHitDistance;
 }
 #elif CONFIG_SIGNAL_PROCESSING == SIGNAL_PROCESSING_DIFFUSE_SPHERICAL_HARMONIC
 {
-	return 1; // TODO
+	// Not used by any pass.
+	return 0;
 }
 #else
 	#error Unimplemented GetSignalWorldBluringRadius()
diff --git a/Engine/Shaders/Private/ScreenSpaceDenoise/SSDSpatialAccumulation.usf b/Engine/Shaders/Private/ScreenSpaceDenoise/SSDSpatialAccumulation.usf
--- a/Engine/Shaders/Private/ScreenSpaceDenoise/SSDSpatialAccumulation.usf
+++ b/Engine/Shaders/Private/ScreenSpaceDenoise/SSDSpatialAccumulation.usf
@@ -40,8 +40,7 @@
 
 // Configures all the pass for each individual signals.
 #if CONFIG_SIGNAL_PROCESSING == SIGNAL_PROCESSING_MONOCHROMATIC_PENUMBRA
-	//#define CONFIG_VIEW_POSITION_BASED_BILATERAL 2
-	#define CONFIG_VIEW_POSITION_BASED_BILATERAL 5
+	#define CONFIG_BILATERAL_PRESET BILATERAL_PRESET_MONOCHROMATIC_PENUMBRA
 	#define CONFIG_MULTIPLEXED_SIGNALS_PER_SIGNAL_DOMAIN 1
 
 	#define MAX_SIGNAL_BATCH_SIZE CONFIG_SIGNAL_BATCH_SIZE
@@ -122,8 +121,7 @@
 #elif CONFIG_SIGNAL_PROCESSING == SIGNAL_PROCESSING_POLYCHROMATIC_PENUMBRA_HARMONIC
 	#define COMPILE_SIGNAL_COLOR 1
 
-	//#define CONFIG_VIEW_POSITION_BASED_BILATERAL 2
-	#define CONFIG_VIEW_POSITION_BASED_BILATERAL 5
+	#define CONFIG_BILATERAL_PRESET BILATERAL_PRESET_POLYCHROMATIC_PENUMBRA
 	//#define CONFIG_MULTIPLEXED_SIGNALS_PER_SIGNAL_DOMAIN 1
 
 	// Denoise diffuse and specular harmonics at the same time.
@@ -162,14 +160,12 @@
 
 		// Configure the bilateral therm.
 		#if DIM_MULTI_SPP
-			#define CONFIG_TOKOYASHI_LOBE_BILATERAL 1
-			#define CONFIG_TOKOYASHI_AXES_BILATERAL 1
+			#define CONFIG_BILATERAL_PRESET BILATERAL_PRESET_REFLECTIONS
 		#else
 			// Reconstruction has 1 ray per input pixel, therefore can use specular ratio estimator instead of Tokoyashi axes bilateral kernel.
 			#define CONFIG_SPECULAR_RATIO_ESTIMATOR 1
 
-			// Uses tokoyashi lobe bilateral rejection to avoid mixing up significantly different lobes that can cause leaking/strong biasing.
-			#define CONFIG_TOKOYASHI_LOBE_BILATERAL 1
+			#define CONFIG_BILATERAL_PRESET BILATERAL_PRESET_REFLECTIONS_1SPP
 		#endif
 
 		//#define CONFIG_SAMPLE_COUNT_POLICY  SAMPLE_COUNT_POLICY_DISABLED
@@ -186,9 +182,7 @@
 		#define CONFIG_SIGNAL_INPUT_LAYOUT  SIGNAL_BUFFER_LAYOUT_REFLECTIONS_HISTORY
 		#define CONFIG_SIGNAL_OUTPUT_LAYOUT SIGNAL_BUFFER_LAYOUT_REFLECTIONS_REJECTION
 
-		// Can only be done using tokoyashi because have more than one sample at a time.
-		#define CONFIG_TOKOYASHI_LOBE_BILATERAL 1
-		#define CONFIG_TOKOYASHI_AXES_BILATERAL 1
+		#define CONFIG_BILATERAL_PRESET BILATERAL_PRESET_REFLECTIONS
 
 		#define CONFIG_ACCUMULATION_COLOR_SPACE REFLECTIONS_REJECTION_COLOR_SPACE
 		#define CONFIG_OUTPUT_COLOR_SPACE       REFLECTIONS_REJECTION_COLOR_SPACE
@@ -207,9 +201,7 @@
 		#define CONFIG_SAMPLE_COUNT_POLICY  SAMPLE_COUNT_POLICY_ACCUMULATION_BASED
 		//#define CONFIG_SAMPLE_COUNT_POLICY  SAMPLE_COUNT_POLICY_DISABLED
 
-		// Can only be done using tokoyashi because have more than one sample at a time.
-		#define CONFIG_TOKOYASHI_LOBE_BILATERAL 1
-		#define CONFIG_TOKOYASHI_AXES_BILATERAL 1
+		#define CONFIG_BILATERAL_PRESET BILATERAL_PRESET_REFLECTIONS
 
 		#define CONFIG_INPUT_TEXTURE_COUNT 2
 		#define CONFIG_OUTPUT_TEXTURE_COUNT 2
@@ -219,7 +211,8 @@
 	#endif
 
 #elif CONFIG_SIGNAL_PROCESSING == SIGNAL_PROCESSING_AO
-	#define CONFIG_VIEW_POSITION_BASED_BILATERAL 2
+	#define CONFIG_BILATERAL_PRESET BILATERAL_PRESET_DIFFUSE
+
 	#define CONFIG_NORMAL_BASED_BILATERAL 1
 
 	#if DIM_STAGE == STAGE_RECONSTRUCTION
@@ -264,7 +257,7 @@
 #elif CONFIG_SIGNAL_PROCESSING == SIGNAL_PROCESSING_DIFFUSE_INDIRECT_AND_AO
 	#define COMPILE_SIGNAL_COLOR 1
 
-	#define CONFIG_VIEW_POSITION_BASED_BILATERAL 2
+	#define CONFIG_BILATERAL_PRESET BILATERAL_PRESET_DIFFUSE
 	#define CONFIG_NORMAL_BASED_BILATERAL 1
 
 	#define CONFIG_INPUT_TEXTURE_COUNT 2
@@ -306,7 +299,7 @@
 	#define COMPILE_SIGNAL_COLOR_SH 1
 
 	// Given it's a spherical harmonic that store directionality, only need position based rejection.
-	#define CONFIG_VIEW_POSITION_BASED_BILATERAL 2
+	#define CONFIG_BILATERAL_PRESET BILATERAL_PRESET_SPHERICAL_HARMONIC
 
 	#define CONFIG_INPUT_TEXTURE_COUNT 4
 	#define CONFIG_OUTPUT_TEXTURE_COUNT 4
@@ -669,6 +662,9 @@ void MainCS(
 			}
 		}
 
+		SetBilateralPreset(CONFIG_BILATERAL_PRESET, /* inout */ KernelConfig);
+
+
 		// SGPRs
 		KernelConfig.BufferSizeAndInvSize = View.BufferSizeAndInvSize;
 		KernelConfig.BufferBilinearUVMinMax = View.BufferBilinearUVMinMax;
diff --git a/Engine/Shaders/Private/ScreenSpaceDenoise/SSDSpatialKernel.ush b/Engine/Shaders/Private/ScreenSpaceDenoise/SSDSpatialKernel.ush
--- a/Engine/Shaders/Private/ScreenSpaceDenoise/SSDSpatialKernel.ush
+++ b/Engine/Shaders/Private/ScreenSpaceDenoise/SSDSpatialKernel.ush
@@ -117,6 +117,9 @@ struct FSSDKernelConfig
 	
 	// The color space of the accumulation.
 	uint AccumulatorColorSpace[MAX_SIGNAL_MULTIPLEXING];
+	
+	// The color space of the accumulation.
+	uint BilateralSettings[MAX_SIGNAL_MULTIPLEXING];
 
 	
 	// --------------------------- Per wave.
@@ -186,6 +189,7 @@ FSSDKernelConfig CreateKernelConfig()
 		{
 			KernelConfig.BufferColorSpace[MultiplexId] = STANDARD_BUFFER_COLOR_SPACE;
 			KernelConfig.AccumulatorColorSpace[MultiplexId] = STANDARD_BUFFER_COLOR_SPACE;
+			KernelConfig.BilateralSettings[MultiplexId] = 0x0000;
 		}
 	}
 	
@@ -217,6 +221,65 @@ FSSDKernelConfig CreateKernelConfig()
 }
 
 
+void SetBilateralPreset(uint BilateralPresetId, inout FSSDKernelConfig KernelConfig)
+{
+	if (BilateralPresetId == BILATERAL_PRESET_MONOCHROMATIC_PENUMBRA)
+	{
+		UNROLL_N(MAX_SIGNAL_MULTIPLEXING)
+		for (uint MultiplexId = 0; MultiplexId < MAX_SIGNAL_MULTIPLEXING; MultiplexId++)
+		{
+			// Shadow masks are normal invarient, so only reject based on position.
+			KernelConfig.BilateralSettings[MultiplexId] = BILATERAL_POSITION_BASED(5);
+		}
+	}
+	else if (BilateralPresetId == BILATERAL_PRESET_POLYCHROMATIC_PENUMBRA)
+	{
+		// Diffuse.
+		KernelConfig.BilateralSettings[0] = BILATERAL_POSITION_BASED(5) | BILATERAL_NORMAL;
+
+		// Specular.
+		KernelConfig.BilateralSettings[1] = BILATERAL_POSITION_BASED(5) | BILATERAL_TOKOYASHI;
+	}
+	else if (BilateralPresetId == BILATERAL_PRESET_REFLECTIONS)
+	{
+		// Specular.
+		// Can only be done using tokoyashi because have more than one sample at a time.
+		KernelConfig.BilateralSettings[0] = BILATERAL_TOKOYASHI;
+
+		// Specular variance for sampling rejection pre convolution.
+		KernelConfig.BilateralSettings[1] = KernelConfig.BilateralSettings[0];
+	}
+	else if (BilateralPresetId == BILATERAL_PRESET_REFLECTIONS_1SPP)
+	{
+		// Specular.
+		// Use specular ratio estomator, so no need to to reject based on the axis of the lobe.
+		KernelConfig.BilateralSettings[0] = BILATERAL_TOKOYASHI_LOBE;
+	}
+	else if (BilateralPresetId == BILATERAL_PRESET_REFLECTIONS_TAA)
+	{
+		// Specular.
+		// Can only be done using tokoyashi because have more than one sample at a time.
+		KernelConfig.BilateralSettings[0] = BILATERAL_POSITION_BASED(1) | BILATERAL_TOKOYASHI;
+
+		// Specular variance for sampling rejection pre convolution.
+		KernelConfig.BilateralSettings[1] = KernelConfig.BilateralSettings[0];
+	}
+	else if (BilateralPresetId == BILATERAL_PRESET_DIFFUSE)
+	{
+		// Diffuse depends on world position, but also the normal of the surface given it doesn't store any directionality.
+		KernelConfig.BilateralSettings[0] = BILATERAL_POSITION_BASED(2) | BILATERAL_NORMAL;
+		
+		// Variance for sampling rejection pre convolution.
+		KernelConfig.BilateralSettings[1] = KernelConfig.BilateralSettings[0];
+	}
+	else if (BilateralPresetId == BILATERAL_PRESET_SPHERICAL_HARMONIC)
+	{
+		// Spherical harmonic encode directionality, so only reject based on world position.
+		KernelConfig.BilateralSettings[0] = BILATERAL_POSITION_BASED(2);
+	}
+}
+
+
 //------------------------------------------------------- CONSTANT
 
 static const float kWaveletFilterWeights5x5[] = { 3.0 / 8.0, 1.0 / 4.0, 1.0 / 16.0 };
@@ -474,6 +537,7 @@ void AccumulateSampledMultiplexedSignals(
 
 		// Compute the weight to be applied to do bilateral rejection.
 		float BilateralWeight = ComputeBilateralWeight(
+			KernelConfig.BilateralSettings[SignalMultiplexId],
 			FinalWorldBluringDistance,
 			KernelConfig.RefSceneMetadata,
 			SampleSceneMetadata,
diff --git a/Engine/Shaders/Private/ScreenSpaceDenoise/SSDTemporalAccumulation.usf b/Engine/Shaders/Private/ScreenSpaceDenoise/SSDTemporalAccumulation.usf
--- a/Engine/Shaders/Private/ScreenSpaceDenoise/SSDTemporalAccumulation.usf
+++ b/Engine/Shaders/Private/ScreenSpaceDenoise/SSDTemporalAccumulation.usf
@@ -33,7 +33,7 @@
 
 
 #if CONFIG_SIGNAL_PROCESSING == SIGNAL_PROCESSING_MONOCHROMATIC_PENUMBRA
-	#define CONFIG_VIEW_POSITION_BASED_BILATERAL 2
+	#define CONFIG_BILATERAL_PRESET BILATERAL_PRESET_MONOCHROMATIC_PENUMBRA
 	#define CONFIG_MULTIPLEXED_SIGNALS_PER_SIGNAL_DOMAIN 1
 
 	#define CONFIG_INPUT_TEXTURE_COUNT CONFIG_SIGNAL_BATCH_SIZE
@@ -65,8 +65,8 @@
 #elif CONFIG_SIGNAL_PROCESSING == SIGNAL_PROCESSING_POLYCHROMATIC_PENUMBRA_HARMONIC
 	#define COMPILE_SIGNAL_COLOR 1
 
-	#define CONFIG_VIEW_POSITION_BASED_BILATERAL 2
-	#define CONFIG_MULTIPLEXED_SIGNALS_PER_SIGNAL_DOMAIN 1
+	#define CONFIG_BILATERAL_PRESET BILATERAL_PRESET_POLYCHROMATIC_PENUMBRA
+	#define CONFIG_MULTIPLEXED_SIGNALS_PER_SIGNAL_DOMAIN 2
 
 	// Input and output layout.
 	#define CONFIG_SIGNAL_INPUT_LAYOUT  SIGNAL_BUFFER_LAYOUT_POLYCHROMATIC_PENUMBRA_HISTORY
@@ -89,10 +89,7 @@
 	#define CONFIG_INPUT_TEXTURE_COUNT 2
 	#define CONFIG_HISTORY_TEXTURE_COUNT 2
 
-	// Can only be done using tokoyashi because have more than one sample at a time.
-	#define CONFIG_TOKOYASHI_LOBE_BILATERAL 1
-	#define CONFIG_TOKOYASHI_AXES_BILATERAL 1
-	#define CONFIG_VIEW_POSITION_BASED_BILATERAL 1
+	#define CONFIG_BILATERAL_PRESET BILATERAL_PRESET_REFLECTIONS_TAA
 
 	// Use variance based rejection.
 	#define CONFIG_HISTORY_REJECTION             HISTORY_REJECTION_VAR_BOUNDARIES
@@ -120,8 +117,7 @@
 
 	#define CONFIG_HISTORY_REJECTION HISTORY_REJECTION_VAR_BOUNDARIES
 
-	#define CONFIG_VIEW_POSITION_BASED_BILATERAL 2
-	#define CONFIG_NORMAL_BASED_BILATERAL 1
+	#define CONFIG_BILATERAL_PRESET BILATERAL_PRESET_DIFFUSE
 
 	// Use pre transformed rejection buffer that contains pre transformed momment 1 & 2.
 	#if 1
@@ -147,8 +143,7 @@
 	//#define CONFIG_HISTORY_REJECTION HISTORY_REJECTION_ADAPTIVE_INTEGRATION
 	#define CONFIG_MULTIPLEXED_SIGNALS_PER_SIGNAL_DOMAIN 3
 
-	#define CONFIG_VIEW_POSITION_BASED_BILATERAL 2
-	#define CONFIG_NORMAL_BASED_BILATERAL 1
+	#define CONFIG_BILATERAL_PRESET BILATERAL_PRESET_DIFFUSE
 
 #elif CONFIG_SIGNAL_PROCESSING == SIGNAL_PROCESSING_DIFFUSE_SPHERICAL_HARMONIC
 	#define COMPILE_SIGNAL_COLOR_SH 1
@@ -166,7 +161,7 @@
 	#define CONFIG_MULTIPLEXED_SIGNALS_PER_SIGNAL_DOMAIN 4
 
 	// Given it's a spherical harmonic that store directionality, only need position based rejection.
-	#define CONFIG_VIEW_POSITION_BASED_BILATERAL 2
+	#define CONFIG_BILATERAL_PRESET BILATERAL_PRESET_SPHERICAL_HARMONIC
 
 #else
 	#error Unimplemented signal processing.
@@ -474,6 +469,8 @@ void TemporallyAccumulate(
 			
 			if (CONFIG_USE_NEAREST_HISTORY)
 				KernelConfig.SampleSet = SAMPLE_SET_1X1;
+			
+			SetBilateralPreset(CONFIG_BILATERAL_PRESET, /* inout */ KernelConfig);
 
 			// SGPR configuration of the kernel.
 			KernelConfig.BufferSizeAndInvSize = View.BufferSizeAndInvSize;
@@ -607,6 +604,8 @@ void TemporallyAccumulate(
 					KernelConfig.bForceAllAccumulation = CurrentFrameSamples.Array[0].SampleCount == 0;
 				}
 			}
+			
+			SetBilateralPreset(CONFIG_BILATERAL_PRESET, /* inout */ KernelConfig);
 		}
 
 		// SGPR configuration of the kernel.
