Fixed tessellation shader for "Volumetric PerVertex" lighting mode and translucent blend mode.

#rnx

diff --git a/Engine/Shaders/Private/BasePassTessellationShaders.usf b/Engine/Shaders/Private/BasePassTessellationShaders.usf
--- a/Engine/Shaders/Private/BasePassTessellationShaders.usf
+++ b/Engine/Shaders/Private/BasePassTessellationShaders.usf
@@ -29,11 +29,11 @@ FBasePassVSToDS PassInterpolate(FBasePassVSToDS a, float aInterp, FBasePassVSToD
 	#endif
 
 	#if TRANSLUCENCY_LIGHTING_VOLUMETRIC_PERVERTEX_NONDIRECTIONAL || TRANSLUCENCY_LIGHTING_VOLUMETRIC_PERVERTEX_DIRECTIONAL
-		TESSELLATION_INTERPOLATE_MEMBER(AmbientLightingVector);
+		TESSELLATION_INTERPOLATE_MEMBER(BasePassInterpolants.AmbientLightingVector);
 	#endif
 
 	#if TRANSLUCENCY_LIGHTING_VOLUMETRIC_PERVERTEX_DIRECTIONAL
-		TESSELLATION_INTERPOLATE_MEMBER(DirectionalLightingVector);
+		TESSELLATION_INTERPOLATE_MEMBER(BasePassInterpolants.DirectionalLightingVector);
 	#endif
 
 	#if USE_WORLD_POSITION_EXCLUDING_SHADER_OFFSETS
