Fixes banding artifacts on the harmonic denoiser, and plug the specular therm through

diff --git a/Engine/Shaders/Private/ScreenSpaceDenoise/SSDComposeHarmonics.usf b/Engine/Shaders/Private/ScreenSpaceDenoise/SSDComposeHarmonics.usf
--- a/Engine/Shaders/Private/ScreenSpaceDenoise/SSDComposeHarmonics.usf
+++ b/Engine/Shaders/Private/ScreenSpaceDenoise/SSDComposeHarmonics.usf
@@ -9,26 +9,28 @@
 #define TILE_PIXEL_SIZE 8
 
 #define CONFIG_SIGNAL_PROCESSING SIGNAL_PROCESSING_POLYCHROMATIC_PENUMBRA_HARMONIC
-#define CONFIG_SIGNAL_BATCH_SIZE 1
 
 #if CONFIG_SIGNAL_PROCESSING == SIGNAL_PROCESSING_POLYCHROMATIC_PENUMBRA_HARMONIC
 	#define COMPILE_SIGNAL_COLOR 1
+	#define COMPILE_SIGNAL_COLOR_ARRAY 2
 
 	//#define CONFIG_VIEW_POSITION_BASED_BILATERAL 2
 	#define CONFIG_MULTIPLEXED_SIGNALS_PER_SIGNAL_DOMAIN 1
 
 	// Input and output layout.
-	#define CONFIG_SIGNAL_INPUT_LAYOUT  SIGNAL_BUFFER_LAYOUT_POLYCHROMATIC_PENUMBRA_INPUT
-	#define CONFIG_SIGNAL_RECONSTRUCTED_LAYOUT  SIGNAL_BUFFER_LAYOUT_POLYCHROMATIC_PENUMBRA_HISTORY
+	#define CONFIG_SIGNAL_INPUT_LAYOUT  SIGNAL_BUFFER_LAYOUT_POLYCHROMATIC_PENUMBRA_HARMONIC_INPUT
+	#define CONFIG_SIGNAL_RECONSTRUCTED_LAYOUT  SIGNAL_BUFFER_LAYOUT_POLYCHROMATIC_PENUMBRA_HARMONIC_RECONSTRUCTION
 	#define CONFIG_SIGNAL_OUTPUT_LAYOUT SIGNAL_BUFFER_LAYOUT_POLYCHROMATIC_PENUMBRA_HISTORY
-	#define CONFIG_INPUT_TEXTURE_COUNT 3
-	#define CONFIG_OUTPUT_TEXTURE_COUNT 3
+	#define CONFIG_INPUT_TEXTURE_COUNT 4
+	#define CONFIG_OUTPUT_TEXTURE_COUNT 2
 
 	#define CONFIG_HARMONIC_COUNT DENOISER_HARMONICS_COUNT
 
 	// Denoise diffuse and specular harmonics at the same time.
 	#define MAX_SIGNAL_BATCH_SIZE 2
 
+	#define CONFIG_SIGNAL_BATCH_SIZE 2
+
 #else
 	#error Unknown signal processing.
 #endif
@@ -113,6 +115,8 @@ FSSDRWTexture2D	SignalOutput_UAVs_3;
 #define SignalOutput_UAVs_3 SignalOutput_UAVs_0
 #endif
 
+FSSDRWTexture2D DebugOutput;
+
 
 //------------------------------------------------------- ENTRY POINT
 
@@ -191,7 +195,7 @@ void MainCS(
 	UNROLL_N(CONFIG_SIGNAL_BATCH_SIZE)
 	for (uint MultiplexId = 0; MultiplexId < CONFIG_SIGNAL_BATCH_SIZE; MultiplexId++)
 	{
-		const uint LowestFrequencyHarmonicId = CONFIG_HARMONIC_COUNT - 1;
+		const uint LowestFrequencyHarmonicId = DENOISER_HARMONICS_COUNT - 1;
 		float LowestFrequencyWeight;
 		{
 			FSSDSignalSample Sample = MultiplexedHarmonics[LowestFrequencyHarmonicId].Array[MultiplexId];
@@ -199,17 +203,52 @@ void MainCS(
 			LowestFrequencyWeight = Sample.SampleCount;
 		}
 
-		UNROLL_N(CONFIG_HARMONIC_COUNT - 1)
-		for (uint HarmonicId = 0; HarmonicId < LowestFrequencyHarmonicId; HarmonicId++)
+		float3 Color = 0;
+		float3 ColorMask = 1.0;
+
+		UNROLL_N(DENOISER_HARMONICS_COUNT)
+		for (uint HarmonicId = 0; HarmonicId < DENOISER_HARMONICS_COUNT; HarmonicId++)
 		{
 			FSSDSignalSample Sample = MultiplexedHarmonics[HarmonicId].Array[MultiplexId];
+			FSSDSignalSample NormalizedSample = NormalizeToOneSample(Sample);
+
+			float3 Incoming = NormalizedSample.ColorArray[0].rgb;
+			float3 Outgoing = NormalizedSample.ColorArray[1].rgb;
+			float3 Stoped = Incoming - Outgoing;
+
+			// TODO.
+			float3 HarmonicMask;
+			HarmonicMask.r = Incoming.r > 0 ? saturate(Outgoing.r / Incoming.r) : 1.0;
+			HarmonicMask.g = Incoming.g > 0 ? saturate(Outgoing.g / Incoming.g) : 1.0;
+			HarmonicMask.b = Incoming.b > 0 ? saturate(Outgoing.b / Incoming.b) : 1.0;
+
+			ColorMask *= HarmonicMask;
+			
+			//if (HarmonicId == View.GeneralPurposeTweak)
+			//{
+			//	DebugOutput[DispatchThreadId] = float4(HarmonicMask, 0);
+			//}
+		}
 
-			float SampleWeight = -min(-LowestFrequencyWeight / Sample.SampleCount, 0.0);
+		float3 TotalLighting = MultiplexedHarmonics[DENOISER_HARMONICS_COUNT - 1].Array[MultiplexId].ColorArray[1];
+		float3 FinalLighting = ColorMask; //TotalLighting * ColorMask;
+		
+		//DebugOutput[DispatchThreadId] = float4(ColorMask, 0);
 
-			OutputSamples.Array[MultiplexId] = AddSignal(
-				OutputSamples.Array[MultiplexId],
-				MulSignal(Sample, SampleWeight));
-		}
+		OutputSamples.Array[MultiplexId].SampleCount = LowestFrequencyWeight;
+		OutputSamples.Array[MultiplexId].SceneColor.rgb = FinalLighting * LowestFrequencyWeight;
+		
+		//OutputSamples.Array[MultiplexId].SampleCount = LowestFrequencyWeight;
+		//OutputSamples.Array[MultiplexId].SceneColor.rgb = MultiplexedHarmonics[2].Array[MultiplexId].ColorArray[0] * LowestFrequencyWeight;
+	}
+
+	// Clamp.
+	UNROLL_N(CONFIG_SIGNAL_BATCH_SIZE)
+	for (uint MultiplexId = 0; MultiplexId < CONFIG_SIGNAL_BATCH_SIZE; MultiplexId++)
+	{
+		#if COMPILE_SIGNAL_COLOR
+			OutputSamples.Array[MultiplexId].SceneColor = max(OutputSamples.Array[MultiplexId].SceneColor, 0);
+		#endif
 	}
 
 	OutputMultiplexedSignal(
diff --git a/Engine/Shaders/Private/ScreenSpaceDenoise/SSDDefinitions.ush b/Engine/Shaders/Private/ScreenSpaceDenoise/SSDDefinitions.ush
--- a/Engine/Shaders/Private/ScreenSpaceDenoise/SSDDefinitions.ush
+++ b/Engine/Shaders/Private/ScreenSpaceDenoise/SSDDefinitions.ush
@@ -95,10 +95,13 @@
 
 
 	/** Buffer layout input for denoising multi polychromatic penumbra. */
-	#define SIGNAL_BUFFER_LAYOUT_POLYCHROMATIC_PENUMBRA_INPUT 0x1000
+	#define SIGNAL_BUFFER_LAYOUT_POLYCHROMATIC_PENUMBRA_HARMONIC_INPUT 0x1000
 
 	/** Internal buffer layout for denoising multi polychromatic penumbra. */
-	#define SIGNAL_BUFFER_LAYOUT_POLYCHROMATIC_PENUMBRA_HISTORY 0x1001
+	#define SIGNAL_BUFFER_LAYOUT_POLYCHROMATIC_PENUMBRA_HARMONIC_RECONSTRUCTION 0x1001
+
+	/** Internal buffer layout for denoising multi polychromatic penumbra. */
+	#define SIGNAL_BUFFER_LAYOUT_POLYCHROMATIC_PENUMBRA_HISTORY 0x1002
 
 
 	/** Buffer layout for the reflection given by the ray generation shader. */
@@ -240,7 +243,7 @@
 
 	#define BILATERAL_PRESET_MONOCHROMATIC_PENUMBRA 0x0011
 
-	#define BILATERAL_PRESET_POLYCHROMATIC_PENUMBRA 0x0021
+	#define BILATERAL_PRESET_POLYCHROMATIC_PENUMBRA 0x0022
 
 	#define BILATERAL_PRESET_REFLECTIONS 0x1001
 	#define BILATERAL_PRESET_REFLECTIONS_1SPP 0x1002
diff --git a/Engine/Shaders/Private/ScreenSpaceDenoise/SSDPublicHarmonics.ush b/Engine/Shaders/Private/ScreenSpaceDenoise/SSDPublicHarmonics.ush
--- a/Engine/Shaders/Private/ScreenSpaceDenoise/SSDPublicHarmonics.ush
+++ b/Engine/Shaders/Private/ScreenSpaceDenoise/SSDPublicHarmonics.ush
@@ -11,7 +11,9 @@
 
 //------------------------------------------------------- DEFINES
 
+// Number of armonic being denoised.
 #define DENOISER_HARMONICS_COUNT 4
+#define DENOISER_HARMONIC_BORDERS_COUNT (DENOISER_HARMONICS_COUNT + 1)
 
 
 //------------------------------------------------------- FUNCTIONS
@@ -67,5 +69,72 @@ float WorldBluringRadiusToHarmonic(float WorldDepth, float WorldBluringRadius)
 
 	const float TemporalHistoryRejectionDiameter = 3.0;
 
-	return clamp(log2(PixelBluringRadius - (TemporalHistoryRejectionDiameter - 1.0)), 0.0, DENOISER_HARMONICS_COUNT - 1.0);
+	return clamp(log2(PixelBluringRadius - (TemporalHistoryRejectionDiameter - 1.0)) * 0.5, 0.0, DENOISER_HARMONICS_COUNT - 1.0);
+}
+
+
+//------------------------------------------------------- Harmonic accumulator for RGS.
+
+// TODO(Denoiser): That is a lot of VGPR for ray gen, implement texture atomic one.
+struct FSSDHarmonicAccumulator
+{
+	float3 HarmonicBorders[DENOISER_HARMONIC_BORDERS_COUNT];
+};
+
+FSSDHarmonicAccumulator InitHarmonicAccumulator()
+{
+	FSSDHarmonicAccumulator Accumulator;
+	UNROLL_N(DENOISER_HARMONIC_BORDERS_COUNT)
+	for (uint HarmonicBorderId = 0; HarmonicBorderId < DENOISER_HARMONIC_BORDERS_COUNT; HarmonicBorderId++)
+	{
+		Accumulator.HarmonicBorders[HarmonicBorderId] = 0.0;
+	}
+	return Accumulator;
+}
+
+void AccumulateSampleHarmonic(inout FSSDHarmonicAccumulator Accumulator, float3 Color, float SceneDepth, float HitHarmonic)
+{
+	UNROLL_N(DENOISER_HARMONIC_BORDERS_COUNT)
+	for (uint HarmonicBorderId = 0; HarmonicBorderId < DENOISER_HARMONIC_BORDERS_COUNT; HarmonicBorderId++)
+	{
+		float Incoming = saturate(1.0 - (float(HarmonicBorderId) - HitHarmonic));
+		Accumulator.HarmonicBorders[HarmonicBorderId] += Incoming * Color;
+	}
+}
+
+void OutputHarmonicAccumulator(
+	RWTexture2D<float4> OutputTexture0,
+	RWTexture2D<float4> OutputTexture1,
+	RWTexture2D<float4> OutputTexture2,
+	RWTexture2D<float4> OutputTexture3,
+	RWTexture2D<float4> OutputTexture4,
+	uint2 PixelPosition,
+	FSSDHarmonicAccumulator Accumulator,
+	bool bIsValidPixel)
+{
+	float4 RawOutputs[DENOISER_HARMONIC_BORDERS_COUNT];
+	
+	UNROLL_N(DENOISER_HARMONIC_BORDERS_COUNT)
+	for (uint HarmonicBorderId = 0; HarmonicBorderId < DENOISER_HARMONIC_BORDERS_COUNT; HarmonicBorderId++)
+	{
+		RawOutputs[HarmonicBorderId] = float4(Accumulator.HarmonicBorders[HarmonicBorderId], bIsValidPixel ? 1.0 : 0.0);
+	}
+	
+	OutputTexture0[PixelPosition] = RawOutputs[0];
+	
+	#if DENOISER_HARMONIC_BORDERS_COUNT >= 2
+		OutputTexture1[PixelPosition] = RawOutputs[1];
+	#endif
+		
+	#if DENOISER_HARMONIC_BORDERS_COUNT >= 3
+		OutputTexture2[PixelPosition] = RawOutputs[2];
+	#endif
+
+	#if DENOISER_HARMONIC_BORDERS_COUNT >= 4
+		OutputTexture3[PixelPosition] = RawOutputs[3];
+	#endif
+
+	#if DENOISER_HARMONIC_BORDERS_COUNT >= 5
+		OutputTexture4[PixelPosition] = RawOutputs[4];
+	#endif
 }
diff --git a/Engine/Shaders/Private/ScreenSpaceDenoise/SSDSignalBufferEncoding.ush b/Engine/Shaders/Private/ScreenSpaceDenoise/SSDSignalBufferEncoding.ush
--- a/Engine/Shaders/Private/ScreenSpaceDenoise/SSDSignalBufferEncoding.ush
+++ b/Engine/Shaders/Private/ScreenSpaceDenoise/SSDSignalBufferEncoding.ush
@@ -55,7 +55,12 @@ void DecodeMultiplexedSignals(
 {
 	OutSamples = CreateSignalArrayFromScalarValue(0.0);
 
-	if (SignalBufferLayout == SIGNAL_BUFFER_LAYOUT_UNINITIALIZED)
+	if (0)
+	{
+	
+	}
+	#if 1
+	else if (SignalBufferLayout == SIGNAL_BUFFER_LAYOUT_UNINITIALIZED)
 	{
 		// trash out the output sample to identify quickly a miss configured buffer layout.
 		OutSamples = CreateSignalArrayFromScalarValue(-INFINITE_FLOAT);
@@ -159,55 +164,87 @@ void DecodeMultiplexedSignals(
 			OutSamples.Array[MultiplexId + 1].MissCount = Channels.g * SampleCount;
 		}
 	}
-	#if COMPILE_SIGNAL_COLOR
-	else if (SignalBufferLayout == SIGNAL_BUFFER_LAYOUT_POLYCHROMATIC_PENUMBRA_INPUT)
+	#endif
+	#if COMPILE_SIGNAL_COLOR_ARRAY >= 2
+	else if (SignalBufferLayout == SIGNAL_BUFFER_LAYOUT_POLYCHROMATIC_PENUMBRA_HARMONIC_INPUT)
 	{
-		// Decode diffuse harmonic
+		// The RGS output same weigth for diffuse and specular given it depends only on shading model id.
+		float Weight = RawSample[0].a;
+
+		// Decode diffuse
 		{
-			float Weight = RawSample[0].a < 0.0 ? 0 : 1;
+			const uint MultiplexId = 0;
 
-			OutSamples.Array[0].SampleCount = Weight;
-			OutSamples.Array[0].MissCount = RawSample[0].a * Weight;
-			OutSamples.Array[0].SceneColor.rgb = RawSample[0].rgb * Weight;
+			OutSamples.Array[MultiplexId].SampleCount = Weight;
+			OutSamples.Array[MultiplexId].ColorArray[0] = RawSample[0].rgb * Weight;
+			OutSamples.Array[MultiplexId].ColorArray[1] = RawSample[1].rgb * Weight;
 		}
 		
-		// Decode specular harmonic
+		// Decode specular
 		{
-			float Weight = RawSample[1].a < 0.0 ? 0 : 1;
+			const uint MultiplexId = 1;
 
-			OutSamples.Array[1].SampleCount = Weight;
-			OutSamples.Array[1].MissCount = RawSample[1].a * Weight;
-			OutSamples.Array[1].SceneColor.rgb = RawSample[1].rgb * Weight;
+			OutSamples.Array[MultiplexId].SampleCount = Weight;
+			OutSamples.Array[MultiplexId].ColorArray[0] = RawSample[2].rgb * Weight;
+			OutSamples.Array[MultiplexId].ColorArray[1] = RawSample[3].rgb * Weight;
 		}
 	}
+	else if (SignalBufferLayout == SIGNAL_BUFFER_LAYOUT_POLYCHROMATIC_PENUMBRA_HARMONIC_RECONSTRUCTION)
+	{
+		// Decode diffuse
+		{
+			const uint MultiplexId = 0;
+
+			float Weight = RawSample[0].a;
+
+			OutSamples.Array[MultiplexId].SampleCount = Weight;
+			OutSamples.Array[MultiplexId].ColorArray[0] = RawSample[0].rgb * Weight;
+			OutSamples.Array[MultiplexId].ColorArray[1] = RawSample[1].rgb * Weight;
+		}
+		
+		// Decode specular
+		{
+			const uint MultiplexId = 1;
+
+			float Weight = RawSample[2].a;
+
+			OutSamples.Array[MultiplexId].SampleCount = Weight;
+			OutSamples.Array[MultiplexId].ColorArray[0] = RawSample[2].rgb * Weight;
+			OutSamples.Array[MultiplexId].ColorArray[1] = RawSample[3].rgb * Weight;
+		}
+	}
+	#endif
+	#if COMPILE_SIGNAL_COLOR
 	else if (SignalBufferLayout == SIGNAL_BUFFER_LAYOUT_POLYCHROMATIC_PENUMBRA_HISTORY)
 	{
 		// Decode diffuse harmonic
 		{
-			float Weight = RawSample[2].r;
+			const uint MultiplexId = 0;
+
+			float Weight = RawSample[0].a;
 
 			if (bNormalizeSample)
 			{
 				Weight = Weight > 0 ? 1 : 0;
 			}
 
-			OutSamples.Array[0].SampleCount = Weight;
-			OutSamples.Array[0].MissCount = RawSample[0].a * Weight;
-			OutSamples.Array[0].SceneColor.rgb = RawSample[0].rgb * Weight;
+			OutSamples.Array[MultiplexId].SampleCount = Weight;
+			OutSamples.Array[MultiplexId].SceneColor.rgb = RawSample[0].rgb * Weight;
 		}
 		
 		// Decode specular harmonic
 		{
-			float Weight = RawSample[2].g;
-			
+			const uint MultiplexId = 1;
+
+			float Weight = RawSample[1].a;
+
 			if (bNormalizeSample)
 			{
 				Weight = Weight > 0 ? 1 : 0;
 			}
 
-			OutSamples.Array[1].SampleCount = Weight;
-			OutSamples.Array[1].MissCount = RawSample[1].a * Weight;
-			OutSamples.Array[1].SceneColor.rgb = RawSample[1].rgb * Weight;
+			OutSamples.Array[MultiplexId].SampleCount = Weight;
+			OutSamples.Array[MultiplexId].SceneColor.rgb = RawSample[1].rgb * Weight;
 		}
 	}
 	#endif
@@ -263,6 +300,7 @@ void DecodeMultiplexedSignals(
 		#endif
 	}
 	#endif
+	#if 1
 	else if (SignalBufferLayout == SIGNAL_BUFFER_LAYOUT_AO_INPUT)
 	{
 		OutSamples.Array[0].SampleCount = RawSample[1].r < 0 ? 0 : 1;
@@ -301,6 +339,7 @@ void DecodeMultiplexedSignals(
 			OutSamples.Array[0].MissCount = RawSample[0].g * RawSample[0].r;
 		}
 	}
+	#endif
 	#if COMPILE_SIGNAL_COLOR
 	else if (SignalBufferLayout == SIGNAL_BUFFER_LAYOUT_DIFFUSE_INDIRECT_AND_AO_INPUT_NSPP)
 	{
@@ -388,7 +427,12 @@ void EncodeMultiplexedSignals(
 	// Number of buffer the signal get encoded onto <= MAX_MULTIPLEXED_TEXTURES.
 	OutBufferCount = 1;
 
-	if (SignalBufferLayout == SIGNAL_BUFFER_LAYOUT_PENUMBRA_INJESTION_1SPP)
+	if (0)
+	{
+		// NOP
+	}
+	#if 1
+	else if (SignalBufferLayout == SIGNAL_BUFFER_LAYOUT_PENUMBRA_INJESTION_1SPP)
 	{
 		UNROLL
 		for (uint MultiplexId = 0; MultiplexId < MultiplexCount; MultiplexId++)
@@ -488,28 +532,59 @@ void EncodeMultiplexedSignals(
 				NormalizedWorldBluringRadius == WORLD_RADIUS_MISS ? -1 : NormalizedWorldBluringRadius);
 		}
 	}
+	#endif
+	#if COMPILE_SIGNAL_COLOR_ARRAY >= 2
+	else if (SignalBufferLayout == SIGNAL_BUFFER_LAYOUT_POLYCHROMATIC_PENUMBRA_HARMONIC_RECONSTRUCTION)
+	{
+		// diffuse harmonic
+		{
+			const uint MultiplexId = 0;
+
+			float NormalizationFactor = SafeRcp(Sample[MultiplexId].SampleCount);
+
+			OutRawSample[0].rgb = Sample[MultiplexId].ColorArray[0] * NormalizationFactor;
+			OutRawSample[0].a = Sample[MultiplexId].SampleCount;
+			OutRawSample[1].rgb = Sample[MultiplexId].ColorArray[1] * NormalizationFactor;
+		}
+		
+		// specular harmonic
+		{
+			const uint MultiplexId = 1;
+
+			float NormalizationFactor = SafeRcp(Sample[1].SampleCount);
+			
+			OutRawSample[2].rgb = Sample[MultiplexId].ColorArray[0] * NormalizationFactor;
+			OutRawSample[2].a = Sample[MultiplexId].SampleCount;
+			OutRawSample[3].rgb = Sample[MultiplexId].ColorArray[1] * NormalizationFactor;
+		}
+
+		OutBufferCount = 4;
+	}
+	#endif
 	#if COMPILE_SIGNAL_COLOR
 	else if (SignalBufferLayout == SIGNAL_BUFFER_LAYOUT_POLYCHROMATIC_PENUMBRA_HISTORY)
 	{
 		// Encode diffuse harmonic
 		{
-			float NormalizationFactor = SafeRcp(Sample[0].SampleCount);
+			const uint MultiplexId = 0;
 
-			OutRawSample[0] = ClampColorForEncoding(Sample[0].SceneColor * NormalizationFactor);
-			OutRawSample[0].a = Sample[0].MissCount * NormalizationFactor;
-			OutRawSample[2].r = Sample[0].SampleCount;
+			float NormalizationFactor = SafeRcp(Sample[MultiplexId].SampleCount);
+
+			OutRawSample[0].rgb = ClampColorForEncoding(Sample[MultiplexId].SceneColor.rgb * NormalizationFactor);
+			OutRawSample[0].a = Sample[MultiplexId].SampleCount;
 		}
 
 		// Encode specular harmonic
 		{
-			float NormalizationFactor = SafeRcp(Sample[1].SampleCount);
+			const uint MultiplexId = 1;
+
+			float NormalizationFactor = SafeRcp(Sample[MultiplexId].SampleCount);
 
-			OutRawSample[1] = ClampColorForEncoding(Sample[1].SceneColor * NormalizationFactor);
-			OutRawSample[1].a = Sample[1].MissCount * NormalizationFactor;
-			OutRawSample[2].g = Sample[1].SampleCount;
+			OutRawSample[1].rgb = ClampColorForEncoding(Sample[MultiplexId].SceneColor.rgb * NormalizationFactor);
+			OutRawSample[1].a = Sample[MultiplexId].SampleCount;
 		}
 
-		OutBufferCount = 3;
+		OutBufferCount = 2;
 	}
 	#endif
 	#if COMPILE_SIGNAL_COLOR
@@ -547,6 +622,7 @@ void EncodeMultiplexedSignals(
 		#endif
 	}
 	#endif
+	#if 1
 	else if (SignalBufferLayout == SIGNAL_BUFFER_LAYOUT_AO_REJECTION)
 	{
 		// Outputs number of ray miss in red to be compatible as a SSAO output.
@@ -565,6 +641,7 @@ void EncodeMultiplexedSignals(
 		OutRawSample[0].r = Sample[0].SampleCount > 0 ? Sample[0].MissCount / Sample[0].SampleCount : 1.0;
 		OutRawSample[0].g = Sample[0].SampleCount;
 	}
+	#endif
 	#if COMPILE_SIGNAL_COLOR
 	else if (SignalBufferLayout == SIGNAL_BUFFER_LAYOUT_DIFFUSE_INDIRECT_AND_AO_RECONSTRUCTION)
 	{
diff --git a/Engine/Shaders/Private/ScreenSpaceDenoise/SSDSignalCore.ush b/Engine/Shaders/Private/ScreenSpaceDenoise/SSDSignalCore.ush
--- a/Engine/Shaders/Private/ScreenSpaceDenoise/SSDSignalCore.ush
+++ b/Engine/Shaders/Private/ScreenSpaceDenoise/SSDSignalCore.ush
@@ -9,11 +9,14 @@
 	#define COMPILE_SIGNAL_COLOR 0
 #endif
 
-
 #ifndef COMPILE_SIGNAL_COLOR_SH
 	#define COMPILE_SIGNAL_COLOR_SH 0
 #endif
 
+#ifndef COMPILE_SIGNAL_COLOR_ARRAY
+	#define COMPILE_SIGNAL_COLOR_ARRAY 0
+#endif
+
 
 //------------------------------------------------------- STRUCTURE
 
@@ -28,6 +31,11 @@ struct FSSDSignalSample
 		float4 SceneColor;
 	#endif
 
+	// Array of colors.
+	#if COMPILE_SIGNAL_COLOR_ARRAY
+		float3 ColorArray[COMPILE_SIGNAL_COLOR_ARRAY];
+	#endif
+
 	// Spherical harmonic of the color.
 	#if COMPILE_SIGNAL_COLOR_SH
 		FSSDSphericalHarmonicRGB ColorSH;
@@ -36,9 +44,6 @@ struct FSSDSignalSample
 	// Number of ray that did not find any intersections.
 	float MissCount;
 
-	// Number of ray that are uncertain due to occlusion with higher frequency.
-	float UncertainRayCount;
-
 	// Hit distance of a this sample, valid only if SampleCount == 1.
 	float ClosestHitDistance;
 	
@@ -59,6 +64,13 @@ FSSDSignalSample CreateSignalSampleFromScalarValue(float Scalar)
 	#if COMPILE_SIGNAL_COLOR
 		OutSample.SceneColor = Scalar;
 	#endif
+	#if COMPILE_SIGNAL_COLOR_ARRAY
+	{
+		UNROLL_N(COMPILE_SIGNAL_COLOR_ARRAY)
+		for (uint ColorId = 0; ColorId < COMPILE_SIGNAL_COLOR_ARRAY; ColorId++)
+			OutSample.ColorArray[ColorId] = Scalar;
+	}
+	#endif
 	#if COMPILE_SIGNAL_COLOR_SH && SPHERICAL_HARMONIC_ORDER == 2
 		OutSample.ColorSH.R.V = Scalar;
 		OutSample.ColorSH.G.V = Scalar;
@@ -77,7 +89,6 @@ FSSDSignalSample CreateSignalSampleFromScalarValue(float Scalar)
 	OutSample.SampleCount = Scalar;
 	OutSample.MissCount = Scalar;
 	OutSample.ClosestHitDistance = Scalar;
-	OutSample.UncertainRayCount = Scalar;
 	OutSample.WorldBluringRadius = Scalar;
 	return OutSample;
 }
@@ -107,13 +118,19 @@ FSSDSignalSample MulSignal(FSSDSignalSample Sample, float Scalar)
 	#if COMPILE_SIGNAL_COLOR
 		OutSample.SceneColor         = Sample.SceneColor * Scalar;
 	#endif
+	#if COMPILE_SIGNAL_COLOR_ARRAY
+	{
+		UNROLL_N(COMPILE_SIGNAL_COLOR_ARRAY)
+		for (uint ColorId = 0; ColorId < COMPILE_SIGNAL_COLOR_ARRAY; ColorId++)
+			OutSample.ColorArray[ColorId] = Sample.ColorArray[ColorId] * Scalar;
+	}
+	#endif
 	#if COMPILE_SIGNAL_COLOR_SH
 		OutSample.ColorSH            = MulSH(Sample.ColorSH, Scalar);
 	#endif
 	OutSample.SampleCount            = Sample.SampleCount * Scalar;
 	OutSample.MissCount              = Sample.MissCount * Scalar;
 	OutSample.ClosestHitDistance     = Sample.ClosestHitDistance * Scalar;
-	OutSample.UncertainRayCount      = Sample.UncertainRayCount * Scalar;
 	OutSample.WorldBluringRadius     = Sample.WorldBluringRadius * Scalar;
 	return OutSample;
 }
@@ -123,13 +140,19 @@ FSSDSignalSample AddSignal(FSSDSignalSample SampleA, FSSDSignalSample SampleB)
 	#if COMPILE_SIGNAL_COLOR
 		SampleA.SceneColor         += SampleB.SceneColor;
 	#endif
+	#if COMPILE_SIGNAL_COLOR_ARRAY
+	{
+		UNROLL_N(COMPILE_SIGNAL_COLOR_ARRAY)
+		for (uint ColorId = 0; ColorId < COMPILE_SIGNAL_COLOR_ARRAY; ColorId++)
+			SampleA.ColorArray[ColorId] += SampleB.ColorArray[ColorId];
+	}
+	#endif
 	#if COMPILE_SIGNAL_COLOR_SH
 		SampleA.ColorSH            = AddSH(SampleA.ColorSH, SampleB.ColorSH);
 	#endif
 	SampleA.SampleCount            += SampleB.SampleCount;
 	SampleA.MissCount              += SampleB.MissCount;
 	SampleA.ClosestHitDistance     += SampleB.ClosestHitDistance;
-	SampleA.UncertainRayCount      += SampleB.UncertainRayCount;
 	SampleA.WorldBluringRadius     += SampleB.WorldBluringRadius;
 	return SampleA;
 }
@@ -139,13 +162,19 @@ FSSDSignalSample MinusSignal(FSSDSignalSample Sample)
 	#if COMPILE_SIGNAL_COLOR
 		Sample.SceneColor         = -Sample.SceneColor;
 	#endif
+	#if COMPILE_SIGNAL_COLOR_ARRAY
+	{
+		UNROLL_N(COMPILE_SIGNAL_COLOR_ARRAY)
+		for (uint ColorId = 0; ColorId < COMPILE_SIGNAL_COLOR_ARRAY; ColorId++)
+			Sample.ColorArray[ColorId] = -Sample.ColorArray[ColorId];
+	}
+	#endif
 	#if COMPILE_SIGNAL_COLOR_SH
 		Sample.ColorSH         = MulSH(Sample.ColorSH, -1.0);
 	#endif
 	Sample.SampleCount            = -Sample.SampleCount;
 	Sample.MissCount              = -Sample.MissCount;
 	Sample.ClosestHitDistance     = -Sample.ClosestHitDistance;
-	Sample.UncertainRayCount      = -Sample.UncertainRayCount;
 	Sample.WorldBluringRadius     = -Sample.WorldBluringRadius;
 	return Sample;
 }
@@ -160,6 +189,13 @@ FSSDSignalSample AbsSignal(FSSDSignalSample Sample)
 	#if COMPILE_SIGNAL_COLOR
 		Sample.SceneColor         = abs(Sample.SceneColor);
 	#endif
+	#if COMPILE_SIGNAL_COLOR_ARRAY
+	{
+		UNROLL_N(COMPILE_SIGNAL_COLOR_ARRAY)
+		for (uint ColorId = 0; ColorId < COMPILE_SIGNAL_COLOR_ARRAY; ColorId++)
+			Sample.ColorArray[ColorId] = abs(Sample.ColorArray[ColorId]);
+	}
+	#endif
 	#if COMPILE_SIGNAL_COLOR_SH && SPHERICAL_HARMONIC_ORDER == 2
 		Sample.ColorSH.R.V = abs(Sample.ColorSH.R.V);
 		Sample.ColorSH.G.V = abs(Sample.ColorSH.G.V);
@@ -178,7 +214,6 @@ FSSDSignalSample AbsSignal(FSSDSignalSample Sample)
 	Sample.SampleCount            = abs(Sample.SampleCount);
 	Sample.MissCount              = abs(Sample.MissCount);
 	Sample.ClosestHitDistance     = abs(Sample.ClosestHitDistance);
-	Sample.UncertainRayCount      = abs(Sample.UncertainRayCount);
 	Sample.WorldBluringRadius     = abs(Sample.WorldBluringRadius);
 	return Sample;
 }
@@ -188,6 +223,13 @@ FSSDSignalSample SqrtSignal(FSSDSignalSample Sample)
 	#if COMPILE_SIGNAL_COLOR
 		Sample.SceneColor         = sqrt(Sample.SceneColor);
 	#endif
+	#if COMPILE_SIGNAL_COLOR_ARRAY
+	{
+		UNROLL_N(COMPILE_SIGNAL_COLOR_ARRAY)
+		for (uint ColorId = 0; ColorId < COMPILE_SIGNAL_COLOR_ARRAY; ColorId++)
+			Sample.ColorArray[ColorId] = sqrt(Sample.ColorArray[ColorId]);
+	}
+	#endif
 	#if COMPILE_SIGNAL_COLOR_SH && SPHERICAL_HARMONIC_ORDER == 2
 		Sample.ColorSH.R.V = sqrt(Sample.ColorSH.R.V);
 		Sample.ColorSH.G.V = sqrt(Sample.ColorSH.G.V);
@@ -206,7 +248,6 @@ FSSDSignalSample SqrtSignal(FSSDSignalSample Sample)
 	Sample.SampleCount            = sqrt(Sample.SampleCount);
 	Sample.MissCount              = sqrt(Sample.MissCount);
 	Sample.ClosestHitDistance     = sqrt(Sample.ClosestHitDistance);
-	Sample.UncertainRayCount      = sqrt(Sample.UncertainRayCount);
 	Sample.WorldBluringRadius     = sqrt(Sample.WorldBluringRadius);
 	return Sample;
 }
@@ -219,6 +260,13 @@ FSSDSignalSample PowerSignal(FSSDSignalSample Sample, float Exponent)
 		Sample.SceneColor.b       = pow(Sample.SceneColor.b, Exponent);
 		Sample.SceneColor.a       = pow(Sample.SceneColor.a, Exponent);
 	#endif
+	#if COMPILE_SIGNAL_COLOR_ARRAY
+	{
+		UNROLL_N(COMPILE_SIGNAL_COLOR_ARRAY)
+		for (uint ColorId = 0; ColorId < COMPILE_SIGNAL_COLOR_ARRAY; ColorId++)
+			Sample.ColorArray[ColorId] = pow(Sample.ColorArray[ColorId], Exponent);
+	}
+	#endif
 	#if COMPILE_SIGNAL_COLOR_SH && SPHERICAL_HARMONIC_ORDER == 2
 		Sample.ColorSH.R.V = pow(Sample.ColorSH.R.V, Exponent);
 		Sample.ColorSH.G.V = pow(Sample.ColorSH.G.V, Exponent);
@@ -237,7 +285,6 @@ FSSDSignalSample PowerSignal(FSSDSignalSample Sample, float Exponent)
 	Sample.SampleCount            = pow(Sample.SampleCount, Exponent);
 	Sample.MissCount              = pow(Sample.MissCount, Exponent);
 	Sample.ClosestHitDistance     = pow(Sample.ClosestHitDistance, Exponent);
-	Sample.UncertainRayCount      = pow(Sample.UncertainRayCount, Exponent);
 	Sample.WorldBluringRadius     = pow(Sample.WorldBluringRadius, Exponent);
 	return Sample;
 }
@@ -248,6 +295,13 @@ FSSDSignalSample MinSignal(FSSDSignalSample SampleA, FSSDSignalSample SampleB)
 	#if COMPILE_SIGNAL_COLOR
 		OutSample.SceneColor = min(SampleA.SceneColor, SampleB.SceneColor);
 	#endif
+	#if COMPILE_SIGNAL_COLOR_ARRAY
+	{
+		UNROLL_N(COMPILE_SIGNAL_COLOR_ARRAY)
+		for (uint ColorId = 0; ColorId < COMPILE_SIGNAL_COLOR_ARRAY; ColorId++)
+			OutSample.ColorArray[ColorId] = min(SampleA.ColorArray[ColorId], SampleB.ColorArray[ColorId]);
+	}
+	#endif
 	#if COMPILE_SIGNAL_COLOR_SH && SPHERICAL_HARMONIC_ORDER == 2
 		OutSample.ColorSH.R.V = min(SampleA.ColorSH.R.V, SampleB.ColorSH.R.V);
 		OutSample.ColorSH.G.V = min(SampleA.ColorSH.G.V, SampleB.ColorSH.G.V);
@@ -266,7 +320,6 @@ FSSDSignalSample MinSignal(FSSDSignalSample SampleA, FSSDSignalSample SampleB)
 	OutSample.SampleCount = min(SampleA.SampleCount, SampleB.SampleCount);
 	OutSample.MissCount = min(SampleA.MissCount, SampleB.MissCount);
 	OutSample.ClosestHitDistance = min(SampleA.ClosestHitDistance, SampleB.ClosestHitDistance);
-	OutSample.UncertainRayCount = min(SampleA.UncertainRayCount, SampleB.UncertainRayCount);
 	OutSample.WorldBluringRadius = min(SampleA.WorldBluringRadius, SampleB.WorldBluringRadius);
 	return OutSample;
 }
@@ -277,6 +330,13 @@ FSSDSignalSample MaxSignal(FSSDSignalSample SampleA, FSSDSignalSample SampleB)
 	#if COMPILE_SIGNAL_COLOR
 		OutSample.SceneColor = max(SampleA.SceneColor, SampleB.SceneColor);
 	#endif
+	#if COMPILE_SIGNAL_COLOR_ARRAY
+	{
+		UNROLL_N(COMPILE_SIGNAL_COLOR_ARRAY)
+		for (uint ColorId = 0; ColorId < COMPILE_SIGNAL_COLOR_ARRAY; ColorId++)
+			OutSample.ColorArray[ColorId] = max(SampleA.ColorArray[ColorId], SampleB.ColorArray[ColorId]);
+	}
+	#endif
 	#if COMPILE_SIGNAL_COLOR_SH && SPHERICAL_HARMONIC_ORDER == 2
 		OutSample.ColorSH.R.V = max(SampleA.ColorSH.R.V, SampleB.ColorSH.R.V);
 		OutSample.ColorSH.G.V = max(SampleA.ColorSH.G.V, SampleB.ColorSH.G.V);
@@ -295,7 +355,6 @@ FSSDSignalSample MaxSignal(FSSDSignalSample SampleA, FSSDSignalSample SampleB)
 	OutSample.SampleCount = max(SampleA.SampleCount, SampleB.SampleCount);
 	OutSample.MissCount = max(SampleA.MissCount, SampleB.MissCount);
 	OutSample.ClosestHitDistance = max(SampleA.ClosestHitDistance, SampleB.ClosestHitDistance);
-	OutSample.UncertainRayCount = max(SampleA.UncertainRayCount, SampleB.UncertainRayCount);
 	OutSample.WorldBluringRadius = max(SampleA.WorldBluringRadius, SampleB.WorldBluringRadius);
 	return OutSample;
 }
@@ -306,6 +365,13 @@ FSSDSignalSample ClampSignal(FSSDSignalSample Sample, FSSDSignalSample SampleMin
 	#if COMPILE_SIGNAL_COLOR
 		OutSample.SceneColor = clamp(Sample.SceneColor, SampleMin.SceneColor, SampleMax.SceneColor);
 	#endif
+	#if COMPILE_SIGNAL_COLOR_ARRAY
+	{
+		UNROLL_N(COMPILE_SIGNAL_COLOR_ARRAY)
+		for (uint ColorId = 0; ColorId < COMPILE_SIGNAL_COLOR_ARRAY; ColorId++)
+			OutSample.ColorArray[ColorId] = clamp(Sample.ColorArray[ColorId], SampleMin.ColorArray[ColorId], SampleMax.ColorArray[ColorId]);
+	}
+	#endif
 	#if COMPILE_SIGNAL_COLOR_SH && SPHERICAL_HARMONIC_ORDER == 2
 		OutSample.ColorSH.R.V = clamp(Sample.ColorSH.R.V, SampleMin.ColorSH.R.V, SampleMax.ColorSH.R.V);
 		OutSample.ColorSH.G.V = clamp(Sample.ColorSH.G.V, SampleMin.ColorSH.G.V, SampleMax.ColorSH.G.V);
@@ -324,7 +390,6 @@ FSSDSignalSample ClampSignal(FSSDSignalSample Sample, FSSDSignalSample SampleMin
 	OutSample.SampleCount = clamp(Sample.SampleCount, SampleMin.SampleCount, SampleMax.SampleCount);
 	OutSample.MissCount = clamp(Sample.MissCount, SampleMin.MissCount, SampleMax.MissCount);
 	OutSample.ClosestHitDistance = clamp(Sample.ClosestHitDistance, SampleMin.ClosestHitDistance, SampleMax.ClosestHitDistance);
-	OutSample.UncertainRayCount = clamp(Sample.UncertainRayCount, SampleMin.UncertainRayCount, SampleMax.UncertainRayCount);
 	OutSample.WorldBluringRadius = clamp(Sample.WorldBluringRadius, SampleMin.WorldBluringRadius, SampleMax.WorldBluringRadius);
 	return OutSample;
 }
diff --git a/Engine/Shaders/Private/ScreenSpaceDenoise/SSDSpatialAccumulation.usf b/Engine/Shaders/Private/ScreenSpaceDenoise/SSDSpatialAccumulation.usf
--- a/Engine/Shaders/Private/ScreenSpaceDenoise/SSDSpatialAccumulation.usf
+++ b/Engine/Shaders/Private/ScreenSpaceDenoise/SSDSpatialAccumulation.usf
@@ -119,10 +119,12 @@
 	#endif
 
 #elif CONFIG_SIGNAL_PROCESSING == SIGNAL_PROCESSING_POLYCHROMATIC_PENUMBRA_HARMONIC
-	#define COMPILE_SIGNAL_COLOR 1
+	#define COMPILE_SIGNAL_COLOR_ARRAY 2
 
 	#define CONFIG_BILATERAL_PRESET BILATERAL_PRESET_POLYCHROMATIC_PENUMBRA
-	//#define CONFIG_MULTIPLEXED_SIGNALS_PER_SIGNAL_DOMAIN 1
+	#define CONFIG_MULTIPLEXED_SIGNALS_PER_SIGNAL_DOMAIN 1
+	#undef CONFIG_SIGNAL_BATCH_SIZE
+	#define CONFIG_SIGNAL_BATCH_SIZE 2
 
 	// Denoise diffuse and specular harmonics at the same time.
 	#define MAX_SIGNAL_BATCH_SIZE 2
@@ -132,11 +134,11 @@
 
 	#if DIM_STAGE == STAGE_RECONSTRUCTION
 		// Input and output layout.
-		#define CONFIG_SIGNAL_INPUT_LAYOUT  SIGNAL_BUFFER_LAYOUT_POLYCHROMATIC_PENUMBRA_INPUT
-		#define CONFIG_SIGNAL_OUTPUT_LAYOUT SIGNAL_BUFFER_LAYOUT_POLYCHROMATIC_PENUMBRA_HISTORY
+		#define CONFIG_SIGNAL_INPUT_LAYOUT  SIGNAL_BUFFER_LAYOUT_POLYCHROMATIC_PENUMBRA_HARMONIC_INPUT
+		#define CONFIG_SIGNAL_OUTPUT_LAYOUT SIGNAL_BUFFER_LAYOUT_POLYCHROMATIC_PENUMBRA_HARMONIC_RECONSTRUCTION
 
-		#define CONFIG_INPUT_TEXTURE_COUNT 1
-		#define CONFIG_OUTPUT_TEXTURE_COUNT 3
+		#define CONFIG_INPUT_TEXTURE_COUNT 4
+		#define CONFIG_OUTPUT_TEXTURE_COUNT 4
 
 		#define CONFIG_SAMPLE_SET           SAMPLE_SET_STACKOWIAK_4_SETS
 		//#define CONFIG_MAX_WITH_REF_DISTANCE 1
diff --git a/Engine/Shaders/Private/ScreenSpaceDenoise/SSDTemporalAccumulation.usf b/Engine/Shaders/Private/ScreenSpaceDenoise/SSDTemporalAccumulation.usf
--- a/Engine/Shaders/Private/ScreenSpaceDenoise/SSDTemporalAccumulation.usf
+++ b/Engine/Shaders/Private/ScreenSpaceDenoise/SSDTemporalAccumulation.usf
@@ -71,8 +71,8 @@
 	// Input and output layout.
 	#define CONFIG_SIGNAL_INPUT_LAYOUT  SIGNAL_BUFFER_LAYOUT_POLYCHROMATIC_PENUMBRA_HISTORY
 	#define CONFIG_SIGNAL_HISTORY_LAYOUT SIGNAL_BUFFER_LAYOUT_POLYCHROMATIC_PENUMBRA_HISTORY
-	#define CONFIG_INPUT_TEXTURE_COUNT 3
-	#define CONFIG_HISTORY_TEXTURE_COUNT 3
+	#define CONFIG_INPUT_TEXTURE_COUNT 2
+	#define CONFIG_HISTORY_TEXTURE_COUNT 2
 
 	#define CONFIG_HISTORY_REJECTION HISTORY_REJECTION_VAR_BOUNDARIES
 
diff --git a/Engine/Source/Runtime/Renderer/Private/ScreenSpaceDenoise.cpp b/Engine/Source/Runtime/Renderer/Private/ScreenSpaceDenoise.cpp
--- a/Engine/Source/Runtime/Renderer/Private/ScreenSpaceDenoise.cpp
+++ b/Engine/Source/Runtime/Renderer/Private/ScreenSpaceDenoise.cpp
@@ -372,7 +372,7 @@ const TCHAR* const kReconstructionResourceNames[] = {
 	TEXT("PolychromaticPenumbraHarmonicReconstruction0"),
 	TEXT("PolychromaticPenumbraHarmonicReconstruction1"),
 	TEXT("PolychromaticPenumbraHarmonicReconstruction2"),
-	nullptr,
+	TEXT("PolychromaticPenumbraHarmonicReconstruction3"),
 
 	// Reflections
 	TEXT("ReflectionsReconstruction0"),
@@ -485,7 +485,7 @@ const TCHAR* const kTemporalAccumulationResourceNames[] = {
 	// PolychromaticPenumbraHarmonic
 	TEXT("PolychromaticPenumbraHistory0"),
 	TEXT("PolychromaticPenumbraHistory1"),
-	TEXT("PolychromaticPenumbraHistory2"),
+	nullptr,
 	nullptr,
 
 	// Reflections
@@ -987,6 +987,8 @@ class FSSDComposeHarmonicsCS : public FGlobalShader
 		SHADER_PARAMETER_STRUCT_ARRAY(FSSDSignalTextures, SignalHarmonics, [IScreenSpaceDenoiser::kMultiPolychromaticPenumbraHarmonics])
 		SHADER_PARAMETER_STRUCT_INCLUDE(FSSDCommonParameters, CommonParameters)
 		SHADER_PARAMETER_STRUCT(FSSDSignalUAVs, SignalOutput)
+
+		SHADER_PARAMETER_RDG_TEXTURE_UAV(RWTexture2D, DebugOutput)
 	END_SHADER_PARAMETER_STRUCT()
 };
 
@@ -1117,13 +1119,15 @@ static void DenoiseSignalAtConstantPixelDensity(
 		}
 		else if (Settings.SignalProcessing == ESignalProcessing::PolychromaticPenumbraHarmonic)
 		{
-			ReconstructionTextureCount = 3;
+			ReconstructionTextureCount = 4;
 			ReconstructionDescs[0].Format = PF_FloatRGBA;
 			ReconstructionDescs[1].Format = PF_FloatRGBA;
 			ReconstructionDescs[2].Format = PF_FloatRGBA;
+			ReconstructionDescs[3].Format = PF_FloatRGBA;
 
-			HistoryTextureCountPerSignal = ReconstructionTextureCount;
-			HistoryDescs = ReconstructionDescs;
+			HistoryTextureCountPerSignal = 2;
+			HistoryDescs[0].Format = PF_FloatRGBA;
+			HistoryDescs[1].Format = PF_FloatRGBA;
 		}
 		else if (Settings.SignalProcessing == ESignalProcessing::Reflections)
 		{
@@ -1663,6 +1667,36 @@ static void DenoiseSignalAtConstantPixelDensity(
 	}
 } // DenoiseSignalAtConstantPixelDensity()
 
+// static
+IScreenSpaceDenoiser::FHarmonicTextures IScreenSpaceDenoiser::CreateHarmonicTextures(FRDGBuilder& GraphBuilder, FIntPoint Extent, const TCHAR* DebugName)
+{
+	FRDGTextureDesc Desc = FRDGTextureDesc::Create2DDesc(
+		Extent,
+		PF_FloatRGBA,
+		FClearValueBinding::None,
+		TexCreate_None,
+		TexCreate_ShaderResource | TexCreate_UAV,
+		/* bInForceSeparateTargetAndShaderResource = */ false);
+
+	FHarmonicTextures HarmonicTextures;
+	for (int32 HarmonicBorderId = 0; HarmonicBorderId < kHarmonicBordersCount; HarmonicBorderId++)
+	{
+		HarmonicTextures.Harmonics[HarmonicBorderId] = GraphBuilder.CreateTexture(Desc, DebugName);
+	}
+	return HarmonicTextures;
+}
+
+// static
+IScreenSpaceDenoiser::FHarmonicUAVs IScreenSpaceDenoiser::CreateUAVs(FRDGBuilder& GraphBuilder, const FHarmonicTextures& Textures)
+{
+	FHarmonicUAVs UAVs;
+	for (int32 HarmonicBorderId = 0; HarmonicBorderId < kHarmonicBordersCount; HarmonicBorderId++)
+	{
+		UAVs.Harmonics[HarmonicBorderId] = GraphBuilder.CreateUAV(Textures.Harmonics[HarmonicBorderId]);
+	}
+	return UAVs;
+}
+
 /** The implementation of the default denoiser of the renderer. */
 class FDefaultScreenSpaceDenoiser : public IScreenSpaceDenoiser
 {
@@ -1770,8 +1804,10 @@ public:
 		// Harmonic 0 doesn't need any reconstruction given it's the highest frequency details.
 		{
 			const int32 HarmonicId = 0;
-			ComposePassParameters->SignalHarmonics[HarmonicId].Textures[0] = Inputs.Diffuse[HarmonicId];
-			ComposePassParameters->SignalHarmonics[HarmonicId].Textures[1] = Inputs.Specular[HarmonicId];
+			ComposePassParameters->SignalHarmonics[HarmonicId].Textures[0] = Inputs.Diffuse.Harmonics[0];
+			ComposePassParameters->SignalHarmonics[HarmonicId].Textures[1] = Inputs.Diffuse.Harmonics[1];
+			ComposePassParameters->SignalHarmonics[HarmonicId].Textures[2] = Inputs.Specular.Harmonics[0];
+			ComposePassParameters->SignalHarmonics[HarmonicId].Textures[3] = Inputs.Specular.Harmonics[1];
 		}
 
 		// Reconstruct each harmonic independently
@@ -1791,8 +1827,10 @@ public:
 			NewHistories[0] = nullptr;
 
 			FSSDSignalTextures InputSignal;
-			InputSignal.Textures[0] = Inputs.Diffuse[HarmonicId];
-			InputSignal.Textures[1] = Inputs.Specular[HarmonicId];
+			InputSignal.Textures[0] = Inputs.Diffuse.Harmonics[HarmonicId + 0];
+			InputSignal.Textures[1] = Inputs.Diffuse.Harmonics[HarmonicId + 1];
+			InputSignal.Textures[2] = Inputs.Specular.Harmonics[HarmonicId + 0];
+			InputSignal.Textures[3] = Inputs.Specular.Harmonics[HarmonicId + 1];
 
 			FSSDSignalTextures SignalOutput;
 			DenoiseSignalAtConstantPixelDensity(
@@ -1816,9 +1854,8 @@ public:
 					/* InTargetableFlags = */ TexCreate_ShaderResource | TexCreate_UAV,
 					/* bInForceSeparateTargetAndShaderResource = */ false);
 
-				ComposedHarmonics.Textures[0] = GraphBuilder.CreateTexture(Desc, TEXT("MultiPolychromaticComposition0"));
-				ComposedHarmonics.Textures[1] = GraphBuilder.CreateTexture(Desc, TEXT("MultiPolychromaticComposition1"));
-				ComposedHarmonics.Textures[2] = GraphBuilder.CreateTexture(Desc, TEXT("MultiPolychromaticComposition2"));
+				ComposedHarmonics.Textures[0] = GraphBuilder.CreateTexture(Desc, TEXT("PolychromaticPenumbraComposition0"));
+				ComposedHarmonics.Textures[1] = GraphBuilder.CreateTexture(Desc, TEXT("PolychromaticPenumbraComposition1"));
 			}
 
 			ComposePassParameters->CommonParameters.ViewUniformBuffer = View.ViewUniformBuffer;
@@ -1826,6 +1863,19 @@ public:
 
 			ComposePassParameters->SignalOutput = CreateMultiplexedUAVs(GraphBuilder, ComposedHarmonics);
 
+			{
+				FRDGTextureDesc DebugDesc = FRDGTextureDesc::Create2DDesc(
+					SceneTextures.SceneDepthBuffer->Desc.Extent,
+					PF_FloatRGBA,
+					FClearValueBinding::Black,
+					/* InFlags = */ TexCreate_None,
+					/* InTargetableFlags = */ TexCreate_ShaderResource | TexCreate_UAV,
+					/* bInForceSeparateTargetAndShaderResource = */ false);
+
+				FRDGTextureRef DebugTexture = GraphBuilder.CreateTexture(DebugDesc, TEXT("DebugHarmonicComposition"));
+				ComposePassParameters->DebugOutput = GraphBuilder.CreateUAV(DebugTexture);
+			}
+
 			TShaderMapRef<FSSDComposeHarmonicsCS> ComputeShader(View.ShaderMap);
 			FComputeShaderUtils::AddPass(
 				GraphBuilder,
diff --git a/Engine/Source/Runtime/Renderer/Private/ScreenSpaceDenoise.h b/Engine/Source/Runtime/Renderer/Private/ScreenSpaceDenoise.h
--- a/Engine/Source/Runtime/Renderer/Private/ScreenSpaceDenoise.h
+++ b/Engine/Source/Runtime/Renderer/Private/ScreenSpaceDenoise.h
@@ -116,10 +116,26 @@ public:
 	// Number of harmonics to be feed when denoising multiple lights.
 	static constexpr int32 kMultiPolychromaticPenumbraHarmonics = 4;
 
-	/** All the inputs to denoise polychromatic penumbra of multiple light. */
+	// Number of border between harmonics are used to denoise harmonical signal.
+	static constexpr int32 kHarmonicBordersCount = kMultiPolychromaticPenumbraHarmonics + 1;
+
+	/** High level for one harmonic of a signal to denoise. */
+	BEGIN_SHADER_PARAMETER_STRUCT(FHarmonicTextures, )
+		SHADER_PARAMETER_RDG_TEXTURE_ARRAY(Texture2D, Harmonics, [kHarmonicBordersCount])
+	END_SHADER_PARAMETER_STRUCT()
+
+	BEGIN_SHADER_PARAMETER_STRUCT(FHarmonicUAVs, )
+		SHADER_PARAMETER_RDG_TEXTURE_UAV_ARRAY(Texture2D, Harmonics, [kHarmonicBordersCount])
+	END_SHADER_PARAMETER_STRUCT()
+
+	static FHarmonicTextures CreateHarmonicTextures(FRDGBuilder& GraphBuilder, FIntPoint Extent, const TCHAR* DebugName);
+	static FHarmonicUAVs CreateUAVs(FRDGBuilder& GraphBuilder, const FHarmonicTextures& Textures);
+
+
+	/** All the inputs to denoise polychromatic penumbra of multiple lights. */
 	BEGIN_SHADER_PARAMETER_STRUCT(FPolychromaticPenumbraHarmonics, )
-		SHADER_PARAMETER_RDG_TEXTURE_ARRAY(Texture2D, Diffuse, [kMultiPolychromaticPenumbraHarmonics])
-		SHADER_PARAMETER_RDG_TEXTURE_ARRAY(Texture2D, Specular, [kMultiPolychromaticPenumbraHarmonics])
+		SHADER_PARAMETER_STRUCT(FHarmonicTextures, Diffuse)
+		SHADER_PARAMETER_STRUCT(FHarmonicTextures, Specular)
 	END_SHADER_PARAMETER_STRUCT()
 
 	/** All the outputs when denoising polychromatic penumbra. */
