Implements signal integration in multi polychromatic penumbra denoiser

diff --git a/Engine/Shaders/Private/ScreenSpaceDenoise/SSDComposeHarmonics.usf b/Engine/Shaders/Private/ScreenSpaceDenoise/SSDComposeHarmonics.usf
--- a/Engine/Shaders/Private/ScreenSpaceDenoise/SSDComposeHarmonics.usf
+++ b/Engine/Shaders/Private/ScreenSpaceDenoise/SSDComposeHarmonics.usf
@@ -54,16 +54,20 @@ FSSDTexture2D SignalHarmonics_1_Textures_0;
 FSSDTexture2D SignalHarmonics_2_Textures_0;
 FSSDTexture2D SignalHarmonics_3_Textures_0;
 
+FSSDTexture2D SignalIntegrand_Textures_0;
+
 #if CONFIG_INPUT_TEXTURE_COUNT > 1
 FSSDTexture2D SignalHarmonics_0_Textures_1;
 FSSDTexture2D SignalHarmonics_1_Textures_1;
 FSSDTexture2D SignalHarmonics_2_Textures_1;
 FSSDTexture2D SignalHarmonics_3_Textures_1;
+FSSDTexture2D SignalIntegrand_Textures_1;
 #else
 #define SignalHarmonics_0_Textures_1 SignalHarmonics_0_Textures_0
 #define SignalHarmonics_1_Textures_1 SignalHarmonics_1_Textures_0
 #define SignalHarmonics_2_Textures_1 SignalHarmonics_2_Textures_0
 #define SignalHarmonics_3_Textures_1 SignalHarmonics_3_Textures_0
+#define SignalIntegrand_Textures_1 SignalIntegrand_Textures_0
 #endif
 
 #if CONFIG_INPUT_TEXTURE_COUNT > 2
@@ -71,11 +75,13 @@ FSSDTexture2D SignalHarmonics_0_Textures_2;
 FSSDTexture2D SignalHarmonics_1_Textures_2;
 FSSDTexture2D SignalHarmonics_2_Textures_2;
 FSSDTexture2D SignalHarmonics_3_Textures_2;
+FSSDTexture2D SignalIntegrand_Textures_2;
 #else
 #define SignalHarmonics_0_Textures_2 SignalHarmonics_0_Textures_0
 #define SignalHarmonics_1_Textures_2 SignalHarmonics_1_Textures_0
 #define SignalHarmonics_2_Textures_2 SignalHarmonics_2_Textures_0
 #define SignalHarmonics_3_Textures_2 SignalHarmonics_3_Textures_0
+#define SignalIntegrand_Textures_2 SignalIntegrand_Textures_0
 #endif
 
 #if CONFIG_INPUT_TEXTURE_COUNT > 3
@@ -83,11 +89,13 @@ FSSDTexture2D SignalHarmonics_0_Textures_3;
 FSSDTexture2D SignalHarmonics_1_Textures_3;
 FSSDTexture2D SignalHarmonics_2_Textures_3;
 FSSDTexture2D SignalHarmonics_3_Textures_3;
+FSSDTexture2D SignalIntegrand_Textures_3;
 #else
 #define SignalHarmonics_0_Textures_3 SignalHarmonics_0_Textures_0
 #define SignalHarmonics_1_Textures_3 SignalHarmonics_1_Textures_0
 #define SignalHarmonics_2_Textures_3 SignalHarmonics_2_Textures_0
 #define SignalHarmonics_3_Textures_3 SignalHarmonics_3_Textures_0
+#define SignalIntegrand_Textures_3 SignalIntegrand_Textures_0
 #endif
 
 
@@ -138,7 +146,7 @@ void MainCS(
 	float2 ScreenPosition = ViewportUVToScreenPos(BufferUVToViewportUV(SceneBufferUV));
 	FSSDSampleSceneInfos RefSceneMetadata = FetchCurrentSceneInfos(ScreenPosition, SceneBufferUV);
 
-	// Sample the reference sample.
+	// Sample the harmonics.
 	FSSDSignalArray MultiplexedHarmonics[4];
 	{
 		MultiplexedHarmonics[0] = SampleMultiplexedSignals(
@@ -189,6 +197,19 @@ void MainCS(
 			SceneBufferUV,
 			/* MipLevel = */ 0);
 	}
+	
+	// Sample the integrated signal.
+	FSSDSignalArray MultiplexedIntegrand = SampleMultiplexedSignals(
+		SignalIntegrand_Textures_0,
+		SignalIntegrand_Textures_1,
+		SignalIntegrand_Textures_2,
+		SignalIntegrand_Textures_3,
+		GlobalPointClampedSampler,
+		CONFIG_SIGNAL_RECONSTRUCTED_LAYOUT,
+		/* MultiplexedSampleId = */ 0,
+		/* bNormalizeSample = */ true,
+		SceneBufferUV,
+		/* MipLevel = */ 0);
 
 	// Whitelists computation that needs to be done.
 	FSSDSignalArray OutputSamples = CreateSignalArrayFromScalarValue(0.0);
@@ -224,16 +245,24 @@ void MainCS(
 
 			ColorMask *= HarmonicMask;
 			
-			//if (HarmonicId == View.GeneralPurposeTweak)
+			//if (HarmonicId == View.GeneralPurposeTweak && MultiplexId == 0)
 			//{
 			//	DebugOutput[DispatchThreadId] = float4(HarmonicMask, 0);
 			//}
 		}
+		
+		FSSDSignalSample IntegrandSample = MultiplexedIntegrand.Array[MultiplexId];
+		float3 TotalLighting = IntegrandSample.ColorArray[0].rgb;
 
-		float3 TotalLighting = MultiplexedHarmonics[DENOISER_HARMONICS_COUNT - 1].Array[MultiplexId].ColorArray[1];
-		float3 FinalLighting = ColorMask; //TotalLighting * ColorMask;
+		float3 FinalLighting = ColorMask * TotalLighting;
 		
-		//DebugOutput[DispatchThreadId] = float4(ColorMask, 0);
+		//if (MultiplexId == 0)
+		//{
+		//	if (SceneBufferUV.x > 0.5)
+		//		DebugOutput[DispatchThreadId] = float4(TotalLighting, 0);
+		//	else
+		//		DebugOutput[DispatchThreadId] = float4(TotalLighting, 0);
+		//}
 
 		OutputSamples.Array[MultiplexId].SampleCount = LowestFrequencyWeight;
 		OutputSamples.Array[MultiplexId].SceneColor.rgb = FinalLighting * LowestFrequencyWeight;
diff --git a/Engine/Shaders/Private/ScreenSpaceDenoise/SSDSignalBufferEncoding.ush b/Engine/Shaders/Private/ScreenSpaceDenoise/SSDSignalBufferEncoding.ush
--- a/Engine/Shaders/Private/ScreenSpaceDenoise/SSDSignalBufferEncoding.ush
+++ b/Engine/Shaders/Private/ScreenSpaceDenoise/SSDSignalBufferEncoding.ush
@@ -171,6 +171,11 @@ void DecodeMultiplexedSignals(
 		// The RGS output same weigth for diffuse and specular given it depends only on shading model id.
 		float Weight = RawSample[0].a;
 
+		if (bNormalizeSample)
+		{
+			Weight = Weight > 0.0 ? 1.0 : 0.0;
+		}
+
 		// Decode diffuse
 		{
 			const uint MultiplexId = 0;
@@ -196,6 +201,10 @@ void DecodeMultiplexedSignals(
 			const uint MultiplexId = 0;
 
 			float Weight = RawSample[0].a;
+			if (bNormalizeSample)
+			{
+				Weight = Weight > 0.0 ? 1.0 : 0.0;
+			}
 
 			OutSamples.Array[MultiplexId].SampleCount = Weight;
 			OutSamples.Array[MultiplexId].ColorArray[0] = RawSample[0].rgb * Weight;
@@ -207,6 +216,10 @@ void DecodeMultiplexedSignals(
 			const uint MultiplexId = 1;
 
 			float Weight = RawSample[2].a;
+			if (bNormalizeSample)
+			{
+				Weight = Weight > 0.0 ? 1.0 : 0.0;
+			}
 
 			OutSamples.Array[MultiplexId].SampleCount = Weight;
 			OutSamples.Array[MultiplexId].ColorArray[0] = RawSample[2].rgb * Weight;
diff --git a/Engine/Source/Runtime/Renderer/Private/ScreenSpaceDenoise.cpp b/Engine/Source/Runtime/Renderer/Private/ScreenSpaceDenoise.cpp
--- a/Engine/Source/Runtime/Renderer/Private/ScreenSpaceDenoise.cpp
+++ b/Engine/Source/Runtime/Renderer/Private/ScreenSpaceDenoise.cpp
@@ -985,6 +985,7 @@ class FSSDComposeHarmonicsCS : public FGlobalShader
 
 	BEGIN_SHADER_PARAMETER_STRUCT(FParameters, )
 		SHADER_PARAMETER_STRUCT_ARRAY(FSSDSignalTextures, SignalHarmonics, [IScreenSpaceDenoiser::kMultiPolychromaticPenumbraHarmonics])
+		SHADER_PARAMETER_STRUCT(FSSDSignalTextures, SignalIntegrand)
 		SHADER_PARAMETER_STRUCT_INCLUDE(FSSDCommonParameters, CommonParameters)
 		SHADER_PARAMETER_STRUCT(FSSDSignalUAVs, SignalOutput)
 
@@ -1799,6 +1800,9 @@ public:
 	{
 		RDG_GPU_STAT_SCOPE(GraphBuilder, ShadowsDenoiser);
 
+		FRDGTextureRef BlackDummy = GraphBuilder.RegisterExternalTexture(GSystemTextures.BlackDummy);
+		FRDGTextureRef WhiteDummy = GraphBuilder.RegisterExternalTexture(GSystemTextures.WhiteDummy);
+
 		FSSDComposeHarmonicsCS::FParameters* ComposePassParameters = GraphBuilder.AllocParameters<FSSDComposeHarmonicsCS::FParameters>();
 
 		// Harmonic 0 doesn't need any reconstruction given it's the highest frequency details.
@@ -1842,6 +1846,46 @@ public:
 			ComposePassParameters->SignalHarmonics[HarmonicId] = SignalOutput;
 		}
 
+		// Denoise the entire integrand signal.
+		// TODO(Denoiser): this assume all the lights are going into lowest frequency harmonic.
+		if (1)
+		{
+			const int32 HarmonicId = IScreenSpaceDenoiser::kMultiPolychromaticPenumbraHarmonics - 1;
+
+			int32 Periode = 1 << HarmonicId;
+
+			FSSDConstantPixelDensitySettings Settings;
+			Settings.SignalProcessing = ESignalProcessing::PolychromaticPenumbraHarmonic;
+			Settings.HarmonicPeriode = Periode;
+			Settings.ReconstructionSamples = Periode * Periode; // TODO: should use preconvolution instead for harmonic 3
+			Settings.bUseTemporalAccumulation = false;
+
+			TStaticArray<FScreenSpaceFilteringHistory*, IScreenSpaceDenoiser::kMaxBatchSize> PrevHistories;
+			TStaticArray<FScreenSpaceFilteringHistory*, IScreenSpaceDenoiser::kMaxBatchSize> NewHistories;
+			PrevHistories[0] = nullptr;
+			NewHistories[0] = nullptr;
+
+			// TODO(Denoiser): pipeline permutation to be faster.
+			FSSDSignalTextures InputSignal;
+			InputSignal.Textures[0] = Inputs.Diffuse.Harmonics[0];
+			InputSignal.Textures[1] = BlackDummy;
+			InputSignal.Textures[2] = Inputs.Specular.Harmonics[0];
+			InputSignal.Textures[3] = BlackDummy;
+
+			DenoiseSignalAtConstantPixelDensity(
+				GraphBuilder, View, SceneTextures,
+				InputSignal, Settings,
+				PrevHistories, NewHistories,
+				/* out */ &ComposePassParameters->SignalIntegrand);
+		}
+		else
+		{
+			ComposePassParameters->SignalIntegrand.Textures[0] = WhiteDummy;
+			ComposePassParameters->SignalIntegrand.Textures[1] = BlackDummy;
+			ComposePassParameters->SignalIntegrand.Textures[2] = WhiteDummy;
+			ComposePassParameters->SignalIntegrand.Textures[3] = BlackDummy;
+		}
+
 		// Merges the different harmonics.
 		FSSDSignalTextures ComposedHarmonics;
 		{
