Fixes specular that was not going through the temporal accumulation

diff --git a/Engine/Shaders/Private/ScreenSpaceDenoise/SSDTemporalAccumulation.usf b/Engine/Shaders/Private/ScreenSpaceDenoise/SSDTemporalAccumulation.usf
--- a/Engine/Shaders/Private/ScreenSpaceDenoise/SSDTemporalAccumulation.usf
+++ b/Engine/Shaders/Private/ScreenSpaceDenoise/SSDTemporalAccumulation.usf
@@ -74,6 +74,9 @@
 	#define CONFIG_INPUT_TEXTURE_COUNT 2
 	#define CONFIG_HISTORY_TEXTURE_COUNT 2
 
+	#undef CONFIG_SIGNAL_BATCH_SIZE
+	#define CONFIG_SIGNAL_BATCH_SIZE 2
+
 	#define CONFIG_HISTORY_REJECTION HISTORY_REJECTION_VAR_BOUNDARIES
 
 	// Denoise diffuse and specular harmonics at the same time.
