Fix VSOutput to PSInput type mismatch.

diff --git a/Engine/Shaders/Private/PostProcessDownsample.usf b/Engine/Shaders/Private/PostProcessDownsample.usf
--- a/Engine/Shaders/Private/PostProcessDownsample.usf
+++ b/Engine/Shaders/Private/PostProcessDownsample.usf
@@ -70,9 +70,9 @@ float4 DownsampleCommon(float2 UV)
 }
 
 // pixel shader entry point
-void MainPS(noperspective float2 InUV : TEXCOORD0, out float4 OutColor : SV_Target0)
+void MainPS(noperspective float4 InUV : TEXCOORD0, out float4 OutColor : SV_Target0)
 {
-	OutColor = DownsampleCommon(InUV);
+	OutColor = DownsampleCommon(InUV.xy);
 }
 
 #if COMPUTESHADER
@@ -93,4 +93,4 @@ void MainCS(uint2 DispatchThreadId : SV_DispatchThreadID)
 		OutComputeTexture[PixelPos] = DownsampleCommon(UV);
 	}
 }
-#endif 
\ No newline at end of file
+#endif 
