Wire up MaxShadowThickness to ray tracing SkyLight pass

diff --git a/Engine/Shaders/Private/RayTracing/RayTracingSkyLightRGS.usf b/Engine/Shaders/Private/RayTracing/RayTracingSkyLightRGS.usf
--- a/Engine/Shaders/Private/RayTracing/RayTracingSkyLightRGS.usf
+++ b/Engine/Shaders/Private/RayTracing/RayTracingSkyLightRGS.usf
@@ -228,12 +228,10 @@ void SkyLightRGS()
 			Payload.SetMinimalPayloadMode(); // Only HitT will be filled by closest hit shader, skipping full material evaluation.
 
 			RayDesc RayLocal;
-			//RayLocal.Origin = WorldPosition + L * SkyLight.MaxShadowThickness;
-			RayLocal.Origin = WorldPosition + L * 1000.0;
+			RayLocal.Origin = WorldPosition + L * SkyLight.MaxShadowThickness;
 			RayLocal.Direction = -L;
 			RayLocal.TMin = 0.0;
-			//RayLocal.TMax = SkyLight.MaxShadowThickness;
-			RayLocal.TMax = 1000.0;
+			RayLocal.TMax = SkyLight.MaxShadowThickness;
 
 			FPackedMaterialClosestHitPayload PayloadLocal = (FPackedMaterialClosestHitPayload)0;
 			PayloadLocal.SetMinimalPayloadMode(); // Only HitT will be filled by closest hit shader, skipping full material evaluation.
