Outline assets not switching between multiple static mesh components on vulkan. This was due to some legacy fix that prevented vulkan from reading stencil buffer causing this behaviour. Now that vulkan supports VK_IMAGE_DEPTH_STENCIL READ_ONLY, we should remove this fix.

diff --git a/Engine/Shaders/Private/PostProcessSelectionOutline.usf b/Engine/Shaders/Private/PostProcessSelectionOutline.usf
--- a/Engine/Shaders/Private/PostProcessSelectionOutline.usf
+++ b/Engine/Shaders/Private/PostProcessSelectionOutline.usf
@@ -84,7 +84,7 @@ FPixelInfo GetPixelInfo(int2 PixelPos, int SampleID, int OffsetX, int OffsetY, f
 	// outline even between multiple selected objects (best usability)
 	Ret.ObjectId = Stencil;
 
-#if (COMPILER_GLSL == 1 && FEATURE_LEVEL < FEATURE_LEVEL_SM4) || (COMPILER_METAL && FEATURE_LEVEL < FEATURE_LEVEL_SM4) || COMPILER_VULKAN
+#if (COMPILER_GLSL == 1 && FEATURE_LEVEL < FEATURE_LEVEL_SM4) || (COMPILER_METAL && FEATURE_LEVEL < FEATURE_LEVEL_SM4)
 	// Stencil read in opengl is not supported on older versions
 	Ret.ObjectId = Ret.bObjectMask ? 2 : 0;
 #endif
