Fix VolumetricLightmapStreaming.usf console compilation by adding the missing Platform.ush include

diff --git a/Engine/Shaders/Private/VolumetricLightmapStreaming.usf b/Engine/Shaders/Private/VolumetricLightmapStreaming.usf
--- a/Engine/Shaders/Private/VolumetricLightmapStreaming.usf
+++ b/Engine/Shaders/Private/VolumetricLightmapStreaming.usf
@@ -2,6 +2,8 @@
 	VolumetricLightmapStreaming.usf
 =============================================================================*/
 
+#include "/Engine/Public/Platform.ush"
+
 uint3 ComputeBrickLayoutPosition(int BrickLayoutAllocation, uint3 BrickLayoutDimensions)
 {
 	uint3 BrickPosition = uint3(
