Fix missing UNORM causing DEVICE_UNORDEREDACCESSVIEW_RETURN_TYPE_MISMATCH on DX 11.0

diff --git a/Engine/Shaders/Private/VolumetricLightmapStreaming.usf b/Engine/Shaders/Private/VolumetricLightmapStreaming.usf
--- a/Engine/Shaders/Private/VolumetricLightmapStreaming.usf
+++ b/Engine/Shaders/Private/VolumetricLightmapStreaming.usf
@@ -44,9 +44,9 @@ Texture3D SkyBentNormal;
 Texture3D DirectionalLightShadowing;
 
 RWTexture3D<float3> OutAmbientVector;
-RWTexture3D<float4> OutSHCoefficients;
-RWTexture3D<float4> OutSkyBentNormal;
-RWTexture3D<float> OutDirectionalLightShadowing;
+RWTexture3D<UNORM float4> OutSHCoefficients;
+RWTexture3D<UNORM float4> OutSkyBentNormal;
+RWTexture3D<UNORM float> OutDirectionalLightShadowing;
 
 uint StartPosInOldVolume;
 uint StartPosInNewVolume;
