Switches ImportanceSampleVisibleGGX() to take DiskE

diff --git a/Engine/Shaders/Private/MonteCarlo.ush b/Engine/Shaders/Private/MonteCarlo.ush
--- a/Engine/Shaders/Private/MonteCarlo.ush
+++ b/Engine/Shaders/Private/MonteCarlo.ush
@@ -236,7 +236,7 @@ float4 ImportanceSampleGGX( float2 E, float a2 )
 // [ Heitz 2018, "Sampling the GGX Distribution of Visible Normals" ]
 // http://jcgt.org/published/0007/04/01/
 
-float4 ImportanceSampleVisibleGGX( float2 E, float a2, float3 V )
+float4 ImportanceSampleVisibleGGX( float2 DiskE, float a2, float3 V )
 {
 	// TODO float2 alpha for anisotropic
 	float a = sqrt(a2);
@@ -249,10 +249,7 @@ float4 ImportanceSampleVisibleGGX( float2 E, float a2, float3 V )
 	float3 Tangent0 = (Vh.z < 0.9999) ? normalize( cross( float3(0, 0, 1), Vh ) ) : float3(1, 0, 0);
 	float3 Tangent1 = cross( Vh, Tangent0 );
 
-	float Radius = sqrt( E.x );
-	float Phi = 2 * PI * E.y;
-
-	float2 p = Radius * float2( cos( Phi ), sin( Phi ) );
+	float2 p = DiskE;
 	float s = 0.5 + 0.5 * Vh.z;
 	p.y = (1 - s) * sqrt( 1 - p.x * p.x ) + s * p.y;
 
diff --git a/Engine/Shaders/Private/PathTracing/Material/PathTracingGlossy.ush b/Engine/Shaders/Private/PathTracing/Material/PathTracingGlossy.ush
--- a/Engine/Shaders/Private/PathTracing/Material/PathTracingGlossy.ush
+++ b/Engine/Shaders/Private/PathTracing/Material/PathTracingGlossy.ush
@@ -98,7 +98,7 @@ void Glossy_SampleMaterial(
 	if ((RandSample.z <= ProbSpecular)) 
 	{	
 		bool Flip = V.z < 0.0f;
-		float4 ResGGX = ImportanceSampleVisibleGGX(RandSample.xy, Pow4(Roughness), Flip ? -V : V);
+		float4 ResGGX = ImportanceSampleVisibleGGX(UniformSampleDisk(RandSample.xy), Pow4(Roughness), Flip ? -V : V);
 		float3 H = ResGGX.xyz;
 		if (Flip)
 		{
diff --git a/Engine/Shaders/Private/RayTracing/RayTracingLightingCommon.ush b/Engine/Shaders/Private/RayTracing/RayTracingLightingCommon.ush
--- a/Engine/Shaders/Private/RayTracing/RayTracingLightingCommon.ush
+++ b/Engine/Shaders/Private/RayTracing/RayTracingLightingCommon.ush
@@ -62,7 +62,7 @@ float3 GenerateReflectedRayDirection(
 
 		float NoV = saturate(dot(V, WorldNormal));
 
-		float4 Sample = ImportanceSampleVisibleGGX(E, Pow4(Roughness), TangentV);
+		float4 Sample = ImportanceSampleVisibleGGX(UniformSampleDisk(E), Pow4(Roughness), TangentV);
 
 		float3 H = mul(Sample.xyz, TangentBasis);
 		float3 L = 2 * dot(V, H) * H - V;
diff --git a/Engine/Shaders/Private/RayTracing/RayTracingReflections.usf b/Engine/Shaders/Private/RayTracing/RayTracingReflections.usf
--- a/Engine/Shaders/Private/RayTracing/RayTracingReflections.usf
+++ b/Engine/Shaders/Private/RayTracing/RayTracingReflections.usf
@@ -121,7 +121,7 @@ float3 GenerateReflectionDirection(inout RandomSequence RandSequence, uint2 Disp
 		if (BounceIndex == 0 && SamplePerPixel == 1 && RetryCount == 0)
 		{
 			E = Rand1SPPDenoiserInput(DispatchThreadId);
-			E.x *= 1 - GGX_IMPORTANT_SAMPLE_BIAS;
+			E.y *= 1 - GGX_IMPORTANT_SAMPLE_BIAS;
 		}
 
 		outDirection = GenerateReflectedRayDirection(IncidentDirection, WorldNormal, Roughness, E);
diff --git a/Engine/Shaders/Private/SSRT/SSRTReflections.usf b/Engine/Shaders/Private/SSRT/SSRTReflections.usf
--- a/Engine/Shaders/Private/SSRT/SSRTReflections.usf
+++ b/Engine/Shaders/Private/SSRT/SSRTReflections.usf
@@ -188,7 +188,7 @@ void ScreenSpaceReflectionsPS(
 			float2 E = Hammersley16( i, NumRays, Random );
 
 #if 1
-			float3 H = mul( ImportanceSampleVisibleGGX(E, a2, TangentV ).xyz, TangentBasis );
+			float3 H = mul( ImportanceSampleVisibleGGX(UniformSampleDisk(E), a2, TangentV ).xyz, TangentBasis );
 			float3 L = 2 * dot( V, H ) * H - V;
 #elif 0
 			float3 H = mul( ImportanceSampleGGX( E, a2 ).xyz, TangentBasis );
@@ -299,14 +299,14 @@ void ScreenSpaceReflectionsPS(
 			
 			#if SSR_OUTPUT_FOR_DENOISER
 			{
-				E.x *= 1 - GGX_IMPORTANT_SAMPLE_BIAS;
+				E.y *= 1 - GGX_IMPORTANT_SAMPLE_BIAS;
 			}
 			#endif
 			
 			float3x3 TangentBasis = GetTangentBasis( N );
 			float3 TangentV = mul( TangentBasis, V );
 
-			float3 H = mul( ImportanceSampleVisibleGGX(E, a2, TangentV ).xyz, TangentBasis );
+			float3 H = mul( ImportanceSampleVisibleGGX(UniformSampleDisk(E), a2, TangentV ).xyz, TangentBasis );
 			L = 2 * dot( V, H ) * H - V;
 		}
 		else
diff --git a/Engine/Shaders/Private/ScreenSpaceDenoise/SSDSignalFramework.ush b/Engine/Shaders/Private/ScreenSpaceDenoise/SSDSignalFramework.ush
--- a/Engine/Shaders/Private/ScreenSpaceDenoise/SSDSignalFramework.ush
+++ b/Engine/Shaders/Private/ScreenSpaceDenoise/SSDSignalFramework.ush
@@ -369,7 +369,7 @@ float ComputeInfinityBluringRadius(FSSDSampleSceneInfos SceneMetadata)
 		for (uint i = 0; i < 2; i++)
 		{
 			float P = i * 0.5 + 0.0;
-			float3 H = ImportanceSampleVisibleGGX(float2(E, P), a2, V).xyz;
+			float3 H = ImportanceSampleVisibleGGX(UniformSampleDisk(float2(P, E)), a2, V).xyz;
 
 			L[i] = 2 * dot( V, H ) * H - V;
 		}
@@ -427,7 +427,7 @@ float3 GenerateReflectedRayDirection(
 
 		float NoV = saturate(dot(V, WorldNormal));
 
-		float4 Sample = ImportanceSampleVisibleGGX(E, Pow4(Roughness), TangentV);
+		float4 Sample = ImportanceSampleVisibleGGX(UniformSampleDisk(E), Pow4(Roughness), TangentV);
 
 		float3 H = mul(Sample.xyz, TangentBasis);
 		float3 L = 2 * dot( V, H ) * H - V;
@@ -451,7 +451,7 @@ void GetInputRayDirectionPixel(
 	float2 E;
 	{
 		E = Rand1SPPDenoiserInput(PixelCoord);
-		E.x *= 1 - GGX_IMPORTANT_SAMPLE_BIAS;
+		E.y *= 1 - GGX_IMPORTANT_SAMPLE_BIAS;
 	}
 
 	GenerateReflectedRayDirection(
