1. Fix subsurface noise when r.RDG.ClobberResources = 1 and r.Test.ViewRectOffset = 1.
2. Fix subsurface depth issue and black artifacts. The "World Scale" is renamed as "World Unit Scale" suggested by Peter.

diff --git a/Engine/Shaders/Private/BurleyNormalizedSSSCommon.ush b/Engine/Shaders/Private/BurleyNormalizedSSSCommon.ush
--- a/Engine/Shaders/Private/BurleyNormalizedSSSCommon.ush
+++ b/Engine/Shaders/Private/BurleyNormalizedSSSCommon.ush
@@ -16,18 +16,18 @@
 #define SSS_TYPE_BURLEY	    0
 #define SSS_TYPE_SSSS		1
 
-//make sure UIMax|ClampMax of WorldScale * ENC_WORLDSCALE_TO_UNIT <= 1
-#define ENC_WORLDSCALE_IN_CM_TO_UNIT 0.02f
-#define DEC_UNIT_TO_WORLDSCALE_IN_CM 1/ENC_WORLDSCALE_IN_CM_TO_UNIT
+// Make sure UIMax|ClampMax of WorldUNitScale * ENC_WORLDUNITSCALE_IN_CM_TO_UNIT <= 1
+#define ENC_WORLDUNITSCALE_IN_CM_TO_UNIT 0.02f
+#define DEC_UNIT_TO_WORLDUNITSCALE_IN_CM 1/ENC_WORLDUNITSCALE_IN_CM_TO_UNIT
 
-//make sure UIMax|ClampMax of DiffuseMeanFreePath * ENC_DIFFUSEMEANFREEPATH_IN_MM_TO_UNIT <= 1
-#define ENC_DIFFUSEMEANFREEPATH_IN_MM_TO_UNIT 0.01f
+// Make sure UIMax|ClampMax of DiffuseMeanFreePath * ENC_DIFFUSEMEANFREEPATH_IN_MM_TO_UNIT <= 1
+#define ENC_DIFFUSEMEANFREEPATH_IN_MM_TO_UNIT (0.01f*0.2f)
 #define DEC_UNIT_TO_DIFFUSEMEANFREEPATH_IN_MM 1/ENC_DIFFUSEMEANFREEPATH_IN_MM_TO_UNIT
 //------------------------------------------------------------------------------------------
 
-float DecodeWorldScale(float EncodedWorldScale)
+float DecodeWorldUnitScale(float EncodedWorldUnitScale)
 {
-	return EncodedWorldScale * DEC_UNIT_TO_WORLDSCALE_IN_CM;
+	return EncodedWorldUnitScale * DEC_UNIT_TO_WORLDUNITSCALE_IN_CM;
 }
 
 float4 DecodeDiffuseMeanFreePath(float4 EncodedDiffuseMeanFreePath)
diff --git a/Engine/Shaders/Private/PostProcessSubsurface.usf b/Engine/Shaders/Private/PostProcessSubsurface.usf
--- a/Engine/Shaders/Private/PostProcessSubsurface.usf
+++ b/Engine/Shaders/Private/PostProcessSubsurface.usf
@@ -590,7 +590,7 @@ void SetupIndirectCS(uint3 DT_ID : SV_DispatchThreadID,
 #endif
 	}
 
-	if (all(Pos < SubsurfaceInput0_ViewportMax))
+	if (all(Pos <= SubsurfaceInput0_ViewportMax))
 	{
 		SetupTexture[Pos] = OutColor;
 		//cache the edge mask, depth, and the original uv in the buffer.
@@ -642,7 +642,7 @@ void MainIndirectDispatchCS(
 	// second pass prepares the setup from the recombine pass which doesn't need depth but wants to reconstruct the color
 	OutColor.a = GetMaskFromDepthInAlpha(OutColor.a);
 #endif
-	if (all(Pos < SubsurfaceInput0_ViewportMax))
+	if (all(Pos <= SubsurfaceInput0_ViewportMax))
 	{
 		SSSColorUAV[Pos] = OutColor;
 	}
diff --git a/Engine/Shaders/Private/SubsurfaceBurleyNormalized.ush b/Engine/Shaders/Private/SubsurfaceBurleyNormalized.ush
--- a/Engine/Shaders/Private/SubsurfaceBurleyNormalized.ush
+++ b/Engine/Shaders/Private/SubsurfaceBurleyNormalized.ush
@@ -16,7 +16,7 @@ struct FBurleyParameter
 {
 	float4 SurfaceAlbedo;
 	float4 DiffuseMeanFreePath;
-	float  WorldScale;
+	float  WorldUnitScale;
 };
 
 //Surface albedo and mean free path length
@@ -36,7 +36,7 @@ float4 GetSubsurfaceProfileDiffuseMeanFreePath(FGBufferData GBufferData)
 	return ActualSSProfilesTexture.Load(int3(BSSS_DMFP_OFFSET, SubsurfaceProfileInt, 0));
 }
 
-float GetSubsurfaceProfileWorldScale(FGBufferData GBufferData)
+float GetSubsurfaceProfileWorldUnitScale(FGBufferData GBufferData)
 {
 	// 0..255, which SubSurface profile to pick
 	uint SubsurfaceProfileInt = ExtractSubsurfaceProfileInt(GBufferData);
@@ -66,8 +66,8 @@ FBurleyParameter GetBurleyParameters(FGBufferData GBuffer)
 	BurleyParameter.SurfaceAlbedo = GetSubsurfaceProfileSurfaceAlbedo(GBuffer);
 	float4 DiffuseMeanFreePath = GetSubsurfaceProfileDiffuseMeanFreePath(GBuffer);
 	BurleyParameter.DiffuseMeanFreePath = DecodeDiffuseMeanFreePath(DiffuseMeanFreePath);
-	float WorldScale = DecodeWorldScale(GetSubsurfaceProfileWorldScale(GBuffer))*BURLEY_CM_2_MM;
-	BurleyParameter.WorldScale = WorldScale;
+	float WorldUnitScale = DecodeWorldUnitScale(GetSubsurfaceProfileWorldUnitScale(GBuffer))*BURLEY_CM_2_MM;
+	BurleyParameter.WorldUnitScale = WorldUnitScale;
 	BRANCH if (GBuffer.ShadingModelID == SHADINGMODELID_SUBSURFACE_PROFILE)
 	{
 		// scale with opacity stored in CustomData.a
@@ -144,7 +144,7 @@ float GetPdf(float Radius, float L, float S)
 
 struct FBurleySampleInfo
 {
-	float Radius;
+	float RadiusInMM;
 	float Theta;
 	float Pdf;
 };
@@ -164,7 +164,7 @@ FBurleySampleInfo GenerateSampleInfo(float2 Rand0T1, float DiffuseMeanFreePathFo
 	float FoundRoot = RadiusRootFindByApproximation(DiffuseMeanFreePathForSample / SpectralForSample, Rand0T1.x);
 #endif
 
-	BurleySampleInfo.Radius = max(FoundRoot, 0.00001f);//in mm
+	BurleySampleInfo.RadiusInMM = max(FoundRoot, 0.00001f);
 	
 	// Sample angle
 #ifdef SAMPLE_ANGLE_RANDOM
@@ -175,7 +175,7 @@ FBurleySampleInfo GenerateSampleInfo(float2 Rand0T1, float DiffuseMeanFreePathFo
 #endif
 
 	// Estimate Pdf
-	BurleySampleInfo.Pdf = GetPdf(BurleySampleInfo.Radius, DiffuseMeanFreePathForSample, SpectralForSample); //
+	BurleySampleInfo.Pdf = GetPdf(BurleySampleInfo.RadiusInMM, DiffuseMeanFreePathForSample, SpectralForSample);
 
 	return BurleySampleInfo;
 }
@@ -242,7 +242,7 @@ int GetNumOfSamplesBasedOnQuality(float2 UV)
 
 	//2. Clamp down the number of samples when the gamma corrected luminance is too large
 	//   because it does not help to improve the quality, but degrade the performance.
-	NumOfSamples = lerp(NumOfSamples, HIGH_LUMA_SAMPLE_COUNT, step(1.5, QualityMatrix.a)); // 1.5 is an experimental value
+	NumOfSamples = lerp(NumOfSamples, HIGH_LUMA_SAMPLE_COUNT, step(3.2, QualityMatrix.a)); // 3.2 is an experimental value to have best quality.
 
 	//3. Reduce the number of samples at the edge due to TAA (Flickering)
 	float Mask = Texture2DSample(SubsurfaceInput2_Texture, SharedSubsurfaceSampler2, UV).r;
@@ -385,11 +385,16 @@ float4 SampleSharedDiffuse(float2 SampleGroupThreadId)
 float4 SampleSSSColorConsideringLocalShared(float2 CenterUV, float2 UVOffset, uint2 CenterGroupThreadID, float MipLevel)
 {
 	float2 ClampedUV = clamp(CenterUV+ UVOffset, SubsurfaceInput0_UVViewportBilinearMin, SubsurfaceInput0_UVViewportBilinearMax);
+	
+	// Burley works only with point sampler when we support world unit scale. Bilinear and trilinear will create artifacts.
+	// So we use SubsurfaceSampler1 here instead of using SubsurfaceSampler0, which is a point sampler.
+	// The reason that we can have different sampler for SubsurfaceSampler0 is that we have Separable running in the same pass.
+
 #ifdef SUBSURFACE_BURLEY_COMPUTE
 
 #ifdef TEXTURE_CACHE_DISABLED
 	
-	return Texture2DSampleLevel(SubsurfaceInput0_Texture, SubsurfaceSampler0, ClampedUV, MipLevel);
+	return Texture2DSampleLevel(SubsurfaceInput0_Texture, SubsurfaceSampler1, ClampedUV, MipLevel);
 #else
 	float2 SampleGroupThreadId = ConvertUVOffset2GTIDOffsetForInput0(UVOffset) + CenterGroupThreadID; // Subtract 0.5 is to make sampling match. (0,0), samples at 1/w*0.5;
 	bool bUseLocalShared = MipLevel == 0 && IsIDInsideLocalShared(SampleGroupThreadId);// We will have artifacts if we do not limit the miplevel
@@ -401,12 +406,12 @@ float4 SampleSSSColorConsideringLocalShared(float2 CenterUV, float2 UVOffset, ui
 		}
 		else
 		{
-			return Texture2DSampleLevel(SubsurfaceInput0_Texture, SubsurfaceSampler0, ClampedUV, MipLevel);
+			return Texture2DSampleLevel(SubsurfaceInput0_Texture, SubsurfaceSampler1, ClampedUV, MipLevel);
 		}
 #endif
 
 #else
-	return Texture2DSampleLevel(SubsurfaceInput0_Texture, SubsurfaceSampler0, ClampedUV, MipLevel);
+	return Texture2DSampleLevel(SubsurfaceInput0_Texture, SubsurfaceSampler1, ClampedUV, MipLevel);
 #endif
 }
 
@@ -496,7 +501,7 @@ float4 BurleyNormalizedSS(float2 BufferUV, uint2 GroupThreadID)
 
 		// World unit to screen space unit, we use world scale to control
 		// the scaling of the radius unit when mean free path is fixed.
-		float2 UVOffset = float2(XScale, YScale)*BurleySampleInfo.Radius*BurleyParameter.WorldScale;
+		float2 UVOffset = float2(XScale, YScale)*BurleySampleInfo.RadiusInMM*BurleyParameter.WorldUnitScale;
 		UVOffset.x *= cos(BurleySampleInfo.Theta);
 		UVOffset.y *= sin(BurleySampleInfo.Theta);
 
@@ -510,7 +515,7 @@ float4 BurleyNormalizedSS(float2 BufferUV, uint2 GroupThreadID)
 #else
 
 			// Determine the miplevel with the expected number of samples at the pixel.
-			float MipLevel = 0.5*max(-log(BurleySampleInfo.Pdf*NumOfSamples / (BurleyParameter.WorldScale*BurleyParameter.WorldScale)), 0);
+			float MipLevel = 0.5*max(-log(BurleySampleInfo.Pdf*NumOfSamples), 0);
 
 			// If we are using half resolution,we should shift the mip level by -1
 #if SUBSURFACE_HALF_RES
@@ -525,18 +530,15 @@ float4 BurleyNormalizedSS(float2 BufferUV, uint2 GroupThreadID)
 #define USE_BILATERAL_FILTERING
 
 #ifdef USE_BILATERAL_FILTERING
-			// 10.0/depth is the depth in world unit (cm) ???
-			float DeltaDepth = min(abs(10.0f / SampledRadianceAndDepth.w - 10.0f / DepthAtDiscCenter), 60000)*BURLEY_CM_2_MM;//in mm
-#ifdef SUBSURFACE_HALF_RES
+			float DeltaDepth = abs(SampledRadianceAndDepth.w - DepthAtDiscCenter)*BURLEY_CM_2_MM;
+			float RadiusSampledInMM = sqrt(BurleySampleInfo.RadiusInMM * BurleySampleInfo.RadiusInMM + DeltaDepth * DeltaDepth);
 
-#endif
-			float RadiusSampled = sqrt(BurleySampleInfo.Radius * BurleySampleInfo.Radius + DeltaDepth * DeltaDepth); //mm.
 //#define RESAMPLE_PDF
 #ifdef RESAMPLE_PDF
-			BurleySampleInfo.Pdf = GetPdf(RadiusSampled, DiffuseMeanFreePathForSample, S);
+			BurleySampleInfo.Pdf = GetPdf(RadiusSampledInMM, DiffuseMeanFreePathForSample, S);
 #endif
 #else
-			float RadiusSampled = BurleySampleInfo.Radius;
+			float RadiusSampledInMM = BurleySampleInfo.RadiusInMM;
 #endif				
 
 			// Determine the tint color, if the sampling pixel is not subsurface, we use tint color
@@ -546,12 +548,10 @@ float4 BurleyNormalizedSS(float2 BufferUV, uint2 GroupThreadID)
 			Mask = UseSubsurfaceProfile(GetScreenSpaceDataFromCache(SampledDiscUV).GBuffer.ShadingModelID) ? 1 : 0;
 
 			// Step 4: accumulate radiance from the diffusion profile rR(r)
-			// make sure the DiffuseMeanFreePath is not zero and in m.
-			float3 DiffusionProfile = GetDiffuseReflectProfileWithDiffuseMeanFreePath(BurleyParameter.SurfaceAlbedo.xyz, BurleyParameter.DiffuseMeanFreePath.xyz, S3D.xyz, RadiusSampled);
+			// make sure the DiffuseMeanFreePath is not zero and in mm.
+			float3 DiffusionProfile = GetDiffuseReflectProfileWithDiffuseMeanFreePath(BurleyParameter.SurfaceAlbedo.xyz, BurleyParameter.DiffuseMeanFreePath.xyz, S3D.xyz, RadiusSampledInMM);
 			RadianceAccumulated.xyz += (DiffusionProfile / BurleySampleInfo.Pdf) * (SampledRadianceAndDepth.xyz*TintColor*Mask);
 
-			// Use Normal and variance to the bilaterial filter.
-
 
 #ifdef USE_BILATERAL_FILTERING
 			WeightingFactor += (DiffusionProfile / BurleySampleInfo.Pdf)*Mask;
diff --git a/Engine/Source/Runtime/Engine/Classes/Engine/SubsurfaceProfile.h b/Engine/Source/Runtime/Engine/Classes/Engine/SubsurfaceProfile.h
--- a/Engine/Source/Runtime/Engine/Classes/Engine/SubsurfaceProfile.h
+++ b/Engine/Source/Runtime/Engine/Classes/Engine/SubsurfaceProfile.h
@@ -28,14 +28,14 @@ struct FSubsurfaceProfileStruct
 	/**
 	* Subsurface mean free path distance in world/unreal units (cm)
 	*/
-	UPROPERTY(Category = "Burley Normalized", EditAnywhere, BluePrintReadOnly, meta = (ClampMin = "0.1", UIMax = "10.0", ClampMax = "10.0"))
+	UPROPERTY(Category = "Burley Normalized", EditAnywhere, BluePrintReadOnly, meta = (ClampMin = "0.1", UIMax = "50.0", ClampMax = "50.0"))
 		float MeanFreePathDistance;
 
 	/**
 	* Control the scale of world/unreal units (cm)
 	*/
 	UPROPERTY(Category = "Burley Normalized", EditAnywhere, BlueprintReadOnly, meta = (ClampMin = "0.1", UIMax = "50.0", ClampMax = "50.0"))
-		float WorldScale;
+		float WorldUnitScale;
 
 	/**
 	* Effective only when Burley subsurface scattering is enabled in cmd.
@@ -106,9 +106,8 @@ struct FSubsurfaceProfileStruct
 		//match FalloffColor exactly for the default setting
 		SurfaceAlbedo = FLinearColor(0.1447f,0.0535f,0.0432f);
 		MeanFreePathColor = FLinearColor(0.5963f,0.2679f,0.2253f);
-		MeanFreePathDistance = 0.1f;
-		//match ScatterRadius
-		WorldScale = 1.2f;
+		MeanFreePathDistance = 1.2f;
+		WorldUnitScale = 0.1f;
 	}
 
 	void Invalidate()
@@ -128,8 +127,8 @@ struct FSubsurfaceProfileStruct
 		bEnableBurley = false;
 		SurfaceAlbedo = FLinearColor(0, 0, 0);
 		MeanFreePathColor = FLinearColor(0, 0, 0);
-		MeanFreePathDistance = 0.1f;
-		WorldScale = 0.0f;
+		MeanFreePathDistance = 0.0f;
+		WorldUnitScale = 0.0f;
 	}
 };
 
diff --git a/Engine/Source/Runtime/Engine/Private/Rendering/SubsurfaceProfile.cpp b/Engine/Source/Runtime/Engine/Private/Rendering/SubsurfaceProfile.cpp
--- a/Engine/Source/Runtime/Engine/Private/Rendering/SubsurfaceProfile.cpp
+++ b/Engine/Source/Runtime/Engine/Private/Rendering/SubsurfaceProfile.cpp
@@ -152,24 +152,25 @@ static float GetNextSmallerPositiveFloat(float x)
 #define SSS_TYPE_BURLEY	    0
 #define SSS_TYPE_SSSS		1
 
-//make sure UIMax|ClampMax of WorldScale * ENC_WORLDSCALE_TO_UNIT <= 1
-#define ENC_WORLDSCALE_IN_CM_TO_UNIT 0.02f
-#define DEC_UNIT_TO_WORLDSCALE_IN_CM 1/ENC_WORLDSCALE_IN_CM_TO_UNIT
+// Make sure UIMax|ClampMax of WorldUnitScale * ENC_WORLDUNITSCALE_IN_CM_TO_UNIT <= 1
+#define ENC_WORLDUNITSCALE_IN_CM_TO_UNIT 0.02f
+#define DEC_UNIT_TO_WORLDUNITSCALE_IN_CM 1/ENC_WORLDUNITSCALE_IN_CM_TO_UNIT
 
-//make sure UIMax|ClampMax of DiffuseMeanFreePath * ENC_DIFFUSEMEANFREEPATH_IN_MM_TO_UNIT <= 1
-#define ENC_DIFFUSEMEANFREEPATH_IN_MM_TO_UNIT 0.01f
+// Make sure UIMax|ClampMax of DiffuseMeanFreePath * 10(cm to mm) * ENC_DIFFUSEMEANFREEPATH_IN_MM_TO_UNIT <= 1
+//
+#define ENC_DIFFUSEMEANFREEPATH_IN_MM_TO_UNIT (0.01f*0.2f)
 #define DEC_UNIT_TO_DIFFUSEMEANFREEPATH_IN_MM 1/ENC_DIFFUSEMEANFREEPATH_IN_MM_TO_UNIT
 //------------------------------------------------------------------------------------------
 
 //in [0,1]
-float EncodeWorldScale(float WorldScale)
+float EncodeWorldUnitScale(float WorldUnitScale)
 {
-	return WorldScale * ENC_WORLDSCALE_IN_CM_TO_UNIT;
+	return WorldUnitScale * ENC_WORLDUNITSCALE_IN_CM_TO_UNIT;
 }
 
-float DecodeWorldScale(float EncodedWorldScale)
+float DecodeWorldUnitScale(float EncodedWorldUnitScale)
 {
-	return EncodedWorldScale * DEC_UNIT_TO_WORLDSCALE_IN_CM;
+	return EncodedWorldUnitScale * DEC_UNIT_TO_WORLDUNITSCALE_IN_CM;
 }
 
 //in [0,1]
@@ -251,7 +252,7 @@ void FSubsurfaceProfileTexture::CreateTexture(FRHICommandListImmediate& RHICmdLi
 		
 		// to allow blending of the Subsurface with fullres in the shader
 		TextureRow[SSSS_SUBSURFACE_COLOR_OFFSET] = Data.SubsurfaceColor;
-		TextureRow[SSSS_SUBSURFACE_COLOR_OFFSET].A = EncodeWorldScale(Data.WorldScale);
+		TextureRow[SSSS_SUBSURFACE_COLOR_OFFSET].A = EncodeWorldUnitScale(Data.WorldUnitScale);
 		
 		FLinearColor DifffuseMeanFreePath = Data.MeanFreePathColor*Data.MeanFreePathDistance*10.0f; // convert cm to mm.
 		SetupSurfaceAlbedoAndDiffuseMeanFreePath(Data.SurfaceAlbedo, DifffuseMeanFreePath);
@@ -372,7 +373,7 @@ bool FSubsurfaceProfileTexture::GetEntryString(uint32 Index, FString& Out) const
 
 
 	Out = FString::Printf(TEXT(" %c. %p ScatterRadius=%.1f, SubsurfaceColor=%.1f %.1f %.1f, FalloffColor=%.1f %.1f %.1f, \
-								SurfaceAlbedo=%.1f %.1f %.1f, MeanFreePathColor=%.1f %.1f %.1f, MeanFreePathDistance=%.1f, WorldScale=%.1f"), 
+								SurfaceAlbedo=%.1f %.1f %.1f, MeanFreePathColor=%.1f %.1f %.1f, MeanFreePathDistance=%.1f, WorldUnitScale=%.1f"), 
 		MiniFontCharFromIndex(Index), 
 		SubsurfaceProfileEntries[Index].Profile,
 		ref.ScatterRadius,
@@ -381,7 +382,7 @@ bool FSubsurfaceProfileTexture::GetEntryString(uint32 Index, FString& Out) const
 		ref.SurfaceAlbedo.R, ref.SurfaceAlbedo.G, ref.SurfaceAlbedo.B,
 		ref.MeanFreePathColor.R, ref.MeanFreePathColor.G, ref.MeanFreePathColor.B,
 		ref.MeanFreePathDistance,
-		ref.WorldScale);
+		ref.WorldUnitScale);
 
 	return true;
 }
@@ -424,8 +425,8 @@ void FSubsurfaceProfileTexture::Dump()
 				SubsurfaceProfileEntries[i].Settings.MeanFreePathColor.R, SubsurfaceProfileEntries[i].Settings.MeanFreePathColor.G, SubsurfaceProfileEntries[i].Settings.MeanFreePathColor.B);
 			UE_LOG(LogSubsurfaceProfile, Log, TEXT("     MeanFreePathDistance=%f"),
 				SubsurfaceProfileEntries[i].Settings.MeanFreePathDistance);
-			UE_LOG(LogSubsurfaceProfile, Log, TEXT("     WorldScale=%f"),
-				SubsurfaceProfileEntries[i].Settings.WorldScale);
+			UE_LOG(LogSubsurfaceProfile, Log, TEXT("     WorldUnitScale=%f"),
+				SubsurfaceProfileEntries[i].Settings.WorldUnitScale);
 		}
 	}
 
@@ -488,6 +489,5 @@ void USubsurfaceProfile::PostLoad()
 		MapFallOffColor2SurfaceAlbedoAndDiffuseMeanFreePath(Settings.FalloffColor.R, Settings.SurfaceAlbedo.R, Settings.MeanFreePathColor.R);
 		MapFallOffColor2SurfaceAlbedoAndDiffuseMeanFreePath(Settings.FalloffColor.G, Settings.SurfaceAlbedo.G, Settings.MeanFreePathColor.G);
 		MapFallOffColor2SurfaceAlbedoAndDiffuseMeanFreePath(Settings.FalloffColor.B, Settings.SurfaceAlbedo.B, Settings.MeanFreePathColor.B);
-		Settings.WorldScale = Settings.ScatterRadius * 0.63; // Experimental value
 	}
 }
\ No newline at end of file
