Fixes SSR

diff --git a/Engine/Shaders/Private/SSRT/SSRTRayCast.ush b/Engine/Shaders/Private/SSRT/SSRTRayCast.ush
--- a/Engine/Shaders/Private/SSRT/SSRTRayCast.ush
+++ b/Engine/Shaders/Private/SSRT/SSRTRayCast.ush
@@ -74,6 +74,8 @@ FSSRTRay InitScreenSpaceRayFromWorldSpace(
 	FSSRTRay Ray;
 	Ray.RayStartScreen = RayStartScreen;
 	Ray.RayStepScreen = RayEndScreen - RayStartScreen;
+	
+	Ray.RayStepScreen *= GetStepScreenFactorToClipAtScreenEdge(RayStartScreen.xy, Ray.RayStepScreen.xy);
 
 	// TODO
 	#if IS_SSGI_SHADER
@@ -119,6 +121,8 @@ FSSRTRay InitScreenSpaceRay(
 	FSSRTRay Ray;
 	Ray.RayStartScreen = RayStartScreen;
 	Ray.RayStepScreen = RayEndScreen - RayStartScreen;
+	
+	Ray.RayStepScreen *= GetStepScreenFactorToClipAtScreenEdge(RayStartScreen.xy, Ray.RayStepScreen.xy);
 
 	// TODO
 	#if IS_SSGI_SHADER
@@ -144,8 +148,6 @@ bool CastScreenSpaceRay(
 	const float3 RayStartScreen = Ray.RayStartScreen;
 	float3 RayStepScreen = Ray.RayStepScreen;
 
-	RayStepScreen *= GetStepScreenFactorToClipAtScreenEdge(RayStartScreen.xy, RayStepScreen.xy);
-
 	float3 RayStartUVz = float3( (RayStartScreen.xy * float2( 0.5, -0.5 ) + 0.5) * HZBUvFactorAndInvFactor.xy, RayStartScreen.z );
 	float3 RayStepUVz  = float3(  RayStepScreen.xy  * float2( 0.5, -0.5 )		 * HZBUvFactorAndInvFactor.xy, RayStepScreen.z );
 	
