Fixed copy statement of base-struct members for DXC on Metal.

#rnx

diff --git a/Engine/Shaders/Private/BasePassTessellationShaders.usf b/Engine/Shaders/Private/BasePassTessellationShaders.usf
--- a/Engine/Shaders/Private/BasePassTessellationShaders.usf
+++ b/Engine/Shaders/Private/BasePassTessellationShaders.usf
@@ -57,7 +57,50 @@ FBasePassVSToPS PassFinalizeTessellationOutput(FBasePassVSToDS Interpolants, flo
 	O.FactoryInterpolants = VertexFactoryAssignInterpolants(Interpolants.FactoryInterpolants);
 
 	// Copy everything in the shared base class over to the pixel shader input
+#if METAL_TESSELLATION_PROFILE
+
+	#if NEEDS_BASEPASS_VERTEX_FOGGING
+		O.BasePassInterpolants.VertexFog = Interpolants.BasePassInterpolants.VertexFog;
+	#endif
+
+	#if !TESSELLATION_SUPPORTED
+		#if USE_WORLD_POSITION_EXCLUDING_SHADER_OFFSETS
+		O.BasePassInterpolants.PixelPositionExcludingWPO = Interpolants.BasePassInterpolants.PixelPositionExcludingWPO;
+		#endif
+	#endif
+
+	#if TRANSLUCENCY_PERVERTEX_LIGHTING_VOLUME
+		O.BasePassInterpolants.AmbientLightingVector = Interpolants.BasePassInterpolants.AmbientLightingVector;
+	#endif
+
+	#if TRANSLUCENCY_PERVERTEX_LIGHTING_VOLUME && TRANSLUCENCY_LIGHTING_VOLUMETRIC_PERVERTEX_DIRECTIONAL
+		O.BasePassInterpolants.DirectionalLightingVector = Interpolants.BasePassInterpolants.DirectionalLightingVector;
+	#endif
+
+	#if TRANSLUCENCY_PERVERTEX_FORWARD_SHADING
+		O.BasePassInterpolants.VertexDiffuseLighting = Interpolants.BasePassInterpolants.VertexDiffuseLighting;
+	#endif
+
+	#if PRECOMPUTED_IRRADIANCE_VOLUME_LIGHTING
+		#if TRANSLUCENCY_LIGHTING_VOLUMETRIC_PERVERTEX_NONDIRECTIONAL
+			O.BasePassInterpolants.VertexIndirectAmbient = Interpolants.BasePassInterpolants.VertexIndirectAmbient;
+		#elif TRANSLUCENCY_LIGHTING_VOLUMETRIC_PERVERTEX_DIRECTIONAL
+			O.BasePassInterpolants.VertexIndirectSH[0] = Interpolants.BasePassInterpolants.VertexIndirectSH[0];
+			O.BasePassInterpolants.VertexIndirectSH[1] = Interpolants.BasePassInterpolants.VertexIndirectSH[1];
+			O.BasePassInterpolants.VertexIndirectSH[2] = Interpolants.BasePassInterpolants.VertexIndirectSH[2];
+		#endif
+	#endif
+
+	#if WRITES_VELOCITY_TO_GBUFFER
+		O.BasePassInterpolants.VelocityPrevScreenPosition = Interpolants.BasePassInterpolants.VelocityPrevScreenPosition;
+		#if WRITES_VELOCITY_TO_GBUFFER_USE_POS_INTERPOLATOR
+			O.BasePassInterpolants.VelocityScreenPosition = Interpolants.BasePassInterpolants.VelocityScreenPosition;
+		#endif
+	#endif
+	
+#else // METAL_TESSELLATION_PROFILE
 	(FSharedBasePassInterpolants)O.BasePassInterpolants = (FSharedBasePassInterpolants)Interpolants.BasePassInterpolants;
+#endif // METAL_TESSELLATION_PROFILE
 
 	// Transform position to clip-space
 	{
