Cherry-pick 9318627: Removes max(0.04, ...) of the Roughness of material when applying DBuffer decals.

This is issue was making the roughess different between translucency and opaque, and based on whether the project supported decals. And this specific hard constant value was making the reflection denoiser blurrier for perfect mirror surfaces at high resolution.

diff --git a/Engine/Shaders/Private/DBufferDecalShared.ush b/Engine/Shaders/Private/DBufferDecalShared.ush
--- a/Engine/Shaders/Private/DBufferDecalShared.ush
+++ b/Engine/Shaders/Private/DBufferDecalShared.ush
@@ -127,8 +127,8 @@ void ApplyDBufferData(
 {
 	// We normalise the normal to get smoother visual result (it helps to avoid having D_GGX explodes toward infinity, and matches DecodeGBufferData)
 	WorldNormal = normalize(WorldNormal * DBufferData.NormalOpacity + DBufferData.PreMulWorldNormal);
-	// We clamp roughness to a minimum value to avoid (also helps with D_GGX)
-	Roughness = max(0.04, Roughness * DBufferData.RoughnessOpacity + DBufferData.PreMulRoughness);
+
+	Roughness = Roughness * DBufferData.RoughnessOpacity + DBufferData.PreMulRoughness;
 
 	Metallic = Metallic * DBufferData.RoughnessOpacity + DBufferData.PreMulMetallic;
 	Specular = Specular * DBufferData.RoughnessOpacity + DBufferData.PreMulSpecular;
