Copying //UE4/Dev-Mobile@9319092 to Dev-RenderPlat-Staging (//UE4/Dev-RenderPlat-Staging)

diff --git a/Engine/Shaders/Private/MobileBasePassPixelShader.usf b/Engine/Shaders/Private/MobileBasePassPixelShader.usf
--- a/Engine/Shaders/Private/MobileBasePassPixelShader.usf
+++ b/Engine/Shaders/Private/MobileBasePassPixelShader.usf
@@ -37,7 +37,7 @@ int NumDynamicPointLights;
 float4 LightPositionAndInvRadius[MAX_DYNAMIC_POINT_LIGHTS];
 float4 LightColorAndFalloffExponent[MAX_DYNAMIC_POINT_LIGHTS];
 float4 SpotLightAngles[MAX_DYNAMIC_POINT_LIGHTS];
-float4 SpotLightDirection[MAX_DYNAMIC_POINT_LIGHTS]; 
+float4 SpotLightDirectionAndSpecularScale[MAX_DYNAMIC_POINT_LIGHTS];
 #endif
 
 #ifndef PROJECT_MOBILE_ENABLE_MOVABLE_SPOTLIGHTS
@@ -437,7 +437,7 @@ void Main(
 #endif
 	}
 #endif
-
+	
 #if MATERIAL_SHADINGMODEL_SINGLELAYERWATER
 	// Fade out diffuse as this will be hanlded by the single scattering lighting. when over the water surface.
 	// We keep the SpecularColor for sun/water interactions
@@ -544,7 +544,7 @@ void Main(
 				#endif
 
 				#if FADE_CSM
-					float Fade = saturate(MaterialParameters.ScreenPosition.w * MobileDirectionalLight.DirectionalLightDistanceFadeMAD.x + MobileDirectionalLight.DirectionalLightDistanceFadeMAD.y);
+					float Fade = saturate(MaterialParameters.ScreenPosition.w * MobileDirectionalLight.DirectionalLightDistanceFadeMADAndSpecularScale.x + MobileDirectionalLight.DirectionalLightDistanceFadeMADAndSpecularScale.y);
 					// lerp out shadow based on fade params.
 					ShadowMap = lerp(ShadowMap, 1.0, Fade * Fade);
 				#endif
@@ -566,7 +566,8 @@ void Main(
 		#if FULLY_ROUGH
 			Color += (Shadow * NoL) * MobileDirectionalLight.DirectionalLightColor.rgb * DiffuseColor;
 		#else
-			Color += (Shadow * NoL) * MobileDirectionalLight.DirectionalLightColor.rgb * (DiffuseColor + SpecularColor * CalcSpecular(Roughness, RoL, NoH, H, MaterialParameters.WorldNormal));
+			// MobileDirectionalLight.DirectionalLightDistanceFadeMADAndSpecularScale.z saves SpecularScale for direction light.
+			Color += (Shadow * NoL) * MobileDirectionalLight.DirectionalLightColor.rgb * (DiffuseColor + SpecularColor * MobileDirectionalLight.DirectionalLightDistanceFadeMADAndSpecularScale.z * CalcSpecular(Roughness, RoL, NoH, H, MaterialParameters.WorldNormal));
 			// Environment map has been prenormalized, scale by lightmap luminance
 			half3 SpecularIBL = GetImageBasedReflectionLighting(MaterialParameters, Roughness) * IndirectIrradiance;
 
@@ -636,12 +637,12 @@ void Main(
 				#if PROJECT_MOBILE_ENABLE_MOVABLE_SPOTLIGHTS
 				if(SpotLightAngles[i].w > 0.0)
 				{
-					Attenuation *= SpotAttenuation(L, -SpotLightDirection[i].xyz, SpotLightAngles[i].xy);
+					Attenuation *= SpotAttenuation(L, -SpotLightDirectionAndSpecularScale[i].xyz, SpotLightAngles[i].xy);
 				}
 				#endif
 
 				#if !FULLY_ROUGH
-					Color += min(65000.0, (Attenuation * PointNoL) * LightColorAndFalloffExponent[i].rgb * (1.0/PI) * (DiffuseColor + SpecularColor * CalcSpecular(Roughness, PointRoL, PointNoH, H, MaterialParameters.WorldNormal)));
+					Color += min(65000.0, (Attenuation * PointNoL) * LightColorAndFalloffExponent[i].rgb * (1.0/PI) * (DiffuseColor + SpecularColor * SpotLightDirectionAndSpecularScale[i].w * CalcSpecular(Roughness, PointRoL, PointNoH, H, MaterialParameters.WorldNormal)));
 				#else
 					Color += (Attenuation * PointNoL) * LightColorAndFalloffExponent[i].rgb * DiffuseColor;
 				#endif
diff --git a/Engine/Source/Programs/UnrealBuildTool/Platform/Android/UEDeployAndroid.cs b/Engine/Source/Programs/UnrealBuildTool/Platform/Android/UEDeployAndroid.cs
--- a/Engine/Source/Programs/UnrealBuildTool/Platform/Android/UEDeployAndroid.cs
+++ b/Engine/Source/Programs/UnrealBuildTool/Platform/Android/UEDeployAndroid.cs
@@ -3158,6 +3158,7 @@ namespace UnrealBuildTool
 
 			GradleProperties.AppendLine("org.gradle.daemon=false");
 			GradleProperties.AppendLine("org.gradle.jvmargs=-XX:MaxHeapSize=4096m -Xmx9216m");
+			GradleProperties.AppendLine("android.injected.testOnly=false");
 			GradleProperties.AppendLine(string.Format("COMPILE_SDK_VERSION={0}", CompileSDKVersion));
 			GradleProperties.AppendLine(string.Format("BUILD_TOOLS_VERSION={0}", BuildToolsVersion));
 			GradleProperties.AppendLine(string.Format("PACKAGE_NAME={0}", PackageName));
diff --git a/Engine/Source/Runtime/Engine/Classes/Materials/MaterialInterface.h b/Engine/Source/Runtime/Engine/Classes/Materials/MaterialInterface.h
--- a/Engine/Source/Runtime/Engine/Classes/Materials/MaterialInterface.h
+++ b/Engine/Source/Runtime/Engine/Classes/Materials/MaterialInterface.h
@@ -83,6 +83,7 @@ struct ENGINE_API FMaterialRelevance
 	uint8 bUsesSingleLayerWaterMaterial : 1;
 	uint8 bHasVolumeMaterialDomain : 1;
 	uint8 bUsesDistanceCullFade : 1;
+	uint8 bUsesCustomDepthStencil : 1;
 
 	/** Default constructor */
 	FMaterialRelevance()
diff --git a/Engine/Source/Runtime/Engine/Private/Materials/MaterialInterface.cpp b/Engine/Source/Runtime/Engine/Private/Materials/MaterialInterface.cpp
--- a/Engine/Source/Runtime/Engine/Private/Materials/MaterialInterface.cpp
+++ b/Engine/Source/Runtime/Engine/Private/Materials/MaterialInterface.cpp
@@ -56,6 +56,7 @@ void FMaterialRelevance::SetPrimitiveViewRelevance(FPrimitiveViewRelevance& OutV
 	OutViewRelevance.bUsesSkyMaterial = bUsesSkyMaterial;
 	OutViewRelevance.bUsesSingleLayerWaterMaterial = bUsesSingleLayerWaterMaterial;
 	OutViewRelevance.bHasVolumeMaterialDomain = bHasVolumeMaterialDomain;
+	OutViewRelevance.bUsesCustomDepthStencil = bUsesCustomDepthStencil;
 }
 
 //////////////////////////////////////////////////////////////////////////
@@ -168,6 +169,7 @@ FMaterialRelevance UMaterialInterface::GetRelevance_Internal(const UMaterial* Ma
 			MaterialRelevance.bUsesSceneDepth = MaterialResource->MaterialUsesSceneDepthLookup_GameThread();
 			MaterialRelevance.bHasVolumeMaterialDomain = MaterialResource->IsVolumetricPrimitive();
 			MaterialRelevance.bUsesDistanceCullFade = MaterialResource->MaterialUsesDistanceCullFade_GameThread();
+			MaterialRelevance.bUsesCustomDepthStencil = MaterialResource->UsesCustomDepthStencil_GameThread();
 			MaterialRelevance.bUsesSkyMaterial = Material->bIsSky;
 			MaterialRelevance.bUsesSingleLayerWaterMaterial = bUsesSingleLayerWaterMaterial;
 		}
diff --git a/Engine/Source/Runtime/Engine/Private/Materials/MaterialShared.cpp b/Engine/Source/Runtime/Engine/Private/Materials/MaterialShared.cpp
--- a/Engine/Source/Runtime/Engine/Private/Materials/MaterialShared.cpp
+++ b/Engine/Source/Runtime/Engine/Private/Materials/MaterialShared.cpp
@@ -881,6 +881,11 @@ bool FMaterial::HasRuntimeVirtualTextureOutput_RenderThread() const
 	return RenderingThreadShaderMap ? RenderingThreadShaderMap->HasRuntimeVirtualTextureOutput() : false;
 }
 
+bool FMaterial::UsesCustomDepthStencil_GameThread() const
+{
+	return GameThreadShaderMap.GetReference() ? (GameThreadShaderMap->UsesSceneTexture(PPI_CustomDepth) || GameThreadShaderMap->UsesSceneTexture(PPI_CustomStencil)) : false;
+}
+
 FMaterialShaderMap* FMaterial::GetRenderingThreadShaderMap() const 
 { 
 	check(IsInParallelRenderingThread());
diff --git a/Engine/Source/Runtime/Engine/Private/SceneManagement.cpp b/Engine/Source/Runtime/Engine/Private/SceneManagement.cpp
--- a/Engine/Source/Runtime/Engine/Private/SceneManagement.cpp
+++ b/Engine/Source/Runtime/Engine/Private/SceneManagement.cpp
@@ -662,7 +662,7 @@ FMobileDirectionalLightShaderParameters::FMobileDirectionalLightShaderParameters
 	DirectionalLightShadowTexture = GWhiteTexture->TextureRHI;
 	DirectionalLightShadowSampler = TStaticSamplerState<SF_Point, AM_Clamp, AM_Clamp, AM_Clamp>::GetRHI();
 	DirectionalLightShadowSize = FVector4(EForceInit::ForceInitToZero);
-	DirectionalLightDistanceFadeMAD = FVector4(EForceInit::ForceInitToZero);
+	DirectionalLightDistanceFadeMADAndSpecularScale = FVector4(EForceInit::ForceInitToZero);
 	for (int32 i = 0; i < MAX_MOBILE_SHADOWCASCADES; ++i)
 	{
 		DirectionalLightScreenToShadow[i].SetIdentity();
diff --git a/Engine/Source/Runtime/Engine/Public/MaterialShared.h b/Engine/Source/Runtime/Engine/Public/MaterialShared.h
--- a/Engine/Source/Runtime/Engine/Public/MaterialShared.h
+++ b/Engine/Source/Runtime/Engine/Public/MaterialShared.h
@@ -1619,6 +1619,9 @@ public:
 	/** Does the material have a runtime virtual texture output node. */
 	ENGINE_API bool HasRuntimeVirtualTextureOutput_RenderThread() const;
 
+	/** Does the material use CustomDepth or CustomStencil lookup */
+	ENGINE_API bool UsesCustomDepthStencil_GameThread() const;
+
 	/** Note: This function is only intended for use in deciding whether or not shader permutations are required before material translation occurs. */
 	ENGINE_API bool MaterialMayModifyMeshPosition() const;
 
diff --git a/Engine/Source/Runtime/Engine/Public/PrimitiveViewRelevance.h b/Engine/Source/Runtime/Engine/Public/PrimitiveViewRelevance.h
--- a/Engine/Source/Runtime/Engine/Public/PrimitiveViewRelevance.h
+++ b/Engine/Source/Runtime/Engine/Public/PrimitiveViewRelevance.h
@@ -83,6 +83,8 @@ struct FPrimitiveViewRelevance
 	uint32 bUsesSingleLayerWaterMaterial : 1;
 	/** Whether the view use custom data. */
 	uint32 bUseCustomViewData : 1;
+	/** Whether the primitive has materials that read the custom depth or custom stencil */
+	uint32 bUsesCustomDepthStencil : 1;
 
 	/** 
 	 * Whether this primitive view relevance has been initialized this frame.  
diff --git a/Engine/Source/Runtime/Engine/Public/SceneView.h b/Engine/Source/Runtime/Engine/Public/SceneView.h
--- a/Engine/Source/Runtime/Engine/Public/SceneView.h
+++ b/Engine/Source/Runtime/Engine/Public/SceneView.h
@@ -539,7 +539,7 @@ BEGIN_GLOBAL_SHADER_PARAMETER_STRUCT_WITH_CONSTRUCTOR(FMobileDirectionalLightSha
 	SHADER_PARAMETER_EX(FLinearColor, DirectionalLightColor, EShaderPrecisionModifier::Half)
 	SHADER_PARAMETER_EX(FVector4, DirectionalLightDirectionAndShadowTransition, EShaderPrecisionModifier::Half)
 	SHADER_PARAMETER_EX(FVector4, DirectionalLightShadowSize, EShaderPrecisionModifier::Half)
-	SHADER_PARAMETER_EX(FVector4, DirectionalLightDistanceFadeMAD, EShaderPrecisionModifier::Half) // .zw is not used atm
+	SHADER_PARAMETER_EX(FVector4, DirectionalLightDistanceFadeMADAndSpecularScale, EShaderPrecisionModifier::Half) // .z is used for SpecularScale, .w is not used atm
 	SHADER_PARAMETER_EX(FVector4, DirectionalLightShadowDistances, EShaderPrecisionModifier::Half)
 	SHADER_PARAMETER_ARRAY(FMatrix, DirectionalLightScreenToShadow, [MAX_MOBILE_SHADOWCASCADES])
 	SHADER_PARAMETER_TEXTURE(Texture2D, DirectionalLightShadowTexture)
diff --git a/Engine/Source/Runtime/RenderCore/Private/ProfilingDebugging/RealtimeGPUProfiler.cpp b/Engine/Source/Runtime/RenderCore/Private/ProfilingDebugging/RealtimeGPUProfiler.cpp
--- a/Engine/Source/Runtime/RenderCore/Private/ProfilingDebugging/RealtimeGPUProfiler.cpp
+++ b/Engine/Source/Runtime/RenderCore/Private/ProfilingDebugging/RealtimeGPUProfiler.cpp
@@ -421,12 +421,9 @@ public:
 			const bool bTracingStatsEnabled = !!CVarGPUTracingStatsEnabled.GetValueOnRenderThread();
 			if (bTracingStatsEnabled)
 			{
-				ANSICHAR EventName[NAME_SIZE];
-				Event.GetName().GetPlainANSIString(EventName);
-
 				const uint32 GPUIndex = 0;
 				FTracingProfiler::Get()->AddGPUEvent(
-					EventName,
+					Event.GetName(),
 					Event.GetStartResultMicroseconds(),
 					Event.GetEndResultMicroseconds(),
 					GPUIndex,
diff --git a/Engine/Source/Runtime/RenderCore/Private/ProfilingDebugging/TracingProfiler.cpp b/Engine/Source/Runtime/RenderCore/Private/ProfilingDebugging/TracingProfiler.cpp
--- a/Engine/Source/Runtime/RenderCore/Private/ProfilingDebugging/TracingProfiler.cpp
+++ b/Engine/Source/Runtime/RenderCore/Private/ProfilingDebugging/TracingProfiler.cpp
@@ -12,6 +12,7 @@
 #include "Misc/Paths.h"
 #include "RHI.h"
 #include "RenderingThread.h"
+#include "UObject/NameTypes.h"
 
 static TAutoConsoleVariable<int32> GTracingProfileBufferSize(
 	TEXT("TracingProfiler.BufferSize"),
@@ -112,7 +113,7 @@ void FTracingProfiler::EndFrame()
 
 	GameThreadFrameEndCycle = FPlatformTime::Cycles64();
 
-	AddCPUEvent("GameThread",
+	AddCPUEvent(NAME_GameThread,
 		GameThreadFrameBeginCycle,
 		GameThreadFrameEndCycle,
 		GGameThreadId,
@@ -162,7 +163,7 @@ void FTracingProfiler::EndFrameRT()
 		return;
 	}
 
-	AddCPUEvent("RenderThread",
+	AddCPUEvent(NAME_RenderThread,
 		RenderThreadFrameBeginCycle,
 		RenderThreadFrameEndCycle,
 		GRenderThreadId,
@@ -344,10 +345,13 @@ void FTracingProfiler::WriteCaptureToFile()
 			ThreadOrGpuId = Event.GPU.GPUIndex;
 		}
 
+		ANSICHAR EventName[NAME_SIZE];
+		Event.Name.GetPlainANSIString(EventName);
+
 		FCStringAnsi::Snprintf(StringBuffer, StringBufferSize,
 			R"({"pid":%d, "tid":%d, "ph": "X", "name": "%s", "ts": %llu, "dur": %llu, "args":{"frame":%d}},)"
 			"\n",
-			Pid, ThreadOrGpuId, Event.Name, TimeBeginMicroseconds,
+			Pid, ThreadOrGpuId, EventName, TimeBeginMicroseconds,
 			TimeEndMicroseconds - TimeBeginMicroseconds,
 			Event.FrameNumber);
 
diff --git a/Engine/Source/Runtime/RenderCore/Public/ProfilingDebugging/TracingProfiler.h b/Engine/Source/Runtime/RenderCore/Public/ProfilingDebugging/TracingProfiler.h
--- a/Engine/Source/Runtime/RenderCore/Public/ProfilingDebugging/TracingProfiler.h
+++ b/Engine/Source/Runtime/RenderCore/Public/ProfilingDebugging/TracingProfiler.h
@@ -41,7 +41,7 @@ public:
 
 	struct FEvent
 	{
-		const char* Name;
+		FName Name;
 		uint32 FrameNumber;
 		uint32 SessionId;
 		EEventType Type;
@@ -67,7 +67,7 @@ public:
 
 	RENDERCORE_API void Init();
 
-	uint32 AddCPUEvent(const char* Name,
+	uint32 AddCPUEvent(FName Name,
 		uint64 TimestampBeginCycles,
 		uint64 TimestempEndCycles,
 		uint32 ThreadId,
@@ -84,7 +84,7 @@ public:
 		return AddEvent(Event);
 	}
 
-	uint32 AddGPUEvent(const char* Name,
+	uint32 AddGPUEvent(FName Name,
 		uint64 TimestempBeginMicroseconds,
 		uint64 TimestampEndMicroseconds,
 		uint64 GPUIndex,
diff --git a/Engine/Source/Runtime/Renderer/Private/MobileBasePass.cpp b/Engine/Source/Runtime/Renderer/Private/MobileBasePass.cpp
--- a/Engine/Source/Runtime/Renderer/Private/MobileBasePass.cpp
+++ b/Engine/Source/Runtime/Renderer/Private/MobileBasePass.cpp
@@ -302,7 +302,9 @@ ELightMapPolicyType MobileBasePass::SelectMeshLightmapPolicy(
 		else if (IsIndirectLightingCacheAllowed(FeatureLevel) /* implies bAllowStaticLighting*/
 			&& PrimitiveSceneProxy
 			// Movable objects need to get their GI from the indirect lighting cache
-			&& PrimitiveSceneProxy->IsMovable())
+			&& (PrimitiveSceneProxy->IsMovable() 
+			// Static objects with lightmap type set to ForceVolumetric also need to get GI from ILC
+			|| LightMapInteraction.GetType() == LMIT_GlobalVolume) )
 		{
 			if (bUseMovableLight)
 			{
@@ -413,7 +415,7 @@ void MobileBasePass::SetTranslucentRenderState(FMeshPassProcessorRenderState& Dr
 				DrawRenderState.SetBlendState(TStaticBlendState<CW_RGB, BO_Add, BF_One, BF_InverseSourceAlpha, BO_Add, BF_Zero, BF_InverseSourceAlpha>::GetRHI());
 			}
 			else
-				check(0);
+			check(0);
 		};
 	}
 	else
@@ -509,14 +511,14 @@ void TMobileBasePassPSPolicyParamType<FUniformLightMapPolicy>::GetShaderBindings
 			ShaderBindings.Add(ReflectionParameter, ReflectionUB);
 		}
 		
-		if (LightPositionAndInvRadiusParameter.IsBound() || SpotLightDirectionParameter.IsBound())
+		if (LightPositionAndInvRadiusParameter.IsBound() || SpotLightDirectionAndSpecularScaleParameter.IsBound())
 		{
 			// Set dynamic point lights
 			FMobileBasePassMovableLightInfo LightInfo(PrimitiveSceneProxy);
 			ShaderBindings.Add(NumDynamicPointLightsParameter, LightInfo.NumMovablePointLights);
 			ShaderBindings.Add(LightPositionAndInvRadiusParameter, LightInfo.LightPositionAndInvRadius);
 			ShaderBindings.Add(LightColorAndFalloffExponentParameter, LightInfo.LightColorAndFalloffExponent);
-			ShaderBindings.Add(SpotLightDirectionParameter, LightInfo.SpotLightDirection);
+			ShaderBindings.Add(SpotLightDirectionAndSpecularScaleParameter, LightInfo.SpotLightDirectionAndSpecularScale);
 			ShaderBindings.Add(SpotLightAnglesParameter, LightInfo.SpotLightAngles);
 		}
 	}
@@ -617,7 +619,8 @@ void FMobileBasePassMeshProcessor::Process(
 		FBaseDS,
 		TMobileBasePassPSPolicyParamType<FUniformLightMapPolicy>> BasePassShaders;
 	
-	bool bEnableSkyLight = ShadingModels.IsLit() && Scene && Scene->ShouldRenderSkylightInBasePass(BlendMode);
+	//The stationary skylight contribution has been added to the LowQuality Lightmap for StaticMeshActor on mobile, so we should skip the sky light spherical harmonic contribution for it.
+	bool bEnableSkyLight = ShadingModels.IsLit() && Scene && Scene->ShouldRenderSkylightInBasePass(BlendMode) && !(FReadOnlyCVARCache::Get().bAllowStaticLighting && PrimitiveSceneProxy && PrimitiveSceneProxy->IsStatic());
 	int32 NumMovablePointLights = MobileBasePass::CalcNumMovablePointLights(MaterialResource, PrimitiveSceneProxy);
 
 	MobileBasePass::GetShaders(
diff --git a/Engine/Source/Runtime/Renderer/Private/MobileBasePassRendering.cpp b/Engine/Source/Runtime/Renderer/Private/MobileBasePassRendering.cpp
--- a/Engine/Source/Runtime/Renderer/Private/MobileBasePassRendering.cpp
+++ b/Engine/Source/Runtime/Renderer/Private/MobileBasePassRendering.cpp
@@ -129,7 +129,8 @@ FMobileBasePassMovableLightInfo::FMobileBasePassMovableLightInfo(const FPrimitiv
 
 				LightPositionAndInvRadius[NumMovablePointLights] = FVector4(LightParameters.Position, LightParameters.InvRadius);
 				LightColorAndFalloffExponent[NumMovablePointLights] = FVector4(LightParameters.Color, LightParameters.FalloffExponent);
-				SpotLightDirection[NumMovablePointLights] = LightParameters.Direction;
+				SpotLightDirectionAndSpecularScale[NumMovablePointLights] = LightParameters.Direction;
+				SpotLightDirectionAndSpecularScale[NumMovablePointLights].Component(3) = LightProxy->GetSpecularScale();
 				SpotLightAngles[NumMovablePointLights].Set(LightParameters.SpotAngles.X, LightParameters.SpotAngles.Y, 0.f, LightType == LightType_Spot ? 1.0f : 0.0f);
 
 				if (LightType == LightType_Rect)
@@ -196,8 +197,9 @@ void SetupMobileDirectionalLightUniformParameters(
 		Params.DirectionalLightDirectionAndShadowTransition = FVector4(-Light->Proxy->GetDirection(), 0.f);
 
 		const FVector2D FadeParams = Light->Proxy->GetDirectionalLightDistanceFadeParameters(FeatureLevel, Light->IsPrecomputedLightingValid(), SceneView.MaxShadowCascades);
-		Params.DirectionalLightDistanceFadeMAD.X = FadeParams.Y;
-		Params.DirectionalLightDistanceFadeMAD.Y = -FadeParams.X * FadeParams.Y;
+		Params.DirectionalLightDistanceFadeMADAndSpecularScale.X = FadeParams.Y;
+		Params.DirectionalLightDistanceFadeMADAndSpecularScale.Y = -FadeParams.X * FadeParams.Y;
+		Params.DirectionalLightDistanceFadeMADAndSpecularScale.Z = Light->Proxy->GetSpecularScale();
 
 		if (bDynamicShadows && VisibleLightInfos.IsValidIndex(Light->Id) && VisibleLightInfos[Light->Id].AllProjectedShadows.Num() > 0)
 		{
diff --git a/Engine/Source/Runtime/Renderer/Private/MobileBasePassRendering.h b/Engine/Source/Runtime/Renderer/Private/MobileBasePassRendering.h
--- a/Engine/Source/Runtime/Renderer/Private/MobileBasePassRendering.h
+++ b/Engine/Source/Runtime/Renderer/Private/MobileBasePassRendering.h
@@ -76,7 +76,7 @@ public:
 	int32 NumMovablePointLights;
 	FVector4 LightPositionAndInvRadius[MAX_BASEPASS_DYNAMIC_POINT_LIGHTS];
 	FVector4 LightColorAndFalloffExponent[MAX_BASEPASS_DYNAMIC_POINT_LIGHTS];
-	FVector4 SpotLightDirection[MAX_BASEPASS_DYNAMIC_POINT_LIGHTS];
+	FVector4 SpotLightDirectionAndSpecularScale[MAX_BASEPASS_DYNAMIC_POINT_LIGHTS];
 	FVector4 SpotLightAngles[MAX_BASEPASS_DYNAMIC_POINT_LIGHTS];
 };
 
@@ -252,7 +252,7 @@ public:
 		LightColorAndFalloffExponentParameter.Bind(Initializer.ParameterMap, TEXT("LightColorAndFalloffExponent"));
 		NumDynamicPointLightsParameter.Bind(Initializer.ParameterMap, TEXT("NumDynamicPointLights"));
 		SpotLightAnglesParameter.Bind(Initializer.ParameterMap, TEXT("SpotLightAngles"));
-		SpotLightDirectionParameter.Bind(Initializer.ParameterMap, TEXT("SpotLightDirection"));
+		SpotLightDirectionAndSpecularScaleParameter.Bind(Initializer.ParameterMap, TEXT("SpotLightDirectionAndSpecularScale"));
 						
 		CSMDebugHintParams.Bind(Initializer.ParameterMap, TEXT("CSMDebugHint"));
 	}
@@ -280,7 +280,7 @@ public:
 		Ar << LightColorAndFalloffExponentParameter;
 		Ar << NumDynamicPointLightsParameter;
 		Ar << SpotLightAnglesParameter;
-		Ar << SpotLightDirectionParameter;
+		Ar << SpotLightDirectionAndSpecularScaleParameter;
 			
 		Ar << CSMDebugHintParams;
 
@@ -300,7 +300,7 @@ private:
 	FShaderParameter LightPositionAndInvRadiusParameter;
 	FShaderParameter LightColorAndFalloffExponentParameter;
 	FShaderParameter NumDynamicPointLightsParameter;
-	FShaderParameter SpotLightDirectionParameter;
+	FShaderParameter SpotLightDirectionAndSpecularScaleParameter;
 	FShaderParameter SpotLightAnglesParameter;
 
 	FShaderParameter CSMDebugHintParams;
diff --git a/Engine/Source/Runtime/Renderer/Private/MobileShadingRenderer.cpp b/Engine/Source/Runtime/Renderer/Private/MobileShadingRenderer.cpp
--- a/Engine/Source/Runtime/Renderer/Private/MobileShadingRenderer.cpp
+++ b/Engine/Source/Runtime/Renderer/Private/MobileShadingRenderer.cpp
@@ -111,7 +111,7 @@ static bool UsesCustomDepthStencilLookup(const FViewInfo& View)
 	}
 
 	//TODO: check if translucency uses CustomDepth
-	return bPPUsesCustomDepth || bPPUsesCustomStencil;
+	return bPPUsesCustomDepth || bPPUsesCustomStencil || View.bUsesCustomDepthStencil;
 }
 
 
@@ -282,17 +282,12 @@ void FMobileSceneRenderer::InitViews(FRHICommandListImmediate& RHICmdList)
 
 	// update buffers used in cached mesh path
 	// in case there are multiple views, these buffers will be updated before rendering each view
+	// OpaqueBasePassUniformBuffer, TranslucentBasePassUniformBuffer and DistortionPassUniformBuffer may depend on custom depth render target, so they should be updated after custom depth rendering 
 	if (Views.Num() > 0)
 	{
 		const FViewInfo& View = Views[0];
 		Scene->UniformBuffers.UpdateViewUniformBuffer(View);
-		UpdateOpaqueBasePassUniformBuffer(RHICmdList, View);
-		UpdateTranslucentBasePassUniformBuffer(RHICmdList, View);
 		UpdateDirectionalLightUniformBuffers(RHICmdList, View);
-
-		FMobileDistortionPassUniformParameters DistortionPassParameters;
-		SetupMobileDistortionPassUniformBuffer(RHICmdList, View, DistortionPassParameters);
-		Scene->UniformBuffers.MobileDistortionPassUniformBuffer.UpdateUniformBufferImmediate(DistortionPassParameters);
 	}
 	UpdateSkyReflectionUniformBuffer();
 
@@ -448,6 +443,16 @@ void FMobileSceneRenderer::Render(FRHICommandListImmediate& RHICmdList)
 			RenderCustomDepthPass(RHICmdList);
 		}
 	}
+
+	// Update OpaqueBasePassUniformBuffer, TranslucentBasePassUniformBuffer and DistortionUniformBuffer right after CustomDepthPass and before BeginRenderPass of MobileBasePass
+	{
+		UpdateOpaqueBasePassUniformBuffer(RHICmdList, View);
+		UpdateTranslucentBasePassUniformBuffer(RHICmdList, View);
+
+		FMobileDistortionPassUniformParameters DistortionPassParameters;
+		SetupMobileDistortionPassUniformBuffer(RHICmdList, View, DistortionPassParameters);
+		Scene->UniformBuffers.MobileDistortionPassUniformBuffer.UpdateUniformBufferImmediate(DistortionPassParameters);
+	}
 		
 	//
 	FRHITexture* SceneColor = nullptr;
diff --git a/Engine/Source/Runtime/Renderer/Private/SceneRendering.cpp b/Engine/Source/Runtime/Renderer/Private/SceneRendering.cpp
--- a/Engine/Source/Runtime/Renderer/Private/SceneRendering.cpp
+++ b/Engine/Source/Runtime/Renderer/Private/SceneRendering.cpp
@@ -873,6 +873,7 @@ void FViewInfo::Init()
 	bHasCustomDepthPrimitives = false;
 	bHasDistortionPrimitives = false;
 	bAllowStencilDither = false;
+	bUsesCustomDepthStencil = false;
 
 	ForwardLightingResources = nullptr;
 
diff --git a/Engine/Source/Runtime/Renderer/Private/SceneRendering.h b/Engine/Source/Runtime/Renderer/Private/SceneRendering.h
--- a/Engine/Source/Runtime/Renderer/Private/SceneRendering.h
+++ b/Engine/Source/Runtime/Renderer/Private/SceneRendering.h
@@ -1058,7 +1058,7 @@ public:
 	uint32 bTranslucentSurfaceLighting : 1;
 	/** Whether the view has any materials that read from scene depth. */
 	uint32 bUsesSceneDepth : 1;
-
+	uint32 bUsesCustomDepthStencil : 1;
 
 	/** Whether fog should only be computed on rendered opaque pixels or not. */
 	uint32 bFogOnlyOnRenderedOpaque : 1;
diff --git a/Engine/Source/Runtime/Renderer/Private/SceneSoftwareOcclusion.cpp b/Engine/Source/Runtime/Renderer/Private/SceneSoftwareOcclusion.cpp
--- a/Engine/Source/Runtime/Renderer/Private/SceneSoftwareOcclusion.cpp
+++ b/Engine/Source/Runtime/Renderer/Private/SceneSoftwareOcclusion.cpp
@@ -908,6 +908,11 @@ static int32 ApplyResults(const FScene* Scene, FViewInfo& View, const FOcclusion
 				View.PrimitiveDefinitelyUnoccludedMap[PrimitiveIndex] = true;
 			}
 		}
+		else
+		{
+			// This applies for those are not occludees, so just follow the frustum culling result.
+			View.PrimitiveDefinitelyUnoccludedMap[PrimitiveIndex] = View.PrimitiveVisibilityMap[PrimitiveIndex];
+		}
 	}
 
 	INC_DWORD_STAT_BY(STAT_SoftwareCulledPrimitives, NumOccluded);
diff --git a/Engine/Source/Runtime/Renderer/Private/SceneVisibility.cpp b/Engine/Source/Runtime/Renderer/Private/SceneVisibility.cpp
--- a/Engine/Source/Runtime/Renderer/Private/SceneVisibility.cpp
+++ b/Engine/Source/Runtime/Renderer/Private/SceneVisibility.cpp
@@ -1892,6 +1892,7 @@ struct FRelevancePacket
 	bool bUsesLightingChannels;
 	bool bTranslucentSurfaceLighting;
 	bool bUsesSceneDepth;
+	bool bUsesCustomDepthStencil;
 	bool bSceneHasSkyMaterial;
 	bool bHasSingleLayerWaterMaterial;
 
@@ -1930,6 +1931,7 @@ struct FRelevancePacket
 		, bUsesLightingChannels(false)
 		, bTranslucentSurfaceLighting(false)
 		, bUsesSceneDepth(false)
+		, bUsesCustomDepthStencil(false)
 		, bSceneHasSkyMaterial(false)
 		, bHasSingleLayerWaterMaterial(false)
 	{
@@ -2050,6 +2052,7 @@ struct FRelevancePacket
 			bUsesLightingChannels |= ViewRelevance.bUsesLightingChannels;
 			bTranslucentSurfaceLighting |= ViewRelevance.bTranslucentSurfaceLighting;
 			bUsesSceneDepth |= ViewRelevance.bUsesSceneDepth;
+			bUsesCustomDepthStencil |= ViewRelevance.bUsesCustomDepthStencil;
 			bSceneHasSkyMaterial |= ViewRelevance.bUsesSkyMaterial;
 			bHasSingleLayerWaterMaterial |= ViewRelevance.bUsesSingleLayerWaterMaterial;
 
@@ -2399,6 +2402,7 @@ struct FRelevancePacket
 		WriteView.TranslucentPrimCount.Append(TranslucentPrimCount);
 		WriteView.bHasDistortionPrimitives |= bHasDistortionPrimitives;
 		WriteView.bHasCustomDepthPrimitives |= bHasCustomDepthPrimitives;
+		WriteView.bUsesCustomDepthStencil |= bUsesCustomDepthStencil;
 		DirtyIndirectLightingCacheBufferPrimitives.AppendTo(WriteView.DirtyIndirectLightingCacheBufferPrimitives);
 
 		WriteView.MeshDecalBatches.Append(MeshDecalBatches);
