Cherry-pick 9387135: Categorises explicitly all TODOs related to denoising.

diff --git a/Engine/Shaders/Private/RayTracing/RayTracingAmbientOcclusionRGS.usf b/Engine/Shaders/Private/RayTracing/RayTracingAmbientOcclusionRGS.usf
--- a/Engine/Shaders/Private/RayTracing/RayTracingAmbientOcclusionRGS.usf
+++ b/Engine/Shaders/Private/RayTracing/RayTracingAmbientOcclusionRGS.usf
@@ -82,7 +82,7 @@ void AmbientOcclusionRGS()
 		GenerateCosineNormalRay(WorldPosition, WorldNormal, RandSequence, Ray.Origin, Ray.Direction, Ray.TMin, Ray.TMax);
 		ApplyCameraRelativeDepthBias(Ray, PixelCoord, DeviceZ, WorldNormal, MaxNormalBias);
 
-		if (dot(WorldNormal, Ray.Direction) <= 0.0) // TODO: does it needs to be handled by the denoiser?
+		if (dot(WorldNormal, Ray.Direction) <= 0.0) // TODO(Denoiser): does it needs to be handled by the denoiser?
 			continue;
 		
 		uint RayFlags = 0;
diff --git a/Engine/Shaders/Private/RayTracing/RayTracingCommon.ush b/Engine/Shaders/Private/RayTracing/RayTracingCommon.ush
--- a/Engine/Shaders/Private/RayTracing/RayTracingCommon.ush
+++ b/Engine/Shaders/Private/RayTracing/RayTracingCommon.ush
@@ -113,7 +113,7 @@ uint2 GetPixelCoord(uint2 DispatchThreadId, uint UpscaleFactor)
 {
 	uint UpscaleFactorPow2 = UpscaleFactor * UpscaleFactor;
 
-	// TODO: find a way to not interfer with TAA's jittering.
+	// TODO(Denoiser): find a way to not interfer with TAA's jittering.
 	uint SubPixelId = View.StateFrameIndex & (UpscaleFactorPow2 - 1);
 
 	return DispatchThreadId * UpscaleFactor + uint2(SubPixelId & (UpscaleFactor - 1), SubPixelId / UpscaleFactor);
diff --git a/Engine/Shaders/Private/RayTracing/RayTracingOcclusionRGS.usf b/Engine/Shaders/Private/RayTracing/RayTracingOcclusionRGS.usf
--- a/Engine/Shaders/Private/RayTracing/RayTracingOcclusionRGS.usf
+++ b/Engine/Shaders/Private/RayTracing/RayTracingOcclusionRGS.usf
@@ -333,7 +333,7 @@ void OcclusionRGS()
 			AvgHitDistance = 1.0e27;
 		}
 
-		// TODO: the denoiser would much prefer a single RG texture.
+		// TODO(Denoiser): the denoiser would much prefer a single RG texture.
 		RWOcclusionMaskUAV[PixelCoord] = Shadow;
 		RWRayDistanceUAV[PixelCoord] = AvgHitDistance;
 	}
diff --git a/Engine/Shaders/Private/RayTracing/RayTracingPrimaryRays.usf b/Engine/Shaders/Private/RayTracing/RayTracingPrimaryRays.usf
--- a/Engine/Shaders/Private/RayTracing/RayTracingPrimaryRays.usf
+++ b/Engine/Shaders/Private/RayTracing/RayTracingPrimaryRays.usf
@@ -110,7 +110,7 @@ void RayTracingPrimaryRaysRGS()
 {
 	uint2 DispatchThreadId = DispatchRaysIndex().xy + View.ViewRectMin;
 	uint2 PixelCoord = GetPixelCoord(DispatchThreadId, UpscaleFactor);
-	uint LinearIndex = PixelCoord.y * View.BufferSizeAndInvSize.x + PixelCoord.x; // TODO: PixelCoord or DispatchThreadId
+	uint LinearIndex = PixelCoord.y * View.BufferSizeAndInvSize.x + PixelCoord.x; // TODO(Denoiser): PixelCoord or DispatchThreadId
 	
 	RandomSequence RandSequence;
 	RandomSequence_Initialize(RandSequence, LinearIndex, View.StateFrameIndex);
diff --git a/Engine/Shaders/Private/RayTracing/RayTracingReflections.usf b/Engine/Shaders/Private/RayTracing/RayTracingReflections.usf
--- a/Engine/Shaders/Private/RayTracing/RayTracingReflections.usf
+++ b/Engine/Shaders/Private/RayTracing/RayTracingReflections.usf
@@ -167,7 +167,7 @@ void RayTracingReflectionsRGS()
 
 	uint2 PixelCoord = GetPixelCoord(DispatchThreadId, UpscaleFactor); // GBuffer pixel coordinates (full resolution)
 
-	uint LinearIndex = PixelCoord.y * View.BufferSizeAndInvSize.x + PixelCoord.x; // TODO: PixelCoord or DispatchThreadId
+	uint LinearIndex = PixelCoord.y * View.BufferSizeAndInvSize.x + PixelCoord.x; // TODO(Denoiser): PixelCoord or DispatchThreadId
 	
 	RandomSequence RandSequence;
 	RandomSequence_Initialize(RandSequence, LinearIndex, View.StateFrameIndex);
@@ -737,7 +737,7 @@ void RayTracingReflectionsRGS()
 		RayHitDistanceOutput[DispatchThreadId] = HitDistance;
 
 		// Imaginary depth computation for Nvidia's denoiser.
-		// TODO: shader permutation?
+		// TODO(Denoiser): shader permutation?
 		{
 			float3 OriginalWorldPos;
 			float3 CameraToPixel;
diff --git a/Engine/Shaders/Private/ScreenSpaceDenoise/SSDCommon.ush b/Engine/Shaders/Private/ScreenSpaceDenoise/SSDCommon.ush
--- a/Engine/Shaders/Private/ScreenSpaceDenoise/SSDCommon.ush
+++ b/Engine/Shaders/Private/ScreenSpaceDenoise/SSDCommon.ush
@@ -120,7 +120,7 @@ float4x4 GetScreenToViewDistanceMatrix()
 	float4x4 LocalScreenToView = ScreenToView;
 
 	// Intentionally nullout some values, so the compiler can remove some of the VALU, that is all the point of CONFIG_USE_VIEW_SPACE.
-	// TODO: off axis projection.
+	// TODO(Denoiser): off axis projection.
 	{
 		LocalScreenToView[0][1] = 0.0;
 		LocalScreenToView[0][2] = 0.0;
@@ -220,7 +220,7 @@ FSSDCompressedSceneInfos SampleCompressedSceneMetadata(
 
 	if (bPrevFrame)
 	{
-		// TODO: so what to do with mip level of these guys?
+		// TODO(Denoiser): so what to do with mip level of these guys?
 		float WorldSceneDepth = ConvertFromDeviceZ(PrevDepthBuffer.SampleLevel(GlobalPointClampedSampler, BufferUV, 0).r);
 		float3 WorldNormal = DecodeNormal(PrevGBufferA.SampleLevel(GlobalPointClampedSampler, BufferUV, 0).rgb);
 		float Roughness = PrevGBufferB.SampleLevel(GlobalPointClampedSampler, BufferUV, 0).b;
diff --git a/Engine/Shaders/Private/ScreenSpaceDenoise/SSDComposeHarmonics.usf b/Engine/Shaders/Private/ScreenSpaceDenoise/SSDComposeHarmonics.usf
--- a/Engine/Shaders/Private/ScreenSpaceDenoise/SSDComposeHarmonics.usf
+++ b/Engine/Shaders/Private/ScreenSpaceDenoise/SSDComposeHarmonics.usf
@@ -234,7 +234,7 @@ void MainCS(
 			float3 Outgoing = NormalizedSample.ColorArray[1].rgb;
 			float3 Stoped = Incoming - Outgoing;
 
-			// TODO.
+			// TODO(Denoiser).
 			float3 HarmonicMask;
 			HarmonicMask.r = Incoming.r > 0 ? saturate(Outgoing.r / Incoming.r) : 1.0;
 			HarmonicMask.g = Incoming.g > 0 ? saturate(Outgoing.g / Incoming.g) : 1.0;
diff --git a/Engine/Shaders/Private/ScreenSpaceDenoise/SSDDefinitions.ush b/Engine/Shaders/Private/ScreenSpaceDenoise/SSDDefinitions.ush
--- a/Engine/Shaders/Private/ScreenSpaceDenoise/SSDDefinitions.ush
+++ b/Engine/Shaders/Private/ScreenSpaceDenoise/SSDDefinitions.ush
@@ -17,7 +17,7 @@
 #define DEBUG_ENABLE_TILE_CLASSIFICATION 1
 
 /** Should output invalid data when know there is a need. */
-// TODO: there is an artifact in temporal reprojection.
+// TODO(Denoiser): there is an artifact in temporal reprojection.
 #define DEBUG_WASTFUL_OUTPUT 1
 
 /** Force the signal to a specific mip level. */
@@ -287,7 +287,7 @@
 //------------------------------------------------------- GLOBAL CONFIGURATION OF ALL DENOISER SHADERS
 
 /** Do normal based anysotropy. */
-// TODO: the hierarchical denoiser needs that.
+// TODO(Denoiser): the hierarchical denoiser needs that.
 #define CONFIG_NORMAL_ANISOTROPY 0
 
 /** Whether should enable VGPR optimizations in the entry point of the shaders. */
@@ -297,7 +297,7 @@
 #define CONFIG_VALU_OPTIMIZATIONS 1
 
 /** Whether should compile hint to the shader compiler to force shift some variable to SGPR. */
-// TODO: the shader compiler is not doing the right thing.
+// TODO(Denoiser): the shader compiler is not doing the right thing.
 #define CONFIG_SGPR_HINT_OPTIMIZATION 0
 
 /** Configures whether the world normal should be compressed in FSSDSampleSceneInfos. */
@@ -313,5 +313,5 @@
 #define REFLECTIONS_REJECTION_COLOR_SPACE (COLOR_SPACE_LCOCG | COLOR_SPACE_UNPREMULTIPLY)
 
 /** Defines the order of the spherical harmonic to denoise. */
-// TODO: quality permutation instead?
+// TODO(Denoiser): quality permutation instead?
 #define SPHERICAL_HARMONIC_ORDER 2
diff --git a/Engine/Shaders/Private/ScreenSpaceDenoise/SSDPublicHarmonics.ush b/Engine/Shaders/Private/ScreenSpaceDenoise/SSDPublicHarmonics.ush
--- a/Engine/Shaders/Private/ScreenSpaceDenoise/SSDPublicHarmonics.ush
+++ b/Engine/Shaders/Private/ScreenSpaceDenoise/SSDPublicHarmonics.ush
@@ -21,13 +21,13 @@
 float ComputeLightSampleWorldBluringRadius(float3 TranslatedWorldPosition, uint LightType, FLightShaderParameters Light, float HitT)
 {
 	// Sometime, artists might put occluder very close to the area light compared to they area, that may lead to negative values.
-	// TODO: the correct way to fix this is to move this world bluring radius computation into the RGS, and have DistanceFromLight = Ray's MaxT.
+	// TODO(Denoiser): the correct way to fix this is to move this world bluring radius computation into the RGS, and have DistanceFromLight = Ray's MaxT.
 	const float ClosestLightDistance = 0.001;
 
 	BRANCH
 	if (LightType == LIGHT_TYPE_DIRECTIONAL)
 	{
-		// TODO: compute this on CPU.
+		// TODO(Denoiser): compute this on CPU.
 		return tan(asin(Light.SourceRadius)) * HitT;
 	}
 	else if (LightType == LIGHT_TYPE_POINT || LightType == LIGHT_TYPE_SPOT)
@@ -50,10 +50,10 @@ float ComputeLightSampleWorldBluringRadius(float3 TranslatedWorldPosition, uint
 
 		float2 LightDimensions = GetRectLightDimensions(Light);
 
-		// TODO: area light anysotropy.
+		// TODO(Denoiser): area light anysotropy.
 		float SmallestLightDimension = min(LightDimensions.x, LightDimensions.y);
 
-		// TODO: HitT on the direction of the ray, but does not DistanceFromLight, witch is bad for large area light.
+		// TODO(Denoiser): HitT on the direction of the ray, but does not DistanceFromLight, witch is bad for large area light.
 		return SmallestLightDimension * HitT / max(DistanceFromLight - HitT, ClosestLightDistance);
 	}
 }
diff --git a/Engine/Shaders/Private/ScreenSpaceDenoise/SSDReduce.usf b/Engine/Shaders/Private/ScreenSpaceDenoise/SSDReduce.usf
--- a/Engine/Shaders/Private/ScreenSpaceDenoise/SSDReduce.usf
+++ b/Engine/Shaders/Private/ScreenSpaceDenoise/SSDReduce.usf
@@ -32,7 +32,7 @@
 	#define CONFIG_SIGNAL_OUTPUT_TEXTURE_TYPE   SIGNAL_TEXTURE_TYPE_UINT2
 
 	// Spherical harmonics are a lot of data, need to shrink VGPR pressure to improve lattency hidding when fetching the buffer.
-	// TODO: some shader compiler completly falls apparts with the current implementation of
+	// TODO(Denoiser): some shader compiler completly falls apparts with the current implementation of
 	//       CONFIG_SIGNAL_VGPR_COMPRESSION and actually drops in occupency.
 	//#define CONFIG_SIGNAL_VGPR_COMPRESSION SIGNAL_COMPRESSION_DIFFUSE_INDIRECT_HARMONIC
 	//#define CONFIG_NEIGHBOR_TO_REF_COMPUTATION NEIGHBOR_TO_REF_LOWEST_VGPR_PRESSURE
@@ -171,7 +171,7 @@ void MainCS(
 			OutputSamples.Array[MultiplexId] = RefSamples.Array[MultiplexId];
 		}
 
-		// TODO: TODO: reduce VGPR pressure here
+		// TODO(Denoiser): reduce VGPR pressure here
 		uint2 OutputPixelPostion = ViewportMax + DispatchThreadId;
 
 		OutputMultiplexedSignal(
diff --git a/Engine/Shaders/Private/ScreenSpaceDenoise/SSDSignalAccumulator.ush b/Engine/Shaders/Private/ScreenSpaceDenoise/SSDSignalAccumulator.ush
--- a/Engine/Shaders/Private/ScreenSpaceDenoise/SSDSignalAccumulator.ush
+++ b/Engine/Shaders/Private/ScreenSpaceDenoise/SSDSignalAccumulator.ush
@@ -111,7 +111,7 @@ FSSDSignalAccumulator CreateSignalAccumulator()
 	
 	#if COMPILE_MINMAX_ACCUMULATOR
 		Accumulator.Min = CreateSignalSampleFromScalarValue(INFINITE_FLOAT);
-		Accumulator.Max = CreateSignalSampleFromScalarValue(0.0); // TODO: -INFINITE_FLOAT? otherwise there is a risk to color clamp with YCoCg.
+		Accumulator.Max = CreateSignalSampleFromScalarValue(0.0); // TODO(Denoiser): -INFINITE_FLOAT? otherwise there is a risk to color clamp with YCoCg.
 	#endif
 		
 	#if COMPILE_MININVFREQ_ACCUMULATOR
@@ -227,7 +227,7 @@ void StartAccumulatingClusterInDRB(
 
 	Accumulator.BorderingRadius = RefPixelWorldBluringRadius * (ClusterInfo.OutterBoundaryRadius + ErrorCorrection);
 
-	// TODO: could this be constant.
+	// TODO(Denoiser): could this be constant.
 	Accumulator.HighestInvFrequency = Accumulator.BorderingRadius * 2;
 
 	// Reset current bucket.
@@ -264,8 +264,8 @@ void AccumulateSampleInDRB(inout FSSDSignalAccumulator Accumulator, FSSDSampleAc
 	{
 		float SampleTranslucency = saturate(ClampedInvFrequency - Accumulator.BorderingRadius);
 		
-		// TODO: do need (A.Sample.MissCount * SafeRcp(A.Sample.SampleCount)) computation onen 1spp reconstruction.
-		// TODO: could pass down a normalised version of A.Sample to avoid Rcp.
+		// TODO(Denoiser): do need (A.Sample.MissCount * SafeRcp(A.Sample.SampleCount)) computation onen 1spp reconstruction.
+		// TODO(Denoiser): could pass down a normalised version of A.Sample to avoid Rcp.
 		float R = A.Sample.MissCount * SafeRcp(A.Sample.SampleCount);
 		
 		float Translucency = lerp(R, 1, SampleTranslucency);
@@ -302,13 +302,13 @@ void DijestSampleClusterInDRB(
 	float CurrentInvFrequency = Accumulator.CurrentInvFrequency * rcp(Accumulator.Current.SampleCount);
 
 	// Whether current bucket is occluding previous bucket.
-	// TODO: put a contrast.
+	// TODO(Denoiser): put a contrast.
 	float bOccludingCluster = saturate((PreviousInvFrequency - CurrentInvFrequency) * rcp(Accumulator.BorderingRadius));
 
 	// Compute final factor to use to with previous bucket.
 	float PreviousBucketFactor = (Accumulator.Previous.SampleCount == 0.0) ? 0.0 : (1.0 - CurrentClusterOpacity * bOccludingCluster);
 
-	// TODO: the big difference in sample count could lead to wrong opacity between sample clusters.
+	// TODO(Denoiser): the big difference in sample count could lead to wrong opacity between sample clusters.
 	//if (PreviousBucketFactor < 1)
 	//	PreviousBucketFactor *= Accumulator.Current.SampleCount * SafeRcp(Accumulator.Previous.SampleCount);
 
diff --git a/Engine/Shaders/Private/ScreenSpaceDenoise/SSDSignalAccumulatorArray.ush b/Engine/Shaders/Private/ScreenSpaceDenoise/SSDSignalAccumulatorArray.ush
--- a/Engine/Shaders/Private/ScreenSpaceDenoise/SSDSignalAccumulatorArray.ush
+++ b/Engine/Shaders/Private/ScreenSpaceDenoise/SSDSignalAccumulatorArray.ush
@@ -119,7 +119,7 @@ FSSDSignalAccumulatorArray UncompressAccumulatorArray(in FSSDCompressedSignalAcc
 	#if COMPILE_DRB_ACCUMULATOR && COMPILE_MININVFREQ_ACCUMULATOR
 	else if (CompressionLayout == ACCUMULATOR_COMPRESSION_PENUMBRA_DRB)
 	{
-		// TODO: CONFIG_SIGNAL_BATCH_SIZE?
+		// TODO(Denoiser): CONFIG_SIGNAL_BATCH_SIZE?
 		UNROLL_N(MAX_SIGNAL_BATCH_SIZE)
 		for (uint SignalBatchId = 0; SignalBatchId < MAX_SIGNAL_BATCH_SIZE; SignalBatchId++)
 		{
diff --git a/Engine/Shaders/Private/ScreenSpaceDenoise/SSDSignalBufferEncoding.ush b/Engine/Shaders/Private/ScreenSpaceDenoise/SSDSignalBufferEncoding.ush
--- a/Engine/Shaders/Private/ScreenSpaceDenoise/SSDSignalBufferEncoding.ush
+++ b/Engine/Shaders/Private/ScreenSpaceDenoise/SSDSignalBufferEncoding.ush
@@ -103,7 +103,7 @@ void DecodeMultiplexedSignalsFromFloat4(
 		for (uint BatchSignalId = 0; BatchSignalId < MAX_SIGNAL_BATCH_SIZE; BatchSignalId++)
 		{
 			uint MultiplexId = BatchSignalId;
-			// TODO: feed the actual number of sample.
+			// TODO(Denoiser): feed the actual number of sample.
 			OutSamples.Array[MultiplexId].SampleCount = (RawSample[MultiplexId].g == DENOISER_INVALID_HIT_DISTANCE ? 0.0 : 1.0);
 			OutSamples.Array[MultiplexId].MissCount = OutSamples.Array[MultiplexId].SampleCount * RawSample[MultiplexId].r;
 			OutSamples.Array[MultiplexId].ClosestHitDistance = OutSamples.Array[MultiplexId].SampleCount * (RawSample[MultiplexId].g == DENOISER_MISS_HIT_DISTANCE ? 1000000 : RawSample[MultiplexId].g);
@@ -117,7 +117,7 @@ void DecodeMultiplexedSignalsFromFloat4(
 			uint MultiplexId = BatchSignalId;
 			float2 Channels = float2(RawSample[MultiplexId / 2][2 * (MultiplexId % 2)], RawSample[MultiplexId / 2][2 * (MultiplexId % 2) + 1]);
 	
-			// TODO: feed the actual number of sample.
+			// TODO(Denoiser): feed the actual number of sample.
 			OutSamples.Array[MultiplexId].SampleCount = (Channels.g == -2.0 ? 0.0 : 1.0);
 			OutSamples.Array[MultiplexId].MissCount = OutSamples.Array[MultiplexId].SampleCount * Channels.r;
 			OutSamples.Array[MultiplexId].WorldBluringRadius = OutSamples.Array[MultiplexId].SampleCount * (Channels.g == -1 ? WORLD_RADIUS_MISS : Channels.g);
@@ -459,7 +459,7 @@ FSSDSignalArray DecodeMultiplexedSignals(
 			/* out */ MultiplexedSamples.Array[0].MissCount,
 			/* out */ MultiplexedSamples.Array[0].ColorSH);
 
-		// TODO: SRV on RG F32 and manually decode with f16tof32(), going to reduce VGPR pressure between fetch and result.
+		// TODO(Denoiser): SRV on RG F32 and manually decode with f16tof32(), going to reduce VGPR pressure between fetch and result.
 	}
 	#endif
 	// Uncompress assuming this is packed as 16bit floats.
@@ -780,7 +780,7 @@ FSSDCompressedMultiplexedSample SampleCompressedMultiplexedSignals(
 		}
 		#elif CONFIG_SIGNAL_INPUT_TEXTURE_TYPE >= SIGNAL_TEXTURE_TYPE_UINT1 && CONFIG_SIGNAL_INPUT_TEXTURE_TYPE <= SIGNAL_TEXTURE_TYPE_UINT4
 		{
-			// TODO: Exposed the int3 instead as function parameter.
+			// TODO(Denoiser): Exposed the int3 instead as function parameter.
 			int3 Coord = int3(PixelCoord >> uint(Level), int(Level));
 			CompressedSample.VGPRArray[0] = SignalBuffer0.Load(Coord);
 			CompressedSample.VGPRArray[1] = SignalBuffer1.Load(Coord);
diff --git a/Engine/Shaders/Private/ScreenSpaceDenoise/SSDSignalFramework.ush b/Engine/Shaders/Private/ScreenSpaceDenoise/SSDSignalFramework.ush
--- a/Engine/Shaders/Private/ScreenSpaceDenoise/SSDSignalFramework.ush
+++ b/Engine/Shaders/Private/ScreenSpaceDenoise/SSDSignalFramework.ush
@@ -136,7 +136,7 @@ float Luma_To_LumaLog(float x)
 /** Returns the exposure multiplier that should be used for the color transformation of the signal. */
 float GetSignalColorTransformationExposure()
 {
-	// TODO: expose that guy.
+	// TODO(Denoiser): expose that guy.
 	const float Tweak = 1;
 
 	float FrameExposureScale = EyeAdaptationLookup();
@@ -152,10 +152,10 @@ float GetSignalColorTransformationExposure()
 float GetSignalLuma4(FSSDSignalSample Sample, const uint Basis)
 #if COMPILE_SIGNAL_COLOR
 {
-	// TODO: should force this to be scalar load.
+	// TODO(Denoiser): should force this to be scalar load.
 	float FrameExposureScale = GetSignalColorTransformationExposure();
 		
-	// TODO: exposes optimisation if sample is normalized.
+	// TODO(Denoiser): exposes optimisation if sample is normalized.
 	const bool bIsNormalizedSample = false;
 
 	if (Basis & COLOR_SPACE_KARIS_WEIGHTING)
@@ -199,10 +199,10 @@ FSSDSignalSample TransformSignal(FSSDSignalSample Sample, const uint SrcBasis, c
 		return Sample;
 	}
 
-	// TODO: should force this to be scalar load.
+	// TODO(Denoiser): should force this to be scalar load.
 	float FrameExposureScale = GetSignalColorTransformationExposure();
 		
-	// TODO: exposes optimisation if sample is normalized.
+	// TODO(Denoiser): exposes optimisation if sample is normalized.
 	const bool bIsNormalizedSample = false;
 	const bool bDebugForceNormalizeColor = true;
 	const bool bIsNormalizedColor = bIsNormalizedSample || bDebugForceNormalizeColor;
@@ -346,7 +346,7 @@ FSphericalGaussian ComputeRoughnessLobe(FSSDSampleSceneInfos RefSceneMetadata)
 {
 	#if CONFIG_USE_VIEW_SPACE
 		float3 N = RefSceneMetadata.ViewNormal;
-		float3 V = float3(1, 0, 0); // TODO.
+		float3 V = float3(1, 0, 0); // TODO(Denoiser).
 	#else
 		float3 N = RefSceneMetadata.WorldNormal;
 		float3 V = normalize(View.TranslatedWorldCameraOrigin - GetTranslatedWorldPosition(RefSceneMetadata));
@@ -401,12 +401,12 @@ void ProjectSpecularLobeToScreenSpace(
 		float3x3 TangentBasis = GetTangentBasis(N);
 		float3 TangentV = mul(TangentBasis, V);
 		
-		// TODO: optimise using lobe symetry along anisotropy major axis.
+		// TODO(Denoiser): optimise using lobe symetry along anisotropy major axis.
 		float3 TangentL[4];
 		UNROLL_N(4)
 		for (uint i = 0; i < 4; i++)
 		{
-			// TODO: optimise VALU here.
+			// TODO(Denoiser): optimise VALU here.
 			float3 TangentH = ImportanceSampleVisibleGGX(ArrayE[i], a2, TangentV).xyz;
 			TangentL[i] = 2 * dot(TangentV, TangentH) * TangentH - TangentV;
 		}
@@ -442,7 +442,7 @@ float ComputeBluringRadiusOfNormalizedSample(FSSDSignalSample Sample, FSSDSample
 	return ComputeInfinityBluringRadius(SceneMetadata) * ComputeBluringOutOfFocusOfNormalizedSample(Sample, SceneMetadata);
 }
 
-// TODO: duplicated...
+// TODO(Denoiser): duplicated...
 float3 GenerateReflectedRayDirection(
 	float3 IncidentDirection,
 	float3 WorldNormal,
@@ -479,7 +479,7 @@ float3 GenerateReflectedRayDirection(
 	return RayDirection;
 }
 
-// TODO: duplicated...
+// TODO(Denoiser): duplicated...
 void GetInputRayDirectionPixel(
 	uint2 PixelCoord,
 	FSSDSampleSceneInfos SceneMetadata,
@@ -561,7 +561,7 @@ float ComputeBilateralWeight(
 	}
 	else if ((BilateralSettings & 0xF) == 2)
 	{
-		// TODO: save 2 VGPR by using SceneDepth instead of RefSceneMetadata.TranslatedWorldPosition (not sure it can work in math yet).
+		// TODO(Denoiser): save 2 VGPR by using SceneDepth instead of RefSceneMetadata.TranslatedWorldPosition (not sure it can work in math yet).
 		// float2 ClipPosition = ScreenPosition * (View.ViewToClip[3][3] < 1.0f ? WorldDepth : 1.0f);
 		// Infos.TranslatedWorldPosition = mul(float4(ClipPosition, WorldDepth, 1), ScreenToTranslatedWorld).xyz;
 
@@ -578,14 +578,14 @@ float ComputeBilateralWeight(
 	}
 	else if ((BilateralSettings & 0xF) == 4)
 	{
-		// TODO: extremely expensive with rcp()...
+		// TODO(Denoiser): extremely expensive with rcp()...
 		PositionBilateralWeight = saturate(1 - dot(NeighborToRefVector, NeighborToRefVector) * rcp(MaxWorldBluringDistance * MaxWorldBluringDistance));
 	}
 	else if ((BilateralSettings & 0xF) == 5)
 	{
 		float DistFromRefPlane = dot(RefNormal, NeighborToRefVector);
 
-		// TODO: extremely expensive with rcp()...
+		// TODO(Denoiser): extremely expensive with rcp()...
 		PositionBilateralWeight = (
 			saturate(1 - dot(NeighborToRefVector, NeighborToRefVector) * rcp(MaxWorldBluringDistance * MaxWorldBluringDistance)) *
 			saturate(1 - 6 * abs(DistFromRefPlane) * rcp(MaxWorldBluringDistance)));
@@ -596,11 +596,11 @@ float ComputeBilateralWeight(
 	{
 		float NoN = max(dot(RefNormal, NeighborNormal), 0);
 
-		#if 1 // TODO
+		#if 1 // TODO(Denoiser)
 			NormalBilateralWeight = pow(NoN, 4);
 		#else
 			// [ Kontkanen et al 2004, "Irradiance Filtering for Monte Carlo Ray Tracing" ]
-			// TODO: looks broken
+			// TODO(Denoiser): looks broken
 			float m = MaxWorldBluringDistance;
 			float a = 0.3;
 			float d = sqrt(1 - NoN); // + length(RefSceneMetadata.TranslatedWorldPosition - NeighborSceneMetadata.TranslatedWorldPosition) * rcp(m);
@@ -647,7 +647,7 @@ float GetSignalWorldBluringRadius(FSSDSignalSample Sample, FSSDSampleSceneInfos
 	FSSDSignalDomainKnowledge DomainKnowledge = GetSignalDomainKnowledge(BatchedSignalId);
 
 	// Sometime, artists might put occluder very close to the area light compared to they area, that may lead to negative values.
-	// TODO: the correct way to fix this is to move this world bluring radius computation into the RGS, and have DistanceFromLight = Ray's MaxT.
+	// TODO(Denoiser): the correct way to fix this is to move this world bluring radius computation into the RGS, and have DistanceFromLight = Ray's MaxT.
 	const float ClosestLightDistance = 0.001;
 
 	BRANCH
@@ -701,7 +701,7 @@ float GetSamplePenumbraSafe(FSSDSignalSample Sample)
 // Get penumbra, but be aware of NaN.
 float GetSamplePenumbra(FSSDSignalSample Sample)
 {
-	if (1) // TODO: For debugging purpose of division by 0.
+	if (1) // TODO(Denoiser): For debugging purpose of division by 0.
 	{
 		return GetSamplePenumbraSafe(Sample);
 	}
@@ -741,7 +741,7 @@ float GetRatioEstimatorWeight(
 
 	float a2 = Pow4(RefSceneMetadata.Roughness);
 
-	// TODO: Improve quantization in CONFIG_COMPRESS_WORLD_NORMAL and use it.
+	// TODO(Denoiser): Improve quantization in CONFIG_COMPRESS_WORLD_NORMAL and use it.
 	float3 N = GetWorldNormal(RefSceneMetadata);
 
 	float NoH;
@@ -793,7 +793,7 @@ float GetRatioEstimatorWeight(
 
 
 // Matches GenerateCosineNormalRay()
-// TODO: expose a Weight that multiply directly onto the Sample.SceneColor to reduce ALU
+// TODO(Denoiser): expose a Weight that multiply directly onto the Sample.SceneColor to reduce ALU
 #if COMPILE_SIGNAL_COLOR_SH && COMPILE_SIGNAL_COLOR
 FSSDSphericalHarmonicRGB ComputeSampleColorSH(FSSDSampleSceneInfos SampleSceneMetadata, FSSDSignalSample Sample, uint2 NeighborPixelCoord)
 {
diff --git a/Engine/Shaders/Private/ScreenSpaceDenoise/SSDSpatialAccumulation.usf b/Engine/Shaders/Private/ScreenSpaceDenoise/SSDSpatialAccumulation.usf
--- a/Engine/Shaders/Private/ScreenSpaceDenoise/SSDSpatialAccumulation.usf
+++ b/Engine/Shaders/Private/ScreenSpaceDenoise/SSDSpatialAccumulation.usf
@@ -177,7 +177,7 @@
 		//#define CONFIG_SAMPLE_COUNT_POLICY  SAMPLE_COUNT_POLICY_DISABLED
 
 		// Do color accumulation with karis weighting to avoid flickering specular highlight to show up the kernel pattern.
-		// TODO: This is a bit agressive.
+		// TODO(Denoiser): This is a bit agressive.
 		#define CONFIG_ACCUMULATION_COLOR_SPACE (COLOR_SPACE_RGB | COLOR_SPACE_KARIS_WEIGHTING)
 
 		#define CONFIG_INPUT_TEXTURE_COUNT 2
@@ -316,7 +316,7 @@
 	#define CONFIG_SIGNAL_OUTPUT_TEXTURE_TYPE   SIGNAL_TEXTURE_TYPE_UINT2
 
 	// Spherical harmonics are a lot of data, need to shrink VGPR pressure to improve lattency hidding when fetching the buffer.
-	// TODO: some shader compiler completly falls apparts with the current implementation of
+	// TODO(Denoiser): some shader compiler completly falls apparts with the current implementation of
 	//       CONFIG_SIGNAL_VGPR_COMPRESSION and actually drops in occupency.
 	#define CONFIG_SIGNAL_VGPR_COMPRESSION SIGNAL_COMPRESSION_DIFFUSE_INDIRECT_HARMONIC
 
@@ -585,12 +585,12 @@ FSSDRWTexture2D	SignalOutput_UAVs_3;
 
 //------------------------------------------------------- FUNCTIONS
 
-// TODO: duplicated with reflection code.
+// TODO(Denoiser): duplicated with reflection code.
 uint2 GetPixelCoord(uint2 DispatchThreadId)
 {
 	uint UpscaleFactorPow2 = UpscaleFactor * UpscaleFactor;
 
-	// TODO: find a way to not interfer with TAA's jittering.
+	// TODO(Denoiser): find a way to not interfer with TAA's jittering.
 	uint SubPixelId = View.StateFrameIndex & (UpscaleFactorPow2 - 1);
 
 	return DispatchThreadId * UpscaleFactor + uint2(SubPixelId & (UpscaleFactor - 1), SubPixelId / UpscaleFactor);
@@ -826,7 +826,7 @@ void MainCS(
 		{
 			KernelConfig.RingCount = 1;
 
-			// TODO: could be improved.
+			// TODO(Denoiser): could be improved.
 			//KernelConfig.bMinSamplePairInvFrequency = true;
 
 			float2 E = float2(
@@ -864,7 +864,7 @@ void MainCS(
 				
 				if (KernelConfig.BilateralDistanceComputation == SIGNAL_WORLD_FREQUENCY_PRECOMPUTED_BLURING_RADIUS)
 				{
-					// TODO: this is ineficient, could fetch the normalised WorldBluringRadius instead of SafeRcp().
+					// TODO(Denoiser): this is ineficient, could fetch the normalised WorldBluringRadius instead of SafeRcp().
 					SampleInfos.InvFrequency = CenterSample.WorldBluringRadius * SafeRcp(CenterSample.SampleCount);
 				}
 
@@ -880,7 +880,7 @@ void MainCS(
 
 			#if CONFIG_UPSCALE
 			{
-				// TODO: could be optimised, but currently reusing same peace of code as reflection for maintainability.
+				// TODO(Denoiser): could be optimised, but currently reusing same peace of code as reflection for maintainability.
 				uint2 RayDispatchThreadId = (DispatchThreadId - UpscaleFactor / 2) / UpscaleFactor;
 				uint2 ClosestRayPixelCoord = GetPixelCoord(RayDispatchThreadId);
 			
@@ -910,7 +910,7 @@ void MainCS(
 			#if CONFIG_VGPR_OPTIMIZATION
 				// Keek sample SampleTrackId & SceneBufferUV arround for computation of pixel output coordinate.
 				// Should be VGPR free given it's curernt is being used in accumulation has well that is highest VGPR pressure of the shader.
-				// TODO: could save 1 VGPR by using 2 SGPR instead of SampleTrackId.
+				// TODO(Denoiser): could save 1 VGPR by using 2 SGPR instead of SampleTrackId.
 				SampleTrackId = KernelConfig.SampleTrackId;
 				KernelBufferUV = KernelConfig.BufferUV;
 			#endif
@@ -930,7 +930,7 @@ void MainCS(
 			#if CONFIG_VGPR_OPTIMIZATION
 				// Keek sample SampleTrackId & SceneBufferUV arround for computation of pixel output coordinate.
 				// Should be VGPR free given it's curernt is being used in accumulation has well that is highest VGPR pressure of the shader.
-				// TODO: could save 1 VGPR by using 2 SGPR instead of SampleTrackId.
+				// TODO(Denoiser): could save 1 VGPR by using 2 SGPR instead of SampleTrackId.
 				SampleTrackId = KernelConfig.SampleTrackId;
 				KernelBufferUV = KernelConfig.BufferUV;
 			#endif
@@ -1007,7 +1007,7 @@ void MainCS(
 		if (!KernelConfig.bSampleKernelCenter && KernelConfig.bDescOrder)
 		{
 			// Remove any jitter the kernel may have. Won't have ant VGPR cost when no jittering, because KernelConfig.BufferUV == SceneBufferUV.
-			// TODO: This is costly for VGPR pressure if using KernelConfig.BufferUV was != SceneBufferUV.
+			// TODO(Denoiser): This is costly for VGPR pressure if using KernelConfig.BufferUV was != SceneBufferUV.
 			KernelConfig.BufferUV = SceneBufferUV;
 
 			SampleAndAccumulateCenterSampleAsItsOwnCluster(
@@ -1069,7 +1069,7 @@ void MainCS(
 			}
 			#endif // CONFIG_ACCUMULATION_COLOR_SPACE != CONFIG_OUTPUT_COLOR_SPACE
 
-			// TODO: it might be better to just uncompress before this for loop and remain uncompressed,
+			// TODO(Denoiser): it might be better to just uncompress before this for loop and remain uncompressed,
 			// so the color operation get done in practice during the output sample transcoding.
 			CompressSignalAccumulator(/* inout */ SignalAccumulators.Array[MultiplexId]);
 		}
@@ -1092,7 +1092,7 @@ void MainCS(
 				OutputSamples.Array[MultiplexId] = SignalAccumulators.Array[MultiplexId].Moment1;
 				
 				// Output the minimal inverse frequency as new world bluring radius for subsequent passes.
-				// TODO: store this in its own signal to avoid *= N; *= rcp(N);
+				// TODO(Denoiser): store this in its own signal to avoid *= N; *= rcp(N);
 				OutputSamples.Array[MultiplexId].WorldBluringRadius = 
 					OutputSamples.Array[MultiplexId].SampleCount * SignalAccumulators.Array[MultiplexId].MinInvFrequency;
 			}
@@ -1122,7 +1122,7 @@ void MainCS(
 				OutputSamples.Array[MultiplexId] = SignalAccumulators.Array[MultiplexId].Previous;
 
 				// Output the minimal inverse frequency as new world bluring radius for subsequent passes.
-				// TODO: store this in its own signal to avoid *= N; *= rcp(N);
+				// TODO(Denoiser): store this in its own signal to avoid *= N; *= rcp(N);
 				OutputSamples.Array[MultiplexId].WorldBluringRadius = 
 					OutputSamples.Array[MultiplexId].SampleCount * SignalAccumulators.Array[MultiplexId].MinInvFrequency;
 			
@@ -1152,9 +1152,9 @@ void MainCS(
 	}
 	#endif // DIM_STAGE == STAGE_POST_FILTERING
 
-	// TODO:  LeaveRayCount = (LeaveRayCount - 1) * 9 / (9 - 2) post processing to reject when for history rejection
+	// TODO(Denoiser):  LeaveRayCount = (LeaveRayCount - 1) * 9 / (9 - 2) post processing to reject when for history rejection
 	uint2 OutputPixelPostion;
-	#if CONFIG_VGPR_OPTIMIZATION && 1 // TODO
+	#if CONFIG_VGPR_OPTIMIZATION && 1 // TODO(Denoiser)
 	{
 		// No need to keep DispatchThreadId, can recompute the output pixel position based on information stored in VGPRs for spatial kernel.
 		#if CONFIG_SAMPLE_SET == SAMPLE_SET_STACKOWIAK_4_SETS || CONFIG_SAMPLE_SET == SAMPLE_SET_STOCASTIC_HIERARCHY
diff --git a/Engine/Shaders/Private/ScreenSpaceDenoise/SSDSpatialKernel.ush b/Engine/Shaders/Private/ScreenSpaceDenoise/SSDSpatialKernel.ush
--- a/Engine/Shaders/Private/ScreenSpaceDenoise/SSDSpatialKernel.ush
+++ b/Engine/Shaders/Private/ScreenSpaceDenoise/SSDSpatialKernel.ush
@@ -393,7 +393,7 @@ float2 ComputeRefBufferUV(FSSDKernelConfig KernelConfig)
 		#endif
 
 		// Matches first line of kStackowiakSampleSet0
-		// TODO: could be optimised further by just setting sign bit on 0.5.
+		// TODO(Denoiser): could be optimised further by just setting sign bit on 0.5.
 		float2 SampleOffset = float2(
 			SampleTrackId & 0x1 ? 0.5 : -0.5,
 			SampleTrackId & 0x2 ? 0.5 : -0.5);
@@ -479,7 +479,7 @@ FSSDSignalSample TransformSignalSampleForAccumulation(
 	uint2 SamplePixelCoord)
 {
 	// Transform the color space.
-	// TODO: could pass down information that this sample may be normalized.
+	// TODO(Denoiser): could pass down information that this sample may be normalized.
 	Sample = TransformSignal(
 		Sample, 
 		/* SrcBasis  = */ KernelConfig.BufferColorSpace[MultiplexId], 
@@ -543,7 +543,7 @@ FSSDSignalArray UncompressMultiplexedSignals(
 	float2 SampleBufferUV,
 	FSSDCompressedMultiplexedSample CompressedMultiplexedSamples)
 {
-	// TODO: offer multiplier to apply to each signal during Decode, to save mul VALU.
+	// TODO(Denoiser): offer multiplier to apply to each signal during Decode, to save mul VALU.
 	FSSDSignalArray MultiplexedSamples = DecodeMultiplexedSignals(
 		KernelConfig.BufferLayout,
 		/* MultiplexedSampleId = */ 0,
@@ -632,7 +632,7 @@ void AccumulateSampledMultiplexedSignals(
 		// Compute at compile time the id of the signal being processed.
 		const uint BatchedSignalId = ComputeSignalBatchIdFromSignalMultiplexId(KernelConfig, SignalMultiplexId);
 		
-		// TODO: direction of the ray should be cached by injest or output by RGS, otherwise ends up with VGPR pressure because of SampleBufferUV.
+		// TODO(Denoiser): direction of the ray should be cached by injest or output by RGS, otherwise ends up with VGPR pressure because of SampleBufferUV.
 		uint2 NeighborPixelCoord = floor(SampleBufferUV * KernelConfig.BufferSizeAndInvSize.xy);
 
 		// Fetch and pre process the sample for accumulation.
@@ -648,7 +648,7 @@ void AccumulateSampledMultiplexedSignals(
 		// But the signal's bluring radius might already be pre computed.
 		if (KernelConfig.BilateralDistanceComputation == SIGNAL_WORLD_FREQUENCY_PRECOMPUTED_BLURING_RADIUS)
 		{
-			// TODO: this is ineficient, could fetch the normalised WorldBluringRadius instead of SafeRcp().
+			// TODO(Denoiser): this is ineficient, could fetch the normalised WorldBluringRadius instead of SafeRcp().
 			SignalConvolutionBluringRadius = Sample.WorldBluringRadius * SafeRcp(Sample.SampleCount);
 		}
 
@@ -709,7 +709,7 @@ void AccumulateSampledMultiplexedSignals(
 			SampleInfos.FinalWeight = 1;
 		}
 
-		// TODO: bIsOutsideFrustum could afect number of samples for DRB.
+		// TODO(Denoiser): bIsOutsideFrustum could afect number of samples for DRB.
 		FLATTEN
 		if (SampleInfos.Sample.SampleCount != 0 && !bIsOutsideFrustum)
 		{
@@ -846,7 +846,7 @@ void SampleAndAccumulateMultiplexedSignalsPair(
 		}
 
 		// Take the min inverse frequency per signal if desired.
-		// TODO: this should be normalized in theory...
+		// TODO(Denoiser): this should be normalized in theory...
 		if (KernelConfig.bMinSamplePairInvFrequency)
 		{
 			UNROLL_N(MAX_SIGNAL_MULTIPLEXING)
@@ -1023,8 +1023,8 @@ void AccumulateBilinear(
 	{
 		float2 SampleOffset = BilinearSamplingOffsets2x2[i];
 		
-		// TODO: could be more ALU efficient for this.
-		// TODO: -0.5 full res pixel to ensure always select the mip, regardless of mantissa precision?
+		// TODO(Denoiser): could be more ALU efficient for this.
+		// TODO(Denoiser): -0.5 full res pixel to ensure always select the mip, regardless of mantissa precision?
 		float2 SampleBufferUV = (BilinearInfos.TopLeftPixelCoord + (SampleOffset + 0.5)) * MipLevelPow2 * KernelConfig.BufferSizeAndInvSize.zw;
 		
 		float BilinearWeight = GetSampleWeight(BilinearInfos, i);
@@ -1074,8 +1074,8 @@ void AccumulateStocasticBilinear(
 			break;
 	}
 	
-	// TODO: could be more ALU efficient for this.
-	// TODO: -0.5 full res pixel to ensure always select the mip, regardless of mantissa precision?
+	// TODO(Denoiser): could be more ALU efficient for this.
+	// TODO(Denoiser): -0.5 full res pixel to ensure always select the mip, regardless of mantissa precision?
 	float2 SampleBufferUV = (BilinearInfos.TopLeftPixelCoord + (SampleOffset + 0.5)) * MipLevelPow2 * KernelConfig.BufferSizeAndInvSize.zw;
 		
 	SampleAndAccumulateMultiplexedSignals(
@@ -1153,7 +1153,7 @@ void AccumulateSquareKernel(
 	}
 	else
 	{
-		// TODO: lattency hiding of this is terrible.
+		// TODO(Denoiser): lattency hiding of this is terrible.
 		LOOP for (int x = -KernelRadius; x <= KernelRadius; x++)
 		LOOP for (int y = -KernelRadius; y <= KernelRadius; y++)
 		{
@@ -1285,7 +1285,7 @@ void AccumulateSquare3x3Kernel(
 	{
 		float2 SampleBufferUV = KernelConfig.BufferUV;
 			
-		// TODO:
+		// TODO(Denoiser):
 		const bool bIsOutsideFrustum = false;
 
 		FSSDCompressedSceneInfos CompressedSampleSceneMetadata;
@@ -1609,7 +1609,7 @@ void ConvolveStackowiakKernel(
 		}
 		else
 		{
-			// TODO
+			// TODO(Denoiser)
 		}
 		
 		FSSDSampleClusterInfo ClusterInfo;
@@ -1753,7 +1753,7 @@ void ConvolveStackowiakKernel(
 		const uint BatchCountCount = (KernelConfig.SampleCount - 1) / kSamplingBatchSize;
 
 		// Processes the samples in batches so that the compiler can do lattency hidding.
-		// TODO: kSamplingBatchSize for lattency hidding
+		// TODO(Denoiser): kSamplingBatchSize for lattency hidding
 		LOOP
 		for (uint BatchId = 0; BatchId < BatchCountCount; BatchId++)
 		{
@@ -1802,7 +1802,7 @@ void ConvolveStackowiakKernel(
 // Returns the position of the sample on the unit circle (radius = 1) for a given ring.
 float2 GetDiskSampleOnUnitCirle(uint RingId, uint RingSampleIteration, uint RingSampleId)
 {
-	RingId -= 1; // TODO.
+	RingId -= 1; // TODO(Denoiser).
 
 	float SampleRingPos = RingSampleId;
 
@@ -1893,7 +1893,7 @@ void GatherRingSamples(
 	const uint RingSamplePairCount = GetRingSamplingPairCount(KernelConfig.SampleSet, RingId);
 
 	// Number of sample pair to process per batch.
-	// TODO: Could potentially do 4 using symetries? Might be unpracticable because of VGPR pressure. 
+	// TODO(Denoiser): Could potentially do 4 using symetries? Might be unpracticable because of VGPR pressure. 
 	const uint SamplePairBatchSize = (KernelConfig.SampleSet == SAMPLE_SET_HEXAWEB) ? 1 : 2;
 
 	// Number of batch to process.
@@ -2154,7 +2154,7 @@ void ConvolveDirectionalRect(
 	const uint kSamplingBatchSize = 2;
 
 	// Number of batch of samples to perform. It's not round up because also sampling the center of kernel anyway.
-	// TODO: store in a SGPR array instead to save 1 VGPR.
+	// TODO(Denoiser): store in a SGPR array instead to save 1 VGPR.
 	const uint BatchCountCount = KernelConfig.SampleCount / kSamplingBatchSize;
 
 	// Processes the samples in batches so that the compiler can do lattency hidding.
diff --git a/Engine/Shaders/Private/ScreenSpaceDenoise/SSDTemporalAccumulation.usf b/Engine/Shaders/Private/ScreenSpaceDenoise/SSDTemporalAccumulation.usf
--- a/Engine/Shaders/Private/ScreenSpaceDenoise/SSDTemporalAccumulation.usf
+++ b/Engine/Shaders/Private/ScreenSpaceDenoise/SSDTemporalAccumulation.usf
@@ -46,8 +46,8 @@
 	#define CONFIG_HISTORY_REJECTION HISTORY_REJECTION_VAR_BOUNDARIES
 
 	// Uses nearest to not leak informations on geometric edges by bilateral. This is OK because just blury greyscale details.
-	// TODO: a bit hacky, need find better solution.
-	// TODO: cause a regression on spot light for some reasons.
+	// TODO(Denoiser): a bit hacky, need find better solution.
+	// TODO(Denoiser): cause a regression on spot light for some reasons.
 	//#define CONFIG_USE_NEAREST_HISTORY 1
 
 	// Use pre transformed rejection buffer that contains pre transformed momment 1 & 2.
@@ -488,7 +488,7 @@ void TemporallyAccumulate(
 
 			TargetedSampleCount = PI * SignalPixelBluringRadius * SignalPixelBluringRadius;
 
-			// TODO: try to make this better? * 0.125 to avoid banding caused by floating point precision.
+			// TODO(Denoiser): try to make this better? * 0.125 to avoid banding caused by floating point precision.
 			TargetedSampleCount = clamp(TargetedSampleCount, 1, TARGETED_SAMPLE_COUNT);
 		}
 		#else
@@ -505,7 +505,7 @@ void TemporallyAccumulate(
 		float2 ClampedHistoryBufferUV = clamp(HistoryBufferUV, BufferBilinearUVMinMax.xy, BufferBilinearUVMinMax.zw);
 		bool bIsPreviousFrameOffscreen = any(HistoryBufferUV != ClampedHistoryBufferUV);
 
-		// TODO: global camera cut to save perf that is going to be required for panic post filtering.
+		// TODO(Denoiser): global camera cut to save perf that is going to be required for panic post filtering.
 		BRANCH
 		if (!bIsPreviousFrameOffscreen)
 		{
@@ -638,8 +638,8 @@ void TemporallyAccumulate(
 			KernelConfig.BilateralDistanceComputation = SIGNAL_WORLD_FREQUENCY_REF_METADATA_ONLY;
 			KernelConfig.NeighborToRefComputation = NEIGHBOR_TO_REF_LOWEST_VGPR_PRESSURE;
 
-			// TODO: Should be done for all signals.
-			// TODO: Should also take into account the world bluring radius du to size of the pixel.
+			// TODO(Denoiser): Should be done for all signals.
+			// TODO(Denoiser): Should also take into account the world bluring radius du to size of the pixel.
 			if (CONFIG_SIGNAL_PROCESSING == SIGNAL_PROCESSING_SHADOW_VISIBILITY_MASK)
 				KernelConfig.BilateralDistanceComputation = SIGNAL_WORLD_FREQUENCY_PRECOMPUTED_BLURING_RADIUS;
 			KernelConfig.WorldBluringDistanceMultiplier = CONFIG_BILATERAL_DISTANCE_MULTIPLIER;
@@ -694,7 +694,7 @@ void TemporallyAccumulate(
 			
 				if (MAX_SIGNAL_BATCH_SIZE == 1)
 				{
-					// TODO: missing && abs(PixelOffset.x) <= 1 && abs(PixelOffset.y) <= 1;
+					// TODO(Denoiser): missing && abs(PixelOffset.x) <= 1 && abs(PixelOffset.y) <= 1;
 					KernelConfig.bForceAllAccumulation = CurrentFrameSamples.Array[0].SampleCount == 0;
 				}
 			}
@@ -925,7 +925,7 @@ void TemporallyAccumulate(
 		FSSDSignalArray CurrentTemporalAnalysis = CreateSignalArrayFromScalarValue(0.0);
 		{
 			float Luma = GetSignalLuma4(CurrentFrameSample, STANDARD_BUFFER_COLOR_SPACE);
-			float SampleCount = CurrentFrameSample.SampleCount; // TODO: 1 instead?
+			float SampleCount = CurrentFrameSample.SampleCount; // TODO(Denoiser): 1 instead?
 
 			CurrentTemporalAnalysis.Array[1].SampleCount = SampleCount;
 			CurrentTemporalAnalysis.Array[1].SceneColor.x = Luma * SampleCount;
@@ -936,7 +936,7 @@ void TemporallyAccumulate(
 		// Update the temporal analysis.
 		FSSDSignalArray TemporalAnalysis;
 		{
-			const float MinimumAnalasisHisteresis = 0.3; // TODO: expose cvar
+			const float MinimumAnalasisHisteresis = 0.3; // TODO(Denoiser): expose cvar
 
 			float PreviousFrameWeight = min(HistorySample.SampleCount, HistorySamples.Array[1].SampleCount * rcp(MinimumAnalasisHisteresis) - CurrentFrameSample.SampleCount);
 			float PreviousFrameMultiplier = HistorySample.SampleCount > 0 ? PreviousFrameWeight / HistorySample.SampleCount : 0;
@@ -990,7 +990,7 @@ void TemporallyAccumulate(
 			float PreviousFrameWeight = min(HistorySample.SampleCount, TargetedSampleCount - CurrentFrameSample.SampleCount);
 			float PreviousFrameMultiplier = HistorySample.SampleCount > 0 ? PreviousFrameWeight / HistorySample.SampleCount : 0;
 
-			// TODO: use WeightedLerp instead for VALU perf.
+			// TODO(Denoiser): use WeightedLerp instead for VALU perf.
 			// Pre transformation of the signal.
 			HistorySample = TransformSignal(
 				HistorySample,
diff --git a/Engine/Source/Runtime/Renderer/Private/ScreenSpaceDenoise.cpp b/Engine/Source/Runtime/Renderer/Private/ScreenSpaceDenoise.cpp
--- a/Engine/Source/Runtime/Renderer/Private/ScreenSpaceDenoise.cpp
+++ b/Engine/Source/Runtime/Renderer/Private/ScreenSpaceDenoise.cpp
@@ -102,7 +102,7 @@ static TAutoConsoleVariable<float> CVarGIHistoryConvolutionKernelSpreadFactor(
 
 
 /** The maximum number of mip level supported in the denoiser. */
-// TODO: jump to 3 because bufefr size already have a size multiple of 4.
+// TODO(Denoiser): jump to 3 because bufefr size already have a size multiple of 4.
 static const int32 kMaxMipLevel = 2;
 
 /** Maximum number of sample per pixel supported in the stackowiak sample set. */
@@ -1097,7 +1097,7 @@ class FSSDSpatialAccumulationCS : public FGlobalShader
 		SHADER_PARAMETER_STRUCT(FSSDSignalSRVs, SignalInputUint)
 		SHADER_PARAMETER_STRUCT(FSSDSignalUAVs, SignalOutput)
 
-		SHADER_PARAMETER_RDG_TEXTURE_UAV(RWTexture2D, DebugOutput) // TODO: remove
+		SHADER_PARAMETER_RDG_TEXTURE_UAV(RWTexture2D, DebugOutput) // TODO(Denoiser): remove
 	END_SHADER_PARAMETER_STRUCT()
 };
 
@@ -1147,7 +1147,7 @@ class FSSDTemporalAccumulationCS : public FGlobalShader
 		SHADER_PARAMETER_RDG_TEXTURE(Texture2D, PrevGBufferB)
 		SHADER_PARAMETER_RDG_TEXTURE_ARRAY(Texture2D<uint>, PrevCompressedMetadata, [kCompressedMetadataTextures])
 
-		SHADER_PARAMETER_RDG_TEXTURE_UAV(RWTexture2D, DebugOutput) // TODO: remove
+		SHADER_PARAMETER_RDG_TEXTURE_UAV(RWTexture2D, DebugOutput) // TODO(Denoiser): remove
 	END_SHADER_PARAMETER_STRUCT()
 };
 
@@ -1352,7 +1352,7 @@ static void DenoiseSignalAtConstantPixelDensity(
 			HistoryDescs[3].Format = PF_R32_FLOAT;
 
 			ReconstructionTextureCount = IScreenSpaceDenoiser::kSphericalHarmonicTextureCount;
-			HistoryTextureCountPerSignal = IScreenSpaceDenoiser::kSphericalHarmonicTextureCount; // TODO: only 3 textures for history
+			HistoryTextureCountPerSignal = IScreenSpaceDenoiser::kSphericalHarmonicTextureCount; // TODO(Denoiser): only 3 textures for history
 			bHasReconstructionLayoutDifferentFromHistory = false;
 
 			InjestTextureCount = 4;
@@ -1409,7 +1409,7 @@ static void DenoiseSignalAtConstantPixelDensity(
 		}
 
 		float PixelPositionToFullResPixel = 1.0f / Settings.DenoisingResolutionFraction;
-		FVector2D FullResPixelOffset = FVector2D(0.5f, 0.5f); // TODO.
+		FVector2D FullResPixelOffset = FVector2D(0.5f, 0.5f); // TODO(Denoiser).
 
 		CommonParameters.ThreadIdToBufferUV.X = PixelPositionToFullResPixel / float(FullResBufferExtent.X);
 		CommonParameters.ThreadIdToBufferUV.Y = PixelPositionToFullResPixel / float(FullResBufferExtent.Y);
@@ -2180,7 +2180,7 @@ public:
 			Settings.FullResViewport = View.ViewRect;
 			Settings.SignalProcessing = ESignalProcessing::PolychromaticPenumbraHarmonic;
 			Settings.HarmonicPeriode = Periode;
-			Settings.ReconstructionSamples = Periode * Periode; // TODO: should use preconvolution instead for harmonic 3
+			Settings.ReconstructionSamples = Periode * Periode; // TODO(Denoiser): should use preconvolution instead for harmonic 3
 			Settings.bUseTemporalAccumulation = false;
 
 			TStaticArray<FScreenSpaceDenoiserHistory*, IScreenSpaceDenoiser::kMaxBatchSize> PrevHistories;
@@ -2216,7 +2216,7 @@ public:
 			Settings.FullResViewport = View.ViewRect;
 			Settings.SignalProcessing = ESignalProcessing::PolychromaticPenumbraHarmonic;
 			Settings.HarmonicPeriode = Periode;
-			Settings.ReconstructionSamples = Periode * Periode; // TODO: should use preconvolution instead for harmonic 3
+			Settings.ReconstructionSamples = Periode * Periode; // TODO(Denoiser): should use preconvolution instead for harmonic 3
 			Settings.bUseTemporalAccumulation = false;
 
 			TStaticArray<FScreenSpaceDenoiserHistory*, IScreenSpaceDenoiser::kMaxBatchSize> PrevHistories;
@@ -2334,7 +2334,7 @@ public:
 		RDG_GPU_STAT_SCOPE(GraphBuilder, ReflectionsDenoiser);
 
 		// Imaginary depth is only used for Nvidia denoiser.
-		// TODO: permutation to not generate it?
+		// TODO(Denoiser): permutation to not generate it?
 		GraphBuilder.RemoveUnusedTextureWarning(ReflectionInputs.RayImaginaryDepth);
 
 		FSSDSignalTextures InputSignal;
