Fixed bad merge to staging with post process material pre-exposure.

diff --git a/Engine/Shaders/Private/PostProcessMaterialShaders.usf b/Engine/Shaders/Private/PostProcessMaterialShaders.usf
--- a/Engine/Shaders/Private/PostProcessMaterialShaders.usf
+++ b/Engine/Shaders/Private/PostProcessMaterialShaders.usf
@@ -200,7 +200,7 @@ void MainPS(
 #endif
 	OutColor = float4(GetMaterialEmissive(PixelMaterialInputs), Alpha );
 
-#if !MATERIALBLENDING_MODULATE && USE_PREEXPOSURE && !POST_PROCESS_MATERIAL_BEFORE_TONEMAP
+#if !MATERIALBLENDING_MODULATE && USE_PREEXPOSURE && POST_PROCESS_MATERIAL_BEFORE_TONEMAP
 	OutColor.xyz *= View.PreExposure;
 #endif
 } 
