Fixed wrong condition from 9491309.
Now ONLY shadow and local lights are disabled on mobile. Lighting is kept.
Tested on PC ES3.1 preview

diff --git a/Engine/Shaders/Private/MobileBasePassPixelShader.usf b/Engine/Shaders/Private/MobileBasePassPixelShader.usf
--- a/Engine/Shaders/Private/MobileBasePassPixelShader.usf
+++ b/Engine/Shaders/Private/MobileBasePassPixelShader.usf
@@ -501,9 +501,9 @@ void Main(
 	#endif
 
 		half Shadow = 1.0f;
-	#if !MATERIAL_SHADINGMODEL_UNLIT && !MATERIAL_SHADINGMODEL_SINGLELAYERWATER
+	#if !MATERIAL_SHADINGMODEL_UNLIT
 		Shadow = GetPrimaryPrecomputedShadowMask(Interpolants).r;
-		#if DIRECTIONAL_LIGHT_CSM && !MOBILE_MULTI_VIEW
+		#if DIRECTIONAL_LIGHT_CSM && !MOBILE_MULTI_VIEW && !MATERIAL_SHADINGMODEL_SINGLELAYERWATER
 		// Cascaded Shadow Map
 		{
 			FPCFSamplerSettings Settings;
@@ -594,7 +594,7 @@ void Main(
 			Color += SkyDiffuseLighting * ResolvedView.SkyLightColor.rgb * DiffuseColor * MaterialAO;
 		#endif
 
-		#if MAX_DYNAMIC_POINT_LIGHTS > 0
+		#if MAX_DYNAMIC_POINT_LIGHTS > 0 && !MATERIAL_SHADINGMODEL_SINGLELAYERWATER
 			#if VARIABLE_NUM_DYNAMIC_POINT_LIGHTS
 			#if WEBGL || ES2_PROFILE || VULKAN_PROFILE
 			// webgl and ES2 needs a constant loop counter
