Fix an uninitialized out parameter causing undefined behavior in the spot light estimator. Fix other similar instances through path and ray tracing shaders to prevent similar bugs in the future.

(ushell-p4-cherrypick of 9604175 by riley.alston)

diff --git a/Engine/Shaders/Private/PathTracing/Light/PathTracingLightSampling.ush b/Engine/Shaders/Private/PathTracing/Light/PathTracingLightSampling.ush
--- a/Engine/Shaders/Private/PathTracing/Light/PathTracingLightSampling.ush
+++ b/Engine/Shaders/Private/PathTracing/Light/PathTracingLightSampling.ush
@@ -13,8 +13,10 @@ LightSampling.usf: Light sampling for next-event estimation
 #include "PathTracingSkyLight.ush"
 #include "PathTracingSpotLight.ush"
 
-void GenerateLightRay(RayDesc Ray, int LightId, float3 LightUV, out RayDesc OutLightRay)
+void GenerateLightRay(RayDesc Ray, int LightId, float3 LightUV, inout RayDesc OutLightRay)
 {
+	OutLightRay = (RayDesc)0; // #dxr_todo: Return false in the case of no matching light type and handle case from caller rather than blindly using a degenerate ray
+
 	if (IsPointLight(LightId))
 	{
 		PointLight_GenerateLightRay(Ray, LightId, LightUV, OutLightRay);
@@ -42,10 +44,11 @@ void EstimateLight(
 	int LightId,
 	RayDesc Ray,
 	FMaterialClosestHitPayload Payload,
-	out float OutIrradiance
+	inout float OutIrradiance
 )
 {
-	OutIrradiance = 0.0;
+	OutIrradiance = 0.0f;
+
 	if (IsPointLight(LightId))
 	{
 		PointLight_EstimateLight(LightId, Ray, Payload, OutIrradiance);
@@ -72,10 +75,12 @@ bool SampleLightSelection(
 	RayDesc Ray,
 	FMaterialClosestHitPayload Payload,
 	float RandSample,
-	out int OutLightId,
-	out float OutLightSelectionPdf
+	inout int OutLightId,
+	inout float OutLightSelectionPdf
 )
 {
+	OutLightSelectionPdf = 0.0;
+
 //#define UNIFORM_LIGHT_SELECTION
 #ifdef UNIFORM_LIGHT_SELECTION
 	OutLightId = RandSample * SceneLightsData.Count;
@@ -126,14 +131,15 @@ void PdfLightSelection(
 	RayDesc Ray,
 	FMaterialClosestHitPayload Payload,
 	uint LightId,
-	out float OutLightSelectionPdf
+	inout float OutLightSelectionPdf
 )
 {
+	OutLightSelectionPdf = 0.0;
+
 #ifdef UNIFORM_LIGHT_SELECTION
 	OutLightSelectionPdf = 1.0 / SceneLightsData.Count;
 	return;
 #else
-	OutLightSelectionPdf = 0.0;
 	if (SceneLightsData.Count == 0) return;
 
 	// Build irradiance estimate prefix sum
@@ -170,7 +176,7 @@ void PdfLight(
 	FMaterialClosestHitPayload Payload,
 	int LightId,
 	float3 LightUV,
-	out float OutPdf)
+	inout float OutPdf)
 {
 	// Light selection
 	// TODO: Use radiance estimate instead of uniform random selection
@@ -207,9 +213,11 @@ void EvalLight(
 	int LightId,
 	float3 LightUV,
 	RayDesc Ray,
-	out float3 OutRadiance
+	inout float3 OutRadiance
 )
 {
+	OutRadiance = float3(0.0f, 0.0f, 0.0f);
+
 	if (IsPointLight(LightId))
 	{
 		PointLight_EvalLight(LightId, LightUV, Ray, OutRadiance);
@@ -236,11 +244,11 @@ void SampleLight(
 	RayDesc Ray,
 	FMaterialClosestHitPayload Payload,
 	float4 RandSample,
-	out int OutLightId,
-	out float3 OutLightUV,
-	out float OutPdf)
+	inout int OutLightId,
+	inout float3 OutLightUV,
+	inout float OutPdf)
 {
-	OutPdf = 0.0;
+	OutPdf = 0.0f;
 
 	// Light selection
 	float LightSelectionPdf = 0.0;
@@ -250,7 +258,6 @@ void SampleLight(
 	}
 
 	// Light sampling
-	OutLightUV = 0.0;
 	float LightPdf = 0.0;
 	if (IsPointLight(OutLightId))
 	{
diff --git a/Engine/Shaders/Private/PathTracing/Light/PathTracingSpotLight.ush b/Engine/Shaders/Private/PathTracing/Light/PathTracingSpotLight.ush
--- a/Engine/Shaders/Private/PathTracing/Light/PathTracingSpotLight.ush
+++ b/Engine/Shaders/Private/PathTracing/Light/PathTracingSpotLight.ush
@@ -147,6 +147,10 @@ void SpotLight_EstimateLight(
 		OutIrradiance = LightPower * Falloff / LightDistanceSquared;
 		// #dxr_todo: UE-72534 apply inner core
 	}
+	else
+	{
+		OutIrradiance = 0.0f;
+	}
 }
 
 void SpotLight_SampleLight(
diff --git a/Engine/Shaders/Private/PathTracing/Material/PathTracingGlossy.ush b/Engine/Shaders/Private/PathTracing/Material/PathTracingGlossy.ush
--- a/Engine/Shaders/Private/PathTracing/Material/PathTracingGlossy.ush
+++ b/Engine/Shaders/Private/PathTracing/Material/PathTracingGlossy.ush
@@ -78,6 +78,10 @@ void Glossy_SampleMaterial(
 	out float3 OutThroughput,
 	out float OutPdf)
 {
+	OutDirection = float3(0, 0, 0);
+	OutThroughput = float3(0, 0, 0);
+	OutPdf = 0.0f;
+
 	float3 N_World = GetWorldNormal(Payload);
 
 	// Diffuse and Specular color are calculated in the payload unpacking, both depend on BaseColor and Metallic attributes
@@ -92,9 +96,6 @@ void Glossy_SampleMaterial(
 	const float NoV = V.z;
 	const float ProbSpecular = GetSpecularEventProbability(SpecularColor, NoV); 
 
-	OutPdf = 0.0f;
-	OutThroughput = float3(0, 0, 0);
-
 	if ((RandSample.z <= ProbSpecular)) 
 	{	
 		bool Flip = V.z < 0.0f;
diff --git a/Engine/Shaders/Private/RayTracing/RayTracingRectLightRGS.usf b/Engine/Shaders/Private/RayTracing/RayTracingRectLightRGS.usf
--- a/Engine/Shaders/Private/RayTracing/RayTracingRectLightRGS.usf
+++ b/Engine/Shaders/Private/RayTracing/RayTracingRectLightRGS.usf
@@ -30,6 +30,12 @@ bool GenerateRectLightImportanceSampledRay(
 	out float RayPdf
 )
 {
+	RayOrigin = float3(0.0f, 0.0f, 0.0f);
+	RayDirection = float3(0.0f, 0.0f, 0.0f);
+	RayTMin = 0.0f;
+	RayTMax = 0.0f;
+	RayPdf = 0.0f;
+
 	uint MipCount = log2(RectLight.MipTreeDimensions.x);
 	uint2 MipPixel = 0;
 	float MipPdf = 1.0;
@@ -126,7 +132,6 @@ bool GenerateRectLightImportanceSampledRay(
 	// Light-normal culling
 	if (dot(-LightDirection, LightNormal) <= 0.0)
 	{
-		RayPdf = 0.0;
 		return false;
 	}
 
