Fix incorrect throughput calculation path tracing specular reflection material sampling.

(ushell-p4-cherrypick of 9548268 by riley.alston)

diff --git a/Engine/Shaders/Private/PathTracing/Material/PathTracingSpecularReflection.ush b/Engine/Shaders/Private/PathTracing/Material/PathTracingSpecularReflection.ush
--- a/Engine/Shaders/Private/PathTracing/Material/PathTracingSpecularReflection.ush
+++ b/Engine/Shaders/Private/PathTracing/Material/PathTracingSpecularReflection.ush
@@ -20,8 +20,7 @@ void SpecularReflection_SampleMaterial(
 
 	float NoV = saturate(dot(V_World, N_World));
 	float VoH = NoV;
-	float NoL = NoV;
-	OutThroughput = F_Schlick(SpecularColor, VoH) / NoL;
+	OutThroughput = F_Schlick(SpecularColor, VoH);
 
 	OutPdf = 1.0;
 }
