Improve generic TraceRay helper functions and replace existing common TraceRay calls which fit established material and visibility ray patterns.

(ushell-p4-cherrypick of 9568132 by riley.alston)

diff --git a/Engine/Shaders/Private/PathTracing/PathTracing.usf b/Engine/Shaders/Private/PathTracing/PathTracing.usf
--- a/Engine/Shaders/Private/PathTracing/PathTracing.usf
+++ b/Engine/Shaders/Private/PathTracing/PathTracing.usf
@@ -67,19 +67,16 @@ RAY_TRACING_ENTRY_RAYGEN(PathTracingMainRG)
 		FRayCone RayCone = (FRayCone)0;
 		RayCone.SpreadAngle = View.EyeToPixelSpreadAngle;
 
-		uint RayFlags = RAY_FLAG_CULL_BACK_FACING_TRIANGLES;
-
 		// Cast ray
-		FMaterialClosestHitPayload Payload = TraceRayInternal(
-			TLAS,   // AccelerationStructure
-			RayFlags,
-			RAY_TRACING_MASK_ALL,
-			RAY_TRACING_SHADER_SLOT_MATERIAL, // RayContributionToHitGroupIndex
-			RAY_TRACING_NUM_SHADER_SLOTS,     // MultiplierForGeometryContributionToShaderIndex
-			0,      // MissShaderIndex
-			Ray,    // RayDesc
-			RayCone
-		);
+		const uint InitialRayFlags = RAY_FLAG_CULL_BACK_FACING_TRIANGLES;
+		const uint InitialInstanceInclusionMask = RAY_TRACING_MASK_ALL;
+
+		FMaterialClosestHitPayload Payload = TraceMaterialRay(
+			TLAS,
+			InitialRayFlags,
+			InitialInstanceInclusionMask,
+			Ray,
+			RayCone);
 		RayCounter++;
 
 		float MaterialPdf = 0.0;
@@ -167,23 +164,17 @@ RAY_TRACING_ENTRY_RAYGEN(PathTracingMainRG)
 				{
 					RayDesc LightRay;
 
-					FRayCone LightRayCone = RayCone;
-
 					GenerateLightRay(Ray, LightId, LightUV, LightRay);
 					ApplyPositionBias(LightRay, Payload.WorldNormal, AdaptiveSamplingData.MaxNormalBias);
-					
-					uint RayFlags = RAY_FLAG_CULL_BACK_FACING_TRIANGLES;
 
-					FMaterialClosestHitPayload NeePayload = TraceRayInternal(
-						TLAS,      // AccelerationStructure
+					const uint RayFlags = RAY_FLAG_CULL_BACK_FACING_TRIANGLES;
+					const uint InstanceInclusionMask = RAY_TRACING_MASK_SHADOW;
+
+					FMinimalPayload NeePayload = TraceVisibilityRay(
+						TLAS,
 						RayFlags,
-						RAY_TRACING_MASK_SHADOW,
-						RAY_TRACING_SHADER_SLOT_SHADOW,   // RayContributionToHitGroupIndex
-						RAY_TRACING_NUM_SHADER_SLOTS,     // MultiplierForGeometryContributionToShaderIndex
-						0,         // MissShaderIndex
-						LightRay,
-						LightRayCone
-					);
+						InstanceInclusionMask,
+						LightRay);
 					RayCounter++;
 
 					if (NeePayload.IsMiss())
@@ -271,15 +262,15 @@ RAY_TRACING_ENTRY_RAYGEN(PathTracingMainRG)
 			}
 
 			// Trace again
-			Payload = TraceRayInternal(TLAS,
-				RAY_FLAG_CULL_BACK_FACING_TRIANGLES,
-				RAY_TRACING_MASK_ALL,
-				RAY_TRACING_SHADER_SLOT_MATERIAL, // RayContributionToHitGroupIndex
-				RAY_TRACING_NUM_SHADER_SLOTS,     // MultiplierForGeometryContributionToShaderIndex
-				0,    // MissShaderIndex
+			const uint RayFlags = RAY_FLAG_CULL_BACK_FACING_TRIANGLES;
+			const uint InstanceInclusionMask = RAY_TRACING_MASK_ALL;
+
+			Payload = TraceMaterialRay(
+				TLAS,
+				RayFlags,
+				InstanceInclusionMask,
 				Ray,
-				RayCone /* propagated by the material CHS */
-			);
+				RayCone);
 			RayCounter++;
 		}
 
diff --git a/Engine/Shaders/Private/RayTracing/RayTracingAmbientOcclusionRGS.usf b/Engine/Shaders/Private/RayTracing/RayTracingAmbientOcclusionRGS.usf
--- a/Engine/Shaders/Private/RayTracing/RayTracingAmbientOcclusionRGS.usf
+++ b/Engine/Shaders/Private/RayTracing/RayTracingAmbientOcclusionRGS.usf
@@ -83,8 +83,10 @@ RAY_TRACING_ENTRY_RAYGEN(AmbientOcclusionRGS)
 
 		if (dot(WorldNormal, Ray.Direction) <= 0.0) // TODO(Denoiser): does it needs to be handled by the denoiser?
 			continue;
-		
+
 		uint RayFlags = 0;
+		const uint InstanceInclusionMask = RAY_TRACING_MASK_SHADOW | RAY_TRACING_MASK_THIN_SHADOW;
+
 #if !ENABLE_MATERIALS
 		RayFlags |= RAY_FLAG_FORCE_OPAQUE;
 #endif
@@ -92,26 +94,17 @@ RAY_TRACING_ENTRY_RAYGEN(AmbientOcclusionRGS)
 		RayFlags |= RAY_FLAG_CULL_BACK_FACING_TRIANGLES;
 #endif
 
-		FPackedMaterialClosestHitPayload Payload = (FPackedMaterialClosestHitPayload)0;
-		Payload.SetMinimalPayloadMode(); // Only HitT will be filled by closest hit shader, skipping full material evaluation.
-		
-		const uint InstanceInclusionMask = RAY_TRACING_MASK_SHADOW | RAY_TRACING_MASK_THIN_SHADOW;
-		TraceRay(
+		FMinimalPayload MinimalPayload = TraceVisibilityRay(
 			TLAS,
 			RayFlags,
 			InstanceInclusionMask,
-			RAY_TRACING_SHADER_SLOT_SHADOW,   // RayContributionToHitGroupIndex
-			RAY_TRACING_NUM_SHADER_SLOTS,     // MultiplierForGeometryContributionToShaderIndex
-			0,    // MissShaderIndex
-			Ray,
-			Payload
-		);
-		
+			Ray);
+
 		RayCount += 1;
-		Visibility += Payload.IsMiss() ? 1.0 : 1.0 - IntensityLocal;
-		if (Payload.IsHit())
+		Visibility += MinimalPayload.IsMiss() ? 1.0 : 1.0 - IntensityLocal;
+		if (MinimalPayload.IsHit())
 		{
-			ClosestRayHitDistance = min(ClosestRayHitDistance, Payload.HitT);
+			ClosestRayHitDistance = min(ClosestRayHitDistance, MinimalPayload.HitT);
 		}
 	}
 
diff --git a/Engine/Shaders/Private/RayTracing/RayTracingBuiltInShaders.usf b/Engine/Shaders/Private/RayTracing/RayTracingBuiltInShaders.usf
--- a/Engine/Shaders/Private/RayTracing/RayTracingBuiltInShaders.usf
+++ b/Engine/Shaders/Private/RayTracing/RayTracingBuiltInShaders.usf
@@ -35,21 +35,15 @@ RAY_TRACING_ENTRY_RAYGEN(OcclusionMainRG)
 	              | RAY_FLAG_ACCEPT_FIRST_HIT_AND_END_SEARCH // terminate immediately upon detecting primitive intersection
 	              | RAY_FLAG_FORCE_OPAQUE                    // don't run anyhit shader
 	              | RAY_FLAG_SKIP_CLOSEST_HIT_SHADER;        // don't run closesthit shader
+	const uint InstanceInclusionMask = RAY_TRACING_MASK_OPAQUE;
 
-	FMinimalPayload Payload = (FMinimalPayload)0;
-
-	TraceRay(
-		TLAS,   // AccelerationStructure
+	FMinimalPayload MinimalPayload = TraceVisibilityRay(
+		TLAS,
 		RayFlags,
-		RAY_TRACING_MASK_OPAQUE,          // InstanceInclusionMask
-		RAY_TRACING_SHADER_SLOT_MATERIAL, // RayContributionToHitGroupIndex
-		RAY_TRACING_NUM_SHADER_SLOTS,     // MultiplierForGeometryContributionToShaderIndex
-		0,      // MissShaderIndex
-		Ray,    // RayDesc
-		Payload // Payload
-	);
+		InstanceInclusionMask,
+		Ray);
 
-	OcclusionOutput[RayIndex] = Payload.IsHit() ? ~0 : 0;
+	OcclusionOutput[RayIndex] = MinimalPayload.IsHit() ? ~0 : 0;
 }
 
 RAY_TRACING_ENTRY_RAYGEN(IntersectionMainRG)
diff --git a/Engine/Shaders/Private/RayTracing/RayTracingCommon.ush b/Engine/Shaders/Private/RayTracing/RayTracingCommon.ush
--- a/Engine/Shaders/Private/RayTracing/RayTracingCommon.ush
+++ b/Engine/Shaders/Private/RayTracing/RayTracingCommon.ush
@@ -192,6 +192,8 @@ struct FMaterialClosestHitPayload : FMinimalPayload
 
 	void SetFrontFace() { Flags |= RAY_TRACING_PAYLOAD_FLAG_FRONT_FACE; }
 	bool IsFrontFace() { return (Flags & RAY_TRACING_PAYLOAD_FLAG_FRONT_FACE) != 0; }
+
+	void SetRayCone(FRayCone NewRayCone) { RayCone = NewRayCone; }
 };
 
 #define USE_PACKED_PAYLOAD 1
@@ -229,6 +231,7 @@ struct FPackedMaterialClosestHitPayload : FMinimalPayload
 	void SetMinimalPayloadMode() { IorAndShadingModelIDAndBlendingModeAndFlags |= RAY_TRACING_PAYLOAD_FLAG_MINIMAL_PAYLOAD << 24; }
 	bool IsMinimalPayloadMode() { return ((IorAndShadingModelIDAndBlendingModeAndFlags >> 24) & RAY_TRACING_PAYLOAD_FLAG_MINIMAL_PAYLOAD) != 0; }
 
+	void SetRayCone(FRayCone NewRayCone) { RayCone = NewRayCone; }
 };
 
 FPackedMaterialClosestHitPayload PackRayTracingPayload(FMaterialClosestHitPayload Input, in FRayCone RayCone)
@@ -321,34 +324,6 @@ FMaterialClosestHitPayload UnpackRayTracingPayload(FPackedMaterialClosestHitPayl
 
 #endif // USE_PACKED_PAYLOAD
 
-FMaterialClosestHitPayload TraceRayInternal(in RaytracingAccelerationStructure TLAS,
-											in uint RayFlags,
-											in uint InstanceInclusionMask,
-											in uint RayContributionToHitGroupIndex,
-											in uint MultiplierForGeometryContributionToShaderIndex,
-											in uint MissShaderIndex,
-											in RayDesc Ray,
-											inout FRayCone RayCone)
-{
-	FPackedMaterialClosestHitPayload Payload = (FPackedMaterialClosestHitPayload)0;
-	Payload.RayCone = RayCone;
-
-	TraceRay
-	(
-			TLAS,
-			RayFlags,
-			InstanceInclusionMask,
-			RayContributionToHitGroupIndex,
-			MultiplierForGeometryContributionToShaderIndex,
-			MissShaderIndex,
-			Ray,
-			Payload
-	);
-
-	RayCone = Payload.RayCone;
-	return UnpackRayTracingPayload(Payload, Ray);
-}
-
 // WorldNormal is the vector towards which the ray position will be offseted.
 void ApplyPositionBias(inout RayDesc Ray, const float3 WorldNormal, const float MaxNormalBias)
 {
@@ -381,6 +356,131 @@ void ApplyCameraRelativeDepthBias(inout RayDesc Ray, uint2 PixelCoord, float Dev
 	ApplyPositionBias(Ray, WorldNormal, RayBias);
 }
 
+FMinimalPayload TraceMinimalRayInternal(
+	in RaytracingAccelerationStructure TLAS,
+	in uint RayFlags,
+	in uint InstanceInclusionMask,
+	in uint RayContributionToHitGroupIndex,
+	in uint MultiplierForGeometryContributionToShaderIndex,
+	in uint MissShaderIndex,
+	in RayDesc Ray)
+{
+	FPackedMaterialClosestHitPayload PackedPayload = (FPackedMaterialClosestHitPayload)0;
+
+	// Set payload inputs
+
+	// By enabling minimal payload mode all other payload information is ignored, meaning these functions need no payload inputs
+	PackedPayload.SetMinimalPayloadMode();
+
+	// Trace the ray
+
+	TraceRay(
+		TLAS,
+		RayFlags,
+		InstanceInclusionMask,
+		RayContributionToHitGroupIndex,
+		MultiplierForGeometryContributionToShaderIndex,
+		MissShaderIndex,
+		Ray,
+		PackedPayload);
+
+	// Unpack the payload
+
+	FMinimalPayload MinimalPayload = (FMinimalPayload)0;
+
+	// In theory this unpacking setp is not needed as FPackedMaterialClosestHitPayload derives from FMinimalPayload,
+	// but the compiler currently dislikes a direct cast between them. Additionally in the future if HitT is ever packed
+	// differently and the FMinimalPayload is not directly inherited from this will need to change.
+	MinimalPayload.HitT = PackedPayload.HitT;
+
+	return MinimalPayload;
+}
+
+FMaterialClosestHitPayload TraceRayInternal(
+	in RaytracingAccelerationStructure TLAS,
+	in uint RayFlags,
+	in uint InstanceInclusionMask,
+	in uint RayContributionToHitGroupIndex,
+	in uint MultiplierForGeometryContributionToShaderIndex,
+	in uint MissShaderIndex,
+	in RayDesc Ray,
+	// Payload Inputs
+	in FRayCone RayCone)
+{
+	FPackedMaterialClosestHitPayload PackedPayload = (FPackedMaterialClosestHitPayload)0;
+
+	// Set payload inputs
+
+	PackedPayload.SetRayCone(RayCone);
+
+	// Trace the ray
+
+	TraceRay(
+		TLAS,
+		RayFlags,
+		InstanceInclusionMask,
+		RayContributionToHitGroupIndex,
+		MultiplierForGeometryContributionToShaderIndex,
+		MissShaderIndex,
+		Ray,
+		PackedPayload);
+
+	// Unpack the payload
+
+	FMaterialClosestHitPayload Payload = UnpackRayTracingPayload(PackedPayload, Ray);
+
+	return Payload;
+}
+
+FMinimalPayload TraceVisibilityRay(
+	in RaytracingAccelerationStructure TLAS,
+	in uint RayFlags,
+	in uint InstanceInclusionMask,
+	in RayDesc Ray)
+{
+	const uint RayContributionToHitGroupIndex = RAY_TRACING_SHADER_SLOT_SHADOW;
+	const uint MultiplierForGeometryContributionToShaderIndex = RAY_TRACING_NUM_SHADER_SLOTS;
+	const uint MissShaderIndex = 0;
+
+	FMinimalPayload Payload = TraceMinimalRayInternal(
+		TLAS,
+		RayFlags,
+		InstanceInclusionMask,
+		RayContributionToHitGroupIndex,
+		MultiplierForGeometryContributionToShaderIndex,
+		MissShaderIndex,
+		Ray);
+
+	return Payload;
+}
+
+FMaterialClosestHitPayload TraceMaterialRay(
+	in RaytracingAccelerationStructure TLAS,
+	in uint RayFlags,
+	in uint InstanceInclusionMask,
+	in RayDesc Ray,
+	// Payload Inputs
+	inout FRayCone RayCone)
+{
+	const uint RayContributionToHitGroupIndex = RAY_TRACING_SHADER_SLOT_MATERIAL;
+	const uint MultiplierForGeometryContributionToShaderIndex = RAY_TRACING_NUM_SHADER_SLOTS;
+	const uint MissShaderIndex = 0;
+
+	FMaterialClosestHitPayload Payload = TraceRayInternal(
+		TLAS,
+		RayFlags,
+		InstanceInclusionMask,
+		RayContributionToHitGroupIndex,
+		MultiplierForGeometryContributionToShaderIndex,
+		MissShaderIndex,
+		Ray,
+		RayCone);
+
+	RayCone = Payload.RayCone;
+
+	return Payload;
+}
+
 #endif // !COMPUTESHADER
 
 #endif // RAYTRACINGCOMMON_USH_INCLUDED // Workarround for UE-66460
diff --git a/Engine/Shaders/Private/RayTracing/RayTracingCreateGatherPointsRGS.usf b/Engine/Shaders/Private/RayTracing/RayTracingCreateGatherPointsRGS.usf
--- a/Engine/Shaders/Private/RayTracing/RayTracingCreateGatherPointsRGS.usf
+++ b/Engine/Shaders/Private/RayTracing/RayTracingCreateGatherPointsRGS.usf
@@ -146,20 +146,19 @@ RAY_TRACING_ENTRY_RAYGEN(RayTracingCreateGatherPointsRGS)
 		{
 			// Cast ray
 			uint RayFlags = 0;
+			const uint InstanceInclusionMask = RAY_TRACING_MASK_OPAQUE;
+
 #if !ENABLE_TWO_SIDED_GEOMETRY
 			RayFlags |= RAY_FLAG_CULL_BACK_FACING_TRIANGLES;
 #endif
+
 			FRayCone RayCone = (FRayCone)0;
-			FMaterialClosestHitPayload Payload = TraceRayInternal(
-				TLAS,   // AccelerationStructure
+			FMaterialClosestHitPayload Payload = TraceMaterialRay(
+				TLAS,
 				RayFlags,
-				RAY_TRACING_MASK_OPAQUE,
-				RAY_TRACING_SHADER_SLOT_MATERIAL, // RayContributionToHitGroupIndex
-				RAY_TRACING_NUM_SHADER_SLOTS,     // MultiplierForGeometryContributionToShaderIndex
-				0,      // MissShaderIndex
-				Ray,    // RayDesc
-				RayCone
-			);
+				InstanceInclusionMask,
+				Ray,
+				RayCone);
 
 			// TODO: impossible based on formulation
 			if (Payload.HitT > MaxRayDistanceForGI) continue;
@@ -207,21 +206,18 @@ RAY_TRACING_ENTRY_RAYGEN(RayTracingCreateGatherPointsRGS)
 					ApplyPositionBias(LightRay, Payload.WorldNormal, MaxNormalBias);
 
 					// Trace visibility ray
-					uint RayFlags = RAY_FLAG_ACCEPT_FIRST_HIT_AND_END_SEARCH | RAY_FLAG_SKIP_CLOSEST_HIT_SHADER;
+					uint NeeRayFlags = RAY_FLAG_ACCEPT_FIRST_HIT_AND_END_SEARCH | RAY_FLAG_SKIP_CLOSEST_HIT_SHADER;
+					const uint NeeInstanceInclusionMask = RAY_TRACING_MASK_OPAQUE;
+
 #if !ENABLE_TWO_SIDED_GEOMETRY
-					RayFlags |= RAY_FLAG_CULL_BACK_FACING_TRIANGLES;
+					NeeRayFlags |= RAY_FLAG_CULL_BACK_FACING_TRIANGLES;
 #endif
-					FRayCone LightRayCone = (FRayCone)0;
-					FMaterialClosestHitPayload NeePayload = TraceRayInternal(
-						TLAS,   // AccelerationStructure
-						RayFlags,
-						RAY_TRACING_MASK_OPAQUE,
-						RAY_TRACING_SHADER_SLOT_SHADOW,  // RayContributionToHitGroupIndex
-						RAY_TRACING_NUM_SHADER_SLOTS,     // MultiplierForGeometryContributionToShaderIndex
-						0,      // MissShaderIndex
-						LightRay,    // RayDesc
-						LightRayCone
-					);
+
+					FMinimalPayload NeePayload = TraceVisibilityRay(
+						TLAS,
+						NeeRayFlags,
+						NeeInstanceInclusionMask,
+						LightRay);
 
 					// No hit indicates successful next-event connection
 					if (NeePayload.IsMiss())
diff --git a/Engine/Shaders/Private/RayTracing/RayTracingDebug.usf b/Engine/Shaders/Private/RayTracing/RayTracingDebug.usf
--- a/Engine/Shaders/Private/RayTracing/RayTracingDebug.usf
+++ b/Engine/Shaders/Private/RayTracing/RayTracingDebug.usf
@@ -23,18 +23,15 @@ RAY_TRACING_ENTRY_RAYGEN(RayTracingDebugMainRGS)
     FRayCone RayCone = (FRayCone)0;
 	RayCone.SpreadAngle = View.EyeToPixelSpreadAngle;
 
-	uint RayFlags = RAY_FLAG_CULL_BACK_FACING_TRIANGLES;
+	const uint RayFlags = RAY_FLAG_CULL_BACK_FACING_TRIANGLES;
+	const uint InstanceInclusionMask = RAY_TRACING_MASK_ALL;
 
-	FMaterialClosestHitPayload Payload = TraceRayInternal(
-		TLAS,   // AccelerationStructure
+	FMaterialClosestHitPayload Payload = TraceMaterialRay(
+		TLAS,
 		RayFlags,
-		RAY_TRACING_MASK_ALL,
-		RAY_TRACING_SHADER_SLOT_MATERIAL, // RayContributionToHitGroupIndex
-		RAY_TRACING_NUM_SHADER_SLOTS,     // MultiplierForGeometryContributionToShaderIndex
-		0,      // MissShaderIndex
-		Ray,    // RayDesc
-		RayCone
-	);
+		InstanceInclusionMask,
+		Ray,
+		RayCone);
 
 	float4 Result = float4(0, 0, 0, 1);
 
diff --git a/Engine/Shaders/Private/RayTracing/RayTracingFinalGatherRGS.usf b/Engine/Shaders/Private/RayTracing/RayTracingFinalGatherRGS.usf
--- a/Engine/Shaders/Private/RayTracing/RayTracingFinalGatherRGS.usf
+++ b/Engine/Shaders/Private/RayTracing/RayTracingFinalGatherRGS.usf
@@ -116,7 +116,7 @@ RAY_TRACING_ENTRY_RAYGEN(RayTracingFinalGatherRGS)
 		float3 RayThroughput = LightingSample.Diffuse / DiffuseColor;
 		float RayPdf = NoL / PI;
 
-		// Trace visibility ray
+		// Create ray
 		RayDesc Ray;
 		Ray.Origin = WorldPosition;
 		Ray.Direction = normalize(GPPosition - WorldPosition);
@@ -124,25 +124,22 @@ RAY_TRACING_ENTRY_RAYGEN(RayTracingFinalGatherRGS)
 		float RayLength = length(GPPosition - WorldPosition);
 		Ray.TMax = max(RayLength - 0.01, Ray.TMin);
 
-		// Cast ray
+		// Trace ray
 		uint RayFlags = 0;
+		const uint InstanceInclusionMask = RAY_TRACING_MASK_OPAQUE;
+
 #if !ENABLE_TWO_SIDED_GEOMETRY
 		RayFlags |= RAY_FLAG_CULL_BACK_FACING_TRIANGLES;
 #endif
-		FRayCone RayCone = (FRayCone)0;
-		FMaterialClosestHitPayload Payload = TraceRayInternal(
-			TLAS,   // AccelerationStructure
+
+		FMinimalPayload MinimalPayload = TraceVisibilityRay(
+			TLAS,
 			RayFlags,
-			RAY_TRACING_MASK_OPAQUE,
-			RAY_TRACING_SHADER_SLOT_SHADOW, // RayContributionToHitGroupIndex
-			RAY_TRACING_NUM_SHADER_SLOTS,     // MultiplierForGeometryContributionToShaderIndex
-			0,      // MissShaderIndex
-			Ray,    // RayDesc
-			RayCone
-		);
+			InstanceInclusionMask,
+			Ray);
 
 		// No hit indicates successful next-event connection
-		if (Payload.IsMiss() && RayPdf > 0.0)
+		if (MinimalPayload.IsMiss() && RayPdf > 0.0)
 		{
 			ExitantRadiance += GPIrradiance * RayThroughput / RayPdf;
 			HitDistance = min(HitDistance, Ray.TMax);
diff --git a/Engine/Shaders/Private/RayTracing/RayTracingGlobalIlluminationRGS.usf b/Engine/Shaders/Private/RayTracing/RayTracingGlobalIlluminationRGS.usf
--- a/Engine/Shaders/Private/RayTracing/RayTracingGlobalIlluminationRGS.usf
+++ b/Engine/Shaders/Private/RayTracing/RayTracingGlobalIlluminationRGS.usf
@@ -150,20 +150,19 @@ RAY_TRACING_ENTRY_RAYGEN(GlobalIlluminationRGS)
 		{
 			// Cast ray
 			uint RayFlags = 0;
+			const uint InstanceInclusionMask = RAY_TRACING_MASK_OPAQUE;
+
 #if !ENABLE_TWO_SIDED_GEOMETRY
 			RayFlags |= RAY_FLAG_CULL_BACK_FACING_TRIANGLES;
 #endif
+
 			FRayCone RayCone = (FRayCone)0;
-			FMaterialClosestHitPayload Payload = TraceRayInternal(
-				TLAS,   // AccelerationStructure
+			FMaterialClosestHitPayload Payload = TraceMaterialRay(
+				TLAS,
 				RayFlags,
-				RAY_TRACING_MASK_OPAQUE,
-				RAY_TRACING_SHADER_SLOT_MATERIAL, // RayContributionToHitGroupIndex
-				RAY_TRACING_NUM_SHADER_SLOTS,     // MultiplierForGeometryContributionToShaderIndex
-				0,      // MissShaderIndex
-				Ray,    // RayDesc
-				RayCone
-			);
+				InstanceInclusionMask,
+				Ray,
+				RayCone);
 
 			// Environment hit
 			if (Payload.IsMiss())
@@ -230,21 +229,18 @@ RAY_TRACING_ENTRY_RAYGEN(GlobalIlluminationRGS)
 					ApplyPositionBias(LightRay, Payload.WorldNormal, MaxNormalBias);
 
 					// Trace visibility ray
-					uint RayFlags = RAY_FLAG_ACCEPT_FIRST_HIT_AND_END_SEARCH | RAY_FLAG_SKIP_CLOSEST_HIT_SHADER;
+					uint NeeRayFlags = RAY_FLAG_ACCEPT_FIRST_HIT_AND_END_SEARCH | RAY_FLAG_SKIP_CLOSEST_HIT_SHADER;
+					const uint NeeInstanceInclusionMask = RAY_TRACING_MASK_OPAQUE;
+
 #if !ENABLE_TWO_SIDED_GEOMETRY
-					RayFlags |= RAY_FLAG_CULL_BACK_FACING_TRIANGLES;
+					NeeRayFlags |= RAY_FLAG_CULL_BACK_FACING_TRIANGLES;
 #endif
-					FRayCone LightRayCone = (FRayCone)0;
-					FMaterialClosestHitPayload NeePayload = TraceRayInternal(
-						TLAS,   // AccelerationStructure
-						RayFlags,
-						RAY_TRACING_MASK_OPAQUE,
-						RAY_TRACING_SHADER_SLOT_SHADOW,  // RayContributionToHitGroupIndex
-						RAY_TRACING_NUM_SHADER_SLOTS,     // MultiplierForGeometryContributionToShaderIndex
-						0,      // MissShaderIndex
-						LightRay,    // RayDesc
-						LightRayCone
-					);
+
+					FMinimalPayload NeePayload = TraceVisibilityRay(
+						TLAS,
+						NeeRayFlags,
+						NeeInstanceInclusionMask,
+						LightRay);
 
 					// No hit indicates successful next-event connection
 					if (NeePayload.IsMiss())
diff --git a/Engine/Shaders/Private/RayTracing/RayTracingLightingCommon.ush b/Engine/Shaders/Private/RayTracing/RayTracingLightingCommon.ush
--- a/Engine/Shaders/Private/RayTracing/RayTracingLightingCommon.ush
+++ b/Engine/Shaders/Private/RayTracing/RayTracingLightingCommon.ush
@@ -1,5 +1,7 @@
 // Copyright 1998-2019 Epic Games, Inc. All Rights Reserved.
 
+#include "RayTracingCommon.ush"
+
 // Light types: should match SceneTypes.h until there is a common header
 #define LIGHT_TYPE_DIRECTIONAL		0 
 #define LIGHT_TYPE_POINT			1 
@@ -73,31 +75,6 @@ float3 GenerateReflectedRayDirection(
 	return RayDirection;
 }
 
-float TraceShadowRay(
-	in RayDesc Ray,
-	in FRayCone ShadowRayCone,
-	in RaytracingAccelerationStructure TLAS)
-{
-	uint RayFlags = RAY_FLAG_ACCEPT_FIRST_HIT_AND_END_SEARCH | RAY_FLAG_SKIP_CLOSEST_HIT_SHADER;
-
-#if !ENABLE_TWO_SIDED_GEOMETRY
-	RayFlags |= RAY_FLAG_CULL_BACK_FACING_TRIANGLES;
-#endif
-
-	FMaterialClosestHitPayload Payload = TraceRayInternal(
-		TLAS,      // AccelerationStructure
-		RayFlags,
-		RAY_TRACING_MASK_SHADOW,
-		RAY_TRACING_SHADER_SLOT_SHADOW, // RayContributionToHitGroupIndex
-		RAY_TRACING_NUM_SHADER_SLOTS,   // MultiplierForGeometryContributionToShaderIndex
-		0,         // MissShaderIndex
-		Ray,       // RayDesc
-		ShadowRayCone
-	);
-
-	return Payload.HitT;
-}
-
 float3 SampleAreaLightDirection(
 	in FDeferredLightData LightData, 
 	in FMaterialClosestHitPayload Payload,
@@ -186,34 +163,6 @@ float3 SampleAreaLightDirection(
 	return ShadowRayDirection;
 }
 
-FPackedMaterialClosestHitPayload TraceRayInternalPacked(in RaytracingAccelerationStructure TLAS,
-	in uint RayFlags,
-	in uint InstanceInclusionMask,
-	in uint RayContributionToHitGroupIndex,
-	in uint MultiplierForGeometryContributionToShaderIndex,
-	in uint MissShaderIndex,
-	in RayDesc Ray,
-	inout FRayCone RayCone)
-{
-	FPackedMaterialClosestHitPayload Payload = (FPackedMaterialClosestHitPayload)0;
-	Payload.RayCone = RayCone;
-
-	TraceRay
-	(
-		TLAS,
-		RayFlags,
-		InstanceInclusionMask,
-		RayContributionToHitGroupIndex,
-		MultiplierForGeometryContributionToShaderIndex,
-		MissShaderIndex,
-		Ray,
-		Payload
-	);
-
-	RayCone = Payload.RayCone;
-	return Payload;
-}
-
 uint2 CullDirectLighting(
 	in RayDesc Ray,
 	in FRayCone RayCone,
@@ -458,8 +407,6 @@ float3 ComputeDirectLightingCulled(
 				ShadowRayDirection = SampleAreaLightDirection(LightData, Payload, LightType, RandSequence);
 			}
 
-			FRayCone ShadowRayCone = RayCone;
-
 			RayDesc ShadowRay;
 			ShadowRay.Origin = WorldPosition;
 			ShadowRay.Direction = ShadowRayDirection;
@@ -467,8 +414,19 @@ float3 ComputeDirectLightingCulled(
 			ShadowRay.TMax = Active ? length(ToLight) : ShadowRay.TMin;
 			ApplyPositionBias(ShadowRay, Payload.WorldNormal, ShadowMaxNormalBias);
 
-			float HitT = TraceShadowRay(ShadowRay, ShadowRayCone, TLAS);
-			AmbientOcclusion = HitT < 0.0f && Active;
+			uint RayFlags = RAY_FLAG_ACCEPT_FIRST_HIT_AND_END_SEARCH | RAY_FLAG_SKIP_CLOSEST_HIT_SHADER;
+			const uint InstanceInclusionMask = RAY_TRACING_MASK_SHADOW;
+
+#if !ENABLE_TWO_SIDED_GEOMETRY
+			RayFlags |= RAY_FLAG_CULL_BACK_FACING_TRIANGLES;
+#endif
+
+			FMinimalPayload ShadowRayPayload = TraceVisibilityRay(
+				TLAS,
+				RayFlags,
+				InstanceInclusionMask,
+				ShadowRay);
+			AmbientOcclusion = ShadowRayPayload.IsMiss() && Active;
 		}
 
 		if (Active)
@@ -515,38 +473,12 @@ void ComputeBottomLayerMaterialProperties(RayDesc Ray, inout FMaterialClosestHit
 	}
 }
 
-FMaterialClosestHitPayload TraceRayResults(
-	in RayDesc Ray,
-	in RaytracingAccelerationStructure TLAS,
-	in RandomSequence RandSequence,
-	in uint RayFlags,
-	in float ShadowMaxNormalBias,
-	in uint Mask,
-	inout FRayCone RayCone,
-	inout float3 Radiance)
-{
-	FPackedMaterialClosestHitPayload PackedPayload = TraceRayInternalPacked(
-		TLAS,								// AccelerationStructure
-		RayFlags,
-		Mask,
-		RAY_TRACING_SHADER_SLOT_MATERIAL,	// RayContributionToHitGroupIndex
-		RAY_TRACING_NUM_SHADER_SLOTS,		// MultiplierForGeometryContributionToShaderIndex
-		0,									// MissShaderIndex
-		Ray,								// RayDesc
-		RayCone
-	);
-	FMaterialClosestHitPayload Payload = UnpackRayTracingPayload(PackedPayload, Ray);
-	return Payload;
-}
-
 void AccumulateResults(
 	in FMaterialClosestHitPayload Payload,
 	in RayDesc Ray,
 	in RaytracingAccelerationStructure TLAS,
 	in RandomSequence RandSequence,
-	in uint RayFlags,
 	in float ShadowMaxNormalBias,
-	in uint Mask,
 	in uint ReflectedShadowsType,
 	in uint ShouldDoDirectLighting,
 	in uint ShouldDoEmissiveAndIndirectLighting,
@@ -587,35 +519,30 @@ void AccumulateResults(
 
 FMaterialClosestHitPayload TraceRayAndAccumulateResults(
 	in RayDesc Ray,
-	in RaytracingAccelerationStructure TLAS, 
-	in RandomSequence RandSequence,
+	in RaytracingAccelerationStructure TLAS,
 	in uint RayFlags,
+	in uint InstanceInclusionMask,
+	in RandomSequence RandSequence,
 	in float ShadowMaxNormalBias,
-	in uint Mask,
 	in uint ReflectedShadowsType,
 	in uint ShouldDoDirectLighting,
 	in uint ShouldDoEmissiveAndIndirectLighting,
 	inout FRayCone RayCone,
 	inout float3 Radiance)
 {
-	FMaterialClosestHitPayload Payload = TraceRayResults(
-		Ray,
+	FMaterialClosestHitPayload Payload = TraceMaterialRay(
 		TLAS,
-		RandSequence,
 		RayFlags,
-		ShadowMaxNormalBias,
-		Mask,
-		RayCone,
-		Radiance);
+		InstanceInclusionMask,
+		Ray,
+		RayCone);
 
 	AccumulateResults(
 		Payload,
 		Ray,
 		TLAS,
 		RandSequence,
-		RayFlags,
 		ShadowMaxNormalBias,
-		Mask,
 		ReflectedShadowsType,
 		ShouldDoDirectLighting,
 		ShouldDoEmissiveAndIndirectLighting,
@@ -628,25 +555,22 @@ FMaterialClosestHitPayload TraceRayAndAccumulateResults(
 FMaterialClosestHitPayload TraceRayAndAccumulateBottomLayerResults(
 	in RayDesc Ray,
 	in RaytracingAccelerationStructure TLAS,
-	in RandomSequence RandSequence,
 	in uint RayFlags,
+	in uint InstanceInclusionMask,
+	in RandomSequence RandSequence,
 	in float ShadowMaxNormalBias,
-	in uint Mask,
 	in uint ReflectedShadowsType,
 	in uint ShouldDoDirectLighting,
 	in uint ShouldDoEmissiveAndIndirectLighting,
 	inout FRayCone RayCone,
 	inout float3 Radiance)
 {
-	FMaterialClosestHitPayload BottomLayerPayload = TraceRayResults(
-		Ray,
+	FMaterialClosestHitPayload BottomLayerPayload = TraceMaterialRay(
 		TLAS,
-		RandSequence,
 		RayFlags,
-		ShadowMaxNormalBias,
-		Mask,
-		RayCone,
-		Radiance);
+		InstanceInclusionMask,
+		Ray,
+		RayCone);
 
 	ComputeBottomLayerMaterialProperties(Ray, BottomLayerPayload);
 
@@ -655,9 +579,7 @@ FMaterialClosestHitPayload TraceRayAndAccumulateBottomLayerResults(
 		Ray,
 		TLAS,
 		RandSequence,
-		RayFlags,
 		ShadowMaxNormalBias,
-		Mask,
 		ReflectedShadowsType,
 		ShouldDoDirectLighting,
 		ShouldDoEmissiveAndIndirectLighting,
diff --git a/Engine/Shaders/Private/RayTracing/RayTracingOcclusionRGS.usf b/Engine/Shaders/Private/RayTracing/RayTracingOcclusionRGS.usf
--- a/Engine/Shaders/Private/RayTracing/RayTracingOcclusionRGS.usf
+++ b/Engine/Shaders/Private/RayTracing/RayTracingOcclusionRGS.usf
@@ -269,41 +269,35 @@ RAY_TRACING_ENTRY_RAYGEN(OcclusionRGS)
 
 			uint RayFlags = 0;
 
-	#if !ENABLE_TWO_SIDED_GEOMETRY
+#if !ENABLE_TWO_SIDED_GEOMETRY
 			RayFlags |= RAY_FLAG_CULL_BACK_FACING_TRIANGLES;
-	#endif
-			
+#endif
+
 			if (DIM_DENOISER_OUTPUT == 0)
 			{
 				// Denoiser mode 0 doesn't use depth, so take first hit
 				RayFlags |= RAY_FLAG_ACCEPT_FIRST_HIT_AND_END_SEARCH;
 			}
 
-			FPackedMaterialClosestHitPayload Payload = (FPackedMaterialClosestHitPayload)0;
-			Payload.SetMinimalPayloadMode(); // Only HitT will be filled by closest hit shader, skipping full material evaluation.
-		
-			TraceRay(
-				TLAS,      // AccelerationStructure
+			FMinimalPayload MinimalPayload = TraceVisibilityRay(
+				TLAS,
 				RayFlags,
-				RaytracingMask,                 // InstanceInclusionMask
-				RAY_TRACING_SHADER_SLOT_SHADOW, // RayContributionToHitGroupIndex
-				RAY_TRACING_NUM_SHADER_SLOTS,   // MultiplierForGeometryContributionToShaderIndex
-				0,         // MissShaderIndex
-				Ray,       // RayDesc
-				Payload    // Payload
-			);
+				RaytracingMask,
+				Ray);
 
 			RayCount += 1.0;
-			Visibility += Payload.IsMiss() ? 1.0 : 0.0;
-			if (Payload.IsHit())
+			Visibility += MinimalPayload.IsMiss() ? 1.0 : 0.0;
+			if (MinimalPayload.IsHit())
 			{
-				ClosestRayDistance = ClosestRayDistance == DENOISER_INVALID_HIT_DISTANCE || Payload.HitT < ClosestRayDistance ? Payload.HitT : ClosestRayDistance;
-				SumRayDistance += Payload.HitT;
+				ClosestRayDistance =
+					(ClosestRayDistance == DENOISER_INVALID_HIT_DISTANCE) ||
+					(MinimalPayload.HitT < ClosestRayDistance) ? MinimalPayload.HitT : ClosestRayDistance;
+				SumRayDistance += MinimalPayload.HitT;
 				HitCount += 1.0;
 			}
 			else
 			{
-				ClosestRayDistance = ClosestRayDistance == DENOISER_INVALID_HIT_DISTANCE ? DENOISER_MISS_HIT_DISTANCE : ClosestRayDistance;
+				ClosestRayDistance = (ClosestRayDistance == DENOISER_INVALID_HIT_DISTANCE) ? DENOISER_MISS_HIT_DISTANCE : ClosestRayDistance;
 			}
 		}
 	}
diff --git a/Engine/Shaders/Private/RayTracing/RayTracingPrimaryRays.usf b/Engine/Shaders/Private/RayTracing/RayTracingPrimaryRays.usf
--- a/Engine/Shaders/Private/RayTracing/RayTracingPrimaryRays.usf
+++ b/Engine/Shaders/Private/RayTracing/RayTracingPrimaryRays.usf
@@ -144,9 +144,6 @@ RAY_TRACING_ENTRY_RAYGEN(RayTracingPrimaryRaysRGS)
 	
 	bool bHasScattered = (ERayTracingPrimaryRaysFlag_ConsiderSurfaceScatter & PrimaryRayFlags) != 0;
 
-	const bool bRefractionBackFaceCulling = false;
-	const bool bReflectionBackFaceCulling = true;
-
 	// Integrated data by path tracing
 	float3 PathRadiance = 0.0;
 	float PathThroughput = 1.0;	// A float for now because UE does not support colored translucency as of today.
@@ -156,10 +153,22 @@ RAY_TRACING_ENTRY_RAYGEN(RayTracingPrimaryRaysRGS)
 
 	for (uint RefractionRayIndex = 0; RefractionRayIndex < MaxRefractionRays; ++RefractionRayIndex)
 	{
+		const uint RefractionRayFlags = 0;
+		const uint RefractionInstanceInclusionMask = RAY_TRACING_MASK_ALL;
 		float3 PathVertexRadiance = float3(0, 0, 0);
-		//#dxr-todo: RayFlags should be 0 for refraction rays
-		const uint RayFlags = bRefractionBackFaceCulling ? RAY_FLAG_CULL_BACK_FACING_TRIANGLES : 0;
-		FMaterialClosestHitPayload Payload = TraceRayAndAccumulateResults(Ray, TLAS, RandSequence, RayFlags, MaxNormalBias, RAY_TRACING_MASK_ALL, ReflectedShadowsType, ShouldDoDirectLighting, ShouldDoEmissiveAndIndirectLighting, RayCone, PathVertexRadiance);
+
+		FMaterialClosestHitPayload Payload = TraceRayAndAccumulateResults(
+			Ray,
+			TLAS,
+			RefractionRayFlags,
+			RefractionInstanceInclusionMask,
+			RandSequence,
+			MaxNormalBias,
+			ReflectedShadowsType,
+			ShouldDoDirectLighting,
+			ShouldDoEmissiveAndIndirectLighting,
+			RayCone,
+			PathVertexRadiance);
 		LastRoughness = Payload.Roughness;
 
 		//
@@ -208,9 +217,22 @@ RAY_TRACING_ENTRY_RAYGEN(RayTracingPrimaryRaysRGS)
 			ReflectionRay.Direction = GenerateReflectedRayDirection(Ray.Direction, Payload.WorldNormal, Payload.Roughness, RandSample);
 			ApplyPositionBias(ReflectionRay, Payload.WorldNormal, MaxNormalBias);
 
+			const uint ReflectionRayFlags = RAY_FLAG_CULL_BACK_FACING_TRIANGLES;
+			const uint ReflectionInstanceInclusionMask = RAY_TRACING_MASK_ALL;
 			float3 ReflectionRadiance = float3(0, 0, 0);
-			const uint RayFlags = bReflectionBackFaceCulling ? RAY_FLAG_CULL_BACK_FACING_TRIANGLES : 0;
-			FMaterialClosestHitPayload ReflectionPayload = TraceRayAndAccumulateResults(ReflectionRay, TLAS, RandSequence, RayFlags, MaxNormalBias, RAY_TRACING_MASK_ALL, ReflectedShadowsType, ShouldDoDirectLighting, ShouldDoEmissiveAndIndirectLighting, RayCone, ReflectionRadiance);
+
+			FMaterialClosestHitPayload ReflectionPayload = TraceRayAndAccumulateResults(
+				ReflectionRay,
+				TLAS,
+				ReflectionRayFlags,
+				ReflectionInstanceInclusionMask,
+				RandSequence,
+				MaxNormalBias,
+				ReflectedShadowsType,
+				ShouldDoDirectLighting,
+				ShouldDoEmissiveAndIndirectLighting,
+				RayCone,
+				ReflectionRadiance);
 
 			// If we have not hit anything, sample the distance sky radiance.
 			if (ReflectionPayload.IsMiss())
diff --git a/Engine/Shaders/Private/RayTracing/RayTracingRectLightRGS.usf b/Engine/Shaders/Private/RayTracing/RayTracingRectLightRGS.usf
--- a/Engine/Shaders/Private/RayTracing/RayTracingRectLightRGS.usf
+++ b/Engine/Shaders/Private/RayTracing/RayTracingRectLightRGS.usf
@@ -267,27 +267,18 @@ RAY_TRACING_ENTRY_RAYGEN(RectLightRGS)
 			continue;
 		}
 
-		uint RayFlags = 0
-			| RAY_FLAG_CULL_BACK_FACING_TRIANGLES
-			| RAY_FLAG_FORCE_OPAQUE; // #dxr_todo: UE-72563 generate a special shader permutation for occlusion ray tracing that includes any-hit shaders for masked geometry
+		const uint RayFlags = RAY_FLAG_CULL_BACK_FACING_TRIANGLES | RAY_FLAG_FORCE_OPAQUE; // #dxr_todo: UE-72563 generate a special shader permutation for occlusion ray tracing that includes any-hit shaders for masked geometry
+		const uint InstanceInclusionMask = RAY_TRACING_MASK_SHADOW;
 
-		FPackedMaterialClosestHitPayload Payload = (FPackedMaterialClosestHitPayload)0;
-		Payload.SetMinimalPayloadMode(); // Only HitT will be filled by closest hit shader, skipping full material evaluation.
-
-		TraceRay(
-			TLAS,      // AccelerationStructure
+		FMinimalPayload MinimalPayload = TraceVisibilityRay(
+			TLAS,
 			RayFlags,
-			RAY_TRACING_MASK_SHADOW,          // InstanceInclusionMask
-			RAY_TRACING_SHADER_SLOT_SHADOW,   // RayContributionToHitGroupIndex
-			RAY_TRACING_NUM_SHADER_SLOTS,     // MultiplierForGeometryContributionToShaderIndex
-			0,         // MissShaderIndex
-			Ray,       // RayDesc
-			Payload    // Payload
-		);
+			InstanceInclusionMask,
+			Ray);
 
-		if (Payload.IsHit())
+		if (MinimalPayload.IsHit())
 		{
-			RayDistance += Payload.HitT;
+			RayDistance += MinimalPayload.HitT;
 			HitCount += 1.0;
 		}
 		else
diff --git a/Engine/Shaders/Private/RayTracing/RayTracingReflections.usf b/Engine/Shaders/Private/RayTracing/RayTracingReflections.usf
--- a/Engine/Shaders/Private/RayTracing/RayTracingReflections.usf
+++ b/Engine/Shaders/Private/RayTracing/RayTracingReflections.usf
@@ -373,7 +373,6 @@ RAY_TRACING_ENTRY_RAYGEN(RayTracingReflectionsRGS)
 				// Trace the top layer reflection
 				float3 TopLayerRadiance = float3(0, 0, 0);
 				FMaterialClosestHitPayload Payload = (FMaterialClosestHitPayload)0;
-				const uint RayFlags = 0;
 
 				if (DIM_HYBRID && DIM_DEFERRED_MATERIAL_MODE == DEFERRED_MATERIAL_MODE_SHADE 
 					&& DeferredMaterialPayload.SortKey == RAY_TRACING_DEFERRED_MATERIAL_KEY_SCREEN_SPACE)
@@ -394,15 +393,15 @@ RAY_TRACING_ENTRY_RAYGEN(RayTracingReflectionsRGS)
 				}
 				else
 				{
-					Payload = TraceRayResults(
-						TopLayerRay,
+					const uint RayFlags = 0;
+					const uint InstanceInclusionMask = RAY_TRACING_MASK_OPAQUE;
+
+					Payload = TraceMaterialRay(
 						TLAS,
-						RandSequence,
 						RayFlags,
-						ReflectionMaxNormalBias,
-						RAY_TRACING_MASK_OPAQUE,
-						RayCone,
-						TopLayerRadiance);
+						InstanceInclusionMask,
+						TopLayerRay,
+						RayCone);
 
 					RayDesc LightingRay = TopLayerRay;
 					if (DIM_HYBRID)
@@ -417,9 +416,7 @@ RAY_TRACING_ENTRY_RAYGEN(RayTracingReflectionsRGS)
 						LightingRay,
 						TLAS,
 						RandSequence,
-						RayFlags,
 						ReflectionMaxNormalBias,
-						RAY_TRACING_MASK_OPAQUE,
 						ReflectedShadowsType,
 						ShouldDoDirectLighting,
 						ShouldDoEmissiveAndIndirectLighting,
@@ -495,8 +492,21 @@ RAY_TRACING_ENTRY_RAYGEN(RayTracingReflectionsRGS)
 
 					if (TraceBottomRay)
 					{
-						FRayCone BottomLayerRayCone = (FRayCone)0;
-						FMaterialClosestHitPayload BottomLayerPayload = TraceRayAndAccumulateBottomLayerResults(BottomLayerRay, TLAS, RandSequence, RayFlags, ReflectionMaxNormalBias, RAY_TRACING_MASK_OPAQUE, ReflectedShadowsType, ShouldDoDirectLighting, ShouldDoEmissiveAndIndirectLighting, RayCone, BottomLayerRadiance);
+						const uint RayFlags = 0;
+						const uint InstanceInclusionMask = RAY_TRACING_MASK_OPAQUE;
+
+						FMaterialClosestHitPayload BottomLayerPayload = TraceRayAndAccumulateBottomLayerResults(
+							BottomLayerRay,
+							TLAS,
+							RayFlags,
+							InstanceInclusionMask,
+							RandSequence,
+							ReflectionMaxNormalBias,
+							ReflectedShadowsType,
+							ShouldDoDirectLighting,
+							ShouldDoEmissiveAndIndirectLighting,
+							RayCone,
+							BottomLayerRadiance);
 
 						// Distance and sky: similar processing to main ray.
 						float BottomLayerHitDistance = -1;
diff --git a/Engine/Shaders/Private/RayTracing/RayTracingSkyLightRGS.usf b/Engine/Shaders/Private/RayTracing/RayTracingSkyLightRGS.usf
--- a/Engine/Shaders/Private/RayTracing/RayTracingSkyLightRGS.usf
+++ b/Engine/Shaders/Private/RayTracing/RayTracingSkyLightRGS.usf
@@ -140,29 +140,24 @@ RAY_TRACING_ENTRY_RAYGEN(SkyLightRGS)
 		ApplyCameraRelativeDepthBias(Ray, PixelCoord, DeviceZ, WorldNormal, SkyLight.MaxNormalBias);
 
 		uint RayFlags = 0;
+		const uint InstanceInclusionMask = RAY_TRACING_MASK_SHADOW;
+
 #if !ENABLE_MATERIALS
 		RayFlags |= RAY_FLAG_FORCE_OPAQUE;
 #endif
 #if !ENABLE_TWO_SIDED_GEOMETRY
 		RayFlags |= RAY_FLAG_CULL_BACK_FACING_TRIANGLES;
 #endif
-		FPackedMaterialClosestHitPayload Payload = (FPackedMaterialClosestHitPayload)0;
-		Payload.SetMinimalPayloadMode(); // Only HitT will be filled by closest hit shader, skipping full material evaluation.
 
-		TraceRay(
-			TLAS,      // AccelerationStructure
+		FMinimalPayload MinimalPayload = TraceVisibilityRay(
+			TLAS,
 			RayFlags,
-			RAY_TRACING_MASK_SHADOW,          // InstanceInclusionMask
-			RAY_TRACING_SHADER_SLOT_SHADOW,   // RayContributionToHitGroupIndex
-			RAY_TRACING_NUM_SHADER_SLOTS,     // MultiplierForGeometryContributionToShaderIndex
-			0,         // MissShaderIndex
-			Ray,       // RayDesc
-			Payload    // Payload
-		);
-
-		if (Payload.IsHit())
+			InstanceInclusionMask,
+			Ray);
+
+		if (MinimalPayload.IsHit())
 		{
-			RayDistance += Payload.HitT;
+			RayDistance += MinimalPayload.HitT;
 			HitCount += 1.0;
 		}
 		else
@@ -217,36 +212,29 @@ RAY_TRACING_ENTRY_RAYGEN(SkyLightRGS)
 		// Cast ray from light to shading point to get estimation of "shadow thickness"
 		float ShadowThickness = 0.0;
 		{
+			RayDesc RayLocal;
+			RayLocal.Origin = WorldPosition + L * SkyLight.MaxShadowThickness;
+			RayLocal.Direction = -L;
+			RayLocal.TMin = 0.0;
+			RayLocal.TMax = SkyLight.MaxShadowThickness;
+
 			uint RayFlagsLocal = 0;
+			const uint InstanceInclusionMaskLocal = RAY_TRACING_MASK_SHADOW;
+
 #if !ENABLE_MATERIALS
 			RayFlagsLocal |= RAY_FLAG_FORCE_OPAQUE;
 #endif
 #if !ENABLE_TWO_SIDED_GEOMETRY
 			RayFlagsLocal |= RAY_FLAG_CULL_BACK_FACING_TRIANGLES;
 #endif
-			FPackedMaterialClosestHitPayload Payload = (FPackedMaterialClosestHitPayload)0;
-			Payload.SetMinimalPayloadMode(); // Only HitT will be filled by closest hit shader, skipping full material evaluation.
-
-			RayDesc RayLocal;
-			RayLocal.Origin = WorldPosition + L * SkyLight.MaxShadowThickness;
-			RayLocal.Direction = -L;
-			RayLocal.TMin = 0.0;
-			RayLocal.TMax = SkyLight.MaxShadowThickness;
 
-			FPackedMaterialClosestHitPayload PayloadLocal = (FPackedMaterialClosestHitPayload)0;
-			PayloadLocal.SetMinimalPayloadMode(); // Only HitT will be filled by closest hit shader, skipping full material evaluation.
-			TraceRay(
-				TLAS,      // AccelerationStructure
+			FMinimalPayload MinimalPayloadLocal = TraceVisibilityRay(
+				TLAS,
 				RayFlagsLocal,
-				RAY_TRACING_MASK_SHADOW,          // InstanceInclusionMask
-				RAY_TRACING_SHADER_SLOT_SHADOW,   // RayContributionToHitGroupIndex
-				RAY_TRACING_NUM_SHADER_SLOTS,     // MultiplierForGeometryContributionToShaderIndex
-				0,         // MissShaderIndex
-				RayLocal,       // RayDesc
-				PayloadLocal    // Payload
-			);
-
-			float3 P = RayLocal.Origin + RayLocal.Direction * PayloadLocal.HitT;
+				InstanceInclusionMaskLocal,
+				RayLocal);
+
+			float3 P = RayLocal.Origin + RayLocal.Direction * MinimalPayloadLocal.HitT;
 			ShadowThickness = min(length(WorldPosition - P), 1.0);
 		}
 		float3 ScatterThroughput = PI * KajiyaKayDiffuseAttenuation(GBufferData, L, V, N, ShadowThickness);
diff --git a/Engine/Shaders/Private/RayTracing/RayTracingTest.usf b/Engine/Shaders/Private/RayTracing/RayTracingTest.usf
--- a/Engine/Shaders/Private/RayTracing/RayTracingTest.usf
+++ b/Engine/Shaders/Private/RayTracing/RayTracingTest.usf
@@ -23,19 +23,13 @@ RAY_TRACING_ENTRY_RAYGEN(TestMainRGS)
 	              | RAY_FLAG_ACCEPT_FIRST_HIT_AND_END_SEARCH // terminate immediately upon detecting primitive intersection
 	              | RAY_FLAG_FORCE_OPAQUE                    // don't run anyhit shader
 	              | RAY_FLAG_SKIP_CLOSEST_HIT_SHADER;        // don't run closesthit shader
+	const uint InstanceInclusionMask = RAY_TRACING_MASK_ALL;
 
-	FMinimalPayload Payload = (FMinimalPayload)0;
-
-	TraceRay(
-		TLAS,   // AccelerationStructure
+	FMinimalPayload MinimalPayload = TraceVisibilityRay(
+		TLAS,
 		RayFlags,
-		RAY_TRACING_MASK_ALL,             // InstanceInclusionMask
-		RAY_TRACING_SHADER_SLOT_MATERIAL, // RayContributionToHitGroupIndex
-		RAY_TRACING_NUM_SHADER_SLOTS,     // MultiplierForGeometryContributionToShaderIndex
-		0,      // MissShaderIndex
-		Ray,    // RayDesc
-		Payload // Payload
-	);
+		InstanceInclusionMask,
+		Ray);
 
-	Output[RayIndex] = Payload.IsHit() ? ~0 : 0;
+	Output[RayIndex] = MinimalPayload.IsHit() ? ~0 : 0;
 }
