Add flags to improve ray tracing performance on visibility rays not requiring hit depth.

(ushell-p4-cherrypick of 9568384 by riley.alston)

diff --git a/Engine/Shaders/Private/PathTracing/PathTracing.usf b/Engine/Shaders/Private/PathTracing/PathTracing.usf
--- a/Engine/Shaders/Private/PathTracing/PathTracing.usf
+++ b/Engine/Shaders/Private/PathTracing/PathTracing.usf
@@ -167,7 +167,7 @@ RAY_TRACING_ENTRY_RAYGEN(PathTracingMainRG)
 					GenerateLightRay(Ray, LightId, LightUV, LightRay);
 					ApplyPositionBias(LightRay, Payload.WorldNormal, AdaptiveSamplingData.MaxNormalBias);
 
-					const uint RayFlags = RAY_FLAG_CULL_BACK_FACING_TRIANGLES;
+					const uint RayFlags = RAY_FLAG_ACCEPT_FIRST_HIT_AND_END_SEARCH | RAY_FLAG_SKIP_CLOSEST_HIT_SHADER | RAY_FLAG_CULL_BACK_FACING_TRIANGLES;
 					const uint InstanceInclusionMask = RAY_TRACING_MASK_SHADOW;
 
 					FMinimalPayload NeePayload = TraceVisibilityRay(
diff --git a/Engine/Shaders/Private/RayTracing/RayTracingFinalGatherRGS.usf b/Engine/Shaders/Private/RayTracing/RayTracingFinalGatherRGS.usf
--- a/Engine/Shaders/Private/RayTracing/RayTracingFinalGatherRGS.usf
+++ b/Engine/Shaders/Private/RayTracing/RayTracingFinalGatherRGS.usf
@@ -125,7 +125,7 @@ RAY_TRACING_ENTRY_RAYGEN(RayTracingFinalGatherRGS)
 		Ray.TMax = max(RayLength - 0.01, Ray.TMin);
 
 		// Trace ray
-		uint RayFlags = 0;
+		uint RayFlags = RAY_FLAG_ACCEPT_FIRST_HIT_AND_END_SEARCH | RAY_FLAG_SKIP_CLOSEST_HIT_SHADER;
 		const uint InstanceInclusionMask = RAY_TRACING_MASK_OPAQUE;
 
 #if !ENABLE_TWO_SIDED_GEOMETRY
diff --git a/Engine/Shaders/Private/RayTracing/RayTracingOcclusionRGS.usf b/Engine/Shaders/Private/RayTracing/RayTracingOcclusionRGS.usf
--- a/Engine/Shaders/Private/RayTracing/RayTracingOcclusionRGS.usf
+++ b/Engine/Shaders/Private/RayTracing/RayTracingOcclusionRGS.usf
@@ -276,7 +276,7 @@ RAY_TRACING_ENTRY_RAYGEN(OcclusionRGS)
 			if (DIM_DENOISER_OUTPUT == 0)
 			{
 				// Denoiser mode 0 doesn't use depth, so take first hit
-				RayFlags |= RAY_FLAG_ACCEPT_FIRST_HIT_AND_END_SEARCH;
+				RayFlags |= RAY_FLAG_SKIP_CLOSEST_HIT_SHADER | RAY_FLAG_ACCEPT_FIRST_HIT_AND_END_SEARCH;
 			}
 
 			FMinimalPayload MinimalPayload = TraceVisibilityRay(
