LandscapeContinuousLODParameters is a ES3.1 level feature -- restore ES2 bits

#rnx

diff --git a/Engine/Shaders/Private/LandscapeVertexFactory.ush b/Engine/Shaders/Private/LandscapeVertexFactory.ush
--- a/Engine/Shaders/Private/LandscapeVertexFactory.ush
+++ b/Engine/Shaders/Private/LandscapeVertexFactory.ush
@@ -32,11 +32,12 @@ float4 LodBias;
 */
 float4 LodValues;
 
-int ForcedLod;
-
 /* Fractional Lod values for each section */
 float4 SectionLods;
 
+#if ! ES2_PROFILE
+int ForcedLod;
+
 int GetComponentLinearIndex()
 {
 	return (LandscapeParameters.ComponentBaseY - LandscapeContinuousLODParameters.Min.y) * LandscapeContinuousLODParameters.Size.x + (LandscapeParameters.ComponentBaseX - LandscapeContinuousLODParameters.Min.x);
@@ -83,6 +84,7 @@ float4 GetNeighborSectionLod()
 	}
 	return float4(NeighborLod[0], NeighborLod[1], NeighborLod[2], NeighborLod[3]);
 }
+#endif
 
 float4 GetLodBias()
 {
@@ -587,7 +589,18 @@ float CalcLOD(float2 xy, float2 Subsection)
 	// LOD value based on Barycentric coordinate
 	float4 L0 = float4(xy.y, xy.x, (1 - xy.x), (1 - xy.y)) * 2;
 
+#if ES2_PROFILE
+	float4 LODCalculated4 = lerp(
+								(lerp(L0 * SectionLods.x + (float4(1, 1, 1, 1) - L0) * NeighborSectionLod[0],
+									  L0 * SectionLods.y + (float4(1, 1, 1, 1) - L0) * NeighborSectionLod[1],
+									  Subsection.x)),
+								(lerp(L0 * SectionLods.z + (float4(1, 1, 1, 1) - L0) * NeighborSectionLod[2],
+									  L0 * SectionLods.w + (float4(1, 1, 1, 1) - L0) * NeighborSectionLod[3],
+									  Subsection.x)),
+								Subsection.y);
+#else
 	float4 LODCalculated4 = L0 * GetSectionLods() + (float4(1, 1, 1, 1) - L0) * GetNeighborSectionLod();
+#endif
 
 	// Find which quadrant of the subsection we're in - top, left, right or bottom.
 	// 0 --------- 1 X  1 ----- 2  1          1 
