Fix for potential sun flashing red due to inf/nan.
It may be that it also simply is the bloom failing on high value but I have not been able to repro. the fp16 safe range seems like a good thing to add in anyway. (however it does not prevent issue with out of bound additive blending later on.)

(ushell-p4-cherrypick of 9615127 by [REDACTED_PERSON])

diff --git a/Engine/Shaders/Private/BasePassPixelShader.usf b/Engine/Shaders/Private/BasePassPixelShader.usf
--- a/Engine/Shaders/Private/BasePassPixelShader.usf
+++ b/Engine/Shaders/Private/BasePassPixelShader.usf
@@ -1187,7 +1187,7 @@ ShadingModel);
 	}
 #endif
 
-#if NEEDS_BASEPASS_PIXEL_FOGGING && PROJECT_SUPPORT_SKY_ATMOSPHERE && MATERIAL_FORCE_SKIP_AERIAL_PERSPECTIVE==0
+#if NEEDS_BASEPASS_PIXEL_FOGGING && PROJECT_SUPPORT_SKY_ATMOSPHERE && MATERIAL_IS_SKY==0 // Do not apply aerial perpsective on sky materials
 	if (ResolvedView.SkyAtmosphereApplyCameraAerialPerspectiveVolume > 0.0f)
 	{
 		const float OneOverPreExposure = USE_PREEXPOSURE ? ResolvedView.OneOverPreExposure : 1.0f;
@@ -1395,6 +1395,10 @@ ShadingModel);
 #if !MATERIALBLENDING_MODULATE && USE_PREEXPOSURE
 	Out.MRT[0].xyz *= View.PreExposure;
 #endif
+#if MATERIAL_IS_SKY
+	// Sky materials can result in high luminance values, e.g. the sun disk. So we make sure to at least stay within the boundaries of fp10 for some platforms.
+	Out.MRT[0].xyz = min(Out.MRT[0].xyz, Max10BitsFloat.xxx);
+#endif
 
 #if NUM_VIRTUALTEXTURE_SAMPLES || LIGHTMAP_VT_ENABLED
 	FinalizeVirtualTextureFeedback(
diff --git a/Engine/Shaders/Private/BasePassVertexShader.usf b/Engine/Shaders/Private/BasePassVertexShader.usf
--- a/Engine/Shaders/Private/BasePassVertexShader.usf
+++ b/Engine/Shaders/Private/BasePassVertexShader.usf
@@ -128,7 +128,7 @@ void Main(
 	Output.BasePassInterpolants.VertexFog = CalculateHeightFog(WorldPosition.xyz - ResolvedView.TranslatedWorldCameraOrigin);
 	#endif
 
-#if PROJECT_SUPPORT_SKY_ATMOSPHERE && MATERIAL_FORCE_SKIP_AERIAL_PERSPECTIVE==0
+#if PROJECT_SUPPORT_SKY_ATMOSPHERE && MATERIAL_IS_SKY==0 // Do not apply aerial perpsective on sky materials
 	if (ResolvedView.SkyAtmosphereApplyCameraAerialPerspectiveVolume > 0.0f)
 	{
 		const float OneOverPreExposure = USE_PREEXPOSURE ? ResolvedView.OneOverPreExposure : 1.0f;
diff --git a/Engine/Shaders/Private/Common.ush b/Engine/Shaders/Private/Common.ush
--- a/Engine/Shaders/Private/Common.ush
+++ b/Engine/Shaders/Private/Common.ush
@@ -81,6 +81,7 @@
 
 const static MaterialFloat PI = 3.1415926535897932f;
 const static float MaxHalfFloat = 65504.0f;
+const static float Max10BitsFloat = 64512.0f;
 
 #ifndef USE_RAYTRACED_TEXTURE_RAYCONE_LOD
 #define USE_RAYTRACED_TEXTURE_RAYCONE_LOD (RAYHITGROUPSHADER)
diff --git a/Engine/Shaders/Private/MobileBasePassPixelShader.usf b/Engine/Shaders/Private/MobileBasePassPixelShader.usf
--- a/Engine/Shaders/Private/MobileBasePassPixelShader.usf
+++ b/Engine/Shaders/Private/MobileBasePassPixelShader.usf
@@ -727,6 +727,11 @@ void Main(
 		OutColor.rgb *= ResolvedView.PreExposure;
 	#endif
 
+	#if MATERIAL_IS_SKY
+		// Sky materials can result in high luminance values, e.g. the sun disk. So we make sure to at least stay within the boundaries of fp10 for some platforms.
+		OutColor.rgb = min(OutColor.rgb, Max10BitsFloat.xxx);
+	#endif
+
 #if USE_EDITOR_COMPOSITING && (MOBILE_EMULATION)
 	// Editor primitive depth testing
 	OutColor.a = 1.0;
diff --git a/Engine/Shaders/Private/MobileBasePassVertexShader.usf b/Engine/Shaders/Private/MobileBasePassVertexShader.usf
--- a/Engine/Shaders/Private/MobileBasePassVertexShader.usf
+++ b/Engine/Shaders/Private/MobileBasePassVertexShader.usf
@@ -92,7 +92,7 @@ void Main(
 #if USE_VERTEX_FOG
 	float4 VertexFog = CalculateHeightFog(WorldPosition.xyz - ResolvedView.TranslatedWorldCameraOrigin);
 
-#if PROJECT_SUPPORT_SKY_ATMOSPHERE && MATERIAL_FORCE_SKIP_AERIAL_PERSPECTIVE==0
+#if PROJECT_SUPPORT_SKY_ATMOSPHERE && MATERIAL_IS_SKY==0 // Do not apply aerial perpsective on sky materials
 	if (ResolvedView.SkyAtmosphereApplyCameraAerialPerspectiveVolume > 0.0f)
 	{
 		const float OneOverPreExposure = USE_PREEXPOSURE ? ResolvedView.OneOverPreExposure : 1.0f;
diff --git a/Engine/Source/Runtime/Engine/Private/Materials/HLSLMaterialTranslator.h b/Engine/Source/Runtime/Engine/Private/Materials/HLSLMaterialTranslator.h
--- a/Engine/Source/Runtime/Engine/Private/Materials/HLSLMaterialTranslator.h
+++ b/Engine/Source/Runtime/Engine/Private/Materials/HLSLMaterialTranslator.h
@@ -1426,7 +1426,7 @@ ResourcesString = TEXT("");
 		OutEnvironment.SetDefine(TEXT("USES_DISTORTION"), Material->IsDistorted()); 
 
 		OutEnvironment.SetDefine(TEXT("MATERIAL_ENABLE_TRANSLUCENCY_FOGGING"), Material->ShouldApplyFogging());
-		OutEnvironment.SetDefine(TEXT("MATERIAL_FORCE_SKIP_AERIAL_PERSPECTIVE"), Material->IsSky());
+		OutEnvironment.SetDefine(TEXT("MATERIAL_IS_SKY"), Material->IsSky());
 		OutEnvironment.SetDefine(TEXT("MATERIAL_COMPUTE_FOG_PER_PIXEL"), Material->ComputeFogPerPixel());
 		OutEnvironment.SetDefine(TEXT("MATERIAL_FULLY_ROUGH"), bIsFullyRough || Material->IsFullyRough());
 
