Fixed black atmosphere pixels appearing sometimes when transitioning to space view.

diff --git a/Engine/Shaders/Private/SkyAtmosphere.usf b/Engine/Shaders/Private/SkyAtmosphere.usf
--- a/Engine/Shaders/Private/SkyAtmosphere.usf
+++ b/Engine/Shaders/Private/SkyAtmosphere.usf
@@ -799,8 +799,6 @@ void RenderSkyAtmosphereRayMarchingPS(
 
 #else // FASTAERIALPERSPECTIVE_ENABLED
 
-	WorldPos += WorldDir * AerialPerspectiveStartDepth; // Offset according to start depth
-
 	// Move to top atmosphere as the starting point for ray marching.
 	// This is critical to be after the above to not disrupt above atmosphere tests and voxel selection.
 	if (!MoveToTopAtmosphere(WorldPos, WorldDir, Atmosphere.TopRadius))
@@ -810,6 +808,9 @@ void RenderSkyAtmosphereRayMarchingPS(
 		return;
 	}
 
+	// Apply the start depth offset after moving to the top of atmosphere for consistency (and to avoid wrong out-of-atmosphere test resulting in black pixels).
+	WorldPos += WorldDir * AerialPerspectiveStartDepth;
+
 	SamplingSetup Sampling;
 	{
 		Sampling.VariableSampleCount = true;
