Fixed visual pop when transitioning from in atmosphere rendering using LUT to space view rendering using ray marching. ==> using a more accurate, slighlty slower, fast math acos.

diff --git a/Engine/Shaders/Private/SkyAtmosphere.usf b/Engine/Shaders/Private/SkyAtmosphere.usf
--- a/Engine/Shaders/Private/SkyAtmosphere.usf
+++ b/Engine/Shaders/Private/SkyAtmosphere.usf
@@ -264,7 +264,7 @@ void UvToSkyViewLutParams(out float3 ViewDir, in float ViewHeight, in float2 UV)
 
 	float Vhorizon = sqrt(ViewHeight * ViewHeight - Atmosphere.BottomRadius * Atmosphere.BottomRadius);
 	float CosBeta = Vhorizon / ViewHeight;				// cos of zenith angle from horizon to zeniht
-	float Beta = acosFast(CosBeta);
+	float Beta = acosFast4(CosBeta);
 	float ZenithHorizonAngle = PI - Beta;
 
 	float ViewZenithAngle;
diff --git a/Engine/Shaders/Private/SkyAtmosphereCommon.ush b/Engine/Shaders/Private/SkyAtmosphereCommon.ush
--- a/Engine/Shaders/Private/SkyAtmosphereCommon.ush
+++ b/Engine/Shaders/Private/SkyAtmosphereCommon.ush
@@ -129,12 +129,13 @@ void SkyViewLutParamsToUv(
 {
 	float Vhorizon = sqrt(ViewHeight * ViewHeight - BottomRadius * BottomRadius);
 	float CosBeta = Vhorizon / ViewHeight;				// GroundToHorizonCos
-	float Beta = acosFast(CosBeta);
+	float Beta = acosFast4(CosBeta);
 	float ZenithHorizonAngle = PI - Beta;
+	float ViewZenithAngle = acosFast4(ViewZenithCosAngle);
 
 	if (!IntersectGround)
 	{
-		float Coord = acosFast(ViewZenithCosAngle) / ZenithHorizonAngle;
+		float Coord = ViewZenithAngle / ZenithHorizonAngle;
 		Coord = 1.0f - Coord;
 		Coord = sqrt(Coord);
 		Coord = 1.0f - Coord;
@@ -142,7 +143,7 @@ void SkyViewLutParamsToUv(
 	}
 	else
 	{
-		float Coord = (acosFast(ViewZenithCosAngle) - ZenithHorizonAngle) / Beta;
+		float Coord = (ViewZenithAngle - ZenithHorizonAngle) / Beta;
 		Coord = sqrt(Coord);
 		UV.y = Coord * 0.5f + 0.5f;
 	}
