Updates to GTAO
Improved Spatial and temporal filters.
Removed upsample for full res render
Improved spatial offsets per sample

diff --git a/Engine/Shaders/Private/PostProcessAmbientOcclusion.usf b/Engine/Shaders/Private/PostProcessAmbientOcclusion.usf
--- a/Engine/Shaders/Private/PostProcessAmbientOcclusion.usf
+++ b/Engine/Shaders/Private/PostProcessAmbientOcclusion.usf
@@ -952,6 +952,8 @@ float4 GTAOParams[3];
 
 #if COMPUTE_SHADER
 RWTexture2D<float2> HorizonOutTexture;
+RWTexture2D<float> DepthOutTexture;
+RWTexture2D<float> DepthsTexture;
 #endif
 
 Texture2D		HistoryTexture;
@@ -960,6 +962,9 @@ SamplerState	HistoryTextureSampler;
 Texture2D		DepthHistoryTexture;
 SamplerState	DepthHistoryTextureSampler;
 
+Texture2D		ZReadTexture;
+SamplerState	ZReadTextureSampler;
+
 
 float4 PrevScreenPositionScaleBias;
 float4 FilterParams;
@@ -968,14 +973,9 @@ float4 FilterParams;
 
 float Falloff(float distSq)
 {
-	float dist = sqrt(distSq);
-	float Falloff_Scale = 1.0f / (FALLOFF_END - FALLOFF_START);
-	float Falloff_Bias = -FALLOFF_START * Falloff_Scale; 
-	return  2.0 * saturate(dist* Falloff_Scale + Falloff_Bias   );
-
-//	float Falloff_Scale = 1.0f / (FALLOFF_END_SQ - FALLOFF_START_SQ);
-//	float Falloff_Bias = -FALLOFF_START_SQ * Falloff_Scale; 
-//	return  2.0 * saturate(distSq* Falloff_Scale + Falloff_Bias   );
+	float Falloff_Scale = 1.0f / (FALLOFF_END_SQ - FALLOFF_START_SQ);
+	float Falloff_Bias = -FALLOFF_START_SQ * Falloff_Scale; 
+	return  2.0 * saturate(distSq* Falloff_Scale + Falloff_Bias   );
 }
 
 float3 GetViewSpacePosFromHZB(float2 UV, float MipLevel)
@@ -1061,7 +1061,6 @@ float SearchForLargestAngleRef(float2 inUV, float2 deltaUV, float3 ViewSpacePosi
 	float BiggestAngle = -1.0f;
 	float PrevAngle    = -1.0f;
 
-	deltaUV.y	  *= -1;
 	float2 UV	   = inUV ;
 
 #if GTAO_OPT
@@ -1141,7 +1140,6 @@ float SearchForLargestAngle(float2 inUV, float2 deltaUV, float3 ViewSpacePositio
 	float BiggestAngle = -1.0f;
 	float PrevAngle	   = -1.0f;
 
-	deltaUV.y	  *=-1;
 	float2 UV	= inUV;
 
 	float MipLevel = 0;
@@ -1161,26 +1159,21 @@ float SearchForLargestAngle(float2 inUV, float2 deltaUV, float3 ViewSpacePositio
 	UNROLL
 	for(int i=1; i<=GTAO_NUMTAPS; i++)
 	{
-		if(i==2)
+		if(i==3)
 		{
 			HZB_dUV += HZB_dUV;
 			DeltaVS_XYLen	+= DeltaVS_XYLen;
 			MipLevel++;
 		}
 		
-		if(i==GTAO_NUMTAPS/4)
-		{
-			HZB_dUV += HZB_dUV;
-			DeltaVS_XYLen	+= DeltaVS_XYLen;
-			MipLevel++;
-		}
-		else if(i==GTAO_NUMTAPS/2)
+		if(i==GTAO_NUMTAPS/2)
 		{
 			HZB_dUV += HZB_dUV;
 			DeltaVS_XYLen	+= DeltaVS_XYLen;
 			MipLevel++;
 		}
 
+
 		float DeviceZ		= Texture2DSampleLevel( PostprocessInput1, PostprocessInput1Sampler, HZB_UV, MipLevel).r;
 		float SceneDepth    = 1.0f / (DeviceZ * View.InvDeviceZToWorldZTransform[2] - View.InvDeviceZToWorldZTransform[3]);
 
@@ -1223,15 +1216,14 @@ float SearchForLargestAngle(float2 inUV, float2 deltaUV, float3 ViewSpacePositio
 			deltaUV *=1.5;
 			MipLevel++;
 		}
-		//else if(i==GTAO_NUMTAPS/2)
-		//{
-		//	deltaUV += deltaUV;
-		//	MipLevel++;
-		//}
 
 		UV				+= deltaUV;
 		
-		float3 SampVS   = GetViewSpacePosFromHZB(UV,MipLevel);
+		// Quantize UV to pixel centre
+		//float2 UV2 = floor( (View.BufferSizeAndInvSize.xy * UV) + float2(0.25,0.25) ) * View.BufferSizeAndInvSize.zw;
+		//UV2 += View.BufferSizeAndInvSize.zw*0.5;
+		
+		float3 SampVS	= GetViewSpacePosFromHZB(UV,MipLevel);
 		float3 DiffV	= (SampVS - ViewSpacePosition) ;
 		float distSQ	= dot(DiffV, DiffV);
 
@@ -1245,8 +1237,6 @@ float SearchForLargestAngle(float2 inUV, float2 deltaUV, float3 ViewSpacePositio
 			float falloff	= Falloff(distSQ);
 			
 			cosh			= cosh - falloff;
-
-
 			BiggestAngle	= max(BiggestAngle, cosh);
 
 			//BiggestAngle	= UpdateMaxAngle(cosh, PrevAngle);
@@ -1261,19 +1251,15 @@ float SearchForLargestAngle(float2 inUV, float2 deltaUV, float3 ViewSpacePositio
 
 
 // Given a screen space Axis this will search for the min and max angles
-float2 SearchAxisForAngles(float2 UV, float3 ScreenDir, float3 ViewDir, float3 ViewSpacePos, float Scale)
+float2 SearchAxisForAngles(float2 UV, float3 ScreenDir, float3 ViewDir, float3 ViewSpacePos, float Scale, float Offset)
 {
 	// TODO - Fix aspect ratio
 	float2 DeltaUV	= Scale * float2(View.BufferSizeAndInvSize.z * ScreenDir.x, View.BufferSizeAndInvSize.w * ScreenDir.y);
 	
-	float O =  0.0f;//GTAOParams[0].z; 
-
-	float2 StartUV = UV + (DeltaUV*O);
-	float Ang1 = SearchForLargestAngle( StartUV,  DeltaUV, ViewSpacePos, ViewDir);
-
-	StartUV = UV - (DeltaUV*O);
-	float Ang2 = SearchForLargestAngle( StartUV, -DeltaUV, ViewSpacePos, ViewDir);
+	DeltaUV.y	  *=-1;
 
+	float Ang1 = SearchForLargestAngle( UV + (DeltaUV*Offset),   DeltaUV, ViewSpacePos, ViewDir);
+	float Ang2 = SearchForLargestAngle( UV - (DeltaUV*Offset),  -DeltaUV, ViewSpacePos, ViewDir);
 	return float2(Ang1, Ang2);
 }
 
@@ -1285,66 +1271,67 @@ float ComputeInnerIntegral(float2 Angles, float3 ScreenDir, float3 ViewDir, floa
 	float3 ProjNormal		= ViewSpaceNormal - PlaneNormal * dot(ViewSpaceNormal, PlaneNormal);
 
 	float LenProjNormal		= length(ProjNormal);
-	//float ProjNormalMagSq	= dot(ProjNormal, ProjNormal);
 	float RecipMag			= 1.0f / (LenProjNormal+1e-6);
 	float CosAng			= dot(ProjNormal, Perp) * RecipMag;	
 	float Gamma				= acosFast(CosAng) - PI_HALF;				
 	float CosGamma			= dot(ProjNormal, ViewDir) * RecipMag;
 	float SinGamma  		= CosAng * -2.0f;					
 
-	const float ANGLEBIAS = 0;// PI_HALF/4.0f;
-
 	// clamp to normal hemisphere 
-	Angles.x = Gamma + max(-Angles.x - Gamma, -(PI_HALF-ANGLEBIAS) );
-	Angles.y = Gamma + min( Angles.y - Gamma,  (PI_HALF-ANGLEBIAS) );
+	Angles.x = Gamma + max(-Angles.x - Gamma, -(PI_HALF) );
+	Angles.y = Gamma + min( Angles.y - Gamma,  (PI_HALF) );
+
+	LenProjNormal = 1;
 
-	float AO = (/* LenProjNormal * */   0.25 * 
+	float AO = ( LenProjNormal *  0.25 * 
 					    ( (Angles.x * SinGamma + CosGamma - cos((2.0 * Angles.x) - Gamma)) +
 				  	      (Angles.y * SinGamma + CosGamma - cos((2.0 * Angles.y) - Gamma)) ));
 
-	AO = saturate(AO);
-
-	//const float Bias = 0.2;
-	//AO = saturate( (AO - Bias)  *  (1.0/(1.0-Bias))  );
+	AO = saturate(AO / PI_HALF);
 
 	// Fade out based on user defined distance
 	float Mul = ScreenSpaceAOParams[4].x;
 	float Add = ScreenSpaceAOParams[4].y;
 	AO = lerp(AO, 1, saturate(SceneDepth * Mul + Add));
 
-
 	return AO;
 }
 
 
 
-/*
-*
-* HORIZON SEARCH AND INNER INTEGRATE COMBINED
-*
-*/
 
-float2 GetRandomVector(float2 UV )
+float3 GetRandomVector(float2 UV )
 {
 	float2 ViewportUVToRandomUV = ScreenSpaceAOParams[1].xy;
+	float3 RandomTexVec = Texture2DSample(RandomNormalTexture, RandomNormalTextureSampler, UV * ViewportUVToRandomUV).rgb;
+	RandomTexVec.xy = RandomTexVec.xy * 2 -1;
+
 	float TemporalCos = GTAOParams[0].x;
 	float TemporalSin = GTAOParams[0].y;
-	float2 RandomTexVec = (Texture2DSample(RandomNormalTexture, RandomNormalTextureSampler, UV * ViewportUVToRandomUV).rg * 2 - 1) ;
-	float2 RandomVec;
+	float3 RandomVec;
 
 	// Rotate the random Vec
 	RandomVec.x = dot(RandomTexVec.xy, float2(TemporalCos, -TemporalSin ));
 	RandomVec.y = dot(RandomTexVec.xy, float2(TemporalSin, TemporalCos ));
+	RandomVec.z =RandomTexVec.z;
 
 	return RandomVec;
 }
 
 
 
-void GTAOCombinedPSandCS(in float2 UV, out float4 OutColor)
+/*
+*
+* HORIZON SEARCH AND INNER INTEGRATE COMBINED
+*
+*/
+void GTAOCombinedPSandCS(in float2 UV, out float4 OutColor, inout float Depth)
 {
 	OutColor = 0;
 
+	// Offset by a fraction of a pixel to unsure we don't hit between pixels when running at half res
+	UV +=  PostprocessInput0Size.zw*0.125;
+
 	float SceneDepth = GetDepthFromAOInput(UV);		
 	if(SceneDepth > ScreenSpaceAOParams[4].w)	
 	{		
@@ -1352,7 +1339,9 @@ void GTAOCombinedPSandCS(in float2 UV, out float4 OutColor)
 		return;
 	}
 	
-	float2 RandomVec		= GetRandomVector(UV);
+	float3 RandomAndOffset	= GetRandomVector(UV);
+	float2 RandomVec		= RandomAndOffset.xy;
+	float  Offset			= RandomAndOffset.z;
 	float3 ViewSpacePos		= ScreenToViewPos(UV,SceneDepth);
 
 	// Get Normal and transform into viewspace
@@ -1371,7 +1360,7 @@ void GTAOCombinedPSandCS(in float2 UV, out float4 OutColor)
 	float Scale = ScreenPixelsToCutoff /  GTAO_NUMTAPS;
 
 	// Clamp the scale. Too small and it won't be jumping by more than 1 pixel. Too much and it'll thrash the cache
-	Scale = clamp(Scale, 2.0, 8.0);
+	Scale = clamp(Scale, 2.0, 16.0);
 
 #if GTAO_REF
 	for(int i=0; i<GTAO_NUMROTS; i++)
@@ -1382,7 +1371,7 @@ void GTAOCombinedPSandCS(in float2 UV, out float4 OutColor)
 		float3 ScreenDir	= float3(RandomVec.x, RandomVec.y, 0.0);
 #endif
 
-		float2 Angles = SearchAxisForAngles( UV ,  ScreenDir,  ViewDir,  ViewSpacePos, Scale);
+		float2 Angles = SearchAxisForAngles( UV ,  ScreenDir,  ViewDir,  ViewSpacePos, Scale, Offset);
 
 		Angles.x = acosFast(Angles.x);
 		Angles.y = acosFast(Angles.y);
@@ -1392,21 +1381,20 @@ void GTAOCombinedPSandCS(in float2 UV, out float4 OutColor)
 	}
 #endif
 
-	float AmbientOcclusionIntensity = ScreenSpaceAOParams[0].w;
-	float AmbientOcclusionPower		= ScreenSpaceAOParams[0].x;
+
 	// user adjust AO
 	Sum = Sum   / (float)GTAO_NUMROTS;
-	Sum = 1 - (1 - pow(abs(Sum), AmbientOcclusionPower)) * AmbientOcclusionIntensity;
 
 	OutColor.r = Sum ;
-
+	Depth = SceneDepth;
 	return;
 }
 
 
-void GTAOCombinedPS(in noperspective float4 UVAndScreenPos : TEXCOORD0, out float4 OutColor : SV_Target0)
+void GTAOCombinedPS(in float4 UVAndScreenPos : TEXCOORD0, out float4 OutColor : SV_Target0)
 {
-	GTAOCombinedPSandCS(UVAndScreenPos.xy, OutColor);	
+	float Depth =0;
+	GTAOCombinedPSandCS(UVAndScreenPos.xy, OutColor, Depth);	
 }
 
 
@@ -1418,13 +1406,16 @@ void GTAOCombinedCS(
 	uint2 GroupThreadId : SV_GroupThreadID) 
 {
 	float4 OutColor		= 0;
+	float OutDepth 		= 10e7;
+
 	int2   PixelPos		= DispatchThreadId + ScreenSpaceAOParams[5].zw; 
 	float2 PixelCenter	= (float2)PixelPos + float2(0.5, 0.5);
 	float2 BufferUV		= PixelCenter.xy * GTAOParams[2].zw;
 
-	GTAOCombinedPSandCS(BufferUV , OutColor);
+	GTAOCombinedPSandCS(BufferUV , OutColor, OutDepth);
 
-	OutTexture[PixelPos] = OutColor.r;
+	OutTexture[PixelPos]		= OutColor.r;
+	DepthOutTexture[PixelPos]   = 1.0f/OutDepth;
 }
 #endif
 
@@ -1433,7 +1424,7 @@ void GTAOCombinedCS(
 
 /*
 *
-*	JUST INNER INTEGRATE
+*	INNER INTEGRATE
 *
 */
 
@@ -1455,7 +1446,7 @@ float GTAOInnerIntegratePSandCS(in float2 UV)
 	Angles				    = Angles *2-1;
 
 	// Get Angle
-	float2 RandomVec		= GetRandomVector(UV);
+	float2 RandomVec		= GetRandomVector(UV).xy;
 	float3 ScreenDir		= float3(RandomVec.x, RandomVec.y, 0.0);
 
 	// ViewspacePos and Normal
@@ -1469,10 +1460,6 @@ float GTAOInnerIntegratePSandCS(in float2 UV)
 
 	float AO = ComputeInnerIntegral(Angles, ScreenDir, ViewDir, ViewSpaceNormal, SceneDepth);
 
-	float AmbientOcclusionIntensity = ScreenSpaceAOParams[0].w;
-	float AmbientOcclusionPower		= ScreenSpaceAOParams[0].x;
-	// user adjust AO
-	AO = 1 - (1 - pow(abs(AO), AmbientOcclusionPower)) * AmbientOcclusionIntensity;
 
 	return AO ;
 }
@@ -1519,8 +1506,7 @@ void GTAOInnerIntegrateCS(
 
 float4 HorizonSearchParams;
 
-
-float2 HorizonSearchPSandCS(in float2 UV)
+float2 HorizonSearchPSandCS(in float2 UV, inout float Depth)
 {
 	float2 OutHorizons = 0;
 
@@ -1531,7 +1517,9 @@ float2 HorizonSearchPSandCS(in float2 UV)
 		return OutHorizons;
 	}
 
-	float2 RandomVec			= GetRandomVector(UV);
+	float3 RandomAndOffset		= GetRandomVector(UV);
+	float2 RandomVec			= RandomAndOffset.xy;
+	float  Offset				= RandomAndOffset.z;
 	float3 ScreenDir			= float3(RandomVec.x, RandomVec.y, 0.0);
 	
 	float3 ViewSpacePos			= ScreenToViewPos(UV, SceneDepth);
@@ -1543,19 +1531,21 @@ float2 HorizonSearchPSandCS(in float2 UV)
 	float Scale = ScreenPixelsToCutoff /  GTAO_NUMTAPS;
 
 	// Clamp the scale. Too small and it won't be jumping by more than 1 pixel. Too much and it'll thrash the cache
-	Scale = clamp(Scale, 2.0, 8.0);
+	Scale = clamp(Scale, 2.0, 16.0);
 
-
-	float2 Angles = SearchAxisForAngles( UV,  ScreenDir,  ViewDir,  ViewSpacePos, Scale);
+	float2 Angles = SearchAxisForAngles( UV,  ScreenDir,  ViewDir,  ViewSpacePos, Scale, Offset);
 	OutHorizons.rg = Angles*0.5+0.5;
 
+	Depth = SceneDepth;
 	return OutHorizons;
 }
 
 
 void HorizonSearchPS(in noperspective float4 UVAndScreenPos : TEXCOORD0, float4 SvPosition : SV_POSITION, out float4 OutColor : SV_Target0)
 {
-	float2 Horizons = HorizonSearchPSandCS(UVAndScreenPos.xy);	
+	float Depth;
+
+	float2 Horizons = HorizonSearchPSandCS(UVAndScreenPos.xy, Depth);	
 	OutColor = float4(Horizons.xy,0,0);
 }
 
@@ -1574,9 +1564,11 @@ void HorizonSearchCS(
 	float2 PixelCenter	= (float2)PixelPos + float2(0.5, 0.5);
 	float2 BufferUV		= PixelCenter.xy * GTAOParams[2].zw;
 
-	float2 Horizons = HorizonSearchPSandCS(BufferUV);
+	float OutDepth;
+	float2 Horizons = HorizonSearchPSandCS(BufferUV,OutDepth);
 	
 	HorizonOutTexture[PixelPos] = Horizons;
+	DepthOutTexture[PixelPos]   = OutDepth;
 }
 #endif
 
@@ -1630,7 +1622,7 @@ void GTAOTemporalFilterPSAndCS(float2 UV, out float OutAO, out float OutDepth)
 	}
 
 	// Do a neighbourhood clamp 
-	float2 dUV = PostprocessInput0Size.zw*3;
+	float2 dUV = PostprocessInput0Size.zw*2;
 	float AONeighbour1 = Texture2DSample(PostprocessInput0, PostprocessInput0Sampler, UV + float2(-dUV.x,-dUV.y) ).r;
 	float AONeighbour2 = Texture2DSample(PostprocessInput0, PostprocessInput0Sampler, UV + float2(-dUV.x, dUV.y) ).r;
 	float AONeighbour3 = Texture2DSample(PostprocessInput0, PostprocessInput0Sampler, UV + float2( dUV.x,-dUV.y) ).r;
@@ -1640,9 +1632,11 @@ void GTAOTemporalFilterPSAndCS(float2 UV, out float OutAO, out float OutDepth)
 	float MaxAO = max(NewAO, max(max(AONeighbour1, AONeighbour2), max(AONeighbour3, AONeighbour4)));
 	float HalfAO = (MaxAO - MinAO) *0.5;
 	float MidAO  = (MinAO + HalfAO);
-	const float ExpandNeighbour = 1.0;
-	MinAO = saturate(MidAO - (HalfAO * ExpandNeighbour));
-	MaxAO = saturate(MidAO + (HalfAO * ExpandNeighbour));
+
+	// Ensure the HalfAO is at least some minimum ammount. This helps with the history clamp.
+	HalfAO = max(HalfAO,0.03);
+	MinAO = saturate(MidAO - HalfAO);
+	MaxAO = saturate(MidAO + HalfAO);
 
 	float BlurredAO = 0.2 * (NewAO + AONeighbour1 + AONeighbour2 + AONeighbour3 + AONeighbour4);
 
@@ -1685,65 +1679,271 @@ void GTAOTemporalFilterCS(
 }
 #endif
 
-void GTAOSpatialFilterPS(in noperspective float4 UVAndScreenPos : TEXCOORD0,  out float4 OutColor : SV_Target0)
+
+
+
+
+
+
+float4 GTAOUpsamplePSAndCS(float2 UV)
 {
-	OutColor = Texture2DSample(PostprocessInput0, PostprocessInput0Sampler, UVAndScreenPos.xy);
+	float AO = saturate(PI_HALF * Texture2DSample(PostprocessInput0, PostprocessInput0Sampler, UV).r);
+
+	float AmbientOcclusionIntensity = ScreenSpaceAOParams[0].w;
+	float AmbientOcclusionPower		= ScreenSpaceAOParams[0].x;
+	// user adjust AO
+	AO = 1 - (1 - pow(abs(AO), AmbientOcclusionPower)) * AmbientOcclusionIntensity;
+
+	return saturate(AO);
 }
 
+void GTAOUpsamplePS(in noperspective float4 UVAndScreenPos : TEXCOORD0,  out float4 OutColor : SV_Target0)
+{
+	OutColor = GTAOUpsamplePSAndCS(UVAndScreenPos.xy);
+}
 
 #if COMPUTE_SHADER
+
+// This array will store the low res Z samples with a 1 texel border. 
+// It is half rez therefore for 8x8 inputs it is 6x6
+
+// Stores the low res samples dilated 1.
+// For 8x8 input we only need 36 entries here but we store 64 to simplify access
+groupshared float LowResZArray[  THREADGROUP_SIZEX * THREADGROUP_SIZEY]; 
+groupshared float LowResAOArray[ THREADGROUP_SIZEX * THREADGROUP_SIZEY]; 
+
+void GatherLowResAODepthSample(int2 pos, out float4 AO, out float4 Depths, out float4 BilinearWeights)
+{
+	// X and Y here are in full res. Need to gather the 4 closest low res
+	int LowResX = (pos.x+1)>>1;
+	int LowResY = (pos.y+1)>>1;
+
+	// Get the weights;
+	float WeightX = float(pos.x&1)  * 0.5 + 0.25;
+	float WeightY = float(pos.y&1)  * 0.5 + 0.25;
+	float WeightX1= 1.0f - WeightX;
+	float WeightY1= 1.0f - WeightY;
+
+	BilinearWeights = float4(	WeightX  * WeightY, 
+								WeightX1 * WeightY, 
+								WeightX  * WeightY1, 
+								WeightX1 * WeightY1);
+
+	AO		=  float4( 
+					LowResAOArray[LowResX + (LowResY*THREADGROUP_SIZEX) + 0],
+					LowResAOArray[LowResX + (LowResY*THREADGROUP_SIZEX) + 1],
+					LowResAOArray[LowResX + (LowResY*THREADGROUP_SIZEX) + THREADGROUP_SIZEX],
+					LowResAOArray[LowResX + (LowResY*THREADGROUP_SIZEX) + THREADGROUP_SIZEX + 1]);
+
+	Depths		= float4(
+					LowResZArray[LowResX + (LowResY*THREADGROUP_SIZEX) + 0],
+					LowResZArray[LowResX + (LowResY*THREADGROUP_SIZEX) + 1],
+					LowResZArray[LowResX + (LowResY*THREADGROUP_SIZEX) + THREADGROUP_SIZEX],
+					LowResZArray[LowResX + (LowResY*THREADGROUP_SIZEX) + THREADGROUP_SIZEX + 1]);
+
+}
+
+
+
 [numthreads(THREADGROUP_SIZEX, THREADGROUP_SIZEY, 1)]
-void GTAOSpatialFilterCS(
-	uint2 GroupId : SV_GroupID,
-	uint2 DispatchThreadId : SV_DispatchThreadID,
-	uint2 GroupThreadId : SV_GroupThreadID) 
+void GTAOUpsampleCS(
+	uint2 GroupId			: SV_GroupID,
+	uint2 DispatchThreadId	: SV_DispatchThreadID,
+	uint2 GroupThreadId		: SV_GroupThreadID) 
 {
-	BRANCH
-	if (any(DispatchThreadId >= (uint2)FilterParams.zw))
-	{
-		return;
-	}
+	float2 OutColor		= 0;
 
-	uint2 DestPixelPos = FilterParams.xy + DispatchThreadId;
-	float2 SampleUV = DestPixelPos * PostprocessInput0Size.zw;
+	// Get Pixel Pos of the group we are processing
+	int2   PixelPos				= DispatchThreadId + ScreenSpaceAOParams[5].zw; 
 
-	// Use a 4x4 box filter because the random texture is tiled 4x4
-	float Result;
-	Result = PostprocessInput0.SampleLevel(PostprocessInput0Sampler, SampleUV, 0).r;
-	Result += PostprocessInput0.SampleLevel(PostprocessInput0Sampler, SampleUV, 0, int2(2, 0)).r;
-	Result += PostprocessInput0.SampleLevel(PostprocessInput0Sampler, SampleUV, 0, int2(0, 2)).r;
-	Result += PostprocessInput0.SampleLevel(PostprocessInput0Sampler, SampleUV, 0, int2(2, 2)).r;
+	int2   FullGroupOrigin		= int2(GroupId.x * THREADGROUP_SIZEX, GroupId.y * THREADGROUP_SIZEY) + ScreenSpaceAOParams[5].zw; 
+	int2   HalfGroupOrigin		= int2(FullGroupOrigin.x >>1, FullGroupOrigin.y >>1);
+	int2   HalfGroupOriginM1	= HalfGroupOrigin.xy - int2(1,1);
+
+	int2   PosInGroup			= GroupThreadId.xy;	// Position with a group
+
+	// Get Full Rez Depth
+	float DeviceZ  = PostprocessInput0.Load(int3(PixelPos, 0)).r;
+	float FullResZ = ConvertFromDeviceZ(DeviceZ); 
+	
+	// Get all the low res samples
+	float HalfZ   = ZReadTexture.Load(int3(HalfGroupOriginM1 +PosInGroup.xy , 0)).r;
+	LowResZArray[PosInGroup.x + (PosInGroup.y * THREADGROUP_SIZEX)] = HalfZ;
+
+	float AO      = saturate(PostprocessInput1.Load(int3(HalfGroupOriginM1 +PosInGroup.xy , 0)).r);
+	LowResAOArray[PosInGroup.x + (PosInGroup.y * THREADGROUP_SIZEX)] = AO;
+
+	GroupMemoryBarrierWithGroupSync();
+
+	float4 AO4, Z4, BilinearWeights;
+	GatherLowResAODepthSample(PosInGroup, AO4, Z4, BilinearWeights);
+
+	float MinZ = min(min(Z4.x, Z4.y) , min(Z4.z, Z4.w)); 
+	float MaxZ = max(max(Z4.x, Z4.y) , max(Z4.z, Z4.w)); 
+	float DiffZ = 1.0f / max(MaxZ - MinZ, 0.000000001);
+
+	FullResZ = clamp(FullResZ, MinZ, MaxZ);
+
+	BilinearWeights = 1;
+
+	BilinearWeights.x *= 1.0f - (abs(Z4.x - FullResZ)*DiffZ);
+	BilinearWeights.y *= 1.0f - (abs(Z4.y - FullResZ)*DiffZ);
+	BilinearWeights.z *= 1.0f - (abs(Z4.z - FullResZ)*DiffZ);
+	BilinearWeights.w *= 1.0f - (abs(Z4.w - FullResZ)*DiffZ);
+
+	float Mag = BilinearWeights.x + BilinearWeights.y + BilinearWeights.z + BilinearWeights.w;
+
+	float AmbientOcclusionIntensity = ScreenSpaceAOParams[0].w;
+	float AmbientOcclusionPower		= ScreenSpaceAOParams[0].x*0.5;
+	// user adjust AO
+	//AO = 1 - (1 - pow(abs(AO), AmbientOcclusionPower)) * AmbientOcclusionIntensity;
+
+	// user adjust AO
+	AO = dot(AO4, BilinearWeights) / Mag;
+//	AO = lerp(1, AO, AmbientOcclusionIntensity);
+//	AO = 1 - (1 - pow(abs(AO), AmbientOcclusionPower));
+
+	OutTexture[PixelPos] = AO;
 
-	OutTexture[DestPixelPos] = Result * 0.25;
 }
 #endif
 
+void GTAOSpatialFilterPS(in noperspective float4 UVAndScreenPos : TEXCOORD0,  out float4 OutColor : SV_Target0)
+{
+	OutColor = Texture2DSample(PostprocessInput0, PostprocessInput0Sampler, UVAndScreenPos.xy);
+}
 
 
-float4 GTAOUpsamplePSAndCS(float2 UV)
+
+#if COMPUTE_SHADER
+
+#define ARRAY_SIZE ((THREADGROUP_SIZEX*2) * (THREADGROUP_SIZEY+4) )
+groupshared float AOArray[ARRAY_SIZE]; 
+groupshared float ZArray[  ARRAY_SIZE];
+groupshared float ZDiffArray[  THREADGROUP_SIZEX * THREADGROUP_SIZEY];
+
+
+void SetAOVal(float AO, int x, int y)
 {
-	return Texture2DSample(PostprocessInput0, PostprocessInput0Sampler, UV);
+	AOArray[ x + (y*THREADGROUP_SIZEX*2)] = AO;
 }
 
+void SetZVal(float Z, int x, int y)
+{
+	ZArray[ x + (y*THREADGROUP_SIZEX*2)] =Z;
+}
 
-void GTAOUpsamplePS(in noperspective float4 UVAndScreenPos : TEXCOORD0,  out float4 OutColor : SV_Target0)
+float GetAOVal(int x, int y)
 {
-	OutColor = GTAOUpsamplePSAndCS(UVAndScreenPos.xy);
+	x+=2;
+	y+=2;
+	return AOArray[ x + (y*THREADGROUP_SIZEX*2)];
 }
 
-#if COMPUTE_SHADER
+float GetZVal(int x, int y)
+{
+	x+=2;
+	y+=2;
+	return ZArray[ x + (y*THREADGROUP_SIZEX*2)];
+}
+
+
+
 [numthreads(THREADGROUP_SIZEX, THREADGROUP_SIZEY, 1)]
-void GTAOUpsampleCS(
+void GTAOSpatialFilterCS(
+	int   GroupIndex: SV_GroupIndex,
 	uint2 GroupId : SV_GroupID,
 	uint2 DispatchThreadId : SV_DispatchThreadID,
 	uint2 GroupThreadId : SV_GroupThreadID) 
 {
-	float2 OutColor		= 0;
-	int2   PixelPos		= DispatchThreadId + ScreenSpaceAOParams[5].zw; 
-	float2 PixelCenter	= (float2)PixelPos + float2(0.5, 0.5);
-	float2 BufferUV		= SvPositionToBufferUV(float4(PixelCenter, 0, 0));
+	int2 GTId = int2(GroupThreadId);
+
+	// Position on the screen
+	int2   PixelPos				= DispatchThreadId + ScreenSpaceAOParams[5].zw; 
+
+	// Origin of the 8x8 block we are filtering
+	int2   FullGroupOrigin		= int2(GroupId.x * THREADGROUP_SIZEX, GroupId.y * THREADGROUP_SIZEY) + ScreenSpaceAOParams[5].zw; 
+	int2   FullGroupOriginM2	= FullGroupOrigin.xy - int2(2,2);
+
+	// Read 2 depths per thread and place them in the array 
+	uint pixIdx = (GroupIndex*2);
+	uint XPos = pixIdx%11;
+	uint YPos = pixIdx/11;
+	SetAOVal(PostprocessInput0.Load(int3(FullGroupOriginM2 + int2(XPos,YPos), 0)).r, XPos, YPos);
+	SetZVal( ZReadTexture.Load(    int3(FullGroupOriginM2 + int2(XPos,YPos), 0)).r, XPos,YPos);
+
+	pixIdx = 1+(GroupIndex*2);
+	XPos = pixIdx%11;
+	YPos = pixIdx/11;
+	SetAOVal(PostprocessInput0.Load(int3(FullGroupOriginM2 + int2(XPos,YPos), 0)).r, XPos, YPos);
+	SetZVal( ZReadTexture.Load(    int3(FullGroupOriginM2 + int2(XPos,YPos), 0)).r, XPos,YPos);
+
+	GroupMemoryBarrierWithGroupSync();
+
+	// Get the ZDiffs array
+	float ThisZ		= GetZVal(GTId.x, GTId.y);
+
+	float XMinusZ	= GetZVal(GTId.x-1, GTId.y);
+	float XPlusZ	= GetZVal(GTId.x+1, GTId.y);
+	float YMinusZ	= GetZVal(GTId.x, GTId.y-1);
+	float YPlusZ	= GetZVal(GTId.x, GTId.y+1);
+
+	float DiffZ_X = min( abs(ThisZ - XMinusZ), abs(ThisZ - XPlusZ) );
+	float DiffZ_Y = min( abs(ThisZ - YMinusZ), abs(ThisZ - YPlusZ) );
+
+	BRANCH
+	if (any(DispatchThreadId >= (uint2)FilterParams.zw))
+	{
+		return;
+	}
+
+	int x,y;
+
+	float SumAO		= 0;
+	float SumWeight = 0;
+
+	for(y=-2; y<2; y++)
+	{
+		float YDiff = abs(y);
+
+		for(x=-2; x<2; x++)
+		{
+			// Get value and see how much it compares to the centre with the gradients
+			float XDiff = abs(x);
+
+			float Sample_AO = GetAOVal(GTId.x + x, GTId.y + y);
+			float SampleZ   = GetZVal( GTId.x + x, GTId.y + y);
+
+			float Weight = 1.0f;
+			if((x==0) && (y==0))
+			{
+				Weight = 1.0f;
+			}
+			else
+			{
+				// Compare the Z at this sample with the gradients 
+				float SampleZDiff = abs(ThisZ - SampleZ);
+				
+				SampleZDiff -= XDiff * DiffZ_X;
+				SampleZDiff -= YDiff * DiffZ_Y;
+				
+				Weight = 1.0f - saturate(SampleZDiff*1000 );
+			}
+
+			SumAO += Sample_AO * Weight;
+			SumWeight += Weight;
+		}
+	}
+	SumAO /=SumWeight;
+
+	SumAO = saturate(SumAO * PI_HALF);
+
+	float AmbientOcclusionIntensity = ScreenSpaceAOParams[0].w;
+	float AmbientOcclusionPower		= ScreenSpaceAOParams[0].x*0.5;
+	// user adjust AO
+	SumAO = 1 - (1 - pow(abs(SumAO), AmbientOcclusionPower)) * AmbientOcclusionIntensity;
+
 
-	float4 Color = GTAOUpsamplePSAndCS(BufferUV);
-	OutTexture[PixelPos] = Color;
+	OutTexture[PixelPos] = SumAO ;
 }
 #endif
\ No newline at end of file
diff --git a/Engine/Source/Runtime/Renderer/Private/CompositionLighting/CompositionLighting.cpp b/Engine/Source/Runtime/Renderer/Private/CompositionLighting/CompositionLighting.cpp
--- a/Engine/Source/Runtime/Renderer/Private/CompositionLighting/CompositionLighting.cpp
+++ b/Engine/Source/Runtime/Renderer/Private/CompositionLighting/CompositionLighting.cpp
@@ -155,6 +155,9 @@ static FRenderingCompositeOutputRef AddPostProcessingGTAOAsyncHorizonSearch(FRHI
 	}
 	GRenderTargetPool.FindFreeElement(RHICmdList, Desc, SceneContext.ScreenSpaceGTAOHorizons, TEXT("ScreenSpaceGTAOHorizons"));
 
+	Desc.Format = PF_R32_FLOAT;
+	GRenderTargetPool.FindFreeElement(RHICmdList, Desc, SceneContext.ScreenSpaceGTAODepths, TEXT("ScreenSpaceGTAODepths"));
+
 	FRenderingCompositePass* HZBInput = Context.Graph.RegisterPass(new FRCPassPostProcessInput(const_cast<FViewInfo&>(Context.View).HZB));
 	FRenderingCompositePass* AmbientOcclusionHorizonSearch = Context.Graph.RegisterPass(new(FMemStack::Get()) FRCPassPostProcessAmbientOcclusion_HorizonSearch(Context.View, DownsampleFactor, ESSAOType::EAsyncCS));
 
@@ -177,6 +180,16 @@ static FRenderingCompositeOutputRef AddPostProcessingGTAOCombined(FRHICommandLis
 	uint32 DownsampleFactor = CVarGTAODownsample.GetValueOnRenderThread() > 0 ? 2 : 1;
 	FSceneRenderTargets& SceneContext = FSceneRenderTargets::Get(RHICmdList);
 
+	FIntPoint HorizonBufferSize = FIntPoint::DivideAndRoundUp(SceneContext.GetBufferSizeXY(), DownsampleFactor);
+	{
+		FPooledRenderTargetDesc Desc(FPooledRenderTargetDesc::Create2DDesc(HorizonBufferSize, PF_R32_FLOAT, FClearValueBinding::White, TexCreate_None, TexCreate_RenderTargetable, false));
+		if (SceneContext.GetCurrentFeatureLevel() >= ERHIFeatureLevel::SM5)
+		{
+			Desc.TargetableFlags |= TexCreate_UAV;
+		}
+		GRenderTargetPool.FindFreeElement(RHICmdList, Desc, SceneContext.ScreenSpaceGTAODepths, TEXT("ScreenSpaceGTAODepths"));
+	}
+
 	if(CVarGTAOCombined.GetValueOnRenderThread() ==1)
 	{
 		FRenderingCompositePass* AmbientOcclusionGTAO = Context.Graph.RegisterPass(new(FMemStack::Get()) FRCPassPostProcessAmbientOcclusion_GTAOCombined(Context.View, DownsampleFactor, false));
@@ -186,16 +199,12 @@ static FRenderingCompositeOutputRef AddPostProcessingGTAOCombined(FRHICommandLis
 	}
 	else
 	{
-
-		FIntPoint BufferSize = SceneContext.GetBufferSizeXY();
-		FIntPoint HorizonBufferSize = FIntPoint::DivideAndRoundUp(BufferSize, DownsampleFactor);
 		FPooledRenderTargetDesc Desc(FPooledRenderTargetDesc::Create2DDesc(HorizonBufferSize, PF_R8G8, FClearValueBinding::White, TexCreate_None, TexCreate_RenderTargetable, false));
 		if (SceneContext.GetCurrentFeatureLevel() >= ERHIFeatureLevel::SM5)
 		{
 			Desc.TargetableFlags |= TexCreate_UAV;
 		}
 		GRenderTargetPool.FindFreeElement(RHICmdList, Desc, SceneContext.ScreenSpaceGTAOHorizons, TEXT("ScreenSpaceGTAOHorizons"));
-
 		FRenderingCompositePass* AmbientOcclusionHorizonSearch = Context.Graph.RegisterPass(new(FMemStack::Get()) FRCPassPostProcessAmbientOcclusion_HorizonSearch(Context.View, DownsampleFactor, ESSAOType::ECS));
 
 		AmbientOcclusionHorizonSearch->SetInput(ePId_Input0, Context.SceneDepth);
@@ -237,10 +246,13 @@ static FRenderingCompositeOutputRef AddPostProcessingGTAOCombined(FRHICommandLis
 
 	}
 	
+	// If downsampled then upsample the final result
+	if(DownsampleFactor > 1)
 	{
 		FRenderingCompositePass* UpsamplePass;
 		UpsamplePass = Context.Graph.RegisterPass(new (FMemStack::Get()) FRCPassPostProcessAmbientOcclusion_GTAO_Upsample(Context.View, DownsampleFactor));
-		UpsamplePass->SetInput(ePId_Input0, FinalOutputPass);
+		UpsamplePass->SetInput(ePId_Input0, Context.SceneDepth);
+		UpsamplePass->SetInput(ePId_Input1, FRenderingCompositeOutputRef(FinalOutputPass, ePId_Output0));
 		FinalOutputPass = UpsamplePass;
 	}
 
@@ -292,7 +304,8 @@ static FRenderingCompositeOutputRef AddPostProcessingGTAOIntegration(FRHICommand
 	{
 		FRenderingCompositePass* UpsamplePass;
 		UpsamplePass = Context.Graph.RegisterPass(new (FMemStack::Get()) FRCPassPostProcessAmbientOcclusion_GTAO_Upsample(Context.View, DownsampleFactor));
-		UpsamplePass->SetInput(ePId_Input0, FinalOutputPass);
+		UpsamplePass->SetInput(ePId_Input0, Context.SceneDepth);
+		UpsamplePass->SetInput(ePId_Input1, FRenderingCompositeOutputRef(FinalOutputPass, ePId_Output0));
 		FinalOutputPass = UpsamplePass;
 	}
 
@@ -570,6 +583,7 @@ void FCompositionLighting::ProcessAfterBasePass(FRHICommandListImmediate& RHICmd
 	}
 
 	SceneContext.ScreenSpaceGTAOHorizons.SafeRelease();
+	SceneContext.ScreenSpaceGTAODepths.SafeRelease();
 }
 
 
diff --git a/Engine/Source/Runtime/Renderer/Private/CompositionLighting/PostProcessAmbientOcclusion.cpp b/Engine/Source/Runtime/Renderer/Private/CompositionLighting/PostProcessAmbientOcclusion.cpp
--- a/Engine/Source/Runtime/Renderer/Private/CompositionLighting/PostProcessAmbientOcclusion.cpp
+++ b/Engine/Source/Runtime/Renderer/Private/CompositionLighting/PostProcessAmbientOcclusion.cpp
@@ -1218,12 +1218,14 @@ public:
 		const FFinalPostProcessSettings& Settings = View.FinalPostProcessSettings;
 
 		uint32 TemporalFrame = 0;
+		uint32 Frame = 0;
 		
 		const FSceneViewState* ViewState = static_cast<const FSceneViewState*>(View.State);
 
 		if (ViewState)
 		{
-			TemporalFrame	= ViewState->GetCurrentUnclampedTemporalAASampleIndex();
+			TemporalFrame = ViewState->GetCurrentUnclampedTemporalAASampleIndex();
+			Frame = ViewState->GetFrameIndex();
 		}
 			
 		const int ArraySize = 3;
@@ -1239,7 +1241,7 @@ public:
 
 		FIntPoint RandomizationSize = GSystemTextures.GTAORandomization->GetDesc().Extent;
 
-		GTAOParam[1] = FVector4(0.0f, 0.0f, 0.0f, 0.0f);
+		GTAOParam[1] = FVector4(Frame, 0.0f, 0.0f, 0.0f);
 
 		float Fx = float(DestSize.X);
 		float Fy = float(DestSize.Y);
@@ -1301,6 +1303,7 @@ public:
 	FScreenSpaceAOParameters	  ScreenSpaceAOParams;
 	FGTAOParameters				  GTAOParams;
 	FShaderResourceParameter 	  HorizonOutTexture;
+	FShaderResourceParameter 	  DepthOutTexture;
 	FShaderResourceParameter      RandomNormalTexture;
 	FShaderResourceParameter      RandomNormalTextureSampler;
 	FShaderParameter			  HZBRemapping;
@@ -1322,6 +1325,7 @@ public:
 		if (bComputeShader)
 		{
 			HorizonOutTexture.Bind(Initializer.ParameterMap, TEXT("HorizonOutTexture"));
+			DepthOutTexture.Bind(Initializer.ParameterMap, TEXT("DepthOutTexture"));
 		}
 	}
 
@@ -1334,7 +1338,7 @@ public:
 	}
 
 	template <typename TRHICmdList>
-	void SetParametersCS(TRHICmdList& RHICmdList, const FRenderingCompositePassContext& Context, FIntPoint DestSize, FIntPoint InputTextureSize, uint32 DownScaleFactor, FRHIUnorderedAccessView *OutUAV)
+	void SetParametersCS(TRHICmdList& RHICmdList, const FRenderingCompositePassContext& Context, FIntPoint DestSize, FIntPoint InputTextureSize, uint32 DownScaleFactor, FRHIUnorderedAccessView *OutUAV, FRHIUnorderedAccessView *OutDepthsUAV)
 	{
 		const FFinalPostProcessSettings& Settings = Context.View.FinalPostProcessSettings;
 		FRHIComputeShader* ShaderRHI = GetComputeShader();
@@ -1348,6 +1352,7 @@ public:
 		PostprocessParameter.SetCS(ShaderRHI, Context, RHICmdList, TStaticSamplerState<SF_Point, AM_Clamp, AM_Clamp, AM_Clamp>::GetRHI());
 
 		RHICmdList.SetUAVParameter(ShaderRHI, HorizonOutTexture.GetBaseIndex(), OutUAV);
+		RHICmdList.SetUAVParameter(ShaderRHI, DepthOutTexture.GetBaseIndex(), OutDepthsUAV);
 
 		FVector4 HZBRemappingValue = GetHZBRemapVal(Context, DestSize, InputTextureSize);
 		SetShaderValue(RHICmdList, ShaderRHI, HZBRemapping, HZBRemappingValue);
@@ -1384,6 +1389,7 @@ public:
 	{
 		FRHIComputeShader* ShaderRHI = GetComputeShader();
 		RHICmdList.SetUAVParameter(ShaderRHI, HorizonOutTexture.GetBaseIndex(), nullptr);
+		RHICmdList.SetUAVParameter(ShaderRHI, DepthOutTexture.GetBaseIndex(), nullptr);
 	}
 
 
@@ -1391,7 +1397,7 @@ public:
 	virtual bool Serialize(FArchive& Ar) override
 	{
 		bool bShaderHasOutdatedParameters = FGlobalShader::Serialize(Ar);
-		Ar << PostprocessParameter << SceneTextureParameters << ScreenSpaceAOParams << GTAOParams << HorizonOutTexture << HZBRemapping << RandomNormalTexture << RandomNormalTextureSampler << HorizonSearchParams;
+		Ar << PostprocessParameter << SceneTextureParameters << ScreenSpaceAOParams << GTAOParams << HorizonOutTexture << DepthOutTexture << HZBRemapping << RandomNormalTexture << RandomNormalTextureSampler << HorizonSearchParams;
 		return bShaderHasOutdatedParameters;
 	}
 
@@ -1455,6 +1461,7 @@ public:
 	FPostProcessPassParameters	  PostprocessParameter;
 	FScreenSpaceAOParameters	  ScreenSpaceAOParams;
 	FShaderResourceParameter 	  OutTexture;
+	FShaderResourceParameter 	  DepthOutTexture;
 	FShaderParameter			  HZBRemapping;
 	FShaderResourceParameter      RandomNormalTexture;
 	FShaderResourceParameter      RandomNormalTextureSampler;
@@ -1475,6 +1482,7 @@ public:
 		if (bComputeShader)
 		{
 			OutTexture.Bind(Initializer.ParameterMap, TEXT("OutTexture"));
+			DepthOutTexture.Bind(Initializer.ParameterMap, TEXT("DepthOutTexture"));
 		}
 	}
 
@@ -1489,7 +1497,8 @@ public:
 
 
 	template <typename TRHICmdList>
-	void SetParametersCS(TRHICmdList& RHICmdList, const FRenderingCompositePassContext& Context, FIntPoint DestSize, FIntPoint InputTextureSize, uint32 DownScaleFactor, FRHIUnorderedAccessView *OutUAV)
+	void SetParametersCS(TRHICmdList& RHICmdList, const FRenderingCompositePassContext& Context, FIntPoint DestSize, FIntPoint InputTextureSize, uint32 DownScaleFactor,
+		FRHIUnorderedAccessView *OutUAV, FRHIUnorderedAccessView *DepthOutUAV)
 	{
 		const FFinalPostProcessSettings& Settings = Context.View.FinalPostProcessSettings;
 		FRHIComputeShader* ShaderRHI = GetComputeShader();
@@ -1502,6 +1511,7 @@ public:
 
 		GTAOParams.Set(RHICmdList, Context.View, DestSize, ShaderRHI);
 		RHICmdList.SetUAVParameter(ShaderRHI, OutTexture.GetBaseIndex(), OutUAV);
+		RHICmdList.SetUAVParameter(ShaderRHI, DepthOutTexture.GetBaseIndex(), DepthOutUAV);
 
 		FVector4 HZBRemappingValue = GetHZBRemapVal(Context, DestSize, InputTextureSize);
 		SetShaderValue(Context.RHICmdList, ShaderRHI, HZBRemapping, HZBRemappingValue);
@@ -1535,6 +1545,7 @@ public:
 	{
 		FRHIComputeShader* ShaderRHI = GetComputeShader();
 		RHICmdList.SetUAVParameter(ShaderRHI, OutTexture.GetBaseIndex(), nullptr);
+		RHICmdList.SetUAVParameter(ShaderRHI, DepthOutTexture.GetBaseIndex(), nullptr);
 	}
 
 
@@ -1542,7 +1553,7 @@ public:
 	virtual bool Serialize(FArchive& Ar) override
 	{
 		bool bShaderHasOutdatedParameters = FGlobalShader::Serialize(Ar);
-		Ar << PostprocessParameter << SceneTextureParameters << ScreenSpaceAOParams << OutTexture << GTAOParams <<HZBRemapping << RandomNormalTexture << RandomNormalTextureSampler;
+		Ar << PostprocessParameter << SceneTextureParameters << ScreenSpaceAOParams << OutTexture << DepthOutTexture << GTAOParams <<HZBRemapping << RandomNormalTexture << RandomNormalTextureSampler;
 		return bShaderHasOutdatedParameters;
 	}
 
@@ -1726,9 +1737,10 @@ void FRCPassPostProcessAmbientOcclusion_GTAOCombined::DispatchCS(const FRenderin
 	TShaderMapRef<FPostProcessGTAOCombinedPSandCS<true, ShaderQuality, UseNormals>> ComputeShader(Context.GetShaderMap());
 	FSceneRenderTargets& SceneContext = FSceneRenderTargets::Get(Context.RHICmdList);
 	const FSceneRenderTargetItem& DestRenderTarget = PassOutputs[0].RequestSurface(Context);
+	const FSceneRenderTargetItem& DestRenderTarget1 = SceneContext.ScreenSpaceGTAODepths->GetRenderTargetItem();
 
 	Context.RHICmdList.SetComputeShader(ComputeShader->GetComputeShader());
-	ComputeShader->SetParametersCS(Context.RHICmdList, Context, DestSize, TexSize, DownScaleFactor, DestRenderTarget.UAV);
+	ComputeShader->SetParametersCS(Context.RHICmdList, Context, DestSize, TexSize, DownScaleFactor, DestRenderTarget.UAV, DestRenderTarget1.UAV);
 
 	uint32 GroupSizeX = FMath::DivideAndRoundUp(ViewRect.Width(), 8);
 	uint32 GroupSizeY = FMath::DivideAndRoundUp(ViewRect.Height(), 8);
@@ -1995,13 +2007,14 @@ template <uint32 ShaderQuality>
 void FRCPassPostProcessAmbientOcclusion_HorizonSearch::DispatchCS(const FRenderingCompositePassContext& Context, FIntRect ViewRect, FIntPoint DestSize, FIntPoint TexSize)
 {
 	FSceneRenderTargets& SceneContext = FSceneRenderTargets::Get(Context.RHICmdList);
-	const FSceneRenderTargetItem& DestRenderTarget = SceneContext.ScreenSpaceGTAOHorizons->GetRenderTargetItem();
+	const FSceneRenderTargetItem& DestRenderTarget  = SceneContext.ScreenSpaceGTAOHorizons->GetRenderTargetItem();
+	const FSceneRenderTargetItem& DestRenderTarget2 = SceneContext.ScreenSpaceGTAODepths->GetRenderTargetItem();
 
 	FRHIAsyncComputeCommandListImmediate& RHICmdListComputeImmediate = FRHICommandListExecutor::GetImmediateAsyncComputeCommandList();
 
 	TShaderMapRef<FPostProcessGTAOHorizonSearchPSandCS<true, ShaderQuality>> ComputeShader(Context.GetShaderMap());
 	RHICmdListComputeImmediate.SetComputeShader(ComputeShader->GetComputeShader());
-	ComputeShader->SetParametersCS(RHICmdListComputeImmediate, Context, DestSize, TexSize, DownScaleFactor, DestRenderTarget.UAV);
+	ComputeShader->SetParametersCS(RHICmdListComputeImmediate, Context, DestSize, TexSize, DownScaleFactor, DestRenderTarget.UAV, DestRenderTarget2.UAV);
 
 	uint32 GroupSizeX = FMath::DivideAndRoundUp(ViewRect.Width(), 8);
 	uint32 GroupSizeY = FMath::DivideAndRoundUp(ViewRect.Height(), 8);
@@ -2322,27 +2335,6 @@ FRCPassPostProcessAmbientOcclusion_GTAO_TemporalFilter::FRCPassPostProcessAmbien
 {
 }
 
-#if 0
-template <uint32 ShaderQuality>
-void FRCPassPostProcessAmbientOcclusion_GTAO_TemporalFilter::DispatchCS(const FRenderingCompositePassContext& Context, FIntRect OutputViewRect, FIntPoint OutputTexSize)
-{
-	FSceneRenderTargets& SceneContext = FSceneRenderTargets::Get(Context.RHICmdList);
-	const FSceneRenderTargetItem& DestRenderTarget = PassOutputs[0].RequestSurface(Context);
-
-	FRHIAsyncComputeCommandListImmediate& RHICmdListComputeImmediate = FRHICommandListExecutor::GetImmediateAsyncComputeCommandList();
-
-	TShaderMapRef<FPostProcessGTAOTemporalFilterPSandCS<true>> ComputeShader(Context.GetShaderMap());
-	RHICmdListComputeImmediate.SetComputeShader(ComputeShader->GetComputeShader());
-	ComputeShader->SetParametersCS(RHICmdListComputeImmediate, Context, DestRenderTarget.UAV, InputHistory);
-	dawd
-	uint32 GroupSizeX = FMath::DivideAndRoundUp(OutputViewRect.Max.X, 8);
-	uint32 GroupSizeY = FMath::DivideAndRoundUp(OutputViewRect.Max.Y, 8);
-	DispatchComputeShader(RHICmdListComputeImmediate, *ComputeShader, GroupSizeX, GroupSizeY, 1);
-
-	ComputeShader->UnsetParameters(RHICmdListComputeImmediate);
-}
-#endif
-
 void FRCPassPostProcessAmbientOcclusion_GTAO_TemporalFilter::Process(FRenderingCompositePassContext& Context)
 {
 	SCOPED_GPU_STAT(Context.RHICmdList, GTAO_TemporalFilter);
@@ -2351,7 +2343,8 @@ void FRCPassPostProcessAmbientOcclusion_GTAO_TemporalFilter::Process(FRenderingC
 	const FPooledRenderTargetDesc* InputDesc0 = GetInputDesc(ePId_Input0);
 	FSceneRenderTargets& SceneContext = FSceneRenderTargets::Get(Context.RHICmdList);
 
-	const FSceneRenderTargetItem& DestRenderTarget0 = PassOutputs[0].RequestSurface(Context);
+
+	const FSceneRenderTargetItem& DestRenderTarget0 = (DownScaleFactor==1) ? SceneContext.ScreenSpaceAO->GetRenderTargetItem() : PassOutputs[0].RequestSurface(Context);
 	const FSceneRenderTargetItem& DestRenderTarget1 = PassOutputs[1].RequestSurface(Context);
 
 	FIntPoint InputTexSize  = InputDesc0->Extent;
@@ -2481,6 +2474,9 @@ class FPostProcessGTAOSpatialFilterPSandCS : public FGlobalShader
 	FPostProcessPassParameters	PostprocessParameter;
 	FShaderParameter			OutTexture;
 	FShaderParameter			FilterParams;
+	FScreenSpaceAOParameters	ScreenSpaceAOParams;
+	FShaderResourceParameter	ZReadTexture;
+	FShaderResourceParameter    ZReadTextureSampler;
 
 public:
 
@@ -2491,13 +2487,16 @@ public:
 		PostprocessParameter.Bind(Initializer.ParameterMap);
 		OutTexture.Bind(Initializer.ParameterMap, TEXT("OutTexture"));
 		FilterParams.Bind(Initializer.ParameterMap, TEXT("FilterParams"));
+		ScreenSpaceAOParams.Bind(Initializer.ParameterMap);
+		ZReadTexture.Bind(Initializer.ParameterMap,			TEXT("ZReadTexture"));
+		ZReadTextureSampler.Bind(Initializer.ParameterMap,	TEXT("ZReadTextureSampler"));
 	}
 
 	// FShader interface.
 	virtual bool Serialize(FArchive& Ar) override
 	{
 		bool bShaderHasOutdatedParameters = FGlobalShader::Serialize(Ar);
-		Ar << PostprocessParameter << OutTexture << FilterParams;
+		Ar << PostprocessParameter << OutTexture << FilterParams << ScreenSpaceAOParams << ZReadTexture << ZReadTextureSampler;
 		return bShaderHasOutdatedParameters;
 	}
 
@@ -2507,8 +2506,11 @@ public:
 		FRHIPixelShader* ShaderRHI = GetPixelShader();
 
 		FGlobalShader::SetParameters<FViewUniformShaderParameters>(Context.RHICmdList, ShaderRHI, Context.View.ViewUniformBuffer);
-
+		ScreenSpaceAOParams.Set(RHICmdList, Context.View, ShaderRHI, InputTextureSize);
 		PostprocessParameter.SetPS(Context.RHICmdList, ShaderRHI, Context, TStaticSamplerState<SF_Bilinear, AM_Clamp, AM_Clamp, AM_Clamp>::GetRHI());
+	
+		SetTextureParameter(RHICmdList, ShaderRHI, ZReadTexture, ZReadTextureSampler, TStaticSamplerState<SF_Point, AM_Wrap, AM_Wrap, AM_Wrap>::GetRHI(),
+			SceneContext.ScreenSpaceGTAODepths->GetRenderTargetItem().ShaderResourceTexture);
 	}
 
 
@@ -2520,8 +2522,13 @@ public:
 
 		FGlobalShader::SetParameters<FViewUniformShaderParameters>(RHICmdList, ShaderRHI, Context.View.ViewUniformBuffer);
 		PostprocessParameter.SetCS(ShaderRHI, Context, RHICmdList, TStaticSamplerState<SF_Bilinear, AM_Clamp, AM_Clamp, AM_Clamp>::GetRHI());
+		ScreenSpaceAOParams.Set(RHICmdList, Context.View, ShaderRHI, InputTextureSize);
 
 		RHICmdList.SetUAVParameter(ShaderRHI, OutTexture.GetBaseIndex(), OutUAV);
+		FSceneRenderTargets& SceneContext = FSceneRenderTargets::Get(Context.RHICmdList);
+
+		SetTextureParameter(RHICmdList, ShaderRHI, ZReadTexture, ZReadTextureSampler, TStaticSamplerState<SF_Point, AM_Wrap, AM_Wrap, AM_Wrap>::GetRHI(),
+			SceneContext.ScreenSpaceGTAODepths->GetRenderTargetItem().ShaderResourceTexture);
 
 		FVector4 FilterParamsValue(OutputRect.Min.X, OutputRect.Min.Y, OutputRect.Width(), OutputRect.Height());
 		SetShaderValue(RHICmdList, ShaderRHI, FilterParams, FilterParamsValue);
@@ -2629,9 +2636,11 @@ class FPostProcessGTAOUpsamplePSandCS : public FGlobalShader
 	}
 
 	FPostProcessGTAOUpsamplePSandCS() {}
-	FPostProcessPassParameters	  PostprocessParameter;
-	FShaderParameter OutTexture;
-	FScreenSpaceAOParameters ScreenSpaceAOParams;
+	FPostProcessPassParameters		PostprocessParameter;
+	FShaderParameter				OutTexture;
+	FScreenSpaceAOParameters		ScreenSpaceAOParams;
+	FShaderResourceParameter	ZReadTexture;
+	FShaderResourceParameter    ZReadTextureSampler;
 
 public:
 
@@ -2642,26 +2651,62 @@ public:
 		PostprocessParameter.Bind(Initializer.ParameterMap);
 		OutTexture.Bind(Initializer.ParameterMap, TEXT("OutTexture"));
 		ScreenSpaceAOParams.Bind(Initializer.ParameterMap);
+		ZReadTexture.Bind(Initializer.ParameterMap, TEXT("ZReadTexture"));
+		ZReadTextureSampler.Bind(Initializer.ParameterMap, TEXT("ZReadTextureSampler"));
 	}
 
 	// FShader interface.
 	virtual bool Serialize(FArchive& Ar) override
 	{
 		bool bShaderHasOutdatedParameters = FGlobalShader::Serialize(Ar);
-		Ar << PostprocessParameter << OutTexture << ScreenSpaceAOParams;
+		Ar << PostprocessParameter << OutTexture << ScreenSpaceAOParams << ZReadTexture << ZReadTextureSampler;
 		return bShaderHasOutdatedParameters;
 	}
 
-	void SetParametersPS(const FRenderingCompositePassContext& Context)
+	void SetParametersPS(const FRenderingCompositePassContext& Context, FIntPoint InputTextureSize)
 	{
 		const FFinalPostProcessSettings& Settings = Context.View.FinalPostProcessSettings;
 		FRHIPixelShader* ShaderRHI = GetPixelShader();
 
 		FGlobalShader::SetParameters<FViewUniformShaderParameters>(Context.RHICmdList, ShaderRHI, Context.View.ViewUniformBuffer);
+		ScreenSpaceAOParams.Set(Context.RHICmdList, Context.View, ShaderRHI, InputTextureSize);
 
 		PostprocessParameter.SetPS(Context.RHICmdList, ShaderRHI, Context, TStaticSamplerState<SF_Bilinear, AM_Clamp, AM_Clamp, AM_Clamp>::GetRHI());
+
+		FSceneRenderTargets& SceneContext = FSceneRenderTargets::Get(Context.RHICmdList);
+
+		SetTextureParameter(RHICmdList, ShaderRHI, ZReadTexture, ZReadTextureSampler, TStaticSamplerState<SF_Point, AM_Wrap, AM_Wrap, AM_Wrap>::GetRHI(),
+			SceneContext.ScreenSpaceGTAODepths->GetRenderTargetItem().ShaderResourceTexture);
+
 	}
 
+	template <typename TRHICmdList>
+	void SetParametersCS(TRHICmdList& RHICmdList, const FRenderingCompositePassContext& Context, FIntPoint InputTextureSize, FRHIUnorderedAccessView *OutUAV)
+	{
+		const FFinalPostProcessSettings& Settings = Context.View.FinalPostProcessSettings;
+		FRHIComputeShader* ShaderRHI = GetComputeShader();
+
+		FGlobalShader::SetParameters<FViewUniformShaderParameters>(RHICmdList, ShaderRHI, Context.View.ViewUniformBuffer);
+		ScreenSpaceAOParams.Set(RHICmdList, Context.View, ShaderRHI, InputTextureSize);
+
+		PostprocessParameter.SetCS(ShaderRHI, Context, RHICmdList, TStaticSamplerState<SF_Point, AM_Clamp, AM_Clamp, AM_Clamp>::GetRHI());
+
+		RHICmdList.SetUAVParameter(ShaderRHI, OutTexture.GetBaseIndex(), OutUAV);
+	
+		FSceneRenderTargets& SceneContext = FSceneRenderTargets::Get(Context.RHICmdList);
+
+		SetTextureParameter(RHICmdList, ShaderRHI, ZReadTexture, ZReadTextureSampler, TStaticSamplerState<SF_Point, AM_Wrap, AM_Wrap, AM_Wrap>::GetRHI(),
+			SceneContext.ScreenSpaceGTAODepths->GetRenderTargetItem().ShaderResourceTexture);
+	}
+
+	template <typename TRHICmdList>
+	void UnsetParameters(TRHICmdList& RHICmdList)
+	{
+		FRHIComputeShader* ShaderRHI = GetComputeShader();
+		RHICmdList.SetUAVParameter(ShaderRHI, OutTexture.GetBaseIndex(),nullptr);
+	}
+	
+
 	static const TCHAR* GetSourceFilename()
 	{
 		return TEXT("/Engine/Private/PostProcessAmbientOcclusion.usf");
@@ -2690,53 +2735,23 @@ void FRCPassPostProcessAmbientOcclusion_GTAO_Upsample::Process(FRenderingComposi
 	FSceneRenderTargets& SceneContext = FSceneRenderTargets::Get(Context.RHICmdList);
 	const FSceneRenderTargetItem& DestRenderTarget = SceneContext.ScreenSpaceAO->GetRenderTargetItem();
 
-	// Do a simple upsample filter from source to dest
-	FIntPoint InputTexSize = InputDesc0->Extent;
+	FIntPoint InputTexSize = InputDesc0->Extent;	// This is the full res Z
 	FIntPoint OutputTexSize = SceneContext.GetBufferSizeXY();
-	FIntRect OutputFullRect(FIntPoint(0, 0), OutputTexSize);
-
-	FIntRect InputViewRect = View.ViewRect / DownScaleFactor;
-	FIntRect OutputViewRect = View.ViewRect;
-
-	Context.RHICmdList.TransitionResource(EResourceTransitionAccess::EWritable, DestRenderTarget.TargetableTexture);
-
-	FRHIRenderPassInfo RPInfo(DestRenderTarget.TargetableTexture, ERenderTargetActions::Load_Store);
-	Context.RHICmdList.BeginRenderPass(RPInfo, TEXT("GTAO_Filter"));
-	{
-		Context.SetViewportAndCallRHI(OutputFullRect);
-
-		FGraphicsPipelineStateInitializer GraphicsPSOInit;
-		Context.RHICmdList.ApplyCachedRenderTargets(GraphicsPSOInit);
 
-		TShaderMapRef<FPostProcessVS> VertexShader(Context.GetShaderMap());
-		TShaderMapRef<FPostProcessGTAOUpsamplePSandCS<false>> PixelShader(Context.GetShaderMap());
-
-		// set the state
-		GraphicsPSOInit.BlendState = TStaticBlendState<>::GetRHI();
-		GraphicsPSOInit.RasterizerState = TStaticRasterizerState<>::GetRHI();
-		GraphicsPSOInit.DepthStencilState = TStaticDepthStencilState<false, CF_Always>::GetRHI();
-		GraphicsPSOInit.PrimitiveType = PT_TriangleList;
-		GraphicsPSOInit.bDepthBounds = false;
+	// Compute  Version
+	Context.RHICmdList.TransitionResource(EResourceTransitionAccess::ERWBarrier, EResourceTransitionPipeline::EGfxToCompute, DestRenderTarget.UAV, nullptr);
 
-		GraphicsPSOInit.BoundShaderState.VertexDeclarationRHI = GFilterVertexDeclaration.VertexDeclarationRHI;
-		GraphicsPSOInit.BoundShaderState.VertexShaderRHI = GETSAFERHISHADER_VERTEX(*VertexShader);
-		GraphicsPSOInit.BoundShaderState.PixelShaderRHI = GETSAFERHISHADER_PIXEL(*PixelShader);
-		SetGraphicsPipelineState(Context.RHICmdList, GraphicsPSOInit);
+	UnbindRenderTargets(Context.RHICmdList);
 
-		VertexShader->SetParameters(Context);
-		PixelShader->SetParametersPS(Context);
+	TShaderMapRef<FPostProcessGTAOUpsamplePSandCS<true>> ComputeShader(Context.GetShaderMap());
+	Context.RHICmdList.SetComputeShader(ComputeShader->GetComputeShader());
+	ComputeShader->SetParametersCS(Context.RHICmdList, Context, InputTexSize, DestRenderTarget.UAV);
 
-		DrawRectangle(
-			Context.RHICmdList,
-			OutputViewRect.Min.X, OutputViewRect.Min.Y, OutputViewRect.Width(), OutputViewRect.Height(),
-			InputViewRect.Min.X, InputViewRect.Min.Y, InputViewRect.Width(), InputViewRect.Height(),
-			OutputTexSize,
-			InputTexSize,
-			*VertexShader,
-			EDRF_UseTriangleOptimization);
+	uint32 GroupSizeX = FMath::DivideAndRoundUp(View.ViewRect.Width(), 8);
+	uint32 GroupSizeY = FMath::DivideAndRoundUp(View.ViewRect.Height(), 8);
+	DispatchComputeShader(Context.RHICmdList, *ComputeShader, GroupSizeX, GroupSizeY, 1);
 
-	}
-	Context.RHICmdList.EndRenderPass();
+	ComputeShader->UnsetParameters(Context.RHICmdList);
 
 	Context.RHICmdList.TransitionResource(EResourceTransitionAccess::EReadable, DestRenderTarget.TargetableTexture);
 }
diff --git a/Engine/Source/Runtime/Renderer/Private/CompositionLighting/PostProcessAmbientOcclusion.h b/Engine/Source/Runtime/Renderer/Private/CompositionLighting/PostProcessAmbientOcclusion.h
--- a/Engine/Source/Runtime/Renderer/Private/CompositionLighting/PostProcessAmbientOcclusion.h
+++ b/Engine/Source/Runtime/Renderer/Private/CompositionLighting/PostProcessAmbientOcclusion.h
@@ -144,7 +144,7 @@ private:
 };
 
 
-class FRCPassPostProcessAmbientOcclusion_HorizonSearch : public TRenderingCompositePassBase<2, 1>
+class FRCPassPostProcessAmbientOcclusion_HorizonSearch : public TRenderingCompositePassBase<2, 2>
 {
 public:
 
@@ -184,7 +184,7 @@ private:
 };
 
 
-class FRCPassPostProcessAmbientOcclusion_GTAOCombined : public TRenderingCompositePassBase<2, 1>
+class FRCPassPostProcessAmbientOcclusion_GTAOCombined : public TRenderingCompositePassBase<2, 2>
 {
 public:
 
@@ -246,7 +246,7 @@ private:
 
 };
 
-class FRCPassPostProcessAmbientOcclusion_GTAO_Upsample : public TRenderingCompositePassBase<1, 1>
+class FRCPassPostProcessAmbientOcclusion_GTAO_Upsample : public TRenderingCompositePassBase<2, 1>
 {
 public:
 
diff --git a/Engine/Source/Runtime/Renderer/Private/PostProcess/SceneRenderTargets.h b/Engine/Source/Runtime/Renderer/Private/PostProcess/SceneRenderTargets.h
--- a/Engine/Source/Runtime/Renderer/Private/PostProcess/SceneRenderTargets.h
+++ b/Engine/Source/Runtime/Renderer/Private/PostProcess/SceneRenderTargets.h
@@ -549,6 +549,7 @@ public:
 	// for AmbientOcclusion, only valid for a short time during the frame to allow reuse
 	TRefCountPtr<IPooledRenderTarget> ScreenSpaceAO;
 	TRefCountPtr<IPooledRenderTarget> ScreenSpaceGTAOHorizons;
+	TRefCountPtr<IPooledRenderTarget> ScreenSpaceGTAODepths;
 	// for shader/quad complexity, the temporary quad descriptors and complexity.
 	TRefCountPtr<IPooledRenderTarget> QuadOverdrawBuffer;
 	// used by the CustomDepth material feature, is allocated on demand or if r.CustomDepth is 2
diff --git a/Engine/Source/Runtime/Renderer/Private/SystemTextures.cpp b/Engine/Source/Runtime/Renderer/Private/SystemTextures.cpp
--- a/Engine/Source/Runtime/Renderer/Private/SystemTextures.cpp
+++ b/Engine/Source/Runtime/Renderer/Private/SystemTextures.cpp
@@ -379,7 +379,7 @@ void FSystemTextures::InitializeFeatureLevelDependentTextures(FRHICommandListImm
 			}
 	
 			{
-				FPooledRenderTargetDesc Desc(FPooledRenderTargetDesc::Create2DDesc(FIntPoint(4, 4), PF_R8G8, FClearValueBinding::None, TexCreate_HideInVisualizeTexture, TexCreate_None | TexCreate_NoFastClear, false));
+				FPooledRenderTargetDesc Desc(FPooledRenderTargetDesc::Create2DDesc(FIntPoint(4, 4), PF_R8G8B8A8, FClearValueBinding::None, TexCreate_HideInVisualizeTexture, TexCreate_None | TexCreate_NoFastClear, false));
 				Desc.AutoWritable = false;
 				GRenderTargetPool.FindFreeElement(RHICmdList, Desc,GTAORandomization, TEXT("GTAORandomization"), true, ERenderTargetTransience::NonTransient);
 				// Write the contents of the texture.
@@ -390,16 +390,18 @@ void FSystemTextures::InitializeFeatureLevelDependentTextures(FRHICommandListImm
 				{
 					for(int32 x = 0; x < Desc.Extent.X; ++x)
 					{
-						uint8* Dest = (uint8*)(DestBuffer + x * sizeof(uint16) + y * DestStride);
+						uint8* Dest = (uint8*)(DestBuffer + x * sizeof(uint32) + y * DestStride);
 			
 						float Angle  = (PI/16.0f)  * ((((x + y) & 0x3) << 2) + (x & 0x3));
-						float Step	 = 0.5f + ((1.0f/8.0f) * ((y - x) & 0x3));
+						float Step   = (1.0f / 4.0f) * ((y - x) & 0x3);
 
-						float ScaleCos = FMath::Cos(Angle) ;
-						float ScaleSin = FMath::Sin(Angle) ;
+						float ScaleCos = FMath::Cos(Angle);
+						float ScaleSin = FMath::Sin(Angle);
 			
 						Dest[0] = (uint8_t)(ScaleCos*127.5f + 127.5f);
 						Dest[1] = (uint8_t)(ScaleSin*127.5f + 127.5f);
+						Dest[2] = (uint8_t)(Step*255.0f);
+						Dest[3] = 0;
 					}
 				}
 			}
