- Fix for blinking on a Niagara Mesh Particle System on D3D11 only
- This looks like a driver / compiler bug to me and this works around the issue

diff --git a/Engine/Shaders/Private/AtmosphereCommon.ush b/Engine/Shaders/Private/AtmosphereCommon.ush
--- a/Engine/Shaders/Private/AtmosphereCommon.ush
+++ b/Engine/Shaders/Private/AtmosphereCommon.ush
@@ -570,6 +570,7 @@ float4 GetAtmosphericFog(float3 ViewPosition, float3 ViewVector, float SceneDept
 	float ShadowFactor = 1.f; // shadow approximation
 	float DistanceRatio = min(FogDepth * 0.1f / View.AtmosphericFogStartDistance, 1.f);
 	bool bIsSceneGeometry = (SceneDepth < DepthThreshold); // Assume as scene geometry
+	BRANCH
 	if (bIsSceneGeometry)
 	{
 		ShadowFactor = DistanceRatio * View.AtmosphericFogPower;
