Fix mip sky light pdf texture and solid angle texture defines set in the incorrect location for sky light direction sampling.

(ushell-p4-cherrypick of 9659095 by riley.alston)

diff --git a/Engine/Shaders/Private/RayTracing/RayTracingSkyLightCommon.ush b/Engine/Shaders/Private/RayTracing/RayTracingSkyLightCommon.ush
--- a/Engine/Shaders/Private/RayTracing/RayTracingSkyLightCommon.ush
+++ b/Engine/Shaders/Private/RayTracing/RayTracingSkyLightCommon.ush
@@ -4,6 +4,9 @@
 	SkyLightMipTreeCommon.ush: Common utilities for MipTree-based SkyLight sampling
 ===============================================================================================*/
 
+#define USE_MIP_PDF_TEXTURE 1
+#define USE_SOLID_ANGLE_TEXTURE 1
+
 float PdfSkyLightOcclusionRay(
 	float FaceCdf[7],
 	float3 WorldDirection
diff --git a/Engine/Shaders/Private/RayTracing/RayTracingSkyLightRGS.usf b/Engine/Shaders/Private/RayTracing/RayTracingSkyLightRGS.usf
--- a/Engine/Shaders/Private/RayTracing/RayTracingSkyLightRGS.usf
+++ b/Engine/Shaders/Private/RayTracing/RayTracingSkyLightRGS.usf
@@ -23,9 +23,6 @@ StructuredBuffer<FSkyLightVisibilityRays> SkyLightVisibilityRays;
 RWTexture2D<float4> RWOcclusionMaskUAV;
 RWTexture2D<float2> RWRayDistanceUAV;
 
-#define USE_MIP_PDF_TEXTURE 1
-#define USE_SOLID_ANGLE_TEXTURE 1
-
 float3 EvalSkyLight(float3 Direction)
 {
 	return SkyLight.Color * TextureCubeSampleLevel(SkyLight.Texture, SkyLight.TextureSampler, Direction, 0).rgb;
