Write barycentric data during the hair projection.

(ushell-p4-cherrypick of 9480152 by Charles.deRousiers)

(ushell-p4-cherrypick of 9480169 by Charles.deRousiers)

diff --git a/Engine/Plugins/Experimental/HairStrands/Source/HairStrandsCore/Private/HairStrandsRendering.cpp b/Engine/Plugins/Experimental/HairStrands/Source/HairStrandsCore/Private/HairStrandsRendering.cpp
--- a/Engine/Plugins/Experimental/HairStrands/Source/HairStrandsCore/Private/HairStrandsRendering.cpp
+++ b/Engine/Plugins/Experimental/HairStrands/Source/HairStrandsCore/Private/HairStrandsRendering.cpp
@@ -139,7 +139,8 @@ class FHairInterpolationCS : public FGlobalShader
 		SHADER_PARAMETER_SRV(Buffer<float4>, DeformedPosition1Buffer)
 		SHADER_PARAMETER_SRV(Buffer<float4>, DeformedPosition2Buffer)
 
-		SHADER_PARAMETER_SRV(Buffer<uint>,	RootToTriangleIndex)
+		SHADER_PARAMETER_SRV(Buffer<uint>, RootBarycentricBuffer)
+		SHADER_PARAMETER_SRV(Buffer<uint>, RootToTriangleIndex)
 		SHADER_PARAMETER_SRV(Buffer<uint>, VertexToRootIndexBuffer)
 
 		END_SHADER_PARAMETER_STRUCT()
@@ -202,6 +203,7 @@ static void AddHairStrandsInterpolationPass(
 		Parameters->DeformedPosition2Buffer = InHairData.LODDatas[LODIndex].DeformedRootTrianglePosition2Buffer->SRV;
 
 		Parameters->RootToTriangleIndex = InHairData.LODDatas[LODIndex].RootTriangleIndexBuffer->SRV;
+		Parameters->RootBarycentricBuffer = InHairData.LODDatas[LODIndex].RootTriangleBarycentricBuffer->SRV;
 		Parameters->VertexToRootIndexBuffer = InHairData.VertexToCurveIndexBuffer->SRV;
 	}
 
diff --git a/Engine/Shaders/Private/HairStrands/HairStrandsMeshProjection.usf b/Engine/Shaders/Private/HairStrands/HairStrandsMeshProjection.usf
--- a/Engine/Shaders/Private/HairStrands/HairStrandsMeshProjection.usf
+++ b/Engine/Shaders/Private/HairStrands/HairStrandsMeshProjection.usf
@@ -22,7 +22,7 @@ Buffer<float4> RootPositionBuffer;
 Buffer<float4> RootNormalBuffer;
 
 RWBuffer<uint> OutRootTriangleIndex;
-//RWBuffer<uint> OutRootTriangleBarycentrics;
+RWBuffer<uint> OutRootTriangleBarycentrics;
 RWStructuredBuffer<float> OutRootTriangleDistance;
 
 ////////////////////////////////////////////////////////////////////////////////////////////////////
@@ -33,7 +33,13 @@ RWStructuredBuffer<float> OutRootTriangleDistance;
 
 // Closest point on A triangle from another point
 // Code from the book "Real-Time Collision Detection" by Christer Ericson
-float3 ClosestPoint(const FHairMeshTriangle Tri, const float3 P)
+struct FTrianglePoint
+{
+	float3 P;
+	float3 Barycentric;
+};
+
+FTrianglePoint ComputeClosestPoint(const FHairMeshTriangle Tri, const float3 P)
 {
 	const float3 A = Tri.P0;
 	const float3 B = Tri.P1;
@@ -46,21 +52,35 @@ float3 ClosestPoint(const FHairMeshTriangle Tri, const float3 P)
 	float D1 = dot(AB, AP);
 	float D2 = dot(AC, AP);
 	if (D1 <= 0.f && D2 <= 0.f)
-		return A; // Barycentric coordinates are (1,0,0).
+	{
+		FTrianglePoint Out;
+		Out.P = A;
+		Out.Barycentric = float3(1, 0, 0);
+		return Out;
+	}
 
 	// Check if P is in vertex region outside B.
 	float3 BP = P - B;
 	float D3 = dot(AB, BP);
 	float D4 = dot(AC, BP);
 	if (D3 >= 0.f && D4 <= D3)
-		return B; // Barycentric coordinates are (0,1,0).
+	{
+		FTrianglePoint Out;
+		Out.P = B;
+		Out.Barycentric = float3(0, 1, 0);
+		return Out;
+	}
 
 	// Check if P is in edge region of AB, and if so, return the projection of P onto AB.
 	float VC = D1 * D4 - D3 * D2;
 	if (VC <= 0.f && D1 >= 0.f && D3 <= 0.f)
 	{
 		float V = D1 / (D1 - D3);
-		return A + V * AB; // The barycentric coordinates are (1-V, V, 0).
+
+		FTrianglePoint Out;
+		Out.P = A + V * AB;
+		Out.Barycentric = float3(1 - V, V, 0);
+		return Out;
 	}
 
 	// Check if P is in vertex region outside C.
@@ -68,14 +88,22 @@ float3 ClosestPoint(const FHairMeshTriangle Tri, const float3 P)
 	float D5 = dot(AB, CP);
 	float D6 = dot(AC, CP);
 	if (D6 >= 0.f && D5 <= D6)
-		return C; // The barycentric coordinates are (0,0,1).
+	{
+		FTrianglePoint Out;
+		Out.P = C;
+		Out.Barycentric = float3(0, 0, 1);
+		return Out;
+	}
 
 	// Check if P is in edge region of AC, and if so, return the projection of P onto AC.
 	float VB = D5 * D2 - D1 * D6;
 	if (VB <= 0.f && D2 >= 0.f && D6 <= 0.f)
 	{
 		float W = D2 / (D2 - D6);
-		return A + W * AC; // The barycentric coordinates are (1-W, 0, W).
+		FTrianglePoint Out;
+		Out.P = A + W * AC;
+		Out.Barycentric = float3(1 - W, 0, W);
+		return Out;
 	}
 
 	// Check if P is in edge region of BC, and if so, return the projection of P onto BC.
@@ -83,14 +111,21 @@ float3 ClosestPoint(const FHairMeshTriangle Tri, const float3 P)
 	if (VA <= 0.f && D4 - D3 >= 0.f && D5 - D6 >= 0.f)
 	{
 		float W = (D4 - D3) / (D4 - D3 + D5 - D6);
-		return B + W * (C - B); // The barycentric coordinates are (0, 1-W, W).
+		FTrianglePoint Out;
+		Out.P = B + W * (C - B);
+		Out.Barycentric = float3(0, 1 - W, W);
+		return Out;
 	}
 
 	// P must be inside the face region. Compute the closest point through its barycentric coordinates (u,V,W).
 	float Denom = 1.f / (VA + VB + VC);
 	float V = VB * Denom;
 	float W = VC * Denom;
-	return A + AB * V + AC * W;
+
+	FTrianglePoint Out;
+	Out.P = A + AB * V + AC * W;
+	Out.Barycentric = float3(1-V-W, V, W);
+	return Out;
 }
 
 [numthreads(32, 1, 1)]
@@ -108,24 +143,28 @@ void MainCS(uint3 DispatchThreadId : SV_DispatchThreadID)
 
 	float MinDistance = bClear ? 999999999 : OutRootTriangleDistance[RootIndex];
 	int TriangleIndex = -1;
+	float3 TriangleBarycentric = float3(0,0,0);
 	for (uint PrimitiveIt = 0; PrimitiveIt < MeshNumPrimitives; ++PrimitiveIt)
 	{
 		const FHairMeshTriangle Tri = GetTriangleData(PrimitiveIt, MeshLocalToWorld, MeshPositionBuffer, MeshIndexBuffer, MeshIndexOffset, MeshMaxIndexCount, MeshMaxVertexCount);
 		
 		const float3 B = 0.5f;
-		const float3 C = ClosestPoint(Tri, RootPosition);
-		const float Distance = distance(C, RootPosition);
+		FTrianglePoint C = ComputeClosestPoint(Tri, RootPosition);
+		const float Distance = distance(C.P, RootPosition);
 		if (abs(Distance) < MinDistance)
 		{			
 			MinDistance = abs(Distance);
+			TriangleBarycentric = C.Barycentric;
 			TriangleIndex = PrimitiveIt;
 		}
 	}
 
 	if (TriangleIndex >= 0)
 	{
-		const uint EncodeIndex = EncodeTriangleIndex(TriangleIndex, MeshSectionIndex);
-		OutRootTriangleIndex[RootIndex] = EncodeIndex;
+		const uint EncodedIndex = EncodeTriangleIndex(TriangleIndex, MeshSectionIndex);
+		const uint EncodedBarycentric = EncodeBarycentrics(TriangleBarycentric.xy);
+		OutRootTriangleIndex[RootIndex] = EncodedIndex;
+		OutRootTriangleBarycentrics[RootIndex] = EncodedBarycentric;
 		OutRootTriangleDistance[RootIndex] = MinDistance;
 	}
 
diff --git a/Engine/Source/Runtime/Renderer/Private/HairStrands/HairStrandsMeshProjection.cpp b/Engine/Source/Runtime/Renderer/Private/HairStrands/HairStrandsMeshProjection.cpp
--- a/Engine/Source/Runtime/Renderer/Private/HairStrands/HairStrandsMeshProjection.cpp
+++ b/Engine/Source/Runtime/Renderer/Private/HairStrands/HairStrandsMeshProjection.cpp
@@ -34,8 +34,8 @@ class FHairMeshProjectionCS : public FGlobalShader
 		SHADER_PARAMETER_SRV(Buffer, RootPositionBuffer)
 		SHADER_PARAMETER_SRV(Buffer, RootNormalBuffer)
 
-		SHADER_PARAMETER_UAV(StructuredBuffer, OutRootTriangleIndex)
-		SHADER_PARAMETER_UAV(StructuredBuffer, OutRootTriangleBarycentrics)
+		SHADER_PARAMETER_UAV(RWBuffer, OutRootTriangleIndex)
+		SHADER_PARAMETER_UAV(RWBuffer, OutRootTriangleBarycentrics)
 		SHADER_PARAMETER_RDG_BUFFER_UAV(StructuredBuffer, OutRootTriangleDistance)
 	END_SHADER_PARAMETER_STRUCT()
 
