Fix hair skin issue.

Since the interpolation constant buffer was updated, the shader needed to be touched/updated as well.

(ushell-p4-cherrypick of 9480740 by Charles.deRousiers)

diff --git a/Engine/Shaders/Private/HairStrands/HairStrandsInterpolation.usf b/Engine/Shaders/Private/HairStrands/HairStrandsInterpolation.usf
--- a/Engine/Shaders/Private/HairStrands/HairStrandsInterpolation.usf
+++ b/Engine/Shaders/Private/HairStrands/HairStrandsInterpolation.usf
@@ -29,6 +29,7 @@ Buffer<float4>  DeformedPosition0Buffer;
 Buffer<float4>  DeformedPosition1Buffer;
 Buffer<float4>  DeformedPosition2Buffer;
 
+Buffer<uint>	RootBarycentricBuffer;
 Buffer<uint>	RootToTriangleIndex;
 Buffer<uint>	VertexToRootIndexBuffer;
 #endif
