Add stretching component for skin deformation

(ushell-p4-cherrypick of 9485618 by Charles.deRousiers)

diff --git a/Engine/Shaders/Private/HairStrands/HairStrandsInterpolation.usf b/Engine/Shaders/Private/HairStrands/HairStrandsInterpolation.usf
--- a/Engine/Shaders/Private/HairStrands/HairStrandsInterpolation.usf
+++ b/Engine/Shaders/Private/HairStrands/HairStrandsInterpolation.usf
@@ -82,10 +82,11 @@ void MainCS(uint2 DispatchThreadId : SV_DispatchThreadID)
 		#if PERMUTATION_DYNAMIC_GEOMETRY == 1
 		{
 			const uint RootIndex						= VertexToRootIndexBuffer[VertexIndex];
+			const float3 RootBarycentric				= DecodeBarycentrics(RootBarycentricBuffer[RootIndex]);
 			const FHairMeshTriangle RestTriangle		= GetTriangleTransformation(RootIndex, RestPosition0Buffer, RestPosition1Buffer, RestPosition2Buffer, RestPositionWorldCenter);
 			const FHairMeshTriangle DeformedTriangle	= GetTriangleTransformation(RootIndex, DeformedPosition0Buffer, DeformedPosition1Buffer, DeformedPosition2Buffer, DeformedPositionWorldCenter);
 
-			ControlPoint = TransformPoint(ControlPoint + HairWorldOffset, RestTriangle, DeformedTriangle) - HairWorldOffset;
+			ControlPoint = TransformPoint(ControlPoint + HairWorldOffset, RootBarycentric, RestTriangle, DeformedTriangle) - HairWorldOffset;
 
 			// Debug
 			#if 0
diff --git a/Engine/Shaders/Private/HairStrands/HairStrandsMeshProjectionCommon.ush b/Engine/Shaders/Private/HairStrands/HairStrandsMeshProjectionCommon.ush
--- a/Engine/Shaders/Private/HairStrands/HairStrandsMeshProjectionCommon.ush
+++ b/Engine/Shaders/Private/HairStrands/HairStrandsMeshProjectionCommon.ush
@@ -71,9 +71,10 @@ FHairMeshTriangle GetTriangleTransformation(
 	return Out;
 }
 
-float3 ToTriangle(const float3 P, const FHairMeshTriangle T)
+float3 ToTriangle(const float3 P, float3 B, const FHairMeshTriangle T)
 {
-	float3 LocalP = P - T.P0;
+	const float3 LocalP = P - (T.P0 * B.x + T.P1 * B.y + T.P2 * B.z);
+
 	float3 Out;
 	Out.x = dot(LocalP, T.E0);
 	Out.y = dot(LocalP, T.E1);
@@ -81,15 +82,18 @@ float3 ToTriangle(const float3 P, const FHairMeshTriangle T)
 	return Out;
 }
 
-float3 ToWorld(const float3 P, const FHairMeshTriangle T)
+float3 ToWorld(const float3 P, float3 B, const FHairMeshTriangle T)
 {
-	return P.x * T.E0 + P.y * T.E1 + P.z * T.N + T.P0;
+	const float3 Rotation = P.x * T.E0 + P.y * T.E1 + P.z * T.N;
+	const float3 Translation = T.P0 * B.x + T.P1 * B.y + T.P2 * B.z;
+	return Rotation + Translation;
+
 }
 
-float3 TransformPoint(float3 P, const FHairMeshTriangle RestTriangle, const FHairMeshTriangle DeformedTriangle)
+float3 TransformPoint(float3 P, float3 B, const FHairMeshTriangle RestTriangle, const FHairMeshTriangle DeformedTriangle)
 {
-	float3 LocalP = ToTriangle(P, RestTriangle);
-	return ToWorld(LocalP, DeformedTriangle);
+	float3 LocalP = ToTriangle(P, B, RestTriangle);
+	return ToWorld(LocalP, B, DeformedTriangle);
 }
 
 struct FHairProjectionData
@@ -124,9 +128,13 @@ uint EncodeBarycentrics(float2 B)
 	return f32tof16(B.x) | (f32tof16(B.y) << 16);
 }
 
-float2 DecodeBarycentrics(uint E)
+float3 DecodeBarycentrics(uint E)
 {
-	return float2(f16tof32(E & 0xFFFF), f16tof32((E>>16) & 0xFFFF));
+	float3 Out;
+	Out.x = f16tof32(E & 0xFFFF);
+	Out.y = f16tof32((E >> 16) & 0xFFFF);
+	Out.z = 1 - Out.x - Out.y;
+	return Out;
 }
 
 
diff --git a/Engine/Shaders/Private/HairStrands/HairStrandsMeshProjectionHairDebug.usf b/Engine/Shaders/Private/HairStrands/HairStrandsMeshProjectionHairDebug.usf
--- a/Engine/Shaders/Private/HairStrands/HairStrandsMeshProjectionHairDebug.usf
+++ b/Engine/Shaders/Private/HairStrands/HairStrandsMeshProjectionHairDebug.usf
@@ -14,6 +14,7 @@ uint			DeformedFrameEnable;
 #if PERMUTATION_INPUT_TYPE == INPUT_ROOT_FRAME
 Buffer<float4> RootPositionBuffer;
 Buffer<float4> RootNormalBuffer;
+Buffer<uint>   RootBarycentricBuffer;
 #endif
 
 Buffer<float4>  RestPosition0Buffer;
@@ -47,13 +48,14 @@ void MainVS(
 
 		float3 P0 = RootPositionBuffer[LineIndex].xyz;
 		float3 LN = RootNormalBuffer[LineIndex].xyz;
+		float3 B  = DecodeBarycentrics(RootBarycentricBuffer[LineIndex]);
 		float3 P1 = P0 + LN;
 
 
 		if (DeformedFrameEnable > 0)
 		{
-			P0 = TransformPoint(P0, RestTriangle, DeformedTriangle);
-			P1 = TransformPoint(P1, RestTriangle, DeformedTriangle);
+			P0 = TransformPoint(P0, B, RestTriangle, DeformedTriangle);
+			P1 = TransformPoint(P1, B, RestTriangle, DeformedTriangle);
 		}
 
 		P0 = mul(float4(P0, 1), RootLocalToWorld).xyz;
@@ -92,13 +94,23 @@ void MainVS(
 	#if 0
 	if (DeformedFrameEnable > 0)
 	{
+		float3 Barycentric = 0;
 		if (TriangleVertex == 0)
+		{
 			P = RestTriangle.P0;
+			Barycentric = float3(1, 0, 0);
+		}
 		if (TriangleVertex == 1)
+		{
 			P = RestTriangle.P1;
+			Barycentric = float3(0, 1, 0);
+		}
 		if (TriangleVertex == 2)
+		{
 			P = RestTriangle.P2;
-		P = TransformPoint(P, RestTriangle, DeformedTriangle);
+			Barycentric = float3(0, 0, 1);
+		}
+		P = TransformPoint(P, Barycentric, RestTriangle, DeformedTriangle);
 	}
 	#endif
 
diff --git a/Engine/Source/Runtime/Renderer/Private/HairStrands/HairStrandsDebug.cpp b/Engine/Source/Runtime/Renderer/Private/HairStrands/HairStrandsDebug.cpp
--- a/Engine/Source/Runtime/Renderer/Private/HairStrands/HairStrandsDebug.cpp
+++ b/Engine/Source/Runtime/Renderer/Private/HairStrands/HairStrandsDebug.cpp
@@ -645,6 +645,7 @@ BEGIN_SHADER_PARAMETER_STRUCT(FHairProjectionHairDebugParameters, )
 	// Change for actual frame data (stored or computed only)
 	SHADER_PARAMETER_SRV(StructuredBuffer, RootPositionBuffer)
 	SHADER_PARAMETER_SRV(StructuredBuffer, RootNormalBuffer)
+	SHADER_PARAMETER_SRV(StructuredBuffer, RootBarycentricBuffer)
 
 	SHADER_PARAMETER_STRUCT_REF(FViewUniformShaderParameters, ViewUniformBuffer)
 	RENDER_TARGET_BINDING_SLOTS()
@@ -713,7 +714,7 @@ static void AddDebugProjectionHairPass(
 	if (PrimitiveCount == 0 || LODIndex < 0 || LODIndex >= HairData.LODDatas.Num())
 		return;
 
-	if (EDebugProjectionHairType::HairFrame == GeometryType && (!HairData.RootPositionBuffer || !HairData.RootNormalBuffer))
+	if (EDebugProjectionHairType::HairFrame == GeometryType && (!HairData.RootPositionBuffer || !HairData.RootNormalBuffer || !HairData.LODDatas[LODIndex].RootTriangleBarycentricBuffer))
 		return;
 
 	if (!HairData.LODDatas[LODIndex].RestRootTrianglePosition0Buffer ||
@@ -737,6 +738,7 @@ static void AddDebugProjectionHairPass(
 	{
 		Parameters->RootPositionBuffer = HairData.RootPositionBuffer;
 		Parameters->RootNormalBuffer = HairData.RootNormalBuffer;
+		Parameters->RootBarycentricBuffer = HairData.LODDatas[LODIndex].RootTriangleBarycentricBuffer->SRV;
 	}
 
 	Parameters->RestPosition0Buffer = HairData.LODDatas[LODIndex].RestRootTrianglePosition0Buffer->SRV;
