Add per pixel hair debug info.

(ushell-p4-cherrypick of 9567767 by Charles.deRousiers)

diff --git a/Engine/Shaders/Private/HairStrands/HairStrandsDebugPrint.usf b/Engine/Shaders/Private/HairStrands/HairStrandsDebugPrint.usf
new file mode 100644
--- /dev/null
+++ b/Engine/Shaders/Private/HairStrands/HairStrandsDebugPrint.usf
@@ -0,0 +1,331 @@
+// Copyright 1998-2019 Epic Games, Inc. All Rights Reserved.
+
+#include "../Common.ush"
+#include "../ShaderPrintCommon.ush"
+#include "HairStrandsVisibilityCommon.ush"
+
+int2 PixelCoord;
+int2 MaxResolution;
+uint FastResolveMask;
+
+Texture2D<uint4> CategorizationTexture;
+Texture2D<uint2> DepthStencilTexture;
+SamplerState LinearSampler;
+Texture2D<uint> HairVisibilityNodeOffsetAndCount;
+StructuredBuffer<FNodeData> HairVisibilityNodeData;
+
+#define DEBUG_MODE_SAMPLE_COUNT 0 
+#define DEBUG_MODE_COVERAGE_TYPE 1
+#define DEBUG_MODE_TAA_RESOLVE_TYPE 2
+
+#define MATERIAL_DATA 0
+#define DETAIL_DATA 0
+
+float3 HUEtoRGB(in float H)
+{
+	float R = abs(H * 6 - 3) - 1;
+	float G = 2 - abs(H * 6 - 2);
+	float B = 2 - abs(H * 6 - 4);
+	return saturate(float3(R, G, B));
+}
+
+float3 HSVtoRGB(in float3 HSV)
+{
+	float3 RGB = HUEtoRGB(HSV.x);
+	return ((RGB - 1) * HSV.y + 1) * HSV.z;
+}
+
+float3 GetColorCode(float x)
+{
+	float c = (1 - saturate(x)) * 0.6; // Remap [0,1] to Blue-Red
+	return x > 0 ? HSVtoRGB(float3(c, 1, 1)) : float3(0, 0, 0);
+}
+
+[numthreads(1, 1, 1)]
+void MainCS(uint GroupIndex : SV_GroupIndex, uint3 DispatchThreadId : SV_DispatchThreadID)
+{	
+	const float2 PixelPos = PixelCoord;
+	if (PixelCoord.x < 0 || PixelCoord.y < 0 || PixelCoord.x >= MaxResolution.x || PixelCoord.y >= MaxResolution.y)
+		return;
+
+	ShaderPrintFilter(true);
+
+	const FCategorizationData CatData = DecodeCategorizationData(CategorizationTexture.Load(uint3(PixelPos,0)));
+	//const bool bIsValid = CatData.PixelCoverage > 0;
+	const bool bIsValid = CatData.TotalCoverage8bit > 0;
+	if (!bIsValid)
+		return;
+
+	const float2 OriginalPos = PixelPos / float2(MaxResolution);
+	float2 Pos = OriginalPos;
+
+	uint Total8BitCoverage = 0;
+	const FNodeDesc NodeDesc = DecodeNodeDesc(HairVisibilityNodeOffsetAndCount.Load(uint3(PixelCoord, 0)));
+	for (uint SampleIt = 0; SampleIt < NodeDesc.Count; SampleIt++)
+	{
+		const uint LocalOffset = NodeDesc.Offset + SampleIt;
+		const FNodeData NodeData = HairVisibilityNodeData[LocalOffset];
+		float			SampleDepth;
+		FVisibilityData SampleVisData;
+		float3			SampleBaseColor;
+		float			SampleRoughness;
+		float			SampleSpecular;
+		DecodeNodeData(NodeData, SampleDepth, SampleVisData, SampleBaseColor, SampleRoughness, SampleSpecular);
+		Total8BitCoverage += SampleVisData.Coverage8bit;
+	}
+
+
+	// Max sample
+	Pos = ShaderPrintSymbol(Pos, _M_);
+	Pos = ShaderPrintSymbol(Pos, _A_);
+	Pos = ShaderPrintSymbol(Pos, _X_);
+	Pos = ShaderPrintSymbol(Pos, _SPC_);
+	Pos = ShaderPrintSymbol(Pos, _S_);
+	Pos = ShaderPrintSymbol(Pos, _A_);
+	Pos = ShaderPrintSymbol(Pos, _M_);
+	Pos = ShaderPrintSymbol(Pos, _P_);
+	Pos = ShaderPrintSymbol(Pos, _L_);
+	Pos = ShaderPrintSymbol(Pos, _E_);
+	Pos = ShaderPrintSymbol(Pos, _SPC_);
+	Pos = ShaderPrint(Pos, uint(HAIR_VISIBILITY_MAXSAMPLECOUNT));
+
+	Pos = ShaderPrintNewline(Pos);
+	Pos.x = OriginalPos.x;
+
+	// Unique sample
+	Pos = ShaderPrintSymbol(Pos, _S_);
+	Pos = ShaderPrintSymbol(Pos, _A_);
+	Pos = ShaderPrintSymbol(Pos, _M_);
+	Pos = ShaderPrintSymbol(Pos, _P_);
+	Pos = ShaderPrintSymbol(Pos, _L_);
+	Pos = ShaderPrintSymbol(Pos, _E_);
+	Pos = ShaderPrintSymbol(Pos, _SPC_);
+	Pos = ShaderPrint(Pos, CatData.SampleCount);
+
+	Pos = ShaderPrintNewline(Pos);
+	Pos.x = OriginalPos.x;
+
+	// Coverage
+	#if 0
+	Pos = ShaderPrintSymbol(Pos, _C_);
+	Pos = ShaderPrintSymbol(Pos, _O_);
+	Pos = ShaderPrintSymbol(Pos, _V_);
+	Pos = ShaderPrintSymbol(Pos, _E_);
+	Pos = ShaderPrintSymbol(Pos, _R_);
+	Pos = ShaderPrintSymbol(Pos, _A_);
+	Pos = ShaderPrintSymbol(Pos, _G_);
+	Pos = ShaderPrintSymbol(Pos, _E_);
+	Pos = ShaderPrintSymbol(Pos, _SPC_);
+	Pos = ShaderPrint(Pos, CatData.PixelCoverage);
+
+	Pos = ShaderPrintNewline(Pos);
+	Pos.x = OriginalPos.x;
+	#endif
+
+	// 8 bit Coverage (cateogory texture)
+	Pos = ShaderPrintSymbol(Pos, _C_);
+	Pos = ShaderPrintSymbol(Pos, _O_);
+	Pos = ShaderPrintSymbol(Pos, _V_);
+	Pos = ShaderPrintSymbol(Pos, _E_);
+	Pos = ShaderPrintSymbol(Pos, _R_);
+	Pos = ShaderPrintSymbol(Pos, _A_);
+	Pos = ShaderPrintSymbol(Pos, _G_);
+	Pos = ShaderPrintSymbol(Pos, _E_);
+	Pos = ShaderPrintSymbol(Pos, _SPC_);
+	Pos = ShaderPrintSymbol(Pos, _8_); 
+	Pos = ShaderPrintSymbol(Pos, _B_);
+	Pos = ShaderPrintSymbol(Pos, _I_);
+	Pos = ShaderPrintSymbol(Pos, _T_);
+	Pos = ShaderPrintSymbol(Pos, _S_);
+	Pos = ShaderPrintSymbol(Pos, _SPC_);
+	Pos = ShaderPrint(Pos, CatData.TotalCoverage8bit);
+	
+
+	// 8 bit Coverage (summed)
+	Pos = ShaderPrintSymbol(Pos, _S_);
+	Pos = ShaderPrintSymbol(Pos, _U_);
+	Pos = ShaderPrintSymbol(Pos, _M_);
+	Pos = ShaderPrintSymbol(Pos, _M_);
+	Pos = ShaderPrintSymbol(Pos, _E_);
+	Pos = ShaderPrintSymbol(Pos, _D_);
+	Pos = ShaderPrintSymbol(Pos, _SPC_);
+	Pos = ShaderPrintSymbol(Pos, _C_);
+	Pos = ShaderPrintSymbol(Pos, _O_);
+	Pos = ShaderPrintSymbol(Pos, _V_);
+	Pos = ShaderPrintSymbol(Pos, _E_);
+	Pos = ShaderPrintSymbol(Pos, _R_);
+	Pos = ShaderPrintSymbol(Pos, _A_);
+	Pos = ShaderPrintSymbol(Pos, _G_);
+	Pos = ShaderPrintSymbol(Pos, _E_);
+	Pos = ShaderPrintSymbol(Pos, _SPC_);
+	Pos = ShaderPrintSymbol(Pos, _8_); 
+	Pos = ShaderPrintSymbol(Pos, _B_);
+	Pos = ShaderPrintSymbol(Pos, _I_);
+	Pos = ShaderPrintSymbol(Pos, _T_);
+	Pos = ShaderPrintSymbol(Pos, _S_);
+	Pos = ShaderPrintSymbol(Pos, _SPC_);
+	Pos = ShaderPrint(Pos, Total8BitCoverage);
+
+	Pos = ShaderPrintNewline(Pos);
+	Pos.x = OriginalPos.x;
+
+	// Resolve type
+	const uint Stencil = DepthStencilTexture.Load(uint3(PixelPos,0)).y;
+	const bool bIsFastResolve = (Stencil & FastResolveMask) != 0;
+	Pos = ShaderPrintSymbol(Pos, _R_);
+	Pos = ShaderPrintSymbol(Pos, _E_);
+	Pos = ShaderPrintSymbol(Pos, _S_);
+	Pos = ShaderPrintSymbol(Pos, _O_);
+	Pos = ShaderPrintSymbol(Pos, _L_);
+	Pos = ShaderPrintSymbol(Pos, _V_);
+	Pos = ShaderPrintSymbol(Pos, _E_);
+	Pos = ShaderPrintSymbol(Pos, _SPC_);
+	if (bIsFastResolve)
+	{
+		Pos = ShaderPrintSymbol(Pos, _F_);
+		Pos = ShaderPrintSymbol(Pos, _A_);
+		Pos = ShaderPrintSymbol(Pos, _S_);
+		Pos = ShaderPrintSymbol(Pos, _T_);
+	}
+	else
+	{
+		Pos = ShaderPrintSymbol(Pos, _R_);
+		Pos = ShaderPrintSymbol(Pos, _E_);
+		Pos = ShaderPrintSymbol(Pos, _G_);
+		Pos = ShaderPrintSymbol(Pos, _U_);
+		Pos = ShaderPrintSymbol(Pos, _L_);
+		Pos = ShaderPrintSymbol(Pos, _A_);
+		Pos = ShaderPrintSymbol(Pos, _R_);
+	}
+
+	Pos = ShaderPrintNewline(Pos);
+	Pos.x = OriginalPos.x;
+
+	// Per sample description
+	for (uint SampleIt = 0; SampleIt < NodeDesc.Count; SampleIt++)
+	{
+		const uint LocalOffset = NodeDesc.Offset + SampleIt;
+		const FNodeData NodeData = HairVisibilityNodeData[LocalOffset];
+		float			SampleDepth;
+		FVisibilityData SampleVisData;
+		float3			SampleBaseColor;
+		float			SampleRoughness;
+		float			SampleSpecular;
+		DecodeNodeData(NodeData, SampleDepth, SampleVisData, SampleBaseColor, SampleRoughness, SampleSpecular);
+
+		// Base Color
+		Pos = ShaderPrintSymbol(Pos, _MINUS_);
+		Pos = ShaderPrintSymbol(Pos, _MINUS_);
+		Pos = ShaderPrintSymbol(Pos, _MINUS_);
+		Pos = ShaderPrintSymbol(Pos, _MINUS_);
+		Pos = ShaderPrintSymbol(Pos, _MINUS_);
+		Pos = ShaderPrintSymbol(Pos, _MINUS_);
+		Pos = ShaderPrintSymbol(Pos, _MINUS_);
+		Pos = ShaderPrintSymbol(Pos, _MINUS_);
+		Pos = ShaderPrintSymbol(Pos, _MINUS_);
+		Pos = ShaderPrintSymbol(Pos, _MINUS_);
+		Pos = ShaderPrintSymbol(Pos, _MINUS_);
+		Pos = ShaderPrintSymbol(Pos, _MINUS_);
+		Pos = ShaderPrintSymbol(Pos, _MINUS_);
+		Pos = ShaderPrintSymbol(Pos, _MINUS_);
+		Pos = ShaderPrintSymbol(Pos, _MINUS_);
+		Pos = ShaderPrintSymbol(Pos, _MINUS_);
+		Pos = ShaderPrintSymbol(Pos, _MINUS_);
+		Pos = ShaderPrintSymbol(Pos, _MINUS_);
+		Pos = ShaderPrintSymbol(Pos, _MINUS_);
+		Pos = ShaderPrintSymbol(Pos, _MINUS_);
+
+		Pos = ShaderPrintNewline(Pos);
+		Pos.x = OriginalPos.x;
+
+#if MATERIAL_DATA
+		// Base Color
+		Pos = ShaderPrintSymbol(Pos, _DOT_);
+		Pos = ShaderPrintSymbol(Pos, _B_);
+		Pos = ShaderPrintSymbol(Pos, _A_);
+		Pos = ShaderPrintSymbol(Pos, _S_);
+		Pos = ShaderPrintSymbol(Pos, _E_);
+		Pos = ShaderPrintSymbol(Pos, _SPC_);
+		Pos = ShaderPrintSymbol(Pos, _C_);
+		Pos = ShaderPrintSymbol(Pos, _O_);
+		Pos = ShaderPrintSymbol(Pos, _L_);
+		Pos = ShaderPrintSymbol(Pos, _O_);
+		Pos = ShaderPrintSymbol(Pos, _R_);
+		Pos = ShaderPrintSymbol(Pos, _SPC_);
+		Pos = ShaderPrint(Pos, SampleBaseColor);
+		
+		Pos = ShaderPrintNewline(Pos);
+		Pos.x = OriginalPos.x;
+
+		// Roughness
+		Pos = ShaderPrintSymbol(Pos, _DOT_);
+		Pos = ShaderPrintSymbol(Pos, _R_);
+		Pos = ShaderPrintSymbol(Pos, _O_);
+		Pos = ShaderPrintSymbol(Pos, _U_);
+		Pos = ShaderPrintSymbol(Pos, _G_);
+		Pos = ShaderPrintSymbol(Pos, _H_);
+		Pos = ShaderPrintSymbol(Pos, _N_);
+		Pos = ShaderPrintSymbol(Pos, _E_);
+		Pos = ShaderPrintSymbol(Pos, _S_);
+		Pos = ShaderPrintSymbol(Pos, _S_);
+		Pos = ShaderPrintSymbol(Pos, _SPC_);
+		Pos = ShaderPrint(Pos, SampleRoughness);
+
+		Pos = ShaderPrintNewline(Pos);
+		Pos.x = OriginalPos.x;
+#endif 
+
+#if DETAIL_DATA
+		// PrimitiveId
+		Pos = ShaderPrintSymbol(Pos, _DOT_);
+		Pos = ShaderPrintSymbol(Pos, _P_);
+		Pos = ShaderPrintSymbol(Pos, _R_);
+		Pos = ShaderPrintSymbol(Pos, _I_);
+		Pos = ShaderPrintSymbol(Pos, _M_);
+		Pos = ShaderPrintSymbol(Pos, _I_);
+		Pos = ShaderPrintSymbol(Pos, _T_);
+		Pos = ShaderPrintSymbol(Pos, _I_);
+		Pos = ShaderPrintSymbol(Pos, _V_);
+		Pos = ShaderPrintSymbol(Pos, _E_);
+		Pos = ShaderPrintSymbol(Pos, _SPC_);
+		Pos = ShaderPrint(Pos, SampleVisData.PrimitiveId);
+
+		Pos = ShaderPrintNewline(Pos);
+		Pos.x = OriginalPos.x;
+
+		// ClusterId
+		Pos = ShaderPrintSymbol(Pos, _DOT_);
+		Pos = ShaderPrintSymbol(Pos, _C_);
+		Pos = ShaderPrintSymbol(Pos, _L_);
+		Pos = ShaderPrintSymbol(Pos, _U_);
+		Pos = ShaderPrintSymbol(Pos, _S_);
+		Pos = ShaderPrintSymbol(Pos, _T_);
+		Pos = ShaderPrintSymbol(Pos, _E_);
+		Pos = ShaderPrintSymbol(Pos, _R_);
+		Pos = ShaderPrintSymbol(Pos, _SPC_);
+		Pos = ShaderPrintSymbol(Pos, _I_);
+		Pos = ShaderPrintSymbol(Pos, _D_);
+		Pos = ShaderPrintSymbol(Pos, _SPC_);
+		Pos = ShaderPrint(Pos, SampleVisData.ClusterId);
+
+		Pos = ShaderPrintNewline(Pos);
+		Pos.x = OriginalPos.x;		
+#endif
+
+		// Coverage8bit
+		Pos = ShaderPrintSymbol(Pos, _DOT_);
+		Pos = ShaderPrintSymbol(Pos, _C_);
+		Pos = ShaderPrintSymbol(Pos, _O_);
+		Pos = ShaderPrintSymbol(Pos, _V_);
+		Pos = ShaderPrintSymbol(Pos, _E_);
+		Pos = ShaderPrintSymbol(Pos, _R_);
+		Pos = ShaderPrintSymbol(Pos, _A_);
+		Pos = ShaderPrintSymbol(Pos, _G_);
+		Pos = ShaderPrintSymbol(Pos, _E_);
+		Pos = ShaderPrintSymbol(Pos, _SPC_);
+		Pos = ShaderPrint(Pos, SampleVisData.Coverage8bit);
+
+		Pos = ShaderPrintNewline(Pos);
+		Pos.x = OriginalPos.x;
+	}
+}
diff --git a/Engine/Source/Runtime/Renderer/Private/HairStrands/HairStrandsDebug.cpp b/Engine/Source/Runtime/Renderer/Private/HairStrands/HairStrandsDebug.cpp
--- a/Engine/Source/Runtime/Renderer/Private/HairStrands/HairStrandsDebug.cpp
+++ b/Engine/Source/Runtime/Renderer/Private/HairStrands/HairStrandsDebug.cpp
@@ -21,6 +21,7 @@
 #include "DynamicPrimitiveDrawing.h"
 #include "RenderTargetTemp.h"
 #include "CanvasTypes.h"
+#include "ShaderPrintParameters.h"
 
 static int32 GDeepShadowDebugIndex = 0;
 static float GDeepShadowDebugScale = 20;
@@ -112,7 +113,7 @@ static const TCHAR* ToString(EHairDebugMode DebugMode)
 	case EHairDebugMode::DeepOpacityMaps: return TEXT("Deep opacity maps");
 	case EHairDebugMode::SamplePerPixel: return TEXT("Sub-pixel sample count");
 	case EHairDebugMode::TAAResolveType: return TEXT("TAA resolve type (regular/responsive)");
-	case EHairDebugMode::CoverageType: return TEXT("Type of hair coverage (full/partial)");
+	case EHairDebugMode::CoverageType: return TEXT("Type of hair coverage - Fully covered : Green / Partially covered : Red");
 	case EHairDebugMode::VoxelsDensity: return TEXT("Hair density volume");
 	case EHairDebugMode::VoxelsTangent: return TEXT("Hair tangent volume");
 	case EHairDebugMode::VoxelsBaseColor: return TEXT("Hair base color volume");
@@ -156,6 +157,73 @@ static const TCHAR* ToString(EHairStrandsDebugMode DebugMode)
 }
 
 
+///////////////////////////////////////////////////////////////////////////////////////////////////
+class FHairDebugPrintCS : public FGlobalShader
+{
+	DECLARE_GLOBAL_SHADER(FHairDebugPrintCS);
+	SHADER_USE_PARAMETER_STRUCT(FHairDebugPrintCS, FGlobalShader);
+
+	BEGIN_SHADER_PARAMETER_STRUCT(FParameters, )
+		SHADER_PARAMETER(FIntPoint, PixelCoord)
+		SHADER_PARAMETER(FIntPoint, MaxResolution)
+		SHADER_PARAMETER(uint32, FastResolveMask)
+		SHADER_PARAMETER_RDG_TEXTURE(Texture2D, CategorizationTexture)
+		SHADER_PARAMETER_RDG_TEXTURE(Texture2D, HairVisibilityNodeOffsetAndCount)
+		SHADER_PARAMETER_RDG_BUFFER_SRV(StructuredBuffer, HairVisibilityNodeData)
+		SHADER_PARAMETER_SRV(Texture2D, DepthStencilTexture)
+		SHADER_PARAMETER_SAMPLER(SamplerState, LinearSampler)
+		SHADER_PARAMETER_STRUCT_REF(FViewUniformShaderParameters, ViewUniformBuffer)
+		SHADER_PARAMETER_STRUCT_INCLUDE(ShaderPrint::FShaderParameters, ShaderPrintUniformBuffer)
+	END_SHADER_PARAMETER_STRUCT()
+
+public:
+	static bool ShouldCompilePermutation(const FGlobalShaderPermutationParameters& Parameters) { return IsHairStrandsSupported(Parameters.Platform); }
+};
+
+IMPLEMENT_GLOBAL_SHADER(FHairDebugPrintCS, "/Engine/Private/HairStrands/HairStrandsDebugPrint.usf", "MainCS", SF_Compute);
+
+static void AddDebugHairPrintPass(
+	FRDGBuilder& GraphBuilder,
+	const FViewInfo* View,
+	const EHairDebugMode InDebugMode,
+	const FHairStrandsVisibilityData& VisibilityData,
+	const FShaderResourceViewRHIRef& InDepthStencilTexture)
+{
+	const TRefCountPtr<IPooledRenderTarget>& InCategorizationTexture = VisibilityData.CategorizationTexture;
+	const TRefCountPtr<IPooledRenderTarget>& InNodeIndex = VisibilityData.NodeIndex;
+	const TRefCountPtr<FPooledRDGBuffer>& InNodeData = VisibilityData.NodeData;
+
+	if (!InCategorizationTexture || !InNodeIndex || !InNodeData || !InDepthStencilTexture) return;
+
+	FRDGTextureRef CategorizationTexture = InCategorizationTexture ? GraphBuilder.RegisterExternalTexture(InCategorizationTexture, TEXT("CategorizationTexture")) : nullptr;
+	FRDGTextureRef NodeIndex = InNodeIndex ? GraphBuilder.RegisterExternalTexture(InNodeIndex, TEXT("NodeIndex")) : nullptr;
+	FRDGBufferRef  NodeData  = InNodeData ? GraphBuilder.RegisterExternalBuffer(InNodeData, TEXT("NodeData")) : nullptr;
+
+	const FIntRect Viewport = View->ViewRect;
+	const FIntPoint Resolution(Viewport.Width(), Viewport.Height());
+
+	FHairDebugPrintCS::FParameters* Parameters = GraphBuilder.AllocParameters<FHairDebugPrintCS::FParameters>();
+	Parameters->MaxResolution = CategorizationTexture->Desc.Extent;
+	Parameters->PixelCoord = View->CursorPos;
+	Parameters->FastResolveMask = STENCIL_TEMPORAL_RESPONSIVE_AA_MASK;
+	Parameters->CategorizationTexture = CategorizationTexture;
+	Parameters->HairVisibilityNodeData = GraphBuilder.CreateSRV(NodeData);
+	Parameters->HairVisibilityNodeOffsetAndCount = NodeIndex;
+	Parameters->DepthStencilTexture = InDepthStencilTexture;
+	Parameters->LinearSampler = TStaticSamplerState<SF_Bilinear, AM_Clamp, AM_Clamp, AM_Clamp>::GetRHI();
+	ShaderPrint::SetParameters(*View, Parameters->ShaderPrintUniformBuffer);
+	TShaderMapRef<FHairDebugPrintCS> ComputeShader(View->ShaderMap);
+
+	ClearUnusedGraphResources(*ComputeShader, Parameters);
+
+	FComputeShaderUtils::AddPass(
+		GraphBuilder,
+		RDG_EVENT_NAME("HairStrandsDebugPrint"),
+		*ComputeShader,
+		Parameters,
+		FIntVector(1, 1, 1));
+}
+
 ///////////////////////////////////////////////////////////////////////////////////////////////////
 class FHairDebugPS : public FGlobalShader
 {
@@ -906,6 +974,14 @@ void RenderHairStrandsDebugInfo(FRHICommandListImmediate& RHICmdList, TArray<FVi
 
 	if (HairDebugMode == EHairDebugMode::ClusterData)
 	{
+		if (ViewIndex < uint32(HairDatas->HairVisibilityViews.HairDatas.Num()))
+		{
+			FRDGBuilder GraphBuilder(RHICmdList);
+			const FHairStrandsVisibilityData& VisibilityData = HairDatas->HairVisibilityViews.HairDatas[ViewIndex];
+			AddDebugHairPrintPass(GraphBuilder, &View, HairDebugMode, VisibilityData, SceneTargets.SceneStencilSRV);
+			GraphBuilder.Execute();	
+		}
+
 		// Clusters part of the clusters
 		FViewElementPDI ShadowFrustumPDI(&View, nullptr, nullptr);
 		const FHairStrandsClusterDatas& ClusterDatas = InClusterViews.Views[ViewIndex];
@@ -1031,6 +1107,7 @@ void RenderHairStrandsDebugInfo(FRHICommandListImmediate& RHICmdList, TArray<FVi
 		{
 			const FHairStrandsVisibilityData& VisibilityData = HairDatas->HairVisibilityViews.HairDatas[ViewIndex];
 			AddDebugHairPass(GraphBuilder, &View, HairDebugMode, VisibilityData.CategorizationTexture, SceneTargets.SceneStencilSRV, SceneColorTexture);
+			AddDebugHairPrintPass(GraphBuilder, &View, HairDebugMode, VisibilityData,  SceneTargets.SceneStencilSRV);
 		}
 
 		GraphBuilder.Execute();
