Add a permutation to evaluate sky lighting  on a per sample evaluation.

(ushell-p4-cherrypick of 9570010 by Charles.deRousiers)

diff --git a/Engine/Shaders/Private/HairStrands/HairStrandsEnvironmentLighting.usf b/Engine/Shaders/Private/HairStrands/HairStrandsEnvironmentLighting.usf
--- a/Engine/Shaders/Private/HairStrands/HairStrandsEnvironmentLighting.usf
+++ b/Engine/Shaders/Private/HairStrands/HairStrandsEnvironmentLighting.usf
@@ -467,11 +467,11 @@ void MainPS(
 		// Note:
 		// * Currently: compute take the closest sample
 		// * Change this to per sample evaluation once we have better hair transmittance evaluation for sky
-		float Coverage = 1;
+		float Coverage = 0;
 		float ClosestDepth = 0;
 		bool bIsClusterCompatible = false;
 		FGBufferData HairGBuffer = (FGBufferData)0;
-		Coverage = 0;
+		float3 SkyLighting = 0;
 
 		LOOP
 		for (uint SampleIt = 0; SampleIt < NodeDesc.Count; SampleIt++)
@@ -488,7 +488,9 @@ void MainPS(
 
 			// Inject material data from the visibility/mini-gbuffer for hair
 			const float SceneDepth = ConvertFromDeviceZ(SampleDepth);
+			#if PERMUTATION_PER_SAMPLE == 0
 			if (SampleDepth > ClosestDepth)
+			#endif
 			{
 				HairGBuffer.WorldNormal = SampleVisData.Tangent;
 				HairGBuffer.BaseColor = SampleBaseColor;
@@ -503,10 +505,19 @@ void MainPS(
 				HairGBuffer.CustomData = 0;
 				HairGBuffer.IndirectIrradiance = 1000000;
 			}
-			ClosestDepth = max(ClosestDepth, SampleDepth);  // Inverse-Z
 
+			#if PERMUTATION_PER_SAMPLE == 0
+			ClosestDepth = max(ClosestDepth, SampleDepth);  // Inverse-Z
 			const float LocalCoverage = From8bitCoverage(SampleVisData.Coverage8bit);
 			Coverage += LocalCoverage;
+			#else
+			if (bIsClusterCompatible = SampleVisData.ClusterId == Voxel_ClusterId)
+			{
+				const float SampleCoverage = From8bitCoverage(SampleVisData.Coverage8bit);
+				SkyLighting += ReflectionEnvironment(HairGBuffer, BufferUV, ScreenPosition, SvPosition)*SampleCoverage;
+				Coverage += SampleCoverage;
+			}
+			#endif
 		}
 		
 		// Lit fully covered or partially covered pixel
@@ -514,7 +525,9 @@ void MainPS(
 		// * Partially covered pixel output their result into the sub pixel color target. the main color target will store the underlying pixel color 
 		if (bIsClusterCompatible)
 		{
-			const float3 SkyLighting = ReflectionEnvironment(HairGBuffer, BufferUV, ScreenPosition, SvPosition);
+			#if PERMUTATION_PER_SAMPLE == 0
+			SkyLighting = ReflectionEnvironment(HairGBuffer, BufferUV, ScreenPosition, SvPosition) * saturate(Coverage);
+			#endif
 
 			// Scene color alpha is used for ScreenSpaceSubsurfaceScattering (if that is not needed it can be disabled with SUBSURFACE_CHANNEL_MODE)
 			FLightAccumulator LightAccumulator = (FLightAccumulator)0;
@@ -522,13 +535,12 @@ void MainPS(
 			const float4 LocalOutColor = LightAccumulator_GetResult(LightAccumulator);
 			if (Coverage < 1)
 			{
-				OutColorSubPixel = LocalOutColor * Coverage;
+				OutColorSubPixel = LocalOutColor;
 				OutColorSubPixel.a = Coverage;
 			}
 			else
 			{
 				OutColor += LocalOutColor;
-				return;
 			}
 		}
 	}
diff --git a/Engine/Source/Runtime/Renderer/Private/HairStrands/HairStrandsEnvironment.cpp b/Engine/Source/Runtime/Renderer/Private/HairStrands/HairStrandsEnvironment.cpp
--- a/Engine/Source/Runtime/Renderer/Private/HairStrands/HairStrandsEnvironment.cpp
+++ b/Engine/Source/Runtime/Renderer/Private/HairStrands/HairStrandsEnvironment.cpp
@@ -29,6 +29,9 @@ static FAutoConsoleVariableRef CVarHairSkyAOEnable(TEXT("r.HairStrands.SkyAOEnab
 static float GHairSkylightingConeAngle = 3;
 static FAutoConsoleVariableRef CVarHairSkylightingConeAngle(TEXT("r.HairStrands.SkyLightingConeAngle"), GHairSkylightingConeAngle, TEXT("Cone angle for tracing sky lighting on hair."));
 
+static float GHairSkylightingPerSample = 0;
+static FAutoConsoleVariableRef CVarHairSkylightingPerSample(TEXT("r.HairStrands.SkyLightingPerSample"), GHairSkylightingPerSample, TEXT("Evaluate sky lighting per hair sample."));
+
 static bool GetHairStrandsSkyLightingEnable() { return GHairSkylightingEnable > 0; }
 static bool GetHairStrandsSkyAOEnable() { return GHairSkyAOEnable > 0; }
 static float GetHairStrandsSkyLightingConeAngle() { return FMath::Max(0.f, GHairSkylightingConeAngle); }
@@ -43,11 +46,12 @@ class FHairEnvironmentLightingPS : public FGlobalShader
 	SHADER_USE_PARAMETER_STRUCT(FHairEnvironmentLightingPS, FGlobalShader)
 
 
-		class FRenderMode : SHADER_PERMUTATION_INT("PERMUTATION_RENDER_MODE", 2);
-	using FPermutationDomain = TShaderPermutationDomain<FRenderMode>;
+	class FRenderMode : SHADER_PERMUTATION_INT("PERMUTATION_RENDER_MODE", 2);
+	class FPerSample : SHADER_PERMUTATION_INT("PERMUTATION_PER_SAMPLE", 2);
+	using FPermutationDomain = TShaderPermutationDomain<FRenderMode, FPerSample>;
 	static bool ShouldCompilePermutation(const FGlobalShaderPermutationParameters& Parameters)
 	{
-		return IsHairStrandsSupported(Parameters.Platform);;
+		return IsHairStrandsSupported(Parameters.Platform);
 	}
 
 	BEGIN_SHADER_PARAMETER_STRUCT(FParameters, )
@@ -166,6 +170,7 @@ static void AddHairStrandsEnvironmentPass(
 
 	FHairEnvironmentLightingPS::FPermutationDomain PermutationVector;
 	PermutationVector.Set<FHairEnvironmentLightingPS::FRenderMode>(RenderMode == EEnvRenderMode::AO ? 1 : 0);
+	PermutationVector.Set<FHairEnvironmentLightingPS::FPerSample>(GHairSkylightingPerSample > 0 ? 1 : 0);
 	TShaderMapRef<FHairEnvironmentLightingPS> PixelShader(View.ShaderMap, PermutationVector);
 	ClearUnusedGraphResources(*PixelShader, PassParameters);
 
