Pipe coverage data into the various passes.

This should allow to plug custom coverage more easily.

diff --git a/Engine/Shaders/Private/DeferredLightPixelShaders.usf b/Engine/Shaders/Private/DeferredLightPixelShaders.usf
--- a/Engine/Shaders/Private/DeferredLightPixelShaders.usf
+++ b/Engine/Shaders/Private/DeferredLightPixelShaders.usf
@@ -99,6 +99,7 @@ uint GetLightingChannelMask(float2 UV)
 
 #if USE_HAIR_LIGHTING
 
+Texture2D<uint4> HairCategorizationTexture;
 Texture2D<uint> HairVisibilityNodeOffsetAndCount;
 StructuredBuffer<FNodeData> HairVisibilityNodeData;
 
@@ -137,12 +138,12 @@ float4 PerceptualToLinear(float4 Color)
 
 void ComputeHairLighting(
 	in FInputParams InputParams,
-	in FNodeDesc NodeDesc, 
-	out float OutOpaqueWeight,
-	out float4 OutColor, 
-	out float4 OutColorSubPixel)
+	in FNodeDesc NodeDesc,
+	in FCategorizationData CatData,
+	inout float OutOpaqueWeight,
+	inout float4 OutColor, 
+	inout float4 OutColorSubPixel)
 {
-	float Coverage = 0;
 	float4 LocalColor = 0;
 	LOOP
 	for (uint SampleIt = 0; SampleIt < NodeDesc.Count; SampleIt++)
@@ -215,15 +216,16 @@ void ComputeHairLighting(
 		const float4 PercepualColor = LinearToPerceptual(LinearColor);
 		const float LocalCoverage = From8bitCoverage(SampleVisData.Coverage8bit);
 		LocalColor += PercepualColor * LocalCoverage;
-		Coverage += LocalCoverage;
 	}
 
+	const float PixelCoverage = min(CatData.PixelCoverage, 1);
+	float Coverage = 0;
 	OutColor = 0;
 	OutColorSubPixel = 0;
-	if (Coverage < 1)
+	if (PixelCoverage < 1)
 	{
-		OutColorSubPixel = LocalColor;
-		OutColorSubPixel.a = Coverage;
+		OutColorSubPixel = LocalColor; // * PixelCoverage <-- this would be need if Coverage8bit was storing the ratio (sampleCoverage / validCoverage)
+		OutColorSubPixel.a = PixelCoverage;
 		Coverage = 0;
 	}
 	else
@@ -270,6 +272,7 @@ void DeferredLightPixelMain(
 #if USE_HAIR_LIGHTING
 	const uint3 PixelCoord = uint3(floor(InputParams.PixelPos), 0);
 	const FNodeDesc NodeDesc = DecodeNodeDesc(HairVisibilityNodeOffsetAndCount.Load(PixelCoord));
+	const FCategorizationData CatData = DecodeCategorizationData(HairCategorizationTexture.Load(PixelCoord));
 	
 	// Only light pixels marked as using deferred shading
 	OutColorSubPixel = 0;
@@ -281,7 +284,7 @@ void DeferredLightPixelMain(
 	#endif
 		)
 	{
-		ComputeHairLighting(InputParams, NodeDesc, OpaqueWeight, OutColor, OutColorSubPixel);
+		ComputeHairLighting(InputParams, NodeDesc, CatData, OpaqueWeight, OutColor, OutColorSubPixel);
 	#if USE_PREEXPOSURE
 		OutColorSubPixel.rgb *= View.PreExposure;
 	#endif
diff --git a/Engine/Shaders/Private/HairStrands/HairStrandsDebug.usf b/Engine/Shaders/Private/HairStrands/HairStrandsDebug.usf
--- a/Engine/Shaders/Private/HairStrands/HairStrandsDebug.usf
+++ b/Engine/Shaders/Private/HairStrands/HairStrandsDebug.usf
@@ -44,7 +44,7 @@ void MainPS(
 	const float2 PixelPos = UV * OutputResolution;
 
 	const FCategorizationData CatData = DecodeCategorizationData(CategorizationTexture.Load(uint3(PixelPos,0)));
-	const bool bIsValid = CatData.TotalCoverage8bit > 0;
+	const bool bIsValid = CatData.PixelCoverage > 0;
 	if (!bIsValid)
 		discard;
 
@@ -53,7 +53,7 @@ void MainPS(
 #endif
 
 #if PERMUTATION_DEBUG_MODE == DEBUG_MODE_COVERAGE_TYPE
-	const bool bIsFullyCovered = CatData.TotalCoverage8bit >= 255;
+	const bool bIsFullyCovered = CatData.PixelCoverage >= 1;
 	const float3 Color = bIsFullyCovered ? float3(0.25,0.85,0.25) : float3(0.85,0.25,0.25);
 #endif
 
diff --git a/Engine/Shaders/Private/HairStrands/HairStrandsDebugPrint.usf b/Engine/Shaders/Private/HairStrands/HairStrandsDebugPrint.usf
--- a/Engine/Shaders/Private/HairStrands/HairStrandsDebugPrint.usf
+++ b/Engine/Shaders/Private/HairStrands/HairStrandsDebugPrint.usf
@@ -48,8 +48,7 @@ void MainCS(uint GroupIndex : SV_GroupIndex, uint3 DispatchThreadId : SV_Dispatc
 	ShaderPrintFilter(true);
 
 	const FCategorizationData CatData = DecodeCategorizationData(CategorizationTexture.Load(uint3(PixelPos,0)));
-	//const bool bIsValid = CatData.PixelCoverage > 0;
-	const bool bIsValid = CatData.TotalCoverage8bit > 0;
+	const bool bIsValid = CatData.PixelCoverage > 0;
 	if (!bIsValid)
 		return;
 
@@ -163,6 +162,7 @@ void MainCS(uint GroupIndex : SV_GroupIndex, uint3 DispatchThreadId : SV_Dispatc
 	Pos = ShaderPrintSymbol(Pos, _S_);
 	Pos = ShaderPrintSymbol(Pos, _SPC_);
 	Pos = ShaderPrint(Pos, CatData.TotalCoverage8bit);
+	Pos = ShaderPrint(Pos, CatData.TotalCoverage8bit / 255.f);
 	
 	Pos = ShaderPrintNewline(Pos);
 	Pos.x = OriginalPos.x;
diff --git a/Engine/Shaders/Private/HairStrands/HairStrandsDeepShadowMask.usf b/Engine/Shaders/Private/HairStrands/HairStrandsDeepShadowMask.usf
--- a/Engine/Shaders/Private/HairStrands/HairStrandsDeepShadowMask.usf
+++ b/Engine/Shaders/Private/HairStrands/HairStrandsDeepShadowMask.usf
@@ -277,7 +277,7 @@ void MainPS(
 	const uint2 PixelCoord = floor(Input.Position.xy);
 	FGBufferData GBuffer = GetGBufferDataFromSceneTextures(UV);
 	const FCategorizationData Categorization = DecodeCategorizationData(CategorizationTexture.Load(uint3(PixelCoord, 0)));
-	const bool bIsHairPixel = Categorization.TotalCoverage8bit >= 255;
+	const bool bIsHairPixel = Categorization.PixelCoverage >= 1;
 
 	const float3 WorldPosition		= TotWorldPosition(UV, GBuffer.Depth);
 	const float3 LightSpacePosition = ToLightPosition(WorldPosition);
@@ -335,7 +335,7 @@ void MainPS(
 	const uint2 PixelCoord = floor(Input.Position.xy);
 	FGBufferData GBuffer = GetGBufferDataFromSceneTextures(UV);
 	const FCategorizationData Categorization = DecodeCategorizationData(CategorizationTexture.Load(uint3(PixelCoord, 0)));
-	const bool bIsHairPixel = Categorization.TotalCoverage8bit >= 255;
+	const bool bIsHairPixel = Categorization.PixelCoverage >= 1;
 
 	float3 WorldPosition = TotWorldPosition(UV, GBuffer.Depth);
 
diff --git a/Engine/Shaders/Private/HairStrands/HairStrandsEnvironmentLighting.usf b/Engine/Shaders/Private/HairStrands/HairStrandsEnvironmentLighting.usf
--- a/Engine/Shaders/Private/HairStrands/HairStrandsEnvironmentLighting.usf
+++ b/Engine/Shaders/Private/HairStrands/HairStrandsEnvironmentLighting.usf
@@ -41,8 +41,8 @@ float3 GatherRadiance(float3 RayDirection, float Roughness)
 }
 
 
-Texture2D<uint4> HairCategorizationTexture;
-Texture2D<uint> HairVisibilityNodeOffsetAndCount;
+Texture2D<uint4>			HairCategorizationTexture;
+Texture2D<uint>				HairVisibilityNodeOffsetAndCount;
 StructuredBuffer<FNodeData> HairVisibilityNodeData;
 
 Texture3D<float4>			HairScatteringLUTTexture;
@@ -464,10 +464,11 @@ void MainPS(
 
 	if (NodeDesc.Count > 0)
 	{
+		const FCategorizationData CatData = DecodeCategorizationData(HairCategorizationTexture.Load(PixelCoord));
+
 		// Note:
 		// * Currently: compute take the closest sample
 		// * Change this to per sample evaluation once we have better hair transmittance evaluation for sky
-		float Coverage = 0;
 		float ClosestDepth = 0;
 		bool bIsClusterCompatible = false;
 		FGBufferData HairGBuffer = (FGBufferData)0;
@@ -508,14 +509,11 @@ void MainPS(
 
 			#if PERMUTATION_PER_SAMPLE == 0
 			ClosestDepth = max(ClosestDepth, SampleDepth);  // Inverse-Z
-			const float LocalCoverage = From8bitCoverage(SampleVisData.Coverage8bit);
-			Coverage += LocalCoverage;
 			#else
 			if (bIsClusterCompatible = SampleVisData.ClusterId == Voxel_ClusterId)
 			{
 				const float SampleCoverage = From8bitCoverage(SampleVisData.Coverage8bit);
-				SkyLighting += ReflectionEnvironment(HairGBuffer, BufferUV, ScreenPosition, SvPosition)*SampleCoverage;
-				Coverage += SampleCoverage;
+				SkyLighting += ReflectionEnvironment(HairGBuffer, BufferUV, ScreenPosition, SvPosition) * SampleCoverage;
 			}
 			#endif
 		}
@@ -526,17 +524,17 @@ void MainPS(
 		if (bIsClusterCompatible)
 		{
 			#if PERMUTATION_PER_SAMPLE == 0
-			SkyLighting = ReflectionEnvironment(HairGBuffer, BufferUV, ScreenPosition, SvPosition) * saturate(Coverage);
+			SkyLighting = ReflectionEnvironment(HairGBuffer, BufferUV, ScreenPosition, SvPosition) * saturate(CatData.PixelCoverage);
 			#endif
 
 			// Scene color alpha is used for ScreenSpaceSubsurfaceScattering (if that is not needed it can be disabled with SUBSURFACE_CHANNEL_MODE)
 			FLightAccumulator LightAccumulator = (FLightAccumulator)0;
 			LightAccumulator_Add(LightAccumulator, SkyLighting, SkyLighting, 1.0f, false);
 			const float4 LocalOutColor = LightAccumulator_GetResult(LightAccumulator);
-			if (Coverage < 1)
+			if (CatData.PixelCoverage < 1)
 			{
-				OutColorSubPixel = LocalOutColor;
-				OutColorSubPixel.a = Coverage;
+				OutColorSubPixel = LocalOutColor; // * PixelCoverage <-- this would be need if Coverage8bit was storing the ratio (sampleCoverage / validCoverage)
+				OutColorSubPixel.a = CatData.PixelCoverage;
 			}
 			else
 			{
@@ -555,12 +553,11 @@ void MainPS(
 	const float3 WorldNormal = GBuffer.WorldNormal;
 
 	// Trace AO for all pixel which are partially covered or not covered at all.
+	const FCategorizationData CatData = DecodeCategorizationData(HairCategorizationTexture.Load(PixelCoord));
 	bool bIsValid = false;
 	bool bIsClusterCompatible = false;
-	if (NodeDesc.Count < 255)
+	if (CatData.PixelCoverage > 0 && CatData.PixelCoverage < 1)
 	{
-		float ClosestDepth = 0;
-		uint Coverage8Bits = 0;
 		LOOP
 		for (uint SampleIt = 0; SampleIt < NodeDesc.Count; SampleIt++)
 		{
@@ -573,11 +570,8 @@ void MainPS(
 			float			SampleSpecular;
 			DecodeNodeData(NodeData, SampleDepth, SampleVisData, SampleBaseColor, SampleRoughness, SampleSpecular);
 			bIsClusterCompatible = bIsClusterCompatible || SampleVisData.ClusterId == Voxel_ClusterId;
-
-			const float LocalCoverage = From8bitCoverage(SampleVisData.Coverage8bit);
-			Coverage8Bits += SampleVisData.Coverage8bit;
 		}
-		bIsValid = Coverage8Bits <= 255;
+		bIsValid = true;
 	}
 	else
 	{
diff --git a/Engine/Shaders/Private/RayTracing/RayTracingOcclusionRGS.usf b/Engine/Shaders/Private/RayTracing/RayTracingOcclusionRGS.usf
--- a/Engine/Shaders/Private/RayTracing/RayTracingOcclusionRGS.usf
+++ b/Engine/Shaders/Private/RayTracing/RayTracingOcclusionRGS.usf
@@ -163,7 +163,7 @@ RAY_TRACING_ENTRY_RAYGEN(OcclusionRGS)
 		{
 			RaytracingMask = RAY_TRACING_MASK_SHADOW;
 			const FCategorizationData Categorization = DecodeCategorizationData(HairCategorizationTexture.Load(uint3(PixelCoord, 0)));
-			bIsValidPixel = Categorization.TotalCoverage8bit > 0;
+			bIsValidPixel = Categorization.PixelCoverage > 0;
 			if (bIsValidPixel)
 			{
 				DeviceZ = Categorization.ClosestDepth;
diff --git a/Engine/Shaders/Private/ShadowProjectionPixelShader.usf b/Engine/Shaders/Private/ShadowProjectionPixelShader.usf
--- a/Engine/Shaders/Private/ShadowProjectionPixelShader.usf
+++ b/Engine/Shaders/Private/ShadowProjectionPixelShader.usf
@@ -116,7 +116,7 @@ void Main(
 #if SUBPIXEL_SHADOW
 	const uint3 PixelCoord = uint3(floor(SVPos.xy), 0);
 	const FCategorizationData CategorizationData = DecodeCategorizationData(HairCategorizationTexture.Load(PixelCoord));
-	if (CategorizationData.TotalCoverage8bit > 0)
+	if (CategorizationData.PixelCoverage > 0)
 	{
 		SceneW = ConvertFromDeviceZ(CategorizationData.ClosestDepth);
 		bIsSubsurfaceCompatible = false;
diff --git a/Engine/Source/Runtime/Renderer/Private/HairStrands/HairStrandsRasterCommon.cpp b/Engine/Source/Runtime/Renderer/Private/HairStrands/HairStrandsRasterCommon.cpp
--- a/Engine/Source/Runtime/Renderer/Private/HairStrands/HairStrandsRasterCommon.cpp
+++ b/Engine/Source/Runtime/Renderer/Private/HairStrands/HairStrandsRasterCommon.cpp
@@ -323,6 +323,7 @@ void FDeepShadowMeshProcessor::Process(
 	ERasterizerCullMode MeshCullMode)
 {
 	const FVertexFactory* VertexFactory = MeshBatch.VertexFactory;
+	static const FVertexFactoryType* CompatibleVF = FVertexFactoryType::GetVFByName(TEXT("FHairStrandsVertexFactory"));
 
 	TMeshProcessorShaders<
 		VertexShaderType,
@@ -332,6 +333,10 @@ void FDeepShadowMeshProcessor::Process(
 	{
 		const EMaterialTessellationMode MaterialTessellationMode = MaterialResource.GetTessellationMode();
 		FVertexFactoryType* VertexFactoryType = VertexFactory->GetType();
+		const bool bIsHairStrandsFactory = MeshBatch.VertexFactory->GetType()->GetId() == CompatibleVF->GetId();
+		if (!bIsHairStrandsFactory)
+			return;
+
 		PassShaders.DomainShader = nullptr;
 		PassShaders.HullShader = nullptr;
 		PassShaders.VertexShader = MaterialResource.GetShader<VertexShaderType>(VertexFactoryType);
diff --git a/Engine/Source/Runtime/Renderer/Private/LightRendering.cpp b/Engine/Source/Runtime/Renderer/Private/LightRendering.cpp
--- a/Engine/Source/Runtime/Renderer/Private/LightRendering.cpp
+++ b/Engine/Source/Runtime/Renderer/Private/LightRendering.cpp
@@ -209,6 +209,7 @@ struct FRenderLightParams
 	bool bUseTransmittanceData = false;
 	
 	// Visibility buffer data
+	IPooledRenderTarget* HairCategorizationTexture = nullptr;
 	IPooledRenderTarget* HairVisibilityNodeOffsetAndCount = nullptr;
 	FShaderResourceViewRHIRef HairVisibilityNodeDataSRV = nullptr;
 
@@ -316,6 +317,7 @@ class FDeferredLightPS : public FGlobalShader
 		HairComponents.Bind(Initializer.ParameterMap, TEXT("HairComponents"));
 		HairTransmittanceEnable.Bind(Initializer.ParameterMap, TEXT("HairTransmittanceEnable"));
 
+		HairCategorizationTexture.Bind(Initializer.ParameterMap, TEXT("HairCategorizationTexture"));
 		HairVisibilityNodeOffsetAndCount.Bind(Initializer.ParameterMap, TEXT("HairVisibilityNodeOffsetAndCount"));
 		HairVisibilityNodeData.Bind(Initializer.ParameterMap, TEXT("HairVisibilityNodeData"));
 	}
@@ -363,6 +365,7 @@ public:
 		Ar << HairTransmittanceBuffer;
 		Ar << HairTransmittanceBufferMaxCount;
 
+		Ar << HairCategorizationTexture;
 		Ar << HairVisibilityNodeOffsetAndCount;
 		Ar << HairVisibilityNodeData;
 		Ar << ScreenShadowMaskSubPixelTexture;
@@ -491,11 +494,25 @@ private:
 					ShaderRHI,
 					ScreenShadowMaskSubPixelTexture,
 					LightAttenuationTextureSampler,
-					TStaticSamplerState<SF_Point, AM_Wrap, AM_Wrap, AM_Wrap>::GetRHI(),
+					TStaticSamplerState<SF_Point, AM_Clamp, AM_Clamp, AM_Clamp>::GetRHI(),
 					(RenderLightParams && RenderLightParams->ScreenShadowMaskSubPixelTexture) ? RenderLightParams->ScreenShadowMaskSubPixelTexture->GetRenderTargetItem().ShaderResourceTexture : GWhiteTexture->TextureRHI);
 			}
 		}
 
+		if (HairCategorizationTexture.IsBound())
+		{
+			if (RenderLightParams && RenderLightParams->HairCategorizationTexture)
+			{
+				SetTextureParameter(
+					RHICmdList,
+					ShaderRHI,
+					HairCategorizationTexture,
+					LightAttenuationTextureSampler,
+					TStaticSamplerState<SF_Point, AM_Clamp, AM_Clamp, AM_Clamp>::GetRHI(),
+					RenderLightParams->HairCategorizationTexture->GetRenderTargetItem().TargetableTexture);
+			}
+		}
+
 		if (HairVisibilityNodeOffsetAndCount.IsBound())
 		{
 			if (RenderLightParams && RenderLightParams->HairVisibilityNodeOffsetAndCount)
@@ -505,7 +522,7 @@ private:
 					ShaderRHI,
 					HairVisibilityNodeOffsetAndCount,
 					LightAttenuationTextureSampler,
-					TStaticSamplerState<SF_Point, AM_Wrap, AM_Wrap, AM_Wrap>::GetRHI(),
+					TStaticSamplerState<SF_Point, AM_Clamp, AM_Clamp, AM_Clamp>::GetRHI(),
 					RenderLightParams->HairVisibilityNodeOffsetAndCount->GetRenderTargetItem().TargetableTexture);
 			}
 		}
@@ -570,8 +587,9 @@ private:
 	FShaderResourceParameter TransmissionProfilesTexture;
 	FShaderResourceParameter TransmissionProfilesLinearSampler;
 
-	FShaderResourceParameter HairTransmittanceBuffer;
 	FShaderParameter HairTransmittanceBufferMaxCount;
+	FShaderResourceParameter HairTransmittanceBuffer;
+	FShaderResourceParameter HairCategorizationTexture;
 	FShaderResourceParameter HairVisibilityNodeOffsetAndCount;
 	FShaderResourceParameter HairVisibilityNodeData;
 	FShaderResourceParameter ScreenShadowMaskSubPixelTexture;
@@ -1844,6 +1862,7 @@ void FDeferredShadingSceneRenderer::RenderLight(FRHICommandList& RHICmdList, con
 				const FHairStrandsVisibilityData& HairVisibilityData = InHairVisibilityViews->HairDatas[ViewIndex];
 				RenderLightParams.HairVisibilityNodeOffsetAndCount = HairVisibilityData.NodeIndex;
 				RenderLightParams.HairVisibilityNodeDataSRV = HairVisibilityData.NodeDataSRV;
+				RenderLightParams.HairCategorizationTexture = HairVisibilityData.CategorizationTexture;
 			}
 		}
 
