Change RT shadow to have uniform shadow mask when the hair permutation is not selected.

diff --git a/Engine/Shaders/Private/RayTracing/RayTracingOcclusionRGS.usf b/Engine/Shaders/Private/RayTracing/RayTracingOcclusionRGS.usf
--- a/Engine/Shaders/Private/RayTracing/RayTracingOcclusionRGS.usf
+++ b/Engine/Shaders/Private/RayTracing/RayTracingOcclusionRGS.usf
@@ -153,13 +153,9 @@ RAY_TRACING_ENTRY_RAYGEN(OcclusionRGS)
 		bool bIsValidPixel = ShadingModelID != SHADINGMODELID_UNLIT && bIsDepthValid;
 
 		// Check if there is hair geometry
-		uint RaytracingMask = 0;
-		if (bIsValidPixel && ShadowMaskType == ShadowMaskType_Opaque)
-		{
-			RaytracingMask = RAY_TRACING_MASK_SHADOW | RAY_TRACING_MASK_THIN_SHADOW;
-		}
+		uint RaytracingMask = RAY_TRACING_MASK_SHADOW | RAY_TRACING_MASK_THIN_SHADOW;
 	#if USE_HAIR_LIGHTING
-		else if (ShadowMaskType == ShadowMaskType_Hair)
+		if (ShadowMaskType == ShadowMaskType_Hair)
 		{
 			RaytracingMask = RAY_TRACING_MASK_SHADOW;
 			const FCategorizationData Categorization = DecodeCategorizationData(HairCategorizationTexture.Load(uint3(PixelCoord, 0)));
diff --git a/Engine/Source/Runtime/Renderer/Private/RayTracing/RayTracingShadows.cpp b/Engine/Source/Runtime/Renderer/Private/RayTracing/RayTracingShadows.cpp
--- a/Engine/Source/Runtime/Renderer/Private/RayTracing/RayTracingShadows.cpp
+++ b/Engine/Source/Runtime/Renderer/Private/RayTracing/RayTracingShadows.cpp
@@ -73,7 +73,8 @@ class FOcclusionRGS : public FGlobalShader
 		OutEnvironment.SetDefine(TEXT("UE_RAY_TRACING_DYNAMIC_MISS_SHADER"),        0);
 
 		FPermutationDomain PermutationVector(Parameters.PermutationId);
-		if (PermutationVector.Get<FLightTypeDim>() == LightType_Directional)
+		if (PermutationVector.Get<FLightTypeDim>() == LightType_Directional &&
+			PermutationVector.Get<FHairLighting>() == 0)
 		{
 			OutEnvironment.SetDefine(TEXT("UE_RAY_TRACING_COHERENT_RAYS"), 1);
 		}
