Fix shader unsupported calls min(int, uint) and min(int, uint)

diff --git a/Engine/Shaders/Private/HairStrands/HairStrandsVisibilityCommon.ush b/Engine/Shaders/Private/HairStrands/HairStrandsVisibilityCommon.ush
--- a/Engine/Shaders/Private/HairStrands/HairStrandsVisibilityCommon.ush
+++ b/Engine/Shaders/Private/HairStrands/HairStrandsVisibilityCommon.ush
@@ -57,7 +57,7 @@ uint4 EncodeVisibilityData(FVisibilityData Data)
 	float3 T = saturate(EncodeTangent(Data.Tangent));
 	Output.xy = uint2(
 		(0xFF & uint(T.x * 255)) | ((0xFF & uint(T.y * 255)) << 8),
-		(0xFF & uint(T.z * 255)) | (min(0xFF, Data.Coverage8bit) << 8));
+		(0xFF & uint(T.z * 255)) | (min(0xFFu, Data.Coverage8bit) << 8));
 
 	// PrimitiveId and ClusterId
 	// Encode into 2 x 16bits
@@ -265,7 +265,7 @@ FNodeData ConvertToNodeData(float Depth, FVisibilityData VisibilityData, float3
 	 const uint ClosestDepthUint = asuint(Data.ClosestDepth);
 	 const uint ClosestDepthLow  = ClosestDepthUint & 0xFFFF;
 	 const uint ClosestDepthHigh = (ClosestDepthUint >> 16) & 0xFFFF;
-	 const uint ComponentX		 = min(Data.TotalCoverage8bit, 0xFF) | ((Data.SampleCount & 0xFF)<<8);
+	 const uint ComponentX		 = min(Data.TotalCoverage8bit, 0xFFu) | ((Data.SampleCount & 0xFF)<<8);
 	 return uint4(ComponentX, f32tof16(Data.PixelCoverage), ClosestDepthLow, ClosestDepthHigh);
  }
 
@@ -307,4 +307,4 @@ FNodeData ConvertToNodeData(float Depth, FVisibilityData VisibilityData, float3
 	 BaseColor.y = (float((Encoded >> 8) & 0xFF) / 255.f);
 	 BaseColor.z = (float((Encoded >> 16) & 0xFF) / 255.f);
 	 Roughness   = (float((Encoded >> 24) & 0xFF) / 255.f);
- }
\ No newline at end of file
+ }
