Change all 'out' in hair shader code to be inout, for avoiding potential issues.

diff --git a/Engine/Shaders/Private/HairStrands/HairStrandsDeepShadowMask.usf b/Engine/Shaders/Private/HairStrands/HairStrandsDeepShadowMask.usf
--- a/Engine/Shaders/Private/HairStrands/HairStrandsDeepShadowMask.usf
+++ b/Engine/Shaders/Private/HairStrands/HairStrandsDeepShadowMask.usf
@@ -74,7 +74,7 @@ struct FHairSamplingSettings
 	float			SceneDepth;					// SceneDepth in lightspace.
 };
 
-void InternalFetchRowOfThree(float2 Sample00TexelCenter, float VerticalOffset, out float3 Values0, FHairSamplingSettings Settings)
+void InternalFetchRowOfThree(float2 Sample00TexelCenter, float VerticalOffset, inout float3 Values0, FHairSamplingSettings Settings)
 {
 	Values0.x = Settings.ShadowDepthTexture.SampleLevel(Settings.ShadowDepthTextureSampler, (Sample00TexelCenter + float2(0, VerticalOffset)) * Settings.ShadowAtlasInvResolution, 0).r;
 	Values0.y = Settings.ShadowDepthTexture.SampleLevel(Settings.ShadowDepthTextureSampler, (Sample00TexelCenter + float2(1, VerticalOffset)) * Settings.ShadowAtlasInvResolution, 0).r;
@@ -82,7 +82,7 @@ void InternalFetchRowOfThree(float2 Sample00TexelCenter, float VerticalOffset, o
 	Values0 = float3(Settings.SceneDepth < Values0);
 }
 
-void FetchRowOfFour(float2 Sample00TexelCenter, float VerticalOffset, out float4 Values0, FHairSamplingSettings Settings)
+void FetchRowOfFour(float2 Sample00TexelCenter, float VerticalOffset, inout float4 Values0, FHairSamplingSettings Settings)
 {
 	Values0.x = Settings.ShadowDepthTexture.SampleLevel(Settings.ShadowDepthTextureSampler, (Sample00TexelCenter + float2(0, VerticalOffset)) * Settings.ShadowAtlasInvResolution, 0).r;
 	Values0.y = Settings.ShadowDepthTexture.SampleLevel(Settings.ShadowDepthTextureSampler, (Sample00TexelCenter + float2(1, VerticalOffset)) * Settings.ShadowAtlasInvResolution, 0).r;
diff --git a/Engine/Shaders/Private/HairStrands/HairStrandsDeepTransmittanceDualScattering.ush b/Engine/Shaders/Private/HairStrands/HairStrandsDeepTransmittanceDualScattering.ush
--- a/Engine/Shaders/Private/HairStrands/HairStrandsDeepTransmittanceDualScattering.ush
+++ b/Engine/Shaders/Private/HairStrands/HairStrandsDeepTransmittanceDualScattering.ush
@@ -18,8 +18,8 @@ void ComputeDualScatteringTerms(
 	const float3 L,
 	const float3 T,
 	const float4 HairComponents,
-	out float3 OutGlobalScattering,
-	out float3 OutLocalScattering)
+	inout float3 OutGlobalScattering,
+	inout float3 OutLocalScattering)
 {
 	const float SinThetaL = clamp(dot(T, L), -1, 1);
 	const float SinThetaV = clamp(dot(T, V), -1, 1);
diff --git a/Engine/Shaders/Private/HairStrands/HairStrandsVertexFactory.ush b/Engine/Shaders/Private/HairStrands/HairStrandsVertexFactory.ush
--- a/Engine/Shaders/Private/HairStrands/HairStrandsVertexFactory.ush
+++ b/Engine/Shaders/Private/HairStrands/HairStrandsVertexFactory.ush
@@ -278,7 +278,7 @@ float4 CalcWorldPosition(float3 Position, uint PrimitiveId)
 	return TransformLocalToTranslatedWorld(Position, PrimitiveId);
 }
 
-half3x3 CalcTangentToLocal(FVertexFactoryInput Input, out float TangentSign)
+half3x3 CalcTangentToLocal(FVertexFactoryInput Input, inout float TangentSign)
 {
 	const uint QuadIndex = Input.VertexId % 6;
 	const uint BaseIndex = Input.VertexId / 6;
diff --git a/Engine/Shaders/Private/HairStrands/HairStrandsVisibilityCommon.ush b/Engine/Shaders/Private/HairStrands/HairStrandsVisibilityCommon.ush
--- a/Engine/Shaders/Private/HairStrands/HairStrandsVisibilityCommon.ush
+++ b/Engine/Shaders/Private/HairStrands/HairStrandsVisibilityCommon.ush
@@ -133,12 +133,12 @@ FNodeData ConvertToNodeData(float Depth, FVisibilityData VisibilityData, float3
 }
 
  void DecodeNodeData(
-	 const FNodeData	 InNode, 
-	 out float			 OutDepth, 
-	 out FVisibilityData OutVisibilityData,
-	 out float3			 OutBaseColor, 
-	 out float			 OutRoughness,
-	 out float			 OutSpecular)
+	 const FNodeData		InNode, 
+	 inout float			OutDepth, 
+	 inout FVisibilityData	OutVisibilityData,
+	 inout float3			OutBaseColor, 
+	 inout float			OutRoughness,
+	 inout float			OutSpecular)
 {
 	OutDepth = InNode.Depth;
 
@@ -178,8 +178,8 @@ FNodeData ConvertToNodeData(float Depth, FVisibilityData VisibilityData, float3
  };
 
  void ComputeUniqueSamples(
-	 out FSampleDesc OutSamples[HAIR_VISIBILITY_MAXSAMPLECOUNT],
-	 out FSampleSetDesc OutSet,
+	 inout FSampleDesc OutSamples[HAIR_VISIBILITY_MAXSAMPLECOUNT],
+	 inout FSampleSetDesc OutSet,
 	 const uint2 PixelCoord,
 	 const uint SampleCount,
 	 Texture2D<float4> InSceneDepthTexture,
@@ -301,7 +301,7 @@ FNodeData ConvertToNodeData(float Depth, FVisibilityData VisibilityData, float3
 		 (uint(Roughness   * 0xFF) & 0xFF) << 24;
  }
 
- void DecodeBaseColorAndRoughness(uint Encoded, out float3 BaseColor, out float Roughness)
+ void DecodeBaseColorAndRoughness(uint Encoded, inout float3 BaseColor, out float Roughness)
  {
 	 BaseColor.x = (float((Encoded) & 0xFF) / 255.f);
 	 BaseColor.y = (float((Encoded >> 8) & 0xFF) / 255.f);
diff --git a/Engine/Shaders/Private/HairStrands/HairStrandsVisibilityCompaction.usf b/Engine/Shaders/Private/HairStrands/HairStrandsVisibilityCompaction.usf
--- a/Engine/Shaders/Private/HairStrands/HairStrandsVisibilityCompaction.usf
+++ b/Engine/Shaders/Private/HairStrands/HairStrandsVisibilityCompaction.usf
@@ -16,8 +16,8 @@ float CosTangentThreshold;
 #endif
 
 void ComputeUniqueSamplesWithoutDepthTest(
-	out uint3 OutSamples[HAIR_VISIBILITY_MAXSAMPLECOUNT],
-	out FSampleSetDesc OutSet,
+	inout uint3 OutSamples[HAIR_VISIBILITY_MAXSAMPLECOUNT],
+	inout FSampleSetDesc OutSet,
 	const uint2 PixelCoord,
 	const uint SampleCount,
 	const float OpaqueDepth,
