Removed hair rendering hack by patching view uniform buffer locally for DOM.
Nest step: avoid duplicated CreateUniformBufferImmediate

diff --git a/Engine/Plugins/Experimental/HairStrands/Source/HairStrandsCore/Private/HairStrandsVertexFactory.cpp b/Engine/Plugins/Experimental/HairStrands/Source/HairStrandsCore/Private/HairStrandsVertexFactory.cpp
--- a/Engine/Plugins/Experimental/HairStrands/Source/HairStrandsCore/Private/HairStrandsVertexFactory.cpp
+++ b/Engine/Plugins/Experimental/HairStrands/Source/HairStrandsCore/Private/HairStrandsVertexFactory.cpp
@@ -37,8 +37,8 @@ public:
 
 	FShaderParameter Radius;
 	FShaderParameter Length;
-	FShaderParameter RadiusAtDepth1_Primary;
-	FShaderParameter RadiusAtDepth1_Velocity;
+	FShaderParameter RadiusAtDepth1_Primary;	// unused
+	FShaderParameter RadiusAtDepth1_Velocity;	// unused
 	FShaderParameter PositionOffset;
 	FShaderParameter PreviousPositionOffset;
 	FShaderParameter Density;
@@ -52,8 +52,6 @@ public:
 	{
 		Radius.Bind(ParameterMap, TEXT("HairStrandsVF_Radius"));
 		Length.Bind(ParameterMap, TEXT("HairStrandsVF_Length"));
-		RadiusAtDepth1_Primary.Bind(ParameterMap, TEXT("HairStrandsVF_RadiusAtDepth1_Primary"));
-		RadiusAtDepth1_Velocity.Bind(ParameterMap, TEXT("HairStrandsVF_RadiusAtDepth1_Velocity"));
 		PositionOffset.Bind(ParameterMap, TEXT("HairStrandsVF_PositionOffset"));
 		PreviousPositionOffset.Bind(ParameterMap, TEXT("HairStrandsVF_PreviousPositionOffset"));
 		Density.Bind(ParameterMap, TEXT("HairStrandsVF_Density"));	
@@ -68,8 +66,8 @@ public:
 	{
 		Ar << Radius;
 		Ar << Length;
-		Ar << RadiusAtDepth1_Primary;
-		Ar << RadiusAtDepth1_Velocity;
+		Ar << RadiusAtDepth1_Primary;	// unused
+		Ar << RadiusAtDepth1_Velocity;	// unused
 		Ar << PositionOffset;
 		Ar << PreviousPositionOffset;
 		Ar << Density;
@@ -93,15 +91,8 @@ public:
 	) const override
 	{
 		const FHairStrandsVertexFactory* VF = static_cast<const FHairStrandsVertexFactory*>(VertexFactory);
-		
-		const FMinHairRadiusAtDepth1 MinRadiusAtDepth1 = ComputeMinStrandRadiusAtDepth1(
-			FIntPoint(View->UnconstrainedViewRect.Width(), View->UnconstrainedViewRect.Height()),
-			View->FOV,
-			GetHairVisibilitySampleCount(),
-			0.f);
 
 		const uint64 GroupIndex = reinterpret_cast<uint64>(BatchElement.UserData);
-
 		BindParam(ShaderBindings, PositionBuffer, VF->GetPositionSRV(GroupIndex));
 		BindParam(ShaderBindings, PreviousPositionBuffer, VF->GetPreviousPositionSRV(GroupIndex));
 		BindParam(ShaderBindings, AttributeBuffer, VF->GetAttributeSRV(GroupIndex));
@@ -111,8 +102,6 @@ public:
 		BindParam(ShaderBindings, PositionOffset, VF->GetPositionOffset(GroupIndex));
 		BindParam(ShaderBindings, PreviousPositionOffset, VF->GetPreviousPositionOffset(GroupIndex));
 		BindParam(ShaderBindings, Density, VF->GetHairDensity(GroupIndex));
-		BindParam(ShaderBindings, RadiusAtDepth1_Primary, MinRadiusAtDepth1.Primary);
-		BindParam(ShaderBindings, RadiusAtDepth1_Velocity, MinRadiusAtDepth1.Velocity);		
 	}
 };
 
diff --git a/Engine/Shaders/Private/HairStrands/HairStrandsDeepShadowVS.usf b/Engine/Shaders/Private/HairStrands/HairStrandsDeepShadowVS.usf
--- a/Engine/Shaders/Private/HairStrands/HairStrandsDeepShadowVS.usf
+++ b/Engine/Shaders/Private/HairStrands/HairStrandsDeepShadowVS.usf
@@ -4,6 +4,8 @@
 #include "/Engine/Generated/Material.ush"
 #include "/Engine/Generated/VertexFactory.ush"
 
+#include "HairStrandsVisibilityCommon.ush"
+
 #define MESH_RENDER_DEPTH 0
 #define MESH_RENDER_DOM   1
 #define MESH_RENDER_VOXEL 2
@@ -36,14 +38,6 @@ void Main(
 {
 	ResolvedView = ResolveView();
 
-	// #hair_todo
-	// Override the View direction with the light direction, since the Deep Shadow pass does not setup the View data, but passes data through a dedicated constant buffer
-	// Hack: Use 1337 as a flag for the strand vertex factory to not use MinRadiusAtDepth1 for computing the minimum radius for directional light. 
-#if SM5_PROFILE && COMPILER_HLSL
-	ResolvedView.TranslatedWorldCameraOrigin = float3(1337, DeepShadowPass.bIsOrtho, DeepShadowPass.MinStrandRadius);
-	ResolvedView.ViewForward = DeepShadowPass.LightDirection;
-#endif
-
 	FVertexFactoryIntermediates VFIntermediates = GetVertexFactoryIntermediates(Input);
 	float4 WorldPosition = VertexFactoryGetWorldPosition(Input, VFIntermediates);		
 	WorldPosition.xyz -= ResolvedView.PreViewTranslation.xyz; // Remove the camera relative position (todo move this into matrix transform)
@@ -60,8 +54,9 @@ void Main(
 
 #if MESH_RENDER_MODE == MESH_RENDER_DOM || MESH_RENDER_MODE == MESH_RENDER_VOXEL
   #ifdef VF_STRAND_HAIR	
+	FHairRenderingInfo HairRenderingInfo = GetHairRenderInfo(ResolvedView.HairRenderInfo);
 	const float CurrentDepth = Output.Position.z / Output.Position.w;
-	const float PixelRadius = CurrentDepth * DeepShadowPass.MinStrandRadius;
+	const float PixelRadius = HairRenderingInfo.IsOrthoView ? HairRenderingInfo.RadiusAtDepth1Primary : (CurrentDepth * HairRenderingInfo.RadiusAtDepth1Primary);
 	const float StrandRealRadius = VFIntermediates.HairDimensions.y;
 	const float Coverage = StrandRealRadius / max(StrandRealRadius, PixelRadius);
 	Output.HairCoverage = Coverage * VFIntermediates.HairDensity;
diff --git a/Engine/Shaders/Private/HairStrands/HairStrandsVertexFactory.ush b/Engine/Shaders/Private/HairStrands/HairStrandsVertexFactory.ush
--- a/Engine/Shaders/Private/HairStrands/HairStrandsVertexFactory.ush
+++ b/Engine/Shaders/Private/HairStrands/HairStrandsVertexFactory.ush
@@ -8,6 +8,8 @@
 
 #include "/Engine/Generated/UniformBuffers/PrecomputedLightingBuffer.ush"
 
+#include "HairStrandsVisibilityCommon.ush"
+
 #if MANUAL_VERTEX_FETCH
 	Buffer<float4> VertexFetch_InstanceOriginBuffer;
 	Buffer<float4> VertexFetch_InstanceTransformBuffer;
@@ -25,8 +27,6 @@ float3 HairStrandsVF_PreviousPositionOffset;	// Vertex position (3x16bits) are r
 float  HairStrandsVF_Density;
 float  HairStrandsVF_Radius;
 float  HairStrandsVF_Length;
-float  HairStrandsVF_RadiusAtDepth1_Primary;
-float  HairStrandsVF_RadiusAtDepth1_Velocity;
 
 Buffer<uint4>  HairStrandsVF_PositionBuffer;
 Buffer<uint4>  HairStrandsVF_PreviousPositionBuffer;
@@ -361,19 +361,13 @@ float4 ComputeViewAlignedWorldPosition(FVertexFactoryInput Input, float3 WorldTa
 {
 	const uint LocalVertexId = Input.VertexId % 6;
 	const bool bIsLeft = LocalVertexId == 0 || LocalVertexId == 1 || LocalVertexId == 5;
-	
-	// #hair_todo
-	// Use either the min strand radius computated within the strand factory, or a custom one, provided 
-	// (in a hacky way) throught the View variables. This is used by the DeepShadow pass, which inject 
-	// its data this way.
-	const bool  bUseCustomStrandRadius = ResolvedView.TranslatedWorldCameraOrigin.x == 1337.f;
-	const bool  bIsOrtho			   = bUseCustomStrandRadius && ResolvedView.TranslatedWorldCameraOrigin.y > 0;
-	const float CustomMinStrandRadius  = ResolvedView.TranslatedWorldCameraOrigin.z;
-	const float DistanceToCamera	   = length(ResolvedView.TranslatedWorldCameraOrigin - WorldPosition.xyz);
+
+	FHairRenderingInfo HairRenderingInfo = GetHairRenderInfo(ResolvedView.HairRenderInfo);
 
 	// Minimal radius to snap the strand to a sample/pixel center (to avoid aliasing)
-	const float RadiusAtDepth1 = bUseCustomStrandRadius ? CustomMinStrandRadius : lerp(HairStrandsVF_RadiusAtDepth1_Primary, HairStrandsVF_RadiusAtDepth1_Velocity, VelocityScale);
-	const float MinStrandHairRadius  = bIsOrtho ? RadiusAtDepth1 : (DistanceToCamera * RadiusAtDepth1);
+	const float RadiusAtDepth1 = lerp(HairRenderingInfo.RadiusAtDepth1Primary, HairRenderingInfo.RadiusAtDepth1Velocity, VelocityScale);
+	const float DistanceToCamera = length(ResolvedView.TranslatedWorldCameraOrigin - WorldPosition.xyz);
+	const float MinStrandHairRadius = HairRenderingInfo.IsOrthoView ? RadiusAtDepth1 : (DistanceToCamera * RadiusAtDepth1);
 	const float3 ViewDir = -ResolvedView.ViewForward;
 	const float3 Right = normalize(cross(WorldTangent, ViewDir));
 	const float3 OutWorldPosition = WorldPosition.xyz + (bIsLeft ? -Right : Right) * max(WorldStrandRadius, MinStrandHairRadius);
diff --git a/Engine/Shaders/Private/HairStrands/HairStrandsVisibilityCommon.ush b/Engine/Shaders/Private/HairStrands/HairStrandsVisibilityCommon.ush
--- a/Engine/Shaders/Private/HairStrands/HairStrandsVisibilityCommon.ush
+++ b/Engine/Shaders/Private/HairStrands/HairStrandsVisibilityCommon.ush
@@ -1,5 +1,7 @@
 // Copyright 1998-2019 Epic Games, Inc. All Rights Reserved.
 
+#pragma once
+
 struct FVisibilityData
 {
 	float3 Tangent;		// 24 bits
@@ -308,3 +310,19 @@ FNodeData ConvertToNodeData(float Depth, FVisibilityData VisibilityData, float3
 	 BaseColor.z = (float((Encoded >> 16) & 0xFF) / 255.f);
 	 Roughness   = (float((Encoded >> 24) & 0xFF) / 255.f);
  }
+
+ struct FHairRenderingInfo
+ {
+	 float RadiusAtDepth1Primary;
+	 float RadiusAtDepth1Velocity;
+	 bool IsOrthoView;
+ };
+
+ FHairRenderingInfo GetHairRenderInfo(float4 ViewHairRenderInfo)
+ {
+	 FHairRenderingInfo Info;
+	 Info.RadiusAtDepth1Primary = ViewHairRenderInfo.x;
+	 Info.RadiusAtDepth1Velocity = ViewHairRenderInfo.y;
+	 Info.IsOrthoView = ViewHairRenderInfo.z > 0.0f ? true : false;
+	 return Info;
+ }
diff --git a/Engine/Source/Runtime/Engine/Public/SceneView.h b/Engine/Source/Runtime/Engine/Public/SceneView.h
--- a/Engine/Source/Runtime/Engine/Public/SceneView.h
+++ b/Engine/Source/Runtime/Engine/Public/SceneView.h
@@ -725,6 +725,7 @@ enum ETranslucencyVolumeCascade
 	VIEW_UNIFORM_BUFFER_MEMBER(int32, FarShadowStaticMeshLODBias) \
 	VIEW_UNIFORM_BUFFER_MEMBER(float, MinRoughness) \
 	VIEW_UNIFORM_BUFFER_MEMBER(float, ConstantWaterDepth) \
+	VIEW_UNIFORM_BUFFER_MEMBER(FVector4, HairRenderInfo) \
 
 #define VIEW_UNIFORM_BUFFER_MEMBER(type, identifier) \
 	SHADER_PARAMETER(type, identifier)
diff --git a/Engine/Source/Runtime/Renderer/Private/HairStrands/HairStrandsRasterCommon.cpp b/Engine/Source/Runtime/Renderer/Private/HairStrands/HairStrandsRasterCommon.cpp
--- a/Engine/Source/Runtime/Renderer/Private/HairStrands/HairStrandsRasterCommon.cpp
+++ b/Engine/Source/Runtime/Renderer/Private/HairStrands/HairStrandsRasterCommon.cpp
@@ -19,11 +19,8 @@
 BEGIN_GLOBAL_SHADER_PARAMETER_STRUCT(FDeepShadowPassUniformParameters, )
 	//SHADER_PARAMETER_STRUCT(FSceneTexturesUniformParameters, SceneTextures)
 	SHADER_PARAMETER(FMatrix, WorldToClipMatrix)
-	SHADER_PARAMETER(FVector, LightDirection)
 	SHADER_PARAMETER(FVector4, SliceValue)
-	SHADER_PARAMETER(uint32, bIsOrtho)
 	SHADER_PARAMETER(FIntRect, AtlasRect)
-	SHADER_PARAMETER(float, MinStrandRadius)
 	SHADER_PARAMETER(FVector, VoxelMinAABB)
 	SHADER_PARAMETER(FVector, VoxelMaxAABB)
 	SHADER_PARAMETER(uint32, VoxelResolution)
@@ -403,11 +400,8 @@ void RasterHairStrands(
 
 		FDeepShadowPassUniformParameters PassParameters;
 		PassParameters.WorldToClipMatrix = LightDesc.WorldToLightClipTransform;
-		PassParameters.LightDirection = LightDesc.LightDirection;
 		PassParameters.SliceValue = FVector4(1, 1, 1, 1);
 		PassParameters.FrontDepthTexture = DeepShadowDepthTexture;
-		PassParameters.MinStrandRadius = LightDesc.MinStrandRadiusAtDepth1.Primary;
-		PassParameters.bIsOrtho = LightDesc.bIsOrtho ? 1u : 0u;
 		PassParameters.AtlasRect = AtlasRect;
 		PassParameters.VoxelMinAABB = MinAABB;
 		PassParameters.VoxelMaxAABB = MaxAABB;
@@ -416,6 +410,24 @@ void RasterHairStrands(
 	}
 
 	FMeshPassProcessorRenderState DrawRenderState(*ViewInfo, PassUniformBuffer);
+
+	if (ShadowPassType == EHairStrandsRasterPassType::DeepOpacityMap
+		|| ShadowPassType == EHairStrandsRasterPassType::FrontDepth
+		|| ShadowPassType == EHairStrandsRasterPassType::VoxelizationMaterial 
+		|| ShadowPassType == EHairStrandsRasterPassType::Voxelization)
+	{
+		TUniformBufferRef<FViewUniformShaderParameters> ViewPassUniformBuffer;
+		// Create a ViewUniformBuffer by patching the cached one
+
+		ViewInfo->CachedViewUniformShaderParameters->HairRenderInfo.X = LightDesc.MinStrandRadiusAtDepth1.Primary;
+		ViewInfo->CachedViewUniformShaderParameters->HairRenderInfo.Y = LightDesc.MinStrandRadiusAtDepth1.Velocity;
+		ViewInfo->CachedViewUniformShaderParameters->HairRenderInfo.Z = LightDesc.bIsOrtho ? 1.0f : 0.0f;
+		ViewInfo->CachedViewUniformShaderParameters->ViewForward = LightDesc.LightDirection;
+		ViewPassUniformBuffer = TUniformBufferRef<FViewUniformShaderParameters>::CreateUniformBufferImmediate(*ViewInfo->CachedViewUniformShaderParameters, UniformBuffer_SingleFrame);
+		DrawRenderState.SetViewUniformBuffer(ViewPassUniformBuffer);
+	}
+
+
 	{
 		RHICmdList.SetViewport(AtlasRect.Min.X, AtlasRect.Min.Y, 0.0f, AtlasRect.Max.X, AtlasRect.Max.Y, 1.0f);
 
diff --git a/Engine/Source/Runtime/Renderer/Private/SceneRendering.cpp b/Engine/Source/Runtime/Renderer/Private/SceneRendering.cpp
--- a/Engine/Source/Runtime/Renderer/Private/SceneRendering.cpp
+++ b/Engine/Source/Runtime/Renderer/Private/SceneRendering.cpp
@@ -58,6 +58,7 @@
 #include "IHeadMountedDisplay.h"
 #include "DiaphragmDOF.h" 
 #include "SingleLayerWaterRendering.h" 
+#include "HairStrands\HairStrandsUtils.h"
 
 /*-----------------------------------------------------------------------------
 	Globals
@@ -1665,6 +1666,15 @@ void FViewInfo::SetupUniformBufferParameters(
 			ViewUniformShaderParameters.LightmapSceneData = Scene->GPUScene.LightmapDataBuffer.SRV;
 		}
 	}
+
+	// Deep opacity maps info
+	{
+		const FMinHairRadiusAtDepth1 MinRadiusAtDepth1 = ComputeMinStrandRadiusAtDepth1(
+			FIntPoint(UnconstrainedViewRect.Width(), UnconstrainedViewRect.Height()), FOV, GetHairVisibilitySampleCount(), 0.0f);
+		ViewUniformShaderParameters.HairRenderInfo.X = MinRadiusAtDepth1.Primary;
+		ViewUniformShaderParameters.HairRenderInfo.Y = MinRadiusAtDepth1.Velocity;
+		ViewUniformShaderParameters.HairRenderInfo.Z = IsPerspectiveProjection() ? 0.0f : 1.0f;
+	}
 }
 
 void FViewInfo::InitRHIResources()
