Coverage view supports ortho projection.
Renaming of the HairRenderInfo for consistency.

diff --git a/Engine/Shaders/Private/HairStrands/HairStrandsDeepShadowVS.usf b/Engine/Shaders/Private/HairStrands/HairStrandsDeepShadowVS.usf
--- a/Engine/Shaders/Private/HairStrands/HairStrandsDeepShadowVS.usf
+++ b/Engine/Shaders/Private/HairStrands/HairStrandsDeepShadowVS.usf
@@ -54,9 +54,9 @@ void Main(
 
 #if MESH_RENDER_MODE == MESH_RENDER_DOM || MESH_RENDER_MODE == MESH_RENDER_VOXEL
   #ifdef VF_STRAND_HAIR	
-	FHairRenderingInfo HairRenderingInfo = GetHairRenderInfo(ResolvedView.HairRenderInfo);
+	FHairRenderInfo HairRenderInfo = GetHairRenderInfo(ResolvedView.HairRenderInfo);
 	const float CurrentDepth = Output.Position.z / Output.Position.w;
-	const float PixelRadius = HairRenderingInfo.bIsOrthoView ? HairRenderingInfo.RadiusAtDepth1Primary : (CurrentDepth * HairRenderingInfo.RadiusAtDepth1Primary);
+	const float PixelRadius = HairRenderInfo.bIsOrthoView ? HairRenderInfo.RadiusAtDepth1Primary : (CurrentDepth * HairRenderInfo.RadiusAtDepth1Primary);
 	const float StrandRealRadius = VFIntermediates.HairDimensions.y;
 	const float Coverage = StrandRealRadius / max(StrandRealRadius, PixelRadius);
 	Output.HairCoverage = Coverage * VFIntermediates.HairDensity;
diff --git a/Engine/Shaders/Private/HairStrands/HairStrandsVertexFactory.ush b/Engine/Shaders/Private/HairStrands/HairStrandsVertexFactory.ush
--- a/Engine/Shaders/Private/HairStrands/HairStrandsVertexFactory.ush
+++ b/Engine/Shaders/Private/HairStrands/HairStrandsVertexFactory.ush
@@ -362,12 +362,12 @@ float4 ComputeViewAlignedWorldPosition(FVertexFactoryInput Input, float3 WorldTa
 	const uint LocalVertexId = Input.VertexId % 6;
 	const bool bIsLeft = LocalVertexId == 0 || LocalVertexId == 1 || LocalVertexId == 5;
 
-	FHairRenderingInfo HairRenderingInfo = GetHairRenderInfo(ResolvedView.HairRenderInfo);
+	FHairRenderInfo HairRenderInfo = GetHairRenderInfo(ResolvedView.HairRenderInfo);
 
 	// Minimal radius to snap the strand to a sample/pixel center (to avoid aliasing)
-	const float RadiusAtDepth1 = lerp(HairRenderingInfo.RadiusAtDepth1Primary, HairRenderingInfo.RadiusAtDepth1Velocity, VelocityScale);
+	const float RadiusAtDepth1 = lerp(HairRenderInfo.RadiusAtDepth1Primary, HairRenderInfo.RadiusAtDepth1Velocity, VelocityScale);
 	const float DistanceToCamera = length(ResolvedView.TranslatedWorldCameraOrigin - WorldPosition.xyz);
-	const float MinStrandHairRadius = HairRenderingInfo.bIsOrthoView ? RadiusAtDepth1 : (DistanceToCamera * RadiusAtDepth1);
+	const float MinStrandHairRadius = HairRenderInfo.bIsOrthoView ? RadiusAtDepth1 : (DistanceToCamera * RadiusAtDepth1);
 	const float3 ViewDir = -ResolvedView.ViewForward;
 	const float3 Right = normalize(cross(WorldTangent, ViewDir));
 	const float3 OutWorldPosition = WorldPosition.xyz + (bIsLeft ? -Right : Right) * max(WorldStrandRadius, MinStrandHairRadius);
diff --git a/Engine/Shaders/Private/HairStrands/HairStrandsVisibilityCommon.ush b/Engine/Shaders/Private/HairStrands/HairStrandsVisibilityCommon.ush
--- a/Engine/Shaders/Private/HairStrands/HairStrandsVisibilityCommon.ush
+++ b/Engine/Shaders/Private/HairStrands/HairStrandsVisibilityCommon.ush
@@ -311,7 +311,7 @@ FNodeData ConvertToNodeData(float Depth, FVisibilityData VisibilityData, float3
 	 Roughness   = (float((Encoded >> 24) & 0xFF) / 255.f);
  }
 
- struct FHairRenderingInfo
+ struct FHairRenderInfo
  {
 	 float RadiusAtDepth1Primary;
 	 float RadiusAtDepth1Velocity;
@@ -319,9 +319,9 @@ FNodeData ConvertToNodeData(float Depth, FVisibilityData VisibilityData, float3
 	 float VelocityMagnitudeScale;
  };
 
- FHairRenderingInfo GetHairRenderInfo(float4 ViewHairRenderInfo)
+ FHairRenderInfo GetHairRenderInfo(float4 ViewHairRenderInfo)
  {
-	 FHairRenderingInfo Info;
+	 FHairRenderInfo Info;
 	 Info.RadiusAtDepth1Primary = ViewHairRenderInfo.x;
 	 Info.RadiusAtDepth1Velocity = ViewHairRenderInfo.y;
 	 Info.bIsOrthoView = ViewHairRenderInfo.z > 0.0f ? true : false;
diff --git a/Engine/Shaders/Private/HairStrands/HairStrandsVisibilityPS.usf b/Engine/Shaders/Private/HairStrands/HairStrandsVisibilityPS.usf
--- a/Engine/Shaders/Private/HairStrands/HairStrandsVisibilityPS.usf
+++ b/Engine/Shaders/Private/HairStrands/HairStrandsVisibilityPS.usf
@@ -68,9 +68,9 @@ void MainVisibility(
 
 	float Coverage = 1;
 	{
-		FHairRenderingInfo HairRenderingInfo = GetHairRenderInfo(ResolvedView.HairRenderInfo);
+		FHairRenderInfo HairRenderInfo = GetHairRenderInfo(ResolvedView.HairRenderInfo);
 		const float SceneDepth = ConvertFromDeviceZ(SvPosition.z); // Linear depth in world unit
-		const float PixelRadius = SceneDepth * HairRenderingInfo.RadiusAtDepth1Primary; // Not correct but the coverage is not used (we count instead the number of sub-sample covered)
+		const float PixelRadius = HairRenderInfo.bIsOrthoView ? HairRenderInfo.RadiusAtDepth1Primary : SceneDepth * HairRenderInfo.RadiusAtDepth1Primary; // Not correct but the coverage is not used (we count instead the number of sub-sample covered)
 		const float StrandRealRadius = WorldStrandRadius;
 		Coverage = StrandRealRadius / max(StrandRealRadius, PixelRadius);
 	}
diff --git a/Engine/Shaders/Private/HairStrands/HairStrandsVisibilityVS.usf b/Engine/Shaders/Private/HairStrands/HairStrandsVisibilityVS.usf
--- a/Engine/Shaders/Private/HairStrands/HairStrandsVisibilityVS.usf
+++ b/Engine/Shaders/Private/HairStrands/HairStrandsVisibilityVS.usf
@@ -55,8 +55,8 @@ void Main(
 		Output.PackedVelocityC			= INVARIANT(PrevScreenPosObj);
 
 		const float2 ScreenVelocity = Calculate2DVelocity(Output.PackedVelocityA, Output.PackedVelocityC);
-		FHairRenderingInfo HairRenderingInfo = GetHairRenderInfo(ResolvedView.HairRenderInfo);
-		VelocityScale = saturate(max(abs(ScreenVelocity.x), abs(ScreenVelocity.y)) / HairRenderingInfo.VelocityMagnitudeScale);
+		FHairRenderInfo HairRenderInfo = GetHairRenderInfo(ResolvedView.HairRenderInfo);
+		VelocityScale = saturate(max(abs(ScreenVelocity.x), abs(ScreenVelocity.y)) / HairRenderInfo.VelocityMagnitudeScale);
 	}
 	#endif
 
