Change BSDF components, global/local/multiple scattering into a bitfield.

diff --git a/Engine/Shaders/Private/DeferredLightPixelShaders.usf b/Engine/Shaders/Private/DeferredLightPixelShaders.usf
--- a/Engine/Shaders/Private/DeferredLightPixelShaders.usf
+++ b/Engine/Shaders/Private/DeferredLightPixelShaders.usf
@@ -110,7 +110,8 @@ StructuredBuffer<float4> HairTransmittanceBuffer;
 Texture3D<float4> HairLUTTexture;
 SamplerState HairLUTSampler;
 uint HairTransmittanceEnable;
-float4 HairComponents; // R, TT, TRT, LocaScattering
+uint HairComponents; // R, TT, TRT, LocalScattering, GlobalScattering, MultipleScattering
+
 
 float Luma(float3 Color)
 {
@@ -177,10 +178,9 @@ void ComputeHairLighting(
 
 		const FDerivedParams DerivedParams = GetDerivedParams(InputParams, SceneDepth);
 		FDeferredLightData LightData = SetupLightDataForStandardDeferred();
-		LightData.HairTransmittance.ScatteringComponent = HairComponents.xyz;
+		LightData.HairTransmittance.ScatteringComponent = HairComponents;
 
 		// Compute Hair transmittance
-		BRANCH if (HairTransmittanceEnable > 0)
 		{
 			const float3 LightPosition = LightData.Position;
 			float3 L = float3(0, 0, 1);
@@ -199,7 +199,6 @@ void ComputeHairLighting(
 				HairTransmittanceBuffer,
 				HairLUTTexture,
 				HairLUTSampler,
-				HairTransmittanceEnable,
 				HairComponents);
 		}
 
diff --git a/Engine/Shaders/Private/DeferredLightingCommon.ush b/Engine/Shaders/Private/DeferredLightingCommon.ush
--- a/Engine/Shaders/Private/DeferredLightingCommon.ush
+++ b/Engine/Shaders/Private/DeferredLightingCommon.ush
@@ -205,7 +205,8 @@ void GetShadowTerms(FGBufferData GBuffer, FDeferredLightData LightData, float3 W
 		{
 			// If hair transmittance is enabled, sharp shadowing should already be handled by 
 			// the dedicated deep shadow maps. Thus no need for contact shadow
-			if (!LightData.HairTransmittance.Enable)
+			const bool bUseComplexTransmittance = (LightData.HairTransmittance.ScatteringComponent & HAIR_COMPONENT_MULTISCATTER) > 0;
+			if (!bUseComplexTransmittance)
 			{
 				Shadow.TransmissionShadow *= ContactShadow;
 			}
diff --git a/Engine/Shaders/Private/HairBsdf.ush b/Engine/Shaders/Private/HairBsdf.ush
--- a/Engine/Shaders/Private/HairBsdf.ush
+++ b/Engine/Shaders/Private/HairBsdf.ush
@@ -145,7 +145,7 @@ float3 HairShading( FGBufferData GBuffer, float3 L, float3 V, half3 N, float Sha
 	float3 S = 0;
 
 	// R
-	if(HairTransmittance.ScatteringComponent.x > 0)
+	if(HairTransmittance.ScatteringComponent & HAIR_COMPONENT_R)
 	{
 		const float sa = sin( Alpha[0] );
 		const float ca = cos( Alpha[0] );
@@ -173,7 +173,7 @@ float3 HairShading( FGBufferData GBuffer, float3 L, float3 V, half3 N, float Sha
 	}
 
 	// TT
-	if(HairTransmittance.ScatteringComponent.y > 0)
+	if(HairTransmittance.ScatteringComponent & HAIR_COMPONENT_TT)
 	{
 		float Mp = Hair_g( B[1], SinThetaL + SinThetaV - Alpha[1] );
 
@@ -210,7 +210,7 @@ float3 HairShading( FGBufferData GBuffer, float3 L, float3 V, half3 N, float Sha
 	}
 
 	// TRT
-	if(HairTransmittance.ScatteringComponent.z > 0)
+	if(HairTransmittance.ScatteringComponent & HAIR_COMPONENT_TRT)
 	{
 		float Mp = Hair_g( B[2], SinThetaL + SinThetaV - Alpha[2] );
 		
@@ -236,31 +236,12 @@ float3 HairShading( FGBufferData GBuffer, float3 L, float3 V, half3 N, float Sha
 	}
 #endif
 
-	if (!HairTransmittance.Enable)
-	{
-		// Use soft Kajiya Kay diffuse attenuation
-		float KajiyaDiffuse = 1 - abs( dot(N,L) );
-	
-		float3 FakeNormal = normalize( V - N * dot(V,N) );
-		//N = normalize( DiffuseN + FakeNormal * 2 );
-		N = FakeNormal;
-	
-		// Hack approximation for multiple scattering.
-		float Wrap = 1;
-		float NoL = saturate( ( dot(N, L) + Wrap ) / Square( 1 + Wrap ) );
-		float DiffuseScatter = (1 / PI) * lerp( NoL, KajiyaDiffuse, 0.33 ) * GBuffer.Metallic;
-		float Luma = Luminance( GBuffer.BaseColor );
-		float3 ScatterTint = pow( GBuffer.BaseColor / Luma, 1 - Shadow );
-		S += GBuffer.BaseColor * DiffuseScatter * pow(ScatterTint,.5);
-	}
-
-	S = -min(-S, 0.0);
-
-	if (HairTransmittance.Enable && HairTransmittance.MultipleScattering > 0)
+	if ((HairTransmittance.ScatteringComponent & HAIR_COMPONENT_MULTISCATTER) && HairTransmittance.MultipleScattering > 0)
 	{
 		S = EvaluateHairMultipleScattering(HairTransmittance, ClampedRoughness, S);
 	}
 
+	S = -min(-S, 0.0);
 	return S;
 }
 
diff --git a/Engine/Shaders/Private/HairStrands/HairStrandsDeepTransmittanceDualScattering.ush b/Engine/Shaders/Private/HairStrands/HairStrandsDeepTransmittanceDualScattering.ush
--- a/Engine/Shaders/Private/HairStrands/HairStrandsDeepTransmittanceDualScattering.ush
+++ b/Engine/Shaders/Private/HairStrands/HairStrandsDeepTransmittanceDualScattering.ush
@@ -1,5 +1,13 @@
 // Copyright 1998-2019 Epic Games, Inc. All Rights Reserved.
 
+#ifndef HAIR_COMPONENT_LS
+	#error HAIR_COMPONENT_LS is not defined, please include ShadingCommon.ush
+#endif
+
+#ifndef HAIR_COMPONENT_GS
+	#error HAIR_COMPONENT_GS is not defined, please include ShadingCommon.ush
+#endif
+
 float Hair_g2(float B,float Theta)
 {
 	#if 0
@@ -17,7 +25,7 @@ void ComputeDualScatteringTerms(
 	const float3 V,
 	const float3 L,
 	const float3 T,
-	const float4 HairComponents,
+	const uint HairComponents,
 	inout float3 OutGlobalScattering,
 	inout float3 OutLocalScattering)
 {
@@ -105,10 +113,9 @@ void ComputeDualScatteringTerms(
 	const float3 LocalScattering  = 2 * Ab * Sb * db;
 	
 	// Decode if the Local/Global component should be enabled or not (for debug purpose)
-    const float LocalWeights  = floor(HairComponents.w / 10.0f) > 0?1:0;
-    const float GlobalWeights = frac(HairComponents.w / 10.0f)*10.f > 0?1:0;
+	const float LocalWeights  = (HairComponents & HAIR_COMPONENT_LS) > 0 ? 1 : 0;
 
-	OutGlobalScattering = GlobalWeights > 0 ? GlobalScattering : 1;
+	OutGlobalScattering = (HairComponents & HAIR_COMPONENT_GS) > 0 ? GlobalScattering : 1;
 	OutLocalScattering  = LocalScattering * LocalWeights;
 
 	// Final computation is done in ShadingModels.ush with the following formula
@@ -122,8 +129,7 @@ FHairTransmittanceData GetHairTransmittance(
 	float4 HairTransmittance,
 	Texture3D<float4> HairLUTTexture,
 	SamplerState HairLUTSampler,
-	const uint HairTransmittanceEnable,
-	const float4 HairComponents)
+	const uint HairComponents)
 {
 	FHairTransmittanceData TransmittanceData = InitHairTransmittanceData();
 	TransmittanceData.Transmittance = HairTransmittance.xyz;
@@ -131,9 +137,8 @@ FHairTransmittanceData GetHairTransmittance(
 	TransmittanceData.MultipleScattering = 1;
 	TransmittanceData.GlobalScattering = 1;
 	TransmittanceData.LocalScattering = 0;
-	TransmittanceData.ScatteringComponent = float3(HairComponents.x > 0 ? 1 : 0, HairComponents.y > 0 ? 1 : 0, HairComponents.z > 0 ? 1 : 0);
+	TransmittanceData.ScatteringComponent = HairComponents;
 	TransmittanceData.OpaqueVisibility = 1;
-	TransmittanceData.Enable = HairTransmittanceEnable > 0;
 
 	const float3 T = GBuffer.WorldNormal;
 	{
@@ -167,8 +172,7 @@ FHairTransmittanceData GetHairTransmittance(
 	StructuredBuffer<float4> HairTransmittanceBuffer,
 	Texture3D<float4> HairLUTTexture,
 	SamplerState HairLUTSampler,
-	const uint HairTransmittanceEnable,
-	const float4 HairComponents)
+	const uint HairComponents)
 {
 	float4 HairTransmittance = float4(1, 1, 1, 0);
 	if (SampleOffset < HairTransmittanceBufferMaxCount)
@@ -182,6 +186,5 @@ FHairTransmittanceData GetHairTransmittance(
 		HairTransmittance,
 		HairLUTTexture,
 		HairLUTSampler,
-		HairTransmittanceEnable,
 		HairComponents);
 }
diff --git a/Engine/Shaders/Private/HairStrands/HairStrandsEnvironmentLighting.usf b/Engine/Shaders/Private/HairStrands/HairStrandsEnvironmentLighting.usf
--- a/Engine/Shaders/Private/HairStrands/HairStrandsEnvironmentLighting.usf
+++ b/Engine/Shaders/Private/HairStrands/HairStrandsEnvironmentLighting.usf
@@ -155,8 +155,8 @@ float3 ReflectionEnvironment(FGBufferData GBuffer, float2 ScreenPosition, float2
 			const float ScalarTransmittance = max(SampleTransmittance.x, max(SampleTransmittance.y, SampleTransmittance.z));
 			AvgR += SampleL * SampleTransmittance.xyz;
 
-			const float4 ScatteringComponent = float4(1, 0, 1, 11);
-			const FHairTransmittanceData TransmittanceData = GetHairTransmittance(V, SampleL, GBuffer, SampleTransmittance, HairScatteringLUTTexture, HairLUTSampler, true, ScatteringComponent);
+			const uint ScatteringComponent = HAIR_COMPONENT_R | HAIR_COMPONENT_TRT | HAIR_COMPONENT_LS | HAIR_COMPONENT_GS;// | HAIR_MULTISCATTER;
+			const FHairTransmittanceData TransmittanceData = GetHairTransmittance(V, SampleL, GBuffer, SampleTransmittance, HairScatteringLUTTexture, HairLUTSampler, ScatteringComponent);
 
 			const float Area = 0.2;
 			const float BackLit = 1;
@@ -283,8 +283,7 @@ float3 ReflectionEnvironment(FGBufferData GBuffer, float2 ScreenPosition, float2
 		{			
 			TransmittanceData.MultipleScattering = 1;
 			TransmittanceData.GlobalScattering = Transmittance.xyz;
-			TransmittanceData.ScatteringComponent = uint3(1, 0, 1); // Only R & TRT
-			TransmittanceData.Enable = true;
+			TransmittanceData.ScatteringComponent = HAIR_COMPONENT_R | HAIR_COMPONENT_TRT | HAIR_COMPONENT_MULTISCATTER;
 
 			const float Area = 0.2;
 			const float BackLit = 0;
@@ -326,7 +325,8 @@ float3 ReflectionEnvironment(FGBufferData GBuffer, float2 ScreenPosition, float2
 		// BSDF with dual-scattering
 		#if 0
 		{
-			const FHairTransmittanceData TransmittanceData = GetHairTransmittance(V, L, GBuffer, Transmittance, HairScatteringLUTTexture, HairLUTSampler, 1, float4(0, 1, 0, 11));
+			const uint HairComponent = HAIR_COMPONENT_TT | HAIR_COMPONENT_LS | HAIR_COMPONENT_GS;
+			const FHairTransmittanceData TransmittanceData = GetHairTransmittance(V, L, GBuffer, Transmittance, HairScatteringLUTTexture, HairLUTSampler, HairComponent);
 			const float Area = 0.2;
 			ColorTT.rgb *= HairShading(GBuffer, L, V, T, 1, TransmittanceData, 1, Area, uint2(0, 0));
 		}
@@ -347,8 +347,7 @@ float3 ReflectionEnvironment(FGBufferData GBuffer, float2 ScreenPosition, float2
 			TransmittanceData.MultipleScattering = 1;
 			TransmittanceData.GlobalScattering = Transmittance.xyz;
 			TransmittanceData.LocalScattering = 0;
-			TransmittanceData.ScatteringComponent = uint3(0, 1, 0);
-			TransmittanceData.Enable = true;
+			TransmittanceData.ScatteringComponent = HAIR_COMPONENT_TT | HAIR_COMPONENT_MULTISCATTER;
 
 			const float Area = 0.2;
 			ColorTT.rgb *= HairShading(GBuffer, L, V, T, 1, TransmittanceData, 0, Area, uint2(0, 0));
diff --git a/Engine/Shaders/Private/HairStrands/HairStrandsLUT.usf b/Engine/Shaders/Private/HairStrands/HairStrandsLUT.usf
--- a/Engine/Shaders/Private/HairStrands/HairStrandsLUT.usf
+++ b/Engine/Shaders/Private/HairStrands/HairStrandsLUT.usf
@@ -214,9 +214,9 @@ void MainCS(uint3 DispatchThreadId : SV_DispatchThreadID)
 	const float3 N = float3(0, 0, 1); // N is the vector parallel to hair pointing toward root. I.e., the tangent T is up
 	const float3 V = float3(CosAngle, 0, SinAngle);
 	const float OpaqueVisibility = 1;
-	FHairTransmittanceData Setting_R   = InitHairTransmittanceData(); Setting_R.ScatteringComponent.x = 1;		Setting_R.ScatteringComponent.y = 0;	Setting_R.ScatteringComponent.z = 0;
-	FHairTransmittanceData Setting_TT  = InitHairTransmittanceData(); Setting_TT.ScatteringComponent.x = 0;		Setting_TT.ScatteringComponent.y = 1;	Setting_TT.ScatteringComponent.z = 0;
-	FHairTransmittanceData Setting_TRT = InitHairTransmittanceData(); Setting_TRT.ScatteringComponent.x = 0;	Setting_TRT.ScatteringComponent.y = 0;	Setting_TRT.ScatteringComponent.z = 1;
+	FHairTransmittanceData Setting_R   = InitHairTransmittanceData(); Setting_R.ScatteringComponent		= HAIR_COMPONENT_R;
+	FHairTransmittanceData Setting_TT  = InitHairTransmittanceData(); Setting_TT.ScatteringComponent	= HAIR_COMPONENT_TT;
+	FHairTransmittanceData Setting_TRT = InitHairTransmittanceData(); Setting_TRT.ScatteringComponent	= HAIR_COMPONENT_TRT;
 
 	const float MaxCosThetaRadius = cos(0.25f * PI / float(ThetaCount)); // [0, Pi/2] / ThetaCount which is divided by 2 for getting the actual radius
 	float3x3 ToViewBasis = GetTangentBasis(V);
diff --git a/Engine/Shaders/Private/ShadingCommon.ush b/Engine/Shaders/Private/ShadingCommon.ush
--- a/Engine/Shaders/Private/ShadingCommon.ush
+++ b/Engine/Shaders/Private/ShadingCommon.ush
@@ -24,6 +24,14 @@
 #define ZERO_PRECSHADOW_MASK			(1 << 6)
 #define SKIP_VELOCITY_MASK				(1 << 7)
 
+// Hair reflectance component (R, TT, TRT, Local Scattering, Global Scattering, Multi Scattering,...)
+#define HAIR_COMPONENT_R			0x1u
+#define HAIR_COMPONENT_TT			0x2u
+#define HAIR_COMPONENT_TRT			0x4u
+#define HAIR_COMPONENT_LS			0x8u 
+#define HAIR_COMPONENT_GS			0x10u
+#define HAIR_COMPONENT_MULTISCATTER	0x12u
+
 // for debugging and to visualize
 float3 GetShadingModelColor(uint ShadingModelID)
 {
diff --git a/Engine/Shaders/Private/ShadingModels.ush b/Engine/Shaders/Private/ShadingModels.ush
--- a/Engine/Shaders/Private/ShadingModels.ush
+++ b/Engine/Shaders/Private/ShadingModels.ush
@@ -67,8 +67,7 @@ struct FHairTransmittanceData
 	float3 LocalScattering;
 	float3 GlobalScattering;
 
-	uint3 ScatteringComponent; // R TT TRT
-	bool  Enable;
+	uint ScatteringComponent;
 };
 
 FHairTransmittanceData InitHairTransmittanceData()
@@ -80,12 +79,11 @@ FHairTransmittanceData InitHairTransmittanceData()
 	o.OpaqueVisibility = 1;
 	o.HairCount = 0;
 	o.MultipleScattering = 0;
-	o.Enable = false;
 	
 	// TEMP: for fastning iteration times
 	o.LocalScattering = 0;
 	o.GlobalScattering = 1;
-	o.ScatteringComponent = 1;
+	o.ScatteringComponent = HAIR_COMPONENT_R | HAIR_COMPONENT_TT | HAIR_COMPONENT_TRT;
 
 	return o;
 }
@@ -794,7 +792,7 @@ float3 HairShading( FGBufferData GBuffer, float3 L, float3 V, half3 N, float Sha
 	float3 S = 0;
 
 	// R
-	if(HairTransmittance.ScatteringComponent.x > 0)
+	if(HairTransmittance.ScatteringComponent & HAIR_COMPONENT_R)
 	{
 		const float sa = sin( Alpha[0] );
 		const float ca = cos( Alpha[0] );
@@ -807,7 +805,7 @@ float3 HairShading( FGBufferData GBuffer, float3 L, float3 V, half3 N, float Sha
 	}
 
 	// TT
-	if(HairTransmittance.ScatteringComponent.y > 0)
+	if(HairTransmittance.ScatteringComponent & HAIR_COMPONENT_TT)
 	{
 		float Mp = Hair_g( B[1], SinThetaL + SinThetaV - Alpha[1] );
 
@@ -837,7 +835,7 @@ float3 HairShading( FGBufferData GBuffer, float3 L, float3 V, half3 N, float Sha
 	}
 
 	// TRT
-	if(HairTransmittance.ScatteringComponent.z > 0)
+	if(HairTransmittance.ScatteringComponent & HAIR_COMPONENT_TRT)
 	{
 		float Mp = Hair_g( B[2], SinThetaL + SinThetaV - Alpha[2] );
 		
@@ -855,17 +853,19 @@ float3 HairShading( FGBufferData GBuffer, float3 L, float3 V, half3 N, float Sha
 	}
 #endif
 
-	if (!HairTransmittance.Enable && bEvalMultiScatter)
+	const bool bComplexMultiScatter = (HairTransmittance.ScatteringComponent & HAIR_COMPONENT_MULTISCATTER);
+	if (!bComplexMultiScatter && bEvalMultiScatter)
 	{
 		S += KajiyaKayDiffuseAttenuation(GBuffer, L, V, N, Shadow);
 	}
 
-	S = -min(-S, 0.0);
-
-	if (HairTransmittance.Enable && HairTransmittance.MultipleScattering > 0)
+	if (bComplexMultiScatter && HairTransmittance.MultipleScattering > 0)
 	{
 		S = EvaluateHairMultipleScattering(HairTransmittance, ClampedRoughness, S);
 	}
+
+	S = -min(-S, 0.0);
+
 	return S;
 }
 
diff --git a/Engine/Source/Runtime/Renderer/Private/HairStrands/HairStrandsDebug.cpp b/Engine/Source/Runtime/Renderer/Private/HairStrands/HairStrandsDebug.cpp
--- a/Engine/Source/Runtime/Renderer/Private/HairStrands/HairStrandsDebug.cpp
+++ b/Engine/Source/Runtime/Renderer/Private/HairStrands/HairStrandsDebug.cpp
@@ -1247,14 +1247,8 @@ void RenderHairStrandsDebugInfo(FRHICommandListImmediate& RHICmdList, TArray<FVi
 		FLinearColor DebugColor(1, 1, 0);
 		FString Line;
 
-		const FVector4 HairComponent = GetHairComponents();
-		const uint32 bHairR = HairComponent.X > 0 ? 1 : 0;
-		const uint32 bHairTT = HairComponent.Y > 0 ? 1 : 0;
-		const uint32 bHairTRT = HairComponent.Z > 0 ? 1 : 0;
-		const uint32 bHairGlobalScattering = FMath::FloorToInt(HairComponent.W / 10.0f) > 0 ? 1 : 0;
-		const uint32 bHairLocalScattering  = FMath::Frac(HairComponent.W / 10.0f)*10.f > 0 ? 1 : 0;
-
-		Line = FString::Printf(TEXT("Hair Components : (R=%d, TT=%d, TRT=%d, GS=%d, LS=%d)"), bHairR, bHairTT, bHairTRT, bHairGlobalScattering, bHairLocalScattering);
+		const FHairComponent HairComponent = GetHairComponents();
+		Line = FString::Printf(TEXT("Hair Components : (R=%d, TT=%d, TRT=%d, GS=%d, LS=%d)"), HairComponent.R, HairComponent.TT, HairComponent.TRT, HairComponent.GlobalScattering, HairComponent.LocalScattering);
 		Canvas.DrawShadowedString(X, Y += YStep, *Line, GetStatsFont(), DebugColor);
 		Line = FString::Printf(TEXT("----------------------------------------------------------------"));				Canvas.DrawShadowedString(X, Y += YStep, *Line, GetStatsFont(), DebugColor);
 		Line = FString::Printf(TEXT("Debug strands mode : %s"), ToString(StrandsDebugMode));							Canvas.DrawShadowedString(X, Y += YStep, *Line, GetStatsFont(), DebugColor);
diff --git a/Engine/Source/Runtime/Renderer/Private/HairStrands/HairStrandsUtils.cpp b/Engine/Source/Runtime/Renderer/Private/HairStrands/HairStrandsUtils.cpp
--- a/Engine/Source/Runtime/Renderer/Private/HairStrands/HairStrandsUtils.cpp
+++ b/Engine/Source/Runtime/Renderer/Private/HairStrands/HairStrandsUtils.cpp
@@ -41,15 +41,26 @@ float SampleCountToSubPixelSize(uint32 SamplePerPixelCount)
 	}
 	return Scale;
 }
-FVector4 GetHairComponents()
+FHairComponent GetHairComponents()
 {
-	return FVector4(
-		GHairR > 0 ? 1 : 0,
-		GHairTT > 0 ? 1 : 0,
-		GHairTRT > 0 ? 1 : 0,
-		(GHairGlobalScattering > 0 ? 1 : 0) + (GHairLocalScattering > 0 ? 10 : 0));
+	FHairComponent Out;
+	Out.R = GHairR > 0;
+	Out.TT = GHairTT > 0;
+	Out.TRT = GHairTRT > 0;
+	Out.LocalScattering = GHairLocalScattering > 0;
+	Out.GlobalScattering = GHairGlobalScattering > 0;
+	return Out;
 }
 
+uint32 ToBitfield(const FHairComponent& C)
+{
+	return
+		(C.R				? 1u : 0u)      |
+		(C.TT				? 1u : 0u) << 1 |
+		(C.TRT				? 1u : 0u) << 2 |
+		(C.LocalScattering  ? 1u : 0u) << 3 |
+		(C.GlobalScattering ? 1u : 0u) << 4 ;
+}
 FMinHairRadiusAtDepth1 ComputeMinStrandRadiusAtDepth1(
 	const FIntPoint& Resolution,
 	const float FOV,
diff --git a/Engine/Source/Runtime/Renderer/Private/HairStrands/HairStrandsUtils.h b/Engine/Source/Runtime/Renderer/Private/HairStrands/HairStrandsUtils.h
--- a/Engine/Source/Runtime/Renderer/Private/HairStrands/HairStrandsUtils.h
+++ b/Engine/Source/Runtime/Renderer/Private/HairStrands/HairStrandsUtils.h
@@ -21,7 +21,16 @@ void ComputeWorldToLightClip(
 	const ELightComponentType LightType,
 	const FIntPoint& ShadowResolution);
 
-FVector4 GetHairComponents();
+struct FHairComponent
+{
+	bool R = true;
+	bool TT = true;
+	bool TRT = true;
+	bool GlobalScattering = true;
+	bool LocalScattering = true;
+};
+FHairComponent GetHairComponents();
+uint32 ToBitfield(const FHairComponent& Component);
 
 float SampleCountToSubPixelSize(uint32 SamplePerPixelCount);
 
diff --git a/Engine/Source/Runtime/Renderer/Private/LightRendering.cpp b/Engine/Source/Runtime/Renderer/Private/LightRendering.cpp
--- a/Engine/Source/Runtime/Renderer/Private/LightRendering.cpp
+++ b/Engine/Source/Runtime/Renderer/Private/LightRendering.cpp
@@ -206,7 +206,6 @@ struct FRenderLightParams
 	// Precompute transmittance
 	FShaderResourceViewRHIRef DeepShadow_TransmittanceMaskBuffer = nullptr;
 	uint32 DeepShadow_TransmittanceMaskBufferMaxCount = 0;
-	bool bUseTransmittanceData = false;
 	
 	// Visibility buffer data
 	IPooledRenderTarget* HairCategorizationTexture = nullptr;
@@ -315,7 +314,6 @@ class FDeferredLightPS : public FGlobalShader
 		HairLUTTexture.Bind(Initializer.ParameterMap, TEXT("HairLUTTexture"));
 		HairLUTSampler.Bind(Initializer.ParameterMap, TEXT("HairLUTSampler"));
 		HairComponents.Bind(Initializer.ParameterMap, TEXT("HairComponents"));
-		HairTransmittanceEnable.Bind(Initializer.ParameterMap, TEXT("HairTransmittanceEnable"));
 
 		HairCategorizationTexture.Bind(Initializer.ParameterMap, TEXT("HairCategorizationTexture"));
 		HairVisibilityNodeOffsetAndCount.Bind(Initializer.ParameterMap, TEXT("HairVisibilityNodeOffsetAndCount"));
@@ -373,7 +371,6 @@ public:
 		Ar << HairLUTTexture;
 		Ar << HairLUTSampler;
 		Ar << HairComponents;
-		Ar << HairTransmittanceEnable;
 
 		return bShaderHasOutdatedParameters;
 	}
@@ -554,23 +551,13 @@ private:
 
 		if (HairComponents.IsBound())
 		{
-			FVector4 InHairComponents = GetHairComponents();
+			uint32 InHairComponents = ToBitfield(GetHairComponents());
 			SetShaderValue(
 				RHICmdList,
 				ShaderRHI,
 				HairComponents,
 				InHairComponents);
 		}
-
-		if (HairTransmittanceEnable.IsBound())
-		{
-			uint32 bIsDeepShadowEnable = (RenderLightParams && RenderLightParams->bUseTransmittanceData) ? 1 : 0;
-			SetShaderValue(
-				RHICmdList,
-				ShaderRHI,
-				HairTransmittanceEnable,
-				bIsDeepShadowEnable);
-		}
 	}
 
 	FSceneTextureShaderParameters SceneTextureParameters;
@@ -597,7 +584,6 @@ private:
 	FShaderResourceParameter HairLUTTexture;
 	FShaderResourceParameter HairLUTSampler;
 	FShaderParameter HairComponents;
-	FShaderParameter HairTransmittanceEnable;
 	};
 
 IMPLEMENT_GLOBAL_SHADER(FDeferredLightPS, "/Engine/Private/DeferredLightPixelShaders.usf", "DeferredLightPixelMain", SF_Pixel);
@@ -1967,7 +1953,6 @@ void FDeferredShadingSceneRenderer::RenderLight(FRHICommandList& RHICmdList, con
 		FRenderLightParams RenderLightParams;
 		if (bHairRenderingEnabled)
 		{
-			RenderLightParams.bUseTransmittanceData = bHairRenderingEnabled;
 			RenderLightParams.DeepShadow_TransmittanceMaskBuffer = InTransmittanceMaskData ? InTransmittanceMaskData->TransmittanceMaskSRV : nullptr;
 			RenderLightParams.DeepShadow_TransmittanceMaskBufferMaxCount = InTransmittanceMaskData && InTransmittanceMaskData->TransmittanceMask ? InTransmittanceMaskData->TransmittanceMask->Desc.NumElements : 0;
 			RenderLightParams.ScreenShadowMaskSubPixelTexture = ScreenShadowMaskSubPixelTexture;
