Remove multiplscattering enable on the shading transmittance struct and use haircomponent/flags instead.

diff --git a/Engine/Shaders/Private/HairBsdf.ush b/Engine/Shaders/Private/HairBsdf.ush
--- a/Engine/Shaders/Private/HairBsdf.ush
+++ b/Engine/Shaders/Private/HairBsdf.ush
@@ -89,6 +89,15 @@ float3 SigmaAFromReflectance(const float3 c, float beta_n)
 		SigmaAFromReflectance(c.z, beta_n));
 }
 
+
+float3 EvaluateHairMultipleScattering(
+	const FHairTransmittanceData TransmittanceData,
+	const float Roughness,
+	const float3 Fs)
+{
+	return TransmittanceData.GlobalScattering * (Fs + TransmittanceData.LocalScattering) * TransmittanceData.OpaqueVisibility;
+}
+
 float3 HairShading( FGBufferData GBuffer, float3 L, float3 V, half3 N, float Shadow, FHairTransmittanceData HairTransmittance, float Backlit, float Area, uint2 Random )
 {
 	// Strand tangent is oriented from root to tip, while the existing shading code expect tip to root orientation
@@ -236,7 +245,7 @@ float3 HairShading( FGBufferData GBuffer, float3 L, float3 V, half3 N, float Sha
 	}
 #endif
 
-	if ((HairTransmittance.ScatteringComponent & HAIR_COMPONENT_MULTISCATTER) && HairTransmittance.MultipleScattering > 0)
+	if (HairTransmittance.ScatteringComponent & HAIR_COMPONENT_MULTISCATTER)
 	{
 		S = EvaluateHairMultipleScattering(HairTransmittance, ClampedRoughness, S);
 	}
diff --git a/Engine/Shaders/Private/HairStrands/HairStrandsDeepTransmittanceDualScattering.ush b/Engine/Shaders/Private/HairStrands/HairStrandsDeepTransmittanceDualScattering.ush
--- a/Engine/Shaders/Private/HairStrands/HairStrandsDeepTransmittanceDualScattering.ush
+++ b/Engine/Shaders/Private/HairStrands/HairStrandsDeepTransmittanceDualScattering.ush
@@ -134,10 +134,9 @@ FHairTransmittanceData GetHairTransmittance(
 	FHairTransmittanceData TransmittanceData = InitHairTransmittanceData();
 	TransmittanceData.Transmittance = HairTransmittance.xyz;
 	TransmittanceData.HairCount = HairTransmittance.w;
-	TransmittanceData.MultipleScattering = 1;
 	TransmittanceData.GlobalScattering = 1;
 	TransmittanceData.LocalScattering = 0;
-	TransmittanceData.ScatteringComponent = HairComponents;
+	TransmittanceData.ScatteringComponent = HairComponents | HAIR_COMPONENT_MULTISCATTER;
 	TransmittanceData.OpaqueVisibility = 1;
 
 	const float3 T = GBuffer.WorldNormal;
diff --git a/Engine/Shaders/Private/HairStrands/HairStrandsEnvironmentLighting.usf b/Engine/Shaders/Private/HairStrands/HairStrandsEnvironmentLighting.usf
--- a/Engine/Shaders/Private/HairStrands/HairStrandsEnvironmentLighting.usf
+++ b/Engine/Shaders/Private/HairStrands/HairStrandsEnvironmentLighting.usf
@@ -281,7 +281,6 @@ float3 ReflectionEnvironment(FGBufferData GBuffer, float2 ScreenPosition, float2
 		// BSDF
 		#if 1
 		{			
-			TransmittanceData.MultipleScattering = 1;
 			TransmittanceData.GlobalScattering = Transmittance.xyz;
 			TransmittanceData.ScatteringComponent = HAIR_COMPONENT_R | HAIR_COMPONENT_TRT | HAIR_COMPONENT_MULTISCATTER;
 
diff --git a/Engine/Shaders/Private/ShadingModels.ush b/Engine/Shaders/Private/ShadingModels.ush
--- a/Engine/Shaders/Private/ShadingModels.ush
+++ b/Engine/Shaders/Private/ShadingModels.ush
@@ -60,8 +60,6 @@ struct FHairTransmittanceData
 
 	float OpaqueVisibility;
 	float HairCount;
-	float MultipleScattering;
-
 
 	// TEMP: for fastning iteration times
 	float3 LocalScattering;
@@ -78,7 +76,6 @@ FHairTransmittanceData InitHairTransmittanceData()
 	o.A_back = 1;			
 	o.OpaqueVisibility = 1;
 	o.HairCount = 0;
-	o.MultipleScattering = 0;
 	
 	// TEMP: for fastning iteration times
 	o.LocalScattering = 0;
@@ -729,14 +726,6 @@ float3 KajiyaKayDiffuseAttenuation(FGBufferData GBuffer, float3 L, float3 V, hal
 	return sqrt(GBuffer.BaseColor) * DiffuseScatter * ScatterTint;
 }
 
-float3 EvaluateHairMultipleScattering(
-	const FHairTransmittanceData TransmittanceData,
-	const float Roughness,
-	const float3 Fs)
-{
-	return TransmittanceData.GlobalScattering * (Fs + TransmittanceData.LocalScattering) * TransmittanceData.OpaqueVisibility;
-}
-
 // Approximation to HairShadingRef using concepts from the following papers:
 // [Marschner et al. 2003, "Light Scattering from Human Hair Fibers"]
 // [Pekelis et al. 2015, "A Data-Driven Light Scattering Model for Hair"]
@@ -853,17 +842,11 @@ float3 HairShading( FGBufferData GBuffer, float3 L, float3 V, half3 N, float Sha
 	}
 #endif
 
-	const bool bComplexMultiScatter = (HairTransmittance.ScatteringComponent & HAIR_COMPONENT_MULTISCATTER);
-	if (!bComplexMultiScatter && bEvalMultiScatter)
+	if (bEvalMultiScatter)
 	{
 		S += KajiyaKayDiffuseAttenuation(GBuffer, L, V, N, Shadow);
 	}
 
-	if (bComplexMultiScatter && HairTransmittance.MultipleScattering > 0)
-	{
-		S = EvaluateHairMultipleScattering(HairTransmittance, ClampedRoughness, S);
-	}
-
 	S = -min(-S, 0.0);
 
 	return S;
