Add UDIM support for root UV. The texture index is limited to 256x256 to limit memory bandwith

diff --git a/Engine/Plugins/Experimental/HairStrands/Source/HairStrandsCore/Private/GroomBuilder.cpp b/Engine/Plugins/Experimental/HairStrands/Source/HairStrandsCore/Private/GroomBuilder.cpp
--- a/Engine/Plugins/Experimental/HairStrands/Source/HairStrandsCore/Private/GroomBuilder.cpp
+++ b/Engine/Plugins/Experimental/HairStrands/Source/HairStrandsCore/Private/GroomBuilder.cpp
@@ -240,10 +240,19 @@ namespace HairStrandsBuilder
 
 				const FVector2D RootUV = Curves.CurvesRootUV[CurveIndex];
 				FHairStrandsAttributeFormat::Type& PackedAttributes = OutPackedAttributes[PointIndex + IndexOffset];
-				PackedAttributes.RootU = uint32(FMath::Clamp(RootUV.X, 0.f, 1.f) * 0xFF);
-				PackedAttributes.RootV = uint32(FMath::Clamp(RootUV.Y, 0.f, 1.f) * 0xFF);
+				PackedAttributes.Unused0 = 0;
+				PackedAttributes.Unused1 = 0;
 				PackedAttributes.UCoord = FMath::Clamp(CoordU * 255.f, 0.f, 255.f);
 				PackedAttributes.Seed = CurveSeed;
+
+				// Root UV support UDIM texture coordinate but limit the spans of the UDIM to be in 256x256 instead of 9999x9999.
+				// The internal UV coords are also limited to 8bits, which means if sampling need to be super precise, this is no enough.
+				const FVector2D TextureRootUV(FMath::Fractional(RootUV.X), FMath::Fractional(RootUV.Y));
+				const FVector2D TextureIndexUV = RootUV - TextureRootUV;
+				PackedAttributes.RootU  = uint32(FMath::Clamp(TextureRootUV.X*255.f, 0.f, 255.f));
+				PackedAttributes.RootV  = uint32(FMath::Clamp(TextureRootUV.Y*255.f, 0.f, 255.f));
+				PackedAttributes.IndexU = uint32(FMath::Clamp(TextureIndexUV.X, 0.f, 255.f));
+				PackedAttributes.IndexV = uint32(FMath::Clamp(TextureIndexUV.Y, 0.f, 255.f));
 			}
 		}
 	}
diff --git a/Engine/Plugins/Experimental/HairStrands/Source/HairStrandsCore/Private/GroomComponent.cpp b/Engine/Plugins/Experimental/HairStrands/Source/HairStrandsCore/Private/GroomComponent.cpp
--- a/Engine/Plugins/Experimental/HairStrands/Source/HairStrandsCore/Private/GroomComponent.cpp
+++ b/Engine/Plugins/Experimental/HairStrands/Source/HairStrandsCore/Private/GroomComponent.cpp
@@ -253,6 +253,7 @@ public:
 			case EHairStrandsDebugMode::RenderHairSeed				: DebugModeScalar = 3.f; break;
 			case EHairStrandsDebugMode::RenderHairDimension			: DebugModeScalar = 4.f; break;
 			case EHairStrandsDebugMode::RenderHairRadiusVariation	: DebugModeScalar = 5.f; break;
+			case EHairStrandsDebugMode::RenderHairRootUDIM			: DebugModeScalar = 6.f; break;
 			};
 
 			float HairMinRadius = FLT_MAX;
diff --git a/Engine/Plugins/Experimental/HairStrands/Source/HairStrandsCore/Private/HairStrandsDatas.cpp b/Engine/Plugins/Experimental/HairStrands/Source/HairStrandsCore/Private/HairStrandsDatas.cpp
--- a/Engine/Plugins/Experimental/HairStrands/Source/HairStrandsCore/Private/HairStrandsDatas.cpp
+++ b/Engine/Plugins/Experimental/HairStrands/Source/HairStrandsCore/Private/HairStrandsDatas.cpp
@@ -1,6 +1,7 @@
 // Copyright 1998-2019 Epic Games, Inc. All Rights Reserved.
 
 #include "HairStrandsDatas.h"
+#include "UObject/ReleaseObjectVersion.h"
 
 void FHairStrandsInterpolationDatas::SetNum(const uint32 NumCurves)
 {
@@ -71,11 +72,29 @@ FArchive& operator<<(FArchive& Ar, FPackedHairVertex& Vertex)
 
 FArchive& operator<<(FArchive& Ar, FPackedHairAttributeVertex& Vertex)
 {
+	Ar.UsingCustomVersion(FReleaseObjectVersion::GUID);
+
 	Ar << Vertex.RootU;
 	Ar << Vertex.RootV;
 	Ar << Vertex.UCoord;
 	Ar << Vertex.Seed;
 
+	if (Ar.CustomVer(FReleaseObjectVersion::GUID) >= FReleaseObjectVersion::GroomAssetVersion1)
+	{
+		Ar << Vertex.IndexU;
+		Ar << Vertex.IndexV;
+		Ar << Vertex.Unused0;
+		Ar << Vertex.Unused1;
+	}
+	else
+	{
+		Vertex.IndexU = 0;
+		Vertex.IndexV = 0;
+		Vertex.Unused0 = 0;
+		Vertex.Unused1 = 0;
+	}
+	
+
 	return Ar;
 }
 
diff --git a/Engine/Plugins/Experimental/HairStrands/Source/HairStrandsCore/Public/HairStrandsDatas.h b/Engine/Plugins/Experimental/HairStrands/Source/HairStrandsCore/Public/HairStrandsDatas.h
--- a/Engine/Plugins/Experimental/HairStrands/Source/HairStrandsCore/Public/HairStrandsDatas.h
+++ b/Engine/Plugins/Experimental/HairStrands/Source/HairStrandsCore/Public/HairStrandsDatas.h
@@ -27,6 +27,11 @@ struct FPackedHairAttributeVertex
 	uint8 UCoord;
 	uint8 Seed;
 
+	uint8 IndexU;
+	uint8 IndexV;
+	uint8 Unused0;
+	uint8 Unused1;
+
 	friend FArchive& operator<<(FArchive& Ar, FPackedHairAttributeVertex& Vertex);
 };
 
@@ -74,8 +79,8 @@ struct FHairStrandsAttributeFormat
 	typedef FPackedHairAttributeVertex Type;
 	static const uint32 ComponentCount = 1;
 	static const uint32 SizeInByte = sizeof(Type);
-	static const EVertexElementType VertexElementType = VET_UByte4;
-	static const EPixelFormat Format = PF_R8G8B8A8;
+	static const EVertexElementType VertexElementType = VET_UShort4;
+	static const EPixelFormat Format = PF_R16G16B16A16_UINT;
 };
 
 struct FHairStrandsTangentFormat
@@ -291,8 +296,8 @@ struct HAIRSTRANDSCORE_API FHairStrandsCurves
 	/** Normalized length relative to the max one */
 	TArray<float> CurvesLength; // [0..1]
 
-	/** Roots UV in the bounding sphere */
-	TArray<FVector2D> CurvesRootUV; // [0..1]
+	/** Roots UV. Support UDIM coordinate up to 256x256 */
+	TArray<FVector2D> CurvesRootUV; // [0..256]
 
 	/** Max strands Curves length */
 	float MaxLength = 0;
diff --git a/Engine/Shaders/Private/HairStrands/HairStrandsVertexFactory.ush b/Engine/Shaders/Private/HairStrands/HairStrandsVertexFactory.ush
--- a/Engine/Shaders/Private/HairStrands/HairStrandsVertexFactory.ush
+++ b/Engine/Shaders/Private/HairStrands/HairStrandsVertexFactory.ush
@@ -5,6 +5,7 @@
 =============================================================================*/
 #include "../VertexFactoryCommon.ush"
 #include "../LocalVertexFactoryCommon.ush"
+#include "HairStrandsVisibilityCommon.ush"
 
 #include "/Engine/Generated/UniformBuffers/PrecomputedLightingBuffer.ush"
 
@@ -30,7 +31,7 @@ float  HairStrandsVF_Length;
 
 Buffer<uint4>  HairStrandsVF_PositionBuffer;
 Buffer<uint4>  HairStrandsVF_PreviousPositionBuffer;
-Buffer<float4> HairStrandsVF_AttributeBuffer;
+Buffer<uint4>  HairStrandsVF_AttributeBuffer;
 Buffer<float4> HairStrandsVF_TangentBuffer;
 
 // ## A quad is made of the following topology and indexing:
@@ -154,11 +155,11 @@ FControlPoint GetVertexPosition(FVertexFactoryInput Input)
 	// Create a degenerated quad the end of each strand to cut between each strands
 	// #hair_todo: This is not efficient for short strand like fur for instance. Need to revisit that at some point
 	const bool bIsInvalidQuad = (ControlPointType == 1 && OffsetIndex == 1) || (ControlPointType == 2 && OffsetIndex == 0);
-	const float4 Attributes		= HairStrandsVF_AttributeBuffer[BaseIndex + OffsetIndex];
+	const FHairAttributes Attributes = DecodeHairAttribute(HairStrandsVF_AttributeBuffer[BaseIndex + OffsetIndex]);
 	Output.Position				= bIsInvalidQuad ? float3(1.#INF, 1.#INF, 1.#INF) : ControlPoint;
-	Output.RootUV				= Attributes.xy;
-	Output.UCoord				= Attributes.z;
-	Output.Seed					= Attributes.w;
+	Output.RootUV				= Attributes.RootUV;
+	Output.UCoord				= Attributes.UCoord;
+	Output.Seed					= Attributes.Seed;
 
 	return Output;
 }
diff --git a/Engine/Shaders/Private/HairStrands/HairStrandsVisibilityCommon.ush b/Engine/Shaders/Private/HairStrands/HairStrandsVisibilityCommon.ush
--- a/Engine/Shaders/Private/HairStrands/HairStrandsVisibilityCommon.ush
+++ b/Engine/Shaders/Private/HairStrands/HairStrandsVisibilityCommon.ush
@@ -328,3 +328,25 @@ FNodeData ConvertToNodeData(float Depth, FVisibilityData VisibilityData, float3
 	 Info.VelocityMagnitudeScale = ViewHairRenderInfo.w;
 	 return Info;
  }
+
+ struct FHairAttributes
+ {
+	 float2 RootUV;
+	 float	Seed;
+	 float	UCoord;
+ };
+
+ FHairAttributes DecodeHairAttribute(uint4 Attributes)
+ {
+	 float2 RootUV, RootIndex;
+	 RootUV.x		= (Attributes.x & 0xFF) / 255.f;
+	 RootUV.y		= ((Attributes.x>>8) & 0xFF) / 255.f;
+	 RootIndex.x	= Attributes.z & 0xFF;
+	 RootIndex.y	= (Attributes.z >> 8) & 0xFF;
+
+	 FHairAttributes Out;
+	 Out.RootUV		= RootUV + RootIndex;
+	 Out.UCoord		= (Attributes.y & 0xFF) / 255.f;
+	 Out.Seed		= ((Attributes.y>>8) & 0xFF) / 255.f;
+	 return Out;
+ }
\ No newline at end of file
diff --git a/Engine/Shaders/Private/MaterialTemplate.ush b/Engine/Shaders/Private/MaterialTemplate.ush
--- a/Engine/Shaders/Private/MaterialTemplate.ush
+++ b/Engine/Shaders/Private/MaterialTemplate.ush
@@ -293,7 +293,7 @@ struct FMaterialPixelParameters
 	float2 HairUV;			// U: parametric coord along the length V: parametric coord along the width)
 	float2 HairDimensions;  // X: Length, Y: Radius
 	float  HairSeed;		// [0,1] quantized on 8bits
-	float2 HairRootUV;		// [0,1]^2	8bits
+	float2 HairRootUV;		// [0,256]^2 16bits (8 bits for UDIM index, 8 bits for UV coord)
 #endif
 
 	/** Per-particle properties. Only valid for particle vertex factories. */
diff --git a/Engine/Source/Runtime/Core/Public/UObject/ReleaseObjectVersion.h b/Engine/Source/Runtime/Core/Public/UObject/ReleaseObjectVersion.h
--- a/Engine/Source/Runtime/Core/Public/UObject/ReleaseObjectVersion.h
+++ b/Engine/Source/Runtime/Core/Public/UObject/ReleaseObjectVersion.h
@@ -83,6 +83,9 @@ struct CORE_API FReleaseObjectVersion
 		// Removed versioning, but version entry must still exist to keep assets saved with this version loadable
 		Unused_SoundClass2DReverbSend,
 
+		// Groom asset version
+		GroomAssetVersion1,
+
 		// -----<new versions can be added above this line>-------------------------------------------------
 		VersionPlusOne,
 		LatestVersion = VersionPlusOne - 1
diff --git a/Engine/Source/Runtime/Renderer/Private/HairStrands/HairStrandsDebug.cpp b/Engine/Source/Runtime/Renderer/Private/HairStrands/HairStrandsDebug.cpp
--- a/Engine/Source/Runtime/Renderer/Private/HairStrands/HairStrandsDebug.cpp
+++ b/Engine/Source/Runtime/Renderer/Private/HairStrands/HairStrandsDebug.cpp
@@ -139,6 +139,7 @@ EHairStrandsDebugMode GetHairStrandsDebugStrandsMode()
 	case 5:  return EHairStrandsDebugMode::RenderHairSeed;
 	case 6:  return EHairStrandsDebugMode::RenderHairDimension;
 	case 7:  return EHairStrandsDebugMode::RenderHairRadiusVariation;
+	case 8:  return EHairStrandsDebugMode::RenderHairRootUDIM;
 	default: return EHairStrandsDebugMode::None;
 	};
 }
@@ -155,6 +156,7 @@ static const TCHAR* ToString(EHairStrandsDebugMode DebugMode)
 	case EHairStrandsDebugMode::RenderHairSeed				: return TEXT("Hair seed");
 	case EHairStrandsDebugMode::RenderHairDimension			: return TEXT("Hair dimensions");
 	case EHairStrandsDebugMode::RenderHairRadiusVariation	: return TEXT("Hair radius variation");
+	case EHairStrandsDebugMode::RenderHairRootUDIM			: return TEXT("Root UV UDIM texture index");		
 	default													: return TEXT("None");
 	};
 }
diff --git a/Engine/Source/Runtime/Renderer/Public/HairStrandsInterface.h b/Engine/Source/Runtime/Renderer/Public/HairStrandsInterface.h
--- a/Engine/Source/Runtime/Renderer/Public/HairStrandsInterface.h
+++ b/Engine/Source/Runtime/Renderer/Public/HairStrandsInterface.h
@@ -24,6 +24,7 @@ enum class EHairStrandsDebugMode : uint8
 	RenderHairSeed,
 	RenderHairDimension,
 	RenderHairRadiusVariation,
+	RenderHairRootUDIM,
 	Count
 };
 
