Fixed voxel shadow to work outside hair volume bounds.

diff --git a/Engine/Shaders/Private/HairStrands/HairStrandsDeepVoxelCommon.ush b/Engine/Shaders/Private/HairStrands/HairStrandsDeepVoxelCommon.ush
--- a/Engine/Shaders/Private/HairStrands/HairStrandsDeepVoxelCommon.ush
+++ b/Engine/Shaders/Private/HairStrands/HairStrandsDeepVoxelCommon.ush
@@ -21,36 +21,42 @@ bool IsInVoxelBounds(float3 P, float3 MinP, float3 MaxP)
 		P.x <= MaxP.x && P.y <= MaxP.y && P.z <= MaxP.z;
 }
 
-float ComputeHairCountVoxel(float3 WorldPosition, float3 LightDirection, float3 InVoxelAABBMin, float3 InVoxelAABBMax, float InVoxelResolution, Texture3D<uint> InDensityTexture, float DensityScale, float CountThreshold)
+float ComputeHairCountVoxel(
+	float3 WorldPosition, 
+	float3 LightDirection, 
+	float3 InVoxelAABBMin, 
+	float3 InVoxelAABBMax, 
+	float InVoxelResolution, 
+	Texture3D<uint> InDensityTexture, 
+	float DensityScale, 
+	float CountThreshold)
 {
-	// Reverse-Z
-	const float ConservativeAABBDiagonal = (InVoxelAABBMax.x - InVoxelAABBMin.x) * 3;
-	const float3 O = WorldPosition;
-	const float3 E = O + LightDirection * ConservativeAABBDiagonal;
-	const float3 D = E - O;
-
-	const float VoxelSize = (InVoxelAABBMax.x - InVoxelAABBMin.x) / InVoxelResolution;
-
 	float OutHairCount = 0;
 
-	// The rest of the code assumes O is within the AABB
-	const bool bIsInAABB = IsInVoxelBounds(O, InVoxelAABBMin, InVoxelAABBMax);
-	if (!bIsInAABB)
-		return OutHairCount;
+	const float3 VolumeBoundSize = InVoxelAABBMax - InVoxelAABBMin;
+	const float3 VolumeBoundSizeInv = 1.0f / (InVoxelAABBMax - InVoxelAABBMin);
+	const float VoxelWorldSize = VolumeBoundSize.x / InVoxelResolution;				// Only using voxel size along x for now...
 
-	const float2 HitT = LineBoxIntersect(O, E, InVoxelAABBMin, InVoxelAABBMax);
+	const float VoxelShadowMaxDistance = 100000.0; // 1 km shadow distance
+	const float3 IntersectEndPoint = WorldPosition + LightDirection * VoxelShadowMaxDistance;
+	const float2 HitT = LineBoxIntersect(WorldPosition, IntersectEndPoint, InVoxelAABBMin, InVoxelAABBMax);
 	if (HitT.x < HitT.y)
 	{
 		// Count the number of fibers which are within a cylinder defined by the voxel size, 
 		// and the distance between the origin and the extent of the volume
 		// This assumes that the voxel volume is cubic (i.e. equal dimensions on all sides)
-		const float Distance = length(HitT.y * D);
-		const uint StepCount = min(ceil(Distance / VoxelSize), 256u);
-		const float Delta = VoxelSize;
-		for (uint StepIt = 0; StepIt < StepCount; ++StepIt)
+		float3 O = lerp(WorldPosition, IntersectEndPoint, HitT.xxx);
+		float3 E = lerp(WorldPosition, IntersectEndPoint, HitT.yyy);
+		const float OELength = length(E - O);
+
+		// Step according to voxel size
+		const float StepCount = float(min(ceil(OELength / VoxelWorldSize), 256u));	// TO TRY: sample with fixed count budget and jitter temporally
+		LightDirection *= VoxelWorldSize;
+
+		for (float StepIt = 0.0f; StepIt < StepCount; ++StepIt)
 		{
-			const float3 HitP = O + (StepIt*Delta) * LightDirection;
-			const float3 UVW = (HitP - InVoxelAABBMin) / (InVoxelAABBMax - InVoxelAABBMin);
+			const float3 HitP = O + StepIt * LightDirection;
+			const float3 UVW = (HitP - InVoxelAABBMin) * VolumeBoundSizeInv;
 			const int3 VolumeCoord = clamp(int3(UVW * InVoxelResolution), int3(0, 0, 0), int3(InVoxelResolution.xxx-1));
 			const float Density = GetHairVoxelDensity(VolumeCoord, InDensityTexture) * DensityScale;
 			OutHairCount += Density;
