Improved coverage again!
Now that I am not remote I have been able to really look at pixel quality :)

diff --git a/Engine/Shaders/Private/HairStrands/HairStrandsVisibilityComposeSubPixelPS.usf b/Engine/Shaders/Private/HairStrands/HairStrandsVisibilityComposeSubPixelPS.usf
--- a/Engine/Shaders/Private/HairStrands/HairStrandsVisibilityComposeSubPixelPS.usf
+++ b/Engine/Shaders/Private/HairStrands/HairStrandsVisibilityComposeSubPixelPS.usf
@@ -23,11 +23,12 @@ void SubColorPS(
 	//	discard;
 
 	// Brute force reading of the depth value. This could be pre-computed during the categorize pass or the compaction pass
+	// This must match what is done in VisibilityDepthPS to fonction. TODO: do not write depth here as it is already in. 
 	OutDepth = 0;
 	for (uint SampleIt = 0; SampleIt < HAIR_VISIBILITY_MAXSAMPLECOUNT; ++SampleIt)
 	{
 		const float SampleDepth = HairVisibilityDepthTexture.Load(PixelCoord, SampleIt);
-		OutDepth = max(OutDepth, SampleDepth); // Inverse-Z
+		OutDepth = max(OutDepth, SampleDepth); // HAS_INVERTED_Z_BUFFER
 	}
 }
 
diff --git a/Engine/Shaders/Private/HairStrands/HairStrandsVisibilityDepthPS.usf b/Engine/Shaders/Private/HairStrands/HairStrandsVisibilityDepthPS.usf
--- a/Engine/Shaders/Private/HairStrands/HairStrandsVisibilityDepthPS.usf
+++ b/Engine/Shaders/Private/HairStrands/HairStrandsVisibilityDepthPS.usf
@@ -20,27 +20,35 @@ void MainPS(
 	out float4 OutColor		: SV_Target1
 )
 {
-	float MaxDepth = 0;
-	float MinDepth = 1;
+	// HAS_INVERTED_Z_BUFFER
+	float ClosestDepth = 0;
+	float FarthestDepth = 1;
+
 	const uint2 PixelCoord = floor(Input.Position.xy);
 	uint ValidSampleCount = 0;
 	for (uint SampleIt=0; SampleIt<HAIR_VISIBILITY_MAXSAMPLECOUNT; ++SampleIt)
 	{
 		const float SampleDepth = HairVisibilityDepthTexture.Load(PixelCoord, SampleIt).x;
-		MinDepth = min(MinDepth, SampleDepth);
-		MaxDepth = max(MaxDepth, SampleDepth);
-		const bool bIsValid = SampleDepth > 0;
+		FarthestDepth = min(FarthestDepth, SampleDepth);
+		ClosestDepth  = max(ClosestDepth, SampleDepth);
+		const bool bIsValid = SampleDepth > 0.0f;
 		ValidSampleCount += bIsValid ? 1 : 0;
 	}
 
-	if (ValidSampleCount < HAIR_VISIBILITY_MAXSAMPLECOUNT
 #if PERMUTATION_COVERAGE == 1
-		|| CoverageTexture.Load(uint3(PixelCoord, 0))<1.0f
+	if (CoverageTexture.Load(uint3(PixelCoord, 0)) < 1.0f || ValidSampleCount == 0)
+	{
+		discard; // If coverage is 1. we skip writting depth to match all the other passes relying on coverage==1 meaning opaque.
+	}
+	SvDepth = ClosestDepth; // Keep closest sample (some mayb be far plane but we only want hair samples)
+#else
+	if (ValidSampleCount < HAIR_VISIBILITY_MAXSAMPLECOUNT)
+	{
+		discard; // If we do not have all samples we skip to ignore depth from far clip plane
+	}
+	SvDepth = FarthestDepth; // Farthest is fine to keep because all samples contains hair data
 #endif
-		)
-		discard;
 
-	SvDepth = MinDepth;
 	OutGBufferB = 0;
 	OutColor = 0;
 }
