Move attributes sampling into vertex shader.

This make the sampling optional, based on the material properties. This will allows to provide more (optional) attributes later on, as this won't cluter anymore the interpolants.

diff --git a/Engine/Plugins/Experimental/HairStrands/Source/HairStrandsCore/Private/HairStrandsVertexFactory.cpp b/Engine/Plugins/Experimental/HairStrands/Source/HairStrandsCore/Private/HairStrandsVertexFactory.cpp
--- a/Engine/Plugins/Experimental/HairStrands/Source/HairStrandsCore/Private/HairStrandsVertexFactory.cpp
+++ b/Engine/Plugins/Experimental/HairStrands/Source/HairStrandsCore/Private/HairStrandsVertexFactory.cpp
@@ -189,6 +189,11 @@ FVertexFactoryShaderParameters* FHairStrandsVertexFactory::ConstructShaderParame
 		return new FHairStrandsVertexFactoryShaderParameters();
 	}
 
+	if (ShaderFrequency == SF_Pixel)
+	{
+		return new FHairStrandsVertexFactoryShaderParameters();
+	}
+
 #if RHI_RAYTRACING
 	if (ShaderFrequency == SF_RayHitGroup)
 	{
diff --git a/Engine/Shaders/Private/HairStrands/HairStrandsVertexFactory.ush b/Engine/Shaders/Private/HairStrands/HairStrandsVertexFactory.ush
--- a/Engine/Shaders/Private/HairStrands/HairStrandsVertexFactory.ush
+++ b/Engine/Shaders/Private/HairStrands/HairStrandsVertexFactory.ush
@@ -72,6 +72,63 @@ Buffer<float4> HairStrandsVF_TangentBuffer;
 #define FPositionOnlyVertexFactoryInput FVertexFactoryInput
 #define FPositionAndNormalOnlyVertexFactoryInput FVertexFactoryInput
 
+struct FHairVertexAttributes
+{
+	float2 HairUV;
+	float2 HairRootUV;
+	float2 HairDimensions;
+	float3 HairBaseColor;
+	float  HairSeed;
+	float  HairRoughness;
+	float  HairSpecular;
+};
+
+FHairVertexAttributes GetHairStrandsAttributes(uint HairPrimitiveId, float2 VertexUV)
+{
+	const uint VertexId0 = HairPrimitiveId;
+	const uint VertexId1 = HairPrimitiveId+1;
+
+	// Instead of doing the actual interpolation we could rely on dithering as an optimization
+	const FHairAttributes Attributes0 = DecodeHairAttribute(HairStrandsVF_AttributeBuffer[VertexId0]);
+	const FHairAttributes Attributes1 = DecodeHairAttribute(HairStrandsVF_AttributeBuffer[VertexId1]);
+
+	const uint4 PackedData = HairStrandsVF_PositionBuffer[VertexId0];
+	const uint PackedAlpha = asuint(PackedData.a);
+	const float WorldRadius = ((PackedAlpha & 0x00FC) >> 2) / 63.f * HairStrandsVF_Radius;
+	const float WorldLength = ((PackedAlpha & 0xFF00) >> 8) / 255.f * HairStrandsVF_Length;
+
+	FHairVertexAttributes Out;
+	// Constant along the strand
+	Out.HairDimensions	= float2(WorldLength, WorldRadius);
+	Out.HairSeed		= Attributes0.Seed;
+	Out.HairRootUV		= Attributes0.RootUV;
+
+	// Varying along the strand
+	Out.HairUV.y		= VertexUV.y;
+	Out.HairUV.x		= lerp(Attributes0.UCoord, Attributes1.UCoord, VertexUV.x);
+	return Out; 
+}
+
+float2 GetHairStrandsUV(uint HairPrimitiveId, float2 HairPrimitiveUV)
+{
+	return GetHairStrandsAttributes(HairPrimitiveId, HairPrimitiveUV).HairUV;
+}
+
+float2 GetHairStrandsDimensions(uint HairPrimitiveId, float2 HairPrimitiveUV)
+{
+	return GetHairStrandsAttributes(HairPrimitiveId, HairPrimitiveUV).HairDimensions;
+}
+
+float2 GetHairStrandsRootUV(uint HairPrimitiveId, float2 HairPrimitiveUV)
+{
+	return GetHairStrandsAttributes(HairPrimitiveId, HairPrimitiveUV).HairRootUV;
+}
+
+float  GetHairStrandsSeed(uint HairPrimitiveId, float2 HairPrimitiveUV)
+{
+	return GetHairStrandsAttributes(HairPrimitiveId, HairPrimitiveUV).HairSeed;
+}
+
 struct FVertexFactoryInput
 {
 #if VF_USE_PRIMITIVE_SCENE_DATA
@@ -106,10 +163,8 @@ struct FVertexFactoryIntermediates
 	half TangentToWorldSign;
 	uint PrimitiveId;
 	uint   HairPrimitiveId;
-	float2 HairDimensions;
-	float2 HairUV;
-	float2 HairRootUV;
-	float  HairSeed;
+	float2 HairPrimitiveUV;
+	float2 HairDimensions;  // This is kept on the intermediate as this is used in several vertex shader for the actual coverage.
 	float  HairDensity;
 };
 
@@ -131,10 +186,7 @@ struct FControlPoint
 {
 	float3 Position;
 	float  WorldRadius;
-	float  UCoord;
 	float  WorldLength;
-	float  Seed;
-	float2 RootUV;
 };
 
 FControlPoint GetVertexPosition(FVertexFactoryInput Input)
@@ -155,15 +207,21 @@ FControlPoint GetVertexPosition(FVertexFactoryInput Input)
 	// Create a degenerated quad the end of each strand to cut between each strands
 	// #hair_todo: This is not efficient for short strand like fur for instance. Need to revisit that at some point
 	const bool bIsInvalidQuad = (ControlPointType == 1 && OffsetIndex == 1) || (ControlPointType == 2 && OffsetIndex == 0);
-	const FHairAttributes Attributes = DecodeHairAttribute(HairStrandsVF_AttributeBuffer[BaseIndex + OffsetIndex]);
-	Output.Position				= bIsInvalidQuad ? float3(1.#INF, 1.#INF, 1.#INF) : ControlPoint;
-	Output.RootUV				= Attributes.RootUV;
-	Output.UCoord				= Attributes.UCoord;
-	Output.Seed					= Attributes.Seed;
+	Output.Position = bIsInvalidQuad ? float3(1.#INF, 1.#INF, 1.#INF) : ControlPoint;
 
 	return Output;
 }
 
+float GetWorldStrandRadius(FVertexFactoryInput Input)
+{
+	return GetVertexPosition(Input).WorldRadius;
+}
+
+float GetWorldStrandLength(FVertexFactoryInput Input)
+{
+	return GetVertexPosition(Input).WorldLength;
+}
+
 float3 GetVertexPreviousPosition(FVertexFactoryInput Input)
 {
 	const uint QuadIndex = Input.VertexId % 6;
@@ -183,15 +241,6 @@ float3 GetVertexPreviousPosition(FVertexFactoryInput Input)
 	return Position;
 }
 
-float GetWorldStrandRadius(FVertexFactoryInput Input)
-{
-	return GetVertexPosition(Input).WorldRadius;
-}
-
-float GetUCoord(FVertexFactoryInput Input)
-{
-	return GetVertexPosition(Input).UCoord;
-}
 
 float GetVCoord(FVertexFactoryInput Input)
 {
@@ -200,21 +249,6 @@ float GetVCoord(FVertexFactoryInput Input)
 	return VCoord;
 }
 
-float GetWorldStrandLength(FVertexFactoryInput Input)
-{
-	return GetVertexPosition(Input).WorldLength;
-}
-
-float GetHairSeed(FVertexFactoryInput Input)
-{
-	return GetVertexPosition(Input).Seed;
-}
-
-float2 GetHairRootUV(FVertexFactoryInput Input)
-{
-	return GetVertexPosition(Input).RootUV;
-}
-
 /** Converts from vertex factory specific interpolants to a FMaterialPixelParameters, which is used by material inputs. */
 FMaterialPixelParameters GetMaterialPixelParameters(FVertexFactoryInterpolantsVSToPS Interpolants, float4 SvPosition)
 {
@@ -236,10 +270,7 @@ FMaterialPixelParameters GetMaterialPixelParameters(FVertexFactoryInterpolantsVS
 	Result.TwoSidedSign = 1;
 	Result.PrimitiveId = GetPrimitiveId(Interpolants);
 	Result.HairPrimitiveId = Interpolants.HairPrimitiveId;
-	Result.HairDimensions = Interpolants.HairDimensions;
-	Result.HairUV = Interpolants.HairUV;
-	Result.HairRootUV = Interpolants.HairRootUV;
-	Result.HairSeed = Interpolants.HairSeed;
+	Result.HairPrimitiveUV = float2(Interpolants.HairPrimitiveUV);
 
 #if COMPILER_SUPPORTS_TO_SCALAR_MEMORY
 	Result.PrimitiveId = ToScalarMemory(Result.PrimitiveId);
@@ -331,13 +362,11 @@ float3 GetTangent(FVertexFactoryInput Input)
 FVertexFactoryIntermediates GetVertexFactoryIntermediates(FVertexFactoryInput Input)
 {
 	FVertexFactoryIntermediates Intermediates;
-	Intermediates.PrimitiveId = GetPrimitiveId(Input);
-	Intermediates.HairPrimitiveId = GetHairPrimitiveId(Input);
-	Intermediates.HairDimensions = float2(GetWorldStrandLength(Input), GetWorldStrandRadius(Input));
-	Intermediates.HairUV = float2(GetUCoord(Input), GetVCoord(Input));
-	Intermediates.HairDensity = HairStrandsVF_Density;
-	Intermediates.HairSeed = GetHairSeed(Input);
-	Intermediates.HairRootUV = GetHairRootUV(Input);
+	Intermediates.PrimitiveId		= GetPrimitiveId(Input);
+	Intermediates.HairPrimitiveId	= GetHairPrimitiveId(Input);
+	Intermediates.HairPrimitiveUV	= float2(0 /*TODO*/, GetVCoord(Input));
+	Intermediates.HairDensity		= HairStrandsVF_Density;
+	Intermediates.HairDimensions	= float2(GetWorldStrandLength(Input), GetWorldStrandRadius(Input));
 
 	float TangentSign;
 	Intermediates.TangentToLocal = CalcTangentToLocal(Input, TangentSign);
@@ -418,10 +447,7 @@ FVertexFactoryInterpolantsVSToPS VertexFactoryGetInterpolantsVSToPS(FVertexFacto
 #endif
 	
 	Interpolants.HairPrimitiveId = GetHairPrimitiveId(Input);
-	Interpolants.HairDimensions = float2(GetWorldStrandLength(Input), GetWorldStrandRadius(Input));
-	Interpolants.HairUV = float2(GetUCoord(Input), GetVCoord(Input));
-	Interpolants.HairSeed = GetHairSeed(Input);
-	Interpolants.HairRootUV = GetHairRootUV(Input);
+	Interpolants.HairPrimitiveUV = float2(0/*TODO*/, GetVCoord(Input));
 	SetPrimitiveId(Interpolants, Intermediates.PrimitiveId);
 
 	return Interpolants;
diff --git a/Engine/Shaders/Private/LocalVertexFactoryCommon.ush b/Engine/Shaders/Private/LocalVertexFactoryCommon.ush
--- a/Engine/Shaders/Private/LocalVertexFactoryCommon.ush
+++ b/Engine/Shaders/Private/LocalVertexFactoryCommon.ush
@@ -41,11 +41,8 @@ struct FVertexFactoryInterpolantsVSToPS
 	#endif
 #endif
 #if VF_STRAND_HAIR
-	nointerpolation uint HairPrimitiveId	: HAIR_PRIMITIVE_ID;
-	float2 HairDimensions					: HAIR_DIMENSIONS;
-	float2 HairUV							: HAIR_UV;
-	float2 HairRootUV						: HAIR_ROOT_UV;
-	float  HairSeed							: HAIR_SEED;
+	nointerpolation uint HairPrimitiveId	: HAIR_PRIMITIVE_ID; // Control point ID
+	float2 HairPrimitiveUV					: HAIR_PRIMITIVE_UV; // U: parameteric distance between the two surrounding control points. V: parametric distance along the width.
 #endif
 };
 
diff --git a/Engine/Shaders/Private/MaterialTemplate.ush b/Engine/Shaders/Private/MaterialTemplate.ush
--- a/Engine/Shaders/Private/MaterialTemplate.ush
+++ b/Engine/Shaders/Private/MaterialTemplate.ush
@@ -289,11 +289,8 @@ struct FMaterialPixelParameters
 
 	// Actual primitive Id
 #if VF_STRAND_HAIR
-	uint HairPrimitiveId;
-	float2 HairUV;			// U: parametric coord along the length V: parametric coord along the width)
-	float2 HairDimensions;  // X: Length, Y: Radius
-	float  HairSeed;		// [0,1] quantized on 8bits
-	float2 HairRootUV;		// [0,256]^2 16bits (8 bits for UDIM index, 8 bits for UV coord)
+	uint	HairPrimitiveId;	// Control point ID
+	float2	HairPrimitiveUV;	// U: parametric distance between the two surrounding control point. V: parametric distance along hair width
 #endif
 
 	/** Per-particle properties. Only valid for particle vertex factories. */
@@ -999,10 +996,18 @@ float3 MaterialExpressionBlackBody( float Temp )
 	return mul( XYZtoRGB, float3( X, Y, Z ) ) * pow( 0.0004 * Temp, 4 );
 }
 
+
+#if VF_STRAND_HAIR
+float2 GetHairStrandsRootUV(uint HairPrimitiveId, float2 VertexUV);
+float2 GetHairStrandsUV(uint HairPrimitiveId, float2 VertexUV);
+float2 GetHairStrandsDimensions(uint HairPrimitiveId, float2 VertexUV);
+float  GetHairStrandsSeed(uint HairPrimitiveId, float2 VertexUV);
+#endif
+
 float2 MaterialExpressionGetHairRootUV(FMaterialPixelParameters Parameters)
 {
 #if VF_STRAND_HAIR
-	return Parameters.HairRootUV;
+	return GetHairStrandsRootUV(Parameters.HairPrimitiveId, Parameters.HairPrimitiveUV);
 #else
 	return float2(0, 0);
 #endif
@@ -1011,7 +1016,7 @@ float2 MaterialExpressionGetHairRootUV(FMaterialPixelParameters Parameters)
 float2 MaterialExpressionGetHairUV(FMaterialPixelParameters Parameters)
 {	
 #if VF_STRAND_HAIR
-	return Parameters.HairUV;
+	return GetHairStrandsUV(Parameters.HairPrimitiveId, Parameters.HairPrimitiveUV);
 #else
 	return float2(0,0);
 #endif
@@ -1020,7 +1025,7 @@ float2 MaterialExpressionGetHairUV(FMaterialPixelParameters Parameters)
 float2 MaterialExpressionGetHairDimensions(FMaterialPixelParameters Parameters)
 {
 #if VF_STRAND_HAIR
-	return Parameters.HairDimensions; // Length | Radius
+	return GetHairStrandsDimensions(Parameters.HairPrimitiveId, Parameters.HairPrimitiveUV);
 #else
 	return float2(0, 0);
 #endif
@@ -1029,7 +1034,7 @@ float2 MaterialExpressionGetHairDimensions(FMaterialPixelParameters Parameters)
 float MaterialExpressionGetHairSeed(FMaterialPixelParameters Parameters)
 {
 #if VF_STRAND_HAIR
-	return Parameters.HairSeed;
+	return GetHairStrandsSeed(Parameters.HairPrimitiveId, Parameters.HairPrimitiveUV);
 #else
 	return 0;
 #endif
