Add tangent print into hair details debug info.

diff --git a/Engine/Shaders/Private/HairStrands/HairStrandsDebugPrint.usf b/Engine/Shaders/Private/HairStrands/HairStrandsDebugPrint.usf
--- a/Engine/Shaders/Private/HairStrands/HairStrandsDebugPrint.usf
+++ b/Engine/Shaders/Private/HairStrands/HairStrandsDebugPrint.usf
@@ -299,6 +299,22 @@ void MainCS(uint GroupIndex : SV_GroupIndex, uint3 DispatchThreadId : SV_Dispatc
 
 		Pos = ShaderPrintNewline(Pos);
 		Pos.x = OriginalPos.x;
+
+
+		// Tangent
+		Pos = ShaderPrintSymbol(Pos, _DOT_);
+		Pos = ShaderPrintSymbol(Pos, _T_);
+		Pos = ShaderPrintSymbol(Pos, _A_);
+		Pos = ShaderPrintSymbol(Pos, _N_);
+		Pos = ShaderPrintSymbol(Pos, _G_);
+		Pos = ShaderPrintSymbol(Pos, _E_);
+		Pos = ShaderPrintSymbol(Pos, _N_);
+		Pos = ShaderPrintSymbol(Pos, _T_);
+		Pos = ShaderPrintSymbol(Pos, _SPC_);
+		Pos = ShaderPrint(Pos, SampleVisData.Tangent);
+
+		Pos = ShaderPrintNewline(Pos);
+		Pos.x = OriginalPos.x;
 #endif 
 
 #if DETAIL_DATA
