Fix missing U interpolation for hair attribute/material.

diff --git a/Engine/Shaders/Private/HairStrands/HairStrandsVertexFactory.ush b/Engine/Shaders/Private/HairStrands/HairStrandsVertexFactory.ush
--- a/Engine/Shaders/Private/HairStrands/HairStrandsVertexFactory.ush
+++ b/Engine/Shaders/Private/HairStrands/HairStrandsVertexFactory.ush
@@ -260,12 +260,13 @@ float3 GetVertexPreviousPosition(FVertexFactoryInput Input)
 	return Position;
 }
 
-
-float GetVCoord(FVertexFactoryInput Input)
+// Segment UV coord of an hair segment. This is different from the UV coord of the hair strands
+float GetSegmentUVCoord(FVertexFactoryInput Input)
 {
 	const uint QuadIndex = Input.VertexId % 6;
 	const float VCoord = QuadIndex == 0 || QuadIndex == 1 || QuadIndex == 5 ? 0.0f : 1.0f;
-	return VCoord;
+	const float UCoord = QuadIndex == 0 || QuadIndex == 2 || QuadIndex == 4 ? 0.0f : 1.0f;
+	return float2(UCoord, VCoord);
 }
 
 /** Converts from vertex factory specific interpolants to a FMaterialPixelParameters, which is used by material inputs. */
@@ -383,7 +384,7 @@ FVertexFactoryIntermediates GetVertexFactoryIntermediates(FVertexFactoryInput In
 	FVertexFactoryIntermediates Intermediates;
 	Intermediates.PrimitiveId		= GetPrimitiveId(Input);
 	Intermediates.HairPrimitiveId	= GetHairPrimitiveId(Input);
-	Intermediates.HairPrimitiveUV	= float2(0 /*TODO*/, GetVCoord(Input));
+	Intermediates.HairPrimitiveUV	= GetSegmentUVCoord(Input);
 	Intermediates.HairDensity		= HairStrandsVF_Density;
 	Intermediates.HairDimensions	= float2(GetWorldStrandLength(Input), GetWorldStrandRadius(Input));
 
@@ -466,7 +467,7 @@ FVertexFactoryInterpolantsVSToPS VertexFactoryGetInterpolantsVSToPS(FVertexFacto
 #endif
 	
 	Interpolants.HairPrimitiveId = GetHairPrimitiveId(Input);
-	Interpolants.HairPrimitiveUV = float2(0/*TODO*/, GetVCoord(Input));
+	Interpolants.HairPrimitiveUV = GetSegmentUVCoord(Input);
 	SetPrimitiveId(Interpolants, Intermediates.PrimitiveId);
 
 	return Interpolants;
