Try to keep  ImageBasedReflectionLighting coherence  between mobile and PC.

diff --git a/Engine/Shaders/Private/MobileBasePassPixelShader.usf b/Engine/Shaders/Private/MobileBasePassPixelShader.usf
--- a/Engine/Shaders/Private/MobileBasePassPixelShader.usf
+++ b/Engine/Shaders/Private/MobileBasePassPixelShader.usf
@@ -79,7 +79,7 @@ float4 SpotLightDirectionAndSpecularScale[MAX_DYNAMIC_POINT_LIGHTS];
 	TextureCube ReflectionCubemap2;
 	SamplerState ReflectionCubemapSampler2;
 	// x,y,z - inverted average brightness for 0, 1, 2
-	float4 ReflectionInvAverageBrigtness;
+	float4 ReflectionAverageBrigtness;
 	float4 ReflectionPositionsAndRadii[MAX_HQ_REFLECTIONS];
 		#if ALLOW_CUBE_REFLECTIONS
 		float4x4 CaptureBoxTransformArray[MAX_HQ_REFLECTIONS];
@@ -194,7 +194,7 @@ float4 GetBoxCaptureVector(float3 ReflectionVector, float3 WorldPosition, float4
 }
 #endif
 
-half3 BlendReflectionCapture(FMaterialPixelParameters MaterialParameters, half Roughness, TextureCube ReflectionCube, SamplerState ReflectionSampler, half MipLevel, int ReflectionIndex, half InvReflectionCubemapAverageBrightness, inout half BlendFactor)
+half3 BlendReflectionCapture(FMaterialPixelParameters MaterialParameters, TextureCube ReflectionCube, SamplerState ReflectionSampler, half MipLevel, int ReflectionIndex, inout half BlendFactor)
 {
 	half4 ProjectedCaptureVector;
 #if ALLOW_CUBE_REFLECTIONS
@@ -209,21 +209,33 @@ half3 BlendReflectionCapture(FMaterialPixelParameters MaterialParameters, half R
 	}
 
 	// Fetch from cubemap and convert to linear HDR
-	float4 Reflection = ReflectionCube.SampleLevel(ReflectionSampler, ProjectedCaptureVector.xyz, MipLevel);
+	half4 Reflection = ReflectionCube.SampleLevel(ReflectionSampler, ProjectedCaptureVector.xyz, MipLevel);
 	half3 SpecularIBL = RGBMDecode(Reflection, 16.0f);
 
-#if !LQ_TEXTURE_LIGHTMAP && !CACHED_POINT_INDIRECT_LIGHTING
-	// ignore average brightness for non-lightmap use case.
-	InvReflectionCubemapAverageBrightness = 1.0f;
-#endif
-
-	SpecularIBL = SpecularIBL * SpecularIBL * (InvReflectionCubemapAverageBrightness * ProjectedCaptureVector.w * BlendFactor);
+	SpecularIBL = SpecularIBL * SpecularIBL * (ProjectedCaptureVector.w * BlendFactor);
 
 	BlendFactor = BlendFactor * (1.0 - ProjectedCaptureVector.w);
 	return SpecularIBL;
 }
 
-half3 BlendReflectionCaptures(FMaterialPixelParameters MaterialParameters, half Roughness)
+// Half precision version for mobile's high quality reflection
+half3 MobileComputeMixingWeight3(half IndirectIrradiance, half3 AverageBrightness, half Roughness)
+{
+	// Mirror surfaces should have no mixing, so they match reflections from other sources (SSR, planar reflections)
+	half MixingAlpha = smoothstep(0, 1, saturate(Roughness * View.ReflectionEnvironmentRoughnessMixingScaleBiasAndLargestWeight.x + View.ReflectionEnvironmentRoughnessMixingScaleBiasAndLargestWeight.y));
+
+	// We have high frequency directional data but low frequency spatial data in the envmap.
+	// We have high frequency spatial data but low frequency directional data in the lightmap.
+	// So, we combine the two for the best of both. This is done by removing the low spatial frequencies from the envmap and replacing them with the lightmap data.
+	// This is only done with luma so as to not get odd color shifting.
+	half3 MixingWeight = half3(IndirectIrradiance, IndirectIrradiance, IndirectIrradiance) / max(AverageBrightness, half3(.0001f, .0001f, .0001f));
+
+	MixingWeight = min(MixingWeight, half3(View.ReflectionEnvironmentRoughnessMixingScaleBiasAndLargestWeight.z, View.ReflectionEnvironmentRoughnessMixingScaleBiasAndLargestWeight.z, View.ReflectionEnvironmentRoughnessMixingScaleBiasAndLargestWeight.z));
+
+	return lerp(half3(1.0f, 1.0f, 1.0f), MixingWeight, MixingAlpha);
+}
+
+half3 BlendReflectionCaptures(FMaterialPixelParameters MaterialParameters, half Roughness, float IndirectIrradiance)
 {
 	// Compute fractional mip from roughness
 	half AbsoluteSpecularMip = ComputeReflectionCaptureMipFromRoughness(Roughness, ResolvedView.ReflectionCubemapMaxMip);
@@ -231,9 +243,21 @@ half3 BlendReflectionCaptures(FMaterialPixelParameters MaterialParameters, half
 	half3 SpecularIBL = half3(0, 0, 0);
 	half BlendFactor = 1;
 
-	SpecularIBL += BlendReflectionCapture(MaterialParameters, Roughness, ReflectionCubemap0, ReflectionCubemapSampler0, AbsoluteSpecularMip, 0, ReflectionInvAverageBrigtness.x, BlendFactor);
-	SpecularIBL += BlendReflectionCapture(MaterialParameters, Roughness, ReflectionCubemap1, ReflectionCubemapSampler1, AbsoluteSpecularMip, 1, ReflectionInvAverageBrigtness.y, BlendFactor);
-	SpecularIBL += BlendReflectionCapture(MaterialParameters, Roughness, ReflectionCubemap2, ReflectionCubemapSampler2, AbsoluteSpecularMip, 2, ReflectionInvAverageBrigtness.z, BlendFactor);
+	half3 SpecularIBL0 = BlendReflectionCapture(MaterialParameters, ReflectionCubemap0, ReflectionCubemapSampler0, AbsoluteSpecularMip, 0, BlendFactor);
+	half3 SpecularIBL1 = BlendReflectionCapture(MaterialParameters, ReflectionCubemap1, ReflectionCubemapSampler1, AbsoluteSpecularMip, 1, BlendFactor);
+	half3 SpecularIBL2 = BlendReflectionCapture(MaterialParameters, ReflectionCubemap2, ReflectionCubemapSampler2, AbsoluteSpecularMip, 2, BlendFactor);
+
+	half3 ReflectionCubemapAverageBrightness = ReflectionAverageBrigtness.xyz;
+
+#if !LQ_TEXTURE_LIGHTMAP && !CACHED_POINT_INDIRECT_LIGHTING
+	// ignore average brightness for non-lightmap use case.
+	ReflectionCubemapAverageBrightness = 0.0f;
+#endif
+
+	//To keep ImageBasedReflectionLighting conherence with PC, use ComputeMixingWeight instead of InvReflectionAverageBrightness to calulate the IBL constribution
+	half3 MixingWeight = MobileComputeMixingWeight3(IndirectIrradiance, ReflectionCubemapAverageBrightness, Roughness);
+
+	SpecularIBL = SpecularIBL0 * MixingWeight.x + SpecularIBL1 * MixingWeight.y + SpecularIBL2 * MixingWeight.z;
 
 	return SpecularIBL;
 }
@@ -253,10 +277,27 @@ float4 GetPlanarReflection(float3 WorldPosition, float3 WorldNormal, float Rough
 	return PlanarReflection;
 }
 
-half3 GetImageBasedReflectionLighting(FMaterialPixelParameters MaterialParameters, half Roughness)
+// Half precision version for mobile's high quality reflection
+half MobileComputeMixingWeight(half IndirectIrradiance, half AverageBrightness, half Roughness)
+{
+	// Mirror surfaces should have no mixing, so they match reflections from other sources (SSR, planar reflections)
+	half MixingAlpha = smoothstep(0, 1, saturate(Roughness * View.ReflectionEnvironmentRoughnessMixingScaleBiasAndLargestWeight.x + View.ReflectionEnvironmentRoughnessMixingScaleBiasAndLargestWeight.y));
+
+	// We have high frequency directional data but low frequency spatial data in the envmap.
+	// We have high frequency spatial data but low frequency directional data in the lightmap.
+	// So, we combine the two for the best of both. This is done by removing the low spatial frequencies from the envmap and replacing them with the lightmap data.
+	// This is only done with luma so as to not get odd color shifting.
+	half MixingWeight = IndirectIrradiance / max(AverageBrightness, .0001f);
+
+	MixingWeight = min(MixingWeight, View.ReflectionEnvironmentRoughnessMixingScaleBiasAndLargestWeight.z);
+
+	return lerp(1.0f, MixingWeight, MixingAlpha);
+}
+
+half3 GetImageBasedReflectionLighting(FMaterialPixelParameters MaterialParameters, half Roughness, half IndirectIrradiance)
 {
 #if HQ_REFLECTIONS
-	half3 SpecularIBL = BlendReflectionCaptures(MaterialParameters, Roughness);
+	half3 SpecularIBL = BlendReflectionCaptures(MaterialParameters, Roughness, IndirectIrradiance);
 #else
 
 	half3 ProjectedCaptureVector = MaterialParameters.ReflectionVector;	
@@ -279,10 +320,10 @@ half3 GetImageBasedReflectionLighting(FMaterialPixelParameters MaterialParameter
 	{
 		SpecularIBL = RGBMDecode(SpecularIBLSample, 16.0);
 		SpecularIBL = SpecularIBL * SpecularIBL;
-		#if LQ_TEXTURE_LIGHTMAP || CACHED_POINT_INDIRECT_LIGHTING
-			// divide by average brightness for lightmap use cases only.
-			SpecularIBL *= MobileReflectionCapture.Params.x;
-		#endif
+		
+		//To keep ImageBasedReflectionLighting conherence with PC, use ComputeMixingWeight instead of InvReflectionAverageBrightness to calulate the IBL constribution
+		SpecularIBL *= MobileComputeMixingWeight(IndirectIrradiance, MobileReflectionCapture.Params.x, Roughness);
+
 	}
 
 #endif
@@ -446,7 +487,8 @@ void Main(
 #endif
 	
 	half3 Color = 0;
-	half IndirectIrradiance = 1;
+	//To keep IndirectLightingCache conherence with PC, initialize the IndirectIrradiance to zero.
+	half IndirectIrradiance = 0;
 
 	#if LQ_TEXTURE_LIGHTMAP
 		float2 LightmapUV0, LightmapUV1;
@@ -570,7 +612,7 @@ void Main(
 			// MobileDirectionalLight.DirectionalLightDistanceFadeMADAndSpecularScale.z saves SpecularScale for direction light.
 			Color += (Shadow * NoL) * MobileDirectionalLight.DirectionalLightColor.rgb * (DiffuseColor + SpecularColor * MobileDirectionalLight.DirectionalLightDistanceFadeMADAndSpecularScale.z * CalcSpecular(Roughness, RoL, NoH, H, MaterialParameters.WorldNormal));
 			// Environment map has been prenormalized, scale by lightmap luminance
-			half3 SpecularIBL = GetImageBasedReflectionLighting(MaterialParameters, Roughness) * IndirectIrradiance;
+			half3 SpecularIBL = GetImageBasedReflectionLighting(MaterialParameters, Roughness, IndirectIrradiance);
 
 			#if MATERIAL_PLANAR_FORWARD_REFLECTIONS
 				BRANCH
diff --git a/Engine/Source/Runtime/Engine/Private/Components/ReflectionCaptureComponent.cpp b/Engine/Source/Runtime/Engine/Private/Components/ReflectionCaptureComponent.cpp
--- a/Engine/Source/Runtime/Engine/Private/Components/ReflectionCaptureComponent.cpp
+++ b/Engine/Source/Runtime/Engine/Private/Components/ReflectionCaptureComponent.cpp
@@ -1284,7 +1284,6 @@ void FReflectionCaptureProxy::SetTransform(const FMatrix& InTransform)
 
 void FReflectionCaptureProxy::UpdateMobileUniformBuffer()
 {
-	float InvBrightness = FMath::Max(FMath::Min(1.0f / EncodedHDRAverageBrightness, 65504.f), -65504.f);
 	FTexture* CaptureTexture = GBlackTextureCube;
 	if (EncodedHDRCubemap)
 	{
@@ -1293,7 +1292,8 @@ void FReflectionCaptureProxy::UpdateMobileUniformBuffer()
 	}
 		
 	FMobileReflectionCaptureShaderParameters Parameters;
-	Parameters.Params = FVector4(InvBrightness, 0.f, 0.f, 0.f);
+	//To keep ImageBasedReflectionLighting coherence with PC, use AverageBrightness instead of InvAverageBrightness to calculate the IBL contribution
+	Parameters.Params = FVector4(EncodedHDRAverageBrightness, 0.f, 0.f, 0.f);
 	Parameters.Texture = CaptureTexture->TextureRHI;
 	Parameters.TextureSampler = CaptureTexture->SamplerStateRHI;
 
diff --git a/Engine/Source/Runtime/Renderer/Private/MobileBasePass.cpp b/Engine/Source/Runtime/Renderer/Private/MobileBasePass.cpp
--- a/Engine/Source/Runtime/Renderer/Private/MobileBasePass.cpp
+++ b/Engine/Source/Runtime/Renderer/Private/MobileBasePass.cpp
@@ -463,7 +463,7 @@ void TMobileBasePassPSPolicyParamType<FUniformLightMapPolicy>::GetShaderBindings
 			// set reflection parameters
 			FTexture* ReflectionCubemapTextures[MaxNumReflections] = { GBlackTextureCube, GBlackTextureCube, GBlackTextureCube };
 			FVector4 CapturePositions[MaxNumReflections] = { FVector4(0, 0, 0, 0), FVector4(0, 0, 0, 0), FVector4(0, 0, 0, 0) };
-			FVector4 ReflectionParams(1.0f, 1.0f, 1.0f, 0.0f);
+			FVector4 ReflectionParams(0.0f, 0.0f, 0.0f, 0.0f);
 			FPrimitiveSceneInfo* PrimitiveSceneInfo = PrimitiveSceneProxy ? PrimitiveSceneProxy->GetPrimitiveSceneInfo() : nullptr;
 			if (PrimitiveSceneInfo)
 			{
@@ -478,7 +478,8 @@ void TMobileBasePassPSPolicyParamType<FUniformLightMapPolicy>::GetShaderBindings
 						{
 							ReflectionCubemapTextures[i] = PrimitiveSceneInfo->CachedReflectionCaptureProxies[i]->EncodedHDRCubemap;
 						}
-						ReflectionParams[i] = FMath::Max(FMath::Min(1.0f / ReflectionProxy->EncodedHDRAverageBrightness, 65504.f), -65504.f);
+						//To keep ImageBasedReflectionLighting coherence with PC, use AverageBrightness instead of InvAverageBrightness to calculate the IBL contribution
+						ReflectionParams[i] = ReflectionProxy->EncodedHDRAverageBrightness;
 					}
 				}
 			}
@@ -617,7 +618,8 @@ void FMobileBasePassMeshProcessor::Process(
 		TMobileBasePassPSPolicyParamType<FUniformLightMapPolicy>> BasePassShaders;
 	
 	//The stationary skylight contribution has been added to the LowQuality Lightmap for StaticMeshActor on mobile, so we should skip the sky light spherical harmonic contribution for it.
-	bool bEnableSkyLight = ShadingModels.IsLit() && Scene && Scene->ShouldRenderSkylightInBasePass(BlendMode) && !(FReadOnlyCVARCache::Get().bAllowStaticLighting && PrimitiveSceneProxy && PrimitiveSceneProxy->IsStatic());
+	//Enable skylight if LowQualityLightmaps is disabled or the Lightmap has not been built.
+	bool bEnableSkyLight = ShadingModels.IsLit() && Scene && Scene->ShouldRenderSkylightInBasePass(BlendMode) && !(FReadOnlyCVARCache::Get().bAllowStaticLighting && FReadOnlyCVARCache::Get().bEnableLowQualityLightmaps && PrimitiveSceneProxy && PrimitiveSceneProxy->IsStatic() && MeshBatch.LCI && MeshBatch.LCI->GetLightMapInteraction(FeatureLevel).GetType() == LMIT_Texture);
 	int32 NumMovablePointLights = MobileBasePass::CalcNumMovablePointLights(MaterialResource, PrimitiveSceneProxy);
 
 	MobileBasePass::GetShaders(
diff --git a/Engine/Source/Runtime/Renderer/Private/MobileBasePassRendering.cpp b/Engine/Source/Runtime/Renderer/Private/MobileBasePassRendering.cpp
--- a/Engine/Source/Runtime/Renderer/Private/MobileBasePassRendering.cpp
+++ b/Engine/Source/Runtime/Renderer/Private/MobileBasePassRendering.cpp
@@ -228,7 +228,7 @@ void SetupMobileDirectionalLightUniformParameters(
 
 void SetupMobileSkyReflectionUniformParameters(FSkyLightSceneProxy* SkyLight, FMobileReflectionCaptureShaderParameters& Parameters)
 {
-	float InvBrightness = 1.f;
+	float Brightness = 0.f;
 	float SkyMaxMipIndex = 0.f;
 	FTexture* CaptureTexture = GBlackTextureCube;
 
@@ -237,10 +237,11 @@ void SetupMobileSkyReflectionUniformParameters(FSkyLightSceneProxy* SkyLight, FM
 		check(SkyLight->ProcessedTexture->IsInitialized());
 		CaptureTexture = SkyLight->ProcessedTexture;
 		SkyMaxMipIndex = FMath::Log2(CaptureTexture->GetSizeX());
-		InvBrightness = FMath::Max(FMath::Min(1.0f / SkyLight->AverageBrightness, 65504.f), -65504.f);
+		Brightness = SkyLight->AverageBrightness;
 	}
 	
-	Parameters.Params = FVector4(InvBrightness, SkyMaxMipIndex, 0.f, 0.f);
+	//To keep ImageBasedReflectionLighting coherence with PC, use AverageBrightness instead of InvAverageBrightness to calculate the IBL contribution
+	Parameters.Params = FVector4(Brightness, SkyMaxMipIndex, 0.f, 0.f);
 	Parameters.Texture = CaptureTexture->TextureRHI;
 	Parameters.TextureSampler = CaptureTexture->SamplerStateRHI;
 }
diff --git a/Engine/Source/Runtime/Renderer/Private/MobileBasePassRendering.h b/Engine/Source/Runtime/Renderer/Private/MobileBasePassRendering.h
--- a/Engine/Source/Runtime/Renderer/Private/MobileBasePassRendering.h
+++ b/Engine/Source/Runtime/Renderer/Private/MobileBasePassRendering.h
@@ -245,7 +245,7 @@ public:
 		HQReflectionSamplers[1].Bind(Initializer.ParameterMap, TEXT("ReflectionCubemapSampler1"));
 		HQReflectionCubemaps[2].Bind(Initializer.ParameterMap, TEXT("ReflectionCubemap2"));
 		HQReflectionSamplers[2].Bind(Initializer.ParameterMap, TEXT("ReflectionCubemapSampler2"));
-		HQReflectionInvAverageBrigtnessParams.Bind(Initializer.ParameterMap, TEXT("ReflectionInvAverageBrigtness"));
+		HQReflectionInvAverageBrigtnessParams.Bind(Initializer.ParameterMap, TEXT("ReflectionAverageBrigtness"));
 		HQReflectionPositionsAndRadii.Bind(Initializer.ParameterMap, TEXT("ReflectionPositionsAndRadii"));
 
 		LightPositionAndInvRadiusParameter.Bind(Initializer.ParameterMap, TEXT("LightPositionAndInvRadius"));
diff --git a/Engine/Source/Runtime/Renderer/Private/MobileShadingRenderer.cpp b/Engine/Source/Runtime/Renderer/Private/MobileShadingRenderer.cpp
--- a/Engine/Source/Runtime/Renderer/Private/MobileShadingRenderer.cpp
+++ b/Engine/Source/Runtime/Renderer/Private/MobileShadingRenderer.cpp
@@ -1011,7 +1011,9 @@ void FMobileSceneRenderer::UpdateSkyReflectionUniformBuffer()
 	FSkyLightSceneProxy* SkyLight = nullptr;
 	if (Scene->ReflectionSceneData.RegisteredReflectionCapturePositions.Num() == 0
 		&& Scene->SkyLight
-		&& Scene->SkyLight->ProcessedTexture->TextureRHI)
+		&& Scene->SkyLight->ProcessedTexture->TextureRHI
+		// Don't use skylight reflection if it is a static sky light for keeping coherence with PC.
+		&& !Scene->SkyLight->bHasStaticLighting)
 	{
 		SkyLight = Scene->SkyLight;
 	}
diff --git a/Engine/Source/Runtime/Renderer/Private/RendererScene.cpp b/Engine/Source/Runtime/Renderer/Private/RendererScene.cpp
--- a/Engine/Source/Runtime/Renderer/Private/RendererScene.cpp
+++ b/Engine/Source/Runtime/Renderer/Private/RendererScene.cpp
@@ -1534,7 +1534,10 @@ void FScene::AddLightSceneInfo_RenderThread(FLightSceneInfo* LightSceneInfo)
 
 	if (bDirectionalLight &&
 		// Only use a stationary or movable light
-		!LightSceneInfo->Proxy->HasStaticLighting())
+		!(LightSceneInfo->Proxy->HasStaticLighting() 
+		// if it is a Static DirectionalLight and the light has not been built, add it to MobileDirectionalLights for mobile preview.
+		&& LightSceneInfo->IsPrecomputedLightingValid())
+		)
 	{
 		// Set SimpleDirectionalLight
 		if(!SimpleDirectionalLight)
