Fix a precision issue at mobile base pass pixel shader.

diff --git a/Engine/Shaders/Private/MobileBasePassPixelShader.usf b/Engine/Shaders/Private/MobileBasePassPixelShader.usf
--- a/Engine/Shaders/Private/MobileBasePassPixelShader.usf
+++ b/Engine/Shaders/Private/MobileBasePassPixelShader.usf
@@ -235,7 +235,7 @@ half3 MobileComputeMixingWeight3(half IndirectIrradiance, half3 AverageBrightnes
 	return lerp(half3(1.0f, 1.0f, 1.0f), MixingWeight, MixingAlpha);
 }
 
-half3 BlendReflectionCaptures(FMaterialPixelParameters MaterialParameters, half Roughness, float IndirectIrradiance)
+half3 BlendReflectionCaptures(FMaterialPixelParameters MaterialParameters, half Roughness, half IndirectIrradiance)
 {
 	// Compute fractional mip from roughness
 	half AbsoluteSpecularMip = ComputeReflectionCaptureMipFromRoughness(Roughness, ResolvedView.ReflectionCubemapMaxMip);
