Fixed HLSL warning and potential uv issue.

diff --git a/Engine/Shaders/Private/HairStrands/HairStrandsVertexFactory.ush b/Engine/Shaders/Private/HairStrands/HairStrandsVertexFactory.ush
--- a/Engine/Shaders/Private/HairStrands/HairStrandsVertexFactory.ush
+++ b/Engine/Shaders/Private/HairStrands/HairStrandsVertexFactory.ush
@@ -261,7 +261,7 @@ float3 GetVertexPreviousPosition(FVertexFactoryInput Input)
 }
 
 // Segment UV coord of an hair segment. This is different from the UV coord of the hair strands
-float GetSegmentUVCoord(FVertexFactoryInput Input)
+float2 GetSegmentUVCoord(FVertexFactoryInput Input)
 {
 	const uint QuadIndex = Input.VertexId % 6;
 	const float VCoord = QuadIndex == 0 || QuadIndex == 1 || QuadIndex == 5 ? 0.0f : 1.0f;
