Reorganize and isolate non-dynamic interpolation code.

diff --git a/Engine/Shaders/Private/HairStrands/HairStrandsInterpolation.usf b/Engine/Shaders/Private/HairStrands/HairStrandsInterpolation.usf
--- a/Engine/Shaders/Private/HairStrands/HairStrandsInterpolation.usf
+++ b/Engine/Shaders/Private/HairStrands/HairStrandsInterpolation.usf
@@ -78,64 +78,61 @@ void MainCS(uint2 DispatchThreadId : SV_DispatchThreadID)
 		float3 GuideVertexWeights = float3(saturate((InterpolationData0.w & 0xFF) / 255.f), saturate(((InterpolationData0.w >> 8) & 0xFF) / 255.f), 0);
 		GuideVertexWeights.z = saturate(1.0f - (GuideVertexWeights.x + GuideVertexWeights.y));
 
+		// 1. Rigid transformation
 		// ControlPoint is in the local hair referential
-		float3 CurrOffset = 0;
-		float3 ControlPoint = f16tof32(RenderRestPosePositionBuffer[VertexIndex].xyz) + InHairPositionOffset;
-		[unroll]
-		for (uint K = 0; K < 3; ++K)
+		// CurrOffset takes only translation component into account, and is done in object local space (vs. triangle/hair local space)
+		#if PERMUTATION_DYNAMIC_GEOMETRY == 0
 		{
-			const uint GuideIndex = GuideVertexIndices[K];
-			const float3 RestGuidePoint	= f16tof32(SimRestPosePositionBuffer[GuideIndex].xyz) + InHairPositionOffset;
+			float3 CurrOffset = 0;
+			float3 ControlPoint = f16tof32(RenderRestPosePositionBuffer[VertexIndex].xyz) + InHairPositionOffset;
+			[unroll]
+			for (uint K = 0; K < 3; ++K)
+			{
+				const uint GuideIndex = GuideVertexIndices[K];
+				const float3 RestGuidePoint	= f16tof32(SimRestPosePositionBuffer[GuideIndex].xyz) + InHairPositionOffset;
 
-			const float3 DeformedGuidePoint	= f16tof32(DeformedSimPositionBuffer[GuideIndex].xyz) + OutHairPositionOffset;
-			CurrOffset += (DeformedGuidePoint - RestGuidePoint) * GuideVertexWeights[K];
+				const float3 DeformedGuidePoint	= f16tof32(DeformedSimPositionBuffer[GuideIndex].xyz) + OutHairPositionOffset;
+				CurrOffset += (DeformedGuidePoint - RestGuidePoint) * GuideVertexWeights[K];
+			}
+			OutRenderDeformedPositionBuffer[VertexIndex] = float4(f32tof16(ControlPoint + CurrOffset - OutHairPositionOffset), RenderRestPosePositionBuffer[VertexIndex].w);
 		}
-		OutRenderDeformedPositionBuffer[VertexIndex] = float4(f32tof16(ControlPoint + CurrOffset - OutHairPositionOffset), RenderRestPosePositionBuffer[VertexIndex].w);
+		#endif	
 
+		// 2. Skin transformation
 		// Apply dynamic mesh deformation (translation / rotation)
-		// hair_todo: applied the transformation in local space and then rotation everything in global space
 		#if PERMUTATION_DYNAMIC_GEOMETRY == 1
 		{
+			float3 ControlPoint = f16tof32(RenderRestPosePositionBuffer[VertexIndex].xyz) + InHairPositionOffset;
+			float3 CurrOffset = 0;
+
+			// Compute the simulation offset in hair local space (i.e., triangle)
 			#if PERMUTATION_LOCAL_INTERPOLATION == 1
+			for (uint K = 0; K < 3; ++K)
 			{
-				CurrOffset = 0;
-				for (uint K = 0; K < 3; ++K)
-				{
-					const uint GuideIndex = GuideVertexIndices[K];
-					const uint RootIndex = SimVertexToRootIndexBuffer[GuideIndex];
-					const float3 RootBarycentric = DecodeBarycentrics(SimRootBarycentricBuffer[RootIndex]);
-
-					const FHairMeshTriangle RestTriangle		= GetTriangleTransformation(RootIndex, SimRestPosition0Buffer, SimRestPosition1Buffer, SimRestPosition2Buffer, SimRestPositionOffset);
-					const FHairMeshTriangle DeformedTriangle	= GetTriangleTransformation(RootIndex, SimDeformedPosition0Buffer, SimDeformedPosition1Buffer, SimDeformedPosition2Buffer, SimDeformedPositionOffset);
-
-					const float3 RestGuidePoint	= f16tof32(SimRestPosePositionBuffer[GuideIndex].xyz) + InHairPositionOffset;
-					const float3 LocalRestGuidePoint = ToTriangle(RestGuidePoint,RootBarycentric,RestTriangle);
-
-					const float3 DeformedGuidePoint	= f16tof32(DeformedSimPositionBuffer[GuideIndex].xyz) + OutHairPositionOffset;
-					const float3 LocalDeformedGuidePoint = ToTriangle(DeformedGuidePoint,RootBarycentric,DeformedTriangle);
-					CurrOffset += (LocalDeformedGuidePoint - LocalRestGuidePoint) * GuideVertexWeights[K];
-				}
-			
-				const uint RootIndex						= VertexToRootIndexBuffer[VertexIndex];
-				const float3 RootBarycentric				= DecodeBarycentrics(RootBarycentricBuffer[RootIndex]);
-				const FHairMeshTriangle RestTriangle		= GetTriangleTransformation(RootIndex, RestPosition0Buffer, RestPosition1Buffer, RestPosition2Buffer, RestPositionOffset);
-				const FHairMeshTriangle DeformedTriangle	= GetTriangleTransformation(RootIndex, DeformedPosition0Buffer, DeformedPosition1Buffer, DeformedPosition2Buffer, DeformedPositionOffset);
-
-				ControlPoint = ToTriangle(ControlPoint,RootBarycentric,RestTriangle) + CurrOffset;
-				ControlPoint = ToWorld(ControlPoint,RootBarycentric,DeformedTriangle);
-			}
-			#endif	
+				const uint GuideIndex = GuideVertexIndices[K];
+				const uint RootIndex = SimVertexToRootIndexBuffer[GuideIndex];
+				const float3 RootBarycentric = DecodeBarycentrics(SimRootBarycentricBuffer[RootIndex]);
 
-			#if PERMUTATION_LOCAL_INTERPOLATION == 0
-			{
-				const uint RootIndex						= VertexToRootIndexBuffer[VertexIndex];
-				const float3 RootBarycentric				= DecodeBarycentrics(RootBarycentricBuffer[RootIndex]);
-				const FHairMeshTriangle RestTriangle		= GetTriangleTransformation(RootIndex, RestPosition0Buffer, RestPosition1Buffer, RestPosition2Buffer, RestPositionOffset);
-				const FHairMeshTriangle DeformedTriangle	= GetTriangleTransformation(RootIndex, DeformedPosition0Buffer, DeformedPosition1Buffer, DeformedPosition2Buffer, DeformedPositionOffset);
+				const FHairMeshTriangle RestTriangle		= GetTriangleTransformation(RootIndex, SimRestPosition0Buffer, SimRestPosition1Buffer, SimRestPosition2Buffer, SimRestPositionOffset);
+				const FHairMeshTriangle DeformedTriangle	= GetTriangleTransformation(RootIndex, SimDeformedPosition0Buffer, SimDeformedPosition1Buffer, SimDeformedPosition2Buffer, SimDeformedPositionOffset);
+
+				const float3 RestGuidePoint	= f16tof32(SimRestPosePositionBuffer[GuideIndex].xyz) + InHairPositionOffset;
+				const float3 LocalRestGuidePoint = ToTriangle(RestGuidePoint,RootBarycentric,RestTriangle);
 
-				ControlPoint = TransformPoint(ControlPoint, RootBarycentric, RestTriangle, DeformedTriangle);
+				const float3 DeformedGuidePoint	= f16tof32(DeformedSimPositionBuffer[GuideIndex].xyz) + OutHairPositionOffset;
+				const float3 LocalDeformedGuidePoint = ToTriangle(DeformedGuidePoint,RootBarycentric,DeformedTriangle);
+				CurrOffset += (LocalDeformedGuidePoint - LocalRestGuidePoint) * GuideVertexWeights[K];
 			}
-			#endif
+			#endif		
+
+			// Transform hair from rest post to deformed pose, based triangle deformation + simulation offset (optional depending of the simulation is actual running or not
+			const uint RootIndex						= VertexToRootIndexBuffer[VertexIndex];
+			const float3 RootBarycentric				= DecodeBarycentrics(RootBarycentricBuffer[RootIndex]);
+			const FHairMeshTriangle RestTriangle		= GetTriangleTransformation(RootIndex, RestPosition0Buffer, RestPosition1Buffer, RestPosition2Buffer, RestPositionOffset);
+			const FHairMeshTriangle DeformedTriangle	= GetTriangleTransformation(RootIndex, DeformedPosition0Buffer, DeformedPosition1Buffer, DeformedPosition2Buffer, DeformedPositionOffset);
+
+			ControlPoint = ToTriangle(ControlPoint,RootBarycentric,RestTriangle) + CurrOffset;
+			ControlPoint = ToWorld(ControlPoint,RootBarycentric,DeformedTriangle);
 
 			OutRenderDeformedPositionBuffer[VertexIndex] = float4(f32tof16(ControlPoint - OutHairPositionOffset), RenderRestPosePositionBuffer[VertexIndex].w);
 
