Fix wrong position offset for rest sim strands.

diff --git a/Engine/Plugins/Experimental/HairStrands/Source/HairStrandsCore/Private/GroomComponent.cpp b/Engine/Plugins/Experimental/HairStrands/Source/HairStrandsCore/Private/GroomComponent.cpp
--- a/Engine/Plugins/Experimental/HairStrands/Source/HairStrandsCore/Private/GroomComponent.cpp
+++ b/Engine/Plugins/Experimental/HairStrands/Source/HairStrandsCore/Private/GroomComponent.cpp
@@ -802,12 +802,13 @@ void UGroomComponent::InitResources()
 		BeginInitResource(Res.RenderDeformedResources);
 		BeginInitResource(Res.SimDeformedResources);
 
-		const FVector RestHairPositionOffset = Res.RenderRestResources->PositionOffset;
+		const FVector RenderRestHairPositionOffset = Res.RenderRestResources->PositionOffset;
+		const FVector SimRestHairPositionOffset = Res.SimRestResources->PositionOffset;
 
-		Res.RenderDeformedResources->PositionOffset = RestHairPositionOffset;
-		Res.RenderRestResources->PositionOffset = RestHairPositionOffset;
-		Res.SimDeformedResources->PositionOffset = RestHairPositionOffset;
-		Res.SimRestResources->PositionOffset = RestHairPositionOffset;
+		Res.RenderDeformedResources->PositionOffset = RenderRestHairPositionOffset;
+		Res.RenderRestResources->PositionOffset = RenderRestHairPositionOffset;
+		Res.SimDeformedResources->PositionOffset = SimRestHairPositionOffset;
+		Res.SimRestResources->PositionOffset = SimRestHairPositionOffset;
 
 		FHairStrandsInterpolationOutput::HairGroup& InterpolationOutputGroup = InterpolationOutput->HairGroups.AddDefaulted_GetRef();
 		FHairStrandsInterpolationInput::FHairGroup& InterpolationInputGroup = InterpolationInput->HairGroups.AddDefaulted_GetRef();
@@ -816,10 +817,13 @@ void UGroomComponent::InitResources()
 		const FHairGroupDesc& InGroupDesc = GroomGroupsDesc[GroupIt];
 		InterpolationInputGroup.HairRadius = GetGroupMaxHairRadius(InGroupDesc, GroupData);
 		InterpolationInputGroup.HairRaytracingRadiusScale = HairRaytracingRadiusScale;
-		InterpolationInputGroup.InHairPositionOffset = RestHairPositionOffset;
+		InterpolationInputGroup.InRenderHairPositionOffset = RenderRestHairPositionOffset;
+		InterpolationInputGroup.InSimHairPositionOffset = SimRestHairPositionOffset;
+
 		// For skinned groom, these value will be updated during TickComponent() call
-		InterpolationInputGroup.OutHairPositionOffset = RestHairPositionOffset;
-		InterpolationInputGroup.OutHairPreviousPositionOffset = RestHairPositionOffset;
+		// Deformed sim & render are expressed within the referencial (unlike rest pose)
+		InterpolationInputGroup.OutHairPositionOffset = RenderRestHairPositionOffset;
+		InterpolationInputGroup.OutHairPreviousPositionOffset = RenderRestHairPositionOffset;
 		InterpolationInputGroup.bIsSimulationEnable = bIsSimulationEnable;
 
 		GroupIt++;
diff --git a/Engine/Plugins/Experimental/HairStrands/Source/HairStrandsCore/Private/HairStrandsRendering.cpp b/Engine/Plugins/Experimental/HairStrands/Source/HairStrandsCore/Private/HairStrandsRendering.cpp
--- a/Engine/Plugins/Experimental/HairStrands/Source/HairStrandsCore/Private/HairStrandsRendering.cpp
+++ b/Engine/Plugins/Experimental/HairStrands/Source/HairStrandsCore/Private/HairStrandsRendering.cpp
@@ -117,7 +117,8 @@ class FHairInterpolationCS : public FGlobalShader
 	BEGIN_SHADER_PARAMETER_STRUCT(FParameters, )
 		SHADER_PARAMETER(uint32, VertexCount)
 		SHADER_PARAMETER(uint32, DispatchCountX)
-		SHADER_PARAMETER(FVector, InHairPositionOffset)
+		SHADER_PARAMETER(FVector, InRenderHairPositionOffset)
+		SHADER_PARAMETER(FVector, InSimHairPositionOffset)
 		SHADER_PARAMETER(FVector, OutHairPositionOffset)
 
 		SHADER_PARAMETER(FVector, RestPositionOffset)
@@ -173,7 +174,8 @@ static void AddHairStrandsInterpolationPass(
 	FRDGBuilder& GraphBuilder,
 	const FHairStrandsProjectionHairData::HairGroup& InRenHairData,
 	const FHairStrandsProjectionHairData::HairGroup& InSimHairData,
-	const FVector& InHairWorldOffset, 
+	const FVector& InRenderHairWorldOffset, 
+	const FVector& InSimHairWorldOffset,
 	const FVector& OutHairWorldOffset,
 	const int32 LODIndex,
 	const bool bHasSimulationEnable,
@@ -204,7 +206,8 @@ static void AddHairStrandsInterpolationPass(
 		Parameters->OutRenderAttributeBuffer = OutRenderAttributeBuffer;
 	}
 	Parameters->VertexCount = VertexCount;
-	Parameters->InHairPositionOffset = InHairWorldOffset;
+	Parameters->InRenderHairPositionOffset = InRenderHairWorldOffset;
+	Parameters->InSimHairPositionOffset = InSimHairWorldOffset;
 	Parameters->OutHairPositionOffset = OutHairWorldOffset;
 	Parameters->DispatchCountX = DispatchCount.X;
 
@@ -493,7 +496,8 @@ void ComputeHairStrandsInterpolation(
 				GraphBuilder,
 				InRenHairDatas.HairGroups[GroupIndex],
 				InSimHairDatas.HairGroups[GroupIndex],
-				Input.InHairPositionOffset,
+				Input.InRenderHairPositionOffset,
+				Input.InSimHairPositionOffset,
 				Input.OutHairPositionOffset,
 				LODIndex,
 				bHasSimulationEnabled,
diff --git a/Engine/Plugins/Experimental/HairStrands/Source/HairStrandsCore/Public/HairStrandsRendering.h b/Engine/Plugins/Experimental/HairStrands/Source/HairStrandsCore/Public/HairStrandsRendering.h
--- a/Engine/Plugins/Experimental/HairStrands/Source/HairStrandsCore/Public/HairStrandsRendering.h
+++ b/Engine/Plugins/Experimental/HairStrands/Source/HairStrandsCore/Public/HairStrandsRendering.h
@@ -35,7 +35,8 @@ struct FHairStrandsInterpolationInput
 		bool bIsSimulationEnable = false;
 		float HairRadius = 0;
 		float HairRaytracingRadiusScale = 0;
-		FVector InHairPositionOffset = FVector::ZeroVector;
+		FVector InRenderHairPositionOffset = FVector::ZeroVector;
+		FVector InSimHairPositionOffset = FVector::ZeroVector;
 		FVector OutHairPositionOffset = FVector::ZeroVector;
 		FVector OutHairPreviousPositionOffset = FVector::ZeroVector;
 
diff --git a/Engine/Shaders/Private/HairStrands/HairStrandsInterpolation.usf b/Engine/Shaders/Private/HairStrands/HairStrandsInterpolation.usf
--- a/Engine/Shaders/Private/HairStrands/HairStrandsInterpolation.usf
+++ b/Engine/Shaders/Private/HairStrands/HairStrandsInterpolation.usf
@@ -5,7 +5,8 @@
 
 uint VertexCount;
 uint DispatchCountX;
-float3 InHairPositionOffset;
+float3 InRenderHairPositionOffset;
+float3 InSimHairPositionOffset;
 float3 OutHairPositionOffset;
 
 Buffer<uint4>  RenderRestPosePositionBuffer;
@@ -89,13 +90,13 @@ void MainCS(uint2 DispatchThreadId : SV_DispatchThreadID)
 		#if PERMUTATION_DYNAMIC_GEOMETRY == 0
 		{
 			float3 CurrOffset = 0;
-			float3 ControlPoint = f16tof32(RenderRestPosePositionBuffer[VertexIndex].xyz) + InHairPositionOffset;
+			float3 ControlPoint = f16tof32(RenderRestPosePositionBuffer[VertexIndex].xyz) + InRenderHairPositionOffset;
 			#if PERMUTATION_SIMULATION == 1
 			[unroll]
 			for (uint K = 0; K < 3; ++K)
 			{
 				const uint GuideIndex = GuideVertexIndices[K];
-				const float3 RestGuidePoint	= f16tof32(SimRestPosePositionBuffer[GuideIndex].xyz) + InHairPositionOffset;
+				const float3 RestGuidePoint	= f16tof32(SimRestPosePositionBuffer[GuideIndex].xyz) + InSimHairPositionOffset;
 
 				const float3 DeformedGuidePoint	= f16tof32(DeformedSimPositionBuffer[GuideIndex].xyz) + OutHairPositionOffset;
 				CurrOffset += (DeformedGuidePoint - RestGuidePoint) * GuideVertexWeights[K];
@@ -109,7 +110,7 @@ void MainCS(uint2 DispatchThreadId : SV_DispatchThreadID)
 		// Apply dynamic mesh deformation (translation / rotation)
 		#if PERMUTATION_DYNAMIC_GEOMETRY == 1
 		{
-			float3 ControlPoint = f16tof32(RenderRestPosePositionBuffer[VertexIndex].xyz) + InHairPositionOffset;
+			float3 ControlPoint = f16tof32(RenderRestPosePositionBuffer[VertexIndex].xyz) + InRenderHairPositionOffset;
 			float3 CurrOffset = 0;
 
 			// Compute the simulation offset in hair local space (i.e., triangle)
@@ -123,7 +124,7 @@ void MainCS(uint2 DispatchThreadId : SV_DispatchThreadID)
 				const FHairMeshTriangle RestTriangle		= GetTriangleTransformation(RootIndex, SimRestPosition0Buffer, SimRestPosition1Buffer, SimRestPosition2Buffer, SimRestPositionOffset);
 				const FHairMeshTriangle DeformedTriangle	= GetTriangleTransformation(RootIndex, SimDeformedPosition0Buffer, SimDeformedPosition1Buffer, SimDeformedPosition2Buffer, SimDeformedPositionOffset);
 
-				const float3 RestGuidePoint	= f16tof32(SimRestPosePositionBuffer[GuideIndex].xyz) + InHairPositionOffset;
+				const float3 RestGuidePoint	= f16tof32(SimRestPosePositionBuffer[GuideIndex].xyz) + InSimHairPositionOffset;
 				const float3 LocalRestGuidePoint = ToTriangle(RestGuidePoint,RootBarycentric,RestTriangle);
 
 				const float3 DeformedGuidePoint	= f16tof32(DeformedSimPositionBuffer[GuideIndex].xyz) + OutHairPositionOffset;
