Cherry-pick 9845073: Fixes mistake in the projection of the GGX lobe into screen space for the reflection denoiser.

diff --git a/Engine/Shaders/Private/ScreenSpaceDenoise/SSDSignalFramework.ush b/Engine/Shaders/Private/ScreenSpaceDenoise/SSDSignalFramework.ush
--- a/Engine/Shaders/Private/ScreenSpaceDenoise/SSDSignalFramework.ush
+++ b/Engine/Shaders/Private/ScreenSpaceDenoise/SSDSignalFramework.ush
@@ -365,6 +365,35 @@ float ComputeBluringOutOfFocusOfNormalizedSample(FSSDSignalSample Sample, FSSDSa
 	return (1 - DistanceCameraToPixel / (DistanceCameraToPixel + Sample.ClosestHitDistance));
 }
 
+
+// TODO(Denoiser): Update ImportanceSampleVisibleGGX() instead?
+float3 StableImportanceSampleVisibleGGX( float2 DiskE, float a2, float3 V )
+{
+	// TODO float2 alpha for anisotropic
+	float a = sqrt(a2);
+
+	// stretch
+	float3 Vh = normalize( float3( a * V.xy, V.z ) );
+
+	// Orthonormal basis
+	// Tangent0 is orthogonal to N.
+	float3 Tangent0 = (V.z < 0.9999) ? normalize( cross( float3(0, 0, 1), V ) ) : float3(1, 0, 0);
+	float3 Tangent1 = normalize(cross( Vh, Tangent0 ));
+
+	float2 p = DiskE;
+	float s = 0.5 + 0.5 * Vh.z;
+	p.y = (1 - s) * sqrt( 1 - p.x * p.x ) + s * p.y;
+
+	float3 H;
+	H  = p.x * Tangent0;
+	H += p.y * Tangent1;
+	H += sqrt( saturate( 1 - dot( p, p ) ) ) * Vh;
+
+	// unstretch
+	return normalize( float3( a * H.xy, max(0.0, H.z) ) );
+}
+
+
 /** Project the specular lobe in screen space. Returns the major axis of the lobe anysotropy, and the inifinity bluring in ViewportUV.x space. */
 void ProjectSpecularLobeToScreenSpace(
 	FSSDSampleSceneInfos SceneMetadata,
@@ -407,7 +436,7 @@ void ProjectSpecularLobeToScreenSpace(
 		for (uint i = 0; i < 4; i++)
 		{
 			// TODO(Denoiser): optimise VALU here.
-			float3 TangentH = ImportanceSampleVisibleGGX(ArrayE[i], a2, TangentV).xyz;
+			float3 TangentH = StableImportanceSampleVisibleGGX(ArrayE[i], a2, TangentV).xyz;
 			TangentL[i] = 2 * dot(TangentV, TangentH) * TangentH - TangentV;
 		}
 	
diff --git a/Engine/Shaders/Private/ScreenSpaceDenoise/SSDSpatialAccumulation.usf b/Engine/Shaders/Private/ScreenSpaceDenoise/SSDSpatialAccumulation.usf
--- a/Engine/Shaders/Private/ScreenSpaceDenoise/SSDSpatialAccumulation.usf
+++ b/Engine/Shaders/Private/ScreenSpaceDenoise/SSDSpatialAccumulation.usf
@@ -966,6 +966,12 @@ void MainCS(
 			float ConvolutionArea = max(KernelConfig.MajorPixelRadius, MinimalPixelRadius) * max(KernelConfig.MinorPixelRadius, MinimalPixelRadius);
 
 			KernelConfig.SampleCount = clamp(ConvolutionArea * TargetSamplePerPixel, 0, MaxSampleCount);
+			
+			// DebugOutput[DispatchThreadId] = float4(
+			// 	KernelConfig.SampleCount,
+			// 	KernelConfig.MajorPixelRadius,
+			// 	KernelConfig.MinorPixelRadius,
+			// 	0.0);
 		}
 		#endif // CONFIG_SAMPLE_SET == SAMPLE_SET_STOCASTIC_HIERARCHY
 		
