Cherry-pick 9845127: Replaces stackowiak kernel with GGX lobe projection in the denoiser's 1spp code path.

diff --git a/Engine/Shaders/Private/ScreenSpaceDenoise/SSDSpatialAccumulation.usf b/Engine/Shaders/Private/ScreenSpaceDenoise/SSDSpatialAccumulation.usf
--- a/Engine/Shaders/Private/ScreenSpaceDenoise/SSDSpatialAccumulation.usf
+++ b/Engine/Shaders/Private/ScreenSpaceDenoise/SSDSpatialAccumulation.usf
@@ -163,10 +163,14 @@
 			//#define CONFIG_BILATERAL_PRESET BILATERAL_PRESET_REFLECTIONS_1SPP
 			//#define CONFIG_BILATERAL_PRESET BILATERAL_PRESET_DISABLED
 		#else
-			// Use stackowiak sample set to keep noise in the signal to avoid flicker in history rejection,
-			// and adapt number of samples based on the sample.
-			#define CONFIG_SAMPLE_SET           SAMPLE_SET_STACKOWIAK_4_SETS
-			#define CONFIG_SAMPLE_COUNT_POLICY  SAMPLE_COUNT_POLICY_HIT_DISTANCE_BASED
+			#if 1
+				#define CONFIG_SAMPLE_SET           SAMPLE_SET_DIRECTIONAL_RECT
+			#else
+				// Use stackowiak sample set to keep noise in the signal to avoid flicker in history rejection,
+				// and adapt number of samples based on the sample.
+				#define CONFIG_SAMPLE_SET           SAMPLE_SET_STACKOWIAK_4_SETS
+				#define CONFIG_SAMPLE_COUNT_POLICY  SAMPLE_COUNT_POLICY_HIT_DISTANCE_BASED
+			#endif
 
 			// Reconstruction has 1 ray per input pixel, therefore can use specular ratio estimator instead of Tokoyashi axes bilateral kernel.
 			#define CONFIG_SPECULAR_RATIO_ESTIMATOR 1
@@ -937,41 +941,47 @@ void MainCS(
 		}
 		#elif CONFIG_SAMPLE_SET == SAMPLE_SET_DIRECTIONAL_RECT
 		{
-			const float TargetSamplePerPixel = 0.25;
-			const float MinimalPixelRadius = 0.5 * rsqrt(2.0);
-
-			// Project GGX lobe into screen space.
-			float2 NormalizedScreenMajorAxis;
-			float InifinityMajorViewportRadius;
-			float InifinityMinorViewportRadius;
-			ProjectSpecularLobeToScreenSpace(
-				RefSceneMetadata,
-				/* out */ NormalizedScreenMajorAxis,
-				/* out */ InifinityMajorViewportRadius,
-				/* out */ InifinityMinorViewportRadius);
+			#if CONFIG_SIGNAL_PROCESSING == SIGNAL_PROCESSING_REFLECTIONS
+			{
+				const float TargetSamplePerPixel = 0.25;
+				const float MinimalPixelRadius = 0.5 * rsqrt(2.0);
+
+				// Project GGX lobe into screen space.
+				float2 NormalizedScreenMajorAxis;
+				float InifinityMajorViewportRadius;
+				float InifinityMinorViewportRadius;
+				ProjectSpecularLobeToScreenSpace(
+					RefSceneMetadata,
+					/* out */ NormalizedScreenMajorAxis,
+					/* out */ InifinityMajorViewportRadius,
+					/* out */ InifinityMinorViewportRadius);
 
-			float OutOfFocus = ComputeBluringOutOfFocusOfNormalizedSample(NormalizedRefSamples.Array[0], RefSceneMetadata);
+				float OutOfFocus = ComputeBluringOutOfFocusOfNormalizedSample(NormalizedRefSamples.Array[0], RefSceneMetadata);
 
-			KernelConfig.MajorAxis = NormalizedScreenMajorAxis * float2(1, -1);
-			KernelConfig.MajorPixelRadius = InifinityMajorViewportRadius * OutOfFocus * View.ViewSizeAndInvSize.x;
+				KernelConfig.MajorAxis = NormalizedScreenMajorAxis * float2(1, -1);
+				KernelConfig.MajorPixelRadius = InifinityMajorViewportRadius * OutOfFocus * View.ViewSizeAndInvSize.x;
 
-			float AspectRatio = InifinityMinorViewportRadius / InifinityMajorViewportRadius;
+				float AspectRatio = InifinityMinorViewportRadius / InifinityMajorViewportRadius;
 
-			float MaxPixelRadius = sqrt(rcp(TargetSamplePerPixel) * MaxSampleCount / AspectRatio);
+				float MaxPixelRadius = sqrt(rcp(TargetSamplePerPixel) * MaxSampleCount / AspectRatio);
 
 
-			KernelConfig.MajorPixelRadius = clamp(KernelConfig.MajorPixelRadius, 0, MaxPixelRadius);
-			KernelConfig.MinorPixelRadius = AspectRatio * KernelConfig.MajorPixelRadius;
+				KernelConfig.MajorPixelRadius = clamp(KernelConfig.MajorPixelRadius, 0, MaxPixelRadius);
+				KernelConfig.MinorPixelRadius = AspectRatio * KernelConfig.MajorPixelRadius;
 
-			float ConvolutionArea = max(KernelConfig.MajorPixelRadius, MinimalPixelRadius) * max(KernelConfig.MinorPixelRadius, MinimalPixelRadius);
+				float ConvolutionArea = max(KernelConfig.MajorPixelRadius, MinimalPixelRadius) * max(KernelConfig.MinorPixelRadius, MinimalPixelRadius);
 
-			KernelConfig.SampleCount = clamp(ConvolutionArea * TargetSamplePerPixel, 0, MaxSampleCount);
+				KernelConfig.SampleCount = clamp(ConvolutionArea * TargetSamplePerPixel, 0, MaxSampleCount);
 			
-			// DebugOutput[DispatchThreadId] = float4(
-			// 	KernelConfig.SampleCount,
-			// 	KernelConfig.MajorPixelRadius,
-			// 	KernelConfig.MinorPixelRadius,
-			// 	0.0);
+				// DebugOutput[DispatchThreadId] = float4(
+				// 	KernelConfig.SampleCount,
+				// 	KernelConfig.MajorPixelRadius,
+				// 	KernelConfig.MinorPixelRadius,
+				// 	0.0);
+			}
+			#else
+				#error Directional rect sample set is not supported.
+			#endif
 		}
 		#endif // CONFIG_SAMPLE_SET == SAMPLE_SET_STOCASTIC_HIERARCHY
 		
