Fixes more overblurring issues in the multi spp reflection denoiser

diff --git a/Engine/Shaders/Private/ScreenSpaceDenoise/SSDSignalFramework.ush b/Engine/Shaders/Private/ScreenSpaceDenoise/SSDSignalFramework.ush
--- a/Engine/Shaders/Private/ScreenSpaceDenoise/SSDSignalFramework.ush
+++ b/Engine/Shaders/Private/ScreenSpaceDenoise/SSDSignalFramework.ush
@@ -362,7 +362,7 @@ float ComputeBluringOutOfFocusOfNormalizedSample(FSSDSignalSample Sample, FSSDSa
 {
 	float DistanceCameraToPixel = length(View.TranslatedWorldCameraOrigin - GetTranslatedWorldPosition(SceneMetadata));
 
-	return (1 - DistanceCameraToPixel / (DistanceCameraToPixel + Sample.ClosestHitDistance));
+	return Sample.ClosestHitDistance / (DistanceCameraToPixel + Sample.ClosestHitDistance);
 }
 
 
@@ -415,7 +415,7 @@ void ProjectSpecularLobeToScreenSpace(
 
 	{
 		// Measure the angle that only truncate 10% of D_GGX's tale.
-		const float E = (1.0 - GGX_IMPORTANT_SAMPLE_BIAS);
+		const float E = (1.0 - 0.5);
 		const float SqrtE = sqrt(E);
 
 		const float2 ArrayE[4] = {
diff --git a/Engine/Shaders/Private/ScreenSpaceDenoise/SSDSpatialAccumulation.usf b/Engine/Shaders/Private/ScreenSpaceDenoise/SSDSpatialAccumulation.usf
--- a/Engine/Shaders/Private/ScreenSpaceDenoise/SSDSpatialAccumulation.usf
+++ b/Engine/Shaders/Private/ScreenSpaceDenoise/SSDSpatialAccumulation.usf
@@ -963,21 +963,27 @@ void MainCS(
 
 				float AspectRatio = InifinityMinorViewportRadius / InifinityMajorViewportRadius;
 
-				float MaxPixelRadius = sqrt(rcp(TargetSamplePerPixel) * MaxSampleCount / AspectRatio);
-
+				float MaxPixelDiameter = sqrt(rcp(TargetSamplePerPixel) * MaxSampleCount / AspectRatio);
+				float MaxPixelRadius = 0.5 * MaxPixelDiameter;
 
 				KernelConfig.MajorPixelRadius = clamp(KernelConfig.MajorPixelRadius, 0, MaxPixelRadius);
 				KernelConfig.MinorPixelRadius = AspectRatio * KernelConfig.MajorPixelRadius;
 
-				float ConvolutionArea = max(KernelConfig.MajorPixelRadius, MinimalPixelRadius) * max(KernelConfig.MinorPixelRadius, MinimalPixelRadius);
+				// *4 to multiply from radii area to diameters area.
+				float ConvolutionArea = 4.0 * max(KernelConfig.MajorPixelRadius, MinimalPixelRadius) * max(KernelConfig.MinorPixelRadius, MinimalPixelRadius);
 
 				KernelConfig.SampleCount = clamp(ConvolutionArea * TargetSamplePerPixel, 0, MaxSampleCount);
 			
-				// DebugOutput[DispatchThreadId] = float4(
-				// 	KernelConfig.SampleCount,
-				// 	KernelConfig.MajorPixelRadius,
-				// 	KernelConfig.MinorPixelRadius,
-				// 	0.0);
+				//DebugOutput[DispatchThreadId] = float4(
+				//	NormalizedScreenMajorAxis,
+				//	KernelConfig.MinorPixelRadius,
+				//	OutOfFocus);
+
+				DebugOutput[DispatchThreadId] = float4(
+					KernelConfig.SampleCount,
+					KernelConfig.MajorPixelRadius,
+					KernelConfig.MinorPixelRadius,
+					OutOfFocus);
 			}
 			#else
 				#error Directional rect sample set is not supported.
diff --git a/Engine/Shaders/Private/ScreenSpaceDenoise/SSDSpatialKernel.ush b/Engine/Shaders/Private/ScreenSpaceDenoise/SSDSpatialKernel.ush
--- a/Engine/Shaders/Private/ScreenSpaceDenoise/SSDSpatialKernel.ush
+++ b/Engine/Shaders/Private/ScreenSpaceDenoise/SSDSpatialKernel.ush
@@ -2174,7 +2174,7 @@ void ConvolveDirectionalRect(
 				float2(KernelConfig.MajorAxis) * KernelConfig.MajorPixelRadius * (E.x * 2.0 - 1.0) +
 				float2(-KernelConfig.MajorAxis.y, KernelConfig.MajorAxis.x) * KernelConfig.MinorPixelRadius * (E.y * 2.0 - 1.0);
 			
-			SampleBufferUV[SampleBatchId] = KernelConfig.BufferUV + (SampleOffset * KernelConfig.KernelSpreadFactor) * KernelConfig.BufferSizeAndInvSize.zw;
+			SampleBufferUV[SampleBatchId] = KernelConfig.BufferUV + SampleOffset * KernelConfig.BufferSizeAndInvSize.zw;
 		}
 
 		SampleAndAccumulateMultiplexedSignalsPair(
