Fixes projection of the direction of the GGX lobe

diff --git a/Engine/Shaders/Private/ScreenSpaceDenoise/SSDDefinitions.ush b/Engine/Shaders/Private/ScreenSpaceDenoise/SSDDefinitions.ush
--- a/Engine/Shaders/Private/ScreenSpaceDenoise/SSDDefinitions.ush
+++ b/Engine/Shaders/Private/ScreenSpaceDenoise/SSDDefinitions.ush
@@ -223,6 +223,7 @@
 	#define SAMPLE_SET_STOCASTIC_HIERARCHY 12
 
 	#define SAMPLE_SET_DIRECTIONAL_RECT 14
+	#define SAMPLE_SET_DIRECTIONAL_ELLIPSE 15
 
 
 /** Compression of signal for VGPR occupency. */
diff --git a/Engine/Shaders/Private/ScreenSpaceDenoise/SSDSignalFramework.ush b/Engine/Shaders/Private/ScreenSpaceDenoise/SSDSignalFramework.ush
--- a/Engine/Shaders/Private/ScreenSpaceDenoise/SSDSignalFramework.ush
+++ b/Engine/Shaders/Private/ScreenSpaceDenoise/SSDSignalFramework.ush
@@ -408,9 +408,12 @@ void ProjectSpecularLobeToScreenSpace(
 	{
 		float3 R = 2 * dot(V, N) * N - V;
 
-		float3 ViewR = mul(float4(R, 0.0), View.TranslatedWorldToView).xyz;
+		float4 RayEndClip = mul(float4(GetTranslatedWorldPosition(SceneMetadata) + R * (0.5 * SceneMetadata.WorldDepth), 1), View.TranslatedWorldToClip);
+		float2 RayEndScreen = RayEndClip.xy * rcp(RayEndClip.w);
 
-		NormalizedScreenMajorAxis = normalize(ViewR.xy);
+		float2 RayDirectionScreen = RayEndScreen - SceneMetadata.ScreenPosition;
+
+		NormalizedScreenMajorAxis = normalize(RayDirectionScreen * View.ViewSizeAndInvSize.xy);
 	}
 
 	{
diff --git a/Engine/Shaders/Private/ScreenSpaceDenoise/SSDSpatialAccumulation.usf b/Engine/Shaders/Private/ScreenSpaceDenoise/SSDSpatialAccumulation.usf
--- a/Engine/Shaders/Private/ScreenSpaceDenoise/SSDSpatialAccumulation.usf
+++ b/Engine/Shaders/Private/ScreenSpaceDenoise/SSDSpatialAccumulation.usf
@@ -157,14 +157,14 @@
 
 		// Configure the bilateral therm.
 		#if DIM_MULTI_SPP
-			#define CONFIG_SAMPLE_SET           SAMPLE_SET_DIRECTIONAL_RECT
+			#define CONFIG_SAMPLE_SET           SAMPLE_SET_DIRECTIONAL_ELLIPSE
 
 			#define CONFIG_BILATERAL_PRESET BILATERAL_PRESET_REFLECTIONS
 			//#define CONFIG_BILATERAL_PRESET BILATERAL_PRESET_REFLECTIONS_1SPP
 			//#define CONFIG_BILATERAL_PRESET BILATERAL_PRESET_DISABLED
 		#else
 			#if 1
-				#define CONFIG_SAMPLE_SET           SAMPLE_SET_DIRECTIONAL_RECT
+				#define CONFIG_SAMPLE_SET           SAMPLE_SET_DIRECTIONAL_ELLIPSE
 			#else
 				// Use stackowiak sample set to keep noise in the signal to avoid flicker in history rejection,
 				// and adapt number of samples based on the sample.
@@ -453,11 +453,6 @@
 	#define CONFIG_MAX_WITH_REF_DISTANCE 0
 #endif
 
-/** Project GGX lobe to compute the kernel. */
-#ifndef CONFIG_PROJECT_GGX_LOBE
-	#define CONFIG_PROJECT_GGX_LOBE 0
-#endif
-
 
 //------------------------------------------------------- COMPILATION WHITE LISTING
 
@@ -466,7 +461,7 @@
 	#define COMPILE_STACKOWIAK_KERNEL 1
 #elif CONFIG_SAMPLE_SET == SAMPLE_SET_HEXAWEB
 	#define COMPILE_DISK_KERNEL 1
-#elif CONFIG_SAMPLE_SET == SAMPLE_SET_DIRECTIONAL_RECT
+#elif CONFIG_SAMPLE_SET == SAMPLE_SET_DIRECTIONAL_RECT || CONFIG_SAMPLE_SET == SAMPLE_SET_DIRECTIONAL_ELLIPSE
 	#define COMPILE_DIRECTIONAL_KERNEL 1
 #else
 	#define COMPILE_BOX_KERNEL 1
@@ -939,7 +934,7 @@ void MainCS(
 				KernelBufferUV = KernelConfig.BufferUV;
 			#endif
 		}
-		#elif CONFIG_SAMPLE_SET == SAMPLE_SET_DIRECTIONAL_RECT
+		#elif CONFIG_SAMPLE_SET == SAMPLE_SET_DIRECTIONAL_RECT || CONFIG_SAMPLE_SET == SAMPLE_SET_DIRECTIONAL_ELLIPSE
 		{
 			#if CONFIG_SIGNAL_PROCESSING == SIGNAL_PROCESSING_REFLECTIONS
 			{
@@ -979,11 +974,11 @@ void MainCS(
 				//	KernelConfig.MinorPixelRadius,
 				//	OutOfFocus);
 
-				DebugOutput[DispatchThreadId] = float4(
-					KernelConfig.SampleCount,
-					KernelConfig.MajorPixelRadius,
-					KernelConfig.MinorPixelRadius,
-					OutOfFocus);
+				// DebugOutput[DispatchThreadId] = float4(
+				// 	KernelConfig.SampleCount,
+				// 	KernelConfig.MajorPixelRadius,
+				// 	KernelConfig.MinorPixelRadius,
+				// 	OutOfFocus);
 			}
 			#else
 				#error Directional rect sample set is not supported.
diff --git a/Engine/Shaders/Private/ScreenSpaceDenoise/SSDSpatialKernel.ush b/Engine/Shaders/Private/ScreenSpaceDenoise/SSDSpatialKernel.ush
--- a/Engine/Shaders/Private/ScreenSpaceDenoise/SSDSpatialKernel.ush
+++ b/Engine/Shaders/Private/ScreenSpaceDenoise/SSDSpatialKernel.ush
@@ -3,6 +3,7 @@
 #include "SSDSignalAccumulatorArray.ush"
 #include "SSDSignalBufferEncoding.ush"
 #include "../TextureSampling.ush"
+#include "../MonteCarlo.ush"
 
 
 //------------------------------------------------------- ENUMS
@@ -2170,9 +2171,18 @@ void ConvolveDirectionalRect(
 
 			float2 E = Hammersley16(SampleId, BatchCountCount * kSamplingBatchSize, KernelConfig.HammersleySeed);
 
+			if (KernelConfig.SampleSet == SAMPLE_SET_DIRECTIONAL_RECT)
+			{
+				E = (E * 2.0 - 1.0);
+			}
+			else // if (KernelConfig.SampleSet == SAMPLE_SET_DIRECTIONAL_ELLIPSE)
+			{
+				E = UniformSampleDiskConcentric(E);
+			}
+
 			float2 SampleOffset = 
-				float2(KernelConfig.MajorAxis) * KernelConfig.MajorPixelRadius * (E.x * 2.0 - 1.0) +
-				float2(-KernelConfig.MajorAxis.y, KernelConfig.MajorAxis.x) * KernelConfig.MinorPixelRadius * (E.y * 2.0 - 1.0);
+				float2(KernelConfig.MajorAxis) * KernelConfig.MajorPixelRadius * E.x +
+				float2(-KernelConfig.MajorAxis.y, KernelConfig.MajorAxis.x) * KernelConfig.MinorPixelRadius * E.y;
 			
 			SampleBufferUV[SampleBatchId] = KernelConfig.BufferUV + SampleOffset * KernelConfig.BufferSizeAndInvSize.zw;
 		}
@@ -2361,7 +2371,7 @@ void AccumulateKernel(
 	}
 #endif // COMPILE_HEIRARCHY_KERNEL
 #if COMPILE_DIRECTIONAL_KERNEL
-	else if (KernelConfig.SampleSet == SAMPLE_SET_DIRECTIONAL_RECT)
+	else if (KernelConfig.SampleSet == SAMPLE_SET_DIRECTIONAL_RECT || KernelConfig.SampleSet == SAMPLE_SET_DIRECTIONAL_ELLIPSE)
 	{
 		ConvolveDirectionalRect(
 			KernelConfig,
