Add roughness remapping from Hair BSDF to GGX sampling.

This improves a little bit the precision on the TT term, but nothing really noticeable for R/TRT as the remapping is nearly linear for lower roughnesses. It provide better mapping for high roughnesses, but since the energy is relatively low/spread, it is not strongly noticeable.

diff --git a/Engine/Shaders/Private/HairStrands/HairStrandsEnvironmentLighting.usf b/Engine/Shaders/Private/HairStrands/HairStrandsEnvironmentLighting.usf
--- a/Engine/Shaders/Private/HairStrands/HairStrandsEnvironmentLighting.usf
+++ b/Engine/Shaders/Private/HairStrands/HairStrandsEnvironmentLighting.usf
@@ -17,6 +17,7 @@
 
 #include "../HairBsdf.ush"
 #include "HairStrandsVisibilityCommon.ush"
+#include "HairStrandsEnvironmentLightingCommon.ush"
 #include "HairStrandsDeepVoxelCommon.ush"
 #include "HairStrandsDeepShadowCommon.ush"
 #include "HairStrandsDeepTransmittanceCommon.ush"
@@ -203,8 +204,9 @@ float3 ReflectionEnvironment(FGBufferData GBuffer, float2 ScreenPosition, float2
 		}
 		#endif	
 
-		const float RoughnessBias = 0.1;
-		ColorMS.rgb += PreExposure * GatherRadiance(R, GBuffer.Roughness + RoughnessBias);
+		// Select the TT term for representing the MS roughness lookup as TT lobe carries the most energy
+		const float GGXRoughness = HairToGGXRoughness_TT(GBuffer.Roughness);
+		ColorMS.rgb += PreExposure * GatherRadiance(R, GGXRoughness);
 		//ColorMS.rgb += PreExposure * GetSkySHDiffuse(R) * View.SkyLightColor.rgb;
 		ColorMS.rgb *= AvgDualScattering;
 
@@ -246,9 +248,10 @@ float3 ReflectionEnvironment(FGBufferData GBuffer, float2 ScreenPosition, float2
 
 		}
 
-		// Top of regular reflection or bottom layer of clear coat.
-		const float RoughnessBias = 0.0;
-		ColorR.rgb += PreExposure * GatherRadiance(R, GBuffer.Roughness + RoughnessBias);
+		// Select the TRT for remapping the hair roughness onto the GGX based lobe, as it is the 
+		// widest lobe between the R & TRT components
+		const float GGXRoughness = HairToGGXRoughness_TRT(GBuffer.Roughness);
+		ColorR.rgb += PreExposure * GatherRadiance(R, GGXRoughness);
 		//ColorR.rgb += PreExposure * GetSkySHDiffuse(R) * View.SkyLightColor.rgb * PI;
 
 		const float3 L = R;
@@ -304,7 +307,8 @@ float3 ReflectionEnvironment(FGBufferData GBuffer, float2 ScreenPosition, float2
 		const float SinLightAngle = dot(L, T);
 		const float SinViewAngle = dot(V, T);
 
-		ColorTT.rgb += PreExposure * GatherRadiance(L, GBuffer.Roughness);
+		const float GGXRoughness = HairToGGXRoughness_TT(GBuffer.Roughness);
+		ColorTT.rgb += PreExposure * GatherRadiance(L, GGXRoughness);
 		//ColorTT.rgb += PreExposure * GetSkySHDiffuse(L) * View.SkyLightColor.rgb;
 
 		// Compute the transmittance
