Add support for local reflection volumes to hair environment lighting.

diff --git a/Engine/Shaders/Private/HairStrands/HairStrandsEnvironmentLighting.usf b/Engine/Shaders/Private/HairStrands/HairStrandsEnvironmentLighting.usf
--- a/Engine/Shaders/Private/HairStrands/HairStrandsEnvironmentLighting.usf
+++ b/Engine/Shaders/Private/HairStrands/HairStrandsEnvironmentLighting.usf
@@ -15,6 +15,18 @@
 #include "../LightGridCommon.ush"
 #include "../SceneTextureParameters.ush"
 
+// Disable static lighting, as hair don't have any diffused lighting baked, 
+// and rely solely on specular information (i.e. reflection capture)
+#undef ALLOW_STATIC_LIGHTING
+#define ALLOW_STATIC_LIGHTING 0
+
+#define ENABLE_SKY_LIGHT 1
+#define REFLECTION_COMPOSITE_USE_BLENDED_REFLECTION_CAPTURES 1
+#define REFLECTION_COMPOSITE_SUPPORT_SKYLIGHT_BLEND 0
+#define REFLECTION_COMPOSITE_HAS_BOX_CAPTURES 1
+#define REFLECTION_COMPOSITE_HAS_SPHERE_CAPTURES 1
+#include "../ReflectionEnvironmentComposite.ush"
+
 #include "../HairBsdf.ush"
 #include "HairStrandsVisibilityCommon.ush"
 #include "HairStrandsEnvironmentLightingCommon.ush"
@@ -26,19 +38,34 @@
 // Note: 
 // * This version of environement lighting is very much in WIP state, with a lot of option to tune tweak to explore 
 //   what kind of approximation works the best for certains type of grooms and lighting condition
-// * There is a few fudges factor as well (MS roughness bias, TT density for instance, reflected vector orientations, ...)
-// * The source of lighting is only sky lighting at the moment for simplicity. This needs to be revisited to support properly reflection probes and indirect lighting in various scenario (static lighting, stationnary lighting, dynamic lighting)
+// * There is a few fudges factor as well (roughness mapping, TT density, reflected vector orientations, ...)
 
-// This is copy/pasta from ReflectionEnvonmentComposite & co. For simplicity, while still in developement, support only sky lighting. No local reflection probe yet.
-// @hair_todo
-float3 GatherRadiance(float3 RayDirection, float Roughness)
+float3 GatherRadiance(float3 RayDirection, float Roughness, uint2 PixelPosition, float PixelSceneDepth, float3 AbsoluteWorldPosition)
 {
-	//float SkyAverageBrightness = 1.0f;
-	float SkyAverageBrightness = ReflectionStruct.SkyLightCubemapBrightness;
-	return GetSkyLightReflection(RayDirection, Roughness, SkyAverageBrightness);
+	const float IndirectIrradiance = 1;
+	const int SingleCaptureIndex = 0;
+
+	const uint EyeIndex = 0;
+	const uint GridIndex = ComputeLightGridCellIndex(PixelPosition, PixelSceneDepth, EyeIndex);
+
+	const uint NumCulledEntryIndex = (ForwardLightData.NumGridCells + GridIndex) * NUM_CULLED_LIGHTS_GRID_STRIDE;
+	const uint NumLocalReflectionCaptures = min(ForwardLightData.NumCulledLightsGrid[NumCulledEntryIndex + 0], ForwardLightData.NumReflectionCaptures);
+	const uint DataStartIndex = ForwardLightData.NumCulledLightsGrid[NumCulledEntryIndex + 1];
+
+	return CompositeReflectionCapturesAndSkylight(
+		1.0f,
+		AbsoluteWorldPosition,
+		RayDirection,
+		Roughness,
+		IndirectIrradiance,
+		1.0f,
+		0.0f,
+		NumLocalReflectionCaptures,
+		DataStartIndex,
+		SingleCaptureIndex,
+		EyeIndex);
 }
 
-
 Texture2D<uint4>			HairCategorizationTexture;
 Texture2D<uint>				HairVisibilityNodeOffsetAndCount;
 StructuredBuffer<FNodeData> HairVisibilityNodeData;
@@ -101,8 +128,10 @@ float3 ReflectionEnvironment(FGBufferData GBuffer, float2 ScreenPosition, float2
 	const float PreExposure = 1.f;
 #endif
 	
-	const float3 WorldPosition = mul(float4(ScreenPosition * GBuffer.Depth, GBuffer.Depth, 1), View.ScreenToWorld).xyz;
-	const float3 CameraToPixel = normalize(WorldPosition - View.WorldCameraOrigin);	
+	const float3 WorldPosition	= mul(float4(ScreenPosition * GBuffer.Depth, GBuffer.Depth, 1), View.ScreenToWorld).xyz;
+	const float3 CameraToPixel	= normalize(WorldPosition - View.WorldCameraOrigin);
+	const uint2  PixelPosition	= SvPosition - ResolvedView.ViewRectMin.xy;
+	const float  PixelSceneDepth= GBuffer.Depth;
 
 	const float3 VoxelSize = (Voxel_MaxAABB - Voxel_MinAABB) / Voxel_Resolution;
 	const float  VoxelMaxSize = max(VoxelSize.x, max(VoxelSize.y, VoxelSize.z));
@@ -220,7 +249,7 @@ float3 ReflectionEnvironment(FGBufferData GBuffer, float2 ScreenPosition, float2
 
 		// Select the TT term for representing the MS roughness lookup as TT lobe carries the most energy
 		const float GGXRoughness = HairToGGXRoughness_TT(GBuffer.Roughness);
-		ColorMS.rgb += PreExposure * GatherRadiance(R, GGXRoughness);
+		ColorMS.rgb += PreExposure * GatherRadiance(R, GGXRoughness, PixelPosition, PixelSceneDepth, WorldPosition);
 		//ColorMS.rgb += PreExposure * GetSkySHDiffuse(R) * View.SkyLightColor.rgb;
 		ColorMS.rgb *= AvgDualScattering;
 
@@ -266,7 +295,7 @@ float3 ReflectionEnvironment(FGBufferData GBuffer, float2 ScreenPosition, float2
 		// Select the TRT for remapping the hair roughness onto the GGX based lobe, as it is the 
 		// widest lobe between the R & TRT components
 		const float GGXRoughness = HairToGGXRoughness_TRT(GBuffer.Roughness);
-		ColorR.rgb += PreExposure * GatherRadiance(R, GGXRoughness);
+		ColorR.rgb += PreExposure * GatherRadiance(R, GGXRoughness, PixelPosition, PixelSceneDepth, WorldPosition);
 		//ColorR.rgb += PreExposure * GetSkySHDiffuse(R) * View.SkyLightColor.rgb * PI;
 
 		const float3 L = R;
@@ -324,7 +353,7 @@ float3 ReflectionEnvironment(FGBufferData GBuffer, float2 ScreenPosition, float2
 		const float SinViewAngle = dot(V, T);
 
 		const float GGXRoughness = HairToGGXRoughness_TT(GBuffer.Roughness);
-		ColorTT.rgb += PreExposure * GatherRadiance(L, GGXRoughness);
+		ColorTT.rgb += PreExposure * GatherRadiance(L, GGXRoughness, PixelPosition, PixelSceneDepth, WorldPosition);
 		//ColorTT.rgb += PreExposure * GetSkySHDiffuse(L) * View.SkyLightColor.rgb;
 
 		// Compute the transmittance
diff --git a/Engine/Source/Runtime/Renderer/Private/HairStrands/HairStrandsEnvironment.cpp b/Engine/Source/Runtime/Renderer/Private/HairStrands/HairStrandsEnvironment.cpp
--- a/Engine/Source/Runtime/Renderer/Private/HairStrands/HairStrandsEnvironment.cpp
+++ b/Engine/Source/Runtime/Renderer/Private/HairStrands/HairStrandsEnvironment.cpp
@@ -228,6 +228,11 @@ class FHairEnvironmentLightingCS : public FGlobalShader
 	class FGroupSize : SHADER_PERMUTATION_INT("PERMUTATION_GROUP_SIZE", 2);
 	using FPermutationDomain = TShaderPermutationDomain<FGroupSize>;
 
+	static void ModifyCompilationEnvironment(const FGlobalShaderPermutationParameters& Parameters, FShaderCompilerEnvironment& OutEnvironment)
+	{
+		FForwardLightingParameters::ModifyCompilationEnvironment(Parameters.Platform, OutEnvironment);
+	}
+
 	static bool ShouldCompilePermutation(const FGlobalShaderPermutationParameters& Parameters)
 	{
 		return IsHairStrandsSupported(Parameters.Platform);
@@ -358,6 +363,12 @@ class FHairEnvironmentLightingPS : public FGlobalShader
 
 	class FPerSample : SHADER_PERMUTATION_INT("PERMUTATION_PER_SAMPLE", 2);
 	using FPermutationDomain = TShaderPermutationDomain<FPerSample>;
+
+	static void ModifyCompilationEnvironment(const FGlobalShaderPermutationParameters& Parameters, FShaderCompilerEnvironment& OutEnvironment)
+	{
+		FForwardLightingParameters::ModifyCompilationEnvironment(Parameters.Platform, OutEnvironment);
+	}
+
 	static bool ShouldCompilePermutation(const FGlobalShaderPermutationParameters& Parameters)
 	{
 		return IsHairStrandsSupported(Parameters.Platform);
