Added support for translucent objects in reflections

(Integration of patrick.kelly work at CL 9344747)

diff --git a/Engine/Shaders/Private/RayTracing/RayTracingReflections.usf b/Engine/Shaders/Private/RayTracing/RayTracingReflections.usf
--- a/Engine/Shaders/Private/RayTracing/RayTracingReflections.usf
+++ b/Engine/Shaders/Private/RayTracing/RayTracingReflections.usf
@@ -88,6 +88,7 @@ uint2 TileAlignedResolution;
 uint TestPathRoughness;
 float MinClearCoatLevel;
 uint MaxUnderCoatBounces;
+uint EnableTranslucency;
 
 float ReflectionMinRayDistance;
 float ReflectionMaxRayDistance;
@@ -141,6 +142,55 @@ void FixSampleDirectionIfNeeded(float3 SmoothSurfaceNormal, inout float3 SampleD
 	}
 }
 
+// Volume rendering ala emission-absorption model along an arbitrary ray
+float Transmit(in RayDesc TransmissionRay, inout RandomSequence RandSequence, inout float3 OutRadiance)
+{
+	float Transmission = 1.0;
+
+	const uint RayFlags = 0;
+	const uint InstanceInclusionMask = RAY_TRACING_MASK_TRANSLUCENT;
+	FRayCone RayCone = (FRayCone)0;
+	FMaterialClosestHitPayload TransmPayload = TraceMaterialRay(
+		TLAS,
+		RayFlags,
+		InstanceInclusionMask,
+		TransmissionRay,
+		RayCone);
+
+	const uint HitCountMax = 32;
+	uint HitCount = 0;
+	bool bIsHit = TransmPayload.IsHit();
+	float Epsilon = 1.0e-2;
+	while (bIsHit && Transmission > Epsilon && HitCount < HitCountMax)
+	{
+		HitCount++;
+		float3 Radiance = 0.0;
+		float NoV = dot(TransmPayload.WorldNormal, -TransmissionRay.Direction);
+		if (NoV > 0.0)
+		{
+			//#dxr_todo: Pass in all the parameters to support in-scattering
+			Radiance += TransmPayload.Radiance;
+		}
+
+		OutRadiance += Radiance * TransmPayload.Opacity * Transmission;
+		Transmission *= 1.0 - TransmPayload.Opacity;
+
+		//#dxr_todo: consider modifying t-min instead Origin for efficiency
+		TransmissionRay.Origin += TransmissionRay.Direction * (TransmPayload.HitT + 0.1);
+		TransmPayload = TraceMaterialRay(
+			TLAS,
+			RayFlags,
+			RAY_TRACING_MASK_TRANSLUCENT,
+			TransmissionRay,
+			RayCone);
+
+		bIsHit = TransmPayload.IsHit();
+	}
+
+	return Transmission;
+}
+
+
 RAY_TRACING_ENTRY_RAYGEN(RayTracingReflectionsRGS)
 {
 	const float ClearCoatSpecularColor = 0.04; // See ReflectionEnvironmentPixelShader.usf
@@ -255,7 +305,7 @@ RAY_TRACING_ENTRY_RAYGEN(RayTracingReflectionsRGS)
 
 			FRayCone RayCone = (FRayCone)0;
 			RayCone.SpreadAngle = View.EyeToPixelSpreadAngle;
-			float SurfaceCurvature = 0.0f; /* #todo_dxr assume no curvature */
+			float SurfaceCurvature = 0.0f; /* #dxr_todo assume no curvature */
 			RayCone = PropagateRayCone(RayCone, SurfaceCurvature, Depth);
 
 			//
@@ -403,6 +453,17 @@ RAY_TRACING_ENTRY_RAYGEN(RayTracingReflectionsRGS)
 						TopLayerRay,
 						RayCone);
 
+					// Trace separate ray for translucent objects
+					float Transmission = 1.0;
+					//#dxr-dxr_todo: do permutation instead of dynamic branching?
+					if (EnableTranslucency)
+					{
+						RayDesc TransmissionRay = TopLayerRay;
+						TransmissionRay.TMin = 0.01;
+						TransmissionRay.TMax = Payload.HitT;
+						Transmission = Transmit(TransmissionRay, RandSequence, TopLayerRadiance);
+					}
+
 					RayDesc LightingRay = TopLayerRay;
 					if (DIM_HYBRID)
 					{
@@ -411,6 +472,7 @@ RAY_TRACING_ENTRY_RAYGEN(RayTracingReflectionsRGS)
 						LightingRay.Direction = normalize(Payload.WorldPos - View.WorldCameraOrigin);
 					}
 
+					float3 Radiance = 0.0;
 					AccumulateResults(
 						Payload,
 						LightingRay,
@@ -421,7 +483,8 @@ RAY_TRACING_ENTRY_RAYGEN(RayTracingReflectionsRGS)
 						ShouldDoDirectLighting,
 						ShouldDoEmissiveAndIndirectLighting,
 						RayCone,
-						TopLayerRadiance);
+						Radiance);
+					TopLayerRadiance += Radiance * Transmission;
 				}
 
 				bool isTopLayerRayValid = Payload.IsHit() || (bAllowSkySampling && ReflectionStruct.SkyLightParameters.y > 0);
@@ -538,7 +601,7 @@ RAY_TRACING_ENTRY_RAYGEN(RayTracingReflectionsRGS)
 						float IndirectIrradiance = 0;
 
 						// Not possible to use the screen culled grid. So going over every capture.
-						// dxr_todo This would need a spatial data structure.
+						//#dxr_todo This would need a spatial data structure.
 						uint NumCulledReflectionCaptures = ForwardLightData.NumReflectionCaptures;
 						uint ReflectionCapturesStartIndex = 0;
 
@@ -596,7 +659,7 @@ RAY_TRACING_ENTRY_RAYGEN(RayTracingReflectionsRGS)
 				// If there is a contribution on this path, apply the required fog on the sample. 
 				// Also accumulate the path contribution.
 				// In case no contribution is accumulated, the ReflectionEnvironmentAndSky pass will fill it up using reflection volumes.
-				//#dxr - todo: when MaxBounces>1, we will need to take into account a path throughput affected from the fog transmittance.
+				//#dxr_todo: when MaxBounces>1, we will need to take into account a path throughput affected from the fog transmittance.
 				float4 HeightFogInscatteringAndTransmittance = float4(0.0f, 0.0f, 0.0f, 1.0f);	// no fog 
 				if (bApplyHeightFog)
 				{
@@ -667,7 +730,7 @@ RAY_TRACING_ENTRY_RAYGEN(RayTracingReflectionsRGS)
 					float IndirectIrradiance = 0; 
 
 					// Not possible to use the screen culled grid. So going over every capture.
-					// dxr_todo This would need a spatial data structure.
+					//#dxr_todo This would need a spatial data structure.
 					uint NumCulledReflectionCaptures = ForwardLightData.NumReflectionCaptures;
 					uint ReflectionCapturesStartIndex = 0;
 
@@ -738,6 +801,26 @@ RAY_TRACING_ENTRY_RAYGEN(RayTracingReflectionsRGS)
 			ReflectedColor *= DistanceFade * max(RoughnessFade, ClearCoatRoughnessFade) * rcp(LocalSamplesPerPixel);
 		}
 
+		if (EnableTranslucency)
+		{
+			float3 WorldPosition;
+			float3 CameraDirection;
+			ReconstructWorldPositionAndCameraDirectionFromDeviceZ(PixelCoord, DeviceZ, WorldPosition, CameraDirection);
+
+			RayDesc TransmissionRay;
+			TransmissionRay.Origin = View.WorldCameraOrigin;
+			TransmissionRay.Direction = CameraDirection;
+			TransmissionRay.TMin = 0.0;
+			TransmissionRay.TMax = length(WorldPosition - View.WorldCameraOrigin);
+			float3 Inscattering = 0.0;
+			RandomSequence RandSequence;
+			float Transmission = Transmit(TransmissionRay, RandSequence, Inscattering);
+			//#dxr_todo: add Inscattering support (this is emission only so far)
+			//#dxr_todo: depending on if there is a translucency pass the component might already be attenuated. Check for double contribution
+			ReflectedColor.rgb *= Transmission;
+			ReflectedColor.rgb += Inscattering;
+		}
+
 	#if USE_PREEXPOSURE
 		ReflectedColor.xyz *= View.PreExposure;
 	#endif
diff --git a/Engine/Source/Runtime/Engine/Classes/Engine/Scene.h b/Engine/Source/Runtime/Engine/Classes/Engine/Scene.h
--- a/Engine/Source/Runtime/Engine/Classes/Engine/Scene.h
+++ b/Engine/Source/Runtime/Engine/Classes/Engine/Scene.h
@@ -1062,6 +1062,9 @@ struct FPostProcessSettings
 	UPROPERTY(EditAnywhere, BlueprintReadWrite, Category = Overrides, meta = (PinHiddenByDefault, InlineEditConditionToggle))
 	uint32 bOverride_RayTracingReflectionsShadows : 1;
 
+	UPROPERTY(EditAnywhere, BlueprintReadWrite, Category = Overrides, meta = (PinHiddenByDefault, InlineEditConditionToggle))
+	uint32 bOverride_RayTracingReflectionsTranslucency : 1;
+
 	UPROPERTY(EditAnywhere, BlueprintReadWrite, Category = Overrides, meta = (PinHiddenByDefault, InlineEditConditionToggle))
 	uint32 bOverride_TranslucencyType : 1;
 
@@ -1739,6 +1742,10 @@ struct FPostProcessSettings
 	UPROPERTY(EditAnywhere, BlueprintReadWrite, Category = "Rendering Features|Ray Tracing Reflections", meta = (editcondition = "bOverride_RayTracingReflectionsShadows", DisplayName = "Shadows"))
 	EReflectedAndRefractedRayTracedShadows RayTracingReflectionsShadows;
 
+	/** Enables ray tracing translucency in reflections. */
+	UPROPERTY(EditAnywhere, BlueprintReadWrite, Category = "Rendering Features|Ray Tracing Reflections", meta = (editcondition = "bOverride_RayTracingReflectionsTranslucency", DisplayName = "Include Translucent Objects"))
+	uint8 RayTracingReflectionsTranslucency : 1;
+
 
 	/** Sets the translucency type */
 	UPROPERTY(EditAnywhere, BlueprintReadWrite, Category = "Rendering Features|Translucency", meta = (editcondition = "bOverride_TranslucencyType", DisplayName = "Type"))
diff --git a/Engine/Source/Runtime/Engine/Private/Scene.cpp b/Engine/Source/Runtime/Engine/Private/Scene.cpp
--- a/Engine/Source/Runtime/Engine/Private/Scene.cpp
+++ b/Engine/Source/Runtime/Engine/Private/Scene.cpp
@@ -557,7 +557,8 @@ FPostProcessSettings::FPostProcessSettings()
 	RayTracingReflectionsMaxBounces = 1;
 	RayTracingReflectionsSamplesPerPixel = 1;
 	RayTracingReflectionsShadows = EReflectedAndRefractedRayTracedShadows::Hard_shadows;
-	
+	RayTracingReflectionsTranslucency = 0;
+
 	TranslucencyType = ETranslucencyType::Raster;
 	RayTracingTranslucencyMaxRoughness = 0.6f;
 	RayTracingTranslucencyRefractionRays = 3; // 3 to: first hit surface, second hit back inner surface and a third to fetch the background.
@@ -743,6 +744,7 @@ FPostProcessSettings::FPostProcessSettings(const FPostProcessSettings& Settings)
 	, bOverride_RayTracingReflectionsMaxBounces(Settings.bOverride_RayTracingReflectionsMaxBounces)
 	, bOverride_RayTracingReflectionsSamplesPerPixel(Settings.bOverride_RayTracingReflectionsSamplesPerPixel)
 	, bOverride_RayTracingReflectionsShadows(Settings.bOverride_RayTracingReflectionsShadows)
+	, bOverride_RayTracingReflectionsTranslucency(Settings.bOverride_RayTracingReflectionsTranslucency)
 	, bOverride_TranslucencyType(Settings.bOverride_TranslucencyType)
 	, bOverride_RayTracingTranslucencyMaxRoughness(Settings.bOverride_RayTracingTranslucencyMaxRoughness)
 	, bOverride_RayTracingTranslucencyRefractionRays(Settings.bOverride_RayTracingTranslucencyRefractionRays)
@@ -919,6 +921,7 @@ FPostProcessSettings::FPostProcessSettings(const FPostProcessSettings& Settings)
 	, RayTracingReflectionsMaxBounces(Settings.RayTracingReflectionsMaxBounces)
 	, RayTracingReflectionsSamplesPerPixel(Settings.RayTracingReflectionsSamplesPerPixel)
 	, RayTracingReflectionsShadows(Settings.RayTracingReflectionsShadows)
+	, RayTracingReflectionsTranslucency(Settings.RayTracingReflectionsTranslucency)
 
 	, TranslucencyType(Settings.TranslucencyType)
 	, RayTracingTranslucencyMaxRoughness(Settings.RayTracingTranslucencyMaxRoughness)
diff --git a/Engine/Source/Runtime/Engine/Private/SceneView.cpp b/Engine/Source/Runtime/Engine/Private/SceneView.cpp
--- a/Engine/Source/Runtime/Engine/Private/SceneView.cpp
+++ b/Engine/Source/Runtime/Engine/Private/SceneView.cpp
@@ -1419,6 +1419,11 @@ void FSceneView::OverridePostProcessSettings(const FPostProcessSettings& Src, fl
 			Dest.RayTracingReflectionsShadows = Src.RayTracingReflectionsShadows;
 		}
 
+		if (Src.bOverride_RayTracingReflectionsTranslucency)
+		{
+			Dest.RayTracingReflectionsTranslucency = Src.RayTracingReflectionsTranslucency;
+		}
+
 		if (Src.bOverride_TranslucencyType)
 		{
 			Dest.TranslucencyType = Src.TranslucencyType;
diff --git a/Engine/Source/Runtime/Renderer/Private/RayTracing/RayTracingReflections.cpp b/Engine/Source/Runtime/Renderer/Private/RayTracing/RayTracingReflections.cpp
--- a/Engine/Source/Runtime/Renderer/Private/RayTracing/RayTracingReflections.cpp
+++ b/Engine/Source/Runtime/Renderer/Private/RayTracing/RayTracingReflections.cpp
@@ -61,6 +61,16 @@ static FAutoConsoleVariableRef CVarRayTracingReflectionsShadows(
 	TEXT(" 2: Soft area shadows")
 );
 
+static int32 GRayTracingReflectionsTranslucency = -1;
+static FAutoConsoleVariableRef CVarRayTracingReflectionsTranslucency(
+	TEXT("r.RayTracing.Reflections.Translucency"),
+	GRayTracingReflectionsTranslucency,
+	TEXT("Translucent objects visible in ray tracing reflections)")
+	TEXT(" -1: Driven by postprocessing volume (default)")
+	TEXT(" 0: Translucent objects not visible")
+	TEXT(" 1: Translucent objects visible")
+);
+
 static int32 GRayTracingReflectionsCaptures = 0;
 static FAutoConsoleVariableRef CVarRayTracingReflectionsCaptures(
 	TEXT("r.RayTracing.Reflections.ReflectionCaptures"),
@@ -138,7 +148,6 @@ static TAutoConsoleVariable<int32> CVarRayTracingReflectionsMaxUnderCoatBounces(
 	ECVF_RenderThreadSafe);
 
 
-
 static const int32 GReflectionLightCountMaximum = 64;
 
 class FRayTracingReflectionsRGS : public FGlobalShader
@@ -173,6 +182,7 @@ class FRayTracingReflectionsRGS : public FGlobalShader
 		SHADER_PARAMETER(int32, TestPathRoughness)
 		SHADER_PARAMETER(float, MinClearCoatLevel)
 		SHADER_PARAMETER(int32, MaxUnderCoatBounces)
+		SHADER_PARAMETER(uint32, EnableTranslucency)
 
 		SHADER_PARAMETER_SRV(RaytracingAccelerationStructure, TLAS)
 		SHADER_PARAMETER_RDG_TEXTURE(Texture2D, SceneColor)
@@ -234,11 +244,13 @@ void FDeferredShadingSceneRenderer::PrepareRayTracingReflections(const FViewInfo
 	}
 	else
 	{
-		FRayTracingReflectionsRGS::FPermutationDomain PermutationVector;
-		PermutationVector.Set<FRayTracingReflectionsRGS::FEnableTwoSidedGeometryForShadowDim>(EnableRayTracingShadowTwoSidedGeometry());
-		PermutationVector.Set<FRayTracingReflectionsRGS::FDeferredMaterialMode>(EDeferredMaterialMode::None);
-		auto RayGenShader = View.ShaderMap->GetShader<FRayTracingReflectionsRGS>(PermutationVector);
-		OutRayGenShaders.Add(RayGenShader->GetRayTracingShader());
+		{
+			FRayTracingReflectionsRGS::FPermutationDomain PermutationVector;
+			PermutationVector.Set<FRayTracingReflectionsRGS::FEnableTwoSidedGeometryForShadowDim>(EnableRayTracingShadowTwoSidedGeometry());
+			PermutationVector.Set<FRayTracingReflectionsRGS::FDeferredMaterialMode>(EDeferredMaterialMode::None);
+			auto RayGenShader = View.ShaderMap->GetShader<FRayTracingReflectionsRGS>(PermutationVector);
+			OutRayGenShaders.Add(RayGenShader->GetRayTracingShader());
+		}
 	}
 }
 
@@ -261,6 +273,7 @@ void FDeferredShadingSceneRenderer::RenderRayTracingReflections(
 
 	const bool bHybridReflections = CVarRayTracingReflectionsHybrid.GetValueOnRenderThread() != 0;
 	const bool bSortMaterials = bHybridReflections || CVarRayTracingReflectionsSortMaterials.GetValueOnRenderThread() != 0;
+	const uint32 EnableTranslucency = GRayTracingReflectionsTranslucency > -1 ? (uint32)GRayTracingReflectionsTranslucency : (uint32) View.FinalPostProcessSettings.RayTracingReflectionsTranslucency;
 
 	FSceneRenderTargets& SceneContext = FSceneRenderTargets::Get(GraphBuilder.RHICmdList);
 
@@ -309,6 +322,7 @@ void FDeferredShadingSceneRenderer::RenderRayTracingReflections(
 	CommonParameters.TestPathRoughness = CVarRayTracingReflectionsTestPathRoughness.GetValueOnRenderThread();
 	CommonParameters.MinClearCoatLevel = CVarRayTracingReflectionsMinClearCoatLevel.GetValueOnRenderThread();
 	CommonParameters.MaxUnderCoatBounces = CVarRayTracingReflectionsMaxUnderCoatBounces.GetValueOnRenderThread();
+	CommonParameters.EnableTranslucency = EnableTranslucency;
 
 	CommonParameters.TLAS = View.RayTracingScene.RayTracingSceneRHI->GetShaderResourceView();
 	CommonParameters.ViewUniformBuffer = View.ViewUniformBuffer;
