Fixed fix in RTGI shaders binding

diff --git a/Engine/Shaders/Private/RayTracing/RayTracingFinalGatherRGS.usf b/Engine/Shaders/Private/RayTracing/RayTracingFinalGatherRGS.usf
--- a/Engine/Shaders/Private/RayTracing/RayTracingFinalGatherRGS.usf
+++ b/Engine/Shaders/Private/RayTracing/RayTracingFinalGatherRGS.usf
@@ -20,8 +20,8 @@
 uint SampleIndex;
 uint SamplesPerPixel;
 uint UpscaleFactor;
-uint TileOffsetX;
-uint TileOffsetY;
+uint RenderTileOffsetX;
+uint RenderTileOffsetY;
 float DiffuseThreshold;
 float MaxNormalBias;
 float FinalGatherDistance;
@@ -41,7 +41,7 @@ uint DispatchThreadIdToLinearIndex(uint2 DispatchThreadId)
 
 RAY_TRACING_ENTRY_RAYGEN(RayTracingFinalGatherRGS)
 {
-	uint2 DispatchThreadId = DispatchRaysIndex().xy + View.ViewRectMin + uint2(TileOffsetX, TileOffsetY);
+	uint2 DispatchThreadId = DispatchRaysIndex().xy + View.ViewRectMin + uint2(RenderTileOffsetX, RenderTileOffsetY);
 	RWGlobalIlluminationUAV[DispatchThreadId] = 0.0;
 	RWRayDistanceUAV[DispatchThreadId] = float2(-1.0, 0.0);
 
