Packed ray tracing payload is now used during the main ray traced reflection rendering. Same payload instance is used for main material ray trace and for shadow rays during direct lighting. This significantly improves performance for regular code path, when not using miss shaders to evaluate lighting.

diff --git a/Engine/Shaders/Private/RayTracing/RayTracingCommon.ush b/Engine/Shaders/Private/RayTracing/RayTracingCommon.ush
--- a/Engine/Shaders/Private/RayTracing/RayTracingCommon.ush
+++ b/Engine/Shaders/Private/RayTracing/RayTracingCommon.ush
@@ -176,7 +176,7 @@ struct FMaterialClosestHitPayload : FMinimalPayload
 	uint ShadingModelID;				// 92        1       4bits
 	uint BlendingMode;					// 96        0       4bits (packed with ShadingModelID)
 	uint PrimitiveLightingChannelMask;	// 100       0       (removed)
-	float4 CustomData;					// 104       4       32bits -- only xy are packed at the moment
+	float4 CustomData;					// 104       4       32bits -- only zw are packed at the moment
 	float GBufferAO;					// 120       0       (removed)
 	float3 IndirectIrradiance;			// 124       0       48bits -- gbuffer only has float payload and there are truncation HLSL warnings
 
@@ -196,8 +196,6 @@ struct FMaterialClosestHitPayload : FMinimalPayload
 	void SetRayCone(FRayCone NewRayCone) { RayCone = NewRayCone; }
 };
 
-#define USE_PACKED_PAYLOAD 1
-
 #if !COMPUTESHADER
 
 RayDesc CreatePrimaryRay(float2 UV)
@@ -215,7 +213,6 @@ RayDesc CreatePrimaryRay(float2 UV)
 	return Ray;
 }
 
-#if USE_PACKED_PAYLOAD
 struct FPackedMaterialClosestHitPayload : FMinimalPayload
 {
 	// float FMinimalPayload::HitT                       // 4  bytes
@@ -231,7 +228,80 @@ struct FPackedMaterialClosestHitPayload : FMinimalPayload
 	void SetMinimalPayloadMode() { IorAndShadingModelIDAndBlendingModeAndFlags |= RAY_TRACING_PAYLOAD_FLAG_MINIMAL_PAYLOAD << 24; }
 	bool IsMinimalPayloadMode() { return ((IorAndShadingModelIDAndBlendingModeAndFlags >> 24) & RAY_TRACING_PAYLOAD_FLAG_MINIMAL_PAYLOAD) != 0; }
 
-	void SetRayCone(FRayCone NewRayCone) { RayCone = NewRayCone; }
+	void SetRayCone(FRayCone InRayCone) { RayCone = InRayCone; }
+
+	void SetRadiance(in float3 InRadiance)
+	{
+		RadianceAndNormal[0] = f32tof16(InRadiance.x) | (f32tof16(InRadiance.y) << 16);
+		RadianceAndNormal[1] &= 0xffff << 16; // Clear the radiance (low bits), keep packed normal component (high bits)
+		RadianceAndNormal[1] |= f32tof16(InRadiance.z);
+	}
+
+	float3 GetRadiance()
+	{
+		float3 Result;
+		Result.x = f16tof32(RadianceAndNormal[0]);
+		Result.y = f16tof32(RadianceAndNormal[0] >> 16);
+		Result.z = f16tof32(RadianceAndNormal[1]);
+		return Result;
+	}
+
+	void SetIndirectIrradiance(float3 InIndirectIrradiance)
+	{
+		IndirectIrradiance[0] = f32tof16(InIndirectIrradiance.x);
+		IndirectIrradiance[0] |= f32tof16(InIndirectIrradiance.y) << 16;
+		IndirectIrradiance[1] = f32tof16(InIndirectIrradiance.z);
+	}
+
+	float3 GetIndirectIrradiance()
+	{
+		float3 Result;
+		Result.x = f16tof32(IndirectIrradiance[0]);
+		Result.y = f16tof32(IndirectIrradiance[0] >> 16);
+		Result.z = f16tof32(IndirectIrradiance[1]);
+		return Result;
+	}
+
+	float3 GetWorldNormal()
+	{
+		float3 Result;
+		Result.x = f16tof32(RadianceAndNormal[1] >> 16);
+		Result.y = f16tof32(RadianceAndNormal[2]);
+		Result.z = f16tof32(RadianceAndNormal[2] >> 16);
+		return Result;
+	}
+
+	float4 GetCustomData()
+	{
+		float4 Result;
+		Result.z = f16tof32(CustomData.x & 0xFFFF);
+		Result.w = f16tof32(CustomData.x >> 16);
+		Result.xy = 0;
+		return Result;
+	}
+	
+
+	float3 GetBaseColor()
+	{
+		float3 Result;
+		Result.x = f16tof32(BaseColorAndOpacity[0]);
+		Result.y = f16tof32(BaseColorAndOpacity[0] >> 16);
+		Result.z = f16tof32(BaseColorAndOpacity[1]);
+		return Result;
+	}
+
+	float GetOpacity()       { return f16tof32(BaseColorAndOpacity[1] >> 16); }
+	float GetMetallic()      { return float(MetallicAndSpecularAndRoughness & 0xFF) / 255.0f; }
+	float GetSpecular()      { return float((MetallicAndSpecularAndRoughness >> 8) & 0xFF) / 255.0f; }
+	float GetRoughness()     { return f16tof32(MetallicAndSpecularAndRoughness >> 16); }
+	float GetIor()           { return f16tof32(IorAndShadingModelIDAndBlendingModeAndFlags); }
+	uint GetBlendingMode()   { return (IorAndShadingModelIDAndBlendingModeAndFlags >> 20) & 0xF;  }
+	uint GetFlags()          { return (IorAndShadingModelIDAndBlendingModeAndFlags >> 24) & 0xFF; }
+	uint GetShadingModelID() { return (IorAndShadingModelIDAndBlendingModeAndFlags >> 16) & 0xF; }
+
+	float3 GetDiffuseColor() { return GetBaseColor() - GetBaseColor() * GetMetallic(); }
+	float3 GetSpecularColor() { return ComputeF0(GetSpecular(), GetBaseColor(), GetMetallic()); }
+
 };
 
 FPackedMaterialClosestHitPayload PackRayTracingPayload(FMaterialClosestHitPayload Input, in FRayCone RayCone)
@@ -264,66 +334,33 @@ FPackedMaterialClosestHitPayload PackRayTracingPayload(FMaterialClosestHitPayloa
 	return Output;
 }
 
-float3 UnpackRayTracingNormal(FPackedMaterialClosestHitPayload Input)
-{
-	float3 OutWorldNormal;
-	OutWorldNormal.x = f16tof32(Input.RadianceAndNormal[1] >> 16);
-	OutWorldNormal.y = f16tof32(Input.RadianceAndNormal[2]);
-	OutWorldNormal.z = f16tof32(Input.RadianceAndNormal[2] >> 16);
-	return OutWorldNormal;
-}
-
 FMaterialClosestHitPayload UnpackRayTracingPayload(FPackedMaterialClosestHitPayload Input, RayDesc Ray)
 {
 	FMaterialClosestHitPayload Output = (FMaterialClosestHitPayload)0;
-	Output.HitT = Input.HitT;
-	Output.RayCone = Input.RayCone; 
-	Output.Radiance.x = f16tof32(Input.RadianceAndNormal[0]);
-	Output.Radiance.y = f16tof32(Input.RadianceAndNormal[0] >> 16);
-	Output.Radiance.z = f16tof32(Input.RadianceAndNormal[1]);
-	Output.WorldNormal.x = f16tof32(Input.RadianceAndNormal[1] >> 16);
-	Output.WorldNormal.y = f16tof32(Input.RadianceAndNormal[2]);
-	Output.WorldNormal.z = f16tof32(Input.RadianceAndNormal[2] >> 16);
-	Output.BaseColor.x = f16tof32(Input.BaseColorAndOpacity[0]);
-	Output.BaseColor.y = f16tof32(Input.BaseColorAndOpacity[0] >> 16);
-	Output.BaseColor.z = f16tof32(Input.BaseColorAndOpacity[1]);
-	Output.Opacity	   = f16tof32(Input.BaseColorAndOpacity[1] >> 16);
-	Output.Metallic	   = float(Input.MetallicAndSpecularAndRoughness & 0xFF) / 255.0f;
-    Output.Specular	   = float((Input.MetallicAndSpecularAndRoughness >> 8) & 0xFF) / 255.0f;
-    Output.Roughness   = f16tof32(Input.MetallicAndSpecularAndRoughness >> 16);
-	Output.Ior            = f16tof32(Input.IorAndShadingModelIDAndBlendingModeAndFlags);      // 16 bits
-	Output.ShadingModelID = (Input.IorAndShadingModelIDAndBlendingModeAndFlags >> 16) & 0xF;  // 4 bits
-	Output.BlendingMode	  = (Input.IorAndShadingModelIDAndBlendingModeAndFlags >> 20) & 0xF;  // 4 bits
-	Output.Flags          = (Input.IorAndShadingModelIDAndBlendingModeAndFlags >> 24) & 0xFF; // 8 bits
-	Output.IndirectIrradiance.x	= f16tof32(Input.IndirectIrradiance[0]);
-	Output.IndirectIrradiance.y	= f16tof32(Input.IndirectIrradiance[0] >> 16);
-	Output.IndirectIrradiance.z	= f16tof32(Input.IndirectIrradiance[1]);
-	Output.CustomData.z = f16tof32(Input.CustomData.x & 0xFFFF);
-	Output.CustomData.w = f16tof32(Input.CustomData.x >> 16);
-	Output.CustomData.xy = 0;
+
+	Output.HitT               = Input.HitT;
+	Output.RayCone            = Input.RayCone;
+	Output.Radiance           = Input.GetRadiance();
+	Output.WorldNormal        = Input.GetWorldNormal();
+	Output.BaseColor          = Input.GetBaseColor();
+	Output.Opacity	          = Input.GetOpacity();
+	Output.Metallic	          = Input.GetMetallic();
+	Output.Specular           = Input.GetSpecular();
+    Output.Roughness          = Input.GetRoughness();
+	Output.Ior                = Input.GetIor();
+	Output.ShadingModelID     = Input.GetShadingModelID();
+	Output.BlendingMode	      = Input.GetBlendingMode();
+	Output.Flags              = Input.GetFlags();
+	Output.IndirectIrradiance = Input.GetIndirectIrradiance();
+	Output.CustomData         = Input.GetCustomData();
 
 	Output.DiffuseColor  = Output.BaseColor - Output.BaseColor * Output.Metallic;
 	Output.SpecularColor = ComputeF0(Output.Specular, Output.BaseColor, Output.Metallic);
 	Output.WorldPos		 = Ray.Origin + Output.HitT * Ray.Direction;
-	return Output;
-}
-
-#else // USE_PACKED_PAYLOAD
 
-typedef FMaterialClosestHitPayload FPackedMaterialClosestHitPayload;
-
-FPackedMaterialClosestHitPayload PackRayTracingPayload(FMaterialClosestHitPayload Input, in FRayCone RayCone)
-{
-	return Input;
-}
-
-FMaterialClosestHitPayload UnpackRayTracingPayload(FPackedMaterialClosestHitPayload Input, RayDesc Ray)
-{
-	return Input;
+	return Output;
 }
 
-#endif // USE_PACKED_PAYLOAD
-
 // WorldNormal is the vector towards which the ray position will be offseted.
 void ApplyPositionBias(inout RayDesc Ray, const float3 WorldNormal, const float MaxNormalBias)
 {
@@ -356,21 +393,20 @@ void ApplyCameraRelativeDepthBias(inout RayDesc Ray, uint2 PixelCoord, float Dev
 	ApplyPositionBias(Ray, WorldNormal, RayBias);
 }
 
-FMinimalPayload TraceMinimalRayInternal(
+void TraceVisibilityRayPacked(
+	inout FPackedMaterialClosestHitPayload PackedPayload,
 	in RaytracingAccelerationStructure TLAS,
 	in uint RayFlags,
 	in uint InstanceInclusionMask,
-	in uint RayContributionToHitGroupIndex,
-	in uint MultiplierForGeometryContributionToShaderIndex,
-	in uint MissShaderIndex,
 	in RayDesc Ray)
 {
-	FPackedMaterialClosestHitPayload PackedPayload = (FPackedMaterialClosestHitPayload)0;
-
-	// Set payload inputs
+	const uint RayContributionToHitGroupIndex = RAY_TRACING_SHADER_SLOT_SHADOW;
+	const uint MultiplierForGeometryContributionToShaderIndex = RAY_TRACING_NUM_SHADER_SLOTS;
+	const uint MissShaderIndex = 0;
 
 	// By enabling minimal payload mode all other payload information is ignored, meaning these functions need no payload inputs
 	PackedPayload.SetMinimalPayloadMode();
+	PackedPayload.HitT = 0;
 
 	// Trace the ray
 
@@ -383,6 +419,16 @@ FMinimalPayload TraceMinimalRayInternal(
 		MissShaderIndex,
 		Ray,
 		PackedPayload);
+}
+
+FMinimalPayload TraceVisibilityRay(
+	in RaytracingAccelerationStructure TLAS,
+	in uint RayFlags,
+	in uint InstanceInclusionMask,
+	in RayDesc Ray)
+{
+	FPackedMaterialClosestHitPayload PackedPayload = (FPackedMaterialClosestHitPayload)0;
+	TraceVisibilityRayPacked(PackedPayload, TLAS, RayFlags, InstanceInclusionMask, Ray);
 
 	// Unpack the payload
 
@@ -396,25 +442,26 @@ FMinimalPayload TraceMinimalRayInternal(
 	return MinimalPayload;
 }
 
-FMaterialClosestHitPayload TraceRayInternal(
+void TraceMaterialRayPacked(
+	inout FPackedMaterialClosestHitPayload Payload,
 	in RaytracingAccelerationStructure TLAS,
 	in uint RayFlags,
 	in uint InstanceInclusionMask,
-	in uint RayContributionToHitGroupIndex,
-	in uint MultiplierForGeometryContributionToShaderIndex,
-	in uint MissShaderIndex,
 	in RayDesc Ray,
 	// Payload Inputs
-	in FRayCone RayCone)
+	inout FRayCone RayCone)
 {
-	FPackedMaterialClosestHitPayload PackedPayload = (FPackedMaterialClosestHitPayload)0;
+	const uint RayContributionToHitGroupIndex = RAY_TRACING_SHADER_SLOT_MATERIAL;
+	const uint MultiplierForGeometryContributionToShaderIndex = RAY_TRACING_NUM_SHADER_SLOTS;
+	const uint MissShaderIndex = 0;
 
 	// Set payload inputs
 
-	PackedPayload.SetRayCone(RayCone);
+	Payload.SetRayCone(RayCone);
 
 	// Trace the ray
 
+	Payload.HitT = 0;
 	TraceRay(
 		TLAS,
 		RayFlags,
@@ -423,35 +470,9 @@ FMaterialClosestHitPayload TraceRayInternal(
 		MultiplierForGeometryContributionToShaderIndex,
 		MissShaderIndex,
 		Ray,
-		PackedPayload);
-
-	// Unpack the payload
+		Payload);
 
-	FMaterialClosestHitPayload Payload = UnpackRayTracingPayload(PackedPayload, Ray);
-
-	return Payload;
-}
-
-FMinimalPayload TraceVisibilityRay(
-	in RaytracingAccelerationStructure TLAS,
-	in uint RayFlags,
-	in uint InstanceInclusionMask,
-	in RayDesc Ray)
-{
-	const uint RayContributionToHitGroupIndex = RAY_TRACING_SHADER_SLOT_SHADOW;
-	const uint MultiplierForGeometryContributionToShaderIndex = RAY_TRACING_NUM_SHADER_SLOTS;
-	const uint MissShaderIndex = 0;
-
-	FMinimalPayload Payload = TraceMinimalRayInternal(
-		TLAS,
-		RayFlags,
-		InstanceInclusionMask,
-		RayContributionToHitGroupIndex,
-		MultiplierForGeometryContributionToShaderIndex,
-		MissShaderIndex,
-		Ray);
-
-	return Payload;
+	RayCone = Payload.RayCone;
 }
 
 FMaterialClosestHitPayload TraceMaterialRay(
@@ -466,7 +487,15 @@ FMaterialClosestHitPayload TraceMaterialRay(
 	const uint MultiplierForGeometryContributionToShaderIndex = RAY_TRACING_NUM_SHADER_SLOTS;
 	const uint MissShaderIndex = 0;
 
-	FMaterialClosestHitPayload Payload = TraceRayInternal(
+	FPackedMaterialClosestHitPayload PackedPayload = (FPackedMaterialClosestHitPayload)0;
+
+	// Set payload inputs
+
+	PackedPayload.SetRayCone(RayCone);
+
+	// Trace the ray
+
+	TraceRay(
 		TLAS,
 		RayFlags,
 		InstanceInclusionMask,
@@ -474,7 +503,11 @@ FMaterialClosestHitPayload TraceMaterialRay(
 		MultiplierForGeometryContributionToShaderIndex,
 		MissShaderIndex,
 		Ray,
-		RayCone);
+		PackedPayload);
+
+	// Unpack the payload
+
+	FMaterialClosestHitPayload Payload = UnpackRayTracingPayload(PackedPayload, Ray);
 
 	RayCone = Payload.RayCone;
 
diff --git a/Engine/Shaders/Private/RayTracing/RayTracingLightingCommon.ush b/Engine/Shaders/Private/RayTracing/RayTracingLightingCommon.ush
--- a/Engine/Shaders/Private/RayTracing/RayTracingLightingCommon.ush
+++ b/Engine/Shaders/Private/RayTracing/RayTracingLightingCommon.ush
@@ -35,6 +35,106 @@ struct FRTLightingData
 
 StructuredBuffer<FRTLightingData> LightDataBuffer;
 
+FRTLightingData GetRayTracingLightData(int LightIndex)
+{
+	return LightDataBuffer[LightIndex];
+}
+
+FDeferredLightData GetRayTracingDeferredLightData(int LightIndex,
+	inout int OutLightProfileIndex,
+	inout int OutRectLightTextureIndex,
+	inout uint OutLightType)
+{
+	FDeferredLightData LightData = (FDeferredLightData)0;
+	FRTLightingData RayTracingLightData = GetRayTracingLightData(LightIndex);
+
+	const uint LightType = RayTracingLightData.Type;
+
+	LightData.Position = RayTracingLightData.LightPosition;
+	LightData.InvRadius = RayTracingLightData.InvRadius;
+	LightData.Color = RayTracingLightData.LightColor;
+	LightData.FalloffExponent = RayTracingLightData.FalloffExponent;
+	LightData.Direction = RayTracingLightData.Direction;
+	LightData.Tangent = RayTracingLightData.Tangent;
+	LightData.SpotAngles = RayTracingLightData.SpotAngles;
+	LightData.SourceRadius = RayTracingLightData.SourceRadius;
+	LightData.SourceLength = RayTracingLightData.SourceLength;
+	LightData.SoftSourceRadius = RayTracingLightData.SoftSourceRadius;
+	LightData.SpecularScale = RayTracingLightData.SpecularScale;
+	LightData.RectLightBarnCosAngle = RayTracingLightData.RectLightBarnCosAngle;
+	LightData.RectLightBarnLength = RayTracingLightData.RectLightBarnLength;
+	LightData.ContactShadowLength = 0.0;
+	LightData.DistanceFadeMAD = RayTracingLightData.DistanceFadeMAD;
+
+	LightData.ShadowMapChannelMask = float4(0, 0, 0, 0);
+	LightData.ShadowedBits = 0; // Not lit dynamic shadows 
+	LightData.ContactShadowLengthInWS = false;
+
+	LightData.bRadialLight = (LightType != LIGHT_TYPE_DIRECTIONAL);
+	LightData.bSpotLight = (LightType == LIGHT_TYPE_SPOT);
+	LightData.bRectLight = (LightType == LIGHT_TYPE_RECT);
+
+	if (LightType == LIGHT_TYPE_DIRECTIONAL)
+	{
+		LightData.bInverseSquared = false;
+	}
+	else
+	{
+		LightData.bInverseSquared = LightData.FalloffExponent == 0;
+	}
+
+	OutLightProfileIndex = RayTracingLightData.LightProfileIndex;
+	OutRectLightTextureIndex = RayTracingLightData.RectLightTextureIndex;
+	OutLightType = LightType;
+
+	return LightData;
+}
+
+FGBufferData GetGBufferDataFromPayload(in FMaterialClosestHitPayload Payload)
+{
+	FGBufferData GBufferData = (FGBufferData)0;
+	GBufferData.Depth = 1.f; // Do not use depth
+	GBufferData.WorldNormal = Payload.WorldNormal;
+	GBufferData.BaseColor = Payload.BaseColor;
+	GBufferData.Metallic = Payload.Metallic;
+	GBufferData.Specular = Payload.Specular;
+	GBufferData.Roughness = Payload.Roughness;
+	GBufferData.CustomData = Payload.CustomData;
+	GBufferData.GBufferAO = Payload.GBufferAO;
+	GBufferData.IndirectIrradiance = (Payload.IndirectIrradiance.x + Payload.IndirectIrradiance.y + Payload.IndirectIrradiance.z) / 3.f;
+	GBufferData.ShadingModelID = Payload.ShadingModelID;
+	GBufferData.SpecularColor = Payload.SpecularColor;
+	GBufferData.DiffuseColor = Payload.DiffuseColor;
+	return GBufferData;
+}
+
+FGBufferData GetGBufferDataFromPayload(in FPackedMaterialClosestHitPayload Payload)
+{
+	RayDesc DummyRay = (RayDesc)0; // Not needed to unpack GBuffer
+	return GetGBufferDataFromPayload(UnpackRayTracingPayload(Payload, DummyRay));
+}
+
+FRectTexture GetRayTracingRectTextureData()
+{
+	FRectTexture RectTexture;
+
+#if USE_SOURCE_TEXTURE_ARRAY 
+	RectTexture.SourceTexture0 = RaytracingLightsDataPacked.RectLightTexture0;
+	RectTexture.SourceTexture1 = RaytracingLightsDataPacked.RectLightTexture1;
+	RectTexture.SourceTexture2 = RaytracingLightsDataPacked.RectLightTexture2;
+	RectTexture.SourceTexture3 = RaytracingLightsDataPacked.RectLightTexture3;
+	RectTexture.SourceTexture4 = RaytracingLightsDataPacked.RectLightTexture4;
+	RectTexture.SourceTexture5 = RaytracingLightsDataPacked.RectLightTexture5;
+	RectTexture.SourceTexture6 = RaytracingLightsDataPacked.RectLightTexture6;
+	RectTexture.SourceTexture7 = RaytracingLightsDataPacked.RectLightTexture7;
+	RectTexture.SourceTextureIndex = 99;
+#else
+	RectTexture = InitRectTexture(RaytracingLightsDataPacked.RectLightTexture0);
+#endif
+
+	return RectTexture;
+}
+
 float GetRoughnessFade(float Roughness, float MaxRoughness)
 {
 	float RoughnessMaskScale = -2.0 / MaxRoughness;
@@ -77,12 +177,12 @@ float3 GenerateReflectedRayDirection(
 
 float3 SampleAreaLightDirection(
 	in FDeferredLightData LightData, 
-	in FMaterialClosestHitPayload Payload,
+	in float3 WorldPosition,
+	in float3 WorldNormal,
 	in uint LightType, 
 	in RandomSequence RandSequence )
 {
 	float3 ShadowRayDirection;
-	float3 WorldPosition = Payload.WorldPos;
 
 	float3 RayOrigin;
 	float RayTMin, RayTMax;
@@ -103,7 +203,7 @@ float3 SampleAreaLightDirection(
 	{
 		GenerateDirectionalLightOcclusionRay(
 			LightParameters,
-			WorldPosition, Payload.WorldNormal,
+			WorldPosition, WorldNormal,
 			RandSample,
 			/* out */ RayOrigin,
 			/* out */ ShadowRayDirection,
@@ -114,7 +214,7 @@ float3 SampleAreaLightDirection(
 	{
 		GenerateSpotLightOcclusionRay(
 			LightParameters,
-			WorldPosition, Payload.WorldNormal,
+			WorldPosition, WorldNormal,
 			RandSample,
 			/* out */ RayOrigin,
 			/* out */ ShadowRayDirection,
@@ -127,7 +227,7 @@ float3 SampleAreaLightDirection(
 		{
 			GeneratePointLightOcclusionRay(
 				LightParameters,
-				WorldPosition, Payload.WorldNormal,
+				WorldPosition, WorldNormal,
 				RandSample,
 				/* out */ RayOrigin,
 				/* out */ ShadowRayDirection,
@@ -138,7 +238,7 @@ float3 SampleAreaLightDirection(
 		{
 			GenerateSphereLightOcclusionRay(
 				LightParameters,
-				WorldPosition, Payload.WorldNormal,
+				WorldPosition, WorldNormal,
 				RandSample,
 				/* out */ RayOrigin,
 				/* out */ ShadowRayDirection,
@@ -151,7 +251,7 @@ float3 SampleAreaLightDirection(
 	{
 		GenerateRectLightOcclusionRay(
 			LightParameters,
-			WorldPosition, Payload.WorldNormal,
+			WorldPosition, WorldNormal,
 			RandSample,
 			/* out */ RayOrigin,
 			/* out */ ShadowRayDirection,
@@ -252,50 +352,21 @@ uint2 CullDirectLighting(
 
 float3 ComputeDirectLightingCulled(
 	in uint2 LightCullMask,
-	in RayDesc Ray,
+	in float3 WorldPosition,
+	in float3 ViewDirection,
 	in FRayCone RayCone,
 	in RaytracingAccelerationStructure TLAS,
-	in FMaterialClosestHitPayload Payload,
+	inout FPackedMaterialClosestHitPayload Payload,
 	in RandomSequence RandSequence,
 	in uint ReflectedShadowsType,
 	in float ShadowMaxNormalBias)
 {
 	float3 DirectLighting = float3(0.0, 0.0, 0.0);
 
-	float3 WorldPosition = Payload.WorldPos;
-	float3 CameraVector = -Ray.Direction;
 	float AmbientOcclusion = 1.0;
 
-	FGBufferData GBufferData = (FGBufferData)0;
-	GBufferData.Depth = 1.f; // Do not use depth
-	GBufferData.WorldNormal = Payload.WorldNormal;
-	GBufferData.BaseColor = Payload.BaseColor;
-	GBufferData.Metallic = Payload.Metallic;
-	GBufferData.Specular = Payload.Specular;
-	GBufferData.Roughness = Payload.Roughness;
-	GBufferData.CustomData = Payload.CustomData;
-	GBufferData.GBufferAO = Payload.GBufferAO;
-	GBufferData.IndirectIrradiance = (Payload.IndirectIrradiance.x + Payload.IndirectIrradiance.y + Payload.IndirectIrradiance.z) / 3.f; 
-	GBufferData.ShadingModelID = Payload.ShadingModelID;
-	GBufferData.SpecularColor = Payload.SpecularColor;
-	GBufferData.DiffuseColor = Payload.DiffuseColor;
-
-	FRectTexture RectTexture;
-#if USE_SOURCE_TEXTURE_ARRAY 
-	RectTexture.SourceTexture0 = RaytracingLightsDataPacked.RectLightTexture0;
-	RectTexture.SourceTexture1 = RaytracingLightsDataPacked.RectLightTexture1;
-	RectTexture.SourceTexture2 = RaytracingLightsDataPacked.RectLightTexture2;
-	RectTexture.SourceTexture3 = RaytracingLightsDataPacked.RectLightTexture3;
-	RectTexture.SourceTexture4 = RaytracingLightsDataPacked.RectLightTexture4;
-	RectTexture.SourceTexture5 = RaytracingLightsDataPacked.RectLightTexture5;
-	RectTexture.SourceTexture6 = RaytracingLightsDataPacked.RectLightTexture6;
-	RectTexture.SourceTexture7 = RaytracingLightsDataPacked.RectLightTexture7;
-	RectTexture.SourceTextureIndex = 99;
-#elif USE_SOURCE_TEXTURE
-	RectTexture = InitRectTexture(RaytracingLightsDataPacked.RectLightTexture0);
-#else
-	RectTexture = (FRectTexture)0;
-#endif
+	FGBufferData GBufferData = GetGBufferDataFromPayload(Payload);
+	FRectTexture RectTexture = GetRayTracingRectTextureData();
 
 	uint DirIndex = 0;
 
@@ -404,7 +475,7 @@ float3 ComputeDirectLightingCulled(
 		{
 			if (ReflectedShadowsType == 2)
 			{
-				ShadowRayDirection = SampleAreaLightDirection(LightData, Payload, LightType, RandSequence);
+				ShadowRayDirection = SampleAreaLightDirection(LightData, WorldPosition, Payload.GetWorldNormal(), LightType, RandSequence);
 			}
 
 			RayDesc ShadowRay;
@@ -412,7 +483,7 @@ float3 ComputeDirectLightingCulled(
 			ShadowRay.Direction = ShadowRayDirection;
 			ShadowRay.TMin = 1e-4f;
 			ShadowRay.TMax = Active ? length(ToLight) : ShadowRay.TMin;
-			ApplyPositionBias(ShadowRay, Payload.WorldNormal, ShadowMaxNormalBias);
+			ApplyPositionBias(ShadowRay, Payload.GetWorldNormal(), ShadowMaxNormalBias);
 
 			uint RayFlags = RAY_FLAG_ACCEPT_FIRST_HIT_AND_END_SEARCH | RAY_FLAG_SKIP_CLOSEST_HIT_SHADER;
 			const uint InstanceInclusionMask = RAY_TRACING_MASK_SHADOW;
@@ -432,7 +503,7 @@ float3 ComputeDirectLightingCulled(
 		if (Active)
 		{
 			float SurfaceShadow = 1.0f;
-			float3 LightContribution = GetDynamicLighting(WorldPosition, -CameraVector, GBufferData, AmbientOcclusion, GBufferData.ShadingModelID, LightData, LightAttenuation, 0.5, uint2(0, 0), RectTexture, SurfaceShadow).xyz;
+			float3 LightContribution = GetDynamicLighting(WorldPosition, ViewDirection, GBufferData, AmbientOcclusion, GBufferData.ShadingModelID, LightData, LightAttenuation, 0.5, uint2(0, 0), RectTexture, SurfaceShadow).xyz;
 			DirectLighting += LightContribution * LightProfileMultiplier;
 		}
 
@@ -446,6 +517,114 @@ float3 ComputeDirectLightingCulled(
 	return DirectLighting;
 }
 
+float3 ComputeDirectLightingMonolithic(
+	in float3 WorldPosition,
+	in float3 ViewDirection,
+	in FRayCone RayCone,
+	in RaytracingAccelerationStructure TLAS,
+	inout FPackedMaterialClosestHitPayload Payload,
+	in RandomSequence RandSequence,
+	in uint ReflectedShadowsType,
+	in float ShadowMaxNormalBias)
+{
+	float3 DirectLighting = (float3)0;
+
+	for (uint LightIndex = 0; LightIndex < RaytracingLightsDataPacked.Count; LightIndex++)
+	{
+		uint Lit = 1;
+
+		int LightProfileIndex = -1;
+		int RectLightTextureIndex = -1;
+		uint LightType = 0;
+		FDeferredLightData LightData = GetRayTracingDeferredLightData(LightIndex, LightProfileIndex, RectLightTextureIndex, LightType);
+
+		float3 ShadowRayDirection;
+		// ToLight should not be normalized because its length is used to compute the shadow ray TMax 
+		float3 ToLight = LightData.Position - WorldPosition;
+		float LightMask = 1.0;
+
+		if (LightType == LIGHT_TYPE_DIRECTIONAL)
+		{
+			ShadowRayDirection = LightData.Direction;
+			ToLight = LightData.Direction * 100000.0f;
+		}
+		else
+		{
+			LightMask = GetLocalLightAttenuation(WorldPosition, LightData, ToLight, ShadowRayDirection);
+
+			// Skip the light sample that does not contribute anything due to attenuation.
+			if (LightMask <= 0.0)
+			{
+				Lit = 0;
+			}
+		}
+
+		// Skip the light sample pointing backwards
+		if (dot(Payload.GetWorldNormal(), normalize(ToLight)) <= 0)
+		{
+			Lit = 0;
+		}
+
+		if (WaveActiveAllTrue(Lit == 0))
+		{
+			continue;
+		}
+
+		bool IsLightVisible = false;
+		if (ReflectedShadowsType == 2)
+		{
+			ShadowRayDirection = SampleAreaLightDirection(LightData, WorldPosition, Payload.GetWorldNormal(), LightType, RandSequence);
+		}
+
+		RayDesc ShadowRay;
+		ShadowRay.Origin = WorldPosition;
+		ShadowRay.Direction = ShadowRayDirection;
+		ShadowRay.TMin = 1e-4f;
+		ShadowRay.TMax = (Lit && ReflectedShadowsType!=0) ? length(ToLight) : -ShadowRay.TMin;
+		ApplyPositionBias(ShadowRay, Payload.GetWorldNormal(), ShadowMaxNormalBias);
+
+		uint RayFlags = RAY_FLAG_ACCEPT_FIRST_HIT_AND_END_SEARCH | RAY_FLAG_SKIP_CLOSEST_HIT_SHADER;
+		const uint InstanceInclusionMask = RAY_TRACING_MASK_SHADOW;
+
+		#if !ENABLE_TWO_SIDED_GEOMETRY
+		RayFlags |= RAY_FLAG_CULL_BACK_FACING_TRIANGLES;
+		#endif // !ENABLE_TWO_SIDED_GEOMETRY
+
+		TraceVisibilityRayPacked(
+			Payload,
+			TLAS,
+			RayFlags,
+			InstanceInclusionMask,
+			ShadowRay);
+		IsLightVisible = Payload.IsMiss() && Lit;
+
+		if (!IsLightVisible) continue;
+
+		FRectTexture RectTexture = GetRayTracingRectTextureData();
+
+		#if USE_SOURCE_TEXTURE_ARRAY
+		RectTexture.SourceTextureIndex = RectLightTextureIndex;
+		#endif // USE_SOURCE_TEXTURE_ARRAY
+
+		float LightProfileMultiplier = 1.0;
+
+		if (LightProfileIndex >= 0)
+		{
+			LightProfileMultiplier = ComputeRayTracingLightProfileMultiplier(WorldPosition, LightData.Position, LightData.Direction, LightProfileIndex);
+		}
+
+		float SurfaceShadow = 1.0f;
+		float4 LightAttenuation = 1.0f;
+
+		FGBufferData GBufferData = GetGBufferDataFromPayload(Payload);
+
+		float3 LightContribution = GetDynamicLighting(WorldPosition, ViewDirection, GBufferData, 1.0, GBufferData.ShadingModelID, LightData, LightAttenuation, 0.5, uint2(0, 0), RectTexture, SurfaceShadow).xyz;
+		DirectLighting += LightContribution * LightProfileMultiplier;
+	}
+
+	return DirectLighting;
+}
+
 void ComputeBottomLayerMaterialProperties(RayDesc Ray, inout FMaterialClosestHitPayload Payload)
 {
 	if (Payload.ShadingModelID == SHADINGMODELID_CLEAR_COAT)
@@ -474,8 +653,9 @@ void ComputeBottomLayerMaterialProperties(RayDesc Ray, inout FMaterialClosestHit
 }
 
 void AccumulateResults(
-	in FMaterialClosestHitPayload Payload,
-	in RayDesc Ray,
+	inout FPackedMaterialClosestHitPayload Payload,
+	in float3 WorldPosition,
+	in float3 ViewDirection,
 	in RaytracingAccelerationStructure TLAS,
 	in RandomSequence RandSequence,
 	in float ShadowMaxNormalBias,
@@ -485,35 +665,35 @@ void AccumulateResults(
 	inout FRayCone RayCone,
 	inout float3 Radiance)
 {
-	if (Payload.IsHit())
+	if (Payload.IsMiss())
 	{
-		Payload.DiffuseColor = Payload.BaseColor - Payload.BaseColor * Payload.Metallic;
-		Payload.SpecularColor = ComputeF0(Payload.Specular, Payload.BaseColor, Payload.Metallic);
+		return;
+	}
 
-		if (ShouldDoDirectLighting)
-		{
-			uint2 LightCullMask = (uint2)0;
-			const float3 WorldPosition = Ray.Origin + Payload.HitT * Ray.Direction;
-			LightCullMask = CullDirectLighting(Ray, RayCone, WorldPosition, Payload.WorldNormal, RandSequence, ReflectedShadowsType > 0, ShadowMaxNormalBias);
+	if (ShouldDoDirectLighting)
+	{
+		// Save and restore original payload HitT, as it's modified during shadow ray tracing
+		float OldHitT = Payload.HitT;
+	
+		float3 DirectLighting = 0;
 			
-			float3 DirectLighting = 0;
-				
-			if (Payload.ShadingModelID != SHADINGMODELID_UNLIT)
-			{
-				DirectLighting = ComputeDirectLightingCulled(LightCullMask, Ray, RayCone, TLAS, Payload, RandSequence, ReflectedShadowsType, ShadowMaxNormalBias);
-			}
-
-			// Transform NaNs to black, transform negative colors to black.
-			DirectLighting = -min(-DirectLighting, float3(0, 0, 0));
-			Radiance += DirectLighting;
-		}
-
-		if (ShouldDoEmissiveAndIndirectLighting)
+		if (Payload.GetShadingModelID() != SHADINGMODELID_UNLIT)
 		{
-			// Emissive & indirect contribution
-			Radiance += Payload.Radiance;
-			Radiance += Payload.DiffuseColor * Payload.IndirectIrradiance;
+			DirectLighting = ComputeDirectLightingMonolithic(WorldPosition, ViewDirection, RayCone, TLAS, Payload, RandSequence, ReflectedShadowsType, ShadowMaxNormalBias);
 		}
+
+		// Transform NaNs to black, transform negative colors to black.
+		DirectLighting = -min(-DirectLighting, float3(0, 0, 0));
+		Radiance += DirectLighting;
+
+		Payload.HitT = OldHitT;
+	}
+
+	if (ShouldDoEmissiveAndIndirectLighting)
+	{
+		// Emissive & indirect contribution
+		Radiance += Payload.GetRadiance();
+		Radiance += Payload.GetDiffuseColor() * Payload.GetIndirectIrradiance();
 	}
 }
 
@@ -530,16 +710,22 @@ FMaterialClosestHitPayload TraceRayAndAccumulateResults(
 	inout FRayCone RayCone,
 	inout float3 Radiance)
 {
-	FMaterialClosestHitPayload Payload = TraceMaterialRay(
+	FPackedMaterialClosestHitPayload Payload = (FPackedMaterialClosestHitPayload)0;
+	TraceMaterialRayPacked(
+		Payload,
 		TLAS,
 		RayFlags,
 		InstanceInclusionMask,
 		Ray,
 		RayCone);
 
+	float3 WorldPosition = Ray.Origin + Ray.Direction * Payload.HitT;
+	float3 ViewDirection = Ray.Direction;
+
 	AccumulateResults(
 		Payload,
-		Ray,
+		WorldPosition,
+		ViewDirection,
 		TLAS,
 		RandSequence,
 		ShadowMaxNormalBias,
@@ -549,7 +735,7 @@ FMaterialClosestHitPayload TraceRayAndAccumulateResults(
 		RayCone,
 		Radiance);
 
-	return Payload;
+	return UnpackRayTracingPayload(Payload, Ray);
 }
 
 FMaterialClosestHitPayload TraceRayAndAccumulateBottomLayerResults(
@@ -574,9 +760,14 @@ FMaterialClosestHitPayload TraceRayAndAccumulateBottomLayerResults(
 
 	ComputeBottomLayerMaterialProperties(Ray, BottomLayerPayload);
 
+	float3 WorldPosition = Ray.Origin + Ray.Direction * BottomLayerPayload.HitT;
+	float3 ViewDirection = Ray.Direction;
+
+	FPackedMaterialClosestHitPayload PackedBottomLayerPayload = PackRayTracingPayload(BottomLayerPayload, RayCone);
 	AccumulateResults(
-		BottomLayerPayload,
-		Ray,
+		PackedBottomLayerPayload,
+		WorldPosition,
+		ViewDirection,
 		TLAS,
 		RandSequence,
 		ShadowMaxNormalBias,
@@ -586,5 +777,5 @@ FMaterialClosestHitPayload TraceRayAndAccumulateBottomLayerResults(
 		RayCone,
 		Radiance);
 
-	return BottomLayerPayload;
+	return UnpackRayTracingPayload(PackedBottomLayerPayload, Ray);
 }
diff --git a/Engine/Shaders/Private/RayTracing/RayTracingReflections.usf b/Engine/Shaders/Private/RayTracing/RayTracingReflections.usf
--- a/Engine/Shaders/Private/RayTracing/RayTracingReflections.usf
+++ b/Engine/Shaders/Private/RayTracing/RayTracingReflections.usf
@@ -214,7 +214,7 @@ RAY_TRACING_ENTRY_RAYGEN(RayTracingReflectionsRGS)
 	uint2 DispatchThreadId = DispatchRaysIndex().xy + uint2(RenderTileOffsetX, RenderTileOffsetY);
 #endif
 
-	DispatchThreadId += View.ViewRectMin;
+	DispatchThreadId += (uint2)View.ViewRectMin;
 
 	uint2 PixelCoord = GetPixelCoord(DispatchThreadId, UpscaleFactor); // GBuffer pixel coordinates (full resolution)
 
@@ -424,7 +424,7 @@ RAY_TRACING_ENTRY_RAYGEN(RayTracingReflectionsRGS)
 
 				// Trace the top layer reflection
 				float3 TopLayerRadiance = float3(0, 0, 0);
-				FMaterialClosestHitPayload Payload = (FMaterialClosestHitPayload)0;
+				FPackedMaterialClosestHitPayload Payload = (FPackedMaterialClosestHitPayload)0;
 
 				if (DIM_HYBRID && DIM_DEFERRED_MATERIAL_MODE == DEFERRED_MATERIAL_MODE_SHADE 
 					&& DeferredMaterialPayload.SortKey == RAY_TRACING_DEFERRED_MATERIAL_KEY_SCREEN_SPACE)
@@ -438,17 +438,17 @@ RAY_TRACING_ENTRY_RAYGEN(RayTracingReflectionsRGS)
 					//   FScreenSpaceData HitScreenSpaceData = GetScreenSpaceData(RayHitUV);
 					//   Payload = MaterialPayloadFromScreenSpaceData(HitScreenSpaceData);
 
-					Payload.WorldPos = RayHitWorldPos;
 					Payload.HitT = DeferredMaterialPayload.HitT;
 
-					TopLayerRadiance += SceneColor.SampleLevel(GlobalPointClampedSampler, RayHitUV, 0) / View.PreExposure;
+					TopLayerRadiance += SceneColor.SampleLevel(GlobalPointClampedSampler, RayHitUV, 0).rgb / View.PreExposure;
 				}
 				else
 				{
 					const uint RayFlags = 0;
 					const uint InstanceInclusionMask = RAY_TRACING_MASK_OPAQUE;
 
-					Payload = TraceMaterialRay(
+					TraceMaterialRayPacked(
+						Payload,
 						TLAS,
 						RayFlags,
 						InstanceInclusionMask,
@@ -466,18 +466,20 @@ RAY_TRACING_ENTRY_RAYGEN(RayTracingReflectionsRGS)
 						Transmission = Transmit(TransmissionRay, RandSequence, TopLayerRadiance);
 					}
 
-					RayDesc LightingRay = TopLayerRay;
+					float3 RayHitWorldPos = TopLayerRay.Origin + Payload.HitT * TopLayerRay.Direction;
+					float3 LightingViewDirection = TopLayerRay.Direction;
 					if (DIM_HYBRID)
 					{
 						// Enforce consistent eye vector for computing lighting between screen space and ray traced reflections when using hybrid mode.
 						// This ensures that specular lighting matches and there is no visual discontinuity, even though it is wrong.
-						LightingRay.Direction = normalize(Payload.WorldPos - View.WorldCameraOrigin);
+						LightingViewDirection = normalize(RayHitWorldPos - View.WorldCameraOrigin);
 					}
 
 					float3 Radiance = 0.0;
 					AccumulateResults(
 						Payload,
-						LightingRay,
+						RayHitWorldPos,
+						LightingViewDirection,
 						TLAS,
 						RandSequence,
 						ReflectionMaxNormalBias,
@@ -504,7 +506,8 @@ RAY_TRACING_ENTRY_RAYGEN(RayTracingReflectionsRGS)
 					}
 
 					bApplyHeightFog = HeightFog > 0;
-					OriginToCollider = Payload.WorldPos - TopLayerRay.Origin;
+					float3 RayHitWorldPos = TopLayerRay.Origin + Payload.HitT * TopLayerRay.Direction;
+					OriginToCollider = RayHitWorldPos - TopLayerRay.Origin;
 				}
 				else if (bAllowSkySampling && ReflectionStruct.SkyLightParameters.y > 0)
 				{
@@ -681,10 +684,10 @@ RAY_TRACING_ENTRY_RAYGEN(RayTracingReflectionsRGS)
 				IncidentDirection = TopLayerRay.Direction;
 				WorldPosition = TopLayerRay.Origin + Payload.HitT * TopLayerRay.Direction;
 
-				ShadingModelID = Payload.ShadingModelID;
-				TopLayerWorldNormal = Payload.WorldNormal;
-				TopLayerRoughness = Payload.Roughness;
-				TopLayerSpecularColor = Payload.SpecularColor;
+				ShadingModelID = Payload.GetShadingModelID();
+				TopLayerWorldNormal = Payload.GetWorldNormal();
+				TopLayerRoughness = Payload.GetRoughness();
+				TopLayerSpecularColor = Payload.GetSpecularColor();
 
 				if (ShadingModelID == SHADINGMODELID_CLEAR_COAT)
 				{
@@ -696,8 +699,9 @@ RAY_TRACING_ENTRY_RAYGEN(RayTracingReflectionsRGS)
 					BottomLayerWorldNormal = OctahedronToUnitVector(oct1);
 					#endif
 
-					ClearCoat = Payload.CustomData.x;				// Clear coat weight
-					TopLayerRoughness = Payload.CustomData.y;		// Clear coat roughness
+					ClearCoat = Payload.GetCustomData().x;				// Clear coat weight
+					TopLayerRoughness = Payload.GetCustomData().y;		// Clear coat roughness
+
 					TopLayerSpecularColor = ClearCoatSpecularColor;	// Hard coded top layer specular color
 				}
 
