Redo the landscape fix that was overwritten

diff --git a/Engine/Shaders/Private/LandscapeVertexFactory.ush b/Engine/Shaders/Private/LandscapeVertexFactory.ush
--- a/Engine/Shaders/Private/LandscapeVertexFactory.ush
+++ b/Engine/Shaders/Private/LandscapeVertexFactory.ush
@@ -584,8 +584,9 @@ void GetShadowMapCoordinate(FVertexFactoryInterpolantsVSToPS Interpolants, out f
 }
 #endif
 
-float CalcLOD(float2 xy, float2 Subsection)
+float CalcLOD(float2 xyLocalToSubsection, float2 Subsection)
 {
+	float2 xy = (xyLocalToSubsection + Subsection) / LandscapeParameters.NumSubsections;
 	// LOD value based on Barycentric coordinate
 	float4 L0 = float4(xy.y, xy.x, (1 - xy.x), (1 - xy.y)) * 2;
 
