Remove sample count cap.

diff --git a/Engine/Shaders/Private/HairStrands/HairStrandsEnvironmentLighting.usf b/Engine/Shaders/Private/HairStrands/HairStrandsEnvironmentLighting.usf
--- a/Engine/Shaders/Private/HairStrands/HairStrandsEnvironmentLighting.usf
+++ b/Engine/Shaders/Private/HairStrands/HairStrandsEnvironmentLighting.usf
@@ -166,6 +166,9 @@ float3 ReflectionEnvironment(FGBufferData GBuffer, float2 ScreenPosition, float2
 			float2(0.19984126, 0.78641367),
 			float2(0.14383161, -0.14100790)
 		};
+		const uint SampleCount = clamp(MultipleScatterSampleCount, 1, 16);
+		#else
+		const uint SampleCount = max(MultipleScatterSampleCount, 1);
 		#endif
 
 		float2 RandOffset = 0;
@@ -177,7 +180,6 @@ float3 ReflectionEnvironment(FGBufferData GBuffer, float2 ScreenPosition, float2
 		}
 
 		float3 AvgR = 0;
-		const uint SampleCount = clamp(MultipleScatterSampleCount, 1, 16);
 		for (uint i = 0; i < SampleCount; ++i)
 		{
 		#if PERMUTATION_SAMPLESET == 0
