Decorrelate sample to remove structured/aliasing artifact when tracing against the density volume

diff --git a/Engine/Shaders/Private/HairStrands/HairStrandsEnvironmentLighting.usf b/Engine/Shaders/Private/HairStrands/HairStrandsEnvironmentLighting.usf
--- a/Engine/Shaders/Private/HairStrands/HairStrandsEnvironmentLighting.usf
+++ b/Engine/Shaders/Private/HairStrands/HairStrandsEnvironmentLighting.usf
@@ -115,7 +115,7 @@ float InternalComputeHairCountVoxel(float3 WorldPosition, float3 LightDirection,
 //////////////////////////////////////////////////////////////////////////////////////////////////////////////////////////
 //////////////////////////////////////////////////////////////////////////////////////////////////////////////////////////
 
-float3 ReflectionEnvironment(FGBufferData GBuffer, float2 ScreenPosition, float2 SvPosition)
+float3 ReflectionEnvironment(FGBufferData GBuffer, float2 ScreenPosition, float2 SvPosition, uint PrimitiveId)
 {
 #if USE_PREEXPOSURE
 	const float PreExposure = View.PreExposure;
@@ -128,9 +128,10 @@ float3 ReflectionEnvironment(FGBufferData GBuffer, float2 ScreenPosition, float2
 	const uint2  PixelPosition	= SvPosition - ResolvedView.ViewRectMin.xy;
 	const float  PixelSceneDepth= GBuffer.Depth;
 
+	const uint PrimitiveIdMod8 = (PrimitiveId % 8);
 	const float3 VoxelSize = (Voxel_MaxAABB - Voxel_MinAABB) / Voxel_Resolution;
 	const float  VoxelMaxSize = max(VoxelSize.x, max(VoxelSize.y, VoxelSize.z));
-	const float3 VoxelOffset = VoxelMaxSize * GetHairVoxelJitter(SvPosition, View.StateFrameIndexMod8);
+	const float3 VoxelOffset = VoxelMaxSize * GetHairVoxelJitter(SvPosition, View.StateFrameIndexMod8 + PrimitiveIdMod8);
 
 	const float3 V = -CameraToPixel;
 	const float3 T = GBuffer.WorldNormal;
@@ -175,8 +176,10 @@ float3 ReflectionEnvironment(FGBufferData GBuffer, float2 ScreenPosition, float2
 		if (JitterSphericalIntegration > 0)
 		{
 			RandOffset += float2(
-				InterleavedGradientNoise(SvPosition, View.StateFrameIndexMod8 * 11),
-				InterleavedGradientNoise(SvPosition, View.StateFrameIndexMod8 * 7));
+				InterleavedGradientNoise(SvPosition, PrimitiveIdMod8 * 11),
+				InterleavedGradientNoise(SvPosition, PrimitiveIdMod8 * 7));
+			//	InterleavedGradientNoise(SvPosition, View.StateFrameIndexMod8 * 11),
+			//	InterleavedGradientNoise(SvPosition, View.StateFrameIndexMod8 * 7));
 		}
 
 		float3 AvgR = 0;
@@ -190,8 +193,8 @@ float3 ReflectionEnvironment(FGBufferData GBuffer, float2 ScreenPosition, float2
 			const float3 SampleL = UniformSampleSphere(frac(Rand2)).xyz;
 
 			// Depth bias
-			// Origin is shifted voxels away towards the light + a constant bias of the size of the voxel
-			const float3 SampleWorldPosition = WorldPosition + VoxelOffset /*+ VoxelMaxSize*Voxel_DepthBiasScale*SampleL*/;
+			const float3 SampleVoxelOffset = VoxelMaxSize * GetHairVoxelJitter(SvPosition, PrimitiveIdMod8 + i);
+			const float3 SampleWorldPosition = WorldPosition + SampleVoxelOffset;
 
 			// Compute the number of hair count between light & shading point
 			const float SinLightAngle = dot(SampleL, T);
@@ -468,7 +471,7 @@ void MainCS(uint3 DispatchThreadId : SV_DispatchThreadID)
 	if (bIsClusterCompatible)
 	{
 		const float SampleCoverage = saturate(From8bitCoverage(SampleVisData.Coverage8bit));
-		const float3 SkyLighting = ReflectionEnvironment(GBuffer, ScreenPosition, SvPosition);
+		const float3 SkyLighting = ReflectionEnvironment(GBuffer, ScreenPosition, SvPosition, SampleVisData.PrimitiveId);
 
 		FLightAccumulator LightAccumulator = (FLightAccumulator)0;
 		LightAccumulator_Add(LightAccumulator, SkyLighting, SkyLighting, 1.0f, false);
