Adding support for Specular in hair PPLL pass

diff --git a/Engine/Shaders/Private/HairStrands/HairStrandsVisibilityCommon.ush b/Engine/Shaders/Private/HairStrands/HairStrandsVisibilityCommon.ush
--- a/Engine/Shaders/Private/HairStrands/HairStrandsVisibilityCommon.ush
+++ b/Engine/Shaders/Private/HairStrands/HairStrandsVisibilityCommon.ush
@@ -295,17 +295,19 @@ FNodeData ConvertToNodeData(float Depth, FVisibilityData VisibilityData, float3
  //////////
  // PPLL
 
-// Footprint: 160bits | 20bytes
+// Footprint: 192bits | 20bytes
  struct FPPLLNodeData
  {
+	 // (backloged)TODO try f16(depth) f16(Specular) and make sure depth does not lead to artefact
 	 float Depth;
+	 float Specular;
 	 uint PrimitiveID_ClusterID;
 	 uint Tangent_Coverage8bit;
 	 uint BaseColor_Roughness;
 	 uint NextNodeIndex;
  };
 
- FPPLLNodeData ConvertToPPLLNodeData(float Depth, FVisibilityData VisibilityData, float3 BaseColor, float Roughness, uint NextNodeIndex)
+ FPPLLNodeData ConvertToPPLLNodeData(float Depth, FVisibilityData VisibilityData, float3 BaseColor, float Roughness, float Specular, uint NextNodeIndex)
  {
 	 // 4x16bits
 	 const uint4 EncodedVisibilityData = EncodeVisibilityData(VisibilityData);
@@ -315,6 +317,7 @@ FNodeData ConvertToNodeData(float Depth, FVisibilityData VisibilityData, float3
 	 Out.Tangent_Coverage8bit = (0xFFFF & EncodedVisibilityData.x) | ((0xFFFF & EncodedVisibilityData.y) << 16);	// 32bits uint
 	 Out.PrimitiveID_ClusterID = (0xFFFF & EncodedVisibilityData.z) | ((0xFFFF & EncodedVisibilityData.w) << 16);	// 32bits uint
 	 Out.BaseColor_Roughness = float4ToUint(float4(sqrt(BaseColor), Roughness));									// 32bits uint
+	 Out.Specular = Specular;																						// 32 bits float
 	 Out.NextNodeIndex = NextNodeIndex;																				// 32bits next node
 
 	 return Out;
@@ -327,8 +330,7 @@ FNodeData ConvertToNodeData(float Depth, FVisibilityData VisibilityData, float3
 	 Out.PrimitiveID_ClusterID = In.PrimitiveID_ClusterID;
 	 Out.Tangent_Coverage8bit = In.Tangent_Coverage8bit;
 	 Out.BaseColor_Roughness = In.BaseColor_Roughness;
-	 Out.Specular = 0.5f;		// TODO
-
+	 Out.Specular = In.Specular;
 	 return Out;
  }
 
diff --git a/Engine/Shaders/Private/HairStrands/HairStrandsVisibilityPS.usf b/Engine/Shaders/Private/HairStrands/HairStrandsVisibilityPS.usf
--- a/Engine/Shaders/Private/HairStrands/HairStrandsVisibilityPS.usf
+++ b/Engine/Shaders/Private/HairStrands/HairStrandsVisibilityPS.usf
@@ -156,7 +156,7 @@ void MainVisibility(
 		// If we can append new node, add it and make it point to the node previously in head of list
 		uint NextNodeIndex;
 		InterlockedExchange(PPLLNodeIndex[PixelCoord], NodeIndex, NextNodeIndex);
-		const FPPLLNodeData Node = ConvertToPPLLNodeData(SvPosition.z, VisibilityData, BaseColor, Roughness, NextNodeIndex);
+		const FPPLLNodeData Node = ConvertToPPLLNodeData(SvPosition.z, VisibilityData, BaseColor, Roughness, Specular, NextNodeIndex);
 		PPLLNodes[NodeIndex] = Node;
 	}
 #endif
diff --git a/Engine/Source/Runtime/Renderer/Private/HairStrands/HairStrandsVisibility.cpp b/Engine/Source/Runtime/Renderer/Private/HairStrands/HairStrandsVisibility.cpp
--- a/Engine/Source/Runtime/Renderer/Private/HairStrands/HairStrandsVisibility.cpp
+++ b/Engine/Source/Runtime/Renderer/Private/HairStrands/HairStrandsVisibility.cpp
@@ -1022,17 +1022,18 @@ static void AddHairVisibilityPPLLPass(
 
 	const uint32 MaxNodeCount = GetMaxNodePerPixel(Resolution);
 	{
-		// Example: 20bytes * 8spp = 160bytes per pixel = 316.4Mb @ 1080p
-		struct NodeData
+		// Example: 24bytes * 8spp = 192bytes per pixel = 379Mb @ 1080p
+		struct PPLLNodeData
 		{
 			uint32 Depth;
 			uint32 PrimitiveId_ClusterId;
 			uint32 Tangent_Coverage;
 			uint32 BaseColor_Roughness;
+			uint32 Specular;
 			uint32 NextNodeIndex;
 		};
 
-		VisibilityPPLLNodeData = GraphBuilder.CreateBuffer(FRDGBufferDesc::CreateStructuredDesc(sizeof(NodeData), MaxNodeCount), TEXT("HairVisibilityPPLLNodeData"));
+		VisibilityPPLLNodeData = GraphBuilder.CreateBuffer(FRDGBufferDesc::CreateStructuredDesc(sizeof(PPLLNodeData), MaxNodeCount), TEXT("HairVisibilityPPLLNodeData"));
 	}
 	AddClearUAVPass(GraphBuilder, RDG_EVENT_NAME("ClearHairVisibilityPPLLCounter"), VisibilityPPLLNodeCounter, 0);
 	AddClearUAVPass(GraphBuilder, RDG_EVENT_NAME("ClearHairVisibilityPPLLNodeIndex"), VisibilityPPLLNodeIndex, 0xFFFFFFFF);
