Fixes signal array size mess by explicitely setting the size in each shader permutation of the denoiser.

diff --git a/Engine/Shaders/Private/ScreenSpaceDenoise/SSDCommon.ush b/Engine/Shaders/Private/ScreenSpaceDenoise/SSDCommon.ush
--- a/Engine/Shaders/Private/ScreenSpaceDenoise/SSDCommon.ush
+++ b/Engine/Shaders/Private/ScreenSpaceDenoise/SSDCommon.ush
@@ -27,15 +27,6 @@
 
 //------------------------------------------------------- SHARED CONSTANTS DEPENDING ON CONFIG OF SHADER
 
-// Maximum number of signal that can be denoised at same time.
-#ifndef MAX_SIGNAL_BATCH_SIZE
-	#define MAX_SIGNAL_BATCH_SIZE 1
-#endif
-
-// Maximum number of signal that can be multiplexed into single output.
-#define MAX_SIGNAL_MULTIPLEXING (MAX_SIGNAL_BATCH_SIZE * 2)
-
-
 #if CONFIG_SIGNAL_PROCESSING == SIGNAL_PROCESSING_REFLECTIONS
 	#define TARGETED_SAMPLE_COUNT 128
 
diff --git a/Engine/Shaders/Private/ScreenSpaceDenoise/SSDSignalAccumulatorArray.ush b/Engine/Shaders/Private/ScreenSpaceDenoise/SSDSignalAccumulatorArray.ush
--- a/Engine/Shaders/Private/ScreenSpaceDenoise/SSDSignalAccumulatorArray.ush
+++ b/Engine/Shaders/Private/ScreenSpaceDenoise/SSDSignalAccumulatorArray.ush
@@ -10,7 +10,7 @@
 /** Structure that is just an array of accumulators, to work arround shader compiler bug on "out" of arrays. */
 struct FSSDSignalAccumulatorArray
 {
-	FSSDSignalAccumulator Array[MAX_SIGNAL_MULTIPLEXING];
+	FSSDSignalAccumulator Array[SIGNAL_ARRAY_SIZE];
 };
 
 /** Accumulator array that is compressed to have lower VGPR footprint. */
@@ -24,8 +24,8 @@ FSSDSignalAccumulatorArray CreateSignalAccumulatorArray()
 {
 	FSSDSignalAccumulatorArray Accumulators;
 	
-	UNROLL_N(MAX_SIGNAL_MULTIPLEXING)
-	for (uint i = 0; i < MAX_SIGNAL_MULTIPLEXING; i++)
+	UNROLL_N(SIGNAL_ARRAY_SIZE)
+	for (uint i = 0; i < SIGNAL_ARRAY_SIZE; i++)
 	{
 		Accumulators.Array[i] = CreateSignalAccumulator();
 	}
diff --git a/Engine/Shaders/Private/ScreenSpaceDenoise/SSDSignalArray.ush b/Engine/Shaders/Private/ScreenSpaceDenoise/SSDSignalArray.ush
--- a/Engine/Shaders/Private/ScreenSpaceDenoise/SSDSignalArray.ush
+++ b/Engine/Shaders/Private/ScreenSpaceDenoise/SSDSignalArray.ush
@@ -5,132 +5,132 @@
 #include "SSDSignalCore.ush"
 
 
-FSSDSignalArray MulSignal(FSSDSignalArray Samples, float Scalar[MAX_SIGNAL_BATCH_SIZE])
+FSSDSignalArray MulSignalArray(FSSDSignalArray Samples, float Scalar[SIGNAL_ARRAY_SIZE])
 {
 	FSSDSignalArray OutSamples;
-	UNROLL_N(MAX_SIGNAL_BATCH_SIZE)
-	for (uint BatchedSignalId = 0; BatchedSignalId < MAX_SIGNAL_BATCH_SIZE; BatchedSignalId++)
+	UNROLL_N(SIGNAL_ARRAY_SIZE)
+	for (uint BatchedSignalId = 0; BatchedSignalId < SIGNAL_ARRAY_SIZE; BatchedSignalId++)
 	{
 		OutSamples.Array[BatchedSignalId] = MulSignal(Samples.Array[BatchedSignalId], Scalar[BatchedSignalId]);
 	}
 	return OutSamples;
 }
 
-FSSDSignalArray AddSignal(FSSDSignalArray SamplesA, FSSDSignalArray SamplesB)
+FSSDSignalArray AddSignalArray(FSSDSignalArray SamplesA, FSSDSignalArray SamplesB)
 {
 	FSSDSignalArray OutSamples;
-	UNROLL_N(MAX_SIGNAL_BATCH_SIZE)
-	for (uint BatchedSignalId = 0; BatchedSignalId < MAX_SIGNAL_BATCH_SIZE; BatchedSignalId++)
+	UNROLL_N(SIGNAL_ARRAY_SIZE)
+	for (uint BatchedSignalId = 0; BatchedSignalId < SIGNAL_ARRAY_SIZE; BatchedSignalId++)
 	{
 		OutSamples.Array[BatchedSignalId] = AddSignal(SamplesA.Array[BatchedSignalId], SamplesB.Array[BatchedSignalId]);
 	}
 	return OutSamples;
 }
 
-FSSDSignalArray MinusSignal(FSSDSignalArray Samples)
+FSSDSignalArray MinusSignalArray(FSSDSignalArray Samples)
 {
 	FSSDSignalArray OutSamples;
-	UNROLL_N(MAX_SIGNAL_BATCH_SIZE)
-	for (uint BatchedSignalId = 0; BatchedSignalId < MAX_SIGNAL_BATCH_SIZE; BatchedSignalId++)
+	UNROLL_N(SIGNAL_ARRAY_SIZE)
+	for (uint BatchedSignalId = 0; BatchedSignalId < SIGNAL_ARRAY_SIZE; BatchedSignalId++)
 	{
 		OutSamples.Array[BatchedSignalId] = MinusSignal(Samples.Array[BatchedSignalId]);
 	}
 	return OutSamples;
 }
 
-FSSDSignalArray SubtractSignal(FSSDSignalArray SamplesA, FSSDSignalArray SamplesB)
+FSSDSignalArray SubtractSignalArray(FSSDSignalArray SamplesA, FSSDSignalArray SamplesB)
 {
 	FSSDSignalArray OutSamples;
-	UNROLL_N(MAX_SIGNAL_BATCH_SIZE)
-	for (uint BatchedSignalId = 0; BatchedSignalId < MAX_SIGNAL_BATCH_SIZE; BatchedSignalId++)
+	UNROLL_N(SIGNAL_ARRAY_SIZE)
+	for (uint BatchedSignalId = 0; BatchedSignalId < SIGNAL_ARRAY_SIZE; BatchedSignalId++)
 	{
 		OutSamples.Array[BatchedSignalId] = SubtractSignal(SamplesA.Array[BatchedSignalId], SamplesB.Array[BatchedSignalId]);
 	}
 	return OutSamples;
 }
 
-FSSDSignalArray AbsSignal(FSSDSignalArray Samples)
+FSSDSignalArray AbsSignalArray(FSSDSignalArray Samples)
 {
 	FSSDSignalArray OutSamples;
-	UNROLL_N(MAX_SIGNAL_BATCH_SIZE)
-	for (uint BatchedSignalId = 0; BatchedSignalId < MAX_SIGNAL_BATCH_SIZE; BatchedSignalId++)
+	UNROLL_N(SIGNAL_ARRAY_SIZE)
+	for (uint BatchedSignalId = 0; BatchedSignalId < SIGNAL_ARRAY_SIZE; BatchedSignalId++)
 	{
 		OutSamples.Array[BatchedSignalId] = AbsSignal(Samples.Array[BatchedSignalId]);
 	}
 	return OutSamples;
 }
 
-FSSDSignalArray SqrtSignal(FSSDSignalArray Samples)
+FSSDSignalArray SqrtSignalArray(FSSDSignalArray Samples)
 {
 	FSSDSignalArray OutSamples;
-	UNROLL_N(MAX_SIGNAL_BATCH_SIZE)
-	for (uint BatchedSignalId = 0; BatchedSignalId < MAX_SIGNAL_BATCH_SIZE; BatchedSignalId++)
+	UNROLL_N(SIGNAL_ARRAY_SIZE)
+	for (uint BatchedSignalId = 0; BatchedSignalId < SIGNAL_ARRAY_SIZE; BatchedSignalId++)
 	{
 		OutSamples.Array[BatchedSignalId] = SqrtSignal(Samples.Array[BatchedSignalId]);
 	}
 	return OutSamples;
 }
 
-FSSDSignalArray PowerSignal(FSSDSignalArray Samples, float Exponent)
+FSSDSignalArray PowerSignalArray(FSSDSignalArray Samples, float Exponent)
 {
 	FSSDSignalArray OutSamples;
-	UNROLL_N(MAX_SIGNAL_BATCH_SIZE)
-	for (uint BatchedSignalId = 0; BatchedSignalId < MAX_SIGNAL_BATCH_SIZE; BatchedSignalId++)
+	UNROLL_N(SIGNAL_ARRAY_SIZE)
+	for (uint BatchedSignalId = 0; BatchedSignalId < SIGNAL_ARRAY_SIZE; BatchedSignalId++)
 	{
 		OutSamples.Array[BatchedSignalId] = PowerSignal(Samples.Array[BatchedSignalId], Exponent);
 	}
 	return OutSamples;
 }
 
-FSSDSignalArray MinSignal(FSSDSignalArray SamplesA, FSSDSignalArray SamplesB)
+FSSDSignalArray MinSignalArray(FSSDSignalArray SamplesA, FSSDSignalArray SamplesB)
 {
 	FSSDSignalArray OutSamples;
-	UNROLL_N(MAX_SIGNAL_BATCH_SIZE)
-	for (uint BatchedSignalId = 0; BatchedSignalId < MAX_SIGNAL_BATCH_SIZE; BatchedSignalId++)
+	UNROLL_N(SIGNAL_ARRAY_SIZE)
+	for (uint BatchedSignalId = 0; BatchedSignalId < SIGNAL_ARRAY_SIZE; BatchedSignalId++)
 	{
 		OutSamples.Array[BatchedSignalId] = MinSignal(SamplesA.Array[BatchedSignalId], SamplesB.Array[BatchedSignalId]);
 	}
 	return OutSamples;
 }
 
-FSSDSignalArray MaxSignal(FSSDSignalArray SamplesA, FSSDSignalArray SamplesB)
+FSSDSignalArray MaxSignalArray(FSSDSignalArray SamplesA, FSSDSignalArray SamplesB)
 {
 	FSSDSignalArray OutSamples;
-	UNROLL_N(MAX_SIGNAL_BATCH_SIZE)
-	for (uint BatchedSignalId = 0; BatchedSignalId < MAX_SIGNAL_BATCH_SIZE; BatchedSignalId++)
+	UNROLL_N(SIGNAL_ARRAY_SIZE)
+	for (uint BatchedSignalId = 0; BatchedSignalId < SIGNAL_ARRAY_SIZE; BatchedSignalId++)
 	{
 		OutSamples.Array[BatchedSignalId] = MaxSignal(SamplesA.Array[BatchedSignalId], SamplesB.Array[BatchedSignalId]);
 	}
 	return OutSamples;
 }
 
-FSSDSignalArray ClampSignal(FSSDSignalArray Samples, FSSDSignalArray SamplesMin, FSSDSignalArray SamplesMax)
+FSSDSignalArray ClampSignalArray(FSSDSignalArray Samples, FSSDSignalArray SamplesMin, FSSDSignalArray SamplesMax)
 {
 	FSSDSignalArray OutSamples;
-	UNROLL_N(MAX_SIGNAL_BATCH_SIZE)
-	for (uint BatchedSignalId = 0; BatchedSignalId < MAX_SIGNAL_BATCH_SIZE; BatchedSignalId++)
+	UNROLL_N(SIGNAL_ARRAY_SIZE)
+	for (uint BatchedSignalId = 0; BatchedSignalId < SIGNAL_ARRAY_SIZE; BatchedSignalId++)
 	{
 		OutSamples.Array[BatchedSignalId] = ClampSignal(Samples.Array[BatchedSignalId], SamplesMin.Array[BatchedSignalId], SamplesMax.Array[BatchedSignalId]);
 	}
 	return OutSamples;
 }
 
-FSSDSignalArray LerpSignal(FSSDSignalArray Samples0, FSSDSignalArray Samples1, float Interp[MAX_SIGNAL_BATCH_SIZE])
+FSSDSignalArray LerpSignalArray(FSSDSignalArray Samples0, FSSDSignalArray Samples1, float Interp[SIGNAL_ARRAY_SIZE])
 {
 	FSSDSignalArray OutSamples;
-	UNROLL_N(MAX_SIGNAL_BATCH_SIZE)
-	for (uint BatchedSignalId = 0; BatchedSignalId < MAX_SIGNAL_BATCH_SIZE; BatchedSignalId++)
+	UNROLL_N(SIGNAL_ARRAY_SIZE)
+	for (uint BatchedSignalId = 0; BatchedSignalId < SIGNAL_ARRAY_SIZE; BatchedSignalId++)
 	{
 		OutSamples.Array[BatchedSignalId] = LerpSignal(Samples0.Array[BatchedSignalId], Samples1.Array[BatchedSignalId], Interp[BatchedSignalId]);
 	}
 	return OutSamples;
 }
 
-FSSDSignalArray NormalizeToOneSample(FSSDSignalArray Samples)
+FSSDSignalArray NormalizeToOneSampleArray(FSSDSignalArray Samples)
 {
 	FSSDSignalArray OutSamples;
-	UNROLL_N(MAX_SIGNAL_BATCH_SIZE)
-	for (uint BatchedSignalId = 0; BatchedSignalId < MAX_SIGNAL_BATCH_SIZE; BatchedSignalId++)
+	UNROLL_N(SIGNAL_ARRAY_SIZE)
+	for (uint BatchedSignalId = 0; BatchedSignalId < SIGNAL_ARRAY_SIZE; BatchedSignalId++)
 	{
 		OutSamples.Array[BatchedSignalId] = NormalizeToOneSample(Samples.Array[BatchedSignalId]);
 	}
@@ -140,8 +140,8 @@ FSSDSignalArray NormalizeToOneSample(FSSDSignalArray Samples)
 FSSDSignalArray WaveBroadcastSignalArray(const FWaveBroadcastSettings BroadcastSettings, FSSDSignalArray Samples)
 {
 	FSSDSignalArray OutSamples;
-	UNROLL_N(MAX_SIGNAL_BATCH_SIZE)
-	for (uint BatchedSignalId = 0; BatchedSignalId < MAX_SIGNAL_BATCH_SIZE; BatchedSignalId++)
+	UNROLL_N(SIGNAL_ARRAY_SIZE)
+	for (uint BatchedSignalId = 0; BatchedSignalId < SIGNAL_ARRAY_SIZE; BatchedSignalId++)
 	{
 		OutSamples.Array[BatchedSignalId] = WaveBroadcastSignal(BroadcastSettings, Samples.Array[BatchedSignalId]);
 	}
diff --git a/Engine/Shaders/Private/ScreenSpaceDenoise/SSDSignalBufferEncoding.ush b/Engine/Shaders/Private/ScreenSpaceDenoise/SSDSignalBufferEncoding.ush
--- a/Engine/Shaders/Private/ScreenSpaceDenoise/SSDSignalBufferEncoding.ush
+++ b/Engine/Shaders/Private/ScreenSpaceDenoise/SSDSignalBufferEncoding.ush
@@ -125,8 +125,8 @@ void DecodeMultiplexedSignalsFromFloat4(
 	}
 	//else if (SignalBufferLayout == SIGNAL_BUFFER_LAYOUT_PENUMBRA_RECONSTRUCTION)
 	//{
-	//	UNROLL_N(MAX_SIGNAL_MULTIPLEXING)
-	//	for (uint MultiplexId = 0; MultiplexId < MAX_SIGNAL_MULTIPLEXING; MultiplexId++)
+	//	UNROLL_N(SIGNAL_ARRAY_SIZE)
+	//	for (uint MultiplexId = 0; MultiplexId < SIGNAL_ARRAY_SIZE; MultiplexId++)
 	//	{
 	//		float2 Channels = float2(RawSample[MultiplexId / 2][2 * (MultiplexId % 2)], RawSample[MultiplexId / 2][2 * (MultiplexId % 2) + 1]);
 	//
@@ -314,7 +314,7 @@ void DecodeMultiplexedSignalsFromFloat4(
 		OutSamples.Array[0].SceneColor = OutSamples.Array[0].SampleCount * RawSample[0];
 
 		// Temporal analysis of the history.
-		#if MAX_SIGNAL_MULTIPLEXING >= 3
+		#if SIGNAL_ARRAY_SIZE >= 3
 			OutSamples.Array[1].SampleCount = OutSamples.Array[0].SampleCount;
 			OutSamples.Array[1].SceneColor.x = OutSamples.Array[1].SampleCount * RawSample[1].g;
 			OutSamples.Array[2].SampleCount = OutSamples.Array[0].SampleCount;
@@ -402,7 +402,7 @@ void DecodeMultiplexedSignalsFromFloat4(
 		OutSamples.Array[0].SceneColor = OutSamples.Array[0].SampleCount * RawSample[0];
 		
 		// Temporal analysis.
-		#if MAX_SIGNAL_MULTIPLEXING >= 3
+		#if SIGNAL_ARRAY_SIZE >= 3
 			OutSamples.Array[1].SampleCount = OutSamples.Array[0].SampleCount;
 			OutSamples.Array[1].SceneColor.x = OutSamples.Array[1].SampleCount * RawSample[1].g;
 			OutSamples.Array[2].SampleCount = OutSamples.Array[0].SampleCount;
@@ -483,7 +483,7 @@ FSSDSignalArray DecodeMultiplexedSignals(
 /** Encode output signal sample. */
 void EncodeMultiplexedSignals(
 	const uint SignalBufferLayout, const uint MultiplexCount,
-	FSSDSignalSample Sample[MAX_SIGNAL_MULTIPLEXING],
+	FSSDSignalSample Sample[SIGNAL_ARRAY_SIZE],
 	out FSSDOutputRawSample OutRawSample[MAX_MULTIPLEXED_TEXTURES],
 	out uint OutBufferCount)
 #if CONFIG_SIGNAL_OUTPUT_TEXTURE_TYPE == SIGNAL_TEXTURE_TYPE_FLOAT4
@@ -662,7 +662,7 @@ void EncodeMultiplexedSignals(
 		OutBufferCount = 2;
 
 		// Temporal analysis.
-		#if MAX_SIGNAL_MULTIPLEXING >= 3
+		#if SIGNAL_ARRAY_SIZE >= 3
 		if (MultiplexCount == 3)
 		{
 			OutRawSample[1].g = Sample[1].SceneColor.x * NormalizationFactor;
@@ -710,7 +710,7 @@ void EncodeMultiplexedSignals(
 		OutBufferCount = 2;
 
 		// Temporal analysis.
-		#if MAX_SIGNAL_MULTIPLEXING >= 3
+		#if SIGNAL_ARRAY_SIZE >= 3
 		if (MultiplexCount == 3)
 		{
 			OutRawSample[1].g = Sample[1].SceneColor.x * NormalizationFactor;
@@ -822,7 +822,7 @@ void OutputMultiplexedSignal(
 	FSSDRWTexture2D OutputSignalBuffer2,
 	FSSDRWTexture2D OutputSignalBuffer3,
 	const uint SignalBufferLayout, const uint MultiplexCount,
-	const uint2 PixelPosition, FSSDSignalSample MultiplexedSamples[MAX_SIGNAL_MULTIPLEXING])
+	const uint2 PixelPosition, FSSDSignalSample MultiplexedSamples[SIGNAL_ARRAY_SIZE])
 {
 	// Encode the output signal.
 	FSSDOutputRawSample RawSample[MAX_MULTIPLEXED_TEXTURES];
diff --git a/Engine/Shaders/Private/ScreenSpaceDenoise/SSDSignalCore.ush b/Engine/Shaders/Private/ScreenSpaceDenoise/SSDSignalCore.ush
--- a/Engine/Shaders/Private/ScreenSpaceDenoise/SSDSignalCore.ush
+++ b/Engine/Shaders/Private/ScreenSpaceDenoise/SSDSignalCore.ush
@@ -17,6 +17,16 @@
 	#define COMPILE_SIGNAL_COLOR_ARRAY 0
 #endif
 
+// Maximum number of indepedendent signal domain that can be denoised at same time.
+#ifndef MAX_SIGNAL_BATCH_SIZE
+	#error Need to specify the number of signal domain being denoised.
+#endif
+
+// Number of sample that can be multiplexed into FSSDSignalSample.
+#ifndef SIGNAL_ARRAY_SIZE
+	#error Need to specify the size of the signal array.
+#endif
+
 
 //------------------------------------------------------- STRUCTURE
 
@@ -54,7 +64,7 @@ struct FSSDSignalSample
 /** Array of signal samples. Technically should be a typedef, but there is a bug in HLSL with the out operator of an array. */
 struct FSSDSignalArray
 {
-	FSSDSignalSample Array[MAX_SIGNAL_MULTIPLEXING];
+	FSSDSignalSample Array[SIGNAL_ARRAY_SIZE];
 };
 
 
@@ -101,8 +111,8 @@ FSSDSignalSample CreateSignalSample()
 FSSDSignalArray CreateSignalArrayFromScalarValue(float Scalar)
 {
 	FSSDSignalArray OutSamples;
-	UNROLL_N(MAX_SIGNAL_MULTIPLEXING)
-	for (uint BatchedSignalId = 0; BatchedSignalId < MAX_SIGNAL_MULTIPLEXING; BatchedSignalId++)
+	UNROLL_N(SIGNAL_ARRAY_SIZE)
+	for (uint BatchedSignalId = 0; BatchedSignalId < SIGNAL_ARRAY_SIZE; BatchedSignalId++)
 	{
 		OutSamples.Array[BatchedSignalId] = CreateSignalSampleFromScalarValue(Scalar);
 	}
diff --git a/Engine/Shaders/Private/ScreenSpaceDenoise/SSDSpatialAccumulation.usf b/Engine/Shaders/Private/ScreenSpaceDenoise/SSDSpatialAccumulation.usf
--- a/Engine/Shaders/Private/ScreenSpaceDenoise/SSDSpatialAccumulation.usf
+++ b/Engine/Shaders/Private/ScreenSpaceDenoise/SSDSpatialAccumulation.usf
@@ -40,11 +40,12 @@
 
 // Configures all the pass for each individual signals.
 #if CONFIG_SIGNAL_PROCESSING == SIGNAL_PROCESSING_SHADOW_VISIBILITY_MASK
+	#define MAX_SIGNAL_BATCH_SIZE CONFIG_SIGNAL_BATCH_SIZE
+	#define SIGNAL_ARRAY_SIZE CONFIG_SIGNAL_BATCH_SIZE
+
 	#define CONFIG_BILATERAL_PRESET BILATERAL_PRESET_MONOCHROMATIC_PENUMBRA
 	#define CONFIG_MULTIPLEXED_SIGNALS_PER_SIGNAL_DOMAIN 1
 
-	#define MAX_SIGNAL_BATCH_SIZE CONFIG_SIGNAL_BATCH_SIZE
-
 	#if DIM_STAGE == STAGE_RECONSTRUCTION
 		// Input and output layout.
 		#define CONFIG_SIGNAL_INPUT_LAYOUT  SIGNAL_BUFFER_LAYOUT_PENUMBRA_INJESTION_NSPP
@@ -77,16 +78,6 @@
 		#define CONFIG_MAX_WITH_REF_DISTANCE 1
 		#define CONFIG_OUTPUT_MODE          OUTPUT_MODE_DRB
 
-	#elif DIM_STAGE == STAGE_REJECTION_PRE_CONVOLUTION
-		// Input and output layout.
-		#define CONFIG_SIGNAL_INPUT_LAYOUT  SIGNAL_BUFFER_LAYOUT_PENUMBRA_RECONSTRUCTION
-		#define CONFIG_SIGNAL_OUTPUT_LAYOUT SIGNAL_BUFFER_LAYOUT_PENUMBRA_REJECTION
-
-		#define CONFIG_INPUT_TEXTURE_COUNT CONFIG_SIGNAL_BATCH_SIZE
-		#define CONFIG_OUTPUT_TEXTURE_COUNT CONFIG_SIGNAL_BATCH_SIZE
-
-		#define CONFIG_BILATERAL_DISTANCE_COMPUTATION SIGNAL_WORLD_FREQUENCY_PRECOMPUTED_BLURING_RADIUS
-
 	#elif DIM_STAGE == STAGE_POST_FILTERING
 		// Input and output layout.
 		#define CONFIG_SIGNAL_INPUT_LAYOUT  SIGNAL_BUFFER_LAYOUT_PENUMBRA_HISTORY
@@ -119,15 +110,17 @@
 	#endif
 
 #elif CONFIG_SIGNAL_PROCESSING == SIGNAL_PROCESSING_POLYCHROMATIC_PENUMBRA_HARMONIC
+	// Denoise diffuse and specular harmonics at the same time.
+	#define MAX_SIGNAL_BATCH_SIZE 2
+	#define SIGNAL_ARRAY_SIZE 2
+	#undef CONFIG_SIGNAL_BATCH_SIZE
+	#define CONFIG_SIGNAL_BATCH_SIZE 2
+
+	// Each harmonic requires input and output RGB.
 	#define COMPILE_SIGNAL_COLOR_ARRAY 2
 
 	#define CONFIG_BILATERAL_PRESET BILATERAL_PRESET_POLYCHROMATIC_PENUMBRA
 	#define CONFIG_MULTIPLEXED_SIGNALS_PER_SIGNAL_DOMAIN 1
-	#undef CONFIG_SIGNAL_BATCH_SIZE
-	#define CONFIG_SIGNAL_BATCH_SIZE 2
-
-	// Denoise diffuse and specular harmonics at the same time.
-	#define MAX_SIGNAL_BATCH_SIZE 2
 
 	// Any world distance depends on the harmonic being processed.
 	#define CONFIG_BILATERAL_DISTANCE_COMPUTATION SIGNAL_WORLD_FREQUENCY_HARMONIC
@@ -148,9 +141,14 @@
 	#endif
 
 #elif CONFIG_SIGNAL_PROCESSING == SIGNAL_PROCESSING_REFLECTIONS
+	// Denoise only specular.
+	#define MAX_SIGNAL_BATCH_SIZE 1
+
 	#define COMPILE_SIGNAL_COLOR 1
 
 	#if DIM_STAGE == STAGE_RECONSTRUCTION
+		#define SIGNAL_ARRAY_SIZE 1
+
 		// Input and output layout.
 		#define CONFIG_SIGNAL_INPUT_LAYOUT  SIGNAL_BUFFER_LAYOUT_REFLECTIONS_INPUT
 		#define CONFIG_SIGNAL_OUTPUT_LAYOUT SIGNAL_BUFFER_LAYOUT_REFLECTIONS_HISTORY
@@ -188,6 +186,8 @@
 		#define CONFIG_OUTPUT_TEXTURE_COUNT 2
 
 	#elif DIM_STAGE == STAGE_REJECTION_PRE_CONVOLUTION
+		#define SIGNAL_ARRAY_SIZE 2 // first and second momment to measure variance in temporal accumulation.
+
 		// Input and output layout.
 		#define CONFIG_SIGNAL_INPUT_LAYOUT  SIGNAL_BUFFER_LAYOUT_REFLECTIONS_HISTORY
 		#define CONFIG_SIGNAL_OUTPUT_LAYOUT SIGNAL_BUFFER_LAYOUT_REFLECTIONS_REJECTION
@@ -201,6 +201,8 @@
 		#define CONFIG_OUTPUT_TEXTURE_COUNT 3
 
 	#elif DIM_STAGE == STAGE_POST_FILTERING
+		#define SIGNAL_ARRAY_SIZE 1
+
 		// Input and output layout.
 		#define CONFIG_SIGNAL_INPUT_LAYOUT  SIGNAL_BUFFER_LAYOUT_REFLECTIONS_HISTORY
 		#define CONFIG_SIGNAL_OUTPUT_LAYOUT SIGNAL_BUFFER_LAYOUT_REFLECTIONS_HISTORY
@@ -221,9 +223,14 @@
 	#endif
 
 #elif CONFIG_SIGNAL_PROCESSING == SIGNAL_PROCESSING_AO
+	// Denoise only AO.
+	#define MAX_SIGNAL_BATCH_SIZE 1
+
 	#define CONFIG_BILATERAL_PRESET BILATERAL_PRESET_DIFFUSE
 
 	#if DIM_STAGE == STAGE_RECONSTRUCTION
+		#define SIGNAL_ARRAY_SIZE 1
+
 		// Input and output layout.
 		#define CONFIG_SIGNAL_INPUT_LAYOUT  SIGNAL_BUFFER_LAYOUT_AO_INPUT
 		#define CONFIG_SIGNAL_OUTPUT_LAYOUT SIGNAL_BUFFER_LAYOUT_AO_HISTORY
@@ -236,6 +243,8 @@
 		#define CONFIG_OUTPUT_MODE          OUTPUT_MODE_DRB
 
 	#elif DIM_STAGE == STAGE_REJECTION_PRE_CONVOLUTION
+		#define SIGNAL_ARRAY_SIZE 2 // first and second momment to measure variance in temporal accumulation.
+
 		// Input and output layout.
 		#define CONFIG_SIGNAL_INPUT_LAYOUT  SIGNAL_BUFFER_LAYOUT_AO_HISTORY
 		#define CONFIG_SIGNAL_OUTPUT_LAYOUT SIGNAL_BUFFER_LAYOUT_AO_REJECTION
@@ -246,6 +255,8 @@
 		#define CONFIG_BILATERAL_DISTANCE_COMPUTATION SIGNAL_WORLD_FREQUENCY_HIT_DISTANCE
 
 	#elif DIM_STAGE == STAGE_POST_FILTERING
+		#define SIGNAL_ARRAY_SIZE 1
+
 		// Input and output layout.
 		#define CONFIG_SIGNAL_INPUT_LAYOUT  SIGNAL_BUFFER_LAYOUT_AO_HISTORY
 		#define CONFIG_SIGNAL_OUTPUT_LAYOUT SIGNAL_BUFFER_LAYOUT_AO_HISTORY
@@ -263,6 +274,10 @@
 	#endif
 
 #elif CONFIG_SIGNAL_PROCESSING == SIGNAL_PROCESSING_DIFFUSE_INDIRECT_AND_AO
+	// Denoise diffuse and AO, but AO is FSSDSignalSample::MissCount.
+	#define MAX_SIGNAL_BATCH_SIZE 1
+	#define SIGNAL_ARRAY_SIZE 1
+
 	#define COMPILE_SIGNAL_COLOR 1
 
 	#define CONFIG_BILATERAL_PRESET BILATERAL_PRESET_DIFFUSE
@@ -303,6 +318,9 @@
 	#endif
 
 #elif CONFIG_SIGNAL_PROCESSING == SIGNAL_PROCESSING_DIFFUSE_SPHERICAL_HARMONIC
+	#define MAX_SIGNAL_BATCH_SIZE 1
+	#define SIGNAL_ARRAY_SIZE 1
+
 	#define COMPILE_SIGNAL_COLOR_SH 1
 
 	// Given it's a spherical harmonic that store directionality, only need position based rejection.
@@ -335,6 +353,10 @@
 	#endif
 
 #elif CONFIG_SIGNAL_PROCESSING == SIGNAL_PROCESSING_SSGI
+	// Denoise diffuse and AO, but AO is FSSDSignalSample::MissCount.
+	#define MAX_SIGNAL_BATCH_SIZE 1
+	#define SIGNAL_ARRAY_SIZE 1
+
 	#define COMPILE_SIGNAL_COLOR 1
 
 	#define CONFIG_BILATERAL_PRESET BILATERAL_PRESET_DIFFUSE
@@ -435,7 +457,7 @@
 
 /** The number of signal that should be processed per signal domain. */
 #ifndef CONFIG_MULTIPLEXED_SIGNALS_PER_SIGNAL_DOMAIN
-	#define CONFIG_MULTIPLEXED_SIGNALS_PER_SIGNAL_DOMAIN MAX_SIGNAL_MULTIPLEXING
+	#define CONFIG_MULTIPLEXED_SIGNALS_PER_SIGNAL_DOMAIN SIGNAL_ARRAY_SIZE
 #endif
 
 /** Selects how the world distance should be computed for bilateral rejection. */
@@ -655,7 +677,8 @@ void MainCS(
 		#if CONFIG_NORMALIZE_INPUT
 			FSSDSignalArray NormalizedRefSamples = RefSamples;
 		#else
-			FSSDSignalArray NormalizedRefSamples = NormalizeToOneSample(RefSamples);
+			// TODO(Denoiser): Decode twice instead.
+			FSSDSignalArray NormalizedRefSamples = NormalizeToOneSampleArray(RefSamples);
 		#endif
 	#endif
 
@@ -749,8 +772,8 @@ void MainCS(
 		KernelConfig.bComputeSampleColorSH = DIM_STAGE == STAGE_RECONSTRUCTION && DIM_MULTI_SPP == 0;
 		
 		{
-			UNROLL_N(MAX_SIGNAL_MULTIPLEXING)
-			for (uint MultiplexId = 0; MultiplexId < MAX_SIGNAL_MULTIPLEXING; MultiplexId++)
+			UNROLL_N(SIGNAL_ARRAY_SIZE)
+			for (uint MultiplexId = 0; MultiplexId < SIGNAL_ARRAY_SIZE; MultiplexId++)
 			{
 				KernelConfig.BufferColorSpace[MultiplexId] = CONFIG_INPUT_COLOR_SPACE;
 				KernelConfig.AccumulatorColorSpace[MultiplexId] = CONFIG_ACCUMULATION_COLOR_SPACE;
@@ -796,8 +819,8 @@ void MainCS(
 		{
 			KernelConfig.bMaxWithRefBilateralDistance = true;
 
-			UNROLL_N(MAX_SIGNAL_MULTIPLEXING)
-			for (uint MultiplexId = 0; MultiplexId < MAX_SIGNAL_MULTIPLEXING; MultiplexId++)
+			UNROLL_N(SIGNAL_ARRAY_SIZE)
+			for (uint MultiplexId = 0; MultiplexId < SIGNAL_ARRAY_SIZE; MultiplexId++)
 			{
 				if (KernelConfig.BilateralDistanceComputation == SIGNAL_WORLD_FREQUENCY_PRECOMPUTED_BLURING_RADIUS)
 				{
@@ -843,8 +866,8 @@ void MainCS(
 		// When not upscaling, manually force accumulate the sample of the kernel.
 		if (!KernelConfig.bSampleKernelCenter && !KernelConfig.bDescOrder)
 		{
-			UNROLL_N(MAX_SIGNAL_MULTIPLEXING)
-			for (uint SignalMultiplexId = 0; SignalMultiplexId < MAX_SIGNAL_MULTIPLEXING; SignalMultiplexId++)
+			UNROLL_N(SIGNAL_ARRAY_SIZE)
+			for (uint SignalMultiplexId = 0; SignalMultiplexId < SIGNAL_ARRAY_SIZE; SignalMultiplexId++)
 			{
 				const uint BatchedSignalId = ComputeSignalBatchIdFromSignalMultiplexId(KernelConfig, SignalMultiplexId);
 				
@@ -1116,6 +1139,10 @@ void MainCS(
 		}
 		#elif CONFIG_OUTPUT_MODE == OUTPUT_MODE_2MOMMENT_SUM
 		{
+			#if SIGNAL_ARRAY_SIZE != 2 * MAX_SIGNAL_BATCH_SIZE
+				#error Invalid signal array size.
+			#endif
+
 			MultiplexCount = 2 * CONFIG_SIGNAL_BATCH_SIZE;
 
 			UNROLL_N(CONFIG_SIGNAL_BATCH_SIZE)
diff --git a/Engine/Shaders/Private/ScreenSpaceDenoise/SSDSpatialKernel.ush b/Engine/Shaders/Private/ScreenSpaceDenoise/SSDSpatialKernel.ush
--- a/Engine/Shaders/Private/ScreenSpaceDenoise/SSDSpatialKernel.ush
+++ b/Engine/Shaders/Private/ScreenSpaceDenoise/SSDSpatialKernel.ush
@@ -127,13 +127,13 @@ struct FSSDKernelConfig
 	bool bClampUVPerMultiplexedSignal;
 
 	// The color space that has been encoded in the buffer.
-	uint BufferColorSpace[MAX_SIGNAL_MULTIPLEXING];
+	uint BufferColorSpace[SIGNAL_ARRAY_SIZE];
 	
 	// The color space of the accumulation.
-	uint AccumulatorColorSpace[MAX_SIGNAL_MULTIPLEXING];
+	uint AccumulatorColorSpace[SIGNAL_ARRAY_SIZE];
 	
 	// The color space of the accumulation.
-	uint BilateralSettings[MAX_SIGNAL_MULTIPLEXING];
+	uint BilateralSettings[SIGNAL_ARRAY_SIZE];
 
 	
 	// --------------------------- Per wave.
@@ -152,7 +152,7 @@ struct FSSDKernelConfig
 	float HarmonicPeriode;
 
 	// Buffer's min and max UV, per texture.
-	float4 PerSignalUVMinMax[MAX_SIGNAL_MULTIPLEXING];
+	float4 PerSignalUVMinMax[SIGNAL_ARRAY_SIZE];
 
 	
 	// --------------------------- Per lane.
@@ -183,7 +183,7 @@ struct FSSDKernelConfig
 	uint SampleTrackId;
 
 	// Reference meta data.
-	float RefBilateralDistance[MAX_SIGNAL_MULTIPLEXING];
+	float RefBilateralDistance[SIGNAL_ARRAY_SIZE];
 
 	// Uniform random values required for stocastic kernel.
 	float Randoms[1];
@@ -210,7 +210,7 @@ FSSDKernelConfig CreateKernelConfig()
 	KernelConfig.SampleSet = SAMPLE_SET_1X1;
 	KernelConfig.SampleSubSetId = 0;
 	KernelConfig.BufferLayout = SIGNAL_BUFFER_LAYOUT_UNINITIALIZED;
-	KernelConfig.MultiplexedSignalsPerSignalDomain = MAX_SIGNAL_MULTIPLEXING;
+	KernelConfig.MultiplexedSignalsPerSignalDomain = SIGNAL_ARRAY_SIZE;
 	KernelConfig.NeighborToRefComputation = NEIGHBOR_TO_REF_CACHE_WORLD_POSITION;
 	KernelConfig.RefSceneMetadataLayout = METADATA_BUFFER_LAYOUT_DISABLED;
 	KernelConfig.RingCount = 0;
@@ -227,8 +227,8 @@ FSSDKernelConfig CreateKernelConfig()
 	KernelConfig.bClampUVPerMultiplexedSignal = false;
 	
 	{
-		UNROLL_N(MAX_SIGNAL_MULTIPLEXING)
-		for (uint MultiplexId = 0; MultiplexId < MAX_SIGNAL_MULTIPLEXING; MultiplexId++)
+		UNROLL_N(SIGNAL_ARRAY_SIZE)
+		for (uint MultiplexId = 0; MultiplexId < SIGNAL_ARRAY_SIZE; MultiplexId++)
 		{
 			KernelConfig.BufferColorSpace[MultiplexId] = STANDARD_BUFFER_COLOR_SPACE;
 			KernelConfig.AccumulatorColorSpace[MultiplexId] = STANDARD_BUFFER_COLOR_SPACE;
@@ -244,8 +244,8 @@ FSSDKernelConfig CreateKernelConfig()
 	KernelConfig.HarmonicPeriode = 1.0;
 	
 	{
-		UNROLL_N(MAX_SIGNAL_MULTIPLEXING)
-		for (uint MultiplexId = 0; MultiplexId < MAX_SIGNAL_MULTIPLEXING; MultiplexId++)
+		UNROLL_N(SIGNAL_ARRAY_SIZE)
+		for (uint MultiplexId = 0; MultiplexId < SIGNAL_ARRAY_SIZE; MultiplexId++)
 		{
 			KernelConfig.PerSignalUVMinMax[MultiplexId] = 0.0;
 		}
@@ -266,8 +266,8 @@ FSSDKernelConfig CreateKernelConfig()
 	KernelConfig.HammersleySeed = 0;
 
 	{	
-		UNROLL_N(MAX_SIGNAL_MULTIPLEXING)
-		for (uint MultiplexId = 0; MultiplexId < MAX_SIGNAL_MULTIPLEXING; MultiplexId++)
+		UNROLL_N(SIGNAL_ARRAY_SIZE)
+		for (uint MultiplexId = 0; MultiplexId < SIGNAL_ARRAY_SIZE; MultiplexId++)
 		{
 			KernelConfig.RefBilateralDistance[MultiplexId] = 0.0;
 		}
@@ -296,8 +296,8 @@ void SetBilateralPreset(uint BilateralPresetId, inout FSSDKernelConfig KernelCon
 {
 	if (BilateralPresetId == BILATERAL_PRESET_MONOCHROMATIC_PENUMBRA)
 	{
-		UNROLL_N(MAX_SIGNAL_MULTIPLEXING)
-		for (uint MultiplexId = 0; MultiplexId < MAX_SIGNAL_MULTIPLEXING; MultiplexId++)
+		UNROLL_N(SIGNAL_ARRAY_SIZE)
+		for (uint MultiplexId = 0; MultiplexId < SIGNAL_ARRAY_SIZE; MultiplexId++)
 		{
 			// Shadow masks are normal invarient, so only reject based on position.
 			KernelConfig.BilateralSettings[MultiplexId] = BILATERAL_POSITION_BASED(5);
@@ -307,18 +307,22 @@ void SetBilateralPreset(uint BilateralPresetId, inout FSSDKernelConfig KernelCon
 	{
 		// Diffuse.
 		KernelConfig.BilateralSettings[0] = BILATERAL_POSITION_BASED(5) | BILATERAL_NORMAL;
-
+		
 		// Specular.
-		KernelConfig.BilateralSettings[1] = BILATERAL_POSITION_BASED(5) | BILATERAL_TOKOYASHI;
+		#if SIGNAL_ARRAY_SIZE > 1
+			KernelConfig.BilateralSettings[1] = BILATERAL_POSITION_BASED(5) | BILATERAL_TOKOYASHI;
+		#endif
 	}
 	else if (BilateralPresetId == BILATERAL_PRESET_REFLECTIONS)
 	{
 		// Specular.
 		// Can only be done using tokoyashi because have more than one sample at a time.
 		KernelConfig.BilateralSettings[0] = BILATERAL_TOKOYASHI;
-
-		// Specular variance for sampling rejection pre convolution.
-		KernelConfig.BilateralSettings[1] = KernelConfig.BilateralSettings[0];
+		
+		#if SIGNAL_ARRAY_SIZE > 1
+			// Specular variance for sampling rejection pre convolution.
+			KernelConfig.BilateralSettings[1] = KernelConfig.BilateralSettings[0];
+		#endif
 	}
 	else if (BilateralPresetId == BILATERAL_PRESET_REFLECTIONS_1SPP)
 	{
@@ -331,17 +335,21 @@ void SetBilateralPreset(uint BilateralPresetId, inout FSSDKernelConfig KernelCon
 		// Specular.
 		// Can only be done using tokoyashi because have more than one sample at a time.
 		KernelConfig.BilateralSettings[0] = BILATERAL_POSITION_BASED(1) | BILATERAL_TOKOYASHI;
-
-		// Specular variance for sampling rejection pre convolution.
-		KernelConfig.BilateralSettings[1] = KernelConfig.BilateralSettings[0];
+		
+		#if SIGNAL_ARRAY_SIZE > 1
+			// Specular variance for sampling rejection pre convolution.
+			KernelConfig.BilateralSettings[1] = KernelConfig.BilateralSettings[0];
+		#endif
 	}
 	else if (BilateralPresetId == BILATERAL_PRESET_DIFFUSE)
 	{
 		// Diffuse depends on world position, but also the normal of the surface given it doesn't store any directionality.
 		KernelConfig.BilateralSettings[0] = BILATERAL_POSITION_BASED(2) | BILATERAL_NORMAL;
 		
-		// Variance for sampling rejection pre convolution.
-		KernelConfig.BilateralSettings[1] = KernelConfig.BilateralSettings[0];
+		#if SIGNAL_ARRAY_SIZE > 1
+			// Variance for sampling rejection pre convolution.
+			KernelConfig.BilateralSettings[1] = KernelConfig.BilateralSettings[0];
+		#endif
 	}
 	else if (BilateralPresetId == BILATERAL_PRESET_SPHERICAL_HARMONIC)
 	{
@@ -553,8 +561,8 @@ FSSDSignalArray UncompressMultiplexedSignals(
 
 	if (KernelConfig.bClampUVPerMultiplexedSignal)
 	{
-		UNROLL_N(MAX_SIGNAL_MULTIPLEXING)
-		for (uint SignalMultiplexId = 0; SignalMultiplexId < MAX_SIGNAL_MULTIPLEXING; SignalMultiplexId++)
+		UNROLL_N(SIGNAL_ARRAY_SIZE)
+		for (uint SignalMultiplexId = 0; SignalMultiplexId < SIGNAL_ARRAY_SIZE; SignalMultiplexId++)
 		{
 			bool bInvalidSample = any(SampleBufferUV != clamp(
 				SampleBufferUV, KernelConfig.PerSignalUVMinMax[SignalMultiplexId].xy, KernelConfig.PerSignalUVMinMax[SignalMultiplexId].zw));
@@ -563,7 +571,7 @@ FSSDSignalArray UncompressMultiplexedSignals(
 			{
 				MultiplexedSamples.Array[SignalMultiplexId] = CreateSignalSampleFromScalarValue(0.0);
 			}
-		} // for (uint SignalMultiplexId = 0; SignalMultiplexId < MAX_SIGNAL_MULTIPLEXING; SignalMultiplexId++)
+		} // for (uint SignalMultiplexId = 0; SignalMultiplexId < SIGNAL_ARRAY_SIZE; SignalMultiplexId++)
 	}
 
 	return MultiplexedSamples;
@@ -627,8 +635,8 @@ void AccumulateSampledMultiplexedSignals(
 	}
 	#endif
 
-	UNROLL_N(MAX_SIGNAL_MULTIPLEXING)
-	for (uint SignalMultiplexId = 0; SignalMultiplexId < MAX_SIGNAL_MULTIPLEXING; SignalMultiplexId++)
+	UNROLL_N(SIGNAL_ARRAY_SIZE)
+	for (uint SignalMultiplexId = 0; SignalMultiplexId < SIGNAL_ARRAY_SIZE; SignalMultiplexId++)
 	{
 		// Compute at compile time the id of the signal being processed.
 		const uint BatchedSignalId = ComputeSignalBatchIdFromSignalMultiplexId(KernelConfig, SignalMultiplexId);
@@ -728,7 +736,7 @@ void AccumulateSampledMultiplexedSignals(
 			}
 			#endif
 		}
-	} // for (uint SignalMultiplexId = 0; SignalMultiplexId < MAX_SIGNAL_MULTIPLEXING; SignalMultiplexId++)
+	} // for (uint SignalMultiplexId = 0; SignalMultiplexId < SIGNAL_ARRAY_SIZE; SignalMultiplexId++)
 	
 	#if CONFIG_ACCUMULATOR_VGPR_COMPRESSION != ACCUMULATOR_COMPRESSION_DISABLED
 		CompressedAccumulators = CompressAccumulatorArray(Accumulators, CONFIG_ACCUMULATOR_VGPR_COMPRESSION);
@@ -850,8 +858,8 @@ void SampleAndAccumulateMultiplexedSignalsPair(
 		// TODO(Denoiser): this should be normalized in theory...
 		if (KernelConfig.bMinSamplePairInvFrequency)
 		{
-			UNROLL_N(MAX_SIGNAL_MULTIPLEXING)
-			for (uint SignalMultiplexId = 0; SignalMultiplexId < MAX_SIGNAL_MULTIPLEXING; SignalMultiplexId++)
+			UNROLL_N(SIGNAL_ARRAY_SIZE)
+			for (uint SignalMultiplexId = 0; SignalMultiplexId < SIGNAL_ARRAY_SIZE; SignalMultiplexId++)
 			{
 				float MinInvFrequency = min(
 					MultiplexedSamples[0].Array[SignalMultiplexId].WorldBluringRadius,
@@ -901,8 +909,8 @@ void StartAccumulatingCluster(
 		FSSDSignalAccumulatorArray Accumulators = UncompressAccumulatorArray(CompressedAccumulators, CONFIG_ACCUMULATOR_VGPR_COMPRESSION);
 	#endif
 
-	UNROLL_N(MAX_SIGNAL_MULTIPLEXING)
-	for (uint SignalMultiplexId = 0; SignalMultiplexId < MAX_SIGNAL_MULTIPLEXING; SignalMultiplexId++)
+	UNROLL_N(SIGNAL_ARRAY_SIZE)
+	for (uint SignalMultiplexId = 0; SignalMultiplexId < SIGNAL_ARRAY_SIZE; SignalMultiplexId++)
 	{
 		#if CONFIG_ACCUMULATOR_VGPR_COMPRESSION == ACCUMULATOR_COMPRESSION_DISABLED
 		{
@@ -935,8 +943,8 @@ void DijestAccumulatedClusterSamples(
 		FSSDSignalAccumulatorArray Accumulators = UncompressAccumulatorArray(CompressedAccumulators, CONFIG_ACCUMULATOR_VGPR_COMPRESSION);
 	#endif
 		
-	UNROLL_N(MAX_SIGNAL_MULTIPLEXING)
-	for (uint SignalMultiplexId = 0; SignalMultiplexId < MAX_SIGNAL_MULTIPLEXING; SignalMultiplexId++)
+	UNROLL_N(SIGNAL_ARRAY_SIZE)
+	for (uint SignalMultiplexId = 0; SignalMultiplexId < SIGNAL_ARRAY_SIZE; SignalMultiplexId++)
 	{
 		#if CONFIG_ACCUMULATOR_VGPR_COMPRESSION == ACCUMULATOR_COMPRESSION_DISABLED
 		{
diff --git a/Engine/Shaders/Private/ScreenSpaceDenoise/SSDTemporalAccumulation.usf b/Engine/Shaders/Private/ScreenSpaceDenoise/SSDTemporalAccumulation.usf
--- a/Engine/Shaders/Private/ScreenSpaceDenoise/SSDTemporalAccumulation.usf
+++ b/Engine/Shaders/Private/ScreenSpaceDenoise/SSDTemporalAccumulation.usf
@@ -33,6 +33,9 @@
 
 
 #if CONFIG_SIGNAL_PROCESSING == SIGNAL_PROCESSING_SHADOW_VISIBILITY_MASK
+	#define MAX_SIGNAL_BATCH_SIZE CONFIG_SIGNAL_BATCH_SIZE
+	#define SIGNAL_ARRAY_SIZE CONFIG_SIGNAL_BATCH_SIZE
+
 	#define CONFIG_BILATERAL_PRESET BILATERAL_PRESET_MONOCHROMATIC_PENUMBRA
 	#define CONFIG_HISTORY_BILATERAL_PRESET BILATERAL_PRESET_MONOCHROMATIC_PENUMBRA
 	#define CONFIG_MULTIPLEXED_SIGNALS_PER_SIGNAL_DOMAIN 1
@@ -50,24 +53,17 @@
 	// TODO(Denoiser): cause a regression on spot light for some reasons.
 	//#define CONFIG_USE_NEAREST_HISTORY 1
 
-	// Use pre transformed rejection buffer that contains pre transformed momment 1 & 2.
-	#if 0
-		#define CONFIG_USE_REJECTION_BUFFER               1
-		#define CONFIG_SIGNAL_HISTORY_REJECTION_LAYOUT    SIGNAL_BUFFER_LAYOUT_PENUMBRA_REJECTION
-		
-		#define CONFIG_REJECTION_TEXTURE_COUNT CONFIG_SIGNAL_BATCH_SIZE
-
-		#define CONFIG_REJECTION_INPUT_MODE               REJECTION_INPUT_MODE_2PRETRANSFORMED_MOMMENTS
-		#define CONFIG_REJECTION_SAMPLE_SET               REJECTION_SAMPLE_SET_SOBEK2018
-	#endif
-
 	#if DIM_SIGNAL_BATCH_SIZE > 1
 		#define CONFIG_CLAMP_UV_PER_SIGNAL 1
 	#endif
 
-	#define MAX_SIGNAL_BATCH_SIZE CONFIG_SIGNAL_BATCH_SIZE
-
 #elif CONFIG_SIGNAL_PROCESSING == SIGNAL_PROCESSING_POLYCHROMATIC_PENUMBRA_HARMONIC
+	// Denoise diffuse and specular harmonics at the same time.
+	#define MAX_SIGNAL_BATCH_SIZE 2
+	#define SIGNAL_ARRAY_SIZE 2
+	#undef CONFIG_SIGNAL_BATCH_SIZE
+	#define CONFIG_SIGNAL_BATCH_SIZE 2
+
 	#define COMPILE_SIGNAL_COLOR 1
 
 	#define CONFIG_BILATERAL_PRESET BILATERAL_PRESET_POLYCHROMATIC_PENUMBRA
@@ -80,17 +76,14 @@
 	#define CONFIG_INPUT_TEXTURE_COUNT 2
 	#define CONFIG_HISTORY_TEXTURE_COUNT 2
 
-	#undef CONFIG_SIGNAL_BATCH_SIZE
-	#define CONFIG_SIGNAL_BATCH_SIZE 2
-
 	#define CONFIG_HISTORY_REJECTION HISTORY_REJECTION_VAR_BOUNDARIES
 
-	// Denoise diffuse and specular harmonics at the same time.
-	#define MAX_SIGNAL_BATCH_SIZE 2
-
 #elif CONFIG_SIGNAL_PROCESSING == SIGNAL_PROCESSING_REFLECTIONS
 	#define COMPILE_SIGNAL_COLOR 1
 
+	#define MAX_SIGNAL_BATCH_SIZE 1
+	#define SIGNAL_ARRAY_SIZE 2 // First and Second moment in rejection pre convolution.
+
 	// Input and output layout.
 	#define CONFIG_SIGNAL_INPUT_LAYOUT   SIGNAL_BUFFER_LAYOUT_REFLECTIONS_HISTORY
 	#define CONFIG_SIGNAL_HISTORY_LAYOUT SIGNAL_BUFFER_LAYOUT_REFLECTIONS_HISTORY
@@ -118,6 +111,9 @@
 	#endif
 
 #elif CONFIG_SIGNAL_PROCESSING == SIGNAL_PROCESSING_AO
+	#define MAX_SIGNAL_BATCH_SIZE 1
+	#define SIGNAL_ARRAY_SIZE 2 // First and Second moment in rejection pre convolution.
+
 	// Input and output layout.
 	#define CONFIG_SIGNAL_INPUT_LAYOUT   SIGNAL_BUFFER_LAYOUT_AO_HISTORY
 	#define CONFIG_SIGNAL_HISTORY_LAYOUT SIGNAL_BUFFER_LAYOUT_AO_HISTORY
@@ -141,6 +137,9 @@
 	#endif
 
 #elif CONFIG_SIGNAL_PROCESSING == SIGNAL_PROCESSING_DIFFUSE_INDIRECT_AND_AO
+	#define MAX_SIGNAL_BATCH_SIZE 1
+	#define SIGNAL_ARRAY_SIZE 1
+
 	#define COMPILE_SIGNAL_COLOR 1
 
 	// Input and output layout.
@@ -158,6 +157,9 @@
 	#define CONFIG_HISTORY_BILATERAL_PRESET BILATERAL_PRESET_DIFFUSE
 
 #elif CONFIG_SIGNAL_PROCESSING == SIGNAL_PROCESSING_DIFFUSE_SPHERICAL_HARMONIC
+	#define MAX_SIGNAL_BATCH_SIZE 1
+	#define SIGNAL_ARRAY_SIZE 1
+
 	#define COMPILE_SIGNAL_COLOR_SH 1
 
 	// Input and output layout.
@@ -178,6 +180,9 @@
 	#define CONFIG_HISTORY_BILATERAL_PRESET BILATERAL_PRESET_SPHERICAL_HARMONIC
 
 #elif CONFIG_SIGNAL_PROCESSING == SIGNAL_PROCESSING_SSGI
+	#define MAX_SIGNAL_BATCH_SIZE 1
+	#define SIGNAL_ARRAY_SIZE 1
+
 	#define COMPILE_SIGNAL_COLOR 1
 	#define CONFIG_MULTIPLEXED_SIGNALS_PER_SIGNAL_DOMAIN 1
 
@@ -241,7 +246,7 @@
 
 /** The number of signal that should be processed per signal domain. */
 #ifndef CONFIG_MULTIPLEXED_SIGNALS_PER_SIGNAL_DOMAIN
-	#define CONFIG_MULTIPLEXED_SIGNALS_PER_SIGNAL_DOMAIN MAX_SIGNAL_MULTIPLEXING
+	#define CONFIG_MULTIPLEXED_SIGNALS_PER_SIGNAL_DOMAIN SIGNAL_ARRAY_SIZE
 #endif
 
 /** Whether can use a nearest sampler when sampling the history */
@@ -672,7 +677,7 @@ void TemporallyAccumulate(
 				KernelConfig.BufferLayout = CONFIG_SIGNAL_HISTORY_REJECTION_LAYOUT;
 				KernelConfig.bNormalizeSample = false;
 			
-				for (uint MultiplexId = 0; MultiplexId < MAX_SIGNAL_MULTIPLEXING; MultiplexId++)
+				for (uint MultiplexId = 0; MultiplexId < SIGNAL_ARRAY_SIZE; MultiplexId++)
 				{
 					KernelConfig.BufferColorSpace[MultiplexId] = CONFIG_REJECTION_BUFFER_COLOR_SPACE;
 					KernelConfig.AccumulatorColorSpace[MultiplexId] = CONFIG_HISTORY_REJECTION_COLOR_SPACE;
@@ -687,7 +692,7 @@ void TemporallyAccumulate(
 				KernelConfig.BufferLayout = CONFIG_SIGNAL_INPUT_LAYOUT;
 				KernelConfig.bNormalizeSample = true;
 			
-				for (uint MultiplexId = 0; MultiplexId < MAX_SIGNAL_MULTIPLEXING; MultiplexId++)
+				for (uint MultiplexId = 0; MultiplexId < SIGNAL_ARRAY_SIZE; MultiplexId++)
 				{
 					KernelConfig.AccumulatorColorSpace[MultiplexId] = CONFIG_HISTORY_REJECTION_COLOR_SPACE;
 				}
@@ -775,6 +780,10 @@ void TemporallyAccumulate(
 			}
 			#elif CONFIG_REJECTION_INPUT_MODE == REJECTION_INPUT_MODE_2PRETRANSFORMED_MOMMENTS
 			{
+				#if SIGNAL_ARRAY_SIZE != 2 * MAX_SIGNAL_BATCH_SIZE
+					#error Invalid signal array size.
+				#endif
+
 				float NormalizeFactor = SafeRcp(SignalAccumulators.Array[BatchedSignalId * 2 + 0].Moment1.SampleCount);
 				NeighborMoment1 = MulSignal(SignalAccumulators.Array[BatchedSignalId * 2 + 0].Moment1, NormalizeFactor);
 				NeighborMoment2 = MulSignal(SignalAccumulators.Array[BatchedSignalId * 2 + 1].Moment1, NormalizeFactor);
