Fixes shader compilation failure in denoiser.

diff --git a/Engine/Shaders/Private/ScreenSpaceDenoise/SSDComposeHarmonics.usf b/Engine/Shaders/Private/ScreenSpaceDenoise/SSDComposeHarmonics.usf
--- a/Engine/Shaders/Private/ScreenSpaceDenoise/SSDComposeHarmonics.usf
+++ b/Engine/Shaders/Private/ScreenSpaceDenoise/SSDComposeHarmonics.usf
@@ -11,9 +11,17 @@
 #define CONFIG_SIGNAL_PROCESSING SIGNAL_PROCESSING_POLYCHROMATIC_PENUMBRA_HARMONIC
 
 #if CONFIG_SIGNAL_PROCESSING == SIGNAL_PROCESSING_POLYCHROMATIC_PENUMBRA_HARMONIC
-	#define COMPILE_SIGNAL_COLOR 1
+	// Denoise diffuse and specular harmonics at the same time.
+	#define MAX_SIGNAL_BATCH_SIZE 2
+	#define SIGNAL_ARRAY_SIZE 2
+	#define CONFIG_SIGNAL_BATCH_SIZE 2
+
+	// Each harmonic requires input and output RGB.
 	#define COMPILE_SIGNAL_COLOR_ARRAY 2
 
+	// Outputing a single composed color for each diffuse and specular.
+	#define COMPILE_SIGNAL_COLOR 1
+
 	//#define CONFIG_VIEW_POSITION_BASED_BILATERAL 2
 	#define CONFIG_MULTIPLEXED_SIGNALS_PER_SIGNAL_DOMAIN 1
 
@@ -26,11 +34,6 @@
 
 	#define CONFIG_HARMONIC_COUNT DENOISER_HARMONICS_COUNT
 
-	// Denoise diffuse and specular harmonics at the same time.
-	#define MAX_SIGNAL_BATCH_SIZE 2
-
-	#define CONFIG_SIGNAL_BATCH_SIZE 2
-
 #else
 	#error Unknown signal processing.
 #endif
diff --git a/Engine/Shaders/Private/ScreenSpaceDenoise/SSDInjest.usf b/Engine/Shaders/Private/ScreenSpaceDenoise/SSDInjest.usf
--- a/Engine/Shaders/Private/ScreenSpaceDenoise/SSDInjest.usf
+++ b/Engine/Shaders/Private/ScreenSpaceDenoise/SSDInjest.usf
@@ -11,6 +11,9 @@
 #define CONFIG_SIGNAL_BATCH_SIZE DIM_SIGNAL_BATCH_SIZE
 
 #if CONFIG_SIGNAL_PROCESSING == SIGNAL_PROCESSING_SHADOW_VISIBILITY_MASK
+	#define MAX_SIGNAL_BATCH_SIZE CONFIG_SIGNAL_BATCH_SIZE
+	#define SIGNAL_ARRAY_SIZE CONFIG_SIGNAL_BATCH_SIZE
+
 	#define CONFIG_SIGNAL_INPUT_LAYOUT  SIGNAL_BUFFER_LAYOUT_PENUMBRA_INPUT_NSPP
 	#define CONFIG_SIGNAL_OUTPUT_LAYOUT SIGNAL_BUFFER_LAYOUT_PENUMBRA_INJESTION_NSPP
 	#define CONFIG_MULTIPLEXED_SIGNALS_PER_SIGNAL_DOMAIN 1
@@ -18,8 +21,6 @@
 	#define CONFIG_INPUT_TEXTURE_COUNT DIM_SIGNAL_BATCH_SIZE
 	#define CONFIG_OUTPUT_TEXTURE_COUNT ((DIM_SIGNAL_BATCH_SIZE + 1) / 2)
 
-	#define MAX_SIGNAL_BATCH_SIZE CONFIG_SIGNAL_BATCH_SIZE
-
 #else
 	#error Unknown signal processing.
 #endif
diff --git a/Engine/Shaders/Private/ScreenSpaceDenoise/SSDReduce.usf b/Engine/Shaders/Private/ScreenSpaceDenoise/SSDReduce.usf
--- a/Engine/Shaders/Private/ScreenSpaceDenoise/SSDReduce.usf
+++ b/Engine/Shaders/Private/ScreenSpaceDenoise/SSDReduce.usf
@@ -15,6 +15,9 @@
 
 // Configures all the pass for each individual signals.
 #if CONFIG_SIGNAL_PROCESSING == SIGNAL_PROCESSING_DIFFUSE_SPHERICAL_HARMONIC
+	#define MAX_SIGNAL_BATCH_SIZE 1
+	#define SIGNAL_ARRAY_SIZE 1
+
 	#define COMPILE_SIGNAL_COLOR_SH 1
 
 	// Given it's a spherical harmonic that store directionality, only need position based rejection.
