Made hair PPLL allocate memory in block per thread group.

diff --git a/Engine/Shaders/Private/HairStrands/HairStrandsVisibilityCompaction.usf b/Engine/Shaders/Private/HairStrands/HairStrandsVisibilityCompaction.usf
--- a/Engine/Shaders/Private/HairStrands/HairStrandsVisibilityCompaction.usf
+++ b/Engine/Shaders/Private/HairStrands/HairStrandsVisibilityCompaction.usf
@@ -51,6 +51,10 @@ RWTexture2D<uint4> OutCategorizationTexture;
 #endif
 
 
+groupshared uint AllocationNodeCount;
+groupshared uint AllocationNodeOffset;
+
+
 #if PERMUTATION_PPLL
 
 Texture2D<uint> PPLLCounter;
@@ -78,18 +82,18 @@ void MainCS(uint GroupIndex : SV_GroupIndex, uint3 DispatchThreadId : SV_Dispatc
 	float2 MaxVelocity = 0;
 	float MaxVelocityMagnitude2 = 0;
 #endif
-
-	if (FirstNodeIndex != 0xFFFFFFFF && all(PixelCoord < OutputResolution))
+	
+	const bool bExecute = FirstNodeIndex != 0xFFFFFFFF && all(PixelCoord < OutputResolution);
+
+	uint SortedIndex[HAIR_VISIBILITY_MAXSAMPLECOUNT];
+	float SortedDepth[HAIR_VISIBILITY_MAXSAMPLECOUNT];
+	float ValidPixelSampleTotalCoverage = 0.0f;
+	float Transmittance = 1.0f;
+	uint PixelNodeCount = 0;
+	if (bExecute)
 	{
 		//////////
 		// First pass: Pre compute some values based on the linked list
-		float ValidPixelSampleTotalCoverage = 0.0f;
-		float Transmittance = 1.0f;
-
-		uint SortedIndex[HAIR_VISIBILITY_MAXSAMPLECOUNT];
-		float SortedDepth[HAIR_VISIBILITY_MAXSAMPLECOUNT];
-
-		uint PixelNodeCount = 0;
 		uint NodeIndex = FirstNodeIndex;
 		while (NodeIndex != 0xFFFFFFFF)// && Transmittance > 0.0f)// && PixelNodeCount < HAIR_VISIBILITY_MAXSAMPLECOUNT)
 		{
@@ -118,11 +122,32 @@ void MainCS(uint GroupIndex : SV_GroupIndex, uint3 DispatchThreadId : SV_Dispatc
 			PixelNodeCount++;
 			NodeIndex = PPLLNode.NextNodeIndex;
 		}
+	}
+
+	// Initialise group allocation node count
+	if (GroupIndex == 0)
+	{
+		AllocationNodeCount = 0;
+		AllocationNodeOffset= 0;
+	}
+	GroupMemoryBarrierWithGroupSync();
+	// Now notify how many nodes this thread wants to allocate
+	uint PixelDataSubOffsetInStorage = 0; 
+	PixelSamplecount = min(PixelNodeCount, HAIR_VISIBILITY_MAXSAMPLECOUNT);
+	InterlockedAdd(AllocationNodeCount, PixelSamplecount, PixelDataSubOffsetInStorage);
+	GroupMemoryBarrierWithGroupSync();
+	// Now allocate all the nodes for this group contiguously in memory
+	if (GroupIndex == 0 && AllocationNodeCount > 0)
+	{
+		InterlockedAdd(OutCompactNodeCounter[uint2(0, 0)], AllocationNodeCount, AllocationNodeOffset);
+	}
+	GroupMemoryBarrierWithGroupSync();
+	OffsetInStorage = AllocationNodeOffset + PixelDataSubOffsetInStorage;
 
+	if(bExecute)
+	{
 		//////////
 		// Second pass: write out compact pixel sample data
-		PixelSamplecount = min(PixelNodeCount, HAIR_VISIBILITY_MAXSAMPLECOUNT);
-		InterlockedAdd(OutCompactNodeCounter[uint2(0, 0)], PixelSamplecount, OffsetInStorage);
 		for (uint i = 0; i < PixelSamplecount; ++i)
 		{
 			const uint NodePPLLIndex = SortedIndex[i];
@@ -282,9 +307,6 @@ Texture2DMS<float, HAIR_VISIBILITY_MAXSAMPLECOUNT> MSAA_DepthTexture;
 Texture2DMS<float4,HAIR_VISIBILITY_MAXSAMPLECOUNT> MSAA_MaterialTexture;
 Texture2DMS<float4,HAIR_VISIBILITY_MAXSAMPLECOUNT> MSAA_AttributeTexture;
 
-groupshared uint AllocationNodeCount;
-groupshared uint AllocationNodeOffset;
-
 [numthreads(TILE_PIXEL_SIZE_X, TILE_PIXEL_SIZE_Y, 1)]
 void MainCS(uint GroupIndex : SV_GroupIndex, uint3 DispatchThreadId : SV_DispatchThreadID)
 {
