Fixed distance filed local light culling which was incorectly picking up too many objects.

diff --git a/Engine/Shaders/Private/DistanceFieldShadowing.usf b/Engine/Shaders/Private/DistanceFieldShadowing.usf
--- a/Engine/Shaders/Private/DistanceFieldShadowing.usf
+++ b/Engine/Shaders/Private/DistanceFieldShadowing.usf
@@ -448,7 +448,7 @@ void CullObjectsToTileWithGather(
 			float4 SphereCenterAndRadius = LoadObjectPositionAndRadius(ObjectIndex);
 
 			BRANCH
-			if (RaySegmentHitSphere(WorldGroup0Center, LightDirection0, RayLength0, SphereCenterAndRadius.xyz, SphereCenterAndRadius.w + Group0BoundingRadius + LightSourceRadius))
+			if (RaySegmentHitSphere(WorldGroup0Center, LightDirection0, RayLength0, SphereCenterAndRadius.xyz, SphereCenterAndRadius.w + Group0BoundingRadius))
 			{
 				uint ListIndex;
 				InterlockedAdd(TileNumObjects0, 1U, ListIndex);
@@ -458,7 +458,7 @@ void CullObjectsToTileWithGather(
 			}
 
 			BRANCH
-			if (RaySegmentHitSphere(WorldGroup1Center, LightDirection1, RayLength1, SphereCenterAndRadius.xyz, SphereCenterAndRadius.w + Group1BoundingRadius + LightSourceRadius))
+			if (RaySegmentHitSphere(WorldGroup1Center, LightDirection1, RayLength1, SphereCenterAndRadius.xyz, SphereCenterAndRadius.w + Group1BoundingRadius))
 			{
 				uint ListIndex;
 				InterlockedAdd(TileNumObjects1, 1U, ListIndex);
