Fixes regressions when donoiser have to upscale.

diff --git a/Engine/Shaders/Private/ScreenSpaceDenoise/SSDSpatialAccumulation.usf b/Engine/Shaders/Private/ScreenSpaceDenoise/SSDSpatialAccumulation.usf
--- a/Engine/Shaders/Private/ScreenSpaceDenoise/SSDSpatialAccumulation.usf
+++ b/Engine/Shaders/Private/ScreenSpaceDenoise/SSDSpatialAccumulation.usf
@@ -1198,7 +1198,7 @@ void MainCS(
 
 	// TODO(Denoiser):  LeaveRayCount = (LeaveRayCount - 1) * 9 / (9 - 2) post processing to reject when for history rejection
 	uint2 OutputPixelPostion;
-	#if CONFIG_VGPR_OPTIMIZATION && 1 // TODO(Denoiser)
+	#if CONFIG_VGPR_OPTIMIZATION && !CONFIG_UPSCALE // TODO(Denoiser)
 	{
 		// No need to keep DispatchThreadId, can recompute the output pixel position based on information stored in VGPRs for spatial kernel.
 		#if CONFIG_SAMPLE_SET == SAMPLE_SET_STACKOWIAK_4_SETS || CONFIG_SAMPLE_SET == SAMPLE_SET_STOCASTIC_HIERARCHY
