Integrating CL 9474326 from dev-render, to fix scene color node in translucent shaders.

diff --git a/Engine/Shaders/Private/MaterialTemplate.ush b/Engine/Shaders/Private/MaterialTemplate.ush
--- a/Engine/Shaders/Private/MaterialTemplate.ush
+++ b/Engine/Shaders/Private/MaterialTemplate.ush
@@ -1791,7 +1791,13 @@ float3 DecodeSceneColorForMaterialNode(float2 ScreenUV)
 	float4 EncodedSceneColor = Texture2DSample(SceneTexturesStruct.SceneColorCopyTexture, SceneTexturesStruct.SceneColorCopyTextureSampler, ScreenUV);
 
 	// Undo the function in EncodeSceneColorForMaterialNode
-	return pow(EncodedSceneColor.rgb, 4) * 10;
+	float3 SampledColor = pow(EncodedSceneColor.rgb, 4) * 10;
+
+#if USE_PREEXPOSURE
+	SampledColor *= View.OneOverPreExposure.xxx;
+#endif
+
+	return SampledColor;
 #endif
 }
 
