Fixed missing top part of sky in rare case when stars align.

diff --git a/Engine/Shaders/Private/MaterialTemplate.ush b/Engine/Shaders/Private/MaterialTemplate.ush
--- a/Engine/Shaders/Private/MaterialTemplate.ush
+++ b/Engine/Shaders/Private/MaterialTemplate.ush
@@ -1148,7 +1148,7 @@ float3 MaterialExpressionSkyAtmosphereViewLuminance(FMaterialPixelParameters Par
 	const float3 ViewDir = -Parameters.CameraVector;
 
 	// The referencial used to build the Sky View lut
-	float3x3 LocalReferencial = GetSkyViewLutReferential(PlanetCenterToWorldCameraPos, ResolvedView.ViewForward);
+	float3x3 LocalReferencial = GetSkyViewLutReferential(PlanetCenterToWorldCameraPos, ResolvedView.ViewForward, ResolvedView.ViewRight);
 	// Compute inputs in this referential
 	float3 WorldPosLocal = float3(0.0, 0.0, ViewHeight);
 	float3 UpVectorLocal = float3(0.0, 0.0, 1.0);
diff --git a/Engine/Shaders/Private/SkyAtmosphere.usf b/Engine/Shaders/Private/SkyAtmosphere.usf
--- a/Engine/Shaders/Private/SkyAtmosphere.usf
+++ b/Engine/Shaders/Private/SkyAtmosphere.usf
@@ -749,7 +749,7 @@ void RenderSkyAtmosphereRayMarchingPS(
 		float2 UV;
 
 		// The referencial used to build the Sky View lut
-		float3x3 LocalReferencial = GetSkyViewLutReferential(WorldPos, View.ViewForward);
+		float3x3 LocalReferencial = GetSkyViewLutReferential(WorldPos, View.ViewForward, View.ViewRight);
 
 		// Input vectors expressed in this referencial: Up is always Z. Also note that ViewHeight is unchanged in this referencial.
 		float3 WorldPosLocal = float3(0.0, 0.0, ViewHeight);
@@ -1032,7 +1032,7 @@ void RenderSkyViewLutCS(uint3 ThreadId : SV_DispatchThreadID)
 	// into a referential with UP being perpendicular to the ground. And with origin at the planet center.
 
 	// This is the local referencial
-	float3x3 LocalReferencial = GetSkyViewLutReferential(WorldPos, View.ViewForward);
+	float3x3 LocalReferencial = GetSkyViewLutReferential(WorldPos, View.ViewForward, View.ViewRight);
 
 	// This is the LUT camera height and position in the local referential
 	float ViewHeight = length(WorldPos);
diff --git a/Engine/Shaders/Private/SkyAtmosphereCommon.ush b/Engine/Shaders/Private/SkyAtmosphereCommon.ush
--- a/Engine/Shaders/Private/SkyAtmosphereCommon.ush
+++ b/Engine/Shaders/Private/SkyAtmosphereCommon.ush
@@ -156,18 +156,20 @@ void SkyViewLutParamsToUv(
 	UV = FromUnitToSubUvs(UV, SkyViewLutSizeAndInvSize);
 }
 
-float3x3 GetSkyViewLutReferential(in float3 WorldPos, in float3 ViewForward)
+float3x3 GetSkyViewLutReferential(in float3 WorldPos, in float3 ViewForward, in float3 ViewRight)
 {
 	float3	Up = normalize(WorldPos);
 	float3	Forward = ViewForward;		// This can make texel visible when the camera is rotating. Use constant worl direction instead?
 	float3	Left = normalize(cross(Forward, Up));
+	if (abs(dot(Forward, Up)) > 0.99f)
+	{
+		Left = -ViewRight;
+	}
 	Forward = normalize(cross(Up, Left));
 	float3x3 LocalReferencial = transpose(float3x3(Forward, Left, Up));
 	return LocalReferencial;
 }
 
-
-
 float3 GetLightDiskLuminance(
 	float3 WorldPos, float3 WorldDir, float BottomRadius, float TopRadius,
 	Texture2D<float4> TransmittanceLutTexture, SamplerState TransmittanceLutTextureSampler,
