Fixed light function causing artefact due to clamp at the edge of the translucency volume.
Light function now fades out at the edge similarly to shadow.

diff --git a/Engine/Shaders/Private/TranslucentLightInjectionShaders.usf b/Engine/Shaders/Private/TranslucentLightInjectionShaders.usf
--- a/Engine/Shaders/Private/TranslucentLightInjectionShaders.usf
+++ b/Engine/Shaders/Private/TranslucentLightInjectionShaders.usf
@@ -176,6 +176,9 @@ void InjectMainPS(
 		#endif
 
 		float ShadowFactor = ComputeVolumeShadowing(WorldPositionForLighting, bPointLight, bSpotLight);
+
+		// Apply light function (it will also fade out at the edge of the volume so overall dimming of light using light function is not a recommended workflow)
+		ShadowFactor *= GetLightFunctionShadowFactor(WorldPositionForLighting);
 		
 		if (VolumeCascadeIndex == 1)
 		{
@@ -198,9 +201,6 @@ void InjectMainPS(
 			#endif
 		}
 
-		// Apply light function after edge fading, so that a light function which is modifying the light's brightness still works
-		ShadowFactor *= GetLightFunctionShadowFactor(WorldPositionForLighting);
-
 		float3 Lighting = DeferredLightUniforms.Color / PI * Attenuation * ShadowFactor;
 
 		FTwoBandSHVectorRGB SHLighting = MulSH(SHBasisFunction(NormalizedLightVector), Lighting);
