Fix for tonemapper sharpen operator

diff --git a/Engine/Shaders/Private/PostProcessTonemap.usf b/Engine/Shaders/Private/PostProcessTonemap.usf
--- a/Engine/Shaders/Private/PostProcessTonemap.usf
+++ b/Engine/Shaders/Private/PostProcessTonemap.usf
@@ -424,7 +424,8 @@ float4 TonemapCommonPS(
 
 	#if HAS_PIXEL_QUAD_MESSAGE_PASSING_SUPPORT 
 		// Use Wave Intrinsics to reduce texture taps
-		FPQMPContext PQMP = PQMPInit(ScreenPos);
+		FPQMPContext PQMP = PQMPInit(floor(SceneUV* Color_Extent));
+				
 		half4 LuminanceNeighbors;
 
 		half3 C1 = Texture2DSample(ColorTexture, ColorSampler, SceneUV + float2(PQMP.QuadVector.x, 0) * Color_ExtentInverse).rgb * OneOverPreExposure;
