Fix coverage pass.

diff --git a/Engine/Shaders/Private/HairStrands/HairStrandsVisibilityPS.usf b/Engine/Shaders/Private/HairStrands/HairStrandsVisibilityPS.usf
--- a/Engine/Shaders/Private/HairStrands/HairStrandsVisibilityPS.usf
+++ b/Engine/Shaders/Private/HairStrands/HairStrandsVisibilityPS.usf
@@ -44,24 +44,20 @@ uint HairVisibilityPass_HairClusterIndex;
 EARLYDEPTHSTENCIL	// Required for UAV operations to not happen.
 #endif
 void MainVisibility(
-	centroid FVertexFactoryInterpolantsVSToPS Interpolants,
-	in float4 SvPosition	: SV_Position
+	centroid FVertexFactoryInterpolantsVSToPS Interpolants
+	, in float4 SvPosition	: SV_Position
 #if HAIR_RENDER_MODE == RENDER_MODE_MSAA || HAIR_RENDER_MODE == RENDER_MODE_PPLL
-	,
-	centroid in float WorldStrandRadius : TEXCOORD8,
-	const float4 PackedVelocityA : TEXCOORD9,
-	const float4 PackedVelocityC : TEXCOORD12
+	, centroid in float WorldStrandRadius : TEXCOORD8
+	, const float4 PackedVelocityA : TEXCOORD9
+	, const float4 PackedVelocityC : TEXCOORD12
 #elif HAIR_RENDER_MODE == RENDER_MODE_COVERAGE 
-	,
-	in float WorldStrandRadius : TEXCOORD8
+	, centroid in float WorldStrandRadius : TEXCOORD8
 #endif
-	OPTIONAL_IsFrontFace
+	  OPTIONAL_IsFrontFace
 #if HAIR_RENDER_MODE == RENDER_MODE_COVERAGE
-	,
-	out float OutColor0 : SV_Target0
+	, out float OutColor0 : SV_Target0
 #elif HAIR_RENDER_MODE == RENDER_MODE_MSAA
-	,
-	out uint4 OutColor0	: SV_Target0
+	, out uint4 OutColor0	: SV_Target0
   #if SUPPORT_MATERIAL_PROPERTY
 	, out float4 OutColor1	: SV_Target1
 	, out float4 OutColor2	: SV_Target2
