Hair PPL depth sorting

TODO
 - Correct coverage from front strand over hiden ones
 - faster compilation by forcing loop using constants
 - TODOAGREGATE

diff --git a/Engine/Shaders/Private/HairStrands/HairStrandsVisibilityCompaction.usf b/Engine/Shaders/Private/HairStrands/HairStrandsVisibilityCompaction.usf
--- a/Engine/Shaders/Private/HairStrands/HairStrandsVisibilityCompaction.usf
+++ b/Engine/Shaders/Private/HairStrands/HairStrandsVisibilityCompaction.usf
@@ -56,7 +56,7 @@ groupshared uint AllocationNodeOffset;
 
 
 #if PERMUTATION_PPLL
-#define PPLL_MAXELEMENTCOUNT PERMUTATION_PPLL
+#define PPLL_MAX_RENDER_NODE_COUNT PERMUTATION_PPLL
 
 Texture2D<uint> PPLLCounter;
 Texture2D<uint> PPLLNodeIndex;
@@ -70,8 +70,6 @@ void MainCS(uint GroupIndex : SV_GroupIndex, uint3 DispatchThreadId : SV_Dispatc
 	const uint FirstNodeIndex = PPLLNodeIndex[PixelCoord];
 
 	FCategorizationData Categorization = (FCategorizationData)0;
-	uint OffsetInStorage = 0;
-	uint PixelSamplecount = 0;
 
 #if PERMUTATION_VELOCITY == VELOCITY_TYPE_AVG
 	float2 AverageVelocity = 0;
@@ -86,17 +84,20 @@ void MainCS(uint GroupIndex : SV_GroupIndex, uint3 DispatchThreadId : SV_Dispatc
 	
 	const bool bExecute = FirstNodeIndex != 0xFFFFFFFF && all(PixelCoord < OutputResolution);
 
-	uint SortedIndex[PPLL_MAXELEMENTCOUNT];
-	float SortedDepth[PPLL_MAXELEMENTCOUNT];
-	float ValidPixelSampleTotalCoverage = 0.0f;
-	float Transmittance = 1.0f;
+	uint SortedIndex[PPLL_MAX_RENDER_NODE_COUNT];
+	float SortedDepth[PPLL_MAX_RENDER_NODE_COUNT];
+	for (int NodeIndex = 0; NodeIndex < PPLL_MAX_RENDER_NODE_COUNT; ++NodeIndex)
+	{
+		SortedDepth[NodeIndex] = 0.0f;
+	}
+
 	uint PixelNodeCount = 0;
 	if (bExecute)
 	{
 		//////////
-		// First pass: Pre compute some values based on the linked list
+		// First pass: sort fragment
 		uint NodeIndex = FirstNodeIndex;
-		while (NodeIndex != 0xFFFFFFFF)// && Transmittance > 0.0f)// && PixelNodeCount < HAIR_VISIBILITY_MAXSAMPLECOUNT)
+		while (NodeIndex != 0xFFFFFFFF)
 		{
 			FPPLLNodeData PPLLNode = PPLLNodeData[NodeIndex];
 
@@ -108,19 +109,39 @@ void MainCS(uint GroupIndex : SV_GroupIndex, uint3 DispatchThreadId : SV_Dispatc
 			float			Specular = 0;
 			DecodeNodeData(ConvertToNodeData(PPLLNode), Depth, VisibilityData, BaseColor, Roughness, Specular);
 
-			const float Coverage = From8bitCoverage(VisibilityData.Coverage8bit);
-			Transmittance *= 1.0f - Coverage;
-			if (PixelNodeCount < PPLL_MAXELEMENTCOUNT)
+			int InsertPos;
+			for (InsertPos = 0; InsertPos < PPLL_MAX_RENDER_NODE_COUNT; ++InsertPos)
+			{
+				if (PPLLNode.Depth > SortedDepth[InsertPos])
+				{
+					// The new node is closer
+					break;
+				}
+			}
+
+			if (InsertPos == PPLL_MAX_RENDER_NODE_COUNT)
 			{
-				ValidPixelSampleTotalCoverage += Coverage;
+				// TODOAGREGATE Need to merge the node into last node
+				// ==> cull it out for now
+			}
+			else
+			{
+				// Shift existing nodes from the end
+				for (int Shift = PPLL_MAX_RENDER_NODE_COUNT - 1; Shift > InsertPos; --Shift) // TODO use PixelNodeCount in place of PPLL_MAX_RENDER_NODE_COUNT
+				{
+					SortedIndex[Shift] = SortedIndex[Shift-1];
+					SortedDepth[Shift] = SortedDepth[Shift-1];
+				}
+
+				// TODOAGREGATE merge last node into agregated material
 
-				// Always insert up to PPLL_MAXELEMENTCOUNT data
-				SortedIndex[PixelNodeCount] = NodeIndex;
-				SortedDepth[PixelNodeCount] = PPLLNode.Depth;
-				// TODO sort and keep closest, merge furthest 
+				// Insert new node
+				SortedIndex[InsertPos] = NodeIndex;
+				SortedDepth[InsertPos] = Depth;
+
+				PixelNodeCount++;
 			}
 
-			PixelNodeCount++;
 			NodeIndex = PPLLNode.NextNodeIndex;
 		}
 	}
@@ -134,8 +155,8 @@ void MainCS(uint GroupIndex : SV_GroupIndex, uint3 DispatchThreadId : SV_Dispatc
 	GroupMemoryBarrierWithGroupSync();
 	// Now notify how many nodes this thread wants to allocate
 	uint PixelDataSubOffsetInStorage = 0; 
-	PixelSamplecount = min(PixelNodeCount, PPLL_MAXELEMENTCOUNT);
-	InterlockedAdd(AllocationNodeCount, PixelSamplecount, PixelDataSubOffsetInStorage);
+	uint PixelRenderNodecount = min(PixelNodeCount, PPLL_MAX_RENDER_NODE_COUNT);
+	InterlockedAdd(AllocationNodeCount, PixelRenderNodecount, PixelDataSubOffsetInStorage);
 	GroupMemoryBarrierWithGroupSync();
 	// Now allocate all the nodes for this group contiguously in memory
 	if (GroupIndex == 0 && AllocationNodeCount > 0)
@@ -143,19 +164,39 @@ void MainCS(uint GroupIndex : SV_GroupIndex, uint3 DispatchThreadId : SV_Dispatc
 		InterlockedAdd(OutCompactNodeCounter[uint2(0, 0)], AllocationNodeCount, AllocationNodeOffset);
 	}
 	GroupMemoryBarrierWithGroupSync();
-	OffsetInStorage = AllocationNodeOffset + PixelDataSubOffsetInStorage;
+	uint OffsetInStorage = AllocationNodeOffset + PixelDataSubOffsetInStorage;
 
 	if(bExecute)
 	{
 		//////////
-		// Second pass: write out compact pixel sample data
-		for (uint i = 0; i < PixelSamplecount; ++i)
+		// Second pass: Compute total coverage for validated nodes
+		float TotalTransmittance = 1.0;
+		float ValidPixelSampleTotalCoverage = 0.0f;
+		for (int i = 0; i < PixelRenderNodecount; ++i)
 		{
 			const uint NodePPLLIndex = SortedIndex[i];
 			FPPLLNodeData PPLLNode = PPLLNodeData[NodePPLLIndex];
 
-			FNodeData OutNodeData = ConvertToNodeData(PPLLNode);
+			// Decode data we need
+			float			Depth = 0.0f;
+			FVisibilityData	VisibilityData = (FVisibilityData)0;
+			float3			BaseColor = 0.0f;
+			float			Roughness = 0.0f;
+			float			Specular = 0;
+			DecodeNodeData(ConvertToNodeData(PPLLNode), Depth, VisibilityData, BaseColor, Roughness, Specular);
 
+			const float Coverage = From8bitCoverage(VisibilityData.Coverage8bit);
+			ValidPixelSampleTotalCoverage += Coverage;
+			TotalTransmittance *= 1.0f - Coverage;
+		}
+
+		//////////
+		// third pass: write out compact pixel sample data
+		for (uint i = 0; i < PixelRenderNodecount; ++i)
+		{
+			const uint NodePPLLIndex = SortedIndex[i];
+			FPPLLNodeData PPLLNode = PPLLNodeData[NodePPLLIndex];
+			FNodeData OutNodeData = ConvertToNodeData(PPLLNode);
 			float			Depth = 0.0f;
 			FVisibilityData	VisibilityData = (FVisibilityData)0;
 			float3			BaseColor = 0.0f;
@@ -172,9 +213,15 @@ void MainCS(uint GroupIndex : SV_GroupIndex, uint3 DispatchThreadId : SV_Dispatc
 			Categorization.TotalCoverage8bit += Coverage8bit;
 			Categorization.ClosestDepth = max(Categorization.ClosestDepth, OutNodeData.Depth); // Inverse-Z
 
+
+
+			// TODOAGREGATE if last, Create FVisibilityData and FNodeData from aggregated  data
+			//FNodeData ConvertToNodeData(float Depth, FVisibilityData VisibilityData, float3 BaseColor, float Roughness, float Specular)
 			OutCompactNodeData[OffsetInStorage + i] = OutNodeData;
 			OutCompactNodeCoord[OffsetInStorage + i] = (PixelCoord.x & 0xFFFF) | ((PixelCoord.y & 0xFFFF) << 16);
 
+
+
 		#if PERMUTATION_VELOCITY != VELOCITY_TYPE_NONE
 			const float2 Velocity = UnpackVelocity(PPLLNode.PackedVelocity) * 0.0f;	// TODO: investigate why TAA jitter is still visible. Disabled for now.
 		#endif
@@ -201,7 +248,7 @@ void MainCS(uint GroupIndex : SV_GroupIndex, uint3 DispatchThreadId : SV_Dispatc
 	#if PERMUTATION_COVERAGE == 1
 		Categorization.PixelCoverage = CoverageTexture.Load(uint3(PixelCoord, 0));
 	#else
-		Categorization.PixelCoverage = 1.0f - Transmittance;
+		Categorization.PixelCoverage = 1.0f - TotalTransmittance;
 	#endif
 
 	#if PERMUTATION_VELOCITY == VELOCITY_TYPE_AVG
@@ -214,11 +261,11 @@ void MainCS(uint GroupIndex : SV_GroupIndex, uint3 DispatchThreadId : SV_Dispatc
 	}
 
 	FNodeDesc NodeDesc;
-	NodeDesc.Count = PixelSamplecount;
+	NodeDesc.Count = PixelRenderNodecount;
 	NodeDesc.Offset = OffsetInStorage;
 	OutCompactNodeIndex[PixelCoord] = EncodeNodeDesc(NodeDesc);
 
-	Categorization.SampleCount = PixelSamplecount;
+	Categorization.SampleCount = PixelRenderNodecount;
 	OutCategorizationTexture[PixelCoord] = EncodeCategorizationData(Categorization);
 }
 
