Visibility for each node is now correctly accounted according to depth sorting.

diff --git a/Engine/Shaders/Private/HairStrands/HairStrandsVisibilityCompaction.usf b/Engine/Shaders/Private/HairStrands/HairStrandsVisibilityCompaction.usf
--- a/Engine/Shaders/Private/HairStrands/HairStrandsVisibilityCompaction.usf
+++ b/Engine/Shaders/Private/HairStrands/HairStrandsVisibilityCompaction.usf
@@ -91,11 +91,12 @@ void MainCS(uint GroupIndex : SV_GroupIndex, uint3 DispatchThreadId : SV_Dispatc
 		SortedDepth[NodeIndex] = 0.0f;
 	}
 
+	float TotalTransmittance = 1.0f;
 	uint PixelNodeCount = 0;
 	if (bExecute)
 	{
 		//////////
-		// First pass: sort fragment
+		// First pass: sort PPLL element into nodes we are going to render
 		uint NodeIndex = FirstNodeIndex;
 		while (NodeIndex != 0xFFFFFFFF)
 		{
@@ -109,6 +110,9 @@ void MainCS(uint GroupIndex : SV_GroupIndex, uint3 DispatchThreadId : SV_Dispatc
 			float			Specular = 0;
 			DecodeNodeData(ConvertToNodeData(PPLLNode), Depth, VisibilityData, BaseColor, Roughness, Specular);
 
+			const float Coverage = From8bitCoverage(VisibilityData.Coverage8bit);
+			TotalTransmittance *= 1.0f - Coverage;
+
 			int InsertPos;
 			for (InsertPos = 0; InsertPos < PPLL_MAX_RENDER_NODE_COUNT; ++InsertPos)
 			{
@@ -169,9 +173,10 @@ void MainCS(uint GroupIndex : SV_GroupIndex, uint3 DispatchThreadId : SV_Dispatc
 	if(bExecute)
 	{
 		//////////
-		// Second pass: Compute total coverage for validated nodes
-		float TotalTransmittance = 1.0;
+		// Second pass: compute total coverage for validated nodes we are going to render
 		float ValidPixelSampleTotalCoverage = 0.0f;
+		float SortedCoverage[PPLL_MAX_RENDER_NODE_COUNT];
+		float TotalSortedTransmittance = 1.0f;
 		for (int i = 0; i < PixelRenderNodecount; ++i)
 		{
 			const uint NodePPLLIndex = SortedIndex[i];
@@ -186,12 +191,19 @@ void MainCS(uint GroupIndex : SV_GroupIndex, uint3 DispatchThreadId : SV_Dispatc
 			DecodeNodeData(ConvertToNodeData(PPLLNode), Depth, VisibilityData, BaseColor, Roughness, Specular);
 
 			const float Coverage = From8bitCoverage(VisibilityData.Coverage8bit);
-			ValidPixelSampleTotalCoverage += Coverage;
-			TotalTransmittance *= 1.0f - Coverage;
+
+			// Update current node coverage as a function of previous nodes coverage
+			SortedCoverage[i] = TotalSortedTransmittance * Coverage;
+			
+			// Update transmittance for the next strands
+			TotalSortedTransmittance *= 1.0f - Coverage;
+
+			// Accumulate total coverage.
+			ValidPixelSampleTotalCoverage += SortedCoverage[i];
 		}
 
 		//////////
-		// third pass: write out compact pixel sample data
+		// Third pass: write out compact nodes for rendering
 		for (uint i = 0; i < PixelRenderNodecount; ++i)
 		{
 			const uint NodePPLLIndex = SortedIndex[i];
@@ -205,7 +217,7 @@ void MainCS(uint GroupIndex : SV_GroupIndex, uint3 DispatchThreadId : SV_Dispatc
 			DecodeNodeData(OutNodeData, Depth, VisibilityData, BaseColor, Roughness, Specular);
 
 			// Coverage8bit is a weight normalising to 1 the contribution of all the compacted samples. Because later it is weighted by Categorization.PixelCoverage.
-			uint Coverage8bit = To8bitCoverage(From8bitCoverage(VisibilityData.Coverage8bit) / float(ValidPixelSampleTotalCoverage));
+			uint Coverage8bit = To8bitCoverage(SortedCoverage[i] / float(ValidPixelSampleTotalCoverage));
 			// Patch the coverage on the out node
 			VisibilityData.Coverage8bit = Coverage8bit;
 			OutNodeData = ConvertToNodeData(Depth, VisibilityData, BaseColor, Roughness, Specular);
