Fixed velocity not being applied correctly due to undefined macro.

diff --git a/Engine/Shaders/Private/HairStrands/HairStrandsVisibilityCommon.ush b/Engine/Shaders/Private/HairStrands/HairStrandsVisibilityCommon.ush
--- a/Engine/Shaders/Private/HairStrands/HairStrandsVisibilityCommon.ush
+++ b/Engine/Shaders/Private/HairStrands/HairStrandsVisibilityCommon.ush
@@ -2,6 +2,10 @@
 
 #pragma once
 
+#define RENDER_MODE_MSAA 0
+#define RENDER_MODE_COVERAGE 1
+#define RENDER_MODE_PPLL 2
+
 struct FVisibilityData
 {
 	float3 Tangent;		// 24 bits
diff --git a/Engine/Shaders/Private/HairStrands/HairStrandsVisibilityCompaction.usf b/Engine/Shaders/Private/HairStrands/HairStrandsVisibilityCompaction.usf
--- a/Engine/Shaders/Private/HairStrands/HairStrandsVisibilityCompaction.usf
+++ b/Engine/Shaders/Private/HairStrands/HairStrandsVisibilityCompaction.usf
@@ -235,7 +235,7 @@ void MainCS(uint GroupIndex : SV_GroupIndex, uint3 DispatchThreadId : SV_Dispatc
 
 
 		#if PERMUTATION_VELOCITY != VELOCITY_TYPE_NONE
-			const float2 Velocity = UnpackVelocity(PPLLNode.PackedVelocity) * 0.0f;	// TODO: investigate why TAA jitter is still visible. Disabled for now.
+			const float2 Velocity = UnpackVelocity(PPLLNode.PackedVelocity);
 		#endif
 		#if PERMUTATION_VELOCITY == VELOCITY_TYPE_AVG
 			AverageVelocity += Velocity;
diff --git a/Engine/Shaders/Private/HairStrands/HairStrandsVisibilityPS.usf b/Engine/Shaders/Private/HairStrands/HairStrandsVisibilityPS.usf
--- a/Engine/Shaders/Private/HairStrands/HairStrandsVisibilityPS.usf
+++ b/Engine/Shaders/Private/HairStrands/HairStrandsVisibilityPS.usf
@@ -13,10 +13,6 @@
 #include "/Engine/Generated/Material.ush"
 #include "/Engine/Generated/VertexFactory.ush"
 
-#define RENDER_MODE_MSAA 0
-#define RENDER_MODE_COVERAGE 1
-#define RENDER_MODE_PPLL 2
-
 #ifndef HAIR_RENDER_MODE
 #error Undefined hair render mode
 #endif
diff --git a/Engine/Shaders/Private/HairStrands/HairStrandsVisibilityVS.usf b/Engine/Shaders/Private/HairStrands/HairStrandsVisibilityVS.usf
--- a/Engine/Shaders/Private/HairStrands/HairStrandsVisibilityVS.usf
+++ b/Engine/Shaders/Private/HairStrands/HairStrandsVisibilityVS.usf
@@ -4,9 +4,7 @@
 #include "../VelocityCommon.ush"
 #include "/Engine/Generated/Material.ush"
 #include "/Engine/Generated/VertexFactory.ush"
-
-#define RENDER_MODE_MSAA 0
-#define RENDER_MODE_COVERAGE 1
+#include "HairStrandsVisibilityCommon.ush"
 
 struct FHairVisibilityAccumulateVSToPS
 {
