Changed min clamp value from 0 to a tiny values to avoid nans/infs in PositiveClampedPow. This now matches the regular ClampedPow functionality

diff --git a/Engine/Shaders/Private/Common.ush b/Engine/Shaders/Private/Common.ush
--- a/Engine/Shaders/Private/Common.ush
+++ b/Engine/Shaders/Private/Common.ush
@@ -438,50 +438,41 @@ MaterialFloat4 UnClampedPow(MaterialFloat4 X, MaterialFloat4 Y)
 	return pow(X, INVARIANT(Y));
 }
 
+#define POW_CLAMP 0.000001f
+
 // Clamp the base, so it's never <= 0.0f (INF/NaN).
 MaterialFloat ClampedPow(MaterialFloat X,MaterialFloat Y)
 {
-	return pow(max(abs(X),0.000001f),INVARIANT(Y));
+	return pow(max(abs(X),POW_CLAMP),Y);
 }
 MaterialFloat2 ClampedPow(MaterialFloat2 X,MaterialFloat2 Y)
 {
-	return pow(max(abs(X),MaterialFloat2(0.000001f,0.000001f)),INVARIANT(Y));
+	return pow(max(abs(X),MaterialFloat2(POW_CLAMP,POW_CLAMP)),Y);
 }
 MaterialFloat3 ClampedPow(MaterialFloat3 X,MaterialFloat3 Y)
 {
-	return pow(max(abs(X),MaterialFloat3(0.000001f,0.000001f,0.000001f)),INVARIANT(Y));
+	return pow(max(abs(X),MaterialFloat3(POW_CLAMP,POW_CLAMP,POW_CLAMP)),Y);
 }  
 MaterialFloat4 ClampedPow(MaterialFloat4 X,MaterialFloat4 Y)
 {
-	return pow(max(abs(X),MaterialFloat4(0.000001f,0.000001f,0.000001f,0.000001f)),INVARIANT(Y));
+	return pow(max(abs(X),MaterialFloat4(POW_CLAMP,POW_CLAMP,POW_CLAMP,POW_CLAMP)),Y);
 } 
+
 MaterialFloat PositiveClampedPow(MaterialFloat X,MaterialFloat Y)
 {
-	return pow(max(X,0.0f),INVARIANT(Y));
+	return pow(max(X,POW_CLAMP),Y);
 }
 MaterialFloat2 PositiveClampedPow(MaterialFloat2 X,MaterialFloat2 Y)
 {
-	return pow(max(X,MaterialFloat2(0.0f,0.0f)),INVARIANT(Y));
-}
-MaterialFloat2 PositiveClampedPow(MaterialFloat2 X,MaterialFloat Y)
-{
-	return pow(max(X,MaterialFloat2(0.0f,0.0f)),MaterialFloat2(INVARIANT(Y),INVARIANT(Y)));
+	return pow(max(X,MaterialFloat2(POW_CLAMP,POW_CLAMP)),Y);
 }
 MaterialFloat3 PositiveClampedPow(MaterialFloat3 X,MaterialFloat3 Y)
 {
-	return pow(max(X,MaterialFloat3(0.0f,0.0f,0.0f)),INVARIANT(Y));
-} 
-MaterialFloat3 PositiveClampedPow(MaterialFloat3 X,MaterialFloat Y)
-{
-	return pow(max(X,MaterialFloat3(0.0f,0.0f,0.0f)),MaterialFloat3(INVARIANT(Y),INVARIANT(Y),INVARIANT(Y)));
-}
+	return pow(max(X,MaterialFloat3(POW_CLAMP,POW_CLAMP,POW_CLAMP)),Y);
+}  
 MaterialFloat4 PositiveClampedPow(MaterialFloat4 X,MaterialFloat4 Y)
 {
-	return pow(max(X,MaterialFloat4(0.0f,0.0f,0.0f,0.0f)),INVARIANT(Y));
-}
-MaterialFloat4 PositiveClampedPow(MaterialFloat4 X,MaterialFloat Y)
-{
-	return pow(max(X,MaterialFloat4(0.0f,0.0f,0.0f,0.0f)),MaterialFloat4(INVARIANT(Y),INVARIANT(Y),INVARIANT(Y),INVARIANT(Y)));
+	return pow(max(X,MaterialFloat4(POW_CLAMP,POW_CLAMP,POW_CLAMP,POW_CLAMP)),Y);
 } 
 
 float DDX(float Input)
