Fixed compression artefacts on landscape for mobile

#rnx

diff --git a/Engine/Shaders/Private/LandscapeVertexFactory.ush b/Engine/Shaders/Private/LandscapeVertexFactory.ush
--- a/Engine/Shaders/Private/LandscapeVertexFactory.ush
+++ b/Engine/Shaders/Private/LandscapeVertexFactory.ush
@@ -121,7 +121,7 @@ struct FVertexFactoryInput
 	float4 Position: ATTRIBUTE0;
 #else
 	float4 PackedPosition: ATTRIBUTE0; 
-	float4 LODHeights[2]: ATTRIBUTE1;  
+	float4 LODHeights[3]: ATTRIBUTE1;  
 #endif
 };
 
@@ -624,8 +624,7 @@ FVertexFactoryIntermediates GetVertexFactoryIntermediates(FVertexFactoryInput In
 	Intermediates.InputPosition = Input.Position;
 #else
 	Intermediates.InputPosition = Input.PackedPosition * 255.f;
-	float2 SubXY = frac((Intermediates.InputPosition.zw) * .5f) * 2.f;
-	float2 LOD45 = (Intermediates.InputPosition.zw - SubXY) / 255.f;
+	float2 SubXY = Intermediates.InputPosition.zw;
 	Intermediates.InputPosition.zw = SubXY;
 #endif
 
@@ -646,24 +645,16 @@ FVertexFactoryIntermediates GetVertexFactoryIntermediates(FVertexFactoryInput In
 	float2 ActualLODCoordsInt = floor(Intermediates.InputPosition.xy * pow(2, -(LodValue - LodValues.x)));
 	float InvLODScaleFactor = pow(2, -LodValue);
 #else
-	float3 HeightData = (LodValue < 1 ? float3(1.0, Input.LODHeights[1][0], Input.LODHeights[1][1]) :
-						(LodValue < 2 ? float3(0.5, Input.LODHeights[1][1], Input.LODHeights[1][2]) :
-						(LodValue < 3 ? float3(0.25, Input.LODHeights[1][2], Input.LODHeights[1][3]) :
-						(LodValue < 4 ? float3(0.125, Input.LODHeights[1][3], LOD45[0]) :
-						(LodValue < 5 ? float3(0.0625, LOD45[0], LOD45[1]) :
-						(float3(0.03125, LOD45[1], LOD45[1])))))));
-
-	// we can hardcode InvLODScaleFactor
-	float InvLODScaleFactor = HeightData.x;
-	// select the height values from the vertex input corresponding to the current and next LOD
-	float InputHeight = HeightData.y;
-	float InputHeightNextLOD = HeightData.z;
-
-	float MinHeight = DecodePackedHeight(Input.LODHeights[0].xy);
-	float MaxHeight = DecodePackedHeight(Input.LODHeights[0].zw);
+	float InvLODScaleFactor = 1.0f / (float)(1 << (uint)LodValue);
+	uint LODHeightU = (uint)floor(LodValue * 0.5);
+	uint LODHeightV = Mod((uint)LodValue, 2) * 2;
+	float NextLodValue = LodValue + 1;
+	uint NextLODHeightU = (uint)floor(NextLodValue * 0.5);
+	uint NextLODHeightV = Mod((uint)NextLodValue, 2) * 2;
 
-	float Height = lerp(MinHeight, MaxHeight, InputHeight);
-	float HeightNextLOD = lerp(MinHeight, MaxHeight, InputHeightNextLOD);
+	// select the height values from the vertex input corresponding to the current and next LOD
+	float Height = DecodePackedHeight(float2(Input.LODHeights[LODHeightU][LODHeightV], Input.LODHeights[LODHeightU][LODHeightV + 1]));
+	float HeightNextLOD = DecodePackedHeight(float2(Input.LODHeights[NextLODHeightU][NextLODHeightV], Input.LODHeights[NextLODHeightU][NextLODHeightV + 1]));
 
 	// LodValues.x is always 0 on mobile.
 	float2 ActualLODCoordsInt = floor(Intermediates.InputPosition.xy * InvLODScaleFactor);
diff --git a/Engine/Source/Runtime/Landscape/Private/LandscapeEdit.cpp b/Engine/Source/Runtime/Landscape/Private/LandscapeEdit.cpp
--- a/Engine/Source/Runtime/Landscape/Private/LandscapeEdit.cpp
+++ b/Engine/Source/Runtime/Landscape/Private/LandscapeEdit.cpp
@@ -5952,7 +5952,6 @@ void ULandscapeComponent::GeneratePlatformVertexData(const ITargetPlatform* Targ
 		TArray<int32> MipHeights;
 		MipHeights.AddZeroed(HeightmapMipData.Num());
 		int32 LastIndex = 0;
-		uint16 MaxHeight = 0, MinHeight = 65535;
 
 		for (int32 Mip = 0; Mip < HeightmapMipData.Num(); ++Mip)
 		{
@@ -5968,25 +5967,12 @@ void ULandscapeComponent::GeneratePlatformVertexData(const ITargetPlatform* Targ
 			uint16 Height = CurrentMipSrcRow[MipX].R << 8 | CurrentMipSrcRow[MipX].G;
 
 			MipHeights[Mip] = Height;
-			MaxHeight = FMath::Max(MaxHeight, Height);
-			MinHeight = FMath::Min(MinHeight, Height);
 		}
 
-		DstVert->LODHeights[0] = MinHeight >> 8;
-		DstVert->LODHeights[1] = (MinHeight & 255);
-		DstVert->LODHeights[2] = MaxHeight >> 8;
-		DstVert->LODHeights[3] = (MaxHeight & 255);
-
 		for (int32 Mip = 0; Mip < HeightmapMipData.Num(); ++Mip)
 		{
-			if (Mip < 4)
-			{
-				DstVert->LODHeights[4 + Mip] = FMath::RoundToInt(float(MipHeights[Mip] - MinHeight) / (MaxHeight - MinHeight) * 255);
-			}
-			else // Mip 4 5 packed into SubX, SubY
-			{
-				DstVert->Position[Mip - 2] += (FMath::RoundToInt(float(MipHeights[Mip] - MinHeight) / (MaxHeight - MinHeight) * 255)) & (0xfffe);
-			}
+			DstVert->LODHeights[Mip * 2] = MipHeights[Mip] >> 8;
+			DstVert->LODHeights[Mip * 2 + 1] = MipHeights[Mip] & 255;
 		}
 
 		DstVert++;
diff --git a/Engine/Source/Runtime/Landscape/Public/LandscapeRenderMobile.h b/Engine/Source/Runtime/Landscape/Public/LandscapeRenderMobile.h
--- a/Engine/Source/Runtime/Landscape/Public/LandscapeRenderMobile.h
+++ b/Engine/Source/Runtime/Landscape/Public/LandscapeRenderMobile.h
@@ -13,12 +13,12 @@ LandscapeRenderMobile.h: Mobile landscape rendering
 #include "LandscapeRender.h"
 #include "Runtime/Landscape/Private/LandscapePrivate.h"
 
-#define LANDSCAPE_MAX_ES_LOD_COMP	2
+#define LANDSCAPE_MAX_ES_LOD_COMP	3
 #define LANDSCAPE_MAX_ES_LOD		6
 
 struct FLandscapeMobileVertex
 {
-	uint8 Position[4]; // Pos + LOD 0 Height
+	uint8 Position[4]; // X, Y, SubX, SubY
 	uint8 LODHeights[LANDSCAPE_MAX_ES_LOD_COMP*4];
 };
 
