Cherry-pick 10068683: Includes mendatory Platform.ush in some hair shaders

Note: Probably duplicated code with one of the clearing utils?

(ushell-p4-cherrypick of 10068690 by [REDACTED_PERSON])

diff --git a/Engine/Shaders/Private/HairStrands/HairStrandsVisibilityClearPS.usf b/Engine/Shaders/Private/HairStrands/HairStrandsVisibilityClearPS.usf
--- a/Engine/Shaders/Private/HairStrands/HairStrandsVisibilityClearPS.usf
+++ b/Engine/Shaders/Private/HairStrands/HairStrandsVisibilityClearPS.usf
@@ -1,5 +1,8 @@
 // Copyright 1998-2019 Epic Games, Inc. All Rights Reserved.
 
+#include "/Engine/Public/Platform.ush"
+
+
 uint ClearValue;
 void ClearPS(out uint4 OutColor : SV_Target0)
 {
diff --git a/Engine/Shaders/Private/HairStrands/HairStrandsVoxelMip.usf b/Engine/Shaders/Private/HairStrands/HairStrandsVoxelMip.usf
--- a/Engine/Shaders/Private/HairStrands/HairStrandsVoxelMip.usf
+++ b/Engine/Shaders/Private/HairStrands/HairStrandsVoxelMip.usf
@@ -1,5 +1,8 @@
 // Copyright 1998-2019 Epic Games, Inc. All Rights Reserved.
 
+#include "/Engine/Public/Platform.ush"
+
+
 uint VoxelResolution;
 uint SourceMip;
 uint TargetMip;
