Reorganizing data layout of gather points buffer. Data is now organized in planes of samples, to be more cache friendly. Based on performance work done by [REDACTED_PERSON] [at] NVIDIA.

diff --git a/Engine/Shaders/Private/RayTracing/RayTracingCreateGatherPointsRGS.usf b/Engine/Shaders/Private/RayTracing/RayTracingCreateGatherPointsRGS.usf
--- a/Engine/Shaders/Private/RayTracing/RayTracingCreateGatherPointsRGS.usf
+++ b/Engine/Shaders/Private/RayTracing/RayTracingCreateGatherPointsRGS.usf
@@ -69,9 +69,8 @@ RAY_TRACING_ENTRY_RAYGEN(RayTracingCreateGatherPointsRGS)
 	for (uint Index = 0; Index < SamplesPerPixel; ++Index)
 	{
 		uint SampleIndexOffset = SampleIndex + Index;
-		RWGatherPointsBuffer[GatherPointsIndex].CreationPoint[SampleIndexOffset] = 0.0;
-		RWGatherPointsBuffer[GatherPointsIndex].Position[SampleIndexOffset] = 0.0;
-		RWGatherPointsBuffer[GatherPointsIndex].Irradiance[SampleIndexOffset] = 0.0;
+		FGatherSample GatherSample = (FGatherSample)0.0;
+		WriteGatherSample(RWGatherPointsBuffer, GatherPointsIndex, SampleIndexOffset, GatherPointsResolution.x * GatherPointsResolution.y, GatherSample);
 	}
 #endif
 	// Get G-Buffer surface data
@@ -243,8 +242,10 @@ RAY_TRACING_ENTRY_RAYGEN(RayTracingCreateGatherPointsRGS)
 			}
 		}
 
-		RWGatherPointsBuffer[GatherPointsIndex].CreationPoint[SampleIndexOffset] = WorldPosition;
-		RWGatherPointsBuffer[GatherPointsIndex].Position[SampleIndexOffset] = Position;
-		RWGatherPointsBuffer[GatherPointsIndex].Irradiance[SampleIndexOffset] = Irradiance;
+		FGatherSample GatherSample = (FGatherSample)0.0;
+		GatherSample.CreationPoint = WorldPosition;
+		GatherSample.Position = Position;
+		GatherSample.Irradiance = Irradiance;
+		WriteGatherSample(RWGatherPointsBuffer, GatherPointsIndex, SampleIndexOffset, GatherPointsResolution.x * GatherPointsResolution.y, GatherSample);
 	}
 }
diff --git a/Engine/Shaders/Private/RayTracing/RayTracingFinalGatherRGS.usf b/Engine/Shaders/Private/RayTracing/RayTracingFinalGatherRGS.usf
--- a/Engine/Shaders/Private/RayTracing/RayTracingFinalGatherRGS.usf
+++ b/Engine/Shaders/Private/RayTracing/RayTracingFinalGatherRGS.usf
@@ -88,7 +88,7 @@ RAY_TRACING_ENTRY_RAYGEN(RayTracingFinalGatherRGS)
 	// Define rejection criteria and evaluate
 	bool bShouldReject[16];
 	FRejectionCriteria RejectionCriteria = CreateRejectionCriteria(WorldPosition, FinalGatherDistance);
-	CreateRejectionMask(RejectionCriteria, GatherPointsBuffer[GatherPointsIndex], SamplesPerPixel, bShouldReject);
+	CreateRejectionMask(RejectionCriteria, GatherPointsBuffer, GatherPointsIndex, 0, GatherPointsResolution.x * GatherPointsResolution.y, SamplesPerPixel, bShouldReject);
 
 	float3 ExitantRadiance = 0.0;
 	float HitDistance = 1.0e7;
@@ -103,8 +103,9 @@ RAY_TRACING_ENTRY_RAYGEN(RayTracingFinalGatherRGS)
 		}
 		ValidSamplesPerPixel++;
 
-		float3 GPPosition = GatherPointsBuffer[GatherPointsIndex].Position[LocalSampleIndex];
-		float3 GPIrradiance = GatherPointsBuffer[GatherPointsIndex].Irradiance[LocalSampleIndex];
+		FGatherSample GatherSample = ReadGatherSample(GatherPointsBuffer, GatherPointsIndex, LocalSampleIndex, GatherPointsResolution.x * GatherPointsResolution.y);
+		float3 GPPosition = GatherSample.Position;
+		float3 GPIrradiance = GatherSample.Irradiance;
 
 		// Initialize throughput based on BRDF
 		half3 N = WorldNormal;
diff --git a/Engine/Shaders/Private/RayTracing/RayTracingGatherPoints.ush b/Engine/Shaders/Private/RayTracing/RayTracingGatherPoints.ush
--- a/Engine/Shaders/Private/RayTracing/RayTracingGatherPoints.ush
+++ b/Engine/Shaders/Private/RayTracing/RayTracingGatherPoints.ush
@@ -2,9 +2,16 @@
 
 struct FGatherPoints
 {
-	float3 CreationPoint[16];
-	float3 Position[16];
-	float3 Irradiance[16];
+	float3 CreationPoint;
+	float3 Position;
+	uint2 Irradiance;
+};
+
+struct FGatherSample
+{
+	float3 CreationPoint;
+	float3 Position;
+	float3 Irradiance;
 };
 
 struct FRejectionCriteria
@@ -13,6 +20,48 @@ struct FRejectionCriteria
 	float DistanceThreshold2;
 };
 
+uint2 PackIrradiance(float3 Irradiance)
+{
+	uint2 Packed = (uint2)0;
+	Irradiance = ClampToHalfFloatRange(Irradiance);
+
+	// Note: Consider keeping Lumiance(Y) in full range
+	Packed.x = f32tof16(Irradiance.x) | (f32tof16(Irradiance.y) << 16);
+	Packed.y = f32tof16(Irradiance.z);
+	return Packed;
+}
+
+float3 UnpackIrradiance(uint2 Irradiance)
+{
+	float3 Unpacked;
+
+	Unpacked.x = f16tof32(Irradiance.x & 0xffff);
+	Unpacked.y = f16tof32(Irradiance.x >> 16);
+	Unpacked.z = f16tof32(Irradiance.y & 0xffff);
+
+	return Unpacked;
+}
+
+FGatherSample ReadGatherSample(StructuredBuffer<FGatherPoints> GatherBuffer, uint BaseIndex, uint SampleIndex, uint ElementCount)
+{
+	FGatherSample GatherSample = (FGatherSample) 0.0;
+
+	uint Index = BaseIndex + SampleIndex * ElementCount;
+	GatherSample.CreationPoint = GatherBuffer[Index].CreationPoint;
+	GatherSample.Position = GatherBuffer[Index].Position;
+	GatherSample.Irradiance = UnpackIrradiance(GatherBuffer[Index].Irradiance);
+
+	return GatherSample;
+}
+
+void WriteGatherSample(RWStructuredBuffer<FGatherPoints> GatherBuffer, uint BaseIndex, uint SampleIndex, uint ElementCount, FGatherSample GatherSample)
+{
+	uint Index = BaseIndex + SampleIndex * ElementCount;
+	GatherBuffer[Index].CreationPoint = GatherSample.CreationPoint;
+	GatherBuffer[Index].Position = GatherSample.Position;
+	GatherBuffer[Index].Irradiance = PackIrradiance(GatherSample.Irradiance);
+}
+
 FRejectionCriteria CreateRejectionCriteria(float3 CreationPoint, float DistanceThreshold)
 {
 	FRejectionCriteria RejectionCriteria;
@@ -21,11 +70,12 @@ FRejectionCriteria CreateRejectionCriteria(float3 CreationPoint, float DistanceT
 	return RejectionCriteria;
 }
 
-void CreateRejectionMask(FRejectionCriteria RejectionCriteria, FGatherPoints GatherPoints, int SampleCount, out bool bShouldReject[16])
+void CreateRejectionMask(FRejectionCriteria RejectionCriteria, StructuredBuffer<FGatherPoints> GatherBuffer, uint BaseIndex, uint SampleIndex, uint ElementCount, int SampleCount, out bool bShouldReject[16])
 {
-	for (int SampleIndex = 0; SampleIndex < SampleCount; ++SampleIndex)
+	for (int Index = 0; Index < SampleCount; ++Index)
 	{
-		float3 GPDistance2 = RejectionCriteria.CreationPoint - GatherPoints.CreationPoint[SampleIndex];
-		bShouldReject[SampleIndex] = dot(GPDistance2, GPDistance2) > RejectionCriteria.DistanceThreshold2;
+		FGatherSample GatherSample = ReadGatherSample(GatherBuffer, BaseIndex, SampleIndex + Index, ElementCount);
+		float3 GPDistance2 = RejectionCriteria.CreationPoint - GatherSample.CreationPoint;
+		bShouldReject[Index] = dot(GPDistance2, GPDistance2) > RejectionCriteria.DistanceThreshold2;
 	}
 }
diff --git a/Engine/Source/Runtime/Renderer/Private/RayTracing/RayTracingGlobalIllumination.cpp b/Engine/Source/Runtime/Renderer/Private/RayTracing/RayTracingGlobalIllumination.cpp
--- a/Engine/Source/Runtime/Renderer/Private/RayTracing/RayTracingGlobalIllumination.cpp
+++ b/Engine/Source/Runtime/Renderer/Private/RayTracing/RayTracingGlobalIllumination.cpp
@@ -306,11 +306,13 @@ class FGlobalIlluminationRGS : public FGlobalShader
 
 IMPLEMENT_GLOBAL_SHADER(FGlobalIlluminationRGS, "/Engine/Private/RayTracing/RayTracingGlobalIlluminationRGS.usf", "GlobalIlluminationRGS", SF_RayGen);
 
-struct GatherPoints
+// Placeholder structure used to allocate gather points buffer
+// #dxr_todo: rework to minimal set, based on active samples-per-pixel
+struct FGatherPoints
 {
 	FVector CreationPoint[16];
 	FVector Position[16];
-	FVector Irradiance[16];
+	FIntPoint Irradiance[16];
 };
 
 class FRayTracingGlobalIlluminationCreateGatherPointsRGS : public FGlobalShader
@@ -586,7 +588,7 @@ void FDeferredShadingSceneRenderer::RayTracingGlobalIlluminationCreateGatherPoin
 	if (GatherPointsResolution != LocalGatherPointsResolution)
 	{
 		GatherPointsResolution = LocalGatherPointsResolution;
-		FRDGBufferDesc BufferDesc = FRDGBufferDesc::CreateStructuredDesc(sizeof(GatherPoints), GatherPointsResolution.X * GatherPointsResolution.Y);
+		FRDGBufferDesc BufferDesc = FRDGBufferDesc::CreateStructuredDesc(sizeof(FGatherPoints), GatherPointsResolution.X * GatherPointsResolution.Y);
 		GatherPointsBuffer = GraphBuilder.CreateBuffer(BufferDesc, TEXT("GatherPointsBuffer"), ERDGResourceFlags::MultiFrame);
 	}
 	else
