Correct the stationary sky light contribution for two sided primtives.

diff --git a/Engine/Shaders/Private/MobileBasePassPixelShader.usf b/Engine/Shaders/Private/MobileBasePassPixelShader.usf
--- a/Engine/Shaders/Private/MobileBasePassPixelShader.usf
+++ b/Engine/Shaders/Private/MobileBasePassPixelShader.usf
@@ -634,11 +634,20 @@ void Main(
 		#if ENABLE_SKY_LIGHT
 			//@mw todo
 			// TODO: Also need to do specular.
-			float3 SkyDiffuseLighting = GetSkySHDiffuseSimple(MaterialParameters.WorldNormal);
+#if MATERIAL_TWO_SIDED && LQ_TEXTURE_LIGHTMAP
+			if (NoL == 0)
+			{
+#endif
+			half3 SkyDiffuseLighting = GetSkySHDiffuseSimple(MaterialParameters.WorldNormal);
+
 		#if MATERIAL_SHADINGMODEL_SINGLELAYERWATER
 			WaterDiffuseIndirectLuminance += SkyDiffuseLighting;
 		#endif
-			Color += SkyDiffuseLighting * ResolvedView.SkyLightColor.rgb * DiffuseColor * MaterialAO;
+			Color += SkyDiffuseLighting * half3(ResolvedView.SkyLightColor.rgb) * DiffuseColor * MaterialAO;
+#if MATERIAL_TWO_SIDED && LQ_TEXTURE_LIGHTMAP
+			
+			}
+#endif
 		#endif
 
 		#if MAX_DYNAMIC_POINT_LIGHTS > 0 && !MATERIAL_SHADINGMODEL_SINGLELAYERWATER
diff --git a/Engine/Source/Runtime/Engine/Private/Materials/HLSLMaterialTranslator.h b/Engine/Source/Runtime/Engine/Private/Materials/HLSLMaterialTranslator.h
--- a/Engine/Source/Runtime/Engine/Private/Materials/HLSLMaterialTranslator.h
+++ b/Engine/Source/Runtime/Engine/Private/Materials/HLSLMaterialTranslator.h
@@ -1429,6 +1429,7 @@ ResourcesString = TEXT("");
 		OutEnvironment.SetDefine(TEXT("MATERIAL_IS_SKY"), Material->IsSky());
 		OutEnvironment.SetDefine(TEXT("MATERIAL_COMPUTE_FOG_PER_PIXEL"), Material->ComputeFogPerPixel());
 		OutEnvironment.SetDefine(TEXT("MATERIAL_FULLY_ROUGH"), bIsFullyRough || Material->IsFullyRough());
+		OutEnvironment.SetDefine(TEXT("MATERIAL_TWO_SIDED"), Material->IsTwoSided());
 
 		// Count the number of VTStacks (each stack will allocate a feedback slot)
 		OutEnvironment.SetDefine(TEXT("NUM_VIRTUALTEXTURE_SAMPLES"), VTStacks.Num());
