Fix erroneous calculation of NEE Pdfs for enviornment contribution off of diracs.

diff --git a/Engine/Shaders/Private/PathTracing/PathTracing.usf b/Engine/Shaders/Private/PathTracing/PathTracing.usf
--- a/Engine/Shaders/Private/PathTracing/PathTracing.usf
+++ b/Engine/Shaders/Private/PathTracing/PathTracing.usf
@@ -81,6 +81,7 @@ RAY_TRACING_ENTRY_RAYGEN(PathTracingMainRG)
 
 		float MaterialPdf = 0.0;
 		uint Bounces = 0;
+		bool bIsPureSpecular = false;
 		while (Bounces < PathTracingData.MaxBounces)
 		{
 			// Environment contribution
@@ -96,7 +97,7 @@ RAY_TRACING_ENTRY_RAYGEN(PathTracingMainRG)
 					{
 						// Apply NEE Pdf for correct MIS weight
 						float MisWeight = 1.0;
-						if (IsNextEventEstimationEnabled && Bounces > 0)
+						if (IsNextEventEstimationEnabled && Bounces > 0 && !bIsPureSpecular)
 						{
 							float NeePdf = 0.0;
 							float3 LightUV = GetLightUV(SkyLightId, Ray, Payload);
@@ -149,8 +150,9 @@ RAY_TRACING_ENTRY_RAYGEN(PathTracingMainRG)
 			Ray.Origin = Ray.Origin + Ray.Direction * Payload.HitT;
 			FMaterialClosestHitPayload RayHitInfo = Payload; 
 
+			bIsPureSpecular = IsPureSpecular(RayHitInfo);
+
 			// Perform next-event estimation
-			bool bIsPureSpecular = IsPureSpecular(RayHitInfo);
 			if (IsNextEventEstimationEnabled && !bIsPureSpecular)
 			{
 				int LightId;
