Fix old path tracing regression where the payload was expected to carry previous material hit information upon a miss (for MIS with next-event estimation).

diff --git a/Engine/Shaders/Private/PathTracing/PathTracing.usf b/Engine/Shaders/Private/PathTracing/PathTracing.usf
--- a/Engine/Shaders/Private/PathTracing/PathTracing.usf
+++ b/Engine/Shaders/Private/PathTracing/PathTracing.usf
@@ -80,6 +80,7 @@ RAY_TRACING_ENTRY_RAYGEN(PathTracingMainRG)
 		RayCounter++;
 
 		float MaterialPdf = 0.0;
+		FMaterialClosestHitPayload PrevMaterialPayload;
 		uint Bounces = 0;
 		bool bIsPureSpecular = false;
 		while (Bounces < PathTracingData.MaxBounces)
@@ -101,7 +102,7 @@ RAY_TRACING_ENTRY_RAYGEN(PathTracingMainRG)
 						{
 							float NeePdf = 0.0;
 							float3 LightUV = GetLightUV(SkyLightId, Ray, Payload);
-							PdfLight(Ray, Payload, SkyLightId, LightUV, NeePdf);
+							PdfLight(Ray, PrevMaterialPayload, SkyLightId, LightUV, NeePdf);
 
 							// Note: MaterialPdf will need to be stored on the ray for an iterative solution
 							MisWeight = MaterialPdf / (MaterialPdf + NeePdf);
@@ -130,7 +131,7 @@ RAY_TRACING_ENTRY_RAYGEN(PathTracingMainRG)
 						float NeePdf = 0.0;
 						int LightId = GetLightId(Payload);
 						float3 LightUV = GetLightUV(LightId, Ray, Payload);
-						PdfLight(Ray, Payload, LightId, LightUV, NeePdf);
+						PdfLight(Ray, PrevMaterialPayload, LightId, LightUV, NeePdf);
 
 						float MisWeight = MaterialPdf / (MaterialPdf + NeePdf);
 					}
@@ -267,6 +268,7 @@ RAY_TRACING_ENTRY_RAYGEN(PathTracingMainRG)
 			const uint RayFlags = RAY_FLAG_CULL_BACK_FACING_TRIANGLES;
 			const uint InstanceInclusionMask = RAY_TRACING_MASK_ALL;
 
+			PrevMaterialPayload = Payload;
 			Payload = TraceMaterialRay(
 				TLAS,
 				RayFlags,
