Fix erroneous calculation of NEE Pdfs for emissive contribution off of diracs.

diff --git a/Engine/Shaders/Private/PathTracing/PathTracing.usf b/Engine/Shaders/Private/PathTracing/PathTracing.usf
--- a/Engine/Shaders/Private/PathTracing/PathTracing.usf
+++ b/Engine/Shaders/Private/PathTracing/PathTracing.usf
@@ -126,7 +126,7 @@ RAY_TRACING_ENTRY_RAYGEN(PathTracingMainRG)
 				{
 					// Apply NEE Pdf for correct MIS weight
 					float MisWeight = 1.0;
-					if (IsNextEventEstimationEnabled && Bounces > 0)
+					if (IsNextEventEstimationEnabled && Bounces > 0 && !bIsPureSpecular)
 					{
 						float NeePdf = 0.0;
 						int LightId = GetLightId(Payload);
