Fix more incorrect NEE Pdf calculations due to not taking exact material determination and ordering into account.

diff --git a/Engine/Shaders/Private/PathTracing/Material/PathTracingMaterialSampling.ush b/Engine/Shaders/Private/PathTracing/Material/PathTracingMaterialSampling.ush
--- a/Engine/Shaders/Private/PathTracing/Material/PathTracingMaterialSampling.ush
+++ b/Engine/Shaders/Private/PathTracing/Material/PathTracingMaterialSampling.ush
@@ -28,6 +28,34 @@ void SampleMaterialAsWorldNormal(
 	OutPdf = 1.0;
 }
 
+bool IsNeeValidMaterial(FMaterialClosestHitPayload Payload)
+{
+	// Checks done in the same order as other sampling and evaluation functions to ensure consistency
+	if (IsPureSpecularReflection(Payload))
+	{
+		return false;
+	}
+	else if (IsDefaultLitMaterial(Payload))
+	{
+		return true;
+	}
+	else if (IsClearCoatMaterial(Payload))
+	{
+		return true;
+	}
+	else if (IsPureSpecularTransmission(Payload))
+	{
+		return false;
+	}
+	else if(IsLambert(Payload))
+	{
+		return true;
+	}
+	else
+	{
+		return true;
+	}
+}
 
 void SampleMaterial(
 	float3 RayDirection,
diff --git a/Engine/Shaders/Private/PathTracing/PathTracing.usf b/Engine/Shaders/Private/PathTracing/PathTracing.usf
--- a/Engine/Shaders/Private/PathTracing/PathTracing.usf
+++ b/Engine/Shaders/Private/PathTracing/PathTracing.usf
@@ -82,7 +82,7 @@ RAY_TRACING_ENTRY_RAYGEN(PathTracingMainRG)
 		float MaterialPdf = 0.0;
 		FMaterialClosestHitPayload PrevMaterialPayload;
 		uint Bounces = 0;
-		bool bIsPureSpecular = false;
+		bool bIsNeeValid = true;
 		while (Bounces < PathTracingData.MaxBounces)
 		{
 			// Environment contribution
@@ -98,7 +98,7 @@ RAY_TRACING_ENTRY_RAYGEN(PathTracingMainRG)
 					{
 						// Apply NEE Pdf for correct MIS weight
 						float MisWeight = 1.0;
-						if (IsNextEventEstimationEnabled && Bounces > 0 && !bIsPureSpecular)
+						if (IsNextEventEstimationEnabled && Bounces > 0 && bIsNeeValid)
 						{
 							float NeePdf = 0.0;
 							float3 LightUV = GetLightUV(SkyLightId, Ray, Payload);
@@ -126,7 +126,7 @@ RAY_TRACING_ENTRY_RAYGEN(PathTracingMainRG)
 				{
 					// Apply NEE Pdf for correct MIS weight
 					float MisWeight = 1.0;
-					if (IsNextEventEstimationEnabled && Bounces > 0 && !bIsPureSpecular)
+					if (IsNextEventEstimationEnabled && Bounces > 0 && bIsNeeValid)
 					{
 						float NeePdf = 0.0;
 						int LightId = GetLightId(Payload);
@@ -151,10 +151,10 @@ RAY_TRACING_ENTRY_RAYGEN(PathTracingMainRG)
 			Ray.Origin = Ray.Origin + Ray.Direction * Payload.HitT;
 			FMaterialClosestHitPayload RayHitInfo = Payload; 
 
-			bIsPureSpecular = IsPureSpecular(RayHitInfo);
+			bIsNeeValid = IsNeeValidMaterial(RayHitInfo);
 
 			// Perform next-event estimation
-			if (IsNextEventEstimationEnabled && !bIsPureSpecular)
+			if (IsNextEventEstimationEnabled && bIsNeeValid)
 			{
 				int LightId;
 				float3 LightUV;
