Initialize the path tracer's previous material payload to fix validation errors.

diff --git a/Engine/Shaders/Private/PathTracing/PathTracing.usf b/Engine/Shaders/Private/PathTracing/PathTracing.usf
--- a/Engine/Shaders/Private/PathTracing/PathTracing.usf
+++ b/Engine/Shaders/Private/PathTracing/PathTracing.usf
@@ -80,7 +80,7 @@ RAY_TRACING_ENTRY_RAYGEN(PathTracingMainRG)
 		RayCounter++;
 
 		float MaterialPdf = 0.0;
-		FMaterialClosestHitPayload PrevMaterialPayload;
+		FMaterialClosestHitPayload PrevMaterialPayload = (FMaterialClosestHitPayload)0;
 		uint Bounces = 0;
 		bool bIsNeeValid = true;
 		while (Bounces < PathTracingData.MaxBounces)
