Fixed shader compilation in path tracing shader

diff --git a/Engine/Shaders/Private/PathTracing/Material/PathTracingMaterialSampling.ush b/Engine/Shaders/Private/PathTracing/Material/PathTracingMaterialSampling.ush
--- a/Engine/Shaders/Private/PathTracing/Material/PathTracingMaterialSampling.ush
+++ b/Engine/Shaders/Private/PathTracing/Material/PathTracingMaterialSampling.ush
@@ -35,14 +35,6 @@ bool IsNeeValidMaterial(FMaterialClosestHitPayload Payload)
 	{
 		return false;
 	}
-	else if (IsDefaultLitMaterial(Payload))
-	{
-		return true;
-	}
-	else if (IsClearCoatMaterial(Payload))
-	{
-		return true;
-	}
 	else if (IsPureSpecularTransmission(Payload))
 	{
 		return false;
