Added missing #include "LightGridCommon.ush" needed for ForwardLightDataISR in "ReflectionEnvironmentComposite.ush"

diff --git a/Engine/Shaders/Private/SingleLayerWaterComposite.usf b/Engine/Shaders/Private/SingleLayerWaterComposite.usf
--- a/Engine/Shaders/Private/SingleLayerWaterComposite.usf
+++ b/Engine/Shaders/Private/SingleLayerWaterComposite.usf
@@ -10,6 +10,7 @@
 #include "ReflectionEnvironmentShared.ush"
 #include "ShadingModels.ush"
 #include "SceneTextureParameters.ush"
+#include "LightGridCommon.ush"
 
 // Scale for encoding depth into SceneDepthWithoutSingleLayerWaterTexture
 #define SINGLE_LAYER_WATER_DEPTH_SCALE 100.0f
