Fixed shader compilation error due invalid usage of packed payload

diff --git a/Engine/Shaders/Private/RayTracing/RayTracingReflections.usf b/Engine/Shaders/Private/RayTracing/RayTracingReflections.usf
--- a/Engine/Shaders/Private/RayTracing/RayTracingReflections.usf
+++ b/Engine/Shaders/Private/RayTracing/RayTracingReflections.usf
@@ -695,7 +695,7 @@ RAY_TRACING_ENTRY_RAYGEN(RayTracingReflectionsRGS)
 					BottomLayerSpecularColor = TopLayerSpecularColor;
 					BottomLayerWorldNormal = TopLayerWorldNormal;
 					#if CLEAR_COAT_BOTTOM_NORMAL
-					const float2 oct1 = ((float2(Payload.CustomData.a, Payload.CustomData.z) * 2) - (256.0 / 255.0)) + UnitVectorToOctahedron(TopLayerWorldNormal);
+					const float2 oct1 = ((float2(Payload.GetCustomData().a, Payload.GetCustomData().z) * 2) - (256.0 / 255.0)) + UnitVectorToOctahedron(TopLayerWorldNormal);
 					BottomLayerWorldNormal = OctahedronToUnitVector(oct1);
 					#endif
 
