Fix duplicate include after CL# 10086128 merged to UE4/Main
#rnx

diff --git a/Engine/Shaders/Private/SingleLayerWaterComposite.usf b/Engine/Shaders/Private/SingleLayerWaterComposite.usf
--- a/Engine/Shaders/Private/SingleLayerWaterComposite.usf
+++ b/Engine/Shaders/Private/SingleLayerWaterComposite.usf
@@ -17,8 +17,6 @@
 
 #define ENABLE_SKY_LIGHT 1
 
-#include "LightGridCommon.ush"
-
 #define REFLECTION_COMPOSITE_USE_BLENDED_REFLECTION_CAPTURES (FEATURE_LEVEL >= FEATURE_LEVEL_SM5)
 #define REFLECTION_COMPOSITE_SUPPORT_SKYLIGHT_BLEND 0
 #include "ReflectionEnvironmentComposite.ush"
