Add a flag to disable Sky Light contribution through ray payloads.

(ushell-p4-cherrypick of 9621763 by riley.alston)

diff --git a/Engine/Shaders/Private/PathTracing/PathTracing.usf b/Engine/Shaders/Private/PathTracing/PathTracing.usf
--- a/Engine/Shaders/Private/PathTracing/PathTracing.usf
+++ b/Engine/Shaders/Private/PathTracing/PathTracing.usf
@@ -70,13 +70,15 @@ RAY_TRACING_ENTRY_RAYGEN(PathTracingMainRG)
 		// Cast ray
 		const uint InitialRayFlags = RAY_FLAG_CULL_BACK_FACING_TRIANGLES;
 		const uint InitialInstanceInclusionMask = RAY_TRACING_MASK_ALL;
+		const bool bInitialDisableSkyLightContribution = false;
 
 		FMaterialClosestHitPayload Payload = TraceMaterialRay(
 			TLAS,
 			InitialRayFlags,
 			InitialInstanceInclusionMask,
 			Ray,
-			RayCone);
+			RayCone,
+			bInitialDisableSkyLightContribution);
 		RayCounter++;
 
 		float MaterialPdf = 0.0;
@@ -267,6 +269,7 @@ RAY_TRACING_ENTRY_RAYGEN(PathTracingMainRG)
 			// Trace again
 			const uint RayFlags = RAY_FLAG_CULL_BACK_FACING_TRIANGLES;
 			const uint InstanceInclusionMask = RAY_TRACING_MASK_ALL;
+			const bool bDisableSkyLightContribution = false;
 
 			PrevMaterialPayload = Payload;
 			Payload = TraceMaterialRay(
@@ -274,7 +277,8 @@ RAY_TRACING_ENTRY_RAYGEN(PathTracingMainRG)
 				RayFlags,
 				InstanceInclusionMask,
 				Ray,
-				RayCone);
+				RayCone,
+				bDisableSkyLightContribution);
 			RayCounter++;
 		}
 
diff --git a/Engine/Shaders/Private/RayTracing/RayTracingCommon.ush b/Engine/Shaders/Private/RayTracing/RayTracingCommon.ush
--- a/Engine/Shaders/Private/RayTracing/RayTracingCommon.ush
+++ b/Engine/Shaders/Private/RayTracing/RayTracingCommon.ush
@@ -154,8 +154,10 @@ uint2 GetPixelCoord(uint2 DispatchThreadId, uint UpscaleFactor)
 #define RAY_TRACING_BLEND_MODE_MODULATE				4
 #define RAY_TRACING_BLEND_MODE_ALPHA_HOLDOUT		5
 
-#define RAY_TRACING_PAYLOAD_FLAG_FRONT_FACE         1 // Indicates that ray has hit front face of a primitive. Set by closest hit shader.
-#define RAY_TRACING_PAYLOAD_FLAG_MINIMAL_PAYLOAD    2 // Indicates that closest hit shader should only fill FMinimalPayload (HitT), skipping full material evaluation. Set by RayGen shader, before TraceRay().
+// Check how Flags are packed in the FPackedMaterialClosestHitPayload before adding to ensure bits are free
+#define RAY_TRACING_PAYLOAD_FLAG_FRONT_FACE                     (1 << 0) // Indicates that ray has hit front face of a primitive. Set by closest hit shader.
+#define RAY_TRACING_PAYLOAD_FLAG_MINIMAL_PAYLOAD                (1 << 1) // Indicates that closest hit shader should only fill FMinimalPayload (HitT), skipping full material evaluation. Set by RayGen shader, before TraceRay().
+#define RAY_TRACING_PAYLOAD_FLAG_DISABLE_SKY_LIGHT_CONTRIBUTION (1 << 2) // Indicates that hit shaders should ignore the Sky Light contribution in their output. Set by RayGen shader, before TraceRay().
 
 struct FMaterialClosestHitPayload : FMinimalPayload
 {
@@ -193,6 +195,9 @@ struct FMaterialClosestHitPayload : FMinimalPayload
 	void SetFrontFace() { Flags |= RAY_TRACING_PAYLOAD_FLAG_FRONT_FACE; }
 	bool IsFrontFace() { return (Flags & RAY_TRACING_PAYLOAD_FLAG_FRONT_FACE) != 0; }
 
+	void SetDisableSkyLightContribution() { Flags |= RAY_TRACING_PAYLOAD_FLAG_DISABLE_SKY_LIGHT_CONTRIBUTION; }
+	bool IsDisableSkyLightContribution() { return (Flags & RAY_TRACING_PAYLOAD_FLAG_DISABLE_SKY_LIGHT_CONTRIBUTION) != 0; }
+
 	void SetRayCone(FRayCone NewRayCone) { RayCone = NewRayCone; }
 };
 
@@ -228,6 +233,9 @@ struct FPackedMaterialClosestHitPayload : FMinimalPayload
 	void SetMinimalPayloadMode() { IorAndShadingModelIDAndBlendingModeAndFlags |= RAY_TRACING_PAYLOAD_FLAG_MINIMAL_PAYLOAD << 24; }
 	bool IsMinimalPayloadMode() { return ((IorAndShadingModelIDAndBlendingModeAndFlags >> 24) & RAY_TRACING_PAYLOAD_FLAG_MINIMAL_PAYLOAD) != 0; }
 
+	void SetDisableSkyLightContribution() { IorAndShadingModelIDAndBlendingModeAndFlags |= RAY_TRACING_PAYLOAD_FLAG_DISABLE_SKY_LIGHT_CONTRIBUTION << 24; }
+	bool IsDisableSkyLightContribution() { return ((IorAndShadingModelIDAndBlendingModeAndFlags >> 24) & RAY_TRACING_PAYLOAD_FLAG_DISABLE_SKY_LIGHT_CONTRIBUTION) != 0; }
+
 	void SetRayCone(FRayCone InRayCone) { RayCone = InRayCone; }
 
 	void SetRadiance(in float3 InRadiance)
@@ -449,7 +457,8 @@ void TraceMaterialRayPacked(
 	in uint InstanceInclusionMask,
 	in RayDesc Ray,
 	// Payload Inputs
-	inout FRayCone RayCone)
+	inout FRayCone RayCone,
+	in bool bDisableSkyLightContribution)
 {
 	const uint RayContributionToHitGroupIndex = RAY_TRACING_SHADER_SLOT_MATERIAL;
 	const uint MultiplierForGeometryContributionToShaderIndex = RAY_TRACING_NUM_SHADER_SLOTS;
@@ -459,6 +468,12 @@ void TraceMaterialRayPacked(
 
 	Payload.SetRayCone(RayCone);
 
+	if (bDisableSkyLightContribution)
+	{
+		// Disable Sky Light contribution in the payload if requested
+		Payload.SetDisableSkyLightContribution();
+	}
+
 	// Trace the ray
 
 	Payload.HitT = 0;
@@ -481,7 +496,8 @@ FMaterialClosestHitPayload TraceMaterialRay(
 	in uint InstanceInclusionMask,
 	in RayDesc Ray,
 	// Payload Inputs
-	inout FRayCone RayCone)
+	inout FRayCone RayCone,
+	in bool bDisableSkyLightContribution)
 {
 	const uint RayContributionToHitGroupIndex = RAY_TRACING_SHADER_SLOT_MATERIAL;
 	const uint MultiplierForGeometryContributionToShaderIndex = RAY_TRACING_NUM_SHADER_SLOTS;
@@ -493,6 +509,12 @@ FMaterialClosestHitPayload TraceMaterialRay(
 
 	PackedPayload.SetRayCone(RayCone);
 
+	if (bDisableSkyLightContribution)
+	{
+		// Disable Sky Light contribution in the payload if requested
+		PackedPayload.SetDisableSkyLightContribution();
+	}
+
 	// Trace the ray
 
 	TraceRay(
diff --git a/Engine/Shaders/Private/RayTracing/RayTracingCreateGatherPointsRGS.usf b/Engine/Shaders/Private/RayTracing/RayTracingCreateGatherPointsRGS.usf
--- a/Engine/Shaders/Private/RayTracing/RayTracingCreateGatherPointsRGS.usf
+++ b/Engine/Shaders/Private/RayTracing/RayTracingCreateGatherPointsRGS.usf
@@ -147,6 +147,7 @@ RAY_TRACING_ENTRY_RAYGEN(RayTracingCreateGatherPointsRGS)
 			// Cast ray
 			uint RayFlags = 0;
 			const uint InstanceInclusionMask = RAY_TRACING_MASK_OPAQUE;
+			const bool bDisableSkyLightContribution = false;
 
 #if !ENABLE_TWO_SIDED_GEOMETRY
 			RayFlags |= RAY_FLAG_CULL_BACK_FACING_TRIANGLES;
@@ -158,7 +159,8 @@ RAY_TRACING_ENTRY_RAYGEN(RayTracingCreateGatherPointsRGS)
 				RayFlags,
 				InstanceInclusionMask,
 				Ray,
-				RayCone);
+				RayCone,
+				bDisableSkyLightContribution);
 
 			// TODO: impossible based on formulation
 			if (Payload.HitT > MaxRayDistanceForGI) continue;
diff --git a/Engine/Shaders/Private/RayTracing/RayTracingDebug.usf b/Engine/Shaders/Private/RayTracing/RayTracingDebug.usf
--- a/Engine/Shaders/Private/RayTracing/RayTracingDebug.usf
+++ b/Engine/Shaders/Private/RayTracing/RayTracingDebug.usf
@@ -25,13 +25,15 @@ RAY_TRACING_ENTRY_RAYGEN(RayTracingDebugMainRGS)
 
 	const uint RayFlags = RAY_FLAG_CULL_BACK_FACING_TRIANGLES;
 	const uint InstanceInclusionMask = RAY_TRACING_MASK_ALL;
+	const bool bDisableSkyLightContribution = false;
 
 	FMaterialClosestHitPayload Payload = TraceMaterialRay(
 		TLAS,
 		RayFlags,
 		InstanceInclusionMask,
 		Ray,
-		RayCone);
+		RayCone,
+		bDisableSkyLightContribution);
 
 	float4 Result = float4(0, 0, 0, 1);
 
diff --git a/Engine/Shaders/Private/RayTracing/RayTracingGlobalIlluminationRGS.usf b/Engine/Shaders/Private/RayTracing/RayTracingGlobalIlluminationRGS.usf
--- a/Engine/Shaders/Private/RayTracing/RayTracingGlobalIlluminationRGS.usf
+++ b/Engine/Shaders/Private/RayTracing/RayTracingGlobalIlluminationRGS.usf
@@ -151,6 +151,7 @@ RAY_TRACING_ENTRY_RAYGEN(GlobalIlluminationRGS)
 			// Cast ray
 			uint RayFlags = 0;
 			const uint InstanceInclusionMask = RAY_TRACING_MASK_OPAQUE;
+			const bool bDisableSkyLightContribution = false;
 
 #if !ENABLE_TWO_SIDED_GEOMETRY
 			RayFlags |= RAY_FLAG_CULL_BACK_FACING_TRIANGLES;
@@ -162,7 +163,8 @@ RAY_TRACING_ENTRY_RAYGEN(GlobalIlluminationRGS)
 				RayFlags,
 				InstanceInclusionMask,
 				Ray,
-				RayCone);
+				RayCone,
+				bDisableSkyLightContribution);
 
 			// Environment hit
 			if (Payload.IsMiss())
diff --git a/Engine/Shaders/Private/RayTracing/RayTracingLightingCommon.ush b/Engine/Shaders/Private/RayTracing/RayTracingLightingCommon.ush
--- a/Engine/Shaders/Private/RayTracing/RayTracingLightingCommon.ush
+++ b/Engine/Shaders/Private/RayTracing/RayTracingLightingCommon.ush
@@ -708,6 +708,7 @@ FMaterialClosestHitPayload TraceRayAndAccumulateResults(
 	in uint ShouldDoDirectLighting,
 	in uint ShouldDoEmissiveAndIndirectLighting,
 	inout FRayCone RayCone,
+	in bool bDisableSkyLightContribution,
 	inout float3 Radiance)
 {
 	FPackedMaterialClosestHitPayload Payload = (FPackedMaterialClosestHitPayload)0;
@@ -717,7 +718,8 @@ FMaterialClosestHitPayload TraceRayAndAccumulateResults(
 		RayFlags,
 		InstanceInclusionMask,
 		Ray,
-		RayCone);
+		RayCone,
+		bDisableSkyLightContribution);
 
 	float3 WorldPosition = Ray.Origin + Ray.Direction * Payload.HitT;
 	float3 ViewDirection = Ray.Direction;
@@ -749,6 +751,7 @@ FMaterialClosestHitPayload TraceRayAndAccumulateBottomLayerResults(
 	in uint ShouldDoDirectLighting,
 	in uint ShouldDoEmissiveAndIndirectLighting,
 	inout FRayCone RayCone,
+	in bool bDisableSkyLightContribution,
 	inout float3 Radiance)
 {
 	FMaterialClosestHitPayload BottomLayerPayload = TraceMaterialRay(
@@ -756,7 +759,8 @@ FMaterialClosestHitPayload TraceRayAndAccumulateBottomLayerResults(
 		RayFlags,
 		InstanceInclusionMask,
 		Ray,
-		RayCone);
+		RayCone,
+		bDisableSkyLightContribution);
 
 	ComputeBottomLayerMaterialProperties(Ray, BottomLayerPayload);
 
diff --git a/Engine/Shaders/Private/RayTracing/RayTracingMaterialHitShaders.usf b/Engine/Shaders/Private/RayTracing/RayTracingMaterialHitShaders.usf
--- a/Engine/Shaders/Private/RayTracing/RayTracingMaterialHitShaders.usf
+++ b/Engine/Shaders/Private/RayTracing/RayTracingMaterialHitShaders.usf
@@ -31,13 +31,29 @@
 #include "/Engine/Private/LightmapCommon.ush"     
 
 /** Computes sky diffuse lighting, including precomputed shadowing. */
-void GetSkyLighting(FMaterialPixelParameters MaterialParameters, VTPageTableResult LightmapVTPageTableResult, float3 WorldNormal, float2 LightmapUV, uint LightmapDataIndex, float3 SkyOcclusionUV3D, uint ShadingModel, out float3 OutDiffuseLighting, out float3 OutSubsurfaceLighting)
+void GetSkyLighting(
+	FMaterialPixelParameters MaterialParameters,
+	VTPageTableResult LightmapVTPageTableResult,
+	float3 WorldNormal,
+	float2 LightmapUV,
+	uint LightmapDataIndex,
+	float3 SkyOcclusionUV3D,
+	uint ShadingModel,
+	bool bDisableSkyLightContribution,
+	out float3 OutDiffuseLighting,
+	out float3 OutSubsurfaceLighting)
 {
 	OutDiffuseLighting = 0;
 	OutSubsurfaceLighting = 0;
 
 #if ENABLE_SKY_LIGHT
 
+	// Check if the Sky Light contribution should be ignored
+	if (bDisableSkyLightContribution)
+	{
+		return;
+	}
+
 	float SkyVisibility = 1;
 	float GeometryTerm = 1;
 	float3 SkyLightingNormal = WorldNormal;
@@ -114,6 +130,7 @@ void GetPrecomputedIndirectLightingAndSkyLight(
 	float3 DiffuseDir,
 	float3 VolumetricLightmapBrickTextureUVs,
 	uint ShadingModel,
+	bool bDisableSkyLightContribution,
 	out float3 OutDiffuseLighting,
 	out float3 OutSubsurfaceLighting,
 	out float OutIndirectIrradiance)
@@ -167,7 +184,17 @@ void GetPrecomputedIndirectLightingAndSkyLight(
 
 	float3 SkyDiffuseLighting;
 	float3 SkySubsurfaceLighting;
-	GetSkyLighting(MaterialParameters, LightmapVTPageTableResult, DiffuseDir, SkyOcclusionUV, SkyOcclusionDataIndex, VolumetricLightmapBrickTextureUVs, ShadingModel, SkyDiffuseLighting, SkySubsurfaceLighting);
+	GetSkyLighting(
+		MaterialParameters,
+		LightmapVTPageTableResult,
+		DiffuseDir,
+		SkyOcclusionUV,
+		SkyOcclusionDataIndex,
+		VolumetricLightmapBrickTextureUVs,
+		ShadingModel,
+		bDisableSkyLightContribution,
+		SkyDiffuseLighting,
+		SkySubsurfaceLighting);
 
 	OutSubsurfaceLighting += SkySubsurfaceLighting;
 
@@ -354,7 +381,17 @@ RAY_TRACING_ENTRY_CLOSEST_HIT(MaterialCHS,
 #endif
 
 	VTPageTableResult LightmapVTPageTableResult = (VTPageTableResult)0.0f;
-	GetPrecomputedIndirectLightingAndSkyLight(MaterialParameters, Interpolants, LightmapVTPageTableResult, Payload.WorldNormal, VolumetricLightmapBrickTextureUVs, Payload.ShadingModelID, DiffuseIndirectLighting, SubsurfaceIndirectLighting, IndirectIrradiance);
+	GetPrecomputedIndirectLightingAndSkyLight(
+		MaterialParameters,
+		Interpolants,
+		LightmapVTPageTableResult,
+		Payload.WorldNormal,
+		VolumetricLightmapBrickTextureUVs,
+		Payload.ShadingModelID,
+		PackedPayload.IsDisableSkyLightContribution(),
+		DiffuseIndirectLighting,
+		SubsurfaceIndirectLighting,
+		IndirectIrradiance);
 	
 	Payload.DiffuseColor = BaseColor - BaseColor * Metallic;
 	Payload.SpecularColor = ComputeF0(Specular, BaseColor, Metallic);
diff --git a/Engine/Shaders/Private/RayTracing/RayTracingPrimaryRays.usf b/Engine/Shaders/Private/RayTracing/RayTracingPrimaryRays.usf
--- a/Engine/Shaders/Private/RayTracing/RayTracingPrimaryRays.usf
+++ b/Engine/Shaders/Private/RayTracing/RayTracingPrimaryRays.usf
@@ -155,6 +155,7 @@ RAY_TRACING_ENTRY_RAYGEN(RayTracingPrimaryRaysRGS)
 	{
 		const uint RefractionRayFlags = 0;
 		const uint RefractionInstanceInclusionMask = RAY_TRACING_MASK_ALL;
+		const bool bRefractionDisableSkyLightContribution = false;
 		float3 PathVertexRadiance = float3(0, 0, 0);
 
 		FMaterialClosestHitPayload Payload = TraceRayAndAccumulateResults(
@@ -168,6 +169,7 @@ RAY_TRACING_ENTRY_RAYGEN(RayTracingPrimaryRaysRGS)
 			ShouldDoDirectLighting,
 			ShouldDoEmissiveAndIndirectLighting,
 			RayCone,
+			bRefractionDisableSkyLightContribution,
 			PathVertexRadiance);
 		LastRoughness = Payload.Roughness;
 
@@ -219,6 +221,7 @@ RAY_TRACING_ENTRY_RAYGEN(RayTracingPrimaryRaysRGS)
 
 			const uint ReflectionRayFlags = RAY_FLAG_CULL_BACK_FACING_TRIANGLES;
 			const uint ReflectionInstanceInclusionMask = RAY_TRACING_MASK_ALL;
+			const bool bReflectionDisableSkyLightContribution = false;
 			float3 ReflectionRadiance = float3(0, 0, 0);
 
 			FMaterialClosestHitPayload ReflectionPayload = TraceRayAndAccumulateResults(
@@ -232,6 +235,7 @@ RAY_TRACING_ENTRY_RAYGEN(RayTracingPrimaryRaysRGS)
 				ShouldDoDirectLighting,
 				ShouldDoEmissiveAndIndirectLighting,
 				RayCone,
+				bReflectionDisableSkyLightContribution,
 				ReflectionRadiance);
 
 			// If we have not hit anything, sample the distance sky radiance.
diff --git a/Engine/Shaders/Private/RayTracing/RayTracingReflections.usf b/Engine/Shaders/Private/RayTracing/RayTracingReflections.usf
--- a/Engine/Shaders/Private/RayTracing/RayTracingReflections.usf
+++ b/Engine/Shaders/Private/RayTracing/RayTracingReflections.usf
@@ -152,12 +152,15 @@ float Transmit(in RayDesc TransmissionRay, inout RandomSequence RandSequence, in
 	const uint RayFlags = 0;
 	const uint InstanceInclusionMask = RAY_TRACING_MASK_TRANSLUCENT;
 	FRayCone RayCone = (FRayCone)0;
+	bool bDisableSkyLightContribution = false;
+
 	FMaterialClosestHitPayload TransmPayload = TraceMaterialRay(
 		TLAS,
 		RayFlags,
 		InstanceInclusionMask,
 		TransmissionRay,
-		RayCone);
+		RayCone,
+		bDisableSkyLightContribution);
 
 	const uint HitCountMax = 32;
 	uint HitCount = 0;
@@ -184,7 +187,8 @@ float Transmit(in RayDesc TransmissionRay, inout RandomSequence RandSequence, in
 			RayFlags,
 			RAY_TRACING_MASK_TRANSLUCENT,
 			TransmissionRay,
-			RayCone);
+			RayCone,
+			bDisableSkyLightContribution);
 
 		bIsHit = TransmPayload.IsHit();
 	}
@@ -446,6 +450,7 @@ RAY_TRACING_ENTRY_RAYGEN(RayTracingReflectionsRGS)
 				{
 					const uint RayFlags = 0;
 					const uint InstanceInclusionMask = RAY_TRACING_MASK_OPAQUE;
+					const bool bDisableSkyLightContribution = false;
 
 					TraceMaterialRayPacked(
 						Payload,
@@ -453,7 +458,8 @@ RAY_TRACING_ENTRY_RAYGEN(RayTracingReflectionsRGS)
 						RayFlags,
 						InstanceInclusionMask,
 						TopLayerRay,
-						RayCone);
+						RayCone,
+						bDisableSkyLightContribution);
 
 					// Trace separate ray for translucent objects
 					float Transmission = 1.0;
@@ -562,6 +568,7 @@ RAY_TRACING_ENTRY_RAYGEN(RayTracingReflectionsRGS)
 					{
 						const uint RayFlags = 0;
 						const uint InstanceInclusionMask = RAY_TRACING_MASK_OPAQUE;
+						const bool bDisableSkyLightContribution = false;
 
 						FMaterialClosestHitPayload BottomLayerPayload = TraceRayAndAccumulateBottomLayerResults(
 							BottomLayerRay,
@@ -574,6 +581,7 @@ RAY_TRACING_ENTRY_RAYGEN(RayTracingReflectionsRGS)
 							ShouldDoDirectLighting,
 							ShouldDoEmissiveAndIndirectLighting,
 							RayCone,
+							bDisableSkyLightContribution,
 							BottomLayerRadiance);
 
 						// Distance and sky: similar processing to main ray.
