Sky light affect reflection flag manual integration to complete what was left out of CL 10088081.

diff --git a/Engine/Shaders/Private/RayTracing/RayTracingPrimaryRays.usf b/Engine/Shaders/Private/RayTracing/RayTracingPrimaryRays.usf
--- a/Engine/Shaders/Private/RayTracing/RayTracingPrimaryRays.usf
+++ b/Engine/Shaders/Private/RayTracing/RayTracingPrimaryRays.usf
@@ -141,6 +141,8 @@ RAY_TRACING_ENTRY_RAYGEN(RayTracingPrimaryRaysRGS)
 	{
 		bAllowSkySampling = true;
 	}
+	// Check if the Sky Light should affect reflection rays within translucency.
+	const bool bSkyLightAffectReflection = ShouldSkyLightAffectReflection();
 	
 	bool bHasScattered = (ERayTracingPrimaryRaysFlag_ConsiderSurfaceScatter & PrimaryRayFlags) != 0;
 
@@ -221,7 +223,7 @@ RAY_TRACING_ENTRY_RAYGEN(RayTracingPrimaryRaysRGS)
 
 			const uint ReflectionRayFlags = RAY_FLAG_CULL_BACK_FACING_TRIANGLES;
 			const uint ReflectionInstanceInclusionMask = RAY_TRACING_MASK_ALL;
-			const bool bReflectionDisableSkyLightContribution = false;
+			const bool bReflectionDisableSkyLightContribution = !bSkyLightAffectReflection;
 			float3 ReflectionRadiance = float3(0, 0, 0);
 
 			FMaterialClosestHitPayload ReflectionPayload = TraceRayAndAccumulateResults(
diff --git a/Engine/Shaders/Private/RayTracing/RayTracingReflections.usf b/Engine/Shaders/Private/RayTracing/RayTracingReflections.usf
--- a/Engine/Shaders/Private/RayTracing/RayTracingReflections.usf
+++ b/Engine/Shaders/Private/RayTracing/RayTracingReflections.usf
@@ -257,6 +257,8 @@ RAY_TRACING_ENTRY_RAYGEN(RayTracingReflectionsRGS)
 
 	// Sky is only sampled when infinite reflection rays are used.
 	const bool bAllowSkySampling = ReflectionMaxRayDistance < 0;
+	// Check if the Sky Light should affect reflection.
+	const bool bSkyLightAffectReflection = ShouldSkyLightAffectReflection();
 
 	// Shorten the rays on rougher surfaces between user-provided min and max ray lengths.
 	// When a shortened ray misses the geometry, we fall back to local reflection capture sampling (similar to SSR).
@@ -450,7 +452,7 @@ RAY_TRACING_ENTRY_RAYGEN(RayTracingReflectionsRGS)
 				{
 					const uint RayFlags = 0;
 					const uint InstanceInclusionMask = RAY_TRACING_MASK_OPAQUE;
-					const bool bDisableSkyLightContribution = false;
+					const bool bDisableSkyLightContribution = !bSkyLightAffectReflection;
 
 					TraceMaterialRayPacked(
 						Payload,
@@ -497,7 +499,7 @@ RAY_TRACING_ENTRY_RAYGEN(RayTracingReflectionsRGS)
 					TopLayerRadiance += Radiance * Transmission;
 				}
 
-				bool isTopLayerRayValid = Payload.IsHit() || (bAllowSkySampling && ReflectionStruct.SkyLightParameters.y > 0);
+				bool isTopLayerRayValid = Payload.IsHit() || (bAllowSkySampling && bSkyLightAffectReflection && ReflectionStruct.SkyLightParameters.y > 0);
 
 
 				// Compute some flags and values
@@ -515,7 +517,7 @@ RAY_TRACING_ENTRY_RAYGEN(RayTracingReflectionsRGS)
 					float3 RayHitWorldPos = TopLayerRay.Origin + Payload.HitT * TopLayerRay.Direction;
 					OriginToCollider = RayHitWorldPos - TopLayerRay.Origin;
 				}
-				else if (bAllowSkySampling && ReflectionStruct.SkyLightParameters.y > 0)
+				else if (bAllowSkySampling && bSkyLightAffectReflection && ReflectionStruct.SkyLightParameters.y > 0)
 				{
 					float SkyAverageBrightness = 1.0f;
 					float3 SkyLighting = GetSkyLightReflection(TopLayerRay.Direction, TopLayerRoughness, SkyAverageBrightness);
@@ -568,7 +570,7 @@ RAY_TRACING_ENTRY_RAYGEN(RayTracingReflectionsRGS)
 					{
 						const uint RayFlags = 0;
 						const uint InstanceInclusionMask = RAY_TRACING_MASK_OPAQUE;
-						const bool bDisableSkyLightContribution = false;
+						const bool bDisableSkyLightContribution = !bSkyLightAffectReflection;
 
 						FMaterialClosestHitPayload BottomLayerPayload = TraceRayAndAccumulateBottomLayerResults(
 							BottomLayerRay,
@@ -592,7 +594,7 @@ RAY_TRACING_ENTRY_RAYGEN(RayTracingReflectionsRGS)
 							IsBottomSampleValid = true;
 							BottomLayerHitDistance = BottomLayerPayload.HitT;
 						}
-						else if (bAllowSkySampling && ReflectionStruct.SkyLightParameters.y > 0)
+						else if (bAllowSkySampling && bSkyLightAffectReflection && ReflectionStruct.SkyLightParameters.y > 0)
 						{
 							IsBottomSampleValid = true;
 							float SkyAverageBrightness = 1.0f;
