Modify reflection compositing in ray traced reflections to omit the sky light if requested.

(ushell-p4-cherrypick of 10087278 by riley.alston)

diff --git a/Engine/Shaders/Private/ForwardLightingCommon.ush b/Engine/Shaders/Private/ForwardLightingCommon.ush
--- a/Engine/Shaders/Private/ForwardLightingCommon.ush
+++ b/Engine/Shaders/Private/ForwardLightingCommon.ush
@@ -341,6 +341,7 @@ half3 GetImageBasedReflectionLighting(FMaterialPixelParameters MaterialParameter
 	#endif
 #endif
 
+	const bool bCompositeSkylight = true;
 	float3 SpecularIBL = CompositeReflectionCapturesAndSkylight(
 		1.0f, 
 		MaterialParameters.AbsoluteWorldPosition, 
@@ -352,6 +353,7 @@ half3 GetImageBasedReflectionLighting(FMaterialPixelParameters MaterialParameter
 		NumLocalReflectionCaptures, 
 		DataStartIndex, 
 		SingleCaptureIndex,
+		bCompositeSkylight,
 		EyeIndex);
 
 #if MATERIAL_SSR && !FORWARD_SHADING
diff --git a/Engine/Shaders/Private/HairStrands/HairStrandsEnvironmentLighting.usf b/Engine/Shaders/Private/HairStrands/HairStrandsEnvironmentLighting.usf
--- a/Engine/Shaders/Private/HairStrands/HairStrandsEnvironmentLighting.usf
+++ b/Engine/Shaders/Private/HairStrands/HairStrandsEnvironmentLighting.usf
@@ -52,6 +52,7 @@ float3 GatherRadiance(float3 RayDirection, float Roughness, uint2 PixelPosition,
 	const uint NumLocalReflectionCaptures = min(ForwardLightData.NumCulledLightsGrid[NumCulledEntryIndex + 0], ForwardLightData.NumReflectionCaptures);
 	const uint DataStartIndex = ForwardLightData.NumCulledLightsGrid[NumCulledEntryIndex + 1];
 
+	const bool bCompositeSkylight = true;
 	return CompositeReflectionCapturesAndSkylight(
 		1.0f,
 		AbsoluteWorldPosition,
@@ -63,6 +64,7 @@ float3 GatherRadiance(float3 RayDirection, float Roughness, uint2 PixelPosition,
 		NumLocalReflectionCaptures,
 		DataStartIndex,
 		SingleCaptureIndex,
+		bCompositeSkylight,
 		EyeIndex);
 }
 
diff --git a/Engine/Shaders/Private/RayTracing/RayTracingReflections.usf b/Engine/Shaders/Private/RayTracing/RayTracingReflections.usf
--- a/Engine/Shaders/Private/RayTracing/RayTracingReflections.usf
+++ b/Engine/Shaders/Private/RayTracing/RayTracingReflections.usf
@@ -640,7 +640,8 @@ RAY_TRACING_ENTRY_RAYGEN(RayTracingReflectionsRGS)
 								ExtraIndirectSpecular,
 								NumCulledReflectionCaptures,
 								ReflectionCapturesStartIndex,
-								0);
+								0,
+								bSkyLightAffectReflection);
 					}
 
 					// Apply correct weights to the bottom and top layers (See ReflectionEnvironmentPixelShader)
@@ -767,7 +768,8 @@ RAY_TRACING_ENTRY_RAYGEN(RayTracingReflectionsRGS)
 							ExtraIndirectSpecular,
 							NumCulledReflectionCaptures,
 							ReflectionCapturesStartIndex,
-							0);
+							0,
+							bSkyLightAffectReflection);
 				}
 
 				ReflectedColor += float4(PathRadiance, 1.0f);
diff --git a/Engine/Shaders/Private/ReflectionEnvironmentComposite.ush b/Engine/Shaders/Private/ReflectionEnvironmentComposite.ush
--- a/Engine/Shaders/Private/ReflectionEnvironmentComposite.ush
+++ b/Engine/Shaders/Private/ReflectionEnvironmentComposite.ush
@@ -15,6 +15,7 @@ float3 CompositeReflectionCapturesAndSkylight(
 	uint NumCapturesAffectingTile,
 	uint CaptureDataStartIndex, 
 	int SingleCaptureIndex,
+	bool bCompositeSkylight,
 	uint EyeIndex)
 {
 	float Mip = ComputeReflectionCaptureMipFromRoughness(Roughness, View.ReflectionCubemapMaxMip);
@@ -164,7 +165,7 @@ float3 CompositeReflectionCapturesAndSkylight(
 #if ENABLE_SKY_LIGHT
 
 	BRANCH  
-	if (ReflectionStruct.SkyLightParameters.y > 0)
+	if (ReflectionStruct.SkyLightParameters.y > 0 && bCompositeSkylight)
 	{
 		float SkyAverageBrightness = 1.0f;
 
@@ -209,7 +210,8 @@ float3 CompositeReflectionCapturesAndSkylight(
 	float3 ExtraIndirectSpecular,
 	uint NumCapturesAffectingTile,
 	uint CaptureDataStartIndex,
-	int SingleCaptureIndex)
+	int SingleCaptureIndex,
+	bool bCompositeSkylight)
 {
 	return CompositeReflectionCapturesAndSkylight(
 		CompositeAlpha, 
@@ -222,5 +224,6 @@ float3 CompositeReflectionCapturesAndSkylight(
 		NumCapturesAffectingTile, 
 		CaptureDataStartIndex,
 		SingleCaptureIndex,
+		bCompositeSkylight,
 		0);
 }
\ No newline at end of file
diff --git a/Engine/Shaders/Private/ReflectionEnvironmentPixelShader.usf b/Engine/Shaders/Private/ReflectionEnvironmentPixelShader.usf
--- a/Engine/Shaders/Private/ReflectionEnvironmentPixelShader.usf
+++ b/Engine/Shaders/Private/ReflectionEnvironmentPixelShader.usf
@@ -31,6 +31,7 @@ float3 GatherRadiance(float CompositeAlpha, float3 WorldPosition, float3 RayDire
 	IndirectIrradiance *= IndirectDiffuseOcclusion;
 #endif
 
+	const bool bCompositeSkylight = true;
 	return CompositeReflectionCapturesAndSkylight(
 		CompositeAlpha, 
 		WorldPosition, 
@@ -41,7 +42,8 @@ float3 GatherRadiance(float CompositeAlpha, float3 WorldPosition, float3 RayDire
 		ExtraIndirectSpecular, 
 		NumCulledReflectionCaptures, 
 		CaptureDataStartIndex, 
-		0);
+		0,
+		bCompositeSkylight);
 }
 
 
diff --git a/Engine/Shaders/Private/SingleLayerWaterComposite.usf b/Engine/Shaders/Private/SingleLayerWaterComposite.usf
--- a/Engine/Shaders/Private/SingleLayerWaterComposite.usf
+++ b/Engine/Shaders/Private/SingleLayerWaterComposite.usf
@@ -86,6 +86,7 @@ SingleLayerWaterCompositeOutput SingleLayerWaterComposite(float2 BufferUV, float
 			float NoV = saturate(dot(N, V));
 			// Point lobe in off-specular peak direction
 			R = GetOffSpecularPeakReflectionDir(N, R, GBuffer.Roughness);
+			const bool bCompositeSkylight = true;
 			Reflection += PreExposure * CompositeReflectionCapturesAndSkylight(
 				1,
 				WorldPosition,
@@ -96,7 +97,8 @@ SingleLayerWaterCompositeOutput SingleLayerWaterComposite(float2 BufferUV, float
 				ExtraIndirectSpecular,
 				NumCulledReflectionCaptures,
 				CaptureDataStartIndex,
-				0);
+				0,
+				bCompositeSkylight);
 
 			// Then combine reflection with SSR
 			float4 SSR = Texture2DSample(ScreenSpaceReflectionsTexture, ScreenSpaceReflectionsSampler, BufferUV);
