Change payload sky light contribution flag from a disable flag to an enable flag.

(ushell-p4-cherrypick of 10103129 by riley.alston)

diff --git a/Engine/Shaders/Private/RayTracing/RayTracingCommon.ush b/Engine/Shaders/Private/RayTracing/RayTracingCommon.ush
--- a/Engine/Shaders/Private/RayTracing/RayTracingCommon.ush
+++ b/Engine/Shaders/Private/RayTracing/RayTracingCommon.ush
@@ -155,9 +155,9 @@ uint2 GetPixelCoord(uint2 DispatchThreadId, uint UpscaleFactor)
 #define RAY_TRACING_BLEND_MODE_ALPHA_HOLDOUT		5
 
 // Check how Flags are packed in the FPackedMaterialClosestHitPayload before adding to ensure bits are free
-#define RAY_TRACING_PAYLOAD_FLAG_FRONT_FACE                     (1 << 0) // Indicates that ray has hit front face of a primitive. Set by closest hit shader.
-#define RAY_TRACING_PAYLOAD_FLAG_MINIMAL_PAYLOAD                (1 << 1) // Indicates that closest hit shader should only fill FMinimalPayload (HitT), skipping full material evaluation. Set by RayGen shader, before TraceRay().
-#define RAY_TRACING_PAYLOAD_FLAG_DISABLE_SKY_LIGHT_CONTRIBUTION (1 << 2) // Indicates that hit shaders should ignore the Sky Light contribution in their output. Set by RayGen shader, before TraceRay().
+#define RAY_TRACING_PAYLOAD_FLAG_FRONT_FACE                    (1 << 0) // Indicates that ray has hit front face of a primitive. Set by closest hit shader.
+#define RAY_TRACING_PAYLOAD_FLAG_MINIMAL_PAYLOAD               (1 << 1) // Indicates that closest hit shader should only fill FMinimalPayload (HitT), skipping full material evaluation. Set by RayGen shader, before TraceRay().
+#define RAY_TRACING_PAYLOAD_FLAG_ENABLE_SKY_LIGHT_CONTRIBUTION (1 << 2) // Indicates that hit shaders should add in the Sky Light contribution in their output. Set by RayGen shader, before TraceRay().
 
 struct FMaterialClosestHitPayload : FMinimalPayload
 {
@@ -195,8 +195,8 @@ struct FMaterialClosestHitPayload : FMinimalPayload
 	void SetFrontFace() { Flags |= RAY_TRACING_PAYLOAD_FLAG_FRONT_FACE; }
 	bool IsFrontFace() { return (Flags & RAY_TRACING_PAYLOAD_FLAG_FRONT_FACE) != 0; }
 
-	void SetDisableSkyLightContribution() { Flags |= RAY_TRACING_PAYLOAD_FLAG_DISABLE_SKY_LIGHT_CONTRIBUTION; }
-	bool IsDisableSkyLightContribution() { return (Flags & RAY_TRACING_PAYLOAD_FLAG_DISABLE_SKY_LIGHT_CONTRIBUTION) != 0; }
+	void SetEnableSkyLightContribution() { Flags |= RAY_TRACING_PAYLOAD_FLAG_ENABLE_SKY_LIGHT_CONTRIBUTION; }
+	bool IsEnableSkyLightContribution() { return (Flags & RAY_TRACING_PAYLOAD_FLAG_ENABLE_SKY_LIGHT_CONTRIBUTION) != 0; }
 
 	void SetRayCone(FRayCone NewRayCone) { RayCone = NewRayCone; }
 };
@@ -233,8 +233,8 @@ struct FPackedMaterialClosestHitPayload : FMinimalPayload
 	void SetMinimalPayloadMode() { IorAndShadingModelIDAndBlendingModeAndFlags |= RAY_TRACING_PAYLOAD_FLAG_MINIMAL_PAYLOAD << 24; }
 	bool IsMinimalPayloadMode() { return ((IorAndShadingModelIDAndBlendingModeAndFlags >> 24) & RAY_TRACING_PAYLOAD_FLAG_MINIMAL_PAYLOAD) != 0; }
 
-	void SetDisableSkyLightContribution() { IorAndShadingModelIDAndBlendingModeAndFlags |= RAY_TRACING_PAYLOAD_FLAG_DISABLE_SKY_LIGHT_CONTRIBUTION << 24; }
-	bool IsDisableSkyLightContribution() { return ((IorAndShadingModelIDAndBlendingModeAndFlags >> 24) & RAY_TRACING_PAYLOAD_FLAG_DISABLE_SKY_LIGHT_CONTRIBUTION) != 0; }
+	void SetEnableSkyLightContribution() { IorAndShadingModelIDAndBlendingModeAndFlags |= RAY_TRACING_PAYLOAD_FLAG_ENABLE_SKY_LIGHT_CONTRIBUTION << 24; }
+	bool IsEnableSkyLightContribution() { return ((IorAndShadingModelIDAndBlendingModeAndFlags >> 24) & RAY_TRACING_PAYLOAD_FLAG_ENABLE_SKY_LIGHT_CONTRIBUTION) != 0; }
 
 	void SetRayCone(FRayCone InRayCone) { RayCone = InRayCone; }
 
@@ -458,7 +458,7 @@ void TraceMaterialRayPacked(
 	in RayDesc Ray,
 	// Payload Inputs
 	inout FRayCone RayCone,
-	in bool bDisableSkyLightContribution)
+	in bool bEnableSkyLightContribution)
 {
 	const uint RayContributionToHitGroupIndex = RAY_TRACING_SHADER_SLOT_MATERIAL;
 	const uint MultiplierForGeometryContributionToShaderIndex = RAY_TRACING_NUM_SHADER_SLOTS;
@@ -468,10 +468,10 @@ void TraceMaterialRayPacked(
 
 	Payload.SetRayCone(RayCone);
 
-	if (bDisableSkyLightContribution)
+	if (bEnableSkyLightContribution)
 	{
-		// Disable Sky Light contribution in the payload if requested
-		Payload.SetDisableSkyLightContribution();
+		// Enable Sky Light contribution in the payload if requested
+		Payload.SetEnableSkyLightContribution();
 	}
 
 	// Trace the ray
@@ -497,7 +497,7 @@ FMaterialClosestHitPayload TraceMaterialRay(
 	in RayDesc Ray,
 	// Payload Inputs
 	inout FRayCone RayCone,
-	in bool bDisableSkyLightContribution)
+	in bool bEnableSkyLightContribution)
 {
 	const uint RayContributionToHitGroupIndex = RAY_TRACING_SHADER_SLOT_MATERIAL;
 	const uint MultiplierForGeometryContributionToShaderIndex = RAY_TRACING_NUM_SHADER_SLOTS;
@@ -509,10 +509,10 @@ FMaterialClosestHitPayload TraceMaterialRay(
 
 	PackedPayload.SetRayCone(RayCone);
 
-	if (bDisableSkyLightContribution)
+	if (bEnableSkyLightContribution)
 	{
-		// Disable Sky Light contribution in the payload if requested
-		PackedPayload.SetDisableSkyLightContribution();
+		// Enable Sky Light contribution in the payload if requested
+		PackedPayload.SetEnableSkyLightContribution();
 	}
 
 	// Trace the ray
diff --git a/Engine/Shaders/Private/RayTracing/RayTracingCreateGatherPointsRGS.usf b/Engine/Shaders/Private/RayTracing/RayTracingCreateGatherPointsRGS.usf
--- a/Engine/Shaders/Private/RayTracing/RayTracingCreateGatherPointsRGS.usf
+++ b/Engine/Shaders/Private/RayTracing/RayTracingCreateGatherPointsRGS.usf
@@ -147,7 +147,7 @@ RAY_TRACING_ENTRY_RAYGEN(RayTracingCreateGatherPointsRGS)
 			// Cast ray
 			uint RayFlags = 0;
 			const uint InstanceInclusionMask = RAY_TRACING_MASK_OPAQUE;
-			const bool bDisableSkyLightContribution = false;
+			const bool bEnableSkyLightContribution = true;
 
 #if !ENABLE_TWO_SIDED_GEOMETRY
 			RayFlags |= RAY_FLAG_CULL_BACK_FACING_TRIANGLES;
@@ -160,7 +160,7 @@ RAY_TRACING_ENTRY_RAYGEN(RayTracingCreateGatherPointsRGS)
 				InstanceInclusionMask,
 				Ray,
 				RayCone,
-				bDisableSkyLightContribution);
+				bEnableSkyLightContribution);
 
 			// TODO: impossible based on formulation
 			if (Payload.HitT > MaxRayDistanceForGI) continue;
diff --git a/Engine/Shaders/Private/RayTracing/RayTracingDebug.usf b/Engine/Shaders/Private/RayTracing/RayTracingDebug.usf
--- a/Engine/Shaders/Private/RayTracing/RayTracingDebug.usf
+++ b/Engine/Shaders/Private/RayTracing/RayTracingDebug.usf
@@ -25,7 +25,7 @@ RAY_TRACING_ENTRY_RAYGEN(RayTracingDebugMainRGS)
 
 	const uint RayFlags = RAY_FLAG_CULL_BACK_FACING_TRIANGLES;
 	const uint InstanceInclusionMask = RAY_TRACING_MASK_ALL;
-	const bool bDisableSkyLightContribution = false;
+	const bool bEnableSkyLightContribution = true;
 
 	FMaterialClosestHitPayload Payload = TraceMaterialRay(
 		TLAS,
@@ -33,7 +33,7 @@ RAY_TRACING_ENTRY_RAYGEN(RayTracingDebugMainRGS)
 		InstanceInclusionMask,
 		Ray,
 		RayCone,
-		bDisableSkyLightContribution);
+		bEnableSkyLightContribution);
 
 	float4 Result = float4(0, 0, 0, 1);
 
diff --git a/Engine/Shaders/Private/RayTracing/RayTracingGlobalIlluminationRGS.usf b/Engine/Shaders/Private/RayTracing/RayTracingGlobalIlluminationRGS.usf
--- a/Engine/Shaders/Private/RayTracing/RayTracingGlobalIlluminationRGS.usf
+++ b/Engine/Shaders/Private/RayTracing/RayTracingGlobalIlluminationRGS.usf
@@ -151,7 +151,7 @@ RAY_TRACING_ENTRY_RAYGEN(GlobalIlluminationRGS)
 			// Cast ray
 			uint RayFlags = 0;
 			const uint InstanceInclusionMask = RAY_TRACING_MASK_OPAQUE;
-			const bool bDisableSkyLightContribution = false;
+			const bool bEnableSkyLightContribution = true;
 
 #if !ENABLE_TWO_SIDED_GEOMETRY
 			RayFlags |= RAY_FLAG_CULL_BACK_FACING_TRIANGLES;
@@ -164,7 +164,7 @@ RAY_TRACING_ENTRY_RAYGEN(GlobalIlluminationRGS)
 				InstanceInclusionMask,
 				Ray,
 				RayCone,
-				bDisableSkyLightContribution);
+				bEnableSkyLightContribution);
 
 			// Environment hit
 			if (Payload.IsMiss())
diff --git a/Engine/Shaders/Private/RayTracing/RayTracingLightingCommon.ush b/Engine/Shaders/Private/RayTracing/RayTracingLightingCommon.ush
--- a/Engine/Shaders/Private/RayTracing/RayTracingLightingCommon.ush
+++ b/Engine/Shaders/Private/RayTracing/RayTracingLightingCommon.ush
@@ -708,7 +708,7 @@ FMaterialClosestHitPayload TraceRayAndAccumulateResults(
 	in uint ShouldDoDirectLighting,
 	in uint ShouldDoEmissiveAndIndirectLighting,
 	inout FRayCone RayCone,
-	in bool bDisableSkyLightContribution,
+	in bool bEnableSkyLightContribution,
 	inout float3 Radiance)
 {
 	FPackedMaterialClosestHitPayload Payload = (FPackedMaterialClosestHitPayload)0;
@@ -719,7 +719,7 @@ FMaterialClosestHitPayload TraceRayAndAccumulateResults(
 		InstanceInclusionMask,
 		Ray,
 		RayCone,
-		bDisableSkyLightContribution);
+		bEnableSkyLightContribution);
 
 	float3 WorldPosition = Ray.Origin + Ray.Direction * Payload.HitT;
 	float3 ViewDirection = Ray.Direction;
@@ -751,7 +751,7 @@ FMaterialClosestHitPayload TraceRayAndAccumulateBottomLayerResults(
 	in uint ShouldDoDirectLighting,
 	in uint ShouldDoEmissiveAndIndirectLighting,
 	inout FRayCone RayCone,
-	in bool bDisableSkyLightContribution,
+	in bool bEnableSkyLightContribution,
 	inout float3 Radiance)
 {
 	FMaterialClosestHitPayload BottomLayerPayload = TraceMaterialRay(
@@ -760,7 +760,7 @@ FMaterialClosestHitPayload TraceRayAndAccumulateBottomLayerResults(
 		InstanceInclusionMask,
 		Ray,
 		RayCone,
-		bDisableSkyLightContribution);
+		bEnableSkyLightContribution);
 
 	ComputeBottomLayerMaterialProperties(Ray, BottomLayerPayload);
 
diff --git a/Engine/Shaders/Private/RayTracing/RayTracingMaterialHitShaders.usf b/Engine/Shaders/Private/RayTracing/RayTracingMaterialHitShaders.usf
--- a/Engine/Shaders/Private/RayTracing/RayTracingMaterialHitShaders.usf
+++ b/Engine/Shaders/Private/RayTracing/RayTracingMaterialHitShaders.usf
@@ -39,7 +39,7 @@ void GetSkyLighting(
 	uint LightmapDataIndex,
 	float3 SkyOcclusionUV3D,
 	uint ShadingModel,
-	bool bDisableSkyLightContribution,
+	bool bEnableSkyLightContribution,
 	out float3 OutDiffuseLighting,
 	out float3 OutSubsurfaceLighting)
 {
@@ -49,7 +49,7 @@ void GetSkyLighting(
 #if ENABLE_SKY_LIGHT
 
 	// Check if the Sky Light contribution should be ignored
-	if (bDisableSkyLightContribution)
+	if (!bEnableSkyLightContribution)
 	{
 		return;
 	}
@@ -130,7 +130,7 @@ void GetPrecomputedIndirectLightingAndSkyLight(
 	float3 DiffuseDir,
 	float3 VolumetricLightmapBrickTextureUVs,
 	uint ShadingModel,
-	bool bDisableSkyLightContribution,
+	bool bEnableSkyLightContribution,
 	out float3 OutDiffuseLighting,
 	out float3 OutSubsurfaceLighting,
 	out float OutIndirectIrradiance)
@@ -192,7 +192,7 @@ void GetPrecomputedIndirectLightingAndSkyLight(
 		SkyOcclusionDataIndex,
 		VolumetricLightmapBrickTextureUVs,
 		ShadingModel,
-		bDisableSkyLightContribution,
+		bEnableSkyLightContribution,
 		SkyDiffuseLighting,
 		SkySubsurfaceLighting);
 
@@ -388,7 +388,7 @@ RAY_TRACING_ENTRY_CLOSEST_HIT(MaterialCHS,
 		Payload.WorldNormal,
 		VolumetricLightmapBrickTextureUVs,
 		Payload.ShadingModelID,
-		PackedPayload.IsDisableSkyLightContribution(),
+		PackedPayload.IsEnableSkyLightContribution(),
 		DiffuseIndirectLighting,
 		SubsurfaceIndirectLighting,
 		IndirectIrradiance);
diff --git a/Engine/Shaders/Private/RayTracing/RayTracingPrimaryRays.usf b/Engine/Shaders/Private/RayTracing/RayTracingPrimaryRays.usf
--- a/Engine/Shaders/Private/RayTracing/RayTracingPrimaryRays.usf
+++ b/Engine/Shaders/Private/RayTracing/RayTracingPrimaryRays.usf
@@ -157,7 +157,7 @@ RAY_TRACING_ENTRY_RAYGEN(RayTracingPrimaryRaysRGS)
 	{
 		const uint RefractionRayFlags = 0;
 		const uint RefractionInstanceInclusionMask = RAY_TRACING_MASK_ALL;
-		const bool bRefractionDisableSkyLightContribution = false;
+		const bool bRefractionEnableSkyLightContribution = true;
 		float3 PathVertexRadiance = float3(0, 0, 0);
 
 		FMaterialClosestHitPayload Payload = TraceRayAndAccumulateResults(
@@ -171,7 +171,7 @@ RAY_TRACING_ENTRY_RAYGEN(RayTracingPrimaryRaysRGS)
 			ShouldDoDirectLighting,
 			ShouldDoEmissiveAndIndirectLighting,
 			RayCone,
-			bRefractionDisableSkyLightContribution,
+			bRefractionEnableSkyLightContribution,
 			PathVertexRadiance);
 		LastRoughness = Payload.Roughness;
 
@@ -223,7 +223,7 @@ RAY_TRACING_ENTRY_RAYGEN(RayTracingPrimaryRaysRGS)
 
 			const uint ReflectionRayFlags = RAY_FLAG_CULL_BACK_FACING_TRIANGLES;
 			const uint ReflectionInstanceInclusionMask = RAY_TRACING_MASK_ALL;
-			const bool bReflectionDisableSkyLightContribution = !bSkyLightAffectReflection;
+			const bool bReflectionEnableSkyLightContribution = bSkyLightAffectReflection;
 			float3 ReflectionRadiance = float3(0, 0, 0);
 
 			FMaterialClosestHitPayload ReflectionPayload = TraceRayAndAccumulateResults(
@@ -237,7 +237,7 @@ RAY_TRACING_ENTRY_RAYGEN(RayTracingPrimaryRaysRGS)
 				ShouldDoDirectLighting,
 				ShouldDoEmissiveAndIndirectLighting,
 				RayCone,
-				bReflectionDisableSkyLightContribution,
+				bReflectionEnableSkyLightContribution,
 				ReflectionRadiance);
 
 			// If we have not hit anything, sample the distance sky radiance.
diff --git a/Engine/Shaders/Private/RayTracing/RayTracingReflections.usf b/Engine/Shaders/Private/RayTracing/RayTracingReflections.usf
--- a/Engine/Shaders/Private/RayTracing/RayTracingReflections.usf
+++ b/Engine/Shaders/Private/RayTracing/RayTracingReflections.usf
@@ -152,7 +152,7 @@ float Transmit(in RayDesc TransmissionRay, inout RandomSequence RandSequence, in
 	const uint RayFlags = 0;
 	const uint InstanceInclusionMask = RAY_TRACING_MASK_TRANSLUCENT;
 	FRayCone RayCone = (FRayCone)0;
-	bool bDisableSkyLightContribution = false;
+	bool bEnableSkyLightContribution = true;
 
 	FMaterialClosestHitPayload TransmPayload = TraceMaterialRay(
 		TLAS,
@@ -160,7 +160,7 @@ float Transmit(in RayDesc TransmissionRay, inout RandomSequence RandSequence, in
 		InstanceInclusionMask,
 		TransmissionRay,
 		RayCone,
-		bDisableSkyLightContribution);
+		bEnableSkyLightContribution);
 
 	const uint HitCountMax = 32;
 	uint HitCount = 0;
@@ -188,7 +188,7 @@ float Transmit(in RayDesc TransmissionRay, inout RandomSequence RandSequence, in
 			RAY_TRACING_MASK_TRANSLUCENT,
 			TransmissionRay,
 			RayCone,
-			bDisableSkyLightContribution);
+			bEnableSkyLightContribution);
 
 		bIsHit = TransmPayload.IsHit();
 	}
@@ -452,7 +452,7 @@ RAY_TRACING_ENTRY_RAYGEN(RayTracingReflectionsRGS)
 				{
 					const uint RayFlags = 0;
 					const uint InstanceInclusionMask = RAY_TRACING_MASK_OPAQUE;
-					const bool bDisableSkyLightContribution = !bSkyLightAffectReflection;
+					const bool bEnableSkyLightContribution = bSkyLightAffectReflection;
 
 					TraceMaterialRayPacked(
 						Payload,
@@ -461,7 +461,7 @@ RAY_TRACING_ENTRY_RAYGEN(RayTracingReflectionsRGS)
 						InstanceInclusionMask,
 						TopLayerRay,
 						RayCone,
-						bDisableSkyLightContribution);
+						bEnableSkyLightContribution);
 
 					// Trace separate ray for translucent objects
 					float Transmission = 1.0;
@@ -570,7 +570,7 @@ RAY_TRACING_ENTRY_RAYGEN(RayTracingReflectionsRGS)
 					{
 						const uint RayFlags = 0;
 						const uint InstanceInclusionMask = RAY_TRACING_MASK_OPAQUE;
-						const bool bDisableSkyLightContribution = !bSkyLightAffectReflection;
+						const bool bEnableSkyLightContribution = bSkyLightAffectReflection;
 
 						FMaterialClosestHitPayload BottomLayerPayload = TraceRayAndAccumulateBottomLayerResults(
 							BottomLayerRay,
@@ -583,7 +583,7 @@ RAY_TRACING_ENTRY_RAYGEN(RayTracingReflectionsRGS)
 							ShouldDoDirectLighting,
 							ShouldDoEmissiveAndIndirectLighting,
 							RayCone,
-							bDisableSkyLightContribution,
+							bEnableSkyLightContribution,
 							BottomLayerRadiance);
 
 						// Distance and sky: similar processing to main ray.
