Cherry-pick 10119105: Fixes a bug in FBilinearSampleInfos

diff --git a/Engine/Shaders/Private/TextureSampling.ush b/Engine/Shaders/Private/TextureSampling.ush
--- a/Engine/Shaders/Private/TextureSampling.ush
+++ b/Engine/Shaders/Private/TextureSampling.ush
@@ -48,7 +48,7 @@ FBilinearSampleInfos GetBilinearSampleLevelInfosEx(float2 SampleUV, float2 Textu
 {
 	FBilinearSampleInfos Infos;
 
-	float2 HistoryBufferPixelCoord = SampleUV * View.BufferSizeAndInvSize.xy * InvMipLevelPow2;
+	float2 HistoryBufferPixelCoord = SampleUV * TextureSize * InvMipLevelPow2;
 
 	Infos.TextureInvSize = TextureInvSize * MipLevelPow2;
 	Infos.TopLeftPixelCoord = floor(HistoryBufferPixelCoord - 0.5);
