Cherry-pick 10110444: Fixes a bug with pixel position computation when the denoiser is running at half res.

diff --git a/Engine/Shaders/Private/ScreenSpaceDenoise/SSDCommon.ush b/Engine/Shaders/Private/ScreenSpaceDenoise/SSDCommon.ush
--- a/Engine/Shaders/Private/ScreenSpaceDenoise/SSDCommon.ush
+++ b/Engine/Shaders/Private/ScreenSpaceDenoise/SSDCommon.ush
@@ -73,6 +73,7 @@
 uint2 ViewportMin;
 uint2 ViewportMax;
 float4 ThreadIdToBufferUV;
+float4 BufferSizeAndInvSize;
 float4 BufferBilinearUVMinMax;
 float2 BufferUVToOutputPixelPosition;
 float4 BufferUVToScreenPosition;
diff --git a/Engine/Shaders/Private/ScreenSpaceDenoise/SSDDefinitions.ush b/Engine/Shaders/Private/ScreenSpaceDenoise/SSDDefinitions.ush
--- a/Engine/Shaders/Private/ScreenSpaceDenoise/SSDDefinitions.ush
+++ b/Engine/Shaders/Private/ScreenSpaceDenoise/SSDDefinitions.ush
@@ -41,7 +41,7 @@
 #define DEBUG_DISABLE_SIGNAL_COMPRESSION 0
 
 /** Enables debugging output from shaders. */
-#define DEBUG_OUTPUT 1
+#define DEBUG_OUTPUT 0
 
 
 //------------------------------------------------------- SHARED CONSTANTS
diff --git a/Engine/Shaders/Private/ScreenSpaceDenoise/SSDSpatialAccumulation.usf b/Engine/Shaders/Private/ScreenSpaceDenoise/SSDSpatialAccumulation.usf
--- a/Engine/Shaders/Private/ScreenSpaceDenoise/SSDSpatialAccumulation.usf
+++ b/Engine/Shaders/Private/ScreenSpaceDenoise/SSDSpatialAccumulation.usf
@@ -783,7 +783,7 @@ void MainCS(
 		SetBilateralPreset(CONFIG_BILATERAL_PRESET, /* inout */ KernelConfig);
 
 		// SGPRs
-		KernelConfig.BufferSizeAndInvSize = View.BufferSizeAndInvSize;
+		KernelConfig.BufferSizeAndInvSize = BufferSizeAndInvSize;
 		KernelConfig.BufferBilinearUVMinMax = BufferBilinearUVMinMax;
 		KernelConfig.KernelSpreadFactor = KernelSpreadFactor;
 		KernelConfig.HarmonicPeriode = HarmonicPeriode;
@@ -908,7 +908,7 @@ void MainCS(
 			
 				uint RaySubPixelId = View.StateFrameIndex & (UpscaleFactor * UpscaleFactor - 1);
 
-				KernelConfig.BufferUV = ((ViewportMin + ClosestRayPixelCoord + (0.5 * KernelSpreadFactor + 0.5))) * View.BufferSizeAndInvSize.zw;
+				KernelConfig.BufferUV = ((ViewportMin + ClosestRayPixelCoord + (0.5 * KernelSpreadFactor + 0.5))) * KernelConfig.BufferSizeAndInvSize.zw;
 			
 				// Sample the center of the kernel by comparing it against the RefSceneMetadata, since it may no match.
 				KernelConfig.bSampleKernelCenter = true;
@@ -944,7 +944,7 @@ void MainCS(
 			#endif
 
 			// Put the kernel center at the center of the quad. Half pixel shift is done in the sample offsets.
-			KernelConfig.BufferUV = (DispatchThreadId | 1) * View.BufferSizeAndInvSize.zw;
+			KernelConfig.BufferUV = (DispatchThreadId | 1) * KernelConfig.BufferSizeAndInvSize.zw;
 
 			// Id of the pixel in the quad. This is to match hard coded first samples of the sample set.
 			KernelConfig.SampleTrackId = ((DispatchThreadId.x & 1) | ((DispatchThreadId.y & 1) << 1));
diff --git a/Engine/Shaders/Private/ScreenSpaceDenoise/SSDTemporalAccumulation.usf b/Engine/Shaders/Private/ScreenSpaceDenoise/SSDTemporalAccumulation.usf
--- a/Engine/Shaders/Private/ScreenSpaceDenoise/SSDTemporalAccumulation.usf
+++ b/Engine/Shaders/Private/ScreenSpaceDenoise/SSDTemporalAccumulation.usf
@@ -540,7 +540,7 @@ void TemporallyAccumulate(
 			SetBilateralPreset(CONFIG_HISTORY_BILATERAL_PRESET, /* inout */ KernelConfig);
 
 			// SGPR configuration of the kernel.
-			KernelConfig.BufferSizeAndInvSize = View.BufferSizeAndInvSize;
+			KernelConfig.BufferSizeAndInvSize = BufferSizeAndInvSize;
 			KernelConfig.BufferBilinearUVMinMax = BufferBilinearUVMinMax;
 			
 			#if CONFIG_CLAMP_UV_PER_SIGNAL
@@ -709,7 +709,7 @@ void TemporallyAccumulate(
 
 		// SGPR configuration of the kernel.
 		{
-			KernelConfig.BufferSizeAndInvSize = View.BufferSizeAndInvSize;
+			KernelConfig.BufferSizeAndInvSize = BufferSizeAndInvSize;
 			KernelConfig.BufferBilinearUVMinMax = BufferBilinearUVMinMax;
 		}
 
diff --git a/Engine/Source/Runtime/Renderer/Private/ScreenSpaceDenoise.cpp b/Engine/Source/Runtime/Renderer/Private/ScreenSpaceDenoise.cpp
--- a/Engine/Source/Runtime/Renderer/Private/ScreenSpaceDenoise.cpp
+++ b/Engine/Source/Runtime/Renderer/Private/ScreenSpaceDenoise.cpp
@@ -706,6 +706,7 @@ BEGIN_SHADER_PARAMETER_STRUCT(FSSDCommonParameters, )
 	SHADER_PARAMETER(FIntPoint, ViewportMin)
 	SHADER_PARAMETER(FIntPoint, ViewportMax)
 	SHADER_PARAMETER(FVector4, ThreadIdToBufferUV)
+	SHADER_PARAMETER(FVector4, BufferSizeAndInvSize)
 	SHADER_PARAMETER(FVector4, BufferBilinearUVMinMax)
 	SHADER_PARAMETER(FVector2D, BufferUVToOutputPixelPosition)
 	SHADER_PARAMETER(float, WorldDepthToPixelWorldRadius)
@@ -1392,6 +1393,11 @@ static void DenoiseSignalAtConstantPixelDensity(
 	{
 		CommonParameters.ViewportMin = Viewport.Min;
 		CommonParameters.ViewportMax = Viewport.Max;
+		CommonParameters.BufferSizeAndInvSize = FVector4(
+			float(BufferExtent.X),
+			float(BufferExtent.Y),
+			1.0f / float(BufferExtent.X),
+			1.0f / float(BufferExtent.Y));
 		CommonParameters.BufferBilinearUVMinMax = FVector4(
 			float(Viewport.Min.X + 0.5f) / float(BufferExtent.X),
 			float(Viewport.Min.Y + 0.5f) / float(BufferExtent.Y),
