Fixes UV miss-alignement on old r.SSGI.LeakFreeReprojection=0

diff --git a/Engine/Shaders/Private/SSRT/SSRTDiffuseIndirect.usf b/Engine/Shaders/Private/SSRT/SSRTDiffuseIndirect.usf
--- a/Engine/Shaders/Private/SSRT/SSRTDiffuseIndirect.usf
+++ b/Engine/Shaders/Private/SSRT/SSRTDiffuseIndirect.usf
@@ -539,9 +539,6 @@ void MainCS(
 			if (bHit)
 			{
 				float2 ReducedColorUV =  HitUVz.xy * ColorBufferScaleBias.xy + ColorBufferScaleBias.zw;
-				Level = min(Level, View.GeneralPurposeTweak);
-
-
 
 				float4 SampleColor = ColorTexture.SampleLevel(ColorTextureSampler, ReducedColorUV, Level);
 				
diff --git a/Engine/Shaders/Private/SSRT/SSRTPrevFrameReduction.usf b/Engine/Shaders/Private/SSRT/SSRTPrevFrameReduction.usf
--- a/Engine/Shaders/Private/SSRT/SSRTPrevFrameReduction.usf
+++ b/Engine/Shaders/Private/SSRT/SSRTPrevFrameReduction.usf
@@ -88,19 +88,11 @@ void MainCS(
 	float2 ViewportUV = BufferUVToViewportUV(SceneBufferUV);
     float2 ScreenPosition = ViewportUVToScreenPos(ViewportUV);
 	
-	float2 PrevFrameUV;
-	float2 PrevScreen;
-	bool bIsSky = false;
-
 	float4 PrevColor;
-
 	float WorldDepth;
 
 	#if DIM_LOWER_MIPS
 	{
-		PrevFrameUV = HigherMipDownScaleFactor * (float2(DispatchThreadId) * 2.0 + 1.0) * ReducedSceneColorTexelSize;
-		PrevScreen = 0.0; // unused.
-		
 		#if DIM_LEAK_FREE
 		{
 			{
@@ -132,15 +124,15 @@ void MainCS(
 
 				float4 SampleColor = float4(
 					HigherMipTexture.SampleLevel(HigherMipTextureSampler, SampleUV, 0).rgb,
-					HigherAlphaMipTexture.SampleLevel(HigherAlphaMipTextureSampler, SampleUV, 0).r);
+					HigherAlphaMipTexture.SampleLevel(HigherAlphaMipTextureSampler, SampleUV, 0).a);
 
 				PrevColor += SampleColor * SampleWeight;
 			}
 		}
 		#else
 		{
-			WorldDepth = 0.0;
-			PrevColor = float4(HigherMipTexture.SampleLevel(HigherMipTextureSampler, PrevFrameUV, 0).rgb, 1);
+			float2 HigherMipUV = HigherMipDownScaleFactor * (float2(DispatchThreadId) * 1.0 + 0.5) * ReducedSceneColorTexelSize;
+			PrevColor = float4(HigherMipTexture.SampleLevel(HigherMipTextureSampler, HigherMipUV, 0).rgb, 1);
 		}
 		#endif
 	}
@@ -152,9 +144,9 @@ void MainCS(
 		// Camera motion for pixel (in ScreenPos space).
 		float4 ThisClip = float4(ScreenPosition, DeviceZ, 1);
 		float4 PrevClip = mul(ThisClip, View.ClipToPrevClip);
-		PrevScreen = PrevClip.xy / PrevClip.w;
+		float2 PrevScreen = PrevClip.xy / PrevClip.w;
 
-		bIsSky = WorldDepth > 100 * 1000;
+		bool bIsSky = WorldDepth > 100 * 1000;
 
 		float2 EncodedVelocity = SceneVelocityBuffer.SampleLevel(SceneVelocityBufferSampler, SceneBufferUV, 0).xy;
 		if (EncodedVelocity.x > 0.0)
@@ -162,7 +154,7 @@ void MainCS(
 			PrevScreen = ThisClip.xy - DecodeVelocityFromTexture(EncodedVelocity);
 		}
 
-		PrevFrameUV = PrevScreen.xy * PrevScreenPositionScaleBias.xy + PrevScreenPositionScaleBias.zw;
+		float2 PrevFrameUV = PrevScreen.xy * PrevScreenPositionScaleBias.xy + PrevScreenPositionScaleBias.zw;
 
 		// Fetch color.
 		#if DIM_LEAK_FREE
