Use constant defined instance of literal for node size dispatch.

diff --git a/Engine/Shaders/Private/HairStrands/HairStrandsVisibilityCopyIndirectArg.usf b/Engine/Shaders/Private/HairStrands/HairStrandsVisibilityCopyIndirectArg.usf
--- a/Engine/Shaders/Private/HairStrands/HairStrandsVisibilityCopyIndirectArg.usf
+++ b/Engine/Shaders/Private/HairStrands/HairStrandsVisibilityCopyIndirectArg.usf
@@ -33,7 +33,7 @@ void CopyCS(uint GroupIndex : SV_GroupIndex, uint3 DispatchThreadId : SV_Dispatc
 	else
 	{
 		const uint DispatchY = ceil(float(TotalGroup) / float(HAIR_VISIBILITY_GROUP_COUNT_WIDTH));
-		OutArgBuffer[0] = 64;
+		OutArgBuffer[0] = HAIR_VISIBILITY_GROUP_COUNT_WIDTH;
 		OutArgBuffer[1] = DispatchY;
 		OutArgBuffer[2] = 1;
 	}
