Fix issue where the velocity vector wouldn't take the closest sample.

diff --git a/Engine/Shaders/Private/HairStrands/HairStrandsVisibilityCompaction.usf b/Engine/Shaders/Private/HairStrands/HairStrandsVisibilityCompaction.usf
--- a/Engine/Shaders/Private/HairStrands/HairStrandsVisibilityCompaction.usf
+++ b/Engine/Shaders/Private/HairStrands/HairStrandsVisibilityCompaction.usf
@@ -236,10 +236,10 @@ void MainCS(uint GroupIndex : SV_GroupIndex, uint3 DispatchThreadId : SV_Dispatc
 			AverageVelocity += Velocity;
 			++ValidVelocityCount;
 		#elif PERMUTATION_VELOCITY == VELOCITY_TYPE_CLOSEST
-			if (ClosestDepth < Categorization.ClosestDepth) // Inverse-Z
+			if (ClosestDepth < OutNodeData.Depth) // Inverse-Z
 			{
 				ClosestVelocity = Velocity;
-				ClosestDepth = Categorization.ClosestDepth;
+				ClosestDepth = OutNodeData.Depth;
 			}
 		#elif PERMUTATION_VELOCITY == VELOCITY_TYPE_MAX
 			const float2 CurrentVelocity = Velocity;
@@ -420,7 +420,7 @@ void MainCS(uint GroupIndex : SV_GroupIndex, uint3 DispatchThreadId : SV_Dispatc
 
 			#if PERMUTATION_VELOCITY == VELOCITY_TYPE_CLOSEST
 			float2 ClosestEncodedVelocity = 0;
-			float ClosestDepth = 0;
+			float ClosestDepth = 0; // Inverse-Z
 			#endif		
 
 			#if PERMUTATION_VELOCITY == VELOCITY_TYPE_MAX
@@ -465,11 +465,12 @@ void MainCS(uint GroupIndex : SV_GroupIndex, uint3 DispatchThreadId : SV_Dispatc
 
 					#if PERMUTATION_VELOCITY == VELOCITY_TYPE_CLOSEST
 					const float2 EncodedVelocity = MSAA_VelocityTexture.Load(PixelCoord, SampleIt);
-					if (ClosestDepth < Categorization.ClosestDepth) // Inverse-Z
+					if (ClosestDepth < SampleDepth) // Inverse-Z
 					{
 						ClosestEncodedVelocity = EncodedVelocity;
-						ClosestDepth = Categorization.ClosestDepth;
+						ClosestDepth = SampleDepth;
 					}
+
 					#endif
 
 					#if PERMUTATION_VELOCITY == VELOCITY_TYPE_MAX
