Insure we are always using the closest sample for MSAA visibility pass.

diff --git a/Engine/Shaders/Private/HairStrands/HairStrandsVisibilityCompaction.usf b/Engine/Shaders/Private/HairStrands/HairStrandsVisibilityCompaction.usf
--- a/Engine/Shaders/Private/HairStrands/HairStrandsVisibilityCompaction.usf
+++ b/Engine/Shaders/Private/HairStrands/HairStrandsVisibilityCompaction.usf
@@ -296,6 +296,7 @@ void ComputeUniqueSamplesWithoutDepthTest(
 #if PERMUTATION_MATERIAL_COMPACTION == 1
 	float4 UniqueTangentAndDepths[HAIR_VISIBILITY_MAXSAMPLECOUNT];
 #endif
+	float ClosestDepth = 0; // Inverse-Z
 	for (uint SampleIt = 0; SampleIt < OutSet.HairSampleCount; ++SampleIt)
 	{
 		const uint4 EncodedData = InMSAA_IDTexture.Load(PixelCoord, SampleIt);
@@ -304,7 +305,7 @@ void ComputeUniqueSamplesWithoutDepthTest(
 			continue;
 
 		const float SampleDepth = InMSAA_DepthTexture.Load(PixelCoord, SampleIt);
-		if (OpaqueDepth > SampleDepth) // Reverse-Z
+		if (OpaqueDepth > SampleDepth) // Inverse-Z
 			continue;
 
 		++OutSet.ValidSampleCount;
@@ -318,7 +319,6 @@ void ComputeUniqueSamplesWithoutDepthTest(
 		for (uint UniqueIt = 0; UniqueIt < OutSet.UniqueSampleCount; ++UniqueIt)
 		{
 		#if PERMUTATION_MATERIAL_COMPACTION == 1
-			const float SceneDepth = ConvertFromDeviceZ(SampleDepth);
 			const bool bIsSimilar  = 
 				VisibilityData.PrimitiveId == OutSamples[UniqueIt].x ||
 				(abs(dot(UniqueTangentAndDepths[UniqueIt].xyz, VisibilityData.Tangent)) > CosTangentThreshold &&
@@ -330,6 +330,10 @@ void ComputeUniqueSamplesWithoutDepthTest(
 			if (bIsSimilar)
 			{
 				OutSamples[UniqueIt].y += 1;
+				if (ClosestDepth < SampleDepth)
+				{
+					OutSamples[OutSet.UniqueSampleCount].z = SampleIt;
+				}
 				bAlreadyExist = true;
 				break;
 			}
