Fixed compilation issue with EyeIndex.

Tested VR with -game -emulatestereo, disto and water in deferred, forward and Es3.1.

diff --git a/Engine/Shaders/Private/BasePassPixelShader.usf b/Engine/Shaders/Private/BasePassPixelShader.usf
--- a/Engine/Shaders/Private/BasePassPixelShader.usf
+++ b/Engine/Shaders/Private/BasePassPixelShader.usf
@@ -774,8 +774,10 @@ void FPixelShaderInOut_MainPS(
 	inout FPixelShaderOut Out)
 {
 #if INSTANCED_STEREO
-	ResolvedView = ResolveView(Interpolants.EyeIndex);
+	const uint EyeIndex = Interpolants.EyeIndex;
+	ResolvedView = ResolveView(EyeIndex);
 #else
+	const uint EyeIndex = 0;
 	ResolvedView = ResolveView();
 #endif
 
@@ -1115,12 +1117,6 @@ ShadingModel);
 		#elif FORWARD_SHADING || TRANSLUCENCY_LIGHTING_SURFACE_FORWARDSHADING || TRANSLUCENCY_LIGHTING_SURFACE_LIGHTINGVOLUME || MATERIAL_SHADINGMODEL_SINGLELAYERWATER
 			uint GridIndex = 0;
 
-			#if INSTANCED_STEREO
-				const uint EyeIndex = Interpolants.EyeIndex;
-			#else
-				const uint EyeIndex = 0;
-			#endif
-
 			#if FEATURE_LEVEL >= FEATURE_LEVEL_SM5
 				GridIndex = ComputeLightGridCellIndex((uint2)(MaterialParameters.SvPosition.xy - ResolvedView.ViewRectMin.xy), MaterialParameters.SvPosition.w, EyeIndex);
 
@@ -1179,11 +1175,7 @@ ShadingModel);
 	if (FogStruct.ApplyVolumetricFog > 0) 
 	{
 		float3 VolumeUV = ComputeVolumeUV(MaterialParameters.AbsoluteWorldPosition, ResolvedView.WorldToClip);
-#if INSTANCED_STEREO
-		Fogging = CombineVolumetricFog(HeightFogging, VolumeUV, Interpolants.EyeIndex);
-#else
-		Fogging = CombineVolumetricFog(HeightFogging, VolumeUV, 0);
-#endif
+		Fogging = CombineVolumetricFog(HeightFogging, VolumeUV, EyeIndex);
 	}
 #endif
 
@@ -1265,7 +1257,7 @@ ShadingModel);
 			OpaqueBasePass.SceneDepthWithoutSingleLayerWaterTexture, SingleLayerWaterSceneDepthSampler, // Scene depth texture
 			Specular, NullDistortionParams,
 			SunIlluminance, WaterDiffuseIndirectIlluminance, EnvBrdf,
-			CameraIsUnderWater, WaterVisibility);
+			CameraIsUnderWater, WaterVisibility, EyeIndex);
 
 		// Add water luminance contribution
 		Color += WaterLighting.Luminance;
@@ -1280,7 +1272,7 @@ ShadingModel);
 			OpaqueBasePass.SceneWithoutSingleLayerWaterMaxUV,
 			Specular, OpaqueBasePass.DistortionParams,
 			SunIlluminance, WaterDiffuseIndirectIlluminance, EnvBrdf,
-			CameraIsUnderWater, WaterVisibility).Luminance;
+			CameraIsUnderWater, WaterVisibility, EyeIndex).Luminance;
 #endif
 	}
 #endif // MATERIAL_SHADINGMODEL_SINGLELAYERWATER
diff --git a/Engine/Shaders/Private/DistortAccumulatePS.usf b/Engine/Shaders/Private/DistortAccumulatePS.usf
--- a/Engine/Shaders/Private/DistortAccumulatePS.usf
+++ b/Engine/Shaders/Private/DistortAccumulatePS.usf
@@ -32,8 +32,10 @@ void Main(
 	)
 {
 #if INSTANCED_STEREO
+	const uint EyeIndex = Interpolants.EyeIndex;
 	ResolvedView = ResolveView(Interpolants.EyeIndex);
 #else
+	const uint EyeIndex = 0;
 	ResolvedView = ResolveView();
 #endif
 
@@ -64,7 +66,7 @@ void Main(
 	// Compute UV distortion
 	float2 BufferUVDistortion = ComputeBufferUVDistortion(
 		MaterialParameters, PixelMaterialInputs, ResolvedView,
-		Normal, DistortionParams, ScreenUV, Refraction, true);
+		Normal, DistortionParams, ScreenUV, Refraction, true, EyeIndex);
 
 	// Sample depth at distortion offset
 	float2 DistortBufferUV = ScreenUV + BufferUVDistortion;
diff --git a/Engine/Shaders/Private/DistortionCommon.ush b/Engine/Shaders/Private/DistortionCommon.ush
--- a/Engine/Shaders/Private/DistortionCommon.ush
+++ b/Engine/Shaders/Private/DistortionCommon.ush
@@ -6,7 +6,7 @@
 
 float2 ComputeBufferUVDistortion(
 	in FMaterialPixelParameters MaterialParameters, in FPixelMaterialInputs PixelMaterialInputs, in ViewState ResolvedView,
-	in half3 Normal, in float4 DistortionParameters, in float2 ScreenUV, in float2 RefractionIORAndBias, in bool TryToClip)
+	in half3 Normal, in float4 DistortionParameters, in float2 ScreenUV, in float2 RefractionIORAndBias, in bool TryToClip, in uint EyeIndex)
 {
 #if MATERIAL_TANGENTSPACENORMAL
 	half3 ViewNormal = normalize(TransformTangentVectorToView(MaterialParameters, Normal));
@@ -35,7 +35,7 @@ float2 ComputeBufferUVDistortion(
 
 	// When ISR is enabled we store two FOVs in the DistortionParameters and compute the aspect ratio here
 #if INSTANCED_STEREO
-	float InvTanHalfFov = DistortionParameters[Interpolants.EyeIndex];
+	float InvTanHalfFov = DistortionParameters[EyeIndex];
 	float Ratio = DistortionParameters.z / DistortionParameters.w;
 #else
 	float InvTanHalfFov = DistortionParameters.x;
diff --git a/Engine/Shaders/Private/MobileBasePassPixelShader.usf b/Engine/Shaders/Private/MobileBasePassPixelShader.usf
--- a/Engine/Shaders/Private/MobileBasePassPixelShader.usf
+++ b/Engine/Shaders/Private/MobileBasePassPixelShader.usf
@@ -717,17 +717,18 @@ void Main(
 
 #if MATERIAL_SHADINGMODEL_SINGLELAYERWATER
 	{
-		const bool CameraIsUnderWater = false;	// Fade out the material contribution over to water contribution according to material opacity.
+		const bool CameraIsUnderWater = false;												// Fade out the material contribution over to water contribution according to material opacity.
 		const float3 SunIlluminance = ResolvedView.DirectionalLightColor.rgb * PI;			// times PI because it is divided by PI on CPU (=luminance) and we want illuminance here. 
 		const float3 WaterDiffuseIndirectIlluminance = WaterDiffuseIndirectLuminance * PI;	// DiffuseIndirectLighting is luminance. So we need to multiply by PI to get illuminance.
-		const float3 EnvBrdf = SpecularColor; // SpecularColor is not F0 as in BasePassPixelShader, it is EnvBRDF.
+		const float3 EnvBrdf = SpecularColor;												// SpecularColor is not F0 as in BasePassPixelShader, it is EnvBRDF.
+		const uint EyeIndex = 0;
 
 		const float4 NullDistortionParams = 1.0f;
 		WaterVolumeLightingOutput WaterLighting = EvaluateWaterVolumeLighting(
 			MaterialParameters, PixelMaterialInputs, ResolvedView,
 			Shadow, Specular, NullDistortionParams,
 			SunIlluminance, WaterDiffuseIndirectIlluminance, EnvBrdf,
-			CameraIsUnderWater, WaterVisibility);
+			CameraIsUnderWater, WaterVisibility, EyeIndex);
 
 		// Accumulate luminance and occlude the background according to transmittance to view and mean transmittance to lights.
 		Color += WaterLighting.Luminance;
diff --git a/Engine/Shaders/Private/SingleLayerWaterCommon.ush b/Engine/Shaders/Private/SingleLayerWaterCommon.ush
--- a/Engine/Shaders/Private/SingleLayerWaterCommon.ush
+++ b/Engine/Shaders/Private/SingleLayerWaterCommon.ush
@@ -53,7 +53,7 @@ WaterVolumeLightingOutput EvaluateWaterVolumeLighting(
 #endif // SIMPLE_SINGLE_LAYER_WATER
 	float Specular, const float4 DistortionParams,
 	float3 SunIlluminance, float3 AmbiantIlluminance, float3 EnvBrdf,
-	bool CameraIsUnderWater, float WaterVisibility)
+	bool CameraIsUnderWater, float WaterVisibility, uint EyeIndex)
 {
 	WaterVolumeLightingOutput Output;
 	Output.Luminance = 0.0f;
@@ -103,7 +103,7 @@ WaterVolumeLightingOutput EvaluateWaterVolumeLighting(
 		float2 Refraction = GetMaterialRefraction(PixelMaterialInputs); // X = IOR, Y = Refraction Bias
 		float2 BufferUVDistortion = ComputeBufferUVDistortion(
 			MaterialParameters, PixelMaterialInputs, ResolvedView,
-			MaterialParameters.WorldNormal, DistortionParams, ViewportUV, Refraction, false);
+			MaterialParameters.WorldNormal, DistortionParams, ViewportUV, Refraction, false, EyeIndex);
 		PostProcessUVDistortion(MaterialParameters, PixelMaterialInputs, PixelSceneDepth, BufferUVDistortion, Refraction);
 		// Also reduce refraction according to water thickness
 		BufferUVDistortion *= saturate((PixelSceneDepth - WaterDepth) * 1.0 / 30.0f);// *(1 - saturate(WaterDepth * 1.0 / 2000.0f));
