Add option for reduce voxel tracing distance, based on parameter CVar value.

Change the default value to shorter distance in order improve performance.

diff --git a/Engine/Shaders/Private/HairStrands/HairStrandsDeepShadowMask.usf b/Engine/Shaders/Private/HairStrands/HairStrandsDeepShadowMask.usf
--- a/Engine/Shaders/Private/HairStrands/HairStrandsDeepShadowMask.usf
+++ b/Engine/Shaders/Private/HairStrands/HairStrandsDeepShadowMask.usf
@@ -355,7 +355,8 @@ void MainPS(
 			WorldPosition += DepthBias;
 		}
 
-		const float HairCount = ComputeHairCountVoxel(WorldPosition, L, Voxel_MinAABB, Voxel_MaxAABB, Voxel_Resolution, Voxel_DensityTexture, Voxel_DensityScale, 1);
+		const float DistanceThreshold = 100000;
+		const float HairCount = ComputeHairCountVoxel(WorldPosition, L, Voxel_MinAABB, Voxel_MaxAABB, Voxel_Resolution, Voxel_DensityTexture, Voxel_DensityScale, 1, DistanceThreshold);
 		Visibility = saturate(1 - HairCount);
 	}
 
diff --git a/Engine/Shaders/Private/HairStrands/HairStrandsDeepTransmittanceMask.usf b/Engine/Shaders/Private/HairStrands/HairStrandsDeepTransmittanceMask.usf
--- a/Engine/Shaders/Private/HairStrands/HairStrandsDeepTransmittanceMask.usf
+++ b/Engine/Shaders/Private/HairStrands/HairStrandsDeepTransmittanceMask.usf
@@ -147,15 +147,16 @@ float4 ComputeTransmittanceVoxel(
 	}
 
 	// Compute the number of hair count between light & shading point
+	const float DistanceThreshold = 100000;
 	float HairCount = 0;
-	if (ClusterId == 0) HairCount = ComputeHairCountVoxel(WorldPosition, L, Voxel_MinAABBs[0], Voxel_MaxAABBs[0], Voxel_Resolution, Voxel_DensityTexture0, Settings.DeepShadowDensityScale, 0);
-	if (ClusterId == 1) HairCount = ComputeHairCountVoxel(WorldPosition, L, Voxel_MinAABBs[1], Voxel_MaxAABBs[1], Voxel_Resolution, Voxel_DensityTexture1, Settings.DeepShadowDensityScale, 0);
-	if (ClusterId == 2) HairCount = ComputeHairCountVoxel(WorldPosition, L, Voxel_MinAABBs[2], Voxel_MaxAABBs[2], Voxel_Resolution, Voxel_DensityTexture2, Settings.DeepShadowDensityScale, 0);
-	if (ClusterId == 3) HairCount = ComputeHairCountVoxel(WorldPosition, L, Voxel_MinAABBs[3], Voxel_MaxAABBs[3], Voxel_Resolution, Voxel_DensityTexture3, Settings.DeepShadowDensityScale, 0);
-	if (ClusterId == 4) HairCount = ComputeHairCountVoxel(WorldPosition, L, Voxel_MinAABBs[4], Voxel_MaxAABBs[4], Voxel_Resolution, Voxel_DensityTexture4, Settings.DeepShadowDensityScale, 0);
-	if (ClusterId == 5) HairCount = ComputeHairCountVoxel(WorldPosition, L, Voxel_MinAABBs[5], Voxel_MaxAABBs[5], Voxel_Resolution, Voxel_DensityTexture5, Settings.DeepShadowDensityScale, 0);
-	if (ClusterId == 6) HairCount = ComputeHairCountVoxel(WorldPosition, L, Voxel_MinAABBs[6], Voxel_MaxAABBs[6], Voxel_Resolution, Voxel_DensityTexture6, Settings.DeepShadowDensityScale, 0);
-	if (ClusterId == 7) HairCount = ComputeHairCountVoxel(WorldPosition, L, Voxel_MinAABBs[7], Voxel_MaxAABBs[7], Voxel_Resolution, Voxel_DensityTexture7, Settings.DeepShadowDensityScale, 0);
+	if (ClusterId == 0) HairCount = ComputeHairCountVoxel(WorldPosition, L, Voxel_MinAABBs[0], Voxel_MaxAABBs[0], Voxel_Resolution, Voxel_DensityTexture0, Settings.DeepShadowDensityScale, 0, DistanceThreshold);
+	if (ClusterId == 1) HairCount = ComputeHairCountVoxel(WorldPosition, L, Voxel_MinAABBs[1], Voxel_MaxAABBs[1], Voxel_Resolution, Voxel_DensityTexture1, Settings.DeepShadowDensityScale, 0, DistanceThreshold);
+	if (ClusterId == 2) HairCount = ComputeHairCountVoxel(WorldPosition, L, Voxel_MinAABBs[2], Voxel_MaxAABBs[2], Voxel_Resolution, Voxel_DensityTexture2, Settings.DeepShadowDensityScale, 0, DistanceThreshold);
+	if (ClusterId == 3) HairCount = ComputeHairCountVoxel(WorldPosition, L, Voxel_MinAABBs[3], Voxel_MaxAABBs[3], Voxel_Resolution, Voxel_DensityTexture3, Settings.DeepShadowDensityScale, 0, DistanceThreshold);
+	if (ClusterId == 4) HairCount = ComputeHairCountVoxel(WorldPosition, L, Voxel_MinAABBs[4], Voxel_MaxAABBs[4], Voxel_Resolution, Voxel_DensityTexture4, Settings.DeepShadowDensityScale, 0, DistanceThreshold);
+	if (ClusterId == 5) HairCount = ComputeHairCountVoxel(WorldPosition, L, Voxel_MinAABBs[5], Voxel_MaxAABBs[5], Voxel_Resolution, Voxel_DensityTexture5, Settings.DeepShadowDensityScale, 0, DistanceThreshold);
+	if (ClusterId == 6) HairCount = ComputeHairCountVoxel(WorldPosition, L, Voxel_MinAABBs[6], Voxel_MaxAABBs[6], Voxel_Resolution, Voxel_DensityTexture6, Settings.DeepShadowDensityScale, 0, DistanceThreshold);
+	if (ClusterId == 7) HairCount = ComputeHairCountVoxel(WorldPosition, L, Voxel_MinAABBs[7], Voxel_MaxAABBs[7], Voxel_Resolution, Voxel_DensityTexture7, Settings.DeepShadowDensityScale, 0, DistanceThreshold);
 
 	return ComputeTransmittanceFromHairCount(HairCount, GBuffer, SinLightAngle, HairLUTTexture, HairLUTSampler);
 }
diff --git a/Engine/Shaders/Private/HairStrands/HairStrandsDeepVoxelCommon.ush b/Engine/Shaders/Private/HairStrands/HairStrandsDeepVoxelCommon.ush
--- a/Engine/Shaders/Private/HairStrands/HairStrandsDeepVoxelCommon.ush
+++ b/Engine/Shaders/Private/HairStrands/HairStrandsDeepVoxelCommon.ush
@@ -29,7 +29,8 @@ float ComputeHairCountVoxel(
 	float InVoxelResolution, 
 	Texture3D<uint> InDensityTexture, 
 	float DensityScale, 
-	float CountThreshold)
+	float CountThreshold,
+	float DistanceThreshold)
 {
 	float OutHairCount = 0;
 
@@ -47,7 +48,7 @@ float ComputeHairCountVoxel(
 		// This assumes that the voxel volume is cubic (i.e. equal dimensions on all sides)
 		float3 O = lerp(WorldPosition, IntersectEndPoint, HitT.xxx);
 		float3 E = lerp(WorldPosition, IntersectEndPoint, HitT.yyy);
-		const float OELength = length(E - O);
+		const float OELength = min(length(E - O), DistanceThreshold);
 
 		// Step according to voxel size
 		const float StepCount = float(min(ceil(OELength / VoxelWorldSize), 256u));	// TO TRY: sample with fixed count budget and jitter temporally
@@ -77,6 +78,7 @@ float ComputeHairCountVoxelCone(
 	Texture3D<uint> InDensityTexture,
 	float InVoxelDensityScale,
 	float CountThreshold,
+	float DistanceThreshold,
 	float TanConeAngle)
 {
 	// Reverse-Z
@@ -101,7 +103,8 @@ float ComputeHairCountVoxelCone(
 		// Count the number of fibers which are within a cylinder defined by the voxel size, 
 		// and the distance between the origin and the extent of the volume
 		// This assumes that the voxel volume is cubic (i.e. equal dimensions on all sides)
-		const float Distance = length(HitT.y * D);
+		const float DistanceThreshold2 = DistanceThreshold * DistanceThreshold;
+		const float Distance = min(DistanceThreshold, length(HitT.y * D));
 		const uint MaxStepCount = min(ceil(Distance / VoxelSize), 256u);
 		float T = 0;
 		float MipLevel = 0;
@@ -110,8 +113,14 @@ float ComputeHairCountVoxelCone(
 		{
 			float3 StepOffset = 0;
 
-			const float3 HitP = O + T * LightDirection + StepOffset;
+			const float3 RaymarchOffset = T * LightDirection;
+			const float3 HitP = O + RaymarchOffset + StepOffset;
 			const float3 UVW = (HitP - InVoxelAABBMin) / (InVoxelAABBMax - InVoxelAABBMin);
+
+			const float Distance2 = dot(RaymarchOffset, RaymarchOffset);
+			if (Distance2 > DistanceThreshold2)
+				return OutHairCount;
+
 			if (!IsInVoxelBounds(UVW, 0, 1))
 				return OutHairCount;
 
diff --git a/Engine/Shaders/Private/HairStrands/HairStrandsEnvironmentAO.usf b/Engine/Shaders/Private/HairStrands/HairStrandsEnvironmentAO.usf
--- a/Engine/Shaders/Private/HairStrands/HairStrandsEnvironmentAO.usf
+++ b/Engine/Shaders/Private/HairStrands/HairStrandsEnvironmentAO.usf
@@ -18,13 +18,25 @@ float Voxel_DepthBiasScale;
 float AO_Power;
 float AO_Intensity;
 uint  AO_SampleCount;
+float AO_DistanceThreshold;
 
 float3 Voxel_MinAABB;
 float3 Voxel_MaxAABB;
 uint   Voxel_Resolution;
 Texture3D<uint> Voxel_DensityTexture;
 
-float InternalComputeHairCountVoxel(float3 WorldPosition, float3 LightDirection, float3 InVoxelAABBMin, float3 InVoxelAABBMax, float InVoxelResolution, Texture3D<uint> InDensityTexture, float DensityScale, float CountThreshold, bool bUseConeTracing)
+// todo: make this function in common with the sky lighting integration pass
+float InternalComputeHairCountVoxel(
+	float3 WorldPosition, 
+	float3 LightDirection, 
+	float3 InVoxelAABBMin, 
+	float3 InVoxelAABBMax, 
+	float InVoxelResolution, 
+	Texture3D<uint> InDensityTexture, 
+	float DensityScale, 
+	float CountThreshold, 
+	float DistanceThreshold, 
+	bool bUseConeTracing)
 {
 	if (bUseConeTracing)
 		return ComputeHairCountVoxelCone(
@@ -36,6 +48,7 @@ float InternalComputeHairCountVoxel(float3 WorldPosition, float3 LightDirection,
 			InDensityTexture,
 			DensityScale,
 			CountThreshold,
+			DistanceThreshold,
 			Voxel_TanConeAngle);
 	else
 		return ComputeHairCountVoxel(
@@ -46,7 +59,8 @@ float InternalComputeHairCountVoxel(float3 WorldPosition, float3 LightDirection,
 			InVoxelResolution,
 			InDensityTexture,
 			DensityScale,
-			CountThreshold);
+			CountThreshold,
+			DistanceThreshold);
 }
 
 float TraceAO(float3 WorldPosition, float3 WorldNormal, float4 SvPosition)
@@ -104,7 +118,7 @@ float TraceAO(float3 WorldPosition, float3 WorldNormal, float4 SvPosition)
 		const float3 SampleWorldPosition = WorldPosition + VoxelOffset;
 
 		// Compute the number of hair count between light & shading point
-		const float HairCount = InternalComputeHairCountVoxel(SampleWorldPosition, SampleL, Voxel_MinAABB, Voxel_MaxAABB, Voxel_Resolution, Voxel_DensityTexture, Voxel_DensityScale, 0, true);
+		const float HairCount = InternalComputeHairCountVoxel(SampleWorldPosition, SampleL, Voxel_MinAABB, Voxel_MaxAABB, Voxel_Resolution, Voxel_DensityTexture, Voxel_DensityScale, 0, AO_DistanceThreshold, true);
 		BentNormal += SampleL * saturate(1 - HairCount);
 		Occlusion += saturate(HairCount);
 	}
diff --git a/Engine/Shaders/Private/HairStrands/HairStrandsEnvironmentLighting.usf b/Engine/Shaders/Private/HairStrands/HairStrandsEnvironmentLighting.usf
--- a/Engine/Shaders/Private/HairStrands/HairStrandsEnvironmentLighting.usf
+++ b/Engine/Shaders/Private/HairStrands/HairStrandsEnvironmentLighting.usf
@@ -74,6 +74,8 @@ Texture3D<float4>			HairScatteringLUTTexture;
 Texture3D<float4>			HairEnergyLUTTexture;
 SamplerState 				HairLUTSampler;
 
+float HairDistanceThreshold;
+
 uint  MultipleScatterSampleCount;
 uint  JitterSphericalIntegration;
 uint  HairComponents;
@@ -88,7 +90,17 @@ float3 Voxel_MaxAABB;
 uint   Voxel_Resolution;
 Texture3D<uint> Voxel_DensityTexture;
 
-float InternalComputeHairCountVoxel(float3 WorldPosition, float3 LightDirection, float3 InVoxelAABBMin, float3 InVoxelAABBMax, float InVoxelResolution, Texture3D<uint> InDensityTexture, float DensityScale, float CountThreshold, bool bUseConeTracing)
+float InternalComputeHairCountVoxel(
+	float3 WorldPosition, 
+	float3 LightDirection, 
+	float3 InVoxelAABBMin, 
+	float3 InVoxelAABBMax, 
+	float InVoxelResolution, 
+	Texture3D<uint> InDensityTexture, 
+	float DensityScale, 
+	float CountThreshold, 
+	float DistanceThreshold,
+	bool bUseConeTracing)
 {
 	if (bUseConeTracing)
 		return ComputeHairCountVoxelCone(
@@ -100,6 +112,7 @@ float InternalComputeHairCountVoxel(float3 WorldPosition, float3 LightDirection,
 			InDensityTexture,
 			DensityScale,
 			CountThreshold,
+			DistanceThreshold,
 			Voxel_TanConeAngle);
 	else
 		return ComputeHairCountVoxel(
@@ -110,7 +123,8 @@ float InternalComputeHairCountVoxel(float3 WorldPosition, float3 LightDirection,
 			InVoxelResolution,
 			InDensityTexture,
 			DensityScale,
-			CountThreshold);
+			CountThreshold,
+			DistanceThreshold);
 }
 
 //////////////////////////////////////////////////////////////////////////////////////////////////////////////////////////
@@ -145,6 +159,7 @@ float3 ReflectionEnvironment(FGBufferData GBuffer, float2 ScreenPosition, float2
 	// Trace/sample transmittance in various direction to have estimate of where is occlusion & computing the integrated dual scattering value
 	float3 UnoccludedN = -V;
 	float3 AvgDualScattering = 0;
+	const float MaxSearchDistance = HairDistanceThreshold; // in cm
 	#if 1
 	if ((HairComponents & HAIR_COMPONENT_R) || (HairComponents & HAIR_COMPONENT_TRT) || (HairComponents & HAIR_COMPONENT_TT))
 	{
@@ -200,7 +215,7 @@ float3 ReflectionEnvironment(FGBufferData GBuffer, float2 ScreenPosition, float2
 
 			// Compute the number of hair count between light & shading point
 			const float SinLightAngle = dot(SampleL, T);
-			const float HairCount = InternalComputeHairCountVoxel(SampleWorldPosition, SampleL, Voxel_MinAABB, Voxel_MaxAABB, Voxel_Resolution, Voxel_DensityTexture, Voxel_DensityScale, 0, true);
+			const float HairCount = InternalComputeHairCountVoxel(SampleWorldPosition, SampleL, Voxel_MinAABB, Voxel_MaxAABB, Voxel_Resolution, Voxel_DensityTexture, Voxel_DensityScale, 0, MaxSearchDistance, true);
 			const float4 SampleTransmittance = ComputeTransmittanceFromHairCount(HairCount, GBuffer, SinLightAngle, HairScatteringLUTTexture, HairLUTSampler);
 
 			const float ScalarTransmittance = max(SampleTransmittance.x, max(SampleTransmittance.y, SampleTransmittance.z));
@@ -319,7 +334,7 @@ float3 ReflectionEnvironment(FGBufferData GBuffer, float2 ScreenPosition, float2
 			const float3 SampleWorldPosition = WorldPosition + DepthBias;
 
 			// Compute the number of hair count between light & shading point
-			const float HairCount = InternalComputeHairCountVoxel(SampleWorldPosition, L, Voxel_MinAABB, Voxel_MaxAABB, Voxel_Resolution, Voxel_DensityTexture, Voxel_DensityScale, 0, true);
+			const float HairCount = InternalComputeHairCountVoxel(SampleWorldPosition, L, Voxel_MinAABB, Voxel_MaxAABB, Voxel_Resolution, Voxel_DensityTexture, Voxel_DensityScale, 0, MaxSearchDistance, true);
 			Transmittance = ComputeTransmittanceFromHairCount(HairCount, GBuffer, SinLightAngle, HairScatteringLUTTexture, HairLUTSampler);
 		}
 
@@ -373,8 +388,10 @@ float3 ReflectionEnvironment(FGBufferData GBuffer, float2 ScreenPosition, float2
 			const float3 SampleWorldPosition = WorldPosition + DepthBias;
 
 			// Compute the number of hair count between light & shading point
+			// For TT, we don't want to abort the hair counting, as this is important to get the proper transmission
+			const float DistanceThresholdTT = 100000;
 			const float DensityScaleFactor = 2;
-			const float HairCount = InternalComputeHairCountVoxel(SampleWorldPosition, L, Voxel_MinAABB, Voxel_MaxAABB, Voxel_Resolution, Voxel_DensityTexture, Voxel_DensityScale * DensityScaleFactor, 0, false);
+			const float HairCount = InternalComputeHairCountVoxel(SampleWorldPosition, L, Voxel_MinAABB, Voxel_MaxAABB, Voxel_Resolution, Voxel_DensityTexture, Voxel_DensityScale * DensityScaleFactor, 0, DistanceThresholdTT, false);
 			Transmittance = ComputeTransmittanceFromHairCount(HairCount, GBuffer, SinLightAngle, HairScatteringLUTTexture, HairLUTSampler);
 		}
 
diff --git a/Engine/Source/Runtime/Renderer/Private/HairStrands/HairStrandsEnvironment.cpp b/Engine/Source/Runtime/Renderer/Private/HairStrands/HairStrandsEnvironment.cpp
--- a/Engine/Source/Runtime/Renderer/Private/HairStrands/HairStrandsEnvironment.cpp
+++ b/Engine/Source/Runtime/Renderer/Private/HairStrands/HairStrandsEnvironment.cpp
@@ -39,6 +39,13 @@ static FAutoConsoleVariableRef CVarHairStrandsSkyLightingJitterSphericalIntegrat
 static int32 GHairStrandsSkyAOSampleCount = 16;
 static FAutoConsoleVariableRef CVarHairStrandsSkyAOSampleCount(TEXT("r.HairStrands.SkyAO.SampleCount"), GHairStrandsSkyAOSampleCount, TEXT("Number of samples used for evaluating hair AO (default is set to 16)."), ECVF_Scalability | ECVF_RenderThreadSafe);
 
+static float GHairStrandsSkyAODistanceThreshold = 10;
+static float GHairStrandsSkyLightingDistanceThreshold = 10;
+static int32 GHairStrandsSkyLightingUseViewHairCount = 1;
+static FAutoConsoleVariableRef CVarHairStrandsSkyAOThreshold(TEXT("r.HairStrands.SkyAO.DistanceThreshold"), GHairStrandsSkyAODistanceThreshold, TEXT("Max distance for occlusion search."), ECVF_Scalability | ECVF_RenderThreadSafe);
+static FAutoConsoleVariableRef CVarHairStrandsSkyLightingDistanceThreshold(TEXT("r.HairStrands.SkyLighting.DistanceThreshold"), GHairStrandsSkyLightingDistanceThreshold, TEXT("Max distance for occlusion search."), ECVF_Scalability | ECVF_RenderThreadSafe);
+static FAutoConsoleVariableRef CVarHairStrandsSkyLightingUseViewHairCount(TEXT("r.HairStrands.SkyLighting.UseViewHairCount"), GHairStrandsSkyLightingUseViewHairCount, TEXT("Use the primary visibility hair count output to compute the TT term instead of the voxel grid."), ECVF_Scalability | ECVF_RenderThreadSafe);
+
 static bool GetHairStrandsSkyLightingEnable() { return GHairSkylightingEnable > 0; }
 static bool GetHairStrandsSkyAOEnable() { return GHairSkyAOEnable > 0; }
 static float GetHairStrandsSkyLightingConeAngle() { return FMath::Max(0.f, GHairSkylightingConeAngle); }
@@ -72,6 +79,7 @@ class FHairEnvironmentAO : public FGlobalShader
 		SHADER_PARAMETER(float, AO_Power)
 		SHADER_PARAMETER(float, AO_Intensity)
 		SHADER_PARAMETER(uint32, AO_SampleCount)
+		SHADER_PARAMETER(float, AO_DistanceThreshold)
 		SHADER_PARAMETER_STRUCT_INCLUDE(FSceneTextureParameters, SceneTextures)
 		SHADER_PARAMETER_STRUCT_INCLUDE(FSceneTextureSamplerParameters, SceneTextureSamplers)
 
@@ -93,6 +101,8 @@ static void AddHairStrandsEnvironmentAOPass(
 	const FHairStrandsClusterData& ClusterData,
 	FRDGTextureRef Output)
 {
+	check(Output);
+
 	FSceneTextureParameters SceneTextures;
 	SetupSceneTextureParameters(GraphBuilder, &SceneTextures);
 
@@ -114,7 +124,7 @@ static void AddHairStrandsEnvironmentAOPass(
 	PassParameters->AO_Power = Settings.AmbientOcclusionPower;
 	PassParameters->AO_Intensity = Settings.AmbientOcclusionIntensity;
 	PassParameters->AO_SampleCount = FMath::Max(uint32(GHairStrandsSkyAOSampleCount), 1u);
-	check(Output);
+	PassParameters->AO_DistanceThreshold = FMath::Max(GHairStrandsSkyAODistanceThreshold, 1.f);
 	PassParameters->RenderTargets[0] = FRenderTargetBinding(Output, ERenderTargetLoadAction::ELoad);
 
 	FHairEnvironmentAO::FPermutationDomain PermutationVector;
@@ -256,6 +266,7 @@ class FHairEnvironmentLightingCS : public FGlobalShader
 		SHADER_PARAMETER(uint32, MultipleScatterSampleCount)
 		SHADER_PARAMETER(uint32, HairComponents)
 		SHADER_PARAMETER(uint32, JitterSphericalIntegration)
+		SHADER_PARAMETER(float, HairDistanceThreshold)
 
 		SHADER_PARAMETER_TEXTURE(Texture2D, PreIntegratedGF)
 		SHADER_PARAMETER_SAMPLER(SamplerState, PreIntegratedGFSampler)
@@ -317,6 +328,7 @@ static FRDGBufferRef AddHairStrandsEnvironmentLightingPassCS(
 	PassParameters->Voxel_DensityScale = GetHairStrandsVoxelizationDensityScale();
 	PassParameters->Voxel_DepthBiasScale = GetHairStrandsVoxelizationDepthBiasScale();
 	PassParameters->Voxel_TanConeAngle = FMath::Tan(FMath::DegreesToRadians(GetHairStrandsSkyLightingConeAngle()));
+	PassParameters->HairDistanceThreshold = FMath::Max(GHairStrandsSkyLightingDistanceThreshold, 1.f);
 	PassParameters->MaxVisibilityNodeCount = VisibilityData.NodeData->Desc.NumElements;
 	PassParameters->MultipleScatterSampleCount = FMath::Max(uint32(GHairStrandsSkyLightingSampleCount), 1u);
 	PassParameters->JitterSphericalIntegration = GHairStrandsSkyLightingJitterSphericalIntegration ? 1 : 0;
