Add (fudged) correct scale for MS roughness lobe.

It was arbitrary taking the TT roughness for evaluating the MS contribution, but usually MS is while more diffuse/high-variance than the actual TT term, even though most of the energy is coming from the TT term. This scale factor (eye-balled) try to compensate for this. This integration needs to be replaced in mid-term.

diff --git a/Engine/Shaders/Private/HairStrands/HairStrandsEnvironmentLighting.usf b/Engine/Shaders/Private/HairStrands/HairStrandsEnvironmentLighting.usf
--- a/Engine/Shaders/Private/HairStrands/HairStrandsEnvironmentLighting.usf
+++ b/Engine/Shaders/Private/HairStrands/HairStrandsEnvironmentLighting.usf
@@ -270,8 +270,9 @@ float3 ReflectionEnvironment(FGBufferData GBuffer, float2 ScreenPosition, float2
 		}
 		#endif	
 
-		// Select the TT term for representing the MS roughness lookup as TT lobe carries the most energy
-		const float GGXRoughness = HairToGGXRoughness_TT(GBuffer.Roughness);
+		// Select the TT term for representing the MS roughness lookup as TT lobe carries the most energy + fudge scale as the MS lobe should be wider than a TT lobe
+		const float FudgeScale = 2;
+		const float GGXRoughness = HairToGGXRoughness_TT(GBuffer.Roughness) * FudgeScale;
 		ColorMS.rgb += PreExposure * GatherRadiance(R, GGXRoughness, PixelPosition, PixelSceneDepth, WorldPosition);
 		//ColorMS.rgb += PreExposure * GetSkySHDiffuse(R) * View.SkyLightColor.rgb;
 		ColorMS.rgb *= AvgDualScattering;
