Fixed: Avoid NaNs in mobile TAA on iOS Metal with shader fast math enabled

diff --git a/Engine/Shaders/Private/PostProcessMobile.usf b/Engine/Shaders/Private/PostProcessMobile.usf
--- a/Engine/Shaders/Private/PostProcessMobile.usf
+++ b/Engine/Shaders/Private/PostProcessMobile.usf
@@ -1213,9 +1213,17 @@ void AaPS_ES2(
 
 	LumaT = max(LumaT, LumaMin);
 	LumaT = min(LumaT, LumaMax);
-
+	
+#if METAL_PROFILE
+	// avoid division by zero 
+	// normally saturate should kill NaNs, but Metal fast math does not do it
+	// this can be fixed by instructing cross-compiler to use precise math for saturate
+	half LumaD = LumaP - LumaM;
+	half BlendFinal = saturate((LumaT - LumaM) * (abs(LumaD) < 0.001 ? 1000.0 : rcp(LumaD)));
+#else
 	half BlendFinal = saturate((LumaT - LumaM) * rcp(LumaP - LumaM));
-
+#endif
+			
 	OutColor.rgb = lerp(ColorM.rgb, ColorP.rgb, BlendFinal);
 	OutColor.a = 0;
 }
