Fixes regressions in denoiser's temporal accumulation when doing variance based history rejections

diff --git a/Engine/Shaders/Private/ScreenSpaceDenoise/SSDTemporalAccumulation.usf b/Engine/Shaders/Private/ScreenSpaceDenoise/SSDTemporalAccumulation.usf
--- a/Engine/Shaders/Private/ScreenSpaceDenoise/SSDTemporalAccumulation.usf
+++ b/Engine/Shaders/Private/ScreenSpaceDenoise/SSDTemporalAccumulation.usf
@@ -736,7 +736,7 @@ void TemporallyAccumulate(
 				/* inout */ UnusedCompressedAccumulators,
 				KernelConfig.RefBufferUV,
 				RefSceneMetadata,
-				CurrentFrameSamples);
+				NormalizeToOneSampleArray(CurrentFrameSamples));
 		}
 
 		{
