DUMMY_INTERP is causing DXC compiler to give the TEXCOORD8 VS to PS user attribute the wrong name - looks like its taken from an incrementing index.  HLSL uses the original identifiers.

Move the dummy interpolator to the end of the structure so the rest of the attributes match up again.

diff --git a/Engine/Shaders/Private/MobileBasePassCommon.ush b/Engine/Shaders/Private/MobileBasePassCommon.ush
--- a/Engine/Shaders/Private/MobileBasePassCommon.ush
+++ b/Engine/Shaders/Private/MobileBasePassCommon.ush
@@ -41,9 +41,7 @@ struct FSharedMobileBasePassInterpolants
 #if USE_VERTEX_FOG && !PACK_INTERPOLANTS
 	INTERP_TYPE VertexFog			: TEXCOORD7;
 #endif
-#if LANDSCAPE_BUG_WORKAROUND
-	INTERP_TYPE DummyInterp		: DUMMY_INTERP;
-#endif
+
 	float4 PixelPosition	: TEXCOORD8; // xyz = world position, w = clip z
 
 #if USE_WORLD_POSITION_EXCLUDING_SHADER_OFFSETS
@@ -57,6 +55,13 @@ struct FSharedMobileBasePassInterpolants
 #if MOBILE_MULTI_VIEW
 	float MultiViewId : VIEW_ID;
 #endif
+
+// Keep this dummy at the bottom of the structure for the DXC compiler
+// as it currently uses indexing for interpolator naming.  Otherwise
+// the VS to PS attribute mapping won't match
+#if LANDSCAPE_BUG_WORKAROUND
+	INTERP_TYPE DummyInterp		: DUMMY_INTERP;
+#endif
 };
 
 #define FMobileBasePassInterpolantsVSToPS FSharedMobileBasePassInterpolants
