Fixed pixel depth offset invariance between depth pass and base pass.

diff --git a/Engine/Shaders/Private/MaterialTemplate.ush b/Engine/Shaders/Private/MaterialTemplate.ush
--- a/Engine/Shaders/Private/MaterialTemplate.ush
+++ b/Engine/Shaders/Private/MaterialTemplate.ush
@@ -2502,7 +2502,7 @@ float ApplyPixelDepthOffsetToMaterialParameters(inout FMaterialPixelParameters M
 	MaterialParameters.SvPosition.w = MaterialParameters.ScreenPosition.w;
 	MaterialParameters.AbsoluteWorldPosition += MaterialParameters.CameraVector * PixelDepthOffset;
 
-	OutDepth = DeviceDepth;
+	OutDepth = INVARIANT(DeviceDepth);
 
 	return PixelDepthOffset;
 }
