Clamp the sampling in the water composite pass properly

#rnx

diff --git a/Engine/Shaders/Private/SingleLayerWaterComposite.usf b/Engine/Shaders/Private/SingleLayerWaterComposite.usf
--- a/Engine/Shaders/Private/SingleLayerWaterComposite.usf
+++ b/Engine/Shaders/Private/SingleLayerWaterComposite.usf
@@ -28,7 +28,7 @@ SamplerState ScreenSpaceReflectionsSampler;
 Texture2D SceneNoWaterDepthTexture;
 SamplerState SceneNoWaterDepthSampler;
 
-float2 SceneNoWaterMaxUV;
+float4 SceneNoWaterMinMaxUV;
 
 struct SingleLayerWaterCompositeOutput
 {
@@ -59,7 +59,7 @@ SingleLayerWaterCompositeOutput SingleLayerWaterComposite(float2 BufferUV, float
 #endif
 	if (ShadingModelID == SHADINGMODELID_SINGLELAYERWATER)
 	{
-		const float2 SceneNoWaterUV = min(BufferUV, SceneNoWaterMaxUV);
+		const float2 SceneNoWaterUV = clamp(BufferUV, SceneNoWaterMinMaxUV.xy, SceneNoWaterMinMaxUV.zw);
 
 		float OpaqueDepth = SceneNoWaterDepthTexture.SampleLevel(SceneNoWaterDepthSampler, SceneNoWaterUV, 0).x * SINGLE_LAYER_WATER_DEPTH_SCALE;
 		float WaterDepth = GBuffer.Depth;
