Updates to GTAO
Improved spatial and temporla filters to bleed better across discontinuities
Added Thickness Heuristic to bias AO around smaller objects - Control using r.GTAO.ThicknessBlend
Added CVARS to control the falloff so it is not hardwired
Control using r.GTAOFalloffEnd and r.GTAO.FalloffRatio

diff --git a/Engine/Shaders/Private/PostProcessAmbientOcclusion.usf b/Engine/Shaders/Private/PostProcessAmbientOcclusion.usf
--- a/Engine/Shaders/Private/PostProcessAmbientOcclusion.usf
+++ b/Engine/Shaders/Private/PostProcessAmbientOcclusion.usf
@@ -28,7 +28,7 @@
 #define GTAO_OPT 0
 
 // This define enables the thickness heuristic which assu
-#define GTAO_THICKNESS_HEURISTIC 0
+#define GTAO_THICKNESS_HEURISTIC 1
 
 
 // 0: classic with weighted sample, 1: don't normalize and adjust the formula to be simpler and faster - can look better and is cheaper (Alchemy like?)
@@ -937,15 +937,38 @@ void MainSSAOSmoothCS(
 #endif
 #endif
 
+float4 GTAOParams[4];
+// [0] - { cos(TemporalAngle), sin(TemporalAngle), TemporalOffset, FrameTemporalOffset}
+// [1] - { FrameNumber, Thicknessblend, unused, unused}
+// [2] - { TargetSizeX, TargetSizeY, 1.0/TargetSizeX, 1.0f/TargetSizeY}
+// [3] - { FallOffStart, FallOffEnd, FalloffScale, FalloffBias}
+
+struct FalloffParams
+{
+	float Start;
+	float End;
+	float EndSq;
+	float Scale;
+	float Bias;
+	uint  MaxTaps;
+};
 
 
-float4 GTAOParams[3];
+// Get the parameters for the fall off and how many taps we must travel before we are guaranteed to hit it. 
+FalloffParams ComputeFalloffParams(float ZDist)
+{
+	FalloffParams params;
 
-#define FALLOFF_START		100.0
-#define FALLOFF_END			200.0
-#define FALLOFF_START_SQ	(FALLOFF_START*FALLOFF_START)
-#define FALLOFF_END_SQ		(FALLOFF_END*FALLOFF_END)
+	// get the distance at which the 
+	params.Start	= GTAOParams[3].x;
+	params.End		= GTAOParams[3].y;
+	params.EndSq	= params.End * params.End;
+	params.Scale	= GTAOParams[3].z;
+	params.Bias  	= GTAOParams[3].w;
+	params.MaxTaps  = 8; //TODO
 
+	return params;
+}	
 
 
 #define PI_HALF (PI*0.5)
@@ -965,17 +988,17 @@ SamplerState	DepthHistoryTextureSampler;
 Texture2D		ZReadTexture;
 SamplerState	ZReadTextureSampler;
 
-
 float4 PrevScreenPositionScaleBias;
 float4 FilterParams;
 
+float ClampScale(float Scale)
+{ 
+	return clamp(Scale, 2.0, 8.0);
+}
 
-
-float Falloff(float distSq)
+float FalloffFn(FalloffParams FalloffParam, float distSq)
 {
-	float Falloff_Scale = 1.0f / (FALLOFF_END_SQ - FALLOFF_START_SQ);
-	float Falloff_Bias = -FALLOFF_START_SQ * Falloff_Scale; 
-	return  2.0 * saturate(distSq* Falloff_Scale + Falloff_Bias   );
+	return  2.0 * saturate(distSq * FalloffParam.Scale + FalloffParam.Bias);
 }
 
 float3 GetViewSpacePosFromHZB(float2 UV, float MipLevel)
@@ -1034,66 +1057,30 @@ float GetLinearDepthProj(float2 ScreenUV)
 	return  1.0f / (DeviceZ * View.InvDeviceZToWorldZTransform[2] - View.InvDeviceZToWorldZTransform[3]);
 }
 
-float UpdateMaxAngle(float CurrAng, float PrevAng)
+float UpdateMaxAngle(float2 UV, float CurrentBiggest, float CurrAng, float PrevAng)
 {
 #ifdef GTAO_THICKNESS_HEURISTIC
-
-	if(CurrAng > PrevAng)
+	if(CurrAng >= PrevAng)
 	{
-		return CurrAng;
+		return max(CurrentBiggest, CurrAng);
 	}
 	else
 	{
-		const float BLEND_FACTOR = 0.9;
-		return 	PrevAng ;//(CurrAng * BLEND_FACTOR) + (PrevAng * (1-BLEND_FACTOR));
+		float   ThicknessBlend = GTAOParams[1].y;
+		return 	lerp(CurrentBiggest, PrevAng ,ThicknessBlend);
 	}
 #else
-	return max(CurrAng, PrevAng);
+	return max(CurrentBiggest, CurrAng);
 #endif
 }
 
-
-float SearchForLargestAngleRef(float2 inUV, float2 deltaUV, float3 ViewSpacePosition, float3 ViewDir)
+float SearchForLargestAngleRef(float2 inUV, float2 deltaUV, float3 ViewSpacePosition, float3 ViewDir, FalloffParams FalloffParam )
 {
 	float BiggestAngle = -1.0f;
 	float PrevAngle    = -1.0f;
 
 	float2 UV	   = inUV ;
 
-#if GTAO_OPT
-	// Precompute as much stuff upfront
-	float CurrDepth			= GetLinearDepthProj(UV);
-	UV						+= deltaUV;
-	float2 DeltaVS_XY	    = ScreenToViewPos(UV, CurrDepth).xy - ViewSpacePosition.xy;
-	float  DeltaVS_XYLen    = sqrt(dot( DeltaVS_XY.xy, DeltaVS_XY.xy));
-	float  CurrVS_XYLen		= DeltaVS_XYLen;
-
-	for(int i=1; i<=GTAO_NUMTAPS; i++)
-	{
-		float SceneDepth	= GetLinearDepthProj(UV);
-		float DiffDepth		= (CurrDepth - SceneDepth);
-
-		float distSQ		= (CurrVS_XYLen*CurrVS_XYLen) + (DiffDepth*DiffDepth);
-
-		UV					+= deltaUV;
-		CurrVS_XYLen		+= DeltaVS_XYLen;
-
-		// At this point we could early out if DistSQ is greater than falloff
-		BRANCH
-		if(distSQ < FALLOFF_END_SQ)
-		{
-			float ooDist    = rsqrt(distSQ);
-			float cosh		= ooDist * DiffDepth;
-
-			float falloff	= Falloff(distSQ);
-			BiggestAngle	= max(BiggestAngle, cosh - falloff);
-		}
-	}
-	
-#else
-
-	deltaUV *=1;
-
 	// Do first sample so we have a previous to work with 
 	for(int i=1; i<=GTAO_NUMTAPS; i++)
 	{
@@ -1105,33 +1092,30 @@ float SearchForLargestAngleRef(float2 inUV, float2 deltaUV, float3 ViewSpacePosi
 		float3 DiffV		= (SampVS - ViewSpacePosition) ;
 		float distSQ		= dot(DiffV, DiffV);
 	
-		// At this point we could early out if DistSQ is greater than falloff
-	//	BRANCH
-	//	if(distSQ < FALLOFF_END_SQ)
+		BRANCH
+		if(distSQ < FalloffParam.EndSq)
 		{
 			float ooDist    = rsqrt(distSQ);
 			float cosh		= ooDist * dot(DiffV, ViewDir );
 	
-			float falloff	= Falloff(distSQ);
+			float falloff	= FalloffFn(FalloffParam, distSQ);
 
 			cosh			= cosh - falloff;
-			BiggestAngle	= max(BiggestAngle, cosh);
-
-			//BiggestAngle	= UpdateMaxAngle(cosh, PrevAngle);
-			//PrevAngle		= cosh;
+			BiggestAngle	= UpdateMaxAngle(inUV, BiggestAngle, cosh, PrevAngle);
+			PrevAngle		= cosh;
 		}
 	}	
-#endif
+
 
 	return BiggestAngle;
 }
 
 
 
-float SearchForLargestAngle(float2 inUV, float2 deltaUV, float3 ViewSpacePosition, float3 ViewDir)
+float SearchForLargestAngle(float2 inUV, float2 deltaUV, float3 ViewSpacePosition, float3 ViewDir, FalloffParams FalloffParam)
 {
 #if GTAO_REF
-	return SearchForLargestAngleRef(inUV,  deltaUV,  ViewSpacePosition,  ViewDir);
+	return SearchForLargestAngleRef(inUV,  deltaUV,  ViewSpacePosition,  ViewDir, FalloffParam);
 #endif
 
 	float BiggestAngle = -1.0f;
@@ -1140,107 +1124,46 @@ float SearchForLargestAngle(float2 inUV, float2 deltaUV, float3 ViewSpacePositio
 	float2 UV	= inUV;
 
 	float MipLevel = 0;
-	
-
-#if GTAO_OPT
-	// Precompute as much stuff upfront
-	float CurrDepth			= GetLinearDepthProj(UV);
-	UV						+= deltaUV;
-	float2 DeltaVS_XY	    = ScreenToViewPos(UV, CurrDepth).xy - ViewSpacePosition.xy;
-	float  DeltaVS_XYLen    = sqrt(dot( DeltaVS_XY.xy, DeltaVS_XY.xy));
-	float  CurrVS_XYLen		= DeltaVS_XYLen;
-
-	float2 HZB_UV  = float2( UV.x * HZBRemapping.x + HZBRemapping.z, UV.y * HZBRemapping.y + HZBRemapping.w);
-	float2 HZB_dUV = float2( deltaUV.x * HZBRemapping.x, deltaUV.y * HZBRemapping.y );
-
-	UNROLL
 	for(int i=1; i<=GTAO_NUMTAPS; i++)
 	{
 		if(i==3)
 		{
-			HZB_dUV += HZB_dUV;
-			DeltaVS_XYLen	+= DeltaVS_XYLen;
+			deltaUV *=2;
 			MipLevel++;
 		}
 		
 		if(i==GTAO_NUMTAPS/2)
 		{
-			HZB_dUV += HZB_dUV;
-			DeltaVS_XYLen	+= DeltaVS_XYLen;
-			MipLevel++;
-		}
-
-
-		float DeviceZ		= Texture2DSampleLevel( PostprocessInput1, PostprocessInput1Sampler, HZB_UV, MipLevel).r;
-		float SceneDepth    = 1.0f / (DeviceZ * View.InvDeviceZToWorldZTransform[2] - View.InvDeviceZToWorldZTransform[3]);
-
-		float DiffDepth		= (CurrDepth - SceneDepth);
-		float distSQ		= (CurrVS_XYLen*CurrVS_XYLen) + (DiffDepth*DiffDepth);
-
-		HZB_UV				+= HZB_dUV;
-		CurrVS_XYLen		+= DeltaVS_XYLen;
-
-		// At this point we could early out if DistSQ is greater than falloff
-		BRANCH
-		if(distSQ < FALLOFF_END_SQ)
-		{
-			float ooDist    = rsqrt(distSQ);
-			float cosh		= ooDist * DiffDepth;
-
-			float falloff	= Falloff(distSQ);
-			cosh			= cosh - falloff;
-
-			BiggestAngle	= max(BiggestAngle, cosh);
-
-		//	BiggestAngle	= UpdateMaxAngle(cosh, PrevAngle);
-		//	PrevAngle		= cosh;
-		}
-	}	
-
-
-#else
-
-	for(int i=1; i<=GTAO_NUMTAPS; i++)
-	{
-		if(i==3)
-		{
-			deltaUV *=1.5;
-			MipLevel++;
-		}
-		
-		if(i==GTAO_NUMTAPS/2)
-		{
-			deltaUV *=1.5;
+			deltaUV *=2;
 			MipLevel++;
 		}
 
 		UV				+= deltaUV;
 		
 		// Quantize UV to pixel centre
-		//float2 UV2 = floor( (View.BufferSizeAndInvSize.xy * UV) + float2(0.25,0.25) ) * View.BufferSizeAndInvSize.zw;
-		//UV2 += View.BufferSizeAndInvSize.zw*0.5;
+		float2 UV2 = floor( (View.BufferSizeAndInvSize.xy * UV) + float2(0.25,0.25) ) * View.BufferSizeAndInvSize.zw;
+		UV2 += View.BufferSizeAndInvSize.zw*0.5;
 		
+		//	UV = saturate(UV2);
+
 		float3 SampVS	= GetViewSpacePosFromHZB(UV,MipLevel);
 		float3 DiffV	= (SampVS - ViewSpacePosition) ;
 		float distSQ	= dot(DiffV, DiffV);
 
-		// At this point we could early out if DistSQ is greater than falloff
-		// Profiling on Xbox with the branch on and off. Without the ALU usage is better but the VGPR count is much higher.
 		BRANCH
-		if(distSQ < FALLOFF_END_SQ)
+		if(distSQ < FalloffParam.EndSq)
 		{
 			float ooDist    = rsqrt(distSQ);
-			float cosh		= ooDist * dot(DiffV,ViewDir );
-			float falloff	= Falloff(distSQ);
-			
-			cosh			= cosh - falloff;
-			BiggestAngle	= max(BiggestAngle, cosh);
+			float cosh		= ooDist * dot(DiffV, ViewDir );
+	
+			float falloff	= FalloffFn(FalloffParam, distSQ);
 
-			//BiggestAngle	= UpdateMaxAngle(cosh, PrevAngle);
-			//PrevAngle		= cosh;
+			cosh			= cosh - falloff;
+			BiggestAngle	= UpdateMaxAngle(inUV, BiggestAngle, cosh, PrevAngle);
+			PrevAngle		= cosh;
 		}
+
 	}
-#endif	
 
 
 	return BiggestAngle;
@@ -1248,21 +1171,19 @@ float SearchForLargestAngle(float2 inUV, float2 deltaUV, float3 ViewSpacePositio
 
 
 // Given a screen space Axis this will search for the min and max angles
-float2 SearchAxisForAngles(float2 UV, float3 ScreenDir, float3 ViewDir, float3 ViewSpacePos, float Scale, float Offset)
+float2 SearchAxisForAngles(float2 UV, float3 ScreenDir, float3 ViewDir, float3 ViewSpacePos, float Scale, float Offset, FalloffParams FalloffParam)
 {
-	// TODO - Fix aspect ratio
-	float2 DeltaUV	= Scale * float2(View.BufferSizeAndInvSize.z * ScreenDir.x, View.BufferSizeAndInvSize.w * ScreenDir.y);
-	
+	float2 DeltaUV	= Scale * float2(GTAOParams[2].z * ScreenDir.x, GTAOParams[2].w * ScreenDir.y);
 	DeltaUV.y	  *=-1;
 
-	float Ang1 = SearchForLargestAngle( UV + (DeltaUV*Offset),   DeltaUV, ViewSpacePos, ViewDir);
-	float Ang2 = SearchForLargestAngle( UV - (DeltaUV*Offset),  -DeltaUV, ViewSpacePos, ViewDir);
+	float Ang1 = SearchForLargestAngle( UV + (DeltaUV*Offset),   DeltaUV, ViewSpacePos, ViewDir, FalloffParam);
+	float Ang2 = SearchForLargestAngle( UV - (DeltaUV*Offset),  -DeltaUV, ViewSpacePos, ViewDir, FalloffParam);
 	return float2(Ang1, Ang2);
 }
 
 float ComputeInnerIntegral(float2 Angles, float3 ScreenDir, float3 ViewDir, float3 ViewSpaceNormal, float SceneDepth)
 {
-	// Given the angles found in the search plane we need to project the View Space GBuffer Normal onto the plane defined by the search axis and the View Direction and perform the inner integrate
+	// Given the angles found in the search plane we need to project the View Space Normal onto the plane defined by the search axis and the View Direction and perform the inner integrate
 	float3 PlaneNormal		= normalize(cross(ScreenDir,ViewDir));
 	float3 Perp				= cross(ViewDir, PlaneNormal);
 	float3 ProjNormal		= ViewSpaceNormal - PlaneNormal * dot(ViewSpaceNormal, PlaneNormal);
@@ -1278,8 +1199,6 @@ float ComputeInnerIntegral(float2 Angles, float3 ScreenDir, float3 ViewDir, floa
 	Angles.x = Gamma + max(-Angles.x - Gamma, -(PI_HALF) );
 	Angles.y = Gamma + min( Angles.y - Gamma,  (PI_HALF) );
 
-	LenProjNormal = 1;
-
 	float AO = ( LenProjNormal *  0.25 * 
 					    ( (Angles.x * SinGamma + CosGamma - cos((2.0 * Angles.x) - Gamma)) +
 				  	      (Angles.y * SinGamma + CosGamma - cos((2.0 * Angles.y) - Gamma)) ));
@@ -1322,14 +1241,16 @@ float3 GetRandomVector(float2 UV )
 * HORIZON SEARCH AND INNER INTEGRATE COMBINED
 *
 */
-void GTAOCombinedPSandCS(in float2 UV, out float4 OutColor, inout float Depth)
+void GTAOCombinedPSandCS(in float2 UV, out float4 OutColor, inout float OutDepth)
 {
 	OutColor = 0;
 
 	// Offset by a fraction of a pixel to unsure we don't hit between pixels when running at half res
 	UV +=  PostprocessInput0Size.zw*0.125;
 
-	float SceneDepth = GetDepthFromAOInput(UV);		
+	float DeviceZ		 = LookupDeviceZ(UV);
+	float SceneDepth	 = ConvertFromDeviceZ(DeviceZ);
+
 	if(SceneDepth > ScreenSpaceAOParams[4].w)	
 	{		
 		OutColor = 1;
@@ -1354,11 +1275,13 @@ void GTAOCombinedPSandCS(in float2 UV, out float4 OutColor, inout float Depth)
 
 	// Given the depth determine how many pixels away we are guranteed to be outside of the falloff distance.
 	float InvTanHalfFov			= ScreenSpaceAOParams[3].w;
-	float ScreenPixelsToCutoff  = View.BufferSizeAndInvSize.x * ((FALLOFF_END*0.5) / ViewSpacePos.z) ; 
+	float ScreenPixelsToCutoff  = (GTAOParams[2].x * InvTanHalfFov * 7)  / ViewSpacePos.z ; 
 	float Scale = ScreenPixelsToCutoff /  GTAO_NUMTAPS;
 
 	// Clamp the scale. Too small and it won't be jumping by more than 1 pixel. Too much and it'll thrash the cache
-	Scale = clamp(Scale, 2.0, 16.0);
+	Scale = ClampScale(Scale);
+
+	FalloffParams FalloffParam = ComputeFalloffParams(ViewSpacePos.z);
 
 #if GTAO_REF
 	for(int i=0; i<GTAO_NUMROTS; i++)
@@ -1368,8 +1291,7 @@ void GTAOCombinedPSandCS(in float2 UV, out float4 OutColor, inout float Depth)
 #else
 		float3 ScreenDir	= float3(RandomVec.x, RandomVec.y, 0.0);
 #endif
-
-		float2 Angles = SearchAxisForAngles( UV ,  ScreenDir,  ViewDir,  ViewSpacePos, Scale, Offset);
+		float2 Angles = SearchAxisForAngles( UV ,  ScreenDir,  ViewDir,  ViewSpacePos, Scale, Offset, FalloffParam);
 
 		Angles.x = acosFast(Angles.x);
 		Angles.y = acosFast(Angles.y);
@@ -1379,12 +1301,12 @@ void GTAOCombinedPSandCS(in float2 UV, out float4 OutColor, inout float Depth)
 	}
 #endif
 
-
 	// user adjust AO
 	Sum = Sum   / (float)GTAO_NUMROTS;
 
-	OutColor.r = Sum ;
-	Depth = SceneDepth;
+	OutColor.r	= Sum ;
+	OutDepth	= SceneDepth;
+
 	return;
 }
 
@@ -1504,11 +1426,14 @@ void GTAOInnerIntegrateCS(
 
 float4 HorizonSearchParams;
 
-float2 HorizonSearchPSandCS(in float2 UV, inout float Depth)
+float2 HorizonSearchPSandCS(in float2 UV, inout float OutDepth)
 {
 	float2 OutHorizons = 0;
 
-	float SceneDepth	 = GetDepthFromAOInput(UV);
+	float DeviceZ		 = LookupDeviceZ(UV);
+	float SceneDepth	 = ConvertFromDeviceZ(DeviceZ);
+	OutDepth = DeviceZ;
+
 	if(SceneDepth > ScreenSpaceAOParams[4].w)	
 	{		
 		OutHorizons = 0;
@@ -1525,16 +1450,18 @@ float2 HorizonSearchPSandCS(in float2 UV, inout float Depth)
 
 	// Given the depth determine how many pixels away we are guranteed to be outside of the falloff distance.
 	float InvTanHalfFov			= ScreenSpaceAOParams[3].w;
-	float ScreenPixelsToCutoff  = View.BufferSizeAndInvSize.x * ((FALLOFF_END*0.5) / ViewSpacePos.z) ; 
+	float ScreenPixelsToCutoff  = (GTAOParams[2].x * InvTanHalfFov * 7)  / ViewSpacePos.z ; 
 	float Scale = ScreenPixelsToCutoff /  GTAO_NUMTAPS;
 
 	// Clamp the scale. Too small and it won't be jumping by more than 1 pixel. Too much and it'll thrash the cache
-	Scale = clamp(Scale, 2.0, 16.0);
+	Scale = ClampScale(Scale);
+
+	FalloffParams FalloffParam = ComputeFalloffParams(ViewSpacePos.z);
 
-	float2 Angles = SearchAxisForAngles( UV,  ScreenDir,  ViewDir,  ViewSpacePos, Scale, Offset);
+
+	float2 Angles = SearchAxisForAngles( UV,  ScreenDir,  ViewDir,  ViewSpacePos, Scale, Offset, FalloffParam);
 	OutHorizons.rg = Angles*0.5+0.5;
 
-	Depth = SceneDepth;
 	return OutHorizons;
 }
 
@@ -1591,11 +1518,68 @@ float2 ReprojectPos(float2 UV, float Depth)
 	}
 	
 	float2 PrevUV = PrevScreen.xy * PrevScreenPositionScaleBias.xy + PrevScreenPositionScaleBias.zw;
-	PrevUV += PostprocessInput0Size.zw*0.25;
-
 	return PrevUV;
 }
 
+float ReadHistory(float2 UV, float CurrDepth, float CurrAO, inout float WeightSum)
+{
+	float	BilinearWeights[4];
+	float4 SizeInvSize = PostprocessInput0Size.xyzw;
+
+	float2 PixUV   = (UV * SizeInvSize.xy)-0.5;
+	float2 FloorUV = floor(PixUV);
+	float2 FracUV  = (PixUV - FloorUV); 
+	UV			   = (FloorUV * SizeInvSize.zw) + (SizeInvSize.zw*0.5);
+
+	BilinearWeights[0] = (1.0 -	FracUV.x) * ( 1.0 -	FracUV.y);
+	BilinearWeights[1] = (		FracUV.x) * ( 1.0 -	FracUV.y);
+	BilinearWeights[2] = (1.0 -	FracUV.x) * (       FracUV.y);
+	BilinearWeights[3] = (		FracUV.x) * (       FracUV.y);
+
+	// Read the 4 previous depths and History
+	float PrevDepth[4];
+	float HistoryAO[4];
+
+	float2 dUV = SizeInvSize.zw;
+	
+	// TODO - Use GatherR when available
+	HistoryAO[0] = Texture2DSample(HistoryTexture, HistoryTextureSampler, UV + float2(	  0,     0)).r;
+	HistoryAO[1] = Texture2DSample(HistoryTexture, HistoryTextureSampler, UV + float2(dUV.x,     0)).r;
+	HistoryAO[2] = Texture2DSample(HistoryTexture, HistoryTextureSampler, UV + float2(    0, dUV.y)).r;
+	HistoryAO[3] = Texture2DSample(HistoryTexture, HistoryTextureSampler, UV + float2(dUV.x, dUV.y)).r;
+
+	PrevDepth[0] = Texture2DSample(DepthHistoryTexture, DepthHistoryTextureSampler, UV + float2(    0,     0)).r;
+	PrevDepth[1] = Texture2DSample(DepthHistoryTexture, DepthHistoryTextureSampler, UV + float2(dUV.x,     0)).r;
+	PrevDepth[2] = Texture2DSample(DepthHistoryTexture, DepthHistoryTextureSampler, UV + float2(    0, dUV.y)).r;
+	PrevDepth[3] = Texture2DSample(DepthHistoryTexture, DepthHistoryTextureSampler, UV + float2(dUV.x, dUV.y)).r;
+
+	const float DepthScale =0.1;
+	float VisHistory = 0;
+	for(int i=0; i<4; i++)
+	{
+		float BilateralWeight = saturate( 1- (DepthScale * (CurrDepth - PrevDepth[i])) );
+		WeightSum  += BilateralWeight;
+		VisHistory += BilinearWeights[i] * lerp( CurrAO, HistoryAO[i], BilateralWeight);
+	}
+
+	WeightSum/=4;
+	return VisHistory;
+}
+
+
+void NeighbourhoodClamp(float2 UV, float BaseAO, inout float MinAO, inout float MaxAO)
+{
+	float2 dUV = PostprocessInput0Size.zw*1.5;
+	float AONeighbour1 = Texture2DSample(PostprocessInput0, PostprocessInput0Sampler, UV + float2(-dUV.x,-dUV.y) ).r;
+	float AONeighbour2 = Texture2DSample(PostprocessInput0, PostprocessInput0Sampler, UV + float2(-dUV.x, dUV.y) ).r;
+	float AONeighbour3 = Texture2DSample(PostprocessInput0, PostprocessInput0Sampler, UV + float2( dUV.x,-dUV.y) ).r;
+	float AONeighbour4 = Texture2DSample(PostprocessInput0, PostprocessInput0Sampler, UV + float2( dUV.x, dUV.y) ).r;
+
+	MinAO = min(BaseAO, min(min(AONeighbour1, AONeighbour2), min(AONeighbour3, AONeighbour4)));
+	MaxAO = max(BaseAO, max(max(AONeighbour1, AONeighbour2), max(AONeighbour3, AONeighbour4)));
+}
+
+
 void GTAOTemporalFilterPSAndCS(float2 UV, inout float OutAO, inout float OutDepth)
 {
 	float NewAO			= Texture2DSample(PostprocessInput0, PostprocessInput0Sampler, UV).r;
@@ -1604,40 +1588,34 @@ void GTAOTemporalFilterPSAndCS(float2 UV, inout float OutAO, inout float OutDept
 	float DeviceZ		= LookupDeviceZ( UV);
 
 	float2 PrevUV		= ReprojectPos( UV,  DeviceZ);
-	float  HistoryAO	= Texture2DSample(HistoryTexture, HistoryTextureSampler, PrevUV).r;
-
+	float WeightSum=0;
 	float CurrDepth		= SceneDepth;
+
+	float HistoryAO		= ReadHistory(PrevUV, SceneDepth, NewAO, WeightSum );
+
 	float PrevDepth		= Texture2DSample(DepthHistoryTexture, DepthHistoryTextureSampler, PrevUV).r;
 
-	float Blend			= 0.15;
+	float Blend			= 0.15;//lerp(1,0.15, WeightSum);
 
-	if( abs(CurrDepth-PrevDepth) > 100)
-		Blend =1;
 
 	if((PrevUV.x <= 0.0) || (PrevUV.y <= 0.0) || (PrevUV.x >= 1.0) || (PrevUV.y >= 1.0))
 	{
 		Blend=1;
 	}
 
-	// Do a neighbourhood clamp 
-	float2 dUV = PostprocessInput0Size.zw*2;
-	float AONeighbour1 = Texture2DSample(PostprocessInput0, PostprocessInput0Sampler, UV + float2(-dUV.x,-dUV.y) ).r;
-	float AONeighbour2 = Texture2DSample(PostprocessInput0, PostprocessInput0Sampler, UV + float2(-dUV.x, dUV.y) ).r;
-	float AONeighbour3 = Texture2DSample(PostprocessInput0, PostprocessInput0Sampler, UV + float2( dUV.x,-dUV.y) ).r;
-	float AONeighbour4 = Texture2DSample(PostprocessInput0, PostprocessInput0Sampler, UV + float2( dUV.x, dUV.y) ).r;
+	// Neighbourhood clamp 
+	float MinAO =0;
+	float MaxAO =0;
+	NeighbourhoodClamp(UV, NewAO, MinAO,  MaxAO);
 
-	float MinAO = min(NewAO, min(min(AONeighbour1, AONeighbour2), min(AONeighbour3, AONeighbour4)));
-	float MaxAO = max(NewAO, max(max(AONeighbour1, AONeighbour2), max(AONeighbour3, AONeighbour4)));
 	float HalfAO = (MaxAO - MinAO) *0.5;
 	float MidAO  = (MinAO + HalfAO);
 
 	// Ensure the HalfAO is at least some minimum ammount. This helps with the history clamp.
-	HalfAO = max(HalfAO,0.03);
+	HalfAO = max(HalfAO,0.05);
 	MinAO = saturate(MidAO - HalfAO);
 	MaxAO = saturate(MidAO + HalfAO);
 
-	float BlurredAO = 0.2 * (NewAO + AONeighbour1 + AONeighbour2 + AONeighbour3 + AONeighbour4);
-
 	OutAO = (NewAO*Blend) + (HistoryAO* (1-Blend));
 	OutAO = clamp(OutAO, MinAO, MaxAO);
 
@@ -1732,8 +1710,6 @@ void GTAOUpsampleCS(
 {
 	float2 OutColor		= 0;
 
-
-
 	// Get Pixel Pos of the group we are processing
 	int2   PixelPos				= DispatchThreadId + ScreenSpaceAOParams[5].zw; 
 
@@ -1774,18 +1750,10 @@ void GTAOUpsampleCS(
 
 	float Mag = BilinearWeights.x + BilinearWeights.y + BilinearWeights.z + BilinearWeights.w;
 
-	float AmbientOcclusionIntensity = ScreenSpaceAOParams[0].w;
-	float AmbientOcclusionPower		= ScreenSpaceAOParams[0].x*0.5;
-	// user adjust AO
-	//AO = 1 - (1 - pow(abs(AO), AmbientOcclusionPower)) * AmbientOcclusionIntensity;
-
 	// user adjust AO
 	AO = dot(AO4, BilinearWeights) / Mag;
-//	AO = lerp(1, AO, AmbientOcclusionIntensity);
-//	AO = 1 - (1 - pow(abs(AO), AmbientOcclusionPower));
 
 	OutTexture[PixelPos] = AO;
-
 }
 #endif
 
@@ -1793,6 +1761,8 @@ void GTAOUpsampleCS(
 
 float4 GTAOUpsamplePSAndCS(float2 UV)
 {
+//	UV += PostprocessInput0Size.zw * -0.5;
+
 	float AO = saturate( Texture2DSample(PostprocessInput0, PostprocessInput0Sampler, UV).r);
 	return float4(AO,0,0,0);
 }
@@ -1919,8 +1889,9 @@ void GTAOSpatialFilterCS(
 				
 				SampleZDiff -= XDiff * DiffZ_X;
 				SampleZDiff -= YDiff * DiffZ_Y;
-				
-				Weight = 1.0f - saturate(SampleZDiff*1000 );
+			
+				const float SpatialFilterWeight = 100000;
+				Weight = 1.0f - saturate(SampleZDiff*SpatialFilterWeight );
 			}
 
 			SumAO += Sample_AO * Weight;
diff --git a/Engine/Source/Runtime/Renderer/Private/CompositionLighting/PostProcessAmbientOcclusion.cpp b/Engine/Source/Runtime/Renderer/Private/CompositionLighting/PostProcessAmbientOcclusion.cpp
--- a/Engine/Source/Runtime/Renderer/Private/CompositionLighting/PostProcessAmbientOcclusion.cpp
+++ b/Engine/Source/Runtime/Renderer/Private/CompositionLighting/PostProcessAmbientOcclusion.cpp
@@ -109,6 +109,30 @@ static TAutoConsoleVariable<int32> CVarGTAOUseNormals(
 	TEXT("1: On (default)\n "),
 	ECVF_RenderThreadSafe | ECVF_Scalability);
 
+static TAutoConsoleVariable<float> CVarGTAOThicknessBlend(
+	TEXT("r.GTAO.ThicknessBlend"),
+	0.1f,
+	TEXT("A heuristic to bias occlusion for thin or thick objects. \n ")
+	TEXT("0  : Off \n ")
+	TEXT(">0 : On - Bigger values lead to reduced occlusion \n ")
+	TEXT("0.1: On (default)\n "),
+	ECVF_RenderThreadSafe | ECVF_Scalability);
+
+static TAutoConsoleVariable<float> CVarGTAOFalloffEnd(
+	TEXT("r.GTAO.FalloffEnd"),
+	200.0f,
+	TEXT("Distance at when the occlusion completes the fall off.  \n "),
+	ECVF_RenderThreadSafe | ECVF_Scalability);
+
+static TAutoConsoleVariable<float> CVarGTAOFalloffStartRatio(
+	TEXT("r.GTAO.FalloffStartRatio"),
+	0.5f,
+	TEXT("Ratio of the r.GTAO.FalloffEnd value at which it starts to fall off. \n ")
+	TEXT("Must be Between 0 and 1. \n "),
+	ECVF_RenderThreadSafe | ECVF_Scalability);
+
+
+
 float FSSAOHelper::GetAmbientOcclusionQualityRT(const FSceneView& View)
 {
 	float CVarValue = CVarAmbientOcclusionMaxQuality.GetValueOnRenderThread();
@@ -1252,7 +1276,7 @@ public:
 			Frame = ViewState->GetFrameIndex();
 		}
 			
-		const int ArraySize = 3;
+		const int ArraySize = 4;
 		FVector4 GTAOParam[ArraySize];
 
 		const float Rots[6]		= { 60.0f, 300.0f, 180.0f, 240.0f, 120.0f, 0.0f };
@@ -1260,17 +1284,30 @@ public:
 
 		float TemporalAngle = Rots[TemporalFrame % 6] * (PI / 360.0f);
 
+		// Angles of rotation that are set per frame
 		GTAOParam[0] = FVector4(cos(TemporalAngle), sin(TemporalAngle), Offsets[(TemporalFrame / 6) % 4], Offsets[TemporalFrame % 4]);
 
+		// Frame X = number , Y = Thickness param, 
+		float ThicknessBlend = CVarGTAOThicknessBlend.GetValueOnRenderThread();
+		GTAOParam[1] = FVector4(Frame, ThicknessBlend, 0.0f, 0.0f);
 
-		FIntPoint RandomizationSize = GSystemTextures.GTAORandomization->GetDesc().Extent;
-
-		GTAOParam[1] = FVector4(Frame, 0.0f, 0.0f, 0.0f);
-
+		// Destination buffer Size and InvSize
 		float Fx = float(DestSize.X);
 		float Fy = float(DestSize.Y);
 		GTAOParam[2] = FVector4(Fx, Fy, 1.0f / Fx, 1.0f / Fy);
 
+		// Fall Off Params
+		float FallOffEnd		= CVarGTAOFalloffEnd.GetValueOnRenderThread();
+		float FallOffStartRatio	= FMath::Clamp(CVarGTAOFalloffStartRatio.GetValueOnRenderThread(), 0.0f, 0.999f);
+		float FallOffStart		= FallOffEnd * FallOffStartRatio;
+		float FallOffStartSq	= FallOffStart * FallOffStart;
+		float FallOffEndSq		= FallOffEnd * FallOffEnd;
+
+		float FallOffScale		= 1.0f / (FallOffEndSq - FallOffStartSq);
+		float FallOffBias		= -FallOffStartSq * FallOffScale;
+
+		GTAOParam[3] = FVector4(FallOffStart, FallOffEnd, FallOffScale, FallOffBias);
+
 		SetShaderValueArray(RHICmdList, ShaderRHI, GTAOParams, GTAOParam, ArraySize);
 	}
 
@@ -2336,19 +2373,19 @@ public:
 		if (InputHistory.IsValid())
 		{
 			SetTextureParameter(RHICmdList, ShaderRHI, HistoryTexture, HistoryTextureSampler,
-				TStaticSamplerState<SF_Bilinear, AM_Border, AM_Border, AM_Border, 0, 0, 0xffffffff >::GetRHI(),
+				TStaticSamplerState<SF_Point, AM_Border, AM_Border, AM_Border, 0, 0, 0xffffffff >::GetRHI(),
 				InputHistory.RT[0]->GetRenderTargetItem().TargetableTexture);
 
 
 			SetTextureParameter(RHICmdList, ShaderRHI, DepthHistoryTexture, DepthHistoryTextureSampler,
-				TStaticSamplerState<SF_Bilinear>::GetRHI(),
+				TStaticSamplerState<SF_Point>::GetRHI(),
 				InputHistory.Depth[0]->GetRenderTargetItem().TargetableTexture);
 		}
 		else
 		{
 			// Need to bind a white dummy
 			SetTextureParameter(RHICmdList, ShaderRHI, HistoryTexture, HistoryTextureSampler,
-				TStaticSamplerState<SF_Bilinear>::GetRHI(),
+				TStaticSamplerState<SF_Point>::GetRHI(),
 				GSystemTextures.WhiteDummy->GetRenderTargetItem().ShaderResourceTexture);
 
 		}
@@ -2390,19 +2427,19 @@ public:
 		if (InputHistory.IsValid())
 		{
 			SetTextureParameter(Context.RHICmdList, ShaderRHI, HistoryTexture, HistoryTextureSampler,
-				TStaticSamplerState<SF_Bilinear,AM_Border, AM_Border, AM_Border, 0,0, 0xffffffff >::GetRHI(),
+				TStaticSamplerState<SF_Point,AM_Border, AM_Border, AM_Border, 0,0, 0xffffffff >::GetRHI(),
 				InputHistory.RT[0]->GetRenderTargetItem().TargetableTexture);
 
 
 			SetTextureParameter(Context.RHICmdList, ShaderRHI, DepthHistoryTexture, DepthHistoryTextureSampler,
-				TStaticSamplerState<SF_Bilinear>::GetRHI(),
+				TStaticSamplerState<SF_Point>::GetRHI(),
 				InputHistory.Depth[0]->GetRenderTargetItem().TargetableTexture);
 		}
 		else
 		{
 			// Need to bind a white dummy
 			SetTextureParameter(Context.RHICmdList, ShaderRHI, HistoryTexture, HistoryTextureSampler,
-				TStaticSamplerState<SF_Bilinear>::GetRHI(),
+				TStaticSamplerState<SF_Point>::GetRHI(),
 				GSystemTextures.WhiteDummy->GetRenderTargetItem().ShaderResourceTexture);
 
 		}
