Fixed testing wrong define: COMPILER_SUPPORT_DEPTHSTENCIL_EARLYTEST_LATEWRITE -> COMPILER_SUPPORTS_DEPTHSTENCIL_EARLYTEST_LATEWRITE

diff --git a/Engine/Shaders/Private/BasePassPixelShader.usf b/Engine/Shaders/Private/BasePassPixelShader.usf
--- a/Engine/Shaders/Private/BasePassPixelShader.usf
+++ b/Engine/Shaders/Private/BasePassPixelShader.usf
@@ -1406,7 +1406,7 @@ ShadingModel);
 
 // If virtual texture is enabled then use early depth test so that UAV feedback buffer writes respect the depth test
 #if NUM_VIRTUALTEXTURE_SAMPLES || LIGHTMAP_VT_ENABLED
-	#if COMPILER_SUPPORT_DEPTHSTENCIL_EARLYTEST_LATEWRITE
+	#if COMPILER_SUPPORTS_DEPTHSTENCIL_EARLYTEST_LATEWRITE
 		// If we support early depth test with late write behaviour then use it since we may be using discard, or modifying depth
 		#define PIXELSHADER_EARLYDEPTHSTENCIL DEPTHSTENCIL_EARLYTEST_LATEWRITE
 	#elif !OUTPUT_PIXEL_DEPTH_OFFSET
