Added missing sampler functions with TextureExternal for material generated shaders (Integrated from //UE4/Dev-Rendering).

#rnx

diff --git a/Engine/Shaders/Private/Common.ush b/Engine/Shaders/Private/Common.ush
--- a/Engine/Shaders/Private/Common.ush
+++ b/Engine/Shaders/Private/Common.ush
@@ -223,6 +223,10 @@ MaterialFloat4 TextureExternalSample(TextureExternal Tex, SamplerState Sampler,
 	return Tex.Sample(Sampler, UV);
 #endif
 }
+MaterialFloat4 TextureExternalSampleGrad(TextureExternal Tex, SamplerState Sampler, float2 UV, MaterialFloat2 DDX, MaterialFloat2 DDY)
+{
+	return Tex.SampleGrad(Sampler, UV, DDX, DDY);
+}
 MaterialFloat4 TextureExternalSampleLevel(TextureExternal Tex, SamplerState Sampler, float2 UV, MaterialFloat Mip)
 {
 	return Tex.SampleLevel(Sampler, UV, Mip);
@@ -259,6 +263,14 @@ MaterialFloat4 TextureCubeSample_Decal(TextureCube Tex, SamplerState Sampler, fl
 	return TextureCubeSample(Tex, Sampler, UV);
 #endif
 }
+MaterialFloat4 TextureExternalSample_Decal(TextureExternal Tex, SamplerState Sampler, float2 UV)
+{
+#if METAL_PROFILE || COMPILER_GLSL_ES3_1
+	return TextureExternalSampleGrad(Tex, Sampler, UV, 0.f, 0.f);
+#else
+	return TextureExternalSample(Tex, Sampler, UV);
+#endif
+}
 
 MaterialFloat4 Texture2DArraySampleLevel(Texture2DArray Tex, SamplerState Sampler, float3 UV, MaterialFloat Mip)
 {
