Continue with the workaround to prevent scene rendering from going black on xb1.

diff --git a/Engine/Shaders/Private/PostProcessEyeAdaptation.usf b/Engine/Shaders/Private/PostProcessEyeAdaptation.usf
--- a/Engine/Shaders/Private/PostProcessEyeAdaptation.usf
+++ b/Engine/Shaders/Private/PostProcessEyeAdaptation.usf
@@ -263,7 +263,11 @@ void BasicEyeAdaptationPS(in float4 UVAndScreenPos : TEXCOORD0, out float4 OutCo
 	OutColor.y = EyeAdaptation_ExposureCompensation * TargetExposureScale;
 	OutColor.z = AverageSceneLuminance;
 	OutColor.w = 0;
-
+	
+#if XBOXONE_PROFILE
+	OutColor = !all(IsFinite(OutColor)) ? float4(1, 1, 1, 0) : OutColor;
+#endif
+	
 #if COMPUTESHADER
 	if(GIndex==0)
 	{
