Lowering minimum value in SafeNormalize to prevent velocity aligned particle sprites from being distorted at low velocities.

diff --git a/Engine/Shaders/Private/ParticleSpriteVertexFactory.ush b/Engine/Shaders/Private/ParticleSpriteVertexFactory.ush
--- a/Engine/Shaders/Private/ParticleSpriteVertexFactory.ush
+++ b/Engine/Shaders/Private/ParticleSpriteVertexFactory.ush
@@ -302,7 +302,7 @@ FMaterialVertexParameters GetMaterialVertexParameters(FVertexFactoryInput Input,
 
 float3 SafeNormalize(float3 V)
 {
-	return V * rsqrt(max(dot(V,V),0.01));
+	return V * rsqrt(max(dot(V,V),0.00000001));
 }
 
 void GetTangents(FVertexFactoryInput Input, float3 WorldPosition, float3 OldWorldPosition, out float3 OutRight, out float3 OutUp)
