Fix the mindistance clamping when injecting height field into global distance field. Previously it would only work for clipmaps that had small enough voxels.

diff --git a/Engine/Shaders/Private/GlobalDistanceField.usf b/Engine/Shaders/Private/GlobalDistanceField.usf
--- a/Engine/Shaders/Private/GlobalDistanceField.usf
+++ b/Engine/Shaders/Private/GlobalDistanceField.usf
@@ -369,15 +369,16 @@ void CompositeHeightfieldsIntoGlobalDistanceFieldCS(
 
 					// Project the vertical height vector onto the normal of the heightfield directly below the point we are computing the distance field for, use the perpendicular distance
 					float DistanceToHeightfieldPlane = dot(WorldHeightfieldNormal, WorldPosition - WorldHeightfieldShadingPosition);
-					float MinDistance = clamp(DistanceToHeightfieldPlane, -AOGlobalMaxSphereQueryRadius, AOGlobalMaxSphereQueryRadius);
 
 					// Limit negative region of a heighfield to a user defined thickness  
 					const float MinInteriorDistance = -HeightfieldThickness;
-					if (MinDistance < MinInteriorDistance)
+					if (DistanceToHeightfieldPlane < MinInteriorDistance)
 					{
-						MinDistance = min(MinInteriorDistance - MinDistance, AOGlobalMaxSphereQueryRadius);
+						DistanceToHeightfieldPlane = MinInteriorDistance - DistanceToHeightfieldPlane;
 					}
 
+					float MinDistance = clamp(DistanceToHeightfieldPlane, -AOGlobalMaxSphereQueryRadius, AOGlobalMaxSphereQueryRadius);
+
 					float3 GlobalUV = ComputeGlobalUV(WorldPosition, ClipmapIndex);
 					uint3 WriteCoordinate = (uint3)(frac(GlobalUV) * GlobalVolumeDimension);
 					RWGlobalDistanceFieldTexture[WriteCoordinate] = min(MinDistance, RWGlobalDistanceFieldTexture[WriteCoordinate]);
