Copying //UE4/Dev-Physics to Main (//UE4/Main)

diff --git a/Engine/Shaders/Private/HairStrands/HairStrandsClusterCulling.usf b/Engine/Shaders/Private/HairStrands/HairStrandsClusterCulling.usf
new file mode 100644
--- /dev/null
+++ b/Engine/Shaders/Private/HairStrands/HairStrandsClusterCulling.usf
@@ -0,0 +1,761 @@
+// Copyright Epic Games, Inc. All Rights Reserved.
+
+#include "HairStrandsDebugCommon.ush"
+#include "HairStrandsAABBCommon.ush"
+#define COMPACTCS_NUMTHREAD 64
+
+// According to https://docs.microsoft.com/en-us/windows/win32/direct3d11/overviews-direct3d-11-devices-downlevel-compute-shaders#thread-group-shared-memory-tgsm
+// A thread can access to 32 byte per thread when number of thread in a group is 512. Thisis largeley enough for this use case. 
+#define PREFIXSUM_PARALLEL  512
+#define PREFIXSUM_LASTTHREAD (PREFIXSUM_PARALLEL-1)
+
+#define D_USE_CONTINUOUSLOD  1
+#define D_LOD_DEBUG_SCALE 1000.0f
+
+// Cluster data in ClusterAABB buffers is packed as {uint3 AABBMin, uint3 AABBMax}
+
+uint ClusterCount;
+
+#ifdef SHADER_CLEARCLUSTERAABB
+
+RWBuffer<int> OutClusterAABBBuffer; // Cluster data packed as {uint3 AABBMin, uint3 AABBMax}
+RWBuffer<int> OutGroupAABBBuffer;	// Group data packed as {uint3 AABBMin, uint3 AABBMax}
+
+[numthreads(64, 1, 1)]
+void MainClearClusterAABBCS(uint2 DispatchThreadId : SV_DispatchThreadID)
+{
+	if (DispatchThreadId.x < ClusterCount*6)
+	{
+		const int ClusterId = DispatchThreadId.x / 6;
+		const bool bIsMin = (DispatchThreadId.x - ClusterId * 6) < 3 ? 1 : 0;
+		const int Value = bIsMin ? 2147483647 : -2147483648;
+		OutClusterAABBBuffer[DispatchThreadId.x] = Value;
+	}
+
+	if (DispatchThreadId.x == 0)
+	{
+		OutGroupAABBBuffer[0] =  2147483647;
+		OutGroupAABBBuffer[1] =  2147483647;
+		OutGroupAABBBuffer[2] =  2147483647;
+		OutGroupAABBBuffer[3] = -2147483647;
+		OutGroupAABBBuffer[4] = -2147483647;
+		OutGroupAABBBuffer[5] = -2147483647;
+	}
+}
+
+#endif
+
+
+#ifdef SHADER_DRAWDEBUGAABB
+
+#include "../Common.ush"
+#include "../ShaderDrawDebug.ush"
+#include "../ShaderPrintCommon.ush"
+
+uint HairGroupId;
+uint TriangleCount;
+int ClusterDebugMode;
+
+Buffer<int> GroupAABBBuffer;
+Buffer<int> ClusterAABBBuffer;
+Buffer<int> ClusterDebugAABBBuffer;
+
+Buffer<uint> CulledDrawIndirectParameters;
+Buffer<uint> CulledDispatchIndirectParametersClusterCountBuffer;
+
+[numthreads(64, 1, 1)]
+void MainDrawDebugAABBCS(uint3 DispatchThreadId : SV_DispatchThreadID)
+{
+	if (DispatchThreadId.x == 0)
+	{
+		float2 StartPos = float2(50.0f, 150.0f + float(HairGroupId) * 70.0f) * View.ViewSizeAndInvSize.zw;
+		float2 Pos = StartPos;
+
+		const uint VisibleClusterCount = CulledDispatchIndirectParametersClusterCountBuffer[0];
+		const uint VisibleTriangleCount = CulledDrawIndirectParameters[0] / (3*6); // TODO investigate 3*6
+
+		Pos = ShaderPrintSymbol(Pos, _H_);
+		Pos = ShaderPrintSymbol(Pos, _A_);
+		Pos = ShaderPrintSymbol(Pos, _I_);
+		Pos = ShaderPrintSymbol(Pos, _R_);
+		Pos = ShaderPrintSymbol(Pos, _G_);
+		Pos = ShaderPrintSymbol(Pos, _R_);
+		Pos = ShaderPrintSymbol(Pos, _O_);
+		Pos = ShaderPrintSymbol(Pos, _U_);
+		Pos = ShaderPrintSymbol(Pos, _P_);
+		Pos = ShaderPrintSymbol(Pos, _SPC_);
+		Pos = ShaderPrintSymbol(Pos, HairGroupId);
+		Pos = ShaderPrintSymbol(Pos, _SPC_);
+		Pos = ShaderPrintSymbol(Pos, _SPC_);
+		Pos = ShaderPrintSymbol(Pos, _SPC_);
+		Pos = ShaderPrintSymbol(Pos, _SPC_);
+		Pos = ShaderPrintSymbol(Pos, _SPC_);
+		Pos = ShaderPrintSymbol(Pos, _SPC_);
+		Pos = ShaderPrintSymbol(Pos, _SPC_);
+		Pos = ShaderPrintSymbol(Pos, _C_);
+		Pos = ShaderPrintSymbol(Pos, _O_);
+		Pos = ShaderPrintSymbol(Pos, _U_);
+		Pos = ShaderPrintSymbol(Pos, _N_);
+		Pos = ShaderPrintSymbol(Pos, _T_);
+		Pos = ShaderPrintSymbol(Pos, _SPC_);
+		Pos = ShaderPrintSymbol(Pos, _SPC_);
+		Pos = ShaderPrintSymbol(Pos, _SPC_);
+		Pos = ShaderPrintSymbol(Pos, _SPC_);
+		Pos = ShaderPrintSymbol(Pos, _SPC_);
+		Pos = ShaderPrintSymbol(Pos, _SPC_);
+		Pos = ShaderPrintSymbol(Pos, _SPC_);
+		Pos = ShaderPrintSymbol(Pos, _SPC_);
+		Pos = ShaderPrintSymbol(Pos, _P_);
+		Pos = ShaderPrintSymbol(Pos, _E_);
+		Pos = ShaderPrintSymbol(Pos, _R_);
+		Pos = ShaderPrintSymbol(Pos, _C_);
+		Pos = ShaderPrintSymbol(Pos, _E_);
+		Pos = ShaderPrintSymbol(Pos, _N_);
+		Pos = ShaderPrintSymbol(Pos, _T_);
+		Pos = ShaderPrintNewline(Pos); Pos.x = StartPos.x;
+
+		Pos = ShaderPrintSymbol(Pos, _V_);
+		Pos = ShaderPrintSymbol(Pos, _I_);
+		Pos = ShaderPrintSymbol(Pos, _S_);
+		Pos = ShaderPrintSymbol(Pos, _I_);
+		Pos = ShaderPrintSymbol(Pos, _B_);
+		Pos = ShaderPrintSymbol(Pos, _L_);
+		Pos = ShaderPrintSymbol(Pos, _E_);
+		Pos = ShaderPrintSymbol(Pos, _SPC_);
+		Pos = ShaderPrintSymbol(Pos, _C_);
+		Pos = ShaderPrintSymbol(Pos, _L_);
+		Pos = ShaderPrintSymbol(Pos, _U_);
+		Pos = ShaderPrintSymbol(Pos, _S_);
+		Pos = ShaderPrintSymbol(Pos, _T_);
+		Pos = ShaderPrintSymbol(Pos, _E_);
+		Pos = ShaderPrintSymbol(Pos, _R_);
+		Pos = ShaderPrintSymbol(Pos, _S_);
+		Pos = ShaderPrintSymbol(Pos, _SPC_);
+		Pos = ShaderPrintSymbol(Pos, _SPC_);
+		Pos = ShaderPrint(Pos, VisibleClusterCount);
+		Pos = ShaderPrintSymbol(Pos, _SPC_);
+		Pos = ShaderPrint(Pos, 100.0f * float(VisibleClusterCount) / float(ClusterCount));
+		Pos = ShaderPrintNewline(Pos); Pos.x = StartPos.x;
+
+		Pos = ShaderPrintSymbol(Pos, _V_);
+		Pos = ShaderPrintSymbol(Pos, _I_);
+		Pos = ShaderPrintSymbol(Pos, _S_);
+		Pos = ShaderPrintSymbol(Pos, _I_);
+		Pos = ShaderPrintSymbol(Pos, _B_);
+		Pos = ShaderPrintSymbol(Pos, _L_);
+		Pos = ShaderPrintSymbol(Pos, _E_);
+		Pos = ShaderPrintSymbol(Pos, _SPC_);
+		Pos = ShaderPrintSymbol(Pos, _T_);
+		Pos = ShaderPrintSymbol(Pos, _R_);
+		Pos = ShaderPrintSymbol(Pos, _I_);
+		Pos = ShaderPrintSymbol(Pos, _A_);
+		Pos = ShaderPrintSymbol(Pos, _N_);
+		Pos = ShaderPrintSymbol(Pos, _G_);
+		Pos = ShaderPrintSymbol(Pos, _L_);
+		Pos = ShaderPrintSymbol(Pos, _E_);
+		Pos = ShaderPrintSymbol(Pos, _S_);
+		Pos = ShaderPrintSymbol(Pos, _SPC_);
+		Pos = ShaderPrint(Pos, VisibleTriangleCount);
+		Pos = ShaderPrintSymbol(Pos, _SPC_);
+		Pos = ShaderPrint(Pos, 100.0f * float(VisibleTriangleCount) / float(TriangleCount));
+		Pos = ShaderPrintNewline(Pos); Pos.x = StartPos.x;
+	}
+
+	if (DispatchThreadId.x < ClusterCount)
+	{
+		if (DispatchThreadId.x == 0)
+		{
+			const float4 GroupColor = float4(1.0f, 0.5f, 0.0f, 1.0f);
+			const FHairAABB Bound = ReadHairAABB(0, GroupAABBBuffer);
+			AddAABB(Bound.Min, Bound.Max, GroupColor);
+		}
+
+		const uint ClusterId = DispatchThreadId;
+		const uint ClusterId6 = DispatchThreadId.x * 6;
+		float3 Min = float3(int3(ClusterAABBBuffer[ClusterId6    ], ClusterAABBBuffer[ClusterId6 + 1], ClusterAABBBuffer[ClusterId6 + 2]));
+		float3 Max = float3(int3(ClusterAABBBuffer[ClusterId6 + 3], ClusterAABBBuffer[ClusterId6 + 4], ClusterAABBBuffer[ClusterId6 + 5]));
+
+#if PERMUTATION_DEBUGAABBBUFFER
+		const int State = ClusterDebugAABBBuffer[ClusterId];
+#else
+		const int State = -1.0f;
+#endif
+
+		float4 Color = 0.0f;
+		if (ClusterDebugMode == 1)
+		{
+			Color = float4(0.5f, 1.0f, 0.0f, 0.1f);
+		}
+#if PERMUTATION_DEBUGAABBBUFFER
+		else if (ClusterDebugMode == 2)
+		{
+			Color = float4(1.0f, 0.1f, 0.1f, 1.0f);
+			if (State >= 0.0f)
+			{
+				return;
+			}
+		}
+		else
+		{
+			if (State < 0.0f)
+			{
+				return;
+			}
+
+			float SelectedLod = float(State) / D_LOD_DEBUG_SCALE;
+			Color = float4(GetColorCode(SelectedLod), 0.05f);
+		}
+#endif
+
+		AddAABB(Min, Max, Color);
+	}
+}
+
+#endif
+
+
+
+#ifdef SHADER_CLUSTERCULLINGINDCLEAR  
+
+RWBuffer<uint> DispatchIndirectParametersClusterCount;
+RWBuffer<uint> DrawIndirectParameters;
+uint VertexCountPerInstance;
+
+[numthreads(1, 1, 1)]
+void MainClusterCullingIndClearCS(uint3 DispatchThreadId : SV_DispatchThreadID)
+{
+#if PERMUTATION_SETINDIRECTDRAW
+	DrawIndirectParameters[0] = VertexCountPerInstance;
+	DrawIndirectParameters[1] = 1;	// 1 instance
+	DrawIndirectParameters[2] = 0;
+	DrawIndirectParameters[3] = 0;
+#else
+	if (DispatchThreadId.x==0)
+	{
+		DispatchIndirectParametersClusterCount[0] = 0;
+		DispatchIndirectParametersClusterCount[1] = 1;
+		DispatchIndirectParametersClusterCount[2] = 1;
+
+		DrawIndirectParameters[0] = 0;	// 0 vertex (to be accumulated)
+		DrawIndirectParameters[1] = 1;	// 1 instance
+		DrawIndirectParameters[2] = 0;
+		DrawIndirectParameters[3] = 0;
+	}
+#endif
+}
+
+#endif
+
+
+
+#ifdef SHADER_CLUSTERCULLING
+
+float4x4 WorldToClipMatrix;
+float3 CameraWorldPos;
+uint NumConvexHullPlanes;
+int LodMode;
+int ShadowViewMode;
+float4 ViewFrustumConvexHull[6];
+
+float2 ViewRectPixel;
+float LodBias;
+float LodAverageVertexPerPixel;
+
+Buffer<int> ClusterAABBBuffer;
+Buffer<uint4> ClusterInfoBuffer;
+Buffer<float> ClusterIndexRadiusScaleInfoBuffer;
+
+RWBuffer<uint> GlobalIndexStartBuffer;
+RWBuffer<uint> GlobalIndexCountBuffer;
+RWBuffer<float> GlobalRadiusScaleBuffer;
+RWBuffer<uint> DispatchIndirectParametersClusterCount;
+RWBuffer<uint> DispatchIndirectParametersVertexCount;
+RWBuffer<uint> DrawIndirectParameters;
+
+#if PERMUTATION_DEBUGAABBBUFFER
+RWBuffer<int> ClusterDebugAABBBuffer;
+#endif
+
+float3			HZBUvFactor;
+float4			HZBSize;
+Texture2D		HZBTexture;
+SamplerState	HZBSampler;
+
+bool ViewFrustumIntersectSphere(float3 SphereOrigin, float SphereRadius)
+{
+	for (uint PlaneIndex = 0; PlaneIndex < NumConvexHullPlanes; PlaneIndex++)
+	{
+		float4 PlaneData = ViewFrustumConvexHull[PlaneIndex];
+		float PlaneDistance = dot(PlaneData.xyz, SphereOrigin) - PlaneData.w;
+
+		if (PlaneDistance > SphereRadius)
+		{
+			return false;
+		}
+	}
+	return true;
+}
+
+// From HZBOcclusion.usf. TODO: make common
+// Rect contains screen coordinates as {minx, miny, maxx, maxy}. Those coordinates are not clamped to screen edges so they mnight be outside the [0,1]^2.
+bool ViewHierarchicalZBufferCulling(float3 WorldBoundsMin, float3 WorldBoundsMax, inout float4 NotClampedRect)
+{
+	float3 WorldBounds[2] = { WorldBoundsMin, WorldBoundsMax };
+
+	// Screen rect from bounds
+	float3 RectMin = float3(1, 1, 1);
+	float3 RectMax = float3(-1, -1, -1);
+	/*UNROLL*/ for (int i = 0; i < 8; i++)
+	{
+		float3 PointSrc;
+		PointSrc.x = WorldBounds[(i >> 0) & 1].x;
+		PointSrc.y = WorldBounds[(i >> 1) & 1].y;
+		PointSrc.z = WorldBounds[(i >> 2) & 1].z;
+
+		float4 PointClip = mul(float4(PointSrc, 1), WorldToClipMatrix);
+		float3 PointScreen = PointClip.xyz / PointClip.w;
+
+		RectMin = min(RectMin, PointScreen);
+		RectMax = max(RectMax, PointScreen);
+	}
+
+	// FIXME assumes DX
+	NotClampedRect = (float4(RectMin.xy, RectMax.xy) * float2(0.5, -0.5).xyxy + 0.5).xwzy;
+	float4 Rect = saturate(NotClampedRect);
+	float4 RectPixels = Rect * HZBSize.xyxy;
+	float2 RectSize = (RectPixels.zw - RectPixels.xy) * 0.5;	// 0.5 for 4x4
+	float Level = max(ceil(log2(max(RectSize.x, RectSize.y))), HZBUvFactor.z);
+
+#if PERMUTATION_HZBCULLING
+	// Check if we can drop one level lower
+	float LevelLower = max(Level - 1, 0);
+	float4 LowerRect = RectPixels * exp2(-LevelLower);
+	float2 LowerRectSize = ceil(LowerRect.zw) - floor(LowerRect.xy);
+	if (all(LowerRectSize <= 4))
+	{
+		Level = LevelLower;
+	}
+
+	// 4x4 samples
+	float2 Scale = HZBUvFactor.xy * (Rect.zw - Rect.xy) / 3;
+	float2 Bias = HZBUvFactor.xy * Rect.xy;
+
+	float4 MinDepth = 1;
+	/*UNROLL*/ for (int i = 0; i < 4; i++)
+	{
+		// TODO could vectorize this
+		float4 Depth;
+		Depth.x = HZBTexture.SampleLevel(HZBSampler, float2(i, 0) * Scale + Bias, Level).r;
+		Depth.y = HZBTexture.SampleLevel(HZBSampler, float2(i, 1) * Scale + Bias, Level).r;
+		Depth.z = HZBTexture.SampleLevel(HZBSampler, float2(i, 2) * Scale + Bias, Level).r;
+		Depth.w = HZBTexture.SampleLevel(HZBSampler, float2(i, 3) * Scale + Bias, Level).r;
+		MinDepth = min(MinDepth, Depth);
+	}
+	MinDepth.x = min(min(MinDepth.x, MinDepth.y), min(MinDepth.z, MinDepth.w));
+
+	// Inverted Z buffer
+	return RectMax.z >= MinDepth.x ? true : false;
+#else
+	return true;
+#endif
+}
+
+void ApplyDebugMode(inout uint GlobalIndexCount, inout int DebugAABBLod, in uint4 ClusterInfo)
+{
+#if PERMUTATION_LODDEBUG
+	switch (LodMode)
+	{
+	case 0:
+		GlobalIndexCount = ClusterInfo.y; // Highest lod
+		DebugAABBLod = 0;
+		break;
+	case 1:
+		GlobalIndexCount = ClusterInfo.z; // Lowest lod
+		DebugAABBLod = 1;
+		break;
+	}
+#endif
+}
+
+[numthreads(64, 1, 1)]
+void MainClusterCullingCS(uint3 DispatchThreadId : SV_DispatchThreadID)
+{
+	const uint ClusterId = DispatchThreadId.x;
+	if (ClusterId < ClusterCount)
+	{
+		uint ClusterId  = DispatchThreadId.x;
+		uint ClusterId6 = ClusterId * 6;
+		float3 Min = float3(int3(ClusterAABBBuffer[ClusterId6 + 0], ClusterAABBBuffer[ClusterId6 + 1], ClusterAABBBuffer[ClusterId6 + 2]));
+		float3 Max = float3(int3(ClusterAABBBuffer[ClusterId6 + 3], ClusterAABBBuffer[ClusterId6 + 4], ClusterAABBBuffer[ClusterId6 + 5]));
+
+		float3 SphereOrigin = (Min + Max)*0.5f;
+		float  SphereRadius = length(Max - SphereOrigin);
+
+		const float3 CamToSphere = SphereOrigin - CameraWorldPos.xyz;
+		const float CamToSphereDistance = length(CamToSphere);
+		bool bVisible = ViewFrustumIntersectSphere(CamToSphere, SphereRadius);
+		int DebugAABBLod = -1;
+		if (bVisible || ShadowViewMode) // ShadowViewMode: pass throught as we still want all cluster to be able to cast shadow
+		{
+			float4 NotClampedRect = float4(0.0f, 1.0f, 0.0f, 1.0f);
+			bVisible = bVisible && ViewHierarchicalZBufferCulling(Min, Max, NotClampedRect);
+			if (bVisible || ShadowViewMode)
+			{
+				int4 ClusterInfo = ClusterInfoBuffer[ClusterId]; // Start and count of index in the hair strand source buffer
+				const uint GlobalIndexStart = ClusterInfo.x;
+				const uint VertexIdCountLod0 = ClusterInfo.y;
+				const uint VertexIdCountLod1 = ClusterInfo.z;
+				float RadiusScale = 1.0f;
+
+				float2 ClusterRectPixel = float2(ViewRectPixel.x*(NotClampedRect.z - NotClampedRect.x), ViewRectPixel.y*(NotClampedRect.w - NotClampedRect.y));
+				const float RectPixelSize = max(1.0f, max(ClusterRectPixel.x, ClusterRectPixel.y));	// Max seems to work the best
+
+				const float Lod0VertexPerPixel = VertexIdCountLod0 / RectPixelSize;
+				const float Lod1VertexPerPixel = VertexIdCountLod1 / RectPixelSize;
+				uint TargetIndexCount = Lod0VertexPerPixel;
+				if (Lod0VertexPerPixel > Lod1VertexPerPixel) // Safety check in case both are equal
+				{
+					float Lod0ToLod1VertexCount = float(VertexIdCountLod1) - float(VertexIdCountLod0);
+					TargetIndexCount = 
+						float(VertexIdCountLod0) +
+						lerp(Lod0ToLod1VertexCount,																				// Scale vertex count from Lod1
+							0.0f,																								// to lod 0
+							saturate((LodAverageVertexPerPixel - Lod1VertexPerPixel) / (Lod0VertexPerPixel - Lod1VertexPerPixel))// According to the LodAverageVertexPerPixel budget
+						);
+					TargetIndexCount = max(VertexIdCountLod1, TargetIndexCount + LodBias * Lod0ToLod1VertexCount);				// and biased as requested
+				}
+
+				// Always have pair over vertices for quads
+				const uint TargetIndexCountRoundUp2 = 2 * ((TargetIndexCount + 1) / 2);
+
+				// When using ShadowViewMode and the cluster is not visible, select lowest lod. Otherwise use the same lod as the view.
+				uint GlobalIndexCount = ShadowViewMode && !bVisible ? VertexIdCountLod1 : clamp(TargetIndexCountRoundUp2, VertexIdCountLod1, VertexIdCountLod0);
+
+				ApplyDebugMode(GlobalIndexCount, DebugAABBLod, ClusterInfo);
+
+				// Cuompute radius scaling using linear equation. See details in FHairStrandsClusterCullingResource CTOR.
+				const float A = min(1.0, ClusterIndexRadiusScaleInfoBuffer[ClusterId]);
+				RadiusScale = max(1.0f, 1.0f + A * (float(GlobalIndexCount) - float(VertexIdCountLod0)));
+
+				DebugAABBLod = int(D_LOD_DEBUG_SCALE * saturate(float(GlobalIndexCount - VertexIdCountLod1) / float(VertexIdCountLod0 - VertexIdCountLod1)));
+
+
+				if (GlobalIndexCount > 0)
+				{
+					uint ClusterId = 0;
+					InterlockedAdd(DispatchIndirectParametersClusterCount[0], 1, ClusterId);
+
+					uint PreviousVertexCount = 0;
+					InterlockedAdd(DrawIndirectParameters[0], GlobalIndexCount * 6, PreviousVertexCount);
+
+					GlobalIndexStartBuffer[ClusterId] = GlobalIndexStart;
+					GlobalIndexCountBuffer[ClusterId] = GlobalIndexCount;
+					// Scale radius base on index count scale as of today, assuming uniform hair un a cluster.
+					// ALso the lowest lod2 uses 1/4 of lod0 strands so clamp radius to [1,4].
+					GlobalRadiusScaleBuffer[ClusterId] = RadiusScale;
+				}
+			}
+		}
+#if PERMUTATION_DEBUGAABBBUFFER
+		ClusterDebugAABBBuffer[ClusterId] = DebugAABBLod;
+#endif
+	}
+}
+
+#endif
+
+
+
+#ifdef SHADER_PREPAREINDIRECTDRAW
+
+Buffer<uint> DispatchIndirectParametersClusterCount;
+
+RWBuffer<uint> DispatchIndirectParametersClusterCount2D;
+RWBuffer<uint> DispatchIndirectParametersClusterCountDiv512;
+RWBuffer<uint> DispatchIndirectParametersClusterCountDiv512Div512;
+
+uint DivideAndRoundUp(uint Dividend, uint Divisor)
+{
+	return (Dividend + Divisor - 1) / Divisor;
+}
+uint3 ComputeDispatchGroupCount2D(uint GroupCount)
+{
+	const uint DispatchCountX = int(floor(sqrt(float(GroupCount))));
+	const uint DispatchCountY = DispatchCountX + DivideAndRoundUp(GroupCount - DispatchCountX * DispatchCountX, DispatchCountX);
+
+	return uint3(DispatchCountX, DispatchCountY, 1);
+}
+
+[numthreads(1, 1, 1)]
+void MainClusterCullingPrepareIndirectDrawsCS(uint3 DispatchThreadId : SV_DispatchThreadID, uint3 GroupId : SV_GroupID, uint3 GroupThreadId : SV_GroupThreadID)
+{
+	if (GroupId.x == 0 && GroupThreadId.x == 0)
+	{
+		uint ClusterCount = DispatchIndirectParametersClusterCount[0];
+		uint ClusterCountDivUp512 = (ClusterCount + PREFIXSUM_PARALLEL - 1) / PREFIXSUM_PARALLEL;
+		uint ClusterCountDivUp512DivUp512 = (ClusterCountDivUp512 + PREFIXSUM_PARALLEL - 1) / PREFIXSUM_PARALLEL;
+
+		DispatchIndirectParametersClusterCountDiv512[0] = ClusterCountDivUp512;
+		DispatchIndirectParametersClusterCountDiv512[1] = 1;
+		DispatchIndirectParametersClusterCountDiv512[2] = 1;
+
+		DispatchIndirectParametersClusterCountDiv512Div512[0] = ClusterCountDivUp512DivUp512;
+		DispatchIndirectParametersClusterCountDiv512Div512[1] = 1;
+		DispatchIndirectParametersClusterCountDiv512Div512[2] = 1;
+
+		uint3 DispatchGroupCount2D = ComputeDispatchGroupCount2D(ClusterCount);
+		DispatchIndirectParametersClusterCount2D[0] = DispatchGroupCount2D.x;
+		DispatchIndirectParametersClusterCount2D[1] = DispatchGroupCount2D.y;
+		DispatchIndirectParametersClusterCount2D[2] = 1.0;
+	}
+}
+
+#endif
+
+
+
+/// PER BLOCK OPTIMISATION
+#ifdef SHADER_PERBLOCKPREFIXSUM
+
+Buffer<uint> DispatchIndirectParametersClusterCount;
+Buffer<uint> GlobalIndexCountBuffer;
+
+RWBuffer<uint2> PerBlocklIndexCountPreFixSumBuffer;
+RWBuffer<uint>  PerBlocklTotalIndexCountBuffer;
+
+groupshared uint PrefixSumClusterCount[2][PREFIXSUM_PARALLEL];
+
+[numthreads(PREFIXSUM_PARALLEL, 1, 1)] 
+void MainPerBlockPreFixSumCS(uint3 DispatchThreadId : SV_DispatchThreadID, uint3 GroupId : SV_GroupID, uint3 GroupThreadId : SV_GroupThreadID)
+{
+	uint ItemCount = DispatchIndirectParametersClusterCount[0];
+	uint ItemCountRoundUp = PREFIXSUM_PARALLEL * ((ItemCount + PREFIXSUM_PARALLEL - 1) / PREFIXSUM_PARALLEL);
+	uint WaveStart = PREFIXSUM_PARALLEL * GroupId.x;
+
+	if (WaveStart < ItemCountRoundUp)
+	{
+		// Load current block data
+		uint GlobalIndex = WaveStart + GroupThreadId.x;
+		if (GlobalIndex < ItemCount)
+		{
+			const uint IndexCount = GlobalIndexCountBuffer[GlobalIndex];
+			PrefixSumClusterCount[0][GroupThreadId.x] = IndexCount;
+		}
+		else
+		{
+			PrefixSumClusterCount[0][GroupThreadId.x] = 0;
+		}
+		GroupMemoryBarrierWithGroupSync();
+
+		// Pre fix sum the current block using ping-pong in shared memory
+		uint Src = 0;
+		for (uint i = 1; i <= (PREFIXSUM_PARALLEL / 2); i *= 2)
+		{
+			if (GroupThreadId.x >= i)
+			{
+				PrefixSumClusterCount[1 - Src][GroupThreadId.x] = PrefixSumClusterCount[Src][GroupThreadId.x] + PrefixSumClusterCount[Src][GroupThreadId.x - i];
+			}
+			else
+			{
+				PrefixSumClusterCount[1 - Src][GroupThreadId.x] = PrefixSumClusterCount[Src][GroupThreadId.x];
+			}
+			Src = 1 - Src;
+			GroupMemoryBarrierWithGroupSync();
+		}
+
+		// Write out result, taking into account previous blocks sum
+		uint Result = 0;
+		if (GlobalIndex < ItemCountRoundUp)
+		{
+			Result = PrefixSumClusterCount[Src][GroupThreadId.x];
+
+			if (GroupThreadId.x == 0)
+			{
+				PerBlocklIndexCountPreFixSumBuffer[GlobalIndex] = uint2(0, Result);
+			}
+			else
+			{
+				PerBlocklIndexCountPreFixSumBuffer[GlobalIndex] = uint2(PrefixSumClusterCount[Src][GroupThreadId.x - 1], Result);
+			}
+		}
+
+		if (GroupThreadId.x == PREFIXSUM_LASTTHREAD)
+		{
+			PerBlocklTotalIndexCountBuffer[GroupId.x] = PrefixSumClusterCount[Src][GroupThreadId.x];
+		}
+	}
+}
+
+#endif
+
+
+
+/// PER BLOCK OPTIMISATION
+#if SHADER_CLUSTERCULLINGCOMPACTVERTEXIDLOCALBLOCK
+
+Buffer<uint> DispatchIndirectParametersClusterCount;
+Buffer<uint> DispatchIndirectParametersClusterCount2D;
+Buffer<uint> DispatchIndirectParametersClusterCountDiv512;
+
+Buffer<uint2> PerBlocklIndexCountPreFixSumBuffer;
+Buffer<uint2> PerBlocklTotalIndexCountPreFixSumBuffer;
+
+Buffer<uint> GlobalIndexStartBuffer;
+Buffer<uint> GlobalIndexCountBuffer;
+Buffer<float> GlobalRadiusScaleBuffer;
+Buffer<uint> ClusterVertexIdBuffer;
+
+RWBuffer<uint> CulledCompactedIndexBuffer;
+RWBuffer<float> CulledCompactedRadiusScaleBuffer;
+
+[numthreads(COMPACTCS_NUMTHREAD, 1, 1)]
+void MainClusterCullingCompactVertexIdsLocalBlockCS(uint3 DispatchThreadId : SV_DispatchThreadID, uint3 GroupId : SV_GroupID, uint3 GroupThreadId : SV_GroupThreadID)
+{
+	const uint LoadedClusterCount = DispatchIndirectParametersClusterCount[0];
+	const uint DispatchCountX = DispatchIndirectParametersClusterCount2D[0];
+	const uint ClusterId = GroupId.x + GroupId.y * DispatchCountX;
+	if (ClusterId >= LoadedClusterCount)
+	{
+		return;
+	}
+
+	// For each cluster, a group is splatting the indices inside its target range (evaluated after a prefix sum)
+
+	const uint ClusterSrcStartIndex = GlobalIndexStartBuffer[ClusterId];
+	const uint ClusterSrcIndexCount = GlobalIndexCountBuffer[ClusterId];
+	const float ClusterRadiusScale  = GlobalRadiusScaleBuffer[ClusterId];
+
+	uint2 TopBlockPrefixSum = PerBlocklTotalIndexCountPreFixSumBuffer[ClusterId / PREFIXSUM_PARALLEL];
+	uint2 PerBlockPrefixSum = PerBlocklIndexCountPreFixSumBuffer[ClusterId];
+
+	uint WriteToIndex = TopBlockPrefixSum.x + PerBlockPrefixSum.x;
+	uint LoopCount = (ClusterSrcIndexCount + COMPACTCS_NUMTHREAD - 1) / COMPACTCS_NUMTHREAD;
+
+	for (uint Loop = 0; Loop < LoopCount; ++Loop)
+	{
+		uint GroupWriteOffset = Loop * COMPACTCS_NUMTHREAD;
+		uint ThreadWriteOffset = GroupWriteOffset + GroupThreadId.x;
+
+		if (ThreadWriteOffset < ClusterSrcIndexCount)
+		{
+			CulledCompactedIndexBuffer[WriteToIndex + ThreadWriteOffset] = ClusterVertexIdBuffer[ClusterSrcStartIndex + ThreadWriteOffset];
+
+			const bool bInvalidVertex = ThreadWriteOffset >= (ClusterSrcIndexCount - 1);
+			CulledCompactedRadiusScaleBuffer[WriteToIndex + ThreadWriteOffset] = bInvalidVertex && D_USE_CONTINUOUSLOD ? 0.0f : ClusterRadiusScale;
+		}
+	}
+}
+
+#endif
+
+
+
+#ifdef SHADER_CLUSTERCULLINGPREFIXSUM
+
+Buffer<uint> DispatchIndirectParametersClusterCount;
+Buffer<uint> GlobalIndexCountBuffer;
+
+RWBuffer<uint2> GlobalIndexCountPreFixSumBuffer;
+
+groupshared uint PreviousSum;
+groupshared uint PrefixSumClusterCount[2][PREFIXSUM_PARALLEL];
+
+[numthreads(PREFIXSUM_PARALLEL, 1, 1)]
+void MainClusterCullingPreFixSumCS(uint3 DispatchThreadId : SV_DispatchThreadID, uint3 GroupId : SV_GroupID)
+{
+#if 0
+
+	// Slow reference
+	if (DispatchThreadId.x == 0)
+	{
+		// ugly linear 1 thread....
+		uint Sum = 0;
+		uint ItemCount = DispatchIndirectParametersClusterCount[0];
+		for (uint i = 0; i < ItemCount; ++i)
+		{
+			uint IndexCount = GlobalIndexCountBuffer[i];
+			GlobalIndexCountPreFixSumBuffer[i] = uint2(Sum, Sum + IndexCount);
+			Sum += IndexCount;
+		}
+	}
+
+#else
+
+	if (DispatchThreadId.x == 0)
+	{
+		// Reset previous block sum
+		PreviousSum = 0;
+	}
+	GroupMemoryBarrierWithGroupSync();
+
+	uint ItemCount = DispatchIndirectParametersClusterCount[0];
+	uint ItemCountRoundUp = PREFIXSUM_PARALLEL * ((ItemCount + PREFIXSUM_PARALLEL - 1) / PREFIXSUM_PARALLEL);
+	uint WaveStart = 0;
+
+	while (WaveStart < ItemCountRoundUp)
+	{
+		// Load current block data
+		uint GlobalIndex = WaveStart + DispatchThreadId.x;
+		if (GlobalIndex < ItemCount)
+		{
+			const uint IndexCount = GlobalIndexCountBuffer[GlobalIndex];
+			PrefixSumClusterCount[0][DispatchThreadId.x] = IndexCount;
+		}
+		else
+		{
+			PrefixSumClusterCount[0][DispatchThreadId.x] = 0;
+		}
+		GroupMemoryBarrierWithGroupSync();
+
+		// Pre fix sum the current block using ping-pong in shared memory
+		uint Src = 0;
+		for (uint i = 1; i <= (PREFIXSUM_PARALLEL/2); i *= 2)
+		{
+			if (DispatchThreadId.x >= i)
+			{
+				PrefixSumClusterCount[1 - Src][DispatchThreadId.x] = PrefixSumClusterCount[Src][DispatchThreadId.x] + PrefixSumClusterCount[Src][DispatchThreadId.x - i];
+			}
+			else
+			{
+				PrefixSumClusterCount[1 - Src][DispatchThreadId.x] = PrefixSumClusterCount[Src][DispatchThreadId.x];
+			}
+			Src = 1 - Src;
+			GroupMemoryBarrierWithGroupSync();
+		}
+
+		// Write out result, taking into account previous blocks sum
+		uint Result = 0;
+		if (GlobalIndex < ItemCountRoundUp)
+		{
+			Result = PrefixSumClusterCount[Src][DispatchThreadId.x];
+
+			if (DispatchThreadId.x == 0)
+			{
+				GlobalIndexCountPreFixSumBuffer[GlobalIndex] = uint2(0, Result) + PreviousSum;
+			}
+			else
+			{
+				GlobalIndexCountPreFixSumBuffer[GlobalIndex] = uint2(PrefixSumClusterCount[Src][DispatchThreadId.x - 1], Result) + PreviousSum;
+			}
+		}
+		GroupMemoryBarrierWithGroupSync();
+
+		// Update previous block sum
+		if (DispatchThreadId.x == PREFIXSUM_LASTTHREAD)
+		{
+			PreviousSum += Result;
+		}
+
+		WaveStart += PREFIXSUM_PARALLEL;
+	}
+
+#endif
+
+	// TODO try multiple pass parallelisation, with dispatch indirect
+	// https://developer.nvidia.com/gpugems/GPUGems3/gpugems3_ch39.html
+	// http://citeseerx.ist.psu.edu/viewdoc/download?doi=10.1.1.177.8755&rep=rep1&type=pdf
+}
+
+#endif
+
+
diff --git a/Engine/Shaders/Private/HairStrands/HairStrandsDebug.usf b/Engine/Shaders/Private/HairStrands/HairStrandsDebug.usf
--- a/Engine/Shaders/Private/HairStrands/HairStrandsDebug.usf
+++ b/Engine/Shaders/Private/HairStrands/HairStrandsDebug.usf
@@ -2,6 +2,7 @@
 
 #include "../Common.ush"
 #include "HairStrandsVisibilityCommon.ush"
+#include "HairStrandsDebugCommon.ush"
 
 float2 OutputResolution;
 uint FastResolveMask;
@@ -15,6 +16,9 @@ Texture2D<uint> NodeIndex;
 StructuredBuffer<FNodeData> NodeData;
 SamplerState LinearSampler;
 
+uint TileSize;
+Texture2D<uint2> TileIndexTexture;
+
 #define DEBUG_MODE_SAMPLE_COUNT 0 
 #define DEBUG_MODE_COVERAGE_TYPE 1
 #define DEBUG_MODE_TAA_RESOLVE_TYPE 2
@@ -24,26 +28,7 @@ SamplerState LinearSampler;
 #define DEBUG_MODE_ROUGHNESS 6
 #define DEBUG_MODE_SPECULAR 7
 #define DEBUG_MODE_TANGENT 8
-
-float3 HUEtoRGB(in float H)
-{
-	float R = abs(H * 6 - 3) - 1;
-	float G = 2 - abs(H * 6 - 2);
-	float B = 2 - abs(H * 6 - 4);
-	return saturate(float3(R, G, B));
-}
-
-float3 HSVtoRGB(in float3 HSV)
-{
-	float3 RGB = HUEtoRGB(HSV.x);
-	return ((RGB - 1) * HSV.y + 1) * HSV.z;
-}
-
-float3 GetColorCode(float x)
-{
-	float c = (1 - saturate(x)) * 0.6; // Remap [0,1] to Blue-Red
-	return x > 0 ? HSVtoRGB(float3(c, 1, 1)) : float3(0, 0, 0);
-}
+#define DEBUG_MODE_TILE 9
 
 void MainPS(
 	in FScreenVertexOutput Input,
@@ -53,11 +38,12 @@ void MainPS(
 	const float2 PixelPos = UV * OutputResolution;
 
 	const FCategorizationData CatData = DecodeCategorizationData(CategorizationTexture.Load(uint3(PixelPos,0)));
-	const bool bIsValid = CatData.PixelCoverage > 0;
+	const bool bIsValid = CatData.PixelCoverage > 0 || DebugMode == DEBUG_MODE_TILE;
 	if (!bIsValid)
 		discard;
 
 	float3 Color = 0;
+	float Alpha = 1;
 	if (DebugMode == DEBUG_MODE_SAMPLE_COUNT)
 	{
 		Color = GetColorCode(CatData.SampleCount/float(HAIR_VISIBILITY_MAXSAMPLECOUNT));
@@ -81,6 +67,14 @@ void MainPS(
 		Color = GetColorCode(saturate(CatData.PixelCoverage));
 	}
 
+	if (DebugMode == DEBUG_MODE_TILE)
+	{
+		const uint2 TileCoord = PixelPos / TileSize;
+		const uint TileOffset = TileIndexTexture[TileCoord];
+		Color = TileOffset != INVALID_TILE_OFFSET ? float3(0.25, 0.85, 0.25) : float3(0.85, 0.25, 0.25);
+		Alpha = 0.5f;
+	}
+
 	const bool bIsDebugModeMaterial =
 		DebugMode == DEBUG_MODE_DEPTH ||
 		DebugMode == DEBUG_MODE_BASECOLOR ||
@@ -164,5 +158,5 @@ void MainPS(
 		}
 	}
 
-	OutColor = float4(Color, 1);
+	OutColor = float4(Color, Alpha);
 }
diff --git a/Engine/Shaders/Private/HairStrands/HairStrandsDebugCommon.ush b/Engine/Shaders/Private/HairStrands/HairStrandsDebugCommon.ush
new file mode 100644
--- /dev/null
+++ b/Engine/Shaders/Private/HairStrands/HairStrandsDebugCommon.ush
@@ -0,0 +1,26 @@
+// Copyright Epic Games, Inc. All Rights Reserved.
+
+#pragma once
+
+float3 HUEtoRGB(in float H)
+{
+	float R = abs(H * 6 - 3) - 1;
+	float G = 2 - abs(H * 6 - 2);
+	float B = 2 - abs(H * 6 - 4);
+	return saturate(float3(R, G, B));
+}
+
+float3 HSVtoRGB(in float3 HSV)
+{
+	float3 RGB = HUEtoRGB(HSV.x);
+	return ((RGB - 1) * HSV.y + 1) * HSV.z;
+}
+
+float3 GetColorCode(float x)
+{
+	float c = (1 - saturate(x)) * 0.6; // Remap [0,1] to Blue-Red
+	return x > 0 ? HSVtoRGB(float3(c, 1, 1)) : float3(0, 0, 0);
+}
+
+
+
diff --git a/Engine/Shaders/Private/HairStrands/HairStrandsDebugPrint.usf b/Engine/Shaders/Private/HairStrands/HairStrandsDebugPrint.usf
--- a/Engine/Shaders/Private/HairStrands/HairStrandsDebugPrint.usf
+++ b/Engine/Shaders/Private/HairStrands/HairStrandsDebugPrint.usf
@@ -2,11 +2,15 @@
 
 #include "../Common.ush"
 #include "../ShaderPrintCommon.ush"
+#include "../ShaderDrawDebug.ush"
 #include "HairStrandsVisibilityCommon.ush"
+#include "HairStrandsDebugCommon.ush"
+#include "HairStrandsAABBCommon.ush"
 
 int2 PixelCoord;
 int2 MaxResolution;
 uint FastResolveMask;
+uint HairMacroGroupCount;
 uint2 GroupSize; 
 
 Texture2D<uint4> CategorizationTexture;
@@ -15,32 +19,25 @@ SamplerState LinearSampler;
 Texture2D<uint> HairVisibilityNodeOffsetAndCount;
 StructuredBuffer<FNodeData> HairVisibilityNodeData;
 
+Buffer<int> HairMacroGroupAABBBuffer;
+
 #define MATERIAL_DATA 0
 #define DETAIL_DATA 0
 
-float3 HUEtoRGB(in float H)
-{
-	float R = abs(H * 6 - 3) - 1;
-	float G = 2 - abs(H * 6 - 2);
-	float B = 2 - abs(H * 6 - 4);
-	return saturate(float3(R, G, B));
-}
-
-float3 HSVtoRGB(in float3 HSV)
-{
-	float3 RGB = HUEtoRGB(HSV.x);
-	return ((RGB - 1) * HSV.y + 1) * HSV.z;
-}
-
-float3 GetColorCode(float x)
-{
-	float c = (1 - saturate(x)) * 0.6; // Remap [0,1] to Blue-Red
-	return x > 0 ? HSVtoRGB(float3(c, 1, 1)) : float3(0, 0, 0);
-}
-
 [numthreads(1, 1, 1)]
 void MainCS(uint GroupIndex : SV_GroupIndex, uint3 DispatchThreadId : SV_DispatchThreadID)
 {	
+	// Draw AABB of the instance groups
+	if (DispatchThreadId.x == 0 && DispatchThreadId.y == 0 && DispatchThreadId.z == 0)
+	{
+		for (uint GroupIndex =0;GroupIndex< HairMacroGroupCount; GroupIndex++)
+		{
+			const float4 InstanceBoundColor = float4(1.f, 0.5f, 0.5f, 1.0f);
+			FHairAABB Bound = ReadHairAABB(GroupIndex, HairMacroGroupAABBBuffer);
+			AddAABB(Bound.Min, Bound.Max, InstanceBoundColor);
+		}
+	}
+
 	const float2 PixelPos = PixelCoord;
 	if (PixelCoord.x < 0 || PixelCoord.y < 0 || PixelCoord.x >= MaxResolution.x || PixelCoord.y >= MaxResolution.y)
 		return;
@@ -335,20 +332,24 @@ void MainCS(uint GroupIndex : SV_GroupIndex, uint3 DispatchThreadId : SV_Dispatc
 		Pos = ShaderPrintNewline(Pos);
 		Pos.x = OriginalPos.x;
 
-		// ClusterId
+		// MacroGroupId
 		Pos = ShaderPrintSymbol(Pos, _DOT_);
+		Pos = ShaderPrintSymbol(Pos, _M_);
+		Pos = ShaderPrintSymbol(Pos, _A_);
 		Pos = ShaderPrintSymbol(Pos, _C_);
-		Pos = ShaderPrintSymbol(Pos, _L_);
-		Pos = ShaderPrintSymbol(Pos, _U_);
-		Pos = ShaderPrintSymbol(Pos, _S_);
-		Pos = ShaderPrintSymbol(Pos, _T_);
-		Pos = ShaderPrintSymbol(Pos, _E_);
 		Pos = ShaderPrintSymbol(Pos, _R_);
+		Pos = ShaderPrintSymbol(Pos, _O_);
+		Pos = ShaderPrintSymbol(Pos, _SPC_);
+		Pos = ShaderPrintSymbol(Pos, _G_);
+		Pos = ShaderPrintSymbol(Pos, _R_);
+		Pos = ShaderPrintSymbol(Pos, _O_);
+		Pos = ShaderPrintSymbol(Pos, _U_);
+		Pos = ShaderPrintSymbol(Pos, _P_);
 		Pos = ShaderPrintSymbol(Pos, _SPC_);
 		Pos = ShaderPrintSymbol(Pos, _I_);
 		Pos = ShaderPrintSymbol(Pos, _D_);
 		Pos = ShaderPrintSymbol(Pos, _SPC_);
-		Pos = ShaderPrint(Pos, SampleVisData.ClusterId);
+		Pos = ShaderPrint(Pos, SampleVisData.MacroGroupId);
 
 		Pos = ShaderPrintNewline(Pos);
 		Pos.x = OriginalPos.x;		
diff --git a/Engine/Shaders/Private/HairStrands/HairStrandsDeepShadowDebug.usf b/Engine/Shaders/Private/HairStrands/HairStrandsDeepShadowDebug.usf
--- a/Engine/Shaders/Private/HairStrands/HairStrandsDeepShadowDebug.usf
+++ b/Engine/Shaders/Private/HairStrands/HairStrandsDeepShadowDebug.usf
@@ -92,7 +92,7 @@ void VisualizeDomPS(
 }
 #endif
 
-// Display screen space rect for vizualizing projected viewrect of each hair cluster 
+// Display screen space rect for vizualizing projected viewrect of each hair macro group
 #if PERMUTATION_OUTPUT_TYPE == 1
 void VisualizeDomPS(
 	in FScreenVertexOutput Input,
diff --git a/Engine/Shaders/Private/HairStrands/HairStrandsDeepShadowMask.usf b/Engine/Shaders/Private/HairStrands/HairStrandsDeepShadowMask.usf
--- a/Engine/Shaders/Private/HairStrands/HairStrandsDeepShadowMask.usf
+++ b/Engine/Shaders/Private/HairStrands/HairStrandsDeepShadowMask.usf
@@ -1,11 +1,27 @@
 // Copyright Epic Games, Inc. All Rights Reserved.
 
+#include "HairStrandsVoxelPageCommonStruct.ush"
+
 #include "../Common.ush"
 #include "../CommonViewUniformBuffer.ush"
 #include "../SceneTextureParameters.ush"
 #include "../DeferredShadingCommon.ush"
 #include "HairStrandsVisibilityCommon.ush"
-#include "HairStrandsDeepVoxelCommon.ush"
+#include "HairStrandsVoxelPageCommon.ush"
+
+#define OPAQUE_MASK_TYPE_DEEP_SHADOW 0
+#define OPAQUE_MASK_TYPE_VOXEL 1
+#define OPAQUE_MASK_TYPE_VIRTUAL_VOXEL 2
+
+#define VOXEL_TRAVERSAL_DEBUG 0
+#if PERMUTATION_OPAQUEMASK_TYPE == OPAQUE_MASK_TYPE_VOXEL || PERMUTATION_OPAQUEMASK_TYPE == OPAQUE_MASK_TYPE_VIRTUAL_VOXEL
+  #if PERMUTATION_OPAQUEMASK_TYPE == OPAQUE_MASK_TYPE_VOXEL
+	#define VOXEL_TRAVERSAL_TYPE VOXEL_TRAVERSAL_LEGACY
+  #elif PERMUTATION_OPAQUEMASK_TYPE == OPAQUE_MASK_TYPE_VIRTUAL_VOXEL
+	#define VOXEL_TRAVERSAL_TYPE VOXEL_TRAVERSAL_LINEAR
+  #endif
+  #include "HairStrandsVoxelPageTraversal.ush"
+#endif
 
 float4x4			DeepShadow_WorldToLightTransform;
 uint				DeepShadow_IsWholeSceneLight;
@@ -19,8 +35,11 @@ float				Voxel_DensityScale;
 float3				Voxel_MinAABB;
 uint				Voxel_Resolution;
 float3				Voxel_MaxAABB;
+uint				Voxel_MacroGroupId;
 Texture3D<uint> 	Voxel_DensityTexture;
 
+Texture2D<float4>	RayMarchMaskTexture;
+
 Texture2D<float4>	DeepShadow_FrontDepthTexture;
 Texture2D<uint4>	CategorizationTexture;
 RWTexture2D<float4>	OutputColor;
@@ -28,9 +47,6 @@ RWTexture2D<float4>	OutputColor;
 SamplerState		LinearSampler;	
 SamplerState 		ShadowSampler;
 
-#define OPAQUE_MASK_TYPE_DEEP_SHADOW 0
-#define OPAQUE_MASK_TYPE_VOXEL 1
-
 #define KERNEL_TYPE_2x2 0 
 #define KERNEL_TYPE_4x4 1
 #define KERNEL_TYPE_Gaussian8 2
@@ -325,7 +341,7 @@ void MainPS(
 
 #endif
 
-#if PERMUTATION_OPAQUEMASK_TYPE == OPAQUE_MASK_TYPE_VOXEL
+#if PERMUTATION_OPAQUEMASK_TYPE == OPAQUE_MASK_TYPE_VOXEL || PERMUTATION_OPAQUEMASK_TYPE == OPAQUE_MASK_TYPE_VIRTUAL_VOXEL
 
 void MainPS(
 	FScreenVertexOutput Input,
@@ -343,21 +359,93 @@ void MainPS(
 	float Visibility = 1;
 	if (!bIsHairPixel)
 	{
+		const float VoxelShadowMaxDistance = 100000.0; // 1 km shadow distance
 		const bool bIsDirectionalLight = Voxel_LightPosition.w == 0;
-		const float3 L = bIsDirectionalLight ? -Voxel_LightDirection : normalize(Voxel_LightPosition.xyz - WorldPosition);
+		const float3 LightPosition  = bIsDirectionalLight ? -Voxel_LightDirection*VoxelShadowMaxDistance + WorldPosition : Voxel_LightPosition.xyz;
+		const float3 LightDirection = bIsDirectionalLight ? -Voxel_LightDirection : normalize(Voxel_LightPosition.xyz - WorldPosition);
 
 		// Depth bias
 		// Origin is shifted 2 voxels away towards the light + a constant bias of the size of the voxel
+		const float3 Random = GetHairVoxelJitter(PixelCoord.xy, View.StateFrameIndexMod8);
 		{
-			const float3 Random = GetHairVoxelJitter(PixelCoord.xy, View.StateFrameIndexMod8);
+		#if PERMUTATION_OPAQUEMASK_TYPE == OPAQUE_MASK_TYPE_VOXEL
 			const float3 VoxelSize = (Voxel_MaxAABB - Voxel_MinAABB) / Voxel_Resolution;
-			const float3 DepthBias = max(VoxelSize.x, max(VoxelSize.y, VoxelSize.z)) * (2 * L + Random);
+			const float MaxVoxelSize = max(VoxelSize.x, max(VoxelSize.y, VoxelSize.z));
+		#elif PERMUTATION_OPAQUEMASK_TYPE == OPAQUE_MASK_TYPE_VIRTUAL_VOXEL
+			const float MaxVoxelSize = VirtualVoxel.VoxelWorldSize;
+		#endif
+			const float3 DepthBias = MaxVoxelSize * (2 * LightDirection + Random);
 			WorldPosition += DepthBias;
 		}
 
-		const float DistanceThreshold = 100000;
-		const float HairCount = ComputeHairCountVoxel(WorldPosition, L, Voxel_MinAABB, Voxel_MaxAABB, Voxel_Resolution, Voxel_DensityTexture, Voxel_DensityScale, 1, DistanceThreshold);
-		Visibility = saturate(1 - HairCount);
+		FHairTraversalResult Result = InitHairTraversalResult();
+
+		const float DistanceThreshold = 100000.0f;
+		const float CoverageThreshold = 0.995f; // When Coverage is high, we do not trace shadow on opaque since hair/fur is covering the background.
+
+	#if PERMUTATION_OPAQUEMASK_TYPE == OPAQUE_MASK_TYPE_VOXEL
+		if (RayMarchMaskTexture.Load(uint3(PixelCoord, 0)).r > 0.0f &&  Categorization.PixelCoverage < CoverageThreshold)
+		{
+			Result = ComputeHairCountVoxel(
+				WorldPosition,
+				LightPosition,
+				Voxel_MinAABB,
+				Voxel_MaxAABB,
+				Voxel_Resolution,
+				Voxel_DensityTexture,
+				Voxel_DensityScale,
+				1,
+				DistanceThreshold);
+		}
+		else
+		{
+			discard;
+		}
+	#endif
+
+	#if PERMUTATION_OPAQUEMASK_TYPE == OPAQUE_MASK_TYPE_VIRTUAL_VOXEL
+		FVirtualVoxelCommonDesc CommonDesc;
+		CommonDesc.PageCountResolution = VirtualVoxel.PageCountResolution;
+		CommonDesc.VoxelWorldSize = VirtualVoxel.VoxelWorldSize;
+		CommonDesc.PageTextureResolution = VirtualVoxel.PageTextureResolution;
+		CommonDesc.PageResolution = VirtualVoxel.PageResolution;
+
+		const FPackedVirtualVoxelNodeDesc PackedNode = VirtualVoxel.NodeDescBuffer[Voxel_MacroGroupId];
+		const FVirtualVoxelNodeDesc NodeDesc = UnpackVoxelNode(PackedNode, VirtualVoxel.PageResolution);
+
+		#if VOXEL_TRAVERSAL_DEBUG
+		const bool bDebugEnabled = PixelCoord.x == GetCursorPos().x && PixelCoord.y == GetCursorPos().y;
+		#else
+		const bool bDebugEnabled = false;
+		#endif	
+
+		FHairTraversalSettings TraversalSettings = InitHairTraversalSettings();
+		TraversalSettings.DensityScale = VirtualVoxel.DensityScale;
+		TraversalSettings.CountThreshold = 1;
+		TraversalSettings.DistanceThreshold = DistanceThreshold;
+		TraversalSettings.bDebugEnabled = bDebugEnabled;
+		TraversalSettings.SteppingScale = VirtualVoxel.SteppingScale;
+		TraversalSettings.Random = Random;
+		
+		if (RayMarchMaskTexture.Load(uint3(PixelCoord, 0)).r > 0.0f &&  Categorization.PixelCoverage < CoverageThreshold)
+		{
+			Result = ComputeHairCountVirtualVoxel(
+				WorldPosition,
+				LightPosition,
+				CommonDesc,
+				NodeDesc,
+				VirtualVoxel.PageIndexBuffer,
+				VirtualVoxel.PageTexture,
+				TraversalSettings);
+
+		}
+		else
+		{
+			discard;
+		}
+	#endif
+
+		Visibility = min(saturate(1 - Result.HairCount), Result.Visibility);
 	}
 
 	// Depending of if the light is direction (i.e., WholeSceneLight) or a local light (Point/Spot/Rect), 
diff --git a/Engine/Shaders/Private/HairStrands/HairStrandsDeepShadowPS.usf b/Engine/Shaders/Private/HairStrands/HairStrandsDeepShadowPS.usf
--- a/Engine/Shaders/Private/HairStrands/HairStrandsDeepShadowPS.usf
+++ b/Engine/Shaders/Private/HairStrands/HairStrandsDeepShadowPS.usf
@@ -1,18 +1,21 @@
 // Copyright Epic Games, Inc. All Rights Reserved.
 
+#include "HairStrandsVoxelPageCommonStruct.ush"
+
 #include "../Common.ush"
 
-//#define SceneTexturesStruct DeepShadowPass.SceneTextures
+//#define SceneTexturesStruct VoxelRasterPass.SceneTextures
 
 #include "../SceneTexturesCommon.ush"
 #include "HairStrandsDeepShadowCommon.ush"
 #include "HairStrandsVisibilityCommon.ush"
+#include "HairStrandsVoxelPageCommon.ush"
 #include "/Engine/Generated/Material.ush"
 #include "/Engine/Generated/VertexFactory.ush"
 
 static const float FixPointMaxValue = 1000;	// 100m @hair_todo: make this camera relative, and expose a CVAR fix changing this value
 static const float FixPointScale	= 100;	// 0.1mm precision
-static const float FixPointRange	= 2 * FixPointMaxValue * FixPointScale;
+static const float FixPointRange	= 2 * FixPointMaxValue * FixPointScale;    
 
 #define SUPPORT_OPACITY_MASK 0
 
@@ -44,24 +47,40 @@ void MainDom(
 	#endif
 
 	const uint2 PixelCoord = uint2(floor(SvPosition.xy));
-    const float FrontDepth = DeepShadowPass.FrontDepthTexture.Load(uint3(PixelCoord, 0));
+	const float FrontDepth = DeepRasterPass.FrontDepthTexture.Load(uint3(PixelCoord, 0));
 	const float DistanceToFrontDepth = GetDomDistanceToFrontDepth(FrontDepth, SvPosition.z);
 
-    OutColor = ComputeDOMWeights(DistanceToFrontDepth) * HairCoverage;
+	OutColor = ComputeDOMWeights(DistanceToFrontDepth) * HairCoverage;
 }
 
+#define DEBUG_OUTPUT 0
+
+#if DEBUG_OUTPUT == 1
+	#define REGISTER_OUT_0 u1
+	#define REGISTER_OUT_1 u2
+	#define REGISTER_OUT_2 u3
+	#define REGISTER_OUT_3 u4
+	#define REGISTER_OUT_4 u5
+#else
+	#define REGISTER_OUT_0 u0
+	#define REGISTER_OUT_1 u1
+	#define REGISTER_OUT_2 u2
+	#define REGISTER_OUT_3 u3
+	#define REGISTER_OUT_4 u4
+#endif
+
 #if VF_STRAND_HAIR
   #if SM5_PROFILE && COMPILER_HLSL
-	RWTexture3D<uint> DensityTexture  : register(u1);
+	RWTexture3D<uint> DensityTexture  : register(REGISTER_OUT_0);
 
 	#if SUPPORT_TANGENT_PROPERTY == 1
-	RWTexture3D<uint> TangentXTexture : register(u2);
-	RWTexture3D<uint> TangentYTexture : register(u3);
-	RWTexture3D<uint> TangentZTexture : register(u4);
+	RWTexture3D<uint> TangentXTexture : register(REGISTER_OUT_1);
+	RWTexture3D<uint> TangentYTexture : register(REGISTER_OUT_2);
+	RWTexture3D<uint> TangentZTexture : register(REGISTER_OUT_3);
 	#endif
 
 	#if SUPPORT_MATERIAL_PROPERTY == 1
-	RWTexture3D<uint> MaterialTexture : register(u5);
+	RWTexture3D<uint> MaterialTexture : register(REGISTER_OUT_4);
 	#endif
   #else
 	RWTexture3D<uint> DensityTexture;
@@ -79,21 +98,66 @@ void MainDom(
 #endif 
 
 void MainVoxel(
-#if SUPPORT_TANGENT_PROPERTY == 1 || SUPPORT_MATERIAL_PROPERTY == 1
+#if SUPPORT_TANGENT_PROPERTY == 1 || SUPPORT_MATERIAL_PROPERTY == 1 
 	FVertexFactoryInterpolantsVSToPS Interpolants,
 #endif
 	in float4 SvPosition : SV_Position,
 	in float4 PixelPosition : TEXCOORD6,
-	in float HairCoverage: CUSTOM_COVERAGE,
-	out uint OutColor : SV_Target0)
+	in float HairCoverage: CUSTOM_COVERAGE
+#if DEBUG_OUTPUT
+,	out uint OutColor : SV_Target0 // For debug purpose only
+#endif
+)
 {
-#if VF_STRAND_HAIR
-	const int3 VoxelCoord = (PixelPosition.xyz - DeepShadowPass.VoxelMinAABB) / (DeepShadowPass.VoxelMaxAABB - DeepShadowPass.VoxelMinAABB) * DeepShadowPass.VoxelResolution;
+#if VF_STRAND_HAIR  
+  #if SUPPORT_SPARSE	
+	const FPackedVirtualVoxelNodeDesc PackedNode = VoxelRasterPass.VirtualVoxel.NodeDescBuffer[VoxelRasterPass.MacroGroupId];
+	const FVirtualVoxelNodeDesc NodeDesc = UnpackVoxelNode(PackedNode, VoxelRasterPass.VirtualVoxel.PageResolution);
+
+	const int3 VoxelCoord = PositionToCoord(PixelPosition.xyz, NodeDesc.MinAABB, NodeDesc.MaxAABB, NodeDesc.VirtualResolution);
+	if (VoxelCoord.x >= 0 || VoxelCoord.y >= 0 || VoxelCoord.z >= 0 ||
+		VoxelCoord.x < int(NodeDesc.VirtualResolution.x) || 
+		VoxelCoord.y < int(NodeDesc.VirtualResolution.y) || 
+		VoxelCoord.z < int(NodeDesc.VirtualResolution.z))
+	{
+		const int3 PageIndexCoord		= VoxelCoord / VoxelRasterPass.VirtualVoxel.PageResolution;
+		const uint LinearPageIndexCoord = CoordToIndex(PageIndexCoord, NodeDesc.PageIndexResolution, NodeDesc.PageIndexOffset);
+		const int3 VoxelCoordBase		= PageIndexCoord * VoxelRasterPass.VirtualVoxel.PageResolution;
+		const int3 VoxelCoordOffset		= VoxelCoord - VoxelCoordBase;
+		if (LinearPageIndexCoord < VoxelRasterPass.VirtualVoxel.PageIndexCount)
+		{
+			const uint PageIndex = VoxelRasterPass.VirtualVoxel.PageIndexBuffer[LinearPageIndexCoord];
+			if (PageIndex != INVALID_VOXEL_PAGE_INDEX)
+			{
+				const uint3 PageCoord			= IndexToCoord(PageIndex, VoxelRasterPass.VirtualVoxel.PageCountResolution);
+				const uint3 PageVoxelCoordBase	= PageCoord * VoxelRasterPass.VirtualVoxel.PageResolution;
+				const uint3 PageVoxelCoord		= PageVoxelCoordBase + VoxelCoordOffset;
+
+				const float HairCount			= HairCoverage;
+				const float VoxelFixPointScale	= GetVoxelDensityFixPointScale();
+				InterlockedAdd(DensityTexture[PageVoxelCoord], HairCount * VoxelFixPointScale);
+
+				#if DEBUG_OUTPUT
+				OutColor = 0xFF;
+				#endif
+				return;
+			}
+			else
+			{
+				#if DEBUG_OUTPUT
+				OutColor = 0x0F;
+				#endif
+				return;
+			}
+		}
+	}
+  #else
+	const int3 VoxelCoord = (PixelPosition.xyz - VoxelRasterPass.VoxelMinAABB) / (VoxelRasterPass.VoxelMaxAABB - VoxelRasterPass.VoxelMinAABB) * VoxelRasterPass.VoxelResolution;
 		 
 	if (VoxelCoord.x >= 0 || VoxelCoord.y >= 0 || VoxelCoord.z >= 0 ||
-		VoxelCoord.x < int(DeepShadowPass.VoxelResolution) || 
-		VoxelCoord.y < int(DeepShadowPass.VoxelResolution) || 
-		VoxelCoord.z < int(DeepShadowPass.VoxelResolution))
+		VoxelCoord.x < int(VoxelRasterPass.VoxelResolution.x) || 
+		VoxelCoord.y < int(VoxelRasterPass.VoxelResolution.y) || 
+		VoxelCoord.z < int(VoxelRasterPass.VoxelResolution.z))
 	{
 		const float HairCount = HairCoverage;
 		const float VoxelFixPointScale = GetVoxelDensityFixPointScale();
@@ -128,6 +192,10 @@ void MainVoxel(
 		#endif
 
 	}
-#endif
+
+	#if DEBUG_OUTPUT
 	OutColor = 0xFF;
+	#endif
+  #endif
+#endif
 }
\ No newline at end of file
diff --git a/Engine/Shaders/Private/HairStrands/HairStrandsDeepShadowVS.usf b/Engine/Shaders/Private/HairStrands/HairStrandsDeepShadowVS.usf
--- a/Engine/Shaders/Private/HairStrands/HairStrandsDeepShadowVS.usf
+++ b/Engine/Shaders/Private/HairStrands/HairStrandsDeepShadowVS.usf
@@ -1,5 +1,6 @@
 // Copyright Epic Games, Inc. All Rights Reserved.
 
+#include "HairStrandsVoxelPageCommonStruct.ush"
 #include "../Common.ush"
 #include "/Engine/Generated/Material.ush"
 #include "/Engine/Generated/VertexFactory.ush"
@@ -27,7 +28,7 @@ struct FDeepShadowAccumulateVSToPS
 #define VS_OUTPUT_TYPE FDeepShadowAccumulateVSToPS
 #define SUPPORT_WPO 0
 
-#if VERTEXSHADER
+#if VERTEXSHADER 
 
 /** transform mesh as normal */
 void Main(
@@ -39,11 +40,50 @@ void Main(
 	ResolvedView = ResolveView();
 
 	FVertexFactoryIntermediates VFIntermediates = GetVertexFactoryIntermediates(Input);
-	float4 WorldPosition = VertexFactoryGetWorldPosition(Input, VFIntermediates);		
+
+#if MESH_RENDER_MODE == MESH_RENDER_VOXEL
+	float4x4 WorldToClipMatrix = VoxelRasterPass.WorldToClipMatrix;
+#elif MESH_RENDER_MODE == MESH_RENDER_DEPTH || MESH_RENDER_MODE == MESH_RENDER_DOM
+	float4x4 WorldToClipMatrix = DeepRasterPass.WorldToClipMatrix;
+#endif
+
+#ifdef VF_STRAND_HAIR	
+	const FHairRenderInfo HairRenderInfo = GetHairRenderInfo(ResolvedView.HairRenderInfo);
+
+	FHairViewInfo HairViewInfo;
+	HairViewInfo.TranslatedWorldCameraOrigin = ResolvedView.TranslatedWorldCameraOrigin;
+	HairViewInfo.ViewForward = ResolvedView.ViewForward;
+	HairViewInfo.RadiusAtDepth1 = HairRenderInfo.RadiusAtDepth1Primary;
+	HairViewInfo.bIsOrthoView = HairRenderInfo.bIsOrthoView;
+
+	#if MESH_RENDER_MODE == MESH_RENDER_VOXEL
+	float2 ViewportScale;
+	if (HairRenderInfo.bIsGPUDriven)
+	{
+		const uint MacroGroupId = VoxelRasterPass.MacroGroupId;
+		FVoxelizationViewInfo VoxelizationViewInfo = VoxelRasterPass.VoxelizationViewInfoBuffer.Load(MacroGroupId);
+		WorldToClipMatrix			= VoxelizationViewInfo.WorldToClip;
+		HairViewInfo.ViewForward	= VoxelizationViewInfo.ViewForward;
+		ViewportScale				= VoxelizationViewInfo.RasterResolution / float2(VoxelRasterPass.ViewportResolution);
+	}
+	#endif
+
+	float4 WorldPosition = VertexFactoryGetWorldPosition(Input, VFIntermediates, HairViewInfo);
+#else
+	float4 WorldPosition = VertexFactoryGetWorldPosition(Input, VFIntermediates);
+#endif
 	WorldPosition.xyz -= ResolvedView.PreViewTranslation.xyz; // Remove the camera relative position (todo move this into matrix transform)
 
-	Output.Position = mul(WorldPosition, DeepShadowPass.WorldToClipMatrix);
+	Output.Position = mul(WorldPosition, WorldToClipMatrix);
+
 #if MESH_RENDER_MODE == MESH_RENDER_VOXEL && defined(VF_STRAND_HAIR)
+	if (HairRenderInfo.bIsGPUDriven)
+	{
+		Output.Position.xy += float2(1, 1);
+		Output.Position.xy *= ViewportScale;
+		Output.Position.xy -= float2(1, 1);
+	}
+
 	Output.PixelPosition = WorldPosition;
   #if SUPPORT_TANGENT_PROPERTY == 1 || SUPPORT_MATERIAL_PROPERTY == 1
 	float3x3 TangentToLocal = VertexFactoryGetTangentToLocal(Input, VFIntermediates);
@@ -54,7 +94,6 @@ void Main(
 
 #if MESH_RENDER_MODE == MESH_RENDER_DOM || MESH_RENDER_MODE == MESH_RENDER_VOXEL
   #ifdef VF_STRAND_HAIR	
-	FHairRenderInfo HairRenderInfo = GetHairRenderInfo(ResolvedView.HairRenderInfo);
 	const float CurrentDepth = Output.Position.z / Output.Position.w;
 	const float PixelRadius = HairRenderInfo.bIsOrthoView ? HairRenderInfo.RadiusAtDepth1Primary : (CurrentDepth * HairRenderInfo.RadiusAtDepth1Primary);
 	const float StrandRealRadius = VFIntermediates.HairDimensions.y;
diff --git a/Engine/Shaders/Private/HairStrands/HairStrandsDeepTransmittanceCommon.ush b/Engine/Shaders/Private/HairStrands/HairStrandsDeepTransmittanceCommon.ush
--- a/Engine/Shaders/Private/HairStrands/HairStrandsDeepTransmittanceCommon.ush
+++ b/Engine/Shaders/Private/HairStrands/HairStrandsDeepTransmittanceCommon.ush
@@ -1,27 +1,49 @@
 // Copyright Epic Games, Inc. All Rights Reserved.
 
-//#include "HairStrandsShadowCommon.ush"
+#pragma once
 
-float3 HUEtoRGB(in float H)
+////////////////////////////////////////////////////////////////////////////////////////////////////
+
+struct FHairTransmittanceMask
 {
-	float R = abs(H * 6 - 3) - 1;
-	float G = 2 - abs(H * 6 - 2);
-	float B = 2 - abs(H * 6 - 4);
-	return saturate(float3(R, G, B));
+	float3 Transmittance;
+	float HairCount;
+	float Visibility;
+};
+
+struct FPackedHairTransmittanceMask
+{
+	float3 Transmittance;
+	uint PackedHairCountAndVisibility;
+};
+
+FHairTransmittanceMask InitHairTransmittanceMask()
+{
+	FHairTransmittanceMask Out;
+	Out.Transmittance = 1;
+	Out.HairCount = 0;
+	Out.Visibility = 1;
+	return Out;
 }
 
-float3 HSVtoRGB(in float3 HSV)
+FPackedHairTransmittanceMask PackTransmittanceMask(FHairTransmittanceMask In)
 {
-	float3 RGB = HUEtoRGB(HSV.x);
-	return ((RGB - 1) * HSV.y + 1) * HSV.z;
+	FPackedHairTransmittanceMask Out;
+	Out.Transmittance = In.Transmittance;
+	Out.PackedHairCountAndVisibility = min(uint(In.HairCount * 1000), 0x00FFFFFF) | (min(uint(In.Visibility * 0xFF), 0xFF) << 24);
+	return Out;
 }
 
-float3 ColorCode(float x)
+FHairTransmittanceMask UnpackTransmittanceMask(FPackedHairTransmittanceMask  In)
 {
-	float c = (1 - saturate(x)) * 0.6; // Remap [0,1] to Blue-Red
-	return x > 0 ? HSVtoRGB(float3(c, 1, 1)) : float3(0, 0, 0);
+	FHairTransmittanceMask Out;
+	Out.Transmittance = In.Transmittance;
+	Out.HairCount = float(In.PackedHairCountAndVisibility & 0x00FFFFFF) / 1000.f;
+	Out.Visibility = float((In.PackedHairCountAndVisibility & 0xFF000000) >> 24) / 255.f;
+	return Out;
 }
 
+////////////////////////////////////////////////////////////////////////////////////////////////////
 
 float3 ToLightPosition(float3 WorldPosition, float4x4 WorldToLightTransform)
 {
@@ -222,20 +244,24 @@ float3 ComputeHairCountDebugColor(
 	const float4 DOMValue = DomTexture.SampleLevel(LinearSampler, LightSpacePosition.xy, 0);
 	const float DistanceToFrontDepth = GetDomDistanceToFrontDepth(FrontDepth, LightSpacePosition.z);
 	const float HairCount = ComputeHairCount(DOMValue, DistanceToFrontDepth);
-	return ColorCode(HairCount / MaxHairCount);
+	return GetColorCode(HairCount / MaxHairCount);
 }
 
-float4 ComputeTransmittanceFromHairCount(float HairCount, FGBufferData GBuffer, float SinLightAngle, Texture3D<float4> InHairScatteringLUTTexture, SamplerState InHairLUTSampler)
+FHairTransmittanceMask ComputeTransmittanceFromHairCount(float HairCount, float HairVisibility, FGBufferData GBuffer, float SinLightAngle, Texture3D<float4> InHairScatteringLUTTexture, SamplerState InHairLUTSampler)
 {
-	HairAverageScattering AbsorptionData = SampleHairLUT(InHairScatteringLUTTexture, InHairLUTSampler, GBuffer.BaseColor, GBuffer.Roughness, SinLightAngle);
+	const HairAverageScattering AbsorptionData = SampleHairLUT(InHairScatteringLUTTexture, InHairLUTSampler, GBuffer.BaseColor, GBuffer.Roughness, SinLightAngle);
 
 	// Always shift the hair count by one to remove self-occlusion/shadow aliasing and have smootheer transition
 	// This insure the the pow function always starts at 0 for front facing hair
-	HairCount = max(0, HairCount - 1);
-	return  float4(pow(AbsorptionData.A_front, HairCount), HairCount);
+	FHairTransmittanceMask Out;
+	Out.HairCount = max(0, HairCount - 1);
+	Out.Transmittance = pow(AbsorptionData.A_front, Out.HairCount);
+	Out.Visibility = HairVisibility;
+
+	return Out;
 }
 
-float4 ComputeTransmittance(
+FHairTransmittanceMask ComputeTransmittance(
 	float3 WorldPosition,
 	FGBufferData GBuffer, 
 	FTransmittanceSettings Settings,
@@ -246,7 +272,8 @@ float4 ComputeTransmittance(
 {
 	const bool bIsHairPixel = GBuffer.ShadingModelID == SHADINGMODELID_HAIR;
 
-	float4 Transmittance = float4(1,1,1,0);
+	FHairTransmittanceMask Out = InitHairTransmittanceMask();
+
 	if (bIsHairPixel)
 	{
 		const bool bIsDirectionalLight = Settings.LightPosition.w == 0;
@@ -288,13 +315,13 @@ float4 ComputeTransmittance(
 			HairCount = 0;
 		}
 
-		Transmittance = ComputeTransmittanceFromHairCount(HairCount * Settings.DeepShadowDensityScale, GBuffer, SinLightAngle, HairLUTTexture, LinearSampler);
+		Out = ComputeTransmittanceFromHairCount(HairCount * Settings.DeepShadowDensityScale, 1, GBuffer, SinLightAngle, HairLUTTexture, LinearSampler);
 
 		if (Settings.DebugMode != 0)
 		{
-			Transmittance.xyz = ComputeDomDebugColor(WorldPosition, Settings.WorldToLightTransform, FrontDepthTexture, LinearSampler);
-			Transmittance.w = 1;
+			Out.Transmittance = ComputeDomDebugColor(WorldPosition, Settings.WorldToLightTransform, FrontDepthTexture, LinearSampler);
+			Out.HairCount = 1;
 		}
 	}
-	return Transmittance;
+	return Out;
 }
diff --git a/Engine/Shaders/Private/HairStrands/HairStrandsDeepTransmittanceDualScattering.ush b/Engine/Shaders/Private/HairStrands/HairStrandsDeepTransmittanceDualScattering.ush
--- a/Engine/Shaders/Private/HairStrands/HairStrandsDeepTransmittanceDualScattering.ush
+++ b/Engine/Shaders/Private/HairStrands/HairStrandsDeepTransmittanceDualScattering.ush
@@ -126,32 +126,35 @@ FHairTransmittanceData GetHairTransmittance(
 	const float3 V,
 	const float3 L,
 	FGBufferData GBuffer,
-	float4 HairTransmittance,
+	FHairTransmittanceMask TransmittanceMask,
 	Texture3D<float4> HairLUTTexture,
 	SamplerState HairLUTSampler,
 	const uint HairComponents)
 {
 	FHairTransmittanceData TransmittanceData = InitHairTransmittanceData();
-	TransmittanceData.Transmittance = HairTransmittance.xyz;
-	TransmittanceData.HairCount = HairTransmittance.w;
+	TransmittanceData.Transmittance = TransmittanceMask.Transmittance;
+	TransmittanceData.HairCount = TransmittanceMask.HairCount;
 	TransmittanceData.GlobalScattering = 1;
 	TransmittanceData.LocalScattering = 0;
 	TransmittanceData.ScatteringComponent = HairComponents | HAIR_COMPONENT_MULTISCATTER;
-	TransmittanceData.OpaqueVisibility = 1;
+	TransmittanceData.OpaqueVisibility = TransmittanceMask.Visibility;
+
+	// Hack: Override the actual roughness, with a different value to achieve a specific look
+	const float Roughness = GBuffer.CustomData.x > 0 ? GBuffer.CustomData.x : GBuffer.Roughness;
 
 	const float3 T = GBuffer.WorldNormal;
 	{
 		const float SinLightAngle = dot(L, T);
 
 		// Compute the transmittance based on precompute Hair transmittance LUT
-		HairAverageScattering AverageScattering = SampleHairLUT(HairLUTTexture, HairLUTSampler, GBuffer.BaseColor, GBuffer.Roughness, SinLightAngle);
+		HairAverageScattering AverageScattering = SampleHairLUT(HairLUTTexture, HairLUTSampler, GBuffer.BaseColor, Roughness, SinLightAngle);
 		TransmittanceData.A_front = AverageScattering.A_front;
 		TransmittanceData.A_back = AverageScattering.A_back;
 	}
 
 	ComputeDualScatteringTerms(
 		TransmittanceData,
-		GBuffer.Roughness,
+		Roughness,
 		V,
 		L,
 		T,
@@ -168,21 +171,22 @@ FHairTransmittanceData GetHairTransmittance(
 	FGBufferData GBuffer,
 	const uint SampleOffset,
 	const uint HairTransmittanceBufferMaxCount,
-	StructuredBuffer<float4> HairTransmittanceBuffer,
+	StructuredBuffer<FPackedHairTransmittanceMask> HairTransmittanceBuffer,
 	Texture3D<float4> HairLUTTexture,
 	SamplerState HairLUTSampler,
 	const uint HairComponents)
 {
-	float4 HairTransmittance = float4(1, 1, 1, 0);
+	FHairTransmittanceMask TransmittanceMask = InitHairTransmittanceMask();
 	if (SampleOffset < HairTransmittanceBufferMaxCount)
 	{
-		HairTransmittance = HairTransmittanceBuffer[SampleOffset];
+		const FPackedHairTransmittanceMask Packed = HairTransmittanceBuffer[SampleOffset];
+		TransmittanceMask = UnpackTransmittanceMask(Packed);
 	}
 	return GetHairTransmittance(
 		V,
 		L,
 		GBuffer,
-		HairTransmittance,
+		TransmittanceMask,
 		HairLUTTexture,
 		HairLUTSampler,
 		HairComponents);
diff --git a/Engine/Shaders/Private/HairStrands/HairStrandsDeepTransmittanceMask.usf b/Engine/Shaders/Private/HairStrands/HairStrandsDeepTransmittanceMask.usf
--- a/Engine/Shaders/Private/HairStrands/HairStrandsDeepTransmittanceMask.usf
+++ b/Engine/Shaders/Private/HairStrands/HairStrandsDeepTransmittanceMask.usf
@@ -1,15 +1,40 @@
 // Copyright Epic Games, Inc. All Rights Reserved.
 
+#include "HairStrandsVoxelPageCommonStruct.ush"
+
 #include "../Common.ush"
 #include "../CommonViewUniformBuffer.ush"
 #include "../SceneTextureParameters.ush"
 #include "../DeferredShadingCommon.ush"
 #include "../ShadowFilteringCommon.ush"
+#include "HairStrandsDebugCommon.ush"
 #include "HairStrandsDeepShadowCommon.ush"
 #include "HairStrandsCommon.ush"
 #include "HairStrandsDeepTransmittanceCommon.ush"
 #include "HairStrandsVisibilityCommon.ush"
-#include "HairStrandsDeepVoxelCommon.ush"
+#include "HairStrandsVoxelPageCommon.ush"
+
+#include "../ShaderDrawDebug.ush"
+#include "../ShaderPrintCommon.ush"
+
+#define PERMUTATION_TRANSMITTANCE_DEEP_SHADOW 0
+#define PERMUTATION_TRANSMITTANCE_VOXEL 1
+#define PERMUTATION_TRANSMITTANCE_VIRTUAL_VOXEL 2
+
+
+#if PERMUTATION_TRANSMITTANCE_TYPE == PERMUTATION_TRANSMITTANCE_DEEP_SHADOW
+	#define VOXEL_TRAVERSAL_TYPE VOXEL_TRAVERSAL_LEGACY
+#elif PERMUTATION_TRANSMITTANCE_TYPE == PERMUTATION_TRANSMITTANCE_VOXEL
+	#define VOXEL_TRAVERSAL_TYPE VOXEL_TRAVERSAL_LEGACY
+#elif PERMUTATION_TRANSMITTANCE_TYPE == PERMUTATION_TRANSMITTANCE_VIRTUAL_VOXEL
+	#if PERMUTATION_TRAVERSAL == 1
+		#define VOXEL_TRAVERSAL_TYPE VOXEL_TRAVERSAL_LINEAR_MIPMAP
+	#else
+		#define VOXEL_TRAVERSAL_TYPE VOXEL_TRAVERSAL_LINEAR
+	#endif
+#endif
+#define VOXEL_TRAVERSAL_DEBUG 0
+#include "HairStrandsVoxelPageTraversal.ush"
 
 float3 ToWorldPosition(float2 UV, float Depth)
 {
@@ -20,28 +45,27 @@ float3 ToWorldPosition(float2 UV, float Depth)
 	return WorldPosition;
 }
 
-#define PERMUTATION_TRANSMITTANCE_DEEP_SHADOW 0
-#define PERMUTATION_TRANSMITTANCE_VOXEL 1
 
-// Need to match FDeepShadowData::MaxClusterCount in HairRendering.h
-#define MAX_HAIR_CLUSTER_COUNT 16 
+// Need to match FDeepShadowData::MaxMacroGroupCount in HairRendering.h
+#define MAX_HAIR_MACROGROUP_COUNT 16 
 
 uint2				DeepShadow_Resolution;
-float4x4			DeepShadow_WorldToLightTransforms[MAX_HAIR_CLUSTER_COUNT];
+float4x4			DeepShadow_WorldToLightTransforms[MAX_HAIR_MACROGROUP_COUNT];
 float3				LightDirection;
 uint				MaxVisibilityNodeCount;
 float4				LightPosition;
 float				DepthBiasScale;
 float				DensityScale;
+float				LightRadius;
 float				DeepShadow_KernelAperture;
 uint				DeepShadow_KernelType;
 uint				DeepShadow_DebugMode;
-int2				DeepShadow_AtlasSlotOffsets[MAX_HAIR_CLUSTER_COUNT];
+int2				DeepShadow_AtlasSlotOffsets[MAX_HAIR_MACROGROUP_COUNT];
 Texture2D<float>	DeepShadow_FrontDepthTexture;
 Texture2D<float4>	DeepShadow_DomTexture;
 
-float3				Voxel_MinAABBs[MAX_HAIR_CLUSTER_COUNT];
-float3				Voxel_MaxAABBs[MAX_HAIR_CLUSTER_COUNT];
+float3				Voxel_MinAABBs[MAX_HAIR_MACROGROUP_COUNT];
+float3				Voxel_MaxAABBs[MAX_HAIR_MACROGROUP_COUNT];
 uint				Voxel_Resolution;
 Texture3D<uint> 	Voxel_DensityTexture0;
 Texture3D<uint> 	Voxel_DensityTexture1;
@@ -52,13 +76,15 @@ Texture3D<uint> 	Voxel_DensityTexture5;
 Texture3D<uint> 	Voxel_DensityTexture6;
 Texture3D<uint> 	Voxel_DensityTexture7;
 
+Texture2D<float4> 	RayMarchMaskTexture;
+
 StructuredBuffer<FNodeData> HairVisibilityNodeData;
 StructuredBuffer<uint>		HairVisibilityNodeCoord;
 Buffer<uint4>				IndirectArgsBuffer;
 
 Texture3D<float4>			HairLUTTexture;
-RWStructuredBuffer<float4>	OutputColor;
 SamplerState				LinearSampler;
+RWStructuredBuffer<FPackedHairTransmittanceMask> OutputColor;
 
 #if PERMUTATION_GROUP_SIZE == 0
 	#define TRANSMITTANCE_GROUP_SIZE 64
@@ -76,7 +102,9 @@ void MainCS(uint2 DispatchThreadId : SV_DispatchThreadID)
 	const uint LocalOffset = DispatchThreadId.x + DispatchThreadId.y * TRANSMITTANCE_GROUP_SIZE * HAIR_VISIBILITY_GROUP_COUNT_WIDTH;
 
 	if (LocalOffset >= MaxVisibilityNodeCount)
+	{
 		return;
+	}
 
 	const uint PackedCoord	 = HairVisibilityNodeCoord[LocalOffset];
 	const uint2 PixelCoord	 = uint2((PackedCoord & 0xFFFF), ((PackedCoord>>16) & 0xFFFF));
@@ -99,9 +127,9 @@ void MainCS(uint2 DispatchThreadId : SV_DispatchThreadID)
 		GBuffer.BaseColor = SampleBaseColor;
 		GBuffer.Roughness = SampleRoughness;
 
-		const uint ClusterIndex = clamp(SampleVisData.ClusterId, 0, MAX_HAIR_CLUSTER_COUNT - 1);
-		const int2 DeepShadowOffset = DeepShadow_AtlasSlotOffsets[ClusterIndex];
-		const float4x4 WorldToLightTransform = DeepShadow_WorldToLightTransforms[ClusterIndex];
+		const uint MacroGroupIndex = clamp(SampleVisData.MacroGroupId, 0, MAX_HAIR_MACROGROUP_COUNT - 1);
+		const int2 DeepShadowOffset = DeepShadow_AtlasSlotOffsets[MacroGroupIndex];
+		const float4x4 WorldToLightTransform = DeepShadow_WorldToLightTransforms[MacroGroupIndex];
 		
 		FTransmittanceSettings Settings;
 		Settings.DeepShadowOffset = DeepShadowOffset;
@@ -118,57 +146,192 @@ void MainCS(uint2 DispatchThreadId : SV_DispatchThreadID)
 
 		const float SceneDepth = ConvertFromDeviceZ(SampleDepth);
 		const float3 WorldPosition = ToWorldPosition(UV, SceneDepth);
-		OutputColor[LocalOffset] = ComputeTransmittance(WorldPosition, GBuffer, Settings, DeepShadow_FrontDepthTexture, DeepShadow_DomTexture, HairLUTTexture, LinearSampler);
+		const FHairTransmittanceMask TransmittanceMask = ComputeTransmittance(WorldPosition, GBuffer, Settings, DeepShadow_FrontDepthTexture, DeepShadow_DomTexture, HairLUTTexture, LinearSampler);
+
+		const FPackedHairTransmittanceMask PackedTransmittanceMask = PackTransmittanceMask(TransmittanceMask);
+		OutputColor[LocalOffset] = PackedTransmittanceMask;
 	}
 }
 #endif
 
-#if PERMUTATION_TRANSMITTANCE_TYPE == PERMUTATION_TRANSMITTANCE_VOXEL
+#if PERMUTATION_TRANSMITTANCE_TYPE == PERMUTATION_TRANSMITTANCE_VOXEL || PERMUTATION_TRANSMITTANCE_TYPE == PERMUTATION_TRANSMITTANCE_VIRTUAL_VOXEL
 
-float4 ComputeTransmittanceVoxel(
+#if VOXEL_TRAVERSAL_DEBUG
+void PrintInfo(uint Index, uint2 PixelCoord, FHairTransmittanceMask TransmittanceMask)
+{
+	#if 0
+	const uint2 OutputResolution = uint2(1024, 1024); // TODO send the actual output resolution
+	float2 Origin = (PixelCoord) / float2(OutputResolution);
+	Origin.y += (Index % 8) * 2 * ShaderPrintUniform.FontSize.w;
+
+	float2 Pos = Origin;
+
+	Pos = ShaderPrintSymbol(Pos, _H_);
+	Pos = ShaderPrintSymbol(Pos, _A_);
+	Pos = ShaderPrintSymbol(Pos, _I_);
+	Pos = ShaderPrintSymbol(Pos, _R_);
+	Pos = ShaderPrintSymbol(Pos, _SPC_);
+	Pos = ShaderPrintSymbol(Pos, _C_);
+	Pos = ShaderPrintSymbol(Pos, _O_);
+	Pos = ShaderPrintSymbol(Pos, _U_);
+	Pos = ShaderPrintSymbol(Pos, _N_);
+	Pos = ShaderPrintSymbol(Pos, _T_);
+	Pos = ShaderPrintSymbol(Pos, _SPC_);
+	Pos = ShaderPrint(Pos, TransmittanceMask.HairCount);
+
+	Pos = ShaderPrintNewline(Pos);
+	Pos.x = Origin.x;
+
+	Pos = ShaderPrintSymbol(Pos, _V_);
+	Pos = ShaderPrintSymbol(Pos, _I_);
+	Pos = ShaderPrintSymbol(Pos, _S_);
+	Pos = ShaderPrintSymbol(Pos, _I_);
+	Pos = ShaderPrintSymbol(Pos, _B_);
+	Pos = ShaderPrintSymbol(Pos, _I_);
+	Pos = ShaderPrintSymbol(Pos, _L_);
+	Pos = ShaderPrintSymbol(Pos, _I_);
+	Pos = ShaderPrintSymbol(Pos, _T_);
+	Pos = ShaderPrintSymbol(Pos, _Y_);
+	Pos = ShaderPrintSymbol(Pos, _SPC_);
+	Pos = ShaderPrint(Pos, TransmittanceMask.Visibility);
+
+	Pos = ShaderPrintNewline(Pos);
+	#endif
+}
+#endif
+
+FHairTransmittanceMask ComputeTransmittanceVoxel(
+	const uint2 PixelCoord,
 	float3 WorldPosition,
-	uint ClusterId,
+	uint MacroGroupId,
 	FGBufferData GBuffer,
 	FTransmittanceSettings Settings,
 	Texture3D<float4>	HairLUTTexture,
-	SamplerState		HairLUTSampler)
+	SamplerState		HairLUTSampler,
+	bool bDebugEnabled)
 {
+
+	const float VoxelShadowMaxDistance = 100000.0; // 1 km shadow distance	
 	const bool bIsDirectionalLight = Settings.LightPosition.w == 0;	
-	const float3 L = bIsDirectionalLight ? -Settings.LightDirection : normalize(Settings.LightPosition.xyz - WorldPosition);
+	const float3 LightPosition  = bIsDirectionalLight ? -Settings.LightDirection*VoxelShadowMaxDistance + WorldPosition : Settings.LightPosition.xyz;
+	const float3 LightDirection = bIsDirectionalLight ? -Settings.LightDirection : normalize(Settings.LightPosition.xyz - WorldPosition);
+	const float TanLightAngle = bIsDirectionalLight ? 0 : LightRadius / length(Settings.LightPosition.xyz - WorldPosition);
 	const float3 T = GBuffer.WorldNormal;
-	const float SinLightAngle = dot(L, T);
+	const float SinLightAngle = dot(LightDirection, T);
 
 	// Depth bias
 	// Origin is shifted voxels away towards the light + a constant bias of the size of the voxel
+	float MaxVoxelSize;
 	{
-		const float3 VoxelSize = (Voxel_MaxAABBs[ClusterId] - Voxel_MinAABBs[ClusterId]) / Voxel_Resolution;
-		const float3 DepthBias = max(VoxelSize.x, max(VoxelSize.y, VoxelSize.z)) * (Settings.DeepShadowDepthBiasScale*L + Settings.Random);
-		WorldPosition += DepthBias;
+	#if PERMUTATION_TRANSMITTANCE_TYPE == PERMUTATION_TRANSMITTANCE_VOXEL
+		const float3 VoxelSize = (Voxel_MaxAABBs[MacroGroupId] - Voxel_MinAABBs[MacroGroupId]) / Voxel_Resolution;
+		MaxVoxelSize = max(VoxelSize.x, max(VoxelSize.y, VoxelSize.z));
+	#elif PERMUTATION_TRANSMITTANCE_TYPE == PERMUTATION_TRANSMITTANCE_VIRTUAL_VOXEL
+		MaxVoxelSize = VirtualVoxel.VoxelWorldSize;
+	#endif
 	}
 
 	// Compute the number of hair count between light & shading point
 	const float DistanceThreshold = 100000;
-	float HairCount = 0;
-	if (ClusterId == 0) HairCount = ComputeHairCountVoxel(WorldPosition, L, Voxel_MinAABBs[0], Voxel_MaxAABBs[0], Voxel_Resolution, Voxel_DensityTexture0, Settings.DeepShadowDensityScale, 0, DistanceThreshold);
-	if (ClusterId == 1) HairCount = ComputeHairCountVoxel(WorldPosition, L, Voxel_MinAABBs[1], Voxel_MaxAABBs[1], Voxel_Resolution, Voxel_DensityTexture1, Settings.DeepShadowDensityScale, 0, DistanceThreshold);
-	if (ClusterId == 2) HairCount = ComputeHairCountVoxel(WorldPosition, L, Voxel_MinAABBs[2], Voxel_MaxAABBs[2], Voxel_Resolution, Voxel_DensityTexture2, Settings.DeepShadowDensityScale, 0, DistanceThreshold);
-	if (ClusterId == 3) HairCount = ComputeHairCountVoxel(WorldPosition, L, Voxel_MinAABBs[3], Voxel_MaxAABBs[3], Voxel_Resolution, Voxel_DensityTexture3, Settings.DeepShadowDensityScale, 0, DistanceThreshold);
-	if (ClusterId == 4) HairCount = ComputeHairCountVoxel(WorldPosition, L, Voxel_MinAABBs[4], Voxel_MaxAABBs[4], Voxel_Resolution, Voxel_DensityTexture4, Settings.DeepShadowDensityScale, 0, DistanceThreshold);
-	if (ClusterId == 5) HairCount = ComputeHairCountVoxel(WorldPosition, L, Voxel_MinAABBs[5], Voxel_MaxAABBs[5], Voxel_Resolution, Voxel_DensityTexture5, Settings.DeepShadowDensityScale, 0, DistanceThreshold);
-	if (ClusterId == 6) HairCount = ComputeHairCountVoxel(WorldPosition, L, Voxel_MinAABBs[6], Voxel_MaxAABBs[6], Voxel_Resolution, Voxel_DensityTexture6, Settings.DeepShadowDensityScale, 0, DistanceThreshold);
-	if (ClusterId == 7) HairCount = ComputeHairCountVoxel(WorldPosition, L, Voxel_MinAABBs[7], Voxel_MaxAABBs[7], Voxel_Resolution, Voxel_DensityTexture7, Settings.DeepShadowDensityScale, 0, DistanceThreshold);
-
-	return ComputeTransmittanceFromHairCount(HairCount, GBuffer, SinLightAngle, HairLUTTexture, HairLUTSampler);
-}
+	float AvgHairCount = 0;
+	float AvgHairVisility = 0;
+	#if PERMUTATION_SUPERSAMPLING
+	const uint SampleCount = 8; 
+	for (uint SampleIt = 0; SampleIt < SampleCount; ++SampleIt)
+	#else
+	const uint SampleIt = 0;
+	#endif
+	{
+		#if PERMUTATION_SUPERSAMPLING
+		const float3 SampleRandom = GetHairVoxelJitter(PixelCoord, SampleIt);
+		#else
+		const float3 SampleRandom = GetHairVoxelJitter(PixelCoord, View.StateFrameIndexMod8);
+		#endif
+
+		const float3 DepthBias = MaxVoxelSize * (Settings.DeepShadowDepthBiasScale*LightDirection + SampleRandom);
+		const float3 SampleWorldPosition = WorldPosition + DepthBias;
+
+		#if PERMUTATION_TRANSMITTANCE_TYPE == PERMUTATION_TRANSMITTANCE_VOXEL
+		FHairTraversalResult Result = InitHairTraversalResult();
+		if (MacroGroupId == 0) Result = ComputeHairCountVoxel(SampleWorldPosition, LightPosition, Voxel_MinAABBs[0], Voxel_MaxAABBs[0], Voxel_Resolution, Voxel_DensityTexture0, Settings.DeepShadowDensityScale, 0, DistanceThreshold);
+		if (MacroGroupId == 1) Result = ComputeHairCountVoxel(SampleWorldPosition, LightPosition, Voxel_MinAABBs[1], Voxel_MaxAABBs[1], Voxel_Resolution, Voxel_DensityTexture1, Settings.DeepShadowDensityScale, 0, DistanceThreshold);
+		if (MacroGroupId == 2) Result = ComputeHairCountVoxel(SampleWorldPosition, LightPosition, Voxel_MinAABBs[2], Voxel_MaxAABBs[2], Voxel_Resolution, Voxel_DensityTexture2, Settings.DeepShadowDensityScale, 0, DistanceThreshold);
+		if (MacroGroupId == 3) Result = ComputeHairCountVoxel(SampleWorldPosition, LightPosition, Voxel_MinAABBs[3], Voxel_MaxAABBs[3], Voxel_Resolution, Voxel_DensityTexture3, Settings.DeepShadowDensityScale, 0, DistanceThreshold);
+		if (MacroGroupId == 4) Result = ComputeHairCountVoxel(SampleWorldPosition, LightPosition, Voxel_MinAABBs[4], Voxel_MaxAABBs[4], Voxel_Resolution, Voxel_DensityTexture4, Settings.DeepShadowDensityScale, 0, DistanceThreshold);
+		if (MacroGroupId == 5) Result = ComputeHairCountVoxel(SampleWorldPosition, LightPosition, Voxel_MinAABBs[5], Voxel_MaxAABBs[5], Voxel_Resolution, Voxel_DensityTexture5, Settings.DeepShadowDensityScale, 0, DistanceThreshold);
+		if (MacroGroupId == 6) Result = ComputeHairCountVoxel(SampleWorldPosition, LightPosition, Voxel_MinAABBs[6], Voxel_MaxAABBs[6], Voxel_Resolution, Voxel_DensityTexture6, Settings.DeepShadowDensityScale, 0, DistanceThreshold);
+		if (MacroGroupId == 7) Result = ComputeHairCountVoxel(SampleWorldPosition, LightPosition, Voxel_MinAABBs[7], Voxel_MaxAABBs[7], Voxel_Resolution, Voxel_DensityTexture7, Settings.DeepShadowDensityScale, 0, DistanceThreshold);
+		#endif
+
+		#if PERMUTATION_TRANSMITTANCE_TYPE == PERMUTATION_TRANSMITTANCE_VIRTUAL_VOXEL
+		FVirtualVoxelCommonDesc CommonDesc;
+		CommonDesc.PageCountResolution	= VirtualVoxel.PageCountResolution;
+		CommonDesc.VoxelWorldSize		= VirtualVoxel.VoxelWorldSize;
+		CommonDesc.PageTextureResolution= VirtualVoxel.PageTextureResolution;
+		CommonDesc.PageResolution		= VirtualVoxel.PageResolution;
+
+		const FPackedVirtualVoxelNodeDesc PackedNode = VirtualVoxel.NodeDescBuffer[MacroGroupId];
+		const FVirtualVoxelNodeDesc NodeDesc = UnpackVoxelNode(PackedNode, VirtualVoxel.PageResolution);
+
+		FHairTraversalSettings TraversalSettings = InitHairTraversalSettings();
+		TraversalSettings.DensityScale = VirtualVoxel.DensityScale;
+		TraversalSettings.CountThreshold = GetOpaqueVoxelValue();
+		TraversalSettings.DistanceThreshold = DistanceThreshold;
+		TraversalSettings.bDebugEnabled = bDebugEnabled;
+		TraversalSettings.SteppingScale = VirtualVoxel.SteppingScale;
+		TraversalSettings.Random = SampleRandom;
+		TraversalSettings.TanConeAngle = 0; // TanLightAngle; // Area light transmission is disabled as it is too noisy for now
+
+		const FHairTraversalResult Result = ComputeHairCountVirtualVoxel(
+			SampleWorldPosition, 
+			LightPosition,
+			CommonDesc,
+			NodeDesc,
+			VirtualVoxel.PageIndexBuffer,
+			VirtualVoxel.PageTexture,
+			TraversalSettings);
+		InitHairTraversalResult();
+		#endif
+
+		AvgHairCount += Result.HairCount;
+		AvgHairVisility += Result.Visibility;
+	}
 
+	#if PERMUTATION_SUPERSAMPLING
+	AvgHairCount /= max(1, SampleCount);
+	AvgHairVisility /= max(1, SampleCount);
+	#endif
+
+	return ComputeTransmittanceFromHairCount(AvgHairCount, AvgHairVisility, GBuffer, SinLightAngle, HairLUTTexture, HairLUTSampler);
+}
 
 [numthreads(TRANSMITTANCE_GROUP_SIZE, 1, 1)]
 void MainCS(uint2 DispatchThreadId : SV_DispatchThreadID)
 {
 	const uint LocalOffset = DispatchThreadId.x + DispatchThreadId.y * TRANSMITTANCE_GROUP_SIZE * HAIR_VISIBILITY_GROUP_COUNT_WIDTH;
+
+	if (LocalOffset >= MaxVisibilityNodeCount)
+	{
+		return;
+	}
+
 	const uint PackedCoord = HairVisibilityNodeCoord[LocalOffset];
 	const uint2 PixelCoord = uint2((PackedCoord & 0xFFFF), ((PackedCoord >> 16) & 0xFFFF));
 
+	if (RayMarchMaskTexture.Load(uint3(PixelCoord, 0)).r == 0.0f)
+	{
+		// This pixel is already fully shadowed opaque. (RayMarchMaskTexture is computed using the closest sample as of today)
+		FHairTransmittanceMask DefaultMask = InitHairTransmittanceMask();
+		const FPackedHairTransmittanceMask PackedTransmittanceMask = PackTransmittanceMask(DefaultMask);
+		OutputColor[LocalOffset] = PackedTransmittanceMask;
+		return;
+	}
+
+	#if VOXEL_TRAVERSAL_DEBUG
+	const bool bDebugEnabled = PixelCoord.x == GetCursorPos().x && PixelCoord.y == GetCursorPos().y;
+	#else
+	const bool bDebugEnabled = false;
+	#endif	
 	const float2 UV = (PixelCoord.xy + float2(0.5f, 0.5f)) / float2(View.BufferSizeAndInvSize.xy);
 	{
 		const FNodeData NodeData = HairVisibilityNodeData[LocalOffset];
@@ -186,7 +349,7 @@ void MainCS(uint2 DispatchThreadId : SV_DispatchThreadID)
 		GBuffer.Roughness = SampleRoughness;
 		GBuffer.Specular = SampleSpecular;
 
-		const uint ClusterIndex = clamp(SampleVisData.ClusterId, 0, MAX_HAIR_CLUSTER_COUNT - 1);
+		const uint MacroGroupIndex = clamp(SampleVisData.MacroGroupId, 0, MAX_HAIR_MACROGROUP_COUNT - 1);
 
 		FTransmittanceSettings Settings;
 		Settings.LightDirection = LightDirection;
@@ -197,10 +360,18 @@ void MainCS(uint2 DispatchThreadId : SV_DispatchThreadID)
 
 		const float SceneDepth = ConvertFromDeviceZ(SampleDepth);
 		const float3 WorldPosition = ToWorldPosition(UV, SceneDepth);
-		OutputColor[LocalOffset] = ComputeTransmittanceVoxel(WorldPosition, ClusterIndex, GBuffer, Settings, HairLUTTexture, LinearSampler);
+	
+		const FHairTransmittanceMask TransmittanceMask = ComputeTransmittanceVoxel(PixelCoord, WorldPosition, MacroGroupIndex, GBuffer, Settings, HairLUTTexture, LinearSampler, bDebugEnabled);
+
+		#if VOXEL_TRAVERSAL_DEBUG
+		if (bDebugEnabled)
+		{
+			PrintInfo(LocalOffset, PixelCoord, TransmittanceMask);
+		}
+		#endif	
+
+		const FPackedHairTransmittanceMask PackedTransmittanceMask = PackTransmittanceMask(TransmittanceMask);
+		OutputColor[LocalOffset] = PackedTransmittanceMask;
 	}
 }
-#endif
-
-
-
+#endif
\ No newline at end of file
diff --git a/Engine/Shaders/Private/HairStrands/HairStrandsDeepVoxelCommon.ush b/Engine/Shaders/Private/HairStrands/HairStrandsDeepVoxelCommon.ush
--- a/Engine/Shaders/Private/HairStrands/HairStrandsDeepVoxelCommon.ush
+++ b/Engine/Shaders/Private/HairStrands/HairStrandsDeepVoxelCommon.ush
@@ -1,168 +1,2 @@
 // Copyright Epic Games, Inc. All Rights Reserved.
 
-float3 GetHairVoxelJitter(uint2 PixelCoord, uint Seed)
-{
-	return float3(
-		InterleavedGradientNoise(PixelCoord.xy, Seed),
-		InterleavedGradientNoise(PixelCoord.xy, Seed * 117),
-		InterleavedGradientNoise(PixelCoord.xy, Seed * 7901));
-}
-
-float GetHairVoxelDensity(uint3 VolumeCoord, Texture3D<uint> InDensityTexture, uint Mip=0)
-{
-	const float VoxelFixPointScale = GetVoxelDensityFixPointScale();
-	return InDensityTexture.Load(int4(VolumeCoord, Mip)) / VoxelFixPointScale;
-}
-
-bool IsInVoxelBounds(float3 P, float3 MinP, float3 MaxP)
-{
-	return
-		P.x >= MinP.x && P.y >= MinP.y && P.z >= MinP.z &&
-		P.x <= MaxP.x && P.y <= MaxP.y && P.z <= MaxP.z;
-}
-
-float ComputeHairCountVoxel(
-	float3 WorldPosition, 
-	float3 LightDirection, 
-	float3 InVoxelAABBMin, 
-	float3 InVoxelAABBMax, 
-	float InVoxelResolution, 
-	Texture3D<uint> InDensityTexture, 
-	float DensityScale, 
-	float CountThreshold,
-	float DistanceThreshold)
-{
-	float OutHairCount = 0;
-
-	const float3 VolumeBoundSize = InVoxelAABBMax - InVoxelAABBMin;
-	const float3 VolumeBoundSizeInv = 1.0f / (InVoxelAABBMax - InVoxelAABBMin);
-	const float VoxelWorldSize = VolumeBoundSize.x / InVoxelResolution;				// Only using voxel size along x for now...
-
-	const float VoxelShadowMaxDistance = 100000.0; // 1 km shadow distance
-	const float3 IntersectEndPoint = WorldPosition + LightDirection * VoxelShadowMaxDistance;
-	const float2 HitT = LineBoxIntersect(WorldPosition, IntersectEndPoint, InVoxelAABBMin, InVoxelAABBMax);
-	if (HitT.x < HitT.y)
-	{
-		// Count the number of fibers which are within a cylinder defined by the voxel size, 
-		// and the distance between the origin and the extent of the volume
-		// This assumes that the voxel volume is cubic (i.e. equal dimensions on all sides)
-		float3 O = lerp(WorldPosition, IntersectEndPoint, HitT.xxx);
-		float3 E = lerp(WorldPosition, IntersectEndPoint, HitT.yyy);
-		const float OELength = min(length(E - O), DistanceThreshold);
-
-		// Step according to voxel size
-		const float StepCount = float(min(ceil(OELength / VoxelWorldSize), 256u));	// TO TRY: sample with fixed count budget and jitter temporally
-		LightDirection *= VoxelWorldSize;
-
-		for (float StepIt = 0.0f; StepIt < StepCount; ++StepIt)
-		{
-			const float3 HitP = O + StepIt * LightDirection;
-			const float3 UVW = (HitP - InVoxelAABBMin) * VolumeBoundSizeInv;
-			const int3 VolumeCoord = clamp(int3(UVW * InVoxelResolution), int3(0, 0, 0), int3(InVoxelResolution.xxx-1));
-			const float Density = GetHairVoxelDensity(VolumeCoord, InDensityTexture) * DensityScale;
-			OutHairCount += Density;
-			if (CountThreshold != 0 && OutHairCount > CountThreshold)
-				return OutHairCount;
-		}
-	}
-
-	return OutHairCount;
-}
-
-float ComputeHairCountVoxelCone(
-	float3 WorldPosition,
-	float3 LightDirection,
-	float3 InVoxelAABBMin,
-	float3 InVoxelAABBMax,
-	float InVoxelResolution,
-	Texture3D<uint> InDensityTexture,
-	float InVoxelDensityScale,
-	float CountThreshold,
-	float DistanceThreshold,
-	float TanConeAngle)
-{
-	// Reverse-Z
-	const float ConservativeAABBDiagonal = (InVoxelAABBMax.x - InVoxelAABBMin.x) * 3;
-	const float3 O = WorldPosition;
-	const float3 E = O + LightDirection * ConservativeAABBDiagonal;
-	const float3 D = E - O;
-
-	float OutHairCount = 0;
-
-	// The rest of the code assumes O is within the AABB
-	const bool bIsInAABB = IsInVoxelBounds(O, InVoxelAABBMin, InVoxelAABBMax);
-	if (!bIsInAABB)
-		return OutHairCount;
-
-	const float2 HitT = LineBoxIntersect(O, E, InVoxelAABBMin, InVoxelAABBMax);
-	if (HitT.x < HitT.y)
-	{
-		const float VolumeSize = (InVoxelAABBMax.x - InVoxelAABBMin.x);
-		const float VoxelSize = VolumeSize / InVoxelResolution;
-
-		// Count the number of fibers which are within a cylinder defined by the voxel size, 
-		// and the distance between the origin and the extent of the volume
-		// This assumes that the voxel volume is cubic (i.e. equal dimensions on all sides)
-		const float DistanceThreshold2 = DistanceThreshold * DistanceThreshold;
-		const float Distance = min(DistanceThreshold, length(HitT.y * D));
-		const uint MaxStepCount = min(ceil(Distance / VoxelSize), 256u);
-		float T = 0;
-		float MipLevel = 0;
-		float MipDensityScale = 1;
-		for (uint StepIt = 0; StepIt < MaxStepCount; ++StepIt)
-		{
-			float3 StepOffset = 0;
-
-			const float3 RaymarchOffset = T * LightDirection;
-			const float3 HitP = O + RaymarchOffset + StepOffset;
-			const float3 UVW = (HitP - InVoxelAABBMin) / (InVoxelAABBMax - InVoxelAABBMin);
-
-			const float Distance2 = dot(RaymarchOffset, RaymarchOffset);
-			if (Distance2 > DistanceThreshold2)
-				return OutHairCount;
-
-			if (!IsInVoxelBounds(UVW, 0, 1))
-				return OutHairCount;
-
-			const uint MipVoxelResolution = uint(InVoxelResolution) >> uint(MipLevel);
-			const int3 VolumeCoord = clamp(int3(UVW * MipVoxelResolution), int3(0, 0, 0), int3(MipVoxelResolution.xxx - 1));
-			const float Density = GetHairVoxelDensity(VolumeCoord, InDensityTexture, MipLevel) * InVoxelDensityScale * MipDensityScale;
-
-			OutHairCount += Density;
-			if (CountThreshold != 0 && OutHairCount > CountThreshold)
-				return OutHairCount;
-
-			const float VoxelRadius = T * TanConeAngle;
-			const float NumVoxel = VolumeSize / max(2 * VoxelRadius, VoxelSize);
-
-			const float ResolutionBias = InVoxelResolution / NumVoxel;
-			float MipBias = log2(max(0.5f, ResolutionBias));
-			MipBias = floor(MipBias);
-			MipLevel = MipBias;
-
-			const float MipVoxelSize = VoxelSize * (MipLevel + 1);
-			{
-				MipDensityScale = VoxelSize * rcp(MipVoxelSize);
-				MipDensityScale = MipDensityScale * MipDensityScale * MipDensityScale;
-			}
-			T += MipVoxelSize;
-		}
-	}
-
-	return OutHairCount;
-}
-
-float3 GetHairVoxelTangent(
-	uint3 InVolumeCoord, 
-	float InHairDensity,
-	Texture3D<uint> InVoxelTangentXTexture,
-	Texture3D<uint> InVoxelTangentYTexture,
-	Texture3D<uint> InVoxelTangentZTexture)
-{
-	float3 T8Bits;
-	T8Bits.x = InVoxelTangentXTexture.Load(int4(InVolumeCoord, 0));
-	T8Bits.y = InVoxelTangentYTexture.Load(int4(InVolumeCoord, 0));
-	T8Bits.z = InVoxelTangentZTexture.Load(int4(InVolumeCoord, 0));
-	T8Bits /= InHairDensity;
-	return From8bits(T8Bits);
-}
\ No newline at end of file
diff --git a/Engine/Shaders/Private/HairStrands/HairStrandsEnvironmentAO.usf b/Engine/Shaders/Private/HairStrands/HairStrandsEnvironmentAO.usf
--- a/Engine/Shaders/Private/HairStrands/HairStrandsEnvironmentAO.usf
+++ b/Engine/Shaders/Private/HairStrands/HairStrandsEnvironmentAO.usf
@@ -1,16 +1,29 @@
 // Copyright Epic Games, Inc. All Rights Reserved.
   
+#include "HairStrandsVoxelPageCommonStruct.ush"
+
 #include "../Common.ush"
 #include "../MonteCarlo.ush"
 #include "../DeferredShadingCommon.ush"  
 #include "../SceneTextureParameters.ush"
 #include "HairStrandsCommon.ush"
 #include "HairStrandsVisibilityCommon.ush"
-#include "HairStrandsDeepVoxelCommon.ush"
+#include "HairStrandsVoxelPageCommon.ush"
+
+#include "../ShaderDrawDebug.ush"
+
+// Note: VOXEL_TRAVERSAL_TYPE is specified as a permutation
+// VOXEL_TRAVERSAL_LEGACY means old plain-voxel traversal
+#if VOXEL_TRAVERSAL_TYPE == 0 
+#undef VOXEL_TRAVERSAL_TYPE
+#define VOXEL_TRAVERSAL_TYPE VOXEL_TRAVERSAL_LEGACY
+#endif 
+#define VOXEL_TRAVERSAL_DEBUG 0 
+#include "HairStrandsVoxelPageTraversal.ush"
 
 Texture2D<uint4> HairCategorizationTexture;
 
-uint  Voxel_ClusterId;
+uint  Voxel_MacroGroupId;
 float Voxel_TanConeAngle;
 float Voxel_DensityScale;
 float Voxel_DepthBiasScale;
@@ -26,9 +39,9 @@ uint   Voxel_Resolution;
 Texture3D<uint> Voxel_DensityTexture;
 
 // todo: make this function in common with the sky lighting integration pass
-float InternalComputeHairCountVoxel(
+FHairTraversalResult InternalComputeHairCountVoxel(
 	float3 WorldPosition, 
-	float3 LightDirection, 
+	float3 LightPosition, 
 	float3 InVoxelAABBMin, 
 	float3 InVoxelAABBMax, 
 	float InVoxelResolution, 
@@ -36,12 +49,43 @@ float InternalComputeHairCountVoxel(
 	float DensityScale, 
 	float CountThreshold, 
 	float DistanceThreshold, 
-	bool bUseConeTracing)
+	bool bUseConeTracing,
+	bool bDebugEnabled)
 {
+#if VOXEL_TRAVERSAL_TYPE != VOXEL_TRAVERSAL_LEGACY
+
+	FVirtualVoxelCommonDesc CommonDesc;
+	CommonDesc.PageCountResolution = VirtualVoxel.PageCountResolution;
+	CommonDesc.VoxelWorldSize = VirtualVoxel.VoxelWorldSize;
+	CommonDesc.PageTextureResolution = VirtualVoxel.PageTextureResolution;
+	CommonDesc.PageResolution = VirtualVoxel.PageResolution;
+
+	const FPackedVirtualVoxelNodeDesc PackedNode = VirtualVoxel.NodeDescBuffer[Voxel_MacroGroupId];
+	const FVirtualVoxelNodeDesc NodeDesc = UnpackVoxelNode(PackedNode, VirtualVoxel.PageResolution);
+
+	FHairTraversalSettings TraversalSettings = InitHairTraversalSettings();
+	TraversalSettings.DensityScale = VirtualVoxel.DensityScale;
+	TraversalSettings.CountThreshold = CountThreshold;
+	TraversalSettings.DistanceThreshold = 1;
+	TraversalSettings.bDebugEnabled = bDebugEnabled;
+	TraversalSettings.SteppingScale = VirtualVoxel.SteppingScale;
+
+	return ComputeHairCountVirtualVoxel(
+		WorldPosition,
+		LightPosition,
+		CommonDesc,
+		NodeDesc,
+		VirtualVoxel.PageIndexBuffer,
+		VirtualVoxel.PageTexture,
+		TraversalSettings);
+
+#else
+
 	if (bUseConeTracing)
+	{
 		return ComputeHairCountVoxelCone(
 			WorldPosition,
-			LightDirection,
+			LightPosition,
 			InVoxelAABBMin,
 			InVoxelAABBMax,
 			InVoxelResolution,
@@ -50,10 +94,12 @@ float InternalComputeHairCountVoxel(
 			CountThreshold,
 			DistanceThreshold,
 			Voxel_TanConeAngle);
+	}
 	else
+	{
 		return ComputeHairCountVoxel(
 			WorldPosition,
-			LightDirection,
+			LightPosition,
 			InVoxelAABBMin,
 			InVoxelAABBMax,
 			InVoxelResolution,
@@ -61,13 +107,21 @@ float InternalComputeHairCountVoxel(
 			DensityScale,
 			CountThreshold,
 			DistanceThreshold);
+	}
+#endif
 }
 
 float TraceAO(float3 WorldPosition, float3 WorldNormal, float4 SvPosition)
 {
+	#if VOXEL_TRAVERSAL_TYPE == VOXEL_TRAVERSAL_LEGACY
 	const float3 VoxelSize = (Voxel_MaxAABB - Voxel_MinAABB) / Voxel_Resolution;
 	const float  VoxelMaxSize = max(VoxelSize.x, max(VoxelSize.y, VoxelSize.z));
+	#else
+	const float  VoxelMaxSize = VirtualVoxel.VoxelWorldSize;
+	#endif	
 	const float3 VoxelOffset = VoxelMaxSize * GetHairVoxelJitter(SvPosition.xy, View.StateFrameIndexMod8);
+
+	const bool bDebugEnabled = all(uint2(SvPosition.xy) == uint2(GetCursorPos()));
 	
 	float3 UnoccludedN = 0;
 	uint TotalHairCount = 0;
@@ -97,8 +151,9 @@ float TraceAO(float3 WorldPosition, float3 WorldNormal, float4 SvPosition)
 	const uint SampleCount = max(AO_SampleCount, 1);
 #endif
 
+	const float VoxelAOMaxDistance = 1000.0;
 	float3 BentNormal = 0;
-	float Occlusion = 0;
+	float Visibility = 0;
 	for (uint i = 0; i < SampleCount; ++i)
 	{
 	#if PERMUTATION_SAMPLESET == 0
@@ -110,23 +165,25 @@ float TraceAO(float3 WorldPosition, float3 WorldNormal, float4 SvPosition)
 			InterleavedGradientNoise(SvPosition.xy, View.StateFrameIndexMod8 + i * 13),
 			InterleavedGradientNoise(SvPosition.xy, View.StateFrameIndexMod8 + i * 117));
 
-		const float3 LocalL = CosineSampleHemisphere(frac(Rand2 + SequenceOffset)).xyz;
-		const float3 SampleL = TangentToWorld(LocalL.xyz, WorldNormal);
-
 		// Depth bias
-		const float3 SampleVoxelOffset = VoxelMaxSize * GetHairVoxelJitter(SvPosition, View.StateFrameIndexMod8 + i * 17);
+		const float3 SampleVoxelOffset = VoxelMaxSize * GetHairVoxelJitter(SvPosition.xy, View.StateFrameIndexMod8 + i * 17);
 		const float3 SampleWorldPosition = WorldPosition + VoxelOffset;
 
+		const float3 LocalL = CosineSampleHemisphere(frac(Rand2 + SequenceOffset)).xyz;
+		const float3 SampleL = TangentToWorld(LocalL.xyz, WorldNormal);
+		const float3 SampleLightPosition = SampleWorldPosition + SampleL * VoxelAOMaxDistance;
+
 		// Compute the number of hair count between light & shading point
-		const float HairCount = InternalComputeHairCountVoxel(SampleWorldPosition, SampleL, Voxel_MinAABB, Voxel_MaxAABB, Voxel_Resolution, Voxel_DensityTexture, Voxel_DensityScale, 0, AO_DistanceThreshold, true);
-		BentNormal += SampleL * saturate(1 - HairCount);
-		Occlusion += saturate(HairCount);
+		const FHairTraversalResult Result = InternalComputeHairCountVoxel(SampleWorldPosition, SampleLightPosition, Voxel_MinAABB, Voxel_MaxAABB, Voxel_Resolution, Voxel_DensityTexture, Voxel_DensityScale, 1, AO_DistanceThreshold, true, bDebugEnabled);
+		const float StepVisibility = saturate(1 - Result.HairCount) * Result.Visibility;
+		BentNormal += SampleL * StepVisibility;
+		Visibility += StepVisibility;
 	}
-	Occlusion /= SampleCount;
+	Visibility /= SampleCount;
 	BentNormal /= SampleCount;
 
 	// User adjust AO
-	const float AO = 1 - (1 - pow(saturate(1 - Occlusion), AO_Power)) * AO_Intensity;
+	const float AO = 1 - (1 - pow(saturate(Visibility), AO_Power)) * AO_Intensity;
 	return saturate(AO);
 }
 
@@ -155,7 +212,7 @@ void MainPS(
 
 	// Trace AO for all pixel which are partially covered or not covered at all.
 	const FCategorizationData CatData = DecodeCategorizationData(HairCategorizationTexture.Load(PixelCoord));
-	const bool bIsValid = (CatData.PixelCoverage > 0 && CatData.PixelCoverage < 1) || (GBuffer.ShadingModelID != SHADINGMODELID_UNLIT);
+	const bool bIsValid = CatData.PixelCoverage < 1 || GBuffer.ShadingModelID != SHADINGMODELID_UNLIT;
 
 	if (bIsValid)
 	{
