Copying //UE4/Dev-Physics to Main (//UE4/Main)

diff --git a/Engine/Shaders/Private/HairStrands/HairStrandsEnvironmentLighting.usf b/Engine/Shaders/Private/HairStrands/HairStrandsEnvironmentLighting.usf
--- a/Engine/Shaders/Private/HairStrands/HairStrandsEnvironmentLighting.usf
+++ b/Engine/Shaders/Private/HairStrands/HairStrandsEnvironmentLighting.usf
@@ -3,7 +3,9 @@
 /*=============================================================================
 	ReflectionEnvironmentComputeShaders - functionality to apply local cubemaps.
 =============================================================================*/
-  
+
+#include "HairStrandsVoxelPageCommonStruct.ush"
+
 #include "../Common.ush"
 #include "HairStrandsCommon.ush"
 #include "../DeferredShadingCommon.ush"  
@@ -28,13 +30,25 @@
 #include "../ReflectionEnvironmentComposite.ush"
 
 #include "../HairBsdf.ush"
+#include "HairStrandsDebugCommon.ush"
 #include "HairStrandsVisibilityCommon.ush"
 #include "HairStrandsEnvironmentLightingCommon.ush"
-#include "HairStrandsDeepVoxelCommon.ush"
+#include "HairStrandsVoxelPageCommon.ush"
 #include "HairStrandsDeepShadowCommon.ush"
 #include "HairStrandsDeepTransmittanceCommon.ush"
 #include "HairStrandsDeepTransmittanceDualScattering.ush"
 
+#include "../ShaderDrawDebug.ush"
+
+// Note: VOXEL_TRAVERSAL_TYPE is specified as a permutation
+// VOXEL_TRAVERSAL_LEGACY means old plain-voxel traversal
+#if VOXEL_TRAVERSAL_TYPE == 0 
+	#undef VOXEL_TRAVERSAL_TYPE
+	#define VOXEL_TRAVERSAL_TYPE VOXEL_TRAVERSAL_LEGACY
+#endif 
+#define VOXEL_TRAVERSAL_DEBUG 0
+#include "HairStrandsVoxelPageTraversal.ush"
+
 // Note: 
 // * This version of environement lighting is very much in WIP state, with a lot of option to tune tweak to explore 
 //   what kind of approximation works the best for certains type of grooms and lighting condition
@@ -81,8 +95,9 @@ Texture2D<float> HairCountTexture;
 uint  MultipleScatterSampleCount;
 uint  JitterSphericalIntegration;
 uint  HairComponents;
+float HairDualScatteringRoughnessOverride;
 
-uint  Voxel_ClusterId;
+uint  Voxel_MacroGroupId;
 float Voxel_TanConeAngle;
 float Voxel_DensityScale;
 float Voxel_DepthBiasScale;
@@ -92,9 +107,9 @@ float3 Voxel_MaxAABB;
 uint   Voxel_Resolution;
 Texture3D<uint> Voxel_DensityTexture;
 
-float InternalComputeHairCountVoxel(
+FHairTraversalResult InternalComputeHairCountVoxel(
 	float3 WorldPosition, 
-	float3 LightDirection, 
+	float3 LightPosition, 
 	float3 InVoxelAABBMin, 
 	float3 InVoxelAABBMax, 
 	float InVoxelResolution, 
@@ -102,12 +117,42 @@ float InternalComputeHairCountVoxel(
 	float DensityScale, 
 	float CountThreshold, 
 	float DistanceThreshold,
-	bool bUseConeTracing)
+	bool bUseConeTracing,
+	bool bDebugEnabled)
 {
+#if VOXEL_TRAVERSAL_TYPE != VOXEL_TRAVERSAL_LEGACY
+
+	FVirtualVoxelCommonDesc CommonDesc;
+	CommonDesc.PageCountResolution = VirtualVoxel.PageCountResolution;
+	CommonDesc.VoxelWorldSize = VirtualVoxel.VoxelWorldSize;
+	CommonDesc.PageTextureResolution = VirtualVoxel.PageTextureResolution;
+	CommonDesc.PageResolution = VirtualVoxel.PageResolution;
+
+	const FPackedVirtualVoxelNodeDesc PackedNode = VirtualVoxel.NodeDescBuffer[Voxel_MacroGroupId];
+	const FVirtualVoxelNodeDesc NodeDesc = UnpackVoxelNode(PackedNode, VirtualVoxel.PageResolution);
+
+	FHairTraversalSettings TraversalSettings = InitHairTraversalSettings();
+	TraversalSettings.DensityScale = DensityScale;
+	TraversalSettings.CountThreshold = GetOpaqueVoxelValue();
+	TraversalSettings.DistanceThreshold = DistanceThreshold;
+	TraversalSettings.bDebugEnabled = bDebugEnabled;
+	TraversalSettings.SteppingScale = VirtualVoxel.SteppingScale;
+
+	return ComputeHairCountVirtualVoxel(
+		WorldPosition,
+		LightPosition,
+		CommonDesc,
+		NodeDesc,
+		VirtualVoxel.PageIndexBuffer,
+		VirtualVoxel.PageTexture,
+		TraversalSettings);
+#else
+
 	if (bUseConeTracing)
+	{
 		return ComputeHairCountVoxelCone(
 			WorldPosition,
-			LightDirection,
+			LightPosition,
 			InVoxelAABBMin,
 			InVoxelAABBMax,
 			InVoxelResolution,
@@ -116,10 +161,12 @@ float InternalComputeHairCountVoxel(
 			CountThreshold,
 			DistanceThreshold,
 			Voxel_TanConeAngle);
+	}
 	else
+	{
 		return ComputeHairCountVoxel(
 			WorldPosition,
-			LightDirection,
+			LightPosition,
 			InVoxelAABBMin,
 			InVoxelAABBMax,
 			InVoxelResolution,
@@ -127,6 +174,9 @@ float InternalComputeHairCountVoxel(
 			DensityScale,
 			CountThreshold,
 			DistanceThreshold);
+	}
+
+#endif
 }
 
 //////////////////////////////////////////////////////////////////////////////////////////////////////////////////////////
@@ -140,19 +190,24 @@ float3 ReflectionEnvironment(FGBufferData GBuffer, float2 ScreenPosition, float2
 #else
 	const float PreExposure = 1.f;
 #endif
-	
+	const float  MaxDistance = 10000.f; // 100m
 	const float3 WorldPosition	= mul(float4(ScreenPosition * GBuffer.Depth, GBuffer.Depth, 1), View.ScreenToWorld).xyz;
 	const float3 CameraToPixel	= normalize(WorldPosition - View.WorldCameraOrigin);
 	const uint2  PixelPosition	= SvPosition - ResolvedView.ViewRectMin.xy;
 	const float  PixelSceneDepth= GBuffer.Depth;
 
-	const uint PrimitiveIdMod8 = (PrimitiveId % 8);
+#if VOXEL_TRAVERSAL_TYPE != VOXEL_TRAVERSAL_LEGACY
+	const float3 VoxelSize = VirtualVoxel.VoxelWorldSize; // Use the correct voxel size when we use virtual voxel volume.
+#else
 	const float3 VoxelSize = (Voxel_MaxAABB - Voxel_MinAABB) / Voxel_Resolution;
+#endif
+	const uint PrimitiveIdMod8 = (PrimitiveId % 8);
 	const float  VoxelMaxSize = max(VoxelSize.x, max(VoxelSize.y, VoxelSize.z));
 	const float3 VoxelOffset = VoxelMaxSize * GetHairVoxelJitter(SvPosition, View.StateFrameIndexMod8 + PrimitiveIdMod8);
 
 	const float3 V = -CameraToPixel;
 	const float3 T = GBuffer.WorldNormal;
+	const bool bDebugEnabled = all(PixelPosition == uint2(GetCursorPos()));
 
 	////////////////////////////////////////////////////////////////////////////////////////////////////////////////////////////////////////////
 	////////////////////////////////////////////////////////////////////////////////////////////////////////////////////////////////////////////
@@ -208,28 +263,30 @@ float3 ReflectionEnvironment(FGBufferData GBuffer, float2 ScreenPosition, float2
 			const float2 Rand2 = (PoissonDisk[i] + float2(1, 1)) * 0.5f + RandOffset;
 		#else // PERMUTATION_SAMPLESET == 1
 			const float2 Rand2 = Hammersley(i, SampleCount, uint2(0,0)) + RandOffset;
-		#endif		
-			const float3 SampleL = UniformSampleSphere(frac(Rand2)).xyz;
+		#endif
 
 			// Depth bias
 			const float3 SampleVoxelOffset = VoxelMaxSize * GetHairVoxelJitter(SvPosition, PrimitiveIdMod8 + i);
 			const float3 SampleWorldPosition = WorldPosition + SampleVoxelOffset;
 
+			const float3 SampleL = UniformSampleSphere(frac(Rand2)).xyz;
+			const float3 SampleLightPosition = WorldPosition + SampleL*MaxDistance;
+
 			// Compute the number of hair count between light & shading point
 			const float SinLightAngle = dot(SampleL, T);
-			const float HairCount = InternalComputeHairCountVoxel(SampleWorldPosition, SampleL, Voxel_MinAABB, Voxel_MaxAABB, Voxel_Resolution, Voxel_DensityTexture, Voxel_DensityScale, 0, MaxSearchDistance, true);
-			const float4 SampleTransmittance = ComputeTransmittanceFromHairCount(HairCount, GBuffer, SinLightAngle, HairScatteringLUTTexture, HairLUTSampler);
+			const FHairTraversalResult Result = InternalComputeHairCountVoxel(SampleWorldPosition, SampleLightPosition, Voxel_MinAABB, Voxel_MaxAABB, Voxel_Resolution, Voxel_DensityTexture, Voxel_DensityScale, 0, MaxSearchDistance, true, bDebugEnabled);
+			const FHairTransmittanceMask TransmittanceMask = ComputeTransmittanceFromHairCount(Result.HairCount, Result.Visibility, GBuffer, SinLightAngle, HairScatteringLUTTexture, HairLUTSampler);
 
-			const float ScalarTransmittance = max(SampleTransmittance.x, max(SampleTransmittance.y, SampleTransmittance.z));
-			AvgR += SampleL * SampleTransmittance.xyz;
+			const float ScalarTransmittance = max(TransmittanceMask.Transmittance.x, max(TransmittanceMask.Transmittance.y, TransmittanceMask.Transmittance.z));
+			AvgR += SampleL * ScalarTransmittance * Result.Visibility;
 
 			const uint ScatteringComponent = (HairComponents & HAIR_COMPONENT_R) | (HairComponents & HAIR_COMPONENT_TRT) | HAIR_COMPONENT_LS | HAIR_COMPONENT_GS;// | HAIR_MULTISCATTER;
-			const FHairTransmittanceData TransmittanceData = GetHairTransmittance(V, SampleL, GBuffer, SampleTransmittance, HairScatteringLUTTexture, HairLUTSampler, ScatteringComponent);
+			const FHairTransmittanceData TransmittanceData = GetHairTransmittance(V, SampleL, GBuffer, TransmittanceMask, HairScatteringLUTTexture, HairLUTSampler, ScatteringComponent);
 
 			const float Area = 0.2;
 			const float BackLit = 1;
 			//AvgDualScattering += HairShading(GBuffer, SampleL, V, T, 1, TransmittanceData, BackLit, Area, uint2(0, 0));
-			AvgDualScattering += TransmittanceData.LocalScattering * TransmittanceData.GlobalScattering;
+			AvgDualScattering += TransmittanceData.LocalScattering * TransmittanceData.GlobalScattering * Result.Visibility;
 		}
 		const float SphereInt = 4 * PI;
 		AvgDualScattering *= SphereInt;
@@ -239,6 +296,13 @@ float3 ReflectionEnvironment(FGBufferData GBuffer, float2 ScreenPosition, float2
 		const float NormR = sqrt(dot(AvgR, AvgR));
 		UnoccludedN = AvgR / NormR;
 		
+		#if VOXEL_TRAVERSAL_DEBUG
+		if (bDebugEnabled)
+		{
+			AddLine(WorldPosition, WorldPosition + UnoccludedN * 50, float4(1, 0, 1, 1), float4(1, 0, 1, 1));
+		}
+		#endif
+
 		// For debug purpose:
 		//const float UnoccludedVarianceR = NormR * (3 - NormR2) / (1 - NormR2);
 		//const float3 ColorDir = (UnoccludedN+ float3(1, 1, 1)) * 0.5f;
@@ -329,16 +393,17 @@ float3 ReflectionEnvironment(FGBufferData GBuffer, float2 ScreenPosition, float2
 		const float SinViewAngle = dot(V, T);
 
 		// Compute the transmittance
-		float4 Transmittance = float4(1,1,1,0);
+		FHairTransmittanceMask TransmittanceMask = InitHairTransmittanceMask();
 		{
 			// Depth bias
 			// Origin is shifted voxels away towards the light + a constant bias of the size of the voxel
 			const float3 DepthBias = VoxelOffset /*+ VoxelMaxSize*Voxel_DepthBiasScale*L*/;
 			const float3 SampleWorldPosition = WorldPosition + DepthBias;
+			const float3 SampleLightPosition = WorldPosition + L * MaxDistance;
 
 			// Compute the number of hair count between light & shading point
-			const float HairCount = InternalComputeHairCountVoxel(SampleWorldPosition, L, Voxel_MinAABB, Voxel_MaxAABB, Voxel_Resolution, Voxel_DensityTexture, Voxel_DensityScale, 0, MaxSearchDistance, true);
-			Transmittance = ComputeTransmittanceFromHairCount(HairCount, GBuffer, SinLightAngle, HairScatteringLUTTexture, HairLUTSampler);
+			const FHairTraversalResult Result = InternalComputeHairCountVoxel(SampleWorldPosition, SampleLightPosition, Voxel_MinAABB, Voxel_MaxAABB, Voxel_Resolution, Voxel_DensityTexture, Voxel_DensityScale, 0, MaxSearchDistance, true, bDebugEnabled);
+			TransmittanceMask = ComputeTransmittanceFromHairCount(Result.HairCount, Result.Visibility, GBuffer, SinLightAngle, HairScatteringLUTTexture, HairLUTSampler);
 		}
 
 		FHairTransmittanceData TransmittanceData = InitHairTransmittanceData();	
@@ -347,19 +412,19 @@ float3 ReflectionEnvironment(FGBufferData GBuffer, float2 ScreenPosition, float2
 		#if 0
 		{
 			HairAverageEnergy Energy = SampleHairEnergyLUT(HairEnergyLUTTexture, HairLUTSampler, GBuffer.BaseColor, GBuffer.Roughness, SinViewAngle);
-			ColorR.rgb *= (Energy.A_TRT + Energy.A_R) * Transmittance.xyz;
+			ColorR.rgb *= (Energy.A_TRT + Energy.A_R) * Transmittance.xyz * Visibility;
 		}
 		#endif
 
 		// BSDF
 		#if 1
 		{			
-			TransmittanceData.GlobalScattering = Transmittance.xyz;
+			TransmittanceData.GlobalScattering = TransmittanceMask.Transmittance;
 			TransmittanceData.ScatteringComponent = (HairComponents & HAIR_COMPONENT_R) | (HairComponents & HAIR_COMPONENT_TRT) | HAIR_COMPONENT_MULTISCATTER;
 
 			const float Area = 0.2;
 			const float BackLit = 0;
-			ColorR.rgb *= HairShading(GBuffer, L, V, T, 1, TransmittanceData, BackLit, Area, uint2(0, 0));
+			ColorR.rgb *= HairShading(GBuffer, L, V, T, 1, TransmittanceData, BackLit, Area, uint2(0, 0)) * TransmittanceMask.Visibility;
 		}
 		#endif
 	}
@@ -383,32 +448,39 @@ float3 ReflectionEnvironment(FGBufferData GBuffer, float2 ScreenPosition, float2
 		//ColorTT.rgb += PreExposure * GetSkySHDiffuse(L) * View.SkyLightColor.rgb;
 
 		// Compute the transmittance
-		float4 Transmittance = float4(1,1,1,0);
+		FHairTransmittanceMask TransmittanceMask = InitHairTransmittanceMask();
 		{
 			// Depth bias
 			// Origin is shifted voxels away towards the light + a constant bias of the size of the voxel
 			const float3 DepthBias = VoxelOffset /*+ VoxelMaxSize*Voxel_DepthBiasScale*L*/;
 			const float3 SampleWorldPosition = WorldPosition + DepthBias;
+			const float3 SampleLightPosition = WorldPosition + L * MaxDistance;
 
 			// Compute the number of hair count between light & shading point
 			// For TT, we don't want to abort the hair counting, as this is important to get the proper transmission
 			const float DistanceThresholdTT = 100000;
 			const float DensityScaleFactor = 2;
-			float HairCount = 0;
+			FHairTraversalResult Result;
 			if (bHairUseViewHairCount)
-				HairCount = HairCountTexture.Load(uint3(SvPosition.xy, 0)) * DensityScaleFactor;
+			{
+				Result.HairCount = HairCountTexture.Load(uint3(SvPosition.xy, 0)) * DensityScaleFactor;
+				Result.Visibility = 1;
+			}
 			else
-				HairCount = InternalComputeHairCountVoxel(SampleWorldPosition, L, Voxel_MinAABB, Voxel_MaxAABB, Voxel_Resolution, Voxel_DensityTexture, Voxel_DensityScale * DensityScaleFactor, 0, DistanceThresholdTT, false);
-			Transmittance = ComputeTransmittanceFromHairCount(HairCount, GBuffer, SinLightAngle, HairScatteringLUTTexture, HairLUTSampler);
+			{
+				Result = InternalComputeHairCountVoxel(SampleWorldPosition, SampleLightPosition, Voxel_MinAABB, Voxel_MaxAABB, Voxel_Resolution, Voxel_DensityTexture, Voxel_DensityScale * DensityScaleFactor, 0, DistanceThresholdTT, false, bDebugEnabled);
+			}
+			TransmittanceMask = ComputeTransmittanceFromHairCount(Result.HairCount, Result.Visibility, GBuffer, SinLightAngle, HairScatteringLUTTexture, HairLUTSampler);
 		}
 
 		// BSDF with dual-scattering
 		#if 0
 		{
+			Error the Transmittance.w does not contain the correct hair count when we hit opaque geometry: TODO
 			const uint HairComponent = HAIR_COMPONENT_TT | HAIR_COMPONENT_LS | HAIR_COMPONENT_GS;
 			const FHairTransmittanceData TransmittanceData = GetHairTransmittance(V, L, GBuffer, Transmittance, HairScatteringLUTTexture, HairLUTSampler, HairComponent);
 			const float Area = 0.2;
-			ColorTT.rgb *= HairShading(GBuffer, L, V, T, 1, TransmittanceData, 1, Area, uint2(0, 0));
+			ColorTT.rgb *= HairShading(GBuffer, L, V, T, 1, TransmittanceData, 1, Area, uint2(0, 0)) * Visibility;
 		}
 		#endif
 
@@ -416,7 +488,7 @@ float3 ReflectionEnvironment(FGBufferData GBuffer, float2 ScreenPosition, float2
 		#if 1
 		{
 			HairAverageEnergy Energy = SampleHairEnergyLUT(HairEnergyLUTTexture, HairLUTSampler, GBuffer.BaseColor, GBuffer.Roughness, SinViewAngle);
-			ColorTT.rgb *= Energy.A_TT * Transmittance.xyz;
+			ColorTT.rgb *= Energy.A_TT * TransmittanceMask.Transmittance * TransmittanceMask.Visibility;
 		}
 		#endif	
 
@@ -430,7 +502,7 @@ float3 ReflectionEnvironment(FGBufferData GBuffer, float2 ScreenPosition, float2
 			TransmittanceData.ScatteringComponent = HAIR_COMPONENT_TT | HAIR_COMPONENT_MULTISCATTER;
 
 			const float Area = 0.2;
-			ColorTT.rgb *= HairShading(GBuffer, L, V, T, 1, TransmittanceData, 0, Area, uint2(0, 0));
+			ColorTT.rgb *= HairShading(GBuffer, L, V, T, 1, TransmittanceData, 0, Area, uint2(0, 0)) * Visibility;
 		}
 		#endif	
 	}
@@ -490,11 +562,11 @@ void MainCS(uint3 DispatchThreadId : SV_DispatchThreadID)
 	GBuffer.Specular = SampleSpecular;
 	GBuffer.Depth = ConvertFromDeviceZ(SampleDepth);
 	GBuffer.GBufferAO = 1;
-	GBuffer.CustomData = 0;
+	GBuffer.CustomData = float4(HairDualScatteringRoughnessOverride,0,0,0);
 	GBuffer.IndirectIrradiance = 1000000;
 
-	const bool bIsClusterCompatible = SampleVisData.ClusterId == Voxel_ClusterId;
-	if (bIsClusterCompatible)
+	const bool bIsMacroGroupCompatible = SampleVisData.MacroGroupId == Voxel_MacroGroupId;
+	if (bIsMacroGroupCompatible)
 	{
 		const float SampleCoverage = saturate(From8bitCoverage(SampleVisData.Coverage8bit));
 		const float3 SkyLighting = ReflectionEnvironment(GBuffer, ScreenPosition, SvPosition, SampleVisData.PrimitiveId);
@@ -512,4 +584,4 @@ void MainCS(uint3 DispatchThreadId : SV_DispatchThreadID)
 		OutLightingBuffer[LocalOffset] = 0;
 	}
 }
-#endif
\ No newline at end of file
+#endif
diff --git a/Engine/Shaders/Private/HairStrands/HairStrandsGuideDeform.usf b/Engine/Shaders/Private/HairStrands/HairStrandsGuideDeform.usf
--- a/Engine/Shaders/Private/HairStrands/HairStrandsGuideDeform.usf
+++ b/Engine/Shaders/Private/HairStrands/HairStrandsGuideDeform.usf
@@ -2,9 +2,11 @@
 
 #include "../Common.ush"
 
-float3 HairWorldOffset;
+float3 SimRestOffset;
+float3 SimDeformedOffset;
 uint VertexCount;
 uint IterationCount;
+uint DispatchCountX;
 
 Buffer<uint4> SimRestPosePositionBuffer;
 Buffer<uint>  SimRootIndexBuffer;
@@ -46,12 +48,16 @@ float3 Displacement(uint VertexIndex, float3 Pos, float Seed)
 	float Time = 0.5f * cos(Seed / 64.f) + 0.5f; // [0..1]
 	return lerp(0, Offset, Time);
 #endif
+
+#if PERMUTATION_DEFORMATION == 3
+	return SimRestOffset - SimDeformedOffset;
+#endif
 }
 
 [numthreads(GROUP_SIZE, 1, 1)]
 void MainCS(uint2 DispatchThreadId : SV_DispatchThreadID)
 {
-	const uint VertexIndex = DispatchThreadId.x;
+	const uint VertexIndex = DispatchThreadId.x + DispatchThreadId.y * DispatchCountX * GROUP_SIZE;
 	if (VertexIndex < VertexCount)
 	{
 		const uint4  Data 			= SimRestPosePositionBuffer[VertexIndex];
diff --git a/Engine/Shaders/Private/HairStrands/HairStrandsInterpolation.usf b/Engine/Shaders/Private/HairStrands/HairStrandsInterpolation.usf
--- a/Engine/Shaders/Private/HairStrands/HairStrandsInterpolation.usf
+++ b/Engine/Shaders/Private/HairStrands/HairStrandsInterpolation.usf
@@ -2,13 +2,29 @@
 
 #include "../Common.ush"
 #include "HairStrandsMeshProjectionCommon.ush"
+#include "HairStrandsVisibilityCommon.ush"
+#include "HairStrandsAABBCommon.ush"
+
+#if PERMUTATION_DEBUG == 1
+#include "../ShaderDrawDebug.ush"
+#endif
 
 uint VertexCount;
 uint DispatchCountX;
+uint HairDebugMode;
 float3 InRenderHairPositionOffset;
 float3 InSimHairPositionOffset;
 float3 OutHairPositionOffset;
 int2 HairStrandsCullIndex;
+float4x4 LocalToWorldMatrix;
+
+float InHairLength;
+float InHairRadius;
+float OutHairRadius;
+float MaxOutHairRadius;
+float HairRadiusRootScale;
+float HairRadiusTipScale;
+float HairLengthClip;
 
 Buffer<uint4>  RenderRestPosePositionBuffer;
 Buffer<uint4>  SimRestPosePositionBuffer;
@@ -16,7 +32,11 @@ Buffer<uint4>  DeformedSimPositionBuffer;
 Buffer<uint4>  Interpolation0Buffer;
 Buffer<uint4>  Interpolation1Buffer;
 
-#if PERMUTATION_DYNAMIC_GEOMETRY == 1 || PERMUTATION_DEBUG == 2
+#if PERMUTATION_SCALE_AND_CLIP
+Buffer<uint4>  AttributeBuffer;
+#endif
+
+#if PERMUTATION_DYNAMIC_GEOMETRY == 1 || PERMUTATION_DEBUG == 1
 Buffer<uint>   RenVertexToRootIndexBuffer;
 Buffer<uint>   SimVertexToRootIndexBuffer;
 #endif
@@ -54,6 +74,7 @@ Buffer<uint>	SimRootToTriangleIndex;
 #endif
 
 #if PERMUTATION_DEBUG == 1
+Buffer<int>		 VertexToClusterIdBuffer;
 Buffer<float4>	 SimAttributeBuffer;
 RWBuffer<float4> OutRenderAttributeBuffer;
 #endif
@@ -90,6 +111,32 @@ float3 ComputeDynamicGeometryOffset(
 	return (LocalDeformedGuidePoint - LocalRestGuidePoint) * GuideVertexWeight;
 }
 
+// Compute the new hair width and return if the current vertex is valid (true), or if it needs to be clipped (false)
+bool ScaleAndClip(uint ControlPointId, inout uint InPackedAlpha)
+{
+#if PERMUTATION_SCALE_AND_CLIP
+	// Instead of doing the actual interpolation we could rely on dithering as an optimization
+	const FHairAttributes Attributes0 = DecodeHairAttribute(AttributeBuffer[ControlPointId]);
+	const float U = Attributes0.UCoord;
+
+	// todo: Move all these computations into shared function which can be shared with the hair VF
+	const float InNormVertexRadius = ((InPackedAlpha & 0x00FC) >> 2) / 63.f;
+	//const float InVertexRadius = InNormVertexRadius * InHairRadius;
+	const float InVertexLength = ((InPackedAlpha & 0xFF00) >> 8) / 255.f * InHairLength;
+
+	const float OutVertexRadius = OutHairRadius * InNormVertexRadius * lerp(HairRadiusRootScale, HairRadiusTipScale, U);
+	const uint PackedVertexRadius = (uint(saturate(OutVertexRadius / MaxOutHairRadius) * 63.f) & 0x3F) << 2;
+
+	InPackedAlpha = (InPackedAlpha & 0xFF03) | PackedVertexRadius;
+
+	const bool bIsValid = InVertexLength <= HairLengthClip;
+	return bIsValid;
+#else
+	return true;
+#endif
+}
+
+
 [numthreads(GROUP_SIZE, 1, 1)]
 void MainCS(uint2 DispatchThreadId : SV_DispatchThreadID)
 {
@@ -117,13 +164,18 @@ void MainCS(uint2 DispatchThreadId : SV_DispatchThreadID)
 		}
 		#endif
 
+		float3 WorldVertexPosition = 0.0f;
+
 		// 1. Rigid transformation
 		// ControlPoint is in the local hair referential
 		// CurrOffset takes only translation component into account, and is done in object local space (vs. triangle/hair local space)
 		#if PERMUTATION_DYNAMIC_GEOMETRY == 0
 		{
 			float3 CurrOffset = 0;
-			float3 ControlPoint = f16tof32(RenderRestPosePositionBuffer[VertexIndex].xyz) + InRenderHairPositionOffset;
+			uint4 PackedRestPosition = RenderRestPosePositionBuffer[VertexIndex];
+			float3 ControlPoint = f16tof32(PackedRestPosition.xyz) + InRenderHairPositionOffset;
+
+
 			#if PERMUTATION_SIMULATION == 1
 			[unroll]
 			for (uint K = 0; K < 3; ++K)
@@ -137,7 +189,11 @@ void MainCS(uint2 DispatchThreadId : SV_DispatchThreadID)
 				CurrOffset += lerp(Offset0, Offset1, GuideVertexLerps[K]);
 			}
 			#endif
-			OutRenderDeformedPositionBuffer[VertexIndex] = float4(f32tof16(ControlPoint + CurrOffset - OutHairPositionOffset), RenderRestPosePositionBuffer[VertexIndex].w);
+
+			const bool bIsValid = ScaleAndClip(VertexIndex, PackedRestPosition.w);
+			WorldVertexPosition = bIsValid ? ControlPoint + CurrOffset : float3(1.#INF, 1.#INF, 1.#INF);
+
+			OutRenderDeformedPositionBuffer[VertexIndex] = float4(f32tof16(WorldVertexPosition - OutHairPositionOffset), PackedRestPosition.w);
 		}
 		#endif	
 
@@ -145,7 +201,10 @@ void MainCS(uint2 DispatchThreadId : SV_DispatchThreadID)
 		// Apply dynamic mesh deformation (translation / rotation)
 		#if PERMUTATION_DYNAMIC_GEOMETRY == 1
 		{
-			float3 ControlPoint = f16tof32(RenderRestPosePositionBuffer[VertexIndex].xyz) + InRenderHairPositionOffset;
+			uint4 PackedRestPosition = RenderRestPosePositionBuffer[VertexIndex];
+			const float3 RestControlPoint = f16tof32(PackedRestPosition.xyz) + InRenderHairPositionOffset;
+			float3 ControlPoint = RestControlPoint;
+
 			float3 CurrOffset = 0;
 
 			// Compute the simulation offset in hair local space (i.e., triangle)
@@ -177,32 +236,55 @@ void MainCS(uint2 DispatchThreadId : SV_DispatchThreadID)
 			ControlPoint = ToTriangle(ControlPoint,RootBarycentric,RestTriangle);
 			ControlPoint = ToWorld(ControlPoint,RootBarycentric,DeformedTriangle) + CurrOffset;
 
-			OutRenderDeformedPositionBuffer[VertexIndex] = float4(f32tof16(ControlPoint - OutHairPositionOffset), RenderRestPosePositionBuffer[VertexIndex].w);
+			const bool bIsValid = ScaleAndClip(VertexIndex, PackedRestPosition.w);
+			WorldVertexPosition = bIsValid ? ControlPoint : float3(1.#INF, 1.#INF, 1.#INF);
 
-			// Debug
-			#if 0
+			#if PERMUTATION_DEBUG == 1
 			{
-				// Show the triangle normal use by hair fiber
-				//ControlPoint = VertexIndex % 2 == 0 ? RestTriangle.P0 : RestTriangle.P0 + RestTriangle.N;
+				const bool bDebugEnabled = HairDebugMode == 4;
+				const float4 RestColor = float4(0, 1, 0, 1);
+				const float4 DeformedColor = float4(1, 0, 0, 1);
+				if (bDebugEnabled && VertexIndex==0)
+				{	
+					AddCross(InRenderHairPositionOffset, 1.2, RestColor);
+					AddCross(OutHairPositionOffset, 1, DeformedColor);
+				}
+			
+				if (bDebugEnabled && VertexIndex >= 0 && VertexIndex < 64)
+				{
+					AddReferential(RestControlPoint, float3(0, 0, 1));
+					AddReferential(WorldVertexPosition, float3(0, 0, 1));
+				}
 
-				// Measure the stretch of the projection (to identify if there any local transform issue during the projeciton steps
-				ControlPoint = VertexIndex % 2 == 0 ? ControlPoint  : RestTriangle.P0;
+				if (bDebugEnabled)
+				{
+					const uint ControlPointType = asuint(PackedRestPosition.w) & 0x3;
+					const bool bIsRoot = ControlPointType == 1;
+					if (bIsRoot)
+					{
+						AddLineTriangle(RestTriangle.P0, RestTriangle.P1, RestTriangle.P2, RestColor);
+						AddLineTriangle(DeformedTriangle.P0, DeformedTriangle.P1, DeformedTriangle.P2, DeformedColor);
+					}
+				}
 			}
 			#endif
+
+			OutRenderDeformedPositionBuffer[VertexIndex] = float4(f32tof16(WorldVertexPosition - OutHairPositionOffset), RenderRestPosePositionBuffer[VertexIndex].w);
 		}
 		#endif	
 
 		#if PERMUTATION_DEBUG == 1
+		if (HairDebugMode == 1)
 		{
 			const uint GuideIndex = GuideVertexIndices[0]; // Take the closest guide
 			const float4 GuideAttribute = SimAttributeBuffer[GuideIndex];
 			OutRenderAttributeBuffer[VertexIndex] = GuideAttribute;
 		}
-		#endif
 
 		// Cull hair strands or clusters which are not matching the provided index
-		#if PERMUTATION_DEBUG == 2
+		if (HairDebugMode == 2)
 		{
+		#if 0
 			bool bIsValid = false;
 			if (HairStrandsCullIndex.y >= 0)
 			{
@@ -211,7 +293,7 @@ void MainCS(uint2 DispatchThreadId : SV_DispatchThreadID)
 				{
 					const uint GuideIndex = GuideVertexIndices[K];
 					const uint RootIndex = SimVertexToRootIndexBuffer[GuideIndex];
-					if (HairStrandsCullIndex.y == RootIndex)
+					if (uint(HairStrandsCullIndex.y) == RootIndex)
 					{
 						bIsValid = true;
 						break;
@@ -221,14 +303,159 @@ void MainCS(uint2 DispatchThreadId : SV_DispatchThreadID)
 			else if (HairStrandsCullIndex.x >= 0)
 			{
 				const uint RootIndex = RenVertexToRootIndexBuffer[VertexIndex];
-				bIsValid = HairStrandsCullIndex.x == RootIndex;
+				bIsValid = uint(HairStrandsCullIndex.x) == RootIndex;
 			}
 
 			if (!bIsValid)
 			{
 				OutRenderDeformedPositionBuffer[VertexIndex] = float4(float3(1.#INF, 1.#INF, 1.#INF), 2);
 			}
+		#endif		
+		}
+
+		if (HairDebugMode == 3)
+		{
+			const uint ClusterId = VertexToClusterIdBuffer[VertexIndex];
+			const uint ClusterIdLo8bits = ClusterId & 0xFF;
+			const uint ClusterIdHi8btis = (ClusterId>>8) & 0xFF;
+			// Using FNV1 hash to break linearity of ClusterId (generated by linearly parsing cluster grid)
+			uint Hash = 0;
+			Hash = Hash * 17;
+			Hash = Hash ^ ClusterIdLo8bits;
+			Hash = Hash * 17;
+			Hash = Hash ^ ClusterIdHi8btis;
+			OutRenderAttributeBuffer[VertexIndex] = asfloat(uint4(0, (Hash & 0xFF) << 8, 0, 0));
 		}
 		#endif
 	}
-}
\ No newline at end of file
+}
+
+///////////////////////////////////////////////////////////////////////////////////////////////////
+
+uint ClusterCount;
+
+Buffer<uint4> RenderDeformedPositionBuffer;
+Buffer<uint> ClusterVertexIdBuffer;
+Buffer<uint4> ClusterInfoBuffer;		// Indices start and count
+
+RWBuffer<int> OutClusterAABBBuffer;		// Cluster data packed as {uint3 AABBMin, uint3 AABBMax}
+RWBuffer<int> OutGroupAABBBuffer;		// Group data packed as {uint3 AABBMin, uint3 AABBMax}
+
+#if GROUP_SIZE!=64 && GROUP_SIZE!=32
+#error Unhandled group size
+#endif
+
+groupshared float3 SharedClusterAABBMin[GROUP_SIZE];
+groupshared float3 SharedClusterAABBMax[GROUP_SIZE];
+
+[numthreads(GROUP_SIZE, 1, 1)]
+void ClusterAABBEvaluationCS(uint3 DispatchThreadId : SV_DispatchThreadID, uint3 GroupId : SV_GroupID, uint3 GroupThreadId : SV_GroupThreadID)
+{
+	const uint ClusterId = GroupId.x + GroupId.y * DispatchCountX;
+	if (ClusterId >= ClusterCount)
+	{
+		return;
+	}
+
+	int4 ClusterInfo = ClusterInfoBuffer[ClusterId];
+	const uint ClusterIndicesStart = ClusterInfo.x;
+	const uint ClusterIndicesCount = ClusterInfo.y;
+
+	int3 ClusterAABBMin = 9999999.0f;
+	int3 ClusterAABBMax =-9999999.0f;
+
+	uint LoopCount = (ClusterIndicesCount + GROUP_SIZE - 1) / GROUP_SIZE;
+	for (uint Loop = 0; Loop < LoopCount; ++Loop)
+	{
+		uint GroupReadOffset = Loop * GROUP_SIZE;
+		uint ThreadReadOffset = GroupReadOffset + GroupThreadId.x;
+
+		if (ThreadReadOffset < ClusterIndicesCount)
+		{
+			uint VertexIndex = ClusterVertexIdBuffer[ClusterIndicesStart + ThreadReadOffset];
+
+			uint4 PackedVertex = RenderDeformedPositionBuffer[VertexIndex];
+			float3 VertexWorldPosition = f16tof32(PackedVertex.xyz);
+			if (IsFinite(VertexWorldPosition.x))
+			{
+				const int3 WorldSpaceCentimeters = int3(mul(float4(VertexWorldPosition + OutHairPositionOffset, 1.0f), LocalToWorldMatrix).xyz);
+				const int3 WorldSpaceCentimetersMin = WorldSpaceCentimeters - 1;
+				const int3 WorldSpaceCentimetersMax = WorldSpaceCentimeters + 1;
+
+				ClusterAABBMin = min(ClusterAABBMin, WorldSpaceCentimetersMin);
+				ClusterAABBMax = max(ClusterAABBMax, WorldSpaceCentimetersMax);
+			}
+		}
+	}
+
+	// Write each thread result to shared memory
+	const uint LocalIndex = GroupThreadId.x;
+	SharedClusterAABBMin[LocalIndex] = ClusterAABBMin;
+	SharedClusterAABBMax[LocalIndex] = ClusterAABBMax;
+
+	// Do a local reduction in shared memory. Assumes ClusterIndicesCount>64 to have all min&max values valid.
+#if GROUP_SIZE==64
+	GroupMemoryBarrierWithGroupSync();
+	if (LocalIndex < 32)
+	{
+		SharedClusterAABBMin[LocalIndex] = min(SharedClusterAABBMin[LocalIndex], SharedClusterAABBMin[LocalIndex + 32]);
+		SharedClusterAABBMax[LocalIndex] = max(SharedClusterAABBMax[LocalIndex], SharedClusterAABBMax[LocalIndex + 32]);
+	}
+#endif
+	GroupMemoryBarrierWithGroupSync();
+	if (LocalIndex < 16)
+	{
+		SharedClusterAABBMin[LocalIndex] = min(SharedClusterAABBMin[LocalIndex], SharedClusterAABBMin[LocalIndex + 16]);
+		SharedClusterAABBMax[LocalIndex] = max(SharedClusterAABBMax[LocalIndex], SharedClusterAABBMax[LocalIndex + 16]);
+	}
+	GroupMemoryBarrierWithGroupSync();
+	if (LocalIndex < 8)
+	{
+		SharedClusterAABBMin[LocalIndex] = min(SharedClusterAABBMin[LocalIndex], SharedClusterAABBMin[LocalIndex + 8]);
+		SharedClusterAABBMax[LocalIndex] = max(SharedClusterAABBMax[LocalIndex], SharedClusterAABBMax[LocalIndex + 8]);
+	}
+	// No hardware has SIMD Vector unit operating in sync with less than 16 threads per group. So can skip sync now.
+	//GroupMemoryBarrierWithGroupSync();
+	if (LocalIndex < 4)
+	{
+		SharedClusterAABBMin[LocalIndex] = min(SharedClusterAABBMin[LocalIndex], SharedClusterAABBMin[LocalIndex + 4]);
+		SharedClusterAABBMax[LocalIndex] = max(SharedClusterAABBMax[LocalIndex], SharedClusterAABBMax[LocalIndex + 4]);
+	}
+	//GroupMemoryBarrierWithGroupSync();
+	if (LocalIndex < 2)
+	{
+		SharedClusterAABBMin[LocalIndex] = min(SharedClusterAABBMin[LocalIndex], SharedClusterAABBMin[LocalIndex + 2]);
+		SharedClusterAABBMax[LocalIndex] = max(SharedClusterAABBMax[LocalIndex], SharedClusterAABBMax[LocalIndex + 2]);
+	}
+	//GroupMemoryBarrierWithGroupSync();
+	if (LocalIndex < 1)
+	{
+		SharedClusterAABBMin[LocalIndex] = min(SharedClusterAABBMin[LocalIndex], SharedClusterAABBMin[LocalIndex + 1]);
+		SharedClusterAABBMax[LocalIndex] = max(SharedClusterAABBMax[LocalIndex], SharedClusterAABBMax[LocalIndex + 1]);
+	}
+
+	// Write out hair cluster AABB
+	if (LocalIndex == 0)
+	{
+		const uint ClusterIdx6 = ClusterId * 6;
+		OutClusterAABBBuffer[ClusterIdx6 + 0] = SharedClusterAABBMin[0].x;
+		OutClusterAABBBuffer[ClusterIdx6 + 1] = SharedClusterAABBMin[0].y;
+		OutClusterAABBBuffer[ClusterIdx6 + 2] = SharedClusterAABBMin[0].z;
+		OutClusterAABBBuffer[ClusterIdx6 + 3] = SharedClusterAABBMax[0].x;
+		OutClusterAABBBuffer[ClusterIdx6 + 4] = SharedClusterAABBMax[0].y;
+		OutClusterAABBBuffer[ClusterIdx6 + 5] = SharedClusterAABBMax[0].z;
+
+		// And contribute to the group full AABB
+		FHairAABB Bound;
+		Bound.Min.x = SharedClusterAABBMin[0].x;
+		Bound.Min.y = SharedClusterAABBMin[0].y;
+		Bound.Min.z = SharedClusterAABBMin[0].z;
+		Bound.Max.x = SharedClusterAABBMax[0].x;
+		Bound.Max.y = SharedClusterAABBMax[0].y;
+		Bound.Max.z = SharedClusterAABBMax[0].z;
+		InterlockHairAABB(0, Bound, OutGroupAABBBuffer);
+	}
+
+	// TODO Atomic Min/Max group AABB
+}
+
diff --git a/Engine/Shaders/Private/HairStrands/HairStrandsMaterialPS.usf b/Engine/Shaders/Private/HairStrands/HairStrandsMaterialPS.usf
--- a/Engine/Shaders/Private/HairStrands/HairStrandsMaterialPS.usf
+++ b/Engine/Shaders/Private/HairStrands/HairStrandsMaterialPS.usf
@@ -11,7 +11,7 @@
 #include "/Engine/Generated/VertexFactory.ush" 
 
 #define SUPPORT_MATERIAL_PROPERTY 1
-uint MaterialPass_ClusterId;
+uint MaterialPass_MacroGroupId;
 uint MaterialPass_MaterialId;  
 uint MaterialPass_PrimitiveId;
 RWStructuredBuffer<FNodeData>	OutNodeData;
@@ -167,7 +167,7 @@ FEvaluateOutput Evaluate(
 	FVisibilityData OutVisibilityData = (FVisibilityData)0;	
 	OutVisibilityData.Coverage8bit = SampleCoverage8bit;
 	OutVisibilityData.PrimitiveId = SampleCoverage8bit;
-	OutVisibilityData.ClusterId = MaterialPass_ClusterId;
+	OutVisibilityData.MacroGroupId = MaterialPass_MacroGroupId;
 
 	half3 OutBaseColor = 0;
 	float OutRoughness = 0;
@@ -191,7 +191,7 @@ FEvaluateOutput Evaluate(
 
 		//OutVisibilityData.Tangent = normalize(Interpolants.TangentToWorld2.xyz);
 		//OutVisibilityData.PrimitiveId = Interpolants.HairPrimitiveId;
-		//OutVisibilityData.ClusterId = HairVisibilityPass_HairClusterIndex;
+		//OutVisibilityData.MacroGroupId = HairVisibilityPass_HairMacroGroupIndex;
 		//OutVisibilityData.Coverage8bit = To8bitCoverage(Coverage);
 
 		// Material computation
@@ -313,7 +313,7 @@ void Main(
 	VisibilityData.Tangent = float3(0, 0, 1);
 	VisibilityData.PrimitiveId = 0;
   #endif
-	VisibilityData.ClusterId = HairVisibilityPass_HairClusterIndex;
+	VisibilityData.MacroGroupId = HairVisibilityPass_HairMacroGroupIndex;
 	VisibilityData.Coverage8bit = To8bitCoverage(Coverage);
 
 	half3 BaseColor = 0;
diff --git a/Engine/Shaders/Private/HairStrands/HairStrandsMeshProjection.usf b/Engine/Shaders/Private/HairStrands/HairStrandsMeshProjection.usf
--- a/Engine/Shaders/Private/HairStrands/HairStrandsMeshProjection.usf
+++ b/Engine/Shaders/Private/HairStrands/HairStrandsMeshProjection.usf
@@ -126,7 +126,7 @@ FTrianglePoint ComputeClosestPoint(const FHairMeshTriangle Tri, const float3 P)
 	return Out;
 }
 
-[numthreads(32, 1, 1)]
+[numthreads(128, 1, 1)]
 void MainCS(uint3 DispatchThreadId : SV_DispatchThreadID)
 {
 	// Shift root index with thread index to reduce/avoid collisions
diff --git a/Engine/Shaders/Private/HairStrands/HairStrandsMeshUpdate.usf b/Engine/Shaders/Private/HairStrands/HairStrandsMeshUpdate.usf
--- a/Engine/Shaders/Private/HairStrands/HairStrandsMeshUpdate.usf
+++ b/Engine/Shaders/Private/HairStrands/HairStrandsMeshUpdate.usf
@@ -18,7 +18,7 @@ RWBuffer<float4> OutRootTrianglePosition0;
 RWBuffer<float4> OutRootTrianglePosition1;
 RWBuffer<float4> OutRootTrianglePosition2;
 
-[numthreads(32, 1, 1)]
+[numthreads(128, 1, 1)]
 void MainCS(uint3 DispatchThreadId : SV_DispatchThreadID)
 { 
 	const uint RootIndex = DispatchThreadId.x;
diff --git a/Engine/Shaders/Private/HairStrands/HairStrandsVertexFactory.ush b/Engine/Shaders/Private/HairStrands/HairStrandsVertexFactory.ush
--- a/Engine/Shaders/Private/HairStrands/HairStrandsVertexFactory.ush
+++ b/Engine/Shaders/Private/HairStrands/HairStrandsVertexFactory.ush
@@ -16,6 +16,11 @@
 	Buffer<float4> VertexFetch_InstanceTransformBuffer;
 #endif
 
+#ifdef USE_CULLED_CLUSTER
+	Buffer<uint> CulledVertexIdsBuffer;
+	Buffer<float> CulledVertexRadiusScaleBuffer;
+#endif
+
 ////////////////////////////////////////////////////////////////////////////////
 // VF_STRAND_HAIR
 // Define used in certain shader like DeepShadow for running specific code (coverage computation for instance) 
@@ -28,6 +33,7 @@ float3 HairStrandsVF_PreviousPositionOffset;	// Vertex position (3x16bits) are r
 float  HairStrandsVF_Density;
 float  HairStrandsVF_Radius;
 float  HairStrandsVF_Length;
+uint   HairStrandsVF_CullingEnable;
 
 Buffer<uint4>  HairStrandsVF_PositionBuffer;
 Buffer<uint4>  HairStrandsVF_PreviousPositionBuffer;
@@ -156,6 +162,44 @@ struct FVertexFactoryInput
 	uint VertexId : SV_VertexID;
 };
 
+struct FVertexInfo
+{
+	uint QuadIndex;
+	uint BaseIndex;
+	uint VertexIndex;
+	uint HairPrimitiveId;
+	uint IsTip;				// The vertex is on the quad side toward the tip of the strand.
+	uint IsLeft;
+	float LodRadiusScale;
+	bool bForceInvalidQuad;
+};
+FVertexInfo GetVertexInfo(FVertexFactoryInput Input)
+{
+	uint VertexId = Input.VertexId;
+
+	FVertexInfo VertexInfo;
+	VertexInfo.QuadIndex	= VertexId % 6;
+	VertexInfo.BaseIndex	= VertexId / 6;
+	VertexInfo.IsTip		= VertexInfo.QuadIndex == 0 || VertexInfo.QuadIndex == 2 || VertexInfo.QuadIndex == 4 ? 0 : 1;
+	VertexInfo.IsLeft		= VertexInfo.QuadIndex == 0 || VertexInfo.QuadIndex == 1 || VertexInfo.QuadIndex == 5 ? 1 : 0;
+
+	VertexInfo.HairPrimitiveId = VertexInfo.BaseIndex;
+	VertexInfo.VertexIndex = VertexInfo.BaseIndex + VertexInfo.IsTip;
+	VertexInfo.LodRadiusScale = 1.0f;
+	VertexInfo.bForceInvalidQuad = false;
+#ifdef USE_CULLED_CLUSTER
+	if (HairStrandsVF_CullingEnable)
+	{
+		VertexInfo.VertexIndex = CulledVertexIdsBuffer[VertexInfo.BaseIndex + VertexInfo.IsTip];
+		VertexInfo.LodRadiusScale = CulledVertexRadiusScaleBuffer[VertexInfo.BaseIndex + VertexInfo.IsTip];
+		VertexInfo.HairPrimitiveId = VertexInfo.VertexIndex - VertexInfo.IsTip;
+		VertexInfo.bForceInvalidQuad = VertexInfo.LodRadiusScale <= 0.0f;
+	}
+#endif
+
+	return VertexInfo;
+}
+
 #if RAYHITGROUPSHADER
 FVertexFactoryInput LoadVertexFactoryInput(uint TriangleIndex, int VertexIndex)
 {
@@ -181,7 +225,7 @@ struct FVertexFactoryIntermediates
 	half3x3 TangentToWorld;
 	half TangentToWorldSign;
 	uint PrimitiveId;
-	uint   HairPrimitiveId;
+	uint HairPrimitiveId;
 	float2 HairPrimitiveUV;
 	float2 HairDimensions;  // This is kept on the intermediate as this is used in several vertex shader for the actual coverage.
 	float  HairDensity;
@@ -198,7 +242,7 @@ uint GetPrimitiveId(FVertexFactoryInput Input)
 
 uint GetHairPrimitiveId(FVertexFactoryInput Input)
 {
-	return Input.VertexId / 6;
+	return GetVertexInfo(Input).HairPrimitiveId;
 }
 
 struct FControlPoint
@@ -210,22 +254,20 @@ struct FControlPoint
 
 FControlPoint GetVertexPosition(FVertexFactoryInput Input)
 {
-	const uint QuadIndex	= Input.VertexId % 6;
-	const uint BaseIndex	= Input.VertexId / 6;
-	const uint OffsetIndex	= QuadIndex == 0 || QuadIndex == 2 || QuadIndex == 4 ? 0 : 1;
-	const uint4 PackedData = HairStrandsVF_PositionBuffer[BaseIndex + OffsetIndex];
+	FVertexInfo VertexInfo = GetVertexInfo(Input);
+	const uint4 PackedData = HairStrandsVF_PositionBuffer[VertexInfo.VertexIndex];
 	
 	FControlPoint Output = (FControlPoint)0;
 
 	const float3 ControlPoint	= f16tof32(PackedData.xyz) + HairStrandsVF_PositionOffset;
 	const uint PackedAlpha		= asuint(PackedData.a);
 	const uint ControlPointType = PackedAlpha & 0x3;
-	Output.WorldRadius			= ((PackedAlpha & 0x00FC) >> 2) / 63.f * HairStrandsVF_Radius;
+	Output.WorldRadius			= ((PackedAlpha & 0x00FC) >> 2) / 63.f * HairStrandsVF_Radius * VertexInfo.LodRadiusScale;
 	Output.WorldLength			= ((PackedAlpha & 0xFF00) >> 8) / 255.f * HairStrandsVF_Length;
 	
 	// Create a degenerated quad the end of each strand to cut between each strands
 	// #hair_todo: This is not efficient for short strand like fur for instance. Need to revisit that at some point
-	const bool bIsInvalidQuad = (ControlPointType == 1 && OffsetIndex == 1) || (ControlPointType == 2 && OffsetIndex == 0);
+	const bool bIsInvalidQuad = (ControlPointType == 1 && VertexInfo.IsTip == 1) || (ControlPointType == 2 && VertexInfo.IsTip == 0) || VertexInfo.bForceInvalidQuad;
 	Output.Position = bIsInvalidQuad ? float3(1.#INF, 1.#INF, 1.#INF) : ControlPoint;
 
 	return Output;
@@ -243,10 +285,8 @@ float GetWorldStrandLength(FVertexFactoryInput Input)
 
 float3 GetVertexPreviousPosition(FVertexFactoryInput Input)
 {
-	const uint QuadIndex = Input.VertexId % 6;
-	const uint BaseIndex = Input.VertexId / 6;
-	const uint OffsetIndex = QuadIndex == 0 || QuadIndex == 2 || QuadIndex == 4 ? 0 : 1;
-	const uint4 PackedData = HairStrandsVF_PreviousPositionBuffer[BaseIndex + OffsetIndex];
+	FVertexInfo VertexInfo = GetVertexInfo(Input);
+	const uint4 PackedData = HairStrandsVF_PreviousPositionBuffer[VertexInfo.VertexIndex];
 
 	const float3 ControlPoint = f16tof32(PackedData.xyz) + HairStrandsVF_PreviousPositionOffset;
 	const uint PackedAlpha = asuint(PackedData.a);
@@ -254,7 +294,7 @@ float3 GetVertexPreviousPosition(FVertexFactoryInput Input)
 
 	// Create a degenerated quad the end of each strand to cut between each strands
 	// #hair_todo: This is not efficient for short strand like fur for instance. Need to revisit that at some point
-	const bool bIsInvalidQuad = (ControlPointType == 1 && OffsetIndex == 1) || (ControlPointType == 2 && OffsetIndex == 0);
+	const bool bIsInvalidQuad = (ControlPointType == 1 && VertexInfo.IsTip == 1) || (ControlPointType == 2 && VertexInfo.IsTip == 0) || VertexInfo.bForceInvalidQuad;
 	const float3 Position = bIsInvalidQuad ? float3(1.#INF, 1.#INF, 1.#INF) : ControlPoint;
 
 	return Position;
@@ -263,9 +303,9 @@ float3 GetVertexPreviousPosition(FVertexFactoryInput Input)
 // Segment UV coord of an hair segment. This is different from the UV coord of the hair strands
 float2 GetSegmentUVCoord(FVertexFactoryInput Input)
 {
-	const uint QuadIndex = Input.VertexId % 6;
-	const float VCoord = QuadIndex == 0 || QuadIndex == 1 || QuadIndex == 5 ? 0.0f : 1.0f;
-	const float UCoord = QuadIndex == 0 || QuadIndex == 2 || QuadIndex == 4 ? 0.0f : 1.0f;
+	FVertexInfo VertexInfo = GetVertexInfo(Input);
+	const float VCoord = VertexInfo.IsLeft ? 0.0f : 1.0f;
+	const float UCoord = VertexInfo.IsTip  ? 1.0f : 0.0f;
 	return float2(UCoord, VCoord);
 }
 
@@ -332,11 +372,9 @@ float4 CalcWorldPosition(float3 Position, uint PrimitiveId)
 
 half3x3 CalcTangentToLocal(FVertexFactoryInput Input, inout float TangentSign)
 {
-	const uint QuadIndex = Input.VertexId % 6;
-	const uint BaseIndex = Input.VertexId / 6;
-	const uint OffsetIndex = QuadIndex == 0 || QuadIndex == 2 || QuadIndex == 4 ? 0 : 1;
-	const half3 TangentInputX = HairStrandsVF_TangentBuffer[(BaseIndex + OffsetIndex) * 2 + 0].xyz;
-	const half4 TangentInputZ = HairStrandsVF_TangentBuffer[(BaseIndex + OffsetIndex) * 2 + 1];
+	FVertexInfo VertexInfo = GetVertexInfo(Input);
+	const half3 TangentInputX = HairStrandsVF_TangentBuffer[VertexInfo.VertexIndex * 2 + 0].xyz;
+	const half4 TangentInputZ = HairStrandsVF_TangentBuffer[VertexInfo.VertexIndex * 2 + 1];
 
 	half3 TangentX = TangentBias(TangentInputX);
 	half4 TangentZ = TangentBias(TangentInputZ);
@@ -407,26 +445,22 @@ half3x3 VertexFactoryGetTangentToLocal( FVertexFactoryInput Input, FVertexFactor
 	return Intermediates.TangentToLocal;
 }
 
-float4 ComputeViewAlignedWorldPosition(FVertexFactoryInput Input, float3 WorldTangent, float4 WorldPosition, float WorldStrandRadius, float VelocityScale)
+float4 ComputeViewAlignedWorldPosition(FVertexFactoryInput Input, float3 WorldTangent, float4 WorldPosition, float WorldStrandRadius, FHairViewInfo HairViewInfo)
 {
-	const uint LocalVertexId = Input.VertexId % 6;
-	const bool bIsLeft = LocalVertexId == 0 || LocalVertexId == 1 || LocalVertexId == 5;
-
-	FHairRenderInfo HairRenderInfo = GetHairRenderInfo(ResolvedView.HairRenderInfo);
+	FVertexInfo VertexInfo = GetVertexInfo(Input);
 
 	// Minimal radius to snap the strand to a sample/pixel center (to avoid aliasing)
-	const float RadiusAtDepth1 = lerp(HairRenderInfo.RadiusAtDepth1Primary, HairRenderInfo.RadiusAtDepth1Velocity, VelocityScale);
-	const float DistanceToCamera = length(ResolvedView.TranslatedWorldCameraOrigin - WorldPosition.xyz);
-	const float MinStrandHairRadius = HairRenderInfo.bIsOrthoView ? RadiusAtDepth1 : (DistanceToCamera * RadiusAtDepth1);
-	const float3 ViewDir = -ResolvedView.ViewForward;
+	const float DistanceToCamera = length(HairViewInfo.TranslatedWorldCameraOrigin - WorldPosition.xyz);
+	const float MinStrandHairRadius = HairViewInfo.bIsOrthoView ? HairViewInfo.RadiusAtDepth1 : (DistanceToCamera * HairViewInfo.RadiusAtDepth1);
+	const float3 ViewDir = -HairViewInfo.ViewForward;
 	const float3 Right = normalize(cross(WorldTangent, ViewDir));
-	const float3 OutWorldPosition = WorldPosition.xyz + (bIsLeft ? -Right : Right) * max(WorldStrandRadius, MinStrandHairRadius);
+	const float3 OutWorldPosition = WorldPosition.xyz + (VertexInfo.IsLeft ? -Right : Right) * max(WorldStrandRadius * VertexInfo.LodRadiusScale, MinStrandHairRadius);
 	
 	return float4(OutWorldPosition, 1);
 }
 
 // @return translated world position
-float4 VertexFactoryGetWorldPosition(FVertexFactoryInput Input, FVertexFactoryIntermediates Intermediates, float VelocityScale)
+float4 VertexFactoryGetWorldPosition(FVertexFactoryInput Input, FVertexFactoryIntermediates Intermediates, FHairViewInfo HairViewInfo)
 {
 	FControlPoint ControlPoint = GetVertexPosition(Input);
 	const float3 VertexPosition = ControlPoint.Position;
@@ -434,12 +468,21 @@ float4 VertexFactoryGetWorldPosition(FVertexFactoryInput Input, FVertexFactoryIn
 
 	// Hair shader expect the WorldNormal to be the tangent vector
 	const float3 WorldTangent = GetTangent(Input);
-	return ComputeViewAlignedWorldPosition(Input, WorldTangent, WorldPosition, ControlPoint.WorldRadius, VelocityScale);
+	return ComputeViewAlignedWorldPosition(Input, WorldTangent, WorldPosition, ControlPoint.WorldRadius, HairViewInfo);
 }
 
+// This function is referenced by several "system" shaders, but none of these shaders (apart HitProxy, are actually used)
 float4 VertexFactoryGetWorldPosition(FVertexFactoryInput Input, FVertexFactoryIntermediates Intermediates)
 {
-	return VertexFactoryGetWorldPosition(Input, Intermediates, 0);
+	const FHairRenderInfo HairRenderInfo = GetHairRenderInfo(ResolvedView.HairRenderInfo);
+
+	FHairViewInfo HairViewInfo;
+	HairViewInfo.TranslatedWorldCameraOrigin = ResolvedView.TranslatedWorldCameraOrigin;
+	HairViewInfo.ViewForward = ResolvedView.ViewForward;
+	HairViewInfo.RadiusAtDepth1 = HairRenderInfo.RadiusAtDepth1Primary;
+	HairViewInfo.bIsOrthoView = HairRenderInfo.bIsOrthoView;
+
+	return VertexFactoryGetWorldPosition(Input, Intermediates, HairViewInfo);
 }
 
 float4 VertexFactoryGetRasterizedWorldPosition(FVertexFactoryInput Input, FVertexFactoryIntermediates Intermediates, float4 InWorldPosition)
@@ -473,16 +516,10 @@ FVertexFactoryInterpolantsVSToPS VertexFactoryGetInterpolantsVSToPS(FVertexFacto
 	return Interpolants;
 }
 
-/** for depth-only pass */
+/** for depth-only pass (Not used by the actual hair shaders)*/
 float4 VertexFactoryGetWorldPosition(FPositionOnlyVertexFactoryInput Input)
 {
-	const FControlPoint ControlPoint = GetVertexPosition(Input);
-	const float3 Position = ControlPoint.Position;
-	const uint PrimitiveId = GetPrimitiveId(Input);
-
-	const float4 WorldPosition = CalcWorldPosition(Position, PrimitiveId);
-    const float3 WorldTangent = GetTangent(Input);
-	return ComputeViewAlignedWorldPosition(Input, WorldTangent, WorldPosition, ControlPoint.WorldRadius, 0);
+	return 0;
 }
 
 // @return translated world position (without quad extension/reorientation).This is used only for velocity computation
diff --git a/Engine/Shaders/Private/HairStrands/HairStrandsVisibilityCommon.ush b/Engine/Shaders/Private/HairStrands/HairStrandsVisibilityCommon.ush
--- a/Engine/Shaders/Private/HairStrands/HairStrandsVisibilityCommon.ush
+++ b/Engine/Shaders/Private/HairStrands/HairStrandsVisibilityCommon.ush
@@ -8,12 +8,14 @@
 #define RENDER_MODE_MSAA_VISIBILITY 3
 #define RENDER_MODE_TRANSMITTANCE_AND_HAIRCOUNT 4
 
+#define INVALID_TILE_OFFSET ~0
+
 struct FVisibilityData
 {
 	float3 Tangent;		// 24 bits
 	uint   Coverage8bit;//  8 bits
 	uint   PrimitiveId;	// 28 bits
-	uint   ClusterId;	//  4 bits
+	uint   MacroGroupId;//  4 bits
 };
 
 struct FNodeVis // 128 bit aligments. TODO: Move the coverage computation into the material evaluation for saving 64 bits? 
@@ -32,7 +34,7 @@ struct FNodeVis // 128 bit aligments. TODO: Move the coverage computation into t
 struct FNodeData
 {
 	float Depth;
-	uint  PrimitiveID_ClusterID;
+	uint  PrimitiveID_MacroGroupID;
 	uint  Tangent_Coverage8bit;
 	uint  BaseColor_Roughness;
 	float Specular;
@@ -76,12 +78,12 @@ uint4 EncodeVisibilityData(FVisibilityData Data)
 		(0xFF & uint(T.x * 255)) | ((0xFF & uint(T.y * 255)) << 8),
 		(0xFF & uint(T.z * 255)) | (min(0xFFu, Data.Coverage8bit) << 8));
 
-	// PrimitiveId and ClusterId
+	// PrimitiveId and MacroGroupId
 	// Encode into 2 x 16bits
-	// PrimitiveId is on 28bits | ClusterId is on 4bits
+	// PrimitiveId is on 28bits | MacroGroupId is on 4bits
 	Output.zw = uint2(
 		0xFFFF & Data.PrimitiveId,
-		((Data.ClusterId & 0x0F) << 12) | (0x0FFF & (uint(Data.PrimitiveId) >> 16)));
+		((Data.MacroGroupId & 0x0F) << 12) | (0x0FFF & (uint(Data.PrimitiveId) >> 16)));
 
 	return Output;
 }
@@ -96,7 +98,7 @@ FVisibilityData DecodeVisibilityData(uint4 EncodedData)
 	Output.Coverage8bit = (EncodedData.y >> 8) & 0xFF;
 
 	Output.PrimitiveId = EncodedData.z | ((EncodedData.w & 0x0FFF) << 16);
-	Output.ClusterId = (EncodedData.w & 0xF000) >> 12;
+	Output.MacroGroupId = (EncodedData.w & 0xF000) >> 12;
 
 	return Output;
 }
@@ -107,7 +109,7 @@ bool IsValidHairData(uint4 EncodedVisibilityData)
 	return EncodedVisibilityData.w != InvalidSample;
 }
 
-uint DecodeClusterId(uint4 EncodedVisibilityData)
+uint DecodeMacroGroupId(uint4 EncodedVisibilityData)
 {
 	return (EncodedVisibilityData.w >> 12) & 0xF;
 }
@@ -142,7 +144,7 @@ FNodeData ConvertToNodeData(float Depth, FVisibilityData VisibilityData, float3
 	FNodeData Out;
 	Out.Depth					= Depth;																			// 32bits float
 	Out.Tangent_Coverage8bit	= (0xFFFF & EncodedVisibilityData.x) | ((0xFFFF & EncodedVisibilityData.y) << 16);	// 32bits uint
-	Out.PrimitiveID_ClusterID	= (0xFFFF & EncodedVisibilityData.z) | ((0xFFFF & EncodedVisibilityData.w) << 16);	// 32bits uint
+	Out.PrimitiveID_MacroGroupID= (0xFFFF & EncodedVisibilityData.z) | ((0xFFFF & EncodedVisibilityData.w) << 16);	// 32bits uint
 	Out.BaseColor_Roughness		= float4ToUint(float4(sqrt(BaseColor), Roughness));									// 32bits uint
 	Out.Specular				= Specular;																			// 32bits float
 
@@ -162,8 +164,8 @@ FNodeData ConvertToNodeData(float Depth, FVisibilityData VisibilityData, float3
 	const uint4 EncodedVisibility = uint4(
 		(InNode.Tangent_Coverage8bit & 0xFFFF),
 		(InNode.Tangent_Coverage8bit >> 16) & 0xFFFF,
-		(InNode.PrimitiveID_ClusterID & 0xFFFF),
-		(InNode.PrimitiveID_ClusterID >> 16) & 0xFFFF);
+		(InNode.PrimitiveID_MacroGroupID & 0xFFFF),
+		(InNode.PrimitiveID_MacroGroupID >> 16) & 0xFFFF);
 	OutVisibilityData = DecodeVisibilityData(EncodedVisibility);
 
 	const float4 BaseColorAndRoughness = UintToFloat4(InNode.BaseColor_Roughness);
@@ -178,7 +180,7 @@ FNodeData ConvertToNodeData(float Depth, FVisibilityData VisibilityData, float3
  struct FSampleDesc
  {
 	 uint   SampleId;
-	 uint   ClusterId;
+	 uint   MacroGroupId;
 	 float  Depths;
 	 float  Weights;
 	 float3 Tangents;
@@ -204,10 +206,14 @@ FNodeData ConvertToNodeData(float Depth, FVisibilityData VisibilityData, float3
 	 return float(Coverage8bit) / 255.f;
  }
 
- float GetVoxelDensityFixPointScale()
+ uint To16bitCoverage(float Coverage)
+ {
+	 return min(uint(Coverage * 0x10000u), 0xFFFFu);
+ }
+
+ float From16bitCoverage(uint Coverage16bit)
  {
-	 // Constant for scaling the hair coverage during the voxelization (done with atomic integer)
-	 return 1000.f;
+	 return float(Coverage16bit) / float(0xFFFF);
  }
 
  struct FCategorizationData
@@ -271,8 +277,13 @@ FNodeData ConvertToNodeData(float Depth, FVisibilityData VisibilityData, float3
  {
 	 float RadiusAtDepth1Primary;
 	 float RadiusAtDepth1Velocity;
-	 bool  bIsOrthoView;
 	 float VelocityMagnitudeScale;
+
+	 bool bIsOrthoView;
+	 bool bIsGPUDriven;
+
+	 float3 ViewForward;
+	 float3 TranslatedWorldCameraOrigin;
  };
 
  FHairRenderInfo GetHairRenderInfo(float4 ViewHairRenderInfo)
@@ -280,10 +291,22 @@ FNodeData ConvertToNodeData(float Depth, FVisibilityData VisibilityData, float3
 	 FHairRenderInfo Info;
 	 Info.RadiusAtDepth1Primary = ViewHairRenderInfo.x;
 	 Info.RadiusAtDepth1Velocity = ViewHairRenderInfo.y;
-	 Info.bIsOrthoView = ViewHairRenderInfo.z > 0.0f ? true : false;
-	 Info.VelocityMagnitudeScale = ViewHairRenderInfo.w;
+	 Info.VelocityMagnitudeScale = ViewHairRenderInfo.z;
+
+	 const uint BitField = asuint(ViewHairRenderInfo.w);
+	 Info.bIsOrthoView = (BitField & 0x1) != 0;
+	 Info.bIsGPUDriven = (BitField & 0x2) != 0;
+
 	 return Info;
  }
+ 
+ struct FHairViewInfo
+ {
+	 float3 TranslatedWorldCameraOrigin;
+	 float3 ViewForward;
+	 float  RadiusAtDepth1;
+	 bool	bIsOrthoView;
+ };
 
  struct FHairAttributes
  {
@@ -313,10 +336,9 @@ FNodeData ConvertToNodeData(float Depth, FVisibilityData VisibilityData, float3
 // Footprint: 224bits | 28bytes
  struct FPPLLNodeData
  {
-	 // (backloged)TODO try f16(depth) f16(Specular) and make sure depth does not lead to artefact
 	 float Depth;
-	 float Specular;
-	 uint PrimitiveID_ClusterID;
+	 uint Specular_Coverage16bit;
+	 uint PrimitiveID_MacroGroupID;
 	 uint Tangent_Coverage8bit;
 	 uint BaseColor_Roughness;
 	 uint NextNodeIndex;
@@ -333,31 +355,35 @@ FNodeData ConvertToNodeData(float Depth, FVisibilityData VisibilityData, float3
 	 return float2(f16tof32(PackedVelocity >> 16), f16tof32(PackedVelocity));
  }
 
- FPPLLNodeData ConvertToPPLLNodeData(float Depth, FVisibilityData VisibilityData, float3 BaseColor, float Roughness, float Specular, float2 Velocity, uint NextNodeIndex)
+ FPPLLNodeData ConvertToPPLLNodeData(float Depth, float CoverageFullPrecision, FVisibilityData VisibilityData, float3 BaseColor, float Roughness, float Specular, float2 Velocity, uint NextNodeIndex)
  {
 	 // 4x16bits
 	 const uint4 EncodedVisibilityData = EncodeVisibilityData(VisibilityData);
+	 const uint Coverage16bit = To16bitCoverage(CoverageFullPrecision);
 
 	 FPPLLNodeData Out;
 	 Out.Depth = Depth;
 	 Out.Tangent_Coverage8bit = (0xFFFF & EncodedVisibilityData.x) | ((0xFFFF & EncodedVisibilityData.y) << 16);
-	 Out.PrimitiveID_ClusterID = (0xFFFF & EncodedVisibilityData.z) | ((0xFFFF & EncodedVisibilityData.w) << 16);
+	 Out.PrimitiveID_MacroGroupID = (0xFFFF & EncodedVisibilityData.z) | ((0xFFFF & EncodedVisibilityData.w) << 16);
 	 Out.BaseColor_Roughness = float4ToUint(float4(sqrt(BaseColor), Roughness));
-	 Out.Specular = Specular;
+	 Out.Specular_Coverage16bit = clamp(uint(Specular * 0xFFFF), 0, 0xFFFF) | ((0xFFFF & Coverage16bit) << 16);
 	 Out.NextNodeIndex = NextNodeIndex;
 	 Out.PackedVelocity = PackVelocity(Velocity);
 
 	 return Out;
  }
 
- FNodeData ConvertToNodeData(FPPLLNodeData In)
+ FNodeData ConvertToNodeData(FPPLLNodeData In, inout uint OutCoverage16bit)
  {
 	 FNodeData Out;
 	 Out.Depth = In.Depth;
-	 Out.PrimitiveID_ClusterID = In.PrimitiveID_ClusterID;
+	 Out.PrimitiveID_MacroGroupID = In.PrimitiveID_MacroGroupID;
 	 Out.Tangent_Coverage8bit = In.Tangent_Coverage8bit;
 	 Out.BaseColor_Roughness = In.BaseColor_Roughness;
-	 Out.Specular = In.Specular;
+	 Out.Specular = float(In.Specular_Coverage16bit & 0xFFFF) / float(0xFFFF);
+
+	 OutCoverage16bit = (In.Specular_Coverage16bit >> 16) & 0xFFFF;
+
 	 return Out;
  }
 
diff --git a/Engine/Shaders/Private/HairStrands/HairStrandsVisibilityCompaction.usf b/Engine/Shaders/Private/HairStrands/HairStrandsVisibilityCompaction.usf
--- a/Engine/Shaders/Private/HairStrands/HairStrandsVisibilityCompaction.usf
+++ b/Engine/Shaders/Private/HairStrands/HairStrandsVisibilityCompaction.usf
@@ -110,7 +110,8 @@ void MainCS(uint GroupIndex : SV_GroupIndex, uint3 DispatchThreadId : SV_Dispatc
 			float3			BaseColor = 0.0f;
 			float			Roughness = 0.0f;
 			float			Specular = 0;
-			DecodeNodeData(ConvertToNodeData(PPLLNode), Depth, VisibilityData, BaseColor, Roughness, Specular);
+			uint			Coverage16bit = 0;
+			DecodeNodeData(ConvertToNodeData(PPLLNode, Coverage16bit), Depth, VisibilityData, BaseColor, Roughness, Specular);
 
 			const float Coverage = From8bitCoverage(VisibilityData.Coverage8bit);
 			TotalTransmittance *= 1.0f - Coverage;
@@ -190,9 +191,10 @@ void MainCS(uint GroupIndex : SV_GroupIndex, uint3 DispatchThreadId : SV_Dispatc
 			float3			BaseColor = 0.0f;
 			float			Roughness = 0.0f;
 			float			Specular = 0;
-			DecodeNodeData(ConvertToNodeData(PPLLNode), Depth, VisibilityData, BaseColor, Roughness, Specular);
+			float			Coverage16bit = 0;
+			DecodeNodeData(ConvertToNodeData(PPLLNode, Coverage16bit), Depth, VisibilityData, BaseColor, Roughness, Specular);
 
-			const float Coverage = From8bitCoverage(VisibilityData.Coverage8bit);
+			const float Coverage = From16bitCoverage(Coverage16bit);
 
 			// Update current node coverage as a function of previous nodes coverage
 			SortedCoverage[i] = TotalSortedTransmittance * Coverage;
@@ -209,8 +211,9 @@ void MainCS(uint GroupIndex : SV_GroupIndex, uint3 DispatchThreadId : SV_Dispatc
 		for (uint i = 0; i < PixelRenderNodecount; ++i)
 		{
 			const uint NodePPLLIndex = SortedIndex[i];
+			uint Coverage16bit = 0;
 			FPPLLNodeData PPLLNode = PPLLNodeData[NodePPLLIndex];
-			FNodeData OutNodeData = ConvertToNodeData(PPLLNode);
+			FNodeData OutNodeData = ConvertToNodeData(PPLLNode, Coverage16bit);
 			float			Depth = 0.0f;
 			FVisibilityData	VisibilityData = (FVisibilityData)0;
 			float3			BaseColor = 0.0f;
diff --git a/Engine/Shaders/Private/HairStrands/HairStrandsVisibilityCopyIndirectArg.usf b/Engine/Shaders/Private/HairStrands/HairStrandsVisibilityCopyIndirectArg.usf
--- a/Engine/Shaders/Private/HairStrands/HairStrandsVisibilityCopyIndirectArg.usf
+++ b/Engine/Shaders/Private/HairStrands/HairStrandsVisibilityCopyIndirectArg.usf
@@ -8,11 +8,12 @@
 Texture2D<uint> CounterTexture;
 RWBuffer<uint> OutArgBuffer;
 uint ThreadGroupSize;
+uint ItemCountPerGroup;
 
 [numthreads(1, 1, 1)]
 void CopyCS(uint GroupIndex : SV_GroupIndex, uint3 DispatchThreadId : SV_DispatchThreadID)
 {
-	const uint NodeCount  = CounterTexture.Load(uint3(0,0,0));
+	const uint NodeCount  = CounterTexture.Load(uint3(0,0,0)) * ItemCountPerGroup;
 	const uint TotalGroup = (NodeCount + ThreadGroupSize-1) / ThreadGroupSize;
 #if 0
 	// Minimize unused groups
diff --git a/Engine/Shaders/Private/HairStrands/HairStrandsVisibilityPS.usf b/Engine/Shaders/Private/HairStrands/HairStrandsVisibilityPS.usf
--- a/Engine/Shaders/Private/HairStrands/HairStrandsVisibilityPS.usf
+++ b/Engine/Shaders/Private/HairStrands/HairStrandsVisibilityPS.usf
@@ -1,5 +1,7 @@
 // Copyright Epic Games, Inc. All Rights Reserved.
 
+#include "HairStrandsVoxelPageCommonStruct.ush" 
+
 #include "../Common.ush"
 
 //#define SceneTexturesStruct DeepShadowPass.SceneTextures
@@ -38,7 +40,7 @@ RWTexture2D<uint> PPLLNodeIndex;
 RWStructuredBuffer<FPPLLNodeData> PPLLNodes;
 #endif
 
-uint HairVisibilityPass_HairClusterIndex;
+uint HairVisibilityPass_HairMacroGroupIndex;
 uint HairVisibilityPass_HairMaterialId;
 
 #if HAIR_RENDER_MODE == RENDER_MODE_PPLL
@@ -114,7 +116,7 @@ void MainVisibility(
 	VisibilityData.Tangent = float3(0, 0, 1);
 	VisibilityData.PrimitiveId = 0;
   #endif
-	VisibilityData.ClusterId = HairVisibilityPass_HairClusterIndex;
+	VisibilityData.MacroGroupId = HairVisibilityPass_HairMacroGroupIndex;
 	VisibilityData.Coverage8bit = To8bitCoverage(Coverage);
 
 	half3 BaseColor = 0;
@@ -168,7 +170,7 @@ void MainVisibility(
 		// If we can append new node, add it and make it point to the node previously in head of list
 		uint NextNodeIndex;
 		InterlockedExchange(PPLLNodeIndex[PixelCoord], NodeIndex, NextNodeIndex);
-		const FPPLLNodeData Node = ConvertToPPLLNodeData(SvPosition.z, VisibilityData, BaseColor, Roughness, Specular, Velocity, NextNodeIndex);
+		const FPPLLNodeData Node = ConvertToPPLLNodeData(SvPosition.z, Coverage, VisibilityData, BaseColor, Roughness, Specular, Velocity, NextNodeIndex);
 		PPLLNodes[NodeIndex] = Node;
 	}
 #endif
diff --git a/Engine/Shaders/Private/HairStrands/HairStrandsVisibilityTile.usf b/Engine/Shaders/Private/HairStrands/HairStrandsVisibilityTile.usf
new file mode 100644
--- /dev/null
+++ b/Engine/Shaders/Private/HairStrands/HairStrandsVisibilityTile.usf
@@ -0,0 +1,56 @@
+// Copyright Epic Games, Inc. All Rights Reserved.
+#pragma once
+
+#include "../Common.ush"
+#include "HairStrandsVisibilityCommon.ush"
+
+#define TILE_SIZE 8
+
+uint2 Resolution;
+uint2 TileResolution;
+uint  TileCount;
+uint  TileSize;
+
+Texture2D<uint4> CategorizationTexture;
+RWTexture2D<uint> OutTileCounter;
+RWTexture2D<uint> OutTileIndexTexture;
+RWBuffer<uint2>   OutTileBuffer;
+
+groupshared uint bAllocationTile;
+
+[numthreads(TILE_SIZE, TILE_SIZE, 1)]
+void MainCS(uint GroupIndex : SV_GroupIndex, uint3 DispatchThreadId : SV_DispatchThreadID)
+{
+	if (GroupIndex == 0)
+	{
+		bAllocationTile = 0;
+	}
+
+	uint2 PixelCoord = DispatchThreadId.xy;
+	if (PixelCoord.x >= Resolution.x || PixelCoord.y >= Resolution.y)
+		PixelCoord = Resolution-1;
+
+	const uint2 TileCoord = PixelCoord / TILE_SIZE;
+	const FCategorizationData Categorization = DecodeCategorizationData(CategorizationTexture.Load(uint3(PixelCoord, 0)));
+	if (Categorization.PixelCoverage > 0)
+	{
+		InterlockedOr(bAllocationTile, 1);		
+	}
+
+	GroupMemoryBarrierWithGroupSync();
+
+	if (GroupIndex == 0)
+	{
+		if (bAllocationTile > 0)
+		{
+			uint AllocationOffset = 0;
+			InterlockedAdd(OutTileCounter[uint2(0, 0)], 1, AllocationOffset);
+			OutTileIndexTexture[TileCoord] = AllocationOffset;
+			OutTileBuffer[AllocationOffset] = TileCoord;
+		}
+		else
+		{
+			OutTileIndexTexture[TileCoord] = INVALID_TILE_OFFSET;
+		}
+	}
+}
\ No newline at end of file
diff --git a/Engine/Shaders/Private/HairStrands/HairStrandsVisibilityVS.usf b/Engine/Shaders/Private/HairStrands/HairStrandsVisibilityVS.usf
--- a/Engine/Shaders/Private/HairStrands/HairStrandsVisibilityVS.usf
+++ b/Engine/Shaders/Private/HairStrands/HairStrandsVisibilityVS.usf
@@ -25,7 +25,7 @@ struct FHairVisibilityAccumulateVSToPS
 
 #define VS_OUTPUT_TYPE FHairVisibilityAccumulateVSToPS
 
-#if VERTEXSHADER
+#if VERTEXSHADER 
 
 /** transform mesh as normal */
 void Main(
@@ -51,6 +51,7 @@ void Main(
 	// The velocity is used to adapt/increase the minimum rasterization size in order to avoid aliasing under heavy-motion.
 	// The higher the velocity, the larger a strand becomes.
 	float VelocityScale = 0;
+	const FHairRenderInfo HairRenderInfo = GetHairRenderInfo(ResolvedView.HairRenderInfo);
 	#if HAIR_RENDER_MODE == RENDER_MODE_MSAA || HAIR_RENDER_MODE == RENDER_MODE_PPLL
 	{
 		const float4 ScreenPos			= mul(float4(WorldPositionRaw.xyz, 1), ResolvedView.TranslatedWorldToClip);
@@ -59,13 +60,18 @@ void Main(
 		Output.PackedVelocityC			= INVARIANT(PrevScreenPosObj);
 
 		const float2 ScreenVelocity = Calculate2DVelocity(Output.PackedVelocityA, Output.PackedVelocityC);
-		FHairRenderInfo HairRenderInfo = GetHairRenderInfo(ResolvedView.HairRenderInfo);
 		VelocityScale = saturate(max(abs(ScreenVelocity.x), abs(ScreenVelocity.y)) / HairRenderInfo.VelocityMagnitudeScale);
 	}
 	#endif
 
 #if VF_STRAND_HAIR
-	float4 WorldPosition = VertexFactoryGetWorldPosition(Input, VFIntermediates, VelocityScale);
+	FHairViewInfo HairViewInfo;
+	HairViewInfo.TranslatedWorldCameraOrigin = ResolvedView.TranslatedWorldCameraOrigin;
+	HairViewInfo.ViewForward = ResolvedView.ViewForward;
+	HairViewInfo.RadiusAtDepth1 = lerp(HairRenderInfo.RadiusAtDepth1Primary, HairRenderInfo.RadiusAtDepth1Velocity, VelocityScale);
+	HairViewInfo.bIsOrthoView = HairRenderInfo.bIsOrthoView;
+
+	float4 WorldPosition = VertexFactoryGetWorldPosition(Input, VFIntermediates, HairViewInfo);
 #else
 	float4 WorldPosition = VertexFactoryGetWorldPosition(Input, VFIntermediates);
 #endif
diff --git a/Engine/Shaders/Private/HairStrands/HairStrandsVoxelMip.usf b/Engine/Shaders/Private/HairStrands/HairStrandsVoxelMip.usf
--- a/Engine/Shaders/Private/HairStrands/HairStrandsVoxelMip.usf
+++ b/Engine/Shaders/Private/HairStrands/HairStrandsVoxelMip.usf
@@ -1,7 +1,44 @@
 // Copyright Epic Games, Inc. All Rights Reserved.
 
 #include "/Engine/Public/Platform.ush"
+#include "HairStrandsVoxelPageCommon.ush"
+
+uint ComputeMipDensity(
+	const uint RawDensity0,
+	const uint RawDensity1,
+	const uint RawDensity2,
+	const uint RawDensity3,
+	const uint RawDensity4,
+	const uint RawDensity5,
+	const uint RawDensity6,
+	const uint RawDensity7)
+{
+	const float TotalOpaque =
+		((RawDensity0 & VOXEL_OPAQUE_MASK) >> VOXEL_OPAQUE_SHIFT) +
+		((RawDensity1 & VOXEL_OPAQUE_MASK) >> VOXEL_OPAQUE_SHIFT) +
+		((RawDensity2 & VOXEL_OPAQUE_MASK) >> VOXEL_OPAQUE_SHIFT) +
+		((RawDensity3 & VOXEL_OPAQUE_MASK) >> VOXEL_OPAQUE_SHIFT) +
+		((RawDensity4 & VOXEL_OPAQUE_MASK) >> VOXEL_OPAQUE_SHIFT) +
+		((RawDensity5 & VOXEL_OPAQUE_MASK) >> VOXEL_OPAQUE_SHIFT) +
+		((RawDensity6 & VOXEL_OPAQUE_MASK) >> VOXEL_OPAQUE_SHIFT) +
+		((RawDensity7 & VOXEL_OPAQUE_MASK) >> VOXEL_OPAQUE_SHIFT);
+	const uint OutTotalOpaque = uint(clamp(TotalOpaque / 8.f, TotalOpaque > 0 ? 1 : 0, 0xFF)) << VOXEL_OPAQUE_SHIFT;
+
+	uint TotalHair = min(VOXEL_HAIR_MASK,
+		(RawDensity0 & VOXEL_HAIR_MASK) +
+		(RawDensity1 & VOXEL_HAIR_MASK) +
+		(RawDensity2 & VOXEL_HAIR_MASK) +
+		(RawDensity3 & VOXEL_HAIR_MASK) +
+		(RawDensity4 & VOXEL_HAIR_MASK) +
+		(RawDensity5 & VOXEL_HAIR_MASK) +
+		(RawDensity6 & VOXEL_HAIR_MASK) +
+		(RawDensity7 & VOXEL_HAIR_MASK));
+	TotalHair /= 8;
+
+	return TotalHair | OutTotalOpaque;
+}
 
+#if SHADER_VOXEL
 
 uint VoxelResolution;
 uint SourceMip;
@@ -10,9 +47,9 @@ uint TargetMip;
 Texture3D<uint> InDensityTexture;
 RWTexture3D<uint> OutDensityTexture0;
 
+
 // Write two levels at a time, and store/cache thread result within LDS
 #if PERMUTATION_METHOD == 1
-
 RWTexture3D<uint> OutDensityTexture1;
 groupshared uint Values0[64];
 groupshared uint Values1[8];
@@ -47,16 +84,16 @@ void MainCS(uint3 DispatchThreadId : SV_DispatchThreadID, uint3 GroupThreadId :
 		const uint3 OutLocalCoord = FromMorton(GroupIndex);
 		const uint3 OutVoxelCoord = (GroupId << 1) + OutLocalCoord;
 
-		uint Total = 0;
 		const uint Offset = GroupIndex * 8;
-		Total += Values0[Offset + 0];
-		Total += Values0[Offset + 1];
-		Total += Values0[Offset + 2];
-		Total += Values0[Offset + 3];
-		Total += Values0[Offset + 4];
-		Total += Values0[Offset + 5];
-		Total += Values0[Offset + 6];
-		Total += Values0[Offset + 7];
+		const uint Total = ComputeMipDensity(
+			Values0[Offset + 0],
+			Values0[Offset + 1],
+			Values0[Offset + 2],
+			Values0[Offset + 3],
+			Values0[Offset + 4],
+			Values0[Offset + 5],
+			Values0[Offset + 6],
+			Values0[Offset + 7]);
 		Values1[GroupIndex] = Total;
 		OutDensityTexture0[OutVoxelCoord] = Total;
 	}
@@ -66,15 +103,15 @@ void MainCS(uint3 DispatchThreadId : SV_DispatchThreadID, uint3 GroupThreadId :
 	if (GroupIndex == 0)
 	{
 		const uint3 OutVoxelCoord = GroupId;
-		uint Total = 0;
-		Total += Values1[0];
-		Total += Values1[1];
-		Total += Values1[2];
-		Total += Values1[3];
-		Total += Values1[4];
-		Total += Values1[5];
-		Total += Values1[6];
-		Total += Values1[7];
+		const uint Total = ComputeMipDensity(
+			Values1[0],
+			Values1[1],
+			Values1[2],
+			Values1[3],
+			Values1[4],
+			Values1[5],
+			Values1[6],
+			Values1[7]);
 		OutDensityTexture1[OutVoxelCoord] = Total;
 	}
 }
@@ -90,23 +127,120 @@ void MainCS(uint3 DispatchThreadId : SV_DispatchThreadID)
 	const uint3 InVoxelCoord1 = InVoxelCoord0 + uint3(1,0,0);
 	const uint3 InVoxelCoord2 = InVoxelCoord0 + uint3(0,1,0);
 	const uint3 InVoxelCoord3 = InVoxelCoord0 + uint3(1,1,0);
-
 	const uint3 InVoxelCoord4 = InVoxelCoord0 + uint3(0,0,1);
 	const uint3 InVoxelCoord5 = InVoxelCoord0 + uint3(1,0,1);
 	const uint3 InVoxelCoord6 = InVoxelCoord0 + uint3(0,1,1);
 	const uint3 InVoxelCoord7 = InVoxelCoord0 + uint3(1,1,1);
 
-	const uint Total =
-		InDensityTexture[InVoxelCoord0] +
-		InDensityTexture[InVoxelCoord1] +
-		InDensityTexture[InVoxelCoord2] +
-		InDensityTexture[InVoxelCoord3] +
 
-		InDensityTexture[InVoxelCoord4] +
-		InDensityTexture[InVoxelCoord5] +
-		InDensityTexture[InVoxelCoord6] +
-		InDensityTexture[InVoxelCoord7];
+	const uint RawDensity0 = InDensityTexture[InVoxelCoord0];
+	const uint RawDensity1 = InDensityTexture[InVoxelCoord1];
+	const uint RawDensity2 = InDensityTexture[InVoxelCoord2];
+	const uint RawDensity3 = InDensityTexture[InVoxelCoord3];
+	const uint RawDensity4 = InDensityTexture[InVoxelCoord4];
+	const uint RawDensity5 = InDensityTexture[InVoxelCoord5];
+	const uint RawDensity6 = InDensityTexture[InVoxelCoord6];
+	const uint RawDensity7 = InDensityTexture[InVoxelCoord7];
+
+	const uint Total = ComputeMipDensity(
+		RawDensity0,
+		RawDensity1,
+		RawDensity2,
+		RawDensity3,
+		RawDensity4,
+		RawDensity5,
+		RawDensity6,
+		RawDensity7);
 
 	OutDensityTexture0[OutVoxelCoord] = Total;
 }
 #endif
+
+#endif
+
+#if SHADER_VIRTUALVOXEL
+int3 PageCountResolution;
+uint PageResolution;
+uint SourceMip;
+uint TargetMip;
+
+Texture3D<uint>   InDensityTexture;
+RWTexture3D<uint> OutDensityTexture;
+
+[numthreads(1, 1, 64)]
+void MainCS(uint3 DispatchThreadId : SV_DispatchThreadID)
+{
+	const uint  MaxDispatchCountPerDim = 65535;
+	const uint  PageIndex = DispatchThreadId.x + DispatchThreadId.y * MaxDispatchCountPerDim;
+	const uint3 PageCoord = IndexToCoord(PageIndex, PageCountResolution);
+
+	const uint  InPageResolution = PageResolution >> SourceMip;
+	const uint OutPageResolution = PageResolution >> TargetMip;
+	const uint TotalVoxelPerOutPageResolution = OutPageResolution * OutPageResolution * OutPageResolution;
+	const uint VoxelIndex = DispatchThreadId.z;
+	if (VoxelIndex >= TotalVoxelPerOutPageResolution)
+		return;
+
+	const uint3 OutVoxelCoordOffset = IndexToCoord(VoxelIndex, OutPageResolution);
+	const uint3  InVoxelCoordOffset = OutVoxelCoordOffset << 1;
+
+	const uint3 OutVoxelCoord = PageCoord * OutPageResolution + OutVoxelCoordOffset;
+	const uint3  InVoxelCoord = PageCoord *  InPageResolution +  InVoxelCoordOffset;
+
+	const uint3 InVoxelCoord0 = InVoxelCoord;
+	const uint3 InVoxelCoord1 = InVoxelCoord0 + uint3(1, 0, 0);
+	const uint3 InVoxelCoord2 = InVoxelCoord0 + uint3(0, 1, 0);
+	const uint3 InVoxelCoord3 = InVoxelCoord0 + uint3(1, 1, 0);
+	const uint3 InVoxelCoord4 = InVoxelCoord0 + uint3(0, 0, 1);
+	const uint3 InVoxelCoord5 = InVoxelCoord0 + uint3(1, 0, 1);
+	const uint3 InVoxelCoord6 = InVoxelCoord0 + uint3(0, 1, 1);
+	const uint3 InVoxelCoord7 = InVoxelCoord0 + uint3(1, 1, 1);
+
+	const uint RawDensity0 = InDensityTexture[InVoxelCoord0];
+	const uint RawDensity1 = InDensityTexture[InVoxelCoord1];
+	const uint RawDensity2 = InDensityTexture[InVoxelCoord2];
+	const uint RawDensity3 = InDensityTexture[InVoxelCoord3];
+	const uint RawDensity4 = InDensityTexture[InVoxelCoord4];
+	const uint RawDensity5 = InDensityTexture[InVoxelCoord5];
+	const uint RawDensity6 = InDensityTexture[InVoxelCoord6];
+	const uint RawDensity7 = InDensityTexture[InVoxelCoord7];
+
+	const uint Total = ComputeMipDensity(
+		RawDensity0,
+		RawDensity1,
+		RawDensity2,
+		RawDensity3,
+		RawDensity4,
+		RawDensity5,
+		RawDensity6,
+		RawDensity7);
+
+	OutDensityTexture[OutVoxelCoord] = Total;
+}
+#endif
+
+#if SHADER_INDIRECTARGS
+#include "HairStrandsVoxelPageCommon.ush"
+
+uint PageResolution;
+uint TargetMipIndex;
+uint DispatchGroupSize;
+
+Buffer<uint>   InIndirectArgs;
+RWBuffer<uint> OutIndirectArgs;
+
+[numthreads(1, 1, 1)]
+void MainCS(uint3 DispatchThreadId : SV_DispatchThreadID)
+{
+	const uint MacroGroupId = DispatchThreadId.x;
+
+	const uint TargetPageResolution = PageResolution >> TargetMipIndex;
+	const uint TotalVoxelCount = TargetPageResolution * TargetPageResolution * TargetPageResolution;
+	const uint DispatchZ = ceil(TotalVoxelCount / float(DispatchGroupSize));
+
+	OutIndirectArgs[0] = InIndirectArgs[0];
+	OutIndirectArgs[1] = InIndirectArgs[1];
+	OutIndirectArgs[2] = DispatchZ;
+}
+
+#endif
\ No newline at end of file
diff --git a/Engine/Shaders/Private/HairStrands/HairStrandsVoxelOpaque.usf b/Engine/Shaders/Private/HairStrands/HairStrandsVoxelOpaque.usf
--- a/Engine/Shaders/Private/HairStrands/HairStrandsVoxelOpaque.usf
+++ b/Engine/Shaders/Private/HairStrands/HairStrandsVoxelOpaque.usf
@@ -3,14 +3,17 @@
 #include "../Common.ush"
 #include "../CommonViewUniformBuffer.ush"
 #include "../SceneTextureParameters.ush"
+#include "HairStrandsVoxelPageCommonStruct.ush"
+#include "HairStrandsVoxelPageCommon.ush"
 
-
-#if PIXELSHADER
+#if SHADER_INJECTOPAQUE_VOXEL
 float3 VoxelMinAABB;
 uint   VoxelResolution;
 float3 VoxelMaxAABB;
+uint   VoxelBiasCount;
 float2 OutputResolution;
 float2 SceneDepthResolution;
+uint   VoxelMarkCount;
 
 Texture2D<float> SceneDepthTexture;
 RWTexture3D<uint> DensityTexture;
@@ -45,7 +48,7 @@ void MainPS(
 
 	float4 ClipPos = mul(float4(WorldPosition,1), View.WorldToClip);
 	ClipPos /= ClipPos.w;
-	const float DepthBias 			= 3 * VoxelSize;
+	const float DepthBias 			= VoxelBiasCount * VoxelSize;
 	const float  VoxelDepth 		= ConvertFromDeviceZ(ClipPos.z) - DepthBias;
 	float2 SceneUV 					= float2(0.5f*(ClipPos.x+1), 1-0.5f*(ClipPos.y+1));
 	SceneUV							= ViewportUVToBufferUV(SceneUV);
@@ -68,15 +71,109 @@ void MainPS(
 		return;
 
 	// Inject opaque depth on a thin layer (Dist < DistThreshold) for avoiding weird projection
-	if (ClosestDepth < VoxelDepth && abs(ClosestDepth - VoxelDepth) < DepthBias)
+	if (ClosestDepth < VoxelDepth && abs(ClosestDepth - VoxelDepth) < VoxelMarkCount * VoxelSize)
 	{
 		OutColor0 = float4(1,0,0,1);
-		InterlockedOr(DensityTexture[VoxelPos], 0x10000000);
+		InterlockedOr(DensityTexture[VoxelPos], VOXEL_OPAQUE_ADD);
 	}
 }
 #endif
 
-#if COMPUTESHADER
+#if SHADER_INJECTOPAQUE_VIRTUALVOXEL
+uint3 DispatchedPageIndexResolution;
+uint MacroGroupId;
+float2 SceneDepthResolution;
+uint VoxelBiasCount;
+uint VoxelMarkCount;
+Texture2D<float> SceneDepthTexture;
+RWTexture3D<uint> OutPageTexture;
+
+uint		 VirtualVoxel_PageIndexCount;
+uint		 VirtualVoxel_PageResolution;
+float		 VirtualVoxel_VoxelWorldSize;
+uint3		 VirtualVoxel_PageCountResolution;
+Buffer<uint> VirtualVoxel_PageIndexBuffer;
+Buffer<uint4>VirtualVoxel_PageIndexCoordBuffer;
+ 
+StructuredBuffer<FPackedVirtualVoxelNodeDesc> VirtualVoxel_NodeDescBuffer;
+
+float3 ReconstructWorldPositionFromDepth(float2 UV, float Depth)
+{
+	float2 ScreenPosition = (UV - View.ScreenPositionScaleBias.wz) / View.ScreenPositionScaleBias.xy;
+	float4 HomogeneousWorldPosition = mul(float4(ScreenPosition * Depth, Depth, 1), View.ScreenToWorld);
+	float3 WorldPosition = HomogeneousWorldPosition.xyz / HomogeneousWorldPosition.w;
+
+	return WorldPosition;
+}
+
+#define GROUP_SIZE 64
+[numthreads(1, 1, GROUP_SIZE)]
+void MainCS(uint3 DispatchThreadId : SV_DispatchThreadID)
+{
+	const uint  MaxDispatchCountPerDim = 65535u;
+	const uint  AllocatedPageIndex  = DispatchThreadId.x + DispatchThreadId.y * MaxDispatchCountPerDim;
+	const uint  LinearVoxelCoord = DispatchThreadId.z;
+	const uint3 VoxelCoordOffset = IndexToCoord(LinearVoxelCoord, VirtualVoxel_PageResolution.xxx);
+
+	const FPackedVirtualVoxelNodeDesc PackedNode = VirtualVoxel_NodeDescBuffer[MacroGroupId];
+	const FVirtualVoxelNodeDesc NodeDesc = UnpackVoxelNode(PackedNode, VirtualVoxel_PageResolution);
+
+	const uint4 PageIndexCoord = VirtualVoxel_PageIndexCoordBuffer[NodeDesc.PageIndexOffset + AllocatedPageIndex];
+	const uint  LinearPageIndexCoord = CoordToIndex(PageIndexCoord, NodeDesc.PageIndexResolution, NodeDesc.PageIndexOffset);
+
+	const bool bIsValid = LinearPageIndexCoord < VirtualVoxel_PageIndexCount;
+	if (bIsValid)
+	{
+		const uint PageIndex = VirtualVoxel_PageIndexBuffer.Load(LinearPageIndexCoord);
+		if (PageIndex != INVALID_VOXEL_PAGE_INDEX)
+		{
+			const uint3 VoxelCoordBase = PageIndexCoord.xyz * VirtualVoxel_PageResolution;
+			const uint3 VoxelCoord = VoxelCoordBase + VoxelCoordOffset;
+			const float3 WorldPosition = VoxelCoord / float3(NodeDesc.VirtualResolution) * (NodeDesc.MaxAABB - NodeDesc.MinAABB) + NodeDesc.MinAABB;
+
+			float4 ClipPos = mul(float4(WorldPosition, 1), View.WorldToClip);
+			ClipPos /= ClipPos.w;
+			const float DepthBias = VoxelBiasCount * VirtualVoxel_VoxelWorldSize;
+			const float VoxelDepth = ConvertFromDeviceZ(ClipPos.z) - DepthBias;
+			float2 SceneUV = float2(0.5f*(ClipPos.x + 1), 1 - 0.5f*(ClipPos.y + 1));
+			SceneUV = ViewportUVToBufferUV(SceneUV);
+			const float2 ScenePixelCoord = SceneUV * SceneDepthResolution;
+
+			const bool bIsOnScreen = SceneUV.x >= 0 && SceneUV.x < 1 && SceneUV.y >= 0 && SceneUV.y < 1;
+			if (!bIsOnScreen)
+				return;
+
+			const float ClosestDepth = ConvertFromDeviceZ(SceneDepthTexture.Load(uint3(ScenePixelCoord, 0)));
+			const float3 SceneWorldPos = ReconstructWorldPositionFromDepth(SceneUV, ClosestDepth);
+
+			const bool bIsInVolume =
+				SceneWorldPos.x >= NodeDesc.MinAABB.x && SceneWorldPos.x < NodeDesc.MaxAABB.x &&
+				SceneWorldPos.y >= NodeDesc.MinAABB.y && SceneWorldPos.y < NodeDesc.MaxAABB.y &&
+				SceneWorldPos.z >= NodeDesc.MinAABB.z && SceneWorldPos.z < NodeDesc.MaxAABB.z;
+
+			if (!bIsInVolume)
+				return;
+
+			// Inject opaque depth on a thin layer (Dist < DistThreshold) for avoiding weird projection
+			if (ClosestDepth < VoxelDepth && abs(ClosestDepth - VoxelDepth) < VoxelMarkCount * VirtualVoxel_VoxelWorldSize)
+			{
+				const uint3 VoxelPageIndexCoord = VoxelCoord / VirtualVoxel_PageResolution;
+				const uint3 VoxelIndexCoordBase = VoxelPageIndexCoord * VirtualVoxel_PageResolution;
+				const uint3 VoxelPageOffset = VoxelCoord - VoxelIndexCoordBase;
+
+				const uint3 PageCoord = IndexToCoord(PageIndex, VirtualVoxel_PageCountResolution);
+				const int3 VoxelPageBase = PageCoord * VirtualVoxel_PageResolution;
+				const int3 VoxelPageCoord = VoxelPageOffset + VoxelPageBase;
+
+				InterlockedOr(OutPageTexture[VoxelPageCoord], VOXEL_OPAQUE_ADD);
+			}
+		}
+	}
+}
+#endif
+
+
+#if SHADER_FILTERDEPTH_VOXEL
 
 uint				VoxelResolution;
 RWTexture3D<uint>	VoxelTexture;
@@ -84,11 +181,9 @@ RWTexture3D<uint>	VoxelTexture;
 [numthreads(4, 4, 4)]
 void MainCS(uint3 DispatchThreadId : SV_DispatchThreadID)
 {
-	const uint  OpaqueMask	 =  0x10000000;
-	const uint  DensityMask	 = ~OpaqueMask;
-	const uint3 VoxelCoord   = DispatchThreadId;
-	const uint  VoxelDensity = VoxelTexture[VoxelCoord];
-	const bool  bHasOpaque	 = (OpaqueMask & VoxelDensity) > 0;
+	const uint3 VoxelCoord  = DispatchThreadId;
+	const uint  RawDensity	= VoxelTexture[VoxelCoord];
+	const bool  bHasOpaque	= GetInternalVoxelOpaqueVisibility(RawDensity) > 0;
 
 	// If the current voxel does not contains any opaque data, there is nothing to filter.
 	if (!bHasOpaque)
@@ -103,7 +198,7 @@ void MainCS(uint3 DispatchThreadId : SV_DispatchThreadID)
 	{
 		const int3 Offset = int3(x, y, z);
 		const int3 Sample = clamp(VoxelCoord + Offset, 0, VoxelResolution-1);
-		const uint Density = VoxelTexture[Sample] & DensityMask;
+		const uint Density = VoxelTexture[Sample] & VOXEL_HAIR_MASK;
 		if (Density > 0)
 		{
 			return;
@@ -111,6 +206,6 @@ void MainCS(uint3 DispatchThreadId : SV_DispatchThreadID)
 	}
 
 	// No hair density has been found, so this opaque voxel is a false positive
-	InterlockedAnd(VoxelTexture[VoxelCoord], DensityMask);
+	InterlockedAnd(VoxelTexture[VoxelCoord], VOXEL_HAIR_MASK);
 }
 #endif
\ No newline at end of file
diff --git a/Engine/Shaders/Private/HairStrands/HairStrandsVoxelPageAllocation.usf b/Engine/Shaders/Private/HairStrands/HairStrandsVoxelPageAllocation.usf
new file mode 100644
--- /dev/null
+++ b/Engine/Shaders/Private/HairStrands/HairStrandsVoxelPageAllocation.usf
@@ -0,0 +1,436 @@
+// Copyright Epic Games, Inc. All Rights Reserved.
+
+#include "../Common.ush"
+#include "../Matrices.ush"
+#include "HairStrandsVoxelPageCommon.ush"
+#include "HairStrandsAABBCommon.ush"
+
+#if SHADER_ALLOCATEPAGEINDEX
+
+float				PageWorldSize;
+uint				TotalPageIndexCount; 	// This is the max page index count;
+uint				PageResolution; 		// Resolution of a page
+uint				MacroGroupCount;
+uint				IndirectDispatchGroupSize;
+
+// For testing parity with CPU version
+float3				CPU_MinAABB;
+float3				CPU_MaxAABB;
+int3				CPU_PageIndexResolution;
+uint				CPU_bUseCPUData;
+
+RWBuffer<int>								MacroGroupAABBBuffer;
+RWBuffer<uint4>								OutPageIndexResolutionAndOffsetBuffer;
+RWBuffer<uint>								OutPageIndexAllocationIndirectBufferArgs;
+RWStructuredBuffer<FVoxelizationViewInfo>	OutVoxelizationViewInfoBuffer;
+
+#define MAX_INSTANCE_GROUP_COUNT 32
+groupshared uint PageIndexOffsets[MAX_INSTANCE_GROUP_COUNT];
+
+// This code assume we have less than 32 macro group (which fit into a single CU/SM)
+[numthreads(MAX_INSTANCE_GROUP_COUNT, 1, 1)]
+void AllocatePageIndex(uint2 DispatchThreadId : SV_DispatchThreadID)
+{
+	const uint MacroGroupId = DispatchThreadId.x;
+	FHairAABB Bound = ReadHairAABB(MacroGroupId, MacroGroupAABBBuffer);
+
+	// Page index allocation
+	int3 PageIndexResolution = 0;
+	{
+		// Snap the max AABB to the voxel size
+		// Allocate enough pages to cover the AABB, where page (0,0,0) origin sit on MinAABB.
+		float3 MacroGroupSize = Bound.Max - Bound.Min;
+		PageIndexResolution = ceil(MacroGroupSize / PageWorldSize);
+		MacroGroupSize = PageIndexResolution * PageWorldSize;
+		Bound.Max = MacroGroupSize + Bound.Min;
+
+		const uint TotalPageIndex = PageIndexResolution.x * PageIndexResolution.y * PageIndexResolution.z;		
+		PageIndexOffsets[MacroGroupId] = TotalPageIndex;
+		GroupMemoryBarrierWithGroupSync();
+
+		// Postfix sum to instance group are always ordered by index
+		if (DispatchThreadId.x == 0)
+		{
+			uint PageIndexOffset = 0;
+			for (uint MacroGroupId = 0; MacroGroupId < MacroGroupCount; ++MacroGroupId)
+			{
+				const uint PageCount = PageIndexOffsets[MacroGroupId];
+				PageIndexOffsets[MacroGroupId] = PageIndexOffset;
+				PageIndexOffset += PageCount;
+			}
+		}
+		GroupMemoryBarrierWithGroupSync();
+
+		const uint PageIndexOffset = PageIndexOffsets[MacroGroupId];
+		OutPageIndexResolutionAndOffsetBuffer[MacroGroupId] = uint4(PageIndexResolution, PageIndexOffset);
+		WriteHairAABB(MacroGroupId, Bound, MacroGroupAABBBuffer);
+	}
+
+	// Use CPU data, which allow to compare the GPU driven part with CPU data
+	// This code path is only for debug purpose and work only when there is a single instance group
+	if (CPU_bUseCPUData > 0)
+	{
+		Bound.Min = CPU_MinAABB;
+		Bound.Max = CPU_MaxAABB;
+		PageIndexResolution = CPU_PageIndexResolution;
+		WriteHairAABB(MacroGroupId, Bound, MacroGroupAABBBuffer);
+	}
+
+	// Prepare indirect buffer for doing the actual page index allocation and filling the page index
+	{
+		const uint AllocatedPageIndexCount = PageIndexResolution.x * PageIndexResolution.y * PageIndexResolution.z;
+		const uint ArgsOffset = MacroGroupId * 3;
+		OutPageIndexAllocationIndirectBufferArgs[ArgsOffset] = ceil(AllocatedPageIndexCount / float(IndirectDispatchGroupSize));
+		OutPageIndexAllocationIndirectBufferArgs[ArgsOffset + 1] = 1;
+		OutPageIndexAllocationIndirectBufferArgs[ArgsOffset + 2] = 1;
+	}
+
+	// Matrix generation
+	if (MacroGroupId < MacroGroupCount)
+	{
+		// Find the largest resolution and its dominant axis
+		uint2 RasterResolution = uint2(0, 0);
+		float3 RasterProjectionSize = float3(0,0,0);
+		float3 RasterDirection = float3(0, 0, 0);
+		float3 RasterUp = float3(0, 0, 0);
+		const int3 TotalVoxelResolution = PageIndexResolution * PageResolution;
+		{
+			const uint ResolutionXY = TotalVoxelResolution.x * TotalVoxelResolution.y;
+			const uint ResolutionXZ = TotalVoxelResolution.x * TotalVoxelResolution.z;
+			const uint ResolutionYZ = TotalVoxelResolution.y * TotalVoxelResolution.y;
+			if (ResolutionXY >= ResolutionXZ && ResolutionXY >= ResolutionYZ)
+			{
+				RasterResolution = int2(TotalVoxelResolution.x, TotalVoxelResolution.y);
+				RasterDirection = float3(0, 0, 1);
+				RasterUp = float3(0, 1, 0);
+
+				FHairAABB ProjRasterAABB;
+				ProjRasterAABB.Min = float3(Bound.Min.x, Bound.Min.y, Bound.Min.z);
+				ProjRasterAABB.Max = float3(Bound.Max.x, Bound.Max.y, Bound.Max.z);
+				RasterProjectionSize = ProjRasterAABB.Max - ProjRasterAABB.Min;
+			}
+			else if (ResolutionXZ >= ResolutionXY && ResolutionXZ >= ResolutionYZ)
+			{
+				RasterResolution = int2(TotalVoxelResolution.x, TotalVoxelResolution.z);
+				RasterDirection = float3(0, -1, 0);
+				RasterUp = float3(0, 0, 1);
+
+				FHairAABB ProjRasterAABB;
+				ProjRasterAABB.Min = float3(Bound.Min.x, Bound.Min.z, Bound.Min.y);
+				ProjRasterAABB.Max = float3(Bound.Max.x, Bound.Max.z, Bound.Max.y);
+				RasterProjectionSize = ProjRasterAABB.Max - ProjRasterAABB.Min;
+			}
+			else
+			{
+				RasterResolution = int2(TotalVoxelResolution.y, TotalVoxelResolution.z);
+				RasterDirection = float3(1, 0, 0);
+				RasterUp = float3(0, 0, 1);
+
+				FHairAABB ProjRasterAABB;
+				ProjRasterAABB.Min = float3(Bound.Min.y, Bound.Min.z, Bound.Min.x);
+				ProjRasterAABB.Max = float3(Bound.Max.y, Bound.Max.z, Bound.Max.x);
+				RasterProjectionSize = ProjRasterAABB.Max - ProjRasterAABB.Min;
+			}
+		}
+
+		const float3 RasterAABBSize = Bound.Max - Bound.Min;
+		const float3 RasterAABBCenter = (Bound.Max + Bound.Min) * 0.5f;
+
+		const float4x4 ProjMatrix  = ReversedZOrthoMatrix(0.5f * RasterProjectionSize.x, 0.5f * RasterProjectionSize.y, 1.0f / RasterProjectionSize.z, 0);
+		const float4x4 ViewMatrix  = LookAtMatrix(RasterAABBCenter - RasterDirection * RasterProjectionSize.z * 0.5f, RasterAABBCenter, RasterUp);
+
+		FVoxelizationViewInfo ViewInfo;
+		ViewInfo.WorldToClip = mul(ViewMatrix, ProjMatrix);
+		ViewInfo.ViewForward = RasterDirection;
+		ViewInfo.RasterResolution = RasterResolution;
+		ViewInfo.Pad0 = 0;
+		ViewInfo.Pad1 = 0;
+		OutVoxelizationViewInfoBuffer[MacroGroupId] = ViewInfo;
+	}
+}
+#endif
+
+
+#if SHADER_MARKVALID
+uint			MaxClusterCount;
+int3			CPU_PageIndexResolution;
+float3			CPU_MinAABB;
+uint			CPU_PageIndexOffset;
+float3			CPU_MaxAABB;
+uint			MacroGroupId;
+
+Buffer<int>		ClusterAABBsBuffer;
+Buffer<int>		MacroGroupAABBBuffer;
+Buffer<uint4>	PageIndexResolutionAndOffsetBuffer;
+RWBuffer<uint>	OutValidPageIndexBuffer;
+
+#define GROUP_SIZE 32
+
+// PageIndexBuffer      is sampled with linear coordinate computed from the 3d page coordinate. VALID NODE ARE NOT COMPACTED. It contains the LINEAR PAGE INDEX (to map to the 3d volume).
+// PageIndexCoordBuffer is sampled with linear coordinate for allocated nodes. VALID NODE ARE COMPACTED. It contains the 3d page coordinate and ClustedId. Only used for opaque voxel injection.
+
+[numthreads(GROUP_SIZE, 1, 1)]
+void MarkValidCS(uint2 DispatchThreadId : SV_DispatchThreadID)
+{
+	const uint ClusterIndex = DispatchThreadId.x;
+	if (ClusterIndex >= MaxClusterCount)
+		return;
+
+	const uint BaseClusterIndex = 6 * ClusterIndex;
+
+	FHairAABB ClusterBound;
+	ClusterBound.Min.x = ClusterAABBsBuffer[BaseClusterIndex + 0];
+	ClusterBound.Min.y = ClusterAABBsBuffer[BaseClusterIndex + 1];
+	ClusterBound.Min.z = ClusterAABBsBuffer[BaseClusterIndex + 2];
+
+	ClusterBound.Max.x = ClusterAABBsBuffer[BaseClusterIndex + 3];
+	ClusterBound.Max.y = ClusterAABBsBuffer[BaseClusterIndex + 4];
+	ClusterBound.Max.z = ClusterAABBsBuffer[BaseClusterIndex + 5];
+
+	FHairAABB MacroGroupBound;
+	int3 PageIndexResolution;
+	uint PageIndexOffset;
+
+#if PERMUTATION_GPU_DRIVEN == 1
+	const uint4 PageIndexResolutionAndOffset = PageIndexResolutionAndOffsetBuffer.Load(MacroGroupId);
+
+	MacroGroupBound = ReadHairAABB(MacroGroupId, MacroGroupAABBBuffer);
+	PageIndexResolution		= PageIndexResolutionAndOffset.xyz;
+	PageIndexOffset			= PageIndexResolutionAndOffset.w;
+#else
+	MacroGroupBound.Min		= CPU_MinAABB;
+	MacroGroupBound.Max		= CPU_MaxAABB;
+	PageIndexResolution		= CPU_PageIndexResolution;
+	PageIndexOffset			= CPU_PageIndexOffset;
+#endif
+
+	uint3 MinCoord = PositionToCoord(ClusterBound.Min, MacroGroupBound.Min, MacroGroupBound.Max, PageIndexResolution);
+	uint3 MaxCoord = PositionToCoord(ClusterBound.Max, MacroGroupBound.Min, MacroGroupBound.Max, PageIndexResolution);
+
+	MinCoord = clamp(MinCoord, 0, PageIndexResolution - 1);
+	MaxCoord = clamp(MaxCoord, 0, PageIndexResolution - 1);	
+
+	// Splat cluster contribution
+	// Change this to be a bit smarter (use graphics pipe?)
+	for (uint z = MinCoord.z; z <= MaxCoord.z; ++z)
+	for (uint y = MinCoord.y; y <= MaxCoord.y; ++y)
+	for (uint x = MinCoord.x; x <= MaxCoord.x; ++x)
+	{
+		const uint3 PageIndexCoord = uint3(x, y, z);
+		const uint LinearPageIndexCoord = CoordToIndex(PageIndexCoord, PageIndexResolution, PageIndexOffset);
+		InterlockedOr(OutValidPageIndexBuffer[LinearPageIndexCoord], 1u);
+	}
+}
+#endif
+
+#if SHADER_ALLOCATE
+uint MacroGroupId;
+uint PageCount;
+uint CPU_PageIndexCount;
+uint CPU_PageIndexOffset;
+uint3 CPU_PageIndexResolution;
+
+Buffer<uint4> PageIndexResolutionAndOffsetBuffer;
+Buffer<int> IndirectBufferArgs;
+
+RWBuffer<uint> PageIndexGlobalCounter;
+RWBuffer<uint> PageIndexBuffer;
+RWBuffer<uint4> PageIndexCoordBuffer;
+
+groupshared uint LocalCounter;
+groupshared uint GroupBase;
+
+#define GROUP_SIZE 32
+
+[numthreads(GROUP_SIZE, 1, 1)]
+void AllocateCS(uint GroupIndex : SV_GroupIndex, uint3 DispatchThreadId : SV_DispatchThreadID)
+{
+	if (GroupIndex==0)
+	{
+		GroupBase = 0;
+		LocalCounter = 0;
+	}
+	GroupMemoryBarrierWithGroupSync();
+
+#if PERMUTATION_GPU_DRIVEN == 1
+	const uint4 PageIndexResolutionAndOffset = PageIndexResolutionAndOffsetBuffer.Load(MacroGroupId);
+	const uint3 PageIndexResolution			 = PageIndexResolutionAndOffset.xyz;
+	const uint  PageIndexOffset				 = PageIndexResolutionAndOffset.w;
+	const uint  PageIndexCount				 = PageIndexResolution.x*PageIndexResolution.y*PageIndexResolution.z;
+#else
+	const uint3 PageIndexResolution	= CPU_PageIndexResolution;
+	const uint  PageIndexOffset		= CPU_PageIndexOffset;
+	const uint  PageIndexCount		= CPU_PageIndexCount;
+#endif
+
+	const uint GridIndex = DispatchThreadId.x + PageIndexOffset;
+	bool bIsValid = false;
+	if (DispatchThreadId.x < PageIndexCount)
+	{
+		bIsValid = PageIndexBuffer[GridIndex] > 0;
+	}
+
+	uint Offset = 0;
+	if (bIsValid)
+	{
+		InterlockedAdd(LocalCounter, 1u, Offset);
+	}
+	GroupMemoryBarrierWithGroupSync();
+
+	if (GroupIndex==0)
+	{
+		InterlockedAdd(PageIndexGlobalCounter[0], LocalCounter, GroupBase);
+	}
+	GroupMemoryBarrierWithGroupSync();
+
+	if (bIsValid)
+	{
+		const uint PageIndex = GroupBase + Offset;
+		PageIndexBuffer[GridIndex] = PageIndex < PageCount ? PageIndex : INVALID_VOXEL_PAGE_INDEX;
+
+		// Output the coordinates of the allocated page for indirect dispatch usage
+		const uint LinearIndex = DispatchThreadId.x;
+		const uint3 PageIndexCoord = IndexToCoord(LinearIndex, PageIndexResolution);
+		const uint IndexRelativeToMacroGroupIndexOffset = PageIndex - PageIndexGlobalCounter[1];
+		PageIndexCoordBuffer[PageIndexOffset + IndexRelativeToMacroGroupIndexOffset] = uint4(PageIndexCoord, MacroGroupId);
+	}
+	// Mark page index as invalid
+	// Insure that even if write more (due to larger dispatch count that needed), we do not stomp other instance group page index
+	else if (DispatchThreadId.x < PageIndexCount)
+	{
+		PageIndexBuffer[GridIndex] = INVALID_VOXEL_PAGE_INDEX;
+	}
+}
+#endif
+
+#if SHADER_ADDDESC
+
+float3 CPU_MinAABB;
+uint MacroGroupId;
+float3 CPU_MaxAABB;
+uint CPU_PageIndexOffset;
+int3 CPU_PageIndexResolution;
+
+Buffer<int> MacroGroupAABBBuffer;
+Buffer<uint4> PageIndexResolutionAndOffsetBuffer;
+RWStructuredBuffer<FPackedVirtualVoxelNodeDesc> OutNodeDescBuffer;
+
+[numthreads(1, 1, 1)]
+void AddDescCS(uint GroupIndex : SV_GroupIndex, uint3 DispatchThreadId : SV_DispatchThreadID)
+{
+	FVirtualVoxelNodeDesc Node;
+
+#if PERMUTATION_GPU_DRIVEN == 1
+	const uint4 PageIndexResolutionAndOffset = PageIndexResolutionAndOffsetBuffer.Load(MacroGroupId);
+	const FHairAABB Bound = ReadHairAABB(MacroGroupId, MacroGroupAABBBuffer);
+
+	Node.MinAABB = Bound.Min;
+	Node.MaxAABB = Bound.Max;
+	Node.PageIndexResolution = PageIndexResolutionAndOffset.xyz;
+	Node.PageIndexOffset = PageIndexResolutionAndOffset.w;
+#else
+	Node.MinAABB = CPU_MinAABB;
+	Node.MaxAABB = CPU_MaxAABB;
+	Node.PageIndexResolution = CPU_PageIndexResolution;
+	Node.PageIndexOffset = CPU_PageIndexOffset;
+#endif
+
+	FPackedVirtualVoxelNodeDesc PackedNode = PackVoxelNode(Node);
+	OutNodeDescBuffer[MacroGroupId] = PackedNode;
+}
+#endif
+
+
+#if SHADER_ADDINDIRECTBUFFER
+
+uint PageResolution;
+uint MacroGroupId;
+uint IndirectGroupSize;
+
+RWBuffer<uint> OutPageIndexGlobalCounter;
+RWBuffer<uint> OutIndirectArgsBuffer;
+
+[numthreads(1, 1, 1)]
+void AddIndirectBufferCS(uint GroupIndex : SV_GroupIndex, uint3 DispatchThreadId : SV_DispatchThreadID)
+{
+	const uint TotalAllocatedPageCount = OutPageIndexGlobalCounter[0];
+	const uint PrevTotalAllocatedPageCount = OutPageIndexGlobalCounter[1];
+	const uint AllocatedPageCount = TotalAllocatedPageCount - PrevTotalAllocatedPageCount;
+	const uint VoxelCountPerPage = PageResolution * PageResolution * PageResolution;
+
+	const uint MaxGroupCountPerDim = 65535;
+	const uint DispatchX = AllocatedPageCount % MaxGroupCountPerDim;
+	const uint DispatchY = ceil(AllocatedPageCount / float(MaxGroupCountPerDim));
+
+	const uint ArgsIndex = MacroGroupId * 3;
+	OutIndirectArgsBuffer[ArgsIndex]	 = DispatchX;
+	OutIndirectArgsBuffer[ArgsIndex + 1] = DispatchY;
+	OutIndirectArgsBuffer[ArgsIndex + 2] = ceil(VoxelCountPerPage / float(IndirectGroupSize));
+
+	// Copy current total to previous total
+	OutPageIndexGlobalCounter[1] = TotalAllocatedPageCount;
+}
+
+#endif
+
+
+#define IND_CLEAR_GROUP_SIZE_Z 64
+
+#if SHADER_INDPAGECLEARBUFFERGEN
+
+uint PageResolution;
+Buffer<uint> PageIndexGlobalCounter;
+RWBuffer<uint> OutIndirectArgsBuffer;
+
+[numthreads(1, 1, 1)]
+void VoxelIndPageClearBufferGenCS(uint GroupIndex : SV_GroupIndex, uint3 DispatchThreadId : SV_DispatchThreadID)
+{
+	const uint TotalAllocatedPageCount = PageIndexGlobalCounter[0];
+	const uint VoxelCountPerPage = PageResolution * PageResolution * PageResolution;
+
+	const uint MaxGroupCountPerDim = 65535;
+	const uint DispatchX = TotalAllocatedPageCount % MaxGroupCountPerDim;
+	const uint DispatchY = ceil(TotalAllocatedPageCount / float(MaxGroupCountPerDim));
+
+	OutIndirectArgsBuffer[0] = DispatchX;
+	OutIndirectArgsBuffer[1] = DispatchY;
+	OutIndirectArgsBuffer[2] = ceil(VoxelCountPerPage / float(IND_CLEAR_GROUP_SIZE_Z));
+}
+
+#endif
+
+
+#if SHADER_INDPAGECLEAR
+
+uint				VirtualVoxel_PageResolution;
+uint				VirtualVoxel_PageCountResolution;
+
+Buffer<uint4>		VirtualVoxel_PageIndexCoordBuffer;
+
+RWTexture3D<uint>	OutPageTexture;
+
+[numthreads(1, 1, IND_CLEAR_GROUP_SIZE_Z)]
+void VoxelIndPageClearCS(uint3 DispatchThreadId : SV_DispatchThreadID, uint3 GroupId : SV_GroupID, uint3 GroupThreadId : SV_GroupThreadID)
+{
+	const uint  LinearVoxelCoord = DispatchThreadId.z;
+
+	const uint VoxelCountPerPage = VirtualVoxel_PageResolution.x * VirtualVoxel_PageResolution.x * VirtualVoxel_PageResolution.x;
+	//	if (AllocatedPageIndex < TotalAllocatedPageCount)		// Not needed today because the XY size exactly matches the range (page count) we need to clear.
+	if(LinearVoxelCoord < VoxelCountPerPage)
+	{
+		const uint  MaxDispatchCountPerDim = 65535u;
+		const uint  AllocatedPageIndex = DispatchThreadId.x + DispatchThreadId.y * MaxDispatchCountPerDim;
+		const uint3 VoxelCoordOffset = IndexToCoord(LinearVoxelCoord, VirtualVoxel_PageResolution.xxx);
+
+		const uint PageIndex = AllocatedPageIndex; // PageIndexBuffer is not needed, we already know those tiles are allocated linearly in 3D within OutPageTexture.
+
+		const uint3 PageCoord = IndexToCoord(PageIndex, VirtualVoxel_PageCountResolution);
+		const int3 VoxelPageBase = PageCoord * VirtualVoxel_PageResolution;
+		const int3 VoxelCoord = VoxelPageBase + VoxelCoordOffset;
+
+		OutPageTexture[VoxelCoord] = 0;
+	}
+}
+
+#endif
+
diff --git a/Engine/Shaders/Private/HairStrands/HairStrandsVoxelPageCommonStruct.ush b/Engine/Shaders/Private/HairStrands/HairStrandsVoxelPageCommonStruct.ush
new file mode 100644
--- /dev/null
+++ b/Engine/Shaders/Private/HairStrands/HairStrandsVoxelPageCommonStruct.ush
@@ -0,0 +1,20 @@
+// Copyright Epic Games, Inc. All Rights Reserved.
+
+#pragma once
+
+struct FPackedVirtualVoxelNodeDesc  
+{
+	float3	MinAABB;
+	uint	PackedPageIndexResolution; // 8bits encoding
+	float3	MaxAABB;
+	uint	PageIndexOffset;
+}; 
+
+struct FVoxelizationViewInfo
+{
+	float4x4 WorldToClip;
+	float3   ViewForward;
+	float    Pad0;
+	uint2    RasterResolution;
+	float2   Pad1;
+};
\ No newline at end of file
