Copying //UE4/Dev-Physics to Main (//UE4/Main)

diff --git a/Engine/Shaders/Private/HairStrands/HairStrandsVoxelPageRayMarching.usf b/Engine/Shaders/Private/HairStrands/HairStrandsVoxelPageRayMarching.usf
new file mode 100644
--- /dev/null
+++ b/Engine/Shaders/Private/HairStrands/HairStrandsVoxelPageRayMarching.usf
@@ -0,0 +1,577 @@
+// Copyright Epic Games, Inc. All Rights Reserved.
+
+#include "HairStrandsVoxelPageCommonStruct.ush"
+
+#include "../Common.ush"
+#include "../CommonViewUniformBuffer.ush"
+#include "../SceneTextureParameters.ush"
+#include "../DeferredShadingCommon.ush"
+#include "../ShaderDrawDebug.ush"
+#include "../ShaderPrintCommon.ush"
+#include "HairStrandsVisibilityCommon.ush"
+#include "HairStrandsVoxelPageCommon.ush"
+#include "HairStrandsDebugCommon.ush" 
+
+float2 OutputResolution; 
+uint bDrawPage;
+uint MacroGroupId;
+uint MacroGroupCount;
+uint MaxTotalPageIndexCount;
+
+RWTexture2D<float4> OutputTexture;
+Buffer<uint> TotalValidPageCounter;
+
+float3 ReconstructWorldPositionFromDepth(float2 UV, float Depth)
+{
+	float2 ScreenPosition = (UV - View.ScreenPositionScaleBias.wz) / View.ScreenPositionScaleBias.xy;
+	float4 HomogeneousWorldPosition = mul(float4(ScreenPosition * Depth, Depth, 1), View.ScreenToWorld);
+	float3 WorldPosition = HomogeneousWorldPosition.xyz / HomogeneousWorldPosition.w;
+
+	return WorldPosition;
+}
+
+#define VOXEL_BLOCK 0 
+#define VOXEL_SMOOTH 1
+#define VOXEL_FILTERING VOXEL_BLOCK
+
+uint GetVoxelPageIndex(uint3 InPageIndexCoord, uint3 InPageIndexResolution, uint InLinearPageIndexOffset)
+{
+	const uint LinearPageIndexCoord = CoordToIndex(InPageIndexCoord, InPageIndexResolution, InLinearPageIndexOffset);
+	return VirtualVoxel.PageIndexBuffer.Load(LinearPageIndexCoord);
+}
+
+[numthreads(8, 8, 1)]
+void MainCS(uint3 DispatchThreadId : SV_DispatchThreadID)
+{
+	const float2 PixelCoord = DispatchThreadId.xy + 0.5f;
+	const float2 UV = PixelCoord / float2(OutputResolution);
+
+	// Reverse-Z
+	const float3 O = View.WorldCameraOrigin;
+	const float3 E = ReconstructWorldPositionFromDepth(UV, ConvertFromDeviceZ(0.0001f));
+	const float3 D = (E - O);
+
+	float  WorldDepth = 0;
+	float  ClipDepth = 0;
+
+	const FPackedVirtualVoxelNodeDesc PackedNode = VirtualVoxel.NodeDescBuffer[MacroGroupId];
+	const FVirtualVoxelNodeDesc NodeDesc = UnpackVoxelNode(PackedNode, VirtualVoxel.PageResolution);
+
+	// Draw macro group AABB & debug infos
+	if (uint(PixelCoord.x) == 0 && uint(PixelCoord.y) == 0)
+	{
+		const uint TotalPageIndexCount = MaxTotalPageIndexCount;
+		const uint TotalPageCount = VirtualVoxel.PageCountResolution.x * VirtualVoxel.PageCountResolution.y * VirtualVoxel.PageCountResolution.z;
+		const uint TotalAllocatedValidPageCount = TotalValidPageCounter.Load(0);
+		uint TotalAllocatedPageCount = 0;
+		for (uint MacroGroupIt = 0; MacroGroupIt < MacroGroupCount; ++MacroGroupIt)
+		{
+			const FPackedVirtualVoxelNodeDesc LocalPackedNode = VirtualVoxel.NodeDescBuffer[MacroGroupIt];
+			const FVirtualVoxelNodeDesc LocalNodeDesc = UnpackVoxelNode(LocalPackedNode, VirtualVoxel.PageResolution);
+			TotalAllocatedPageCount += LocalNodeDesc.PageIndexResolution.x * LocalNodeDesc.PageIndexResolution.y * LocalNodeDesc.PageIndexResolution.z;
+		}
+
+		AddAABB(NodeDesc.MinAABB, NodeDesc.MaxAABB, float4(1, 0.5f, 0, 1));
+
+		// General information
+		// General
+		//  Voxel World size        : 
+		//  Page Resolution         :
+		//  Page Count Resolution   :
+		//  Page Texture Resolution :
+		//
+		//  Total Allocatled Page Index   :
+		//  Total Page Index              :
+		//  Total Page Index Occupancy    :
+		//
+		//  Total Allocated Page          :
+		//  Total Page                    :
+		//  Total Page Occupancy          :
+		if (MacroGroupId == 0)
+		{
+			float2 StartPos = float2(50.0f, 100.0f ) / OutputResolution;
+			float2 Pos = StartPos;
+
+			Pos = ShaderPrintSymbol(Pos, _G_);
+			Pos = ShaderPrintSymbol(Pos, _E_);
+			Pos = ShaderPrintSymbol(Pos, _N_);
+			Pos = ShaderPrintSymbol(Pos, _E_);
+			Pos = ShaderPrintSymbol(Pos, _R_);
+			Pos = ShaderPrintSymbol(Pos, _A_);
+			Pos = ShaderPrintSymbol(Pos, _L_);
+			Pos = ShaderPrintNewline(Pos); Pos.x = StartPos.x;
+
+			Pos = ShaderPrintSymbol(Pos, _SPC_);
+			Pos = ShaderPrintSymbol(Pos, _V_);
+			Pos = ShaderPrintSymbol(Pos, _O_);
+			Pos = ShaderPrintSymbol(Pos, _X_);
+			Pos = ShaderPrintSymbol(Pos, _E_);
+			Pos = ShaderPrintSymbol(Pos, _L_);
+			Pos = ShaderPrintSymbol(Pos, _SPC_);
+			Pos = ShaderPrintSymbol(Pos, _W_);
+			Pos = ShaderPrintSymbol(Pos, _O_);
+			Pos = ShaderPrintSymbol(Pos, _R_);
+			Pos = ShaderPrintSymbol(Pos, _L_);
+			Pos = ShaderPrintSymbol(Pos, _D_);
+			Pos = ShaderPrintSymbol(Pos, _SPC_);
+			Pos = ShaderPrintSymbol(Pos, _S_);
+			Pos = ShaderPrintSymbol(Pos, _I_);
+			Pos = ShaderPrintSymbol(Pos, _Z_);
+			Pos = ShaderPrintSymbol(Pos, _E_);
+			Pos = ShaderPrintSymbol(Pos, _SPC_);
+			Pos = ShaderPrintSymbol(Pos, _SPC_);
+			Pos = ShaderPrint(Pos, VirtualVoxel.VoxelWorldSize);
+			Pos = ShaderPrintNewline(Pos); Pos.x = StartPos.x;
+
+			Pos = ShaderPrintSymbol(Pos, _SPC_);
+			Pos = ShaderPrintSymbol(Pos, _P_);
+			Pos = ShaderPrintSymbol(Pos, _A_);
+			Pos = ShaderPrintSymbol(Pos, _G_);
+			Pos = ShaderPrintSymbol(Pos, _E_);
+			Pos = ShaderPrintSymbol(Pos, _SPC_);
+			Pos = ShaderPrintSymbol(Pos, _R_);
+			Pos = ShaderPrintSymbol(Pos, _E_);
+			Pos = ShaderPrintSymbol(Pos, _S_);
+			Pos = ShaderPrintSymbol(Pos, _SPC_);
+			Pos = ShaderPrintSymbol(Pos, _SPC_);
+			Pos = ShaderPrintSymbol(Pos, _SPC_);
+			Pos = ShaderPrintSymbol(Pos, _SPC_);
+			Pos = ShaderPrintSymbol(Pos, _SPC_);
+			Pos = ShaderPrintSymbol(Pos, _SPC_);
+			Pos = ShaderPrintSymbol(Pos, _SPC_);
+			Pos = ShaderPrintSymbol(Pos, _SPC_);
+			Pos = ShaderPrintSymbol(Pos, _SPC_);
+			Pos = ShaderPrintSymbol(Pos, _SPC_);
+			Pos = ShaderPrint(Pos, VirtualVoxel.PageResolution);
+			Pos = ShaderPrintNewline(Pos); Pos.x = StartPos.x;
+
+			Pos = ShaderPrintSymbol(Pos, _SPC_);
+			Pos = ShaderPrintSymbol(Pos, _P_);
+			Pos = ShaderPrintSymbol(Pos, _A_);
+			Pos = ShaderPrintSymbol(Pos, _G_);
+			Pos = ShaderPrintSymbol(Pos, _E_);
+			Pos = ShaderPrintSymbol(Pos, _SPC_);
+			Pos = ShaderPrintSymbol(Pos, _C_);
+			Pos = ShaderPrintSymbol(Pos, _O_);
+			Pos = ShaderPrintSymbol(Pos, _U_);
+			Pos = ShaderPrintSymbol(Pos, _N_);
+			Pos = ShaderPrintSymbol(Pos, _T_);
+			Pos = ShaderPrintSymbol(Pos, _SPC_);
+			Pos = ShaderPrintSymbol(Pos, _R_);
+			Pos = ShaderPrintSymbol(Pos, _E_);
+			Pos = ShaderPrintSymbol(Pos, _S_);
+			Pos = ShaderPrintSymbol(Pos, _SPC_);
+			Pos = ShaderPrintSymbol(Pos, _SPC_);
+			Pos = ShaderPrintSymbol(Pos, _SPC_);
+			Pos = ShaderPrintSymbol(Pos, _SPC_);
+			Pos = ShaderPrint(Pos, VirtualVoxel.PageCountResolution.x);
+			Pos = ShaderPrint(Pos, VirtualVoxel.PageCountResolution.y);
+			Pos = ShaderPrint(Pos, VirtualVoxel.PageCountResolution.z);
+			Pos = ShaderPrintNewline(Pos); Pos.x = StartPos.x;
+
+			Pos = ShaderPrintSymbol(Pos, _SPC_);
+			Pos = ShaderPrintSymbol(Pos, _P_);
+			Pos = ShaderPrintSymbol(Pos, _A_);
+			Pos = ShaderPrintSymbol(Pos, _G_);
+			Pos = ShaderPrintSymbol(Pos, _E_);
+			Pos = ShaderPrintSymbol(Pos, _SPC_);
+			Pos = ShaderPrintSymbol(Pos, _T_);
+			Pos = ShaderPrintSymbol(Pos, _E_);
+			Pos = ShaderPrintSymbol(Pos, _X_);
+			Pos = ShaderPrintSymbol(Pos, _T_);
+			Pos = ShaderPrintSymbol(Pos, _U_);
+			Pos = ShaderPrintSymbol(Pos, _R_);
+			Pos = ShaderPrintSymbol(Pos, _E_);
+			Pos = ShaderPrintSymbol(Pos, _SPC_);
+			Pos = ShaderPrintSymbol(Pos, _R_);
+			Pos = ShaderPrintSymbol(Pos, _E_);
+			Pos = ShaderPrintSymbol(Pos, _S_);
+			Pos = ShaderPrintSymbol(Pos, _SPC_);
+			Pos = ShaderPrintSymbol(Pos, _SPC_);
+			Pos = ShaderPrint(Pos, VirtualVoxel.PageTextureResolution.x);
+			Pos = ShaderPrint(Pos, VirtualVoxel.PageTextureResolution.y);
+			Pos = ShaderPrint(Pos, VirtualVoxel.PageTextureResolution.z);
+			Pos = ShaderPrintNewline(Pos); Pos.x = StartPos.x;
+
+			// Total page indices stats
+			Pos = ShaderPrintNewline(Pos); Pos.x = StartPos.x;
+
+			Pos = ShaderPrintSymbol(Pos, _SPC_);
+			Pos = ShaderPrintSymbol(Pos, _T_);
+			Pos = ShaderPrintSymbol(Pos, _O_);
+			Pos = ShaderPrintSymbol(Pos, _T_);
+			Pos = ShaderPrintSymbol(Pos, _A_);
+			Pos = ShaderPrintSymbol(Pos, _L_);
+			Pos = ShaderPrintSymbol(Pos, _SPC_);
+			Pos = ShaderPrintSymbol(Pos, _A_);
+			Pos = ShaderPrintSymbol(Pos, _L_);
+			Pos = ShaderPrintSymbol(Pos, _L_);
+			Pos = ShaderPrintSymbol(Pos, _O_);
+			Pos = ShaderPrintSymbol(Pos, _C_);
+			Pos = ShaderPrintSymbol(Pos, _A_);
+			Pos = ShaderPrintSymbol(Pos, _T_);
+			Pos = ShaderPrintSymbol(Pos, _E_);
+			Pos = ShaderPrintSymbol(Pos, _D_);
+			Pos = ShaderPrintSymbol(Pos, _SPC_);
+			Pos = ShaderPrintSymbol(Pos, _P_);
+			Pos = ShaderPrintSymbol(Pos, _A_);
+			Pos = ShaderPrintSymbol(Pos, _G_);
+			Pos = ShaderPrintSymbol(Pos, _E_);
+			Pos = ShaderPrintSymbol(Pos, _SPC_);
+			Pos = ShaderPrintSymbol(Pos, _I_);
+			Pos = ShaderPrintSymbol(Pos, _N_);
+			Pos = ShaderPrintSymbol(Pos, _D_);
+			Pos = ShaderPrintSymbol(Pos, _E_);
+			Pos = ShaderPrintSymbol(Pos, _X_);
+			Pos = ShaderPrintSymbol(Pos, _SPC_);
+			Pos = ShaderPrintSymbol(Pos, _SPC_);
+			Pos = ShaderPrint(Pos, TotalAllocatedPageCount);
+			Pos = ShaderPrintNewline(Pos); Pos.x = StartPos.x;
+
+			Pos = ShaderPrintSymbol(Pos, _SPC_);
+			Pos = ShaderPrintSymbol(Pos, _T_);
+			Pos = ShaderPrintSymbol(Pos, _O_);
+			Pos = ShaderPrintSymbol(Pos, _T_);
+			Pos = ShaderPrintSymbol(Pos, _A_);
+			Pos = ShaderPrintSymbol(Pos, _L_);
+			Pos = ShaderPrintSymbol(Pos, _SPC_);
+			Pos = ShaderPrintSymbol(Pos, _P_);
+			Pos = ShaderPrintSymbol(Pos, _A_);
+			Pos = ShaderPrintSymbol(Pos, _G_);
+			Pos = ShaderPrintSymbol(Pos, _E_);
+			Pos = ShaderPrintSymbol(Pos, _SPC_);
+			Pos = ShaderPrintSymbol(Pos, _I_);
+			Pos = ShaderPrintSymbol(Pos, _N_);
+			Pos = ShaderPrintSymbol(Pos, _D_);
+			Pos = ShaderPrintSymbol(Pos, _E_);
+			Pos = ShaderPrintSymbol(Pos, _X_);
+			Pos = ShaderPrintSymbol(Pos, _SPC_);
+			Pos = ShaderPrintSymbol(Pos, _SPC_);
+			Pos = ShaderPrintSymbol(Pos, _SPC_);
+			Pos = ShaderPrintSymbol(Pos, _SPC_);
+			Pos = ShaderPrintSymbol(Pos, _SPC_);
+			Pos = ShaderPrintSymbol(Pos, _SPC_);
+			Pos = ShaderPrintSymbol(Pos, _SPC_);
+			Pos = ShaderPrintSymbol(Pos, _SPC_);
+			Pos = ShaderPrintSymbol(Pos, _SPC_);
+			Pos = ShaderPrintSymbol(Pos, _SPC_);
+			Pos = ShaderPrintSymbol(Pos, _SPC_);
+			Pos = ShaderPrintSymbol(Pos, _SPC_);
+			Pos = ShaderPrint(Pos, TotalPageIndexCount);
+			Pos = ShaderPrintNewline(Pos); Pos.x = StartPos.x;
+
+			Pos = ShaderPrintSymbol(Pos, _SPC_);
+			Pos = ShaderPrintSymbol(Pos, _P_);
+			Pos = ShaderPrintSymbol(Pos, _A_);
+			Pos = ShaderPrintSymbol(Pos, _G_);
+			Pos = ShaderPrintSymbol(Pos, _E_);
+			Pos = ShaderPrintSymbol(Pos, _SPC_);
+			Pos = ShaderPrintSymbol(Pos, _I_);
+			Pos = ShaderPrintSymbol(Pos, _N_);
+			Pos = ShaderPrintSymbol(Pos, _D_);
+			Pos = ShaderPrintSymbol(Pos, _E_);
+			Pos = ShaderPrintSymbol(Pos, _X_);
+			Pos = ShaderPrintSymbol(Pos, _SPC_);
+			Pos = ShaderPrintSymbol(Pos, _O_);
+			Pos = ShaderPrintSymbol(Pos, _C_);
+			Pos = ShaderPrintSymbol(Pos, _C_);
+			Pos = ShaderPrintSymbol(Pos, _U_);
+			Pos = ShaderPrintSymbol(Pos, _P_);
+			Pos = ShaderPrintSymbol(Pos, _A_);
+			Pos = ShaderPrintSymbol(Pos, _N_);
+			Pos = ShaderPrintSymbol(Pos, _C_);
+			Pos = ShaderPrintSymbol(Pos, _Y_);
+			Pos = ShaderPrintSymbol(Pos, _SPC_);
+			Pos = ShaderPrintSymbol(Pos, _SPC_);
+			Pos = ShaderPrintSymbol(Pos, _SPC_);
+			Pos = ShaderPrintSymbol(Pos, _SPC_);
+			Pos = ShaderPrintSymbol(Pos, _SPC_);
+			Pos = ShaderPrintSymbol(Pos, _SPC_);
+			Pos = ShaderPrintSymbol(Pos, _SPC_);
+			Pos = ShaderPrintSymbol(Pos, _SPC_);
+			Pos = ShaderPrint(Pos, TotalAllocatedPageCount / float(TotalPageIndexCount) * 100.f);
+			Pos = ShaderPrintNewline(Pos); Pos.x = StartPos.x;
+
+			// Total pages stats
+			Pos = ShaderPrintNewline(Pos); Pos.x = StartPos.x;
+
+			Pos = ShaderPrintSymbol(Pos, _SPC_);
+			Pos = ShaderPrintSymbol(Pos, _T_);
+			Pos = ShaderPrintSymbol(Pos, _O_);
+			Pos = ShaderPrintSymbol(Pos, _T_);
+			Pos = ShaderPrintSymbol(Pos, _A_);
+			Pos = ShaderPrintSymbol(Pos, _L_);
+			Pos = ShaderPrintSymbol(Pos, _SPC_);
+			Pos = ShaderPrintSymbol(Pos, _A_);
+			Pos = ShaderPrintSymbol(Pos, _L_);
+			Pos = ShaderPrintSymbol(Pos, _L_);
+			Pos = ShaderPrintSymbol(Pos, _O_);
+			Pos = ShaderPrintSymbol(Pos, _C_);
+			Pos = ShaderPrintSymbol(Pos, _A_);
+			Pos = ShaderPrintSymbol(Pos, _T_);
+			Pos = ShaderPrintSymbol(Pos, _E_);
+			Pos = ShaderPrintSymbol(Pos, _D_);
+			Pos = ShaderPrintSymbol(Pos, _SPC_);
+			Pos = ShaderPrintSymbol(Pos, _V_);
+			Pos = ShaderPrintSymbol(Pos, _A_);
+			Pos = ShaderPrintSymbol(Pos, _L_);
+			Pos = ShaderPrintSymbol(Pos, _I_);
+			Pos = ShaderPrintSymbol(Pos, _D_);
+			Pos = ShaderPrintSymbol(Pos, _SPC_);
+			Pos = ShaderPrintSymbol(Pos, _P_);
+			Pos = ShaderPrintSymbol(Pos, _A_);
+			Pos = ShaderPrintSymbol(Pos, _G_);
+			Pos = ShaderPrintSymbol(Pos, _E_);
+			Pos = ShaderPrintSymbol(Pos, _SPC_);
+			Pos = ShaderPrintSymbol(Pos, _SPC_);
+			Pos = ShaderPrint(Pos, TotalAllocatedValidPageCount);
+			Pos = ShaderPrintNewline(Pos); Pos.x = StartPos.x;
+
+			Pos = ShaderPrintSymbol(Pos, _SPC_);
+			Pos = ShaderPrintSymbol(Pos, _T_);
+			Pos = ShaderPrintSymbol(Pos, _O_);
+			Pos = ShaderPrintSymbol(Pos, _T_);
+			Pos = ShaderPrintSymbol(Pos, _A_);
+			Pos = ShaderPrintSymbol(Pos, _L_);
+			Pos = ShaderPrintSymbol(Pos, _SPC_);
+			Pos = ShaderPrintSymbol(Pos, _P_);
+			Pos = ShaderPrintSymbol(Pos, _A_);
+			Pos = ShaderPrintSymbol(Pos, _G_);
+			Pos = ShaderPrintSymbol(Pos, _E_);
+			Pos = ShaderPrintSymbol(Pos, _SPC_);
+			Pos = ShaderPrintSymbol(Pos, _SPC_);
+			Pos = ShaderPrintSymbol(Pos, _SPC_);
+			Pos = ShaderPrintSymbol(Pos, _SPC_);
+			Pos = ShaderPrintSymbol(Pos, _SPC_);
+			Pos = ShaderPrintSymbol(Pos, _SPC_);
+			Pos = ShaderPrintSymbol(Pos, _SPC_);
+			Pos = ShaderPrintSymbol(Pos, _SPC_);
+			Pos = ShaderPrintSymbol(Pos, _SPC_);
+			Pos = ShaderPrintSymbol(Pos, _SPC_);
+			Pos = ShaderPrintSymbol(Pos, _SPC_);
+			Pos = ShaderPrintSymbol(Pos, _SPC_);
+			Pos = ShaderPrintSymbol(Pos, _SPC_);
+			Pos = ShaderPrintSymbol(Pos, _SPC_);
+			Pos = ShaderPrintSymbol(Pos, _SPC_);
+			Pos = ShaderPrintSymbol(Pos, _SPC_);
+			Pos = ShaderPrintSymbol(Pos, _SPC_);
+			Pos = ShaderPrintSymbol(Pos, _SPC_);
+			Pos = ShaderPrint(Pos, TotalPageCount);
+			Pos = ShaderPrintNewline(Pos); Pos.x = StartPos.x;
+
+			Pos = ShaderPrintSymbol(Pos, _SPC_);
+			Pos = ShaderPrintSymbol(Pos, _P_);
+			Pos = ShaderPrintSymbol(Pos, _A_);
+			Pos = ShaderPrintSymbol(Pos, _G_);
+			Pos = ShaderPrintSymbol(Pos, _E_);
+			Pos = ShaderPrintSymbol(Pos, _SPC_);
+			Pos = ShaderPrintSymbol(Pos, _O_);
+			Pos = ShaderPrintSymbol(Pos, _C_);
+			Pos = ShaderPrintSymbol(Pos, _C_);
+			Pos = ShaderPrintSymbol(Pos, _U_);
+			Pos = ShaderPrintSymbol(Pos, _P_);
+			Pos = ShaderPrintSymbol(Pos, _A_);
+			Pos = ShaderPrintSymbol(Pos, _N_);
+			Pos = ShaderPrintSymbol(Pos, _C_);
+			Pos = ShaderPrintSymbol(Pos, _Y_);
+			Pos = ShaderPrintSymbol(Pos, _SPC_);
+			Pos = ShaderPrintSymbol(Pos, _SPC_);
+			Pos = ShaderPrintSymbol(Pos, _SPC_);
+			Pos = ShaderPrintSymbol(Pos, _SPC_);
+			Pos = ShaderPrintSymbol(Pos, _SPC_);
+			Pos = ShaderPrintSymbol(Pos, _SPC_);
+			Pos = ShaderPrintSymbol(Pos, _SPC_);
+			Pos = ShaderPrintSymbol(Pos, _SPC_);
+			Pos = ShaderPrintSymbol(Pos, _SPC_);
+			Pos = ShaderPrintSymbol(Pos, _SPC_);
+			Pos = ShaderPrintSymbol(Pos, _SPC_);
+			Pos = ShaderPrintSymbol(Pos, _SPC_);
+			Pos = ShaderPrintSymbol(Pos, _SPC_);
+			Pos = ShaderPrintSymbol(Pos, _SPC_);
+			Pos = ShaderPrint(Pos, TotalAllocatedValidPageCount / float(TotalPageCount) * 100.f);
+			Pos = ShaderPrintNewline(Pos); Pos.x = StartPos.x;
+		}
+
+		// MacroGroupId XX
+		//  Page Index Resolution: X Y Z (Linear/Total - %) 
+		//  Virtual Resolution   : X Y Z
+		//  Page Index Offset    : O
+		{
+			float2 StartPos = float2(50.0f, 350.0f + float(MacroGroupId) * 100.0f) / OutputResolution;
+			float2 Pos = StartPos;
+
+			Pos = ShaderPrintSymbol(Pos, _M_);
+			Pos = ShaderPrintSymbol(Pos, _A_);
+			Pos = ShaderPrintSymbol(Pos, _C_);
+			Pos = ShaderPrintSymbol(Pos, _R_);
+			Pos = ShaderPrintSymbol(Pos, _O_);
+			Pos = ShaderPrintSymbol(Pos, _SPC_);
+			Pos = ShaderPrintSymbol(Pos, _G_);
+			Pos = ShaderPrintSymbol(Pos, _R_);
+			Pos = ShaderPrintSymbol(Pos, _O_);
+			Pos = ShaderPrintSymbol(Pos, _U_);
+			Pos = ShaderPrintSymbol(Pos, _P_);
+			Pos = ShaderPrintSymbol(Pos, _SPC_);
+			Pos = ShaderPrintSymbol(Pos, _SPC_);
+			Pos = ShaderPrintSymbol(Pos, _SPC_);
+			Pos = ShaderPrintSymbol(Pos, _SPC_);
+			Pos = ShaderPrint(Pos, MacroGroupId);
+			Pos = ShaderPrintNewline(Pos); Pos.x = StartPos.x;
+
+			Pos = ShaderPrintSymbol(Pos, _SPC_);
+			Pos = ShaderPrintSymbol(Pos, _V_);
+			Pos = ShaderPrintSymbol(Pos, _I_);
+			Pos = ShaderPrintSymbol(Pos, _R_);
+			Pos = ShaderPrintSymbol(Pos, _T_);
+			Pos = ShaderPrintSymbol(Pos, _U_);
+			Pos = ShaderPrintSymbol(Pos, _A_);
+			Pos = ShaderPrintSymbol(Pos, _L_);
+			Pos = ShaderPrintSymbol(Pos, _SPC_);
+			Pos = ShaderPrintSymbol(Pos, _R_);
+			Pos = ShaderPrintSymbol(Pos, _E_);
+			Pos = ShaderPrintSymbol(Pos, _S_);
+			Pos = ShaderPrintSymbol(Pos, _SPC_);
+			Pos = ShaderPrintSymbol(Pos, _SPC_);
+			Pos = ShaderPrintSymbol(Pos, _SPC_);
+			Pos = ShaderPrintSymbol(Pos, _SPC_);
+			Pos = ShaderPrintSymbol(Pos, _SPC_);
+			Pos = ShaderPrintSymbol(Pos, _SPC_);
+			Pos = ShaderPrintSymbol(Pos, _SPC_);
+			Pos = ShaderPrint(Pos, NodeDesc.VirtualResolution.x);
+			Pos = ShaderPrint(Pos, NodeDesc.VirtualResolution.y);
+			Pos = ShaderPrint(Pos, NodeDesc.VirtualResolution.z);
+			Pos = ShaderPrintNewline(Pos); Pos.x = StartPos.x;
+
+			Pos = ShaderPrintSymbol(Pos, _SPC_);
+			Pos = ShaderPrintSymbol(Pos, _P_);
+			Pos = ShaderPrintSymbol(Pos, _A_);
+			Pos = ShaderPrintSymbol(Pos, _G_);
+			Pos = ShaderPrintSymbol(Pos, _E_);
+			Pos = ShaderPrintSymbol(Pos, _SPC_);
+			Pos = ShaderPrintSymbol(Pos, _I_);
+			Pos = ShaderPrintSymbol(Pos, _N_);
+			Pos = ShaderPrintSymbol(Pos, _D_);
+			Pos = ShaderPrintSymbol(Pos, _E_);
+			Pos = ShaderPrintSymbol(Pos, _X_);
+			Pos = ShaderPrintSymbol(Pos, _SPC_);
+			Pos = ShaderPrintSymbol(Pos, _R_);
+			Pos = ShaderPrintSymbol(Pos, _E_);
+			Pos = ShaderPrintSymbol(Pos, _S_);
+			Pos = ShaderPrintSymbol(Pos, _SPC_);
+			Pos = ShaderPrintSymbol(Pos, _SPC_);
+			Pos = ShaderPrintSymbol(Pos, _SPC_);
+			Pos = ShaderPrintSymbol(Pos, _SPC_);
+			Pos = ShaderPrint(Pos, NodeDesc.PageIndexResolution.x);
+			Pos = ShaderPrint(Pos, NodeDesc.PageIndexResolution.y);
+			Pos = ShaderPrint(Pos, NodeDesc.PageIndexResolution.z);
+			Pos = ShaderPrintSymbol(Pos, _SPC_);
+			Pos = ShaderPrintNewline(Pos); Pos.x = StartPos.x;
+
+			Pos = ShaderPrintSymbol(Pos, _SPC_);
+			Pos = ShaderPrintSymbol(Pos, _P_);
+			Pos = ShaderPrintSymbol(Pos, _A_);
+			Pos = ShaderPrintSymbol(Pos, _G_);
+			Pos = ShaderPrintSymbol(Pos, _E_);
+			Pos = ShaderPrintSymbol(Pos, _SPC_);
+			Pos = ShaderPrintSymbol(Pos, _I_);
+			Pos = ShaderPrintSymbol(Pos, _N_);
+			Pos = ShaderPrintSymbol(Pos, _D_);
+			Pos = ShaderPrintSymbol(Pos, _E_);
+			Pos = ShaderPrintSymbol(Pos, _X_);
+			Pos = ShaderPrintSymbol(Pos, _SPC_);
+			Pos = ShaderPrintSymbol(Pos, _O_);
+			Pos = ShaderPrintSymbol(Pos, _F_);
+			Pos = ShaderPrintSymbol(Pos, _F_);
+			Pos = ShaderPrintSymbol(Pos, _S_);
+			Pos = ShaderPrintSymbol(Pos, _E_);
+			Pos = ShaderPrintSymbol(Pos, _T_);
+			Pos = ShaderPrintSymbol(Pos, _SPC_);
+			Pos = ShaderPrint(Pos, NodeDesc.PageIndexOffset);
+			Pos = ShaderPrintNewline(Pos); Pos.x = StartPos.x;
+
+			Pos = ShaderPrintNewline(Pos); Pos.x = StartPos.x;
+		}
+	}
+
+	const float2 HitT = LineBoxIntersect(O, E, NodeDesc.MinAABB, NodeDesc.MaxAABB);
+	float3 HitColor = float3(0,0,0);
+	float AccDensity = 0;
+	bool bIsValid = false;
+	if (HitT.x < HitT.y)
+	{
+		// Use different offseting to see either: 
+		// * Individual voxel (for debugging mainly)
+		// * Smoothed voxel (as used for transmittance computation)
+	#if VOXEL_FILTERING == VOXEL_BLOCK
+		const float3 Offset = 0;
+		const float Dither = InterleavedGradientNoise(PixelCoord.xy, 1);
+	#endif
+	#if VOXEL_FILTERING == VOXEL_SMOOTH
+		const float3 Offset = float3(InterleavedGradientNoise(PixelCoord.xy, 1), InterleavedGradientNoise(PixelCoord.xy, 2), InterleavedGradientNoise(PixelCoord.xy, 3));
+		const float Dither = 0;
+	#endif
+
+
+		// Intersect the indirection table to see if a brick is empty or not
+		if (bDrawPage)
+		{
+			const uint MaxStep = 256;
+			const float Delta = (HitT.y-HitT.x)/float(MaxStep);
+			for (uint StepIt=0; StepIt<MaxStep; ++StepIt)
+			{			
+				const float3 HitP = Offset + O + ((Dither+StepIt)*Delta + HitT.x) * D;
+				const int3 PageIndexCoord = PositionToCoord(HitP, NodeDesc.MinAABB, NodeDesc.MaxAABB, NodeDesc.PageIndexResolution);
+				const uint PageIndex = GetVoxelPageIndex(PageIndexCoord, NodeDesc.PageIndexResolution, NodeDesc.PageIndexOffset);
+				if (PageIndex != INVALID_VOXEL_PAGE_INDEX)
+				{
+					const uint3 PageCoord = IndexToCoord(PageIndex, VirtualVoxel.PageCountResolution);
+
+					AccDensity = 1;
+					bIsValid = true;
+					HitColor = PageCoord / 12.f;
+					break;
+				}
+			}
+		}
+		else
+		{
+			const uint MipLevel = 0; //clamp(View.GeneralPurposeTweak, 0, 6);
+			const uint MaxStep = 256;
+			const float Delta = (HitT.y - HitT.x) / float(MaxStep);
+			for (uint StepIt = 0; StepIt < MaxStep; ++StepIt)
+			{
+				const float3 HitP = Offset + O + ((Dither + StepIt)*Delta + HitT.x) * D;
+				const int3 VolumeCoord = PositionToCoord(HitP, NodeDesc.MinAABB, NodeDesc.MaxAABB, NodeDesc.VirtualResolution);
+				const int3 PageIndexCoord = VolumeCoord / VirtualVoxel.PageResolution;
+				const uint PageIndex = GetVoxelPageIndex(PageIndexCoord, NodeDesc.PageIndexResolution, NodeDesc.PageIndexOffset);
+				if (PageIndex != INVALID_VOXEL_PAGE_INDEX)
+				{
+					const uint3 PageCoord		= IndexToCoord(PageIndex, VirtualVoxel.PageCountResolution);
+					const int3 VoxelPageBase	= PageCoord * VirtualVoxel.PageResolution;
+					const int3 VoxelPageOffset	= VolumeCoord - PageIndexCoord * VirtualVoxel.PageResolution;
+					const int3 VoxelPageCoord	= VoxelPageBase + VoxelPageOffset;
+
+					const float DensityScale	= 1;
+					const FHairTraversalResult Result = GetHairVirtualVoxelDensity(VoxelPageCoord, VirtualVoxel.PageTexture, MipLevel, DensityScale);
+					const float Density = Result.HairCount + (Result.Visibility > 0 ? 0 : 1000000000);
+					if (Density > 0)
+					{
+						AccDensity = max(AccDensity, float(Density));
+						bIsValid = true;
+					}
+
+				}
+			}
+		}
+	}
+
+	if (!bDrawPage)
+	{
+		HitColor = GetColorCode(saturate(AccDensity));
+	}
+
+	if (bIsValid)
+	{
+		OutputTexture[DispatchThreadId.xy] = float4(HitColor,1);
+	}
+}
diff --git a/Engine/Shaders/Private/HairStrands/HairStrandsVoxelPageTraversal.ush b/Engine/Shaders/Private/HairStrands/HairStrandsVoxelPageTraversal.ush
new file mode 100644
--- /dev/null
+++ b/Engine/Shaders/Private/HairStrands/HairStrandsVoxelPageTraversal.ush
@@ -0,0 +1,546 @@
+// Copyright Epic Games, Inc. All Rights Reserved.
+
+#pragma once
+
+
+// VOXEL_TRAVERSAL_TYPE needs to be defined for selecting which type of traversal needs to be used
+// Definition of traversal types are defined in HairStrandsVoxelPageCommon.ush
+#ifndef VOXEL_TRAVERSAL_TYPE
+#define VOXEL_TRAVERSAL_TYPE VOXEL_TRAVERSAL_NONE
+#endif
+
+#ifndef VOXEL_TRAVERSAL_DEBUG
+#define VOXEL_TRAVERSAL_DEBUG 0
+#else
+#include "../ShaderDrawDebug.ush"
+#endif
+
+////////////////////////////////////////////////////////////////////////////////////////////////////////////////////////////////
+
+#if VOXEL_TRAVERSAL_DEBUG 
+void AddLine(bool bEnabled, float3 P0, float3 P1)
+{
+	if (!bEnabled)
+		return;
+	const float4 DebugColor = float4(1, 1, 0, 1);
+	AddLine(P0, P1, DebugColor, DebugColor);
+}
+
+void AddMacroGroupAABB(bool bEnabled, float3 P0, float3 P1)
+{
+	if (!bEnabled)
+		return;
+	AddAABB(P0, P1, float4(1, 0, 1, 1));
+}
+
+void AddPageAABB(bool bEnabled, float3 P0, float3 P1, bool bIsPageValid)
+{
+	if (!bEnabled)
+		return;
+	const float4 DebugColor = bIsPageValid ? float4(0, 1, 1, 1) : float4(1, 0, 0, 1);
+	AddAABB(P0, P1, DebugColor);
+}
+
+void AddStepAABB(bool bEnabled, float3 P0, float3 P1, float3 Color)
+{
+	if (!bEnabled)
+		return;
+	const float4 DebugColor = float4(Color, 1);
+	AddAABB(P0, P1, DebugColor);
+}
+
+void AddAABBVolume(bool bEnabled, float3 P0, float3 P1)
+{
+	if (!bEnabled)
+		return;
+	const float4 DebugColor = float4(1,1,0, 1);
+	AddAABB(P0, P1, DebugColor);
+}
+
+#endif
+
+////////////////////////////////////////////////////////////////////////////////////////////////////////////////////////////////
+
+bool ShouldStopTraversal(const FHairTraversalResult InResult, float InHairCountThreshold)
+{
+	// Change this function once we have better result/control over the opaque filtering
+	return InResult.Visibility == 0 || (InHairCountThreshold != 0 && InResult.HairCount > InHairCountThreshold);
+}
+
+float3 GetHairVoxelJitter(uint2 PixelCoord, uint Seed)
+{
+	return float3(
+		InterleavedGradientNoise(PixelCoord.xy, Seed),
+		InterleavedGradientNoise(PixelCoord.xy, Seed * 117),
+		InterleavedGradientNoise(PixelCoord.xy, Seed * 7901));
+}
+
+/////////////////////////////////////////////////////////////////////////////////////////////////////////////////////////////////
+// Non-Sparse voxel traversal routines
+#if VOXEL_TRAVERSAL_TYPE ==  VOXEL_TRAVERSAL_LEGACY
+FHairTraversalResult ComputeHairCountVoxel(
+	float3 WorldPosition0,
+	float3 WorldPosition1,
+	float3 InVoxelAABBMin,
+	float3 InVoxelAABBMax,
+	float InVoxelResolution,
+	Texture3D<uint> InDensityTexture,
+	float DensityScale,
+	float CountThreshold,
+	float DistanceThreshold)
+{
+	FHairTraversalResult Out = InitHairTraversalResult();
+
+	const float3 VolumeBoundSize = InVoxelAABBMax - InVoxelAABBMin;
+	const float3 VolumeBoundSizeInv = 1.0f / (InVoxelAABBMax - InVoxelAABBMin);
+	const float VoxelWorldSize = VolumeBoundSize.x / InVoxelResolution;				// Only using voxel size along x for now...
+
+	const float2 HitT = LineBoxIntersect(WorldPosition0, WorldPosition1, InVoxelAABBMin, InVoxelAABBMax);
+	if (HitT.x < HitT.y)
+	{
+		// Count the number of fibers which are within a cylinder defined by the voxel size, 
+		// and the distance between the origin and the extent of the volume
+		// This assumes that the voxel volume is cubic (i.e. equal dimensions on all sides)
+		float3 O = lerp(WorldPosition0, WorldPosition1, HitT.xxx);
+		float3 E = lerp(WorldPosition0, WorldPosition1, HitT.yyy);
+		const float OELength = min(length(E - O), DistanceThreshold);
+
+		// Step according to voxel size
+		const float StepCount = float(min(ceil(OELength / VoxelWorldSize), 256u));	// TO TRY: sample with fixed count budget and jitter temporally
+		const float3 D = normalize(WorldPosition1 - WorldPosition0) * VoxelWorldSize;
+
+		for (float StepIt = 0.0f; StepIt < StepCount; ++StepIt)
+		{
+			const float3 HitP = O + StepIt * D;
+			const float3 UVW = (HitP - InVoxelAABBMin) * VolumeBoundSizeInv;
+			const int3 VolumeCoord = clamp(int3(UVW * InVoxelResolution), int3(0, 0, 0), int3(InVoxelResolution.xxx - 1));
+			const FHairTraversalResult StepResult = GetHairVoxelDensity(VolumeCoord, InDensityTexture, 0, DensityScale);
+			Acc(Out, StepResult);
+
+			if (ShouldStopTraversal(Out, CountThreshold))
+				return Out;
+		}
+	}
+
+	return Out;
+}
+
+FHairTraversalResult ComputeHairCountVoxelCone(
+	float3 WorldPosition0,
+	float3 WorldPosition1,
+	float3 InVoxelAABBMin,
+	float3 InVoxelAABBMax,
+	float InVoxelResolution,
+	Texture3D<uint> InDensityTexture,
+	float InVoxelDensityScale,
+	float CountThreshold,
+	float DistanceThreshold,
+	float TanConeAngle)
+{
+	FHairTraversalResult Out = InitHairTraversalResult();
+
+	// Reverse-Z
+	const float ConservativeAABBDiagonal = (InVoxelAABBMax.x - InVoxelAABBMin.x) * 3;
+	const float3 O = WorldPosition0;
+	const float3 E = WorldPosition1;
+	const float3 D = E - O;
+	const float3 NormD = normalize(D);
+
+	// The rest of the code assumes O is within the AABB
+	const bool bIsInAABB = IsInVoxelBounds(O, InVoxelAABBMin, InVoxelAABBMax);
+	if (!bIsInAABB)
+		return Out;
+
+	const float2 HitT = LineBoxIntersect(O, E, InVoxelAABBMin, InVoxelAABBMax);
+	if (HitT.x < HitT.y)
+	{
+		const float VolumeSize = (InVoxelAABBMax.x - InVoxelAABBMin.x);
+		const float VoxelSize = VolumeSize / InVoxelResolution;
+
+		// Count the number of fibers which are within a cylinder defined by the voxel size, 
+		// and the distance between the origin and the extent of the volume
+		// This assumes that the voxel volume is cubic (i.e. equal dimensions on all sides)
+		const float DistanceThreshold2 = DistanceThreshold * DistanceThreshold;
+		const float Distance = min(DistanceThreshold, length(HitT.y * D));
+		const uint MaxStepCount = min(ceil(Distance / VoxelSize), 256u);
+		float T = 0;
+		float MipLevel = 0;
+		float MipDensityScale = 1;
+		for (uint StepIt = 0; StepIt < MaxStepCount; ++StepIt)
+		{
+			float3 StepOffset = 0;
+
+			const float3 RaymarchOffset = T * NormD;
+			const float3 HitP = O + RaymarchOffset + StepOffset;
+			const float3 UVW = (HitP - InVoxelAABBMin) / (InVoxelAABBMax - InVoxelAABBMin);
+
+			const float Distance2 = dot(RaymarchOffset, RaymarchOffset);
+			if (Distance2 > DistanceThreshold2)
+				return Out;
+
+			if (!IsInVoxelBounds(UVW, 0, 1))
+				return Out;
+
+			const uint MipVoxelResolution = uint(InVoxelResolution) >> uint(MipLevel);
+			const int3 VolumeCoord = clamp(int3(UVW * MipVoxelResolution), int3(0, 0, 0), int3(MipVoxelResolution.xxx - 1));
+			const FHairTraversalResult StepResult = GetHairVoxelDensity(VolumeCoord, InDensityTexture, MipLevel, InVoxelDensityScale * MipDensityScale);
+			Acc(Out, StepResult);
+
+			if (ShouldStopTraversal(Out, CountThreshold))
+				return Out;
+
+			const float VoxelRadius = T * TanConeAngle;
+			const float NumVoxel = VolumeSize / max(2 * VoxelRadius, VoxelSize);
+
+			const float ResolutionBias = InVoxelResolution / NumVoxel;
+			float MipBias = log2(max(0.5f, ResolutionBias));
+			MipBias = floor(MipBias);
+			MipLevel = MipBias;
+
+			const float MipVoxelSize = VoxelSize * (MipLevel + 1);
+			{
+				MipDensityScale = VoxelSize * rcp(MipVoxelSize);
+				MipDensityScale = MipDensityScale * MipDensityScale * MipDensityScale;
+			}
+			T += MipVoxelSize;
+		}
+	}
+
+	return Out;
+}
+#endif
+
+/////////////////////////////////////////////////////////////////////////////////////////////////////////////////////////////////
+
+struct FHairTraversalSettings
+{
+	float	DensityScale;
+	float	CountThreshold;
+	float	DistanceThreshold;
+	float	SteppingScale;		// Rate at which the raymarching step increase
+	float3	Random;
+	float	TanConeAngle;
+	bool	bDebugEnabled;
+};
+
+FHairTraversalSettings InitHairTraversalSettings()
+{
+	FHairTraversalSettings Out;
+	Out.DensityScale = 1;
+	Out.CountThreshold = 0;
+	Out.DistanceThreshold = 100000.f;
+	Out.SteppingScale = 1;
+	Out.Random = 0.5f;
+	Out.TanConeAngle = 0;
+	Out.bDebugEnabled = false;
+	return Out;
+}
+
+/////////////////////////////////////////////////////////////////////////////////////////////////////////////////////////////////
+// Sparse voxel traversal routines
+FHairTraversalResult ComputeHairCountVirtualVoxel(
+	float3 WorldPosition0,
+	float3 WorldPosition1,
+	FVirtualVoxelCommonDesc InCommonDesc,
+	FVirtualVoxelNodeDesc InNodeDesc,
+	Buffer<uint> InPageIndexBuffer,
+	Texture3D<uint> InPageTexture,
+	FHairTraversalSettings InSettings)
+	////////////////////////////////////////////////////////////////////////////////////////////////////////////////////////////////
+#if VOXEL_TRAVERSAL_TYPE == VOXEL_TRAVERSAL_NONE
+{
+	Error: VOXEL TRAVERSAL TYPE is undefined
+}
+#endif
+#if VOXEL_TRAVERSAL_TYPE == VOXEL_TRAVERSAL_LEGACY
+; // This function is never used. When VOXEL_TRAVERSAL_LEGACY is defined, the old traversal (with plain voxel) is used instead
+#endif
+#if VOXEL_TRAVERSAL_TYPE == VOXEL_TRAVERSAL_LINEAR || VOXEL_TRAVERSAL_TYPE == VOXEL_TRAVERSAL_LINEAR_MIPMAP
+{
+	const float MipLevelDistanceScale = 1.0f / InCommonDesc.VoxelWorldSize;
+
+	FHairTraversalResult Out = InitHairTraversalResult();
+
+	int3 CurrentPageIndexCoord = -1;
+	bool bIsPageValid = false;
+	uint3 PageCoord = 0;
+
+	const float2 HitT = LineBoxIntersect(WorldPosition0, WorldPosition1, InNodeDesc.MinAABB, InNodeDesc.MaxAABB);
+	if (HitT.x < HitT.y)
+	{
+		// Count the number of fibers which are within a cylinder defined by the voxel size, 
+		// and the distance between the origin and the extent of the volume
+		// This assumes that the voxel volume is cubic (i.e. equal dimensions on all sides)
+		const float3 O = lerp(WorldPosition0, WorldPosition1, HitT.xxx);
+		const float3 E = lerp(WorldPosition0, WorldPosition1, HitT.yyy);
+		const float OELength = min(length(E - O), InSettings.DistanceThreshold);
+		#if VOXEL_TRAVERSAL_DEBUG
+		if (InSettings.bDebugEnabled)
+		{
+			AddLine(true, O, E);
+			AddAABBVolume(true, InNodeDesc.MinAABB, InNodeDesc.MaxAABB);
+		}
+		#endif
+
+		// Step according to voxel size
+		const float3 D = normalize(E - O) * InCommonDesc.VoxelWorldSize;
+		const float MaxStep = float(min(ceil(OELength / InCommonDesc.VoxelWorldSize), 1024u));
+		const float DeltaWorld = OELength / float(MaxStep);
+		float StepScale = 1.0f;
+		float3 PreviousP = O;
+
+		const float3 RandomStepJitter = InSettings.Random.xyz * 2 - 1; // [-1..1]
+		for (float StepIt = 0.0f; StepIt < MaxStep; StepIt += StepScale)
+		{
+			const float SteppingWorldSize = max(DeltaWorld * StepScale, InSettings.TanConeAngle * StepIt * InCommonDesc.VoxelWorldSize);
+			const float3 HitP = O + StepIt * D + RandomStepJitter * SteppingWorldSize * 0.5f;
+
+			#if VOXEL_TRAVERSAL_DEBUG
+			AddStepAABB(InSettings.bDebugEnabled, PreviousP, HitP, float3(0, 1, 0));
+			#endif
+
+			const int3 VolumeCoord = PositionToCoord(HitP, InNodeDesc.MinAABB, InNodeDesc.MaxAABB, InNodeDesc.VirtualResolution);
+			const int3 PageIndexCoord = VolumeCoord / InCommonDesc.PageResolution;
+
+			// Update page index only when needed
+			if (PageIndexCoord.x != CurrentPageIndexCoord.x ||
+				PageIndexCoord.y != CurrentPageIndexCoord.y ||
+				PageIndexCoord.z != CurrentPageIndexCoord.z)
+			{
+				CurrentPageIndexCoord = PageIndexCoord;
+				const uint LinearPageIndexCoord = CoordToIndex(PageIndexCoord, InNodeDesc.PageIndexResolution, InNodeDesc.PageIndexOffset);
+				const uint PageIndex = InPageIndexBuffer.Load(LinearPageIndexCoord);
+
+				bIsPageValid = PageIndex != INVALID_VOXEL_PAGE_INDEX;
+				if (bIsPageValid)
+				{
+					PageCoord = IndexToCoord(PageIndex, InCommonDesc.PageCountResolution);
+				}
+			}
+
+			if (bIsPageValid)
+			{
+				const int3 VoxelPageBase = PageCoord * InCommonDesc.PageResolution;
+				const int3 VoxelPageOffset = VolumeCoord - PageIndexCoord * InCommonDesc.PageResolution;
+				const int3 VoxelPageCoord = VoxelPageBase + VoxelPageOffset;
+				float HairCountScale = 1;
+				float MipLevel = 0.0f;
+				#if VOXEL_TRAVERSAL_TYPE == VOXEL_TRAVERSAL_LINEAR_MIPMAP
+				HairCountScale = SteppingWorldSize * MipLevelDistanceScale;
+				MipLevel = log2(SteppingWorldSize * MipLevelDistanceScale);
+				#endif
+
+				const FHairTraversalResult StepResult = GetHairVirtualVoxelDensity(VoxelPageCoord, InPageTexture, uint(MipLevel), InSettings.DensityScale * HairCountScale);
+				Acc(Out, StepResult);
+
+				#if VOXEL_TRAVERSAL_DEBUG
+				if (InSettings.bDebugEnabled)
+				{
+					const float3 WorldOffset = PageIndexCoord * InCommonDesc.PageResolution * InCommonDesc.VoxelWorldSize;
+					const float MipVoxelWorldSize = InCommonDesc.VoxelWorldSize * InCommonDesc.PageResolution / (InCommonDesc.PageResolution >> uint(MipLevel));
+
+					const uint3 MinCoord = VoxelPageOffset >> uint(MipLevel);
+					const uint3 MaxCoord = MinCoord+1;
+					const float3 FetchMinAABB = InNodeDesc.MinAABB + WorldOffset + MinCoord * MipVoxelWorldSize;
+					const float3 FetchMaxAABB = InNodeDesc.MinAABB + WorldOffset + MaxCoord * MipVoxelWorldSize;
+					AddStepAABB(true, FetchMinAABB, FetchMaxAABB, StepResult.Visibility > 0.5f ? float3(0, 0.5f, 1) : float3(1, 0.5f, 0));
+				}
+				#endif
+
+				if (ShouldStopTraversal(Out, InSettings.CountThreshold))
+					return Out;
+			}
+			#if VOXEL_TRAVERSAL_TYPE == VOXEL_TRAVERSAL_LINEAR_MIPMAP
+			StepScale = min(InCommonDesc.PageResolution, StepScale * InSettings.SteppingScale);
+			#endif
+			PreviousP = HitP;
+		}
+
+	}
+
+	return Out;
+}
+#endif
+////////////////////////////////////////////////////////////////////////////////////////////////////////////////////////////////
+#if VOXEL_TRAVERSAL_TYPE == VOXEL_TRAVERSAL_SPARSE_LINEAR || VOXEL_TRAVERSAL_TYPE == VOXEL_TRAVERSAL_SPARSE_MIPMAP
+{
+	FHairTraversalResult Out = InitHairTraversalResult();
+
+	int3 CurrentPageIndexCoord = -1;
+	bool bIsPageValid = false;
+	uint3 PageCoord = 0;
+
+	#if VOXEL_TRAVERSAL_DEBUG
+	AddMacroGroupAABB(InSettings.bDebugEnabled, InNodeDesc.MinAABB, InNodeDesc.MaxAABB);
+	AddLine(InSettings.bDebugEnabled, WorldPosition0, WorldPosition1);
+	#endif
+
+	const float2 HitT = LineBoxIntersect(WorldPosition0, WorldPosition1, InNodeDesc.MinAABB, InNodeDesc.MaxAABB);
+	if (HitT.x < HitT.y)
+	{
+		// Count the number of fibers which are within a cylinder defined by the voxel size, 
+		// and the distance between the origin and the extent of the volume
+		// This assumes that the voxel volume is cubic (i.e. equal dimensions on all sides)
+		const float3 O = WorldPosition0; // lerp(WorldPosition, IntersectEndPoint, HitT.xxx);
+		const float3 E = WorldPosition1; // lerp(WorldPosition, IntersectEndPoint, HitT.yyy);
+		const float OELength = min(length(E - O), InSettings.DistanceThreshold);
+
+	
+		// Init to 1 or -1 depending of the orientation of stepping
+		const float3 UNormD = WorldPosition1 - WorldPosition0;
+		const int3 Step = sign(UNormD);
+
+		// Step according to voxel size
+		const float3 D = normalize(UNormD) * InCommonDesc.VoxelWorldSize;
+
+		float t = HitT.x;
+
+		// this is slop coeff for each axis: how far we need to move in units of t for each axis
+		const float PageWorldSize = InCommonDesc.PageResolution * InCommonDesc.VoxelWorldSize;
+		const float3 tDelta = Step * PageWorldSize / UNormD;
+
+		// Init to the starting voxel
+		int3 PageIndexCoord = -1;
+		float3 tMax = 0;
+		{
+			const float3 HitP = O + HitT.x * UNormD;
+
+			const float Epsilon = 0.000001f;
+			const float3 Coords = clamp(
+				saturate((HitP - InNodeDesc.MinAABB) / (InNodeDesc.MaxAABB - InNodeDesc.MinAABB)) * InNodeDesc.PageIndexResolution,
+				0,
+				InNodeDesc.PageIndexResolution - Epsilon);
+
+			const float3 FractCoords = max(Step, 0) - Step * frac(Coords);
+			tMax = FractCoords * tDelta;
+
+			PageIndexCoord = clamp(int3(Coords), uint3(0, 0, 0), InNodeDesc.PageIndexResolution-1);
+		}
+
+		// Page stepping is the walking quantity (i.e. number of voxel) for fine ray-marching within a valid page
+		float PageStepping = 1;
+
+		const uint LoopCount = 256u;
+		for (uint LoopIt = 0; LoopIt < LoopCount; ++LoopIt)
+		{
+			const bool bIsInside =
+				PageIndexCoord.x >= 0 && 
+				PageIndexCoord.y >= 0 && 
+				PageIndexCoord.z >= 0 &&
+				PageIndexCoord.x < int(InNodeDesc.PageIndexResolution.x) && 
+				PageIndexCoord.y < int(InNodeDesc.PageIndexResolution.y) && 
+				PageIndexCoord.z < int(InNodeDesc.PageIndexResolution.z);
+			if (!bIsInside)
+			{
+				return Out;
+			}
+
+			const uint LinearPageIndexCoord = CoordToIndex(PageIndexCoord, InNodeDesc.PageIndexResolution, InNodeDesc.PageIndexOffset);
+			const uint PageIndex = InPageIndexBuffer.Load(LinearPageIndexCoord);
+
+			const bool bIsPageValid = PageIndex != INVALID_VOXEL_PAGE_INDEX;
+			#if VOXEL_TRAVERSAL_DEBUG
+			{
+				const float3 PageMinAABB = PageIndexCoord * InCommonDesc.PageResolution * InCommonDesc.VoxelWorldSize + InNodeDesc.MinAABB;
+				const float3 PageMaxAABB = (PageIndexCoord+float3(1,1,1)) * InCommonDesc.PageResolution * InCommonDesc.VoxelWorldSize + InNodeDesc.MinAABB;
+				AddPageAABB(InSettings.bDebugEnabled, PageMinAABB, PageMaxAABB, bIsPageValid);
+			}
+			#endif
+
+			if (bIsPageValid)
+			{
+				const uint3 PageCoord = IndexToCoord(PageIndex, InCommonDesc.PageCountResolution);
+
+				float3 InnerO = O + t * UNormD;
+				const uint MaxInnerStep = InCommonDesc.PageResolution * 1.75f; // ~ Diagonal step count
+
+				#if VOXEL_TRAVERSAL_TYPE == VOXEL_TRAVERSAL_SPARSE_LINEAR
+				const uint IntPageStepping = 1;
+				const uint MipLevel = 0;
+				#endif			
+
+				#if VOXEL_TRAVERSAL_TYPE == VOXEL_TRAVERSAL_SPARSE_MIPMAP
+				const uint IntPageStepping = uint(PageStepping);
+				const uint MipLevel = log2(IntPageStepping);
+				#endif
+
+				const uint MipPageResolution = InCommonDesc.PageResolution;
+				const int3 VoxelPageBase = PageCoord * MipPageResolution;
+
+				for (uint InnerStepIt = 0; InnerStepIt < MaxInnerStep; InnerStepIt += IntPageStepping)
+				{
+					const float3 InnerHitP = InnerO + D * InnerStepIt;
+					const int3 VolumeCoord = PositionToCoordUnclampled(InnerHitP, InNodeDesc.MinAABB, InNodeDesc.MaxAABB, InNodeDesc.VirtualResolution);
+					const int3 VoxelPageOffset = VolumeCoord - PageIndexCoord * InCommonDesc.PageResolution;
+
+					#if VOXEL_TRAVERSAL_DEBUG
+					if (InSettings.bDebugEnabled)
+					{
+						if (InnerStepIt==0)
+							AddStepAABB(InSettings.bDebugEnabled, InnerHitP - D * 2, InnerHitP + D*2, float3(0,1,0));
+						else
+							AddStepAABB(InSettings.bDebugEnabled, InnerHitP, InnerHitP + D, float3(1, 0, 1));
+					}
+					#endif
+
+					const bool bIsInsideInner =
+						VoxelPageOffset.x >= 0 && 
+						VoxelPageOffset.y >= 0 && 
+						VoxelPageOffset.z >= 0 &&
+						VoxelPageOffset.x < int(MipPageResolution) &&
+						VoxelPageOffset.y < int(MipPageResolution) &&
+						VoxelPageOffset.z < int(MipPageResolution);
+					if (!bIsInsideInner)
+					{
+						break;
+					}
+
+					const int3 VoxelPageCoord = VoxelPageBase + VoxelPageOffset;
+
+					#if VOXEL_TRAVERSAL_DEBUG
+					{
+						const float VoxelExtent = InCommonDesc.VoxelWorldSize * (1 << MipLevel) * 0.5f;
+						const float3 P = (PageIndexCoord*InCommonDesc.PageResolution + VoxelPageOffset) * InCommonDesc.VoxelWorldSize + InNodeDesc.MinAABB;
+						AddStepAABB(InSettings.bDebugEnabled, P - VoxelExtent, P + VoxelExtent, float3(0, 0.5f, 1));
+					}
+					#endif
+
+					const FHairTraversalResult StepResult = GetHairVirtualVoxelDensity(VoxelPageCoord, InPageTexture, MipLevel, InSettings.DensityScale);
+					Acc(Out, StepResult);
+
+					if (ShouldStopTraversal(Out, InSettings.CountThreshold))
+					{
+						#if VOXEL_TRAVERSAL_DEBUG
+						AddStepAABB(InSettings.bDebugEnabled, InnerHitP - D * 2, InnerHitP + D * 2, float3(1, 0.25f, 0));
+						#endif
+						return Out;
+					}
+				}
+			}
+
+			#if VOXEL_TRAVERSAL_TYPE == VOXEL_TRAVERSAL_SPARSE_MIPMAP
+			// Increase the stepping size as we walk away from the start point. 
+			// Hypothesis: 
+			// * don't increase the stepping too soon as we might miss important closeby details
+			// * don't increase too much the stepping other wise we will miss important thing or look too coarse
+			PageStepping = clamp(PageStepping += 0.5f, 1.f, 8.f);
+			#endif		
+
+			// t is used for defining the intersection point at the entry of a valid page
+			t = min(tMax.x, min(tMax.y, tMax.z));
+
+			// Find the next page indx to visit and update the tmax, accordingly
+			const float3 Mask = tMax.x < tMax.y ?
+				(tMax.x < tMax.z ? float3(1, 0, 0) : float3(0, 0, 1)) :
+				(tMax.y < tMax.z ? float3(0, 1, 0) : float3(0, 0, 1));
+			PageIndexCoord += Step * Mask;
+			tMax += tDelta * Mask;
+			
+			const float3 NewDelta = tDelta * Mask;
+		}
+	}
+
+
+	return Out;
+}
+#endif
\ No newline at end of file
diff --git a/Engine/Shaders/Private/HairStrands/HairStrandsVoxelRayMarching.usf b/Engine/Shaders/Private/HairStrands/HairStrandsVoxelRayMarching.usf
--- a/Engine/Shaders/Private/HairStrands/HairStrandsVoxelRayMarching.usf
+++ b/Engine/Shaders/Private/HairStrands/HairStrandsVoxelRayMarching.usf
@@ -5,7 +5,8 @@
 #include "../SceneTextureParameters.ush"
 #include "../DeferredShadingCommon.ush"
 #include "HairStrandsVisibilityCommon.ush"
-#include "HairStrandsDeepVoxelCommon.ush"
+#include "HairStrandsVoxelPageCommon.ush"
+#include "HairStrandsDebugCommon.ush"
 
 float3 VoxelMinAABB;
 uint   VoxelResolution;
@@ -22,28 +23,6 @@ Texture3D<uint> MaterialTexture;
 
 SamplerState LinearSampler;
 
-
-float3 HUEtoRGB(in float H)
-{
-	float R = abs(H * 6 - 3) - 1;
-	float G = 2 - abs(H * 6 - 2);
-	float B = 2 - abs(H * 6 - 4);
-	return saturate(float3(R, G, B));
-}
-
-float3 HSVtoRGB(in float3 HSV)
-{
-	float3 RGB = HUEtoRGB(HSV.x);
-	return ((RGB - 1) * HSV.y + 1) * HSV.z;
-}
-
-float3 GetColorCode(float x)
-{
-	float c = (1 - saturate(x)) * 0.6; // Remap [0,1] to Blue-Red
-	return x > 0 ? HSVtoRGB(float3(c, 1, 1)) : float3(0, 0, 0);
-}
-
-
 float3 ReconstructWorldPositionFromDepth(float2 UV, float Depth)
 {
 	float2 ScreenPosition = (UV - View.ScreenPositionScaleBias.wz) / View.ScreenPositionScaleBias.xy;
@@ -53,6 +32,21 @@ float3 ReconstructWorldPositionFromDepth(float2 UV, float Depth)
 	return WorldPosition;
 }
 
+float3 GetHairVoxelTangent(
+	uint3 InVolumeCoord,
+	float InHairDensity,
+	Texture3D<uint> InVoxelTangentXTexture,
+	Texture3D<uint> InVoxelTangentYTexture,
+	Texture3D<uint> InVoxelTangentZTexture)
+{
+	float3 T8Bits;
+	T8Bits.x = InVoxelTangentXTexture.Load(int4(InVolumeCoord, 0));
+	T8Bits.y = InVoxelTangentYTexture.Load(int4(InVolumeCoord, 0));
+	T8Bits.z = InVoxelTangentZTexture.Load(int4(InVolumeCoord, 0));
+	T8Bits /= InHairDensity;
+	return From8bits(T8Bits);
+}
+
 // Config
 #define DEBUG_MODE_DENSITY 0
 #define DEBUG_MODE_TANGENT 1
@@ -105,7 +99,8 @@ void MainPS(
 			const float3 HitP = Offset + O + ((Dither+StepIt)*Delta + HitT.x) * D;
 			const float3 UVW  = (HitP - VoxelMinAABB) / (VoxelMaxAABB - VoxelMinAABB);
 			const int3 VolumeCoord = clamp(int3(UVW * MipResolution), int3(0,0,0), int3(MipResolution.xxx-1));
-			const float Density = GetHairVoxelDensity(VolumeCoord, DensityTexture, MipLevel) * VoxelDensityScale;
+			const FHairTraversalResult Result = GetHairVoxelDensity(VolumeCoord, DensityTexture, MipLevel, VoxelDensityScale);
+			const float Density = Result.HairCount + (Result.Visibility == 0 ? 1000000 : 0);
 		#if PERMUTATION_DEBUG_MODE == DEBUG_MODE_DENSITY
 			if (Density > 0)
 			{
diff --git a/Engine/Shaders/Private/ShaderDrawDebug.usf b/Engine/Shaders/Private/ShaderDrawDebug.usf
new file mode 100644
--- /dev/null
+++ b/Engine/Shaders/Private/ShaderDrawDebug.usf
@@ -0,0 +1,98 @@
+// Copyright Epic Games, Inc. All Rights Reserved.
+
+
+#include "Common.ush"
+
+#include "ShaderDrawDebug.ush"
+
+// The shader used to render the debug element on view
+
+#if PERMUTATION_INPUT_TYPE == 0 
+StructuredBuffer<PackedShaderDrawElement>	ShaderDrawDebugPrimitive;
+#else
+StructuredBuffer<PackedShaderDrawElement>	LockedShaderDrawDebugPrimitive;
+#endif
+
+float ColorScale;
+
+#if GPU_DEBUG_RENDERING_VS 
+void ShaderDrawDebugVS(
+	in uint InstanceId : SV_InstanceID,
+	in uint VertexId : SV_VertexID,
+	out float4 Position : SV_POSITION,
+	out float4 ColorAlpha : TEXCOORD0)
+{
+#if PERMUTATION_INPUT_TYPE == 0
+	PackedShaderDrawElement PackedLine = ShaderDrawDebugPrimitive[InstanceId];
+#else
+	PackedShaderDrawElement PackedLine = LockedShaderDrawDebugPrimitive[InstanceId];
+#endif
+	ShaderDrawElement Line = UnpackShaderElement(PackedLine);
+
+	float4 VertexPosition = float4(VertexId == 0 ? Line.Pos0 : Line.Pos1, 1.0f);
+	float4 VertexColor = VertexId == 0 ? Line.Color0 : Line.Color1;
+
+	Position = mul(VertexPosition, View.WorldToClip);
+	ColorAlpha = VertexColor;
+}
+#endif
+
+#if GPU_DEBUG_RENDERING_PS
+void ShaderDrawDebugPS(
+	in float4 Position : SV_POSITION,
+	in float4 ColorAlpha : TEXCOORD0,
+	out float4 OutSceneColor : SV_Target0)
+{
+	uint2 PixelCoord = Position.xy;
+	const float CheckerboardScale = ((PixelCoord.x%2) ^ (PixelCoord.y%2)) > 0 ? 1 : 0.0;
+	const float OcclusionScale = ColorScale < 0.5 ? CheckerboardScale : 1;
+	OutSceneColor = OcclusionScale * ColorAlpha * float4(ColorAlpha.aaa * ColorScale, 1.0); // Pre multiplied alpha with color scale
+}
+#endif
+
+#if GPU_DEBUG_RENDERING_CLEAR_CS
+
+RWStructuredBuffer<PackedShaderDrawElement>	DataBuffer;
+RWBuffer<uint> IndirectBuffer;
+
+[numthreads(1, 1, 1)]
+void ShaderDrawDebugClearCS(uint3 DispatchThreadId : SV_DispatchThreadID)
+{
+	IndirectBuffer[0] = 2;
+	IndirectBuffer[1] = 0;
+	IndirectBuffer[2] = 0;
+	IndirectBuffer[3] = 0;
+}
+#endif
+
+#if GPU_DEBUG_RENDERING_COPY_CS
+
+	#if PERMUTATION_BUFFER_TYPE == 0
+	uint NumElements;
+	StructuredBuffer<PackedShaderDrawElement>	InStructuredBuffer;
+	RWStructuredBuffer<PackedShaderDrawElement>	OutStructuredBuffer;
+	[numthreads(1024, 1, 1)]
+	void ShaderDrawDebugCopyCS(uint3 DispatchThreadId : SV_DispatchThreadID)
+	{
+		const uint Index = DispatchThreadId.x;
+		if (Index < NumElements)
+		{
+			OutStructuredBuffer[Index] = InStructuredBuffer[Index];
+		}
+	}
+	#endif
+
+	#if PERMUTATION_BUFFER_TYPE == 1
+	Buffer<uint> InBuffer;
+	RWBuffer<uint> OutBuffer;
+	[numthreads(1, 1, 1)]
+	void ShaderDrawDebugCopyCS(uint3 DispatchThreadId : SV_DispatchThreadID)
+	{
+		OutBuffer[0] = InBuffer[0];
+		OutBuffer[1] = InBuffer[1];
+		OutBuffer[2] = InBuffer[2];
+		OutBuffer[3] = InBuffer[3];
+	}
+	#endif
+
+#endif
diff --git a/Engine/Shaders/Private/ShaderDrawDebug.ush b/Engine/Shaders/Private/ShaderDrawDebug.ush
new file mode 100644
--- /dev/null
+++ b/Engine/Shaders/Private/ShaderDrawDebug.ush
@@ -0,0 +1,181 @@
+// Copyright Epic Games, Inc. All Rights Reserved.
+
+#pragma once
+
+#ifndef SHADER_DRAW_DEBUG
+#define SHADER_DRAW_DEBUG 1
+#endif
+
+struct ShaderDrawElement 
+{
+	float3 Pos0;
+	float3 Pos1;
+	float4 Color0;
+	float4 Color1;
+};
+
+struct PackedShaderDrawElement 
+{
+	uint3 Pos01;
+	uint2 Color01;
+};
+
+PackedShaderDrawElement PackShaderElement(ShaderDrawElement E)
+{
+	PackedShaderDrawElement PackedE;
+	PackedE.Pos01.x = (f32tof16(E.Pos0.x) << 16) | f32tof16(E.Pos0.y);
+	PackedE.Pos01.y = (f32tof16(E.Pos0.z) << 16) | f32tof16(E.Pos1.x);
+	PackedE.Pos01.z = (f32tof16(E.Pos1.y) << 16) | f32tof16(E.Pos1.z);
+	uint4 PackedC0 = 255.0f * saturate(E.Color0);
+	uint4 PackedC1 = 255.0f * saturate(E.Color1);
+	PackedE.Color01.x = (PackedC0.x << 24) | (PackedC0.y << 16) | (PackedC0.z << 8) | (PackedC0.w);
+	PackedE.Color01.y = (PackedC1.x << 24) | (PackedC1.y << 16) | (PackedC1.z << 8) | (PackedC1.w);
+	return PackedE;
+}
+
+ShaderDrawElement UnpackShaderElement(PackedShaderDrawElement PackedE)
+{
+	ShaderDrawElement E;
+	E.Pos0 = float3(f16tof32(PackedE.Pos01.x >> 16),	f16tof32(PackedE.Pos01.x),			f16tof32(PackedE.Pos01.y >> 16));
+	E.Pos1 = float3(f16tof32(PackedE.Pos01.y),			f16tof32(PackedE.Pos01.z >> 16),	f16tof32(PackedE.Pos01.z));
+	E.Color0 = float4((PackedE.Color01.x >> 24) & 0xFF, (PackedE.Color01.x >> 16) & 0xFF, (PackedE.Color01.x >> 8) & 0xFF, (PackedE.Color01.x) & 0xFF) / 255.0f;
+	E.Color1 = float4((PackedE.Color01.y >> 24) & 0xFF, (PackedE.Color01.y >> 16) & 0xFF, (PackedE.Color01.y >> 8) & 0xFF, (PackedE.Color01.y) & 0xFF) / 255.0f;
+	return E;
+}
+
+#if SHADER_DRAW_DEBUG == 1
+
+RWStructuredBuffer<PackedShaderDrawElement>	OutShaderDrawPrimitive;
+RWBuffer<uint>								OutputShaderDrawIndirect;
+
+int2 ShaderDrawCursorPos;
+uint ShaderDrawMaxElementCount; 
+
+int2 GetCursorPos()
+{
+	return ShaderDrawCursorPos;
+}
+
+void AddLine(float3 Pos0, float3 Pos1, float4 Color0, float4 Color1)
+{
+	uint PreviousPrimitiveCount = 0;
+	InterlockedAdd(OutputShaderDrawIndirect[1], 1, PreviousPrimitiveCount);
+	if ((PreviousPrimitiveCount+1) <= ShaderDrawMaxElementCount)
+    {
+		ShaderDrawElement Element;
+		Element.Pos0 = Pos0;
+		Element.Pos1 = Pos1;
+		Element.Color0 = Color0;
+		Element.Color1 = Color1;
+
+		OutShaderDrawPrimitive[PreviousPrimitiveCount] = PackShaderElement(Element);
+    }
+	else
+	{
+		InterlockedAdd(OutputShaderDrawIndirect[1], -1, PreviousPrimitiveCount);// Make sure we stay in the safe zone for indirect draw
+	}
+}
+
+void AddQuad(float3 Pos0, float3 Pos1, float3 Pos2, float3 Pos3, float4 Color)
+{
+	uint PreviousPrimitiveCount = 0;
+	InterlockedAdd(OutputShaderDrawIndirect[1], 4, PreviousPrimitiveCount);
+	if ((PreviousPrimitiveCount+4) <= ShaderDrawMaxElementCount)
+	{
+		ShaderDrawElement Element;
+		Element.Color0 = Element.Color1 = Color;
+
+		Element.Pos0 = Pos0; Element.Pos1 = Pos1;
+		OutShaderDrawPrimitive[PreviousPrimitiveCount + 0] = PackShaderElement(Element);
+
+		Element.Pos0 = Pos1; Element.Pos1 = Pos2;
+		OutShaderDrawPrimitive[PreviousPrimitiveCount + 1] = PackShaderElement(Element);
+
+		Element.Pos0 = Pos2; Element.Pos1 = Pos3;
+		OutShaderDrawPrimitive[PreviousPrimitiveCount + 2] = PackShaderElement(Element);
+
+		Element.Pos0 = Pos3; Element.Pos1 = Pos0;
+		OutShaderDrawPrimitive[PreviousPrimitiveCount + 3] = PackShaderElement(Element);
+	}
+	else
+	{
+		InterlockedAdd(OutputShaderDrawIndirect[1], -4, PreviousPrimitiveCount);// Make sure we stay in the safe zone for indirect draw
+	}
+}
+
+void AddAABB(float3 Min, float3 Max, float4 Color)
+{
+	uint PreviousPrimitiveCount = 0;
+	InterlockedAdd(OutputShaderDrawIndirect[1], 12, PreviousPrimitiveCount);
+	if ((PreviousPrimitiveCount + 12) <= ShaderDrawMaxElementCount)
+	{
+		float3 P0 = float3(Min.x, Min.y, Min.z);
+		float3 P1 = float3(Max.x, Min.y, Min.z);
+		float3 P2 = float3(Max.x, Max.y, Min.z);
+		float3 P3 = float3(Min.x, Max.y, Min.z);
+		float3 P4 = float3(Min.x, Min.y, Max.z);
+		float3 P5 = float3(Max.x, Min.y, Max.z);
+		float3 P6 = float3(Max.x, Max.y, Max.z);
+		float3 P7 = float3(Min.x, Max.y, Max.z);
+
+		ShaderDrawElement Element;
+		Element.Color0 = Element.Color1 = Color;
+
+		Element.Pos0 = P0; Element.Pos1 = P1; OutShaderDrawPrimitive[PreviousPrimitiveCount + 0] = PackShaderElement(Element);
+		Element.Pos0 = P1; Element.Pos1 = P2; OutShaderDrawPrimitive[PreviousPrimitiveCount + 1] = PackShaderElement(Element);
+		Element.Pos0 = P2; Element.Pos1 = P3; OutShaderDrawPrimitive[PreviousPrimitiveCount + 2] = PackShaderElement(Element);
+		Element.Pos0 = P3; Element.Pos1 = P0; OutShaderDrawPrimitive[PreviousPrimitiveCount + 3] = PackShaderElement(Element);
+
+		Element.Pos0 = P4; Element.Pos1 = P5; OutShaderDrawPrimitive[PreviousPrimitiveCount + 4] = PackShaderElement(Element);
+		Element.Pos0 = P5; Element.Pos1 = P6; OutShaderDrawPrimitive[PreviousPrimitiveCount + 5] = PackShaderElement(Element);
+		Element.Pos0 = P6; Element.Pos1 = P7; OutShaderDrawPrimitive[PreviousPrimitiveCount + 6] = PackShaderElement(Element);
+		Element.Pos0 = P7; Element.Pos1 = P4; OutShaderDrawPrimitive[PreviousPrimitiveCount + 7] = PackShaderElement(Element);
+
+		Element.Pos0 = P0; Element.Pos1 = P4; OutShaderDrawPrimitive[PreviousPrimitiveCount + 8] = PackShaderElement(Element);
+		Element.Pos0 = P1; Element.Pos1 = P5; OutShaderDrawPrimitive[PreviousPrimitiveCount + 9] = PackShaderElement(Element);
+		Element.Pos0 = P2; Element.Pos1 = P6; OutShaderDrawPrimitive[PreviousPrimitiveCount +10] = PackShaderElement(Element);
+		Element.Pos0 = P3; Element.Pos1 = P7; OutShaderDrawPrimitive[PreviousPrimitiveCount +11] = PackShaderElement(Element);
+	}
+	else
+	{
+		InterlockedAdd(OutputShaderDrawIndirect[1], -12, PreviousPrimitiveCount);// Make sure we stay in the safe zone for indirect draw
+	}
+}
+
+void AddCross(float3 Pos, float Size, float4 Color)
+{
+	AddLine(Pos - float3(Size,0,0), Pos + float3(Size,0,0), Color, Color);
+	AddLine(Pos - float3(0,Size,0), Pos + float3(0,Size,0), Color, Color);
+	AddLine(Pos - float3(0,0,Size), Pos + float3(0,0,Size), Color, Color);
+}
+
+void AddReferential(float3 Pos, float3 T, float3 B, float3 N)
+{
+	const float4 Red = float4(1, 0, 0, 1);
+	const float4 Green = float4(0, 1, 0, 1);
+	const float4 Blue = float4(0, 0, 1, 1);
+
+	AddLine(Pos, Pos + T, Red, Red);
+	AddLine(Pos, Pos + B, Green, Green);
+	AddLine(Pos, Pos + N, Blue, Blue);
+}
+
+void AddReferential(float3 Pos, float3 TangentZ)
+{
+	const float Sign = TangentZ.z >= 0 ? 1 : -1;
+	const float a = -rcp(Sign + TangentZ.z);
+	const float b = TangentZ.x * TangentZ.y * a;
+
+	const float3 TangentX = { 1 + Sign * a * Pow2(TangentZ.x), Sign * b, -Sign * TangentZ.x };
+	const float3 TangentY = { b,  Sign + a * Pow2(TangentZ.y), -TangentZ.y };
+
+	AddReferential(Pos, TangentX, TangentY, TangentZ);
+}
+
+void AddLineTriangle(float3 P0, float3 P1, float3 P2, float4 Color)
+{
+	AddLine(P0, P1, Color, Color);
+	AddLine(P1, P2, Color, Color);
+	AddLine(P2, P0, Color, Color);
+}
+#endif
\ No newline at end of file
diff --git a/Engine/Shaders/Private/ShadowProjectionPixelShader.usf b/Engine/Shaders/Private/ShadowProjectionPixelShader.usf
--- a/Engine/Shaders/Private/ShadowProjectionPixelShader.usf
+++ b/Engine/Shaders/Private/ShadowProjectionPixelShader.usf
@@ -362,6 +362,20 @@ void MainOnePassPointLightPS(
 {
 	float2 ScreenUV = float2( SVPos.xy * View.BufferSizeAndInvSize.zw );
 	float SceneW = CalcSceneDepth( ScreenUV );
+
+#if SUBPIXEL_SHADOW
+	const uint3 PixelCoord = uint3(floor(SVPos.xy), 0);
+	const FCategorizationData CategorizationData = DecodeCategorizationData(HairCategorizationTexture.Load(PixelCoord));
+	if (CategorizationData.PixelCoverage > 0)
+	{
+		SceneW = ConvertFromDeviceZ(CategorizationData.ClosestDepth);
+	}
+	else
+	{
+		discard;
+	}
+#endif
+
 	float2 ScreenPosition = ( ScreenUV.xy - View.ScreenPositionScaleBias.wz ) / View.ScreenPositionScaleBias.xy;
 	float3 WorldPosition = mul(float4(ScreenPosition.xy * SceneW, SceneW, 1), View.ScreenToWorld).xyz;
 
@@ -380,6 +394,7 @@ void MainOnePassPointLightPS(
 	// Unlike spot/rect/direction light, point lights performs depth bias only at projection time. 
 	// This is why the slope bias drives the normal receiver bias here.
 	float SlopeBias = 0;
+#if !SUBPIXEL_SHADOW
 	{
 		FGBufferData GBufferData = GetGBufferData(ScreenUV);
 		const bool bIsDirectional = LightPositionOrDirection.w == 0;
@@ -387,6 +402,7 @@ void MainOnePassPointLightPS(
 		const float NoL = saturate(dot(GBufferData.WorldNormal, LightDirection));
 		SlopeBias = (1-NoL) * PointLightDepthBias.y;
 	}
+#endif
 
 	float Shadow = CubemapHardwarePCF(WorldPosition, LightPositionAndInvRadius.xyz, LightPositionAndInvRadius.w, PointLightDepthBias.x, SlopeBias, PointLightDepthBias.z);
 
diff --git a/Engine/Source/Editor/DetailCustomizations/Private/CollisionProfileDetails.cpp b/Engine/Source/Editor/DetailCustomizations/Private/CollisionProfileDetails.cpp
--- a/Engine/Source/Editor/DetailCustomizations/Private/CollisionProfileDetails.cpp
+++ b/Engine/Source/Editor/DetailCustomizations/Private/CollisionProfileDetails.cpp
@@ -1258,7 +1258,7 @@ TSharedRef<SWidget> SProfileListItem::GenerateWidgetForColumn(const FName& Colum
 				[
 					SNew(SImage)
 					.Image(FEditorStyle::GetBrush("SettingsEditor.Collision_Game"))
-					.ToolTipText(LOCTEXT("CanModify_Tooltip", "This is your custom project profie"))
+					.ToolTipText(LOCTEXT("CanModify_Tooltip", "This is your custom project profile"))
 				];
 		}
 	}
diff --git a/Engine/Source/Editor/DetailCustomizations/Private/PhysicsSettingsDetails.cpp b/Engine/Source/Editor/DetailCustomizations/Private/PhysicsSettingsDetails.cpp
--- a/Engine/Source/Editor/DetailCustomizations/Private/PhysicsSettingsDetails.cpp
+++ b/Engine/Source/Editor/DetailCustomizations/Private/PhysicsSettingsDetails.cpp
@@ -254,6 +254,7 @@ private:
 			}
 
 			PhysicsSettings->UpdateDefaultConfigFile();
+			PhysicsSettings->LoadSurfaceType();
 
 			PhysicalSurfacesProperty->NotifyPostChange();
 		}
diff --git a/Engine/Source/Runtime/Core/Public/UObject/PhysicsObjectVersion.h b/Engine/Source/Runtime/Core/Public/UObject/PhysicsObjectVersion.h
--- a/Engine/Source/Runtime/Core/Public/UObject/PhysicsObjectVersion.h
+++ b/Engine/Source/Runtime/Core/Public/UObject/PhysicsObjectVersion.h
@@ -15,6 +15,11 @@ struct CORE_API FPhysicsObjectVersion
 		PerShapeData,
 		// Add serialization from handle back to particle
 		SerializeGTGeometryParticles,
+		// Groom serialization with hair description as bulk data
+		GroomWithDescription,
+		// Groom serialization with import option
+		GroomWithImportSettings,
+
 		// -----<new versions can be added above this line>-------------------------------------------------
 		VersionPlusOne,
 		LatestVersion = VersionPlusOne - 1
diff --git a/Engine/Source/Runtime/Engine/Classes/Components/SkinnedMeshComponent.h b/Engine/Source/Runtime/Engine/Classes/Components/SkinnedMeshComponent.h
--- a/Engine/Source/Runtime/Engine/Classes/Components/SkinnedMeshComponent.h
+++ b/Engine/Source/Runtime/Engine/Classes/Components/SkinnedMeshComponent.h
@@ -659,6 +659,7 @@ public:
 	/** Object responsible for sending bone transforms, morph target state etc. to render thread. */
 	class FSkeletalMeshObject*	MeshObject;
 	FSkeletalMeshObjectCallbackData MeshObjectCallbackData;
+	void SetMeshObjectCallbackData(FSkeletalMeshObjectCallbackData& MeshObjectCallbackData);
 
 	/** Gets the skeletal mesh resource used for rendering the component. */
 	FSkeletalMeshRenderData* GetSkeletalMeshRenderData() const;
diff --git a/Engine/Source/Runtime/Engine/Private/Components/SkinnedMeshComponent.cpp b/Engine/Source/Runtime/Engine/Private/Components/SkinnedMeshComponent.cpp
--- a/Engine/Source/Runtime/Engine/Private/Components/SkinnedMeshComponent.cpp
+++ b/Engine/Source/Runtime/Engine/Private/Components/SkinnedMeshComponent.cpp
@@ -778,6 +778,15 @@ void USkinnedMeshComponent::TickUpdateRate(float DeltaTime, bool bNeedsValidRoot
 	}
 }
 
+void USkinnedMeshComponent::SetMeshObjectCallbackData(FSkeletalMeshObjectCallbackData& InData)
+{
+	MeshObjectCallbackData = InData;
+	if (MeshObject)
+	{
+		MeshObject->SetCallbackData(InData);
+	}
+}
+
 void USkinnedMeshComponent::TickComponent(float DeltaTime, enum ELevelTick TickType, FActorComponentTickFunction *ThisTickFunction)
 {
 	SCOPED_NAMED_EVENT(USkinnedMeshComponent_TickComponent, FColor::Yellow);
diff --git a/Engine/Source/Runtime/Engine/Private/SkeletalRenderGPUSkin.cpp b/Engine/Source/Runtime/Engine/Private/SkeletalRenderGPUSkin.cpp
--- a/Engine/Source/Runtime/Engine/Private/SkeletalRenderGPUSkin.cpp
+++ b/Engine/Source/Runtime/Engine/Private/SkeletalRenderGPUSkin.cpp
@@ -127,6 +127,10 @@ FSkeletalMeshObjectGPUSkin
 void FSkeletalMeshObjectGPUSkin::SetCallbackData(FSkeletalMeshObjectCallbackData& InMeshObjectCallbackData)
 {
 	CallbackData = InMeshObjectCallbackData;
+	if (CallbackData.Run)
+	{
+		CallbackData.Run(FSkeletalMeshObjectCallbackData::EEventType::Register, this, CallbackData.UserData);
+	}
 }
 
 FSkeletalMeshObjectGPUSkin::FSkeletalMeshObjectGPUSkin(USkinnedMeshComponent* InMeshComponent, FSkeletalMeshRenderData* InSkelMeshRenderData, ERHIFeatureLevel::Type InFeatureLevel)
@@ -145,21 +149,12 @@ FSkeletalMeshObjectGPUSkin::FSkeletalMeshObjectGPUSkin(USkinnedMeshComponent* In
 	}
 
 	InitResources(InMeshComponent);
-	CallbackData = InMeshComponent->MeshObjectCallbackData;
-	if (CallbackData.Run)
-	{
-		CallbackData.Run(FSkeletalMeshObjectCallbackData::EEventType::Register, this, CallbackData.UserData);
-	}
+	SetCallbackData(InMeshComponent->MeshObjectCallbackData);
 }
 
 
 FSkeletalMeshObjectGPUSkin::~FSkeletalMeshObjectGPUSkin()
 {
-	if (CallbackData.Run)
-	{
-		CallbackData.Run(FSkeletalMeshObjectCallbackData::EEventType::Unregister, this, CallbackData.UserData);
-	}
-
 	check(!RHIThreadFenceForDynamicData.GetReference());
 	if (DynamicData)
 	{
@@ -197,6 +192,11 @@ void FSkeletalMeshObjectGPUSkin::InitResources(USkinnedMeshComponent* InMeshComp
 
 void FSkeletalMeshObjectGPUSkin::ReleaseResources()
 {
+	if (CallbackData.Run)
+	{
+		CallbackData.Run(FSkeletalMeshObjectCallbackData::EEventType::Unregister, this, CallbackData.UserData);
+	}
+
 	for( int32 LODIndex=0;LODIndex < LODs.Num();LODIndex++ )
 	{
 		FSkeletalMeshObjectLOD& SkelLOD = LODs[LODIndex];
diff --git a/Engine/Source/Runtime/RenderCore/Private/RenderGraphValidation.cpp b/Engine/Source/Runtime/RenderCore/Private/RenderGraphValidation.cpp
--- a/Engine/Source/Runtime/RenderCore/Private/RenderGraphValidation.cpp
+++ b/Engine/Source/Runtime/RenderCore/Private/RenderGraphValidation.cpp
@@ -63,9 +63,9 @@ public:
 	{
 		check(Buffer);
 
-		ensureMsgf(!bIsGenerateMips,
-			TEXT("Attempting to read from Buffer %s on Pass %s which is marked as 'GenerateMips'. Only textures are supported on this type of pass."),
-			Buffer->Name, Pass.GetName());
+//		ensureMsgf(!bIsGenerateMips,
+//			TEXT("Attempting to read from Buffer %s on Pass %s which is marked as 'GenerateMips'. Only textures are supported on this type of pass."),
+//			Buffer->Name, Pass.GetName());
 
 		FBufferAccess& BufferAccess = BufferAccessMap.FindOrAdd(Buffer);
 		BufferAccess.Bits.Read = 1;
diff --git a/Engine/Source/Runtime/RenderCore/Public/VertexFactory.h b/Engine/Source/Runtime/RenderCore/Public/VertexFactory.h
--- a/Engine/Source/Runtime/RenderCore/Public/VertexFactory.h
+++ b/Engine/Source/Runtime/RenderCore/Public/VertexFactory.h
@@ -17,6 +17,7 @@
 
 class FMaterial;
 class FMeshDrawSingleShaderBindings;
+class FPrimitiveSceneProxy;
 
 #if PLATFORM_SUPPORTS_PRAGMA_PACK
 	#pragma pack (push,4)
diff --git a/Engine/Source/Runtime/Renderer/Private/DeferredShadingRenderer.cpp b/Engine/Source/Runtime/Renderer/Private/DeferredShadingRenderer.cpp
--- a/Engine/Source/Runtime/Renderer/Private/DeferredShadingRenderer.cpp
+++ b/Engine/Source/Runtime/Renderer/Private/DeferredShadingRenderer.cpp
@@ -39,6 +39,7 @@
 #include "RayTracingDefinitions.h"
 #include "RayTracingInstance.h"
 #include "ShaderPrint.h"
+#include "GpuDebugRendering.h"
 #include "HairStrands/HairStrandsRendering.h"
 
 static TAutoConsoleVariable<int32> CVarStencilForLODDither(
@@ -454,6 +455,16 @@ static FORCEINLINE bool NeedsPrePass(const FDeferredShadingSceneRenderer* Render
 		(Renderer->EarlyZPassMode != DDM_None || Renderer->bEarlyZPassMovable != 0);
 }
 
+static bool DoesHairStrandsRequestHzb(EShaderPlatform Platform)
+{
+	auto HzbRequested = [&]()
+	{
+		IConsoleVariable* CVar = IConsoleManager::Get().FindConsoleVariable(TEXT("HairStrands.Cluster.CullingUsesHzb"));
+		return CVar && CVar->GetInt() > 0;
+	};
+	return IsHairStrandsEnable(Platform) && HzbRequested();
+}
+
 bool FDeferredShadingSceneRenderer::RenderHzb(FRHICommandListImmediate& RHICmdList)
 {
 	FSceneRenderTargets& SceneContext = FSceneRenderTargets::Get(RHICmdList);
@@ -475,8 +486,9 @@ bool FDeferredShadingSceneRenderer::RenderHzb(FRHICommandListImmediate& RHICmdLi
 		const bool bSSR  = ShouldRenderScreenSpaceReflections(View);
 		const bool bSSAO = ShouldRenderScreenSpaceAmbientOcclusion(View);
 		const bool bSSGI = ShouldRenderScreenSpaceDiffuseIndirect(View);
+		const bool bHair = DoesHairStrandsRequestHzb(Scene->GetShaderPlatform());
 
-		if (bSSAO || bHZBOcclusion || bSSR || bSSGI)
+		if (bSSAO || bHZBOcclusion || bSSR || bSSGI || bHair)
 		{
 			FRDGBuilder GraphBuilder(RHICmdList);
 
@@ -1159,6 +1171,7 @@ void FDeferredShadingSceneRenderer::Render(FRHICommandListImmediate& RHICmdList)
 			FViewInfo& View = Views[ViewIndex];
 			SCOPED_GPU_MASK(RHICmdList, View.GPUMask);
 			ShaderPrint::BeginView(RHICmdList, View);
+			ShaderDrawDebug::BeginView(RHICmdList, View);
 		}
 
 		for (int32 ViewIndex = 0; ViewIndex < Views.Num(); ViewIndex++)
@@ -1393,11 +1406,13 @@ void FDeferredShadingSceneRenderer::Render(FRHICommandListImmediate& RHICmdList)
 
 	// Interpolation needs to happen after the skin cache run as there is a dependency 
 	// on the skin cache output.
-	if (IsHairStrandsEnable(Scene->GetShaderPlatform()) && Views.Num() > 0)
+	const bool bRunHairStrands = IsHairStrandsEnable(Scene->GetShaderPlatform()) && Views.Num() > 0;
+	FHairStrandClusterData HairClusterData;
+	if (bRunHairStrands)
 	{
 		const EWorldType::Type WorldType = Views[0].Family->Scene->GetWorld()->WorldType;
 		auto ShaderMap = GetGlobalShaderMap(FeatureLevel);
-		RunHairStrandsInterpolation(RHICmdList, WorldType, ShaderMap, EHairStrandsInterpolationType::RenderStrands);
+		RunHairStrandsInterpolation(RHICmdList, WorldType, &Views[0].ShaderDrawData, ShaderMap, EHairStrandsInterpolationType::RenderStrands, &HairClusterData); // Send data to full up with culling
 	}
 
 	// Before starting the render, all async task for the Custom data must be completed
@@ -1937,10 +1952,30 @@ void FDeferredShadingSceneRenderer::Render(FRHICommandListImmediate& RHICmdList)
 	if (IsHairStrandsEnable(Scene->GetShaderPlatform()) && bIsViewCompatible)
 	{
 		SCOPED_GPU_STAT(RHICmdList, HairRendering);
-		HairDatasStorage.HairClusterPerViews = CreateHairStrandsClusters(RHICmdList, Scene, Views);
-		VoxelizeHairStrands(RHICmdList, Scene, Views, HairDatasStorage.HairClusterPerViews);
-		HairDatasStorage.DeepShadowViews = RenderHairStrandsDeepShadows(RHICmdList, Scene, Views, HairDatasStorage.HairClusterPerViews);
-		HairDatasStorage.HairVisibilityViews = RenderHairStrandsVisibilityBuffer(RHICmdList, Scene, Views, SceneContext.GBufferB, SceneContext.GetSceneColor(), SceneContext.SceneDepthZ, SceneContext.SceneVelocity, HairDatasStorage.HairClusterPerViews);
+		HairDatasStorage.MacroGroupsPerViews = CreateHairStrandsMacroGroups(RHICmdList, Scene, Views);
+
+		// Culling/LOD pass for DOM and Voxelisation altogether
+		FHairCullingParams CullingParams;
+		CullingParams.bShadowViewMode = true;
+		CullingParams.bCullingProcessSkipped = false;
+		ComputeHairStrandsClustersCulling(RHICmdList, *GetGlobalShaderMap(FeatureLevel), Views, CullingParams, HairClusterData);
+
+		ServiceLocalQueue();
+
+		// Voxelization and Deep Opacity Maps
+		VoxelizeHairStrands(RHICmdList, Scene, Views, HairDatasStorage.MacroGroupsPerViews);
+		HairDatasStorage.DeepShadowViews = RenderHairStrandsDeepShadows(RHICmdList, Scene, Views, HairDatasStorage.MacroGroupsPerViews);
+
+		ServiceLocalQueue();
+
+		// Culling/LOD pass for visibility (must be done after HZB is generated)
+		CullingParams.bShadowViewMode = false;
+		ComputeHairStrandsClustersCulling(RHICmdList, *GetGlobalShaderMap(FeatureLevel), Views, CullingParams, HairClusterData);
+		// Hair visibility pass
+		HairDatasStorage.HairVisibilityViews = RenderHairStrandsVisibilityBuffer(RHICmdList, Scene, Views, SceneContext.GBufferB, SceneContext.GetSceneColor(), SceneContext.SceneDepthZ, SceneContext.SceneVelocity, HairDatasStorage.MacroGroupsPerViews);
+		// Reset indirect draw buffer
+		ResetHairStrandsClusterToLOD0(RHICmdList, *GetGlobalShaderMap(FeatureLevel), HairClusterData);
+
 		ServiceLocalQueue();
 
 		HairDatas = &HairDatasStorage;
@@ -2320,7 +2355,7 @@ void FDeferredShadingSceneRenderer::Render(FRHICommandListImmediate& RHICmdList)
 
 		if (IsHairStrandsEnable(Scene->GetShaderPlatform()))
 		{
-			RenderHairStrandsDebugInfo(RHICmdList, Views, HairDatas);
+			RenderHairStrandsDebugInfo(RHICmdList, Views, HairDatas, HairClusterData);
 		}
 	}
 
@@ -2493,6 +2528,7 @@ void FDeferredShadingSceneRenderer::Render(FRHICommandListImmediate& RHICmdList)
 		for (int32 ViewIndex = 0; ViewIndex < Views.Num(); ViewIndex++)
 		{
 			ShaderPrint::EndView(Views[ViewIndex]);
+			ShaderDrawDebug::EndView(Views[ViewIndex]);
 		}
 
 #if WITH_MGPU
diff --git a/Engine/Source/Runtime/Renderer/Private/GpuDebugRendering.cpp b/Engine/Source/Runtime/Renderer/Private/GpuDebugRendering.cpp
--- a/Engine/Source/Runtime/Renderer/Private/GpuDebugRendering.cpp
+++ b/Engine/Source/Runtime/Renderer/Private/GpuDebugRendering.cpp
@@ -1,145 +1,399 @@
 // Copyright Epic Games, Inc. All Rights Reserved.
 
 #include "GpuDebugRendering.h"
-#include "DynamicPrimitiveDrawing.h"
+#include "SceneRendering.h"
+#include "GlobalShader.h"
+#include "RenderGraphBuilder.h"
 #include "RenderGraphUtils.h"
-#include "RenderTargetPool.h"
+#include "Containers/ResourceArray.h"
+#include "CommonRenderResources.h"
 
-FGpuDebugPrimitiveBuffers AllocateGpuDebugPrimitiveBuffers(FRHICommandListImmediate& RHICmdList)
+namespace ShaderDrawDebug 
 {
-	FGpuDebugPrimitiveBuffers Output;
+	// Console variables
+	static TAutoConsoleVariable<int32> CVarShaderDrawEnable(
+		TEXT("r.ShaderDrawDebug"),
+		1,
+		TEXT("ShaderDrawDebug debugging toggle.\n"),
+		ECVF_Cheat | ECVF_RenderThreadSafe);
+
+	static TAutoConsoleVariable<int32> CVarShaderDrawMaxElementCount(
+		TEXT("r.ShaderDrawDebug.MaxElementCount"),
+		10000,
+		TEXT("ShaderDraw output buffer size in element.\n"),
+		ECVF_Cheat | ECVF_RenderThreadSafe);
+
+	static TAutoConsoleVariable<int32> CVarShaderDrawLock(
+		TEXT("r.ShaderDrawDebug.Lock"),
+		0,
+		TEXT("Lock the shader draw buffer.\n"),
+		ECVF_Cheat | ECVF_RenderThreadSafe);
+
+	bool IsShaderDrawDebugEnabled()
+	{
+#if WITH_EDITOR
+		return CVarShaderDrawEnable.GetValueOnAnyThread() > 0;
+#else
+		return false;
+#endif
+	}
+
+	bool IsShaderDrawLocked()
+	{
+#if WITH_EDITOR
+		return CVarShaderDrawLock.GetValueOnAnyThread() > 0;
+#else
+		return false;
+#endif
+	}
+
+	static bool IsShaderDrawDebugEnabled(const EShaderPlatform Platform)
+	{
+		return IsFeatureLevelSupported(Platform, ERHIFeatureLevel::SM5) && IsPCPlatform(Platform);
+	}
+
+	bool IsShaderDrawDebugEnabled(const FViewInfo& View)
+	{
+		return IsShaderDrawDebugEnabled() && IsShaderDrawDebugEnabled(View.GetShaderPlatform());
+	}
+
+	static uint32 GetMaxShaderDrawElementCount()
 	{
-		FIntPoint Resolution(1, 1);
-		EPixelFormat Format = PF_R32_UINT;
+		return uint32(FMath::Max(1, CVarShaderDrawMaxElementCount.GetValueOnRenderThread()));
+	}
+
+	struct ShaderDrawDebugElement
+	{
+		uint32 Pos[3];
+		uint32 Color[2];
+	};
+
+	// This needs to be allocated per view, or move into a more persistent place
+	struct FLockedData
+	{
+		FRWBufferStructured Buffer;
+		FRWBuffer IndirectBuffer;
+		bool bIsLocked = false;
+	};
+
+	static FLockedData LockedData = {};
+
+	//////////////////////////////////////////////////////////////////////////
+
+	class FShaderDrawDebugClearCS : public FGlobalShader
+	{
+		DECLARE_GLOBAL_SHADER(FShaderDrawDebugClearCS);
+		SHADER_USE_PARAMETER_STRUCT(FShaderDrawDebugClearCS, FGlobalShader);
+
+		BEGIN_SHADER_PARAMETER_STRUCT(FParameters, )
+		SHADER_PARAMETER_RDG_BUFFER_UAV(RWStructuredBuffer, DataBuffer)
+		SHADER_PARAMETER_RDG_BUFFER_UAV(RWBuffer, IndirectBuffer)
+		END_SHADER_PARAMETER_STRUCT()
+
+		static bool ShouldCompilePermutation(const FGlobalShaderPermutationParameters& Parameters)
 		{
-			FPooledRenderTargetDesc Desc(FPooledRenderTargetDesc::Create2DDesc(Resolution, Format, FClearValueBinding::Black, TexCreate_None, TexCreate_UAV, false));
-			GRenderTargetPool.FindFreeElement(RHICmdList, Desc, Output.DebugPrimitiveCountTexture, TEXT("DebugPrimitiveCountTexture"));
+			return IsShaderDrawDebugEnabled(Parameters.Platform);
 		}
+
+		static void ModifyCompilationEnvironment(const FGlobalShaderPermutationParameters& Parameters, FShaderCompilerEnvironment& OutEnvironment)
 		{
-			FPooledRenderTargetDesc Desc(FPooledRenderTargetDesc::Create2DDesc(Resolution, Format, FClearValueBinding::None, TexCreate_CPUReadback, TexCreate_None, false));
-			GRenderTargetPool.FindFreeElement(RHICmdList, Desc, Output.DebugPrimitiveCountStagingTexture, TEXT("DebugPrimitiveCountStagingTexture"));
+			FGlobalShader::ModifyCompilationEnvironment(Parameters, OutEnvironment);
+			OutEnvironment.SetDefine(TEXT("GPU_DEBUG_RENDERING"), 1);
+			OutEnvironment.SetDefine(TEXT("GPU_DEBUG_RENDERING_CLEAR_CS"), 1);
 		}
-	}
+	};
+
+	IMPLEMENT_GLOBAL_SHADER(FShaderDrawDebugClearCS, "/Engine/Private/ShaderDrawDebug.usf", "ShaderDrawDebugClearCS", SF_Compute);
 
+
+	//////////////////////////////////////////////////////////////////////////
+
+	class FShaderDrawDebugCopyCS : public FGlobalShader
 	{
-		FIntPoint Resolution(1024, 1);
-		EPixelFormat Format = PF_A32B32G32R32F;
+		DECLARE_GLOBAL_SHADER(FShaderDrawDebugCopyCS);
+		SHADER_USE_PARAMETER_STRUCT(FShaderDrawDebugCopyCS, FGlobalShader);
+
+		class FBufferType : SHADER_PERMUTATION_INT("PERMUTATION_BUFFER_TYPE", 2);
+		using FPermutationDomain = TShaderPermutationDomain<FBufferType>;
+
+		BEGIN_SHADER_PARAMETER_STRUCT(FParameters, )
+			SHADER_PARAMETER(uint32, NumElements)
+			SHADER_PARAMETER_RDG_BUFFER_SRV(Buffer, InBuffer)
+			SHADER_PARAMETER_UAV(RWBuffer, OutBuffer)
+
+			SHADER_PARAMETER_RDG_BUFFER_SRV(StructuredBuffer, InStructuredBuffer)
+			SHADER_PARAMETER_UAV(RWStructuredBuffer, OutStructuredBuffer)
+		END_SHADER_PARAMETER_STRUCT()
+
+		static bool ShouldCompilePermutation(const FGlobalShaderPermutationParameters& Parameters)
 		{
-			FPooledRenderTargetDesc Desc(FPooledRenderTargetDesc::Create2DDesc(Resolution, Format, FClearValueBinding::Black, TexCreate_None, TexCreate_UAV, false));
-			GRenderTargetPool.FindFreeElement(RHICmdList, Desc, Output.DebugPrimitiveTexture, TEXT("DebugPrimitiveTexture"));
+			return IsShaderDrawDebugEnabled(Parameters.Platform);
 		}
+
+		static void ModifyCompilationEnvironment(const FGlobalShaderPermutationParameters& Parameters, FShaderCompilerEnvironment& OutEnvironment)
 		{
-			FPooledRenderTargetDesc Desc(FPooledRenderTargetDesc::Create2DDesc(Resolution, Format, FClearValueBinding::None, TexCreate_CPUReadback, TexCreate_None, false));
-			GRenderTargetPool.FindFreeElement(RHICmdList, Desc, Output.DebugPrimitiveStagingTexture, TEXT("DebugPrimitiveStagingTexture"));
+			FGlobalShader::ModifyCompilationEnvironment(Parameters, OutEnvironment);
+			OutEnvironment.SetDefine(TEXT("GPU_DEBUG_RENDERING"), 1);
+			OutEnvironment.SetDefine(TEXT("GPU_DEBUG_RENDERING_COPY_CS"), 1);
 		}
-	}
+	};
 
-	uint32 ClearValues[4] = { 0, 0, 0, 0 };
-	RHICmdList.ClearTinyUAV(Output.DebugPrimitiveCountTexture->GetRenderTargetItem().UAV, ClearValues);
-	RHICmdList.ClearTinyUAV(Output.DebugPrimitiveTexture->GetRenderTargetItem().UAV, ClearValues);
+	IMPLEMENT_GLOBAL_SHADER(FShaderDrawDebugCopyCS, "/Engine/Private/ShaderDrawDebug.usf", "ShaderDrawDebugCopyCS", SF_Compute);
 
-	return Output;
-}
+	//////////////////////////////////////////////////////////////////////////
 
-class FGpuDebugLine { public: FVector Start; FVector End; FLinearColor Color; };
-typedef TArray<FGpuDebugLine> FGpuDebugLineArray;
-static FGpuDebugLineArray ReadGPUDebugPrimitives(FRHICommandListImmediate& RHICmdList, FGpuDebugPrimitiveBuffers& DebugPrimitiveBuffer)
-{
-	// Count
-	uint32 PointCount = 0;
+	class FShaderDrawDebugVS : public FGlobalShader
 	{
-		FTextureRHIRef SourceTexture = DebugPrimitiveBuffer.DebugPrimitiveCountTexture->GetRenderTargetItem().TargetableTexture;
-		FTextureRHIRef StagingTexture = DebugPrimitiveBuffer.DebugPrimitiveCountStagingTexture->GetRenderTargetItem().ShaderResourceTexture;
+		DECLARE_GLOBAL_SHADER(FShaderDrawDebugVS);
+		SHADER_USE_PARAMETER_STRUCT(FShaderDrawDebugVS, FGlobalShader);
+
+		class FInputType : SHADER_PERMUTATION_INT("PERMUTATION_INPUT_TYPE", 2);
+		using FPermutationDomain = TShaderPermutationDomain<FInputType>;
 
-		// Transfer memory GPU -> CPU
-		RHICmdList.CopyToResolveTarget(SourceTexture, StagingTexture, FResolveParams());
-		if (StagingTexture.IsValid())
+		BEGIN_SHADER_PARAMETER_STRUCT(FParameters, )
+			SHADER_PARAMETER_STRUCT_REF(FViewUniformShaderParameters, View)
+			SHADER_PARAMETER_SRV(StructuredBuffer, LockedShaderDrawDebugPrimitive)
+			SHADER_PARAMETER_RDG_BUFFER_SRV(StructuredBuffer, ShaderDrawDebugPrimitive)
+			SHADER_PARAMETER_RDG_BUFFER(StructuredBuffer, IndirectBuffer)
+		END_SHADER_PARAMETER_STRUCT()
+
+		static bool ShouldCompilePermutation(const FGlobalShaderPermutationParameters& Parameters)
 		{
-			const uint32* CountData = nullptr;
-			int32 BufferWidth = 0, BufferHeight = 0;
-			RHICmdList.MapStagingSurface(StagingTexture, *(void**)&CountData, BufferWidth, BufferHeight);
+			return IsShaderDrawDebugEnabled(Parameters.Platform);
+		}
 
-			if (CountData)
-			{
-				PointCount = CountData[0];
-				PointCount = FMath::Min(PointCount, 1024u);
-			}
-			RHICmdList.UnmapStagingSurface(StagingTexture);
+		static void ModifyCompilationEnvironment(const FGlobalShaderPermutationParameters& Parameters, FShaderCompilerEnvironment& OutEnvironment)
+		{
+			FGlobalShader::ModifyCompilationEnvironment(Parameters, OutEnvironment);
+			OutEnvironment.SetDefine(TEXT("GPU_DEBUG_RENDERING"), 1);
+			OutEnvironment.SetDefine(TEXT("GPU_DEBUG_RENDERING_VS"), 1);
+			OutEnvironment.SetDefine(TEXT("GPU_DEBUG_RENDERING_PS"), 0);
+		}
+	};
+
+	IMPLEMENT_GLOBAL_SHADER(FShaderDrawDebugVS, "/Engine/Private/ShaderDrawDebug.usf", "ShaderDrawDebugVS", SF_Vertex);
+
+	//////////////////////////////////////////////////////////////////////////
+
+	class FShaderDrawDebugPS : public FGlobalShader
+	{
+		DECLARE_GLOBAL_SHADER(FShaderDrawDebugPS);
+		SHADER_USE_PARAMETER_STRUCT(FShaderDrawDebugPS, FGlobalShader);
+		
+		BEGIN_SHADER_PARAMETER_STRUCT(FParameters, )
+			SHADER_PARAMETER(float, ColorScale)
+			RENDER_TARGET_BINDING_SLOTS()
+		END_SHADER_PARAMETER_STRUCT()
+
+		using FPermutationDomain = TShaderPermutationDomain<>;
+
+		static bool ShouldCompilePermutation(const FGlobalShaderPermutationParameters& Parameters)
+		{
+			return IsShaderDrawDebugEnabled(Parameters.Platform);
 		}
 
+		static void ModifyCompilationEnvironment(const FGlobalShaderPermutationParameters& Parameters, FShaderCompilerEnvironment& OutEnvironment)
+		{
+			FGlobalShader::ModifyCompilationEnvironment(Parameters, OutEnvironment);
+			OutEnvironment.SetDefine(TEXT("GPU_DEBUG_RENDERING"), 1);
+			OutEnvironment.SetDefine(TEXT("GPU_DEBUG_RENDERING_VS"), 0);
+			OutEnvironment.SetDefine(TEXT("GPU_DEBUG_RENDERING_PS"), 1);
+		}
+	};
+
+	IMPLEMENT_GLOBAL_SHADER(FShaderDrawDebugPS, "/Engine/Private/ShaderDrawDebug.usf", "ShaderDrawDebugPS", SF_Pixel);
+
+	BEGIN_SHADER_PARAMETER_STRUCT(ShaderDrawVSPSParameters, )
+		SHADER_PARAMETER_STRUCT_INCLUDE(FShaderDrawDebugVS::FParameters, ShaderDrawVSParameters)
+		SHADER_PARAMETER_STRUCT_INCLUDE(FShaderDrawDebugPS::FParameters, ShaderDrawPSParameters)
+	END_SHADER_PARAMETER_STRUCT()
+
+	//////////////////////////////////////////////////////////////////////////
+
+	void BeginView(FRHICommandListImmediate& InRHICmdList, FViewInfo& View)
+	{
+		if (!IsShaderDrawDebugEnabled(View))
+		{
+			return;
+		}
+
+		FRDGBuilder GraphBuilder(InRHICmdList);
+		FRDGBufferRef DataBuffer = GraphBuilder.CreateBuffer(FRDGBufferDesc::CreateStructuredDesc(sizeof(ShaderDrawDebugElement), GetMaxShaderDrawElementCount()), TEXT("ShaderDrawDataBuffer"));
+		FRDGBufferRef IndirectBuffer = GraphBuilder.CreateBuffer(FRDGBufferDesc::CreateIndirectDesc<FRHIDrawIndirectParameters>(1), TEXT("ShaderDrawDataIndirectBuffer"));
+
+		FShaderDrawDebugClearCS::FParameters* Parameters = GraphBuilder.AllocParameters<FShaderDrawDebugClearCS::FParameters>();
+		Parameters->DataBuffer = GraphBuilder.CreateUAV(DataBuffer);
+		Parameters->IndirectBuffer = GraphBuilder.CreateUAV(IndirectBuffer);
+
+		TShaderMapRef<FShaderDrawDebugClearCS> ComputeShader(View.ShaderMap);
+
+		// Note: we do not call ClearUnusedGraphResources here as we want to for the allocation of DataBuffer
+		GraphBuilder.AddPass(
+			RDG_EVENT_NAME("ShaderDrawClear"),
+			Parameters,
+			ERDGPassFlags::Compute,
+			[Parameters, ComputeShader](FRHICommandList& RHICmdList)
+		{
+			FComputeShaderUtils::Dispatch(RHICmdList, *ComputeShader, *Parameters, FIntVector(1,1,1));
+		});
+
+		GraphBuilder.QueueBufferExtraction(DataBuffer, &View.ShaderDrawData.Buffer, FRDGResourceState::EAccess::Write, FRDGResourceState::EPipeline::Compute);
+		GraphBuilder.QueueBufferExtraction(IndirectBuffer, &View.ShaderDrawData.IndirectBuffer, FRDGResourceState::EAccess::Write, FRDGResourceState::EPipeline::Compute);
+
+		GraphBuilder.Execute();
+
+		if (IsShaderDrawLocked() && !LockedData.bIsLocked)
+		{
+			LockedData.Buffer.Initialize(sizeof(ShaderDrawDebugElement), GetMaxShaderDrawElementCount(), 0U, TEXT("ShaderDrawDataBuffer"));
+			LockedData.IndirectBuffer.Initialize(sizeof(uint32), 4, PF_R32_UINT, BUF_DrawIndirect, TEXT("ShaderDrawDataIndirectBuffer"));
+		}
+
+		View.ShaderDrawData.CursorPosition = View.CursorPos;
 	}
 
-	// Primitives
-	FGpuDebugLineArray Output;
-	if (PointCount > 0)
+	void DrawView(FRDGBuilder& GraphBuilder, const FViewInfo& View, FRDGTextureRef OutputTexture, FRDGTextureRef DepthTexture)
 	{
-		auto ToColor = [](uint32 ColorIndex)
+		if (!IsShaderDrawDebugEnabled(View))
 		{
-			FLinearColor Color;
-			switch (ColorIndex)
+			return;
+		}
+
+		auto RunPass = [&](
+			bool bIsBehindDepth, 
+			bool bUseRdgInput,
+			FRDGBufferRef DataBuffer, 
+			FRDGBufferRef IndirectBuffer, 
+			FShaderResourceViewRHIRef LockedDataBuffer,
+			FRHIVertexBuffer* LockedIndirectBuffer)
+		{
+			FShaderDrawDebugVS::FPermutationDomain PermutationVector;
+			PermutationVector.Set<FShaderDrawDebugVS::FInputType>(bUseRdgInput ? 0 : 1);
+			TShaderMapRef<FShaderDrawDebugVS> VertexShader(View.ShaderMap, PermutationVector);
+			TShaderMapRef<FShaderDrawDebugPS> PixelShader(View.ShaderMap);
+
+			ShaderDrawVSPSParameters* PassParameters = GraphBuilder.AllocParameters<ShaderDrawVSPSParameters>();
+			PassParameters->ShaderDrawPSParameters.RenderTargets[0] = FRenderTargetBinding(OutputTexture, ERenderTargetLoadAction::ELoad);
+			PassParameters->ShaderDrawPSParameters.RenderTargets.DepthStencil = FDepthStencilBinding(DepthTexture, ERenderTargetLoadAction::ELoad, FExclusiveDepthStencil::DepthRead_StencilNop);
+			PassParameters->ShaderDrawPSParameters.ColorScale = bIsBehindDepth ? 0.4f : 1.0f;	// When debug primitive are behind the depth buffer, make them look darker.
+			PassParameters->ShaderDrawVSParameters.View = View.ViewUniformBuffer;
+			if (bUseRdgInput)
 			{
-			case 0:		return FLinearColor::Red;
-			case 1:		return FLinearColor::Green;
-			case 2:		return FLinearColor::Blue;
-			case 3:		return FLinearColor::Yellow;
-			default:	return FLinearColor::White;
+				PassParameters->ShaderDrawVSParameters.ShaderDrawDebugPrimitive = GraphBuilder.CreateSRV(DataBuffer);
+				PassParameters->ShaderDrawVSParameters.IndirectBuffer = IndirectBuffer;
+			}
+			else
+			{
+				PassParameters->ShaderDrawVSParameters.LockedShaderDrawDebugPrimitive = LockedDataBuffer;
 			}
-		};
 
-		FTextureRHIRef SourceTexture = DebugPrimitiveBuffer.DebugPrimitiveTexture->GetRenderTargetItem().TargetableTexture;
-		FTextureRHIRef StagingTexture = DebugPrimitiveBuffer.DebugPrimitiveStagingTexture->GetRenderTargetItem().ShaderResourceTexture;
-
-		// Transfer memory GPU -> CPU
-		RHICmdList.CopyToResolveTarget(SourceTexture, StagingTexture, FResolveParams());
-		if (StagingTexture.IsValid())
-		{	
-			const FVector4* PrimitiveData = nullptr;
-			int32 BufferWidth = 0, BufferHeight = 0;
-			RHICmdList.MapStagingSurface(StagingTexture, *(void**)&PrimitiveData, BufferWidth, BufferHeight);
-			
-			if (PrimitiveData)
+			ValidateShaderParameters(*PixelShader, PassParameters->ShaderDrawPSParameters);
+			ClearUnusedGraphResources(*PixelShader, &PassParameters->ShaderDrawPSParameters);
+			ValidateShaderParameters(*VertexShader, PassParameters->ShaderDrawVSParameters);
+			ClearUnusedGraphResources(*VertexShader, &PassParameters->ShaderDrawVSParameters);
+
+			GraphBuilder.AddPass(
+				RDG_EVENT_NAME("ShaderDrawDebug"),
+				PassParameters,
+				ERDGPassFlags::Raster,
+				[VertexShader, PixelShader, PassParameters, IndirectBuffer, LockedIndirectBuffer, bIsBehindDepth, bUseRdgInput](FRHICommandListImmediate& RHICmdListImmediate)
 			{
-				for (uint32 i = 0; i < PointCount; i += 2)
+				FGraphicsPipelineStateInitializer GraphicsPSOInit;
+				RHICmdListImmediate.ApplyCachedRenderTargets(GraphicsPSOInit);
+				GraphicsPSOInit.DepthStencilState = bIsBehindDepth ? TStaticDepthStencilState<false, CF_DepthFarther>::GetRHI() : TStaticDepthStencilState<false, CF_DepthNearOrEqual>::GetRHI();
+				GraphicsPSOInit.BlendState = TStaticBlendState<CW_RGBA, BO_Add, BF_One, BF_InverseSourceAlpha, BO_Add, BF_Zero, BF_One>::GetRHI(); // Premultiplied-alpha composition
+				GraphicsPSOInit.RasterizerState = TStaticRasterizerState<FM_Solid, CM_None, true>::GetRHI();
+				GraphicsPSOInit.PrimitiveType = PT_LineList;
+				GraphicsPSOInit.BoundShaderState.VertexDeclarationRHI = GEmptyVertexDeclaration.VertexDeclarationRHI;
+				GraphicsPSOInit.BoundShaderState.VertexShaderRHI = GETSAFERHISHADER_VERTEX(*VertexShader);
+				GraphicsPSOInit.BoundShaderState.PixelShaderRHI = GETSAFERHISHADER_PIXEL(*PixelShader);
+				SetGraphicsPipelineState(RHICmdListImmediate, GraphicsPSOInit);
+
+				SetShaderParameters(RHICmdListImmediate, *VertexShader, VertexShader->GetVertexShader(), PassParameters->ShaderDrawVSParameters);
+				SetShaderParameters(RHICmdListImmediate, *PixelShader, PixelShader->GetPixelShader(), PassParameters->ShaderDrawPSParameters);
+
+
+				if (bUseRdgInput)
 				{
-					const FVector P0 = PrimitiveData[i+0];
-					const FVector P1 = PrimitiveData[i+1];
-					const FLinearColor Color = ToColor(uint32(PrimitiveData[i + 0].W));
-					Output.Add({ P0, P1, Color });
+					RHICmdListImmediate.DrawPrimitiveIndirect(IndirectBuffer->GetIndirectRHICallBuffer(), 0);
 				}
-			}
-			RHICmdList.UnmapStagingSurface(StagingTexture);
+				else
+				{
+					RHICmdListImmediate.DrawPrimitiveIndirect(LockedIndirectBuffer, 0);
+				}
+			});
+		};
+
+		FRDGBufferRef DataBuffer = GraphBuilder.RegisterExternalBuffer(View.ShaderDrawData.Buffer, TEXT("ShaderDrawDebugDataBuffer"));
+		FRDGBufferRef IndirectBuffer = GraphBuilder.RegisterExternalBuffer(View.ShaderDrawData.IndirectBuffer, TEXT("ShaderDrawDebugIndirectDataBuffer"));
+		{
+			RunPass(true,  true, DataBuffer, IndirectBuffer, nullptr, nullptr);	// Render what is behind the depth buffer
+			RunPass(false, true, DataBuffer, IndirectBuffer, nullptr, nullptr); // Render what is in front of the depth buffer
 		}
 
-	}
-	return Output;
-}
+		if (LockedData.bIsLocked)
+		{
+			FShaderResourceViewRHIRef LockedDataBuffer = LockedData.Buffer.SRV;
+			FRHIVertexBuffer* LockedIndirectBuffer = LockedData.IndirectBuffer.Buffer;
+			RunPass(false, false, nullptr, nullptr, LockedDataBuffer, LockedIndirectBuffer); // Render what is in front of the depth buffer
+			RunPass(true,  false, nullptr, nullptr, LockedDataBuffer, LockedIndirectBuffer); // Render what is behind the depth buffer
+		}
 
-void BindGpuDebugPrimitiveBuffers(FRHIRenderPassInfo& RPInfo, FGpuDebugPrimitiveBuffers& DebugPrimitiveBuffer, uint32 UAVIndex)
-{
-	// Currently this value is hardcoded into the shader. Once the index is configurable in the shader, this check should be remove
-	check(UAVIndex == 0); 
+		if (IsShaderDrawLocked() && !LockedData.bIsLocked)
+		{
+			{
+				const uint32 NumElements = LockedData.Buffer.NumBytes / sizeof(ShaderDrawDebugElement);
 
-	RPInfo.UAVIndex = UAVIndex;
-	RPInfo.NumUAVs = 2;
-	RPInfo.UAVs[UAVIndex+0] = DebugPrimitiveBuffer.DebugPrimitiveCountTexture->GetRenderTargetItem().UAV;
-	RPInfo.UAVs[UAVIndex+1] = DebugPrimitiveBuffer.DebugPrimitiveTexture->GetRenderTargetItem().UAV;
-}
+				FShaderDrawDebugCopyCS::FPermutationDomain PermutationVector;
+				PermutationVector.Set<FShaderDrawDebugCopyCS::FBufferType>(0);
+				TShaderMapRef<FShaderDrawDebugCopyCS> ComputeShader(View.ShaderMap, PermutationVector);
 
-void DrawGpuDebugPrimitives(FRHICommandListImmediate& RHICmdList, TArray<FViewInfo>& Views, FGpuDebugPrimitiveBuffers& DebugPrimitiveBuffer)
-{
-	const FGpuDebugLineArray DebugLines = ReadGPUDebugPrimitives(RHICmdList, DebugPrimitiveBuffer);
+				FShaderDrawDebugCopyCS::FParameters* Parameters = GraphBuilder.AllocParameters<FShaderDrawDebugCopyCS::FParameters>();
+				Parameters->NumElements = NumElements;
+				Parameters->InStructuredBuffer = GraphBuilder.CreateSRV(DataBuffer);
+				Parameters->OutStructuredBuffer = LockedData.Buffer.UAV;
+				FComputeShaderUtils::AddPass(GraphBuilder, RDG_EVENT_NAME("ShaderDrawDebugCopy"), *ComputeShader, Parameters, FIntVector(FMath::CeilToInt(NumElements / 1024.f), 1, 1));
+			}
+			{
+				FShaderDrawDebugCopyCS::FPermutationDomain PermutationVector;
+				PermutationVector.Set<FShaderDrawDebugCopyCS::FBufferType>(1);
+				TShaderMapRef<FShaderDrawDebugCopyCS> ComputeShader(View.ShaderMap, PermutationVector);
 
-	uint32 ViewIndex = 0;
-	for (FViewInfo& View : Views)
+				FShaderDrawDebugCopyCS::FParameters* Parameters = GraphBuilder.AllocParameters<FShaderDrawDebugCopyCS::FParameters>();
+				Parameters->InBuffer = GraphBuilder.CreateSRV(IndirectBuffer);
+				Parameters->OutBuffer = LockedData.IndirectBuffer.UAV;
+				FComputeShaderUtils::AddPass(GraphBuilder, RDG_EVENT_NAME("ShaderDrawDebugCopy"), *ComputeShader, Parameters, FIntVector(1, 1, 1));
+			}
+		}
+	}
+
+	void EndView(FViewInfo& View)
 	{
-		FViewElementPDI PDI(&View, nullptr, nullptr);
+		if (!IsShaderDrawDebugEnabled(View))
+		{
+			return;
+		}
+
+		if (IsShaderDrawLocked() && !LockedData.bIsLocked)
+		{
+			LockedData.bIsLocked = true;
+		}
 
-		for(int32 i=0; i < DebugLines.Num(); ++i)
+		if (!IsShaderDrawLocked() && LockedData.bIsLocked)
 		{
-			const FGpuDebugLine& Line = DebugLines[i];
-			PDI.DrawLine(Line.Start, Line.End, Line.Color, 0);
+			LockedData.Buffer.Release();
+			LockedData.IndirectBuffer.Release();
+			LockedData.bIsLocked = false;
 		}
+	}
+
+	void SetParameters(FRDGBuilder& GraphBuilder, const FShaderDrawDebugData& Data, FShaderDrawDebugParameters& OutParameters)
+	{
+		FRDGBufferRef DataBuffer = GraphBuilder.RegisterExternalBuffer(Data.Buffer, TEXT("ShaderDrawDebugDataBuffer"));
+		FRDGBufferRef IndirectBuffer = GraphBuilder.RegisterExternalBuffer(Data.IndirectBuffer, TEXT("ShaderDrawDebugIndirectDataBuffer"));
 
-		ViewIndex++;
+		OutParameters.ShaderDrawCursorPos = Data.CursorPosition;
+		OutParameters.ShaderDrawMaxElementCount = GetMaxShaderDrawElementCount();
+		OutParameters.OutShaderDrawPrimitive = GraphBuilder.CreateUAV(DataBuffer);
+		OutParameters.OutputShaderDrawIndirect = GraphBuilder.CreateUAV(IndirectBuffer);
 	}
 }
