Fix shader compile error

diff --git a/Engine/Shaders/Private/HeightFieldAtlasManagement.usf b/Engine/Shaders/Private/HeightFieldAtlasManagement.usf
--- a/Engine/Shaders/Private/HeightFieldAtlasManagement.usf
+++ b/Engine/Shaders/Private/HeightFieldAtlasManagement.usf
@@ -4,6 +4,8 @@
 	HeightFieldAtlasManagement.usf
 =============================================================================*/
 
+#include "Common.ush"
+
 uint4 UpdateRegionOffsetAndSize;
 float4 SourceScaleBias;
 uint SourceMipBias;
@@ -21,6 +23,6 @@ void UploadHeightFieldToAtlasCS(uint3 DispatchThreadId : SV_DispatchThreadID)
 		float2 SourceUV = DispatchThreadId.xy * SourceScaleBias.xy + SourceScaleBias.zw;
 		uint2 WriteCoord = DispatchThreadId.xy + UpdateRegionOffsetAndSize.xy;
 
-		RWHeightFieldAtlas[WriteCoord] = SourceTexture.SampleLevel(SourceTextureSampler, SourceUV, SourceMipBias).xy;
+		RWHeightFieldAtlas[WriteCoord] = Texture2DSampleLevel(SourceTexture, SourceTextureSampler, SourceUV, SourceMipBias).xy;
 	}
 }
