Prevent SSGI from magnifying nans that come from other passes (Guillaume.Abadie)

diff --git a/Engine/Shaders/Private/SSRT/SSRTPrevFrameReduction.usf b/Engine/Shaders/Private/SSRT/SSRTPrevFrameReduction.usf
--- a/Engine/Shaders/Private/SSRT/SSRTPrevFrameReduction.usf
+++ b/Engine/Shaders/Private/SSRT/SSRTPrevFrameReduction.usf
@@ -193,6 +193,8 @@ void MainCS(
 		}
 		#endif
 
+		PrevColor = -min(-PrevColor, 0.0);
+
 		#if CONFIG_COLOR_TILE_CLASSIFICATION
 		{
 			if (bIsSky)
