Reinstate missing clear coat custom data parameters in the ray tracing payload:
1) Clear coat scalar
2) Clear coat roughness scalar

Partial cherry-pick of CL 9894543.

#udn 544855

diff --git a/Engine/Shaders/Private/RayTracing/RayTracingCommon.ush b/Engine/Shaders/Private/RayTracing/RayTracingCommon.ush
--- a/Engine/Shaders/Private/RayTracing/RayTracingCommon.ush
+++ b/Engine/Shaders/Private/RayTracing/RayTracingCommon.ush
@@ -178,7 +178,7 @@ struct FMaterialClosestHitPayload : FMinimalPayload
 	uint ShadingModelID;				// 92        1       4bits
 	uint BlendingMode;					// 96        0       4bits (packed with ShadingModelID)
 	uint PrimitiveLightingChannelMask;	// 100       0       (removed)
-	float4 CustomData;					// 104       4       32bits -- only zw are packed at the moment
+	float4 CustomData;					// 104       4       32bits
 	float GBufferAO;					// 120       0       (removed)
 	float3 IndirectIrradiance;			// 124       0       48bits -- gbuffer only has float payload and there are truncation HLSL warnings
 
@@ -282,9 +282,11 @@ struct FPackedMaterialClosestHitPayload : FMinimalPayload
 	float4 GetCustomData()
 	{
 		float4 Result;
-		Result.z = f16tof32(CustomData.x & 0xFFFF);
-		Result.w = f16tof32(CustomData.x >> 16);
-		Result.xy = 0;
+		Result.x = CustomData & 0xFF;
+		Result.y = (CustomData >> 8) & 0xFF;
+		Result.z = (CustomData >> 16) & 0xFF;
+		Result.w = (CustomData >> 24) & 0xFF;
+		Result /= 255.0;
 		return Result;
 	}
 	
@@ -337,7 +339,8 @@ FPackedMaterialClosestHitPayload PackRayTracingPayload(FMaterialClosestHitPayloa
 	Output.IndirectIrradiance[0] = f32tof16(Input.IndirectIrradiance.x);
 	Output.IndirectIrradiance[0] |= f32tof16(Input.IndirectIrradiance.y) << 16;
 	Output.IndirectIrradiance[1] = f32tof16(Input.IndirectIrradiance.z);
-	Output.CustomData = f32tof16(Input.CustomData.z) | (f32tof16(Input.CustomData.w) << 16);
+	int4 CustomData = round(Input.CustomData * 255);
+	Output.CustomData = CustomData.x | (CustomData.y << 8) | (CustomData.z << 16) | (CustomData.w << 24);
 
 	return Output;
 }
