Workaround to avoid flickering with ray tracing enabled

diff --git a/Engine/Shaders/Private/RayTracing/RayTracingAmbientOcclusionRGS.usf b/Engine/Shaders/Private/RayTracing/RayTracingAmbientOcclusionRGS.usf
--- a/Engine/Shaders/Private/RayTracing/RayTracingAmbientOcclusionRGS.usf
+++ b/Engine/Shaders/Private/RayTracing/RayTracingAmbientOcclusionRGS.usf
@@ -6,6 +6,7 @@
 #include "../SceneTextureParameters.ush"
 #include "../PathTracing/Utilities/PathTracingRandomSequence.ush"
 #include "RayTracingCommon.ush"
+#include "RayTracingDeferredShadingCommon.ush"
 
 uint SamplesPerPixel;
 float MaxRayDistance;
@@ -52,7 +53,14 @@ RAY_TRACING_ENTRY_RAYGEN(AmbientOcclusionRGS)
 	// Get G-Buffer surface data
 	float2 InvBufferSize = View.BufferSizeAndInvSize.zw;
 	float2 UV = (float2(PixelCoord) + 0.5) * InvBufferSize;
-	FGBufferData GBufferData = GetGBufferDataFromSceneTextures(UV);	
+
+#if 0
+	FGBufferData GBufferData = GetGBufferDataFromSceneTextures(UV);
+#else
+	//#dxr-todo: workaround for flickering. UE-87281
+	FGBufferData GBufferData = GetGBufferDataFromSceneTexturesLoad(PixelCoord);
+#endif
+
 	float Depth = GBufferData.Depth;
 	float DeviceZ = SceneDepthBuffer.Load(int3(PixelCoord, 0)).r;
 	float3 WorldPosition = ReconstructWorldPositionFromDeviceZ(PixelCoord, DeviceZ);
diff --git a/Engine/Shaders/Private/RayTracing/RayTracingCreateGatherPointsRGS.usf b/Engine/Shaders/Private/RayTracing/RayTracingCreateGatherPointsRGS.usf
--- a/Engine/Shaders/Private/RayTracing/RayTracingCreateGatherPointsRGS.usf
+++ b/Engine/Shaders/Private/RayTracing/RayTracingCreateGatherPointsRGS.usf
@@ -16,6 +16,7 @@
 #include "../SceneTextureParameters.ush"
 
 #include "RayTracingCommon.ush"
+#include "RayTracingDeferredShadingCommon.ush"
 #include "RayTracingHitGroupCommon.ush"
 #include "RayTracingGatherPoints.ush"
 
@@ -116,7 +117,14 @@ RAY_TRACING_ENTRY_RAYGEN(RayTracingCreateGatherPointsTraceRGS)
 	uint2 PixelCoord = GetPixelCoord(DispatchThreadId, UpscaleFactor);
 	float2 InvBufferSize = View.BufferSizeAndInvSize.zw;
 	float2 UV = (float2(PixelCoord) + 0.5) * InvBufferSize;
+
+#if 0
 	FGBufferData GBufferData = GetGBufferDataFromSceneTextures(UV);
+#else
+	//#dxr-todo: workaround for flickering. UE-87281
+	FGBufferData GBufferData = GetGBufferDataFromSceneTexturesLoad(PixelCoord);
+#endif
+
 	// Remap DiffuseColor when using SubsurfaceProfile (GBuffer decoding replaces with 100% albedo)
 	if (UseSubsurfaceProfile(GBufferData.ShadingModelID))
 	{
@@ -213,7 +221,14 @@ RAY_TRACING_ENTRY_RAYGEN(RayTracingCreateGatherPointsRGS)
 	uint2 PixelCoord = GetPixelCoord(DispatchThreadId, UpscaleFactor);
 	float2 InvBufferSize = View.BufferSizeAndInvSize.zw;
 	float2 UV = (float2(PixelCoord) + 0.5) * InvBufferSize;
+
+#if 0
 	FGBufferData GBufferData = GetGBufferDataFromSceneTextures(UV);
+#else
+	//#dxr-todo: workaround for flickering. UE-87281
+	FGBufferData GBufferData = GetGBufferDataFromSceneTexturesLoad(PixelCoord);
+#endif
+
 	// Remap DiffuseColor when using SubsurfaceProfile (GBuffer decoding replaces with 100% albedo)
 	if (UseSubsurfaceProfile(GBufferData.ShadingModelID))
 	{
diff --git a/Engine/Shaders/Private/RayTracing/RayTracingDeferredShadingCommon.ush b/Engine/Shaders/Private/RayTracing/RayTracingDeferredShadingCommon.ush
new file mode 100644
--- /dev/null
+++ b/Engine/Shaders/Private/RayTracing/RayTracingDeferredShadingCommon.ush
@@ -0,0 +1,50 @@
+// Copyright Epic Games, Inc. All Rights Reserved.
+
+#pragma once
+
+//#include "../DeferredShadingCommon.ush"
+
+//#dxr-todo: workaround for flickering. UE-87281
+
+FGBufferData GetGBufferDataFromSceneTexturesLoad(uint2 PixelCoord, bool bGetNormalizedNormal = true)
+{
+	float4 GBufferA = SceneGBufferA.Load(int3(PixelCoord, 0));
+	float4 GBufferB = SceneGBufferB.Load(int3(PixelCoord, 0));
+	float4 GBufferC = SceneGBufferC.Load(int3(PixelCoord, 0));
+	float4 GBufferD = SceneGBufferD.Load(int3(PixelCoord, 0));
+	float4 GBufferE = SceneGBufferE.Load(int3(PixelCoord, 0));
+	float4 GBufferVelocity = SceneVelocityBuffer.Load(int3(PixelCoord, 0));
+
+	uint CustomStencil = 0;
+	float CustomNativeDepth = 0;
+
+	float DeviceZ = SceneDepthBuffer.Load(int3(PixelCoord, 0)).r;;
+
+	float SceneDepth = ConvertFromDeviceZ(DeviceZ);
+
+	return DecodeGBufferData(GBufferA, GBufferB, GBufferC, GBufferD, GBufferE, GBufferVelocity, CustomNativeDepth, CustomStencil, SceneDepth, bGetNormalizedNormal, CheckerFromPixelPos(PixelCoord));
+}
+
+FGBufferData GetGBufferDataFromPayload(in FMaterialClosestHitPayload Payload)
+{
+	FGBufferData GBufferData = (FGBufferData)0;
+	GBufferData.Depth = 1.f; // Do not use depth
+	GBufferData.WorldNormal = Payload.WorldNormal;
+	GBufferData.BaseColor = Payload.BaseColor;
+	GBufferData.Metallic = Payload.Metallic;
+	GBufferData.Specular = Payload.Specular;
+	GBufferData.Roughness = Payload.Roughness;
+	GBufferData.CustomData = Payload.CustomData;
+	GBufferData.GBufferAO = Payload.GBufferAO;
+	GBufferData.IndirectIrradiance = (Payload.IndirectIrradiance.x + Payload.IndirectIrradiance.y + Payload.IndirectIrradiance.z) / 3.f;
+	GBufferData.ShadingModelID = Payload.ShadingModelID;
+	GBufferData.SpecularColor = Payload.SpecularColor;
+	GBufferData.DiffuseColor = Payload.DiffuseColor;
+	return GBufferData;
+}
+
+FGBufferData GetGBufferDataFromPayload(in FPackedMaterialClosestHitPayload Payload)
+{
+	RayDesc DummyRay = (RayDesc)0; // Not needed to unpack GBuffer
+	return GetGBufferDataFromPayload(UnpackRayTracingPayload(Payload, DummyRay));
+}
\ No newline at end of file
diff --git a/Engine/Shaders/Private/RayTracing/RayTracingFinalGatherRGS.usf b/Engine/Shaders/Private/RayTracing/RayTracingFinalGatherRGS.usf
--- a/Engine/Shaders/Private/RayTracing/RayTracingFinalGatherRGS.usf
+++ b/Engine/Shaders/Private/RayTracing/RayTracingFinalGatherRGS.usf
@@ -7,6 +7,7 @@
 #include "../SceneTextureParameters.ush"
 
 #include "RayTracingCommon.ush"
+#include "RayTracingDeferredShadingCommon.ush"
 #include "RayTracingHitGroupCommon.ush"
 #include "RayTracingGatherPoints.ush"
 
@@ -49,7 +50,14 @@ RAY_TRACING_ENTRY_RAYGEN(RayTracingFinalGatherRGS)
 	uint2 PixelCoord = GetPixelCoord(DispatchThreadId, UpscaleFactor);
 	float2 InvBufferSize = View.BufferSizeAndInvSize.zw;
 	float2 UV = (float2(PixelCoord) + 0.5) * InvBufferSize;
+
+#if 0
 	FGBufferData GBufferData = GetGBufferDataFromSceneTextures(UV);
+#else
+	//#dxr-todo: workaround for flickering. UE-87281
+	FGBufferData GBufferData = GetGBufferDataFromSceneTexturesLoad(PixelCoord);
+#endif
+
 	// Remap DiffuseColor when using SubsurfaceProfile (GBuffer decoding replaces with 100% albedo)
 	if (UseSubsurfaceProfile(GBufferData.ShadingModelID))
 	{
diff --git a/Engine/Shaders/Private/RayTracing/RayTracingGlobalIlluminationRGS.usf b/Engine/Shaders/Private/RayTracing/RayTracingGlobalIlluminationRGS.usf
--- a/Engine/Shaders/Private/RayTracing/RayTracingGlobalIlluminationRGS.usf
+++ b/Engine/Shaders/Private/RayTracing/RayTracingGlobalIlluminationRGS.usf
@@ -5,9 +5,10 @@
 //#include "../MonteCarlo.ush"
 #include "../DeferredShadingCommon.ush"
 #include "../ShadingModels.ush"
+#include "../SceneTextureParameters.ush"
 #include "RayTracingCommon.ush"
+#include "RayTracingDeferredShadingCommon.ush"
 #include "RayTracingHitGroupCommon.ush"
-#include "../SceneTextureParameters.ush"
 
 #include "../PathTracing/Light/PathTracingLightSampling.ush"
 #include "../PathTracing/Material/PathTracingMaterialSampling.ush"
@@ -69,7 +70,13 @@ RAY_TRACING_ENTRY_RAYGEN(GlobalIlluminationRGS)
 	// Get G-Buffer surface data
 	float2 InvBufferSize = View.BufferSizeAndInvSize.zw;
 	float2 UV = (float2(PixelCoord) + 0.5) * InvBufferSize;
+
+#if 0
 	FGBufferData GBufferData = GetGBufferDataFromSceneTextures(UV);
+#else
+	//#dxr-todo: workaround for flickering. UE-87281
+	FGBufferData GBufferData = GetGBufferDataFromSceneTexturesLoad(PixelCoord);
+#endif
 	// Remap DiffuseColor when using SubsurfaceProfile (GBuffer decoding replaces with 100% albedo)
 	if (UseSubsurfaceProfile(GBufferData.ShadingModelID))
 	{
diff --git a/Engine/Shaders/Private/RayTracing/RayTracingLightingCommon.ush b/Engine/Shaders/Private/RayTracing/RayTracingLightingCommon.ush
--- a/Engine/Shaders/Private/RayTracing/RayTracingLightingCommon.ush
+++ b/Engine/Shaders/Private/RayTracing/RayTracingLightingCommon.ush
@@ -90,30 +90,6 @@ FDeferredLightData GetRayTracingDeferredLightData(int LightIndex,
 	return LightData;
 }
 
-FGBufferData GetGBufferDataFromPayload(in FMaterialClosestHitPayload Payload)
-{
-	FGBufferData GBufferData = (FGBufferData)0;
-	GBufferData.Depth = 1.f; // Do not use depth
-	GBufferData.WorldNormal = Payload.WorldNormal;
-	GBufferData.BaseColor = Payload.BaseColor;
-	GBufferData.Metallic = Payload.Metallic;
-	GBufferData.Specular = Payload.Specular;
-	GBufferData.Roughness = Payload.Roughness;
-	GBufferData.CustomData = Payload.CustomData;
-	GBufferData.GBufferAO = Payload.GBufferAO;
-	GBufferData.IndirectIrradiance = (Payload.IndirectIrradiance.x + Payload.IndirectIrradiance.y + Payload.IndirectIrradiance.z) / 3.f;
-	GBufferData.ShadingModelID = Payload.ShadingModelID;
-	GBufferData.SpecularColor = Payload.SpecularColor;
-	GBufferData.DiffuseColor = Payload.DiffuseColor;
-	return GBufferData;
-}
-
-FGBufferData GetGBufferDataFromPayload(in FPackedMaterialClosestHitPayload Payload)
-{
-	RayDesc DummyRay = (RayDesc)0; // Not needed to unpack GBuffer
-	return GetGBufferDataFromPayload(UnpackRayTracingPayload(Payload, DummyRay));
-}
-
 FRectTexture GetRayTracingRectTextureData()
 {
 	FRectTexture RectTexture;
diff --git a/Engine/Shaders/Private/RayTracing/RayTracingOcclusionRGS.usf b/Engine/Shaders/Private/RayTracing/RayTracingOcclusionRGS.usf
--- a/Engine/Shaders/Private/RayTracing/RayTracingOcclusionRGS.usf
+++ b/Engine/Shaders/Private/RayTracing/RayTracingOcclusionRGS.usf
@@ -9,6 +9,7 @@
 #include "../SobolRandom.ush"
 #include "../HairStrands/HairStrandsVisibilityCommon.ush"
 #include "RayTracingCommon.ush"
+#include "RayTracingDeferredShadingCommon.ush"
 #include "RayTracingDirectionalLight.ush"
 #include "RayTracingRectLight.ush"
 #include "RayTracingSphereLight.ush"
@@ -140,15 +141,24 @@ RAY_TRACING_ENTRY_RAYGEN(OcclusionRGS)
 		uint ShadingModelID = SHADINGMODELID_UNLIT;
 		float3 WorldNormal = 0;
 		{
-			const FGBufferData GBufferData = GetGBufferDataFromSceneTextures(BufferUV);
+#if 0
+			FGBufferData GBufferData = GetGBufferDataFromSceneTextures(UV);
+#else
+			//#dxr-todo: workaround for flickering. UE-87281
+			FGBufferData GBufferData = GetGBufferDataFromSceneTexturesLoad(PixelCoord);
+#endif
 			ShadingModelID = GBufferData.ShadingModelID;
 			WorldNormal = GBufferData.WorldNormal;
 		}
 
-
 		// Mask out depth values that are infinitely far away
 
+#if 0
 		float DeviceZ = SceneDepthBuffer.SampleLevel(SceneDepthBufferSampler, BufferUV, 0).r;
+#else
+		//#dxr-todo: workaround for flickering. UE-87281
+		float DeviceZ = SceneDepthBuffer.Load(int3(PixelCoord, 0)).r;
+#endif
 		const bool bIsDepthValid = SceneDepthBuffer.Load(int3(PixelCoord, 0)).r > 0.0;
 		bool bIsValidPixel = ShadingModelID != SHADINGMODELID_UNLIT && bIsDepthValid;
 
diff --git a/Engine/Shaders/Private/RayTracing/RayTracingPrimaryRays.usf b/Engine/Shaders/Private/RayTracing/RayTracingPrimaryRays.usf
--- a/Engine/Shaders/Private/RayTracing/RayTracingPrimaryRays.usf
+++ b/Engine/Shaders/Private/RayTracing/RayTracingPrimaryRays.usf
@@ -26,6 +26,7 @@
 #include "../SobolRandom.ush"
 #include "../SceneTextureParameters.ush"
 #include "RayTracingCommon.ush"
+#include "RayTracingDeferredShadingCommon.ush"
 #include "RayTracingHitGroupCommon.ush"
 #include "RayTracingDirectionalLight.ush"
 #include "RayTracingRectLight.ush"
@@ -116,7 +117,14 @@ RAY_TRACING_ENTRY_RAYGEN(RayTracingPrimaryRaysRGS)
 
 	float2 InvBufferSize = View.BufferSizeAndInvSize.zw;
 	float2 UV = (float2(PixelCoord) + 0.5) * InvBufferSize;
+
+#if 0
 	FGBufferData GBufferData = GetGBufferDataFromSceneTextures(UV);
+#else
+	//#dxr-todo: workaround for flickering. UE-87281
+	FGBufferData GBufferData = GetGBufferDataFromSceneTexturesLoad(PixelCoord);
+#endif
+
 	float Depth = GBufferData.Depth;
 	float3 WorldPosition = ReconstructWorldPositionFromDepth(UV, Depth);
 	float WorldSpaceDistance = length(WorldPosition - View.WorldViewOrigin);
diff --git a/Engine/Shaders/Private/RayTracing/RayTracingReflections.usf b/Engine/Shaders/Private/RayTracing/RayTracingReflections.usf
--- a/Engine/Shaders/Private/RayTracing/RayTracingReflections.usf
+++ b/Engine/Shaders/Private/RayTracing/RayTracingReflections.usf
@@ -49,13 +49,14 @@
 #include "../ScreenSpaceDenoise/SSDPublic.ush"
 #include "../SobolRandom.ush"
 #include "../SceneTextureParameters.ush"
+#include "RayTracingCommon.ush"
+#include "RayTracingDeferredShadingCommon.ush"
 #include "RayTracingHitGroupCommon.ush"
 #include "RayTracingDirectionalLight.ush"
 #include "RayTracingRectLight.ush"
 #include "RayTracingSphereLight.ush"
 #include "RayTracingSpotLight.ush"
 #include "RayTracingPointLight.ush"
-#include "RayTracingCommon.ush"
 #include "RayTracingLightingCommon.ush"
 
 
@@ -239,7 +240,12 @@ RAY_TRACING_ENTRY_RAYGEN(RayTracingReflectionsRGS)
 
 	float2 InvBufferSize = View.BufferSizeAndInvSize.zw;
 	float2 UV = (float2(PixelCoord) + 0.5) * InvBufferSize;
+#if 0
 	FGBufferData GBufferData = GetGBufferDataFromSceneTextures(UV);
+#else
+	//#dxr-todo: workaround for flickering. UE-87281
+	FGBufferData GBufferData = GetGBufferDataFromSceneTexturesLoad(PixelCoord);
+#endif
 	float Depth = GBufferData.Depth;
 
 	uint LocalSamplesPerPixel = SamplesPerPixel; 
diff --git a/Engine/Shaders/Private/RayTracing/RayTracingSkyLightRGS.usf b/Engine/Shaders/Private/RayTracing/RayTracingSkyLightRGS.usf
--- a/Engine/Shaders/Private/RayTracing/RayTracingSkyLightRGS.usf
+++ b/Engine/Shaders/Private/RayTracing/RayTracingSkyLightRGS.usf
@@ -6,10 +6,11 @@
 #include "../ShadingModels.ush"
 #include "../QuasiRandom.ush"
 #include "SkyLightMipTreeCommon.ush"
+#include "../SceneTextureParameters.ush"
 #include "RayTracingCommon.ush"
+#include "RayTracingDeferredShadingCommon.ush"
 #include "RayTracingSkyLightCommon.ush"
 #include "../ReflectionEnvironmentShared.ush"
-#include "../SceneTextureParameters.ush"
 
 struct FSkyLightVisibilityRays
 {
@@ -37,7 +38,14 @@ RAY_TRACING_ENTRY_RAYGEN(SkyLightRGS)
 	float2 InvBufferSize = View.BufferSizeAndInvSize.zw;
 	float2 UV = (float2(PixelCoord) + 0.5) * InvBufferSize;
 	FScreenSpaceData ScreenSpaceData = GetScreenSpaceData(UV);
+
+#if 0
 	FGBufferData GBufferData = GetGBufferDataFromSceneTextures(UV);
+#else
+	//#dxr-todo: workaround for flickering. UE-87281
+	FGBufferData GBufferData = GetGBufferDataFromSceneTexturesLoad(PixelCoord);
+#endif
+
 	float DeviceZ = SceneDepthBuffer.Load(int3(PixelCoord, 0)).r;
 	float3 WorldPosition;
 	float3 CameraDirection;
