Fixed bad merge issue with ray tracing shaders

diff --git a/Engine/Shaders/Private/RayTracing/RayTracingLightingMS.usf b/Engine/Shaders/Private/RayTracing/RayTracingLightingMS.usf
--- a/Engine/Shaders/Private/RayTracing/RayTracingLightingMS.usf
+++ b/Engine/Shaders/Private/RayTracing/RayTracingLightingMS.usf
@@ -14,6 +14,9 @@
 #define PreIntegratedGFSampler		ReflectionStruct.PreIntegratedGFSampler	
 #define SSProfilesTexture			RaytracingLightsDataPacked.SSProfilesTexture
 
+#include "../SceneTextureParameters.ush"
+#include "RayTracingCommon.ush"
+#include "RayTracingDeferredShadingCommon.ush"
 #include "RayTracingLightingCommon.ush" 
 
 RAY_TRACING_ENTRY_MISS(RayTracingLightingMS, FPackedMaterialClosestHitPayload, PackedPayload)
diff --git a/Engine/Shaders/Private/RayTracing/RayTracingSkyLightRGS.usf b/Engine/Shaders/Private/RayTracing/RayTracingSkyLightRGS.usf
--- a/Engine/Shaders/Private/RayTracing/RayTracingSkyLightRGS.usf
+++ b/Engine/Shaders/Private/RayTracing/RayTracingSkyLightRGS.usf
@@ -4,6 +4,7 @@
 #include "../SceneTextureParameters.ush"
 #include "../SceneTextureParameters.ush"
 #include "RayTracingSkyLightEvaluation.ush"
+#include "RayTracingDeferredShadingCommon.ush"
 
 RaytracingAccelerationStructure TLAS;
 
