Fixed: EARLYDEPTHSTENCIL usage for HLSL based mobile preview

diff --git a/Engine/Shaders/Private/MobileBasePassPixelShader.usf b/Engine/Shaders/Private/MobileBasePassPixelShader.usf
--- a/Engine/Shaders/Private/MobileBasePassPixelShader.usf
+++ b/Engine/Shaders/Private/MobileBasePassPixelShader.usf
@@ -395,7 +395,16 @@ void ApplyPixelDepthOffsetForMobileBasePass(inout FMaterialPixelParameters Mater
     float PixelDepthOffset = ApplyPixelDepthOffsetToMaterialParameters(MaterialParameters, PixelMaterialInputs, OutDepth);
 }
 
-EARLYDEPTHSTENCIL
+
+// Disable early depth_stencil for masked materials and depth writes
+// Cross-compiler (hlslcc) checks for 'discard' and depth write usage anyway and disables it, this ifdef is needed for HLSL
+#if MATERIALBLENDING_MASKED || USE_DITHERED_LOD_TRANSITION || OUTPUT_PIXEL_DEPTH_OFFSET
+	#define PIXELSHADER_EARLYDEPTHSTENCIL 
+#else
+	#define PIXELSHADER_EARLYDEPTHSTENCIL EARLYDEPTHSTENCIL
+#endif
+
+PIXELSHADER_EARLYDEPTHSTENCIL
 void Main( 
 	FVertexFactoryInterpolantsVSToPS Interpolants
 	, FMobileBasePassInterpolantsVSToPS BasePassInterpolants
