Applied same workaround to GpuSkinCacheComputeShader as for GpySkinVertexFactory.

#rnx

diff --git a/Engine/Shaders/Private/GpuSkinCacheComputeShader.usf b/Engine/Shaders/Private/GpuSkinCacheComputeShader.usf
--- a/Engine/Shaders/Private/GpuSkinCacheComputeShader.usf
+++ b/Engine/Shaders/Private/GpuSkinCacheComputeShader.usf
@@ -142,9 +142,16 @@ FBoneMatrix CalcBoneMatrix( FVertexUnpacked VertexUnpacked, uint VertexIndex )
 		if (InputWeightIndexSize > 1)
 		{
 			BoneIndex = InputWeightStream[BoneIndexOffset + 1] << 8 | BoneIndex;
+			//@todo-lh: Workaround to fix issue in SPIRVEmitter of DXC; this block must be inside the if branch
+			float BoneWeight = float(InputWeightStream[WeightsOffset + InfluenceIdx]) / 255.0f;
+			BoneMatrix += BoneWeight * FBoneMatrix(BoneMatrices[BoneIndex*3], BoneMatrices[BoneIndex*3+1], BoneMatrices[BoneIndex*3+2]);
+		}
+		else
+		{
+			//@todo-lh: Workaround to fix issue in SPIRVEmitter of DXC; this block must be inside the if branch
+			float BoneWeight = float(InputWeightStream[WeightsOffset + InfluenceIdx]) / 255.0f;
+			BoneMatrix += BoneWeight * FBoneMatrix(BoneMatrices[BoneIndex*3], BoneMatrices[BoneIndex*3+1], BoneMatrices[BoneIndex*3+2]);
 		}
-		float BoneWeight = float(InputWeightStream[WeightsOffset + InfluenceIdx]) / 255.0f;
-		BoneMatrix += BoneWeight * FBoneMatrix(BoneMatrices[BoneIndex*3], BoneMatrices[BoneIndex*3+1], BoneMatrices[BoneIndex*3+2]);
 	}
 #else // GPUSKIN_UNLIMITED_BONE_INFLUENCE
 	FBoneMatrix BoneMatrix = VertexUnpacked.BlendWeights.x * FBoneMatrix(BoneMatrices[VertexUnpacked.BlendIndices.x*3], BoneMatrices[VertexUnpacked.BlendIndices.x*3+1], BoneMatrices[VertexUnpacked.BlendIndices.x*3+2]);
