Changed geometry tolerance of primary rays in translucency pass to avoid artifacts due transparent objects contacting opaque objects

diff --git a/Engine/Shaders/Private/RayTracing/RayTracingPrimaryRays.usf b/Engine/Shaders/Private/RayTracing/RayTracingPrimaryRays.usf
--- a/Engine/Shaders/Private/RayTracing/RayTracingPrimaryRays.usf
+++ b/Engine/Shaders/Private/RayTracing/RayTracingPrimaryRays.usf
@@ -136,7 +136,7 @@ RAY_TRACING_ENTRY_RAYGEN(RayTracingPrimaryRaysRGS)
 
 	if((ERayTracingPrimaryRaysFlag_UseGBufferForMaxDistance & PrimaryRayFlags) != 0) 
 	{
-		Ray.TMax = WorldSpaceDistance - 0.1;
+		Ray.TMax = WorldSpaceDistance - 0.01;
 	}
 
 	bool bAllowSkySampling;
