Removed the last remaining pieces of the emulated instancing code. Added new FRenderingObjectVersion entry for this.

Picked from 11156339

diff --git a/Engine/Shaders/Private/LocalVertexFactory.ush b/Engine/Shaders/Private/LocalVertexFactory.ush
--- a/Engine/Shaders/Private/LocalVertexFactory.ush
+++ b/Engine/Shaders/Private/LocalVertexFactory.ush
@@ -27,7 +27,6 @@
 #endif
 
 #if USE_INSTANCING
-
 #if USE_DITHERED_LOD_TRANSITION
 	float4 InstancingViewZCompareZero;  // w contains random lod scale 
 	float4 InstancingViewZCompareOne;
@@ -38,22 +37,6 @@
 
 float4 InstancingOffset;
 float4 InstancingFadeOutParams;
-
-#ifndef USE_INSTANCING_EMULATED
-	#define USE_INSTANCING_EMULATED 0
-#endif	// USE_INSTANCING_EMULATED
-
-#if USE_INSTANCING_EMULATED
-	// Required for CPU emulation of Instancing
-	float4 CPUInstanceOrigin;  // per-instance random in w 
-	float3x4 CPUInstanceTransform; // hitproxy.r + 256 * selected in .w; hitproxy.g in .w; hitproxy.b in .w
-	float4 CPUInstanceLightmapAndShadowMapBias; 
-#endif	// USE_INSTANCING_EMULATED
-
-#else // USE_INSTANCING
-	#ifndef USE_INSTANCING_EMULATED
-		#define USE_INSTANCING_EMULATED (0)
-	#endif
 #endif	// USE_INSTANCING
 
 #if USE_SPLINEDEFORM
@@ -87,7 +70,7 @@ float4 SplineParams[10];
 	#define FV_LightMapIndex_Index 2
 	#define VF_VertexOffset 3
 
-#if !USE_INSTANCING_EMULATED
+
 	Buffer<float4> VertexFetch_InstanceOriginBuffer;
 	Buffer<float4> VertexFetch_InstanceTransformBuffer;
 	Buffer<float4> VertexFetch_InstanceLightmapBuffer;
@@ -98,7 +81,7 @@ float4 SplineParams[10];
 	#endif
 
 	uint InstanceOffset;
-#endif //! !USE_INSTANCING_EMULATED
+
 #endif //! MANUAL_VERTEX_FETCH
 
 /**
@@ -151,13 +134,13 @@ struct FVertexFactoryInput
 	float2	TexCoords[1] : ATTRIBUTE4;
 #endif
 
-#if USE_INSTANCING && !USE_INSTANCING_EMULATED && !MANUAL_VERTEX_FETCH
+#if USE_INSTANCING && !MANUAL_VERTEX_FETCH
 	float4 InstanceOrigin : ATTRIBUTE8;  // per-instance random in w 
 	half4 InstanceTransform1 : ATTRIBUTE9;  // hitproxy.r + 256 * selected in .w
 	half4 InstanceTransform2 : ATTRIBUTE10; // hitproxy.g in .w
 	half4 InstanceTransform3 : ATTRIBUTE11; // hitproxy.b in .w
 	float4 InstanceLightmapAndShadowMapUVBias : ATTRIBUTE12; 
-#endif //USE_INSTANCING && !USE_INSTANCING_EMULATED
+#endif //USE_INSTANCING
 
 #if VF_USE_PRIMITIVE_SCENE_DATA
 	uint PrimitiveId : ATTRIBUTE13;
@@ -168,7 +151,7 @@ struct FVertexFactoryInput
 #endif
 
 /** Optional instance ID for vertex layered rendering */
-#if FEATURE_LEVEL >= FEATURE_LEVEL_ES3_1 && (MANUAL_VERTEX_FETCH && (USE_INSTANCING && !USE_INSTANCING_EMULATED))
+#if FEATURE_LEVEL >= FEATURE_LEVEL_ES3_1 && (MANUAL_VERTEX_FETCH && USE_INSTANCING)
 	uint InstanceId	: SV_InstanceID;
 #endif
 
@@ -247,19 +230,19 @@ struct FPositionOnlyVertexFactoryInput
 {
 	float4	Position	: ATTRIBUTE0;
 
-#if USE_INSTANCING && !USE_INSTANCING_EMULATED && !MANUAL_VERTEX_FETCH
+#if USE_INSTANCING && !MANUAL_VERTEX_FETCH
 	float4 InstanceOrigin : ATTRIBUTE8;  // per-instance random in w 
 	half4 InstanceTransform1 : ATTRIBUTE9;  // hitproxy.r + 256 * selected in .w
 	half4 InstanceTransform2 : ATTRIBUTE10; // hitproxy.g in .w
 	half4 InstanceTransform3 : ATTRIBUTE11; // hitproxy.b in .w
-#endif	// USE_INSTANCING && !USE_INSTANCING_EMULATED
+#endif	// USE_INSTANCING
 
 #if VF_USE_PRIMITIVE_SCENE_DATA
 	uint PrimitiveId : ATTRIBUTE1;
 #endif
 
 /** Optional instance ID for vertex layered rendering */
-#if (FEATURE_LEVEL >= FEATURE_LEVEL_ES3_1) && (MANUAL_VERTEX_FETCH && (USE_INSTANCING && !USE_INSTANCING_EMULATED))
+#if (FEATURE_LEVEL >= FEATURE_LEVEL_ES3_1) && (MANUAL_VERTEX_FETCH && USE_INSTANCING)
 	uint InstanceId	: SV_InstanceID;
 #endif
 
@@ -276,19 +259,19 @@ struct FPositionAndNormalOnlyVertexFactoryInput
 	float4	Position	: ATTRIBUTE0;
 	float4	Normal		: ATTRIBUTE2;
 
-#if USE_INSTANCING && !USE_INSTANCING_EMULATED && !MANUAL_VERTEX_FETCH
+#if USE_INSTANCING && !MANUAL_VERTEX_FETCH
 	float4 InstanceOrigin : ATTRIBUTE8;  // per-instance random in w 
 	half4 InstanceTransform1 : ATTRIBUTE9;  // hitproxy.r + 256 * selected in .w
 	half4 InstanceTransform2 : ATTRIBUTE10; // hitproxy.g in .w
 	half4 InstanceTransform3 : ATTRIBUTE11; // hitproxy.b in .w
-#endif	// USE_INSTANCING && !USE_INSTANCING_EMULATED
+#endif	// USE_INSTANCING
 
 #if VF_USE_PRIMITIVE_SCENE_DATA
 	uint PrimitiveId : ATTRIBUTE1;
 #endif
 
 /** Optional instance ID for vertex layered rendering */
-#if (FEATURE_LEVEL >= FEATURE_LEVEL_ES3_1) && (MANUAL_VERTEX_FETCH && (USE_INSTANCING && !USE_INSTANCING_EMULATED))
+#if (FEATURE_LEVEL >= FEATURE_LEVEL_ES3_1) && (MANUAL_VERTEX_FETCH && USE_INSTANCING)
 	uint InstanceId	: SV_InstanceID;
 #endif
 
@@ -308,7 +291,6 @@ struct FVertexFactoryIntermediates
 
 	half4 Color;
 #if USE_INSTANCING
-	#if !USE_INSTANCING_EMULATED
 		float4 InstanceOrigin;
 		float4 InstanceTransform1;
 		float4 InstanceTransform2;
@@ -322,7 +304,7 @@ struct FVertexFactoryIntermediates
 		#endif
 
 		float4 InstanceLightmapAndShadowMapUVBias;
-	#endif
+
 	// x = per-instance random, y = per-instance fade out amount, z = hide/show flag, w dither fade cutoff
 	float4 PerInstanceParams;
 #endif	// USE_INSTANCING
@@ -332,19 +314,11 @@ struct FVertexFactoryIntermediates
 #if USE_INSTANCING
 float4x4 GetInstanceTransform(FVertexFactoryIntermediates Intermediates)
 {
-	#if !USE_INSTANCING_EMULATED
 	return float4x4(
 		float4(Intermediates.InstanceTransform1.xyz, 0.0f),
 		float4(Intermediates.InstanceTransform2.xyz, 0.0f),
 		float4(Intermediates.InstanceTransform3.xyz, 0.0f),
 		float4(Intermediates.InstanceOrigin.xyz, 1.0f));
-	#else
-	return float4x4(
-		float4(CPUInstanceTransform[0].xyz, 0.0f), 
-		float4(CPUInstanceTransform[1].xyz, 0.0f), 
-		float4(CPUInstanceTransform[2].xyz, 0.0f), 
-		float4(CPUInstanceOrigin.xyz, 1.0f));
-	#endif	// !USE_INSTANCING_EMULATED
 }
 
 #if USE_INSTANCING && USE_INSTANCING_BONEMAP
@@ -360,7 +334,7 @@ float4x4 GetInstancePrevTransform(FVertexFactoryIntermediates Intermediates)
 
 float4x4 GetInstanceTransform(FPositionOnlyVertexFactoryInput Input)
 {
-	#if !USE_INSTANCING_EMULATED && MANUAL_VERTEX_FETCH && !USE_INSTANCING_BONEMAP
+	#if MANUAL_VERTEX_FETCH && !USE_INSTANCING_BONEMAP
 	uint InstanceId = GetInstanceId(Input.InstanceId);
 	return float4x4(
 		float4(VertexFetch_InstanceTransformBuffer[3 * (InstanceId + InstanceOffset) + 0].xyz, 0.0f),
@@ -374,24 +348,18 @@ float4x4 GetInstanceTransform(FPositionOnlyVertexFactoryInput Input)
 		float4(VertexFetch_InstanceTransformBuffer[4 * InstanceIndex + 1].xyz, 0.0f),
 		float4(VertexFetch_InstanceTransformBuffer[4 * InstanceIndex + 2].xyz, 0.0f),
 		float4(VertexFetch_InstanceTransformBuffer[4 * InstanceIndex + 3].xyz, 1.0f));
-	#elif !USE_INSTANCING_EMULATED
+	#else 
 	return float4x4(
 		float4(Input.InstanceTransform1.xyz, 0.0f),
 		float4(Input.InstanceTransform2.xyz, 0.0f),
 		float4(Input.InstanceTransform3.xyz, 0.0f),
 		float4(Input.InstanceOrigin.xyz, 1.0f));
-	#else
-	return float4x4(
-		float4(CPUInstanceTransform[0].xyz, 0.0f), 
-		float4(CPUInstanceTransform[1].xyz, 0.0f), 
-		float4(CPUInstanceTransform[2].xyz, 0.0f), 
-		float4(CPUInstanceOrigin.xyz, 1.0f));
-	#endif	// !USE_INSTANCING_EMULATED
+	#endif
 }
 
 float4x4 GetInstanceTransform(FPositionAndNormalOnlyVertexFactoryInput Input)
 {
-	#if !USE_INSTANCING_EMULATED && MANUAL_VERTEX_FETCH && !USE_INSTANCING_BONEMAP
+	#if MANUAL_VERTEX_FETCH && !USE_INSTANCING_BONEMAP
 	uint InstanceId = GetInstanceId(Input.InstanceId);
 	return float4x4(
 		float4(VertexFetch_InstanceTransformBuffer[3 * (InstanceId + InstanceOffset) + 0].xyz, 0.0f),
@@ -405,76 +373,43 @@ float4x4 GetInstanceTransform(FPositionAndNormalOnlyVertexFactoryInput Input)
 			float4(VertexFetch_InstanceTransformBuffer[4 * InstanceIndex + 1].xyz, 0.0f),
 			float4(VertexFetch_InstanceTransformBuffer[4 * InstanceIndex + 2].xyz, 0.0f),
 			float4(VertexFetch_InstanceTransformBuffer[4 * InstanceIndex + 3].xyz, 1.0f));
-	#elif !USE_INSTANCING_EMULATED
+	#else 
 	return float4x4(
 		float4(Input.InstanceTransform1.xyz, 0.0f),
 		float4(Input.InstanceTransform2.xyz, 0.0f),
 		float4(Input.InstanceTransform3.xyz, 0.0f),
 		float4(Input.InstanceOrigin.xyz, 1.0f));
-	#else
-	return float4x4(
-		float4(CPUInstanceTransform[0].xyz, 0.0f), 
-		float4(CPUInstanceTransform[1].xyz, 0.0f), 
-		float4(CPUInstanceTransform[2].xyz, 0.0f), 
-		float4(CPUInstanceOrigin.xyz, 1.0f));
-	#endif	// !USE_INSTANCING_EMULATED
+	#endif
 }
 
 half3x3 GetInstanceToLocal3x3(FVertexFactoryIntermediates Intermediates)
 {
-	#if !USE_INSTANCING_EMULATED
 	return half3x3(Intermediates.InstanceTransform1.xyz, Intermediates.InstanceTransform2.xyz, Intermediates.InstanceTransform3.xyz);
-	#else
-	return half3x3(half3(CPUInstanceTransform[0].xyz), half3(CPUInstanceTransform[1].xyz), half3(CPUInstanceTransform[2].xyz));
-	#endif	// !USE_INSTANCING_EMULATED
 }
 
 float2 GetInstanceShadowMapBias(FVertexFactoryIntermediates Intermediates)
 {
-	#if !USE_INSTANCING_EMULATED
 	return Intermediates.InstanceLightmapAndShadowMapUVBias.zw;
-	#else
-	return CPUInstanceLightmapAndShadowMapBias.zw;
-	#endif	// !USE_INSTANCING_EMULATED
 }
 
 float2 GetInstanceLightMapBias(FVertexFactoryIntermediates Intermediates)
 {
-	#if !USE_INSTANCING_EMULATED
 	return Intermediates.InstanceLightmapAndShadowMapUVBias.xy;
-	#else
-	return CPUInstanceLightmapAndShadowMapBias.xy;
-	#endif	// !USE_INSTANCING_EMULATED
 }
 
 float GetInstanceSelected(FVertexFactoryIntermediates Intermediates)
 {
-	#if !USE_INSTANCING_EMULATED
-		float SelectedValue = trunc(Intermediates.InstanceTransform1.w * (1.0 / 256.0));
-	#else
-		float SelectedValue = trunc(CPUInstanceTransform[0].w * (1.0 / 256.0));
-	#endif	// !USE_INSTANCING_EMULATED
-	return SelectedValue;
+	return trunc(Intermediates.InstanceTransform1.w * (1.0 / 256.0));;
 }
 
 float GetInstanceRandom(FVertexFactoryIntermediates Intermediates)
 {
-	#if !USE_INSTANCING_EMULATED
-		float RandomVal = Intermediates.InstanceOrigin.w;
-	#else
-		float RandomVal = CPUInstanceOrigin.w;
-	#endif	// !USE_INSTANCING_EMULATED
-	return RandomVal;
+	return Intermediates.InstanceOrigin.w;
 }
 
 float3 GetInstanceOrigin(FVertexFactoryIntermediates Intermediates)
 {
-	#if !USE_INSTANCING_EMULATED
-		float3 Origin = Intermediates.InstanceOrigin.xyz + InstancingOffset;
-	#else
-		float3 Origin = CPUInstanceOrigin.xyz;
-	#endif	// !USE_INSTANCING_EMULATED
-	return Origin;
+	return Intermediates.InstanceOrigin.xyz + InstancingOffset;
 }
 
 #endif	// USE_INSTANCING
@@ -807,14 +742,14 @@ FVertexFactoryIntermediates GetVertexFactoryIntermediates(FVertexFactoryInput In
 	Intermediates.Color = Input.Color FCOLOR_COMPONENT_SWIZZLE; // Swizzle vertex color.
 #endif
 
-#if USE_INSTANCING && !USE_INSTANCING_EMULATED && MANUAL_VERTEX_FETCH && !USE_INSTANCING_BONEMAP
+#if USE_INSTANCING && MANUAL_VERTEX_FETCH && !USE_INSTANCING_BONEMAP
 	uint InstanceId = GetInstanceId(Input.InstanceId);
 	Intermediates.InstanceTransform1 = VertexFetch_InstanceTransformBuffer[3 * (InstanceId + InstanceOffset) + 0];
 	Intermediates.InstanceTransform2 = VertexFetch_InstanceTransformBuffer[3 * (InstanceId + InstanceOffset) + 1];
 	Intermediates.InstanceTransform3 = VertexFetch_InstanceTransformBuffer[3 * (InstanceId + InstanceOffset) + 2];
 	Intermediates.InstanceOrigin = VertexFetch_InstanceOriginBuffer[(InstanceId + InstanceOffset)];
 	Intermediates.InstanceLightmapAndShadowMapUVBias = VertexFetch_InstanceLightmapBuffer[(InstanceId + InstanceOffset)];
-#elif !USE_INSTANCING_EMULATED && MANUAL_VERTEX_FETCH && USE_INSTANCING_BONEMAP
+#elif MANUAL_VERTEX_FETCH && USE_INSTANCING_BONEMAP
 	uint InstanceIndex = VertexFetch_InstanceBoneMapBuffer[LocalVF.VertexFetch_Parameters[VF_VertexOffset] + Input.VertexId];
 	Intermediates.InstanceTransform1 = VertexFetch_InstanceTransformBuffer[4 * InstanceIndex + 0];
 	Intermediates.InstanceTransform2 = VertexFetch_InstanceTransformBuffer[4 * InstanceIndex + 1];
@@ -827,7 +762,7 @@ FVertexFactoryIntermediates GetVertexFactoryIntermediates(FVertexFactoryInput In
 	Intermediates.InstancePrevOrigin = VertexFetch_InstancePrevTransformBuffer[4 * InstanceIndex + 3];	
 
 	Intermediates.InstanceLightmapAndShadowMapUVBias = float4(0,0,0,0);
-#elif USE_INSTANCING && !USE_INSTANCING_EMULATED
+#elif USE_INSTANCING
 	Intermediates.InstanceTransform1 = Input.InstanceTransform1;
 	Intermediates.InstanceTransform2 = Input.InstanceTransform2;
 	Intermediates.InstanceTransform3 = Input.InstanceTransform3;
@@ -1202,15 +1137,10 @@ float4 VertexFactoryGetPreviousWorldPosition(FVertexFactoryInput Input, FVertexF
 #if USE_INSTANCING
 float4 VertexFactoryGetInstanceHitProxyId(FVertexFactoryInput Input, FVertexFactoryIntermediates Intermediates)
 {
-	#if !USE_INSTANCING_EMULATED
 		float R = Intermediates.InstanceTransform1.w - 256.0 * GetInstanceSelected(Intermediates);
 		float G = Intermediates.InstanceTransform2.w;
 		float B = Intermediates.InstanceTransform3.w;
-	#else
-		float R = CPUInstanceTransform[0].w - 256.0 * GetInstanceSelected(Intermediates);
-		float G = CPUInstanceTransform[1].w;
-		float B = CPUInstanceTransform[2].w;
-	#endif
+
 	return float4(R/255.0, G/255.0, B/255.0, 0);
 }
 #endif	// USE_INSTANCING
diff --git a/Engine/Source/Runtime/Core/Public/UObject/RenderingObjectVersion.h b/Engine/Source/Runtime/Core/Public/UObject/RenderingObjectVersion.h
--- a/Engine/Source/Runtime/Core/Public/UObject/RenderingObjectVersion.h
+++ b/Engine/Source/Runtime/Core/Public/UObject/RenderingObjectVersion.h
@@ -106,6 +106,9 @@ struct CORE_API FRenderingObjectVersion
 		// Add force opaque flag for static mesh
 		AutoExposureChanges,
 
+		// Removed emulated instancing from instanced static meshes
+		RemovedEmulatedInstancing,
+
 		// -----<new versions can be added above this line>-------------------------------------------------
 		VersionPlusOne,
 		LatestVersion = VersionPlusOne - 1
diff --git a/Engine/Source/Runtime/Engine/Private/InstancedStaticMesh.h b/Engine/Source/Runtime/Engine/Private/InstancedStaticMesh.h
--- a/Engine/Source/Runtime/Engine/Private/InstancedStaticMesh.h
+++ b/Engine/Source/Runtime/Engine/Private/InstancedStaticMesh.h
@@ -305,9 +305,6 @@ class FInstancedStaticMeshVertexFactoryShaderParameters : public FLocalVertexFac
 		InstancingOffsetParameter.Bind(ParameterMap, TEXT("InstancingOffset"));
 		InstancingWorldViewOriginZeroParameter.Bind(ParameterMap, TEXT("InstancingWorldViewOriginZero"));
 		InstancingWorldViewOriginOneParameter.Bind(ParameterMap, TEXT("InstancingWorldViewOriginOne"));
-		CPUInstanceOrigin.Bind(ParameterMap, TEXT("CPUInstanceOrigin"));
-		CPUInstanceTransform.Bind(ParameterMap, TEXT("CPUInstanceTransform"));
-		CPUInstanceLightmapAndShadowMapBias.Bind(ParameterMap, TEXT("CPUInstanceLightmapAndShadowMapBias"));
 		VertexFetch_InstanceOriginBufferParameter.Bind(ParameterMap, TEXT("VertexFetch_InstanceOriginBuffer"));
 		VertexFetch_InstanceTransformBufferParameter.Bind(ParameterMap, TEXT("VertexFetch_InstanceTransformBuffer"));
 		VertexFetch_InstanceLightmapBufferParameter.Bind(ParameterMap, TEXT("VertexFetch_InstanceLightmapBuffer"));
@@ -328,6 +325,8 @@ class FInstancedStaticMeshVertexFactoryShaderParameters : public FLocalVertexFac
 
 	void Serialize(FArchive& Ar) override
 	{
+		Ar.UsingCustomVersion(FRenderingObjectVersion::GUID);
+
 		FLocalVertexFactoryShaderParametersBase::Serialize(Ar);
 		Ar << InstancingFadeOutParamsParameter;
 		Ar << InstancingViewZCompareZeroParameter;
@@ -336,9 +335,14 @@ class FInstancedStaticMeshVertexFactoryShaderParameters : public FLocalVertexFac
 		Ar << InstancingOffsetParameter;
 		Ar << InstancingWorldViewOriginZeroParameter;
 		Ar << InstancingWorldViewOriginOneParameter;
-		Ar << CPUInstanceOrigin;
-		Ar << CPUInstanceTransform;
-		Ar << CPUInstanceLightmapAndShadowMapBias;
+		
+		if (Ar.IsLoading() && Ar.CustomVer(FRenderingObjectVersion::GUID) < FRenderingObjectVersion::RemovedEmulatedInstancing)
+		{
+			// Older version need to also load legacy emulated instancing shader parameters
+			FShaderParameter Dummy;
+			Ar << Dummy << Dummy << Dummy;
+		}
+
 		Ar << VertexFetch_InstanceOriginBufferParameter;
 		Ar << VertexFetch_InstanceTransformBufferParameter;
 		Ar << VertexFetch_InstanceLightmapBufferParameter;
@@ -356,10 +360,6 @@ private:
 	FShaderParameter InstancingWorldViewOriginZeroParameter;
 	FShaderParameter InstancingWorldViewOriginOneParameter;
 
-	FShaderParameter CPUInstanceOrigin;
-	FShaderParameter CPUInstanceTransform;
-	FShaderParameter CPUInstanceLightmapAndShadowMapBias;
-
 	FShaderResourceParameter VertexFetch_InstanceOriginBufferParameter;
 	FShaderResourceParameter VertexFetch_InstanceTransformBufferParameter;
 	FShaderResourceParameter VertexFetch_InstanceLightmapBufferParameter;
diff --git a/Engine/Source/Runtime/Engine/Public/GeometryCollectionRendering.h b/Engine/Source/Runtime/Engine/Public/GeometryCollectionRendering.h
--- a/Engine/Source/Runtime/Engine/Public/GeometryCollectionRendering.h
+++ b/Engine/Source/Runtime/Engine/Public/GeometryCollectionRendering.h
@@ -95,8 +95,7 @@ public:
 		if (RHISupportsManualVertexFetch(Platform))
 		{
 			OutEnvironment.SetDefine(TEXT("USE_INSTANCING"), TEXT("1"));
-			OutEnvironment.SetDefine(TEXT("USE_INSTANCING_EMULATED"), TEXT("0"));
-		
+
 			OutEnvironment.SetDefine(TEXT("USE_INSTANCING_BONEMAP"), TEXT("1"));
 			OutEnvironment.SetDefine(TEXT("USE_DITHERED_LOD_TRANSITION_FOR_INSTANCED"), TEXT("0"));
 		}
