Removed unused shader #ifdef that was missed in the previous CL

diff --git a/Engine/Shaders/Private/LocalVertexFactory.ush b/Engine/Shaders/Private/LocalVertexFactory.ush
--- a/Engine/Shaders/Private/LocalVertexFactory.ush
+++ b/Engine/Shaders/Private/LocalVertexFactory.ush
@@ -341,7 +341,7 @@ float4x4 GetInstanceTransform(FPositionOnlyVertexFactoryInput Input)
 		float4(VertexFetch_InstanceTransformBuffer[3 * (InstanceId + InstanceOffset) + 1].xyz, 0.0f),
 		float4(VertexFetch_InstanceTransformBuffer[3 * (InstanceId + InstanceOffset) + 2].xyz, 0.0f),
 		float4(VertexFetch_InstanceOriginBuffer[(InstanceId + InstanceOffset)].xyz, 1.0f));
-	#elif !USE_INSTANCING_EMULATED && MANUAL_VERTEX_FETCH && USE_INSTANCING_BONEMAP
+	#elif MANUAL_VERTEX_FETCH && USE_INSTANCING_BONEMAP
 	uint InstanceIndex = VertexFetch_InstanceBoneMapBuffer[LocalVF.VertexFetch_Parameters[VF_VertexOffset] + Input.VertexId];
 	return float4x4(
 		float4(VertexFetch_InstanceTransformBuffer[4 * InstanceIndex + 0].xyz, 0.0f),
@@ -366,7 +366,7 @@ float4x4 GetInstanceTransform(FPositionAndNormalOnlyVertexFactoryInput Input)
 		float4(VertexFetch_InstanceTransformBuffer[3 * (InstanceId + InstanceOffset) + 1].xyz, 0.0f),
 		float4(VertexFetch_InstanceTransformBuffer[3 * (InstanceId + InstanceOffset) + 2].xyz, 0.0f),
 		float4(VertexFetch_InstanceOriginBuffer[(InstanceId + InstanceOffset)].xyz, 1.0f));
-	#elif !USE_INSTANCING_EMULATED && MANUAL_VERTEX_FETCH && USE_INSTANCING_BONEMAP
+	#elif MANUAL_VERTEX_FETCH && USE_INSTANCING_BONEMAP
 		uint InstanceIndex = VertexFetch_InstanceBoneMapBuffer[LocalVF.VertexFetch_Parameters[VF_VertexOffset] + Input.VertexId];
 		return float4x4(
 			float4(VertexFetch_InstanceTransformBuffer[4 * InstanceIndex + 0].xyz, 0.0f),
