Fixing chromatic aberration for a non-center principal point.

diff --git a/Engine/Shaders/Private/PostProcessTonemap.usf b/Engine/Shaders/Private/PostProcessTonemap.usf
--- a/Engine/Shaders/Private/PostProcessTonemap.usf
+++ b/Engine/Shaders/Private/PostProcessTonemap.usf
@@ -69,6 +69,7 @@ Texture2D BloomDirtMaskTexture;
 SamplerState BloomDirtMaskSampler;
 
 float4 LensPrincipalPointOffsetScale;
+float4 LensPrincipalPointOffsetScaleInverse;
 
 half GrainFromUV(float2 GrainUV) 
 {
@@ -76,11 +77,21 @@ half GrainFromUV(float2 GrainUV)
 	return Grain;
 }
 
+// converts from screen [-1,1] space to the lens [-1,1] viewport space
 float2 ConvertScreenViewportSpaceToLensViewportSpace(float2 UV)
 {
 	return LensPrincipalPointOffsetScale.xy + UV * LensPrincipalPointOffsetScale.zw;
 }
 
+float2 ConvertLensViewportSpaceToScreenViewportSpace(float2 UV)
+{
+	// reference version
+	//return (UV - LensPrincipalPointOffsetScale.xy)/LensPrincipalPointOffsetScale.zw;
+
+	// optimized version
+	return LensPrincipalPointOffsetScaleInverse.xy + UV * LensPrincipalPointOffsetScaleInverse.zw;
+}
+
 // LUT for color grading
 #if USE_VOLUME_LUT == 1
 	Texture3D		ColorGradingLUT;
@@ -385,9 +396,15 @@ float4 TonemapCommonPS(
 	        float2 CAScale = ChromaticAberrationParams.rg;
 	        float StartOffset = ChromaticAberrationParams.z;
 
+			float2 LensUV = ConvertScreenViewportSpaceToLensViewportSpace(ScreenPos);
+
 	        float4 CAUV;
-	        CAUV = ScreenPos.xyxy - sign(ScreenPos).xyxy * saturate(abs(ScreenPos) - StartOffset).xyxy * CAScale.rrgg;
-	        CAUV.xy = ScreenPosToUV(CAUV.xy, Color_ExtentInverse);
+	        CAUV = LensUV.xyxy - sign(LensUV).xyxy * saturate(abs(LensUV) - StartOffset).xyxy * CAScale.rrgg;
+
+	        CAUV.xy = ConvertLensViewportSpaceToScreenViewportSpace(CAUV.xy);
+	        CAUV.zw = ConvertLensViewportSpaceToScreenViewportSpace(CAUV.zw);
+
+			CAUV.xy = ScreenPosToUV(CAUV.xy, Color_ExtentInverse);
 	        CAUV.zw = ScreenPosToUV(CAUV.zw, Color_ExtentInverse);
 
 		half4 SceneColor = SampleSceneColor(CAUV.xy + SceneUVJitter.xy);
@@ -504,8 +521,11 @@ float4 TonemapCommonPS(
 			float4 CombinedBloom = Texture2DSample(BloomTexture, BloomSampler, BloomUV);
 			CombinedBloom.rgb *= OneOverPreExposure; // Standard bloom path, convolution path has bloom already in scene color.
 
-			float2 DirtViewportUV = (SvPosition.xy - float2(Output_ViewportMin)) * Output_ViewportSizeInverse;
-			float3 BloomDirtMaskColor = Texture2DSample(BloomDirtMaskTexture, BloomDirtMaskSampler, DirtViewportUV).rgb * BloomDirtMaskTint.rgb;
+			//float2 DirtViewportUV = (SvPosition.xy - float2(Output_ViewportMin)) * Output_ViewportSizeInverse; 
+
+			float2 DirtLensUV = ConvertScreenViewportSpaceToLensViewportSpace(ScreenPos) * float2(1.0f, -1.0f);
+
+			float3 BloomDirtMaskColor = Texture2DSample(BloomDirtMaskTexture, BloomDirtMaskSampler, DirtLensUV * .5f + .5f).rgb * BloomDirtMaskTint.rgb;
 
 #if DEBUG_TONEMAPPER
 			if (all(DebugTile == int2(-1, -1)))
diff --git a/Engine/Source/Runtime/Renderer/Private/PostProcess/PostProcessTonemap.cpp b/Engine/Source/Runtime/Renderer/Private/PostProcess/PostProcessTonemap.cpp
--- a/Engine/Source/Runtime/Renderer/Private/PostProcess/PostProcessTonemap.cpp
+++ b/Engine/Source/Runtime/Renderer/Private/PostProcess/PostProcessTonemap.cpp
@@ -534,6 +534,7 @@ BEGIN_SHADER_PARAMETER_STRUCT(FTonemapParameters, )
 	SHADER_PARAMETER(FVector4, ChromaticAberrationParams)
 	SHADER_PARAMETER(FVector4, TonemapperParams)
 	SHADER_PARAMETER(FVector4, LensPrincipalPointOffsetScale)
+	SHADER_PARAMETER(FVector4, LensPrincipalPointOffsetScaleInverse)
 	SHADER_PARAMETER(float, SwitchVerticalAxis)
 	SHADER_PARAMETER(float, DefaultEyeExposure)
 	SHADER_PARAMETER(float, EditorNITLevel)
@@ -780,6 +781,17 @@ FScreenPassTexture AddTonemapPass(FRDGBuilder& GraphBuilder, const FViewInfo& Vi
 	CommonParameters.bOutputInHDR = ViewFamily.bIsHDR;
 	CommonParameters.LensPrincipalPointOffsetScale = View.LensPrincipalPointOffsetScale;
 
+	// forward transformation from shader:
+	//return LensPrincipalPointOffsetScale.xy + UV * LensPrincipalPointOffsetScale.zw;
+
+	// reverse transformation from shader:
+	//return UV*(1.0f/LensPrincipalPointOffsetScale.zw) - LensPrincipalPointOffsetScale.xy/LensPrincipalPointOffsetScale.zw;
+
+	CommonParameters.LensPrincipalPointOffsetScaleInverse.X = -View.LensPrincipalPointOffsetScale.X / View.LensPrincipalPointOffsetScale.Z;
+	CommonParameters.LensPrincipalPointOffsetScaleInverse.Y = -View.LensPrincipalPointOffsetScale.Y / View.LensPrincipalPointOffsetScale.W;
+	CommonParameters.LensPrincipalPointOffsetScaleInverse.Z = 1.0f / View.LensPrincipalPointOffsetScale.Z;
+	CommonParameters.LensPrincipalPointOffsetScaleInverse.W = 1.0f / View.LensPrincipalPointOffsetScale.W;
+
 	// Generate permutation vector for the desktop tonemapper.
 	TonemapperPermutation::FDesktopDomain DesktopPermutationVector;
 
