Final rename resolves (I hope) from the CL8684840 merge-down

diff --git a/Engine/Shaders/Private/RayTracing/RayTracingPrimaryRays.usf b/Engine/Shaders/Private/RayTracing/RayTracingPrimaryRays.usf
new file mode 100644
--- /dev/null
+++ b/Engine/Shaders/Private/RayTracing/RayTracingPrimaryRays.usf
@@ -0,0 +1,292 @@
+// Copyright 1998-2019 Epic Games, Inc. All Rights Reserved.
+
+#ifndef DIM_DENOISER_OUTPUT
+	#define DIM_DENOISER_OUTPUT 0
+#endif
+
+#include "../Common.ush"
+
+#define SUPPORT_CONTACT_SHADOWS		0
+#define USE_SOURCE_TEXTURE			1
+#define USE_SOURCE_TEXTURE_ARRAY	1
+
+#define LTCMatTexture				RaytracingLightsDataPacked.LTCMatTexture
+#define LTCMatSampler				RaytracingLightsDataPacked.LTCMatSampler
+#define LTCAmpTexture				RaytracingLightsDataPacked.LTCAmpTexture
+#define LTCAmpSampler				RaytracingLightsDataPacked.LTCAmpSampler
+#define PreIntegratedGF				ReflectionStruct.PreIntegratedGF			
+#define PreIntegratedGFSampler		ReflectionStruct.PreIntegratedGFSampler	
+
+#include "../DeferredShadingCommon.ush"
+#include "../DeferredLightingCommon.ush"
+#include "../ReflectionEnvironmentShared.ush"
+#include "../Montecarlo.ush"
+#include "../PathTracing/Utilities/PathTracingRandomSequence.ush" 
+#include "../HeightFogCommon.ush"
+#include "../SobolRandom.ush"
+#include "../SceneTextureParameters.ush"
+#include "RayTracingCommon.ush"
+#include "RayTracingHitGroupCommon.ush"
+#include "RayTracingDirectionalLight.ush"
+#include "RayTracingRectLight.ush"
+#include "RayTracingSphereLight.ush"
+#include "RayTracingSpotLight.ush"
+#include "RayTracingPointLight.ush"
+
+#define	ERayTracingPrimaryRaysFlag_None                               0
+#define	ERayTracingPrimaryRaysFlag_UseGBufferForMaxDistance    (1u << 0)
+#define	ERayTracingPrimaryRaysFlag_ConsiderSurfaceScatter	   (1u << 1)
+#define	ERayTracingPrimaryRaysFlag_AllowSkipSkySample		   (1u << 2)
+
+
+int SamplesPerPixel;
+int MaxRefractionRays;
+int HeightFog;
+int ReflectedShadowsType;
+int ShouldDoDirectLighting;
+int ShouldDoEmissiveAndIndirectLighting;
+int UpscaleFactor;
+int ShouldUsePreExposure;
+uint PrimaryRayFlags;
+
+float TranslucencyMinRayDistance;
+float TranslucencyMaxRayDistance;
+float TranslucencyMaxRoughness;
+int  TranslucencyRefraction;
+float MaxNormalBias;
+
+Texture2D SceneColorTexture;
+
+RaytracingAccelerationStructure TLAS;
+
+RWTexture2D<float4> ColorOutput;
+RWTexture2D<float> RayHitDistanceOutput;
+
+#include "RayTracingLightingCommon.ush"
+
+float FresnelDielectric(float Eta, float IoH, float ToH)
+{
+	float Rs = Square((Eta * IoH - ToH) / (Eta * IoH + ToH));
+	float Rp = Square((Eta * ToH - IoH) / (Eta * ToH + IoH));
+	return (Rs + Rp) / 2;
+}
+
+bool RefractRay(float3 V, float3 N, float Ior1, float Ior2, bool bIsEntering, out float3 L, out float EventThroughput)
+{
+	float Eta = Ior1 / Ior2;
+	if (bIsEntering)
+	{
+		Eta = 1.0f / Ior1;
+	}
+
+	V = -V;
+	const float CosThetaI = dot(V, N);
+	const float CosThetaT2 = 1.f - Eta * Eta * (1.0f - CosThetaI * CosThetaI);
+	if (CosThetaT2 < 0.0f)
+	{
+		// Total internal reflection
+		EventThroughput = 1.0f;
+		L = reflect(-V, N); 
+		return false;
+	}
+
+	const float CosThetaT = -(CosThetaI >= 0.0f ? 1.0f : -1.0f) * sqrt(CosThetaT2);
+	L = (Eta * CosThetaI + CosThetaT) * N - Eta * V;
+
+	float VoH = abs(dot(N, V));
+	float ToH = abs(dot(N, L));
+	EventThroughput = 1 - FresnelDielectric(Eta, VoH, ToH);
+	return true;
+}
+
+float3 GetSkyRadiance(float3 Direction, float Roughness)
+{
+	float SkyAverageBrightness = 1.0f;
+	return GetSkyLightReflection(Direction, Roughness, SkyAverageBrightness);
+}
+
+[shader("raygeneration")]
+void RayTracingPrimaryRaysRGS()
+{
+	uint2 DispatchThreadId = DispatchRaysIndex().xy + View.ViewRectMin;
+	uint2 PixelCoord = GetPixelCoord(DispatchThreadId, UpscaleFactor);
+	uint LinearIndex = PixelCoord.y * View.BufferSizeAndInvSize.x + PixelCoord.x; // TODO: PixelCoord or DispatchThreadId
+	
+	RandomSequence RandSequence;
+	RandomSequence_Initialize(RandSequence, LinearIndex, View.StateFrameIndex);
+
+	float2 InvBufferSize = View.BufferSizeAndInvSize.zw;
+	float2 UV = (float2(PixelCoord) + 0.5) * InvBufferSize;
+	FGBufferData GBufferData = GetGBufferDataFromSceneTextures(UV);
+	float Depth = GBufferData.Depth;
+	float3 WorldPosition = ReconstructWorldPositionFromDepth(UV, Depth);
+	float WorldSpaceDistance = length(WorldPosition - View.WorldViewOrigin);
+
+	// Trace rays from camera origin to (Gbuffer - epsilon) to only intersect translucent objects
+	RayDesc Ray = CreatePrimaryRay(UV);
+	FRayCone RayCone = (FRayCone)0;
+	RayCone.SpreadAngle = View.EyeToPixelSpreadAngle;
+
+	if((ERayTracingPrimaryRaysFlag_UseGBufferForMaxDistance & PrimaryRayFlags) != 0) 
+	{
+		Ray.TMax = WorldSpaceDistance - 0.1;
+	}
+
+	bool bAllowSkySampling;
+	if((ERayTracingPrimaryRaysFlag_AllowSkipSkySample & PrimaryRayFlags) != 0) 
+	{
+		// Sky is only sampled when infinite reflection rays are used.
+		bAllowSkySampling =  TranslucencyMaxRayDistance < 0;
+	} 
+	else 
+	{
+		bAllowSkySampling = true;
+	}
+	
+	bool bHasScattered = (ERayTracingPrimaryRaysFlag_ConsiderSurfaceScatter & PrimaryRayFlags) != 0;
+
+	const bool bRefractionBackFaceCulling = false;
+	const bool bReflectionBackFaceCulling = true;
+
+	// Integrated data by path tracing
+	float3 PathRadiance = 0.0;
+	float PathThroughput = 1.0;	// A float for now because UE does not support colored translucency as of today.
+	float LastRoughness = 0.0;
+	float LastIor = 1.0f;
+	float HitDistance = 0.0f;
+
+	for (uint RefractionRayIndex = 0; RefractionRayIndex < MaxRefractionRays; ++RefractionRayIndex)
+	{
+		float3 PathVertexRadiance = float3(0, 0, 0);
+		//#dxr-todo: RayFlags should be 0 for refraction rays
+		const uint RayFlags = bRefractionBackFaceCulling ? RAY_FLAG_CULL_BACK_FACING_TRIANGLES : 0;
+		FMaterialClosestHitPayload Payload = TraceRayAndAccumulateResults(Ray, TLAS, RandSequence, RayFlags, MaxNormalBias, RAY_TRACING_MASK_ALL, ReflectedShadowsType, ShouldDoDirectLighting, ShouldDoEmissiveAndIndirectLighting, RayCone, PathVertexRadiance);
+		LastRoughness = Payload.Roughness;
+
+		//
+		// Handle no hit condition
+		//
+		if (Payload.IsMiss())
+		{
+			if (bHasScattered && bAllowSkySampling)
+			{
+				// We only sample the sky if the ray has scattered (i.e. been refracted or reflected). Otherwise we are going ot use the regular scene color.
+				PathRadiance += PathThroughput * GetSkyRadiance(Ray.Direction, LastRoughness);
+			}
+			break;
+		}
+		float3 HitPoint = Ray.Origin + Ray.Direction * Payload.HitT;
+		float NextMaxRayDistance = Ray.TMax - Payload.HitT;
+
+		//
+		// Handle surface lighting
+		//
+
+		float vertexRadianceWeight = Payload.Opacity;	// Opacity as coverage. This works for RAY_TRACING_BLEND_MODE_OPAQUE and RAY_TRACING_BLEND_MODE_TRANSLUCENT.
+		// It is also needed for RAY_TRACING_BLEND_MODE_ADDITIVE and  RAY_TRACING_BLEND_MODE_ALPHA_COMPOSITE: radiance continbution is alway weighted by coverage.
+		// #dxr_todo: I have not been able to setup a material using RAY_TRACING_BLEND_MODE_MODULATE.
+
+		PathRadiance += PathThroughput * PathVertexRadiance * vertexRadianceWeight;
+
+
+		//
+		// Handle reflection tracing with a ray per vertex of the refraction path
+		//
+
+		// Shorten the rays on rougher surfaces between user-provided min and max ray lengths.
+		// When a shortened ray misses the geometry, we fall back to local reflection capture sampling (similar to SSR).
+		const float LocalMaxRayDistance = bAllowSkySampling ? 1e27f : lerp(TranslucencyMaxRayDistance, TranslucencyMinRayDistance, Payload.Roughness);
+		if (Payload.Roughness < TranslucencyMaxRoughness)
+		{
+			// Trace reflection ray 
+			uint DummyVariable;
+			float2 RandSample = RandomSequence_GenerateSample2D(RandSequence, DummyVariable);
+
+			RayDesc ReflectionRay;
+			ReflectionRay.TMin = 0.01;
+			ReflectionRay.TMax = LocalMaxRayDistance;
+			ReflectionRay.Origin = HitPoint;
+			ReflectionRay.Direction = GenerateReflectedRayDirection(Ray.Direction, Payload.WorldNormal, Payload.Roughness, RandSample);
+			ApplyPositionBias(ReflectionRay, Payload.WorldNormal, MaxNormalBias);
+
+			float3 ReflectionRadiance = float3(0, 0, 0);
+			const uint RayFlags = bReflectionBackFaceCulling ? RAY_FLAG_CULL_BACK_FACING_TRIANGLES : 0;
+			FMaterialClosestHitPayload ReflectionPayload = TraceRayAndAccumulateResults(ReflectionRay, TLAS, RandSequence, RayFlags, MaxNormalBias, RAY_TRACING_MASK_ALL, ReflectedShadowsType, ShouldDoDirectLighting, ShouldDoEmissiveAndIndirectLighting, RayCone, ReflectionRadiance);
+
+			// If we have not hit anything, sample the distance sky radiance.
+			if (ReflectionPayload.IsMiss())
+			{
+				ReflectionRadiance = GetSkyRadiance(ReflectionRay.Direction, LastRoughness);
+			}
+
+			// #dxr_todo: reflection IOR and clear coat also? This only handles default material.
+			float NoV = saturate(dot(-Ray.Direction, Payload.WorldNormal));
+			const float3 ReflectionThroughput = EnvBRDF(Payload.SpecularColor, Payload.Roughness, NoV);
+
+			PathRadiance += PathThroughput * ReflectionThroughput * ReflectionRadiance * vertexRadianceWeight;
+		}
+
+
+		//
+		// Handle refraction through the surface.
+		//
+
+		// Update the refraction path transmittance and check stop condition
+		float PathVertexTransmittance = Payload.BlendingMode == RAY_TRACING_BLEND_MODE_ADDITIVE ? 1.0 : 1.0 - Payload.Opacity;
+		PathThroughput *= PathVertexTransmittance;
+		if (PathThroughput <= 0.0)
+		{
+			break;
+		}
+
+		// Set refraction ray for next iteration
+		float3 RefractedDirection = Ray.Direction;
+		if (TranslucencyRefraction)
+		{
+			float Ior1 = DielectricF0ToIor(DielectricSpecularToF0(Payload.Specular)); // Not using Payload.Ior but parameterisation from specular.
+			float Ior2 = LastIor;
+			bHasScattered |= Ior1 != Ior2;
+			float ReflectionRefractionEventThroughput;
+			bool bIsEntering = Payload.IsFrontFace();
+
+			if (!RefractRay(Ray.Direction, Payload.WorldNormal, Ior1, Ior2, bIsEntering, RefractedDirection, ReflectionRefractionEventThroughput))
+			{
+				//#dxr_todo: handle total internal reflection ray event. We could just continue form this path. And not use the new IOR.
+				// We should combine that step with the one above reflection evaluation. For now we stop tracing.
+				break;
+			}
+			LastIor = Ior1;
+			PathThroughput *= ReflectionRefractionEventThroughput;
+
+			// ray has bent, so it may need to go arbitrarily far
+			NextMaxRayDistance = LocalMaxRayDistance;
+		}
+
+
+		//
+		// Setup refracted ray to be traced
+		//
+
+		Ray.Origin = HitPoint;
+		Ray.TMin = 0.01;
+		Ray.TMax = NextMaxRayDistance;
+		Ray.Direction = RefractedDirection;
+		float SurfaceCurvature = 0.0f; /* #todo_dxr assume no curvature */
+		RayCone = PropagateRayCone(RayCone, SurfaceCurvature, Depth);
+	}
+
+	if (!bHasScattered)
+	{
+		// Use the scene radiance for ray that has not been scattered/refracted (no surface or IORin=IORout). Still apply the throughtput in case we have traversed surfaces with opacity>0.
+		PathRadiance += PathThroughput * SceneColorTexture.SampleLevel(GlobalPointClampedSampler, UV, 0).xyz / View.PreExposure;
+	}
+
+	if(ShouldUsePreExposure) 
+	{
+		PathRadiance.rgb *= View.PreExposure;
+	}
+
+	PathRadiance = ClampToHalfFloatRange(PathRadiance);
+	ColorOutput[DispatchThreadId] = float4(PathRadiance, 0.0f);
+	RayHitDistanceOutput[DispatchThreadId] = HitDistance;
+}
\ No newline at end of file
diff --git a/Engine/Shaders/Private/RayTracing/RayTracingTranslucency.usf b/Engine/Shaders/Private/RayTracing/RayTracingTranslucency.usf
deleted file mode 100644
