Disable LookupDeviceZ() and CalcSceneDepth() for non-pixel shaders on mobile

diff --git a/Engine/Shaders/Private/SceneTexturesCommon.ush b/Engine/Shaders/Private/SceneTexturesCommon.ush
--- a/Engine/Shaders/Private/SceneTexturesCommon.ush
+++ b/Engine/Shaders/Private/SceneTexturesCommon.ush
@@ -98,7 +98,7 @@ float4 CalcFullSceneColor(float2 ScreenUV)
 /** Returns clip space W, which is world space distance along the View Z axis. Note if you need DeviceZ LookupDeviceZ() is the faster option */
 float CalcSceneDepth(float2 ScreenUV)
 {
-#if SCENE_TEXTURES_DISABLED
+#if SCENE_TEXTURES_DISABLED || !PIXELSHADER
 	return SCENE_TEXTURES_DISABLED_SCENE_DEPTH_VALUE;
 #else
 	#if MOBILE_FORCE_DEPTH_TEXTURE_READS
@@ -145,7 +145,7 @@ float CalcSceneDepth(float2 ScreenUV)
 /** Returns DeviceZ which is the z value stored in the depth buffer. */
 float LookupDeviceZ( float2 ScreenUV )
 {
-#if	SCENE_TEXTURES_DISABLED
+#if	SCENE_TEXTURES_DISABLED || !PIXELSHADER
 	return SCENE_TEXTURES_DISABLED_SCENE_DEPTH_VALUE;
 #elif MOBILE_FORCE_DEPTH_TEXTURE_READS
 	// native Depth buffer lookup
