Fix for bloom application in splitscreen (Integrated from CL #11299051)

#rnx

diff --git a/Engine/Shaders/Private/PostProcessTonemap.usf b/Engine/Shaders/Private/PostProcessTonemap.usf
--- a/Engine/Shaders/Private/PostProcessTonemap.usf
+++ b/Engine/Shaders/Private/PostProcessTonemap.usf
@@ -521,7 +521,8 @@ float4 TonemapCommonPS(
 		half3 LinearColor = SceneColor.rgb * ColorScale0.rgb;
 
 		#if USE_BLOOM
-			float2 BloomUV = clamp(UV, Bloom_UVViewportBilinearMin, Bloom_UVViewportBilinearMax);
+			float2 BloomUV = Bloom_UVViewportMin + UV * Bloom_UVViewportSize;			
+			BloomUV = clamp(BloomUV, Bloom_UVViewportBilinearMin, Bloom_UVViewportBilinearMax);
 			#if ES2_PROFILE || ES3_1_PROFILE
 				BloomUV = FullViewUV.xy;
 			#endif
