Add SpecularScale support for all light types for mobile renderer.

diff --git a/Engine/Shaders/Private/MobileBasePassPixelShader.usf b/Engine/Shaders/Private/MobileBasePassPixelShader.usf
--- a/Engine/Shaders/Private/MobileBasePassPixelShader.usf
+++ b/Engine/Shaders/Private/MobileBasePassPixelShader.usf
@@ -32,7 +32,7 @@ int NumDynamicPointLights;
 float4 LightPositionAndInvRadius[MAX_DYNAMIC_POINT_LIGHTS];
 float4 LightColorAndFalloffExponent[MAX_DYNAMIC_POINT_LIGHTS];
 float4 SpotLightAngles[MAX_DYNAMIC_POINT_LIGHTS];
-float4 SpotLightDirection[MAX_DYNAMIC_POINT_LIGHTS]; 
+float4 SpotLightDirectionAndSpecularScale[MAX_DYNAMIC_POINT_LIGHTS];
 #endif
 
 #ifndef PROJECT_MOBILE_ENABLE_MOVABLE_SPOTLIGHTS
@@ -530,7 +530,7 @@ void Main(
 				#endif
 
 				#if FADE_CSM
-					float Fade = saturate(MaterialParameters.ScreenPosition.w * MobileDirectionalLight.DirectionalLightDistanceFadeMAD.x + MobileDirectionalLight.DirectionalLightDistanceFadeMAD.y);
+					float Fade = saturate(MaterialParameters.ScreenPosition.w * MobileDirectionalLight.DirectionalLightDistanceFadeMADAndSpecularScale.x + MobileDirectionalLight.DirectionalLightDistanceFadeMADAndSpecularScale.y);
 					// lerp out shadow based on fade params.
 					ShadowMap = lerp(ShadowMap, 1.0, Fade * Fade);
 				#endif
@@ -552,7 +552,8 @@ void Main(
 		#if FULLY_ROUGH
 			Color += (Shadow * NoL) * MobileDirectionalLight.DirectionalLightColor.rgb * DiffuseColor;
 		#else
-			Color += (Shadow * NoL) * MobileDirectionalLight.DirectionalLightColor.rgb * (DiffuseColor + SpecularColor * CalcSpecular(Roughness, RoL, NoH, H, MaterialParameters.WorldNormal));
+			// MobileDirectionalLight.DirectionalLightDistanceFadeMADAndSpecularScale.z saves SpecularScale for direction light.
+			Color += (Shadow * NoL) * MobileDirectionalLight.DirectionalLightColor.rgb * (DiffuseColor + SpecularColor * MobileDirectionalLight.DirectionalLightDistanceFadeMADAndSpecularScale.z * CalcSpecular(Roughness, RoL, NoH, H, MaterialParameters.WorldNormal));
 			// Environment map has been prenormalized, scale by lightmap luminance
 			half3 SpecularIBL = GetImageBasedReflectionLighting(MaterialParameters, Roughness) * IndirectIrradiance;
 
@@ -618,12 +619,12 @@ void Main(
 				#if PROJECT_MOBILE_ENABLE_MOVABLE_SPOTLIGHTS
 				if(SpotLightAngles[i].w > 0.0)
 				{
-					Attenuation *= SpotAttenuation(L, -SpotLightDirection[i].xyz, SpotLightAngles[i].xy);
+					Attenuation *= SpotAttenuation(L, -SpotLightDirectionAndSpecularScale[i].xyz, SpotLightAngles[i].xy);
 				}
 				#endif
 
 				#if !FULLY_ROUGH
-					Color += min(65000.0, (Attenuation * PointNoL) * LightColorAndFalloffExponent[i].rgb * (1.0/PI) * (DiffuseColor + SpecularColor * CalcSpecular(Roughness, PointRoL, PointNoH, H, MaterialParameters.WorldNormal)));
+					Color += min(65000.0, (Attenuation * PointNoL) * LightColorAndFalloffExponent[i].rgb * (1.0/PI) * (DiffuseColor + SpecularColor * SpotLightDirectionAndSpecularScale[i].w * CalcSpecular(Roughness, PointRoL, PointNoH, H, MaterialParameters.WorldNormal)));
 				#else
 					Color += (Attenuation * PointNoL) * LightColorAndFalloffExponent[i].rgb * DiffuseColor;
 				#endif
diff --git a/Engine/Source/Runtime/Engine/Private/SceneManagement.cpp b/Engine/Source/Runtime/Engine/Private/SceneManagement.cpp
--- a/Engine/Source/Runtime/Engine/Private/SceneManagement.cpp
+++ b/Engine/Source/Runtime/Engine/Private/SceneManagement.cpp
@@ -661,7 +661,7 @@ FMobileDirectionalLightShaderParameters::FMobileDirectionalLightShaderParameters
 	DirectionalLightShadowTexture = GWhiteTexture->TextureRHI;
 	DirectionalLightShadowSampler = TStaticSamplerState<SF_Point, AM_Clamp, AM_Clamp, AM_Clamp>::GetRHI();
 	DirectionalLightShadowSize = FVector4(EForceInit::ForceInitToZero);
-	DirectionalLightDistanceFadeMAD = FVector4(EForceInit::ForceInitToZero);
+	DirectionalLightDistanceFadeMADAndSpecularScale = FVector4(EForceInit::ForceInitToZero);
 	for (int32 i = 0; i < MAX_MOBILE_SHADOWCASCADES; ++i)
 	{
 		DirectionalLightScreenToShadow[i].SetIdentity();
diff --git a/Engine/Source/Runtime/Engine/Public/SceneView.h b/Engine/Source/Runtime/Engine/Public/SceneView.h
--- a/Engine/Source/Runtime/Engine/Public/SceneView.h
+++ b/Engine/Source/Runtime/Engine/Public/SceneView.h
@@ -539,7 +539,7 @@ BEGIN_GLOBAL_SHADER_PARAMETER_STRUCT_WITH_CONSTRUCTOR(FMobileDirectionalLightSha
 	SHADER_PARAMETER_EX(FLinearColor, DirectionalLightColor, EShaderPrecisionModifier::Half)
 	SHADER_PARAMETER_EX(FVector4, DirectionalLightDirectionAndShadowTransition, EShaderPrecisionModifier::Half)
 	SHADER_PARAMETER_EX(FVector4, DirectionalLightShadowSize, EShaderPrecisionModifier::Half)
-	SHADER_PARAMETER_EX(FVector4, DirectionalLightDistanceFadeMAD, EShaderPrecisionModifier::Half) // .zw is not used atm
+	SHADER_PARAMETER_EX(FVector4, DirectionalLightDistanceFadeMADAndSpecularScale, EShaderPrecisionModifier::Half) // .z is used for SpecularScale, .w is not used atm
 	SHADER_PARAMETER_EX(FVector4, DirectionalLightShadowDistances, EShaderPrecisionModifier::Half)
 	SHADER_PARAMETER_ARRAY(FMatrix, DirectionalLightScreenToShadow, [MAX_MOBILE_SHADOWCASCADES])
 	SHADER_PARAMETER_TEXTURE(Texture2D, DirectionalLightShadowTexture)
diff --git a/Engine/Source/Runtime/Renderer/Private/MobileBasePass.cpp b/Engine/Source/Runtime/Renderer/Private/MobileBasePass.cpp
--- a/Engine/Source/Runtime/Renderer/Private/MobileBasePass.cpp
+++ b/Engine/Source/Runtime/Renderer/Private/MobileBasePass.cpp
@@ -503,14 +503,14 @@ void TMobileBasePassPSPolicyParamType<FUniformLightMapPolicy>::GetShaderBindings
 			ShaderBindings.Add(ReflectionParameter, ReflectionUB);
 		}
 		
-		if (LightPositionAndInvRadiusParameter.IsBound() || SpotLightDirectionParameter.IsBound())
+		if (LightPositionAndInvRadiusParameter.IsBound() || SpotLightDirectionAndSpecularScaleParameter.IsBound())
 		{
 			// Set dynamic point lights
 			FMobileBasePassMovableLightInfo LightInfo(PrimitiveSceneProxy);
 			ShaderBindings.Add(NumDynamicPointLightsParameter, LightInfo.NumMovablePointLights);
 			ShaderBindings.Add(LightPositionAndInvRadiusParameter, LightInfo.LightPositionAndInvRadius);
 			ShaderBindings.Add(LightColorAndFalloffExponentParameter, LightInfo.LightColorAndFalloffExponent);
-			ShaderBindings.Add(SpotLightDirectionParameter, LightInfo.SpotLightDirection);
+			ShaderBindings.Add(SpotLightDirectionAndSpecularScaleParameter, LightInfo.SpotLightDirectionAndSpecularScale);
 			ShaderBindings.Add(SpotLightAnglesParameter, LightInfo.SpotLightAngles);
 		}
 	}
diff --git a/Engine/Source/Runtime/Renderer/Private/MobileBasePassRendering.cpp b/Engine/Source/Runtime/Renderer/Private/MobileBasePassRendering.cpp
--- a/Engine/Source/Runtime/Renderer/Private/MobileBasePassRendering.cpp
+++ b/Engine/Source/Runtime/Renderer/Private/MobileBasePassRendering.cpp
@@ -129,7 +129,8 @@ FMobileBasePassMovableLightInfo::FMobileBasePassMovableLightInfo(const FPrimitiv
 
 				LightPositionAndInvRadius[NumMovablePointLights] = FVector4(LightParameters.Position, LightParameters.InvRadius);
 				LightColorAndFalloffExponent[NumMovablePointLights] = FVector4(LightParameters.Color, LightParameters.FalloffExponent);
-				SpotLightDirection[NumMovablePointLights] = LightParameters.Direction;
+				SpotLightDirectionAndSpecularScale[NumMovablePointLights] = LightParameters.Direction;
+				SpotLightDirectionAndSpecularScale[NumMovablePointLights].Component(3) = LightProxy->GetSpecularScale();
 				SpotLightAngles[NumMovablePointLights].Set(LightParameters.SpotAngles.X, LightParameters.SpotAngles.Y, 0.f, LightType == LightType_Spot ? 1.0f : 0.0f);
 
 				if (LightType == LightType_Rect)
@@ -196,8 +197,9 @@ void SetupMobileDirectionalLightUniformParameters(
 		Params.DirectionalLightDirectionAndShadowTransition = FVector4(-Light->Proxy->GetDirection(), 0.f);
 
 		const FVector2D FadeParams = Light->Proxy->GetDirectionalLightDistanceFadeParameters(FeatureLevel, Light->IsPrecomputedLightingValid(), SceneView.MaxShadowCascades);
-		Params.DirectionalLightDistanceFadeMAD.X = FadeParams.Y;
-		Params.DirectionalLightDistanceFadeMAD.Y = -FadeParams.X * FadeParams.Y;
+		Params.DirectionalLightDistanceFadeMADAndSpecularScale.X = FadeParams.Y;
+		Params.DirectionalLightDistanceFadeMADAndSpecularScale.Y = -FadeParams.X * FadeParams.Y;
+		Params.DirectionalLightDistanceFadeMADAndSpecularScale.Z = Light->Proxy->GetSpecularScale();
 
 		if (bDynamicShadows && VisibleLightInfos.IsValidIndex(Light->Id) && VisibleLightInfos[Light->Id].AllProjectedShadows.Num() > 0)
 		{
diff --git a/Engine/Source/Runtime/Renderer/Private/MobileBasePassRendering.h b/Engine/Source/Runtime/Renderer/Private/MobileBasePassRendering.h
--- a/Engine/Source/Runtime/Renderer/Private/MobileBasePassRendering.h
+++ b/Engine/Source/Runtime/Renderer/Private/MobileBasePassRendering.h
@@ -76,7 +76,7 @@ public:
 	int32 NumMovablePointLights;
 	FVector4 LightPositionAndInvRadius[MAX_BASEPASS_DYNAMIC_POINT_LIGHTS];
 	FVector4 LightColorAndFalloffExponent[MAX_BASEPASS_DYNAMIC_POINT_LIGHTS];
-	FVector4 SpotLightDirection[MAX_BASEPASS_DYNAMIC_POINT_LIGHTS];
+	FVector4 SpotLightDirectionAndSpecularScale[MAX_BASEPASS_DYNAMIC_POINT_LIGHTS];
 	FVector4 SpotLightAngles[MAX_BASEPASS_DYNAMIC_POINT_LIGHTS];
 };
 
@@ -252,7 +252,7 @@ public:
 		LightColorAndFalloffExponentParameter.Bind(Initializer.ParameterMap, TEXT("LightColorAndFalloffExponent"));
 		NumDynamicPointLightsParameter.Bind(Initializer.ParameterMap, TEXT("NumDynamicPointLights"));
 		SpotLightAnglesParameter.Bind(Initializer.ParameterMap, TEXT("SpotLightAngles"));
-		SpotLightDirectionParameter.Bind(Initializer.ParameterMap, TEXT("SpotLightDirection"));
+		SpotLightDirectionAndSpecularScaleParameter.Bind(Initializer.ParameterMap, TEXT("SpotLightDirectionAndSpecularScale"));
 						
 		CSMDebugHintParams.Bind(Initializer.ParameterMap, TEXT("CSMDebugHint"));
 	}
@@ -280,7 +280,7 @@ public:
 		Ar << LightColorAndFalloffExponentParameter;
 		Ar << NumDynamicPointLightsParameter;
 		Ar << SpotLightAnglesParameter;
-		Ar << SpotLightDirectionParameter;
+		Ar << SpotLightDirectionAndSpecularScaleParameter;
 			
 		Ar << CSMDebugHintParams;
 
@@ -300,7 +300,7 @@ private:
 	FShaderParameter LightPositionAndInvRadiusParameter;
 	FShaderParameter LightColorAndFalloffExponentParameter;
 	FShaderParameter NumDynamicPointLightsParameter;
-	FShaderParameter SpotLightDirectionParameter;
+	FShaderParameter SpotLightDirectionAndSpecularScaleParameter;
 	FShaderParameter SpotLightAnglesParameter;
 
 	FShaderParameter CSMDebugHintParams;
