DRPS - Fix missing position output on Metal shaders
#rnx

diff --git a/Engine/Shaders/Private/ShadowDepthVertexShader.usf b/Engine/Shaders/Private/ShadowDepthVertexShader.usf
--- a/Engine/Shaders/Private/ShadowDepthVertexShader.usf
+++ b/Engine/Shaders/Private/ShadowDepthVertexShader.usf
@@ -373,10 +373,17 @@ void MainForGS(
 #endif
 #else
 	out FShadowDepthVSToPS OutParameters
+#endif
+#if COMPILER_METAL
+#define SKIP_UNUSED_POSITION
+	// Metal needs a position output
+	, out float4 UnusedPosition : SV_POSITION
 #endif
 	)
 {
+#ifndef SKIP_UNUSED_POSITION
 	float4 UnusedPosition;
+#endif
 
 	Main(Input,
 		OutParameters,
@@ -462,10 +469,19 @@ void PositionOnlyMainForGS(
 	out FShadowDepthVSToGS OutParameters
 #else
 	out FShadowDepthVSToPS OutParameters
+#endif
+#if COMPILER_METAL
+	#ifndef SKIP_UNUSED_POSITION
+		#define SKIP_UNUSED_POSITION
+	#endif
+	// Metal needs a position output
+	, out float4 UnusedPosition : SV_POSITION
 #endif
 	)
 {
+#ifndef SKIP_UNUSED_POSITION
 	float4 UnusedPosition;
+#endif
 
 	PositionOnlyMain(Input,
 #if USING_ONEPASS_POINT_VERTEX_LAYER
