DRPS - Fix compile error when vr.InstancedStereo=1
#rnx

diff --git a/Engine/Shaders/Private/SingleLayerWaterComposite.usf b/Engine/Shaders/Private/SingleLayerWaterComposite.usf
--- a/Engine/Shaders/Private/SingleLayerWaterComposite.usf
+++ b/Engine/Shaders/Private/SingleLayerWaterComposite.usf
@@ -16,6 +16,8 @@
 
 #define ENABLE_SKY_LIGHT 1
 
+#include "LightGridCommon.ush"
+
 #define REFLECTION_COMPOSITE_USE_BLENDED_REFLECTION_CAPTURES (FEATURE_LEVEL >= FEATURE_LEVEL_SM5)
 #define REFLECTION_COMPOSITE_SUPPORT_SKYLIGHT_BLEND 0
 #include "ReflectionEnvironmentComposite.ush"
