Fix roughness clap causing IBL mip0 never be used.

diff --git a/Engine/Shaders/Private/MobileBasePassPixelShader.usf b/Engine/Shaders/Private/MobileBasePassPixelShader.usf
--- a/Engine/Shaders/Private/MobileBasePassPixelShader.usf
+++ b/Engine/Shaders/Private/MobileBasePassPixelShader.usf
@@ -633,7 +633,8 @@ void Main(
 			// MobileDirectionalLight.DirectionalLightDistanceFadeMADAndSpecularScale.z saves SpecularScale for direction light.
 			Color += (Shadow * NoL) * MobileDirectionalLight.DirectionalLightColor.rgb * (DiffuseColor + SpecularColor * MobileDirectionalLight.DirectionalLightDistanceFadeMADAndSpecularScale.z * CalcSpecular(Roughness, RoL, NoH, H, MaterialParameters.WorldNormal));
 			// Environment map has been prenormalized, scale by lightmap luminance
-			half3 SpecularIBL = GetImageBasedReflectionLighting(MaterialParameters, Roughness, IndirectIrradiance);
+			half RoughnessRaw = clamp(saturate(GetMaterialRoughnessRaw(PixelMaterialInputs)), 1e-6, 1.0);
+			half3 SpecularIBL = GetImageBasedReflectionLighting(MaterialParameters, RoughnessRaw, IndirectIrradiance);
 
 			#if MATERIAL_PLANAR_FORWARD_REFLECTIONS
 				BRANCH
