Add ClearCoat shading model for mobile renderer.
Refactor the MobileBasepassPixelShader to make further shading model extension to be more easy and maintainable.
The instrunctions count is identical to pre-refactor.
Fix a small bug that local lights diffuse not correctly cosine-weighted.

diff --git a/Engine/Shaders/Private/MaterialTemplate.ush b/Engine/Shaders/Private/MaterialTemplate.ush
--- a/Engine/Shaders/Private/MaterialTemplate.ush
+++ b/Engine/Shaders/Private/MaterialTemplate.ush
@@ -1899,30 +1899,9 @@ half GetMaterialRoughness(FPixelMaterialInputs PixelMaterialInputs)
 #if MATERIAL_FULLY_ROUGH
 	return 1;
 #endif
-#if (ES2_PROFILE || ES3_1_PROFILE)
-	// The smallest normalized value that can be represented in IEEE 754 (FP16) is 2^-14 = 6.1e-5.
-	// The code will make the following computation involving roughness: 1.0 / Roughness^4.
-	// Therefore to prevent division by zero on devices that do not support denormals, Roughness^4
-	// must be >= 6.1e-5. We will clamp to 0.09 because 0.09^4 = 6.5e-5.
-	//
-	// Note that we also clamp to 1.0 to match the deferred renderer on PC where the roughness is 
-	// stored in an 8-bit value and thus automatically clamped at 1.0.
-
-	half Roughness = GetMaterialRoughnessRaw(PixelMaterialInputs);
-
-	#if MOBILE_EMULATION
-	{
-		// this is only needed for mobile preview on PC
-		Roughness = Roughness * ResolvedView.RoughnessOverrideParameter.y + ResolvedView.RoughnessOverrideParameter.x;
-	}
-	#endif
-
-	// Increase value from 0.09 to 0.12 to fix missing specular lobe problem on device
-	return clamp( Roughness, 0.12, 1.0 );
-#else
 	half Roughness = saturate(GetMaterialRoughnessRaw(PixelMaterialInputs));
 	
-	#if (USE_EDITOR_SHADERS && !ESDEFERRED_PROFILE)
+	#if (USE_EDITOR_SHADERS && !ESDEFERRED_PROFILE) || MOBILE_EMULATION
 	{
 		// this feature is only needed for development/editor - we can compile it out for a shipping build (see r.CompileShadersForDevelopment cvar help)
 		Roughness = Roughness * ResolvedView.RoughnessOverrideParameter.y + ResolvedView.RoughnessOverrideParameter.x;
@@ -1930,7 +1909,6 @@ half GetMaterialRoughness(FPixelMaterialInputs PixelMaterialInputs)
 	#endif
 	
 	return Roughness;
-#endif
 }
 
 half GetMaterialTranslucencyDirectionalLightingIntensity()
diff --git a/Engine/Shaders/Private/MobileBasePassPixelShader.usf b/Engine/Shaders/Private/MobileBasePassPixelShader.usf
--- a/Engine/Shaders/Private/MobileBasePassPixelShader.usf
+++ b/Engine/Shaders/Private/MobileBasePassPixelShader.usf
@@ -27,12 +27,10 @@
 #include "/Engine/Generated/VertexFactory.ush"
 #include "ReflectionEnvironmentShared.ush"
 #include "LightmapCommon.ush"  
-#include "BRDF.ush"
-#include "SHCommon.ush"
+#include "MobileShadingModels.ush"
 #include "ShadowFilteringCommon.ush"
 #include "DynamicLightingCommon.ush"
 #include "PlanarReflectionShared.ush"
-#include "MobileGGX.ush"
 
 #if MATERIAL_SHADINGMODEL_SINGLELAYERWATER
 #define SIMPLE_SINGLE_LAYER_WATER 1
@@ -119,15 +117,6 @@ half PhongApprox( half Roughness, half RoL )
 }
 
 
-half CalcSpecular(half Roughness, float RoL, float NoH, float3 H, float3 N)
-{
-#if USE_LEGACY_SPECULAR
-	return PhongApprox(Roughness, RoL);
-#else
-	return (Roughness*0.25 + 0.25) * GGX_Mobile(Roughness, NoH, H, N);
-#endif
-}
-
 #if !FULLY_ROUGH
 #if HQ_REFLECTIONS
 
@@ -443,307 +432,283 @@ void Main(
 	//Clip if the blend mode requires it.
 	GetMaterialCoverageAndClipping(MaterialParameters, PixelMaterialInputs);
 
-	half Opacity = GetMaterialOpacity(PixelMaterialInputs);
-	const float BaseMaterialCoverageOverWater = Opacity;
-	const float WaterVisibility = 1.0 - BaseMaterialCoverageOverWater;
-
 	// Store the results in local variables and reuse instead of calling the functions multiple times.
-	// Store the results in local variables and reuse instead of calling the functions multiple times.
-	half3 BaseColor = GetMaterialBaseColor(PixelMaterialInputs);
-	half  Metallic = GetMaterialMetallic(PixelMaterialInputs);
-	half  Specular = GetMaterialSpecular(PixelMaterialInputs);
-
-#if NONMETAL
-	half3 DiffuseColor = BaseColor;
-	half SpecularColor = 0.04;
-#else
-	half DielectricSpecular = 0.08 * Specular;
-	half3 DiffuseColor = BaseColor - BaseColor * Metallic;	// 1 mad
-	half3 SpecularColor = (DielectricSpecular - DielectricSpecular * Metallic) + BaseColor * Metallic;	// 2 mad
-#endif
-
-	half Roughness = GetMaterialRoughness(PixelMaterialInputs);
-
-#if FULLY_ROUGH
-	EnvBRDFApproxFullyRough(DiffuseColor, SpecularColor);
-#else
-	half NoV = max( dot( MaterialParameters.WorldNormal, MaterialParameters.CameraVector ), 0 );
-
-#if NONMETAL
-	// If nothing is hooked up to Metalic and Specular,
-	// then defaults are the same as a non-metal,
-	// so this define is safe.
-	SpecularColor = EnvBRDFApproxNonmetal( Roughness, NoV );
-#else
-#if FORWARDSHADING_USE_HQ_ENV_BRDF
-	SpecularColor = EnvBRDF(SpecularColor, Roughness, NoV);
-#else
-	SpecularColor = EnvBRDFApprox( SpecularColor, Roughness, NoV );
-#endif
-#endif
-#endif
+	FMobileShadingModelContext ShadingModelContext;
+	ShadingModelContext.Opacity = GetMaterialOpacity(PixelMaterialInputs);
+	ShadingModelContext.BaseColor = GetMaterialBaseColor(PixelMaterialInputs);
+	ShadingModelContext.Metallic = GetMaterialMetallic(PixelMaterialInputs);
+	ShadingModelContext.Specular = GetMaterialSpecular(PixelMaterialInputs);
+	ShadingModelContext.Roughness = GetMaterialRoughness(PixelMaterialInputs);
+	// The smallest normalized value that can be represented in IEEE 754 (FP16) is 2^-14 = 6.1e-5.
+	// The code will make the following computation involving roughness: 1.0 / Roughness^4.
+	// Therefore to prevent division by zero on devices that do not support denormals, Roughness^4
+	// must be >= 6.1e-5. We will clamp to 0.09 because 0.09^4 = 6.5e-5.
+	// Increase value from 0.09 to 0.12 to fix missing specular lobe problem on device
+	//
+	// Note that we also clamp to 1.0 to match the deferred renderer on PC where the roughness is 
+	// stored in an 8-bit value and thus automatically clamped at 1.0.
+	// To help with IBL cube sample, we postpone the clamp only when calculate GGX_Mobile.
+	ShadingModelContext.RoughnessWithClamp = max(0.12, ShadingModelContext.Roughness);
+
+	InitShadingModelContext(ShadingModelContext, MaterialParameters);
 
-#if MOBILE_EMULATION
-	{
-		// this feature is only needed for development/editor - we can compile it out for a shipping build (see r.CompileShadersForDevelopment cvar help)
-		DiffuseColor = DiffuseColor * ResolvedView.DiffuseOverrideParameter.w + ResolvedView.DiffuseOverrideParameter.xyz;
-#if NONMETAL
-    SpecularColor = SpecularColor * ResolvedView.SpecularOverrideParameter.w + ResolvedView.SpecularOverrideParameter.x;
-#else
-    SpecularColor = SpecularColor * ResolvedView.SpecularOverrideParameter.w + ResolvedView.SpecularOverrideParameter.xyz;
-#endif
-	}
-#endif
-	
-#if MATERIAL_SHADINGMODEL_SINGLELAYERWATER
-	// Fade out diffuse as this will be hanlded by the single scattering lighting. when over the water surface.
-	// We keep the SpecularColor for sun/water interactions
-	float3 WaterDiffuseIndirectLuminance = 0;
-	DiffuseColor *= BaseMaterialCoverageOverWater;
-#endif
-	
 	half3 Color = 0;
+
+#if !MATERIAL_SHADINGMODEL_UNLIT
 	//To keep IndirectLightingCache conherence with PC, initialize the IndirectIrradiance to zero.
 	half IndirectIrradiance = 0;
-
-	#if LQ_TEXTURE_LIGHTMAP
-		float2 LightmapUV0, LightmapUV1;
-		uint LightmapDataIndex;
-		GetLightMapCoordinates(Interpolants, LightmapUV0, LightmapUV1, LightmapDataIndex);
-
-		half4 LightmapColor = GetLightMapColorLQ( LightmapUV0, LightmapUV1, LightmapDataIndex, MaterialParameters.WorldNormal );
-		Color += LightmapColor.rgb * DiffuseColor * View.IndirectLightingColorScale;
-		IndirectIrradiance = LightmapColor.a;
-	#elif CACHED_POINT_INDIRECT_LIGHTING
-		#if MATERIALBLENDING_MASKED || MATERIALBLENDING_SOLID
-
-			// Take the normal into account for opaque
-			FThreeBandSHVectorRGB PointIndirectLighting;
-			PointIndirectLighting.R.V0 = IndirectLightingCache.IndirectLightingSHCoefficients0[0];
-			PointIndirectLighting.R.V1 = IndirectLightingCache.IndirectLightingSHCoefficients1[0];
-			PointIndirectLighting.R.V2 = IndirectLightingCache.IndirectLightingSHCoefficients2[0];
-
-			PointIndirectLighting.G.V0 = IndirectLightingCache.IndirectLightingSHCoefficients0[1];
-			PointIndirectLighting.G.V1 = IndirectLightingCache.IndirectLightingSHCoefficients1[1];
-			PointIndirectLighting.G.V2 = IndirectLightingCache.IndirectLightingSHCoefficients2[1];
-
-			PointIndirectLighting.B.V0 = IndirectLightingCache.IndirectLightingSHCoefficients0[2];
-			PointIndirectLighting.B.V1 = IndirectLightingCache.IndirectLightingSHCoefficients1[2];
-			PointIndirectLighting.B.V2 = IndirectLightingCache.IndirectLightingSHCoefficients2[2];
-
-			FThreeBandSHVector DiffuseTransferSH = CalcDiffuseTransferSH3(MaterialParameters.WorldNormal, 1);
-
-			// Compute diffuse lighting which takes the normal into account
-			half3 DiffuseGI = max(half3(0,0,0), DotSH3(PointIndirectLighting, DiffuseTransferSH));
-
-			IndirectIrradiance = Luminance(DiffuseGI);
-			Color += DiffuseColor * DiffuseGI * View.IndirectLightingColorScale;			
-		#else 
-			
-			// Non-directional for translucency
-			// Ambient terms packed in xyz
-			// Already divided by PI and SH ambient on CPU
-			half3 PointIndirectLighting = IndirectLightingCache.IndirectLightingSHSingleCoefficient;
-			half3 DiffuseGI = PointIndirectLighting;
-
-			IndirectIrradiance = Luminance(DiffuseGI);
-			Color += DiffuseColor * DiffuseGI * View.IndirectLightingColorScale;
-
-		#endif
-	#endif
-
-	#if !MATERIAL_SHADINGMODEL_UNLIT
-		half MaterialAO = GetMaterialAmbientOcclusion(PixelMaterialInputs);
-		Color *= MaterialAO;	
-		IndirectIrradiance *= MaterialAO;
+	// Indirect Diffuse
+#if LQ_TEXTURE_LIGHTMAP
+	float2 LightmapUV0, LightmapUV1;
+	uint LightmapDataIndex;
+	GetLightMapCoordinates(Interpolants, LightmapUV0, LightmapUV1, LightmapDataIndex);
+
+	half4 LightmapColor = GetLightMapColorLQ(LightmapUV0, LightmapUV1, LightmapDataIndex, MaterialParameters.WorldNormal);
+	Color += LightmapColor.rgb * ShadingModelContext.DiffuseColor * View.IndirectLightingColorScale;
+	IndirectIrradiance = LightmapColor.a;
+#elif CACHED_POINT_INDIRECT_LIGHTING
+	#if MATERIALBLENDING_MASKED || MATERIALBLENDING_SOLID
+		// Take the normal into account for opaque
+		FThreeBandSHVectorRGB PointIndirectLighting;
+		PointIndirectLighting.R.V0 = IndirectLightingCache.IndirectLightingSHCoefficients0[0];
+		PointIndirectLighting.R.V1 = IndirectLightingCache.IndirectLightingSHCoefficients1[0];
+		PointIndirectLighting.R.V2 = IndirectLightingCache.IndirectLightingSHCoefficients2[0];
+
+		PointIndirectLighting.G.V0 = IndirectLightingCache.IndirectLightingSHCoefficients0[1];
+		PointIndirectLighting.G.V1 = IndirectLightingCache.IndirectLightingSHCoefficients1[1];
+		PointIndirectLighting.G.V2 = IndirectLightingCache.IndirectLightingSHCoefficients2[1];
+
+		PointIndirectLighting.B.V0 = IndirectLightingCache.IndirectLightingSHCoefficients0[2];
+		PointIndirectLighting.B.V1 = IndirectLightingCache.IndirectLightingSHCoefficients1[2];
+		PointIndirectLighting.B.V2 = IndirectLightingCache.IndirectLightingSHCoefficients2[2];
+
+		FThreeBandSHVector DiffuseTransferSH = CalcDiffuseTransferSH3(MaterialParameters.WorldNormal, 1);
+
+		// Compute diffuse lighting which takes the normal into account
+		half3 DiffuseGI = max(half3(0, 0, 0), DotSH3(PointIndirectLighting, DiffuseTransferSH));
+
+		IndirectIrradiance = Luminance(DiffuseGI);
+		Color += ShadingModelContext.DiffuseColor * DiffuseGI * View.IndirectLightingColorScale;
+	#else 
+		// Non-directional for translucency
+		// Ambient terms packed in xyz
+		// Already divided by PI and SH ambient on CPU
+		half3 PointIndirectLighting = IndirectLightingCache.IndirectLightingSHSingleCoefficient;
+		half3 DiffuseGI = PointIndirectLighting;
+
+		IndirectIrradiance = Luminance(DiffuseGI);
+		Color += ShadingModelContext.DiffuseColor * DiffuseGI * View.IndirectLightingColorScale;
 	#endif
+#endif
+	half MaterialAO = GetMaterialAmbientOcclusion(PixelMaterialInputs);
+	Color *= MaterialAO;
+	IndirectIrradiance *= MaterialAO;
 
-		half Shadow = 1.0f;
-	#if !MATERIAL_SHADINGMODEL_UNLIT
-		Shadow = GetPrimaryPrecomputedShadowMask(Interpolants).r;
-		#if DIRECTIONAL_LIGHT_CSM && !MATERIAL_SHADINGMODEL_SINGLELAYERWATER
-		// Cascaded Shadow Map
+	// Shadow
+	half Shadow = GetPrimaryPrecomputedShadowMask(Interpolants).r;
+#if DIRECTIONAL_LIGHT_CSM && !MATERIAL_SHADINGMODEL_SINGLELAYERWATER
+	// Cascaded Shadow Map
+	{
+		FPCFSamplerSettings Settings;
+		Settings.ShadowDepthTexture = MobileDirectionalLight.DirectionalLightShadowTexture;
+		Settings.ShadowDepthTextureSampler = MobileDirectionalLight.DirectionalLightShadowSampler;
+		Settings.TransitionScale = MobileDirectionalLight.DirectionalLightDirectionAndShadowTransition.w;
+		Settings.ShadowBufferSize = MobileDirectionalLight.DirectionalLightShadowSize;
+		Settings.bSubsurface = false;
+		Settings.bTreatMaxDepthUnshadowed = false;
+		Settings.DensityMulConstant = 0;
+		Settings.ProjectionDepthBiasParameters = 0;
+
+		float4 ShadowPosition = float4(0, 0, 0, 0);
+		for (int i = 0; i < MAX_MOBILE_SHADOWCASCADES; i++)
 		{
-			FPCFSamplerSettings Settings;
-			Settings.ShadowDepthTexture = MobileDirectionalLight.DirectionalLightShadowTexture;
-			Settings.ShadowDepthTextureSampler = MobileDirectionalLight.DirectionalLightShadowSampler;
-			Settings.TransitionScale = MobileDirectionalLight.DirectionalLightDirectionAndShadowTransition.w;
-			Settings.ShadowBufferSize = MobileDirectionalLight.DirectionalLightShadowSize;
-			Settings.bSubsurface = false;
-			Settings.bTreatMaxDepthUnshadowed = false;
-			Settings.DensityMulConstant = 0;
-			Settings.ProjectionDepthBiasParameters = 0;
-
-			float4 ShadowPosition = float4(0,0,0,0);
-			for (int i = 0; i < MAX_MOBILE_SHADOWCASCADES; i++)
-			{
-				if (MaterialParameters.ScreenPosition.w < MobileDirectionalLight.DirectionalLightShadowDistances[i])
-				{
-				#if MOBILE_MULTI_VIEW
-					ShadowPosition = mul(float4(MaterialParameters.ScreenPosition.xyw, 1), ResolvedView.MobileMultiviewShadowTransform);
-					ShadowPosition = mul(ShadowPosition, MobileDirectionalLight.DirectionalLightScreenToShadow[i]);
-				#else
-					ShadowPosition = mul(float4(MaterialParameters.ScreenPosition.xyw, 1), MobileDirectionalLight.DirectionalLightScreenToShadow[i]);
-				#endif
-					break; // position found.
-				}
-			}
-
-			// Process CSM only when ShadowPosition is valid.
-			if (ShadowPosition.z > 0)
+			if (MaterialParameters.ScreenPosition.w < MobileDirectionalLight.DirectionalLightShadowDistances[i])
 			{
-				// Clamp pixel depth in light space for shadowing opaque, because areas of the shadow depth buffer that weren't rendered to will have been cleared to 1
-				// We want to force the shadow comparison to result in 'unshadowed' in that case, regardless of whether the pixel being shaded is in front or behind that plane
-				float LightSpacePixelDepthForOpaque = min(ShadowPosition.z, 0.99999f);
-				Settings.SceneDepth = LightSpacePixelDepthForOpaque;
-
-				#if MOBILE_CSM_QUALITY == 0
-					half ShadowMap = ManualNoFiltering(ShadowPosition.xy, Settings);
-				#elif MOBILE_CSM_QUALITY == 1
-					half ShadowMap = Manual1x1PCF(ShadowPosition.xy, Settings);
-				#elif MOBILE_CSM_QUALITY == 2
-					half ShadowMap = Manual2x2PCF(ShadowPosition.xy, Settings);
-				#else
-					#error Unsupported MOBILE_CSM_QUALITY value.
-				#endif
-
-				#if FADE_CSM
-					float Fade = saturate(MaterialParameters.ScreenPosition.w * MobileDirectionalLight.DirectionalLightDistanceFadeMADAndSpecularScale.x + MobileDirectionalLight.DirectionalLightDistanceFadeMADAndSpecularScale.y);
-					// lerp out shadow based on fade params.
-					ShadowMap = lerp(ShadowMap, 1.0, Fade * Fade);
-				#endif
-
-				#if MOVABLE_DIRECTIONAL_LIGHT
-					Shadow = ShadowMap;
-				#else
-					Shadow = min(ShadowMap, Shadow);
-				#endif
+			#if MOBILE_MULTI_VIEW
+				ShadowPosition = mul(float4(MaterialParameters.ScreenPosition.xyw, 1), ResolvedView.MobileMultiviewShadowTransform);
+				ShadowPosition = mul(ShadowPosition, MobileDirectionalLight.DirectionalLightScreenToShadow[i]);
+			#else
+				ShadowPosition = mul(float4(MaterialParameters.ScreenPosition.xyw, 1), MobileDirectionalLight.DirectionalLightScreenToShadow[i]);
+			#endif
+				break; // position found.
 			}
 		}
-		#endif /* DIRECTIONAL_LIGHT_CSM */
 
-		float NoL = max(0, dot(float3(MaterialParameters.WorldNormal), float3(MobileDirectionalLight.DirectionalLightDirectionAndShadowTransition.xyz)));
-		float RoL = max(0, dot(float3(MaterialParameters.ReflectionVector), float3(MobileDirectionalLight.DirectionalLightDirectionAndShadowTransition.xyz)));
-		float3 H = normalize(float3(MaterialParameters.CameraVector) + float3(MobileDirectionalLight.DirectionalLightDirectionAndShadowTransition.xyz));
-		float NoH = max(0,dot(MaterialParameters.WorldNormal, H));
-
-		#if FULLY_ROUGH
-			Color += (Shadow * NoL) * MobileDirectionalLight.DirectionalLightColor.rgb * DiffuseColor;
-		#else
-			// MobileDirectionalLight.DirectionalLightDistanceFadeMADAndSpecularScale.z saves SpecularScale for direction light.
-			Color += (Shadow * NoL) * MobileDirectionalLight.DirectionalLightColor.rgb * (DiffuseColor + SpecularColor * MobileDirectionalLight.DirectionalLightDistanceFadeMADAndSpecularScale.z * CalcSpecular(Roughness, RoL, NoH, H, MaterialParameters.WorldNormal));
-			// Environment map has been prenormalized, scale by lightmap luminance
-			half RoughnessRaw = clamp(saturate(GetMaterialRoughnessRaw(PixelMaterialInputs)), 1e-6, 1.0);
-			half3 SpecularIBL = GetImageBasedReflectionLighting(MaterialParameters, RoughnessRaw, IndirectIrradiance);
-
-			#if MATERIAL_PLANAR_FORWARD_REFLECTIONS
-				BRANCH
-				if (abs(dot(PlanarReflectionStruct.ReflectionPlane.xyz, 1)) > .0001f)
-				{
-					half4 PlanarReflection = GetPlanarReflection(MaterialParameters.AbsoluteWorldPosition, MaterialParameters.WorldNormal, Roughness);
-					// Planar reflections win over reflection environment
-					SpecularIBL = lerp(SpecularIBL, PlanarReflection.rgb, PlanarReflection.a);
-				}
+		// Process CSM only when ShadowPosition is valid.
+		if (ShadowPosition.z > 0)
+		{
+			// Clamp pixel depth in light space for shadowing opaque, because areas of the shadow depth buffer that weren't rendered to will have been cleared to 1
+			// We want to force the shadow comparison to result in 'unshadowed' in that case, regardless of whether the pixel being shaded is in front or behind that plane
+			float LightSpacePixelDepthForOpaque = min(ShadowPosition.z, 0.99999f);
+			Settings.SceneDepth = LightSpacePixelDepthForOpaque;
+
+			#if MOBILE_CSM_QUALITY == 0
+				half ShadowMap = ManualNoFiltering(ShadowPosition.xy, Settings);
+			#elif MOBILE_CSM_QUALITY == 1
+				half ShadowMap = Manual1x1PCF(ShadowPosition.xy, Settings);
+			#elif MOBILE_CSM_QUALITY == 2
+				half ShadowMap = Manual2x2PCF(ShadowPosition.xy, Settings);
+			#else
+				#error Unsupported MOBILE_CSM_QUALITY value.
 			#endif
-			Color += SpecularIBL * SpecularColor;
-		#endif
-
-		#if ENABLE_SKY_LIGHT
-			//@mw todo
-			// TODO: Also need to do specular.
-#if MATERIAL_TWOSIDED && LQ_TEXTURE_LIGHTMAP
-			if (NoL == 0)
-			{
-#endif
-			half3 SkyDiffuseLighting = GetSkySHDiffuseSimple(MaterialParameters.WorldNormal);
-
-		#if MATERIAL_SHADINGMODEL_SINGLELAYERWATER
-			WaterDiffuseIndirectLuminance += SkyDiffuseLighting;
-		#endif
-			Color += SkyDiffuseLighting * half3(ResolvedView.SkyLightColor.rgb) * DiffuseColor * MaterialAO;
-#if MATERIAL_TWOSIDED && LQ_TEXTURE_LIGHTMAP
-			
-			}
-#endif
-		#endif
 
-		#if MAX_DYNAMIC_POINT_LIGHTS > 0 && !MATERIAL_SHADINGMODEL_SINGLELAYERWATER
-			#if VARIABLE_NUM_DYNAMIC_POINT_LIGHTS
-			#if WEBGL || ES2_PROFILE || VULKAN_PROFILE
-			// webgl and ES2 needs a constant loop counter
-			UNROLL
-			for (int i = 0; i < MAX_DYNAMIC_POINT_LIGHTS; i++)
-			{
-				if (i < NumDynamicPointLights)
-				{
-			#else
-			for (int i = 0; i < NumDynamicPointLights; i++)
-			{
+			#if FADE_CSM
+				float Fade = saturate(MaterialParameters.ScreenPosition.w * MobileDirectionalLight.DirectionalLightDistanceFadeMADAndSpecularScale.x + MobileDirectionalLight.DirectionalLightDistanceFadeMADAndSpecularScale.y);
+			// lerp out shadow based on fade params.
+				ShadowMap = lerp(ShadowMap, 1.0, Fade * Fade);
 			#endif
+
+			#if MOVABLE_DIRECTIONAL_LIGHT
+				Shadow = ShadowMap;
 			#else
-			for (int i = 0; i < NUM_DYNAMIC_POINT_LIGHTS; i++)
-			{
+				Shadow = min(ShadowMap, Shadow);
 			#endif
-				float3 ToLight = LightPositionAndInvRadius[i].xyz - MaterialParameters.AbsoluteWorldPosition;
-				float DistanceSqr = dot(ToLight, ToLight);
-				float3 L = ToLight * rsqrt(DistanceSqr);
-				float3 H = normalize(MaterialParameters.CameraVector + L);
+		}
+	}
+#endif /* DIRECTIONAL_LIGHT_CSM */
+
+	float NoL = max(0, dot(float3(MaterialParameters.WorldNormal), float3(MobileDirectionalLight.DirectionalLightDirectionAndShadowTransition.xyz)));
+	float RoL = max(0, dot(float3(MaterialParameters.ReflectionVector), float3(MobileDirectionalLight.DirectionalLightDirectionAndShadowTransition.xyz)));
+	float3 H = normalize(float3(MaterialParameters.CameraVector) + float3(MobileDirectionalLight.DirectionalLightDirectionAndShadowTransition.xyz));
+	float NoH = max(0, dot(MaterialParameters.WorldNormal, H));
+
+	// Direct Lighting (Directional light) + IBL
+#if FULLY_ROUGH
+	Color += (Shadow * NoL) * MobileDirectionalLight.DirectionalLightColor.rgb * ShadingModelContext.DiffuseColor;
+#else
+	FMobileDirectLighting Lighting = MobileIntegrateBxDF(ShadingModelContext, NoL, RoL, MaterialParameters.CameraVector, MaterialParameters.WorldNormal, H, NoH);
+	// MobileDirectionalLight.DirectionalLightDistanceFadeMADAndSpecularScale.z saves SpecularScale for direction light.
+	Color += (Shadow * NoL) * MobileDirectionalLight.DirectionalLightColor.rgb * (Lighting.Diffuse + Lighting.Specular * MobileDirectionalLight.DirectionalLightDistanceFadeMADAndSpecularScale.z);
+
+	// Environment map has been prenormalized, scale by lightmap luminance
+	half3 SpecularIBL = GetImageBasedReflectionLighting(MaterialParameters, ShadingModelContext.Roughness, IndirectIrradiance);
+	#if MATERIAL_PLANAR_FORWARD_REFLECTIONS
+		BRANCH
+		if (abs(dot(PlanarReflectionStruct.ReflectionPlane.xyz, 1)) > .0001f)
+		{
+			half4 PlanarReflection = GetPlanarReflection(MaterialParameters.AbsoluteWorldPosition, MaterialParameters.WorldNormal, ShadingModelContext.Roughness);
+			// Planar reflections win over reflection environment
+			SpecularIBL = lerp(SpecularIBL, PlanarReflection.rgb, PlanarReflection.a);
+		}
+	#endif
+	#if MATERIAL_SHADINGMODEL_CLEARCOAT
+		half SpecularOcclusion = 1; // GetSpecularOcclusion(ShadingModelContext.NoV, ShadingModelContext.Roughness*ShadingModelContext.Roughness, MaterialAO);
+		//Color += SpecularIBL * ShadingModelContext.SpecularColor * ShadingModelContext.ClearCoat * SpecularOcclusion;
+		float2 AB = PreIntegratedGF.SampleLevel(PreIntegratedGFSampler, float2(ShadingModelContext.NoV, ShadingModelContext.Roughness), 0).rg;
+
+		// F_Schlick
+		float F = EnvBRDF(0.04, ShadingModelContext.ClearCoatRoughness, ShadingModelContext.NoV).x;
+
+		F *= ShadingModelContext.ClearCoat;
 
-				float PointNoL = max(0, dot(MaterialParameters.WorldNormal, L));
-				float PointRoL = max(0, dot(MaterialParameters.ReflectionVector, L));
-				float PointNoH = max(0, dot(MaterialParameters.WorldNormal, H));
+		float LayerAttenuation = (1 - F);
 
-				float Attenuation;
+		Color += SpecularIBL * ShadingModelContext.SpecularColor * LayerAttenuation * (ShadingModelContext.SpecularColor * AB.x + AB.y * saturate(50 * ShadingModelContext.SpecularColor.g) * (1 - ShadingModelContext.ClearCoat));
 
-				if (LightColorAndFalloffExponent[i].w == 0)
-				{
-					// Sphere falloff (technically just 1/d2 but this avoids inf)
-					Attenuation = 1 / ( DistanceSqr + 1 );
+		SpecularIBL = GetImageBasedReflectionLighting(MaterialParameters, ShadingModelContext.ClearCoatRoughness, IndirectIrradiance);
+		Color += SpecularIBL * ShadingModelContext.SpecularColor * F * SpecularOcclusion;
+	#else
+		Color += SpecularIBL * ShadingModelContext.SpecularColor;
+	#endif
+#endif /* FULLY_ROUGH */
+
+	// Local lights
+#if MAX_DYNAMIC_POINT_LIGHTS > 0 && !MATERIAL_SHADINGMODEL_SINGLELAYERWATER
+	#if VARIABLE_NUM_DYNAMIC_POINT_LIGHTS
+	#if WEBGL || ES2_PROFILE || VULKAN_PROFILE
+	// webgl and ES2 needs a constant loop counter
+	UNROLL
+	for (int i = 0; i < MAX_DYNAMIC_POINT_LIGHTS; i++)
+	{
+		if (i < NumDynamicPointLights)
+		{
+	#else
+	for (int i = 0; i < NumDynamicPointLights; i++)
+	{
+	#endif
+	#else
+	for (int i = 0; i < NUM_DYNAMIC_POINT_LIGHTS; i++)
+	{
+	#endif
+		float3 ToLight = LightPositionAndInvRadius[i].xyz - MaterialParameters.AbsoluteWorldPosition;
+		float DistanceSqr = dot(ToLight, ToLight);
+		float3 L = ToLight * rsqrt(DistanceSqr);
+		float3 PointH = normalize(MaterialParameters.CameraVector + L);
+
+		float PointNoL = max(0, dot(MaterialParameters.WorldNormal, L));
+		float PointRoL = max(0, dot(MaterialParameters.ReflectionVector, L));
+		float PointNoH = max(0, dot(MaterialParameters.WorldNormal, PointH));
+
+		float Attenuation;
+
+		if (LightColorAndFalloffExponent[i].w == 0)
+		{
+			// Sphere falloff (technically just 1/d2 but this avoids inf)
+			Attenuation = 1 / ( DistanceSqr + 1 );
 	
-					float LightRadiusMask = Square(saturate(1 - Square(DistanceSqr * (LightPositionAndInvRadius[i].w * LightPositionAndInvRadius[i].w))));
-					Attenuation *= LightRadiusMask;
-				}
-				else
-				{
-					Attenuation = RadialAttenuation(ToLight * LightPositionAndInvRadius[i].w, LightColorAndFalloffExponent[i].w);		
-				}
+			float LightRadiusMask = Square(saturate(1 - Square(DistanceSqr * (LightPositionAndInvRadius[i].w * LightPositionAndInvRadius[i].w))));
+			Attenuation *= LightRadiusMask;
+		}
+		else
+		{
+			Attenuation = RadialAttenuation(ToLight * LightPositionAndInvRadius[i].w, LightColorAndFalloffExponent[i].w);		
+		}
 				
-				#if PROJECT_MOBILE_ENABLE_MOVABLE_SPOTLIGHTS
-				if(SpotLightAngles[i].w > 0.0)
-				{
-					Attenuation *= SpotAttenuation(L, -SpotLightDirectionAndSpecularScale[i].xyz, SpotLightAngles[i].xy);
-				}
-				#endif
-
-				#if !FULLY_ROUGH
-					Color += min(65000.0, (Attenuation * PointNoL) * LightColorAndFalloffExponent[i].rgb * (1.0/PI) * (DiffuseColor + SpecularColor * SpotLightDirectionAndSpecularScale[i].w * CalcSpecular(Roughness, PointRoL, PointNoH, H, MaterialParameters.WorldNormal)));
-				#else
-					Color += (Attenuation * PointNoL) * LightColorAndFalloffExponent[i].rgb * DiffuseColor;
-				#endif
-			#if VARIABLE_NUM_DYNAMIC_POINT_LIGHTS && ( WEBGL || ES2_PROFILE || VULKAN_PROFILE )
-				}
-			#endif
-			}
+		#if PROJECT_MOBILE_ENABLE_MOVABLE_SPOTLIGHTS
+		if(SpotLightAngles[i].w > 0.0)
+		{
+			Attenuation *= SpotAttenuation(L, -SpotLightDirectionAndSpecularScale[i].xyz, SpotLightAngles[i].xy);
+		}
 		#endif
 
-	#endif /* !MATERIAL_SHADINGMODEL_UNLIT */
+		#if !FULLY_ROUGH
+			FMobileDirectLighting Lighting = MobileIntegrateBxDF(ShadingModelContext, PointNoL, PointRoL, MaterialParameters.CameraVector, MaterialParameters.WorldNormal, PointH, PointNoH);
+			Color += min(65000.0, (Attenuation * PointNoL) * LightColorAndFalloffExponent[i].rgb * (1.0/PI) * (Lighting.Diffuse + Lighting.Specular * SpotLightDirectionAndSpecularScale[i].w));
+		#else
+			Color += (Attenuation * PointNoL) * LightColorAndFalloffExponent[i].rgb * (1.0/PI) * ShadingModelContext.DiffuseColor;
+		#endif
+	#if VARIABLE_NUM_DYNAMIC_POINT_LIGHTS && ( WEBGL || ES2_PROFILE || VULKAN_PROFILE )
+		}
+	#endif
+	}
+#endif
+
+	// Skylight
+#if ENABLE_SKY_LIGHT
+	//@mw todo
+	// TODO: Also need to do specular.
+	#if MATERIAL_TWOSIDED && LQ_TEXTURE_LIGHTMAP
+	if (NoL == 0)
+	{
+	#endif
+		half3 SkyDiffuseLighting = GetSkySHDiffuseSimple(MaterialParameters.WorldNormal);
+
+	#if MATERIAL_SHADINGMODEL_SINGLELAYERWATER
+		ShadingModelContext.WaterDiffuseIndirectLuminance += SkyDiffuseLighting;
+	#endif
+		Color += SkyDiffuseLighting * half3(ResolvedView.SkyLightColor.rgb) * ShadingModelContext.DiffuseColor * MaterialAO;
+	#if MATERIAL_TWOSIDED && LQ_TEXTURE_LIGHTMAP
+	}
+	#endif
+#endif
+
+#endif /* !MATERIAL_SHADINGMODEL_UNLIT */
 		 
 	half3 Emissive = GetMaterialEmissive(PixelMaterialInputs);
-
 	Color += Emissive;
 
 #if MATERIAL_SHADINGMODEL_SINGLELAYERWATER
 	{
 		const bool CameraIsUnderWater = false;	// Fade out the material contribution over to water contribution according to material opacity.
 		const float3 SunIlluminance = ResolvedView.DirectionalLightColor.rgb * PI;			// times PI because it is divided by PI on CPU (=luminance) and we want illuminance here. 
-		const float3 WaterDiffuseIndirectIlluminance = WaterDiffuseIndirectLuminance * PI;	// DiffuseIndirectLighting is luminance. So we need to multiply by PI to get illuminance.
-		const float3 EnvBrdf = SpecularColor; // SpecularColor is not F0 as in BasePassPixelShader, it is EnvBRDF.
+		const float3 WaterDiffuseIndirectIlluminance = ShadingModelContext.WaterDiffuseIndirectLuminance * PI;	// DiffuseIndirectLighting is luminance. So we need to multiply by PI to get illuminance.
+		const float3 EnvBrdf = ShadingModelContext.SpecularColor; // SpecularColor is not F0 as in BasePassPixelShader, it is EnvBRDF.
 		const uint EyeIndex = 0;
 
 		const float4 NullDistortionParams = 1.0f;
 		WaterVolumeLightingOutput WaterLighting = EvaluateWaterVolumeLighting(
 			MaterialParameters, PixelMaterialInputs, ResolvedView,
-			Shadow, Specular, NullDistortionParams,
+			Shadow, ShadingModelContext.Specular, NullDistortionParams,
 			SunIlluminance, WaterDiffuseIndirectIlluminance, EnvBrdf,
 			CameraIsUnderWater, WaterVisibility, EyeIndex);
 
@@ -765,22 +730,22 @@ void Main(
 	// NEEDS_BASEPASS_PIXEL_FOGGING is not allowed on mobile for the sake of performance.
 
 #if !MATERIAL_SHADINGMODEL_UNLIT && MOBILE_EMULATION
-	Color = lerp(Color, DiffuseColor + SpecularColor, ResolvedView.UnlitViewmodeMask);
+	Color = lerp(Color, ShadingModelContext.DiffuseColor + ShadingModelContext.SpecularColor, ResolvedView.UnlitViewmodeMask);
 #endif
 
 	// On mobile, water (an opaque material) is rendered as trnaslucent with forced premultiplied alpha blending (see MobileBasePass::SetTranslucentRenderState)
 	#if MATERIALBLENDING_ALPHACOMPOSITE || MATERIAL_SHADINGMODEL_SINGLELAYERWATER
-		OutColor = half4(Color * VertexFog.a + VertexFog.rgb * Opacity, Opacity);
+		OutColor = half4(Color * VertexFog.a + VertexFog.rgb * ShadingModelContext.Opacity, ShadingModelContext.Opacity);
 	#elif MATERIALBLENDING_ALPHAHOLDOUT
 		// not implemented for holdout
-		OutColor = half4(Color * VertexFog.a + VertexFog.rgb * Opacity, Opacity);
+		OutColor = half4(Color * VertexFog.a + VertexFog.rgb * ShadingModelContext.Opacity, ShadingModelContext.Opacity);
 	#elif MATERIALBLENDING_TRANSLUCENT
-		OutColor = half4(Color * VertexFog.a + VertexFog.rgb, Opacity);
+		OutColor = half4(Color * VertexFog.a + VertexFog.rgb, ShadingModelContext.Opacity);
 	#elif MATERIALBLENDING_ADDITIVE
-		OutColor = half4(Color * (VertexFog.a * Opacity.x), 0.0f);
+		OutColor = half4(Color * (VertexFog.a * ShadingModelContext.Opacity.x), 0.0f);
 	#elif MATERIALBLENDING_MODULATE
 		half3 FoggedColor = lerp(half3(1, 1, 1), Color, VertexFog.aaa * VertexFog.aaa);
-		OutColor = half4(FoggedColor, Opacity);
+		OutColor = half4(FoggedColor, ShadingModelContext.Opacity);
 	#else
 		OutColor.rgb = Color * VertexFog.a + VertexFog.rgb;
 
diff --git a/Engine/Shaders/Private/MobileShadingModels.ush b/Engine/Shaders/Private/MobileShadingModels.ush
new file mode 100644
--- /dev/null
+++ b/Engine/Shaders/Private/MobileShadingModels.ush
@@ -0,0 +1,181 @@
+// Copyright 1998-2020 Epic Games, Inc. All Rights Reserved.
+
+#pragma once
+
+#include "BRDF.ush"
+#include "MobileGGX.ush"
+
+struct FMobileDirectLighting
+{
+	half3 Diffuse;
+	half3 Specular;
+};
+
+struct FMobileShadingModelContext
+{
+	half Opacity;
+	half3 BaseColor;
+	half Metallic;
+	half Specular;
+	half Roughness;
+	half RoughnessWithClamp;
+
+	half3 DiffuseColor;
+#if NONMETAL
+	half SpecularColor;
+#else
+	half3 SpecularColor;
+#endif
+
+#if MATERIAL_SHADINGMODEL_CLEAR_COAT
+	half ClearCoat;
+	half ClearCoatRoughness;
+	float NoV;
+#elif MATERIAL_SHADINGMODEL_SINGLELAYERWATER
+	half BaseMaterialCoverageOverWater;
+	half WaterVisibility;
+	float3 WaterDiffuseIndirectLuminance;
+#endif
+};
+
+
+
+#if NONMETAL
+half GetEnvBRDF(half SpecularColor, half Roughness, float NoV)
+{
+	// If nothing is hooked up to Metalic and Specular,
+	// then defaults are the same as a non-metal,
+	// so this define is safe.
+	return EnvBRDFApproxNonmetal(Roughness, NoV);
+}
+#else
+half3 GetEnvBRDF(half3 SpecularColor, half Roughness, float NoV)
+{
+#if FORWARDSHADING_USE_HQ_ENV_BRDF
+	return EnvBRDF(SpecularColor, Roughness, NoV);
+#else
+	return EnvBRDFApprox(SpecularColor, Roughness, NoV);
+#endif
+}
+#endif
+
+void InitShadingModelContext( inout FMobileShadingModelContext ShadingModelContext, FMaterialPixelParameters MaterialParameters)
+{
+#if !FULLY_ROUGH
+	float NoV = max(dot(MaterialParameters.WorldNormal, MaterialParameters.CameraVector), 0);
+#endif
+
+#if NONMETAL
+	ShadingModelContext.DiffuseColor = BaseColor;
+	ShadingModelContext.SpecularColor = 0.04;
+#else
+
+	#if MATERIAL_SHADINGMODEL_CLEAR_COAT
+		ShadingModelContext.ClearCoat = saturate(GetMaterialCustomData0(MaterialParameters));
+		ShadingModelContext.ClearCoatRoughness = clamp(GetMaterialCustomData1(MaterialParameters), 0.12, 1.0);
+		const half MetalSpec = 0.9;
+
+		// Approximation of refraction's effect on EnvBRDF
+		half RefractionScale = ((NoV * 0.5 + 0.5) * NoV - 1) * saturate(1.25 - 1.25 * ShadingModelContext.Roughness) + 1;
+
+		// Approximation of absorption integral, tuned for Roughness=0.4
+		half3 AbsorptionColor = ShadingModelContext.BaseColor * (1 / MetalSpec);
+		half3 Absorption = AbsorptionColor * ((NoV - 1) * 0.85 * (1 - lerp(AbsorptionColor, Square(AbsorptionColor), -0.78)) + 1);
+
+		// F_Schlick
+		half F0 = 0.04;
+		half Fc = Pow5(1 - NoV);
+		half F = Fc + (1 - Fc) * F0;
+		half LayerAttenuation = lerp(1, (1 - F), ShadingModelContext.ClearCoat);
+
+		ShadingModelContext.Specular *= lerp(1, RefractionScale, ShadingModelContext.ClearCoat);
+		half DielectricSpecular = 0.08 * ShadingModelContext.Specular;
+		ShadingModelContext.SpecularColor = (DielectricSpecular - DielectricSpecular * ShadingModelContext.Metallic) + ShadingModelContext.BaseColor * ShadingModelContext.Metallic;	// 2 mad
+		ShadingModelContext.DiffuseColor = ShadingModelContext.BaseColor - ShadingModelContext.BaseColor * ShadingModelContext.Metallic;
+
+		// Adjust BaseColor according to ClearCoat layer attenuation
+		ShadingModelContext.BaseColor = lerp(ShadingModelContext.BaseColor * LayerAttenuation, MetalSpec * Absorption * RefractionScale, ShadingModelContext.Metallic * ShadingModelContext.ClearCoat);
+		const half Dither = InterleavedGradientNoise(MaterialParameters.SvPosition.xy, View.StateFrameIndexMod8);
+		ShadingModelContext.BaseColor += Dither / 255.f;
+
+		// This is to prevent Vis to get inf when both NoL and NoV are 0.
+		ShadingModelContext.NoV = saturate(abs(NoV) + 1e-5);
+
+	#elif MATERIAL_SHADINGMODEL_SINGLELAYERWATER
+		ShadingModelContext.BaseMaterialCoverageOverWater = ShadingModelContext.Opacity;
+		ShadingModelContext.WaterVisibility = 1.0 - ShadingModelContext.BaseMaterialCoverageOverWater;
+
+		half DielectricSpecular = 0.08 * ShadingModelContext.Specular;
+		ShadingModelContext.DiffuseColor = ShadingModelContext.BaseColor - ShadingModelContext.BaseColor * ShadingModelContext.Metallic;	// 1 mad
+		ShadingModelContext.SpecularColor = (DielectricSpecular - DielectricSpecular * ShadingModelContext.Metallic) + ShadingModelContext.BaseColor * ShadingModelContext.Metallic;	// 2 mad
+
+	#else // MATERIAL_SHADINGMODEL_DEFAULTLIT
+		half DielectricSpecular = 0.08 * ShadingModelContext.Specular;
+		ShadingModelContext.DiffuseColor = ShadingModelContext.BaseColor - ShadingModelContext.BaseColor * ShadingModelContext.Metallic;	// 1 mad
+		ShadingModelContext.SpecularColor = (DielectricSpecular - DielectricSpecular * ShadingModelContext.Metallic) + ShadingModelContext.BaseColor * ShadingModelContext.Metallic;	// 2 mad
+
+	#endif
+
+#endif
+
+#if FULLY_ROUGH
+	EnvBRDFApproxFullyRough(ShadingModelContext.DiffuseColor, ShadingModelContext.SpecularColor);
+#else
+	ShadingModelContext.SpecularColor = GetEnvBRDF(ShadingModelContext.SpecularColor, ShadingModelContext.Roughness, NoV);
+#endif
+
+#if MOBILE_EMULATION
+	{
+		// this feature is only needed for development/editor - we can compile it out for a shipping build (see r.CompileShadersForDevelopment cvar help)
+		ShadingModelContext.DiffuseColor = ShadingModelContext.DiffuseColor * ResolvedView.DiffuseOverrideParameter.w + ResolvedView.DiffuseOverrideParameter.xyz;
+#if NONMETAL
+		ShadingModelContext.SpecularColor = ShadingModelContext.SpecularColor * ResolvedView.SpecularOverrideParameter.w + ResolvedView.SpecularOverrideParameter.x;
+#else
+		ShadingModelContext.SpecularColor = ShadingModelContext.SpecularColor * ResolvedView.SpecularOverrideParameter.w + ResolvedView.SpecularOverrideParameter.xyz;
+#endif
+	}
+#endif
+
+#if MATERIAL_SHADINGMODEL_SINGLELAYERWATER
+	// Fade out diffuse as this will be handled by the single scattering lighting. when over the water surface.
+	// We keep the SpecularColor for sun/water interactions
+	ShadingModelContext.WaterDiffuseIndirectLuminance = 0;
+	ShadingModelContext.DiffuseColor *= ShadingModelContext.BaseMaterialCoverageOverWater;
+#endif
+}
+
+half CalcSpecular(half Roughness, half RoughnessWithClamp, float RoL, float NoH, float3 H, float3 N)
+{
+#if USE_LEGACY_SPECULAR
+	return PhongApprox(Roughness, RoL);
+#else
+	return (Roughness*0.25 + 0.25) * GGX_Mobile(RoughnessWithClamp, NoH, H, N);
+#endif
+}
+
+FMobileDirectLighting MobileIntegrateBxDF(FMobileShadingModelContext ShadingModelContext, float NoL, float RoL, float3 V, float3 N, float3 H, float3 NoH)
+{
+	FMobileDirectLighting Lighting;
+#if MATERIAL_SHADINGMODEL_CLEAR_COAT
+	float VoH = max(0, dot(V, H));
+	half F0 = 0.04;
+	half Fc = Pow5(1 - VoH);
+	half F = Fc + (1 - Fc) * F0;
+	F *= ShadingModelContext.ClearCoat;
+	half LayerAttenuation = 1 - F;
+
+	// Vis_SmithJointApprox
+	half a = ShadingModelContext.ClearCoatRoughness * ShadingModelContext.ClearCoatRoughness;
+	float Vis_SmithV = NoL * (ShadingModelContext.NoV * (1 - a) + a);
+	float Vis_SmithL = ShadingModelContext.NoV * (NoL * (1 - a) + a);
+	float Vis = 0.5 * rcp(Vis_SmithV + Vis_SmithL);
+
+	Lighting.Specular = F * Vis * GGX_Mobile(ShadingModelContext.ClearCoatRoughness, NoH, H, N);
+	Lighting.Specular += LayerAttenuation * ShadingModelContext.SpecularColor * CalcSpecular(ShadingModelContext.Roughness, ShadingModelContext.RoughnessWithClamp, RoL, NoH, H, N);
+	Lighting.Diffuse = LayerAttenuation * ShadingModelContext.DiffuseColor;
+#else // MATERIAL_SHADINGMODEL_DEFAULT_LIT
+	Lighting.Specular = ShadingModelContext.SpecularColor * CalcSpecular(ShadingModelContext.Roughness, ShadingModelContext.RoughnessWithClamp, RoL, NoH, H, N);
+	Lighting.Diffuse = ShadingModelContext.DiffuseColor;
+#endif
+	return Lighting;
+}
\ No newline at end of file
