Some ClearCoat fix and cleanup.

diff --git a/Engine/Shaders/Private/MobileBasePassPixelShader.usf b/Engine/Shaders/Private/MobileBasePassPixelShader.usf
--- a/Engine/Shaders/Private/MobileBasePassPixelShader.usf
+++ b/Engine/Shaders/Private/MobileBasePassPixelShader.usf
@@ -592,21 +592,19 @@ void Main(
 		}
 	#endif
 	#if MATERIAL_SHADINGMODEL_CLEARCOAT
-		half SpecularOcclusion = 1; // GetSpecularOcclusion(ShadingModelContext.NoV, ShadingModelContext.Roughness*ShadingModelContext.Roughness, MaterialAO);
-		//Color += SpecularIBL * ShadingModelContext.SpecularColor * ShadingModelContext.ClearCoat * SpecularOcclusion;
 		float2 AB = PreIntegratedGF.SampleLevel(PreIntegratedGFSampler, float2(ShadingModelContext.NoV, ShadingModelContext.Roughness), 0).rg;
 
 		// F_Schlick
-		float F = EnvBRDF(0.04, ShadingModelContext.ClearCoatRoughness, ShadingModelContext.NoV).x;
-
+		half F = GetEnvBRDF(0.04, ShadingModelContext.ClearCoatRoughness, ShadingModelContext.NoV).x;
 		F *= ShadingModelContext.ClearCoat;
+		half LayerAttenuation = (1 - F);
 
-		float LayerAttenuation = (1 - F);
-
-		Color += SpecularIBL * ShadingModelContext.SpecularColor * LayerAttenuation * (ShadingModelContext.SpecularColor * AB.x + AB.y * saturate(50 * ShadingModelContext.SpecularColor.g) * (1 - ShadingModelContext.ClearCoat));
+		half DielectricSpecular = 0.08 * ShadingModelContext.Specular;
+		half SpecularColor = (DielectricSpecular - DielectricSpecular * ShadingModelContext.Metallic) + ShadingModelContext.BaseColor * ShadingModelContext.Metallic;	// 2 mad
+		Color += SpecularIBL * ShadingModelContext.SpecularColor * LayerAttenuation * (ShadingModelContext.SpecPreEnvBrdf * AB.x + AB.y * saturate(50 * ShadingModelContext.SpecPreEnvBrdf.g) * (1 - ShadingModelContext.ClearCoat));
 
 		SpecularIBL = GetImageBasedReflectionLighting(MaterialParameters, ShadingModelContext.ClearCoatRoughness, IndirectIrradiance);
-		Color += SpecularIBL * ShadingModelContext.SpecularColor * F * SpecularOcclusion;
+		Color += F * SpecularIBL * ShadingModelContext.SpecularColor;
 	#else
 		Color += SpecularIBL * ShadingModelContext.SpecularColor;
 	#endif
diff --git a/Engine/Shaders/Private/MobileShadingModels.ush b/Engine/Shaders/Private/MobileShadingModels.ush
--- a/Engine/Shaders/Private/MobileShadingModels.ush
+++ b/Engine/Shaders/Private/MobileShadingModels.ush
@@ -31,6 +31,7 @@ struct FMobileShadingModelContext
 	half ClearCoat;
 	half ClearCoatRoughness;
 	float NoV;
+	half SpecPreEnvBrdf;
 #elif MATERIAL_SHADINGMODEL_SINGLELAYERWATER
 	half BaseMaterialCoverageOverWater;
 	half WaterVisibility;
@@ -38,8 +39,6 @@ struct FMobileShadingModelContext
 #endif
 };
 
-
-
 #if NONMETAL
 half GetEnvBRDF(half SpecularColor, half Roughness, float NoV)
 {
@@ -91,6 +90,7 @@ void InitShadingModelContext( inout FMobileShadingModelContext ShadingModelConte
 		ShadingModelContext.Specular *= lerp(1, RefractionScale, ShadingModelContext.ClearCoat);
 		half DielectricSpecular = 0.08 * ShadingModelContext.Specular;
 		ShadingModelContext.SpecularColor = (DielectricSpecular - DielectricSpecular * ShadingModelContext.Metallic) + ShadingModelContext.BaseColor * ShadingModelContext.Metallic;	// 2 mad
+		ShadingModelContext.SpecPreEnvBrdf = ShadingModelContext.SpecularColor;
 		ShadingModelContext.DiffuseColor = ShadingModelContext.BaseColor - ShadingModelContext.BaseColor * ShadingModelContext.Metallic;
 
 		// Adjust BaseColor according to ClearCoat layer attenuation
@@ -173,6 +173,7 @@ FMobileDirectLighting MobileIntegrateBxDF(FMobileShadingModelContext ShadingMode
 	Lighting.Specular = F * Vis * GGX_Mobile(ShadingModelContext.ClearCoatRoughness, NoH, H, N);
 	Lighting.Specular += LayerAttenuation * ShadingModelContext.SpecularColor * CalcSpecular(ShadingModelContext.Roughness, ShadingModelContext.RoughnessWithClamp, RoL, NoH, H, N);
 	Lighting.Diffuse = LayerAttenuation * ShadingModelContext.DiffuseColor;
+
 #else // MATERIAL_SHADINGMODEL_DEFAULT_LIT
 	Lighting.Specular = ShadingModelContext.SpecularColor * CalcSpecular(ShadingModelContext.Roughness, ShadingModelContext.RoughnessWithClamp, RoL, NoH, H, N);
 	Lighting.Diffuse = ShadingModelContext.DiffuseColor;
