More fix for incorrect clearcoat IBL calculation.

diff --git a/Engine/Shaders/Private/MobileBasePassPixelShader.usf b/Engine/Shaders/Private/MobileBasePassPixelShader.usf
--- a/Engine/Shaders/Private/MobileBasePassPixelShader.usf
+++ b/Engine/Shaders/Private/MobileBasePassPixelShader.usf
@@ -592,19 +592,16 @@ void Main(
 		}
 	#endif
 	#if MATERIAL_SHADINGMODEL_CLEARCOAT
+		// Fc * Vis
 		float2 AB = PreIntegratedGF.SampleLevel(PreIntegratedGFSampler, float2(ShadingModelContext.NoV, ShadingModelContext.Roughness), 0).rg;
 
-		// F_Schlick
 		half F = GetEnvBRDF(0.04, ShadingModelContext.ClearCoatRoughness, ShadingModelContext.NoV).x;
 		F *= ShadingModelContext.ClearCoat;
 		half LayerAttenuation = (1 - F);
-
-		half DielectricSpecular = 0.08 * ShadingModelContext.Specular;
-		half SpecularColor = (DielectricSpecular - DielectricSpecular * ShadingModelContext.Metallic) + ShadingModelContext.BaseColor * ShadingModelContext.Metallic;	// 2 mad
-		Color += SpecularIBL * ShadingModelContext.SpecularColor * LayerAttenuation * (ShadingModelContext.SpecPreEnvBrdf * AB.x + AB.y * saturate(50 * ShadingModelContext.SpecPreEnvBrdf.g) * (1 - ShadingModelContext.ClearCoat));
+		Color += SpecularIBL * LayerAttenuation * (ShadingModelContext.SpecPreEnvBrdf * AB.x + AB.y * saturate(50 * ShadingModelContext.SpecPreEnvBrdf.g) * (1 - ShadingModelContext.ClearCoat));
 
 		SpecularIBL = GetImageBasedReflectionLighting(MaterialParameters, ShadingModelContext.ClearCoatRoughness, IndirectIrradiance);
-		Color += F * SpecularIBL * ShadingModelContext.SpecularColor;
+		Color += F * SpecularIBL;
 	#else
 		Color += SpecularIBL * ShadingModelContext.SpecularColor;
 	#endif
