More fix for clearcoat IBL calculation.

diff --git a/Engine/Shaders/Private/MobileBasePassPixelShader.usf b/Engine/Shaders/Private/MobileBasePassPixelShader.usf
--- a/Engine/Shaders/Private/MobileBasePassPixelShader.usf
+++ b/Engine/Shaders/Private/MobileBasePassPixelShader.usf
@@ -591,17 +591,22 @@ void Main(
 			SpecularIBL = lerp(SpecularIBL, PlanarReflection.rgb, PlanarReflection.a);
 		}
 	#endif
-	#if MATERIAL_SHADINGMODEL_CLEARCOAT
-		// Fc * Vis
-		float2 AB = PreIntegratedGF.SampleLevel(PreIntegratedGFSampler, float2(ShadingModelContext.NoV, ShadingModelContext.Roughness), 0).rg;
-
+	#if MATERIAL_SHADINGMODEL_CLEAR_COAT
 		half F = GetEnvBRDF(0.04, ShadingModelContext.ClearCoatRoughness, ShadingModelContext.NoV).x;
 		F *= ShadingModelContext.ClearCoat;
 		half LayerAttenuation = (1 - F);
+
+		// Fc * Vis
+		#if FORWARDSHADING_USE_HQ_ENV_BRDF
+			half2 AB = PreIntegratedGF.SampleLevel(PreIntegratedGFSampler, float2(ShadingModelContext.NoV, ShadingModelContext.Roughness), 0).rg;
+		#else
+			half3 Fc = (ShadingModelContext.SpecularColor - ShadingModelContext.SpecPreEnvBrdf) * rcp(saturate(50 * ShadingModelContext.SpecPreEnvBrdf.g) - ShadingModelContext.SpecPreEnvBrdf);
+			half2 AB = half2(1 - Fc.x, Fc.x);
+		#endif
 		Color += SpecularIBL * LayerAttenuation * (ShadingModelContext.SpecPreEnvBrdf * AB.x + AB.y * saturate(50 * ShadingModelContext.SpecPreEnvBrdf.g) * (1 - ShadingModelContext.ClearCoat));
 
 		SpecularIBL = GetImageBasedReflectionLighting(MaterialParameters, ShadingModelContext.ClearCoatRoughness, IndirectIrradiance);
-		Color += F * SpecularIBL;
+		Color += SpecularIBL * F;
 	#else
 		Color += SpecularIBL * ShadingModelContext.SpecularColor;
 	#endif
diff --git a/Engine/Shaders/Private/MobileShadingModels.ush b/Engine/Shaders/Private/MobileShadingModels.ush
--- a/Engine/Shaders/Private/MobileShadingModels.ush
+++ b/Engine/Shaders/Private/MobileShadingModels.ush
@@ -31,7 +31,7 @@ struct FMobileShadingModelContext
 	half ClearCoat;
 	half ClearCoatRoughness;
 	float NoV;
-	half SpecPreEnvBrdf;
+	half3 SpecPreEnvBrdf;
 #elif MATERIAL_SHADINGMODEL_SINGLELAYERWATER
 	half BaseMaterialCoverageOverWater;
 	half WaterVisibility;
