Fix for shader compile error "error C1020: invalid operands to "/""
An integer argument to rcp() is not being automatically cast to a float.  This adds a manual cast (same as the line above).

Partial fix for:
#rnx

diff --git a/Engine/Shaders/Private/DiaphragmDOF/DOFBokehLUT.usf b/Engine/Shaders/Private/DiaphragmDOF/DOFBokehLUT.usf
--- a/Engine/Shaders/Private/DiaphragmDOF/DOFBokehLUT.usf
+++ b/Engine/Shaders/Private/DiaphragmDOF/DOFBokehLUT.usf
@@ -77,7 +77,7 @@ float DiaphragmBladeCenterOffset;
 float ComputeCocRadiusToBokehEdgeFactor(float PixelAngle)
 {
 	float EdgeId = floor(PixelAngle * rcp(2 * PI) * float(BladeCount));
-	float EdgeAngle = (EdgeId + 0.5) * (2 * PI) * rcp(BladeCount);
+	float EdgeAngle = (EdgeId + 0.5) * (2 * PI) * rcp(float(BladeCount));
 	
 	float Alpha = PixelAngle - EdgeAngle;
 
