Fix shader issue related to hair rendering.

diff --git a/Engine/Shaders/Private/HairStrands/HairStrandsAABB.usf b/Engine/Shaders/Private/HairStrands/HairStrandsAABB.usf
--- a/Engine/Shaders/Private/HairStrands/HairStrandsAABB.usf
+++ b/Engine/Shaders/Private/HairStrands/HairStrandsAABB.usf
@@ -1,5 +1,6 @@
 // Copyright Epic Games, Inc. All Rights Reserved.
 
+#include "../Common.ush"
 #include "HairStrandsAABBCommon.ush"
 
 #if SHADER_AABBUPDATE
diff --git a/Engine/Shaders/Private/HairStrands/HairStrandsClusterCulling.usf b/Engine/Shaders/Private/HairStrands/HairStrandsClusterCulling.usf
--- a/Engine/Shaders/Private/HairStrands/HairStrandsClusterCulling.usf
+++ b/Engine/Shaders/Private/HairStrands/HairStrandsClusterCulling.usf
@@ -1,5 +1,6 @@
 // Copyright Epic Games, Inc. All Rights Reserved.
 
+#include "../Common.ush"
 #include "HairStrandsDebugCommon.ush"
 #include "HairStrandsAABBCommon.ush"
 #define COMPACTCS_NUMTHREAD 64
diff --git a/Engine/Shaders/Private/HairStrands/HairStrandsDeepTransmittanceMask.usf b/Engine/Shaders/Private/HairStrands/HairStrandsDeepTransmittanceMask.usf
--- a/Engine/Shaders/Private/HairStrands/HairStrandsDeepTransmittanceMask.usf
+++ b/Engine/Shaders/Private/HairStrands/HairStrandsDeepTransmittanceMask.usf
@@ -60,12 +60,12 @@ float				LightRadius;
 float				DeepShadow_KernelAperture;
 uint				DeepShadow_KernelType;
 uint				DeepShadow_DebugMode;
-int2				DeepShadow_AtlasSlotOffsets[MAX_HAIR_MACROGROUP_COUNT];
+int4				DeepShadow_AtlasSlotOffsets[MAX_HAIR_MACROGROUP_COUNT];
 Texture2D<float>	DeepShadow_FrontDepthTexture;
 Texture2D<float4>	DeepShadow_DomTexture;
 
-float3				Voxel_MinAABBs[MAX_HAIR_MACROGROUP_COUNT];
-float3				Voxel_MaxAABBs[MAX_HAIR_MACROGROUP_COUNT];
+float4				Voxel_MinAABBs[MAX_HAIR_MACROGROUP_COUNT];
+float4				Voxel_MaxAABBs[MAX_HAIR_MACROGROUP_COUNT];
 uint				Voxel_Resolution;
 Texture3D<uint> 	Voxel_DensityTexture0;
 Texture3D<uint> 	Voxel_DensityTexture1;
diff --git a/Engine/Shaders/Private/HairStrands/HairStrandsVoxelPageAllocation.usf b/Engine/Shaders/Private/HairStrands/HairStrandsVoxelPageAllocation.usf
--- a/Engine/Shaders/Private/HairStrands/HairStrandsVoxelPageAllocation.usf
+++ b/Engine/Shaders/Private/HairStrands/HairStrandsVoxelPageAllocation.usf
@@ -403,7 +403,7 @@ void VoxelIndPageClearBufferGenCS(uint GroupIndex : SV_GroupIndex, uint3 Dispatc
 #if SHADER_INDPAGECLEAR
 
 uint				VirtualVoxel_PageResolution;
-uint				VirtualVoxel_PageCountResolution;
+int3				VirtualVoxel_PageCountResolution;
 
 Buffer<uint4>		VirtualVoxel_PageIndexCoordBuffer;
 
