Copying up //Dev-Mobile @ CL 11191662

diff --git a/Engine/Plugins/Runtime/AndroidDeviceProfileSelector/AndroidDeviceProfileSelector.uplugin b/Engine/Plugins/Runtime/AndroidDeviceProfileSelector/AndroidDeviceProfileSelector.uplugin
--- a/Engine/Plugins/Runtime/AndroidDeviceProfileSelector/AndroidDeviceProfileSelector.uplugin
+++ b/Engine/Plugins/Runtime/AndroidDeviceProfileSelector/AndroidDeviceProfileSelector.uplugin
@@ -19,7 +19,7 @@
 		{
 			"Name" : "AndroidDeviceProfileSelector",
 			"Type": "Editor",
-			"LoadingPhase" : "PostConfigInit"
+			"LoadingPhase" : "None"
 		},
 		{
 			"Name" : "AndroidDeviceProfileSelectorRuntime",
diff --git a/Engine/Shaders/Private/MaterialTemplate.ush b/Engine/Shaders/Private/MaterialTemplate.ush
--- a/Engine/Shaders/Private/MaterialTemplate.ush
+++ b/Engine/Shaders/Private/MaterialTemplate.ush
@@ -1900,30 +1900,9 @@ half GetMaterialRoughness(FPixelMaterialInputs PixelMaterialInputs)
 #if MATERIAL_FULLY_ROUGH
 	return 1;
 #endif
-#if (ES2_PROFILE || ES3_1_PROFILE)
-	// The smallest normalized value that can be represented in IEEE 754 (FP16) is 2^-14 = 6.1e-5.
-	// The code will make the following computation involving roughness: 1.0 / Roughness^4.
-	// Therefore to prevent division by zero on devices that do not support denormals, Roughness^4
-	// must be >= 6.1e-5. We will clamp to 0.09 because 0.09^4 = 6.5e-5.
-	//
-	// Note that we also clamp to 1.0 to match the deferred renderer on PC where the roughness is 
-	// stored in an 8-bit value and thus automatically clamped at 1.0.
-
-	half Roughness = GetMaterialRoughnessRaw(PixelMaterialInputs);
-
-	#if MOBILE_EMULATION
-	{
-		// this is only needed for mobile preview on PC
-		Roughness = Roughness * ResolvedView.RoughnessOverrideParameter.y + ResolvedView.RoughnessOverrideParameter.x;
-	}
-	#endif
-
-	// Increase value from 0.09 to 0.12 to fix missing specular lobe problem on device
-	return clamp( Roughness, 0.12, 1.0 );
-#else
 	half Roughness = saturate(GetMaterialRoughnessRaw(PixelMaterialInputs));
 	
-	#if (USE_EDITOR_SHADERS && !ESDEFERRED_PROFILE)
+	#if (USE_EDITOR_SHADERS && !ESDEFERRED_PROFILE) || MOBILE_EMULATION
 	{
 		// this feature is only needed for development/editor - we can compile it out for a shipping build (see r.CompileShadersForDevelopment cvar help)
 		Roughness = Roughness * ResolvedView.RoughnessOverrideParameter.y + ResolvedView.RoughnessOverrideParameter.x;
@@ -1931,7 +1910,6 @@ half GetMaterialRoughness(FPixelMaterialInputs PixelMaterialInputs)
 	#endif
 	
 	return Roughness;
-#endif
 }
 
 half GetMaterialTranslucencyDirectionalLightingIntensity()
diff --git a/Engine/Shaders/Private/MobileBasePassPixelShader.usf b/Engine/Shaders/Private/MobileBasePassPixelShader.usf
--- a/Engine/Shaders/Private/MobileBasePassPixelShader.usf
+++ b/Engine/Shaders/Private/MobileBasePassPixelShader.usf
@@ -27,12 +27,10 @@
 #include "/Engine/Generated/VertexFactory.ush"
 #include "ReflectionEnvironmentShared.ush"
 #include "LightmapCommon.ush"  
-#include "BRDF.ush"
-#include "SHCommon.ush"
+#include "MobileShadingModels.ush"
 #include "ShadowFilteringCommon.ush"
 #include "DynamicLightingCommon.ush"
 #include "PlanarReflectionShared.ush"
-#include "MobileGGX.ush"
 
 #if MATERIAL_SHADINGMODEL_SINGLELAYERWATER
 #define SIMPLE_SINGLE_LAYER_WATER 1
@@ -127,15 +125,6 @@ half PhongApprox( half Roughness, half RoL )
 }
 
 
-half CalcSpecular(half Roughness, float RoL, float NoH, float3 H, float3 N)
-{
-#if USE_LEGACY_SPECULAR
-	return PhongApprox(Roughness, RoL);
-#else
-	return (Roughness*0.25 + 0.25) * GGX_Mobile(Roughness, NoH, H, N);
-#endif
-}
-
 #if !FULLY_ROUGH
 #if HQ_REFLECTIONS
 
@@ -461,306 +450,283 @@ void Main(
 	//Clip if the blend mode requires it.
 	GetMaterialCoverageAndClipping(MaterialParameters, PixelMaterialInputs);
 
-	half Opacity = GetMaterialOpacity(PixelMaterialInputs);
-	const float BaseMaterialCoverageOverWater = Opacity;
-	const float WaterVisibility = 1.0 - BaseMaterialCoverageOverWater;
-
 	// Store the results in local variables and reuse instead of calling the functions multiple times.
-	// Store the results in local variables and reuse instead of calling the functions multiple times.
-	half3 BaseColor = GetMaterialBaseColor(PixelMaterialInputs);
-	half  Metallic = GetMaterialMetallic(PixelMaterialInputs);
-	half  Specular = GetMaterialSpecular(PixelMaterialInputs);
-
-#if NONMETAL
-	half3 DiffuseColor = BaseColor;
-	half SpecularColor = 0.04;
-#else
-	half DielectricSpecular = 0.08 * Specular;
-	half3 DiffuseColor = BaseColor - BaseColor * Metallic;	// 1 mad
-	half3 SpecularColor = (DielectricSpecular - DielectricSpecular * Metallic) + BaseColor * Metallic;	// 2 mad
-#endif
+	FMobileShadingModelContext ShadingModelContext;
+	ShadingModelContext.Opacity = GetMaterialOpacity(PixelMaterialInputs);
+	ShadingModelContext.BaseColor = GetMaterialBaseColor(PixelMaterialInputs);
+	ShadingModelContext.Metallic = GetMaterialMetallic(PixelMaterialInputs);
+	ShadingModelContext.Specular = GetMaterialSpecular(PixelMaterialInputs);
+	ShadingModelContext.Roughness = GetMaterialRoughness(PixelMaterialInputs);
+	// The smallest normalized value that can be represented in IEEE 754 (FP16) is 2^-14 = 6.1e-5.
+	// The code will make the following computation involving roughness: 1.0 / Roughness^4.
+	// Therefore to prevent division by zero on devices that do not support denormals, Roughness^4
+	// must be >= 6.1e-5. We will clamp to 0.09 because 0.09^4 = 6.5e-5.
+	// Increase value from 0.09 to 0.12 to fix missing specular lobe problem on device
+	//
+	// Note that we also clamp to 1.0 to match the deferred renderer on PC where the roughness is 
+	// stored in an 8-bit value and thus automatically clamped at 1.0.
+	// To help with IBL cube sample, we postpone the clamp only when calculate GGX_Mobile.
+	ShadingModelContext.RoughnessWithClamp = max(0.12, ShadingModelContext.Roughness);
+
+	InitShadingModelContext(ShadingModelContext, MaterialParameters);
 
-	half Roughness = GetMaterialRoughness(PixelMaterialInputs);
-
-#if FULLY_ROUGH
-	EnvBRDFApproxFullyRough(DiffuseColor, SpecularColor);
-#else
-	half NoV = max( dot( MaterialParameters.WorldNormal, MaterialParameters.CameraVector ), 0 );
-
-#if NONMETAL
-	// If nothing is hooked up to Metalic and Specular,
-	// then defaults are the same as a non-metal,
-	// so this define is safe.
-	SpecularColor = EnvBRDFApproxNonmetal( Roughness, NoV );
-#else
-#if FORWARDSHADING_USE_HQ_ENV_BRDF
-	SpecularColor = EnvBRDF(SpecularColor, Roughness, NoV);
-#else
-	SpecularColor = EnvBRDFApprox( SpecularColor, Roughness, NoV );
-#endif
-#endif
-#endif
-
-#if MOBILE_EMULATION
-	{
-		// this feature is only needed for development/editor - we can compile it out for a shipping build (see r.CompileShadersForDevelopment cvar help)
-		DiffuseColor = DiffuseColor * ResolvedView.DiffuseOverrideParameter.w + ResolvedView.DiffuseOverrideParameter.xyz;
-#if NONMETAL
-    SpecularColor = SpecularColor * ResolvedView.SpecularOverrideParameter.w + ResolvedView.SpecularOverrideParameter.x;
-#else
-    SpecularColor = SpecularColor * ResolvedView.SpecularOverrideParameter.w + ResolvedView.SpecularOverrideParameter.xyz;
-#endif
-	}
-#endif
-	
-#if MATERIAL_SHADINGMODEL_SINGLELAYERWATER
-	// Fade out diffuse as this will be hanlded by the single scattering lighting. when over the water surface.
-	// We keep the SpecularColor for sun/water interactions
-	float3 WaterDiffuseIndirectLuminance = 0;
-	DiffuseColor *= BaseMaterialCoverageOverWater;
-#endif
-	
 	half3 Color = 0;
+
+#if !MATERIAL_SHADINGMODEL_UNLIT
 	//To keep IndirectLightingCache conherence with PC, initialize the IndirectIrradiance to zero.
 	half IndirectIrradiance = 0;
-
-	#if LQ_TEXTURE_LIGHTMAP
-		float2 LightmapUV0, LightmapUV1;
-		uint LightmapDataIndex;
-		GetLightMapCoordinates(Interpolants, LightmapUV0, LightmapUV1, LightmapDataIndex);
-
-		half4 LightmapColor = GetLightMapColorLQ( LightmapUV0, LightmapUV1, LightmapDataIndex, MaterialParameters.WorldNormal );
-		Color += LightmapColor.rgb * DiffuseColor * View.IndirectLightingColorScale;
-		IndirectIrradiance = LightmapColor.a;
-	#elif CACHED_POINT_INDIRECT_LIGHTING
-		#if MATERIALBLENDING_MASKED || MATERIALBLENDING_SOLID
-
-			// Take the normal into account for opaque
-			FThreeBandSHVectorRGB PointIndirectLighting;
-			PointIndirectLighting.R.V0 = IndirectLightingCache.IndirectLightingSHCoefficients0[0];
-			PointIndirectLighting.R.V1 = IndirectLightingCache.IndirectLightingSHCoefficients1[0];
-			PointIndirectLighting.R.V2 = IndirectLightingCache.IndirectLightingSHCoefficients2[0];
-
-			PointIndirectLighting.G.V0 = IndirectLightingCache.IndirectLightingSHCoefficients0[1];
-			PointIndirectLighting.G.V1 = IndirectLightingCache.IndirectLightingSHCoefficients1[1];
-			PointIndirectLighting.G.V2 = IndirectLightingCache.IndirectLightingSHCoefficients2[1];
-
-			PointIndirectLighting.B.V0 = IndirectLightingCache.IndirectLightingSHCoefficients0[2];
-			PointIndirectLighting.B.V1 = IndirectLightingCache.IndirectLightingSHCoefficients1[2];
-			PointIndirectLighting.B.V2 = IndirectLightingCache.IndirectLightingSHCoefficients2[2];
-
-			FThreeBandSHVector DiffuseTransferSH = CalcDiffuseTransferSH3(MaterialParameters.WorldNormal, 1);
-
-			// Compute diffuse lighting which takes the normal into account
-			half3 DiffuseGI = max(half3(0,0,0), DotSH3(PointIndirectLighting, DiffuseTransferSH));
-
-			IndirectIrradiance = Luminance(DiffuseGI);
-			Color += DiffuseColor * DiffuseGI * View.IndirectLightingColorScale;			
-		#else 
-			
-			// Non-directional for translucency
-			// Ambient terms packed in xyz
-			// Already divided by PI and SH ambient on CPU
-			half3 PointIndirectLighting = IndirectLightingCache.IndirectLightingSHSingleCoefficient;
-			half3 DiffuseGI = PointIndirectLighting;
-
-			IndirectIrradiance = Luminance(DiffuseGI);
-			Color += DiffuseColor * DiffuseGI * View.IndirectLightingColorScale;
-
-		#endif
-	#endif
-
-	#if !MATERIAL_SHADINGMODEL_UNLIT
-		half MaterialAO = GetMaterialAmbientOcclusion(PixelMaterialInputs);
-		Color *= MaterialAO;	
-		IndirectIrradiance *= MaterialAO;
+	// Indirect Diffuse
+#if LQ_TEXTURE_LIGHTMAP
+	float2 LightmapUV0, LightmapUV1;
+	uint LightmapDataIndex;
+	GetLightMapCoordinates(Interpolants, LightmapUV0, LightmapUV1, LightmapDataIndex);
+
+	half4 LightmapColor = GetLightMapColorLQ(LightmapUV0, LightmapUV1, LightmapDataIndex, MaterialParameters.WorldNormal);
+	Color += LightmapColor.rgb * ShadingModelContext.DiffuseColor * View.IndirectLightingColorScale;
+	IndirectIrradiance = LightmapColor.a;
+#elif CACHED_POINT_INDIRECT_LIGHTING
+	#if MATERIALBLENDING_MASKED || MATERIALBLENDING_SOLID
+		// Take the normal into account for opaque
+		FThreeBandSHVectorRGB PointIndirectLighting;
+		PointIndirectLighting.R.V0 = IndirectLightingCache.IndirectLightingSHCoefficients0[0];
+		PointIndirectLighting.R.V1 = IndirectLightingCache.IndirectLightingSHCoefficients1[0];
+		PointIndirectLighting.R.V2 = IndirectLightingCache.IndirectLightingSHCoefficients2[0];
+
+		PointIndirectLighting.G.V0 = IndirectLightingCache.IndirectLightingSHCoefficients0[1];
+		PointIndirectLighting.G.V1 = IndirectLightingCache.IndirectLightingSHCoefficients1[1];
+		PointIndirectLighting.G.V2 = IndirectLightingCache.IndirectLightingSHCoefficients2[1];
+
+		PointIndirectLighting.B.V0 = IndirectLightingCache.IndirectLightingSHCoefficients0[2];
+		PointIndirectLighting.B.V1 = IndirectLightingCache.IndirectLightingSHCoefficients1[2];
+		PointIndirectLighting.B.V2 = IndirectLightingCache.IndirectLightingSHCoefficients2[2];
+
+		FThreeBandSHVector DiffuseTransferSH = CalcDiffuseTransferSH3(MaterialParameters.WorldNormal, 1);
+
+		// Compute diffuse lighting which takes the normal into account
+		half3 DiffuseGI = max(half3(0, 0, 0), DotSH3(PointIndirectLighting, DiffuseTransferSH));
+
+		IndirectIrradiance = Luminance(DiffuseGI);
+		Color += ShadingModelContext.DiffuseColor * DiffuseGI * View.IndirectLightingColorScale;
+	#else 
+		// Non-directional for translucency
+		// Ambient terms packed in xyz
+		// Already divided by PI and SH ambient on CPU
+		half3 PointIndirectLighting = IndirectLightingCache.IndirectLightingSHSingleCoefficient;
+		half3 DiffuseGI = PointIndirectLighting;
+
+		IndirectIrradiance = Luminance(DiffuseGI);
+		Color += ShadingModelContext.DiffuseColor * DiffuseGI * View.IndirectLightingColorScale;
 	#endif
+#endif
+	half MaterialAO = GetMaterialAmbientOcclusion(PixelMaterialInputs);
+	Color *= MaterialAO;
+	IndirectIrradiance *= MaterialAO;
 
-		half Shadow = 1.0f;
-	#if !MATERIAL_SHADINGMODEL_UNLIT
-		Shadow = GetPrimaryPrecomputedShadowMask(Interpolants).r;
-		#if DIRECTIONAL_LIGHT_CSM && !MATERIAL_SHADINGMODEL_SINGLELAYERWATER
-		// Cascaded Shadow Map
+	// Shadow
+	half Shadow = GetPrimaryPrecomputedShadowMask(Interpolants).r;
+#if DIRECTIONAL_LIGHT_CSM && !MATERIAL_SHADINGMODEL_SINGLELAYERWATER
+	// Cascaded Shadow Map
+	{
+		FPCFSamplerSettings Settings;
+		Settings.ShadowDepthTexture = MobileDirectionalLight.DirectionalLightShadowTexture;
+		Settings.ShadowDepthTextureSampler = MobileDirectionalLight.DirectionalLightShadowSampler;
+		Settings.TransitionScale = MobileDirectionalLight.DirectionalLightDirectionAndShadowTransition.w;
+		Settings.ShadowBufferSize = MobileDirectionalLight.DirectionalLightShadowSize;
+		Settings.bSubsurface = false;
+		Settings.bTreatMaxDepthUnshadowed = false;
+		Settings.DensityMulConstant = 0;
+		Settings.ProjectionDepthBiasParameters = 0;
+
+		float4 ShadowPosition = float4(0, 0, 0, 0);
+		for (int i = 0; i < MAX_MOBILE_SHADOWCASCADES; i++)
 		{
-			FPCFSamplerSettings Settings;
-			Settings.ShadowDepthTexture = MobileDirectionalLight.DirectionalLightShadowTexture;
-			Settings.ShadowDepthTextureSampler = MobileDirectionalLight.DirectionalLightShadowSampler;
-			Settings.TransitionScale = MobileDirectionalLight.DirectionalLightDirectionAndShadowTransition.w;
-			Settings.ShadowBufferSize = MobileDirectionalLight.DirectionalLightShadowSize;
-			Settings.bSubsurface = false;
-			Settings.bTreatMaxDepthUnshadowed = false;
-			Settings.DensityMulConstant = 0;
-			Settings.ProjectionDepthBiasParameters = 0;
-
-			float4 ShadowPosition = float4(0,0,0,0);
-			for (int i = 0; i < MAX_MOBILE_SHADOWCASCADES; i++)
-			{
-				if (MaterialParameters.ScreenPosition.w < MobileDirectionalLight.DirectionalLightShadowDistances[i])
-				{
-				#if MOBILE_MULTI_VIEW
-					ShadowPosition = mul(float4(MaterialParameters.ScreenPosition.xyw, 1), ResolvedView.MobileMultiviewShadowTransform);
-					ShadowPosition = mul(ShadowPosition, MobileDirectionalLight.DirectionalLightScreenToShadow[i]);
-				#else
-					ShadowPosition = mul(float4(MaterialParameters.ScreenPosition.xyw, 1), MobileDirectionalLight.DirectionalLightScreenToShadow[i]);
-				#endif
-					break; // position found.
-				}
-			}
-
-			// Process CSM only when ShadowPosition is valid.
-			if (ShadowPosition.z > 0)
+			if (MaterialParameters.ScreenPosition.w < MobileDirectionalLight.DirectionalLightShadowDistances[i])
 			{
-				// Clamp pixel depth in light space for shadowing opaque, because areas of the shadow depth buffer that weren't rendered to will have been cleared to 1
-				// We want to force the shadow comparison to result in 'unshadowed' in that case, regardless of whether the pixel being shaded is in front or behind that plane
-				float LightSpacePixelDepthForOpaque = min(ShadowPosition.z, 0.99999f);
-				Settings.SceneDepth = LightSpacePixelDepthForOpaque;
-
-				#if MOBILE_CSM_QUALITY == 0
-					half ShadowMap = ManualNoFiltering(ShadowPosition.xy, Settings);
-				#elif MOBILE_CSM_QUALITY == 1
-					half ShadowMap = Manual1x1PCF(ShadowPosition.xy, Settings);
-				#elif MOBILE_CSM_QUALITY == 2
-					half ShadowMap = Manual2x2PCF(ShadowPosition.xy, Settings);
-				#else
-					#error Unsupported MOBILE_CSM_QUALITY value.
-				#endif
-
-				#if FADE_CSM
-					float Fade = saturate(MaterialParameters.ScreenPosition.w * MobileDirectionalLight.DirectionalLightDistanceFadeMADAndSpecularScale.x + MobileDirectionalLight.DirectionalLightDistanceFadeMADAndSpecularScale.y);
-					// lerp out shadow based on fade params.
-					ShadowMap = lerp(ShadowMap, 1.0, Fade * Fade);
-				#endif
-
-				#if MOVABLE_DIRECTIONAL_LIGHT
-					Shadow = ShadowMap;
-				#else
-					Shadow = min(ShadowMap, Shadow);
-				#endif
+			#if MOBILE_MULTI_VIEW
+				ShadowPosition = mul(float4(MaterialParameters.ScreenPosition.xyw, 1), ResolvedView.MobileMultiviewShadowTransform);
+				ShadowPosition = mul(ShadowPosition, MobileDirectionalLight.DirectionalLightScreenToShadow[i]);
+			#else
+				ShadowPosition = mul(float4(MaterialParameters.ScreenPosition.xyw, 1), MobileDirectionalLight.DirectionalLightScreenToShadow[i]);
+			#endif
+				break; // position found.
 			}
 		}
-		#endif /* DIRECTIONAL_LIGHT_CSM */
 
-		float NoL = max(0, dot(float3(MaterialParameters.WorldNormal), float3(MobileDirectionalLight.DirectionalLightDirectionAndShadowTransition.xyz)));
-		float RoL = max(0, dot(float3(MaterialParameters.ReflectionVector), float3(MobileDirectionalLight.DirectionalLightDirectionAndShadowTransition.xyz)));
-		float3 H = normalize(float3(MaterialParameters.CameraVector) + float3(MobileDirectionalLight.DirectionalLightDirectionAndShadowTransition.xyz));
-		float NoH = max(0,dot(MaterialParameters.WorldNormal, H));
+		// Process CSM only when ShadowPosition is valid.
+		if (ShadowPosition.z > 0)
+		{
+			// Clamp pixel depth in light space for shadowing opaque, because areas of the shadow depth buffer that weren't rendered to will have been cleared to 1
+			// We want to force the shadow comparison to result in 'unshadowed' in that case, regardless of whether the pixel being shaded is in front or behind that plane
+			float LightSpacePixelDepthForOpaque = min(ShadowPosition.z, 0.99999f);
+			Settings.SceneDepth = LightSpacePixelDepthForOpaque;
+
+			#if MOBILE_CSM_QUALITY == 0
+				half ShadowMap = ManualNoFiltering(ShadowPosition.xy, Settings);
+			#elif MOBILE_CSM_QUALITY == 1
+				half ShadowMap = Manual1x1PCF(ShadowPosition.xy, Settings);
+			#elif MOBILE_CSM_QUALITY == 2
+				half ShadowMap = Manual2x2PCF(ShadowPosition.xy, Settings);
+			#else
+				#error Unsupported MOBILE_CSM_QUALITY value.
+			#endif
 
-		#if FULLY_ROUGH
-			Color += (Shadow * NoL) * MobileDirectionalLight.DirectionalLightColor.rgb * DiffuseColor;
-		#else
-			// MobileDirectionalLight.DirectionalLightDistanceFadeMADAndSpecularScale.z saves SpecularScale for direction light.
-			Color += (Shadow * NoL) * MobileDirectionalLight.DirectionalLightColor.rgb * (DiffuseColor + SpecularColor * MobileDirectionalLight.DirectionalLightDistanceFadeMADAndSpecularScale.z * CalcSpecular(Roughness, RoL, NoH, H, MaterialParameters.WorldNormal));
-			// Environment map has been prenormalized, scale by lightmap luminance
-			half3 SpecularIBL = GetImageBasedReflectionLighting(MaterialParameters, Roughness, IndirectIrradiance);
-
-			#if MATERIAL_PLANAR_FORWARD_REFLECTIONS
-				BRANCH
-				if (abs(dot(PlanarReflectionStruct.ReflectionPlane.xyz, 1)) > .0001f)
-				{
-					half4 PlanarReflection = GetPlanarReflection(MaterialParameters.AbsoluteWorldPosition, MaterialParameters.WorldNormal, Roughness);
-					// Planar reflections win over reflection environment
-					SpecularIBL = lerp(SpecularIBL, PlanarReflection.rgb, PlanarReflection.a);
-				}
+			#if FADE_CSM
+				float Fade = saturate(MaterialParameters.ScreenPosition.w * MobileDirectionalLight.DirectionalLightDistanceFadeMADAndSpecularScale.x + MobileDirectionalLight.DirectionalLightDistanceFadeMADAndSpecularScale.y);
+			// lerp out shadow based on fade params.
+				ShadowMap = lerp(ShadowMap, 1.0, Fade * Fade);
 			#endif
-			Color += SpecularIBL * SpecularColor;
-		#endif
 
-		#if ENABLE_SKY_LIGHT
-			//@mw todo
-			// TODO: Also need to do specular.
-#if MATERIAL_TWOSIDED && LQ_TEXTURE_LIGHTMAP
-			if (NoL == 0)
-			{
-#endif
-			half3 SkyDiffuseLighting = GetSkySHDiffuseSimple(MaterialParameters.WorldNormal);
+			#if MOVABLE_DIRECTIONAL_LIGHT
+				Shadow = ShadowMap;
+			#else
+				Shadow = min(ShadowMap, Shadow);
+			#endif
+		}
+	}
+#endif /* DIRECTIONAL_LIGHT_CSM */
 
-		#if MATERIAL_SHADINGMODEL_SINGLELAYERWATER
-			WaterDiffuseIndirectLuminance += SkyDiffuseLighting;
-		#endif
-			Color += SkyDiffuseLighting * half3(ResolvedView.SkyLightColor.rgb) * DiffuseColor * MaterialAO;
-#if MATERIAL_TWOSIDED && LQ_TEXTURE_LIGHTMAP
-			
-			}
-#endif
+	float NoL = max(0, dot(float3(MaterialParameters.WorldNormal), float3(MobileDirectionalLight.DirectionalLightDirectionAndShadowTransition.xyz)));
+	float RoL = max(0, dot(float3(MaterialParameters.ReflectionVector), float3(MobileDirectionalLight.DirectionalLightDirectionAndShadowTransition.xyz)));
+	float3 H = normalize(float3(MaterialParameters.CameraVector) + float3(MobileDirectionalLight.DirectionalLightDirectionAndShadowTransition.xyz));
+	float NoH = max(0, dot(MaterialParameters.WorldNormal, H));
+
+	// Direct Lighting (Directional light) + IBL
+#if FULLY_ROUGH
+	Color += (Shadow * NoL) * MobileDirectionalLight.DirectionalLightColor.rgb * ShadingModelContext.DiffuseColor;
+#else
+	FMobileDirectLighting Lighting = MobileIntegrateBxDF(ShadingModelContext, NoL, RoL, MaterialParameters.CameraVector, MaterialParameters.WorldNormal, H, NoH);
+	// MobileDirectionalLight.DirectionalLightDistanceFadeMADAndSpecularScale.z saves SpecularScale for direction light.
+	Color += (Shadow * NoL) * MobileDirectionalLight.DirectionalLightColor.rgb * (Lighting.Diffuse + Lighting.Specular * MobileDirectionalLight.DirectionalLightDistanceFadeMADAndSpecularScale.z);
+
+	// Environment map has been prenormalized, scale by lightmap luminance
+	half3 SpecularIBL = GetImageBasedReflectionLighting(MaterialParameters, ShadingModelContext.Roughness, IndirectIrradiance);
+	#if MATERIAL_PLANAR_FORWARD_REFLECTIONS
+		BRANCH
+		if (abs(dot(PlanarReflectionStruct.ReflectionPlane.xyz, 1)) > .0001f)
+		{
+			half4 PlanarReflection = GetPlanarReflection(MaterialParameters.AbsoluteWorldPosition, MaterialParameters.WorldNormal, ShadingModelContext.Roughness);
+			// Planar reflections win over reflection environment
+			SpecularIBL = lerp(SpecularIBL, PlanarReflection.rgb, PlanarReflection.a);
+		}
+	#endif
+	#if MATERIAL_SHADINGMODEL_CLEAR_COAT
+		half F = GetEnvBRDF(0.04, ShadingModelContext.ClearCoatRoughness, ShadingModelContext.NoV).x;
+		F *= ShadingModelContext.ClearCoat;
+		half LayerAttenuation = (1 - F);
+
+		// Fc * Vis
+		#if FORWARDSHADING_USE_HQ_ENV_BRDF
+			half2 AB = PreIntegratedGF.SampleLevel(PreIntegratedGFSampler, float2(ShadingModelContext.NoV, ShadingModelContext.Roughness), 0).rg;
+		#else
+			half3 Fc = (ShadingModelContext.SpecularColor - ShadingModelContext.SpecPreEnvBrdf) * rcp(saturate(50 * ShadingModelContext.SpecPreEnvBrdf.g) - ShadingModelContext.SpecPreEnvBrdf);
+			half2 AB = half2(1 - Fc.x, Fc.x);
 		#endif
+		Color += SpecularIBL * LayerAttenuation * (ShadingModelContext.SpecPreEnvBrdf * AB.x + AB.y * saturate(50 * ShadingModelContext.SpecPreEnvBrdf.g) * (1 - ShadingModelContext.ClearCoat));
 
-		#if MAX_DYNAMIC_POINT_LIGHTS > 0 && !MATERIAL_SHADINGMODEL_SINGLELAYERWATER
-			#if VARIABLE_NUM_DYNAMIC_POINT_LIGHTS
-			#if WEBGL || ES2_PROFILE || VULKAN_PROFILE
-			// webgl and ES2 needs a constant loop counter
-			UNROLL
-			for (int i = 0; i < MAX_DYNAMIC_POINT_LIGHTS; i++)
-			{
-				if (i < NumDynamicPointLights)
-				{
-			#else
-			for (int i = 0; i < NumDynamicPointLights; i++)
-			{
-			#endif
-			#else
-			for (int i = 0; i < NUM_DYNAMIC_POINT_LIGHTS; i++)
-			{
-			#endif
-				float3 ToLight = LightPositionAndInvRadius[i].xyz - MaterialParameters.AbsoluteWorldPosition;
-				float DistanceSqr = dot(ToLight, ToLight);
-				float3 L = ToLight * rsqrt(DistanceSqr);
-				float3 H = normalize(MaterialParameters.CameraVector + L);
+		SpecularIBL = GetImageBasedReflectionLighting(MaterialParameters, ShadingModelContext.ClearCoatRoughness, IndirectIrradiance);
+		Color += SpecularIBL * F;
+	#else
+		Color += SpecularIBL * ShadingModelContext.SpecularColor;
+	#endif
+#endif /* FULLY_ROUGH */
+
+	// Local lights
+#if MAX_DYNAMIC_POINT_LIGHTS > 0 && !MATERIAL_SHADINGMODEL_SINGLELAYERWATER
+	#if VARIABLE_NUM_DYNAMIC_POINT_LIGHTS
+	#if WEBGL || ES2_PROFILE || VULKAN_PROFILE
+	// webgl and ES2 needs a constant loop counter
+	UNROLL
+	for (int i = 0; i < MAX_DYNAMIC_POINT_LIGHTS; i++)
+	{
+		if (i < NumDynamicPointLights)
+		{
+	#else
+	for (int i = 0; i < NumDynamicPointLights; i++)
+	{
+	#endif
+	#else
+	for (int i = 0; i < NUM_DYNAMIC_POINT_LIGHTS; i++)
+	{
+	#endif
+		float3 ToLight = LightPositionAndInvRadius[i].xyz - MaterialParameters.AbsoluteWorldPosition;
+		float DistanceSqr = dot(ToLight, ToLight);
+		float3 L = ToLight * rsqrt(DistanceSqr);
+		float3 PointH = normalize(MaterialParameters.CameraVector + L);
 
-				float PointNoL = max(0, dot(MaterialParameters.WorldNormal, L));
-				float PointRoL = max(0, dot(MaterialParameters.ReflectionVector, L));
-				float PointNoH = max(0, dot(MaterialParameters.WorldNormal, H));
+		float PointNoL = max(0, dot(MaterialParameters.WorldNormal, L));
+		float PointRoL = max(0, dot(MaterialParameters.ReflectionVector, L));
+		float PointNoH = max(0, dot(MaterialParameters.WorldNormal, PointH));
 
-				float Attenuation;
+		float Attenuation;
 
-				if (LightColorAndFalloffExponent[i].w == 0)
-				{
-					// Sphere falloff (technically just 1/d2 but this avoids inf)
-					Attenuation = 1 / ( DistanceSqr + 1 );
+		if (LightColorAndFalloffExponent[i].w == 0)
+		{
+			// Sphere falloff (technically just 1/d2 but this avoids inf)
+			Attenuation = 1 / ( DistanceSqr + 1 );
 	
-					float LightRadiusMask = Square(saturate(1 - Square(DistanceSqr * (LightPositionAndInvRadius[i].w * LightPositionAndInvRadius[i].w))));
-					Attenuation *= LightRadiusMask;
-				}
-				else
-				{
-					Attenuation = RadialAttenuation(ToLight * LightPositionAndInvRadius[i].w, LightColorAndFalloffExponent[i].w);		
-				}
+			float LightRadiusMask = Square(saturate(1 - Square(DistanceSqr * (LightPositionAndInvRadius[i].w * LightPositionAndInvRadius[i].w))));
+			Attenuation *= LightRadiusMask;
+		}
+		else
+		{
+			Attenuation = RadialAttenuation(ToLight * LightPositionAndInvRadius[i].w, LightColorAndFalloffExponent[i].w);		
+		}
 				
-				#if PROJECT_MOBILE_ENABLE_MOVABLE_SPOTLIGHTS
-				if(SpotLightAngles[i].w > 0.0)
-				{
-					Attenuation *= SpotAttenuation(L, -SpotLightDirectionAndSpecularScale[i].xyz, SpotLightAngles[i].xy);
-				}
-				#endif
-
-				#if !FULLY_ROUGH
-					Color += min(65000.0, (Attenuation * PointNoL) * LightColorAndFalloffExponent[i].rgb * (1.0/PI) * (DiffuseColor + SpecularColor * SpotLightDirectionAndSpecularScale[i].w * CalcSpecular(Roughness, PointRoL, PointNoH, H, MaterialParameters.WorldNormal)));
-				#else
-					Color += (Attenuation * PointNoL) * LightColorAndFalloffExponent[i].rgb * DiffuseColor;
-				#endif
-			#if VARIABLE_NUM_DYNAMIC_POINT_LIGHTS && ( WEBGL || ES2_PROFILE || VULKAN_PROFILE )
-				}
-			#endif
-			}
+		#if PROJECT_MOBILE_ENABLE_MOVABLE_SPOTLIGHTS
+		if(SpotLightAngles[i].w > 0.0)
+		{
+			Attenuation *= SpotAttenuation(L, -SpotLightDirectionAndSpecularScale[i].xyz, SpotLightAngles[i].xy);
+		}
+		#endif
+
+		#if !FULLY_ROUGH
+			FMobileDirectLighting Lighting = MobileIntegrateBxDF(ShadingModelContext, PointNoL, PointRoL, MaterialParameters.CameraVector, MaterialParameters.WorldNormal, PointH, PointNoH);
+			Color += min(65000.0, (Attenuation * PointNoL) * LightColorAndFalloffExponent[i].rgb * (1.0/PI) * (Lighting.Diffuse + Lighting.Specular * SpotLightDirectionAndSpecularScale[i].w));
+		#else
+			Color += (Attenuation * PointNoL) * LightColorAndFalloffExponent[i].rgb * (1.0/PI) * ShadingModelContext.DiffuseColor;
 		#endif
+	#if VARIABLE_NUM_DYNAMIC_POINT_LIGHTS && ( WEBGL || ES2_PROFILE || VULKAN_PROFILE )
+		}
+	#endif
+	}
+#endif
+
+	// Skylight
+#if ENABLE_SKY_LIGHT
+	//@mw todo
+	// TODO: Also need to do specular.
+	#if MATERIAL_TWOSIDED && LQ_TEXTURE_LIGHTMAP
+	if (NoL == 0)
+	{
+	#endif
+		half3 SkyDiffuseLighting = GetSkySHDiffuseSimple(MaterialParameters.WorldNormal);
 
-	#endif /* !MATERIAL_SHADINGMODEL_UNLIT */
+	#if MATERIAL_SHADINGMODEL_SINGLELAYERWATER
+		ShadingModelContext.WaterDiffuseIndirectLuminance += SkyDiffuseLighting;
+	#endif
+		Color += SkyDiffuseLighting * half3(ResolvedView.SkyLightColor.rgb) * ShadingModelContext.DiffuseColor * MaterialAO;
+	#if MATERIAL_TWOSIDED && LQ_TEXTURE_LIGHTMAP
+	}
+	#endif
+#endif
+
+#endif /* !MATERIAL_SHADINGMODEL_UNLIT */
 		 
 	half3 Emissive = GetMaterialEmissive(PixelMaterialInputs);
-
 	Color += Emissive;
 
 #if MATERIAL_SHADINGMODEL_SINGLELAYERWATER
 	{
 		const bool CameraIsUnderWater = false;	// Fade out the material contribution over to water contribution according to material opacity.
 		const float3 SunIlluminance = ResolvedView.DirectionalLightColor.rgb * PI;			// times PI because it is divided by PI on CPU (=luminance) and we want illuminance here. 
-		const float3 WaterDiffuseIndirectIlluminance = WaterDiffuseIndirectLuminance * PI;	// DiffuseIndirectLighting is luminance. So we need to multiply by PI to get illuminance.
-		const float3 EnvBrdf = SpecularColor; // SpecularColor is not F0 as in BasePassPixelShader, it is EnvBRDF.
+		const float3 WaterDiffuseIndirectIlluminance = ShadingModelContext.WaterDiffuseIndirectLuminance * PI;	// DiffuseIndirectLighting is luminance. So we need to multiply by PI to get illuminance.
+		const float3 EnvBrdf = ShadingModelContext.SpecularColor; // SpecularColor is not F0 as in BasePassPixelShader, it is EnvBRDF.
 		const uint EyeIndex = 0;
 
 		const float4 NullDistortionParams = 1.0f;
 		WaterVolumeLightingOutput WaterLighting = EvaluateWaterVolumeLighting(
 			MaterialParameters, PixelMaterialInputs, ResolvedView,
-			Shadow, Specular, NullDistortionParams,
+			Shadow, ShadingModelContext.Specular, NullDistortionParams,
 			SunIlluminance, WaterDiffuseIndirectIlluminance, EnvBrdf,
 			CameraIsUnderWater, WaterVisibility, EyeIndex);
 
@@ -782,22 +748,22 @@ void Main(
 	// NEEDS_BASEPASS_PIXEL_FOGGING is not allowed on mobile for the sake of performance.
 
 #if !MATERIAL_SHADINGMODEL_UNLIT && MOBILE_EMULATION
-	Color = lerp(Color, DiffuseColor + SpecularColor, ResolvedView.UnlitViewmodeMask);
+	Color = lerp(Color, ShadingModelContext.DiffuseColor + ShadingModelContext.SpecularColor, ResolvedView.UnlitViewmodeMask);
 #endif
 
 	// On mobile, water (an opaque material) is rendered as trnaslucent with forced premultiplied alpha blending (see MobileBasePass::SetTranslucentRenderState)
 	#if MATERIALBLENDING_ALPHACOMPOSITE || MATERIAL_SHADINGMODEL_SINGLELAYERWATER
-		OutColor = half4(Color * VertexFog.a + VertexFog.rgb * Opacity, Opacity);
+		OutColor = half4(Color * VertexFog.a + VertexFog.rgb * ShadingModelContext.Opacity, ShadingModelContext.Opacity);
 	#elif MATERIALBLENDING_ALPHAHOLDOUT
 		// not implemented for holdout
-		OutColor = half4(Color * VertexFog.a + VertexFog.rgb * Opacity, Opacity);
+		OutColor = half4(Color * VertexFog.a + VertexFog.rgb * ShadingModelContext.Opacity, ShadingModelContext.Opacity);
 	#elif MATERIALBLENDING_TRANSLUCENT
-		OutColor = half4(Color * VertexFog.a + VertexFog.rgb, Opacity);
+		OutColor = half4(Color * VertexFog.a + VertexFog.rgb, ShadingModelContext.Opacity);
 	#elif MATERIALBLENDING_ADDITIVE
-		OutColor = half4(Color * (VertexFog.a * Opacity.x), 0.0f);
+		OutColor = half4(Color * (VertexFog.a * ShadingModelContext.Opacity.x), 0.0f);
 	#elif MATERIALBLENDING_MODULATE
 		half3 FoggedColor = lerp(half3(1, 1, 1), Color, VertexFog.aaa * VertexFog.aaa);
-		OutColor = half4(FoggedColor, Opacity);
+		OutColor = half4(FoggedColor, ShadingModelContext.Opacity);
 	#else
 		OutColor.rgb = Color * VertexFog.a + VertexFog.rgb;
 
diff --git a/Engine/Shaders/Private/MobileShadingModels.ush b/Engine/Shaders/Private/MobileShadingModels.ush
new file mode 100644
--- /dev/null
+++ b/Engine/Shaders/Private/MobileShadingModels.ush
@@ -0,0 +1,182 @@
+// Copyright 1998-2020 Epic Games, Inc. All Rights Reserved.
+
+#pragma once
+
+#include "BRDF.ush"
+#include "MobileGGX.ush"
+
+struct FMobileDirectLighting
+{
+	half3 Diffuse;
+	half3 Specular;
+};
+
+struct FMobileShadingModelContext
+{
+	half Opacity;
+	half3 BaseColor;
+	half Metallic;
+	half Specular;
+	half Roughness;
+	half RoughnessWithClamp;
+
+	half3 DiffuseColor;
+#if NONMETAL
+	half SpecularColor;
+#else
+	half3 SpecularColor;
+#endif
+
+#if MATERIAL_SHADINGMODEL_CLEAR_COAT
+	half ClearCoat;
+	half ClearCoatRoughness;
+	float NoV;
+	half3 SpecPreEnvBrdf;
+#elif MATERIAL_SHADINGMODEL_SINGLELAYERWATER
+	half BaseMaterialCoverageOverWater;
+	half WaterVisibility;
+	float3 WaterDiffuseIndirectLuminance;
+#endif
+};
+
+#if NONMETAL
+half GetEnvBRDF(half SpecularColor, half Roughness, float NoV)
+{
+	// If nothing is hooked up to Metalic and Specular,
+	// then defaults are the same as a non-metal,
+	// so this define is safe.
+	return EnvBRDFApproxNonmetal(Roughness, NoV);
+}
+#else
+half3 GetEnvBRDF(half3 SpecularColor, half Roughness, float NoV)
+{
+#if FORWARDSHADING_USE_HQ_ENV_BRDF
+	return EnvBRDF(SpecularColor, Roughness, NoV);
+#else
+	return EnvBRDFApprox(SpecularColor, Roughness, NoV);
+#endif
+}
+#endif
+
+void InitShadingModelContext( inout FMobileShadingModelContext ShadingModelContext, FMaterialPixelParameters MaterialParameters)
+{
+#if !FULLY_ROUGH
+	float NoV = max(dot(MaterialParameters.WorldNormal, MaterialParameters.CameraVector), 0);
+#endif
+
+#if NONMETAL
+	ShadingModelContext.DiffuseColor = BaseColor;
+	ShadingModelContext.SpecularColor = 0.04;
+#else
+
+	#if MATERIAL_SHADINGMODEL_CLEAR_COAT
+		ShadingModelContext.ClearCoat = saturate(GetMaterialCustomData0(MaterialParameters));
+		ShadingModelContext.ClearCoatRoughness = clamp(GetMaterialCustomData1(MaterialParameters), 0.12, 1.0);
+		const half MetalSpec = 0.9;
+
+		// Approximation of refraction's effect on EnvBRDF
+		half RefractionScale = ((NoV * 0.5 + 0.5) * NoV - 1) * saturate(1.25 - 1.25 * ShadingModelContext.Roughness) + 1;
+
+		// Approximation of absorption integral, tuned for Roughness=0.4
+		half3 AbsorptionColor = ShadingModelContext.BaseColor * (1 / MetalSpec);
+		half3 Absorption = AbsorptionColor * ((NoV - 1) * 0.85 * (1 - lerp(AbsorptionColor, Square(AbsorptionColor), -0.78)) + 1);
+
+		// F_Schlick
+		half F0 = 0.04;
+		half Fc = Pow5(1 - NoV);
+		half F = Fc + (1 - Fc) * F0;
+		half LayerAttenuation = lerp(1, (1 - F), ShadingModelContext.ClearCoat);
+
+		ShadingModelContext.Specular *= lerp(1, RefractionScale, ShadingModelContext.ClearCoat);
+		half DielectricSpecular = 0.08 * ShadingModelContext.Specular;
+		ShadingModelContext.SpecularColor = (DielectricSpecular - DielectricSpecular * ShadingModelContext.Metallic) + ShadingModelContext.BaseColor * ShadingModelContext.Metallic;	// 2 mad
+		ShadingModelContext.SpecPreEnvBrdf = ShadingModelContext.SpecularColor;
+		ShadingModelContext.DiffuseColor = ShadingModelContext.BaseColor - ShadingModelContext.BaseColor * ShadingModelContext.Metallic;
+
+		// Adjust BaseColor according to ClearCoat layer attenuation
+		ShadingModelContext.BaseColor = lerp(ShadingModelContext.BaseColor * LayerAttenuation, MetalSpec * Absorption * RefractionScale, ShadingModelContext.Metallic * ShadingModelContext.ClearCoat);
+		const half Dither = InterleavedGradientNoise(MaterialParameters.SvPosition.xy, View.StateFrameIndexMod8);
+		ShadingModelContext.BaseColor += Dither / 255.f;
+
+		// This is to prevent Vis to get inf when both NoL and NoV are 0.
+		ShadingModelContext.NoV = saturate(abs(NoV) + 1e-5);
+
+	#elif MATERIAL_SHADINGMODEL_SINGLELAYERWATER
+		ShadingModelContext.BaseMaterialCoverageOverWater = ShadingModelContext.Opacity;
+		ShadingModelContext.WaterVisibility = 1.0 - ShadingModelContext.BaseMaterialCoverageOverWater;
+
+		half DielectricSpecular = 0.08 * ShadingModelContext.Specular;
+		ShadingModelContext.DiffuseColor = ShadingModelContext.BaseColor - ShadingModelContext.BaseColor * ShadingModelContext.Metallic;	// 1 mad
+		ShadingModelContext.SpecularColor = (DielectricSpecular - DielectricSpecular * ShadingModelContext.Metallic) + ShadingModelContext.BaseColor * ShadingModelContext.Metallic;	// 2 mad
+
+	#else // MATERIAL_SHADINGMODEL_DEFAULTLIT
+		half DielectricSpecular = 0.08 * ShadingModelContext.Specular;
+		ShadingModelContext.DiffuseColor = ShadingModelContext.BaseColor - ShadingModelContext.BaseColor * ShadingModelContext.Metallic;	// 1 mad
+		ShadingModelContext.SpecularColor = (DielectricSpecular - DielectricSpecular * ShadingModelContext.Metallic) + ShadingModelContext.BaseColor * ShadingModelContext.Metallic;	// 2 mad
+
+	#endif
+
+#endif
+
+#if FULLY_ROUGH
+	EnvBRDFApproxFullyRough(ShadingModelContext.DiffuseColor, ShadingModelContext.SpecularColor);
+#else
+	ShadingModelContext.SpecularColor = GetEnvBRDF(ShadingModelContext.SpecularColor, ShadingModelContext.Roughness, NoV);
+#endif
+
+#if MOBILE_EMULATION
+	{
+		// this feature is only needed for development/editor - we can compile it out for a shipping build (see r.CompileShadersForDevelopment cvar help)
+		ShadingModelContext.DiffuseColor = ShadingModelContext.DiffuseColor * ResolvedView.DiffuseOverrideParameter.w + ResolvedView.DiffuseOverrideParameter.xyz;
+#if NONMETAL
+		ShadingModelContext.SpecularColor = ShadingModelContext.SpecularColor * ResolvedView.SpecularOverrideParameter.w + ResolvedView.SpecularOverrideParameter.x;
+#else
+		ShadingModelContext.SpecularColor = ShadingModelContext.SpecularColor * ResolvedView.SpecularOverrideParameter.w + ResolvedView.SpecularOverrideParameter.xyz;
+#endif
+	}
+#endif
+
+#if MATERIAL_SHADINGMODEL_SINGLELAYERWATER
+	// Fade out diffuse as this will be handled by the single scattering lighting. when over the water surface.
+	// We keep the SpecularColor for sun/water interactions
+	ShadingModelContext.WaterDiffuseIndirectLuminance = 0;
+	ShadingModelContext.DiffuseColor *= ShadingModelContext.BaseMaterialCoverageOverWater;
+#endif
+}
+
+half CalcSpecular(half Roughness, half RoughnessWithClamp, float RoL, float NoH, float3 H, float3 N)
+{
+#if USE_LEGACY_SPECULAR
+	return PhongApprox(Roughness, RoL);
+#else
+	return (Roughness*0.25 + 0.25) * GGX_Mobile(RoughnessWithClamp, NoH, H, N);
+#endif
+}
+
+FMobileDirectLighting MobileIntegrateBxDF(FMobileShadingModelContext ShadingModelContext, float NoL, float RoL, float3 V, float3 N, float3 H, float3 NoH)
+{
+	FMobileDirectLighting Lighting;
+#if MATERIAL_SHADINGMODEL_CLEAR_COAT
+	float VoH = max(0, dot(V, H));
+	half F0 = 0.04;
+	half Fc = Pow5(1 - VoH);
+	half F = Fc + (1 - Fc) * F0;
+	F *= ShadingModelContext.ClearCoat;
+	half LayerAttenuation = 1 - F;
+
+	// Vis_SmithJointApprox
+	half a = ShadingModelContext.ClearCoatRoughness * ShadingModelContext.ClearCoatRoughness;
+	float Vis_SmithV = NoL * (ShadingModelContext.NoV * (1 - a) + a);
+	float Vis_SmithL = ShadingModelContext.NoV * (NoL * (1 - a) + a);
+	float Vis = 0.5 * rcp(Vis_SmithV + Vis_SmithL);
+
+	Lighting.Specular = F * Vis * GGX_Mobile(ShadingModelContext.ClearCoatRoughness, NoH, H, N);
+	Lighting.Specular += LayerAttenuation * ShadingModelContext.SpecularColor * CalcSpecular(ShadingModelContext.Roughness, ShadingModelContext.RoughnessWithClamp, RoL, NoH, H, N);
+	Lighting.Diffuse = LayerAttenuation * ShadingModelContext.DiffuseColor;
+
+#else // MATERIAL_SHADINGMODEL_DEFAULT_LIT
+	Lighting.Specular = ShadingModelContext.SpecularColor * CalcSpecular(ShadingModelContext.Roughness, ShadingModelContext.RoughnessWithClamp, RoL, NoH, H, N);
+	Lighting.Diffuse = ShadingModelContext.DiffuseColor;
+#endif
+	return Lighting;
+}
\ No newline at end of file
diff --git a/Engine/Source/Developer/IOS/IOSTargetPlatform/Private/Windows/IOSDeviceHelperWindows.cpp b/Engine/Source/Developer/IOS/IOSTargetPlatform/Private/Windows/IOSDeviceHelperWindows.cpp
--- a/Engine/Source/Developer/IOS/IOSTargetPlatform/Private/Windows/IOSDeviceHelperWindows.cpp
+++ b/Engine/Source/Developer/IOS/IOSTargetPlatform/Private/Windows/IOSDeviceHelperWindows.cpp
@@ -166,7 +166,7 @@ private:
 			RetryQuery--;
 			if (RetryQuery < 0 || Response < 0)
 			{
-				//UE_LOG(LogTemp, Log, TEXT("IOS device listing is disabled (to many failed attempts)!"));
+				UE_LOG(LogTemp, Log, TEXT("IOS device listing is disabled (too many failed attempts)!"));
 				Enable(false);
 			}
 			return;
diff --git a/Engine/Source/Editor/UnrealEd/Private/PlayLevelNewProcess.cpp b/Engine/Source/Editor/UnrealEd/Private/PlayLevelNewProcess.cpp
--- a/Engine/Source/Editor/UnrealEd/Private/PlayLevelNewProcess.cpp
+++ b/Engine/Source/Editor/UnrealEd/Private/PlayLevelNewProcess.cpp
@@ -132,8 +132,8 @@ void UEditorEngine::LaunchNewProcess(const FRequestPlaySessionParams& InParams,
 		}
 		else
 		{
-			// Otherwise, we'll fall back to ES2 emulation.
-			CommandLine += TEXT(" -featureleveles2");
+			// Otherwise, we'll fall back to ES31 emulation.
+			CommandLine += TEXT(" -featureleveles31");
 		}
 
 		// If we're currently running in OpenGL mode, pass that onto our newly spawned processes.
@@ -148,6 +148,22 @@ void UEditorEngine::LaunchNewProcess(const FRequestPlaySessionParams& InParams,
 		// Ensure the executable writes out a differently named config file to avoid multiple instances overwriting each other.
 		// ToDo: Should this be on all multi-client launches?
 		CommandLine += TEXT(" -MultiprocessSaveConfig");
+
+		// In order for the mobile previewer to adjust its safe zone according to the device profile specified in the editor play settings,
+		// we need to pass the PIESafeZoneOverride's values as command line variables to the new process that we are about to launch.
+		FMargin PIESafeZoneOverride = InParams.EditorPlaySettings->PIESafeZoneOverride;
+		CommandLine += FString::Printf(TEXT(" -SafeZonePaddingLeft=%f -SafeZonePaddingRight=%f -SafeZonePaddingTop=%f -SafeZonePaddingBottom=%f"),
+			PIESafeZoneOverride.Left,
+			PIESafeZoneOverride.Right,
+			PIESafeZoneOverride.Top,
+			PIESafeZoneOverride.Bottom
+		);
+
+		if (!PIESafeZoneOverride.GetDesiredSize().IsZero())
+		{
+			// Send -DrawUnSafeZones so that the red "unsafe" zones show up in the mobile previewer.
+			CommandLine += TEXT(" -DrawUnSafeZones");
+		}
 	}
 
 	if (InParams.SessionPreviewTypeOverride.Get(EPlaySessionPreviewType::NoPreview) == EPlaySessionPreviewType::VulkanPreview)
diff --git a/Engine/Source/Programs/AutomationTool/IOS/IOSPlatform.Automation.cs b/Engine/Source/Programs/AutomationTool/IOS/IOSPlatform.Automation.cs
--- a/Engine/Source/Programs/AutomationTool/IOS/IOSPlatform.Automation.cs
+++ b/Engine/Source/Programs/AutomationTool/IOS/IOSPlatform.Automation.cs
@@ -211,7 +211,7 @@ public class IOSPlatform : Platform
 		IOSExports.GetProvisioningData(InProject, bDistribution, out MobileProvision, out SigningCertificate, out TeamUUID, out bAutomaticSigning);
 	}
 
-	public virtual bool DeployGeneratePList(FileReference ProjectFile, UnrealTargetConfiguration Config, DirectoryReference ProjectDirectory, bool bIsUE4Game, string GameName, string ProjectName, DirectoryReference InEngineDir, DirectoryReference AppDirectory, string InExecutablePath, out bool bSupportsPortrait, out bool bSupportsLandscape, out bool bSkipIcons)
+	public virtual bool DeployGeneratePList(FileReference ProjectFile, UnrealTargetConfiguration Config, DirectoryReference ProjectDirectory, bool bIsUE4Game, string GameName, bool bIsClient, string ProjectName, DirectoryReference InEngineDir, DirectoryReference AppDirectory, string InExecutablePath, out bool bSupportsPortrait, out bool bSupportsLandscape, out bool bSkipIcons)
 	{
 		string TargetName = Path.GetFileNameWithoutExtension(InExecutablePath).Split("-".ToCharArray())[0];
 		FileReference TargetReceiptFileName;
@@ -223,7 +223,7 @@ public class IOSPlatform : Platform
 		{
 			TargetReceiptFileName = TargetReceipt.GetDefaultPath(ProjectDirectory, TargetName, UnrealTargetPlatform.IOS, Config, "");
 		}
-		return IOSExports.GeneratePList(ProjectFile, Config, ProjectDirectory, bIsUE4Game, GameName, ProjectName, InEngineDir, AppDirectory, TargetReceiptFileName, out bSupportsPortrait, out bSupportsLandscape, out bSkipIcons);
+		return IOSExports.GeneratePList(ProjectFile, Config, ProjectDirectory, bIsUE4Game, GameName, bIsClient, ProjectName, InEngineDir, AppDirectory, TargetReceiptFileName, out bSupportsPortrait, out bSupportsLandscape, out bSkipIcons);
 	}
 
 	protected string MakeIPAFileName(UnrealTargetConfiguration TargetConfiguration, ProjectParams Params, DeploymentContext SC, bool bAllowDistroPrefix)
@@ -867,8 +867,9 @@ public class IOSPlatform : Platform
 							(SC.IsCodeBasedProject ? SC.ProjectRoot : DirectoryReference.Combine(SC.LocalRoot, "Engine")),
 							!SC.IsCodeBasedProject,
 							(SC.IsCodeBasedProject ? SC.StageExecutables[0] : "UE4Game"),
+							Params.Client,
 							SC.ShortProjectName, DirectoryReference.Combine(SC.LocalRoot, "Engine"),
-														DirectoryReference.Combine((SC.IsCodeBasedProject ? SC.ProjectRoot : DirectoryReference.Combine(SC.LocalRoot, "Engine")), "Binaries", PlatformName, "Payload", (SC.IsCodeBasedProject ? SC.ShortProjectName : "UE4Game") + ".app"),
+							DirectoryReference.Combine((SC.IsCodeBasedProject ? SC.ProjectRoot : DirectoryReference.Combine(SC.LocalRoot, "Engine")), "Binaries", PlatformName, "Payload", (SC.IsCodeBasedProject ? SC.ShortProjectName : "UE4Game") + ".app"),
 							SC.StageExecutables[0],
 							out bSupportsPortrait,
 							out bSupportsLandscape,
diff --git a/Engine/Source/Programs/AutomationTool/TVOS/TVOSPlatform.Automation.cs b/Engine/Source/Programs/AutomationTool/TVOS/TVOSPlatform.Automation.cs
--- a/Engine/Source/Programs/AutomationTool/TVOS/TVOSPlatform.Automation.cs
+++ b/Engine/Source/Programs/AutomationTool/TVOS/TVOSPlatform.Automation.cs
@@ -37,7 +37,7 @@ public class TVOSPlatform : IOSPlatform
 		TVOSExports.GetProvisioningData(InProject, bDistribution, out MobileProvision, out SigningCertificate, out Team, out bAutomaticSigning);
     }
 
-	public override bool DeployGeneratePList(FileReference ProjectFile, UnrealTargetConfiguration Config, DirectoryReference ProjectDirectory, bool bIsUE4Game, string GameName, string ProjectName, DirectoryReference InEngineDir, DirectoryReference AppDirectory, string InExecutablePath, out bool bSupportsPortrait, out bool bSupportsLandscape, out bool bSkipIcons)
+	public override bool DeployGeneratePList(FileReference ProjectFile, UnrealTargetConfiguration Config, DirectoryReference ProjectDirectory, bool bIsUE4Game, string GameName, bool bIsClient, string ProjectName, DirectoryReference InEngineDir, DirectoryReference AppDirectory, string InExecutablePath, out bool bSupportsPortrait, out bool bSupportsLandscape, out bool bSkipIcons)
 	{
 		string TargetName = Path.GetFileNameWithoutExtension(InExecutablePath).Split("-".ToCharArray())[0];
 		FileReference TargetReceiptFileName;
@@ -49,7 +49,7 @@ public class TVOSPlatform : IOSPlatform
 		{
 			TargetReceiptFileName = TargetReceipt.GetDefaultPath(ProjectDirectory, TargetName, UnrealTargetPlatform.TVOS, Config, "");
 		}
-		return TVOSExports.GeneratePList(ProjectFile, Config, ProjectDirectory, bIsUE4Game, GameName, ProjectName, InEngineDir, AppDirectory, TargetReceiptFileName, out bSupportsPortrait, out bSupportsLandscape, out bSkipIcons);
+		return TVOSExports.GeneratePList(ProjectFile, Config, ProjectDirectory, bIsUE4Game, GameName, bIsClient, ProjectName, InEngineDir, AppDirectory, TargetReceiptFileName, out bSupportsPortrait, out bSupportsLandscape, out bSkipIcons);
 	}
 
     public override string GetCookPlatform(bool bDedicatedServer, bool bIsClientOnly)
@@ -112,7 +112,7 @@ public class TVOSPlatform : IOSPlatform
                     bool bSupportsPortrait = false;
                     bool bSupportsLandscape = false;
                     bool bSkipIcons = false;
-                    DeployGeneratePList(SC.RawProjectPath, TargetConfiguration, (SC.IsCodeBasedProject ? SC.ProjectRoot : SC.EngineRoot), !SC.IsCodeBasedProject, (SC.IsCodeBasedProject ? SC.ShortProjectName : "UE4Game"), SC.ShortProjectName, SC.EngineRoot, DirectoryReference.Combine((SC.IsCodeBasedProject ? SC.ProjectRoot : SC.EngineRoot), "Binaries", "TVOS", "Payload", (SC.IsCodeBasedProject ? SC.ShortProjectName : "UE4Game") + ".app"), SC.StageExecutables[0], out bSupportsPortrait, out bSupportsLandscape, out bSkipIcons);
+                    DeployGeneratePList(SC.RawProjectPath, TargetConfiguration, (SC.IsCodeBasedProject ? SC.ProjectRoot : SC.EngineRoot), !SC.IsCodeBasedProject, (SC.IsCodeBasedProject ? SC.ShortProjectName : "UE4Game"), Params.Client, SC.ShortProjectName, SC.EngineRoot, DirectoryReference.Combine((SC.IsCodeBasedProject ? SC.ProjectRoot : SC.EngineRoot), "Binaries", "TVOS", "Payload", (SC.IsCodeBasedProject ? SC.ShortProjectName : "UE4Game") + ".app"), SC.StageExecutables[0], out bSupportsPortrait, out bSupportsLandscape, out bSkipIcons);
                 }
 
                 SC.StageFile(StagedFileType.SystemNonUFS, TargetPListFile, new StagedFileReference("Info.plist"));
diff --git a/Engine/Source/Programs/IOS/iPhonePackager/SSHCommandHelper.cs b/Engine/Source/Programs/IOS/iPhonePackager/SSHCommandHelper.cs
--- a/Engine/Source/Programs/IOS/iPhonePackager/SSHCommandHelper.cs
+++ b/Engine/Source/Programs/IOS/iPhonePackager/SSHCommandHelper.cs
@@ -260,7 +260,7 @@ namespace iPhonePackager
 			// make simple rsync commandline to send a file
 			RsyncProcess.StartInfo.FileName = RSyncExePath;
 			RsyncProcess.StartInfo.Arguments = string.Format(
-                "-zrltgoDe \"ssh {0}\" --rsync-path=\"mkdir -p {1} && rsync\" --chmod=ug=rwX,o=rxX '{2}' \"{3}@{4}\":'{1}/{5}'",
+                "-zrltgoDe \"./ssh {0}\" --rsync-path=\"mkdir -p {1} && rsync\" --chmod=ug=rwX,o=rxX '{2}' \"{3}@{4}\":'{1}/{5}'",
 				SSHAuthentication,
 				RemoteDir,
 				ConvertPathToCygwin(LocalPath),
@@ -316,7 +316,7 @@ namespace iPhonePackager
 			// make simple rsync commandline to send a file
 			RsyncProcess.StartInfo.FileName = RSyncExePath;
 			RsyncProcess.StartInfo.Arguments = string.Format(
-                "-zrltgoDe \"ssh {0}\" \"{2}@{3}\":'{4}' \"{1}\"",
+                "-zrltgoDe \"./ssh {0}\" \"{2}@{3}\":'{4}' \"{1}\"",
 				SSHAuthentication,
 				ConvertPathToCygwin(LocalPath),
 				SSHUser,
diff --git a/Engine/Source/Programs/UnrealBuildTool/Platform/IOS/IOSExports.cs b/Engine/Source/Programs/UnrealBuildTool/Platform/IOS/IOSExports.cs
--- a/Engine/Source/Programs/UnrealBuildTool/Platform/IOS/IOSExports.cs
+++ b/Engine/Source/Programs/UnrealBuildTool/Platform/IOS/IOSExports.cs
@@ -93,6 +93,7 @@ namespace UnrealBuildTool
 		/// <param name="ProjectDirectory"></param>
 		/// <param name="bIsUE4Game"></param>
 		/// <param name="GameName"></param>
+		/// <param name="bIsClient"></param>
 		/// <param name="ProjectName"></param>
 		/// <param name="InEngineDir"></param>
 		/// <param name="AppDirectory"></param>
@@ -101,10 +102,10 @@ namespace UnrealBuildTool
 		/// <param name="bSupportsLandscape"></param>
 		/// <param name="bSkipIcons"></param>
 		/// <returns></returns>
-		public static bool GeneratePList(FileReference ProjectFile, UnrealTargetConfiguration Config, DirectoryReference ProjectDirectory, bool bIsUE4Game, string GameName, string ProjectName, DirectoryReference InEngineDir, DirectoryReference AppDirectory, FileReference BuildReceiptFileName, out bool bSupportsPortrait, out bool bSupportsLandscape, out bool bSkipIcons)
+		public static bool GeneratePList(FileReference ProjectFile, UnrealTargetConfiguration Config, DirectoryReference ProjectDirectory, bool bIsUE4Game, string GameName, bool bIsClient, string ProjectName, DirectoryReference InEngineDir, DirectoryReference AppDirectory, FileReference BuildReceiptFileName, out bool bSupportsPortrait, out bool bSupportsLandscape, out bool bSkipIcons)
 		{
 			TargetReceipt Receipt = TargetReceipt.Read(BuildReceiptFileName);
-			return new UEDeployIOS().GeneratePList(ProjectFile, Config, ProjectDirectory.FullName, bIsUE4Game, GameName, ProjectName, InEngineDir.FullName, AppDirectory.FullName, Receipt, out bSupportsPortrait, out bSupportsLandscape, out bSkipIcons);
+			return new UEDeployIOS().GeneratePList(ProjectFile, Config, ProjectDirectory.FullName, bIsUE4Game, GameName, bIsClient, ProjectName, InEngineDir.FullName, AppDirectory.FullName, Receipt, out bSupportsPortrait, out bSupportsLandscape, out bSkipIcons);
 		}
 
 		/// <summary>
diff --git a/Engine/Source/Programs/UnrealBuildTool/Platform/IOS/UEDeployIOS.cs b/Engine/Source/Programs/UnrealBuildTool/Platform/IOS/UEDeployIOS.cs
--- a/Engine/Source/Programs/UnrealBuildTool/Platform/IOS/UEDeployIOS.cs
+++ b/Engine/Source/Programs/UnrealBuildTool/Platform/IOS/UEDeployIOS.cs
@@ -195,7 +195,7 @@ namespace UnrealBuildTool
 			return result;
 		}
 
-		public static bool GenerateIOSPList(FileReference ProjectFile, UnrealTargetConfiguration Config, string ProjectDirectory, bool bIsUE4Game, string GameName, string ProjectName, string InEngineDir, string AppDirectory, VersionNumber SdkVersion, UnrealPluginLanguage UPL, string BundleID, bool bBuildAsFramework, out bool bSupportsPortrait, out bool bSupportsLandscape, out bool bSkipIcons)
+		public static bool GenerateIOSPList(FileReference ProjectFile, UnrealTargetConfiguration Config, string ProjectDirectory, bool bIsUE4Game, string GameName, bool bIsClient, string ProjectName, string InEngineDir, string AppDirectory, VersionNumber SdkVersion, UnrealPluginLanguage UPL, string BundleID, bool bBuildAsFramework, out bool bSupportsPortrait, out bool bSupportsLandscape, out bool bSkipIcons)
 		{
 			// generate the Info.plist for future use
 			string BuildDirectory = ProjectDirectory + "/Build/IOS";
@@ -430,7 +430,12 @@ namespace UnrealBuildTool
 			Text.AppendLine("\t<key>CFBundleDisplayName</key>");
 			Text.AppendLine(string.Format("\t<string>{0}</string>", EncodeBundleName(BundleDisplayName, ProjectName)));
 			Text.AppendLine("\t<key>CFBundleExecutable</key>");
-			Text.AppendLine(string.Format("\t<string>{0}</string>", bIsUE4Game ? "UE4Game" : GameName));
+			string BundleExecutable = bIsUE4Game ? "UE4Game" : GameName;
+			if (bIsClient)
+			{
+				BundleExecutable += "Client";
+			}
+			Text.AppendLine(string.Format("\t<string>{0}</string>", BundleExecutable));
 			Text.AppendLine("\t<key>CFBundleIdentifier</key>");
 			Text.AppendLine(string.Format("\t<string>{0}</string>", BundleIdentifier.Replace("[PROJECT_NAME]", ProjectName).Replace("_", "")));
 			Text.AppendLine("\t<key>CFBundleInfoDictionaryVersion</key>");
@@ -814,15 +819,15 @@ namespace UnrealBuildTool
 			return false;
 		}
 
-		public bool GeneratePList(FileReference ProjectFile, UnrealTargetConfiguration Config, string ProjectDirectory, bool bIsUE4Game, string GameName, string ProjectName, string InEngineDir, string AppDirectory, TargetReceipt Receipt, out bool bSupportsPortrait, out bool bSupportsLandscape, out bool bSkipIcons)
+ 		public bool GeneratePList(FileReference ProjectFile, UnrealTargetConfiguration Config, string ProjectDirectory, bool bIsUE4Game, string GameName, bool bIsClient, string ProjectName, string InEngineDir, string AppDirectory, TargetReceipt Receipt, out bool bSupportsPortrait, out bool bSupportsLandscape, out bool bSkipIcons)
 		{
 			List<string> UPLScripts = CollectPluginDataPaths(Receipt.AdditionalProperties);
 			VersionNumber SdkVersion = GetSdkVersion(Receipt);
 			bool bBuildAsFramework = GetCompileAsDll(Receipt);
-			return GeneratePList(ProjectFile, Config, ProjectDirectory, bIsUE4Game, GameName, ProjectName, InEngineDir, AppDirectory, UPLScripts, SdkVersion, "", bBuildAsFramework, out bSupportsPortrait, out bSupportsLandscape, out bSkipIcons);
+			return GeneratePList(ProjectFile, Config, ProjectDirectory, bIsUE4Game, GameName, bIsClient, ProjectName, InEngineDir, AppDirectory, UPLScripts, SdkVersion, "", bBuildAsFramework, out bSupportsPortrait, out bSupportsLandscape, out bSkipIcons);
 		}
 
-		public virtual bool GeneratePList(FileReference ProjectFile, UnrealTargetConfiguration Config, string ProjectDirectory, bool bIsUE4Game, string GameName, string ProjectName, string InEngineDir, string AppDirectory, List<string> UPLScripts, VersionNumber SdkVersion, string BundleID, bool bBuildAsFramework, out bool bSupportsPortrait, out bool bSupportsLandscape, out bool bSkipIcons)
+		public virtual bool GeneratePList(FileReference ProjectFile, UnrealTargetConfiguration Config, string ProjectDirectory, bool bIsUE4Game, string GameName, bool bIsClient, string ProjectName, string InEngineDir, string AppDirectory, List<string> UPLScripts, VersionNumber SdkVersion, string BundleID, bool bBuildAsFramework, out bool bSupportsPortrait, out bool bSupportsLandscape, out bool bSkipIcons)
 		{
 			// remember name with -IOS-Shipping, etc
 			string ExeName = GameName;
@@ -854,7 +859,7 @@ namespace UnrealBuildTool
 			// Passing in true for distribution is not ideal here but given the way that ios packaging happens and this call chain it seems unavoidable for now, maybe there is a way to correctly pass it in that I can't find?
 			UPL.Init(ProjectArches, true, RelativeEnginePath, BundlePath, ProjectDirectory, Config.ToString(), false);
 
-			return GenerateIOSPList(ProjectFile, Config, ProjectDirectory, bIsUE4Game, GameName, ProjectName, InEngineDir, AppDirectory, SdkVersion, UPL, BundleID, bBuildAsFramework, out bSupportsPortrait, out bSupportsLandscape, out bSkipIcons);
+			return GenerateIOSPList(ProjectFile, Config, ProjectDirectory, bIsUE4Game, GameName, bIsClient, ProjectName, InEngineDir, AppDirectory, SdkVersion, UPL, BundleID, bBuildAsFramework, out bSupportsPortrait, out bSupportsLandscape, out bSkipIcons);
 		}
 
 		protected virtual void CopyCloudResources(string InEngineDir, string AppDirectory)
@@ -1125,7 +1130,7 @@ namespace UnrealBuildTool
 			bool bSupportsPortrait = true;
 			bool bSupportsLandscape = false;
 			bool bSkipIcons = false;
-			bool bSkipDefaultPNGs = GeneratePList(ProjectFile, Config, InProjectDirectory, bIsUE4Game, GameExeName, InProjectName, InEngineDir, AppDirectory, UPLScripts, SdkVersion, BundleID, bBuildAsFramework, out bSupportsPortrait, out bSupportsLandscape, out bSkipIcons);
+			bool bSkipDefaultPNGs = GeneratePList(ProjectFile, Config, InProjectDirectory, bIsUE4Game, GameExeName, false, InProjectName, InEngineDir, AppDirectory, UPLScripts, SdkVersion, BundleID, bBuildAsFramework, out bSupportsPortrait, out bSupportsLandscape, out bSkipIcons);
 
 			// ensure the destination is writable
 			if (File.Exists(AppDirectory + "/" + GameName))
@@ -1202,7 +1207,7 @@ namespace UnrealBuildTool
 			{
 				// @todo tvos merge: This used to copy the bundle back - where did that code go? It needs to be fixed up for TVOS directories
 				bool bSupportPortrait, bSupportLandscape, bSkipIcons;
-				GeneratePList(ProjectFile, Configuration, ProjectDirectory, bIsUE4Game, GameName, (ProjectFile == null) ? "" : Path.GetFileNameWithoutExtension(ProjectFile.FullName), "../../Engine", "", UPLScripts, SdkVersion, BundleID, bBuildAsFramework, out bSupportPortrait, out bSupportLandscape, out bSkipIcons);
+				GeneratePList(ProjectFile, Configuration, ProjectDirectory, bIsUE4Game, GameName, false, (ProjectFile == null) ? "" : Path.GetFileNameWithoutExtension(ProjectFile.FullName), "../../Engine", "", UPLScripts, SdkVersion, BundleID, bBuildAsFramework, out bSupportPortrait, out bSupportLandscape, out bSkipIcons);
 			}
 			return true;
 		}
diff --git a/Engine/Source/Programs/UnrealBuildTool/Platform/TVOS/TVOSExports.cs b/Engine/Source/Programs/UnrealBuildTool/Platform/TVOS/TVOSExports.cs
--- a/Engine/Source/Programs/UnrealBuildTool/Platform/TVOS/TVOSExports.cs
+++ b/Engine/Source/Programs/UnrealBuildTool/Platform/TVOS/TVOSExports.cs
@@ -90,6 +90,7 @@ namespace UnrealBuildTool
 		/// <param name="ProjectDirectory"></param>
 		/// <param name="bIsUE4Game"></param>
 		/// <param name="GameName"></param>
+		/// <param name="bIsClient"></param>
 		/// <param name="ProjectName"></param>
 		/// <param name="InEngineDir"></param>
 		/// <param name="AppDirectory"></param>
@@ -98,10 +99,10 @@ namespace UnrealBuildTool
 		/// <param name="bSupportsLandscape"></param>
 		/// <param name="bSkipIcons"></param>
 		/// <returns></returns>
-		public static bool GeneratePList(FileReference ProjectFile, UnrealTargetConfiguration Config, DirectoryReference ProjectDirectory, bool bIsUE4Game, string GameName, string ProjectName, DirectoryReference InEngineDir, DirectoryReference AppDirectory, FileReference BuildReceiptFileName, out bool bSupportsPortrait, out bool bSupportsLandscape, out bool bSkipIcons)
+		public static bool GeneratePList(FileReference ProjectFile, UnrealTargetConfiguration Config, DirectoryReference ProjectDirectory, bool bIsUE4Game, string GameName, bool bIsClient, string ProjectName, DirectoryReference InEngineDir, DirectoryReference AppDirectory, FileReference BuildReceiptFileName, out bool bSupportsPortrait, out bool bSupportsLandscape, out bool bSkipIcons)
 		{
 			TargetReceipt Receipt = TargetReceipt.Read(BuildReceiptFileName);
-			return new UEDeployTVOS().GeneratePList(ProjectFile, Config, ProjectDirectory.FullName, bIsUE4Game, GameName, ProjectName, InEngineDir.FullName, AppDirectory.FullName, Receipt, out bSupportsPortrait, out bSupportsLandscape, out bSkipIcons);
+			return new UEDeployTVOS().GeneratePList(ProjectFile, Config, ProjectDirectory.FullName, bIsUE4Game, GameName, bIsClient, ProjectName, InEngineDir.FullName, AppDirectory.FullName, Receipt, out bSupportsPortrait, out bSupportsLandscape, out bSkipIcons);
 		}
 	}
 }
diff --git a/Engine/Source/Programs/UnrealBuildTool/Platform/TVOS/UEDeployTVOS.cs b/Engine/Source/Programs/UnrealBuildTool/Platform/TVOS/UEDeployTVOS.cs
--- a/Engine/Source/Programs/UnrealBuildTool/Platform/TVOS/UEDeployTVOS.cs
+++ b/Engine/Source/Programs/UnrealBuildTool/Platform/TVOS/UEDeployTVOS.cs
@@ -24,7 +24,7 @@ namespace UnrealBuildTool
 			return "TVOS";
 		}
 
-		public static bool GenerateTVOSPList(string ProjectDirectory, bool bIsUE4Game, string GameName, string ProjectName, string InEngineDir, string AppDirectory, UnrealPluginLanguage UPL, string BundleID)
+		public static bool GenerateTVOSPList(string ProjectDirectory, bool bIsUE4Game, string GameName, bool bIsClient, string ProjectName, string InEngineDir, string AppDirectory, UnrealPluginLanguage UPL, string BundleID)
 		{
 			// @todo tvos: THIS!
 
@@ -116,7 +116,12 @@ namespace UnrealBuildTool
 			Text.AppendLine("\t<key>CFBundleDisplayName</key>");
 			Text.AppendLine(string.Format("\t<string>{0}</string>", EncodeBundleName(BundleDisplayName, ProjectName)));
 			Text.AppendLine("\t<key>CFBundleExecutable</key>");
-			Text.AppendLine(string.Format("\t<string>{0}</string>", bIsUE4Game ? "UE4Game" : GameName));
+			string BundleExecutable = bIsUE4Game ? "UE4Game" : GameName;
+			if (bIsClient)
+			{
+				BundleExecutable += "Client";
+			}
+			Text.AppendLine(string.Format("\t<string>{0}</string>", BundleExecutable));
 			Text.AppendLine("\t<key>CFBundleIdentifier</key>");
 			Text.AppendLine(string.Format("\t<string>{0}</string>", BundleIdentifier.Replace("[PROJECT_NAME]", ProjectName).Replace("_","")));
 			Text.AppendLine("\t<key>CFBundleInfoDictionaryVersion</key>");
@@ -349,11 +354,11 @@ namespace UnrealBuildTool
 			return bSkipDefaultPNGs;
 		}
 
-		public override bool GeneratePList(FileReference ProjectFile, UnrealTargetConfiguration Config, string ProjectDirectory, bool bIsUE4Game, string GameName, string ProjectName, string InEngineDir, string AppDirectory, List<string> UPLScripts, VersionNumber SdkVersion, string BundleID, bool bBuildAsFramework, out bool bSupportsPortrait, out bool bSupportsLandscape, out bool bSkipIcons)
+		public override bool GeneratePList(FileReference ProjectFile, UnrealTargetConfiguration Config, string ProjectDirectory, bool bIsUE4Game, string GameName, bool bIsClient, string ProjectName, string InEngineDir, string AppDirectory, List<string> UPLScripts, VersionNumber SdkVersion, string BundleID, bool bBuildAsFramework, out bool bSupportsPortrait, out bool bSupportsLandscape, out bool bSkipIcons)
 		{
 			bSupportsLandscape = bSupportsPortrait = true;
             bSkipIcons = true;
-			return GenerateTVOSPList(ProjectDirectory, bIsUE4Game, GameName, ProjectName, InEngineDir, AppDirectory, null, BundleID);
+			return GenerateTVOSPList(ProjectDirectory, bIsUE4Game, GameName, bIsClient, ProjectName, InEngineDir, AppDirectory, null, BundleID);
 		}
 
         protected override void CopyGraphicsResources(bool bSkipDefaultPNGs, bool bSkipIcons, string InEngineDir, string AppDirectory, string BuildDirectory, string IntermediateDir, bool bSupportsPortrait, bool bSupportsLandscape)
diff --git a/Engine/Source/Programs/UnrealBuildTool/ProjectFiles/Xcode/XcodeProject.cs b/Engine/Source/Programs/UnrealBuildTool/ProjectFiles/Xcode/XcodeProject.cs
--- a/Engine/Source/Programs/UnrealBuildTool/ProjectFiles/Xcode/XcodeProject.cs
+++ b/Engine/Source/Programs/UnrealBuildTool/ProjectFiles/Xcode/XcodeProject.cs
@@ -1265,12 +1265,12 @@ namespace UnrealBuildTool
 							TargetReceipt.TryRead(ReceiptFilename, out Receipt);
 							VersionNumber SdkVersion = UEDeployIOS.GetSdkVersion(Receipt);
 							bool bBuildAsFramework = UEDeployIOS.GetCompileAsDll(Receipt);
-							UEDeployIOS.GenerateIOSPList(ProjectFile, Config.BuildConfig, ProjectPath.FullName, bIsUE4Game, GameName, Config.BuildTarget, EngineDir.FullName, ProjectPath + "/Binaries/IOS/Payload", SdkVersion, null, BundleIdentifier, bBuildAsFramework, out bSupportPortrait, out bSupportLandscape, out bSkipIcons);
+							UEDeployIOS.GenerateIOSPList(ProjectFile, Config.BuildConfig, ProjectPath.FullName, bIsUE4Game, GameName, false, Config.BuildTarget, EngineDir.FullName, ProjectPath + "/Binaries/IOS/Payload", SdkVersion, null, BundleIdentifier, bBuildAsFramework, out bSupportPortrait, out bSupportLandscape, out bSkipIcons);
 						}
 						if (bCreateTVOSInfoPlist)
 						{
 							Directory.CreateDirectory(Path.GetDirectoryName(TVOSInfoPlistPath));
-							UEDeployTVOS.GenerateTVOSPList(ProjectPath.FullName, bIsUE4Game, GameName, Config.BuildTarget, EngineDir.FullName, ProjectPath + "/Binaries/TVOS/Payload", null, BundleIdentifier);
+							UEDeployTVOS.GenerateTVOSPList(ProjectPath.FullName, bIsUE4Game, GameName, false, Config.BuildTarget, EngineDir.FullName, ProjectPath + "/Binaries/TVOS/Payload", null, BundleIdentifier);
 						}
 					}
 				}
diff --git a/Engine/Source/Programs/UnrealBuildTool/ToolChain/RemoteMac.cs b/Engine/Source/Programs/UnrealBuildTool/ToolChain/RemoteMac.cs
--- a/Engine/Source/Programs/UnrealBuildTool/ToolChain/RemoteMac.cs
+++ b/Engine/Source/Programs/UnrealBuildTool/ToolChain/RemoteMac.cs
@@ -58,7 +58,7 @@ namespace UnrealBuildTool
 		/// The authentication used for Rsync (for the -e rsync flag).
 		/// </summary>
 		[XmlConfigFile]
-		private string RsyncAuthentication = "ssh -i '${CYGWIN_SSH_PRIVATE_KEY}'";
+		private string RsyncAuthentication = "./ssh -i '${CYGWIN_SSH_PRIVATE_KEY}'";
 
 		/// <summary>
 		/// The authentication used for SSH (probably similar to RsyncAuthentication).
diff --git a/Engine/Source/Runtime/ApplicationCore/Private/GenericPlatform/GenericApplication.cpp b/Engine/Source/Runtime/ApplicationCore/Private/GenericPlatform/GenericApplication.cpp
--- a/Engine/Source/Runtime/ApplicationCore/Private/GenericPlatform/GenericApplication.cpp
+++ b/Engine/Source/Runtime/ApplicationCore/Private/GenericPlatform/GenericApplication.cpp
@@ -3,6 +3,7 @@
 #include "GenericPlatform/GenericApplication.h"
 #include "HAL/IConsoleManager.h"
 #include "Misc/CoreDelegates.h"
+#include "Misc/CommandLine.h"
 #include "IInputDevice.h"
 
 const FGamepadKeyNames::Type FGamepadKeyNames::Invalid(NAME_None);
@@ -197,11 +198,22 @@ float FDisplayMetrics::GetDebugActionSafeZoneRatio()
 
 void FDisplayMetrics::ApplyDefaultSafeZones()
 {
-	const float SafeZoneRatio = GetDebugTitleSafeZoneRatio();
-	if (SafeZoneRatio < 1.0f)
+	// Allow safe zones to be overridden by the command line. Used by mobile PIE.
+	TitleSafePaddingSize = FVector4(0.0f, 0.0f, 0.0f, 0.0f);
+	bool bSetByCommandLine;
+	bSetByCommandLine = FParse::Value(FCommandLine::Get(), TEXT("SafeZonePaddingLeft="),   TitleSafePaddingSize.X);
+	bSetByCommandLine = FParse::Value(FCommandLine::Get(), TEXT("SafeZonePaddingRight="),  TitleSafePaddingSize.Y) || bSetByCommandLine;
+	bSetByCommandLine = FParse::Value(FCommandLine::Get(), TEXT("SafeZonePaddingTop="),    TitleSafePaddingSize.Z) || bSetByCommandLine;
+	bSetByCommandLine = FParse::Value(FCommandLine::Get(), TEXT("SafeZonePaddingBottom="), TitleSafePaddingSize.W) || bSetByCommandLine;
+
+	if (!bSetByCommandLine)
 	{
-		const float HalfUnsafeRatio = (1.0f - SafeZoneRatio) * 0.5f;
-		TitleSafePaddingSize = FVector4(PrimaryDisplayWidth * HalfUnsafeRatio, PrimaryDisplayHeight * HalfUnsafeRatio, PrimaryDisplayWidth * HalfUnsafeRatio, PrimaryDisplayHeight * HalfUnsafeRatio);
+		const float SafeZoneRatio = GetDebugTitleSafeZoneRatio();
+		if (SafeZoneRatio < 1.0f)
+		{
+			const float HalfUnsafeRatio = (1.0f - SafeZoneRatio) * 0.5f;
+			TitleSafePaddingSize = FVector4(PrimaryDisplayWidth * HalfUnsafeRatio, PrimaryDisplayHeight * HalfUnsafeRatio, PrimaryDisplayWidth * HalfUnsafeRatio, PrimaryDisplayHeight * HalfUnsafeRatio);
+		}
 	}
 
 	const float ActionSafeZoneRatio = GetDebugActionSafeZoneRatio();
diff --git a/Engine/Source/Runtime/Engine/Private/GameViewportClient.cpp b/Engine/Source/Runtime/Engine/Private/GameViewportClient.cpp
--- a/Engine/Source/Runtime/Engine/Private/GameViewportClient.cpp
+++ b/Engine/Source/Runtime/Engine/Private/GameViewportClient.cpp
@@ -2446,6 +2446,9 @@ void UGameViewportClient::DrawTitleSafeArea( UCanvas* Canvas )
 	FCanvasTileItem TileItem(FVector2D::ZeroVector, GWhiteTexture, UnsafeZoneColor);
 	TileItem.BlendMode = SE_BLEND_Translucent;
 
+	// Command line override used by mobile PIE.
+	static bool bDrawUnSafeZones = FParse::Param(FCommandLine::Get(), TEXT("DrawUnSafeZones"));
+
 	// CalculateSafeZoneValues() can be slow, so we only want to run it if we have boundaries to draw
 	if (FDisplayMetrics::GetDebugTitleSafeZoneRatio() < 1.f)
 	{
@@ -2471,7 +2474,7 @@ void UGameViewportClient::DrawTitleSafeArea( UCanvas* Canvas )
 		TileItem.Size = FVector2D(SafeZone.Right, HeightOfSides);
 		Canvas->DrawItem(TileItem);
 	}
-	else if (!FSlateApplication::Get().GetCustomSafeZone().GetDesiredSize().IsZero())
+	else if (!FSlateApplication::Get().GetCustomSafeZone().GetDesiredSize().IsZero() || bDrawUnSafeZones)
 	{
 		ULevelEditorPlaySettings* PlaySettings = GetMutableDefault<ULevelEditorPlaySettings>();
 		PlaySettings->CalculateCustomUnsafeZones(PlaySettings->CustomUnsafeZoneStarts, PlaySettings->CustomUnsafeZoneDimensions, PlaySettings->DeviceToEmulate, FVector2D(Width, Height));
diff --git a/Engine/Source/Runtime/Engine/Private/PrimitiveSceneProxy.cpp b/Engine/Source/Runtime/Engine/Private/PrimitiveSceneProxy.cpp
--- a/Engine/Source/Runtime/Engine/Private/PrimitiveSceneProxy.cpp
+++ b/Engine/Source/Runtime/Engine/Private/PrimitiveSceneProxy.cpp
@@ -403,6 +403,12 @@ void FPrimitiveSceneProxy::SetTransform(const FMatrix& InLocalToWorld, const FBo
 	LocalBounds = InLocalBounds;
 	ActorPosition = InActorPosition;
 	
+	// Update cached reflection capture.
+	if (PrimitiveSceneInfo)
+	{
+		PrimitiveSceneInfo->bNeedsCachedReflectionCaptureUpdate = true;
+	}
+	
 	UpdateUniformBuffer();
 	
 	// Notify the proxy's implementation of the change.
diff --git a/Engine/Source/Runtime/Renderer/Private/PostProcess/PostProcessing.cpp b/Engine/Source/Runtime/Renderer/Private/PostProcess/PostProcessing.cpp
--- a/Engine/Source/Runtime/Renderer/Private/PostProcess/PostProcessing.cpp
+++ b/Engine/Source/Runtime/Renderer/Private/PostProcess/PostProcessing.cpp
@@ -1773,7 +1773,7 @@ void FPostProcessing::ProcessES2(FRHICommandListImmediate& RHICmdList, FScene* S
 			Desc.Format = PF_B8G8R8A8;
 			Desc.NumMips = 1;
 			Desc.DebugName = TEXT("OverriddenRenderTarget");
-			Desc.TargetableFlags |= TexCreate_RenderTargetable;
+			Desc.TargetableFlags |= TexCreate_RenderTargetable | TexCreate_ShaderResource;
 #if !(UE_BUILD_SHIPPING || UE_BUILD_TEST)
 			// for FVisualizeTexture
 			Desc.TargetableFlags |= TexCreate_ShaderResource;
