Fix for subsurface flickering artifact caused by lack of memory barrier group sync

#rnx

diff --git a/Engine/Shaders/Private/PostProcessSubsurface.usf b/Engine/Shaders/Private/PostProcessSubsurface.usf
--- a/Engine/Shaders/Private/PostProcessSubsurface.usf
+++ b/Engine/Shaders/Private/PostProcessSubsurface.usf
@@ -502,7 +502,7 @@ void SetupIndirectCS(uint3 DT_ID : SV_DispatchThreadID,
 	{
 		SubsurfaceTypeFlag = 0;
 	}
-	GroupMemoryBarrier();
+	GroupMemoryBarrierWithGroupSync();
 
 	uint2 Pos = DT_ID.xy + Output_ViewportMin;
 	bool bHasSubsurface = false;
