Fixing burley SSS when opacity is very low.

diff --git a/Engine/Shaders/Private/ShadingModelsMaterial.ush b/Engine/Shaders/Private/ShadingModelsMaterial.ush
--- a/Engine/Shaders/Private/ShadingModelsMaterial.ush
+++ b/Engine/Shaders/Private/ShadingModelsMaterial.ush
@@ -57,7 +57,7 @@ void SetGBufferForShadingModel(
 	{
 		// Optimization: if opacity is 0 then revert to default shading model 
 	#if SUBSURFACE_PROFILE_OPACITY_THRESHOLD
-		if (Opacity > 0.0)
+		if (Opacity > SSSS_OPACITY_THRESHOLD_EPS)
 	#endif
 		{
 			GBuffer.CustomData.rgb = EncodeSubsurfaceProfile(SubsurfaceProfile);
@@ -67,8 +67,9 @@ void SetGBufferForShadingModel(
 			uint SubsurfaceProfileUInt = uint(SubsurfaceProfile * 255.0f + 0.5f);
 			float MaterialRoughnessToAverage = ActualSSProfilesTexture.Load(int3(SSSS_DUAL_SPECULAR_OFFSET, SubsurfaceProfileUInt, 0)).w * SSSS_MAX_DUAL_SPECULAR_ROUGHNESS;
 
-			// Smooth blend out dual specular when opacity is low.
-			MaterialRoughnessToAverage = lerp(1.0f, MaterialRoughnessToAverage, saturate(Opacity * 10.0f));
+			// Smooth blend out dual specular when opacity is low, we have the extra SSSS_OPACITY_THRESHOLD_EPS so that we fade out by the time we
+			// get to 0.01, as opposed to 0.0.
+			MaterialRoughnessToAverage = lerp(1.0f, MaterialRoughnessToAverage, saturate((Opacity - SSSS_OPACITY_THRESHOLD_EPS) * 10.0f));
 
 			GBuffer.Roughness = saturate(GBuffer.Roughness * MaterialRoughnessToAverage);
 		}
diff --git a/Engine/Shaders/Private/SubsurfaceBurleyNormalized.ush b/Engine/Shaders/Private/SubsurfaceBurleyNormalized.ush
--- a/Engine/Shaders/Private/SubsurfaceBurleyNormalized.ush
+++ b/Engine/Shaders/Private/SubsurfaceBurleyNormalized.ush
@@ -9,7 +9,6 @@
 #include "BurleyNormalizedSSSCommon.ush"
 #include "DeferredShadingCommon.ush"
 
-
 // Setup the max number of Burley samples
 #define BURLEY_NUM_SAMPLES	64
 #define BURLEY_INV_NUM_SAMPLES (1.0f/BURLEY_NUM_SAMPLES)
@@ -135,6 +134,7 @@ struct FBurleyParameter
 	float4 SurfaceAlbedo;
 	float4 DiffuseMeanFreePath;
 	float  WorldUnitScale;
+	float  SurfaceOpacity; // 1.0 means full Burley, 0.0 means full Default Lit
 };
 
 //Surface albedo and mean free path length
@@ -190,7 +190,8 @@ FBurleyParameter GetBurleyParameters(FGBufferData GBuffer)
 
 	// Mask should always be true, we can probably ignore it?
 	float Mask = UseSubsurfaceProfile(GBuffer.ShadingModelID) ? 1.0f : 0.0f;
-	BurleyParameter.DiffuseMeanFreePath *= Mask * GBuffer.CustomData.a;
+	BurleyParameter.SurfaceOpacity = Mask * GBuffer.CustomData.a;
+	BurleyParameter.DiffuseMeanFreePath *= BurleyParameter.SurfaceOpacity;
 
 	return BurleyParameter;
 }
@@ -646,14 +647,10 @@ float4 BurleyNormalizedSS(float2 BufferUV, uint2 GroupThreadID)
 	FScreenSpaceData ScreenSpaceData = GetScreenSpaceData(BufferUV);
 	float3 WorldNormal = ScreenSpaceData.GBuffer.WorldNormal;
 
-#if SUBSURFACE_SINGLE_PASS
-	float DepthAtDiscCenter = CalcSceneDepth(BufferUV);
-	uint ReconstructMethod = (SUBSURFACE_RECOMBINE_QUALITY == SUBSURFACE_RECOMBINE_QUALITY_HIGH) ? 3 : 1;
-#else
-	//float DepthAtDiscCenter = SampleSSSColorConsideringLocalShared(BufferUV, 0, GroupThreadID, 0).DiffuseLighting.w;
 	BurleySampleDiffuseNormal CenterSample = SampleSSSColorConsideringLocalShared(BufferUV, 0, GroupThreadID, 0);
 	float DepthAtDiscCenter = CenterSample.DiffuseLighting.w;
-#endif
+
+	float3 OriginalColor = CenterSample.DiffuseLighting.rgb;
 
 	float4 OutColor = 0;
 
@@ -817,7 +814,13 @@ float4 BurleyNormalizedSS(float2 BufferUV, uint2 GroupThreadID)
 	RadianceAccumulated.xyz *= (WeightingFactor == 0)? 0:(1.0f /WeightingFactor*EnergyNormalization);
 #endif
 
-	OutColor.xyz = RadianceAccumulated.xyz;
+	// The opacity works by reducing the radius based on opacity, but this runs into precision issues with low opacity values.
+	// So as the opacity goes from 0.02 to 0.01, we transition to fully disabling SSS by the time we get to 0.01.
+	float LowOpacityEps = SSSS_OPACITY_THRESHOLD_EPS;
+
+	float OriginalLerp = saturate((BurleyParameter.SurfaceOpacity - LowOpacityEps) / LowOpacityEps);
+
+	OutColor.xyz = lerp(OriginalColor,RadianceAccumulated.xyz,OriginalLerp);
 	OutColor.w = NumOfSamples;
 
 	return OutColor;
diff --git a/Engine/Shaders/Private/SubsurfaceProfileCommon.ush b/Engine/Shaders/Private/SubsurfaceProfileCommon.ush
--- a/Engine/Shaders/Private/SubsurfaceProfileCommon.ush
+++ b/Engine/Shaders/Private/SubsurfaceProfileCommon.ush
@@ -28,6 +28,9 @@
 #define	SSSS_MAX_TRANSMISSION_PROFILE_DISTANCE	5.0f // See MaxTransmissionProfileDistance in ComputeTransmissionProfile(), SeparableSSS.cpp
 #define SSSS_MAX_DUAL_SPECULAR_ROUGHNESS		2.0f
 
+// Threshold value at which model switches from SSS to default lit
+#define SSSS_OPACITY_THRESHOLD_EPS 0.01
+
 // One profile per line, encoded using constants above. See FSubsurfaceProfileTexture::CreateTexture() in SubsurfaceProfile.cpp.
 // By default, SSProfilesTexture is "Texture2D SSProfilesTexture", but it can be remapped as in BasePassCommon.usf.
 #ifndef SSProfilesTexture
