GTAO - Added default filter size of 5x5 rather than 4x4

diff --git a/Engine/Shaders/Private/PostProcessAmbientOcclusion.usf b/Engine/Shaders/Private/PostProcessAmbientOcclusion.usf
--- a/Engine/Shaders/Private/PostProcessAmbientOcclusion.usf
+++ b/Engine/Shaders/Private/PostProcessAmbientOcclusion.usf
@@ -31,7 +31,10 @@
 #define GTAO_THICKNESS_HEURISTIC 1
 
 // Define to Use Cosine Weighted Integral
-#define GTAO_PROJECTED_NORMAL_WEIGHTING 0
+#define GTAO_PROJECTED_NORMAL_WEIGHTING 1
+
+#define GTAO_PATTERN_SIZE 5
+
 
 // Define to Use Cosine Weighted Integral
 #define GTAO_COSINE_WEIGHTED_INTEGRAL 1
@@ -1220,7 +1223,9 @@ float ComputeInnerIntegral(float2 Angles, float3 ScreenDir, float3 ViewDir, floa
 	float AO = ( LenProjNormal *  0.25 * 
 					    ( (Angles.x * SinGamma + CosGamma - cos((2.0 * Angles.x) - Gamma)) +
 				  	      (Angles.y * SinGamma + CosGamma - cos((2.0 * Angles.y) - Gamma)) ));
+
 	AO = saturate(AO / PI_HALF);
+
 #else
 	// Uniform weighted AO
 	float AO =  2 - (cos(Angles.x) + cos(Angles.y));
@@ -1239,20 +1244,35 @@ float ComputeInnerIntegral(float2 Angles, float3 ScreenDir, float3 ViewDir, floa
 
 
 
-float3 GetRandomVector(float2 UV )
+
+
+float3 GetRandomVector(float2 UV, uint2 iPos )
 {
-	float2 ViewportUVToRandomUV = ScreenSpaceAOParams[1].xy;
-	float3 RandomTexVec = Texture2DSample(RandomNormalTexture, RandomNormalTextureSampler, UV * ViewportUVToRandomUV).rgb;
-	RandomTexVec.xy = RandomTexVec.xy * 2 -1;
+	float3 RandomVec=float3(0,0,0);
+	float3 RandomTexVec=float3(0,0,0);
+	float ScaleOffset;
 
 	float TemporalCos = GTAOParams[0].x;
 	float TemporalSin = GTAOParams[0].y;
-	float3 RandomVec;
+
+#if GTAO_PATTERN_SIZE == 4
+	float2 ViewportUVToRandomUV = ScreenSpaceAOParams[1].xy;
+	RandomTexVec = Texture2DSample(RandomNormalTexture, RandomNormalTextureSampler, UV * ViewportUVToRandomUV).rgb;
+	ScaleOffset = RandomTexVec.z;
+	RandomTexVec.xy = RandomTexVec.xy * 2 -1;
+#else
+	// 5x5 So compute this analytically
+	float Rot	= (1.0/25.0) * ((((iPos.x + iPos.y) % 5) * 5 ) + (iPos.x % 5));
+	ScaleOffset = (1.0/5.0)  *  ((iPos.y - iPos.x) % 5);
+
+	RandomTexVec.x = cos(Rot*PI);
+	RandomTexVec.y = sin(Rot*PI);
+#endif
 
 	// Rotate the random Vec
 	RandomVec.x = dot(RandomTexVec.xy, float2(TemporalCos, -TemporalSin ));
     RandomVec.y = dot(RandomTexVec.xy, float2(TemporalSin, TemporalCos ));
-	RandomVec.z = RandomTexVec.z;
+	RandomVec.z = ScaleOffset;
 
 	return RandomVec;
 }
@@ -1264,7 +1284,7 @@ float3 GetRandomVector(float2 UV )
 * HORIZON SEARCH AND INNER INTEGRATE COMBINED
 *
 */
-void GTAOCombinedPSandCS(in float2 UV, out float4 OutColor, inout float OutDepth)
+void GTAOCombinedPSandCS(in float2 UV, in uint2 iPos, out float4 OutColor, inout float OutDepth)
 {
 	OutColor = 0;
 
@@ -1280,7 +1300,7 @@ void GTAOCombinedPSandCS(in float2 UV, out float4 OutColor, inout float OutDepth
 		return;
 	}
 	
-	float3 RandomAndOffset	= GetRandomVector(UV);
+	float3 RandomAndOffset	= GetRandomVector(UV, iPos);
 	float2 RandomVec		= RandomAndOffset.xy;
 	float  Offset			= RandomAndOffset.z;
 	Offset += GTAOParams[0].z;
@@ -1339,7 +1359,10 @@ void GTAOCombinedPSandCS(in float2 UV, out float4 OutColor, inout float OutDepth
 void GTAOCombinedPS(in float4 UVAndScreenPos : TEXCOORD0, out float4 OutColor : SV_Target0)
 {
 	float Depth =0;
-	GTAOCombinedPSandCS(UVAndScreenPos.xy, OutColor, Depth);	
+	int2 iPos = int2( UVAndScreenPos.xy * PostprocessInput0Size.xy );
+
+
+	GTAOCombinedPSandCS(UVAndScreenPos.xy, iPos, OutColor, Depth);	
 }
 
 
@@ -1357,7 +1380,7 @@ void GTAOCombinedCS(
 	float2 PixelCenter	= (float2)PixelPos + float2(0.5, 0.5);
 	float2 BufferUV		= PixelCenter.xy * GTAOParams[2].zw;
 
-	GTAOCombinedPSandCS(BufferUV , OutColor, OutDepth);
+	GTAOCombinedPSandCS(BufferUV , PixelPos,  OutColor, OutDepth);
 
 	OutTexture[PixelPos]		= OutColor.r;
 	DepthOutTexture[PixelPos]   = OutDepth;
@@ -1378,7 +1401,7 @@ SamplerState	HorizonsTextureSampler;
 float4			InnerIntegrateParams;
 
 
-float GTAOInnerIntegratePSandCS(in float2 UV)
+float GTAOInnerIntegratePSandCS(in float2 UV, in uint2 iPos)
 {
 	// Read the angles buffer
 	float SceneDepth	 = GetDepthFromAOInput(UV);		
@@ -1391,7 +1414,7 @@ float GTAOInnerIntegratePSandCS(in float2 UV)
 	Angles				    = Angles *2-1;
 
 	// Get Angle
-	float2 RandomVec		= GetRandomVector(UV).xy;
+	float2 RandomVec		= GetRandomVector(UV, iPos).xy;
 	float3 ScreenDir		= float3(RandomVec.x, RandomVec.y, 0.0);
 
 	// ViewspacePos and Normal
@@ -1413,7 +1436,9 @@ float GTAOInnerIntegratePSandCS(in float2 UV)
 
 void GTAOInnerIntegratePS(in noperspective float4 UVAndScreenPos : TEXCOORD0, out float4 OutColor : SV_Target0)
 {
-	float AO = 	GTAOInnerIntegratePSandCS(UVAndScreenPos.xy);	
+	int2 iPos = int2( UVAndScreenPos.xy * PostprocessInput0Size.xy );
+
+	float AO = 	GTAOInnerIntegratePSandCS(UVAndScreenPos.xy, iPos);	
 	OutColor = AO;
 }
 
@@ -1432,7 +1457,7 @@ void GTAOInnerIntegrateCS(
 	float2 PixelCenter	= (float2)PixelPos + float2(0.5, 0.5);
 	float2 BufferUV		= PixelCenter.xy * GTAOParams[2].zw;
 
-	float AO = GTAOInnerIntegratePSandCS(BufferUV);
+	float AO = GTAOInnerIntegratePSandCS(BufferUV, PixelPos);
 	
 	OutTexture[PixelPos] = AO;
 }
@@ -1451,7 +1476,7 @@ void GTAOInnerIntegrateCS(
 
 float4 HorizonSearchParams;
 
-float2 HorizonSearchPSandCS(in float2 UV, inout float OutDepth)
+float2 HorizonSearchPSandCS(in float2 UV,  in uint2 iPos, inout float OutDepth)
 {
 	float2 OutHorizons = 0;
 
@@ -1465,7 +1490,7 @@ float2 HorizonSearchPSandCS(in float2 UV, inout float OutDepth)
 		return OutHorizons;
 	}
 
-	float3 RandomAndOffset		= GetRandomVector(UV);
+	float3 RandomAndOffset		= GetRandomVector(UV, iPos);
 	float2 RandomVec			= RandomAndOffset.xy;
 	float  Offset				= RandomAndOffset.z;
 	float3 ScreenDir			= float3(RandomVec.x, RandomVec.y, 0.0);
@@ -1494,8 +1519,9 @@ float2 HorizonSearchPSandCS(in float2 UV, inout float OutDepth)
 void HorizonSearchPS(in noperspective float4 UVAndScreenPos : TEXCOORD0, float4 SvPosition : SV_POSITION, out float4 OutColor : SV_Target0)
 {
 	float Depth;
+	int2 iPos = int2( UVAndScreenPos.xy * PostprocessInput0Size.xy );
 
-	float2 Horizons = HorizonSearchPSandCS(UVAndScreenPos.xy, Depth);	
+	float2 Horizons = HorizonSearchPSandCS(UVAndScreenPos.xy, iPos, Depth);	
 	OutColor = float4(Horizons.xy,0,0);
 }
 
@@ -1515,7 +1541,7 @@ void HorizonSearchCS(
 	float2 BufferUV		= PixelCenter.xy * GTAOParams[2].zw;
 
 	float OutDepth;
-	float2 Horizons = HorizonSearchPSandCS(BufferUV,OutDepth);
+	float2 Horizons = HorizonSearchPSandCS(BufferUV,PixelPos,OutDepth);
 	
 	HorizonOutTexture[PixelPos] = Horizons;
 	DepthOutTexture[PixelPos]   = OutDepth;
@@ -1670,7 +1696,7 @@ void GTAOTemporalFilterPSAndCS(float2 UV, inout float OutAO, inout float OutDept
 	{
 		float DestDeviceZ		= Texture2DSample(ZCurrTexture, ZCurrTextureSampler, PrevUVDepth.xy).r;
 		float3 Reproj			= ReprojectPos( PrevUVDepth.xy,  DestDeviceZ); 
-		DestVeloc				= PrevUVDepth - Reproj.xy;
+		DestVeloc				= PrevUVDepth.xy - Reproj.xy;
 	}
 
 	float VelocCompare			= CompareVeloc(PixVelocity, DestVeloc);
@@ -1713,7 +1739,7 @@ void GTAOTemporalFilterPSAndCS2(float2 UV, inout float OutAO, inout float OutDep
 	{
 		float DestDeviceZ		= Texture2DSample(ZCurrTexture, ZCurrTextureSampler, PrevUVDepth.xy).r;
 		float3 Blah				= ReprojectPos( PrevUVDepth.xy,  DestDeviceZ); 
-		DestVeloc				= PrevUVDepth-Blah.xy;
+		DestVeloc				= PrevUVDepth.xy-Blah.xy;
 	}
 
 	// Compare velocities 
@@ -1925,7 +1951,172 @@ void GTAOUpsamplePS(in noperspective float4 UVAndScreenPos : TEXCOORD0,  out flo
 
 #if COMPUTE_SHADER
 
+#if SPATIAL_FILTER_SIZE == 5
+
+// The 5x5 filter works on a threadgroup of size 16x8 (128 pixels)
+// We need to read in the 16x8 and a 2 pixel border around. So this is 20x12 (240 pixels)
+// Each thread reads in 2 pixels each
+
+groupshared float AOData[ 16*8*2]; 
+groupshared float ZData[  16*8*2];
+
+float GetAO(int x, int y)
+{
+	x+=2; y+=2;
+	return AOData [(y*20) + x];
+}
+
+float GetZ(int x, int y)
+{
+	x+=2; y+=2;
+	return ZData [(y*20) + x];
+}
+
+
+[numthreads(16, 8, 1)]
+void GTAOSpatialFilter5x5CS(
+	int   GroupIndex: SV_GroupIndex,
+	uint2 GroupId : SV_GroupID,
+	uint2 DispatchThreadId : SV_DispatchThreadID,
+	uint2 GroupThreadId : SV_GroupThreadID) 
+{
+	int2 GTId = int2(GroupThreadId);
+
+	
+	// Position on the screen We care about
+	int2   PixelPos				= DispatchThreadId + ScreenSpaceAOParams[5].zw; 
+
+	// Firstly get the origin in the screen of the 16x8 inner box
+	int2   FullGroupOrigin		= int2(GroupId.x * 16, GroupId.y * 8) + ScreenSpaceAOParams[5].zw; 
+	int2   FullGroupOriginM2	= FullGroupOrigin.xy - int2(2,2);
+
+	uint pixIdx = (GroupIndex*2);
+	uint XPos = pixIdx%20;
+	uint YPos = pixIdx/20;
+	float AO = PostprocessInput0.Load(int3(FullGroupOriginM2 + int2(XPos,YPos), 0)).r;
+	float  Z = ZReadTexture.Load(     int3(FullGroupOriginM2 + int2(XPos,YPos), 0)).r;
+	AOData[ (YPos*20) + XPos ] = AO;
+	ZData[  (YPos*20) + XPos ] = Z;
+
+	// Next pixel
+	XPos++;
+	AO = PostprocessInput0.Load(int3(FullGroupOriginM2 + int2(XPos,YPos), 0)).r;
+	Z  = ZReadTexture.Load(     int3(FullGroupOriginM2 + int2(XPos,YPos), 0)).r;
+	AOData[ (YPos*20) + XPos ] = AO;
+	ZData[  (YPos*20) + XPos ] = Z;
+
+	GroupMemoryBarrierWithGroupSync();
+
+	// Get the differences in Z at this pixel. This is needed for the bilateral filter
+	float ThisZ	= GetZ(GTId.x, GTId.y);
+
+	float2 ZDiff;
+
+	{
+		float XM2Z	= GetZ(GTId.x-2, GTId.y);
+		float XM1Z	= GetZ(GTId.x-1, GTId.y);
+		float XP1Z	= GetZ(GTId.x+1, GTId.y);
+		float XP2Z	= GetZ(GTId.x+2, GTId.y);
+
+		// Get extrapolated point either side
+		float C1 = abs((XM1Z + (XM1Z - XM2Z)) - ThisZ);
+		float C2 = abs((XP1Z + (XP1Z - XP2Z)) - ThisZ);
+
+		if(C1 < C2)
+		{
+			ZDiff.x = XM1Z - XM2Z;
+		}
+		else
+		{
+			ZDiff.x = XP2Z - XP1Z;
+		}
+	}
+
+
+	{
+		float YM2Z	= GetZ(GTId.x, GTId.y-2);
+		float YM1Z	= GetZ(GTId.x, GTId.y-1);
+		float YP1Z	= GetZ(GTId.x, GTId.y+1);
+		float YP2Z	= GetZ(GTId.x, GTId.y+2);
+
+		// Get extrapolated point either side
+		float C1 = abs((YM1Z + (YM1Z - YM2Z)) - ThisZ);
+		float C2 = abs((YP1Z + (YP1Z - YP2Z)) - ThisZ);
+
+		if(C1 < C2)
+		{
+			ZDiff.y = YM1Z - YM2Z;
+		}
+		else
+		{
+			ZDiff.y = YP2Z - YP1Z;
+		}
+	}
+
+	// Do the blur
+	float SumAO		= 0;
+	float SumWeight = 0;
+
+	int x,y;
+
+	// Get the Z Value to compare against 
+
+	float DepthBase = ThisZ  -(ZDiff.x*2) -(ZDiff.y*2);
+
+	for(y=-2; y<=2; y++)
+	{
+		float PlaneZ = DepthBase;
+
+		for(x=-2; x<=2; x++)
+		{
+			// Get value and see how much it compares to the centre with the gradients
+			float XDiff = abs(x);
+
+			float Sample_AO = GetAO(GTId.x + x, GTId.y + y);
+			float SampleZ   = GetZ( GTId.x + x, GTId.y + y);
+
+			float Weight = 1.0f;
+			if((x==0) && (y==0))
+			{
+				Weight = 1.0f;
+			}
+			else
+			{
+				// Get the bilateral weight. This is a function of the difference in height between the plane equation and the base depth
+				// Compare the Z at this sample with the gradients 
+				float SampleZDiff = abs(PlaneZ - SampleZ);
+			
+				const float SpatialFilterWeight = 1000;
+				Weight = 1.0f - saturate(SampleZDiff*SpatialFilterWeight );
+			}
+
+			SumAO += Sample_AO * Weight;
+			SumWeight += Weight;
+		
+			PlaneZ += ZDiff.x;
+		}
+		DepthBase += ZDiff.y;
+	}
+	SumAO /=SumWeight;
+
+	SumAO = saturate(SumAO * PI_HALF);
+
+	float AmbientOcclusionIntensity = ScreenSpaceAOParams[0].w;
+	float AmbientOcclusionPower		= ScreenSpaceAOParams[0].x*0.5;
+	// user adjust AO
+	SumAO = 1 - (1 - pow(abs(SumAO), AmbientOcclusionPower)) * AmbientOcclusionIntensity;
+
+
+	OutTexture[PixelPos] = SumAO ;
+
+}
+
+
+#else
+
+
 #define ARRAY_SIZE ((THREADGROUP_SIZEX*2) * (THREADGROUP_SIZEY+4) )
+
 groupshared float AOArray[ARRAY_SIZE]; 
 groupshared float ZArray[  ARRAY_SIZE];
 groupshared float ZDiffArray[  THREADGROUP_SIZEX * THREADGROUP_SIZEY];
@@ -1958,7 +2149,7 @@ float GetZVal(int x, int y)
 
 
 [numthreads(THREADGROUP_SIZEX, THREADGROUP_SIZEY, 1)]
-void GTAOSpatialFilterCS(
+void GTAOSpatialFilter4x4CS(
 	int   GroupIndex: SV_GroupIndex,
 	uint2 GroupId : SV_GroupID,
 	uint2 DispatchThreadId : SV_DispatchThreadID,
@@ -2056,7 +2247,6 @@ void GTAOSpatialFilterCS(
 }
 #endif
 
-void GTAOSpatialFilterPS(in noperspective float4 UVAndScreenPos : TEXCOORD0,  out float4 OutColor : SV_Target0)
-{
-	OutColor = Texture2DSample(PostprocessInput0, PostprocessInput0Sampler, UVAndScreenPos.xy);
-}
\ No newline at end of file
+
+#endif
+
diff --git a/Engine/Source/Runtime/Renderer/Private/CompositionLighting/PostProcessAmbientOcclusion.cpp b/Engine/Source/Runtime/Renderer/Private/CompositionLighting/PostProcessAmbientOcclusion.cpp
--- a/Engine/Source/Runtime/Renderer/Private/CompositionLighting/PostProcessAmbientOcclusion.cpp
+++ b/Engine/Source/Runtime/Renderer/Private/CompositionLighting/PostProcessAmbientOcclusion.cpp
@@ -109,6 +109,14 @@ static TAutoConsoleVariable<int32> CVarGTAOUseNormals(
 	TEXT("1: On (default)\n "),
 	ECVF_RenderThreadSafe | ECVF_Scalability);
 
+static TAutoConsoleVariable<int32> CVarGTAOFilterWidth(
+	TEXT("r.GTAO.FilterWidth"),
+	5,
+	TEXT("Size of the noise pattern and filter width\n ")
+	TEXT("5: 5x5 Pattern (default) \n ")
+	TEXT("4: 4x4 Pattern \n "),
+	ECVF_RenderThreadSafe | ECVF_Scalability);
+
 static TAutoConsoleVariable<float> CVarGTAOThicknessBlend(
 	TEXT("r.GTAO.ThicknessBlend"),
 	0.03f,
@@ -2620,10 +2628,10 @@ FPooledRenderTargetDesc FRCPassPostProcessAmbientOcclusion_GTAO_TemporalFilter::
 }
 
 
-template<uint32 bComputeShader>
-class FPostProcessGTAOSpatialFilterPSandCS : public FGlobalShader
+template<uint32 FilterSize>
+class FPostProcessGTAOSpatialFilterCS : public FGlobalShader
 {
-	DECLARE_SHADER_TYPE(FPostProcessGTAOSpatialFilterPSandCS, Global);
+	DECLARE_SHADER_TYPE(FPostProcessGTAOSpatialFilterCS, Global);
 
 	static bool ShouldCompilePermutation(const FGlobalShaderPermutationParameters& Parameters)
 	{
@@ -2633,15 +2641,22 @@ class FPostProcessGTAOSpatialFilterPSandCS : public FGlobalShader
 	static void ModifyCompilationEnvironment(const FGlobalShaderPermutationParameters& Parameters, FShaderCompilerEnvironment& OutEnvironment)
 	{
 		FGlobalShader::ModifyCompilationEnvironment(Parameters, OutEnvironment);
-		OutEnvironment.SetDefine(TEXT("COMPUTE_SHADER"), bComputeShader);
-		if (bComputeShader)
+		OutEnvironment.SetDefine(TEXT("COMPUTE_SHADER"), true);
+		OutEnvironment.SetDefine(TEXT("SPATIAL_FILTER_SIZE"), FilterSize);
+		if (FilterSize == 5)
+		{
+			OutEnvironment.SetDefine(TEXT("THREADGROUP_SIZEX"), 16);
+		}
+		else
 		{
 			OutEnvironment.SetDefine(TEXT("THREADGROUP_SIZEX"), 8);
-			OutEnvironment.SetDefine(TEXT("THREADGROUP_SIZEY"), 8);
 		}
+		OutEnvironment.SetDefine(TEXT("THREADGROUP_SIZEY"), 8);
+
+
 	}
 
-	FPostProcessGTAOSpatialFilterPSandCS() {}
+	FPostProcessGTAOSpatialFilterCS() {}
 	LAYOUT_FIELD(FPostProcessPassParameters, PostprocessParameter);
 	LAYOUT_FIELD(FShaderParameter, OutTexture);
 	LAYOUT_FIELD(FShaderParameter, FilterParams);
@@ -2652,7 +2667,7 @@ class FPostProcessGTAOSpatialFilterPSandCS : public FGlobalShader
 public:
 
 	/** Initialization constructor. */
-	FPostProcessGTAOSpatialFilterPSandCS(const ShaderMetaType::CompiledShaderInitializerType& Initializer)
+	FPostProcessGTAOSpatialFilterCS(const ShaderMetaType::CompiledShaderInitializerType& Initializer)
 		: FGlobalShader(Initializer)
 	{
 		PostprocessParameter.Bind(Initializer.ParameterMap);
@@ -2715,12 +2730,12 @@ public:
 
 	static const TCHAR* GetFunctionName()
 	{
-		return bComputeShader ? TEXT("GTAOSpatialFilterCS") : TEXT("GTAOSpatialFilterPS");
+		return FilterSize==4 ? TEXT("GTAOSpatialFilter4x4CS") : TEXT("GTAOSpatialFilter5x5CS");
 	}
 };
 
-IMPLEMENT_SHADER_TYPE2(FPostProcessGTAOSpatialFilterPSandCS<true>,  SF_Compute);
-IMPLEMENT_SHADER_TYPE2(FPostProcessGTAOSpatialFilterPSandCS<false>, SF_Pixel);
+IMPLEMENT_SHADER_TYPE2(FPostProcessGTAOSpatialFilterCS<4>,  SF_Compute);
+IMPLEMENT_SHADER_TYPE2(FPostProcessGTAOSpatialFilterCS<5>,  SF_Compute);
 
 FRCPassPostProcessAmbientOcclusion_GTAO_SpatialFilter::FRCPassPostProcessAmbientOcclusion_GTAO_SpatialFilter(const FSceneView& View, uint32 InDownScaleFactor, const EGTAOType InAOType)
 	: AOType(InAOType),
@@ -2744,33 +2759,68 @@ void FRCPassPostProcessAmbientOcclusion_GTAO_SpatialFilter::Process(FRenderingCo
 
 	// Compute  Version
 	UnbindRenderTargets(Context.RHICmdList);
-	TShaderMapRef<FPostProcessGTAOSpatialFilterPSandCS<true>> ComputeShader(Context.GetShaderMap());
-	uint32 GroupSizeX = FMath::DivideAndRoundUp(OutputViewRect.Width(), 8);
-	uint32 GroupSizeY = FMath::DivideAndRoundUp(OutputViewRect.Height(), 8);
 
-	if (AOType == EGTAOType::EAsyncCombinedSpatial)
+	// Get filter size
+	uint32 FilterWidth = CVarGTAOFilterWidth.GetValueOnRenderThread();
+
+	if (FilterWidth == 4)
 	{
-		// If the Spatial Filter is running as part of the async then we'll render to the R channel of the horizons texture so it can be read in as part of the temporal
-		const FSceneRenderTargetItem& DestRenderTarget = SceneContext.ScreenSpaceGTAOHorizons->GetRenderTargetItem();
+		TShaderMapRef<FPostProcessGTAOSpatialFilterCS<4>> ComputeShader(Context.GetShaderMap());
+		uint32 GroupSizeX = FMath::DivideAndRoundUp(OutputViewRect.Width(), 8);
+		uint32 GroupSizeY = FMath::DivideAndRoundUp(OutputViewRect.Height(), 8);
+		if (AOType == EGTAOType::EAsyncCombinedSpatial)
+		{
+			const FSceneRenderTargetItem& DestRenderTarget = SceneContext.ScreenSpaceGTAOHorizons->GetRenderTargetItem();
 
-		FRHIAsyncComputeCommandListImmediate& RHICmdListComputeImmediate = FRHICommandListExecutor::GetImmediateAsyncComputeCommandList();
+			FRHIAsyncComputeCommandListImmediate& RHICmdListComputeImmediate = FRHICommandListExecutor::GetImmediateAsyncComputeCommandList();
+			RHICmdListComputeImmediate.SetComputeShader(ComputeShader->GetComputeShader());
+			RHICmdListComputeImmediate.TransitionResource(EResourceTransitionAccess::ERWBarrier, EResourceTransitionPipeline::EComputeToCompute, DestRenderTarget.UAV, nullptr);
+			ComputeShader->SetParametersCS(RHICmdListComputeImmediate, Context, InputTexSize, OutputViewRect, DestRenderTarget.UAV);
+			DispatchComputeShader(RHICmdListComputeImmediate, *ComputeShader, GroupSizeX, GroupSizeY, 1);
+			ComputeShader->UnsetParameters(Context.RHICmdList);
+		}
+		else
+		{
+			const FSceneRenderTargetItem& DestRenderTarget = PassOutputs[0].RequestSurface(Context);
 
-		RHICmdListComputeImmediate.SetComputeShader(ComputeShader.GetComputeShader());
-		RHICmdListComputeImmediate.TransitionResource(EResourceTransitionAccess::ERWBarrier, EResourceTransitionPipeline::EComputeToCompute, DestRenderTarget.UAV, nullptr);
-		ComputeShader->SetParametersCS(RHICmdListComputeImmediate, Context, InputTexSize, OutputViewRect, DestRenderTarget.UAV);
-		DispatchComputeShader(RHICmdListComputeImmediate, ComputeShader, GroupSizeX, GroupSizeY, 1);
+			Context.RHICmdList.SetComputeShader(ComputeShader->GetComputeShader());
+			Context.RHICmdList.TransitionResource(EResourceTransitionAccess::ERWBarrier, EResourceTransitionPipeline::EGfxToCompute, DestRenderTarget.UAV, nullptr);
+			ComputeShader->SetParametersCS(Context.RHICmdList, Context, InputTexSize, OutputViewRect, DestRenderTarget.UAV);
+			DispatchComputeShader(Context.RHICmdList, *ComputeShader, GroupSizeX, GroupSizeY, 1);
+			ComputeShader->UnsetParameters(Context.RHICmdList);
+		}
 	}
 	else
 	{
-		const FSceneRenderTargetItem& DestRenderTarget = PassOutputs[0].RequestSurface(Context);
+		TShaderMapRef<FPostProcessGTAOSpatialFilterCS<5>> ComputeShader(Context.GetShaderMap());
+		uint32 GroupSizeX = FMath::DivideAndRoundUp(OutputViewRect.Width(), 16);
+		uint32 GroupSizeY = FMath::DivideAndRoundUp(OutputViewRect.Height(), 8);
+		if (AOType == EGTAOType::EAsyncCombinedSpatial)
+		{
+			const FSceneRenderTargetItem& DestRenderTarget = SceneContext.ScreenSpaceGTAOHorizons->GetRenderTargetItem();
 
-		Context.RHICmdList.SetComputeShader(ComputeShader.GetComputeShader());
-		Context.RHICmdList.TransitionResource(EResourceTransitionAccess::ERWBarrier, EResourceTransitionPipeline::EGfxToCompute, DestRenderTarget.UAV, nullptr);
-		ComputeShader->SetParametersCS(Context.RHICmdList, Context, InputTexSize, OutputViewRect, DestRenderTarget.UAV);
-		DispatchComputeShader(Context.RHICmdList, ComputeShader, GroupSizeX, GroupSizeY, 1);
+			FRHIAsyncComputeCommandListImmediate& RHICmdListComputeImmediate = FRHICommandListExecutor::GetImmediateAsyncComputeCommandList();
+			RHICmdListComputeImmediate.SetComputeShader(ComputeShader->GetComputeShader());
+			RHICmdListComputeImmediate.TransitionResource(EResourceTransitionAccess::ERWBarrier, EResourceTransitionPipeline::EComputeToCompute, DestRenderTarget.UAV, nullptr);
+			ComputeShader->SetParametersCS(RHICmdListComputeImmediate, Context, InputTexSize, OutputViewRect, DestRenderTarget.UAV);
+			DispatchComputeShader(RHICmdListComputeImmediate, *ComputeShader, GroupSizeX, GroupSizeY, 1);
+			ComputeShader->UnsetParameters(Context.RHICmdList);
+		}
+		else
+		{
+			const FSceneRenderTargetItem& DestRenderTarget = PassOutputs[0].RequestSurface(Context);
+
+			Context.RHICmdList.SetComputeShader(ComputeShader->GetComputeShader());
+			Context.RHICmdList.TransitionResource(EResourceTransitionAccess::ERWBarrier, EResourceTransitionPipeline::EGfxToCompute, DestRenderTarget.UAV, nullptr);
+			ComputeShader->SetParametersCS(Context.RHICmdList, Context, InputTexSize, OutputViewRect, DestRenderTarget.UAV);
+			DispatchComputeShader(Context.RHICmdList, *ComputeShader, GroupSizeX, GroupSizeY, 1);
+			ComputeShader->UnsetParameters(Context.RHICmdList);
+		}
 	}
 
-	ComputeShader->UnsetParameters(Context.RHICmdList);
+
+
+
 }
 
 
