Remove loose global parameters from landscape rendering.
Most were moved to a shared uniform buffer recently but a few remained.
The only one now left is the mobile LodBias which is actually the TexCoordOffset, so it's been renamed to match the intent.
Cherry pick integration from Dev-Rendering

diff --git a/Engine/Shaders/Private/LandscapeVertexFactory.ush b/Engine/Shaders/Private/LandscapeVertexFactory.ush
--- a/Engine/Shaders/Private/LandscapeVertexFactory.ush
+++ b/Engine/Shaders/Private/LandscapeVertexFactory.ush
@@ -11,62 +11,75 @@
 // Set by FLightMapTexturePolicy
 #include "/Engine/Generated/UniformBuffers/PrecomputedLightingBuffer.ush"
 
-/* SM4+:
-   x = unused
-   y = unused
-   z = Heightmap texture LOD difference from current LOD to highest LOD
-   w = XYOffset texture LOD difference from current LOD to highest LOD
-
-   ES3.1:
-   x = unused
-   y = unused
-   z = TexCoordOffset.x
-   w = TexCoordOffset.y
-*/
-float4 LodBias;
-
-/* x = current LOD
-   y = unused
-   z = SectionSizeQuads in current LOD
-   w = 1/SectionSizeQuads in current LOD
-*/
-float4 LodValues;
-
-/* Fractional Lod values for each section */
-float4 SectionLods;
+#if LANDSCAPE_XYOFFSET
+	#define XYOFFSET_SCALE (1.0f/256.f)
+#endif
 
-int ForcedLod;
+#if (ES3_1_PROFILE)
+/* Offset for UV localization for large UV values. */
+float2 TexCoordOffset;
+#endif
 
+/* Return index into buffers for this section instance. */
 int GetComponentLinearIndex()
 {
 	return (LandscapeParameters.ComponentBaseY - LandscapeContinuousLODParameters.Min.y) * LandscapeContinuousLODParameters.Size.x + (LandscapeParameters.ComponentBaseX - LandscapeContinuousLODParameters.Min.x);
 }
 
-float4 GetSectionLods()
+/**
+ * x = Current mesh LOD
+ * y = unused
+ * z = SectionSizeQuads in current mesh LOD
+ * w = 1/SectionSizeQuads in current mesh LOD
+ **/
+float4 GetLodValues()
+{
+#if FIXED_GRID 
+	// Fixed grid has fixed Lod
+	return LandscapeFixedGrid.LodValues;
+#elif FEATURE_LEVEL < FEATURE_LEVEL_SM4
+	// On mobile we always use the Lod 0 mesh data
+	return float4(0, 0, LandscapeParameters.SubsectionSizeVertsLayerUVPan.xy);
+#else
+	// Get Lod for the current section and infer the values
+	uint Lod = (uint)floor(LandscapeContinuousLODParameters.SectionLOD[GetComponentLinearIndex()]);
+	float LodSubSectionSizeQuads = (float)((LandscapeParameters.SubsectionSizeVerts >> Lod) - 1);
+	return float4(Lod, 0, LodSubSectionSizeQuads, rcp(LodSubSectionSizeQuads));
+#endif
+}
+
+/**
+ * x = Heightmap texture LOD difference from current LOD to highest LOD
+ * y = todo: XYOffset texture LOD difference from current LOD to highest LOD
+ */
+float2 GetLodBias()
+{
+#if FIXED_GRID || FEATURE_LEVEL < FEATURE_LEVEL_SM4
+	return 0;
+#else
+	float LodBias = LandscapeContinuousLODParameters.SectionLODBias[GetComponentLinearIndex()];
+	return float2(LodBias, 0.f);
+#endif
+}
+
+/* Get current Lod for this section instance. */
+float GetSectionLod()
 {
-	if (ForcedLod >= 0)
-	{
-		return LodValues.x;
-	}
-	
 	return min(LandscapeContinuousLODParameters.SectionLOD[GetComponentLinearIndex()] + View.FarShadowStaticMeshLODBias, LandscapeParameters.LastLOD);
 }
 
+/** 
+ * Get current Lods for neighbour section instances. 
+ *       -Y
+ *    - - 0 - -
+ *    |       |
+ * -X 1   P   2 +X
+ *    |       |
+ *    - - 3 - -
+ *       +Y
+ */
 float4 GetNeighborSectionLod()
 {
-	if (ForcedLod >= 0)
-	{
-		return LodValues.xxxx;
-	}
-
-	//       -Y       
-	//    - - 0 - -   
-	//    |       |   
-	// -X 1   P   2 +X
-	//    |       |   
-	//    - - 3 - -   
-	//       +Y
-
 	int2 NeighborOffset[4];
     NeighborOffset[0] = int2(0, -1);
     NeighborOffset[1] = int2(-1, 0);
@@ -78,41 +91,12 @@ float4 GetNeighborSectionLod()
 		int2 NeighborPos = NeighborOffset[i] + int2(LandscapeParameters.ComponentBaseX, LandscapeParameters.ComponentBaseY);
 		NeighborPos = max(LandscapeContinuousLODParameters.Min                                                     , NeighborPos);
 		NeighborPos = min(LandscapeContinuousLODParameters.Min + LandscapeContinuousLODParameters.Size - int2(1, 1), NeighborPos);
-		NeighborLod[i] = max(LandscapeContinuousLODParameters.SectionLOD[(NeighborPos.y - LandscapeContinuousLODParameters.Min.y) * LandscapeContinuousLODParameters.Size.x + (NeighborPos.x - LandscapeContinuousLODParameters.Min.x)], GetSectionLods());
+		NeighborLod[i] = max(LandscapeContinuousLODParameters.SectionLOD[(NeighborPos.y - LandscapeContinuousLODParameters.Min.y) * LandscapeContinuousLODParameters.Size.x + (NeighborPos.x - LandscapeContinuousLODParameters.Min.x)], GetSectionLod());
 		NeighborLod[i] = min(NeighborLod[i] + View.FarShadowStaticMeshLODBias, LandscapeParameters.LastLOD);
 	}
 	return float4(NeighborLod[0], NeighborLod[1], NeighborLod[2], NeighborLod[3]);
 }
 
-float4 GetLodBias()
-{
-#if FIXED_GRID || FEATURE_LEVEL < FEATURE_LEVEL_SM4
-	return LodBias;
-#else
-	return float4(
-		0,
-		0,
-		LandscapeContinuousLODParameters.SectionLODBias[GetComponentLinearIndex()],
-		0); // TODO: handle XY offset texture mipmap bias
-#endif
-}
-
-/* Array index is section
-   x = neighbor section's LOD in position 0               0     -> X
-   y = neighbor section's LOD in position 1        1   Current  2
-   z = neighbor section's LOD in position 2        |      3
-   w = neighbor section's LOD in position 3        V Y
-*/
-float4 NeighborSectionLod[4];
-
-#if LANDSCAPE_XYOFFSET
-	#define XYOFFSET_SCALE			(1.0f/256.f)
-#endif
-
-#if ES3_1_PROFILE
-float3 BlendableLayerMask;
-#endif
-
 struct FVertexFactoryInput
 {
 #if FEATURE_LEVEL >= FEATURE_LEVEL_SM4
@@ -129,6 +113,7 @@ FVertexFactoryInput LoadVertexFactoryInputForHGS(uint TriangleIndex, int VertexI
 {
 	FVertexFactoryInput Input;
 
+	float4 LodValues = GetLodValues();
 	uint SubsectionSizeQuadsCurrentLOD = LodValues.z;
 	uint SubsectionSizeVertsCurrentLOD = LodValues.z + 1;
 
@@ -162,6 +147,7 @@ FVertexFactoryInput LoadVertexFactoryInputForDynamicUpdate(uint TriangleIndex, i
 {
 	FVertexFactoryInput Input;
 
+	float4 LodValues = GetLodValues();
 	uint SubsectionSizeQuadsCurrentLOD = LodValues.z;
 	uint SubsectionSizeVertsCurrentLOD = LodValues.z + 1;
 
@@ -272,6 +258,8 @@ struct FVertexFactoryIntermediates
 	float4 InputPosition;
 	float3 LocalPosition;
 	float3 WorldNormal;
+	float4 LodValues;
+	float2 LodBias;
 };
 
 float3 GetLocalPosition(FVertexFactoryIntermediates Intermediates)
@@ -318,7 +306,7 @@ FLandscapeTexCoords GetLandscapeTexCoords(FVertexFactoryInput Input, FVertexFact
 #endif
 
 #if NEEDS_LIGHTMAP_COORDINATE
-	Result.LightMapCoordinate			= (Intermediates.LocalPosition.xy * LandscapeParameters.LandscapeLightmapScaleBias.xy + LandscapeParameters.LandscapeLightmapScaleBias.wz + Intermediates.InputPosition.zw * LandscapeParameters.LightmapSubsectionOffsetParams.xy);
+	Result.LightMapCoordinate	= (Intermediates.LocalPosition.xy * LandscapeParameters.LandscapeLightmapScaleBias.xy + LandscapeParameters.LandscapeLightmapScaleBias.wz + Intermediates.InputPosition.zw * LandscapeParameters.LightmapSubsectionOffsetParams.xy);
 #endif
 	
 	return Result;
@@ -497,9 +485,8 @@ FMaterialPixelParameters GetMaterialPixelParameters(FVertexFactoryInterpolantsVS
 #endif
 
 #if (ES3_1_PROFILE) && PIXELSHADER
-
 	Result.LayerWeights.xy = Texture2DSample(LandscapeParameters.NormalmapTexture, LandscapeParameters.NormalmapTextureSampler, Interpolants.WeightMapTexCoord).rg;
-	Result.LayerWeights.z = saturate(BlendableLayerMask.z - dot(Result.LayerWeights.xy, BlendableLayerMask.xy));
+	Result.LayerWeights.z = saturate(LandscapeParameters.BlendableLayerMask.z - dot(Result.LayerWeights.xy, LandscapeParameters.BlendableLayerMask.xy));
 	Result.LayerWeights.w = 0; // Total weight need to be 1
 #endif
 
@@ -525,7 +512,7 @@ FMaterialVertexParameters GetMaterialVertexParameters(FVertexFactoryInput Input,
 #if NUM_MATERIAL_TEXCOORDS_VERTEX     // XY layer
 	Result.TexCoords[0] = LandscapeTexCoords.LayerTexCoord.xy;
 	#if (ES3_1_PROFILE)
-		Result.TexCoordOffset = LodBias.zw;
+		Result.TexCoordOffset = TexCoordOffset;
 	#endif
 #if NUM_MATERIAL_TEXCOORDS_VERTEX > 1 // XZ layer
 	Result.TexCoords[1] = float2(LandscapeTexCoords.LayerTexCoord.x, Intermediates.LocalPosition.z);
@@ -592,7 +579,7 @@ float CalcLOD(float2 xyLocalToSubsection, float2 Subsection)
 	// LOD value based on Barycentric coordinate
 	float4 L0 = float4(xy.y, xy.x, (1 - xy.x), (1 - xy.y)) * 2;
 
-	float4 LODCalculated4 = L0 * GetSectionLods() + (float4(1, 1, 1, 1) - L0) * GetNeighborSectionLod();
+	float4 LODCalculated4 = L0 * GetSectionLod() + (float4(1, 1, 1, 1) - L0) * GetNeighborSectionLod();
 
 	// Find which quadrant of the subsection we're in - top, left, right or bottom.
 	// 0 --------- 1 X  1 ----- 2  1          1 
@@ -610,6 +597,8 @@ float CalcLOD(float2 xyLocalToSubsection, float2 Subsection)
 FVertexFactoryIntermediates GetVertexFactoryIntermediates(FVertexFactoryInput Input)
 {
 	FVertexFactoryIntermediates Intermediates;
+	Intermediates.LodValues = GetLodValues();
+	Intermediates.LodBias = GetLodBias();
 
 #if FEATURE_LEVEL >= FEATURE_LEVEL_SM4
 	Intermediates.InputPosition = Input.Position;
@@ -622,10 +611,10 @@ FVertexFactoryIntermediates GetVertexFactoryIntermediates(FVertexFactoryInput In
 	float HoleVertexLod = (float)((uint)(Input.PackedPosition.w * 255.f) & 0xf);
 #endif
 
-	float2 xy = Intermediates.InputPosition.xy * LodValues.w;
+	float2 xy = Intermediates.InputPosition.xy * Intermediates.LodValues.w;
 
 #if FIXED_GRID
-	float LODCalculated = floor(LodValues.x);
+	float LODCalculated = floor(Intermediates.LodValues.x);
 	float LodValue = LODCalculated;
 	float MorphAlpha = 0;
 #else
@@ -636,7 +625,7 @@ FVertexFactoryIntermediates GetVertexFactoryIntermediates(FVertexFactoryInput In
 
 	// InputPositionLODAdjusted : Position for actual LOD in base LOD units
 #if FEATURE_LEVEL >= FEATURE_LEVEL_SM4
-	float2 ActualLODCoordsInt = floor(Intermediates.InputPosition.xy * pow(2, -(LodValue - LodValues.x)));
+	float2 ActualLODCoordsInt = floor(Intermediates.InputPosition.xy * pow(2, -(LodValue - Intermediates.LodValues.x)));
 	float InvLODScaleFactor = pow(2, -LodValue);
 #else
 
@@ -680,20 +669,20 @@ FVertexFactoryIntermediates GetVertexFactoryIntermediates(FVertexFactoryInput In
 	// Get the height and normal XY for current and next LOD out of the textures
 #if FEATURE_LEVEL >= FEATURE_LEVEL_SM4
 	float2 SampleCoords = InputPositionLODAdjusted * LandscapeParameters.HeightmapUVScaleBias.xy + LandscapeParameters.HeightmapUVScaleBias.zw + 0.5*LandscapeParameters.HeightmapUVScaleBias.xy + Intermediates.InputPosition.zw * LandscapeParameters.SubsectionOffsetParams.xy;
-	float4 SampleValue = Texture2DSampleLevel(LandscapeParameters.HeightmapTexture, LandscapeParameters.HeightmapTextureSampler, SampleCoords, LodValue-GetLodBias().z);
+	float4 SampleValue = Texture2DSampleLevel(LandscapeParameters.HeightmapTexture, LandscapeParameters.HeightmapTextureSampler, SampleCoords, LodValue-Intermediates.LodBias.x);
 	float Height = DecodePackedHeight(SampleValue.xy);
 
 	float2 SampleCoordsNextLOD = InputPositionNextLOD * LandscapeParameters.HeightmapUVScaleBias.xy + LandscapeParameters.HeightmapUVScaleBias.zw + 0.5*LandscapeParameters.HeightmapUVScaleBias.xy + Intermediates.InputPosition.zw * LandscapeParameters.SubsectionOffsetParams.xy;
-	float4 SampleValueNextLOD = Texture2DSampleLevel(LandscapeParameters.HeightmapTexture, LandscapeParameters.HeightmapTextureSampler, SampleCoordsNextLOD, LodValue+1-GetLodBias().z);
+	float4 SampleValueNextLOD = Texture2DSampleLevel(LandscapeParameters.HeightmapTexture, LandscapeParameters.HeightmapTextureSampler, SampleCoordsNextLOD, LodValue+1-Intermediates.LodBias.x);
 	float HeightNextLOD = DecodePackedHeight(SampleValueNextLOD.xy);
 #endif
 
 
 #if LANDSCAPE_XYOFFSET // FEATURE_LEVEL >= FEATURE_LEVEL_SM4 only
 	float2 SampleCoords2 = float2(InputPositionLODAdjusted * LandscapeParameters.WeightmapUVScaleBias.xy + LandscapeParameters.WeightmapUVScaleBias.zw + Intermediates.InputPosition.zw * LandscapeParameters.SubsectionOffsetParams.zz);
-	float4 OffsetValue = Texture2DSampleLevel( LandscapeParameters.XYOffsetmapTexture, LandscapeParameters.XYOffsetmapTextureSampler, SampleCoords2, LodValue-GetLodBias().w );
+	float4 OffsetValue = Texture2DSampleLevel( LandscapeParameters.XYOffsetmapTexture, LandscapeParameters.XYOffsetmapTextureSampler, SampleCoords2, LodValue- Intermediates.LodBias.y );
 	float2 SampleCoordsNextLOD2 = float2(InputPositionNextLOD * LandscapeParameters.WeightmapUVScaleBias.xy + LandscapeParameters.WeightmapUVScaleBias.zw + Intermediates.InputPosition.zw * LandscapeParameters.SubsectionOffsetParams.zz);
-	float4 OffsetValueNextLOD = Texture2DSampleLevel( LandscapeParameters.XYOffsetmapTexture, LandscapeParameters.XYOffsetmapTextureSampler, SampleCoordsNextLOD2, LodValue+1-GetLodBias().w );
+	float4 OffsetValueNextLOD = Texture2DSampleLevel( LandscapeParameters.XYOffsetmapTexture, LandscapeParameters.XYOffsetmapTextureSampler, SampleCoordsNextLOD2, LodValue+1-Intermediates.LodBias.y );
 	float2 XYOffset = float2(((OffsetValue.r * 255.0 * 256.0 + OffsetValue.g * 255.0) - 32768.0) * XYOFFSET_SCALE, ((OffsetValue.b * 255.0 * 256.0 + OffsetValue.a * 255.0) - 32768.0) * XYOFFSET_SCALE );
 	float2 XYOffsetNextLOD = float2(((OffsetValueNextLOD.r * 255.0 * 256.0 + OffsetValueNextLOD.g * 255.0) - 32768.0) * XYOFFSET_SCALE, ((OffsetValueNextLOD.b * 255.0 * 256.0 + OffsetValueNextLOD.a * 255.0) - 32768.0) * XYOFFSET_SCALE );
 
@@ -725,17 +714,17 @@ float2 LandscapeVertexFactorySampleHeight(FVertexFactoryInput Input, FVertexFact
 	return SampleValue.xy;
 	*/
 
-	float2 xy = Intermediates.InputPosition.xy * LodValues.w;
+	float2 xy = Intermediates.InputPosition.xy * Intermediates.LodValues.w;
 	float LODCalculated = CalcLOD(xy, Intermediates.InputPosition.zw);
 	float LodValue = floor(LODCalculated);
 	float CoordTranslate_x = (LandscapeParameters.SubsectionSizeVertsLayerUVPan.x * pow(2, -LodValue) - 1) * LandscapeParameters.SubsectionSizeVertsLayerUVPan.y;
 
 	// InputPositionLODAdjusted : Position for actual LOD in base LOD units
-	float2 ActualLODCoordsInt = floor(Intermediates.InputPosition.xy * pow(2, -(LodValue - LodValues.x)));
+	float2 ActualLODCoordsInt = floor(Intermediates.InputPosition.xy * pow(2, -(LodValue - Intermediates.LodValues.x)));
 
 	float2 InputPositionLODAdjusted = ActualLODCoordsInt / CoordTranslate_x;
 	float2 SampleCoords = InputPositionLODAdjusted * LandscapeParameters.HeightmapUVScaleBias.xy + LandscapeParameters.HeightmapUVScaleBias.zw + 0.5*LandscapeParameters.HeightmapUVScaleBias.xy + Intermediates.InputPosition.zw * LandscapeParameters.SubsectionOffsetParams.xy;
-	float4 SampleValue = Texture2DSampleLevel(LandscapeParameters.HeightmapTexture, LandscapeParameters.HeightmapTextureSampler, SampleCoords, LodValue-GetLodBias().z);
+	float4 SampleValue = Texture2DSampleLevel(LandscapeParameters.HeightmapTexture, LandscapeParameters.HeightmapTextureSampler, SampleCoords, LodValue-Intermediates.LodBias.x);
 	return SampleValue.rg;
 }
 #endif
@@ -757,7 +746,7 @@ float3x3 VertexFactoryGetTangentToLocal( FVertexFactoryInput Input, FVertexFacto
 float4 VertexFactoryGetRasterizedWorldPosition(FVertexFactoryInput Input, FVertexFactoryIntermediates Intermediates, float4 InWorldPosition)
 {
 #if RENDERING_LANDSCAPE_GRASS
-	float2 xy = Intermediates.InputPosition.xy * LodValues.w;
+	float2 xy = Intermediates.InputPosition.xy * Intermediates.LodValues.w;
 	float LODCalculated = CalcLOD(xy, Intermediates.InputPosition.zw);
 	float LodValue = floor(LODCalculated);
 	float CoordTranslate_x = (LandscapeParameters.SubsectionSizeVertsLayerUVPan.x * pow(2, -LodValue) - 1) * LandscapeParameters.SubsectionSizeVertsLayerUVPan.y;
@@ -862,7 +851,7 @@ FVertexFactoryInterpolantsVSToDS VertexFactoryGetInterpolantsVSToDS(FVertexFacto
 	Interpolants.TangentToWorld0 = TangentToWorld[0];
 	Interpolants.TangentToWorld2 = float4(TangentToWorld[2], 1);
 
-	float2 xy = Intermediates.InputPosition.xy * LodValues.w;
+	float2 xy = Intermediates.InputPosition.xy * Intermediates.LodValues.w;
 	Interpolants.CalculatedLOD = CalcLOD(xy, Intermediates.InputPosition.zw);
 #endif
 	
diff --git a/Engine/Source/Runtime/Landscape/Private/LandscapeRender.cpp b/Engine/Source/Runtime/Landscape/Private/LandscapeRender.cpp
--- a/Engine/Source/Runtime/Landscape/Private/LandscapeRender.cpp
+++ b/Engine/Source/Runtime/Landscape/Private/LandscapeRender.cpp
@@ -43,6 +43,7 @@ LandscapeRender.cpp: New terrain rendering
 #include "ProfilingDebugging/LoadTimeTracker.h"
 
 IMPLEMENT_GLOBAL_SHADER_PARAMETER_STRUCT(FLandscapeUniformShaderParameters, "LandscapeParameters");
+IMPLEMENT_GLOBAL_SHADER_PARAMETER_STRUCT(FLandscapeFixedGridUniformShaderParameters, "LandscapeFixedGrid");
 IMPLEMENT_TYPE_LAYOUT(FLandscapeVertexFactoryPixelShaderParameters);
 
 int32 GLandscapeMeshLODBias = 0;
@@ -1128,6 +1129,7 @@ FLandscapeComponentSceneProxy::FLandscapeComponentSceneProxy(ULandscapeComponent
 	, BaseColorForGITexture(InComponent->GIBakedBaseColorTexture)
 	, HeightmapScaleBias(InComponent->HeightmapScaleBias)
 	, XYOffsetmapTexture(InComponent->XYOffsetmapTexture)
+	, BlendableLayerMask(InComponent->MobileBlendableLayerMask)
 	, SharedBuffersKey(0)
 	, SharedBuffers(nullptr)
 	, VertexFactory(nullptr)
@@ -1460,6 +1462,20 @@ void FLandscapeComponentSceneProxy::CreateRenderThreadResources()
 	// Assign LandscapeUniformShaderParameters
 	LandscapeUniformShaderParameters.InitResource();
 
+	// Create per Lod uniform buffers
+	LandscapeFixedGridUniformShaderParameters.AddDefaulted(MaxLOD+1);
+	for (int32 LodIndex = 0; LodIndex <= MaxLOD; ++LodIndex)
+	{
+		LandscapeFixedGridUniformShaderParameters[LodIndex].InitResource();
+		FLandscapeFixedGridUniformShaderParameters Parameters;
+		Parameters.LodValues = FVector4(
+			LodIndex, 
+			0.f,
+			(float)((SubsectionSizeVerts >> LodIndex) - 1),
+			1.f / (float)((SubsectionSizeVerts >> LodIndex) - 1));
+		LandscapeFixedGridUniformShaderParameters[LodIndex].SetContents(Parameters);
+	}
+
 #if WITH_EDITOR
 	// Create MeshBatch for grass rendering
 	if (SharedBuffers->GrassIndexBuffer)
@@ -1483,13 +1499,9 @@ void FLandscapeComponentSceneProxy::CreateRenderThreadResources()
 		// Combined grass rendering batch element
 		FMeshBatchElement* GrassBatchElement = &GrassMeshBatch.Elements[0];
 		FLandscapeBatchElementParams* BatchElementParams = &GrassBatchParams[0];
-		BatchElementParams->LocalToWorldNoScalingPtr = &LocalToWorldNoScaling;
 		BatchElementParams->LandscapeUniformShaderParametersResource = &LandscapeUniformShaderParameters;
 		BatchElementParams->SceneProxy = this;
-		BatchElementParams->SubX = -1;
-		BatchElementParams->SubY = -1;
 		BatchElementParams->CurrentLOD = 0;
-		BatchElementParams->ForcedLOD = 0;
 		GrassBatchElement->UserData = BatchElementParams;
 		GrassBatchElement->PrimitiveUniformBuffer = GetUniformBuffer();
 		GrassBatchElement->IndexBuffer = SharedBuffers->GrassIndexBuffer;
@@ -1562,6 +1574,12 @@ FLandscapeComponentSceneProxy::~FLandscapeComponentSceneProxy()
 	// Free the subsection uniform buffer
 	LandscapeUniformShaderParameters.ReleaseResource();
 
+	// Free the lod uniform buffers
+	for (int32 i = 0; i < LandscapeFixedGridUniformShaderParameters.Num(); ++i)
+	{
+		LandscapeFixedGridUniformShaderParameters[i].ReleaseResource();
+	}
+
 	if (SharedBuffers)
 	{
 		check(SharedBuffers == FLandscapeComponentSceneProxy::SharedBuffersMap.FindRef(SharedBuffersKey));
@@ -1861,6 +1879,12 @@ void FLandscapeComponentSceneProxy::OnTransformChanged()
 		0,
 		0
 	);
+	LandscapeParams.BlendableLayerMask = FVector4(
+		BlendableLayerMask & (1 << 0) ? 1 : 0,
+		BlendableLayerMask & (1 << 1) ? 1 : 0,
+		BlendableLayerMask & (1 << 2) ? 1 : 0,
+		0
+	);
 
 	if (HeightmapTexture)
 	{
@@ -1999,11 +2023,8 @@ void FLandscapeComponentSceneProxy::BuildDynamicMeshElement(const FViewCustomDat
 
 				if (!InToolMesh)
 				{
-					BatchElementParams.LocalToWorldNoScalingPtr = &LocalToWorldNoScaling;
 					BatchElementParams.LandscapeUniformShaderParametersResource = &LandscapeUniformShaderParameters;
 					BatchElementParams.SceneProxy = this;
-					BatchElementParams.SubX = SubX;
-					BatchElementParams.SubY = SubY;
 					BatchElementParams.CurrentLOD = CurrentLOD;
 				}
 
@@ -2051,10 +2072,7 @@ void FLandscapeComponentSceneProxy::BuildDynamicMeshElement(const FViewCustomDat
 		{
 			FLandscapeBatchElementParams& BatchElementParams = OutStaticBatchParamArray[0];
 			BatchElementParams.LandscapeUniformShaderParametersResource = &LandscapeUniformShaderParameters;
-			BatchElementParams.LocalToWorldNoScalingPtr = &LocalToWorldNoScaling;
 			BatchElementParams.SceneProxy = this;
-			BatchElementParams.SubX = -1;
-			BatchElementParams.SubY = -1;
 			BatchElementParams.CurrentLOD = CurrentLODIndex;
 
 			BatchElement.UserData = &BatchElementParams;
@@ -2121,10 +2139,7 @@ bool FLandscapeComponentSceneProxy::GetMeshElement(bool UseSeperateBatchForShado
 				{
 					FLandscapeBatchElementParams* BatchElementParams = new(OutStaticBatchParamArray) FLandscapeBatchElementParams;
 					BatchElementParams->LandscapeUniformShaderParametersResource = &LandscapeUniformShaderParameters;
-					BatchElementParams->LocalToWorldNoScalingPtr = &LocalToWorldNoScaling;
 					BatchElementParams->SceneProxy = this;
-					BatchElementParams->SubX = SubX;
-					BatchElementParams->SubY = SubY;
 					BatchElementParams->CurrentLOD = i;
 
 					FMeshBatchElement BatchElement;
@@ -2153,10 +2168,7 @@ bool FLandscapeComponentSceneProxy::GetMeshElement(bool UseSeperateBatchForShado
 		// Combined batch element
 		FLandscapeBatchElementParams* BatchElementParams = new(OutStaticBatchParamArray) FLandscapeBatchElementParams;
 		BatchElementParams->LandscapeUniformShaderParametersResource = &LandscapeUniformShaderParameters;
-		BatchElementParams->LocalToWorldNoScalingPtr = &LocalToWorldNoScaling;
 		BatchElementParams->SceneProxy = this;
-		BatchElementParams->SubX = -1;
-		BatchElementParams->SubY = -1;
 		BatchElementParams->CurrentLOD = i;
 
 		FMeshBatchElement BatchElement;
@@ -2205,10 +2217,7 @@ bool FLandscapeComponentSceneProxy::GetMeshElementForVirtualTexture(int32 InLodI
 	FLandscapeBatchElementParams* BatchElementParams = new(OutStaticBatchParamArray) FLandscapeBatchElementParams;
 	BatchElementParams->SceneProxy = this;
 	BatchElementParams->LandscapeUniformShaderParametersResource = &LandscapeUniformShaderParameters;
-	BatchElementParams->LocalToWorldNoScalingPtr = &LocalToWorldNoScaling;
 	BatchElementParams->CurrentLOD = InLodIndex;
-	BatchElementParams->SubX = -1;
-	BatchElementParams->SubY = -1;
 
 	int32 LodSubsectionSizeVerts = SubsectionSizeVerts >> InLodIndex;
 
@@ -2315,12 +2324,8 @@ bool FLandscapeComponentSceneProxy::GetStaticMeshElement(int32 LODIndex, bool bF
 
 		FLandscapeBatchElementParams* BatchElementParams = new(OutStaticBatchParamArray) FLandscapeBatchElementParams;
 		BatchElementParams->LandscapeUniformShaderParametersResource = &LandscapeUniformShaderParameters;
-		BatchElementParams->LocalToWorldNoScalingPtr = &LocalToWorldNoScaling;
 		BatchElementParams->SceneProxy = this;
-		BatchElementParams->SubX = -1;
-		BatchElementParams->SubY = -1;
 		BatchElementParams->CurrentLOD = LODIndex;
-		BatchElementParams->ForcedLOD = bForcedLOD ? LODIndex : -1;
 
 		BatchElement.UserData = BatchElementParams;
 		BatchElement.PrimitiveUniformBuffer = GetUniformBuffer();
@@ -3636,13 +3641,9 @@ void FLandscapeComponentSceneProxy::GetDynamicRayTracingInstances(FRayTracingMat
 			FMeshBatchElement BatchElement;
 			FLandscapeBatchElementParams& BatchElementParams = ParameterArray.ElementParams[SubSectionIdx];
 
-			BatchElementParams.LocalToWorldNoScalingPtr = &LocalToWorldNoScaling;
 			BatchElementParams.LandscapeUniformShaderParametersResource = &LandscapeUniformShaderParameters;
 			BatchElementParams.SceneProxy = this;
-			BatchElementParams.SubX = SubX;
-			BatchElementParams.SubY = SubY;
 			BatchElementParams.CurrentLOD = CurrentLOD;
-			BatchElementParams.ForcedLOD = ForcedLODLevel;
 			BatchElement.UserData = &BatchElementParams;
 			BatchElement.PrimitiveUniformBuffer = GetUniformBuffer();
 
@@ -4280,12 +4281,6 @@ public:
 	*/
 	void Bind(const FShaderParameterMap& ParameterMap)
 	{
-		LodBiasParameter.Bind(ParameterMap, TEXT("LodBias"));
-		LodValuesParameter.Bind(ParameterMap, TEXT("LodValues"));
-		ForcedLodParameter.Bind(ParameterMap, TEXT("ForcedLod"));
-		LodTessellationParameter.Bind(ParameterMap, TEXT("LodTessellationParams"));
-		SectionLodsParameter.Bind(ParameterMap, TEXT("SectionLods"));
-		NeighborSectionLodParameter.Bind(ParameterMap, TEXT("NeighborSectionLod"));
 	}
 
 	void GetElementShaderBindings(
@@ -4324,26 +4319,7 @@ public:
 			ShaderBindings.Add(Shader->GetUniformBufferParameter<FLandscapeVertexFactoryMVFParameters>(), BatchElementParams->LandscapeVertexFactoryMVFUniformBuffer);
 		}
 #endif
-
-		if (LodValuesParameter.IsBound())
-		{
-			ShaderBindings.Add(LodValuesParameter, SceneProxy->GetShaderLODValues(BatchElementParams->CurrentLOD));
-		}
-
-		if (ForcedLodParameter.IsBound())
-		{
-			ShaderBindings.Add(ForcedLodParameter, BatchElementParams->ForcedLOD);
-		}
 	}
-
-protected:
-	LAYOUT_FIELD(FShaderParameter, LodTessellationParameter);
-	LAYOUT_FIELD(FShaderParameter, LodValuesParameter);
-	LAYOUT_FIELD(FShaderParameter, ForcedLodParameter);
-	LAYOUT_FIELD(FShaderParameter, NeighborSectionLodParameter);
-	LAYOUT_FIELD(FShaderParameter, LodBiasParameter);
-	LAYOUT_FIELD(FShaderParameter, SectionLodsParameter);
-	LAYOUT_FIELD(TShaderUniformBufferParameter<FLandscapeUniformShaderParameters>, LandscapeShaderParameters);
 };
 
 /** 
@@ -4374,25 +4350,8 @@ public:
 		check(SceneProxy);
 
 		ShaderBindings.Add(Shader->GetUniformBufferParameter<FLandscapeUniformShaderParameters>(), *BatchElementParams->LandscapeUniformShaderParametersResource);
-
-		if (LodValuesParameter.IsBound())
-		{
-			ShaderBindings.Add(LodValuesParameter, SceneProxy->GetShaderLODValues(BatchElementParams->CurrentLOD));
+		ShaderBindings.Add(Shader->GetUniformBufferParameter<FLandscapeFixedGridUniformShaderParameters>(), SceneProxy->LandscapeFixedGridUniformShaderParameters[BatchElementParams->CurrentLOD]);
 		}
-
-		if (LodBiasParameter.IsBound())
-		{
-			ShaderBindings.Add(LodBiasParameter, FVector4(ForceInitToZero));
-		}
-
-		if (ForcedLodParameter.IsBound())
-		{
-			ShaderBindings.Add(ForcedLodParameter, BatchElementParams->ForcedLOD);
-		}
-	}
-
-	
-	
 };
 
 //
diff --git a/Engine/Source/Runtime/Landscape/Private/LandscapeRenderMobile.cpp b/Engine/Source/Runtime/Landscape/Private/LandscapeRenderMobile.cpp
--- a/Engine/Source/Runtime/Landscape/Private/LandscapeRenderMobile.cpp
+++ b/Engine/Source/Runtime/Landscape/Private/LandscapeRenderMobile.cpp
@@ -59,12 +59,7 @@ public:
 	*/
 	void Bind(const FShaderParameterMap& ParameterMap)
 	{
-		LodValuesParameter.Bind(ParameterMap,TEXT("LodValues"));
-		ForcedLodParameter.Bind(ParameterMap,TEXT("ForcedLod"));
-		LodTessellationParameter.Bind(ParameterMap, TEXT("LodTessellationParams"));
-		NeighborSectionLodParameter.Bind(ParameterMap,TEXT("NeighborSectionLod"));
-		LodBiasParameter.Bind(ParameterMap,TEXT("LodBias"));
-		SectionLodsParameter.Bind(ParameterMap,TEXT("SectionLods"));
+		TexCoordOffsetParameter.Bind(ParameterMap,TEXT("TexCoordOffset"));
 	}
 
 	void GetElementShaderBindings(
@@ -87,28 +82,15 @@ public:
 		const FLandscapeComponentSceneProxyMobile* SceneProxy = (const FLandscapeComponentSceneProxyMobile*)BatchElementParams->SceneProxy;
 		ShaderBindings.Add(Shader->GetUniformBufferParameter<FLandscapeUniformShaderParameters>(),*BatchElementParams->LandscapeUniformShaderParametersResource);
 
-		if (LodValuesParameter.IsBound())
-		{
-			FVector4 LodValues(
-				0.0f, // this is the mesh's LOD, ES2 always uses the LOD0 mesh
-				0.0f, // unused
-				(float)SceneProxy->SubsectionSizeQuads,
-				1.f / (float)SceneProxy->SubsectionSizeQuads);
-
-			ShaderBindings.Add(LodValuesParameter, LodValues);
-		}
-
-		if (LodBiasParameter.IsBound())
+		if (TexCoordOffsetParameter.IsBound())
 		{
 			FVector CameraLocalPos3D = SceneProxy->WorldToLocal.TransformPosition(InView->ViewMatrices.GetViewOrigin());
 
-			FVector4 LodBias(
-				0.0f, // unused
-				0.0f, // unused
+			FVector2D TexCoordOffset(
 				CameraLocalPos3D.X + SceneProxy->SectionBase.X,
 				CameraLocalPos3D.Y + SceneProxy->SectionBase.Y
 			);
-			ShaderBindings.Add(LodBiasParameter, LodBias);
+			ShaderBindings.Add(TexCoordOffsetParameter, TexCoordOffset);
 		}
 
 		if (SceneProxy->bRegistered)
@@ -119,77 +101,10 @@ public:
 		{
 			ShaderBindings.Add(Shader->GetUniformBufferParameter<FLandscapeSectionLODUniformParameters>(), GNullLandscapeRenderSystemResources.UniformBuffer);
 		}
-
-		if (ForcedLodParameter.IsBound())
-		{
-			ShaderBindings.Add(ForcedLodParameter, BatchElementParams->ForcedLOD);
-		}
-
-#if 0
-		FLandscapeComponentSceneProxy::FViewCustomDataLOD* LODData = (FLandscapeComponentSceneProxy::FViewCustomDataLOD*)InView->GetCustomData(SceneProxy->GetPrimitiveSceneInfo()->GetIndex());
-		int32 SubSectionIndex = BatchElementParams->SubX + BatchElementParams->SubY * SceneProxy->NumSubsections;
-
-		if (LODData != nullptr)
-		{
-			SceneProxy->PostInitViewCustomData(*InView, LODData);
-
-			if (LodTessellationParameter.IsBound())
-			{
-				ShaderBindings.Add(LodTessellationParameter, LODData->LodTessellationParams);
-			}
-
-			if (SectionLodsParameter.IsBound())
-			{
-				if (LODData->UseCombinedMeshBatch)
-				{
-					ShaderBindings.Add(SectionLodsParameter, LODData->ShaderCurrentLOD);
-				}
-				else // in non combined, only the one representing us as we'll be called 4 times (once per sub section)
-				{
-					check(SubSectionIndex >= 0);
-					FVector4 ShaderCurrentLOD(ForceInitToZero);
-					ShaderCurrentLOD.Component(SubSectionIndex) = LODData->ShaderCurrentLOD.Component(SubSectionIndex);
-
-					ShaderBindings.Add(SectionLodsParameter, ShaderCurrentLOD);
-				}
 			}
 
-			if (NeighborSectionLodParameter.IsBound())
-			{
-				FVector4 ShaderCurrentNeighborLOD[FLandscapeComponentSceneProxy::NEIGHBOR_COUNT] = { FVector4(ForceInitToZero), FVector4(ForceInitToZero), FVector4(ForceInitToZero), FVector4(ForceInitToZero) };
-
-				if (LODData->UseCombinedMeshBatch)
-				{
-					int32 SubSectionCount = SceneProxy->NumSubsections == 1 ? 1 : FLandscapeComponentSceneProxy::MAX_SUBSECTION_COUNT;
-
-					for (int32 NeighborSubSectionIndex = 0; NeighborSubSectionIndex < SubSectionCount; ++NeighborSubSectionIndex)
-					{
-						ShaderCurrentNeighborLOD[NeighborSubSectionIndex] = LODData->SubSections[NeighborSubSectionIndex].ShaderCurrentNeighborLOD;
-						check(ShaderCurrentNeighborLOD[NeighborSubSectionIndex].X != -1.0f); // they should all match so only check the 1st one for simplicity
-					}
-
-					ShaderBindings.Add(NeighborSectionLodParameter, ShaderCurrentNeighborLOD);
-				}
-				else // in non combined, only the one representing us as we'll be called 4 times (once per sub section)
-				{
-					check(SubSectionIndex >= 0);
-					ShaderCurrentNeighborLOD[SubSectionIndex] = LODData->SubSections[SubSectionIndex].ShaderCurrentNeighborLOD;
-					check(ShaderCurrentNeighborLOD[SubSectionIndex].X != -1.0f); // they should all match so only check the 1st one for simplicity
-
-					ShaderBindings.Add(NeighborSectionLodParameter, ShaderCurrentNeighborLOD);
-				}
-			}
-		}
-#endif
-	}
 protected:
-	LAYOUT_FIELD(FShaderParameter, LodValuesParameter);
-	LAYOUT_FIELD(FShaderParameter, ForcedLodParameter);
-	LAYOUT_FIELD(FShaderParameter, LodTessellationParameter);
-	LAYOUT_FIELD(FShaderParameter, NeighborSectionLodParameter);
-	LAYOUT_FIELD(FShaderParameter, LodBiasParameter);
-	LAYOUT_FIELD(FShaderParameter, SectionLodsParameter);
-	LAYOUT_FIELD(TShaderUniformBufferParameter<FLandscapeUniformShaderParameters>, LandscapeShaderParameters);
+	LAYOUT_FIELD(FShaderParameter, TexCoordOffsetParameter);
 };
 
 /** Shader parameters for use with FLandscapeVertexFactory */
@@ -204,7 +119,6 @@ public:
 	void Bind(const FShaderParameterMap& ParameterMap)
 	{
 		FLandscapeVertexFactoryPixelShaderParameters::Bind(ParameterMap);
-		BlendableLayerMaskParameter.Bind(ParameterMap, TEXT("BlendableLayerMask"));
 	}
 
 	void GetElementShaderBindings(
@@ -222,23 +136,7 @@ public:
 		SCOPE_CYCLE_COUNTER(STAT_LandscapeVFDrawTimePS);
 		
 		FLandscapeVertexFactoryPixelShaderParameters::GetElementShaderBindings(Scene, InView, Shader, InputStreamType, FeatureLevel, VertexFactory, BatchElement, ShaderBindings, VertexStreams);
-
-		if (BlendableLayerMaskParameter.IsBound())
-		{
-			const FLandscapeBatchElementParams* BatchElementParams = (const FLandscapeBatchElementParams*)BatchElement.UserData;
-			check(BatchElementParams);
-			const FLandscapeComponentSceneProxyMobile* SceneProxy = (const FLandscapeComponentSceneProxyMobile*)BatchElementParams->SceneProxy;
-			
-			FVector MaskVector;
-			MaskVector[0] = (SceneProxy->BlendableLayerMask & (1 << 0)) ? 1 : 0;
-			MaskVector[1] = (SceneProxy->BlendableLayerMask & (1 << 1)) ? 1 : 0;
-			MaskVector[2] = (SceneProxy->BlendableLayerMask & (1 << 2)) ? 1 : 0;
-			ShaderBindings.Add(BlendableLayerMaskParameter, MaskVector);
-		}
 	}
-
-protected:
-	LAYOUT_FIELD(FShaderParameter, BlendableLayerMaskParameter);
 };
 
 /**
@@ -266,20 +164,11 @@ public:
 		check(BatchElementParams);
 		const FLandscapeComponentSceneProxyMobile* SceneProxy = (const FLandscapeComponentSceneProxyMobile*)BatchElementParams->SceneProxy;
 		ShaderBindings.Add(Shader->GetUniformBufferParameter<FLandscapeUniformShaderParameters>(), *BatchElementParams->LandscapeUniformShaderParametersResource);
+		ShaderBindings.Add(Shader->GetUniformBufferParameter<FLandscapeFixedGridUniformShaderParameters>(), SceneProxy->LandscapeFixedGridUniformShaderParameters[BatchElementParams->CurrentLOD]);
 
-		if (LodValuesParameter.IsBound())
-		{
-			ShaderBindings.Add(LodValuesParameter, SceneProxy->GetShaderLODValues(BatchElementParams->CurrentLOD));
-		}
-
-		if (LodBiasParameter.IsBound())
+		if (TexCoordOffsetParameter.IsBound())
 		{
-			ShaderBindings.Add(LodBiasParameter, FVector4(ForceInitToZero));
-		}
-
-		if (ForcedLodParameter.IsBound())
-		{
-			ShaderBindings.Add(ForcedLodParameter, BatchElementParams->ForcedLOD);
+			ShaderBindings.Add(TexCoordOffsetParameter, FVector4(ForceInitToZero));
 		}
 	}
 };
@@ -461,8 +350,6 @@ FLandscapeComponentSceneProxyMobile::FLandscapeComponentSceneProxyMobile(ULandsc
 	WeightmapTextures = InComponent->MobileWeightmapTextures;
 	NormalmapTexture = InComponent->MobileWeightmapTextures[0];
 
-	BlendableLayerMask = InComponent->MobileBlendableLayerMask;
-
 #if WITH_EDITOR
 	TArray<FWeightmapLayerAllocationInfo>& LayerAllocations = InComponent->MobileWeightmapLayerAllocations.Num() ? InComponent->MobileWeightmapLayerAllocations : InComponent->GetWeightmapLayerAllocations();
 	LayerColors.Empty();
diff --git a/Engine/Source/Runtime/Landscape/Public/LandscapeRender.h b/Engine/Source/Runtime/Landscape/Public/LandscapeRender.h
--- a/Engine/Source/Runtime/Landscape/Public/LandscapeRender.h
+++ b/Engine/Source/Runtime/Landscape/Public/LandscapeRender.h
@@ -88,7 +88,6 @@ LANDSCAPE_API extern UMaterialInterface* GLandscapeDirtyMaterial;
 
 /** The uniform shader parameters for a landscape draw call. */
 BEGIN_GLOBAL_SHADER_PARAMETER_STRUCT(FLandscapeUniformShaderParameters, LANDSCAPE_API)
-/** vertex shader parameters */
 SHADER_PARAMETER(int32, ComponentBaseX)
 SHADER_PARAMETER(int32, ComponentBaseY)
 SHADER_PARAMETER(int32, SubsectionSizeVerts)
@@ -100,6 +99,7 @@ SHADER_PARAMETER(FVector4, LandscapeLightmapScaleBias)
 SHADER_PARAMETER(FVector4, SubsectionSizeVertsLayerUVPan)
 SHADER_PARAMETER(FVector4, SubsectionOffsetParams)
 SHADER_PARAMETER(FVector4, LightmapSubsectionOffsetParams)
+	SHADER_PARAMETER(FVector4, BlendableLayerMask)
 SHADER_PARAMETER(FMatrix, LocalToWorldNoScaling)
 SHADER_PARAMETER_TEXTURE(Texture2D, HeightmapTexture)
 SHADER_PARAMETER_SAMPLER(SamplerState, HeightmapTextureSampler)
@@ -109,7 +109,6 @@ SHADER_PARAMETER_TEXTURE(Texture2D, XYOffsetmapTexture)
 SHADER_PARAMETER_SAMPLER(SamplerState, XYOffsetmapTextureSampler)
 END_GLOBAL_SHADER_PARAMETER_STRUCT()
 
-
 BEGIN_GLOBAL_SHADER_PARAMETER_STRUCT(FLandscapeVertexFactoryMVFParameters, LANDSCAPE_API)
 	SHADER_PARAMETER(FIntPoint, SubXY)
 END_GLOBAL_SHADER_PARAMETER_STRUCT()
@@ -125,6 +124,10 @@ BEGIN_GLOBAL_SHADER_PARAMETER_STRUCT(FLandscapeSectionLODUniformParameters, )
 	SHADER_PARAMETER_SRV(Buffer<float>, SectionTessellationFalloffK)
 END_GLOBAL_SHADER_PARAMETER_STRUCT()
 
+BEGIN_GLOBAL_SHADER_PARAMETER_STRUCT(FLandscapeFixedGridUniformShaderParameters, )
+	SHADER_PARAMETER(FVector4, LodValues)
+END_GLOBAL_SHADER_PARAMETER_STRUCT()
+
 /* Data needed for the landscape vertex factory to set the render state for an individual batch element */
 struct FLandscapeBatchElementParams
 {
@@ -132,14 +135,8 @@ struct FLandscapeBatchElementParams
 	FLandscapeVertexFactoryMVFUniformBufferRef LandscapeVertexFactoryMVFUniformBuffer;
 #endif
 	const TUniformBuffer<FLandscapeUniformShaderParameters>* LandscapeUniformShaderParametersResource;
-	const FMatrix* LocalToWorldNoScalingPtr;
-
-	// LOD calculation-related params
 	const FLandscapeComponentSceneProxy* SceneProxy;
-	int32 SubX;
-	int32 SubY;
 	int32 CurrentLOD;
-	int32 ForcedLOD = -1;
 };
 
 class FLandscapeElementParamArray : public FOneFrameResource
@@ -913,6 +910,8 @@ protected:
 
 	UTexture2D* XYOffsetmapTexture;
 
+	uint8 BlendableLayerMask;
+
 	uint32						SharedBuffersKey;
 	FLandscapeSharedBuffers*	SharedBuffers;
 	FLandscapeVertexFactory*	VertexFactory;
@@ -966,6 +965,8 @@ protected:
 
 	TUniformBuffer<FLandscapeUniformShaderParameters> LandscapeUniformShaderParameters;
 
+	TArray< TUniformBuffer<FLandscapeFixedGridUniformShaderParameters> > LandscapeFixedGridUniformShaderParameters;
+
 	// Cached versions of these
 	FMatrix					WorldToLocal;
 
diff --git a/Engine/Source/Runtime/Landscape/Public/LandscapeRenderMobile.h b/Engine/Source/Runtime/Landscape/Public/LandscapeRenderMobile.h
--- a/Engine/Source/Runtime/Landscape/Public/LandscapeRenderMobile.h
+++ b/Engine/Source/Runtime/Landscape/Public/LandscapeRenderMobile.h
@@ -123,8 +123,6 @@ public:
 	virtual void CreateRenderThreadResources() override;
 	virtual int32 CollectOccluderElements(FOccluderElementsCollector& Collector) const override;
 
-	uint8 BlendableLayerMask;
-
 	friend class FLandscapeVertexBufferMobile;
 
 	virtual void ApplyMeshElementModifier(FMeshBatchElement& InOutMeshElement, int32 InLodIndex) const override;
