Fix for compiler errors when building sun temple

diff --git a/Engine/Shaders/Private/MobileBasePassPixelShader.usf b/Engine/Shaders/Private/MobileBasePassPixelShader.usf
--- a/Engine/Shaders/Private/MobileBasePassPixelShader.usf
+++ b/Engine/Shaders/Private/MobileBasePassPixelShader.usf
@@ -107,25 +107,6 @@ float4 SpotLightDirectionAndSpecularScale[MAX_DYNAMIC_POINT_LIGHTS];
 	#endif
 #endif
 
-half PhongApprox( half Roughness, half RoL )
-{
-	half a = Roughness * Roughness;			// 1 mul
-	//!! Ronin Hack?
-	a = max(a, 0.008);						// avoid underflow in FP16, next sqr should be bigger than 6.1e-5
-	half a2 = a * a;						// 1 mul
-	half rcp_a2 = rcp(a2);					// 1 rcp
-	//half rcp_a2 = exp2( -6.88886882 * Roughness + 6.88886882 );
-
-	// Spherical Gaussian approximation: pow( x, n ) ~= exp( (n + 0.775) * (x - 1) )
-	// Phong: n = 0.5 / a2 - 0.5
-	// 0.5 / ln(2), 0.275 / ln(2)
-	half c = 0.72134752 * rcp_a2 + 0.39674113;	// 1 mad
-	half p = rcp_a2 * exp2(c * RoL - c);		// 2 mad, 1 exp2, 1 mul
-	// Total 7 instr
-	return min(p, rcp_a2);						// Avoid overflow/underflow on Mali GPUs
-}
-
-
 #if !FULLY_ROUGH
 #if HQ_REFLECTIONS
 
diff --git a/Engine/Shaders/Private/MobileShadingModels.ush b/Engine/Shaders/Private/MobileShadingModels.ush
--- a/Engine/Shaders/Private/MobileShadingModels.ush
+++ b/Engine/Shaders/Private/MobileShadingModels.ush
@@ -58,6 +58,24 @@ half3 GetEnvBRDF(half3 SpecularColor, half Roughness, float NoV)
 }
 #endif
 
+half PhongApprox( half Roughness, half RoL )
+{
+	half a = Roughness * Roughness;			// 1 mul
+	//!! Ronin Hack?
+	a = max(a, 0.008);						// avoid underflow in FP16, next sqr should be bigger than 6.1e-5
+	half a2 = a * a;						// 1 mul
+	half rcp_a2 = rcp(a2);					// 1 rcp
+	//half rcp_a2 = exp2( -6.88886882 * Roughness + 6.88886882 );
+
+	// Spherical Gaussian approximation: pow( x, n ) ~= exp( (n + 0.775) * (x - 1) )
+	// Phong: n = 0.5 / a2 - 0.5
+	// 0.5 / ln(2), 0.275 / ln(2)
+	half c = 0.72134752 * rcp_a2 + 0.39674113;	// 1 mad
+	half p = rcp_a2 * exp2(c * RoL - c);		// 2 mad, 1 exp2, 1 mul
+	// Total 7 instr
+	return min(p, rcp_a2);						// Avoid overflow/underflow on Mali GPUs
+}
+
 void InitShadingModelContext( inout FMobileShadingModelContext ShadingModelContext, FMaterialPixelParameters MaterialParameters)
 {
 #if !FULLY_ROUGH
