Grass map rendering now uses the fixed grid vertex factory.
This keeps it independent of view lod.
Cherry pick integration from Dev-Rendering

diff --git a/Engine/Shaders/Private/LandscapeVertexFactory.ush b/Engine/Shaders/Private/LandscapeVertexFactory.ush
--- a/Engine/Shaders/Private/LandscapeVertexFactory.ush
+++ b/Engine/Shaders/Private/LandscapeVertexFactory.ush
@@ -23,7 +23,12 @@ float2 TexCoordOffset;
 /* Return index into buffers for this section instance. */
 int GetComponentLinearIndex()
 {
+#if FIXED_GRID
+	// Fixed grid doesn't use LandscapeContinuousLODParameters buffers
+	return 0;
+#else
 	return (LandscapeParameters.ComponentBaseY - LandscapeContinuousLODParameters.Min.y) * LandscapeContinuousLODParameters.Size.x + (LandscapeParameters.ComponentBaseX - LandscapeContinuousLODParameters.Min.x);
+#endif
 }
 
 /**
@@ -32,7 +37,7 @@ int GetComponentLinearIndex()
  * z = SectionSizeQuads in current mesh LOD
  * w = 1/SectionSizeQuads in current mesh LOD
  **/
-float4 GetLodValues()
+float4 GetLodValues(uint ComponentIndex)
 {
 #if FIXED_GRID 
 	// Fixed grid has fixed Lod
@@ -42,7 +47,7 @@ float4 GetLodValues()
 	return float4(0, 0, LandscapeParameters.SubsectionSizeVertsLayerUVPan.xy);
 #else
 	// Get Lod for the current section and infer the values
-	uint Lod = (uint)floor(LandscapeContinuousLODParameters.SectionLOD[GetComponentLinearIndex()]);
+	uint Lod = (uint)floor(LandscapeContinuousLODParameters.SectionLOD[ComponentIndex]);
 	float LodSubSectionSizeQuads = (float)((LandscapeParameters.SubsectionSizeVerts >> Lod) - 1);
 	return float4(Lod, 0, LodSubSectionSizeQuads, rcp(LodSubSectionSizeQuads));
 #endif
@@ -52,20 +57,38 @@ float4 GetLodValues()
  * x = Heightmap texture LOD difference from current LOD to highest LOD
  * y = todo: XYOffset texture LOD difference from current LOD to highest LOD
  */
-float2 GetLodBias()
+float2 GetLodBias(uint ComponentIndex)
 {
 #if FIXED_GRID || FEATURE_LEVEL < FEATURE_LEVEL_SM4
+	// Fixed grid and mobile don't use LodBias
 	return 0;
 #else
-	float LodBias = LandscapeContinuousLODParameters.SectionLODBias[GetComponentLinearIndex()];
+	float LodBias = LandscapeContinuousLODParameters.SectionLODBias[ComponentIndex];
 	return float2(LodBias, 0.f);
 #endif
 }
 
+float2 GetSectionTessellationFalloff(uint ComponentIndex)
+{
+#if FIXED_GRID
+	// Fixed grid doesn't support tessellation
+	return 0;
+#else
+	return float2(
+		LandscapeContinuousLODParameters.SectionTessellationFalloffC[ComponentIndex],
+		LandscapeContinuousLODParameters.SectionTessellationFalloffK[ComponentIndex]);
+#endif
+}
+
 /* Get current Lod for this section instance. */
-float GetSectionLod()
+float GetSectionLod(uint ComponentIndex)
 {
-	return min(LandscapeContinuousLODParameters.SectionLOD[GetComponentLinearIndex()] + View.FarShadowStaticMeshLODBias, LandscapeParameters.LastLOD);
+#if FIXED_GRID
+	// Fixed grid has fixed Lod
+	return LandscapeFixedGrid.LodValues.x;
+#else
+	return min(LandscapeContinuousLODParameters.SectionLOD[ComponentIndex] + View.FarShadowStaticMeshLODBias, LandscapeParameters.LastLOD);
+#endif
 }
 
 /** 
@@ -78,8 +101,12 @@ float GetSectionLod()
  *    - - 3 - -
  *       +Y
  */
-float4 GetNeighborSectionLod()
+float4 GetNeighborSectionLod(uint ComponentIndex)
 {
+#if FIXED_GRID
+	// Fixed grid has fixed Lod
+	return LandscapeFixedGrid.LodValues.xxxx;
+#else
 	int2 NeighborOffset[4];
     NeighborOffset[0] = int2(0, -1);
     NeighborOffset[1] = int2(-1, 0);
@@ -91,10 +118,38 @@ float4 GetNeighborSectionLod()
 		int2 NeighborPos = NeighborOffset[i] + int2(LandscapeParameters.ComponentBaseX, LandscapeParameters.ComponentBaseY);
 		NeighborPos = max(LandscapeContinuousLODParameters.Min                                                     , NeighborPos);
 		NeighborPos = min(LandscapeContinuousLODParameters.Min + LandscapeContinuousLODParameters.Size - int2(1, 1), NeighborPos);
-		NeighborLod[i] = max(LandscapeContinuousLODParameters.SectionLOD[(NeighborPos.y - LandscapeContinuousLODParameters.Min.y) * LandscapeContinuousLODParameters.Size.x + (NeighborPos.x - LandscapeContinuousLODParameters.Min.x)], GetSectionLod());
+		NeighborLod[i] = max(LandscapeContinuousLODParameters.SectionLOD[(NeighborPos.y - LandscapeContinuousLODParameters.Min.y) * LandscapeContinuousLODParameters.Size.x + (NeighborPos.x - LandscapeContinuousLODParameters.Min.x)], GetSectionLod(ComponentIndex));
 		NeighborLod[i] = min(NeighborLod[i] + View.FarShadowStaticMeshLODBias, LandscapeParameters.LastLOD);
 	}
 	return float4(NeighborLod[0], NeighborLod[1], NeighborLod[2], NeighborLod[3]);
+#endif
+}
+
+/** Calculate fractional render LOD for a subsection. */
+float CalcLOD(uint ComponentIndex, float2 xyLocalToSubsection, float2 Subsection)
+{
+#if FIXED_GRID
+	// Fixed grid has fixed Lod
+	return LandscapeFixedGrid.LodValues.x;
+#else
+	float2 xy = (xyLocalToSubsection + Subsection) / LandscapeParameters.NumSubsections;
+	// LOD value based on Barycentric coordinate
+	float4 L0 = float4(xy.y, xy.x, (1 - xy.x), (1 - xy.y)) * 2;
+
+	float4 LODCalculated4 = L0 * GetSectionLod(ComponentIndex) + (float4(1, 1, 1, 1) - L0) * GetNeighborSectionLod(ComponentIndex);
+
+	// Find which quadrant of the subsection we're in - top, left, right or bottom.
+	// 0 --------- 1 X  1 ----- 2  1          1 
+	//   | \ 0 / |        \   /    | \      / | 
+	//   | 1 + 2 |          0      |  0    0  |      0       
+	//   | / 3 \ |                 | /      \ |    /   \     
+	// 2 --------- 3               2          2  1 ----- 2   
+	float LODCalculated = ((xy.x + xy.y) > 1) ?
+		(xy.x < xy.y ? LODCalculated4.w : LODCalculated4.z) :
+		(xy.x < xy.y ? LODCalculated4.y : LODCalculated4.x);
+
+	return LODCalculated;
+#endif
 }
 
 struct FVertexFactoryInput
@@ -113,7 +168,7 @@ FVertexFactoryInput LoadVertexFactoryInputForHGS(uint TriangleIndex, int VertexI
 {
 	FVertexFactoryInput Input;
 
-	float4 LodValues = GetLodValues();
+	float4 LodValues = GetLodValues(GetComponentLinearIndex());
 	uint SubsectionSizeQuadsCurrentLOD = LodValues.z;
 	uint SubsectionSizeVertsCurrentLOD = LodValues.z + 1;
 
@@ -147,7 +202,7 @@ FVertexFactoryInput LoadVertexFactoryInputForDynamicUpdate(uint TriangleIndex, i
 {
 	FVertexFactoryInput Input;
 
-	float4 LodValues = GetLodValues();
+	float4 LodValues = GetLodValues(GetComponentLinearIndex());
 	uint SubsectionSizeQuadsCurrentLOD = LodValues.z;
 	uint SubsectionSizeVertsCurrentLOD = LodValues.z + 1;
 
@@ -258,8 +313,10 @@ struct FVertexFactoryIntermediates
 	float4 InputPosition;
 	float3 LocalPosition;
 	float3 WorldNormal;
+	uint ComponentIndex;
 	float4 LodValues;
 	float2 LodBias;
+	float2 TessellationFalloff;
 };
 
 float3 GetLocalPosition(FVertexFactoryIntermediates Intermediates)
@@ -291,22 +348,19 @@ FLandscapeTexCoords GetLandscapeTexCoords(FVertexFactoryInput Input, FVertexFact
 	FLandscapeTexCoords Result;
 
 	Result.LayerTexCoord.xy = Intermediates.LocalPosition.xy + LandscapeParameters.SubsectionSizeVertsLayerUVPan.zw + Intermediates.InputPosition.zw * LandscapeParameters.SubsectionOffsetParams.ww;
+
 #if USING_TESSELLATION
-#if FIXED_GRID
-	Result.LayerTexCoord.zw = 0;
-#else
-	Result.LayerTexCoord.z = LandscapeContinuousLODParameters.SectionTessellationFalloffC[GetComponentLinearIndex()];
-	Result.LayerTexCoord.w = LandscapeContinuousLODParameters.SectionTessellationFalloffK[GetComponentLinearIndex()];
-#endif
+	Result.LayerTexCoord.zw = Intermediates.TessellationFalloff;
 #endif
 
-	Result.WeightMapTexCoord	= Intermediates.LocalPosition.xy * LandscapeParameters.WeightmapUVScaleBias.xy + LandscapeParameters.WeightmapUVScaleBias.zw + Intermediates.InputPosition.zw * LandscapeParameters.SubsectionOffsetParams.zz;
+	Result.WeightMapTexCoord = Intermediates.LocalPosition.xy * LandscapeParameters.WeightmapUVScaleBias.xy + LandscapeParameters.WeightmapUVScaleBias.zw + Intermediates.InputPosition.zw * LandscapeParameters.SubsectionOffsetParams.zz;
+
 #if !(ES3_1_PROFILE)
-	Result.HeightMapTexCoord	= Intermediates.LocalPosition.xy * LandscapeParameters.HeightmapUVScaleBias.xy + LandscapeParameters.HeightmapUVScaleBias.zw + 0.5*LandscapeParameters.HeightmapUVScaleBias.xy + Intermediates.InputPosition.zw * LandscapeParameters.SubsectionOffsetParams.xy;
+	Result.HeightMapTexCoord = Intermediates.LocalPosition.xy * LandscapeParameters.HeightmapUVScaleBias.xy + LandscapeParameters.HeightmapUVScaleBias.zw + 0.5*LandscapeParameters.HeightmapUVScaleBias.xy + Intermediates.InputPosition.zw * LandscapeParameters.SubsectionOffsetParams.xy;
 #endif
 
 #if NEEDS_LIGHTMAP_COORDINATE
-	Result.LightMapCoordinate	= (Intermediates.LocalPosition.xy * LandscapeParameters.LandscapeLightmapScaleBias.xy + LandscapeParameters.LandscapeLightmapScaleBias.wz + Intermediates.InputPosition.zw * LandscapeParameters.LightmapSubsectionOffsetParams.xy);
+	Result.LightMapCoordinate = (Intermediates.LocalPosition.xy * LandscapeParameters.LandscapeLightmapScaleBias.xy + LandscapeParameters.LandscapeLightmapScaleBias.wz + Intermediates.InputPosition.zw * LandscapeParameters.LightmapSubsectionOffsetParams.xy);
 #endif
 	
 	return Result;
@@ -573,32 +627,14 @@ void GetShadowMapCoordinate(FVertexFactoryInterpolantsVSToPS Interpolants, out f
 }
 #endif
 
-float CalcLOD(float2 xyLocalToSubsection, float2 Subsection)
-{
-	float2 xy = (xyLocalToSubsection + Subsection) / LandscapeParameters.NumSubsections;
-	// LOD value based on Barycentric coordinate
-	float4 L0 = float4(xy.y, xy.x, (1 - xy.x), (1 - xy.y)) * 2;
-
-	float4 LODCalculated4 = L0 * GetSectionLod() + (float4(1, 1, 1, 1) - L0) * GetNeighborSectionLod();
-
-	// Find which quadrant of the subsection we're in - top, left, right or bottom.
-	// 0 --------- 1 X  1 ----- 2  1          1 
-	//   | \ 0 / |        \   /    | \      / | 
-	//   | 1 + 2 |          0      |  0    0  |      0       
-	//   | / 3 \ |                 | /      \ |    /   \     
-	// 2 --------- 3               2          2  1 ----- 2   
-	float LODCalculated = ((xy.x + xy.y) > 1) ? 
-							(xy.x < xy.y ? LODCalculated4.w : LODCalculated4.z) :
-							(xy.x < xy.y ? LODCalculated4.y : LODCalculated4.x);
-
-	return LODCalculated;
-}
-
 FVertexFactoryIntermediates GetVertexFactoryIntermediates(FVertexFactoryInput Input)
 {
+
 	FVertexFactoryIntermediates Intermediates;
-	Intermediates.LodValues = GetLodValues();
-	Intermediates.LodBias = GetLodBias();
+	Intermediates.ComponentIndex = GetComponentLinearIndex();
+	Intermediates.LodValues = GetLodValues(Intermediates.ComponentIndex);
+	Intermediates.LodBias = GetLodBias(Intermediates.ComponentIndex);
+	Intermediates.TessellationFalloff = GetSectionTessellationFalloff(Intermediates.ComponentIndex);
 
 #if FEATURE_LEVEL >= FEATURE_LEVEL_SM4
 	Intermediates.InputPosition = Input.Position;
@@ -613,15 +649,9 @@ FVertexFactoryIntermediates GetVertexFactoryIntermediates(FVertexFactoryInput In
 
 	float2 xy = Intermediates.InputPosition.xy * Intermediates.LodValues.w;
 
-#if FIXED_GRID
-	float LODCalculated = floor(Intermediates.LodValues.x);
-	float LodValue = LODCalculated;
-	float MorphAlpha = 0;
-#else
-	float LODCalculated = CalcLOD(xy, Intermediates.InputPosition.zw);
+	float LODCalculated = CalcLOD(Intermediates.ComponentIndex, xy, Intermediates.InputPosition.zw);
 	float LodValue = floor(LODCalculated);
 	float MorphAlpha = LODCalculated - LodValue;
-#endif
 
 	// InputPositionLODAdjusted : Position for actual LOD in base LOD units
 #if FEATURE_LEVEL >= FEATURE_LEVEL_SM4
@@ -707,15 +737,8 @@ FVertexFactoryIntermediates GetVertexFactoryIntermediates(FVertexFactoryInput In
 #if RENDERING_LANDSCAPE_GRASS && FEATURE_LEVEL >= FEATURE_LEVEL_SM4 && VERTEXSHADER
 float2 LandscapeVertexFactorySampleHeight(FVertexFactoryInput Input, FVertexFactoryIntermediates Intermediates)
 {
-	/*
-	FLandscapeTexCoords TexCoords = GetLandscapeTexCoords(Input, Intermediates);
-
-	float4 SampleValue = Texture2DSampleLevel(HeightmapTexture, HeightmapTextureSampler, TexCoords.HeightMapTexCoord, 0);
-	return SampleValue.xy;
-	*/
-
 	float2 xy = Intermediates.InputPosition.xy * Intermediates.LodValues.w;
-	float LODCalculated = CalcLOD(xy, Intermediates.InputPosition.zw);
+	float LODCalculated = CalcLOD(Intermediates.ComponentIndex, xy, Intermediates.InputPosition.zw);
 	float LodValue = floor(LODCalculated);
 	float CoordTranslate_x = (LandscapeParameters.SubsectionSizeVertsLayerUVPan.x * pow(2, -LodValue) - 1) * LandscapeParameters.SubsectionSizeVertsLayerUVPan.y;
 
@@ -747,7 +770,7 @@ float4 VertexFactoryGetRasterizedWorldPosition(FVertexFactoryInput Input, FVerte
 {
 #if RENDERING_LANDSCAPE_GRASS
 	float2 xy = Intermediates.InputPosition.xy * Intermediates.LodValues.w;
-	float LODCalculated = CalcLOD(xy, Intermediates.InputPosition.zw);
+	float LODCalculated = CalcLOD(Intermediates.ComponentIndex, xy, Intermediates.InputPosition.zw);
 	float LodValue = floor(LODCalculated);
 	float CoordTranslate_x = (LandscapeParameters.SubsectionSizeVertsLayerUVPan.x * pow(2, -LodValue) - 1) * LandscapeParameters.SubsectionSizeVertsLayerUVPan.y;
 
@@ -852,7 +875,7 @@ FVertexFactoryInterpolantsVSToDS VertexFactoryGetInterpolantsVSToDS(FVertexFacto
 	Interpolants.TangentToWorld2 = float4(TangentToWorld[2], 1);
 
 	float2 xy = Intermediates.InputPosition.xy * Intermediates.LodValues.w;
-	Interpolants.CalculatedLOD = CalcLOD(xy, Intermediates.InputPosition.zw);
+	Interpolants.CalculatedLOD = CalcLOD(Intermediates.ComponentIndex, xy, Intermediates.InputPosition.zw);
 #endif
 	
 	return Interpolants;
diff --git a/Engine/Source/Runtime/Landscape/Private/LandscapeGrass.cpp b/Engine/Source/Runtime/Landscape/Private/LandscapeGrass.cpp
--- a/Engine/Source/Runtime/Landscape/Private/LandscapeGrass.cpp
+++ b/Engine/Source/Runtime/Landscape/Private/LandscapeGrass.cpp
@@ -223,8 +223,8 @@ static bool ShouldCacheLandscapeGrassShaders(const FMeshMaterialShaderPermutatio
 	// We only need grass weight shaders for Landscape vertex factories on desktop platforms
 	return (Parameters.MaterialParameters.bIsUsedWithLandscape || Parameters.MaterialParameters.bIsSpecialEngineMaterial) &&
 		IsFeatureLevelSupported(Parameters.Platform, ERHIFeatureLevel::SM5) &&
-		((Parameters.VertexFactoryType == FindVertexFactoryType(FName(TEXT("FLandscapeVertexFactory"), FNAME_Find))) || (Parameters.VertexFactoryType == FindVertexFactoryType(FName(TEXT("FLandscapeXYOffsetVertexFactory"), FNAME_Find))))
-		&& !IsConsolePlatform(Parameters.Platform);
+		Parameters.VertexFactoryType == FindVertexFactoryType(FName(TEXT("FLandscapeFixedGridVertexFactory"), FNAME_Find)) && 
+		!IsConsolePlatform(Parameters.Platform);
 }
 
 class FLandscapeGrassWeightShaderElementData : public FMeshMaterialShaderElementData
diff --git a/Engine/Source/Runtime/Landscape/Private/LandscapeRender.cpp b/Engine/Source/Runtime/Landscape/Private/LandscapeRender.cpp
--- a/Engine/Source/Runtime/Landscape/Private/LandscapeRender.cpp
+++ b/Engine/Source/Runtime/Landscape/Private/LandscapeRender.cpp
@@ -1515,8 +1515,9 @@ void FLandscapeComponentSceneProxy::CreateRenderThreadResources()
 		GrassBatchParams.AddDefaulted(NumMips);
 		
 		// Grass is being generated using LOD0 material only
+		// It uses the fixed grid vertex factory so it doesn't support XY offsets
 		FMaterialRenderProxy* RenderProxy = AvailableMaterials[LODIndexToMaterialIndex[0]]->GetRenderProxy();
-		GrassMeshBatch.VertexFactory = VertexFactory;
+		GrassMeshBatch.VertexFactory = FixedGridVertexFactory;
 		GrassMeshBatch.MaterialRenderProxy = RenderProxy;
 		GrassMeshBatch.LCI = nullptr;
 		GrassMeshBatch.ReverseCulling = false;
@@ -1538,6 +1539,7 @@ void FLandscapeComponentSceneProxy::CreateRenderThreadResources()
 		GrassBatchElement->MinVertexIndex = 0;
 		GrassBatchElement->MaxVertexIndex = SharedBuffers->NumVertices - 1;
 
+		// Grass system is also used to bake out heights which are source for collision data when bBakeMaterialPositionOffsetIntoCollision is enabled
 		for (int32 Mip = 1; Mip < NumMips; ++Mip)
 		{
 			const int32 MipSubsectionSizeVerts = SubsectionSizeVerts >> Mip;
@@ -1546,6 +1548,7 @@ void FLandscapeComponentSceneProxy::CreateRenderThreadResources()
 			*CollisionBatchElement = *GrassBatchElement;
 			FLandscapeBatchElementParams* CollisionBatchElementParams = &GrassBatchParams[Mip];
 			*CollisionBatchElementParams = *BatchElementParams;
+			CollisionBatchElementParams->CurrentLOD = Mip;
 			CollisionBatchElement->UserData = CollisionBatchElementParams;
 			CollisionBatchElement->NumPrimitives = FMath::Square(NumSubsections) * FMath::Square(MipSubsectionSizeVerts);
 			CollisionBatchElement->FirstIndex = SharedBuffers->GrassIndexMipOffsets[Mip];
@@ -4685,8 +4688,9 @@ public:
 			{
 				// Landscape MICs are only for use with the Landscape vertex factories
 
-				// For now only compile FLandscapeFixedGridVertexFactory for runtime virtual texture page rendering (can change if we need for other cases)
+				// For now only compile FLandscapeFixedGridVertexFactory for grass and runtime virtual texture page rendering (can change if we need for other cases)
 				// Todo: only compile LandscapeXYOffsetVertexFactory if we are using it
+				bool bIsGrassShaderType = Algo::Find(GetGrassShaderTypes(), ShaderType->GetFName()) != nullptr;
 				bool bIsRuntimeVirtualTextureShaderType = Algo::Find(GetRuntimeVirtualTextureShaderTypes(), ShaderType->GetFName()) != nullptr;
 
 				static const FName LandscapeVertexFactory = FName(TEXT("FLandscapeVertexFactory"));
@@ -4704,7 +4708,7 @@ public:
 				if (VertexFactoryType->GetFName() == LandscapeFixedGridVertexFactory ||
 					VertexFactoryType->GetFName() == LandscapeFixedGridVertexFactoryMobile)
 				{
-					return bIsRuntimeVirtualTextureShaderType && FMaterialResource::ShouldCache(Platform, ShaderType, VertexFactoryType);
+					return (bIsGrassShaderType || bIsRuntimeVirtualTextureShaderType) && FMaterialResource::ShouldCache(Platform, ShaderType, VertexFactoryType);
 				}
 			}
 		}
@@ -4975,6 +4979,16 @@ public:
 		return ExcludedShaderTypes;
 	}
 
+	static const TArray<FName>& GetGrassShaderTypes()
+	{
+		static const TArray<FName> ShaderTypes =
+		{
+			FName(TEXT("FLandscapeGrassWeightVS")),
+			FName(TEXT("FLandscapeGrassWeightPS")),
+		};
+		return ShaderTypes;
+	}
+
 	static const TArray<FName>& GetRuntimeVirtualTextureShaderTypes()
 	{
 		static const TArray<FName> ShaderTypes =
