Dithered LOD Fix

diff --git a/Engine/Shaders/Private/MaterialTemplate.ush b/Engine/Shaders/Private/MaterialTemplate.ush
--- a/Engine/Shaders/Private/MaterialTemplate.ush
+++ b/Engine/Shaders/Private/MaterialTemplate.ush
@@ -2221,7 +2221,12 @@ void ClipLODTransition(float2 SvPosition, float DitherFactor)
 {
 	if (abs(DitherFactor) > .001)
 	{
-		float RandCos = cos(dot(floor(SvPosition.xy), float2(347.83451793,3343.28371963)));
+		float ArgCos = dot(floor(SvPosition.xy), float2(347.83451793, 3343.28371963));
+#if FEATURE_LEVEL <= FEATURE_LEVEL_ES3_1
+		// Temporary workaround for precision issues on mobile when the argument is bigger than 10k
+		ArgCos = fmod(ArgCos, 10000);
+#endif
+		float RandCos = cos(ArgCos);
 		float RandomVal = frac(RandCos * 1000.0);
 		half RetVal = (DitherFactor < 0.0) ?
 			(DitherFactor + 1.0 > RandomVal) :
