Fixed: Fixed some merging issues related to water shading

diff --git a/Engine/Shaders/Private/MobileBasePassPixelShader.usf b/Engine/Shaders/Private/MobileBasePassPixelShader.usf
--- a/Engine/Shaders/Private/MobileBasePassPixelShader.usf
+++ b/Engine/Shaders/Private/MobileBasePassPixelShader.usf
@@ -739,11 +739,11 @@ void Main(
 			MaterialParameters, PixelMaterialInputs, ResolvedView,
 			Shadow, ShadingModelContext.Specular, NullDistortionParams,
 			SunIlluminance, WaterDiffuseIndirectIlluminance, EnvBrdf,
-			CameraIsUnderWater, WaterVisibility, EyeIndex);
+			CameraIsUnderWater, ShadingModelContext.WaterVisibility, EyeIndex);
 
 		// Accumulate luminance and occlude the background according to transmittance to view and mean transmittance to lights.
 		Color += WaterLighting.Luminance;
-		Opacity = 1.0 - ((1.0 - Opacity) * dot(WaterLighting.WaterToSceneToLightTransmittance, float3(1.0 / 3.0, 1.0 / 3.0, 1.0 / 3.0)));
+		ShadingModelContext.Opacity = 1.0 - ((1.0 - ShadingModelContext.Opacity) * dot(WaterLighting.WaterToSceneToLightTransmittance, float3(1.0 / 3.0, 1.0 / 3.0, 1.0 / 3.0)));
 	}
 #endif // MATERIAL_SHADINGMODEL_SINGLELAYERWATER
 
