Integrate hair changes (rendering+simulation)

diff --git a/Engine/Shaders/Private/HairStrands/HairStrandsInterpolation.usf b/Engine/Shaders/Private/HairStrands/HairStrandsInterpolation.usf
--- a/Engine/Shaders/Private/HairStrands/HairStrandsInterpolation.usf
+++ b/Engine/Shaders/Private/HairStrands/HairStrandsInterpolation.usf
@@ -6,7 +6,7 @@
 #include "HairStrandsAABBCommon.ush" 
 
 #if PERMUTATION_DEBUG == 1
-#include "../ShaderDrawDebug.ush"
+#include "../ShaderDrawDebug.ush" 
 #endif
 
 uint VertexCount;
@@ -31,19 +31,20 @@ Buffer<uint4>  SimRestPosePositionBuffer;
 Buffer<uint4>  DeformedSimPositionBuffer;
 Buffer<uint4>  Interpolation0Buffer;
 Buffer<uint4>  Interpolation1Buffer;
+Buffer<uint>   SimRootPointIndexBuffer;
 
 #if PERMUTATION_SCALE_AND_CLIP
 Buffer<uint4>  AttributeBuffer;
 #endif
 
-#if PERMUTATION_DYNAMIC_GEOMETRY == 1 || PERMUTATION_DEBUG == 1
+#if PERMUTATION_DYNAMIC_GEOMETRY == 1 || PERMUTATION_DEBUG == 1 || PERMUTATION_DYNAMIC_GEOMETRY == 2
 Buffer<uint>   RenVertexToRootIndexBuffer;
 Buffer<uint>   SimVertexToRootIndexBuffer;
 #endif
 
 RWBuffer<uint4> OutRenderDeformedPositionBuffer;
 
-#if PERMUTATION_DYNAMIC_GEOMETRY == 1
+#if PERMUTATION_DYNAMIC_GEOMETRY == 1 || PERMUTATION_DYNAMIC_GEOMETRY == 2
 
 // Compact all these buffers into 2 buffers: translation + quaternion
 Buffer<float4>  RestPosition0Buffer;
@@ -89,6 +90,17 @@ float3 ComputeStaticGeometryOffset(uint GuideIndex, float GuideVertexWeight)
 	return (DeformedGuidePoint - RestGuidePoint) * GuideVertexWeight;
 }
 
+float3 ComputeLocalGeometryOffset(uint GuideIndex, float GuideVertexWeight)
+{	
+	const float3 RestGuidePoint = f16tof32(SimRestPosePositionBuffer[GuideIndex].xyz) - 
+					f16tof32(SimRestPosePositionBuffer[SimRootPointIndexBuffer[GuideIndex]].xyz) ;
+
+	const float3 DeformedGuidePoint = f16tof32(DeformedSimPositionBuffer[GuideIndex].xyz) - 
+					f16tof32(DeformedSimPositionBuffer[SimRootPointIndexBuffer[GuideIndex]].xyz) ;
+
+	return (DeformedGuidePoint - RestGuidePoint) * GuideVertexWeight;
+}
+
 float3 ComputeDynamicGeometryOffset(
 	uint GuideIndex, 
 	float GuideVertexWeight, 
@@ -140,7 +152,7 @@ void MainCS(uint2 DispatchThreadId : SV_DispatchThreadID)
 		uint3 GuideVertexIndices = 0;
 		float3 GuideVertexLerps = 0;
 		float3 GuideVertexWeights = 0;
-		#if PERMUTATION_SIMULATION == 1 || PERMUTATION_DEBUG == 1 || PERMUTATION_DEBUG == 2
+		#if PERMUTATION_SIMULATION == 1 || PERMUTATION_DEBUG == 1 || PERMUTATION_DEBUG == 2 
 		{
 			const uint4 InterpolationData0	= Interpolation0Buffer[VertexIndex];
 			const uint4 InterpolationData1	= Interpolation1Buffer[VertexIndex];
@@ -169,7 +181,6 @@ void MainCS(uint2 DispatchThreadId : SV_DispatchThreadID)
 			uint4 PackedRestPosition = RenderRestPosePositionBuffer[VertexIndex];
 			float3 ControlPoint = f16tof32(PackedRestPosition.xyz) + InRenderHairPositionOffset;
 
-
 			#if PERMUTATION_SIMULATION == 1
 			[unroll]
 			for (uint K = 0; K < 3; ++K)
@@ -267,6 +278,45 @@ void MainCS(uint2 DispatchThreadId : SV_DispatchThreadID)
 		}
 		#endif	
 
+		// 3. Skin transformation
+		// Apply dynamic mesh deformation (translation)
+		#if PERMUTATION_DYNAMIC_GEOMETRY == 2
+		{
+			uint4 PackedRestPosition = RenderRestPosePositionBuffer[VertexIndex];
+			const float3 RestControlPoint = f16tof32(PackedRestPosition.xyz) + InRenderHairPositionOffset;
+			float3 ControlPoint = RestControlPoint;
+
+			float3 CurrOffset = 0;
+
+			#if PERMUTATION_SIMULATION == 1
+			[unroll]
+			for (uint K = 0; K < 3; ++K)
+			{
+				const uint GuideIndex0 = GuideVertexIndices[K];
+				const uint GuideIndex1 = GuideIndex0 + 1;
+
+				const float3 Offset0 = ComputeLocalGeometryOffset(GuideIndex0, GuideVertexWeights[K]);
+				const float3 Offset1 = ComputeLocalGeometryOffset(GuideIndex1, GuideVertexWeights[K]);
+
+				CurrOffset += lerp(Offset0, Offset1, GuideVertexLerps[K]);
+			}
+			#endif
+
+			const uint RootIndex						= RenVertexToRootIndexBuffer[VertexIndex];
+			const float3 RootBarycentric				= DecodeBarycentrics(RootBarycentricBuffer[RootIndex]);
+			const FHairMeshTriangle RestTriangle		= GetTriangleTransformation(RootIndex, RestPosition0Buffer, RestPosition1Buffer, RestPosition2Buffer);
+			const FHairMeshTriangle DeformedTriangle	= GetTriangleTransformation(RootIndex, DeformedPosition0Buffer, DeformedPosition1Buffer, DeformedPosition2Buffer);
+			
+			ControlPoint -= RootBarycentric.x * RestTriangle.P0     + RootBarycentric.y * RestTriangle.P1     + RootBarycentric.z * RestTriangle.P2;
+			ControlPoint += RootBarycentric.x * DeformedTriangle.P0 + RootBarycentric.y * DeformedTriangle.P1 + RootBarycentric.z * DeformedTriangle.P2;
+
+			const bool bIsValid = ScaleAndClip(VertexIndex, PackedRestPosition.w);
+			WorldVertexPosition = bIsValid ? ControlPoint + CurrOffset : float3(1.#INF, 1.#INF, 1.#INF);
+
+			OutRenderDeformedPositionBuffer[VertexIndex] = float4(f32tof16(WorldVertexPosition - OutHairPositionOffset), PackedRestPosition.w);
+		}
+		#endif	
+
 		#if PERMUTATION_DEBUG == 1
 		if (HairDebugMode == 1)
 		{
diff --git a/Engine/Shaders/Private/HairStrands/HairStrandsLUT.usf b/Engine/Shaders/Private/HairStrands/HairStrandsLUT.usf
--- a/Engine/Shaders/Private/HairStrands/HairStrandsLUT.usf
+++ b/Engine/Shaders/Private/HairStrands/HairStrandsLUT.usf
@@ -90,6 +90,7 @@ void MainCS(uint3 DispatchThreadId : SV_DispatchThreadID)
 	GBufferData.BaseColor	= Absorption.xxx;		// Perceptual absorption
 	GBufferData.Metallic	= 0;		 			// This disable the fake multiple scattering
 	GBufferData.Roughness 	= Roughness; 			// Perceptual roughness
+	GBufferData.CustomData  = float4(0, 0, 1, 0); 	// Backlit
 
 	float FrontHemisphereOutput = 0;
 	float BackHemisphereOutput  = 0;
@@ -199,6 +200,7 @@ void MainCS(uint3 DispatchThreadId : SV_DispatchThreadID)
 	GBufferData.BaseColor = Absorption.xxx;		// Perceptual absorption
 	GBufferData.Metallic = 0;		 			// This disable the fake multiple scattering
 	GBufferData.Roughness = Roughness; 			// Perceptual roughness
+	GBufferData.CustomData = float4(0,0,1,0); 	// Backlit
 
 	float R_Output = 0;
 	float TT_Output = 0;
diff --git a/Engine/Shaders/Private/HairStrands/HairStrandsLightChannelMask.usf b/Engine/Shaders/Private/HairStrands/HairStrandsLightChannelMask.usf
--- a/Engine/Shaders/Private/HairStrands/HairStrandsLightChannelMask.usf
+++ b/Engine/Shaders/Private/HairStrands/HairStrandsLightChannelMask.usf
@@ -10,7 +10,7 @@
 
 uint2						OutputResolution;
 Texture2D<uint> 			NodeOffsetAndCount;
-StructuredBuffer<FNodeData> NodeData;
+StructuredBuffer<FPackedHairSample> NodeData;
 RWTexture2D<uint>			OutLightChannelMaskTexture;
 
 #if PERMUTATION_GROUP_SIZE == 0
@@ -36,15 +36,9 @@ void MainCS(uint3 DispatchThreadId : SV_DispatchThreadID)
 	for (uint SampleIt = 0; SampleIt < NodeDesc.Count; SampleIt++)
 	{
 		const uint LocalOffset = NodeDesc.Offset + SampleIt;
-		const FNodeData Data = NodeData[LocalOffset];
-		float			SampleDepth;
-		FVisibilityData SampleVisData;
-		float3			SampleBaseColor;
-		float			SampleRoughness;
-		float			SampleSpecular;
-		uint			SampleLightChannelMask;
-		DecodeNodeData(Data, SampleDepth, SampleVisData, SampleBaseColor, SampleRoughness, SampleSpecular, SampleLightChannelMask);
-		OutLightChannelMask |= SampleLightChannelMask;
+		const FPackedHairSample Data = NodeData[LocalOffset];
+		const FHairSample Sample = UnpackHairSample(Data);
+		OutLightChannelMask |= Sample.LightChannelMask;
 	}
 
 	OutLightChannelMaskTexture[PixelCoord] = OutLightChannelMask;
diff --git a/Engine/Shaders/Private/HairStrands/HairStrandsMaterialPS.usf b/Engine/Shaders/Private/HairStrands/HairStrandsMaterialPS.usf
--- a/Engine/Shaders/Private/HairStrands/HairStrandsMaterialPS.usf
+++ b/Engine/Shaders/Private/HairStrands/HairStrandsMaterialPS.usf
@@ -1,6 +1,6 @@
 // Copyright Epic Games, Inc. All Rights Reserved.
 
-#include "HairStrandsVisibilityCommon.ush"
+#include "HairStrandsVisibilityCommon.ush" 
 #include "../Common.ush"
 #include "../SceneTexturesCommon.ush"
 #include "HairStrandsDeepShadowCommon.ush"
@@ -15,59 +15,24 @@ uint MaterialPass_MacroGroupId;
 uint MaterialPass_MaterialId;  
 uint MaterialPass_PrimitiveId;
 uint MaterialPass_LightChannelMask;
-RWStructuredBuffer<FNodeData>	OutNodeData;
-RWBuffer<float2>				OutNodeVelocity; 
 
-// For reference:
-#if 0
-struct FVertexFactoryInterpolantsVSToPS
-{
-	TANGENTTOWORLD_INTERPOLATOR_BLOCK
-
-
-#if USE_INSTANCING
-	// x = per-instance random, y = per-instance fade out amount, z = hide/show flag, w dither fade cutoff
-	float4  PerInstanceParams : COLOR1;
-#endif
-
-#if INSTANCED_STEREO
-	nointerpolation uint EyeIndex : PACKED_EYE_INDEX;
-#endif
-
-#if VF_USE_PRIMITIVE_SCENE_DATA
-	nointerpolation uint PrimitiveId : PRIMITIVE_ID;
-#endif
-#if VF_STRAND_HAIR
-	nointerpolation uint HairPrimitiveId : HAIR_PRIMITIVE_ID; // Control point ID
-	float2 HairPrimitiveUV : HAIR_PRIMITIVE_UV; // U: parameteric distance between the two surrounding control points. V: parametric distance along the width.
-#endif
-};
-
-struct FVertexFactoryInput
-{
-#if VF_USE_PRIMITIVE_SCENE_DATA
-	uint PrimitiveId : ATTRIBUTE13;
-#endif	
-	uint VertexId : SV_VertexID;
-};
-
-#endif
 
 #if VF_STRAND_HAIR
 struct FEvaluateOutput
 {
-	FNodeData NodeData;
+	FPackedHairSample NodeData;
 	float2 NodeVelocity;
-//	float3 P0;
-//	float3 P1;
-//	float3 P;
-//	float3 A;
-//	float3 B;
-//	float U;
+};
 
+struct FHairVertexData
+{
+	float3 World_P;
+	float3 World_P_Prev;
+	float  Radius;
+	bool   bIsValid;
 };
 
-FVertexFactoryInterpolantsVSToPS GetInterpolants(uint PrimitiveId, uint HairPrimitiveId, inout float3 World_P, inout float3 World_P_Prev)
+FVertexFactoryInterpolantsVSToPS GetInterpolants(uint PrimitiveId, uint HairPrimitiveId, inout FHairVertexData OutHairVertex)
 {
 	FMaterialVertexParameters Empty = (FMaterialVertexParameters)0;
 
@@ -78,8 +43,12 @@ FVertexFactoryInterpolantsVSToPS GetInterpolants(uint PrimitiveId, uint HairPrim
 	FVertexFactoryIntermediates	Intermediates = GetVertexFactoryIntermediates(Input);
 	FVertexFactoryInterpolantsVSToPS Interpolants = VertexFactoryGetInterpolantsVSToPS(Input, Intermediates, Empty);
 
-	World_P = VertexFactoryGetWorldPositionRaw(Input, Intermediates).xyz;
-	World_P_Prev = VertexFactoryGetPreviousWorldPosition(Input, Intermediates).xyz;
+	OutHairVertex.World_P = VertexFactoryGetWorldPositionRaw(Input, Intermediates).xyz;
+	OutHairVertex.World_P_Prev = VertexFactoryGetPreviousWorldPosition(Input, Intermediates).xyz;
+
+	const FHairControlPoint ControlPoint = GetVertexPosition(Input);
+	OutHairVertex.Radius = ControlPoint.WorldRadius;
+	OutHairVertex.bIsValid = ControlPoint.Type != HAIR_CONTROLPOINT_END;
 
 	return Interpolants;
 }
@@ -95,47 +64,45 @@ FEvaluateOutput Evaluate(
 	const float SceneDepth = ConvertFromDeviceZ(SampleDepth);
 	const float2 UV = SamplePixelCoord / float2(View.BufferSizeAndInvSize.xy);
 	const float2 ScreenPosition = (UV - View.ScreenPositionScaleBias.wz) / View.ScreenPositionScaleBias.xy;
-//	const float3 World_P = mul(float4(ScreenPosition * SceneDepth, SceneDepth, 1), View.ScreenToWorld).xyz;
 	const float3 World_P = mul(float4(ScreenPosition * SceneDepth, SceneDepth, 1), View.ScreenToTranslatedWorld).xyz;
 
-//	Out.P = World_P;
-
 	float3 Velocity_World_P;
 	float3 Velocity_World_P_Prev;
 	float U = 0;
 	FVertexFactoryInterpolantsVSToPS Interpolants;
 	{
 		const uint HairPrimitiveId0 = SamplePrimitiveId;
-		const uint HairPrimitiveId1 = HairPrimitiveId0 + 1;
+		const uint HairPrimitiveId1 = HairPrimitiveId0 + 1; //TODO: min(HairPrimitiveId0 + 1, MaxVertexCount-1);
 
-		float3 World_P0;
-		float3 World_P1;
+		FHairVertexData Vertex0;
+		FHairVertexData Vertex1;
 
-		float3 World_P0_Prev;
-		float3 World_P1_Prev;
-
-		
-		FVertexFactoryInterpolantsVSToPS Interpolants0 = GetInterpolants(PrimitiveId, HairPrimitiveId0, World_P0, World_P0_Prev);
-		FVertexFactoryInterpolantsVSToPS Interpolants1 = GetInterpolants(PrimitiveId, HairPrimitiveId1, World_P1, World_P1_Prev);
-//		Out.P0 = World_P0;
-//		Out.P1 = World_P1;
+		FVertexFactoryInterpolantsVSToPS Interpolants0 = GetInterpolants(PrimitiveId, HairPrimitiveId0, Vertex0);
+		FVertexFactoryInterpolantsVSToPS Interpolants1 = GetInterpolants(PrimitiveId, HairPrimitiveId1, Vertex1);
 
 		Interpolants = Interpolants0;
 
-		// Compute the projection of the current sample point onto the hair segment (which is in the middle of the hair strands)
-		const float3 A = (World_P - World_P0);
-		const float3 B = (World_P1 - World_P0);
-		const float DotA = dot(A, B);
+		// Compute U
+		// Use the projection of the current sample point onto the hair segment (which is in the middle of the hair strands)
+		const float3 A = (World_P  - Vertex0.World_P);
+		const float3 B = (Vertex1.World_P - Vertex0.World_P);
+		const float CosA = dot(A, B);
 		const float LengthB2 = dot(B, B);
-		const float SegmentU = LengthB2 > 0 ? saturate(DotA / LengthB2) : 0;
+		const float SegmentU = Vertex1.bIsValid && LengthB2 > 0 ? saturate(CosA / LengthB2) : 0;
 
-//		Out.U = SegmentU;
-//		Out.A = A;
-//		Out.B = B;
-		const float SegmentV = 0.5f; // Todo handle case when the actual geometry is close and cover multiple pixels
+		// Compute V
+		float SegmentV = 0.5f;
+		{
+			const float3 W = cross(A, B);
+			const bool bIsRight = dot(W, View.ViewForward) < 0;
+			const float WorldRadius	= lerp(Vertex0.Radius, Vertex1.Radius, SegmentU);
+			const float WorldV	= length(A - B * SegmentU);
+			SegmentV = WorldV / WorldRadius;
+			SegmentV = bIsRight ? SegmentV * 0.5f + 0.5f : (1 - SegmentV) * 0.5f;
+		}
 
-		Velocity_World_P	  = lerp(World_P0, World_P1, SegmentU);
-		Velocity_World_P_Prev = lerp(World_P0_Prev, World_P1_Prev, SegmentU);
+		Velocity_World_P	  = lerp(Vertex0.World_P, Vertex1.World_P, SegmentU);
+		Velocity_World_P_Prev = lerp(Vertex0.World_P_Prev, Vertex1.World_P_Prev, SegmentU);
 
 		Interpolants.TangentToWorld0 = lerp(Interpolants0.TangentToWorld0, Interpolants1.TangentToWorld0, SegmentU);
 		Interpolants.TangentToWorld2 = lerp(Interpolants0.TangentToWorld2, Interpolants1.TangentToWorld2, SegmentU);
@@ -165,64 +132,54 @@ FEvaluateOutput Evaluate(
 	const float4 SvPosition = float4(SamplePixelCoord, SampleDepth, 1);
 
 	// PS
-	FVisibilityData OutVisibilityData = (FVisibilityData)0;	
-	OutVisibilityData.Coverage8bit = SampleCoverage8bit;
-	OutVisibilityData.PrimitiveId = SampleCoverage8bit;
-	OutVisibilityData.MacroGroupId = MaterialPass_MacroGroupId;
-
-	half3 OutBaseColor = 0;
-	float OutRoughness = 0;
-	float OutSpecular = 0;
+	FHairSample OutSample		= (FHairSample)0;
+	OutSample.Depth				= SampleDepth;
+	OutSample.Coverage8bit		= SampleCoverage8bit;
+	OutSample.PrimitiveId		= SamplePrimitiveId;
+	OutSample.MacroGroupId		= MaterialPass_MacroGroupId;
+	OutSample.LightChannelMask	= MaterialPass_LightChannelMask;
+	OutSample.bScatterSceneLighting = UseScatterSceneLighting();
+
+	// Coverage computation
+	// We don't use the coverage information
+	#if 0
+	float Coverage = 1;
 	{
-		// Coverage computation
-		// We don't use the coverage information
-		#if 0
-		float Coverage = 1;
-		{
-			float WorldStrandRadius = VFIntermediates.HairDimensions.y; // Coming from the vertex factory evaluation
+		const bool bUseStableRasterization = UseStableRasterization();
 
-			FHairRenderInfo HairRenderInfo = GetHairRenderInfo(ResolvedView.HairRenderInfo);
-			const float SceneDepth = ConvertFromDeviceZ(SvPosition.z); // Linear depth in world unit
-			const float PixelRadius = HairRenderInfo.bIsOrthoView ? HairRenderInfo.RadiusAtDepth1Primary : SceneDepth * HairRenderInfo.RadiusAtDepth1Primary; // Not correct but the coverage is not used (we count instead the number of sub-sample covered)
-			const float StrandRealRadius = WorldStrandRadius;
-			Coverage = StrandRealRadius / max(StrandRealRadius, PixelRadius);
-		}
-		float Transmittance = saturate(1.0f - Coverage);
-		#endif
-
-		//OutVisibilityData.Tangent = normalize(Interpolants.TangentToWorld2.xyz);
-		//OutVisibilityData.PrimitiveId = Interpolants.HairPrimitiveId;
-		//OutVisibilityData.MacroGroupId = HairVisibilityPass_HairMacroGroupIndex;
-		//OutVisibilityData.Coverage8bit = To8bitCoverage(Coverage);
+		float WorldStrandRadius = VFIntermediates.HairDimensions.y; // Coming from the vertex factory evaluation
 
-		// Material computation
-		#if 1
-		{
-			const bool bIsFrontFace = true;
-			FMaterialPixelParameters MaterialParameters = GetMaterialPixelParameters(Interpolants, SvPosition);
-			FPixelMaterialInputs PixelMaterialInputs;
+		FHairRenderInfo HairRenderInfo = GetHairRenderInfo(ResolvedView.HairRenderInfo, ResolvedView.HairRenderInfoBits, bUseStableRasterization);
+		const float SceneDepth = ConvertFromDeviceZ(SvPosition.z); // Linear depth in world unit
+		const float PixelRadius = HairRenderInfo.bIsOrthoView ? HairRenderInfo.RadiusAtDepth1Primary : SceneDepth * HairRenderInfo.RadiusAtDepth1Primary; // Not correct but the coverage is not used (we count instead the number of sub-sample covered)
+		const float StrandRealRadius = WorldStrandRadius;
+		Coverage = StrandRealRadius / max(StrandRealRadius, PixelRadius);
+	}
+	float Transmittance = saturate(1.0f - Coverage);
+	#endif
 
-			float4 ScreenPosition = SvPositionToResolvedScreenPosition(SvPosition);
-			float3 TranslatedWorldPosition = SvPositionToResolvedTranslatedWorld(SvPosition);
-			CalcMaterialParametersEx(MaterialParameters, PixelMaterialInputs, SvPosition, ScreenPosition, bIsFrontFace, TranslatedWorldPosition, TranslatedWorldPosition);
+	// Material computation
+	{
+		const bool bIsFrontFace = true;
+		FMaterialPixelParameters MaterialParameters = GetMaterialPixelParameters(Interpolants, SvPosition);
+		FPixelMaterialInputs PixelMaterialInputs;
 
-			OutBaseColor = GetMaterialBaseColor(PixelMaterialInputs);
-			OutRoughness = GetMaterialRoughness(PixelMaterialInputs);
-			OutSpecular = GetMaterialSpecular(PixelMaterialInputs);
-			//float MaterialAO = GetMaterialAmbientOcclusion(PixelMaterialInputs);
+		float4 ScreenPosition = SvPositionToResolvedScreenPosition(SvPosition);
+		float3 TranslatedWorldPosition = SvPositionToResolvedTranslatedWorld(SvPosition);
+		CalcMaterialParametersEx(MaterialParameters, PixelMaterialInputs, SvPosition, ScreenPosition, bIsFrontFace, TranslatedWorldPosition, TranslatedWorldPosition);
 
-			//OutVisibilityData.Tangent = mul(normalize(GetMaterialNormalRaw(PixelMaterialInputs)), MaterialParameters.TangentToWorld);
-			OutVisibilityData.Tangent = mul(normalize(GetMaterialNormalRaw(PixelMaterialInputs)), MaterialParameters.TangentToWorld);
-//			OutVisibilityData.Tangent = normalize(Interpolants.TangentToWorld2.xyz);
+		OutSample.BaseColor = GetMaterialBaseColor(PixelMaterialInputs);
+		OutSample.Roughness = GetMaterialRoughness(PixelMaterialInputs);
+		OutSample.Specular = GetMaterialSpecular(PixelMaterialInputs);
+		OutSample.Backlit = saturate(GetMaterialCustomData0(MaterialParameters));
 
-			OutVisibilityData.Tangent = normalize(OutVisibilityData.Tangent); 
-		}
-		#endif
+		OutSample.Tangent = mul(normalize(GetMaterialNormalRaw(PixelMaterialInputs)), MaterialParameters.TangentToWorld);
+		OutSample.Tangent = normalize(OutSample.Tangent);
 	}
 
 	FEvaluateOutput Out;
-	Out.NodeData		= ConvertToNodeData(SampleDepth, OutVisibilityData, OutBaseColor, OutRoughness, OutSpecular, MaterialPass_LightChannelMask);
-	Out.NodeVelocity	= EncodedVelocity;
+	Out.NodeData	 = PackHairSample(OutSample);
+	Out.NodeVelocity = EncodedVelocity;
 	return Out;
 }
 #endif
@@ -230,44 +187,20 @@ FEvaluateOutput Evaluate(
 void Main(
 	in float4 SvPosition: SV_Position
 	, in uint2 Resolution : DISPATCH_RESOLUTION
-	, out float4 OutColor0 : SV_Target0
-//	, out float4 OutColor1 : SV_Target1
-//	, out float4 OutColor2 : SV_Target2
-//	, out float4 OutColor3 : SV_Target3
-//	, out float4 OutColor4 : SV_Target4
-)
+	, out float4 OutColor0 : SV_Target0)
 {
 	ResolvedView = ResolveView();
 
 	const uint3 PixelCoord = uint3(floor(SvPosition.xy), 0);
-//	const FNodeDesc NodeDesc = DecodeNodeDesc(MaterialPassParameters.NodeIndex.Load(PixelCoord));
-
-//	OutColor0 = 0;
 	const uint SampleIndex = PixelCoord.x + PixelCoord.y * Resolution.x;// MaterialPassParameters.MaxResolution.x;
 	bool bIsValid = false;
 	if (SampleIndex < MaterialPassParameters.MaxSampleCount)
 	{
-//		FNodeData InNode = MaterialPassParameters.NodeVis[SampleIndex];
 		const FNodeVis InNode = MaterialPassParameters.NodeVis.Load(SampleIndex);
 		
 		const uint PackedSamplePixelCoord = MaterialPassParameters.NodeCoord[SampleIndex];
 		const float2 SamplePixelCoord = float2(PackedSamplePixelCoord & 0xFFFF, (PackedSamplePixelCoord>>16) & 0xFFFF) + float2(0.5f, 0.5f);
 
-#if 0
-		float			SampleDepth;
-		FVisibilityData	VisibilityData;
-		float3			BaseColor;
-		float			Roughness;
-		float			Specular;
-
-		DecodeNodeData(
-			InNode,
-			SampleDepth,
-			VisibilityData,
-			BaseColor,
-			Roughness,
-			Specular);
-#else
 		float SampleDepth = InNode.Depth;
 		uint SampleCoverage8bit = InNode.Coverage8bit;
 		uint SamplePrimitiveId = InNode.PrimitiveId;
@@ -275,106 +208,15 @@ void Main(
 
 		if (SampleMaterialId == MaterialPass_MaterialId)
 		{
-#endif
 			FEvaluateOutput Out = Evaluate(SamplePixelCoord, SampleDepth, SamplePrimitiveId, SampleCoverage8bit, MaterialPass_PrimitiveId);
-			OutNodeData[SampleIndex] = Out.NodeData;
-			OutNodeVelocity[SampleIndex] = Out.NodeVelocity;
+			MaterialPassParameters.OutNodeData[SampleIndex] = Out.NodeData;
+			MaterialPassParameters.OutNodeVelocity[SampleIndex] = Out.NodeVelocity;
 
 			bIsValid = true;
-//		OutColor0 = float4(Out.P0, 1);
-//		OutColor1 = float4(Out.P1, 1);
-//		OutColor2 = float4(Out.P, Out.U);
-//		OutColor3 = float4(Out.A, 1);
-//		OutColor4 = float4(Out.B, 1);
 		}
 	}
 
 	OutColor0 = bIsValid? float4(0, 1, 0, 1) : float4(1, 0, 0, 1);
-
-#if 0
-	float Coverage = 1;
-	{
-		FHairRenderInfo HairRenderInfo = GetHairRenderInfo(ResolvedView.HairRenderInfo);
-		const float SceneDepth = ConvertFromDeviceZ(SvPosition.z); // Linear depth in world unit
-		const float PixelRadius = HairRenderInfo.bIsOrthoView ? HairRenderInfo.RadiusAtDepth1Primary : SceneDepth * HairRenderInfo.RadiusAtDepth1Primary; // Not correct but the coverage is not used (we count instead the number of sub-sample covered)
-		const float StrandRealRadius = WorldStrandRadius;
-		Coverage = StrandRealRadius / max(StrandRealRadius, PixelRadius);
-	}
-
-#if HAIR_RENDER_MODE == RENDER_MODE_TRANSMITTANCE
-	OutColor0 = saturate(1.0f - Coverage);
-#endif
-
-#if HAIR_RENDER_MODE == RENDER_MODE_MSAA || HAIR_RENDER_MODE == RENDER_MODE_PPLL
-	FVisibilityData VisibilityData;
-  #if VF_STRAND_HAIR
-	VisibilityData.Tangent = normalize(Interpolants.TangentToWorld2.xyz);
-	VisibilityData.PrimitiveId = Interpolants.HairPrimitiveId;
-  #else
-	VisibilityData.Tangent = float3(0, 0, 1);
-	VisibilityData.PrimitiveId = 0;
-  #endif
-	VisibilityData.MacroGroupId = HairVisibilityPass_HairMacroGroupIndex;
-	VisibilityData.Coverage8bit = To8bitCoverage(Coverage);
-
-	half3 BaseColor = 0;
-	float Roughness = 0;
-	float Specular  = 0;
-	float2 Velocity = 0;
-	float2 EncodedVelocity = 0;
-#if SUPPORT_MATERIAL_PROPERTY
-	{
-		FMaterialPixelParameters MaterialParameters = GetMaterialPixelParameters(Interpolants, SvPosition);
-		FPixelMaterialInputs PixelMaterialInputs;
-	
-		float4 ScreenPosition = SvPositionToResolvedScreenPosition(SvPosition);
-		float3 TranslatedWorldPosition = SvPositionToResolvedTranslatedWorld(SvPosition);
-		CalcMaterialParametersEx(MaterialParameters, PixelMaterialInputs, SvPosition, ScreenPosition, bIsFrontFace, TranslatedWorldPosition, TranslatedWorldPosition);
-
-		BaseColor		 = GetMaterialBaseColor(PixelMaterialInputs);
-		Roughness		 = GetMaterialRoughness(PixelMaterialInputs);
-		half  Metallic	 = GetMaterialMetallic(PixelMaterialInputs);
-		Specular		 = GetMaterialSpecular(PixelMaterialInputs);
-		float MaterialAO = GetMaterialAmbientOcclusion(PixelMaterialInputs);
-
-		VisibilityData.Tangent = mul(normalize(GetMaterialNormalRaw(PixelMaterialInputs)), MaterialParameters.TangentToWorld);
-	}
-	{
-		Velocity = Calculate2DVelocity(PackedVelocityA, PackedVelocityC);
-		EncodedVelocity = EncodeVelocityToTexture(Velocity);
-	}
-#endif
-#endif
-
-#if HAIR_RENDER_MODE == RENDER_MODE_MSAA
-	OutColor0 = EncodeVisibilityData(VisibilityData);
-	OutColor1 = float4(sqrt(BaseColor), Roughness);
-	OutColor2 = float4(Specular,0,0,0);
-	OutColor3 = EncodedVelocity;
-#endif
-
-#if HAIR_RENDER_MODE == RENDER_MODE_PPLL
-	const uint2 PixelCoord = SvPosition.xy;
-
-	float2 NDC = SvPosition.xy / SvPosition.w;
-	float2 ScreenUV = NDC * ResolvedView.ScreenPositionScaleBias.xy + ResolvedView.ScreenPositionScaleBias.wz;
-
-	// Allocate space for a new node
-	uint NodeIndex;
-	InterlockedAdd(PPLLCounter[uint2(0, 0)], 1, NodeIndex);
-
-	if (NodeIndex < VisibilityPassGlobalParameters.MaxPPLLNodeCount)
-	{
-		// If we can append new node, add it and make it point to the node previously in head of list
-		uint NextNodeIndex;
-		InterlockedExchange(PPLLNodeIndex[PixelCoord], NodeIndex, NextNodeIndex);
-		const FPPLLNodeData Node = ConvertToPPLLNodeData(SvPosition.z, VisibilityData, BaseColor, Roughness, Specular, Velocity, NextNodeIndex);
-		PPLLNodes[NodeIndex] = Node;
-	}
-#endif
-
-#endif
-
 }
 
 
diff --git a/Engine/Shaders/Private/HairStrands/HairStrandsMeshInterpolate.usf b/Engine/Shaders/Private/HairStrands/HairStrandsMeshInterpolate.usf
new file mode 100644
--- /dev/null
+++ b/Engine/Shaders/Private/HairStrands/HairStrandsMeshInterpolate.usf
@@ -0,0 +1,66 @@
+// Copyright Epic Games, Inc. All Rights Reserved.
+
+#include "../Common.ush"
+#include "HairStrandsMeshProjectionCommon.ush"
+
+uint MaxRootCount;
+uint MaxSampleCount;
+
+Buffer<float4> RestSamplePositionsBuffer;
+Buffer<float4> MeshSampleWeightsBuffer;
+
+Buffer<float4> RestRootTrianglePosition0;
+Buffer<float4> RestRootTrianglePosition1;
+Buffer<float4> RestRootTrianglePosition2;
+
+RWBuffer<float4> OutDeformedRootTrianglePosition0;
+RWBuffer<float4> OutDeformedRootTrianglePosition1;
+RWBuffer<float4> OutDeformedRootTrianglePosition2;
+
+[numthreads(128, 1, 1)]
+void MainCS(uint3 DispatchThreadId : SV_DispatchThreadID)
+{ 
+	const uint RootIndex = DispatchThreadId.x;
+	if (RootIndex >= MaxRootCount) 
+		return;
+
+	{
+		const float3 RestControlPoint = RestRootTrianglePosition0[RootIndex];
+		float3 ControlPoint = RestControlPoint;
+
+		// Apply rbf interpolation from the samples set
+		for(int i = 0; i < MaxSampleCount; ++i)
+		{
+			const float3 PositionDelta = RestControlPoint - RestSamplePositionsBuffer[i];
+			const float FunctionValue = sqrt( dot(PositionDelta,PositionDelta)+1);
+			ControlPoint += FunctionValue * MeshSampleWeightsBuffer[i].xyz;
+		}
+		OutDeformedRootTrianglePosition0[RootIndex+MaxRootCount] = float4(ControlPoint,1);
+	}
+	{
+		const float3 RestControlPoint = RestRootTrianglePosition1[RootIndex];
+		float3 ControlPoint = RestControlPoint;
+
+		// Apply rbf interpolation from the samples set
+		for(int i = 0; i < MaxSampleCount; ++i)
+		{
+			const float3 PositionDelta = RestControlPoint - RestSamplePositionsBuffer[i];
+			const float FunctionValue = sqrt( dot(PositionDelta,PositionDelta)+1);
+			ControlPoint += FunctionValue * MeshSampleWeightsBuffer[i].xyz;
+		}
+		OutDeformedRootTrianglePosition1[RootIndex+MaxRootCount] = float4(ControlPoint,1);
+	}
+	{
+		const float3 RestControlPoint = RestRootTrianglePosition2[RootIndex];
+		float3 ControlPoint = RestControlPoint;
+
+		// Apply rbf interpolation from the samples set
+		for(int i = 0; i < MaxSampleCount; ++i)
+		{
+			const float3 PositionDelta = RestControlPoint - RestSamplePositionsBuffer[i];
+			const float FunctionValue = sqrt( dot(PositionDelta,PositionDelta)+1);
+			ControlPoint += FunctionValue * MeshSampleWeightsBuffer[i].xyz;
+		}
+		OutDeformedRootTrianglePosition2[RootIndex+MaxRootCount] = float4(ControlPoint,1);
+	}
+}
\ No newline at end of file
diff --git a/Engine/Shaders/Private/HairStrands/HairStrandsMeshProjection.usf b/Engine/Shaders/Private/HairStrands/HairStrandsMeshProjection.usf
--- a/Engine/Shaders/Private/HairStrands/HairStrandsMeshProjection.usf
+++ b/Engine/Shaders/Private/HairStrands/HairStrandsMeshProjection.usf
@@ -382,7 +382,10 @@ void MainMarkSectionIdCS(uint3 DispatchThreadId : SV_DispatchThreadID)
 	const uint VertexIndex = MeshIndexBuffer[Index + MeshIndexOffset];
 	if (VertexIndex < MeshMaxVertexCount)
 	{
-		OutVertexSectionId[VertexIndex] = MeshSectionId;
+		// Use atomic min as a vertex can belong to several sections. Since 
+		// currently one of the assumption is that the transfered section will 
+		// be section 0, this works fine.
+		InterlockedMin(OutVertexSectionId[VertexIndex], MeshSectionId);
 	}
 }
 
diff --git a/Engine/Shaders/Private/HairStrands/HairStrandsMeshProjectionCommon.ush b/Engine/Shaders/Private/HairStrands/HairStrandsMeshProjectionCommon.ush
--- a/Engine/Shaders/Private/HairStrands/HairStrandsMeshProjectionCommon.ush
+++ b/Engine/Shaders/Private/HairStrands/HairStrandsMeshProjectionCommon.ush
@@ -247,13 +247,13 @@ struct FHairTriangleIndex
 
 uint EncodeTriangleIndex(uint TriangleIndex, uint SectionIndex)
 {
-	return ((SectionIndex & 0x3) <<28) | (TriangleIndex & 0xCFFFFFFF);
+	return ((SectionIndex & 0xF)<<28) | (TriangleIndex & 0xFFFFFFF);
 }
 
 FHairTriangleIndex DecodeTriangleIndex(uint Encoded)
 {
 	FHairTriangleIndex Out;
-	Out.SectionIndex = (Encoded>>28) & 0x7;
-	Out.TriangleIndex = Encoded & 0xCFFFFFFF;
+	Out.SectionIndex = (Encoded>>28) & 0xF;
+	Out.TriangleIndex = Encoded & 0xFFFFFFF;
 	return Out;
 }
\ No newline at end of file
diff --git a/Engine/Shaders/Private/HairStrands/HairStrandsMeshUpdate.usf b/Engine/Shaders/Private/HairStrands/HairStrandsMeshUpdate.usf
--- a/Engine/Shaders/Private/HairStrands/HairStrandsMeshUpdate.usf
+++ b/Engine/Shaders/Private/HairStrands/HairStrandsMeshUpdate.usf
@@ -3,12 +3,16 @@
 #include "../Common.ush"
 #include "HairStrandsMeshProjectionCommon.ush"
 
-uint MeshSectionIndex;
-uint MeshMaxIndexCount;
-uint MeshMaxVertexCount;
-uint MeshIndexOffset;
-uint MeshUVsChannelOffset;
-uint MeshUVsChannelCount;
+#ifndef MAX_SECTION_COUNT
+#error MAX_SECTION_COUNT needs to be defined
+#endif
+
+uint MaxSectionCount;
+uint MeshMaxIndexCount		[MAX_SECTION_COUNT];
+uint MeshMaxVertexCount		[MAX_SECTION_COUNT];
+uint MeshIndexOffset		[MAX_SECTION_COUNT];
+uint MeshUVsChannelOffset	[MAX_SECTION_COUNT];
+uint MeshUVsChannelCount	[MAX_SECTION_COUNT];
 
 Buffer<uint> MeshIndexBuffer;
 Buffer<float> MeshPositionBuffer;
@@ -32,17 +36,19 @@ void MainCS(uint3 DispatchThreadId : SV_DispatchThreadID)
 		return;
 	
 	const FHairTriangleIndex TriangleIndex = DecodeTriangleIndex(RootTriangleIndex[RootIndex]);
-	if (TriangleIndex.SectionIndex != MeshSectionIndex)
+	if (TriangleIndex.SectionIndex >= MaxSectionCount)
 		return;
 
+	const uint SectionIndex = TriangleIndex.SectionIndex;
+
 #if PERMUTATION_WITHUV == 1
-	const FHairMeshTriangleWithUVs Tri = GetTriangleData(TriangleIndex.TriangleIndex, MeshPositionBuffer, MeshUVsBuffer, MeshIndexBuffer, MeshIndexOffset, MeshMaxIndexCount, MeshMaxVertexCount, MeshUVsChannelOffset, MeshUVsChannelCount);
+	const FHairMeshTriangleWithUVs Tri = GetTriangleData(TriangleIndex.TriangleIndex, MeshPositionBuffer, MeshUVsBuffer, MeshIndexBuffer, MeshIndexOffset[SectionIndex], MeshMaxIndexCount[SectionIndex], MeshMaxVertexCount[SectionIndex], MeshUVsChannelOffset[SectionIndex], MeshUVsChannelCount[SectionIndex]);
 
 	OutRootTrianglePosition0[RootIndex] = float4(Tri.P0, PackUVsToFloat(Tri.UV0));
 	OutRootTrianglePosition1[RootIndex] = float4(Tri.P1, PackUVsToFloat(Tri.UV1));
 	OutRootTrianglePosition2[RootIndex] = float4(Tri.P2, PackUVsToFloat(Tri.UV2));
 #else
-	const FHairMeshTriangle Tri = GetTriangleData(TriangleIndex.TriangleIndex, MeshPositionBuffer, MeshIndexBuffer, MeshIndexOffset, MeshMaxIndexCount, MeshMaxVertexCount);
+	const FHairMeshTriangle Tri = GetTriangleData(TriangleIndex.TriangleIndex, MeshPositionBuffer, MeshIndexBuffer, MeshIndexOffset[SectionIndex], MeshMaxIndexCount[SectionIndex], MeshMaxVertexCount[SectionIndex]);
 
 	OutRootTrianglePosition0[RootIndex] = (float4(Tri.P0, 1));
 	OutRootTrianglePosition1[RootIndex] = (float4(Tri.P1, 1));
diff --git a/Engine/Shaders/Private/HairStrands/HairStrandsSamplesInit.usf b/Engine/Shaders/Private/HairStrands/HairStrandsSamplesInit.usf
new file mode 100644
--- /dev/null
+++ b/Engine/Shaders/Private/HairStrands/HairStrandsSamplesInit.usf
@@ -0,0 +1,28 @@
+// Copyright Epic Games, Inc. All Rights Reserved.
+#include "../Common.ush"
+
+uint MaxVertexCount;
+Buffer<float> VertexPositionsBuffer;
+
+uint MaxSampleCount;
+Buffer<uint> SampleIndicesBuffer;
+RWBuffer<float4> OutSamplePositionsBuffer;
+
+[numthreads(128, 1, 1)]
+void MainCS(uint3 DispatchThreadId : SV_DispatchThreadID)
+{ 
+	const uint SampleIndex = DispatchThreadId.x;
+	if (SampleIndex >= MaxSampleCount) 
+		return;
+
+	const int VertexIndex = SampleIndicesBuffer[SampleIndex];
+	if (VertexIndex >= MaxVertexCount) 
+		return;
+
+	float3 SamplePosition = float3(0,0,0);
+	SamplePosition.x = VertexPositionsBuffer.Load(VertexIndex*3 + 0);
+	SamplePosition.y = VertexPositionsBuffer.Load(VertexIndex*3 + 1);
+	SamplePosition.z = VertexPositionsBuffer.Load(VertexIndex*3 + 2);
+
+	OutSamplePositionsBuffer[SampleIndex] = float4(SamplePosition,1);
+}
\ No newline at end of file
diff --git a/Engine/Shaders/Private/HairStrands/HairStrandsSamplesUpdate.usf b/Engine/Shaders/Private/HairStrands/HairStrandsSamplesUpdate.usf
new file mode 100644
--- /dev/null
+++ b/Engine/Shaders/Private/HairStrands/HairStrandsSamplesUpdate.usf
@@ -0,0 +1,30 @@
+// Copyright Epic Games, Inc. All Rights Reserved.
+
+#include "../Common.ush"
+
+uint MaxSampleCount;
+Buffer<uint> SampleIndicesBuffer;
+Buffer<float> InterpolationWeightsBuffer;
+
+Buffer<float4> SampleRestPositionsBuffer;
+Buffer<float4> SampleDeformedPositionsBuffer;
+
+RWBuffer<float4> OutSampleDeformationsBuffer;
+
+[numthreads(128, 1, 1)]
+void MainCS(uint3 DispatchThreadId : SV_DispatchThreadID)
+{ 
+	const uint SampleIndex = DispatchThreadId.x;
+	if (SampleIndex >= MaxSampleCount) 
+		return;
+	
+	uint WeightsOffset = SampleIndex * MaxSampleCount;
+	float3 SampleDeformation = float3(0,0,0);
+	for(uint i = 0; i < MaxSampleCount; ++i, ++WeightsOffset)
+	{
+		SampleDeformation += InterpolationWeightsBuffer[WeightsOffset] * 
+				(SampleDeformedPositionsBuffer[i]-SampleRestPositionsBuffer[i]);
+	}
+
+	OutSampleDeformationsBuffer[SampleIndex] = float4(SampleDeformation,1);
+}
\ No newline at end of file
diff --git a/Engine/Shaders/Private/HairStrands/HairStrandsVertexFactory.ush b/Engine/Shaders/Private/HairStrands/HairStrandsVertexFactory.ush
--- a/Engine/Shaders/Private/HairStrands/HairStrandsVertexFactory.ush
+++ b/Engine/Shaders/Private/HairStrands/HairStrandsVertexFactory.ush
@@ -6,6 +6,7 @@
 #include "../VertexFactoryCommon.ush"
 #include "../LocalVertexFactoryCommon.ush"
 #include "HairStrandsVisibilityCommon.ush"
+#include "HairStrandsVertexFactoryCommon.ush"
 
 #include "/Engine/Generated/UniformBuffers/PrecomputedLightingBuffer.ush"
 
@@ -16,7 +17,11 @@
 	Buffer<float4> VertexFetch_InstanceTransformBuffer;
 #endif
 
-#ifdef USE_CULLED_CLUSTER
+#ifndef USE_CULLED_CLUSTER
+#define USE_CULLED_CLUSTER 0
+#endif
+
+#if USE_CULLED_CLUSTER == 1
 	Buffer<uint> CulledVertexIdsBuffer;
 	Buffer<float> CulledVertexRadiusScaleBuffer;
 #endif
@@ -34,6 +39,8 @@ float  HairStrandsVF_Density;
 float  HairStrandsVF_Radius;
 float  HairStrandsVF_Length;
 uint   HairStrandsVF_CullingEnable;
+uint   HairStrandsVF_bUseStableRasterization;
+uint   HairStrandsVF_bScatterSceneLighing;
 
 Buffer<uint4>  HairStrandsVF_PositionBuffer;
 Buffer<uint4>  HairStrandsVF_PreviousPositionBuffer;
@@ -154,6 +161,16 @@ float GetHairStrandsRoughness(uint HairPrimitiveId, float2 HairPrimitiveUV)
 	return GetHairStrandsAttributes(HairPrimitiveId, HairPrimitiveUV).HairRoughness;
 }
 
+bool UseStableRasterization()
+{
+	return HairStrandsVF_bUseStableRasterization > 0;
+}
+
+bool UseScatterSceneLighting()
+{
+	return HairStrandsVF_bScatterSceneLighing > 0;
+}
+
 struct FVertexFactoryInput
 {
 #if VF_USE_PRIMITIVE_SCENE_DATA
@@ -187,7 +204,7 @@ FVertexInfo GetVertexInfo(FVertexFactoryInput Input)
 	VertexInfo.VertexIndex = VertexInfo.BaseIndex + VertexInfo.IsTip;
 	VertexInfo.LodRadiusScale = 1.0f;
 	VertexInfo.bForceInvalidQuad = false;
-#ifdef USE_CULLED_CLUSTER
+#if USE_CULLED_CLUSTER == 1
 	if (HairStrandsVF_CullingEnable)
 	{
 		VertexInfo.VertexIndex = CulledVertexIdsBuffer[VertexInfo.BaseIndex + VertexInfo.IsTip];
@@ -245,32 +262,27 @@ uint GetHairPrimitiveId(FVertexFactoryInput Input)
 	return GetVertexInfo(Input).HairPrimitiveId;
 }
 
-struct FControlPoint
-{
-	float3 Position;
-	float  WorldRadius;
-	float  WorldLength;
-};
-
-FControlPoint GetVertexPosition(FVertexFactoryInput Input)
+FHairControlPoint GetVertexPosition(const FVertexInfo VertexInfo)
 {
-	FVertexInfo VertexInfo = GetVertexInfo(Input);
 	const uint4 PackedData = HairStrandsVF_PositionBuffer[VertexInfo.VertexIndex];
-	
-	FControlPoint Output = (FControlPoint)0;
-
-	const float3 ControlPoint	= f16tof32(PackedData.xyz) + HairStrandsVF_PositionOffset;
-	const uint PackedAlpha		= asuint(PackedData.a);
-	const uint ControlPointType = PackedAlpha & 0x3;
-	Output.WorldRadius			= ((PackedAlpha & 0x00FC) >> 2) / 63.f * HairStrandsVF_Radius * VertexInfo.LodRadiusScale;
-	Output.WorldLength			= ((PackedAlpha & 0xFF00) >> 8) / 255.f * HairStrandsVF_Length;
+	FHairControlPoint Out = UnpackHairControlPoint(
+		PackedData, 
+		HairStrandsVF_PositionOffset, 
+		HairStrandsVF_Radius * VertexInfo.LodRadiusScale, 
+		HairStrandsVF_Length);
 	
 	// Create a degenerated quad the end of each strand to cut between each strands
 	// #hair_todo: This is not efficient for short strand like fur for instance. Need to revisit that at some point
-	const bool bIsInvalidQuad = (ControlPointType == 1 && VertexInfo.IsTip == 1) || (ControlPointType == 2 && VertexInfo.IsTip == 0) || VertexInfo.bForceInvalidQuad;
-	Output.Position = bIsInvalidQuad ? float3(1.#INF, 1.#INF, 1.#INF) : ControlPoint;
+	const bool bIsInvalidQuad = (Out.Type == HAIR_CONTROLPOINT_START && VertexInfo.IsTip == 1) || (Out.Type == HAIR_CONTROLPOINT_END && VertexInfo.IsTip == 0) || VertexInfo.bForceInvalidQuad;
+	Out.Position = bIsInvalidQuad ? float3(1.#INF, 1.#INF, 1.#INF) : Out.Position;
 
-	return Output;
+	return Out;
+}
+
+FHairControlPoint GetVertexPosition(FVertexFactoryInput Input)
+{
+	const FVertexInfo VertexInfo = GetVertexInfo(Input);
+	return GetVertexPosition(VertexInfo);
 }
 
 float GetWorldStrandRadius(FVertexFactoryInput Input)
@@ -287,17 +299,18 @@ float3 GetVertexPreviousPosition(FVertexFactoryInput Input)
 {
 	FVertexInfo VertexInfo = GetVertexInfo(Input);
 	const uint4 PackedData = HairStrandsVF_PreviousPositionBuffer[VertexInfo.VertexIndex];
-
-	const float3 ControlPoint = f16tof32(PackedData.xyz) + HairStrandsVF_PreviousPositionOffset;
-	const uint PackedAlpha = asuint(PackedData.a);
-	const uint ControlPointType = PackedAlpha & 0x3;
+	FHairControlPoint Out = UnpackHairControlPoint(
+		PackedData,
+		HairStrandsVF_PreviousPositionOffset,
+		HairStrandsVF_Radius * VertexInfo.LodRadiusScale,
+		HairStrandsVF_Length);
 
 	// Create a degenerated quad the end of each strand to cut between each strands
 	// #hair_todo: This is not efficient for short strand like fur for instance. Need to revisit that at some point
-	const bool bIsInvalidQuad = (ControlPointType == 1 && VertexInfo.IsTip == 1) || (ControlPointType == 2 && VertexInfo.IsTip == 0) || VertexInfo.bForceInvalidQuad;
-	const float3 Position = bIsInvalidQuad ? float3(1.#INF, 1.#INF, 1.#INF) : ControlPoint;
+	const bool bIsInvalidQuad = (Out.Type == HAIR_CONTROLPOINT_START && VertexInfo.IsTip == 1) || (Out.Type == HAIR_CONTROLPOINT_END && VertexInfo.IsTip == 0) || VertexInfo.bForceInvalidQuad;
+	Out.Position = bIsInvalidQuad ? float3(1.#INF, 1.#INF, 1.#INF) : Out.Position;
 
-	return Position;
+	return Out.Position;
 }
 
 // Segment UV coord of an hair segment. This is different from the UV coord of the hair strands
@@ -459,10 +472,52 @@ float4 ComputeViewAlignedWorldPosition(FVertexFactoryInput Input, float3 WorldTa
 	return float4(OutWorldPosition, 1);
 }
 
+// @return translated world position 
+// Specialized version of VertexFactoryGetWorldPosition & ComputeViewAlignedWorldPosition, for visibility VS shader.
+float4 VertexFactoryGetWorldPosition_Visibility(FVertexFactoryInput Input, FHairViewInfo HairViewInfo, inout uint OutHairPrimitiveId)
+{
+	const FVertexInfo VertexInfo = GetVertexInfo(Input);
+	OutHairPrimitiveId = VertexInfo.HairPrimitiveId;
+
+	FHairControlPoint ControlPoint = GetVertexPosition(VertexInfo);
+	const float4 WorldPosition = CalcWorldPosition(ControlPoint.Position, GetPrimitiveId(Input));
+	const float3 WorldTangent = GetTangent(Input);
+
+	// Minimal radius to snap the strand to a sample/pixel center (to avoid aliasing)
+	const float DistanceToCamera = length(HairViewInfo.TranslatedWorldCameraOrigin - WorldPosition.xyz);
+	const float MinStrandHairRadius = DistanceToCamera * HairViewInfo.RadiusAtDepth1;
+	const float3 ViewDir = -HairViewInfo.ViewForward;
+	const float3 Right = normalize(cross(WorldTangent, ViewDir));
+	const float3 OutWorldPosition = WorldPosition.xyz + (VertexInfo.IsLeft ? -Right : Right) * max(ControlPoint.WorldRadius * VertexInfo.LodRadiusScale, MinStrandHairRadius);
+
+	return float4(OutWorldPosition, 1);
+}
+
+// @return translated world position 
+// Specialized version of VertexFactoryGetWorldPosition & ComputeViewAlignedWorldPosition, for voxelization VS shader.
+float4 VertexFactoryGetWorldPosition_Voxelization(FVertexFactoryInput Input, FHairViewInfo HairViewInfo, inout float OutHairWorldRadius, inout float OutHairDensity)
+{
+	const FVertexInfo VertexInfo = GetVertexInfo(Input);
+
+	FHairControlPoint ControlPoint = GetVertexPosition(VertexInfo);
+	OutHairWorldRadius = ControlPoint.WorldRadius;
+	OutHairDensity = HairStrandsVF_Density;
+
+	const float4 WorldPosition = CalcWorldPosition(ControlPoint.Position, GetPrimitiveId(Input));
+	const float3 WorldTangent = GetTangent(Input);
+
+	// Minimal radius to snap the strand to a sample/pixel center (to avoid aliasing)
+	const float3 ViewDir = -HairViewInfo.ViewForward;
+	const float3 Right = normalize(cross(WorldTangent, ViewDir));
+	const float3 OutWorldPosition = WorldPosition.xyz + (VertexInfo.IsLeft ? -Right : Right) * max(ControlPoint.WorldRadius * VertexInfo.LodRadiusScale, HairViewInfo.RadiusAtDepth1);
+
+	return float4(OutWorldPosition, 1);
+}
+
 // @return translated world position
 float4 VertexFactoryGetWorldPosition(FVertexFactoryInput Input, FVertexFactoryIntermediates Intermediates, FHairViewInfo HairViewInfo)
 {
-	FControlPoint ControlPoint = GetVertexPosition(Input);
+	FHairControlPoint ControlPoint = GetVertexPosition(Input);
 	const float3 VertexPosition = ControlPoint.Position;
 	const float4 WorldPosition = CalcWorldPosition(VertexPosition, Intermediates.PrimitiveId);
 
@@ -474,7 +529,7 @@ float4 VertexFactoryGetWorldPosition(FVertexFactoryInput Input, FVertexFactoryIn
 // This function is referenced by several "system" shaders, but none of these shaders (apart HitProxy, are actually used)
 float4 VertexFactoryGetWorldPosition(FVertexFactoryInput Input, FVertexFactoryIntermediates Intermediates)
 {
-	const FHairRenderInfo HairRenderInfo = GetHairRenderInfo(ResolvedView.HairRenderInfo);
+	const FHairRenderInfo HairRenderInfo = GetHairRenderInfo(ResolvedView.HairRenderInfo, ResolvedView.HairRenderInfoBits, UseStableRasterization());
 
 	FHairViewInfo HairViewInfo;
 	HairViewInfo.TranslatedWorldCameraOrigin = ResolvedView.TranslatedWorldCameraOrigin;
diff --git a/Engine/Shaders/Private/HairStrands/HairStrandsVertexFactoryCommon.ush b/Engine/Shaders/Private/HairStrands/HairStrandsVertexFactoryCommon.ush
new file mode 100644
--- /dev/null
+++ b/Engine/Shaders/Private/HairStrands/HairStrandsVertexFactoryCommon.ush
@@ -0,0 +1,30 @@
+// Copyright Epic Games, Inc. All Rights Reserved.
+
+#define HAIR_CONTROLPOINT_INSIDE 0
+#define HAIR_CONTROLPOINT_START	1
+#define HAIR_CONTROLPOINT_END	2
+
+struct FHairControlPoint
+{
+	float3 Position;
+	float  WorldRadius;
+	float  WorldLength;
+	uint   Type;
+};
+
+FHairControlPoint UnpackHairControlPoint(
+	uint4 InPackedData,
+	float3 InVF_PositionOffset,
+	float InVF_Radius,
+	float InVF_Length)
+{
+	const uint PackedAlpha	= asuint(InPackedData.a);
+
+	FHairControlPoint Out = (FHairControlPoint)0;
+	Out.Position	= f16tof32(InPackedData.xyz) + InVF_PositionOffset;
+	Out.WorldRadius	= ((PackedAlpha & 0x00FC) >> 2) / 63.f * InVF_Radius;
+	Out.WorldLength	= ((PackedAlpha & 0xFF00) >> 8) / 255.f * InVF_Length;
+	Out.Type		= PackedAlpha & 0x3;
+
+	return Out;
+}
\ No newline at end of file
diff --git a/Engine/Shaders/Private/HairStrands/HairStrandsVisibilityCommon.ush b/Engine/Shaders/Private/HairStrands/HairStrandsVisibilityCommon.ush
--- a/Engine/Shaders/Private/HairStrands/HairStrandsVisibilityCommon.ush
+++ b/Engine/Shaders/Private/HairStrands/HairStrandsVisibilityCommon.ush
@@ -2,6 +2,7 @@
 
 #pragma once
 
+// Render modes
 #define RENDER_MODE_MSAA 0
 #define RENDER_MODE_TRANSMITTANCE 1
 #define RENDER_MODE_PPLL 2
@@ -10,35 +11,12 @@
 
 #define INVALID_TILE_OFFSET ~0
 
-struct FVisibilityData
-{
-	float3 Tangent;		// 24 bits
-	uint   Coverage8bit;//  8 bits
-	uint   PrimitiveId;	// 28 bits
-	uint   MacroGroupId;//  4 bits
-};
+#define HAIR_VISIBILITY_GROUP_COUNT_WIDTH 64
 
-struct FNodeVis // 128 bit aligments. TODO: Move the coverage computation into the material evaluation for saving 64 bits? 
-{
-	float Depth;		// 32 bits
-	uint  PrimitiveId;	// 32 bits
-	uint  MaterialId;	//  8 bits : this could be packed into a the 3 high bits of PrimitiveId: up to 8 material on screen
-	uint  Coverage8bit;	//  8 bits
-};
+#include "HairStrandsVisibilityCommonStruct.ush"
 
-// Footprint: 128bits | 16bytes
-// Option to reduce the footprint: 
-// * we don't use PrimitiveID at the moment which is encoded on 3bytes
-// * Deferred the evaluation of Coverage & material properties into some other pass?
-
-struct FNodeData
-{
-	float Depth;
-	uint  PrimitiveID_MacroGroupID;
-	uint  Tangent_Coverage8bit;
-	uint  BaseColor_Roughness;
-	uint  Specular_LightChannels;
-};
+ ////////////////////////////////////////////////////////////////////////////////////////////////////////////////////////////////
+ // General util functions
 
 float3 EncodeTangent(float3 N)
 {
@@ -60,63 +38,111 @@ float4 UintToFloat4(uint In)
 {
 	uint4 Out;
 	Out.x = (0xFF &  In);
-	Out.y = (0xFF & (In >>  8));
+	Out.y = (0xFF & (In >> 8));
 	Out.z = (0xFF & (In >> 16));
 	Out.w = (0xFF & (In >> 24));
 	return Out / 255.f;
 }
 
-uint4 EncodeVisibilityData(FVisibilityData Data)
+uint Uint16ToUint32(uint2 In)
 {
-	uint4 Output;
+	return (In.x & 0xFFFF) | ((In.y & 0xFFFF) << 16);
+}
 
-	// Tangent and Coverage
-	// Encode into 2 x 16bits
-	// Coverage is clamped o 255 as we only have 8 bits to store its value
-	float3 T = saturate(EncodeTangent(Data.Tangent));
-	Output.xy = uint2(
-		(0xFF & uint(T.x * 255)) | ((0xFF & uint(T.y * 255)) << 8),
-		(0xFF & uint(T.z * 255)) | (min(uint(0xFF), Data.Coverage8bit) << 8));
+uint2 Uint32ToUint16(uint In)
+{
+	uint2 A;
+	A.x = In & 0xFFFF;
+	A.y = (In >> 16) & 0xFFFF;
+	return A;
+}
 
-	// PrimitiveId and MacroGroupId
-	// Encode into 2 x 16bits
-	// PrimitiveId is on 28bits | MacroGroupId is on 4bits
-	Output.zw = uint2(
-		0xFFFF & Data.PrimitiveId,
-		((Data.MacroGroupId & 0x0F) << 12) | (0x0FFF & (uint(Data.PrimitiveId) >> 16)));
+uint Float16ToUint32(float2 In)
+{
+	return Uint16ToUint32(f32tof16(In));
+}
+
+float2 Uint32ToFloat16(uint In)
+{
+	return f16tof32(Uint32ToUint16(In));
+}
 
-	return Output;
+uint To8bitCoverage(float Coverage)
+{
+	return min(uint(Coverage * 0x100), 0xFFu);
 }
 
-FVisibilityData DecodeVisibilityData(uint4 EncodedData)
+float From8bitCoverage(uint Coverage8bit)
 {
-	FVisibilityData Output;
-	Output.Tangent = DecodeTangent(float3(
-		float((EncodedData.x     ) & 0xFF) / 255.f,
-		float((EncodedData.x >> 8) & 0xFF) / 255.f,
-		float((EncodedData.y     ) & 0xFF) / 255.f));
-	Output.Coverage8bit = (EncodedData.y >> 8) & 0xFF;
+	return float(Coverage8bit) / 255.f;
+}
 
-	Output.PrimitiveId = EncodedData.z | ((EncodedData.w & 0x0FFF) << 16);
-	Output.MacroGroupId = (EncodedData.w & 0xF000) >> 12;
+uint To16bitCoverage(float Coverage)
+{
+	return min(uint(Coverage * 0x10000u), 0xFFFFu);
+}
 
-	return Output;
+float From16bitCoverage(uint Coverage16bit)
+{
+	return float(Coverage16bit) / float(0xFFFF);
 }
 
-bool IsValidHairData(uint4 EncodedVisibilityData)
+uint3 QuantizeTo8Bits(float3 T)
 {
-	const uint InvalidSample = 0xFFFF;
-	return EncodedVisibilityData.w != InvalidSample;
+	return saturate((T + float(1).xxx) * 0.5f) * 0xFF;
 }
 
-uint DecodeMacroGroupId(uint4 EncodedVisibilityData)
+float3 From8bits(float3 In8bits)
 {
-	return (EncodedVisibilityData.w >> 12) & 0xF;
+	return (normalize(In8bits / 255.f) * 2) - float(1).xxx;
 }
 
+uint PackVelocity(float2 EncodedVelocity)
+{
+	return f32tof16(EncodedVelocity.x) << 16 | f32tof16(EncodedVelocity.y);
+}
+
+float2 UnpackVelocity(uint PackedVelocity)
+{
+	return float2(f16tof32(PackedVelocity >> 16), f16tof32(PackedVelocity));
+}
+
+////////////////////////////////////////////////////////////////////////////////////////////////////////////////////////////////
+// Hair samples
+
+// Footprint: 128bits | 16bytes
+// Option to reduce the footprint: 
+// * we don't use PrimitiveID at the moment which is encoded on 3bytes
+// * Deferred the evaluation of Coverage & material properties into some other pass?
+
+struct FPackedHairSample
+{
+	float Depth;
+	uint  PrimitiveID_MacroGroupID;
+	uint  Tangent_Coverage8bit;
+	uint  BaseColor_Roughness;
+	uint  Specular_LightChannels_Backlit;
+};
+
+struct FHairSample
+{
+	float	Depth;
+	float3	Tangent;
+	uint	Coverage8bit;
+	uint	PrimitiveId;
+	uint	MacroGroupId;
+	float3	BaseColor;
+	float	Roughness;
+	float	Specular;
+	float	Backlit;
+	uint	LightChannelMask;
+	bool	bScatterSceneLighting;
+};
+
 // * Max 67,108,864 total nodes (This is 32 nodes per pixel at 1080p if all pixel were covered)
 // * Max 63 nodes per pixel
 // 26bits for offset | 6 bits for count (max 63 nodes)
+//TODO rename FHairSamplesDesc
 struct FNodeDesc
 {
 	uint Offset;
@@ -136,303 +162,325 @@ FNodeDesc DecodeNodeDesc(uint In)
 	return Out;
 }
 
-FNodeData ConvertToNodeData(float Depth, FVisibilityData VisibilityData, float3 BaseColor, float Roughness, float Specular, uint LightChannelMask)
+FPackedHairSample PackHairSample(const FHairSample In)
 {
-	// 4x16bits
-	const uint4 EncodedVisibilityData = EncodeVisibilityData(VisibilityData);
+	FPackedHairSample Out;
 
-	FNodeData Out;
-	Out.Depth					= Depth;																			// 32bits float
-	Out.Tangent_Coverage8bit	= (0xFFFF & EncodedVisibilityData.x) | ((0xFFFF & EncodedVisibilityData.y) << 16);	// 32bits uint
-	Out.PrimitiveID_MacroGroupID= (0xFFFF & EncodedVisibilityData.z) | ((0xFFFF & EncodedVisibilityData.w) << 16);	// 32bits uint
-	Out.BaseColor_Roughness		= float4ToUint(float4(sqrt(BaseColor), Roughness));									// 32bits uint
-	Out.Specular_LightChannels	= clamp(uint(Specular * 0xFF), 0, 0xFF) | ((LightChannelMask & 0x7) << 8);
+	// Tangent and Coverage
+	// Encode into 1 x 32bits uint
+	// Coverage is clamped o 255 as we only have 8 bits to store its value
+	const float3 T = saturate(EncodeTangent(In.Tangent));
+	Out.Tangent_Coverage8bit =
+		((0xFF & uint(T.x * 255)))      | 
+		((0xFF & uint(T.y * 255)) << 8) |
+		((0xFF & uint(T.z * 255)) << 16)| 
+		(min(uint(0xFF), In.Coverage8bit) << 24);
+
+	// PrimitiveId and MacroGroupId
+	// Encode into 1 x 32bits uint
+	// PrimitiveId is on 28bits | MacroGroupId is on 4bits
+	Out.PrimitiveID_MacroGroupID = 
+		((In.PrimitiveId  & 0x0FFFFFFF)) |
+		((In.MacroGroupId & 0xF) << 28);
+
+	Out.Depth					= In.Depth;																			// 32bits float
+	Out.BaseColor_Roughness		= float4ToUint(float4(sqrt(In.BaseColor), In.Roughness));							// 32bits uint
+	Out.Specular_LightChannels_Backlit = 
+		clamp(uint(In.Specular * 0xFF), 0, 0xFF) | 
+		((In.LightChannelMask & 0x7) << 8)  | 
+		((In.bScatterSceneLighting ? 1 : 0) << 12) |
+		(clamp(uint(In.Backlit * 0xFF), 0, 0xFF) << 16);
+
+	return Out;
+}
+
+FHairSample UnpackHairSample(const FPackedHairSample In)
+{
+	FHairSample Out;
+	Out.Depth = In.Depth;
+
+	Out.Tangent = DecodeTangent(float3(
+		float((In.Tangent_Coverage8bit)       & 0xFF) / 255.f,
+		float((In.Tangent_Coverage8bit >> 8)  & 0xFF) / 255.f,
+		float((In.Tangent_Coverage8bit >> 16) & 0xFF) / 255.f));
+	Out.Coverage8bit = (In.Tangent_Coverage8bit >> 24) & 0xFF;
+
+	Out.PrimitiveId  = In.PrimitiveID_MacroGroupID & 0x0FFFFFFF;
+	Out.MacroGroupId = (In.PrimitiveID_MacroGroupID>>28) & 0xF;
+
+	const float4 BaseColorAndRoughness = UintToFloat4(In.BaseColor_Roughness);
+	Out.BaseColor			= BaseColorAndRoughness.xyz * BaseColorAndRoughness.xyz;
+	Out.Roughness			= BaseColorAndRoughness.w;
+	Out.Specular			= float(0xFF & In.Specular_LightChannels_Backlit) / 0xFF;
+	Out.LightChannelMask	= (In.Specular_LightChannels_Backlit >>8) & 0x7;
+	Out.Backlit				= float(0xFF & (In.Specular_LightChannels_Backlit>>16)) / 0xFF;
+	Out.bScatterSceneLighting = (In.Specular_LightChannels_Backlit & (1<<12)) > 0;
 
 	return Out;
 }
 
- void DecodeNodeData(
-	 const FNodeData		InNode, 
-	 inout float			OutDepth, 
-	 inout FVisibilityData	OutVisibilityData,
-	 inout float3			OutBaseColor, 
-	 inout float			OutRoughness,
-	 inout float			OutSpecular,
-	 inout uint				OutLightChannelMask)
+////////////////////////////////////////////////////////////////////////////////////////////////////////////////////////////////
+// Hair categorization
+
+// Rename FHairCatetorizationData
+struct FCategorizationData
 {
-	OutDepth = InNode.Depth;
+	uint  TotalCoverage8bit;
+	uint  SampleCount;
+	float PixelCoverage;
+	float ClosestDepth;
+};
+
+uint4 EncodeCategorizationData(FCategorizationData Data)
+{
+	const uint ClosestDepthUint = asuint(Data.ClosestDepth);
+	const uint ClosestDepthLow  = ClosestDepthUint & 0xFFFF;
+	const uint ClosestDepthHigh	= (ClosestDepthUint >> 16) & 0xFFFF;
+	const uint ComponentX		= min(Data.TotalCoverage8bit, uint(0xFF)) | ((Data.SampleCount & 0xFF)<<8);
+	return uint4(ComponentX, f32tof16(Data.PixelCoverage), ClosestDepthLow, ClosestDepthHigh);
+}
 
-	const uint4 EncodedVisibility = uint4(
-		(InNode.Tangent_Coverage8bit & 0xFFFF),
-		(InNode.Tangent_Coverage8bit >> 16) & 0xFFFF,
-		(InNode.PrimitiveID_MacroGroupID & 0xFFFF),
-		(InNode.PrimitiveID_MacroGroupID >> 16) & 0xFFFF);
-	OutVisibilityData = DecodeVisibilityData(EncodedVisibility);
+FCategorizationData DecodeCategorizationData(uint4 Data)
+{
+	FCategorizationData Out;
+	Out.TotalCoverage8bit		= Data.x & 0xFF;
+	Out.SampleCount				=(Data.x>>8) & 0xFF;
+	Out.PixelCoverage			= f16tof32(Data.y);
+	const uint ClosestDepthLow	= Data.z & 0xFFFF;
+	const uint ClosestDepthHigh	=(Data.w & 0xFFFF) << 16;
+	Out.ClosestDepth			= asfloat(ClosestDepthHigh | ClosestDepthLow);
 
-	const float4 BaseColorAndRoughness = UintToFloat4(InNode.BaseColor_Roughness);
-	OutBaseColor = BaseColorAndRoughness.xyz * BaseColorAndRoughness.xyz;
-	OutRoughness = BaseColorAndRoughness.w;
-	OutSpecular			= float(0xFF & InNode.Specular_LightChannels) / 0xFF;
-	OutLightChannelMask = (InNode.Specular_LightChannels>>8) & 0x7;
+	return Out;
 }
 
-#define HAIR_VISIBILITY_MAXSAMPLECOUNT 8
-#define HAIR_VISIBILITY_GROUP_COUNT_WIDTH 64
+////////////////////////////////////////////////////////////////////////////////////////////////////////////////////////////////
+// Hair view info (used for overriding view constant/uniform buffer)
+
+struct FHairRenderInfo
+{
+	float RadiusAtDepth1Primary;
+	float RadiusAtDepth1Velocity;
+	float VelocityMagnitudeScale;
 
- struct FSampleDesc
- {
-	 uint   SampleId;
-	 uint   MacroGroupId;
-	 float  Depths;
-	 float  Weights;
-	 float3 Tangents;
-	 uint   Coverage8bit;
-	 float4 BaseColorAndRoughness;
-	 float  Specular;
- };
-
- struct FSampleSetDesc
- {
-	 uint   UniqueSampleCount;
-	 uint   ValidSampleCount;
-	 uint   HairSampleCount;
- };
-
- uint To8bitCoverage(float Coverage)
- {
-	 return min(uint(Coverage * 0x100), 0xFFu);
- }
-
- float From8bitCoverage(uint Coverage8bit)
- {
-	 return float(Coverage8bit) / 255.f;
- }
-
- uint To16bitCoverage(float Coverage)
- {
-	 return min(uint(Coverage * 0x10000u), 0xFFFFu);
- }
-
- float From16bitCoverage(uint Coverage16bit)
- {
-	 return float(Coverage16bit) / float(0xFFFF);
- }
-
- struct FCategorizationData
- {
-	 uint  TotalCoverage8bit;
-	 uint  SampleCount;
-	 float PixelCoverage;
-	 float ClosestDepth;
- };
-
- uint4 EncodeCategorizationData(FCategorizationData Data)
- {
-	 const uint ClosestDepthUint = asuint(Data.ClosestDepth);
-	 const uint ClosestDepthLow  = ClosestDepthUint & 0xFFFF;
-	 const uint ClosestDepthHigh = (ClosestDepthUint >> 16) & 0xFFFF;
-	 const uint ComponentX		 = min(Data.TotalCoverage8bit, uint(0xFF)) | ((Data.SampleCount & 0xFF)<<8);
-	 return uint4(ComponentX, f32tof16(Data.PixelCoverage), ClosestDepthLow, ClosestDepthHigh);
- }
-
- FCategorizationData DecodeCategorizationData(uint4 Data)
- {
-	 FCategorizationData Out;
-	 Out.TotalCoverage8bit			= Data.x & 0xFF;
-	 Out.SampleCount				=(Data.x>>8) & 0xFF;
-	 Out.PixelCoverage				= f16tof32(Data.y);
-	 const uint ClosestDepthLow		= Data.z & 0xFFFF;
-	 const uint ClosestDepthHigh	=(Data.w & 0xFFFF) << 16;
-	 Out.ClosestDepth				= asfloat(ClosestDepthHigh | ClosestDepthLow);
-
-	 return Out;
- }
-
- uint3 QuantizeTo8Bits(float3 T)
- {
-	 return saturate((T + float(1).xxx) * 0.5f) * 0xFF;
- }
-
- float3 From8bits(float3 In8bits)
- {
-	 return (normalize(In8bits / 255.f) * 2) - float(1).xxx;
- }
-
- uint EncodeBaseColorAndRoughness(float3 BaseColor, float Roughness)
- {
-	 return
-		 (uint(BaseColor.x * 0xFF) & 0xFF) |
-		 (uint(BaseColor.y * 0xFF) & 0xFF) << 8 |
-		 (uint(BaseColor.z * 0xFF) & 0xFF) << 16 |
-		 (uint(Roughness   * 0xFF) & 0xFF) << 24;
- }
-
- void DecodeBaseColorAndRoughness(uint Encoded, inout float3 BaseColor, out float Roughness)
- {
-	 BaseColor.x = (float((Encoded) & 0xFF) / 255.f);
-	 BaseColor.y = (float((Encoded >> 8) & 0xFF) / 255.f);
-	 BaseColor.z = (float((Encoded >> 16) & 0xFF) / 255.f);
-	 Roughness   = (float((Encoded >> 24) & 0xFF) / 255.f);
- }
-
- struct FHairRenderInfo
- {
-	 float RadiusAtDepth1Primary;
-	 float RadiusAtDepth1Velocity;
-	 float VelocityMagnitudeScale;
-
-	 bool bIsOrthoView;
-	 bool bIsGPUDriven;
-
-	 float3 ViewForward;
-	 float3 TranslatedWorldCameraOrigin;
- };
-
- FHairRenderInfo GetHairRenderInfo(float4 ViewHairRenderInfo)
- {
-	 FHairRenderInfo Info;
-	 Info.RadiusAtDepth1Primary = ViewHairRenderInfo.x;
-	 Info.RadiusAtDepth1Velocity = ViewHairRenderInfo.y;
-	 Info.VelocityMagnitudeScale = ViewHairRenderInfo.z;
-
-	 const uint BitField = asuint(ViewHairRenderInfo.w);
-	 Info.bIsOrthoView = (BitField & 0x1) != 0;
-	 Info.bIsGPUDriven = (BitField & 0x2) != 0;
-
-	 return Info;
- }
+	bool bIsOrthoView;
+	bool bIsGPUDriven;
+
+	float3 ViewForward;
+	float3 TranslatedWorldCameraOrigin;
+};
+
+FHairRenderInfo GetHairRenderInfo(float4 ViewHairRenderInfo, uint ViewHairRenderInfoBits, bool bUseScableRasterization=false)
+{
+	FHairRenderInfo Info;
+	Info.RadiusAtDepth1Primary  = bUseScableRasterization ? ViewHairRenderInfo.y : ViewHairRenderInfo.x;
+	Info.RadiusAtDepth1Velocity = ViewHairRenderInfo.z;
+	Info.VelocityMagnitudeScale = ViewHairRenderInfo.w;
+
+	const uint BitField = ViewHairRenderInfoBits;
+	Info.bIsOrthoView = (BitField & 0x1) != 0;
+	Info.bIsGPUDriven = (BitField & 0x2) != 0;
+
+	return Info;
+}
  
- struct FHairViewInfo
- {
-	 float3 TranslatedWorldCameraOrigin;
-	 float3 ViewForward;
-	 float  RadiusAtDepth1;
-	 bool	bIsOrthoView;
- };
-
- struct FHairAttributes
- {
-	 float2 RootUV;
-	 float	Seed;
-	 float	UCoord;
- };
-
- FHairAttributes DecodeHairAttribute(uint4 Attributes)
- {
-	 float2 RootUV, RootIndex;
-	 RootUV.x		= (Attributes.x & 0xFF) / 255.f;
-	 RootUV.y		= ((Attributes.x>>8) & 0xFF) / 255.f;
-	 RootIndex.x	= Attributes.z & 0xFF;
-	 RootIndex.y	= (Attributes.z >> 8) & 0xFF;
-
-	 FHairAttributes Out;
-	 Out.RootUV		= RootUV + RootIndex;
-	 Out.UCoord		= (Attributes.y & 0xFF) / 255.f;
-	 Out.Seed		= ((Attributes.y>>8) & 0xFF) / 255.f;
-	 return Out;
- }
-
- //////////
- // PPLL
-
-// Footprint: 224bits | 28bytes
- struct FPPLLNodeData
- {
-	 float Depth;
-	 uint Specular_LightChannel_Coverage16bit;
-	 uint PrimitiveID_MacroGroupID;
-	 uint Tangent_Coverage8bit;
-	 uint BaseColor_Roughness;
-	 uint NextNodeIndex;
-	 uint PackedVelocity;
- };
-
- uint PackVelocity(float2 EncodedVelocity)
- {
-	 return f32tof16(EncodedVelocity.x) << 16 | f32tof16(EncodedVelocity.y);
- }
-
- float2 UnpackVelocity(uint PackedVelocity)
- {
-	 return float2(f16tof32(PackedVelocity >> 16), f16tof32(PackedVelocity));
- }
-
- FPPLLNodeData ConvertToPPLLNodeData(float Depth, float CoverageFullPrecision, FVisibilityData VisibilityData, float3 BaseColor, float Roughness, float Specular, uint LightChannelMask, float2 Velocity, uint NextNodeIndex)
- {
-	 // 4x16bits
-	 const uint4 EncodedVisibilityData = EncodeVisibilityData(VisibilityData);
-	 const uint Coverage16bit = To16bitCoverage(CoverageFullPrecision);
-
-	 FPPLLNodeData Out;
-	 Out.Depth = Depth;
-	 Out.Tangent_Coverage8bit = (0xFFFF & EncodedVisibilityData.x) | ((0xFFFF & EncodedVisibilityData.y) << 16);
-	 Out.PrimitiveID_MacroGroupID = (0xFFFF & EncodedVisibilityData.z) | ((0xFFFF & EncodedVisibilityData.w) << 16);
-	 Out.BaseColor_Roughness = float4ToUint(float4(sqrt(BaseColor), Roughness));
-	 Out.Specular_LightChannel_Coverage16bit = clamp(uint(Specular * 0xFF), 0, 0xFF) | ((LightChannelMask & 0x7)<<8) | ((0xFFFF & Coverage16bit) << 16);
-	 Out.NextNodeIndex = NextNodeIndex;
-	 Out.PackedVelocity = PackVelocity(Velocity);
-
-	 return Out;
- }
-
- FNodeData ConvertToNodeData(FPPLLNodeData In, inout uint OutCoverage16bit)
- {
-	 FNodeData Out;
-	 Out.Depth = In.Depth;
-	 Out.PrimitiveID_MacroGroupID = In.PrimitiveID_MacroGroupID;
-	 Out.Tangent_Coverage8bit = In.Tangent_Coverage8bit;
-	 Out.BaseColor_Roughness = In.BaseColor_Roughness;
-	 Out.Specular_LightChannels = In.Specular_LightChannel_Coverage16bit & 0xFFFF;
-
-	 OutCoverage16bit = (In.Specular_LightChannel_Coverage16bit >> 16) & 0xFFFF;
-
-	 return Out;
- }
-
- uint Uint16ToUint32(uint2 In)
- {
-	 return (In.x & 0xFFFF) | ((In.y & 0xFFFF) << 16);
- }
-
- uint2 Uint32ToUint16(uint In)
- {
-	 uint2 A;
-	 A.x = In & 0xFFFF;
-	 A.y = (In >> 16) & 0xFFFF;
-	 return A;
- }
-
- uint Float16ToUint32(float2 In)
- {
-	 return Uint16ToUint32(f32tof16(In));
- }
-
- float2 Uint32ToFloat16(uint In)
- {
-	return f16tof32(Uint32ToUint16(In));
- }
-
- struct FVisSample
- {
-	 uint PrimitiveId;
-	 uint MaterialId;
- };
-
- uint GetInvalidHairPrimitiveId()
- {
-	return 0xFFFFFFFF;
- }
-
- uint EncodeVisSample(uint PrimitiveId, uint MaterialId)
- {
-	 return (PrimitiveId & 0x0FFFFFFF) | ((MaterialId & 0xF) << 28);
- }
-
- FVisSample DecodeVisSample(uint RawVis)
- {
-	 FVisSample Out;
-	 Out.PrimitiveId = RawVis & 0x0FFFFFFF;
-	 Out.MaterialId = (RawVis >> 28) & 0xF;
-	 return Out;
- }
+struct FHairViewInfo
+{
+	float3 TranslatedWorldCameraOrigin;
+	float3 ViewForward;
+	float  RadiusAtDepth1;
+	bool	bIsOrthoView;
+};
+
+////////////////////////////////////////////////////////////////////////////////////////////////////////////////////////////////
+// Hair Attributes
+
+struct FHairAttributes
+{
+	float2 RootUV;
+	float	Seed;
+	float	UCoord;
+};
+
+FHairAttributes DecodeHairAttribute(uint4 Attributes)
+{
+	float2 RootUV, RootIndex;
+	RootUV.x	= (Attributes.x & 0xFF) / 255.f;
+	RootUV.y	= ((Attributes.x>>8) & 0xFF) / 255.f;
+	RootIndex.x	= Attributes.z & 0xFF;
+	RootIndex.y	= (Attributes.z >> 8) & 0xFF;
+
+	FHairAttributes Out;
+	Out.RootUV	= RootUV + RootIndex;
+	Out.UCoord	= (Attributes.y & 0xFF) / 255.f;
+	Out.Seed	= ((Attributes.y>>8) & 0xFF) / 255.f;
+	return Out;
+}
+
+////////////////////////////////////////////////////////////////////////////////////////////////////////////////////////////////
+// Hair PPLL
+
+uint GetCoverage8Bits(FPPLLNodeData In)
+{
+	return (In.Tangent_Coverage8bit >> 24) & 0xFF;
+}
+
+uint GetCoverage16Bits(FPPLLNodeData In)
+{
+	return (In.Specular_LightChannel_Backlit_Coverage16bit >> 16) & 0xFFFF;
+}
+
+FPPLLNodeData CreatePPLLNodeData(
+	float	Depth, 
+	uint	PrimitiveId,
+	uint	MacroGroupId,
+	float3	Tangent,
+	float	CoverageFullPrecision, 
+	uint	Coverage8bit,
+	float3	BaseColor,
+	float	Roughness, 
+	float	Specular, 
+	float	Backlit, 
+	uint	LightChannelMask, 
+	float2	Velocity, 
+	uint	NextNodeIndex)
+{
+	FHairSample S;
+	S.Depth = Depth;
+	S.PrimitiveId = PrimitiveId;
+	S.MacroGroupId = MacroGroupId;
+	S.Tangent = Tangent;
+	S.Coverage8bit = Coverage8bit;
+	S.BaseColor = BaseColor;
+	S.Roughness = Roughness;
+	S.Specular = 0;			// Custom encoding for PPLL
+	S.Backlit = 0;			// Custom encoding for PPLL
+	S.LightChannelMask = 0; // Custom encoding for PPLL
+	
+	const FPackedHairSample PackedS = PackHairSample(S);
+	const uint Coverage16bit = To16bitCoverage(CoverageFullPrecision);
+
+	FPPLLNodeData Out;
+	Out.Depth = PackedS.Depth;
+	Out.PrimitiveID_MacroGroupID = PackedS.PrimitiveID_MacroGroupID;
+	Out.Tangent_Coverage8bit = PackedS.Tangent_Coverage8bit;
+	Out.BaseColor_Roughness = PackedS.BaseColor_Roughness;
+	Out.Specular_LightChannel_Backlit_Coverage16bit = clamp(uint(Specular * 0xFF), 0, 0xFF) |  (clamp(uint(Backlit * 0x1F), 0, 0x1F)<<11) | ((LightChannelMask & 0x7)<<8) |  ((0xFFFF & Coverage16bit) << 16);
+	Out.NextNodeIndex = NextNodeIndex;
+	Out.PackedVelocity = PackVelocity(Velocity);
+
+	return Out;
+}
+
+FPackedHairSample ConvertToPackedHairSample(FPPLLNodeData In, uint InNewCoverage8bit)
+{
+	const uint Specular_LightChannel = In.Specular_LightChannel_Backlit_Coverage16bit & 0x7FF;
+	const float Backlit = float((In.Specular_LightChannel_Backlit_Coverage16bit >> 11) & 0x1F) / float(0x1F);
+
+	FPackedHairSample Out;
+	Out.Depth = In.Depth;
+	Out.PrimitiveID_MacroGroupID = In.PrimitiveID_MacroGroupID;
+	Out.Tangent_Coverage8bit = (In.Tangent_Coverage8bit & 0x00FFFFFF) | ((InNewCoverage8bit & 0xFF)<<24);
+	Out.BaseColor_Roughness = In.BaseColor_Roughness;
+	Out.Specular_LightChannels_Backlit = Specular_LightChannel | (clamp(uint(Backlit * 0xFF), 0, 0xFF) << 16);
+
+	return Out;
+}
+
+ ////////////////////////////////////////////////////////////////////////////////////////////////////////////////////////////////
+ // Hair visibility buffer 
+
+// Renamed into FHairNodeVisFat
+struct FNodeVis // 128 bit aligments. TODO: Move the coverage computation into the material evaluation for saving 64 bits? 
+{
+	float Depth;		// 32 bits
+	uint  PrimitiveId;	// 32 bits
+	uint  MaterialId;	//  8 bits : this could be packed into a the 3 high bits of PrimitiveId: up to 8 material on screen
+	uint  Coverage8bit;	//  8 bits
+};
+
+struct FHairVisSample
+{
+	uint PrimitiveId;
+	uint MaterialId;
+};
+
+uint GetInvalidHairPrimitiveId()
+{
+return 0xFFFFFFFF;
+}
+
+uint PackHairVisSample(uint PrimitiveId, uint MaterialId)
+{
+	return (PrimitiveId & 0x0FFFFFFF) | ((MaterialId & 0xF) << 28);
+}
+
+FHairVisSample UnpackHairVisSample(uint RawVis)
+{
+	FHairVisSample Out;
+	Out.PrimitiveId = RawVis & 0x0FFFFFFF;
+	Out.MaterialId = (RawVis >> 28) & 0xF;
+	return Out;
+}
  
+ ////////////////////////////////////////////////////////////////////////////////////////////////////////////////////////////////
+ // Hair MSAA buffer 
+
+bool IsValidHairData(uint2 MRT0Data)
+{
+	const uint InvalidSample = GetInvalidHairPrimitiveId();// &0x0FFFFFFF; // PrimitiveId mask, see EncodeVisSample()
+	return MRT0Data.x != InvalidSample;
+}
+
+uint2 PackHairMRT0(uint PrimitiveId, uint MacroGroupId, uint Coverage8bit, float3 Tangent)
+{
+	uint2 Out;
+	Out.x =
+		((PrimitiveId & 0x0FFFFFFF)) |
+		((MacroGroupId & 0xF) << 28);
+
+	const float3 T = saturate(EncodeTangent(Tangent));
+	Out.y =
+		((0xFF & uint(T.x * 255))) |
+		((0xFF & uint(T.y * 255)) << 8) |
+		((0xFF & uint(T.z * 255)) << 16) |
+		(min(uint(0xFF), Coverage8bit) << 24);
+
+	return Out;
+}
+
+float4 PackHairMRT1(float3 BaseColor, float Roughness)
+{
+	return float4(sqrt(BaseColor), Roughness);
+}
+
+float4 PackHairMRT2(float Specular, uint LightChannelMask, float Backlit, bool bUseScatterSceneLighting)
+{
+	return float4(Specular, float(LightChannelMask) / 255.f, Backlit, bUseScatterSceneLighting? 1 : 0);
+}
+
+void UnpackHairMRT0(uint2 In, inout uint PrimitiveId, inout uint MacroGroupId, inout uint Coverage8bit, inout float3 Tangent)
+{
+	PrimitiveId  =  In.x & 0x0FFFFFFF;
+	MacroGroupId = (In.x >> 28) & 0xF;
+
+	Tangent = DecodeTangent(float3(
+		float((In.y) & 0xFF) / 255.f,
+		float((In.y >> 8) & 0xFF) / 255.f,
+		float((In.y >> 16) & 0xFF) / 255.f));
+	Coverage8bit = (In.y >> 24) & 0xFF;
+}
+
+void UnpackHairMRT1(float4 In, inout float3 BaseColor, inout float Roughness)
+{
+	// Cheap gamma encoding since this is encoded into 8bit target
+	BaseColor = In.xyz * In.xyz;
+	Roughness = In.w;
+}
+
+void UnpackHairMRT2(float4 In, inout float Specular, inout uint LightChannelMask, inout float Backlit, inout bool bUseScatterSceneLighting)
+{
+	Specular = In.x;
+	LightChannelMask = uint(In.y * 0xFF);
+	Backlit = In.z;
+	bUseScatterSceneLighting = In.w > 0;
+}
\ No newline at end of file
diff --git a/Engine/Shaders/Private/HairStrands/HairStrandsVisibilityCommonStruct.ush b/Engine/Shaders/Private/HairStrands/HairStrandsVisibilityCommonStruct.ush
new file mode 100644
--- /dev/null
+++ b/Engine/Shaders/Private/HairStrands/HairStrandsVisibilityCommonStruct.ush
@@ -0,0 +1,15 @@
+// Copyright Epic Games, Inc. All Rights Reserved.
+
+#pragma once
+
+// Footprint: 224bits | 28bytes
+struct FPPLLNodeData
+{
+	float Depth;
+	uint Specular_LightChannel_Backlit_Coverage16bit;
+	uint PrimitiveID_MacroGroupID;
+	uint Tangent_Coverage8bit;
+	uint BaseColor_Roughness;
+	uint NextNodeIndex;
+	uint PackedVelocity;
+};
diff --git a/Engine/Shaders/Private/HairStrands/HairStrandsVisibilityCompaction.usf b/Engine/Shaders/Private/HairStrands/HairStrandsVisibilityCompaction.usf
--- a/Engine/Shaders/Private/HairStrands/HairStrandsVisibilityCompaction.usf
+++ b/Engine/Shaders/Private/HairStrands/HairStrandsVisibilityCompaction.usf
@@ -28,7 +28,6 @@ float CosTangentThreshold;
 
 uint2 OutputResolution;
 uint MaxNodeCount;
-uint HairVisibilityMSAASampleCount;
 uint2 ResolutionOffset;
 
 RWTexture2D<uint> OutCompactNodeCounter;
@@ -36,7 +35,7 @@ RWTexture2D<uint> OutCompactNodeIndex;
 #if PERMUTATION_VISIBILITY
 	RWStructuredBuffer<FNodeVis> OutCompactNodeData;
 #else
-	RWStructuredBuffer<FNodeData> OutCompactNodeData;
+	RWStructuredBuffer<FPackedHairSample> OutCompactNodeData;
 #endif
 RWStructuredBuffer<uint> OutCompactNodeCoord;
 RWTexture2D<uint4> OutCategorizationTexture;
@@ -57,6 +56,10 @@ RWTexture2D<uint4> OutCategorizationTexture;
 groupshared uint AllocationNodeCount;
 groupshared uint AllocationNodeOffset;
 
+///////////////////////////////////////////////////////////////////////////////////////////////////////////////
+// PPLL
+///////////////////////////////////////////////////////////////////////////////////////////////////////////////
+
 #if PERMUTATION_PPLL
 #define PPLL_MAX_RENDER_NODE_COUNT PERMUTATION_PPLL
 
@@ -102,19 +105,10 @@ void MainCS(uint GroupIndex : SV_GroupIndex, uint3 DispatchThreadId : SV_Dispatc
 		uint NodeIndex = FirstNodeIndex;
 		while (NodeIndex != 0xFFFFFFFF)
 		{
-			FPPLLNodeData PPLLNode = PPLLNodeData[NodeIndex];
-
-			// Decode data we need
-			float			Depth = 0.0f;
-			FVisibilityData	VisibilityData = (FVisibilityData)0;
-			float3			BaseColor = 0.0f;
-			float			Roughness = 0.0f;
-			float			Specular = 0;
-			uint			LightChannelMask = 0;
-			uint			Coverage16bit = 0;
-			DecodeNodeData(ConvertToNodeData(PPLLNode, Coverage16bit), Depth, VisibilityData, BaseColor, Roughness, Specular, LightChannelMask);
-
-			const float Coverage = From8bitCoverage(VisibilityData.Coverage8bit);
+			const FPPLLNodeData PPLLNode = PPLLNodeData[NodeIndex];
+			const uint Coverage8bit		 = GetCoverage8Bits(PPLLNode);
+			const float Coverage		 = From8bitCoverage(Coverage8bit);
+
 			TotalTransmittance *= 1.0f - Coverage;
 
 			int InsertPos;
@@ -145,7 +139,7 @@ void MainCS(uint GroupIndex : SV_GroupIndex, uint3 DispatchThreadId : SV_Dispatc
 
 				// Insert new node
 				SortedIndex[InsertPos] = NodeIndex;
-				SortedDepth[InsertPos] = Depth;
+				SortedDepth[InsertPos] = PPLLNode.Depth;
 
 				PixelNodeCount++;
 			}
@@ -181,20 +175,11 @@ void MainCS(uint GroupIndex : SV_GroupIndex, uint3 DispatchThreadId : SV_Dispatc
 		float ValidPixelSampleTotalCoverage = 0.0f;
 		float SortedCoverage[PPLL_MAX_RENDER_NODE_COUNT];
 		float TotalSortedTransmittance = 1.0f;
-		for (int i = 0; i < PixelRenderNodecount; ++i)
+		for (uint i = 0; i < PixelRenderNodecount; ++i)
 		{
-			const uint NodePPLLIndex = SortedIndex[i];
-			FPPLLNodeData PPLLNode = PPLLNodeData[NodePPLLIndex];
-
-			// Decode data we need
-			float			Depth = 0.0f;
-			FVisibilityData	VisibilityData = (FVisibilityData)0;
-			float3			BaseColor = 0.0f;
-			float			Roughness = 0.0f;
-			float			Specular = 0;
-			uint			LightChannelMask = 0;
-			float			Coverage16bit = 0;
-			DecodeNodeData(ConvertToNodeData(PPLLNode, Coverage16bit), Depth, VisibilityData, BaseColor, Roughness, Specular, LightChannelMask);
+			const uint NodePPLLIndex	 = SortedIndex[i];
+			const FPPLLNodeData PPLLNode = PPLLNodeData[NodePPLLIndex];
+			const uint Coverage16bit	 = GetCoverage16Bits(PPLLNode);
 
 			const float Coverage = From16bitCoverage(Coverage16bit);
 
@@ -210,35 +195,23 @@ void MainCS(uint GroupIndex : SV_GroupIndex, uint3 DispatchThreadId : SV_Dispatc
 
 		//////////
 		// Third pass: write out compact nodes for rendering
-		for (uint i = 0; i < PixelRenderNodecount; ++i)
+		for (uint j = 0; j < PixelRenderNodecount; ++j)
 		{
-			const uint NodePPLLIndex = SortedIndex[i];
+			const uint NodePPLLIndex = SortedIndex[j];
 			uint Coverage16bit = 0;
 			FPPLLNodeData PPLLNode = PPLLNodeData[NodePPLLIndex];
-			FNodeData OutNodeData = ConvertToNodeData(PPLLNode, Coverage16bit);
-			float			Depth = 0.0f;
-			FVisibilityData	VisibilityData = (FVisibilityData)0;
-			float3			BaseColor = 0.0f;
-			float			Roughness = 0.0f;
-			float			Specular = 0.0f;
-			uint			LightChannelMask = 0;
-			DecodeNodeData(OutNodeData, Depth, VisibilityData, BaseColor, Roughness, Specular, LightChannelMask);
 
 			// Coverage8bit is a weight normalising to 1 the contribution of all the compacted samples. Because later it is weighted by Categorization.PixelCoverage.
-			uint Coverage8bit = To8bitCoverage(SortedCoverage[i] / float(ValidPixelSampleTotalCoverage));
 			// Patch the coverage on the out node
-			VisibilityData.Coverage8bit = Coverage8bit;
-			OutNodeData = ConvertToNodeData(Depth, VisibilityData, BaseColor, Roughness, Specular, LightChannelMask);
+			const float PatchedCoverage8bit = To8bitCoverage(SortedCoverage[j] / float(ValidPixelSampleTotalCoverage));
+			FPackedHairSample OutNodeData = ConvertToPackedHairSample(PPLLNode, PatchedCoverage8bit);
 
-			Categorization.TotalCoverage8bit += Coverage8bit;
+			Categorization.TotalCoverage8bit += PatchedCoverage8bit;
 			Categorization.ClosestDepth = max(Categorization.ClosestDepth, OutNodeData.Depth); // Inverse-Z
 
-
-
-			// TODOAGREGATE if last, Create FVisibilityData and FNodeData from aggregated  data
-			//FNodeData ConvertToNodeData(float Depth, FVisibilityData VisibilityData, float3 BaseColor, float Roughness, float Specular)
-			OutCompactNodeData[OffsetInStorage + i] = OutNodeData;
-			OutCompactNodeCoord[OffsetInStorage + i] = (PixelCoord.x & 0xFFFF) | ((PixelCoord.y & 0xFFFF) << 16);
+			// TODOAGREGATE if last, Create FPackedHairSample from aggregated  data
+			OutCompactNodeData[OffsetInStorage + j] = OutNodeData;
+			OutCompactNodeCoord[OffsetInStorage + j] = (PixelCoord.x & 0xFFFF) | ((PixelCoord.y & 0xFFFF) << 16);
 
 
 
@@ -286,24 +259,37 @@ void MainCS(uint GroupIndex : SV_GroupIndex, uint3 DispatchThreadId : SV_Dispatc
 }
 
 
-#else // PERMUTATION_PPLL //////////////////////////////////////////////////////////////////////////////////////////////////////////////
+#else // PERMUTATION_PPLL
+
+///////////////////////////////////////////////////////////////////////////////////////////////////////////////
+// MSAA visiblity buffer
+///////////////////////////////////////////////////////////////////////////////////////////////////////////////
 
+struct FSampleSetDesc
+{
+	uint   UniqueSampleCount;
+	uint   ValidSampleCount;
+	uint   HairSampleCount;
+};
+
+#ifndef PERMUTATION_MSAACOUNT
+	#error PERMUTATION_MSAACOUNT is not defined
+#endif
 
 #if PERMUTATION_VISIBILITY
 
 void ComputeUniqueSamplesWithoutDepthTest(
-	inout uint3 OutSamples[HAIR_VISIBILITY_MAXSAMPLECOUNT],
+	inout uint3 OutSamples[PERMUTATION_MSAACOUNT],
 	inout FSampleSetDesc OutSet,
 	const uint2 PixelCoord,
-	const uint SampleCount,
 	const float OpaqueDepth,
-	Texture2DMS<uint, HAIR_VISIBILITY_MAXSAMPLECOUNT>  InMSAA_IDTexture,
-	Texture2DMS<float, HAIR_VISIBILITY_MAXSAMPLECOUNT>  InMSAA_DepthTexture)
+	Texture2DMS<uint, PERMUTATION_MSAACOUNT>  InMSAA_IDTexture,
+	Texture2DMS<float, PERMUTATION_MSAACOUNT>  InMSAA_DepthTexture)
 {
 	OutSet.UniqueSampleCount = 0;
 	OutSet.ValidSampleCount = 0;
 
-	OutSet.HairSampleCount = clamp(SampleCount, 0, HAIR_VISIBILITY_MAXSAMPLECOUNT);
+	OutSet.HairSampleCount = PERMUTATION_MSAACOUNT;
 	const uint InvalidSample = 0xFFFF;
 
 	float ClosestDepth = 0; // Inverse-Z
@@ -351,12 +337,8 @@ void ComputeUniqueSamplesWithoutDepthTest(
 	}
 }
 
-#ifndef HAIR_VISIBILITY_MAXSAMPLECOUNT
-#error HAIR_VISIBILITY_MAXSAMPLECOUNT is not defined
-#endif
-
-Texture2DMS<uint, HAIR_VISIBILITY_MAXSAMPLECOUNT> MSAA_IDTexture;
-Texture2DMS<float, HAIR_VISIBILITY_MAXSAMPLECOUNT> MSAA_DepthTexture;
+Texture2DMS<uint, PERMUTATION_MSAACOUNT> MSAA_IDTexture;
+Texture2DMS<float, PERMUTATION_MSAACOUNT> MSAA_DepthTexture;
 
 #if PERMUTATION_VIEWTRANSMITTANCE == 1
 Texture2D<float> ViewTransmittanceTexture;
@@ -372,18 +354,17 @@ void MainCS(uint GroupIndex : SV_GroupIndex, uint3 DispatchThreadId : SV_Dispatc
 	}
 
 	uint2 PixelCoord = DispatchThreadId.xy + ResolutionOffset;
-	if (HairVisibilityMSAASampleCount == 0 || PixelCoord.x >= OutputResolution.x || PixelCoord.y >= OutputResolution.y)
+	if (PixelCoord.x >= OutputResolution.x || PixelCoord.y >= OutputResolution.y)
 		PixelCoord = uint2(0, 0);
 
 	FSampleSetDesc SampleDesc;
 
 	const float OpaqueDepth = SceneTexturesStruct.SceneDepthTexture.Load(uint3(PixelCoord, 0)).r;
-	uint3 Samples[HAIR_VISIBILITY_MAXSAMPLECOUNT];		// x:PrimitiveId|MaterialId, y:Weight, z:SampleIt
+	uint3 Samples[PERMUTATION_MSAACOUNT];		// x:PrimitiveId|MaterialId, y:Weight, z:SampleIt
 	ComputeUniqueSamplesWithoutDepthTest(
 		Samples,
 		SampleDesc,
 		PixelCoord,
-		HairVisibilityMSAASampleCount,
 		OpaqueDepth,
 		MSAA_IDTexture, 
 		MSAA_DepthTexture);
@@ -411,7 +392,7 @@ void MainCS(uint GroupIndex : SV_GroupIndex, uint3 DispatchThreadId : SV_Dispatc
 		{
 			for (uint OutIndex = 0; OutIndex < NodeDesc.Count; ++OutIndex)
 			{
-				const FVisSample VisSample = DecodeVisSample(Samples[OutIndex].x);
+				const FHairVisSample VisSample = UnpackHairVisSample(Samples[OutIndex].x);
 
 				// VisibilityData.Coverage8bit is a weight normalising to 1 the contribution of all the compacted samples. Because later it is weighted by Categorization.PixelCoverage.
 				FNodeVis OutNodeVis;
@@ -436,7 +417,7 @@ void MainCS(uint GroupIndex : SV_GroupIndex, uint3 DispatchThreadId : SV_Dispatc
 #if PERMUTATION_VIEWTRANSMITTANCE == 1
 	Categorization.PixelCoverage = 1.0f - ViewTransmittanceTexture.Load(uint3(PixelCoord, 0));
 #else
-	Categorization.PixelCoverage = SampleDesc.ValidSampleCount / float(HAIR_VISIBILITY_MAXSAMPLECOUNT);
+	Categorization.PixelCoverage = SampleDesc.ValidSampleCount / float(PERMUTATION_MSAACOUNT);
 #endif
 
 	OutCompactNodeIndex[PixelCoord] = EncodeNodeDesc(NodeDesc);
@@ -445,28 +426,31 @@ void MainCS(uint GroupIndex : SV_GroupIndex, uint3 DispatchThreadId : SV_Dispatc
 
 #else // PERMUTATION_VISIBILITY
 
+///////////////////////////////////////////////////////////////////////////////////////////////////////////////
+// MSAA buffer
+///////////////////////////////////////////////////////////////////////////////////////////////////////////////
+
 void ComputeUniqueSamplesWithoutDepthTest(
-	inout uint3 OutSamples[HAIR_VISIBILITY_MAXSAMPLECOUNT],
+	inout uint3 OutSamples[PERMUTATION_MSAACOUNT],
 	inout FSampleSetDesc OutSet,
 	const uint2 PixelCoord,
-	const uint SampleCount,
 	const float OpaqueDepth,
-	Texture2DMS<uint4, HAIR_VISIBILITY_MAXSAMPLECOUNT>  InMSAA_IDTexture,
-	Texture2DMS<float, HAIR_VISIBILITY_MAXSAMPLECOUNT>  InMSAA_DepthTexture)
+	Texture2DMS<uint2, PERMUTATION_MSAACOUNT>  InMSAA_IDTexture,
+	Texture2DMS<float, PERMUTATION_MSAACOUNT>  InMSAA_DepthTexture)
 {
 	OutSet.UniqueSampleCount = 0;
 	OutSet.ValidSampleCount = 0;
 
-	OutSet.HairSampleCount = clamp(SampleCount, 0, HAIR_VISIBILITY_MAXSAMPLECOUNT);
+	OutSet.HairSampleCount = PERMUTATION_MSAACOUNT;
 	const uint InvalidSample = 0xFFFF;
 
 #if PERMUTATION_MATERIAL_COMPACTION == 1
-	float4 UniqueTangentAndDepths[HAIR_VISIBILITY_MAXSAMPLECOUNT];
+	float4 UniqueTangentAndDepths[PERMUTATION_MSAACOUNT];
 #endif
 	float ClosestDepth = 0; // Inverse-Z
 	for (uint SampleIt = 0; SampleIt < OutSet.HairSampleCount; ++SampleIt)
 	{
-		const uint4 EncodedData = InMSAA_IDTexture.Load(PixelCoord, SampleIt);
+		const uint2 EncodedData = InMSAA_IDTexture.Load(PixelCoord, SampleIt);
 		const bool bIsValid = IsValidHairData(EncodedData);
 		if (!bIsValid)
 			continue;
@@ -481,17 +465,22 @@ void ComputeUniqueSamplesWithoutDepthTest(
 		const float SceneDepth = ConvertFromDeviceZ(SampleDepth);
 	#endif
 
-		const FVisibilityData VisibilityData = DecodeVisibilityData(EncodedData);
+		uint PrimitiveId;
+		uint MacroGroupId;
+		uint Coverage8bit;
+		float3 Tangent;
+
+		UnpackHairMRT0(EncodedData, PrimitiveId, MacroGroupId, Coverage8bit, Tangent);
 		bool bAlreadyExist = false;
 		for (uint UniqueIt = 0; UniqueIt < OutSet.UniqueSampleCount; ++UniqueIt)
 		{
 		#if PERMUTATION_MATERIAL_COMPACTION == 1
 			const bool bIsSimilar  = 
-				VisibilityData.PrimitiveId == OutSamples[UniqueIt].x ||
-				(abs(dot(UniqueTangentAndDepths[UniqueIt].xyz, VisibilityData.Tangent)) > CosTangentThreshold &&
+				PrimitiveId == OutSamples[UniqueIt].x ||
+				(abs(dot(UniqueTangentAndDepths[UniqueIt].xyz, Tangent)) > CosTangentThreshold &&
 					abs(UniqueTangentAndDepths[UniqueIt].w - SceneDepth) < DepthTheshold);
 		#else
-			const bool bIsSimilar = VisibilityData.PrimitiveId == OutSamples[UniqueIt].x;
+			const bool bIsSimilar = PrimitiveId == OutSamples[UniqueIt].x;
 		#endif
 
 			if (bIsSimilar)
@@ -510,10 +499,10 @@ void ComputeUniqueSamplesWithoutDepthTest(
 		if (!bAlreadyExist)
 		{
 		#if PERMUTATION_MATERIAL_COMPACTION == 1
-			UniqueTangentAndDepths[OutSet.UniqueSampleCount] = float4(VisibilityData.Tangent, SceneDepth);
+			UniqueTangentAndDepths[OutSet.UniqueSampleCount] = float4(Tangent, SceneDepth);
 		#endif
 
-			OutSamples[OutSet.UniqueSampleCount].x = VisibilityData.PrimitiveId;
+			OutSamples[OutSet.UniqueSampleCount].x = PrimitiveId;
 			OutSamples[OutSet.UniqueSampleCount].y = 1;
 			OutSamples[OutSet.UniqueSampleCount].z = SampleIt;
 			++OutSet.UniqueSampleCount;
@@ -521,17 +510,13 @@ void ComputeUniqueSamplesWithoutDepthTest(
 	}
 }
 
-#ifndef HAIR_VISIBILITY_MAXSAMPLECOUNT
-#error HAIR_VISIBILITY_MAXSAMPLECOUNT is not defined
-#endif
-
-Texture2DMS<uint4, HAIR_VISIBILITY_MAXSAMPLECOUNT> MSAA_IDTexture;
-Texture2DMS<float, HAIR_VISIBILITY_MAXSAMPLECOUNT> MSAA_DepthTexture;
-Texture2DMS<float4,HAIR_VISIBILITY_MAXSAMPLECOUNT> MSAA_MaterialTexture;
-Texture2DMS<float4,HAIR_VISIBILITY_MAXSAMPLECOUNT> MSAA_AttributeTexture;
+Texture2DMS<uint2, PERMUTATION_MSAACOUNT> MSAA_IDTexture;
+Texture2DMS<float, PERMUTATION_MSAACOUNT> MSAA_DepthTexture;
+Texture2DMS<float4,PERMUTATION_MSAACOUNT> MSAA_MaterialTexture;
+Texture2DMS<float4,PERMUTATION_MSAACOUNT> MSAA_AttributeTexture;
 
 #if PERMUTATION_VELOCITY != VELOCITY_TYPE_NONE
-Texture2DMS<float2, HAIR_VISIBILITY_MAXSAMPLECOUNT> MSAA_VelocityTexture;
+Texture2DMS<float2, PERMUTATION_MSAACOUNT> MSAA_VelocityTexture;
 #endif
 
 #if PERMUTATION_VIEWTRANSMITTANCE == 1
@@ -548,18 +533,17 @@ void MainCS(uint GroupIndex : SV_GroupIndex, uint3 DispatchThreadId : SV_Dispatc
 	}
 
 	uint2 PixelCoord = DispatchThreadId.xy + ResolutionOffset;
-	if (HairVisibilityMSAASampleCount == 0 || PixelCoord.x >= OutputResolution.x || PixelCoord.y >= OutputResolution.y)
+	if (PixelCoord.x >= OutputResolution.x || PixelCoord.y >= OutputResolution.y)
 		PixelCoord = uint2(0, 0);
 
 	FSampleSetDesc SampleDesc;
 
 	const float OpaqueDepth = SceneTexturesStruct.SceneDepthTexture.Load(uint3(PixelCoord, 0)).r;
-	uint3 Samples[HAIR_VISIBILITY_MAXSAMPLECOUNT]; // x:PrimitiveId, y:Weight, z:SampleIndex
+	uint3 Samples[PERMUTATION_MSAACOUNT]; // x:PrimitiveId, y:Weight, z:SampleIndex
 	ComputeUniqueSamplesWithoutDepthTest(
 		Samples,
 		SampleDesc,
 		PixelCoord,
-		HairVisibilityMSAASampleCount,
 		OpaqueDepth,
 		MSAA_IDTexture, 
 		MSAA_DepthTexture);
@@ -605,32 +589,21 @@ void MainCS(uint GroupIndex : SV_GroupIndex, uint3 DispatchThreadId : SV_Dispatc
 				const uint SampleIt = Samples[OutIndex].z;
 
 				// Read all material data
-				// #hair_todo: simplify the encoding. technically we don't need to decompress the visibility data, 
-				// as we only need to patch the weight
-				// #hair_todo: merge similar sample
-				FNodeData OutNodeData;
+				FPackedHairSample OutNodeData;
 				{
-					const float SampleDepth = MSAA_DepthTexture.Load(PixelCoord, SampleIt);
-					const uint4 EncodedData = MSAA_IDTexture.Load(PixelCoord, SampleIt);
-					FVisibilityData VisibilityData = DecodeVisibilityData(EncodedData);
+					FHairSample HairSample;
+					HairSample.Depth = MSAA_DepthTexture.Load(PixelCoord, SampleIt);
+
+					UnpackHairMRT0(MSAA_IDTexture.Load(PixelCoord, SampleIt), HairSample.PrimitiveId, HairSample.MacroGroupId, HairSample.Coverage8bit, HairSample.Tangent);
+					UnpackHairMRT1(MSAA_MaterialTexture.Load(PixelCoord, SampleIt), HairSample.BaseColor, HairSample.Roughness);
+					UnpackHairMRT2(MSAA_AttributeTexture.Load(PixelCoord, SampleIt), HairSample.Specular, HairSample.LightChannelMask, HairSample.Backlit, HairSample.bScatterSceneLighting);
+
 					// VisibilityData.Coverage8bit is a weight normalising to 1 the contribution of all the compacted samples. Because later it is weighted by Categorization.PixelCoverage.
-					VisibilityData.Coverage8bit = To8bitCoverage(Samples[OutIndex].y / float(SampleDesc.ValidSampleCount));
-					Categorization.TotalCoverage8bit += VisibilityData.Coverage8bit;
-					Categorization.ClosestDepth = max(Categorization.ClosestDepth, SampleDepth); // Inverse-Z
-
-					float4 BaseColorAndRoughness = MSAA_MaterialTexture.Load(PixelCoord, SampleIt);
-					BaseColorAndRoughness.xyz = BaseColorAndRoughness.xyz * BaseColorAndRoughness.xyz; // Cheap gamma encoding since this is encoded into 8bit target
-					const float2 SpecularAndLigthChannelMask = MSAA_AttributeTexture.Load(PixelCoord, SampleIt).xy;
-					const float Specular = SpecularAndLigthChannelMask.x;
-					const uint LightChannelMask = SpecularAndLigthChannelMask.y * 0xFF;
-
-					OutNodeData = ConvertToNodeData(
-						SampleDepth,
-						VisibilityData,
-						BaseColorAndRoughness.xyz,
-						BaseColorAndRoughness.w,
-						Specular,
-						LightChannelMask);
+					HairSample.Coverage8bit = To8bitCoverage(Samples[OutIndex].y / float(SampleDesc.ValidSampleCount));
+					Categorization.TotalCoverage8bit += HairSample.Coverage8bit;
+					Categorization.ClosestDepth = max(Categorization.ClosestDepth, HairSample.Depth); // Inverse-Z
+
+					OutNodeData = PackHairSample(HairSample);
 
 					#if PERMUTATION_VELOCITY == VELOCITY_TYPE_AVG
 					const float2 EncodedVelocity = MSAA_VelocityTexture.Load(PixelCoord, SampleIt);
@@ -640,10 +613,10 @@ void MainCS(uint GroupIndex : SV_GroupIndex, uint3 DispatchThreadId : SV_Dispatc
 
 					#if PERMUTATION_VELOCITY == VELOCITY_TYPE_CLOSEST
 					const float2 EncodedVelocity = MSAA_VelocityTexture.Load(PixelCoord, SampleIt);
-					if (ClosestDepth < SampleDepth) // Inverse-Z
+					if (ClosestDepth < HairSample.Depth) // Inverse-Z
 					{
 						ClosestEncodedVelocity = EncodedVelocity;
-						ClosestDepth = SampleDepth;
+						ClosestDepth = HairSample.Depth;
 					}
 
 					#endif
@@ -668,7 +641,7 @@ void MainCS(uint GroupIndex : SV_GroupIndex, uint3 DispatchThreadId : SV_Dispatc
 			#if PERMUTATION_VIEWTRANSMITTANCE == 1
 			Categorization.PixelCoverage = 1.0f - ViewTransmittanceTexture.Load(uint3(PixelCoord,0));
 			#else
-			Categorization.PixelCoverage = SampleDesc.ValidSampleCount / float(HAIR_VISIBILITY_MAXSAMPLECOUNT);
+			Categorization.PixelCoverage = SampleDesc.ValidSampleCount / float(PERMUTATION_MSAACOUNT);
 			#endif
 
 			#if PERMUTATION_VELOCITY == VELOCITY_TYPE_AVG
diff --git a/Engine/Shaders/Private/HairStrands/HairStrandsVisibilityComposeSubPixelPS.usf b/Engine/Shaders/Private/HairStrands/HairStrandsVisibilityComposeSubPixelPS.usf
--- a/Engine/Shaders/Private/HairStrands/HairStrandsVisibilityComposeSubPixelPS.usf
+++ b/Engine/Shaders/Private/HairStrands/HairStrandsVisibilityComposeSubPixelPS.usf
@@ -7,6 +7,8 @@
 #include "../DeferredShadingCommon.ush"
 #include "HairStrandsVisibilityCommon.ush"
 
+#if SHADER_SUBCOLOR
+
 Texture2D<float4> SubPixelColorTexture;
 Texture2D<uint4> CategorisationTexture;
 
@@ -27,15 +29,23 @@ void SubColorPS(
 	OutDepth = CategorisationData.ClosestDepth;
 }
 
+#endif // SHADER_SUBCOLOR
+
+#if SHADER_FASTRESOLVE
+
+#ifndef PERMUTATION_MSAACOUNT
+#error PERMUTATION_MSAACOUNT is not defined
+#endif
+
 float VelocityThreshold;
-Texture2DMS<float2, HAIR_VISIBILITY_MAXSAMPLECOUNT> HairVisibilityVelocityTexture;
+Texture2DMS<float2, PERMUTATION_MSAACOUNT> HairVisibilityVelocityTexture;
 
 void FastResolvePS(in FScreenVertexOutput Input)
 {
 	const uint2 PixelCoord = floor(Input.Position.xy);
 
 	bool bNeedFastResolve = false;
-	for (uint SampleIt = 0; SampleIt < HAIR_VISIBILITY_MAXSAMPLECOUNT; ++SampleIt)
+	for (uint SampleIt = 0; SampleIt < PERMUTATION_MSAACOUNT; ++SampleIt)
 	{
 		const float2 EncodedVelocity = HairVisibilityVelocityTexture.Load(PixelCoord, SampleIt);
 		const float2 Velocity = DecodeVelocityFromTexture(EncodedVelocity);
@@ -49,4 +59,6 @@ void FastResolvePS(in FScreenVertexOutput Input)
 
 	if (!bNeedFastResolve)
 		discard;
-}
\ No newline at end of file
+}
+
+#endif // SHADER_FASTRESOLVE
\ No newline at end of file
diff --git a/Engine/Shaders/Private/HairStrands/HairStrandsVisibilityPS.usf b/Engine/Shaders/Private/HairStrands/HairStrandsVisibilityPS.usf
--- a/Engine/Shaders/Private/HairStrands/HairStrandsVisibilityPS.usf
+++ b/Engine/Shaders/Private/HairStrands/HairStrandsVisibilityPS.usf
@@ -1,5 +1,6 @@
 // Copyright Epic Games, Inc. All Rights Reserved.
 
+#include "HairStrandsVisibilityCommonStruct.ush"
 #include "HairStrandsVoxelPageCommonStruct.ush" 
 
 #include "../Common.ush"
@@ -33,13 +34,6 @@
 #define SUPPORT_MATERIAL_PROPERTY 0
 #endif
 
-#if HAIR_RENDER_MODE == RENDER_MODE_PPLL 
-uint HairVisibilityPass_MaxPPLLNodeCount;
-RWTexture2D<uint> PPLLCounter;
-RWTexture2D<uint> PPLLNodeIndex;
-RWStructuredBuffer<FPPLLNodeData> PPLLNodes;
-#endif
-
 uint HairVisibilityPass_HairMacroGroupIndex;
 uint HairVisibilityPass_HairMaterialId;
 uint HairVisibilityPass_LightChannelMask;
@@ -68,7 +62,7 @@ void MainVisibility(
 	, out float OutColor0 : SV_Target0
 	, out float OutColor1 : SV_Target1
 #elif HAIR_RENDER_MODE == RENDER_MODE_MSAA
-	, out uint4 OutColor0	: SV_Target0
+	, out uint2 OutColor0	: SV_Target0
   #if SUPPORT_MATERIAL_PROPERTY
 	, out float4 OutColor1	: SV_Target1
 	, out float4 OutColor2	: SV_Target2
@@ -93,7 +87,12 @@ void MainVisibility(
 #if HAIR_RENDER_MODE != RENDER_MODE_MSAA_VISIBILITY
 	float Coverage = 1;
 	{
-		FHairRenderInfo HairRenderInfo = GetHairRenderInfo(ResolvedView.HairRenderInfo);
+		bool bUseStableRasterization = false;
+	#if VF_STRAND_HAIR
+		bUseStableRasterization = UseStableRasterization();
+	#endif
+
+		FHairRenderInfo HairRenderInfo = GetHairRenderInfo(ResolvedView.HairRenderInfo, ResolvedView.HairRenderInfoBits, bUseStableRasterization);
 		const float SceneDepth = ConvertFromDeviceZ(SvPosition.z); // Linear depth in world unit
 		const float PixelRadius = HairRenderInfo.bIsOrthoView ? HairRenderInfo.RadiusAtDepth1Primary : SceneDepth * HairRenderInfo.RadiusAtDepth1Primary; // Not correct but the coverage is not used (we count instead the number of sub-sample covered)
 		const float StrandRealRadius = WorldStrandRadius;
@@ -109,23 +108,25 @@ void MainVisibility(
 #endif
 
 #if HAIR_RENDER_MODE == RENDER_MODE_MSAA || HAIR_RENDER_MODE == RENDER_MODE_PPLL
-	FVisibilityData VisibilityData;
   #if VF_STRAND_HAIR
-	VisibilityData.Tangent = normalize(Interpolants.TangentToWorld2.xyz);
-	VisibilityData.PrimitiveId = Interpolants.HairPrimitiveId;
+	float3 Tangent = normalize(Interpolants.TangentToWorld2.xyz);
+	const uint PrimitiveId = Interpolants.HairPrimitiveId;
+	const bool bUseScatterSceneLighting = UseScatterSceneLighting();
   #else
-	VisibilityData.Tangent = float3(0, 0, 1);
-	VisibilityData.PrimitiveId = 0;
+	float3 Tangent = float3(0, 0, 1);
+	const uint PrimitiveId = 0;
+	const bool bUseScatterSceneLighting = false;
   #endif
-	VisibilityData.MacroGroupId = HairVisibilityPass_HairMacroGroupIndex;
-	VisibilityData.Coverage8bit = To8bitCoverage(Coverage);
+	const uint MacroGroupId = HairVisibilityPass_HairMacroGroupIndex;
+	const uint Coverage8bit = To8bitCoverage(Coverage);
 
 	half3 BaseColor = 0;
 	float Roughness = 0;
 	float Specular  = 0;
+	float Backlit	= 1;
 	float2 Velocity = 0;
 	float2 EncodedVelocity = 0;
-#if SUPPORT_MATERIAL_PROPERTY
+	#if SUPPORT_MATERIAL_PROPERTY
 	{
 		FMaterialPixelParameters MaterialParameters = GetMaterialPixelParameters(Interpolants, SvPosition);
 		FPixelMaterialInputs PixelMaterialInputs;
@@ -138,21 +139,22 @@ void MainVisibility(
 		Roughness		 = GetMaterialRoughness(PixelMaterialInputs);
 		half  Metallic	 = GetMaterialMetallic(PixelMaterialInputs);
 		Specular		 = GetMaterialSpecular(PixelMaterialInputs);
+		Backlit			 = GetMaterialCustomData0(MaterialParameters);
 		float MaterialAO = GetMaterialAmbientOcclusion(PixelMaterialInputs);
 
-		VisibilityData.Tangent = mul(normalize(GetMaterialNormalRaw(PixelMaterialInputs)), MaterialParameters.TangentToWorld);
+		Tangent = mul(normalize(GetMaterialNormalRaw(PixelMaterialInputs)), MaterialParameters.TangentToWorld);
 	}
 	{
 		Velocity = Calculate2DVelocity(PackedVelocityA, PackedVelocityC);
 		EncodedVelocity = EncodeVelocityToTexture(Velocity);
 	}
-#endif
+	#endif
 #endif
 
 #if HAIR_RENDER_MODE == RENDER_MODE_MSAA
-	OutColor0 = EncodeVisibilityData(VisibilityData);
-	OutColor1 = float4(sqrt(BaseColor), Roughness);
-	OutColor2 = float4(Specular, float(HairVisibilityPass_LightChannelMask) / 255.f,0,0);
+	OutColor0 = PackHairMRT0(PrimitiveId, MacroGroupId, Coverage8bit, Tangent);
+	OutColor1 = PackHairMRT1(BaseColor, Roughness);
+	OutColor2 = PackHairMRT2(Specular, HairVisibilityPass_LightChannelMask, Backlit, bUseScatterSceneLighting);
 	OutColor3 = EncodedVelocity;
 #endif
 
@@ -164,21 +166,34 @@ void MainVisibility(
 
 	// Allocate space for a new node
 	uint NodeIndex;
-	InterlockedAdd(PPLLCounter[uint2(0, 0)], 1, NodeIndex);
+	InterlockedAdd(HairVisibilityPass.PPLLCounter[uint2(0, 0)], 1, NodeIndex);
 
-	if (NodeIndex < HairVisibilityPass_MaxPPLLNodeCount)
+	if (NodeIndex < HairVisibilityPass.MaxPPLLNodeCount)
 	{
 		// If we can append new node, add it and make it point to the node previously in head of list
 		uint NextNodeIndex;
-		InterlockedExchange(PPLLNodeIndex[PixelCoord], NodeIndex, NextNodeIndex);
-		const FPPLLNodeData Node = ConvertToPPLLNodeData(SvPosition.z, Coverage, VisibilityData, BaseColor, Roughness, Specular, HairVisibilityPass_LightChannelMask, Velocity, NextNodeIndex);
-		PPLLNodes[NodeIndex] = Node;
+		InterlockedExchange(HairVisibilityPass.PPLLNodeIndex[PixelCoord], NodeIndex, NextNodeIndex);
+		const FPPLLNodeData Node = CreatePPLLNodeData(
+			SvPosition.z, 
+			PrimitiveId, 
+			MacroGroupId, 
+			Tangent, 
+			Coverage, 
+			Coverage8bit, 
+			BaseColor, 
+			Roughness, 
+			Specular, 
+			Backlit, 
+			HairVisibilityPass_LightChannelMask, 
+			Velocity, 
+			NextNodeIndex);
+		HairVisibilityPass.PPLLNodeData[NodeIndex] = Node;
 	}
 #endif
 
 #if HAIR_RENDER_MODE == RENDER_MODE_MSAA_VISIBILITY
   #if VF_STRAND_HAIR
-	OutColor0 = EncodeVisSample(HairPrimitiveId, HairVisibilityPass_HairMaterialId);
+	OutColor0 = PackHairVisSample(HairPrimitiveId, HairVisibilityPass_HairMaterialId);
   #else
 	OutColor0 = 0;
   #endif
