1) Remove float4 channel compression from batch shadow denoiser
2) Add transmission distance to batch shadow denoiser

diff --git a/Engine/Shaders/Private/ScreenSpaceDenoise/SSDInjest.usf b/Engine/Shaders/Private/ScreenSpaceDenoise/SSDInjest.usf
--- a/Engine/Shaders/Private/ScreenSpaceDenoise/SSDInjest.usf
+++ b/Engine/Shaders/Private/ScreenSpaceDenoise/SSDInjest.usf
@@ -19,7 +19,7 @@
 	#define CONFIG_MULTIPLEXED_SIGNALS_PER_SIGNAL_DOMAIN 1
 
 	#define CONFIG_INPUT_TEXTURE_COUNT DIM_SIGNAL_BATCH_SIZE
-	#define CONFIG_OUTPUT_TEXTURE_COUNT ((DIM_SIGNAL_BATCH_SIZE + 1) / 2)
+	#define CONFIG_OUTPUT_TEXTURE_COUNT DIM_SIGNAL_BATCH_SIZE
 
 #else
 	#error Unknown signal processing.
diff --git a/Engine/Shaders/Private/ScreenSpaceDenoise/SSDSignalBufferEncoding.ush b/Engine/Shaders/Private/ScreenSpaceDenoise/SSDSignalBufferEncoding.ush
--- a/Engine/Shaders/Private/ScreenSpaceDenoise/SSDSignalBufferEncoding.ush
+++ b/Engine/Shaders/Private/ScreenSpaceDenoise/SSDSignalBufferEncoding.ush
@@ -107,6 +107,7 @@ void DecodeMultiplexedSignalsFromFloat4(
 			OutSamples.Array[MultiplexId].SampleCount = (RawSample[MultiplexId].g == DENOISER_INVALID_HIT_DISTANCE ? 0.0 : 1.0);
 			OutSamples.Array[MultiplexId].MissCount = OutSamples.Array[MultiplexId].SampleCount * RawSample[MultiplexId].r;
 			OutSamples.Array[MultiplexId].ClosestHitDistance = OutSamples.Array[MultiplexId].SampleCount * (RawSample[MultiplexId].g == DENOISER_MISS_HIT_DISTANCE ? 1000000 : RawSample[MultiplexId].g);
+			OutSamples.Array[MultiplexId].TransmissionDistance = OutSamples.Array[MultiplexId].SampleCount * RawSample[MultiplexId].a;
 		}
 	}
 	else if (SignalBufferLayout == SIGNAL_BUFFER_LAYOUT_PENUMBRA_INJESTION_NSPP)
@@ -115,12 +116,13 @@ void DecodeMultiplexedSignalsFromFloat4(
 		for (uint BatchSignalId = 0; BatchSignalId < MAX_SIGNAL_BATCH_SIZE; BatchSignalId++)
 		{
 			uint MultiplexId = BatchSignalId;
-			float2 Channels = float2(RawSample[MultiplexId / 2][2 * (MultiplexId % 2)], RawSample[MultiplexId / 2][2 * (MultiplexId % 2) + 1]);
-	
+			float4 Channels = RawSample[MultiplexId];
+
 			// TODO(Denoiser): feed the actual number of sample.
 			OutSamples.Array[MultiplexId].SampleCount = (Channels.g == -2.0 ? 0.0 : 1.0);
 			OutSamples.Array[MultiplexId].MissCount = OutSamples.Array[MultiplexId].SampleCount * Channels.r;
 			OutSamples.Array[MultiplexId].WorldBluringRadius = OutSamples.Array[MultiplexId].SampleCount * (Channels.g == -1 ? WORLD_RADIUS_MISS : Channels.g);
+			OutSamples.Array[MultiplexId].TransmissionDistance = OutSamples.Array[MultiplexId].SampleCount * Channels.a;
 		}
 	}
 	//else if (SignalBufferLayout == SIGNAL_BUFFER_LAYOUT_PENUMBRA_RECONSTRUCTION)
@@ -147,13 +149,14 @@ void DecodeMultiplexedSignalsFromFloat4(
 		UNROLL_N(MAX_SIGNAL_BATCH_SIZE)
 		for (uint MultiplexId = 0; MultiplexId < MAX_SIGNAL_BATCH_SIZE; MultiplexId++)
 		{
-			float3 Channels = RawSample[MultiplexId].xyz;
+			float4 Channels = RawSample[MultiplexId].xyzw;
 			
 			float SampleCount = bNormalizeSample ? (Channels.g > 0 ? 1 : 0) : (Channels.g);
 
 			OutSamples.Array[MultiplexId].MissCount = Channels.r * SampleCount;
 			OutSamples.Array[MultiplexId].SampleCount = SampleCount;
 			OutSamples.Array[MultiplexId].WorldBluringRadius = (Channels.b == DENOISER_MISS_HIT_DISTANCE ? WORLD_RADIUS_MISS : Channels.b) * SampleCount;
+			OutSamples.Array[MultiplexId].TransmissionDistance = Channels.a * SampleCount;
 		}
 	}
 	else if (SignalBufferLayout == SIGNAL_BUFFER_LAYOUT_PENUMBRA_REJECTION)
@@ -161,16 +164,14 @@ void DecodeMultiplexedSignalsFromFloat4(
 		UNROLL_N(MAX_SIGNAL_BATCH_SIZE)
 		for (uint BatchSignalId = 0; BatchSignalId < MAX_SIGNAL_BATCH_SIZE; BatchSignalId++)
 		{
-			uint MultiplexId = BatchSignalId * 2;
+			uint MultiplexId = BatchSignalId;
 			float4 Channels = RawSample[BatchSignalId];
-			float SampleCount = Channels.b;
+			float SampleCount = Channels.g;
 
-			OutSamples.Array[MultiplexId + 0].SampleCount = SampleCount;
-			OutSamples.Array[MultiplexId + 0].MissCount = Channels.r * SampleCount;
-			OutSamples.Array[MultiplexId + 0].WorldBluringRadius = (Channels.a == -1 ? WORLD_RADIUS_MISS : Channels.a) * SampleCount;
-			
-			OutSamples.Array[MultiplexId + 1].SampleCount = SampleCount;
-			OutSamples.Array[MultiplexId + 1].MissCount = Channels.g * SampleCount;
+			OutSamples.Array[MultiplexId].SampleCount = SampleCount;
+			OutSamples.Array[MultiplexId].MissCount = Channels.r * SampleCount;
+			OutSamples.Array[MultiplexId].WorldBluringRadius = (Channels.b == -1 ? WORLD_RADIUS_MISS : Channels.b * SampleCount);
+			OutSamples.Array[MultiplexId].TransmissionDistance = Channels.a * SampleCount;
 		}
 	}
 	#endif
@@ -506,9 +507,12 @@ void EncodeMultiplexedSignals(
 		UNROLL
 		for (uint MultiplexId = 0; MultiplexId < MultiplexCount; MultiplexId++)
 		{
-			float2 EncodedChannels = float2(
+			float4 EncodedChannels = float4(
 				Sample[MultiplexId].MissCount,
-				Sample[MultiplexId].WorldBluringRadius) * SafeRcp(Sample[MultiplexId].SampleCount);
+				Sample[MultiplexId].WorldBluringRadius,
+				0.0,
+				Sample[MultiplexId].TransmissionDistance);
+			EncodedChannels *= SafeRcp(Sample[MultiplexId].SampleCount);
 			
 			if (Sample[MultiplexId].SampleCount == 0)
 			{
@@ -519,12 +523,9 @@ void EncodeMultiplexedSignals(
 				EncodedChannels.y = -1;
 			}
 
-			if (MultiplexId % 2)
-				OutRawSample[MultiplexId / 2].zw = EncodedChannels;
-			else
-				OutRawSample[MultiplexId / 2].xy = EncodedChannels;
+			OutRawSample[MultiplexId] = EncodedChannels;
 		}
-		OutBufferCount = (MultiplexCount + 1) / 2;
+		OutBufferCount = MultiplexCount;
 	}
 	//else if (SignalBufferLayout == SIGNAL_BUFFER_LAYOUT_PENUMBRA_RECONSTRUCTION)
 	//{
@@ -550,24 +551,25 @@ void EncodeMultiplexedSignals(
 			float NormalizationFactor = SafeRcp(Sample[MultiplexId].SampleCount);
 
 			float NormalizedWorldBluringRadius = Sample[MultiplexId].WorldBluringRadius * NormalizationFactor;
+			float NormalizedTransmissionDistance = Sample[MultiplexId].TransmissionDistance * NormalizationFactor;
 
 			OutRawSample[MultiplexId] = float4(
 				Sample[MultiplexId].MissCount * NormalizationFactor,
 				Sample[MultiplexId].SampleCount,
 				NormalizedWorldBluringRadius == WORLD_RADIUS_MISS ? -1 : NormalizedWorldBluringRadius,
-				0);
+				NormalizedTransmissionDistance);
 		}
 		OutBufferCount = MultiplexCount;
 	}
 	else if (SignalBufferLayout == SIGNAL_BUFFER_LAYOUT_PENUMBRA_REJECTION)
 	{
-		const uint BatchSignalCount = MultiplexCount / 2;
+		const uint BatchSignalCount = MultiplexCount;
 		OutBufferCount = BatchSignalCount;
 
 		UNROLL
 		for (uint BatchSignalId = 0; BatchSignalId < BatchSignalCount; BatchSignalId++)
 		{
-			uint MultiplexId = BatchSignalId * 2;
+			uint MultiplexId = BatchSignalId;
 
 			float NormalizationFactor = SafeRcp(Sample[MultiplexId + 0].SampleCount);
 
@@ -576,10 +578,10 @@ void EncodeMultiplexedSignals(
 			// Samples are normalised when doing history preconvolution, so the number of is either 0 or 1 on both momments. Therefore
 			// Sample[MultiplexId + 0].SampleCount == Sample[MultiplexId + 1].SampleCount.
 			OutRawSample[BatchSignalId] = float4(
-				Sample[MultiplexId + 0].MissCount * NormalizationFactor,
-				Sample[MultiplexId + 1].MissCount * NormalizationFactor,
-				Sample[MultiplexId + 0].SampleCount,
-				NormalizedWorldBluringRadius == WORLD_RADIUS_MISS ? -1 : NormalizedWorldBluringRadius);
+				Sample[MultiplexId].MissCount * NormalizationFactor,
+				Sample[MultiplexId].SampleCount,
+				NormalizedWorldBluringRadius == WORLD_RADIUS_MISS ? -1 : NormalizedWorldBluringRadius,
+				Sample[MultiplexId].TransmissionDistance * NormalizationFactor);
 		}
 	}
 	#endif
diff --git a/Engine/Shaders/Private/ScreenSpaceDenoise/SSDSignalCore.ush b/Engine/Shaders/Private/ScreenSpaceDenoise/SSDSignalCore.ush
--- a/Engine/Shaders/Private/ScreenSpaceDenoise/SSDSignalCore.ush
+++ b/Engine/Shaders/Private/ScreenSpaceDenoise/SSDSignalCore.ush
@@ -59,6 +59,8 @@ struct FSSDSignalSample
 	
 	// Blur radius allowed for this sample, valid only if SampleCount == 1.
 	float WorldBluringRadius;
+
+	float TransmissionDistance;
 };
 
 /** Array of signal samples. Technically should be a typedef, but there is a bug in HLSL with the out operator of an array. */
@@ -100,6 +102,7 @@ FSSDSignalSample CreateSignalSampleFromScalarValue(float Scalar)
 	OutSample.MissCount = Scalar;
 	OutSample.ClosestHitDistance = Scalar;
 	OutSample.WorldBluringRadius = Scalar;
+	OutSample.TransmissionDistance = Scalar;
 	return OutSample;
 }
 
@@ -142,6 +145,7 @@ FSSDSignalSample MulSignal(FSSDSignalSample Sample, float Scalar)
 	OutSample.MissCount              = Sample.MissCount * Scalar;
 	OutSample.ClosestHitDistance     = Sample.ClosestHitDistance * Scalar;
 	OutSample.WorldBluringRadius     = Sample.WorldBluringRadius * Scalar;
+	OutSample.TransmissionDistance   = Sample.TransmissionDistance * Scalar;
 	return OutSample;
 }
 
@@ -164,6 +168,7 @@ FSSDSignalSample AddSignal(FSSDSignalSample SampleA, FSSDSignalSample SampleB)
 	SampleA.MissCount              += SampleB.MissCount;
 	SampleA.ClosestHitDistance     += SampleB.ClosestHitDistance;
 	SampleA.WorldBluringRadius     += SampleB.WorldBluringRadius;
+	SampleA.TransmissionDistance   += SampleB.TransmissionDistance;
 	return SampleA;
 }
 
@@ -186,6 +191,7 @@ FSSDSignalSample MinusSignal(FSSDSignalSample Sample)
 	Sample.MissCount              = -Sample.MissCount;
 	Sample.ClosestHitDistance     = -Sample.ClosestHitDistance;
 	Sample.WorldBluringRadius     = -Sample.WorldBluringRadius;
+	Sample.TransmissionDistance   = -Sample.TransmissionDistance;
 	return Sample;
 }
 
@@ -225,6 +231,7 @@ FSSDSignalSample AbsSignal(FSSDSignalSample Sample)
 	Sample.MissCount              = abs(Sample.MissCount);
 	Sample.ClosestHitDistance     = abs(Sample.ClosestHitDistance);
 	Sample.WorldBluringRadius     = abs(Sample.WorldBluringRadius);
+	Sample.TransmissionDistance = abs(Sample.TransmissionDistance);
 	return Sample;
 }
 
@@ -259,6 +266,7 @@ FSSDSignalSample SqrtSignal(FSSDSignalSample Sample)
 	Sample.MissCount              = sqrt(Sample.MissCount);
 	Sample.ClosestHitDistance     = sqrt(Sample.ClosestHitDistance);
 	Sample.WorldBluringRadius     = sqrt(Sample.WorldBluringRadius);
+	Sample.TransmissionDistance   = sqrt(Sample.TransmissionDistance);
 	return Sample;
 }
 
@@ -296,6 +304,7 @@ FSSDSignalSample PowerSignal(FSSDSignalSample Sample, float Exponent)
 	Sample.MissCount              = pow(Sample.MissCount, Exponent);
 	Sample.ClosestHitDistance     = pow(Sample.ClosestHitDistance, Exponent);
 	Sample.WorldBluringRadius     = pow(Sample.WorldBluringRadius, Exponent);
+	Sample.TransmissionDistance   = pow(Sample.TransmissionDistance, Exponent);
 	return Sample;
 }
 
@@ -331,6 +340,7 @@ FSSDSignalSample MinSignal(FSSDSignalSample SampleA, FSSDSignalSample SampleB)
 	OutSample.MissCount = min(SampleA.MissCount, SampleB.MissCount);
 	OutSample.ClosestHitDistance = min(SampleA.ClosestHitDistance, SampleB.ClosestHitDistance);
 	OutSample.WorldBluringRadius = min(SampleA.WorldBluringRadius, SampleB.WorldBluringRadius);
+	OutSample.TransmissionDistance = min(SampleA.TransmissionDistance, SampleB.TransmissionDistance);
 	return OutSample;
 }
 
@@ -366,6 +376,7 @@ FSSDSignalSample MaxSignal(FSSDSignalSample SampleA, FSSDSignalSample SampleB)
 	OutSample.MissCount = max(SampleA.MissCount, SampleB.MissCount);
 	OutSample.ClosestHitDistance = max(SampleA.ClosestHitDistance, SampleB.ClosestHitDistance);
 	OutSample.WorldBluringRadius = max(SampleA.WorldBluringRadius, SampleB.WorldBluringRadius);
+	OutSample.TransmissionDistance = max(SampleA.TransmissionDistance, SampleB.TransmissionDistance);
 	return OutSample;
 }
 
@@ -401,6 +412,7 @@ FSSDSignalSample ClampSignal(FSSDSignalSample Sample, FSSDSignalSample SampleMin
 	OutSample.MissCount = clamp(Sample.MissCount, SampleMin.MissCount, SampleMax.MissCount);
 	OutSample.ClosestHitDistance = clamp(Sample.ClosestHitDistance, SampleMin.ClosestHitDistance, SampleMax.ClosestHitDistance);
 	OutSample.WorldBluringRadius = clamp(Sample.WorldBluringRadius, SampleMin.WorldBluringRadius, SampleMax.WorldBluringRadius);
+	OutSample.TransmissionDistance = clamp(Sample.TransmissionDistance, SampleMin.TransmissionDistance, SampleMax.TransmissionDistance);
 	return OutSample;
 }
 
@@ -449,5 +461,6 @@ FSSDSignalSample WaveBroadcastSignal(const FWaveBroadcastSettings BroadcastSetti
 	OutSample.MissCount          = WaveBroadcast(BroadcastSettings, Sample.MissCount);
 	OutSample.ClosestHitDistance = WaveBroadcast(BroadcastSettings, Sample.ClosestHitDistance);
 	OutSample.WorldBluringRadius = WaveBroadcast(BroadcastSettings, Sample.WorldBluringRadius);
+	OutSample.TransmissionDistance = WaveBroadcast(BroadcastSettings, Sample.TransmissionDistance);
 	return OutSample;
 }
diff --git a/Engine/Shaders/Private/ScreenSpaceDenoise/SSDSpatialAccumulation.usf b/Engine/Shaders/Private/ScreenSpaceDenoise/SSDSpatialAccumulation.usf
--- a/Engine/Shaders/Private/ScreenSpaceDenoise/SSDSpatialAccumulation.usf
+++ b/Engine/Shaders/Private/ScreenSpaceDenoise/SSDSpatialAccumulation.usf
@@ -51,7 +51,7 @@
 		#define CONFIG_SIGNAL_INPUT_LAYOUT  SIGNAL_BUFFER_LAYOUT_PENUMBRA_INJESTION_NSPP
 		#define CONFIG_SIGNAL_OUTPUT_LAYOUT SIGNAL_BUFFER_LAYOUT_PENUMBRA_RECONSTRUCTION
 
-		#define CONFIG_INPUT_TEXTURE_COUNT ((CONFIG_SIGNAL_BATCH_SIZE + 1) / 2)
+		#define CONFIG_INPUT_TEXTURE_COUNT CONFIG_SIGNAL_BATCH_SIZE
 		#define CONFIG_OUTPUT_TEXTURE_COUNT CONFIG_SIGNAL_BATCH_SIZE
 
 		#define CONFIG_SAMPLE_SET           SAMPLE_SET_STACKOWIAK_4_SETS
@@ -1226,7 +1226,7 @@ void MainCS(
 				float Shadow = GetSamplePenumbraSafe(OutputSamples.Array[MultiplexId]);
 
 				const float ShadowFadeFraction = 1;
-				float SSSTransmission = Shadow;
+				float SSSTransmission = (OutputSamples.Array[MultiplexId].SampleCount > 0 ? OutputSamples.Array[MultiplexId].TransmissionDistance / OutputSamples.Array[MultiplexId].SampleCount : OutputSamples.Array[MultiplexId].TransmissionDistance);
 
 				// 0 is shadowed, 1 is unshadowed
 				// RETURN_COLOR not needed unless writing to SceneColor;
diff --git a/Engine/Source/Runtime/Renderer/Private/ScreenSpaceDenoise.cpp b/Engine/Source/Runtime/Renderer/Private/ScreenSpaceDenoise.cpp
--- a/Engine/Source/Runtime/Renderer/Private/ScreenSpaceDenoise.cpp
+++ b/Engine/Source/Runtime/Renderer/Private/ScreenSpaceDenoise.cpp
@@ -319,8 +319,8 @@ const TCHAR* const kInjestResourceNames[] = {
 	// ShadowVisibilityMask
 	TEXT("ShadowDenoiserInjest0"),
 	TEXT("ShadowDenoiserInjest1"),
-	nullptr,
-	nullptr,
+	TEXT("ShadowDenoiserInjest2"),
+	TEXT("ShadowDenoiserInjest3"),
 
 	// PolychromaticPenumbraHarmonic
 	nullptr,
@@ -1345,15 +1345,16 @@ static void DenoiseSignalAtConstantPixelDensity(
 			check(Settings.SignalBatchSize >= 1 && Settings.SignalBatchSize <= IScreenSpaceDenoiser::kMaxBatchSize);
 			for (int32 BatchedSignalId = 0; BatchedSignalId < Settings.SignalBatchSize; BatchedSignalId++)
 			{
-				InjestDescs[BatchedSignalId / 2].Format = (BatchedSignalId % 2) ? PF_FloatRGBA : PF_G16R16F;
-				InjestTextureCount = BatchedSignalId / 2 + 1;
+				InjestDescs[BatchedSignalId].Format = PF_FloatRGBA;
+				InjestTextureCount = BatchedSignalId;
 				ReconstructionDescs[BatchedSignalId].Format = PF_FloatRGBA;
 				HistoryDescs[BatchedSignalId].Format = PF_FloatRGBA;
 			}
+			InjestTextureCount = Settings.SignalBatchSize;
 
 			HistoryTextureCountPerSignal = 1;
 			ReconstructionTextureCount = Settings.SignalBatchSize;
-			bHasReconstructionLayoutDifferentFromHistory = true;
+			bHasReconstructionLayoutDifferentFromHistory = false;
 		}
 		else if (Settings.SignalProcessing == ESignalProcessing::PolychromaticPenumbraHarmonic)
 		{
