Adding diffuse transmission distance to ray traced shadows.

Enabled when selecting "transmission" on a candidate light. Ray tracing "lights" cost is effectively 50% more for a light with this functionality enabled.

diff --git a/Engine/Shaders/Private/RayTracing/RayTracingOcclusionRGS.usf b/Engine/Shaders/Private/RayTracing/RayTracingOcclusionRGS.usf
--- a/Engine/Shaders/Private/RayTracing/RayTracingOcclusionRGS.usf
+++ b/Engine/Shaders/Private/RayTracing/RayTracingOcclusionRGS.usf
@@ -9,6 +9,7 @@
 #include "../SceneTextureParameters.ush"
 #include "../ScreenSpaceDenoise/SSDPublic.ush"
 #include "../SobolRandom.ush"
+#include "../TransmissionCommon.ush"
 #include "../HairStrands/HairStrandsVisibilityCommon.ush"
 #include "RayTracingCommon.ush"
 #include "RayTracingDeferredShadingCommon.ush"
@@ -196,6 +197,95 @@ bool GenerateOcclusionRay(
 	return true;
 }
 
+float3 ComputeDiffuseTransmission(float3 V, float3 P, float3 N, float MeanFreePath, float Eta, FLightShaderParameters LightParameters, float2 LightSample, inout RandomSequence RandSequence, uint RejectionRetryMax)
+{
+	float TransmissionDistance = SSSS_MAX_TRANSMISSION_PROFILE_DISTANCE;
+	float RcpMeanFreePath = rcp(MeanFreePath);
+
+	// Refract inward and force a scattering event
+	float3 Tr = refract(V, N, Eta);
+
+	// TODO: decide between stochastic or expectation
+	//float ScatterDistance = MeanFreePath;
+	uint DummyVariable;
+	float RandSample = RandomSequence_GenerateSample1D(RandSequence, DummyVariable);
+	float ScatterDistance = -log(RandSample) * MeanFreePath;
+	float ScatterDistancePdf = RcpMeanFreePath * exp(-ScatterDistance * RcpMeanFreePath);
+
+	// Reject scatter distances which penetrate through the medium
+	const uint RayFlags = 0;
+	const uint InstanceInclusionMask = RAY_TRACING_MASK_OPAQUE;
+
+	// Use rejection sampling to improve 1D estimate
+	if (RejectionRetryMax > 1)
+	{
+		RayDesc RejectionTestRay;
+		RejectionTestRay.Origin = P;
+		RejectionTestRay.Direction = Tr;
+		RejectionTestRay.TMin = 0.01;
+		RejectionTestRay.TMax = ScatterDistance;
+		FMinimalPayload RejectionTestPayload = TraceVisibilityRay(
+			TLAS,
+			RayFlags,
+			InstanceInclusionMask,
+			RejectionTestRay);
+
+		float RejectionDistancePdf = ScatterDistancePdf;
+		uint RetryCount = 0;
+		while (RejectionTestPayload.IsHit() && RetryCount < RejectionRetryMax)
+		{
+			RetryCount++;
+			ScatterDistance /= 2.0;
+			// Account for forcibly scattering at a half-distance, given that a previous trial has failed
+			ScatterDistancePdf = RcpMeanFreePath * exp(-ScatterDistance * RcpMeanFreePath);
+			RejectionDistancePdf = ScatterDistancePdf / RejectionDistancePdf;
+			RejectionTestRay.TMax = ScatterDistance;
+
+			RejectionTestPayload = TraceVisibilityRay(
+				TLAS,
+				RayFlags,
+				InstanceInclusionMask,
+				RejectionTestRay);
+		}
+		if (RejectionTestPayload.IsHit()) return 0.0;
+		ScatterDistance = ScatterDistance * ScatterDistancePdf / RejectionDistancePdf;
+	}
+
+	// Build sub-surface scattering ray
+	RayDesc SSSRay;
+	SSSRay.Origin = P + Tr * ScatterDistance;
+	GenerateOcclusionRay(LightParameters, SSSRay.Origin, N, LightSample, SSSRay.Origin, SSSRay.Direction, SSSRay.TMin, SSSRay.TMax);
+
+	// Clip the ray length by the maximum transmission distance
+	float LightTMax = SSSRay.TMax;
+	SSSRay.TMax = SSSS_MAX_TRANSMISSION_PROFILE_DISTANCE - ScatterDistance;
+
+	FMinimalPayload SSSPayload = TraceVisibilityRay(
+		TLAS,
+		RayFlags,
+		InstanceInclusionMask,
+		SSSRay);
+
+	if (SSSPayload.IsHit())
+	{
+		// Confirm that there is not an actual occluder beyond the maximum transmission distance
+		SSSRay.TMin = SSSPayload.HitT + 0.01;
+		SSSRay.TMax = LightTMax;
+		FMinimalPayload VisibilityPayload = TraceVisibilityRay(
+			TLAS,
+			RayFlags,
+			InstanceInclusionMask,
+			SSSRay);
+		if (VisibilityPayload.IsMiss())
+		{
+			TransmissionDistance = ScatterDistance + SSSPayload.HitT;
+		}
+	}
+
+	// Conversion to raster model: 1.0 - (distance / SSSS_MAX_TRANSMISSION_PROFILE_DISTANCE)
+	return EncodeThickness(TransmissionDistance / SSSS_MAX_TRANSMISSION_PROFILE_DISTANCE);
+}
+
 RAY_TRACING_ENTRY_RAYGEN(OcclusionRGS)
 {
 	uint2 PixelCoord = DispatchRaysIndex().xy + View.ViewRectMin.xy + PixelOffset;
@@ -203,6 +293,7 @@ RAY_TRACING_ENTRY_RAYGEN(OcclusionRGS)
 	float Visibility = 0.0;
 	float SumRayDistance = 0.0;
 	float ClosestRayDistance = DENOISER_INVALID_HIT_DISTANCE;
+	float TransmissionDistance = 0.0;
 	float HitCount = 0.0;
 	float RayCount = 0.0;
 	uint LocalSamplesPerPixel = ENABLE_MULTIPLE_SAMPLES_PER_PIXEL ? SamplesPerPixel : 1;
@@ -220,6 +311,7 @@ RAY_TRACING_ENTRY_RAYGEN(OcclusionRGS)
 		float2 BufferUV = (float2(PixelCoord) + 0.5) * InvBufferSize;
 		uint ShadingModelID = SHADINGMODELID_UNLIT;
 		float3 WorldNormal = 0;
+		FTransmissionProfileParams TransmissionProfileParams = (FTransmissionProfileParams)0;
 		{
 #if 0
 			FGBufferData GBufferData = GetGBufferDataFromSceneTextures(UV);
@@ -229,6 +321,9 @@ RAY_TRACING_ENTRY_RAYGEN(OcclusionRGS)
 #endif
 			ShadingModelID = GBufferData.ShadingModelID;
 			WorldNormal = GBufferData.WorldNormal;
+#if ENABLE_TRANSMISSION
+			TransmissionProfileParams = GetTransmissionProfileParams(GBufferData);
+#endif
 		}
 
 		// Mask out depth values that are infinitely far away
@@ -406,7 +501,18 @@ RAY_TRACING_ENTRY_RAYGEN(OcclusionRGS)
 #if !ENABLE_MULTIPLE_SAMPLES_PER_PIXEL
 		while (0);
 #endif //!ENABLE_MULTIPLE_SAMPLES_PER_PIXEL
+		if (ENABLE_TRANSMISSION && bTraceRay && ShadingModelID == SHADINGMODELID_SUBSURFACE_PROFILE)
+		{
+			// Snag two variables to fix light sampling for all computations..
+			uint DummyVariable;
+			float2 LightSample = RandomSequence_GenerateSample2D(RandSequence, DummyVariable);
 
+			float3 V = normalize(WorldPosition - View.WorldViewOrigin);
+			float MeanFreePath = TransmissionProfileParams.ExtinctionScale;
+			float Eta = TransmissionProfileParams.OneOverIOR;
+			uint RejectionRetryMax = LocalSamplesPerPixel - 1;
+			TransmissionDistance = ComputeDiffuseTransmission(V, WorldPosition, WorldNormal, MeanFreePath, Eta, LightParameters, LightSample, RandSequence, RejectionRetryMax);
+		}
 	}
 
 	float Shadow = (LocalSamplesPerPixel > 0) ? Visibility / LocalSamplesPerPixel : Visibility;
@@ -417,7 +523,7 @@ RAY_TRACING_ENTRY_RAYGEN(OcclusionRGS)
 			Shadow,
 			ClosestRayDistance,
 			0,
-			0);
+			TransmissionDistance);
 	}
 	else if (DIM_DENOISER_OUTPUT == 1)
 	{
@@ -432,13 +538,13 @@ RAY_TRACING_ENTRY_RAYGEN(OcclusionRGS)
 		}
 
 		// TODO(Denoiser): the denoiser would much prefer a single RG texture.
-		RWOcclusionMaskUAV[PixelCoord] = Shadow;
+		RWOcclusionMaskUAV[PixelCoord] = float4(Shadow, TransmissionDistance, Shadow, TransmissionDistance);
 		RWRayDistanceUAV[PixelCoord] = AvgHitDistance;
 	}
 	else
 	{
 		const float ShadowFadeFraction = 1;
-		float SSSTransmission = Shadow;
+		float SSSTransmission = TransmissionDistance;
 
 		// 0 is shadowed, 1 is unshadowed
 		// RETURN_COLOR not needed unless writing to SceneColor;
diff --git a/Engine/Source/Runtime/Renderer/Private/RayTracing/RayTracingShadows.cpp b/Engine/Source/Runtime/Renderer/Private/RayTracing/RayTracingShadows.cpp
--- a/Engine/Source/Runtime/Renderer/Private/RayTracing/RayTracingShadows.cpp
+++ b/Engine/Source/Runtime/Renderer/Private/RayTracing/RayTracingShadows.cpp
@@ -56,8 +56,9 @@ class FOcclusionRGS : public FGlobalShader
 	class FEnableTwoSidedGeometryDim : SHADER_PERMUTATION_BOOL("ENABLE_TWO_SIDED_GEOMETRY");
 	class FEnableMultipleSamplesPerPixel : SHADER_PERMUTATION_BOOL("ENABLE_MULTIPLE_SAMPLES_PER_PIXEL");
 	class FHairLighting : SHADER_PERMUTATION_INT("USE_HAIR_LIGHTING", 2);
+	class FEnableTransmissionDim : SHADER_PERMUTATION_INT("ENABLE_TRANSMISSION", 2);
 
-	using FPermutationDomain = TShaderPermutationDomain<FLightTypeDim, FDenoiserOutputDim, FEnableTwoSidedGeometryDim, FHairLighting, FEnableMultipleSamplesPerPixel>;
+	using FPermutationDomain = TShaderPermutationDomain<FLightTypeDim, FDenoiserOutputDim, FEnableTwoSidedGeometryDim, FHairLighting, FEnableMultipleSamplesPerPixel, FEnableTransmissionDim>;
 
 	static bool ShouldCompilePermutation(const FGlobalShaderPermutationParameters& Parameters)
 	{
@@ -99,6 +100,7 @@ class FOcclusionRGS : public FGlobalShader
 
 		SHADER_PARAMETER_RDG_TEXTURE(Texture2D, HairCategorizationTexture)
 		SHADER_PARAMETER_RDG_TEXTURE(Texture2D, HairLightChannelMaskTexture)
+		SHADER_PARAMETER_RDG_TEXTURE(Texture2D, SSProfilesTexture)
 		SHADER_PARAMETER_SRV(RaytracingAccelerationStructure, TLAS)
 		SHADER_PARAMETER_RDG_TEXTURE_UAV(RWTexture2D<float4>, RWOcclusionMaskUAV)
 		SHADER_PARAMETER_RDG_TEXTURE_UAV(RWTexture2D<float>, RWRayDistanceUAV)
@@ -127,21 +129,25 @@ void FDeferredShadingSceneRenderer::PrepareRayTracingShadows(const FViewInfo& Vi
 
 	for (int32 MultiSPP = 0; MultiSPP < 2; ++MultiSPP)
 	{
-		for (int32 HairLighting = 0; HairLighting < 2; ++HairLighting)
+		for (int32 EnableTransmissionDim = 0; EnableTransmissionDim < 2; ++EnableTransmissionDim)
 		{
-			for (int32 LightType = 0; LightType < LightType_MAX; ++LightType)
+			for (int32 HairLighting = 0; HairLighting < 2; ++HairLighting)
 			{
-				for (IScreenSpaceDenoiser::EShadowRequirements DenoiserRequirement : DenoiserRequirements)
+				for (int32 LightType = 0; LightType < LightType_MAX; ++LightType)
 				{
-					FOcclusionRGS::FPermutationDomain PermutationVector;
-					PermutationVector.Set<FOcclusionRGS::FLightTypeDim>(LightType);
-					PermutationVector.Set<FOcclusionRGS::FDenoiserOutputDim>((int32)DenoiserRequirement);
-					PermutationVector.Set<FOcclusionRGS::FEnableTwoSidedGeometryDim>(CVarRayTracingShadowsEnableTwoSidedGeometry.GetValueOnRenderThread() != 0);
-					PermutationVector.Set<FOcclusionRGS::FHairLighting>(HairLighting);
-					PermutationVector.Set<FOcclusionRGS::FEnableMultipleSamplesPerPixel>(MultiSPP != 0);
-
-					TShaderMapRef<FOcclusionRGS> RayGenerationShader(View.ShaderMap, PermutationVector);
-					OutRayGenShaders.Add(RayGenerationShader.GetRayTracingShader());
+					for (IScreenSpaceDenoiser::EShadowRequirements DenoiserRequirement : DenoiserRequirements)
+					{
+						FOcclusionRGS::FPermutationDomain PermutationVector;
+						PermutationVector.Set<FOcclusionRGS::FLightTypeDim>(LightType);
+						PermutationVector.Set<FOcclusionRGS::FDenoiserOutputDim>((int32)DenoiserRequirement);
+						PermutationVector.Set<FOcclusionRGS::FEnableTwoSidedGeometryDim>(CVarRayTracingShadowsEnableTwoSidedGeometry.GetValueOnRenderThread() != 0);
+						PermutationVector.Set<FOcclusionRGS::FHairLighting>(HairLighting);
+						PermutationVector.Set<FOcclusionRGS::FEnableMultipleSamplesPerPixel>(MultiSPP != 0);
+						PermutationVector.Set<FOcclusionRGS::FEnableTransmissionDim>(EnableTransmissionDim);
+
+						TShaderMapRef<FOcclusionRGS> RayGenerationShader(View.ShaderMap, PermutationVector);
+						OutRayGenShaders.Add(RayGenerationShader.GetRayTracingShader());
+					}
 				}
 			}
 		}
@@ -212,6 +218,7 @@ void FDeferredShadingSceneRenderer::RenderRayTracingShadows(
 		PassParameters->LightScissor = ScissorRect;
 		PassParameters->PixelOffset = PixelOffset;
 		PassParameters->ShadowMaskType = bUseHairLighting && bSubPixelShadowMask ? 1 : 0;
+		PassParameters->SSProfilesTexture = GraphBuilder.RegisterExternalTexture(View.RayTracingSubSurfaceProfileTexture);
 		if (bUseHairLighting)
 		{
 			PassParameters->bUseHairVoxel = IsHairRayTracingEnabled() ? 0 : 1;
@@ -242,6 +249,7 @@ void FDeferredShadingSceneRenderer::RenderRayTracingShadows(
 		PermutationVector.Set<FOcclusionRGS::FHairLighting>(bUseHairLighting? 1 : 0);
 
 		PermutationVector.Set<FOcclusionRGS::FEnableMultipleSamplesPerPixel>(RayTracingConfig.RayCountPerPixel > 1);
+		PermutationVector.Set<FOcclusionRGS::FEnableTransmissionDim>(LightSceneProxy->Transmission() ? 1 : 0);
 
 		TShaderMapRef<FOcclusionRGS> RayGenerationShader(GetGlobalShaderMap(FeatureLevel), PermutationVector);
 
