Quick fix to eyes, inadvertently modified in CL 11595310.

diff --git a/Engine/Shaders/Private/ScreenSpaceDenoise/SSDSpatialAccumulation.usf b/Engine/Shaders/Private/ScreenSpaceDenoise/SSDSpatialAccumulation.usf
--- a/Engine/Shaders/Private/ScreenSpaceDenoise/SSDSpatialAccumulation.usf
+++ b/Engine/Shaders/Private/ScreenSpaceDenoise/SSDSpatialAccumulation.usf
@@ -1226,7 +1226,8 @@ void MainCS(
 				float Shadow = GetSamplePenumbraSafe(OutputSamples.Array[MultiplexId]);
 
 				const float ShadowFadeFraction = 1;
-				float SSSTransmission = (OutputSamples.Array[MultiplexId].SampleCount > 0 ? OutputSamples.Array[MultiplexId].TransmissionDistance / OutputSamples.Array[MultiplexId].SampleCount : OutputSamples.Array[MultiplexId].TransmissionDistance);
+				//float SSSTransmission = (OutputSamples.Array[MultiplexId].SampleCount > 0 ? OutputSamples.Array[MultiplexId].TransmissionDistance / OutputSamples.Array[MultiplexId].SampleCount : OutputSamples.Array[MultiplexId].TransmissionDistance);
+				float SSSTransmission = Shadow;
 
 				// 0 is shadowed, 1 is unshadowed
 				// RETURN_COLOR not needed unless writing to SceneColor;
