Correctly transfer transmission result through ray tracing shadows, depending upon ShadingModel.

diff --git a/Engine/Shaders/Private/RayTracing/RayTracingOcclusionRGS.usf b/Engine/Shaders/Private/RayTracing/RayTracingOcclusionRGS.usf
--- a/Engine/Shaders/Private/RayTracing/RayTracingOcclusionRGS.usf
+++ b/Engine/Shaders/Private/RayTracing/RayTracingOcclusionRGS.usf
@@ -501,6 +501,7 @@ RAY_TRACING_ENTRY_RAYGEN(OcclusionRGS)
 #if !ENABLE_MULTIPLE_SAMPLES_PER_PIXEL
 		while (0);
 #endif //!ENABLE_MULTIPLE_SAMPLES_PER_PIXEL
+
 		if (ENABLE_TRANSMISSION && bTraceRay && ShadingModelID == SHADINGMODELID_SUBSURFACE_PROFILE)
 		{
 			// Snag two variables to fix light sampling for all computations..
@@ -513,6 +514,10 @@ RAY_TRACING_ENTRY_RAYGEN(OcclusionRGS)
 			uint RejectionRetryMax = LocalSamplesPerPixel - 1;
 			TransmissionDistance = ComputeDiffuseTransmission(V, WorldPosition, WorldNormal, MeanFreePath, Eta, LightParameters, LightSample, RandSequence, RejectionRetryMax);
 		}
+		else if (ShadingModelID == SHADINGMODELID_EYE)
+		{
+			TransmissionDistance = (LocalSamplesPerPixel > 0) ? Visibility / LocalSamplesPerPixel : Visibility;
+		}
 	}
 
 	float Shadow = (LocalSamplesPerPixel > 0) ? Visibility / LocalSamplesPerPixel : Visibility;
diff --git a/Engine/Shaders/Private/ScreenSpaceDenoise/SSDSpatialAccumulation.usf b/Engine/Shaders/Private/ScreenSpaceDenoise/SSDSpatialAccumulation.usf
--- a/Engine/Shaders/Private/ScreenSpaceDenoise/SSDSpatialAccumulation.usf
+++ b/Engine/Shaders/Private/ScreenSpaceDenoise/SSDSpatialAccumulation.usf
@@ -1226,8 +1226,7 @@ void MainCS(
 				float Shadow = GetSamplePenumbraSafe(OutputSamples.Array[MultiplexId]);
 
 				const float ShadowFadeFraction = 1;
-				//float SSSTransmission = (OutputSamples.Array[MultiplexId].SampleCount > 0 ? OutputSamples.Array[MultiplexId].TransmissionDistance / OutputSamples.Array[MultiplexId].SampleCount : OutputSamples.Array[MultiplexId].TransmissionDistance);
-				float SSSTransmission = Shadow;
+				float SSSTransmission = (OutputSamples.Array[MultiplexId].SampleCount > 0 ? OutputSamples.Array[MultiplexId].TransmissionDistance / OutputSamples.Array[MultiplexId].SampleCount : OutputSamples.Array[MultiplexId].TransmissionDistance);
 
 				// 0 is shadowed, 1 is unshadowed
 				// RETURN_COLOR not needed unless writing to SceneColor;
