Relax constraint and propagate transmission term as shadow term for all shading models except sub-surface profile.

diff --git a/Engine/Shaders/Private/RayTracing/RayTracingOcclusionRGS.usf b/Engine/Shaders/Private/RayTracing/RayTracingOcclusionRGS.usf
--- a/Engine/Shaders/Private/RayTracing/RayTracingOcclusionRGS.usf
+++ b/Engine/Shaders/Private/RayTracing/RayTracingOcclusionRGS.usf
@@ -514,7 +514,7 @@ RAY_TRACING_ENTRY_RAYGEN(OcclusionRGS)
 			uint RejectionRetryMax = LocalSamplesPerPixel - 1;
 			TransmissionDistance = ComputeDiffuseTransmission(V, WorldPosition, WorldNormal, MeanFreePath, Eta, LightParameters, LightSample, RandSequence, RejectionRetryMax);
 		}
-		else if (ShadingModelID == SHADINGMODELID_EYE)
+		else // if (ShadingModelID == SHADINGMODELID_EYE)
 		{
 			TransmissionDistance = (LocalSamplesPerPixel > 0) ? Visibility / LocalSamplesPerPixel : Visibility;
 		}
