Change occlusion RGS to run in a single pass when hair is enabled.

This allows to reduce the tracing cost by roughly 20% in test cases.

diff --git a/Engine/Shaders/Private/RayTracing/RayTracingOcclusionRGS.usf b/Engine/Shaders/Private/RayTracing/RayTracingOcclusionRGS.usf
--- a/Engine/Shaders/Private/RayTracing/RayTracingOcclusionRGS.usf
+++ b/Engine/Shaders/Private/RayTracing/RayTracingOcclusionRGS.usf
@@ -31,8 +31,8 @@
 RaytracingAccelerationStructure TLAS;
 RWTexture2D<float4> RWOcclusionMaskUAV;
 RWTexture2D<float> RWRayDistanceUAV;
+RWTexture2D<float4> RWSubPixelOcclusionMaskUAV;
 
-uint ShadowMaskType;
 uint LightingChannelMask;
 uint SamplesPerPixel;
 float NormalBias;
@@ -197,7 +197,7 @@ bool GenerateOcclusionRay(
 	return true;
 }
 
-float3 ComputeDiffuseTransmission(float3 V, float3 P, float3 N, float MeanFreePath, float Eta, FLightShaderParameters LightParameters, float2 LightSample, inout RandomSequence RandSequence, uint RejectionRetryMax)
+float ComputeDiffuseTransmission(float3 V, float3 P, float3 N, float MeanFreePath, float Eta, FLightShaderParameters LightParameters, float2 LightSample, inout RandomSequence RandSequence, uint RejectionRetryMax)
 {
 	float TransmissionDistance = SSSS_MAX_TRANSMISSION_PROFILE_DISTANCE;
 	float RcpMeanFreePath = rcp(MeanFreePath);
@@ -286,16 +286,216 @@ float3 ComputeDiffuseTransmission(float3 V, float3 P, float3 N, float MeanFreePa
 	return EncodeThickness(TransmissionDistance / SSSS_MAX_TRANSMISSION_PROFILE_DISTANCE);
 }
 
+struct FOcclusionResult
+{
+	float Visibility;
+	float SumRayDistance;
+	float ClosestRayDistance;
+	float TransmissionDistance;
+	float HitCount;
+	float RayCount;
+};
+
+FOcclusionResult InitOcclusionResult()
+{
+	FOcclusionResult Out;
+
+	Out.Visibility = 0.0;
+	Out.SumRayDistance = 0.0;
+	Out.ClosestRayDistance = DENOISER_INVALID_HIT_DISTANCE;
+	Out.TransmissionDistance = 0.0;
+	Out.HitCount = 0.0;
+	Out.RayCount = 0.0;
+
+	return Out;
+}
+
+float OcclusionToShadow(FOcclusionResult In, uint LocalSamplesPerPixel)
+{
+	return (LocalSamplesPerPixel > 0) ? In.Visibility / LocalSamplesPerPixel : In.Visibility;
+}
+
+FOcclusionResult ComputeOcclusion(
+	const uint2 PixelCoord,
+	RandomSequence RandSequence,
+	const uint ShadingModelID,
+	const uint RaytracingMask,
+	const float DeviceZ,
+	float3 WorldNormal,
+	const FLightShaderParameters LightParameters,
+	const FTransmissionProfileParams TransmissionProfileParams,
+	const uint LocalSamplesPerPixel)
+{
+	FOcclusionResult Out = InitOcclusionResult();
+
+	const float3 WorldPosition = ReconstructWorldPositionFromDeviceZ(PixelCoord, DeviceZ);
+
+	// Cheap solution for double sided foliage to have valid lighing through the denoiser
+	if (ShadingModelID == SHADINGMODELID_TWOSIDED_FOLIAGE)
+	{
+	#if LIGHT_TYPE == LIGHT_TYPE_DIRECTIONAL
+		const float3 LightDirection = LightParameters.Direction;
+	#else
+		const float3 LightDirection = normalize(LightParameters.Position - WorldPosition);
+	#endif
+		if (dot(LightDirection, WorldNormal) < 0)
+			WorldNormal = -WorldNormal;
+	}
+
+	// For hair shading model, WorldNormal is actually the shading tangent.
+	// To generate proper bias, we compute a normal oriented toward the light.
+	// Normal clipping is removed from hair since the BSDF is spherical, rather 
+	// than hemispherical
+	//
+	// Note: 
+	// Since we don't have a notion of backfacing here, there is no correct way 
+	// to avoid self-intersection with hair. Either we push towards the light, 
+	// (this creates some issue in triangle facing back the light), or we push 
+	// toward the view point (this create issue for transmitted light)
+	// It seems having proper transmission is more important, thus the light 
+	// version is enabled by default
+	bool bApplyNormalCulling = true;
+	if (ShadingModelID == SHADINGMODELID_HAIR)
+	{
+	#if LIGHT_TYPE == LIGHT_TYPE_DIRECTIONAL
+		const float3 LightDirection = LightParameters.Direction;
+	#else
+		const float3 LightDirection = normalize(LightParameters.Position - WorldPosition);
+	#endif
+		WorldNormal = LightDirection;
+		bApplyNormalCulling = false;
+	}
+
+	// Fixes clipping of eye caustic normals in raytrace shadows
+	if (ShadingModelID == SHADINGMODELID_EYE)
+	{
+		bApplyNormalCulling = false;
+	}
+
+#if ENABLE_MULTIPLE_SAMPLES_PER_PIXEL
+	LOOP for (uint SampleIndex = 0; SampleIndex < LocalSamplesPerPixel; ++SampleIndex)
+#else // ENABLE_MULTIPLE_SAMPLES_PER_PIXEL
+	do if (LocalSamplesPerPixel > 0)
+#endif // ENABLE_MULTIPLE_SAMPLES_PER_PIXEL
+	{
+		uint DummyVariable;
+		float2 RandSample = RandomSequence_GenerateSample2D(RandSequence, DummyVariable);
+
+		#if !ENABLE_MULTIPLE_SAMPLES_PER_PIXEL
+		{
+			const uint TemporalBits = 10;
+			uint FrameIndexMod1024 = ReverseBitsN(GetPowerOfTwoModulatedFrameIndex(1 << TemporalBits), TemporalBits);
+
+			uint2 SobolBase = SobolPixel(PixelCoord);
+			uint2 SobolFrame = SobolIndex(SobolBase, FrameIndexMod1024, TemporalBits);
+
+			RandSample = SobolIndexToUniformUnitSquare(SobolFrame);
+		}
+		#endif // !ENABLE_MULTIPLE_SAMPLES_PER_PIXEL
+
+		RayDesc Ray;
+		bool bIsValidRay = GenerateOcclusionRay(
+			LightParameters,
+			WorldPosition, WorldNormal,
+			RandSample,
+			/* out */ Ray.Origin,
+			/* out */ Ray.Direction,
+			/* out */ Ray.TMin,
+			/* out */ Ray.TMax);
+		ApplyCameraRelativeDepthBias(Ray, PixelCoord, DeviceZ, WorldNormal, NormalBias);
+
+		BRANCH
+		if (!bIsValidRay)
+		{
+			continue;
+		}
+		else if (bApplyNormalCulling && dot(WorldNormal, Ray.Direction) <= 0.0)
+		{
+			Out.ClosestRayDistance = 0.001;
+			continue;
+		}
+
+		// Attenuation check
+		if (LightParameters.InvRadius > 0.0)
+		{
+			const float MaxAttenuationDistance = 1.0 / LightParameters.InvRadius;
+			if (Ray.TMax > MaxAttenuationDistance)
+			{
+				continue;
+			}
+		}
+
+		uint RayFlags = 0;
+
+	#if !ENABLE_TWO_SIDED_GEOMETRY
+		RayFlags |= RAY_FLAG_CULL_BACK_FACING_TRIANGLES;
+	#endif
+
+		if (DIM_DENOISER_OUTPUT == 0 || ShadingModelID == SHADINGMODELID_HAIR)
+		{
+			// Denoiser mode 0 doesn't use depth, so take first hit
+			RayFlags |= RAY_FLAG_SKIP_CLOSEST_HIT_SHADER | RAY_FLAG_ACCEPT_FIRST_HIT_AND_END_SEARCH;
+		}
+
+		FMinimalPayload MinimalPayload = TraceVisibilityRay(
+			TLAS,
+			RayFlags,
+			RaytracingMask,
+			Ray);
+
+		#if USE_HAIR_LIGHTING
+		if (ShadingModelID != SHADINGMODELID_HAIR && bUseHairVoxel)
+		{
+			MinimalPayload.HitT = TraverseHair(PixelCoord, Ray.Origin, Ray.Direction, MinimalPayload.HitT);
+		}
+		#endif
+
+		Out.RayCount += 1.0;
+		Out.Visibility += MinimalPayload.IsMiss() ? 1.0 : 0.0;
+		if (MinimalPayload.IsHit())
+		{
+			Out.ClosestRayDistance =
+				(Out.ClosestRayDistance == DENOISER_INVALID_HIT_DISTANCE) ||
+				(MinimalPayload.HitT < Out.ClosestRayDistance) ? MinimalPayload.HitT : Out.ClosestRayDistance;
+			Out.SumRayDistance += MinimalPayload.HitT;
+			Out.HitCount += 1.0;
+		}
+		else
+		{
+			Out.ClosestRayDistance = (Out.ClosestRayDistance == DENOISER_INVALID_HIT_DISTANCE) ? DENOISER_MISS_HIT_DISTANCE : Out.ClosestRayDistance;
+		}
+	}
+#if !ENABLE_MULTIPLE_SAMPLES_PER_PIXEL
+	while (0);
+#endif // ENABLE_MULTIPLE_SAMPLES_PER_PIXEL
+
+	if (ENABLE_TRANSMISSION && LocalSamplesPerPixel > 0 && ShadingModelID == SHADINGMODELID_SUBSURFACE_PROFILE)
+	{
+		// Snag two variables to fix light sampling for all computations..
+		uint DummyVariable;
+		const float2 LightSample = RandomSequence_GenerateSample2D(RandSequence, DummyVariable);
+		const float3 V = normalize(WorldPosition - View.WorldViewOrigin);
+
+		const float MeanFreePath = TransmissionProfileParams.ExtinctionScale;
+		const float Eta = TransmissionProfileParams.OneOverIOR;
+		const uint RejectionRetryMax = LocalSamplesPerPixel - 1;
+
+		Out.TransmissionDistance = ComputeDiffuseTransmission(V, WorldPosition, WorldNormal, MeanFreePath, Eta, LightParameters, LightSample, RandSequence, RejectionRetryMax);
+	}
+	else // if (ShadingModelID == SHADINGMODELID_EYE)
+	{
+		Out.TransmissionDistance = (LocalSamplesPerPixel > 0) ? Out.Visibility / LocalSamplesPerPixel : Out.Visibility;
+	}
+	return Out;
+}
+
 RAY_TRACING_ENTRY_RAYGEN(OcclusionRGS)
 {
 	uint2 PixelCoord = DispatchRaysIndex().xy + View.ViewRectMin.xy + PixelOffset;
 
-	float Visibility = 0.0;
-	float SumRayDistance = 0.0;
-	float ClosestRayDistance = DENOISER_INVALID_HIT_DISTANCE;
-	float TransmissionDistance = 0.0;
-	float HitCount = 0.0;
-	float RayCount = 0.0;
+	FOcclusionResult Occlusion = InitOcclusionResult();
+	FOcclusionResult HairOcclusion = InitOcclusionResult();
+
 	uint LocalSamplesPerPixel = ENABLE_MULTIPLE_SAMPLES_PER_PIXEL ? SamplesPerPixel : 1;
 
 	if (all(PixelCoord >= LightScissor.xy) && all(PixelCoord <= LightScissor.zw))
@@ -309,8 +509,8 @@ RAY_TRACING_ENTRY_RAYGEN(OcclusionRGS)
 		// Get G-Buffer surface data
 		float2 InvBufferSize = View.BufferSizeAndInvSize.zw;
 		float2 BufferUV = (float2(PixelCoord) + 0.5) * InvBufferSize;
-		uint ShadingModelID = SHADINGMODELID_UNLIT;
 		float3 WorldNormal = 0;
+		uint ShadingModelID = SHADINGMODELID_UNLIT;
 		FTransmissionProfileParams TransmissionProfileParams = (FTransmissionProfileParams)0;
 		{
 #if 0
@@ -335,221 +535,96 @@ RAY_TRACING_ENTRY_RAYGEN(OcclusionRGS)
 		float DeviceZ = SceneDepthBuffer.Load(int3(PixelCoord, 0)).r;
 #endif
 		const bool bIsDepthValid = SceneDepthBuffer.Load(int3(PixelCoord, 0)).r > 0.0;
-		bool bIsValidPixel = ShadingModelID != SHADINGMODELID_UNLIT && bIsDepthValid;
-		uint LightChannel = GetSceneLightingChannel(PixelCoord);
-
-		// Check if there is hair geometry
-		uint RaytracingMask = RAY_TRACING_MASK_SHADOW | RAY_TRACING_MASK_THIN_SHADOW;
-	#if USE_HAIR_LIGHTING
-		if (ShadowMaskType == ShadowMaskType_Hair)
-		{
-			RaytracingMask = RAY_TRACING_MASK_SHADOW;
-			const FCategorizationData Categorization = DecodeCategorizationData(HairCategorizationTexture.Load(uint3(PixelCoord, 0)));
-			bIsValidPixel = Categorization.PixelCoverage > 0;
-			if (bIsValidPixel)
-			{
-				DeviceZ = Categorization.ClosestDepth;
-				ShadingModelID = SHADINGMODELID_HAIR;
-				LightChannel = HairLightChannelMaskTexture.Load(uint3(PixelCoord, 0));
-			}
-		}
-	#endif
-
-		const float3 WorldPosition = ReconstructWorldPositionFromDeviceZ(PixelCoord, DeviceZ);
-
+		const bool bIsValidPixel = ShadingModelID != SHADINGMODELID_UNLIT && bIsDepthValid;
+		const uint LightChannel = GetSceneLightingChannel(PixelCoord);
 		const bool bTraceRay = bIsValidPixel && (LightChannel & LightingChannelMask) != 0;
 		if (!bTraceRay)
 		{
 			LocalSamplesPerPixel = 0;
 		}
 
-		// Cheap solution for double sided foliage to have valid lighing through the denoiser
-		if (ShadingModelID == SHADINGMODELID_TWOSIDED_FOLIAGE)
-		{
-		#if LIGHT_TYPE == LIGHT_TYPE_DIRECTIONAL
-			const float3 LightDirection = LightParameters.Direction;
-		#else
-			const float3 LightDirection = normalize(LightParameters.Position - WorldPosition);
-		#endif
-			if (dot(LightDirection, WorldNormal) < 0)
-				WorldNormal = -WorldNormal;
-		}
-
-		// For hair shading model, WorldNormal is actually the shading tangent.
-		// To generate proper bias, we compute a normal oriented toward the light.
-		// Normal clipping is removed from hair since the BSDF is spherical, rather 
-		// than hemispherical
-		//
-		// Note: 
-		// Since we don't have a notion of backfacing here, there is no correct way 
-		// to avoid self-intersection with hair. Either we push towards the light, 
-		// (this creates some issue in triangle facing back the light), or we push 
-		// toward the view point (this create issue for transmitted light)
-		// It seems having proper transmission is more important, thus the light 
-		// version is enabled by default
-		bool bApplyNormalCulling = true;
-		if (ShadingModelID == SHADINGMODELID_HAIR)
-		{
-#if LIGHT_TYPE == LIGHT_TYPE_DIRECTIONAL
-			const float3 LightDirection = LightParameters.Direction;
-#else
-			const float3 LightDirection = normalize(LightParameters.Position - WorldPosition);
-#endif
-			WorldNormal = LightDirection;
-			bApplyNormalCulling = false;
-		}
-			
-		// Fixes clipping of eye caustic normals in raytrace shadows
-		if (ShadingModelID == SHADINGMODELID_EYE)
+	#if USE_HAIR_LIGHTING
+		const FCategorizationData Categorization = DecodeCategorizationData(HairCategorizationTexture.Load(uint3(PixelCoord, 0)));
+		bool bTraceHairRay = false;
+		float HairDeviceZ = 0;
+		if (Categorization.PixelCoverage > 0)
 		{
-			bApplyNormalCulling = false;
+			const uint HairLightChannel = HairLightChannelMaskTexture.Load(uint3(PixelCoord, 0));
+			HairDeviceZ	= Categorization.ClosestDepth;
+			bTraceHairRay = (HairLightChannel & LightingChannelMask) != 0;
 		}
+	#endif
 
-#if ENABLE_MULTIPLE_SAMPLES_PER_PIXEL
-		LOOP for (uint SampleIndex = 0; SampleIndex < LocalSamplesPerPixel; ++SampleIndex)
-#else // ENABLE_MULTIPLE_SAMPLES_PER_PIXEL
-		do if (bTraceRay)
-#endif // ENABLE_MULTIPLE_SAMPLES_PER_PIXEL
+		Occlusion = ComputeOcclusion(
+			PixelCoord,
+			RandSequence,
+			ShadingModelID,
+			RAY_TRACING_MASK_SHADOW | RAY_TRACING_MASK_THIN_SHADOW,
+			DeviceZ,
+			WorldNormal,
+			LightParameters, 
+			TransmissionProfileParams,
+			LocalSamplesPerPixel);
+
+		#if USE_HAIR_LIGHTING
+		if (bTraceHairRay)
 		{
-			uint DummyVariable;
-			float2 RandSample = RandomSequence_GenerateSample2D(RandSequence, DummyVariable);
-
-			#if !ENABLE_MULTIPLE_SAMPLES_PER_PIXEL
-			{
-				const uint TemporalBits = 10;
-				uint FrameIndexMod1024 = ReverseBitsN(GetPowerOfTwoModulatedFrameIndex(1 << TemporalBits), TemporalBits);
-
-				uint2 SobolBase = SobolPixel(PixelCoord);
-				uint2 SobolFrame = SobolIndex(SobolBase, FrameIndexMod1024, TemporalBits);
-
-				RandSample = SobolIndexToUniformUnitSquare(SobolFrame);
-			}
-			#endif // !ENABLE_MULTIPLE_SAMPLES_PER_PIXEL
-
-			RayDesc Ray;
-			bool bIsValidRay = GenerateOcclusionRay(
+			HairOcclusion = ComputeOcclusion(
+				PixelCoord,
+				RandSequence,
+				SHADINGMODELID_HAIR,
+				RAY_TRACING_MASK_SHADOW,
+				HairDeviceZ,
+				float3(1,0,0),
 				LightParameters,
-				WorldPosition, WorldNormal,
-				RandSample,
-				/* out */ Ray.Origin,
-				/* out */ Ray.Direction,
-				/* out */ Ray.TMin,
-				/* out */ Ray.TMax);
-			ApplyCameraRelativeDepthBias(Ray, PixelCoord, DeviceZ, WorldNormal, NormalBias);
-
-			BRANCH
-			if (!bIsValidRay)
-			{
-				continue;
-			}
-			else if (bApplyNormalCulling && dot(WorldNormal, Ray.Direction) <= 0.0)
-			{
-				ClosestRayDistance = 0.001;
-				continue;
-			}
-
-			// Attenuation check
-			if (LightParameters.InvRadius > 0.0)
-			{
-				float MaxAttenuationDistance = 1.0 / LightParameters.InvRadius;
-				if (Ray.TMax > MaxAttenuationDistance)
-				{
-					continue;
-				}
-			}
-
-			uint RayFlags = 0;
-
-#if !ENABLE_TWO_SIDED_GEOMETRY
-			RayFlags |= RAY_FLAG_CULL_BACK_FACING_TRIANGLES;
-#endif
-
-			if (DIM_DENOISER_OUTPUT == 0)
-			{
-				// Denoiser mode 0 doesn't use depth, so take first hit
-				RayFlags |= RAY_FLAG_SKIP_CLOSEST_HIT_SHADER | RAY_FLAG_ACCEPT_FIRST_HIT_AND_END_SEARCH;
-			}
-
-			FMinimalPayload MinimalPayload = TraceVisibilityRay(
-				TLAS,
-				RayFlags,
-				RaytracingMask,
-				Ray);
-
-			#if USE_HAIR_LIGHTING
-			if (ShadowMaskType != ShadowMaskType_Hair && bUseHairVoxel)
-			{
-				MinimalPayload.HitT = TraverseHair(PixelCoord, Ray.Origin, Ray.Direction, MinimalPayload.HitT);
-			}
-			#endif
-
-			RayCount += 1.0;
-			Visibility += MinimalPayload.IsMiss() ? 1.0 : 0.0;
-			if (MinimalPayload.IsHit())
-			{
-				ClosestRayDistance =
-					(ClosestRayDistance == DENOISER_INVALID_HIT_DISTANCE) ||
-					(MinimalPayload.HitT < ClosestRayDistance) ? MinimalPayload.HitT : ClosestRayDistance;
-				SumRayDistance += MinimalPayload.HitT;
-				HitCount += 1.0;
-			}
-			else
-			{
-				ClosestRayDistance = (ClosestRayDistance == DENOISER_INVALID_HIT_DISTANCE) ? DENOISER_MISS_HIT_DISTANCE : ClosestRayDistance;
-			}
+				TransmissionProfileParams,
+				1);
 		}
-#if !ENABLE_MULTIPLE_SAMPLES_PER_PIXEL
-		while (0);
-#endif //!ENABLE_MULTIPLE_SAMPLES_PER_PIXEL
+		#endif
+	}
 
-		if (ENABLE_TRANSMISSION && bTraceRay && ShadingModelID == SHADINGMODELID_SUBSURFACE_PROFILE)
-		{
-			// Snag two variables to fix light sampling for all computations..
-			uint DummyVariable;
-			float2 LightSample = RandomSequence_GenerateSample2D(RandSequence, DummyVariable);
-
-			float3 V = normalize(WorldPosition - View.WorldViewOrigin);
-			float MeanFreePath = TransmissionProfileParams.ExtinctionScale;
-			float Eta = TransmissionProfileParams.OneOverIOR;
-			uint RejectionRetryMax = LocalSamplesPerPixel - 1;
-			TransmissionDistance = ComputeDiffuseTransmission(V, WorldPosition, WorldNormal, MeanFreePath, Eta, LightParameters, LightSample, RandSequence, RejectionRetryMax);
-		}
-		else // if (ShadingModelID == SHADINGMODELID_EYE)
-		{
-			TransmissionDistance = (LocalSamplesPerPixel > 0) ? Visibility / LocalSamplesPerPixel : Visibility;
-		}
+	// Opaque shadow on ray does not need any denoising due to the high frequency nature of hair.
+	#if USE_HAIR_LIGHTING
+	{
+		// the channel assignment is documented in ShadowRendering.cpp (look for Light Attenuation channel assignment)
+		const float HairShadow = OcclusionToShadow(HairOcclusion, LocalSamplesPerPixel);
+		const float FadedShadow = Square(HairShadow);
+		const float4 OutColor = EncodeLightAttenuation(FadedShadow.xxxx);
+		RWSubPixelOcclusionMaskUAV[PixelCoord] = OutColor;
 	}
+	#endif
+
+	const float Shadow = OcclusionToShadow(Occlusion, LocalSamplesPerPixel);
 
-	float Shadow = (LocalSamplesPerPixel > 0) ? Visibility / LocalSamplesPerPixel : Visibility;
-	
 	if (DIM_DENOISER_OUTPUT == 2)
 	{
 		RWOcclusionMaskUAV[PixelCoord] = float4(
 			Shadow,
-			ClosestRayDistance,
+			Occlusion.ClosestRayDistance,
 			0,
-			TransmissionDistance);
+			Occlusion.TransmissionDistance);
+
 	}
 	else if (DIM_DENOISER_OUTPUT == 1)
 	{
 		float AvgHitDistance = -1.0;
-		if (HitCount > 0.0)
+		if (Occlusion.HitCount > 0.0)
 		{
-			AvgHitDistance = SumRayDistance / HitCount;
+			AvgHitDistance = Occlusion.SumRayDistance / Occlusion.HitCount;
 		}
-		else if (RayCount > 0.0)
+		else if (Occlusion.RayCount > 0.0)
 		{
 			AvgHitDistance = 1.0e27;
 		}
 
 		// TODO(Denoiser): the denoiser would much prefer a single RG texture.
-		RWOcclusionMaskUAV[PixelCoord] = float4(Shadow, TransmissionDistance, Shadow, TransmissionDistance);
+		RWOcclusionMaskUAV[PixelCoord] = float4(Shadow, Occlusion.TransmissionDistance, Shadow, Occlusion.TransmissionDistance);
 		RWRayDistanceUAV[PixelCoord] = AvgHitDistance;
 	}
 	else
 	{
 		const float ShadowFadeFraction = 1;
-		float SSSTransmission = TransmissionDistance;
+		float SSSTransmission = Occlusion.TransmissionDistance;
 
 		// 0 is shadowed, 1 is unshadowed
 		// RETURN_COLOR not needed unless writing to SceneColor;
diff --git a/Engine/Source/Runtime/Renderer/Private/DeferredShadingRenderer.h b/Engine/Source/Runtime/Renderer/Private/DeferredShadingRenderer.h
--- a/Engine/Source/Runtime/Renderer/Private/DeferredShadingRenderer.h
+++ b/Engine/Source/Runtime/Renderer/Private/DeferredShadingRenderer.h
@@ -528,10 +528,10 @@ private:
 		const FLightSceneInfo& LightSceneInfo,
 		const IScreenSpaceDenoiser::FShadowRayTracingConfig& RayTracingConfig,
 		const IScreenSpaceDenoiser::EShadowRequirements DenoiserRequirements,
-		const bool bSubPixelShadowMask,
 		const struct FHairStrandsOcclusionResources* HairResources,
 		FRDGTextureUAV* OutShadowMaskUAV,
-		FRDGTextureUAV* OutRayHitDistanceUAV);
+		FRDGTextureUAV* OutRayHitDistanceUAV,
+		FRDGTextureUAV* SubPixelRayTracingShadowMaskUAV);
 
 	void RenderRayTracingStochasticRectLight(FRHICommandListImmediate& RHICmdList, const FLightSceneInfo& RectLightSceneInfo, TRefCountPtr<IPooledRenderTarget>& RectLightRT, TRefCountPtr<IPooledRenderTarget>& HitDistanceRT);
 	void CompositeRayTracingSkyLight(FRHICommandListImmediate& RHICmdList, TRefCountPtr<IPooledRenderTarget>& SkyLightRT, TRefCountPtr<IPooledRenderTarget>& HitDistanceRT);
diff --git a/Engine/Source/Runtime/Renderer/Private/LightRendering.cpp b/Engine/Source/Runtime/Renderer/Private/LightRendering.cpp
--- a/Engine/Source/Runtime/Renderer/Private/LightRendering.cpp
+++ b/Engine/Source/Runtime/Renderer/Private/LightRendering.cpp
@@ -1372,6 +1372,21 @@ void FDeferredShadingSceneRenderer::RenderLights(FRHICommandListImmediate& RHICm
 									RayDistanceTexture = GraphBuilder.CreateTexture(Desc, TEXT("RayTracingOcclusionDistance"));
 								}
 
+								FRDGTextureRef SubPixelRayTracingShadowMaskTexture = nullptr;
+								FRDGTextureUAV* SubPixelRayTracingShadowMaskUAV = nullptr;
+								if (bUseHairLighting)
+								{
+									FRDGTextureDesc Desc = FRDGTextureDesc::Create2DDesc(
+										SceneTextures.SceneDepthBuffer->Desc.Extent,
+										PF_FloatRGBA,
+										FClearValueBinding::Black,
+										TexCreate_None,
+										TexCreate_ShaderResource | TexCreate_RenderTargetable | TexCreate_UAV,
+										/* bInForceSeparateTargetAndShaderResource = */ false);
+									SubPixelRayTracingShadowMaskTexture = GraphBuilder.CreateTexture(Desc, TEXT("SubPixelRayTracingOcclusion"));
+									SubPixelRayTracingShadowMaskUAV = GraphBuilder.CreateUAV(FRDGTextureUAVDesc(SubPixelRayTracingShadowMaskTexture));
+								}
+
 								FString BatchLightNameWithLevel;
 								GetLightNameForDrawEvent(BatchLightSceneInfo.Proxy, BatchLightNameWithLevel);
 
@@ -1381,6 +1396,8 @@ void FDeferredShadingSceneRenderer::RenderLights(FRHICommandListImmediate& RHICm
 								{
 									RDG_EVENT_SCOPE(GraphBuilder, "%s", *BatchLightNameWithLevel);
 
+									// Ray trace the shadow cast by opaque geometries on to hair strands geometries
+									// Note: No denoiser is required on this output, as the hair strands are geometrically noisy, which make it hard to denoise
 									RenderRayTracingShadows(
 										GraphBuilder,
 										SceneTextures,
@@ -1388,62 +1405,18 @@ void FDeferredShadingSceneRenderer::RenderLights(FRHICommandListImmediate& RHICm
 										BatchLightSceneInfo,
 										BatchRayTracingConfig,
 										DenoiserRequirements,
-										false,
 										&HairResources,
 										RayTracingShadowMaskUAV,
-										RayHitDistanceUAV);
-								}
-
-								// Ray trace the shadow cast by opaque geometries on to hair strands geometries
-								// Note: No denoiser is required on this output, as the hair strands are geometrically noisy, and so hard to denoise
-								if (HasHairStrandsClusters(ViewIndex))
-								{
-									//#dxr_todo: support multiview for the batching case
-									FRDGTextureRef SubPixelRayTracingShadowMaskTexture;
+										RayHitDistanceUAV,
+										SubPixelRayTracingShadowMaskUAV);
+									
+									if (HasHairStrandsClusters(ViewIndex))
 									{
-										FRDGTextureDesc Desc = FRDGTextureDesc::Create2DDesc(
-											SceneTextures.SceneDepthBuffer->Desc.Extent,
-											PF_FloatRGBA,
-											FClearValueBinding::Black,
-											TexCreate_None,
-											TexCreate_ShaderResource | TexCreate_RenderTargetable | TexCreate_UAV,
-											/* bInForceSeparateTargetAndShaderResource = */ false);
-										SubPixelRayTracingShadowMaskTexture = GraphBuilder.CreateTexture(Desc, TEXT("RayTracingOcclusion"));
-									}
+										TRefCountPtr<IPooledRenderTarget>* RefDestination = &PreprocessedShadowMaskSubPixelTextures[LightBatchIndex - AttenuationLightStart];
+										check(*RefDestination == nullptr);
 
-									FRDGTextureRef SubPixelRayDistanceTexture;
-									{
-										FRDGTextureDesc Desc = FRDGTextureDesc::Create2DDesc(
-											SceneTextures.SceneDepthBuffer->Desc.Extent,
-											PF_R16F,
-											FClearValueBinding::Black,
-											TexCreate_None,
-											TexCreate_ShaderResource | TexCreate_RenderTargetable | TexCreate_UAV,
-											/* bInForceSeparateTargetAndShaderResource = */ false);
-										SubPixelRayDistanceTexture = GraphBuilder.CreateTexture(Desc, TEXT("RayTracingOcclusionDistance"));
+										GraphBuilder.QueueTextureExtraction(SubPixelRayTracingShadowMaskTexture, RefDestination);
 									}
-
-									FRDGTextureUAV* SubPixelRayTracingShadowMaskUAV = GraphBuilder.CreateUAV(FRDGTextureUAVDesc(SubPixelRayTracingShadowMaskTexture));
-									FRDGTextureUAV* SubPixelRayHitDistanceUAV = GraphBuilder.CreateUAV(FRDGTextureUAVDesc(SubPixelRayDistanceTexture));
-
-									RDG_EVENT_SCOPE(GraphBuilder, "%s Hair Strands", *BatchLightNameWithLevel);
-
-									RenderRayTracingShadows(
-										GraphBuilder,
-										SceneTextures,
-										View,
-										BatchLightSceneInfo,
-										BatchRayTracingConfig,
-										IScreenSpaceDenoiser::EShadowRequirements::Bailout,
-										true,
-										&HairResources,
-										SubPixelRayTracingShadowMaskUAV,
-										SubPixelRayHitDistanceUAV);
-
-									TRefCountPtr<IPooledRenderTarget>* RefDestination = &PreprocessedShadowMaskSubPixelTextures[LightBatchIndex - AttenuationLightStart];
-									check(*RefDestination == nullptr);
-
-									GraphBuilder.QueueTextureExtraction(SubPixelRayTracingShadowMaskTexture, RefDestination);
 								}
 
 								bool bBatchFull = false;
@@ -1546,8 +1519,9 @@ void FDeferredShadingSceneRenderer::RenderLights(FRHICommandListImmediate& RHICm
 						FRDGTextureUAV* RayTracingShadowMaskUAV = GraphBuilder.CreateUAV(FRDGTextureUAVDesc(RayTracingShadowMaskTexture));
 						FRDGTextureUAV* RayHitDistanceUAV = GraphBuilder.CreateUAV(FRDGTextureUAVDesc(RayDistanceTexture));
 
-
-						FRDGTextureRef SubPixelRayTracingShadowMaskTexture;
+						FRDGTextureRef SubPixelRayTracingShadowMaskTexture = nullptr;
+						FRDGTextureUAV* SubPixelRayTracingShadowMaskUAV = nullptr;
+						if (bUseHairLighting)
 						{
 							FRDGTextureDesc Desc = FRDGTextureDesc::Create2DDesc(
 								SceneTextures.SceneDepthBuffer->Desc.Extent,
@@ -1557,22 +1531,9 @@ void FDeferredShadingSceneRenderer::RenderLights(FRHICommandListImmediate& RHICm
 								TexCreate_ShaderResource | TexCreate_RenderTargetable | TexCreate_UAV,
 								/* bInForceSeparateTargetAndShaderResource = */ false);
 							SubPixelRayTracingShadowMaskTexture = GraphBuilder.CreateTexture(Desc, TEXT("RayTracingOcclusion"));
+							SubPixelRayTracingShadowMaskUAV = GraphBuilder.CreateUAV(FRDGTextureUAVDesc(SubPixelRayTracingShadowMaskTexture));
 						}
 
-						FRDGTextureRef SubPixelRayDistanceTexture;
-						{
-							FRDGTextureDesc Desc = FRDGTextureDesc::Create2DDesc(
-								SceneTextures.SceneDepthBuffer->Desc.Extent,
-								PF_R16F,
-								FClearValueBinding::Black,
-								TexCreate_None,
-								TexCreate_ShaderResource | TexCreate_RenderTargetable | TexCreate_UAV,
-								/* bInForceSeparateTargetAndShaderResource = */ false);
-							SubPixelRayDistanceTexture = GraphBuilder.CreateTexture(Desc, TEXT("RayTracingOcclusionDistance"));
-						}
-
-						FRDGTextureUAV* SubPixelRayTracingShadowMaskUAV = GraphBuilder.CreateUAV(FRDGTextureUAVDesc(SubPixelRayTracingShadowMaskTexture));
-						FRDGTextureUAV* SubPixelRayHitDistanceUAV = GraphBuilder.CreateUAV(FRDGTextureUAVDesc(SubPixelRayDistanceTexture));
 
 						FRDGTextureRef RayTracingShadowMaskTileTexture;
 						{
@@ -1609,10 +1570,10 @@ void FDeferredShadingSceneRenderer::RenderLights(FRHICommandListImmediate& RHICm
 								LightSceneInfo,
 								RayTracingConfig,
 								DenoiserRequirements,
-								false,
 								&HairResources,
 								RayTracingShadowMaskUAV,
-								RayHitDistanceUAV);
+								RayHitDistanceUAV,
+								SubPixelRayTracingShadowMaskUAV);
 
 							if (DenoiserRequirements != IScreenSpaceDenoiser::EShadowRequirements::Bailout)
 							{
@@ -1656,22 +1617,8 @@ void FDeferredShadingSceneRenderer::RenderLights(FRHICommandListImmediate& RHICm
 								GraphBuilder.QueueTextureExtraction(RayTracingShadowMaskTexture, &ScreenShadowMaskTexture);
 							}
 
-							// Ray trace the shadow cast by opaque geometries on to hair strands geometries
-							// Note: No denoiser is required on this output, as the hair strands are geometrically noisy, and so hard to denoise
 							if (HasHairStrandsClusters(ViewIndex))
 							{
-								RenderRayTracingShadows(
-									GraphBuilder,
-									SceneTextures,
-									View,
-									LightSceneInfo,
-									RayTracingConfig,
-									IScreenSpaceDenoiser::EShadowRequirements::Bailout,
-									true,
-									&HairResources,
-									SubPixelRayTracingShadowMaskUAV,
-									SubPixelRayHitDistanceUAV);
-
 								GraphBuilder.QueueTextureExtraction(SubPixelRayTracingShadowMaskTexture, &ScreenShadowMaskSubPixelTexture);
 							}
 						}
diff --git a/Engine/Source/Runtime/Renderer/Private/RayTracing/RayTracingShadows.cpp b/Engine/Source/Runtime/Renderer/Private/RayTracing/RayTracingShadows.cpp
--- a/Engine/Source/Runtime/Renderer/Private/RayTracing/RayTracingShadows.cpp
+++ b/Engine/Source/Runtime/Renderer/Private/RayTracing/RayTracingShadows.cpp
@@ -90,7 +90,6 @@ class FOcclusionRGS : public FGlobalShader
 		SHADER_PARAMETER(uint32, SamplesPerPixel)
 		SHADER_PARAMETER(float, NormalBias)
 		SHADER_PARAMETER(uint32, LightingChannelMask)
-		SHADER_PARAMETER(uint32, ShadowMaskType)
 		SHADER_PARAMETER(FIntRect, LightScissor)
 		SHADER_PARAMETER(FIntPoint, PixelOffset)
 		SHADER_PARAMETER(uint32, bUseHairVoxel)
@@ -104,6 +103,7 @@ class FOcclusionRGS : public FGlobalShader
 		SHADER_PARAMETER_SRV(RaytracingAccelerationStructure, TLAS)
 		SHADER_PARAMETER_RDG_TEXTURE_UAV(RWTexture2D<float4>, RWOcclusionMaskUAV)
 		SHADER_PARAMETER_RDG_TEXTURE_UAV(RWTexture2D<float>, RWRayDistanceUAV)
+		SHADER_PARAMETER_RDG_TEXTURE_UAV(RWTexture2D<float4>, RWSubPixelOcclusionMaskUAV)
 		SHADER_PARAMETER_STRUCT_REF(FViewUniformShaderParameters, ViewUniformBuffer)
 		SHADER_PARAMETER_STRUCT_REF(FVirtualVoxelParameters, VirtualVoxel)
 	END_SHADER_PARAMETER_STRUCT()
@@ -163,10 +163,10 @@ void FDeferredShadingSceneRenderer::RenderRayTracingShadows(
 	const FLightSceneInfo& LightSceneInfo,
 	const IScreenSpaceDenoiser::FShadowRayTracingConfig& RayTracingConfig,
 	const IScreenSpaceDenoiser::EShadowRequirements DenoiserRequirements,
-	const bool bSubPixelShadowMask,
 	const FHairStrandsOcclusionResources* HairResources,
 	FRDGTextureUAV* OutShadowMaskUAV,
-	FRDGTextureUAV* OutRayHitDistanceUAV)
+	FRDGTextureUAV* OutRayHitDistanceUAV,
+	FRDGTextureUAV* SubPixelRayTracingShadowMaskUAV)
 #if RHI_RAYTRACING
 {
 	FLightSceneProxy* LightSceneProxy = LightSceneInfo.Proxy;
@@ -208,6 +208,7 @@ void FDeferredShadingSceneRenderer::RenderRayTracingShadows(
 		FOcclusionRGS::FParameters* PassParameters = GraphBuilder.AllocParameters<FOcclusionRGS::FParameters>();
 		PassParameters->RWOcclusionMaskUAV = OutShadowMaskUAV;
 		PassParameters->RWRayDistanceUAV = OutRayHitDistanceUAV;
+		PassParameters->RWSubPixelOcclusionMaskUAV = SubPixelRayTracingShadowMaskUAV;
 		PassParameters->SamplesPerPixel = RayTracingConfig.RayCountPerPixel;
 		PassParameters->NormalBias = GetRaytracingMaxNormalBias();
 		PassParameters->LightingChannelMask = LightSceneProxy->GetLightingChannelMask();
@@ -217,7 +218,6 @@ void FDeferredShadingSceneRenderer::RenderRayTracingShadows(
 		PassParameters->SceneTextures = SceneTextures;
 		PassParameters->LightScissor = ScissorRect;
 		PassParameters->PixelOffset = PixelOffset;
-		PassParameters->ShadowMaskType = bUseHairLighting && bSubPixelShadowMask ? 1 : 0;
 		PassParameters->SSProfilesTexture = GraphBuilder.RegisterExternalTexture(View.RayTracingSubSurfaceProfileTexture);
 		if (bUseHairLighting)
 		{
