Improve aliasing on hair shadow.

diff --git a/Engine/Shaders/Private/RayTracing/RayTracingOcclusionRGS.usf b/Engine/Shaders/Private/RayTracing/RayTracingOcclusionRGS.usf
--- a/Engine/Shaders/Private/RayTracing/RayTracingOcclusionRGS.usf
+++ b/Engine/Shaders/Private/RayTracing/RayTracingOcclusionRGS.usf
@@ -51,16 +51,22 @@ uint				bUseHairVoxel;
 Texture2D<uint4>	HairCategorizationTexture;
 Texture2D<uint>		HairLightChannelMaskTexture;
 
-float TraverseHair(uint2 PixelCoord, float3 RayOrigin, float3 RayDirection, float InMinT)
+float TraverseHair(uint2 PixelCoord, inout RandomSequence RandSequence, float3 RayOrigin, float3 RayDirection, float InMinT)
 {
 	float HitT = InMinT;
 
 	// Depth bias
-	// Origin is shifted 2 voxels away towards the light + a constant bias of the size of the voxel
-	const float3 Random = 0.5;// GetHairVoxelJitter(PixelCoord.xy, View.StateFrameIndexMod8);
+	// Origin is shifted voxels away towards the light + a constant bias of the size of the voxel
+#if 0
+	const float3 VoxelRandom = 0.5f;
+	const float  RayRandom   = 1;
+#else
+	uint DummyVariable;
+	const float3 VoxelRandom = RandomSequence_GenerateSample3D(RandSequence, DummyVariable);
+	const float RayRandom = RandomSequence_GenerateSample1D(RandSequence, DummyVariable);
+#endif
 	{
-		const float MaxVoxelSize = VirtualVoxel.VoxelWorldSize;
-		const float3 DepthBias = MaxVoxelSize * (2 * RayDirection + Random);
+		const float3 DepthBias = VirtualVoxel.VoxelWorldSize * (RayRandom * RayDirection + VoxelRandom);
 		RayOrigin += DepthBias;
 	}
 
@@ -95,7 +101,7 @@ float TraverseHair(uint2 PixelCoord, float3 RayOrigin, float3 RayDirection, floa
 		TraversalSettings.DistanceThreshold = DistanceThreshold;
 		TraversalSettings.bDebugEnabled = bDebugEnabled;
 		TraversalSettings.SteppingScale = VirtualVoxel.SteppingScale;
-		TraversalSettings.Random = Random;
+		TraversalSettings.Random = VoxelRandom;
 		TraversalSettings.bUseOpaqueVisibility = false;
 
 		Result = ComputeHairCountVirtualVoxel(
@@ -446,7 +452,7 @@ FOcclusionResult ComputeOcclusion(
 		#if USE_HAIR_LIGHTING
 		if (ShadingModelID != SHADINGMODELID_HAIR && bUseHairVoxel)
 		{
-			MinimalPayload.HitT = TraverseHair(PixelCoord, Ray.Origin, Ray.Direction, MinimalPayload.HitT);
+			MinimalPayload.HitT = TraverseHair(PixelCoord, RandSequence, Ray.Origin, Ray.Direction, MinimalPayload.HitT);
 		}
 		#endif
 
