Add debug print for tracking the number of lit hair samples.

diff --git a/Engine/Shaders/Private/HairStrands/HairStrandsDebugPrint.usf b/Engine/Shaders/Private/HairStrands/HairStrandsDebugPrint.usf
--- a/Engine/Shaders/Private/HairStrands/HairStrandsDebugPrint.usf
+++ b/Engine/Shaders/Private/HairStrands/HairStrandsDebugPrint.usf
@@ -13,6 +13,7 @@ uint FastResolveMask;
 uint HairMacroGroupCount;
 uint MaxSampleCount;
 uint2 GroupSize; 
+uint HairVisibilityNodeGroupSize;
 
 Texture2D<float> HairCountTexture;
 Texture2D<uint> HairCountUintTexture;
@@ -21,6 +22,7 @@ Texture2D<uint2> DepthStencilTexture;
 SamplerState LinearSampler;
 Texture2D<uint> HairVisibilityNodeOffsetAndCount;
 StructuredBuffer<FPackedHairSample> HairVisibilityNodeData;
+Buffer<uint> HairVisibilityIndirectArgsBuffer;
 
 Buffer<int> HairMacroGroupAABBBuffer;
 
@@ -39,6 +41,56 @@ void MainCS(uint GroupIndex : SV_GroupIndex, uint3 DispatchThreadId : SV_Dispatc
 			FHairAABB Bound = ReadHairAABB(GroupIndex, HairMacroGroupAABBBuffer);
 			AddAABB(Bound.Min, Bound.Max, InstanceBoundColor);
 		}
+
+
+		const uint DispatchX = HairVisibilityIndirectArgsBuffer[0];
+		const uint DispatchY = HairVisibilityIndirectArgsBuffer[1];
+		const uint TotalSampleCount = DispatchX * DispatchY * HairVisibilityNodeGroupSize;
+		const uint2 EquivalentResolution = ceil(sqrt(TotalSampleCount));
+
+		// Pixel coord
+		float2 OriginalPos = float2(50, 50) / float2(MaxResolution);
+		float2 Pos = OriginalPos;
+		Pos = ShaderPrintSymbol(Pos, _T_);
+		Pos = ShaderPrintSymbol(Pos, _O_);
+		Pos = ShaderPrintSymbol(Pos, _T_);
+		Pos = ShaderPrintSymbol(Pos, _A_);
+		Pos = ShaderPrintSymbol(Pos, _L_);
+		Pos = ShaderPrintSymbol(Pos, _SPC_);
+		Pos = ShaderPrintSymbol(Pos, _S_);
+		Pos = ShaderPrintSymbol(Pos, _A_);
+		Pos = ShaderPrintSymbol(Pos, _M_);
+		Pos = ShaderPrintSymbol(Pos, _P_);
+		Pos = ShaderPrintSymbol(Pos, _L_);
+		Pos = ShaderPrintSymbol(Pos, _E_);
+		Pos = ShaderPrintSymbol(Pos, _SPC_);
+		Pos = ShaderPrintSymbol(Pos, _C_);
+		Pos = ShaderPrintSymbol(Pos, _O_);
+		Pos = ShaderPrintSymbol(Pos, _U_);
+		Pos = ShaderPrintSymbol(Pos, _N_);
+		Pos = ShaderPrintSymbol(Pos, _T_);
+		Pos = ShaderPrintSymbol(Pos, _SPC_);
+		Pos = ShaderPrint(Pos, TotalSampleCount);
+
+		Pos = ShaderPrintNewline(Pos);
+		Pos.x = OriginalPos.x;
+
+		Pos = ShaderPrintSymbol(Pos, _E_);
+		Pos = ShaderPrintSymbol(Pos, _Q_);
+		Pos = ShaderPrintSymbol(Pos, _U_);
+		Pos = ShaderPrintSymbol(Pos, _I_);
+		Pos = ShaderPrintSymbol(Pos, _V_);
+		Pos = ShaderPrintSymbol(Pos, _A_);
+		Pos = ShaderPrintSymbol(Pos, _L_);
+		Pos = ShaderPrintSymbol(Pos, _E_);
+		Pos = ShaderPrintSymbol(Pos, _N_);
+		Pos = ShaderPrintSymbol(Pos, _T_);
+		Pos = ShaderPrintSymbol(Pos, _SPC_);
+		Pos = ShaderPrintSymbol(Pos, _R_);
+		Pos = ShaderPrintSymbol(Pos, _E_);
+		Pos = ShaderPrintSymbol(Pos, _S_);
+		Pos = ShaderPrintSymbol(Pos, _SPC_);
+		Pos = ShaderPrint(Pos, EquivalentResolution);
 	}
 
 	const float2 PixelPos = PixelCoord;
diff --git a/Engine/Source/Runtime/Renderer/Private/HairStrands/HairStrandsDebug.cpp b/Engine/Source/Runtime/Renderer/Private/HairStrands/HairStrandsDebug.cpp
--- a/Engine/Source/Runtime/Renderer/Private/HairStrands/HairStrandsDebug.cpp
+++ b/Engine/Source/Runtime/Renderer/Private/HairStrands/HairStrandsDebug.cpp
@@ -242,11 +242,13 @@ class FHairDebugPrintCS : public FGlobalShader
 		SHADER_PARAMETER(uint32, FastResolveMask)
 		SHADER_PARAMETER(uint32, HairMacroGroupCount)
 		SHADER_PARAMETER(uint32, MaxSampleCount)
+		SHADER_PARAMETER(uint32, HairVisibilityNodeGroupSize)
 		SHADER_PARAMETER_RDG_TEXTURE(Texture2D, HairCountTexture)
 		SHADER_PARAMETER_RDG_TEXTURE(Texture2D, HairCountUintTexture)
 		SHADER_PARAMETER_RDG_TEXTURE(Texture2D, CategorizationTexture)
 		SHADER_PARAMETER_RDG_TEXTURE(Texture2D, HairVisibilityNodeOffsetAndCount)
 		SHADER_PARAMETER_RDG_BUFFER_SRV(StructuredBuffer, HairVisibilityNodeData)
+		SHADER_PARAMETER_RDG_BUFFER_SRV(Buffer<uint>, HairVisibilityIndirectArgsBuffer)
 		SHADER_PARAMETER_RDG_BUFFER_SRV(Buffer, HairMacroGroupAABBBuffer)
 		SHADER_PARAMETER_SRV(Texture2D, DepthStencilTexture)
 		SHADER_PARAMETER_SAMPLER(SamplerState, LinearSampler)
@@ -296,6 +298,8 @@ static void AddDebugHairPrintPass(
 	Parameters->HairCountUintTexture = ViewHairCountUintTexture;
 	Parameters->HairVisibilityNodeData = GraphBuilder.CreateSRV(NodeData);
 	Parameters->HairVisibilityNodeOffsetAndCount = NodeIndex;
+	Parameters->HairVisibilityIndirectArgsBuffer = GraphBuilder.CreateSRV(GraphBuilder.RegisterExternalBuffer(VisibilityData.NodeIndirectArg), PF_R32_UINT);
+	Parameters->HairVisibilityNodeGroupSize = VisibilityData.NodeGroupSize;
 	Parameters->DepthStencilTexture = InDepthStencilTexture;
 	Parameters->LinearSampler = TStaticSamplerState<SF_Bilinear, AM_Clamp, AM_Clamp, AM_Clamp>::GetRHI();
 	Parameters->HairMacroGroupCount = MacroGroupDatas.Datas.Num();
