Add hair coverage threshold for marking hair pixel as fully opaque after a certain coverage.

This is necessary, as the transmittance is asymptotically converging toward 0, but never reaches it. So the coverage is almost never 1.

diff --git a/Engine/Shaders/Private/HairStrands/HairStrandsVisibilityCompaction.usf b/Engine/Shaders/Private/HairStrands/HairStrandsVisibilityCompaction.usf
--- a/Engine/Shaders/Private/HairStrands/HairStrandsVisibilityCompaction.usf
+++ b/Engine/Shaders/Private/HairStrands/HairStrandsVisibilityCompaction.usf
@@ -10,11 +10,6 @@
 #define VENDOR_NVIDIA 1
 #define VENDOR_OTHER 2
 
-#if PERMUTATION_MATERIAL_COMPACTION == 1
-float DepthTheshold;
-float CosTangentThreshold;
-#endif
-
 #if PERMUTATION_VENDOR == VENDOR_AMD
 #define TILE_PIXEL_SIZE_X 8
 #define TILE_PIXEL_SIZE_Y 8
@@ -26,9 +21,12 @@ float CosTangentThreshold;
 #define TILE_PIXEL_SIZE_Y 8
 #endif
 
-uint2 OutputResolution;
+int2 OutputResolution;
+int2 ResolutionOffset;
 uint MaxNodeCount;
-uint2 ResolutionOffset;
+float DepthTheshold; 
+float CosTangentThreshold;
+float CoverageThreshold; // Allow to scale the transmittance to compensate its asymptotic behavior 
 
 RWTexture2D<uint> OutCompactNodeCounter;
 RWTexture2D<uint> OutCompactNodeIndex;
@@ -52,6 +50,10 @@ RWTexture2D<uint4> OutCategorizationTexture;
 	RWTexture2D<float2> OutVelocityTexture;
 #endif
 
+float TransmittanceToCoverage(float InTransmittance)
+{
+	return saturate(min(1 - InTransmittance, 1) / CoverageThreshold);
+}
 
 groupshared uint AllocationNodeCount;
 groupshared uint AllocationNodeOffset;
@@ -87,7 +89,7 @@ void MainCS(uint GroupIndex : SV_GroupIndex, uint3 DispatchThreadId : SV_Dispatc
 	float MaxVelocityMagnitude2 = 0;
 #endif
 	
-	const bool bExecute = FirstNodeIndex != 0xFFFFFFFF && all(PixelCoord < OutputResolution);
+	const bool bExecute = FirstNodeIndex != 0xFFFFFFFF && all(PixelCoord < uint2(OutputResolution));
 
 	uint SortedIndex[PPLL_MAX_RENDER_NODE_COUNT];
 	float SortedDepth[PPLL_MAX_RENDER_NODE_COUNT];
@@ -238,7 +240,7 @@ void MainCS(uint GroupIndex : SV_GroupIndex, uint3 DispatchThreadId : SV_Dispatc
 		#endif		
 		}
 
-		Categorization.PixelCoverage = 1.0f - TotalTransmittance;
+		Categorization.PixelCoverage = TransmittanceToCoverage(TotalTransmittance);
 
 	#if PERMUTATION_VELOCITY == VELOCITY_TYPE_AVG
 		OutVelocityTexture[PixelCoord] = EncodeVelocityToTexture(AverageVelocity / max(ValidVelocityCount, 1));
@@ -354,7 +356,7 @@ void MainCS(uint GroupIndex : SV_GroupIndex, uint3 DispatchThreadId : SV_Dispatc
 	}
 
 	uint2 PixelCoord = DispatchThreadId.xy + ResolutionOffset;
-	if (PixelCoord.x >= OutputResolution.x || PixelCoord.y >= OutputResolution.y)
+	if (PixelCoord.x >= uint(OutputResolution.x) || PixelCoord.y >= uint(OutputResolution.y))
 		PixelCoord = uint2(0, 0);
 
 	FSampleSetDesc SampleDesc;
@@ -415,7 +417,7 @@ void MainCS(uint GroupIndex : SV_GroupIndex, uint3 DispatchThreadId : SV_Dispatc
 	}
 
 #if PERMUTATION_VIEWTRANSMITTANCE == 1
-	Categorization.PixelCoverage = 1.0f - ViewTransmittanceTexture.Load(uint3(PixelCoord, 0));
+	Categorization.PixelCoverage = TransmittanceToCoverage(ViewTransmittanceTexture.Load(uint3(PixelCoord, 0)));
 #else
 	Categorization.PixelCoverage = SampleDesc.ValidSampleCount / float(PERMUTATION_MSAACOUNT);
 #endif
@@ -533,7 +535,7 @@ void MainCS(uint GroupIndex : SV_GroupIndex, uint3 DispatchThreadId : SV_Dispatc
 	}
 
 	uint2 PixelCoord = DispatchThreadId.xy + ResolutionOffset;
-	if (PixelCoord.x >= OutputResolution.x || PixelCoord.y >= OutputResolution.y)
+	if (PixelCoord.x >= uint(OutputResolution.x) || PixelCoord.y >= uint(OutputResolution.y))
 		PixelCoord = uint2(0, 0);
 
 	FSampleSetDesc SampleDesc;
@@ -639,7 +641,7 @@ void MainCS(uint GroupIndex : SV_GroupIndex, uint3 DispatchThreadId : SV_Dispatc
 			}
 
 			#if PERMUTATION_VIEWTRANSMITTANCE == 1
-			Categorization.PixelCoverage = 1.0f - ViewTransmittanceTexture.Load(uint3(PixelCoord,0));
+			Categorization.PixelCoverage = TransmittanceToCoverage(ViewTransmittanceTexture.Load(uint3(PixelCoord,0)));
 			#else
 			Categorization.PixelCoverage = SampleDesc.ValidSampleCount / float(PERMUTATION_MSAACOUNT);
 			#endif
diff --git a/Engine/Source/Runtime/Renderer/Private/HairStrands/HairStrandsVisibility.cpp b/Engine/Source/Runtime/Renderer/Private/HairStrands/HairStrandsVisibility.cpp
--- a/Engine/Source/Runtime/Renderer/Private/HairStrands/HairStrandsVisibility.cpp
+++ b/Engine/Source/Runtime/Renderer/Private/HairStrands/HairStrandsVisibility.cpp
@@ -63,6 +63,9 @@ static FAutoConsoleVariableRef CVarHairStrandsVisibilityComputeRaster(TEXT("r.Ha
 static FAutoConsoleVariableRef CVarHairStrandsVisibilityComputeRasterMaxPixelCount(TEXT("r.HairStrands.Visibility.ComputeRaster.MaxPixelCount"), GHairStrandsVisibilityComputeRasterMaxPixelCount, TEXT("Define the maximal length rasterize in compute."));
 static FAutoConsoleVariableRef CVarHairStrandsVisibilityComputeRasterSampleCount(TEXT("r.HairStrands.Visibility.ComputeRaster.SampleCount"), GHairStrandsVisibilityComputeRasterSampleCount, TEXT("Define sample count used in rasterize in compute."));
 
+static float GHairStrandsFullCoverageThreshold = 0.98f;
+static FAutoConsoleVariableRef CVarHairStrandsFullCoverageThreshold(TEXT("r.HairStrands.Visibility.FullCoverageThreshold"), GHairStrandsFullCoverageThreshold, TEXT("Define the coverage threshold at which a pixel is considered fully covered."));
+
 /////////////////////////////////////////////////////////////////////////////////////////
 
 namespace HairStrandsVisibilityInternal
@@ -1124,11 +1127,11 @@ class FHairVisibilityPrimitiveIdCompactionCS : public FGlobalShader
 
 	BEGIN_SHADER_PARAMETER_STRUCT(FParameters, )
 		SHADER_PARAMETER(FIntPoint, OutputResolution)
-		SHADER_PARAMETER(uint32, MaxNodeCount)
-		SHADER_PARAMETER(uint32, MSAASampleCount)
 		SHADER_PARAMETER(FIntPoint, ResolutionOffset)
+		SHADER_PARAMETER(uint32, MaxNodeCount)
 		SHADER_PARAMETER(float, DepthTheshold)
 		SHADER_PARAMETER(float, CosTangentThreshold)
+		SHADER_PARAMETER(float, CoverageThreshold)
 
 		// Available for the MSAA path
 		SHADER_PARAMETER_RDG_TEXTURE(Texture2D, MSAA_DepthTexture)
@@ -1303,6 +1306,7 @@ static void AddHairVisibilityPrimitiveIdCompactionPass(
 
 	PassParameters->OutputResolution = Resolution;
 	PassParameters->MaxNodeCount = MaxRenderNodeCount;
+	PassParameters->CoverageThreshold = FMath::Clamp(GHairStrandsFullCoverageThreshold, 0.1f, 1.f);
 	PassParameters->DepthTheshold = FMath::Clamp(GHairStrandsMaterialCompactionDepthThreshold, 0.f, 100.f);
 	PassParameters->CosTangentThreshold = FMath::Cos(FMath::DegreesToRadians(FMath::Clamp(GHairStrandsMaterialCompactionTangentThreshold, 0.f, 90.f)));
 	PassParameters->SceneTexturesStruct = CreateUniformBufferImmediate(SceneTextures, EUniformBufferUsage::UniformBuffer_SingleDraw);
