Fixed shader compile errors due to mismatched parameter types for clamp().

diff --git a/Engine/Shaders/Private/HairStrands/HairStrandsVoxelPageTraversal.ush b/Engine/Shaders/Private/HairStrands/HairStrandsVoxelPageTraversal.ush
--- a/Engine/Shaders/Private/HairStrands/HairStrandsVoxelPageTraversal.ush
+++ b/Engine/Shaders/Private/HairStrands/HairStrandsVoxelPageTraversal.ush
@@ -430,7 +430,7 @@ FHairTraversalResult ComputeHairCountVirtualVoxel(
 			const float3 FractCoords = max(Step, 0) - Step * frac(Coords);
 			tMax = FractCoords * tDelta;
 
-			PageIndexCoord = clamp(int3(Coords), uint3(0, 0, 0), InNodeDesc.PageIndexResolution-1);
+			PageIndexCoord = clamp(uint3(Coords), uint3(0, 0, 0), InNodeDesc.PageIndexResolution-1);
 		}
 
 		// Page stepping is the walking quantity (i.e. number of voxel) for fine ray-marching within a valid page
