Fixed shader compilation errors in Vulkan

diff --git a/Engine/Shaders/Private/HairStrands/HairStrandsDeepTransmittanceCommon.ush b/Engine/Shaders/Private/HairStrands/HairStrandsDeepTransmittanceCommon.ush
--- a/Engine/Shaders/Private/HairStrands/HairStrandsDeepTransmittanceCommon.ush
+++ b/Engine/Shaders/Private/HairStrands/HairStrandsDeepTransmittanceCommon.ush
@@ -38,7 +38,7 @@ FPackedHairTransmittanceMask PackTransmittanceMask(FHairTransmittanceMask In)
 {
 	FPackedHairTransmittanceMask Out;
 	Out.Transmittance = In.Transmittance;
-	Out.PackedHairCountAndVisibility = min(uint(In.HairCount * 1000), 0x00FFFFFF) | (min(uint(In.Visibility * 0xFF), 0xFF) << 24);
+	Out.PackedHairCountAndVisibility = min(uint(In.HairCount * 1000), uint(0x00FFFFFF)) | (min(uint(In.Visibility * 0xFF), uint(0xFF)) << 24);
 	return Out;
 }
 
diff --git a/Engine/Shaders/Private/HairStrands/HairStrandsEnvironmentLighting.usf b/Engine/Shaders/Private/HairStrands/HairStrandsEnvironmentLighting.usf
--- a/Engine/Shaders/Private/HairStrands/HairStrandsEnvironmentLighting.usf
+++ b/Engine/Shaders/Private/HairStrands/HairStrandsEnvironmentLighting.usf
@@ -241,9 +241,9 @@ float3 ReflectionEnvironment(FGBufferData GBuffer, float2 ScreenPosition, float2
 			float2(0.19984126, 0.78641367),
 			float2(0.14383161, -0.14100790)
 		};
-		const uint SampleCount = clamp(MultipleScatterSampleCount, 1, 16);
+		const uint SampleCount = clamp(MultipleScatterSampleCount, uint(1), uint(16));
 		#else
-		const uint SampleCount = max(MultipleScatterSampleCount, 1);
+		const uint SampleCount = max(MultipleScatterSampleCount, uint(1));
 		#endif
 
 		float2 RandOffset = 0;
