Fixing LPVs for pre-exposure.

diff --git a/Engine/Shaders/Private/PostProcessLpvIndirect.usf b/Engine/Shaders/Private/PostProcessLpvIndirect.usf
--- a/Engine/Shaders/Private/PostProcessLpvIndirect.usf
+++ b/Engine/Shaders/Private/PostProcessLpvIndirect.usf
@@ -191,6 +191,13 @@ void MainPS(in noperspective float4 UVAndScreenPos : TEXCOORD0, float4 SVPos: SV
 #endif 
 #endif
 
+#if USE_PREEXPOSURE
+	OutColor.rgb *= View.PreExposure;
+#if APPLY_SEPARATE_SPECULAR_RT
+	OutSpecular.rgb *= View.PreExposure;
+#endif 
+#endif
+
 #if DEBUG_SHOW_EDGES
 	float3 gridSmooth = WorldToGridSmooth( WorldPosition );
 	float gridExtentSmooth = max3(abs(gridSmooth-float3(16,16,16)));
