Fix crash in hair system when a skin cache entry contains several sections not using the sharing position buffer.

diff --git a/Engine/Shaders/Private/HairStrands/HairStrandsMeshUpdate.usf b/Engine/Shaders/Private/HairStrands/HairStrandsMeshUpdate.usf
--- a/Engine/Shaders/Private/HairStrands/HairStrandsMeshUpdate.usf
+++ b/Engine/Shaders/Private/HairStrands/HairStrandsMeshUpdate.usf
@@ -13,12 +13,36 @@ uint MeshMaxVertexCount		[MAX_SECTION_COUNT];
 uint MeshIndexOffset		[MAX_SECTION_COUNT];
 uint MeshUVsChannelOffset	[MAX_SECTION_COUNT];
 uint MeshUVsChannelCount	[MAX_SECTION_COUNT];
+uint MeshSectionBufferIndex	[MAX_SECTION_COUNT];
 
-Buffer<uint> MeshIndexBuffer;
-Buffer<float> MeshPositionBuffer;
+Buffer<uint> MeshIndexBuffer0;
+Buffer<uint> MeshIndexBuffer1;
+Buffer<uint> MeshIndexBuffer2;
+Buffer<uint> MeshIndexBuffer3;
+Buffer<uint> MeshIndexBuffer4;
+Buffer<uint> MeshIndexBuffer5;
+Buffer<uint> MeshIndexBuffer6;
+Buffer<uint> MeshIndexBuffer7;
+
+Buffer<float> MeshPositionBuffer0;
+Buffer<float> MeshPositionBuffer1;
+Buffer<float> MeshPositionBuffer2;
+Buffer<float> MeshPositionBuffer3;
+Buffer<float> MeshPositionBuffer4;
+Buffer<float> MeshPositionBuffer5;
+Buffer<float> MeshPositionBuffer6;
+Buffer<float> MeshPositionBuffer7;
 
 #if PERMUTATION_WITHUV == 1
-Buffer<float2> MeshUVsBuffer;
+Buffer<float2> MeshUVsBuffer0;
+Buffer<float2> MeshUVsBuffer1;
+Buffer<float2> MeshUVsBuffer2;
+Buffer<float2> MeshUVsBuffer3;
+Buffer<float2> MeshUVsBuffer4;
+Buffer<float2> MeshUVsBuffer5;
+Buffer<float2> MeshUVsBuffer6;
+Buffer<float2> MeshUVsBuffer7;
+Buffer<float2> MeshUVsBuffer8;
 #endif
 
 uint MaxRootCount;
@@ -28,6 +52,85 @@ RWBuffer<float4> OutRootTrianglePosition0;
 RWBuffer<float4> OutRootTrianglePosition1;
 RWBuffer<float4> OutRootTrianglePosition2;
 
+#if PERMUTATION_WITHUV == 1
+FHairMeshTriangleWithUVs GetTriangleData(FHairTriangleIndex TriangleIndex, uint SectionIndex)
+{
+	const uint MeshSectionIndex = MeshSectionBufferIndex[SectionIndex];
+	if (MeshSectionIndex == 0)
+	{
+		return GetTriangleData(TriangleIndex.TriangleIndex, MeshPositionBuffer0, MeshUVsBuffer0, MeshIndexBuffer0, MeshIndexOffset[SectionIndex], MeshMaxIndexCount[SectionIndex], MeshMaxVertexCount[SectionIndex], MeshUVsChannelOffset[SectionIndex], MeshUVsChannelCount[SectionIndex]);
+	}
+	else if (MeshSectionIndex == 1)
+	{
+		return GetTriangleData(TriangleIndex.TriangleIndex, MeshPositionBuffer1, MeshUVsBuffer1, MeshIndexBuffer1, MeshIndexOffset[SectionIndex], MeshMaxIndexCount[SectionIndex], MeshMaxVertexCount[SectionIndex], MeshUVsChannelOffset[SectionIndex], MeshUVsChannelCount[SectionIndex]);
+	}
+	else if (MeshSectionIndex == 2)
+	{
+		return GetTriangleData(TriangleIndex.TriangleIndex, MeshPositionBuffer2, MeshUVsBuffer2, MeshIndexBuffer2, MeshIndexOffset[SectionIndex], MeshMaxIndexCount[SectionIndex], MeshMaxVertexCount[SectionIndex], MeshUVsChannelOffset[SectionIndex], MeshUVsChannelCount[SectionIndex]);
+	}
+	else if (MeshSectionIndex == 3)
+	{
+		return GetTriangleData(TriangleIndex.TriangleIndex, MeshPositionBuffer3, MeshUVsBuffer3, MeshIndexBuffer3, MeshIndexOffset[SectionIndex], MeshMaxIndexCount[SectionIndex], MeshMaxVertexCount[SectionIndex], MeshUVsChannelOffset[SectionIndex], MeshUVsChannelCount[SectionIndex]);
+	}
+	else if (MeshSectionIndex == 4)
+	{
+		return GetTriangleData(TriangleIndex.TriangleIndex, MeshPositionBuffer4, MeshUVsBuffer4, MeshIndexBuffer4, MeshIndexOffset[SectionIndex], MeshMaxIndexCount[SectionIndex], MeshMaxVertexCount[SectionIndex], MeshUVsChannelOffset[SectionIndex], MeshUVsChannelCount[SectionIndex]);
+	}
+	else if (MeshSectionIndex == 5)
+	{
+		return GetTriangleData(TriangleIndex.TriangleIndex, MeshPositionBuffer5, MeshUVsBuffer5, MeshIndexBuffer5, MeshIndexOffset[SectionIndex], MeshMaxIndexCount[SectionIndex], MeshMaxVertexCount[SectionIndex], MeshUVsChannelOffset[SectionIndex], MeshUVsChannelCount[SectionIndex]);
+	}
+	else if (MeshSectionIndex == 6)
+	{
+		return GetTriangleData(TriangleIndex.TriangleIndex, MeshPositionBuffer6, MeshUVsBuffer6, MeshIndexBuffer6, MeshIndexOffset[SectionIndex], MeshMaxIndexCount[SectionIndex], MeshMaxVertexCount[SectionIndex], MeshUVsChannelOffset[SectionIndex], MeshUVsChannelCount[SectionIndex]);
+	}
+	else if (MeshSectionIndex == 7)
+	{
+		return GetTriangleData(TriangleIndex.TriangleIndex, MeshPositionBuffer7, MeshUVsBuffer7, MeshIndexBuffer7, MeshIndexOffset[SectionIndex], MeshMaxIndexCount[SectionIndex], MeshMaxVertexCount[SectionIndex], MeshUVsChannelOffset[SectionIndex], MeshUVsChannelCount[SectionIndex]);
+	}
+	return (FHairMeshTriangleWithUVs)0;
+}
+#else
+FHairMeshTriangle GetTriangleData(FHairTriangleIndex TriangleIndex, uint SectionIndex)
+{
+	const uint MeshSectionIndex = MeshSectionBufferIndex[SectionIndex];
+	if (MeshSectionIndex == 0)
+	{
+		return GetTriangleData(TriangleIndex.TriangleIndex, MeshPositionBuffer0, MeshIndexBuffer0, MeshIndexOffset[SectionIndex], MeshMaxIndexCount[SectionIndex], MeshMaxVertexCount[SectionIndex]);
+	}
+	else if (MeshSectionIndex == 1)
+	{
+		return GetTriangleData(TriangleIndex.TriangleIndex, MeshPositionBuffer1, MeshIndexBuffer1, MeshIndexOffset[SectionIndex], MeshMaxIndexCount[SectionIndex], MeshMaxVertexCount[SectionIndex]);
+	}
+	else if (MeshSectionIndex == 2)
+	{
+		return GetTriangleData(TriangleIndex.TriangleIndex, MeshPositionBuffer2, MeshIndexBuffer2, MeshIndexOffset[SectionIndex], MeshMaxIndexCount[SectionIndex], MeshMaxVertexCount[SectionIndex]);
+	}
+	else if (MeshSectionIndex == 3)
+	{
+		return GetTriangleData(TriangleIndex.TriangleIndex, MeshPositionBuffer3, MeshIndexBuffer3, MeshIndexOffset[SectionIndex], MeshMaxIndexCount[SectionIndex], MeshMaxVertexCount[SectionIndex]);
+	}
+	else if (MeshSectionIndex == 4)
+	{
+		return GetTriangleData(TriangleIndex.TriangleIndex, MeshPositionBuffer4, MeshIndexBuffer4, MeshIndexOffset[SectionIndex], MeshMaxIndexCount[SectionIndex], MeshMaxVertexCount[SectionIndex]);
+	}
+	else if (MeshSectionIndex == 5)
+	{
+		return GetTriangleData(TriangleIndex.TriangleIndex, MeshPositionBuffer5, MeshIndexBuffer5, MeshIndexOffset[SectionIndex], MeshMaxIndexCount[SectionIndex], MeshMaxVertexCount[SectionIndex]);
+	}
+	else if (MeshSectionIndex == 6)
+	{
+		return GetTriangleData(TriangleIndex.TriangleIndex, MeshPositionBuffer6, MeshIndexBuffer6, MeshIndexOffset[SectionIndex], MeshMaxIndexCount[SectionIndex], MeshMaxVertexCount[SectionIndex]);
+	}
+	else if (MeshSectionIndex == 7)
+	{
+		return GetTriangleData(TriangleIndex.TriangleIndex, MeshPositionBuffer7, MeshIndexBuffer7, MeshIndexOffset[SectionIndex], MeshMaxIndexCount[SectionIndex], MeshMaxVertexCount[SectionIndex]);
+	}
+	return (FHairMeshTriangle)0;
+}
+
+#endif
+
 [numthreads(128, 1, 1)]
 void MainCS(uint3 DispatchThreadId : SV_DispatchThreadID)
 { 
@@ -42,13 +145,13 @@ void MainCS(uint3 DispatchThreadId : SV_DispatchThreadID)
 	const uint SectionIndex = TriangleIndex.SectionIndex;
 
 #if PERMUTATION_WITHUV == 1
-	const FHairMeshTriangleWithUVs Tri = GetTriangleData(TriangleIndex.TriangleIndex, MeshPositionBuffer, MeshUVsBuffer, MeshIndexBuffer, MeshIndexOffset[SectionIndex], MeshMaxIndexCount[SectionIndex], MeshMaxVertexCount[SectionIndex], MeshUVsChannelOffset[SectionIndex], MeshUVsChannelCount[SectionIndex]);
+	const FHairMeshTriangleWithUVs Tri = GetTriangleData(TriangleIndex, SectionIndex);
 
 	OutRootTrianglePosition0[RootIndex] = float4(Tri.P0, PackUVsToFloat(Tri.UV0));
 	OutRootTrianglePosition1[RootIndex] = float4(Tri.P1, PackUVsToFloat(Tri.UV1));
 	OutRootTrianglePosition2[RootIndex] = float4(Tri.P2, PackUVsToFloat(Tri.UV2));
 #else
-	const FHairMeshTriangle Tri = GetTriangleData(TriangleIndex.TriangleIndex, MeshPositionBuffer, MeshIndexBuffer, MeshIndexOffset[SectionIndex], MeshMaxIndexCount[SectionIndex], MeshMaxVertexCount[SectionIndex]);
+	const FHairMeshTriangle Tri = GetTriangleData(TriangleIndex, SectionIndex);
 
 	OutRootTrianglePosition0[RootIndex] = (float4(Tri.P0, 1));
 	OutRootTrianglePosition1[RootIndex] = (float4(Tri.P1, 1));
diff --git a/Engine/Source/Runtime/Renderer/Private/HairStrands/HairStrandsMeshProjection.cpp b/Engine/Source/Runtime/Renderer/Private/HairStrands/HairStrandsMeshProjection.cpp
--- a/Engine/Source/Runtime/Renderer/Private/HairStrands/HairStrandsMeshProjection.cpp
+++ b/Engine/Source/Runtime/Renderer/Private/HairStrands/HairStrandsMeshProjection.cpp
@@ -385,9 +385,34 @@ private:
 		SHADER_PARAMETER_ARRAY(uint32, MeshIndexOffset, [SectionArrayCount])
 		SHADER_PARAMETER_ARRAY(uint32, MeshUVsChannelOffset, [SectionArrayCount])
 		SHADER_PARAMETER_ARRAY(uint32, MeshUVsChannelCount, [SectionArrayCount])
-		SHADER_PARAMETER_SRV(Buffer, MeshIndexBuffer)
-		SHADER_PARAMETER_SRV(Buffer, MeshPositionBuffer)
-		SHADER_PARAMETER_SRV(Buffer, MeshUVsBuffer)
+		SHADER_PARAMETER_ARRAY(uint32, MeshSectionBufferIndex, [SectionArrayCount])
+
+		SHADER_PARAMETER_SRV(Buffer, MeshPositionBuffer0)
+		SHADER_PARAMETER_SRV(Buffer, MeshPositionBuffer1)
+		SHADER_PARAMETER_SRV(Buffer, MeshPositionBuffer2)
+		SHADER_PARAMETER_SRV(Buffer, MeshPositionBuffer3)
+		SHADER_PARAMETER_SRV(Buffer, MeshPositionBuffer4)
+		SHADER_PARAMETER_SRV(Buffer, MeshPositionBuffer5)
+		SHADER_PARAMETER_SRV(Buffer, MeshPositionBuffer6)
+		SHADER_PARAMETER_SRV(Buffer, MeshPositionBuffer7)
+
+		SHADER_PARAMETER_SRV(Buffer, MeshIndexBuffer0)
+		SHADER_PARAMETER_SRV(Buffer, MeshIndexBuffer1)
+		SHADER_PARAMETER_SRV(Buffer, MeshIndexBuffer2)
+		SHADER_PARAMETER_SRV(Buffer, MeshIndexBuffer3)
+		SHADER_PARAMETER_SRV(Buffer, MeshIndexBuffer4)
+		SHADER_PARAMETER_SRV(Buffer, MeshIndexBuffer5)
+		SHADER_PARAMETER_SRV(Buffer, MeshIndexBuffer6)
+		SHADER_PARAMETER_SRV(Buffer, MeshIndexBuffer7)
+		
+		SHADER_PARAMETER_SRV(Buffer, MeshUVsBuffer0)
+		SHADER_PARAMETER_SRV(Buffer, MeshUVsBuffer1)
+		SHADER_PARAMETER_SRV(Buffer, MeshUVsBuffer2)
+		SHADER_PARAMETER_SRV(Buffer, MeshUVsBuffer3)
+		SHADER_PARAMETER_SRV(Buffer, MeshUVsBuffer4)
+		SHADER_PARAMETER_SRV(Buffer, MeshUVsBuffer5)
+		SHADER_PARAMETER_SRV(Buffer, MeshUVsBuffer6)
+		SHADER_PARAMETER_SRV(Buffer, MeshUVsBuffer7)
 
 		SHADER_PARAMETER_SRV(Buffer, RootTriangleIndex)
 		SHADER_PARAMETER_UAV(StructuredBuffer, OutRootTrianglePosition0)
@@ -425,19 +450,110 @@ static void AddHairStrandUpdateMeshTrianglesPass(
 	FHairUpdateMeshTriangleCS::FParameters* Parameters = GraphBuilder.AllocParameters<FHairUpdateMeshTriangleCS::FParameters>();
 	Parameters->MaxRootCount		= RootData.RootCount;
 	Parameters->MaxSectionCount		= SectionCount;
-	Parameters->MeshPositionBuffer	= MeshData.Sections[0].PositionBuffer;
-	Parameters->MeshIndexBuffer		= MeshData.Sections[0].IndexBuffer;
-	Parameters->MeshUVsBuffer		= MeshData.Sections[0].UVsBuffer;
+
+	const uint32 MaxSectionBufferCount = 8;
+	struct FMeshSectionBuffers
+	{
+		uint32 MeshSectionBufferIndex = 0;
+		FRHIShaderResourceView* PositionBuffer = nullptr;
+		FRHIShaderResourceView* IndexBuffer = nullptr;
+		FRHIShaderResourceView* UVsBuffer = nullptr;
+	};
+	TMap<FRHIShaderResourceView*, FMeshSectionBuffers> UniqueMeshSectionBuffers;
+	uint32 UniqueMeshSectionBufferIndex = 0;
+
+	auto SetMeshSectionBuffers = [Parameters](uint32 UniqueIndex, const FHairStrandsProjectionMeshData::Section& MeshSectionData)
+	{
+		switch (UniqueIndex)
+		{
+		case 0:
+		{
+			Parameters->MeshPositionBuffer0 = MeshSectionData.PositionBuffer;
+			Parameters->MeshIndexBuffer0 = MeshSectionData.IndexBuffer;
+			Parameters->MeshUVsBuffer0 = MeshSectionData.UVsBuffer;
+			break;
+		}
+		case 1:
+		{
+			Parameters->MeshPositionBuffer1 = MeshSectionData.PositionBuffer;
+			Parameters->MeshIndexBuffer1 = MeshSectionData.IndexBuffer;
+			Parameters->MeshUVsBuffer1 = MeshSectionData.UVsBuffer;
+			break;
+		}
+		case 2:
+		{
+			Parameters->MeshPositionBuffer2 = MeshSectionData.PositionBuffer;
+			Parameters->MeshIndexBuffer2 = MeshSectionData.IndexBuffer;
+			Parameters->MeshUVsBuffer2 = MeshSectionData.UVsBuffer;
+			break;
+		}
+		case 3:
+		{
+			Parameters->MeshPositionBuffer3 = MeshSectionData.PositionBuffer;
+			Parameters->MeshIndexBuffer3 = MeshSectionData.IndexBuffer;
+			Parameters->MeshUVsBuffer3 = MeshSectionData.UVsBuffer;
+			break;
+		}
+		case 4:
+		{
+			Parameters->MeshPositionBuffer4 = MeshSectionData.PositionBuffer;
+			Parameters->MeshIndexBuffer4 = MeshSectionData.IndexBuffer;
+			Parameters->MeshUVsBuffer4 = MeshSectionData.UVsBuffer;
+			break;
+		}
+		case 5:
+		{
+			Parameters->MeshPositionBuffer5 = MeshSectionData.PositionBuffer;
+			Parameters->MeshIndexBuffer5 = MeshSectionData.IndexBuffer;
+			Parameters->MeshUVsBuffer5 = MeshSectionData.UVsBuffer;
+			break;
+		}
+		case 6:
+		{
+			Parameters->MeshPositionBuffer6 = MeshSectionData.PositionBuffer;
+			Parameters->MeshIndexBuffer6 = MeshSectionData.IndexBuffer;
+			Parameters->MeshUVsBuffer6 = MeshSectionData.UVsBuffer;
+			break;
+		}
+		case 7:
+		{
+			Parameters->MeshPositionBuffer7 = MeshSectionData.PositionBuffer;
+			Parameters->MeshIndexBuffer7 = MeshSectionData.IndexBuffer;
+			Parameters->MeshUVsBuffer7 = MeshSectionData.UVsBuffer;
+			break;
+		}
+		}
+	};
 
 	check(SectionCount < FHairUpdateMeshTriangleCS::SectionArrayCount);
 	for (int32 SectionIt = 0; SectionIt < SectionCount; ++SectionIt)
 	{
 		const FHairStrandsProjectionMeshData::Section& MeshSectionData = MeshData.Sections[SectionIt];
 		
-		// Sanity check to insure all sections share the same underlying buffer 
-		check(Parameters->MeshPositionBuffer == MeshSectionData.PositionBuffer);
-		check(Parameters->MeshIndexBuffer == MeshSectionData.IndexBuffer);
-		check(Parameters->MeshUVsBuffer == MeshSectionData.UVsBuffer);
+		if (const FMeshSectionBuffers* Buffers = UniqueMeshSectionBuffers.Find(MeshSectionData.PositionBuffer))
+		{
+			// Insure that all buffers actually match
+			check(Buffers->PositionBuffer == MeshSectionData.PositionBuffer);
+			check(Buffers->IndexBuffer == MeshSectionData.IndexBuffer);
+			check(Buffers->UVsBuffer == MeshSectionData.UVsBuffer);
+
+			Parameters->MeshSectionBufferIndex[SectionIt] = Buffers->MeshSectionBufferIndex;
+		}
+		else
+		{
+			// Only support 8 unique different buffer at the moment
+			check(UniqueMeshSectionBufferIndex < MaxSectionBufferCount);
+			SetMeshSectionBuffers(UniqueMeshSectionBufferIndex, MeshSectionData);
+
+			FMeshSectionBuffers& Entry = UniqueMeshSectionBuffers.Add(MeshSectionData.PositionBuffer);
+			Entry.MeshSectionBufferIndex = UniqueMeshSectionBufferIndex;
+			Entry.PositionBuffer = MeshSectionData.PositionBuffer;
+			Entry.IndexBuffer = MeshSectionData.IndexBuffer;
+			Entry.UVsBuffer = MeshSectionData.UVsBuffer;
+
+			Parameters->MeshSectionBufferIndex[SectionIt] = UniqueMeshSectionBufferIndex;
+			++UniqueMeshSectionBufferIndex;
+		}
 
 		Parameters->MeshSectionIndex[SectionIt]		= MeshSectionData.SectionIndex;
 		Parameters->MeshMaxIndexCount[SectionIt]	= MeshSectionData.TotalIndexCount;
@@ -447,6 +563,14 @@ static void AddHairStrandUpdateMeshTrianglesPass(
 		Parameters->MeshUVsChannelCount[SectionIt]	= MeshSectionData.UVsChannelCount;
 	}
 
+	if (MeshData.Sections.Num() > 0)
+	{
+		for (uint32 Index = UniqueMeshSectionBufferIndex; Index < MaxSectionBufferCount; ++Index)
+		{
+			SetMeshSectionBuffers(Index, MeshData.Sections[0]);
+		}
+	}
+
 	Parameters->RootTriangleIndex = LODData.RootTriangleIndexBuffer->SRV;
 	if (Type == HairStrandsTriangleType::RestPose)
 	{
