Fix deep shadow frustum bounds being either too small or too large depending on light orientation/type.

diff --git a/Engine/Shaders/Private/HairStrands/HairStrandsAABBCommon.ush b/Engine/Shaders/Private/HairStrands/HairStrandsAABBCommon.ush
--- a/Engine/Shaders/Private/HairStrands/HairStrandsAABBCommon.ush
+++ b/Engine/Shaders/Private/HairStrands/HairStrandsAABBCommon.ush
@@ -26,44 +26,76 @@ FHairAABB InitHairAABB()
 	return O;
 }
 
-FHairAABB ReadHairAABB(uint Index, RWBuffer<int> InBuffer)
+
+FHairAABB Transform(const FHairAABB In, float4x4 T)
+{
+	float4 P00 = float4(In.Min.x, In.Min.y, In.Min.z, 1);
+	float4 P10 = float4(In.Max.x, In.Min.y, In.Min.z, 1);
+	float4 P20 = float4(In.Max.x, In.Max.y, In.Min.z, 1);
+	float4 P30 = float4(In.Min.x, In.Max.y, In.Min.z, 1);
+
+	float4 P01 = float4(In.Min.x, In.Min.y, In.Max.z, 1);
+	float4 P11 = float4(In.Max.x, In.Min.y, In.Max.z, 1);
+	float4 P21 = float4(In.Max.x, In.Max.y, In.Max.z, 1);
+	float4 P31 = float4(In.Min.x, In.Max.y, In.Max.z, 1);
+
+	P00 = mul(P00, T); P00 /= P00.w;
+	P10 = mul(P10, T); P10 /= P10.w;
+	P20 = mul(P20, T); P20 /= P20.w;
+	P30 = mul(P30, T); P30 /= P30.w;
+
+	P01 = mul(P01, T); P01 /= P01.w;
+	P11 = mul(P11, T); P11 /= P11.w;
+	P21 = mul(P21, T); P21 /= P21.w;
+	P31 = mul(P31, T); P31 /= P31.w;
+
+	FHairAABB Out;
+	Out.Min = min(P00.xyz, min(P10.xyz, min(P20.xyz, P30.xyz)));
+	Out.Min = min(P01.xyz, min(P11.xyz, min(P21.xyz, min(P31.xyz, Out.Min))));
+
+	Out.Max = max(P00.xyz, max(P10.xyz, max(P20.xyz, P30.xyz)));
+	Out.Max = max(P01.xyz, max(P11.xyz, max(P21.xyz, max(P31.xyz, Out.Max))));
+	return Out;
+}
+
+FHairAABB ReadHairAABB(uint Index, RWBuffer<int> InBuffer, float PrecisionScale = HAIR_AABB_PRECISION_SCALE)
 {
 	const uint Index6 = Index * 6;
 	FHairAABB O;
-	O.Min.x = InBuffer[Index6 + 0] / HAIR_AABB_PRECISION_SCALE;
-	O.Min.y = InBuffer[Index6 + 1] / HAIR_AABB_PRECISION_SCALE;
-	O.Min.z = InBuffer[Index6 + 2] / HAIR_AABB_PRECISION_SCALE;
+	O.Min.x = InBuffer[Index6 + 0] / PrecisionScale;
+	O.Min.y = InBuffer[Index6 + 1] / PrecisionScale;
+	O.Min.z = InBuffer[Index6 + 2] / PrecisionScale;
 
-	O.Max.x = InBuffer[Index6 + 3] / HAIR_AABB_PRECISION_SCALE;
-	O.Max.y = InBuffer[Index6 + 4] / HAIR_AABB_PRECISION_SCALE;
-	O.Max.z = InBuffer[Index6 + 5] / HAIR_AABB_PRECISION_SCALE;
+	O.Max.x = InBuffer[Index6 + 3] / PrecisionScale;
+	O.Max.y = InBuffer[Index6 + 4] / PrecisionScale;
+	O.Max.z = InBuffer[Index6 + 5] / PrecisionScale;
 	return O;
 }
 
-FHairAABB ReadHairAABB(uint Index, Buffer<int> InBuffer)
+FHairAABB ReadHairAABB(uint Index, Buffer<int> InBuffer, float PrecisionScale = HAIR_AABB_PRECISION_SCALE)
 {
 	const uint Index6 = Index * 6;
 	FHairAABB O;
-	O.Min.x = InBuffer[Index6 + 0] / HAIR_AABB_PRECISION_SCALE;
-	O.Min.y = InBuffer[Index6 + 1] / HAIR_AABB_PRECISION_SCALE;
-	O.Min.z = InBuffer[Index6 + 2] / HAIR_AABB_PRECISION_SCALE;
+	O.Min.x = InBuffer[Index6 + 0] / PrecisionScale;
+	O.Min.y = InBuffer[Index6 + 1] / PrecisionScale;
+	O.Min.z = InBuffer[Index6 + 2] / PrecisionScale;
 
-	O.Max.x = InBuffer[Index6 + 3] / HAIR_AABB_PRECISION_SCALE;
-	O.Max.y = InBuffer[Index6 + 4] / HAIR_AABB_PRECISION_SCALE;
-	O.Max.z = InBuffer[Index6 + 5] / HAIR_AABB_PRECISION_SCALE;
+	O.Max.x = InBuffer[Index6 + 3] / PrecisionScale;
+	O.Max.y = InBuffer[Index6 + 4] / PrecisionScale;
+	O.Max.z = InBuffer[Index6 + 5] / PrecisionScale;
 	return O;
 }
 
-void WriteHairAABB(uint Index, FHairAABB In, RWBuffer<int> OutBuffer)
+void WriteHairAABB(uint Index, FHairAABB In, RWBuffer<int> OutBuffer, float PrecisionScale = HAIR_AABB_PRECISION_SCALE)
 {
 	const uint Index6 = Index * 6;
-	OutBuffer[Index6 + 0] = In.Min.x * HAIR_AABB_PRECISION_SCALE;
-	OutBuffer[Index6 + 1] = In.Min.y * HAIR_AABB_PRECISION_SCALE;
-	OutBuffer[Index6 + 2] = In.Min.z * HAIR_AABB_PRECISION_SCALE;
+	OutBuffer[Index6 + 0] = In.Min.x * PrecisionScale;
+	OutBuffer[Index6 + 1] = In.Min.y * PrecisionScale;
+	OutBuffer[Index6 + 2] = In.Min.z * PrecisionScale;
 
-	OutBuffer[Index6 + 3] = In.Max.x * HAIR_AABB_PRECISION_SCALE;
-	OutBuffer[Index6 + 4] = In.Max.y * HAIR_AABB_PRECISION_SCALE;
-	OutBuffer[Index6 + 5] = In.Max.z * HAIR_AABB_PRECISION_SCALE;
+	OutBuffer[Index6 + 3] = In.Max.x * PrecisionScale;
+	OutBuffer[Index6 + 4] = In.Max.y * PrecisionScale;
+	OutBuffer[Index6 + 5] = In.Max.z * PrecisionScale;
 }
 
 FHairAABB UnionHairAABB(FHairAABB A, FHairAABB B)
@@ -79,13 +111,14 @@ FHairAABB UnionHairAABB(FHairAABB A, FHairAABB B)
 	return O;
 }
 
-void InterlockHairAABB(uint Index, FHairAABB In, RWBuffer<int> OutBuffer)
+void InterlockHairAABB(uint Index, FHairAABB In, RWBuffer<int> OutBuffer, float PrecisionScale = HAIR_AABB_PRECISION_SCALE)
 {
 	const uint Index6 = Index * 6;
-	InterlockedMin(OutBuffer[Index6 + 0], int(In.Min.x * HAIR_AABB_PRECISION_SCALE));
-	InterlockedMin(OutBuffer[Index6 + 1], int(In.Min.y * HAIR_AABB_PRECISION_SCALE));
-	InterlockedMin(OutBuffer[Index6 + 2], int(In.Min.z * HAIR_AABB_PRECISION_SCALE));
-	InterlockedMax(OutBuffer[Index6 + 3], int(In.Max.x * HAIR_AABB_PRECISION_SCALE));
-	InterlockedMax(OutBuffer[Index6 + 4], int(In.Max.y * HAIR_AABB_PRECISION_SCALE));
-	InterlockedMax(OutBuffer[Index6 + 5], int(In.Max.z * HAIR_AABB_PRECISION_SCALE));
-}
\ No newline at end of file
+	InterlockedMin(OutBuffer[Index6 + 0], int(In.Min.x * PrecisionScale));
+	InterlockedMin(OutBuffer[Index6 + 1], int(In.Min.y * PrecisionScale));
+	InterlockedMin(OutBuffer[Index6 + 2], int(In.Min.z * PrecisionScale));
+
+	InterlockedMax(OutBuffer[Index6 + 3], int(In.Max.x * PrecisionScale));
+	InterlockedMax(OutBuffer[Index6 + 4], int(In.Max.y * PrecisionScale));
+	InterlockedMax(OutBuffer[Index6 + 5], int(In.Max.z * PrecisionScale));
+}
diff --git a/Engine/Shaders/Private/HairStrands/HairStrandsDeepShadowAllocation.usf b/Engine/Shaders/Private/HairStrands/HairStrandsDeepShadowAllocation.usf
--- a/Engine/Shaders/Private/HairStrands/HairStrandsDeepShadowAllocation.usf
+++ b/Engine/Shaders/Private/HairStrands/HairStrandsDeepShadowAllocation.usf
@@ -13,6 +13,30 @@
 #error MAX_SLOT_COUNT needs to be defined
 #endif
 
+////////////////////////////////////////////////////////////////////////////////
+
+float4x4 ComputeWorldToLight(
+	const FHairAABB Aabb,
+	const float3 LightDirection,
+	const float3 LightPosition,
+	const bool bIsDirectional)
+{
+	const float3 Extents = GetExtents(Aabb);
+	const float3 Center = GetCenter(Aabb);
+	const float  Radius = length(Extents);
+
+	if (bIsDirectional) // (LightType == LightType_Directional)
+	{
+		return LookAtMatrix(Center - LightDirection * Radius, Center, float3(0, 0, 1));
+	}
+	else // if (LightType == LightType_Spot || LightType == LightType_Point || LightType == LightType_Rect)
+	{
+		return LookAtMatrix(LightPosition, Center, float3(0, 0, 1));
+	}
+}
+
+////////////////////////////////////////////////////////////////////////////////
+
 float4				LightDirections[MAX_SLOT_COUNT];
 float4				LightPositions[MAX_SLOT_COUNT];
 uint4				MacroGroupIndices[MAX_SLOT_COUNT];
@@ -45,25 +69,32 @@ float ComputeMinStrandRadiusAtDepth1(const int2 Resolution, const float FOV, con
 	return DiameterToRadius * InRasterizationScale * StrandDiameterAtDepth1;
 }
 
-// This function is equivalent of the CPU version in HairStrandsDeepShadow.cpp
-// TODO: create a tighter bound
+// This function is similar to the CPU version in HairStrandsDeepShadow.cpp
 void ComputeWorldToLightClip(
 	inout float4x4	OutWorldToClipTransform,
 	inout float		OutMinStrandRadiusAtDepth1,
 	inout float3	OutLightDirection,
-	const FHairAABB Aabb,
+	const FHairAABB World_AABB,
 	const uint		SlotIndex)
 {
-	const float3 Extents = GetExtents(Aabb);
-	const float3 Center  = GetCenter(Aabb);
-	const float  Radius  = length(Extents);
-
 	const float3 LightDirection = LightDirections[SlotIndex].xyz;
 	const float3 LightPosition  = LightPositions[SlotIndex].xyz;
-	const uint bIsDirectional	= LightPositions[SlotIndex].w == 0;
+	const bool bIsDirectional	= LightPositions[SlotIndex].w == 0;
+
+	float3 Center = GetCenter(World_AABB);
+	OutLightDirection = bIsDirectional ? LightDirection : normalize(LightPosition - Center);
 
-	const float MinZ = max(0.1f, length(LightPosition - Center)) - Radius;
-	const float MaxZ = max(0.2f, length(LightPosition - Center)) + Radius;
+	const float4x4 Coarse_WorldToLight = ComputeWorldToLight(World_AABB, LightDirection, LightPosition, bIsDirectional);
+
+	const FHairAABB Light_AABB = Transform(World_AABB, Coarse_WorldToLight);
+	const float3 Light_Extents = GetExtents(Light_AABB);
+
+	const float Radius = length(Light_Extents.xy);
+	float MinZ = -Light_Extents.z;
+	float MaxZ = +Light_Extents.z;
+
+	// Update AABB center based on the tigher Z entents
+	Center = ((MaxZ * OutLightDirection) + (MinZ * OutLightDirection)) * 0.5f + Center;
 
 	const float StrandHairStableRasterizationScale = max(RasterizationScale, 1.0f);
 	OutMinStrandRadiusAtDepth1 = 1;
@@ -71,10 +102,8 @@ void ComputeWorldToLightClip(
 
 	if (bIsDirectional) // (LightType == LightType_Directional)
 	{
-		OutLightDirection			= LightDirection;
-
-		const float4x4 ProjMatrix	= ReversedZOrthoMatrix(Radius, Radius, 1.f / (2 * Radius), 0);
-		const float4x4 WorldToLight	= LookAtMatrix(Center - LightDirection * Radius, Center, float3(0, 0, 1));
+		const float4x4 WorldToLight	= LookAtMatrix(Center - LightDirection * abs(MinZ), Center, float3(0, 0, 1));
+		const float4x4 ProjMatrix	= ReversedZOrthoMatrix(Radius, Radius, 1.f / (MaxZ-MinZ), 0);
 		OutWorldToClipTransform		= mul(WorldToLight, ProjMatrix);
 
 		const float RadiusAtDepth1	= Radius / min(SlotResolution.x, SlotResolution.y);
@@ -82,13 +111,14 @@ void ComputeWorldToLightClip(
 	}
 	else // if (LightType == LightType_Spot || LightType == LightType_Point || LightType == LightType_Rect)
 	{
-		OutLightDirection			= normalize(LightPosition - Center);
+		MinZ = max(0.1f, length(LightPosition - Center)) + MinZ;
+		MaxZ = max(0.2f, length(LightPosition - Center)) + MaxZ;
 
 		const float SphereDistance	= length(LightPosition - Center);
 		const float HalfFov			= asin(Radius / SphereDistance);
 
-		const float4x4 ProjMatrix	= ReversedZPerspectiveMatrix(HalfFov, 1, 1, MinZ, MaxZ);
 		const float4x4 WorldToLight = LookAtMatrix(LightPosition, Center, float3(0, 0, 1));
+		const float4x4 ProjMatrix	= ReversedZPerspectiveMatrix(HalfFov, 1, 1, MinZ, MaxZ);
 		OutWorldToClipTransform		= mul(WorldToLight, ProjMatrix);
 
 		OutMinStrandRadiusAtDepth1 = ComputeMinStrandRadiusAtDepth1(SlotResolution, 2 * HalfFov, RasterizationScale);
