Fix a bug of EyeAdaptation on mobile platform.

diff --git a/Engine/Shaders/Private/MobileBasePassPixelShader.usf b/Engine/Shaders/Private/MobileBasePassPixelShader.usf
--- a/Engine/Shaders/Private/MobileBasePassPixelShader.usf
+++ b/Engine/Shaders/Private/MobileBasePassPixelShader.usf
@@ -786,9 +786,9 @@ void Main(
 		#endif
 	#endif
 
-	#if !MATERIALBLENDING_MODULATE
+	#if !MATERIALBLENDING_MODULATE && USE_PREEXPOSURE
 		// MobileHDR applies PreExposure in tonemapper
-		OutColor.rgb *= USE_PREEXPOSURE ? ResolvedView.PreExposure : 1;
+		OutColor.rgb *= ResolvedView.PreExposure;
 	#endif
 
 	#if MATERIAL_IS_SKY
diff --git a/Engine/Shaders/Private/PostProcessMaterialShaders.usf b/Engine/Shaders/Private/PostProcessMaterialShaders.usf
--- a/Engine/Shaders/Private/PostProcessMaterialShaders.usf
+++ b/Engine/Shaders/Private/PostProcessMaterialShaders.usf
@@ -379,6 +379,10 @@ void MainPS(
 #if OUTPUT_GAMMA_SPACE
 	FullSceneColor.rgb = sqrt(FullSceneColor.rgb);
 #endif
+
+#if !MATERIALBLENDING_MODULATE && USE_PREEXPOSURE
+	FullSceneColor.xyz *= View.PreExposure;
+#endif
 #endif
 	OutColor = FullSceneColor;
 }
diff --git a/Engine/Shaders/Private/PostProcessMobile.usf b/Engine/Shaders/Private/PostProcessMobile.usf
--- a/Engine/Shaders/Private/PostProcessMobile.usf
+++ b/Engine/Shaders/Private/PostProcessMobile.usf
@@ -237,6 +237,15 @@ void BloomPS_ES2(
 	)
 {
 	half3 AverageColor = 0.0f;
+
+#if USE_PREEXPOSURE
+	const float OneOverPreExposure = View.OneOverPreExposure;
+	const float PreExposure = View.PreExposure;
+#else
+	const float OneOverPreExposure = 1.0f;
+	const float PreExposure = View.PreExposure;
+#endif
+
 #if ES2_USE_BLOOM || ES2_USE_SUN || ES2_USE_EYEADAPTATION
 	half4 C0 = PostprocessInput0.Sample(PostprocessInput0Sampler, InUVs[0]);
 	half4 C1 = PostprocessInput0.Sample(PostprocessInput0Sampler, InUVs[1]);
@@ -257,7 +266,7 @@ void BloomPS_ES2(
 	#endif
 
 	// Try to kill negatives and NaNs here
-	AverageColor.rgb = max(AverageColor.rgb, 0);
+	AverageColor.rgb = max(AverageColor.rgb, 0) * OneOverPreExposure;
 
 	
 	#if ES2_USE_BLOOM || ES2_USE_SUN
@@ -272,7 +281,7 @@ void BloomPS_ES2(
 		OutColor.rgb *= Amount;
 		OutColor.a = 0;
 
-		OutColor = OutColor;
+		OutColor = OutColor * PreExposure;
 	#endif
 #endif
 
@@ -320,10 +329,6 @@ void BloomPS_ES2(
 
 #if ES2_USE_EYEADAPTATION
 
-	#if USE_PREEXPOSURE
-		AverageColor.xyz *= View.OneOverPreExposure;
-	#endif
-
 	const half Intensity = max(dot(AverageColor, half3(1.0f, 1.0f, 1.0f) / 3.0f), EyeAdaptation_LuminanceMin);
 	const float LogIntensity = clamp(log2(Intensity), -10.0f, 20.0f);
 
