Remove plain voxelization for hair AO.

diff --git a/Engine/Shaders/Private/HairStrands/HairStrandsEnvironmentAO.usf b/Engine/Shaders/Private/HairStrands/HairStrandsEnvironmentAO.usf
--- a/Engine/Shaders/Private/HairStrands/HairStrandsEnvironmentAO.usf
+++ b/Engine/Shaders/Private/HairStrands/HairStrandsEnvironmentAO.usf
@@ -12,12 +12,7 @@
 
 #include "../ShaderDrawDebug.ush"
 
-// Note: VOXEL_TRAVERSAL_TYPE is specified as a permutation
-// VOXEL_TRAVERSAL_LEGACY means old plain-voxel traversal
-#if VOXEL_TRAVERSAL_TYPE == 0 
-#undef VOXEL_TRAVERSAL_TYPE
-#define VOXEL_TRAVERSAL_TYPE VOXEL_TRAVERSAL_LEGACY
-#endif 
+#define VOXEL_TRAVERSAL_TYPE VOXEL_TRAVERSAL_LINEAR_MIPMAP
 #define VOXEL_TRAVERSAL_DEBUG 0 
 #include "HairStrandsVoxelPageTraversal.ush"
 
@@ -25,101 +20,14 @@ Texture2D<uint4> HairCategorizationTexture;
 
 uint  Voxel_MacroGroupId;
 float Voxel_TanConeAngle;
-float Voxel_DensityScale;
-float Voxel_DepthBiasScale;
-
 float AO_Power;
 float AO_Intensity;
 uint  AO_SampleCount;
 float AO_DistanceThreshold;
 
-float3 Voxel_MinAABB;
-float3 Voxel_MaxAABB;
-uint   Voxel_Resolution;
-Texture3D<uint> Voxel_DensityTexture;
-
-// todo: make this function in common with the sky lighting integration pass
-FHairTraversalResult InternalComputeHairCountVoxel(
-	float3 WorldPosition, 
-	float3 LightPosition, 
-	float3 InVoxelAABBMin, 
-	float3 InVoxelAABBMax, 
-	float InVoxelResolution, 
-	Texture3D<uint> InDensityTexture, 
-	float DensityScale, 
-	float CountThreshold, 
-	float DistanceThreshold, 
-	bool bUseConeTracing,
-	bool bDebugEnabled)
-{
-#if VOXEL_TRAVERSAL_TYPE != VOXEL_TRAVERSAL_LEGACY
-
-	FVirtualVoxelCommonDesc CommonDesc;
-	CommonDesc.PageCountResolution = VirtualVoxel.PageCountResolution;
-	CommonDesc.VoxelWorldSize = VirtualVoxel.VoxelWorldSize;
-	CommonDesc.PageTextureResolution = VirtualVoxel.PageTextureResolution;
-	CommonDesc.PageResolution = VirtualVoxel.PageResolution;
-
-	const FPackedVirtualVoxelNodeDesc PackedNode = VirtualVoxel.NodeDescBuffer[Voxel_MacroGroupId];
-	const FVirtualVoxelNodeDesc NodeDesc = UnpackVoxelNode(PackedNode, VirtualVoxel.PageResolution);
-
-	FHairTraversalSettings TraversalSettings = InitHairTraversalSettings();
-	TraversalSettings.DensityScale = VirtualVoxel.DensityScale;
-	TraversalSettings.CountThreshold = CountThreshold;
-	TraversalSettings.DistanceThreshold = 1;
-	TraversalSettings.bDebugEnabled = bDebugEnabled;
-	TraversalSettings.SteppingScale = VirtualVoxel.SteppingScale;
-
-	return ComputeHairCountVirtualVoxel(
-		WorldPosition,
-		LightPosition,
-		CommonDesc,
-		NodeDesc,
-		VirtualVoxel.PageIndexBuffer,
-		VirtualVoxel.PageTexture,
-		TraversalSettings);
-
-#else
-
-	if (bUseConeTracing)
-	{
-		return ComputeHairCountVoxelCone(
-			WorldPosition,
-			LightPosition,
-			InVoxelAABBMin,
-			InVoxelAABBMax,
-			InVoxelResolution,
-			InDensityTexture,
-			DensityScale,
-			CountThreshold,
-			DistanceThreshold,
-			Voxel_TanConeAngle);
-	}
-	else
-	{
-		return ComputeHairCountVoxel(
-			WorldPosition,
-			LightPosition,
-			InVoxelAABBMin,
-			InVoxelAABBMax,
-			InVoxelResolution,
-			InDensityTexture,
-			DensityScale,
-			CountThreshold,
-			DistanceThreshold);
-	}
-#endif
-}
-
 float TraceAO(float3 WorldPosition, float3 WorldNormal, float4 SvPosition)
 {
-	#if VOXEL_TRAVERSAL_TYPE == VOXEL_TRAVERSAL_LEGACY
-	const float3 VoxelSize = (Voxel_MaxAABB - Voxel_MinAABB) / Voxel_Resolution;
-	const float  VoxelMaxSize = max(VoxelSize.x, max(VoxelSize.y, VoxelSize.z));
-	#else
-	const float  VoxelMaxSize = VirtualVoxel.VoxelWorldSize;
-	#endif	
-	const float3 VoxelOffset = VoxelMaxSize * GetHairVoxelJitter(SvPosition.xy, View.StateFrameIndexMod8);
+	const float3 VoxelOffset = VirtualVoxel.VoxelWorldSize * GetHairVoxelJitter(SvPosition.xy, View.StateFrameIndexMod8);
 
 	const bool bDebugEnabled = all(uint2(SvPosition.xy) == uint2(GetCursorPos()));
 	
@@ -151,6 +59,22 @@ float TraceAO(float3 WorldPosition, float3 WorldNormal, float4 SvPosition)
 	const uint SampleCount = max(AO_SampleCount, 1);
 #endif
 
+	FVirtualVoxelCommonDesc CommonDesc;
+	CommonDesc.PageCountResolution = VirtualVoxel.PageCountResolution;
+	CommonDesc.VoxelWorldSize = VirtualVoxel.VoxelWorldSize;
+	CommonDesc.PageTextureResolution = VirtualVoxel.PageTextureResolution;
+	CommonDesc.PageResolution = VirtualVoxel.PageResolution;
+
+	const FPackedVirtualVoxelNodeDesc PackedNode = VirtualVoxel.NodeDescBuffer[Voxel_MacroGroupId];
+	const FVirtualVoxelNodeDesc NodeDesc = UnpackVoxelNode(PackedNode, VirtualVoxel.PageResolution);
+
+	FHairTraversalSettings TraversalSettings = InitHairTraversalSettings();
+	TraversalSettings.DensityScale = VirtualVoxel.DensityScale;
+	TraversalSettings.CountThreshold = 1;
+	TraversalSettings.DistanceThreshold = AO_DistanceThreshold;
+	TraversalSettings.bDebugEnabled = bDebugEnabled;
+	TraversalSettings.SteppingScale = VirtualVoxel.SteppingScale;
+
 	const float VoxelAOMaxDistance = 1000.0;
 	float3 BentNormal = 0;
 	float Visibility = 0;
@@ -166,7 +90,7 @@ float TraceAO(float3 WorldPosition, float3 WorldNormal, float4 SvPosition)
 			InterleavedGradientNoise(SvPosition.xy, View.StateFrameIndexMod8 + i * 117));
 
 		// Depth bias
-		const float3 SampleVoxelOffset = VoxelMaxSize * GetHairVoxelJitter(SvPosition.xy, View.StateFrameIndexMod8 + i * 17);
+		const float3 SampleVoxelOffset = VirtualVoxel.VoxelWorldSize * GetHairVoxelJitter(SvPosition.xy, View.StateFrameIndexMod8 + i * 17);
 		const float3 SampleWorldPosition = WorldPosition + VoxelOffset;
 
 		const float3 LocalL = CosineSampleHemisphere(frac(Rand2 + SequenceOffset)).xyz;
@@ -174,7 +98,16 @@ float TraceAO(float3 WorldPosition, float3 WorldNormal, float4 SvPosition)
 		const float3 SampleLightPosition = SampleWorldPosition + SampleL * VoxelAOMaxDistance;
 
 		// Compute the number of hair count between light & shading point
-		const FHairTraversalResult Result = InternalComputeHairCountVoxel(SampleWorldPosition, SampleLightPosition, Voxel_MinAABB, Voxel_MaxAABB, Voxel_Resolution, Voxel_DensityTexture, Voxel_DensityScale, 1, AO_DistanceThreshold, true, bDebugEnabled);
+
+		const FHairTraversalResult Result = ComputeHairCountVirtualVoxel(
+			SampleWorldPosition,
+			SampleLightPosition,
+			CommonDesc,
+			NodeDesc,
+			VirtualVoxel.PageIndexBuffer,
+			VirtualVoxel.PageTexture,
+			TraversalSettings);
+
 		const float StepVisibility = saturate(1 - Result.HairCount) * Result.Visibility;
 		BentNormal += SampleL * StepVisibility;
 		Visibility += StepVisibility;
diff --git a/Engine/Source/Runtime/Renderer/Private/HairStrands/HairStrandsEnvironment.cpp b/Engine/Source/Runtime/Renderer/Private/HairStrands/HairStrandsEnvironment.cpp
--- a/Engine/Source/Runtime/Renderer/Private/HairStrands/HairStrandsEnvironment.cpp
+++ b/Engine/Source/Runtime/Renderer/Private/HairStrands/HairStrandsEnvironment.cpp
@@ -31,9 +31,6 @@ static FAutoConsoleVariableRef CVarHairSkylightingEnable(TEXT("r.HairStrands.Sky
 static int32 GHairSkyAOEnable = 1;
 static FAutoConsoleVariableRef CVarHairSkyAOEnable(TEXT("r.HairStrands.SkyAO"), GHairSkyAOEnable, TEXT("Enable (sky) AO on hair."));
 
-static int32 GHairSkyAOVirtualVoxel = 4;
-static FAutoConsoleVariableRef CVarHairSkyAOVirtualVoxel(TEXT("r.HairStrands.SkyAO.VirtualVoxel"), GHairSkyAOVirtualVoxel, TEXT("Use virtual voxel for visibility/hair tracing/counting if available."));
-
 static float GHairSkylightingConeAngle = 3;
 static FAutoConsoleVariableRef CVarHairSkylightingConeAngle(TEXT("r.HairStrands.SkyLighting.ConeAngle"), GHairSkylightingConeAngle, TEXT("Cone angle for tracing sky lighting on hair."));
 
@@ -74,8 +71,7 @@ class FHairEnvironmentAO : public FGlobalShader
 	SHADER_USE_PARAMETER_STRUCT(FHairEnvironmentAO, FGlobalShader)
 
 	class FSampleSet : SHADER_PERMUTATION_INT("PERMUTATION_SAMPLESET", 2);
-	class FVirtualVoxel : SHADER_PERMUTATION_INT("VOXEL_TRAVERSAL_TYPE", 5);
-	using FPermutationDomain = TShaderPermutationDomain<FSampleSet, FVirtualVoxel>;
+	using FPermutationDomain = TShaderPermutationDomain<FSampleSet>;
 
 	static bool ShouldCompilePermutation(const FGlobalShaderPermutationParameters& Parameters)
 	{
@@ -84,12 +80,7 @@ class FHairEnvironmentAO : public FGlobalShader
 
 	BEGIN_SHADER_PARAMETER_STRUCT(FParameters, )
 
-		SHADER_PARAMETER(FVector, Voxel_MinAABB)
 		SHADER_PARAMETER(uint32, Voxel_MacroGroupId)
-		SHADER_PARAMETER(FVector, Voxel_MaxAABB)
-		SHADER_PARAMETER(uint32, Voxel_Resolution)
-		SHADER_PARAMETER(float, Voxel_DensityScale)
-		SHADER_PARAMETER(float, Voxel_DepthBiasScale)
 		SHADER_PARAMETER(float, Voxel_TanConeAngle)
 		SHADER_PARAMETER(float, AO_Power)
 		SHADER_PARAMETER(float, AO_Intensity)
@@ -99,8 +90,6 @@ class FHairEnvironmentAO : public FGlobalShader
 		SHADER_PARAMETER_STRUCT_INCLUDE(FSceneTextureSamplerParameters, SceneTextureSamplers)
 
 		SHADER_PARAMETER_RDG_TEXTURE(Texture2D, HairCategorizationTexture)
-		SHADER_PARAMETER_RDG_TEXTURE(Texture3D, Voxel_DensityTexture)
-
 		SHADER_PARAMETER_STRUCT_REF(FViewUniformShaderParameters, ViewUniformBuffer)
 
 		SHADER_PARAMETER_STRUCT_REF(FVirtualVoxelParameters, VirtualVoxel)
@@ -122,23 +111,12 @@ static void AddHairStrandsEnvironmentAOPass(
 	FRDGTextureRef Output)
 {
 	check(Output);
-	const bool bUseVirtualVoxel = MacroGroupDatas.VirtualVoxelResources.IsValid() && GHairSkyAOVirtualVoxel > 0;
-	const int32 VoxelTraversalType = FMath::Clamp(GHairSkyAOVirtualVoxel, 0, 4);
 
 	FSceneTextureParameters SceneTextures;
 	SetupSceneTextureParameters(GraphBuilder, &SceneTextures);
 
 	FHairEnvironmentAO::FParameters* PassParameters = GraphBuilder.AllocParameters<FHairEnvironmentAO::FParameters>();
 	PassParameters->Voxel_MacroGroupId = MacroGroupData.MacroGroupId;
-	if (!bUseVirtualVoxel)
-	{
-		PassParameters->Voxel_MinAABB = MacroGroupData.GetMinBound();
-		PassParameters->Voxel_MaxAABB = MacroGroupData.GetMaxBound();
-		PassParameters->Voxel_Resolution = MacroGroupData.GetResolution();
-		PassParameters->Voxel_DensityTexture = GraphBuilder.RegisterExternalTexture(MacroGroupData.VoxelResources.DensityTexture);
-	}
-	PassParameters->Voxel_DensityScale = GetHairStrandsVoxelizationDensityScale();
-	PassParameters->Voxel_DepthBiasScale = GetHairStrandsVoxelizationDepthBiasScale();
 	PassParameters->Voxel_TanConeAngle = FMath::Tan(FMath::DegreesToRadians(GetHairStrandsSkyLightingConeAngle()));
 	PassParameters->SceneTextures = SceneTextures;
 	PassParameters->VirtualVoxel = MacroGroupDatas.VirtualVoxelResources.UniformBuffer;
@@ -160,7 +138,6 @@ static void AddHairStrandsEnvironmentAOPass(
 
 	FHairEnvironmentAO::FPermutationDomain PermutationVector;
 	PermutationVector.Set<FHairEnvironmentAO::FSampleSet>(PassParameters->AO_SampleCount <= 16 ? 0 : 1);
-	PermutationVector.Set<FHairEnvironmentAO::FVirtualVoxel>(bUseVirtualVoxel ? VoxelTraversalType : 0);
 
 	TShaderMapRef<FHairEnvironmentAO> PixelShader(View.ShaderMap, PermutationVector);
 	ClearUnusedGraphResources(PixelShader, PassParameters);
