Remove plain voxeliation support for hair env. lighting

diff --git a/Engine/Shaders/Private/HairStrands/HairStrandsEnvironmentLighting.usf b/Engine/Shaders/Private/HairStrands/HairStrandsEnvironmentLighting.usf
--- a/Engine/Shaders/Private/HairStrands/HairStrandsEnvironmentLighting.usf
+++ b/Engine/Shaders/Private/HairStrands/HairStrandsEnvironmentLighting.usf
@@ -40,12 +40,8 @@
 
 #include "../ShaderDrawDebug.ush"
 
-// Note: VOXEL_TRAVERSAL_TYPE is specified as a permutation
-// VOXEL_TRAVERSAL_LEGACY means old plain-voxel traversal
-#if VOXEL_TRAVERSAL_TYPE == 0 
-	#undef VOXEL_TRAVERSAL_TYPE
-	#define VOXEL_TRAVERSAL_TYPE VOXEL_TRAVERSAL_LEGACY
-#endif 
+#define VOXEL_TRAVERSAL_TYPE VOXEL_TRAVERSAL_LINEAR_MIPMAP
+//#define VOXEL_TRAVERSAL_TYPE VOXEL_TRAVERSAL_LINEAR
 #define VOXEL_TRAVERSAL_DEBUG 0
 #include "HairStrandsVoxelPageTraversal.ush"
 
@@ -100,29 +96,16 @@ float TransmissionDensityScaleFactor;
 
 uint  Voxel_MacroGroupId;
 float Voxel_TanConeAngle;
-float Voxel_DensityScale;
-float Voxel_DepthBiasScale;
-
-float3 Voxel_MinAABB;
-float3 Voxel_MaxAABB;
-uint   Voxel_Resolution;
-Texture3D<uint> Voxel_DensityTexture;
 
 FHairTraversalResult InternalComputeHairCountVoxel(
 	float3 WorldPosition, 
 	float3 LightPosition, 
-	float3 InVoxelAABBMin, 
-	float3 InVoxelAABBMax, 
-	float InVoxelResolution, 
-	Texture3D<uint> InDensityTexture, 
-	float DensityScale, 
+	float ExtraDensityScale, 
 	float CountThreshold, 
 	float DistanceThreshold,
 	bool bUseConeTracing,
 	bool bDebugEnabled)
 {
-#if VOXEL_TRAVERSAL_TYPE != VOXEL_TRAVERSAL_LEGACY
-
 	FVirtualVoxelCommonDesc CommonDesc;
 	CommonDesc.PageCountResolution = VirtualVoxel.PageCountResolution;
 	CommonDesc.VoxelWorldSize = VirtualVoxel.VoxelWorldSize;
@@ -133,7 +116,7 @@ FHairTraversalResult InternalComputeHairCountVoxel(
 	const FVirtualVoxelNodeDesc NodeDesc = UnpackVoxelNode(PackedNode, VirtualVoxel.PageResolution);
 
 	FHairTraversalSettings TraversalSettings = InitHairTraversalSettings();
-	TraversalSettings.DensityScale = DensityScale;
+	TraversalSettings.DensityScale = VirtualVoxel.DensityScale * ExtraDensityScale;
 	TraversalSettings.CountThreshold = GetOpaqueVoxelValue();
 	TraversalSettings.DistanceThreshold = DistanceThreshold;
 	TraversalSettings.bDebugEnabled = bDebugEnabled;
@@ -147,37 +130,6 @@ FHairTraversalResult InternalComputeHairCountVoxel(
 		VirtualVoxel.PageIndexBuffer,
 		VirtualVoxel.PageTexture,
 		TraversalSettings);
-#else
-
-	if (bUseConeTracing)
-	{
-		return ComputeHairCountVoxelCone(
-			WorldPosition,
-			LightPosition,
-			InVoxelAABBMin,
-			InVoxelAABBMax,
-			InVoxelResolution,
-			InDensityTexture,
-			DensityScale,
-			CountThreshold,
-			DistanceThreshold,
-			Voxel_TanConeAngle);
-	}
-	else
-	{
-		return ComputeHairCountVoxel(
-			WorldPosition,
-			LightPosition,
-			InVoxelAABBMin,
-			InVoxelAABBMax,
-			InVoxelResolution,
-			InDensityTexture,
-			DensityScale,
-			CountThreshold,
-			DistanceThreshold);
-	}
-
-#endif
 }
 
 //////////////////////////////////////////////////////////////////////////////////////////////////////////////////////////
@@ -197,13 +149,8 @@ float3 ReflectionEnvironment(FGBufferData GBuffer, float2 ScreenPosition, float2
 	const uint2  PixelPosition	= SvPosition - ResolvedView.ViewRectMin.xy;
 	const float  PixelSceneDepth= GBuffer.Depth;
 
-#if VOXEL_TRAVERSAL_TYPE != VOXEL_TRAVERSAL_LEGACY
-	const float3 VoxelSize = VirtualVoxel.VoxelWorldSize; // Use the correct voxel size when we use virtual voxel volume.
-#else
-	const float3 VoxelSize = (Voxel_MaxAABB - Voxel_MinAABB) / Voxel_Resolution;
-#endif
 	const uint PrimitiveIdMod8 = (PrimitiveId % 8);
-	const float  VoxelMaxSize = max(VoxelSize.x, max(VoxelSize.y, VoxelSize.z));
+	const float  VoxelMaxSize = VirtualVoxel.VoxelWorldSize;
 	const float3 VoxelOffset = VoxelMaxSize * GetHairVoxelJitter(SvPosition, View.StateFrameIndexMod8 + PrimitiveIdMod8);
 
 	const float3 V = -CameraToPixel;
@@ -275,7 +222,7 @@ float3 ReflectionEnvironment(FGBufferData GBuffer, float2 ScreenPosition, float2
 
 			// Compute the number of hair count between light & shading point
 			const float SinLightAngle = dot(SampleL, T);
-			const FHairTraversalResult Result = InternalComputeHairCountVoxel(SampleWorldPosition, SampleLightPosition, Voxel_MinAABB, Voxel_MaxAABB, Voxel_Resolution, Voxel_DensityTexture, Voxel_DensityScale, 0, MaxSearchDistance, true, bDebugEnabled);
+			const FHairTraversalResult Result = InternalComputeHairCountVoxel(SampleWorldPosition, SampleLightPosition, 1, 0, MaxSearchDistance, true, bDebugEnabled);
 			const FHairTransmittanceMask TransmittanceMask = ComputeTransmittanceFromHairCount(Result.HairCount, Result.Visibility, GBuffer, SinLightAngle, HairScatteringLUTTexture, HairLUTSampler);
 
 			const float ScalarTransmittance = max(TransmittanceMask.Transmittance.x, max(TransmittanceMask.Transmittance.y, TransmittanceMask.Transmittance.z));
@@ -403,7 +350,7 @@ float3 ReflectionEnvironment(FGBufferData GBuffer, float2 ScreenPosition, float2
 			const float3 SampleLightPosition = WorldPosition + L * MaxDistance;
 
 			// Compute the number of hair count between light & shading point
-			const FHairTraversalResult Result = InternalComputeHairCountVoxel(SampleWorldPosition, SampleLightPosition, Voxel_MinAABB, Voxel_MaxAABB, Voxel_Resolution, Voxel_DensityTexture, Voxel_DensityScale, 0, MaxSearchDistance, true, bDebugEnabled);
+			const FHairTraversalResult Result = InternalComputeHairCountVoxel(SampleWorldPosition, SampleLightPosition, 1, 0, MaxSearchDistance, true, bDebugEnabled);
 			TransmittanceMask = ComputeTransmittanceFromHairCount(Result.HairCount, Result.Visibility, GBuffer, SinLightAngle, HairScatteringLUTTexture, HairLUTSampler);
 		}
 
@@ -468,7 +415,7 @@ float3 ReflectionEnvironment(FGBufferData GBuffer, float2 ScreenPosition, float2
 			}
 			else
 			{
-				Result = InternalComputeHairCountVoxel(SampleWorldPosition, SampleLightPosition, Voxel_MinAABB, Voxel_MaxAABB, Voxel_Resolution, Voxel_DensityTexture, Voxel_DensityScale * TransmissionDensityScaleFactor, 0, DistanceThresholdTT, false, bDebugEnabled);
+				Result = InternalComputeHairCountVoxel(SampleWorldPosition, SampleLightPosition, TransmissionDensityScaleFactor, 0, DistanceThresholdTT, false, bDebugEnabled);
 			}
 			TransmittanceMask = ComputeTransmittanceFromHairCount(Result.HairCount, Result.Visibility, GBuffer, SinLightAngle, HairScatteringLUTTexture, HairLUTSampler);
 		}
diff --git a/Engine/Source/Runtime/Renderer/Private/HairStrands/HairStrandsEnvironment.cpp b/Engine/Source/Runtime/Renderer/Private/HairStrands/HairStrandsEnvironment.cpp
--- a/Engine/Source/Runtime/Renderer/Private/HairStrands/HairStrandsEnvironment.cpp
+++ b/Engine/Source/Runtime/Renderer/Private/HairStrands/HairStrandsEnvironment.cpp
@@ -34,9 +34,6 @@ static FAutoConsoleVariableRef CVarHairSkyAOEnable(TEXT("r.HairStrands.SkyAO"),
 static float GHairSkylightingConeAngle = 3;
 static FAutoConsoleVariableRef CVarHairSkylightingConeAngle(TEXT("r.HairStrands.SkyLighting.ConeAngle"), GHairSkylightingConeAngle, TEXT("Cone angle for tracing sky lighting on hair."));
 
-static int GHairSkylightingVirtualVoxel = 4;
-static FAutoConsoleVariableRef CVarHairSkylightingVirtualVoxel(TEXT("r.HairStrands.SkyLighting.VirtualVoxel"), GHairSkylightingVirtualVoxel, TEXT("Use virtual voxel for visibility/hair tracing/counting if available."));
-
 static int32 GHairStrandsSkyLightingSampleCount = 16;
 static FAutoConsoleVariableRef CVarHairStrandsSkyLightingSampleCount(TEXT("r.HairStrands.SkyLighting.SampleCount"), GHairStrandsSkyLightingSampleCount, TEXT("Number of samples used for evaluating multiple scattering and visible area (default is set to 16)."), ECVF_Scalability | ECVF_RenderThreadSafe);
 
@@ -166,9 +163,8 @@ class FHairEnvironmentLighting
 public:
 	class FGroupSize : SHADER_PERMUTATION_INT("PERMUTATION_GROUP_SIZE", 2);
 	class FSampleSet : SHADER_PERMUTATION_INT("PERMUTATION_SAMPLESET", 2);
-	class FVirtualVoxel : SHADER_PERMUTATION_INT("VOXEL_TRAVERSAL_TYPE", 5);
 	class FUseSceneColor : SHADER_PERMUTATION_INT("USE_SCENE_COLOR", 2);
-	using FPermutationDomain = TShaderPermutationDomain<FSampleSet, FVirtualVoxel, FUseSceneColor>;
+	using FPermutationDomain = TShaderPermutationDomain<FSampleSet, FUseSceneColor>;
 
 	static void ModifyCompilationEnvironment(const FGlobalShaderPermutationParameters& Parameters, FShaderCompilerEnvironment& OutEnvironment)
 	{
@@ -183,8 +179,7 @@ public:
 		}
 
 		const FPermutationDomain PermutationVector(Parameters.PermutationId);
-		if (PermutationVector.Get<FUseSceneColor>() == 1 &&
-		   (PermutationVector.Get<FSampleSet>() == 1 || PermutationVector.Get<FVirtualVoxel>() == 0))
+		if (PermutationVector.Get<FUseSceneColor>() == 1 && PermutationVector.Get<FSampleSet>() == 1)
 		{
 			return false;
 		}
@@ -198,7 +193,6 @@ public:
 		if (PermutationVector.Get<FUseSceneColor>() == 1)
 		{
 			PermutationVector.Set<FSampleSet>(0);
-			PermutationVector.Set<FVirtualVoxel>(1);
 		}
 
 		return PermutationVector;
@@ -206,12 +200,7 @@ public:
 
 	BEGIN_SHADER_PARAMETER_STRUCT(FParameters, )
 
-		SHADER_PARAMETER(FVector, Voxel_MinAABB)
 		SHADER_PARAMETER(uint32, Voxel_MacroGroupId)
-		SHADER_PARAMETER(FVector, Voxel_MaxAABB)
-		SHADER_PARAMETER(uint32, Voxel_Resolution)
-		SHADER_PARAMETER(float, Voxel_DensityScale)
-		SHADER_PARAMETER(float, Voxel_DepthBiasScale)
 		SHADER_PARAMETER(float, Voxel_TanConeAngle)
 
 		SHADER_PARAMETER(uint32, MaxVisibilityNodeCount)
@@ -241,7 +230,6 @@ public:
 		SHADER_PARAMETER_RDG_TEXTURE(TEXTURE3D, HairScatteringLUTTexture)
 		SHADER_PARAMETER_SAMPLER(SamplerState, HairLUTSampler)
 
-		SHADER_PARAMETER_RDG_TEXTURE(Texture3D, Voxel_DensityTexture)
 		SHADER_PARAMETER_RDG_BUFFER_UAV(RWBuffer, OutLightingBuffer)
 
 		SHADER_PARAMETER_STRUCT_REF(FViewUniformShaderParameters, ViewUniformBuffer)
@@ -315,8 +303,7 @@ static void AddHairStrandsEnvironmentLightingPassPS(
 	FSceneTextureParameters SceneTextures;
 	SetupSceneTextureParameters(GraphBuilder, &SceneTextures);
 
-	const bool bUseVirtualVoxel = MacroGroupDatas.VirtualVoxelResources.IsValid();
-	const int32 VoxelTraversalType = FMath::Clamp(GHairSkylightingVirtualVoxel, 0, 4);
+	check(MacroGroupDatas.VirtualVoxelResources.IsValid());
 
 	// Render the reflection environment with tiled deferred culling
 	const bool bHasBoxCaptures = (View.NumBoxReflectionCaptures > 0);
@@ -329,13 +316,7 @@ static void AddHairStrandsEnvironmentLightingPassPS(
 	PassParameters->HairScatteringLUTTexture = GraphBuilder.RegisterExternalTexture(InHairLUT.Textures[HairLUTType_DualScattering], TEXT("HairScatteringEnergyLUTTexture"));
 	PassParameters->HairLUTSampler = TStaticSamplerState<SF_Bilinear, AM_Clamp, AM_Clamp, AM_Clamp>::GetRHI();
 	PassParameters->Voxel_MacroGroupId = MacroGroupData.MacroGroupId;
-	if (!bUseVirtualVoxel)
-	{
-		PassParameters->Voxel_MinAABB = MacroGroupData.GetMinBound();
-		PassParameters->Voxel_MaxAABB = MacroGroupData.GetMaxBound();
-		PassParameters->Voxel_Resolution = MacroGroupData.GetResolution();
-		PassParameters->Voxel_DensityTexture = GraphBuilder.RegisterExternalTexture(MacroGroupData.VoxelResources.DensityTexture);
-	}
+
 	const bool bUseSceneColor = SceneColorTexture != nullptr;
 	if (bUseSceneColor)
 	{
@@ -344,8 +325,6 @@ static void AddHairStrandsEnvironmentLightingPassPS(
 	}
 	PassParameters->MaxViewportResolution = VisibilityData.SampleLightingViewportResolution;
 	PassParameters->HairVisibilityNodeCount = GraphBuilder.RegisterExternalTexture(VisibilityData.NodeCount);
-	PassParameters->Voxel_DensityScale = GetHairStrandsVoxelizationDensityScale();
-	PassParameters->Voxel_DepthBiasScale = GetHairStrandsVoxelizationDepthBiasScale();
 	PassParameters->Voxel_TanConeAngle = FMath::Tan(FMath::DegreesToRadians(GetHairStrandsSkyLightingConeAngle()));
 	PassParameters->HairDistanceThreshold = FMath::Max(GHairStrandsSkyLightingDistanceThreshold, 1.f);
 	PassParameters->bHairUseViewHairCount = VisibilityData.ViewHairCountTexture && GHairStrandsSkyLightingUseHairCountTexture ? 1 : 0;
@@ -384,7 +363,6 @@ static void AddHairStrandsEnvironmentLightingPassPS(
 
 	FHairEnvironmentLightingPS::FPermutationDomain PermutationVector;
 	PermutationVector.Set<FHairEnvironmentLighting::FSampleSet>(PassParameters->MultipleScatterSampleCount <= 16 ? 0 : 1);
-	PermutationVector.Set<FHairEnvironmentLighting::FVirtualVoxel>(bUseVirtualVoxel ? VoxelTraversalType : 0);
 	PermutationVector.Set<FHairEnvironmentLighting::FUseSceneColor>(bUseSceneColor ? 1 : 0);
 	PermutationVector = FHairEnvironmentLighting::RemapPermutation(PermutationVector);
 
