Fix debug voxel hair view using wrong density parameters.

diff --git a/Engine/Shaders/Private/HairStrands/HairStrandsVoxelPageRayMarching.usf b/Engine/Shaders/Private/HairStrands/HairStrandsVoxelPageRayMarching.usf
--- a/Engine/Shaders/Private/HairStrands/HairStrandsVoxelPageRayMarching.usf
+++ b/Engine/Shaders/Private/HairStrands/HairStrandsVoxelPageRayMarching.usf
@@ -551,8 +551,7 @@ void MainCS(uint3 DispatchThreadId : SV_DispatchThreadID)
 					const int3 VoxelPageOffset	= VolumeCoord - PageIndexCoord * VirtualVoxel.PageResolution;
 					const int3 VoxelPageCoord	= VoxelPageBase + VoxelPageOffset;
 
-					const float DensityScale	= 1;
-					const FHairTraversalResult Result = GetHairVirtualVoxelDensity(VoxelPageCoord, VirtualVoxel.PageTexture, MipLevel, DensityScale);
+					const FHairTraversalResult Result = GetHairVirtualVoxelDensity(VoxelPageCoord, VirtualVoxel.PageTexture, MipLevel, VirtualVoxel.DensityScale);
 					const float Density = Result.HairCount + (Result.Visibility > 0 ? 0 : 1000000000);
 					if (Density > 0)
 					{
