Fix incorrect Hair Primitive ID interpolation when visibility pass is enabled.

diff --git a/Engine/Shaders/Private/HairStrands/HairStrandsVisibilityPS.usf b/Engine/Shaders/Private/HairStrands/HairStrandsVisibilityPS.usf
--- a/Engine/Shaders/Private/HairStrands/HairStrandsVisibilityPS.usf
+++ b/Engine/Shaders/Private/HairStrands/HairStrandsVisibilityPS.usf
@@ -43,7 +43,7 @@ EARLYDEPTHSTENCIL	// Required for UAV operations to not happen.
 #endif
 void MainVisibility(
 #if HAIR_RENDER_MODE == RENDER_MODE_MSAA_VISIBILITY
-	uint HairPrimitiveId : HAIR_PRIMITIVE_ID
+	nointerpolation uint HairPrimitiveId : HAIR_PRIMITIVE_ID 
 #else
 	centroid FVertexFactoryInterpolantsVSToPS Interpolants
 #endif
