Fix non-hair pixel getting categorization data.

diff --git a/Engine/Shaders/Private/HairStrands/HairStrandsVisibilityCompaction.usf b/Engine/Shaders/Private/HairStrands/HairStrandsVisibilityCompaction.usf
--- a/Engine/Shaders/Private/HairStrands/HairStrandsVisibilityCompaction.usf
+++ b/Engine/Shaders/Private/HairStrands/HairStrandsVisibilityCompaction.usf
@@ -397,7 +397,10 @@ void MainCS(uint GroupIndex : SV_GroupIndex, uint3 DispatchThreadId : SV_Dispatc
 	NodeDesc.Count = SampleDesc.UniqueSampleCount;
 	NodeDesc.Offset = 0;
 
-	InterlockedAdd(AllocationNodeCount, NodeDesc.Count, NodeDesc.Offset);
+	if (NodeDesc.Count > 0)
+	{
+		InterlockedAdd(AllocationNodeCount, NodeDesc.Count, NodeDesc.Offset);
+	}
 	GroupMemoryBarrierWithGroupSync();
 	if (GroupIndex == 0 && AllocationNodeCount > 0)
 	{
@@ -435,15 +438,15 @@ void MainCS(uint GroupIndex : SV_GroupIndex, uint3 DispatchThreadId : SV_Dispatc
 
 			NodeDesc.Count = NodeDesc.Count;
 			Categorization.SampleCount = NodeDesc.Count;
+
+		#if PERMUTATION_VIEWTRANSMITTANCE == 1
+			Categorization.PixelCoverage = TransmittanceToCoverage(ViewTransmittanceTexture.Load(uint3(PixelCoord, 0)));
+		#else
+			Categorization.PixelCoverage = SampleDesc.ValidSampleCount / float(PERMUTATION_MSAACOUNT);
+		#endif
 		}
 	}
 
-#if PERMUTATION_VIEWTRANSMITTANCE == 1
-	Categorization.PixelCoverage = TransmittanceToCoverage(ViewTransmittanceTexture.Load(uint3(PixelCoord, 0)));
-#else
-	Categorization.PixelCoverage = SampleDesc.ValidSampleCount / float(PERMUTATION_MSAACOUNT);
-#endif
-
 	OutCompactNodeIndex[PixelCoord] = EncodeNodeDesc(NodeDesc);
 	OutCategorizationTexture[PixelCoord] = EncodeCategorizationData(Categorization);
 }
@@ -576,7 +579,10 @@ void MainCS(uint GroupIndex : SV_GroupIndex, uint3 DispatchThreadId : SV_Dispatc
 	NodeDesc.Count = SampleDesc.UniqueSampleCount;
 	NodeDesc.Offset = 0;
 
-	InterlockedAdd(AllocationNodeCount, NodeDesc.Count, NodeDesc.Offset);
+	if (NodeDesc.Count > 0)
+	{
+		InterlockedAdd(AllocationNodeCount, NodeDesc.Count, NodeDesc.Offset);
+	}
 	GroupMemoryBarrierWithGroupSync();
 	if (GroupIndex == 0 && AllocationNodeCount > 0)
 	{
@@ -689,4 +695,4 @@ void MainCS(uint GroupIndex : SV_GroupIndex, uint3 DispatchThreadId : SV_Dispatc
 	OutCategorizationTexture[PixelCoord] = EncodeCategorizationData(Categorization);
 }
 
-#endif // PERMUTATION_PPLL or PERMUTATION_VISIBILITY or PERMUTATION_MSAA
\ No newline at end of file
+#endif // PERMUTATION_PPLL or PERMUTATION_VISIBILITY or PERMUTATION_MSAA
