Fix material pass attribute interpolation lookup.

Primitive ID is now fully clear to avoid adding sample from padding border.

diff --git a/Engine/Shaders/Private/HairStrands/HairStrandsMaterialPS.usf b/Engine/Shaders/Private/HairStrands/HairStrandsMaterialPS.usf
--- a/Engine/Shaders/Private/HairStrands/HairStrandsMaterialPS.usf
+++ b/Engine/Shaders/Private/HairStrands/HairStrandsMaterialPS.usf
@@ -43,10 +43,10 @@ FVertexFactoryInterpolantsVSToPS GetInterpolants(uint PrimitiveId, uint HairPrim
 	FVertexFactoryIntermediates	Intermediates = GetVertexFactoryIntermediates(Input);
 	FVertexFactoryInterpolantsVSToPS Interpolants = VertexFactoryGetInterpolantsVSToPS(Input, Intermediates, Empty);
 
-	OutHairVertex.World_P = VertexFactoryGetWorldPositionRaw(Input, Intermediates).xyz;
-	OutHairVertex.World_P_Prev = VertexFactoryGetPreviousWorldPosition(Input, Intermediates).xyz;
+	OutHairVertex.World_P = VertexFactoryGetWorldPositionRaw(Input, Intermediates, false).xyz;
+	OutHairVertex.World_P_Prev = VertexFactoryGetPreviousWorldPosition(Input, Intermediates, false).xyz;
 
-	const FHairControlPoint ControlPoint = GetVertexPosition(Input);
+	const FHairControlPoint ControlPoint = GetVertexPosition(Input, false);
 	OutHairVertex.Radius = ControlPoint.WorldRadius;
 	OutHairVertex.bIsValid = ControlPoint.Type != HAIR_CONTROLPOINT_END;
 
@@ -89,7 +89,7 @@ FEvaluateOutput Evaluate(
 		const float3 B = (Vertex1.World_P - Vertex0.World_P);
 		const float CosA = dot(A, B);
 		const float LengthB2 = dot(B, B);
-		const float SegmentU = Vertex1.bIsValid && LengthB2 > 0 ? saturate(CosA / LengthB2) : 0;
+		const float SegmentU = LengthB2 > 0 ? saturate(CosA / LengthB2) : 0;
 
 		// Compute V
 		float SegmentV = 0.5f;
diff --git a/Engine/Shaders/Private/HairStrands/HairStrandsVertexFactory.ush b/Engine/Shaders/Private/HairStrands/HairStrandsVertexFactory.ush
--- a/Engine/Shaders/Private/HairStrands/HairStrandsVertexFactory.ush
+++ b/Engine/Shaders/Private/HairStrands/HairStrandsVertexFactory.ush
@@ -262,7 +262,7 @@ uint GetHairPrimitiveId(FVertexFactoryInput Input)
 	return GetVertexInfo(Input).HairPrimitiveId;
 }
 
-FHairControlPoint GetVertexPosition(const FVertexInfo VertexInfo)
+FHairControlPoint GetVertexPosition(const FVertexInfo VertexInfo, bool bInvalidJointVertex=true)
 {
 	const uint4 PackedData = HairStrandsVF_PositionBuffer[VertexInfo.VertexIndex];
 	FHairControlPoint Out = UnpackHairControlPoint(
@@ -273,16 +273,16 @@ FHairControlPoint GetVertexPosition(const FVertexInfo VertexInfo)
 	
 	// Create a degenerated quad the end of each strand to cut between each strands
 	// #hair_todo: This is not efficient for short strand like fur for instance. Need to revisit that at some point
-	const bool bIsInvalidQuad = (Out.Type == HAIR_CONTROLPOINT_START && VertexInfo.IsTip == 1) || (Out.Type == HAIR_CONTROLPOINT_END && VertexInfo.IsTip == 0) || VertexInfo.bForceInvalidQuad;
+	const bool bIsInvalidQuad = (bInvalidJointVertex && Out.Type == HAIR_CONTROLPOINT_START && VertexInfo.IsTip == 1) || (bInvalidJointVertex && Out.Type == HAIR_CONTROLPOINT_END && VertexInfo.IsTip == 0) || VertexInfo.bForceInvalidQuad;
 	Out.Position = bIsInvalidQuad ? float3(1.#INF, 1.#INF, 1.#INF) : Out.Position;
 
 	return Out;
 }
 
-FHairControlPoint GetVertexPosition(FVertexFactoryInput Input)
+FHairControlPoint GetVertexPosition(FVertexFactoryInput Input, bool bInvalidJointVertex=true)
 {
 	const FVertexInfo VertexInfo = GetVertexInfo(Input);
-	return GetVertexPosition(VertexInfo);
+	return GetVertexPosition(VertexInfo, bInvalidJointVertex);
 }
 
 float GetWorldStrandRadius(FVertexFactoryInput Input)
@@ -295,7 +295,7 @@ float GetWorldStrandLength(FVertexFactoryInput Input)
 	return GetVertexPosition(Input).WorldLength;
 }
 
-float3 GetVertexPreviousPosition(FVertexFactoryInput Input)
+float3 GetVertexPreviousPosition(FVertexFactoryInput Input, bool bInvalidJointVertex = true)
 {
 	FVertexInfo VertexInfo = GetVertexInfo(Input);
 	const uint4 PackedData = HairStrandsVF_PreviousPositionBuffer[VertexInfo.VertexIndex];
@@ -307,7 +307,7 @@ float3 GetVertexPreviousPosition(FVertexFactoryInput Input)
 
 	// Create a degenerated quad the end of each strand to cut between each strands
 	// #hair_todo: This is not efficient for short strand like fur for instance. Need to revisit that at some point
-	const bool bIsInvalidQuad = (Out.Type == HAIR_CONTROLPOINT_START && VertexInfo.IsTip == 1) || (Out.Type == HAIR_CONTROLPOINT_END && VertexInfo.IsTip == 0) || VertexInfo.bForceInvalidQuad;
+	const bool bIsInvalidQuad = (bInvalidJointVertex && Out.Type == HAIR_CONTROLPOINT_START && VertexInfo.IsTip == 1) || (bInvalidJointVertex && Out.Type == HAIR_CONTROLPOINT_END && VertexInfo.IsTip == 0) || VertexInfo.bForceInvalidQuad;
 	Out.Position = bIsInvalidQuad ? float3(1.#INF, 1.#INF, 1.#INF) : Out.Position;
 
 	return Out.Position;
@@ -578,20 +578,20 @@ float4 VertexFactoryGetWorldPosition(FPositionOnlyVertexFactoryInput Input)
 }
 
 // @return translated world position (without quad extension/reorientation).This is used only for velocity computation
-float4 VertexFactoryGetWorldPositionRaw(FVertexFactoryInput Input, FVertexFactoryIntermediates Intermediates)
+float4 VertexFactoryGetWorldPositionRaw(FVertexFactoryInput Input, FVertexFactoryIntermediates Intermediates, bool bInvalidJointVertex = true)
 {
-	return CalcWorldPosition(GetVertexPosition(Input).Position, Intermediates.PrimitiveId);
+	return CalcWorldPosition(GetVertexPosition(Input, bInvalidJointVertex).Position, Intermediates.PrimitiveId);
 }
 
 // @return previous translated world position
-float4 VertexFactoryGetPreviousWorldPosition(FVertexFactoryInput Input, FVertexFactoryIntermediates Intermediates)
+float4 VertexFactoryGetPreviousWorldPosition(FVertexFactoryInput Input, FVertexFactoryIntermediates Intermediates, bool bInvalidJointVertex = true)
 {
 	float4x4 PreviousLocalToWorldTranslated = GetPrimitiveData(Intermediates.PrimitiveId).PreviousLocalToWorld;
 	PreviousLocalToWorldTranslated[3][0] += ResolvedView.PrevPreViewTranslation.x;
 	PreviousLocalToWorldTranslated[3][1] += ResolvedView.PrevPreViewTranslation.y;
 	PreviousLocalToWorldTranslated[3][2] += ResolvedView.PrevPreViewTranslation.z;
 
-	const float3 VertexPosition = GetVertexPreviousPosition(Input);
+	const float3 VertexPosition = GetVertexPreviousPosition(Input, bInvalidJointVertex);
 	return mul(float4(VertexPosition,1), PreviousLocalToWorldTranslated);
 }
 
diff --git a/Engine/Source/Runtime/Renderer/Private/HairStrands/HairStrandsVisibility.cpp b/Engine/Source/Runtime/Renderer/Private/HairStrands/HairStrandsVisibility.cpp
--- a/Engine/Source/Runtime/Renderer/Private/HairStrands/HairStrandsVisibility.cpp
+++ b/Engine/Source/Runtime/Renderer/Private/HairStrands/HairStrandsVisibility.cpp
@@ -1185,7 +1185,7 @@ static void AddClearGraphicPass(
 
 	TShaderMapRef<FPostProcessVS> VertexShader(View->ShaderMap);
 	TShaderMapRef<FClearUIntGraphicPS> PixelShader(View->ShaderMap);
-	const FIntRect Viewport = View->ViewRect;
+	const FIntRect Viewport = FIntRect(FIntPoint(0, 0), OutTarget->Desc.Extent);// View->ViewRect;
 	const FIntPoint Resolution = OutTarget->Desc.Extent;
 
 	ClearUnusedGraphResources(PixelShader, Parameters);
