Fix sample count debug mode.

diff --git a/Engine/Shaders/Private/HairStrands/HairStrandsDebug.usf b/Engine/Shaders/Private/HairStrands/HairStrandsDebug.usf
--- a/Engine/Shaders/Private/HairStrands/HairStrandsDebug.usf
+++ b/Engine/Shaders/Private/HairStrands/HairStrandsDebug.usf
@@ -41,19 +41,19 @@ void MainPS(
 	const float2 UV = Input.UV;
 	const float2 PixelPos = UV * OutputResolution;
 
-	if (DebugMode == DEBUG_MODE_SAMPLE_COUNT)
-	{
-		const float HairCount = HairCountUintTexture.Load(uint3(PixelPos, 0)) / 1000.f;
-		//const float HairCount = HairCountTexture.Load(uint3(PixelPos, 0));
-
-		if (HairCount==0)
-			discard;
-		OutColor = float4(float3(1, 0, 0), 1);
-
-		OutColor = float4(GetColorCode(HairCount / float(32)), 1);
-
-		return;
-	}
+	//if (DebugMode == DEBUG_MODE_SAMPLE_COUNT)
+	//{
+	//	const float HairCount = HairCountUintTexture.Load(uint3(PixelPos, 0)) / 1000.f;
+	//	//const float HairCount = HairCountTexture.Load(uint3(PixelPos, 0));
+	//
+	//	if (HairCount==0)
+	//		discard;
+	//	OutColor = float4(float3(1, 0, 0), 1);
+	//
+	//	OutColor = float4(GetColorCode(HairCount / float(32)), 1);
+	//
+	//	return;
+	//}
 
 	const FCategorizationData CatData = DecodeCategorizationData(CategorizationTexture.Load(uint3(PixelPos,0)));
 	const bool bIsValid = CatData.PixelCoverage > 0 || DebugMode == DEBUG_MODE_TILE;
