Add uniform sampling for hair environment lighting.

diff --git a/Engine/Shaders/Private/HairStrands/HairStrandsEnvironmentLighting.usf b/Engine/Shaders/Private/HairStrands/HairStrandsEnvironmentLighting.usf
--- a/Engine/Shaders/Private/HairStrands/HairStrandsEnvironmentLighting.usf
+++ b/Engine/Shaders/Private/HairStrands/HairStrandsEnvironmentLighting.usf
@@ -45,11 +45,25 @@
 #define VOXEL_TRAVERSAL_DEBUG 0
 #include "HairStrandsVoxelPageTraversal.ush"
 
+#define INTEGRATION_SCENECOLOR 0
+#define INTEGRATION_ADHOC 1
+#define INTEGRATION_UNIFORM 2
+
 // Note: 
 // * This version of environement lighting is very much in WIP state, with a lot of option to tune tweak to explore 
 //   what kind of approximation works the best for certains type of grooms and lighting condition
 // * There is a few fudges factor as well (roughness mapping, TT density, reflected vector orientations, ...)
 
+float GetPreExposure()
+{
+	#if USE_PREEXPOSURE
+	const float PreExposure = View.PreExposure;
+	#else
+	const float PreExposure = 1.f;
+	#endif
+	return PreExposure;
+}
+
 float3 GatherRadiance(float3 RayDirection, float Roughness, uint2 PixelPosition, float PixelSceneDepth, float3 AbsoluteWorldPosition)
 {
 	const float IndirectIrradiance = 1;
@@ -135,14 +149,16 @@ FHairTraversalResult InternalComputeHairCountVoxel(
 //////////////////////////////////////////////////////////////////////////////////////////////////////////////////////////
 //////////////////////////////////////////////////////////////////////////////////////////////////////////////////////////
 //////////////////////////////////////////////////////////////////////////////////////////////////////////////////////////
+#if PERMUTATION_INTEGRATION_TYPE == INTEGRATION_ADHOC
+
+#define ADHOC_SAMPLING_R 1
+#define ADHOC_SAMPLING_TT 1
+#define ADHOC_SAMPLING_TRT 1
+#define ADHOC_SAMPLING_MS 1
 
 float3 ReflectionEnvironment(FGBufferData GBuffer, float2 ScreenPosition, float2 SvPosition, uint PrimitiveId)
 {
-#if USE_PREEXPOSURE
-	const float PreExposure = View.PreExposure;
-#else
-	const float PreExposure = 1.f;
-#endif
+	const float PreExposure = GetPreExposure();
 	const float  MaxDistance = 10000.f; // 100m
 	const float3 WorldPosition	= mul(float4(ScreenPosition * GBuffer.Depth, GBuffer.Depth, 1), View.ScreenToWorld).xyz;
 	const float3 CameraToPixel	= normalize(WorldPosition - View.WorldCameraOrigin);
@@ -165,7 +181,7 @@ float3 ReflectionEnvironment(FGBufferData GBuffer, float2 ScreenPosition, float2
 	float3 UnoccludedN = -V;
 	float3 AvgDualScattering = 0;
 	const float MaxSearchDistance = HairDistanceThreshold; // in cm
-	#if 1
+	#if ADHOC_SAMPLING_MS || ADHOC_SAMPLING_R || ADHOC_SAMPLING_TT || ADHOC_SAMPLING_TRT
 	if ((HairComponents & HAIR_COMPONENT_R) || (HairComponents & HAIR_COMPONENT_TRT) || (HairComponents & HAIR_COMPONENT_TT))
 	{
 		#if PERMUTATION_SAMPLESET == 0
@@ -267,7 +283,7 @@ float3 ReflectionEnvironment(FGBufferData GBuffer, float2 ScreenPosition, float2
 	////////////////////////////////////////////////////////////////////////////////////////////////////////////////////////////////////////////
 	// MS ==============================================================================
 	float3 ColorMS = 0;
-	#if 1
+	#if ADHOC_SAMPLING_MS
 	if ((HairComponents & HAIR_COMPONENT_R) || (HairComponents & HAIR_COMPONENT_TRT) || (HairComponents & HAIR_COMPONENT_TT))
 	{
 		float3 N = UnoccludedN;
@@ -308,7 +324,7 @@ float3 ReflectionEnvironment(FGBufferData GBuffer, float2 ScreenPosition, float2
 	////////////////////////////////////////////////////////////////////////////////////////////////////////////////////////////////////////////
 	// R ==============================================================================
 	float3 ColorR = 0;
-	#if 1
+	#if ADHOC_SAMPLING_R || ADHOC_SAMPLING_TRT
 	if ((HairComponents & HAIR_COMPONENT_R) || (HairComponents & HAIR_COMPONENT_TRT))
 	{
 		float3 R = 0;
@@ -384,7 +400,7 @@ float3 ReflectionEnvironment(FGBufferData GBuffer, float2 ScreenPosition, float2
 	////////////////////////////////////////////////////////////////////////////////////////////////////////////////////////////////////////////
 	// TT ==============================================================================
 	float3 ColorTT = 0;
-	#if 1
+	#if ADHOC_SAMPLING_TT
 	if (HairComponents & HAIR_COMPONENT_TT)
 	{
 		const float3 L = -V;
@@ -461,8 +477,130 @@ float3 ReflectionEnvironment(FGBufferData GBuffer, float2 ScreenPosition, float2
 	// Transform NaNs to black, transform negative colors to black.
 	return -min(-Color.rgb, 0.0);
 }
+#endif // PERMUTATION_INTEGRATION_TYPE == INTEGRATION_ADHOC
+
+//////////////////////////////////////////////////////////////////////////////////////////////////////////////////////////
+// Reference light integration using uniform spherical sampling
+
+#if PERMUTATION_INTEGRATION_TYPE == INTEGRATION_UNIFORM
+float3 ReflectionEnvironmentUniformSampling(
+	FGBufferData GBuffer, 
+	float2 ScreenPosition, 
+	float2 SvPosition, 
+	uint PrimitiveId, 
+	bool bDebugEnabled)
+{
+	const float PreExposure = GetPreExposure();
+	const float3 WorldPosition = mul(float4(ScreenPosition * GBuffer.Depth, GBuffer.Depth, 1), View.ScreenToWorld).xyz;
+	const float3 CameraToPixel = normalize(WorldPosition - View.WorldCameraOrigin);
+	const uint2  PixelPosition = SvPosition - ResolvedView.ViewRectMin.xy;
+	const float  PixelSceneDepth = GBuffer.Depth;
+
+	const uint PrimitiveIdMod8 = (PrimitiveId % 8);
+	const float  VoxelMaxSize = VirtualVoxel.VoxelWorldSize;
+	const float3 VoxelOffset = VoxelMaxSize * GetHairVoxelJitter(SvPosition, View.StateFrameIndexMod8 + PrimitiveIdMod8);
 
-#if USE_SCENE_COLOR
+	const float3 V = -CameraToPixel;
+	const float3 T = GBuffer.WorldNormal;
+
+	// Trace/sample transmittance in various direction to have estimate of where is occlusion & computing the integrated dual scattering value
+	const float MaxSearchDistance = HairDistanceThreshold; // in cm
+#if PERMUTATION_SAMPLESET == 0
+	// Poisson disk position http://developer.download.nvidia.com/whitepapers/2008/PCSS_Integration.pdf
+	float2 PoissonDisk[16] =
+	{
+		float2(-0.94201624, -0.39906216),
+		float2(0.94558609, -0.76890725),
+		float2(-0.094184101, -0.92938870),
+		float2(0.34495938, 0.29387760),
+		float2(-0.91588581, 0.45771432),
+		float2(-0.81544232, -0.87912464),
+		float2(-0.38277543, 0.27676845),
+		float2(0.97484398, 0.75648379),
+		float2(0.44323325, -0.97511554),
+		float2(0.53742981, -0.47373420),
+		float2(-0.26496911, -0.41893023),
+		float2(0.79197514, 0.19090188),
+		float2(-0.24188840, 0.99706507),
+		float2(-0.81409955, 0.91437590),
+		float2(0.19984126, 0.78641367),
+		float2(0.14383161, -0.14100790)
+	};
+	const uint SampleCount = clamp(MultipleScatterSampleCount, uint(1), uint(16));
+#else
+	const uint SampleCount = max(MultipleScatterSampleCount, uint(1));
+#endif
+
+	float2 RandOffset = 0;
+	if (JitterSphericalIntegration > 0)
+	{
+		RandOffset += float2(
+			InterleavedGradientNoise(SvPosition, PrimitiveIdMod8 * 11),
+			InterleavedGradientNoise(SvPosition, PrimitiveIdMod8 * 7));
+	}
+
+	float3 AccLighting = 0;
+	for (uint SampleIt = 0; SampleIt < SampleCount; ++SampleIt)
+	{
+		#if PERMUTATION_SAMPLESET == 0
+		const float2 Rand2 = frac((PoissonDisk[SampleIt] + float2(1, 1)) * 0.5f + RandOffset);
+		#else // PERMUTATION_SAMPLESET == 1
+		const float2 Rand2 = frac(Hammersley(SampleIt, SampleCount, uint2(0, 0)) + RandOffset);
+		#endif
+
+		// Depth bias
+		const float3 SampleVoxelOffset = VoxelMaxSize * GetHairVoxelJitter(SvPosition, PrimitiveIdMod8 + SampleIt);
+		const float3 SampleWorldPosition = WorldPosition + SampleVoxelOffset;
+
+		const float3 SampleL = UniformSampleSphere(frac(Rand2)).xyz;
+		const float3 SampleLightPosition = WorldPosition + SampleL * MaxSearchDistance;
+
+		const float SinLightAngle = dot(SampleL, T);
+		const FHairTraversalResult Result = InternalComputeHairCountVoxel(
+			SampleWorldPosition, 
+			SampleLightPosition, 
+			1,
+			0, 
+			MaxSearchDistance, 
+			true, 
+			false/*bDebugEnabled*/);
+
+		const FHairTransmittanceMask TransmittanceMask = ComputeTransmittanceFromHairCount(
+			Result.HairCount, 
+			Result.Visibility, 
+			GBuffer, 
+			SinLightAngle, 
+			HairScatteringLUTTexture, 
+			HairLUTSampler);
+
+		FHairTransmittanceData TransmittanceData = GetHairTransmittance(
+			V,
+			SampleL,
+			GBuffer,
+			TransmittanceMask,
+			HairScatteringLUTTexture,
+			HairLUTSampler,
+			HairComponents);	
+
+		const float Area = 0;
+		AccLighting += 
+			HairShading(GBuffer, SampleL, V, T, 1, TransmittanceData, 0, Area, uint2(0, 0)) * 
+			Result.Visibility * 
+			PreExposure * 
+			GatherRadiance(SampleL, 0.1f, PixelPosition, PixelSceneDepth, WorldPosition);
+	}
+	const float SphereInt = 4 * PI;
+	const float3 Color = SphereInt * AccLighting / SampleCount;
+	
+	// Transform NaNs to black, transform negative colors to black.
+	return -min(-Color, 0.0);
+}
+#endif // #if PERMUTATION_INTEGRATION_TYPE == INTEGRATION_UNIFORM
+
+//////////////////////////////////////////////////////////////////////////////////////////////////////////////////////////
+// Simple lookup of the scene color for single/low count of hair, close to the surface
+
+#if PERMUTATION_INTEGRATION_TYPE == INTEGRATION_SCENECOLOR
 Texture2D<float4> SceneColorTexture;
 Texture2D<uint4> HairCategorizationTexture;
 
@@ -490,7 +628,7 @@ float3 ComputeThinTransmission(FGBufferData GBuffer, float2 ScreenPosition, floa
 	
 	const float SinViewAngle = dot(T, V);
 	const HairAverageScattering Energy = SampleHairLUT(HairScatteringLUTTexture, HairLUTSampler, GBuffer.BaseColor, GBuffer.Roughness, SinViewAngle);
-	float3 Color = Energy.A_front * Luminance * PixelCoverage;
+	float3 Color = Energy.A_front * Luminance.xyz * PixelCoverage;
 	Color *= PreExposure;
 
 	return -min(-Color, 0.0);
@@ -569,23 +707,29 @@ void MainPS(
 	GBuffer.CustomData = float4(HairDualScatteringRoughnessOverride, 0, Sample.Backlit, 0);
 	GBuffer.IndirectIrradiance = 1000000;
 
-	bool bEnabled = true;
 	const bool bIsMacroGroupCompatible = Sample.MacroGroupId == Voxel_MacroGroupId;
-#if USE_SCENE_COLOR
-	bEnabled = Sample.bScatterSceneLighting;
+#if PERMUTATION_INTEGRATION_TYPE == INTEGRATION_SCENECOLOR
+	const bool bEnabled = Sample.bScatterSceneLighting;
+#else
+	const bool bEnabled = true;
 #endif
+
 	if (bIsMacroGroupCompatible && bEnabled)
 	{
+		const float3 SkyLighting =
+		#if PERMUTATION_INTEGRATION_TYPE == INTEGRATION_SCENECOLOR
+			ComputeThinTransmission(GBuffer, ScreenPosition, SvPosition, Sample.PrimitiveId);
+		#elif PERMUTATION_INTEGRATION_TYPE == INTEGRATION_ADHOC
+			ReflectionEnvironment(GBuffer, ScreenPosition, SvPosition, Sample.PrimitiveId);
+		#elif PERMUTATION_INTEGRATION_TYPE == INTEGRATION_UNIFORM
+			ReflectionEnvironmentUniformSampling(GBuffer, ScreenPosition, SvPosition, Sample.PrimitiveId, false);
+		#endif
+
 		const float SampleCoverage = saturate(From8bitCoverage(Sample.Coverage8bit));
-#if USE_SCENE_COLOR
-		const float3 SkyLighting = ComputeThinTransmission(GBuffer, ScreenPosition, SvPosition, Sample.PrimitiveId);
-#else
-		const float3 SkyLighting = ReflectionEnvironment(GBuffer, ScreenPosition, SvPosition, Sample.PrimitiveId);
-#endif
 
 		FLightAccumulator LightAccumulator = (FLightAccumulator)0;
 		LightAccumulator_Add(LightAccumulator, SkyLighting, SkyLighting, 1.0f, false);
-		OutColor.xyz = LightAccumulator_GetResult(LightAccumulator) * SampleCoverage;
+		OutColor.xyz = LightAccumulator_GetResult(LightAccumulator).xyz * SampleCoverage;
 		OutColor.a = SampleCoverage;
 	}
 }
diff --git a/Engine/Source/Runtime/Renderer/Private/HairStrands/HairStrandsEnvironment.cpp b/Engine/Source/Runtime/Renderer/Private/HairStrands/HairStrandsEnvironment.cpp
--- a/Engine/Source/Runtime/Renderer/Private/HairStrands/HairStrandsEnvironment.cpp
+++ b/Engine/Source/Runtime/Renderer/Private/HairStrands/HairStrandsEnvironment.cpp
@@ -54,6 +54,18 @@ static float GHairStrandsSkyLightingDistanceThreshold = 10;
 static FAutoConsoleVariableRef CVarHairStrandsSkyAOThreshold(TEXT("r.HairStrands.SkyAO.DistanceThreshold"), GHairStrandsSkyAODistanceThreshold, TEXT("Max distance for occlusion search."), ECVF_Scalability | ECVF_RenderThreadSafe);
 static FAutoConsoleVariableRef CVarHairStrandsSkyLightingDistanceThreshold(TEXT("r.HairStrands.SkyLighting.DistanceThreshold"), GHairStrandsSkyLightingDistanceThreshold, TEXT("Max distance for occlusion search."), ECVF_Scalability | ECVF_RenderThreadSafe);
 
+static int32 GHairStrandsSkyLighting_IntegrationType = 0;
+static FAutoConsoleVariableRef CVarHairStrandsSkyLighting_IntegrationType(TEXT("r.HairStrands.SkyLighting.IntegrationType"), GHairStrandsSkyLighting_IntegrationType, TEXT("Hair env. lighting integration type (0:Adhoc, 1:Uniform."), ECVF_Scalability | ECVF_RenderThreadSafe);
+
+///////////////////////////////////////////////////////////////////////////////////////////////////
+
+enum class EHairLightingIntegrationType : uint8
+{
+	SceneColor = 0,
+	AdHoc = 1,
+	Uniform = 2
+};
+
 bool GetHairStrandsSkyLightingEnable() { return GHairSkylightingEnable > 0; }
 static bool GetHairStrandsSkyAOEnable() { return GHairSkyAOEnable > 0; }
 static float GetHairStrandsSkyLightingConeAngle() { return FMath::Max(0.f, GHairSkylightingConeAngle); }
@@ -163,8 +175,8 @@ class FHairEnvironmentLighting
 public:
 	class FGroupSize : SHADER_PERMUTATION_INT("PERMUTATION_GROUP_SIZE", 2);
 	class FSampleSet : SHADER_PERMUTATION_INT("PERMUTATION_SAMPLESET", 2);
-	class FUseSceneColor : SHADER_PERMUTATION_INT("USE_SCENE_COLOR", 2);
-	using FPermutationDomain = TShaderPermutationDomain<FSampleSet, FUseSceneColor>;
+	class FIntegrationType : SHADER_PERMUTATION_INT("PERMUTATION_INTEGRATION_TYPE", 3);
+	using FPermutationDomain = TShaderPermutationDomain<FSampleSet, FIntegrationType>;
 
 	static void ModifyCompilationEnvironment(const FGlobalShaderPermutationParameters& Parameters, FShaderCompilerEnvironment& OutEnvironment)
 	{
@@ -179,7 +191,7 @@ public:
 		}
 
 		const FPermutationDomain PermutationVector(Parameters.PermutationId);
-		if (PermutationVector.Get<FUseSceneColor>() == 1 && PermutationVector.Get<FSampleSet>() == 1)
+		if (PermutationVector.Get<FIntegrationType>() == uint32(EHairLightingIntegrationType::SceneColor) && PermutationVector.Get<FSampleSet>() == 1)
 		{
 			return false;
 		}
@@ -190,7 +202,7 @@ public:
 	static FPermutationDomain RemapPermutation(FPermutationDomain PermutationVector)
 	{
 		// If not rendering the sky, ignore the fastsky and sundisk permutations
-		if (PermutationVector.Get<FUseSceneColor>() == 1)
+		if (PermutationVector.Get<FIntegrationType>() == uint32(EHairLightingIntegrationType::SceneColor))
 		{
 			PermutationVector.Set<FSampleSet>(0);
 		}
@@ -317,12 +329,23 @@ static void AddHairStrandsEnvironmentLightingPassPS(
 	PassParameters->HairLUTSampler = TStaticSamplerState<SF_Bilinear, AM_Clamp, AM_Clamp, AM_Clamp>::GetRHI();
 	PassParameters->Voxel_MacroGroupId = MacroGroupData.MacroGroupId;
 
+	EHairLightingIntegrationType IntegrationType = EHairLightingIntegrationType::AdHoc;
 	const bool bUseSceneColor = SceneColorTexture != nullptr;
 	if (bUseSceneColor)
 	{
+		IntegrationType = EHairLightingIntegrationType::SceneColor;
 		PassParameters->SceneColorTexture = SceneColorTexture;
 		PassParameters->HairCategorizationTexture = GraphBuilder.RegisterExternalTexture(VisibilityData.CategorizationTexture ? VisibilityData.CategorizationTexture : GSystemTextures.BlackDummy);
 	}
+	else
+	{
+		switch (GHairStrandsSkyLighting_IntegrationType)
+		{
+			case 0: IntegrationType = EHairLightingIntegrationType::AdHoc; break;
+			case 1: IntegrationType = EHairLightingIntegrationType::Uniform; break;
+		}
+	}
+
 	PassParameters->MaxViewportResolution = VisibilityData.SampleLightingViewportResolution;
 	PassParameters->HairVisibilityNodeCount = GraphBuilder.RegisterExternalTexture(VisibilityData.NodeCount);
 	PassParameters->Voxel_TanConeAngle = FMath::Tan(FMath::DegreesToRadians(GetHairStrandsSkyLightingConeAngle()));
@@ -363,7 +386,7 @@ static void AddHairStrandsEnvironmentLightingPassPS(
 
 	FHairEnvironmentLightingPS::FPermutationDomain PermutationVector;
 	PermutationVector.Set<FHairEnvironmentLighting::FSampleSet>(PassParameters->MultipleScatterSampleCount <= 16 ? 0 : 1);
-	PermutationVector.Set<FHairEnvironmentLighting::FUseSceneColor>(bUseSceneColor ? 1 : 0);
+	PermutationVector.Set<FHairEnvironmentLighting::FIntegrationType>(uint32(IntegrationType));
 	PermutationVector = FHairEnvironmentLighting::RemapPermutation(PermutationVector);
 
 	FIntPoint ViewportResolution = VisibilityData.SampleLightingViewportResolution;
