Fix shader compile fail when 'GBUFFER_HAS_TANGENT' is enabled

diff --git a/Engine/Shaders/Private/RayTracing/RayTracingReflectionsGenerateRaysCS.usf b/Engine/Shaders/Private/RayTracing/RayTracingReflectionsGenerateRaysCS.usf
--- a/Engine/Shaders/Private/RayTracing/RayTracingReflectionsGenerateRaysCS.usf
+++ b/Engine/Shaders/Private/RayTracing/RayTracingReflectionsGenerateRaysCS.usf
@@ -5,6 +5,7 @@
 #include "../SceneTextureParameters.ush"
 #include "../ScreenSpaceDenoise/SSDPublic.ush"
 #include "../MortonCode.ush"
+#include "../BRDF.ush"
 
 #include "RayTracingCommon.ush"
 #include "RayTracingReflectionsCommon.ush"
@@ -83,7 +84,9 @@ void GenerateReflectionRaysCS(
 	const uint2 PixelPos         = TileBasePixelPos + uint2(RayIndexInTile % TileSize.x, RayIndexInTile / TileSize.x);
 	const float2 ScreenPos       = ViewportUVToScreenPos((PixelPos.xy - View.ViewRectMin.xy) * View.ViewSizeAndInvSize.zw);
 	const float2 UV              = PixelPos.xy * View.BufferSizeAndInvSize.zw;
-	const FGBufferData GBuffer   = GetGBufferDataFromSceneTextures(UV);
+
+	// not 'const', may need to be modified below
+	FGBufferData GBuffer         = GetGBufferDataFromSceneTextures(UV);
 	const float SceneDepth       = GBuffer.Depth;
 
 	const float3 PositionTranslatedWorld = mul(float4(ScreenPos * SceneDepth, SceneDepth, 1), View.ScreenToTranslatedWorld).xyz;
