Fix a bug that the IBL is missing when ALLOW_STATIC_LIGHTING is disabled on mobile platform.

diff --git a/Engine/Shaders/Private/MobileBasePassPixelShader.usf b/Engine/Shaders/Private/MobileBasePassPixelShader.usf
--- a/Engine/Shaders/Private/MobileBasePassPixelShader.usf
+++ b/Engine/Shaders/Private/MobileBasePassPixelShader.usf
@@ -245,11 +245,14 @@ half3 BlendReflectionCaptures(FMaterialPixelParameters MaterialParameters, half
 	ReflectionCubemapAverageBrightness = 0.0f;
 #endif
 
+#if ALLOW_STATIC_LIGHTING
 	//To keep ImageBasedReflectionLighting conherence with PC, use ComputeMixingWeight instead of InvReflectionAverageBrightness to calulate the IBL constribution
 	half3 MixingWeight = MobileComputeMixingWeight3(IndirectIrradiance, ReflectionCubemapAverageBrightness, Roughness);
 
 	SpecularIBL = SpecularIBL0 * MixingWeight.x + SpecularIBL1 * MixingWeight.y + SpecularIBL2 * MixingWeight.z;
-
+#else
+	SpecularIBL = SpecularIBL0 + SpecularIBL1 + SpecularIBL2;
+#endif
 	return SpecularIBL;
 }
 #endif // HQ_REFLECTIONS
@@ -313,8 +316,10 @@ half3 GetImageBasedReflectionLighting(FMaterialPixelParameters MaterialParameter
 		SpecularIBL = RGBMDecode(SpecularIBLSample, 16.0);
 		SpecularIBL = SpecularIBL * SpecularIBL;
 		
+#if ALLOW_STATIC_LIGHTING
 		//To keep ImageBasedReflectionLighting conherence with PC, use ComputeMixingWeight instead of InvReflectionAverageBrightness to calulate the IBL constribution
 		SpecularIBL *= MobileComputeMixingWeight(IndirectIrradiance, MobileReflectionCapture.Params.x, Roughness);
+#endif
 
 	}
 
