Remove the skylight shader permutation.
Saves some shader compiles at the cost of a shader branch. It was not possible to measure any difference in performance
rolando.caloca, chris.bunner, arciel.rekman, lukas.hermanss

diff --git a/Engine/Shaders/Private/BasePassPixelShader.usf b/Engine/Shaders/Private/BasePassPixelShader.usf
--- a/Engine/Shaders/Private/BasePassPixelShader.usf
+++ b/Engine/Shaders/Private/BasePassPixelShader.usf
@@ -261,79 +261,81 @@ void GetSkyLighting(FMaterialPixelParameters MaterialParameters, VTPageTableResu
 	OutDiffuseLighting = 0;
 	OutSubsurfaceLighting = 0;
 
-#if ENABLE_SKY_LIGHT
-
-	float SkyVisibility = 1;
-	float GeometryTerm = 1;
-	float3 SkyLightingNormal = WorldNormal;
+	uint SkyLightMask = MATERIALBLENDING_TRANSLUCENT ? 2 : 1;
+	bool EnableSkyLight = 0 != (View.EnableSkyLight & SkyLightMask) && !MATERIAL_SHADINGMODEL_UNLIT;
+	BRANCH
+	if(EnableSkyLight)
+	{
+		float SkyVisibility = 1;
+		float GeometryTerm = 1;
+		float3 SkyLightingNormal = WorldNormal;
 	
-	#if HQ_TEXTURE_LIGHTMAP || CACHED_POINT_INDIRECT_LIGHTING || CACHED_VOLUME_INDIRECT_LIGHTING || PRECOMPUTED_IRRADIANCE_VOLUME_LIGHTING
-		BRANCH
-		if (ShouldSkyLightApplyPrecomputedBentNormalShadowing())
-		{
-			float3 NormalizedBentNormal = SkyLightingNormal;
+		#if HQ_TEXTURE_LIGHTMAP || CACHED_POINT_INDIRECT_LIGHTING || CACHED_VOLUME_INDIRECT_LIGHTING || PRECOMPUTED_IRRADIANCE_VOLUME_LIGHTING
+			BRANCH
+			if (ShouldSkyLightApplyPrecomputedBentNormalShadowing())
+			{
+				float3 NormalizedBentNormal = SkyLightingNormal;
 
-			#if PRECOMPUTED_IRRADIANCE_VOLUME_LIGHTING
+				#if PRECOMPUTED_IRRADIANCE_VOLUME_LIGHTING
 			
-				float3 SkyBentNormal = GetVolumetricLightmapSkyBentNormal(SkyOcclusionUV3D);
-				SkyVisibility = length(SkyBentNormal);
-				NormalizedBentNormal = SkyBentNormal / max(SkyVisibility, .0001f);
+					float3 SkyBentNormal = GetVolumetricLightmapSkyBentNormal(SkyOcclusionUV3D);
+					SkyVisibility = length(SkyBentNormal);
+					NormalizedBentNormal = SkyBentNormal / max(SkyVisibility, .0001f);
 
-			#elif HQ_TEXTURE_LIGHTMAP
+				#elif HQ_TEXTURE_LIGHTMAP
 
-				// Bent normal from precomputed texture
-				float4 WorldSkyBentNormalAndOcclusion = GetSkyBentNormalAndOcclusion(LightmapVTPageTableResult, LightmapUV * float2(1, 2), LightmapDataIndex, MaterialParameters.SvPosition.xy);
-				// Renormalize as vector was quantized and compressed
-				NormalizedBentNormal = normalize(WorldSkyBentNormalAndOcclusion.xyz);
-				SkyVisibility = WorldSkyBentNormalAndOcclusion.w;
+					// Bent normal from precomputed texture
+					float4 WorldSkyBentNormalAndOcclusion = GetSkyBentNormalAndOcclusion(LightmapVTPageTableResult, LightmapUV * float2(1, 2), LightmapDataIndex, MaterialParameters.SvPosition.xy);
+					// Renormalize as vector was quantized and compressed
+					NormalizedBentNormal = normalize(WorldSkyBentNormalAndOcclusion.xyz);
+					SkyVisibility = WorldSkyBentNormalAndOcclusion.w;
 
-			#elif CACHED_POINT_INDIRECT_LIGHTING || CACHED_VOLUME_INDIRECT_LIGHTING
+				#elif CACHED_POINT_INDIRECT_LIGHTING || CACHED_VOLUME_INDIRECT_LIGHTING
 
-				// Bent normal from the indirect lighting cache - one value for the whole object
-				if (View.IndirectLightingCacheShowFlag > 0.0f)
-				{
-					NormalizedBentNormal = IndirectLightingCache.PointSkyBentNormal.xyz;
-					SkyVisibility = IndirectLightingCache.PointSkyBentNormal.w;
-				}
+					// Bent normal from the indirect lighting cache - one value for the whole object
+					if (View.IndirectLightingCacheShowFlag > 0.0f)
+					{
+						NormalizedBentNormal = IndirectLightingCache.PointSkyBentNormal.xyz;
+						SkyVisibility = IndirectLightingCache.PointSkyBentNormal.w;
+					}
 
-			#endif
+				#endif
 
-			#if (MATERIALBLENDING_TRANSLUCENT || MATERIALBLENDING_ADDITIVE) && (TRANSLUCENCY_LIGHTING_VOLUMETRIC_NONDIRECTIONAL || TRANSLUCENCY_LIGHTING_VOLUMETRIC_PERVERTEX_NONDIRECTIONAL)
-				// NonDirectional lighting can't depend on the normal
-				SkyLightingNormal = NormalizedBentNormal;
-			#else
+				#if (MATERIALBLENDING_TRANSLUCENT || MATERIALBLENDING_ADDITIVE) && (TRANSLUCENCY_LIGHTING_VOLUMETRIC_NONDIRECTIONAL || TRANSLUCENCY_LIGHTING_VOLUMETRIC_PERVERTEX_NONDIRECTIONAL)
+					// NonDirectional lighting can't depend on the normal
+					SkyLightingNormal = NormalizedBentNormal;
+				#else
 				
-				// Weight toward the material normal to increase directionality
-				float BentNormalWeightFactor = 1 - (1 - SkyVisibility) * (1 - SkyVisibility);
-
-				// We are lerping between the inputs of two lighting scenarios based on occlusion
-				// In the mostly unoccluded case, evaluate sky lighting with the material normal, because it has higher detail
-				// In the mostly occluded case, evaluate sky lighting with the bent normal, because it is a better representation of the incoming lighting
-				// Then treat the lighting evaluated along the bent normal as an area light, so we must apply the lambert term
-				SkyLightingNormal = lerp(NormalizedBentNormal, WorldNormal, BentNormalWeightFactor);
-
-				float DotProductFactor = lerp(saturate(dot(NormalizedBentNormal, WorldNormal)), 1, BentNormalWeightFactor);
-				// Account for darkening due to the geometry term
-				GeometryTerm = DotProductFactor;
-			#endif
-		}
-	#endif
+					// Weight toward the material normal to increase directionality
+					float BentNormalWeightFactor = 1 - (1 - SkyVisibility) * (1 - SkyVisibility);
+
+					// We are lerping between the inputs of two lighting scenarios based on occlusion
+					// In the mostly unoccluded case, evaluate sky lighting with the material normal, because it has higher detail
+					// In the mostly occluded case, evaluate sky lighting with the bent normal, because it is a better representation of the incoming lighting
+					// Then treat the lighting evaluated along the bent normal as an area light, so we must apply the lambert term
+					SkyLightingNormal = lerp(NormalizedBentNormal, WorldNormal, BentNormalWeightFactor);
+
+					float DotProductFactor = lerp(saturate(dot(NormalizedBentNormal, WorldNormal)), 1, BentNormalWeightFactor);
+					// Account for darkening due to the geometry term
+					GeometryTerm = DotProductFactor;
+				#endif
+			}
+		#endif
 			
-	// Compute the preconvolved incoming lighting with the bent normal direction
-	float3 DiffuseLookup = GetEffectiveSkySHDiffuse(SkyLightingNormal) * ResolvedView.SkyLightColor.rgb;
+		// Compute the preconvolved incoming lighting with the bent normal direction
+		float3 DiffuseLookup = GetEffectiveSkySHDiffuse(SkyLightingNormal) * ResolvedView.SkyLightColor.rgb;
 
-	// Apply AO to the sky diffuse
-	OutDiffuseLighting += DiffuseLookup * (SkyVisibility * GeometryTerm);
+		// Apply AO to the sky diffuse
+		OutDiffuseLighting += DiffuseLookup * (SkyVisibility * GeometryTerm);
 
-	#if MATERIAL_SHADINGMODEL_TWOSIDED_FOLIAGE
-	if (GBuffer.ShadingModelID == SHADINGMODELID_TWOSIDED_FOLIAGE)
-	{
-		float3 BackfaceDiffuseLookup = GetEffectiveSkySHDiffuse(-WorldNormal) * ResolvedView.SkyLightColor.rgb;
-		OutSubsurfaceLighting += BackfaceDiffuseLookup * SkyVisibility;
+		#if MATERIAL_SHADINGMODEL_TWOSIDED_FOLIAGE
+		if (GBuffer.ShadingModelID == SHADINGMODELID_TWOSIDED_FOLIAGE)
+		{
+			float3 BackfaceDiffuseLookup = GetEffectiveSkySHDiffuse(-WorldNormal) * ResolvedView.SkyLightColor.rgb;
+			OutSubsurfaceLighting += BackfaceDiffuseLookup * SkyVisibility;
+		}
+		#endif
 	}
-	#endif
-
-#endif
 }
 
 #if SUPPORTS_INDEPENDENT_SAMPLERS
diff --git a/Engine/Shaders/Private/ReflectionEnvironmentComposite.ush b/Engine/Shaders/Private/ReflectionEnvironmentComposite.ush
--- a/Engine/Shaders/Private/ReflectionEnvironmentComposite.ush
+++ b/Engine/Shaders/Private/ReflectionEnvironmentComposite.ush
@@ -162,34 +162,38 @@ float3 CompositeReflectionCapturesAndSkylight(
 	ImageBasedReflections.rgb *= View.IndirectLightingColorScale;
 	CompositedAverageBrightness.x *= Luminance( View.IndirectLightingColorScale );
 
-#if ENABLE_SKY_LIGHT
-
-	BRANCH  
-	if (ReflectionStruct.SkyLightParameters.y > 0 && bCompositeSkylight)
+	uint SkyLightMask = MATERIALBLENDING_TRANSLUCENT ? 2 : 1;
+	bool EnableSkyLight = 0 != (View.EnableSkyLight & SkyLightMask) && !MATERIAL_SHADINGMODEL_UNLIT;
+	BRANCH
+	if(EnableSkyLight)
 	{
-		float SkyAverageBrightness = 1.0f;
+		BRANCH  
+		if (ReflectionStruct.SkyLightParameters.y > 0 && bCompositeSkylight)
+		{
+			float SkyAverageBrightness = 1.0f;
 
-		#if REFLECTION_COMPOSITE_SUPPORT_SKYLIGHT_BLEND
-			float3 SkyLighting = GetSkyLightReflectionSupportingBlend(RayDirection, Roughness, SkyAverageBrightness);
-		#else
-			float3 SkyLighting = GetSkyLightReflection(RayDirection, Roughness, SkyAverageBrightness);
-		#endif
+			#if REFLECTION_COMPOSITE_SUPPORT_SKYLIGHT_BLEND
+				float3 SkyLighting = GetSkyLightReflectionSupportingBlend(RayDirection, Roughness, SkyAverageBrightness);
+			#else
+				float3 SkyLighting = GetSkyLightReflection(RayDirection, Roughness, SkyAverageBrightness);
+			#endif
 
-		// Normalize for static skylight types which mix with lightmaps
-		bool bNormalize = ReflectionStruct.SkyLightParameters.z < 1 && ALLOW_STATIC_LIGHTING;
+			// Normalize for static skylight types which mix with lightmaps
+			bool bNormalize = ReflectionStruct.SkyLightParameters.z < 1 && ALLOW_STATIC_LIGHTING;
 
-		FLATTEN
-		if (bNormalize)
-		{
-			ImageBasedReflections.rgb += ImageBasedReflections.a * SkyLighting * IndirectSpecularOcclusion;
-			CompositedAverageBrightness.x += SkyAverageBrightness * CompositedAverageBrightness.y;
-		}
-		else
-		{
-			ExtraIndirectSpecular += SkyLighting * IndirectSpecularOcclusion;
+			FLATTEN
+			if (bNormalize)
+			{
+				ImageBasedReflections.rgb += ImageBasedReflections.a * SkyLighting * IndirectSpecularOcclusion;
+				CompositedAverageBrightness.x += SkyAverageBrightness * CompositedAverageBrightness.y;
+			}
+			else
+			{
+				ExtraIndirectSpecular += SkyLighting * IndirectSpecularOcclusion;
+			}
 		}
 	}
-#endif
+
 
 #if ALLOW_STATIC_LIGHTING
 	ImageBasedReflections.rgb *= ComputeMixingWeight(IndirectIrradiance, CompositedAverageBrightness.x, Roughness);
diff --git a/Engine/Shaders/Private/ReflectionEnvironmentPixelShader.usf b/Engine/Shaders/Private/ReflectionEnvironmentPixelShader.usf
--- a/Engine/Shaders/Private/ReflectionEnvironmentPixelShader.usf
+++ b/Engine/Shaders/Private/ReflectionEnvironmentPixelShader.usf
@@ -222,12 +222,18 @@ float3 ReflectionEnvironment(FGBufferData GBuffer, float AmbientOcclusion, float
 	float3 ReflectionVector = reflect(CameraToPixel, GBuffer.WorldNormal);
 	float IndirectIrradiance = GBuffer.IndirectIrradiance;
 	
-#if ENABLE_SKY_LIGHT && ALLOW_STATIC_LIGHTING
+#if ALLOW_STATIC_LIGHTING
+	uint SkyLightMask = MATERIALBLENDING_TRANSLUCENT ? 2 : 1;
+	bool EnableSkyLight = 0 != (View.EnableSkyLight & SkyLightMask) && !MATERIAL_SHADINGMODEL_UNLIT;
 	BRANCH
-	// Add in diffuse contribution from dynamic skylights so reflection captures will have something to mix with
-	if (ReflectionStruct.SkyLightParameters.y > 0 && ReflectionStruct.SkyLightParameters.z > 0)
+	if (EnableSkyLight)
 	{
-		IndirectIrradiance += GetDynamicSkyIndirectIrradiance(BentNormal, GBuffer.WorldNormal);
+		BRANCH
+		// Add in diffuse contribution from dynamic skylights so reflection captures will have something to mix with
+		if (ReflectionStruct.SkyLightParameters.y > 0 && ReflectionStruct.SkyLightParameters.z > 0)
+		{
+			IndirectIrradiance += GetDynamicSkyIndirectIrradiance(BentNormal, GBuffer.WorldNormal);
+		}
 	}
 #endif
 
diff --git a/Engine/Shaders/Public/ShaderVersion.ush b/Engine/Shaders/Public/ShaderVersion.ush
--- a/Engine/Shaders/Public/ShaderVersion.ush
+++ b/Engine/Shaders/Public/ShaderVersion.ush
@@ -3,4 +3,4 @@
 // in Platform.ush (which should be included in any shader) it allows to invalidate the shader DDC.
 // 
 // If you are merging streams and there is a conflict with this GUID you should make a new GUID rather than taking one or the other.
-// GUID = 2C9048786E8942F699D8C862E803EB60
+// GUID = F8A2B52ADE3D4F4095316DED2FE4DF3B
diff --git a/Engine/Source/Runtime/Engine/Private/SceneView.cpp b/Engine/Source/Runtime/Engine/Private/SceneView.cpp
--- a/Engine/Source/Runtime/Engine/Private/SceneView.cpp
+++ b/Engine/Source/Runtime/Engine/Private/SceneView.cpp
@@ -2368,6 +2368,9 @@ void FSceneView::SetupCommonViewUniformBufferParameters(
 
 	ViewUniformShaderParameters.MinRoughness = FMath::Clamp(CVarGlobalMinRoughnessOverride.GetValueOnRenderThread(), 0.02f, 1.0f);
 
+	FSceneInterface* Scene = Family->Scene;
+	uint32 SkylightMask = Scene ? Scene->ShouldRenderSkylightMask() : 0;
+	ViewUniformShaderParameters.EnableSkyLight = SkylightMask;
 	//to tail call keep the order and number of parameters of the caller function
 	SetupViewRectUniformBufferParameters(ViewUniformShaderParameters, BufferSize, EffectiveViewRect, InViewMatrices, InPrevViewMatrices);
 }
diff --git a/Engine/Source/Runtime/Engine/Public/SceneInterface.h b/Engine/Source/Runtime/Engine/Public/SceneInterface.h
--- a/Engine/Source/Runtime/Engine/Public/SceneInterface.h
+++ b/Engine/Source/Runtime/Engine/Public/SceneInterface.h
@@ -430,6 +430,12 @@ public:
 		return GetShadingPath(GetFeatureLevel());
 	}
 
+	virtual uint32 ShouldRenderSkylightMask() const
+	{
+		return 3;
+	}
+
+
 #if WITH_EDITOR
 	/**
 	 * Initialize the pixel inspector buffers.
diff --git a/Engine/Source/Runtime/Engine/Public/SceneView.h b/Engine/Source/Runtime/Engine/Public/SceneView.h
--- a/Engine/Source/Runtime/Engine/Public/SceneView.h
+++ b/Engine/Source/Runtime/Engine/Public/SceneView.h
@@ -730,6 +730,7 @@ enum ETranslucencyVolumeCascade
 	VIEW_UNIFORM_BUFFER_MEMBER(int32, FarShadowStaticMeshLODBias) \
 	VIEW_UNIFORM_BUFFER_MEMBER(float, MinRoughness) \
 	VIEW_UNIFORM_BUFFER_MEMBER(FVector4, HairRenderInfo) \
+	VIEW_UNIFORM_BUFFER_MEMBER(uint32, EnableSkyLight) \
 	VIEW_UNIFORM_BUFFER_MEMBER(uint32, HairRenderInfoBits) \
 
 #define VIEW_UNIFORM_BUFFER_MEMBER(type, identifier) \
diff --git a/Engine/Source/Runtime/Renderer/Private/BasePassRendering.cpp b/Engine/Source/Runtime/Renderer/Private/BasePassRendering.cpp
--- a/Engine/Source/Runtime/Renderer/Private/BasePassRendering.cpp
+++ b/Engine/Source/Runtime/Renderer/Private/BasePassRendering.cpp
@@ -94,16 +94,16 @@ IMPLEMENT_GLOBAL_SHADER_PARAMETER_STRUCT(FTranslucentBasePassUniformParameters,
 	typedef TBasePassHS< LightMapPolicyType, true> TBasePassHS##LightMapPolicyName##AtmosphericFogShaderName; \
 	IMPLEMENT_MATERIAL_SHADER_TYPE(template<>,TBasePassHS##LightMapPolicyName##AtmosphericFogShaderName,TEXT("/Engine/Private/BasePassTessellationShaders.usf"),TEXT("MainHull"),SF_Hull);
 
-#define IMPLEMENT_BASEPASS_PIXELSHADER_TYPE(LightMapPolicyType,LightMapPolicyName,bEnableSkyLight,SkyLightName) \
-	typedef TBasePassPS<LightMapPolicyType, bEnableSkyLight> TBasePassPS##LightMapPolicyName##SkyLightName; \
-	IMPLEMENT_MATERIAL_SHADER_TYPE(template<>,TBasePassPS##LightMapPolicyName##SkyLightName,TEXT("/Engine/Private/BasePassPixelShader.usf"),TEXT("MainPS"),SF_Pixel);
+#define IMPLEMENT_BASEPASS_PIXELSHADER_TYPE(LightMapPolicyType,LightMapPolicyName) \
+	typedef TBasePassPS<LightMapPolicyType> TBasePassPS##LightMapPolicyName; \
+	IMPLEMENT_MATERIAL_SHADER_TYPE(template<>,TBasePassPS##LightMapPolicyName,TEXT("/Engine/Private/BasePassPixelShader.usf"),TEXT("MainPS"),SF_Pixel);
 
 // Implement a pixel shader type for skylights and one without, and one vertex shader that will be shared between them
 #define IMPLEMENT_BASEPASS_LIGHTMAPPED_SHADER_TYPE(LightMapPolicyType,LightMapPolicyName) \
 	IMPLEMENT_BASEPASS_VERTEXSHADER_TYPE(LightMapPolicyType,LightMapPolicyName) \
 	IMPLEMENT_BASEPASS_VERTEXSHADER_ONLY_TYPE(LightMapPolicyType,LightMapPolicyName,AtmosphericFog) \
-	IMPLEMENT_BASEPASS_PIXELSHADER_TYPE(LightMapPolicyType,LightMapPolicyName,true,Skylight) \
-	IMPLEMENT_BASEPASS_PIXELSHADER_TYPE(LightMapPolicyType,LightMapPolicyName,false,);
+	IMPLEMENT_BASEPASS_PIXELSHADER_TYPE(LightMapPolicyType,LightMapPolicyName);
+
 
 // Implement shader types per lightmap policy
 // If renaming or refactoring these, remember to update FMaterialResource::GetRepresentativeInstructionCounts and FPreviewMaterial::ShouldCache().
@@ -430,20 +430,13 @@ void GetUniformBasePassShaders(
 		VertexShader = TShaderRef<TBasePassVertexShaderPolicyParamType<FUniformLightMapPolicy>>::ReinterpretCast(Material.GetShader<TBasePassVS<TUniformLightMapPolicy<Policy>, false> >(VertexFactoryType));
 	}
 
-	if (bEnableSkyLight)
+	if (bUse128bitRT && (Policy == LMP_NO_LIGHTMAP))
 	{
-		PixelShader = TShaderRef<TBasePassPixelShaderPolicyParamType<FUniformLightMapPolicy>>::ReinterpretCast(Material.GetShader<TBasePassPS<TUniformLightMapPolicy<Policy>, true> >(VertexFactoryType));
+		PixelShader = TShaderRef<TBasePassPixelShaderPolicyParamType<FUniformLightMapPolicy>>::ReinterpretCast(Material.GetShader<F128BitRTBasePassPS>(VertexFactoryType));
 	}
 	else
 	{
-		if (bUse128bitRT && (Policy == LMP_NO_LIGHTMAP))
-		{
-			PixelShader = TShaderRef<TBasePassPixelShaderPolicyParamType<FUniformLightMapPolicy>>::ReinterpretCast(Material.GetShader<F128BitRTBasePassPS>(VertexFactoryType));
-		}
-		else
-		{
-			PixelShader = TShaderRef<TBasePassPixelShaderPolicyParamType<FUniformLightMapPolicy>>::ReinterpretCast(Material.GetShader<TBasePassPS<TUniformLightMapPolicy<Policy>, false> >(VertexFactoryType));
-		}
+		PixelShader = TShaderRef<TBasePassPixelShaderPolicyParamType<FUniformLightMapPolicy>>::ReinterpretCast(Material.GetShader<TBasePassPS<TUniformLightMapPolicy<Policy>> >(VertexFactoryType));
 	}
 }
 
diff --git a/Engine/Source/Runtime/Renderer/Private/BasePassRendering.h b/Engine/Source/Runtime/Renderer/Private/BasePassRendering.h
--- a/Engine/Source/Runtime/Renderer/Private/BasePassRendering.h
+++ b/Engine/Source/Runtime/Renderer/Private/BasePassRendering.h
@@ -470,7 +470,7 @@ public:
 };
 
 /** The concrete base pass pixel shader type. */
-template<typename LightMapPolicyType, bool bEnableSkyLight>
+template<typename LightMapPolicyType>
 class TBasePassPS : public TBasePassPixelShaderBaseType<LightMapPolicyType>
 {
 	DECLARE_SHADER_TYPE(TBasePassPS,MeshMaterial);
@@ -487,7 +487,7 @@ public:
 		const bool bTranslucent = IsTranslucentBlendMode(Parameters.MaterialParameters.BlendMode);
 		const bool bForceAllPermutations = SupportAllShaderPermutations && SupportAllShaderPermutations->GetValueOnAnyThread() != 0;
 		const bool bProjectSupportsStationarySkylight = !SupportStationarySkylight || SupportStationarySkylight->GetValueOnAnyThread() != 0 || bForceAllPermutations;
-
+		const bool bEnableSkyLight = false; //makes bCacheShaders redundant
 		const bool bCacheShaders = !bEnableSkyLight
 			//translucent materials need to compile skylight support to support MOVABLE skylights also.
 			|| bTranslucent
@@ -504,8 +504,6 @@ public:
 	{
 		// For deferred decals, the shader class used is FDeferredDecalPS. the TBasePassPS is only used in the material editor and will read wrong values.
 		OutEnvironment.SetDefine(TEXT("SCENE_TEXTURES_DISABLED"), Parameters.MaterialParameters.MaterialDomain != MD_Surface); 
-
-		OutEnvironment.SetDefine(TEXT("ENABLE_SKY_LIGHT"), bEnableSkyLight);
 		TBasePassPixelShaderBaseType<LightMapPolicyType>::ModifyCompilationEnvironment(Parameters, OutEnvironment);
 	}
 	
@@ -592,14 +590,7 @@ void GetBasePassShaders(
 	{
 		VertexShader = Material.GetShader<TBasePassVS<LightMapPolicyType, false> >(VertexFactoryType);
 	}
-	if (bEnableSkyLight)
-	{
-		PixelShader = Material.GetShader<TBasePassPS<LightMapPolicyType, true> >(VertexFactoryType);
-	}
-	else
-	{
-		PixelShader = Material.GetShader<TBasePassPS<LightMapPolicyType, false> >(VertexFactoryType);
-	}
+	PixelShader = Material.GetShader<TBasePassPS<LightMapPolicyType> >(VertexFactoryType);
 }
 
 template <>
diff --git a/Engine/Source/Runtime/Renderer/Private/DebugViewModeRendering.cpp b/Engine/Source/Runtime/Renderer/Private/DebugViewModeRendering.cpp
--- a/Engine/Source/Runtime/Renderer/Private/DebugViewModeRendering.cpp
+++ b/Engine/Source/Runtime/Renderer/Private/DebugViewModeRendering.cpp
@@ -299,12 +299,12 @@ void FDebugViewModeMeshProcessor::UpdateInstructionCount(FDebugViewModeShaderEle
 			if (IsSimpleForwardShadingEnabled(ShaderPlatform))
 			{
 				OutShaderElementData.NumVSInstructions = InBatchMaterial->GetShader<TBasePassVS<TUniformLightMapPolicy<LMP_SIMPLE_NO_LIGHTMAP>, false>>(InVertexFactoryType)->GetNumInstructions();
-				OutShaderElementData.NumPSInstructions = InBatchMaterial->GetShader<TBasePassPS<TUniformLightMapPolicy<LMP_SIMPLE_NO_LIGHTMAP>, false>>(InVertexFactoryType)->GetNumInstructions();
+				OutShaderElementData.NumPSInstructions = InBatchMaterial->GetShader<TBasePassPS<TUniformLightMapPolicy<LMP_SIMPLE_NO_LIGHTMAP>>>(InVertexFactoryType)->GetNumInstructions();
 			}
 			else
 			{
 				OutShaderElementData.NumVSInstructions = InBatchMaterial->GetShader<TBasePassVS<TUniformLightMapPolicy<LMP_NO_LIGHTMAP>, false>>(InVertexFactoryType)->GetNumInstructions();
-				OutShaderElementData.NumPSInstructions = InBatchMaterial->GetShader<TBasePassPS<TUniformLightMapPolicy<LMP_NO_LIGHTMAP>, false>>(InVertexFactoryType)->GetNumInstructions();
+				OutShaderElementData.NumPSInstructions = InBatchMaterial->GetShader<TBasePassPS<TUniformLightMapPolicy<LMP_NO_LIGHTMAP>>>(InVertexFactoryType)->GetNumInstructions();
 
 				if (IsForwardShadingEnabled(ShaderPlatform) && !IsTranslucentBlendMode(InBatchMaterial->GetBlendMode()))
 				{
diff --git a/Engine/Source/Runtime/Renderer/Private/ScenePrivate.h b/Engine/Source/Runtime/Renderer/Private/ScenePrivate.h
--- a/Engine/Source/Runtime/Renderer/Private/ScenePrivate.h
+++ b/Engine/Source/Runtime/Renderer/Private/ScenePrivate.h
@@ -2954,6 +2954,33 @@ public:
 
 		return bRenderSkyLight;
 	}
+	virtual uint32 ShouldRenderSkylightMask() const
+	{
+		bool bRenderSkyLight = SkyLight && !SkyLight->bHasStaticLighting && !SkyLight->bCastRayTracedShadow;
+		uint32 Mask = 0;
+		{
+			// Both stationary and movable skylights are applied in base pass for translucent materials
+			if(bRenderSkyLight && (ReadOnlyCVARCache.bEnableStationarySkylight || !SkyLight->bWantsStaticShadowing))
+			{
+				Mask |= 2;
+
+			}
+		}
+		{
+			// For opaque materials, stationary skylight is applied in base pass but movable skylight
+			// is applied in a separate render pass (bWantssStaticShadowing means stationary skylight)
+			if(bRenderSkyLight && ((ReadOnlyCVARCache.bEnableStationarySkylight && SkyLight->bWantsStaticShadowing)
+					|| (!SkyLight->bWantsStaticShadowing
+						&& (IsAnyForwardShadingEnabled(GetShaderPlatform()) || IsMobilePlatform(GetShaderPlatform())))))
+			{
+				Mask |= 1;
+			}
+		}
+
+		return Mask;
+	}
+
+
 
 	virtual TArray<FPrimitiveComponentId> GetScenePrimitiveComponentIds() const override
 	{
