Implement simplified ray tracing material hit shaders which forces everything to diffuse and no static lighting evaluation

diff --git a/Engine/Shaders/Private/RayTracing/RayTracingMaterialHitShaders.usf b/Engine/Shaders/Private/RayTracing/RayTracingMaterialHitShaders.usf
--- a/Engine/Shaders/Private/RayTracing/RayTracingMaterialHitShaders.usf
+++ b/Engine/Shaders/Private/RayTracing/RayTracingMaterialHitShaders.usf
@@ -9,6 +9,9 @@
 	// Needed for VertexFactoryInterpolate
 	#define USING_TESSELLATION 1
 #endif
+#ifndef SIMPLIFIED_MATERIAL_SHADER
+	#define SIMPLIFIED_MATERIAL_SHADER 0
+#endif
 
 // #dxr_todo: UE-72537 move to common.ush
 #define clip(x)
@@ -304,6 +307,33 @@ RAY_TRACING_ENTRY_CLOSEST_HIT(MaterialCHS,
 
 	FMaterialClosestHitPayload Payload = (FMaterialClosestHitPayload)0;
 
+	/**
+	 * Set common material attributes for both full and simplified materials
+	 **/
+	Payload.ShadingModelID = GetMaterialShadingModel(PixelMaterialInputs);
+
+	#if MATERIALBLENDING_ALPHACOMPOSITE
+		uint MaterialBlendingMode = RAY_TRACING_BLEND_MODE_ALPHA_COMPOSITE;
+	#elif MATERIALBLENDING_ALPHAHOLDOUT
+		uint MaterialBlendingMode = RAY_TRACING_BLEND_MODE_ALPHA_HOLDOUT;
+	#elif MATERIALBLENDING_TRANSLUCENT
+		uint MaterialBlendingMode = RAY_TRACING_BLEND_MODE_TRANSLUCENT;
+	#elif MATERIALBLENDING_ADDITIVE
+		uint MaterialBlendingMode = RAY_TRACING_BLEND_MODE_ADDITIVE;
+	#elif MATERIALBLENDING_MODULATE
+		uint MaterialBlendingMode = RAY_TRACING_BLEND_MODE_MODULATE;
+	#else
+		uint MaterialBlendingMode = RAY_TRACING_BLEND_MODE_OPAQUE;
+	#endif
+
+	Payload.BlendingMode = MaterialBlendingMode;
+
+	Payload.HitT = RayTCurrent();
+	if (HitKind() == HIT_KIND_TRIANGLE_FRONT_FACE)
+	{
+		Payload.SetFrontFace();
+	}
+
 	// Store the results in local variables and reuse instead of calling the functions multiple times.
 	half3 BaseColor = GetMaterialBaseColor(PixelMaterialInputs);
 	half2 Refraction = GetMaterialRefraction(PixelMaterialInputs);
@@ -311,6 +341,11 @@ RAY_TRACING_ENTRY_CLOSEST_HIT(MaterialCHS,
 	half  Specular = GetMaterialSpecular(PixelMaterialInputs);
 	half Roughness = GetMaterialRoughness(PixelMaterialInputs);
 
+#if !SIMPLIFIED_MATERIAL_SHADER
+
+	/**
+	 * Set material attributes for full materials
+	 **/
 	Payload.WorldPos = MaterialParameters.AbsoluteWorldPosition;
 	Payload.WorldNormal = normalize(MaterialParameters.WorldNormal);
 	Payload.Radiance = GetMaterialEmissive(PixelMaterialInputs);
@@ -318,8 +353,10 @@ RAY_TRACING_ENTRY_CLOSEST_HIT(MaterialCHS,
 	Payload.Specular = Specular;
 	Payload.Roughness = Roughness;
 	Payload.Metallic = Metallic;
+	Payload.Opacity = GetMaterialOpacity(PixelMaterialInputs);
 	Payload.Ior = Refraction.x;
-    Payload.CustomData = float4(GetMaterialCustomData0(MaterialParameters), GetMaterialCustomData1(MaterialParameters), 0, 0);
+	Payload.CustomData = float4(GetMaterialCustomData0(MaterialParameters), GetMaterialCustomData1(MaterialParameters), 0, 0);
+	
 	#if CLEAR_COAT_BOTTOM_NORMAL
 	{
 		float2 oct2 = UnitVectorToOctahedron(Payload.WorldNormal);
@@ -341,58 +378,34 @@ RAY_TRACING_ENTRY_CLOSEST_HIT(MaterialCHS,
 		#endif
 	}
 	#endif
-	Payload.Opacity = GetMaterialOpacity(PixelMaterialInputs);
-    Payload.ShadingModelID = GetMaterialShadingModel(PixelMaterialInputs);
-
- #if MATERIALBLENDING_ALPHACOMPOSITE
-	uint MaterialBlendingMode = RAY_TRACING_BLEND_MODE_ALPHA_COMPOSITE;
- #elif MATERIALBLENDING_ALPHAHOLDOUT
-	uint MaterialBlendingMode = RAY_TRACING_BLEND_MODE_ALPHA_HOLDOUT;
- #elif MATERIALBLENDING_TRANSLUCENT
-	uint MaterialBlendingMode = RAY_TRACING_BLEND_MODE_TRANSLUCENT;
- #elif MATERIALBLENDING_ADDITIVE
-	uint MaterialBlendingMode = RAY_TRACING_BLEND_MODE_ADDITIVE;
- #elif MATERIALBLENDING_MODULATE
-	uint MaterialBlendingMode = RAY_TRACING_BLEND_MODE_MODULATE;
-#else
-	uint MaterialBlendingMode = RAY_TRACING_BLEND_MODE_OPAQUE;
-#endif
 
-	Payload.BlendingMode = MaterialBlendingMode;
-
-	Payload.HitT = RayTCurrent();
-	if (HitKind() == HIT_KIND_TRIANGLE_FRONT_FACE)
-	{
-		Payload.SetFrontFace();
-	}
-
-#if GBUFFER_HAS_TANGENT
-	#if MATERIAL_TANGENTSPACENORMAL
-		half3 MaterialTangent = normalize(GetMaterialTangent(PixelMaterialInputs));
-
-		float3 WorldTangent = TransformTangentNormalToWorld(MaterialParameters.TangentToWorld, MaterialTangent);
-
-		// flip the tangent for backfaces being rendered with a two-sided material
-		WorldTangent *= MaterialParameters.TwoSidedSign;
-	#else 
-		float3 WorldTangent = normalize(GetMaterialTangent(PixelMaterialInputs));
-	#endif 
-		Payload.WorldTangent = WorldTangent;
-		Payload.Anisotropy = GetMaterialAnisotropy(PixelMaterialInputs);
-#else
-	Payload.WorldTangent = float3(1, 0, 0);
-	Payload.Anisotropy = 1.0;
-#endif	
+	#if GBUFFER_HAS_TANGENT
+		#if MATERIAL_TANGENTSPACENORMAL
+			half3 MaterialTangent = normalize(GetMaterialTangent(PixelMaterialInputs));
+
+			float3 WorldTangent = TransformTangentNormalToWorld(MaterialParameters.TangentToWorld, MaterialTangent);
+
+			// flip the tangent for backfaces being rendered with a two-sided material
+			WorldTangent *= MaterialParameters.TwoSidedSign;
+		#else 
+			WorldTangent = normalize(GetMaterialTangent(PixelMaterialInputs));
+		#endif 
+			Payload.WorldTangent = WorldTangent;
+			Payload.Anisotropy = GetMaterialAnisotropy(PixelMaterialInputs);
+	#else
+		Payload.WorldTangent = float3(1, 0, 0);
+		Payload.Anisotropy = 1.0;
+	#endif
 
 	float3 DiffuseIndirectLighting = 0;
 	float3 SubsurfaceIndirectLighting;
 	float IndirectIrradiance = 0;
 	
 	float3 VolumetricLightmapBrickTextureUVs;
-	
-#if PRECOMPUTED_IRRADIANCE_VOLUME_LIGHTING
-	VolumetricLightmapBrickTextureUVs = ComputeVolumetricLightmapBrickTextureUVs(MaterialParameters.AbsoluteWorldPosition);
-#endif
+		
+	#if PRECOMPUTED_IRRADIANCE_VOLUME_LIGHTING
+		VolumetricLightmapBrickTextureUVs = ComputeVolumetricLightmapBrickTextureUVs(MaterialParameters.AbsoluteWorldPosition);
+	#endif
 
 	VTPageTableResult LightmapVTPageTableResult = (VTPageTableResult)0.0f;
 	GetPrecomputedIndirectLightingAndSkyLight(
@@ -412,6 +425,26 @@ RAY_TRACING_ENTRY_CLOSEST_HIT(MaterialCHS,
 	
 	Payload.IndirectIrradiance += DiffuseIndirectLighting; // Contains both lightmap & skylight
 
+#else
+
+	/**
+	 * Set material attributes for simplified materials
+	 **/
+	// Anything that is not foliage will be forced to default lit
+	if (Payload.ShadingModelID != SHADINGMODELID_TWOSIDED_FOLIAGE)
+	{
+		Payload.ShadingModelID = SHADINGMODELID_DEFAULT_LIT;
+	}
+	Payload.WorldPos = MaterialParameters.AbsoluteWorldPosition;
+	Payload.WorldNormal = normalize(MaterialParameters.WorldNormal);
+	Payload.Radiance = GetMaterialEmissive(PixelMaterialInputs);
+	Payload.BaseColor = BaseColor;
+	Payload.Specular = 0;
+	Payload.Roughness = 1;
+	Payload.Metallic = 0;
+
+#endif
+
 	PackedPayload = PackRayTracingPayload(Payload, PropagatedCone);
 }
 
