Fix HitT not set on full payload

* PackedPayload is overwritten at the end of the shader

diff --git a/Engine/Shaders/Private/RayTracing/RayTracingMaterialHitShaders.usf b/Engine/Shaders/Private/RayTracing/RayTracingMaterialHitShaders.usf
--- a/Engine/Shaders/Private/RayTracing/RayTracingMaterialHitShaders.usf
+++ b/Engine/Shaders/Private/RayTracing/RayTracingMaterialHitShaders.usf
@@ -331,6 +331,7 @@ RAY_TRACING_ENTRY_CLOSEST_HIT(MaterialCHS,
 
 	Payload.BlendingMode = MaterialBlendingMode;
 
+	Payload.HitT = RayTCurrent();
 	if (HitKind() == HIT_KIND_TRIANGLE_FRONT_FACE)
 	{
 		Payload.SetFrontFace();
