Fixed post integration compilation issue

diff --git a/Engine/Shaders/Private/RayTracing/RayTracingMaterialHitShaders.usf b/Engine/Shaders/Private/RayTracing/RayTracingMaterialHitShaders.usf
--- a/Engine/Shaders/Private/RayTracing/RayTracingMaterialHitShaders.usf
+++ b/Engine/Shaders/Private/RayTracing/RayTracingMaterialHitShaders.usf
@@ -400,6 +400,24 @@ RAY_TRACING_ENTRY_CLOSEST_HIT(MaterialCHS,
 		Payload.Anisotropy = 1.0;
 	#endif
 
+#if GBUFFER_HAS_TANGENT
+	#if MATERIAL_TANGENTSPACENORMAL
+		half3 MaterialTangent = normalize(GetMaterialTangent(PixelMaterialInputs));
+
+		float3 WorldTangent = TransformTangentNormalToWorld(MaterialParameters.TangentToWorld, MaterialTangent);
+
+		// flip the tangent for backfaces being rendered with a two-sided material
+		WorldTangent *= MaterialParameters.TwoSidedSign;
+	#else 
+		float3 WorldTangent = normalize(GetMaterialTangent(PixelMaterialInputs));
+	#endif 
+		Payload.WorldTangent = WorldTangent;
+		Payload.Anisotropy = GetMaterialAnisotropy(PixelMaterialInputs);
+#else
+	Payload.WorldTangent = float3(1, 0, 0);
+	Payload.Anisotropy = 1.0;
+#endif	
+
 	float3 DiffuseIndirectLighting = 0;
 	float3 SubsurfaceIndirectLighting;
 	float IndirectIrradiance = 0;
