Fixed code duplication due a merge issue.

diff --git a/Engine/Shaders/Private/RayTracing/RayTracingMaterialHitShaders.usf b/Engine/Shaders/Private/RayTracing/RayTracingMaterialHitShaders.usf
--- a/Engine/Shaders/Private/RayTracing/RayTracingMaterialHitShaders.usf
+++ b/Engine/Shaders/Private/RayTracing/RayTracingMaterialHitShaders.usf
@@ -382,24 +382,6 @@ RAY_TRACING_ENTRY_CLOSEST_HIT(MaterialCHS,
 	}
 	#endif
 
-	#if GBUFFER_HAS_TANGENT
-		#if MATERIAL_TANGENTSPACENORMAL
-			half3 MaterialTangent = normalize(GetMaterialTangent(PixelMaterialInputs));
-
-			float3 WorldTangent = TransformTangentNormalToWorld(MaterialParameters.TangentToWorld, MaterialTangent);
-
-			// flip the tangent for backfaces being rendered with a two-sided material
-			WorldTangent *= MaterialParameters.TwoSidedSign;
-		#else 
-			WorldTangent = normalize(GetMaterialTangent(PixelMaterialInputs));
-		#endif 
-			Payload.WorldTangent = WorldTangent;
-			Payload.Anisotropy = GetMaterialAnisotropy(PixelMaterialInputs);
-	#else
-		Payload.WorldTangent = float3(1, 0, 0);
-		Payload.Anisotropy = 1.0;
-	#endif
-
 #if GBUFFER_HAS_TANGENT
 	#if MATERIAL_TANGENTSPACENORMAL
 		half3 MaterialTangent = normalize(GetMaterialTangent(PixelMaterialInputs));
