Fixed ray tracing shader compilation issue

diff --git a/Engine/Shaders/Private/RayTracing/RayTracingMaterialHitShaders.usf b/Engine/Shaders/Private/RayTracing/RayTracingMaterialHitShaders.usf
--- a/Engine/Shaders/Private/RayTracing/RayTracingMaterialHitShaders.usf
+++ b/Engine/Shaders/Private/RayTracing/RayTracingMaterialHitShaders.usf
@@ -503,11 +503,5 @@ RAY_TRACING_ENTRY_ANY_HIT(MaterialAHS,
 	}
 #endif // MATERIALBLENDING_MASKED
 
-	// For any other blending models, ignore hit in anyhit shader if we're tracing visibility rays until we have proper handling for translucent
-	if (PackedPayload.IsIgnoreTranslucentMaterials())
-	{
-		IgnoreHit();
-	}
-
 #endif // USE_MATERIAL_ANY_HIT_SHADER
 }  
