Fix incorrect pre-exposure in deferred pixel lighting.

diff --git a/Engine/Shaders/Private/DeferredLightPixelShaders.usf b/Engine/Shaders/Private/DeferredLightPixelShaders.usf
--- a/Engine/Shaders/Private/DeferredLightPixelShaders.usf
+++ b/Engine/Shaders/Private/DeferredLightPixelShaders.usf
@@ -109,7 +109,7 @@ float GetExposure()
 #if USE_PREEXPOSURE
 	return View.PreExposure;
 #else
-	return 0;
+	return 1;
 #endif
 }
 
