Fix pre-exposition being applied twice on scene color in hair env lighting pass.

diff --git a/Engine/Shaders/Private/HairStrands/HairStrandsEnvironmentLighting.usf b/Engine/Shaders/Private/HairStrands/HairStrandsEnvironmentLighting.usf
--- a/Engine/Shaders/Private/HairStrands/HairStrandsEnvironmentLighting.usf
+++ b/Engine/Shaders/Private/HairStrands/HairStrandsEnvironmentLighting.usf
@@ -685,11 +685,6 @@ Texture2D<uint4> HairCategorizationTexture;
 
 float3 ComputeThinTransmission(FGBufferData GBuffer, float2 ScreenPosition, float2 SvPosition, uint PrimitiveId)
 {
-#if USE_PREEXPOSURE
-	const float PreExposure = View.PreExposure;
-#else
-	const float PreExposure = 1.f;
-#endif
 	const float3 WorldPosition = mul(float4(ScreenPosition * GBuffer.Depth, GBuffer.Depth, 1), View.ScreenToWorld).xyz;
 	const float3 CameraToPixel = normalize(WorldPosition - View.WorldCameraOrigin);
 	const uint2  PixelPosition = SvPosition - ResolvedView.ViewRectMin.xy;
@@ -704,11 +699,11 @@ float3 ComputeThinTransmission(FGBufferData GBuffer, float2 ScreenPosition, floa
 	// value, we divide by the hair coverage
 	const float4 OccludedLuminance = SceneColorTexture.Load(uint3(PixelPosition, 0));
 	const float4 Luminance = OccludedLuminance / PixelCoverage;
-	
+
+	// The pre-exposure is not applied onto the output as the SceneColorTexture value are already preexposed.
 	const float SinViewAngle = dot(T, V);
 	const HairAverageScattering Energy = SampleHairLUT(HairScatteringLUTTexture, HairLUTSampler, GBuffer.BaseColor, GBuffer.Roughness, SinViewAngle);
 	float3 Color = Energy.A_front * Luminance.xyz * PixelCoverage;
-	Color *= PreExposure;
 
 	return -min(-Color, 0.0);
 }
