Fixed ray tracing shader compilation issue

diff --git a/Engine/Shaders/Private/RayTracing/RayTracingSkyLightRGS.usf b/Engine/Shaders/Private/RayTracing/RayTracingSkyLightRGS.usf
--- a/Engine/Shaders/Private/RayTracing/RayTracingSkyLightRGS.usf
+++ b/Engine/Shaders/Private/RayTracing/RayTracingSkyLightRGS.usf
@@ -3,7 +3,6 @@
 #include "../Common.ush"
 #include "../SceneTextureParameters.ush"
 #include "../SceneTextureParameters.ush"
-#include "RayTracingDeferredShadingCommon.ush"
 #include "RayTracingSkyLightEvaluation.ush"
 #include "RayTracingDeferredShadingCommon.ush"
 
