Fix incorrect doubling hair voxelization when a raster sample fall several time within the same voxel.
Fix half coverage issue.

diff --git a/Engine/Shaders/Private/HairStrands/HairStrandsVoxelRasterCompute.usf b/Engine/Shaders/Private/HairStrands/HairStrandsVoxelRasterCompute.usf
--- a/Engine/Shaders/Private/HairStrands/HairStrandsVoxelRasterCompute.usf
+++ b/Engine/Shaders/Private/HairStrands/HairStrandsVoxelRasterCompute.usf
@@ -66,11 +66,12 @@ void MainCS(uint3 DispatchThreadID : SV_DispatchThreadID)
 		HairStrandsVF_Radius,
 		HairStrandsVF_Length);
 
+	const float DiameterToRadius = 0.5f;
 	const float3 WP0 = mul(float4(CP0.Position, 1), HairStrandsVF_LocalToWorldPrimitiveTransform).xyz;
-	const float HairCoverage0 = CP0.WorldRadius / max(CP0.WorldRadius, VirtualVoxel_VoxelWorldSize) * HairStrandsVF_Density;
+	const float HairCoverage0 = CP0.WorldRadius / max(CP0.WorldRadius, DiameterToRadius * VirtualVoxel_VoxelWorldSize);
 
 	const float3 WP1 = mul(float4(CP1.Position, 1), HairStrandsVF_LocalToWorldPrimitiveTransform).xyz;
-	const float HairCoverage1 = CP1.WorldRadius / max(CP1.WorldRadius, VirtualVoxel_VoxelWorldSize) * HairStrandsVF_Density;
+	const float HairCoverage1 = CP1.WorldRadius / max(CP1.WorldRadius, DiameterToRadius * VirtualVoxel_VoxelWorldSize);
 
 	const FVirtualVoxelNodeDesc NodeDesc = UnpackVoxelNode(VirtualVoxel_NodeDescBuffer[MacroGroupId], VirtualVoxel_PageResolution);
 
@@ -107,6 +108,7 @@ void MainCS(uint3 DispatchThreadID : SV_DispatchThreadID)
 	const float fMaxStep = LineLength / CommonDesc.VoxelWorldSize;
 	const float MaxStep = float(min(ceil(fMaxStep), MaxRasterCount));
 
+	int3 PreviousCoord = -1;
 	for (float StepIt = 0.0f; StepIt < MaxStep; ++StepIt)
 	{
 		const float U = (StepIt + 0.5f) / float(MaxStep);
@@ -138,9 +140,17 @@ void MainCS(uint3 DispatchThreadID : SV_DispatchThreadID)
 			const int3 VoxelPageOffset = VolumeCoord - PageIndexCoord * CommonDesc.PageResolution;
 			const int3 VoxelPageCoord = VoxelPageBase + VoxelPageOffset;
 
-			const float VoxelFixPointScale = GetVoxelDensityFixPointScale();
+			// Insure we don't write multiple time within the same voxel.
+			// This can happen for small hair segment, where both start & end points could be writtent into the same voxel.
+			const bool bRasterize = any(VoxelPageCoord != PreviousCoord);
+
 			const float HairCoverage = lerp(HairCoverage0, HairCoverage1, U);
-			InterlockedAdd(OutPageTexture[VoxelPageCoord], HairCoverage * VoxelFixPointScale);
+			if (bRasterize)
+			{
+				const float VoxelFixPointScale = GetVoxelDensityFixPointScale();
+				InterlockedAdd(OutPageTexture[VoxelPageCoord], HairCoverage * VoxelFixPointScale * HairStrandsVF_Density);
+			}
+			PreviousCoord = VoxelPageCoord;
 		}
 	}
 }
\ No newline at end of file
