Fix parameter name conflict with global named constant buffer.

diff --git a/Engine/Shaders/Private/HairStrands/HairStrandsVoxelOpaque.usf b/Engine/Shaders/Private/HairStrands/HairStrandsVoxelOpaque.usf
--- a/Engine/Shaders/Private/HairStrands/HairStrandsVoxelOpaque.usf
+++ b/Engine/Shaders/Private/HairStrands/HairStrandsVoxelOpaque.usf
@@ -1,5 +1,6 @@
 // Copyright Epic Games, Inc. All Rights Reserved.
 
+#include "HairStrandsVoxelPageCommonStruct.ush"
 #include "../Common.ush"
 #include "../CommonViewUniformBuffer.ush"
 #include "../SceneTextureParameters.ush"
@@ -15,14 +16,13 @@ uint VoxelMarkCount;
 Texture2D<float> SceneDepthTexture;
 RWTexture3D<uint> OutPageTexture;
 
-uint		 VirtualVoxel_PageIndexCount;
-uint		 VirtualVoxel_PageResolution;
-float		 VirtualVoxel_VoxelWorldSize;
-uint3		 VirtualVoxel_PageCountResolution;
-Buffer<uint> VirtualVoxel_PageIndexBuffer;
-Buffer<uint4>VirtualVoxel_PageIndexCoordBuffer;
- 
-StructuredBuffer<FPackedVirtualVoxelNodeDesc> VirtualVoxel_NodeDescBuffer;
+uint		 VirtualVoxelParams_PageIndexCount;
+uint		 VirtualVoxelParams_PageResolution;
+float		 VirtualVoxelParams_VoxelWorldSize;
+uint3		 VirtualVoxelParams_PageCountResolution;
+Buffer<uint> VirtualVoxelParams_PageIndexBuffer;
+Buffer<uint4>VirtualVoxelParams_PageIndexCoordBuffer;
+StructuredBuffer<FPackedVirtualVoxelNodeDesc> VirtualVoxelParams_NodeDescBuffer;
 
 float3 ReconstructWorldPositionFromDepth(float2 UV, float Depth)
 {
@@ -40,27 +40,27 @@ void MainCS(uint3 DispatchThreadId : SV_DispatchThreadID)
 	const uint  MaxDispatchCountPerDim = 65535u;
 	const uint  AllocatedPageIndex  = DispatchThreadId.x + DispatchThreadId.y * MaxDispatchCountPerDim;
 	const uint  LinearVoxelCoord = DispatchThreadId.z;
-	const uint3 VoxelCoordOffset = IndexToCoord(LinearVoxelCoord, VirtualVoxel_PageResolution.xxx);
+	const uint3 VoxelCoordOffset = IndexToCoord(LinearVoxelCoord, VirtualVoxelParams_PageResolution.xxx);
 
-	const FPackedVirtualVoxelNodeDesc PackedNode = VirtualVoxel_NodeDescBuffer[MacroGroupId];
-	const FVirtualVoxelNodeDesc NodeDesc = UnpackVoxelNode(PackedNode, VirtualVoxel_PageResolution);
+	const FPackedVirtualVoxelNodeDesc PackedNode = VirtualVoxelParams_NodeDescBuffer[MacroGroupId];
+	const FVirtualVoxelNodeDesc NodeDesc = UnpackVoxelNode(PackedNode, VirtualVoxelParams_PageResolution);
 
-	const uint4 PageIndexCoord = VirtualVoxel_PageIndexCoordBuffer[NodeDesc.PageIndexOffset + AllocatedPageIndex];
+	const uint4 PageIndexCoord = VirtualVoxelParams_PageIndexCoordBuffer[NodeDesc.PageIndexOffset + AllocatedPageIndex];
 	const uint  LinearPageIndexCoord = CoordToIndex(PageIndexCoord, NodeDesc.PageIndexResolution, NodeDesc.PageIndexOffset);
 
-	const bool bIsValid = LinearPageIndexCoord < VirtualVoxel_PageIndexCount;
+	const bool bIsValid = LinearPageIndexCoord < VirtualVoxelParams_PageIndexCount;
 	if (bIsValid)
 	{
-		const uint PageIndex = VirtualVoxel_PageIndexBuffer.Load(LinearPageIndexCoord);
+		const uint PageIndex = VirtualVoxelParams_PageIndexBuffer.Load(LinearPageIndexCoord);
 		if (PageIndex != INVALID_VOXEL_PAGE_INDEX)
 		{
-			const uint3 VoxelCoordBase = PageIndexCoord.xyz * VirtualVoxel_PageResolution;
+			const uint3 VoxelCoordBase = PageIndexCoord.xyz * VirtualVoxelParams_PageResolution;
 			const uint3 VoxelCoord = VoxelCoordBase + VoxelCoordOffset;
 			const float3 WorldPosition = VoxelCoord / float3(NodeDesc.VirtualResolution) * (NodeDesc.MaxAABB - NodeDesc.MinAABB) + NodeDesc.MinAABB;
 
 			float4 ClipPos = mul(float4(WorldPosition, 1), View.WorldToClip);
 			ClipPos /= ClipPos.w;
-			const float DepthBias = VoxelBiasCount * VirtualVoxel_VoxelWorldSize;
+			const float DepthBias = VoxelBiasCount * VirtualVoxelParams_VoxelWorldSize;
 			const float VoxelDepth = ConvertFromDeviceZ(ClipPos.z) - DepthBias;
 			float2 SceneUV = float2(0.5f*(ClipPos.x + 1), 1 - 0.5f*(ClipPos.y + 1));
 			SceneUV = ViewportUVToBufferUV(SceneUV);
@@ -82,14 +82,14 @@ void MainCS(uint3 DispatchThreadId : SV_DispatchThreadID)
 				return;
 
 			// Inject opaque depth on a thin layer (Dist < DistThreshold) for avoiding weird projection
-			if (ClosestDepth < VoxelDepth && abs(ClosestDepth - VoxelDepth) < VoxelMarkCount * VirtualVoxel_VoxelWorldSize)
+			if (ClosestDepth < VoxelDepth && abs(ClosestDepth - VoxelDepth) < VoxelMarkCount * VirtualVoxelParams_VoxelWorldSize)
 			{
-				const uint3 VoxelPageIndexCoord = VoxelCoord / VirtualVoxel_PageResolution;
-				const uint3 VoxelIndexCoordBase = VoxelPageIndexCoord * VirtualVoxel_PageResolution;
+				const uint3 VoxelPageIndexCoord = VoxelCoord / VirtualVoxelParams_PageResolution;
+				const uint3 VoxelIndexCoordBase = VoxelPageIndexCoord * VirtualVoxelParams_PageResolution;
 				const uint3 VoxelPageOffset = VoxelCoord - VoxelIndexCoordBase;
 
-				const uint3 PageCoord = IndexToCoord(PageIndex, VirtualVoxel_PageCountResolution);
-				const int3 VoxelPageBase = PageCoord * VirtualVoxel_PageResolution;
+				const uint3 PageCoord = IndexToCoord(PageIndex, VirtualVoxelParams_PageCountResolution);
+				const int3 VoxelPageBase = PageCoord * VirtualVoxelParams_PageResolution;
 				const int3 VoxelPageCoord = VoxelPageOffset + VoxelPageBase;
 
 				InterlockedOr(OutPageTexture[VoxelPageCoord], VOXEL_OPAQUE_ADD);
diff --git a/Engine/Shaders/Private/HairStrands/HairStrandsVoxelPageAllocation.usf b/Engine/Shaders/Private/HairStrands/HairStrandsVoxelPageAllocation.usf
--- a/Engine/Shaders/Private/HairStrands/HairStrandsVoxelPageAllocation.usf
+++ b/Engine/Shaders/Private/HairStrands/HairStrandsVoxelPageAllocation.usf
@@ -1,5 +1,6 @@
 // Copyright Epic Games, Inc. All Rights Reserved.
 
+#include "HairStrandsVoxelPageCommonStruct.ush"
 #include "../Common.ush"
 #include "../Matrices.ush"
 #include "HairStrandsVoxelPageCommon.ush"
@@ -610,10 +611,9 @@ void VoxelIndPageClearBufferGenCS(uint GroupIndex : SV_GroupIndex, uint3 Dispatc
 
 #if SHADER_INDPAGECLEAR
 
-uint				VirtualVoxel_PageResolution;
-int3				VirtualVoxel_PageCountResolution;
-
-Buffer<uint4>		VirtualVoxel_PageIndexCoordBuffer;
+uint				VirtualVoxelParams_PageResolution;
+int3				VirtualVoxelParams_PageCountResolution;
+Buffer<uint4>		VirtualVoxelParams_PageIndexCoordBuffer;
 
 RWTexture3D<uint>	OutPageTexture;
 
@@ -622,18 +622,18 @@ void VoxelIndPageClearCS(uint3 DispatchThreadId : SV_DispatchThreadID, uint3 Gro
 {
 	const uint  LinearVoxelCoord = DispatchThreadId.z;
 
-	const uint VoxelCountPerPage = VirtualVoxel_PageResolution.x * VirtualVoxel_PageResolution.x * VirtualVoxel_PageResolution.x;
+	const uint VoxelCountPerPage = VirtualVoxelParams_PageResolution * VirtualVoxelParams_PageResolution * VirtualVoxelParams_PageResolution;
 	//	if (AllocatedPageIndex < TotalAllocatedPageCount)		// Not needed today because the XY size exactly matches the range (page count) we need to clear.
 	if(LinearVoxelCoord < VoxelCountPerPage)
 	{
 		const uint  MaxDispatchCountPerDim = 65535u;
 		const uint  AllocatedPageIndex = DispatchThreadId.x + DispatchThreadId.y * MaxDispatchCountPerDim;
-		const uint3 VoxelCoordOffset = IndexToCoord(LinearVoxelCoord, VirtualVoxel_PageResolution.xxx);
+		const uint3 VoxelCoordOffset = IndexToCoord(LinearVoxelCoord, VirtualVoxelParams_PageResolution.xxx);
 
 		const uint PageIndex = AllocatedPageIndex; // PageIndexBuffer is not needed, we already know those tiles are allocated linearly in 3D within OutPageTexture.
 
-		const uint3 PageCoord = IndexToCoord(PageIndex, VirtualVoxel_PageCountResolution);
-		const int3 VoxelPageBase = PageCoord * VirtualVoxel_PageResolution;
+		const uint3 PageCoord = IndexToCoord(PageIndex, VirtualVoxelParams_PageCountResolution);
+		const int3 VoxelPageBase = PageCoord * VirtualVoxelParams_PageResolution;
 		const int3 VoxelCoord = VoxelPageBase + VoxelCoordOffset;
 
 		OutPageTexture[VoxelCoord] = 0;
diff --git a/Engine/Shaders/Private/HairStrands/HairStrandsVoxelRasterCompute.usf b/Engine/Shaders/Private/HairStrands/HairStrandsVoxelRasterCompute.usf
--- a/Engine/Shaders/Private/HairStrands/HairStrandsVoxelRasterCompute.usf
+++ b/Engine/Shaders/Private/HairStrands/HairStrandsVoxelRasterCompute.usf
@@ -19,14 +19,15 @@ uint			HairStrandsVF_VertexCount;
 float4x4		HairStrandsVF_LocalToWorldPrimitiveTransform;
 Buffer<uint4>	HairStrandsVF_PositionBuffer;
 
-uint			VirtualVoxel_PageIndexCount;
-uint			VirtualVoxel_PageResolution;
-float			VirtualVoxel_VoxelWorldSize;
-uint3			VirtualVoxel_PageCountResolution;
-uint3			VirtualVoxel_PageTextureResolution;
-Buffer<uint>	VirtualVoxel_PageIndexBuffer;
-Buffer<uint4>	VirtualVoxel_PageIndexCoordBuffer;
-StructuredBuffer<FPackedVirtualVoxelNodeDesc> VirtualVoxel_NodeDescBuffer;
+uint			VirtualVoxelParams_PageIndexCount;
+uint			VirtualVoxelParams_PageResolution;
+float			VirtualVoxelParams_VoxelWorldSize;
+uint3			VirtualVoxelParams_PageCountResolution;
+uint3			VirtualVoxelParams_PageTextureResolution;
+
+Buffer<uint>	VirtualVoxelParams_PageIndexBuffer;
+Buffer<uint4>	VirtualVoxelParams_PageIndexCoordBuffer;
+StructuredBuffer<FPackedVirtualVoxelNodeDesc> VirtualVoxelParams_NodeDescBuffer;
 
 RWTexture3D<uint> OutPageTexture;
 
@@ -68,21 +69,21 @@ void MainCS(uint3 DispatchThreadID : SV_DispatchThreadID)
 
 	const float DiameterToRadius = 0.5f;
 	const float3 WP0 = mul(float4(CP0.Position, 1), HairStrandsVF_LocalToWorldPrimitiveTransform).xyz;
-	const float HairCoverage0 = CP0.WorldRadius / max(CP0.WorldRadius, DiameterToRadius * VirtualVoxel_VoxelWorldSize);
+	const float HairCoverage0 = CP0.WorldRadius / max(CP0.WorldRadius, DiameterToRadius * VirtualVoxelParams_VoxelWorldSize);
 
 	const float3 WP1 = mul(float4(CP1.Position, 1), HairStrandsVF_LocalToWorldPrimitiveTransform).xyz;
-	const float HairCoverage1 = CP1.WorldRadius / max(CP1.WorldRadius, DiameterToRadius * VirtualVoxel_VoxelWorldSize);
+	const float HairCoverage1 = CP1.WorldRadius / max(CP1.WorldRadius, DiameterToRadius * VirtualVoxelParams_VoxelWorldSize);
 
-	const FVirtualVoxelNodeDesc NodeDesc = UnpackVoxelNode(VirtualVoxel_NodeDescBuffer[MacroGroupId], VirtualVoxel_PageResolution);
+	const FVirtualVoxelNodeDesc NodeDesc = UnpackVoxelNode(VirtualVoxelParams_NodeDescBuffer[MacroGroupId], VirtualVoxelParams_PageResolution);
 
 	// In order to reduce aliasing, we increase the number of steps. This makes the result more comparable to the raster pass.
 	const float LineStepMultiplier = 1.5f;
 
 	FVirtualVoxelCommonDesc CommonDesc;
-	CommonDesc.PageCountResolution		= VirtualVoxel_PageCountResolution;
-	CommonDesc.VoxelWorldSize			= VirtualVoxel_VoxelWorldSize / LineStepMultiplier;
-	CommonDesc.PageTextureResolution	= VirtualVoxel_PageTextureResolution;
-	CommonDesc.PageResolution			= VirtualVoxel_PageResolution;
+	CommonDesc.PageCountResolution		= VirtualVoxelParams_PageCountResolution;
+	CommonDesc.VoxelWorldSize			= VirtualVoxelParams_VoxelWorldSize / LineStepMultiplier;
+	CommonDesc.PageTextureResolution	= VirtualVoxelParams_PageTextureResolution;
+	CommonDesc.PageResolution			= VirtualVoxelParams_PageResolution;
 
 	// Count the number of fibers which are within a cylinder defined by the voxel size, 
 	// and the distance between the origin and the extent of the volume
@@ -125,7 +126,7 @@ void MainCS(uint3 DispatchThreadID : SV_DispatchThreadID)
 		{
 			CurrentPageIndexCoord = PageIndexCoord;
 			const uint LinearPageIndexCoord = CoordToIndex(PageIndexCoord, NodeDesc.PageIndexResolution, NodeDesc.PageIndexOffset);
-			const uint PageIndex = VirtualVoxel_PageIndexBuffer.Load(LinearPageIndexCoord);
+			const uint PageIndex = VirtualVoxelParams_PageIndexBuffer.Load(LinearPageIndexCoord);
 
 			bIsPageValid = PageIndex != INVALID_VOXEL_PAGE_INDEX;
 			if (bIsPageValid)
diff --git a/Engine/Source/Runtime/Renderer/Private/HairStrands/HairStrandsVoxelization.cpp b/Engine/Source/Runtime/Renderer/Private/HairStrands/HairStrandsVoxelization.cpp
--- a/Engine/Source/Runtime/Renderer/Private/HairStrands/HairStrandsVoxelization.cpp
+++ b/Engine/Source/Runtime/Renderer/Private/HairStrands/HairStrandsVoxelization.cpp
@@ -93,7 +93,7 @@ class FVirtualVoxelInjectOpaqueCS : public FGlobalShader
 
 	BEGIN_SHADER_PARAMETER_STRUCT(FParameters, )
 		SHADER_PARAMETER_STRUCT_INCLUDE(FSceneTextureParameters, SceneTextures)
-		SHADER_PARAMETER_STRUCT(FVirtualVoxelCommonParameters, VirtualVoxel)
+		SHADER_PARAMETER_STRUCT(FVirtualVoxelCommonParameters, VirtualVoxelParams)
 		SHADER_PARAMETER(FIntVector, DispatchedPageIndexResolution)
 		SHADER_PARAMETER(uint32, MacroGroupId)
 		SHADER_PARAMETER(FVector2D, SceneDepthResolution)
@@ -137,7 +137,7 @@ static void AddVirtualVoxelInjectOpaquePass(
 
 	FVirtualVoxelInjectOpaqueCS::FParameters* Parameters = GraphBuilder.AllocParameters<FVirtualVoxelInjectOpaqueCS::FParameters>();
 	Parameters->ViewUniformBuffer = View.ViewUniformBuffer;
-	Parameters->VirtualVoxel = VoxelResources.Parameters.Common;
+	Parameters->VirtualVoxelParams = VoxelResources.Parameters.Common;
 	Parameters->VoxelBiasCount = FMath::Max(0, GHairStransVoxelInjectOpaqueBiasCount);
 	Parameters->VoxelMarkCount = FMath::Max(0, GHairStransVoxelInjectOpaqueMarkCount);
 	Parameters->SceneDepthResolution = SceneTextures.SceneDepthBuffer->Desc.Extent;
@@ -400,7 +400,7 @@ class FVoxelIndPageClearCS : public FGlobalShader
 	SHADER_USE_PARAMETER_STRUCT(FVoxelIndPageClearCS, FGlobalShader);
 
 	BEGIN_SHADER_PARAMETER_STRUCT(FParameters, )
-		SHADER_PARAMETER_STRUCT(FVirtualVoxelCommonParameters, VirtualVoxel)
+		SHADER_PARAMETER_STRUCT(FVirtualVoxelCommonParameters, VirtualVoxelParams)
 		SHADER_PARAMETER_RDG_TEXTURE_UAV(Texture3D, OutPageTexture)
 		SHADER_PARAMETER_RDG_BUFFER(Buffer, IndirectDispatchBuffer)
 	END_SHADER_PARAMETER_STRUCT()
@@ -962,7 +962,7 @@ static FRDGBufferRef IndirectVoxelPageClear(
 		FRDGBufferRef IndirectDispatchArgsBuffer = GraphBuilder.RegisterExternalBuffer(VoxelResources.IndirectArgsBuffer, TEXT("HairVoxelIndirectDispatchArgs"));
 
 		FVoxelIndPageClearCS::FParameters* Parameters = GraphBuilder.AllocParameters<FVoxelIndPageClearCS::FParameters>();
-		Parameters->VirtualVoxel = VoxelResources.Parameters.Common;
+		Parameters->VirtualVoxelParams = VoxelResources.Parameters.Common;
 		Parameters->OutPageTexture = GraphBuilder.CreateUAV(FRDGTextureUAVDesc(OutPageTexture));
 		Parameters->IndirectDispatchBuffer = ClearIndArgsBuffer;
 
@@ -987,7 +987,7 @@ class FVoxelRasterComputeCS : public FGlobalShader
 	SHADER_USE_PARAMETER_STRUCT(FVoxelRasterComputeCS, FGlobalShader);
 
 	BEGIN_SHADER_PARAMETER_STRUCT(FParameters, )
-		SHADER_PARAMETER_STRUCT(FVirtualVoxelCommonParameters, VirtualVoxel)
+		SHADER_PARAMETER_STRUCT(FVirtualVoxelCommonParameters, VirtualVoxelParams)
 		SHADER_PARAMETER(uint32, MacroGroupId)
 		SHADER_PARAMETER(uint32, DispatchCountX)
 		SHADER_PARAMETER(uint32, MaxRasterCount)
@@ -1041,7 +1041,7 @@ static void AddVirtualVoxelizationComputeRasterPass(
 		{
 			FVoxelRasterComputeCS::FParameters* PassParameters = GraphBuilder.AllocParameters<FVoxelRasterComputeCS::FParameters>();
 			PassParameters->MaxRasterCount = FMath::Clamp(GHairStrandsVoxelComputeRasterMaxVoxelCount, 1, 256);
-			PassParameters->VirtualVoxel = VoxelResources.Parameters.Common;
+			PassParameters->VirtualVoxelParams = VoxelResources.Parameters.Common;
 			PassParameters->MacroGroupId = MacroGroup.MacroGroupId;
 			PassParameters->VoxelizationViewInfoBuffer = VoxelizationViewInfoBufferSRV;
 			PassParameters->DispatchCountX = DispatchCountX;
