Fix debug voxel traversal to be similar to shadow/transmittance evaluation.

diff --git a/Engine/Shaders/Private/HairStrands/HairStrandsVoxelPageRayMarching.usf b/Engine/Shaders/Private/HairStrands/HairStrandsVoxelPageRayMarching.usf
--- a/Engine/Shaders/Private/HairStrands/HairStrandsVoxelPageRayMarching.usf
+++ b/Engine/Shaders/Private/HairStrands/HairStrandsVoxelPageRayMarching.usf
@@ -12,6 +12,14 @@
 #include "HairStrandsVoxelPageCommon.ush"
 #include "HairStrandsDebugCommon.ush" 
 
+#define VOXEL_TRAVERSAL_DEBUG 0
+#if PERMUTATION_TRAVERSAL == 1
+	#define VOXEL_TRAVERSAL_TYPE VOXEL_TRAVERSAL_LINEAR_MIPMAP
+#else
+	#define VOXEL_TRAVERSAL_TYPE VOXEL_TRAVERSAL_LINEAR
+#endif
+#include "HairStrandsVoxelPageTraversal.ush"
+
 float2 OutputResolution; 
 uint bDrawPage;
 uint MacroGroupId;
@@ -46,10 +54,6 @@ void MainCS(uint3 DispatchThreadId : SV_DispatchThreadID)
 	const float2 PixelCoord = DispatchThreadId.xy + 0.5f;
 	const float2 UV = PixelCoord / float2(OutputResolution);
 
-	// Reverse-Z
-	const float3 O = View.WorldCameraOrigin;
-	const float3 E = ReconstructWorldPositionFromDepth(UV, ConvertFromDeviceZ(0.0001f));
-	const float3 D = (E - O);
 
 	float  WorldDepth = 0;
 	float  ClipDepth = 0;
@@ -493,28 +497,34 @@ void MainCS(uint3 DispatchThreadId : SV_DispatchThreadID)
 		}
 	}
 
-	const float2 HitT = LineBoxIntersect(O, E, NodeDesc.MinAABB, NodeDesc.MaxAABB);
-	float3 HitColor = float3(0,0,0);
-	float AccDensity = 0;
 	bool bIsValid = false;
-	if (HitT.x < HitT.y)
+	float3 HitColor = 0;
+	if (bDrawPage)
 	{
-		// Use different offseting to see either: 
-		// * Individual voxel (for debugging mainly)
-		// * Smoothed voxel (as used for transmittance computation)
-	#if VOXEL_FILTERING == VOXEL_BLOCK
-		const float3 Offset = 0;
-		const float Dither = InterleavedGradientNoise(PixelCoord.xy, 1);
-	#endif
-	#if VOXEL_FILTERING == VOXEL_SMOOTH
-		const float3 Offset = float3(InterleavedGradientNoise(PixelCoord.xy, 1), InterleavedGradientNoise(PixelCoord.xy, 2), InterleavedGradientNoise(PixelCoord.xy, 3));
-		const float Dither = 0;
-	#endif
-
-
-		// Intersect the indirection table to see if a brick is empty or not
-		if (bDrawPage)
+		// Reverse-Z
+		const float3 O = View.WorldCameraOrigin;
+		const float3 E = ReconstructWorldPositionFromDepth(UV, ConvertFromDeviceZ(0.0001f));
+		const float3 D = (E - O);
+
+		float PixelRadius = -1;
+		const float2 HitT = LineBoxIntersect(O, E, NodeDesc.MinAABB, NodeDesc.MaxAABB);
+		float AccDensity = 0;
+		if (HitT.x < HitT.y)
 		{
+			// Use different offseting to see either: 
+			// * Individual voxel (for debugging mainly)
+			// * Smoothed voxel (as used for transmittance computation)
+		#if VOXEL_FILTERING == VOXEL_BLOCK
+			const float3 Offset = 0;
+			const float Dither = InterleavedGradientNoise(PixelCoord.xy, 1);
+		#endif
+		#if VOXEL_FILTERING == VOXEL_SMOOTH
+			const float3 Offset = float3(InterleavedGradientNoise(PixelCoord.xy, 1), InterleavedGradientNoise(PixelCoord.xy, 2), InterleavedGradientNoise(PixelCoord.xy, 3));
+			const float Dither = 0;
+		#endif
+
+
+			// Intersect the indirection table to see if a brick is empty or not
 			const uint MaxStep = 256;
 			const float Delta = (HitT.y-HitT.x)/float(MaxStep);
 			for (uint StepIt=0; StepIt<MaxStep; ++StepIt)
@@ -533,44 +543,93 @@ void MainCS(uint3 DispatchThreadId : SV_DispatchThreadID)
 				}
 			}
 		}
-		else
+	}
+	else 
+	{
+		FHairTraversalResult TraversalResult = (FHairTraversalResult)0;
+
+		const float DistanceThreshold = 10000;
+		const bool bDebugEnabled = false;
+		const float3 SampleRandom = GetHairVoxelJitter(PixelCoord, View.StateFrameIndexMod8);
+
+		const float3 WP0 = View.WorldCameraOrigin;
+		const float3 E = ReconstructWorldPositionFromDepth(UV, ConvertFromDeviceZ(0.0001f));
+		const float3 WP1 = normalize(E - WP0) * DistanceThreshold + WP0;
+
+		FVirtualVoxelCommonDesc CommonDesc;
+		CommonDesc.PageCountResolution	= VirtualVoxel.PageCountResolution;
+		CommonDesc.VoxelWorldSize		= VirtualVoxel.VoxelWorldSize;
+		CommonDesc.PageTextureResolution= VirtualVoxel.PageTextureResolution;
+		CommonDesc.PageResolution		= VirtualVoxel.PageResolution;
+
+		const FPackedVirtualVoxelNodeDesc PackedNode = VirtualVoxel.NodeDescBuffer[MacroGroupId];
+		const FVirtualVoxelNodeDesc NodeDesc = UnpackVoxelNode(PackedNode, VirtualVoxel.PageResolution);
+
+		FHairTraversalSettings TraversalSettings = InitHairTraversalSettings();
+		TraversalSettings.DensityScale		= VirtualVoxel.DensityScale;
+		TraversalSettings.CountThreshold	= GetOpaqueVoxelValue();
+		TraversalSettings.DistanceThreshold = DistanceThreshold;
+		TraversalSettings.bDebugEnabled		= bDebugEnabled;
+		TraversalSettings.SteppingScale		= VirtualVoxel.SteppingScale;
+		TraversalSettings.Random			= SampleRandom;
+		TraversalSettings.TanConeAngle		= 0;
+		TraversalSettings.bIsPrimaryRay		= true;
+
+		TraversalResult = ComputeHairCountVirtualVoxel(
+			WP0,
+			WP1,
+			CommonDesc,
+			NodeDesc,
+			VirtualVoxel.PageIndexBuffer,
+			VirtualVoxel.PageTexture,
+			TraversalSettings);
+
+		HitColor = GetColorCode(saturate(TraversalResult.HairCount / 16));
+		bIsValid = TraversalResult.HairCount > 0;
+
+		const bool bIsPixelEnable = all(int2(PixelCoord) == GetCursorPos());
+		if (bIsPixelEnable && bIsValid)
 		{
-			const uint MipLevel = 0; //clamp(View.GeneralPurposeTweak, 0, 6);
-			const uint MaxStep = 256;
-			const float Delta = (HitT.y - HitT.x) / float(MaxStep);
-			for (uint StepIt = 0; StepIt < MaxStep; ++StepIt)
-			{
-				const float3 HitP = Offset + O + ((Dither + StepIt)*Delta + HitT.x) * D;
-				const int3 VolumeCoord = PositionToCoord(HitP, NodeDesc.MinAABB, NodeDesc.MaxAABB, NodeDesc.VirtualResolution);
-				const int3 PageIndexCoord = VolumeCoord / VirtualVoxel.PageResolution;
-				const uint PageIndex = GetVoxelPageIndex(PageIndexCoord, NodeDesc.PageIndexResolution, NodeDesc.PageIndexOffset);
-				if (PageIndex != INVALID_VOXEL_PAGE_INDEX)
-				{
-					const uint3 PageCoord		= IndexToCoord(PageIndex, VirtualVoxel.PageCountResolution);
-					const int3 VoxelPageBase	= PageCoord * VirtualVoxel.PageResolution;
-					const int3 VoxelPageOffset	= VolumeCoord - PageIndexCoord * VirtualVoxel.PageResolution;
-					const int3 VoxelPageCoord	= VoxelPageBase + VoxelPageOffset;
-
-					const FHairTraversalResult Result = GetHairVirtualVoxelDensity(VoxelPageCoord, VirtualVoxel.PageTexture, MipLevel, VirtualVoxel.DensityScale);
-					const float Density = Result.HairCount + (Result.Visibility > 0 ? 0 : 1000000000);
-					if (Density > 0)
-					{
-						AccDensity = max(AccDensity, float(Density));
-						bIsValid = true;
-					}
+			float2 StartPos = float2(50.0f, 50.0f) / OutputResolution;
+			float2 Pos = StartPos;
 
-				}
-			}
-		}
-	}
+			Pos = ShaderPrintSymbol(Pos, _X_);
+			Pos = ShaderPrintSymbol(Pos, _Y_);
+			Pos = ShaderPrintSymbol(Pos, _SPC_);
+			Pos = ShaderPrint(Pos, GetCursorPos());
+			Pos = ShaderPrintNewline(Pos); Pos.x = StartPos.x;
 
-	if (!bDrawPage)
-	{
-		HitColor = GetColorCode(saturate(AccDensity));
+			Pos = ShaderPrintSymbol(Pos, _T_);
+			Pos = ShaderPrintSymbol(Pos, _DOT_);
+			Pos = ShaderPrintSymbol(Pos, _C_);
+			Pos = ShaderPrintSymbol(Pos, _O_);
+			Pos = ShaderPrintSymbol(Pos, _U_);
+			Pos = ShaderPrintSymbol(Pos, _N_);
+			Pos = ShaderPrintSymbol(Pos, _T_);
+			Pos = ShaderPrintSymbol(Pos, _SPC_);
+			Pos = ShaderPrint(Pos, TraversalResult.HairCount);
+			Pos = ShaderPrintNewline(Pos); Pos.x = StartPos.x;
+			
+			Pos = ShaderPrintSymbol(Pos, _T_);
+			Pos = ShaderPrintSymbol(Pos, _DOT_);
+			Pos = ShaderPrintSymbol(Pos, _V_);
+			Pos = ShaderPrintSymbol(Pos, _I_);
+			Pos = ShaderPrintSymbol(Pos, _S_);
+			Pos = ShaderPrintSymbol(Pos, _I_);
+			Pos = ShaderPrintSymbol(Pos, _B_);
+			Pos = ShaderPrintSymbol(Pos, _I_);
+			Pos = ShaderPrintSymbol(Pos, _L_);
+			Pos = ShaderPrintSymbol(Pos, _I_);
+			Pos = ShaderPrintSymbol(Pos, _T_);
+			Pos = ShaderPrintSymbol(Pos, _Y_);
+			Pos = ShaderPrintSymbol(Pos, _SPC_);
+			Pos = ShaderPrint(Pos, TraversalResult.Visibility);
+			Pos = ShaderPrintNewline(Pos); Pos.x = StartPos.x;
+		}
 	}
 
 	if (bIsValid)
 	{
-		OutputTexture[DispatchThreadId.xy] = float4(HitColor,1);
+		OutputTexture[DispatchThreadId.xy] = float4(HitColor, 1);
 	}
 }
diff --git a/Engine/Shaders/Private/HairStrands/HairStrandsVoxelPageTraversal.ush b/Engine/Shaders/Private/HairStrands/HairStrandsVoxelPageTraversal.ush
--- a/Engine/Shaders/Private/HairStrands/HairStrandsVoxelPageTraversal.ush
+++ b/Engine/Shaders/Private/HairStrands/HairStrandsVoxelPageTraversal.ush
@@ -228,6 +228,7 @@ struct FHairTraversalSettings
 	float	TanConeAngle;
 	bool	bUseOpaqueVisibility;
 	bool	bDebugEnabled;
+	bool	bIsPrimaryRay;
 };
 
 float GetHairTraversalMaxT()
@@ -246,6 +247,7 @@ FHairTraversalSettings InitHairTraversalSettings()
 	Out.TanConeAngle = 0;
 	Out.bUseOpaqueVisibility = true;
 	Out.bDebugEnabled = false;
+	Out.bIsPrimaryRay = false;
 	return Out;
 }
 
@@ -301,6 +303,7 @@ FHairTraversalResult ComputeHairCountVirtualVoxel(
 		const float DeltaWorld = OELength / float(MaxStep);
 		float StepScale = 1.0f;
 		float3 PreviousP = O;
+		float ClosestHitT = -1;
 
 		const float3 RandomStepJitter = InSettings.Random.xyz * 2 - 1; // [-1..1]
 		for (float StepIt = 0.0f; StepIt < MaxStep; StepIt += StepScale)
@@ -346,6 +349,11 @@ FHairTraversalResult ComputeHairCountVirtualVoxel(
 				const FHairTraversalResult StepResult = GetHairVirtualVoxelDensity(VoxelPageCoord, InPageTexture, uint(MipLevel), InSettings.DensityScale * HairCountScale);
 				Acc(Out, StepResult);
 
+				if (InSettings.bIsPrimaryRay && StepResult.HairCount > 0)
+				{
+					ClosestHitT = 1;
+				}
+
 				#if VOXEL_TRAVERSAL_DEBUG
 				if (InSettings.bDebugEnabled)
 				{
@@ -367,7 +375,14 @@ FHairTraversalResult ComputeHairCountVirtualVoxel(
 				}
 			}
 			#if VOXEL_TRAVERSAL_TYPE == VOXEL_TRAVERSAL_LINEAR_MIPMAP
-			StepScale = min(InCommonDesc.PageResolution, StepScale * InSettings.SteppingScale);
+			if (InSettings.bIsPrimaryRay && ClosestHitT < 0)
+			{
+				// Start the mipmap traversal only when we have reach the first hit (primaryview)
+			}
+			else
+			{
+				StepScale = min(InCommonDesc.PageResolution, StepScale * InSettings.SteppingScale);
+			}
 			#endif
 			PreviousP = HitP;
 		}
@@ -560,4 +575,4 @@ FHairTraversalResult ComputeHairCountVirtualVoxel(
 
 	return Out;
 }
-#endif
\ No newline at end of file
+#endif
diff --git a/Engine/Source/Runtime/Renderer/Private/HairStrands/HairStrandsDebug.cpp b/Engine/Source/Runtime/Renderer/Private/HairStrands/HairStrandsDebug.cpp
--- a/Engine/Source/Runtime/Renderer/Private/HairStrands/HairStrandsDebug.cpp
+++ b/Engine/Source/Runtime/Renderer/Private/HairStrands/HairStrandsDebug.cpp
@@ -83,7 +83,8 @@ static int32 GHairStrandsDebugPPLL = 0;
 static FAutoConsoleVariableRef CVarHairStrandsDebugPPLL(									TEXT("r.HairStrands.DebugPPLL"),									GHairStrandsDebugPPLL, TEXT("Draw debug per pixel light list rendering."));
 static int32 GHairVirtualVoxel_DrawDebugPage = 0;
 static FAutoConsoleVariableRef CVarHairVirtualVoxel_DrawDebugPage(TEXT("r.HairStrands.Voxelization.Virtual.DrawDebugPage"), GHairVirtualVoxel_DrawDebugPage, TEXT("When voxel debug rendering is enable, render the page bounds, instead of the voxel"));
-
+static int32 GHairVirtualVoxel_DebugTraversalType = 0;
+static FAutoConsoleVariableRef CVarHairVirtualVoxel_DebugTraversalType(TEXT("r.HairStrands.Voxelization.Virtual.DebugTraversalType"), GHairVirtualVoxel_DebugTraversalType, TEXT("Traversal mode (0:linear, 1:mip) for debug voxel visualization."));
 
 // Helper functions for accessing interpolation data for debug purpose.
 // Definitions is in HairStrandsInterface.cpp
@@ -699,6 +700,9 @@ class FVoxelVirtualRaymarchingCS : public FGlobalShader
 	DECLARE_GLOBAL_SHADER(FVoxelVirtualRaymarchingCS);
 	SHADER_USE_PARAMETER_STRUCT(FVoxelVirtualRaymarchingCS, FGlobalShader);
 
+	class FTraversalType : SHADER_PERMUTATION_INT("PERMUTATION_TRAVERSAL", 2);
+	using FPermutationDomain = TShaderPermutationDomain<FTraversalType>;
+	
 	BEGIN_SHADER_PARAMETER_STRUCT(FParameters, )
 		SHADER_PARAMETER_STRUCT_INCLUDE(FSceneTextureParameters, SceneTextures)
 		SHADER_PARAMETER_STRUCT_INCLUDE(ShaderDrawDebug::FShaderDrawDebugParameters, ShaderDrawParameters)
@@ -751,7 +755,9 @@ static void AddVoxelPageRaymarchingPass(
 		ShaderPrint::SetParameters(View, Parameters->ShaderPrintParameters);
 		Parameters->OutputTexture			= GraphBuilder.CreateUAV(OutputTexture);
 
-		TShaderMapRef<FVoxelVirtualRaymarchingCS> ComputeShader(View.ShaderMap);
+		FVoxelVirtualRaymarchingCS::FPermutationDomain PermutationVector;
+		PermutationVector.Set<FVoxelVirtualRaymarchingCS::FTraversalType>(GHairVirtualVoxel_DebugTraversalType > 0 ? 1 : 0);
+		TShaderMapRef<FVoxelVirtualRaymarchingCS> ComputeShader(View.ShaderMap, PermutationVector);
 
 		const FIntVector DispatchCount = DispatchCount.DivideAndRoundUp(FIntVector(OutputTexture->Desc.Extent.X, OutputTexture->Desc.Extent.Y, 1), FIntVector(8, 8, 1));
 		FComputeShaderUtils::AddPass(GraphBuilder, RDG_EVENT_NAME("HairStrandsVoxelVirtualRaymarching"), ComputeShader, Parameters, DispatchCount);
