Fix invalid MIP selection during hair material pass.

The material pass is a pixel shader, but adjacent pixel are not spatially/topologically related, issuing invalid ddx/ddy value and thus wrong MIP level selection. This changelist use the same mechanism than raytracing for forcing material texture to a specific level (when raycone are not available.

diff --git a/Engine/Shaders/Private/Common.ush b/Engine/Shaders/Private/Common.ush
--- a/Engine/Shaders/Private/Common.ush
+++ b/Engine/Shaders/Private/Common.ush
@@ -83,6 +83,23 @@ const static MaterialFloat PI = 3.1415926535897932f;
 const static float MaxHalfFloat = 65504.0f;
 const static float Max10BitsFloat = 64512.0f;
 
+// Control MIP level used for material texture fetches. By default only raytracing 
+// shaders (i.e., !PIXELSHADER) use manual MIP level selection. A material shader 
+// can opt. in to force a specific MIP level.
+//
+// * USE_FORCE_TEXTURE_MIP : enable/disable manual MIP level selection
+// * FORCED_TEXTURE_MIP    : force a specific MIP level
+//
+#if !PIXELSHADER && !defined(USE_FORCE_TEXTURE_MIP)
+	#define USE_FORCE_TEXTURE_MIP 1
+#endif
+#ifndef USE_FORCE_TEXTURE_MIP
+	#define USE_FORCE_TEXTURE_MIP 0
+#endif
+#ifndef FORCED_TEXTURE_MIP
+	#define FORCED_TEXTURE_MIP 0.0f
+#endif
+
 #ifndef USE_RAYTRACED_TEXTURE_RAYCONE_LOD
 #define USE_RAYTRACED_TEXTURE_RAYCONE_LOD (RAYHITGROUPSHADER)
 #endif // USE_RAYTRACED_TEXTURE_RAYCONE_LOD
@@ -96,7 +113,7 @@ float ComputeRayConeLod(Texture2D Tex)
 	int TexArea = Dimensions.x * Dimensions.y;
 	return 0.5f * log2(GlobalRayCone_TexArea * TexArea);
 #else
-    return 0.0f;
+    return FORCED_TEXTURE_MIP;
 #endif
 }
 
@@ -109,7 +126,7 @@ float4 ClampToHalfFloatRange(float4 X) { return clamp(X, float(0).xxxx, MaxHalfF
 // for type checking of the parameters - ironically the type checking is really only needed in GLSL!
 MaterialFloat4 Texture1DSample(Texture1D Tex, SamplerState Sampler, float UV)
 {
-#if !PIXELSHADER
+#if USE_FORCE_TEXTURE_MIP
 	return Tex.SampleLevel(Sampler, UV, 0);
 #else
 	return Tex.Sample(Sampler, UV);
@@ -117,7 +134,7 @@ MaterialFloat4 Texture1DSample(Texture1D Tex, SamplerState Sampler, float UV)
 }
 MaterialFloat4 Texture2DSample(Texture2D Tex, SamplerState Sampler, float2 UV)
 {
-#if !PIXELSHADER
+#if USE_FORCE_TEXTURE_MIP
 	return Tex.SampleLevel(Sampler, UV, ComputeRayConeLod(Tex));
 #else
 	return Tex.Sample(Sampler, UV);
@@ -125,7 +142,7 @@ MaterialFloat4 Texture2DSample(Texture2D Tex, SamplerState Sampler, float2 UV)
 }
 MaterialFloat Texture2DSample_A8(Texture2D Tex, SamplerState Sampler, float2 UV)
 {
-#if !PIXELSHADER
+#if USE_FORCE_TEXTURE_MIP
 	return Tex.SampleLevel(Sampler, UV, ComputeRayConeLod(Tex)) A8_SAMPLE_MASK;
 #else
 	return Tex.Sample(Sampler, UV) A8_SAMPLE_MASK;
@@ -133,7 +150,7 @@ MaterialFloat Texture2DSample_A8(Texture2D Tex, SamplerState Sampler, float2 UV)
 }
 MaterialFloat4 Texture3DSample(Texture3D Tex, SamplerState Sampler, float3 UV)
 {
-#if !PIXELSHADER
+#if USE_FORCE_TEXTURE_MIP
 	return Tex.SampleLevel(Sampler, UV, 0);
 #else
 	return Tex.Sample(Sampler, UV);
@@ -141,7 +158,7 @@ MaterialFloat4 Texture3DSample(Texture3D Tex, SamplerState Sampler, float3 UV)
 }
 MaterialFloat4 TextureCubeSample(TextureCube Tex, SamplerState Sampler, float3 UV)
 {
-#if !PIXELSHADER
+#if USE_FORCE_TEXTURE_MIP
 	return Tex.SampleLevel(Sampler, UV, 0);
 #else
 	return Tex.Sample(Sampler, UV);
@@ -149,7 +166,7 @@ MaterialFloat4 TextureCubeSample(TextureCube Tex, SamplerState Sampler, float3 U
 }
 MaterialFloat4 Texture2DArraySample(Texture2DArray Tex, SamplerState Sampler, float3 UV)
 {
-#if !PIXELSHADER
+#if USE_FORCE_TEXTURE_MIP
 	return Tex.SampleLevel(Sampler, UV, 0);
 #else
 	return Tex.Sample(Sampler, UV);
@@ -165,7 +182,7 @@ MaterialFloat4 Texture2DSampleLevel(Texture2D Tex, SamplerState Sampler, float2
 }
 MaterialFloat4 Texture2DSampleBias(Texture2D Tex, SamplerState Sampler, float2 UV, MaterialFloat MipBias)
 {
-#if !PIXELSHADER
+#if USE_FORCE_TEXTURE_MIP
 	return Tex.SampleLevel(Sampler, UV, ComputeRayConeLod(Tex) + MipBias);
 #else
 	return Tex.SampleBias(Sampler, UV, MipBias);
@@ -181,7 +198,7 @@ MaterialFloat4 Texture3DSampleLevel(Texture3D Tex, SamplerState Sampler, float3
 }
 MaterialFloat4 Texture3DSampleBias(Texture3D Tex, SamplerState Sampler, float3 UV, MaterialFloat MipBias)
 {
-#if !PIXELSHADER
+#if USE_FORCE_TEXTURE_MIP
 	return Tex.SampleBias(Sampler, UV, 0);
 #else
 	return Tex.SampleBias(Sampler, UV, MipBias);
@@ -201,7 +218,7 @@ MaterialFloat TextureCubeSampleDepthLevel(TextureCube TexDepth, SamplerState Sam
 }
 MaterialFloat4 TextureCubeSampleBias(TextureCube Tex, SamplerState Sampler, float3 UV, MaterialFloat MipBias)
 {
-#if !PIXELSHADER
+#if USE_FORCE_TEXTURE_MIP
 	return Tex.SampleLevel(Sampler, UV, 0);
 #else
 	return Tex.SampleBias(Sampler, UV, MipBias);
@@ -213,7 +230,7 @@ MaterialFloat4 TextureCubeSampleGrad(TextureCube Tex, SamplerState Sampler, floa
 }
 MaterialFloat4 TextureExternalSample(TextureExternal Tex, SamplerState Sampler, float2 UV)
 {
-#if !PIXELSHADER
+#if USE_FORCE_TEXTURE_MIP
 	#if SUPPORTS_TEXTURE_EXTERNAL
 		return Tex.SampleLevel(Sampler, UV, 0);
 	#else
@@ -278,7 +295,7 @@ MaterialFloat4 Texture2DArraySampleLevel(Texture2DArray Tex, SamplerState Sample
 }
 MaterialFloat4 Texture2DArraySampleBias(Texture2DArray Tex, SamplerState Sampler, float3 UV, MaterialFloat MipBias)
 {
-#if !PIXELSHADER
+#if USE_FORCE_TEXTURE_MIP
 	return Tex.SampleLevel(Sampler, UV, 0);
 #else
 	return Tex.SampleBias(Sampler, UV, MipBias);
@@ -497,7 +514,7 @@ MaterialFloat4 PositiveClampedPow(MaterialFloat4 X,MaterialFloat4 Y)
 
 float DDX(float Input)
 {
-#if !PIXELSHADER
+#if USE_FORCE_TEXTURE_MIP
 	return 0;
 #else
 	return ddx(Input);
@@ -506,7 +523,7 @@ float DDX(float Input)
 
 float2 DDX(float2 Input)
 {
-#if !PIXELSHADER
+#if USE_FORCE_TEXTURE_MIP
 	return 0;
 #else
 	return ddx(Input);
@@ -515,7 +532,7 @@ float2 DDX(float2 Input)
 
 float3 DDX(float3 Input)
 {
-#if !PIXELSHADER
+#if USE_FORCE_TEXTURE_MIP
 	return 0;
 #else
 	return ddx(Input);
@@ -524,7 +541,7 @@ float3 DDX(float3 Input)
 
 float4 DDX(float4 Input)
 {
-#if !PIXELSHADER
+#if USE_FORCE_TEXTURE_MIP
 	return 0;
 #else
 	return ddx(Input);
@@ -533,7 +550,7 @@ float4 DDX(float4 Input)
 
 float DDY(float Input)
 {
-#if !PIXELSHADER
+#if USE_FORCE_TEXTURE_MIP
 	return 0;
 #else
 	return ddy(Input);
@@ -542,7 +559,7 @@ float DDY(float Input)
 
 float2 DDY(float2 Input)
 {
-#if !PIXELSHADER
+#if USE_FORCE_TEXTURE_MIP
 	return 0;
 #else
 	return ddy(Input);
@@ -551,7 +568,7 @@ float2 DDY(float2 Input)
 
 float3 DDY(float3 Input)
 {
-#if !PIXELSHADER
+#if USE_FORCE_TEXTURE_MIP
 	return 0;
 #else
 	return ddy(Input);
@@ -560,7 +577,7 @@ float3 DDY(float3 Input)
 
 float4 DDY(float4 Input)
 {
-#if !PIXELSHADER
+#if USE_FORCE_TEXTURE_MIP
 	return 0;
 #else
 	return ddy(Input);
diff --git a/Engine/Shaders/Private/HairStrands/HairStrandsMaterialPS.usf b/Engine/Shaders/Private/HairStrands/HairStrandsMaterialPS.usf
--- a/Engine/Shaders/Private/HairStrands/HairStrandsMaterialPS.usf
+++ b/Engine/Shaders/Private/HairStrands/HairStrandsMaterialPS.usf
@@ -1,7 +1,11 @@
 // Copyright Epic Games, Inc. All Rights Reserved.
 
 #include "HairStrandsVisibilityCommon.ush" 
+
+#define USE_FORCE_TEXTURE_MIP 1
+#define FORCED_TEXTURE_MIP 0.0f
 #include "../Common.ush"
+
 #include "../SceneTexturesCommon.ush"
 #include "HairStrandsDeepShadowCommon.ush"
 #include "../DeferredShadingCommon.ush"
