Add pixel radius at depth1 support for hair shader.

diff --git a/Engine/Shaders/Private/HairStrands/HairScatter.usf b/Engine/Shaders/Private/HairStrands/HairScatter.usf
--- a/Engine/Shaders/Private/HairStrands/HairScatter.usf
+++ b/Engine/Shaders/Private/HairStrands/HairScatter.usf
@@ -176,6 +176,7 @@ SamplerState 				LinearSampler;
 
 FHairTransmittanceMask ComputeTransmittanceVoxel(
 	const uint2 PixelCoord,
+	const float PixelRadius, 
 	float3 WorldPosition,
 	float3 SamplePosition,
 	uint MacroGroupId,
@@ -212,6 +213,7 @@ FHairTransmittanceMask ComputeTransmittanceVoxel(
 	TraversalSettings.bDebugEnabled = bDebugEnabled;
 	TraversalSettings.SteppingScale = VirtualVoxel.SteppingScale;
 	TraversalSettings.Random = SampleRandom;
+	TraversalSettings.PixelRadius = PixelRadius;
 
 	FHairTraversalResult Result = ComputeHairCountVirtualVoxel(
 		WorldPosition + DepthBias,
@@ -359,9 +361,11 @@ void MainPS(
 				Settings.DeepShadowDepthBiasScale = 1; // Voxel_DepthBiasScale;
 				Settings.Random = TransmittanceRandom;
 
+				const float PixelRadius = SampleDepth * VirtualVoxel.HairCoveragePixelRadiusAtDepth1;
 				const uint MacroGroupId = 0;
 				TransmittanceMask = ComputeTransmittanceVoxel(
 					PixelCoord,
+					PixelRadius, 
 					WorldPosition,
 					SampleWorldPos,
 					MacroGroupId,
diff --git a/Engine/Shaders/Private/HairStrands/HairStrandsDeepShadowMask.usf b/Engine/Shaders/Private/HairStrands/HairStrandsDeepShadowMask.usf
--- a/Engine/Shaders/Private/HairStrands/HairStrandsDeepShadowMask.usf
+++ b/Engine/Shaders/Private/HairStrands/HairStrandsDeepShadowMask.usf
@@ -456,6 +456,7 @@ void MainPS(
 		TraversalSettings.bDebugEnabled = bDebugEnabled;
 		TraversalSettings.SteppingScale = VirtualVoxel.SteppingScale;
 		TraversalSettings.Random = Random;
+		TraversalSettings.PixelRadius = GBuffer.Depth * VirtualVoxel.HairCoveragePixelRadiusAtDepth1;
 		
 		if (RayMarchMaskTexture.Load(uint3(PixelCoord, 0)).r > 0.0f &&  Categorization.PixelCoverage < CoverageThreshold)
 		{
diff --git a/Engine/Shaders/Private/HairStrands/HairStrandsDeepTransmittanceMask.usf b/Engine/Shaders/Private/HairStrands/HairStrandsDeepTransmittanceMask.usf
--- a/Engine/Shaders/Private/HairStrands/HairStrandsDeepTransmittanceMask.usf
+++ b/Engine/Shaders/Private/HairStrands/HairStrandsDeepTransmittanceMask.usf
@@ -22,17 +22,15 @@
 #define PERMUTATION_TRANSMITTANCE_VIRTUAL_VOXEL 1
 
 
-#if PERMUTATION_TRANSMITTANCE_TYPE == PERMUTATION_TRANSMITTANCE_DEEP_SHADOW
-	#define VOXEL_TRAVERSAL_TYPE VOXEL_TRAVERSAL_LEGACY
-#elif PERMUTATION_TRANSMITTANCE_TYPE == PERMUTATION_TRANSMITTANCE_VIRTUAL_VOXEL
+#if PERMUTATION_TRANSMITTANCE_TYPE == PERMUTATION_TRANSMITTANCE_VIRTUAL_VOXEL
 	#if PERMUTATION_TRAVERSAL == 1
 		#define VOXEL_TRAVERSAL_TYPE VOXEL_TRAVERSAL_LINEAR_MIPMAP
 	#else
 		#define VOXEL_TRAVERSAL_TYPE VOXEL_TRAVERSAL_LINEAR
 	#endif
+	#define VOXEL_TRAVERSAL_DEBUG 0
+	#include "HairStrandsVoxelPageTraversal.ush"
 #endif
-#define VOXEL_TRAVERSAL_DEBUG 0
-#include "HairStrandsVoxelPageTraversal.ush"
 
 float3 ToWorldPosition(float2 UV, float Depth)
 {
@@ -201,6 +199,7 @@ void PrintInfo(uint Index, uint2 PixelCoord, FHairTransmittanceMask Transmittanc
 
 FHairTransmittanceMask ComputeTransmittanceVoxel(
 	const uint2 PixelCoord,
+	const float PixelRadius,
 	float3 WorldPosition,
 	uint MacroGroupId,
 	FGBufferData GBuffer,
@@ -218,10 +217,6 @@ FHairTransmittanceMask ComputeTransmittanceVoxel(
 	const float3 T = GBuffer.WorldNormal;
 	const float SinLightAngle = dot(LightDirection, T);
 
-	// Depth bias
-	// Origin is shifted voxels away towards the light + a constant bias of the size of the voxel
-	const float MaxVoxelSize = VirtualVoxel.VoxelWorldSize;
-
 	// Compute the number of hair count between light & shading point
 	const float DistanceThreshold = 100000;
 	float AvgHairCount = 0;
@@ -239,7 +234,7 @@ FHairTransmittanceMask ComputeTransmittanceVoxel(
 		const float3 SampleRandom = GetHairVoxelJitter(PixelCoord, View.StateFrameIndexMod8);
 		#endif
 		
-		const float3 DepthBias = MaxVoxelSize * (VirtualVoxel.DepthBiasScale*LightDirection + SampleRandom);
+		const float3 DepthBias = VirtualVoxel.VoxelWorldSize * (VirtualVoxel.DepthBiasScale*LightDirection + SampleRandom);
 		const float3 SampleWorldPosition = WorldPosition + DepthBias;
 
 		FVirtualVoxelCommonDesc CommonDesc;
@@ -259,6 +254,7 @@ FHairTransmittanceMask ComputeTransmittanceVoxel(
 		TraversalSettings.SteppingScale = VirtualVoxel.SteppingScale;
 		TraversalSettings.Random = SampleRandom;
 		TraversalSettings.TanConeAngle = 0; // TanLightAngle; // Area light transmission is disabled as it is too noisy for now
+		TraversalSettings.PixelRadius = PixelRadius;
 
 		const FHairTraversalResult Result = ComputeHairCountVirtualVoxel(
 			SampleWorldPosition, 
@@ -338,8 +334,10 @@ void MainCS(uint2 DispatchThreadId : SV_DispatchThreadID)
 
 		const float SceneDepth = ConvertFromDeviceZ(Sample.Depth);
 		const float3 WorldPosition = ToWorldPosition(UV, SceneDepth);
+
+		const float PixelRadius = SceneDepth * VirtualVoxel.HairCoveragePixelRadiusAtDepth1;
 	
-		const FHairTransmittanceMask TransmittanceMask = ComputeTransmittanceVoxel(PixelCoord, WorldPosition, MacroGroupIndex, GBuffer, Settings, HairLUTTexture, LinearSampler, bDebugEnabled);
+		const FHairTransmittanceMask TransmittanceMask = ComputeTransmittanceVoxel(PixelCoord, PixelRadius, WorldPosition, MacroGroupIndex, GBuffer, Settings, HairLUTTexture, LinearSampler, bDebugEnabled);
 
 		#if VOXEL_TRAVERSAL_DEBUG
 		if (bDebugEnabled)
diff --git a/Engine/Shaders/Private/HairStrands/HairStrandsEnvironmentAO.usf b/Engine/Shaders/Private/HairStrands/HairStrandsEnvironmentAO.usf
--- a/Engine/Shaders/Private/HairStrands/HairStrandsEnvironmentAO.usf
+++ b/Engine/Shaders/Private/HairStrands/HairStrandsEnvironmentAO.usf
@@ -25,7 +25,7 @@ float AO_Intensity;
 uint  AO_SampleCount;
 float AO_DistanceThreshold;
 
-float TraceAO(float3 WorldPosition, float3 WorldNormal, float4 SvPosition)
+float TraceAO(float3 WorldPosition, float3 WorldNormal, float4 SvPosition, float PixelRadius)
 {
 	const float3 VoxelOffset = VirtualVoxel.VoxelWorldSize * GetHairVoxelJitter(SvPosition.xy, View.StateFrameIndexMod8);
 
@@ -149,7 +149,8 @@ void MainPS(
 
 	if (bIsValid)
 	{
-		OutAO = TraceAO(WorldPosition, WorldNormal, SvPosition);
+		const float PixelRadius = GBuffer.Depth * VirtualVoxel.HairCoveragePixelRadiusAtDepth1;
+		OutAO = TraceAO(WorldPosition, WorldNormal, SvPosition, PixelRadius);
 	}
 	else
 	{
diff --git a/Engine/Shaders/Private/HairStrands/HairStrandsEnvironmentLighting.usf b/Engine/Shaders/Private/HairStrands/HairStrandsEnvironmentLighting.usf
--- a/Engine/Shaders/Private/HairStrands/HairStrandsEnvironmentLighting.usf
+++ b/Engine/Shaders/Private/HairStrands/HairStrandsEnvironmentLighting.usf
@@ -168,6 +168,7 @@ FHairTraversalResult InternalComputeHairCountVoxel(
 	float ExtraDensityScale, 
 	float CountThreshold, 
 	float DistanceThreshold,
+	float PixelRadius, 
 	bool bUseConeTracing,
 	bool bDebugEnabled)
 {
@@ -186,6 +187,7 @@ FHairTraversalResult InternalComputeHairCountVoxel(
 	TraversalSettings.DistanceThreshold = DistanceThreshold;
 	TraversalSettings.bDebugEnabled = bDebugEnabled;
 	TraversalSettings.SteppingScale = VirtualVoxel.SteppingScale;
+	TraversalSettings.PixelRadius = PixelRadius;
 
 	return ComputeHairCountVirtualVoxel(
 		WorldPosition,
@@ -215,6 +217,7 @@ float3 ReflectionEnvironment(FGBufferData GBuffer, float2 ScreenPosition, float2
 	const float3 CameraToPixel	= normalize(WorldPosition - View.WorldCameraOrigin);
 	const uint2  PixelPosition	= SvPosition - ResolvedView.ViewRectMin.xy;
 	const float  PixelSceneDepth= GBuffer.Depth;
+	const float  PixelRadius	= VirtualVoxel.HairCoveragePixelRadiusAtDepth1 * PixelSceneDepth;
 
 	const float VoxelMaxSize = VirtualVoxel.VoxelWorldSize;
 	const uint PrimitiveIdMod8 = (PrimitiveId % 8);
@@ -289,7 +292,7 @@ float3 ReflectionEnvironment(FGBufferData GBuffer, float2 ScreenPosition, float2
 
 			// Compute the number of hair count between light & shading point
 			const float SinLightAngle = dot(SampleL, T);
-			const FHairTraversalResult Result = InternalComputeHairCountVoxel(SampleWorldPosition, SampleLightPosition, 1, 0, MaxSearchDistance, true, bDebugEnabled);
+			const FHairTraversalResult Result = InternalComputeHairCountVoxel(SampleWorldPosition, SampleLightPosition, 1, 0, MaxSearchDistance, PixelRadius, true, bDebugEnabled);
 			const FHairTransmittanceMask TransmittanceMask = ComputeTransmittanceFromHairCount(Result.HairCount, Result.Visibility, GBuffer, SinLightAngle, HairScatteringLUTTexture, HairLUTSampler);
 
 			const float ScalarTransmittance = max(TransmittanceMask.Transmittance.x, max(TransmittanceMask.Transmittance.y, TransmittanceMask.Transmittance.z));
@@ -417,7 +420,7 @@ float3 ReflectionEnvironment(FGBufferData GBuffer, float2 ScreenPosition, float2
 			const float3 SampleLightPosition = WorldPosition + L * MaxDistance;
 
 			// Compute the number of hair count between light & shading point
-			const FHairTraversalResult Result = InternalComputeHairCountVoxel(SampleWorldPosition, SampleLightPosition, 1, 0, MaxSearchDistance, true, bDebugEnabled);
+			const FHairTraversalResult Result = InternalComputeHairCountVoxel(SampleWorldPosition, SampleLightPosition, 1, 0, MaxSearchDistance, PixelRadius, true, bDebugEnabled);
 			TransmittanceMask = ComputeTransmittanceFromHairCount(Result.HairCount, Result.Visibility, GBuffer, SinLightAngle, HairScatteringLUTTexture, HairLUTSampler);
 		}
 
@@ -482,7 +485,7 @@ float3 ReflectionEnvironment(FGBufferData GBuffer, float2 ScreenPosition, float2
 			}
 			else
 			{
-				Result = InternalComputeHairCountVoxel(SampleWorldPosition, SampleLightPosition, TransmissionDensityScaleFactor, 0, DistanceThresholdTT, false, bDebugEnabled);
+				Result = InternalComputeHairCountVoxel(SampleWorldPosition, SampleLightPosition, TransmissionDensityScaleFactor, 0, DistanceThresholdTT, PixelRadius, false, bDebugEnabled);
 			}
 			TransmittanceMask = ComputeTransmittanceFromHairCount(Result.HairCount, Result.Visibility, GBuffer, SinLightAngle, HairScatteringLUTTexture, HairLUTSampler);
 		}
@@ -546,12 +549,7 @@ float3 ReflectionEnvironmentUniformSampling(
 	const float3 CameraToPixel = normalize(WorldPosition - View.WorldCameraOrigin);
 	const uint2  PixelPosition = SvPosition - ResolvedView.ViewRectMin.xy;
 	const float  PixelSceneDepth = GBuffer.Depth;
-
-#if VOXEL_TRAVERSAL_TYPE != VOXEL_TRAVERSAL_LEGACY
-	const float3 VoxelSize = VirtualVoxel.VoxelWorldSize; // Use the correct voxel size when we use virtual voxel volume.
-#else
-	const float3 VoxelSize = (Voxel_MaxAABB - Voxel_MinAABB) / Voxel_Resolution;
-#endif
+	const float  PixelRadius = VirtualVoxel.HairCoveragePixelRadiusAtDepth1 * PixelSceneDepth;
 
 	const uint PrimitiveIdMod8 = (PrimitiveId % 8);
 	const float VoxelMaxSize = VirtualVoxel.VoxelWorldSize;
@@ -634,6 +632,7 @@ float3 ReflectionEnvironmentUniformSampling(
 			1,
 			0, 
 			MaxSearchDistance, 
+			PixelRadius,
 			true, 
 			false/*bDebugEnabled*/);
 
diff --git a/Engine/Shaders/Private/HairStrands/HairStrandsVoxelPageCommon.ush b/Engine/Shaders/Private/HairStrands/HairStrandsVoxelPageCommon.ush
--- a/Engine/Shaders/Private/HairStrands/HairStrandsVoxelPageCommon.ush
+++ b/Engine/Shaders/Private/HairStrands/HairStrandsVoxelPageCommon.ush
@@ -11,7 +11,6 @@
 #define VOXEL_TRAVERSAL_SPARSE_LINEAR 2
 #define VOXEL_TRAVERSAL_SPARSE_MIPMAP 3
 #define VOXEL_TRAVERSAL_LINEAR_MIPMAP 4
-#define VOXEL_TRAVERSAL_LEGACY 5
 
 ///////////////////////////////////////////////////////////////////////////////////////////////////
 uint3 PositionToCoord(float3 P, float3 InMinAABB, float3 InMaxAABB, uint3 InResolution)
diff --git a/Engine/Shaders/Private/HairStrands/HairStrandsVoxelPageRayMarching.usf b/Engine/Shaders/Private/HairStrands/HairStrandsVoxelPageRayMarching.usf
--- a/Engine/Shaders/Private/HairStrands/HairStrandsVoxelPageRayMarching.usf
+++ b/Engine/Shaders/Private/HairStrands/HairStrandsVoxelPageRayMarching.usf
@@ -574,6 +574,7 @@ void MainCS(uint3 DispatchThreadId : SV_DispatchThreadID)
 		TraversalSettings.Random			= SampleRandom;
 		TraversalSettings.TanConeAngle		= 0;
 		TraversalSettings.bIsPrimaryRay		= true;
+		TraversalSettings.PixelRadius		= -1;
 
 		TraversalResult = ComputeHairCountVirtualVoxel(
 			WP0,
diff --git a/Engine/Shaders/Private/HairStrands/HairStrandsVoxelPageTraversal.ush b/Engine/Shaders/Private/HairStrands/HairStrandsVoxelPageTraversal.ush
--- a/Engine/Shaders/Private/HairStrands/HairStrandsVoxelPageTraversal.ush
+++ b/Engine/Shaders/Private/HairStrands/HairStrandsVoxelPageTraversal.ush
@@ -75,147 +75,6 @@ float3 GetHairVoxelJitter(uint2 PixelCoord, uint Seed)
 		InterleavedGradientNoise(PixelCoord.xy, Seed * 7901));
 }
 
-/////////////////////////////////////////////////////////////////////////////////////////////////////////////////////////////////
-// Non-Sparse voxel traversal routines
-#if VOXEL_TRAVERSAL_TYPE ==  VOXEL_TRAVERSAL_LEGACY
-FHairTraversalResult ComputeHairCountVoxel(
-	float3 WorldPosition0,
-	float3 WorldPosition1,
-	float3 InVoxelAABBMin,
-	float3 InVoxelAABBMax,
-	float InVoxelResolution,
-	Texture3D<uint> InDensityTexture,
-	float DensityScale,
-	float CountThreshold,
-	float DistanceThreshold)
-{
-	FHairTraversalResult Out = InitHairTraversalResult();
-
-	const float3 VolumeBoundSize = InVoxelAABBMax - InVoxelAABBMin;
-	const float3 VolumeBoundSizeInv = 1.0f / (InVoxelAABBMax - InVoxelAABBMin);
-	const float VoxelWorldSize = VolumeBoundSize.x / InVoxelResolution;				// Only using voxel size along x for now...
-
-	const float2 HitT = LineBoxIntersect(WorldPosition0, WorldPosition1, InVoxelAABBMin, InVoxelAABBMax);
-	if (HitT.x < HitT.y)
-	{
-		// Count the number of fibers which are within a cylinder defined by the voxel size, 
-		// and the distance between the origin and the extent of the volume
-		// This assumes that the voxel volume is cubic (i.e. equal dimensions on all sides)
-		float3 O = lerp(WorldPosition0, WorldPosition1, HitT.xxx);
-		float3 E = lerp(WorldPosition0, WorldPosition1, HitT.yyy);
-		const float OELength = min(length(E - O), DistanceThreshold);
-
-		// Step according to voxel size
-		const float StepCount = float(min(ceil(OELength / VoxelWorldSize), 256u));	// TO TRY: sample with fixed count budget and jitter temporally
-		const float3 D = normalize(WorldPosition1 - WorldPosition0) * VoxelWorldSize;
-
-		for (float StepIt = 0.0f; StepIt < StepCount; ++StepIt)
-		{
-			const float3 HitP = O + StepIt * D;
-			const float3 UVW = (HitP - InVoxelAABBMin) * VolumeBoundSizeInv;
-			const int3 VolumeCoord = clamp(int3(UVW * InVoxelResolution), int3(0, 0, 0), int3(InVoxelResolution.xxx - 1));
-			const FHairTraversalResult StepResult = GetHairVoxelDensity(VolumeCoord, InDensityTexture, 0, DensityScale);
-			Acc(Out, StepResult);
-
-			if (ShouldStopTraversal(Out, CountThreshold))
-			{
-				Out.HitT = StepIt * VoxelWorldSize;
-				return Out;
-			}
-		}
-	}
-
-	return Out;
-}
-
-FHairTraversalResult ComputeHairCountVoxelCone(
-	float3 WorldPosition0,
-	float3 WorldPosition1,
-	float3 InVoxelAABBMin,
-	float3 InVoxelAABBMax,
-	float InVoxelResolution,
-	Texture3D<uint> InDensityTexture,
-	float InVoxelDensityScale,
-	float CountThreshold,
-	float DistanceThreshold,
-	float TanConeAngle)
-{
-	FHairTraversalResult Out = InitHairTraversalResult();
-
-	// Reverse-Z
-	const float ConservativeAABBDiagonal = (InVoxelAABBMax.x - InVoxelAABBMin.x) * 3;
-	const float3 O = WorldPosition0;
-	const float3 E = WorldPosition1;
-	const float3 D = E - O;
-	const float3 NormD = normalize(D);
-
-	// The rest of the code assumes O is within the AABB
-	const bool bIsInAABB = IsInVoxelBounds(O, InVoxelAABBMin, InVoxelAABBMax);
-	if (!bIsInAABB)
-		return Out;
-
-	const float2 HitT = LineBoxIntersect(O, E, InVoxelAABBMin, InVoxelAABBMax);
-	if (HitT.x < HitT.y)
-	{
-		const float VolumeSize = (InVoxelAABBMax.x - InVoxelAABBMin.x);
-		const float VoxelSize = VolumeSize / InVoxelResolution;
-
-		// Count the number of fibers which are within a cylinder defined by the voxel size, 
-		// and the distance between the origin and the extent of the volume
-		// This assumes that the voxel volume is cubic (i.e. equal dimensions on all sides)
-		const float DistanceThreshold2 = DistanceThreshold * DistanceThreshold;
-		const float Distance = min(DistanceThreshold, length(HitT.y * D));
-		const uint MaxStepCount = min(ceil(Distance / VoxelSize), 256u);
-		float T = 0;
-		float MipLevel = 0;
-		float MipDensityScale = 1;
-		for (uint StepIt = 0; StepIt < MaxStepCount; ++StepIt)
-		{
-			float3 StepOffset = 0;
-
-			const float3 RaymarchOffset = T * NormD;
-			const float3 HitP = O + RaymarchOffset + StepOffset;
-			const float3 UVW = (HitP - InVoxelAABBMin) / (InVoxelAABBMax - InVoxelAABBMin);
-
-			const float Distance2 = dot(RaymarchOffset, RaymarchOffset);
-			if (Distance2 > DistanceThreshold2)
-				return Out;
-
-			if (!IsInVoxelBounds(UVW, 0, 1))
-				return Out;
-
-			const uint MipVoxelResolution = uint(InVoxelResolution) >> uint(MipLevel);
-			const int3 VolumeCoord = clamp(int3(UVW * MipVoxelResolution), int3(0, 0, 0), int3(MipVoxelResolution.xxx - 1));
-			const FHairTraversalResult StepResult = GetHairVoxelDensity(VolumeCoord, InDensityTexture, MipLevel, InVoxelDensityScale * MipDensityScale);
-			Acc(Out, StepResult);
-
-			if (ShouldStopTraversal(Out, CountThreshold))
-			{
-				Out.HitT = T;
-				return Out;
-			}
-
-			const float VoxelRadius = T * TanConeAngle;
-			const float NumVoxel = VolumeSize / max(2 * VoxelRadius, VoxelSize);
-
-			const float ResolutionBias = InVoxelResolution / NumVoxel;
-			float MipBias = log2(max(0.5f, ResolutionBias));
-			MipBias = floor(MipBias);
-			MipLevel = MipBias;
-
-			const float MipVoxelSize = VoxelSize * (MipLevel + 1);
-			{
-				MipDensityScale = VoxelSize * rcp(MipVoxelSize);
-				MipDensityScale = MipDensityScale * MipDensityScale * MipDensityScale;
-			}
-			T += MipVoxelSize;
-		}
-	}
-
-	return Out;
-}
-#endif
-
 /////////////////////////////////////////////////////////////////////////////////////////////////////////////////////////////////
 
 struct FHairTraversalSettings
@@ -226,6 +85,7 @@ struct FHairTraversalSettings
 	float	SteppingScale;		// Rate at which the raymarching step increase
 	float3	Random;
 	float	TanConeAngle;
+	float	PixelRadius;
 	bool	bUseOpaqueVisibility;
 	bool	bDebugEnabled;
 	bool	bIsPrimaryRay;
@@ -245,6 +105,7 @@ FHairTraversalSettings InitHairTraversalSettings()
 	Out.SteppingScale = 1;
 	Out.Random = 0.5f;
 	Out.TanConeAngle = 0;
+	Out.PixelRadius = 0;
 	Out.bUseOpaqueVisibility = true;
 	Out.bDebugEnabled = false;
 	Out.bIsPrimaryRay = false;
@@ -267,9 +128,6 @@ FHairTraversalResult ComputeHairCountVirtualVoxel(
 	Error: VOXEL TRAVERSAL TYPE is undefined
 }
 #endif
-#if VOXEL_TRAVERSAL_TYPE == VOXEL_TRAVERSAL_LEGACY
-; // This function is never used. When VOXEL_TRAVERSAL_LEGACY is defined, the old traversal (with plain voxel) is used instead
-#endif
 #if VOXEL_TRAVERSAL_TYPE == VOXEL_TRAVERSAL_LINEAR || VOXEL_TRAVERSAL_TYPE == VOXEL_TRAVERSAL_LINEAR_MIPMAP
 {
 	const float MipLevelDistanceScale = 1.0f / InCommonDesc.VoxelWorldSize;
diff --git a/Engine/Shaders/Private/RayTracing/RayTracingOcclusionRGS.usf b/Engine/Shaders/Private/RayTracing/RayTracingOcclusionRGS.usf
--- a/Engine/Shaders/Private/RayTracing/RayTracingOcclusionRGS.usf
+++ b/Engine/Shaders/Private/RayTracing/RayTracingOcclusionRGS.usf
@@ -51,10 +51,14 @@ uint				bUseHairVoxel;
 Texture2D<uint4>	HairCategorizationTexture;
 Texture2D<uint>		HairLightChannelMaskTexture;
 
-float TraverseHair(uint2 PixelCoord, inout RandomSequence RandSequence, float3 RayOrigin, float3 RayDirection, float InMinT)
+float TraverseHair(uint2 PixelCoord, float DeviceZ, inout RandomSequence RandSequence, float3 RayOrigin, float3 RayDirection, float InMinT)
 {
 	float HitT = InMinT;
 
+
+	const float SceneDepth = ConvertFromDeviceZ(DeviceZ);
+	const float PixelRadius = VirtualVoxel.HairCoveragePixelRadiusAtDepth1 * SceneDepth;
+
 	// Depth bias
 	// Origin is shifted voxels away towards the light + a constant bias of the size of the voxel
 #if 0
@@ -103,6 +107,7 @@ float TraverseHair(uint2 PixelCoord, inout RandomSequence RandSequence, float3 R
 		TraversalSettings.SteppingScale = VirtualVoxel.SteppingScale;
 		TraversalSettings.Random = VoxelRandom;
 		TraversalSettings.bUseOpaqueVisibility = false;
+		TraversalSettings.PixelRadius = PixelRadius;
 
 		Result = ComputeHairCountVirtualVoxel(
 			P0,
@@ -452,7 +457,7 @@ FOcclusionResult ComputeOcclusion(
 		#if USE_HAIR_LIGHTING
 		if (ShadingModelID != SHADINGMODELID_HAIR && bUseHairVoxel)
 		{
-			MinimalPayload.HitT = TraverseHair(PixelCoord, RandSequence, Ray.Origin, Ray.Direction, MinimalPayload.HitT);
+			MinimalPayload.HitT = TraverseHair(PixelCoord, DeviceZ, RandSequence, Ray.Origin, Ray.Direction, MinimalPayload.HitT);
 		}
 		#endif
 
diff --git a/Engine/Source/Runtime/Renderer/Private/HairStrands/HairStrandsVoxelization.cpp b/Engine/Source/Runtime/Renderer/Private/HairStrands/HairStrandsVoxelization.cpp
--- a/Engine/Source/Runtime/Renderer/Private/HairStrands/HairStrandsVoxelization.cpp
+++ b/Engine/Source/Runtime/Renderer/Private/HairStrands/HairStrandsVoxelization.cpp
@@ -822,6 +822,10 @@ FVirtualVoxelResources AllocateVirtualVoxelResources(
 	Out.Parameters.Common.SteppingScale				= FMath::Clamp(GHairStransVoxelRaymarchingSteppingScale, 1.f, 10.f);
 	Out.Parameters.Common.NodeDescCount				= MacroGroups.Datas.Num();
 	Out.Parameters.Common.IndirectDispatchGroupSize = 64;
+	
+	Out.Parameters.Common.HairCoveragePixelRadiusAtDepth1	= ComputeMinStrandRadiusAtDepth1(FIntPoint(View.ViewRect.Width(), View.ViewRect.Height()), View.FOV, 1/*SampleCount*/, 1/*RasterizationScale*/).Primary;
+	Out.Parameters.Common.HairCoverageLUT					= GetHairLUT(RHICmdList, View).Textures[HairLUTType_Coverage]->GetRenderTargetItem().TargetableTexture;
+	Out.Parameters.Common.HairCoverageSampler				= TStaticSamplerState<SF_Bilinear, AM_Clamp, AM_Clamp, AM_Clamp>::GetRHI();
 
 	AddAllocateVoxelPagesPass(
 		GraphBuilder, 
diff --git a/Engine/Source/Runtime/Renderer/Private/HairStrands/HairStrandsVoxelization.h b/Engine/Source/Runtime/Renderer/Private/HairStrands/HairStrandsVoxelization.h
--- a/Engine/Source/Runtime/Renderer/Private/HairStrands/HairStrandsVoxelization.h
+++ b/Engine/Source/Runtime/Renderer/Private/HairStrands/HairStrandsVoxelization.h
@@ -31,6 +31,9 @@ struct FPackedVirtualVoxelNodeDesc
 	uint32	PageIndexOffset;
 };
 
+// PixelRadiusAtDepth1 shouldn't be stored into this structure should be view independent, 
+// but is put here for convenience at the moment since multiple views are not supported 
+// at the moment
 BEGIN_SHADER_PARAMETER_STRUCT(FVirtualVoxelCommonParameters, )
 	SHADER_PARAMETER(FIntVector, PageCountResolution)
 	SHADER_PARAMETER(float, VoxelWorldSize)
@@ -43,9 +46,12 @@ BEGIN_SHADER_PARAMETER_STRUCT(FVirtualVoxelCommonParameters, )
 	SHADER_PARAMETER(float, DensityScale)
 	SHADER_PARAMETER(float, DepthBiasScale)
 	SHADER_PARAMETER(float, SteppingScale)
+	SHADER_PARAMETER(float, HairCoveragePixelRadiusAtDepth1) 
 	SHADER_PARAMETER_SRV(Buffer<uint>, PageIndexBuffer)
 	SHADER_PARAMETER_SRV(Buffer<uint>, PageIndexCoordBuffer)
 	SHADER_PARAMETER_SRV(StructuredBuffer<FPackedVirtualVoxelNodeDesc>, NodeDescBuffer) // Packed into 2 x uint4
+	SHADER_PARAMETER_TEXTURE(Texture2D<float>, HairCoverageLUT)
+	SHADER_PARAMETER_SAMPLER(SamplerState, HairCoverageSampler)
 END_SHADER_PARAMETER_STRUCT()
 
 BEGIN_GLOBAL_SHADER_PARAMETER_STRUCT(FVirtualVoxelParameters, RENDERER_API)
