Fix hair sky uniform integration having missing back lighting.

diff --git a/Engine/Shaders/Private/HairStrands/HairStrandsBsdfPlot.usf b/Engine/Shaders/Private/HairStrands/HairStrandsBsdfPlot.usf
--- a/Engine/Shaders/Private/HairStrands/HairStrandsBsdfPlot.usf
+++ b/Engine/Shaders/Private/HairStrands/HairStrandsBsdfPlot.usf
@@ -232,6 +232,8 @@ void MainPS(
 			TransmittanceData.ScatteringComponent = HairComponents;
 			TransmittanceData.OpaqueVisibility = 1;
 
+			const float SampleBacklit = 1;
+			const float SampleSpecular = 0.5f;
 			FGBufferData HairGBuffer = (FGBufferData)0;
 			HairGBuffer.WorldNormal = P.T;
 			HairGBuffer.BaseColor = P.BaseColor;
@@ -239,14 +241,16 @@ void MainPS(
 			HairGBuffer.ShadingModelID = SHADINGMODELID_HAIR;
 			HairGBuffer.DiffuseColor = 0;
 			HairGBuffer.SpecularColor = 0;
-			HairGBuffer.Specular = 0.5f; //TODO
+			HairGBuffer.Specular = SampleSpecular;
 			HairGBuffer.Metallic = 0;
 			HairGBuffer.Depth = 0.5f;
 			HairGBuffer.GBufferAO = 1;
-			HairGBuffer.CustomData = 0;
+			HairGBuffer.CustomData = float4(0,0,SampleBacklit,0);
 			HairGBuffer.IndirectIrradiance = 1000000;
 
-			OutColor.xyz = HairShading(HairGBuffer, PixelDir_Local, V_Local, T_Local, 1, TransmittanceData, 1, 0, uint2(0, 0)) * Exposure;
+			const float Area = 0;
+			const float BackLit = 1;
+			OutColor.xyz = HairShading(HairGBuffer, PixelDir_Local, V_Local, T_Local, 1, TransmittanceData, BackLit, Area, uint2(0, 0)) * Exposure;
 		}
 
 		for (uint SampleIt = 0; SampleIt < P.SampleCount; SampleIt++)
diff --git a/Engine/Shaders/Private/HairStrands/HairStrandsEnvironmentLighting.usf b/Engine/Shaders/Private/HairStrands/HairStrandsEnvironmentLighting.usf
--- a/Engine/Shaders/Private/HairStrands/HairStrandsEnvironmentLighting.usf
+++ b/Engine/Shaders/Private/HairStrands/HairStrandsEnvironmentLighting.usf
@@ -654,11 +654,13 @@ float3 ReflectionEnvironmentUniformSampling(
 			HairComponents);	
 
 		const float Area = 0;
+		const float Backlit = 1;
+		const float RoughnessBias = 0.1f;
 		AccLighting += 
-			HairShading(GBuffer, SampleL, V, T, 1, TransmittanceData, 0, Area, uint2(0, 0)) * 
+			HairShading(GBuffer, SampleL, V, T, 1, TransmittanceData, Backlit, Area, uint2(0, 0)) * 
 			Result.Visibility * 
 			PreExposure * 
-			GatherRadiance(SampleL, 0.1f, PixelPosition, PixelSceneDepth, WorldPosition);
+			GatherRadiance(SampleL, RoughnessBias, PixelPosition, PixelSceneDepth, WorldPosition);
 
 		#if PERMUTATION_DEBUG
 		if (bDebugEnabled)
