Fix HQreflection mobile. First empty spot in the in the ReflectionCubeMapTextures(if there is any) will be taken by the Sky

diff --git a/Engine/Shaders/Private/MobileBasePassPixelShader.usf b/Engine/Shaders/Private/MobileBasePassPixelShader.usf
--- a/Engine/Shaders/Private/MobileBasePassPixelShader.usf
+++ b/Engine/Shaders/Private/MobileBasePassPixelShader.usf
@@ -116,6 +116,10 @@ float4 GetSphereCaptureVector(float3 ReflectionVector, float3 WorldPosition, flo
 	ProjectedCaptureVector.w = 0;
 	ProjectedCaptureVector.xyz = ReflectionVector;
 
+	// Use ReflectionVector if we are defaulting to SkySphere
+	if (CapturePositionAndRadius.w < 0)
+		return ProjectedCaptureVector;
+
 	float3 RayDirection = ReflectionVector;
 	float ProjectionSphereRadius = CapturePositionAndRadius.w * 1.2f;
 	float SphereRadiusSquared = ProjectionSphereRadius * ProjectionSphereRadius;
@@ -188,6 +192,8 @@ float4 GetBoxCaptureVector(float3 ReflectionVector, float3 WorldPosition, float4
 half3 BlendReflectionCapture(FMaterialPixelParameters MaterialParameters, TextureCube ReflectionCube, SamplerState ReflectionSampler, half MipLevel, int ReflectionIndex, inout half BlendFactor)
 {
 	half4 ProjectedCaptureVector;
+	half3 SpecularIBL = half3(0.0f, 0.0f, 0.0f);
+
 #if ALLOW_CUBE_REFLECTIONS
 	if (CaptureBoxScalesArray[ReflectionIndex].w > 0)
 	{
@@ -201,11 +207,24 @@ half3 BlendReflectionCapture(FMaterialPixelParameters MaterialParameters, Textur
 
 	// Fetch from cubemap and convert to linear HDR
 	half4 Reflection = ReflectionCube.SampleLevel(ReflectionSampler, ProjectedCaptureVector.xyz, MipLevel);
-	half3 SpecularIBL = RGBMDecode(Reflection, 16.0f);
+	
+	if (ReflectionPositionsAndRadii[ReflectionIndex].w > 0)
+	{
+		SpecularIBL = RGBMDecode(Reflection, 16.0f);
 
-	SpecularIBL = SpecularIBL * SpecularIBL * (ProjectedCaptureVector.w * BlendFactor);
+		SpecularIBL = SpecularIBL * SpecularIBL * (ProjectedCaptureVector.w * BlendFactor);
 
-	BlendFactor = BlendFactor * (1.0 - ProjectedCaptureVector.w);
+		BlendFactor = BlendFactor * (1.0 - ProjectedCaptureVector.w);
+	}
+	else
+	{
+		SpecularIBL = Reflection.rgb;
+
+		SpecularIBL *= ResolvedView.SkyLightColor.rgb;
+
+		// If we default to SkyLight it means there are no ReflectionComponents left in the scene, so set BlendFactor to 0
+		BlendFactor = 0.0f;
+	}
 	return SpecularIBL;
 }
 
diff --git a/Engine/Source/Runtime/Renderer/Private/MobileBasePass.cpp b/Engine/Source/Runtime/Renderer/Private/MobileBasePass.cpp
--- a/Engine/Source/Runtime/Renderer/Private/MobileBasePass.cpp
+++ b/Engine/Source/Runtime/Renderer/Private/MobileBasePass.cpp
@@ -477,8 +477,16 @@ void TMobileBasePassPSPolicyParamType<FUniformLightMapPolicy>::GetShaderBindings
 						//To keep ImageBasedReflectionLighting coherence with PC, use AverageBrightness instead of InvAverageBrightness to calculate the IBL contribution
 						ReflectionParams[i] = ReflectionProxy->EncodedHDRAverageBrightness;
 					}
+					else if (Scene->SkyLight != nullptr)
+					{
+						// NegativeInfluence to signal the shader we are defaulting to SkyLight if there are no ReflectionComponents in the Level
+						CapturePositions[i].W = -1.0f;
+						ReflectionCubemapTextures[i] = Scene->SkyLight->ProcessedTexture;
+						break;
+					}
 				}
 			}
+
 			for (int32 i = 0; i < MaxNumReflections; i++)
 			{
 				ShaderBindings.AddTexture(HQReflectionCubemaps[i], HQReflectionSamplers[i], ReflectionCubemapTextures[i]->SamplerStateRHI, ReflectionCubemapTextures[i]->TextureRHI);
