Fix BSDF plot (green tint + missing backlit term).

diff --git a/Engine/Shaders/Private/HairStrands/HairStrandsBsdfPlot.usf b/Engine/Shaders/Private/HairStrands/HairStrandsBsdfPlot.usf
--- a/Engine/Shaders/Private/HairStrands/HairStrandsBsdfPlot.usf
+++ b/Engine/Shaders/Private/HairStrands/HairStrandsBsdfPlot.usf
@@ -62,7 +62,9 @@ void MainPS(
 		discard;
 
 	const float2 UV = ((PixelCoord-OutputOffset) + float2(0.5f,0.5f)) / float2(OutputResolution);
-	const float2 InputUV = saturate(float2(InputCoord - OutputOffset) / float2(OutputResolution));
+	float2 InputUV = saturate(float2(InputCoord - OutputOffset) / float2(OutputResolution));
+	// InputUV = float2(0.5, 0.5);
+
 
 	float V_Phi   = (UV.x*2-1) * PI; //[-Pi..Pi]
 	float V_Theta = (UV.y*2-1) * PI * 0.5f;
@@ -88,6 +90,8 @@ void MainPS(
 	TransmittanceData.ScatteringComponent = HairComponents;
 	TransmittanceData.OpaqueVisibility = 1;
 
+	const float Backlit = 1;
+	const float Specular = 0.5f;
 	FGBufferData HairGBuffer = (FGBufferData)0;
 	HairGBuffer.WorldNormal = T;
 	HairGBuffer.BaseColor = BaseColor.xxx;
@@ -95,11 +99,11 @@ void MainPS(
 	HairGBuffer.ShadingModelID = SHADINGMODELID_HAIR;
 	HairGBuffer.DiffuseColor = 0;
 	HairGBuffer.SpecularColor = 0;
-	HairGBuffer.Specular = 0.5f;
+	HairGBuffer.Specular = Specular;
 	HairGBuffer.Metallic = 0;
 	HairGBuffer.Depth = 0.5f;
 	HairGBuffer.GBufferAO = 1;
-	HairGBuffer.CustomData = 0;
+	HairGBuffer.CustomData = float4(0,0,Backlit,0);
 	HairGBuffer.IndirectIrradiance = 1000000;
 
 	OutColor = float4(0,0,0,1);
@@ -116,7 +120,7 @@ void MainPS(
 	OutColor.xyz += Fs_Lit_TRT * float3(0, 1, 1);
 	OutColor.xyz += Fs_Lit_TT  * float3(0, 0, 1);
 #endif
-	OutColor.xyz += Fs_Hair    * float3(0, 1, 0);
+	OutColor.xyz += Fs_Hair;
 	OutColor.xyz *= Exposure;
 
 	//const float3 ColorV = (V_World+1)* 0.5f;
