UE4.25 - Fixed a bunch of shader warnings
#rnx

diff --git a/Engine/Shaders/Private/LocalVertexFactory.ush b/Engine/Shaders/Private/LocalVertexFactory.ush
--- a/Engine/Shaders/Private/LocalVertexFactory.ush
+++ b/Engine/Shaders/Private/LocalVertexFactory.ush
@@ -409,7 +409,7 @@ float GetInstanceRandom(FVertexFactoryIntermediates Intermediates)
 
 float3 GetInstanceOrigin(FVertexFactoryIntermediates Intermediates)
 {
-	return Intermediates.InstanceOrigin.xyz + InstancingOffset;
+	return Intermediates.InstanceOrigin.xyz + InstancingOffset.xyz;
 }
 
 #endif	// USE_INSTANCING
diff --git a/Engine/Shaders/Private/MeshParticleVertexFactory.ush b/Engine/Shaders/Private/MeshParticleVertexFactory.ush
--- a/Engine/Shaders/Private/MeshParticleVertexFactory.ush
+++ b/Engine/Shaders/Private/MeshParticleVertexFactory.ush
@@ -459,7 +459,7 @@ float4 VertexFactoryGetWorldPosition(FVertexFactoryInput Input, FVertexFactoryIn
 
 float3 VertexFactoryGetWorldNormal(FVertexFactoryInput Input, FVertexFactoryIntermediates Intermediates)
 {
-	return TangentBias(Input.TangentZ);
+	return TangentBias(Input.TangentZ.xyz);
 }
 
 float4 VertexFactoryGetRasterizedWorldPosition(FVertexFactoryInput Input, FVertexFactoryIntermediates Intermediates, float4 InWorldPosition)
diff --git a/Engine/Shaders/Private/PostProcessMaterialShaders.usf b/Engine/Shaders/Private/PostProcessMaterialShaders.usf
--- a/Engine/Shaders/Private/PostProcessMaterialShaders.usf
+++ b/Engine/Shaders/Private/PostProcessMaterialShaders.usf
@@ -133,15 +133,19 @@ void MainVS(
 	GetMaterialCustomizedUVs(VertexParameters, CustomizedUVs);
 	GetCustomInterpolators(VertexParameters, CustomizedUVs);
 	
-	UNROLL
-	for (int CoordinateIndex = 0; CoordinateIndex < NUM_MATERIAL_TEXCOORDS; CoordinateIndex++)
 	{
-		SetUV(Output, CoordinateIndex, InPosition.xy);
+		UNROLL
+		for (int CoordinateIndex = 0; CoordinateIndex < NUM_MATERIAL_TEXCOORDS; CoordinateIndex++)
+		{
+			SetUV(Output, CoordinateIndex, InPosition.xy);
+		}
 	}
-	UNROLL
-	for (int CoordinateIndex = NUM_MATERIAL_TEXCOORDS; CoordinateIndex < NUM_TEX_COORD_INTERPOLATORS; CoordinateIndex++)
 	{
-		SetUV(Output, CoordinateIndex, CustomizedUVs[CoordinateIndex]);
+		UNROLL
+		for (int CoordinateIndex = NUM_MATERIAL_TEXCOORDS; CoordinateIndex < NUM_TEX_COORD_INTERPOLATORS; CoordinateIndex++)
+		{
+			SetUV(Output, CoordinateIndex, CustomizedUVs[CoordinateIndex]);
+		}
 	}
 #endif
 }
@@ -172,15 +176,19 @@ void MainPS(
 	float2 ViewportUV = (SvPosition.xy - PostProcessOutput_ViewportMin.xy) * PostProcessOutput_ViewportSizeInverse.xy;
 
 #if NUM_TEX_COORD_INTERPOLATORS
-	UNROLL
-	for (int CoordinateIndex = 0; CoordinateIndex < NUM_MATERIAL_TEXCOORDS; CoordinateIndex++)
 	{
-		Parameters.TexCoords[CoordinateIndex] = ViewportUV;
+		UNROLL
+		for (int CoordinateIndex = 0; CoordinateIndex < NUM_MATERIAL_TEXCOORDS; CoordinateIndex++)
+		{
+			Parameters.TexCoords[CoordinateIndex] = ViewportUV;
+		}
 	}
-	UNROLL
-	for (int CoordinateIndex = NUM_MATERIAL_TEXCOORDS; CoordinateIndex < NUM_TEX_COORD_INTERPOLATORS; CoordinateIndex++)
 	{
-		Parameters.TexCoords[CoordinateIndex] = GetUV(Input, CoordinateIndex);
+		UNROLL
+		for (int CoordinateIndex = NUM_MATERIAL_TEXCOORDS; CoordinateIndex < NUM_TEX_COORD_INTERPOLATORS; CoordinateIndex++)
+		{
+			Parameters.TexCoords[CoordinateIndex] = GetUV(Input, CoordinateIndex);
+		}
 	}
 #endif
 
diff --git a/Engine/Shaders/Private/RayTracing/RayTracingMaterialHitShaders.usf b/Engine/Shaders/Private/RayTracing/RayTracingMaterialHitShaders.usf
--- a/Engine/Shaders/Private/RayTracing/RayTracingMaterialHitShaders.usf
+++ b/Engine/Shaders/Private/RayTracing/RayTracingMaterialHitShaders.usf
@@ -138,7 +138,7 @@ void GetPrecomputedIndirectLightingAndSkyLight(
 	out float3 OutSubsurfaceLighting,
 	out float OutIndirectIrradiance)
 {
-	OutIndirectIrradiance = float3(0,0,0);
+	OutIndirectIrradiance = 0;
 	OutDiffuseLighting = float3(0,0,0);
 	OutSubsurfaceLighting = 0;
 	float2 SkyOcclusionUV = 0;
diff --git a/Engine/Shaders/Private/ShadowDepthVertexShader.usf b/Engine/Shaders/Private/ShadowDepthVertexShader.usf
--- a/Engine/Shaders/Private/ShadowDepthVertexShader.usf
+++ b/Engine/Shaders/Private/ShadowDepthVertexShader.usf
@@ -458,7 +458,7 @@ void PositionOnlyMain(
 #endif // #if ONEPASS_POINTLIGHT_SHADOW
 
 #if INTERPOLATE_POSITION
-	OutParameters.PixelPosition = WorldPos;
+	OutParameters.PixelPosition = WorldPos.xyz;
 #endif
 }
 
diff --git a/Engine/Shaders/Private/SlateVertexShader.usf b/Engine/Shaders/Private/SlateVertexShader.usf
--- a/Engine/Shaders/Private/SlateVertexShader.usf
+++ b/Engine/Shaders/Private/SlateVertexShader.usf
@@ -63,10 +63,12 @@ VertexToPixelInterpolants Main(
 		float2(0,0)  //uv channel 7
 	};
 
-	UNROLL
-	for( int CoordinateIndex = 0; CoordinateIndex < NUM_MATERIAL_TEXCOORDS_VERTEX; CoordinateIndex++ )
 	{
-		VertexParameters.TexCoords[CoordinateIndex] = UVArrays[CoordinateIndex];
+		UNROLL
+		for( int CoordinateIndex = 0; CoordinateIndex < NUM_MATERIAL_TEXCOORDS_VERTEX; CoordinateIndex++ )
+		{
+			VertexParameters.TexCoords[CoordinateIndex] = UVArrays[CoordinateIndex];
+		}
 	}
 
 	#if HAS_MATERIAL_TEXCOORDS
diff --git a/Engine/Shaders/Private/VolumetricFogLightFunction.usf b/Engine/Shaders/Private/VolumetricFogLightFunction.usf
--- a/Engine/Shaders/Private/VolumetricFogLightFunction.usf
+++ b/Engine/Shaders/Private/VolumetricFogLightFunction.usf
@@ -42,7 +42,7 @@ void Main(
 	float3 AbsoluteWorldPosition = SvPositionToWorld(SvPosition);
 	*/
 
-	float2 ScreenPosition = (SvPosition * LightFunctionTexelSize - .5f) * float2(2, -2);
+	float2 ScreenPosition = (SvPosition.xy * LightFunctionTexelSize - float2(0.5f, 0.5f)) * float2(2, -2);
 	float3 AbsoluteWorldPosition = GetWorldPositionForLightFunctionTexturePosition(ScreenPosition);
 
 	float4 LightVector = mul(float4(AbsoluteWorldPosition, 1), LightFunctionWorldToLight);
