Fix root/tip hair scaling when hair are simulated & attached to geometry.

diff --git a/Engine/Shaders/Private/HairStrands/HairStrandsInterpolation.usf b/Engine/Shaders/Private/HairStrands/HairStrandsInterpolation.usf
--- a/Engine/Shaders/Private/HairStrands/HairStrandsInterpolation.usf
+++ b/Engine/Shaders/Private/HairStrands/HairStrandsInterpolation.usf
@@ -274,7 +274,7 @@ void MainCS(uint2 DispatchThreadId : SV_DispatchThreadID)
 			}
 			#endif
 
-			OutRenderDeformedPositionBuffer[VertexIndex] = float4(f32tof16(WorldVertexPosition - OutHairPositionOffset), RenderRestPosePositionBuffer[VertexIndex].w);
+			OutRenderDeformedPositionBuffer[VertexIndex] = float4(f32tof16(WorldVertexPosition - OutHairPositionOffset), PackedRestPosition.w);
 		}
 		#endif	
 
