Merge 11372317 from Dev-VR:

Fix for ARCore/ARKit camera passthrough material

- "MOBILE_STENCIL_COMPARE_FUNCTION" was introduced in the shader code without a sensible default value. Its value is only set by FPostProcessMaterialPS but there're other native code also uses the same shader code (FGoogleARCoreCameraOverlayPS, TARKitCameraOverlayPS)
- Now we set the default value for "MOBILE_STENCIL_COMPARE_FUNCTION" to "MOBILE_STENCIL_COMPARE_NEVER, making it backward compatible

#rnx

diff --git a/Engine/Shaders/Private/PostProcessMaterialShaders.usf b/Engine/Shaders/Private/PostProcessMaterialShaders.usf
--- a/Engine/Shaders/Private/PostProcessMaterialShaders.usf
+++ b/Engine/Shaders/Private/PostProcessMaterialShaders.usf
@@ -304,6 +304,10 @@ void MainVS(
 #define MOBILE_STENCIL_COMPARE_NEVER			6
 #define MOBILE_STENCIL_COMPARE_ALWAYS			7
 
+#ifndef MOBILE_STENCIL_COMPARE_FUNCTION
+#define MOBILE_STENCIL_COMPARE_FUNCTION MOBILE_STENCIL_COMPARE_NEVER
+#endif
+
 void MainPS(
 	in FPostProcessMaterialVSToPS Input,
 	in float2 InViewportUV : TEXCOORD8,
