Ray tracing reflection rays are not traced from pixels that have until shader model.

[REDACTED_HANDLE]

diff --git a/Engine/Shaders/Private/RayTracing/RayTracingReflections.usf b/Engine/Shaders/Private/RayTracing/RayTracingReflections.usf
--- a/Engine/Shaders/Private/RayTracing/RayTracingReflections.usf
+++ b/Engine/Shaders/Private/RayTracing/RayTracingReflections.usf
@@ -330,7 +330,8 @@ RAY_TRACING_ENTRY_RAYGEN(RayTracingReflectionsRGS)
 	// Cutoff
 	// * Roughness value
 	// * Hair material is not for RT reflection at the moment. Their indirect contributuion is based on diffuse SH evaluation in ReflectionEnvironemenPixelShader.usf
-    const bool bIsValidPixel = (RoughnessFade > 0 || ClearCoatRoughnessFade > 0) && GBufferData.ShadingModelID != SHADINGMODELID_HAIR;
+	// * Unlit materials do not participate in lighting computations, results are unused.
+	bool bIsValidPixel = (RoughnessFade > 0 || ClearCoatRoughnessFade > 0) && GBufferData.ShadingModelID != SHADINGMODELID_HAIR && GBufferData.ShadingModelID != SHADINGMODELID_UNLIT;
 
     const uint LocalMaxBounces = DIM_HYBRID ? 1 : MaxBounces;
 
