Fix inverted V component on root UV.

diff --git a/Engine/Shaders/Private/HairStrands/HairStrandsVisibilityCommon.ush b/Engine/Shaders/Private/HairStrands/HairStrandsVisibilityCommon.ush
--- a/Engine/Shaders/Private/HairStrands/HairStrandsVisibilityCommon.ush
+++ b/Engine/Shaders/Private/HairStrands/HairStrandsVisibilityCommon.ush
@@ -315,7 +315,7 @@ FHairAttributes DecodeHairAttribute(uint4 Attributes)
 {
 	float2 RootUV, RootIndex;
 	RootUV.x	= (Attributes.x & 0xFF) / 255.f;
-	RootUV.y	= ((Attributes.x>>8) & 0xFF) / 255.f;
+	RootUV.y	= 1-(((Attributes.x>>8) & 0xFF) / 255.f); // Invert V to compensate image origin flip
 	RootIndex.x	= Attributes.z & 0xFF;
 	RootIndex.y	= (Attributes.z >> 8) & 0xFF;
 
