Fix a bug of pre exposure for planar reflection.

diff --git a/Engine/Shaders/Private/MobileBasePassPixelShader.usf b/Engine/Shaders/Private/MobileBasePassPixelShader.usf
--- a/Engine/Shaders/Private/MobileBasePassPixelShader.usf
+++ b/Engine/Shaders/Private/MobileBasePassPixelShader.usf
@@ -281,10 +281,6 @@ float4 GetPlanarReflection(float3 WorldPosition, float3 WorldNormal, float Rough
 	PlanarReflection.rgb *= PlanarReflection.rgb;
 #endif
 
-	// the capture has applied exposure scale so need to cancel that.
-#if !MATERIALBLENDING_MODULATE
-	PlanarReflection.rgb *= USE_PREEXPOSURE ? ResolvedView.OneOverPreExposure : 1;
-#endif
 	return PlanarReflection;
 }
 
