Adding support for sampling scene depth in single layer water pass through a material expression

diff --git a/Engine/Shaders/Private/BasePassPixelShader.usf b/Engine/Shaders/Private/BasePassPixelShader.usf
--- a/Engine/Shaders/Private/BasePassPixelShader.usf
+++ b/Engine/Shaders/Private/BasePassPixelShader.usf
@@ -16,12 +16,13 @@
 
 // Material setting to allow forward shading (including mobile) to use preintegrated GF lut for simple IBL.
 #if MATERIAL_SHADINGMODEL_SINGLELAYERWATER || (FORWARD_SHADING && MATERIAL_USE_PREINTEGRATED_GF)
-#define PreIntegratedGF			OpaqueBasePass.PreIntegratedGFTexture
-#if SUPPORTS_INDEPENDENT_SAMPLERS
-#define PreIntegratedGFSampler	View.SharedBilinearClampedSampler
-#else
-#define PreIntegratedGFSampler	OpaqueBasePass.PreIntegratedGFSampler
-#endif
+	#define PreIntegratedGF	OpaqueBasePass.PreIntegratedGFTexture
+
+	#if SUPPORTS_INDEPENDENT_SAMPLERS
+		#define PreIntegratedGFSampler	View.SharedBilinearClampedSampler
+	#else
+		#define PreIntegratedGFSampler	OpaqueBasePass.PreIntegratedGFSampler
+	#endif
 #endif
 
 #include "SHCommon.ush"
@@ -96,7 +97,7 @@ float NormalCurvatureToRoughness(float3 WorldNormal)
 
 #endif
 
-#include "SingleLayerWaterCommon.ush"
+#include "SingleLayerWaterShading.ush"
 
 #endif // MATERIAL_SHADINGMODEL_SINGLELAYERWATER
 
@@ -339,15 +340,9 @@ void GetSkyLighting(FMaterialPixelParameters MaterialParameters, VTPageTableResu
 #if SUPPORTS_INDEPENDENT_SAMPLERS
 	#define ILCSharedSampler1 View.SharedBilinearClampedSampler
 	#define ILCSharedSampler2 View.SharedBilinearClampedSampler
-
-	#define SingleLayerWaterSceneColorSampler View.SharedBilinearClampedSampler
-	#define SingleLayerWaterSceneDepthSampler View.SharedPointClampedSampler
 #else
 	#define ILCSharedSampler1 IndirectLightingCache.IndirectLightingCacheTextureSampler1
 	#define ILCSharedSampler2 IndirectLightingCache.IndirectLightingCacheTextureSampler2
-
-	#define SingleLayerWaterSceneColorSampler OpaqueBasePass.SceneColorWithoutSingleLayerWaterSampler
-	#define SingleLayerWaterSceneDepthSampler OpaqueBasePass.SceneDepthWithoutSingleLayerWaterSampler
 #endif
 
 /** Calculates indirect lighting contribution on this object from precomputed data. */
diff --git a/Engine/Shaders/Private/MaterialTemplate.ush b/Engine/Shaders/Private/MaterialTemplate.ush
--- a/Engine/Shaders/Private/MaterialTemplate.ush
+++ b/Engine/Shaders/Private/MaterialTemplate.ush
@@ -77,7 +77,11 @@
 #endif
 
 #if MATERIAL_SKY_ATMOSPHERE && PROJECT_SUPPORT_SKY_ATMOSPHERE
-#include "/Engine/Private/SkyAtmosphereCommon.ush"
+	#include "/Engine/Private/SkyAtmosphereCommon.ush"
+#endif
+
+#if MATERIAL_SHADINGMODEL_SINGLELAYERWATER
+	#include "/Engine/Private/SingleLayerWaterCommon.ush"
 #endif
 
 #include "/Engine/Private/PaniniProjection.ush"
@@ -1237,6 +1241,22 @@ float3 MaterialExpressionSkyAtmosphereDistantLightScatteredLuminance(FMaterialPi
 #endif
 }
 
+/** 
+	Get the scene depth from the scene below the single layer water surface. This is only valid in the single layer water rendering pass.
+	Returns the chosen fallback depth if the material doesn't support reading back the correct depth.
+*/
+float MaterialExpressionSceneDepthWithoutWater(float2 ViewportUV, float FallbackDepth)
+{
+#if MATERIAL_SHADINGMODEL_SINGLELAYERWATER && !SCENE_TEXTURES_DISABLED && (!SIMPLE_SINGLE_LAYER_WATER || SINGLE_LAYER_WATER_SIMPLE_FORWARD)
+
+	const float2 ClampedViewportUV = clamp(ViewportUV, OpaqueBasePass.SceneWithoutSingleLayerWaterMinMaxUV.xy, OpaqueBasePass.SceneWithoutSingleLayerWaterMinMaxUV.zw);
+
+	return OpaqueBasePass.SceneDepthWithoutSingleLayerWaterTexture.SampleLevel(SingleLayerWaterSceneDepthSampler, ClampedViewportUV, 0).x * SINGLE_LAYER_WATER_DEPTH_SCALE;
+#else
+	return FallbackDepth;
+#endif
+}
+
 /**
  * Utility function to unmirror one coordinate value to the other side
  * UnMirrored == 1 if normal
diff --git a/Engine/Shaders/Private/SingleLayerWaterCommon.ush b/Engine/Shaders/Private/SingleLayerWaterCommon.ush
--- a/Engine/Shaders/Private/SingleLayerWaterCommon.ush
+++ b/Engine/Shaders/Private/SingleLayerWaterCommon.ush
@@ -1,27 +1,23 @@
-/*=============================================================================
-	SingleLayerWaterCommon.ush
-=============================================================================*/
+// Copyright Epic Games, Inc. All Rights Reserved.
 
-#include "Common.ush"
-#include "ParticipatingMediaCommon.ush" 
-#include "ShadingCommon.ush" 
-#include "DistortionCommon.ush"
+#pragma once
 
 // Scale for encoding depth into SceneDepthWithoutSingleLayerWaterTexture
 #define SINGLE_LAYER_WATER_DEPTH_SCALE 100.0f
 
 #ifndef SIMPLE_SINGLE_LAYER_WATER
-#define SIMPLE_SINGLE_LAYER_WATER 0
+	#define SIMPLE_SINGLE_LAYER_WATER 0
 #endif
 
-struct WaterVolumeLightingOutput
-{
-	float3 Luminance;
-	float3 WaterToSceneTransmittance;
-	float3 WaterToSceneToLightTransmittance;
-};
+#if SUPPORTS_INDEPENDENT_SAMPLERS
+	#define SingleLayerWaterSceneColorSampler View.SharedBilinearClampedSampler
+	#define SingleLayerWaterSceneDepthSampler View.SharedPointClampedSampler
+#else
+	#define SingleLayerWaterSceneColorSampler OpaqueBasePass.SceneColorWithoutSingleLayerWaterSampler
+	#define SingleLayerWaterSceneDepthSampler OpaqueBasePass.SceneDepthWithoutSingleLayerWaterSampler
+#endif
 
-// We use our own refract because the HLSL one was not working on all platofrm (ES3.1 preview for instance)
+// We use our own refract because the HLSL one was not working on all platforms (ES3.1 preview for instance)
 bool WaterRefract(float3 V, float3 N, float Eta, out float3 R)
 {
 	const float CosThetaI = dot(V, N);
@@ -36,193 +32,4 @@ bool WaterRefract(float3 V, float3 N, float Eta, out float3 R)
 	const float CosThetaT = -(CosThetaI >= 0.0f ? 1.0f : -1.0f) * sqrt(CosThetaT2);
 	R = (Eta * CosThetaI + CosThetaT) * N - Eta * V;
 	return true;
-}
-
-
-// Returns as luminance, the result of lighting scattering in water and the under water scene color contribution.
-WaterVolumeLightingOutput EvaluateWaterVolumeLighting(
-	FMaterialPixelParameters MaterialParameters, FPixelMaterialInputs PixelMaterialInputs, ViewState ResolvedView,
-	float DirectionalLightShadow,
-#if SINGLE_LAYER_WATER_SIMPLE_FORWARD
-	Texture2D SceneDepthWithoutSingleLayerWaterTexture, SamplerState SceneDepthWithoutSingleLayerWaterSampler,
-#endif
-#if !SIMPLE_SINGLE_LAYER_WATER
-	Texture2D SceneColorWithoutSingleLayerWaterTexture, SamplerState SceneColorWithoutSingleLayerWaterSampler,
-	Texture2D SceneDepthWithoutSingleLayerWaterTexture, SamplerState SceneDepthWithoutSingleLayerWaterSampler,
-	float2 SceneWithoutSingleLayerWaterMinUV,
-	float2 SceneWithoutSingleLayerWaterMaxUV,
-#endif // SIMPLE_SINGLE_LAYER_WATER
-	float Specular, const float4 DistortionParams,
-	float3 SunIlluminance, float3 AmbiantIlluminance, float3 EnvBrdf,
-	bool CameraIsUnderWater, float WaterVisibility, uint EyeIndex)
-{
-	WaterVolumeLightingOutput Output;
-	Output.Luminance = 0.0f;
-	Output.WaterToSceneTransmittance = 1.0f;
-	Output.WaterToSceneToLightTransmittance = 1.0f;
-
-	float2 WaterNDC = MaterialParameters.ScreenPosition.xy / MaterialParameters.ScreenPosition.w;
-	float2 ViewportUV = WaterNDC * ResolvedView.ScreenPositionScaleBias.xy + ResolvedView.ScreenPositionScaleBias.wz;
-
-	if (WaterVisibility > 0.0f)
-	{
-		float WaterDepth = ConvertFromDeviceZ(MaterialParameters.SvPosition.z);
-#if SIMPLE_SINGLE_LAYER_WATER
-
-#if SINGLE_LAYER_WATER_SIMPLE_FORWARD 
-		float PixelSceneDepth = SceneDepthWithoutSingleLayerWaterTexture.SampleLevel(SceneDepthWithoutSingleLayerWaterSampler, ViewportUV, 0).x * SINGLE_LAYER_WATER_DEPTH_SCALE;
-		float PixelSceneDeviceZ = ConvertToDeviceZ(PixelSceneDepth);
-#else
-		float PixelSceneDeviceZ = LookupDeviceZ(ViewportUV);
-		float PixelSceneDepth = ConvertFromDeviceZ(PixelSceneDeviceZ);
-#endif
-
-#if FEATURE_LEVEL <= FEATURE_LEVEL_ES3_1
-		// Only mobile run the SIMPLE_SINGLE_LAYER_WATER path today.
-		// On mobile, depth is stored linear in alpha channel of scene color as FP16 so it is easy to reach it max value. It would then result in fully transparent water.
-		// To counter that, we assume any visible water pixel has a constant under water depth and we progressively lerp towards this simplified representation.
-		// Lighting remains correct thanks to the analytical integration of scattering.
-		const float FadeLenght = 10000.0f;								// 0.1 kilometers
-		const float FadeStartDistance = 65504.0f - FadeLenght;
-		const float FadeBehindWaterDepth = WaterDepth + FadeLenght;
-		if (WaterDepth > FadeStartDistance)
-		{
-			PixelSceneDepth = lerp(PixelSceneDepth, FadeBehindWaterDepth, saturate((WaterDepth - FadeStartDistance) / FadeLenght));
-			PixelSceneDeviceZ = ConvertToDeviceZ(PixelSceneDepth);
-		}
-#endif
-
-		float SceneDeviceZ = PixelSceneDeviceZ;
-		float SceneDepth = PixelSceneDepth;
-
-#else  // !SIMPLE_SINGLE_LAYER_WATER
-
-		ViewportUV = clamp(ViewportUV, SceneWithoutSingleLayerWaterMinUV, SceneWithoutSingleLayerWaterMaxUV);
-
-		float PixelSceneDepth = SceneDepthWithoutSingleLayerWaterTexture.SampleLevel(SceneDepthWithoutSingleLayerWaterSampler, ViewportUV, 0).x * SINGLE_LAYER_WATER_DEPTH_SCALE;
-		
-		float2 Refraction = GetMaterialRefraction(PixelMaterialInputs); // X = IOR, Y = Refraction Bias
-		float2 BufferUVDistortion = ComputeBufferUVDistortion(
-			MaterialParameters, PixelMaterialInputs, ResolvedView,
-			MaterialParameters.WorldNormal, DistortionParams, ViewportUV, Refraction, false, EyeIndex);
-		PostProcessUVDistortion(MaterialParameters, PixelMaterialInputs, PixelSceneDepth, BufferUVDistortion, Refraction);
-		// Also reduce refraction according to water thickness
-		BufferUVDistortion *= saturate((PixelSceneDepth - WaterDepth) * 1.0 / 30.0f);// *(1 - saturate(WaterDepth * 1.0 / 2000.0f));
-
-		const float2 DistortedUV = clamp(ViewportUV + BufferUVDistortion, SceneWithoutSingleLayerWaterMinUV, SceneWithoutSingleLayerWaterMaxUV);
-
-		float SceneDepth = SceneDepthWithoutSingleLayerWaterTexture.SampleLevel(SceneDepthWithoutSingleLayerWaterSampler, DistortedUV, 0).x * SINGLE_LAYER_WATER_DEPTH_SCALE;
-
-		if (SceneDepth < WaterDepth)
-		{
-			SceneDepth = PixelSceneDepth; // The pixel we try to sample is closer than the water pixel: falback to no distortion
-		}
-		else
-		{
-			ViewportUV = DistortedUV;
-		}
-
-		const float SceneDeviceZ = ConvertToDeviceZ(SceneDepth);
-
-#endif // !SIMPLE_SINGLE_LAYER_WATER
-
-		const float BehindWaterDeltaDepth = CameraIsUnderWater ? WaterDepth : max(0.0f, SceneDepth - WaterDepth);
-
-		const float3 ScatteringCoeff = max(0.0f, GetSingleLayerWaterMaterialOutput0(MaterialParameters));
-		const float3 AbsorptionCoeff = max(0.0f, GetSingleLayerWaterMaterialOutput1(MaterialParameters));
-		const float PhaseG = clamp(GetSingleLayerWaterMaterialOutput2(MaterialParameters), -1.0f, 1.0f);
-		const float3 ExtinctionCoeff = ScatteringCoeff + AbsorptionCoeff; 
-		// Max to avoid division by 0 with the analytical integral below.
-		// 1e-5 is high enough to avoid denorms on mobile
-		const float3 ExtinctionCoeffSafe = max(ScatteringCoeff + AbsorptionCoeff, 1e-5); 
-		
-		float DirLightPhaseValue = 0.0f; // Default when Total Internal Reflection happens.
-		{
-#if SIMPLE_SINGLE_LAYER_WATER
-			DirLightPhaseValue = IsotropicPhase();
-#else
-			float IorFrom = 1.0f; // assumes we come from air
-			float IorTo   = DielectricF0ToIor(DielectricSpecularToF0(Specular)); // Wrong if metal is set to >1. But we still keep refraction on the water surface nonetheless.
-			const float relativeIOR = IorFrom / IorTo;
-			float3 UnderWaterRayDir = 0.0f;
-			if (WaterRefract(MaterialParameters.CameraVector, MaterialParameters.WorldNormal, relativeIOR, UnderWaterRayDir))
-			{
-				DirLightPhaseValue = SchlickPhase(PhaseG, dot(-ResolvedView.DirectionalLightDirection.xyz, UnderWaterRayDir));
-			}
-#endif
-		}
-
-		// We also apply transmittance from light to under water surface. However, the scene has been lit by many sources already.
-		// So the transmittance toabove surface is simply approximated using the travel distance from the scene pixel to the water top, assuming a flat water surface.
-		// We cannot combine this transmittance with the transmittance from view because this would change the behavior of the analytical integration of light scattering integration.
-		const float3 BehindWaterSceneWorldPos = SvPositionToWorld(float4(MaterialParameters.SvPosition.xy, SceneDeviceZ, 1.0));
-		const float DistanceFromScenePixelToWaterTop = max(0.0, MaterialParameters.AbsoluteWorldPosition.z - BehindWaterSceneWorldPos.z);
-		const float3 MeanTransmittanceToLightSources = exp(-DistanceFromScenePixelToWaterTop * ExtinctionCoeff);
-
-#if SIMPLE_SINGLE_LAYER_WATER
-		const float3 BehindWaterSceneLuminance = 0.0f; // Cannot read back the scene color in this case
-#else
-		// We use the pixel SvPosition instead of the scene one pre refraction/distortion to avoid those extra ALUs.
-		float3 BehindWaterSceneLuminance = SceneColorWithoutSingleLayerWaterTexture.SampleLevel(SceneColorWithoutSingleLayerWaterSampler, ViewportUV, 0).rgb;
-		BehindWaterSceneLuminance = MeanTransmittanceToLightSources * (USE_PREEXPOSURE ? ResolvedView.OneOverPreExposure : 1.0f) * BehindWaterSceneLuminance;
-#endif
-
-		float3 SunScattLuminance = DirLightPhaseValue * SunIlluminance;
-		float3 AmbScattLuminance = IsotropicPhase()   * AmbiantIlluminance;
-
-#define VOLUMETRICSHADOW	0
-#if !VOLUMETRICSHADOW || SIMPLE_SINGLE_LAYER_WATER
-
-		const float3 OpticalDepth = ExtinctionCoeff * BehindWaterDeltaDepth;
-		float3 Transmittance = exp(-OpticalDepth);
-		float3 ScatteredLuminance = ScatteringCoeff * (AmbScattLuminance + SunScattLuminance * DirectionalLightShadow);
-		ScatteredLuminance = (ScatteredLuminance - ScatteredLuminance * Transmittance) / ExtinctionCoeffSafe;
-
-#else
-		// TODO Make the volumetric shadow part work again
-		float3 Transmittance = 1.0f;
-		float3 ScatteredLuminance = 0.0f;
-		const float RayMarchMaxDistance = min(BehindWaterDeltaDepth, 200.0f);  // 20 meters
-		const float RayMarchStepSize = RayMarchMaxDistance / 10.0f; // Less samples wil lresult in a bit brighter look due to TransmittanceToLightThroughWater being 1 on a longer first sample. Would need it part of analiytical integration
-		const float ShadowDither = RayMarchStepSize * GBufferDither;
-		for (float s = 0.0f; s < RayMarchMaxDistance; s += RayMarchStepSize)
-		{
-			// Only jitter shadow map sampling to not lose energy on first sample
-			float Shadow = ComputeDirectionalLightDynamicShadowing(MaterialParameters.AbsoluteWorldPosition - (s + ShadowDither)*MaterialParameters.CameraVector, GBuffer.Depth);
-
-			float3 WP = MaterialParameters.AbsoluteWorldPosition - s * MaterialParameters.CameraVector;
-			float WaterHeightAboveSample = max(0.0, MaterialParameters.AbsoluteWorldPosition.z - WP.z);
-			float3 TransmittanceToLightThroughWater = 1.0;												// no self shadow, same energy as above analytical solution
-			//float3 TransmittanceToLightThroughWater = exp(-ExtinctionCoeff * WaterHeightAboveSample);	// self shadow as transmittance to water level, close to reference, depends a bit on sample count due to first sample being critical for dense medium
-
-			float3 SampleTransmittance = exp(-ExtinctionCoeff * RayMarchStepSize); // Constant
-			float3 SS = (ScatteringCoeff * TransmittanceToLightThroughWater * (SunScattLuminance * Shadow + AmbScattLuminance));
-			ScatteredLuminance += Transmittance * (SS - SS * SampleTransmittance) / ExtinctionCoeffSafe;
-			Transmittance *= SampleTransmittance;
-		}
-
-		// The rest of the medium
-		const float3 OpticalDepth2 = ExtinctionCoeff * max(0.0, BehindWaterDeltaDepth - RayMarchMaxDistance);
-		if (any(OpticalDepth2 > 0.0f))
-		{
-			float3 Transmittance2 = exp(-OpticalDepth2);
-			float3 ScatteredLuminance2 = ScatteringCoeff * (SunScattLuminance + AmbScattLuminance);
-			ScatteredLuminance += Transmittance * (ScatteredLuminance2 - ScatteredLuminance2 * Transmittance2) / ExtinctionCoeffSafe;
-			Transmittance *= Transmittance2;
-		}
-#endif
-
-		// Apply Fresnel effect to out-scattering towards the view
-		ScatteredLuminance *= CameraIsUnderWater ? 1.0 : (1.0 - EnvBrdf);	// Under water is less visible due to Fresnel effect
-		Transmittance *= CameraIsUnderWater ? (1.0 - EnvBrdf) : 1.0;		// Above	"		"		"		"		"
-
-		// Add single in-scattering apply colored transmittance to scene color
-		Output.Luminance = WaterVisibility * (ScatteredLuminance + Transmittance * BehindWaterSceneLuminance);
-		Output.WaterToSceneTransmittance = Transmittance;
-		Output.WaterToSceneToLightTransmittance = Transmittance * MeanTransmittanceToLightSources;
-	}
-
-	return Output;
-}
-
-
+}
\ No newline at end of file
diff --git a/Engine/Shaders/Private/SingleLayerWaterShading.ush b/Engine/Shaders/Private/SingleLayerWaterShading.ush
new file mode 100644
--- /dev/null
+++ b/Engine/Shaders/Private/SingleLayerWaterShading.ush
@@ -0,0 +1,202 @@
+// Copyright Epic Games, Inc. All Rights Reserved.
+
+#pragma once
+
+#include "Common.ush"
+#include "ParticipatingMediaCommon.ush" 
+#include "ShadingCommon.ush" 
+#include "DistortionCommon.ush"
+#include "SingleLayerWaterCommon.ush"
+
+struct WaterVolumeLightingOutput
+{
+	float3 Luminance;
+	float3 WaterToSceneTransmittance;
+	float3 WaterToSceneToLightTransmittance;
+};
+
+// Returns as luminance, the result of lighting scattering in water and the under water scene color contribution.
+WaterVolumeLightingOutput EvaluateWaterVolumeLighting(
+	FMaterialPixelParameters MaterialParameters, FPixelMaterialInputs PixelMaterialInputs, ViewState ResolvedView,
+	float DirectionalLightShadow,
+#if SINGLE_LAYER_WATER_SIMPLE_FORWARD
+	Texture2D SceneDepthWithoutSingleLayerWaterTexture, SamplerState SceneDepthWithoutSingleLayerWaterSampler,
+#endif
+#if !SIMPLE_SINGLE_LAYER_WATER
+	Texture2D SceneColorWithoutSingleLayerWaterTexture, SamplerState SceneColorWithoutSingleLayerWaterSampler,
+	Texture2D SceneDepthWithoutSingleLayerWaterTexture, SamplerState SceneDepthWithoutSingleLayerWaterSampler,
+	float2 SceneWithoutSingleLayerWaterMinUV,
+	float2 SceneWithoutSingleLayerWaterMaxUV,
+#endif // SIMPLE_SINGLE_LAYER_WATER
+	float Specular, const float4 DistortionParams,
+	float3 SunIlluminance, float3 AmbiantIlluminance, float3 EnvBrdf,
+	bool CameraIsUnderWater, float WaterVisibility, uint EyeIndex)
+{
+	WaterVolumeLightingOutput Output;
+	Output.Luminance = 0.0f;
+	Output.WaterToSceneTransmittance = 1.0f;
+	Output.WaterToSceneToLightTransmittance = 1.0f;
+
+	float2 WaterNDC = MaterialParameters.ScreenPosition.xy / MaterialParameters.ScreenPosition.w;
+	float2 ViewportUV = WaterNDC * ResolvedView.ScreenPositionScaleBias.xy + ResolvedView.ScreenPositionScaleBias.wz;
+
+	if (WaterVisibility > 0.0f)
+	{
+		float WaterDepth = ConvertFromDeviceZ(MaterialParameters.SvPosition.z);
+#if SIMPLE_SINGLE_LAYER_WATER
+
+#if SINGLE_LAYER_WATER_SIMPLE_FORWARD 
+		float PixelSceneDepth = SceneDepthWithoutSingleLayerWaterTexture.SampleLevel(SceneDepthWithoutSingleLayerWaterSampler, ViewportUV, 0).x * SINGLE_LAYER_WATER_DEPTH_SCALE;
+		float PixelSceneDeviceZ = ConvertToDeviceZ(PixelSceneDepth);
+#else
+		float PixelSceneDeviceZ = LookupDeviceZ(ViewportUV);
+		float PixelSceneDepth = ConvertFromDeviceZ(PixelSceneDeviceZ);
+#endif
+
+#if FEATURE_LEVEL <= FEATURE_LEVEL_ES3_1
+		// Only mobile run the SIMPLE_SINGLE_LAYER_WATER path today.
+		// On mobile, depth is stored linear in alpha channel of scene color as FP16 so it is easy to reach it max value. It would then result in fully transparent water.
+		// To counter that, we assume any visible water pixel has a constant under water depth and we progressively lerp towards this simplified representation.
+		// Lighting remains correct thanks to the analytical integration of scattering.
+		const float FadeLenght = 10000.0f;								// 0.1 kilometers
+		const float FadeStartDistance = 65504.0f - FadeLenght;
+		const float FadeBehindWaterDepth = WaterDepth + FadeLenght;
+		if (WaterDepth > FadeStartDistance)
+		{
+			PixelSceneDepth = lerp(PixelSceneDepth, FadeBehindWaterDepth, saturate((WaterDepth - FadeStartDistance) / FadeLenght));
+			PixelSceneDeviceZ = ConvertToDeviceZ(PixelSceneDepth);
+		}
+#endif
+
+		float SceneDeviceZ = PixelSceneDeviceZ;
+		float SceneDepth = PixelSceneDepth;
+
+#else  // !SIMPLE_SINGLE_LAYER_WATER
+
+		ViewportUV = clamp(ViewportUV, SceneWithoutSingleLayerWaterMinUV, SceneWithoutSingleLayerWaterMaxUV);
+
+		float PixelSceneDepth = SceneDepthWithoutSingleLayerWaterTexture.SampleLevel(SceneDepthWithoutSingleLayerWaterSampler, ViewportUV, 0).x * SINGLE_LAYER_WATER_DEPTH_SCALE;
+		
+		float2 Refraction = GetMaterialRefraction(PixelMaterialInputs); // X = IOR, Y = Refraction Bias
+		float2 BufferUVDistortion = ComputeBufferUVDistortion(
+			MaterialParameters, PixelMaterialInputs, ResolvedView,
+			MaterialParameters.WorldNormal, DistortionParams, ViewportUV, Refraction, false, EyeIndex);
+		PostProcessUVDistortion(MaterialParameters, PixelMaterialInputs, PixelSceneDepth, BufferUVDistortion, Refraction);
+		// Also reduce refraction according to water thickness
+		BufferUVDistortion *= saturate((PixelSceneDepth - WaterDepth) * 1.0 / 30.0f);// *(1 - saturate(WaterDepth * 1.0 / 2000.0f));
+
+		const float2 DistortedUV = clamp(ViewportUV + BufferUVDistortion, SceneWithoutSingleLayerWaterMinUV, SceneWithoutSingleLayerWaterMaxUV);
+
+		float SceneDepth = SceneDepthWithoutSingleLayerWaterTexture.SampleLevel(SceneDepthWithoutSingleLayerWaterSampler, DistortedUV, 0).x * SINGLE_LAYER_WATER_DEPTH_SCALE;
+
+		if (SceneDepth < WaterDepth)
+		{
+			SceneDepth = PixelSceneDepth; // The pixel we try to sample is closer than the water pixel: falback to no distortion
+		}
+		else
+		{
+			ViewportUV = DistortedUV;
+		}
+
+		const float SceneDeviceZ = ConvertToDeviceZ(SceneDepth);
+
+#endif // !SIMPLE_SINGLE_LAYER_WATER
+
+		const float BehindWaterDeltaDepth = CameraIsUnderWater ? WaterDepth : max(0.0f, SceneDepth - WaterDepth);
+
+		const float3 ScatteringCoeff = max(0.0f, GetSingleLayerWaterMaterialOutput0(MaterialParameters));
+		const float3 AbsorptionCoeff = max(0.0f, GetSingleLayerWaterMaterialOutput1(MaterialParameters));
+		const float PhaseG = clamp(GetSingleLayerWaterMaterialOutput2(MaterialParameters), -1.0f, 1.0f);
+		const float3 ExtinctionCoeff = ScatteringCoeff + AbsorptionCoeff; 
+		// Max to avoid division by 0 with the analytical integral below.
+		// 1e-5 is high enough to avoid denorms on mobile
+		const float3 ExtinctionCoeffSafe = max(ScatteringCoeff + AbsorptionCoeff, 1e-5); 
+		
+		float DirLightPhaseValue = 0.0f; // Default when Total Internal Reflection happens.
+		{
+#if SIMPLE_SINGLE_LAYER_WATER
+			DirLightPhaseValue = IsotropicPhase();
+#else
+			float IorFrom = 1.0f; // assumes we come from air
+			float IorTo   = DielectricF0ToIor(DielectricSpecularToF0(Specular)); // Wrong if metal is set to >1. But we still keep refraction on the water surface nonetheless.
+			const float relativeIOR = IorFrom / IorTo;
+			float3 UnderWaterRayDir = 0.0f;
+			if (WaterRefract(MaterialParameters.CameraVector, MaterialParameters.WorldNormal, relativeIOR, UnderWaterRayDir))
+			{
+				DirLightPhaseValue = SchlickPhase(PhaseG, dot(-ResolvedView.DirectionalLightDirection.xyz, UnderWaterRayDir));
+			}
+#endif
+		}
+
+		// We also apply transmittance from light to under water surface. However, the scene has been lit by many sources already.
+		// So the transmittance toabove surface is simply approximated using the travel distance from the scene pixel to the water top, assuming a flat water surface.
+		// We cannot combine this transmittance with the transmittance from view because this would change the behavior of the analytical integration of light scattering integration.
+		const float3 BehindWaterSceneWorldPos = SvPositionToWorld(float4(MaterialParameters.SvPosition.xy, SceneDeviceZ, 1.0));
+		const float DistanceFromScenePixelToWaterTop = max(0.0, MaterialParameters.AbsoluteWorldPosition.z - BehindWaterSceneWorldPos.z);
+		const float3 MeanTransmittanceToLightSources = exp(-DistanceFromScenePixelToWaterTop * ExtinctionCoeff);
+
+#if SIMPLE_SINGLE_LAYER_WATER
+		const float3 BehindWaterSceneLuminance = 0.0f; // Cannot read back the scene color in this case
+#else
+		// We use the pixel SvPosition instead of the scene one pre refraction/distortion to avoid those extra ALUs.
+		float3 BehindWaterSceneLuminance = SceneColorWithoutSingleLayerWaterTexture.SampleLevel(SceneColorWithoutSingleLayerWaterSampler, ViewportUV, 0).rgb;
+		BehindWaterSceneLuminance = MeanTransmittanceToLightSources * (USE_PREEXPOSURE ? ResolvedView.OneOverPreExposure : 1.0f) * BehindWaterSceneLuminance;
+#endif
+
+		float3 SunScattLuminance = DirLightPhaseValue * SunIlluminance;
+		float3 AmbScattLuminance = IsotropicPhase()   * AmbiantIlluminance;
+
+#define VOLUMETRICSHADOW	0
+#if !VOLUMETRICSHADOW || SIMPLE_SINGLE_LAYER_WATER
+
+		const float3 OpticalDepth = ExtinctionCoeff * BehindWaterDeltaDepth;
+		float3 Transmittance = exp(-OpticalDepth);
+		float3 ScatteredLuminance = ScatteringCoeff * (AmbScattLuminance + SunScattLuminance * DirectionalLightShadow);
+		ScatteredLuminance = (ScatteredLuminance - ScatteredLuminance * Transmittance) / ExtinctionCoeffSafe;
+
+#else
+		// TODO Make the volumetric shadow part work again
+		float3 Transmittance = 1.0f;
+		float3 ScatteredLuminance = 0.0f;
+		const float RayMarchMaxDistance = min(BehindWaterDeltaDepth, 200.0f);  // 20 meters
+		const float RayMarchStepSize = RayMarchMaxDistance / 10.0f; // Less samples wil lresult in a bit brighter look due to TransmittanceToLightThroughWater being 1 on a longer first sample. Would need it part of analiytical integration
+		const float ShadowDither = RayMarchStepSize * GBufferDither;
+		for (float s = 0.0f; s < RayMarchMaxDistance; s += RayMarchStepSize)
+		{
+			// Only jitter shadow map sampling to not lose energy on first sample
+			float Shadow = ComputeDirectionalLightDynamicShadowing(MaterialParameters.AbsoluteWorldPosition - (s + ShadowDither)*MaterialParameters.CameraVector, GBuffer.Depth);
+
+			float3 WP = MaterialParameters.AbsoluteWorldPosition - s * MaterialParameters.CameraVector;
+			float WaterHeightAboveSample = max(0.0, MaterialParameters.AbsoluteWorldPosition.z - WP.z);
+			float3 TransmittanceToLightThroughWater = 1.0;												// no self shadow, same energy as above analytical solution
+			//float3 TransmittanceToLightThroughWater = exp(-ExtinctionCoeff * WaterHeightAboveSample);	// self shadow as transmittance to water level, close to reference, depends a bit on sample count due to first sample being critical for dense medium
+
+			float3 SampleTransmittance = exp(-ExtinctionCoeff * RayMarchStepSize); // Constant
+			float3 SS = (ScatteringCoeff * TransmittanceToLightThroughWater * (SunScattLuminance * Shadow + AmbScattLuminance));
+			ScatteredLuminance += Transmittance * (SS - SS * SampleTransmittance) / ExtinctionCoeffSafe;
+			Transmittance *= SampleTransmittance;
+		}
+
+		// The rest of the medium
+		const float3 OpticalDepth2 = ExtinctionCoeff * max(0.0, BehindWaterDeltaDepth - RayMarchMaxDistance);
+		if (any(OpticalDepth2 > 0.0f))
+		{
+			float3 Transmittance2 = exp(-OpticalDepth2);
+			float3 ScatteredLuminance2 = ScatteringCoeff * (SunScattLuminance + AmbScattLuminance);
+			ScatteredLuminance += Transmittance * (ScatteredLuminance2 - ScatteredLuminance2 * Transmittance2) / ExtinctionCoeffSafe;
+			Transmittance *= Transmittance2;
+		}
+#endif
+
+		// Apply Fresnel effect to out-scattering towards the view
+		ScatteredLuminance *= CameraIsUnderWater ? 1.0 : (1.0 - EnvBrdf);	// Under water is less visible due to Fresnel effect
+		Transmittance *= CameraIsUnderWater ? (1.0 - EnvBrdf) : 1.0;		// Above	"		"		"		"		"
+
+		// Add single in-scattering apply colored transmittance to scene color
+		Output.Luminance = WaterVisibility * (ScatteredLuminance + Transmittance * BehindWaterSceneLuminance);
+		Output.WaterToSceneTransmittance = Transmittance;
+		Output.WaterToSceneToLightTransmittance = Transmittance * MeanTransmittanceToLightSources;
+	}
+
+	return Output;
+}
diff --git a/Engine/Source/Runtime/Engine/Classes/Materials/MaterialExpressionSceneDepthWithoutWater.h b/Engine/Source/Runtime/Engine/Classes/Materials/MaterialExpressionSceneDepthWithoutWater.h
new file mode 100644
--- /dev/null
+++ b/Engine/Source/Runtime/Engine/Classes/Materials/MaterialExpressionSceneDepthWithoutWater.h
@@ -0,0 +1,54 @@
+// Copyright Epic Games, Inc. All Rights Reserved.
+
+#pragma once
+
+#include "CoreMinimal.h"
+#include "UObject/ObjectMacros.h"
+#include "MaterialExpressionIO.h"
+#include "Materials/MaterialExpression.h"
+#include "Materials/MaterialExpressionSceneColor.h"
+#include "MaterialExpressionSceneDepthWithoutWater.generated.h"
+
+UCLASS(collapsecategories, hidecategories = Object)
+class UMaterialExpressionSceneDepthWithoutWater : public UMaterialExpression
+{
+	GENERATED_BODY()
+
+public:
+
+	UMaterialExpressionSceneDepthWithoutWater();
+
+	/**
+	* Coordinates - UV coordinates to apply to the scene depth lookup.
+	* OffsetFraction - An offset to apply to the scene depth lookup in a 2d fraction of the screen.
+	*/
+	UPROPERTY(EditAnywhere, Category = MaterialExpressionSceneDepthWithoutWater)
+	TEnumAsByte<enum EMaterialSceneAttributeInputMode::Type> InputMode;
+
+	/**
+	* Based on the input mode the input will be treated as either:
+	* UV coordinates to apply to the scene depth lookup or
+	* an offset to apply to the scene depth lookup, in a 2d fraction of the screen.
+	*/
+	UPROPERTY(meta = (RequiredInput = "false", ToolTip = "Defaults to 'ConstInput' if not specified"))
+	FExpressionInput Input;
+
+	/** only used if Input is not hooked up */
+	UPROPERTY(EditAnywhere, Category = MaterialExpressionSceneDepthWithoutWater)
+	FVector2D ConstInput;
+
+	/** Depth to fall back to in case the needed texture isn't available on a particular platform or configuration */
+	UPROPERTY(EditAnywhere, Category = MaterialExpressionSceneDepthWithoutWater)
+	float FallbackDepth = 1000000.0f;
+
+	//~ Begin UMaterialExpression Interface
+#if WITH_EDITOR
+	virtual int32 Compile(class FMaterialCompiler* Compiler, int32 OutputIndex) override;
+	virtual void GetCaption(TArray<FString>& OutCaptions) const override;
+	virtual FName GetInputName(int32 InputIndex) const override;
+#endif
+	//~ End UMaterialExpression Interface
+};
+
+
+
diff --git a/Engine/Source/Runtime/Engine/Private/Materials/HLSLMaterialTranslator.cpp b/Engine/Source/Runtime/Engine/Private/Materials/HLSLMaterialTranslator.cpp
--- a/Engine/Source/Runtime/Engine/Private/Materials/HLSLMaterialTranslator.cpp
+++ b/Engine/Source/Runtime/Engine/Private/Materials/HLSLMaterialTranslator.cpp
@@ -6567,6 +6567,49 @@ int32 FHLSLMaterialTranslator::SkyAtmosphereDistantLightScatteredLuminance()
 	return AddCodeChunk(MCT_Float3, TEXT("MaterialExpressionSkyAtmosphereDistantLightScatteredLuminance(Parameters)"));
 }
 
+int32 FHLSLMaterialTranslator::SceneDepthWithoutWater(int32 Offset, int32 ViewportUV, bool bUseOffset, float FallbackDepth)
+{
+	if (ShaderFrequency == SF_Vertex && FeatureLevel <= ERHIFeatureLevel::ES3_1)
+	{
+		// mobile currently does not support this, we need to read a separate copy of the depth, we must disable framebuffer fetch and force scene texture reads.
+		return Errorf(TEXT("Cannot read scene depth from the vertex shader with feature level ES3.1 or below."));
+	}
+
+	if (!Material->GetShadingModels().HasShadingModel(MSM_SingleLayerWater))
+	{
+		return Errorf(TEXT("Can only read scene depth below water when material Shading Model is Single Layer Water."));
+	}
+	
+	if (Material->GetMaterialDomain() != MD_Surface)
+	{
+		return Errorf(TEXT("Can only read scene depth below water when material Domain is set to Surface."));
+	}
+
+	if (IsTranslucentBlendMode(Material->GetBlendMode()))
+	{
+		return Errorf(TEXT("Can only read scene depth below water when material Blend Mode isn't translucent."));
+	}
+
+	if (Offset == INDEX_NONE && bUseOffset)
+	{
+		return INDEX_NONE;
+	}
+
+	AddEstimatedTextureSample();
+
+	const FString UserDepthCode(TEXT("MaterialExpressionSceneDepthWithoutWater(%s, %s)"));
+	const FString FallbackString(FString::SanitizeFloat(FallbackDepth));
+	const int32 TexCoordCode = GetScreenAlignedUV(Offset, ViewportUV, bUseOffset);
+
+	// add the code string
+	return AddCodeChunk(
+		MCT_Float,
+		*UserDepthCode,
+		*GetParameterCode(TexCoordCode),
+		*FallbackString
+	);
+}
+
 int32 FHLSLMaterialTranslator::CustomPrimitiveData(int32 OutputIndex, EMaterialValueType Type)
 {
 	check(OutputIndex < FCustomPrimitiveData::NumCustomPrimitiveDataFloats);
diff --git a/Engine/Source/Runtime/Engine/Private/Materials/HLSLMaterialTranslator.h b/Engine/Source/Runtime/Engine/Private/Materials/HLSLMaterialTranslator.h
--- a/Engine/Source/Runtime/Engine/Private/Materials/HLSLMaterialTranslator.h
+++ b/Engine/Source/Runtime/Engine/Private/Materials/HLSLMaterialTranslator.h
@@ -710,6 +710,9 @@ protected:
 	virtual int32 SkyAtmosphereAerialPerspective(int32 WorldPosition) override;
 	virtual int32 SkyAtmosphereDistantLightScatteredLuminance() override;
 
+	// Water
+	virtual int32 SceneDepthWithoutWater(int32 Offset, int32 ViewportUV, bool bUseOffset, float FallbackDepth) override;
+
 	virtual int32 CustomPrimitiveData(int32 OutputIndex, EMaterialValueType Type) override;
 
 	virtual int32 ShadingModel(EMaterialShadingModel InSelectedShadingModel) override;
diff --git a/Engine/Source/Runtime/Engine/Private/Materials/MaterialExpressions.cpp b/Engine/Source/Runtime/Engine/Private/Materials/MaterialExpressions.cpp
--- a/Engine/Source/Runtime/Engine/Private/Materials/MaterialExpressions.cpp
+++ b/Engine/Source/Runtime/Engine/Private/Materials/MaterialExpressions.cpp
@@ -170,6 +170,7 @@
 #include "Materials/MaterialExpressionSaturate.h"
 #include "Materials/MaterialExpressionSceneColor.h"
 #include "Materials/MaterialExpressionSceneDepth.h"
+#include "Materials/MaterialExpressionSceneDepthWithoutWater.h"
 #include "Materials/MaterialExpressionSceneTexelSize.h"
 #include "Materials/MaterialExpressionSceneTexture.h"
 #include "Materials/MaterialExpressionScreenPosition.h"
@@ -17506,7 +17507,71 @@ FString UMaterialExpressionThinTranslucentMaterialOutput::GetDisplayName() const
 }
 
 
+///////////////////////////////////////////////////////////////////////////////
+// UMaterialExpressionSceneDepthWithoutWater
+///////////////////////////////////////////////////////////////////////////////
+UMaterialExpressionSceneDepthWithoutWater::UMaterialExpressionSceneDepthWithoutWater()
+{
+#if WITH_EDITORONLY_DATA
+	MenuCategories.Add(LOCTEXT("Water", "Water"));
+
+	Outputs.Reset();
+	Outputs.Add(FExpressionOutput(TEXT(""), 1, 1, 0, 0, 0));
+	bShaderInputData = true;
+#endif
 
+	ConstInput = FVector2D(0.f, 0.f);
+}
+
+#if WITH_EDITOR
+int32 UMaterialExpressionSceneDepthWithoutWater::Compile(class FMaterialCompiler* Compiler, int32 OutputIndex)
+{
+	int32 OffsetIndex = INDEX_NONE;
+	int32 CoordinateIndex = INDEX_NONE;
+	bool bUseOffset = false;
+
+	if (InputMode == EMaterialSceneAttributeInputMode::OffsetFraction)
+	{
+		if (Input.GetTracedInput().Expression)
+		{
+			OffsetIndex = Input.Compile(Compiler);
+		}
+		else
+		{
+			OffsetIndex = Compiler->Constant2(ConstInput.X, ConstInput.Y);
+		}
+		bUseOffset = true;
+	}
+	else if (InputMode == EMaterialSceneAttributeInputMode::Coordinates)
+	{
+		if (Input.GetTracedInput().Expression)
+		{
+			CoordinateIndex = Input.Compile(Compiler);
+		}
+	}
+
+	int32 Result = Compiler->SceneDepthWithoutWater(OffsetIndex, CoordinateIndex, bUseOffset, FallbackDepth);
+	return Result;
+}
+
+void UMaterialExpressionSceneDepthWithoutWater::GetCaption(TArray<FString>& OutCaptions) const
+{
+	OutCaptions.Add(TEXT("Scene Depth Without Water"));
+}
+
+FName UMaterialExpressionSceneDepthWithoutWater::GetInputName(int32 InputIndex) const
+{
+	if (InputIndex == 0)
+	{
+		// Display the current InputMode enum's display name.
+		FByteProperty* InputModeProperty = FindFProperty<FByteProperty>(UMaterialExpressionSceneDepthWithoutWater::StaticClass(), "InputMode");
+		// Can't use GetNameByValue as GetNameStringByValue does name mangling that GetNameByValue does not
+		return *InputModeProperty->Enum->GetNameStringByValue((int64)InputMode.GetValue());
+	}
+	return NAME_None;
+}
+
+#endif // WITH_EDITOR
 
 
 #undef LOCTEXT_NAMESPACE
diff --git a/Engine/Source/Runtime/Engine/Public/MaterialCompiler.h b/Engine/Source/Runtime/Engine/Public/MaterialCompiler.h
--- a/Engine/Source/Runtime/Engine/Public/MaterialCompiler.h
+++ b/Engine/Source/Runtime/Engine/Public/MaterialCompiler.h
@@ -361,6 +361,10 @@ public:
 	virtual int32 GetHairRootUV() = 0;
 	virtual int32 GetHairBaseColor() = 0;
 	virtual int32 GetHairRoughness() = 0;
+	
+	// Water
+	virtual int32 SceneDepthWithoutWater(int32 Offset, int32 ViewportUV, bool bUseOffset, float FallbackDepth) = 0;
+
 	virtual int32 CustomPrimitiveData(int32 OutputIndex, EMaterialValueType Type) = 0;
 	virtual int32 ShadingModel(EMaterialShadingModel InSelectedShadingModel) = 0;
 
@@ -679,6 +683,11 @@ public:
 		return Compiler->SkyAtmosphereDistantLightScatteredLuminance();
 	}
 	
+	virtual int32 SceneDepthWithoutWater(int32 Offset, int32 ViewportUV, bool bUseOffset, float FallbackDepth) override
+	{
+		return Compiler->SceneDepthWithoutWater(Offset, ViewportUV, bUseOffset, FallbackDepth);
+	}
+
 	virtual int32 CustomPrimitiveData(int32 OutputIndex, EMaterialValueType Type) override
 	{
 		return Compiler->CustomPrimitiveData(OutputIndex, Type);
