Fix macro redefinition warning

#rnx

diff --git a/Engine/Shaders/Private/MobileBasePassPixelShader.usf b/Engine/Shaders/Private/MobileBasePassPixelShader.usf
--- a/Engine/Shaders/Private/MobileBasePassPixelShader.usf
+++ b/Engine/Shaders/Private/MobileBasePassPixelShader.usf
@@ -30,6 +30,10 @@
 	#define FEATURE_GATHER4 0
 #endif
 
+#if MATERIAL_SHADINGMODEL_SINGLELAYERWATER
+	#define SIMPLE_SINGLE_LAYER_WATER 1
+#endif
+
 #include "SHCommon.ush"
 #include "MobileBasePassCommon.ush"
 #include "/Engine/Generated/Material.ush"
@@ -42,8 +46,7 @@
 #include "PlanarReflectionShared.ush"
 
 #if MATERIAL_SHADINGMODEL_SINGLELAYERWATER
-#define SIMPLE_SINGLE_LAYER_WATER 1
-#include "SingleLayerWaterCommon.ush"
+	#include "SingleLayerWaterShading.ush"
 #endif
 
 // Add fading CSM plane:
