Adding hack to allow one-sided materials to be modeled as dielectrics in ray tracing translucency.

diff --git a/Engine/Shaders/Private/RayTracing/RayTracingPrimaryRays.usf b/Engine/Shaders/Private/RayTracing/RayTracingPrimaryRays.usf
--- a/Engine/Shaders/Private/RayTracing/RayTracingPrimaryRays.usf
+++ b/Engine/Shaders/Private/RayTracing/RayTracingPrimaryRays.usf
@@ -281,13 +281,22 @@ RAY_TRACING_ENTRY_RAYGEN(RayTracingPrimaryRaysRGS)
 
 			float ReflectionRefractionEventThroughput;
 			bool bIsEntering = Payload.IsFrontFace();
-			float N1 = bIsEntering ? 1.0 : Ior;
-			float N2 = bIsEntering ? Ior : 1.0;
-			float Eta = N1 / N2;
 
 			float3 N = Payload.WorldNormal;
 			float3 V = -Ray.Direction;
-			float NoV = saturate(dot(N, V));
+			float NoV = dot(N, V);
+
+			// Hack to allow one-sided materials to be modeled as dielectrics
+			if (NoV < 0.0)
+			{
+				NoV = -NoV;
+				N = -N;
+				bIsEntering = true;
+			}
+
+			float N1 = bIsEntering ? 1.0 : Ior;
+			float N2 = bIsEntering ? Ior : 1.0;
+			float Eta = N1 / N2;
 			float NoT = CalcNoT(NoV, N1, N2);
 			float Fr = FresnelDielectric(Eta, NoV, NoT);
 
