Fixed deferred reflection reflection capture compositing on ray miss.

#[fyi] Juan.Canada

diff --git a/Engine/Shaders/Private/RayTracing/RayTracingDeferredReflections.usf b/Engine/Shaders/Private/RayTracing/RayTracingDeferredReflections.usf
--- a/Engine/Shaders/Private/RayTracing/RayTracingDeferredReflections.usf
+++ b/Engine/Shaders/Private/RayTracing/RayTracingDeferredReflections.usf
@@ -97,7 +97,15 @@ RAY_TRACING_ENTRY_RAYGEN(RayTracingDeferredReflectionsRGS)
 
 	const FSortedReflectionRay SortedRay = RayBuffer[SortedRayIndex];
 	const uint2 PixelPos = UnpackPixelCoordinates(SortedRay.PixelCoordinates);
-	if (all(PixelPos < RayTracingResolution) && DeferredMaterialPayload.SortKey < RAY_TRACING_DEFERRED_MATERIAL_KEY_RAY_MISS)
+	if (any(PixelPos >= RayTracingResolution))
+	{
+		return;
+	}
+
+	float4 ResultColor    = (float4)0;
+	float  ResultDistance = DENOISER_MISS_HIT_DISTANCE;
+
+	if (DeferredMaterialPayload.SortKey < RAY_TRACING_DEFERRED_MATERIAL_KEY_RAY_MISS)
 	{
 		RayDesc Ray;
 		Ray.TMin      = 0;
@@ -117,13 +125,8 @@ RAY_TRACING_ENTRY_RAYGEN(RayTracingDeferredReflectionsRGS)
 			Bookmark,
 			PackedPayload);
 
-		float4 ResultColor = (float4)0;
-		float ResultDistance = -1;
-
 		if (PackedPayload.IsHit())
 		{
-			ResultDistance = PackedPayload.HitT;
-
 			bool bShouldDoDirectLighting               = true;
 			bool bShouldDoEmissiveAndIndirectLighting  = true;
 			bool bTopLayerRayTraceSkyLightContribution = false;
@@ -155,12 +158,14 @@ RAY_TRACING_ENTRY_RAYGEN(RayTracingDeferredReflectionsRGS)
 
 			ResultColor.rgb *= View.PreExposure;
 			ResultColor.a = 1;
-		}
 
-		ColorOutput[PixelPos] 			= ResultColor;
-		RayHitDistanceOutput[PixelPos] 	= ResultDistance;
+			ResultDistance = PackedPayload.HitT;
+		}
 	}
 
+	ColorOutput[PixelPos] 			= ResultColor;
+	RayHitDistanceOutput[PixelPos] 	= ResultDistance;
+
 #endif // DIM_DEFERRED_MATERIAL_MODE
 }
 
