Ray Tracing: Modify custom data for hit shaders when intersecting a sub-surface material.

diff --git a/Engine/Shaders/Private/RayTracing/RayTracingMaterialHitShaders.usf b/Engine/Shaders/Private/RayTracing/RayTracingMaterialHitShaders.usf
--- a/Engine/Shaders/Private/RayTracing/RayTracingMaterialHitShaders.usf
+++ b/Engine/Shaders/Private/RayTracing/RayTracingMaterialHitShaders.usf
@@ -360,6 +360,12 @@ RAY_TRACING_ENTRY_CLOSEST_HIT(MaterialCHS,
 	Payload.Ior = Refraction.x;
 	Payload.CustomData = float4(GetMaterialCustomData0(MaterialParameters), GetMaterialCustomData1(MaterialParameters), 0, 0);
 	
+	// Override custom data if sub-surface material
+	if (Payload.ShadingModelID == SHADINGMODELID_SUBSURFACE)
+	{
+		Payload.CustomData = GetMaterialSubsurfaceData(PixelMaterialInputs);
+	}
+
 	#if CLEAR_COAT_BOTTOM_NORMAL
 	{
 		float2 oct2 = UnitVectorToOctahedron(Payload.WorldNormal);
