RTR: Correctly apply opacity to emissive objects.

diff --git a/Engine/Shaders/Private/RayTracing/RayTracingLightingCommon.ush b/Engine/Shaders/Private/RayTracing/RayTracingLightingCommon.ush
--- a/Engine/Shaders/Private/RayTracing/RayTracingLightingCommon.ush
+++ b/Engine/Shaders/Private/RayTracing/RayTracingLightingCommon.ush
@@ -756,7 +756,7 @@ void AccumulateResults(
 	if (ShouldDoEmissiveAndIndirectLighting)
 	{
 		// Emissive & indirect contribution
-		Radiance += Payload.GetRadiance();
+		Radiance += Payload.GetRadiance() * Payload.GetOpacity();
 
 		// Indirect contribution
 		const float3 IndirectLighting = ComputeIndirectLighting(
