Fixed missing SetEnableSkyLightContribution() call in the ray tracing deferred reflections packed payload

diff --git a/Engine/Shaders/Private/RayTracing/RayTracingDeferredReflections.usf b/Engine/Shaders/Private/RayTracing/RayTracingDeferredReflections.usf
--- a/Engine/Shaders/Private/RayTracing/RayTracingDeferredReflections.usf
+++ b/Engine/Shaders/Private/RayTracing/RayTracingDeferredReflections.usf
@@ -125,6 +125,7 @@ RAY_TRACING_ENTRY_RAYGEN(RayTracingDeferredReflectionsRGS)
 		FRayIntersectionBookmark Bookmark = BookmarkBuffer[DeferredMaterialPayload.PixelCoordinates];
 
 		FPackedMaterialClosestHitPayload PackedPayload = (FPackedMaterialClosestHitPayload)0;
+		PackedPayload.SetEnableSkyLightContribution();
 
 		TraceDeferredMaterialShadingRay(
 			TLAS,
