Re-enable channel compression during SSD ingestion for shadow denoising.

diff --git a/Engine/Shaders/Private/ScreenSpaceDenoise/SSDInjest.usf b/Engine/Shaders/Private/ScreenSpaceDenoise/SSDInjest.usf
--- a/Engine/Shaders/Private/ScreenSpaceDenoise/SSDInjest.usf
+++ b/Engine/Shaders/Private/ScreenSpaceDenoise/SSDInjest.usf
@@ -18,8 +18,11 @@
 	#define CONFIG_SIGNAL_OUTPUT_LAYOUT SIGNAL_BUFFER_LAYOUT_PENUMBRA_INJESTION_NSPP
 	#define CONFIG_MULTIPLEXED_SIGNALS_PER_SIGNAL_DOMAIN 1
 
+	#define CONFIG_SIGNAL_INPUT_TEXTURE_TYPE SIGNAL_TEXTURE_TYPE_FLOAT4
+	#define CONFIG_SIGNAL_OUTPUT_TEXTURE_TYPE SIGNAL_TEXTURE_TYPE_UINT2
+
 	#define CONFIG_INPUT_TEXTURE_COUNT DIM_SIGNAL_BATCH_SIZE
-	#define CONFIG_OUTPUT_TEXTURE_COUNT DIM_SIGNAL_BATCH_SIZE
+	#define CONFIG_OUTPUT_TEXTURE_COUNT ((DIM_SIGNAL_BATCH_SIZE + 1) / 2)
 
 #else
 	#error Unknown signal processing.
diff --git a/Engine/Shaders/Private/ScreenSpaceDenoise/SSDSignalBufferEncoding.ush b/Engine/Shaders/Private/ScreenSpaceDenoise/SSDSignalBufferEncoding.ush
--- a/Engine/Shaders/Private/ScreenSpaceDenoise/SSDSignalBufferEncoding.ush
+++ b/Engine/Shaders/Private/ScreenSpaceDenoise/SSDSignalBufferEncoding.ush
@@ -425,6 +425,87 @@ void DecodeMultiplexedSignalsFromFloat4(
 	#endif
 } // DecodeMultiplexedSignalsFromFloat4()
 
+void DecodeMultiplexedSignalsFromUint2(
+	const uint SignalBufferLayout,
+	const uint MultiplexedSampleId,
+	const bool bNormalizeSample,
+	uint2 RawSample[MAX_MULTIPLEXED_TEXTURES],
+	out FSSDSignalArray OutSamples)
+{
+	OutSamples = CreateSignalArrayFromScalarValue(0.0);
+
+	if (0)
+	{
+
+	}
+	else if (SignalBufferLayout == SIGNAL_BUFFER_LAYOUT_PENUMBRA_INPUT_NSPP)
+	{
+		// Shouldn't happen..
+	}
+	else if (SignalBufferLayout == SIGNAL_BUFFER_LAYOUT_PENUMBRA_INJESTION_NSPP)
+	{
+		UNROLL_N(MAX_SIGNAL_BATCH_SIZE)
+		for (uint BatchSignalId = 0; BatchSignalId < MAX_SIGNAL_BATCH_SIZE; BatchSignalId++)
+		{
+			uint MultiplexId = BatchSignalId;
+			uint EncodedData = RawSample[MultiplexId / 2][MultiplexId % 2];
+
+			int MissCountRatio = EncodedData & 0xFF;
+			int TransmissionDistanceRatio = (EncodedData >> 8) & 0xFF;
+			float WorldBluringRadius = f16tof32(EncodedData >> 16);
+			float SampleCount = (WorldBluringRadius == -2.0 ? 0.0 : 1.0);
+
+			float MissCount = MissCountRatio / 256.0 * SampleCount;
+			float TransmissionDistance = TransmissionDistanceRatio / 256.0 * 5.0;
+
+			WorldBluringRadius = (WorldBluringRadius == -1.0 ? WORLD_RADIUS_MISS : WorldBluringRadius);
+
+			// TODO(Denoiser): feed the actual number of sample.
+			OutSamples.Array[MultiplexId].SampleCount = SampleCount;
+			OutSamples.Array[MultiplexId].MissCount = MissCount;
+			OutSamples.Array[MultiplexId].WorldBluringRadius = WorldBluringRadius;
+			OutSamples.Array[MultiplexId].TransmissionDistance = TransmissionDistance;
+		}
+	}
+	else if (SignalBufferLayout == SIGNAL_BUFFER_LAYOUT_PENUMBRA_HISTORY || SignalBufferLayout == SIGNAL_BUFFER_LAYOUT_PENUMBRA_RECONSTRUCTION)
+	{
+		UNROLL_N(MAX_SIGNAL_BATCH_SIZE)
+		for (uint MultiplexId = 0; MultiplexId < MAX_SIGNAL_BATCH_SIZE; MultiplexId++)
+		{
+			uint2 EncodedData = RawSample[MultiplexId];
+
+			float MissCount = f16tof32(EncodedData.x);
+			float SampleCount = f16tof32(EncodedData.x >> 16);
+			float WorldBluringRadius = f16tof32(EncodedData.y);
+			float TransmissionDistance = f16tof32(EncodedData.y >> 16);
+
+			OutSamples.Array[MultiplexId].MissCount = MissCount;
+			OutSamples.Array[MultiplexId].SampleCount = SampleCount;
+			OutSamples.Array[MultiplexId].WorldBluringRadius = (WorldBluringRadius == -1.0 ? WORLD_RADIUS_MISS : WorldBluringRadius);
+			OutSamples.Array[MultiplexId].TransmissionDistance = TransmissionDistance;
+		}
+	}
+	else if (SignalBufferLayout == SIGNAL_BUFFER_LAYOUT_PENUMBRA_REJECTION)
+	{
+		UNROLL_N(MAX_SIGNAL_BATCH_SIZE)
+		for (uint MultiplexId = 0; MultiplexId < MAX_SIGNAL_BATCH_SIZE; MultiplexId++)
+		{
+			uint2 EncodedData = RawSample[MultiplexId];
+
+			float SampleCount = f16tof32(EncodedData.y);
+			float WorldBluringRadius = f16tof32(EncodedData.y >> 16);
+			float MissCount0 = f16tof32(EncodedData.x);
+			float MissCount1 = f16tof32(EncodedData.x >> 16);
+
+			OutSamples.Array[MultiplexId + 0].SampleCount = SampleCount;
+			OutSamples.Array[MultiplexId + 0].MissCount = MissCount0;
+			OutSamples.Array[MultiplexId + 0].WorldBluringRadius = (WorldBluringRadius == -1.0 ? WORLD_RADIUS_MISS : WorldBluringRadius);
+
+			OutSamples.Array[MultiplexId + 1].SampleCount = SampleCount;
+			OutSamples.Array[MultiplexId + 1].MissCount = MissCount1;
+		}
+	}
+}
 
 /** Decode input signal sample from raw float. */
 FSSDSignalArray DecodeMultiplexedSignals(
@@ -443,7 +524,10 @@ FSSDSignalArray DecodeMultiplexedSignals(
 #elif CONFIG_SIGNAL_INPUT_TEXTURE_TYPE == SIGNAL_TEXTURE_TYPE_UINT2
 {
 	FSSDSignalArray MultiplexedSamples = CreateSignalArrayFromScalarValue(0.0);
-	
+	DecodeMultiplexedSignalsFromUint2(
+		SignalBufferLayout, MultiplexedSampleId, bNormalizeSample,
+		CompressedSample.VGPRArray, /* out */ MultiplexedSamples);
+
 	if (0)
 	{
 	
@@ -741,8 +825,94 @@ void EncodeMultiplexedSignals(
 	// Number of buffer the signal get encoded onto <= MAX_MULTIPLEXED_TEXTURES.
 	OutBufferCount = 1;
 
+	if (0)
+	{
+	}
+	else if (SignalBufferLayout == SIGNAL_BUFFER_LAYOUT_PENUMBRA_INJESTION_NSPP)
+	{
+		UNROLL
+		for (uint MultiplexId = 0; MultiplexId < MultiplexCount; MultiplexId++)
+		{
+			float MissCount = Sample[MultiplexId].MissCount;
+			float WorldBluringRadius = Sample[MultiplexId].WorldBluringRadius;
+			float TransmissionDistance = Sample[MultiplexId].TransmissionDistance;
+			float SampleCount = Sample[MultiplexId].SampleCount;
+			if (SampleCount == 0)
+			{
+				WorldBluringRadius = -2.0;
+			}
+			else if (WorldBluringRadius == WORLD_RADIUS_MISS)
+			{
+				WorldBluringRadius = -1.0;
+			}
+
+			int MissCountRatio = MissCount / SampleCount * 256;
+			int TransmissionDistanceRatio = TransmissionDistance / 5.0 * 256;
+			uint EncodedData = MissCountRatio | (TransmissionDistanceRatio << 8);
+			EncodedData |= f32tof16(WorldBluringRadius) << 16;
+
+			if (MultiplexId % 2)
+			{
+				OutRawSample[MultiplexId / 2].y = EncodedData;
+			}
+			else
+			{
+				OutRawSample[MultiplexId / 2].x = EncodedData;
+			}
+		}
+		OutBufferCount = (MultiplexCount + 1) / 2;
+	}
+	else if (SignalBufferLayout == SIGNAL_BUFFER_LAYOUT_PENUMBRA_HISTORY || SignalBufferLayout == SIGNAL_BUFFER_LAYOUT_PENUMBRA_RECONSTRUCTION)
+	{
+		UNROLL
+		for (uint MultiplexId = 0; MultiplexId < MultiplexCount; MultiplexId++)
+		{
+			float SampleCount = Sample[MultiplexId].SampleCount;
+			float NormalizationFactor = rcp(SampleCount);
+
+			float MissCount = Sample[MultiplexId].MissCount;
+			float WorldBluringRadius = (Sample[MultiplexId].WorldBluringRadius == WORLD_RADIUS_MISS ? -1.0 : Sample[MultiplexId].WorldBluringRadius);
+			float TransmissionDistance = Sample[MultiplexId].TransmissionDistance;
+
+			uint2 EncodedData = 0;
+			EncodedData.x = f32tof16(MissCount);
+			EncodedData.x |= f32tof16(SampleCount) << 16;
+			EncodedData.y = f32tof16(WorldBluringRadius);
+			EncodedData.y |= asuint(f32tof16(TransmissionDistance)) << 16;
+
+			OutRawSample[MultiplexId] = EncodedData;
+		}
+		OutBufferCount = MultiplexCount;
+	}
+	else if (SignalBufferLayout == SIGNAL_BUFFER_LAYOUT_PENUMBRA_REJECTION)
+	{
+		const uint BatchSignalCount = MultiplexCount / 2;
+		OutBufferCount = BatchSignalCount;
+
+		UNROLL
+		for (uint BatchSignalId = 0; BatchSignalId < BatchSignalCount; BatchSignalId++)
+		{
+			uint MultiplexId = BatchSignalId * 2;
+
+			float SampleCount = Sample[MultiplexId + 0].SampleCount;
+			float NormalizationFactor = rcp(SampleCount);
+
+			float MissCount0 = Sample[MultiplexId + 0].MissCount;
+			float MissCount1 = Sample[MultiplexId + 1].MissCount;
+			float WorldBluringRadius = (Sample[MultiplexId + 0].WorldBluringRadius == WORLD_RADIUS_MISS ? -1.0 : Sample[MultiplexId + 0].WorldBluringRadius);
+			float TransmissionDistance = 0.0;
+
+			uint2 EncodedData = 0;
+			EncodedData.x  = f32tof16(MissCount0);
+			EncodedData.x |= f32tof16(MissCount1) << 16;
+			EncodedData.y  = f32tof16(SampleCount);
+			EncodedData.y |= f32tof16(WorldBluringRadius) << 16;
+
+			OutRawSample[BatchSignalId] = EncodedData;
+		}
+	}
 	#if COMPILE_SIGNAL_COLOR_SH
-	if (SignalBufferLayout == SIGNAL_BUFFER_LAYOUT_DIFFUSE_INDIRECT_HARMONIC)
+	else if (SignalBufferLayout == SIGNAL_BUFFER_LAYOUT_DIFFUSE_INDIRECT_HARMONIC)
 	{
 		const uint EncodeOptions = (CONFIG_ENCODING_CLAMP_COLOR ? SSD_ENCODE_CLAMP_COLOR : 0) | SSD_ENCODE_NORMALIZE;
 
diff --git a/Engine/Shaders/Private/ScreenSpaceDenoise/SSDSpatialAccumulation.usf b/Engine/Shaders/Private/ScreenSpaceDenoise/SSDSpatialAccumulation.usf
--- a/Engine/Shaders/Private/ScreenSpaceDenoise/SSDSpatialAccumulation.usf
+++ b/Engine/Shaders/Private/ScreenSpaceDenoise/SSDSpatialAccumulation.usf
@@ -51,7 +51,9 @@
 		#define CONFIG_SIGNAL_INPUT_LAYOUT  SIGNAL_BUFFER_LAYOUT_PENUMBRA_INJESTION_NSPP
 		#define CONFIG_SIGNAL_OUTPUT_LAYOUT SIGNAL_BUFFER_LAYOUT_PENUMBRA_RECONSTRUCTION
 
-		#define CONFIG_INPUT_TEXTURE_COUNT CONFIG_SIGNAL_BATCH_SIZE
+		#define CONFIG_SIGNAL_INPUT_TEXTURE_TYPE SIGNAL_TEXTURE_TYPE_UINT2
+		#define CONFIG_SIGNAL_OUTPUT_TEXTURE_TYPE SIGNAL_TEXTURE_TYPE_FLOAT4
+		#define CONFIG_INPUT_TEXTURE_COUNT ((CONFIG_SIGNAL_BATCH_SIZE + 1) / 2)
 		#define CONFIG_OUTPUT_TEXTURE_COUNT CONFIG_SIGNAL_BATCH_SIZE
 
 		#define CONFIG_SAMPLE_SET           SAMPLE_SET_STACKOWIAK_4_SETS
@@ -68,6 +70,9 @@
 		#define CONFIG_SIGNAL_INPUT_LAYOUT  SIGNAL_BUFFER_LAYOUT_PENUMBRA_RECONSTRUCTION
 		#define CONFIG_SIGNAL_OUTPUT_LAYOUT SIGNAL_BUFFER_LAYOUT_PENUMBRA_RECONSTRUCTION
 
+		#define CONFIG_SIGNAL_INPUT_TEXTURE_TYPE SIGNAL_TEXTURE_TYPE_FLOAT4
+		#define CONFIG_SIGNAL_OUTPUT_TEXTURE_TYPE SIGNAL_TEXTURE_TYPE_FLOAT4
+
 		#define CONFIG_INPUT_TEXTURE_COUNT CONFIG_SIGNAL_BATCH_SIZE
 		#define CONFIG_OUTPUT_TEXTURE_COUNT CONFIG_SIGNAL_BATCH_SIZE
 
@@ -78,11 +83,18 @@
 		#define CONFIG_MAX_WITH_REF_DISTANCE 1
 		#define CONFIG_OUTPUT_MODE          OUTPUT_MODE_DRB
 
+	#elif DIM_STAGE == STAGE_REJECTION_PRE_CONVOLUTION
+		#define CONFIG_SIGNAL_INPUT_TEXTURE_TYPE SIGNAL_TEXTURE_TYPE_FLOAT4
+		#define CONFIG_SIGNAL_OUTPUT_TEXTURE_TYPE SIGNAL_TEXTURE_TYPE_FLOAT4
+
 	#elif DIM_STAGE == STAGE_POST_FILTERING
 		// Input and output layout.
 		#define CONFIG_SIGNAL_INPUT_LAYOUT  SIGNAL_BUFFER_LAYOUT_PENUMBRA_HISTORY
 		#define CONFIG_SIGNAL_OUTPUT_LAYOUT SIGNAL_BUFFER_LAYOUT_PENUMBRA_HISTORY
 
+		#define CONFIG_SIGNAL_INPUT_TEXTURE_TYPE SIGNAL_TEXTURE_TYPE_FLOAT4
+		#define CONFIG_SIGNAL_OUTPUT_TEXTURE_TYPE SIGNAL_TEXTURE_TYPE_FLOAT4
+
 		#define CONFIG_INPUT_TEXTURE_COUNT CONFIG_SIGNAL_BATCH_SIZE
 		#define CONFIG_OUTPUT_TEXTURE_COUNT CONFIG_SIGNAL_BATCH_SIZE
 
@@ -94,6 +106,9 @@
 	#elif DIM_STAGE == STAGE_FINAL_OUTPUT
 		#define CONFIG_SIGNAL_INPUT_LAYOUT  SIGNAL_BUFFER_LAYOUT_PENUMBRA_HISTORY
 
+		#define CONFIG_SIGNAL_INPUT_TEXTURE_TYPE SIGNAL_TEXTURE_TYPE_FLOAT4
+		#define CONFIG_SIGNAL_OUTPUT_TEXTURE_TYPE SIGNAL_TEXTURE_TYPE_FLOAT4
+
 		#define CONFIG_INPUT_TEXTURE_COUNT CONFIG_SIGNAL_BATCH_SIZE
 		#define CONFIG_OUTPUT_TEXTURE_COUNT CONFIG_SIGNAL_BATCH_SIZE
 
@@ -551,11 +566,11 @@ float4 InputBufferUVMinMax[CONFIG_SIGNAL_BATCH_SIZE];
 	#error Missing CONFIG_INPUT_TEXTURE_COUNT
 #endif
 
-Texture2D SignalInput_Textures_0;
+FSSDTexture2D SignalInput_Textures_0;
 FSSDTexture2D SignalInputUint_Textures_0;
 
 #if CONFIG_INPUT_TEXTURE_COUNT > 1
-Texture2D SignalInput_Textures_1;
+FSSDTexture2D SignalInput_Textures_1;
 FSSDTexture2D SignalInputUint_Textures_1;
 #else
 #define SignalInput_Textures_1 SignalInput_Textures_0
@@ -563,7 +578,7 @@ FSSDTexture2D SignalInputUint_Textures_1;
 #endif
 
 #if CONFIG_INPUT_TEXTURE_COUNT > 2
-Texture2D SignalInput_Textures_2;
+FSSDTexture2D SignalInput_Textures_2;
 FSSDTexture2D SignalInputUint_Textures_2;
 #else
 #define SignalInput_Textures_2 SignalInput_Textures_0
@@ -571,7 +586,7 @@ FSSDTexture2D SignalInputUint_Textures_2;
 #endif
 
 #if CONFIG_INPUT_TEXTURE_COUNT > 3
-Texture2D SignalInput_Textures_3;
+FSSDTexture2D SignalInput_Textures_3;
 FSSDTexture2D SignalInputUint_Textures_3;
 #else
 #define SignalInput_Textures_3 SignalInput_Textures_0
@@ -633,10 +648,10 @@ void MainCS(
 	Texture2D Signal_Textures_2 = SignalInput_Textures_2;
 	Texture2D Signal_Textures_3 = SignalInput_Textures_3;
 #else
-	FSSDTexture2D Signal_Textures_0 = SignalInputUint_Textures_0;
-	FSSDTexture2D Signal_Textures_1 = SignalInputUint_Textures_1;
-	FSSDTexture2D Signal_Textures_2 = SignalInputUint_Textures_2;
-	FSSDTexture2D Signal_Textures_3 = SignalInputUint_Textures_3;
+	FSSDTexture2D Signal_Textures_0 = SignalInput_Textures_0;
+	FSSDTexture2D Signal_Textures_1 = SignalInput_Textures_1;
+	FSSDTexture2D Signal_Textures_2 = SignalInput_Textures_2;
+	FSSDTexture2D Signal_Textures_3 = SignalInput_Textures_3;
 #endif
 
 	// Find out scene buffer UV.
diff --git a/Engine/Shaders/Private/ScreenSpaceDenoise/SSDTemporalAccumulation.usf b/Engine/Shaders/Private/ScreenSpaceDenoise/SSDTemporalAccumulation.usf
--- a/Engine/Shaders/Private/ScreenSpaceDenoise/SSDTemporalAccumulation.usf
+++ b/Engine/Shaders/Private/ScreenSpaceDenoise/SSDTemporalAccumulation.usf
@@ -40,6 +40,9 @@
 	#define CONFIG_HISTORY_BILATERAL_PRESET BILATERAL_PRESET_MONOCHROMATIC_PENUMBRA
 	#define CONFIG_MULTIPLEXED_SIGNALS_PER_SIGNAL_DOMAIN 1
 
+	#define CONFIG_SIGNAL_INPUT_TEXTURE_TYPE SIGNAL_TEXTURE_TYPE_FLOAT4
+	#define CONFIG_SIGNAL_OUTPUT_TEXTURE_TYPE SIGNAL_TEXTURE_TYPE_FLOAT4
+
 	#define CONFIG_INPUT_TEXTURE_COUNT CONFIG_SIGNAL_BATCH_SIZE
 	#define CONFIG_HISTORY_TEXTURE_COUNT CONFIG_SIGNAL_BATCH_SIZE
 
diff --git a/Engine/Source/Runtime/Renderer/Private/ScreenSpaceDenoise.cpp b/Engine/Source/Runtime/Renderer/Private/ScreenSpaceDenoise.cpp
--- a/Engine/Source/Runtime/Renderer/Private/ScreenSpaceDenoise.cpp
+++ b/Engine/Source/Runtime/Renderer/Private/ScreenSpaceDenoise.cpp
@@ -328,8 +328,8 @@ const TCHAR* const kInjestResourceNames[] = {
 	// ShadowVisibilityMask
 	TEXT("ShadowDenoiserInjest0"),
 	TEXT("ShadowDenoiserInjest1"),
-	TEXT("ShadowDenoiserInjest2"),
-	TEXT("ShadowDenoiserInjest3"),
+	nullptr,
+	nullptr,
 
 	// PolychromaticPenumbraHarmonic
 	nullptr,
@@ -1354,12 +1354,11 @@ static void DenoiseSignalAtConstantPixelDensity(
 			check(Settings.SignalBatchSize >= 1 && Settings.SignalBatchSize <= IScreenSpaceDenoiser::kMaxBatchSize);
 			for (int32 BatchedSignalId = 0; BatchedSignalId < Settings.SignalBatchSize; BatchedSignalId++)
 			{
-				InjestDescs[BatchedSignalId].Format = PF_FloatRGBA;
-				InjestTextureCount = BatchedSignalId;
-				ReconstructionDescs[BatchedSignalId].Format = PF_FloatRGBA;
-				HistoryDescs[BatchedSignalId].Format = PF_FloatRGBA;
+				InjestDescs[BatchedSignalId / 2].Format = (BatchedSignalId % 2) ? PF_R32G32_UINT : PF_R32_UINT;
+				InjestTextureCount = BatchedSignalId / 2 + 1;
+				ReconstructionDescs[BatchedSignalId].Format = PF_R32G32_UINT;
+				HistoryDescs[BatchedSignalId].Format = PF_R32G32_UINT;
 			}
-			InjestTextureCount = Settings.SignalBatchSize;
 
 			HistoryTextureCountPerSignal = 1;
 			ReconstructionTextureCount = Settings.SignalBatchSize;
@@ -1790,7 +1789,7 @@ static void DenoiseSignalAtConstantPixelDensity(
 				{
 					for (int32 BatchedSignalId = 0; BatchedSignalId < Settings.SignalBatchSize; BatchedSignalId++)
 					{
-						RejectionSignalProcessingDescs[BatchedSignalId].Format = PF_FloatRGBA;
+						RejectionSignalProcessingDescs[BatchedSignalId].Format = PF_R32G32_UINT;
 					}
 					RejectionTextureCount = Settings.SignalBatchSize;
 				}
