Use static samplers in few ray generation shaders as a workaround for the outstanding sampler heap corruption bug.

#[fyi] Juan.Canada

diff --git a/Engine/Shaders/Private/RayTracing/RayTracingDeferredReflections.usf b/Engine/Shaders/Private/RayTracing/RayTracingDeferredReflections.usf
--- a/Engine/Shaders/Private/RayTracing/RayTracingDeferredReflections.usf
+++ b/Engine/Shaders/Private/RayTracing/RayTracingDeferredReflections.usf
@@ -18,11 +18,11 @@
 #define USE_SOURCE_TEXTURE_ARRAY	1
 
 #define LTCMatTexture				RaytracingLightsDataPacked.LTCMatTexture
-#define LTCMatSampler				RaytracingLightsDataPacked.LTCMatSampler
+#define LTCMatSampler				GlobalBilinearClampedSampler
 #define LTCAmpTexture				RaytracingLightsDataPacked.LTCAmpTexture
-#define LTCAmpSampler				RaytracingLightsDataPacked.LTCAmpSampler
-#define PreIntegratedGF				ReflectionStruct.PreIntegratedGF			
-#define PreIntegratedGFSampler		ReflectionStruct.PreIntegratedGFSampler	
+#define LTCAmpSampler				GlobalBilinearClampedSampler
+#define PreIntegratedGF				ReflectionStruct.PreIntegratedGF
+#define PreIntegratedGFSampler		GlobalBilinearClampedSampler
 
 #include "../DeferredShadingCommon.ush"
 #include "../DeferredLightingCommon.ush"
diff --git a/Engine/Shaders/Private/RayTracing/RayTracingPrimaryRays.usf b/Engine/Shaders/Private/RayTracing/RayTracingPrimaryRays.usf
--- a/Engine/Shaders/Private/RayTracing/RayTracingPrimaryRays.usf
+++ b/Engine/Shaders/Private/RayTracing/RayTracingPrimaryRays.usf
@@ -11,11 +11,11 @@
 #define USE_SOURCE_TEXTURE_ARRAY	1
 
 #define LTCMatTexture				RaytracingLightsDataPacked.LTCMatTexture
-#define LTCMatSampler				RaytracingLightsDataPacked.LTCMatSampler
+#define LTCMatSampler				GlobalBilinearClampedSampler
 #define LTCAmpTexture				RaytracingLightsDataPacked.LTCAmpTexture
-#define LTCAmpSampler				RaytracingLightsDataPacked.LTCAmpSampler
-#define PreIntegratedGF				ReflectionStruct.PreIntegratedGF			
-#define PreIntegratedGFSampler		ReflectionStruct.PreIntegratedGFSampler	
+#define LTCAmpSampler				GlobalBilinearClampedSampler
+#define PreIntegratedGF				ReflectionStruct.PreIntegratedGF
+#define PreIntegratedGFSampler		GlobalBilinearClampedSampler
 
 #include "../DeferredShadingCommon.ush"
 #include "../DeferredLightingCommon.ush"
diff --git a/Engine/Shaders/Private/RayTracing/RayTracingReflections.usf b/Engine/Shaders/Private/RayTracing/RayTracingReflections.usf
--- a/Engine/Shaders/Private/RayTracing/RayTracingReflections.usf
+++ b/Engine/Shaders/Private/RayTracing/RayTracingReflections.usf
@@ -37,11 +37,11 @@
 #define USE_SOURCE_TEXTURE_ARRAY	1
 
 #define LTCMatTexture				RaytracingLightsDataPacked.LTCMatTexture
-#define LTCMatSampler				RaytracingLightsDataPacked.LTCMatSampler
+#define LTCMatSampler				GlobalBilinearClampedSampler
 #define LTCAmpTexture				RaytracingLightsDataPacked.LTCAmpTexture
-#define LTCAmpSampler				RaytracingLightsDataPacked.LTCAmpSampler
-#define PreIntegratedGF				ReflectionStruct.PreIntegratedGF			
-#define PreIntegratedGFSampler		ReflectionStruct.PreIntegratedGFSampler	
+#define LTCAmpSampler				GlobalBilinearClampedSampler
+#define PreIntegratedGF				ReflectionStruct.PreIntegratedGF
+#define PreIntegratedGFSampler		GlobalBilinearClampedSampler
 
 #include "../DeferredShadingCommon.ush"
 #include "../DeferredLightingCommon.ush"
