Fix EngineTest RayTracing failures.

1) Disable float4 decoding in uint2 decoding path
2) Revert intermediate shadow encoding to use float4 (except for Ingestion)
3) Cleanup uint2 shadow encoding/decoding logic

diff --git a/Engine/Shaders/Private/ScreenSpaceDenoise/SSDSignalBufferEncoding.ush b/Engine/Shaders/Private/ScreenSpaceDenoise/SSDSignalBufferEncoding.ush
--- a/Engine/Shaders/Private/ScreenSpaceDenoise/SSDSignalBufferEncoding.ush
+++ b/Engine/Shaders/Private/ScreenSpaceDenoise/SSDSignalBufferEncoding.ush
@@ -448,17 +448,16 @@ void DecodeMultiplexedSignalsFromUint2(
 		for (uint BatchSignalId = 0; BatchSignalId < MAX_SIGNAL_BATCH_SIZE; BatchSignalId++)
 		{
 			uint MultiplexId = BatchSignalId;
-			uint EncodedData = RawSample[MultiplexId / 2][MultiplexId % 2];
+			uint EncodedData = MultiplexId % 2 ? RawSample[MultiplexId / 2].g : RawSample[MultiplexId / 2].r;
 
-			int MissCountRatio = EncodedData & 0xFF;
-			int TransmissionDistanceRatio = (EncodedData >> 8) & 0xFF;
+			float MissCountRatio = (EncodedData & 0xFF) / 255.0;
+			float TransmissionDistanceRatio = ((EncodedData >> 8) & 0xFF) / 255.0;
 			float WorldBluringRadius = f16tof32(EncodedData >> 16);
 			float SampleCount = (WorldBluringRadius == -2.0 ? 0.0 : 1.0);
 
-			float MissCount = MissCountRatio / 256.0 * SampleCount;
-			float TransmissionDistance = TransmissionDistanceRatio / 256.0 * 5.0;
-
-			WorldBluringRadius = (WorldBluringRadius == -1.0 ? WORLD_RADIUS_MISS : WorldBluringRadius);
+			float MissCount = SampleCount * MissCountRatio;
+			float TransmissionDistance = TransmissionDistanceRatio * 5.0;
+			WorldBluringRadius = SampleCount * (WorldBluringRadius == -1.0 ? WORLD_RADIUS_MISS : WorldBluringRadius);
 
 			// TODO(Denoiser): feed the actual number of sample.
 			OutSamples.Array[MultiplexId].SampleCount = SampleCount;
@@ -547,6 +546,7 @@ FSSDSignalArray DecodeMultiplexedSignals(
 		// TODO(Denoiser): SRV on RG F32 and manually decode with f16tof32(), going to reduce VGPR pressure between fetch and result.
 	}
 	#endif
+#if 0
 	// Uncompress assuming this is packed as 16bit floats.
 	else
 	{
@@ -560,6 +560,7 @@ FSSDSignalArray DecodeMultiplexedSignals(
 			SignalBufferLayout, MultiplexedSampleId, bNormalizeSample,
 			FloatVGPRArray, /* out */ MultiplexedSamples);
 	}
+#endif
 
 	return MultiplexedSamples;
 } // DecodeMultiplexedSignals()
@@ -846,18 +847,20 @@ void EncodeMultiplexedSignals(
 				WorldBluringRadius = -1.0;
 			}
 
-			int MissCountRatio = MissCount / SampleCount * 256;
-			int TransmissionDistanceRatio = TransmissionDistance / 5.0 * 256;
-			uint EncodedData = MissCountRatio | (TransmissionDistanceRatio << 8);
-			EncodedData |= f32tof16(WorldBluringRadius) << 16;
+			float MissCountRatio = MissCount * SafeRcp(SampleCount);
+			float TransmissionDistanceRatio = TransmissionDistance / 5.0;
+
+			uint EncodedData = MissCountRatio * 255;
+			EncodedData |= (uint(TransmissionDistanceRatio * 255) << 8);
+			EncodedData |= (f32tof16(WorldBluringRadius) << 16);
 
 			if (MultiplexId % 2)
 			{
-				OutRawSample[MultiplexId / 2].y = EncodedData;
+				OutRawSample[MultiplexId / 2].g = EncodedData;
 			}
 			else
 			{
-				OutRawSample[MultiplexId / 2].x = EncodedData;
+				OutRawSample[MultiplexId / 2].r = EncodedData;
 			}
 		}
 		OutBufferCount = (MultiplexCount + 1) / 2;
diff --git a/Engine/Source/Runtime/Renderer/Private/ScreenSpaceDenoise.cpp b/Engine/Source/Runtime/Renderer/Private/ScreenSpaceDenoise.cpp
--- a/Engine/Source/Runtime/Renderer/Private/ScreenSpaceDenoise.cpp
+++ b/Engine/Source/Runtime/Renderer/Private/ScreenSpaceDenoise.cpp
@@ -1356,8 +1356,8 @@ static void DenoiseSignalAtConstantPixelDensity(
 			{
 				InjestDescs[BatchedSignalId / 2].Format = (BatchedSignalId % 2) ? PF_R32G32_UINT : PF_R32_UINT;
 				InjestTextureCount = BatchedSignalId / 2 + 1;
-				ReconstructionDescs[BatchedSignalId].Format = PF_R32G32_UINT;
-				HistoryDescs[BatchedSignalId].Format = PF_R32G32_UINT;
+				ReconstructionDescs[BatchedSignalId].Format = PF_FloatRGBA;
+				HistoryDescs[BatchedSignalId].Format = PF_FloatRGBA;
 			}
 
 			HistoryTextureCountPerSignal = 1;
@@ -1789,7 +1789,7 @@ static void DenoiseSignalAtConstantPixelDensity(
 				{
 					for (int32 BatchedSignalId = 0; BatchedSignalId < Settings.SignalBatchSize; BatchedSignalId++)
 					{
-						RejectionSignalProcessingDescs[BatchedSignalId].Format = PF_R32G32_UINT;
+						RejectionSignalProcessingDescs[BatchedSignalId].Format = PF_FloatRGBA;
 					}
 					RejectionTextureCount = Settings.SignalBatchSize;
 				}
