Fix macro redefinition

#rnx

diff --git a/Engine/Shaders/Private/BasePassPixelShader.usf b/Engine/Shaders/Private/BasePassPixelShader.usf
--- a/Engine/Shaders/Private/BasePassPixelShader.usf
+++ b/Engine/Shaders/Private/BasePassPixelShader.usf
@@ -84,6 +84,7 @@ float NormalCurvatureToRoughness(float3 WorldNormal)
 #if SINGLE_LAYER_WATER_SIMPLE_FORWARD
 
 	// Use mobile like simple shading
+	#undef SIMPLE_SINGLE_LAYER_WATER
 	#define SIMPLE_SINGLE_LAYER_WATER 1
 
 	// Remove forward lighting directional light shadow
