Additional fix for MeshParticleVertexFactory, ensuring that the interpolants are poprerly initialized

diff --git a/Engine/Shaders/Private/MeshParticleVertexFactory.ush b/Engine/Shaders/Private/MeshParticleVertexFactory.ush
--- a/Engine/Shaders/Private/MeshParticleVertexFactory.ush
+++ b/Engine/Shaders/Private/MeshParticleVertexFactory.ush
@@ -588,8 +588,9 @@ FVertexFactoryInterpolantsVSToPS VertexFactoryGetInterpolantsVSToPS(FVertexFacto
 #if USE_PARTICLE_TIME
 	Interpolants.ParticleTime = Intermediates.RelativeTime;
 #endif
-#if NEEDS_PARTICLE_TRANSFORM
-	Interpolants.LocalToWorld = Intermediates.LocalToWorld;
+
+#if NEEDS_PARTICLE_LOCAL_TO_WORLD || NEEDS_PARTICLE_WORLD_TO_LOCAL
+	Interpolants.LocalToWorld = (float3x4) transpose(Intermediates.LocalToWorld);
 #endif
 
 #if INSTANCED_STEREO
