Fixing subsurface specular threshold.

diff --git a/Engine/Shaders/Private/SubsurfaceProfileCommon.ush b/Engine/Shaders/Private/SubsurfaceProfileCommon.ush
--- a/Engine/Shaders/Private/SubsurfaceProfileCommon.ush
+++ b/Engine/Shaders/Private/SubsurfaceProfileCommon.ush
@@ -29,7 +29,7 @@
 #define SSSS_MAX_DUAL_SPECULAR_ROUGHNESS		2.0f
 
 // Threshold value at which model switches from SSS to default lit
-#define SSSS_OPACITY_THRESHOLD_EPS 0.01
+#define SSSS_OPACITY_THRESHOLD_EPS 0.10
 
 // One profile per line, encoded using constants above. See FSubsurfaceProfileTexture::CreateTexture() in SubsurfaceProfile.cpp.
 // By default, SSProfilesTexture is "Texture2D SSProfilesTexture", but it can be remapped as in BasePassCommon.usf.
