Addressing a number of DXC shader validation errors and warnings

#rnx

diff --git a/Engine/Shaders/Private/ComposeSeparateTranslucency.usf b/Engine/Shaders/Private/ComposeSeparateTranslucency.usf
--- a/Engine/Shaders/Private/ComposeSeparateTranslucency.usf
+++ b/Engine/Shaders/Private/ComposeSeparateTranslucency.usf
@@ -15,7 +15,7 @@ void MainPS(
 	float4 SvPosition : SV_POSITION,
 	out float4 OutColor : SV_Target0)
 {
-	float2 BufferUV = SvPosition * View.BufferSizeAndInvSize.zw;
+	float2 BufferUV = SvPosition.xy * View.BufferSizeAndInvSize.zw;
 	
 	float4 SceneColorSample = SceneColor.SampleLevel(SceneColorSampler, BufferUV, 0);
 	float4 SeparateTranslucencySample = SeparateTranslucency.SampleLevel(SeparateTranslucencySampler, BufferUV, 0);
diff --git a/Engine/Shaders/Private/CompositeUIPixelShader.usf b/Engine/Shaders/Private/CompositeUIPixelShader.usf
--- a/Engine/Shaders/Private/CompositeUIPixelShader.usf
+++ b/Engine/Shaders/Private/CompositeUIPixelShader.usf
@@ -100,7 +100,7 @@ void Main(
 	)
 {
 #if PLATFORM_SUPPORTS_RENDERTARGET_WRITE_MASK
-	uint CompositeUIMask = DecodeRTWriteMask(Input.Position, UIWriteMaskTexture, 1);
+	uint CompositeUIMask = DecodeRTWriteMask(Input.Position.xy, UIWriteMaskTexture, 1);
 	BRANCH
 	if (CompositeUIMask == 0)
 	{
diff --git a/Engine/Shaders/Private/DistanceFieldAOShared.ush b/Engine/Shaders/Private/DistanceFieldAOShared.ush
--- a/Engine/Shaders/Private/DistanceFieldAOShared.ush
+++ b/Engine/Shaders/Private/DistanceFieldAOShared.ush
@@ -124,8 +124,8 @@ bool SphereIntersectConeWithDepthRanges(float4 SphereCenterAndRadius, float3 Con
 		float ConeAxisDistance = dot(SphereCenterAndRadius.xyz - ConeVertex, ConeAxis);
 		float2 ConeAxisDistanceMinMax = float2(ConeAxisDistance + SphereCenterAndRadius.w, ConeAxisDistance - SphereCenterAndRadius.w);
 
-		if (ConeAxisDistanceMinMax.x > ConeAxisDepthRanges.x && ConeAxisDistanceMinMax.y < ConeAxisDepthRanges.y
-			|| ConeAxisDistanceMinMax.x > ConeAxisDepthRanges.z && ConeAxisDistanceMinMax.y < ConeAxisDepthRanges.w)
+		if ((ConeAxisDistanceMinMax.x > ConeAxisDepthRanges.x && ConeAxisDistanceMinMax.y < ConeAxisDepthRanges.y)
+			|| (ConeAxisDistanceMinMax.x > ConeAxisDepthRanges.z && ConeAxisDistanceMinMax.y < ConeAxisDepthRanges.w))
 		{
 			return true;
 		}
diff --git a/Engine/Shaders/Private/DistanceFieldObjectCulling.usf b/Engine/Shaders/Private/DistanceFieldObjectCulling.usf
--- a/Engine/Shaders/Private/DistanceFieldObjectCulling.usf
+++ b/Engine/Shaders/Private/DistanceFieldObjectCulling.usf
@@ -41,16 +41,20 @@ void UploadObjectsToBufferCS(
 	{
 		uint DestIndex = UploadOperationIndices[UploadOperationIndex];
 
-		UNROLL
-		for (uint VectorIndex = 0; VectorIndex < OBJECT_BOUNDS_STRIDE; ++VectorIndex)
 		{
-			UploadBoundsFloat4(DestIndex * OBJECT_BOUNDS_STRIDE + VectorIndex, UploadOperationIndex * UPLOAD_DATA_STRIDE + VectorIndex);
+			UNROLL
+			for (uint VectorIndex = 0; VectorIndex < OBJECT_BOUNDS_STRIDE; ++VectorIndex)
+			{
+				UploadBoundsFloat4(DestIndex * OBJECT_BOUNDS_STRIDE + VectorIndex, UploadOperationIndex * UPLOAD_DATA_STRIDE + VectorIndex);
+			}
 		}
 
-		UNROLL
-		for (uint VectorIndex = 0; VectorIndex < OBJECT_DATA_STRIDE; ++VectorIndex)
 		{
-			UploadDataFloat4(DestIndex * OBJECT_DATA_STRIDE + VectorIndex, UploadOperationIndex * UPLOAD_DATA_STRIDE + OBJECT_BOUNDS_STRIDE + VectorIndex);
+			UNROLL
+			for (uint VectorIndex = 0; VectorIndex < OBJECT_DATA_STRIDE; ++VectorIndex)
+			{
+				UploadDataFloat4(DestIndex * OBJECT_DATA_STRIDE + VectorIndex, UploadOperationIndex * UPLOAD_DATA_STRIDE + OBJECT_BOUNDS_STRIDE + VectorIndex);
+			}
 		}
 	}
 }
@@ -86,16 +90,20 @@ void CopyObjectBufferCS(
 		uint SourceIndex = CopyOperationIndex;
 		uint DestIndex = CopyOperationIndex;
 
-		UNROLL
-		for (uint VectorIndex = 0; VectorIndex < OBJECT_BOUNDS_STRIDE; ++VectorIndex)
 		{
-			CopyBoundsFloat4(DestIndex * OBJECT_BOUNDS_STRIDE + VectorIndex, SourceIndex * OBJECT_BOUNDS_STRIDE + VectorIndex);
+			UNROLL
+			for (uint VectorIndex = 0; VectorIndex < OBJECT_BOUNDS_STRIDE; ++VectorIndex)
+			{
+				CopyBoundsFloat4(DestIndex * OBJECT_BOUNDS_STRIDE + VectorIndex, SourceIndex * OBJECT_BOUNDS_STRIDE + VectorIndex);
+			}
 		}
 
-		UNROLL
-		for (uint VectorIndex = 0; VectorIndex < OBJECT_DATA_STRIDE; VectorIndex++)
 		{
-			CopyDataFloat4(DestIndex * OBJECT_DATA_STRIDE + VectorIndex, SourceIndex * OBJECT_DATA_STRIDE + VectorIndex);
+			UNROLL
+			for (uint VectorIndex = 0; VectorIndex < OBJECT_DATA_STRIDE; VectorIndex++)
+			{
+				CopyDataFloat4(DestIndex * OBJECT_DATA_STRIDE + VectorIndex, SourceIndex * OBJECT_DATA_STRIDE + VectorIndex);
+			}
 		}
 	}
 }
@@ -144,16 +152,20 @@ void RemoveObjectsFromBufferCS(
 		uint SourceIndex = RemoveOperationIndices[RemoveOperationIndex].y;
 		uint DestIndex = RemoveOperationIndices[RemoveOperationIndex].x;
 
-		UNROLL
-		for (uint VectorIndex = 0; VectorIndex < OBJECT_BOUNDS_STRIDE; ++VectorIndex)
 		{
-			WriteBoundsFloat4(DestIndex * OBJECT_BOUNDS_STRIDE + VectorIndex, SourceIndex * OBJECT_BOUNDS_STRIDE + VectorIndex);
+			UNROLL
+			for (uint VectorIndex = 0; VectorIndex < OBJECT_BOUNDS_STRIDE; ++VectorIndex)
+			{
+				WriteBoundsFloat4(DestIndex * OBJECT_BOUNDS_STRIDE + VectorIndex, SourceIndex * OBJECT_BOUNDS_STRIDE + VectorIndex);
+			}
 		}
 
-		UNROLL
-		for (uint VectorIndex = 0; VectorIndex < OBJECT_DATA_STRIDE; VectorIndex++)
 		{
-			WriteDataFloat4(DestIndex * OBJECT_DATA_STRIDE + VectorIndex, SourceIndex * OBJECT_DATA_STRIDE + VectorIndex);
+			UNROLL
+			for (uint VectorIndex = 0; VectorIndex < OBJECT_DATA_STRIDE; VectorIndex++)
+			{
+				WriteDataFloat4(DestIndex * OBJECT_DATA_STRIDE + VectorIndex, SourceIndex * OBJECT_DATA_STRIDE + VectorIndex);
+			}
 		}
 	}
 }
diff --git a/Engine/Shaders/Private/DistanceFieldShadowing.usf b/Engine/Shaders/Private/DistanceFieldShadowing.usf
--- a/Engine/Shaders/Private/DistanceFieldShadowing.usf
+++ b/Engine/Shaders/Private/DistanceFieldShadowing.usf
@@ -146,8 +146,8 @@ void CullObjectsForShadowCS(
 #if (DISTANCEFIELD_PRIMITIVE_TYPE == DFPT_SignedDistanceField)
 		float4 ObjectBoundingSphere = LoadFloat4(SceneObjectBounds, SourceIndex);
         
-        if (ShadowBoundingSphere.w == 0 && ShadowConvexHullIntersectSphere(ObjectBoundingSphere.xyz, ObjectBoundingSphere.w)
-			|| ShadowBoundingSphere.w > 0 && dot(ShadowBoundingSphere.xyz - ObjectBoundingSphere.xyz, ShadowBoundingSphere.xyz - ObjectBoundingSphere.xyz) < Square(ShadowBoundingSphere.w + ObjectBoundingSphere.w))
+        if ((ShadowBoundingSphere.w == 0 && ShadowConvexHullIntersectSphere(ObjectBoundingSphere.xyz, ObjectBoundingSphere.w))
+			|| (ShadowBoundingSphere.w > 0 && dot(ShadowBoundingSphere.xyz - ObjectBoundingSphere.xyz, ShadowBoundingSphere.xyz - ObjectBoundingSphere.xyz) < Square(ShadowBoundingSphere.w + ObjectBoundingSphere.w)))
 		{
 			// Assume ObjectBoundingSphere is located at (0, 0, 0) in local space
             float ViewDist2 = length2(ObjectBoundingSphere.xyz - View.WorldViewOrigin);
diff --git a/Engine/Shaders/Private/FXAAShader.usf b/Engine/Shaders/Private/FXAAShader.usf
--- a/Engine/Shaders/Private/FXAAShader.usf
+++ b/Engine/Shaders/Private/FXAAShader.usf
@@ -66,7 +66,6 @@ void FxaaVS(
 #endif
 	
 	#define DECLARE_PARAMETER(a, b, c)
-	float4 fxaaConstDir = float4(1.0, -1.0, 0.25, -0.25);
 
 	// see Fxaa.. header file for more details
 	float4 fxaaConsoleRcpFrameOpt;
diff --git a/Engine/Shaders/Private/GPUFastFourierTransform.usf b/Engine/Shaders/Private/GPUFastFourierTransform.usf
--- a/Engine/Shaders/Private/GPUFastFourierTransform.usf
+++ b/Engine/Shaders/Private/GPUFastFourierTransform.usf
@@ -489,9 +489,8 @@ void ComplexMultTexture(bool bIsHorizontal, bool bUseAlpha, bool bIsGAGroup, in
 //float4 FilterTint;
 void ApplyTint(in Complex Tint, inout Complex LocalBuffer[2][RADIX])
 {
-	for (uint r = 0; r < RADIX; ++r) LocalBuffer[0][r] *= Tint.x;
-
-	for (uint r = 0; r < RADIX; ++r) LocalBuffer[1][r] *= Tint.y;
+	{ for (uint r = 0; r < RADIX; ++r) LocalBuffer[0][r] *= Tint.x; }
+	{ for (uint r = 0; r < RADIX; ++r) LocalBuffer[1][r] *= Tint.y; }
 }
 
 
@@ -632,19 +631,23 @@ void GSConvolutionWithTextureCS(uint3 GroupID : SV_GroupID, uint3 GroupThreadID
 
 
 		// Normalize R & G
-		for (uint r = 0; r < RADIX; ++r)
 		{
-			// This is the R or G channel
-			LocalBuffer[0][r] /= NormMax; 
+			for (uint r = 0; r < RADIX; ++r)
+			{
+				// This is the R or G channel
+				LocalBuffer[0][r] /= NormMax; 
+			}
 		}
 		
 		float AorBNorm = (bIsGAGroup) ? AlphaSum : NormMax; 
 
 		// Normalize B & A
-		for (uint r = 0; r < RADIX; ++r)
 		{
-			// This is the B or A channel
-			LocalBuffer[1][r] /= AorBNorm; 
+			for (uint r = 0; r < RADIX; ++r)
+			{
+				// This is the B or A channel
+				LocalBuffer[1][r] /= AorBNorm; 
+			}
 		}
 	}
 
diff --git a/Engine/Shaders/Private/GPUFastFourierTransform2DCore.ush b/Engine/Shaders/Private/GPUFastFourierTransform2DCore.ush
--- a/Engine/Shaders/Private/GPUFastFourierTransform2DCore.ush
+++ b/Engine/Shaders/Private/GPUFastFourierTransform2DCore.ush
@@ -114,8 +114,8 @@ void CopyDataLocalToDst(in Complex LocalBuffer[2][RADIX], bool bIsHorizontal, ui
 
 void CopyDataSrcWindowToLocal(inout Complex LocalBuffer[2][RADIX], bool bIsHorizontal, in uint ScanIdx, uint Loc, uint Stride, uint4 Window)
 {
-	for (uint i = 0; i < RADIX; ++i) LocalBuffer[0][ i ] = float2(0.f, 0.f);
-	for (uint i = 0; i < RADIX; ++i) LocalBuffer[1][ i ] = float2(0.f, 0.f);
+	{ for (uint i = 0; i < RADIX; ++i) LocalBuffer[0][ i ] = float2(0.f, 0.f); }
+	{ for (uint i = 0; i < RADIX; ++i) LocalBuffer[1][ i ] = float2(0.f, 0.f); }
 	
 	if (bIsHorizontal) 
 	{
@@ -158,8 +158,8 @@ void CopyDataSrcWindowToLocal(inout Complex LocalBuffer[2][RADIX], bool bIsHoriz
 
 void CopyDataSrcWindowToLocal(inout Complex LocalBuffer[2][RADIX], bool bIsHorizontal, uint ScanIdx, uint Loc, uint Stride, uint2 WindowMin, uint2 WindowMax )
 {
-	for (uint i = 0; i < RADIX; ++i) LocalBuffer[0][ i ] = float2(0.f, 0.f);
-	for (uint i = 0; i < RADIX; ++i) LocalBuffer[1][ i ] = float2(0.f, 0.f);
+	{ for (uint i = 0; i < RADIX; ++i) LocalBuffer[0][ i ] = float2(0.f, 0.f); }
+	{ for (uint i = 0; i < RADIX; ++i) LocalBuffer[1][ i ] = float2(0.f, 0.f); }
 
 	if (bIsHorizontal) 
 	{
diff --git a/Engine/Shaders/Private/GPUFastFourierTransformCore.ush b/Engine/Shaders/Private/GPUFastFourierTransformCore.ush
--- a/Engine/Shaders/Private/GPUFastFourierTransformCore.ush
+++ b/Engine/Shaders/Private/GPUFastFourierTransformCore.ush
@@ -376,17 +376,19 @@ groupshared float SharedImag[ SCAN_LINE_LENGTH ];
 
 void CopyLocalToGroupShared(in Complex Local[RADIX], in uint Head, in uint Stride)
 {
- 
-	for (uint r = 0, i = Head; r < RADIX; ++r, i += Stride)
 	{
-		SharedReal[ i ] = Local[ r ].x;
+		for (uint r = 0, i = Head; r < RADIX; ++r, i += Stride)
+		{
+			SharedReal[ i ] = Local[ r ].x;
+		}
 	}
 	 
-	for (uint r = 0, i = Head; r < RADIX; ++r, i += Stride)
 	{
-		SharedImag[ i ] = Local[ r ].y;
+		for (uint r = 0, i = Head; r < RADIX; ++r, i += Stride)
+		{
+			SharedImag[ i ] = Local[ r ].y;
+		}
 	}
-
 }
 
 void CopyLocalToGroupSharedWSync(in Complex Local[RADIX], in uint Head, in uint Stride)
@@ -398,15 +400,18 @@ void CopyLocalToGroupSharedWSync(in Complex Local[RADIX], in uint Head, in uint
 
 void CopyGroupSharedToLocal(inout Complex Local[RADIX], in uint Head, in uint Stride)
 {
- 
-	for (uint r = 0, i = Head; r < RADIX; ++r, i += Stride) 
 	{
-		Local[ r ].x = SharedReal[ i ];
+ 		for (uint r = 0, i = Head; r < RADIX; ++r, i += Stride) 
+		{
+			Local[ r ].x = SharedReal[ i ];
+		}
 	}
 
-	for (uint r = 0, i = Head; r < RADIX; ++r, i += Stride) 
 	{
-		Local[ r ].y = SharedImag[ i ];
+		for (uint r = 0, i = Head; r < RADIX; ++r, i += Stride) 
+		{
+			Local[ r ].y = SharedImag[ i ];
+		}
 	}
 }
 
@@ -616,13 +621,18 @@ void GroupSharedFFT(in const bool bIsForward, inout Complex Local[RADIX], in con
 void Scale(inout Complex LocalBuffer[2][RADIX], in float ScaleValue)
 {
 	// Scale
-	for (uint r = 0; r < RADIX; ++r)
-	{	
-		LocalBuffer[0][r] *= ScaleValue; 
+	{
+		for (uint r = 0; r < RADIX; ++r)
+		{	
+			LocalBuffer[0][r] *= ScaleValue; 
+		}
 	}
-	for (uint r = 0; r < RADIX; ++r)
-	{	
-		LocalBuffer[1][r] *= ScaleValue; 
+
+	{
+		for (uint r = 0; r < RADIX; ++r)
+		{	
+			LocalBuffer[1][r] *= ScaleValue; 
+		}
 	}
 }
 
@@ -709,44 +719,50 @@ void SplitTwoForOne(inout Complex LocalBuffer[RADIX], in uint Head, in uint Stri
 
 	// Construct the transform for the two real signals in the LocalBuffer
 
-	UNROLL
-	for (uint i = 0, K = Head; i < RADIX; ++i, K += Stride)
 	{
-		//
-		// K = N/2 - abs(SrcIdx - N/2)
-
-		// DstK = SrcIdx % Non2;
-		// N - k
+		UNROLL
+		for (uint i = 0, K = Head; i < RADIX; ++i, K += Stride)
+		{
+			//
+			// K = N/2 - abs(SrcIdx - N/2)
+
+			// DstK = SrcIdx % Non2;
+			// N - k
 		
-		FLATTEN uint NmK = (K > 0) ? ( N - K) : 0;
+			uint NmK = (K > 0) ? ( N - K) : 0;
 
-		// Z_k = LocalBuffer[i]
-		// If k < N/2  : Store  F_k = 1/2 (Z_k + Z*_{N-k})
-		// If k > N/2  : Compute I*G_k = 1/2 (Z_k - Z*_{N-k})
+			// Z_k = LocalBuffer[i]
+			// If k < N/2  : Store  F_k = 1/2 (Z_k + Z*_{N-k})
+			// If k > N/2  : Compute I*G_k = 1/2 (Z_k - Z*_{N-k})
 		
-		// Tmp =  {+,-}ComplexConjugate( Z_{N-k})
-		Complex Tmp = Complex( SharedReal[NmK], -SharedImag[NmK] );
+			// Tmp =  {+,-}ComplexConjugate( Z_{N-k})
+			Complex Tmp = Complex( SharedReal[NmK], -SharedImag[NmK] );
 		
-		Tmp *= (K > Non2)? -1 : 1;
+			Tmp *= (K > Non2)? -1 : 1;
 
-		LocalBuffer[i] += Tmp;
+			LocalBuffer[i] += Tmp;
 		
+		}
 	}
 
-	for (uint i =0; i < RADIX; ++i) LocalBuffer[ i ] *= 0.5f;
+	{
+		for (uint i =0; i < RADIX; ++i) LocalBuffer[ i ] *= 0.5f;
+	}
 	
-	UNROLL
-	for (uint i = 0, K = Head; i < RADIX; ++i, K += Stride)
 	{
-		// If k > N/2 get G_k from I*G_k:  G_k = -I * (I G_k)
-		if (K > Non2) LocalBuffer[i] = ComplexMult(Complex(0, -1), LocalBuffer[i] );
+		UNROLL
+		for (uint i = 0, K = Head; i < RADIX; ++i, K += Stride)
+		{
+			// If k > N/2 get G_k from I*G_k:  G_k = -I * (I G_k)
+			if (K > Non2) LocalBuffer[i] = ComplexMult(Complex(0, -1), LocalBuffer[i] );
 		
-		if (K == Non2) 
-		{	
-			// F_N/2 + I * G_N/2
-			LocalBuffer[i] = Complex(SharedReal[Non2], SharedImag[Non2]);
-		}
+			if (K == Non2) 
+			{	
+				// F_N/2 + I * G_N/2
+				LocalBuffer[i] = Complex(SharedReal[Non2], SharedImag[Non2]);
+			}
 	
+		}
 	}
 
 	if (Head == 0) 
@@ -773,27 +789,29 @@ void SplitTwoForOne(inout Complex LocalBuffer[RADIX], in uint Head, in uint Stri
 
 	// Construct the transform for the two real signals in the LocalBuffer
 
-	UNROLL
-	for (uint i = 0, K = Head; i < RADIX; ++i, K += Stride)
 	{
-		//
-		// K = N/2 - abs(SrcIdx - N/2)
-
-		// DstK = SrcIdx % Non2;
-		// N - k
+		UNROLL
+		for (uint i = 0, K = Head; i < RADIX; ++i, K += Stride)
+		{
+			//
+			// K = N/2 - abs(SrcIdx - N/2)
+
+			// DstK = SrcIdx % Non2;
+			// N - k
 		
-		FLATTEN uint NmK = (K > 0) ? ( N - K) : 0;
+			uint NmK = (K > 0) ? ( N - K) : 0;
 
-		// Z_k = LocalBuffer[i]
-		// If k < N/2  : Store  F_k = 1/2 (Z_k + Z*_{N-k})
-		// If k > N/2  : Compute I*G_k = 1/2 (Z_k - Z*_{N-k})
+			// Z_k = LocalBuffer[i]
+			// If k < N/2  : Store  F_k = 1/2 (Z_k + Z*_{N-k})
+			// If k > N/2  : Compute I*G_k = 1/2 (Z_k - Z*_{N-k})
 		
-		// Tmp =  {+,-}ComplexConjugate( Z_{N-k})
-		float Tmp = SharedReal[NmK]; 
+			// Tmp =  {+,-}ComplexConjugate( Z_{N-k})
+			float Tmp = SharedReal[NmK]; 
 		
-		Tmp *= (K > Non2)? -1 : 1;
+			Tmp *= (K > Non2)? -1 : 1;
 
-		LocalBuffer[i].x += Tmp;
+			LocalBuffer[i].x += Tmp;
+		}
 	}
 
 	if (Head == 0 ) LocalBuffer[0].x = 2.f * SharedReal[0];
@@ -804,39 +822,45 @@ void SplitTwoForOne(inout Complex LocalBuffer[RADIX], in uint Head, in uint Stri
 
 	FFTMemoryBarrier();
 
-	UNROLL
-	for (uint i = 0, K = Head; i < RADIX; ++i, K += Stride)
 	{
-		//
-		// K = N/2 - abs(SrcIdx - N/2)
-
-		// DstK = SrcIdx % Non2;
-		// N - k
+		UNROLL
+		for (uint i = 0, K = Head; i < RADIX; ++i, K += Stride)
+		{
+			//
+			// K = N/2 - abs(SrcIdx - N/2)
+
+			// DstK = SrcIdx % Non2;
+			// N - k
 		
-		FLATTEN uint NmK = (K > 0) ? ( N - K) : 0;
+			uint NmK = (K > 0) ? ( N - K) : 0;
 
-		// Z_k = LocalBuffer[i]
-		// If k < N/2  : Store  F_k = 1/2 (Z_k + Z*_{N-k})
-		// If k > N/2  : Compute I*G_k = 1/2 (Z_k - Z*_{N-k})
+			// Z_k = LocalBuffer[i]
+			// If k < N/2  : Store  F_k = 1/2 (Z_k + Z*_{N-k})
+			// If k > N/2  : Compute I*G_k = 1/2 (Z_k - Z*_{N-k})
 		
-		// Tmp =  {+,-}ComplexConjugate( Z_{N-k})
-		float Tmp = -SharedReal[NmK];
+			// Tmp =  {+,-}ComplexConjugate( Z_{N-k})
+			float Tmp = -SharedReal[NmK];
 		
-		Tmp *= (K < Non2)? 1 : -1;
+			Tmp *= (K < Non2)? 1 : -1;
 
-		LocalBuffer[i].y += Tmp;
-		//LocalBuffer[i] *= 0.5; 
+			LocalBuffer[i].y += Tmp;
+			//LocalBuffer[i] *= 0.5; 
+		}
 	}
 
 	if (Head == 0) LocalBuffer[0].y = 2.f * SharedReal[0];
  
-	UNROLL for (uint i = 0; i < RADIX; ++i) LocalBuffer[i] *= 0.5;
+	{
+		UNROLL for (uint i = 0; i < RADIX; ++i) LocalBuffer[i] *= 0.5;
+	}
 
-	UNROLL
-	for (uint i = 0, K = Head; i < RADIX; ++i, K += Stride)
 	{
-		// If k > N/2 get G_k from I*G_k:  G_k = -I * (I G_k)
-		if (K > Non2) LocalBuffer[i] = ComplexMult(Complex(0, -1), LocalBuffer[i] );
+		UNROLL
+		for (uint i = 0, K = Head; i < RADIX; ++i, K += Stride)
+		{
+			// If k > N/2 get G_k from I*G_k:  G_k = -I * (I G_k)
+			if (K > Non2) LocalBuffer[i] = ComplexMult(Complex(0, -1), LocalBuffer[i] );
+		}
 	}
 }
 
@@ -880,45 +904,49 @@ void MergeTwoForOne(inout Complex LocalBuffer[RADIX], in uint Head, in uint Stri
 
 
 	// Compose the transform of a single f+ig signal from the two real signals in the LocalBuffer
-	UNROLL
-	for (uint i = 0, K = Head; i < RADIX; ++i, K += Stride)
 	{
+		UNROLL
+		for (uint i = 0, K = Head; i < RADIX; ++i, K += Stride)
+		{
 	
-		// N - k 
-		FLATTEN uint NmK = (K > 0) ? (N - K) : 0 ;
+			// N - k 
+			uint NmK = (K > 0) ? (N - K) : 0 ;
 	
 			
 		
-		// If k < N/2  : LocalBuffer[i] = F_k,  
-		//               Shared[N-K] =  G[N-K] = Conjugate(G[k]) =  (G[k]_r, -G[k]_i) = Complex(G_r, -G_i)
-		//                          want   I G[k] =  float(-G[k]_i, G[k]_r)
-		//                 Tmp = I * G_k   
-		// If k > N/2  : LocalBuffer[i] = G_k 
-		//               Shared[N-K] = F[N-k] = Conjugate(F[k]) = Complex(F[k]_r, -F[k]_i)
-		//                           want -I F[k] = Complex(F[k]_i, -F[k]_r)
-		//                 Tmp = -I * F   
-		//  ComplexConjugate( Z_{N-k})
-
-		// Tmp = (K < Non2) ? I *  G_k : -I * F
-		Complex Tmp = Complex( SharedImag[ NmK ], SharedReal[ NmK ] );
-		Tmp *= (K > Non2) ? -1 : 1;
+			// If k < N/2  : LocalBuffer[i] = F_k,  
+			//               Shared[N-K] =  G[N-K] = Conjugate(G[k]) =  (G[k]_r, -G[k]_i) = Complex(G_r, -G_i)
+			//                          want   I G[k] =  float(-G[k]_i, G[k]_r)
+			//                 Tmp = I * G_k   
+			// If k > N/2  : LocalBuffer[i] = G_k 
+			//               Shared[N-K] = F[N-k] = Conjugate(F[k]) = Complex(F[k]_r, -F[k]_i)
+			//                           want -I F[k] = Complex(F[k]_i, -F[k]_r)
+			//                 Tmp = -I * F   
+			//  ComplexConjugate( Z_{N-k})
+
+			// Tmp = (K < Non2) ? I *  G_k : -I * F
+			Complex Tmp = Complex( SharedImag[ NmK ], SharedReal[ NmK ] );
+			Tmp *= (K > Non2) ? -1 : 1;
 		
-		LocalBuffer[i] += Tmp;
+			LocalBuffer[i] += Tmp;
+		}
 	}
 	// Compose the transform of a single f+ig signal from the two real signals in the LocalBuffer
 
-	UNROLL
-	for (uint i = 0, K = Head; i < RADIX; ++i, K += Stride)
 	{
-		// if k > N/2 we have G - I *  F.  Multiply by I to get F + I * G
-		if (K > Non2) LocalBuffer[ i ] = ComplexMult(Complex(0, 1), LocalBuffer[ i ]);
+		UNROLL
+		for (uint i = 0, K = Head; i < RADIX; ++i, K += Stride)
+		{
+			// if k > N/2 we have G - I *  F.  Multiply by I to get F + I * G
+			if (K > Non2) LocalBuffer[ i ] = ComplexMult(Complex(0, 1), LocalBuffer[ i ]);
+
+			if (K == Non2)
+			{	
+				// F_N/2 + I * G_N/2
+				LocalBuffer[ i ] = Complex(SharedReal[ Non2 ], SharedImag[ Non2 ]);
+			}
 
-		if (K == Non2)
-		{	
-			// F_N/2 + I * G_N/2
-			LocalBuffer[ i ] = Complex(SharedReal[ Non2 ], SharedImag[ Non2 ]);
 		}
-
 	}
 
 	if (Head == 0) 
@@ -936,37 +964,41 @@ void MergeTwoForOne(inout Complex LocalBuffer[RADIX], in uint Head, in uint Stri
 	// Write the two FFTs into shared memory.
 
 	float TmpX[RADIX];
-	for (uint i = 0; i < RADIX; ++i) TmpX[i] = LocalBuffer[i].x;
+	{
+		for (uint i = 0; i < RADIX; ++i) TmpX[i] = LocalBuffer[i].x;
+	}
 
 	CopyLocalYToGroupShared(LocalBuffer, Head, Stride);
 	
 	FFTMemoryBarrier();
 
 	// Compose the transform of a single f+ig signal from the two real signals in the LocalBuffer
-	UNROLL
-	for (uint i = 0, K = Head; i < RADIX; ++i, K += Stride)
 	{
+		UNROLL
+		for (uint i = 0, K = Head; i < RADIX; ++i, K += Stride)
+		{
 	
-		// N - k 
-		FLATTEN uint NmK = (K > 0) ? (N - K) : 0 ;
+			// N - k 
+			uint NmK = (K > 0) ? (N - K) : 0 ;
 	
 			
 		
-		// If k < N/2  : LocalBuffer[i] = F_k,  
-		//               Shared[N-K] =  G[N-K] = Conjugate(G[k]) =  (G[k]_r, -G[k]_i) = Complex(G_r, -G_i)
-		//                          want   I G[k] =  float(-G[k]_i, G[k]_r)
-		//                 Tmp = I * G_k   
-		// If k > N/2  : LocalBuffer[i] = G_k 
-		//               Shared[N-K] = F[N-k] = Conjugate(F[k]) = Complex(F[k]_r, -F[k]_i)
-		//                           want -I F[k] = Complex(F[k]_i, -F[k]_r)
-		//                 Tmp = -I * F   
-		//  ComplexConjugate( Z_{N-k})
-
-		// Tmp = (K < Non2) ? I *  G_k : -I * F
-		float Tmp = SharedReal[ NmK ]; 
-		Tmp *= (K > Non2) ? -1 : 1;
+			// If k < N/2  : LocalBuffer[i] = F_k,  
+			//               Shared[N-K] =  G[N-K] = Conjugate(G[k]) =  (G[k]_r, -G[k]_i) = Complex(G_r, -G_i)
+			//                          want   I G[k] =  float(-G[k]_i, G[k]_r)
+			//                 Tmp = I * G_k   
+			// If k > N/2  : LocalBuffer[i] = G_k 
+			//               Shared[N-K] = F[N-k] = Conjugate(F[k]) = Complex(F[k]_r, -F[k]_i)
+			//                           want -I F[k] = Complex(F[k]_i, -F[k]_r)
+			//                 Tmp = -I * F   
+			//  ComplexConjugate( Z_{N-k})
+
+			// Tmp = (K < Non2) ? I *  G_k : -I * F
+			float Tmp = SharedReal[ NmK ]; 
+			Tmp *= (K > Non2) ? -1 : 1;
 		
-		LocalBuffer[i].x += Tmp;
+			LocalBuffer[i].x += Tmp;
+		}
 	}
 	
 	float2 FirstElement  = float2(0, SharedReal[0]);
@@ -985,45 +1017,47 @@ void MergeTwoForOne(inout Complex LocalBuffer[RADIX], in uint Head, in uint Stri
 	FirstElement.x  = SharedReal[0];
 	MiddleElement.x = SharedReal[Non2];
 
-	UNROLL
-	for (uint i = 0, K = Head; i < RADIX; ++i, K += Stride)
 	{
+		UNROLL
+		for (uint i = 0, K = Head; i < RADIX; ++i, K += Stride)
+		{
 	
-		// N - k 
-		FLATTEN uint NmK = (K > 0) ? (N - K) : 0 ;
+			// N - k 
+			uint NmK = (K > 0) ? (N - K) : 0 ;
 	
 			
 		
-		// If k < N/2  : LocalBuffer[i] = F_k,  
-		//               Shared[N-K] =  G[N-K] = Conjugate(G[k]) =  (G[k]_r, -G[k]_i) = Complex(G_r, -G_i)
-		//                          want   I G[k] =  float(-G[k]_i, G[k]_r)
-		//                 Tmp = I * G_k   
-		// If k > N/2  : LocalBuffer[i] = G_k 
-		//               Shared[N-K] = F[N-k] = Conjugate(F[k]) = Complex(F[k]_r, -F[k]_i)
-		//                           want -I F[k] = Complex(F[k]_i, -F[k]_r)
-		//                 Tmp = -I * F   
-		//  ComplexConjugate( Z_{N-k})
-
-		// Tmp = (K < Non2) ? I *  G_k : -I * F
-		float Tmp = SharedReal[ NmK ]; 
-		Tmp *= (K > Non2) ? -1 : 1;
+			// If k < N/2  : LocalBuffer[i] = F_k,  
+			//               Shared[N-K] =  G[N-K] = Conjugate(G[k]) =  (G[k]_r, -G[k]_i) = Complex(G_r, -G_i)
+			//                          want   I G[k] =  float(-G[k]_i, G[k]_r)
+			//                 Tmp = I * G_k   
+			// If k > N/2  : LocalBuffer[i] = G_k 
+			//               Shared[N-K] = F[N-k] = Conjugate(F[k]) = Complex(F[k]_r, -F[k]_i)
+			//                           want -I F[k] = Complex(F[k]_i, -F[k]_r)
+			//                 Tmp = -I * F   
+			//  ComplexConjugate( Z_{N-k})
+
+			// Tmp = (K < Non2) ? I *  G_k : -I * F
+			float Tmp = SharedReal[ NmK ]; 
+			Tmp *= (K > Non2) ? -1 : 1;
 		
-		LocalBuffer[i].y += Tmp;
+			LocalBuffer[i].y += Tmp;
+		}
 	}
 
-	
-	UNROLL
-	for (uint i = 0, K = Head; i < RADIX; ++i, K += Stride)
 	{
-		// if k > N/2 we have G - I *  F.  Multiply by I to get F + I * G
-		if (K > Non2) LocalBuffer[ i ] = ComplexMult(Complex(0, 1), LocalBuffer[ i ]);
-
-		if (K == Non2)
-		{	
-			// F_N/2 + I * G_N/2
-			LocalBuffer[ i ] = MiddleElement;
+		UNROLL
+		for (uint i = 0, K = Head; i < RADIX; ++i, K += Stride)
+		{
+			// if k > N/2 we have G - I *  F.  Multiply by I to get F + I * G
+			if (K > Non2) LocalBuffer[ i ] = ComplexMult(Complex(0, 1), LocalBuffer[ i ]);
+
+			if (K == Non2)
+			{	
+				// F_N/2 + I * G_N/2
+				LocalBuffer[ i ] = MiddleElement;
+			}
 		}
-
 	}
 
 	if (Head == 0) LocalBuffer[ 0 ] = FirstElement;
diff --git a/Engine/Shaders/Private/GlobalDistanceField.usf b/Engine/Shaders/Private/GlobalDistanceField.usf
--- a/Engine/Shaders/Private/GlobalDistanceField.usf
+++ b/Engine/Shaders/Private/GlobalDistanceField.usf
@@ -238,26 +238,28 @@ void CompositeObjectDistanceFieldsCS(
 	NumCulledObjects = NumGridCulledObjects;
 #endif
 
-	for (uint ListObjectIndex = ThreadIndex; ListObjectIndex < NumCulledObjects; ListObjectIndex += COMPOSITE_THREADGORUP_TOTALSIZE)
 	{
-		#if USE_CULL_GRID
-		uint ObjectIndex = CullGridObjectArray[CullTileIndex * MAX_GRID_CULLED_DF_OBJECTS + ListObjectIndex];
-		#else
-			uint ObjectIndex = ListObjectIndex;
-		#endif
+		for (uint ListObjectIndex = ThreadIndex; ListObjectIndex < NumCulledObjects; ListObjectIndex += COMPOSITE_THREADGORUP_TOTALSIZE)
+		{
+			#if USE_CULL_GRID
+			uint ObjectIndex = CullGridObjectArray[CullTileIndex * MAX_GRID_CULLED_DF_OBJECTS + ListObjectIndex];
+			#else
+				uint ObjectIndex = ListObjectIndex;
+			#endif
 
-		float4 ObjectPositionAndRadius = LoadObjectPositionAndRadius(ObjectIndex);
-		//float BoxDistance = ComputeDistanceFromBoxToPoint(TileBoxCenter - TileBoxExtent, TileBoxCenter + TileBoxExtent, ObjectPositionAndRadius.xyz);
+			float4 ObjectPositionAndRadius = LoadObjectPositionAndRadius(ObjectIndex);
+			//float BoxDistance = ComputeDistanceFromBoxToPoint(TileBoxCenter - TileBoxExtent, TileBoxCenter + TileBoxExtent, ObjectPositionAndRadius.xyz);
 
-		// ComputeSquaredDistanceFromBoxToPoint is correct but adds a lot to the final cost
-		float BoxDistanceSq = ComputeSquaredDistanceFromBoxToPoint(TileBoxCenter, TileBoxExtent, ObjectPositionAndRadius.xyz);
-		float CombinedDistance = ObjectPositionAndRadius.w + AOGlobalMaxSphereQueryRadius;
+			// ComputeSquaredDistanceFromBoxToPoint is correct but adds a lot to the final cost
+			float BoxDistanceSq = ComputeSquaredDistanceFromBoxToPoint(TileBoxCenter, TileBoxExtent, ObjectPositionAndRadius.xyz);
+			float CombinedDistance = ObjectPositionAndRadius.w + AOGlobalMaxSphereQueryRadius;
 
-		if (BoxDistanceSq < CombinedDistance * CombinedDistance)
-		{
-			uint DestIndex;
-			InterlockedAdd(NumTileCulledObjects, 1U, DestIndex);
-			SharedCulledObjectList[DestIndex] = ObjectIndex;
+			if (BoxDistanceSq < CombinedDistance * CombinedDistance)
+			{
+				uint DestIndex;
+				InterlockedAdd(NumTileCulledObjects, 1U, DestIndex);
+				SharedCulledObjectList[DestIndex] = ObjectIndex;
+			}
 		}
 	}
 
@@ -277,35 +279,37 @@ void CompositeObjectDistanceFieldsCS(
 	MinDistance = ParentGlobalDistanceFieldTexture[IntCoordinate].r;
 #endif
 
-	LOOP 
-	for (uint ListObjectIndex = 0; ListObjectIndex < NumCulledObjects; ListObjectIndex++)
 	{
-		#if USE_OBJECT_COMPOSITING_TILE_CULLING
-			uint ObjectIndex = SharedCulledObjectList[ListObjectIndex];
-		#else
-			uint ObjectIndex = ListObjectIndex;
-		#endif
-
-		float3 LocalPositionExtent = LoadObjectLocalPositionExtent(ObjectIndex);
-		float4x4 WorldToVolume = LoadObjectWorldToVolume(ObjectIndex);
-		bool bGeneratedAsTwoSided;
-		float4 UVScaleAndVolumeScale = LoadObjectUVScale(ObjectIndex, bGeneratedAsTwoSided);
-		float3 VolumePosition = mul(float4(WorldPosition, 1), WorldToVolume).xyz;
-		float BoxDistance = ComputeDistanceFromBoxToPoint(-LocalPositionExtent, LocalPositionExtent, VolumePosition) * UVScaleAndVolumeScale.w;
-
-		BRANCH
-		if (BoxDistance < AOGlobalMaxSphereQueryRadius)
+		LOOP
+		for (uint ListObjectIndex = 0; ListObjectIndex < NumCulledObjects; ListObjectIndex++)
 		{
-			float3 UVAdd = LoadObjectUVAddAndSelfShadowBias(ObjectIndex).xyz;
-			float2 DistanceFieldMAD = LoadObjectDistanceFieldMAD(ObjectIndex);
+			#if USE_OBJECT_COMPOSITING_TILE_CULLING
+				uint ObjectIndex = SharedCulledObjectList[ListObjectIndex];
+			#else
+				uint ObjectIndex = ListObjectIndex;
+			#endif
+
+			float3 LocalPositionExtent = LoadObjectLocalPositionExtent(ObjectIndex);
+			float4x4 WorldToVolume = LoadObjectWorldToVolume(ObjectIndex);
+			bool bGeneratedAsTwoSided;
+			float4 UVScaleAndVolumeScale = LoadObjectUVScale(ObjectIndex, bGeneratedAsTwoSided);
+			float3 VolumePosition = mul(float4(WorldPosition, 1), WorldToVolume).xyz;
+			float BoxDistance = ComputeDistanceFromBoxToPoint(-LocalPositionExtent, LocalPositionExtent, VolumePosition) * UVScaleAndVolumeScale.w;
+
+			BRANCH
+			if (BoxDistance < AOGlobalMaxSphereQueryRadius)
+			{
+				float3 UVAdd = LoadObjectUVAddAndSelfShadowBias(ObjectIndex).xyz;
+				float2 DistanceFieldMAD = LoadObjectDistanceFieldMAD(ObjectIndex);
 
-			float3 ClampedSamplePosition = clamp(VolumePosition, -LocalPositionExtent, LocalPositionExtent);
-			float DistanceToClamped = length(VolumePosition - ClampedSamplePosition);
+				float3 ClampedSamplePosition = clamp(VolumePosition, -LocalPositionExtent, LocalPositionExtent);
+				float DistanceToClamped = length(VolumePosition - ClampedSamplePosition);
 
-			float3 StepVolumeUV = DistanceFieldVolumePositionToUV(ClampedSamplePosition, UVScaleAndVolumeScale.xyz, UVAdd);
-			float DistanceToOccluder = (SampleMeshDistanceField(StepVolumeUV, DistanceFieldMAD).x + DistanceToClamped) * UVScaleAndVolumeScale.w;
+				float3 StepVolumeUV = DistanceFieldVolumePositionToUV(ClampedSamplePosition, UVScaleAndVolumeScale.xyz, UVAdd);
+				float DistanceToOccluder = (SampleMeshDistanceField(StepVolumeUV, DistanceFieldMAD).x + DistanceToClamped) * UVScaleAndVolumeScale.w;
 
-			MinDistance = min(MinDistance, DistanceToOccluder);
+				MinDistance = min(MinDistance, DistanceToOccluder);
+			}
 		}
 	}
 
diff --git a/Engine/Shaders/Private/GpuSkinVertexFactory.ush b/Engine/Shaders/Private/GpuSkinVertexFactory.ush
--- a/Engine/Shaders/Private/GpuSkinVertexFactory.ush
+++ b/Engine/Shaders/Private/GpuSkinVertexFactory.ush
@@ -39,11 +39,13 @@ uint PerBoneMotionBlur;
 
 #if GPUSKIN_APEX_CLOTH
 /** Vertex buffer from which to read simulated positions of clothing. */
+#if FEATURE_LEVEL >= FEATURE_LEVEL_SM4
 Buffer<float2> ClothSimulVertsPositionsNormals;
 Buffer<float2> PreviousClothSimulVertsPositionsNormals;
 /** Transform to apply when positions/normals are loaded */
 float4x4 ClothLocalToWorld;
 float4x4 PreviousClothLocalToWorld;
+#endif // #if FEATURE_LEVEL >= FEATURE_LEVEL_SM4
 /** blend weight between simulated positions and original key-framed animation */
 float ClothBlendWeight;
 #endif// #if GPUSKIN_APEX_CLOTH
@@ -301,7 +303,6 @@ float3 GetClothSimulNormal(int Index, bool bPrevious)
 	}
 	
 #endif //FEATURE_LEVEL >= FEATURE_LEVEL_SM4
-
 }
 
 float3 ClothingPosition(FClothVertex Input, bool bPrevious)
diff --git a/Engine/Shaders/Private/HZB.usf b/Engine/Shaders/Private/HZB.usf
--- a/Engine/Shaders/Private/HZB.usf
+++ b/Engine/Shaders/Private/HZB.usf
@@ -181,7 +181,7 @@ void HZBBuildCS(
 
 void HZBBuildPS(float4 SvPosition : SV_POSITION, out float4 OutColor : SV_Target0)
 {
-	float2 BufferUV = SvPosition * DispatchThreadIdToBufferUV.xy + DispatchThreadIdToBufferUV.zw;
+	float2 BufferUV = SvPosition.xy * DispatchThreadIdToBufferUV.xy + DispatchThreadIdToBufferUV.zw;
 	float4 DeviceZ = Gather4(ParentTextureMip, ParentTextureMipSampler, BufferUV);
 
 	float FurthestDeviceZ = min(min(DeviceZ.x, DeviceZ.y), min(DeviceZ.z, DeviceZ.w));
diff --git a/Engine/Shaders/Private/HZBOcclusion.usf b/Engine/Shaders/Private/HZBOcclusion.usf
--- a/Engine/Shaders/Private/HZBOcclusion.usf
+++ b/Engine/Shaders/Private/HZBOcclusion.usf
@@ -37,18 +37,20 @@ void HZBTestPS(float2 InUV : TEXCOORD0, out float4 OutColor : SV_Target0)
 	// Screen rect from bounds
 	float3 RectMin = float3( 1, 1, 1 );
 	float3 RectMax = float3( -1, -1, -1 );
-	UNROLL for( int i = 0; i < 8; i++ )
 	{
-		float3 PointSrc;
-		PointSrc.x = Bounds[ (i >> 0) & 1 ].x;
-		PointSrc.y = Bounds[ (i >> 1) & 1 ].y;
-		PointSrc.z = Bounds[ (i >> 2) & 1 ].z;
+		UNROLL for( int i = 0; i < 8; i++ )
+		{
+			float3 PointSrc;
+			PointSrc.x = Bounds[ (i >> 0) & 1 ].x;
+			PointSrc.y = Bounds[ (i >> 1) & 1 ].y;
+			PointSrc.z = Bounds[ (i >> 2) & 1 ].z;
 
-		float4 PointClip = mul( float4( PointSrc, 1 ), View.TranslatedWorldToClip );
-		float3 PointScreen = PointClip.xyz / PointClip.w;
+			float4 PointClip = mul( float4( PointSrc, 1 ), View.TranslatedWorldToClip );
+			float3 PointScreen = PointClip.xyz / PointClip.w;
 
-		RectMin = min( RectMin, PointScreen );
-		RectMax = max( RectMax, PointScreen );
+			RectMin = min( RectMin, PointScreen );
+			RectMax = max( RectMax, PointScreen );
+		}
 	}
 
 #if 0
@@ -194,15 +196,17 @@ void HZBTestPS(float2 InUV : TEXCOORD0, out float4 OutColor : SV_Target0)
 	float2 Bias = HZBUvFactor.xy * Rect.xy;
 
 	float4 MinDepth = 1;
-	UNROLL for( int i = 0; i < 4; i++ )
 	{
-		// TODO could vectorize this
-		float4 Depth;
-		Depth.x = HZBTexture.SampleLevel( HZBSampler, float2( i, 0 ) * Scale + Bias, Level ).r;
-		Depth.y = HZBTexture.SampleLevel( HZBSampler, float2( i, 1 ) * Scale + Bias, Level ).r;
-		Depth.z = HZBTexture.SampleLevel( HZBSampler, float2( i, 2 ) * Scale + Bias, Level ).r;
-		Depth.w = HZBTexture.SampleLevel( HZBSampler, float2( i, 3 ) * Scale + Bias, Level ).r;
-		MinDepth = min( MinDepth, Depth );
+		UNROLL for( int i = 0; i < 4; i++ )
+		{
+			// TODO could vectorize this
+			float4 Depth;
+			Depth.x = HZBTexture.SampleLevel( HZBSampler, float2( i, 0 ) * Scale + Bias, Level ).r;
+			Depth.y = HZBTexture.SampleLevel( HZBSampler, float2( i, 1 ) * Scale + Bias, Level ).r;
+			Depth.z = HZBTexture.SampleLevel( HZBSampler, float2( i, 2 ) * Scale + Bias, Level ).r;
+			Depth.w = HZBTexture.SampleLevel( HZBSampler, float2( i, 3 ) * Scale + Bias, Level ).r;
+			MinDepth = min( MinDepth, Depth );
+		}
 	}
 	MinDepth.x = min( min(MinDepth.x, MinDepth.y), min(MinDepth.z, MinDepth.w) );
 
diff --git a/Engine/Shaders/Private/LPVPropagate.usf b/Engine/Shaders/Private/LPVPropagate.usf
--- a/Engine/Shaders/Private/LPVPropagate.usf
+++ b/Engine/Shaders/Private/LPVPropagate.usf
@@ -423,20 +423,22 @@ void CSPropagate(uint3 DTid : SV_DispatchThreadID, uint3 GTid : SV_GroupThreadID
 	ldsCacheOrigin -= int3(1,1,1);
 #endif
 
-	for (int i=0;i<4;i++)
 	{
-		// Get the cell index,from 0-215 (6x6x6 block). There will be 4 indices per thread (with some duplication)
-		uint cellCacheIndex = laneID*4+i;
-		if ( cellCacheIndex < (6*6*6) )
+		for (int i=0;i<4;i++)
 		{
-			uint3 offset3;
-			offset3.x = ( cellCacheIndex % 6 );
-			offset3.y = ( cellCacheIndex / 6 ) %  6;
-			offset3.z = cellCacheIndex / ( 6 * 6 );
-			int3 cachedGridPos = ldsCacheOrigin + int3(offset3);
-			LPVCell CachedThreadCell = ReadLpvCell( GetGridAddress(cachedGridPos) );
+			// Get the cell index,from 0-215 (6x6x6 block). There will be 4 indices per thread (with some duplication)
+			uint cellCacheIndex = laneID*4+i;
+			if ( cellCacheIndex < (6*6*6) )
+			{
+				uint3 offset3;
+				offset3.x = ( cellCacheIndex % 6 );
+				offset3.y = ( cellCacheIndex / 6 ) %  6;
+				offset3.z = cellCacheIndex / ( 6 * 6 );
+				int3 cachedGridPos = ldsCacheOrigin + int3(offset3);
+				LPVCell CachedThreadCell = ReadLpvCell( GetGridAddress(cachedGridPos) );
  
-			LDSCachedLPV[cellCacheIndex] = PackLPVCell( CachedThreadCell );
+				LDSCachedLPV[cellCacheIndex] = PackLPVCell( CachedThreadCell );
+			}
 		}
 	}
 
diff --git a/Engine/Shaders/Private/PlanarReflectionShaders.usf b/Engine/Shaders/Private/PlanarReflectionShaders.usf
--- a/Engine/Shaders/Private/PlanarReflectionShaders.usf
+++ b/Engine/Shaders/Private/PlanarReflectionShaders.usf
@@ -86,16 +86,20 @@ void PrefilterPlanarReflectionPS(
 		float AccumulatedWeight = 0;
 
 		// Two texels at a time with bilinear
-		LOOP
-		for (float y = 0; y < KernelRadiusY; y += 2)
 		{
-			VerticalSample(BufferUV, y, KernelRadiusY, CenterDistanceToPlane, CenterReflectionDistance, false, AccumulatedColor, AccumulatedWeight);
+			LOOP
+			for (float y = 0; y < KernelRadiusY; y += 2)
+			{
+				VerticalSample(BufferUV, y, KernelRadiusY, CenterDistanceToPlane, CenterReflectionDistance, false, AccumulatedColor, AccumulatedWeight);
+			}
 		}
 
-		LOOP
-		for (float y = -2; y > -KernelRadiusY; y -= 2)
 		{
-			VerticalSample(BufferUV, y, KernelRadiusY, CenterDistanceToPlane, CenterReflectionDistance, true, AccumulatedColor, AccumulatedWeight);
+			LOOP
+			for (float y = -2; y > -KernelRadiusY; y -= 2)
+			{
+				VerticalSample(BufferUV, y, KernelRadiusY, CenterDistanceToPlane, CenterReflectionDistance, true, AccumulatedColor, AccumulatedWeight);
+			}
 		}
 
 		AccumulatedColor = AccumulatedColor / max(AccumulatedWeight, .0001f);
diff --git a/Engine/Shaders/Private/PostProcessTemporalAA.usf b/Engine/Shaders/Private/PostProcessTemporalAA.usf
--- a/Engine/Shaders/Private/PostProcessTemporalAA.usf
+++ b/Engine/Shaders/Private/PostProcessTemporalAA.usf
@@ -2078,7 +2078,7 @@ FTAAHistoryPayload TemporalAASample(uint2 GroupId, uint2 GroupThreadId, uint Gro
 		bool Dynamic7 = SceneVelocityBuffer.SampleLevel(SceneVelocityBufferSampler, InputParams.NearestBufferUV, 0, int2( 0,  1)).x > 0;
 
 		bool Dynamic = Dynamic1 || Dynamic3 || Dynamic4 || Dynamic5 || Dynamic7;
-		IgnoreHistory = IgnoreHistory || !Dynamic && History.Color.a > 0;
+		IgnoreHistory = IgnoreHistory || (!Dynamic && History.Color.a > 0);
 	}
 	#endif
 	
diff --git a/Engine/Shaders/Private/RadixSortShaders.usf b/Engine/Shaders/Private/RadixSortShaders.usf
--- a/Engine/Shaders/Private/RadixSortShaders.usf
+++ b/Engine/Shaders/Private/RadixSortShaders.usf
@@ -443,13 +443,15 @@ void RadixSort_Spine(
 	}
 
 	// Scan raking totals.
-	[unroll]
-	for ( uint RakingOffset = 1; RakingOffset < RAKING_THREADS_PER_DIGIT; RakingOffset <<= 1 )
 	{
-		// Acquire LocalRakingTotals.
-		PPS_BARRIER();
-		Total = Total + LocalRakingTotals[RakingIndex - RakingOffset];
-		LocalRakingTotals[RakingIndex] = Total;
+		[unroll]
+		for ( uint RakingOffset = 1; RakingOffset < RAKING_THREADS_PER_DIGIT; RakingOffset <<= 1 )
+		{
+			// Acquire LocalRakingTotals.
+			PPS_BARRIER();
+			Total = Total + LocalRakingTotals[RakingIndex - RakingOffset];
+			LocalRakingTotals[RakingIndex] = Total;
+		}
 	}
 
 	// Acquire LocalRakingTotals.
@@ -459,13 +461,15 @@ void RadixSort_Spine(
 	LocalTotals[ThreadId + DIGIT_SCAN_PADDING] = DigitSum;
 
 	// Scan local totals.
-	[unroll]
-	for ( uint RakingOffset = 1; RakingOffset < DIGIT_COUNT; RakingOffset <<= 1 )
 	{
-		// Acquire LocalTotals.
-		PPS_BARRIER();
-		DigitSum += LocalTotals[ThreadId + DIGIT_SCAN_PADDING - RakingOffset];
-		LocalTotals[ThreadId + DIGIT_SCAN_PADDING] = DigitSum;
+		[unroll]
+		for ( uint RakingOffset = 1; RakingOffset < DIGIT_COUNT; RakingOffset <<= 1 )
+		{
+			// Acquire LocalTotals.
+			PPS_BARRIER();
+			DigitSum += LocalTotals[ThreadId + DIGIT_SCAN_PADDING - RakingOffset];
+			LocalTotals[ThreadId + DIGIT_SCAN_PADDING] = DigitSum;
+		}
 	}
 
 	// Acquire LocalTotals.
diff --git a/Engine/Shaders/Private/ReflectionEnvironmentPixelShader.usf b/Engine/Shaders/Private/ReflectionEnvironmentPixelShader.usf
--- a/Engine/Shaders/Private/ReflectionEnvironmentPixelShader.usf
+++ b/Engine/Shaders/Private/ReflectionEnvironmentPixelShader.usf
@@ -84,7 +84,7 @@ void RemapClearCoatDiffuseAndSpecularColor(FGBufferData GBuffer, float2 ScreenPo
 		//BaseColor += Dither / 255.f;
 		DiffuseColor = BaseColor - BaseColor * GBuffer.Metallic;
 
-		float3 Specular = lerp(1, RefractionScale, ClearCoat);
+		float Specular = lerp(1, RefractionScale, ClearCoat);
 		SpecularColor = ComputeF0(Specular, BaseColor, GBuffer.Metallic);
 	}
 }
diff --git a/Engine/Shaders/Private/ShaderPrintDraw.usf b/Engine/Shaders/Private/ShaderPrintDraw.usf
--- a/Engine/Shaders/Private/ShaderPrintDraw.usf
+++ b/Engine/Shaders/Private/ShaderPrintDraw.usf
@@ -97,22 +97,26 @@ void AddUIntSymbols(inout float2 ScreenPos, in uint Value)
 	uint SymbolCount = 1;
 	uint Divisor = 1;
 	uint TestValue = Value;
-	for (uint i = 0; i < MaxSymbols-1; ++i)
 	{
-		TestValue /= 10;
-		if (TestValue > 0)
+		for (uint i = 0; i < MaxSymbols-1; ++i)
 		{
-			SymbolCount ++;
-			Divisor *= 10;
+			TestValue /= 10;
+			if (TestValue > 0)
+			{
+				SymbolCount ++;
+				Divisor *= 10;
+			}
 		}
 	}
 	
-	for (uint i = 0; i < SymbolCount; ++i)
 	{
-		uint Digit = Value / Divisor;
-		AddSymbol(ScreenPos, GetDecimalSymbol((int)Digit));
-		Value = Value - (Digit * Divisor);
-		Divisor /= 10;
+		for (uint i = 0; i < SymbolCount; ++i)
+		{
+			uint Digit = Value / Divisor;
+			AddSymbol(ScreenPos, GetDecimalSymbol((int)Digit));
+			Value = Value - (Digit * Divisor);
+			Divisor /= 10;
+		}
 	}
 }
 
diff --git a/Engine/Shaders/Private/SlateVertexShader.usf b/Engine/Shaders/Private/SlateVertexShader.usf
--- a/Engine/Shaders/Private/SlateVertexShader.usf
+++ b/Engine/Shaders/Private/SlateVertexShader.usf
@@ -65,7 +65,7 @@ VertexToPixelInterpolants Main(
 
 	{
 		UNROLL
-		for( int CoordinateIndex = 0; CoordinateIndex < NUM_MATERIAL_TEXCOORDS_VERTEX; CoordinateIndex++ )
+		for( uint CoordinateIndex = 0; CoordinateIndex < NUM_MATERIAL_TEXCOORDS_VERTEX; CoordinateIndex++ )
 		{
 			VertexParameters.TexCoords[CoordinateIndex] = UVArrays[CoordinateIndex];
 		}
@@ -76,17 +76,19 @@ VertexToPixelInterpolants Main(
 		float2 CustomizedUVs[NUM_SLATE_TEXCOORDS];
 		GetMaterialCustomizedUVs(VertexParameters, CustomizedUVs);
 
-		UNROLL
-		for( uint CoordinateIndex = 0; CoordinateIndex < NUM_SLATE_TEXCOORDS; CoordinateIndex++ )
 		{
-			FLATTEN
-			if (Mod(CoordinateIndex, 2))
-			{
-				VOut.TextureCoordinates[CoordinateIndex / 2].zw = CustomizedUVs[CoordinateIndex];
-			}
-			else
+			UNROLL
+			for( uint CoordinateIndex = 0; CoordinateIndex < NUM_SLATE_TEXCOORDS; CoordinateIndex++ )
 			{
-				VOut.TextureCoordinates[CoordinateIndex / 2].xy = CustomizedUVs[CoordinateIndex];
+				FLATTEN
+				if (Mod(CoordinateIndex, 2))
+				{
+					VOut.TextureCoordinates[CoordinateIndex / 2].zw = CustomizedUVs[CoordinateIndex];
+				}
+				else
+				{
+					VOut.TextureCoordinates[CoordinateIndex / 2].xy = CustomizedUVs[CoordinateIndex];
+				}
 			}
 		}
 		
diff --git a/Engine/Shaders/Private/SubsurfaceBurleyNormalized.ush b/Engine/Shaders/Private/SubsurfaceBurleyNormalized.ush
--- a/Engine/Shaders/Private/SubsurfaceBurleyNormalized.ush
+++ b/Engine/Shaders/Private/SubsurfaceBurleyNormalized.ush
@@ -944,39 +944,43 @@ void BurleyComputeMain(uint2 DT_ID, uint2 G_ID, uint GI)
 #if !TEXTURE_CACHE_DISABLED
 
 	// We do not need to use unroll if we have a fixed number of thread dimension.
-	UNROLL for (uint i = 0; i < NUM_OF_PASS_REQUIRED_FILL_SHARED_DIFFUSE_TEXTURE; i++)
 	{
-		// Calculate the sample uv for the current thread
-		uint LocalSharedIndex = GI * NUM_OF_PASS_REQUIRED_FILL_SHARED_DIFFUSE_TEXTURE + i;
-		BRANCH if (LocalSharedIndex >= LARGE_GROUP_TEXTURE_SHARE_TOTALSIZE)
+		UNROLL for (uint i = 0; i < NUM_OF_PASS_REQUIRED_FILL_SHARED_DIFFUSE_TEXTURE; i++)
 		{
-			break;
-		}
+			// Calculate the sample uv for the current thread
+			uint LocalSharedIndex = GI * NUM_OF_PASS_REQUIRED_FILL_SHARED_DIFFUSE_TEXTURE + i;
+			BRANCH if (LocalSharedIndex >= LARGE_GROUP_TEXTURE_SHARE_TOTALSIZE)
+			{
+				break;
+			}
 
-		float2 SampleUV = Output_ExtentInverse * (TopLeftCorner + float2(LocalSharedIndex % LARGE_GROUP_TEXTURE_DIMENSION_SIZE,
-			LocalSharedIndex / LARGE_GROUP_TEXTURE_DIMENSION_SIZE) + 0.5f);
-		SampleUV = clamp(SampleUV, SubsurfaceInput0_UVViewportBilinearMin, SubsurfaceInput0_UVViewportBilinearMax);
+			float2 SampleUV = Output_ExtentInverse * (TopLeftCorner + float2(LocalSharedIndex % LARGE_GROUP_TEXTURE_DIMENSION_SIZE,
+				LocalSharedIndex / LARGE_GROUP_TEXTURE_DIMENSION_SIZE) + 0.5f);
+			SampleUV = clamp(SampleUV, SubsurfaceInput0_UVViewportBilinearMin, SubsurfaceInput0_UVViewportBilinearMax);
 
-		float4 SubsurfaceColorWithSSSIndicator = Texture2DSample(SubsurfaceInput0_Texture, SubsurfaceSampler0, SampleUV);
-		SharedSubsurfaceDiffuseLighting[LocalSharedIndex] = SubsurfaceColorWithSSSIndicator;
+			float4 SubsurfaceColorWithSSSIndicator = Texture2DSample(SubsurfaceInput0_Texture, SubsurfaceSampler0, SampleUV);
+			SharedSubsurfaceDiffuseLighting[LocalSharedIndex] = SubsurfaceColorWithSSSIndicator;
 
-		// Also fetch the normal and profile ID
-		FScreenSpaceData SampleScreenSpaceData = GetScreenSpaceData(SampleUV);
-		SharedSubsurfaceWorldNormal[LocalSharedIndex] = SampleScreenSpaceData.GBuffer.WorldNormal;
-		SharedSubsurfaceProfileID[LocalSharedIndex] = ExtractSubsurfaceProfileIntWithInvalid(SampleScreenSpaceData.GBuffer);
+			// Also fetch the normal and profile ID
+			FScreenSpaceData SampleScreenSpaceData = GetScreenSpaceData(SampleUV);
+			SharedSubsurfaceWorldNormal[LocalSharedIndex] = SampleScreenSpaceData.GBuffer.WorldNormal;
+			SharedSubsurfaceProfileID[LocalSharedIndex] = ExtractSubsurfaceProfileIntWithInvalid(SampleScreenSpaceData.GBuffer);
+		}
 	}
 
 	GroupMemoryBarrierWithGroupSync();
 
 	// need to perform burley normalized subsurface scattering if it's burley
 	// this would unroll to just 1 sequence.
-	UNROLL for (uint i = 0; i < LOCALGROUP_RATIO; ++i)
 	{
-		UNROLL for (uint j = 0; j < LOCALGROUP_RATIO; ++j)
+		UNROLL for (uint i = 0; i < LOCALGROUP_RATIO; ++i)
 		{
-			uint2 Pos = ConvertGroupIndexToNormal2DGrid(GI, TopLeftCorner + THREAD_TEXTURE_BORDER + uint2(j, i)*THREAD_SIZE_1D);
-			float2 LBufferUV = ConvertGridPos2UV(Pos);
-			SSSColorUAV[Pos] = BurleyNormalizedSS(LBufferUV, Pos - (TopLeftCorner + THREAD_TEXTURE_BORDER));
+			UNROLL for (uint j = 0; j < LOCALGROUP_RATIO; ++j)
+			{
+				uint2 Pos = ConvertGroupIndexToNormal2DGrid(GI, TopLeftCorner + THREAD_TEXTURE_BORDER + uint2(j, i)*THREAD_SIZE_1D);
+				float2 LBufferUV = ConvertGridPos2UV(Pos);
+				SSSColorUAV[Pos] = BurleyNormalizedSS(LBufferUV, Pos - (TopLeftCorner + THREAD_TEXTURE_BORDER));
+			}
 		}
 	}
 
diff --git a/Engine/Shaders/Private/SurfelTree.usf b/Engine/Shaders/Private/SurfelTree.usf
--- a/Engine/Shaders/Private/SurfelTree.usf
+++ b/Engine/Shaders/Private/SurfelTree.usf
@@ -28,16 +28,20 @@ void CopySurfelBufferCS(
 		uint SourceIndex = CopyOperationIndex;
 		uint DestIndex = CopyOperationIndex;
 
-		UNROLL
-		for (uint VectorIndex = 0; VectorIndex < INTERPOLATED_VERTEX_DATA_STRIDE; VectorIndex++)
 		{
-			RWCopyInterpolatedVertexData[DestIndex * INTERPOLATED_VERTEX_DATA_STRIDE + VectorIndex] = InterpolatedVertexData[SourceIndex * INTERPOLATED_VERTEX_DATA_STRIDE + VectorIndex];
+			UNROLL
+			for (uint VectorIndex = 0; VectorIndex < INTERPOLATED_VERTEX_DATA_STRIDE; VectorIndex++)
+			{
+				RWCopyInterpolatedVertexData[DestIndex * INTERPOLATED_VERTEX_DATA_STRIDE + VectorIndex] = InterpolatedVertexData[SourceIndex * INTERPOLATED_VERTEX_DATA_STRIDE + VectorIndex];
+			}
 		}
 
-		UNROLL
-		for (uint VectorIndex = 0; VectorIndex < SURFEL_DATA_STRIDE; VectorIndex++)
 		{
-			RWCopySurfelData[DestIndex * SURFEL_DATA_STRIDE + VectorIndex] = SurfelData[SourceIndex * SURFEL_DATA_STRIDE + VectorIndex];
+			UNROLL
+			for (uint VectorIndex = 0; VectorIndex < SURFEL_DATA_STRIDE; VectorIndex++)
+			{
+				RWCopySurfelData[DestIndex * SURFEL_DATA_STRIDE + VectorIndex] = SurfelData[SourceIndex * SURFEL_DATA_STRIDE + VectorIndex];
+			}
 		}
 	}
 }
diff --git a/Engine/Shaders/Private/VolumetricFog.usf b/Engine/Shaders/Private/VolumetricFog.usf
--- a/Engine/Shaders/Private/VolumetricFog.usf
+++ b/Engine/Shaders/Private/VolumetricFog.usf
@@ -469,171 +469,175 @@ void LightScatteringCS(
 	NumSuperSamples = HistoryAlpha < .001f && all(GridCoordinate < VolumetricFog.GridSizeInt) ? HISTORY_MISS_SUPER_SAMPLE_COUNT : 1;
 
 #endif
-	UNROLL
-	for (uint SampleIndex = 0; SampleIndex < NumSuperSamples; SampleIndex++)
 	{
-		float3 CellOffset = FrameJitterOffsets[SampleIndex].xyz;
-		//float CellOffset = .5f;
+		UNROLL
+		for (uint SampleIndex = 0; SampleIndex < NumSuperSamples; SampleIndex++)
+		{
+			float3 CellOffset = FrameJitterOffsets[SampleIndex].xyz;
+			//float CellOffset = .5f;
 
-		float SceneDepth;
-		float3 WorldPosition = ComputeCellWorldPosition(GridCoordinate, CellOffset, SceneDepth);
-		float CameraVectorLength = length(WorldPosition - View.WorldCameraOrigin);
-		float3 CameraVector = (WorldPosition - View.WorldCameraOrigin) / CameraVectorLength;
+			float SceneDepth;
+			float3 WorldPosition = ComputeCellWorldPosition(GridCoordinate, CellOffset, SceneDepth);
+			float CameraVectorLength = length(WorldPosition - View.WorldCameraOrigin);
+			float3 CameraVector = (WorldPosition - View.WorldCameraOrigin) / CameraVectorLength;
 
-		BRANCH
-		if (ForwardLightData.HasDirectionalLight)
-		{
-			float ShadowFactor = 1;
-			
-			if (UseDirectionalLightShadowing > 0)
+			BRANCH
+			if (ForwardLightData.HasDirectionalLight)
 			{
-				ShadowFactor *= ComputeDirectionalLightStaticShadowing(WorldPosition);
-				ShadowFactor *= ComputeDirectionalLightDynamicShadowing(WorldPosition, SceneDepth);
-			}
+				float ShadowFactor = 1;
+			
+				if (UseDirectionalLightShadowing > 0)
+				{
+					ShadowFactor *= ComputeDirectionalLightStaticShadowing(WorldPosition);
+					ShadowFactor *= ComputeDirectionalLightDynamicShadowing(WorldPosition, SceneDepth);
+				}
 
-			ShadowFactor *= GetLightFunction(WorldPosition);
+				ShadowFactor *= GetLightFunction(WorldPosition);
 
-			float3 DirectionalLightColor = ForwardLightData.DirectionalLightColor;
+				float3 DirectionalLightColor = ForwardLightData.DirectionalLightColor;
 
-			if (UseHeightFogColors.x > 0)
-			{
-				// Attempt to maintain intensity ratio between sky and sun
-				DirectionalLightColor = VolumetricFog.HeightFogDirectionalLightInscatteringColor * Luminance(ForwardLightData.DirectionalLightColor);
-			}
+				if (UseHeightFogColors.x > 0)
+				{
+					// Attempt to maintain intensity ratio between sky and sun
+					DirectionalLightColor = VolumetricFog.HeightFogDirectionalLightInscatteringColor * Luminance(ForwardLightData.DirectionalLightColor);
+				}
 
-			LightScattering += DirectionalLightColor
-				* (ShadowFactor 
-				* ForwardLightData.DirectionalLightVolumetricScatteringIntensity 
-				* PhaseFunction(PhaseG, dot(ForwardLightData.DirectionalLightDirection, -CameraVector)));
-		}
+				LightScattering += DirectionalLightColor
+					* (ShadowFactor 
+					* ForwardLightData.DirectionalLightVolumetricScatteringIntensity 
+					* PhaseFunction(PhaseG, dot(ForwardLightData.DirectionalLightDirection, -CameraVector)));
+			}
 
-		FTwoBandSHVector RotatedHGZonalHarmonic;
-		RotatedHGZonalHarmonic.V = float4(1.0f, CameraVector.y, CameraVector.z, CameraVector.x) * float4(1.0f, PhaseG, PhaseG, PhaseG);
+			FTwoBandSHVector RotatedHGZonalHarmonic;
+			RotatedHGZonalHarmonic.V = float4(1.0f, CameraVector.y, CameraVector.z, CameraVector.x) * float4(1.0f, PhaseG, PhaseG, PhaseG);
 
-		float3 BrickTextureUVs = 0;
+			float3 BrickTextureUVs = 0;
 
 #if ALLOW_STATIC_LIGHTING
-		if (SkyLightVolumetricScatteringIntensity > 0 || StaticLightingScatteringIntensity > 0)
-		{
-			BrickTextureUVs = ComputeVolumetricLightmapBrickTextureUVs(WorldPosition);
-		}
-#endif
-
-		// Skylight
-		if (SkyLightVolumetricScatteringIntensity > 0)
-		{
-			float3 SkyLighting;
-
-			if (UseHeightFogColors.y > 0)
+			if (SkyLightVolumetricScatteringIntensity > 0 || StaticLightingScatteringIntensity > 0)
 			{
-				float3 HeightFogInscatteringColor = ComputeInscatteringColor(CameraVector, CameraVectorLength);
-				float ScalarFactor = SHAmbientFunction();
-				FTwoBandSHVectorRGB SkyIrradianceSH;
-				SkyIrradianceSH.R.V = float4(ScalarFactor * HeightFogInscatteringColor.r, 0, 0, 0);
-				SkyIrradianceSH.G.V = float4(ScalarFactor * HeightFogInscatteringColor.g, 0, 0, 0);
-				SkyIrradianceSH.B.V = float4(ScalarFactor * HeightFogInscatteringColor.b, 0, 0, 0);
-
-				SkyLighting = max(DotSH(SkyIrradianceSH, RotatedHGZonalHarmonic), 0);
+				BrickTextureUVs = ComputeVolumetricLightmapBrickTextureUVs(WorldPosition);
 			}
-			else
+#endif
+
+			// Skylight
+			if (SkyLightVolumetricScatteringIntensity > 0)
 			{
-				FTwoBandSHVectorRGB SkyIrradianceSH;
-				SkyIrradianceSH.R.V = SkySH[0];
-				SkyIrradianceSH.G.V = SkySH[1];
-				SkyIrradianceSH.B.V = SkySH[2];
+				float3 SkyLighting;
 
-				SkyLighting = View.SkyLightColor.rgb * max(DotSH(SkyIrradianceSH, RotatedHGZonalHarmonic), 0) / PI;
-			}
+				if (UseHeightFogColors.y > 0)
+				{
+					float3 HeightFogInscatteringColor = ComputeInscatteringColor(CameraVector, CameraVectorLength);
+					float ScalarFactor = SHAmbientFunction();
+					FTwoBandSHVectorRGB SkyIrradianceSH;
+					SkyIrradianceSH.R.V = float4(ScalarFactor * HeightFogInscatteringColor.r, 0, 0, 0);
+					SkyIrradianceSH.G.V = float4(ScalarFactor * HeightFogInscatteringColor.g, 0, 0, 0);
+					SkyIrradianceSH.B.V = float4(ScalarFactor * HeightFogInscatteringColor.b, 0, 0, 0);
+
+					SkyLighting = max(DotSH(SkyIrradianceSH, RotatedHGZonalHarmonic), 0);
+				}
+				else
+				{
+					FTwoBandSHVectorRGB SkyIrradianceSH;
+					SkyIrradianceSH.R.V = SkySH[0];
+					SkyIrradianceSH.G.V = SkySH[1];
+					SkyIrradianceSH.B.V = SkySH[2];
 
-			float SkyVisibility = ComputeSkyVisibility(WorldPosition, BrickTextureUVs);
-			LightScattering += (SkyVisibility * SkyLightVolumetricScatteringIntensity) * SkyLighting;
-		}
+					SkyLighting = View.SkyLightColor.rgb * max(DotSH(SkyIrradianceSH, RotatedHGZonalHarmonic), 0) / PI;
+				}
+
+				float SkyVisibility = ComputeSkyVisibility(WorldPosition, BrickTextureUVs);
+				LightScattering += (SkyVisibility * SkyLightVolumetricScatteringIntensity) * SkyLighting;
+			}
 
 #if ALLOW_STATIC_LIGHTING
-		// Indirect lighting of Stationary lights and Direct + Indirect lighting of Static lights
-		if (StaticLightingScatteringIntensity > 0)
-		{
-			FTwoBandSHVectorRGB IrradianceSH = GetVolumetricLightmapSH2(BrickTextureUVs);
+			// Indirect lighting of Stationary lights and Direct + Indirect lighting of Static lights
+			if (StaticLightingScatteringIntensity > 0)
+			{
+				FTwoBandSHVectorRGB IrradianceSH = GetVolumetricLightmapSH2(BrickTextureUVs);
 
-			LightScattering += (StaticLightingScatteringIntensity / PI) * max(DotSH(IrradianceSH, RotatedHGZonalHarmonic), 0);
-		}
+				LightScattering += (StaticLightingScatteringIntensity / PI) * max(DotSH(IrradianceSH, RotatedHGZonalHarmonic), 0);
+			}
 #endif
+		}
 	}
 
-	for (uint SampleIndex = 0; SampleIndex < NumSuperSamples; SampleIndex++)
 	{
-		float3 CellOffset = FrameJitterOffsets[SampleIndex].xyz;
-		//float CellOffset = .5f;
-
-		float SceneDepth;
-		float3 WorldPosition = ComputeCellWorldPosition(GridCoordinate, CellOffset, SceneDepth);
-		float CameraVectorLength = length(WorldPosition - View.WorldCameraOrigin);
-		float3 CameraVector = (WorldPosition - View.WorldCameraOrigin) / CameraVectorLength;
-
-		uint GridIndex = ComputeLightGridCellIndex(GridCoordinate.xy * VolumetricFog.FogGridToPixelXY, SceneDepth, 0);
-		const FCulledLightsGridData CulledLightsGrid = GetCulledLightsGrid(GridIndex, 0);
+		for (uint SampleIndex = 0; SampleIndex < NumSuperSamples; SampleIndex++)
+		{
+			float3 CellOffset = FrameJitterOffsets[SampleIndex].xyz;
+			//float CellOffset = .5f;
 
-		float CellRadius = length(WorldPosition - ComputeCellWorldPosition(GridCoordinate + uint3(1, 1, 1), CellOffset));
-		// Bias the inverse squared light falloff based on voxel size to prevent aliasing near the light source
-		float DistanceBiasSqr = max(CellRadius * InverseSquaredLightDistanceBiasScale, 1);
-		DistanceBiasSqr *= DistanceBiasSqr;
-		FRectTexture RectTexture = InitRectTexture(DeferredLightUniforms.SourceTexture);
+			float SceneDepth;
+			float3 WorldPosition = ComputeCellWorldPosition(GridCoordinate, CellOffset, SceneDepth);
+			float CameraVectorLength = length(WorldPosition - View.WorldCameraOrigin);
+			float3 CameraVector = (WorldPosition - View.WorldCameraOrigin) / CameraVectorLength;
 
-		// Forward lighting of unshadowed point and spot lights
-		LOOP
-		for (uint LocalLightListIndex = 0; LocalLightListIndex < CulledLightsGrid.NumLocalLights; LocalLightListIndex++)
-		{
-			const FLocalLightData LocalLight = GetLocalLightData(CulledLightsGrid.DataStartIndex + LocalLightListIndex, 0);
+			uint GridIndex = ComputeLightGridCellIndex(GridCoordinate.xy * VolumetricFog.FogGridToPixelXY, SceneDepth, 0);
+			const FCulledLightsGridData CulledLightsGrid = GetCulledLightsGrid(GridIndex, 0);
 
-			float VolumetricScatteringIntensity = f16tof32(asuint(LocalLight.SpotAnglesAndSourceRadiusPacked.w) >> 16);
+			float CellRadius = length(WorldPosition - ComputeCellWorldPosition(GridCoordinate + uint3(1, 1, 1), CellOffset));
+			// Bias the inverse squared light falloff based on voxel size to prevent aliasing near the light source
+			float DistanceBiasSqr = max(CellRadius * InverseSquaredLightDistanceBiasScale, 1);
+			DistanceBiasSqr *= DistanceBiasSqr;
+			FRectTexture RectTexture = InitRectTexture(DeferredLightUniforms.SourceTexture);
 
-			if (VolumetricScatteringIntensity > 0)
+			// Forward lighting of unshadowed point and spot lights
+			LOOP
+			for (uint LocalLightListIndex = 0; LocalLightListIndex < CulledLightsGrid.NumLocalLights; LocalLightListIndex++)
 			{
-				FDeferredLightData LightData = (FDeferredLightData)0;
-				LightData.Position = LocalLight.LightPositionAndInvRadius.xyz;
-				LightData.InvRadius = LocalLight.LightPositionAndInvRadius.w;
-				LightData.Color = LocalLight.LightColorAndFalloffExponent.xyz;
-				LightData.FalloffExponent = LocalLight.LightColorAndFalloffExponent.w;
-				LightData.Direction = LocalLight.LightDirectionAndShadowMask.xyz;
-				LightData.SpotAngles = LocalLight.SpotAnglesAndSourceRadiusPacked.xy;
-				LightData.SourceRadius = LocalLight.SpotAnglesAndSourceRadiusPacked.z;
-				LightData.SourceLength = f16tof32(asuint(LocalLight.SpotAnglesAndSourceRadiusPacked.w));
-                LightData.Tangent = LocalLight.LightTangentAndSoftSourceRadius.xyz;
-                LightData.SoftSourceRadius = LocalLight.LightTangentAndSoftSourceRadius.w;
-				LightData.bInverseSquared = LightData.FalloffExponent == 0;
-				LightData.RectLightBarnCosAngle = LocalLight.RectBarnDoor.x;
-				LightData.RectLightBarnLength = LocalLight.RectBarnDoor.y;
-				LightData.bRadialLight = true;
-				LightData.bSpotLight = LightData.SpotAngles.x > -2.0f;
-				LightData.bRectLight = LightData.RectLightBarnLength > -2.0f;
-
-				float3 LightColor = LightData.Color;
-
-				float3 L = 0;
-				float3 ToLight = 0;
-				float LightMask = GetLocalLightAttenuation(WorldPosition, LightData, ToLight, L);
-
-				float Lighting;
-				if( LightData.bRectLight )
-				{
-					FRect Rect = GetRect( ToLight, LightData );
+				const FLocalLightData LocalLight = GetLocalLightData(CulledLightsGrid.DataStartIndex + LocalLightListIndex, 0);
 
-					Lighting = IntegrateLight(Rect, RectTexture);
-				}
-				else
-				{
-					FCapsuleLight Capsule = GetCapsule(ToLight, LightData);
-					Capsule.DistBiasSqr = DistanceBiasSqr;
+				float VolumetricScatteringIntensity = f16tof32(asuint(LocalLight.SpotAnglesAndSourceRadiusPacked.w) >> 16);
 
-					Lighting = IntegrateLight(Capsule, LightData.bInverseSquared);
+				if (VolumetricScatteringIntensity > 0)
+				{
+					FDeferredLightData LightData = (FDeferredLightData)0;
+					LightData.Position = LocalLight.LightPositionAndInvRadius.xyz;
+					LightData.InvRadius = LocalLight.LightPositionAndInvRadius.w;
+					LightData.Color = LocalLight.LightColorAndFalloffExponent.xyz;
+					LightData.FalloffExponent = LocalLight.LightColorAndFalloffExponent.w;
+					LightData.Direction = LocalLight.LightDirectionAndShadowMask.xyz;
+					LightData.SpotAngles = LocalLight.SpotAnglesAndSourceRadiusPacked.xy;
+					LightData.SourceRadius = LocalLight.SpotAnglesAndSourceRadiusPacked.z;
+					LightData.SourceLength = f16tof32(asuint(LocalLight.SpotAnglesAndSourceRadiusPacked.w));
+					LightData.Tangent = LocalLight.LightTangentAndSoftSourceRadius.xyz;
+					LightData.SoftSourceRadius = LocalLight.LightTangentAndSoftSourceRadius.w;
+					LightData.bInverseSquared = LightData.FalloffExponent == 0;
+					LightData.RectLightBarnCosAngle = LocalLight.RectBarnDoor.x;
+					LightData.RectLightBarnLength = LocalLight.RectBarnDoor.y;
+					LightData.bRadialLight = true;
+					LightData.bSpotLight = LightData.SpotAngles.x > -2.0f;
+					LightData.bRectLight = LightData.RectLightBarnLength > -2.0f;
+
+					float3 LightColor = LightData.Color;
+
+					float3 L = 0;
+					float3 ToLight = 0;
+					float LightMask = GetLocalLightAttenuation(WorldPosition, LightData, ToLight, L);
+
+					float Lighting;
+					if( LightData.bRectLight )
+					{
+						FRect Rect = GetRect( ToLight, LightData );
+
+						Lighting = IntegrateLight(Rect, RectTexture);
+					}
+					else
+					{
+						FCapsuleLight Capsule = GetCapsule(ToLight, LightData);
+						Capsule.DistBiasSqr = DistanceBiasSqr;
+
+						Lighting = IntegrateLight(Capsule, LightData.bInverseSquared);
+					}
+
+					float CombinedAttenuation = Lighting * LightMask;
+
+					LightScattering += LightColor * (PhaseFunction(PhaseG, dot(L, -CameraVector)) * CombinedAttenuation * VolumetricScatteringIntensity);
+
+					// To debug culling
+					//LightScattering += LocalLight.LightColorAndFalloffExponent.xyz * .0000001f;
 				}
-
-				float CombinedAttenuation = Lighting * LightMask;
-
-				LightScattering += LightColor * (PhaseFunction(PhaseG, dot(L, -CameraVector)) * CombinedAttenuation * VolumetricScatteringIntensity);
-
-				// To debug culling
-				//LightScattering += LocalLight.LightColorAndFalloffExponent.xyz * .0000001f;
 			}
 		}
 	}
diff --git a/Engine/Shaders/Private/VolumetricLightmapStreaming.usf b/Engine/Shaders/Private/VolumetricLightmapStreaming.usf
--- a/Engine/Shaders/Private/VolumetricLightmapStreaming.usf
+++ b/Engine/Shaders/Private/VolumetricLightmapStreaming.usf
@@ -62,11 +62,11 @@ void CopyResidentBricksCS(
 	uint3 NewBrickPos = ComputeBrickLayoutPosition(StartPosInNewVolume + BrickID.x, uint3(256, 256, 256));
 	uint3 NewPos = NewBrickPos * PaddedSize + CellPos;
 	
-	OutAmbientVector[NewPos] = AmbientVector[OldPos];	
+	OutAmbientVector[NewPos] = AmbientVector[OldPos].xyz;
 	#if HAS_SKY_BENT_NORMAL
 		OutSkyBentNormal[NewPos] = SkyBentNormal[OldPos];
 	#endif
-	OutDirectionalLightShadowing[NewPos] = DirectionalLightShadowing[OldPos];
+	OutDirectionalLightShadowing[NewPos] = DirectionalLightShadowing[OldPos].x;
 }
 
 [numthreads(5, 5, 5)]
diff --git a/Engine/Shaders/Private/WideCustomResolveShaders.usf b/Engine/Shaders/Private/WideCustomResolveShaders.usf
--- a/Engine/Shaders/Private/WideCustomResolveShaders.usf
+++ b/Engine/Shaders/Private/WideCustomResolveShaders.usf
@@ -66,7 +66,7 @@ Texture2DMS<float4,MSAA_SAMPLE_COUNT> Tex; // Input MSAA color
 // Box filter
 float3 resolve_bspline(uint2 pos)
 {
-	return (Tex.Load(pos, 0) + Tex.Load(pos, 1) + Tex.Load(pos, 2) + Tex.Load(pos, 3)) * .25;
+	return (Tex.Load(pos, 0) + Tex.Load(pos, 1) + Tex.Load(pos, 2) + Tex.Load(pos, 3)).xyz * .25;
 }
 # elif WIDE_RESOLVE_WIDTH == 1
 #  include "WideCustomResolve_Wide.ush"
