DXC compilation fixes

- remove default values for some global shader uniforms
  - DXC does not support this and I've verified we are setting these uniforms prior to use

#rnx

diff --git a/Engine/Shaders/Private/DiaphragmDOF/DOFCocTileDilate.usf b/Engine/Shaders/Private/DiaphragmDOF/DOFCocTileDilate.usf
--- a/Engine/Shaders/Private/DiaphragmDOF/DOFCocTileDilate.usf
+++ b/Engine/Shaders/Private/DiaphragmDOF/DOFCocTileDilate.usf
@@ -19,8 +19,8 @@
 //------------------------------------------------------- PASSES.
 
 uint4 ViewportRect;
-uint SampleOffsetMultipler = 1;
-float fSampleOffsetMultipler = 1;
+uint SampleOffsetMultipler;
+float fSampleOffsetMultipler;
 
 // Multiplier to compare CocRadius and bucket distance.
 float CocRadiusToBucketDistanceUpperBound;
diff --git a/Engine/Shaders/Private/TonemapCommon.ush b/Engine/Shaders/Private/TonemapCommon.ush
--- a/Engine/Shaders/Private/TonemapCommon.ush
+++ b/Engine/Shaders/Private/TonemapCommon.ush
@@ -123,11 +123,11 @@ BlackClip = 0;
 WhiteClip = 0.035;
 ===========================================
 */
-float FilmSlope = 0.91;
-float FilmToe = 0.53;
-float FilmShoulder = 0.23;
-float FilmBlackClip = 0;
-float FilmWhiteClip = 0.035;
+float FilmSlope;
+float FilmToe;
+float FilmShoulder;
+float FilmBlackClip;
+float FilmWhiteClip;
 
 half3 FilmToneMap( half3 LinearColor ) 
 {
