Enabling ray tracing global illumination on two-sided foliage.

diff --git a/Engine/Shaders/Private/RayTracing/RayTracingCreateGatherPointsRGS.usf b/Engine/Shaders/Private/RayTracing/RayTracingCreateGatherPointsRGS.usf
--- a/Engine/Shaders/Private/RayTracing/RayTracingCreateGatherPointsRGS.usf
+++ b/Engine/Shaders/Private/RayTracing/RayTracingCreateGatherPointsRGS.usf
@@ -83,9 +83,7 @@ bool ShouldTraceRay(uint2 DispatchThreadId, FGBufferData GBufferData)
 
 	// Cull based on shading model
 	uint ShadingModelID = GBufferData.ShadingModelID;
-	if (ShadingModelID == SHADINGMODELID_UNLIT
-		|| ShadingModelID == SHADINGMODELID_TWOSIDED_FOLIAGE
-		)
+	if (ShadingModelID == SHADINGMODELID_UNLIT)
 	{
 		bTraceRay = false;
 	}
diff --git a/Engine/Shaders/Private/RayTracing/RayTracingFinalGatherRGS.usf b/Engine/Shaders/Private/RayTracing/RayTracingFinalGatherRGS.usf
--- a/Engine/Shaders/Private/RayTracing/RayTracingFinalGatherRGS.usf
+++ b/Engine/Shaders/Private/RayTracing/RayTracingFinalGatherRGS.usf
@@ -66,9 +66,7 @@ RAY_TRACING_ENTRY_RAYGEN(RayTracingFinalGatherRGS)
 
 	// Cull based on shading model
 	uint ShadingModelID = GBufferData.ShadingModelID;
-	if (ShadingModelID == SHADINGMODELID_UNLIT
-		|| ShadingModelID == SHADINGMODELID_TWOSIDED_FOLIAGE
-		)
+	if (ShadingModelID == SHADINGMODELID_UNLIT)
 	{
 		return;
 	}
diff --git a/Engine/Shaders/Private/RayTracing/RayTracingGlobalIlluminationRGS.usf b/Engine/Shaders/Private/RayTracing/RayTracingGlobalIlluminationRGS.usf
--- a/Engine/Shaders/Private/RayTracing/RayTracingGlobalIlluminationRGS.usf
+++ b/Engine/Shaders/Private/RayTracing/RayTracingGlobalIlluminationRGS.usf
@@ -89,9 +89,7 @@ RAY_TRACING_ENTRY_RAYGEN(GlobalIlluminationRGS)
 	ReconstructWorldPositionAndCameraDirectionFromDeviceZ(PixelCoord, DeviceZ, WorldPosition, CameraDirection);
 	float3 WorldNormal = GBufferData.WorldNormal;
 	uint ShadingModelID = GBufferData.ShadingModelID;
-	if (ShadingModelID == SHADINGMODELID_UNLIT
-		|| ShadingModelID == SHADINGMODELID_TWOSIDED_FOLIAGE
-		)
+	if (ShadingModelID == SHADINGMODELID_UNLIT)
 	{
 		return;
 	}
