Fix missing sky raytracing occlusion from hair geometry using hair voxel representation.

diff --git a/Engine/Shaders/Private/HairStrands/HairStrandsRaytracing.ush b/Engine/Shaders/Private/HairStrands/HairStrandsRaytracing.ush
new file mode 100644
--- /dev/null
+++ b/Engine/Shaders/Private/HairStrands/HairStrandsRaytracing.ush
@@ -0,0 +1,95 @@
+// Copyright Epic Games, Inc. All Rights Reserved.
+
+#pragma once
+
+#define VOXEL_TRAVERSAL_DEBUG 0
+#define VOXEL_TRAVERSAL_TYPE VOXEL_TRAVERSAL_LINEAR_MIPMAP
+//#define VOXEL_TRAVERSAL_TYPE VOXEL_TRAVERSAL_LINEAR
+
+#include "../HairStrands/HairStrandsVoxelPageCommon.ush"
+#include "../HairStrands/HairStrandsVoxelPageTraversal.ush"
+
+
+float InternalTraverseHair(uint2 PixelCoord, float3 RayOrigin, float3 RayDirection, float InMinT, float3 VoxelRandom, float RayRandom)
+{
+	float HitT = InMinT;
+
+	// Depth bias
+	// Origin is shifted voxels away towards the light + a constant bias of the size of the voxel
+	{
+		const float3 DepthBias = VirtualVoxel.VoxelWorldSize * (RayRandom * RayDirection + VoxelRandom);
+		RayOrigin += DepthBias;
+	}
+
+	#if VOXEL_TRAVERSAL_DEBUG
+	const bool bDebugEnabled = PixelCoord.x == GetCursorPos().x && PixelCoord.y == GetCursorPos().y;
+	#else
+	const bool bDebugEnabled = false;
+	#endif
+
+	FHairTraversalResult Result = InitHairTraversalResult();
+
+	const float DistanceThreshold = 1000.0f; 
+	const float CoverageThreshold = 0.995f; // When Coverage is high, we do not trace shadow on opaque since hair/fur is covering the background.
+
+	FVirtualVoxelCommonDesc CommonDesc;
+	CommonDesc.PageCountResolution = VirtualVoxel.PageCountResolution;
+	CommonDesc.VoxelWorldSize = VirtualVoxel.VoxelWorldSize;
+	CommonDesc.PageTextureResolution = VirtualVoxel.PageTextureResolution;
+	CommonDesc.PageResolution = VirtualVoxel.PageResolution;
+
+	const float3 P0 = RayOrigin;
+	const float3 P1 = RayOrigin + RayDirection * (HitT>=0 ? min(HitT, DistanceThreshold) : DistanceThreshold);
+
+	for (uint MacroGroupId = 0; MacroGroupId < VirtualVoxel.NodeDescCount; ++MacroGroupId)
+	{
+		const FPackedVirtualVoxelNodeDesc PackedNode = VirtualVoxel.NodeDescBuffer[MacroGroupId];
+		const FVirtualVoxelNodeDesc NodeDesc = UnpackVoxelNode(PackedNode, VirtualVoxel.PageResolution);
+
+		FHairTraversalSettings TraversalSettings = InitHairTraversalSettings();
+		TraversalSettings.DensityScale = VirtualVoxel.DensityScale;
+		TraversalSettings.CountThreshold = 1;
+		TraversalSettings.DistanceThreshold = DistanceThreshold;
+		TraversalSettings.bDebugEnabled = bDebugEnabled;
+		TraversalSettings.SteppingScale = VirtualVoxel.SteppingScale;
+		TraversalSettings.Random = VoxelRandom;
+		TraversalSettings.bUseOpaqueVisibility = false;
+
+		Result = ComputeHairCountVirtualVoxel(
+			P0,
+			P1,
+			CommonDesc,
+			NodeDesc,
+			VirtualVoxel.PageIndexBuffer,
+			VirtualVoxel.PageTexture,
+			TraversalSettings);
+
+		if (Result.HitT >= 0)
+		{
+			const float WorldHitT = DistanceThreshold * Result.HitT;
+			HitT = HitT >= 0 ? min(HitT, WorldHitT) : WorldHitT;
+		}
+	}
+
+	return HitT;
+}
+
+float TraverseHair(uint2 PixelCoord, float3 RayOrigin, float3 RayDirection, float InMinT)
+{
+	float HitT = InMinT;
+
+	const float3 VoxelRandom = 0.5f;
+	const float  RayRandom   = 1;
+	return InternalTraverseHair(PixelCoord, RayOrigin, RayDirection, InMinT, VoxelRandom, RayRandom);
+}
+
+float TraverseHair(uint2 PixelCoord, inout RandomSequence RandSequence, float3 RayOrigin, float3 RayDirection, float InMinT)
+{
+	float HitT = InMinT;
+
+	uint DummyVariable;
+	const float3 VoxelRandom = RandomSequence_GenerateSample3D(RandSequence, DummyVariable);
+	const float RayRandom 	 = RandomSequence_GenerateSample1D(RandSequence, DummyVariable);
+
+	return InternalTraverseHair(PixelCoord, RayOrigin, RayDirection, InMinT, VoxelRandom, RayRandom);
+}
\ No newline at end of file
diff --git a/Engine/Shaders/Private/RayTracing/RayTracingOcclusionRGS.usf b/Engine/Shaders/Private/RayTracing/RayTracingOcclusionRGS.usf
--- a/Engine/Shaders/Private/RayTracing/RayTracingOcclusionRGS.usf
+++ b/Engine/Shaders/Private/RayTracing/RayTracingOcclusionRGS.usf
@@ -40,92 +40,25 @@ int4 LightScissor;
 int2 PixelOffset;
 
 #if USE_HAIR_LIGHTING
-#define VOXEL_TRAVERSAL_DEBUG 0
-#define VOXEL_TRAVERSAL_TYPE VOXEL_TRAVERSAL_LINEAR_MIPMAP
-//#define VOXEL_TRAVERSAL_TYPE VOXEL_TRAVERSAL_LINEAR
-
-#include "../HairStrands/HairStrandsVoxelPageCommon.ush"
-#include "../HairStrands/HairStrandsVoxelPageTraversal.ush"
-
+#include "../HairStrands/HairStrandsRaytracing.ush"
 uint				bUseHairVoxel;
 Texture2D<uint4>	HairCategorizationTexture;
 Texture2D<uint>		HairLightChannelMaskTexture;
 
-float TraverseHair(uint2 PixelCoord, float DeviceZ, inout RandomSequence RandSequence, float3 RayOrigin, float3 RayDirection, float InMinT)
+void NeedTraceHair(
+	in uint2 PixelCoord,
+	inout bool bTraceHairRay,
+	inout float HairDeviceZ)
 {
-	float HitT = InMinT;
-
-
-	const float SceneDepth = ConvertFromDeviceZ(DeviceZ);
-	const float PixelRadius = VirtualVoxel.HairCoveragePixelRadiusAtDepth1 * SceneDepth;
-
-	// Depth bias
-	// Origin is shifted voxels away towards the light + a constant bias of the size of the voxel
-#if 0
-	const float3 VoxelRandom = 0.5f;
-	const float  RayRandom   = 1;
-#else
-	uint DummyVariable;
-	const float3 VoxelRandom = RandomSequence_GenerateSample3D(RandSequence, DummyVariable);
-	const float RayRandom = RandomSequence_GenerateSample1D(RandSequence, DummyVariable);
-#endif
-	{
-		const float3 DepthBias = VirtualVoxel.VoxelWorldSize * (RayRandom * RayDirection + VoxelRandom);
-		RayOrigin += DepthBias;
-	}
-
-	#if VOXEL_TRAVERSAL_DEBUG
-	const bool bDebugEnabled = PixelCoord.x == GetCursorPos().x && PixelCoord.y == GetCursorPos().y;
-	#else
-	const bool bDebugEnabled = false;
-	#endif
-
-	FHairTraversalResult Result = InitHairTraversalResult();
-
-	const float DistanceThreshold = 1000.0f; 
-	const float CoverageThreshold = 0.995f; // When Coverage is high, we do not trace shadow on opaque since hair/fur is covering the background.
-
-	FVirtualVoxelCommonDesc CommonDesc;
-	CommonDesc.PageCountResolution = VirtualVoxel.PageCountResolution;
-	CommonDesc.VoxelWorldSize = VirtualVoxel.VoxelWorldSize;
-	CommonDesc.PageTextureResolution = VirtualVoxel.PageTextureResolution;
-	CommonDesc.PageResolution = VirtualVoxel.PageResolution;
-
-	const float3 P0 = RayOrigin;
-	const float3 P1 = RayOrigin + RayDirection * (HitT>=0 ? min(HitT, DistanceThreshold) : DistanceThreshold);
-
-	for (uint MacroGroupId = 0; MacroGroupId < VirtualVoxel.NodeDescCount; ++MacroGroupId)
+	const FCategorizationData Categorization = DecodeCategorizationData(HairCategorizationTexture.Load(uint3(PixelCoord, 0)));
+	bTraceHairRay = false;
+	HairDeviceZ = 0;
+	if (Categorization.PixelCoverage > 0)
 	{
-		const FPackedVirtualVoxelNodeDesc PackedNode = VirtualVoxel.NodeDescBuffer[MacroGroupId];
-		const FVirtualVoxelNodeDesc NodeDesc = UnpackVoxelNode(PackedNode, VirtualVoxel.PageResolution);
-
-		FHairTraversalSettings TraversalSettings = InitHairTraversalSettings();
-		TraversalSettings.DensityScale = VirtualVoxel.DensityScale;
-		TraversalSettings.CountThreshold = 1;
-		TraversalSettings.DistanceThreshold = DistanceThreshold;
-		TraversalSettings.bDebugEnabled = bDebugEnabled;
-		TraversalSettings.SteppingScale = VirtualVoxel.SteppingScale;
-		TraversalSettings.Random = VoxelRandom;
-		TraversalSettings.bUseOpaqueVisibility = false;
-		TraversalSettings.PixelRadius = PixelRadius;
-
-		Result = ComputeHairCountVirtualVoxel(
-			P0,
-			P1,
-			CommonDesc,
-			NodeDesc,
-			VirtualVoxel.PageIndexBuffer,
-			VirtualVoxel.PageTexture,
-			TraversalSettings);
-
-		if (Result.HitT >= 0)
-		{
-			const float WorldHitT = DistanceThreshold * Result.HitT;
-			HitT = HitT >= 0 ? min(HitT, WorldHitT) : WorldHitT;
-		}
+		const uint HairLightChannel = HairLightChannelMaskTexture.Load(uint3(PixelCoord, 0));
+		HairDeviceZ = Categorization.ClosestDepth;
+		bTraceHairRay = (HairLightChannel & LightingChannelMask) != 0;
 	}
-
-	return HitT;
 }
 #endif
 
@@ -462,7 +395,7 @@ FOcclusionResult ComputeOcclusion(
 		#if USE_HAIR_LIGHTING
 		if (ShadingModelID != SHADINGMODELID_HAIR && bUseHairVoxel)
 		{
-			MinimalPayload.HitT = TraverseHair(PixelCoord, DeviceZ, RandSequence, Ray.Origin, Ray.Direction, MinimalPayload.HitT);
+			MinimalPayload.HitT = TraverseHair(PixelCoord, RandSequence, Ray.Origin, Ray.Direction, MinimalPayload.HitT);
 		}
 		#endif
 
@@ -560,15 +493,9 @@ RAY_TRACING_ENTRY_RAYGEN(OcclusionRGS)
 		}
 
 	#if USE_HAIR_LIGHTING
-		const FCategorizationData Categorization = DecodeCategorizationData(HairCategorizationTexture.Load(uint3(PixelCoord, 0)));
 		bool bTraceHairRay = false;
 		float HairDeviceZ = 0;
-		if (Categorization.PixelCoverage > 0)
-		{
-			const uint HairLightChannel = HairLightChannelMaskTexture.Load(uint3(PixelCoord, 0));
-			HairDeviceZ	= Categorization.ClosestDepth;
-			bTraceHairRay = (HairLightChannel & LightingChannelMask) != 0;
-		}
+		NeedTraceHair(PixelCoord, bTraceHairRay, HairDeviceZ);
 	#endif
 
 		Occlusion = ComputeOcclusion(
diff --git a/Engine/Shaders/Private/RayTracing/RayTracingSkyLightEvaluation.ush b/Engine/Shaders/Private/RayTracing/RayTracingSkyLightEvaluation.ush
--- a/Engine/Shaders/Private/RayTracing/RayTracingSkyLightEvaluation.ush
+++ b/Engine/Shaders/Private/RayTracing/RayTracingSkyLightEvaluation.ush
@@ -18,6 +18,10 @@
 uint3 SkyLightVisibilityRaysDimensions;
 StructuredBuffer<FSkyLightVisibilityRays> SkyLightVisibilityRays;
 
+#ifndef USE_HAIR_LIGHTING
+#define USE_HAIR_LIGHTING 0
+#endif
+
 // Evaluates the radiance for the current sky light in a given direction
 float3 EvaluateSkyLightRadiance(float3 Direction)
 {
@@ -165,6 +169,13 @@ void SkyLightEvaluate(
 			InstanceInclusionMask,
 			Ray);
 
+		#if USE_HAIR_LIGHTING
+		if (GBufferData.ShadingModelID != SHADINGMODELID_HAIR)
+		{
+			MinimalPayload.HitT = TraverseHair(SampleCoord, Ray.Origin, Ray.Direction, MinimalPayload.HitT);
+		}
+		#endif
+
 		if (MinimalPayload.IsHit())
 		{
 			RayDistance += MinimalPayload.HitT;
diff --git a/Engine/Shaders/Private/RayTracing/RayTracingSkyLightRGS.usf b/Engine/Shaders/Private/RayTracing/RayTracingSkyLightRGS.usf
--- a/Engine/Shaders/Private/RayTracing/RayTracingSkyLightRGS.usf
+++ b/Engine/Shaders/Private/RayTracing/RayTracingSkyLightRGS.usf
@@ -1,8 +1,15 @@
 // Copyright Epic Games, Inc. All Rights Reserved.
 
+#include "../HairStrands/HairStrandsVoxelPageCommonStruct.ush"
+
 #include "../Common.ush"
 #include "../SceneTextureParameters.ush"
-#include "../SceneTextureParameters.ush"
+#include "../PathTracing/Utilities/PathTracingRandomSequence.ush"
+
+#if USE_HAIR_LIGHTING
+#include "../HairStrands/HairStrandsRaytracing.ush"
+#endif
+
 #include "RayTracingSkyLightEvaluation.ush"
 #include "RayTracingDeferredShadingCommon.ush"
 
diff --git a/Engine/Source/Runtime/Renderer/Private/DeferredShadingRenderer.cpp b/Engine/Source/Runtime/Renderer/Private/DeferredShadingRenderer.cpp
--- a/Engine/Source/Runtime/Renderer/Private/DeferredShadingRenderer.cpp
+++ b/Engine/Source/Runtime/Renderer/Private/DeferredShadingRenderer.cpp
@@ -1923,6 +1923,7 @@ void FDeferredShadingSceneRenderer::Render(FRHICommandListImmediate& RHICmdList)
 	{
 		if (bHairEnable)
 		{
+			RenderHairPrePass(RHICmdList, Scene, Views, HairClusterData, HairDatasStorage);
 			RenderHairBasePass(RHICmdList, Scene, SceneContext, Views, HairClusterData, HairDatasStorage);
 			HairDatas = &HairDatasStorage;
 		}
@@ -2238,6 +2239,13 @@ void FDeferredShadingSceneRenderer::Render(FRHICommandListImmediate& RHICmdList)
 		ServiceLocalQueue();
 	}
 
+	// Hair base pass for deferred shading
+	if (bHairEnable && !IsForwardShadingEnabled(ShaderPlatform))
+	{
+		RenderHairPrePass(RHICmdList, Scene, Views, HairClusterData, HairDatasStorage);
+		HairDatas = &HairDatasStorage;
+	}
+
 #if !UE_BUILD_SHIPPING
 	if (CVarForceBlackVelocityBuffer.GetValueOnRenderThread())
 	{
@@ -2253,7 +2261,7 @@ void FDeferredShadingSceneRenderer::Render(FRHICommandListImmediate& RHICmdList)
 	const bool bRayTracingEnabled = IsRayTracingEnabled();
 	if (bRayTracingEnabled && bCanOverlayRayTracingOutput && !IsForwardShadingEnabled(ShaderPlatform))
 	{		
-		RenderRayTracingSkyLight(RHICmdList, SkyLightRT, SkyLightHitDistanceRT);
+		RenderRayTracingSkyLight(RHICmdList, SkyLightRT, SkyLightHitDistanceRT, HairDatas);
 	}
 #endif // RHI_RAYTRACING
 	checkSlow(RHICmdList.IsOutsideRenderPass());
@@ -2332,10 +2340,10 @@ void FDeferredShadingSceneRenderer::Render(FRHICommandListImmediate& RHICmdList)
 	checkSlow(RHICmdList.IsOutsideRenderPass());
 
 	// Hair base pass for deferred shading
-	if (bHairEnable && !HairDatas)
+	if (bHairEnable && !IsForwardShadingEnabled(ShaderPlatform))
 	{
+		check(HairDatas);
 		RenderHairBasePass(RHICmdList, Scene, SceneContext, Views, HairClusterData, HairDatasStorage);
-		HairDatas = &HairDatasStorage;
 	}
 
 	// Render lighting.
diff --git a/Engine/Source/Runtime/Renderer/Private/DeferredShadingRenderer.h b/Engine/Source/Runtime/Renderer/Private/DeferredShadingRenderer.h
--- a/Engine/Source/Runtime/Renderer/Private/DeferredShadingRenderer.h
+++ b/Engine/Source/Runtime/Renderer/Private/DeferredShadingRenderer.h
@@ -595,7 +595,7 @@ private:
 
 	void GenerateSkyLightVisibilityRays(FRDGBuilder& GraphBuilder, FRDGBufferRef& SkyLightVisibilityRays, FIntVector& Dimensions);
 	void VisualizeSkyLightMipTree(FRHICommandListImmediate& RHICmdList, const FViewInfo& View, FRWBuffer& SkyLightMipTreePosX, FRWBuffer& SkyLightMipTreePosY, FRWBuffer& SkyLightMipTreePosZ, FRWBuffer& SkyLightMipTreeNegX, FRWBuffer& SkyLightMipTreeNegY, FRWBuffer& SkyLightMipTreeNegZ, const FIntVector& SkyLightMipDimensions);
-	void RenderRayTracingSkyLight(FRHICommandListImmediate& RHICmdList, TRefCountPtr<IPooledRenderTarget>& SkyLightRT, TRefCountPtr<IPooledRenderTarget>& HitDistanceRT);
+	void RenderRayTracingSkyLight(FRHICommandListImmediate& RHICmdList, TRefCountPtr<IPooledRenderTarget>& SkyLightRT, TRefCountPtr<IPooledRenderTarget>& HitDistanceRT, const struct FHairStrandsDatas* HairDatas);
 
 	void RenderRayTracingPrimaryRaysView(
 		FRDGBuilder& GraphBuilder,
diff --git a/Engine/Source/Runtime/Renderer/Private/HairStrands/HairStrandsRendering.cpp b/Engine/Source/Runtime/Renderer/Private/HairStrands/HairStrandsRendering.cpp
--- a/Engine/Source/Runtime/Renderer/Private/HairStrands/HairStrandsRendering.cpp
+++ b/Engine/Source/Runtime/Renderer/Private/HairStrands/HairStrandsRendering.cpp
@@ -5,10 +5,10 @@
 #include "ScenePrivate.h"
 
 void ServiceLocalQueue();
-void RenderHairBasePass(
+
+void RenderHairPrePass(
 	FRHICommandListImmediate& RHICmdList,
 	FScene* Scene,
-	FSceneRenderTargets& SceneContext,
 	TArray<FViewInfo>& Views,
 	FHairStrandClusterData HairClusterData,
 	FHairStrandsDatas& OutHairDatas)
@@ -35,8 +35,27 @@ void RenderHairBasePass(
 		RenderHairStrandsDeepShadows(RHICmdList, Scene, Views, OutHairDatas.MacroGroupsPerViews);
 
 		ServiceLocalQueue();
+	}
+}
+
+void RenderHairBasePass(
+	FRHICommandListImmediate& RHICmdList,
+	FScene* Scene,
+	FSceneRenderTargets& SceneContext,
+	TArray<FViewInfo>& Views,
+	FHairStrandClusterData HairClusterData,
+	FHairStrandsDatas& OutHairDatas)
+{
+	// #hair_todo: Add multi-view
+	const bool bIsViewCompatible = Views.Num() > 0 && Views[0].Family->ViewMode == VMI_Lit;
+	if (IsHairStrandsEnable(Scene->GetShaderPlatform()) && bIsViewCompatible)
+	{
+		const ERHIFeatureLevel::Type FeatureLevel = Scene->GetFeatureLevel();
 
+		//SCOPED_GPU_STAT(RHICmdList, HairRendering);
 		// Culling/LOD pass for visibility (must be done after HZB is generated)
+		FHairCullingParams CullingParams;
+		CullingParams.bCullingProcessSkipped = false;
 		CullingParams.bShadowViewMode = false;
 		ComputeHairStrandsClustersCulling(RHICmdList, *GetGlobalShaderMap(FeatureLevel), Views, CullingParams, HairClusterData);
 
diff --git a/Engine/Source/Runtime/Renderer/Private/HairStrands/HairStrandsRendering.h b/Engine/Source/Runtime/Renderer/Private/HairStrands/HairStrandsRendering.h
--- a/Engine/Source/Runtime/Renderer/Private/HairStrands/HairStrandsRendering.h
+++ b/Engine/Source/Runtime/Renderer/Private/HairStrands/HairStrandsRendering.h
@@ -45,6 +45,13 @@ void RunHairStrandsInterpolation(
 	EHairStrandsInterpolationType Type,
 	FHairStrandClusterData* ClusterData);
 
+void RenderHairPrePass(
+	FRHICommandListImmediate& RHICmdList,
+	FScene* Scene,
+	TArray<FViewInfo>& Views,
+	FHairStrandClusterData HairClusterData,
+	FHairStrandsDatas& OutHairDatas);
+
 void RenderHairBasePass(
 	FRHICommandListImmediate& RHICmdList,
 	FScene* Scene,
diff --git a/Engine/Source/Runtime/Renderer/Private/RayTracing/RaytracingSkylight.cpp b/Engine/Source/Runtime/Renderer/Private/RayTracing/RaytracingSkylight.cpp
--- a/Engine/Source/Runtime/Renderer/Private/RayTracing/RaytracingSkylight.cpp
+++ b/Engine/Source/Runtime/Renderer/Private/RayTracing/RaytracingSkylight.cpp
@@ -28,6 +28,7 @@ static int32 GRayTracingSkyLight = -1;
 #include "Raytracing/RaytracingOptions.h"
 #include "PostProcess/PostProcessing.h"
 #include "PostProcess/SceneFilterRendering.h"
+#include "HairStrands/HairStrandsRendering.h"
 
 static FAutoConsoleVariableRef CVarRayTracingSkyLight(
 	TEXT("r.RayTracing.SkyLight"),
@@ -91,6 +92,13 @@ static TAutoConsoleVariable<int32> CVarRayTracingSkyLightDecoupleSampleGeneratio
 	ECVF_RenderThreadSafe
 );
 
+static TAutoConsoleVariable<int32> CVarRayTracingSkyLightEnableHairVoxel(
+	TEXT("r.RayTracing.SkyLight.HairVoxel"),
+	0,
+	TEXT("Include hair voxel representation to estimate sky occlusion"),
+	ECVF_RenderThreadSafe
+);
+
 int32 GetRayTracingSkyLightDecoupleSampleGenerationCVarValue()
 {
 	return CVarRayTracingSkyLightDecoupleSampleGeneration.GetValueOnRenderThread();
@@ -310,8 +318,9 @@ class FRayTracingSkyLightRGS : public FGlobalShader
 	class FEnableTwoSidedGeometryDim : SHADER_PERMUTATION_BOOL("ENABLE_TWO_SIDED_GEOMETRY");
 	class FEnableMaterialsDim : SHADER_PERMUTATION_BOOL("ENABLE_MATERIALS");
 	class FDecoupleSampleGeneration : SHADER_PERMUTATION_BOOL("DECOUPLE_SAMPLE_GENERATION");
+	class FHairLighting : SHADER_PERMUTATION_INT("USE_HAIR_LIGHTING", 2);
 
-	using FPermutationDomain = TShaderPermutationDomain<FEnableTwoSidedGeometryDim, FEnableMaterialsDim, FDecoupleSampleGeneration>;
+	using FPermutationDomain = TShaderPermutationDomain<FEnableTwoSidedGeometryDim, FEnableMaterialsDim, FDecoupleSampleGeneration, FHairLighting>;
 
 	static bool ShouldCompilePermutation(const FGlobalShaderPermutationParameters& Parameters)
 	{
@@ -325,6 +334,7 @@ class FRayTracingSkyLightRGS : public FGlobalShader
 
 		SHADER_PARAMETER_STRUCT_REF(FViewUniformShaderParameters, ViewUniformBuffer)
 		SHADER_PARAMETER_STRUCT_REF(FSkyLightData, SkyLightData)
+		SHADER_PARAMETER_STRUCT_REF(FVirtualVoxelParameters, VirtualVoxel)
 		
 		SHADER_PARAMETER_STRUCT_INCLUDE(FSkyLightQuasiRandomData, SkyLightQuasiRandomData)
 		SHADER_PARAMETER_STRUCT_INCLUDE(FSkyLightVisibilityRaysData, SkyLightVisibilityRaysData)
@@ -433,7 +443,8 @@ DECLARE_GPU_STAT_NAMED(RayTracingSkyLight, TEXT("Ray Tracing SkyLight"));
 void FDeferredShadingSceneRenderer::RenderRayTracingSkyLight(
 	FRHICommandListImmediate& RHICmdList,
 	TRefCountPtr<IPooledRenderTarget>& SkyLightRT,
-	TRefCountPtr<IPooledRenderTarget>& HitDistanceRT
+	TRefCountPtr<IPooledRenderTarget>& HitDistanceRT,
+	const FHairStrandsDatas* HairDatas
 )
 {
 	SCOPED_DRAW_EVENT(RHICmdList, RayTracingSkyLight);
@@ -507,6 +518,7 @@ void FDeferredShadingSceneRenderer::RenderRayTracingSkyLight(
 		SubsurfaceProfileRT = GSystemTextures.BlackDummy;
 	}
 
+	int32 ViewIndex = 0;
 	for (FViewInfo& View : Views)
 	{
 		FSceneViewState* SceneViewState = (FSceneViewState*)View.State;
@@ -533,10 +545,20 @@ void FDeferredShadingSceneRenderer::RenderRayTracingSkyLight(
 		PassParameters->TLAS = View.RayTracingScene.RayTracingSceneRHI->GetShaderResourceView();
 		PassParameters->ViewUniformBuffer = View.ViewUniformBuffer;
 
+		const bool bUseHairLighting = 
+			HairDatas && ViewIndex < HairDatas->MacroGroupsPerViews.Views.Num() && 
+			HairDatas->MacroGroupsPerViews.Views[ViewIndex].VirtualVoxelResources.IsValid() && 
+			CVarRayTracingSkyLightEnableHairVoxel.GetValueOnRenderThread() > 0;
+		if (bUseHairLighting)
+		{
+			PassParameters->VirtualVoxel = HairDatas->MacroGroupsPerViews.Views[ViewIndex].VirtualVoxelResources.UniformBuffer;
+		}
+
 		FRayTracingSkyLightRGS::FPermutationDomain PermutationVector;
 		PermutationVector.Set<FRayTracingSkyLightRGS::FEnableTwoSidedGeometryDim>(CVarRayTracingSkyLightEnableTwoSidedGeometry.GetValueOnRenderThread() != 0);
 		PermutationVector.Set<FRayTracingSkyLightRGS::FEnableMaterialsDim>(CVarRayTracingSkyLightEnableMaterials.GetValueOnRenderThread() != 0);
 		PermutationVector.Set<FRayTracingSkyLightRGS::FDecoupleSampleGeneration>(CVarRayTracingSkyLightDecoupleSampleGeneration.GetValueOnRenderThread() != 0);
+		PermutationVector.Set<FRayTracingSkyLightRGS::FHairLighting>(bUseHairLighting ? 1 : 0);
 		TShaderMapRef<FRayTracingSkyLightRGS> RayGenerationShader(GetGlobalShaderMap(FeatureLevel), PermutationVector);
 		ClearUnusedGraphResources(RayGenerationShader, PassParameters);
 
@@ -617,6 +639,8 @@ void FDeferredShadingSceneRenderer::RenderRayTracingSkyLight(
 				SceneViewState->SkyLightVisibilityRaysDimensions = FIntVector(1);
 			}
 		}
+
+		++ViewIndex;
 	}
 
 	GraphBuilder.QueueTextureExtraction(SkyLightTexture, &SkyLightRT);
