Fix a bug that the IndirectIrradiance is 0 if there is only a dynamic skylight in the scene and the metal material will be black.

diff --git a/Engine/Shaders/Private/MobileBasePassPixelShader.usf b/Engine/Shaders/Private/MobileBasePassPixelShader.usf
--- a/Engine/Shaders/Private/MobileBasePassPixelShader.usf
+++ b/Engine/Shaders/Private/MobileBasePassPixelShader.usf
@@ -545,6 +545,13 @@ void Main(
 		Color += ShadingModelContext.DiffuseColor * DiffuseGI * View.IndirectLightingColorScale;
 	#endif
 #endif
+
+#if ENABLE_SKY_LIGHT
+		half3 SkyDiffuseLighting = GetSkySHDiffuseSimple(MaterialParameters.WorldNormal);
+		half3 DiffuseLookup = SkyDiffuseLighting * ResolvedView.SkyLightColor.rgb;
+		IndirectIrradiance += Luminance(DiffuseLookup);
+#endif
+
 	half MaterialAO = GetMaterialAmbientOcclusion(PixelMaterialInputs);
 	Color *= MaterialAO;
 	IndirectIrradiance *= MaterialAO;
@@ -728,7 +735,6 @@ void Main(
 	if (NoL == 0)
 	{
 	#endif
-		half3 SkyDiffuseLighting = GetSkySHDiffuseSimple(MaterialParameters.WorldNormal);
 
 	#if MATERIAL_SHADINGMODEL_SINGLELAYERWATER
 		ShadingModelContext.WaterDiffuseIndirectLuminance += SkyDiffuseLighting;
