Roll and merge loops in volumetric fog compute shader

- unrolled loop for N=16 can exceed maximum branch length on unoptimized shaders for some DXC variants
  - roll the loop to address this and remerge with the following loop as there is no longer a reason for the split

diff --git a/Engine/Shaders/Private/VolumetricFog.usf b/Engine/Shaders/Private/VolumetricFog.usf
--- a/Engine/Shaders/Private/VolumetricFog.usf
+++ b/Engine/Shaders/Private/VolumetricFog.usf
@@ -467,177 +467,161 @@ void LightScatteringCS(
 	// Supersample if the history was outside the camera frustum
 	// The compute shader is dispatched with extra threads, make sure those don't supersample
 	NumSuperSamples = HistoryAlpha < .001f && all(GridCoordinate < VolumetricFog.GridSizeInt) ? HISTORY_MISS_SUPER_SAMPLE_COUNT : 1;
-
 #endif
+
+	for (uint SampleIndex = 0; SampleIndex < NumSuperSamples; SampleIndex++)
 	{
-		UNROLL
-		for (uint SampleIndex = 0; SampleIndex < NumSuperSamples; SampleIndex++)
-		{
-			float3 CellOffset = FrameJitterOffsets[SampleIndex].xyz;
-			//float CellOffset = .5f;
+		float3 CellOffset = FrameJitterOffsets[SampleIndex].xyz;
+		//float CellOffset = .5f;
 
-			float SceneDepth;
-			float3 WorldPosition = ComputeCellWorldPosition(GridCoordinate, CellOffset, SceneDepth);
-			float CameraVectorLength = length(WorldPosition - View.WorldCameraOrigin);
-			float3 CameraVector = (WorldPosition - View.WorldCameraOrigin) / CameraVectorLength;
+		float SceneDepth;
+		float3 WorldPosition = ComputeCellWorldPosition(GridCoordinate, CellOffset, SceneDepth);
+		float CameraVectorLength = length(WorldPosition - View.WorldCameraOrigin);
+		float3 CameraVector = (WorldPosition - View.WorldCameraOrigin) / CameraVectorLength;
 
-			BRANCH
-			if (ForwardLightData.HasDirectionalLight)
-			{
-				float ShadowFactor = 1;
+		BRANCH
+		if (ForwardLightData.HasDirectionalLight)
+		{
+			float ShadowFactor = 1;
 			
-				if (UseDirectionalLightShadowing > 0)
-				{
-					ShadowFactor *= ComputeDirectionalLightStaticShadowing(WorldPosition);
-					ShadowFactor *= ComputeDirectionalLightDynamicShadowing(WorldPosition, SceneDepth);
-				}
+			if (UseDirectionalLightShadowing > 0)
+			{
+				ShadowFactor *= ComputeDirectionalLightStaticShadowing(WorldPosition);
+				ShadowFactor *= ComputeDirectionalLightDynamicShadowing(WorldPosition, SceneDepth);
+			}
 
-				ShadowFactor *= GetLightFunction(WorldPosition);
+			ShadowFactor *= GetLightFunction(WorldPosition);
 
-				float3 DirectionalLightColor = ForwardLightData.DirectionalLightColor;
+			float3 DirectionalLightColor = ForwardLightData.DirectionalLightColor;
 
-				if (UseHeightFogColors.x > 0)
-				{
-					// Attempt to maintain intensity ratio between sky and sun
-					DirectionalLightColor = VolumetricFog.HeightFogDirectionalLightInscatteringColor * Luminance(ForwardLightData.DirectionalLightColor);
-				}
-
-				LightScattering += DirectionalLightColor
-					* (ShadowFactor 
-					* ForwardLightData.DirectionalLightVolumetricScatteringIntensity 
-					* PhaseFunction(PhaseG, dot(ForwardLightData.DirectionalLightDirection, -CameraVector)));
+			if (UseHeightFogColors.x > 0)
+			{
+				// Attempt to maintain intensity ratio between sky and sun
+				DirectionalLightColor = VolumetricFog.HeightFogDirectionalLightInscatteringColor * Luminance(ForwardLightData.DirectionalLightColor);
 			}
 
-			FTwoBandSHVector RotatedHGZonalHarmonic;
-			RotatedHGZonalHarmonic.V = float4(1.0f, CameraVector.y, CameraVector.z, CameraVector.x) * float4(1.0f, PhaseG, PhaseG, PhaseG);
+			LightScattering += DirectionalLightColor
+				* (ShadowFactor 
+				* ForwardLightData.DirectionalLightVolumetricScatteringIntensity 
+				* PhaseFunction(PhaseG, dot(ForwardLightData.DirectionalLightDirection, -CameraVector)));
+		}
+
+		FTwoBandSHVector RotatedHGZonalHarmonic;
+		RotatedHGZonalHarmonic.V = float4(1.0f, CameraVector.y, CameraVector.z, CameraVector.x) * float4(1.0f, PhaseG, PhaseG, PhaseG);
 
-			float3 BrickTextureUVs = 0;
+		float3 BrickTextureUVs = 0;
 
 #if ALLOW_STATIC_LIGHTING
-			if (SkyLightVolumetricScatteringIntensity > 0 || StaticLightingScatteringIntensity > 0)
-			{
-				BrickTextureUVs = ComputeVolumetricLightmapBrickTextureUVs(WorldPosition);
-			}
+		if (SkyLightVolumetricScatteringIntensity > 0 || StaticLightingScatteringIntensity > 0)
+		{
+			BrickTextureUVs = ComputeVolumetricLightmapBrickTextureUVs(WorldPosition);
+		}
 #endif
 
-			// Skylight
-			if (SkyLightVolumetricScatteringIntensity > 0)
-			{
-				float3 SkyLighting;
-
-				if (UseHeightFogColors.y > 0)
-				{
-					float3 HeightFogInscatteringColor = ComputeInscatteringColor(CameraVector, CameraVectorLength);
-					float ScalarFactor = SHAmbientFunction();
-					FTwoBandSHVectorRGB SkyIrradianceSH;
-					SkyIrradianceSH.R.V = float4(ScalarFactor * HeightFogInscatteringColor.r, 0, 0, 0);
-					SkyIrradianceSH.G.V = float4(ScalarFactor * HeightFogInscatteringColor.g, 0, 0, 0);
-					SkyIrradianceSH.B.V = float4(ScalarFactor * HeightFogInscatteringColor.b, 0, 0, 0);
-
-					SkyLighting = max(DotSH(SkyIrradianceSH, RotatedHGZonalHarmonic), 0);
-				}
-				else
-				{
-					FTwoBandSHVectorRGB SkyIrradianceSH;
-					SkyIrradianceSH.R.V = SkySH[0];
-					SkyIrradianceSH.G.V = SkySH[1];
-					SkyIrradianceSH.B.V = SkySH[2];
-
-					SkyLighting = View.SkyLightColor.rgb * max(DotSH(SkyIrradianceSH, RotatedHGZonalHarmonic), 0) / PI;
-				}
+		// Skylight
+		if (SkyLightVolumetricScatteringIntensity > 0)
+		{
+			float3 SkyLighting;
 
-				float SkyVisibility = ComputeSkyVisibility(WorldPosition, BrickTextureUVs);
-				LightScattering += (SkyVisibility * SkyLightVolumetricScatteringIntensity) * SkyLighting;
+			if (UseHeightFogColors.y > 0)
+			{
+				float3 HeightFogInscatteringColor = ComputeInscatteringColor(CameraVector, CameraVectorLength);
+				float ScalarFactor = SHAmbientFunction();
+				FTwoBandSHVectorRGB SkyIrradianceSH;
+				SkyIrradianceSH.R.V = float4(ScalarFactor * HeightFogInscatteringColor.r, 0, 0, 0);
+				SkyIrradianceSH.G.V = float4(ScalarFactor * HeightFogInscatteringColor.g, 0, 0, 0);
+				SkyIrradianceSH.B.V = float4(ScalarFactor * HeightFogInscatteringColor.b, 0, 0, 0);
+
+				SkyLighting = max(DotSH(SkyIrradianceSH, RotatedHGZonalHarmonic), 0);
 			}
-
-#if ALLOW_STATIC_LIGHTING
-			// Indirect lighting of Stationary lights and Direct + Indirect lighting of Static lights
-			if (StaticLightingScatteringIntensity > 0)
+			else
 			{
-				FTwoBandSHVectorRGB IrradianceSH = GetVolumetricLightmapSH2(BrickTextureUVs);
+				FTwoBandSHVectorRGB SkyIrradianceSH;
+				SkyIrradianceSH.R.V = SkySH[0];
+				SkyIrradianceSH.G.V = SkySH[1];
+				SkyIrradianceSH.B.V = SkySH[2];
 
-				LightScattering += (StaticLightingScatteringIntensity / PI) * max(DotSH(IrradianceSH, RotatedHGZonalHarmonic), 0);
+				SkyLighting = View.SkyLightColor.rgb * max(DotSH(SkyIrradianceSH, RotatedHGZonalHarmonic), 0) / PI;
 			}
-#endif
+
+			float SkyVisibility = ComputeSkyVisibility(WorldPosition, BrickTextureUVs);
+			LightScattering += (SkyVisibility * SkyLightVolumetricScatteringIntensity) * SkyLighting;
 		}
-	}
 
-	{
-		for (uint SampleIndex = 0; SampleIndex < NumSuperSamples; SampleIndex++)
+#if ALLOW_STATIC_LIGHTING
+		// Indirect lighting of Stationary lights and Direct + Indirect lighting of Static lights
+		if (StaticLightingScatteringIntensity > 0)
 		{
-			float3 CellOffset = FrameJitterOffsets[SampleIndex].xyz;
-			//float CellOffset = .5f;
+			FTwoBandSHVectorRGB IrradianceSH = GetVolumetricLightmapSH2(BrickTextureUVs);
 
-			float SceneDepth;
-			float3 WorldPosition = ComputeCellWorldPosition(GridCoordinate, CellOffset, SceneDepth);
-			float CameraVectorLength = length(WorldPosition - View.WorldCameraOrigin);
-			float3 CameraVector = (WorldPosition - View.WorldCameraOrigin) / CameraVectorLength;
+			LightScattering += (StaticLightingScatteringIntensity / PI) * max(DotSH(IrradianceSH, RotatedHGZonalHarmonic), 0);
+		}
+#endif
 
-			uint GridIndex = ComputeLightGridCellIndex(GridCoordinate.xy * VolumetricFog.FogGridToPixelXY, SceneDepth, 0);
-			const FCulledLightsGridData CulledLightsGrid = GetCulledLightsGrid(GridIndex, 0);
+		uint GridIndex = ComputeLightGridCellIndex(GridCoordinate.xy * VolumetricFog.FogGridToPixelXY, SceneDepth, 0);
+		const FCulledLightsGridData CulledLightsGrid = GetCulledLightsGrid(GridIndex, 0);
 
-			float CellRadius = length(WorldPosition - ComputeCellWorldPosition(GridCoordinate + uint3(1, 1, 1), CellOffset));
-			// Bias the inverse squared light falloff based on voxel size to prevent aliasing near the light source
-			float DistanceBiasSqr = max(CellRadius * InverseSquaredLightDistanceBiasScale, 1);
-			DistanceBiasSqr *= DistanceBiasSqr;
-			FRectTexture RectTexture = InitRectTexture(DeferredLightUniforms.SourceTexture);
+		float CellRadius = length(WorldPosition - ComputeCellWorldPosition(GridCoordinate + uint3(1, 1, 1), CellOffset));
+		// Bias the inverse squared light falloff based on voxel size to prevent aliasing near the light source
+		float DistanceBiasSqr = max(CellRadius * InverseSquaredLightDistanceBiasScale, 1);
+		DistanceBiasSqr *= DistanceBiasSqr;
+		FRectTexture RectTexture = InitRectTexture(DeferredLightUniforms.SourceTexture);
 
-			// Forward lighting of unshadowed point and spot lights
-			LOOP
-			for (uint LocalLightListIndex = 0; LocalLightListIndex < CulledLightsGrid.NumLocalLights; LocalLightListIndex++)
-			{
-				const FLocalLightData LocalLight = GetLocalLightData(CulledLightsGrid.DataStartIndex + LocalLightListIndex, 0);
+		// Forward lighting of unshadowed point and spot lights
+		LOOP
+		for (uint LocalLightListIndex = 0; LocalLightListIndex < CulledLightsGrid.NumLocalLights; LocalLightListIndex++)
+		{
+			const FLocalLightData LocalLight = GetLocalLightData(CulledLightsGrid.DataStartIndex + LocalLightListIndex, 0);
+
+			float VolumetricScatteringIntensity = f16tof32(asuint(LocalLight.SpotAnglesAndSourceRadiusPacked.w) >> 16);
 
-				float VolumetricScatteringIntensity = f16tof32(asuint(LocalLight.SpotAnglesAndSourceRadiusPacked.w) >> 16);
+			if (VolumetricScatteringIntensity > 0)
+			{
+				FDeferredLightData LightData = (FDeferredLightData)0;
+				LightData.Position = LocalLight.LightPositionAndInvRadius.xyz;
+				LightData.InvRadius = LocalLight.LightPositionAndInvRadius.w;
+				LightData.Color = LocalLight.LightColorAndFalloffExponent.xyz;
+				LightData.FalloffExponent = LocalLight.LightColorAndFalloffExponent.w;
+				LightData.Direction = LocalLight.LightDirectionAndShadowMask.xyz;
+				LightData.SpotAngles = LocalLight.SpotAnglesAndSourceRadiusPacked.xy;
+				LightData.SourceRadius = LocalLight.SpotAnglesAndSourceRadiusPacked.z;
+				LightData.SourceLength = f16tof32(asuint(LocalLight.SpotAnglesAndSourceRadiusPacked.w));
+				LightData.Tangent = LocalLight.LightTangentAndSoftSourceRadius.xyz;
+				LightData.SoftSourceRadius = LocalLight.LightTangentAndSoftSourceRadius.w;
+				LightData.bInverseSquared = LightData.FalloffExponent == 0;
+				LightData.RectLightBarnCosAngle = LocalLight.RectBarnDoor.x;
+				LightData.RectLightBarnLength = LocalLight.RectBarnDoor.y;
+				LightData.bRadialLight = true;
+				LightData.bSpotLight = LightData.SpotAngles.x > -2.0f;
+				LightData.bRectLight = LightData.RectLightBarnLength > -2.0f;
+
+				float3 LightColor = LightData.Color;
+
+				float3 L = 0;
+				float3 ToLight = 0;
+				float LightMask = GetLocalLightAttenuation(WorldPosition, LightData, ToLight, L);
+
+				float Lighting;
+				if( LightData.bRectLight )
+				{
+					FRect Rect = GetRect( ToLight, LightData );
 
-				if (VolumetricScatteringIntensity > 0)
+					Lighting = IntegrateLight(Rect, RectTexture);
+				}
+				else
 				{
-					FDeferredLightData LightData = (FDeferredLightData)0;
-					LightData.Position = LocalLight.LightPositionAndInvRadius.xyz;
-					LightData.InvRadius = LocalLight.LightPositionAndInvRadius.w;
-					LightData.Color = LocalLight.LightColorAndFalloffExponent.xyz;
-					LightData.FalloffExponent = LocalLight.LightColorAndFalloffExponent.w;
-					LightData.Direction = LocalLight.LightDirectionAndShadowMask.xyz;
-					LightData.SpotAngles = LocalLight.SpotAnglesAndSourceRadiusPacked.xy;
-					LightData.SourceRadius = LocalLight.SpotAnglesAndSourceRadiusPacked.z;
-					LightData.SourceLength = f16tof32(asuint(LocalLight.SpotAnglesAndSourceRadiusPacked.w));
-					LightData.Tangent = LocalLight.LightTangentAndSoftSourceRadius.xyz;
-					LightData.SoftSourceRadius = LocalLight.LightTangentAndSoftSourceRadius.w;
-					LightData.bInverseSquared = LightData.FalloffExponent == 0;
-					LightData.RectLightBarnCosAngle = LocalLight.RectBarnDoor.x;
-					LightData.RectLightBarnLength = LocalLight.RectBarnDoor.y;
-					LightData.bRadialLight = true;
-					LightData.bSpotLight = LightData.SpotAngles.x > -2.0f;
-					LightData.bRectLight = LightData.RectLightBarnLength > -2.0f;
-
-					float3 LightColor = LightData.Color;
-
-					float3 L = 0;
-					float3 ToLight = 0;
-					float LightMask = GetLocalLightAttenuation(WorldPosition, LightData, ToLight, L);
-
-					float Lighting;
-					if( LightData.bRectLight )
-					{
-						FRect Rect = GetRect( ToLight, LightData );
-
-						Lighting = IntegrateLight(Rect, RectTexture);
-					}
-					else
-					{
-						FCapsuleLight Capsule = GetCapsule(ToLight, LightData);
-						Capsule.DistBiasSqr = DistanceBiasSqr;
-
-						Lighting = IntegrateLight(Capsule, LightData.bInverseSquared);
-					}
-
-					float CombinedAttenuation = Lighting * LightMask;
-
-					LightScattering += LightColor * (PhaseFunction(PhaseG, dot(L, -CameraVector)) * CombinedAttenuation * VolumetricScatteringIntensity);
-
-					// To debug culling
-					//LightScattering += LocalLight.LightColorAndFalloffExponent.xyz * .0000001f;
+					FCapsuleLight Capsule = GetCapsule(ToLight, LightData);
+					Capsule.DistBiasSqr = DistanceBiasSqr;
+
+					Lighting = IntegrateLight(Capsule, LightData.bInverseSquared);
 				}
+
+				float CombinedAttenuation = Lighting * LightMask;
+
+				LightScattering += LightColor * (PhaseFunction(PhaseG, dot(L, -CameraVector)) * CombinedAttenuation * VolumetricScatteringIntensity);
+
+				// To debug culling
+				//LightScattering += LocalLight.LightColorAndFalloffExponent.xyz * .0000001f;
 			}
 		}
 	}
