Fix hair interpolation for 4.25

diff --git a/Engine/Plugins/Experimental/HairStrands/Source/HairStrandsCore/Private/GroomComponent.cpp b/Engine/Plugins/Experimental/HairStrands/Source/HairStrandsCore/Private/GroomComponent.cpp
--- a/Engine/Plugins/Experimental/HairStrands/Source/HairStrandsCore/Private/GroomComponent.cpp
+++ b/Engine/Plugins/Experimental/HairStrands/Source/HairStrandsCore/Private/GroomComponent.cpp
@@ -911,6 +911,10 @@ void UGroomComponent::UpdateSimulatedGroups()
 				const bool bIsSimulationEnable = (LocalGroomAsset && GroupIt < LocalGroomAsset->HairGroupsPhysics.Num()) ? 
 					LocalGroomAsset->HairGroupsPhysics[GroupIt].SolverSettings.EnableSimulation : false;
 				const bool bHasGlobalInterpolation = (LocalBindingAsset && LocalGroomAsset->EnableGlobalInterpolation);
+				HairGroup.HairInterpolationType =
+					(LocalGroomAsset && LocalGroomAsset->HairInterpolationType == EGroomInterpolationType::RigidTransform) ? 0 :
+					(LocalGroomAsset && LocalGroomAsset->HairInterpolationType == EGroomInterpolationType::OffsetTransform) ? 1 :
+					(LocalGroomAsset && LocalGroomAsset->HairInterpolationType == EGroomInterpolationType::SmoothTransform) ? 2 : 0;
 				HairGroup.bIsSimulationEnable = bIsSimulationEnable;
 				HairGroup.bHasGlobalInterpolation = bHasGlobalInterpolation;
 				UpdateHairStrandsDebugInfo(Id, WorldType, GroupIt, bIsSimulationEnable);
@@ -1070,6 +1074,11 @@ void UGroomComponent::InitResources(bool bIsBindingReloading)
 
 		const bool bHasGlobalInterpolation = (BindingAsset && GroomAsset->EnableGlobalInterpolation);
 
+		InterpolationInputGroup.HairInterpolationType =
+			(GroomAsset && GroomAsset->HairInterpolationType == EGroomInterpolationType::RigidTransform) ? 0 :
+			(GroomAsset && GroomAsset->HairInterpolationType == EGroomInterpolationType::OffsetTransform) ? 1 :
+			(GroomAsset && GroomAsset->HairInterpolationType == EGroomInterpolationType::SmoothTransform) ? 2 : 0;
+
 		// For skinned groom, these value will be updated during TickComponent() call
 		// Deformed sim & render are expressed within the referential (unlike rest pose)
 		InterpolationInputGroup.OutHairPositionOffset = RenderRestHairPositionOffset;
diff --git a/Engine/Plugins/Experimental/HairStrands/Source/HairStrandsCore/Private/HairStrandsRendering.cpp b/Engine/Plugins/Experimental/HairStrands/Source/HairStrandsCore/Private/HairStrandsRendering.cpp
--- a/Engine/Plugins/Experimental/HairStrands/Source/HairStrandsCore/Private/HairStrandsRendering.cpp
+++ b/Engine/Plugins/Experimental/HairStrands/Source/HairStrandsCore/Private/HairStrandsRendering.cpp
@@ -278,7 +278,7 @@ class FHairInterpolationCS : public FGlobalShader
 
 	class FGroupSize : SHADER_PERMUTATION_INT("PERMUTATION_GROUP_SIZE", 2);
 	class FDebug : SHADER_PERMUTATION_INT("PERMUTATION_DEBUG", 2);
-	class FDynamicGeometry : SHADER_PERMUTATION_INT("PERMUTATION_DYNAMIC_GEOMETRY", 3);
+	class FDynamicGeometry : SHADER_PERMUTATION_INT("PERMUTATION_DYNAMIC_GEOMETRY", 5);
 	class FSimulation : SHADER_PERMUTATION_INT("PERMUTATION_SIMULATION", 2);
 	class FScaleAndClip : SHADER_PERMUTATION_INT("PERMUTATION_SCALE_AND_CLIP", 2);
 	using FPermutationDomain = TShaderPermutationDomain<FGroupSize, FDebug, FDynamicGeometry, FSimulation, FScaleAndClip>;
@@ -374,7 +374,8 @@ static void AddHairStrandsInterpolationPass(
 	const FShaderResourceViewRHIRef& VertexToClusterIdBuffer,
 	const FShaderResourceViewRHIRef& SimRootPointIndexBuffer,
 	FBufferTransitionQueue& OutTransitionQueue,
-	const bool bHasGlobalInterpolation)
+	const bool bHasGlobalInterpolation,
+	const uint32 HairInterpolationType)
 {
 	const uint32 GroupSize = ComputeGroupSize();
 	const FIntVector DispatchCount = ComputeDispatchCount(VertexCount, GroupSize);
@@ -507,8 +508,8 @@ static void AddHairStrandsInterpolationPass(
 	FHairInterpolationCS::FPermutationDomain PermutationVector;
 	PermutationVector.Set<FHairInterpolationCS::FGroupSize>(GetGroupSizePermutation(GroupSize));
 	PermutationVector.Set<FHairInterpolationCS::FDebug>(Parameters->HairDebugMode > 0 ? 1 : 0);
-	PermutationVector.Set<FHairInterpolationCS::FDynamicGeometry>((bSupportDynamicMesh && bSupportGlobalInterpolation) ? 2 :
-							(bSupportDynamicMesh && !bSupportGlobalInterpolation) ? 1 : 0);
+	PermutationVector.Set<FHairInterpolationCS::FDynamicGeometry>((bSupportDynamicMesh && bHasLocalDeformation) ? HairInterpolationType+1 :
+							(bSupportDynamicMesh && !bHasLocalDeformation) ? 1 : 0);
 	PermutationVector.Set<FHairInterpolationCS::FSimulation>(bHasLocalDeformation ? 1 : 0);
 	PermutationVector.Set<FHairInterpolationCS::FScaleAndClip>(bNeedScaleOrClip ? 1 : 0);
 
@@ -1014,7 +1015,8 @@ void ComputeHairStrandsInterpolation(
 					Input.VertexToClusterIdBuffer->SRV,
 					Input.SimRootPointIndexBuffer->SRV,
 					TransitionQueue,
-					Input.bHasGlobalInterpolation);
+					Input.bHasGlobalInterpolation,
+					Input.HairInterpolationType);
 
 			}
 			GraphBuilder.Execute();
diff --git a/Engine/Plugins/Experimental/HairStrands/Source/HairStrandsCore/Public/GroomAsset.h b/Engine/Plugins/Experimental/HairStrands/Source/HairStrandsCore/Public/GroomAsset.h
--- a/Engine/Plugins/Experimental/HairStrands/Source/HairStrandsCore/Public/GroomAsset.h
+++ b/Engine/Plugins/Experimental/HairStrands/Source/HairStrandsCore/Public/GroomAsset.h
@@ -17,6 +17,15 @@
 
 #include "GroomAsset.generated.h"
 
+/** List of interpolation type */
+UENUM(BlueprintType)
+enum class EGroomInterpolationType : uint8
+{
+	None = 0 UMETA(Hidden),
+	RigidTransform = 0x02 UMETA(DisplatName = "RigidTransform"),
+	OffsetTransform = 0x04 UMETA(DisplatName = "OffsetTransform"),
+	SmoothTransform = 0x08 UMETA(DisplatName = "SmoothTransform")
+};
 
 /** List of niagara solvers */
 UENUM(BlueprintType)
@@ -635,6 +644,10 @@ public:
 	UPROPERTY(EditAnywhere, EditFixedSize, BlueprintReadWrite, Category = "HairInterpolation", meta = (ToolTip = "Enable radial basis function interpolation to be used instead of the local skin rigid transform (WIP)"))
 	bool EnableGlobalInterpolation = false;
 
+	/** Type of interpolation used */
+	UPROPERTY(EditAnywhere, EditFixedSize, BlueprintReadWrite, Category = "HairInterpolation", meta = (ToolTip = "Type of interpolation used (WIP)"))
+	EGroomInterpolationType HairInterpolationType = EGroomInterpolationType::RigidTransform;
+
 	TArray<FHairGroupData> HairGroupsData;
 
 	//~ Begin UObject Interface.
diff --git a/Engine/Plugins/Experimental/HairStrands/Source/HairStrandsCore/Public/HairStrandsRendering.h b/Engine/Plugins/Experimental/HairStrands/Source/HairStrandsCore/Public/HairStrandsRendering.h
--- a/Engine/Plugins/Experimental/HairStrands/Source/HairStrandsCore/Public/HairStrandsRendering.h
+++ b/Engine/Plugins/Experimental/HairStrands/Source/HairStrandsCore/Public/HairStrandsRendering.h
@@ -42,6 +42,7 @@ struct FHairStrandsInterpolationInput
 
 		bool bIsSimulationEnable = false;
 		bool bHasGlobalInterpolation = false;
+		uint32 HairInterpolationType = 0;
 		
 		FHairGroupDesc GroupDesc;
 
diff --git a/Engine/Shaders/Private/HairStrands/HairStrandsInterpolation.usf b/Engine/Shaders/Private/HairStrands/HairStrandsInterpolation.usf
--- a/Engine/Shaders/Private/HairStrands/HairStrandsInterpolation.usf
+++ b/Engine/Shaders/Private/HairStrands/HairStrandsInterpolation.usf
@@ -37,14 +37,14 @@ Buffer<uint>   SimRootPointIndexBuffer;
 Buffer<uint4>  AttributeBuffer;
 #endif
 
-#if PERMUTATION_DYNAMIC_GEOMETRY == 1 || PERMUTATION_DEBUG == 1 || PERMUTATION_DYNAMIC_GEOMETRY == 2
+#if PERMUTATION_DYNAMIC_GEOMETRY >= 1 || PERMUTATION_DEBUG == 1 
 Buffer<uint>   RenVertexToRootIndexBuffer;
 Buffer<uint>   SimVertexToRootIndexBuffer;
 #endif
 
 RWBuffer<uint4> OutRenderDeformedPositionBuffer;
 
-#if PERMUTATION_DYNAMIC_GEOMETRY == 1 || PERMUTATION_DYNAMIC_GEOMETRY == 2
+#if PERMUTATION_DYNAMIC_GEOMETRY >= 1 
 
 // Compact all these buffers into 2 buffers: translation + quaternion
 Buffer<float4>  RestPosition0Buffer;
@@ -117,6 +117,32 @@ float3 ComputeDynamicGeometryOffset(
 	return (LocalDeformedGuidePoint - LocalRestGuidePoint) * GuideVertexWeight;
 }
 
+float3 ComputeTranslateGeometryOffset(uint GuideIndex, float GuideVertexWeight)
+{	
+	const float3 RestGuidePoint = f16tof32(SimRestPosePositionBuffer[GuideIndex].xyz) - 
+					f16tof32(SimRestPosePositionBuffer[SimRootPointIndexBuffer[GuideIndex]].xyz) ;
+
+	const float3 DeformedGuidePoint = f16tof32(DeformedSimPositionBuffer[GuideIndex].xyz) - 
+					f16tof32(DeformedSimPositionBuffer[SimRootPointIndexBuffer[GuideIndex]].xyz) ;
+
+	return (DeformedGuidePoint - RestGuidePoint) * GuideVertexWeight;
+}
+
+float3 ComputeRotateGeometryOffset(uint GuideIndex, float GuideVertexWeight,
+	FHairMeshTriangle RestTriangle, 
+	FHairMeshTriangle DeformedTriangle)
+{	
+	const float3 RestGuidePoint = f16tof32(SimRestPosePositionBuffer[GuideIndex].xyz) - 
+					f16tof32(SimRestPosePositionBuffer[SimRootPointIndexBuffer[GuideIndex]].xyz) ;
+	const float3 LocalRestGuidePoint = VectorToTriangle(RestGuidePoint, RestTriangle);
+
+	const float3 DeformedGuidePoint = f16tof32(DeformedSimPositionBuffer[GuideIndex].xyz) - 
+					f16tof32(DeformedSimPositionBuffer[SimRootPointIndexBuffer[GuideIndex]].xyz) ;
+	const float3 LocalDeformedGuidePoint = VectorToTriangle(DeformedGuidePoint, DeformedTriangle);
+
+	return (LocalDeformedGuidePoint - LocalRestGuidePoint) * GuideVertexWeight;
+}
+
 // Compute the new hair width and return if the current vertex is valid (true), or if it needs to be clipped (false)
 bool ScaleAndClip(uint ControlPointId, inout uint InPackedAlpha)
 {
@@ -278,14 +304,20 @@ void MainCS(uint2 DispatchThreadId : SV_DispatchThreadID)
 		}
 		#endif	
 
-		// 3. Skin transformation
-		// Apply dynamic mesh deformation (translation)
+		// 3. Linear blend skinning (translation)
 		#if PERMUTATION_DYNAMIC_GEOMETRY == 2
 		{
 			uint4 PackedRestPosition = RenderRestPosePositionBuffer[VertexIndex];
 			const float3 RestControlPoint = f16tof32(PackedRestPosition.xyz) + InRenderHairPositionOffset;
 			float3 ControlPoint = RestControlPoint;
 
+			const uint RootIndex						= RenVertexToRootIndexBuffer[VertexIndex];
+			const float3 RootBarycentric				= DecodeBarycentrics(RootBarycentricBuffer[RootIndex]);
+			const FHairMeshTriangle RestTriangle		= GetTriangleTransformation(RootIndex, RestPosition0Buffer, RestPosition1Buffer, RestPosition2Buffer);
+			const FHairMeshTriangle DeformedTriangle	= GetTriangleTransformation(RootIndex, DeformedPosition0Buffer, DeformedPosition1Buffer, DeformedPosition2Buffer);
+			
+			ControlPoint -= RootBarycentric.x * RestTriangle.P0 + RootBarycentric.y * RestTriangle.P1 + RootBarycentric.z * RestTriangle.P2;
+			
 			float3 CurrOffset = 0;
 
 			#if PERMUTATION_SIMULATION == 1
@@ -295,27 +327,108 @@ void MainCS(uint2 DispatchThreadId : SV_DispatchThreadID)
 				const uint GuideIndex0 = GuideVertexIndices[K];
 				const uint GuideIndex1 = GuideIndex0 + 1;
 
-				const float3 Offset0 = ComputeLocalGeometryOffset(GuideIndex0, GuideVertexWeights[K]);
-				const float3 Offset1 = ComputeLocalGeometryOffset(GuideIndex1, GuideVertexWeights[K]);
+				const float3 Offset0 = ComputeTranslateGeometryOffset(GuideIndex0, 1.0);
+				const float3 Offset1 = ComputeTranslateGeometryOffset(GuideIndex1, 1.0);
 
-				CurrOffset += lerp(Offset0, Offset1, GuideVertexLerps[K]);
+				CurrOffset += GuideVertexWeights[K] * lerp(Offset0, Offset1, GuideVertexLerps[K]);
 			}
 			#endif
 
+			ControlPoint += RootBarycentric.x * DeformedTriangle.P0 + RootBarycentric.y * DeformedTriangle.P1 + RootBarycentric.z * DeformedTriangle.P2;
+
+			const bool bIsValid = ScaleAndClip(VertexIndex, PackedRestPosition.w);
+			WorldVertexPosition = bIsValid ? ControlPoint + CurrOffset : float3(1.#INF, 1.#INF, 1.#INF);
+
+			OutRenderDeformedPositionBuffer[VertexIndex] = float4(f32tof16(WorldVertexPosition - OutHairPositionOffset), PackedRestPosition.w);
+		}
+		#endif	
+
+		// 4. Linear blend skinning (rotation)
+		#if PERMUTATION_DYNAMIC_GEOMETRY == 3
+		{
+			uint4 PackedRestPosition = RenderRestPosePositionBuffer[VertexIndex];
+			const float3 RestControlPoint = f16tof32(PackedRestPosition.xyz) + InRenderHairPositionOffset;
+			float3 ControlPoint = RestControlPoint;
+
 			const uint RootIndex						= RenVertexToRootIndexBuffer[VertexIndex];
 			const float3 RootBarycentric				= DecodeBarycentrics(RootBarycentricBuffer[RootIndex]);
 			const FHairMeshTriangle RestTriangle		= GetTriangleTransformation(RootIndex, RestPosition0Buffer, RestPosition1Buffer, RestPosition2Buffer);
 			const FHairMeshTriangle DeformedTriangle	= GetTriangleTransformation(RootIndex, DeformedPosition0Buffer, DeformedPosition1Buffer, DeformedPosition2Buffer);
-			
-			ControlPoint -= RootBarycentric.x * RestTriangle.P0     + RootBarycentric.y * RestTriangle.P1     + RootBarycentric.z * RestTriangle.P2;
-			ControlPoint += RootBarycentric.x * DeformedTriangle.P0 + RootBarycentric.y * DeformedTriangle.P1 + RootBarycentric.z * DeformedTriangle.P2;
+
+			ControlPoint -= RootBarycentric.x * RestTriangle.P0 + RootBarycentric.y * RestTriangle.P1 + RootBarycentric.z * RestTriangle.P2;
+
+			float3 CurrOffset = 0;
+
+			#if PERMUTATION_SIMULATION == 1
+			for (uint K = 0; K < 3; ++K)
+			{
+				const uint GuideIndex0 = GuideVertexIndices[K];
+				const uint GuideIndex1 = GuideIndex0 + 1;
+
+				const uint RootIndexG = SimVertexToRootIndexBuffer[GuideIndex0];
+				
+				const FHairMeshTriangle RestTriangleG		= GetTriangleTransformation(RootIndexG, SimRestPosition0Buffer, SimRestPosition1Buffer, SimRestPosition2Buffer);
+				const FHairMeshTriangle DeformedTriangleG	= GetTriangleTransformation(RootIndexG, SimDeformedPosition0Buffer, SimDeformedPosition1Buffer, SimDeformedPosition2Buffer);
+				
+				const float3 LocalPoint = VectorToTriangle(ControlPoint,RestTriangleG);
+
+				const float3 Offset0 = ComputeRotateGeometryOffset(GuideIndex0, 1.0, RestTriangleG, DeformedTriangleG);
+				const float3 Offset1 = ComputeRotateGeometryOffset(GuideIndex1, 1.0, RestTriangleG, DeformedTriangleG);
+
+				CurrOffset += GuideVertexWeights[K] * VectorToWorld(lerp(Offset0, Offset1, GuideVertexLerps[K])+LocalPoint,DeformedTriangleG);
+			}
+			#endif
+
+			ControlPoint =  RootBarycentric.x * DeformedTriangle.P0 + RootBarycentric.y * DeformedTriangle.P1 + RootBarycentric.z * DeformedTriangle.P2;
 
 			const bool bIsValid = ScaleAndClip(VertexIndex, PackedRestPosition.w);
 			WorldVertexPosition = bIsValid ? ControlPoint + CurrOffset : float3(1.#INF, 1.#INF, 1.#INF);
 
 			OutRenderDeformedPositionBuffer[VertexIndex] = float4(f32tof16(WorldVertexPosition - OutHairPositionOffset), PackedRestPosition.w);
 		}
-		#endif	
+		#endif
+
+		// 5. Linear blend skinning (translation + rotation)
+		#if PERMUTATION_DYNAMIC_GEOMETRY == 3
+		{
+			uint4 PackedRestPosition = RenderRestPosePositionBuffer[VertexIndex];
+			const float3 RestControlPoint = f16tof32(PackedRestPosition.xyz) + InRenderHairPositionOffset;
+			float3 ControlPoint = RestControlPoint;
+
+			const uint RootIndex						= RenVertexToRootIndexBuffer[VertexIndex];
+			const float3 RootBarycentric				= DecodeBarycentrics(RootBarycentricBuffer[RootIndex]);
+			const FHairMeshTriangle RestTriangle		= GetTriangleTransformation(RootIndex, RestPosition0Buffer, RestPosition1Buffer, RestPosition2Buffer);
+			const FHairMeshTriangle DeformedTriangle	= GetTriangleTransformation(RootIndex, DeformedPosition0Buffer, DeformedPosition1Buffer, DeformedPosition2Buffer);
+
+			float3 CurrOffset = 0;
+
+			#if PERMUTATION_SIMULATION == 1
+			for (uint K = 0; K < 3; ++K)
+			{
+				const uint GuideIndex0 = GuideVertexIndices[K];
+				const uint GuideIndex1 = GuideIndex0 + 1;
+
+				const uint RootIndexG = SimVertexToRootIndexBuffer[GuideIndex0];
+				const float3 RootBarycentricG				= DecodeBarycentrics(RootBarycentricBuffer[RootIndexG]);
+				
+				const FHairMeshTriangle RestTriangleG		= GetTriangleTransformation(RootIndexG, SimRestPosition0Buffer, SimRestPosition1Buffer, SimRestPosition2Buffer);
+				const FHairMeshTriangle DeformedTriangleG	= GetTriangleTransformation(RootIndexG, SimDeformedPosition0Buffer, SimDeformedPosition1Buffer, SimDeformedPosition2Buffer);
+				
+				const float3 LocalPoint = ToTriangle(ControlPoint,RootBarycentricG,RestTriangleG);
+
+				const float3 Offset0 = ComputeDynamicGeometryOffset(GuideIndex0, 1.0, RestTriangleG, DeformedTriangleG, RootBarycentricG);
+				const float3 Offset1 = ComputeDynamicGeometryOffset(GuideIndex1, 1.0, RestTriangleG, DeformedTriangleG, RootBarycentricG);
+
+				CurrOffset += GuideVertexWeights[K] * ToWorld(lerp(Offset0, Offset1, GuideVertexLerps[K])+LocalPoint,RootBarycentricG,DeformedTriangleG);
+			}
+			#endif
+
+			const bool bIsValid = ScaleAndClip(VertexIndex, PackedRestPosition.w);
+			WorldVertexPosition = bIsValid ? CurrOffset : float3(1.#INF, 1.#INF, 1.#INF);
+
+			OutRenderDeformedPositionBuffer[VertexIndex] = float4(f32tof16(WorldVertexPosition - OutHairPositionOffset), PackedRestPosition.w);
+		}
+		#endif
 
 		#if PERMUTATION_DEBUG == 1
 		if (HairDebugMode == 1)
