Fix bad index

diff --git a/Engine/Shaders/Private/HairStrands/HairStrandsInterpolation.usf b/Engine/Shaders/Private/HairStrands/HairStrandsInterpolation.usf
--- a/Engine/Shaders/Private/HairStrands/HairStrandsInterpolation.usf
+++ b/Engine/Shaders/Private/HairStrands/HairStrandsInterpolation.usf
@@ -389,7 +389,7 @@ void MainCS(uint2 DispatchThreadId : SV_DispatchThreadID)
 		#endif
 
 		// 5. Linear blend skinning (translation + rotation)
-		#if PERMUTATION_DYNAMIC_GEOMETRY == 3
+		#if PERMUTATION_DYNAMIC_GEOMETRY == 4
 		{
 			uint4 PackedRestPosition = RenderRestPosePositionBuffer[VertexIndex];
 			const float3 RestControlPoint = f16tof32(PackedRestPosition.xyz) + InRenderHairPositionOffset;
