Integrate PSO creation crash fix from Dev-Rendering for Lightmap density view (wrong VS-PS interpolators) - From UDN

diff --git a/Engine/Shaders/Private/LightMapDensityShader.usf b/Engine/Shaders/Private/LightMapDensityShader.usf
--- a/Engine/Shaders/Private/LightMapDensityShader.usf
+++ b/Engine/Shaders/Private/LightMapDensityShader.usf
@@ -160,6 +160,9 @@ float4		LightMapDensityDisplayOptions;
 void MainPixelShader(
 	FVertexFactoryInterpolantsVSToPS FactoryInterpolants,
 	float4 WorldPosition	: TEXCOORD6,
+#if USE_WORLD_POSITION_EXCLUDING_SHADER_OFFSETS
+	float3 PixelPositionExcludingWPO : TEXCOORD7,
+#endif
 	in float4 SvPosition	: SV_Position
 	OPTIONAL_IsFrontFace,
 	out float4 OutColor		: SV_Target0
