Fixing issues with clear coat:
1) The normalization term for reflected transmitted color was being applied incorrectly
2) Addressed corner-case where Metallic = 0 was not properly applying transmission
3) Applied lerp(Default, ClearCoat, scalar) to handle mixed materials

diff --git a/Engine/Shaders/Private/ShadingModels.ush b/Engine/Shaders/Private/ShadingModels.ush
--- a/Engine/Shaders/Private/ShadingModels.ush
+++ b/Engine/Shaders/Private/ShadingModels.ush
@@ -278,13 +278,12 @@ float3 CalcThinTransmission(float NoL, float NoV, FGBufferData GBuffer)
 		float Metallic = GBuffer.Metallic;
 		float ClearCoat = GBuffer.CustomData.x;
 
-		// Only absorb for metallic materials
 		Transmission = 1.0;
-		if (Metallic * ClearCoat > 0.0)
+		if (ClearCoat > 0.0)
 		{
 			// Base color represents reflected color viewed at 0 incidence angle, after being absorbed through the substrate.
 			// Because of this, extinction is normalized by traveling through layer thickness twice
-			float3 TransmissionColor = GBuffer.BaseColor;
+			float3 TransmissionColor = Diffuse_Lambert(GBuffer.DiffuseColor);
 			float3 ExtinctionCoefficient = -log(TransmissionColor) / (2.0 * NormalizedLayerThickness);
 			float3 OpticalDepth = ExtinctionCoefficient * ThinDistance;
 			Transmission = exp(-OpticalDepth);
@@ -389,7 +388,7 @@ FDirectLighting ClearCoatBxDF( FGBufferData GBuffer, half3 N, half3 V, half3 L,
 		float Vis = Vis_SmithJointApprox(a2, Context.NoV, NoL);
 
 		float Fr1 = D * Vis * F;
-		Lighting.Specular = AreaLight.FalloffColor * (Falloff * NoL * Fr1);
+		Lighting.Specular = ClearCoat * AreaLight.FalloffColor * (Falloff * NoL * Fr1);
 	}
 
 	// Incoming and exiting Fresnel terms are identical to incoming Fresnel term (VoH == HoL)
@@ -429,11 +428,18 @@ FDirectLighting ClearCoatBxDF( FGBufferData GBuffer, half3 N, half3 V, half3 L,
 	float Alpha = Pow2(GBuffer.Roughness);
 	float a2 = Pow2(Alpha);
 	float Energy = EnergyNormalization(a2, BottomContext.VoH, AreaLight);
-	Lighting.Diffuse = (FresnelCoeff * Energy * Falloff * BottomContext.NoL) * Transmission * AreaLight.FalloffColor * Diffuse_Lambert(GBuffer.DiffuseColor);
+	//Lighting.Diffuse = (FresnelCoeff * Energy * Falloff * BottomContext.NoL) * Transmission * AreaLight.FalloffColor * Diffuse_Lambert(GBuffer.DiffuseColor);
+	float3 CommonDiffuse = (Energy * Falloff) * AreaLight.FalloffColor;
+	float3 DefaultDiffuse = Context.NoL * Diffuse_Lambert(GBuffer.DiffuseColor);
+	float3 RefractedDiffuse = (FresnelCoeff * BottomContext.NoL) * Transmission;
+	Lighting.Diffuse = CommonDiffuse * lerp(DefaultDiffuse, RefractedDiffuse, ClearCoat);
 	if (AreaLight.bIsRect)
 	{
 		// Note: V is used instead of V2 because LTC integration is not tuned to handle refraction direction
-		Lighting.Specular += FresnelCoeff * Transmission * RectGGXApproxLTC(GBuffer.Roughness, GBuffer.SpecularColor, N, V, AreaLight.Rect, AreaLight.Texture);
+		//Lighting.Specular += FresnelCoeff * Transmission * RectGGXApproxLTC(GBuffer.Roughness, GBuffer.SpecularColor, N, V, AreaLight.Rect, AreaLight.Texture);
+		float3 DefaultSpecular = RectGGXApproxLTC(GBuffer.Roughness, GBuffer.SpecularColor, N, V, AreaLight.Rect, AreaLight.Texture);
+		float3 RefractedSpecular = FresnelCoeff * Transmission * DefaultSpecular;
+		Lighting.Specular += lerp(DefaultSpecular, RefractedSpecular, ClearCoat);
 	}
 	else
 	{
@@ -458,7 +464,12 @@ FDirectLighting ClearCoatBxDF( FGBufferData GBuffer, half3 N, half3 V, half3 L,
 		}
 
 		float3 F = F_Schlick(GBuffer.SpecularColor, BottomContext.VoH);
-		Lighting.Specular += (Energy * Falloff * BottomContext.NoL * D2 * Vis2 * FresnelCoeff) * Transmission * AreaLight.FalloffColor * F;
+		//Lighting.Specular += (Energy * Falloff * BottomContext.NoL * D2 * Vis2 * FresnelCoeff) * Transmission * AreaLight.FalloffColor * F;
+		// Note: reusing D, V, and F from refracted context to save computation for when ClearCoat < 1
+		float3 CommonSpecular = (Energy * Falloff * D2 * Vis2) * AreaLight.FalloffColor * F;
+		float3 DefaultSpecular = Context.NoL;
+		float3 RefractedSpecular = Transmission * BottomContext.NoL;
+		Lighting.Specular += CommonSpecular * lerp(DefaultSpecular, RefractedSpecular, ClearCoat);
 	}
 
 	return Lighting;
