Prevent substrate transmission for non metallic, clear coat materials.
Reformulate CalcThinTransmission to align with path tracing material model.

#public okforgithub

diff --git a/Engine/Shaders/Private/ShadingModels.ush b/Engine/Shaders/Private/ShadingModels.ush
--- a/Engine/Shaders/Private/ShadingModels.ush
+++ b/Engine/Shaders/Private/ShadingModels.ush
@@ -277,16 +277,18 @@ float3 CalcThinTransmission(float NoL, float NoV, FGBufferData GBuffer)
 		float ThinDistance = NormalizedLayerThickness * (rcp(NoV) + rcp(NoL));
 		float Metallic = GBuffer.Metallic;
 		float ClearCoat = GBuffer.CustomData.x;
+		float AbsorptionMix = Metallic * ClearCoat;
 
 		Transmission = 1.0;
-		if (ClearCoat > 0.0)
+		if (AbsorptionMix > 0.0)
 		{
 			// Base color represents reflected color viewed at 0 incidence angle, after being absorbed through the substrate.
 			// Because of this, extinction is normalized by traveling through layer thickness twice
 			float3 TransmissionColor = Diffuse_Lambert(GBuffer.DiffuseColor);
 			float3 ExtinctionCoefficient = -log(TransmissionColor) / (2.0 * NormalizedLayerThickness);
-			float3 OpticalDepth = ExtinctionCoefficient * ThinDistance;
+			float3 OpticalDepth = ExtinctionCoefficient * max(ThinDistance - 2.0 * NormalizedLayerThickness, 0.0);
 			Transmission = exp(-OpticalDepth);
+			Transmission = lerp(1.0, Transmission, AbsorptionMix);
 		}
 	}
 
@@ -431,7 +433,7 @@ FDirectLighting ClearCoatBxDF( FGBufferData GBuffer, half3 N, half3 V, half3 L,
 	//Lighting.Diffuse = (FresnelCoeff * Energy * Falloff * BottomContext.NoL) * Transmission * AreaLight.FalloffColor * Diffuse_Lambert(GBuffer.DiffuseColor);
 	float3 CommonDiffuse = (Energy * Falloff) * AreaLight.FalloffColor;
 	float3 DefaultDiffuse = Context.NoL * Diffuse_Lambert(GBuffer.DiffuseColor);
-	float3 RefractedDiffuse = (FresnelCoeff * BottomContext.NoL) * Transmission;
+	float3 RefractedDiffuse = (FresnelCoeff * BottomContext.NoL) * Transmission * Diffuse_Lambert(GBuffer.DiffuseColor);
 	Lighting.Diffuse = CommonDiffuse * lerp(DefaultDiffuse, RefractedDiffuse, ClearCoat);
 	if (AreaLight.bIsRect)
 	{
