Adapt path tracing clear coat model to the new real-time changes.

diff --git a/Engine/Shaders/Private/PathTracing/Material/PathTracingClearCoat.ush b/Engine/Shaders/Private/PathTracing/Material/PathTracingClearCoat.ush
--- a/Engine/Shaders/Private/PathTracing/Material/PathTracingClearCoat.ush
+++ b/Engine/Shaders/Private/PathTracing/Material/PathTracingClearCoat.ush
@@ -78,17 +78,19 @@ float3 CalcThinTransmission(float NoL, float NoV, FMaterialClosestHitPayload Pay
 		float ThinDistance = LayerThickness * (rcp(NoV) + rcp(NoL));
 		float Metallic = GetMetallic(Payload);
 		float ClearCoat = GetClearCoat(Payload);
+		float AbsorptionMix = Metallic * ClearCoat;
 
 		// Only absorb for metallic materials
 		Transmission = 1.0;
-		if (Metallic * ClearCoat > 0.0)
+		if (AbsorptionMix > 0.0)
 		{
 			// Base color represents reflected color viewed at 0 incidence angle, after being absorbed through the substrate.
 			// Because of this, extinction is normalized by traveling through layer thickness twice
-			float3 TransmissionColor = GetBaseColor(Payload);
+			float3 TransmissionColor = Diffuse_Lambert(GetBaseColor(Payload) );
 			float3 ExtinctionCoefficient = -log(TransmissionColor) / (2.0 * LayerThickness);
-			float3 OpticalDepth = ExtinctionCoefficient * ThinDistance;
+			float3 OpticalDepth = ExtinctionCoefficient * max(ThinDistance - 2.0 * LayerThickness, 0.0);
 			Transmission = exp(-OpticalDepth);
+			Transmission = lerp(1.0, Transmission, AbsorptionMix);
 		}
 	}
 	return Transmission;
