Fixed potential flickering in ray tracing translucency due precision issues

diff --git a/Engine/Shaders/Private/RayTracing/RayTracingPrimaryRays.usf b/Engine/Shaders/Private/RayTracing/RayTracingPrimaryRays.usf
--- a/Engine/Shaders/Private/RayTracing/RayTracingPrimaryRays.usf
+++ b/Engine/Shaders/Private/RayTracing/RayTracingPrimaryRays.usf
@@ -116,7 +116,7 @@ RAY_TRACING_ENTRY_RAYGEN(RayTracingPrimaryRaysRGS)
 
 	if((ERayTracingPrimaryRaysFlag_UseGBufferForMaxDistance & PrimaryRayFlags) != 0) 
 	{
-		Ray.TMax = WorldSpaceDistance - 0.01;
+		Ray.TMax = WorldSpaceDistance - 0.025;
 	}
 
 	bool bAllowSkySampling;
