Implements depth velocity

diff --git a/Engine/Shaders/Private/BasePassPixelShader.usf b/Engine/Shaders/Private/BasePassPixelShader.usf
--- a/Engine/Shaders/Private/BasePassPixelShader.usf
+++ b/Engine/Shaders/Private/BasePassPixelShader.usf
@@ -971,13 +971,12 @@ void FPixelShaderInOut_MainPS(
 	{
 		// 2d velocity, includes camera an object motion
 	#if WRITES_VELOCITY_TO_GBUFFER_USE_POS_INTERPOLATOR
-		float2 Velocity = Calculate2DVelocity(BasePassInterpolants.VelocityScreenPosition, BasePassInterpolants.VelocityPrevScreenPosition);
+		float3 Velocity = Calculate3DVelocity(BasePassInterpolants.VelocityScreenPosition, BasePassInterpolants.VelocityPrevScreenPosition);
 	#else
-		float2 Velocity = Calculate2DVelocity(MaterialParameters.ScreenPosition, BasePassInterpolants.VelocityPrevScreenPosition);
+		float3 Velocity = Calculate3DVelocity(MaterialParameters.ScreenPosition, BasePassInterpolants.VelocityPrevScreenPosition);
 	#endif
 
-		// Make sure not to touch 0,0 which is clear color
-		float4 EncodedVelocity = float4(EncodeVelocityToTexture(Velocity) * GetPrimitiveData(MaterialParameters.PrimitiveId).DrawsVelocity, 0, 0);
+		float4 EncodedVelocity = float4(EncodeVelocityToTexture(Velocity) * GetPrimitiveData(MaterialParameters.PrimitiveId).DrawsVelocity, 0);
 	#if USES_GBUFFER
 		GBuffer.Velocity = EncodedVelocity;
 	#else
diff --git a/Engine/Shaders/Private/Common.ush b/Engine/Shaders/Private/Common.ush
--- a/Engine/Shaders/Private/Common.ush
+++ b/Engine/Shaders/Private/Common.ush
@@ -1519,14 +1519,14 @@ struct FScreenVertexOutput
 // for velocity rendering, motionblur and temporal AA
 // velocity needs to support -2..2 screen space range for x and y
 // texture is 16bit 0..1 range per channel
-float2 EncodeVelocityToTexture(float2 In)
+float3 EncodeVelocityToTexture(float3 In)
 {
 	// 0.499f is a value smaller than 0.5f to avoid using the full range to use the clear color (0,0) as special value
 	// 0.5f to allow for a range of -2..2 instead of -1..1 for really fast motions for temporal AA
 	return In * (0.499f * 0.5f) + 32767.0f / 65535.0f;
 }
 // see EncodeVelocityToTexture()
-float2 DecodeVelocityFromTexture(float2 In)
+float3 DecodeVelocityFromTexture(float3 In)
 {
 	const float InvDiv = 1.0f / (0.499f * 0.5f);
 	// reference
diff --git a/Engine/Shaders/Private/DistanceFieldLightingPost.usf b/Engine/Shaders/Private/DistanceFieldLightingPost.usf
--- a/Engine/Shaders/Private/DistanceFieldLightingPost.usf
+++ b/Engine/Shaders/Private/DistanceFieldLightingPost.usf
@@ -178,7 +178,7 @@ void UpdateHistoryDepthRejectionPS(
 		
 	if (VelocityN.x > 0.0)
 	{
-		Screenvelocity = DecodeVelocityFromTexture(VelocityN);
+		Screenvelocity = DecodeVelocityFromTexture(float3(VelocityN, 0.0)).xy;
 	}
 		
 	float PixelSpeed = 0.5 * length(Screenvelocity);
diff --git a/Engine/Shaders/Private/HairStrands/HairStrandsMaterialPS.usf b/Engine/Shaders/Private/HairStrands/HairStrandsMaterialPS.usf
--- a/Engine/Shaders/Private/HairStrands/HairStrandsMaterialPS.usf
+++ b/Engine/Shaders/Private/HairStrands/HairStrandsMaterialPS.usf
@@ -127,8 +127,8 @@ FEvaluateOutput Evaluate(
 		const float4 ScreenPos		= mul(float4(Velocity_World_P.xyz, 1), ResolvedView.TranslatedWorldToClip);
 		const float4 PrevScreenPos	= mul(float4(Velocity_World_P_Prev.xyz, 1), ResolvedView.PrevTranslatedWorldToClip);
 
-		const float2 Velocity	= Calculate2DVelocity(ScreenPos, PrevScreenPos);
-		EncodedVelocity			= EncodeVelocityToTexture(Velocity);
+		const float2 Velocity	= Calculate3DVelocity(ScreenPos, PrevScreenPos).xy;
+		EncodedVelocity			= EncodeVelocityToTexture(float3(Velocity, 0.0)).xy;
 	}
 	#endif
 
diff --git a/Engine/Shaders/Private/HairStrands/HairStrandsVelocity.usf b/Engine/Shaders/Private/HairStrands/HairStrandsVelocity.usf
--- a/Engine/Shaders/Private/HairStrands/HairStrandsVelocity.usf
+++ b/Engine/Shaders/Private/HairStrands/HairStrandsVelocity.usf
@@ -63,7 +63,7 @@ void MainCS(uint GroupIndex : SV_GroupIndex, uint3 DispatchThreadId : SV_Dispatc
 			const float2 EncodedVelocity = NodeVelocity[SampleIndex];
 
 			#if PERMUTATION_VELOCITY == VELOCITY_TYPE_AVG
-			AverageVelocity += DecodeVelocityFromTexture(EncodedVelocity);
+			AverageVelocity += DecodeVelocityFromTexture(float3(EncodedVelocity, 0.0)).xy;
 			#endif			
 
 			#if PERMUTATION_VELOCITY == VELOCITY_TYPE_CLOSEST
@@ -76,7 +76,7 @@ void MainCS(uint GroupIndex : SV_GroupIndex, uint3 DispatchThreadId : SV_Dispatc
 			#endif
 
 			#if PERMUTATION_VELOCITY == VELOCITY_TYPE_MAX
-			const float2 CurrentVelocity = DecodeVelocityFromTexture(EncodedVelocity);
+			const float2 CurrentVelocity = DecodeVelocityFromTexture(float3(EncodedVelocity, 0.0)).xy;
 			const float CurrentVelocityMagnitude2 = dot(CurrentVelocity, CurrentVelocity);
 			if (CurrentVelocityMagnitude2 > MaxVelocityMagnitude2)
 			{
@@ -87,7 +87,7 @@ void MainCS(uint GroupIndex : SV_GroupIndex, uint3 DispatchThreadId : SV_Dispatc
 		}
 
 		#if PERMUTATION_VELOCITY == VELOCITY_TYPE_AVG
-		OutVelocityTexture[PixelCoord] = EncodeVelocityToTexture(AverageVelocity / max(NodeDesc.Count, 1));
+		OutVelocityTexture[PixelCoord] = EncodeVelocityToTexture(float3(AverageVelocity / max(NodeDesc.Count, 1), 0.0)).xy;
 		#endif		
 
 		#if PERMUTATION_VELOCITY == VELOCITY_TYPE_CLOSEST
diff --git a/Engine/Shaders/Private/HairStrands/HairStrandsVisibilityCompaction.usf b/Engine/Shaders/Private/HairStrands/HairStrandsVisibilityCompaction.usf
--- a/Engine/Shaders/Private/HairStrands/HairStrandsVisibilityCompaction.usf
+++ b/Engine/Shaders/Private/HairStrands/HairStrandsVisibilityCompaction.usf
@@ -252,11 +252,11 @@ void MainCS(uint GroupIndex : SV_GroupIndex, uint3 DispatchThreadId : SV_Dispatc
 		Categorization.PixelCoverage = TransmittanceToCoverage(TotalTransmittance);
 
 	#if PERMUTATION_VELOCITY == VELOCITY_TYPE_AVG
-		OutVelocityTexture[PixelCoord] = EncodeVelocityToTexture(AverageVelocity / max(ValidVelocityCount, 1));
+		OutVelocityTexture[PixelCoord] = EncodeVelocityToTexture(float3(AverageVelocity / max(ValidVelocityCount, 1), 0.0)).xy;
 	#elif PERMUTATION_VELOCITY == VELOCITY_TYPE_CLOSEST
-		OutVelocityTexture[PixelCoord] = EncodeVelocityToTexture(ClosestVelocity);
+		OutVelocityTexture[PixelCoord] = EncodeVelocityToTexture(float3(ClosestVelocity, 0.0)).xy;
 	#elif PERMUTATION_VELOCITY == VELOCITY_TYPE_MAX
-		OutVelocityTexture[PixelCoord] = EncodeVelocityToTexture(MaxVelocity);
+		OutVelocityTexture[PixelCoord] = EncodeVelocityToTexture(float3(MaxVelocity, 0.0)).xy;
 	#endif		
 	}
 
@@ -637,7 +637,7 @@ void MainCS(uint GroupIndex : SV_GroupIndex, uint3 DispatchThreadId : SV_Dispatc
 
 					#if PERMUTATION_VELOCITY == VELOCITY_TYPE_AVG
 					const float2 EncodedVelocity = MSAA_VelocityTexture.Load(PixelCoord, SampleIt);
-					AverageVelocity += DecodeVelocityFromTexture(EncodedVelocity);
+					AverageVelocity += DecodeVelocityFromTexture(float3(EncodedVelocity, 0.0)).xy;
 					++ValidVelocityCount;
 					#endif			
 
@@ -653,7 +653,7 @@ void MainCS(uint GroupIndex : SV_GroupIndex, uint3 DispatchThreadId : SV_Dispatc
 
 					#if PERMUTATION_VELOCITY == VELOCITY_TYPE_MAX
 					const float2 EncodedVelocity = MSAA_VelocityTexture.Load(PixelCoord, SampleIt);
-					const float2 CurrentVelocity = DecodeVelocityFromTexture(EncodedVelocity);
+					const float2 CurrentVelocity = DecodeVelocityFromTexture(float3(EncodedVelocity, 0.0)).xy;
 					const float CurrentVelocityMagnitude2 = dot(CurrentVelocity, CurrentVelocity);
 					if (CurrentVelocityMagnitude2 > MaxVelocityMagnitude2)
 					{
@@ -675,7 +675,7 @@ void MainCS(uint GroupIndex : SV_GroupIndex, uint3 DispatchThreadId : SV_Dispatc
 			#endif
 
 			#if PERMUTATION_VELOCITY == VELOCITY_TYPE_AVG
-			OutVelocityTexture[PixelCoord] = EncodeVelocityToTexture(AverageVelocity / max(ValidVelocityCount, 1));
+			OutVelocityTexture[PixelCoord] = EncodeVelocityToTexture(float3(AverageVelocity / max(ValidVelocityCount, 1), 0.0)).xy;
 			#endif		
 
 			#if PERMUTATION_VELOCITY == VELOCITY_TYPE_CLOSEST
diff --git a/Engine/Shaders/Private/HairStrands/HairStrandsVisibilityComposeSubPixelPS.usf b/Engine/Shaders/Private/HairStrands/HairStrandsVisibilityComposeSubPixelPS.usf
--- a/Engine/Shaders/Private/HairStrands/HairStrandsVisibilityComposeSubPixelPS.usf
+++ b/Engine/Shaders/Private/HairStrands/HairStrandsVisibilityComposeSubPixelPS.usf
@@ -47,7 +47,7 @@ void FastResolvePS(in FScreenVertexOutput Input)
 	for (uint SampleIt = 0; SampleIt < PERMUTATION_MSAACOUNT; ++SampleIt)
 	{
 		const float2 EncodedVelocity = HairVisibilityVelocityTexture.Load(PixelCoord, SampleIt);
-		const float2 Velocity = DecodeVelocityFromTexture(EncodedVelocity);
+		const float2 Velocity = DecodeVelocityFromTexture(float3(EncodedVelocity, 0.0)).xy;
 		const float VelocityMagnitude = sqrt(dot(Velocity, Velocity));
 		if (EncodedVelocity.x > 0 && VelocityMagnitude > VelocityThreshold)
 		{
diff --git a/Engine/Shaders/Private/HairStrands/HairStrandsVisibilityPS.usf b/Engine/Shaders/Private/HairStrands/HairStrandsVisibilityPS.usf
--- a/Engine/Shaders/Private/HairStrands/HairStrandsVisibilityPS.usf
+++ b/Engine/Shaders/Private/HairStrands/HairStrandsVisibilityPS.usf
@@ -145,8 +145,8 @@ void MainVisibility(
 		Tangent = mul(normalize(GetMaterialNormalRaw(PixelMaterialInputs)), MaterialParameters.TangentToWorld);
 	}
 	{
-		Velocity = Calculate2DVelocity(PackedVelocityA, PackedVelocityC);
-		EncodedVelocity = EncodeVelocityToTexture(Velocity);
+		Velocity = Calculate3DVelocity(PackedVelocityA, PackedVelocityC).xy;
+		EncodedVelocity = EncodeVelocityToTexture(float3(Velocity, 0.0)).xy;
 	}
 	#endif
 #endif
diff --git a/Engine/Shaders/Private/HairStrands/HairStrandsVisibilityVS.usf b/Engine/Shaders/Private/HairStrands/HairStrandsVisibilityVS.usf
--- a/Engine/Shaders/Private/HairStrands/HairStrandsVisibilityVS.usf
+++ b/Engine/Shaders/Private/HairStrands/HairStrandsVisibilityVS.usf
@@ -85,7 +85,7 @@ void Main(
 		Output.PackedVelocityA			= INVARIANT(ScreenPos);
 		Output.PackedVelocityC			= INVARIANT(PrevScreenPosObj);
 
-		const float2 ScreenVelocity = Calculate2DVelocity(Output.PackedVelocityA, Output.PackedVelocityC);
+		const float2 ScreenVelocity = Calculate3DVelocity(Output.PackedVelocityA, Output.PackedVelocityC).xy;
 		VelocityScale = saturate(max(abs(ScreenVelocity.x), abs(ScreenVelocity.y)) / HairRenderInfo.VelocityMagnitudeScale);
 	}
 	#endif
diff --git a/Engine/Shaders/Private/MaterialTemplate.ush b/Engine/Shaders/Private/MaterialTemplate.ush
--- a/Engine/Shaders/Private/MaterialTemplate.ush
+++ b/Engine/Shaders/Private/MaterialTemplate.ush
@@ -1732,7 +1732,7 @@ float2 PostProcessVelocityLookup(float Depth, float2 UV)
 
 	if( Velocity.x > 0.0 )
 	{
-		Velocity = DecodeVelocityFromTexture(Velocity);
+		Velocity = DecodeVelocityFromTexture(float3(Velocity, 0.0)).xy;
 	}
 	else
 	{
diff --git a/Engine/Shaders/Private/PostProcessAmbientOcclusion.usf b/Engine/Shaders/Private/PostProcessAmbientOcclusion.usf
--- a/Engine/Shaders/Private/PostProcessAmbientOcclusion.usf
+++ b/Engine/Shaders/Private/PostProcessAmbientOcclusion.usf
@@ -1565,7 +1565,7 @@ float3 ReprojectPos(float2 UV, float Depth)
 	float2 Velocity	=  Texture2DSampleLevel(SceneVelocityTexture, SceneVelocityTextureSampler, UV,0).rg;
 	if( Velocity.x > 0.0 )
 	{
-		PrevScreen = ThisClip.xy - DecodeVelocityFromTexture( Velocity );
+		PrevScreen = ThisClip.xy - DecodeVelocityFromTexture(float3(Velocity, 0.0)).xy;
 	}
 	
 	float2 PrevUV = PrevScreen.xy * PrevScreenPositionScaleBias.xy + PrevScreenPositionScaleBias.zw;
diff --git a/Engine/Shaders/Private/PostProcessMotionBlur.usf b/Engine/Shaders/Private/PostProcessMotionBlur.usf
--- a/Engine/Shaders/Private/PostProcessMotionBlur.usf
+++ b/Engine/Shaders/Private/PostProcessMotionBlur.usf
@@ -407,7 +407,7 @@ void VisualizeMotionBlurPS(
 	// relative, in screen space -1...1 -1..1, can be even outside of that range, points into the movement direction
 	float2 VelocityAtTileCenter;
 	{
-		VelocityAtTileCenter = DecodeVelocityFromTexture(VelocityTexture.SampleLevel(SharedVelocitySampler, UVAtTileCenter, 0).xy);
+		VelocityAtTileCenter = DecodeVelocityFromTexture(float3(VelocityTexture.SampleLevel(SharedVelocitySampler, UVAtTileCenter, 0).xy, 0.0)).xy;
 
 		// reconstruct from camera motion if we don't have velocity data
 		if(!bSelectorOpaqueAtTileCenter)
diff --git a/Engine/Shaders/Private/PostProcessTemporalAA.usf b/Engine/Shaders/Private/PostProcessTemporalAA.usf
--- a/Engine/Shaders/Private/PostProcessTemporalAA.usf
+++ b/Engine/Shaders/Private/PostProcessTemporalAA.usf
@@ -2015,7 +2015,7 @@ FTAAHistoryPayload TemporalAASample(uint2 GroupId, uint2 GroupThreadId, uint Gro
 			bool DynamicN = VelocityN.x > 0.0;
 			if(DynamicN)
 			{
-				BackN = DecodeVelocityFromTexture(VelocityN);
+				BackN = DecodeVelocityFromTexture(float3(VelocityN, 0.0)).xy;
 			}
 			BackTemp = BackN * OutputViewportSize.xy;
 		}
diff --git a/Engine/Shaders/Private/PostProcessVelocityFlatten.usf b/Engine/Shaders/Private/PostProcessVelocityFlatten.usf
--- a/Engine/Shaders/Private/PostProcessVelocityFlatten.usf
+++ b/Engine/Shaders/Private/PostProcessVelocityFlatten.usf
@@ -68,7 +68,7 @@ void VelocityFlattenMain(
 
 	if (Velocity.x > 0.0)
 	{
-		Velocity = DecodeVelocityFromTexture(Velocity);
+		Velocity = DecodeVelocityFromTexture(float3(Velocity, 0.0)).xy;
 	}
 	else
 	{
diff --git a/Engine/Shaders/Private/RayTracing/SplitImaginaryReflectionGBufferCS.usf b/Engine/Shaders/Private/RayTracing/SplitImaginaryReflectionGBufferCS.usf
--- a/Engine/Shaders/Private/RayTracing/SplitImaginaryReflectionGBufferCS.usf
+++ b/Engine/Shaders/Private/RayTracing/SplitImaginaryReflectionGBufferCS.usf
@@ -18,7 +18,7 @@ StructuredBuffer<FImaginaryReflectionGBufferData> ImaginaryReflectionGBuffer;
 
 RWTexture2D<float4> ImaginaryReflectionGBufferA;
 RWTexture2D<float> ImaginaryReflectionDepthZ;
-RWTexture2D<float2> ImaginaryReflectionVelocity;
+RWTexture2D<float4> ImaginaryReflectionVelocity;
 
 [numthreads(THREADGROUP_SIZE, THREADGROUP_SIZE, 1)]
 void MainCS(
@@ -51,13 +51,13 @@ void MainCS(
 		// Encode split components to output formats and output
 		ImaginaryReflectionGBufferA[PixelCoord] = float4(EncodeNormal(NormalizedWorldNormal), 0.0f);
 		ImaginaryReflectionDepthZ[PixelCoord] = ConvertToDeviceZ(AveragedSceneDepth); // Output format expecting DeviceZ similar to a proper depth buffer
-		ImaginaryReflectionVelocity[PixelCoord] = EncodeVelocityToTexture(AveragedVelocity);
+		ImaginaryReflectionVelocity[PixelCoord] = float4(EncodeVelocityToTexture(float3(AveragedVelocity, 0.0)), 0.0);
 	}
 	else
 	{
 		// Output special invalid sample values
 		ImaginaryReflectionGBufferA[PixelCoord] = float4(0.0f, 0.0f, 0.0f, 0.0f);
 		ImaginaryReflectionDepthZ[PixelCoord] = 0.0f;
-		ImaginaryReflectionVelocity[PixelCoord] = float2(0.0f, 0.0f);
+		ImaginaryReflectionVelocity[PixelCoord] = float4(0.0f, 0.0f, 0.0f, 0.0f);
 	}
 }
diff --git a/Engine/Shaders/Private/SSRT/SSRTPrevFrameReduction.usf b/Engine/Shaders/Private/SSRT/SSRTPrevFrameReduction.usf
--- a/Engine/Shaders/Private/SSRT/SSRTPrevFrameReduction.usf
+++ b/Engine/Shaders/Private/SSRT/SSRTPrevFrameReduction.usf
@@ -151,7 +151,7 @@ void MainCS(
 		float2 EncodedVelocity = SceneVelocityBuffer.SampleLevel(SceneVelocityBufferSampler, SceneBufferUV, 0).xy;
 		if (EncodedVelocity.x > 0.0)
 		{
-			PrevScreen = ThisClip.xy - DecodeVelocityFromTexture(EncodedVelocity);
+			PrevScreen = ThisClip.xy - DecodeVelocityFromTexture(float3(EncodedVelocity, 0.0)).xy;
 		}
 
 		float2 PrevFrameUV = PrevScreen.xy * PrevScreenPositionScaleBias.xy + PrevScreenPositionScaleBias.zw;
diff --git a/Engine/Shaders/Private/SSRT/SSRTRayCast.ush b/Engine/Shaders/Private/SSRT/SSRTRayCast.ush
--- a/Engine/Shaders/Private/SSRT/SSRTRayCast.ush
+++ b/Engine/Shaders/Private/SSRT/SSRTRayCast.ush
@@ -483,7 +483,7 @@ void ReprojectHit(float4 PrevScreenPositionScaleBias, Texture2D Texture, Sampler
 	float2 Velocity = Texture.SampleLevel( Sampler, HitUVz.xy, 0 ).xy;
 	if( Velocity.x > 0.0 )
 	{
-		PrevScreen = ThisClip.xy - DecodeVelocityFromTexture( Velocity );
+		PrevScreen = ThisClip.xy - DecodeVelocityFromTexture(float3(Velocity, 0.0)).xy;
 	}
 
 	float2 PrevUV = PrevScreen.xy * PrevScreenPositionScaleBias.xy + PrevScreenPositionScaleBias.zw;
diff --git a/Engine/Shaders/Private/ScreenSpaceDenoise/SSDTemporalAccumulation.usf b/Engine/Shaders/Private/ScreenSpaceDenoise/SSDTemporalAccumulation.usf
--- a/Engine/Shaders/Private/ScreenSpaceDenoise/SSDTemporalAccumulation.usf
+++ b/Engine/Shaders/Private/ScreenSpaceDenoise/SSDTemporalAccumulation.usf
@@ -262,6 +262,11 @@
 	#define CONFIG_BILATERAL_DISTANCE_MULTIPLIER 1.0
 #endif
 
+/** Whether should consider the depth velocity to compute the reference scene depth. */
+#ifndef CONFIG_PATCH_PREV_SCENE_DEPTH
+	#define CONFIG_PATCH_PREV_REF_SCENE_DEPTH 0
+#endif
+
 
 //------------------------------------------------------- CONFIG ENABLED DEFAULTS
 
@@ -432,24 +437,38 @@ void TemporallyAccumulate(
 		SceneBufferUV = clamp(SceneBufferUV, BufferBilinearUVMinMax.xy, BufferBilinearUVMinMax.zw);
 	}
 
+	// Sample current frame data.
+	FSSDCompressedSceneInfos CompressedRefSceneMetadata = SampleCompressedSceneMetadata(
+		/* bPrevFrame = */ false,
+		SceneBufferUV, BufferUVToBufferPixelCoord(SceneBufferUV));
+
+	float DeviceZ;
+	{
+		FSSDSampleSceneInfos RefInfo = UncompressSampleSceneInfo(
+			CONFIG_METADATA_BUFFER_LAYOUT, /* bIsPrevFrame = */ false,
+			DenoiserBufferUVToScreenPosition(SceneBufferUV),
+			CompressedRefSceneMetadata);
+
+		DeviceZ = RefInfo.DeviceZ;
+	}
+
 	// Reproject to previous frame.
-	float2 HistoryScreenPosition = DenoiserBufferUVToScreenPosition(SceneBufferUV);
+	float3 HistoryScreenPosition = float3(DenoiserBufferUVToScreenPosition(SceneBufferUV), DeviceZ);
 	if (1)
 	{
-		float DeviceZ = SceneDepthBuffer.SampleLevel(GlobalPointClampedSampler, SceneBufferUV, 0).r;
+		float4 ThisClip = float4(HistoryScreenPosition, 1);
 
-		float4 ThisClip = float4(HistoryScreenPosition, DeviceZ, 1);
 		float4 PrevClip = mul(ThisClip, View.ClipToPrevClip);
-		float2 PrevScreen = PrevClip.xy / PrevClip.w;
-		float2 Velocity = HistoryScreenPosition - PrevScreen;
+		float3 PrevScreen = PrevClip.xyz * rcp(PrevClip.w);
+		float3 Velocity = HistoryScreenPosition - PrevScreen;
 
 		if (1)
 		{
-			float2 VelocityN = SceneVelocityBuffer.SampleLevel(GlobalPointClampedSampler, SceneBufferUV, 0).xy;
-			bool DynamicN = VelocityN.x > 0.0;
-			if (DynamicN)
+			float3 EncodedVelocity = SceneVelocityBuffer.SampleLevel(GlobalPointClampedSampler, SceneBufferUV, 0).xyz;
+			if (EncodedVelocity.x > 0.0)
+
 			{
-				Velocity = DecodeVelocityFromTexture(VelocityN);
+				Velocity = DecodeVelocityFromTexture(EncodedVelocity);
 			}
 		}
 
@@ -457,11 +476,6 @@ void TemporallyAccumulate(
 	}
 
 
-	// Sample current frame data.
-	FSSDCompressedSceneInfos CompressedRefSceneMetadata = SampleCompressedSceneMetadata(
-		/* bPrevFrame = */ false,
-		SceneBufferUV, BufferUVToBufferPixelCoord(SceneBufferUV));
-
 	FSSDSignalArray CurrentFrameSamples = SampleMultiplexedSignals(
 		SignalInput_Textures_0,
 		SignalInput_Textures_1,
@@ -509,7 +523,7 @@ void TemporallyAccumulate(
 	// Sample the history buffer.
 	FSSDSignalArray HistorySamples = CreateSignalArrayFromScalarValue(0.0);
 	{
-		float2 HistoryBufferUV = HistoryScreenPosition * View.ScreenPositionScaleBias.xy + View.ScreenPositionScaleBias.wz;
+		float2 HistoryBufferUV = HistoryScreenPosition.xy * View.ScreenPositionScaleBias.xy + View.ScreenPositionScaleBias.wz;
 		float2 ClampedHistoryBufferUV = clamp(HistoryBufferUV, BufferBilinearUVMinMax.xy, BufferBilinearUVMinMax.zw);
 		bool bIsPreviousFrameOffscreen = any(HistoryBufferUV != ClampedHistoryBufferUV);
 
@@ -559,11 +573,38 @@ void TemporallyAccumulate(
 			
 			// VGPR configuration of the kernel.
 			KernelConfig.BufferUV = HistoryBufferUV + BufferUVBilinearCorrection;
+
+			#if CONFIG_PATCH_PREV_SCENE_DEPTH
+			{
+				KernelConfig.RefBufferUV = HistoryBufferUV;
+				KernelConfig.RefSceneMetadataLayout = METADATA_BUFFER_LAYOUT_WORLD_POS_SHADINGMODEL;
+				//KernelConfig.RefSceneMetadataLayout = CONFIG_METADATA_BUFFER_LAYOUT;
+
+				FSSDSampleSceneInfos PrevRefInfo = UncompressSampleSceneInfo(
+					CONFIG_METADATA_BUFFER_LAYOUT, /* bIsPrevFrame = */ false,
+					BufferUVToBufferPixelCoord(SceneBufferUV),
+					CompressedRefSceneMetadata);
+
+				PrevRefInfo.ScreenPosition = HistoryScreenPosition.xy;
+				PrevRefInfo.DeviceZ = HistoryScreenPosition.z;
+				PrevRefInfo.WorldDepth = ConvertFromDeviceZ(HistoryScreenPosition.z);
+				
+				float4 ClipPosition = float4(HistoryScreenPosition.xy * (View.ViewToClip[3][3] < 1.0f ? PrevRefInfo.WorldDepth : 1.0f), PrevRefInfo.WorldDepth, 1);
+				
+				PrevRefInfo.TranslatedWorldPosition = mul(ClipPosition, View.PrevScreenToTranslatedWorld).xyz + (View.PreViewTranslation.xyz - View.PrevPreViewTranslation.xyz);
+				
+				KernelConfig.CompressedRefSceneMetadata = CompressSampleSceneInfo(
+					KernelConfig.RefSceneMetadataLayout,
+					PrevRefInfo);
+			}
+			#else
 			{
 				KernelConfig.CompressedRefSceneMetadata = CompressedRefSceneMetadata;
 				KernelConfig.RefBufferUV = SceneBufferUV;
 				KernelConfig.RefSceneMetadataLayout = CONFIG_METADATA_BUFFER_LAYOUT;
 			}
+			#endif
+
 
 			// Compute random signals.
 			ISOLATE
diff --git a/Engine/Shaders/Private/SubsurfaceBurleyNormalized.ush b/Engine/Shaders/Private/SubsurfaceBurleyNormalized.ush
--- a/Engine/Shaders/Private/SubsurfaceBurleyNormalized.ush
+++ b/Engine/Shaders/Private/SubsurfaceBurleyNormalized.ush
@@ -864,7 +864,7 @@ float4 UpdateQualityVariance(float4 SubsurfaceColor, float2 BufferUV)
 
 		if (DynamicN)
 		{
-			BackN = DecodeVelocityFromTexture(VelocityN);
+			BackN = DecodeVelocityFromTexture(float3(VelocityN, 0.0)).xy;
 		}
 
 		BackTemp = BackN * SubsurfaceInput1_ViewportSize.xy;
diff --git a/Engine/Shaders/Private/VelocityCommon.ush b/Engine/Shaders/Private/VelocityCommon.ush
--- a/Engine/Shaders/Private/VelocityCommon.ush
+++ b/Engine/Shaders/Private/VelocityCommon.ush
@@ -6,13 +6,16 @@
 
 #pragma once
 
-float2 Calculate2DVelocity(float4 PackedVelocityA, float4 PackedVelocityC)
+float3 Calculate3DVelocity(float4 PackedVelocityA, float4 PackedVelocityC)
 {
 	float2 ScreenPos = PackedVelocityA.xy / PackedVelocityA.w - ResolvedView.TemporalAAJitter.xy;
 	float2 PrevScreenPos = PackedVelocityC.xy / PackedVelocityC.w - ResolvedView.TemporalAAJitter.zw;
 
-	// 2d velocity, includes camera an object motion
-	float2 Velocity = ScreenPos - PrevScreenPos;
+	float DeviceZ = PackedVelocityA.z / PackedVelocityA.w;
+	float PrevDeviceZ = PackedVelocityC.z / PackedVelocityC.w;
+
+	// 3d velocity, includes camera an object motion
+	float3 Velocity = float3(ScreenPos - PrevScreenPos, DeviceZ - PrevDeviceZ);
 
 	// Make sure not to touch 0,0 which is clear color
 	return Velocity;
diff --git a/Engine/Shaders/Private/VelocityShader.usf b/Engine/Shaders/Private/VelocityShader.usf
--- a/Engine/Shaders/Private/VelocityShader.usf
+++ b/Engine/Shaders/Private/VelocityShader.usf
@@ -261,10 +261,9 @@ void MainPixelShader(
 	OutCoverage = DiscardMaterialWithPixelCoverage(MaterialParameters, PixelMaterialInputs);
 #endif
 
-	// 2d velocity, includes camera an object motion
-	float2 Velocity = Calculate2DVelocity(Input.PackedVelocityA, Input.PackedVelocityC);
+	// 3d velocity, includes camera an object motion
+	float3 Velocity = Calculate3DVelocity(Input.PackedVelocityA, Input.PackedVelocityC);
 
-	// Make sure not to touch 0,0 which is clear color
-	OutColor.xy = EncodeVelocityToTexture(Velocity);
-	OutColor.zw = float2(0, 0);
+	OutColor.xyz = EncodeVelocityToTexture(Velocity).xyz;
+	OutColor.w = 0.0;
 }
diff --git a/Engine/Source/Runtime/Renderer/Private/PostProcess/SceneRenderTargets.cpp b/Engine/Source/Runtime/Renderer/Private/PostProcess/SceneRenderTargets.cpp
--- a/Engine/Source/Runtime/Renderer/Private/PostProcess/SceneRenderTargets.cpp
+++ b/Engine/Source/Runtime/Renderer/Private/PostProcess/SceneRenderTargets.cpp
@@ -1370,7 +1370,7 @@ void FSceneRenderTargets::AllocGBufferTargets(FRHICommandList& RHICmdList)
 
 	if (bAllocateVelocityGBuffer)
 	{
-		FPooledRenderTargetDesc VelocityRTDesc = FVelocityRendering::GetRenderTargetDesc();
+		FPooledRenderTargetDesc VelocityRTDesc = FVelocityRendering::GetRenderTargetDesc(ShaderPlatform);
 		VelocityRTDesc.Flags |= GFastVRamConfig.GBufferVelocity;
 		GRenderTargetPool.FindFreeElement(RHICmdList, VelocityRTDesc, SceneVelocity, TEXT("GBufferVelocity"));
 	}
@@ -2586,6 +2586,8 @@ void FSceneRenderTargets::AllocateLightingChannelTexture(FRHICommandList& RHICmd
 
 void FSceneRenderTargets::AllocateDeferredShadingPathRenderTargets(FRHICommandListImmediate& RHICmdList, const int32 NumViews)
 {
+	const EShaderPlatform ShaderPlatform = GetFeatureLevelShaderPlatform(CurrentFeatureLevel);
+
 	AllocateCommonDepthTargets(RHICmdList);
 
 	// Create a quarter-sized version of the scene depth.
@@ -2708,7 +2710,7 @@ void FSceneRenderTargets::AllocateDeferredShadingPathRenderTargets(FRHICommandLi
 
 	if (bAllocateVelocityGBuffer)
 	{
-		FPooledRenderTargetDesc VelocityRTDesc = FVelocityRendering::GetRenderTargetDesc();
+		FPooledRenderTargetDesc VelocityRTDesc = FVelocityRendering::GetRenderTargetDesc(ShaderPlatform);
 		VelocityRTDesc.Flags |= GFastVRamConfig.GBufferVelocity;
 		GRenderTargetPool.FindFreeElement(RHICmdList, VelocityRTDesc, SceneVelocity, TEXT("GBufferVelocity"));
 	}
diff --git a/Engine/Source/Runtime/Renderer/Private/RayTracing/RayTracingReflections.cpp b/Engine/Source/Runtime/Renderer/Private/RayTracing/RayTracingReflections.cpp
--- a/Engine/Source/Runtime/Renderer/Private/RayTracing/RayTracingReflections.cpp
+++ b/Engine/Source/Runtime/Renderer/Private/RayTracing/RayTracingReflections.cpp
@@ -792,7 +792,7 @@ void FDeferredShadingSceneRenderer::RenderRayTracingReflections(
 		// Create a texture for the velocity imaginary reflection g-buffer.
 		FRDGTextureRef ImaginaryReflectionGBufferVelocityTexture;
 		{
-			FRDGTextureDesc Desc(FRDGTextureDesc::Create2DDesc(RayTracingBufferSize, FVelocityRendering::GetFormat(), FClearValueBinding::Transparent, TexCreate_None, TexCreate_ShaderResource | TexCreate_UAV, false));
+			FRDGTextureDesc Desc(FRDGTextureDesc::Create2DDesc(RayTracingBufferSize, FVelocityRendering::GetFormat(View.GetShaderPlatform()), FClearValueBinding::Transparent, TexCreate_None, TexCreate_ShaderResource | TexCreate_UAV, false));
 			ImaginaryReflectionGBufferVelocityTexture = GraphBuilder.CreateTexture(Desc, TEXT("ImaginaryReflectionGBufferVelocity"));
 		}
 
diff --git a/Engine/Source/Runtime/Renderer/Private/VelocityRendering.cpp b/Engine/Source/Runtime/Renderer/Private/VelocityRendering.cpp
--- a/Engine/Source/Runtime/Renderer/Private/VelocityRendering.cpp
+++ b/Engine/Source/Runtime/Renderer/Private/VelocityRendering.cpp
@@ -147,7 +147,7 @@ public:
 	static void ModifyCompilationEnvironment(const FMaterialShaderPermutationParameters& Parameters, FShaderCompilerEnvironment& OutEnvironment)
 	{
 		FMaterialShader::ModifyCompilationEnvironment(Parameters, OutEnvironment);
-		OutEnvironment.SetRenderTargetOutputFormat(0, FVelocityRendering::GetFormat());
+		OutEnvironment.SetRenderTargetOutputFormat(0, FVelocityRendering::GetFormat(Parameters.Platform));
 	}
 
 	FVelocityPS() = default;
@@ -375,7 +375,7 @@ void FDeferredShadingSceneRenderer::RenderVelocitiesInner(FRHICommandListImmedia
 
 bool FDeferredShadingSceneRenderer::ShouldRenderVelocities() const
 {
-	if (!FVelocityRendering::IsSeparateVelocityPassSupported() || ViewFamily.UseDebugViewPS())
+	if (!FVelocityRendering::IsSeparateVelocityPassSupported(ShaderPlatform) || ViewFamily.UseDebugViewPS())
 	{
 		return false;
 	}
@@ -420,7 +420,7 @@ void FDeferredShadingSceneRenderer::RenderVelocities(FRHICommandListImmediate& R
 
 	if (!VelocityRT)
 	{
-		FPooledRenderTargetDesc Desc = FVelocityRendering::GetRenderTargetDesc();
+		FPooledRenderTargetDesc Desc = FVelocityRendering::GetRenderTargetDesc(ShaderPlatform);
 		GRenderTargetPool.FindFreeElement(RHICmdList, Desc, VelocityRT, TEXT("Velocity"));
 	}
 
@@ -464,21 +464,21 @@ void FDeferredShadingSceneRenderer::RenderVelocities(FRHICommandListImmediate& R
 	GVisualizeTexture.SetCheckPoint(RHICmdList, VelocityRT);
 }
 
-EPixelFormat FVelocityRendering::GetFormat()
+EPixelFormat FVelocityRendering::GetFormat(EShaderPlatform ShaderPlatform)
 {
-	return PF_G16R16;
+	return FDataDrivenShaderPlatformInfo::GetSupportsRayTracing(ShaderPlatform) ? PF_A16B16G16R16 : PF_G16R16;
 }
 
-FPooledRenderTargetDesc FVelocityRendering::GetRenderTargetDesc()
+FPooledRenderTargetDesc FVelocityRendering::GetRenderTargetDesc(EShaderPlatform ShaderPlatform)
 {
 	const FIntPoint BufferSize = FSceneRenderTargets::Get_FrameConstantsOnly().GetBufferSizeXY();
 	const FIntPoint VelocityBufferSize = BufferSize;		// full resolution so we can reuse the existing full res z buffer
-	return FPooledRenderTargetDesc(FPooledRenderTargetDesc::Create2DDesc(VelocityBufferSize, GetFormat(), FClearValueBinding::Transparent, TexCreate_None, TexCreate_RenderTargetable | TexCreate_UAV | TexCreate_ShaderResource, false));
+	return FPooledRenderTargetDesc(FPooledRenderTargetDesc::Create2DDesc(VelocityBufferSize, GetFormat(ShaderPlatform), FClearValueBinding::Transparent, TexCreate_None, TexCreate_RenderTargetable | TexCreate_UAV | TexCreate_ShaderResource, false));
 }
 
-bool FVelocityRendering::IsSeparateVelocityPassSupported()
+bool FVelocityRendering::IsSeparateVelocityPassSupported(EShaderPlatform ShaderPlatform)
 {
-	return GPixelFormats[GetFormat()].Supported;
+	return GPixelFormats[GetFormat(ShaderPlatform)].Supported;
 }
 
 bool FVelocityRendering::BasePassCanOutputVelocity(EShaderPlatform ShaderPlatform)
@@ -529,7 +529,7 @@ bool FVelocityMeshProcessor::PrimitiveHasVelocityForView(const FViewInfo& View,
 
 bool FOpaqueVelocityMeshProcessor::PrimitiveCanHaveVelocity(EShaderPlatform ShaderPlatform, const FPrimitiveSceneProxy* PrimitiveSceneProxy)
 {
-	if (!FVelocityRendering::IsSeparateVelocityPassSupported())
+	if (!FVelocityRendering::IsSeparateVelocityPassSupported(ShaderPlatform))
 	{
 		return false;
 	}
@@ -643,7 +643,7 @@ bool FTranslucentVelocityMeshProcessor::PrimitiveCanHaveVelocity(EShaderPlatform
 	 * Therefore, the primitive can't be filtered based on motion, or it will break post
 	 * effects like depth of field which rely on depth information.
 	 */
-	return FVelocityRendering::IsSeparateVelocityPassSupported();
+	return FVelocityRendering::IsSeparateVelocityPassSupported(ShaderPlatform);
 }
 
 bool FTranslucentVelocityMeshProcessor::PrimitiveHasVelocityForFrame(const FPrimitiveSceneProxy* PrimitiveSceneProxy)
diff --git a/Engine/Source/Runtime/Renderer/Private/VelocityRendering.h b/Engine/Source/Runtime/Renderer/Private/VelocityRendering.h
--- a/Engine/Source/Runtime/Renderer/Private/VelocityRendering.h
+++ b/Engine/Source/Runtime/Renderer/Private/VelocityRendering.h
@@ -32,11 +32,11 @@ EMeshPass::Type GetMeshPassFromVelocityPass(EVelocityPass VelocityPass);
 // Group Velocity Rendering accessors, types, etc.
 struct FVelocityRendering
 {
-	static EPixelFormat GetFormat();
-	static FPooledRenderTargetDesc GetRenderTargetDesc();
+	static EPixelFormat GetFormat(EShaderPlatform ShaderPlatform);
+	static FPooledRenderTargetDesc GetRenderTargetDesc(EShaderPlatform ShaderPlatform);
 
 	/** Returns true if the separate velocity pass is enabled. */
-	static bool IsSeparateVelocityPassSupported();
+	static bool IsSeparateVelocityPassSupported(EShaderPlatform ShaderPlatform);
 
 	/** Returns true if the velocity can be output in the BasePass. */
 	static bool BasePassCanOutputVelocity(EShaderPlatform ShaderPlatform);
