Deletes reflections denoiser's 1spp code path

diff --git a/Engine/Shaders/Private/ScreenSpaceDenoise/SSDSpatialAccumulation.usf b/Engine/Shaders/Private/ScreenSpaceDenoise/SSDSpatialAccumulation.usf
--- a/Engine/Shaders/Private/ScreenSpaceDenoise/SSDSpatialAccumulation.usf
+++ b/Engine/Shaders/Private/ScreenSpaceDenoise/SSDSpatialAccumulation.usf
@@ -169,27 +169,11 @@
 		#define CONFIG_SIGNAL_OUTPUT_LAYOUT SIGNAL_BUFFER_LAYOUT_REFLECTIONS_HISTORY
 
 		// Configure the bilateral therm.
-		#if DIM_MULTI_SPP
-			#define CONFIG_SAMPLE_SET           SAMPLE_SET_DIRECTIONAL_ELLIPSE
+		#define CONFIG_SAMPLE_SET           SAMPLE_SET_DIRECTIONAL_ELLIPSE
 
-			#define CONFIG_BILATERAL_PRESET BILATERAL_PRESET_REFLECTIONS
-			//#define CONFIG_BILATERAL_PRESET BILATERAL_PRESET_REFLECTIONS_1SPP
-			//#define CONFIG_BILATERAL_PRESET BILATERAL_PRESET_DISABLED
-		#else
-			#if 1
-				#define CONFIG_SAMPLE_SET           SAMPLE_SET_DIRECTIONAL_ELLIPSE
-			#else
-				// Use stackowiak sample set to keep noise in the signal to avoid flicker in history rejection,
-				// and adapt number of samples based on the sample.
-				#define CONFIG_SAMPLE_SET           SAMPLE_SET_STACKOWIAK_4_SETS
-				#define CONFIG_SAMPLE_COUNT_POLICY  SAMPLE_COUNT_POLICY_HIT_DISTANCE_BASED
-			#endif
-
-			// Reconstruction has 1 ray per input pixel, therefore can use specular ratio estimator instead of Tokoyashi axes bilateral kernel.
-			#define CONFIG_SPECULAR_RATIO_ESTIMATOR 1
-
-			#define CONFIG_BILATERAL_PRESET BILATERAL_PRESET_REFLECTIONS_1SPP
-		#endif
+		#define CONFIG_BILATERAL_PRESET BILATERAL_PRESET_REFLECTIONS
+		//#define CONFIG_BILATERAL_PRESET BILATERAL_PRESET_REFLECTIONS_1SPP
+		//#define CONFIG_BILATERAL_PRESET BILATERAL_PRESET_DISABLED
 
 		//#define CONFIG_SAMPLE_COUNT_POLICY  SAMPLE_COUNT_POLICY_DISABLED
 
diff --git a/Engine/Source/Runtime/Renderer/Private/ScreenSpaceDenoise.cpp b/Engine/Source/Runtime/Renderer/Private/ScreenSpaceDenoise.cpp
--- a/Engine/Source/Runtime/Renderer/Private/ScreenSpaceDenoise.cpp
+++ b/Engine/Source/Runtime/Renderer/Private/ScreenSpaceDenoise.cpp
@@ -294,7 +294,6 @@ static int32 SignalMaxBatchSize(ESignalProcessing SignalProcessing)
 static bool SignalSupport1SPP(ESignalProcessing SignalProcessing)
 {
 	return (
-		SignalProcessing == ESignalProcessing::Reflections ||
 		SignalProcessing == ESignalProcessing::DiffuseAndAmbientOcclusion);
 }
 
