Makes sure all View.GeneralPurposeTweak are commented out.

diff --git a/Engine/Shaders/Private/DiaphragmDOF/DOFRecombine.usf b/Engine/Shaders/Private/DiaphragmDOF/DOFRecombine.usf
--- a/Engine/Shaders/Private/DiaphragmDOF/DOFRecombine.usf
+++ b/Engine/Shaders/Private/DiaphragmDOF/DOFRecombine.usf
@@ -1110,8 +1110,8 @@ void RecombineMainCS(
 		}
 
 		OutputFinalSceneColor.rgb *= DebugColor;
-		if (View.GeneralPurposeTweak > 0.5)
-			OutputFinalSceneColor.rgb = DebugColor;
+		//if (View.GeneralPurposeTweak > 0.5)
+		//	OutputFinalSceneColor.rgb = DebugColor;
 	}
 	#endif
 	
@@ -1120,41 +1120,41 @@ void RecombineMainCS(
 		OutputFinalSceneColor = GatherParameters.KernelPixelRadius;
 	}
 	#elif 0 // DEBUG
-	if (View.GeneralPurposeTweak > 0.5)
-	{
-		//OutputFinalSceneColor = ForegroundColor;
-		//OutputFinalSceneColor = HoleFillingAdditiveColor;
-		//OutputFinalSceneColor = BackgroundColor.a;
-		//OutputFinalSceneColor = SlightOutOfFocusConvolution_SceneColor.SampleLevel(GlobalPointClampedSampler, InputParameters.DOFBufferUV, 0).a;
-		//OutputFinalSceneColor = GatherBackgroundUnpremultipliedColor;
-		//OutputFinalSceneColor.a = GatherBackgroundTranslucency;
-	}
+	//if (View.GeneralPurposeTweak > 0.5)
+	//{
+	//	//OutputFinalSceneColor = ForegroundColor;
+	//	//OutputFinalSceneColor = HoleFillingAdditiveColor;
+	//	//OutputFinalSceneColor = BackgroundColor.a;
+	//	//OutputFinalSceneColor = SlightOutOfFocusConvolution_SceneColor.SampleLevel(GlobalPointClampedSampler, InputParameters.DOFBufferUV, 0).a;
+	//	//OutputFinalSceneColor = GatherBackgroundUnpremultipliedColor;
+	//	//OutputFinalSceneColor.a = GatherBackgroundTranslucency;
+	//}
 	#elif 0 // DEBUG
-	if (View.GeneralPurposeTweak > 0.5)
-	{
-		OutputFinalSceneColor = GatherBackgroundAdditiveColor;
-		if (View.GeneralPurposeTweak > 1.5)
-			OutputFinalSceneColor = GatherBackgroundTranslucency;
-		if (View.GeneralPurposeTweak > 2.5)
-			OutputFinalSceneColor = BackgroundColor;
-		if (View.GeneralPurposeTweak > 3.5)
-			OutputFinalSceneColor = ForegroundColor;
-		if (View.GeneralPurposeTweak > 4.5)
-			OutputFinalSceneColor = GatherForegroundTranslucency;
-		
-		//OutputFinalSceneColor = GatherForegroundTranslucency;
-		//OutputFinalSceneColor = BackgroundConvolution_SceneColor.SampleLevel(GlobalBilinearClampedSampler, InputParameters.DOFBufferUV, 0);
-
-
-		//if (View.GeneralPurposeTweak > 3.5)
-		//	OutputFinalSceneColor = PostprocessInput3.SampleLevel(GlobalBilinearClampedSampler, InputParameters.DOFBufferUV, 0);
-		//if (View.GeneralPurposeTweak > 4.5)
-		//	OutputFinalSceneColor = PostprocessInput3.SampleLevel(GlobalBilinearClampedSampler, InputParameters.DOFBufferUV, 0).a;
-
-		//OutputFinalSceneColor = BackgroundColor;
-		//OutputFinalSceneColor = PostprocessInput3.SampleLevel(GlobalPointClampedSampler, InputParameters.DOFBufferUV, 0);
-		//OutputFinalSceneColor = ForegroundHoleFillingConvolution_SceneColor.SampleLevel(GlobalPointClampedSampler, InputParameters.DOFBufferUV, 0);
-	}
+	//if (View.GeneralPurposeTweak > 0.5)
+	//{
+	//	OutputFinalSceneColor = GatherBackgroundAdditiveColor;
+	//	if (View.GeneralPurposeTweak > 1.5)
+	//		OutputFinalSceneColor = GatherBackgroundTranslucency;
+	//	if (View.GeneralPurposeTweak > 2.5)
+	//		OutputFinalSceneColor = BackgroundColor;
+	//	if (View.GeneralPurposeTweak > 3.5)
+	//		OutputFinalSceneColor = ForegroundColor;
+	//	if (View.GeneralPurposeTweak > 4.5)
+	//		OutputFinalSceneColor = GatherForegroundTranslucency;
+	//	
+	//	//OutputFinalSceneColor = GatherForegroundTranslucency;
+	//	//OutputFinalSceneColor = BackgroundConvolution_SceneColor.SampleLevel(GlobalBilinearClampedSampler, InputParameters.DOFBufferUV, 0);
+	//
+	//
+	//	//if (View.GeneralPurposeTweak > 3.5)
+	//	//	OutputFinalSceneColor = PostprocessInput3.SampleLevel(GlobalBilinearClampedSampler, InputParameters.DOFBufferUV, 0);
+	//	//if (View.GeneralPurposeTweak > 4.5)
+	//	//	OutputFinalSceneColor = PostprocessInput3.SampleLevel(GlobalBilinearClampedSampler, InputParameters.DOFBufferUV, 0).a;
+	//
+	//	//OutputFinalSceneColor = BackgroundColor;
+	//	//OutputFinalSceneColor = PostprocessInput3.SampleLevel(GlobalPointClampedSampler, InputParameters.DOFBufferUV, 0);
+	//	//OutputFinalSceneColor = ForegroundHoleFillingConvolution_SceneColor.SampleLevel(GlobalPointClampedSampler, InputParameters.DOFBufferUV, 0);
+	//}
 	#endif
 	
 	#if 1 // Lower VGPR footprint.
diff --git a/Engine/Shaders/Private/SSRT/SSRTDiffuseIndirect.usf b/Engine/Shaders/Private/SSRT/SSRTDiffuseIndirect.usf
--- a/Engine/Shaders/Private/SSRT/SSRTDiffuseIndirect.usf
+++ b/Engine/Shaders/Private/SSRT/SSRTDiffuseIndirect.usf
@@ -413,10 +413,10 @@ void MainCS(
 			{
 				uint DestPos = SelectedBinStartingPos + OffsetInBin;
 
-				#if 1 // Uncomment me to measure performance improvement.
-					if (View.GeneralPurposeTweak > 0.5)
-						DestPos = GroupPixelId + RaySequenceId * TILE_PIXEL_COUNT;
-				#endif
+				//#if 1 // Uncomment me to measure performance improvement.
+				//	if (View.GeneralPurposeTweak > 0.5)
+				//		DestPos = GroupPixelId + RaySequenceId * TILE_PIXEL_COUNT;
+				//#endif
 		
 				#if LANE_PER_GROUPS > 256
 					#error "LANE_PER_GROUPS needs to be encodable in 8bits."
diff --git a/Engine/Shaders/Private/ScreenSpaceDenoise/SSDTemporalAccumulation.usf b/Engine/Shaders/Private/ScreenSpaceDenoise/SSDTemporalAccumulation.usf
--- a/Engine/Shaders/Private/ScreenSpaceDenoise/SSDTemporalAccumulation.usf
+++ b/Engine/Shaders/Private/ScreenSpaceDenoise/SSDTemporalAccumulation.usf
@@ -1069,7 +1069,7 @@ void TemporallyAccumulate(
 				SafeRcp(CurrentFrameSample.SampleCount) *
 				HistorySample.SampleCount);
 
-			if (View.GeneralPurposeTweak > 0.5)
+			//if (View.GeneralPurposeTweak > 0.5)
 			HistorySamples.Array[BatchedSignalId] = HistorySample;
 		}
 	}
diff --git a/Engine/Shaders/Private/ShadowPercentageCloserFiltering.ush b/Engine/Shaders/Private/ShadowPercentageCloserFiltering.ush
--- a/Engine/Shaders/Private/ShadowPercentageCloserFiltering.ush
+++ b/Engine/Shaders/Private/ShadowPercentageCloserFiltering.ush
@@ -166,14 +166,14 @@ float DirectionalPCSS(FPCSSSamplerSettings Settings, float2 ShadowPosition, floa
 
 	float MinorAnisotropyFactor = clamp(MinorAnisotropy / SearchRadius, 1, 7);
 	float MajorAnisotropyFactor = clamp(MajorAnisotropy / SearchRadius, 1, 7);
-	if (View.GeneralPurposeTweak < 1.5) {
-		MajorAnisotropyFactor = 1;
-		MinorAnisotropyFactor = 1;
-	}
-	if (View.GeneralPurposeTweak < 2.5) {
-		MajorAnisotropyFactor = clamp(length(ShadowAnisotropicVector) / SearchRadius, 1, 7);
-		MinorAnisotropyFactor = 1;
-	}
+	//if (View.GeneralPurposeTweak < 1.5) {
+	//	MajorAnisotropyFactor = 1;
+	//	MinorAnisotropyFactor = 1;
+	//}
+	//if (View.GeneralPurposeTweak < 2.5) {
+	//	MajorAnisotropyFactor = clamp(length(ShadowAnisotropicVector) / SearchRadius, 1, 7);
+	//	MinorAnisotropyFactor = 1;
+	//}
 	//return MajorAnisotropyFactor - MinorAnisotropyFactor;
 	//return (MinorAnisotropyFactor - 1);
 
