Fixes long standing issues in screen space denoiser's temporal accumulation where history data could turn stale when all the clamping box was invalid.

diff --git a/Engine/Shaders/Private/ScreenSpaceDenoise/SSDTemporalAccumulation.usf b/Engine/Shaders/Private/ScreenSpaceDenoise/SSDTemporalAccumulation.usf
--- a/Engine/Shaders/Private/ScreenSpaceDenoise/SSDTemporalAccumulation.usf
+++ b/Engine/Shaders/Private/ScreenSpaceDenoise/SSDTemporalAccumulation.usf
@@ -777,13 +777,19 @@ void TemporallyAccumulate(
 			
 			FSSDCompressedSignalAccumulatorArray UnusedCompressedAccumulators = CreateUninitialisedCompressedAccumulatorArray();
 
+			FSSDSignalArray CenterSample = CurrentFrameSamples;
+			if (KernelConfig.bNormalizeSample)
+			{
+				CenterSample = NormalizeToOneSampleArray(CurrentFrameSamples);
+			}
+
 			AccumulateRefSampleAsKernelCenter(
 				KernelConfig,
 				/* inout */ SignalAccumulators,
 				/* inout */ UnusedCompressedAccumulators,
 				KernelConfig.RefBufferUV,
 				RefSceneMetadata,
-				NormalizeToOneSampleArray(CurrentFrameSamples));
+				CenterSample);
 		}
 
 		{
@@ -934,9 +940,19 @@ void TemporallyAccumulate(
 			{
 				FSSDSignalSample NeighborMin = SignalAccumulators.Array[BatchedSignalId].Min;
 				FSSDSignalSample NeighborMax = SignalAccumulators.Array[BatchedSignalId].Max;
-		
+
+				// If no neighborhood had a sample the maximum sample count will remain 0.
+				bool bIsValid = NeighborMax.SampleCount > 0.0;
+
 				float RejectedSampleCount = 0;
 				HistorySample = MulSignal(TransformSignalForPostRejection(ClampSignal(TransformInputBufferForPreRejection(NormalizeToOneSample(HistorySample)), NeighborMin, NeighborMax)), HistorySample.SampleCount - RejectedSampleCount);
+
+				// All the clamping box was invalid, so invalidate the history sample too.
+				FLATTEN
+				if (!bIsValid)
+				{
+					HistorySample = CreateSignalSampleFromScalarValue(0.0);
+				}
 			}
 			#endif
 
