Deletes HISTORY_REJECTION_ADAPTIVE_INTEGRATION

diff --git a/Engine/Shaders/Private/ScreenSpaceDenoise/SSDTemporalAccumulation.usf b/Engine/Shaders/Private/ScreenSpaceDenoise/SSDTemporalAccumulation.usf
--- a/Engine/Shaders/Private/ScreenSpaceDenoise/SSDTemporalAccumulation.usf
+++ b/Engine/Shaders/Private/ScreenSpaceDenoise/SSDTemporalAccumulation.usf
@@ -19,7 +19,6 @@
 	#define HISTORY_REJECTION_DISABLED 0
 	#define HISTORY_REJECTION_MINMAX_BOUNDARIES 1
 	#define HISTORY_REJECTION_VAR_BOUNDARIES 2
-	#define HISTORY_REJECTION_ADAPTIVE_INTEGRATION 3
 
 
 //------------------------------------------------------- CONFIGS
@@ -153,7 +152,6 @@
 	#define CONFIG_HISTORY_TEXTURE_COUNT 2
 
 	#define CONFIG_HISTORY_REJECTION HISTORY_REJECTION_MINMAX_BOUNDARIES
-	//#define CONFIG_HISTORY_REJECTION HISTORY_REJECTION_ADAPTIVE_INTEGRATION
 	#define CONFIG_MULTIPLEXED_SIGNALS_PER_SIGNAL_DOMAIN 3
 
 	#define CONFIG_BILATERAL_PRESET BILATERAL_PRESET_DIFFUSE
@@ -175,7 +173,6 @@
 	#define CONFIG_HISTORY_TEXTURE_COUNT 4
 
 	#define CONFIG_HISTORY_REJECTION HISTORY_REJECTION_MINMAX_BOUNDARIES
-	//#define CONFIG_HISTORY_REJECTION HISTORY_REJECTION_ADAPTIVE_INTEGRATION
 	#define CONFIG_MULTIPLEXED_SIGNALS_PER_SIGNAL_DOMAIN 4
 
 	// Given it's a spherical harmonic that store directionality, only need position based rejection.
@@ -977,81 +974,6 @@ void TemporallyAccumulate(
 	}
 	#endif // CONFIG_HISTORY_REJECTION > 0
 	
-	// Blends the current and history together.
-	#if CONFIG_HISTORY_REJECTION == HISTORY_REJECTION_ADAPTIVE_INTEGRATION
-	{
-		#if CONFIG_SIGNAL_BATCH_SIZE != 1
-			#error Unsupported.
-		#endif
-	
-		#if !COMPILE_SIGNAL_COLOR
-			#error Need to compile signal color.
-		#endif
-
-		const uint BatchedSignalId = 0;
-	
-		FSSDSignalSample CurrentFrameSample = CurrentFrameSamples.Array[BatchedSignalId];
-		FSSDSignalSample HistorySample = HistorySamples.Array[BatchedSignalId];
-
-		// Construct the current frame sample of the temporal analysis.
-		FSSDSignalArray CurrentTemporalAnalysis = CreateSignalArrayFromScalarValue(0.0);
-		{
-			float Luma = GetSignalLuma4(CurrentFrameSample, STANDARD_BUFFER_COLOR_SPACE);
-			float SampleCount = CurrentFrameSample.SampleCount; // TODO(Denoiser): 1 instead?
-
-			CurrentTemporalAnalysis.Array[1].SampleCount = SampleCount;
-			CurrentTemporalAnalysis.Array[1].SceneColor.x = Luma * SampleCount;
-			CurrentTemporalAnalysis.Array[2].SampleCount = SampleCount;
-			CurrentTemporalAnalysis.Array[2].SceneColor.x = (Luma * Luma) * SampleCount;
-		}
-
-		// Update the temporal analysis.
-		FSSDSignalArray TemporalAnalysis;
-		{
-			const float MinimumAnalasisHisteresis = 0.3; // TODO(Denoiser): expose cvar
-
-			float PreviousFrameWeight = min(HistorySample.SampleCount, HistorySamples.Array[1].SampleCount * rcp(MinimumAnalasisHisteresis) - CurrentFrameSample.SampleCount);
-			float PreviousFrameMultiplier = HistorySample.SampleCount > 0 ? PreviousFrameWeight / HistorySample.SampleCount : 0;
-			
-			UNROLL
-			for (uint SignalMultiplexId = 1; SignalMultiplexId < 3; SignalMultiplexId++)
-			{
-				TemporalAnalysis.Array[SignalMultiplexId] = Add(Mul(HistorySamples.Array[SignalMultiplexId], PreviousFrameMultiplier), CurrentTemporalAnalysis.Array[SignalMultiplexId]);
-			}
-		}
-
-		// Compute the histeresis of the main signal.
-		float MainHisteresiss = 0;
-		{
-			float InvSampleCount = rcp(TemporalAnalysis.Array[1].SampleCount);
-
-			float AvgLuma = TemporalAnalysis.Array[2].SceneColor.x * InvSampleCount;
-			float VarianceLuma = TemporalAnalysis.Array[2].SceneColor.x * InvSampleCount - AvgLuma * AvgLuma;
-			float StdDevLuma = sqrt(VarianceLuma);
-
-			float LikelyNess = 0.5 * sqrt(StdDevLuma * HistorySample.SceneColor.x) / (StdDevLuma + HistorySample.SceneColor.x);
-
-			MainHisteresiss = 0.3 / max(LikelyNess * StdDevLuma, 0.3); 
-		}
-
-		// Accumulate the actual history.
-		{
-			float PreviousFrameWeight = min(HistorySample.SampleCount, CurrentFrameSample.SampleCount * rcp(MainHisteresiss) - CurrentFrameSample.SampleCount);
-			float PreviousFrameMultiplier = HistorySample.SampleCount > 0 ? PreviousFrameWeight / HistorySample.SampleCount : 0;
-			
-			HistorySample = AddSignal(MulSignal(HistorySample, PreviousFrameMultiplier), CurrentFrameSample);
-		}
-		
-		HistorySamples.Array[BatchedSignalId] = HistorySample;
-
-		// Copy the temporal analysis into the history sample.
-		UNROLL
-		for (uint SignalMultiplexId = 1; SignalMultiplexId < 3; SignalMultiplexId++)
-		{
-			HistorySamples.Array[SignalMultiplexId] = TemporalAnalysis.Array[SignalMultiplexId];
-		}
-	}
-	#else // CONFIG_HISTORY_REJECTION != HISTORY_REJECTION_ADAPTIVE_INTEGRATION
 	{
 		UNROLL
 		for (uint BatchedSignalId = 0; BatchedSignalId < MAX_SIGNAL_BATCH_SIZE; BatchedSignalId++)
@@ -1092,7 +1014,6 @@ void TemporallyAccumulate(
 			HistorySamples.Array[BatchedSignalId] = HistorySample;
 		}
 	}
-	#endif
 
 	// Romoves NaN and negative numbers
 	// HistorySample = Minus(Min(Minus(HistorySample), CreateSignalSampleFromScalarValue(0.0)));
@@ -1101,17 +1022,7 @@ void TemporallyAccumulate(
 	uint MultiplexCount = 1;
 	FSSDSignalArray OutputSamples = CreateSignalArrayFromScalarValue(0.0);
 	{
-		#if CONFIG_HISTORY_REJECTION == HISTORY_REJECTION_ADAPTIVE_INTEGRATION
-		{
-			MultiplexCount = CONFIG_MULTIPLEXED_SIGNALS_PER_SIGNAL_DOMAIN;
-		}
-		#elif 1
-		{
-			MultiplexCount = CONFIG_SIGNAL_BATCH_SIZE;
-		}
-		#else
-			#error Unknown output mode.
-		#endif
+		MultiplexCount = CONFIG_SIGNAL_BATCH_SIZE;
 		
 		UNROLL
 		for (uint BatchedSignalId = 0; BatchedSignalId < MultiplexCount; BatchedSignalId++)
