Implements SAMPLE_SET_2X2_ADAPTIVE for denoiser history reprojection

diff --git a/Engine/Shaders/Private/ScreenSpaceDenoise/SSDDefinitions.ush b/Engine/Shaders/Private/ScreenSpaceDenoise/SSDDefinitions.ush
--- a/Engine/Shaders/Private/ScreenSpaceDenoise/SSDDefinitions.ush
+++ b/Engine/Shaders/Private/ScreenSpaceDenoise/SSDDefinitions.ush
@@ -205,6 +205,9 @@
 	// Stocastic 2x2 that only take one sample.
 	#define SAMPLE_SET_2X2_STOCASTIC 13
 
+	// Adapt between bilinear and stocastic 2x2.
+	#define SAMPLE_SET_2X2_ADAPTIVE 17
+
 	// Square kernel
 	#define SAMPLE_SET_3X3 2
 	#define SAMPLE_SET_3X3_SOBEK2018 3
diff --git a/Engine/Shaders/Private/ScreenSpaceDenoise/SSDSpatialKernel.ush b/Engine/Shaders/Private/ScreenSpaceDenoise/SSDSpatialKernel.ush
--- a/Engine/Shaders/Private/ScreenSpaceDenoise/SSDSpatialKernel.ush
+++ b/Engine/Shaders/Private/ScreenSpaceDenoise/SSDSpatialKernel.ush
@@ -179,6 +179,9 @@ struct FSSDKernelConfig
 	// Runtime to force accumulating all sample.
 	bool bForceAllAccumulation;
 
+	// Runtime whether this pixel is dynamic object.
+	bool bIsDynamicPixel;
+
 	// Runtime selection of a track of sample.
 	uint SampleTrackId;
 
@@ -259,6 +262,7 @@ FSSDKernelConfig CreateKernelConfig()
 	KernelConfig.RefBufferUV = 0.0;
 	KernelConfig.bForceKernelCenterAccumulation = false;
 	KernelConfig.bForceAllAccumulation = false;
+	KernelConfig.bIsDynamicPixel = false;
 	KernelConfig.SampleTrackId = 0;
 	KernelConfig.MajorAxis = 0.0;
 	KernelConfig.MajorPixelRadius = 0.0;
@@ -1010,7 +1014,7 @@ void SampleAndAccumulateCenterSampleAsItsOwnCluster(
 
 #if COMPILE_BOX_KERNEL
 
-void AccumulateBilinear(
+void AccumulateBilinear2x2Kernel(
 	FSSDKernelConfig KernelConfig,
 	FSSDTexture2D SignalBuffer0,
 	FSSDTexture2D SignalBuffer1,
@@ -1027,17 +1031,59 @@ void AccumulateBilinear(
 		KernelConfig.BufferSizeAndInvSize.zw,
 		MipLevelPow2, rcp(MipLevelPow2));
 	
-	UNROLL_N(4)
-	for (uint i = 0; i < 4; i++)
+	bool bUseStocasticBilinear = false;
+	if (KernelConfig.SampleSet == SAMPLE_SET_2X2_STOCASTIC)
 	{
-		float2 SampleOffset = BilinearSamplingOffsets2x2[i];
-		
+		bUseStocasticBilinear = true;
+	}
+	else if (KernelConfig.SampleSet == SAMPLE_SET_2X2_ADAPTIVE)
+	{
+		bUseStocasticBilinear = !KernelConfig.bIsDynamicPixel;
+	}
+
+	float2 SampleBufferUVArray[4];
+	float BilinearWeightArray[4];
+
+	FLATTEN
+	if (bUseStocasticBilinear)
+	{
+		float2 SampleOffset;
+		float WeigthAccumulation = 0.0;
+
+		UNROLL_N(4)
+		for (uint i = 0; i < 4; i++)
+		{
+			FLATTEN
+			if (KernelConfig.Randoms[0] > WeigthAccumulation)
+				SampleOffset = BilinearSamplingOffsets2x2[i];
+
+			WeigthAccumulation += GetSampleWeight(BilinearInfos, i);
+
+			BilinearWeightArray[i] = 0.0;
+			SampleBufferUVArray[i] = 0.0;
+		}
+
 		// TODO(Denoiser): could be more ALU efficient for this.
 		// TODO(Denoiser): -0.5 full res pixel to ensure always select the mip, regardless of mantissa precision?
-		float2 SampleBufferUV = (BilinearInfos.TopLeftPixelCoord + (SampleOffset + 0.5)) * MipLevelPow2 * KernelConfig.BufferSizeAndInvSize.zw;
+		SampleBufferUVArray[0] = (BilinearInfos.TopLeftPixelCoord + (SampleOffset + 0.5)) * MipLevelPow2 * KernelConfig.BufferSizeAndInvSize.zw;
+		BilinearWeightArray[0] = 1.0;
+	}
+	else
+	{
+		UNROLL_N(4)
+		for (uint i = 0; i < 4; i++)
+		{
+			float2 SampleOffset = BilinearSamplingOffsets2x2[i];
 		
-		float BilinearWeight = GetSampleWeight(BilinearInfos, i);
+			// TODO(Denoiser): could be more ALU efficient for this.
+			// TODO(Denoiser): -0.5 full res pixel to ensure always select the mip, regardless of mantissa precision?
+			SampleBufferUVArray[i] = (BilinearInfos.TopLeftPixelCoord + (SampleOffset + 0.5)) * MipLevelPow2 * KernelConfig.BufferSizeAndInvSize.zw;
 		
+			BilinearWeightArray[i] = GetSampleWeight(BilinearInfos, i);
+		}
+	}
+
+	{
 		SampleAndAccumulateMultiplexedSignals(
 			KernelConfig,
 			SignalBuffer0,
@@ -1046,60 +1092,33 @@ void AccumulateBilinear(
 			SignalBuffer3,
 			/* inout */ UncompressedAccumulators,
 			/* inout */ CompressedAccumulators,
-			SampleBufferUV,
+			SampleBufferUVArray[0],
 			/* MipLevel = */ 0.0,
-			BilinearWeight,
+			BilinearWeightArray[0],
 			/* bForceSample = */ false);
 	}
-} // AccumulateBilinear()
 
-void AccumulateStocasticBilinear(
-	FSSDKernelConfig KernelConfig,
-	FSSDTexture2D SignalBuffer0,
-	FSSDTexture2D SignalBuffer1,
-	FSSDTexture2D SignalBuffer2,
-	FSSDTexture2D SignalBuffer3,
-	inout FSSDSignalAccumulatorArray UncompressedAccumulators,
-	inout FSSDCompressedSignalAccumulatorArray CompressedAccumulators)
-{
-	const float MipLevelPow2 = 1;
-
-	FBilinearSampleInfos BilinearInfos = GetBilinearSampleLevelInfosEx(
-		KernelConfig.BufferUV,
-		KernelConfig.BufferSizeAndInvSize.xy,
-		KernelConfig.BufferSizeAndInvSize.zw,
-		MipLevelPow2, rcp(MipLevelPow2));
-	
-	float2 SampleOffset;
-	float WeigthAccumulation = 0.0;
-
-	UNROLL_N(4)
-	for (uint i = 0; i < 4; i++)
+	BRANCH
+	if (!bUseStocasticBilinear)
 	{
-		SampleOffset = BilinearSamplingOffsets2x2[i];
-		WeigthAccumulation += GetSampleWeight(BilinearInfos, i);
-		
-		if (KernelConfig.Randoms[0] < WeigthAccumulation)
-			break;
+		UNROLL_N(3)
+		for (uint i = 1; i < 4; i++)
+		{
+			SampleAndAccumulateMultiplexedSignals(
+				KernelConfig,
+				SignalBuffer0,
+				SignalBuffer1,
+				SignalBuffer2,
+				SignalBuffer3,
+				/* inout */ UncompressedAccumulators,
+				/* inout */ CompressedAccumulators,
+				SampleBufferUVArray[i],
+				/* MipLevel = */ 0.0,
+				BilinearWeightArray[i],
+				/* bForceSample = */ false);
+		}
 	}
-	
-	// TODO(Denoiser): could be more ALU efficient for this.
-	// TODO(Denoiser): -0.5 full res pixel to ensure always select the mip, regardless of mantissa precision?
-	float2 SampleBufferUV = (BilinearInfos.TopLeftPixelCoord + (SampleOffset + 0.5)) * MipLevelPow2 * KernelConfig.BufferSizeAndInvSize.zw;
-		
-	SampleAndAccumulateMultiplexedSignals(
-		KernelConfig,
-		SignalBuffer0,
-		SignalBuffer1,
-		SignalBuffer2,
-		SignalBuffer3,
-		/* inout */ UncompressedAccumulators,
-		/* inout */ CompressedAccumulators,
-		SampleBufferUV,
-		/* MipLevel = */ 0.0,
-		/* KernelWeight = */ 1.0,
-		/* bForceSample = */ false);
-} // AccumulateStocasticBilinear()
+} // AccumulateBilinear2x2Kernel()
 
 void AccumulateSquareKernel(
 	FSSDKernelConfig KernelConfig,
@@ -2288,20 +2307,12 @@ void AccumulateKernel(
 				/* inout */ CompressedAccumulators);
 		}
 	}
-	else if (KernelConfig.SampleSet == SAMPLE_SET_2X2_BILINEAR)
-	{
-		AccumulateBilinear(
-			KernelConfig,
-			SignalBuffer0,
-			SignalBuffer1,
-			SignalBuffer2,
-			SignalBuffer3,
-			/* inout */ UncompressedAccumulators,
-			/* inout */ CompressedAccumulators);
-	}
-	else if (KernelConfig.SampleSet == SAMPLE_SET_2X2_STOCASTIC)
+	else if (
+		KernelConfig.SampleSet == SAMPLE_SET_2X2_BILINEAR ||
+		KernelConfig.SampleSet == SAMPLE_SET_2X2_STOCASTIC ||
+		KernelConfig.SampleSet == SAMPLE_SET_2X2_ADAPTIVE)
 	{
-		AccumulateStocasticBilinear(
+		AccumulateBilinear2x2Kernel(
 			KernelConfig,
 			SignalBuffer0,
 			SignalBuffer1,
diff --git a/Engine/Shaders/Private/ScreenSpaceDenoise/SSDTemporalAccumulation.usf b/Engine/Shaders/Private/ScreenSpaceDenoise/SSDTemporalAccumulation.usf
--- a/Engine/Shaders/Private/ScreenSpaceDenoise/SSDTemporalAccumulation.usf
+++ b/Engine/Shaders/Private/ScreenSpaceDenoise/SSDTemporalAccumulation.usf
@@ -454,6 +454,7 @@ void TemporallyAccumulate(
 
 	// Reproject to previous frame.
 	float3 HistoryScreenPosition = float3(DenoiserBufferUVToScreenPosition(SceneBufferUV), DeviceZ);
+	bool bIsDynamicPixel = false;
 	if (1)
 	{
 		float4 ThisClip = float4(HistoryScreenPosition, 1);
@@ -466,7 +467,6 @@ void TemporallyAccumulate(
 		{
 			float3 EncodedVelocity = SceneVelocityBuffer.SampleLevel(GlobalPointClampedSampler, SceneBufferUV, 0).xyz;
 			if (EncodedVelocity.x > 0.0)
-
 			{
 				Velocity = DecodeVelocityFromTexture(EncodedVelocity);
 			}
@@ -573,6 +573,7 @@ void TemporallyAccumulate(
 			
 			// VGPR configuration of the kernel.
 			KernelConfig.BufferUV = HistoryBufferUV + BufferUVBilinearCorrection;
+			KernelConfig.bIsDynamicPixel = bIsDynamicPixel;
 
 			#if CONFIG_PATCH_PREV_SCENE_DEPTH
 			{
