Additional Clear coat fixes:

1) Remove double application of clearcoat scalar to transmission
2) Remove duplicate Diffuse_Lambert() call
3) Apply Fresnel transmission to refracted specular lobe

diff --git a/Engine/Shaders/Private/ShadingModels.ush b/Engine/Shaders/Private/ShadingModels.ush
--- a/Engine/Shaders/Private/ShadingModels.ush
+++ b/Engine/Shaders/Private/ShadingModels.ush
@@ -277,7 +277,7 @@ float3 CalcThinTransmission(float NoL, float NoV, FGBufferData GBuffer)
 		float ThinDistance = NormalizedLayerThickness * (rcp(NoV) + rcp(NoL));
 		float Metallic = GBuffer.Metallic;
 		float ClearCoat = GBuffer.CustomData.x;
-		float AbsorptionMix = Metallic * ClearCoat;
+		float AbsorptionMix = Metallic;
 
 		Transmission = 1.0;
 		if (AbsorptionMix > 0.0)
@@ -431,9 +431,9 @@ FDirectLighting ClearCoatBxDF( FGBufferData GBuffer, half3 N, half3 V, half3 L,
 	float a2 = Pow2(Alpha);
 	float Energy = EnergyNormalization(a2, BottomContext.VoH, AreaLight);
 	//Lighting.Diffuse = (FresnelCoeff * Energy * Falloff * BottomContext.NoL) * Transmission * AreaLight.FalloffColor * Diffuse_Lambert(GBuffer.DiffuseColor);
-	float3 CommonDiffuse = (Energy * Falloff) * AreaLight.FalloffColor;
-	float3 DefaultDiffuse = Context.NoL * Diffuse_Lambert(GBuffer.DiffuseColor);
-	float3 RefractedDiffuse = (FresnelCoeff * BottomContext.NoL) * Transmission * Diffuse_Lambert(GBuffer.DiffuseColor);
+	float3 CommonDiffuse = (Energy * Falloff) * AreaLight.FalloffColor * Diffuse_Lambert(GBuffer.DiffuseColor);
+	float3 DefaultDiffuse = Context.NoL;
+	float3 RefractedDiffuse = (FresnelCoeff * BottomContext.NoL) * Transmission;
 	Lighting.Diffuse = CommonDiffuse * lerp(DefaultDiffuse, RefractedDiffuse, ClearCoat);
 	if (AreaLight.bIsRect)
 	{
@@ -470,7 +470,7 @@ FDirectLighting ClearCoatBxDF( FGBufferData GBuffer, half3 N, half3 V, half3 L,
 		// Note: reusing D, V, and F from refracted context to save computation for when ClearCoat < 1
 		float3 CommonSpecular = (Energy * Falloff * D2 * Vis2) * AreaLight.FalloffColor * F;
 		float3 DefaultSpecular = Context.NoL;
-		float3 RefractedSpecular = Transmission * BottomContext.NoL;
+		float3 RefractedSpecular = FresnelCoeff * Transmission * BottomContext.NoL;
 		Lighting.Specular += CommonSpecular * lerp(DefaultSpecular, RefractedSpecular, ClearCoat);
 	}
 
