Enable hair plugin
Fix hair CIS

diff --git a/Engine/Shaders/Private/HairStrands/HairStrandsVertexFactory.ush b/Engine/Shaders/Private/HairStrands/HairStrandsVertexFactory.ush
--- a/Engine/Shaders/Private/HairStrands/HairStrandsVertexFactory.ush
+++ b/Engine/Shaders/Private/HairStrands/HairStrandsVertexFactory.ush
@@ -274,7 +274,7 @@ FHairControlPoint GetVertexPosition(const FVertexInfo VertexInfo, bool bInvalidJ
 	// Create a degenerated quad the end of each strand to cut between each strands
 	// #hair_todo: This is not efficient for short strand like fur for instance. Need to revisit that at some point
 	const bool bIsInvalidQuad = (bInvalidJointVertex && Out.Type == HAIR_CONTROLPOINT_START && VertexInfo.IsTip == 1) || (bInvalidJointVertex && Out.Type == HAIR_CONTROLPOINT_END && VertexInfo.IsTip == 0) || VertexInfo.bForceInvalidQuad;
-	Out.Position = bIsInvalidQuad ? float3(1.#INF, 1.#INF, 1.#INF) : Out.Position;
+	Out.Position = bIsInvalidQuad ? float3(INFINITE_FLOAT, INFINITE_FLOAT, INFINITE_FLOAT) : Out.Position;
 
 	return Out;
 }
@@ -308,7 +308,7 @@ float3 GetVertexPreviousPosition(FVertexFactoryInput Input, bool bInvalidJointVe
 	// Create a degenerated quad the end of each strand to cut between each strands
 	// #hair_todo: This is not efficient for short strand like fur for instance. Need to revisit that at some point
 	const bool bIsInvalidQuad = (bInvalidJointVertex && Out.Type == HAIR_CONTROLPOINT_START && VertexInfo.IsTip == 1) || (bInvalidJointVertex && Out.Type == HAIR_CONTROLPOINT_END && VertexInfo.IsTip == 0) || VertexInfo.bForceInvalidQuad;
-	Out.Position = bIsInvalidQuad ? float3(1.#INF, 1.#INF, 1.#INF) : Out.Position;
+	Out.Position = bIsInvalidQuad ? float3(INFINITE_FLOAT, INFINITE_FLOAT, INFINITE_FLOAT) : Out.Position;
 
 	return Out.Position;
 }
@@ -409,8 +409,8 @@ half3x3 CalcTangentToLocal(FVertexFactoryInput Input, inout float TangentSign)
 
 half3x3 CalcTangentToWorld(FVertexFactoryIntermediates Intermediates, half3x3 TangentToLocal)
 {
-	half3x3 TangentToWorld = CalcTangentToWorldNoScale(Intermediates, TangentToLocal);
-	return TangentToWorld;
+	half3x3 LocalTangentToWorld = CalcTangentToWorldNoScale(Intermediates, TangentToLocal);
+	return LocalTangentToWorld;
 }
 
 float3 GetTangent(FVertexFactoryInput Input)
diff --git a/Engine/Shaders/Private/HairStrands/HairStrandsVisibilityCommon.ush b/Engine/Shaders/Private/HairStrands/HairStrandsVisibilityCommon.ush
--- a/Engine/Shaders/Private/HairStrands/HairStrandsVisibilityCommon.ush
+++ b/Engine/Shaders/Private/HairStrands/HairStrandsVisibilityCommon.ush
@@ -30,7 +30,7 @@ float3 DecodeTangent(float3 N)
 
 uint float4ToUint(float4 v)
 {
-	uint4 i = v * 255;
+	uint4 i = uint4(v.x * 255, v.y * 255, v.z * 255, v.w * 255);
 	return (0xFF & i.w) << 24 | (0xFF & i.z) << 16 | (0xFF & i.y) << 8 | (0xFF & i.x);
 }
 
@@ -89,7 +89,8 @@ float From16bitCoverage(uint Coverage16bit)
 
 uint3 QuantizeTo8Bits(float3 T)
 {
-	return saturate((T + float(1).xxx) * 0.5f) * 0xFF;
+	const float3 Quanta = saturate((T + float(1).xxx) * 0.5f) * 0xFF;
+	return uint3(Quanta.x, Quanta.y, Quanta.z);
 }
 
 float3 From8bits(float3 In8bits)
@@ -314,8 +315,8 @@ struct FHairAttributes
 FHairAttributes DecodeHairAttribute(uint4 Attributes)
 {
 	float2 RootUV, RootIndex;
-	RootUV.x	= (Attributes.x & 0xFF) / 255.f;
-	RootUV.y	= 1-(((Attributes.x>>8) & 0xFF) / 255.f); // Invert V to compensate image origin flip
+	RootUV.x	= float(Attributes.x & 0xFF) / 255.f;
+	RootUV.y	= 1.0-(float((Attributes.x>>8) & 0xFF) / 255.f); // Invert V to compensate image origin flip
 	RootIndex.x	= Attributes.z & 0xFF;
 	RootIndex.y	= (Attributes.z >> 8) & 0xFF;
 
