Fix a bug of the fully rough on mobile.

diff --git a/Engine/Shaders/Private/MobileShadingModels.ush b/Engine/Shaders/Private/MobileShadingModels.ush
--- a/Engine/Shaders/Private/MobileShadingModels.ush
+++ b/Engine/Shaders/Private/MobileShadingModels.ush
@@ -83,7 +83,7 @@ void InitShadingModelContext( inout FMobileShadingModelContext ShadingModelConte
 #endif
 
 #if NONMETAL
-	ShadingModelContext.DiffuseColor = BaseColor;
+	ShadingModelContext.DiffuseColor = ShadingModelContext.BaseColor;
 	ShadingModelContext.SpecularColor = 0.04;
 #else
 
