Improve basic eye adaptation performance for mobile platform.

diff --git a/Engine/Shaders/Private/PostProcessMobile.usf b/Engine/Shaders/Private/PostProcessMobile.usf
--- a/Engine/Shaders/Private/PostProcessMobile.usf
+++ b/Engine/Shaders/Private/PostProcessMobile.usf
@@ -1245,7 +1245,7 @@ void AverageLuminance_MainCS(
 	half2 BufferUV;
 	float Weight;
 
-	SharedLuminance[GroupIndex] = float2(0.0f,0.0f);
+	float2 GroupLuminance = float2(0.0f, 0.0f);
 
 	LOOP for (uint y = 0; y < LOOP_SIZEY; ++y)
 	{
@@ -1261,12 +1261,14 @@ void AverageLuminance_MainCS(
 
 				Weight = max(AdaptationWeightTexture(BufferUV), 0.05f);
 
-				SharedLuminance[GroupIndex].x += Texture2DSample(InputTexture, InputSampler, BufferUV).x * Weight;
-				SharedLuminance[GroupIndex].y += Weight;
+				GroupLuminance.x += Texture2DSample(InputTexture, InputSampler, BufferUV).x * Weight;
+				GroupLuminance.y += Weight;
 			}
 		}
 	}
 
+	SharedLuminance[GroupIndex] = GroupLuminance;
+
 	GroupMemoryBarrierWithGroupSync();
 
 	UNROLL for (uint cutoff = (ThreadGroupSize >> 1); cutoff > 0; cutoff >>= 1)
