Guard against negative dot products when calculating shading terms for clear coat bottom normal.

diff --git a/Engine/Shaders/Private/ShadingModels.ush b/Engine/Shaders/Private/ShadingModels.ush
--- a/Engine/Shaders/Private/ShadingModels.ush
+++ b/Engine/Shaders/Private/ShadingModels.ush
@@ -403,8 +403,8 @@ FDirectLighting ClearCoatBxDF( FGBufferData GBuffer, half3 N, half3 V, half3 L,
 		// If bottom-normal, update normal-based dot products:
 		float3 H = normalize(V + L);
 		Context.NoH = saturate(dot(N, H));
-		Context.NoV = dot(N, V);
-		Context.NoL = dot(N, L);
+		Context.NoV = saturate(dot(N, V));
+		Context.NoL = saturate(dot(N, L));
 
 		bool bNewtonIteration = true;
 		SphereMaxNoH(Context, AreaLight.SphereSinAlpha, bNewtonIteration);
