Applying saturation guard around absorption to prevent view-dependent, precision-based artifacts.

diff --git a/Engine/Shaders/Private/ShadingModels.ush b/Engine/Shaders/Private/ShadingModels.ush
--- a/Engine/Shaders/Private/ShadingModels.ush
+++ b/Engine/Shaders/Private/ShadingModels.ush
@@ -287,7 +287,7 @@ float3 CalcThinTransmission(float NoL, float NoV, FGBufferData GBuffer)
 			float3 TransmissionColor = Diffuse_Lambert(GBuffer.DiffuseColor);
 			float3 ExtinctionCoefficient = -log(TransmissionColor) / (2.0 * NormalizedLayerThickness);
 			float3 OpticalDepth = ExtinctionCoefficient * max(ThinDistance - 2.0 * NormalizedLayerThickness, 0.0);
-			Transmission = exp(-OpticalDepth);
+			Transmission = saturate(exp(-OpticalDepth));
 			Transmission = lerp(1.0, Transmission, AbsorptionMix);
 		}
 	}
