ClearCoat fixes.

1) Transmitted color should be BaseColor
2) Removing duplicate application of ClearCoat factor to Fresnel term

diff --git a/Engine/Shaders/Private/ShadingModels.ush b/Engine/Shaders/Private/ShadingModels.ush
--- a/Engine/Shaders/Private/ShadingModels.ush
+++ b/Engine/Shaders/Private/ShadingModels.ush
@@ -284,7 +284,7 @@ float3 CalcThinTransmission(float NoL, float NoV, FGBufferData GBuffer)
 		{
 			// Base color represents reflected color viewed at 0 incidence angle, after being absorbed through the substrate.
 			// Because of this, extinction is normalized by traveling through layer thickness twice
-			float3 TransmissionColor = Diffuse_Lambert(GBuffer.DiffuseColor);
+			float3 TransmissionColor = Diffuse_Lambert(GBuffer.BaseColor);
 			float3 ExtinctionCoefficient = -log(TransmissionColor) / (2.0 * NormalizedLayerThickness);
 			float3 OpticalDepth = ExtinctionCoefficient * max(ThinDistance - 2.0 * NormalizedLayerThickness, 0.0);
 			Transmission = saturate(exp(-OpticalDepth));
@@ -375,7 +375,6 @@ FDirectLighting ClearCoatBxDF( FGBufferData GBuffer, half3 N, half3 V, half3 L,
 	float F0 = 0.04;
 	float Fc = Pow5(1 - Context.VoH);
 	float F = Fc + (1 - Fc) * F0;
-	F *= ClearCoat;
 
 	if (AreaLight.bIsRect)
 	{
