Aligning path tracing clear coat model's absorption factor to match the real-time model.

diff --git a/Engine/Shaders/Private/PathTracing/Material/PathTracingClearCoat.ush b/Engine/Shaders/Private/PathTracing/Material/PathTracingClearCoat.ush
--- a/Engine/Shaders/Private/PathTracing/Material/PathTracingClearCoat.ush
+++ b/Engine/Shaders/Private/PathTracing/Material/PathTracingClearCoat.ush
@@ -78,7 +78,7 @@ float3 CalcThinTransmission(float NoL, float NoV, FMaterialClosestHitPayload Pay
 		float ThinDistance = LayerThickness * (rcp(NoV) + rcp(NoL));
 		float Metallic = GetMetallic(Payload);
 		float ClearCoat = GetClearCoat(Payload);
-		float AbsorptionMix = Metallic * ClearCoat;
+		float AbsorptionMix = Metallic;
 
 		// Only absorb for metallic materials
 		Transmission = 1.0;
