Fixed: Correctly handle cases in mobile PP materials where viewport does not match buffer size

diff --git a/Engine/Shaders/Private/PostProcessMaterialShaders.usf b/Engine/Shaders/Private/PostProcessMaterialShaders.usf
--- a/Engine/Shaders/Private/PostProcessMaterialShaders.usf
+++ b/Engine/Shaders/Private/PostProcessMaterialShaders.usf
@@ -269,18 +269,15 @@ uint bFlipYAxis;
 void MainVS(
 	in float4 InPosition : ATTRIBUTE0,
 	in float2 InTexCoord : ATTRIBUTE1,
-	out FPostProcessMaterialVSToPS Output,
-	out float2 OutViewportUV : TEXCOORD8
+	out FPostProcessMaterialVSToPS Output
 	)
 {
 	Output = (FPostProcessMaterialVSToPS)0;
 	float2 OutUV;
 	DrawRectangle(InPosition, InTexCoord, Output.Position, OutUV);
-
-	OutViewportUV = InPosition.xy;
-	
+			
 #if NUM_TEX_COORD_INTERPOLATORS
-	FMaterialVertexParameters VertexParameters = GetPostProcessMaterialVSParameters(OutUV);
+	FMaterialVertexParameters VertexParameters = GetPostProcessMaterialVSParameters(InPosition.xy);
 
 	float2 CustomizedUVs[NUM_TEX_COORD_INTERPOLATORS];
 	GetMaterialCustomizedUVs(VertexParameters, CustomizedUVs);
@@ -289,11 +286,11 @@ void MainVS(
 	UNROLL
 	for (int CoordinateIndex = 0; CoordinateIndex < NUM_MATERIAL_TEXCOORDS; CoordinateIndex++)
 	{
-#if POST_PROCESS_AR_PASSTHROUGH
+	#if POST_PROCESS_AR_PASSTHROUGH
 		SetUV(Output, CoordinateIndex, OutUV);
-#else
+	#else
 		SetUV(Output, CoordinateIndex, InPosition.xy);
-#endif
+	#endif
 	}
 	UNROLL
 	for (int CoordinateIndex = NUM_MATERIAL_TEXCOORDS; CoordinateIndex < NUM_TEX_COORD_INTERPOLATORS; CoordinateIndex++)
@@ -314,14 +311,40 @@ void MainVS(
 
 void MainPS(
 	in FPostProcessMaterialVSToPS Input,
-	in float2 InViewportUV : TEXCOORD8,
 	out half4 OutColor : SV_Target0
 	)
 {
+	ResolvedView = ResolveView();
+	FMaterialPixelParameters Parameters = MakeInitializedMaterialPixelParameters();
+	FPixelMaterialInputs PixelMaterialInputs;
+
+	// can be optimized
+	float4 SvPosition = Input.Position;
+	if (bFlipYAxis)
+	{
+		SvPosition.y = PostProcessOutput_ViewportSize.y - SvPosition.y;
+	}
+	float2 ViewportUV = (SvPosition - PostProcessOutput_ViewportMin) * PostProcessOutput_ViewportSizeInverse;
+	float2 BufferUV = ViewportUVToBufferUV(ViewportUV);
+	
+#if NUM_TEX_COORD_INTERPOLATORS
+	UNROLL
+	for (int CoordinateIndex = 0; CoordinateIndex < NUM_MATERIAL_TEXCOORDS; CoordinateIndex++)
+	{
+		Parameters.TexCoords[CoordinateIndex] = ViewportUV;
+	}
+	UNROLL
+	for (int CoordinateIndex = NUM_MATERIAL_TEXCOORDS; CoordinateIndex < NUM_TEX_COORD_INTERPOLATORS; CoordinateIndex++)
+	{
+		Parameters.TexCoords[CoordinateIndex] = GetUV(Input, CoordinateIndex);
+	}
+#endif
+	Parameters.VertexColor = 1;
+		
 	float ClipValue = 1.0f;
 
 #if MOBILE_STENCIL_COMPARE_FUNCTION != MOBILE_STENCIL_COMPARE_NEVER && MOBILE_STENCIL_COMPARE_FUNCTION != MOBILE_STENCIL_COMPARE_ALWAYS
-	int MobileStencilValue = floor(Texture2DSample(MobileCustomStencilTexture, MobileCustomStencilTextureSampler, InViewportUV).r*255.0f + 0.5f);
+	int MobileStencilValue = floor(Texture2DSample(MobileCustomStencilTexture, MobileCustomStencilTextureSampler, BufferUV).r*255.0f + 0.5f);
 #endif
 
 #if MOBILE_STENCIL_COMPARE_FUNCTION == MOBILE_STENCIL_COMPARE_LESS
@@ -345,28 +368,9 @@ void MainPS(
 	clip(ClipValue);
 #endif
 
-	ResolvedView = ResolveView();
-	FMaterialPixelParameters Parameters = MakeInitializedMaterialPixelParameters();
-	FPixelMaterialInputs PixelMaterialInputs;
-
-#if NUM_TEX_COORD_INTERPOLATORS
-	UNROLL
-	for (int CoordinateIndex = 0; CoordinateIndex < NUM_TEX_COORD_INTERPOLATORS; CoordinateIndex++)
-	{
-		Parameters.TexCoords[CoordinateIndex] = GetUV(Input, CoordinateIndex);
-	}
-#endif
-
-	Parameters.VertexColor = 1;
-	
-	float4 SvPosition = Input.Position;
-	if (bFlipYAxis)
-	{
-		SvPosition.y = PostProcessOutput_ViewportSize.y - SvPosition.y;
-	}
 	// SceneDepth texture is not accessible during post-processing as we discard it at the end of mobile BasePass
     // instead fetch depth from SceneColor.A 
-	float Depth = Texture2DSample(MobileSceneTextures.SceneColorTexture, MobileSceneTextures.SceneColorTextureSampler, InViewportUV).a;
+	float Depth = Texture2DSample(MobileSceneTextures.SceneColorTexture, MobileSceneTextures.SceneColorTextureSampler, BufferUV).a;
 	float DeviceZ = ConvertToDeviceZ(Depth);
 	SvPosition.z = DeviceZ;
 	SvPosition.z = max(SvPosition.z, 1e-18);
