- Fixed some shader #includes for PlatformExtension locations

diff --git a/Engine/Shaders/Private/HdrCustomResolveShaders.usf b/Engine/Shaders/Private/HdrCustomResolveShaders.usf
--- a/Engine/Shaders/Private/HdrCustomResolveShaders.usf
+++ b/Engine/Shaders/Private/HdrCustomResolveShaders.usf
@@ -8,8 +8,6 @@
 
 #ifdef OVERRIDE_HDRCUSTOMRESOLVESHADERS_USH
 # include "/Platform/Private/HdrCustomResolveShaders.ush"
-#elif PS4_PROFILE
-# include "Platform/PS4/HdrCustomResolveShaders.ush"
 #else
 
 float4 HdrCustomResolveVS(uint Id : SV_VertexID) : SV_POSITION
diff --git a/Engine/Shaders/Private/WideCustomResolveShaders.usf b/Engine/Shaders/Private/WideCustomResolveShaders.usf
--- a/Engine/Shaders/Private/WideCustomResolveShaders.usf
+++ b/Engine/Shaders/Private/WideCustomResolveShaders.usf
@@ -8,8 +8,6 @@
 
 #ifdef OVERRIDE_WIDECUSTOMRESOLVESHADERS_USH
 # include "/Platform/Private/WideCustomResolveShaders.ush"
-#elif PS4_PROFILE
-# include "Platform/PS4/WideCustomResolveShaders.ush"
 #else
 
 float2 ResolveOrigin;
