Ray tracing translucency rays TMax bias value reverted to 4.24 value

diff --git a/Engine/Shaders/Private/RayTracing/RayTracingPrimaryRays.usf b/Engine/Shaders/Private/RayTracing/RayTracingPrimaryRays.usf
--- a/Engine/Shaders/Private/RayTracing/RayTracingPrimaryRays.usf
+++ b/Engine/Shaders/Private/RayTracing/RayTracingPrimaryRays.usf
@@ -116,7 +116,7 @@ RAY_TRACING_ENTRY_RAYGEN(RayTracingPrimaryRaysRGS)
 
 	if((ERayTracingPrimaryRaysFlag_UseGBufferForMaxDistance & PrimaryRayFlags) != 0) 
 	{
-		Ray.TMax = WorldSpaceDistance - 0.025;
+		Ray.TMax = WorldSpaceDistance - 0.1;
 	}
 
 	bool bAllowSkySampling;
