Permute shadow ray origin point if transmission is computed. This allows for correct computation of occluded transmission distance when computing shadows for two-sided foliage without the need for a secondary ray.

diff --git a/Engine/Shaders/Private/RayTracing/RayTracingOcclusionRGS.usf b/Engine/Shaders/Private/RayTracing/RayTracingOcclusionRGS.usf
--- a/Engine/Shaders/Private/RayTracing/RayTracingOcclusionRGS.usf
+++ b/Engine/Shaders/Private/RayTracing/RayTracingOcclusionRGS.usf
@@ -274,18 +274,6 @@ FOcclusionResult ComputeOcclusion(
 
 	const float3 WorldPosition = ReconstructWorldPositionFromDeviceZ(PixelCoord, DeviceZ);
 
-	// Cheap solution for double sided foliage to have valid lighing through the denoiser
-	if (ShadingModelID == SHADINGMODELID_TWOSIDED_FOLIAGE)
-	{
-	#if LIGHT_TYPE == LIGHT_TYPE_DIRECTIONAL
-		const float3 LightDirection = LightParameters.Direction;
-	#else
-		const float3 LightDirection = normalize(LightParameters.Position - WorldPosition);
-	#endif
-		if (dot(LightDirection, WorldNormal) < 0)
-			WorldNormal = -WorldNormal;
-	}
-
 	// For hair shading model, WorldNormal is actually the shading tangent.
 	// To generate proper bias, we compute a normal oriented toward the light.
 	// Normal clipping is removed from hair since the BSDF is spherical, rather 
@@ -309,9 +297,7 @@ FOcclusionResult ComputeOcclusion(
 		WorldNormal = LightDirection;
 		bApplyNormalCulling = false;
 	}
-
-	// Fixes clipping of eye caustic normals in raytrace shadows
-	if (ShadingModelID == SHADINGMODELID_EYE)
+	else if (ShadingModelID == SHADINGMODELID_EYE || ShadingModelID == SHADINGMODELID_TWOSIDED_FOLIAGE)
 	{
 		bApplyNormalCulling = false;
 	}
@@ -351,7 +337,19 @@ FOcclusionResult ComputeOcclusion(
 		uint Stencil = SceneStencilBuffer.Load(int3(PixelCoord, 0)) STENCIL_COMPONENT_SWIZZLE;
 		bool bDitheredLODFadingOut = Stencil & 1;
 		
+#if	ENABLE_TRANSMISSION
+		float NoL = dot(WorldNormal, Ray.Direction);
+		if (NoL > 0.0)
+		{
+			ApplyCameraRelativeDepthBias(Ray, PixelCoord, DeviceZ, WorldNormal, NormalBias + (bDitheredLODFadingOut ? 0.8f : 0.0f));
+		}
+		else
+		{
+			ApplyPositionBias(Ray, -WorldNormal, NormalBias);
+		}
+#else
 		ApplyCameraRelativeDepthBias(Ray, PixelCoord, DeviceZ, WorldNormal, NormalBias + (bDitheredLODFadingOut ? 0.8f : 0.0f));
+#endif
 
 		BRANCH
 		if (!bIsValidRay)
@@ -431,7 +429,7 @@ FOcclusionResult ComputeOcclusion(
 
 		Out.TransmissionDistance = ComputeDiffuseTransmission(V, WorldPosition, WorldNormal, MeanFreePath, Eta, LightParameters, LightSample, RandSequence, RejectionRetryMax);
 	}
-	else // if (ShadingModelID == SHADINGMODELID_EYE)
+	else // if ((ShadingModelID == SHADINGMODELID_EYE) || (ShadingModelID == SHADINGMODELID_TWOSIDED_FOLIAGE))
 	{
 		Out.TransmissionDistance = (LocalSamplesPerPixel > 0) ? Out.Visibility / LocalSamplesPerPixel : Out.Visibility;
 	}
