Add transmission to ray tracing SkyLight pass

diff --git a/Engine/Shaders/Private/RayTracing/RayTracingSkyLightEvaluation.ush b/Engine/Shaders/Private/RayTracing/RayTracingSkyLightEvaluation.ush
--- a/Engine/Shaders/Private/RayTracing/RayTracingSkyLightEvaluation.ush
+++ b/Engine/Shaders/Private/RayTracing/RayTracingSkyLightEvaluation.ush
@@ -143,14 +143,23 @@ void SkyLightEvaluate(
 		}
 
 		// Apply a depth bias based on if the sample world position came from the g-buffer or not
-		if (bGBufferSampleOrigin)
+		float NoL = dot(CurrentWorldNormal, Ray.Direction);
+		if (NoL > 0.0)
 		{
-			ApplyCameraRelativeDepthBias(Ray, SampleCoord, DeviceZ, CurrentWorldNormal, SkyLight.MaxNormalBias);
+			if (bGBufferSampleOrigin)
+			{
+				ApplyCameraRelativeDepthBias(Ray, SampleCoord, DeviceZ, CurrentWorldNormal, SkyLight.MaxNormalBias);
+			}
+			else
+			{
+				ApplyPositionBias(Ray, CurrentWorldNormal, SkyLight.MaxNormalBias);
+			}
 		}
 		else
 		{
-			ApplyPositionBias(Ray, CurrentWorldNormal, SkyLight.MaxNormalBias);
+			ApplyPositionBias(Ray, -CurrentWorldNormal, SkyLight.MaxNormalBias);
 		}
+		NoL = saturate(NoL);
 
 		// Trace a visibility ray
 		uint RayFlags = 0;
@@ -189,7 +198,6 @@ void SkyLightEvaluate(
 			const half3 N = WorldNormal;
 			const half3 V = -ViewDirection;
 			const half3 L = Ray.Direction;
-			const float NoL = saturate(dot(CurrentWorldNormal, L));
 			FDirectLighting LightingSample;
 			if (GBufferData.ShadingModelID == SHADINGMODELID_HAIR)
 			{
@@ -212,7 +220,12 @@ void SkyLightEvaluate(
 			float3 IncomingRadiance = EvaluateSkyLightRadiance(Ray.Direction);
 
 			ExitantRadiance += IncomingRadiance * Brdf * RayWeight;
-			DiffuseExitantRadiance += IncomingRadiance * (LightingSample.Diffuse + LightingSample.Transmission) * RayWeight;
+			float3 DiffuseThroughput = LightingSample.Diffuse;
+			if (SkyLight.bTransmission)
+			{
+				DiffuseThroughput += LightingSample.Transmission;
+			}
+			DiffuseExitantRadiance += IncomingRadiance * DiffuseThroughput * RayWeight;
 		}
 	}
 
diff --git a/Engine/Source/Runtime/Engine/Private/Components/SkyLightComponent.cpp b/Engine/Source/Runtime/Engine/Private/Components/SkyLightComponent.cpp
--- a/Engine/Source/Runtime/Engine/Private/Components/SkyLightComponent.cpp
+++ b/Engine/Source/Runtime/Engine/Private/Components/SkyLightComponent.cpp
@@ -184,6 +184,7 @@ FSkyLightSceneProxy::FSkyLightSceneProxy(const USkyLightComponent* InLightCompon
 	, bCastRayTracedShadow(InLightComponent->bCastRaytracedShadow)
 	, bAffectReflection(InLightComponent->bAffectReflection)
 	, bAffectGlobalIllumination(InLightComponent->bAffectGlobalIllumination)
+	, bTransmission(InLightComponent->bTransmission)
 	, OcclusionCombineMode(InLightComponent->OcclusionCombineMode)
 	, IndirectLightingIntensity(InLightComponent->IndirectLightingIntensity)
 	, VolumetricScatteringIntensity(FMath::Max(InLightComponent->VolumetricScatteringIntensity, 0.0f))
diff --git a/Engine/Source/Runtime/Engine/Public/SceneManagement.h b/Engine/Source/Runtime/Engine/Public/SceneManagement.h
--- a/Engine/Source/Runtime/Engine/Public/SceneManagement.h
+++ b/Engine/Source/Runtime/Engine/Public/SceneManagement.h
@@ -1076,6 +1076,7 @@ public:
 	uint8 bCastRayTracedShadow:1;
 	uint8 bAffectReflection:1;
 	uint8 bAffectGlobalIllumination:1;
+	uint8 bTransmission:1;
 	TEnumAsByte<EOcclusionCombineMode> OcclusionCombineMode;
 	float AverageBrightness;
 	float IndirectLightingIntensity;
diff --git a/Engine/Source/Runtime/Renderer/Private/RayTracing/RayTracingSkyLight.h b/Engine/Source/Runtime/Renderer/Private/RayTracing/RayTracingSkyLight.h
--- a/Engine/Source/Runtime/Renderer/Private/RayTracing/RayTracingSkyLight.h
+++ b/Engine/Source/Runtime/Renderer/Private/RayTracing/RayTracingSkyLight.h
@@ -19,6 +19,7 @@ BEGIN_GLOBAL_SHADER_PARAMETER_STRUCT(FSkyLightData, )
 	SHADER_PARAMETER(FIntVector, MipDimensions)
 	SHADER_PARAMETER(float, MaxNormalBias)
 	SHADER_PARAMETER(float, MaxShadowThickness)
+	SHADER_PARAMETER(int, bTransmission)
 	SHADER_PARAMETER_TEXTURE(TextureCube, Texture)
 	SHADER_PARAMETER_SAMPLER(SamplerState, TextureSampler)
 	SHADER_PARAMETER(FIntVector, TextureDimensions)
diff --git a/Engine/Source/Runtime/Renderer/Private/RayTracing/RaytracingSkylight.cpp b/Engine/Source/Runtime/Renderer/Private/RayTracing/RaytracingSkylight.cpp
--- a/Engine/Source/Runtime/Renderer/Private/RayTracing/RaytracingSkylight.cpp
+++ b/Engine/Source/Runtime/Renderer/Private/RayTracing/RaytracingSkylight.cpp
@@ -162,6 +162,7 @@ void SetupSkyLightParameters(
 		SkyLightData->SamplingStopLevel = GRayTracingSkyLightSamplingStopLevel;
 		SkyLightData->MaxRayDistance = GRayTracingSkyLightMaxRayDistance;
 		SkyLightData->MaxNormalBias = GetRaytracingMaxNormalBias();
+		SkyLightData->bTransmission = Scene.SkyLight->bTransmission;
 		SkyLightData->MaxShadowThickness = GRayTracingSkyLightMaxShadowThickness;
 
 		ensure(SkyLightData->SamplesPerPixel > 0);
