Add transmission to ray tracing global illumination pass.

TODO: provide global override control to systemically disable the effect.

diff --git a/Engine/Shaders/Private/PathTracing/Light/PathTracingLightCommon.ush b/Engine/Shaders/Private/PathTracing/Light/PathTracingLightCommon.ush
--- a/Engine/Shaders/Private/PathTracing/Light/PathTracingLightCommon.ush
+++ b/Engine/Shaders/Private/PathTracing/Light/PathTracingLightCommon.ush
@@ -130,4 +130,9 @@ float ComputeAttenuationFalloff(float Distance, float Attenuation)
 #else
 	return 1.0;
 #endif
+}
+
+bool HasTransmission(int LightId)
+{
+	return SceneLightsData.Flags[LightId] & 0x01;
 }
\ No newline at end of file
diff --git a/Engine/Shaders/Private/RayTracing/RayTracingCreateGatherPointsRGS.usf b/Engine/Shaders/Private/RayTracing/RayTracingCreateGatherPointsRGS.usf
--- a/Engine/Shaders/Private/RayTracing/RayTracingCreateGatherPointsRGS.usf
+++ b/Engine/Shaders/Private/RayTracing/RayTracingCreateGatherPointsRGS.usf
@@ -349,7 +349,13 @@ RAY_TRACING_ENTRY_RAYGEN(RayTracingCreateGatherPointsRGS)
 					RayDesc LightRay;
 					GenerateLightRay(Ray, LightId, LightUV, LightRay);
 					LightRay.TMax = min(LightRay.TMax, MaxShadowDistance);
-					ApplyPositionBias(LightRay, Payload.WorldNormal, MaxNormalBias);
+					bool bTransmission = HasTransmission(LightId);
+					float3 BiasNormal = Payload.WorldNormal;
+					if (bTransmission && dot(BiasNormal, LightRay.Direction) < 0.0)
+					{
+						BiasNormal = -BiasNormal;
+					}
+					ApplyPositionBias(LightRay, BiasNormal, MaxNormalBias);
 
 					// Trace visibility ray
 					uint NeeRayFlags = RAY_FLAG_ACCEPT_FIRST_HIT_AND_END_SEARCH | RAY_FLAG_SKIP_CLOSEST_HIT_SHADER;
@@ -380,6 +386,10 @@ RAY_TRACING_ENTRY_RAYGEN(RayTracingCreateGatherPointsRGS)
 						FShadowTerms ShadowTerms = { 0.0, 0.0, 0.0, InitHairTransmittanceData() };
 						FDirectLighting LightingSample = EvaluateBxDF(GBufferData, N, V, L, NoL, ShadowTerms);
 						float3 MaterialThroughput = LightingSample.Diffuse;
+						if (bTransmission)
+						{
+							MaterialThroughput += LightingSample.Transmission;
+						}
 
 						// Record the contribution
 						float3 ExitantRadianceSample = Radiance * MaterialThroughput * SplitFactor / NeePdf;
diff --git a/Engine/Shaders/Private/RayTracing/RayTracingGlobalIlluminationRGS.usf b/Engine/Shaders/Private/RayTracing/RayTracingGlobalIlluminationRGS.usf
--- a/Engine/Shaders/Private/RayTracing/RayTracingGlobalIlluminationRGS.usf
+++ b/Engine/Shaders/Private/RayTracing/RayTracingGlobalIlluminationRGS.usf
@@ -235,7 +235,13 @@ RAY_TRACING_ENTRY_RAYGEN(GlobalIlluminationRGS)
 					RayDesc LightRay;
 					GenerateLightRay(Ray, LightId, LightUV, LightRay);
 					LightRay.TMax = min(LightRay.TMax, MaxShadowDistance);
-					ApplyPositionBias(LightRay, Payload.WorldNormal, MaxNormalBias);
+					bool bTransmission = HasTransmission(LightId);
+					float3 BiasNormal = Payload.WorldNormal;
+					if (bTransmission && dot(BiasNormal, LightRay.Direction) < 0.0)
+					{
+						BiasNormal = -BiasNormal;
+					}
+					ApplyPositionBias(LightRay, BiasNormal, MaxNormalBias);
 
 					// Trace visibility ray
 					uint NeeRayFlags = RAY_FLAG_ACCEPT_FIRST_HIT_AND_END_SEARCH | RAY_FLAG_SKIP_CLOSEST_HIT_SHADER;
@@ -271,6 +277,10 @@ RAY_TRACING_ENTRY_RAYGEN(GlobalIlluminationRGS)
 						FShadowTerms ShadowTerms = { 0.0, 0.0, 0.0, InitHairTransmittanceData() };
 						FDirectLighting LightingSample = EvaluateBxDF(GBufferData, N, V, L, NoL, ShadowTerms);
 						MaterialThroughput = LightingSample.Diffuse;
+						if (bTransmission)
+						{
+							MaterialThroughput += LightingSample.Transmission;
+						}
 						MaterialEvalPdf = 1.0;
 #endif
 						// Apply material Pdf for correct MIS weight
diff --git a/Engine/Source/Runtime/Renderer/Private/PathTracing.cpp b/Engine/Source/Runtime/Renderer/Private/PathTracing.cpp
--- a/Engine/Source/Runtime/Renderer/Private/PathTracing.cpp
+++ b/Engine/Source/Runtime/Renderer/Private/PathTracing.cpp
@@ -196,6 +196,7 @@ public:
 			uint32 SkyLightIndex = 0;
 			LightData.Type[SkyLightIndex] = 0;
 			LightData.Color[SkyLightIndex] = FVector(SkyLightData.Color);
+			LightData.Flags[SkyLightIndex] = SkyLightData.bTransmission & 0x01;
 			LightData.Count++;
 
 			for (auto Light : Lights)
@@ -206,6 +207,8 @@ public:
 
 				FLightShaderParameters LightParameters;
 				Light.LightSceneInfo->Proxy->GetLightShaderParameters(LightParameters);
+				uint32 Transmission = Light.LightSceneInfo->Proxy->Transmission();
+				LightData.Flags[LightData.Count] = Transmission & 0x01;
 
 				ELightComponentType LightComponentType = (ELightComponentType)Light.LightSceneInfo->Proxy->GetLightType();
 				switch (LightComponentType)
diff --git a/Engine/Source/Runtime/Renderer/Private/PathTracingUniformBuffers.h b/Engine/Source/Runtime/Renderer/Private/PathTracingUniformBuffers.h
--- a/Engine/Source/Runtime/Renderer/Private/PathTracingUniformBuffers.h
+++ b/Engine/Source/Runtime/Renderer/Private/PathTracingUniformBuffers.h
@@ -32,6 +32,8 @@ BEGIN_GLOBAL_SHADER_PARAMETER_STRUCT(FPathTracingLightData, )
 	SHADER_PARAMETER_ARRAY(float, Attenuation, [RAY_TRACING_LIGHT_COUNT_MAXIMUM])
 	SHADER_PARAMETER_ARRAY(float, RectLightBarnCosAngle, [RAY_TRACING_LIGHT_COUNT_MAXIMUM])
 	SHADER_PARAMETER_ARRAY(float, RectLightBarnLength, [RAY_TRACING_LIGHT_COUNT_MAXIMUM])
+	// Flags
+	SHADER_PARAMETER_ARRAY(uint32, Flags, [RAY_TRACING_LIGHT_COUNT_MAXIMUM])
 END_GLOBAL_SHADER_PARAMETER_STRUCT()
 
 
diff --git a/Engine/Source/Runtime/Renderer/Private/RayTracing/RayTracingGlobalIllumination.cpp b/Engine/Source/Runtime/Renderer/Private/RayTracing/RayTracingGlobalIllumination.cpp
--- a/Engine/Source/Runtime/Renderer/Private/RayTracing/RayTracingGlobalIllumination.cpp
+++ b/Engine/Source/Runtime/Renderer/Private/RayTracing/RayTracingGlobalIllumination.cpp
@@ -187,12 +187,14 @@ void SetupLightParameters(
 	{
 		SkyLightColor = FVector(SkyLight->GetEffectiveLightColor());
 	}
+	uint32 SkyLightTransmission = SkyLight->bTransmission;
 
 	// Prepend SkyLight to light buffer
 	// WARNING: Until ray payload encodes Light data buffer, the execution depends on this ordering!
 	uint32 SkyLightIndex = 0;
 	LightParameters->Type[SkyLightIndex] = 0;
 	LightParameters->Color[SkyLightIndex] = SkyLightColor;
+	LightParameters->Flags[SkyLightIndex] = SkyLightTransmission & 0x01;
 	LightParameters->Count++;
 
 	uint32 MaxLightCount = FMath::Min(CVarRayTracingGlobalIlluminationMaxLightCount.GetValueOnRenderThread(), RAY_TRACING_LIGHT_COUNT_MAXIMUM);
@@ -205,6 +207,8 @@ void SetupLightParameters(
 
 		FLightShaderParameters LightShaderParameters;
 		Light.LightSceneInfo->Proxy->GetLightShaderParameters(LightShaderParameters);
+		uint32 Transmission = Light.LightSceneInfo->Proxy->Transmission();
+		LightParameters->Flags[LightParameters->Count] = Transmission & 0x01;
 
 		ELightComponentType LightComponentType = (ELightComponentType)Light.LightSceneInfo->Proxy->GetLightType();
 		switch (LightComponentType)
