Add transmission lobe consideration for EstimateIrradiance function

diff --git a/Engine/Shaders/Private/PathTracing/Light/PathTracingDirectionalLight.ush b/Engine/Shaders/Private/PathTracing/Light/PathTracingDirectionalLight.ush
--- a/Engine/Shaders/Private/PathTracing/Light/PathTracingDirectionalLight.ush
+++ b/Engine/Shaders/Private/PathTracing/Light/PathTracingDirectionalLight.ush
@@ -51,7 +51,8 @@ void DirectionalLight_EstimateLight(
 {
 	float3 LightDirection = GetNormal(LightId);
 	float3 SurfaceNormal = Payload.WorldNormal;
-	float NoL = saturate(dot(SurfaceNormal, LightDirection));
+	float NoL = dot(SurfaceNormal, LightDirection);
+	NoL = HasTransmission(LightId) ? abs(NoL) : saturate(NoL);
 
 	float LightPower = Luminance(GetColor(LightId));
 	OutIrradiance = LightPower * NoL;
diff --git a/Engine/Shaders/Private/PathTracing/Light/PathTracingPointLight.ush b/Engine/Shaders/Private/PathTracing/Light/PathTracingPointLight.ush
--- a/Engine/Shaders/Private/PathTracing/Light/PathTracingPointLight.ush
+++ b/Engine/Shaders/Private/PathTracing/Light/PathTracingPointLight.ush
@@ -64,7 +64,8 @@ void PointLight_EstimateLight(
 	// Geometric term
 	float3 SurfaceNormal = Payload.WorldNormal;
 	// dxr_todo: incorporate sphere light radius
-	float NoL = saturate(dot(SurfaceNormal, LightDirection));
+	float NoL = dot(SurfaceNormal, LightDirection);
+	NoL = HasTransmission(LightId) ? abs(NoL) : saturate(NoL);
 
 	float LightPower = Luminance(GetColor(LightId));
 	float Falloff = ComputeAttenuationFalloff(sqrt(LightDistanceSquared), GetAttenuation(LightId));
diff --git a/Engine/Shaders/Private/PathTracing/Light/PathTracingRectLight.ush b/Engine/Shaders/Private/PathTracing/Light/PathTracingRectLight.ush
--- a/Engine/Shaders/Private/PathTracing/Light/PathTracingRectLight.ush
+++ b/Engine/Shaders/Private/PathTracing/Light/PathTracingRectLight.ush
@@ -72,7 +72,8 @@ void RectLight_EstimateLight(
 
 	// Approximate geometric term
 	float3 SurfaceNormal = Payload.WorldNormal;
-	float NoL = saturate(dot(SurfaceNormal, LightDirection));
+	float NoL = dot(SurfaceNormal, LightDirection);
+	NoL = HasTransmission(LightId) ? abs(NoL) : saturate(NoL);
 	float Area = GetWidth(LightId) * GetHeight(LightId);
 
 	float LightPower = Luminance(GetColor(LightId));
diff --git a/Engine/Shaders/Private/PathTracing/Light/PathTracingSpotLight.ush b/Engine/Shaders/Private/PathTracing/Light/PathTracingSpotLight.ush
--- a/Engine/Shaders/Private/PathTracing/Light/PathTracingSpotLight.ush
+++ b/Engine/Shaders/Private/PathTracing/Light/PathTracingSpotLight.ush
@@ -144,7 +144,9 @@ void SpotLight_EstimateLight(
 		float Falloff = NoL - CosConeAngle;
 		Falloff *= Falloff;
 		Falloff *= ComputeAttenuationFalloff(sqrt(LightDistanceSquared), GetAttenuation(LightId));
-		OutIrradiance = LightPower * Falloff / LightDistanceSquared;
+		float CosTheta = dot(Payload.WorldNormal, LightDirection);
+		CosTheta = HasTransmission(LightId) ? abs(CosTheta) : saturate(CosTheta);
+		OutIrradiance = LightPower * CosTheta * Falloff / LightDistanceSquared;
 		// #dxr_todo: UE-72534 apply inner core
 	}
 	else
