Full-screen post process material

diff --git a/Engine/Shaders/Private/DeviceEncodingOnlyCommon.ush b/Engine/Shaders/Private/DeviceEncodingOnlyCommon.ush
new file mode 100644
--- /dev/null
+++ b/Engine/Shaders/Private/DeviceEncodingOnlyCommon.ush
@@ -0,0 +1,416 @@
+// Copyright Epic Games, Inc. All Rights Reserved.
+
+/*=============================================================================
+	DeviceEncodingOnlyCommon.usf: PostProcessing tone mapping common
+=============================================================================*/
+
+#include "ACES.ush"
+#include "GammaCorrectionCommon.ush"
+
+// usually 1/2.2, the .y is used for inverse gamma when "gamma only" mode is not used
+half3 InverseGamma;
+
+// Scale factor for converting pixel values to nits. 
+// This value is required for PQ (ST2084) conversions, because PQ linear values are in nits. 
+// The purpose is to make good use of PQ lut entries. A scale factor of 100 conveniently places 
+// about half of the PQ lut indexing below 1.0, with the other half for input values over 1.0.
+// Also, 100nits is the expected monitor brightness for a 1.0 pixel value without a tone curve.
+static const float LinearToNitsScale = 100.0;
+static const float LinearToNitsScaleInverse = 1.0 / 100.0;
+
+// Film tonal and color control.
+half4 ColorMatrixR_ColorCurveCd1;
+half4 ColorMatrixG_ColorCurveCd3Cm3;
+half4 ColorMatrixB_ColorCurveCm2;
+half4 ColorCurve_Cm0Cd0_Cd2_Ch0Cm1_Ch3;
+half4 ColorCurve_Ch1_Ch2;
+half4 ColorShadow_Luma;
+half4 ColorShadow_Tint1;
+half4 ColorShadow_Tint2;
+
+half3 FilmPostProcess(half3 LinearColor) 
+{
+	// Color and exposure control.
+	half3 MatrixColor;
+	#if USE_COLOR_MATRIX == 1
+		// Apply color matrix (channel mixer, exposure, saturation).
+		MatrixColor.r = dot(LinearColor, ColorMatrixR_ColorCurveCd1.rgb);
+		MatrixColor.g = dot(LinearColor, ColorMatrixG_ColorCurveCd3Cm3.rgb);
+		MatrixColor.b = dot(LinearColor, ColorMatrixB_ColorCurveCm2.rgb);
+		#if USE_SHADOW_TINT == 1
+			MatrixColor *= ColorShadow_Tint1.rgb + ColorShadow_Tint2.rgb * rcp(dot(LinearColor, ColorShadow_Luma.rgb) + 1.0);
+		#endif
+		// Required to insure saturation doesn't create negative colors!
+		MatrixColor = max(half3(0.0, 0.0, 0.0), MatrixColor);
+	#else
+		// Less expensive route when not using saturation and channel mixer.
+		#if USE_SHADOW_TINT == 1
+			MatrixColor = LinearColor * (ColorShadow_Tint1.rgb + ColorShadow_Tint2.rgb * rcp(dot(LinearColor, ColorShadow_Luma.rgb) + 1.0));
+		#else
+			MatrixColor = LinearColor * ColorMatrixB_ColorCurveCm2.rgb;
+		#endif
+	#endif
+	// Apply color curve (includes tonemapping).
+	#if USE_CONTRAST == 1
+		// Full path.
+		half3 MatrixColorD = max(0, ColorCurve_Cm0Cd0_Cd2_Ch0Cm1_Ch3.xxx - MatrixColor);
+		half3 MatrixColorH = max(MatrixColor, ColorCurve_Cm0Cd0_Cd2_Ch0Cm1_Ch3.zzz);
+		half3 MatrixColorM = clamp(MatrixColor, ColorCurve_Cm0Cd0_Cd2_Ch0Cm1_Ch3.xxx, ColorCurve_Cm0Cd0_Cd2_Ch0Cm1_Ch3.zzz);
+		half3 CurveColor = 
+			// Highlights
+			(MatrixColorH*ColorCurve_Ch1_Ch2.xxx + ColorCurve_Ch1_Ch2.yyy) * rcp(MatrixColorH + ColorCurve_Cm0Cd0_Cd2_Ch0Cm1_Ch3.www) + 
+				// Midtones
+				((MatrixColorM*ColorMatrixB_ColorCurveCm2.aaa + 
+					// Darks
+					((MatrixColorD*ColorMatrixR_ColorCurveCd1.aaa) * rcp(MatrixColorD + ColorCurve_Cm0Cd0_Cd2_Ch0Cm1_Ch3.yyy) + ColorMatrixG_ColorCurveCd3Cm3.aaa)));
+	#else
+		// This is for mobile, it assumes color is not negative.
+		// Fast path when contrast=1, can remove the dark part of the curve.
+		half3 MatrixColorH = max(MatrixColor, ColorCurve_Cm0Cd0_Cd2_Ch0Cm1_Ch3.zzz);
+		half3 MatrixColorM = min(MatrixColor, ColorCurve_Cm0Cd0_Cd2_Ch0Cm1_Ch3.zzz);
+		half3 CurveColor = (MatrixColorH*ColorCurve_Ch1_Ch2.xxx + ColorCurve_Ch1_Ch2.yyy) * rcp(MatrixColorH + ColorCurve_Cm0Cd0_Cd2_Ch0Cm1_Ch3.www) + MatrixColorM;
+	#endif
+
+	// Must hit black by log10(-2.474)
+	CurveColor -= 0.002;
+
+	return CurveColor;
+}
+
+half3 DeviceEncodingOnlyAndGammaCorrect(half3 LinearColor)
+{
+	// Clamp input to be positive
+	// This displays negative colors as black
+	LinearColor = max(LinearColor, 0);
+
+	half3 GammaColor = pow(LinearColor, InverseGamma.x);
+	
+	// in all cases it's good to clamp into the 0..1 range (e.g for LUT color grading)
+	GammaColor = saturate(GammaColor);
+
+	return GammaColor;
+}
+
+
+/*
+============================================
+// Uncharted settings
+Slope = 0.63;
+Toe = 0.55;
+Shoulder = 0.47;
+BlackClip= 0;
+WhiteClip = 0.01;
+
+// HP settings
+Slope = 0.65;
+Toe = 0.63;
+Shoulder = 0.45;
+BlackClip = 0;
+WhiteClip = 0;
+
+// Legacy settings
+Slope = 0.98;
+Toe = 0.3;
+Shoulder = 0.22;
+BlackClip = 0;
+WhiteClip = 0.025;
+
+// ACES settings
+Slope = 0.91;
+Toe = 0.53;
+Shoulder = 0.23;
+BlackClip = 0;
+WhiteClip = 0.035;
+===========================================
+*/
+float FilmSlope = 0.91;
+float FilmToe = 0.53;
+float FilmShoulder = 0.23;
+float FilmBlackClip = 0;
+float FilmWhiteClip = 0.035;
+
+half3 FilmDeviceEncodingOnly( half3 LinearColor ) 
+{
+	const float3x3 sRGB_2_AP0 = mul( XYZ_2_AP0_MAT, mul( D65_2_D60_CAT, sRGB_2_XYZ_MAT ) );
+	const float3x3 sRGB_2_AP1 = mul( XYZ_2_AP1_MAT, mul( D65_2_D60_CAT, sRGB_2_XYZ_MAT ) );
+
+	const float3x3 AP0_2_sRGB = mul( XYZ_2_sRGB_MAT, mul( D60_2_D65_CAT, AP0_2_XYZ_MAT ) );
+	const float3x3 AP1_2_sRGB = mul( XYZ_2_sRGB_MAT, mul( D60_2_D65_CAT, AP1_2_XYZ_MAT ) );
+	
+	const float3x3 AP0_2_AP1 = mul( XYZ_2_AP1_MAT, AP0_2_XYZ_MAT );
+	const float3x3 AP1_2_AP0 = mul( XYZ_2_AP0_MAT, AP1_2_XYZ_MAT );
+	
+	float3 ColorAP1 = LinearColor;
+	//float3 ColorAP1 = mul( sRGB_2_AP1, float3(LinearColor) );
+
+#if 0
+	{
+		float3 oces = Inverse_ODT_sRGB_D65( LinearColor );
+		float3 aces = Inverse_RRT( oces );
+		ColorAP1 = mul( AP0_2_AP1, aces );
+	}
+#endif
+	
+#if 0
+	float3 ColorSRGB = mul( AP1_2_sRGB, ColorAP1 );
+	ColorSRGB = max( 0, ColorSRGB );
+	ColorAP1 = mul( sRGB_2_AP1, ColorSRGB );
+#endif
+
+	float3 ColorAP0 = mul( AP1_2_AP0, ColorAP1 );
+
+#if 0
+	{
+		float3 aces = ColorAP0;
+		float3 oces = RRT( aces );
+		LinearColor = ODT_sRGB_D65( oces );
+	}
+	return mul( sRGB_2_AP1, LinearColor );
+#endif
+
+#if 1
+	// "Glow" module constants
+	const float RRT_GLOW_GAIN = 0.05;
+	const float RRT_GLOW_MID = 0.08;
+
+	float saturation = rgb_2_saturation( ColorAP0 );
+	float ycIn = rgb_2_yc( ColorAP0 );
+	float s = sigmoid_shaper( (saturation - 0.4) / 0.2);
+	float addedGlow = 1 + glow_fwd( ycIn, RRT_GLOW_GAIN * s, RRT_GLOW_MID);
+	ColorAP0 *= addedGlow;
+#endif
+
+#if 1
+	// --- Red modifier --- //
+	const float RRT_RED_SCALE = 0.82;
+	const float RRT_RED_PIVOT = 0.03;
+	const float RRT_RED_HUE = 0;
+	const float RRT_RED_WIDTH = 135;
+	float hue = rgb_2_hue( ColorAP0 );
+	float centeredHue = center_hue( hue, RRT_RED_HUE );
+	float hueWeight = Square( smoothstep( 0, 1, 1 - abs( 2 * centeredHue / RRT_RED_WIDTH ) ) );
+		
+	ColorAP0.r += hueWeight * saturation * (RRT_RED_PIVOT - ColorAP0.r) * (1. - RRT_RED_SCALE);
+#endif
+	
+	// Use ACEScg primaries as working space
+	float3 WorkingColor = mul( AP0_2_AP1_MAT, ColorAP0 );
+
+	WorkingColor = max( 0, WorkingColor );
+
+	// Pre desaturate
+	WorkingColor = lerp( dot( WorkingColor, AP1_RGB2Y ), WorkingColor, 0.96 );
+	
+	const half ToeScale			= 1 + FilmBlackClip - FilmToe;
+	const half ShoulderScale	= 1 + FilmWhiteClip - FilmShoulder;
+	
+	const float InMatch = 0.18;
+	const float OutMatch = 0.18;
+
+	float ToeMatch;
+	if( FilmToe > 0.8 )
+	{
+		// 0.18 will be on straight segment
+		ToeMatch = ( 1 - FilmToe  - OutMatch ) / FilmSlope + log10( InMatch );
+	}
+	else
+	{
+		// 0.18 will be on toe segment
+
+		// Solve for ToeMatch such that input of InMatch gives output of OutMatch.
+		const float bt = ( OutMatch + FilmBlackClip ) / ToeScale - 1;
+		ToeMatch = log10( InMatch ) - 0.5 * log( (1+bt)/(1-bt) ) * (ToeScale / FilmSlope);
+	}
+
+	float StraightMatch = ( 1 - FilmToe ) / FilmSlope - ToeMatch;
+	float ShoulderMatch = FilmShoulder / FilmSlope - StraightMatch;
+	
+	half3 LogColor = log10( WorkingColor );
+	half3 StraightColor = FilmSlope * ( LogColor + StraightMatch );
+	
+	half3 ToeColor		= (    -FilmBlackClip ) + (2 *      ToeScale) / ( 1 + exp( (-2 * FilmSlope /      ToeScale) * ( LogColor -      ToeMatch ) ) );
+	half3 ShoulderColor	= ( 1 + FilmWhiteClip ) - (2 * ShoulderScale) / ( 1 + exp( ( 2 * FilmSlope / ShoulderScale) * ( LogColor - ShoulderMatch ) ) );
+
+	ToeColor		= LogColor <      ToeMatch ?      ToeColor : StraightColor;
+	ShoulderColor	= LogColor > ShoulderMatch ? ShoulderColor : StraightColor;
+
+	half3 t = saturate( ( LogColor - ToeMatch ) / ( ShoulderMatch - ToeMatch ) );
+	t = ShoulderMatch < ToeMatch ? 1 - t : t;
+	t = (3-2*t)*t*t;
+	half3 ToneColor = lerp( ToeColor, ShoulderColor, t );
+
+	// Post desaturate
+	ToneColor = lerp( dot( float3(ToneColor), AP1_RGB2Y ), ToneColor, 0.93 );
+
+	// Returning positive AP1 values
+	return max( 0, ToneColor );
+}
+
+half3 FilmDeviceEncodingOnlyInverse( half3 ToneColor ) 
+{
+	const float3x3 sRGB_2_AP1 = mul( XYZ_2_AP1_MAT, mul( D65_2_D60_CAT, sRGB_2_XYZ_MAT ) );
+	const float3x3 AP1_2_sRGB = mul( XYZ_2_sRGB_MAT, mul( D60_2_D65_CAT, AP1_2_XYZ_MAT ) );
+	
+	// Use ACEScg primaries as working space
+	half3 WorkingColor = mul( sRGB_2_AP1, saturate( ToneColor ) );
+
+	WorkingColor = max( 0, WorkingColor );
+	
+	// Post desaturate
+	WorkingColor = lerp( dot( WorkingColor, AP1_RGB2Y ), WorkingColor, 1.0 / 0.93 );
+
+	half3 ToeColor		= 0.374816 * pow( 0.9 / min( WorkingColor, 0.8 ) - 1, -0.588729 );
+	half3 ShoulderColor	= 0.227986 * pow( 1.56 / ( 1.04 - WorkingColor ) - 1, 1.02046 );
+
+	half3 t = saturate( ( WorkingColor - 0.35 ) / ( 0.45 - 0.35 ) );
+	t = (3-2*t)*t*t;
+	half3 LinearColor = lerp( ToeColor, ShoulderColor, t );
+
+	// Pre desaturate
+	LinearColor = lerp( dot( LinearColor, AP1_RGB2Y ), LinearColor, 1.0 / 0.96 );
+
+	LinearColor = mul( AP1_2_sRGB, LinearColor );
+
+	// Returning positive sRGB values
+	return max( 0, LinearColor );
+}
+
+//
+// ACES sRGB D65 Output Transform - Forward and Inverse
+//  Input is scene-referred linear values in the sRGB gamut
+//  Output is output-referred linear values in the sRGB gamut
+//
+float3 ACESOutputTransformsRGBD65( float3 SceneReferredLinearsRGBColor )
+{
+	const float3x3 sRGB_2_AP0 = mul( XYZ_2_AP0_MAT, mul( D65_2_D60_CAT, sRGB_2_XYZ_MAT ) );
+
+	float3 aces = mul( sRGB_2_AP0, SceneReferredLinearsRGBColor * 1.5 );
+	float3 oces = RRT( aces );
+	float3 OutputReferredLinearsRGBColor =  ODT_sRGB_D65( oces );
+	return OutputReferredLinearsRGBColor;
+}
+
+float3 InverseACESOutputTransformsRGBD65( float3 OutputReferredLinearsRGBColor )
+{
+	const float3x3 AP0_2_sRGB = mul( XYZ_2_sRGB_MAT, mul( D60_2_D65_CAT, AP0_2_XYZ_MAT ) );
+
+	float3 oces = Inverse_ODT_sRGB_D65( OutputReferredLinearsRGBColor );
+	float3 aces = Inverse_RRT( oces );
+	float3 SceneReferredLinearsRGBColor = mul( AP0_2_sRGB, aces ) * 0.6666;
+
+	return SceneReferredLinearsRGBColor;
+}
+
+//
+// ACES D65 1000 nit Output Transform - Forward
+//  Input is scene-referred linear values in the sRGB gamut
+//  Output is output-referred linear values in the AP1 gamut
+//
+float3 ACESOutputTransforms1000( float3 SceneReferredLinearsRGBColor )
+{
+	const float3x3 sRGB_2_AP0 = mul( XYZ_2_AP0_MAT, mul( D65_2_D60_CAT, sRGB_2_XYZ_MAT ) );
+
+	float3 aces = mul( sRGB_2_AP0, SceneReferredLinearsRGBColor * 1.5 );
+	float3 oces = RRT( aces );
+	float3 OutputReferredLinearAP1Color = ODT_1000nits( oces );
+	return OutputReferredLinearAP1Color;
+}
+
+//
+// ACES D65 2000 nit Output Transform - Forward
+//  Input is scene-referred linear values in the sRGB gamut
+//  Output is output-referred linear values in the AP1 gamut
+//
+float3 ACESOutputTransforms2000( float3 SceneReferredLinearsRGBColor )
+{
+	const float3x3 sRGB_2_AP0 = mul( XYZ_2_AP0_MAT, mul( D65_2_D60_CAT, sRGB_2_XYZ_MAT ) );
+
+	float3 aces = mul( sRGB_2_AP0, SceneReferredLinearsRGBColor * 1.5 );
+	float3 oces = RRT( aces );
+	float3 OutputReferredLinearAP1Color = ODT_2000nits( oces );
+	return OutputReferredLinearAP1Color;
+}
+
+static const float3x3 GamutMappingIdentityMatrix = { 1.0, 0.0, 0.0, 0.0, 1.0, 0.0, 0.0, 0.0, 1.0 };
+
+//
+// Gamut conversion matrices
+//
+float3x3 OuputGamutMappingMatrix( uint OutputGamut )
+{
+	// Gamut mapping matrices used later
+	const float3x3 AP1_2_sRGB    = mul( XYZ_2_sRGB_MAT, mul( D60_2_D65_CAT, AP1_2_XYZ_MAT ) );
+	const float3x3 AP1_2_DCI_D65 = mul( XYZ_2_P3D65_MAT, mul( D60_2_D65_CAT, AP1_2_XYZ_MAT ) );
+	const float3x3 AP1_2_Rec2020 = mul( XYZ_2_Rec2020_MAT, mul( D60_2_D65_CAT, AP1_2_XYZ_MAT ) );
+
+	// Set gamut mapping matrix 
+	// 0 = sRGB - D65
+	// 1 = P3 - D65
+	// 2 = Rec.2020 - D65
+	// 3 = ACES AP0 - D60
+	// 4 = ACES AP1 - D60
+
+	if( OutputGamut == 1 )
+		return AP1_2_DCI_D65;
+	else if( OutputGamut == 2 )
+		return AP1_2_Rec2020;
+	else if( OutputGamut == 3 )
+		return AP1_2_AP0_MAT;
+	else if( OutputGamut == 4 )
+		return GamutMappingIdentityMatrix;
+	else
+		return AP1_2_sRGB;
+}
+
+float3x3 OuputInverseGamutMappingMatrix( uint OutputGamut )
+{
+	// Gamut mapping matrices used later
+	const float3x3 sRGB_2_AP1    = mul( XYZ_2_AP1_MAT, mul( D65_2_D60_CAT, sRGB_2_XYZ_MAT ) );
+	const float3x3 DCI_D65_2_AP1 = mul( XYZ_2_AP1_MAT, mul( D65_2_D60_CAT, P3D65_2_XYZ_MAT ) );
+	const float3x3 Rec2020_2_AP1 = mul( XYZ_2_AP1_MAT, mul( D65_2_D60_CAT, Rec2020_2_XYZ_MAT ) );
+
+	// Set gamut mapping matrix 
+	// 0 = sRGB - D65
+	// 1 = P3 (DCI) - D65
+	// 2 = Rec.2020 - D65
+	// 3 = ACES AP0 - D60
+	float3x3 GamutMappingMatrix = sRGB_2_AP1;
+	if( OutputGamut == 1 )
+		GamutMappingMatrix = DCI_D65_2_AP1;
+	else if( OutputGamut == 2 )
+		GamutMappingMatrix = Rec2020_2_AP1;
+	else if( OutputGamut == 3 )
+		GamutMappingMatrix = AP0_2_AP1_MAT;
+
+	return GamutMappingMatrix;
+}
+
+float3 ST2084ToScRGB(float3 Color, uint OutputDevice)
+{
+	// Nvidia HDR encoding - Remove PQ, convert to linear scRGB
+	const float3x3 AP1_2_sRGB = mul(XYZ_2_sRGB_MAT, AP1_2_XYZ_MAT);
+	const float WhitePoint = 80.f;
+
+	// 1000.f nit display
+	float MaxODTNits = 1000.0f;
+	float MinODTNits = 0.0001f;
+
+	if (OutputDevice == 4 || OutputDevice == 6)
+	{
+		// 2000 nit display
+		MaxODTNits = 2000.0f;
+		MinODTNits = 0.005f;
+	}
+
+	float3 OutColor = ST2084ToLinear(Color);
+		
+	OutColor = clamp(OutColor, MinODTNits, MaxODTNits);
+	OutColor.x = Y_2_linCV(OutColor.x, MaxODTNits, MinODTNits);
+	OutColor.y = Y_2_linCV(OutColor.y, MaxODTNits, MinODTNits);
+	OutColor.z = Y_2_linCV(OutColor.z, MaxODTNits, MinODTNits);
+
+	float scRGBScale = MaxODTNits / WhitePoint;
+	OutColor = mul(AP1_2_sRGB, OutColor) * scRGBScale;
+
+	return OutColor;
+}
\ No newline at end of file
diff --git a/Engine/Shaders/Private/PostProcessCombineLUTs.usf b/Engine/Shaders/Private/PostProcessCombineLUTs.usf
--- a/Engine/Shaders/Private/PostProcessCombineLUTs.usf
+++ b/Engine/Shaders/Private/PostProcessCombineLUTs.usf
@@ -2,6 +2,9 @@
 
 /*=============================================================================
 	LUTBlender.usf: Filter pixel shader source.   
+
+	Note: Any changes to device encoding logic in CombineLUTsCommon() must
+	also be made in PostProcessDeviceEncodingOnly.usf.
 =============================================================================*/
 
 #define USE_COLOR_MATRIX	1
@@ -387,6 +390,8 @@ float4 CombineLUTsCommon(float2 InUV, uint InLayerIndex)
 	// Apply "gamma" curve adjustment.
 	FilmColor = pow( max(0, FilmColorNoGamma), InverseGamma.y );
 		
+	// Note: Any changes to the following device encoding logic must also be made 
+	// in PostProcessDeviceEncodingOnly.usf's corresponding pixel shader.
 	half3 OutDeviceColor = 0;
 	BRANCH
 
diff --git a/Engine/Shaders/Private/PostProcessDeviceEncodingOnly.usf b/Engine/Shaders/Private/PostProcessDeviceEncodingOnly.usf
new file mode 100644
--- /dev/null
+++ b/Engine/Shaders/Private/PostProcessDeviceEncodingOnly.usf
@@ -0,0 +1,210 @@
+// Copyright Epic Games, Inc. All Rights Reserved.
+
+/*=============================================================================
+	PostProcessDeviceEncodingOnly.usf: PostProcessing device encoding only. 
+	This replaces the tonemapper when displaying debug post process materials
+	used for monitor calibration.
+
+	Note: Any changes to device encoding logic must also be made 
+	in PostProcessCombineLUTs.usf's corresponding pixel shader.
+
+=============================================================================*/
+
+#include "Common.ush"
+#include "PostProcessCommon.ush"
+#include "TonemapCommon.ush"
+#include "ScreenPass.ush"
+
+/** Values of DIM_OUTPUT_DEVICE (matching C++'s ETonemapperOutputDevice) */
+#define TONEMAPPER_OUTPUT_sRGB 0
+#define TONEMAPPER_OUTPUT_Rec709 1
+#define TONEMAPPER_OUTPUT_ExplicitGammaMapping 2
+#define TONEMAPPER_OUTPUT_ACES1000nitST2084 3
+#define TONEMAPPER_OUTPUT_ACES2000nitST2084 4
+#define TONEMAPPER_OUTPUT_ACES1000nitScRGB 5
+#define TONEMAPPER_OUTPUT_ACES2000nitScRGB 6
+#define TONEMAPPER_OUTPUT_LinearEXR 7
+#define TONEMAPPER_OUTPUT_NoToneCurve 8
+#define TONEMAPPER_OUTPUT_WithToneCurve 9
+
+#ifndef DIM_OUTPUT_DEVICE
+	#define DIM_OUTPUT_DEVICE (TONEMAPPER_OUTPUT_sRGB)
+#endif
+
+SCREEN_PASS_TEXTURE_VIEWPORT(Color)
+SCREEN_PASS_TEXTURE_VIEWPORT(Output)
+
+Texture2D ColorTexture;
+SamplerState ColorSampler;
+
+float EditorNITLevel;
+
+uint bOutputInHDR;
+
+uint OutputDevice;
+uint OutputGamut;
+
+uint GetOutputDevice()
+{
+#if FEATURE_LEVEL > FEATURE_LEVEL_ES3_1
+	return OutputDevice;
+#else
+	// Only sRGB output for Mobile
+	return 0;
+#endif 
+}
+
+float4 SampleSceneColor(float2 SceneUV)
+{
+	return Texture2DSample(ColorTexture, ColorSampler, SceneUV);
+}
+
+float4 DeviceEncodingOnlyCommonPS(
+	float2 UV,
+	float4 SvPosition
+	)
+{
+	float4 OutColor = 0;
+
+	float2 SceneUV = UV.xy;
+
+	half4 SceneColor = SampleSceneColor(SceneUV);
+
+	// AP1 to Working space matrices
+	const float3x3 sRGB_2_AP1 = mul( XYZ_2_AP1_MAT, mul( D65_2_D60_CAT, sRGB_2_XYZ_MAT ) );
+	const float3x3 AP1_2_sRGB = mul( XYZ_2_sRGB_MAT, mul( D60_2_D65_CAT, AP1_2_XYZ_MAT ) );
+
+	const float3x3 AP0_2_AP1 = mul( XYZ_2_AP1_MAT, AP0_2_XYZ_MAT );
+	const float3x3 AP1_2_AP0 = mul( XYZ_2_AP0_MAT, AP1_2_XYZ_MAT );
+
+	const float3x3 AP1_2_Output  = OuputGamutMappingMatrix( OutputGamut );
+
+	// Apply "gamma" curve adjustment.
+	float3 FilmColor = pow(max(0, SceneColor.rgb), InverseGamma.y);
+	
+	// Note: Any changes to device encoding logic below must also be made 
+	// in PostProcessCombineLUTs.usf's corresponding pixel shader.
+	half3 OutDeviceColor = 0;
+	BRANCH
+
+	// sRGB, user specified gamut
+	if( GetOutputDevice() == 0 )
+	{		
+		// Convert from sRGB to specified output gamut	
+		// float3 OutputGamutColor = mul( AP1_2_Output, mul( sRGB_2_AP1, FilmColor ) );
+
+		// FIXME: Workaround for UE-29935, pushing all colors with a 0 component to black output
+		// Default parameters seem to cancel out (sRGB->XYZ->AP1->XYZ->sRGB), so should be okay for a temp fix
+		float3 OutputGamutColor = FilmColor;
+
+		// Apply conversion to sRGB (this must be an exact sRGB conversion else darks are bad).
+		OutDeviceColor = LinearToSrgb( OutputGamutColor );
+	}
+
+	// Rec 709, user specified gamut
+	else if( GetOutputDevice() == 1 )
+	{
+		// Convert from sRGB to specified output gamut
+		float3 OutputGamutColor = mul( AP1_2_Output, mul( sRGB_2_AP1, FilmColor ) );
+
+		// Didn't profile yet if the branching version would be faster (different linear segment).
+		OutDeviceColor = LinearTo709Branchless( OutputGamutColor );
+	}
+
+	// ACES 1000nit transform with PQ/2084 encoding, user specified gamut 
+	else if( GetOutputDevice() == 3 || GetOutputDevice() == 5 )
+	{		
+		// 1000 nit ODT
+		float3 ODTColor = ACESOutputTransforms1000( FilmColor );
+
+		// Convert from AP1 to specified output gamut
+		ODTColor = mul( AP1_2_Output, ODTColor );
+
+		// Apply conversion to ST-2084 (Dolby PQ)
+		OutDeviceColor = LinearToST2084( ODTColor );
+	}
+
+	// ACES 2000nit transform with PQ/2084 encoding, user specified gamut 
+	else if( GetOutputDevice() == 4 || GetOutputDevice() == 6 )
+	{		
+		// 2000 nit ODT
+		float3 ODTColor = ACESOutputTransforms2000( FilmColor );
+
+		// Convert from AP1 to specified output gamut
+		ODTColor = mul( AP1_2_Output, ODTColor );
+
+		// Apply conversion to ST-2084 (Dolby PQ)
+		OutDeviceColor = LinearToST2084( ODTColor );
+	}	
+	
+	else if( GetOutputDevice() == 7 )
+	{
+		float3 OutputGamutColor = mul( AP1_2_Output, mul( sRGB_2_AP1, FilmColor ) );
+		OutDeviceColor = LinearToST2084( OutputGamutColor );
+	}
+	// Linear HDR, including all color correction, but no tone curve
+	else if( GetOutputDevice() == 8 ) 
+	{
+		OutDeviceColor = FilmColor;
+	}
+	// "Linear" including all color correction and the tone curve, but no device gamma
+	else if (GetOutputDevice() == 9)
+	{
+		float3 OutputGamutColor = mul(AP1_2_Output, mul(sRGB_2_AP1, SceneColor)); 
+
+		OutDeviceColor = OutputGamutColor;
+	}
+
+	// OutputDevice == 2
+	// Gamma 2.2, user specified gamut
+	else
+	{
+		// Convert from sRGB to specified output gamut
+		float3 OutputGamutColor = mul( AP1_2_Output, mul( sRGB_2_AP1, FilmColor ) );
+
+		// This is different than the prior "gamma" curve adjustment (but reusing the variable).
+		// For displays set to a gamma colorspace.
+		// Note, MacOSX native output is raw gamma 2.2 not sRGB!
+		OutDeviceColor = pow( OutputGamutColor, InverseGamma.z );
+	}
+
+	// Better to saturate(lerp(a,b,t)) than lerp(saturate(a),saturate(b),t)
+	OutColor.rgb = OutDeviceColor / 1.05;
+	OutColor.a = 1;
+
+	return OutColor;
+}
+
+// pixel shader entry point
+void MainPS(
+	in noperspective float2 UV : TEXCOORD0,
+	float4 SvPosition : SV_POSITION,		// after all interpolators
+	out float4 OutColor : SV_Target0
+	)
+{
+    OutColor = DeviceEncodingOnlyCommonPS(UV, SvPosition);
+}
+
+#if COMPUTESHADER
+RWTexture2D<float4> RWOutputTexture;
+
+[numthreads(THREADGROUP_SIZEX, THREADGROUP_SIZEY, 1)]
+void MainCS(uint2 DispatchThreadId : SV_DispatchThreadID)
+{
+	float4 SvPosition = float4((float2)DispatchThreadId + Output_ViewportMin + 0.5f, 0.0f, 1.0f);
+	float2 UV = SvPosition.xy * Output_ExtentInverse;
+	float4 InScreenPos = float4(UV*2-1,0,1);
+
+	if (IsComputeUVOutOfBounds(UV))
+	{
+		return;
+	}
+
+	float4 OutColor = DeviceEncodingOnlyCommonPS(UV, SvPosition);
+
+	uint2 PixelPos = DispatchThreadId + Output_ViewportMin;
+	RWOutputTexture[PixelPos] = OutColor;
+}
+#endif
+
+
diff --git a/Engine/Source/Editor/UnrealEd/Private/EditorViewportClient.cpp b/Engine/Source/Editor/UnrealEd/Private/EditorViewportClient.cpp
--- a/Engine/Source/Editor/UnrealEd/Private/EditorViewportClient.cpp
+++ b/Engine/Source/Editor/UnrealEd/Private/EditorViewportClient.cpp
@@ -14,6 +14,7 @@
 #include "Engine/Canvas.h"
 #include "Settings/LevelEditorViewportSettings.h"
 #include "Settings/LevelEditorMiscSettings.h"
+#include "Engine/RendererSettings.h"
 #include "Components/DirectionalLightComponent.h"
 #include "Components/BillboardComponent.h"
 #include "Audio/AudioDebug.h"
@@ -2565,6 +2566,17 @@ FText FEditorViewportClient::GetCurrentBufferVisualizationModeDisplayName() cons
 		? FBufferVisualizationData::GetMaterialDefaultDisplayName() : GetBufferVisualizationData().GetMaterialDisplayName(CurrentBufferVisualizationMode));
 }
 
+bool FEditorViewportClient::IsVisualizeDebugMaterialEnabled() const
+{
+	// Get the list of requested buffers from the console
+	const URendererSettings* Settings = GetDefault<URendererSettings>();
+	check(Settings);
+
+	return ((EngineShowFlags.VisualizeDebugCustomPostProcessMaterial && Settings->VisualizeDebugCustomMaterialPath.IsValid()) ||
+		(EngineShowFlags.VisualizeDebugColor  && Settings->VisualizeDebugColorMaterialPath.IsValid()) ||
+		(EngineShowFlags.VisualizeDebugGrayscale && Settings->VisualizeDebugGrayscaleMaterialPath.IsValid()));
+}
+
 void FEditorViewportClient::ChangeRayTracingDebugVisualizationMode(FName InName)
 {
 	SetViewMode(VMI_RayTracingDebug);
@@ -2599,7 +2611,7 @@ bool FEditorViewportClient::SupportsPreviewResolutionFraction() const
 	}
 
 	// Don't do preview screen percentage for buffer visualization.
-	if (EngineShowFlags.VisualizeBuffer)
+	if (EngineShowFlags.VisualizeBuffer || IsVisualizeDebugMaterialEnabled())
 	{
 		return false;
 	}
@@ -3865,7 +3877,7 @@ void FEditorViewportClient::Draw(FViewport* InViewport, FCanvas* Canvas)
 	}
 
 	// Axes indicators
-	if (bDrawAxes && !ViewFamily.EngineShowFlags.Game && !GLevelEditorModeTools().IsViewportUIHidden())
+	if (bDrawAxes && !ViewFamily.EngineShowFlags.Game && !GLevelEditorModeTools().IsViewportUIHidden() && !IsVisualizeDebugMaterialEnabled())
 	{
 		switch (GetViewportType())
 		{
diff --git a/Engine/Source/Editor/UnrealEd/Public/EditorViewportClient.h b/Engine/Source/Editor/UnrealEd/Public/EditorViewportClient.h
--- a/Engine/Source/Editor/UnrealEd/Public/EditorViewportClient.h
+++ b/Engine/Source/Editor/UnrealEd/Public/EditorViewportClient.h
@@ -1189,6 +1189,11 @@ public:
 	 */
 	FText GetCurrentBufferVisualizationModeDisplayName() const;
 
+	/**
+	* Returns whether visualize debug material is enabled.
+	*/
+	bool IsVisualizeDebugMaterialEnabled() const;
+
 	/**
 	 * Changes the ray tracing debug visualization mode for this viewport
 	 *
diff --git a/Engine/Source/Runtime/Engine/Classes/Engine/RendererSettings.h b/Engine/Source/Runtime/Engine/Classes/Engine/RendererSettings.h
--- a/Engine/Source/Runtime/Engine/Classes/Engine/RendererSettings.h
+++ b/Engine/Source/Runtime/Engine/Classes/Engine/RendererSettings.h
@@ -821,6 +821,27 @@ class ENGINE_API URendererSettings : public UDeveloperSettings
 		ToolTip="Whether to discard skeletal mesh LODs below minimum LOD levels at cook time."))
 	FPerPlatformBool bDiscardSkeletalMeshOptionalLODs;
 
+	/**
+	" Visualize debug material settings for post process debug materials, used for setting full-screen images used for monitor calibration."
+	*/
+	UPROPERTY(config, EditAnywhere, Category = PostProcessDebugMaterials, meta = (AllowedClasses = "Material",
+		DisplayName = "Visualize Debug Color Material Path",
+		ToolTip = "When the VisualizeDebugColor show flag is enabled, this path will be used as the post-process material to render.",
+		ConfigRestartRequired = true))
+	FSoftObjectPath VisualizeDebugColorMaterialPath;
+
+	UPROPERTY(config, EditAnywhere, Category = PostProcessDebugMaterials, meta = (AllowedClasses = "Material",
+		DisplayName = "Visualize Debug Grayscale Material Path",
+		ToolTip = "When the VisualizeDebugGrayscale show flag is enabled, this path will be used as the post-process material to render.",
+		ConfigRestartRequired = true))
+	FSoftObjectPath VisualizeDebugGrayscaleMaterialPath;
+
+	UPROPERTY(config, EditAnywhere, Category = PostProcessDebugMaterials, meta = (AllowedClasses = "Material",
+		DisplayName = "Visualize Debug Custom Material Path",
+		ToolTip = "When the VisualizeDebugCustom show flag is enabled, this path will be used as the post-process material to render.",
+		ConfigRestartRequired = true))
+	FSoftObjectPath VisualizeDebugCustomMaterialPath;
+
 public:
 
 	//~ Begin UObject Interface
diff --git a/Engine/Source/Runtime/Engine/Private/SceneView.cpp b/Engine/Source/Runtime/Engine/Private/SceneView.cpp
--- a/Engine/Source/Runtime/Engine/Private/SceneView.cpp
+++ b/Engine/Source/Runtime/Engine/Private/SceneView.cpp
@@ -2011,6 +2011,9 @@ void FSceneView::EndFinalPostprocessSettings(const FSceneViewInitOptions& ViewIn
 	if (AllowDebugViewmodes())
 	{
 		ConfigureBufferVisualizationSettings();
+#if !(UE_BUILD_SHIPPING || UE_BUILD_TEST)
+		ConfigureDebugVisualizationSettings();
+#endif
 	}
 
 #if WITH_EDITOR
@@ -2124,6 +2127,71 @@ void FSceneView::ConfigureBufferVisualizationSettings()
 	}
 }
 
+#if !(UE_BUILD_SHIPPING || UE_BUILD_TEST)
+
+void FSceneView::ConfigureDebugVisualizationSettings()
+{
+	// Get the list of requested buffers from the console
+	const URendererSettings* Settings = GetDefault<URendererSettings>();
+	check(Settings);
+
+	if (Family->EngineShowFlags.VisualizeDebugCustomPostProcessMaterial)
+	{
+		if (Settings->VisualizeDebugCustomMaterialPath.IsValid())
+		{
+			if (UMaterial* Material = Cast<UMaterial>(Settings->VisualizeDebugCustomMaterialPath.TryLoad()))
+			{
+				FinalPostProcessSettings.DebugCustomVisualizationMaterial = Material;
+				CurrentVisualizeDebugCustomMaterialName = *Material->GetPathName();
+			}
+			else
+			{
+				UE_LOG(LogBufferVisualization, Warning, TEXT("Error loading material '%s'"), *Settings->VisualizeDebugCustomMaterialPath.ToString());
+				FinalPostProcessSettings.DebugCustomVisualizationMaterial = nullptr;
+				CurrentVisualizeDebugCustomMaterialName = NAME_None;
+			}
+		}
+	}
+
+	if (Family->EngineShowFlags.VisualizeDebugColor)
+	{
+		if (Settings->VisualizeDebugColorMaterialPath.IsValid())
+		{
+			if (UMaterial* Material = Cast<UMaterial>(Settings->VisualizeDebugColorMaterialPath.TryLoad()))
+			{
+				FinalPostProcessSettings.DebugColorVisualizationMaterial = Material;
+				CurrentVisualizeDebugColorMaterialName = *Material->GetPathName();
+			}
+			else
+			{
+				UE_LOG(LogBufferVisualization, Warning, TEXT("Error loading material '%s'"), *Settings->VisualizeDebugColorMaterialPath.ToString());
+				FinalPostProcessSettings.DebugColorVisualizationMaterial = nullptr;
+				CurrentVisualizeDebugColorMaterialName = NAME_None;
+			}
+		}
+	}
+
+	if (Family->EngineShowFlags.VisualizeDebugGrayscale)
+	{
+		if (Settings->VisualizeDebugGrayscaleMaterialPath.IsValid())
+		{
+			if (UMaterial* Material = Cast<UMaterial>(Settings->VisualizeDebugGrayscaleMaterialPath.TryLoad()))
+			{
+				FinalPostProcessSettings.DebugGrayscaleVisualizationMaterial = Material;
+				CurrentVisualizeDebugGrayscaleMaterialName = *Material->GetPathName();
+			}
+			else
+			{
+				UE_LOG(LogBufferVisualization, Warning, TEXT("Error loading material '%s'"), *Settings->VisualizeDebugGrayscaleMaterialPath.ToString());
+				FinalPostProcessSettings.DebugGrayscaleVisualizationMaterial = nullptr;
+				CurrentVisualizeDebugGrayscaleMaterialName = NAME_None;
+			}
+		}
+	}
+}
+
+#endif
+
 EShaderPlatform FSceneView::GetShaderPlatform() const
 {
 	return GShaderPlatformForFeatureLevel[GetFeatureLevel()];
diff --git a/Engine/Source/Runtime/Engine/Private/ShowFlags.cpp b/Engine/Source/Runtime/Engine/Private/ShowFlags.cpp
--- a/Engine/Source/Runtime/Engine/Private/ShowFlags.cpp
+++ b/Engine/Source/Runtime/Engine/Private/ShowFlags.cpp
@@ -333,6 +333,11 @@ void ApplyViewMode(EViewModeIndex ViewModeIndex, bool bPerspective, FEngineShowF
 	EngineShowFlags.SetPathTracing(ViewModeIndex == VMI_PathTracing);
 }
 
+static bool IsVisualizeDebugMaterialEnabled(FEngineShowFlags& EngineShowFlags)
+{
+	return (EngineShowFlags.VisualizeDebugColor || EngineShowFlags.VisualizeDebugGrayscale || EngineShowFlags.VisualizeDebugCustomPostProcessMaterial);
+}
+
 void EngineShowFlagOverride(EShowFlagInitMode ShowFlagInitMode, EViewModeIndex ViewModeIndex, FEngineShowFlags& EngineShowFlags, bool bCanDisableTonemapper)
 {
 	if(ShowFlagInitMode == ESFIM_Game)
@@ -509,8 +514,8 @@ void EngineShowFlagOverride(EShowFlagInitMode ShowFlagInitMode, EViewModeIndex V
 		}
 	}
 
-	// disable AA in full screen GBuffer visualization
-	if (bCanDisableTonemapper && EngineShowFlags.VisualizeBuffer)
+	// disable AA in full screen GBuffer visualization or debug material visualization
+	if (bCanDisableTonemapper && (EngineShowFlags.VisualizeBuffer || IsVisualizeDebugMaterialEnabled(EngineShowFlags)))
 	{
 		EngineShowFlags.SetTonemapper(false);
 	}
diff --git a/Engine/Source/Runtime/Engine/Public/FinalPostProcessSettings.h b/Engine/Source/Runtime/Engine/Public/FinalPostProcessSettings.h
--- a/Engine/Source/Runtime/Engine/Public/FinalPostProcessSettings.h
+++ b/Engine/Source/Runtime/Engine/Public/FinalPostProcessSettings.h
@@ -17,7 +17,10 @@ class FFinalPostProcessSettings : public FPostProcessSettings
 {
 public:
 	FFinalPostProcessSettings()
-	: HighResScreenshotMaterial(NULL)
+	: DebugColorVisualizationMaterial(NULL)
+	, DebugGrayscaleVisualizationMaterial(NULL)
+	, DebugCustomVisualizationMaterial(NULL)
+	, HighResScreenshotMaterial(NULL)
 	, HighResScreenshotMaskMaterial(NULL)
 	, HighResScreenshotCaptureRegionMaterial(NULL)
 	, bBufferVisualizationDumpRequired(false)
@@ -168,6 +171,15 @@ public:
 	// List of materials to use in the buffer visualization overview
 	TArray<UMaterialInterface*> BufferVisualizationOverviewMaterials;
 
+	// Material to use for debug color visualization
+	UMaterialInterface* DebugColorVisualizationMaterial;
+
+	// Material to use for debug grayscale visualization
+	UMaterialInterface* DebugGrayscaleVisualizationMaterial;
+
+	// Material to use for debug custom post-processing visualization
+	UMaterialInterface* DebugCustomVisualizationMaterial;
+
 	// Material to use for rendering high res screenshot with mask. Post process expects this material to be set all the time.
 	UMaterialInterface* HighResScreenshotMaterial;
 
diff --git a/Engine/Source/Runtime/Engine/Public/SceneView.h b/Engine/Source/Runtime/Engine/Public/SceneView.h
--- a/Engine/Source/Runtime/Engine/Public/SceneView.h
+++ b/Engine/Source/Runtime/Engine/Public/SceneView.h
@@ -931,6 +931,15 @@ public:
 	/** Current buffer visualization mode */
 	FName CurrentBufferVisualizationMode;
 
+	/** Current visualize debug color material name */
+	FName CurrentVisualizeDebugColorMaterialName;
+
+	/** Current visualize debug grayscale material name */
+	FName CurrentVisualizeDebugGrayscaleMaterialName;
+
+	/** Current visualize debug custom material name */
+	FName CurrentVisualizeDebugCustomMaterialName;
+
 #if WITH_EDITOR
 	/* Whether to use the pixel inspector */
 	bool bUsePixelInspector;
@@ -1262,6 +1271,11 @@ public:
 	/** Configure post process settings for the buffer visualization system */
 	void ConfigureBufferVisualizationSettings();
 
+#if !(UE_BUILD_SHIPPING || UE_BUILD_TEST)
+	/** Configure post process settings for debug image */
+	void ConfigureDebugVisualizationSettings();
+#endif
+
 	/** Get the feature level for this view (cached from the scene so this is not different per view) **/
 	ERHIFeatureLevel::Type GetFeatureLevel() const { return FeatureLevel; }
 
diff --git a/Engine/Source/Runtime/Engine/Public/ShowFlags.h b/Engine/Source/Runtime/Engine/Public/ShowFlags.h
--- a/Engine/Source/Runtime/Engine/Public/ShowFlags.h
+++ b/Engine/Source/Runtime/Engine/Public/ShowFlags.h
@@ -400,6 +400,9 @@ private:
 		SetPathTracing(false);
 		SetRayTracingDebug(false);
 		SetVisualizeSkyAtmosphere(false);
+		SetVisualizeDebugColor(false);
+		SetVisualizeDebugGrayscale(false);
+		SetVisualizeDebugCustomPostProcessMaterial(false);
 	}
 
 
diff --git a/Engine/Source/Runtime/Engine/Public/ShowFlagsValues.inl b/Engine/Source/Runtime/Engine/Public/ShowFlagsValues.inl
--- a/Engine/Source/Runtime/Engine/Public/ShowFlagsValues.inl
+++ b/Engine/Source/Runtime/Engine/Public/ShowFlagsValues.inl
@@ -350,6 +350,12 @@ SHOWFLAG_FIXED_IN_SHIPPING(0, RayTracingDebug, SFG_Hidden, NSLOCTEXT("UnrealEd",
 
 /** Enable the SkyAtmosphere visualization to be drawn on screen */
 SHOWFLAG_FIXED_IN_SHIPPING(0, VisualizeSkyAtmosphere, SFG_Visualize, NSLOCTEXT("UnrealEd", "VisualizeSkyAtmosphereSF", "Sky Atmosphere"))
+/** Shows a debug color post process material defined in renderer config  */
+SHOWFLAG_FIXED_IN_SHIPPING(0, VisualizeDebugColor, SFG_Developer, NSLOCTEXT("UnrealEd", "VisualizeDebugColorSF", "Visualize Debug Color"))
+/** Shows a debug grayscale post process material defined in renderer config  */
+SHOWFLAG_FIXED_IN_SHIPPING(0, VisualizeDebugGrayscale, SFG_Developer, NSLOCTEXT("UnrealEd", "VisualizeDebugGrayscaleSF", "Visualize Debug Grayscale"))
+/** Shows a debug custom post process material defined in renderer config  */
+SHOWFLAG_FIXED_IN_SHIPPING(0, VisualizeDebugCustomPostProcessMaterial, SFG_Developer, NSLOCTEXT("UnrealEd", "VisualizeDebugCustomPostProcessMaterialSF", "Visualize Debug Custom"))
 
 #undef SHOWFLAG_ALWAYS_ACCESSIBLE
 #undef SHOWFLAG_FIXED_IN_SHIPPING
diff --git a/Engine/Source/Runtime/Renderer/Private/DeferredShadingRenderer.cpp b/Engine/Source/Runtime/Renderer/Private/DeferredShadingRenderer.cpp
--- a/Engine/Source/Runtime/Renderer/Private/DeferredShadingRenderer.cpp
+++ b/Engine/Source/Runtime/Renderer/Private/DeferredShadingRenderer.cpp
@@ -13,6 +13,7 @@
 #include "ScreenRendering.h"
 #include "PostProcess/SceneFilterRendering.h"
 #include "PostProcess/PostProcessSubsurface.h"
+#include "PostProcess/PostProcessVisualizeDebugMaterial.h"
 #include "CompositionLighting/CompositionLighting.h"
 #include "FXSystem.h"
 #include "OneColorShader.h"
@@ -2905,18 +2906,31 @@ void FDeferredShadingSceneRenderer::Render(FRHICommandListImmediate& RHICmdList)
 
 				SCOPED_GPU_MASK(RHICmdList, View.GPUMask);
 				RDG_EVENT_SCOPE_CONDITIONAL(GraphBuilder, Views.Num() > 1, "View%d", ViewIndex);
-				AddDebugPostProcessingPasses(GraphBuilder, View, PostProcessingInputs);
+				AddDebugViewPostProcessingPasses(GraphBuilder, View, PostProcessingInputs);
 			}
 		}
-		else
+		else 
 		{
 			for (int32 ViewIndex = 0; ViewIndex < Views.Num(); ViewIndex++)
 			{
 				FViewInfo& View = Views[ViewIndex];
 
-				SCOPED_GPU_MASK(RHICmdList, View.GPUMask);
-				RDG_EVENT_SCOPE_CONDITIONAL(GraphBuilder, Views.Num() > 1, "View%d", ViewIndex);
-				AddPostProcessingPasses(GraphBuilder, View, PostProcessingInputs);
+				if (IsPostProcessVisualizeDebugMaterialEnabled(View))
+				{
+					const UMaterialInterface* DebugMaterialInterface = GetPostProcessVisualizeDebugMaterialInterface(View);
+					check(DebugMaterialInterface);
+
+					SCOPED_GPU_MASK(RHICmdList, View.GPUMask);
+					RDG_EVENT_SCOPE_CONDITIONAL(GraphBuilder, Views.Num() > 1, "View%d", ViewIndex);
+					AddVisualizeDebugMaterialPostProcessingPasses(GraphBuilder, View, PostProcessingInputs, DebugMaterialInterface);
+
+				}
+				else
+				{
+					SCOPED_GPU_MASK(RHICmdList, View.GPUMask);
+					RDG_EVENT_SCOPE_CONDITIONAL(GraphBuilder, Views.Num() > 1, "View%d", ViewIndex);
+					AddPostProcessingPasses(GraphBuilder, View, PostProcessingInputs);
+				}
 			}
 		}
 
diff --git a/Engine/Source/Runtime/Renderer/Private/PostProcess/PostProcessDeviceEncodingOnly.cpp b/Engine/Source/Runtime/Renderer/Private/PostProcess/PostProcessDeviceEncodingOnly.cpp
new file mode 100644
--- /dev/null
+++ b/Engine/Source/Runtime/Renderer/Private/PostProcess/PostProcessDeviceEncodingOnly.cpp
@@ -0,0 +1,287 @@
+// Copyright Epic Games, Inc. All Rights Reserved.
+
+/*=========================================================================================
+	PostProcessDeviceEncodingOnly.cpp: Post processing device encoding only implementation.
+===========================================================================================*/
+
+#include "PostProcess/PostProcessDeviceEncodingOnly.h"
+#include "EngineGlobals.h"
+#include "ScenePrivate.h"
+#include "RendererModule.h"
+#include "PostProcess/PostProcessing.h"
+#include "PostProcess/SceneFilterRendering.h"
+#include "PostProcess/PostProcessCombineLUTs.h"
+#include "PostProcess/PostProcessMobile.h"
+#include "PostProcess/PostProcessing.h"
+#include "ClearQuad.h"
+#include "PipelineStateCache.h"
+
+namespace
+{
+	const int32 GDeviceEncodingOnlyComputeTileSizeX = 8;
+	const int32 GDeviceEncodingOnlyComputeTileSizeY = 8;
+
+	namespace DeviceEncodingOnlyPermutation
+	{
+		// Desktop renderer permutation dimensions.
+		class FDeviceEncodingOnlyOutputDeviceDim : SHADER_PERMUTATION_ENUM_CLASS("DIM_OUTPUT_DEVICE", EDeviceEncodingOnlyOutputDevice);
+
+		using FDesktopDomain = TShaderPermutationDomain<FDeviceEncodingOnlyOutputDeviceDim>;
+	
+	} // namespace DeviceEncodingOnlyPermutation
+} // namespace
+
+
+FDeviceEncodingOnlyOutputDeviceParameters GetDeviceEncodingOnlyOutputDeviceParameters(const FSceneViewFamily& Family)
+{
+	static TConsoleVariableData<int32>* CVarOutputGamut = IConsoleManager::Get().FindTConsoleVariableDataInt(TEXT("r.HDR.Display.ColorGamut"));
+	static TConsoleVariableData<int32>* CVarOutputDevice = IConsoleManager::Get().FindTConsoleVariableDataInt(TEXT("r.HDR.Display.OutputDevice"));
+	static TConsoleVariableData<float>* CVarOutputGamma = IConsoleManager::Get().FindTConsoleVariableDataFloat(TEXT("r.TonemapperGamma"));
+
+	EDeviceEncodingOnlyOutputDevice OutputDeviceValue;
+
+	if (Family.SceneCaptureSource == SCS_FinalColorHDR)
+	{
+		OutputDeviceValue = EDeviceEncodingOnlyOutputDevice::LinearNoToneCurve;
+	}
+	else if (Family.SceneCaptureSource == SCS_FinalToneCurveHDR)
+	{
+		OutputDeviceValue = EDeviceEncodingOnlyOutputDevice::LinearWithToneCurve;
+	}
+	else if (Family.bIsHDR)
+	{
+		OutputDeviceValue = EDeviceEncodingOnlyOutputDevice::ACES1000nitST2084;
+	}
+	else
+	{
+		OutputDeviceValue = static_cast<EDeviceEncodingOnlyOutputDevice>(CVarOutputDevice->GetValueOnRenderThread());
+		OutputDeviceValue = static_cast<EDeviceEncodingOnlyOutputDevice>(FMath::Clamp(static_cast<int32>(OutputDeviceValue), 0, static_cast<int32>(EDeviceEncodingOnlyOutputDevice::MAX) - 1));
+	}
+
+	float Gamma = CVarOutputGamma->GetValueOnRenderThread();
+
+	if (PLATFORM_APPLE && Gamma == 0.0f)
+	{
+		Gamma = 2.2f;
+	}
+
+	// Enforce user-controlled ramp over sRGB or Rec709
+	if (Gamma > 0.0f && (OutputDeviceValue == EDeviceEncodingOnlyOutputDevice::sRGB || OutputDeviceValue == EDeviceEncodingOnlyOutputDevice::Rec709))
+	{
+		OutputDeviceValue = EDeviceEncodingOnlyOutputDevice::ExplicitGammaMapping;
+	}
+
+	FVector InvDisplayGammaValue;
+	InvDisplayGammaValue.X = 1.0f / Family.RenderTarget->GetDisplayGamma();
+	InvDisplayGammaValue.Y = 2.2f / Family.RenderTarget->GetDisplayGamma();
+	InvDisplayGammaValue.Z = 1.0f / FMath::Max(Gamma, 1.0f);
+
+	FDeviceEncodingOnlyOutputDeviceParameters Parameters;
+	Parameters.InverseGamma = InvDisplayGammaValue;
+	Parameters.OutputDevice = static_cast<uint32>(OutputDeviceValue);
+	Parameters.OutputGamut = CVarOutputGamut->GetValueOnRenderThread();
+	return Parameters;
+}
+
+
+BEGIN_SHADER_PARAMETER_STRUCT(FDeviceEncodingOnlyParameters, )
+	SHADER_PARAMETER_STRUCT_REF(FViewUniformShaderParameters, View)
+	SHADER_PARAMETER_STRUCT_INCLUDE(FDeviceEncodingOnlyOutputDeviceParameters, OutputDevice)
+	SHADER_PARAMETER_STRUCT(FScreenPassTextureViewportParameters, Color)
+	SHADER_PARAMETER_STRUCT(FScreenPassTextureViewportParameters, Output)
+	SHADER_PARAMETER_RDG_TEXTURE(Texture2D, ColorTexture)
+	SHADER_PARAMETER_SAMPLER(SamplerState, ColorSampler)
+	SHADER_PARAMETER(float, EditorNITLevel)
+	SHADER_PARAMETER(uint32, bOutputInHDR)
+END_SHADER_PARAMETER_STRUCT()
+
+class FDeviceEncodingOnlyPS : public FGlobalShader
+{
+public:
+	DECLARE_GLOBAL_SHADER(FDeviceEncodingOnlyPS);
+	SHADER_USE_PARAMETER_STRUCT(FDeviceEncodingOnlyPS, FGlobalShader);
+
+	using FPermutationDomain = DeviceEncodingOnlyPermutation::FDesktopDomain;
+
+	BEGIN_SHADER_PARAMETER_STRUCT(FParameters, )
+		SHADER_PARAMETER_STRUCT_INCLUDE(FDeviceEncodingOnlyParameters, DeviceEncoding)
+		RENDER_TARGET_BINDING_SLOTS()
+	END_SHADER_PARAMETER_STRUCT()
+
+	static bool ShouldCompilePermutation(const FGlobalShaderPermutationParameters& Parameters)
+	{
+		if (!IsFeatureLevelSupported(Parameters.Platform, ERHIFeatureLevel::SM5))
+		{
+			return false;
+		}
+		return true;
+	}
+
+	static void ModifyCompilationEnvironment(const FGlobalShaderPermutationParameters& Parameters, FShaderCompilerEnvironment& OutEnvironment)
+	{
+		const int UseVolumeLut = PipelineVolumeTextureLUTSupportGuaranteedAtRuntime(Parameters.Platform) ? 1 : 0;
+		OutEnvironment.SetDefine(TEXT("USE_VOLUME_LUT"), UseVolumeLut);
+	}
+};
+
+IMPLEMENT_GLOBAL_SHADER(FDeviceEncodingOnlyPS, "/Engine/Private/PostProcessDeviceEncodingOnly.usf", "MainPS", SF_Pixel);
+
+class FDeviceEncodingOnlyCS : public FGlobalShader
+{
+public:
+	DECLARE_GLOBAL_SHADER(FDeviceEncodingOnlyCS);
+	SHADER_USE_PARAMETER_STRUCT(FDeviceEncodingOnlyCS, FGlobalShader);
+
+	using FPermutationDomain = TShaderPermutationDomain<DeviceEncodingOnlyPermutation::FDesktopDomain>;
+
+	BEGIN_SHADER_PARAMETER_STRUCT(FParameters, )
+		SHADER_PARAMETER_STRUCT_INCLUDE(FDeviceEncodingOnlyParameters, DeviceEncoding)
+		SHADER_PARAMETER_RDG_TEXTURE_UAV(RWTexture2D, RWOutputTexture)
+	END_SHADER_PARAMETER_STRUCT()
+
+	static bool ShouldCompilePermutation(const FGlobalShaderPermutationParameters& Parameters)
+	{
+		if (!IsFeatureLevelSupported(Parameters.Platform, ERHIFeatureLevel::SM5))
+		{
+			return false;
+		}
+
+		return true;
+	}
+
+	static void ModifyCompilationEnvironment(const FGlobalShaderPermutationParameters& Parameters, FShaderCompilerEnvironment& OutEnvironment)
+	{
+		OutEnvironment.SetDefine(TEXT("THREADGROUP_SIZEX"), GDeviceEncodingOnlyComputeTileSizeX);
+		OutEnvironment.SetDefine(TEXT("THREADGROUP_SIZEY"), GDeviceEncodingOnlyComputeTileSizeY);
+	}
+};
+
+IMPLEMENT_GLOBAL_SHADER(FDeviceEncodingOnlyCS, "/Engine/Private/PostProcessDeviceEncodingOnly.usf", "MainCS", SF_Compute);
+
+FScreenPassTexture AddDeviceEncodingOnlyPass(FRDGBuilder& GraphBuilder, const FViewInfo& View, const FDeviceEncodingOnlyInputs& Inputs)
+{
+	check(Inputs.SceneColor.IsValid());
+
+	const FSceneViewFamily& ViewFamily = *(View.Family);
+	const FPostProcessSettings& PostProcessSettings = View.FinalPostProcessSettings;
+
+	const FScreenPassTextureViewport SceneColorViewport(Inputs.SceneColor);
+
+	FScreenPassRenderTarget Output = Inputs.OverrideOutput;
+
+	if (!Output.IsValid())
+	{
+		FRDGTextureDesc OutputDesc = Inputs.SceneColor.Texture->Desc;
+		OutputDesc.Reset();
+		OutputDesc.TargetableFlags |= View.bUseComputePasses ? TexCreate_UAV : TexCreate_RenderTargetable;
+		// RGB is the color in LDR, A is the luminance for PostprocessAA
+		OutputDesc.Format = Inputs.bOutputInHDR ? GRHIHDRDisplayOutputFormat : PF_B8G8R8A8;
+		OutputDesc.ClearValue = FClearValueBinding(FLinearColor(0, 0, 0, 0));
+		OutputDesc.Flags |= GFastVRamConfig.Tonemap;
+
+		const FDeviceEncodingOnlyOutputDeviceParameters OutputDeviceParameters = GetDeviceEncodingOnlyOutputDeviceParameters(*View.Family);
+		const EDeviceEncodingOnlyOutputDevice OutputDevice = static_cast<EDeviceEncodingOnlyOutputDevice>(OutputDeviceParameters.OutputDevice);
+
+		if (OutputDevice == EDeviceEncodingOnlyOutputDevice::LinearEXR)
+		{
+			OutputDesc.Format = PF_A32B32G32R32F;
+		}
+		if (OutputDevice == EDeviceEncodingOnlyOutputDevice::LinearNoToneCurve || OutputDevice == EDeviceEncodingOnlyOutputDevice::LinearWithToneCurve)
+		{
+			OutputDesc.Format = PF_FloatRGBA;
+		}
+
+		Output = FScreenPassRenderTarget(
+			GraphBuilder.CreateTexture(OutputDesc, TEXT("DeviceEncodingOnly")),
+			Inputs.SceneColor.ViewRect,
+			View.GetOverwriteLoadAction());
+	}
+
+	const FScreenPassTextureViewport OutputViewport(Output);
+
+	FRHISamplerState* BilinearClampSampler = TStaticSamplerState<SF_Bilinear, AM_Clamp, AM_Clamp, AM_Clamp>::GetRHI();
+	FRHISamplerState* PointClampSampler = TStaticSamplerState<SF_Point, AM_Clamp, AM_Clamp, AM_Clamp>::GetRHI();
+
+	float EditorNITLevel = 160.0f;
+
+#if WITH_EDITOR
+	{
+		static auto CVarHDRNITLevel = IConsoleManager::Get().FindConsoleVariable(TEXT("Editor.HDRNITLevel"));
+		if (CVarHDRNITLevel)
+		{
+			EditorNITLevel = CVarHDRNITLevel->GetFloat();
+		}
+	}
+#endif
+
+	FDeviceEncodingOnlyParameters CommonParameters;
+	CommonParameters.View = View.ViewUniformBuffer;
+	CommonParameters.OutputDevice = GetDeviceEncodingOnlyOutputDeviceParameters(ViewFamily);
+	CommonParameters.Color = GetScreenPassTextureViewportParameters(SceneColorViewport);
+	CommonParameters.Output = GetScreenPassTextureViewportParameters(OutputViewport);
+	CommonParameters.ColorTexture = Inputs.SceneColor.Texture;
+	CommonParameters.ColorSampler = BilinearClampSampler;
+	CommonParameters.EditorNITLevel = EditorNITLevel;
+	CommonParameters.bOutputInHDR = ViewFamily.bIsHDR;
+
+	// Generate permutation vector for the desktop tonemapper.
+	DeviceEncodingOnlyPermutation::FDesktopDomain DesktopPermutationVector;
+	DesktopPermutationVector.Set<DeviceEncodingOnlyPermutation::FDeviceEncodingOnlyOutputDeviceDim>(EDeviceEncodingOnlyOutputDevice(CommonParameters.OutputDevice.OutputDevice));
+
+	// Override output might not support UAVs.
+	const bool bComputePass = (Output.Texture->Desc.TargetableFlags & TexCreate_UAV) == TexCreate_UAV ? View.bUseComputePasses : false;
+
+	if (bComputePass)
+	{
+		FDeviceEncodingOnlyCS::FParameters* PassParameters = GraphBuilder.AllocParameters<FDeviceEncodingOnlyCS::FParameters>();
+		PassParameters->DeviceEncoding = CommonParameters;
+		PassParameters->RWOutputTexture = GraphBuilder.CreateUAV(Output.Texture);
+
+		FDeviceEncodingOnlyCS::FPermutationDomain PermutationVector;
+		PermutationVector.Set<DeviceEncodingOnlyPermutation::FDesktopDomain>(DesktopPermutationVector);
+
+		TShaderMapRef<FDeviceEncodingOnlyCS> ComputeShader(View.ShaderMap, PermutationVector);
+
+		FComputeShaderUtils::AddPass(
+			GraphBuilder,
+			RDG_EVENT_NAME("DeviceEncodingOnly %dx%d (CS)", OutputViewport.Rect.Width(), OutputViewport.Rect.Height()),
+			ComputeShader,
+			PassParameters,
+			FComputeShaderUtils::GetGroupCount(OutputViewport.Rect.Size(), FIntPoint(GDeviceEncodingOnlyComputeTileSizeX, GDeviceEncodingOnlyComputeTileSizeY)));
+	}
+	else
+	{
+		FDeviceEncodingOnlyPS::FParameters* PassParameters = GraphBuilder.AllocParameters<FDeviceEncodingOnlyPS::FParameters>();
+		PassParameters->DeviceEncoding = CommonParameters;
+		PassParameters->RenderTargets[0] = Output.GetRenderTargetBinding();
+
+		TShaderMapRef<FScreenPassVS> VertexShader(View.ShaderMap);
+		TShaderMapRef<FDeviceEncodingOnlyPS> PixelShader(View.ShaderMap, DesktopPermutationVector);
+
+		// If this is a stereo view, there's a good chance we need alpha out of the tonemapper
+		// @todo: Remove this once Oculus fix the bug in their runtime that requires alpha here.
+		//const bool bIsStereo = IStereoRendering::IsStereoEyeView(View);
+		//FRHIBlendState* BlendState = bIsStereo ? FScreenPassPipelineState::FDefaultBlendState::GetRHI() : TStaticBlendStateWriteMask<CW_RGB>::GetRHI();
+		FRHIBlendState* BlendState = TStaticBlendStateWriteMask<CW_RGB>::GetRHI();
+		FRHIDepthStencilState* DepthStencilState = FScreenPassPipelineState::FDefaultDepthStencilState::GetRHI();
+
+		EScreenPassDrawFlags DrawFlags = EScreenPassDrawFlags::AllowHMDHiddenAreaMask;
+
+		AddDrawScreenPass(
+			GraphBuilder,
+			RDG_EVENT_NAME("DeviceEncodingOnly %dx%d (PS)", OutputViewport.Rect.Width(), OutputViewport.Rect.Height()),
+			View,
+			OutputViewport,
+			SceneColorViewport,
+			FScreenPassPipelineState(VertexShader, PixelShader, BlendState, DepthStencilState),
+			DrawFlags,
+			PassParameters,
+			[PixelShader, PassParameters](FRHICommandList& RHICmdList)
+			{
+				SetShaderParameters(RHICmdList, PixelShader, PixelShader.GetPixelShader(), *PassParameters);
+			});
+	}
+
+	return MoveTemp(Output);
+}
+
diff --git a/Engine/Source/Runtime/Renderer/Private/PostProcess/PostProcessDeviceEncodingOnly.h b/Engine/Source/Runtime/Renderer/Private/PostProcess/PostProcessDeviceEncodingOnly.h
new file mode 100644
--- /dev/null
+++ b/Engine/Source/Runtime/Renderer/Private/PostProcess/PostProcessDeviceEncodingOnly.h
@@ -0,0 +1,48 @@
+// Copyright Epic Games, Inc. All Rights Reserved.
+
+#pragma once
+
+#include "PostProcess/RenderingCompositionGraph.h"
+#include "PostProcess/PostProcessEyeAdaptation.h"
+#include "ScreenPass.h"
+#include "OverridePassSequence.h"
+#include "Math/Halton.h"
+
+// You must update values in PostProcessTonemap.usf when changing this enum.
+enum class EDeviceEncodingOnlyOutputDevice
+{
+	sRGB,
+	Rec709,
+	ExplicitGammaMapping,
+	ACES1000nitST2084,
+	ACES2000nitST2084,
+	ACES1000nitScRGB,
+	ACES2000nitScRGB,
+	LinearEXR,
+	LinearNoToneCurve,
+	LinearWithToneCurve,
+
+	MAX
+};
+
+BEGIN_SHADER_PARAMETER_STRUCT(FDeviceEncodingOnlyOutputDeviceParameters, )
+	SHADER_PARAMETER(FVector, InverseGamma)
+	SHADER_PARAMETER(uint32, OutputDevice)
+	SHADER_PARAMETER(uint32, OutputGamut)
+END_SHADER_PARAMETER_STRUCT()
+
+FDeviceEncodingOnlyOutputDeviceParameters GetDeviceEncodingOnlyOutputDeviceParameters(const FSceneViewFamily& Family);
+
+struct FDeviceEncodingOnlyInputs
+{
+	// [Optional] Render to the specified output. If invalid, a new texture is created and returned.
+	FScreenPassRenderTarget OverrideOutput;
+
+	// [Required] HDR scene color to tonemap.
+	FScreenPassTexture SceneColor;
+
+	// Whether to leave the final output in HDR.
+	bool bOutputInHDR = false;
+};
+
+FScreenPassTexture AddDeviceEncodingOnlyPass(FRDGBuilder& GraphBuilder, const FViewInfo& View, const FDeviceEncodingOnlyInputs& Inputs);
diff --git a/Engine/Source/Runtime/Renderer/Private/PostProcess/PostProcessVisualizeDebugMaterial.cpp b/Engine/Source/Runtime/Renderer/Private/PostProcess/PostProcessVisualizeDebugMaterial.cpp
new file mode 100644
--- /dev/null
+++ b/Engine/Source/Runtime/Renderer/Private/PostProcess/PostProcessVisualizeDebugMaterial.cpp
@@ -0,0 +1,43 @@
+// Copyright Epic Games, Inc. All Rights Reserved.
+
+#include "PostProcess/PostProcessVisualizeDebugMaterial.h"
+#include "EngineGlobals.h"
+
+static bool IsPostProcessVisualizeDebugColorMaterialEnabled(const FViewInfo& View) 
+{
+	return View.Family->EngineShowFlags.VisualizeDebugColor && View.CurrentVisualizeDebugColorMaterialName != NAME_None;
+}
+
+static bool IsPostProcessVisualizeDebugGrayscaleMaterialEnabled(const FViewInfo& View) 
+{
+	return View.Family->EngineShowFlags.VisualizeDebugGrayscale && View.CurrentVisualizeDebugGrayscaleMaterialName != NAME_None;
+}
+
+static bool IsPostProcessVisualizeDebugCustomMaterialEnabled(const FViewInfo& View) 
+{
+	return View.Family->EngineShowFlags.VisualizeDebugCustomPostProcessMaterial && View.CurrentVisualizeDebugCustomMaterialName != NAME_None;
+}
+
+bool IsPostProcessVisualizeDebugMaterialEnabled(const FViewInfo& View) 
+{
+	return (IsPostProcessVisualizeDebugColorMaterialEnabled(View) || IsPostProcessVisualizeDebugGrayscaleMaterialEnabled(View) || IsPostProcessVisualizeDebugCustomMaterialEnabled(View));
+}
+
+// Returns whether the debug custom material pass needs to render on screen.
+const UMaterialInterface* GetPostProcessVisualizeDebugMaterialInterface(const FViewInfo& View) 
+{
+	if (IsPostProcessVisualizeDebugCustomMaterialEnabled(View))
+	{
+		return View.FinalPostProcessSettings.DebugCustomVisualizationMaterial;
+	}
+	else if (IsPostProcessVisualizeDebugColorMaterialEnabled(View))
+	{
+		return View.FinalPostProcessSettings.DebugColorVisualizationMaterial;
+	}
+	else if (IsPostProcessVisualizeDebugGrayscaleMaterialEnabled(View))
+	{
+		return View.FinalPostProcessSettings.DebugGrayscaleVisualizationMaterial;
+	}
+
+	return NULL;
+}
diff --git a/Engine/Source/Runtime/Renderer/Private/PostProcess/PostProcessVisualizeDebugMaterial.h b/Engine/Source/Runtime/Renderer/Private/PostProcess/PostProcessVisualizeDebugMaterial.h
new file mode 100644
--- /dev/null
+++ b/Engine/Source/Runtime/Renderer/Private/PostProcess/PostProcessVisualizeDebugMaterial.h
@@ -0,0 +1,11 @@
+// Copyright Epic Games, Inc. All Rights Reserved.
+
+#pragma once
+
+#include "ScreenPass.h"
+
+// Returns whether any debug material pass is enabled.
+bool IsPostProcessVisualizeDebugMaterialEnabled(const FViewInfo& View);
+
+// Returns whether any debug material pass is enabled.
+const UMaterialInterface* GetPostProcessVisualizeDebugMaterialInterface(const FViewInfo& View);
diff --git a/Engine/Source/Runtime/Renderer/Private/PostProcess/PostProcessing.cpp b/Engine/Source/Runtime/Renderer/Private/PostProcess/PostProcessing.cpp
--- a/Engine/Source/Runtime/Renderer/Private/PostProcess/PostProcessing.cpp
+++ b/Engine/Source/Runtime/Renderer/Private/PostProcess/PostProcessing.cpp
@@ -23,6 +23,7 @@
 #include "PostProcess/PostProcessLensFlares.h"
 #include "PostProcess/PostProcessBokehDOF.h"
 #include "PostProcess/PostProcessCombineLUTs.h"
+#include "PostProcess/PostProcessDeviceEncodingOnly.h"
 #include "PostProcess/PostProcessTemporalAA.h"
 #include "PostProcess/PostProcessMotionBlur.h"
 #include "PostProcess/PostProcessDOF.h"
@@ -31,6 +32,7 @@
 #include "PostProcess/PostProcessVisualizeComplexity.h"
 #include "PostProcess/PostProcessCompositeEditorPrimitives.h"
 #include "PostProcess/PostProcessTestImage.h"
+#include "PostProcess/PostProcessVisualizeDebugMaterial.h"
 #include "PostProcess/PostProcessFFTBloom.h"
 #include "PostProcess/PostProcessStreamingAccuracyLegend.h"
 #include "PostProcess/PostProcessSubsurface.h"
@@ -892,7 +894,7 @@ void AddPostProcessingPasses(FRDGBuilder& GraphBuilder, const FViewInfo& View, c
 	}
 }
 
-void AddDebugPostProcessingPasses(FRDGBuilder& GraphBuilder, const FViewInfo& View, const FPostProcessingInputs& Inputs)
+void AddDebugViewPostProcessingPasses(FRDGBuilder& GraphBuilder, const FViewInfo& View, const FPostProcessingInputs& Inputs)
 {
 	CSV_SCOPED_TIMING_STAT_EXCLUSIVE(RenderPostProcessing);
 	QUICK_SCOPE_CYCLE_COUNTER(STAT_PostProcessing_Process);
@@ -1075,6 +1077,48 @@ void AddDebugPostProcessingPasses(FRDGBuilder& GraphBuilder, const FViewInfo& Vi
 	}
 }
 
+#if !(UE_BUILD_SHIPPING || UE_BUILD_TEST)
+
+void AddVisualizeDebugMaterialPostProcessingPasses(FRDGBuilder& GraphBuilder, const FViewInfo& View, const FPostProcessingInputs& Inputs, const UMaterialInterface* InMaterialInterface)
+{
+	CSV_SCOPED_TIMING_STAT_EXCLUSIVE(RenderPostProcessing);
+	QUICK_SCOPE_CYCLE_COUNTER(STAT_PostProcessing_Process);
+
+	check(IsInRenderingThread());
+	check(View.VerifyMembersChecks());
+	check(InMaterialInterface);
+	Inputs.Validate();
+
+	const FIntRect PrimaryViewRect = View.ViewRect;
+
+	const FSceneTextureParameters& SceneTextures = *Inputs.SceneTextures;
+	const FScreenPassRenderTarget ViewFamilyOutput = FScreenPassRenderTarget::CreateViewFamilyOutput(Inputs.ViewFamilyTexture, View);
+
+	// Scene color is updated incrementally through the post process pipeline.
+	FScreenPassTexture SceneColor(Inputs.SceneColor, PrimaryViewRect);
+
+	const FEngineShowFlags& EngineShowFlags = View.Family->EngineShowFlags;
+	const bool bVisualizeHDR = EngineShowFlags.VisualizeHDR;
+	const bool bViewFamilyOutputInHDR = GRHISupportsHDROutput && IsHDREnabled();
+	const bool bOutputInHDR = IsPostProcessingOutputInHDR();
+
+	// Post Process Material - Before Color Correction
+	FPostProcessMaterialInputs PostProcessMaterialInputs;
+	PostProcessMaterialInputs.SetInput(EPostProcessMaterialInput::SceneColor, SceneColor);
+
+	SceneColor = AddPostProcessMaterialPass(GraphBuilder, View, PostProcessMaterialInputs, InMaterialInterface);
+
+	// Replace tonemapper with device encoding only pass, which converts the scene color to device-specific color.
+	FDeviceEncodingOnlyInputs PassInputs;
+	PassInputs.OverrideOutput = ViewFamilyOutput;
+	PassInputs.SceneColor = SceneColor;
+	PassInputs.bOutputInHDR = bViewFamilyOutputInHDR;
+
+	SceneColor = AddDeviceEncodingOnlyPass(GraphBuilder, View, PassInputs);
+}
+
+#endif
+
 ///////////////////////////////////////////////////////////////////////////
 // Mobile Post Processing
 //////////////////////////////////////////////////////////////////////////
diff --git a/Engine/Source/Runtime/Renderer/Private/PostProcess/PostProcessing.h b/Engine/Source/Runtime/Renderer/Private/PostProcess/PostProcessing.h
--- a/Engine/Source/Runtime/Renderer/Private/PostProcess/PostProcessing.h
+++ b/Engine/Source/Runtime/Renderer/Private/PostProcess/PostProcessing.h
@@ -55,7 +55,13 @@ struct FPostProcessingInputs
 
 void AddPostProcessingPasses(FRDGBuilder& GraphBuilder, const FViewInfo& View, const FPostProcessingInputs& Inputs);
 
-void AddDebugPostProcessingPasses(FRDGBuilder& GraphBuilder, const FViewInfo& View, const FPostProcessingInputs& Inputs);
+void AddDebugViewPostProcessingPasses(FRDGBuilder& GraphBuilder, const FViewInfo& View, const FPostProcessingInputs& Inputs);
+
+#if !(UE_BUILD_SHIPPING || UE_BUILD_TEST)
+
+void AddVisualizeDebugMaterialPostProcessingPasses(FRDGBuilder& GraphBuilder, const FViewInfo& View, const FPostProcessingInputs& Inputs, const UMaterialInterface* InMaterialInterface);
+
+#endif
 
 // For compatibility with composition graph passes until they are ported to Render Graph.
 class FPostProcessVS : public FScreenPassVS
