Refactors denoiser's frequency handling into its own dedicated structure

diff --git a/Engine/Shaders/Private/ScreenSpaceDenoise/SSDDefinitions.ush b/Engine/Shaders/Private/ScreenSpaceDenoise/SSDDefinitions.ush
--- a/Engine/Shaders/Private/ScreenSpaceDenoise/SSDDefinitions.ush
+++ b/Engine/Shaders/Private/ScreenSpaceDenoise/SSDDefinitions.ush
@@ -223,8 +223,6 @@
 
 	#define SAMPLE_SET_HEXAWEB 11
 
-	#define SAMPLE_SET_STOCASTIC_HIERARCHY 12
-
 	#define SAMPLE_SET_DIRECTIONAL_RECT 14
 	#define SAMPLE_SET_DIRECTIONAL_ELLIPSE 15
 
diff --git a/Engine/Shaders/Private/ScreenSpaceDenoise/SSDInjest.usf b/Engine/Shaders/Private/ScreenSpaceDenoise/SSDInjest.usf
--- a/Engine/Shaders/Private/ScreenSpaceDenoise/SSDInjest.usf
+++ b/Engine/Shaders/Private/ScreenSpaceDenoise/SSDInjest.usf
@@ -119,7 +119,9 @@ void MainCS(
 	}
 
 	// Sample the reference sample.
-	FSSDSignalArray MultiplexedSamples = SampleMultiplexedSignals(
+	FSSDSignalArray MultiplexedSamples;
+	FSSDSignalFrequencyArray MultiplexedFrequencies;
+	SampleMultiplexedSignals(
 		SignalInput_Textures_0,
 		SignalInput_Textures_1,
 		SignalInput_Textures_2,
@@ -127,28 +129,20 @@ void MainCS(
 		GlobalPointClampedSampler,
 		CONFIG_SIGNAL_INPUT_LAYOUT,
 		/* MultiplexedSampleId = */ 0,
-		/* bNormalizeSample = */ false,
+		/* bNormalizeSample = */ true,
 		SceneBufferUV,
-		/* MipLevel = */ 0);
+		/* out */ MultiplexedSamples,
+		/* out */ MultiplexedFrequencies);
 
 	// Performs sample value precomputation.
 	UNROLL
 	for (uint MultiplexId = 0; MultiplexId < CONFIG_SIGNAL_BATCH_SIZE; MultiplexId++)
 	{
 		const uint BatchedSignalId = MultiplexId / CONFIG_MULTIPLEXED_SIGNALS_PER_SIGNAL_DOMAIN;
-		MultiplexedSamples.Array[MultiplexId].WorldBluringRadius = GetSignalWorldBluringRadius(
-			MultiplexedSamples.Array[MultiplexId], RefSceneMetadata, BatchedSignalId);
+		MultiplexedFrequencies.Array[MultiplexId].WorldBluringRadius = GetSignalWorldBluringRadius(
+			MultiplexedFrequencies.Array[MultiplexId], RefSceneMetadata, BatchedSignalId);
 	}
 
-	// Whitelists computation that needs to be done.
-	const uint MultiplexCount = CONFIG_SIGNAL_BATCH_SIZE;
-	FSSDSignalArray OutputSamples = CreateSignalArrayFromScalarValue(0.0);
-	UNROLL
-	for (uint MultiplexId = 0; MultiplexId < MultiplexCount; MultiplexId++)
-	{
-		OutputSamples.Array[MultiplexId] = MultiplexedSamples.Array[MultiplexId];
-	}
-	
 	// No need to keep DispatchThreadId, while SceneBufferUV is arround at highest VGPR peak because center of the unique pixel to sample.
 	uint2 OutputPixelPostion = BufferUVToBufferPixelCoord(SceneBufferUV);
 		
@@ -160,7 +154,10 @@ void MainCS(
 			SignalOutput_UAVs_1,
 			SignalOutput_UAVs_2,
 			SignalOutput_UAVs_3,
-			CONFIG_SIGNAL_OUTPUT_LAYOUT, MultiplexCount,
-			OutputPixelPostion, OutputSamples.Array);
+			CONFIG_SIGNAL_OUTPUT_LAYOUT,
+			CONFIG_SIGNAL_BATCH_SIZE,
+			OutputPixelPostion,
+			MultiplexedSamples,
+			MultiplexedFrequencies);
 	}
 } // MainCS
diff --git a/Engine/Shaders/Private/ScreenSpaceDenoise/SSDPublic.ush b/Engine/Shaders/Private/ScreenSpaceDenoise/SSDPublic.ush
--- a/Engine/Shaders/Private/ScreenSpaceDenoise/SSDPublic.ush
+++ b/Engine/Shaders/Private/ScreenSpaceDenoise/SSDPublic.ush
@@ -18,7 +18,7 @@
 
 
 /** Random generator used for 1 sample per pixel for denoiser input. */
-float2 Rand1SPPDenoiserInput(uint2 PixelPos)
+float2 Rand1SPPDenoiserInput(uint2 PixelPos) // TODO(Denoiser): kill
 {
 	float2 E;
 
diff --git a/Engine/Shaders/Private/ScreenSpaceDenoise/SSDReduce.usf b/Engine/Shaders/Private/ScreenSpaceDenoise/SSDReduce.usf
deleted file mode 100644
diff --git a/Engine/Shaders/Private/ScreenSpaceDenoise/SSDSignalAccumulator.ush b/Engine/Shaders/Private/ScreenSpaceDenoise/SSDSignalAccumulator.ush
--- a/Engine/Shaders/Private/ScreenSpaceDenoise/SSDSignalAccumulator.ush
+++ b/Engine/Shaders/Private/ScreenSpaceDenoise/SSDSignalAccumulator.ush
@@ -22,8 +22,8 @@
 	#define COMPILE_MINMAX_ACCUMULATOR 0
 #endif
 
-#ifndef COMPILE_MININVFREQ_ACCUMULATOR
-	#define COMPILE_MININVFREQ_ACCUMULATOR 0
+#ifndef COMPILE_MIN_FREQUENCY_ACCUMULATOR
+	#define COMPILE_MIN_FREQUENCY_ACCUMULATOR 0
 #endif
 
 #ifndef COMPILE_DRB_ACCUMULATOR
@@ -63,8 +63,8 @@ struct FSSDSignalAccumulator
 	#endif
 
 	// Minimal inverse frequency found and intersected when sampling.
-	#if COMPILE_MININVFREQ_ACCUMULATOR
-		float MinInvFrequency;
+	#if COMPILE_MIN_FREQUENCY_ACCUMULATOR
+		FSSDSignalFrequency MinFrequency;
 	#endif
 
 	// Descending ring bucketing.
@@ -114,8 +114,9 @@ FSSDSignalAccumulator CreateSignalAccumulator()
 		Accumulator.Max = CreateSignalSampleFromScalarValue(0.0); // TODO(Denoiser): -INFINITE_FLOAT? otherwise there is a risk to color clamp with YCoCg.
 	#endif
 		
-	#if COMPILE_MININVFREQ_ACCUMULATOR
-		Accumulator.MinInvFrequency = WORLD_RADIUS_MISS;
+	#if COMPILE_MIN_FREQUENCY_ACCUMULATOR
+		Accumulator.MinFrequency.ClosestHitDistance = INFINITE_FLOAT;
+		Accumulator.MinFrequency.WorldBluringRadius = WORLD_RADIUS_MISS;
 	#endif
 
 	// DRB initialization.
@@ -144,6 +145,9 @@ struct FSSDSampleAccumulationInfos
 	// Sample to accumulate.
 	FSSDSignalSample Sample;
 
+	// Frequency of the sample.
+	FSSDSignalFrequency Frequency;
+
 	// Weight of the sample in the kernel.
 	float FinalWeight;
 
@@ -202,9 +206,11 @@ void AccumulateSampleInDomainBoundaries(inout FSSDSignalAccumulator Accumulator,
 		
 		// Track the minimal frequency inverse.
 		
-		#if COMPILE_MININVFREQ_ACCUMULATOR
-			if (SampleInfos.InvFrequency != WORLD_RADIUS_MISS)
-				Accumulator.MinInvFrequency = min(Accumulator.MinInvFrequency, SampleInfos.InvFrequency);
+		#if COMPILE_MIN_FREQUENCY_ACCUMULATOR
+		if (SampleInfos.InvFrequency != WORLD_RADIUS_INVALID)
+		{
+			Accumulator.MinFrequency = MinSignalFrequency(Accumulator.MinFrequency, SampleInfos.Frequency);
+		}
 		#endif
 	}
 }
diff --git a/Engine/Shaders/Private/ScreenSpaceDenoise/SSDSignalAccumulatorArray.ush b/Engine/Shaders/Private/ScreenSpaceDenoise/SSDSignalAccumulatorArray.ush
--- a/Engine/Shaders/Private/ScreenSpaceDenoise/SSDSignalAccumulatorArray.ush
+++ b/Engine/Shaders/Private/ScreenSpaceDenoise/SSDSignalAccumulatorArray.ush
@@ -4,7 +4,7 @@
 
 
 // Maximum number of VGPR that can be used to compress accumulator array. Currently bound by ACCUMULATOR_COMPRESSION_PENUMBRA_DRB.
-#define MAX_ACCUMULATOR_COMPRESSED_VGPRS 24
+#define MAX_ACCUMULATOR_COMPRESSED_VGPRS 20
 
 
 /** Structure that is just an array of accumulators, to work arround shader compiler bug on "out" of arrays. */
@@ -62,29 +62,28 @@ FSSDCompressedSignalAccumulatorArray CompressAccumulatorArray(in FSSDSignalAccum
 	{
 		// NOP
 	}
-	#if COMPILE_DRB_ACCUMULATOR && COMPILE_MININVFREQ_ACCUMULATOR
+	#if COMPILE_DRB_ACCUMULATOR && COMPILE_MIN_FREQUENCY_ACCUMULATOR && 0
 	else if (CompressionLayout == ACCUMULATOR_COMPRESSION_PENUMBRA_DRB)
 	{
+		#error Broken
 		UNROLL_N(MAX_SIGNAL_BATCH_SIZE)
 		for (uint SignalBatchId = 0; SignalBatchId < MAX_SIGNAL_BATCH_SIZE; SignalBatchId++)
 		{
 			const FSSDSignalAccumulator Accum = Accumulators.Array[SignalBatchId];
 
-			RawArray[SignalBatchId * 6 + 0].x = Accum.Current.SampleCount;
-			RawArray[SignalBatchId * 6 + 0].y = Accum.Current.MissCount;
-			RawArray[SignalBatchId * 6 + 1].x = Accum.Current.WorldBluringRadius;
-			RawArray[SignalBatchId * 6 + 1].y = Accum.CurrentInvFrequency;
-			RawArray[SignalBatchId * 6 + 2].x = Accum.CurrentTranslucency;
-			RawArray[SignalBatchId * 6 + 2].y = Accum.Previous.SampleCount;
-			RawArray[SignalBatchId * 6 + 3].x = Accum.Previous.MissCount;
-			RawArray[SignalBatchId * 6 + 3].y = Accum.Previous.WorldBluringRadius;
-			RawArray[SignalBatchId * 6 + 4].x = Accum.PreviousInvFrequency;
-			RawArray[SignalBatchId * 6 + 4].y = Accum.MinInvFrequency;
-			RawArray[SignalBatchId * 6 + 5].x = Accum.HighestInvFrequency;
-			RawArray[SignalBatchId * 6 + 5].y = Accum.BorderingRadius;
+			RawArray[SignalBatchId * 5 + 0].x = Accum.Current.SampleCount;
+			RawArray[SignalBatchId * 5 + 0].y = Accum.Current.MissCount;
+			RawArray[SignalBatchId * 5 + 1].x = Accum.CurrentInvFrequency;
+			RawArray[SignalBatchId * 5 + 1].y = Accum.CurrentTranslucency;
+			RawArray[SignalBatchId * 5 + 2].x = Accum.Previous.SampleCount;
+			RawArray[SignalBatchId * 5 + 2].y = Accum.Previous.MissCount;
+			RawArray[SignalBatchId * 5 + 3].x = Accum.PreviousInvFrequency;
+			RawArray[SignalBatchId * 5 + 3].y = Accum.MinFrequency.WorldBluringRadius;
+			RawArray[SignalBatchId * 5 + 4].x = Accum.HighestInvFrequency;
+			RawArray[SignalBatchId * 5 + 4].y = Accum.BorderingRadius;
 		} // for (uint SignalMultiplexId = 0; SignalMultiplexId < MAX_SIGNAL_BATCH_SIZE; SignalMultiplexId++)
 	}
-	#endif // COMPILE_DRB_ACCUMULATOR && COMPILE_MININVFREQ_ACCUMULATOR
+	#endif // COMPILE_DRB_ACCUMULATOR && COMPILE_MIN_FREQUENCY_ACCUMULATOR
 
 	// Compress.
 	FSSDCompressedSignalAccumulatorArray CompressedAccumulators;
@@ -116,7 +115,7 @@ FSSDSignalAccumulatorArray UncompressAccumulatorArray(in FSSDCompressedSignalAcc
 	{
 		// NOP to not clobber information in OutSample that remained uncompressed.
 	}
-	#if COMPILE_DRB_ACCUMULATOR && COMPILE_MININVFREQ_ACCUMULATOR
+	#if COMPILE_DRB_ACCUMULATOR && COMPILE_MIN_FREQUENCY_ACCUMULATOR && 0
 	else if (CompressionLayout == ACCUMULATOR_COMPRESSION_PENUMBRA_DRB)
 	{
 		// TODO(Denoiser): CONFIG_SIGNAL_BATCH_SIZE?
@@ -125,23 +124,21 @@ FSSDSignalAccumulatorArray UncompressAccumulatorArray(in FSSDCompressedSignalAcc
 		{
 			FSSDSignalAccumulator Accum = CreateSignalAccumulator();
 
-			Accum.Current.SampleCount         = RawArray[SignalBatchId * 6 + 0].x;
-			Accum.Current.MissCount           = RawArray[SignalBatchId * 6 + 0].y;
-			Accum.Current.WorldBluringRadius  = RawArray[SignalBatchId * 6 + 1].x;
-			Accum.CurrentInvFrequency         = RawArray[SignalBatchId * 6 + 1].y;
-			Accum.CurrentTranslucency         = RawArray[SignalBatchId * 6 + 2].x;
-			Accum.Previous.SampleCount        = RawArray[SignalBatchId * 6 + 2].y;
-			Accum.Previous.MissCount          = RawArray[SignalBatchId * 6 + 3].x;
-			Accum.Previous.WorldBluringRadius = RawArray[SignalBatchId * 6 + 3].y;
-			Accum.PreviousInvFrequency        = RawArray[SignalBatchId * 6 + 4].x;
-			Accum.MinInvFrequency             = RawArray[SignalBatchId * 6 + 4].y;
-			Accum.HighestInvFrequency         = RawArray[SignalBatchId * 6 + 5].x;
-			Accum.BorderingRadius             = RawArray[SignalBatchId * 6 + 5].y;
+			Accum.Current.SampleCount         = RawArray[SignalBatchId * 5 + 0].x;
+			Accum.Current.MissCount           = RawArray[SignalBatchId * 5 + 0].y;
+			Accum.CurrentInvFrequency         = RawArray[SignalBatchId * 5 + 1].x;
+			Accum.CurrentTranslucency         = RawArray[SignalBatchId * 5 + 1].y;
+			Accum.Previous.SampleCount        = RawArray[SignalBatchId * 5 + 2].x;
+			Accum.Previous.MissCount          = RawArray[SignalBatchId * 5 + 2].y;
+			Accum.PreviousInvFrequency        = RawArray[SignalBatchId * 5 + 3].x;
+			Accum.MinInvFrequency             = RawArray[SignalBatchId * 5 + 3].y;
+			Accum.HighestInvFrequency         = RawArray[SignalBatchId * 5 + 4].x;
+			Accum.BorderingRadius             = RawArray[SignalBatchId * 5 + 4].y;
 
 			Accumulators.Array[SignalBatchId] = Accum;
 		} // for (uint SignalBatchId = 0; SignalBatchId < MAX_SIGNAL_BATCH_SIZE; SignalBatchId++)
 	}
-	#endif // COMPILE_DRB_ACCUMULATOR && COMPILE_MININVFREQ_ACCUMULATOR
+	#endif // COMPILE_DRB_ACCUMULATOR && COMPILE_MIN_FREQUENCY_ACCUMULATOR
 
 	return Accumulators;
 } // UncompressAccumulatorArray()
diff --git a/Engine/Shaders/Private/ScreenSpaceDenoise/SSDSignalArray.ush b/Engine/Shaders/Private/ScreenSpaceDenoise/SSDSignalArray.ush
--- a/Engine/Shaders/Private/ScreenSpaceDenoise/SSDSignalArray.ush
+++ b/Engine/Shaders/Private/ScreenSpaceDenoise/SSDSignalArray.ush
@@ -147,3 +147,14 @@ FSSDSignalArray WaveBroadcastSignalArray(const FWaveBroadcastSettings BroadcastS
 	}
 	return OutSamples;
 }
+
+FSSDSignalFrequencyArray WaveBroadcastSignalFrequenciesArray(const FWaveBroadcastSettings BroadcastSettings, FSSDSignalFrequencyArray Frequencies)
+{
+	FSSDSignalFrequencyArray OutFrequencies;
+	UNROLL_N(SIGNAL_ARRAY_SIZE)
+	for (uint BatchedSignalId = 0; BatchedSignalId < SIGNAL_ARRAY_SIZE; BatchedSignalId++)
+	{
+		OutFrequencies.Array[BatchedSignalId] = WaveBroadcastSignalFrequency(BroadcastSettings, Frequencies.Array[BatchedSignalId]);
+	}
+	return OutFrequencies;
+}
diff --git a/Engine/Shaders/Private/ScreenSpaceDenoise/SSDSignalBufferEncoding.ush b/Engine/Shaders/Private/ScreenSpaceDenoise/SSDSignalBufferEncoding.ush
--- a/Engine/Shaders/Private/ScreenSpaceDenoise/SSDSignalBufferEncoding.ush
+++ b/Engine/Shaders/Private/ScreenSpaceDenoise/SSDSignalBufferEncoding.ush
@@ -83,9 +83,11 @@ void DecodeMultiplexedSignalsFromFloat4(
 	const uint MultiplexedSampleId,
 	const bool bNormalizeSample,
 	float4 RawSample[MAX_MULTIPLEXED_TEXTURES],
-	out FSSDSignalArray OutSamples)
+	out FSSDSignalArray OutSamples,
+	out FSSDSignalFrequencyArray OutSampleFrequencies)
 {
 	OutSamples = CreateSignalArrayFromScalarValue(0.0);
+	OutSampleFrequencies = CreateInvalidSignalFrequencyArray();
 
 	if (0)
 	{
@@ -106,8 +108,9 @@ void DecodeMultiplexedSignalsFromFloat4(
 			// TODO(Denoiser): feed the actual number of sample.
 			OutSamples.Array[MultiplexId].SampleCount = (RawSample[MultiplexId].g == DENOISER_INVALID_HIT_DISTANCE ? 0.0 : 1.0);
 			OutSamples.Array[MultiplexId].MissCount = OutSamples.Array[MultiplexId].SampleCount * RawSample[MultiplexId].r;
-			OutSamples.Array[MultiplexId].ClosestHitDistance = OutSamples.Array[MultiplexId].SampleCount * (RawSample[MultiplexId].g == DENOISER_MISS_HIT_DISTANCE ? 1000000 : RawSample[MultiplexId].g);
 			OutSamples.Array[MultiplexId].TransmissionDistance = OutSamples.Array[MultiplexId].SampleCount * RawSample[MultiplexId].a;
+
+			OutSampleFrequencies.Array[MultiplexId].ClosestHitDistance = RawSample[MultiplexId].g;
 		}
 	}
 	else if (SignalBufferLayout == SIGNAL_BUFFER_LAYOUT_PENUMBRA_INJESTION_NSPP)
@@ -121,8 +124,10 @@ void DecodeMultiplexedSignalsFromFloat4(
 			// TODO(Denoiser): feed the actual number of sample.
 			OutSamples.Array[MultiplexId].SampleCount = (Channels.g == -2.0 ? 0.0 : 1.0);
 			OutSamples.Array[MultiplexId].MissCount = OutSamples.Array[MultiplexId].SampleCount * Channels.r;
-			OutSamples.Array[MultiplexId].WorldBluringRadius = OutSamples.Array[MultiplexId].SampleCount * (Channels.g == -1 ? WORLD_RADIUS_MISS : Channels.g);
 			OutSamples.Array[MultiplexId].TransmissionDistance = OutSamples.Array[MultiplexId].SampleCount * Channels.a;
+
+			OutSampleFrequencies.Array[MultiplexId].WorldBluringRadius = (Channels.g == -1 ? WORLD_RADIUS_MISS : Channels.g);
+			OutSampleFrequencies.Array[MultiplexId].WorldBluringRadius = Channels.g;
 		}
 	}
 	//else if (SignalBufferLayout == SIGNAL_BUFFER_LAYOUT_PENUMBRA_RECONSTRUCTION)
@@ -155,8 +160,10 @@ void DecodeMultiplexedSignalsFromFloat4(
 
 			OutSamples.Array[MultiplexId].MissCount = Channels.r * SampleCount;
 			OutSamples.Array[MultiplexId].SampleCount = SampleCount;
-			OutSamples.Array[MultiplexId].WorldBluringRadius = (Channels.b == DENOISER_MISS_HIT_DISTANCE ? WORLD_RADIUS_MISS : Channels.b) * SampleCount;
 			OutSamples.Array[MultiplexId].TransmissionDistance = Channels.a * SampleCount;
+
+			OutSampleFrequencies.Array[MultiplexId].WorldBluringRadius = (Channels.b == DENOISER_MISS_HIT_DISTANCE ? WORLD_RADIUS_MISS : Channels.b);
+			OutSampleFrequencies.Array[MultiplexId].WorldBluringRadius = Channels.b;
 		}
 	}
 	else if (SignalBufferLayout == SIGNAL_BUFFER_LAYOUT_PENUMBRA_REJECTION)
@@ -170,8 +177,10 @@ void DecodeMultiplexedSignalsFromFloat4(
 
 			OutSamples.Array[MultiplexId].SampleCount = SampleCount;
 			OutSamples.Array[MultiplexId].MissCount = Channels.r * SampleCount;
-			OutSamples.Array[MultiplexId].WorldBluringRadius = (Channels.b == -1 ? WORLD_RADIUS_MISS : Channels.b * SampleCount);
 			OutSamples.Array[MultiplexId].TransmissionDistance = Channels.a * SampleCount;
+
+			OutSampleFrequencies.Array[MultiplexId].WorldBluringRadius = (Channels.b == -1 ? WORLD_RADIUS_MISS : Channels.b);
+			OutSampleFrequencies.Array[MultiplexId].WorldBluringRadius = Channels.b;
 		}
 	}
 	#endif
@@ -279,8 +288,8 @@ void DecodeMultiplexedSignalsFromFloat4(
 		OutSamples.Array[0].SceneColor = RawSample[0];
 
 		// max() because ray hit distance is < 0 when miss / not ray has been shot.
-		//OutSamples.Array[0].ClosestHitDistance = RawSample[1].r * OutSamples.Array[0].SampleCount;
-		OutSamples.Array[0].ClosestHitDistance = ((RawSample[1].r == DENOISER_MISS_HIT_DISTANCE) ? 1.0e20 : RawSample[1].r) * OutSamples.Array[0].SampleCount;
+		OutSampleFrequencies.Array[0].ClosestHitDistance = RawSample[1].r;
+		//OutSampleFrequencies.Array[0].ClosestHitDistance = ((RawSample[1].r == DENOISER_MISS_HIT_DISTANCE) ? 1.0e20 : RawSample[1].r);
 	}
 	else if (SignalBufferLayout == SIGNAL_BUFFER_LAYOUT_REFLECTIONS_REJECTION)
 	{
@@ -328,7 +337,8 @@ void DecodeMultiplexedSignalsFromFloat4(
 	{
 		OutSamples.Array[0].SampleCount = RawSample[1].r < 0 ? 0 : 1;
 		OutSamples.Array[0].MissCount = OutSamples.Array[0].SampleCount * RawSample[0].r;
-		OutSamples.Array[0].ClosestHitDistance = OutSamples.Array[0].SampleCount * RawSample[1].r;
+
+		OutSampleFrequencies.Array[0].ClosestHitDistance = RawSample[1].r;
 	}
 	else if (SignalBufferLayout == SIGNAL_BUFFER_LAYOUT_AO_REJECTION)
 	{
@@ -370,9 +380,7 @@ void DecodeMultiplexedSignalsFromFloat4(
 		OutSamples.Array[0].MissCount = OutSamples.Array[0].SampleCount * (1 - RawSample[0].a);
 		OutSamples.Array[0].SceneColor = OutSamples.Array[0].SampleCount * RawSample[0];
 
-		// max() because ray hit distance is < 0 when miss / not ray has been shot.
-		OutSamples.Array[0].ClosestHitDistance = RawSample[1].r * OutSamples.Array[0].SampleCount;
-		OutSamples.Array[0].ClosestHitDistance = ((RawSample[1].r == -1.0) ? 1.0e20 : RawSample[1].r) * OutSamples.Array[0].SampleCount;
+		OutSampleFrequencies.Array[0].ClosestHitDistance = ((RawSample[1].r == -1.0) ? 1.0e20 : RawSample[1].r);
 	}
 	else if (SignalBufferLayout == SIGNAL_BUFFER_LAYOUT_DIFFUSE_INDIRECT_AND_AO_RECONSTRUCTION)
 	{
@@ -430,17 +438,15 @@ void DecodeMultiplexedSignalsFromUint2(
 	const uint MultiplexedSampleId,
 	const bool bNormalizeSample,
 	uint2 RawSample[MAX_MULTIPLEXED_TEXTURES],
-	out FSSDSignalArray OutSamples)
+	out FSSDSignalArray OutSamples,
+	out FSSDSignalFrequencyArray OutSampleFrequencies)
 {
 	OutSamples = CreateSignalArrayFromScalarValue(0.0);
+	OutSampleFrequencies = CreateInvalidSignalFrequencyArray();
 
 	if (0)
 	{
 
-	}
-	else if (SignalBufferLayout == SIGNAL_BUFFER_LAYOUT_PENUMBRA_INPUT_NSPP)
-	{
-		// Shouldn't happen..
 	}
 	else if (SignalBufferLayout == SIGNAL_BUFFER_LAYOUT_PENUMBRA_INJESTION_NSPP)
 	{
@@ -457,13 +463,14 @@ void DecodeMultiplexedSignalsFromUint2(
 
 			float MissCount = SampleCount * MissCountRatio;
 			float TransmissionDistance = TransmissionDistanceRatio * 5.0;
-			WorldBluringRadius = SampleCount * (WorldBluringRadius == -1.0 ? WORLD_RADIUS_MISS : WorldBluringRadius);
 
 			// TODO(Denoiser): feed the actual number of sample.
 			OutSamples.Array[MultiplexId].SampleCount = SampleCount;
 			OutSamples.Array[MultiplexId].MissCount = MissCount;
-			OutSamples.Array[MultiplexId].WorldBluringRadius = WorldBluringRadius;
 			OutSamples.Array[MultiplexId].TransmissionDistance = TransmissionDistance;
+
+			OutSampleFrequencies.Array[MultiplexId].WorldBluringRadius = (WorldBluringRadius == -1.0 ? WORLD_RADIUS_MISS : WorldBluringRadius);
+			OutSampleFrequencies.Array[MultiplexId].WorldBluringRadius = WorldBluringRadius;
 		}
 	}
 	else if (SignalBufferLayout == SIGNAL_BUFFER_LAYOUT_PENUMBRA_HISTORY || SignalBufferLayout == SIGNAL_BUFFER_LAYOUT_PENUMBRA_RECONSTRUCTION)
@@ -480,8 +487,10 @@ void DecodeMultiplexedSignalsFromUint2(
 
 			OutSamples.Array[MultiplexId].MissCount = MissCount;
 			OutSamples.Array[MultiplexId].SampleCount = SampleCount;
-			OutSamples.Array[MultiplexId].WorldBluringRadius = (WorldBluringRadius == -1.0 ? WORLD_RADIUS_MISS : WorldBluringRadius);
 			OutSamples.Array[MultiplexId].TransmissionDistance = TransmissionDistance;
+
+			OutSampleFrequencies.Array[MultiplexId].WorldBluringRadius = (WorldBluringRadius == -1.0 ? WORLD_RADIUS_MISS : WorldBluringRadius);
+			OutSampleFrequencies.Array[MultiplexId].WorldBluringRadius = WorldBluringRadius;
 		}
 	}
 	else if (SignalBufferLayout == SIGNAL_BUFFER_LAYOUT_PENUMBRA_REJECTION)
@@ -498,71 +507,44 @@ void DecodeMultiplexedSignalsFromUint2(
 
 			OutSamples.Array[MultiplexId + 0].SampleCount = SampleCount;
 			OutSamples.Array[MultiplexId + 0].MissCount = MissCount0;
-			OutSamples.Array[MultiplexId + 0].WorldBluringRadius = (WorldBluringRadius == -1.0 ? WORLD_RADIUS_MISS : WorldBluringRadius);
 
 			OutSamples.Array[MultiplexId + 1].SampleCount = SampleCount;
 			OutSamples.Array[MultiplexId + 1].MissCount = MissCount1;
+
+			OutSampleFrequencies.Array[MultiplexId + 0].WorldBluringRadius = (WorldBluringRadius == -1.0 ? WORLD_RADIUS_MISS : WorldBluringRadius);
+			OutSampleFrequencies.Array[MultiplexId].WorldBluringRadius = WorldBluringRadius;
 		}
 	}
 }
 
 /** Decode input signal sample from raw float. */
-FSSDSignalArray DecodeMultiplexedSignals(
+void DecodeMultiplexedSignals(
 	const uint SignalBufferLayout,
 	const uint MultiplexedSampleId,
 	const bool bNormalizeSample,
-	FSSDCompressedMultiplexedSample CompressedSample)
+	FSSDCompressedMultiplexedSample CompressedSample,
+	out FSSDSignalArray OutSamples,
+	out FSSDSignalFrequencyArray OutSampleFrequencies)
 #if CONFIG_SIGNAL_INPUT_TEXTURE_TYPE == SIGNAL_TEXTURE_TYPE_FLOAT4
 {
-	FSSDSignalArray MultiplexedSamples;
 	DecodeMultiplexedSignalsFromFloat4(
-		SignalBufferLayout, MultiplexedSampleId, bNormalizeSample,
-		CompressedSample.VGPRArray, /* out */ MultiplexedSamples);
-	return MultiplexedSamples;
+		SignalBufferLayout,
+		MultiplexedSampleId,
+		bNormalizeSample,
+		CompressedSample.VGPRArray,
+		/* out */ OutSamples,
+		/* out */ OutSampleFrequencies);
 } // DecodeMultiplexedSignals()
 #elif CONFIG_SIGNAL_INPUT_TEXTURE_TYPE == SIGNAL_TEXTURE_TYPE_UINT2
 {
 	FSSDSignalArray MultiplexedSamples = CreateSignalArrayFromScalarValue(0.0);
 	DecodeMultiplexedSignalsFromUint2(
-		SignalBufferLayout, MultiplexedSampleId, bNormalizeSample,
-		CompressedSample.VGPRArray, /* out */ MultiplexedSamples);
-
-	if (0)
-	{
-	
-	}
-	#if COMPILE_SIGNAL_COLOR_SH
-	else if (SignalBufferLayout == SIGNAL_BUFFER_LAYOUT_DIFFUSE_INDIRECT_HARMONIC)
-	{
-		const uint DecodeOptions = bNormalizeSample ? SSD_DECODE_NORMALIZE : 0x0;
-
-		DecodeDiffuseSphericalHarmonicTexel(
-			CompressedSample.VGPRArray,
-			DecodeOptions,
-			/* out */ MultiplexedSamples.Array[0].SampleCount,
-			/* out */ MultiplexedSamples.Array[0].MissCount,
-			/* out */ MultiplexedSamples.Array[0].ColorSH);
-
-		// TODO(Denoiser): SRV on RG F32 and manually decode with f16tof32(), going to reduce VGPR pressure between fetch and result.
-	}
-	#endif
-#if 0
-	// Uncompress assuming this is packed as 16bit floats.
-	else
-	{
-		float4 FloatVGPRArray[MAX_MULTIPLEXED_TEXTURES];
-		
-		UNROLL_N(MAX_MULTIPLEXED_TEXTURES)
-		for (uint TexId = 0; TexId < MAX_MULTIPLEXED_TEXTURES; TexId++)
-			FloatVGPRArray[TexId] = Uncompress4Floats(CompressedSample.VGPRArray[TexId]);
-		
-		DecodeMultiplexedSignalsFromFloat4(
-			SignalBufferLayout, MultiplexedSampleId, bNormalizeSample,
-			FloatVGPRArray, /* out */ MultiplexedSamples);
-	}
-#endif
-
-	return MultiplexedSamples;
+		SignalBufferLayout,
+		MultiplexedSampleId,
+		bNormalizeSample,
+		CompressedSample.VGPRArray,
+		/* out */ OutSamples,
+		/* out */ OutSampleFrequencies);
 } // DecodeMultiplexedSignals()
 #endif
 
@@ -570,6 +552,7 @@ FSSDSignalArray DecodeMultiplexedSignals(
 void EncodeMultiplexedSignals(
 	const uint SignalBufferLayout, const uint MultiplexCount,
 	FSSDSignalSample Sample[SIGNAL_ARRAY_SIZE],
+	FSSDSignalFrequency SampleFrequency[SIGNAL_ARRAY_SIZE],
 	out FSSDOutputRawSample OutRawSample[MAX_MULTIPLEXED_TEXTURES],
 	out uint OutBufferCount)
 #if CONFIG_SIGNAL_OUTPUT_TEXTURE_TYPE == SIGNAL_TEXTURE_TYPE_FLOAT4
@@ -592,21 +575,22 @@ void EncodeMultiplexedSignals(
 		UNROLL
 		for (uint MultiplexId = 0; MultiplexId < MultiplexCount; MultiplexId++)
 		{
+			float NormalizationFactor = SafeRcp(Sample[MultiplexId].SampleCount);
+
 			float4 EncodedChannels = float4(
 				Sample[MultiplexId].MissCount,
-				Sample[MultiplexId].WorldBluringRadius,
+				SampleFrequency[MultiplexId].WorldBluringRadius,
 				0.0,
-				Sample[MultiplexId].TransmissionDistance);
-			EncodedChannels *= SafeRcp(Sample[MultiplexId].SampleCount);
+				Sample[MultiplexId].TransmissionDistance * NormalizationFactor);
 			
-			if (Sample[MultiplexId].SampleCount == 0)
-			{
-				EncodedChannels.y = -2;
-			}
-			else if (Sample[MultiplexId].WorldBluringRadius == WORLD_RADIUS_MISS)
-			{
-				EncodedChannels.y = -1;
-			}
+			//if (Sample[MultiplexId].SampleCount == 0)
+			//{
+			//	EncodedChannels.y = -2;
+			//}
+			//else if (SampleFrequency[MultiplexId].WorldBluringRadius == WORLD_RADIUS_MISS)
+			//{
+			//	EncodedChannels.y = -1;
+			//}
 
 			OutRawSample[MultiplexId] = EncodedChannels;
 		}
@@ -635,13 +619,12 @@ void EncodeMultiplexedSignals(
 		{
 			float NormalizationFactor = SafeRcp(Sample[MultiplexId].SampleCount);
 
-			float NormalizedWorldBluringRadius = Sample[MultiplexId].WorldBluringRadius * NormalizationFactor;
 			float NormalizedTransmissionDistance = Sample[MultiplexId].TransmissionDistance * NormalizationFactor;
 
 			OutRawSample[MultiplexId] = float4(
 				Sample[MultiplexId].MissCount * NormalizationFactor,
 				Sample[MultiplexId].SampleCount,
-				NormalizedWorldBluringRadius == WORLD_RADIUS_MISS ? -1 : NormalizedWorldBluringRadius,
+				SampleFrequency[MultiplexId].WorldBluringRadius,
 				NormalizedTransmissionDistance);
 		}
 		OutBufferCount = MultiplexCount;
@@ -656,16 +639,12 @@ void EncodeMultiplexedSignals(
 		{
 			uint MultiplexId = BatchSignalId;
 
-			float NormalizationFactor = SafeRcp(Sample[MultiplexId + 0].SampleCount);
-
-			float NormalizedWorldBluringRadius = Sample[MultiplexId + 0].WorldBluringRadius * NormalizationFactor;
+			float NormalizationFactor = SafeRcp(Sample[MultiplexId].SampleCount);
 
-			// Samples are normalised when doing history preconvolution, so the number of is either 0 or 1 on both momments. Therefore
-			// Sample[MultiplexId + 0].SampleCount == Sample[MultiplexId + 1].SampleCount.
 			OutRawSample[BatchSignalId] = float4(
 				Sample[MultiplexId].MissCount * NormalizationFactor,
 				Sample[MultiplexId].SampleCount,
-				NormalizedWorldBluringRadius == WORLD_RADIUS_MISS ? -1 : NormalizedWorldBluringRadius,
+				SampleFrequency[MultiplexId].WorldBluringRadius == WORLD_RADIUS_MISS ? -1 : SampleFrequency[MultiplexId].WorldBluringRadius,
 				Sample[MultiplexId].TransmissionDistance * NormalizationFactor);
 		}
 	}
@@ -835,7 +814,7 @@ void EncodeMultiplexedSignals(
 		for (uint MultiplexId = 0; MultiplexId < MultiplexCount; MultiplexId++)
 		{
 			float MissCount = Sample[MultiplexId].MissCount;
-			float WorldBluringRadius = Sample[MultiplexId].WorldBluringRadius;
+			float WorldBluringRadius = SampleFrequency[MultiplexId].WorldBluringRadius;
 			float TransmissionDistance = Sample[MultiplexId].TransmissionDistance;
 			float SampleCount = Sample[MultiplexId].SampleCount;
 			if (SampleCount == 0)
@@ -874,7 +853,7 @@ void EncodeMultiplexedSignals(
 			float NormalizationFactor = rcp(SampleCount);
 
 			float MissCount = Sample[MultiplexId].MissCount;
-			float WorldBluringRadius = (Sample[MultiplexId].WorldBluringRadius == WORLD_RADIUS_MISS ? -1.0 : Sample[MultiplexId].WorldBluringRadius);
+			float WorldBluringRadius = (SampleFrequency[MultiplexId].WorldBluringRadius == WORLD_RADIUS_MISS ? -1.0 : SampleFrequency[MultiplexId].WorldBluringRadius);
 			float TransmissionDistance = Sample[MultiplexId].TransmissionDistance;
 
 			uint2 EncodedData = 0;
@@ -902,7 +881,7 @@ void EncodeMultiplexedSignals(
 
 			float MissCount0 = Sample[MultiplexId + 0].MissCount;
 			float MissCount1 = Sample[MultiplexId + 1].MissCount;
-			float WorldBluringRadius = (Sample[MultiplexId + 0].WorldBluringRadius == WORLD_RADIUS_MISS ? -1.0 : Sample[MultiplexId + 0].WorldBluringRadius);
+			float WorldBluringRadius = (SampleFrequency[MultiplexId + 0].WorldBluringRadius == WORLD_RADIUS_MISS ? -1.0 : SampleFrequency[MultiplexId + 0].WorldBluringRadius);
 			float TransmissionDistance = 0.0;
 
 			uint2 EncodedData = 0;
@@ -938,7 +917,7 @@ void EncodeMultiplexedSignals(
 /** Sample the raw of multiple input signals that have been multiplexed. */
 FSSDCompressedMultiplexedSample SampleCompressedMultiplexedSignals(
 	FSSDTexture2D SignalBuffer0, FSSDTexture2D SignalBuffer1, FSSDTexture2D SignalBuffer2, FSSDTexture2D SignalBuffer3,
-	SamplerState Sampler, float2 UV, float Level, uint2 PixelCoord)
+	SamplerState Sampler, float2 UV, uint2 PixelCoord)
 {
 	FSSDCompressedMultiplexedSample CompressedSample;
 
@@ -948,15 +927,15 @@ FSSDCompressedMultiplexedSample SampleCompressedMultiplexedSignals(
 	{
 		#if CONFIG_SIGNAL_INPUT_TEXTURE_TYPE == SIGNAL_TEXTURE_TYPE_FLOAT4
 		{
-			CompressedSample.VGPRArray[0] = SignalBuffer0.SampleLevel(Sampler, UV, Level);
-			CompressedSample.VGPRArray[1] = SignalBuffer1.SampleLevel(Sampler, UV, Level);
-			CompressedSample.VGPRArray[2] = SignalBuffer2.SampleLevel(Sampler, UV, Level);
-			CompressedSample.VGPRArray[3] = SignalBuffer3.SampleLevel(Sampler, UV, Level);
+			CompressedSample.VGPRArray[0] = SignalBuffer0.SampleLevel(Sampler, UV, 0.0);
+			CompressedSample.VGPRArray[1] = SignalBuffer1.SampleLevel(Sampler, UV, 0.0);
+			CompressedSample.VGPRArray[2] = SignalBuffer2.SampleLevel(Sampler, UV, 0.0);
+			CompressedSample.VGPRArray[3] = SignalBuffer3.SampleLevel(Sampler, UV, 0.0);
 		}
 		#elif CONFIG_SIGNAL_INPUT_TEXTURE_TYPE >= SIGNAL_TEXTURE_TYPE_UINT1 && CONFIG_SIGNAL_INPUT_TEXTURE_TYPE <= SIGNAL_TEXTURE_TYPE_UINT4
 		{
 			// TODO(Denoiser): Exposed the int3 instead as function parameter.
-			int3 Coord = int3(PixelCoord >> uint(Level), int(Level));
+			int3 Coord = int3(PixelCoord, 0);
 			CompressedSample.VGPRArray[0] = SignalBuffer0.Load(Coord);
 			CompressedSample.VGPRArray[1] = SignalBuffer1.Load(Coord);
 			CompressedSample.VGPRArray[2] = SignalBuffer2.Load(Coord);
@@ -972,21 +951,25 @@ FSSDCompressedMultiplexedSample SampleCompressedMultiplexedSignals(
 
 
 /** Sample multiple input signals that have been multiplexed. */
-FSSDSignalArray SampleMultiplexedSignals(
+void SampleMultiplexedSignals( // TODO
 	FSSDTexture2D SignalBuffer0, FSSDTexture2D SignalBuffer1, FSSDTexture2D SignalBuffer2, FSSDTexture2D SignalBuffer3,
 	SamplerState Sampler,
 	const uint SignalBufferLayout, const uint MultiplexedSampleId,
 	const bool bNormalizeSample,
-	float2 UV, float Level = 0)
+	float2 UV,
+	out FSSDSignalArray OutMultiplexedSamples,
+	out FSSDSignalFrequencyArray OutMultiplexedFrequencies)
 {
 	uint2 PixelCoord = BufferUVToBufferPixelCoord(UV);
 
 	FSSDCompressedMultiplexedSample CompressedSample = SampleCompressedMultiplexedSignals(
 		SignalBuffer0, SignalBuffer1, SignalBuffer2, SignalBuffer3,
-		Sampler, UV, Level, PixelCoord);
+		Sampler, UV, PixelCoord);
 
-	return DecodeMultiplexedSignals(
-		SignalBufferLayout, MultiplexedSampleId, bNormalizeSample, CompressedSample);
+	DecodeMultiplexedSignals(
+		SignalBufferLayout, MultiplexedSampleId, bNormalizeSample, CompressedSample,
+		/* out */ OutMultiplexedSamples,
+		/* out */ OutMultiplexedFrequencies);
 }
 
 
@@ -997,14 +980,17 @@ void OutputMultiplexedSignal(
 	FSSDRWTexture2D OutputSignalBuffer2,
 	FSSDRWTexture2D OutputSignalBuffer3,
 	const uint SignalBufferLayout, const uint MultiplexCount,
-	const uint2 PixelPosition, FSSDSignalSample MultiplexedSamples[SIGNAL_ARRAY_SIZE])
+	const uint2 PixelPosition,
+	FSSDSignalArray MultiplexedSamples,
+	FSSDSignalFrequencyArray MultiplexedFrequencies)
 {
 	// Encode the output signal.
 	FSSDOutputRawSample RawSample[MAX_MULTIPLEXED_TEXTURES];
 	uint BufferCount;
 	EncodeMultiplexedSignals(
 		SignalBufferLayout, MultiplexCount,
-		MultiplexedSamples,
+		MultiplexedSamples.Array,
+		MultiplexedFrequencies.Array,
 		/* out */ RawSample, /* out */ BufferCount);
 
 	// Output the raw encoded sample according to number of RT they requires.
diff --git a/Engine/Shaders/Private/ScreenSpaceDenoise/SSDSignalCore.ush b/Engine/Shaders/Private/ScreenSpaceDenoise/SSDSignalCore.ush
--- a/Engine/Shaders/Private/ScreenSpaceDenoise/SSDSignalCore.ush
+++ b/Engine/Shaders/Private/ScreenSpaceDenoise/SSDSignalCore.ush
@@ -30,7 +30,7 @@
 
 //------------------------------------------------------- STRUCTURE
 
-/** Generic data structure to manipulate any kind of signal. */
+/** Generic linearily filterable data structure to manipulate any kind of signal. */
 struct FSSDSignalSample
 {
 	// Number of sample accumulated.
@@ -54,13 +54,17 @@ struct FSSDSignalSample
 	// Number of ray that did not find any intersections.
 	float MissCount;
 
-	// Hit distance of a this sample, valid only if SampleCount == 1.
+	float TransmissionDistance;
+};
+
+/** Generic data structure about the frequency of a FSSDSignalSample. */
+struct FSSDSignalFrequency
+{
+	// Hit distance of a this sample.
 	float ClosestHitDistance;
-	
-	// Blur radius allowed for this sample, valid only if SampleCount == 1.
-	float WorldBluringRadius;
 
-	float TransmissionDistance;
+	// Blur radius allowed for this sample.
+	float WorldBluringRadius;
 };
 
 /** Array of signal samples. Technically should be a typedef, but there is a bug in HLSL with the out operator of an array. */
@@ -69,6 +73,12 @@ struct FSSDSignalArray
 	FSSDSignalSample Array[SIGNAL_ARRAY_SIZE];
 };
 
+/** Array of signal frequencies. Technically should be a typedef, but there is a bug in HLSL with the out operator of an array. */
+struct FSSDSignalFrequencyArray
+{
+	FSSDSignalFrequency Array[SIGNAL_ARRAY_SIZE];
+};
+
 
 FSSDSignalSample CreateSignalSampleFromScalarValue(float Scalar)
 {
@@ -100,12 +110,18 @@ FSSDSignalSample CreateSignalSampleFromScalarValue(float Scalar)
 	#endif
 	OutSample.SampleCount = Scalar;
 	OutSample.MissCount = Scalar;
-	OutSample.ClosestHitDistance = Scalar;
-	OutSample.WorldBluringRadius = Scalar;
 	OutSample.TransmissionDistance = Scalar;
 	return OutSample;
 }
 
+FSSDSignalFrequency CreateInvalidSignalFrequency()
+{
+	FSSDSignalFrequency Frequency;
+	Frequency.ClosestHitDistance = DENOISER_INVALID_HIT_DISTANCE;
+	Frequency.WorldBluringRadius = WORLD_RADIUS_INVALID;
+	return Frequency;
+}
+
 FSSDSignalSample CreateSignalSample()
 {
 	return CreateSignalSampleFromScalarValue(0.0);
@@ -122,6 +138,17 @@ FSSDSignalArray CreateSignalArrayFromScalarValue(float Scalar)
 	return OutSamples;
 }
 
+FSSDSignalFrequencyArray CreateInvalidSignalFrequencyArray()
+{
+	FSSDSignalFrequencyArray OutFrequencies;
+	UNROLL_N(SIGNAL_ARRAY_SIZE)
+	for (uint BatchedSignalId = 0; BatchedSignalId < SIGNAL_ARRAY_SIZE; BatchedSignalId++)
+	{
+		OutFrequencies.Array[BatchedSignalId] = CreateInvalidSignalFrequency();
+	}
+	return OutFrequencies;
+}
+
 
 //------------------------------------------------------- SIGNAL OPERATORS
 
@@ -143,8 +170,6 @@ FSSDSignalSample MulSignal(FSSDSignalSample Sample, float Scalar)
 	#endif
 	OutSample.SampleCount            = Sample.SampleCount * Scalar;
 	OutSample.MissCount              = Sample.MissCount * Scalar;
-	OutSample.ClosestHitDistance     = Sample.ClosestHitDistance * Scalar;
-	OutSample.WorldBluringRadius     = Sample.WorldBluringRadius * Scalar;
 	OutSample.TransmissionDistance   = Sample.TransmissionDistance * Scalar;
 	return OutSample;
 }
@@ -166,8 +191,6 @@ FSSDSignalSample AddSignal(FSSDSignalSample SampleA, FSSDSignalSample SampleB)
 	#endif
 	SampleA.SampleCount            += SampleB.SampleCount;
 	SampleA.MissCount              += SampleB.MissCount;
-	SampleA.ClosestHitDistance     += SampleB.ClosestHitDistance;
-	SampleA.WorldBluringRadius     += SampleB.WorldBluringRadius;
 	SampleA.TransmissionDistance   += SampleB.TransmissionDistance;
 	return SampleA;
 }
@@ -189,8 +212,6 @@ FSSDSignalSample MinusSignal(FSSDSignalSample Sample)
 	#endif
 	Sample.SampleCount            = -Sample.SampleCount;
 	Sample.MissCount              = -Sample.MissCount;
-	Sample.ClosestHitDistance     = -Sample.ClosestHitDistance;
-	Sample.WorldBluringRadius     = -Sample.WorldBluringRadius;
 	Sample.TransmissionDistance   = -Sample.TransmissionDistance;
 	return Sample;
 }
@@ -229,8 +250,6 @@ FSSDSignalSample AbsSignal(FSSDSignalSample Sample)
 	#endif
 	Sample.SampleCount            = abs(Sample.SampleCount);
 	Sample.MissCount              = abs(Sample.MissCount);
-	Sample.ClosestHitDistance     = abs(Sample.ClosestHitDistance);
-	Sample.WorldBluringRadius     = abs(Sample.WorldBluringRadius);
 	Sample.TransmissionDistance = abs(Sample.TransmissionDistance);
 	return Sample;
 }
@@ -264,8 +283,6 @@ FSSDSignalSample SqrtSignal(FSSDSignalSample Sample)
 	#endif
 	Sample.SampleCount            = sqrt(Sample.SampleCount);
 	Sample.MissCount              = sqrt(Sample.MissCount);
-	Sample.ClosestHitDistance     = sqrt(Sample.ClosestHitDistance);
-	Sample.WorldBluringRadius     = sqrt(Sample.WorldBluringRadius);
 	Sample.TransmissionDistance   = sqrt(Sample.TransmissionDistance);
 	return Sample;
 }
@@ -302,8 +319,6 @@ FSSDSignalSample PowerSignal(FSSDSignalSample Sample, float Exponent)
 	#endif
 	Sample.SampleCount            = pow(Sample.SampleCount, Exponent);
 	Sample.MissCount              = pow(Sample.MissCount, Exponent);
-	Sample.ClosestHitDistance     = pow(Sample.ClosestHitDistance, Exponent);
-	Sample.WorldBluringRadius     = pow(Sample.WorldBluringRadius, Exponent);
 	Sample.TransmissionDistance   = pow(Sample.TransmissionDistance, Exponent);
 	return Sample;
 }
@@ -338,8 +353,6 @@ FSSDSignalSample MinSignal(FSSDSignalSample SampleA, FSSDSignalSample SampleB)
 	#endif
 	OutSample.SampleCount = min(SampleA.SampleCount, SampleB.SampleCount);
 	OutSample.MissCount = min(SampleA.MissCount, SampleB.MissCount);
-	OutSample.ClosestHitDistance = min(SampleA.ClosestHitDistance, SampleB.ClosestHitDistance);
-	OutSample.WorldBluringRadius = min(SampleA.WorldBluringRadius, SampleB.WorldBluringRadius);
 	OutSample.TransmissionDistance = min(SampleA.TransmissionDistance, SampleB.TransmissionDistance);
 	return OutSample;
 }
@@ -374,8 +387,6 @@ FSSDSignalSample MaxSignal(FSSDSignalSample SampleA, FSSDSignalSample SampleB)
 	#endif
 	OutSample.SampleCount = max(SampleA.SampleCount, SampleB.SampleCount);
 	OutSample.MissCount = max(SampleA.MissCount, SampleB.MissCount);
-	OutSample.ClosestHitDistance = max(SampleA.ClosestHitDistance, SampleB.ClosestHitDistance);
-	OutSample.WorldBluringRadius = max(SampleA.WorldBluringRadius, SampleB.WorldBluringRadius);
 	OutSample.TransmissionDistance = max(SampleA.TransmissionDistance, SampleB.TransmissionDistance);
 	return OutSample;
 }
@@ -410,8 +421,6 @@ FSSDSignalSample ClampSignal(FSSDSignalSample Sample, FSSDSignalSample SampleMin
 	#endif
 	OutSample.SampleCount = clamp(Sample.SampleCount, SampleMin.SampleCount, SampleMax.SampleCount);
 	OutSample.MissCount = clamp(Sample.MissCount, SampleMin.MissCount, SampleMax.MissCount);
-	OutSample.ClosestHitDistance = clamp(Sample.ClosestHitDistance, SampleMin.ClosestHitDistance, SampleMax.ClosestHitDistance);
-	OutSample.WorldBluringRadius = clamp(Sample.WorldBluringRadius, SampleMin.WorldBluringRadius, SampleMax.WorldBluringRadius);
 	OutSample.TransmissionDistance = clamp(Sample.TransmissionDistance, SampleMin.TransmissionDistance, SampleMax.TransmissionDistance);
 	return OutSample;
 }
@@ -459,8 +468,22 @@ FSSDSignalSample WaveBroadcastSignal(const FWaveBroadcastSettings BroadcastSetti
 	#endif
 	OutSample.SampleCount        = WaveBroadcast(BroadcastSettings, Sample.SampleCount);
 	OutSample.MissCount          = WaveBroadcast(BroadcastSettings, Sample.MissCount);
-	OutSample.ClosestHitDistance = WaveBroadcast(BroadcastSettings, Sample.ClosestHitDistance);
-	OutSample.WorldBluringRadius = WaveBroadcast(BroadcastSettings, Sample.WorldBluringRadius);
 	OutSample.TransmissionDistance = WaveBroadcast(BroadcastSettings, Sample.TransmissionDistance);
 	return OutSample;
 }
+
+FSSDSignalFrequency MinSignalFrequency(FSSDSignalFrequency FrequencyA, FSSDSignalFrequency FrequencyB)
+{
+	FSSDSignalFrequency OutFrequency;
+	OutFrequency.ClosestHitDistance = min(FrequencyA.ClosestHitDistance, FrequencyB.ClosestHitDistance);
+	OutFrequency.WorldBluringRadius = min(FrequencyA.WorldBluringRadius, FrequencyB.WorldBluringRadius);
+	return OutFrequency;
+}
+
+FSSDSignalFrequency WaveBroadcastSignalFrequency(const FWaveBroadcastSettings BroadcastSettings, FSSDSignalFrequency Frequency)
+{
+	FSSDSignalFrequency OutFrequency;
+	OutFrequency.ClosestHitDistance = WaveBroadcast(BroadcastSettings, Frequency.ClosestHitDistance);
+	OutFrequency.WorldBluringRadius = WaveBroadcast(BroadcastSettings, Frequency.WorldBluringRadius);
+	return OutFrequency;
+}
diff --git a/Engine/Shaders/Private/ScreenSpaceDenoise/SSDSignalFramework.ush b/Engine/Shaders/Private/ScreenSpaceDenoise/SSDSignalFramework.ush
--- a/Engine/Shaders/Private/ScreenSpaceDenoise/SSDSignalFramework.ush
+++ b/Engine/Shaders/Private/ScreenSpaceDenoise/SSDSignalFramework.ush
@@ -358,11 +358,11 @@ FSphericalGaussian ComputeRoughnessLobe(FSSDSampleSceneInfos RefSceneMetadata)
 }
 
 /** Compute the factor to apply to inifinity bluring radius based on the sample's hit distance. */
-float ComputeBluringOutOfFocusOfNormalizedSample(FSSDSignalSample Sample, FSSDSampleSceneInfos SceneMetadata)
+float ComputeBluringOutOfFocusOfNormalizedSample(FSSDSignalFrequency SampleFrequency, FSSDSampleSceneInfos SceneMetadata)
 {
 	float DistanceCameraToPixel = length(View.TranslatedWorldCameraOrigin - GetTranslatedWorldPosition(SceneMetadata));
 
-	return Sample.ClosestHitDistance / (DistanceCameraToPixel + Sample.ClosestHitDistance);
+	return SampleFrequency.ClosestHitDistance / (DistanceCameraToPixel + SampleFrequency.ClosestHitDistance);
 }
 
 
@@ -469,89 +469,14 @@ float ComputeInfinityBluringRadius(FSSDSampleSceneInfos SceneMetadata)
 }
 
 /** Return the size of bluring radius caused by specular lobe in ViewportUV.x space. */
-float ComputeBluringRadiusOfNormalizedSample(FSSDSignalSample Sample, FSSDSampleSceneInfos SceneMetadata)
+float ComputeBluringRadiusOfNormalizedSample(FSSDSignalFrequency SampleFrequency, FSSDSampleSceneInfos SceneMetadata)
 {
-	return ComputeInfinityBluringRadius(SceneMetadata) * ComputeBluringOutOfFocusOfNormalizedSample(Sample, SceneMetadata);
-}
-
-// TODO(Denoiser): duplicated...
-float3 GenerateReflectedRayDirection(
-	float3 IncidentDirection,
-	float3 WorldNormal,
-	float Roughness,
-	float2 RandSample,
-	out float3 RayDirection,
-	out float PDF)
-{
-	if (Roughness < 0.001 ) //ReflectionSmoothClamp)
-	{
-		RayDirection = reflect(IncidentDirection, WorldNormal);
-		PDF = 1.0;
-	}
-	else
-	{
-		float3 N = WorldNormal;
-		float3 V = -IncidentDirection;
-		float2 E = RandSample;
-
-		float3x3 TangentBasis = GetTangentBasis(N);
-		float3 TangentV = mul(TangentBasis, V);
-
-		float NoV = saturate(dot(V, WorldNormal));
-
-		float4 Sample = ImportanceSampleVisibleGGX(UniformSampleDisk(E), Pow4(Roughness), TangentV);
-
-		float3 H = mul(Sample.xyz, TangentBasis);
-		float3 L = 2 * dot( V, H ) * H - V;
-
-		RayDirection = L;
-		PDF = Sample.w;
-	}
-
-	return RayDirection;
-}
-
-// TODO(Denoiser): duplicated...
-void GetInputRayDirectionPixel(
-	uint2 PixelCoord,
-	FSSDSampleSceneInfos SceneMetadata,
-	out float3 RayDirection,
-	out float PDF)
-{
-	float3 IncidentDirection = normalize(GetTranslatedWorldPosition(SceneMetadata) - View.TranslatedWorldCameraOrigin);
-
-	float2 E;
-	{
-		E = Rand1SPPDenoiserInput(PixelCoord);
-		E.y *= 1 - GGX_IMPORTANT_SAMPLE_BIAS;
-	}
-
-	GenerateReflectedRayDirection(
-		IncidentDirection,
-		GetWorldNormal(SceneMetadata),
-		SceneMetadata.Roughness,
-		E,
-		/* out = */ RayDirection,
-		/* out = */ PDF);
+	return ComputeInfinityBluringRadius(SceneMetadata) * ComputeBluringOutOfFocusOfNormalizedSample(SampleFrequency, SceneMetadata);
 }
 
 
 //------------------------------------------------------- MAIN FUNCTIONS TO IMPLEMENT FOR THE PASSES
 
-/** Return whether this signal is occluded or not. */
-#if CONFIG_SIGNAL_PROCESSING == SIGNAL_PROCESSING_SHADOW_VISIBILITY_MASK
-bool IsMissSample(FSSDSignalSample Sample)
-{
-	return Sample.MissCount != 0.0;
-}
-#endif
-
-/** Return whether this signal is valid or not. */
-bool IsInvalidSample(FSSDSignalSample Sample)
-{
-	return Sample.SampleCount == 0.0;
-}
-
 /** COmpute the weight of the bilateral filter to use between a reference sample and a neighbor. */
 float ComputeBilateralWeight(
 	const uint BilateralSettings,
@@ -664,16 +589,16 @@ float ComputeBilateralWeight(
 	return PositionBilateralWeight * NormalBilateralWeight * LobeSimilarity * AxesSimilarity;
 }
 
-float GetSignalWorldBluringRadius(FSSDSignalSample Sample, FSSDSampleSceneInfos SceneMetadata, uint BatchedSignalId)
+float GetSignalWorldBluringRadius(FSSDSignalFrequency SampleFrequency, FSSDSampleSceneInfos SceneMetadata, uint BatchedSignalId)
 #if CONFIG_SIGNAL_PROCESSING == SIGNAL_PROCESSING_SHADOW_VISIBILITY_MASK
 {
-	if (IsInvalidSample(Sample))
+	if (SampleFrequency.ClosestHitDistance == DENOISER_INVALID_HIT_DISTANCE)
 	{
 		return WORLD_RADIUS_INVALID;
 	}
-	else if (IsMissSample(Sample))
+	else if (SampleFrequency.ClosestHitDistance == DENOISER_MISS_HIT_DISTANCE)
 	{
-		return WORLD_RADIUS_MISS;
+		return 100000;
 	}
 
 	FSSDSignalDomainKnowledge DomainKnowledge = GetSignalDomainKnowledge(BatchedSignalId);
@@ -685,14 +610,14 @@ float GetSignalWorldBluringRadius(FSSDSignalSample Sample, FSSDSampleSceneInfos
 	BRANCH
 	if (DomainKnowledge.LightType == LIGHT_TYPE_DIRECTIONAL)
 	{
-		return DomainKnowledge.HitDistanceToWorldBluringRadius * Sample.ClosestHitDistance;
+		return DomainKnowledge.HitDistanceToWorldBluringRadius * SampleFrequency.ClosestHitDistance;
 	}
 	else
 	{
 		return ComputeLightSampleWorldBluringRadius(
 			GetTranslatedWorldPosition(SceneMetadata),
 			DomainKnowledge.LightType, DomainKnowledge.Light,
-			Sample.ClosestHitDistance);
+			SampleFrequency.ClosestHitDistance);
 	}
 }
 #elif CONFIG_SIGNAL_PROCESSING == SIGNAL_PROCESSING_POLYCHROMATIC_PENUMBRA_HARMONIC
@@ -707,8 +632,17 @@ float GetSignalWorldBluringRadius(FSSDSignalSample Sample, FSSDSampleSceneInfos
 }
 #elif CONFIG_SIGNAL_PROCESSING == SIGNAL_PROCESSING_AO || CONFIG_SIGNAL_PROCESSING == SIGNAL_PROCESSING_DIFFUSE_INDIRECT_AND_AO
 {
+	if (SampleFrequency.ClosestHitDistance == DENOISER_INVALID_HIT_DISTANCE)
+	{
+		return WORLD_RADIUS_INVALID;
+	}
+	else if (SampleFrequency.ClosestHitDistance == DENOISER_MISS_HIT_DISTANCE)
+	{
+		return 100000;
+	}
+
 	// How much can be blured is about the the length of hit distance assuming some amount of visiiblity.
-	return Sample.ClosestHitDistance;
+	return SampleFrequency.ClosestHitDistance;
 }
 #elif CONFIG_SIGNAL_PROCESSING == SIGNAL_PROCESSING_DIFFUSE_SPHERICAL_HARMONIC
 {
@@ -740,89 +674,6 @@ float GetSamplePenumbra(FSSDSignalSample Sample)
 	return Sample.MissCount / Sample.SampleCount;
 }
 
-/** Get the weight for multi importance sampling. */
-float GetRatioEstimatorWeight(
-	FSSDSampleSceneInfos RefSceneMetadata,
-	FSSDSampleSceneInfos SampleSceneMetadata,
-	FSSDSignalSample Sample,
-	uint2 SamplePixelCoord)
-#if CONFIG_SPECULAR_RATIO_ESTIMATOR
-// [ Stackowiak 2018, "Stochastic all the things: Raytracing in hybrid real-time rendering" ]
-{
-	#if CONFIG_USE_VIEW_SPACE
-		#error Ratio estimator should be done in world space to ensure exact same ray direction.
-	#endif
-
-	// If invalid ray or geometric miss, just.
-	if (IsInvalidSample(Sample) && 0)
-	{
-		return 0; 
-	}
-
-	float3 OriginalRayDirection;
-	float OriginalRayPDF;
-
-	GetInputRayDirectionPixel(
-		SamplePixelCoord,
-		SampleSceneMetadata, 
-		/* out */ OriginalRayDirection,
-		/* out */ OriginalRayPDF);
-
-	// min to avoid floating point underflow that causes artifact by D anormally high.
-	float ClosestHitDistance = min(Sample.ClosestHitDistance, 1000000);
-
-	float a2 = Pow4(RefSceneMetadata.Roughness);
-
-	// TODO(Denoiser): Improve quantization in CONFIG_COMPRESS_WORLD_NORMAL and use it.
-	float3 N = GetWorldNormal(RefSceneMetadata);
-
-	float NoH;
-	#if CONFIG_VALU_OPTIMIZATIONS
-	{
-		float3 V = (View.TranslatedWorldCameraOrigin - GetTranslatedWorldPosition(RefSceneMetadata));
-		float3 L = ((GetTranslatedWorldPosition(SampleSceneMetadata) - GetTranslatedWorldPosition(RefSceneMetadata)) + OriginalRayDirection * ClosestHitDistance);
-		float3 H = (L + V * (sqrt(length2(L)) * rsqrt(length2(V))));
-
-		NoH = saturate(dot(N, H) * rsqrt(length2(H)));
-	}
-	#else
-	{
-		float3 V = normalize(View.TranslatedWorldCameraOrigin - GetTranslatedWorldPosition(RefSceneMetadata));
-		float3 L = normalize((GetTranslatedWorldPosition(SampleSceneMetadata) - GetTranslatedWorldPosition(RefSceneMetadata)) + OriginalRayDirection * ClosestHitDistance);
-		float3 H = normalize(L + V);
-
-		NoH = abs(dot(N, H));
-	}
-	#endif
-
-
-	#if 1
-	{
-		const float MaximumWeight = 10;
-
-		float d = saturate(clamp( NoH * a2 - NoH, -NoH, 0) * NoH + 1);
-	
-		if (d == 0)
-		{
-			return MaximumWeight;
-		}
-
-		return clamp(a2 / (PI * d * d * OriginalRayPDF), 0, MaximumWeight);
-	}
-	#else
-	{
-		float D = D_GGX(a2, NoH);
-	
-		return clamp(D / OriginalRayPDF, 0, 10);
-	}
-	#endif
-}
-#else // !CONFIG_SPECULAR_RATIO_ESTIMATOR
-{
-	return 1.0;
-}
-#endif // CONFIG_SPECULAR_RATIO_ESTIMATOR
-
 
 // Matches GenerateCosineNormalRay()
 // TODO(Denoiser): expose a Weight that multiply directly onto the Sample.SceneColor to reduce ALU
diff --git a/Engine/Shaders/Private/ScreenSpaceDenoise/SSDSpatialAccumulation.usf b/Engine/Shaders/Private/ScreenSpaceDenoise/SSDSpatialAccumulation.usf
--- a/Engine/Shaders/Private/ScreenSpaceDenoise/SSDSpatialAccumulation.usf
+++ b/Engine/Shaders/Private/ScreenSpaceDenoise/SSDSpatialAccumulation.usf
@@ -165,7 +165,7 @@
 		#define SIGNAL_ARRAY_SIZE 1
 
 		// Input and output layout.
-		#define CONFIG_SIGNAL_INPUT_LAYOUT  SIGNAL_BUFFER_LAYOUT_REFLECTIONS_INPUT
+			#define CONFIG_SIGNAL_INPUT_LAYOUT  SIGNAL_BUFFER_LAYOUT_REFLECTIONS_INPUT
 		#define CONFIG_SIGNAL_OUTPUT_LAYOUT SIGNAL_BUFFER_LAYOUT_REFLECTIONS_HISTORY
 
 		// Configure the bilateral therm.
@@ -342,7 +342,6 @@
 	#define CONFIG_SIGNAL_VGPR_COMPRESSION SIGNAL_COMPRESSION_DIFFUSE_INDIRECT_HARMONIC
 
 	#if DIM_STAGE == STAGE_RECONSTRUCTION
-		//#define CONFIG_SAMPLE_SET           SAMPLE_SET_STOCASTIC_HIERARCHY
 		#define CONFIG_SAMPLE_SET           SAMPLE_SET_STACKOWIAK_4_SETS
 
 	#elif DIM_STAGE == STAGE_REJECTION_PRE_CONVOLUTION
@@ -491,7 +490,7 @@
 // White list accumulators to compile.
 #if CONFIG_OUTPUT_MODE == OUTPUT_MODE_DRB
 	#define COMPILE_DRB_ACCUMULATOR 1
-	#define COMPILE_MININVFREQ_ACCUMULATOR 1
+	#define COMPILE_MIN_FREQUENCY_ACCUMULATOR 1
 
 #elif CONFIG_OUTPUT_MODE == OUTPUT_MODE_2MOMMENT_SUM
 	#define COMPILE_MOMENT1_ACCUMULATOR 1
@@ -499,7 +498,7 @@
 
 #elif CONFIG_OUTPUT_MODE == OUTPUT_MODE_SUM
 	#define COMPILE_MOMENT1_ACCUMULATOR 1
-	#define COMPILE_MININVFREQ_ACCUMULATOR 1
+	#define COMPILE_MIN_FREQUENCY_ACCUMULATOR 1
 
 #else
 	#error Unknown output mode.
@@ -662,7 +661,9 @@ void MainCS(
 
 	// Sample the reference sample.
 	#if !CONFIG_UPSCALE || 1
-		FSSDSignalArray RefSamples = SampleMultiplexedSignals(
+		FSSDSignalArray RefSamples;
+		FSSDSignalFrequencyArray RefFrequencies;
+		SampleMultiplexedSignals(
 			Signal_Textures_0,
 			Signal_Textures_1,
 			Signal_Textures_2,
@@ -671,7 +672,9 @@ void MainCS(
 			CONFIG_SIGNAL_INPUT_LAYOUT,
 			/* MultiplexedSampleId = */ 0,
 			/* bNormalizeSample = */ CONFIG_NORMALIZE_INPUT != 0,
-			SceneBufferUV);
+			SceneBufferUV,
+			/* out */ RefSamples,
+			/* out */ RefFrequencies);
 		
 		#if CONFIG_NORMALIZE_INPUT
 			FSSDSignalArray NormalizedRefSamples = RefSamples;
@@ -707,7 +710,7 @@ void MainCS(
 		#if CONFIG_SIGNAL_BATCH_SIZE != 1
 			#error Unable to support more than one signal.
 		#endif
-		float RefBluringRadius = ComputeBluringRadiusOfNormalizedSample(NormalizedRefSamples.Array[0], RefSceneMetadata) * View.ViewSizeAndInvSize.x * rcp(KernelSpreadFactor);
+		float RefBluringRadius = ComputeBluringRadiusOfNormalizedSample(RefFrequencies, RefSceneMetadata) * View.ViewSizeAndInvSize.x * rcp(KernelSpreadFactor);
 		RequestedSampleCount = PI * RefBluringRadius * RefBluringRadius;
 	}
 	#elif CONFIG_SAMPLE_COUNT_POLICY == SAMPLE_COUNT_POLICY_ACCUMULATION_BASED
@@ -735,7 +738,7 @@ void MainCS(
 	#endif
 
 	// Register renaming of members of FSSDKernelConfig to survive until the output to UAV
-	#if (CONFIG_SAMPLE_SET == SAMPLE_SET_STACKOWIAK_4_SETS || CONFIG_SAMPLE_SET == SAMPLE_SET_STOCASTIC_HIERARCHY) && CONFIG_VGPR_OPTIMIZATION
+	#if (CONFIG_SAMPLE_SET == SAMPLE_SET_STACKOWIAK_4_SETS) && CONFIG_VGPR_OPTIMIZATION
 		float2 KernelBufferUV;
 		uint SampleTrackId;
 	#endif
@@ -823,13 +826,13 @@ void MainCS(
 			{
 				if (KernelConfig.BilateralDistanceComputation == SIGNAL_WORLD_FREQUENCY_PRECOMPUTED_BLURING_RADIUS)
 				{
-					KernelConfig.RefBilateralDistance[MultiplexId] = RefSamples.Array[MultiplexId].WorldBluringRadius;
+					KernelConfig.RefBilateralDistance[MultiplexId] = RefFrequencies.Array[MultiplexId].WorldBluringRadius;
 				}
 				else
 				{
 					const uint BatchedSignalId = ComputeSignalBatchIdFromSignalMultiplexId(KernelConfig, MultiplexId);
 
-					KernelConfig.RefBilateralDistance[MultiplexId] = GetSignalWorldBluringRadius(RefSamples.Array[MultiplexId], RefSceneMetadata, BatchedSignalId);
+					KernelConfig.RefBilateralDistance[MultiplexId] = GetSignalWorldBluringRadius(RefFrequencies.Array[MultiplexId], RefSceneMetadata, BatchedSignalId);
 				}
 			}
 		}
@@ -880,13 +883,13 @@ void MainCS(
 				
 				FSSDSampleAccumulationInfos SampleInfos;
 				SampleInfos.Sample = CenterSample;
+				SampleInfos.Frequency = RefFrequencies.Array[SignalMultiplexId];
 				SampleInfos.FinalWeight = 1.0;
-				SampleInfos.InvFrequency = GetSignalWorldBluringRadius(CenterSample, RefSceneMetadata, BatchedSignalId);
+				SampleInfos.InvFrequency = GetSignalWorldBluringRadius(SampleInfos.Frequency, RefSceneMetadata, BatchedSignalId);
 				
 				if (KernelConfig.BilateralDistanceComputation == SIGNAL_WORLD_FREQUENCY_PRECOMPUTED_BLURING_RADIUS)
 				{
-					// TODO(Denoiser): this is ineficient, could fetch the normalised WorldBluringRadius instead of SafeRcp().
-					SampleInfos.InvFrequency = CenterSample.WorldBluringRadius * SafeRcp(CenterSample.SampleCount);
+					SampleInfos.InvFrequency = SampleInfos.Frequency.WorldBluringRadius;
 				}
 
 				AccumulateSample(
@@ -936,26 +939,6 @@ void MainCS(
 				KernelBufferUV = KernelConfig.BufferUV;
 			#endif
 		}
-		#elif CONFIG_SAMPLE_SET == SAMPLE_SET_STOCASTIC_HIERARCHY
-		{
-			#if CONFIG_UPSCALE
-				#error Unsupported.
-			#endif
-
-			// Put the kernel center at the center of the quad. Half pixel shift is done in the sample offsets.
-			KernelConfig.BufferUV = (DispatchThreadId | 1) * KernelConfig.BufferSizeAndInvSize.zw;
-
-			// Id of the pixel in the quad. This is to match hard coded first samples of the sample set.
-			KernelConfig.SampleTrackId = ((DispatchThreadId.x & 1) | ((DispatchThreadId.y & 1) << 1));
-
-			#if CONFIG_VGPR_OPTIMIZATION
-				// Keek sample SampleTrackId & SceneBufferUV arround for computation of pixel output coordinate.
-				// Should be VGPR free given it's curernt is being used in accumulation has well that is highest VGPR pressure of the shader.
-				// TODO(Denoiser): could save 1 VGPR by using 2 SGPR instead of SampleTrackId.
-				SampleTrackId = KernelConfig.SampleTrackId;
-				KernelBufferUV = KernelConfig.BufferUV;
-			#endif
-		}
 		#elif CONFIG_SAMPLE_SET == SAMPLE_SET_DIRECTIONAL_RECT || CONFIG_SAMPLE_SET == SAMPLE_SET_DIRECTIONAL_ELLIPSE
 		{
 			#if CONFIG_SIGNAL_PROCESSING == SIGNAL_PROCESSING_REFLECTIONS
@@ -973,7 +956,7 @@ void MainCS(
 					/* out */ InifinityMajorViewportRadius,
 					/* out */ InifinityMinorViewportRadius);
 
-				float OutOfFocus = ComputeBluringOutOfFocusOfNormalizedSample(NormalizedRefSamples.Array[0], RefSceneMetadata);
+				float OutOfFocus = ComputeBluringOutOfFocusOfNormalizedSample(RefFrequencies.Array[0], RefSceneMetadata);
 
 				KernelConfig.MajorAxis = NormalizedScreenMajorAxis * float2(1, -1);
 				KernelConfig.MajorPixelRadius = InifinityMajorViewportRadius * OutOfFocus * View.ViewSizeAndInvSize.x;
@@ -1002,7 +985,7 @@ void MainCS(
 				#error Directional rect sample set is not supported.
 			#endif
 		}
-		#endif // CONFIG_SAMPLE_SET == SAMPLE_SET_STOCASTIC_HIERARCHY
+		#endif // CONFIG_SAMPLE_SET == SAMPLE_SET_DIRECTIONAL_*
 		
 		FSSDCompressedSignalAccumulatorArray CompressedAccumulators = CompressAccumulatorArray(UncompressedAccumulators, CONFIG_ACCUMULATOR_VGPR_COMPRESSION);
 
@@ -1114,6 +1097,7 @@ void MainCS(
 	// Transcode the spatial accumulation into multiplexed signal according to different modes.
 	uint MultiplexCount = 1;
 	FSSDSignalArray OutputSamples = CreateSignalArrayFromScalarValue(0.0);
+	FSSDSignalFrequencyArray OutputFrequencies = CreateInvalidSignalFrequencyArray();
 	{
 		#if CONFIG_OUTPUT_MODE == OUTPUT_MODE_SUM
 		{
@@ -1127,9 +1111,7 @@ void MainCS(
 				OutputSamples.Array[MultiplexId] = SignalAccumulators.Array[MultiplexId].Moment1;
 				
 				// Output the minimal inverse frequency as new world bluring radius for subsequent passes.
-				// TODO(Denoiser): store this in its own signal to avoid *= N; *= rcp(N);
-				OutputSamples.Array[MultiplexId].WorldBluringRadius = 
-					OutputSamples.Array[MultiplexId].SampleCount * SignalAccumulators.Array[MultiplexId].MinInvFrequency;
+				OutputFrequencies.Array[MultiplexId] = SignalAccumulators.Array[MultiplexId].MinFrequency;
 			}
 		}
 		#elif CONFIG_OUTPUT_MODE == OUTPUT_MODE_2MOMMENT_SUM
@@ -1161,14 +1143,12 @@ void MainCS(
 				OutputSamples.Array[MultiplexId] = SignalAccumulators.Array[MultiplexId].Previous;
 
 				// Output the minimal inverse frequency as new world bluring radius for subsequent passes.
-				// TODO(Denoiser): store this in its own signal to avoid *= N; *= rcp(N);
-				OutputSamples.Array[MultiplexId].WorldBluringRadius = 
-					OutputSamples.Array[MultiplexId].SampleCount * SignalAccumulators.Array[MultiplexId].MinInvFrequency;
+				OutputFrequencies.Array[MultiplexId] = SignalAccumulators.Array[MultiplexId].MinFrequency;
 			
 				// No need to keep the VGPR pressure at this point for WorldBluringRadius, because no passes use it after.
 				if (DIM_STAGE == STAGE_POST_FILTERING && 0)
 				{
-					OutputSamples.Array[MultiplexId].WorldBluringRadius = 0;
+					OutputFrequencies.Array[MultiplexId].WorldBluringRadius = 0;
 				}
 			}
 		}
@@ -1196,7 +1176,7 @@ void MainCS(
 	#if CONFIG_VGPR_OPTIMIZATION && !CONFIG_UPSCALE // TODO(Denoiser)
 	{
 		// No need to keep DispatchThreadId, can recompute the output pixel position based on information stored in VGPRs for spatial kernel.
-		#if CONFIG_SAMPLE_SET == SAMPLE_SET_STACKOWIAK_4_SETS || CONFIG_SAMPLE_SET == SAMPLE_SET_STOCASTIC_HIERARCHY
+		#if CONFIG_SAMPLE_SET == SAMPLE_SET_STACKOWIAK_4_SETS
 			#if CONFIG_VGPR_FREE_SAMPLE_TRACK_ID
 				SampleTrackId = GetSampleTrackIdFromLaneIndex();
 			#endif
@@ -1252,8 +1232,11 @@ void MainCS(
 				SignalOutput_UAVs_1,
 				SignalOutput_UAVs_2,
 				SignalOutput_UAVs_3,
-				CONFIG_SIGNAL_OUTPUT_LAYOUT, MultiplexCount,
-				OutputPixelPostion, OutputSamples.Array);
+				CONFIG_SIGNAL_OUTPUT_LAYOUT,
+				MultiplexCount,
+				OutputPixelPostion,
+				OutputSamples,
+				OutputFrequencies);
 		}
 		#endif
 	}
diff --git a/Engine/Shaders/Private/ScreenSpaceDenoise/SSDSpatialKernel.ush b/Engine/Shaders/Private/ScreenSpaceDenoise/SSDSpatialKernel.ush
--- a/Engine/Shaders/Private/ScreenSpaceDenoise/SSDSpatialKernel.ush
+++ b/Engine/Shaders/Private/ScreenSpaceDenoise/SSDSpatialKernel.ush
@@ -54,10 +54,6 @@
 	#define COMPILE_DISK_KERNEL 0
 #endif
 
-#ifndef COMPILE_HEIRARCHY_KERNEL
-	#define COMPILE_HEIRARCHY_KERNEL 0
-#endif
-
 #ifndef COMPILE_DIRECTIONAL_KERNEL
 	#define COMPILE_DIRECTIONAL_KERNEL 0
 #endif
@@ -142,9 +138,6 @@ struct FSSDKernelConfig
 	float4 BufferSizeAndInvSize;
 	float4 BufferBilinearUVMinMax;
 
-	// Mip level in the buffer to sample.
-	float BufferMipLevel;
-
 	// Multiplier on the sample's offset.
 	float KernelSpreadFactor;
 
@@ -242,7 +235,6 @@ FSSDKernelConfig CreateKernelConfig()
 	// SGPRs.
 	KernelConfig.BufferSizeAndInvSize = float4(0, 0, 0, 0);
 	KernelConfig.BufferBilinearUVMinMax = float4(0, 0, 0, 0);
-	KernelConfig.BufferMipLevel = 0.0;
 	KernelConfig.KernelSpreadFactor = 1;
 	KernelConfig.HarmonicPeriode = 1.0;
 	
@@ -529,7 +521,6 @@ void SampleMultiplexedSignals(
 	FSSDTexture2D SignalBuffer2,
 	FSSDTexture2D SignalBuffer3,
 	float2 SampleBufferUV,
-	float MipLevel,
 	out FSSDCompressedSceneInfos OutCompressedSampleSceneMetadata,
 	out FSSDCompressedMultiplexedSample OutCompressedMultiplexedSamples)
 {
@@ -546,22 +537,25 @@ void SampleMultiplexedSignals(
 		SignalBuffer3,
 		GlobalPointClampedSampler,
 		SampleBufferUV,
-		KernelConfig.BufferMipLevel + MipLevel,
 		PixelCoord);
 } // SampleMultiplexedSignals()
 
 /** Uncompressed multiplexed signal for accumulation. */
-FSSDSignalArray UncompressMultiplexedSignals(
+void UncompressMultiplexedSignals(
 	FSSDKernelConfig KernelConfig,
 	float2 SampleBufferUV,
-	FSSDCompressedMultiplexedSample CompressedMultiplexedSamples)
+	FSSDCompressedMultiplexedSample CompressedMultiplexedSamples,
+	out FSSDSignalArray MultiplexedSamples,
+	out FSSDSignalFrequencyArray MultiplexedFrequencies)
 {
 	// TODO(Denoiser): offer multiplier to apply to each signal during Decode, to save mul VALU.
-	FSSDSignalArray MultiplexedSamples = DecodeMultiplexedSignals(
+	DecodeMultiplexedSignals(
 		KernelConfig.BufferLayout,
 		/* MultiplexedSampleId = */ 0,
 		KernelConfig.bNormalizeSample,
-		CompressedMultiplexedSamples);
+		CompressedMultiplexedSamples,
+		/* out */ MultiplexedSamples,
+		/* out */ MultiplexedFrequencies);
 
 	if (KernelConfig.bClampUVPerMultiplexedSignal)
 	{
@@ -577,8 +571,6 @@ FSSDSignalArray UncompressMultiplexedSignals(
 			}
 		} // for (uint SignalMultiplexId = 0; SignalMultiplexId < SIGNAL_ARRAY_SIZE; SignalMultiplexId++)
 	}
-
-	return MultiplexedSamples;
 }
 
 /** Accumulate multiplexed samples and their metadata to an accumulator. */
@@ -590,6 +582,7 @@ void AccumulateSampledMultiplexedSignals(
 	float2 SampleBufferUV,
 	FSSDSampleSceneInfos SampleSceneMetadata,
 	FSSDSignalArray MultiplexedSamples,
+	FSSDSignalFrequencyArray MultiplexedFrequencies,
 	float KernelSampleWeight,
 	const bool bForceSample,
 	bool bIsOutsideFrustum)
@@ -651,18 +644,20 @@ void AccumulateSampledMultiplexedSignals(
 		// Fetch and pre process the sample for accumulation.
 		FSSDSignalSample Sample = MultiplexedSamples.Array[SignalMultiplexId];
 		Sample = TransformSignalSampleForAccumulation(KernelConfig, SignalMultiplexId, SampleSceneMetadata, Sample, NeighborPixelCoord);
-		
+
+		// Fetch sample's frequency for accumulation.
+		FSSDSignalFrequency SampleFrequency = MultiplexedFrequencies.Array[SignalMultiplexId];
+
 		// Compute the bluring radius of pixel itself.
 		float SamplePixelWorldBluringRadius = ComputeWorldBluringRadiusCausedByPixelSize(SampleSceneMetadata);
 		
 		// Compute the bluring radius of the signal from ray hit distance and signal domain knowledge.
-		float SignalConvolutionBluringRadius = GetSignalWorldBluringRadius(Sample, SampleSceneMetadata, BatchedSignalId);
+		float SignalConvolutionBluringRadius = GetSignalWorldBluringRadius(SampleFrequency, SampleSceneMetadata, BatchedSignalId);
 		
 		// But the signal's bluring radius might already be pre computed.
 		if (KernelConfig.BilateralDistanceComputation == SIGNAL_WORLD_FREQUENCY_PRECOMPUTED_BLURING_RADIUS)
 		{
-			// TODO(Denoiser): this is ineficient, could fetch the normalised WorldBluringRadius instead of SafeRcp().
-			SignalConvolutionBluringRadius = Sample.WorldBluringRadius * SafeRcp(Sample.SampleCount);
+			SignalConvolutionBluringRadius = SampleFrequency.WorldBluringRadius;
 		}
 
 		// Compute the final world distance to use for bilateral rejection.
@@ -709,12 +704,10 @@ void AccumulateSampledMultiplexedSignals(
 			SampleSceneMetadata,
 			NeighborToRefVector);
 		
-		float RatioEstimatorWeight = GetRatioEstimatorWeight(
-			RefSceneMetadata, SampleSceneMetadata, Sample, NeighborPixelCoord);
-
 		FSSDSampleAccumulationInfos SampleInfos;
 		SampleInfos.Sample = Sample;
-		SampleInfos.FinalWeight = KernelSampleWeight * BilateralWeight * RatioEstimatorWeight;
+		SampleInfos.Frequency = SampleFrequency;
+		SampleInfos.FinalWeight = KernelSampleWeight * BilateralWeight;
 		SampleInfos.InvFrequency = SignalConvolutionBluringRadius;
 
 		if (bForceSample || KernelConfig.bForceAllAccumulation)
@@ -760,7 +753,6 @@ void SampleAndAccumulateMultiplexedSignals(
 	inout FSSDSignalAccumulatorArray UncompressedAccumulators,
 	inout FSSDCompressedSignalAccumulatorArray CompressedAccumulators,
 	float2 SampleBufferUV,
-	float MipLevel,
 	float KernelSampleWeight,
 	const bool bForceSample)
 {
@@ -780,15 +772,18 @@ void SampleAndAccumulateMultiplexedSignals(
 			SignalBuffer2,
 			SignalBuffer3,
 			SampleBufferUV,
-			MipLevel,
 			/* out */ CompressedSampleSceneMetadata,
 			/* out */ CompressedMultiplexedSamples);
 	}
 	
 	// Accumulate the samples, giving full freedom for shader compiler scheduler to put instructions in most optimal way.
 	{
-		FSSDSignalArray MultiplexedSamples = UncompressMultiplexedSignals(
-			KernelConfig, SampleBufferUV, CompressedMultiplexedSamples);
+		FSSDSignalArray MultiplexedSamples;
+		FSSDSignalFrequencyArray MultiplexedFrequencies;
+		UncompressMultiplexedSignals(
+			KernelConfig, SampleBufferUV, CompressedMultiplexedSamples,
+			/* out */ MultiplexedSamples,
+			/* out */ MultiplexedFrequencies);
 	
 		FSSDSampleSceneInfos RefSceneMetadata = UncompressRefSceneMetadata(KernelConfig);
 	
@@ -803,6 +798,7 @@ void SampleAndAccumulateMultiplexedSignals(
 			SampleBufferUV,
 			SampleSceneMetadata,
 			MultiplexedSamples,
+			MultiplexedFrequencies,
 			KernelSampleWeight,
 			bForceSample,
 			bIsOutsideFrustum);
@@ -841,7 +837,6 @@ void SampleAndAccumulateMultiplexedSignalsPair(
 				SignalBuffer2,
 				SignalBuffer3,
 				SampleBufferUV[PairFetchId],
-				/* MipLevel = */ 0.0,
 				/* out */ CompressedSampleSceneMetadata[PairFetchId],
 				/* out */ CompressedMultiplexedSamples[PairFetchId]);
 		}
@@ -851,29 +846,33 @@ void SampleAndAccumulateMultiplexedSignalsPair(
 	{
 		// Uncompress the multiplexed signal.
 		FSSDSignalArray MultiplexedSamples[2];
+		FSSDSignalFrequencyArray MultiplexedFrequencies[2];
 		UNROLL_N(2)
 		for (uint PairUncompressId = 0; PairUncompressId < 2; PairUncompressId++)
 		{
-			MultiplexedSamples[PairUncompressId] = UncompressMultiplexedSignals(
-				KernelConfig, SampleBufferUV[PairUncompressId], CompressedMultiplexedSamples[PairUncompressId]);
+			UncompressMultiplexedSignals(
+				KernelConfig,
+				SampleBufferUV[PairUncompressId],
+				CompressedMultiplexedSamples[PairUncompressId],
+				/* out */ MultiplexedSamples[PairUncompressId],
+				/* out */ MultiplexedFrequencies[PairUncompressId]);
 		}
 
 		// Take the min inverse frequency per signal if desired.
-		// TODO(Denoiser): this should be normalized in theory...
 		if (KernelConfig.bMinSamplePairInvFrequency)
 		{
 			UNROLL_N(SIGNAL_ARRAY_SIZE)
 			for (uint SignalMultiplexId = 0; SignalMultiplexId < SIGNAL_ARRAY_SIZE; SignalMultiplexId++)
 			{
 				float MinInvFrequency = min(
-					MultiplexedSamples[0].Array[SignalMultiplexId].WorldBluringRadius,
-					MultiplexedSamples[1].Array[SignalMultiplexId].WorldBluringRadius);
+					MultiplexedFrequencies[0].Array[SignalMultiplexId].WorldBluringRadius,
+					MultiplexedFrequencies[1].Array[SignalMultiplexId].WorldBluringRadius);
 
 				FLATTEN
 				if (MinInvFrequency > 0)
 				{
-					MultiplexedSamples[0].Array[SignalMultiplexId].WorldBluringRadius = MinInvFrequency;
-					MultiplexedSamples[1].Array[SignalMultiplexId].WorldBluringRadius = MinInvFrequency;
+					MultiplexedFrequencies[0].Array[SignalMultiplexId].WorldBluringRadius = MinInvFrequency;
+					MultiplexedFrequencies[1].Array[SignalMultiplexId].WorldBluringRadius = MinInvFrequency;
 				}
 			}
 		}
@@ -894,6 +893,7 @@ void SampleAndAccumulateMultiplexedSignalsPair(
 				SampleBufferUV[PairAccumulateId],
 				SampleSceneMetadata,
 				MultiplexedSamples[PairAccumulateId],
+				MultiplexedFrequencies[PairAccumulateId],
 				KernelSampleWeight,
 				/* bForceSample = */ false,
 				bIsOutsideFrustum[PairAccumulateId]);
@@ -999,7 +999,6 @@ void SampleAndAccumulateCenterSampleAsItsOwnCluster(
 		/* inout */ UncompressedAccumulators,
 		/* inout */ CompressedAccumulators,
 		KernelConfig.BufferUV,
-		/* MipLevel = */ 0.0,
 		/* KernelSampleWeight = */ 1.0,
 		/* bForceSample = */ KernelConfig.bForceKernelCenterAccumulation);
 
@@ -1093,7 +1092,6 @@ void AccumulateBilinear2x2Kernel(
 			/* inout */ UncompressedAccumulators,
 			/* inout */ CompressedAccumulators,
 			SampleBufferUVArray[0],
-			/* MipLevel = */ 0.0,
 			BilinearWeightArray[0],
 			/* bForceSample = */ false);
 	}
@@ -1113,7 +1111,6 @@ void AccumulateBilinear2x2Kernel(
 				/* inout */ UncompressedAccumulators,
 				/* inout */ CompressedAccumulators,
 				SampleBufferUVArray[i],
-				/* MipLevel = */ 0.0,
 				BilinearWeightArray[i],
 				/* bForceSample = */ false);
 		}
@@ -1174,7 +1171,6 @@ void AccumulateSquareKernel(
 				/* inout */ UncompressedAccumulators,
 				/* inout */ CompressedAccumulators,
 				SampleBufferUV,
-				/* MipLevel = */ 0.0,
 				KernelWeight,
 				/* bForceSample = */ bIsKernelCenterSample && KernelConfig.bForceKernelCenterAccumulation);
 		}
@@ -1215,7 +1211,6 @@ void AccumulateSquareKernel(
 				/* inout */ UncompressedAccumulators,
 				/* inout */ CompressedAccumulators,
 				SampleBufferUV,
-				/* MipLevel = */ 0.0,
 				KernelWeight,
 				/* bForceSample = */ bIsKernelCenterSample && KernelConfig.bForceKernelCenterAccumulation);
 		}
@@ -1229,7 +1224,8 @@ void BroadcastAccumulateSquare3x3KernelCenter(
  	FSSDSampleSceneInfos RefSceneMetadata,
  	float2 SampleBufferUV,
  	FSSDSampleSceneInfos SampleSceneMetadata,
- 	FSSDSignalArray SampleMultiplexedSamples)
+ 	FSSDSignalArray SampleMultiplexedSamples,
+	FSSDSignalFrequencyArray SampleMultiplexedFrequencies)
 #if CONFIG_ENABLE_WAVE_BROADCAST
 {
 	const FWaveBroadcastSettings BroadcastSettingsX = InitWaveSwapWithinLaneGroup(/* LaneGroupSize = */ 2);
@@ -1239,6 +1235,9 @@ void BroadcastAccumulateSquare3x3KernelCenter(
 	SampleBufferUV = WaveBroadcast(BroadcastSettingsX, SampleBufferUV);
 	SampleSceneMetadata = WaveBroadcastSceneMetadata(BroadcastSettingsX, SampleSceneMetadata);
 	SampleMultiplexedSamples = WaveBroadcastSignalArray(BroadcastSettingsX, SampleMultiplexedSamples);
+	SampleMultiplexedFrequencies = WaveBroadcastSignalFrequencyArray(BroadcastSettingsX, SampleMultiplexedFrequencies);
+
+	FSSDSignalFrequency WaveBroadcastSignalFrequency(const FWaveBroadcastSettings BroadcastSettings, FSSDSignalFrequency Frequency)
 
 	if (KernelConfig.SampleSet == SAMPLE_SET_3X3 || KernelConfig.SampleSet == SAMPLE_SET_3X3_PLUS)
 	{
@@ -1250,6 +1249,7 @@ void BroadcastAccumulateSquare3x3KernelCenter(
 			SampleBufferUV,
 			SampleSceneMetadata,
 			SampleMultiplexedSamples,
+			SampleMultiplexedFrequencies,
 			/* KernelWeight = */ 1.0,
 			/* bForceSample = */ false,
 			/* bIsOutsideFrustum = */ false);
@@ -1259,6 +1259,7 @@ void BroadcastAccumulateSquare3x3KernelCenter(
 	SampleBufferUV = WaveBroadcast(BroadcastSettingsY, SampleBufferUV);
 	SampleSceneMetadata = WaveBroadcastSceneMetadata(BroadcastSettingsY, SampleSceneMetadata);
 	SampleMultiplexedSamples = WaveBroadcastSignalArray(BroadcastSettingsY, SampleMultiplexedSamples);
+	SampleMultiplexedFrequencies = WaveBroadcastSignalFrequencyArray(BroadcastSettingsY, SampleMultiplexedFrequencies);
 
 	if (KernelConfig.SampleSet == SAMPLE_SET_3X3 || KernelConfig.SampleSet == SAMPLE_SET_3X3_CROSS)
 	{
@@ -1270,6 +1271,7 @@ void BroadcastAccumulateSquare3x3KernelCenter(
 			SampleBufferUV,
 			SampleSceneMetadata,
 			SampleMultiplexedSamples,
+			SampleMultiplexedFrequencies,
 			/* KernelWeight = */ 1.0,
 			/* bForceSample = */ false,
 			/* bIsOutsideFrustum = */ false);
@@ -1279,6 +1281,7 @@ void BroadcastAccumulateSquare3x3KernelCenter(
 	SampleBufferUV = WaveBroadcast(BroadcastSettingsX, SampleBufferUV);
 	SampleSceneMetadata = WaveBroadcastSceneMetadata(BroadcastSettingsX, SampleSceneMetadata);
 	SampleMultiplexedSamples = WaveBroadcastSignalArray(BroadcastSettingsX, SampleMultiplexedSamples);
+	SampleMultiplexedFrequencies = WaveBroadcastSignalFrequencyArray(BroadcastSettingsX, SampleMultiplexedFrequencies);
 	
 	if (KernelConfig.SampleSet == SAMPLE_SET_3X3 || KernelConfig.SampleSet == SAMPLE_SET_3X3_PLUS)
 	{
@@ -1290,6 +1293,7 @@ void BroadcastAccumulateSquare3x3KernelCenter(
 			SampleBufferUV,
 			SampleSceneMetadata,
 			SampleMultiplexedSamples,
+			SampleMultiplexedFrequencies,
 			/* KernelWeight = */ 1.0,
 			/* bForceSample = */ false,
 			/* bIsOutsideFrustum = */ false);
@@ -1329,15 +1333,19 @@ void AccumulateSquare3x3Kernel(
 				SignalBuffer2,
 				SignalBuffer3,
 				SampleBufferUV,
-				/* MipLevel = */ 0,
 				/* out */ CompressedSampleSceneMetadata,
 				/* out */ CompressedMultiplexedSamples);
 		}
 
 		FSSDSampleSceneInfos RefSceneMetadata = UncompressRefSceneMetadata(KernelConfig);
-				
-		FSSDSignalArray MultiplexedSamples = UncompressMultiplexedSignals(
+
+		FSSDSignalArray MultiplexedSamples;
+		FSSDSignalFrequencyArray MultiplexedFrequencies;
+		UncompressMultiplexedSignals(
 			KernelConfig, SampleBufferUV, CompressedMultiplexedSamples);
+			/* out */ MultiplexedSamples,
+			/* out */ MultiplexedFrequencies);
+
 		FSSDSampleSceneInfos SampleSceneMetadata = UncompressSampleSceneMetadata(
 			KernelConfig, SampleBufferUV, CompressedSampleSceneMetadata);
 
@@ -1349,6 +1357,7 @@ void AccumulateSquare3x3Kernel(
 			SampleBufferUV,
 			SampleSceneMetadata,
 			MultiplexedSamples,
+			MultiplexedFrequencies,
 			/* KernelWeight = */ 1.0,
 			/* bForceSample = */ true,
 			bIsOutsideFrustum);
@@ -1360,7 +1369,8 @@ void AccumulateSquare3x3Kernel(
 			RefSceneMetadata,
 			SampleBufferUV,
 			SampleSceneMetadata,
-			MultiplexedSamples);
+			MultiplexedSamples,
+			MultiplexedFrequencies);
 	}
 
 	// Store whether needs to flip offsets to have lowest VGPR pressure.
@@ -1383,7 +1393,6 @@ void AccumulateSquare3x3Kernel(
 			/* inout */ UncompressedAccumulators,
 			/* inout */ CompressedAccumulators,
 			SampleBufferUV,
-			/* MipLevel = */ 0.0,
 			/* KernelWeight = */ 1.0,
 			/* bForceSample = */ false);
 	}
@@ -1683,7 +1692,6 @@ void ConvolveStackowiakKernel(
 					/* inout */ UncompressedAccumulators,
 					/* inout */ CompressedAccumulators,
 					SampleBufferUV,
-					/* MipLevel = */ 0.0,
 					/* KernelWeight = */ 1.0,
 					/* bForceSample = */ bIsKernelCenterSample && KernelConfig.bForceKernelCenterAccumulation);
 
@@ -1772,7 +1780,6 @@ void ConvolveStackowiakKernel(
 				/* inout */ UncompressedAccumulators,
 				/* inout */ CompressedAccumulators,
 				SampleBufferUV,
-				/* MipLevel = */ 0.0,
 				/* KernelWeight = */ 1.0,
 				/* bForceSample = */ false);
 		}
@@ -2067,91 +2074,6 @@ void ConvolveDiskKernel(
 #endif // COMPILE_DISK_KERNEL
 
 
-//------------------------------------------------------- STOCASTIC HIERARCHY
-
-#if COMPILE_HEIRARCHY_KERNEL
-
-void ConvolveStocasticHierarchy(
-	FSSDKernelConfig KernelConfig,
-	FSSDTexture2D SignalBuffer0,
-	FSSDTexture2D SignalBuffer1,
-	FSSDTexture2D SignalBuffer2,
-	FSSDTexture2D SignalBuffer3,
-	inout FSSDSignalAccumulatorArray UncompressedAccumulators,
-	inout FSSDCompressedSignalAccumulatorArray CompressedAccumulators)
-{
-	// Accumulate the center of the kernel.
-	if (KernelConfig.bSampleKernelCenter && !KernelConfig.bDescOrder)
-	{
-		SampleAndAccumulateCenterSampleAsItsOwnCluster(
-			KernelConfig,
-			SignalBuffer0,
-			SignalBuffer1,
-			SignalBuffer2,
-			SignalBuffer3,
-			/* inout */ UncompressedAccumulators,
-			/* inout */ CompressedAccumulators);
-	}
-
-	UNROLL_N(1)
-	for (uint MipLevel = 1; MipLevel < 2; MipLevel++)
-	{
-		uint SampleTrackId = KernelConfig.SampleTrackId;
-		#if CONFIG_VGPR_FREE_SAMPLE_TRACK_ID
-			SampleTrackId = GetSampleTrackIdFromLaneIndex();
-		#endif
-
-		SampleTrackId ^= View.StateFrameIndex;
-
-		float2 BasisX;
-		BasisX.x = float((~SampleTrackId) & 0x1) * (SampleTrackId & 0x2 ? -1.0 : 1.0);
-		BasisX.y = float(SampleTrackId & 0x1) * (SampleTrackId & 0x2 ? -1.0 : 1.0);
-
-		float2 BasisY = float2(-BasisX.y, BasisX.x);
-
-		float2 SampleTexelOffset[2];
-		SampleTexelOffset[0] = BasisX;
-		SampleTexelOffset[1] = -BasisX + BasisY;
-		
-		float MipLevelSize = float(1 << MipLevel);
-
-		const uint SampleCount = 2;
-		UNROLL_N(2)
-		for (uint SampleId = 0; SampleId < SampleCount; SampleId++)
-		{
-			float2 SampleBufferUV = KernelConfig.BufferUV + (SampleTexelOffset[SampleId] * MipLevelSize * KernelConfig.KernelSpreadFactor) * KernelConfig.BufferSizeAndInvSize.zw;
-
-			SampleAndAccumulateMultiplexedSignals(
-				KernelConfig,
-				SignalBuffer0,
-				SignalBuffer1,
-				SignalBuffer2,
-				SignalBuffer3,
-				/* inout */ UncompressedAccumulators,
-				/* inout */ CompressedAccumulators,
-				SampleBufferUV,
-				/* MipLevel = */ float(MipLevel),
-				/* KernelSampleWeight = */ 1.0,
-				/* bForceSample = */ false);
-		}
-	}
-	
-	// Accumulate the center of the kernel.
-	if (KernelConfig.bSampleKernelCenter && KernelConfig.bDescOrder)
-	{
-		SampleAndAccumulateCenterSampleAsItsOwnCluster(
-			KernelConfig,
-			SignalBuffer0,
-			SignalBuffer1,
-			SignalBuffer2,
-			SignalBuffer3,
-			/* inout */ UncompressedAccumulators,
-			/* inout */ CompressedAccumulators);
-	}
-}
-
-#endif // COMPILE_HEIRARCHY_KERNEL
-
 //------------------------------------------------------- DIRECTIONAL KERNELS
 
 #if COMPILE_DIRECTIONAL_KERNEL
@@ -2243,7 +2165,8 @@ void AccumulateRefSampleAsKernelCenter(
 	inout FSSDCompressedSignalAccumulatorArray CompressedAccumulators,
 	float2 RefBufferUV,
 	FSSDSampleSceneInfos RefSceneMetadata,
-	FSSDSignalArray RefMultiplexedSamples)
+	FSSDSignalArray RefMultiplexedSamples,
+	FSSDSignalFrequencyArray RefMultiplexedFrequencies)
 {
 	if (!KernelConfig.bSampleKernelCenter)
 	{
@@ -2255,6 +2178,7 @@ void AccumulateRefSampleAsKernelCenter(
 			RefBufferUV,
 			RefSceneMetadata,
 			RefMultiplexedSamples,
+			RefMultiplexedFrequencies,
 			/* KernelWeight = */ 1.0,
 			/* bForceSample = */ true,
 			/* bIsOutsideFrustum = */ false);
@@ -2274,7 +2198,8 @@ void AccumulateRefSampleAsKernelCenter(
 				RefSceneMetadata,
 				RefBufferUV,
 				RefSceneMetadata,
-				RefMultiplexedSamples);
+				RefMultiplexedSamples,
+				RefMultiplexedFrequencies);
 		}
 #endif
 	}
@@ -2376,19 +2301,6 @@ void AccumulateKernel(
 			/* inout */ CompressedAccumulators);
 	}
 #endif // COMPILE_DISK_KERNEL
-#if COMPILE_HEIRARCHY_KERNEL
-	else if (KernelConfig.SampleSet == SAMPLE_SET_STOCASTIC_HIERARCHY)
-	{
-		ConvolveStocasticHierarchy(
-			KernelConfig,
-			SignalBuffer0,
-			SignalBuffer1,
-			SignalBuffer2,
-			SignalBuffer3,
-			/* inout */ UncompressedAccumulators,
-			/* inout */ CompressedAccumulators);
-	}
-#endif // COMPILE_HEIRARCHY_KERNEL
 #if COMPILE_DIRECTIONAL_KERNEL
 	else if (KernelConfig.SampleSet == SAMPLE_SET_DIRECTIONAL_RECT || KernelConfig.SampleSet == SAMPLE_SET_DIRECTIONAL_ELLIPSE)
 	{
diff --git a/Engine/Shaders/Private/ScreenSpaceDenoise/SSDTemporalAccumulation.usf b/Engine/Shaders/Private/ScreenSpaceDenoise/SSDTemporalAccumulation.usf
--- a/Engine/Shaders/Private/ScreenSpaceDenoise/SSDTemporalAccumulation.usf
+++ b/Engine/Shaders/Private/ScreenSpaceDenoise/SSDTemporalAccumulation.usf
@@ -291,7 +291,7 @@
 #define COMPILE_MOMENT1_ACCUMULATOR 1
 
 #if CONFIG_HISTORY_REJECTION != HISTORY_REJECTION_DISABLED
-	#define COMPILE_MININVFREQ_ACCUMULATOR 1
+	#define COMPILE_MIN_FREQUENCY_ACCUMULATOR 1
 #endif
 
 #if CONFIG_HISTORY_REJECTION == HISTORY_REJECTION_MINMAX_BOUNDARIES
@@ -476,7 +476,9 @@ void TemporallyAccumulate(
 	}
 
 
-	FSSDSignalArray CurrentFrameSamples = SampleMultiplexedSignals(
+	FSSDSignalArray CurrentFrameSamples;
+	FSSDSignalFrequencyArray CurrentFrameFrequencies;
+	SampleMultiplexedSignals(
 		SignalInput_Textures_0,
 		SignalInput_Textures_1,
 		SignalInput_Textures_2,
@@ -486,8 +488,8 @@ void TemporallyAccumulate(
 		/* MultiplexedSampleId = */ 0,
 		/* bNormalizeSample = */ false,
 		SceneBufferUV,
-		/* MipLevel = */ 0);
-
+		/* out */ CurrentFrameSamples,
+		/* out */ CurrentFrameFrequencies);
 
 	float WorldBluringRadius;
 	float TargetedSampleCount;
@@ -643,12 +645,6 @@ void TemporallyAccumulate(
 				FSSDSignalSample CurrentFrameSample = CurrentFrameSamples.Array[BatchedSignalId];
 				FSSDSignalSample HistorySample = HistorySamples.Array[BatchedSignalId];
 
-				// Override the previous frame's inverse frequency with the current frame's one, because don't know how much it changed.
-				HistorySamples.Array[BatchedSignalId].WorldBluringRadius = (
-					HistorySample.SampleCount *
-					CurrentFrameSample.WorldBluringRadius *
-					SafeRcp(CurrentFrameSample.SampleCount));
-
 				// Applies correction when pre exposure change happen.
 				#if COMPILE_SIGNAL_COLOR
 					HistorySamples.Array[BatchedSignalId].SceneColor.rgb *= HistoryPreExposureCorrection;
@@ -787,7 +783,8 @@ void TemporallyAccumulate(
 				/* inout */ UnusedCompressedAccumulators,
 				KernelConfig.RefBufferUV,
 				RefSceneMetadata,
-				CenterSample);
+				CenterSample,
+				CurrentFrameFrequencies);
 		}
 
 		{
@@ -957,16 +954,9 @@ void TemporallyAccumulate(
 			// Dilate the smallest inv frequency.
 			if (1)
 			{
-				float NeightborMinBluringRadius = SignalAccumulators.Array[BatchedSignalId].MinInvFrequency;
-				CurrentFrameSample.WorldBluringRadius = min(CurrentFrameSample.WorldBluringRadius, NeightborMinBluringRadius * CurrentFrameSample.SampleCount); 
-			
-				#if CONFIG_SIGNAL_PROCESSING == SIGNAL_PROCESSING_SHADOW_VISIBILITY_MASK && 0
-					DebugOutput[DispatchThreadId] = float4(
-						NeightborMinBluringRadius,
-						NeightborMinBluringRadius == WORLD_RADIUS_MISS,
-						0,
-						0);
-				#endif
+				CurrentFrameFrequencies.Array[BatchedSignalId] = MinSignalFrequency(
+					CurrentFrameFrequencies.Array[BatchedSignalId],
+					SignalAccumulators.Array[BatchedSignalId].MinFrequency);
 			}
 			
 			HistorySamples.Array[BatchedSignalId] = HistorySample;
@@ -1004,12 +994,6 @@ void TemporallyAccumulate(
 				HistorySample,
 				/* SrcBasis = */ CONFIG_HISTORY_BLENDING_COLOR_SPACE,
 				/* DestBasis = */ STANDARD_BUFFER_COLOR_SPACE);
-			
-			// Ignore the bluring radius of previous frame frequencies as it gets clamped anyway, just take the current frame.
-			HistorySample.WorldBluringRadius = (
-				CurrentFrameSample.WorldBluringRadius *
-				SafeRcp(CurrentFrameSample.SampleCount) *
-				HistorySample.SampleCount);
 
 			//if (View.GeneralPurposeTweak > 0.5)
 			HistorySamples.Array[BatchedSignalId] = HistorySample;
@@ -1022,6 +1006,7 @@ void TemporallyAccumulate(
 	// Whitelists what should be outputed to ensure the compiler compiles out evreything that ends up unneeded.
 	uint MultiplexCount = 1;
 	FSSDSignalArray OutputSamples = CreateSignalArrayFromScalarValue(0.0);
+	FSSDSignalFrequencyArray OutputFrequencies = CreateInvalidSignalFrequencyArray();
 	{
 		MultiplexCount = CONFIG_SIGNAL_BATCH_SIZE;
 		
@@ -1029,6 +1014,7 @@ void TemporallyAccumulate(
 		for (uint BatchedSignalId = 0; BatchedSignalId < MultiplexCount; BatchedSignalId++)
 		{
 			OutputSamples.Array[BatchedSignalId] = HistorySamples.Array[BatchedSignalId];
+			OutputFrequencies.Array[BatchedSignalId] = CurrentFrameFrequencies.Array[BatchedSignalId];
 		}
 	}
 	
@@ -1043,8 +1029,11 @@ void TemporallyAccumulate(
 			SignalHistoryOutput_UAVs_1,
 			SignalHistoryOutput_UAVs_2,
 			SignalHistoryOutput_UAVs_3,
-			CONFIG_SIGNAL_HISTORY_LAYOUT, MultiplexCount,
-			OutputPixelPostion, OutputSamples.Array);
+			CONFIG_SIGNAL_HISTORY_LAYOUT,
+			MultiplexCount,
+			OutputPixelPostion,
+			OutputSamples,
+			OutputFrequencies);
 	}
 } // TemporallyAccumulate
 
diff --git a/Engine/Source/Runtime/Renderer/Private/ScreenSpaceDenoise.cpp b/Engine/Source/Runtime/Renderer/Private/ScreenSpaceDenoise.cpp
--- a/Engine/Source/Runtime/Renderer/Private/ScreenSpaceDenoise.cpp
+++ b/Engine/Source/Runtime/Renderer/Private/ScreenSpaceDenoise.cpp
@@ -220,12 +220,6 @@ static bool SignalUsesInjestion(ESignalProcessing SignalProcessing)
 	return SignalProcessing == ESignalProcessing::ShadowVisibilityMask;
 }
 
-/** Returns whether a signal processing uses a reduction pass before the reconstruction. */
-static bool SignalUsesReduction(ESignalProcessing SignalProcessing)
-{
-	return SignalProcessing == ESignalProcessing::DiffuseSphericalHarmonic;
-}
-
 /** Returns whether a signal processing uses an additional pre convolution pass. */
 static bool SignalUsesPreConvolution(ESignalProcessing SignalProcessing)
 {
@@ -367,49 +361,6 @@ const TCHAR* const kInjestResourceNames[] = {
 	nullptr,
 };
 
-const TCHAR* const kReduceResourceNames[] = {
-	// ShadowVisibilityMask
-	nullptr,
-	nullptr,
-	nullptr,
-	nullptr,
-
-	// PolychromaticPenumbraHarmonic
-	nullptr,
-	nullptr,
-	nullptr,
-	nullptr,
-
-	// Reflections
-	nullptr,
-	nullptr,
-	nullptr,
-	nullptr,
-
-	// AmbientOcclusion
-	nullptr,
-	nullptr,
-	nullptr,
-	nullptr,
-
-	// DiffuseIndirect
-	nullptr,
-	nullptr,
-	nullptr,
-	nullptr,
-
-	// DiffuseSphericalHarmonic
-	TEXT("DiffuseHarmonicReduce0"),
-	TEXT("DiffuseHarmonicReduce1"),
-	TEXT("DiffuseHarmonicReduce2"),
-	TEXT("DiffuseHarmonicReduce3"),
-
-	// ScreenSpaceDiffuseIndirect
-	nullptr,
-	nullptr,
-	nullptr,
-	nullptr,
-};
 
 const TCHAR* const kReconstructionResourceNames[] = {
 	// ShadowVisibilityMask
@@ -697,13 +648,16 @@ bool ShouldCompileSignalPipeline(ESignalProcessing SignalProcessing, EShaderPlat
 	}
 	else if (
 		SignalProcessing == ESignalProcessing::ShadowVisibilityMask ||
-		SignalProcessing == ESignalProcessing::PolychromaticPenumbraHarmonic ||
 		SignalProcessing == ESignalProcessing::AmbientOcclusion ||
 		SignalProcessing == ESignalProcessing::DiffuseAndAmbientOcclusion)
 	{
 		// Only for ray tracing denoising.
 		return RHISupportsRayTracingShaders(Platform);
 	}
+	else if (SignalProcessing == ESignalProcessing::PolychromaticPenumbraHarmonic)
+	{
+		return false;
+	}
 	check(0);
 	return false;
 }
@@ -842,7 +796,6 @@ class FSSDCompressMetadataCS : public FGlobalShader
 	END_SHADER_PARAMETER_STRUCT()
 };
 
-// TOOD: collapse within reduce pass.
 class FSSDInjestCS : public FGlobalShader
 {
 	DECLARE_GLOBAL_SHADER(FSSDInjestCS);
@@ -904,67 +857,6 @@ class FSSDInjestCS : public FGlobalShader
 	END_SHADER_PARAMETER_STRUCT()
 };
 
-class FSSDReduceCS : public FGlobalShader
-{
-	DECLARE_GLOBAL_SHADER(FSSDReduceCS);
-	SHADER_USE_PARAMETER_STRUCT(FSSDReduceCS, FGlobalShader);
-
-	using FPermutationDomain = TShaderPermutationDomain<FSignalProcessingDim, FSignalBatchSizeDim, FMultiSPPDim>;
-
-	static bool ShouldCompilePermutation(const FGlobalShaderPermutationParameters& Parameters)
-	{
-		FPermutationDomain PermutationVector(Parameters.PermutationId);
-		ESignalProcessing SignalProcessing = PermutationVector.Get<FSignalProcessingDim>();
-
-		// Only compile this shader for signal processing that uses it.
-		if (!SignalUsesReduction(SignalProcessing))
-		{
-			return false;
-		}
-
-		// Not all signal processing allow to batch multiple signals at the same time.
-		if (PermutationVector.Get<FSignalBatchSizeDim>() > SignalMaxBatchSize(SignalProcessing))
-		{
-			return false;
-		}
-
-		// Only compiler multi SPP permutation for signal that supports it.
-		if (PermutationVector.Get<FMultiSPPDim>() && !SignalSupportMultiSPP(SignalProcessing))
-		{
-			return false;
-		}
-
-		// Compile out the shader if this permutation gets remapped.
-		if (RemapPermutationVector(PermutationVector) != PermutationVector)
-		{
-			return false;
-		}
-
-		return ShouldCompileSignalPipeline(SignalProcessing, Parameters.Platform);
-	}
-
-	static FPermutationDomain RemapPermutationVector(FPermutationDomain PermutationVector)
-	{
-		ESignalProcessing SignalProcessing = PermutationVector.Get<FSignalProcessingDim>();
-
-		// force use the multi sample per pixel code path.
-		if (!SignalSupport1SPP(SignalProcessing))
-		{
-			PermutationVector.Set<FMultiSPPDim>(true);
-		}
-
-		return PermutationVector;
-	}
-
-	BEGIN_SHADER_PARAMETER_STRUCT(FParameters, )
-		SHADER_PARAMETER_STRUCT_INCLUDE(FSSDCommonParameters, CommonParameters)
-		SHADER_PARAMETER_STRUCT_INCLUDE(FSSDConvolutionMetaData, ConvolutionMetaData)
-
-		SHADER_PARAMETER_STRUCT(FSSDSignalTextures, SignalInput)
-		SHADER_PARAMETER_STRUCT_ARRAY(FSSDSignalUAVs, SignalOutputMips, [kMaxMipLevel])
-	END_SHADER_PARAMETER_STRUCT()
-};
-
 class FSSDSpatialAccumulationCS : public FGlobalShader
 {
 	DECLARE_GLOBAL_SHADER(FSSDSpatialAccumulationCS);
@@ -1185,7 +1077,6 @@ class FSSDComposeHarmonicsCS : public FGlobalShader
 
 IMPLEMENT_GLOBAL_SHADER(FSSDCompressMetadataCS, "/Engine/Private/ScreenSpaceDenoise/SSDCompressMetadata.usf", "MainCS", SF_Compute);
 IMPLEMENT_GLOBAL_SHADER(FSSDInjestCS, "/Engine/Private/ScreenSpaceDenoise/SSDInjest.usf", "MainCS", SF_Compute);
-IMPLEMENT_GLOBAL_SHADER(FSSDReduceCS, "/Engine/Private/ScreenSpaceDenoise/SSDReduce.usf", "MainCS", SF_Compute);
 IMPLEMENT_GLOBAL_SHADER(FSSDSpatialAccumulationCS, "/Engine/Private/ScreenSpaceDenoise/SSDSpatialAccumulation.usf", "MainCS", SF_Compute);
 IMPLEMENT_GLOBAL_SHADER(FSSDTemporalAccumulationCS, "/Engine/Private/ScreenSpaceDenoise/SSDTemporalAccumulation.usf", "MainCS", SF_Compute);
 IMPLEMENT_GLOBAL_SHADER(FSSDComposeHarmonicsCS, "/Engine/Private/ScreenSpaceDenoise/SSDComposeHarmonics.usf", "MainCS", SF_Compute);
@@ -1314,7 +1205,6 @@ static void DenoiseSignalAtConstantPixelDensity(
 	// Descriptor to allocate internal denoising buffer.
 	bool bHasReconstructionLayoutDifferentFromHistory = false;
 	TStaticArray<FRDGTextureDesc, kMaxBufferProcessingCount> InjestDescs;
-	TStaticArray<FRDGTextureDesc, kMaxBufferProcessingCount> ReduceDescs;
 	TStaticArray<FRDGTextureDesc, kMaxBufferProcessingCount> ReconstructionDescs;
 	TStaticArray<FRDGTextureDesc, kMaxBufferProcessingCount> HistoryDescs;
 	FRDGTextureDesc DebugDesc;
@@ -1344,8 +1234,6 @@ static void DenoiseSignalAtConstantPixelDensity(
 		for (int32 i = 0; i < kMaxBufferProcessingCount; i++)
 		{
 			InjestDescs[i] = RefDesc;
-			ReduceDescs[i] = RefDesc;
-			ReduceDescs[i].NumMips = kMaxMipLevel;
 			ReconstructionDescs[i] = RefDesc;
 			HistoryDescs[i] = RefDesc;
 		}
@@ -1405,12 +1293,10 @@ static void DenoiseSignalAtConstantPixelDensity(
 		{
 			for (int32 i = 0; i < 3; i++)
 			{
-				ReduceDescs[i].Format = PF_G32R32F;
 				ReconstructionDescs[i].Format = PF_G32R32F;
 				HistoryDescs[i].Format = PF_G32R32F;
 			}
 
-			ReduceDescs[3].Format = PF_R32_FLOAT;
 			ReconstructionDescs[3].Format = PF_R32_FLOAT;
 			HistoryDescs[3].Format = PF_R32_FLOAT;
 
@@ -1636,40 +1522,6 @@ static void DenoiseSignalAtConstantPixelDensity(
 		SignalHistory = NewSignalOutput;
 	}
 
-	// Reduce the input to speed up the reconstruction pass.
-	if (SignalUsesReduction(Settings.SignalProcessing))
-	{
-		FSSDSignalTextures NewSignalOutput = CreateMultiplexedTextures(
-			GraphBuilder,
-			InjestTextureCount, ReduceDescs,
-			GetResourceNames(kReduceResourceNames));
-
-		FSSDReduceCS::FParameters* PassParameters = GraphBuilder.AllocParameters<FSSDReduceCS::FParameters>();
-		PassParameters->CommonParameters = CommonParameters;
-		PassParameters->ConvolutionMetaData = ConvolutionMetaData;
-		PassParameters->SignalInput = SignalHistory;
-		for (int32 MipLevel = 0; MipLevel < kMaxMipLevel; MipLevel++)
-			PassParameters->SignalOutputMips[MipLevel] = CreateMultiplexedUAVs(GraphBuilder, NewSignalOutput, MipLevel);
-
-		FSSDReduceCS::FPermutationDomain PermutationVector;
-		PermutationVector.Set<FSignalProcessingDim>(Settings.SignalProcessing);
-		PermutationVector.Set<FSignalBatchSizeDim>(Settings.SignalBatchSize);
-		PermutationVector.Set<FMultiSPPDim>(bUseMultiInputSPPShaderPath);
-		PermutationVector = FSSDReduceCS::RemapPermutationVector(PermutationVector);
-
-		TShaderMapRef<FSSDReduceCS> ComputeShader(View.ShaderMap, PermutationVector);
-		FComputeShaderUtils::AddPass(
-			GraphBuilder,
-			RDG_EVENT_NAME("SSD Reduce(Mips=%i MultiSPP=%i)",
-				kMaxMipLevel,
-				int32(PermutationVector.Get<FMultiSPPDim>())),
-			ComputeShader,
-			PassParameters,
-			FComputeShaderUtils::GetGroupCount(Viewport.Size(), FComputeShaderUtils::kGolden2DGroupSize));
-
-		SignalHistory = NewSignalOutput;
-	}
-
 	// Spatial reconstruction with ratio estimator to be more precise in the history rejection.
 	if (Settings.bEnableReconstruction)
 	{
@@ -1711,7 +1563,7 @@ static void DenoiseSignalAtConstantPixelDensity(
 		TShaderMapRef<FSSDSpatialAccumulationCS> ComputeShader(View.ShaderMap, PermutationVector);
 		FComputeShaderUtils::AddPass(
 			GraphBuilder,
-			RDG_EVENT_NAME("SSD SpatialAccumulation(Reconstruction MaxSamples=%i Scissor=%ix%i%s%s)",
+			RDG_EVENT_NAME("SSD Reconstruction(MaxSamples=%i Scissor=%ix%i%s%s)",
 				PassParameters->MaxSampleCount,
 				Viewport.Width(), Viewport.Height(),
 				PermutationVector.Get<FSSDSpatialAccumulationCS::FUpscaleDim>() ? TEXT(" Upscale") : TEXT(""),
