Fixes history reprojection issues in the denoiser

diff --git a/Engine/Shaders/Private/ScreenSpaceDenoise/SSDTemporalAccumulation.usf b/Engine/Shaders/Private/ScreenSpaceDenoise/SSDTemporalAccumulation.usf
--- a/Engine/Shaders/Private/ScreenSpaceDenoise/SSDTemporalAccumulation.usf
+++ b/Engine/Shaders/Private/ScreenSpaceDenoise/SSDTemporalAccumulation.usf
@@ -490,8 +490,6 @@ void TemporallyAccumulate(
 		float PixelWorldBluringRadius = ComputeWorldBluringRadiusCausedByPixelSize(RefSceneMetadata);
 		WorldBluringRadius = WorldBluringRadiusToBilateralWorldDistance(PixelWorldBluringRadius);
 	
-		TargetedSampleCount = TARGETED_SAMPLE_COUNT;
-
 		#if CONFIG_SIGNAL_PROCESSING == SIGNAL_PROCESSING_REFLECTIONS
 		{
 			float2 NormalizedScreenMajorAxis;
@@ -521,7 +519,7 @@ void TemporallyAccumulate(
 	// Sample the history buffer.
 	FSSDSignalArray HistorySamples = CreateSignalArrayFromScalarValue(0.0);
 	{
-		float2 HistoryBufferUV = HistoryScreenPosition.xy * ScreenPosToHistoryBufferUV.xy + ScreenPosToHistoryBufferUV.wz;
+		float2 HistoryBufferUV = HistoryScreenPosition.xy* ScreenPosToHistoryBufferUV.xy + ScreenPosToHistoryBufferUV.zw;
 		float2 ClampedHistoryBufferUV = clamp(HistoryBufferUV, HistoryBufferUVMinMax.xy, HistoryBufferUVMinMax.zw);
 		bool bIsPreviousFrameOffscreen = any(HistoryBufferUV != ClampedHistoryBufferUV);
 
