Fixes computation of ScreenPos when sampling screen space denoiser's history to handle resolution changes

diff --git a/Engine/Shaders/Private/ScreenSpaceDenoise/SSDCommon.ush b/Engine/Shaders/Private/ScreenSpaceDenoise/SSDCommon.ush
--- a/Engine/Shaders/Private/ScreenSpaceDenoise/SSDCommon.ush
+++ b/Engine/Shaders/Private/ScreenSpaceDenoise/SSDCommon.ush
@@ -80,6 +80,9 @@ float4 BufferUVToScreenPosition;
 float4x4 ScreenToView;
 float2 BufferUVBilinearCorrection;
 
+
+float4 PrevSceneBufferUVToScreenPosition;
+
 float WorldDepthToPixelWorldRadius;
 
 Texture2D PrevDepthBuffer;
diff --git a/Engine/Shaders/Private/ScreenSpaceDenoise/SSDSpatialKernel.ush b/Engine/Shaders/Private/ScreenSpaceDenoise/SSDSpatialKernel.ush
--- a/Engine/Shaders/Private/ScreenSpaceDenoise/SSDSpatialKernel.ush
+++ b/Engine/Shaders/Private/ScreenSpaceDenoise/SSDSpatialKernel.ush
@@ -416,10 +416,20 @@ FSSDSampleSceneInfos UncompressRefSceneMetadata(FSSDKernelConfig KernelConfig)
 	// Find out the buffer UV of the reference pixel.
 	float2 RefBufferUV = ComputeRefBufferUV(KernelConfig);
 
+	float2 ScreenPos;
+	if (KernelConfig.bPreviousFrameMetadata)
+	{
+		ScreenPos = RefBufferUV * PrevSceneBufferUVToScreenPosition.xy + PrevSceneBufferUVToScreenPosition.zw;
+	}
+	else
+	{
+		ScreenPos = DenoiserBufferUVToScreenPosition(RefBufferUV);
+	}
+
 	// Uncompress the reference scene metadata to keep a low VGPR pressure.
 	return UncompressSampleSceneInfo(
 		KernelConfig.RefSceneMetadataLayout, /* bIsPrevFrame = */ false,
-		DenoiserBufferUVToScreenPosition(RefBufferUV),
+		ScreenPos,
 		KernelConfig.CompressedRefSceneMetadata);
 }
 
@@ -429,9 +439,19 @@ FSSDSampleSceneInfos UncompressSampleSceneMetadata(
 	float2 SampleBufferUV,
 	FSSDCompressedSceneInfos CompressedSampleSceneMetadata)
 {
+	float2 ScreenPos;
+	if (KernelConfig.bPreviousFrameMetadata)
+	{
+		ScreenPos = SampleBufferUV * PrevSceneBufferUVToScreenPosition.xy + PrevSceneBufferUVToScreenPosition.zw;
+	}
+	else
+	{
+		ScreenPos = DenoiserBufferUVToScreenPosition(SampleBufferUV);
+	}
+
 	return UncompressSampleSceneInfo(
 		CONFIG_METADATA_BUFFER_LAYOUT, KernelConfig.bPreviousFrameMetadata,
-		DenoiserBufferUVToScreenPosition(SampleBufferUV),
+		ScreenPos,
 		CompressedSampleSceneMetadata);
 }
 
diff --git a/Engine/Source/Runtime/Renderer/Private/ScreenSpaceDenoise.cpp b/Engine/Source/Runtime/Renderer/Private/ScreenSpaceDenoise.cpp
--- a/Engine/Source/Runtime/Renderer/Private/ScreenSpaceDenoise.cpp
+++ b/Engine/Source/Runtime/Renderer/Private/ScreenSpaceDenoise.cpp
@@ -1038,6 +1038,7 @@ class FSSDTemporalAccumulationCS : public FGlobalShader
 		SHADER_PARAMETER(FVector4, HistoryBufferSizeAndInvSize)
 		SHADER_PARAMETER(FVector4, HistoryBufferUVMinMax)
 		SHADER_PARAMETER_ARRAY(FVector4, HistoryBufferScissorUVMinMax, [IScreenSpaceDenoiser::kMaxBatchSize])
+		SHADER_PARAMETER(FVector4, PrevSceneBufferUVToScreenPosition)
 
 		SHADER_PARAMETER_STRUCT_INCLUDE(FSSDCommonParameters, CommonParameters)
 		SHADER_PARAMETER_STRUCT_INCLUDE(FSSDConvolutionMetaData, ConvolutionMetaData)
@@ -1762,6 +1763,11 @@ static void DenoiseSignalAtConstantPixelDensity(
 				(ViewportOffset.Y + ViewportExtent.Y - 0.5f) * InvBufferSizeY);
 
 			PassParameters->HistoryBufferSizeAndInvSize = FVector4(PrevFrameBufferExtent.X, PrevFrameBufferExtent.Y, InvBufferSizeX, InvBufferSizeY);
+
+			PassParameters->PrevSceneBufferUVToScreenPosition.X = float(PrevFrameBufferExtent.X) / float(ViewportExtent.X) * 2.0f;
+			PassParameters->PrevSceneBufferUVToScreenPosition.Y = -float(PrevFrameBufferExtent.Y) / float(ViewportExtent.Y) * 2.0f;
+			PassParameters->PrevSceneBufferUVToScreenPosition.Z = float(ViewportOffset.X) / float(ViewportExtent.X) * 2.0f - 1.0f;
+			PassParameters->PrevSceneBufferUVToScreenPosition.W = -float(ViewportOffset.Y) / float(ViewportExtent.Y) * 2.0f + 1.0f;
 		}
 
 		FScreenSpaceDenoiserHistory DummyPrevFrameHistory;
