Increases the maximum number accumulatable samples for reflection denoiser

diff --git a/Engine/Shaders/Private/ScreenSpaceDenoise/SSDCommon.ush b/Engine/Shaders/Private/ScreenSpaceDenoise/SSDCommon.ush
--- a/Engine/Shaders/Private/ScreenSpaceDenoise/SSDCommon.ush
+++ b/Engine/Shaders/Private/ScreenSpaceDenoise/SSDCommon.ush
@@ -28,7 +28,7 @@
 //------------------------------------------------------- SHARED CONSTANTS DEPENDING ON CONFIG OF SHADER
 
 #if CONFIG_SIGNAL_PROCESSING == SIGNAL_PROCESSING_REFLECTIONS
-	#define TARGETED_SAMPLE_COUNT 128
+	#define TARGETED_SAMPLE_COUNT 1024
 
 #elif CONFIG_SIGNAL_PROCESSING == SIGNAL_PROCESSING_AO || CONFIG_SIGNAL_PROCESSING == SIGNAL_PROCESSING_SHADOW_VISIBILITY_MASK
 	// A gray scale valued encoded into a 16bit float only have 10bits mantissa.
