Fixes SM6 compatible platforms

diff --git a/Engine/Shaders/Private/ScreenSpaceDenoise/SSDSpatialKernel.ush b/Engine/Shaders/Private/ScreenSpaceDenoise/SSDSpatialKernel.ush
--- a/Engine/Shaders/Private/ScreenSpaceDenoise/SSDSpatialKernel.ush
+++ b/Engine/Shaders/Private/ScreenSpaceDenoise/SSDSpatialKernel.ush
@@ -1255,9 +1255,7 @@ void BroadcastAccumulateSquare3x3KernelCenter(
 	SampleBufferUV = WaveBroadcast(BroadcastSettingsX, SampleBufferUV);
 	SampleSceneMetadata = WaveBroadcastSceneMetadata(BroadcastSettingsX, SampleSceneMetadata);
 	SampleMultiplexedSamples = WaveBroadcastSignalArray(BroadcastSettingsX, SampleMultiplexedSamples);
-	SampleMultiplexedFrequencies = WaveBroadcastSignalFrequencyArray(BroadcastSettingsX, SampleMultiplexedFrequencies);
-
-	FSSDSignalFrequency WaveBroadcastSignalFrequency(const FWaveBroadcastSettings BroadcastSettings, FSSDSignalFrequency Frequency)
+	SampleMultiplexedFrequencies = WaveBroadcastSignalFrequenciesArray(BroadcastSettingsX, SampleMultiplexedFrequencies);
 
 	if (KernelConfig.SampleSet == SAMPLE_SET_3X3 || KernelConfig.SampleSet == SAMPLE_SET_3X3_PLUS)
 	{
@@ -1279,7 +1277,7 @@ void BroadcastAccumulateSquare3x3KernelCenter(
 	SampleBufferUV = WaveBroadcast(BroadcastSettingsY, SampleBufferUV);
 	SampleSceneMetadata = WaveBroadcastSceneMetadata(BroadcastSettingsY, SampleSceneMetadata);
 	SampleMultiplexedSamples = WaveBroadcastSignalArray(BroadcastSettingsY, SampleMultiplexedSamples);
-	SampleMultiplexedFrequencies = WaveBroadcastSignalFrequencyArray(BroadcastSettingsY, SampleMultiplexedFrequencies);
+	SampleMultiplexedFrequencies = WaveBroadcastSignalFrequenciesArray(BroadcastSettingsY, SampleMultiplexedFrequencies);
 
 	if (KernelConfig.SampleSet == SAMPLE_SET_3X3 || KernelConfig.SampleSet == SAMPLE_SET_3X3_CROSS)
 	{
@@ -1301,7 +1299,7 @@ void BroadcastAccumulateSquare3x3KernelCenter(
 	SampleBufferUV = WaveBroadcast(BroadcastSettingsX, SampleBufferUV);
 	SampleSceneMetadata = WaveBroadcastSceneMetadata(BroadcastSettingsX, SampleSceneMetadata);
 	SampleMultiplexedSamples = WaveBroadcastSignalArray(BroadcastSettingsX, SampleMultiplexedSamples);
-	SampleMultiplexedFrequencies = WaveBroadcastSignalFrequencyArray(BroadcastSettingsX, SampleMultiplexedFrequencies);
+	SampleMultiplexedFrequencies = WaveBroadcastSignalFrequenciesArray(BroadcastSettingsX, SampleMultiplexedFrequencies);
 	
 	if (KernelConfig.SampleSet == SAMPLE_SET_3X3 || KernelConfig.SampleSet == SAMPLE_SET_3X3_PLUS)
 	{
@@ -1362,7 +1360,9 @@ void AccumulateSquare3x3Kernel(
 		FSSDSignalArray MultiplexedSamples;
 		FSSDSignalFrequencyArray MultiplexedFrequencies;
 		UncompressMultiplexedSignals(
-			KernelConfig, SampleBufferUV, CompressedMultiplexedSamples);
+			KernelConfig,
+			SampleBufferUV,
+			CompressedMultiplexedSamples,
 			/* out */ MultiplexedSamples,
 			/* out */ MultiplexedFrequencies);
 
