Adding primitive lighting channel mask to ray tracing payload.

diff --git a/Engine/Shaders/Private/RayTracing/RayTracingCommon.ush b/Engine/Shaders/Private/RayTracing/RayTracingCommon.ush
--- a/Engine/Shaders/Private/RayTracing/RayTracingCommon.ush
+++ b/Engine/Shaders/Private/RayTracing/RayTracingCommon.ush
@@ -177,7 +177,7 @@ struct FMaterialClosestHitPayload : FMinimalPayload
 	float Ior;							// 88        2       16bits
 	uint ShadingModelID;				// 92        1       4bits
 	uint BlendingMode;					// 96        0       4bits (packed with ShadingModelID)
-	uint PrimitiveLightingChannelMask;	// 100       0       (removed)
+	uint PrimitiveLightingChannelMask;	// 100       0       3bits (packed with ShadingModelID)
 	float4 CustomData;					// 104       4       32bits
 	float GBufferAO;					// 120       0       (removed)
 	float3 IndirectIrradiance;			// 124       0       48bits -- gbuffer only has float payload and there are truncation HLSL warnings
@@ -186,7 +186,7 @@ struct FMaterialClosestHitPayload : FMinimalPayload
 	// So keep it here and force to re-construct it in Unpack call using ray information.
 	// It is not packed in FRayHitInfoPacked
 	float3 WorldPos;					// 136       0       (derived)
-	uint Flags;							// 148       4       32bits
+	uint Flags;							// 148       0       5bits (packed with ShadingModelID)
 	float3 WorldTangent;				// 152       6       48bits
 	float Anisotropy;					// 164       2       16bits (packed with WorldTangent)
 										// 166 total
@@ -227,17 +227,17 @@ struct FPackedMaterialClosestHitPayload : FMinimalPayload
 	uint RadianceAndNormal[3];                           // 12 bytes
 	uint BaseColorAndOpacity[2];                         // 8  bytes
 	uint MetallicAndSpecularAndRoughness;                // 4  bytes
-	uint IorAndShadingModelIDAndBlendingModeAndFlags;    // 4  bytes
+	uint IorAndShadingModelIDAndBlendingModeAndFlagsAndPrimitiveLightingChannelMask;    // 4  bytes
 	uint PackedIndirectIrradiance[2];                    // 8  bytes
 	uint PackedCustomData;                               // 4  bytes
 	uint WorldTangentAndAnisotropy[2];					 // 8  bytes
 	                                                     // 60 bytes total
 
-	void SetMinimalPayloadMode() { IorAndShadingModelIDAndBlendingModeAndFlags |= RAY_TRACING_PAYLOAD_FLAG_MINIMAL_PAYLOAD << 24; }
-	bool IsMinimalPayloadMode() { return ((IorAndShadingModelIDAndBlendingModeAndFlags >> 24) & RAY_TRACING_PAYLOAD_FLAG_MINIMAL_PAYLOAD) != 0; }
+	void SetMinimalPayloadMode() { IorAndShadingModelIDAndBlendingModeAndFlagsAndPrimitiveLightingChannelMask |= RAY_TRACING_PAYLOAD_FLAG_MINIMAL_PAYLOAD << 24; }
+	bool IsMinimalPayloadMode() { return ((IorAndShadingModelIDAndBlendingModeAndFlagsAndPrimitiveLightingChannelMask >> 24) & RAY_TRACING_PAYLOAD_FLAG_MINIMAL_PAYLOAD) != 0; }
 
-	void SetEnableSkyLightContribution() { IorAndShadingModelIDAndBlendingModeAndFlags |= RAY_TRACING_PAYLOAD_FLAG_ENABLE_SKY_LIGHT_CONTRIBUTION << 24; }
-	bool IsEnableSkyLightContribution() { return ((IorAndShadingModelIDAndBlendingModeAndFlags >> 24) & RAY_TRACING_PAYLOAD_FLAG_ENABLE_SKY_LIGHT_CONTRIBUTION) != 0; }
+	void SetEnableSkyLightContribution() { IorAndShadingModelIDAndBlendingModeAndFlagsAndPrimitiveLightingChannelMask |= RAY_TRACING_PAYLOAD_FLAG_ENABLE_SKY_LIGHT_CONTRIBUTION << 24; }
+	bool IsEnableSkyLightContribution() { return ((IorAndShadingModelIDAndBlendingModeAndFlagsAndPrimitiveLightingChannelMask >> 24) & RAY_TRACING_PAYLOAD_FLAG_ENABLE_SKY_LIGHT_CONTRIBUTION) != 0; }
 
 	void SetRayCone(FRayCone InRayCone) { RayCone = InRayCone; }
 
@@ -315,10 +315,11 @@ struct FPackedMaterialClosestHitPayload : FMinimalPayload
 	float GetMetallic()      { return float(MetallicAndSpecularAndRoughness & 0xFF) / 255.0f; }
 	float GetSpecular()      { return float((MetallicAndSpecularAndRoughness >> 8) & 0xFF) / 255.0f; }
 	float GetRoughness()     { return f16tof32(MetallicAndSpecularAndRoughness >> 16); }
-	float GetIor()           { return f16tof32(IorAndShadingModelIDAndBlendingModeAndFlags); }
-	uint GetBlendingMode()   { return (IorAndShadingModelIDAndBlendingModeAndFlags >> 20) & 0xF;  }
-	uint GetFlags()          { return (IorAndShadingModelIDAndBlendingModeAndFlags >> 24) & 0xFF; }
-	uint GetShadingModelID() { return (IorAndShadingModelIDAndBlendingModeAndFlags >> 16) & 0xF; }
+	float GetIor()           { return f16tof32(IorAndShadingModelIDAndBlendingModeAndFlagsAndPrimitiveLightingChannelMask); }
+	uint GetShadingModelID() { return (IorAndShadingModelIDAndBlendingModeAndFlagsAndPrimitiveLightingChannelMask >> 16) & 0xF; }
+	uint GetBlendingMode()   { return (IorAndShadingModelIDAndBlendingModeAndFlagsAndPrimitiveLightingChannelMask >> 20) & 0xF; }
+	uint GetFlags()          { return (IorAndShadingModelIDAndBlendingModeAndFlagsAndPrimitiveLightingChannelMask >> 24) & 0x1F; }
+	uint GetPrimitiveLightingChannelMask() { return (IorAndShadingModelIDAndBlendingModeAndFlagsAndPrimitiveLightingChannelMask >> 29) & 0x7; }
 	float GetAnisotropy()    { return f16tof32(WorldTangentAndAnisotropy[1] >> 16); }
 
 	float3 GetDiffuseColor() { return GetBaseColor() - GetBaseColor() * GetMetallic(); }
@@ -344,10 +345,11 @@ FPackedMaterialClosestHitPayload PackRayTracingPayload(FMaterialClosestHitPayloa
 	Output.MetallicAndSpecularAndRoughness  = (uint(round(Input.Metallic * 255.0f)) & 0xFF);
 	Output.MetallicAndSpecularAndRoughness |= (uint(round(Input.Specular * 255.0f)) & 0xFF) << 8;
 	Output.MetallicAndSpecularAndRoughness |= f32tof16(Input.Roughness) << 16;
-	Output.IorAndShadingModelIDAndBlendingModeAndFlags  = f32tof16(Input.Ior);                // 16 bits
-	Output.IorAndShadingModelIDAndBlendingModeAndFlags |= (Input.ShadingModelID & 0xF) << 16; // 4 bits
-	Output.IorAndShadingModelIDAndBlendingModeAndFlags |= (Input.BlendingMode & 0xF) << 20;   // 4 bits
-	Output.IorAndShadingModelIDAndBlendingModeAndFlags |= (Input.Flags & 0xF) << 24;          // 8 bits
+	Output.IorAndShadingModelIDAndBlendingModeAndFlagsAndPrimitiveLightingChannelMask  = f32tof16(Input.Ior);                              // 16 bits
+	Output.IorAndShadingModelIDAndBlendingModeAndFlagsAndPrimitiveLightingChannelMask |= (Input.ShadingModelID & 0xF) << 16;               // 4 bits
+	Output.IorAndShadingModelIDAndBlendingModeAndFlagsAndPrimitiveLightingChannelMask |= (Input.BlendingMode & 0xF) << 20;                 // 4 bits
+	Output.IorAndShadingModelIDAndBlendingModeAndFlagsAndPrimitiveLightingChannelMask |= (Input.Flags & 0x1F) << 24;                       // 5 bits
+	Output.IorAndShadingModelIDAndBlendingModeAndFlagsAndPrimitiveLightingChannelMask |= (Input.PrimitiveLightingChannelMask & 0x7) << 29; // 3 bits
 	Output.PackedIndirectIrradiance[0] = f32tof16(Input.IndirectIrradiance.x);
 	Output.PackedIndirectIrradiance[0] |= f32tof16(Input.IndirectIrradiance.y) << 16;
 	Output.PackedIndirectIrradiance[1] = f32tof16(Input.IndirectIrradiance.z);
@@ -377,6 +379,7 @@ FMaterialClosestHitPayload UnpackRayTracingPayload(FPackedMaterialClosestHitPayl
 	Output.Ior                = Input.GetIor();
 	Output.ShadingModelID     = Input.GetShadingModelID();
 	Output.BlendingMode	      = Input.GetBlendingMode();
+	Output.PrimitiveLightingChannelMask = Input.GetPrimitiveLightingChannelMask();
 	Output.Flags              = Input.GetFlags();
 	Output.IndirectIrradiance = Input.GetIndirectIrradiance();
 	Output.CustomData         = Input.GetCustomData();
diff --git a/Engine/Shaders/Private/RayTracing/RayTracingDebug.usf b/Engine/Shaders/Private/RayTracing/RayTracingDebug.usf
--- a/Engine/Shaders/Private/RayTracing/RayTracingDebug.usf
+++ b/Engine/Shaders/Private/RayTracing/RayTracingDebug.usf
@@ -80,7 +80,7 @@ RAY_TRACING_ENTRY_RAYGEN(RayTracingDebugMainRGS)
 			Result = float4(Payload.BlendingMode, Payload.BlendingMode, Payload.BlendingMode, 1.0f);
 			break;
 		case RAY_TRACING_DEBUG_VIZ_LIGHTING_CHANNEL_MASK:
-			Result = float4(Payload.PrimitiveLightingChannelMask, Payload.PrimitiveLightingChannelMask, Payload.PrimitiveLightingChannelMask, 1.0f);
+			Result = float4(Payload.PrimitiveLightingChannelMask & 0x1, Payload.PrimitiveLightingChannelMask & 0x2, Payload.PrimitiveLightingChannelMask & 0x4, 1.0f);
 			break;
 		case RAY_TRACING_DEBUG_VIZ_CUSTOM_DATA:
 			Result = float4(Payload.CustomData);
diff --git a/Engine/Shaders/Private/RayTracing/RayTracingMaterialHitShaders.usf b/Engine/Shaders/Private/RayTracing/RayTracingMaterialHitShaders.usf
--- a/Engine/Shaders/Private/RayTracing/RayTracingMaterialHitShaders.usf
+++ b/Engine/Shaders/Private/RayTracing/RayTracingMaterialHitShaders.usf
@@ -330,6 +330,7 @@ RAY_TRACING_ENTRY_CLOSEST_HIT(MaterialCHS,
 	#endif
 
 	Payload.BlendingMode = MaterialBlendingMode;
+	Payload.PrimitiveLightingChannelMask = GetPrimitiveData(MaterialParameters.PrimitiveId).LightingChannelMask;
 
 	Payload.HitT = RayTCurrent();
 	if (HitKind() == HIT_KIND_TRIANGLE_FRONT_FACE)
