Fix for transforming cascade material transforms in the vertex shader

diff --git a/Engine/Shaders/Private/MeshParticleVertexFactory.ush b/Engine/Shaders/Private/MeshParticleVertexFactory.ush
--- a/Engine/Shaders/Private/MeshParticleVertexFactory.ush
+++ b/Engine/Shaders/Private/MeshParticleVertexFactory.ush
@@ -336,7 +336,7 @@ FMaterialVertexParameters GetMaterialVertexParameters(FVertexFactoryInput Input,
 #endif	//USE_DYNAMIC_PARAMETERS
 
 	float4x4 ParticleTransform = GetParticleTransform(Input);
-	Result.InstanceLocalToWorld = Primitive_LocalToWorld;
+	Result.InstanceLocalToWorld = ParticleTransform;
 
 #if NEEDS_PARTICLE_LOCAL_TO_WORLD
 	Result.Particle.ParticleToWorld = ParticleTransform;
