- workaround for cook issue.

diff --git a/Engine/Shaders/Private/HTileEncoding.ush b/Engine/Shaders/Private/HTileEncoding.ush
--- a/Engine/Shaders/Private/HTileEncoding.ush
+++ b/Engine/Shaders/Private/HTileEncoding.ush
@@ -67,7 +67,12 @@ uint EncodeTileMinMaxDepth(float MinDepth, float MaxDepth)
 
 	// Shift up min by 2 bits, then set all four low bits.
 	HTileValue = BitFieldInsertU32(BitFieldMaskU32(14, 18), HTileValue, HTileValue << 2);
-	return HTileValue |= HTILE_ZMASK_EXPANDED;
+	// TODO: needs fixing HTILE_ZMASK_EXPANDED is missing now
+	#ifdef HTILE_ZMASK_EXPANDED
+		return HTileValue |= HTILE_ZMASK_EXPANDED;
+	#else
+		return HTileValue;
+	#endif
 }
 
 #endif // COMPILER_SUPPORTS_PACK_INTRINSICS
