Minor cleanup to EncodeTileMinMaxDepth

diff --git a/Engine/Shaders/Private/HTileEncoding.ush b/Engine/Shaders/Private/HTileEncoding.ush
--- a/Engine/Shaders/Private/HTileEncoding.ush
+++ b/Engine/Shaders/Private/HTileEncoding.ush
@@ -69,13 +69,11 @@ uint EncodeTileMinMaxDepth(float MinDepth, float MaxDepth)
 
 	// Shift up min by 2 bits, then set all four low bits.
 	HTileValue = BitFieldInsertU32(BitFieldMaskU32(14, 18), HTileValue, HTileValue << 2);
-	// TODO: needs fixing HTILE_ZMASK_EXPANDED is missing now
-	#ifdef HTILE_ZMASK_EXPANDED
-		return HTileValue |= HTILE_ZMASK_EXPANDED;
-	#else
-		return HTileValue;
-	#endif
-}
+#ifdef HTILE_ZMASK_EXPANDED
+	HTileValue |= HTILE_ZMASK_EXPANDED;
+#endif
+
+	return HTileValue;
 
 #endif // COMPILER_SUPPORTS_PACK_INTRINSICS
 
