Revert reverb CL11806278 to fix some edge cases where nearer shadow cascades are stomping further ones data in the screen space shadow mask.

Will revisit the reverb issue and look for another fix.

diff --git a/Engine/Shaders/Private/ShadowProjectionPixelShader.usf b/Engine/Shaders/Private/ShadowProjectionPixelShader.usf
--- a/Engine/Shaders/Private/ShadowProjectionPixelShader.usf
+++ b/Engine/Shaders/Private/ShadowProjectionPixelShader.usf
@@ -144,13 +144,6 @@ void Main(
 	float ShadowZ = 1 - ShadowPosition.z;
 	ShadowPosition.xyz /= ShadowPosition.w;
 
-	if (any(ShadowPosition.xy > ShadowTileOffsetAndSize.xy+ShadowTileOffsetAndSize.zw) ||
-		any(ShadowPosition.xy < ShadowTileOffsetAndSize.xy))
-	{
-		OutColor = (float4)1;
-		return;
-	}
-
 #if MODULATED_SHADOWS
 	// UE-29083 : work around precision issues with ScreenToShadowMatrix on low end devices.
 	ShadowPosition.xy *= ShadowTileOffsetAndSize.zw;
