SS - Fixing Ps4 shader compile errors.

diff --git a/Engine/Shaders/Private/HTileEncoding.ush b/Engine/Shaders/Private/HTileEncoding.ush
--- a/Engine/Shaders/Private/HTileEncoding.ush
+++ b/Engine/Shaders/Private/HTileEncoding.ush
@@ -74,6 +74,7 @@ uint EncodeTileMinMaxDepth(float MinDepth, float MaxDepth)
 #endif
 
 	return HTileValue;
+}
 
 #endif // COMPILER_SUPPORTS_PACK_INTRINSICS
 
