Fixes some shader compilation failure on mac & linux in DOF, TAAU2, Denoiser, SSGI

diff --git a/Engine/Shaders/Private/DiaphragmDOF/DOFRecombine.usf b/Engine/Shaders/Private/DiaphragmDOF/DOFRecombine.usf
--- a/Engine/Shaders/Private/DiaphragmDOF/DOFRecombine.usf
+++ b/Engine/Shaders/Private/DiaphragmDOF/DOFRecombine.usf
@@ -886,7 +886,7 @@ void RecombineMainCS(
 			float RecommendedPairCount = KenelArea * 0.5;
 
 			// Number of pair of sample.
-			GatherParameters.SamplePairCount = min(CONFIG_GATHER_PAIR_COUNT, uint(RecommendedPairCount));
+			GatherParameters.SamplePairCount = min(uint(CONFIG_GATHER_PAIR_COUNT), uint(RecommendedPairCount));
 		}
 	}
 	#endif // CONFIG_SLIGHT_FOCUS_METHOD != SLIGHT_FOCUS_METHOD_DISABLED
diff --git a/Engine/Shaders/Private/SSRT/SSRTRayCast.ush b/Engine/Shaders/Private/SSRT/SSRTRayCast.ush
--- a/Engine/Shaders/Private/SSRT/SSRTRayCast.ush
+++ b/Engine/Shaders/Private/SSRT/SSRTRayCast.ush
@@ -195,7 +195,7 @@ FSSRTCastingSettings CreateDefaultCastSettings()
 	Settings.bStopWhenUncertain = false;
 	Settings.TileInfos = (FSSRTTileInfos)0;
 	return Settings;
-};
+}
 
 uint ComputeRayDirectionId(FSSRTRay Ray)
 {
diff --git a/Engine/Shaders/Private/ScreenSpaceDenoise/SSDCommon.ush b/Engine/Shaders/Private/ScreenSpaceDenoise/SSDCommon.ush
--- a/Engine/Shaders/Private/ScreenSpaceDenoise/SSDCommon.ush
+++ b/Engine/Shaders/Private/ScreenSpaceDenoise/SSDCommon.ush
@@ -163,12 +163,12 @@ float4x4 GetScreenToViewDistanceMatrix()
 	// NOP
 
 #elif SPHERICAL_HARMONIC_ORDER == 2
-typedef FTwoBandSHVector FSSDSphericalHarmonic;
-typedef FTwoBandSHVectorRGB FSSDSphericalHarmonicRGB;
+#define FSSDSphericalHarmonic FTwoBandSHVector
+#define FSSDSphericalHarmonicRGB FTwoBandSHVectorRGB
 
 #elif SPHERICAL_HARMONIC_ORDER == 3
-typedef FThreeBandSHVector FSSDSphericalHarmonic;
-typedef FThreeBandSHVectorRGB FSSDSphericalHarmonicRGB;
+#define FSSDSphericalHarmonic FThreeBandSHVector
+#define FSSDSphericalHarmonicRGB FThreeBandSHVectorRGB
 
 #else
 	#error Unknown spherical harmonic configuration.
diff --git a/Engine/Shaders/Private/TemporalAA/TAADilateRejection.usf b/Engine/Shaders/Private/TemporalAA/TAADilateRejection.usf
--- a/Engine/Shaders/Private/TemporalAA/TAADilateRejection.usf
+++ b/Engine/Shaders/Private/TemporalAA/TAADilateRejection.usf
@@ -42,25 +42,27 @@ void MainCS(
 	// Fetch neighborhood of the kernel
 	UNROLL
 	for (int x = -KernelSize; x <= KernelSize; x++)
-	UNROLL
-	for (int y = -KernelSize; y <= KernelSize; y++)
-	//UNROLL
-	//for (uint i = 0; i < 9; i++)
 	{
-		//float2 Offset = float2(kOffsets3x3[i]);
-		float2 Offset = float2(x, y);
-		
-		if (all(Offset == 0.0))
+		UNROLL
+		for (int y = -KernelSize; y <= KernelSize; y++)
+		//UNROLL
+		//for (uint i = 0; i < 9; i++)
 		{
-			continue;
-		}
+			//float2 Offset = float2(kOffsets3x3[i]);
+			float2 Offset = float2(x, y);
+		
+			if (all(Offset == 0.0))
+			{
+				continue;
+			}
 
-		float2 SampleRejectionBufferUV = RejectionBufferUV + Offset * RejectionInfo_ExtentInverse;
-		SampleRejectionBufferUV = clamp(SampleRejectionBufferUV, RejectionInfo_UVViewportBilinearMin, RejectionInfo_UVViewportBilinearMax);
+			float2 SampleRejectionBufferUV = RejectionBufferUV + Offset * RejectionInfo_ExtentInverse;
+			SampleRejectionBufferUV = clamp(SampleRejectionBufferUV, RejectionInfo_UVViewportBilinearMin, RejectionInfo_UVViewportBilinearMax);
 
-		taa_half SampleRejectionInfo = HistoryRejectionTexture.SampleLevel(GlobalPointClampedSampler, SampleRejectionBufferUV, 0);
+			taa_half SampleRejectionInfo = HistoryRejectionTexture.SampleLevel(GlobalPointClampedSampler, SampleRejectionBufferUV, 0);
 
-		HistoryRejectionInfo = min(HistoryRejectionInfo, SampleRejectionInfo);
+			HistoryRejectionInfo = min(HistoryRejectionInfo, SampleRejectionInfo);
+		}
 	}
 
 	if (all(RejectionPixelPos < RejectionInfo_ViewportSize))
