fix a shader compile error

diff --git a/Engine/Shaders/Private/VolumetricCloud.usf b/Engine/Shaders/Private/VolumetricCloud.usf
--- a/Engine/Shaders/Private/VolumetricCloud.usf
+++ b/Engine/Shaders/Private/VolumetricCloud.usf
@@ -101,14 +101,13 @@ bool ViewHierarchicalZBufferMinDepth(float4 ScreenUVMinMax, const float MaxZDept
 
 void MainVS(
 	in FVertexFactoryInput FactoryInput,
-	in uint VertexID : SV_VertexID,
 	out RenderVolumetricCloudVSToPS Interpolants
 )
 {
 	ResolvedView = ResolveView();
 
 	// Full screen triangle
-	Interpolants.Position = float4(VertexID == 2 || VertexID == 1 ? 1.0 : -3.0, VertexID == 1 ? -3.0 : 1.0, 0.5f, 1.0f);
+	Interpolants.Position = float4(FactoryInput.VertexId == 2 || FactoryInput.VertexId == 1 ? 1.0 : -3.0, FactoryInput.VertexId == 1 ? -3.0 : 1.0, 0.5f, 1.0f);
 
 	const float3 WorldPosition = 0.0f;
 	FVertexFactoryIntermediates VFIntermediates = GetVertexFactoryIntermediates(FactoryInput);
