Fix for && within || logical op parenthesis error on DXC when compiling SSRTRayCast

diff --git a/Engine/Shaders/Private/SSRT/SSRTRayCast.ush b/Engine/Shaders/Private/SSRT/SSRTRayCast.ush
--- a/Engine/Shaders/Private/SSRT/SSRTRayCast.ush
+++ b/Engine/Shaders/Private/SSRT/SSRTRayCast.ush
@@ -457,7 +457,7 @@ void CastScreenSpaceRay(
 			for (uint j = 0; j < SSRT_SAMPLE_BATCH_SIZE; j++)
 			{
 				bFoundHit = bFoundHit || bMultipleSampleHit[j];
-				bUncertain = (bUncertain || bMultipleSampleUncertain[i] && !bFoundHit);
+				bUncertain |= (bMultipleSampleUncertain[i] && !bFoundHit);
 			}
 		}
 
