- Fixed some shader compile errors

diff --git a/Engine/Shaders/Private/SSRT/SSRTRayCast.ush b/Engine/Shaders/Private/SSRT/SSRTRayCast.ush
--- a/Engine/Shaders/Private/SSRT/SSRTRayCast.ush
+++ b/Engine/Shaders/Private/SSRT/SSRTRayCast.ush
@@ -457,7 +457,8 @@ void CastScreenSpaceRay(
 			for (uint j = 0; j < SSRT_SAMPLE_BATCH_SIZE; j++)
 			{
 				bFoundHit = bFoundHit || bMultipleSampleHit[j];
-				bUncertain |= (bMultipleSampleUncertain[i] && !bFoundHit);
+				bool bLocalMultisampleUncertain = bMultipleSampleUncertain[i];
+				bUncertain = bUncertain || (bLocalMultisampleUncertain && !bFoundHit);
 			}
 		}
 
diff --git a/Engine/Shaders/Private/VirtualShadowMaps/PageManagement.usf b/Engine/Shaders/Private/VirtualShadowMaps/PageManagement.usf
--- a/Engine/Shaders/Private/VirtualShadowMaps/PageManagement.usf
+++ b/Engine/Shaders/Private/VirtualShadowMaps/PageManagement.usf
@@ -361,10 +361,10 @@ void MarkPage(uint VirtualShadowMapId, uint MipLevel, FScreenPixelData Pixel)
 }
 
 [numthreads(VSM_DEFAULT_CS_GROUP_XY, VSM_DEFAULT_CS_GROUP_XY, 1)]
-void GeneratePageFlagsFromPixels(uint3 GroupId : SV_GroupID, uint3 GroupIndex : SV_GroupIndex)
+void GeneratePageFlagsFromPixels(uint3 GroupId : SV_GroupID, uint GroupIndex : SV_GroupIndex)
 {
 	// Morton order within a group so page access/atomics are more coherent and wave-swizzled gradients are possible.
-	uint2 PixelPos = VSM_DEFAULT_CS_GROUP_XY * GroupId.xy + MortonDecode(GroupIndex.x);
+	uint2 PixelPos = VSM_DEFAULT_CS_GROUP_XY * GroupId.xy + MortonDecode(GroupIndex);
 
 	if (any(PixelPos >= uint2(View.ViewSizeAndInvSize.xy)))
 	{
