Hierarchical Page-Flag fixes & ensure lights are added to light grid when virtual shadow maps are enabled
- Corrected offset calculation for hierarchical page flags.
- Improved consistency of hierarchical page flag lookup.
- Light grid build checks the r.Shadow.v.Enable cvar to enable light injection.

diff --git a/Engine/Shaders/Private/VirtualShadowMaps/PageAccessCommon.ush b/Engine/Shaders/Private/VirtualShadowMaps/PageAccessCommon.ush
--- a/Engine/Shaders/Private/VirtualShadowMaps/PageAccessCommon.ush
+++ b/Engine/Shaders/Private/VirtualShadowMaps/PageAccessCommon.ush
@@ -232,14 +232,16 @@ bool OverlapsAnyValidPage( uint ShadowMapID, uint MipLevel, int2 ViewPortRectMin
 		// Root level size in pages
 		uint HLevelSizePages = VSM_LEVEL0_DIM_PAGES_XY >> (MipLevel + HMipLevel);
 		
-		int HPageFlagOffset = ShadowMapID * HPageTableSize + HPageFlagLevelOffsets[MipLevel];
-		HPageFlagOffset -= CalcLevelOffsets( MipLevel + 1U );
+		int HPageLevelOffset = ShadowMapID * HPageTableSize + HPageFlagLevelOffsets[MipLevel]
+			+ CalcLevelOffsets( HMipLevel + MipLevel)
+			- CalcLevelOffsets( MipLevel + 1U );
 
 		for( uint y = RectPages.y; y <= RectPages.w; y++ )
 		{
 			for( uint x = RectPages.x; x <= RectPages.z; x++ )
 			{
-				if( (HPageFlags[ HPageFlagOffset + CalcLevelOffsets(HMipLevel + MipLevel) + y * HLevelSizePages + x ] & FlagMask) != 0 )
+				int HPageFlagOffset = HPageLevelOffset + y * HLevelSizePages + x;
+				if( (HPageFlags[ HPageFlagOffset ] & FlagMask) != 0 )
 				{
 					return true;
 				}
diff --git a/Engine/Shaders/Private/VirtualShadowMaps/PageManagement.usf b/Engine/Shaders/Private/VirtualShadowMaps/PageManagement.usf
--- a/Engine/Shaders/Private/VirtualShadowMaps/PageManagement.usf
+++ b/Engine/Shaders/Private/VirtualShadowMaps/PageManagement.usf
@@ -521,15 +521,21 @@ void GenerateHierarchicalPageFlags(uint2 ThreadId : SV_DispatchThreadID, uint2 G
 		InterlockedMax(PageRectBoundsOut[PageBoundIndex].z, PX);
 		InterlockedMax(PageRectBoundsOut[PageBoundIndex].w, PY);
 
-		// Loop over H flag levels:
-		uint NumHLevels = VSM_MAX_MIP_LEVELS - MipLevel - 1U;
+		// Loop over H flag levels, this builds a mip pyramid over _each_ mip level in the page table
+		// the 0-th level in this hiearchy is the page table mip level itself.
+		uint MaxHLevel = VSM_MAX_MIP_LEVELS - MipLevel;
 		uint HLevelSizePages = VSM_LEVEL0_DIM_PAGES_XY >> MipLevel;
-		for (uint HMipLevel = 0; HMipLevel < NumHLevels; ++HMipLevel)
+		// Note: starting from 1 as level 0 is the ordinary flag mip level
+		for (uint HMipLevel = 1U; HMipLevel < MaxHLevel; ++HMipLevel)
 		{
 			HLevelSizePages >>= 1U;
 			PX >>= 1U;
 			PY >>= 1U;
-			uint HPageFlagOffset = ShadowMapID * HPageTableSize + HPageFlagLevelOffsets[MipLevel] + CalcLevelOffsets(HMipLevel + MipLevel + 1U) - CalcLevelOffsets(MipLevel + 1U) + PY * HLevelSizePages + PX;
+			uint HPageFlagOffset = ShadowMapID * HPageTableSize 
+				+ HPageFlagLevelOffsets[MipLevel] 
+				+ CalcLevelOffsets(HMipLevel + MipLevel) 
+				- CalcLevelOffsets(MipLevel + 1U) 
+				+ PY * HLevelSizePages + PX;
 
 			uint PreviousValue = 0;
 			InterlockedOr(OutHPageFlags[HPageFlagOffset], Flag, PreviousValue);
diff --git a/Engine/Source/Runtime/Renderer/Private/LightGridInjection.cpp b/Engine/Source/Runtime/Renderer/Private/LightGridInjection.cpp
--- a/Engine/Source/Runtime/Renderer/Private/LightGridInjection.cpp
+++ b/Engine/Source/Runtime/Renderer/Private/LightGridInjection.cpp
@@ -310,10 +310,13 @@ void FDeferredShadingSceneRenderer::ComputeLightGrid(FRHICommandListImmediate& R
 			const FViewInfo& View = Views[ViewIndex];
 			bAnyViewUsesForwardLighting |= View.bTranslucentSurfaceLighting || ShouldRenderVolumetricFog() || View.bHasSingleLayerWaterMaterial;
 		}
-		
+
+		// TODO: Move decision logic to separate function for readability
+		static const auto EnableVirtualSMCVar = IConsoleManager::Get().FindTConsoleVariableDataInt(TEXT("r.Shadow.v.Enable"));
+
 		const bool bCullLightsToGrid = GLightCullingQuality 
 			&& (ViewFamily.EngineShowFlags.DirectLighting 
-			&& (IsForwardShadingEnabled(ShaderPlatform) || bAnyViewUsesForwardLighting || IsRayTracingEnabled() || ShouldUseClusteredDeferredShading() || GAllowLumenScene));
+			&& (IsForwardShadingEnabled(ShaderPlatform) || bAnyViewUsesForwardLighting || IsRayTracingEnabled() || ShouldUseClusteredDeferredShading() || GAllowLumenScene || EnableVirtualSMCVar->GetValueOnRenderThread() != 0));
 			   
 		// Store this flag if lights are injected in the grids, check with 'AreClusteredLightsInLightGrid()'
 		bClusteredShadingLightsInLightGrid = bCullLightsToGrid;
diff --git a/Engine/Source/Runtime/Renderer/Private/VirtualShadowMaps/VirtualShadowMapArray.cpp b/Engine/Source/Runtime/Renderer/Private/VirtualShadowMaps/VirtualShadowMapArray.cpp
--- a/Engine/Source/Runtime/Renderer/Private/VirtualShadowMaps/VirtualShadowMapArray.cpp
+++ b/Engine/Source/Runtime/Renderer/Private/VirtualShadowMaps/VirtualShadowMapArray.cpp
@@ -152,9 +152,8 @@ FVirtualShadowMapArray::FVirtualShadowMapArray()
 	uint32 HPageFlagOffset = 0;
 	for (uint32 Level = 0; Level < FVirtualShadowMap::MaxMipLevels - 1; ++Level)
 	{
-		uint32 NumHLevels = FVirtualShadowMap::MaxMipLevels - Level - 1U;
 		CommonParameters.HPageFlagLevelOffsets[Level] = HPageFlagOffset;
-		HPageFlagOffset += CommonParameters.LevelOffsets[FVirtualShadowMap::MaxMipLevels - 1] - CommonParameters.LevelOffsets[Level + 1];
+		HPageFlagOffset += CommonParameters.PageTableSize - CommonParameters.LevelOffsets[Level + 1];
 	}
 	// The last mip level is 1x1 and thus does not have any H levels possible.
 	CommonParameters.HPageFlagLevelOffsets[FVirtualShadowMap::MaxMipLevels - 1] = 0;
