Implements experimental TAAU for Gen5 consoles

diff --git a/Engine/Shaders/Private/TemporalAA/TAAStandalone.usf b/Engine/Shaders/Private/TemporalAA/TAAStandalone.usf
--- a/Engine/Shaders/Private/TemporalAA/TAAStandalone.usf
+++ b/Engine/Shaders/Private/TemporalAA/TAAStandalone.usf
@@ -5,7 +5,7 @@
 #include "../EyeAdaptationCommon.ush"
 #include "../TextureSampling.ush"
 
-//------------------------------------------------------- COMPILER CONFIG   
+//------------------------------------------------------- COMPILER CONFIG
 
 // Generate vector truncation warnings to errors.
 #pragma warning(error: 3206)
@@ -37,6 +37,9 @@
 	// Caches 3x3 Neighborhood into VGPR (although my have corner optimised away).
 	#define AA_SAMPLE_CACHE_METHOD_VGPR_3X3 1
 	
+	// Prefetches scene color into 10x10 LDS tile (8x8 when screen percentage < 71%).
+	#define AA_SAMPLE_CACHE_METHOD_LDS 2
+	
 
 /** Clamping method for scene color. */
 	// Min max neighboorhing samples.
@@ -63,7 +66,6 @@
 		#define AA_HISTORY_PAYLOAD (HISTORY_PAYLOAD_RGB)
 	#endif
 	#define AA_BICUBIC 1
-	#define AA_BORDER 1
 	#define AA_CROSS 2
 	#define AA_DYNAMIC 1
 	#define AA_FILTERED 1
@@ -75,6 +77,10 @@
 		#define AA_DYNAMIC_ANTIGHOST 1
 	#endif
 
+	#if COMPUTESHADER
+		#define AA_SAMPLE_CACHE_METHOD (AA_SAMPLE_CACHE_METHOD_LDS)
+	#endif
+
 #elif TAA_PASS_CONFIG == 3 // Used for SSR
 	#define AA_HISTORY_PAYLOAD (HISTORY_PAYLOAD_RGB_OPACITY)
 	#define AA_DYNAMIC 1
@@ -82,7 +88,6 @@
 	#define AA_LERP 8
 	#define AA_MANUALLY_CLAMP_HISTORY_UV 1
 	#define AA_YCOCG 1
-	#define AA_BORDER 1
 
 #elif TAA_PASS_CONFIG == 4 // Used for LightShaft
 	#define AA_HISTORY_PAYLOAD (HISTORY_PAYLOAD_RGB_OPACITY)
@@ -91,7 +96,6 @@
 	#define AA_MANUALLY_CLAMP_HISTORY_UV 1
 	#define AA_YCOCG 1
 	#define AA_LOWER_RESOLUTION 1
-	#define AA_BORDER 1
 
 #elif TAA_PASS_CONFIG == 1 || TAA_PASS_CONFIG == 2 // MainUpsampling & MainSuperSampling
 	#if POST_PROCESS_ALPHA
@@ -100,7 +104,6 @@
 		#define AA_HISTORY_PAYLOAD (HISTORY_PAYLOAD_RGB)
 	#endif
 	#define AA_BICUBIC 1
-	#define AA_BORDER 2
 	#define AA_CROSS 1
 	#define AA_DYNAMIC 1
 	#define AA_FILTERED TAA_UPSAMPLE_FILTERED
@@ -128,11 +131,19 @@
 		#endif
 
 	#else
-		//#define AA_HISTORY_CLAMPING_BOX (HISTORY_CLAMPING_BOX_SAMPLE_DISTANCE)
-		#define AA_ANTIFLICKER 1
+		#define AA_HISTORY_CLAMPING_BOX (HISTORY_CLAMPING_BOX_SAMPLE_DISTANCE)
 		#define AA_DYNAMIC_ANTIGHOST 1
-		#define AA_SAMPLES 6
+		#define AA_SAMPLES 9
+
+	#endif
+
+	// Do not use LDS caching for screen percentage > 100% or < 50%.
+	#if TAA_SCREEN_PERCENTAGE_RANGE == 2 || TAA_SCREEN_PERCENTAGE_RANGE == 3
+		#define AA_SAMPLE_CACHE_METHOD (AA_SAMPLE_CACHE_METHOD_VGPR_3X3)
 
+
+	#else
+		#define AA_SAMPLE_CACHE_METHOD (AA_SAMPLE_CACHE_METHOD_LDS)
 	#endif
 
 	#if !COMPUTESHADER
@@ -152,11 +163,9 @@
 	#define AA_LOWER_RESOLUTION 1
 
 	#if TAA_PASS_CONFIG == 6
-		#define AA_BORDER 2
 		#define AA_UPSAMPLE 1
 		#define AA_UPSAMPLE_ADAPTIVE_FILTERING 1
 	#else
-		#define AA_BORDER 1
 		#define AA_UPSAMPLE 0
 	#endif
 	
@@ -168,6 +177,10 @@
 		#define AA_BICUBIC 1
 		#define AA_FILTERED 1
 	#endif
+
+	#if TAA_SCREEN_PERCENTAGE_RANGE != 2 && AA_UPSAMPLE
+		//#define AA_SAMPLE_CACHE_METHOD (AA_SAMPLE_CACHE_METHOD_LDS) // TODO: doesn't work yet.
+	#endif
 	
 #else
 	#error Unknown TAA pass config. Have you changed ETAAPassConfig without updating me?
@@ -213,9 +226,10 @@
 	#define AA_LERP 0
 #endif
 
-// Fix for lack of borders during current frame filter.
-#ifndef AA_BORDER
-	#define AA_BORDER 0
+// 1 = Use higher quality round clamp.
+// 0 = Use lower quality but faster box clamp.
+#ifndef AA_ROUND
+	#define AA_ROUND 0
 #endif
 
 // Force clamp on alpha.
@@ -273,10 +287,6 @@
 	#define AA_MANUALLY_CLAMP_HISTORY_UV 0
 #endif
 
-#ifndef AA_ANTIFLICKER
-	#define AA_ANTIFLICKER 0
-#endif
-
 
 //------------------------------------------------------- CONFIG ENABLED DEFAULTS
 
@@ -323,16 +333,11 @@
 	#define HISTORY_PAYLOAD_COMPONENTS 4
 #endif
 
-// Defines the number of render target input to TAA.
-#if HISTORY_PAYLOAD_COMPONENTS > 4
-	#define INPUT_RENDER_TARGETS 2
-#else
-	#define INPUT_RENDER_TARGETS 1
-#endif
-
 // Defines the number of render target to store TAA's history.
-#if 1
-	#define HISTORY_RENDER_TARGETS (INPUT_RENDER_TARGETS)
+#if AA_HISTORY_PAYLOAD == HISTORY_PAYLOAD_RGB_OPACITY_COC
+	#define HISTORY_RENDER_TARGETS 2
+#else
+	#define HISTORY_RENDER_TARGETS 1
 #endif
 
 
@@ -342,6 +347,10 @@
 	#error Samples must be 5, (6 for TAAU) or 9
 #endif
 
+#if AA_SAMPLE_CACHE_METHOD >= 2 && !COMPUTESHADER
+	#error Group share only for compute shader.
+#endif
+
 
 //------------------------------------------------------- CONSTANTS
 
@@ -378,10 +387,14 @@ static const int2 kOffsets3x3[9] =
 };
 	
 // Indexes of the 3x3 square.
-static const uint kSquareIndexes3x3[9] = { 4, 0, 1, 2, 3, 5, 6, 7, 8 };
+static const uint kSquareIndexes3x3[9] = { 0, 1, 2, 3, 4, 5, 6, 7, 8 };
 
 // Indexes of the offsets to have plus + shape.
-static const uint kPlusIndexes3x3[5] = { 4, 1, 3, 5, 7 };
+static const uint kPlusIndexes3x3[5] = { 1, 3, 4, 5, 7 };
+
+// Number of neighbors.
+static const uint kNeighborsCount = 9;
+
 
 #if AA_UPSAMPLE
 	// T = Center of the nearest top left pixel input pixel.
@@ -428,44 +441,44 @@ float4 ScreenPosToHistoryBufferUV;
 float4 InputSceneColorSize;
 int2 InputMinPixelCoord;
 int2 InputMaxPixelCoord;
-Texture2D<half4> InputSceneColor;
-Texture2D<half4> InputSceneMetadata;
+Texture2D InputSceneColor;
+SamplerState InputSceneColorSampler;
+Texture2D InputSceneMetadata;
+SamplerState InputSceneMetadataSampler;
 
-Texture2D<half4> SceneDepthBuffer;
+Texture2D SceneDepthBuffer;
+SamplerState SceneDepthBufferSampler;
 
 Texture2D SceneVelocityBuffer;
 SamplerState SceneVelocityBufferSampler;
 
 Texture2D<uint2> StencilTexture;
 
-Texture2D<half4> HistoryBuffer_0;
-Texture2D<half4> HistoryBuffer_1;
-Texture2D<half4> HistoryBuffer_2;
-Texture2D<half4> HistoryBuffer_3;
+Texture2D HistoryBuffer_0;
+Texture2D HistoryBuffer_1;
 SamplerState HistoryBufferSampler_0;
 SamplerState HistoryBufferSampler_1;
-SamplerState HistoryBufferSampler_2;
-SamplerState HistoryBufferSampler_3;
 float4 HistoryBufferSize;
 float4 HistoryBufferUVMinMax;
 
 float4 OutputViewportSize;
-float4 OutputViewportRect; // TODO:Reverb: int
+float4 OutputViewportRect;
 
-RWTexture2D<half4> OutComputeTex_0;
-RWTexture2D<half4> OutComputeTex_1;
-RWTexture2D<half4> OutComputeTex_2;
-RWTexture2D<half4> OutComputeTex_3;
-
-RWTexture2D<float4> DebugOutput;
+#if COMPUTESHADER
+RWTexture2D<float4> OutComputeTex_0;
 
 #if TAA_DOWNSAMPLE 
 // This shader permutation outputs half resolution image in addition to main full-res one.
 // It is more efficient than performing a separate downsampling pass afterwards.
-RWTexture2D<half4> OutComputeTexDownsampled;
-groupshared half4 GroupSharedDownsampleArray[THREADGROUP_SIZEX*THREADGROUP_SIZEY]; // TODO: share this with GroupSharedArrayF4 when possible
+RWTexture2D<float4> OutComputeTexDownsampled;
+groupshared float4 GroupSharedDownsampleArray[THREADGROUP_SIZEX*THREADGROUP_SIZEY]; // TODO: share this with GroupSharedArrayF4 when possible
 #endif // TAA_DOWNSAMPLE
 
+#if HISTORY_RENDER_TARGETS == 2
+	RWTexture2D<float4> OutComputeTex_1;
+#endif // HISTORY_RENDER_TARGETS == 2
+#endif
+
 // Temporal upsample specific params.
 #if AA_UPSAMPLE
 
@@ -483,46 +496,46 @@ float UpscaleFactor; // = 1 / ScreenPercentage
 
 //------------------------------------------------------- FUNCTIONS
 
-half3 RGBToYCoCg(half3 RGB )
+float3 RGBToYCoCg( float3 RGB )
 {
-	half Y  = dot( RGB, half3(  1, 2,  1 ) );
-	half Co = dot( RGB, half3(  2, 0, -2 ) );
-	half Cg = dot( RGB, half3( -1, 2, -1 ) );
+	float Y  = dot( RGB, float3(  1, 2,  1 ) );
+	float Co = dot( RGB, float3(  2, 0, -2 ) );
+	float Cg = dot( RGB, float3( -1, 2, -1 ) );
 	
-	half3 YCoCg = half3( Y, Co, Cg );
+	float3 YCoCg = float3( Y, Co, Cg );
 	return YCoCg;
 }
 
-half3 YCoCgToRGB(half3 YCoCg )
+float3 YCoCgToRGB( float3 YCoCg )
 {
-	half Y  = YCoCg.x * 0.25;
-	half Co = YCoCg.y * 0.25;
-	half Cg = YCoCg.z * 0.25;
+	float Y  = YCoCg.x * 0.25;
+	float Co = YCoCg.y * 0.25;
+	float Cg = YCoCg.z * 0.25;
 
-	half R = Y + Co - Cg;
-	half G = Y + Cg;
-	half B = Y - Co - Cg;
+	float R = Y + Co - Cg;
+	float G = Y + Cg;
+	float B = Y - Co - Cg;
 
-	half3 RGB = half3( R, G, B );
+	float3 RGB = float3( R, G, B );
 	return RGB;
 }
 
 // Faster but less accurate luma computation. 
 // Luma includes a scaling by 4.
-half Luma4(half3 Color)
+float Luma4(float3 Color)
 {
-	return (Color.g * half(2.0)) + (Color.r + Color.b);
+	return (Color.g * 2.0) + (Color.r + Color.b);
 }
 
 // Optimized HDR weighting function.
-half HdrWeight4(half3 Color, half Exposure)
+float HdrWeight4(float3 Color, float Exposure) 
 {
-	return rcp(Luma4(Color * Exposure) + half(4.0));
+	return rcp(Luma4(Color) * Exposure + 4.0);
 }
 
-half HdrWeightY(half Color, half Exposure)
+float HdrWeightY(float Color, float Exposure) 
 {
-	return rcp(Color * Exposure + half(4.0));
+	return rcp(Color * Exposure + 4.0);
 }
 
 // Intersect ray with AABB, knowing there is an intersection.
@@ -530,70 +543,68 @@ half HdrWeightY(half Color, half Exposure)
 //   Org = Start of the ray.
 //   Box = Box is at {0,0,0} with this size.
 // Returns distance on line segment.
-half IntersectAABB(half3 Dir, half3 Org, half3 Box)
+float IntersectAABB(float3 Dir, float3 Org, float3 Box)
 {
 	#if PS4_PROFILE
 		// This causes flicker, it should only be used on PS4 until proper fix is in.
-		if(min(min(abs(Dir.x), abs(Dir.y)), abs(Dir.z)) < half3(1.0/65536.0)) return 1.0;
+		if(min(min(abs(Dir.x), abs(Dir.y)), abs(Dir.z)) < (1.0/65536.0)) return 1.0;
 	#endif
-	half3 RcpDir = rcp(Dir);
-	half3 TNeg = (  Box  - Org) * RcpDir;
-	half3 TPos = ((-Box) - Org) * RcpDir;
+	float3 RcpDir = rcp(Dir);
+	float3 TNeg = (  Box  - Org) * RcpDir;
+	float3 TPos = ((-Box) - Org) * RcpDir;
 	return max(max(min(TNeg.x, TPos.x), min(TNeg.y, TPos.y)), min(TNeg.z, TPos.z));
 }
 
-half HistoryClip(half3 History, half3 Filtered, half3 NeighborMin, half3 NeighborMax)
+float HistoryClip(float3 History, float3 Filtered, float3 NeighborMin, float3 NeighborMax)
 {
 #if 0
-	half3 Min = min(Filtered, min(NeighborMin, NeighborMax));
-	half3 Max = max(Filtered, max(NeighborMin, NeighborMax));	
-	half3 Avg2 = Max + Min;
-	half3 Dir = Filtered - History;
-	half3 Org = History - Avg2 * half(0.5);
-	half3 Scale = Max - Avg2 * half(0.5);
+	float3 Min = min(Filtered, min(NeighborMin, NeighborMax));
+	float3 Max = max(Filtered, max(NeighborMin, NeighborMax));	
+	float3 Avg2 = Max + Min;
+	float3 Dir = Filtered - History;
+	float3 Org = History - Avg2 * 0.5;
+	float3 Scale = Max - Avg2 * 0.5;
 	return saturate(IntersectAABB(Dir, Org, Scale));
 #else
-	half3 BoxMin = NeighborMin;
-	half3 BoxMax = NeighborMax;
-	//half3 BoxMin = min( Filtered, NeighborMin );
-	//half3 BoxMax = max( Filtered, NeighborMax );
-
-	half3 RayOrigin = History;
-	half3 RayDir = Filtered - History;
-	RayDir.x = abs(RayDir.x) < half(1.0 / 65536.0) ? half(1.0 / 65536.0) : RayDir.x;
-	RayDir.y = abs(RayDir.y) < half(1.0 / 65536.0) ? half(1.0 / 65536.0) : RayDir.y;
-	RayDir.z = abs(RayDir.z) < half(1.0 / 65536.0) ? half(1.0 / 65536.0) : RayDir.z;
-	half3 InvRayDir = rcp( RayDir );
-
-	half3 MinIntersect = (BoxMin - RayOrigin) * InvRayDir;
-	half3 MaxIntersect = (BoxMax - RayOrigin) * InvRayDir;
-	half3 EnterIntersect = min( MinIntersect, MaxIntersect );
+	float3 BoxMin = NeighborMin;
+	float3 BoxMax = NeighborMax;
+	//float3 BoxMin = min( Filtered, NeighborMin );
+	//float3 BoxMax = max( Filtered, NeighborMax );
+
+	float3 RayOrigin = History;
+	float3 RayDir = Filtered - History;
+	RayDir = abs( RayDir ) < (1.0/65536.0) ? (1.0/65536.0) : RayDir;
+	float3 InvRayDir = rcp( RayDir );
+
+	float3 MinIntersect = (BoxMin - RayOrigin) * InvRayDir;
+	float3 MaxIntersect = (BoxMax - RayOrigin) * InvRayDir;
+	float3 EnterIntersect = min( MinIntersect, MaxIntersect );
 	return max3( EnterIntersect.x, EnterIntersect.y, EnterIntersect.z );
 #endif
 }
 
-half2 WeightedLerpFactors(half WeightA, half WeightB, half Blend)
+float2 WeightedLerpFactors(float WeightA, float WeightB, float Blend)
 {
-	half BlendA = (half(1.0) - Blend) * WeightA;
-	half BlendB =        Blend  * WeightB;
-	half RcpBlend = rcp(BlendA + BlendB);
+	float BlendA = (1.0 - Blend) * WeightA;
+	float BlendB =        Blend  * WeightB;
+	float RcpBlend = rcp(BlendA + BlendB);
 	BlendA *= RcpBlend;
 	BlendB *= RcpBlend;
-	return half2(BlendA, BlendB);
+	return float2(BlendA, BlendB);
 }
 
 #if AA_HISTORY_PAYLOAD == HISTORY_PAYLOAD_RGB_COC || AA_HISTORY_PAYLOAD == HISTORY_PAYLOAD_RGB_OPACITY_COC
 
 // Computes the bilateral weight according to two Coc radii.
-half ComputeBilateralWeight(half RefCocRadius, half SampleCocRadius)
+float ComputeBilateralWeight(float RefCocRadius, float SampleCocRadius)
 {
-	half Factor = (abs(RefCocRadius) > 1 ? rcp(abs(RefCocRadius)) : 1.0);
+	float Factor = (abs(RefCocRadius) > 1 ? rcp(abs(RefCocRadius)) : 1.0);
 	return saturate(1 - abs(RefCocRadius - SampleCocRadius) * Factor);
 }
 
-half ComputeNeightborSampleBilateralWeight(half CenterCocRadius, half SampleCocRadius)
+float ComputeNeightborSampleBilateralWeight(float CenterCocRadius, float SampleCocRadius)
 {
-	half Factor = (abs(CenterCocRadius) > 1 ? rcp(abs(CenterCocRadius)) : 1.0);
+	float Factor = (abs(CenterCocRadius) > 1 ? rcp(abs(CenterCocRadius)) : 1.0);
 	return saturate(1 - (CenterCocRadius - SampleCocRadius) * Factor);
 }
 
@@ -601,54 +612,54 @@ half ComputeNeightborSampleBilateralWeight(half CenterCocRadius, half SampleCocR
 #endif
 
 
-//------------------------------------------------------- TAA's PAYLOAD
+//------------------------------------------------------- HISTORY's PAYLOAD
 
 // Payload of the TAA's history.
-struct FTAAPayload
+struct FTAAHistoryPayload
 {
 	// Transformed scene color and alpha channel.
-	half4 Color;
+	float4 Color;
 
 	// Radius of the circle of confusion for DOF.
-	half CocRadius;
+	float CocRadius;
 };
 
-FTAAPayload MulPayload(in FTAAPayload Payload, in half x)
+FTAAHistoryPayload MulPayload(in FTAAHistoryPayload Payload, in float x)
 {
 	Payload.Color *= x;
 	Payload.CocRadius *= x;
 	return Payload;
 }
 
-FTAAPayload AddPayload(in FTAAPayload Payload0, in FTAAPayload Payload1)
+FTAAHistoryPayload AddPayload(in FTAAHistoryPayload Payload0, in FTAAHistoryPayload Payload1)
 {
 	Payload0.Color += Payload1.Color;
 	Payload0.CocRadius += Payload1.CocRadius;
 	return Payload0;
 }
 
-FTAAPayload MinPayload(in FTAAPayload Payload0, in FTAAPayload Payload1)
+FTAAHistoryPayload MinPayload(in FTAAHistoryPayload Payload0, in FTAAHistoryPayload Payload1)
 {
 	Payload0.Color = min(Payload0.Color, Payload1.Color);
 	Payload0.CocRadius = min(Payload0.CocRadius, Payload1.CocRadius);
 	return Payload0;
 }
 
-FTAAPayload MaxPayload(in FTAAPayload Payload0, in FTAAPayload Payload1)
+FTAAHistoryPayload MaxPayload(in FTAAHistoryPayload Payload0, in FTAAHistoryPayload Payload1)
 {
 	Payload0.Color = max(Payload0.Color, Payload1.Color);
 	Payload0.CocRadius = max(Payload0.CocRadius, Payload1.CocRadius);
 	return Payload0;
 }
 
-FTAAPayload MinPayload3(in FTAAPayload Payload0, in FTAAPayload Payload1, in FTAAPayload Payload2)
+FTAAHistoryPayload MinPayload3(in FTAAHistoryPayload Payload0, in FTAAHistoryPayload Payload1, in FTAAHistoryPayload Payload2)
 {
 	Payload0.Color = min3(Payload0.Color, Payload1.Color, Payload2.Color);
 	Payload0.CocRadius = min3(Payload0.CocRadius, Payload1.CocRadius, Payload2.CocRadius);
 	return Payload0;
 }
 
-FTAAPayload MaxPayload3(in FTAAPayload Payload0, in FTAAPayload Payload1, in FTAAPayload Payload2)
+FTAAHistoryPayload MaxPayload3(in FTAAHistoryPayload Payload0, in FTAAHistoryPayload Payload1, in FTAAHistoryPayload Payload2)
 {
 	Payload0.Color = max3(Payload0.Color, Payload1.Color, Payload2.Color);
 	Payload0.CocRadius = max3(Payload0.CocRadius, Payload1.CocRadius, Payload2.CocRadius);
@@ -661,6 +672,11 @@ FTAAPayload MaxPayload3(in FTAAPayload Payload0, in FTAAPayload Payload1, in FTA
 // Output pixel parameters. Should not be modified once setup.
 struct FTAAInputParameters
 {
+	// Compute shader dispatch params, set to 0 in pixel shader.
+	uint2 GroupId;
+	uint2 GroupThreadId;
+	uint GroupThreadIndex;
+
 	// Viewport UV of the output pixel.
 	float2 ViewportUV;
 
@@ -669,35 +685,28 @@ struct FTAAInputParameters
 
 	// Buffer UV of the nearest input pixel.
 	float2 NearestBufferUV;
+
+	#if AA_UPSAMPLE
+		// Buffer UV of the nearest top left input pixel.
+		float2 NearestTopLeftBufferUV;
+	#endif
 	
 	// Whether this pixel should be responsive.
-	bool bIsResponsiveAAPixel;
+	float bIsResponsiveAAPixel;
 
 	// Frame exposure's scale.
-	half FrameExposureScale;
+	float FrameExposureScale;
 	
 	// Cache of neightbors' transformed scene color.
 	#if AA_SAMPLE_CACHE_METHOD == AA_SAMPLE_CACHE_METHOD_VGPR_3X3
-		half4 CachedNeighbors0[AA_SAMPLES];
+		float4 CachedNeighbors0[kNeighborsCount];
 
-		#if INPUT_RENDER_TARGETS == 2
-			half CachedNeighbors1[AA_SAMPLES];
+		#if HISTORY_RENDER_TARGETS == 2
+			float CachedNeighbors1[kNeighborsCount];
 		#endif
 	#endif
 };
 
-struct FTAAHistoryInformation
-{
-	bool bIsDynamicPixel;
-};
-
-FTAAHistoryInformation CreateHistoryInformation()
-{
-	FTAAHistoryInformation HistoryInfo;
-	HistoryInfo.bIsDynamicPixel = false;
-	return HistoryInfo;
-}
-
 
 // Intermediary results for major function to conveniently share values between them.
 //
@@ -705,13 +714,13 @@ FTAAHistoryInformation CreateHistoryInformation()
 struct FTAAIntermediaryResult
 {
 	// The filtered input.
-	FTAAPayload Filtered;
+	FTAAHistoryPayload Filtered;
 
 	// Temporal weight of the filtered input.
-	half FilteredTemporalWeight;
+	float FilteredTemporalWeight;
 
 	// 1 / filtering kernel scale factor for AA_UPSAMPLE_ADAPTIVE_FILTERING.
-	half InvFilterScaleFactor;
+	float InvFilterScaleFactor;
 };
 
 
@@ -732,48 +741,45 @@ FTAAIntermediaryResult CreateIntermediaryResult()
 // Transformed scene color's data for a sample.
 struct FTAASceneColorSample
 {
-	// Payload of the sample.
-	FTAAPayload Payload;
+	// Transformed scene color and alpha channel.
+	float4 Color;
 
-	// HDR weight of the scene color sample.
-	half HdrWeight;
+	// Radius of the circle of confusion for DOF.
+	float CocRadius;
 
-	// Pixel offset of the sample
-	int2 SampleOffset;
+	// HDR weight of the scene color sample.
+	float HdrWeight;
 };
 
 
 //------------------------------------------------------- SCENE COLOR SPACE MANAGMENT
 
 // Transform RAW linear scene color RGB to TAA's working color space.
-half4 TransformSceneColor(half4 RawLinearSceneColorRGBA)
+float4 TransformSceneColor(float4 RawLinearSceneColorRGBA)
 {
 	#if AA_YCOCG
-		return half4(RGBToYCoCg(RawLinearSceneColorRGBA.rgb), RawLinearSceneColorRGBA.a);
+		return float4(RGBToYCoCg(RawLinearSceneColorRGBA.rgb), RawLinearSceneColorRGBA.a);
 	#endif
 	return RawLinearSceneColorRGBA;
 }
 
 // Reciprocal of TransformSceneColor().
-half4 TransformBackToRawLinearSceneColor(half4 SceneColor)
+float4 TransformBackToRawLinearSceneColor(float4 SceneColor)
 {
 	#if AA_YCOCG
-		return half4(YCoCgToRGB(SceneColor.xyz), SceneColor.a);
+		return float4(YCoCgToRGB(SceneColor.xyz), SceneColor.a);
 	#endif
 	return SceneColor;
 }
 
 // Transform current frame's RAW scene color RGB to TAA's working color space.
-half4 TransformCurrentFrameSceneColor(half4 RawSceneColorRGBA)
+float4 TransformCurrentFrameSceneColor(float4 RawSceneColorRGBA)
 {
-	// This min is to avoid large values causing NaN/black pixels resulting in bloom flickering on the sun due to half values and high pixel luminance.
-	// It is safe to clamp this way here because SceneColor in Unreal is always pre-exposed. We can accumulate up to 9 samples so MaxHalfFloat/9 ~= 7000.
-	half4 SafeRawSceneColorRGBA = min(RawSceneColorRGBA, half4(7000.0f, 7000.0f, 7000.0f, 7000.0f));
-	return TransformSceneColor(SafeRawSceneColorRGBA);
+	return TransformSceneColor(RawSceneColorRGBA);
 }
 
 // Get the Luma4 of the sceneColor
-half GetSceneColorLuma4(half4 SceneColor)
+float GetSceneColorLuma4(float4 SceneColor)
 {
 	#if AA_YCOCG
 		return SceneColor.x;
@@ -782,9 +788,9 @@ half GetSceneColorLuma4(half4 SceneColor)
 }
 
 // Get the HDR weight of the transform scene color.
-half GetSceneColorHdrWeight(
+float GetSceneColorHdrWeight(
 	in FTAAInputParameters InputParams,
-	half4 SceneColor)
+	float4 SceneColor)
 {
 	#if AA_YCOCG
 		return HdrWeightY(SceneColor.x, InputParams.FrameExposureScale);
@@ -792,52 +798,17 @@ half GetSceneColorHdrWeight(
 	return HdrWeight4(SceneColor.rgb, InputParams.FrameExposureScale);
 }
 
-int2 GetNeighborSampleOffset(in FTAAInputParameters InputParams, uint NeighborId)
-{
-	#if AA_UPSAMPLE
-		// Pixel coordinate of the center of output pixel O in the input viewport.
-		float2 PPCo = InputParams.ViewportUV * InputViewSize.xy + TemporalJitterPixels;
-
-		// Pixel coordinate of the center of the nearest input pixel K.
-		float2 PPCk = floor(PPCo) + 0.5;
-		
-		// Vector in pixel between pixel K -> O.
-		half2 dKO = half2(PPCo - PPCk);
-	#endif
-
-	#if AA_SAMPLES == 9
-		const uint SampleIndexes[9] = kSquareIndexes3x3;
-	#elif AA_SAMPLES == 5 || AA_SAMPLES == 6
-		const uint SampleIndexes[5] = kPlusIndexes3x3;
-	#endif
-
-	// Get the sample offset from the nearest input pixel.
-	int2 SampleOffset;
-			
-	#if AA_UPSAMPLE && AA_SAMPLES == 6
-	if (NeighborId == 5)
-	{
-		SampleOffset = SignFastInt(dKO);
-	}
-	else
-	#endif
-	{
-		const uint SampleIndex = SampleIndexes[NeighborId];
-		SampleOffset = kOffsets3x3[SampleIndex];
-	}
-
-	return SampleOffset;
-}
-
 
 //------------------------------------------------------- INPUT SAMPLE CACHING.
 // API to sample input scene color and depth through caching system.
 // 
 // Precache scene color or depth:
 //		PrecacheInputSceneColor(InputParams);
+//		PrecacheInputSceneDepth(InputParams);
 //
 // Then sample scene color or depth:
 //		SampleCachedSceneColorTexture(InputParams, /* Offset = */ int2(-1, -1));
+//		SampleCachedSceneDepthTexture(InputParams, /* Offset = */ int2(-1, -1));
 //
 //		<Offset> parameter is meant to be compile time constant of the pixel offset from nearest input sample.
 
@@ -846,100 +817,538 @@ int2 GetNeighborSampleOffset(in FTAAInputParameters InputParams, uint NeighborId
 
 	#define AA_PRECACHE_SCENE_COLOR 1
 	
-	ISOLATE
 	void PrecacheInputSceneColor(inout FTAAInputParameters InputParams)
 	{
 		// Precache 3x3 input scene color into FTAAInputParameters::CachedNeighbors.
-		UNROLL_N(AA_SAMPLES)
-		for (uint NeighborId = 0; NeighborId < AA_SAMPLES; NeighborId++)
+		UNROLL
+		for (uint i = 0; i < kNeighborsCount; i++)
 		{
-			int2 SampleOffset = GetNeighborSampleOffset(InputParams, NeighborId);
-
-			int2 Coord = int2(InputParams.NearestBufferUV * InputSceneColorSize.xy) + SampleOffset;
-
-			#if AA_BORDER == 2
-				Coord = clamp(Coord, InputMinPixelCoord, InputMaxPixelCoord);
-			#endif
+			int2 Coord = int2(InputParams.NearestBufferUV * InputSceneColorSize.xy) + kOffsets3x3[i];
+			Coord = clamp(Coord, InputMinPixelCoord, InputMaxPixelCoord);
 
-			InputParams.CachedNeighbors0[NeighborId] = TransformCurrentFrameSceneColor(InputSceneColor[Coord]);
+			InputParams.CachedNeighbors0[i] = TransformCurrentFrameSceneColor(InputSceneColor[Coord]);
 	
-			#if AA_HISTORY_PAYLOAD == HISTORY_PAYLOAD_RGB_OPACITY_COC
-				InputParams.CachedNeighbors1[NeighborId] = InputSceneMetadata[Coord].r;
+			#if HISTORY_RENDER_TARGETS == 2
+				InputParams.CachedNeighbors1[i] = InputSceneMetadata[Coord].r;
 			#endif
 		}
 	}
 
-	FTAASceneColorSample SampleCachedSceneColorTexture(in FTAAInputParameters InputParams, uint NeighborId)
+	FTAASceneColorSample SampleCachedSceneColorTexture(
+		inout FTAAInputParameters InputParams,
+		int2 PixelOffset)
 	{
+		// PixelOffset is const at compile time. Therefore all this computaton is actually free.
+		uint NeighborsId = uint(4 + PixelOffset.x + PixelOffset.y * 3);
 		FTAASceneColorSample Sample;
-		Sample.SampleOffset = GetNeighborSampleOffset(InputParams, NeighborId);
 
-		Sample.Payload.Color = InputParams.CachedNeighbors0[NeighborId];
-		Sample.Payload.CocRadius = 0;
+		Sample.Color = InputParams.CachedNeighbors0[NeighborsId];
+		Sample.CocRadius = 0;
+
 		#if AA_HISTORY_PAYLOAD == HISTORY_PAYLOAD_RGB_COC
-			Sample.Payload.CocRadius = Sample.Payload.Color.a;
+			Sample.CocRadius = Sample.Color.a;
 		#elif AA_HISTORY_PAYLOAD == HISTORY_PAYLOAD_RGB_OPACITY_COC
-			Sample.Payload.CocRadius = InputParams.CachedNeighbors1[NeighborId];
+			Sample.CocRadius = InputParams.CachedNeighbors1[NeighborsId];
 		#endif
 
-		Sample.HdrWeight = GetSceneColorHdrWeight(InputParams, Sample.Payload.Color);
+		Sample.HdrWeight = GetSceneColorHdrWeight(InputParams, Sample.Color);
 		return Sample;
 	}
 
-#else // AA_SAMPLE_CACHE_METHOD != AA_SAMPLE_CACHE_METHOD_VGPR_3X3
 
-	// Silently do no scene color precaching.
-	void PrecacheInputSceneColor(in FTAAInputParameters InputParams)
-	{ }
+#elif AA_SAMPLE_CACHE_METHOD == AA_SAMPLE_CACHE_METHOD_LDS
+//-------------------------------- Generic LDS functions.
+
+// Returns SV_GroupIndex.
+uint GetGroupIndex(in FTAAInputParameters InputParams)
+{
+	if (0)
+	{
+		// Group thread index doesn't actually suppress ALU instruction on consoles.
+		return InputParams.GroupThreadIndex;
+	}
+	return InputParams.GroupThreadId.x + InputParams.GroupThreadId.y * THREADGROUP_SIZEX;
+}
+
+
+//------------------------------------------------------- COMPUTE SHADER GROUP SHARE OPTIMIZATION
+
+#if THREADGROUP_SIZEX != THREADGROUP_SIZEY
+	#error AA_SAMPLE_CACHE_METHOD >= 2 assume square shaped tiles.
+#endif
+
+#ifndef TAA_SCREEN_PERCENTAGE_RANGE
+	#error LDS cache needs to know the screen percentage range.
+#endif
+
+#define AA_PRECACHE_IMPLEMENTATIONS 1
+
+// Total number of thread per group.
+#define THREADGROUP_TOTAL (THREADGROUP_SIZEX * THREADGROUP_SIZEY)
+
+// Configure base width of the LDS tile according to screen percentage range.
+// If screen percentage is lower than 75%, we can set the base tile width to 6, which will subsequently create a scene color
+// tile width of 8 that can be cache with only one iteration per lane in PrecacheInputSceneColorToLDS(), and reduce size(LDS) = 1k.
+#if TAA_SCREEN_PERCENTAGE_RANGE == 1 // screen percentage < 75%
+	#define LDS_BASE_TILE_WIDTH 6
+#elif TAA_SCREEN_PERCENTAGE_RANGE == 0 // screen percentage in [75%; 100%]
+	#define LDS_BASE_TILE_WIDTH THREADGROUP_SIZEX
+#elif TAA_SCREEN_PERCENTAGE_RANGE == 2 // screen percentage > 100%
+	#error Should not use LDS caching.
+#else
+	#error Unknown screen percentage range.
+#endif
+
+
+//-------------------------------- Configuration.
+
+// Configuration of what should be prefetched.
+// 1: use Load; 2: use gather4. 
+
+#if !AA_UPSAMPLE
+// Disables scene depth caching for TAA upsample because the extra screen percentage ALU is making things worst.
+#define AA_PRECACHE_SCENE_DEPTH 2
+#endif
+
+// 1: use load.
+#define AA_PRECACHE_SCENE_COLOR 1
+
+// Precache GetSceneColorHdrWeight() into scene color's alpha channel.
+#define AA_PRECACHE_SCENE_HDR_WEIGHT (AA_TONE && HISTORY_PAYLOAD_COMPONENTS == 3)
+
+// Layout of scene color in LDS.
+// 0: AoS
+// 1: SoA
+// 2: AoSoA (SoA scene color, AoS GetSceneColorHdrWeight())
+#define LDS_COLOR_LAYOUT 0
+
+
+//-------------------------------- Depth tile constants.
+
+// Number of texels arround the group tile for depth.
+#define LDS_DEPTH_TILE_BORDER_SIZE (AA_CROSS)
+
+// Width in texels of the depth tile cached into LDS.
+#define LDS_DEPTH_TILE_WIDTH (LDS_BASE_TILE_WIDTH + 2 * LDS_DEPTH_TILE_BORDER_SIZE)
+
+// Total number of texels cached in the depth tile.
+#define LDS_DEPTH_ARRAY_SIZE (LDS_DEPTH_TILE_WIDTH * LDS_DEPTH_TILE_WIDTH)
+
 
+//-------------------------------- Scene color tile constants.
 
-	// Sample and transform scene color.
-	FTAASceneColorSample SampleCachedSceneColorTexture(in FTAAInputParameters InputParams, uint NeighborId)
+// TODO: shader permutation for screen percentage <= 75% with AA_UPSAMPLE to fit in 1k LDS.
+
+// Number of scene color component that gets cached.
+#if HISTORY_PAYLOAD_COMPONENTS == 4 || AA_PRECACHE_SCENE_HDR_WEIGHT
+#define LDS_COLOR_COMPONENT_COUNT 4
+#else
+#define LDS_COLOR_COMPONENT_COUNT 3
+#endif
+
+// Number of texels arround the group tile for scene color.
+#define LDS_COLOR_TILE_BORDER_SIZE (1)
+
+// Width in texels of the depth tile cached into LDS.
+#define LDS_COLOR_TILE_WIDTH (LDS_BASE_TILE_WIDTH + 2 * LDS_COLOR_TILE_BORDER_SIZE)
+
+// Total number of texels cached in the scene color tile.
+#define LDS_COLOR_ARRAY_SIZE (LDS_COLOR_TILE_WIDTH * LDS_COLOR_TILE_WIDTH)
+
+
+//-------------------------------- Group shared global.
+
+// Size of the LDS to be allocated.
+#define LDS_ARRAY_SIZE (LDS_COLOR_ARRAY_SIZE * LDS_COLOR_COMPONENT_COUNT)
+#if LDS_ARRAY_SIZE < LDS_DEPTH_ARRAY_SIZE
+	#error LDS_ARRAY_SIZE assumed scene color caching is bigger than scene depth caching.
+#endif
+
+// Some compilers may have issues optimising LDS store instructions, therefore we give the compiler a hint by using a float4 LDS.
+#if defined(AA_PRECACHE_SCENE_DEPTH)
+#define LDS_USE_FLOAT4_ARRAY 0
+#else
+#define LDS_USE_FLOAT4_ARRAY (LDS_COLOR_COMPONENT_COUNT == 4 && LDS_COLOR_LAYOUT == 0)
+#endif
+
+#if LDS_USE_FLOAT4_ARRAY
+groupshared float4 GroupSharedArrayF4[LDS_ARRAY_SIZE/4];
+#else
+groupshared float GroupSharedArray[LDS_ARRAY_SIZE];
+#endif
+
+
+//-------------------------------- Generic LDS tile functions.
+
+#if AA_UPSAMPLE
+
+// Get the pixel coordinate of the nearest input pixel K for group's thread 0.
+float2 GetGroupThread0InputPixelCoord(in FTAAInputParameters InputParams)
+{
+	// Output pixel center position of the group thread index 0, relative to top left corner of the viewport.
+	float2 Thread0SvPosition = InputParams.GroupId * uint2(THREADGROUP_SIZEX, THREADGROUP_SIZEY) + 0.5;
+		
+	// Output pixel's viewport UV group thread index 0.
+	float2 Thread0ViewportUV = Thread0SvPosition * OutputViewportSize.zw;
+
+	// Pixel coordinate of the center of output pixel O in the input viewport.
+	float2 Thread0PPCo = Thread0ViewportUV * InputViewSize.xy + TemporalJitterPixels;
+	
+	// Pixel coordinate of the center of the nearest input pixel K.
+	float2 Thread0PPCk = floor(Thread0PPCo) + 0.5;
+		
+	return InputViewMin.xy + Thread0PPCk;
+}
+
+#endif
+
+// Get the texel offset of a LDS tile's top left corner.
+uint2 GetGroupTileTexelOffset(in FTAAInputParameters InputParams, uint TileBorderSize)
+{
+	#if AA_UPSAMPLE
 	{
-		FTAASceneColorSample Sample;
-		Sample.SampleOffset = GetNeighborSampleOffset(InputParams, NeighborId);
+		// Pixel coordinate of the center of the nearest input pixel K.
+		float2 Thread0PPCk = GetGroupThread0InputPixelCoord(InputParams);
+		
+		return uint2(floor(Thread0PPCk) - TileBorderSize);
+	}
+	#else // !AA_UPSAMPLE
+	{
+		return OutputViewportRect.xy + InputParams.GroupId * uint2(THREADGROUP_SIZEX, THREADGROUP_SIZEY) - TileBorderSize;
+	}
+	#endif
+}
+
+// Get the index within the LDS array.
+uint GetTileArrayIndexFromPixelOffset(in FTAAInputParameters InputParams, int2 PixelOffset, uint TileBorderSize)
+{
+	#if AA_UPSAMPLE
+	{
+		const float2 RowMultiplier = float2(1, TileBorderSize * 2 + LDS_BASE_TILE_WIDTH);
+
+		float2 Thread0PPCk = GetGroupThread0InputPixelCoord(InputParams);
+		float2 PPCk = InputParams.NearestBufferUV * InputSceneColorSize.xy;
+
+		float2 TilePos = floor(PPCk) - floor(Thread0PPCk);
+		return uint(dot(TilePos, RowMultiplier) + dot(float2(PixelOffset) + float(TileBorderSize), RowMultiplier));
+	}
+	#else
+	{
+		uint2 TilePos = InputParams.GroupThreadId + uint2(PixelOffset + TileBorderSize);
+		return TilePos.x + TilePos.y * (TileBorderSize * 2 + LDS_BASE_TILE_WIDTH);
+	}
+	#endif
+}
+
 
-		int2 Coord = int2(InputParams.NearestBufferUV * InputSceneColorSize.xy) + Sample.SampleOffset;
+//-------------------------------- Share depth texture fetches.
+#if defined(AA_PRECACHE_SCENE_DEPTH)
+
+// Precache input scene depth into LDS.
+void PrecacheInputSceneDepthToLDS(in FTAAInputParameters InputParams)
+{
+	uint2 GroupTexelOffset = GetGroupTileTexelOffset(InputParams, LDS_DEPTH_TILE_BORDER_SIZE);
+	
+	#if AA_PRECACHE_SCENE_DEPTH == 1
+	// Prefetch depth buffer using Load.
+	{
+		const uint LoadCount = (LDS_DEPTH_ARRAY_SIZE + THREADGROUP_TOTAL - 1) / THREADGROUP_TOTAL;
 		
-		#if AA_BORDER == 2
-			Coord = clamp(Coord, InputMinPixelCoord, InputMaxPixelCoord);
-		#endif
+		uint LinearGroupThreadId = GetGroupIndex(InputParams);
 
-		Sample.Payload.Color = TransformCurrentFrameSceneColor(InputSceneColor[Coord]);
-		Sample.Payload.CocRadius = 0;
-		#if AA_HISTORY_PAYLOAD == HISTORY_PAYLOAD_RGB_COC
-			Sample.Payload.CocRadius = Sample.Color.a;
-		#elif AA_HISTORY_PAYLOAD == HISTORY_PAYLOAD_RGB_OPACITY_COC
-			Sample.Payload.CocRadius = InputSceneMetadata[Coord].r;
+		UNROLL
+		for (uint i = 0; i < LoadCount; i++)
+		{
+			uint2 TexelLocation =  GroupTexelOffset + uint2(
+				LinearGroupThreadId % LDS_DEPTH_TILE_WIDTH,
+				LinearGroupThreadId / LDS_DEPTH_TILE_WIDTH);
+
+			if ((LinearGroupThreadId < LDS_DEPTH_ARRAY_SIZE) ||
+				(i != LoadCount - 1) ||
+				(LDS_DEPTH_ARRAY_SIZE % THREADGROUP_TOTAL) == 0)
+			{
+				GroupSharedArray[LinearGroupThreadId] = SceneDepthBuffer.Load(uint3(TexelLocation, 0)).x;
+			}
+			LinearGroupThreadId += THREADGROUP_TOTAL;
+		}
+
+	}
+
+	#elif AA_PRECACHE_SCENE_DEPTH == 2
+	// Prefetch depth buffer using Gather.
+	{
+		const uint LoadCount = (LDS_DEPTH_ARRAY_SIZE / 4 + THREADGROUP_TOTAL - 1) / THREADGROUP_TOTAL;
+		
+		uint LinearGroupThreadId = GetGroupIndex(InputParams);
+	
+		UNROLL
+		for (uint i = 0; i < LoadCount; i++)
+		{
+			uint2 TileDest = uint2(
+				(2 * LinearGroupThreadId) % LDS_DEPTH_TILE_WIDTH,
+				2 * ((2 * LinearGroupThreadId) / LDS_DEPTH_TILE_WIDTH));
+
+			uint2 TexelLocation =  GroupTexelOffset + TileDest;
+
+			uint DestI = TileDest.x + TileDest.y * LDS_DEPTH_TILE_WIDTH;
+
+			if ((DestI < LDS_DEPTH_ARRAY_SIZE) ||
+				(i != LoadCount - 1) ||
+				((LDS_DEPTH_ARRAY_SIZE / 4) % THREADGROUP_TOTAL) == 0)
+			{
+				float2 UV = float2(TexelLocation + 0.5) * InputSceneColorSize.zw;
+				float4 Depth = SceneDepthBuffer.Gather(SceneDepthBufferSampler, UV);
+				GroupSharedArray[DestI + 1 * LDS_DEPTH_TILE_WIDTH + 0] = Depth.x;
+				GroupSharedArray[DestI + 1 * LDS_DEPTH_TILE_WIDTH + 1] = Depth.y;
+				GroupSharedArray[DestI + 0 * LDS_DEPTH_TILE_WIDTH + 1] = Depth.z;
+				GroupSharedArray[DestI + 0 * LDS_DEPTH_TILE_WIDTH + 0] = Depth.w;
+			}
+			LinearGroupThreadId += THREADGROUP_TOTAL;
+		}
+	}
+
+	#else
+		#error Wrong AA_PRECACHE_SCENE_DEPTH
+
+	#endif
+}
+
+float SampleCachedSceneDepthTexture(in FTAAInputParameters InputParams, int2 PixelOffset)
+{
+	return GroupSharedArray[GetTileArrayIndexFromPixelOffset(InputParams, PixelOffset, LDS_DEPTH_TILE_BORDER_SIZE)];
+}
+
+
+#endif // define(AA_PRECACHE_SCENE_DEPTH)
+
+
+//-------------------------------- Share color texture fetches.
+
+#if defined(AA_PRECACHE_SCENE_COLOR)
+
+// Return the index GroupSharedArray from a given ArrayIndex and ComponentId.
+uint GetSceneColorLDSIndex(uint ArrayIndex, uint ComponentId)
+{
+	#if LDS_COLOR_LAYOUT == 0 // AoS
+		return ArrayIndex * LDS_COLOR_COMPONENT_COUNT + ComponentId;
+	#elif LDS_COLOR_LAYOUT == 1 // SoA
+		return ArrayIndex + ComponentId * LDS_COLOR_ARRAY_SIZE;
+	#else
+		#error Unknown color layout.
+	#endif
+}
+
+// Precache input scene color into LDS.
+void PrecacheInputSceneColorToLDS(in FTAAInputParameters InputParams)
+{
+	const uint LoadCount = (LDS_COLOR_ARRAY_SIZE + THREADGROUP_TOTAL - 1) / THREADGROUP_TOTAL;
+	
+	#define LDS_FLOAT_UV AA_UPSAMPLE
+
+	#if LDS_FLOAT_UV
+		float LinearGroupThreadId = float(GetGroupIndex(InputParams));
+		float2 Thread0PPCk = GetGroupThread0InputPixelCoord(InputParams);
+		float2 GroupTexelOffset = Thread0PPCk - LDS_COLOR_TILE_BORDER_SIZE;
+	#else
+		uint LinearGroupThreadId = GetGroupIndex(InputParams);
+		uint2 GroupTexelOffset = GetGroupTileTexelOffset(InputParams, LDS_COLOR_TILE_BORDER_SIZE);
+	#endif
+	
+	UNROLL
+	for (uint i = 0; i < LoadCount; i++)
+	{
+		#if LDS_FLOAT_UV
+			float Y = floor(LinearGroupThreadId * (1.0 / LDS_COLOR_TILE_WIDTH));
+			float X = LinearGroupThreadId - LDS_COLOR_TILE_WIDTH * Y;
+
+			float2 TexelLocation = GroupTexelOffset + float2(X, Y);
+		#else
+			uint2 TexelLocation =  GroupTexelOffset + uint2(
+				LinearGroupThreadId % LDS_COLOR_TILE_WIDTH,
+				LinearGroupThreadId / LDS_COLOR_TILE_WIDTH);
 		#endif
 
-		Sample.HdrWeight = GetSceneColorHdrWeight(InputParams, Sample.Payload.Color);
-		return Sample;
+		if ((LinearGroupThreadId < LDS_COLOR_ARRAY_SIZE) ||
+			(i != LoadCount - 1) ||
+			(LDS_COLOR_ARRAY_SIZE % THREADGROUP_TOTAL) == 0)
+		{
+			#if LDS_FLOAT_UV
+				int2 Coord = TexelLocation;
+				Coord = clamp(Coord, InputMinPixelCoord, InputMaxPixelCoord);
+
+				float4 RawColor = InputSceneColor[Coord];
+			#else
+				int2 Coord = int2(TexelLocation);
+				Coord = clamp(Coord, InputMinPixelCoord, InputMaxPixelCoord);
+				float4 RawColor = InputSceneColor.Load(uint3(Coord, 0));
+			#endif
+
+			float4 Color = TransformCurrentFrameSceneColor(RawColor);
+
+			// Precache scene color's HDR weight into alpha channel to reduce rcp() instructions in innerloops.
+			#if AA_PRECACHE_SCENE_HDR_WEIGHT
+				Color.a = GetSceneColorHdrWeight(InputParams, Color);
+			#endif
+
+			#if LDS_USE_FLOAT4_ARRAY
+				GroupSharedArrayF4[uint(LinearGroupThreadId)] = Color;
+
+			#else
+				GroupSharedArray[GetSceneColorLDSIndex(uint(LinearGroupThreadId), 0)] = Color.r;
+				GroupSharedArray[GetSceneColorLDSIndex(uint(LinearGroupThreadId), 1)] = Color.g;
+				GroupSharedArray[GetSceneColorLDSIndex(uint(LinearGroupThreadId), 2)] = Color.b;
+				#if LDS_COLOR_COMPONENT_COUNT == 4
+					GroupSharedArray[GetSceneColorLDSIndex(uint(LinearGroupThreadId), 3)] = Color.a;
+				#endif
+			
+			#endif
+		}
+		LinearGroupThreadId += THREADGROUP_TOTAL;
 	}
+}
+
+FTAASceneColorSample SampleCachedSceneColorTexture(
+	in FTAAInputParameters InputParams,
+	int2 PixelOffset)
+{
+	uint ArrayPos = GetTileArrayIndexFromPixelOffset(InputParams, PixelOffset, LDS_COLOR_TILE_BORDER_SIZE);
+	
+	FTAASceneColorSample Sample;
+	Sample.CocRadius = 0;
+	
+	#if AA_HISTORY_PAYLOAD == HISTORY_PAYLOAD_RGB_COC || AA_HISTORY_PAYLOAD == HISTORY_PAYLOAD_RGB_OPACITY_COC
+		#error Unsupported history payload with LDS scene color caching.Test
+	#endif
 
-#endif // AA_SAMPLE_CACHE_METHOD != AA_SAMPLE_CACHE_METHOD_VGPR_3X3
+	#if LDS_USE_FLOAT4_ARRAY
+		Sample.Color = GroupSharedArrayF4[ArrayPos];
 
-// Sample depth.
-half SampleSceneDepthTexture(in FTAAInputParameters InputParams, int2 PixelOffset)
+		#if AA_PRECACHE_SCENE_HDR_WEIGHT
+			Sample.HdrWeight = Sample.Color.a;
+			Sample.Color.a = 0;
+		#elif HISTORY_PAYLOAD_COMPONENTS != 4
+			#error LDS_USE_FLOAT4_ARRAY assumes 4 components.
+		#endif
+
+	#else
+		Sample.Color.r = GroupSharedArray[GetSceneColorLDSIndex(ArrayPos, 0)];
+		Sample.Color.g = GroupSharedArray[GetSceneColorLDSIndex(ArrayPos, 1)];
+		Sample.Color.b = GroupSharedArray[GetSceneColorLDSIndex(ArrayPos, 2)];
+		Sample.Color.a = 0;
+	
+		#if HISTORY_PAYLOAD_COMPONENTS == 4
+			Sample.Color.a = GroupSharedArray[GetSceneColorLDSIndex(ArrayPos, 3)];
+		#elif AA_PRECACHE_SCENE_HDR_WEIGHT
+			Sample.HdrWeight = GroupSharedArray[GetSceneColorLDSIndex(ArrayPos, 3)];
+		#endif
+
+	#endif
+	
+	// if scene color weight was not precached in LDS, compute it.
+	#if !AA_PRECACHE_SCENE_HDR_WEIGHT
+		Sample.HdrWeight = GetSceneColorHdrWeight(InputParams, Sample.Color);
+	#endif
+	
+	// Color has already been transformed in PrecacheInputSceneColor.
+	return Sample;
+}
+
+
+#endif // defined(AA_PRECACHE_SCENE_COLOR)
+
+
+void PrecacheInputSceneDepth(in FTAAInputParameters InputParams)
 {
+	#if defined(AA_PRECACHE_SCENE_DEPTH)
+		PrecacheInputSceneDepthToLDS(InputParams);
+		GroupMemoryBarrierWithGroupSync();
+	#endif
+}
+
+void PrecacheInputSceneColor(in FTAAInputParameters InputParams)
+{
+	#if defined(AA_PRECACHE_SCENE_DEPTH) && defined(AA_PRECACHE_SCENE_COLOR)
+		GroupMemoryBarrierWithGroupSync();
+	#endif
+	#if defined(AA_PRECACHE_SCENE_COLOR)
+		PrecacheInputSceneColorToLDS(InputParams);
+		GroupMemoryBarrierWithGroupSync();
+	#endif
+}
+
+
+#endif // AA_SAMPLE_CACHE_METHOD == AA_SAMPLE_CACHE_METHOD_LDS
+
+//------------------------------------------------------- FALLBACK TO NO CACHING IMPLEMENTATIONS.
+
+#if !defined(AA_PRECACHE_SCENE_DEPTH)
+
+#if !defined(AA_PRECACHE_IMPLEMENTATIONS)
+
+// Silently do no scene depth precaching.
+void PrecacheInputSceneDepth(in FTAAInputParameters InputParams)
+{ }
+
+#endif
+
+
+// Sample scene color.
+float SampleCachedSceneDepthTexture(in FTAAInputParameters InputParams, int2 PixelOffset)
+{
+	return SceneDepthBuffer.SampleLevel(SceneDepthBufferSampler, InputParams.NearestBufferUV, 0, PixelOffset).r;
+}
+
+
+#endif
+
+#if !defined(AA_PRECACHE_SCENE_COLOR)
+
+#if !defined(AA_PRECACHE_IMPLEMENTATIONS)
+
+// Silently do no scene color precaching.
+void PrecacheInputSceneColor(in FTAAInputParameters InputParams)
+{ }
+
+#endif
+
+
+// Sample and transform scene color.
+FTAASceneColorSample SampleCachedSceneColorTexture(
+	in FTAAInputParameters InputParams,
+	int2 PixelOffset)
+{
+	FTAASceneColorSample Sample;
+
+
 	int2 Coord = int2(InputParams.NearestBufferUV * InputSceneColorSize.xy) + PixelOffset;
+	Coord = clamp(Coord, InputMinPixelCoord, InputMaxPixelCoord);
+
+	Sample.Color = TransformCurrentFrameSceneColor(InputSceneColor[Coord]);
 
-	#if AA_BORDER == 2
-		Coord = clamp(Coord, InputMinPixelCoord, InputMaxPixelCoord);
+	Sample.CocRadius = 0;
+
+	#if AA_HISTORY_PAYLOAD == HISTORY_PAYLOAD_RGB_COC
+		Sample.CocRadius = Sample.Color.a;
+	#elif AA_HISTORY_PAYLOAD == HISTORY_PAYLOAD_RGB_OPACITY_COC
+		Sample.CocRadius = InputSceneMetadata[Coord].r;
 	#endif
 
-	return SceneDepthBuffer[Coord].r;
+	Sample.HdrWeight = GetSceneColorHdrWeight(InputParams, Sample.Color);
+	return Sample;
 }
 
 
+#endif
+
+
 //------------------------------------------------------- TEMPORAL UPSAMPLING
 
 #if AA_UPSAMPLE
 
 // Returns the weight of a pixels at a coordinate <PixelDelta> from the PDF highest point.
-half ComputeSampleWeigth(in FTAAIntermediaryResult IntermediaryResult, float2 PixelDelta)
+float ComputeSampleWeigth(in FTAAIntermediaryResult IntermediaryResult, float2 PixelDelta)
 {
-	half u2 = UpscaleFactor * UpscaleFactor; // TODO: shader parameter
+	float u2 = UpscaleFactor * UpscaleFactor;
 
 	// The point of InvFilterScaleFactor is to blur current frame scene color when upscaling.
 	// Therefore there is no need to do it when downscaling.
@@ -950,15 +1359,15 @@ half ComputeSampleWeigth(in FTAAIntermediaryResult IntermediaryResult, float2 Pi
 
 	#if 1
 		// 1 - 1.9 * x^2 + 0.9 * x^4
-		half x2 = saturate(u2 * dot(PixelDelta, PixelDelta));
+		float x2 = saturate(u2 * dot(PixelDelta, PixelDelta));
 		return (0.905 * x2 - 1.9) * x2 + 1;
 
 	#else
 		// original e ^ (- x^2 / (2 * s^2))
-		const half Sigma = 0.47;
-		const half ExponentInputFactor = (-0.5 / (Sigma * Sigma));
+		const float Sigma = 0.47;
+		const float ExponentInputFactor = (-0.5 / (Sigma * Sigma));
 		
-		half x2 = dot(PixelDelta, PixelDelta) * u2;
+		float x2 = dot(PixelDelta, PixelDelta) * u2;
 		return exp(ExponentInputFactor * x2);
 
 	#endif
@@ -966,9 +1375,9 @@ half ComputeSampleWeigth(in FTAAIntermediaryResult IntermediaryResult, float2 Pi
 
 
 // Returns the weight of a pixels at a coordinate <PixelDelta> from the PDF highest point.
-half ComputePixelWeigth(in FTAAIntermediaryResult IntermediaryResult, half2 PixelDelta)
+float ComputePixelWeigth(in FTAAIntermediaryResult IntermediaryResult, float2 PixelDelta)
 {
-	half u2 = UpscaleFactor * UpscaleFactor;
+	float u2 = UpscaleFactor * UpscaleFactor;
 
 	// The point of InvFilterScaleFactor is to blur current frame scene color when upscaling.
 	// Therefore there is no need to do it when downscaling.
@@ -979,16 +1388,16 @@ half ComputePixelWeigth(in FTAAIntermediaryResult IntermediaryResult, half2 Pixe
 
 	#if 1
 		// 1 - 1.9 * x^2 + 0.9 * x^4
-		half x2 = saturate(u2 * dot(PixelDelta, PixelDelta));
-		half r = (0.905 * x2 - 1.9) * x2 + 1;
+		float x2 = saturate(u2 * dot(PixelDelta, PixelDelta));
+		float r = (0.905 * x2 - 1.9) * x2 + 1;
 
 	#else
 		// original e ^ (- x^2 / (2 * s^2))
-		const half Sigma = 0.47;
-		const half ExponentInputFactor = (-0.5 / (Sigma * Sigma));
+		const float Sigma = 0.47;
+		const float ExponentInputFactor = (-0.5 / (Sigma * Sigma));
 		
-		half x2 = dot(PixelDelta, PixelDelta);
-		half r = exp(ExponentInputFactor * x2);
+		float x2 = dot(PixelDelta, PixelDelta);
+		float r = exp(ExponentInputFactor * x2);
 
 	#endif
 
@@ -1013,12 +1422,13 @@ void FilterCurrentFrameInputSamples(
 {
 	#if !AA_FILTERED
 	{
-		IntermediaryResult.Filtered = SampleCachedSceneColorTexture(InputParams, /* NeighborId = */ 0).Payload;
+		IntermediaryResult.Filtered.Color = SampleCachedSceneColorTexture(InputParams, int2(0, 0)).Color;
+		IntermediaryResult.Filtered.CocRadius = SampleCachedSceneColorTexture(InputParams, int2(0, 0)).CocRadius;
 		return;
 	}
 	#endif
 
-	FTAAPayload Filtered;
+	FTAAHistoryPayload Filtered;
 
 	{
 		#if AA_UPSAMPLE
@@ -1029,62 +1439,93 @@ void FilterCurrentFrameInputSamples(
 			float2 PPCk = floor(PPCo) + 0.5;
 		
 			// Vector in pixel between pixel K -> O.
-			half2 dKO = half2(PPCo - PPCk);
+			float2 dKO = PPCo - PPCk;
+		#endif
+		
+		#if AA_SAMPLES == 9
+			const uint SampleIndexes[9] = kSquareIndexes3x3;
+		#elif AA_SAMPLES == 5 || AA_SAMPLES == 6
+			const uint SampleIndexes[5] = kPlusIndexes3x3;
 		#endif
 
 		#if AA_HISTORY_PAYLOAD == HISTORY_PAYLOAD_RGB_COC || AA_HISTORY_PAYLOAD == HISTORY_PAYLOAD_RGB_OPACITY_COC
 			// Fetches center pixel's Coc for the bilateral filtering.
-			half CenterCocRadius = SampleCachedSceneColorTexture(InputParams, /* NeighborId = */ 0).Payload.CocRadius;
+			float CenterCocRadius = SampleCachedSceneColorTexture(InputParams, int2(0, 0)).CocRadius;
 		#endif
 
-		half NeighborsHdrWeight = 0;
-		half NeighborsFinalWeight = 0;
-		half4 NeighborsColor = 0;
+		float NeighborsHdrWeight = 0;
+		float NeighborsFinalWeight = 0;
+		float4 NeighborsColor = 0;
 
 		UNROLL
-		for (uint NeighborId = 0; NeighborId < AA_SAMPLES; NeighborId++)
+		for (uint i = 0; i < AA_SAMPLES; i++)
 		{
-			FTAASceneColorSample Sample = SampleCachedSceneColorTexture(InputParams, NeighborId);
+			// Get the sample offset from the nearest input pixel.
+			int2 SampleOffset;
+			
+			#if AA_UPSAMPLE && AA_SAMPLES == 6
+				if (i == 5)
+				{
+					SampleOffset = SignFastInt(dKO);
+				}
+				else
+			#endif
+			{
+				const uint SampleIndex = SampleIndexes[i];
+				SampleOffset = kOffsets3x3[SampleIndex];
+			}
+			float2 fSampleOffset = float2(SampleOffset);
+			
+			// When doing Coc bilateral, the center sample is accumulated last.
+			#if AA_HISTORY_PAYLOAD == HISTORY_PAYLOAD_RGB_COC && 0
+				if (all(SampleOffset == 0) && (AA_SAMPLES != 6 || i != 5))
+				{
+					continue;
+				}
+			#endif
 
 			// Finds out the spatial weight of this input sample.
 			#if AA_UPSAMPLE
 				// Compute the pixel delta between output pixels and input pixel I.
 				//  Note: abs() is unecessary because of the dot(dPP, dPP) latter on.
-				half2 dPP = half2(Sample.SampleOffset) - dKO;
+				float2 dPP = fSampleOffset - dKO;
 
-				half SampleSpatialWeight = ComputeSampleWeigth(IntermediaryResult, dPP);
+				float SampleSpatialWeight = ComputeSampleWeigth(IntermediaryResult, dPP);
 
 			#elif AA_SAMPLES == 9
-				half SampleSpatialWeight = SampleWeights[NeighborId];
+				float SampleSpatialWeight = SampleWeights[i];
 
 			#elif AA_SAMPLES == 5
-				half SampleSpatialWeight = PlusWeights[NeighborId];
+				float SampleSpatialWeight = PlusWeights[i];
 
 			#else
 				#error Do not know how to compute filtering sample weight.
 
 			#endif
 
+			// Fetch sample.
+			FTAASceneColorSample Sample = SampleCachedSceneColorTexture(InputParams, SampleOffset);
+				
 			// Finds out the sample's HDR weight.
 			#if AA_TONE
-				half SampleHdrWeight = Sample.HdrWeight;
+				float SampleHdrWeight = Sample.HdrWeight;
 			#else
-				half SampleHdrWeight = 1;
+				float SampleHdrWeight = 1;
 			#endif
 
 			// Finds out the sample's bilateral weight according to the payload.
 			#if AA_HISTORY_PAYLOAD == HISTORY_PAYLOAD_RGB_COC
-				half BilateralWeight = ComputeNeightborSampleBilateralWeight(CenterCocRadius, Sample.Payload.CocRadius);
+				float BilateralWeight = ComputeNeightborSampleBilateralWeight(CenterCocRadius, Sample.CocRadius);
 
 			#else
-				half BilateralWeight = 1;
+				float BilateralWeight = 1;
 
 			#endif
 
-			half SampleFinalWeight = SampleSpatialWeight * SampleHdrWeight * BilateralWeight;
+			float SampleFinalWeight = SampleSpatialWeight * SampleHdrWeight * BilateralWeight;
 
 			// Apply pixel.
-			NeighborsColor       += SampleFinalWeight * Sample.Payload.Color;
+			NeighborsColor       += SampleFinalWeight * Sample.Color;
 			NeighborsFinalWeight += SampleFinalWeight;
 
 			NeighborsHdrWeight   += SampleSpatialWeight * SampleHdrWeight;
@@ -1092,16 +1533,16 @@ void FilterCurrentFrameInputSamples(
 		
 		#if AA_HISTORY_PAYLOAD == HISTORY_PAYLOAD_RGB_COC || AA_HISTORY_PAYLOAD == HISTORY_PAYLOAD_RGB_OPACITY_COC
 		{
-			FTAASceneColorSample Sample = SampleCachedSceneColorTexture(InputParams, /* NeighborId = */ 0);
+			FTAASceneColorSample Sample = SampleCachedSceneColorTexture(InputParams, 0);
 			
 			#if AA_UPSAMPLE
-				half SampleWeight = ComputeSampleWeigth(IntermediaryResult, -dKO);
+				float SampleWeight = ComputeSampleWeigth(IntermediaryResult, -dKO);
 
 			#elif AA_SAMPLES == 9
-				half SampleWeight = SampleWeights[0];
+				float SampleWeight = SampleWeights[4];
 
 			#else
-				half SampleWeight = PlusWeights[0];
+				float SampleWeight = PlusWeights[2];
 
 			#endif
 			
@@ -1133,24 +1574,6 @@ void FilterCurrentFrameInputSamples(
 			IntermediaryResult.FilteredTemporalWeight = ComputePixelWeigth(IntermediaryResult, dKO);
 		#endif
 	}
-	
-	#if AA_BORDER == 1
-	{
-		#if AA_UPSAMPLE
-			#error AA_BORDER == 1 does not work for TAA upsample because also needs to be done in color clamping box.
-		#else
-			// Use unfiltered for 1 pixel border.
-			float2 TestPos = abs(InputParams.ScreenPos) + OutputViewportSize.zw * 2.0;
-			
-			if(max(TestPos.x, TestPos.y) >= 1.0) 
-			{
-				Filtered = SampleCachedSceneColorTexture(InputParams, /* NeighborId = */ 0).Payload;
-			}
-		#endif
-	}
-	#endif
-
-	Filtered.Color.rgb = clamp(Filtered.Color.rgb, -half(MaxHalfFloat), half(MaxHalfFloat));
 
 	IntermediaryResult.Filtered = Filtered;
 }
@@ -1160,29 +1583,45 @@ void FilterCurrentFrameInputSamples(
 void ComputeNeighborhoodBoundingbox(
 	in FTAAInputParameters InputParams,
 	in FTAAIntermediaryResult IntermediaryResult,
-	out FTAAPayload OutNeighborMin,
-	out FTAAPayload OutNeighborMax)
+	out FTAAHistoryPayload OutNeighborMin,
+	out FTAAHistoryPayload OutNeighborMax)
 {
-	FTAAPayload NeighborMin;
-	FTAAPayload NeighborMax;
+	// TODO: clean this up.
+	FTAAHistoryPayload Neighbors[kNeighborsCount];
+	UNROLL
+	for (uint i = 0; i < kNeighborsCount; i++)
+	{
+		Neighbors[i].Color = SampleCachedSceneColorTexture(InputParams, kOffsets3x3[i]).Color;
+		Neighbors[i].CocRadius = SampleCachedSceneColorTexture(InputParams, kOffsets3x3[i]).CocRadius;
+	}
+
+	FTAAHistoryPayload NeighborMin;
+	FTAAHistoryPayload NeighborMax;
 
 	#if AA_HISTORY_CLAMPING_BOX == HISTORY_CLAMPING_BOX_VARIANCE
 	{
-		half4 m1 = 0;
-		half4 m2 = 0;
-		UNROLL
-		for(uint NeighborId = 0; NeighborId < AA_SAMPLES; NeighborId++)
+		#if AA_SAMPLES == 9
+			const uint SampleIndexes[9] = kSquareIndexes3x3;
+		#elif AA_SAMPLES == 5
+			const uint SampleIndexes[5] = kPlusIndexes3x3;
+		#else
+			#error Unknown number of samples.
+		#endif
+
+		float4 m1 = 0;
+		float4 m2 = 0;
+		for( uint i = 0; i < AA_SAMPLES; i++ )
 		{
-			FTAASceneColorSample Sample = SampleCachedSceneColorTexture(InputParams, NeighborId);
+			float4 SampleColor = Neighbors[ SampleIndexes[i] ];
 
-			m1 += Sample.Payload.Color;
-			m2 += Pow2(Sample.Payload.Color);
+			m1 += SampleColor;
+			m2 += Pow2( SampleColor );
 		}
 
-		m1 *= half(1.0 / AA_SAMPLES);
-		m2 *= half(1.0 / AA_SAMPLES);
+		m1 *= (1.0 / AA_SAMPLES);
+		m2 *= (1.0 / AA_SAMPLES);
 
-		half4 StdDev = sqrt( abs(m2 - m1 * m1) );
+		float4 StdDev = sqrt( abs(m2 - m1 * m1) );
 		NeighborMin = m1 - 1.25 * StdDev;
 		NeighborMax = m1 + 1.25 * StdDev;
 
@@ -1194,51 +1633,80 @@ void ComputeNeighborhoodBoundingbox(
 	{
 		float2 PPCo = InputParams.ViewportUV * InputViewSize.xy + TemporalJitterPixels;
 		float2 PPCk = floor(PPCo) + 0.5;
-		half2 dKO = half2(PPCo - PPCk);
+		float2 dKO = PPCo - PPCk;
 		
 		// Sample 4 is is always going to be considered anyway.
-		{
-			FTAASceneColorSample Sample = SampleCachedSceneColorTexture(InputParams, /* NeighborId = */ 0);
-			NeighborMin = Sample.Payload;
-			NeighborMax = Sample.Payload;
-		}
+		NeighborMin = Neighbors[4];
+		NeighborMax = Neighbors[4];
 		
 		// Reduce distance threshold as upsacale factor increase to reduce ghosting.
-		half DistthresholdLerp = UpscaleFactor - 1;
-		half DistThreshold = lerp(1.51, 1.3, DistthresholdLerp); // TODO(REVERB): ShaderParameter
-		half DistThresholdSqr = ToScalarMemory(DistThreshold * DistThreshold);
+		float DistthresholdLerp = UpscaleFactor - 1;
+		float DistThreshold = lerp(1.51, 1.3, DistthresholdLerp);
+
+		#if AA_SAMPLES == 9
+			const uint Indexes[9] = kSquareIndexes3x3;
+		#else
+			const uint Indexes[5] = kPlusIndexes3x3;
+		#endif
 
 		UNROLL
-		for(uint NeighborId = 1; NeighborId < AA_SAMPLES; NeighborId++)
+		for( uint i = 0; i < AA_SAMPLES; i++ )
 		{
-			FTAASceneColorSample Sample = SampleCachedSceneColorTexture(InputParams, NeighborId);
-
-			half2 dPP = half2(Sample.SampleOffset) - dKO; // TODO(REVERB)
-
-			FLATTEN
-			if (dot(dPP, dPP) < DistThresholdSqr)
+			uint NeightborId = Indexes[i];
+			if (NeightborId != 4)
 			{
-				NeighborMin = MinPayload(NeighborMin, Sample.Payload);
-				NeighborMax = MaxPayload(NeighborMax, Sample.Payload);
+				float2 dPP = float2(kOffsets3x3[NeightborId]) - dKO;
+
+				FLATTEN
+				if (dot(dPP, dPP) < (DistThreshold * DistThreshold))
+				{
+					NeighborMin = MinPayload(NeighborMin, Neighbors[NeightborId]);
+					NeighborMax = MaxPayload(NeighborMax, Neighbors[NeightborId]);
+				}
 			}
 		}
 	}
 	#elif AA_HISTORY_CLAMPING_BOX == HISTORY_CLAMPING_BOX_MIN_MAX
 	{
+		NeighborMin = MinPayload3( Neighbors[1], Neighbors[3], Neighbors[4] );
+		NeighborMin = MinPayload3( NeighborMin,  Neighbors[5], Neighbors[7] );
+
+		NeighborMax = MaxPayload3( Neighbors[1], Neighbors[3], Neighbors[4] );
+		NeighborMax = MaxPayload3( NeighborMax,  Neighbors[5], Neighbors[7] );
+		
+		#if AA_SAMPLES == 6
 		{
-			FTAASceneColorSample Sample = SampleCachedSceneColorTexture(InputParams, /* NeighborId = */ 0);
-			NeighborMin = Sample.Payload;
-			NeighborMax = Sample.Payload;
+			float2 PPCo = InputParams.ViewportUV * InputViewSize.xy + TemporalJitterPixels;
+			float2 PPCk = floor(PPCo) + 0.5;
+			float2 dKO = PPCo - PPCk;
+			
+			int2 FifthNeighborOffset = SignFastInt(dKO);
+
+			FTAAHistoryPayload FifthNeighbor;
+			FifthNeighbor.Color = SampleCachedSceneColorTexture(InputParams, FifthNeighborOffset).Color;
+			FifthNeighbor.CocRadius = SampleCachedSceneColorTexture(InputParams, FifthNeighborOffset).CocRadius;
+			
+			NeighborMin = MinPayload(NeighborMin, FifthNeighbor);
+			NeighborMax = MaxPayload(NeighborMax, FifthNeighbor);
 		}
-		
-		UNROLL
-		for(uint NeighborId = 1; NeighborId < AA_SAMPLES; NeighborId++)
+		#elif AA_SAMPLES == 9
 		{
-			FTAASceneColorSample Sample = SampleCachedSceneColorTexture(InputParams, NeighborId);
+			FTAAHistoryPayload NeighborMinPlus = NeighborMin;
+			FTAAHistoryPayload NeighborMaxPlus = NeighborMax;
 
-			NeighborMin = MinPayload(NeighborMin, Sample.Payload);
-			NeighborMax = MaxPayload(NeighborMax, Sample.Payload);
+			NeighborMin = MinPayload3( NeighborMin, Neighbors[0], Neighbors[2] );
+			NeighborMin = MinPayload3( NeighborMin, Neighbors[6], Neighbors[8] );
+
+			NeighborMax = MaxPayload3( NeighborMax, Neighbors[0], Neighbors[2] );
+			NeighborMax = MaxPayload3( NeighborMax, Neighbors[6], Neighbors[8] );
+
+			if( AA_ROUND )
+			{
+				NeighborMin = AddPayload(MulPayload(NeighborMin, 0.5), MulPayload(NeighborMinPlus, 0.5));
+				NeighborMax = AddPayload(MulPayload(NeighborMax, 0.5), MulPayload(NeighborMaxPlus, 0.5));
+			}
 		}
+		#endif
 	}
 	#else
 		#error Unknown history clamping box.
@@ -1250,15 +1718,10 @@ void ComputeNeighborhoodBoundingbox(
 
 
 // Sample history.
-void SampleHistory(
-	in float2 HistoryScreenPosition,
-	out FTAAPayload HistoryPayload,
-	out FTAAHistoryInformation HistoryInfo)
+FTAAHistoryPayload SampleHistory(in float2 HistoryScreenPosition)
 {
-	half4 RawHistory0 = 0;
-	half4 RawHistory1 = 0;
-	half4 RawHistory2 = 0;
-	half4 RawHistory3 = 0;
+	float4 RawHistory0 = 0;
+	float4 RawHistory1 = 0;
 
 	// Sample the history using Catmull-Rom to reduce blur on motion.
 	#if AA_BICUBIC
@@ -1299,33 +1762,8 @@ void SampleHistory(
 			}
 
 			RawHistory0 += HistoryBuffer_0.SampleLevel(HistoryBufferSampler_0, SampleUV, 0) * Samples.Weight[i];
-			
-			#if HISTORY_RENDER_TARGETS >= 1
-				RawHistory1 += HistoryBuffer_1.SampleLevel(HistoryBufferSampler_1, SampleUV, 0) * Samples.Weight[i];
-			#endif
-			
-			#if HISTORY_RENDER_TARGETS >= 2
-				RawHistory2 += HistoryBuffer_2.SampleLevel(HistoryBufferSampler_2, SampleUV, 0) * Samples.Weight[i];
-			#endif
-			
-			#if HISTORY_RENDER_TARGETS >= 3
-				RawHistory3 += HistoryBuffer_3.SampleLevel(HistoryBufferSampler_3, SampleUV, 0) * Samples.Weight[i];
-			#endif
-		} // for (uint i = 0; i < Samples.Count; i++)
-
-		RawHistory0 *= half(Samples.FinalMultiplier);
-		
-		#if HISTORY_RENDER_TARGETS >= 1
-			RawHistory1 *= half(Samples.FinalMultiplier);
-		#endif
-		
-		#if HISTORY_RENDER_TARGETS >= 2
-			RawHistory2 *= half(Samples.FinalMultiplier);
-		#endif
-		
-		#if HISTORY_RENDER_TARGETS >= 3
-			RawHistory3 *= half(Samples.FinalMultiplier);
-		#endif
+		}
+		RawHistory0 *= Samples.FinalMultiplier;
 	}
 
 	// Sample the history using bilinear sampler.
@@ -1340,22 +1778,10 @@ void SampleHistory(
 		float2 HistoryBufferUV = HistoryScreenPosition * ScreenPosToHistoryBufferUV.xy + ScreenPosToHistoryBufferUV.zw;
 
 		RawHistory0 = HistoryBuffer_0.SampleLevel(HistoryBufferSampler_0, HistoryBufferUV, 0);
-		
-		#if HISTORY_RENDER_TARGETS >= 1
-			RawHistory1 = HistoryBuffer_1.SampleLevel(HistoryBufferSampler_1, HistoryBufferUV, 0);
-		#endif
-
-		#if HISTORY_RENDER_TARGETS >= 2
-			RawHistory2 = HistoryBuffer_2.SampleLevel(HistoryBufferSampler_2, HistoryBufferUV, 0);
-		#endif
-
-		#if HISTORY_RENDER_TARGETS >= 3
-			RawHistory3 = HistoryBuffer_3.SampleLevel(HistoryBufferSampler_3, HistoryBufferUV, 0);
-		#endif
 	}
 	#endif
 
-	#if AA_HISTORY_PAYLOAD == HISTORY_PAYLOAD_RGB_OPACITY_COC
+	#if HISTORY_RENDER_TARGETS == 2
 	{
 		if (AA_MANUALLY_CLAMP_HISTORY_UV)
 		{
@@ -1365,80 +1791,43 @@ void SampleHistory(
 		float2 HistoryBufferUV = HistoryScreenPosition * ScreenPosToHistoryBufferUV.xy + ScreenPosToHistoryBufferUV.zw;
 
 		RawHistory1 = HistoryBuffer_1.SampleLevel(HistoryBufferSampler_1, HistoryBufferUV, 0);
-
-		#if HISTORY_RENDER_TARGETS >= 3
-			RawHistory2 = HistoryBuffer_2.SampleLevel(HistoryBufferSampler_2, HistoryBufferUV, 0);
-		#endif
-
-		#if HISTORY_RENDER_TARGETS >= 4
-			RawHistory3 = HistoryBuffer_3.SampleLevel(HistoryBufferSampler_3, HistoryBufferUV, 0);
-		#endif
 	}
 	#endif
 	
-	// Decode history payload
-	{
-		HistoryPayload.Color = 0;
-		HistoryPayload.CocRadius = 0;
+	FTAAHistoryPayload HistoryPayload;
+	HistoryPayload.Color = RawHistory0;
 
-		#if AA_HISTORY_PAYLOAD == HISTORY_PAYLOAD_RGB
-		{
-			HistoryPayload.Color.rgb = RawHistory0.rgb;
-		}
-		#elif AA_HISTORY_PAYLOAD == HISTORY_PAYLOAD_RGB_OPACITY
-		{
-			HistoryPayload.Color = RawHistory0;
-		}
-		#elif AA_HISTORY_PAYLOAD == HISTORY_PAYLOAD_RGB_COC
-		{
-			HistoryPayload.Color.rgb = RawHistory0.rgb;
-			HistoryPayload.CocRadius = RawHistory0.a;
-		}
-		#elif AA_HISTORY_PAYLOAD == HISTORY_PAYLOAD_RGB_OPACITY_COC
-		{
-			HistoryPayload.Color = RawHistory0;
-			HistoryPayload.CocRadius = RawHistory1.r;
-		}
-		#else
-			#error Unknown history payload
-		#endif
-
-		#if USE_PREEXPOSURE
-			HistoryPayload.Color.rgb *= half(HistoryPreExposureCorrection);
-		#endif
+	#if AA_HISTORY_PAYLOAD == HISTORY_PAYLOAD_RGB_OPACITY_COC
+		HistoryPayload.CocRadius = RawHistory1.r;
+	#else
+		HistoryPayload.CocRadius = RawHistory0.a;
+	#endif
 
-		HistoryPayload.Color = TransformSceneColor(HistoryPayload.Color);
-		HistoryPayload.Color.rgb = clamp(HistoryPayload.Color.rgb, -half(MaxHalfFloat), half(MaxHalfFloat));
-	}
+	#if USE_PREEXPOSURE
+		HistoryPayload.Color.rgb *= HistoryPreExposureCorrection;
+	#endif
 
-	// Decode history infos
-	{
-		HistoryInfo = CreateHistoryInformation();
+	HistoryPayload.Color = TransformSceneColor(HistoryPayload.Color);
 
-		#if AA_DYNAMIC_ANTIGHOST && AA_DYNAMIC
-		{
-			HistoryInfo.bIsDynamicPixel = RawHistory0.a > 0;
-		}
-		#endif
-	}
+	return HistoryPayload;
 }
 
 
 // Clamp history.
-FTAAPayload ClampHistory(FTAAPayload History, FTAAPayload NeighborMin, FTAAPayload NeighborMax)
+FTAAHistoryPayload ClampHistory(inout FTAAIntermediaryResult IntermediaryResult, FTAAHistoryPayload History, FTAAHistoryPayload NeighborMin, FTAAHistoryPayload NeighborMax)
 {
 	#if !AA_CLAMP
 		return History;
 	
 	#elif AA_CLIP
 		// Clip history, this uses color AABB intersection for tighter fit.
-		//half4 TargetColor = half(0.5) * ( NeighborMin + NeighborMax );
-		half4 TargetColor = Filtered;
+		//float4 TargetColor = 0.5 * ( NeighborMin + NeighborMax );
+		float4 TargetColor = Filtered;
 
-		half ClipBlend = HistoryClip( HistoryColor.rgb, TargetColor.rgb, NeighborMin.rgb, NeighborMax.rgb );
+		float ClipBlend = HistoryClip( HistoryColor.rgb, TargetColor.rgb, NeighborMin.rgb, NeighborMax.rgb );
 		
-		//half DistToClamp = saturate(-ClipBlend) / ( saturate(-ClipBlend) + 1 );
-		//half DistToClamp = abs( ClipBlend ) / ( 1 - ClipBlend );
+		//float DistToClamp = saturate(-ClipBlend) / ( saturate(-ClipBlend) + 1 );
+		//float DistToClamp = abs( ClipBlend ) / ( 1 - ClipBlend );
 		ClipBlend = saturate( ClipBlend );
 
 		HistoryColor = lerp( HistoryColor, TargetColor, ClipBlend );
@@ -1457,85 +1846,63 @@ FTAAPayload ClampHistory(FTAAPayload History, FTAAPayload NeighborMin, FTAAPaylo
 }
 
 
-
-
 //------------------------------------------------------- TAA MAIN FUNCTION
 
-void TemporalAASample(
-	uint2 GroupId, uint GroupThreadIndex, 
-	out FTAAPayload OutputPayload,
-	out FTAAHistoryInformation OutputHistoryInfo)
+FTAAHistoryPayload TemporalAASample(uint2 GroupId, uint2 GroupThreadId, uint GroupThreadIndex, float2 ViewportUV, float FrameExposureScale)
 {
-	OutputHistoryInfo = CreateHistoryInformation();
-
 	// SETUP
 	// -----
 	FTAAInputParameters InputParams;
 
 	// Per frame setup.
-	{
-		float FrameExposureScale = EyeAdaptationLookup();
-		#if USE_PREEXPOSURE
-			InputParams.FrameExposureScale = ToScalarMemory(half(FrameExposureScale * View.OneOverPreExposure));
-		#else
-			InputParams.FrameExposureScale = ToScalarMemory(half(FrameExposureScale));
-		#endif
-	}
+	#if USE_PREEXPOSURE
+		InputParams.FrameExposureScale = ToScalarMemory(FrameExposureScale * View.OneOverPreExposure);
+	#else
+		InputParams.FrameExposureScale = ToScalarMemory(FrameExposureScale);
+	#endif
+
 
 	// Per pixel setup.
-	uint2 DebugOutputCoord;
 	{
-		{
-			uint2 DispatchThreadId = uint2(GroupThreadIndex % THREADGROUP_SIZEX, GroupThreadIndex / THREADGROUP_SIZEX) + GroupId * uint2(THREADGROUP_SIZEX, THREADGROUP_SIZEY);
+		InputParams.GroupId = GroupId;
+		InputParams.GroupThreadId = GroupThreadId;
+		InputParams.GroupThreadIndex = GroupThreadIndex;
+		InputParams.ViewportUV = ViewportUV;
+		InputParams.ScreenPos = ViewportUVToScreenPos(ViewportUV);
+		InputParams.NearestBufferUV = ViewportUV * ViewportUVToInputBufferUV.xy + ViewportUVToInputBufferUV.zw;
 
-			DebugOutputCoord = DispatchThreadId;
+		// Handle single or multi-pass responsive AA
+		#if AA_SINGLE_PASS_RESPONSIVE
+		{
+			const uint kResponsiveStencilMask = 1 << 3;
+			
+			int2 SceneStencilUV = (int2)trunc(InputParams.NearestBufferUV * InputSceneColorSize.xy);
+			uint SceneStencilRef = StencilTexture.Load(int3(SceneStencilUV, 0)) STENCIL_COMPONENT_SWIZZLE;
 
-			#if AA_LOWER_RESOLUTION
-			{
-				InputParams.ViewportUV = (float2(DispatchThreadId) + 0.5f) * MaxViewportUVAndSvPositionToViewportUV.zw;
-				InputParams.ViewportUV = min(InputParams.ViewportUV, MaxViewportUVAndSvPositionToViewportUV.xy);
-			}
-			#else
-			{
-				InputParams.ViewportUV = (float2(DispatchThreadId) + 0.5f) * OutputViewportSize.zw;
-			}
-			#endif
+			InputParams.bIsResponsiveAAPixel = (SceneStencilRef & kResponsiveStencilMask) ? 1.f : 0.f;
 		}
-
-		InputParams.ScreenPos = ViewportUVToScreenPos(InputParams.ViewportUV);
-		
+		#elif TAA_RESPONSIVE
+			InputParams.bIsResponsiveAAPixel = 1.f;
+		#else
+			InputParams.bIsResponsiveAAPixel = 0.f;
+		#endif
+	
 		#if AA_UPSAMPLE
 		{
 			// Pixel coordinate of the center of output pixel O in the input viewport.
-			float2 PPCo = InputParams.ViewportUV * InputViewSize.xy + TemporalJitterPixels;
+			float2 PPCo = ViewportUV * InputViewSize.xy + TemporalJitterPixels;
 		
 			// Pixel coordinate of the center of the nearest input pixel K.
 			float2 PPCk = floor(PPCo) + 0.5;
 		
+			// Pixel coordinate of the center of the nearest top left input pixel T.
+			float2 PPCt = floor(PPCo - 0.5) + 0.5;
+		
 			InputParams.NearestBufferUV = InputSceneColorSize.zw * (InputViewMin + PPCk);
+			InputParams.NearestTopLeftBufferUV = InputSceneColorSize.zw * (InputViewMin + PPCt);
 
 			// TODO: because use nearest sampler, can be faster in this computation.
 		}
-		#else
-		{
-			InputParams.NearestBufferUV = InputParams.ViewportUV * ViewportUVToInputBufferUV.xy + ViewportUVToInputBufferUV.zw;
-		}
-		#endif
-
-		// Handle single or multi-pass responsive AA
-		#if AA_SINGLE_PASS_RESPONSIVE
-		{
-			const uint kResponsiveStencilMask = 1 << 3;
-			
-			int2 SceneStencilUV = (int2)trunc(InputParams.NearestBufferUV * InputSceneColorSize.xy);
-			uint SceneStencilRef = StencilTexture.Load(int3(SceneStencilUV, 0)) STENCIL_COMPONENT_SWIZZLE;
-
-			InputParams.bIsResponsiveAAPixel = (SceneStencilRef & kResponsiveStencilMask) != 0;
-		}
-		#elif TAA_RESPONSIVE
-			InputParams.bIsResponsiveAAPixel = true;
-		#else
-			InputParams.bIsResponsiveAAPixel = false;
 		#endif
 	}
 
@@ -1547,7 +1914,8 @@ void TemporalAASample(
 	float3 PosN; // Position of this pixel, possibly later nearest pixel in neighborhood.
 	PosN.xy = InputParams.ScreenPos;
 
-	PosN.z = SampleSceneDepthTexture(InputParams, int2(0, 0));
+	PrecacheInputSceneDepth(InputParams);
+	PosN.z = SampleCachedSceneDepthTexture(InputParams, int2(0, 0));
 
 	// Screen position of minimum depth.
 	float2 VelocityOffset = float2(0.0, 0.0);
@@ -1556,11 +1924,11 @@ void TemporalAASample(
 		// For motion vector, use camera/dynamic motion from min depth pixel in pattern around pixel.
 		// This enables better quality outline on foreground against different motion background.
 		// Larger 2 pixel distance "x" works best (because AA dilates surface).
-		half4 Depths;
-		Depths.x = SampleSceneDepthTexture(InputParams, int2(-AA_CROSS, -AA_CROSS));
-		Depths.y = SampleSceneDepthTexture(InputParams, int2( AA_CROSS, -AA_CROSS));
-		Depths.z = SampleSceneDepthTexture(InputParams, int2(-AA_CROSS,  AA_CROSS));
-		Depths.w = SampleSceneDepthTexture(InputParams, int2( AA_CROSS,  AA_CROSS));
+		float4 Depths;
+		Depths.x = SampleCachedSceneDepthTexture(InputParams, int2(-AA_CROSS, -AA_CROSS));
+		Depths.y = SampleCachedSceneDepthTexture(InputParams, int2( AA_CROSS, -AA_CROSS));
+		Depths.z = SampleCachedSceneDepthTexture(InputParams, int2(-AA_CROSS,  AA_CROSS));
+		Depths.w = SampleCachedSceneDepthTexture(InputParams, int2( AA_CROSS,  AA_CROSS));
 
 		float2 DepthOffset = float2(AA_CROSS, AA_CROSS);
 		float DepthOffsetXx = float(AA_CROSS);
@@ -1574,14 +1942,14 @@ void TemporalAASample(
 			{
 				DepthOffset.x = -AA_CROSS;
 			}
-			half DepthsXY = max(Depths.x, Depths.y);
-			half DepthsZW = max(Depths.z, Depths.w);
+			float DepthsXY = max(Depths.x, Depths.y);
+			float DepthsZW = max(Depths.z, Depths.w);
 			if(DepthsXY > DepthsZW) 
 			{
 				DepthOffset.y = -AA_CROSS;
 				DepthOffset.x = DepthOffsetXx; 
 			}
-			half DepthsXYZW = max(DepthsXY, DepthsZW);
+			float DepthsXYZW = max(DepthsXY, DepthsZW);
 			if(DepthsXYZW > PosN.z) 
 			{
 				// This is offset for reading from velocity texture.
@@ -1600,8 +1968,8 @@ void TemporalAASample(
 
 	// Camera motion for pixel or nearest pixel (in ScreenPos space).
 	bool OffScreen = false;
-	half Velocity = 0;
-	half HistoryBlur = 0;
+	float Velocity = 0;
+	float HistoryBlur = 0;
 	float2 HistoryScreenPosition = InputParams.ScreenPos;
 
 	#if 1
@@ -1626,11 +1994,11 @@ void TemporalAASample(
 		}
 		#endif
 
-		Velocity = half(sqrt(dot(BackTemp, BackTemp)));
+		Velocity = sqrt(dot(BackTemp, BackTemp));
 		#if !AA_BICUBIC
 			// Save the amount of pixel offset of just camera motion, used later as the amount of blur introduced by history.
 			float HistoryBlurAmp = 2.0;
-			HistoryBlur = half(saturate(abs(BackTemp.x) * HistoryBlurAmp + abs(BackTemp.y) * HistoryBlurAmp));
+			HistoryBlur = saturate(abs(BackTemp.x) * HistoryBlurAmp + abs(BackTemp.y) * HistoryBlurAmp);
 		#endif
 		// Easier to do off screen check before conversion.
 		// BackN is in units of 2pixels/viewportWidthInPixels
@@ -1653,8 +2021,8 @@ void TemporalAASample(
 	#endif
 	
 	// Compute neighborhood bounding box.
-	FTAAPayload NeighborMin;
-	FTAAPayload NeighborMax;
+	FTAAHistoryPayload NeighborMin;
+	FTAAHistoryPayload NeighborMax;
 
 	ComputeNeighborhoodBoundingbox(
 		InputParams,
@@ -1662,16 +2030,15 @@ void TemporalAASample(
 		NeighborMin, NeighborMax);
 
 	// Sample history.
-	FTAAPayload History;
-	FTAAHistoryInformation HistoryInfo;
-	SampleHistory(HistoryScreenPosition, /* out */ History, /* out */ HistoryInfo);
+	FTAAHistoryPayload History = SampleHistory(HistoryScreenPosition);
 
 	// Whether the feedback needs to be reset.
-	bool bIgnoreHistory = OffScreen || bCameraCut;
+	bool IgnoreHistory = OffScreen || bCameraCut;
 
 	// DYNAMIC ANTI GHOSTING
 	// ---------------------
-	#if AA_DYNAMIC_ANTIGHOST && AA_DYNAMIC
+	#if AA_DYNAMIC_ANTIGHOST && AA_DYNAMIC && HISTORY_PAYLOAD_COMPONENTS == 3
+	bool Dynamic4;
 	{
 		#if !AA_DYNAMIC
 			#error AA_DYNAMIC_ANTIGHOST requires AA_DYNAMIC
@@ -1679,25 +2046,23 @@ void TemporalAASample(
 		// TODO: try a 2x2 for AA_UPSAMPLE
 		bool Dynamic1 = SceneVelocityBuffer.SampleLevel(SceneVelocityBufferSampler, InputParams.NearestBufferUV, 0, int2( 0, -1)).x > 0;
 		bool Dynamic3 = SceneVelocityBuffer.SampleLevel(SceneVelocityBufferSampler, InputParams.NearestBufferUV, 0, int2(-1,  0)).x > 0;
-		bool Dynamic4 = SceneVelocityBuffer.SampleLevel(SceneVelocityBufferSampler, InputParams.NearestBufferUV, 0).x > 0;
+		Dynamic4 = SceneVelocityBuffer.SampleLevel(SceneVelocityBufferSampler, InputParams.NearestBufferUV, 0).x > 0;
 		bool Dynamic5 = SceneVelocityBuffer.SampleLevel(SceneVelocityBufferSampler, InputParams.NearestBufferUV, 0, int2( 1,  0)).x > 0;
 		bool Dynamic7 = SceneVelocityBuffer.SampleLevel(SceneVelocityBufferSampler, InputParams.NearestBufferUV, 0, int2( 0,  1)).x > 0;
 
 		bool Dynamic = Dynamic1 || Dynamic3 || Dynamic4 || Dynamic5 || Dynamic7;
-		bIgnoreHistory = bIgnoreHistory || !Dynamic && HistoryInfo.bIsDynamicPixel;
-		
-		OutputHistoryInfo.bIsDynamicPixel = Dynamic4;
+		IgnoreHistory = IgnoreHistory || (!Dynamic && History.Color.a > 0);
 	}
 	#endif
 	
 	// Save off luma of history before the clamp.
-	half LumaMin     = GetSceneColorLuma4(NeighborMin.Color);
-	half LumaMax     = GetSceneColorLuma4(NeighborMax.Color);
-	half LumaHistory = GetSceneColorLuma4(History.Color);
+	float LumaMin     = GetSceneColorLuma4(NeighborMin.Color);
+	float LumaMax     = GetSceneColorLuma4(NeighborMax.Color);
+	float LumaHistory = GetSceneColorLuma4(History.Color);
 	
 	// Clamp history.
-	FTAAPayload PreClampingHistoryColor = History;
-	History = ClampHistory(History, NeighborMin, NeighborMax);
+	FTAAHistoryPayload PreClampingHistoryColor = History;
+	History = ClampHistory(IntermediaryResult, History, NeighborMin, NeighborMax);
 	
 	// Filter input after color clamping.
 	#if AA_UPSAMPLE_ADAPTIVE_FILTERING == 1
@@ -1706,7 +2071,7 @@ void TemporalAASample(
 			#error AA_VARIANCE and AA_UPSAMPLE_ADAPTIVE_FILTERING are not compatible because of circular code dependency.
 		#endif
 
-		if (bIgnoreHistory) // || any(HistoryColor != PreClampingHistoryColor))
+		if (IgnoreHistory) // || any(HistoryColor != PreClampingHistoryColor))
 		{
 			// Set the input filter infinitely large when we know need to rely on it.
 			IntermediaryResult.InvFilterScaleFactor = 0;
@@ -1734,217 +2099,186 @@ void TemporalAASample(
 		#endif
 		
 		// Blend in non-filtered based on the amount of sub-pixel motion.
-		half AddAliasing = saturate(HistoryBlur) * 0.5;
-		half LumaContrastFactor = 32.0;
+		float AddAliasing = saturate(HistoryBlur) * 0.5;
+		float LumaContrastFactor = 32.0;
 		#if AA_YCOCG // TODO: Probably a bug arround here because using Luma4() even with YCOCG=0.
 			// 1/4 as bright.
 			LumaContrastFactor *= 4.0;
 		#endif
-		half LumaContrast = LumaMax - LumaMin;
+		float LumaContrast = LumaMax - LumaMin;
 		AddAliasing = saturate(AddAliasing + rcp(1.0 + LumaContrast * LumaContrastFactor));
-
-		FTAASceneColorSample Sample = SampleCachedSceneColorTexture(InputParams, /* NeighborId = */ 0);
-		IntermediaryResult.Filtered.Color = lerp(IntermediaryResult.Filtered.Color, Sample.Payload.Color, AddAliasing);
+		IntermediaryResult.Filtered.Color = lerp(IntermediaryResult.Filtered.Color, SampleCachedSceneColorTexture(InputParams, int2(0, 0)).Color, AddAliasing);
 	}
 	#endif
 	
 	// COMPUTE BLEND AMOUNT 
 	// --------------------
-	half BlendFinal;
+	float BlendFinal;
 	{
-		half LumaFiltered = GetSceneColorLuma4(IntermediaryResult.Filtered.Color);
+		float LumaFiltered = GetSceneColorLuma4(IntermediaryResult.Filtered.Color);
 
 		BlendFinal = IntermediaryResult.FilteredTemporalWeight * CurrentFrameWeight;
 
-		BlendFinal = lerp(BlendFinal, half(0.2), saturate(Velocity / 40.0));
+		BlendFinal = lerp(BlendFinal, 0.2, saturate(Velocity / 40));
 
-		#if AA_ANTIFLICKER
-			half DistToClamp = min( abs(LumaHistory - LumaMin), abs(LumaHistory - LumaMax) ) / max3( LumaHistory, LumaFiltered, 1e-4 );
-			BlendFinal *= 0.2 + 0.8 * saturate( 0.5 * DistToClamp );
+		#if 0
+		{
+			// Anti-flicker
+			float DistToClamp = 2 * abs(min(LumaHistory - LumaMin, LumaMax - LumaHistory) / (LumaMax - LumaMin));
+			//BlendFinal *= lerp( 0, 1, saturate(4 * DistToClamp) );
+			BlendFinal += 0.8 * saturate(0.02 * LumaHistory / abs(Filtered.x - LumaHistory));
+			BlendFinal *= (LumaMin * InExposureScale + 0.5) / (LumaMax * InExposureScale + 0.5);
+		}
 		#endif
 
 		// Make sure to have at least some small contribution
-		BlendFinal = max( BlendFinal, saturate( half(0.01) * LumaHistory / abs( LumaFiltered - LumaHistory ) ) );
+		BlendFinal = max( BlendFinal, saturate( 0.01 * LumaHistory / abs( LumaFiltered - LumaHistory ) ) );
 
 		#if AA_NAN && (COMPILER_GLSL || COMPILER_METAL)
 			// The current Metal & GLSL compilers don't handle saturate(NaN) -> 0, instead they return NaN/INF.
-			BlendFinal = -min(-BlendFinal, half(0.0));
+			BlendFinal = -min(-BlendFinal, 0.0);
 		#endif
 
 		// Responsive forces 1/4 of new frame.
-		BlendFinal = InputParams.bIsResponsiveAAPixel ? half(1.0/4.0) : BlendFinal;
+		BlendFinal = InputParams.bIsResponsiveAAPixel ? (1.0/4.0) : BlendFinal;
 
 		#if AA_LERP 
-			BlendFinal = half(1.0/half(AA_LERP));
+			BlendFinal = 1.0/float(AA_LERP);
 		#endif
 	
 		#if AA_HISTORY_PAYLOAD == HISTORY_PAYLOAD_RGB_COC || AA_HISTORY_PAYLOAD == HISTORY_PAYLOAD_RGB_OPACITY_COC
 		{
-			half BilateralWeight = ComputeBilateralWeight(IntermediaryResult.Filtered.CocRadius, History.CocRadius);
+			float BilateralWeight = ComputeBilateralWeight(IntermediaryResult.Filtered.CocRadius, History.CocRadius);
 
-			BlendFinal = lerp(half(1), BlendFinal, BilateralWeight);
+			BlendFinal = lerp(1, BlendFinal, BilateralWeight);
 		}
 		#endif
 		
 		if (bCameraCut)
 		{
-			BlendFinal = half(1.0);
+			BlendFinal = 1.0;
 		}
 	}
 
 	// Offscreen feedback resets.
-	if (bIgnoreHistory)
+	if (IgnoreHistory)
 	{
 		History = IntermediaryResult.Filtered;
+		#if HISTORY_PAYLOAD_COMPONENTS == 3
+			History.Color.a = 0.0;
+		#endif
 	}
 	
 	// DO FINAL BLEND BETWEEN HISTORY AND FILTERED COLOR
 	// -------------------------------------------------
-	{
-		// Luma weighted blend
-		half FilterWeight = GetSceneColorHdrWeight(InputParams, IntermediaryResult.Filtered.Color.x);
-		half HistoryWeight = GetSceneColorHdrWeight(InputParams, History.Color.x);
+	// Luma weighted blend
+	float FilterWeight = GetSceneColorHdrWeight(InputParams, IntermediaryResult.Filtered.Color.x);
+	float HistoryWeight = GetSceneColorHdrWeight(InputParams, History.Color.x);
 
-		half2 Weights = WeightedLerpFactors(HistoryWeight, FilterWeight, BlendFinal);
+	FTAAHistoryPayload OutputPayload;
+	{
+		float2 Weights = WeightedLerpFactors(HistoryWeight, FilterWeight, BlendFinal);
 		OutputPayload = AddPayload(MulPayload(History, Weights.x), MulPayload(IntermediaryResult.Filtered, Weights.y));
 	}
 
+	// Ensure that alpha values that are expected to be opaque (but are only close to opaque) are forced to be opaque.
+	// (0.995 chosen to accomodate handling of 254/255)
+	if (OutputPayload.Color.a > 0.995)
 	{
-		OutputPayload.Color = TransformBackToRawLinearSceneColor(OutputPayload.Color);
-		
-		// Ensure that alpha values that are expected to be opaque (but are only close to opaque) are forced to be opaque.
-		// (0.995 chosen to accomodate handling of 254/255)
-		if (OutputPayload.Color.a > half(0.995))
-		{
-			OutputPayload.Color.a = half(1);
-		}
+		OutputPayload.Color.a = 1;
 	}
+
+	OutputPayload.Color = TransformBackToRawLinearSceneColor(OutputPayload.Color);
+	
+	#if AA_NAN
+		// Transform NaNs to black, transform negative colors to black.
+		OutputPayload.Color = -min(-OutputPayload.Color, 0.0);
+
+		OutputPayload.CocRadius = isnan(OutputPayload.CocRadius) ? 0.0 : OutputPayload.CocRadius;
+	#endif
+
+	#if HISTORY_PAYLOAD_COMPONENTS == 3
+		#if  AA_DYNAMIC_ANTIGHOST && AA_DYNAMIC
+			OutputPayload.Color.a = Dynamic4 ? 1 : 0;
+		#else
+			// Zero out to remove any prior computation of alpha
+			OutputPayload.Color.a = 0;
+		#endif
+	#endif
+
+	return OutputPayload;
 }
 
 //------------------------------------------------------- ENTRY POINT
 
 [numthreads(THREADGROUP_SIZEX, THREADGROUP_SIZEY, 1)]
 void MainCS(
+	uint2 DispatchThreadId : SV_DispatchThreadID,
 	uint2 GroupId : SV_GroupID,
+	uint2 GroupThreadId : SV_GroupThreadID,
 	uint GroupThreadIndex : SV_GroupIndex)
 {
-	FTAAPayload OutputPayload;
-	FTAAHistoryInformation OutputHistoryInfo;
+	float2 ViewportUV = (float2(DispatchThreadId) + 0.5f) * OutputViewportSize.zw;
+	
+	#if AA_LOWER_RESOLUTION
+	{
+		ViewportUV = (float2(DispatchThreadId) + 0.5f) * MaxViewportUVAndSvPositionToViewportUV.zw;
+		ViewportUV = min(ViewportUV, MaxViewportUVAndSvPositionToViewportUV.xy);
+	}
+	#endif
 
-	TemporalAASample(
-		GroupId, GroupThreadIndex,
-		/* out */ OutputPayload,
-		/* out */ OutputHistoryInfo);
+	float FrameExposureScale = EyeAdaptationLookup();
+	FTAAHistoryPayload OutputPayload = TemporalAASample(GroupId, GroupThreadId, GroupThreadIndex, ViewportUV, FrameExposureScale);
 
-	half4 OutColor0 = 0;
-	half4 OutColor1 = 0;
-	half4 OutColor2 = 0;
-	half4 OutColor3 = 0;
+	float4 OutColor0 = 0;
+	float4 OutColor1 = 0;
 
-	// Encode history
+	#if AA_HISTORY_PAYLOAD == HISTORY_PAYLOAD_RGB_COC
 	{
-		// Transform NaNs to black, transform negative colors to black.
-		#if AA_NAN
-			OutputPayload.Color = -min(-OutputPayload.Color, half(0.0));
-			OutputPayload.CocRadius = isnan(OutputPayload.CocRadius) ? 0.0 : OutputPayload.CocRadius;
-		#endif
+		OutColor0.rgb = OutputPayload.Color.rgb;
+		OutColor0.a = OutputPayload.CocRadius;
+	}
+	#elif AA_HISTORY_PAYLOAD == HISTORY_PAYLOAD_RGB_OPACITY_COC
+	{
+		OutColor0 = OutputPayload.Color;
+		OutColor1.r = OutputPayload.CocRadius;
+	}
+	#else
+	{
+		OutColor0 = OutputPayload.Color;
+	}
+	#endif
 
-		OutputPayload.Color.rgb = min(OutputPayload.Color.rgb, half(MaxHalfFloat));
+	uint2 PixelPos = DispatchThreadId + OutputViewportRect.xy;
+	if (all(PixelPos < OutputViewportRect.zw))
+	{
+		OutComputeTex_0[PixelPos] = min(MaxHalfFloat.xxxx, OutColor0);
 
-		#if AA_HISTORY_PAYLOAD == HISTORY_PAYLOAD_RGB
-		{
-			OutColor0.rgb = OutputPayload.Color.rgb;
-		}
-		#elif AA_HISTORY_PAYLOAD == HISTORY_PAYLOAD_RGB_OPACITY
-		{
-			OutColor0 = OutputPayload.Color;
-		}
-		#elif AA_HISTORY_PAYLOAD == HISTORY_PAYLOAD_RGB_COC
-		{
-			OutColor0.rgb = OutputPayload.Color.rgb;
-			OutColor0.a = OutputPayload.CocRadius;
-		}
-		#elif AA_HISTORY_PAYLOAD == HISTORY_PAYLOAD_RGB_OPACITY_COC
-		{
-			OutColor0 = OutputPayload.Color;
-			OutColor1.r = OutputPayload.CocRadius;
-		}
-		#else
-			#error Unknown history payload
-		#endif
-		
-		#if AA_DYNAMIC_ANTIGHOST && AA_DYNAMIC
-		{
-			#if HISTORY_PAYLOAD_COMPONENTS == 3
-				OutColor0.a = OutputHistoryInfo.bIsDynamicPixel ? 1 : 0;
-			#else
-				#error "Can't encode AA_DYNAMIC_ANTIGHOST"
-			#endif
-		}
+		#if HISTORY_RENDER_TARGETS == 2
+			OutComputeTex_1[PixelPos] = OutColor1;
 		#endif
 	}
 
-	// Output history
-	ISOLATE
+	#if TAA_DOWNSAMPLE
 	{
-		// Avoid keeping a VGPR for GroupThreadIndex.
-		#if COMPILER_SUPPORTS_WAVE_ONCE
-		ISOLATE
-		{
-			GroupThreadIndex = WaveGetLaneIndex();
-		}
-		#endif
+		// This shader permutation outputs half resolution image in addition to main full-res one.
+		// It is more efficient than performing a separate downsampling pass afterwards.
+		// Assumes output resolution is always even.
 
-		// Compute new pixel position to avoid keeping VGPR.
-		uint2 OutputPixelPos;
-		ISOLATE
-		{
-			OutputPixelPos = (
-				uint2(GroupThreadIndex % THREADGROUP_SIZEX, GroupThreadIndex / THREADGROUP_SIZEX) +
-				GroupId * uint2(THREADGROUP_SIZEX, THREADGROUP_SIZEY) +
-				OutputViewportRect.xy);
-		}
-
-		if (all(OutputPixelPos < OutputViewportRect.zw))
-		{
-			OutComputeTex_0[OutputPixelPos] = OutColor0;
-
-			#if HISTORY_RENDER_TARGETS >= 2
-				OutComputeTex_1[OutputPixelPos] = OutColor1;
-			#endif
+		uint P0 = GroupThreadId.x + GroupThreadId.y * THREADGROUP_SIZEX;
+		uint P1 = P0 + 1;
+		uint P2 = P0 + THREADGROUP_SIZEX;
+		uint P3 = P2 + 1;
 
-			#if HISTORY_RENDER_TARGETS >= 3
-				OutComputeTex_2[OutputPixelPos] = OutColor2;
-			#endif
+		GroupSharedDownsampleArray[P0] = OutColor0;
 
-			#if HISTORY_RENDER_TARGETS == 4
-				OutComputeTex_3[OutputPixelPos] = OutColor3;
-			#endif
-		}
+		GroupMemoryBarrierWithGroupSync();
 
-		#if TAA_DOWNSAMPLE
+		if (((GroupThreadId.x | GroupThreadId.y) & 1) == 0)
 		{
-			// This shader permutation outputs half resolution image in addition to main full-res one.
-			// It is more efficient than performing a separate downsampling pass afterwards.
-			// Assumes output resolution is always even.
-
-			uint2 GroupThreadId = uint2(GroupThreadIndex % THREADGROUP_SIZEX, GroupThreadIndex / THREADGROUP_SIZEX);
-
-			uint P0 = GroupThreadId.x + GroupThreadId.y * THREADGROUP_SIZEX;
-			uint P1 = P0 + 1;
-			uint P2 = P0 + THREADGROUP_SIZEX;
-			uint P3 = P2 + 1;
-
-			GroupSharedDownsampleArray[P0] = OutColor0;
-
-			GroupMemoryBarrierWithGroupSync();
-
-			if (((GroupThreadId.x | GroupThreadId.y) & 1) == 0)
-			{
-				OutComputeTexDownsampled[OutputPixelPos / 2] =
-					(OutColor0 + GroupSharedDownsampleArray[P1] + GroupSharedDownsampleArray[P2] + GroupSharedDownsampleArray[P3]) * 0.25;
-			}
+			OutComputeTexDownsampled[PixelPos / 2] =
+				(OutColor0 + GroupSharedDownsampleArray[P1] + GroupSharedDownsampleArray[P2] + GroupSharedDownsampleArray[P3]) * 0.25;
 		}
-		#endif //TAA_DOWNSAMPLE
 	}
+	#endif //TAA_DOWNSAMPLE
 }
diff --git a/Engine/Source/Runtime/Renderer/Private/PostProcess/PostProcessMitchellNetravali.cpp b/Engine/Source/Runtime/Renderer/Private/PostProcess/PostProcessMitchellNetravali.cpp
--- a/Engine/Source/Runtime/Renderer/Private/PostProcess/PostProcessMitchellNetravali.cpp
+++ b/Engine/Source/Runtime/Renderer/Private/PostProcess/PostProcessMitchellNetravali.cpp
@@ -1,9 +1,5 @@
 // Copyright Epic Games, Inc. All Rights Reserved.
 
-/*=============================================================================
-	PostProcessTemporalAA.cpp: Post process MotionBlur implementation.
-=============================================================================*/
-
 #include "PostProcess/PostProcessMitchellNetravali.h"
 #include "SceneUtils.h"
 #include "SceneRendering.h"
diff --git a/Engine/Source/Runtime/Renderer/Private/PostProcess/TemporalAA.cpp b/Engine/Source/Runtime/Renderer/Private/PostProcess/TemporalAA.cpp
--- a/Engine/Source/Runtime/Renderer/Private/PostProcess/TemporalAA.cpp
+++ b/Engine/Source/Runtime/Renderer/Private/PostProcess/TemporalAA.cpp
@@ -17,6 +17,13 @@ const int32 GTemporalAATileSizeY = 8;
 
 constexpr int32 kHistoryTextures = 3;
 
+TAutoConsoleVariable<int32> CVarTAAAlgorithm(
+	TEXT("r.TemporalAA.Algorithm"), 0,
+	TEXT("Algorithm to use for Temporal AA\n")
+	TEXT(" 0: Gen 4 TAAU (default)\n")
+	TEXT(" 1: Gen 5 TAAU (experimental)"),
+	ECVF_RenderThreadSafe);
+
 TAutoConsoleVariable<float> CVarTemporalAAFilterSize(
 	TEXT("r.TemporalAAFilterSize"),
 	1.0f,
@@ -59,8 +66,16 @@ TAutoConsoleVariable<int32> CVarTemporalAAAllowDownsampling(
 	TEXT("Allows half-resolution color buffer to be produced during TAA. Only possible when motion blur is off and when using compute shaders for post processing."),
 	ECVF_RenderThreadSafe);
 
+static TAutoConsoleVariable<int32> CVarUseTemporalAAUpscaler(
+	TEXT("r.TemporalAA.Upscaler"),
+	1,
+	TEXT("Choose the upscaling algorithm.\n")
+	TEXT(" 0: Forces the default temporal upscaler of the renderer;\n")
+	TEXT(" 1: GTemporalUpscaler which may be overridden by a third party plugin (default)."),
+	ECVF_RenderThreadSafe);
+
 TAutoConsoleVariable<int32> CVarTAAR11G11B10History(
-	TEXT("r.TemporalAA.R11G11B10History"), 1,
+	TEXT("r.TemporalAA.R11G11B10History"), 0,
 	TEXT("Select the bitdepth of the history."),
 	ECVF_RenderThreadSafe);
 
@@ -69,23 +84,16 @@ TAutoConsoleVariable<int32> CVarTAANyquistHistory(
 	TEXT(""),
 	ECVF_RenderThreadSafe);
 
-static TAutoConsoleVariable<int32> CVarUseTemporalAAUpscaler(
-	TEXT("r.TemporalAA.Upscaler"),
-	1,
-	TEXT("Choose the upscaling algorithm.\n")
-	TEXT(" 0: Forces the default temporal upscaler of the renderer;\n")
-	TEXT(" 1: GTemporalUpscaler which may be overridden by a third party plugin (default)."),
-	ECVF_RenderThreadSafe);
-
 inline bool DoesPlatformSupportTemporalHistoryUpscale(EShaderPlatform Platform)
 {
 	return (IsPCPlatform(Platform) || FDataDrivenShaderPlatformInfo::GetSupportsTemporalHistoryUpscale(Platform))
 		&& IsFeatureLevelSupported(Platform, ERHIFeatureLevel::SM5);
 }
 
-inline bool DoesPlatformSupportTAA2(EShaderPlatform Platform)
+inline bool DoesPlatformSupportGen5TAA(EShaderPlatform Platform)
 {
-	return IsFeatureLevelSupported(Platform, ERHIFeatureLevel::SM5);
+	return Platform == SP_PCD3D_SM5;
+	//return IsFeatureLevelSupported(Platform, ERHIFeatureLevel::SM5);
 }
 
 BEGIN_SHADER_PARAMETER_STRUCT(FTAA2CommonParameters, )
@@ -118,24 +126,24 @@ FTAA2HistoryUAVs CreateUAVs(FRDGBuilder& GraphBuilder, const FTAA2HistoryTexture
 	return UAVs;
 }
 
-class FTAA2Shader : public FGlobalShader
+class FTAAGen5Shader : public FGlobalShader
 {
 public:
-	FTAA2Shader(const ShaderMetaType::CompiledShaderInitializerType& Initializer)
+	FTAAGen5Shader(const ShaderMetaType::CompiledShaderInitializerType& Initializer)
 		: FGlobalShader(Initializer)
 	{ }
 
-	FTAA2Shader()
+	FTAAGen5Shader()
 	{ }
 
 	static bool ShouldCompilePermutation(const FGlobalShaderPermutationParameters& Parameters)
 	{
-		return DoesPlatformSupportTAA2(Parameters.Platform);
+		return DoesPlatformSupportGen5TAA(Parameters.Platform);
 	}
 
 	static void ModifyCompilationEnvironment(const FGlobalShaderPermutationParameters& Parameters, FShaderCompilerEnvironment& OutEnvironment)
 	{
-		OutEnvironment.CompilerFlags.Add(CFLAG_AllowRealTypes);
+		//OutEnvironment.CompilerFlags.Add(CFLAG_AllowRealTypes);
 	}
 }; // class FTAA2Shader
 
@@ -208,8 +216,8 @@ class FTAAStandaloneCS : public FGlobalShader
 		SHADER_PARAMETER(float, ScreenPercentage)
 		SHADER_PARAMETER(float, UpscaleFactor)
 
-		SHADER_PARAMETER_RDG_TEXTURE_UAV_ARRAY(RWTexture2D, OutComputeTex, [FTemporalAAHistory::kRenderTargetCount])
-		SHADER_PARAMETER_RDG_TEXTURE_UAV(RWTexture2D, OutComputeTexDownsampled)
+		SHADER_PARAMETER_RDG_TEXTURE_UAV_ARRAY(Texture2D, OutComputeTex, [FTemporalAAHistory::kRenderTargetCount])
+		SHADER_PARAMETER_RDG_TEXTURE_UAV(Texture2D, OutComputeTexDownsampled)
 
 		SHADER_PARAMETER_RDG_TEXTURE_UAV(RWTexture2D, DebugOutput)
 	END_SHADER_PARAMETER_STRUCT()
@@ -296,14 +304,13 @@ class FTAAStandaloneCS : public FGlobalShader
 	{
 		OutEnvironment.SetDefine(TEXT("THREADGROUP_SIZEX"), GTemporalAATileSizeX);
 		OutEnvironment.SetDefine(TEXT("THREADGROUP_SIZEY"), GTemporalAATileSizeY);
-		//OutEnvironment.CompilerFlags.Add(CFLAG_AllowRealTypes);	// TODO: Causes blinking artifacts. Revisit this optimization later.
 	}
 }; // class FTAAStandaloneCS
 
-class FTAA2DilateVelocityCS : public FTAA2Shader
+class FTAA2DilateVelocityCS : public FTAAGen5Shader
 {
 	DECLARE_GLOBAL_SHADER(FTAA2DilateVelocityCS);
-	SHADER_USE_PARAMETER_STRUCT(FTAA2DilateVelocityCS, FTAA2Shader);
+	SHADER_USE_PARAMETER_STRUCT(FTAA2DilateVelocityCS, FTAAGen5Shader);
 
 	BEGIN_SHADER_PARAMETER_STRUCT(FParameters, )
 		SHADER_PARAMETER_STRUCT_INCLUDE(FTAA2CommonParameters, CommonParameters)
@@ -319,10 +326,10 @@ class FTAA2DilateVelocityCS : public FTAA2Shader
 	END_SHADER_PARAMETER_STRUCT()
 }; // class FTAA2DilateVelocityCS
 
-class FTAA2BuildParallaxMaskCS : public FTAA2Shader
+class FTAA2BuildParallaxMaskCS : public FTAAGen5Shader
 {
 	DECLARE_GLOBAL_SHADER(FTAA2BuildParallaxMaskCS);
-	SHADER_USE_PARAMETER_STRUCT(FTAA2BuildParallaxMaskCS, FTAA2Shader);
+	SHADER_USE_PARAMETER_STRUCT(FTAA2BuildParallaxMaskCS, FTAAGen5Shader);
 
 	BEGIN_SHADER_PARAMETER_STRUCT(FParameters, )
 		SHADER_PARAMETER_STRUCT_INCLUDE(FTAA2CommonParameters, CommonParameters)
@@ -338,10 +345,10 @@ class FTAA2BuildParallaxMaskCS : public FTAA2Shader
 	END_SHADER_PARAMETER_STRUCT()
 }; // class FTAA2BuildParallaxMaskCS
 
-class FTAA2DecimateHistoryCS : public FTAA2Shader
+class FTAA2DecimateHistoryCS : public FTAAGen5Shader
 {
 	DECLARE_GLOBAL_SHADER(FTAA2DecimateHistoryCS);
-	SHADER_USE_PARAMETER_STRUCT(FTAA2DecimateHistoryCS, FTAA2Shader);
+	SHADER_USE_PARAMETER_STRUCT(FTAA2DecimateHistoryCS, FTAAGen5Shader);
 
 	BEGIN_SHADER_PARAMETER_STRUCT(FParameters, )
 		SHADER_PARAMETER_STRUCT_INCLUDE(FTAA2CommonParameters, CommonParameters)
@@ -360,10 +367,10 @@ class FTAA2DecimateHistoryCS : public FTAA2Shader
 	END_SHADER_PARAMETER_STRUCT()
 }; // class FTAA2DecimateHistoryCS
 
-class FTAA2FilterFrequenciesCS : public FTAA2Shader
+class FTAA2FilterFrequenciesCS : public FTAAGen5Shader
 {
 	DECLARE_GLOBAL_SHADER(FTAA2FilterFrequenciesCS);
-	SHADER_USE_PARAMETER_STRUCT(FTAA2FilterFrequenciesCS, FTAA2Shader);
+	SHADER_USE_PARAMETER_STRUCT(FTAA2FilterFrequenciesCS, FTAAGen5Shader);
 
 	BEGIN_SHADER_PARAMETER_STRUCT(FParameters, )
 		SHADER_PARAMETER_STRUCT_INCLUDE(FTAA2CommonParameters, CommonParameters)
@@ -376,10 +383,10 @@ class FTAA2FilterFrequenciesCS : public FTAA2Shader
 	END_SHADER_PARAMETER_STRUCT()
 }; // class FTAA2FilterFrequenciesCS
 
-class FTAA2CompareHistoryCS : public FTAA2Shader
+class FTAA2CompareHistoryCS : public FTAAGen5Shader
 {
 	DECLARE_GLOBAL_SHADER(FTAA2CompareHistoryCS);
-	SHADER_USE_PARAMETER_STRUCT(FTAA2CompareHistoryCS, FTAA2Shader);
+	SHADER_USE_PARAMETER_STRUCT(FTAA2CompareHistoryCS, FTAAGen5Shader);
 
 	BEGIN_SHADER_PARAMETER_STRUCT(FParameters, )
 		SHADER_PARAMETER_STRUCT_INCLUDE(FTAA2CommonParameters, CommonParameters)
@@ -391,10 +398,10 @@ class FTAA2CompareHistoryCS : public FTAA2Shader
 	END_SHADER_PARAMETER_STRUCT()
 }; // class FTAA2CompareHistoryCS
 
-class FTAA2DilateRejectionCS : public FTAA2Shader
+class FTAA2DilateRejectionCS : public FTAAGen5Shader
 {
 	DECLARE_GLOBAL_SHADER(FTAA2DilateRejectionCS);
-	SHADER_USE_PARAMETER_STRUCT(FTAA2DilateRejectionCS, FTAA2Shader);
+	SHADER_USE_PARAMETER_STRUCT(FTAA2DilateRejectionCS, FTAAGen5Shader);
 
 	BEGIN_SHADER_PARAMETER_STRUCT(FParameters, )
 		SHADER_PARAMETER_STRUCT_INCLUDE(FTAA2CommonParameters, CommonParameters)
@@ -404,10 +411,10 @@ class FTAA2DilateRejectionCS : public FTAA2Shader
 	END_SHADER_PARAMETER_STRUCT()
 }; // class FTAA2DilateRejectionCS
 
-class FTAA2UpdateHistoryCS : public FTAA2Shader
+class FTAA2UpdateHistoryCS : public FTAAGen5Shader
 {
 	DECLARE_GLOBAL_SHADER(FTAA2UpdateHistoryCS);
-	SHADER_USE_PARAMETER_STRUCT(FTAA2UpdateHistoryCS, FTAA2Shader);
+	SHADER_USE_PARAMETER_STRUCT(FTAA2UpdateHistoryCS, FTAAGen5Shader);
 
 	BEGIN_SHADER_PARAMETER_STRUCT(FParameters, )
 		SHADER_PARAMETER_STRUCT_INCLUDE(FTAA2CommonParameters, CommonParameters)
@@ -457,11 +464,11 @@ void SetupSampleWeightParameters(FTAAStandaloneCS::FParameters* OutTAAParameters
 
 	static const float SampleOffsets[9][2] =
 	{
-		{  0.0f,  0.0f },
 		{ -1.0f, -1.0f },
 		{  0.0f, -1.0f },
 		{  1.0f, -1.0f },
 		{ -1.0f,  0.0f },
+		{  0.0f,  0.0f },
 		{  1.0f,  0.0f },
 		{ -1.0f,  1.0f },
 		{  0.0f,  1.0f },
@@ -501,18 +508,17 @@ void SetupSampleWeightParameters(FTAAStandaloneCS::FParameters* OutTAAParameters
 
 	// Compute 3x3 + weights.
 	{
-		OutTAAParameters->PlusWeights[0] = OutTAAParameters->SampleWeights[0];
-		OutTAAParameters->PlusWeights[1] = OutTAAParameters->SampleWeights[2];
+		OutTAAParameters->PlusWeights[0] = OutTAAParameters->SampleWeights[1];
+		OutTAAParameters->PlusWeights[1] = OutTAAParameters->SampleWeights[3];
 		OutTAAParameters->PlusWeights[2] = OutTAAParameters->SampleWeights[4];
 		OutTAAParameters->PlusWeights[3] = OutTAAParameters->SampleWeights[5];
 		OutTAAParameters->PlusWeights[4] = OutTAAParameters->SampleWeights[7];
-
 		float TotalWeightPlus = (
-			OutTAAParameters->PlusWeights[0] +
-			OutTAAParameters->PlusWeights[1] +
-			OutTAAParameters->PlusWeights[2] +
-			OutTAAParameters->PlusWeights[3] +
-			OutTAAParameters->PlusWeights[4]);
+			OutTAAParameters->SampleWeights[1] +
+			OutTAAParameters->SampleWeights[3] +
+			OutTAAParameters->SampleWeights[4] +
+			OutTAAParameters->SampleWeights[5] +
+			OutTAAParameters->SampleWeights[7]);
 	
 		for (int32 i = 0; i < 5; i++)
 			OutTAAParameters->PlusWeights[i] /= TotalWeightPlus;
@@ -653,7 +659,7 @@ FTAAOutputs AddTemporalAAPass(
 			NewHistoryTexture[i] = GraphBuilder.CreateTexture(
 				SceneColorDesc,
 				OutputName,
-				ERDGParentResourceFlags::MultiFrame);
+				ERDGResourceFlags::MultiFrame);
 		}
 
 		NewHistoryTexture[0] = Outputs.SceneColor = NewHistoryTexture[0];
@@ -936,11 +942,10 @@ FTAAOutputs AddTemporalAAPass(
 	return Outputs;
 } // AddTemporalAAPass()
 
-void AddTemporalAA2Passes(
+static void AddGen5MainTemporalAAPasses(
 	FRDGBuilder& GraphBuilder,
-	const FSceneTextureParameters& SceneTextures,
 	const FViewInfo& View,
-	FRDGTextureRef InputSceneColorTexture,
+	const ITemporalUpscaler::FPassInputs& PassInputs,
 	FRDGTextureRef* OutSceneColorTexture,
 	FIntRect* OutSceneColorViewRect)
 {
@@ -950,7 +955,7 @@ void AddTemporalAA2Passes(
 	// Whether to use camera cut shader permutation or not.
 	bool bCameraCut = !InputHistory.IsValid() || View.bCameraCut;
 
-	FIntPoint InputExtent = InputSceneColorTexture->Desc.Extent;
+	FIntPoint InputExtent = PassInputs.SceneColorTexture->Desc.Extent;
 	FIntRect InputRect = View.ViewRect;
 
 	FIntPoint LowFrequencyExtent = InputExtent;
@@ -1052,8 +1057,8 @@ void AddTemporalAA2Passes(
 
 			FTAA2DilateVelocityCS::FParameters* PassParameters = GraphBuilder.AllocParameters<FTAA2DilateVelocityCS::FParameters>();
 			PassParameters->CommonParameters = CommonParameters;
-			PassParameters->SceneDepthTexture = SceneTextures.SceneDepthBuffer;
-			PassParameters->SceneVelocityTexture = SceneTextures.SceneVelocityBuffer;
+			PassParameters->SceneDepthTexture = PassInputs.SceneDepthTexture;
+			PassParameters->SceneVelocityTexture = PassInputs.SceneVelocityTexture;
 			PassParameters->DilatedVelocityOutput = GraphBuilder.CreateUAV(DilatedVelocityTexture);
 			PassParameters->ClosestDepthOutput = GraphBuilder.CreateUAV(ClosestDepthTexture);
 			PassParameters->PrevUseCountOutput = GraphBuilder.CreateUAV(PrevUseCountTexture);
@@ -1232,7 +1237,7 @@ void AddTemporalAA2Passes(
 
 			FTAA2FilterFrequenciesCS::FParameters* PassParameters = GraphBuilder.AllocParameters<FTAA2FilterFrequenciesCS::FParameters>();
 			PassParameters->CommonParameters = CommonParameters;
-			PassParameters->InputTexture = InputSceneColorTexture;
+			PassParameters->InputTexture = PassInputs.SceneColorTexture;
 			PassParameters->PredictionSceneColorTexture = PredictionSceneColorTexture;
 			PassParameters->PredictionInfoTexture = PredictionInfoTexture;
 
@@ -1320,8 +1325,9 @@ void AddTemporalAA2Passes(
 
 		FTAA2UpdateHistoryCS::FParameters* PassParameters = GraphBuilder.AllocParameters<FTAA2UpdateHistoryCS::FParameters>();
 		PassParameters->CommonParameters = CommonParameters;
-		PassParameters->InputSceneColorTexture = InputSceneColorTexture;
-		PassParameters->InputSceneStencilTexture = GraphBuilder.CreateSRV(FRDGTextureSRVDesc::CreateWithPixelFormat(SceneTextures.SceneDepthBuffer, PF_X24_G8));
+		PassParameters->InputSceneColorTexture = PassInputs.SceneColorTexture;
+		PassParameters->InputSceneStencilTexture = GraphBuilder.CreateSRV(
+			FRDGTextureSRVDesc::CreateWithPixelFormat(PassInputs.SceneDepthTexture, PF_X24_G8));
 		PassParameters->HistoryRejectionTexture = DilatedHistoryRejectionTexture;
 		PassParameters->DilatedVelocityTexture = DilatedVelocityTexture;
 		PassParameters->ParallaxRejectionMaskTexture = ParallaxRejectionMaskTexture;
@@ -1355,7 +1361,7 @@ void AddTemporalAA2Passes(
 
 			if (bPrevFrameIsntAvailable && !PassParameters->bCameraCut)
 			{
-				ensureMsgf(false, TEXT("Shaders read PrevHistory[%d] but doesn't write HistoryOutput[%d]"), i, i);
+				//ensureMsgf(false, TEXT("Shaders read PrevHistory[%d] but doesn't write HistoryOutput[%d]"), i, i);
 				PassParameters->bCameraCut = true;
 			}
 
@@ -1391,9 +1397,9 @@ void AddTemporalAA2Passes(
 
 	*OutSceneColorTexture = SceneColorOutputTexture;
 	*OutSceneColorViewRect = OutputRect;
-} // AddTemporalAA2Passes()
+} // AddGen5MainTemporalAAPasses()
 
-static void AddTemporalAAPass(
+static void AddGen4MainTemporalAAPasses(
 	FRDGBuilder& GraphBuilder,
 	const FViewInfo& View,
 	const ITemporalUpscaler::FPassInputs& PassInputs,
@@ -1479,7 +1485,7 @@ static void AddTemporalAAPass(
 	*OutSceneColorViewRect = SecondaryViewRect;
 	*OutSceneColorHalfResTexture = TAAOutputs.DownsampledSceneColor;
 	*OutSceneColorHalfResViewRect = FIntRect::DivideAndRoundUp(SecondaryViewRect, 2);
-} // AddTemporalAAPass()
+} // AddGen4MainTemporalAAPasses()
 
 const ITemporalUpscaler* GTemporalUpscaler = nullptr;
 
@@ -1501,14 +1507,29 @@ public:
 		FRDGTextureRef* OutSceneColorHalfResTexture,
 		FIntRect* OutSceneColorHalfResViewRect) const final
 	{
-		return AddTemporalAAPass(
-			GraphBuilder,
-			View,
-			PassInputs,
-			OutSceneColorTexture,
-			OutSceneColorViewRect,
-			OutSceneColorHalfResTexture,
-			OutSceneColorHalfResViewRect);
+		if (CVarTAAAlgorithm.GetValueOnRenderThread())
+		{
+			*OutSceneColorHalfResTexture = nullptr;
+			*OutSceneColorHalfResViewRect; // TODO.
+
+			return AddGen5MainTemporalAAPasses(
+				GraphBuilder,
+				View,
+				PassInputs,
+				OutSceneColorTexture,
+				OutSceneColorViewRect);
+		}
+		else
+		{
+			return AddGen4MainTemporalAAPasses(
+				GraphBuilder,
+				View,
+				PassInputs,
+				OutSceneColorTexture,
+				OutSceneColorViewRect,
+				OutSceneColorHalfResTexture,
+				OutSceneColorHalfResViewRect);
+		}
 	}
 };
 
diff --git a/Engine/Source/Runtime/Renderer/Private/SceneRendering.h b/Engine/Source/Runtime/Renderer/Private/SceneRendering.h
--- a/Engine/Source/Runtime/Renderer/Private/SceneRendering.h
+++ b/Engine/Source/Runtime/Renderer/Private/SceneRendering.h
@@ -723,7 +723,7 @@ struct FTemporalAAHistory
 
 	// Render targets holding's pixel history.
 	//  scene color's RGBA are in RT[0].
-	TRefCountPtr<IPooledRenderTarget> RT[kRenderTargetCount];
+	TStaticArray<TRefCountPtr<IPooledRenderTarget>, kRenderTargetCount> RT;
 
 	// Reference size of RT. Might be different than RT's actual size to handle down res.
 	FIntPoint ReferenceBufferSize;
@@ -756,7 +756,7 @@ struct FScreenSpaceDenoiserHistory
 	FIntRect Scissor;
 
 	// Render target specific to the history.
-	TRefCountPtr<IPooledRenderTarget> RT[RTCount];
+	TStaticArray<TRefCountPtr<IPooledRenderTarget>, RTCount> RT;
 
 	// The texture for tile classification.
 	TRefCountPtr<IPooledRenderTarget> TileClassification;
@@ -785,9 +785,9 @@ struct FGTAOTAAHistory
 
 	// Render targets holding's pixel history.
 	//  scene color's RGBA are in RT[0].
-	TRefCountPtr<IPooledRenderTarget> RT[kRenderTargetCount];
-	TRefCountPtr<IPooledRenderTarget> Depth[kRenderTargetCount];
-	TRefCountPtr<IPooledRenderTarget> Velocity[kRenderTargetCount];
+	TStaticArray<TRefCountPtr<IPooledRenderTarget>, kRenderTargetCount> RT;
+	TStaticArray<TRefCountPtr<IPooledRenderTarget>, kRenderTargetCount> Depth;
+	TStaticArray<TRefCountPtr<IPooledRenderTarget>, kRenderTargetCount> Velocity;
 
 	// Reference size of RT. Might be different than RT's actual size to handle down res.
 	FIntPoint ReferenceBufferSize;
