Shadow projection shader refactoring.
Sampling a virtual shadow map (directional clipmap or spot light) can now be more easily done by:
- Include FVirtualShadowMapProjectionParameters in shader parameter struct
- Call FVirtualShadowMapArray::SetShaderDefines in ModifyCompilationEnvironment
- Set shader parameters via FVirtualShadowMapArray::SetProjectionParameters
- Calling SampleVirtualShadowMapDepth in the shader

diff --git a/Engine/Shaders/Private/VirtualShadowMaps/CacheManagement.usf b/Engine/Shaders/Private/VirtualShadowMaps/CacheManagement.usf
--- a/Engine/Shaders/Private/VirtualShadowMaps/CacheManagement.usf
+++ b/Engine/Shaders/Private/VirtualShadowMaps/CacheManagement.usf
@@ -4,6 +4,7 @@
 #include "../SceneData.ush"
 #include "ProjectionCommon.ush"
 #include "PageAccessCommon.ush"
+#include "../ScreenSpaceDenoise/SSDDefinitions.ush"		// For LIGHT_TYPE's
 
 //#include "NaniteDataDecode.ush"
  #include "../Nanite/HZBCull.ush"
@@ -46,13 +47,14 @@ void VirtualSmInvalidateInstancePagesCS(
 		{
 			// 1. Load cached projection data
 			FVirtualShadowMapProjectionShaderData ProjectionData = ShadowMapProjectionData[ShadowMapDataIndex];
+			const bool bDirectionalLight = (ProjectionData.LightType == LIGHT_TYPE_DIRECTIONAL);
 
 			float4x4 LocalToTranslatedWorld = InstanceData.LocalToWorld;
 			LocalToTranslatedWorld[3].xyz += ProjectionData.ShadowPreViewTranslation;
 			// TODO: put mat mul in pre-computed
 			float4x4 LocalToClip = mul(LocalToTranslatedWorld, mul(ProjectionData.TranslatedWorldToShadowViewMatrix, ProjectionData.ShadowViewToClipMatrix));
 
-			FFrustumCullData Cull = BoxCullFrustum(InstanceData.LocalBoundsCenter, InstanceData.LocalBoundsExtent, LocalToClip, ProjectionData.bNearClip, false);
+			FFrustumCullData Cull = BoxCullFrustum(InstanceData.LocalBoundsCenter, InstanceData.LocalBoundsExtent, LocalToClip, !bDirectionalLight, false);
 
 			if (Cull.bIsVisible)
 			{
diff --git a/Engine/Shaders/Private/VirtualShadowMaps/PageManagement.usf b/Engine/Shaders/Private/VirtualShadowMaps/PageManagement.usf
--- a/Engine/Shaders/Private/VirtualShadowMaps/PageManagement.usf
+++ b/Engine/Shaders/Private/VirtualShadowMaps/PageManagement.usf
@@ -391,7 +391,7 @@ void GeneratePageFlagsFromPixels(uint3 GroupId : SV_GroupID, uint GroupIndex : S
 	{
 		// Mark one page per directional light. Lights can have varying numbers of clipmap levels allocated which
 		// will all be contiguous in the array, so we hop between them using the counts.
-		int ProjectionIndex = 0;
+		uint ProjectionIndex = 0;
 		while (ProjectionIndex < NumDirectionalLightSmInds)
 		{
 			int ClipmapStartId = VirtualShadowMapIdRemap[ProjectionIndex];
@@ -402,7 +402,7 @@ void GeneratePageFlagsFromPixels(uint3 GroupId : SV_GroupID, uint GroupIndex : S
 
 			// Possibly use a single shadow projection data for the whole clipmap with just scales/biases for
 			// each level (to handle snapping). For now we maintain one per level.
-			int ClipmapLevel = CalcClipmapLevel(Pixel.ViewPosition, ClipmapResolutionLodBias);
+			int ClipmapLevel = CalcClipmapLevel(length(Pixel.ViewPosition), ClipmapResolutionLodBias);
 			int ClipmapIndex = max(0, ClipmapLevel - FirstClipmapLevel);
 			if (ClipmapIndex < ClipmapLevelCount)
 			{
diff --git a/Engine/Shaders/Private/VirtualShadowMaps/ProjectionCommon.ush b/Engine/Shaders/Private/VirtualShadowMaps/ProjectionCommon.ush
--- a/Engine/Shaders/Private/VirtualShadowMaps/ProjectionCommon.ush
+++ b/Engine/Shaders/Private/VirtualShadowMaps/ProjectionCommon.ush
@@ -5,6 +5,12 @@ ProjectionCommon.ush:
 =============================================================================*/
 #pragma once
 
+#include "../Common.ush"
+#include "../ScreenSpaceDenoise/SSDDefinitions.ush"		// For LIGHT_TYPE's
+#include "PageAccessCommon.ush"
+
+#define VIRTUAL_SHADOW_MAP_MIN_OCCLUDER_DISTANCE 1e-6f
+
 // Must match C++ structure in VirtualShadowMapArray.h
 struct FVirtualShadowMapProjectionShaderData
 {
@@ -13,18 +19,28 @@ struct FVirtualShadowMapProjectionShaderData
 	float4x4 ShadowViewToClipMatrix;
 	float4x4 TranslatedWorldToShadowUvNormalMatrix;
 	float3 ShadowPreViewTranslation;
-	float Dummy;
+	uint LightType;		// Matches ELightComponentType via defines in SSDDefinitions.ush
+
+	float3 ClipmapWorldOrigin;
 	int VirtualShadowMapId;	
-	uint bNearClip;
 
 	int ClipmapLevel;
 	int ClipmapLevelCount;
 	float ClipmapResolutionLodBias;
 
 	// Seems the FMatrix forces 16-byte alignment, so we need to pad
-	float Padding[3];
+	float Padding[1];
 };
 
+// FVirtualShadowMapProjectionParameters
+StructuredBuffer< FVirtualShadowMapProjectionShaderData > VirtualShadowMapProjectionData;
+Texture2D< uint > PhysicalPagePool;
+
+FVirtualShadowMapProjectionShaderData GetVirtualShadowMapProjectionData(int VirtualShadowMapId)
+{
+	return VirtualShadowMapProjectionData[VirtualShadowMapId];
+}
+
 uint CalcMipLevelFromFootprint(float Footprint, uint MipLevelCount = VSM_MAX_MIP_LEVELS)
 {
 	float MipLevelFloat = log2(max(1.0f, Footprint));
@@ -32,13 +48,95 @@ uint CalcMipLevelFromFootprint(float Footprint, uint MipLevelCount = VSM_MAX_MIP
 	return MipLevel;
 }
 
-int CalcClipmapLevel(float3 ViewPosition, float ResolutionLodBias)
+int CalcClipmapLevel(float DistanceToClipmapOrigin, float ResolutionLodBias)
 {
-	float Log2Distance = log2(length(ViewPosition));
+	float Log2Distance = log2(DistanceToClipmapOrigin);
 	int AbsoluteLevel = int(floor(Log2Distance + ResolutionLodBias));
 	return AbsoluteLevel;
 }
 
+struct FVirtualSmSample
+{
+	float Depth;
+	uint2 VirtualTexelAddress;
+	float2 VirtualTexelAddressFloat;
+	uint2 PhysicalTexelAddress;
+	uint MipLevel;
+	bool bValid;
+};
+
+FVirtualSmSample SampleVirtualSmLevel(int VirtualShadowMapId, float2 ShadowMapUV, uint VSMLevel)
+{
+	float2 VirtualTexelAddressFloat = ShadowMapUV * float(CalcLevelDimsTexels(VSMLevel));
+	uint2 VirtualTexelAddress = uint2(VirtualTexelAddressFloat);
+
+	// Page map to physical address
+	uint2 PhysicalTexelAddress = uint2(0, 0);
+	bool PagePresent = VirtualToPhysicalTexel(VirtualShadowMapId, VSMLevel, VirtualTexelAddress, PhysicalTexelAddress);
+	if (PagePresent)
+	{
+		// TODO: maybe store metadata about Max depth (and/or use sparse HZB) for each tile, and early out if completely in shadow (must test using kernel size).
+		FVirtualSmSample Result;
+		Result.Depth = asfloat(PhysicalPagePool.Load(int3(PhysicalTexelAddress, 0)));
+		Result.VirtualTexelAddress = VirtualTexelAddress;
+		Result.PhysicalTexelAddress = PhysicalTexelAddress;
+		Result.VirtualTexelAddressFloat = VirtualTexelAddressFloat;
+		Result.MipLevel = VSMLevel;
+		Result.bValid = true;
+		return Result;
+	}
+	else
+	{
+		FVirtualSmSample Result;
+		Result.Depth = 0.0f;
+		Result.VirtualTexelAddress = uint2(0, 0);
+		Result.PhysicalTexelAddress = uint2(0, 0);
+		Result.MipLevel = 0;
+		Result.bValid = false;
+		return Result;
+	}
+}
+
+FVirtualSmSample SampleVirtualSm(int VirtualShadowMapId, float2 ShadowMapUV)
+{
+	// Compute virtual address, map to physical & sample
+	for (uint VSMLevel = 0; VSMLevel < VSM_MAX_MIP_LEVELS; ++VSMLevel)
+	{
+		FVirtualSmSample Result = SampleVirtualSmLevel(VirtualShadowMapId, ShadowMapUV, VSMLevel);
+		if (Result.bValid)
+		{
+			return Result;
+		}
+	}
+
+	FVirtualSmSample Result;
+	Result.Depth = 0.0f;
+	Result.VirtualTexelAddress = uint2(0, 0);
+	Result.PhysicalTexelAddress = uint2(0, 0);
+	Result.MipLevel = 0;
+	Result.bValid = false;
+	return Result;
+}
+
+float ComputeVirtualShadowMapOptimalSlopeBias(FVirtualShadowMapProjectionShaderData ProjectionData, FVirtualSmSample SmSample, float3 WorldPosition, float3 EstimatedGeoWorldNormal)
+{
+	// Transform geometry world-space plane eq to shadow 'UV' texture space [0-1] ranges
+	float4 NormalPlaneTranslatedWorld = float4(EstimatedGeoWorldNormal, -dot(EstimatedGeoWorldNormal, WorldPosition + ProjectionData.ShadowPreViewTranslation));
+	float4 NormalPlaneUv = mul(NormalPlaneTranslatedWorld, ProjectionData.TranslatedWorldToShadowUvNormalMatrix);
+
+	float2 DepthSlopeUv = -NormalPlaneUv.xy / NormalPlaneUv.z;
+	float MipLevelDim = float(CalcLevelDimsTexels(SmSample.MipLevel));
+	float2 TexelCenter = float2(SmSample.VirtualTexelAddress) + 0.5;
+	float2 TexelCenterOffset = TexelCenter - SmSample.VirtualTexelAddressFloat;
+	float2 TexelCenterOffsetUv = TexelCenterOffset / MipLevelDim;
+	// 2x factor due to lack of precision (probably)
+	float OptimalSlopeBias = 2.0f * max(0.0f, dot(DepthSlopeUv, TexelCenterOffsetUv));
+	// Clamp to avoid excessive degenerate slope biases causing flickering lit pixels
+	OptimalSlopeBias = min(OptimalSlopeBias, 0.01f);
+
+	return OptimalSlopeBias;
+}
+
 // Small wrapper for point transformed into shadow space
 struct FShadowPosition
 {
@@ -51,7 +149,6 @@ struct FShadowPosition
 FShadowPosition WorldToShadow(FVirtualShadowMapProjectionShaderData ProjectionData, float3 WorldPosition)
 {
 	FShadowPosition ShadowPosition;
-	// 
 	ShadowPosition.View = mul(float4(WorldPosition + ProjectionData.ShadowPreViewTranslation, 1.0f), ProjectionData.TranslatedWorldToShadowViewMatrix).xyz;
 
 	ShadowPosition.ClipH = mul(float4(ShadowPosition.View, 1.0f), ProjectionData.ShadowViewToClipMatrix);
@@ -62,3 +159,117 @@ FShadowPosition WorldToShadow(FVirtualShadowMapProjectionShaderData ProjectionDa
 
 	return ShadowPosition;
 }
+
+struct FVirtualShadowMapSampleResult
+{
+	bool bValid;
+	bool bOccluded;
+	float OccluderDistance;
+};
+
+// Used for orthographic projections (i.e. directional lights)
+// Receiver depth is post-projection-divide.
+float ComputeOccluderDistanceOrtho(float4x4 ShadowViewToClip, float OccluderDepth, float ReceiverDepth)
+{
+	float OccluderViewZ = (OccluderDepth - ShadowViewToClip._43) / ShadowViewToClip._33;
+	float ReceiverViewZ = (ReceiverDepth - ShadowViewToClip._43) / ShadowViewToClip._33;
+
+	// No perspective projection, so simple difference gets us the distance
+	float Result = ReceiverViewZ - OccluderViewZ;
+	return max(VIRTUAL_SHADOW_MAP_MIN_OCCLUDER_DISTANCE, Result);
+}
+
+// Used for perspective projections (i.e. spot lights)
+// Receiver depth is post-projection-divide.
+float ComputeOccluderDistancePerspective(float4x4 ShadowViewToClip, float OccluderDepth, float ReceiverDepth, float3 ReceiverShadowViewPosition)
+{
+	float OccluderViewZ = ShadowViewToClip._43 / (OccluderDepth - ShadowViewToClip._33);
+	float ReceiverViewZ = ShadowViewToClip._43 / (ReceiverDepth - ShadowViewToClip._33);
+
+	// Similar triangles to compute euclidean distance in view/world space
+	float ReceiverDistance = length(ReceiverShadowViewPosition);
+	float OccluderDistance = (ReceiverDistance / ReceiverViewZ) * OccluderViewZ;
+	float Result = ReceiverDistance - OccluderDistance;	
+	return max(VIRTUAL_SHADOW_MAP_MIN_OCCLUDER_DISTANCE, Result);
+}
+
+/**
+* VirtualShadowMapId is the ID of the virtual shadow map to sample
+* WorldPosition is the sample position in world space to project into the shadow map
+* EstimatedGeoWorldNormal is ideally the geometric (flat) normal of the sample point in world space
+*   - This is used to compute a receiver-plane-based bias for the sample point
+*   - The shading normal can be used if the geometric normal is not available, but divergence from the geometric normal can cause biasing issues.
+*/
+FVirtualShadowMapSampleResult SampleVirtualShadowMap(int VirtualShadowMapId, float3 WorldPosition, float3 EstimatedGeoWorldNormal)
+{
+	FVirtualShadowMapProjectionShaderData BaseProjectionData = GetVirtualShadowMapProjectionData(VirtualShadowMapId);
+
+	FVirtualSmSample SmSample;
+	SmSample.Depth = 0.0f;
+	SmSample.VirtualTexelAddress = uint2(0, 0);
+	SmSample.PhysicalTexelAddress = uint2(0, 0);
+	SmSample.MipLevel = 0;
+	SmSample.bValid = false;
+
+	FShadowPosition Position;
+	float OccluderDistance = 0.0f;
+	float BiasedDepth = 0.0f;
+	bool bOccluded = false;
+
+	if (BaseProjectionData.LightType == LIGHT_TYPE_DIRECTIONAL)
+	{
+		const int FirstClipmapLevel = BaseProjectionData.ClipmapLevel;
+		float DistanceToClipmapOrigin = length(WorldPosition - BaseProjectionData.ClipmapWorldOrigin);
+		int ClipmapLevel = CalcClipmapLevel(DistanceToClipmapOrigin, BaseProjectionData.ClipmapResolutionLodBias);
+		int ClipmapIndex = max(0, ClipmapLevel - FirstClipmapLevel);
+
+		// Check if sample is within the clipmap range (from camera)
+		if (ClipmapIndex < BaseProjectionData.ClipmapLevelCount)
+		{
+			// NOTE: Clipmap levels are contiguous in the virtual shadow map list
+			int ClipmapLevelVirtualShadowMapId = VirtualShadowMapId + ClipmapIndex;
+			FVirtualShadowMapProjectionShaderData ProjectionData = GetVirtualShadowMapProjectionData(ClipmapLevelVirtualShadowMapId);
+
+			Position = WorldToShadow(ProjectionData, WorldPosition);
+			SmSample = SampleVirtualSmLevel(ClipmapLevelVirtualShadowMapId, Position.UV, 0);
+			if (SmSample.bValid)
+			{
+				BiasedDepth = SmSample.Depth - ComputeVirtualShadowMapOptimalSlopeBias(ProjectionData, SmSample, WorldPosition, EstimatedGeoWorldNormal);
+
+				if (BiasedDepth > Position.Depth)
+				{
+					bOccluded = true;
+					OccluderDistance = ComputeOccluderDistanceOrtho(
+						ProjectionData.ShadowViewToClipMatrix,
+						BiasedDepth,
+						Position.Depth);
+				}
+			}
+		}
+	}
+	else	// Spot light
+	{
+		Position = WorldToShadow(BaseProjectionData, WorldPosition);
+		SmSample = SampleVirtualSm(VirtualShadowMapId, Position.UV);
+		if (SmSample.bValid)
+		{
+			BiasedDepth = SmSample.Depth - ComputeVirtualShadowMapOptimalSlopeBias(BaseProjectionData, SmSample, WorldPosition, EstimatedGeoWorldNormal);
+
+			if (BiasedDepth > Position.Depth)
+			{
+				bOccluded = true;
+				OccluderDistance = ComputeOccluderDistancePerspective(
+					BaseProjectionData.ShadowViewToClipMatrix,
+					BiasedDepth,
+					Position.Depth,
+					Position.View);
+			}
+		}
+	}
+
+	FVirtualShadowMapSampleResult Result;
+	Result.bValid = SmSample.bValid;
+	Result.bOccluded = bOccluded;
+	Result.OccluderDistance = OccluderDistance;
+	return Result;
+}
diff --git a/Engine/Shaders/Private/VirtualShadowMaps/VirtualShadowMapProjection.usf b/Engine/Shaders/Private/VirtualShadowMaps/VirtualShadowMapProjection.usf
--- a/Engine/Shaders/Private/VirtualShadowMaps/VirtualShadowMapProjection.usf
+++ b/Engine/Shaders/Private/VirtualShadowMaps/VirtualShadowMapProjection.usf
@@ -22,14 +22,7 @@
 #define OUTPUT_TYPE 0
 #endif
 
-#define MIN_OCCLUDER_DISTANCE 1e-6f
-
-// One per clipmap level for directional lights
-StructuredBuffer< FVirtualShadowMapProjectionShaderData > ShadowMapProjectionData;
-Texture2D< uint > PhysicalPagePool;
-
-uint LightType;		// Matches ELightComponentType via defines in SSDDefinitions.ush
-uint ProjectionCount;
+int VirtualShadowMapId;		// For the first layer for clipmaps
 int DebugOutputType;
 float ContactShadowLength;
 
@@ -51,70 +44,6 @@ float3 IntToColor(uint Index)
 	return ColorTable[Index % 8];
 }
 
-struct FVirtualSmSample
-{
-	float Depth;
-	uint2 VirtualTexelAddress;
-	float2 VirtualTexelAddressFloat;
-	uint2 PhysicalTexelAddress;
-	uint MipLevel;
-	bool bValid;
-};
-
-FVirtualSmSample SampleVirtualSmLevel(uint VirtualShadowMapId, float2 ShadowMapUV, uint VSMLevel)
-{
-	uint VSMLevelDim = CalcLevelDimsTexels(VSMLevel);
-	float2 VirtualTexelAddressFloat = ShadowMapUV * float(VSMLevelDim);
-	uint2 VirtualTexelAddress = uint2(VirtualTexelAddressFloat);
-
-	// Page map to physical address
-	uint2 PhysicalTexelAddress = uint2(0, 0);
-	bool PagePresent = VirtualToPhysicalTexel(VirtualShadowMapId, VSMLevel, VirtualTexelAddress, PhysicalTexelAddress);
-	if (PagePresent)
-	{
-		// TODO: maybe store metadata about Max depth (and/or use sparse HZB) for each tile, and early out if completely in shadow (must test using kernel size).
-		FVirtualSmSample Result;
-		Result.Depth = asfloat(PhysicalPagePool.Load(int3(PhysicalTexelAddress, 0)));
-		Result.VirtualTexelAddress = VirtualTexelAddress;
-		Result.PhysicalTexelAddress = PhysicalTexelAddress;
-		Result.VirtualTexelAddressFloat = VirtualTexelAddressFloat;
-		Result.MipLevel = VSMLevel;
-		Result.bValid = true;
-		return Result;
-	}
-	else
-	{
-		FVirtualSmSample Result;
-		Result.Depth = 0.0f;
-		Result.VirtualTexelAddress = uint2(0, 0);
-		Result.PhysicalTexelAddress = uint2(0, 0);
-		Result.MipLevel = 0;
-		Result.bValid = false;
-		return Result;
-	}
-}
-
-FVirtualSmSample SampleVirtualSm(uint VirtualShadowMapId, float2 ShadowMapUV)
-{
-	// Compute virtual address, map to physical & sample
-	for (uint VSMLevel = 0; VSMLevel < VSM_MAX_MIP_LEVELS; ++VSMLevel)
-	{
-		FVirtualSmSample Result = SampleVirtualSmLevel(VirtualShadowMapId, ShadowMapUV, VSMLevel);
-		if (Result.bValid)
-		{
-			return Result;
-		}
-	}
-
-	FVirtualSmSample Result;
-	Result.Depth = 0.0f;
-	Result.VirtualTexelAddress = uint2(0, 0);
-	Result.PhysicalTexelAddress = uint2(0, 0);
-	Result.MipLevel = 0;
-	Result.bValid = false;
-	return Result;
-}
-
 float ShadowRayCast(
 	float3 RayOriginTranslatedWorld, float3 RayDirection, float RayLength,
 	int NumSteps, float StepOffset
@@ -180,33 +109,6 @@ float4 EncodeLightAttenuationFromMask(float ShadowMask)
 	return EncodeLightAttenuation(half4(FadedShadow, FadedSSSShadow, FadedShadow, FadedSSSShadow));
 }
 
-// Receiver depth is post-projection-divide.
-float ComputeOccluderDistance(float4x4 ShadowViewToClip, float OccluderDepth, float ReceiverDepth, float3 ReceiverShadowViewPosition)
-{
-	float Result = 0.0f;
-	if (ShadowViewToClip._44 < 1.0f)
-	{
-		// Perspective projection
-		float OccluderViewZ = ShadowViewToClip._43 / (OccluderDepth - ShadowViewToClip._33);
-		float ReceiverViewZ = ShadowViewToClip._43 / (ReceiverDepth - ShadowViewToClip._33);
-
-		// Similar triangles to compute euclidean distance in view/world space
-		float ReceiverDistance = length(ReceiverShadowViewPosition);
-		float OccluderDistance = (ReceiverDistance / ReceiverViewZ) * OccluderViewZ;
-		Result = ReceiverDistance - OccluderDistance;
-	}
-	else
-	{
-		// Ortho
-		float OccluderViewZ = (OccluderDepth - ShadowViewToClip._43) / ShadowViewToClip._33;
-		float ReceiverViewZ = (ReceiverDepth - ShadowViewToClip._43) / ShadowViewToClip._33;
-
-		// No perspective projection, so simple difference gets us the distance
-		Result = ReceiverViewZ - OccluderViewZ;
-	}
-	return max(MIN_OCCLUDER_DISTANCE, Result);
-}
-
 void VirtualShadowMapProjectionPS(
 	in float4 SvPosition : SV_Position,
 	out float4 OutputTarget : SV_Target
@@ -221,11 +123,13 @@ void VirtualShadowMapProjectionPS(
 		
 	// TODO: Early out on sky, etc.
 	bool bInLightRegion = true;
-	int ProjectionDataIndex = 0;
+
+	FVirtualShadowMapProjectionShaderData ProjectionData = GetVirtualShadowMapProjectionData(VirtualShadowMapId);
+	uint LightType = ProjectionData.LightType;
 
 	FLightShaderParameters Light = GetRootLightShaderParameters();
 	float3 ToLight = Light.Direction;
-	if (LightType == LIGHT_TYPE_SPOT)
+	if (ProjectionData.LightType == LIGHT_TYPE_SPOT)
 	{
 		// Early out if we're outside the spot light cone		
 		ToLight = Light.Position - (TranslatedWorldPosition - View.PreViewTranslation);
@@ -234,17 +138,6 @@ void VirtualShadowMapProjectionPS(
 		bool bInLightCone = dot(ToLight, Light.Direction) >= Light.SpotAngles.x;
 		bInLightRegion = bInLightRegion && bInLightRadius && bInLightCone;
 	}
-	else if (LightType == LIGHT_TYPE_DIRECTIONAL)
-	{
-		float3 ViewPosition = mul(float4(TranslatedWorldPosition, 1.0f), View.TranslatedWorldToView).xyz;
-
-		const int FirstClipmapLevel = ShadowMapProjectionData[0].ClipmapLevel;
-		const int ClipmapLevelCount = ShadowMapProjectionData[0].ClipmapLevelCount;
-		const float ClipmapResolutionLodBias = ShadowMapProjectionData[0].ClipmapResolutionLodBias;
-		int ClipmapLevel = CalcClipmapLevel(ViewPosition, ClipmapResolutionLodBias);
-		ProjectionDataIndex = max(0, ClipmapLevel - FirstClipmapLevel);
-		bInLightRegion = (ProjectionDataIndex < ClipmapLevelCount);
-	}
 	if (!bInLightRegion)
 	{
 		discard;
@@ -264,26 +157,7 @@ void VirtualShadowMapProjectionPS(
 
 	// NOTE: Could skip most of the rest of the shader if contact shadows say we're shadowed, but likely to be small, incoherent branching
 
-	FVirtualShadowMapProjectionShaderData ProjectionData = ShadowMapProjectionData[ProjectionDataIndex];
-	FShadowPosition ShadowPosition = WorldToShadow(ProjectionData, TranslatedWorldPosition - View.PreViewTranslation);
-	
-	FVirtualSmSample SmSample;
-	if (LightType == LIGHT_TYPE_DIRECTIONAL)
-	{
-		SmSample = SampleVirtualSmLevel(ProjectionData.VirtualShadowMapId, ShadowPosition.UV, 0);
-	}
-	else
-	{
-		SmSample = SampleVirtualSm(ProjectionData.VirtualShadowMapId, ShadowPosition.UV);
-	}
-
-	if (!SmSample.bValid)
-	{
-		// Generally shouldn't occur if PT allocation is conservative
-		discard;
-	}
-
-	// 3.1, figure out slope, we do world space since that is the space where we might override using the shading normal...
+	// Figure out slope, we do world space since that is the space where we might override using the shading normal...
 	float3 TranslatedWorldPositionDDX = DDX(TranslatedWorldPosition);
 	float3 TranslatedWorldPositionDDY = DDY(TranslatedWorldPosition);
 	float3 EstimatedGeoWorldNormal = normalize(cross(TranslatedWorldPositionDDX, TranslatedWorldPositionDDY));
@@ -297,38 +171,22 @@ void VirtualShadowMapProjectionPS(
 	}
 #endif
 
-	// Transform geometry world-space plane eq to shadow 'UV' texture space [0-1] ranges
-	float4 NormalPlaneWorld = float4(EstimatedGeoWorldNormal, -dot(EstimatedGeoWorldNormal, TranslatedWorldPosition));
-	float4 NormalPlaneUv = mul(NormalPlaneWorld, ProjectionData.TranslatedWorldToShadowUvNormalMatrix);
-		
-	float2 DepthSlopeUv = -NormalPlaneUv.xy / NormalPlaneUv.z;
-	float MipLevelDim = float(CalcLevelDimsTexels(SmSample.MipLevel));
-	float2 TexelCenter = float2(SmSample.VirtualTexelAddress) + 0.5;
-	float2 TexelCenterOffset = TexelCenter - SmSample.VirtualTexelAddressFloat;
-	float2 TexelCenterOffsetUv = TexelCenterOffset / MipLevelDim;
-	// 2x factor due to lack of precision (probably)
-	float OptimalSlopeBias = 2.0f * max(0.0f, dot(DepthSlopeUv, TexelCenterOffsetUv));
-	// Clamp to avoid excessive degenerate slope biases causing flickering lit pixels
-	OptimalSlopeBias = min(OptimalSlopeBias, 0.01f);
-
-	//SmSample.Depth = asfloat( max( asint( SmSample.Depth ) - 1, 0 ) );
+	float3 WorldPositon = TranslatedWorldPosition - View.PreViewTranslation;
+	FVirtualShadowMapSampleResult VirtualShadowMapSample = SampleVirtualShadowMap(VirtualShadowMapId, WorldPositon, EstimatedGeoWorldNormal);
 
-	float SmSampleBiasedDepth = SmSample.Depth - OptimalSlopeBias;
-
-	if (ContactShadowLength == 0.0f)
+	if (!VirtualShadowMapSample.bValid)
 	{
-		// When not doing contact shadow trace, at least add a constant bias to avoid obvious artifacts
-		// TODO: Could pull this bias itself from the contact shadow length... say if negative or something.
-		// This is mostly a debug path though so fine for the moment.
-		SmSampleBiasedDepth -= 0.0006f * SmSample.Depth;
+		// Generally shouldn't occur if PT allocation is conservative
+		discard;
 	}
 
 	// See signal encoding for SIGNAL_BUFFER_LAYOUT_PENUMBRA_RECONSTRUCTION in SSDSignalBufferEncoding.ush
 	float4 OutputSignal = float4(0, 0, 0, 0);
-	if (SmSampleBiasedDepth > ShadowPosition.Depth)
+	if (VirtualShadowMapSample.bOccluded)
 	{
 		// Hit
-		float OccluderDistance = ComputeOccluderDistance(ProjectionData.ShadowViewToClipMatrix, SmSampleBiasedDepth, ShadowPosition.Depth, ShadowPosition.View);
+		// Recall we offset the initial ray, so add the distance back in
+		float OccluderDistance = VirtualShadowMapSample.OccluderDistance + ContactShadowLengthWorld;
 		float WorldBluringRadius = ComputeLightSampleWorldBluringRadius(TranslatedWorldPosition, LightType, Light, OccluderDistance);
 		OutputSignal = float4(0, 1, WorldBluringRadius, 0);
 	}
@@ -338,7 +196,7 @@ void VirtualShadowMapProjectionPS(
 		OutputSignal = float4(1, 1, -1, 0);
 	}
 
-	if( ContactShadow < 1 )
+	if (ContactShadow < 1)
 	{
 		OutputSignal.x *= ContactShadow;
 		OutputSignal.z = ComputeLightSampleWorldBluringRadius(TranslatedWorldPosition, LightType, Light, ContactShadowLengthWorld);
@@ -354,23 +212,18 @@ void VirtualShadowMapProjectionPS(
 	// Output debug for "vis".
 	float3 Color = float3(0, 0, 0);
 
-	if (DebugOutputType == 1)		// Clipmap level
+	if (DebugOutputType == 1)	// Invalid page
 	{
-		Color = lerp(IntToColor(ProjectionDataIndex), OutputSignal.rrr, 0.4f);
+		Color = VirtualShadowMapSample.bValid ? (0.25f * OutputSignal.rrr) : float3(1, 0, 1);
 	}
-	else if (DebugOutputType == 2)	// Mip level/LOD
+	else if (DebugOutputType == 2)		// Clipmap level
 	{
-		Color = lerp(IntToColor(SmSample.MipLevel + 1), OutputSignal.rrr, 0.0f);
-	}
-	else if (DebugOutputType == 3)	// Invalid page
-	{
-		Color = SmSample.bValid ? (0.25f * OutputSignal.rrr) : float3(1, 0, 1);
+		//Color = lerp(IntToColor(ProjectionDataIndex), OutputSignal.rrr, 0.4f);
 	}
 
 	//Color = OptimalSlopeBias.xxx * 1.0f;
 
-	OutputTarget = float4(Color, 1.0f);
-	
+	OutputTarget = float4(Color, 1.0f);	
 #endif // OUTPUT_TYPE
 }
 
diff --git a/Engine/Source/Runtime/Renderer/Private/VirtualShadowMaps/VirtualShadowMapArray.cpp b/Engine/Source/Runtime/Renderer/Private/VirtualShadowMaps/VirtualShadowMapArray.cpp
--- a/Engine/Source/Runtime/Renderer/Private/VirtualShadowMaps/VirtualShadowMapArray.cpp
+++ b/Engine/Source/Runtime/Renderer/Private/VirtualShadowMaps/VirtualShadowMapArray.cpp
@@ -109,9 +109,9 @@ FVirtualShadowMapProjectionShaderData GetVirtualShadowMapProjectionShaderData(co
 	Data.TranslatedWorldToShadowViewMatrix = ShadowInfo->TranslatedWorldToView;
 	Data.ShadowViewToClipMatrix = ViewToClip;
 	Data.TranslatedWorldToShadowUvNormalMatrix = CalcTranslatedWorldToShadowUvNormalMatrix(ShadowInfo->TranslatedWorldToView, ViewToClip);
-	Data.ShadowPreViewTranslation = FVector4(ShadowInfo->PreShadowTranslation, 666.0f);
+	Data.ShadowPreViewTranslation = FVector(ShadowInfo->PreShadowTranslation);
 	Data.VirtualShadowMapId = ShadowInfo->VirtualShadowMap->ID;
-	Data.bNearClip = uint32(!ShadowInfo->ShouldClampToNearPlane());
+	Data.LightType = ShadowInfo->GetLightSceneInfo().Proxy->GetLightType();
 
 	return Data;
 }
@@ -871,6 +871,13 @@ void FVirtualShadowMapArray::GeneratePageFlagsFromLightGrid(FRDGBuilder& GraphBu
 	}
 }
 
+void FVirtualShadowMapArray::SetProjectionParameters(FRDGBuilder& GraphBuilder, FVirtualShadowMapProjectionParameters& OutParameters)
+{
+	OutParameters.CommonParameters = CommonParameters;
+	OutParameters.PageTable = GraphBuilder.CreateSRV(GraphBuilder.RegisterExternalBuffer(PageTable));
+	OutParameters.PhysicalPagePool = GraphBuilder.RegisterExternalTexture(PhysicalPagePool);
+	OutParameters.VirtualShadowMapProjectionData = GraphBuilder.CreateSRV(GraphBuilder.RegisterExternalBuffer(ShadowMapProjectionDataBuffer));
+}
 
 class FDebugVisualizeVirtualSmCS : public FVirtualPageManagementShader
 {
diff --git a/Engine/Source/Runtime/Renderer/Private/VirtualShadowMaps/VirtualShadowMapArray.h b/Engine/Source/Runtime/Renderer/Private/VirtualShadowMaps/VirtualShadowMapArray.h
--- a/Engine/Source/Runtime/Renderer/Private/VirtualShadowMaps/VirtualShadowMapArray.h
+++ b/Engine/Source/Runtime/Renderer/Private/VirtualShadowMaps/VirtualShadowMapArray.h
@@ -5,6 +5,7 @@
 #pragma once
 
 #include "../Nanite/NaniteRender.h"
+#include "SceneTypes.h"
 
 struct FSortedLightSetSceneInfo;
 class FViewInfo;
@@ -87,17 +88,18 @@ struct FVirtualShadowMapProjectionShaderData
 	/**
 	 * Translation from world space to shadow space (add before transform by TranslatedWorldToShadowViewMatrix).
 	 */
-	FVector4 ShadowPreViewTranslation;
-	int32 VirtualShadowMapId;
-	uint32 bNearClip = 1U;
+	FVector ShadowPreViewTranslation;
+	uint32 LightType = ELightComponentType::LightType_Directional;
+	
+	FVector ClipmapWorldOrigin;
+	int32 VirtualShadowMapId = INDEX_NONE;
 
-	// These could be per-light (first/count), but convenient here and not much overhead
 	int32 ClipmapLevel = 0;					// "Absolute" level, can be negative
 	int32 ClipmapLevelCount = 0;
 	float ClipmapResolutionLodBias = 0.0f;
 
 	// Seems the FMatrix forces 16-byte alignment
-	float Padding[3];
+	//float Padding[3];
 };
 static_assert((sizeof(FVirtualShadowMapProjectionShaderData) % 16) == 0, "FVirtualShadowMapProjectionShaderData size should be a multiple of 16-bytes for alignment.");
 
@@ -117,6 +119,13 @@ BEGIN_SHADER_PARAMETER_STRUCT(FVirtualShadowMapCommonParameters, )
 	SHADER_PARAMETER(uint32, PhysicalPageRowShift)
 END_SHADER_PARAMETER_STRUCT()
 
+BEGIN_SHADER_PARAMETER_STRUCT(FVirtualShadowMapProjectionParameters, )
+	SHADER_PARAMETER_STRUCT_INCLUDE(FVirtualShadowMapCommonParameters, CommonParameters)
+	SHADER_PARAMETER_RDG_BUFFER_SRV(StructuredBuffer<uint2>, PageTable)
+	SHADER_PARAMETER_RDG_TEXTURE(Texture2D<uint>, PhysicalPagePool)
+	SHADER_PARAMETER_RDG_BUFFER_SRV(StructuredBuffer< FVirtualShadowMapProjectionShaderData >, VirtualShadowMapProjectionData)
+END_SHADER_PARAMETER_STRUCT()
+
 class FVirtualShadowMapArray
 {
 public:
@@ -141,6 +150,8 @@ public:
 
 	static void SetShaderDefines(FShaderCompilerEnvironment& OutEnvironment);
 
+	void SetProjectionParameters(FRDGBuilder& GraphBuilder, FVirtualShadowMapProjectionParameters& OutParameters);
+
 	//
 	void GenerateIdentityPageTables(FRDGBuilder& GraphBuilder, uint32 MipLevel = INDEX_NONE);
 	void ClearPhysicalMemory(FRDGBuilder& GraphBuilder, FRDGTextureRef& PhysicalTexture, FVirtualShadowMapArrayCacheManager *VirtualShadowMapArrayCacheManager);
@@ -176,8 +187,8 @@ public:
 	TRefCountPtr<IPooledRenderTarget>	HZBPhysical;
 	TRefCountPtr<FPooledRDGBuffer>		HZBPageTable;
 
-	TRefCountPtr<IPooledRenderTarget> DebugVisualizationOutput;
-	TRefCountPtr<IPooledRenderTarget> DebugVisualizationProjectionOutput;
+	TRefCountPtr<IPooledRenderTarget>	DebugVisualizationOutput;
+	TRefCountPtr<IPooledRenderTarget>	DebugVisualizationProjectionOutput;
 	
 	TRefCountPtr<FPooledRDGBuffer>		AllocatedPagesOffset;
 
@@ -201,6 +212,5 @@ public:
 	// Used to clip the rect size of clusters during culling.
 	TRefCountPtr<FPooledRDGBuffer>		PageRectBounds;
 
-
 	TRefCountPtr<FPooledRDGBuffer>		ShadowMapProjectionDataBuffer;
 };
diff --git a/Engine/Source/Runtime/Renderer/Private/VirtualShadowMaps/VirtualShadowMapClipmap.cpp b/Engine/Source/Runtime/Renderer/Private/VirtualShadowMaps/VirtualShadowMapClipmap.cpp
--- a/Engine/Source/Runtime/Renderer/Private/VirtualShadowMaps/VirtualShadowMapClipmap.cpp
+++ b/Engine/Source/Runtime/Renderer/Private/VirtualShadowMaps/VirtualShadowMapClipmap.cpp
@@ -69,7 +69,7 @@ FVirtualShadowMapClipmap::FVirtualShadowMapClipmap(
 	LevelData.Empty();
 	LevelData.AddDefaulted(LevelCount);
 
-	const FVector ClipmapCenter = CameraViewMatrices.GetViewOrigin();
+	WorldOrigin = CameraViewMatrices.GetViewOrigin();
 
 	for (int32 Index = 0; Index < LevelCount; ++Index)
 	{
@@ -90,19 +90,19 @@ FVirtualShadowMapClipmap::FVirtualShadowMapClipmap(
 
 		// NOTE: We may eventually want to snap/quantize Z as well, but right now we just adjust it each time
 		// we update the cache page. With clipmaps this could ertainly be simplified.
-		FVector CenterView = TranslatedWorldToView.TransformPosition(ClipmapCenter);
+		FVector ViewCenter = TranslatedWorldToView.TransformPosition(WorldOrigin);
 		FIntPoint PageSpaceCenter(
-			FMath::RoundToInt(CenterView.X / PageSizeInWorldSpace),
-			FMath::RoundToInt(CenterView.Y / PageSizeInWorldSpace));
-		CenterView.X = PageSpaceCenter.X * PageSizeInWorldSpace;
-		CenterView.Y = PageSpaceCenter.Y * PageSizeInWorldSpace;
+			FMath::RoundToInt(ViewCenter.X / PageSizeInWorldSpace),
+			FMath::RoundToInt(ViewCenter.Y / PageSizeInWorldSpace));
+		ViewCenter.X = PageSpaceCenter.X * PageSizeInWorldSpace;
+		ViewCenter.Y = PageSpaceCenter.Y * PageSizeInWorldSpace;
 
-		const FVector SnappedWorldCenter = TranslatedWorldToView.InverseTransformPosition(CenterView);
+		const FVector SnappedWorldCenter = TranslatedWorldToView.InverseTransformPosition(ViewCenter);
 		const float SnappedLevelRadius = 0.5f * (PageSizeInWorldSpace * DimPages);
 			
 		// Sanity check that we didn't snap further away than we thought we were going to
 		{
-			float SnapError = (SnappedWorldCenter - ClipmapCenter).Size();
+			float SnapError = (SnappedWorldCenter - WorldOrigin).Size();
 			check(SnapError <= (UE_HALF_SQRT_2 * PageSizeInWorldSpace));
 			check((RawLevelRadius + SnapError) <= SnappedLevelRadius);
 		}
@@ -123,7 +123,7 @@ FVirtualShadowMapClipmap::FVirtualShadowMapClipmap(
 			{
 				FIntPoint PageOffset(PageSpaceCenter);
 				PageOffset.Y = -PageOffset.Y;		// Viewport
-				float DepthOffset = -CenterView.Z * ZScale;
+				float DepthOffset = -ViewCenter.Z * ZScale;
 
 				CacheEntry->UpdateClipmap(Level.VirtualShadowMap->ID, WorldToLight, PageOffset, DepthOffset);
 
@@ -161,12 +161,13 @@ FVirtualShadowMapProjectionShaderData FVirtualShadowMapClipmap::GetProjectionSha
 	Data.TranslatedWorldToShadowViewMatrix = TranslatedWorldToView;
 	Data.ShadowViewToClipMatrix = Level.ViewToClip;
 	Data.TranslatedWorldToShadowUvNormalMatrix = CalcTranslatedWorldToShadowUvNormalMatrix(TranslatedWorldToView, Level.ViewToClip);
-	Data.ShadowPreViewTranslation = FVector4(-Level.WorldCenter, 666.0f);
+	Data.ShadowPreViewTranslation = -Level.WorldCenter;
 	Data.VirtualShadowMapId = Level.VirtualShadowMap->ID;
+	Data.LightType = ELightComponentType::LightType_Directional;
+	Data.ClipmapWorldOrigin = WorldOrigin;
 	Data.ClipmapLevel = FirstLevel + ClipmapIndex;
 	Data.ClipmapLevelCount = LevelData.Num();
 	Data.ClipmapResolutionLodBias = ResolutionLodBias;
-	Data.bNearClip = false;
 
 	return Data;
 }
diff --git a/Engine/Source/Runtime/Renderer/Private/VirtualShadowMaps/VirtualShadowMapClipmap.h b/Engine/Source/Runtime/Renderer/Private/VirtualShadowMaps/VirtualShadowMapClipmap.h
--- a/Engine/Source/Runtime/Renderer/Private/VirtualShadowMaps/VirtualShadowMapClipmap.h
+++ b/Engine/Source/Runtime/Renderer/Private/VirtualShadowMaps/VirtualShadowMapClipmap.h
@@ -27,7 +27,7 @@ public:
 
 	FViewMatrices GetViewMatrices(int32 ClipmapIndex) const;
 
-	FVirtualShadowMap* GetVirtualShadowMap(int32 ClipmapIndex)
+	FVirtualShadowMap* GetVirtualShadowMap(int32 ClipmapIndex = 0)
 	{
 		return LevelData[ClipmapIndex].VirtualShadowMap;
 	}
@@ -48,13 +48,23 @@ public:
 		return LightSceneInfo;
 	}
 
-	FVirtualShadowMapProjectionShaderData GetProjectionShaderData(
-		int32 ClipmapIndex) const;
+	FVirtualShadowMapProjectionShaderData GetProjectionShaderData(int32 ClipmapIndex) const;
+
+	FVector GetWorldOrigin() const
+	{
+		return WorldOrigin;
+	}
+
 
 private:
 	const FLightSceneInfo& LightSceneInfo;
 
-	// A translation that is applied to world-space before transforming by one of the shadow matrices.
+	/** Origin of the clipmap in world space
+	* Usually aligns with the camera position from which it was created.
+	* Note that the centers of each of the levels can be different as they are snapped to page alignment at their respective scales
+	* */
+	FVector WorldOrigin;
+
 	FMatrix TranslatedWorldToView;
 
 	int32 FirstLevel;
diff --git a/Engine/Source/Runtime/Renderer/Private/VirtualShadowMaps/VirtualShadowMapProjection.cpp b/Engine/Source/Runtime/Renderer/Private/VirtualShadowMaps/VirtualShadowMapProjection.cpp
--- a/Engine/Source/Runtime/Renderer/Private/VirtualShadowMaps/VirtualShadowMapProjection.cpp
+++ b/Engine/Source/Runtime/Renderer/Private/VirtualShadowMaps/VirtualShadowMapProjection.cpp
@@ -38,18 +38,14 @@ static TAutoConsoleVariable<int32> CVarVirtualShadowMapDebugProjection(
 	ECVF_RenderThreadSafe
 );
 
-BEGIN_SHADER_PARAMETER_STRUCT(FVirtualShadowMapProjectionParameters, )
-	SHADER_PARAMETER_STRUCT_INCLUDE(FVirtualShadowMapCommonParameters, CommonParameters)
+BEGIN_SHADER_PARAMETER_STRUCT(FProjectionParameters, )
+	SHADER_PARAMETER_STRUCT_INCLUDE(FVirtualShadowMapProjectionParameters, ProjectionParameters)
 	SHADER_PARAMETER_STRUCT(FLightShaderParameters, Light)
 	SHADER_PARAMETER_STRUCT_REF(FSceneTexturesUniformParameters, SceneTexturesStruct)
 	SHADER_PARAMETER_STRUCT_REF(FViewUniformShaderParameters, View)
-	SHADER_PARAMETER(uint32, LightType)
-	SHADER_PARAMETER(uint32, ProjectionCount)
+	SHADER_PARAMETER(int32, VirtualShadowMapId)
 	SHADER_PARAMETER(float, ContactShadowLength)
 	SHADER_PARAMETER(int32, DebugOutputType)
-	SHADER_PARAMETER_RDG_BUFFER_SRV(StructuredBuffer<uint2>, PageTable)
-	SHADER_PARAMETER_RDG_TEXTURE(Texture2D<uint>, PhysicalPagePool)
-	SHADER_PARAMETER_RDG_BUFFER_SRV(StructuredBuffer< FVirtualShadowMapProjectionShaderData >, ShadowMapProjectionData)
 	RENDER_TARGET_BINDING_SLOTS()
 END_SHADER_PARAMETER_STRUCT()
 
@@ -61,7 +57,7 @@ class FVirtualShadowMapProjectionPS : public FGlobalShader
 	class FOutputTypeDim : SHADER_PERMUTATION_ENUM_CLASS("OUTPUT_TYPE", EVirtualShadowMapProjectionOutputType);
 	using FPermutationDomain = TShaderPermutationDomain<FOutputTypeDim>;
 
-	using FParameters = FVirtualShadowMapProjectionParameters;
+	using FParameters = FProjectionParameters;
 
 	static bool ShouldCompilePermutation(const FGlobalShaderPermutationParameters& Parameters)
 	{
@@ -104,10 +100,10 @@ IMPLEMENT_GLOBAL_SHADER(FVirtualShadowMapProjectionCompositePS, "/Engine/Private
 
 static void AddPass_RenderVirtualShadowMapProjection(
 	const FLightSceneProxy* LightProxy,
-	const TArray<FVirtualShadowMapProjectionShaderData, SceneRenderingAllocator>& ShadowMapProjectionData,
 	FRDGBuilder& GraphBuilder,
 	const FViewInfo& View, 
 	FVirtualShadowMapArray& VirtualShadowMapArray,
+	int VirtualShadowMapId,
 	const FIntRect ScissorRect,
 	EVirtualShadowMapProjectionOutputType OutputType,
 	const FRenderTargetBinding& RenderTargetBinding,
@@ -118,30 +114,17 @@ static void AddPass_RenderVirtualShadowMapProjection(
 	FGlobalShaderMap* ShaderMap = View.ShaderMap;
 	TRefCountPtr<FPooledRDGBuffer> PageTable = VirtualShadowMapArray.PageTable;
 
-	FVirtualShadowMapProjectionParameters* PassParameters = GraphBuilder.AllocParameters<FVirtualShadowMapProjectionParameters>();
-	PassParameters->CommonParameters = VirtualShadowMapArray.CommonParameters;
-	PassParameters->PageTable = GraphBuilder.CreateSRV(GraphBuilder.RegisterExternalBuffer(VirtualShadowMapArray.PageTable));
-	PassParameters->PhysicalPagePool = GraphBuilder.RegisterExternalTexture(VirtualShadowMapArray.PhysicalPagePool);
+	FProjectionParameters* PassParameters = GraphBuilder.AllocParameters<FProjectionParameters>();
+	VirtualShadowMapArray.SetProjectionParameters(GraphBuilder, PassParameters->ProjectionParameters);
 	PassParameters->SceneTexturesStruct = CreateSceneTextureUniformBufferSingleDraw(GraphBuilder.RHICmdList, ESceneTextureSetupMode::GBuffers | ESceneTextureSetupMode::SceneDepth, View.FeatureLevel);
 	PassParameters->View = View.ViewUniformBuffer;
 		
 	FLightShaderParameters LightParameters;
 	LightProxy->GetLightShaderParameters(LightParameters);
 	PassParameters->Light = LightParameters;
-	PassParameters->LightType = LightProxy->GetLightType();
+	PassParameters->VirtualShadowMapId = VirtualShadowMapId;
 	PassParameters->DebugOutputType = CVarVirtualShadowMapDebugProjection.GetValueOnRenderThread();
 	PassParameters->ContactShadowLength = CVarContactShadowLength.GetValueOnRenderThread();
-	PassParameters->ProjectionCount = ShadowMapProjectionData.Num();
-
-	FRDGBufferRef ShadowMapProjectionDataRDG = CreateStructuredBuffer(
-		GraphBuilder, 
-		TEXT("ShadowMapProjectionData"),
-		ShadowMapProjectionData.GetTypeSize(),
-		ShadowMapProjectionData.Num(),
-		ShadowMapProjectionData.GetData(),
-		ShadowMapProjectionData.Num() * ShadowMapProjectionData.GetTypeSize());
-
-	PassParameters->ShadowMapProjectionData = GraphBuilder.CreateSRV(ShadowMapProjectionDataRDG);
 
 	PassParameters->RenderTargets[0] = RenderTargetBinding;
 
@@ -177,10 +160,10 @@ static FRDGTextureRef CreateDebugOutput(FRDGBuilder& GraphBuilder, FIntPoint Ext
 
 static void RenderVirtualShadowMapProjectionCommon(
 	const FLightSceneProxy* LightProxy,
-	const TArray<FVirtualShadowMapProjectionShaderData, SceneRenderingAllocator>& ShadowMapProjectionData,
 	FRDGBuilder& GraphBuilder,
 	const FViewInfo& View, 
 	FVirtualShadowMapArray& VirtualShadowMapArray,
+	int32 VirtualShadowMapId,
 	const FIntRect ScissorRect,
 	EVirtualShadowMapProjectionOutputType OutputType,
 	const FRenderTargetBinding& RenderTargetBinding,
@@ -188,10 +171,11 @@ static void RenderVirtualShadowMapProjectionCommon(
 {
 	// Main Pass
 	AddPass_RenderVirtualShadowMapProjection(
-		LightProxy, ShadowMapProjectionData,
+		LightProxy,
 		GraphBuilder,
 		View,
 		VirtualShadowMapArray,
+		VirtualShadowMapId,
 		ScissorRect,
 		OutputType,
 		RenderTargetBinding,
@@ -205,10 +189,11 @@ static void RenderVirtualShadowMapProjectionCommon(
 		FRenderTargetBinding DebugRenderTargetBinding(DebugOutput, ERenderTargetLoadAction::EClear);
 
 		AddPass_RenderVirtualShadowMapProjection(
-			LightProxy, ShadowMapProjectionData,
+			LightProxy,
 			GraphBuilder,
 			View,
 			VirtualShadowMapArray,
+			VirtualShadowMapId,
 			ScissorRect,
 			EVirtualShadowMapProjectionOutputType::Debug,
 			DebugRenderTargetBinding,
@@ -218,21 +203,6 @@ static void RenderVirtualShadowMapProjectionCommon(
 	}
 }
 
-static void GetVirtualShadowMapProjectionData(
-	const TSharedPtr<FVirtualShadowMapClipmap>& Clipmap,
-	TArray<FVirtualShadowMapProjectionShaderData, SceneRenderingAllocator>& OutShadowMapProjectionData)
-{	
-	OutShadowMapProjectionData.Empty();
-	
-	// Expand to one projection data per clipmap level
-	int32 ClipmapLevelCount = Clipmap->GetLevelCount();
-	OutShadowMapProjectionData.AddDefaulted(ClipmapLevelCount);
-	for (int32 ClipmapLevel = 0; ClipmapLevel < ClipmapLevelCount; ++ClipmapLevel)
-	{
-		OutShadowMapProjectionData[ClipmapLevel] = Clipmap->GetProjectionShaderData(ClipmapLevel);
-	}
-}
-
 void RenderVirtualShadowMapProjectionForDenoising(
 	FProjectedShadowInfo* ShadowInfo,
 	FRDGBuilder& GraphBuilder,
@@ -241,15 +211,13 @@ void RenderVirtualShadowMapProjectionForDenoising(
 	const FIntRect ScissorRect,
 	FRDGTextureRef SignalTexture)
 {
-	TArray<FVirtualShadowMapProjectionShaderData, SceneRenderingAllocator> ShadowMapProjectionData;
-	ShadowMapProjectionData.Add(GetVirtualShadowMapProjectionShaderData(ShadowInfo));
-
 	FRenderTargetBinding RenderTargetBinding(SignalTexture, ERenderTargetLoadAction::EClear);
 	RenderVirtualShadowMapProjectionCommon(
-		ShadowInfo->GetLightSceneInfo().Proxy, ShadowMapProjectionData,
+		ShadowInfo->GetLightSceneInfo().Proxy,
 		GraphBuilder,
 		View,
 		VirtualShadowMapArray,
+		ShadowInfo->VirtualShadowMap->ID,
 		ScissorRect,
 		EVirtualShadowMapProjectionOutputType::Denoiser,
 		RenderTargetBinding,
@@ -263,23 +231,20 @@ void RenderVirtualShadowMapProjectionForDenoising(
 	FVirtualShadowMapArray& VirtualShadowMapArray,
 	const FIntRect ScissorRect,
 	FRDGTextureRef SignalTexture)
-{
-	TArray<FVirtualShadowMapProjectionShaderData, SceneRenderingAllocator> ShadowMapProjectionData;
-	GetVirtualShadowMapProjectionData(Clipmap, ShadowMapProjectionData);
-
+{	
 	FRenderTargetBinding RenderTargetBinding(SignalTexture, ERenderTargetLoadAction::EClear);
 	RenderVirtualShadowMapProjectionCommon(
-		Clipmap->GetLightSceneInfo().Proxy, ShadowMapProjectionData,
+		Clipmap->GetLightSceneInfo().Proxy,
 		GraphBuilder,
 		View,
 		VirtualShadowMapArray,
+		Clipmap->GetVirtualShadowMap()->ID,
 		ScissorRect,
 		EVirtualShadowMapProjectionOutputType::Denoiser,
 		RenderTargetBinding,
 		nullptr);
 }
 
-
 void RenderVirtualShadowMapProjection(
 	FProjectedShadowInfo* ShadowInfo,
 	FRDGBuilder& GraphBuilder,
@@ -289,17 +254,14 @@ void RenderVirtualShadowMapProjection(
 	FRDGTextureRef ScreenShadowMaskTexture,
 	bool bProjectingForForwardShading)
 {
-	TArray<FVirtualShadowMapProjectionShaderData, SceneRenderingAllocator> ShadowMapProjectionData;
-	ShadowMapProjectionData.Add(GetVirtualShadowMapProjectionShaderData(ShadowInfo));
-
 	FRenderTargetBinding RenderTargetBinding(ScreenShadowMaskTexture, ERenderTargetLoadAction::ELoad);
-	// Blend state should not vary per cascade
 	FRHIBlendState* BlendState = ShadowInfo->GetBlendStateForProjection(bProjectingForForwardShading, false);
 	RenderVirtualShadowMapProjectionCommon(
-		ShadowInfo->GetLightSceneInfo().Proxy, ShadowMapProjectionData,
+		ShadowInfo->GetLightSceneInfo().Proxy,
 		GraphBuilder,
 		View,
 		VirtualShadowMapArray,
+		ShadowInfo->VirtualShadowMap->ID,
 		ScissorRect,
 		EVirtualShadowMapProjectionOutputType::ScreenShadowMask,
 		RenderTargetBinding,
@@ -315,19 +277,15 @@ void RenderVirtualShadowMapProjection(
 	FRDGTextureRef ScreenShadowMaskTexture,
 	bool bProjectingForForwardShading)
 {
-	TArray<FVirtualShadowMapProjectionShaderData, SceneRenderingAllocator> ShadowMapProjectionData;
-	GetVirtualShadowMapProjectionData(Clipmap, ShadowMapProjectionData);
-
-	FRenderTargetBinding RenderTargetBinding(ScreenShadowMaskTexture, ERenderTargetLoadAction::ELoad);
-	
+	FRenderTargetBinding RenderTargetBinding(ScreenShadowMaskTexture, ERenderTargetLoadAction::ELoad);	
 	// See FProjectedShadowInfo::GetBlendStateForProjection. TODO: Support other modes in this path?
 	FRHIBlendState* BlendState = TStaticBlendState<CW_BA, BO_Min, BF_One, BF_One, BO_Min, BF_One, BF_One>::GetRHI();
-
 	RenderVirtualShadowMapProjectionCommon(
-		Clipmap->GetLightSceneInfo().Proxy, ShadowMapProjectionData,
+		Clipmap->GetLightSceneInfo().Proxy,
 		GraphBuilder,
 		View,
 		VirtualShadowMapArray,
+		Clipmap->GetVirtualShadowMap()->ID,
 		ScissorRect,
 		EVirtualShadowMapProjectionOutputType::ScreenShadowMask,
 		RenderTargetBinding,
