Rename virtual shadow map sampling parameters structure.
Fix issues with contact shadow offset in virtual shadow map shader
- Need to ensure we only offset the  depth test, not the clipmap level selection or occluder distance computation

diff --git a/Engine/Shaders/Private/VirtualShadowMaps/ProjectionCommon.ush b/Engine/Shaders/Private/VirtualShadowMaps/ProjectionCommon.ush
--- a/Engine/Shaders/Private/VirtualShadowMaps/ProjectionCommon.ush
+++ b/Engine/Shaders/Private/VirtualShadowMaps/ProjectionCommon.ush
@@ -199,9 +199,15 @@ float ComputeOccluderDistancePerspective(float4x4 ShadowViewToClip, float Occlud
 * EstimatedGeoWorldNormal is ideally the geometric (flat) normal of the sample point in world space
 *   - This is used to compute a receiver-plane-based bias for the sample point
 *   - The shading normal can be used if the geometric normal is not available, but divergence from the geometric normal can cause biasing issues.
+* RayStartDistance is an optional offset to move the lookup along the shadow ray towards the light
+*   - Should zero or positive
+*   - This offset is useful in that it does *not* affect the selection of clipmap level, unlike offsetting the WorldPosition itself
+*   - OccluderDistance will still be relative to the original sample position
 */
-FVirtualShadowMapSampleResult SampleVirtualShadowMap(int VirtualShadowMapId, float3 WorldPosition, float3 EstimatedGeoWorldNormal)
+FVirtualShadowMapSampleResult SampleVirtualShadowMap(int VirtualShadowMapId, float3 WorldPosition, float3 EstimatedGeoWorldNormal, float RayStartDistance = 0.0f)
 {
+	RayStartDistance= max(RayStartDistance, 0.0f);
+
 	FVirtualShadowMapProjectionShaderData BaseProjectionData = GetVirtualShadowMapProjectionData(VirtualShadowMapId);
 
 	FVirtualSmSample SmSample;
@@ -211,9 +217,7 @@ FVirtualShadowMapSampleResult SampleVirtualShadowMap(int VirtualShadowMapId, flo
 	SmSample.MipLevel = 0;
 	SmSample.bValid = false;
 
-	FShadowPosition Position;
-	float OccluderDistance = 0.0f;
-	float BiasedDepth = 0.0f;
+	float OccluderDistance = -1.0f;
 	bool bOccluded = false;
 
 	if (BaseProjectionData.LightType == LIGHT_TYPE_DIRECTIONAL)
@@ -230,18 +234,20 @@ FVirtualShadowMapSampleResult SampleVirtualShadowMap(int VirtualShadowMapId, flo
 			int ClipmapLevelVirtualShadowMapId = VirtualShadowMapId + ClipmapIndex;
 			FVirtualShadowMapProjectionShaderData ProjectionData = GetVirtualShadowMapProjectionData(ClipmapLevelVirtualShadowMapId);
 
-			Position = WorldToShadow(ProjectionData, WorldPosition);
+			FShadowPosition Position = WorldToShadow(ProjectionData, WorldPosition);
 			SmSample = SampleVirtualSmLevel(ClipmapLevelVirtualShadowMapId, Position.UV, 0);
 			if (SmSample.bValid)
 			{
-				BiasedDepth = SmSample.Depth - ComputeVirtualShadowMapOptimalSlopeBias(ProjectionData, SmSample, WorldPosition, EstimatedGeoWorldNormal);
+				float OptimalSlopeBias = ComputeVirtualShadowMapOptimalSlopeBias(ProjectionData, SmSample, WorldPosition, EstimatedGeoWorldNormal);
+				float RayStartBias = -RayStartDistance * ProjectionData.ShadowViewToClipMatrix._33;
+				float BiasedDepth = SmSample.Depth - OptimalSlopeBias - RayStartBias;
 
 				if (BiasedDepth > Position.Depth)
 				{
 					bOccluded = true;
 					OccluderDistance = ComputeOccluderDistanceOrtho(
 						ProjectionData.ShadowViewToClipMatrix,
-						BiasedDepth,
+						SmSample.Depth,
 						Position.Depth);
 				}
 			}
@@ -249,18 +255,20 @@ FVirtualShadowMapSampleResult SampleVirtualShadowMap(int VirtualShadowMapId, flo
 	}
 	else	// Spot light
 	{
-		Position = WorldToShadow(BaseProjectionData, WorldPosition);
+		FShadowPosition Position = WorldToShadow(BaseProjectionData, WorldPosition);
 		SmSample = SampleVirtualSm(VirtualShadowMapId, Position.UV);
 		if (SmSample.bValid)
 		{
-			BiasedDepth = SmSample.Depth - ComputeVirtualShadowMapOptimalSlopeBias(BaseProjectionData, SmSample, WorldPosition, EstimatedGeoWorldNormal);
+			float OptimalSlopeBias = ComputeVirtualShadowMapOptimalSlopeBias(BaseProjectionData, SmSample, WorldPosition, EstimatedGeoWorldNormal);
+			float RayStartBias = -RayStartDistance * BaseProjectionData.ShadowViewToClipMatrix._33 / Position.ClipH.w;
+			float BiasedDepth = SmSample.Depth - OptimalSlopeBias - RayStartBias;
 
 			if (BiasedDepth > Position.Depth)
 			{
 				bOccluded = true;
-				OccluderDistance = ComputeOccluderDistancePerspective(
+				OccluderDistance = RayStartDistance + ComputeOccluderDistancePerspective(
 					BaseProjectionData.ShadowViewToClipMatrix,
-					BiasedDepth,
+					SmSample.Depth,
 					Position.Depth,
 					Position.View);
 			}
diff --git a/Engine/Shaders/Private/VirtualShadowMaps/VirtualShadowMapProjection.usf b/Engine/Shaders/Private/VirtualShadowMaps/VirtualShadowMapProjection.usf
--- a/Engine/Shaders/Private/VirtualShadowMaps/VirtualShadowMapProjection.usf
+++ b/Engine/Shaders/Private/VirtualShadowMaps/VirtualShadowMapProjection.usf
@@ -150,9 +150,6 @@ void VirtualShadowMapProjectionPS(
 		ContactShadowLengthWorld = ContactShadowLength * View.ClipToView[1][1] * SceneDepth;
 		float StepOffset = InterleavedGradientNoise( SvPosition.xy, View.StateFrameIndexMod8 );
 		ContactShadow = ShadowRayCast( TranslatedWorldPosition, ToLight, ContactShadowLengthWorld, 12, StepOffset );
-	
-		// Move shadow map lookup to end of screen space ray.
-		TranslatedWorldPosition += ToLight * ContactShadowLengthWorld;
 	}
 
 	// NOTE: Could skip most of the rest of the shader if contact shadows say we're shadowed, but likely to be small, incoherent branching
@@ -171,23 +168,25 @@ void VirtualShadowMapProjectionPS(
 	}
 #endif
 
-	float3 WorldPositon = TranslatedWorldPosition - View.PreViewTranslation;
-	FVirtualShadowMapSampleResult VirtualShadowMapSample = SampleVirtualShadowMap(VirtualShadowMapId, WorldPositon, EstimatedGeoWorldNormal);
-
+	float3 WorldPosition = TranslatedWorldPosition - View.PreViewTranslation;
+	// Offset ray to start at the end of the contact shadow trace
+	float RayStartDistance = ContactShadowLengthWorld;
+	FVirtualShadowMapSampleResult VirtualShadowMapSample = SampleVirtualShadowMap(VirtualShadowMapId, WorldPosition, EstimatedGeoWorldNormal, RayStartDistance);
+	
+#if OUTPUT_TYPE != 2		// Don't kill pixels in debug output mode
 	if (!VirtualShadowMapSample.bValid)
 	{
 		// Generally shouldn't occur if PT allocation is conservative
 		discard;
 	}
+#endif
 
 	// See signal encoding for SIGNAL_BUFFER_LAYOUT_PENUMBRA_RECONSTRUCTION in SSDSignalBufferEncoding.ush
 	float4 OutputSignal = float4(0, 0, 0, 0);
 	if (VirtualShadowMapSample.bOccluded)
 	{
 		// Hit
-		// Recall we offset the initial ray, so add the distance back in
-		float OccluderDistance = VirtualShadowMapSample.OccluderDistance + ContactShadowLengthWorld;
-		float WorldBluringRadius = ComputeLightSampleWorldBluringRadius(TranslatedWorldPosition, LightType, Light, OccluderDistance);
+		float WorldBluringRadius = ComputeLightSampleWorldBluringRadius(TranslatedWorldPosition, LightType, Light, VirtualShadowMapSample.OccluderDistance);
 		OutputSignal = float4(0, 1, WorldBluringRadius, 0);
 	}
 	else
@@ -216,12 +215,6 @@ void VirtualShadowMapProjectionPS(
 	{
 		Color = VirtualShadowMapSample.bValid ? (0.25f * OutputSignal.rrr) : float3(1, 0, 1);
 	}
-	else if (DebugOutputType == 2)		// Clipmap level
-	{
-		//Color = lerp(IntToColor(ProjectionDataIndex), OutputSignal.rrr, 0.4f);
-	}
-
-	//Color = OptimalSlopeBias.xxx * 1.0f;
 
 	OutputTarget = float4(Color, 1.0f);	
 #endif // OUTPUT_TYPE
diff --git a/Engine/Source/Runtime/Renderer/Private/VirtualShadowMaps/VirtualShadowMapArray.cpp b/Engine/Source/Runtime/Renderer/Private/VirtualShadowMaps/VirtualShadowMapArray.cpp
--- a/Engine/Source/Runtime/Renderer/Private/VirtualShadowMaps/VirtualShadowMapArray.cpp
+++ b/Engine/Source/Runtime/Renderer/Private/VirtualShadowMaps/VirtualShadowMapArray.cpp
@@ -871,7 +871,7 @@ void FVirtualShadowMapArray::GeneratePageFlagsFromLightGrid(FRDGBuilder& GraphBu
 	}
 }
 
-void FVirtualShadowMapArray::SetProjectionParameters(FRDGBuilder& GraphBuilder, FVirtualShadowMapProjectionParameters& OutParameters)
+void FVirtualShadowMapArray::SetProjectionParameters(FRDGBuilder& GraphBuilder, FVirtualShadowMapSamplingParameters& OutParameters)
 {
 	OutParameters.CommonParameters = CommonParameters;
 	OutParameters.PageTable = GraphBuilder.CreateSRV(GraphBuilder.RegisterExternalBuffer(PageTable));
diff --git a/Engine/Source/Runtime/Renderer/Private/VirtualShadowMaps/VirtualShadowMapArray.h b/Engine/Source/Runtime/Renderer/Private/VirtualShadowMaps/VirtualShadowMapArray.h
--- a/Engine/Source/Runtime/Renderer/Private/VirtualShadowMaps/VirtualShadowMapArray.h
+++ b/Engine/Source/Runtime/Renderer/Private/VirtualShadowMaps/VirtualShadowMapArray.h
@@ -119,7 +119,7 @@ BEGIN_SHADER_PARAMETER_STRUCT(FVirtualShadowMapCommonParameters, )
 	SHADER_PARAMETER(uint32, PhysicalPageRowShift)
 END_SHADER_PARAMETER_STRUCT()
 
-BEGIN_SHADER_PARAMETER_STRUCT(FVirtualShadowMapProjectionParameters, )
+BEGIN_SHADER_PARAMETER_STRUCT(FVirtualShadowMapSamplingParameters, )
 	SHADER_PARAMETER_STRUCT_INCLUDE(FVirtualShadowMapCommonParameters, CommonParameters)
 	SHADER_PARAMETER_RDG_BUFFER_SRV(StructuredBuffer<uint2>, PageTable)
 	SHADER_PARAMETER_RDG_TEXTURE(Texture2D<uint>, PhysicalPagePool)
@@ -150,7 +150,7 @@ public:
 
 	static void SetShaderDefines(FShaderCompilerEnvironment& OutEnvironment);
 
-	void SetProjectionParameters(FRDGBuilder& GraphBuilder, FVirtualShadowMapProjectionParameters& OutParameters);
+	void SetProjectionParameters(FRDGBuilder& GraphBuilder, FVirtualShadowMapSamplingParameters& OutParameters);
 
 	//
 	void GenerateIdentityPageTables(FRDGBuilder& GraphBuilder, uint32 MipLevel = INDEX_NONE);
diff --git a/Engine/Source/Runtime/Renderer/Private/VirtualShadowMaps/VirtualShadowMapProjection.cpp b/Engine/Source/Runtime/Renderer/Private/VirtualShadowMaps/VirtualShadowMapProjection.cpp
--- a/Engine/Source/Runtime/Renderer/Private/VirtualShadowMaps/VirtualShadowMapProjection.cpp
+++ b/Engine/Source/Runtime/Renderer/Private/VirtualShadowMaps/VirtualShadowMapProjection.cpp
@@ -39,7 +39,7 @@ static TAutoConsoleVariable<int32> CVarVirtualShadowMapDebugProjection(
 );
 
 BEGIN_SHADER_PARAMETER_STRUCT(FProjectionParameters, )
-	SHADER_PARAMETER_STRUCT_INCLUDE(FVirtualShadowMapProjectionParameters, ProjectionParameters)
+	SHADER_PARAMETER_STRUCT_INCLUDE(FVirtualShadowMapSamplingParameters, ProjectionParameters)
 	SHADER_PARAMETER_STRUCT(FLightShaderParameters, Light)
 	SHADER_PARAMETER_STRUCT_REF(FSceneTexturesUniformParameters, SceneTexturesStruct)
 	SHADER_PARAMETER_STRUCT_REF(FViewUniformShaderParameters, View)
