Removed a number of LOOP and UNROLL hints that were causing FXC to spend 20 minutes per DoF permutation, just to fail when the simulator gives up trying to forcibly unroll loops it had no business touching... The codegen shoud be looked at later, but in theory it shouldn't really be affected much (FXC wasn't compiling, and DXC generally ignores a lot of the force unroll hints). Even trying UNROLL_N(5), etc.. failed in this case, and other cases with explicit LOOP were failing because the compiler "can't unroll loops marked with loop attribute"... sigh

FXC needs to disappear forever

diff --git a/Engine/Shaders/Private/DiaphragmDOF/DOFGatherKernel.ush b/Engine/Shaders/Private/DiaphragmDOF/DOFGatherKernel.ush
--- a/Engine/Shaders/Private/DiaphragmDOF/DOFGatherKernel.ush
+++ b/Engine/Shaders/Private/DiaphragmDOF/DOFGatherKernel.ush
@@ -390,7 +390,6 @@ void GatherRingSamples(
 	#endif
 
 	// Iterate over the batch of samples.
-	LOOP
 	for (uint BatchId = 0; BatchId < BatchCount; BatchId++)
 	#if defined(WORKARROUND_UE_58394)
 		WORKARROUND_UE_58394 // TODO: nuke that guy once original issue is fixed.
@@ -644,7 +643,6 @@ uint GatherSamplesAndResolve(
 	if (static_condition(GatherParameters.KernelSamplingDensityMode == KERNEL_DENSITY_CONSTANT) ||
 		static_condition(GatherParameters.KernelSamplingDensityMode == KERNEL_DENSITY_CONSTANT_HEXAWEB))
 	{
-		UNROLL
 		#if LARGEST_RINGS_FIRST
 			for (int RingId = GatherParameters.MaxRingCount - 1; RingId >= 0; RingId--)
 		#else
@@ -666,7 +664,6 @@ uint GatherSamplesAndResolve(
 	}
 	else if (static_condition(GatherParameters.KernelSamplingDensityMode == KERNEL_DENSITY_CONSTANT_DYNAMIC_RING_COUNT))
 	{
-		LOOP
 		#if LARGEST_RINGS_FIRST
 			for (int RingId = GatherParameters.MaxRingCount - 1; RingId >= 0; RingId--)
 		#else
@@ -694,7 +691,7 @@ uint GatherSamplesAndResolve(
 
 		int RingId;
 
-		UNROLL
+		UNROLL_N(5)
 		for (RingId = GatherParameters.MaxRingCount - 1; RingId >= LowerRingId; RingId--)
 		{
 			GatherRing(
@@ -708,7 +705,7 @@ uint GatherSamplesAndResolve(
 			bIsFirstRing = false;
 		} // for (int RingId)
 		
-		UNROLL // TODO: LOOP instead? should be totally fine perf wise...
+		LOOP
 		for (uint SamplingDensityChangeId = 0; SamplingDensityChangeId < GatherParameters.MaxDensityChangeCount; SamplingDensityChangeId++)
 		{
 			float MipTexelSize = pow(2.0, uint(KernelScope.GatherInputMipLevel));
@@ -741,7 +738,7 @@ uint GatherSamplesAndResolve(
 					UpdateRuntimeKernelScope(GatherParameters, /* inout */ KernelScope);
 				}
 
-				UNROLL
+				UNROLL_N(5)
 				for (int RingId2 = GatherParameters.MaxRingCount - 1; RingId2 >= LowerRingId; RingId2--)
 				{
 					GatherRing(
@@ -814,7 +811,7 @@ uint GatherSamplesAndResolve(
 			RingId = RemainingRingCount - 1;
 		}
 
-		UNROLL
+		UNROLL_N(5)
 		for (; RingId >= 0; RingId--)
 		{
 			GatherRing(
@@ -858,7 +855,6 @@ uint GatherSamplesAndResolve(
 				/* inout */ FastAccumulator);
 		}
 
-		LOOP
 		for (uint SamplingDensityChangeId = 0; SamplingDensityChangeId < GatherParameters.MaxDensityChangeCount; SamplingDensityChangeId++)
 		{
 			// Gather all remaining rings as single cluster of sample to the accumulator.
@@ -884,7 +880,6 @@ uint GatherSamplesAndResolve(
 				// Start accumulating ring samples
 				StartRingSamples(Accumulator, RingInfos);
 	
-				UNROLL // Perhaps loop is ok?
 				for (uint RingId = 0; RingId < RemainingRingCount; RingId++)
 				{
 					GatherRingSamples(
