Fix Gather4 offsets in shadow map filtering (off-by half error)

diff --git a/Engine/Shaders/Private/ShadowFilteringCommon.ush b/Engine/Shaders/Private/ShadowFilteringCommon.ush
--- a/Engine/Shaders/Private/ShadowFilteringCommon.ush
+++ b/Engine/Shaders/Private/ShadowFilteringCommon.ush
@@ -227,29 +227,32 @@ float3 CalculateOcclusion(float3 ShadowmapDepth, FPCFSamplerSettings Settings)
 	}
 }
 
-void FetchRowOfThree(float2 Sample00TexelCenter, float VerticalOffset, out float3 Values0, FPCFSamplerSettings Settings)
+float3 FetchRowOfThree(float2 Sample00TexelCenter, float VerticalOffset, FPCFSamplerSettings Settings)
 {
-	Values0.x = Texture2DSampleLevel(Settings.ShadowDepthTexture, Settings.ShadowDepthTextureSampler, (Sample00TexelCenter + float2(0, VerticalOffset)) * Settings.ShadowBufferSize.zw, 0).r;
-	Values0.y = Texture2DSampleLevel(Settings.ShadowDepthTexture, Settings.ShadowDepthTextureSampler, (Sample00TexelCenter + float2(1, VerticalOffset)) * Settings.ShadowBufferSize.zw, 0).r;
-	Values0.z = Texture2DSampleLevel(Settings.ShadowDepthTexture, Settings.ShadowDepthTextureSampler, (Sample00TexelCenter + float2(2, VerticalOffset)) * Settings.ShadowBufferSize.zw, 0).r;	
-	Values0 = CalculateOcclusion(Values0, Settings);
+	float3 Values;
+	Values.x = Texture2DSampleLevel(Settings.ShadowDepthTexture, Settings.ShadowDepthTextureSampler, (Sample00TexelCenter + float2(0, VerticalOffset)) * Settings.ShadowBufferSize.zw, 0).r;
+	Values.y = Texture2DSampleLevel(Settings.ShadowDepthTexture, Settings.ShadowDepthTextureSampler, (Sample00TexelCenter + float2(1, VerticalOffset)) * Settings.ShadowBufferSize.zw, 0).r;
+	Values.z = Texture2DSampleLevel(Settings.ShadowDepthTexture, Settings.ShadowDepthTextureSampler, (Sample00TexelCenter + float2(2, VerticalOffset)) * Settings.ShadowBufferSize.zw, 0).r;	
+	return CalculateOcclusion(Values, Settings);
 }
 
-void FetchRowOfFour(float2 Sample00TexelCenter, float VerticalOffset, out float4 Values0, FPCFSamplerSettings Settings)
+float4 FetchRowOfFour(float2 Sample00TexelCenter, float VerticalOffset, FPCFSamplerSettings Settings)
 {
-	Values0.x = Texture2DSampleLevel(Settings.ShadowDepthTexture, Settings.ShadowDepthTextureSampler, (Sample00TexelCenter + float2(0, VerticalOffset)) * Settings.ShadowBufferSize.zw, 0).r;
-	Values0.y = Texture2DSampleLevel(Settings.ShadowDepthTexture, Settings.ShadowDepthTextureSampler, (Sample00TexelCenter + float2(1, VerticalOffset)) * Settings.ShadowBufferSize.zw, 0).r;
-	Values0.z = Texture2DSampleLevel(Settings.ShadowDepthTexture, Settings.ShadowDepthTextureSampler, (Sample00TexelCenter + float2(2, VerticalOffset)) * Settings.ShadowBufferSize.zw, 0).r;
-	Values0.w = Texture2DSampleLevel(Settings.ShadowDepthTexture, Settings.ShadowDepthTextureSampler, (Sample00TexelCenter + float2(3, VerticalOffset)) * Settings.ShadowBufferSize.zw, 0).r;
-	Values0 = CalculateOcclusion(Values0, Settings);
+	float4 Values;
+	Values.x = Texture2DSampleLevel(Settings.ShadowDepthTexture, Settings.ShadowDepthTextureSampler, (Sample00TexelCenter + float2(0, VerticalOffset)) * Settings.ShadowBufferSize.zw, 0).r;
+	Values.y = Texture2DSampleLevel(Settings.ShadowDepthTexture, Settings.ShadowDepthTextureSampler, (Sample00TexelCenter + float2(1, VerticalOffset)) * Settings.ShadowBufferSize.zw, 0).r;
+	Values.z = Texture2DSampleLevel(Settings.ShadowDepthTexture, Settings.ShadowDepthTextureSampler, (Sample00TexelCenter + float2(2, VerticalOffset)) * Settings.ShadowBufferSize.zw, 0).r;
+	Values.w = Texture2DSampleLevel(Settings.ShadowDepthTexture, Settings.ShadowDepthTextureSampler, (Sample00TexelCenter + float2(3, VerticalOffset)) * Settings.ShadowBufferSize.zw, 0).r;
+	return CalculateOcclusion(Values, Settings);
 }
 
-void FetchRowOfThreeAfterFour(float2 Sample00TexelCenter, float VerticalOffset, out float3 Values1, FPCFSamplerSettings Settings)
+float3 FetchRowOfThreeAfterFour(float2 Sample00TexelCenter, float VerticalOffset, FPCFSamplerSettings Settings)
 {
-	Values1.x = Texture2DSampleLevel(Settings.ShadowDepthTexture, Settings.ShadowDepthTextureSampler, (Sample00TexelCenter + float2(4, VerticalOffset)) * Settings.ShadowBufferSize.zw, 0).r;
-	Values1.y = Texture2DSampleLevel(Settings.ShadowDepthTexture, Settings.ShadowDepthTextureSampler, (Sample00TexelCenter + float2(5, VerticalOffset)) * Settings.ShadowBufferSize.zw, 0).r;
-	Values1.z = Texture2DSampleLevel(Settings.ShadowDepthTexture, Settings.ShadowDepthTextureSampler, (Sample00TexelCenter + float2(6, VerticalOffset)) * Settings.ShadowBufferSize.zw, 0).r;
-	Values1 = CalculateOcclusion(Values1, Settings);
+	float3 Values;
+	Values.x = Texture2DSampleLevel(Settings.ShadowDepthTexture, Settings.ShadowDepthTextureSampler, (Sample00TexelCenter + float2(4, VerticalOffset)) * Settings.ShadowBufferSize.zw, 0).r;
+	Values.y = Texture2DSampleLevel(Settings.ShadowDepthTexture, Settings.ShadowDepthTextureSampler, (Sample00TexelCenter + float2(5, VerticalOffset)) * Settings.ShadowBufferSize.zw, 0).r;
+	Values.z = Texture2DSampleLevel(Settings.ShadowDepthTexture, Settings.ShadowDepthTextureSampler, (Sample00TexelCenter + float2(6, VerticalOffset)) * Settings.ShadowBufferSize.zw, 0).r;
+	return CalculateOcclusion(Values, Settings);
 }
 
 float Manual3x3PCF(float2 ShadowPosition, FPCFSamplerSettings Settings)
@@ -258,24 +261,25 @@ float Manual3x3PCF(float2 ShadowPosition, FPCFSamplerSettings Settings)
 	{
 		float2 TexelPos = ShadowPosition * Settings.ShadowBufferSize.xy - 0.5f;	// bias to be consistent with texture filtering hardware
 		float2 Fraction = frac(TexelPos);
-		float2 TexelCenter = floor(TexelPos) + 0.5f;	// bias to get reliable texel center content
+		float2 TexelCorner = floor(TexelPos);	// bias to get reliable texel center content
 		{
-			float2 Sample00TexelCenter = TexelCenter - float2(1, 1);
 
-			float4 SampleValues0, SampleValues1, SampleValues2, SampleValues3;
 
 #if FEATURE_GATHER4
-			float2 SamplePos = TexelCenter * Settings.ShadowBufferSize.zw;	// bias to get reliable texel center content
-			SampleValues0 = CalculateOcclusion(Settings.ShadowDepthTexture.Gather(Settings.ShadowDepthTextureSampler, SamplePos, int2(-1, -1)), Settings);
-			SampleValues1 = CalculateOcclusion(Settings.ShadowDepthTexture.Gather(Settings.ShadowDepthTextureSampler, SamplePos, int2(1, -1)), Settings);
-			SampleValues2 = CalculateOcclusion(Settings.ShadowDepthTexture.Gather(Settings.ShadowDepthTextureSampler, SamplePos, int2(-1, 1)), Settings);
-			SampleValues3 = CalculateOcclusion(Settings.ShadowDepthTexture.Gather(Settings.ShadowDepthTextureSampler, SamplePos, int2(1, 1)), Settings);
+			// Gather4 samples "at the following locations: (-,+),(+,+),(+,-),(-,-), where the magnitude of the deltas are always half a texel" - DX11 Func. Spec.
+			// So we need to offset to the centre of the 2x2 grid we want to sample.
+			float2 SamplePos = (TexelCorner + 1.0f) * Settings.ShadowBufferSize.zw;	// bias to get reliable texel center content
+			float4 SampleValues0 = CalculateOcclusion(Settings.ShadowDepthTexture.Gather(Settings.ShadowDepthTextureSampler, SamplePos, int2(-1, -1)), Settings);
+			float4 SampleValues1 = CalculateOcclusion(Settings.ShadowDepthTexture.Gather(Settings.ShadowDepthTextureSampler, SamplePos, int2(1, -1)), Settings);
+			float4 SampleValues2 = CalculateOcclusion(Settings.ShadowDepthTexture.Gather(Settings.ShadowDepthTextureSampler, SamplePos, int2(-1, 1)), Settings);
+			float4 SampleValues3 = CalculateOcclusion(Settings.ShadowDepthTexture.Gather(Settings.ShadowDepthTextureSampler, SamplePos, int2(1, 1)), Settings);
 			return PCF3x3gather(Fraction, SampleValues0, SampleValues1, SampleValues2, SampleValues3);
 #else // FEATURE_GATHER4 
-			FetchRowOfFour(Sample00TexelCenter, 0, SampleValues0, Settings);
-			FetchRowOfFour(Sample00TexelCenter, 1, SampleValues1, Settings);
-			FetchRowOfFour(Sample00TexelCenter, 2, SampleValues2, Settings);
-			FetchRowOfFour(Sample00TexelCenter, 3, SampleValues3, Settings);
+			float2 Sample00TexelCenter = TexelCorner - 0.5f;
+			float4 SampleValues0 = FetchRowOfFour(Sample00TexelCenter, 0, Settings);
+			float4 SampleValues1 = FetchRowOfFour(Sample00TexelCenter, 1, Settings);
+			float4 SampleValues2 = FetchRowOfFour(Sample00TexelCenter, 2, Settings);
+			float4 SampleValues3 = FetchRowOfFour(Sample00TexelCenter, 3, Settings);
 			return PCF3x3(Fraction, SampleValues0, SampleValues1, SampleValues2, SampleValues3);
 #endif // FEATURE_GATHER4
 		}
@@ -291,11 +295,9 @@ float Manual2x2PCF(float2 ShadowPosition, FPCFSamplerSettings Settings)
 	
 	float2 Sample00TexelCenter = TexelCenter - float2(1, 1);
 
-	float3 SamplesValues0, SamplesValues1, SamplesValues2;
-
-	FetchRowOfThree(Sample00TexelCenter, 0, SamplesValues0, Settings);
-	FetchRowOfThree(Sample00TexelCenter, 1, SamplesValues1, Settings);
-	FetchRowOfThree(Sample00TexelCenter, 2, SamplesValues2, Settings);
+	float3 SamplesValues0 = FetchRowOfThree(Sample00TexelCenter, 0, Settings);
+	float3 SamplesValues1 = FetchRowOfThree(Sample00TexelCenter, 1, Settings);
+	float3 SamplesValues2 = FetchRowOfThree(Sample00TexelCenter, 2, Settings);
 	
 	return PCF2x2(Fraction, SamplesValues0, SamplesValues1, SamplesValues2);
 }
@@ -309,11 +311,9 @@ float Manuax2PCF(float2 ShadowPosition, FPCFSamplerSettings Settings)
 	
 	float2 Sample00TexelCenter = TexelCenter - float2(1, 1);
 
-	float3 SamplesValues0, SamplesValues1, SamplesValues2;
-
-	FetchRowOfThree(Sample00TexelCenter, 0, SamplesValues0, Settings);
-	FetchRowOfThree(Sample00TexelCenter, 1, SamplesValues1, Settings);
-	FetchRowOfThree(Sample00TexelCenter, 2, SamplesValues2, Settings);
+	float3 SamplesValues0 = FetchRowOfThree(Sample00TexelCenter, 0, Settings);
+	float3 SamplesValues1 = FetchRowOfThree(Sample00TexelCenter, 1, Settings);
+	float3 SamplesValues2 = FetchRowOfThree(Sample00TexelCenter, 2, Settings);
 	
 	return PCF2x2(Fraction, SamplesValues0, SamplesValues1, SamplesValues2);
 }
@@ -328,14 +328,17 @@ float Manual1x1PCF(float2 ShadowPosition, FPCFSamplerSettings Settings)
 {
 	float2 TexelPos = ShadowPosition * Settings.ShadowBufferSize.xy - 0.5f;	// bias to be consistent with texture filtering hardware
 	float2 Fraction = frac(TexelPos);
-	float2 TexelCenter = floor(TexelPos) + 0.5f;	// bias to get reliable texel center content
 
 	// using Gather: xyzw in counter clockwise order starting with the sample to the lower left of the queried location
 	float4 Samples;
 
 #if FEATURE_GATHER4
-	Samples = Settings.ShadowDepthTexture.Gather(Settings.ShadowDepthTextureSampler, TexelCenter * Settings.ShadowBufferSize.zw);
+	// Gather4 samples "at the following locations: (-,+),(+,+),(+,-),(-,-), where the magnitude of the deltas are always half a texel" - DX11 Func. Spec.
+	// So we need to offset to the centre of the 2x2 grid we want to sample.
+	float2 QuadCenter = floor(TexelPos) + 1.0f;	
+	Samples = Settings.ShadowDepthTexture.Gather(Settings.ShadowDepthTextureSampler, QuadCenter * Settings.ShadowBufferSize.zw);
 #else
+	float2 TexelCenter = floor(TexelPos) + 0.5f;	// bias to get reliable texel center content
 	Samples.x = Texture2DSampleLevel(Settings.ShadowDepthTexture, Settings.ShadowDepthTextureSampler, (TexelCenter.xy + float2(0, 1)) * Settings.ShadowBufferSize.zw, 0).r;
 	Samples.y = Texture2DSampleLevel(Settings.ShadowDepthTexture, Settings.ShadowDepthTextureSampler, (TexelCenter.xy + float2(1, 1)) * Settings.ShadowBufferSize.zw, 0).r;
 	Samples.z = Texture2DSampleLevel(Settings.ShadowDepthTexture, Settings.ShadowDepthTextureSampler, (TexelCenter.xy + float2(1, 0)) * Settings.ShadowBufferSize.zw, 0).r;
@@ -361,17 +364,16 @@ float ManualPCF(float2 ShadowPosition, FPCFSamplerSettings Settings)
 	{
 		float2 TexelPos = ShadowPosition * Settings.ShadowBufferSize.xy - 0.5f;	// bias to be consistent with texture filtering hardware
 		float2 Fraction = frac(TexelPos);
-		float2 TexelCenter = floor(TexelPos);
-		float2 SamplePos = (TexelCenter + 0.5f) * Settings.ShadowBufferSize.zw;	// bias to get reliable texel center content
-
-		float Results;
+		// Gather4 samples "at the following locations: (-,+),(+,+),(+,-),(-,-), where the magnitude of the deltas are always half a texel" - DX11 Func. Spec.
+		// So we need to offset to the centre of the 2x2 grid we want to sample.
+		float2 SamplePos = (floor(TexelPos) + 1.0f) * Settings.ShadowBufferSize.zw;	// bias to get reliable texel center content
 
 		float4 Values00 = CalculateOcclusion(Settings.ShadowDepthTexture.Gather(Settings.ShadowDepthTextureSampler, SamplePos, int2(-2, -2)), Settings);
 		float4 Values20 = CalculateOcclusion(Settings.ShadowDepthTexture.Gather(Settings.ShadowDepthTextureSampler, SamplePos, int2(0, -2)), Settings);
 		float4 Values40 = CalculateOcclusion(Settings.ShadowDepthTexture.Gather(Settings.ShadowDepthTextureSampler, SamplePos, int2(2, -2)), Settings);
 
 		float2 Row0 = HorizontalPCF5x2(Fraction, Values00, Values20, Values40);
-		Results = Row0.x * (1.0f - Fraction.y) + Row0.y;
+		float Results = Row0.x * (1.0f - Fraction.y) + Row0.y;
 
 		float4 Values02 = CalculateOcclusion(Settings.ShadowDepthTexture.Gather(Settings.ShadowDepthTextureSampler, SamplePos, int2(-2, 0)), Settings);
 		float4 Values22 = CalculateOcclusion(Settings.ShadowDepthTexture.Gather(Settings.ShadowDepthTextureSampler, SamplePos, int2(0, 0)), Settings);
@@ -412,28 +414,25 @@ float ManualPCF(float2 ShadowPosition, FPCFSamplerSettings Settings)
 			// Group work to minimize temporary registers and to split texture work with PCF ALU operations to hide texture latency
 			// Without this layout (all texture lookups at the beginning, PCF ALU's at the end) this shader was 4x slower on Nvidia cards
 			{
-				float4 SampleValues00;
-				FetchRowOfFour(Sample00TexelCenter, 0, SampleValues00, Settings);
+				float4 SampleValues00 =  FetchRowOfFour(Sample00TexelCenter, 0, Settings);
 				SampleValues10.x = SampleValues00.w;
 
-				float4 SampleValues01;
-				FetchRowOfFour(Sample00TexelCenter, 1, SampleValues01, Settings);
+				float4 SampleValues01 =  FetchRowOfFour(Sample00TexelCenter, 1, Settings);
 				SampleValues11.x = SampleValues01.w;
 
-				float4 SampleValues02;
-				FetchRowOfFour(Sample00TexelCenter, 2, SampleValues02, Settings);
+				float4 SampleValues02 = FetchRowOfFour(Sample00TexelCenter, 2, Settings);
 				SampleValues12.x = SampleValues02.w;
 
-				FetchRowOfFour(Sample00TexelCenter, 3, SampleValues03, Settings);
+				SampleValues03 = FetchRowOfFour(Sample00TexelCenter, 3, Settings);
 				SampleValues13.x = SampleValues03.w;
 				Results.x = PCF3x3(Fraction, SampleValues00, SampleValues01, SampleValues02, SampleValues03);
 			}
 
 			{
-				FetchRowOfThreeAfterFour(Sample00TexelCenter, 0, SampleValues10.yzw, Settings);
-				FetchRowOfThreeAfterFour(Sample00TexelCenter, 1, SampleValues11.yzw, Settings);
-				FetchRowOfThreeAfterFour(Sample00TexelCenter, 2, SampleValues12.yzw, Settings);
-				FetchRowOfThreeAfterFour(Sample00TexelCenter, 3, SampleValues13.yzw, Settings);
+				SampleValues10.yzw = FetchRowOfThreeAfterFour(Sample00TexelCenter, 0, Settings);
+				SampleValues11.yzw = FetchRowOfThreeAfterFour(Sample00TexelCenter, 1, Settings);
+				SampleValues12.yzw = FetchRowOfThreeAfterFour(Sample00TexelCenter, 2, Settings);
+				SampleValues13.yzw = FetchRowOfThreeAfterFour(Sample00TexelCenter, 3, Settings);
 				Results.y = PCF3x3(Fraction, SampleValues10, SampleValues11, SampleValues12, SampleValues13);
 			}
 		}
@@ -444,25 +443,22 @@ float ManualPCF(float2 ShadowPosition, FPCFSamplerSettings Settings)
 			float4 SampleValues16;
 
 			{
-				float4 SampleValues04;
-				FetchRowOfFour(Sample00TexelCenter, 4, SampleValues04, Settings);
+				float4 SampleValues04 =  FetchRowOfFour(Sample00TexelCenter, 4, Settings);
 				SampleValues14.x = SampleValues04.w;
 
-				float4 SampleValues05;
-				FetchRowOfFour(Sample00TexelCenter, 5, SampleValues05, Settings);
+				float4 SampleValues05 = FetchRowOfFour(Sample00TexelCenter, 5, Settings);
 				SampleValues15.x = SampleValues05.w;
 
-				float4 SampleValues06;
-				FetchRowOfFour(Sample00TexelCenter, 6, SampleValues06, Settings);
+				float4 SampleValues06 = FetchRowOfFour(Sample00TexelCenter, 6, Settings);
 				SampleValues16.x = SampleValues06.w;
 
 				Results.z = PCF3x3(Fraction, SampleValues03, SampleValues04, SampleValues05, SampleValues06);
 			}
 
 			{
-				FetchRowOfThreeAfterFour(Sample00TexelCenter, 4, SampleValues14.yzw, Settings);
-				FetchRowOfThreeAfterFour(Sample00TexelCenter, 5, SampleValues15.yzw, Settings);
-				FetchRowOfThreeAfterFour(Sample00TexelCenter, 6, SampleValues16.yzw, Settings);
+				SampleValues14.yzw = FetchRowOfThreeAfterFour(Sample00TexelCenter, 4, Settings);
+				SampleValues15.yzw = FetchRowOfThreeAfterFour(Sample00TexelCenter, 5, Settings);
+				SampleValues16.yzw = FetchRowOfThreeAfterFour(Sample00TexelCenter, 6, Settings);
 				Results.w = PCF3x3(Fraction, SampleValues13, SampleValues14, SampleValues15, SampleValues16);
 			}
 		}
