Always specify an output render target for depth pixel shader.

#rnx

diff --git a/Engine/Shaders/Private/ShadowDepthPixelShader.usf b/Engine/Shaders/Private/ShadowDepthPixelShader.usf
--- a/Engine/Shaders/Private/ShadowDepthPixelShader.usf
+++ b/Engine/Shaders/Private/ShadowDepthPixelShader.usf
@@ -48,9 +48,9 @@ float3 GetMaterialBounceColor(FPixelMaterialInputs PixelMaterialInputs)
 
 void Main( 
 	FShadowDepthVSToPS Inputs,
-	in float4 SvPosition : SV_Position		// after all interpolators
+	in float4 SvPosition : SV_Position,		// after all interpolators
+	out float4 OutColor : SV_Target0		// OutReflectiveShadowMapNormal, when REFLECTIVE_SHADOW_MAP is defined.
 #if REFLECTIVE_SHADOW_MAP
-	,out float4 OutReflectiveShadowMapNormal  : SV_Target0
 	,out float4 OutReflectiveShadowMapDiffuse : SV_Target1
 #endif
 #if PERSPECTIVE_CORRECT_DEPTH
@@ -91,6 +91,8 @@ void Main(
 	OutDepth = saturate(Inputs.ShadowDepth);
 #endif
 
+	OutColor = 0;
+
 #if REFLECTIVE_SHADOW_MAP
 	OutReflectiveShadowMapDiffuse = 0;
 	float3 normal = normalize( MaterialParameters.TangentToWorld[2] ); // Smooth normal
@@ -138,22 +140,22 @@ void Main(
 	}
 	else
 	{
-		OutReflectiveShadowMapNormal = 0;
+		OutColor = 0;
 		OutReflectiveShadowMapDiffuse = 0;
 		return;
 	}
 #endif
 
 #if MATERIAL_BLOCK_GI
-	OutReflectiveShadowMapNormal = 1.0f;
+	OutColor = 1.0f;
 	OutReflectiveShadowMapDiffuse = opacity;
 #elif MATERIALBLENDING_TRANSLUCENT || MATERIALBLENDING_ALPHACOMPOSITE || MATERIALBLENDING_ALPHAHOLDOUT
-	OutReflectiveShadowMapNormal = 0;
+	OutColor = 0;
 	OutReflectiveShadowMapDiffuse = 0;
 	clip(-1.0);
 #else
 	// Pass the LPV bias multiplier in the alpha channel of the normal
-	OutReflectiveShadowMapNormal = float4(normal * .5 + .5, GetPrimitiveData(MaterialParameters.PrimitiveId).LpvBiasMultiplier / 32.0f );
+	OutColor = float4(normal * .5 + .5, GetPrimitiveData(MaterialParameters.PrimitiveId).LpvBiasMultiplier / 32.0f );
 	OutReflectiveShadowMapDiffuse = float4(flux, 0);
 #endif
 
