Fix several global shaders related to SSD for HLSLcc:
- Sampler in structs are not supported by HLSLcc, so "SourceTexture" in "FLightShaderParameters" is disabled if a Non-RayTracing shader is used.
- Replaced expressions that assign textures conditionally via ternary operators.
- Fixed macros in DiffuseIndirectComposite.usf that caused array index out of bounds error.
- Added placeholder parameters in "FDiffuseIndirectCompositePS" global shader to enable it for Vulkan again.

#rnx

diff --git a/Engine/Shaders/Private/DiffuseIndirectComposite.usf b/Engine/Shaders/Private/DiffuseIndirectComposite.usf
--- a/Engine/Shaders/Private/DiffuseIndirectComposite.usf
+++ b/Engine/Shaders/Private/DiffuseIndirectComposite.usf
@@ -20,8 +20,8 @@
 #elif DIM_APPLY_DIFFUSE_INDIRECT == 2 // Probe hierarchy output
 	#define CONFIG_SIGNAL_PROCESSING SIGNAL_PROCESSING_DIFFUSE_PROBE_HIERARCHY
 	#define COMPILE_SIGNAL_COLOR_ARRAY 2
-	#define MAX_SIGNAL_BATCH_SIZE 1
-	#define SIGNAL_ARRAY_SIZE 1
+	#define MAX_SIGNAL_BATCH_SIZE 2
+	#define SIGNAL_ARRAY_SIZE 2
 
 	#define CONFIG_SIGNAL_INPUT_LAYOUT  SIGNAL_BUFFER_LAYOUT_DIFFUSE_PROBE_HIERARCHY_HISTORY
 	#define CONFIG_INPUT_TEXTURE_COUNT 3
diff --git a/Engine/Shaders/Private/LightShaderParameters.ush b/Engine/Shaders/Private/LightShaderParameters.ush
--- a/Engine/Shaders/Private/LightShaderParameters.ush
+++ b/Engine/Shaders/Private/LightShaderParameters.ush
@@ -3,6 +3,12 @@
 #pragma once
 
 
+#if RAYCALLABLESHADER || RAYGENSHADER || RAYHITGROUPSHADER || RAYMISSSHADER
+#	define SHADER_PARAMETERS_HAS_SOURCETEXTURE 1
+#else
+#	define SHADER_PARAMETERS_HAS_SOURCETEXTURE 0
+#endif
+
 // Mirrors FLightShaderParameters on C++ side.
 struct FLightShaderParameters
 {
@@ -19,7 +25,9 @@ struct FLightShaderParameters
 	float  SourceLength;
 	float  RectLightBarnCosAngle;
 	float  RectLightBarnLength;
+	#if SHADER_PARAMETERS_HAS_SOURCETEXTURE
 	Texture2D SourceTexture;
+	#endif
 };
 
 
@@ -71,6 +79,8 @@ FLightShaderParameters GetRootLightShaderParameters()
 	LightParameters.SourceLength = Light_SourceLength;
 	LightParameters.RectLightBarnCosAngle = Light_RectLightBarnCosAngle;
 	LightParameters.RectLightBarnLength = Light_RectLightBarnLength;
+	#if SHADER_PARAMETERS_HAS_SOURCETEXTURE
 	LightParameters.SourceTexture = Light_SourceTexture;
+	#endif
 	return LightParameters;
 }
diff --git a/Engine/Shaders/Private/ScreenSpaceDenoise/SSDCommon.ush b/Engine/Shaders/Private/ScreenSpaceDenoise/SSDCommon.ush
--- a/Engine/Shaders/Private/ScreenSpaceDenoise/SSDCommon.ush
+++ b/Engine/Shaders/Private/ScreenSpaceDenoise/SSDCommon.ush
@@ -277,12 +277,18 @@ FSSDCompressedSceneInfos SampleCompressedSceneMetadata(
 {
 	FSSDCompressedSceneInfos CompressedMetadata = CreateCompressedSceneInfos();
 
-	Texture2D<float> CompressedMetadata0 = bPrevFrame ? PrevCompressedMetadata_0 : CompressedMetadata_0;
-	Texture2D<uint>  CompressedMetadata1 = bPrevFrame ? PrevCompressedMetadata_1 : CompressedMetadata_1;
-
 	int3 Coord = int3(PixelCoord, /* MipLevel = */ 0);
-	CompressedMetadata.VGPR[0] = asuint(CompressedMetadata0.Load(Coord));
-	CompressedMetadata.VGPR[1] = CompressedMetadata1.Load(Coord);
+
+	if (bPrevFrame)
+	{
+		CompressedMetadata.VGPR[0] = asuint(PrevCompressedMetadata_0.Load(Coord));
+		CompressedMetadata.VGPR[1] = PrevCompressedMetadata_1.Load(Coord);
+	}
+	else
+	{
+		CompressedMetadata.VGPR[0] = asuint(CompressedMetadata_0.Load(Coord));
+		CompressedMetadata.VGPR[1] = CompressedMetadata_1.Load(Coord);
+	}
 
 	return CompressedMetadata;
 }
@@ -290,10 +296,16 @@ FSSDCompressedSceneInfos SampleCompressedSceneMetadata(
 {
 	FSSDCompressedSceneInfos CompressedMetadata = CreateCompressedSceneInfos();
 
-	Texture2D<uint> CompressedMetadata0 = bPrevFrame ? PrevCompressedMetadata_0 : CompressedMetadata_0; 
-
 	int3 Coord = int3(PixelCoord, /* MipLevel = */ 0);
-	CompressedMetadata.VGPR[0] = CompressedMetadata0.Load(Coord);
+
+	if (bPrevFrame)
+	{
+		CompressedMetadata.VGPR[0] = PrevCompressedMetadata_0.Load(Coord);
+	}
+	else
+	{
+		CompressedMetadata.VGPR[0] = CompressedMetadata_0.Load(Coord);
+	}
 	
 	return CompressedMetadata;
 }
diff --git a/Engine/Shaders/Private/ScreenSpaceDenoise/SSDSignalFramework.ush b/Engine/Shaders/Private/ScreenSpaceDenoise/SSDSignalFramework.ush
--- a/Engine/Shaders/Private/ScreenSpaceDenoise/SSDSignalFramework.ush
+++ b/Engine/Shaders/Private/ScreenSpaceDenoise/SSDSignalFramework.ush
@@ -77,7 +77,9 @@ float4 LightDirectionAndLength[MAX_SIGNAL_BATCH_SIZE];
 float HitDistanceToWorldBluringRadius[MAX_SIGNAL_BATCH_SIZE];
 uint LightType[MAX_SIGNAL_BATCH_SIZE];
 
+#if SHADER_PARAMETERS_HAS_SOURCETEXTURE
 Texture2D UnusedLightSourceTexture;
+#endif
 
 
 /** Returns the FLightShaderParameters from the root shader parameters. */
@@ -92,7 +94,9 @@ FSSDSignalDomainKnowledge GetSignalDomainKnowledge(uint BatchedSignalId)
 	SignalDomain.Light.SpotAngles = 0;
 	SignalDomain.Light.SpecularScale = 0;
 	SignalDomain.Light.SoftSourceRadius = 0;
+	#if SHADER_PARAMETERS_HAS_SOURCETEXTURE
 	SignalDomain.Light.SourceTexture = UnusedLightSourceTexture;
+	#endif
 	
 	SignalDomain.Light.Position = LightPositionAndRadius[BatchedSignalId].xyz;
 	SignalDomain.Light.Direction = LightDirectionAndLength[BatchedSignalId].xyz;
diff --git a/Engine/Shaders/Private/ScreenSpaceDenoise/SSDSpatialAccumulation.usf b/Engine/Shaders/Private/ScreenSpaceDenoise/SSDSpatialAccumulation.usf
--- a/Engine/Shaders/Private/ScreenSpaceDenoise/SSDSpatialAccumulation.usf
+++ b/Engine/Shaders/Private/ScreenSpaceDenoise/SSDSpatialAccumulation.usf
@@ -691,18 +691,6 @@ void MainCS(
 	uint2 GroupThreadId : SV_GroupThreadID,
 	uint GroupThreadIndex : SV_GroupIndex)
 {
-#if CONFIG_SIGNAL_INPUT_TEXTURE_TYPE == SIGNAL_TEXTURE_TYPE_FLOAT4
-	Texture2D Signal_Textures_0 = SignalInput_Textures_0;
-	Texture2D Signal_Textures_1 = SignalInput_Textures_1;
-	Texture2D Signal_Textures_2 = SignalInput_Textures_2;
-	Texture2D Signal_Textures_3 = SignalInput_Textures_3;
-#else
-	FSSDTexture2D Signal_Textures_0 = SignalInput_Textures_0;
-	FSSDTexture2D Signal_Textures_1 = SignalInput_Textures_1;
-	FSSDTexture2D Signal_Textures_2 = SignalInput_Textures_2;
-	FSSDTexture2D Signal_Textures_3 = SignalInput_Textures_3;
-#endif
-
 	// Find out scene buffer UV.
 	float2 SceneBufferUV = DispatchThreadId * ThreadIdToBufferUV.xy + ThreadIdToBufferUV.zw;
 	if (true)
@@ -730,10 +718,10 @@ void MainCS(
 		FSSDSignalArray RefSamples;
 		FSSDSignalFrequencyArray RefFrequencies;
 		SampleMultiplexedSignals(
-			Signal_Textures_0,
-			Signal_Textures_1,
-			Signal_Textures_2,
-			Signal_Textures_3,
+			SignalInput_Textures_0,
+			SignalInput_Textures_1,
+			SignalInput_Textures_2,
+			SignalInput_Textures_3,
 			GlobalPointClampedSampler,
 			CONFIG_SIGNAL_INPUT_LAYOUT,
 			/* MultiplexedSampleId = */ 0,
@@ -1068,10 +1056,10 @@ void MainCS(
 		{
 			AccumulateKernel(
 				KernelConfig,
-				Signal_Textures_0,
-				Signal_Textures_1,
-				Signal_Textures_2,
-				Signal_Textures_3,
+				SignalInput_Textures_0,
+				SignalInput_Textures_1,
+				SignalInput_Textures_2,
+				SignalInput_Textures_3,
 				/* inout */ UncompressedAccumulators,
 				/* inout */ CompressedAccumulators);
 		}
@@ -1089,10 +1077,10 @@ void MainCS(
 			
 				AccumulateKernel(
 					KernelConfig,
-					Signal_Textures_0,
-					Signal_Textures_1,
-					Signal_Textures_2,
-					Signal_Textures_3,
+					SignalInput_Textures_0,
+					SignalInput_Textures_1,
+					SignalInput_Textures_2,
+					SignalInput_Textures_3,
 					/* inout */ UncompressedAccumulators,
 					/* inout */ CompressedAccumulators);
 			}
@@ -1108,10 +1096,10 @@ void MainCS(
 
 			SampleAndAccumulateCenterSampleAsItsOwnCluster(
 				KernelConfig,
-				Signal_Textures_0,
-				Signal_Textures_1,
-				Signal_Textures_2,
-				Signal_Textures_3,
+				SignalInput_Textures_0,
+				SignalInput_Textures_1,
+				SignalInput_Textures_2,
+				SignalInput_Textures_3,
 				/* inout */ UncompressedAccumulators,
 				/* inout */ CompressedAccumulators);
 		}
diff --git a/Engine/Shaders/Public/Platform/D3D/D3DCommon.ush b/Engine/Shaders/Public/Platform/D3D/D3DCommon.ush
--- a/Engine/Shaders/Public/Platform/D3D/D3DCommon.ush
+++ b/Engine/Shaders/Public/Platform/D3D/D3DCommon.ush
@@ -4,6 +4,9 @@
 D3DCommon.usf: Common defines for the d3d shader compiler
 =============================================================================*/
 
+// Update this GUID to invalidate D3D shaders and force a shader recompilation for D3D RHIs.
+// GUID: 420707B5-DB86-4011-97F5-AEAAF78928DD
+
 precise float MakePrecise(precise float v) { return v; }
 precise float2 MakePrecise(precise float2 v) { return v; }
 precise float3 MakePrecise(precise float3 v) { return v; }
diff --git a/Engine/Shaders/Public/Platform/Vulkan/VulkanCommon.ush b/Engine/Shaders/Public/Platform/Vulkan/VulkanCommon.ush
--- a/Engine/Shaders/Public/Platform/Vulkan/VulkanCommon.ush
+++ b/Engine/Shaders/Public/Platform/Vulkan/VulkanCommon.ush
@@ -4,8 +4,8 @@
 	VulkanCommon.usf: Common Vulkan shader code
 =============================================================================*/
 
-// Update this GUID to improve shader recompilation for Vulkan only shaders
-// GUID = F1EAFBAF-17ED-4845-8FA3-F640F48665E6
+// Update this GUID to invalidate Vulkan shaders and force a shader recompilation for the Vulkan RHI.
+// GUID = 97939BC2-746A-481C-B024-9DECD6366A92
 // Glslang 1.1.82.1
 
 #pragma once
diff --git a/Engine/Source/Runtime/Renderer/Private/IndirectLightRendering.cpp b/Engine/Source/Runtime/Renderer/Private/IndirectLightRendering.cpp
--- a/Engine/Source/Runtime/Renderer/Private/IndirectLightRendering.cpp
+++ b/Engine/Source/Runtime/Renderer/Private/IndirectLightRendering.cpp
@@ -137,6 +137,10 @@ class FDiffuseIndirectCompositePS : public FGlobalShader
 
 		SHADER_PARAMETER_RDG_TEXTURE_UAV(RWTexture2D<float4>, PassDebugOutput)
 
+		SHADER_PARAMETER(FVector2D, BufferUVToOutputPixelPosition)
+		SHADER_PARAMETER_RDG_TEXTURE(Texture2D, EyeAdaptation)
+		SHADER_PARAMETER_RDG_TEXTURE_ARRAY(Texture2D<uint>, CompressedMetadata, [2])
+
 		RENDER_TARGET_BINDING_SLOTS()
 	END_SHADER_PARAMETER_STRUCT()
 };
@@ -1059,6 +1063,26 @@ void FDeferredShadingSceneRenderer::RenderDiffuseIndirectAndAmbientOcclusion(
 				PassParameters->ApplyAOToDynamicDiffuseIndirect = 1.0f;
 			}
 
+			const FIntPoint BufferExtent = SceneTextures.SceneDepthBuffer->Desc.Extent;
+
+			{
+				// Placeholder texture for textures pulled in from SSDCommon.ush
+				FRDGTextureDesc Desc = FRDGTextureDesc::Create2DDesc(
+					FIntPoint(1),
+					PF_R32_UINT,
+					FClearValueBinding::Black,
+					/* InFlags = */ TexCreate_None,
+					/* InTargetableFlags = */ TexCreate_ShaderResource,
+					/* bInForceSeparateTargetAndShaderResource = */ false);
+				FRDGTextureRef CompressedMetadataPlaceholder = GraphBuilder.CreateTexture(Desc, TEXT("CompressedMetadataPlaceholder"));
+
+				PassParameters->CompressedMetadata[0] = CompressedMetadataPlaceholder;
+				PassParameters->CompressedMetadata[1] = CompressedMetadataPlaceholder;
+			}
+
+			PassParameters->BufferUVToOutputPixelPosition = BufferExtent;
+			PassParameters->EyeAdaptation = GetEyeAdaptationTexture(GraphBuilder, View);
+
 			PassParameters->bVisualizeDiffuseIndirect = bIsVisualizePass;
 
 			PassParameters->DiffuseIndirect = DenoiserOutputs;
@@ -1087,7 +1111,6 @@ void FDeferredShadingSceneRenderer::RenderDiffuseIndirectAndAmbientOcclusion(
 			PassParameters->RenderTargets[0] = FRenderTargetBinding(
 				SceneColor, ERenderTargetLoadAction::ELoad);
 
-			if (1)
 			{
 				FRDGTextureDesc Desc = FRDGTextureDesc::Create2DDesc(
 					SceneColor->Desc.Extent,
