Fix virtual shadow map denoising
- Update encoding in the shadow projection shader to align with changes to the denoiser
- The coupling is not ideal, but we need to avoid the 2-3 extra passes per light

diff --git a/Engine/Shaders/Private/VirtualShadowMaps/VirtualShadowMapProjection.usf b/Engine/Shaders/Private/VirtualShadowMaps/VirtualShadowMapProjection.usf
--- a/Engine/Shaders/Private/VirtualShadowMaps/VirtualShadowMapProjection.usf
+++ b/Engine/Shaders/Private/VirtualShadowMaps/VirtualShadowMapProjection.usf
@@ -181,25 +181,38 @@ void VirtualShadowMapProjectionPS(
 	}
 #endif
 
-	// See signal encoding for SIGNAL_BUFFER_LAYOUT_PENUMBRA_RECONSTRUCTION in SSDSignalBufferEncoding.ush
+	float RayMissFactor = VirtualShadowMapSample.bOccluded ? 0.0f : 1.0f;
+	float HitDistance = VirtualShadowMapSample.OccluderDistance;
+	if (ContactShadow < 1.0f)
+	{
+		RayMissFactor *= ContactShadow;
+		HitDistance = ContactShadowLengthWorld;
+	}
+
+	// See signal encoding for SIGNAL_BUFFER_LAYOUT_PENUMBRA_HISTORY in SSDSignalBufferEncoding.ush
+	// We fuse that logic in here to avoid needing another pass, but it's hard to call the functions directly as they depend on a large
+	// number of preprocessor defines and logic that we don't want to mirror here.
+	// Unfortunately that means this logic can get out of date when the denoiser changes, but that's the trade-off for now.
 	float4 OutputSignal = float4(0, 0, 0, 0);
-	if (VirtualShadowMapSample.bOccluded)
+	float SampleCount = 1.0f;
+	float TransmissionDistance = RayMissFactor; // As per RayTracingOcclusionRGS.usf
+	if (RayMissFactor < 1.0f)
 	{
-		// Hit
-		float WorldBluringRadius = ComputeLightSampleWorldBluringRadius(TranslatedWorldPosition, LightType, Light, VirtualShadowMapSample.OccluderDistance);
-		OutputSignal = float4(0, 1, WorldBluringRadius, 0);
+		float WorldBluringRadius = ComputeLightSampleWorldBluringRadius(TranslatedWorldPosition, LightType, Light, HitDistance);
+
+		OutputSignal = float4(
+			RayMissFactor,
+			SampleCount,
+			WorldBluringRadius,
+			TransmissionDistance);
 	}
 	else
 	{
 		// Miss
-		OutputSignal = float4(1, 1, -1, 0);
+		OutputSignal = float4(RayMissFactor, SampleCount, WORLD_RADIUS_MISS, TransmissionDistance);
 	}
 
-	if (ContactShadow < 1)
-	{
-		OutputSignal.x *= ContactShadow;
-		OutputSignal.z = ComputeLightSampleWorldBluringRadius(TranslatedWorldPosition, LightType, Light, ContactShadowLengthWorld);
-	}
+	
 
 #if OUTPUT_TYPE == 0
 	// Output feeds into denoiser
@@ -230,8 +243,8 @@ void VirtualShadowMapCompositePS(
 	// NOTE: The signal is encoded as per SSDSignalBufferEncoding.ush, SIGNAL_BUFFER_LAYOUT_PENUMBRA_RECONSTRUCTION
 	float4 Input = InputSignal.Load(int3(SvPosition.xy, 0));
 	float SampleCount = Input.g;
-	float MissCount = Input.r;
-	float Shadow = (SampleCount > 0 ? MissCount : 1.0f);
+	float MissFactor = Input.r;
+	float Shadow = (SampleCount > 0 ? MissFactor : 1.0f);
 	
 	OutShadowMask = EncodeLightAttenuationFromMask(Shadow);
 }
diff --git a/Engine/Source/Runtime/Renderer/Private/ShadowRendering.cpp b/Engine/Source/Runtime/Renderer/Private/ShadowRendering.cpp
--- a/Engine/Source/Runtime/Renderer/Private/ShadowRendering.cpp
+++ b/Engine/Source/Runtime/Renderer/Private/ShadowRendering.cpp
@@ -1789,7 +1789,7 @@ bool FSceneRenderer::RenderShadowProjections(FRHICommandListImmediate& RHICmdLis
 						FSceneTextureParameters SceneTextures;
 						SetupSceneTextureParameters(GraphBuilder, &SceneTextures);
 
-						// Virtual SMs draw into a texture with R=Mask, G=OccluderDistance to feed into denoising, similar to ray traced shadows
+						// Virtual SMs draw into a texture formatted for the denoiser. See VirtualShadowMapProjection.usf
 						FRDGTextureRef SignalTexture;
 						const FLinearColor ClearColor(0.0f, 0.0f, 0.0f, 0.0f);
 						{
