Lumen - restore non DDA global SDF tracing which doesn't have issues with the volume wrappng

diff --git a/Engine/Shaders/Private/GlobalDistanceFieldUtils.ush b/Engine/Shaders/Private/GlobalDistanceFieldUtils.ush
--- a/Engine/Shaders/Private/GlobalDistanceFieldUtils.ush
+++ b/Engine/Shaders/Private/GlobalDistanceFieldUtils.ush
@@ -75,107 +75,74 @@ FGlobalSDFTraceResult RayTraceGlobalDistanceField(FGlobalSDFTraceInput TraceInpu
 
 	// Bias
 	float RayStartDistance = max(TraceInput.MinTraceDistance, GlobalVolumeCenterAndExtent[MinClipmapIndex].w * GlobalVolumeTexelSize);
-	float SampleRayTime = RayStartDistance;
-
-	int3 VoxelStep = (TraceInput.WorldRayDirection < 0) ? int3(-1, -1, -1) : int3(1, 1, 1);
-	const float NumPagesPerClipmap = GlobalDistanceFieldClipmapSizeInPages;
-	const float InvNumPagesPerClipmap = rcp(NumPagesPerClipmap);
+	float MinRayTime = RayStartDistance / TraceInput.MaxTraceDistance;
 
 	LOOP
-	for (uint ClipmapIndex = MinClipmapIndex; ClipmapIndex < NumGlobalSDFClipmaps && TraceResult.HitTime < 0.0f && SampleRayTime < TraceInput.MaxTraceDistance; ++ClipmapIndex)
+	for (uint ClipmapIndex = MinClipmapIndex; ClipmapIndex < NumGlobalSDFClipmaps && TraceResult.HitTime < 0.0f; ++ClipmapIndex)
 	{
-		float3 GlobalVolumeCenter = GlobalVolumeCenterAndExtent[ClipmapIndex].xyz;
-		float GlobalVolumeExtent = GlobalVolumeCenterAndExtent[ClipmapIndex].w;
-
-		float VoxelSize = 2.0f * GlobalVolumeExtent * InvNumPagesPerClipmap;
-		float3 GlobalClipmapMin = GlobalVolumeCenter - GlobalVolumeExtent;
-		float3 GlobalClipmapSize = 2.0f * GlobalVolumeExtent;
-
-		float3 SampleWorldPosition = TraceInput.WorldRayStart + TraceInput.WorldRayDirection * SampleRayTime;
-		int3 VoxelCoord = NumPagesPerClipmap * (SampleWorldPosition - GlobalClipmapMin) / GlobalClipmapSize;
-		float3 VoxelMin = (VoxelCoord - 0.0f) * VoxelSize + GlobalClipmapMin;
-		float3 VoxelMax = (VoxelCoord + 1.0f) * VoxelSize + GlobalClipmapMin;
-		float3 MaxNegT = (VoxelMin - TraceInput.WorldRayStart) / TraceInput.WorldRayDirection;
-		float3 MaxPosT = (VoxelMax - TraceInput.WorldRayStart) / TraceInput.WorldRayDirection;
-		float3 NextVoxelFaceT = (TraceInput.WorldRayDirection < 0) ? MaxNegT : MaxPosT;
-		float3 VoxelFaceStepT = abs(VoxelSize / TraceInput.WorldRayDirection);
-
 		float ClipmapVoxelExtent = GlobalVolumeCenterAndExtent[ClipmapIndex].w * GlobalVolumeTexelSize;
 		float MinStepSize = TraceInput.MinStepFactor * ClipmapVoxelExtent;
 		float ExpandSurfaceDistance = ClipmapVoxelExtent;
-		const float ClipmapInfluenceRange = GLOBAL_DISTANCE_FIELD_INFLUENCE_RANGE_IN_VOXELS * 2.0f * GlobalVolumeCenterAndExtent[ClipmapIndex].w * GlobalVolumeTexelSize;
 
-		uint StepIndex = 0;
-		const uint MaxSteps = 256;
+		float3 GlobalVolumeCenter = GlobalVolumeCenterAndExtent[ClipmapIndex].xyz;
+		float GlobalVolumeExtent = GlobalVolumeCenterAndExtent[ClipmapIndex].w - ClipmapVoxelExtent;
+		float3 WorldRayEnd = TraceInput.WorldRayStart + TraceInput.WorldRayDirection * TraceInput.MaxTraceDistance;
+		float2 IntersectionTimes = LineBoxIntersect(TraceInput.WorldRayStart, WorldRayEnd, GlobalVolumeCenter - GlobalVolumeExtent.xxx, GlobalVolumeCenter + GlobalVolumeExtent.xxx);
+		IntersectionTimes.x = max(IntersectionTimes.x, MinRayTime);
 
-		LOOP
-		for (; StepIndex < MaxSteps; ++StepIndex)
+		// Update the trace start for the next clipmap
+		MinRayTime = IntersectionTimes.y;
+
+		if (IntersectionTimes.x < IntersectionTimes.y)
 		{
-			int3 PageTableCoord = frac(VoxelCoord * InvNumPagesPerClipmap + GlobalDistanceFieldPageTableScrollOffset[ClipmapIndex]) * NumPagesPerClipmap
-				+ int3(0, 0, ClipmapIndex * GlobalDistanceFieldClipmapSizeInPages);
+			float SampleRayTime = IntersectionTimes.x * TraceInput.MaxTraceDistance;
+			const float ClipmapInfluenceRange = GLOBAL_DISTANCE_FIELD_INFLUENCE_RANGE_IN_VOXELS * 2.0f * GlobalVolumeCenterAndExtent[ClipmapIndex].w * GlobalVolumeTexelSize;
 
-			const uint PageIndex = GlobalDistanceFieldPageTableTexture.Load(int4(PageTableCoord, 0));
-			const float MinNextVoxelFaceT = min3(NextVoxelFaceT.x, NextVoxelFaceT.y, NextVoxelFaceT.z);
+			uint StepIndex = 0;
+			const uint MaxSteps = 256;
 
-			if (PageIndex != GLOBAL_DISTANCE_FIELD_INVALID_PAGE_ID && SampleRayTime < MinNextVoxelFaceT)
+			LOOP
+			for (; StepIndex < MaxSteps; ++StepIndex)
 			{
-				float3 BasePagePos = -(GlobalClipmapMin + VoxelCoord * VoxelSize) / VoxelSize;
-				float3 BasePageCoord = GlobalDistanceFieldPageLinearIndexToPageAtlasCoord(PageIndex) + BasePagePos * GLOBAL_DISTANCE_FIELD_PAGE_RESOLUTION + 1.0f;
-				float3 BasePageUV = BasePageCoord * GlobalDistanceFieldInvPageAtlasSize;
+				float3 SampleWorldPosition = TraceInput.WorldRayStart + TraceInput.WorldRayDirection * SampleRayTime;
 
-				// Trace single page
-				LOOP
-				for (; StepIndex < MaxSteps; ++StepIndex)
-				{
-					SampleWorldPosition = TraceInput.WorldRayStart + TraceInput.WorldRayDirection * SampleRayTime;
-					float3 PageUV = BasePageUV + (SampleWorldPosition / VoxelSize) * GLOBAL_DISTANCE_FIELD_PAGE_RESOLUTION * GlobalDistanceFieldInvPageAtlasSize;
+				float DistanceField = ClipmapInfluenceRange;
+
+				float3 ClipmapVolumeUV = frac(ComputeGlobalUV(SampleWorldPosition, ClipmapIndex));
+				uint PageIndex = GetGlobalDistanceFieldPage(ClipmapVolumeUV, ClipmapIndex);
 
+				if (PageIndex < 0xFFFF)
+				{
+					float3 PageUV = ComputeGlobalDistanceFieldPageUV(ClipmapVolumeUV, PageIndex);
 					float DistanceFieldValue = Texture3DSampleLevel(GlobalDistanceFieldPageAtlasTexture, GlobalTrilinearWrappedSampler, PageUV, 0).x;
-					float DistanceField = DecodeGlobalDistanceFieldPageDistance(DistanceFieldValue, ClipmapInfluenceRange);
+					DistanceField = DecodeGlobalDistanceFieldPageDistance(DistanceFieldValue, ClipmapInfluenceRange);
+				}
 
-					MaxDistance = max(DistanceField, MaxDistance);
-					float ExpandSurfaceTime = TraceInput.bExpandSurfaceUsingRayTimeInsteadOfMaxDistance ? SampleRayTime : MaxDistance;
+				MaxDistance = max(DistanceField, MaxDistance);
+				float ExpandSurfaceTime = TraceInput.bExpandSurfaceUsingRayTimeInsteadOfMaxDistance ? SampleRayTime : MaxDistance;
 
-					float ExpandSurfaceFalloff = ExpandSurfaceDistance;
-					float ExpandSurfaceAmount = ExpandSurfaceDistance * saturate((ExpandSurfaceTime - ExpandSurfaceFalloff) / ExpandSurfaceFalloff);
+				float ExpandSurfaceFalloff = ExpandSurfaceDistance;
+				float ExpandSurfaceAmount = ExpandSurfaceDistance * saturate((ExpandSurfaceTime - ExpandSurfaceFalloff) / ExpandSurfaceFalloff);
 
-					if (DistanceField < ExpandSurfaceAmount)
-					{
-						TraceResult.HitTime = max(SampleRayTime + DistanceField - ExpandSurfaceAmount, RayStartDistance);
-						TraceResult.HitClipmapIndex = ClipmapIndex;
-						break;
-					}
+				if (DistanceField < ExpandSurfaceAmount)
+				{
+					TraceResult.HitTime = max(SampleRayTime + DistanceField - ExpandSurfaceAmount, TraceInput.MinTraceDistance);
+					TraceResult.HitClipmapIndex = ClipmapIndex;
+					break;
+				}
 
-					float StepDistance = max(DistanceField * TraceInput.StepFactor, MinStepSize);
-					SampleRayTime += StepDistance;
+				float StepDistance = max(DistanceField * TraceInput.StepFactor, MinStepSize);
+				SampleRayTime += StepDistance;
 
-					// Terminate the trace if we went past the end of page or hit something
-					if (SampleRayTime >= MinNextVoxelFaceT || SampleRayTime >= TraceInput.MaxTraceDistance)
-					{
-						break;
-					}
+				// Terminate the trace if we went past the end of the ray or achieved enough occlusion
+				if (SampleRayTime > IntersectionTimes.y * TraceInput.MaxTraceDistance || TraceResult.HitTime >= 0.0f)
+				{
+					break;
 				}
 			}
 
-			// Step to the next voxel (page)
-			uint3 StepInc = int3(NextVoxelFaceT == MinNextVoxelFaceT);
-			SampleRayTime = max(SampleRayTime, MinNextVoxelFaceT);
-			NextVoxelFaceT += StepInc * VoxelFaceStepT;
-			VoxelCoord += StepInc * VoxelStep;
-
-			// Next clipmap if we are out of bounds
-			if (any(VoxelCoord < 0) || any(VoxelCoord >= NumPagesPerClipmap) || TraceResult.HitTime >= 0.0f || SampleRayTime > TraceInput.MaxTraceDistance)
-			{
-				break;
-			}
+			TraceResult.TotalStepsTaken += StepIndex;
 		}
-
-		TraceResult.TotalStepsTaken += StepIndex;
-	}
-
-	if (TraceResult.HitTime > TraceInput.MaxTraceDistance)
-	{
-		TraceResult.HitTime = -1;
 	}
 
 	return TraceResult;
