Fixes for shader compilation errors on mobile

diff --git a/Engine/Shaders/Private/MobileBasePassPixelShader.usf b/Engine/Shaders/Private/MobileBasePassPixelShader.usf
--- a/Engine/Shaders/Private/MobileBasePassPixelShader.usf
+++ b/Engine/Shaders/Private/MobileBasePassPixelShader.usf
@@ -557,7 +557,7 @@ void Main(
 	IndirectIrradiance *= MaterialAO;
 
 	// Shadow
-	half Shadow = GetPrimaryPrecomputedShadowMask(Interpolants).r;
+	half Shadow = GetPrimaryPrecomputedShadowMask(Interpolants, MaterialParameters).r;
 #if DIRECTIONAL_LIGHT_CSM && !MATERIAL_SHADINGMODEL_SINGLELAYERWATER
 	// Cascaded Shadow Map
 	if(MobileBasePass.UseCSM)
diff --git a/Engine/Shaders/Private/ShadowDepthVertexShader.usf b/Engine/Shaders/Private/ShadowDepthVertexShader.usf
--- a/Engine/Shaders/Private/ShadowDepthVertexShader.usf
+++ b/Engine/Shaders/Private/ShadowDepthVertexShader.usf
@@ -318,7 +318,7 @@ void Main(
 		OutParameters.PixelPosition = WorldPos.xyz;
 	#endif
 
-	#if USING_VERTEX_SHADER_LAYER && VERTEXSHADER
+	#if FEATURE_LEVEL >= FEATURE_LEVEL_SM4 && USING_VERTEX_SHADER_LAYER && VERTEXSHADER
 		LayerIndex = LayerId;
 		OutPosition = mul(WorldPos, PassStruct.ShadowViewProjectionMatrices[LayerIndex]);
 		#if INTERPOLATE_VF_ATTRIBUTES
