Protect FSSDSignalArray::Array from out of range accesses

Change by [REDACTED_PERSON] per slack discussion: https://epic.slack.com/archives/C024FSUU2/p1599682889093700

#rnx

diff --git a/Engine/Shaders/Private/ScreenSpaceDenoise/SSDSignalBufferEncoding.ush b/Engine/Shaders/Private/ScreenSpaceDenoise/SSDSignalBufferEncoding.ush
--- a/Engine/Shaders/Private/ScreenSpaceDenoise/SSDSignalBufferEncoding.ush
+++ b/Engine/Shaders/Private/ScreenSpaceDenoise/SSDSignalBufferEncoding.ush
@@ -141,7 +141,7 @@ void DecodeMultiplexedSignalsFromFloat4(
 		}
 	}
 	#endif
-	#if COMPILE_SIGNAL_COLOR_ARRAY >= 2
+	#if COMPILE_SIGNAL_COLOR_ARRAY >= 2 && SIGNAL_ARRAY_SIZE >= 2
 	else if (SignalBufferLayout == SIGNAL_BUFFER_LAYOUT_POLYCHROMATIC_PENUMBRA_HARMONIC_INPUT)
 	{
 		// The RGS output same weigth for diffuse and specular given it depends only on shading model id.
@@ -540,7 +540,7 @@ void EncodeMultiplexedSignals(
 		OutBufferCount = MultiplexCount;
 	}
 	#endif
-	#if COMPILE_SIGNAL_COLOR_ARRAY >= 2
+	#if COMPILE_SIGNAL_COLOR_ARRAY >= 2 && SIGNAL_ARRAY_SIZE >= 2
 	else if (SignalBufferLayout == SIGNAL_BUFFER_LAYOUT_POLYCHROMATIC_PENUMBRA_HARMONIC_RECONSTRUCTION)
 	{
 		// diffuse harmonic
