Fix GetDistanceFieldGradientGlobal to be resolution indepedent.

diff --git a/Engine/Shaders/Private/GlobalDistanceFieldShared.ush b/Engine/Shaders/Private/GlobalDistanceFieldShared.ush
--- a/Engine/Shaders/Private/GlobalDistanceFieldShared.ush
+++ b/Engine/Shaders/Private/GlobalDistanceFieldShared.ush
@@ -200,7 +200,8 @@ float3 GlobalDistanceFieldPageCentralDiff(float3 PageUV)
 	return float3(R - L, F - B, U - D);
 }
 
-// Returns vector (unnormalized) to the nearest surface point
+// Returns vector (unnormalized) to the nearest surface point.
+// This vector is rescaled by 128 * GlobalVolumeTexelSize for backwards compatiblity.
 float3 GetDistanceFieldGradientGlobal(float3 WorldPosition)
 {
 	float3 Gradient = float3(0, 0, 0.001f);
@@ -231,7 +232,9 @@ float3 GetDistanceFieldGradientGlobal(float3 WorldPosition)
 					const float PageDistanceDecodeFactor = 2.0f * ClipmapInfluenceRange;
 					float ClipmapExtent = GlobalVolumeCenterAndExtent[ClipmapIndex].w;
 
-					Gradient = PageDistanceDecodeFactor * PageCentralDiff / (ClipmapExtent);
+					// Rescale to keep backwards compatability
+					const float BackwardsCompatibilityFactor = 128.0f * GlobalVolumeTexelSize;
+					Gradient = PageDistanceDecodeFactor * PageCentralDiff / (ClipmapExtent * BackwardsCompatibilityFactor);
 
 					break;
 				}
