Missing file to go with 14326107
thanks p4.

diff --git a/Engine/Shaders/Private/BasePassPixelShader.usf b/Engine/Shaders/Private/BasePassPixelShader.usf
--- a/Engine/Shaders/Private/BasePassPixelShader.usf
+++ b/Engine/Shaders/Private/BasePassPixelShader.usf
@@ -57,6 +57,8 @@
 #include "Lumen/LumenTranslucencyVolumeShared.ush"
 #endif
 
+#define STRATA_ENABLED (PROJECT_STRATA && MATERIALBLENDING_ANY_TRANSLUCENT<=0)
+
 float NormalCurvatureToRoughness(float3 WorldNormal)
 {
     float3 dNdx = ddx(WorldNormal);
@@ -1567,10 +1569,18 @@ void FPixelShaderInOut_MainPS(
 		View.VTFeedbackBuffer
 	);
 #endif
+
+#if STRATA_ENABLED
+	OpaqueBasePass.Strata.MaterialLobesTextureUAV[uint2(In.SvPosition.xy)] = 1.0f; // 
+	//OpaqueBasePass.Strata.MaterialLobesBufferUAV.Store(0*4, 0);
+	//OpaqueBasePass.Strata.MaterialLobesBufferUAV.Store(1*4, 1);
+	//OpaqueBasePass.Strata.MaterialLobesBufferUAV.Store(2*4, 2);
+	//OpaqueBasePass.Strata.MaterialLobesBufferUAV.Store(3*4, 3);
+#endif
 }
 
 // If virtual texture is enabled then use early depth test so that UAV feedback buffer writes respect the depth test
-#if NUM_VIRTUALTEXTURE_SAMPLES || LIGHTMAP_VT_ENABLED
+#if NUM_VIRTUALTEXTURE_SAMPLES || LIGHTMAP_VT_ENABLED || STRATA_ENABLED
 	#if COMPILER_SUPPORTS_DEPTHSTENCIL_EARLYTEST_LATEWRITE
 		// If we support early depth test with late write behaviour then use it since we may be using discard, or modifying depth
 		#define PIXELSHADER_EARLYDEPTHSTENCIL DEPTHSTENCIL_EARLYTEST_LATEWRITE
