Change directional light sampling in path tracer to use uniform cone angle sampling

diff --git a/Engine/Shaders/Private/PathTracing/Light/PathTracingDirectionalLight.ush b/Engine/Shaders/Private/PathTracing/Light/PathTracingDirectionalLight.ush
--- a/Engine/Shaders/Private/PathTracing/Light/PathTracingDirectionalLight.ush
+++ b/Engine/Shaders/Private/PathTracing/Light/PathTracingDirectionalLight.ush
@@ -39,7 +39,18 @@ void DirectionalLight_EvalLight(
 )
 {
 	float3 Color = GetColor(LightId);
-	OutRadiance = Color;
+	float Radius = GetRadius(LightId);
+	if (Radius == 0)
+	{
+		OutRadiance = Color;
+	}
+	else
+	{
+		// #dxr_todo: Merge Eval/Sample to avoid redundant computation
+		float SinThetaMax = Radius; // currently sin(Angle / 2)
+		float CosThetaMax = sqrt(1.0 - SinThetaMax * SinThetaMax);
+		OutRadiance = Color * 1.0 / (2 * PI * (1 - CosThetaMax));
+	}
 }
 
 void DirectionalLight_EstimateLight(
@@ -72,31 +83,27 @@ void DirectionalLight_SampleLight(
 )
 {
 	float3 N = normalize(GetNormal(LightId));
-	float3 dPdu = float3(1, 0, 0);
-	if (dot(N, dPdu) != 0)
+	float Radius = GetRadius(LightId);
+
+	if (Radius == 0)
 	{
-		dPdu = cross(N, dPdu);
+		// singular case
+		// #dxr_todo: should return +inf PDF to disable MIS
+		OutLightUV = N;
+		OutPdf = 1.0;
 	}
 	else
 	{
-		dPdu = cross(N, float3(0, 1, 0));
-	}
-	float3 dPdv = cross(dPdu, N);
-
-	float Radius = GetRadius(LightId);
-	float2 DiskUV = UniformSampleDiskConcentric(RandSample.yz) * Radius;
+		// soft case - build a frame, and sample the distant cone
+		float3 X, Y;
+		// #dxr_todo: the implementation of this function should probably follow: "Building an Orthonormal Basis, Revisited" JCGT 2017
+		GenerateCoordinateSystem(N, X, Y);
 
-	float3 Direction = N;
-	OutLightUV = Direction + dPdu * DiskUV.x + dPdv * DiskUV.y;
-	OutPdf = 1.0;
+		float SinThetaMax = Radius; // currently sin(Angle / 2)
+		float CosThetaMax = sqrt(1.0 - SinThetaMax * SinThetaMax);
 
-#if 0
-	// Pdf calculation for when directional light is visible to material rays
-	if (Radius > 0.0)
-	{
-		float Area = PI * Radius * Radius;
-		float SolidAngle = Area;
-		OutPdf /= SolidAngle;
+		float4 DirAndPdf = UniformSampleCone(RandSample.yz, CosThetaMax);
+		OutLightUV = DirAndPdf.x * X + DirAndPdf.y * Y + DirAndPdf.z * N;
+		OutPdf = DirAndPdf.w;
 	}
-#endif
-}
+}
\ No newline at end of file
