Small fixes to light sampling in path tracer
 - use better tangent basis in directional light
 - add missing early out in light sampling if all lights are hidden or invisible to current point
 - add spherical quad sampling option to rect light
 - workaround issue in rect light leading to clipped lighting around shadow terminator
 - fix lighting inconsistency with point lights of radius=0

diff --git a/Engine/Shaders/Private/PathTracing/Light/PathTracingDirectionalLight.ush b/Engine/Shaders/Private/PathTracing/Light/PathTracingDirectionalLight.ush
--- a/Engine/Shaders/Private/PathTracing/Light/PathTracingDirectionalLight.ush
+++ b/Engine/Shaders/Private/PathTracing/Light/PathTracingDirectionalLight.ush
@@ -95,15 +95,13 @@ void DirectionalLight_SampleLight(
 	else
 	{
 		// soft case - build a frame, and sample the distant cone
-		float3 X, Y;
-		// #dxr_todo: the implementation of this function should probably follow: "Building an Orthonormal Basis, Revisited" JCGT 2017
-		GenerateCoordinateSystem(N, X, Y);
+		const float3x3 ToWorld = GetTangentBasis(N);
 
 		float SinThetaMax = Radius; // currently sin(Angle / 2)
 		float CosThetaMax = sqrt(1.0 - SinThetaMax * SinThetaMax);
 
 		float4 DirAndPdf = UniformSampleCone(RandSample.yz, CosThetaMax);
-		OutLightUV = DirAndPdf.x * X + DirAndPdf.y * Y + DirAndPdf.z * N;
+		OutLightUV = mul(DirAndPdf.xyz, ToWorld);
 		OutPdf = DirAndPdf.w;
 	}
 }
\ No newline at end of file
diff --git a/Engine/Shaders/Private/PathTracing/Light/PathTracingLightSampling.ush b/Engine/Shaders/Private/PathTracing/Light/PathTracingLightSampling.ush
--- a/Engine/Shaders/Private/PathTracing/Light/PathTracingLightSampling.ush
+++ b/Engine/Shaders/Private/PathTracing/Light/PathTracingLightSampling.ush
@@ -86,13 +86,17 @@ bool SampleLightSelection(
 {
 	OutLightSelectionPdf = 0.0;
 
+	if (SceneLightsData.Count == 0)
+	{
+		return false;
+	}
+
 //#define UNIFORM_LIGHT_SELECTION
 #ifdef UNIFORM_LIGHT_SELECTION
 	OutLightId = RandSample * SceneLightsData.Count;
 	OutLightSelectionPdf = 1.0 / SceneLightsData.Count;
 	return true;
 #else
-	if (SceneLightsData.Count == 0) return false;
 
 	// Build irradiance estimate prefix sum
 	float IrradianceEstimate[RAY_TRACING_LIGHT_COUNT_MAXIMUM];
@@ -104,7 +108,11 @@ bool SampleLightSelection(
 			IrradianceEstimate[LightIndex] += IrradianceEstimate[LightIndex - 1];
 		}
 	}
-	// dxr_todo: Consider when IrradianceEstimate[SceneLightsData.Count] == 0
+	
+	if (IrradianceEstimate[SceneLightsData.Count - 1] == 0)
+	{
+		return false;
+	}
 
 	// CDF normalization
 	{
diff --git a/Engine/Shaders/Private/PathTracing/Light/PathTracingPointLight.ush b/Engine/Shaders/Private/PathTracing/Light/PathTracingPointLight.ush
--- a/Engine/Shaders/Private/PathTracing/Light/PathTracingPointLight.ush
+++ b/Engine/Shaders/Private/PathTracing/Light/PathTracingPointLight.ush
@@ -95,6 +95,6 @@ void PointLight_SampleLight(
 	else
 	{
 		OutLightUV = 0.0;
-		OutPdf = 1.0;
+		OutPdf = 1.0 / (4 * PI);
 	}
 }
diff --git a/Engine/Shaders/Private/PathTracing/Light/PathTracingRectLight.ush b/Engine/Shaders/Private/PathTracing/Light/PathTracingRectLight.ush
--- a/Engine/Shaders/Private/PathTracing/Light/PathTracingRectLight.ush
+++ b/Engine/Shaders/Private/PathTracing/Light/PathTracingRectLight.ush
@@ -18,12 +18,8 @@ void RectLight_GenerateLightRay(
 {
 	OutLightRay.Origin = Ray.Origin;
 	OutLightRay.TMin = 0.001;
-	// Get rect parameters to determine 3D direction.
-	float2 DeltaP = float2(GetWidth(LightId), GetHeight(LightId)) * LightUV.xy;
-	float3 LightDir = GetPosition(LightId) + DeltaP.x * GetdPdu(LightId) + DeltaP.y * GetdPdv(LightId) - Ray.Origin;
-
-	OutLightRay.Direction = normalize(LightDir);
-	OutLightRay.TMax = length(LightDir);
+	OutLightRay.Direction = normalize(LightUV);
+	OutLightRay.TMax = length(LightUV);
 }
 
 void RectLight_PdfLight(
@@ -72,12 +68,18 @@ void RectLight_EstimateLight(
 
 	// Approximate geometric term
 	float3 SurfaceNormal = Payload.WorldNormal;
+#if 0
 	float NoL = dot(SurfaceNormal, LightDirection);
 	NoL = ENABLE_TRANSMISSION && HasTransmission(LightId) ? abs(NoL) : saturate(NoL);
+#else
+	// Just use a trivial upper bound on cosine term 
+	// #dxr_todo: UE-72533 least intrusive way of fixing the artifact cause by using LightCenter
+	float NoL = 1.0;
+#endif
 	float Area = GetWidth(LightId) * GetHeight(LightId);
 
 	float LightPower = Luminance(GetColor(LightId));
-	float Falloff = ComputeAttenuationFalloff(length(LightDirection), GetAttenuation(LightId));
+	float Falloff = ComputeAttenuationFalloff(sqrt(LightDistanceSquared), GetAttenuation(LightId));
 	OutIrradiance = LightPower * Falloff * Area * NoL / LightDistanceSquared;
 }
 
@@ -97,50 +99,6 @@ void RectLight_SampleLight(
 	float3 LightdPdv = GetdPdv(LightId);
 	float LightWidth = GetWidth(LightId);
 	float LightHeight = GetHeight(LightId);
-	
-	// Compute the visible portion of the light from the current shading point
-	if (false)
-	{
-		float3 ToLight = LightPosition - WorldPosition;
-		FRect VisibleRect = GetRect(
-				ToLight,
-				-LightNormal,
-				LightdPdv,
-				LightWidth*0.5f,
-				LightHeight*0.5f,
-				GetRectLightBarnCosAngle(LightId),
-				GetRectLightBarnLength(LightId),
-				true);
-
-		LightPosition = LightPosition - LightdPdu * VisibleRect.Offset.x - LightdPdv * VisibleRect.Offset.y;	
-		LightWidth	  = VisibleRect.Extent.x * 2.0f;
-		LightHeight   = VisibleRect.Extent.y * 2.0f;
-	}
-
-	// Pick a random position
-	float2 UV = RandSample.yz - 0.5;
-	OutLightUV = float3(UV, 0);
-
-	// Map sample point to quad
-	float3 LightSamplePosition = LightPosition + LightdPdu * LightWidth * UV.x + LightdPdv * LightHeight * UV.y;
-	float3 LightDirection = LightSamplePosition - WorldPosition;
-	float LightDirectionLength = length(LightDirection);
-#if USE_ATTENUATION_TERM
-	float Epsilon = 1.0e-5;
-	if (ComputeAttenuationFalloff(LightDirectionLength, GetAttenuation(LightId)) < Epsilon)
-	{
-		OutPdf = 0.0;
-		return;
-	}
-#endif
-	LightDirection *= rcp(LightDirectionLength);
-
-	// Light-normal culling
-	if (dot(-LightDirection, LightNormal) <= 0.0)
-	{
-		OutPdf = 0.0;
-		return;
-	}
 
 	// Define rectangle and compute solid angle
 	FRect Rect;
@@ -150,5 +108,6 @@ void RectLight_SampleLight(
 	Rect.Axis[2] = cross(LightdPdu, LightdPdv);
 	Rect.Extent = float2(LightWidth, LightHeight) * 0.5f;
 	FSphericalRect SphericalRect = BuildSphericalRect(Rect);
+	OutLightUV = UniformSampleSphericalRect(RandSample.yz, SphericalRect);
 	OutPdf = isfinite(SphericalRect.SolidAngle) ? 1.0 / SphericalRect.SolidAngle : 0.0;
 }
\ No newline at end of file
