Fixed -d3ddebug warning in the Global Distance Field update

diff --git a/Engine/Shaders/Private/GlobalDistanceField.usf b/Engine/Shaders/Private/GlobalDistanceField.usf
--- a/Engine/Shaders/Private/GlobalDistanceField.usf
+++ b/Engine/Shaders/Private/GlobalDistanceField.usf
@@ -47,7 +47,7 @@ RWStructuredBuffer<uint> RWCullGridObjectArray;
 uint3 UpdateRegionSize;
 StructuredBuffer<uint> CullGridObjectHeader;
 StructuredBuffer<uint> CullGridObjectArray;
-RWTexture3D<float> RWPageAtlasTexture;
+RWTexture3D<UNORM float> RWPageAtlasTexture;
 
 #ifndef COMPOSITE_THREADGROUP_SIZEX
 #define COMPOSITE_THREADGROUP_SIZEX 1
