Fix ray traced shadow artifact on point lights with tiny radius

The term inside the sqrt for TMax could become negative numerically even though it should always be positive. Added a clamp to ensure robustness.

I also re-arranged the math slightly to be more numerically stable.

diff --git a/Engine/Shaders/Private/RayTracing/RayTracingSphereLight.ush b/Engine/Shaders/Private/RayTracing/RayTracingSphereLight.ush
--- a/Engine/Shaders/Private/RayTracing/RayTracingSphereLight.ush
+++ b/Engine/Shaders/Private/RayTracing/RayTracingSphereLight.ush
@@ -72,7 +72,8 @@ bool GenerateSphereLightOcclusionRayWithSolidAngleSampling(
 	// float3 RayDirectionTangent = UniformSampleCone(RandSample, CosThetaMax).xyz;
 	// RayPdf = UniformSampleCone(RandSample, CosThetaMax).w;
 	float CosTheta = lerp(1.0, CosThetaMax, RandSample.x);
-	float SinTheta = sqrt(1.0 - CosTheta * CosTheta);
+	float SinTheta2 = 1.0 - CosTheta * CosTheta;
+	float SinTheta = sqrt(SinTheta2);
 	float Phi = RandSample.y * 2.0 * PI;
 	float3 RayDirectionTangent = float3(cos(Phi) * SinTheta, sin(Phi) * SinTheta, CosTheta);
 
@@ -84,7 +85,7 @@ bool GenerateSphereLightOcclusionRayWithSolidAngleSampling(
 	
 	RayTMin = 0.0;
 	// Clip length to closest intersection with the sphere
-	RayTMax = RayLength * CosTheta - sqrt(Radius2 - RayLength2 * SinTheta * SinTheta);
+	RayTMax = RayLength * (CosTheta - sqrt(max(SinThetaMax2 - SinTheta2, 0.0)));
 	RayPdf = 1.0 / (2.0 * PI * (1 - CosThetaMax));
 	return true;
 }
\ No newline at end of file
