Fix shading of point inside the barndoor volume of RectLights
Simplify logic of the GetRect function to make it a bit easier to read

diff --git a/Engine/Shaders/Private/RectLight.ush b/Engine/Shaders/Private/RectLight.ush
--- a/Engine/Shaders/Private/RectLight.ush
+++ b/Engine/Shaders/Private/RectLight.ush
@@ -669,55 +669,48 @@ FRect GetRect(
 	// Only compute the occluded rect if the barn door has an angle less 
 	// than 88 degrees
 	if (bComputeVisibleRect && LightDataRectLightBarnCosAngle > 0.035f)
-	{	
-		const float3 LightNormal	=-Rect.Axis[2];
-		const float3 LightdPdv		=-Rect.Axis[1];
-		const float3 LightdPdu		=-Rect.Axis[0];
-		const float2 LightExtent	= float2(LightDataSourceRadius, LightDataSourceLength);
-		const float BarnLength		= LightDataRectLightBarnLength;
+	{
+		const float3 LightdPdv = -Rect.Axis[1];
+		const float3 LightdPdu = -Rect.Axis[0];
+		const float2 LightExtent = float2(LightDataSourceRadius, LightDataSourceLength);
+		const float BarnLength = LightDataRectLightBarnLength;
 	
 		// Project shading point S into light space
-		const float3 S_World = -ToLight;
-		float3 S_Light = float3(dot(S_World, LightdPdu), dot(S_World, LightdPdv), dot(S_World, LightNormal));
-		
+		float3 S_Light = mul(Rect.Axis, ToLight);
+
 		// Compute barn door projection (D_B). Clamp the projection to the shading point if it is closer 
 		// to the light than the actual barn door
 		// Theta is the angle between the Z axis and the barn door
 		const float CosTheta = LightDataRectLightBarnCosAngle;
-		const float SinTheta = sqrt(1 - max(0, CosTheta * CosTheta));
+		const float SinTheta = sqrt(1 - CosTheta * CosTheta);
 		const float BarnDepth = min(S_Light.z, CosTheta * BarnLength);
 		const float S_ratio = BarnDepth / (CosTheta * BarnLength);
-		const float D_B	= SinTheta * BarnLength * S_ratio;
+		const float D_B = SinTheta * BarnLength * S_ratio;
 		
 		// Clamp shading point onto the closest edge, if it is inside the rect light
 		const float2 SignS = sign(S_Light.xy);
-		S_Light.xy = float2(SignS.x, SignS.y) * max(abs(S_Light.xy), LightExtent + D_B.xx);
+		S_Light.xy = SignS * max(abs(S_Light.xy), LightExtent + D_B.xx);
 		
 		// Compute the closest rect lignt corner, offset by the barn door size
-		const float3 C = float3(SignS.x * (LightExtent.x + D_B), SignS.y * (LightExtent.y + D_B), BarnDepth);
+		const float3 C = float3(SignS * (LightExtent + D_B.xx), BarnDepth);
 			
 		// Compute projected distance (D_S) of barn door onto the rect light
 		// Eta is the angle between the Z axis and the direction vector (S-C)
-		const float3 SProj  = S_Light - C;
-		const float  CosEta = max(SProj.z, 0);
+		const float3 SProj = S_Light - C;
+		const float CosEta = max(SProj.z, 0.001f);
 		const float2 SinEta = abs(SProj.xy);
-		const float2 TanEta	= CosEta > 0.001f ? abs(SProj.xy)/CosEta : 10000.f;
+		const float2 TanEta = abs(SProj.xy) / CosEta;
 		const float2 D_S = BarnDepth * TanEta;
 
 		// Equivalent to (e.g., X axis):
 		//  if (SignS.x < 0) MinMaxX.x += D_S - D_B;
 		//  if (SignS.x > 0) MinMaxX.y -= D_S - D_B;
-		float2 MinMaxX = float2(-LightExtent.x, LightExtent.x);
-		float2 MinMaxY = float2(-LightExtent.y, LightExtent.y);
-		MinMaxX += (D_S.x - D_B) * float2(max(0, -SignS.x), -max(0, SignS.x));
-		MinMaxY += (D_S.y - D_B) * float2(max(0, -SignS.y), -max(0, SignS.y));
-		MinMaxX = clamp(MinMaxX, -LightExtent.x, LightExtent.x);
-		MinMaxY = clamp(MinMaxY, -LightExtent.y, LightExtent.y);
-		
-		const float2 RectOffset = 0.5f * float2(MinMaxX.x + MinMaxX.y, MinMaxY.x + MinMaxY.y);
+		const float2 MinXY = clamp(-LightExtent + (D_S - D_B.xx) * max(0, -SignS), -LightExtent, LightExtent);
+		const float2 MaxXY = clamp( LightExtent - (D_S - D_B.xx) * max(0,  SignS), -LightExtent, LightExtent);
+		const float2 RectOffset = 0.5f * (MinXY + MaxXY);
 
-		Rect.Extent = 0.5f * float2(MinMaxX.y - MinMaxX.x, MinMaxY.y - MinMaxY.x);
-		Rect.Origin = Rect.Origin +  LightdPdu * RectOffset.x + LightdPdv * RectOffset.y;
+		Rect.Extent = 0.5f * (MaxXY - MinXY);
+		Rect.Origin = Rect.Origin + LightdPdu * RectOffset.x + LightdPdv * RectOffset.y;
 		Rect.Offset = -RectOffset;
 		Rect.FullExtent = LightExtent;
 	}
