Missed a file add on last commit.

diff --git a/Engine/Shaders/Private/ComputeKernel.usf b/Engine/Shaders/Private/ComputeKernel.usf
new file mode 100644
--- /dev/null
+++ b/Engine/Shaders/Private/ComputeKernel.usf
@@ -0,0 +1,8 @@
+#include "Common.ush"
+
+[numthreads(32, 1, 1)]
+void Main(
+	uint DispatchIndex : SV_DispatchThreadID,
+	uint GroupIndex : SV_GroupIndex)
+{
+}
\ No newline at end of file
