Fix syntax error in CalculateAnisotropyTangent.  Same fix as 14387590 from main ue4 branch. #rb

diff --git a/Engine/Shaders/Private/MaterialTemplate.ush b/Engine/Shaders/Private/MaterialTemplate.ush
--- a/Engine/Shaders/Private/MaterialTemplate.ush
+++ b/Engine/Shaders/Private/MaterialTemplate.ush
@@ -2288,7 +2288,7 @@ float3 CalculateAnisotropyTangent(FMaterialPixelParameters Parameters, FPixelMat
 #if CLEAR_COAT_BOTTOM_NORMAL && (NUM_MATERIAL_OUTPUTS_CLEARCOATBOTTOMNORMAL > 0)
 	Normal =  ClearCoatBottomNormal0(Parameters);
 	#if MATERIAL_TANGENTSPACENORMAL
-		Normal = TransformTangentVectorToWorld(Parameters.TangentToWorld, Normal));
+		Normal = TransformTangentVectorToWorld(Parameters.TangentToWorld, Normal);
 	#endif
 #endif
 
