Improved Strata BSDF packing.

diff --git a/Engine/Shaders/Private/Strata/Strata.ush b/Engine/Shaders/Private/Strata/Strata.ush
--- a/Engine/Shaders/Private/Strata/Strata.ush
+++ b/Engine/Shaders/Private/Strata/Strata.ush
@@ -13,7 +13,6 @@
 //	- Have a Disney model that is convernient to be super packed and be able to represent current mateiral?
 //	- Shared normal? Up to N normals? On CPU, inputs have a hash that could be use to decuplicate normals
 //	- Transmittance lob modelisation (GGX statistics?) and blocking surface such as FStrataDiffuse?
-//  - 11-11-10 format
 //	- fp16 on gen5
 
 
@@ -62,10 +61,16 @@ struct FStrataVolume
 
 
 
+// The largest Strata BSDF is FStrataConductor having as size of 24 bytes.
+// So we will have 
+//		STRATA_MAX_TOTAL_BSDF * sizeof(FStrataConductor) +				// 4 * 24
+//		STRATA_MAX_TOTAL_BSDF * (sizeof(BSDFType) + sizeof(Weight))		// 4 * 8
+//		sizeof(BSDFCount)												// UINT
+//		= 4 * 24 + 4 * 8 + 1 = 96 + 32 + 1 = 129						// See Strata.cpp
 struct FStrataBSDF
 {
 	int					BSDFType;
-	float3				Weight;				// Setup by users and mixing operations
+	float3				Weight;
 
 	// Even thought the compiler should optimise, it would be nice to have a union with reinterpretation of the register memory.
 	FStrataDiffuse		Diffuse;
@@ -264,6 +269,55 @@ FStrataData StrataVerticalLayering(FStrataData Top, FStrataData Base)
 ///////////////////////////////////////////////////////////////////////////////
 // BSDFs packing/unpacking
 
+uint PackR11G11B10F(float3 rgb)
+{
+	uint r = (f32tof16(rgb.r) << 17) & 0xFFE00000;
+	uint g = (f32tof16(rgb.g) << 6)  & 0x001FFC00;
+	uint b = (f32tof16(rgb.b) >> 5)  & 0x000003FF;
+	return r | g | b;
+}
+
+float3 UnpackR11G11B10F(uint rgb)
+{
+	float r = f16tof32((rgb >> 17) & 0x7FF0);
+	float g = f16tof32((rgb >> 6)  & 0x7FF0);
+	float b = f16tof32((rgb << 5)  & 0x7FE0);
+	return float3(r, g, b);
+}
+
+uint2 PackR16G16B16A16F(float4 rgba)
+{
+	uint A = (f32tof16(rgba.r) << 16) | f32tof16(rgba.g);
+	uint B = (f32tof16(rgba.b) << 16) | f32tof16(rgba.a);
+	return uint2 (A, B);
+}
+
+float4 UnpackR16G16B16A16F(uint2 AB)
+{
+	float r = f16tof32(AB.x >> 16);
+	float g = f16tof32(AB.x);
+	float b = f16tof32(AB.y >> 16);
+	float a = f16tof32(AB.y);
+	return float4(r, g, b, a);
+}
+
+uint PackAlbedoLinearToGamma2(float3 rgb)
+{
+	uint r = (uint(sqrt(saturate(rgb.r)) * 255.0f) << 16);
+	uint g = (uint(sqrt(saturate(rgb.g)) * 255.0f) << 8);
+	uint b = (uint(sqrt(saturate(rgb.b)) * 255.0f) << 0);
+	return r | g | b;
+}
+
+float3 UnpackAlbedoGamma2ToLinear(uint rgb)
+{
+	float r = float((rgb >> 16) & 0x000000FF);
+	float g = float((rgb >> 8)  & 0x000000FF);
+	float b = float((rgb >> 0)  & 0x000000FF);
+	float3 Gamma2 = float3(r, g, b) * (1.0f / 255.0f);
+	return Gamma2 * Gamma2;	// Simple gamma 2.0
+}
+
 
 
 #define STRATA_STORE_UINT1(x) StrataBuffer.Store (ByteOffset, asuint(x)); ByteOffset += 4;
@@ -281,6 +335,35 @@ FStrataData StrataVerticalLayering(FStrataData Top, FStrataData Base)
 #define STRATA_LOAD_FLOAT3(x) x = asfloat(StrataBuffer.Load3(ByteOffset)); ByteOffset += 12;
 #define STRATA_LOAD_FLOAT4(x) x = asfloat(StrataBuffer.Load4(ByteOffset)); ByteOffset += 16;
 
+uint  StrataLoadUint1(ByteAddressBuffer StrataBuffer, inout uint ByteOffset) { uint  x = StrataBuffer.Load (ByteOffset); ByteOffset += 4; return x; }
+uint2 StrataLoadUint2(ByteAddressBuffer StrataBuffer, inout uint ByteOffset) { uint2 x = StrataBuffer.Load2(ByteOffset); ByteOffset += 8; return x; }
+uint3 StrataLoadUint3(ByteAddressBuffer StrataBuffer, inout uint ByteOffset) { uint3 x = StrataBuffer.Load3(ByteOffset); ByteOffset +=12; return x; }
+uint4 StrataLoadUint4(ByteAddressBuffer StrataBuffer, inout uint ByteOffset) { uint4 x = StrataBuffer.Load4(ByteOffset); ByteOffset +=16; return x; }
+
+void StrataLoad_R16G16B16F_A16F(ByteAddressBuffer StrataBuffer, inout uint ByteOffset, inout float3 A, inout float B)
+{
+	float4 temp = UnpackR16G16B16A16F(StrataLoadUint2(StrataBuffer, ByteOffset));
+	A = temp.xyz;
+	B = temp.a;
+}
+
+void StrataLoad_R11G11B10F(ByteAddressBuffer StrataBuffer, inout uint ByteOffset, inout float3 A)
+{
+	A = UnpackR11G11B10F(StrataLoadUint1(StrataBuffer, ByteOffset));
+}
+
+void StrataLoad_R11G11_B10F(ByteAddressBuffer StrataBuffer, inout uint ByteOffset, inout float2 A, inout float B)
+{
+	float3 temp = UnpackR11G11B10F(StrataLoadUint1(StrataBuffer, ByteOffset));
+	A = temp.xy;
+	B = temp.z;
+}
+
+void StrataLoad_AlbedoGamma2ToLinear(ByteAddressBuffer StrataBuffer, inout uint ByteOffset, inout float3 A)
+{
+	A = UnpackAlbedoGamma2ToLinear(StrataLoadUint1(StrataBuffer, ByteOffset));
+}
+
 
 
 void PackStrataOut(
@@ -407,40 +490,39 @@ void PackStrataOut(
 		{
 			FStrataBSDF BSDF = Layer.BSDFs[i];
 
-			STRATA_STORE_UINT1(BSDF.BSDFType);
-			STRATA_STORE_UINT3(BSDF.Weight);
+			STRATA_STORE_UINT2(PackR16G16B16A16F(float4(BSDF.Weight, BSDF.BSDFType)));
 
 			switch (BSDF.BSDFType)
 			{
 			case STRATA_BSDF_TYPE_DIFFUSE_ON:
 			case STRATA_BSDF_TYPE_DIFFUSE_CHAN:
 			{
-				STRATA_STORE_UINT3(BSDF.Diffuse.Albedo);
-				STRATA_STORE_UINT1(BSDF.Diffuse.Roughness);
-				STRATA_STORE_UINT3(BSDF.Diffuse.Normal);
+				STRATA_STORE_UINT1(PackAlbedoLinearToGamma2(BSDF.Diffuse.Albedo));
+				STRATA_STORE_UINT2(PackR16G16B16A16F(float4(BSDF.Diffuse.Normal, BSDF.Diffuse.Roughness)));
+				// 12 bytes
 			}
 			break;
 			case STRATA_BSDF_TYPE_DIELECTRIC: //
 			{
-				STRATA_STORE_UINT1(BSDF.Dielectric.IOR);
-				STRATA_STORE_UINT3(BSDF.Dielectric.Tint);
-				STRATA_STORE_UINT2(BSDF.Dielectric.Roughness);
-				STRATA_STORE_UINT3(BSDF.Dielectric.Normal);
+				STRATA_STORE_UINT1(PackR11G11B10F(BSDF.Dielectric.Tint));
+				STRATA_STORE_UINT1(PackR11G11B10F(float3(BSDF.Dielectric.Roughness, BSDF.Dielectric.IOR)));
+				STRATA_STORE_UINT2(PackR16G16B16A16F(float4(BSDF.Dielectric.Normal, 0.0f)));
+				// 16 bytes
 			}
 			break;
 			case STRATA_BSDF_TYPE_CONDUCTOR:
 			{
-				STRATA_STORE_UINT3(BSDF.Conductor.IOR);
-				STRATA_STORE_UINT3(BSDF.Conductor.Extinction);
-				STRATA_STORE_UINT2(BSDF.Conductor.Roughness);
-				STRATA_STORE_UINT3(BSDF.Conductor.Normal);
+				STRATA_STORE_UINT2(PackR16G16B16A16F(float4(BSDF.Conductor.IOR, BSDF.Conductor.Roughness.x)));
+				STRATA_STORE_UINT2(PackR16G16B16A16F(float4(BSDF.Conductor.Extinction, BSDF.Conductor.Roughness.y)));
+				STRATA_STORE_UINT2(PackR16G16B16A16F(float4(BSDF.Conductor.Normal, 0.0f)));
+				// 24 bytes
 			}
 			break;
 			case STRATA_BSDF_TYPE_VOLUME:
 			{
-				STRATA_STORE_UINT3(BSDF.Volume.Absorption);
-				STRATA_STORE_UINT3(BSDF.Volume.Scattering);
-				STRATA_STORE_UINT1(BSDF.Volume.Anisotropy);
+				STRATA_STORE_UINT1(PackR11G11B10F(BSDF.Volume.Absorption));
+				STRATA_STORE_UINT2(PackR16G16B16A16F(float4(BSDF.Volume.Scattering, BSDF.Volume.Anisotropy)));
+				// 12 bytes
 			}
 			break;
 			}
@@ -469,43 +551,44 @@ FStrataPixelBSDFs UnpackStrataIn(ByteAddressBuffer StrataBuffer, uint ByteOffset
 	STRATA_LOAD_UINT1(PixelBSDFs.BSDFCount);
 	PixelBSDFs.BSDFCount = min(PixelBSDFs.BSDFCount, STRATA_MAX_TOTAL_BSDF);	// To make the compiler happy with unrolling
 
+	float DummyFloat = 0.0f;
+
 	UNROLL
 	for (int i = 0; i < PixelBSDFs.BSDFCount; ++i)
 	{
-		STRATA_LOAD_UINT1(PixelBSDFs.BSDFs[i].BSDFType);
-		STRATA_LOAD_FLOAT3(PixelBSDFs.BSDFs[i].Weight);
+		{
+			float4 temp4  = UnpackR16G16B16A16F(StrataLoadUint2(StrataBuffer, ByteOffset));
+			PixelBSDFs.BSDFs[i].Weight   = temp4.xyz;
+			PixelBSDFs.BSDFs[i].BSDFType = int(temp4.a);
+		}
 
 		switch (PixelBSDFs.BSDFs[i].BSDFType)
 		{
 		case STRATA_BSDF_TYPE_DIFFUSE_ON:
 		case STRATA_BSDF_TYPE_DIFFUSE_CHAN:
 		{
-			STRATA_LOAD_FLOAT3(PixelBSDFs.BSDFs[i].Diffuse.Albedo);
-			STRATA_LOAD_FLOAT1(PixelBSDFs.BSDFs[i].Diffuse.Roughness);
-			STRATA_LOAD_FLOAT3(PixelBSDFs.BSDFs[i].Diffuse.Normal);
+			StrataLoad_AlbedoGamma2ToLinear	(StrataBuffer, ByteOffset, PixelBSDFs.BSDFs[i].Diffuse.Albedo);
+			StrataLoad_R16G16B16F_A16F		(StrataBuffer, ByteOffset, PixelBSDFs.BSDFs[i].Diffuse.Normal,			PixelBSDFs.BSDFs[i].Diffuse.Roughness);
 		}
 		break;
 		case STRATA_BSDF_TYPE_DIELECTRIC:
 		{
-			STRATA_LOAD_FLOAT1(PixelBSDFs.BSDFs[i].Dielectric.IOR);
-			STRATA_LOAD_FLOAT3(PixelBSDFs.BSDFs[i].Dielectric.Tint);
-			STRATA_LOAD_FLOAT2(PixelBSDFs.BSDFs[i].Dielectric.Roughness);
-			STRATA_LOAD_FLOAT3(PixelBSDFs.BSDFs[i].Dielectric.Normal);
+			StrataLoad_R11G11B10F			(StrataBuffer, ByteOffset, PixelBSDFs.BSDFs[i].Dielectric.Tint);
+			StrataLoad_R11G11_B10F			(StrataBuffer, ByteOffset, PixelBSDFs.BSDFs[i].Dielectric.Roughness,	PixelBSDFs.BSDFs[i].Dielectric.IOR);
+			StrataLoad_R16G16B16F_A16F		(StrataBuffer, ByteOffset, PixelBSDFs.BSDFs[i].Dielectric.Normal,		DummyFloat);
 		}
 		break;
 		case STRATA_BSDF_TYPE_CONDUCTOR:
 		{
-			STRATA_LOAD_FLOAT3(PixelBSDFs.BSDFs[i].Conductor.IOR);
-			STRATA_LOAD_FLOAT3(PixelBSDFs.BSDFs[i].Conductor.Extinction);
-			STRATA_LOAD_FLOAT2(PixelBSDFs.BSDFs[i].Conductor.Roughness);
-			STRATA_LOAD_FLOAT3(PixelBSDFs.BSDFs[i].Conductor.Normal);
+			StrataLoad_R16G16B16F_A16F		(StrataBuffer, ByteOffset, PixelBSDFs.BSDFs[i].Conductor.IOR,			PixelBSDFs.BSDFs[i].Conductor.Roughness.x);
+			StrataLoad_R16G16B16F_A16F		(StrataBuffer, ByteOffset, PixelBSDFs.BSDFs[i].Conductor.Extinction,	PixelBSDFs.BSDFs[i].Conductor.Roughness.y);
+			StrataLoad_R16G16B16F_A16F		(StrataBuffer, ByteOffset, PixelBSDFs.BSDFs[i].Conductor.Normal,		DummyFloat);
 		}
 		break;
 		case STRATA_BSDF_TYPE_VOLUME:
 		{
-			STRATA_LOAD_FLOAT3(PixelBSDFs.BSDFs[i].Volume.Absorption);
-			STRATA_LOAD_FLOAT3(PixelBSDFs.BSDFs[i].Volume.Scattering);
-			STRATA_LOAD_FLOAT1(PixelBSDFs.BSDFs[i].Volume.Anisotropy);
+			StrataLoad_R11G11B10F			(StrataBuffer, ByteOffset, PixelBSDFs.BSDFs[i].Volume.Absorption);
+			StrataLoad_R16G16B16F_A16F		(StrataBuffer, ByteOffset, PixelBSDFs.BSDFs[i].Volume.Scattering,		PixelBSDFs.BSDFs[i].Volume.Anisotropy);
 		}
 		break;
 		}
diff --git a/Engine/Source/Runtime/Renderer/Private/Strata/Strata.cpp b/Engine/Source/Runtime/Renderer/Private/Strata/Strata.cpp
--- a/Engine/Source/Runtime/Renderer/Private/Strata/Strata.cpp
+++ b/Engine/Source/Runtime/Renderer/Private/Strata/Strata.cpp
@@ -47,12 +47,13 @@ void InitialiseStrataFrameSceneData(FSceneRenderer& SceneRenderer, FRDGBuilder&
 		ResolutionY = BufferSizeXY.Y;
 
 		// Previous GBuffer when complete was 28bytes
-		StrataData.MaxBytesPerPixel = 256;
-
+		// check out Strata.ush to see how this is computed
+		const uint32 MaterialConservativeByteCountPerPixel = 129u;
+		StrataData.MaxBytesPerPixel = FMath::DivideAndRoundUp(MaterialConservativeByteCountPerPixel, 4u) * 4u;
 	}
 	else
 	{
-		StrataData.MaxBytesPerPixel = 1;
+		StrataData.MaxBytesPerPixel = 4u;
 	}
 
 	FRDGTextureRef MaterialLobesTexture = GraphBuilder.CreateTexture(
