Path tracer lighting improvements:
 - Remove different factors of 4*PI from point and spot lights on the cpp side so the logic between these lights becomes the same
 - Simplify DirectionalLight code some more
 - Implement soft shadows for point lights with solid angle sampling
 - Implement soft shadows for spots using the point light logic
 - Implement cone angle falloff for spot lights
 - Implement barn door support on rect lights
 - Implement approximated projected quad sampling for rect lights

diff --git a/Engine/Shaders/Private/PathTracing/Light/PathTracingDirectionalLight.ush b/Engine/Shaders/Private/PathTracing/Light/PathTracingDirectionalLight.ush
--- a/Engine/Shaders/Private/PathTracing/Light/PathTracingDirectionalLight.ush
+++ b/Engine/Shaders/Private/PathTracing/Light/PathTracingDirectionalLight.ush
@@ -95,13 +95,11 @@ void DirectionalLight_SampleLight(
 	else
 	{
 		// soft case - build a frame, and sample the distant cone
-		const float3x3 ToWorld = GetTangentBasis(N);
-
 		float SinThetaMax = Radius; // currently sin(Angle / 2)
 		float CosThetaMax = sqrt(1.0 - SinThetaMax * SinThetaMax);
 
 		float4 DirAndPdf = UniformSampleCone(RandSample.yz, CosThetaMax);
-		OutLightUV = mul(DirAndPdf.xyz, ToWorld);
+		OutLightUV = TangentToWorld(DirAndPdf.xyz, N);
 		OutPdf = DirAndPdf.w;
 	}
 }
\ No newline at end of file
diff --git a/Engine/Shaders/Private/PathTracing/Light/PathTracingLightCommon.ush b/Engine/Shaders/Private/PathTracing/Light/PathTracingLightCommon.ush
--- a/Engine/Shaders/Private/PathTracing/Light/PathTracingLightCommon.ush
+++ b/Engine/Shaders/Private/PathTracing/Light/PathTracingLightCommon.ush
@@ -87,9 +87,9 @@ float GetDepth(int LightId)
 	return SceneLightsData.Dimensions[LightId].z;
 }
 
-float GetCosConeAngle(int LightId)
+float2 GetCosConeAngles(int LightId)
 {
-	return SceneLightsData.Dimensions[LightId].x;
+	return SceneLightsData.Dimensions[LightId].xy;
 }
 
 float3 GetColor(int LightId)
diff --git a/Engine/Shaders/Private/PathTracing/Light/PathTracingPointLight.ush b/Engine/Shaders/Private/PathTracing/Light/PathTracingPointLight.ush
--- a/Engine/Shaders/Private/PathTracing/Light/PathTracingPointLight.ush
+++ b/Engine/Shaders/Private/PathTracing/Light/PathTracingPointLight.ush
@@ -15,13 +15,24 @@ void PointLight_GenerateLightRay(
 	out RayDesc OutLightRay
 )
 {
+	float3 LightDirection = GetPosition(LightId) - Ray.Origin;
+
 	OutLightRay.Origin = Ray.Origin;
 	OutLightRay.TMin = 0.001;
+	OutLightRay.Direction = TangentToWorld(LightUV, normalize(LightDirection));
+
+	// Compute distance to sphere edge
 	float Radius = GetRadius(LightId);
-	float3 LightDir = GetPosition(LightId) + LightUV * Radius - Ray.Origin;
-	float Length = length(LightDir);
-	OutLightRay.Direction = LightDir / Length;
-	OutLightRay.TMax = Length;
+	float Radius2 = Pow2(Radius2);
+
+	float RayLength2 = dot(LightDirection, LightDirection);
+	float RayLength = sqrt(RayLength2);
+	float SinThetaMax2 = Radius2 / RayLength2;
+
+	float CosTheta = LightUV.z;
+	float SinTheta2 = 1.0 - CosTheta * CosTheta;
+
+	OutLightRay.TMax = RayLength * (CosTheta - sqrt(max(SinThetaMax2 - SinTheta2, 0.0)));
 }
 
 void PointLight_PdfLight(
@@ -41,12 +52,25 @@ void PointLight_EvalLight(
 	out float3 OutRadiance
 )
 {
+	float3 LightDirection = GetPosition(LightId) - Ray.Origin;
+	float LightDistanceSquared = dot(LightDirection, LightDirection);
+
 	float3 Color = GetColor(LightId);
 	float Radius = GetRadius(LightId);
-	float3 LightDir = GetPosition(LightId) + LightUV * Radius - Ray.Origin;
-	float DistanceSquared = dot(LightDir, LightDir);
 	float Falloff = ComputeAttenuationFalloff(Ray.TMax, GetAttenuation(LightId));
-	OutRadiance = Color * Falloff / DistanceSquared;
+	OutRadiance = Color * Falloff;
+#if 0
+	// doing the "right" thing leads to unstable math for small radius lights
+	// #dxr_todo: Revisit this when implementing MIS for all lights
+	if (Radius == 0)
+		OutRadiance *= 1 / LightDistanceSquared;
+	else
+		OutRadiance *= 1 / (PI * Radius * Radius);
+#else
+	// simplified expression for singular and soft lights - just like in the raster path
+	// this is more numerically stable and means we can ignore the PDF for now
+	OutRadiance *= 1.0 / LightDistanceSquared;
+#endif
 }
 
 void PointLight_EstimateLight(
@@ -88,13 +112,31 @@ void PointLight_SampleLight(
 	float Radius = GetRadius(LightId);
 	if (Radius > 0.0)
 	{
-		float4 Result = UniformSampleSphere(RandSample.yz);
-		OutLightUV = Result.xyz;
-		OutPdf = Result.w;
+		// Soft case - sample the solid angle subtended by the sphere
+		float3 LightDirection = GetPosition(LightId) - Payload.WorldPos;
+		float LightDistanceSquared = dot(LightDirection, LightDirection);
+		float Radius2 = Pow2(Radius);
+		
+		// outside the light
+		float SinThetaMax2 = Radius2 / LightDistanceSquared;
+		// #dxr_todo: find a better way of handling the region inside the light
+		float CosThetaMax = sqrt(max(1.0 - SinThetaMax2, 0.0));
+		float4 DirAndPdf = UniformSampleCone(RandSample.yz, CosThetaMax);
+
+		OutLightUV = DirAndPdf.xyz;
+#if 0
+		OutPdf = DirAndPdf.w;
+#else
+		// doing the "right" thing leads to unstable math for small radius lights
+		// #dxr_todo: Revisit this when implementing MIS for all lights
+		OutPdf = 1.0;
+#endif
 	}
 	else
 	{
-		OutLightUV = 0.0;
-		OutPdf = 1.0 / (4 * PI);
+		// Singular case
+		OutLightUV = float3(0, 0, 1);
+		// #dxr_todo: this isn't really part of the pdf, the pdf should be +inf for the singular case
+		OutPdf = 1.0;
 	}
 }
diff --git a/Engine/Shaders/Private/PathTracing/Light/PathTracingRectLight.ush b/Engine/Shaders/Private/PathTracing/Light/PathTracingRectLight.ush
--- a/Engine/Shaders/Private/PathTracing/Light/PathTracingRectLight.ush
+++ b/Engine/Shaders/Private/PathTracing/Light/PathTracingRectLight.ush
@@ -83,6 +83,49 @@ void RectLight_EstimateLight(
 	OutIrradiance = LightPower * Falloff * Area * NoL / LightDistanceSquared;
 }
 
+float3 BilinearQuadWarp(float2 uv, float W00, float W01, float W10, float W11) {
+	// "Practical Product Sampling by Fitting and Composing Warps" - EGSR 2020
+	// https://casual-effects.com/research/Hart2020Sampling/index.html
+	// https://www.shadertoy.com/view/wljyDz
+
+	float a = lerp(W00, W01, .5);
+	float b = lerp(W10, W11, .5);
+	float u = a == b ? uv.x : (sqrt(lerp(a * a, b * b, uv.x)) - a) / (b - a);
+	float c = lerp(W00, W10, u);
+	float d = lerp(W01, W11, u);
+	float v = c == d ? uv.y : (sqrt(lerp(c * c, d * d, uv.y)) - c) / (d - c);
+	float area = lerp(a, b, .5);
+	float pdf = lerp(c, d, v) / area;
+	return float3(u, v, pdf);
+}
+
+// #dxr_todo: This method should be moved to RectLight.ush
+float4 SampleApproxProjectionSphericalRect(float2 Rand, FSphericalRect SphericalRect, float3 WorldNormal)
+{
+	// Construct the quad vertices
+	float3 S00 = float3(SphericalRect.x0, SphericalRect.y0, SphericalRect.z0);
+	float3 S10 = float3(SphericalRect.x1, SphericalRect.y0, SphericalRect.z0);
+	float3 S11 = float3(SphericalRect.x1, SphericalRect.y1, SphericalRect.z0);
+	float3 S01 = float3(SphericalRect.x0, SphericalRect.y1, SphericalRect.z0);
+	
+	// Compute the cosine of the normal against each corner of the quad
+	float3 LightSpaceNormal = mul(SphericalRect.Axis, WorldNormal);
+	float W00 = saturate(dot(normalize(S00), LightSpaceNormal));
+	float W10 = saturate(dot(normalize(S10), LightSpaceNormal));
+	float W11 = saturate(dot(normalize(S11), LightSpaceNormal));
+	float W01 = saturate(dot(normalize(S01), LightSpaceNormal));
+	
+	// Warp a 2D sample according to the cosine at each corner of the quad
+	float3 WarpedRand = BilinearQuadWarp(Rand, W00, W01, W10, W11);
+
+	// Proceed with uniform sampling of the spherical rectangle
+	float3 Direction = UniformSampleSphericalRect(WarpedRand.xy, SphericalRect);
+	float OutPdf = isfinite(SphericalRect.SolidAngle) ? 1.0 / SphericalRect.SolidAngle : 0.0;
+
+	// Return Direction vector and adjusted PDF
+	return float4(Direction, OutPdf * WarpedRand.z);
+}
+
 void RectLight_SampleLight(
 	RayDesc Ray,
 	FMaterialClosestHitPayload Payload,
@@ -99,15 +142,29 @@ void RectLight_SampleLight(
 	float3 LightdPdv = GetdPdv(LightId);
 	float LightWidth = GetWidth(LightId);
 	float LightHeight = GetHeight(LightId);
-
+	
 	// Define rectangle and compute solid angle
-	FRect Rect;
-	Rect.Origin = LightPosition - WorldPosition;
-	Rect.Axis[0] = LightdPdu;
-	Rect.Axis[1] = LightdPdv;
-	Rect.Axis[2] = cross(LightdPdu, LightdPdv);
-	Rect.Extent = float2(LightWidth, LightHeight) * 0.5f;
+	FRect Rect = GetRect(LightPosition - WorldPosition,
+						 -LightNormal,
+						 LightdPdv,
+						 0.5 * LightWidth,
+						 0.5 * LightHeight,
+						 GetRectLightBarnCosAngle(LightId),
+						 GetRectLightBarnLength(LightId),
+						 true /* bComputeVisibleRect */);
 	FSphericalRect SphericalRect = BuildSphericalRect(Rect);
+
+	// #dxr_todo: Analyze the overall efficiency (MSE reduction vs. time)
+	//  Projected solid angle sampling converges faster at equal sample count, but is a bit slower
+	// #dxr_todo: Should we switch to plain area sampling at some distance?
+#if 1
+	// Approximate _projected_ solidangle sampling of the spherical rectangle
+	float4 Result = SampleApproxProjectionSphericalRect(RandSample.yz, SphericalRect, Payload.WorldNormal);
+	OutLightUV = Result.xyz;
+	OutPdf = Result.w;
+#else
+	// Ordinary solid angle sampling of the spherical rectangle
 	OutLightUV = UniformSampleSphericalRect(RandSample.yz, SphericalRect);
 	OutPdf = isfinite(SphericalRect.SolidAngle) ? 1.0 / SphericalRect.SolidAngle : 0.0;
+#endif
 }
\ No newline at end of file
diff --git a/Engine/Shaders/Private/PathTracing/Light/PathTracingSpotLight.ush b/Engine/Shaders/Private/PathTracing/Light/PathTracingSpotLight.ush
--- a/Engine/Shaders/Private/PathTracing/Light/PathTracingSpotLight.ush
+++ b/Engine/Shaders/Private/PathTracing/Light/PathTracingSpotLight.ush
@@ -8,6 +8,24 @@
 
 #include "PathTracingLightCommon.ush"
 
+// #dxr_todo: These are stolen from DynamicLightingCommon.ush but I was having trouble getting that header to include cleanly
+/** 
+ * Calculates attenuation for a spot light.
+ * L normalize vector to light. 
+ * SpotDirection is the direction of the spot light.
+ * SpotAngles.x is CosOuterCone, SpotAngles.y is InvCosConeDifference. 
+ */
+float SpotAttenuationMask(float3 L, float3 SpotDirection, float2 SpotAngles)
+{
+	return saturate((dot(L, -SpotDirection) - SpotAngles.x) * SpotAngles.y);
+}
+float SpotAttenuation(float3 L, float3 SpotDirection, float2 SpotAngles)
+{
+	float ConeAngleFalloff = Square(SpotAttenuationMask(L, SpotDirection, SpotAngles));
+	return ConeAngleFalloff;
+}
+
+
 void SpotLight_GenerateLightRay(
 	RayDesc Ray,
 	int LightId,
@@ -15,21 +33,24 @@ void SpotLight_GenerateLightRay(
 	out RayDesc OutLightRay
 )
 {
+	float3 LightDirection = GetPosition(LightId) - Ray.Origin;
+
 	OutLightRay.Origin = Ray.Origin;
 	OutLightRay.TMin = 0.001;
+	OutLightRay.Direction = TangentToWorld(LightUV, normalize(LightDirection));
+
+	// Compute distance to sphere edge
 	float Radius = GetRadius(LightId);
-	Radius = 0.0;
-	float3 DeltaPosition = 0.0;
-	if (Radius > 0.0)
-	{
-		// Transform UV into normal-aligned direction
-		float3 LightNormal = GetNormal(LightId);
-		DeltaPosition = TangentToWorld(LightUV, LightNormal) * Radius;
-	}
-	float3 LightDir = GetPosition(LightId) + DeltaPosition - Ray.Origin;
-	float Length = length(LightDir);
-	OutLightRay.Direction = LightDir / Length;
-	OutLightRay.TMax = Length;
+	float Radius2 = Pow2(Radius2);
+
+	float RayLength2 = dot(LightDirection, LightDirection);
+	float RayLength = sqrt(RayLength2);
+	float SinThetaMax2 = Radius2 / RayLength2;
+
+	float CosTheta = LightUV.z;
+	float SinTheta2 = 1.0 - CosTheta * CosTheta;
+
+	OutLightRay.TMax = RayLength * (CosTheta - sqrt(max(SinThetaMax2 - SinTheta2, 0.0)));
 }
 
 void SpotLight_PdfLight(
@@ -49,52 +70,21 @@ void SpotLight_EvalLight(
 	out float3 OutRadiance
 )
 {
-	OutRadiance = 0.0;
+	float3 LightDirection = GetPosition(LightId) - Ray.Origin;
+	float LightDistanceSquared = dot(LightDirection, LightDirection);
 
-	float3 LightPosition = GetPosition(LightId);
-	float3 LightNormal = GetNormal(LightId);
 	float3 Color = GetColor(LightId);
 	float Radius = GetRadius(LightId);
-	Radius = 0.0;
-	float CosConeAngle = GetCosConeAngle(LightId);
-
-	// Check for spherical definition
-	float3 ClosestPoint = LightPosition;
-	if (Radius > 0.0)
-	{
-		float SphericalCapArea = 2.0 * PI * Radius * Radius * (1.0 - CosConeAngle);
-		Color /= SphericalCapArea;
-
-		// Create plane equation
-		float3 DiskOrigin = LightPosition + LightNormal * Radius * CosConeAngle;
-		float d = -dot(LightNormal, DiskOrigin);
-		float SignedDistance = dot(LightNormal, Ray.Origin) + d;
-
-		// Project ray origin to find closest point on disk
-		ClosestPoint = Ray.Origin - LightNormal * SignedDistance;
-		float3 DeltaFromOrigin = ClosestPoint - DiskOrigin;
-		float DistanceToDiskOrigin = length(DeltaFromOrigin);
-		if (DistanceToDiskOrigin > Radius)
-		{
-			ClosestPoint = DiskOrigin + DeltaFromOrigin * (Radius / DistanceToDiskOrigin);
-		}
-	}
+	float Falloff = ComputeAttenuationFalloff(Ray.TMax, GetAttenuation(LightId));
+	OutRadiance = Color * Falloff;
 
-	// Perform culling based on closest point
-	float3 LightDirection = ClosestPoint - Ray.Origin;
-	float LightDistanceSquared = dot(LightDirection, LightDirection);
-	LightDirection = normalize(LightDirection);
+	// #dxr-todo: See math in point lights - ideally spot light would share that logic somehow
+	OutRadiance *= 1.0 / LightDistanceSquared;
 
-	// Apply culling
-	float NoL = dot(-LightDirection, LightNormal);
-	if (NoL >= CosConeAngle)
-	{
-		float Falloff = NoL - CosConeAngle;
-		Falloff *= Falloff;
-		Falloff *= ComputeAttenuationFalloff(Ray.TMax, GetAttenuation(LightId));
-		OutRadiance = Color * Falloff / LightDistanceSquared;
-		// dxr_todo: apply inner core
-	}
+	// This appears to be the only difference between point and spot lights
+	float3 LightNormal = GetNormal(LightId);
+	float2 CosConeAngles = GetCosConeAngles(LightId);
+	OutRadiance *= SpotAttenuation(normalize(LightDirection), LightNormal, CosConeAngles);
 }
 
 void SpotLight_EstimateLight(
@@ -108,53 +98,29 @@ void SpotLight_EstimateLight(
 	float3 LightNormal = GetNormal(LightId);
 	float LightPower = Luminance(GetColor(LightId));
 	float Radius = GetRadius(LightId);
-	Radius = 0.0;
-	float CosConeAngle = GetCosConeAngle(LightId);
-
-	// Check for spherical definition
-	float3 ClosestPoint = LightPosition;
-	if (Radius > 0.0)
-	{
-		float SphericalCapArea = 2.0 * PI * Radius * Radius * (1.0 - CosConeAngle);
-		LightPower /= SphericalCapArea;
-
-		// Create plane equation
-		float3 DiskOrigin = LightPosition + LightNormal * Radius * CosConeAngle;
-		float d = -dot(LightNormal, DiskOrigin);
-		float SignedDistance = dot(LightNormal, Ray.Origin) + d;
-
-		// Project ray origin to find closest point on disk
-		ClosestPoint = Ray.Origin - LightNormal * SignedDistance;
-		float3 DeltaFromOrigin = ClosestPoint - DiskOrigin;
-		float DistanceToDiskOrigin = length(DeltaFromOrigin);
-		if (DistanceToDiskOrigin > Radius)
-		{
-			ClosestPoint = DiskOrigin + DeltaFromOrigin * (Radius / DistanceToDiskOrigin);
-		}
-	}
-
-	// Perform culling based on closest point
-	float3 LightDirection = ClosestPoint - Payload.WorldPos;
+	
+	// Distance
+	float3 LightDirection = GetPosition(LightId) - Payload.WorldPos;
 	float LightDistanceSquared = dot(LightDirection, LightDirection);
 	LightDirection = normalize(LightDirection);
 
-	float NoL = dot(LightNormal, -LightDirection);
-	if (NoL >= CosConeAngle)
-	{
-		float Falloff = NoL - CosConeAngle;
-		Falloff *= Falloff;
-		Falloff *= ComputeAttenuationFalloff(sqrt(LightDistanceSquared), GetAttenuation(LightId));
-		float CosTheta = dot(Payload.WorldNormal, LightDirection);
-		CosTheta = ENABLE_TRANSMISSION && HasTransmission(LightId) ? abs(CosTheta) : saturate(CosTheta);
-		OutIrradiance = LightPower * CosTheta * Falloff / LightDistanceSquared;
-		// #dxr_todo: UE-72534 apply inner core
-	}
-	else
+	// Geometric term
+	float3 SurfaceNormal = Payload.WorldNormal;
+
+	float NoL = 1.0;
+	if (GetRadius(LightId) == 0.0)
 	{
-		OutIrradiance = 0.0f;
+		float NoL = dot(SurfaceNormal, LightDirection);
+		NoL = ENABLE_TRANSMISSION && HasTransmission(LightId) ? abs(NoL) : saturate(NoL);
 	}
-}
 
+	float Falloff = ComputeAttenuationFalloff(sqrt(LightDistanceSquared), GetAttenuation(LightId));
+	OutIrradiance = LightPower * Falloff * NoL / LightDistanceSquared;
+	
+	float2 CosConeAngles = GetCosConeAngles(LightId);
+	OutIrradiance *= SpotAttenuation(LightDirection, LightNormal, CosConeAngles);
+}
+	
 void SpotLight_SampleLight(
 	RayDesc Ray,
 	FMaterialClosestHitPayload Payload,
@@ -165,23 +131,33 @@ void SpotLight_SampleLight(
 )
 {
 	float Radius = GetRadius(LightId);
-	Radius = 0.0;
 	if (Radius > 0.0)
 	{
-		// Sample a disk that cuts the spherical cap corresponding to the cone angle
-		float CosConeAngle = GetCosConeAngle(LightId);
-		float2 DiskPosition = UniformSampleDiskConcentric(RandSample.yz);
-
-		// Project disk to cutting slice of bounding sphere
-		float SinConeAngle = sqrtFast(1 - CosConeAngle * CosConeAngle);
-		OutLightUV = normalize(float3(SinConeAngle * DiskPosition, CosConeAngle));
-		float SphericalCapArea = 2.0 * PI * Radius * Radius * (1.0 - CosConeAngle);
-		// #dxr_todo: UE-72535 apply the correct Jacobian
-		OutPdf = 1.0 / SphericalCapArea;
+		// Soft case - sample the solid angle subtended by the sphere
+		float3 LightDirection = GetPosition(LightId) - Payload.WorldPos;
+		float LightDistanceSquared = dot(LightDirection, LightDirection);
+		float Radius2 = Pow2(Radius);
+		
+		// outside the light
+		float SinThetaMax2 = Radius2 / LightDistanceSquared;
+		// #dxr_todo: find a better way of handling the region inside the light
+		float CosThetaMax = sqrt(max(1.0 - SinThetaMax2, 0.0));
+		float4 DirAndPdf = UniformSampleCone(RandSample.yz, CosThetaMax);
+
+		OutLightUV = DirAndPdf.xyz;
+#if 0
+		OutPdf = DirAndPdf.w;
+#else
+		// doing the "right" thing leads to unstable math for small radius lights
+		// #dxr_todo: Revisit this when implementing MIS for all lights
+		OutPdf = 1.0;
+#endif
 	}
 	else
 	{
-		OutLightUV = 0.0;
+		// Singular case
+		OutLightUV = float3(0, 0, 1);
+		// #dxr_todo: this isn't really part of the pdf, the pdf should be +inf for the singular case
 		OutPdf = 1.0;
 	}
 }
diff --git a/Engine/Source/Runtime/Renderer/Private/PathTracing.cpp b/Engine/Source/Runtime/Renderer/Private/PathTracing.cpp
--- a/Engine/Source/Runtime/Renderer/Private/PathTracing.cpp
+++ b/Engine/Source/Runtime/Renderer/Private/PathTracing.cpp
@@ -247,7 +247,7 @@ public:
 						LightData.Position[LightData.Count] = LightParameters.Position;
 						LightData.Normal[LightData.Count] = -LightParameters.Direction;
 						// #dxr_todo: UE-72556  define these differences from Lit..
-						LightData.Color[LightData.Count] = 4.0 * PI * LightParameters.Color;
+						LightData.Color[LightData.Count] = LightParameters.Color;
 						LightData.Dimensions[LightData.Count] = FVector(LightParameters.SpotAngles, LightParameters.SourceRadius);
 						LightData.Attenuation[LightData.Count] = 1.0 / LightParameters.InvRadius;
 						break;
@@ -258,7 +258,7 @@ public:
 						LightData.Type[LightData.Count] = 1;
 						LightData.Position[LightData.Count] = LightParameters.Position;
 						// #dxr_todo: UE-72556  define these differences from Lit..
-						LightData.Color[LightData.Count] = LightParameters.Color / (4.0 * PI);
+						LightData.Color[LightData.Count] = LightParameters.Color;
 						LightData.Dimensions[LightData.Count] = FVector(0.0, 0.0, LightParameters.SourceRadius);
 						LightData.Attenuation[LightData.Count] = 1.0 / LightParameters.InvRadius;
 						break;
diff --git a/Engine/Source/Runtime/Renderer/Private/RayTracing/RayTracingGlobalIllumination.cpp b/Engine/Source/Runtime/Renderer/Private/RayTracing/RayTracingGlobalIllumination.cpp
--- a/Engine/Source/Runtime/Renderer/Private/RayTracing/RayTracingGlobalIllumination.cpp
+++ b/Engine/Source/Runtime/Renderer/Private/RayTracing/RayTracingGlobalIllumination.cpp
@@ -326,7 +326,7 @@ RENDERER_API void SetupLightParameters(
 			LightParameters->Type[LightParameters->Count] = 1;
 			LightParameters->Position[LightParameters->Count] = LightShaderParameters.Position;
 			// #dxr_todo: UE-72556 define these differences from Lit..
-			LightParameters->Color[LightParameters->Count] = LightShaderParameters.Color / (4.0 * PI);
+			LightParameters->Color[LightParameters->Count] = LightShaderParameters.Color;
 			float SourceRadius = 0.0; // LightShaderParameters.SourceRadius causes too much noise for little pay off at this time
 			LightParameters->Dimensions[LightParameters->Count] = FVector(0.0, 0.0, SourceRadius);
 			LightParameters->Attenuation[LightParameters->Count] = 1.0 / LightShaderParameters.InvRadius;
@@ -340,7 +340,7 @@ RENDERER_API void SetupLightParameters(
 			LightParameters->Position[LightParameters->Count] = LightShaderParameters.Position;
 			LightParameters->Normal[LightParameters->Count] = -LightShaderParameters.Direction;
 			// #dxr_todo: UE-72556 define these differences from Lit..
-			LightParameters->Color[LightParameters->Count] = 4.0 * PI * LightShaderParameters.Color;
+			LightParameters->Color[LightParameters->Count] = LightShaderParameters.Color;
 			float SourceRadius = 0.0; // LightShaderParameters.SourceRadius causes too much noise for little pay off at this time
 			LightParameters->Dimensions[LightParameters->Count] = FVector(LightShaderParameters.SpotAngles, SourceRadius);
 			LightParameters->Attenuation[LightParameters->Count] = 1.0 / LightShaderParameters.InvRadius;
