Updated conductor and volume nodes to be artist friendly.
Added thickness input on volumetric (test)

diff --git a/Engine/Shaders/Private/ShadingCommon.ush b/Engine/Shaders/Private/ShadingCommon.ush
--- a/Engine/Shaders/Private/ShadingCommon.ush
+++ b/Engine/Shaders/Private/ShadingCommon.ush
@@ -123,6 +123,22 @@ void ComputeF0AndEdgeTintFromIOR(in float3 n, in float3 k, inout float3 Reflecti
 	EdgeTint = (Common0 - n) / (Common0 + (1 - SqrtR) / (1 + SqrtR));
 }
 
+float3 ComputeF0AndEdgeTintFromComplexIOR_Reflectivity(in float3 n, in float3 k)
+{
+	float3 Reflectivity = 0;
+	float3 EdgeTint = 0;
+	ComputeF0AndEdgeTintFromIOR(n, k, Reflectivity, EdgeTint);
+	return Reflectivity;
+}
+
+float3 ComputeF0AndEdgeTintFromComplexIOR_EdgeTint(in float3 n, in float3 k)
+{
+	float3 Reflectivity = 0;
+	float3 EdgeTint = 0;
+	ComputeF0AndEdgeTintFromIOR(n, k, Reflectivity, EdgeTint);
+	return EdgeTint;
+}
+
 // Compute complex index of refraction N = n + ik from F0 and edge tint.
 void ComputeComplexIORFromF0AndEdgeTint(in float3 Reflectivity, in float3 EdgeTint, inout float3 n, inout float3 k)
 {
diff --git a/Engine/Shaders/Private/Strata/Strata.ush b/Engine/Shaders/Private/Strata/Strata.ush
--- a/Engine/Shaders/Private/Strata/Strata.ush
+++ b/Engine/Shaders/Private/Strata/Strata.ush
@@ -38,8 +38,8 @@ struct FStrataDielectric
 
 struct FStrataConductor
 {
-	float3	IOR;
-	float3	Extinction;
+	float3	Reflectivity;
+	float3	EdgeColor;
 	float2	Roughness;
 
 	float3	Normal;
@@ -47,9 +47,10 @@ struct FStrataConductor
 
 struct FStrataVolume
 {
-	float3	Absorption;
-	float3	Scattering;
+	float3	Albedo;
+	float3	Extinction;
 	float	Anisotropy;
+	float	Thickness;
 };
 
 
@@ -154,15 +155,15 @@ FStrataData GetStrataDielectricBSDF(float2 Roughness, float IOR, float3 Tint, fl
 	return StrataData;
 }
 
-FStrataData GetStrataConductorBSDF(float3 IOR, float3 Extinction, float2 Roughness, float3 Normal)
+FStrataData GetStrataConductorBSDF(float3 Reflectivity, float3 EdgeColor, float2 Roughness, float3 Normal)
 {
 	FStrataData StrataData = GetInitialisedStrataData();
 	StrataData.LayerCount++;
 
 	StrataData.Layers[0].BSDFCount++;
 	StrataData.Layers[0].BSDFs[0].BSDFType						= STRATA_BSDF_TYPE_CONDUCTOR;
-	StrataData.Layers[0].BSDFs[0].Conductor.IOR					= IOR;
-	StrataData.Layers[0].BSDFs[0].Conductor.Extinction			= Extinction;
+	StrataData.Layers[0].BSDFs[0].Conductor.Reflectivity		= Reflectivity;
+	StrataData.Layers[0].BSDFs[0].Conductor.EdgeColor			= EdgeColor;
 	StrataData.Layers[0].BSDFs[0].Conductor.Roughness			= Roughness;
 	StrataData.Layers[0].BSDFs[0].Conductor.Normal				= Normal;
 
@@ -171,16 +172,17 @@ FStrataData GetStrataConductorBSDF(float3 IOR, float3 Extinction, float2 Roughne
 	return StrataData;
 }
 
-FStrataData GetStrataVolumeBSDF(float3 Absorption, float3 Scattering, float Anisotropy)
+FStrataData GetStrataVolumeBSDF(float3 Albedo, float3 Extinction, float Anisotropy, float Thickness)
 {
 	FStrataData StrataData = GetInitialisedStrataData();
 	StrataData.LayerCount++;
 
 	StrataData.Layers[0].BSDFCount++;
 	StrataData.Layers[0].BSDFs[0].BSDFType						= STRATA_BSDF_TYPE_VOLUME;
-	StrataData.Layers[0].BSDFs[0].Volume.Absorption				= Absorption;
-	StrataData.Layers[0].BSDFs[0].Volume.Scattering				= Scattering;
+	StrataData.Layers[0].BSDFs[0].Volume.Albedo					= Albedo;
+	StrataData.Layers[0].BSDFs[0].Volume.Extinction				= Extinction;
 	StrataData.Layers[0].BSDFs[0].Volume.Anisotropy				= Anisotropy;
+	StrataData.Layers[0].BSDFs[0].Volume.Thickness				= Thickness;
 
 	StrataData.Layers[0].BSDFs[0].Weight						= 1.0f;
 
@@ -359,6 +361,14 @@ void StrataLoad_R11G11_B10F(ByteAddressBuffer StrataBuffer, inout uint ByteOffse
 	B = temp.z;
 }
 
+void StrataLoad_R11_G11_B10F(ByteAddressBuffer StrataBuffer, inout uint ByteOffset, inout float A, inout float B, inout float C)
+{
+	float3 temp = UnpackR11G11B10F(StrataLoadUint1(StrataBuffer, ByteOffset));
+	A = temp.x;
+	B = temp.y;
+	C = temp.z;
+}
+
 void StrataLoad_AlbedoGamma2ToLinear(ByteAddressBuffer StrataBuffer, inout uint ByteOffset, inout float3 A)
 {
 	A = UnpackAlbedoGamma2ToLinear(StrataLoadUint1(StrataBuffer, ByteOffset));
@@ -439,10 +449,7 @@ void PackStrataOut(
 
 			case STRATA_BSDF_TYPE_CONDUCTOR:
 			{
-				float3 ReflectivityF0 = float3(1.0, 0.1, 0.1);
-				float3 EdgeTint = float3(0.1, 1.0, 0.1);
-				ComputeF0AndEdgeTintFromIOR(BSDF.Conductor.IOR, BSDF.Conductor.Extinction, ReflectivityF0, EdgeTint);
-				float3 F = F_Schlick(ReflectivityF0, EdgeTint, ContextNoLight.VoH);
+				float3 F = F_Schlick(BSDF.Conductor.Reflectivity, BSDF.Conductor.EdgeColor, ContextNoLight.VoH);
 
 				LayerTransmittance += BSDF.Weight * (1.0f - F); // STRATA_TODO transmittance should only come from coverage
 				break;
@@ -450,16 +457,15 @@ void PackStrataOut(
 
 			case STRATA_BSDF_TYPE_VOLUME:
 			{
-				const float SafeRoughness = clamp(BSDF.Conductor.Roughness.x, 0.001, 1.0f);
-
-				const float3 ScatteringCoefficients = BSDF.Volume.Scattering;
-				const float3 AbsorptionCoefficients = BSDF.Volume.Absorption;
-				const float3 ExtinctionCoefficients = ScatteringCoefficients + AbsorptionCoefficients;
+				const float3 Albedo					= BSDF.Volume.Albedo;
+				const float3 ExtinctionCoefficients = BSDF.Volume.Extinction;
+				const float3 ScatteringCoefficients = Albedo * ExtinctionCoefficients;
+				const float3 AbsorptionCoefficients = ExtinctionCoefficients - ScatteringCoefficients;
+				const float  Thickness				= BSDF.Volume.Thickness;
 
 				const float3 SafeExtinctionThreshold = 0.000001f;
 				const float3 SafeExtinctionCoefficients = max(SafeExtinctionThreshold, ExtinctionCoefficients);
 
-				const float Thickness = 0.01f;	// STRATA_TODO where does thickness come from? For now, 1cm
 				const float PathLength = Thickness / max(0.0001f, abs(ContextNoLight.NoV));
 				const float3 SafePathSegmentTransmittance = exp(-SafeExtinctionCoefficients * PathLength);
 
@@ -512,16 +518,17 @@ void PackStrataOut(
 			break;
 			case STRATA_BSDF_TYPE_CONDUCTOR:
 			{
-				STRATA_STORE_UINT2(PackR16G16B16A16F(float4(BSDF.Conductor.IOR, BSDF.Conductor.Roughness.x)));
-				STRATA_STORE_UINT2(PackR16G16B16A16F(float4(BSDF.Conductor.Extinction, BSDF.Conductor.Roughness.y)));
+				STRATA_STORE_UINT2(PackR16G16B16A16F(float4(BSDF.Conductor.Reflectivity, BSDF.Conductor.Roughness.x)));
+				STRATA_STORE_UINT2(PackR16G16B16A16F(float4(BSDF.Conductor.EdgeColor, BSDF.Conductor.Roughness.y)));
 				STRATA_STORE_UINT2(PackR16G16B16A16F(float4(BSDF.Conductor.Normal, 0.0f)));
 				// 24 bytes
 			}
 			break;
 			case STRATA_BSDF_TYPE_VOLUME:
 			{
-				STRATA_STORE_UINT1(PackR11G11B10F(BSDF.Volume.Absorption));
-				STRATA_STORE_UINT2(PackR16G16B16A16F(float4(BSDF.Volume.Scattering, BSDF.Volume.Anisotropy)));
+				STRATA_STORE_UINT1(PackR11G11B10F(BSDF.Volume.Albedo));
+				STRATA_STORE_UINT1(PackR11G11B10F(BSDF.Volume.Extinction));
+				STRATA_STORE_UINT1(PackR11G11B10F(float3(BSDF.Volume.Anisotropy, BSDF.Volume.Thickness, 0.0f)));
 				// 12 bytes
 			}
 			break;
@@ -580,15 +587,16 @@ FStrataPixelBSDFs UnpackStrataIn(ByteAddressBuffer StrataBuffer, uint ByteOffset
 		break;
 		case STRATA_BSDF_TYPE_CONDUCTOR:
 		{
-			StrataLoad_R16G16B16F_A16F		(StrataBuffer, ByteOffset, PixelBSDFs.BSDFs[i].Conductor.IOR,			PixelBSDFs.BSDFs[i].Conductor.Roughness.x);
-			StrataLoad_R16G16B16F_A16F		(StrataBuffer, ByteOffset, PixelBSDFs.BSDFs[i].Conductor.Extinction,	PixelBSDFs.BSDFs[i].Conductor.Roughness.y);
+			StrataLoad_R16G16B16F_A16F		(StrataBuffer, ByteOffset, PixelBSDFs.BSDFs[i].Conductor.Reflectivity,	PixelBSDFs.BSDFs[i].Conductor.Roughness.x);
+			StrataLoad_R16G16B16F_A16F		(StrataBuffer, ByteOffset, PixelBSDFs.BSDFs[i].Conductor.EdgeColor,		PixelBSDFs.BSDFs[i].Conductor.Roughness.y);
 			StrataLoad_R16G16B16F_A16F		(StrataBuffer, ByteOffset, PixelBSDFs.BSDFs[i].Conductor.Normal,		DummyFloat);
 		}
 		break;
 		case STRATA_BSDF_TYPE_VOLUME:
 		{
-			StrataLoad_R11G11B10F			(StrataBuffer, ByteOffset, PixelBSDFs.BSDFs[i].Volume.Absorption);
-			StrataLoad_R16G16B16F_A16F		(StrataBuffer, ByteOffset, PixelBSDFs.BSDFs[i].Volume.Scattering,		PixelBSDFs.BSDFs[i].Volume.Anisotropy);
+			StrataLoad_R11G11B10F			(StrataBuffer, ByteOffset, PixelBSDFs.BSDFs[i].Volume.Albedo);
+			StrataLoad_R11G11B10F			(StrataBuffer, ByteOffset, PixelBSDFs.BSDFs[i].Volume.Extinction);
+			StrataLoad_R11_G11_B10F			(StrataBuffer, ByteOffset, PixelBSDFs.BSDFs[i].Volume.Anisotropy, PixelBSDFs.BSDFs[i].Volume.Thickness, DummyFloat);
 		}
 		break;
 		}
diff --git a/Engine/Shaders/Private/Strata/StrataDeferredLighting.ush b/Engine/Shaders/Private/Strata/StrataDeferredLighting.ush
--- a/Engine/Shaders/Private/Strata/StrataDeferredLighting.ush
+++ b/Engine/Shaders/Private/Strata/StrataDeferredLighting.ush
@@ -88,10 +88,6 @@ float4 StrataDeferredLighting(
 		break;
 		case STRATA_BSDF_TYPE_CONDUCTOR:
 		{
-			float3 ReflectivityF0 = float3(1.0, 0.1, 0.1);
-			float3 EdgeTint = float3(0.1, 1.0, 0.1);
-			ComputeF0AndEdgeTintFromIOR(BSDF.Conductor.IOR, BSDF.Conductor.Extinction, ReflectivityF0, EdgeTint);
-
 			// See SpecularGGX, isotropic verion
 
 			const float SafeRoughness = MakeRoughnessSafe(BSDF.Conductor.Roughness.x);
@@ -101,23 +97,22 @@ float4 StrataDeferredLighting(
 			// Generalized microfacet specular
 			float D = D_GGX(a2, Context.NoH) * Energy;
 			float Vis = Vis_SmithJointApprox(a2, Context.NoV, SatNoL);
-			float3 F = F_Schlick(ReflectivityF0, EdgeTint, Context.VoH);
+			float3 F = F_Schlick(BSDF.Conductor.Reflectivity, BSDF.Conductor.EdgeColor, Context.VoH);
 
 			Radiance.rgb += D * Vis * F * DirLightGeometryFactor;
 		}
 		break;
 		case STRATA_BSDF_TYPE_VOLUME:
 		{
-			// Strata TODO expose albedo + extinction instead?
-			const float3 ScatteringCoefficients = BSDF.Volume.Scattering;
-			const float3 AbsorptionCoefficients = BSDF.Volume.Absorption;
-			const float3 ExtinctionCoefficients = ScatteringCoefficients + AbsorptionCoefficients;
+			const float3 Albedo = BSDF.Volume.Albedo;
+			const float3 ExtinctionCoefficients = BSDF.Volume.Extinction;
+			const float3 ScatteringCoefficients = Albedo * ExtinctionCoefficients;
+			const float3 AbsorptionCoefficients = ExtinctionCoefficients - ScatteringCoefficients;
+			const float  Thickness = BSDF.Volume.Thickness;
 
 			const float3 SafeExtinctionThreshold = 0.000001f;
 			const float3 SafeExtinctionCoefficients = max(SafeExtinctionThreshold, ExtinctionCoefficients);
-			const float3 SafeAlbedo = ScatteringCoefficients / SafeExtinctionCoefficients;
 
-			const float Thickness = 0.01f;	// STRATA_TODO where does thickness come from? For now, 1cm
 			const float PathLength = Thickness / max(0.0001f, abs(Context.NoV));
 			const float3 SafePathSegmentTransmittance = exp(-SafeExtinctionCoefficients * PathLength);
 			const float3 NormLuminanceTransfert = 1.0f * ScatteringCoefficients;
diff --git a/Engine/Shaders/Private/Strata/StrataEnvironmentLighting.ush b/Engine/Shaders/Private/Strata/StrataEnvironmentLighting.ush
--- a/Engine/Shaders/Private/Strata/StrataEnvironmentLighting.ush
+++ b/Engine/Shaders/Private/Strata/StrataEnvironmentLighting.ush
@@ -64,10 +64,6 @@ float3 StrataEnvironmentLighting(
 			BRANCH
 				if (ReflectionStruct.SkyLightParameters.y > 0)
 				{
-					float3 ReflectivityF0 = float3(1.0, 0.1, 0.1);
-					float3 EdgeTint = float3(0.1, 1.0, 0.1);
-					ComputeF0AndEdgeTintFromIOR(BSDF.Conductor.IOR, BSDF.Conductor.Extinction, ReflectivityF0, EdgeTint);
-
 					// See SpecularGGX, isotropic verion
 
 					const float SafeRoughness = MakeRoughnessSafe(BSDF.Conductor.Roughness.x);
@@ -77,9 +73,9 @@ float3 StrataEnvironmentLighting(
 					// Generalized microfacet specular
 					//float D = D_GGX(a2, Context.NoH) * Energy;
 					//float Vis = Vis_SmithJointApprox(a2, Context.NoV, SatNoL);
-					//float3 F = F_Schlick(ReflectivityF0, EdgeTint, Context.VoH); // STRATA_TODO check that this is correct 
+					//float3 F = F_Schlick(BSDF.Conductor.Reflectivity, BSDF.Conductor.EdgeColor, Context.VoH); // STRATA_TODO check that this is correct 
 
-					float3 EvalEnvBRDF = EnvBRDFApprox(ReflectivityF0, EdgeTint, SafeRoughness, Context.NoV);
+					float3 EvalEnvBRDF = EnvBRDFApprox(BSDF.Conductor.Reflectivity, BSDF.Conductor.EdgeColor, SafeRoughness, Context.NoV);
 					float SkyAverageBrightness = 1.0f;
 					SpecularLuminance += GetSkyLightReflection(R, SafeRoughness, SkyAverageBrightness) * EvalEnvBRDF;
 				}
@@ -87,15 +83,15 @@ float3 StrataEnvironmentLighting(
 		break;
 		case STRATA_BSDF_TYPE_VOLUME:
 		{
-			const float3 ScatteringCoefficients = BSDF.Volume.Scattering;
-			const float3 AbsorptionCoefficients = BSDF.Volume.Absorption;
-			const float3 ExtinctionCoefficients = ScatteringCoefficients + AbsorptionCoefficients;
+			const float3 VolumeAlbedo = BSDF.Volume.Albedo;
+			const float3 ExtinctionCoefficients = BSDF.Volume.Extinction;
+			const float3 ScatteringCoefficients = VolumeAlbedo * ExtinctionCoefficients;
+			const float3 AbsorptionCoefficients = ExtinctionCoefficients - ScatteringCoefficients;
+			const float  Thickness = BSDF.Volume.Thickness;
 
 			const float3 SafeExtinctionThreshold = 0.000001f;
 			const float3 SafeExtinctionCoefficients = max(SafeExtinctionThreshold, ExtinctionCoefficients);
-			const float3 SafeAlbedo = ScatteringCoefficients / SafeExtinctionCoefficients;
 
-			const float Thickness = 0.01f;	// STRATA_TODO where does thickness come from? For now, 1cm
 			const float PathLength = Thickness / max(0.0001f, saturate(Context.NoV));
 			const float3 SafePathSegmentTransmittance = exp(-SafeExtinctionCoefficients * PathLength);
 			const float3 NormLuminanceTransfert = 1.0f * ScatteringCoefficients;
diff --git a/Engine/Shaders/Private/Strata/StrataForwardLighting.ush b/Engine/Shaders/Private/Strata/StrataForwardLighting.ush
--- a/Engine/Shaders/Private/Strata/StrataForwardLighting.ush
+++ b/Engine/Shaders/Private/Strata/StrataForwardLighting.ush
@@ -117,10 +117,6 @@ float3 StrataForwardLighting(
 
 			case STRATA_BSDF_TYPE_CONDUCTOR:
 			{
-				float3 ReflectivityF0 = float3(1.0, 0.1, 0.1);
-				float3 EdgeTint = float3(0.1, 1.0, 0.1);
-				ComputeF0AndEdgeTintFromIOR(BSDF.Conductor.IOR, BSDF.Conductor.Extinction, ReflectivityF0, EdgeTint);
-
 				// See SpecularGGX, isotropic verion
 
 				const float SafeRoughness = MakeRoughnessSafe(BSDF.Conductor.Roughness.x);
@@ -130,7 +126,7 @@ float3 StrataForwardLighting(
 				// Generalized microfacet specular
 				float D = D_GGX(a2, Context.NoH) * Energy;
 				float Vis = Vis_SmithJointApprox(a2, Context.NoV, SatNoL);
-				float3 F = F_Schlick(ReflectivityF0, EdgeTint, Context.VoH);
+				float3 F = F_Schlick(BSDF.Conductor.Reflectivity, BSDF.Conductor.EdgeColor, Context.VoH);
 
 				Color += D * Vis * F * DirLightGeometryFactor;
 
@@ -138,7 +134,7 @@ float3 StrataForwardLighting(
 				//BRANCH
 				if (ReflectionStruct.SkyLightParameters.y > 0)
 				{
-					float3 EvalEnvBRDF = EnvBRDFApprox(ReflectivityF0, EdgeTint, SafeRoughness, Context.NoV);
+					float3 EvalEnvBRDF = EnvBRDFApprox(BSDF.Conductor.Reflectivity, BSDF.Conductor.EdgeColor, SafeRoughness, Context.NoV);
 					float SkyAverageBrightness = 1.0f;
 					Color += GetSkyLightReflection(R, SafeRoughness, SkyAverageBrightness) * EvalEnvBRDF * SkyLightCommonFactor;
 				}
@@ -150,17 +146,15 @@ float3 StrataForwardLighting(
 
 			case STRATA_BSDF_TYPE_VOLUME:
 			{
-				const float SafeRoughness = MakeRoughnessSafe(BSDF.Conductor.Roughness.x);
-
-				const float3 ScatteringCoefficients = BSDF.Volume.Scattering;
-				const float3 AbsorptionCoefficients = BSDF.Volume.Absorption;
-				const float3 ExtinctionCoefficients = ScatteringCoefficients + AbsorptionCoefficients;
+				const float3 Albedo = BSDF.Volume.Albedo;
+				const float3 ExtinctionCoefficients = BSDF.Volume.Extinction;
+				const float3 ScatteringCoefficients = Albedo * ExtinctionCoefficients;
+				const float3 AbsorptionCoefficients = ExtinctionCoefficients - ScatteringCoefficients;
+				const float  Thickness = BSDF.Volume.Thickness;
 
 				const float3 SafeExtinctionThreshold = 0.000001f;
 				const float3 SafeExtinctionCoefficients = max(SafeExtinctionThreshold, ExtinctionCoefficients);
-				const float3 SafeAlbedo = ScatteringCoefficients / SafeExtinctionCoefficients;
 
-				const float Thickness = 0.01f;	// STRATA_TODO where does thickness come from? For now, 1cm
 				const float PathLength = Thickness / max(0.0001f, abs(Context.NoV));
 				const float3 SafePathSegmentTransmittance = exp(-SafeExtinctionCoefficients * PathLength);
 				const float3 NormLuminanceTransfert = 1.0f * ScatteringCoefficients;
diff --git a/Engine/Source/Runtime/Engine/Classes/Materials/MaterialExpressionStrata.h b/Engine/Source/Runtime/Engine/Classes/Materials/MaterialExpressionStrata.h
--- a/Engine/Source/Runtime/Engine/Classes/Materials/MaterialExpressionStrata.h
+++ b/Engine/Source/Runtime/Engine/Classes/Materials/MaterialExpressionStrata.h
@@ -142,13 +142,13 @@ class UMaterialExpressionStrataConductorBSDF : public UMaterialExpression
 	 * float3
 	 */
 	UPROPERTY()
-	FExpressionInput IOR;
+	FExpressionInput Reflectivity;
 	
 	/**
 	 * float3
 	 */
 	UPROPERTY()
-	FExpressionInput Extinction;
+	FExpressionInput EdgeColor;
 
 	/**
 	 * float2
@@ -182,23 +182,29 @@ class UMaterialExpressionStrataVolumeBSDF : public UMaterialExpression
 	GENERATED_UCLASS_BODY()
 		
 	/**
-	 * float3
+	 * Albedo (type = float3, unit = 1/m)
 	 */
 	UPROPERTY()
-	FExpressionInput Absorption;
+	FExpressionInput Albedo;
 	
 	/**
-	 * float3
+	 * Extinction (type = float3, unit = 1/m)
 	 */
 	UPROPERTY()
-	FExpressionInput Scattering;
+	FExpressionInput Extinction;
 
 	/**
-	 * float
+	 * Anisotropy (type = float, unitless)
 	 */
 	UPROPERTY()
 	FExpressionInput Anisotropy;
 
+	/**
+	 * Thickness (type = float, unit = meters, default = 1mm)
+	 */
+	UPROPERTY()
+	FExpressionInput Thickness;
+
 	//~ Begin UMaterialExpression Interface
 #if WITH_EDITOR
 	virtual int32 Compile(class FMaterialCompiler* Compiler, int32 OutputIndex) override;
@@ -362,4 +368,31 @@ class UMaterialExpressionStrataArtisticIOR : public UMaterialExpression
 	//~ End UMaterialExpression Interface
 };
 
+UCLASS(MinimalAPI, collapsecategories, hidecategories = Object)
+class UMaterialExpressionStrataPhysicalIOR : public UMaterialExpression
+{
+	GENERATED_UCLASS_BODY()
+
+	/**
+	* TODO 
+	*/
+	UPROPERTY()
+	FExpressionInput IOR;
+
+	/**
+	 * TODO 
+	 */
+	UPROPERTY()
+	FExpressionInput Extinction;
+
+	//~ Begin UMaterialExpression Interface
+#if WITH_EDITOR
+	virtual int32 Compile(class FMaterialCompiler* Compiler, int32 OutputIndex) override;
+	virtual void GetCaption(TArray<FString>& OutCaptions) const override;
+	virtual uint32 GetOutputType(int32 OutputIndex) override;
+	virtual uint32 GetInputType(int32 InputIndex) override;
+#endif
+	//~ End UMaterialExpression Interface
+};
+
 
diff --git a/Engine/Source/Runtime/Engine/Private/Materials/HLSLMaterialTranslator.cpp b/Engine/Source/Runtime/Engine/Private/Materials/HLSLMaterialTranslator.cpp
--- a/Engine/Source/Runtime/Engine/Private/Materials/HLSLMaterialTranslator.cpp
+++ b/Engine/Source/Runtime/Engine/Private/Materials/HLSLMaterialTranslator.cpp
@@ -7078,29 +7078,34 @@ int32 FHLSLMaterialTranslator::StrataDielectricBSDF(int32 Roughness, int32 IOR,
 	);
 }
 
-int32 FHLSLMaterialTranslator::StrataConductorBSDF(int32 IOR, int32 Extinction, int32 Roughness, int32 Normal)
+int32 FHLSLMaterialTranslator::StrataConductorBSDF(int32 Reflectivity, int32 EdgeColor, int32 Roughness, int32 Normal)
 {
 	return AddInlinedCodeChunk(
 		MCT_Strata, TEXT("GetStrataConductorBSDF(%s, %s, %s, %s)"),
-		*GetParameterCode(IOR),
-		*GetParameterCode(Extinction),
+		*GetParameterCode(Reflectivity),
+		*GetParameterCode(EdgeColor),
 		*GetParameterCode(Roughness),
 		*GetParameterCode(Normal)
 	);
 }
 
-int32 FHLSLMaterialTranslator::StrataVolumeBSDF(int32 Absorption, int32 Scattering, int32 Anisotropy)
+int32 FHLSLMaterialTranslator::StrataVolumeBSDF(int32 Albedo, int32 Extinction, int32 Anisotropy, int32 Thickness)
 {
 	return AddInlinedCodeChunk(
-		MCT_Strata, TEXT("GetStrataVolumeBSDF(%s, %s, %s)"),
-		*GetParameterCode(Absorption),
-		*GetParameterCode(Scattering),
-		*GetParameterCode(Anisotropy)
+		MCT_Strata, TEXT("GetStrataVolumeBSDF(%s, %s, %s, %s)"),
+		*GetParameterCode(Albedo),
+		*GetParameterCode(Extinction),
+		*GetParameterCode(Anisotropy),
+		*GetParameterCode(Thickness)
 	);
 }
 
 int32 FHLSLMaterialTranslator::StrataHorizontalMixing(int32 Foreground, int32 Background, int32 Mix)
 {
+	if (Foreground == INDEX_NONE || Background == INDEX_NONE || Mix == INDEX_NONE)
+	{
+		return INDEX_NONE;
+	}
 	return AddInlinedCodeChunk(
 		MCT_Strata, TEXT("StrataHorizontalMixing(%s, %s, %s)"),
 		*GetParameterCode(Foreground),
@@ -7111,6 +7116,10 @@ int32 FHLSLMaterialTranslator::StrataHorizontalMixing(int32 Foreground, int32 Ba
 
 int32 FHLSLMaterialTranslator::StrataVerticalLayering(int32 Top, int32 Base)
 {
+	if (Top == INDEX_NONE || Base == INDEX_NONE)
+	{
+		return INDEX_NONE;
+	}
 	return AddInlinedCodeChunk(
 		MCT_Strata, TEXT("StrataVerticalLayering(%s, %s)"),
 		*GetParameterCode(Top),
@@ -7120,6 +7129,10 @@ int32 FHLSLMaterialTranslator::StrataVerticalLayering(int32 Top, int32 Base)
 
 int32 FHLSLMaterialTranslator::StrataAdd(int32 A, int32 B)
 {
+	if (A == INDEX_NONE || B == INDEX_NONE)
+	{
+		return INDEX_NONE;
+	}
 	return AddInlinedCodeChunk(
 		MCT_Strata, TEXT("StrataAdd(%s, %s)"),
 		*GetParameterCode(A),
@@ -7129,6 +7142,10 @@ int32 FHLSLMaterialTranslator::StrataAdd(int32 A, int32 B)
 
 int32 FHLSLMaterialTranslator::StrataMultiply(int32 A, int32 Weight)
 {
+	if (A == INDEX_NONE || Weight == INDEX_NONE)
+	{
+		return INDEX_NONE;
+	}
 	return AddInlinedCodeChunk(
 		MCT_Strata, TEXT("StrataMultiply(%s, %s)"),
 		*GetParameterCode(A),
@@ -7138,6 +7155,10 @@ int32 FHLSLMaterialTranslator::StrataMultiply(int32 A, int32 Weight)
 
 int32 FHLSLMaterialTranslator::StrataArtisticIOR(int32 Reflectivity, int32 EdgeColor, int32 OutputIndex)
 {
+	if (Reflectivity == INDEX_NONE || EdgeColor == INDEX_NONE)
+	{
+		return INDEX_NONE;
+	}
 	return AddInlinedCodeChunk(
 		MCT_Strata, 
 		OutputIndex == 0 ? TEXT("ComputeComplexIORFromF0AndEdgeTint_IOR(%s, %s)") : TEXT("ComputeComplexIORFromF0AndEdgeTint_Extinction(%s, %s)"),
@@ -7146,6 +7167,20 @@ int32 FHLSLMaterialTranslator::StrataArtisticIOR(int32 Reflectivity, int32 EdgeC
 	);
 }
 
+int32 FHLSLMaterialTranslator::StrataPhysicalIOR(int32 IOR, int32 Extinction, int32 OutputIndex)
+{
+	if (IOR == INDEX_NONE || Extinction == INDEX_NONE)
+	{
+		return INDEX_NONE;
+	}
+	return AddInlinedCodeChunk(
+		MCT_Strata,
+		OutputIndex == 0 ? TEXT("ComputeF0AndEdgeTintFromComplexIOR_Reflectivity(%s, %s)") : TEXT("ComputeF0AndEdgeTintFromComplexIOR_EdgeTint(%s, %s)"),
+		*GetParameterCode(IOR),
+		*GetParameterCode(Extinction)
+	);
+}
+
 int32 FHLSLMaterialTranslator::MapARPassthroughCameraUV(int32 UV)
 {
 	if (UV == INDEX_NONE)
diff --git a/Engine/Source/Runtime/Engine/Private/Materials/HLSLMaterialTranslator.h b/Engine/Source/Runtime/Engine/Private/Materials/HLSLMaterialTranslator.h
--- a/Engine/Source/Runtime/Engine/Private/Materials/HLSLMaterialTranslator.h
+++ b/Engine/Source/Runtime/Engine/Private/Materials/HLSLMaterialTranslator.h
@@ -763,13 +763,14 @@ protected:
 	virtual int32 StrataDiffuseOrenNayarBSDF(int32 Albedo, int32 Roughness, int32 Normal) override;
 	virtual int32 StrataDiffuseChanBSDF(int32 Albedo, int32 Roughness, int32 Normal) override;
 	virtual int32 StrataDielectricBSDF(int32 Roughness, int32 IOR, int32 Tint, int32 Normal) override;
-	virtual int32 StrataConductorBSDF(int32 IOR, int32 Extinction, int32 Roughness, int32 Normal) override;
-	virtual int32 StrataVolumeBSDF(int32 Absorption, int32 Scattering, int32 Anisotropy) override;
+	virtual int32 StrataConductorBSDF(int32 Reflectivity, int32 EdgeColor, int32 Roughness, int32 Normal) override;
+	virtual int32 StrataVolumeBSDF(int32 Albedo, int32 Extinction, int32 Anisotropy, int32 Thickness) override;
 	virtual int32 StrataHorizontalMixing(int32 Foreground, int32 Background, int32 Mix) override;
 	virtual int32 StrataVerticalLayering(int32 Top, int32 Base) override;
 	virtual int32 StrataAdd(int32 A, int32 B) override;
 	virtual int32 StrataMultiply(int32 A, int32 Weight) override;
 	virtual int32 StrataArtisticIOR(int32 Reflectivity, int32 EdgeColor, int32 OutputIndex) override;
+	virtual int32 StrataPhysicalIOR(int32 IOR, int32 Extinction, int32 OutputIndex) override;
 
 #if HANDLE_CUSTOM_OUTPUTS_AS_MATERIAL_ATTRIBUTES
 	/** Used to translate code for custom output attributes such as ClearCoatBottomNormal */
diff --git a/Engine/Source/Runtime/Engine/Private/Materials/MaterialExpressions.cpp b/Engine/Source/Runtime/Engine/Private/Materials/MaterialExpressions.cpp
--- a/Engine/Source/Runtime/Engine/Private/Materials/MaterialExpressions.cpp
+++ b/Engine/Source/Runtime/Engine/Private/Materials/MaterialExpressions.cpp
@@ -18948,13 +18948,13 @@ UMaterialExpressionStrataConductorBSDF::UMaterialExpressionStrataConductorBSDF(c
 #if WITH_EDITOR
 int32 UMaterialExpressionStrataConductorBSDF::Compile(class FMaterialCompiler* Compiler, int32 OutputIndex)
 {
-	if (!IOR.GetTracedInput().Expression)
+	if (!Reflectivity.GetTracedInput().Expression)
 	{
-		return Compiler->Errorf(TEXT("Missing IOR input"));			// STRATA_TODO: defaults in the detail panel with auto read only if plug in
+		return Compiler->Errorf(TEXT("Missing Reflectivity input"));			// STRATA_TODO: defaults in the detail panel with auto read only if plug in
 	}
-	if (!Extinction.GetTracedInput().Expression)
+	if (!EdgeColor.GetTracedInput().Expression)
 	{
-		return Compiler->Errorf(TEXT("Missing Extinction input"));	// STRATA_TODO: defaults in the detail panel with auto read only if plug in
+		return Compiler->Errorf(TEXT("Missing EdgeColor input"));	// STRATA_TODO: defaults in the detail panel with auto read only if plug in
 	}
 	if (!Roughness.GetTracedInput().Expression)
 	{
@@ -18964,7 +18964,7 @@ int32 UMaterialExpressionStrataConductorBSDF::Compile(class FMaterialCompiler* C
 	{
 		return Compiler->Errorf(TEXT("Missing Normal input"));		// STRATA_TODO: defaults in the detail panel with auto read only if plug in
 	}
-	return Compiler->StrataConductorBSDF(IOR.Compile(Compiler), Extinction.Compile(Compiler), Roughness.Compile(Compiler), Normal.Compile(Compiler));
+	return Compiler->StrataConductorBSDF(Reflectivity.Compile(Compiler), EdgeColor.Compile(Compiler), Roughness.Compile(Compiler), Normal.Compile(Compiler));
 }
 
 void UMaterialExpressionStrataConductorBSDF::GetCaption(TArray<FString>& OutCaptions) const
@@ -19021,19 +19021,23 @@ UMaterialExpressionStrataVolumeBSDF::UMaterialExpressionStrataVolumeBSDF(const F
 #if WITH_EDITOR
 int32 UMaterialExpressionStrataVolumeBSDF::Compile(class FMaterialCompiler* Compiler, int32 OutputIndex)
 {
-	if (!Absorption.GetTracedInput().Expression)
+	if (!Albedo.GetTracedInput().Expression)
 	{
-		return Compiler->Errorf(TEXT("Missing Absorption input"));	// STRATA_TODO: defaults in the detail panel with auto read only if plug in
+		return Compiler->Errorf(TEXT("Missing Albedo input"));	// STRATA_TODO: defaults in the detail panel with auto read only if plug in
 	}
-	if (!Scattering.GetTracedInput().Expression)
+	if (!Extinction.GetTracedInput().Expression)
 	{
-		return Compiler->Errorf(TEXT("Missing Scattering input"));	// STRATA_TODO: defaults in the detail panel with auto read only if plug in
+		return Compiler->Errorf(TEXT("Missing Extinction input"));	// STRATA_TODO: defaults in the detail panel with auto read only if plug in
 	}
 	if (!Anisotropy.GetTracedInput().Expression)
 	{
 		return Compiler->Errorf(TEXT("Missing Anisotropy input"));	// STRATA_TODO: defaults in the detail panel with auto read only if plug in
 	}
-	return Compiler->StrataVolumeBSDF(Absorption.Compile(Compiler), Scattering.Compile(Compiler), Anisotropy.Compile(Compiler));
+	return Compiler->StrataVolumeBSDF(
+		Albedo.Compile(Compiler), 
+		Extinction.Compile(Compiler), 
+		Anisotropy.Compile(Compiler), 
+		Thickness.GetTracedInput().Expression ? Thickness.Compile(Compiler) : Compiler->Constant(0.001f)); // default = 1mm
 }
 
 void UMaterialExpressionStrataVolumeBSDF::GetCaption(TArray<FString>& OutCaptions) const
@@ -19059,6 +19063,9 @@ uint32 UMaterialExpressionStrataVolumeBSDF::GetInputType(int32 InputIndex)
 	case 2:
 		return MCT_Float1;
 		break;
+	case 3:
+		return MCT_Float1;
+		break;
 	}
 
 	check(false);
@@ -19312,6 +19319,55 @@ uint32 UMaterialExpressionStrataArtisticIOR::GetInputType(int32 InputIndex)
 
 
 
+UMaterialExpressionStrataPhysicalIOR::UMaterialExpressionStrataPhysicalIOR(const FObjectInitializer& ObjectInitializer)
+	: Super(ObjectInitializer)
+{
+	struct FConstructorStatics
+	{
+		FText NAME_Strata;
+		FConstructorStatics() : NAME_Strata(LOCTEXT("Strata", "Strata")) { }
+	};
+	static FConstructorStatics ConstructorStatics;
+#if WITH_EDITORONLY_DATA
+	MenuCategories.Add(ConstructorStatics.NAME_Strata);
+
+	bShowOutputNameOnPin = true;
+
+	Outputs.Reset();
+	Outputs.Add(FExpressionOutput(TEXT("Reflectivity")));
+	Outputs.Add(FExpressionOutput(TEXT("EdgeColor")));
+#endif
+}
+
+#if WITH_EDITOR
+int32 UMaterialExpressionStrataPhysicalIOR::Compile(class FMaterialCompiler* Compiler, int32 OutputIndex)
+{
+	if (!IOR.GetTracedInput().Expression)
+	{
+		return Compiler->Errorf(TEXT("Missing IOR input"));
+	}
+	if (!Extinction.GetTracedInput().Expression)
+	{
+		return Compiler->Errorf(TEXT("Missing Extinction input"));
+	}
+	return Compiler->StrataPhysicalIOR(IOR.Compile(Compiler), Extinction.Compile(Compiler), OutputIndex);
+}
+
+void UMaterialExpressionStrataPhysicalIOR::GetCaption(TArray<FString>& OutCaptions) const
+{
+	OutCaptions.Add(TEXT("Strata Artistic IOR"));
+}
+
+uint32 UMaterialExpressionStrataPhysicalIOR::GetOutputType(int32 OutputIndex)
+{
+	return MCT_Float3;
+}
+
+uint32 UMaterialExpressionStrataPhysicalIOR::GetInputType(int32 InputIndex)
+{
+	return MCT_Float3;
+}
+#endif // WITH_EDITOR
 
 
 
diff --git a/Engine/Source/Runtime/Engine/Public/MaterialCompiler.h b/Engine/Source/Runtime/Engine/Public/MaterialCompiler.h
--- a/Engine/Source/Runtime/Engine/Public/MaterialCompiler.h
+++ b/Engine/Source/Runtime/Engine/Public/MaterialCompiler.h
@@ -391,13 +391,14 @@ public:
 	virtual int32 StrataDiffuseOrenNayarBSDF(int32 Albedo, int32 Roughness, int32 Normal) = 0;
 	virtual int32 StrataDiffuseChanBSDF(int32 Albedo, int32 Roughness, int32 Normal) = 0;
 	virtual int32 StrataDielectricBSDF(int32 Roughness, int32 IOR, int32 Tint, int32 Normal) = 0;
-	virtual int32 StrataConductorBSDF(int32 IOR, int32 Extinction, int32 Roughness, int32 Normal) = 0;
-	virtual int32 StrataVolumeBSDF(int32 Absorption, int32 Scattering, int32 Anisotropy) = 0;
+	virtual int32 StrataConductorBSDF(int32 Reflectivity, int32 EdgeColor, int32 Roughness, int32 Normal) = 0;
+	virtual int32 StrataVolumeBSDF(int32 Albedo, int32 Extinction, int32 Anisotropy, int32 Thickness) = 0;
 	virtual int32 StrataHorizontalMixing(int32 Foreground, int32 Background, int32 Mix) = 0;
 	virtual int32 StrataVerticalLayering(int32 Top, int32 Base) = 0;
 	virtual int32 StrataAdd(int32 A, int32 B) = 0;
 	virtual int32 StrataMultiply(int32 A, int32 Weight) = 0;
 	virtual int32 StrataArtisticIOR(int32 Reflectivity, int32 EdgeColor, int32 OutputIndex) = 0;
+	virtual int32 StrataPhysicalIOR(int32 IOR, int32 Extinction, int32 OutputIndex) = 0;
 
 
 	virtual int32 MapARPassthroughCameraUV(int32 UV) = 0;
@@ -799,14 +800,14 @@ public:
 		return Compiler->StrataDielectricBSDF(Roughness, IOR, Tint, Normal);
 	}
 
-	virtual int32 StrataConductorBSDF(int32 IOR, int32 Extinction, int32 Roughness, int32 Normal) override
+	virtual int32 StrataConductorBSDF(int32 Reflectivity, int32 EdgeColor, int32 Roughness, int32 Normal) override
 	{
-		return Compiler->StrataConductorBSDF(IOR, Extinction, Roughness, Normal);
+		return Compiler->StrataConductorBSDF(Reflectivity, EdgeColor, Roughness, Normal);
 	}
 
-	virtual int32 StrataVolumeBSDF(int32 Absorption, int32 Scattering, int32 Anisotropy) override
+	virtual int32 StrataVolumeBSDF(int32 Albedo, int32 Extinction, int32 Anisotropy, int32 Thickness) override
 	{
-		return Compiler->StrataVolumeBSDF(Absorption, Scattering, Anisotropy);
+		return Compiler->StrataVolumeBSDF(Albedo, Extinction, Anisotropy, Thickness);
 	}
 
 	virtual int32 StrataHorizontalMixing(int32 Foreground, int32 Background, int32 Mix) override
@@ -834,6 +835,11 @@ public:
 		return Compiler->StrataArtisticIOR(Reflectivity, EdgeColor, OutputIndex);
 	}
 
+	virtual int32 StrataPhysicalIOR(int32 IOR, int32 Extinction, int32 OutputIndex) override
+	{
+		return Compiler->StrataPhysicalIOR(IOR, Extinction, OutputIndex);
+	}
+
 protected:
 		
 	FMaterialCompiler* Compiler;
