Fix typo in maxdist calulation for point and spot light that could lead to undefined behavior

diff --git a/Engine/Shaders/Private/PathTracing/Light/PathTracingPointLight.ush b/Engine/Shaders/Private/PathTracing/Light/PathTracingPointLight.ush
--- a/Engine/Shaders/Private/PathTracing/Light/PathTracingPointLight.ush
+++ b/Engine/Shaders/Private/PathTracing/Light/PathTracingPointLight.ush
@@ -23,7 +23,7 @@ void PointLight_GenerateLightRay(
 
 	// Compute distance to sphere edge
 	float Radius = GetRadius(LightId);
-	float Radius2 = Pow2(Radius2);
+	float Radius2 = Pow2(Radius);
 
 	float RayLength2 = dot(LightDirection, LightDirection);
 	float RayLength = sqrt(RayLength2);
diff --git a/Engine/Shaders/Private/PathTracing/Light/PathTracingSpotLight.ush b/Engine/Shaders/Private/PathTracing/Light/PathTracingSpotLight.ush
--- a/Engine/Shaders/Private/PathTracing/Light/PathTracingSpotLight.ush
+++ b/Engine/Shaders/Private/PathTracing/Light/PathTracingSpotLight.ush
@@ -41,7 +41,7 @@ void SpotLight_GenerateLightRay(
 
 	// Compute distance to sphere edge
 	float Radius = GetRadius(LightId);
-	float Radius2 = Pow2(Radius2);
+	float Radius2 = Pow2(Radius);
 
 	float RayLength2 = dot(LightDirection, LightDirection);
 	float RayLength = sqrt(RayLength2);
