Cherrypick HairStrandsVisibilityRasterCompute.usf fix from //UE5/Release-Engine-Staging

diff --git a/Engine/Shaders/Private/HairStrands/HairStrandsVisibilityRasterCompute.usf b/Engine/Shaders/Private/HairStrands/HairStrandsVisibilityRasterCompute.usf
--- a/Engine/Shaders/Private/HairStrands/HairStrandsVisibilityRasterCompute.usf
+++ b/Engine/Shaders/Private/HairStrands/HairStrandsVisibilityRasterCompute.usf
@@ -51,8 +51,9 @@ PackedType WriteVisibilityData(RWTexture2D<PackedType> OutTexture, uint2 Coords,
 	// TODO
 	return (PackedType)0;
 #elif PERMUTATION_RASTER_ATOMIC == RASTER_PLATFORM && COMPILER_SUPPORTS_UINT64_IMAGE_ATOMICS && COMPILER_SUPPORTS_ULONG_TYPES
-	PackedType PrevData = ImageInterlockedMaxUInt64(OutTexture, Coords, PackedData);
-	return PrevData;
+	//PackedType PrevData = ImageInterlockedMaxUInt64(OutTexture, Coords, PackedData);
+	//return PrevData;
+	return (PackedType)0;
 #else
 	return (PackedType)0;
 #endif
