Fixing Final Gather sample counting bug.

diff --git a/Engine/Shaders/Private/RayTracing/RayTracingFinalGatherRGS.usf b/Engine/Shaders/Private/RayTracing/RayTracingFinalGatherRGS.usf
--- a/Engine/Shaders/Private/RayTracing/RayTracingFinalGatherRGS.usf
+++ b/Engine/Shaders/Private/RayTracing/RayTracingFinalGatherRGS.usf
@@ -196,6 +196,7 @@ RAY_TRACING_ENTRY_RAYGEN(RayTracingFinalGatherRGS)
 		{
 			continue;
 		}
+		ValidSamplesPerPixel++;
 
 		float3 GPPosition = GatherSample.Position;
 		float3 GPIrradiance = GatherSample.Irradiance;
@@ -219,7 +220,6 @@ RAY_TRACING_ENTRY_RAYGEN(RayTracingFinalGatherRGS)
 
 		// Firefly prevention attempt
 		if (RayPdf < 1.0e-3) continue;
-		ValidSamplesPerPixel++;
 
 		// Create ray
 		RayDesc Ray;
@@ -303,6 +303,7 @@ RAY_TRACING_ENTRY_RAYGEN(RayTracingFinalGatherRGS)
 				float3 LocalWorldNormal = DecodeNormal(GBufferATexture.Load(int3(LocalPixelCoord, 0)).xyz);
 				if (ShouldBilateralReject(LocalPixelCoord, DeviceZ, LocalDeviceZ, WorldNormal, LocalWorldNormal)) continue;
 
+				ValidSamplesPerPixel++;
 				uint GatherPointsIndex = DispatchThreadIdToLinearIndex(GlobalIndex);
 				FGatherSample GatherSample = ReadGatherSample(GatherPointsBuffer, GatherPointsIndex, ParentSampleIndex, GatherPointsResolution.x * GatherPointsResolution.y);
 				float3 GPPosition = GatherSample.Position;
@@ -326,7 +327,9 @@ RAY_TRACING_ENTRY_RAYGEN(RayTracingFinalGatherRGS)
 				float Weight = TargetPdf / SourcePdf;
 				float RayPdf = TargetPdf;
 
-				ValidSamplesPerPixel++;
+				// Firefly prevention attempt
+				if (RayPdf < 1.0e-3) continue;
+
 #if USE_NEIGHBOR_VISIBILITY_TEST
 				// Create ray
 				RayDesc Ray;
