Fixed compilation error on a platform due to Strata.

diff --git a/Engine/Shaders/Private/Strata/Strata.ush b/Engine/Shaders/Private/Strata/Strata.ush
--- a/Engine/Shaders/Private/Strata/Strata.ush
+++ b/Engine/Shaders/Private/Strata/Strata.ush
@@ -484,7 +484,7 @@ void PackStrataOut(
 
 	// Now write out the BSDFs
 
-	BSDFCount = min(BSDFCount, STRATA_MAX_TOTAL_BSDF);
+	BSDFCount = min(BSDFCount, uint(STRATA_MAX_TOTAL_BSDF));
 	STRATA_STORE_UINT1(BSDFCount);
 
 	BSDFCount = 0;
