Strata translucent now uses dual source blending.
Laverage current dual source blending code.
Rename a define to not make it only about thin translucent.
Make weight/transmittance/throughput/coverage more well defined. Throughput not takes into account non coveraed part.

diff --git a/Engine/Shaders/Private/BasePassPixelShader.usf b/Engine/Shaders/Private/BasePassPixelShader.usf
--- a/Engine/Shaders/Private/BasePassPixelShader.usf
+++ b/Engine/Shaders/Private/BasePassPixelShader.usf
@@ -74,7 +74,7 @@
 #include "Lumen/LumenTranslucencyVolumeShared.ush"
 #endif
 
-#define STRATA_ENABLED					PROJECT_STRATA
+#define STRATA_ENABLED					PROJECT_STRATA && MATERIAL_IS_STRATA
 #define STRATA_OPAQUE_ENABLED			(STRATA_ENABLED && MATERIALBLENDING_ANY_TRANSLUCENT<=0)
 #define STRATA_TRANSLUCENT_ENABLED		(STRATA_ENABLED && MATERIALBLENDING_ANY_TRANSLUCENT>0)
 
@@ -115,6 +115,7 @@ float NormalCurvatureToRoughness(float3 WorldNormal)
 #define FORCE_FULLY_ROUGH (SIMPLE_FORWARD_SHADING || MATERIAL_FULLY_ROUGH)
 #define EDITOR_ALPHA2COVERAGE (USE_EDITOR_COMPOSITING && SUPPORTS_PIXEL_COVERAGE)
 #define POST_PROCESS_SUBSURFACE ((MATERIAL_SHADINGMODEL_SUBSURFACE_PROFILE || MATERIAL_SHADINGMODEL_EYE) && USES_GBUFFER)
+#define MATERIAL_WORKS_WITH_DUAL_SOURCE_COLOR_BLENDING (MATERIAL_SHADINGMODEL_THIN_TRANSLUCENT || STRATA_TRANSLUCENT_ENABLED)
 
 #if  MATERIAL_SHADINGMODEL_SINGLELAYERWATER
 
@@ -1427,11 +1428,12 @@ void FPixelShaderInOut_MainPS(
 	}
 #endif // MATERIAL_SHADINGMODEL_SINGLELAYERWATER
 
-#if MATERIAL_SHADINGMODEL_THIN_TRANSLUCENT
-
+#if MATERIAL_WORKS_WITH_DUAL_SOURCE_COLOR_BLENDING
 	float3 DualBlendColorAdd = 0.0f;
 	float3 DualBlendColorMul = 1.0f;
+#endif
 
+#if MATERIAL_SHADINGMODEL_THIN_TRANSLUCENT
 	{
 		AccumulateThinTranslucentModel(
 								DualBlendColorAdd,
@@ -1446,8 +1448,7 @@ void FPixelShaderInOut_MainPS(
 		Color = 0;
 		Opacity = 1.0f;
 	}
-
-#endif
+#endif // MATERIAL_SHADINGMODEL_THIN_TRANSLUCENT
 
 
 #if STRATA_ENABLED && !MATERIAL_SHADINGMODEL_UNLIT
@@ -1456,7 +1457,6 @@ void FPixelShaderInOut_MainPS(
 
 #if STRATA_OPAQUE_ENABLED
 
-#if MATERIAL_IS_STRATA
 	// Need to cancel that to make sure strata material are only lit using Strata lighting passes
 	Color = 0;
 	// Clear only what should be and keep the rest
@@ -1466,7 +1466,6 @@ void FPixelShaderInOut_MainPS(
 	GBuffer.BaseColor = 0;
 	GBuffer.Specular = 0;
 	GBuffer.Metallic = 0;
-#endif
 
 	OpaqueBasePass.Strata.MaterialLobesTextureUAV[uint2(In.SvPosition.xy)] = Strata.LayerCount; // Could store a bit pack of layer count, BSDF type masks for tiling, etc.
 
@@ -1474,20 +1473,7 @@ void FPixelShaderInOut_MainPS(
 	uint PixelStrataDataByteOffset = (PixelPos.y * ResolvedView.BufferSizeAndInvSize.x + PixelPos.x) * OpaqueBasePass.Strata.MaxBytesPerPixel;
 
 	PackStrataOut(OpaqueBasePass.Strata.MaterialLobesBufferUAV, PixelStrataDataByteOffset, Strata, MaterialParameters.WorldNormal, MaterialParameters.CameraVector);
-
-	//switch (uint(Strata.LayerCount))
-	//{
-	//case 1:
-	//	Color = float3(0, 0, 1);	break;
-	//case 2:
-	//	Color = float3(0, 1, 0);	break;
-	//case 3:
-	//	Color = float3(1, 1, 0);	break;
-	//case 4:
-	//	Color = float3(1, 0, 0);	break;
-	//}
-
-#endif
+#endif // STRATA_OPAQUE_ENABLED
 
 #if STRATA_TRANSLUCENT_ENABLED
 	//FORWARD_SHADING
@@ -1497,6 +1483,7 @@ void FPixelShaderInOut_MainPS(
 		float3 V = MaterialParameters.CameraVector;
 		float3 LightColor = ResolvedView.DirectionalLightColor.rgb * PI;
 
+		float3 Throughput = 1.0f;
 		Color = StrataForwardLighting(
 			GBuffer,
 			MaterialParameters.WorldNormal,
@@ -1505,11 +1492,17 @@ void FPixelShaderInOut_MainPS(
 			LightColor,
 			DirectionalLightShadow,
 			IndirectOcclusion,
-			Strata);
+			Strata,
+			Throughput);
+
+		DualBlendColorAdd = Color;
+		DualBlendColorMul = Throughput;	// TODO transmittance
+		Color = 0.0f;
+		Opacity = 1.0f;
 	}
-#endif
+#endif // STRATA_TRANSLUCENT_ENABLED
 
-#endif
+#endif // STRATA_ENABLED
 
 
 	#if MATERIAL_DOMAIN_POSTPROCESS
@@ -1520,7 +1513,7 @@ void FPixelShaderInOut_MainPS(
 		#endif
 		Out.MRT[0] = RETURN_COLOR(Out.MRT[0]);
 	// MATERIAL_SHADINGMODEL_THIN_TRANSLUCENT must come first because it also has MATERIALBLENDING_TRANSLUCENT defined
-	#elif MATERIAL_SHADINGMODEL_THIN_TRANSLUCENT
+	#elif MATERIAL_WORKS_WITH_DUAL_SOURCE_COLOR_BLENDING
 		// After thin translucency, the final color is going to be:
 		//    FinalColor = DualBlendColorAdd + DualBlendColorMul * BackgroundColor;
 	    // To apply fogging, we want the final equation to be:
@@ -1534,7 +1527,7 @@ void FPixelShaderInOut_MainPS(
 		float3 AdjustedDualBlendAdd = Fogging.rgb + Fogging.a * DualBlendColorAdd;
 		float3 AdjustedDualBlendMul =               Fogging.a * DualBlendColorMul;
 
-		#if THIN_TRANSLUCENT_USE_DUAL_BLEND
+		#if DUAL_SOURCE_COLOR_BLENDING_ENABLED
 			// no RETURN_COLOR because these values are explicit multiplies and adds
 			Out.MRT[0] = half4(AdjustedDualBlendAdd,0.0);
 			Out.MRT[1] = half4(AdjustedDualBlendMul,1.0);
@@ -1710,7 +1703,7 @@ void FPixelShaderInOut_MainPS(
 #else //USES_GBUFFER
 	#define PIXELSHADEROUTPUT_MRT0 1
 	// we also need MRT for thin translucency due to dual blending if we are not on the fallback path
-	#define PIXELSHADEROUTPUT_MRT1 (WRITES_VELOCITY_TO_GBUFFER || (MATERIAL_SHADINGMODEL_THIN_TRANSLUCENT && THIN_TRANSLUCENT_USE_DUAL_BLEND))
+	#define PIXELSHADEROUTPUT_MRT1 (WRITES_VELOCITY_TO_GBUFFER || (DUAL_SOURCE_COLOR_BLENDING_ENABLED && MATERIAL_WORKS_WITH_DUAL_SOURCE_COLOR_BLENDING))
 #endif //USES_GBUFFER
 #define PIXELSHADEROUTPUT_A2C ((EDITOR_ALPHA2COVERAGE) != 0)
 #define PIXELSHADEROUTPUT_COVERAGE (MATERIALBLENDING_MASKED_USING_COVERAGE && !EARLY_Z_PASS_ONLY_MATERIAL_MASKING)
diff --git a/Engine/Shaders/Private/Strata/Strata.ush b/Engine/Shaders/Private/Strata/Strata.ush
--- a/Engine/Shaders/Private/Strata/Strata.ush
+++ b/Engine/Shaders/Private/Strata/Strata.ush
@@ -71,7 +71,7 @@ struct FStrataVolume
 struct FStrataBSDF
 {
 	int					BSDFType;
-	float3				Weight;
+	float3				Weight;		// Weight is a float3 because, when loaded from the material buffer, it contains the BSDF throughput (coverage and transmittance combined)
 
 	// Even thought the compiler should optimise, it would be nice to have a union with reinterpretation of the register memory.
 	FStrataDiffuse		Diffuse;
@@ -391,18 +391,22 @@ void PackStrataOut(
 	//	- We could clean up irrelevant BSDFs 
 	//	- We could generate transmittance BSDFs 
 
+	const float OpaqueBSDFThroughput = 0.0f;
+	const float FullThroughput = 1.0f;
 
 	uint BSDFCount = 0;
-	float3 TotalTransmittance = 1.0f;
+	float3 Throughput = 1.0f;
 
 	int l;
 	for (l = 0; l < Strata.LayerCount; ++l)
 	{
-		float3 LayerTransmittance = 0.0f;
+		float3 LayerThroughput = 0.0f;
+		float3 LayerCoverage = 0.0f;
 
 		for (int i = 0; i < Strata.Layers[l].BSDFCount; ++i)
 		{
 			FStrataBSDF BSDF = Strata.Layers[l].BSDFs[i];
+			const float3 BSDFCoverage = BSDF.Weight;
 			BSDFCount++;
 
 			float3 N = MeshWorldNormal;
@@ -428,13 +432,15 @@ void PackStrataOut(
 			{
 			case STRATA_BSDF_TYPE_DIFFUSE_ON:
 			{
-				//LayerTransmittance += 0.0f;
+				LayerCoverage	+= BSDFCoverage;
+				LayerThroughput += BSDFCoverage * OpaqueBSDFThroughput;
 				break;
 			}
 
 			case STRATA_BSDF_TYPE_DIFFUSE_CHAN:
 			{
-				//LayerTransmittance += 0.0f;
+				LayerCoverage	+= BSDFCoverage;
+				LayerThroughput += BSDFCoverage * OpaqueBSDFThroughput;
 				break;
 			}
 
@@ -443,7 +449,8 @@ void PackStrataOut(
 				float F0 = DielectricIorToF0(BSDF.Dielectric.IOR);
 				float3 F = F_Schlick(F0, ContextNoLight.VoH);
 
-				LayerTransmittance += BSDF.Weight * (1.0f - F); // STRATA_TODO use 1-FGD instead of simple 1-F, event if we ignore ray scattering/spreading
+				LayerCoverage	+= BSDFCoverage;
+				LayerThroughput += BSDFCoverage * (1.0f - F); // STRATA_TODO use 1-FGD instead of simple 1-F, event if we ignore ray scattering/spreading
 				break;
 			}
 
@@ -451,7 +458,8 @@ void PackStrataOut(
 			{
 				float3 F = F_Schlick(BSDF.Conductor.Reflectivity, BSDF.Conductor.EdgeColor, ContextNoLight.VoH);
 
-				LayerTransmittance += BSDF.Weight * (1.0f - F); // STRATA_TODO transmittance should only come from coverage
+				LayerCoverage	+= BSDFCoverage;
+				LayerThroughput += BSDFCoverage * OpaqueBSDFThroughput;
 				break;
 			}
 
@@ -469,16 +477,21 @@ void PackStrataOut(
 				const float PathLength = Thickness / max(0.0001f, abs(ContextNoLight.NoV));
 				const float3 SafePathSegmentTransmittance = exp(-SafeExtinctionCoefficients * PathLength);
 
-				LayerTransmittance += BSDF.Weight * SafePathSegmentTransmittance;
+				LayerCoverage	+= BSDFCoverage;
+				LayerThroughput += BSDFCoverage * SafePathSegmentTransmittance;
 				break;
 			}
 			} // switch
 
-			// Update the weights to take into accounr transmittance towards the view
-			Strata.Layers[l].BSDFs[i].Weight *= TotalTransmittance;
+			// Update the weights to take into account throughput towards the view from layers above
+			Strata.Layers[l].BSDFs[i].Weight *= Throughput;
 		}
 
-		TotalTransmittance *= LayerTransmittance;	// STRATA_TODO we need to go over all weight and trnasmittance ot make sure everything would behave corretly. Especially with a Add node for instance. Add saturate where neded, etc.
+		// Take into account the non covered part for this layer and update throughput accordingly
+		LayerThroughput += (1.0 - saturate(LayerCoverage)) * FullThroughput;
+
+		// And combine this layer throughtput with he material throughput
+		Throughput *= LayerThroughput;
 	}
 
 
diff --git a/Engine/Shaders/Private/Strata/StrataDeferredLighting.ush b/Engine/Shaders/Private/Strata/StrataDeferredLighting.ush
--- a/Engine/Shaders/Private/Strata/StrataDeferredLighting.ush
+++ b/Engine/Shaders/Private/Strata/StrataDeferredLighting.ush
@@ -21,6 +21,7 @@ float4 StrataDeferredLighting(
 	for (int i = 0; i < PixelBSDFs.BSDFCount; ++i)
 	{
 		FStrataBSDF BSDF = PixelBSDFs.BSDFs[i];
+		const float3 BSDFThroughput = BSDF.Weight;
 
 		float3 N = GBuffer.WorldNormal;
 		switch (BSDF.BSDFType)		// TODO_STRATA: all have a normal, even subsurface?
@@ -51,7 +52,7 @@ float4 StrataDeferredLighting(
 			UglyIntegrationFactor = IntegrateLight(Capsule, LightData.bInverseSquared);
 		}
 
-		const float3 DirLightBaseFactor = BSDF.Weight * LightData.Color * LightMask * SurfaceShadow * UglyIntegrationFactor;
+		const float3 DirLightBaseFactor = BSDFThroughput * LightData.Color * LightMask * SurfaceShadow * UglyIntegrationFactor;
 		const float3 DirLightGeometryFactor = DirLightBaseFactor * SatNoL;
 
 		switch (BSDF.BSDFType)
diff --git a/Engine/Shaders/Private/Strata/StrataEnvironmentLighting.ush b/Engine/Shaders/Private/Strata/StrataEnvironmentLighting.ush
--- a/Engine/Shaders/Private/Strata/StrataEnvironmentLighting.ush
+++ b/Engine/Shaders/Private/Strata/StrataEnvironmentLighting.ush
@@ -19,6 +19,7 @@ float3 StrataEnvironmentLighting(
 	for (int i = 0; i < PixelBSDFs.BSDFCount; ++i)
 	{
 		FStrataBSDF BSDF = PixelBSDFs.BSDFs[i];
+		const float3 BSDFThroughput = BSDF.Weight;
 
 		BxDFContext Context = (BxDFContext)0;
 		float3 L = N;
@@ -109,10 +110,10 @@ float3 StrataEnvironmentLighting(
 		if (any(Albedo > 0.0f))
 		{
 			// BentNormal?
-			SkyLighting += BSDF.Weight * Diffuse_Lambert(Albedo) * GetSkySHDiffuse(Normal) * View.SkyLightColor.rgb * AmbientOcclusion;
+			SkyLighting += BSDFThroughput * Diffuse_Lambert(Albedo) * GetSkySHDiffuse(Normal) * View.SkyLightColor.rgb * AmbientOcclusion;
 		}
 
-		SkyLighting += BSDF.Weight * SpecularLuminance * View.SkyLightColor.rgb * AmbientOcclusion;
+		SkyLighting += BSDFThroughput * SpecularLuminance * View.SkyLightColor.rgb * AmbientOcclusion;
 	}
 #if USE_PREEXPOSURE
 	SkyLighting *= View.PreExposure;
diff --git a/Engine/Shaders/Private/Strata/StrataForwardLighting.ush b/Engine/Shaders/Private/Strata/StrataForwardLighting.ush
--- a/Engine/Shaders/Private/Strata/StrataForwardLighting.ush
+++ b/Engine/Shaders/Private/Strata/StrataForwardLighting.ush
@@ -12,19 +12,25 @@ float3 StrataForwardLighting(
 	float3 LColor,
 	float3 LShadow,
 	float3 IndirectOcclusion,
-	FStrataData Strata
+	FStrataData Strata,
+	inout float3 OutThroughput
 	)
 {
+	const float OpaqueBSDFThroughput = 0.0f;
+	const float FullThroughput = 1.0f;
+
 	float3 Color = 0;
-	float3 TotalTransmittance = 1.0f;
+	OutThroughput = 1.0f;
 
 	for (int l = 0; l < Strata.LayerCount; ++l)
 	{
-		float3 LayerTransmittance = 0.0f;
+		float3 LayerThroughput = 0.0f;
+		float3 LayerCoverage = 0.0f;
 
 		for (int i = 0; i < Strata.Layers[l].BSDFCount; ++i)
 		{
 			FStrataBSDF BSDF = Strata.Layers[l].BSDFs[i];
+			const float3 BSDFCoverage = BSDF.Weight;
 
 			float3 N = MaterialParametersWorldNormal;
 			switch (BSDF.BSDFType)		// TODO_STRATA: all have a normal, even subsurface?
@@ -44,9 +50,9 @@ float3 StrataForwardLighting(
 			Init(Context, N, V, L);
 			float SatNoL = saturate(Context.NoL);
 
-			const float3 DirLightBaseFactor = TotalTransmittance * BSDF.Weight * LColor * LShadow;
+			const float3 DirLightBaseFactor = OutThroughput * BSDFCoverage * LColor * LShadow;
 			const float3 DirLightGeometryFactor = DirLightBaseFactor * SatNoL;
-			const float3 SkyLightCommonFactor = TotalTransmittance * BSDF.Weight * ResolvedView.SkyLightColor.rgb * IndirectOcclusion;
+			const float3 SkyLightCommonFactor = OutThroughput * BSDFCoverage * ResolvedView.SkyLightColor.rgb * IndirectOcclusion;
 
 			switch (BSDF.BSDFType)
 			{
@@ -63,7 +69,8 @@ float3 StrataForwardLighting(
 					}
 #endif
 
-				//LayerTransmittance += 0.0f;
+				LayerCoverage	+= BSDFCoverage;
+				LayerThroughput += BSDFCoverage * OpaqueBSDFThroughput;
 				break;
 			}
 
@@ -81,8 +88,9 @@ float3 StrataForwardLighting(
 						Color += GetEffectiveSkySHDiffuse(N) * SkyLightCommonFactor * Diffuse_Lambert(BSDF.Diffuse.Albedo); // STRATA_TODO should not be lambert
 					}
 #endif
-
-				//LayerTransmittance += 0.0f;
+				
+				LayerCoverage	+= BSDFCoverage;
+				LayerThroughput += BSDFCoverage * OpaqueBSDFThroughput;
 				break;
 			}
 
@@ -110,8 +118,9 @@ float3 StrataForwardLighting(
 						Color += GetSkyLightReflection(R, SafeRoughness, SkyAverageBrightness) * EvalEnvBRDF * BSDF.Dielectric.Tint * SkyLightCommonFactor;
 					}
 #endif
-
-				LayerTransmittance += BSDF.Weight * (1.0f - F); // STRATA_TODO transmittance should only come from coverage
+				
+				LayerCoverage	+= BSDFCoverage;
+				LayerThroughput += BSDFCoverage * (1.0f - F); // STRATA_TODO (Kulla's (1-FDG))
 				break;
 			}
 
@@ -139,8 +148,9 @@ float3 StrataForwardLighting(
 					Color += GetSkyLightReflection(R, SafeRoughness, SkyAverageBrightness) * EvalEnvBRDF * SkyLightCommonFactor;
 				}
 #endif
-
-				LayerTransmittance += BSDF.Weight * (1.0f - F); // STRATA_TODO use 1-FGD instead of simple 1-F, event if we ignore ray scattering/spreading
+				
+				LayerCoverage	+= BSDFCoverage;
+				LayerThroughput += BSDFCoverage * OpaqueBSDFThroughput;
 				break;
 			}
 
@@ -186,13 +196,18 @@ float3 StrataForwardLighting(
 				}
 #endif
 
-				LayerTransmittance += BSDF.Weight * SafePathSegmentTransmittance;
+				LayerCoverage	+= BSDFCoverage;
+				LayerThroughput += BSDFCoverage * SafePathSegmentTransmittance;
 				break;
 			}
 			}
 		}
 
-		TotalTransmittance *= LayerTransmittance;	// STRATA_TODO we need to go over all weight and trnasmittance ot make sure everything would behave corretly. Especially with a Add node for instance. Add saturate where neded, etc.
+		// Take into account the non covered part for this layer and update throughput accordingly
+		LayerThroughput += (1.0 - saturate(LayerCoverage)) * FullThroughput;
+
+		// And combine this layer throughtput with he material throughput
+		OutThroughput *= LayerThroughput;
 	}
 	return Color;
 }
diff --git a/Engine/Source/Runtime/Engine/Classes/Materials/Material.h b/Engine/Source/Runtime/Engine/Classes/Materials/Material.h
--- a/Engine/Source/Runtime/Engine/Classes/Materials/Material.h
+++ b/Engine/Source/Runtime/Engine/Classes/Materials/Material.h
@@ -1760,18 +1760,17 @@ public:
 	bool HasBaseColorConnected() const { return BaseColor.IsConnected(); }
 	bool HasRoughnessConnected() const { return Roughness.IsConnected(); }
 	bool HasAmbientOcclusionConnected() const { return AmbientOcclusion.IsConnected(); }
-	bool HasStrataFrontMaterialConnected() const { return FrontMaterial.IsConnected(); }
 #else	
 	// Add to runtime data only if we need to call these at runtime
 	bool HasBaseColorConnected() const { check(0); return false; }
 	bool HasRoughnessConnected() const { check(0); return false; }
 	bool HasAmbientOcclusionConnected() const { check(0); return false; }
-	bool HasStrataFrontMaterialConnected() const { check(0); return false; }
 #endif 	
 	bool HasNormalConnected() const { return Normal.IsConnected(); }
 	bool HasSpecularConnected() const { return Specular.IsConnected(); }
 	bool HasEmissiveColorConnected() const { return EmissiveColor.IsConnected(); }
 	bool HasAnisotropyConnected() const { return Anisotropy.IsConnected(); }
+	bool HasStrataFrontMaterialConnected() const { return FrontMaterial.IsConnected(); }
 
 #if WITH_EDITOR
 	static void NotifyCompilationFinished(UMaterialInterface* Material);
diff --git a/Engine/Source/Runtime/Engine/Private/Materials/HLSLMaterialTranslator.cpp b/Engine/Source/Runtime/Engine/Private/Materials/HLSLMaterialTranslator.cpp
--- a/Engine/Source/Runtime/Engine/Private/Materials/HLSLMaterialTranslator.cpp
+++ b/Engine/Source/Runtime/Engine/Private/Materials/HLSLMaterialTranslator.cpp
@@ -938,6 +938,14 @@ bool FHLSLMaterialTranslator::Translate()
 			}
 		}
 
+		if (Material->IsStrataMaterial())
+		{
+			if (!(BlendMode == BLEND_Translucent || BlendMode == BLEND_Opaque))
+			{
+				Errorf(TEXT("Strata materials must be opaque or translucent."));
+			}
+		}
+
 		bool bDBufferAllowed = IsUsingDBuffers(Platform);
 		bool bDBufferBlendMode = IsDBufferDecalBlendMode((EDecalBlendMode)Material->GetDecalBlendMode());
 
@@ -1130,6 +1138,8 @@ bool FHLSLMaterialTranslator::Translate()
 
 void FHLSLMaterialTranslator::GetMaterialEnvironment(EShaderPlatform InPlatform, FShaderCompilerEnvironment& OutEnvironment)
 {
+	bool bMaterialRequestsDualSourceBlending = false;
+
 	if (bNeedsParticlePosition || Material->ShouldGenerateSphericalParticleNormals() || bUsesSphericalParticleOpacity)
 	{
 		OutEnvironment.SetDefine(TEXT("NEEDS_PARTICLE_POSITION"), 1);
@@ -1333,11 +1343,7 @@ void FHLSLMaterialTranslator::GetMaterialEnvironment(EShaderPlatform InPlatform,
 			OutEnvironment.SetDefine(TEXT("MATERIAL_SHADINGMODEL_THIN_TRANSLUCENT"), TEXT("1"));
 			NumSetMaterials++;
 
-			// if it is not enabled, it will fall back to standard alpha blending
-			if (Material->IsDualBlendingEnabled(Platform))
-			{
-				OutEnvironment.SetDefine(TEXT("THIN_TRANSLUCENT_USE_DUAL_BLEND"), TEXT("1"));
-			}
+			bMaterialRequestsDualSourceBlending = true;
 		}
 
 		if(ShadingModels.HasShadingModel(MSM_SingleLayerWater) && 
@@ -1412,9 +1418,12 @@ void FHLSLMaterialTranslator::GetMaterialEnvironment(EShaderPlatform InPlatform,
 		}
 	}
 
-	OutEnvironment.SetDefine(TEXT("MATERIAL_IS_STRATA"),
-		Material->HasStrataFrontMaterialConnected() ? TEXT("1") : TEXT("0"));
-	 
+	const bool bStrataMaterial = Material->IsStrataMaterial();
+	OutEnvironment.SetDefine(TEXT("MATERIAL_IS_STRATA"), bStrataMaterial ? TEXT("1") : TEXT("0"));
+	bMaterialRequestsDualSourceBlending |= bStrataMaterial;
+
+	// if duals source blending (colored transmittance) is not supported on a platform, it will fall back to standard alpha blending (grey scale transmittance)
+	OutEnvironment.SetDefine(TEXT("DUAL_SOURCE_COLOR_BLENDING_ENABLED"), bMaterialRequestsDualSourceBlending && Material->IsDualBlendingEnabled(Platform) ? TEXT("1") : TEXT("0"));
 }
 
 // Assign custom interpolators to slots, packing them as much as possible in unused slots.
diff --git a/Engine/Source/Runtime/Engine/Private/Materials/MaterialShared.cpp b/Engine/Source/Runtime/Engine/Private/Materials/MaterialShared.cpp
--- a/Engine/Source/Runtime/Engine/Private/Materials/MaterialShared.cpp
+++ b/Engine/Source/Runtime/Engine/Private/Materials/MaterialShared.cpp
@@ -1209,9 +1209,9 @@ bool FMaterialResource::IsTranslucencyAfterDOFEnabled() const
 
 bool FMaterialResource::IsDualBlendingEnabled(EShaderPlatform Platform) const
 {
-	const bool bIsMaterialEnabled = (Material->ShadingModel == MSM_ThinTranslucent);
+	const bool bMaterialRequestsDualSourceBlending = Material->ShadingModel == MSM_ThinTranslucent || IsStrataMaterial();
 	const bool bIsPlatformSupported = RHISupportsDualSourceBlending(Platform);
-	return bIsMaterialEnabled && bIsPlatformSupported;
+	return bMaterialRequestsDualSourceBlending && bIsPlatformSupported;
 }
 
 bool FMaterialResource::IsMobileSeparateTranslucencyEnabled() const
@@ -1423,8 +1423,9 @@ bool FMaterialResource::HasAmbientOcclusionConnected() const
 	return HasMaterialAttributesConnected() || Material->HasAmbientOcclusionConnected();
 }
 
-bool FMaterialResource::HasStrataFrontMaterialConnected() const
+bool FMaterialResource::IsStrataMaterial() const
 {
+	// STRATA_TODO IsStrataMaterial should go away once Strata implementation is finished
 	const URendererSettings* RendererSettings = GetDefault<URendererSettings>();
 	return RendererSettings && RendererSettings->bEnableStrata ? Material->HasStrataFrontMaterialConnected() : false;
 }
diff --git a/Engine/Source/Runtime/Engine/Public/MaterialShared.h b/Engine/Source/Runtime/Engine/Public/MaterialShared.h
--- a/Engine/Source/Runtime/Engine/Public/MaterialShared.h
+++ b/Engine/Source/Runtime/Engine/Public/MaterialShared.h
@@ -1625,7 +1625,7 @@ public:
 	virtual bool HasEmissiveColorConnected() const { return false; }
 	virtual bool HasAnisotropyConnected() const { return false; }
 	virtual bool HasAmbientOcclusionConnected() const { return false; }
-	virtual bool HasStrataFrontMaterialConnected() const { return false; }
+	virtual bool IsStrataMaterial() const { return false; }
 	virtual bool RequiresSynchronousCompilation() const { return false; };
 	virtual bool IsDefaultMaterial() const { return false; };
 	virtual int32 GetNumCustomizedUVs() const { return 0; }
@@ -2435,7 +2435,7 @@ public:
 	ENGINE_API virtual bool HasEmissiveColorConnected() const override;
 	ENGINE_API virtual bool HasAnisotropyConnected() const override;
 	ENGINE_API virtual bool HasAmbientOcclusionConnected() const override;
-	ENGINE_API virtual bool HasStrataFrontMaterialConnected() const override;
+	ENGINE_API virtual bool IsStrataMaterial() const override;
 	ENGINE_API virtual FMaterialShadingModelField GetShadingModels() const override;
 	ENGINE_API virtual bool IsShadingModelFromMaterialExpression() const override;
 	ENGINE_API virtual enum ETranslucencyLightingMode GetTranslucencyLightingMode() const override;
diff --git a/Engine/Source/Runtime/Renderer/Private/BasePassRendering.cpp b/Engine/Source/Runtime/Renderer/Private/BasePassRendering.cpp
--- a/Engine/Source/Runtime/Renderer/Private/BasePassRendering.cpp
+++ b/Engine/Source/Runtime/Renderer/Private/BasePassRendering.cpp
@@ -148,7 +148,8 @@ static bool IsBasePassWaitForTasksEnabled()
 
 void SetTranslucentRenderState(FMeshPassProcessorRenderState& DrawRenderState, const FMaterial& Material, const EShaderPlatform Platform, ETranslucencyPass::Type InTranslucencyPassType)
 {
-	if (Material.GetShadingModels().HasShadingModel(MSM_ThinTranslucent))
+	// STRATA_TODO IsStrataMaterial will be removed when all materials are unified to use premultipled alpha
+	if (Material.GetShadingModels().HasShadingModel(MSM_ThinTranslucent) || Material.IsStrataMaterial())
 	{
 		// Special case for dual blending, which is not exposed as a parameter in the material editor
 		if (Material.IsDualBlendingEnabled(Platform))
