Fixes eferred hair shadow in UE5

diff --git a/Engine/Shaders/Private/DeferredLightPixelShaders.usf b/Engine/Shaders/Private/DeferredLightPixelShaders.usf
--- a/Engine/Shaders/Private/DeferredLightPixelShaders.usf
+++ b/Engine/Shaders/Private/DeferredLightPixelShaders.usf
@@ -212,16 +212,10 @@ void DeferredLightPixelMain(
 				uint PixelStrataDataByteOffset = (PixelPos.y * View.BufferSizeAndInvSize.x + PixelPos.x) * Strata.MaxBytesPerPixel;
 				FStrataPixelBSDFs PixelBSDFs = UnpackStrataIn(Strata.MaterialLobesBuffer, PixelStrataDataByteOffset);
 
-				FShadowTerms Shadow;
-				Shadow.SurfaceShadow = ScreenSpaceData.AmbientOcclusion;
-				Shadow.TransmissionShadow = 1;
-				Shadow.TransmissionThickness = 1;
-				//TODO Sebastien Hillaire
-				//Shadow.HairTransmittance.Transmittance = 1;
-				Shadow.HairTransmittance.OpaqueVisibility = 1;
+				FShadowTerms ShadowTerms = { ScreenSpaceData.AmbientOcclusion, 1.0, 1.0, InitHairTransmittanceData() };
 				float4 LightAttenuation = GetPerPixelLightAttenuation(InputParams.ScreenUV);
-				GetShadowTerms(ScreenSpaceData.GBuffer, LightData, DerivedParams.WorldPosition, L, LightAttenuation, Dither, Shadow);
-				SurfaceShadow = Shadow.SurfaceShadow;
+				GetShadowTerms(ScreenSpaceData.GBuffer, LightData, DerivedParams.WorldPosition, L, LightAttenuation, Dither, ShadowTerms);
+				SurfaceShadow = ShadowTerms.SurfaceShadow;
 
 				Radiance = StrataDeferredLighting(
 					LightData,
