Fixed issue with Strata g volume parameter.

diff --git a/Engine/Shaders/Private/Strata/Strata.ush b/Engine/Shaders/Private/Strata/Strata.ush
--- a/Engine/Shaders/Private/Strata/Strata.ush
+++ b/Engine/Shaders/Private/Strata/Strata.ush
@@ -541,7 +541,7 @@ void PackStrataOut(
 			{
 				STRATA_STORE_UINT1(PackR11G11B10F(BSDF.Volume.Albedo));
 				STRATA_STORE_UINT1(PackR11G11B10F(BSDF.Volume.Extinction));
-				STRATA_STORE_UINT1(PackR11G11B10F(float3(BSDF.Volume.Anisotropy, BSDF.Volume.Thickness, 0.0f)));
+				STRATA_STORE_UINT1(PackR11G11B10F(float3(abs(BSDF.Volume.Anisotropy), BSDF.Volume.Thickness, BSDF.Volume.Anisotropy > 0.0f ? 2.0f : 0.0f)));
 				// 12 bytes
 			}
 			break;
@@ -609,7 +609,10 @@ FStrataPixelBSDFs UnpackStrataIn(ByteAddressBuffer StrataBuffer, uint ByteOffset
 		{
 			StrataLoad_R11G11B10F			(StrataBuffer, ByteOffset, PixelBSDFs.BSDFs[i].Volume.Albedo);
 			StrataLoad_R11G11B10F			(StrataBuffer, ByteOffset, PixelBSDFs.BSDFs[i].Volume.Extinction);
-			StrataLoad_R11_G11_B10F			(StrataBuffer, ByteOffset, PixelBSDFs.BSDFs[i].Volume.Anisotropy, PixelBSDFs.BSDFs[i].Volume.Thickness, DummyFloat);
+
+			float SignOfAnisotropy = 0.0f;
+			StrataLoad_R11_G11_B10F			(StrataBuffer, ByteOffset, PixelBSDFs.BSDFs[i].Volume.Anisotropy, PixelBSDFs.BSDFs[i].Volume.Thickness, SignOfAnisotropy);
+			PixelBSDFs.BSDFs[i].Volume.Anisotropy *= (SignOfAnisotropy - 1.0f);
 		}
 		break;
 		}
