Added a tool to debug the Strata material buffer

diff --git a/Engine/Shaders/Private/Strata/StrataVisualize.usf b/Engine/Shaders/Private/Strata/StrataVisualize.usf
new file mode 100644
--- /dev/null
+++ b/Engine/Shaders/Private/Strata/StrataVisualize.usf
@@ -0,0 +1,157 @@
+// Copyright Epic Games, Inc. All Rights Reserved.
+
+#include "/Engine/Private/Common.ush"
+
+#include "/Engine/Private/Strata/Strata.ush"
+
+#include "/Engine/Private/MiniFontCommon.ush"
+
+void VisualizeMaterialPS(
+	float4 SVPos : SV_POSITION,
+	out float4 OutColor : SV_Target0)
+{
+	const int2 PixelPos = SVPos.xy;
+	const int2 MaterialPos = View.CursorPosition;
+	const float InspectSquare = 10.0f;
+
+	OutColor = float4(0.0f, 0.0f, 0.0f, 1.0f);
+
+	if (all(float2(PixelPos) > float2(MaterialPos + InspectSquare)))
+	{
+		OutColor.w = 0.0f;				// Set background visiblity to 0 in the bottom right part of the screen to better reads the debug values
+	}
+
+	if (Strata.MaterialLobesTexture[MaterialPos].x > 0.0f)
+	{
+		uint PixelStrataDataByteOffset = (MaterialPos.y * View.BufferSizeAndInvSize.x + MaterialPos.x) * Strata.MaxBytesPerPixel;
+		FStrataPixelBSDFs PixelBSDFs = UnpackStrataIn(Strata.MaterialLobesBuffer, PixelStrataDataByteOffset);
+
+		float3 BSDFFontColor = float3(0.25, 1.0, 0.25);
+		float3 INFOFontColor = float3(1.0, 0.5, 0.25);
+		float3 PARAMFontColor = float3(0.8, 0.8, 0.8);
+		float3 FontColor = float3(0.25, 1.0, 1.0);
+		int2 TopLeftInit = MaterialPos + InspectSquare + int2(3,3);
+		int2 TopLeft = TopLeftInit;
+
+#define PRINTPARAMCHAR(X) PrintCharacter(PixelPos, OutColor.xyz, PARAMFontColor, TopLeft, X);
+#define PRINTINFOCHAR(X) PrintCharacter(PixelPos, OutColor.xyz, INFOFontColor, TopLeft, X);
+#define PRINTBSDFCHAR(X) PrintCharacter(PixelPos, OutColor.xyz, BSDFFontColor, TopLeft, X);
+
+#define PRINTBSDFWEIGHT()	PRINTBSDFCHAR(_SPC_); PRINTBSDFCHAR(_MINUS_); PRINTBSDFCHAR(_SPC_); PrintFloat(PixelPos, OutColor.xyz, FontColor, TopLeft, BSDF.Weight.x); TopLeft.x+=8*7; PRINTBSDFCHAR(_SPC_); PrintFloat(PixelPos, OutColor.xyz, FontColor, TopLeft, BSDF.Weight.y); TopLeft.x+=8*7; PRINTBSDFCHAR(_SPC_); PrintFloat(PixelPos, OutColor.xyz, FontColor, TopLeft, BSDF.Weight.z);
+#define PRINTFLOAT3(X) PrintFloat(PixelPos, OutColor.xyz, PARAMFontColor, TopLeft, X.x); TopLeft.x+=8*7; PRINTBSDFCHAR(_SPC_); PrintFloat(PixelPos, OutColor.xyz, PARAMFontColor, TopLeft, X.y); TopLeft.x+=8*7; PRINTBSDFCHAR(_SPC_); PrintFloat(PixelPos, OutColor.xyz, PARAMFontColor, TopLeft, X.z);
+#define PRINTFLOAT1(Value) PrintFloat(PixelPos, OutColor.xyz, PARAMFontColor, TopLeft, Value); TopLeft.x+=8*7;
+#define SIGNOFFLOAT1(Value) if(Value < 0.0f) { PrintCharacter(PixelPos, OutColor.xyz, PARAMFontColor, TopLeft, _MINUS_); } /*else { PrintCharacter(PixelPos, OutColor.xyz, PARAMFontColor, TopLeft, _SPC_);}*/
+
+#define NEXTLINE  TopLeft.y += 8; TopLeft.x = TopLeftInit.x;
+#define HORIZONTALSEPARATION  TopLeft.y += 4;
+#define HORIZONTALTAB  TopLeft.x += 8;
+
+		PRINTINFOCHAR(_T_); PRINTINFOCHAR(_O_); PRINTINFOCHAR(_P_);
+		NEXTLINE;
+		HORIZONTALSEPARATION;
+
+		for (int i = 0; i < PixelBSDFs.BSDFCount; ++i)
+		{
+			FStrataBSDF BSDF = PixelBSDFs.BSDFs[i];
+
+			TopLeft.x = TopLeftInit.x;
+
+			switch (BSDF.BSDFType)
+			{
+			case STRATA_BSDF_TYPE_DIFFUSE_ON:
+				PRINTBSDFCHAR(_D_); PRINTBSDFCHAR(_I_); PRINTBSDFCHAR(_F_); PRINTBSDFCHAR(_F_); PRINTBSDFCHAR(_O_); PRINTBSDFCHAR(_N_);
+				PRINTBSDFCHAR(_SPC_); PRINTBSDFCHAR(_SPC_); PRINTBSDFCHAR(_SPC_); PRINTBSDFCHAR(_SPC_); PRINTBSDFCHAR(_SPC_); PRINTBSDFCHAR(_SPC_);
+				PRINTBSDFWEIGHT();
+				NEXTLINE;
+
+				HORIZONTALTAB; PRINTPARAMCHAR(_A_); PRINTPARAMCHAR(_L_); PRINTPARAMCHAR(_B_); PRINTPARAMCHAR(_E_); PRINTPARAMCHAR(_D_); PRINTPARAMCHAR(_O_);
+				HORIZONTALTAB; PRINTFLOAT3(BSDF.Diffuse.Albedo);
+				NEXTLINE;
+				HORIZONTALTAB; PRINTPARAMCHAR(_R_); PRINTPARAMCHAR(_O_); PRINTPARAMCHAR(_U_); PRINTPARAMCHAR(_G_); PRINTPARAMCHAR(_H_); PRINTPARAMCHAR(_SPC_);
+				HORIZONTALTAB; PRINTFLOAT1(BSDF.Diffuse.Roughness);
+				NEXTLINE;
+
+				break;
+
+			case STRATA_BSDF_TYPE_DIFFUSE_CHAN:
+				PRINTBSDFCHAR(_D_); PRINTBSDFCHAR(_I_); PRINTBSDFCHAR(_F_); PRINTBSDFCHAR(_F_);
+				PRINTBSDFCHAR(_C_); PRINTBSDFCHAR(_H_); PRINTBSDFCHAR(_A_); PRINTBSDFCHAR(_N_);
+				PRINTBSDFWEIGHT();
+				NEXTLINE;
+
+				HORIZONTALTAB; PRINTPARAMCHAR(_A_); PRINTPARAMCHAR(_L_); PRINTPARAMCHAR(_B_); PRINTPARAMCHAR(_E_); PRINTPARAMCHAR(_D_); PRINTPARAMCHAR(_O_); 
+				HORIZONTALTAB; PRINTFLOAT3(BSDF.Diffuse.Albedo);
+				NEXTLINE;
+				HORIZONTALTAB; PRINTPARAMCHAR(_R_); PRINTPARAMCHAR(_O_); PRINTPARAMCHAR(_U_); PRINTPARAMCHAR(_G_); PRINTPARAMCHAR(_H_); PRINTPARAMCHAR(_SPC_);
+				HORIZONTALTAB; PRINTFLOAT1(BSDF.Diffuse.Roughness);
+				NEXTLINE;
+
+				break;
+
+			case STRATA_BSDF_TYPE_DIELECTRIC:
+				PRINTBSDFCHAR(_D_); PRINTBSDFCHAR(_I_); PRINTBSDFCHAR(_E_); PRINTBSDFCHAR(_L_); PRINTBSDFCHAR(_E_);
+				PRINTBSDFCHAR(_C_); PRINTBSDFCHAR(_T_); PRINTBSDFCHAR(_R_); PRINTBSDFCHAR(_I_); PRINTBSDFCHAR(_C_);
+				PRINTBSDFWEIGHT();
+				NEXTLINE;
+
+				HORIZONTALTAB; PRINTPARAMCHAR(_I_); PRINTPARAMCHAR(_O_); PRINTPARAMCHAR(_R_); PRINTPARAMCHAR(_SPC_); PRINTPARAMCHAR(_SPC_); PRINTPARAMCHAR(_SPC_);
+				HORIZONTALTAB; PRINTFLOAT1(BSDF.Dielectric.IOR);
+				NEXTLINE;
+				HORIZONTALTAB; PRINTPARAMCHAR(_T_); PRINTPARAMCHAR(_I_); PRINTPARAMCHAR(_N_); PRINTPARAMCHAR(_T_); PRINTPARAMCHAR(_SPC_); PRINTPARAMCHAR(_SPC_);
+				HORIZONTALTAB; PRINTFLOAT3(BSDF.Dielectric.Tint);
+				NEXTLINE;
+				HORIZONTALTAB; PRINTPARAMCHAR(_R_); PRINTPARAMCHAR(_O_); PRINTPARAMCHAR(_U_); PRINTPARAMCHAR(_G_); PRINTPARAMCHAR(_H_); PRINTPARAMCHAR(_SPC_);
+				HORIZONTALTAB; PRINTFLOAT1(BSDF.Dielectric.Roughness.x);
+				NEXTLINE;
+
+				break;
+
+			case STRATA_BSDF_TYPE_CONDUCTOR:
+				PRINTBSDFCHAR(_C_); PRINTBSDFCHAR(_O_); PRINTBSDFCHAR(_N_); PRINTBSDFCHAR(_D_); PRINTBSDFCHAR(_U_);
+				PRINTBSDFCHAR(_C_); PRINTBSDFCHAR(_T_); PRINTBSDFCHAR(_O_); PRINTBSDFCHAR(_R_); PRINTBSDFCHAR(_SPC_);
+				PRINTBSDFWEIGHT();
+				NEXTLINE;
+
+				HORIZONTALTAB; PRINTPARAMCHAR(_R_); PRINTPARAMCHAR(_E_); PRINTPARAMCHAR(_F_); PRINTPARAMCHAR(_L_); PRINTPARAMCHAR(_E_); PRINTPARAMCHAR(_C_);
+				HORIZONTALTAB; PRINTFLOAT3(BSDF.Conductor.Reflectivity);
+				NEXTLINE;
+				HORIZONTALTAB; PRINTPARAMCHAR(_E_); PRINTPARAMCHAR(_D_); PRINTPARAMCHAR(_G_); PRINTPARAMCHAR(_E_); PRINTPARAMCHAR(_C_); PRINTPARAMCHAR(_SPC_);
+				HORIZONTALTAB; PRINTFLOAT3(BSDF.Conductor.EdgeColor);
+				NEXTLINE;
+				HORIZONTALTAB; PRINTPARAMCHAR(_R_); PRINTPARAMCHAR(_O_); PRINTPARAMCHAR(_U_); PRINTPARAMCHAR(_G_); PRINTPARAMCHAR(_H_); PRINTPARAMCHAR(_SPC_);
+				HORIZONTALTAB; PRINTFLOAT1(BSDF.Conductor.Roughness.x);
+				NEXTLINE;
+
+				break;
+
+			case STRATA_BSDF_TYPE_VOLUME:
+				PRINTBSDFCHAR(_V_); PRINTBSDFCHAR(_O_); PRINTBSDFCHAR(_L_); PRINTBSDFCHAR(_U_); PRINTBSDFCHAR(_M_); 
+				PRINTBSDFCHAR(_E_); PRINTBSDFCHAR(_SPC_); PRINTBSDFCHAR(_SPC_); PRINTBSDFCHAR(_SPC_); PRINTBSDFCHAR(_SPC_);
+				PRINTBSDFWEIGHT();
+				NEXTLINE;
+				
+				HORIZONTALTAB; PRINTPARAMCHAR(_A_); PRINTPARAMCHAR(_L_); PRINTPARAMCHAR(_B_); PRINTPARAMCHAR(_E_); PRINTPARAMCHAR(_D_); PRINTPARAMCHAR(_O_); 
+				HORIZONTALTAB; PRINTFLOAT3(BSDF.Volume.Albedo);
+				NEXTLINE;
+				HORIZONTALTAB; PRINTPARAMCHAR(_E_); PRINTPARAMCHAR(_X_); PRINTPARAMCHAR(_T_); PRINTPARAMCHAR(_I_); PRINTPARAMCHAR(_N_); PRINTPARAMCHAR(_C_); 
+				HORIZONTALTAB; PRINTFLOAT3(BSDF.Volume.Extinction);
+				NEXTLINE;
+				HORIZONTALTAB; PRINTPARAMCHAR(_A_); PRINTPARAMCHAR(_N_); PRINTPARAMCHAR(_I_); PRINTPARAMCHAR(_S_); PRINTPARAMCHAR(_O_); PRINTPARAMCHAR(_T_); 
+				HORIZONTALTAB; SIGNOFFLOAT1(BSDF.Volume.Anisotropy); PRINTFLOAT1(abs(BSDF.Volume.Anisotropy));
+				NEXTLINE;
+				HORIZONTALTAB; PRINTPARAMCHAR(_T_); PRINTPARAMCHAR(_H_); PRINTPARAMCHAR(_I_); PRINTPARAMCHAR(_C_); PRINTPARAMCHAR(_K_); PRINTPARAMCHAR(_N_); 
+				HORIZONTALTAB; PRINTFLOAT1(BSDF.Volume.Thickness*1000.0f); PRINTPARAMCHAR(_SPC_); PRINTPARAMCHAR(_M_); PRINTPARAMCHAR(_I_); PRINTPARAMCHAR(_L_); PRINTPARAMCHAR(_L_); PRINTPARAMCHAR(_I_);
+				NEXTLINE;
+
+				break;
+			}
+
+			HORIZONTALSEPARATION;
+		}
+
+		PRINTINFOCHAR(_B_); PRINTINFOCHAR(_O_); PRINTINFOCHAR(_T_); PRINTINFOCHAR(_T_); PRINTINFOCHAR(_O_); PRINTINFOCHAR(_M_);
+		NEXTLINE;
+	}
+}
+
+
diff --git a/Engine/Source/Runtime/Engine/Public/ShowFlags.h b/Engine/Source/Runtime/Engine/Public/ShowFlags.h
--- a/Engine/Source/Runtime/Engine/Public/ShowFlags.h
+++ b/Engine/Source/Runtime/Engine/Public/ShowFlags.h
@@ -452,6 +452,7 @@ private:
 		SetVisualizeCalibrationGrayscale(false);
 		SetVisualizeCalibrationCustom(false);
 		SetVirtualTexturePrimitives(false);
+		SetVisualizeStrataMaterial(false);
 	}
 
 
diff --git a/Engine/Source/Runtime/Engine/Public/ShowFlagsValues.inl b/Engine/Source/Runtime/Engine/Public/ShowFlagsValues.inl
--- a/Engine/Source/Runtime/Engine/Public/ShowFlagsValues.inl
+++ b/Engine/Source/Runtime/Engine/Public/ShowFlagsValues.inl
@@ -371,5 +371,8 @@ SHOWFLAG_FIXED_IN_SHIPPING(0, VisualizeCalibrationCustom, SFG_Developer, NSLOCTE
 /** Draw in the main pass the primitives that would normally only contribute to runtime virtual textures. */
 SHOWFLAG_FIXED_IN_SHIPPING(0, VirtualTexturePrimitives, SFG_Developer, NSLOCTEXT("UnrealEd", "VirtualTexturePrimitivesSF", "Virtual Texture Primitives"))
 
+/** Debug the Strata material buffer content */
+SHOWFLAG_FIXED_IN_SHIPPING(0, VisualizeStrataMaterial, SFG_Visualize, NSLOCTEXT("UnrealEd", "StrataMaterial", "Strata Material"))
+
 #undef SHOWFLAG_ALWAYS_ACCESSIBLE
 #undef SHOWFLAG_FIXED_IN_SHIPPING
diff --git a/Engine/Source/Runtime/Renderer/Private/DeferredShadingRenderer.cpp b/Engine/Source/Runtime/Renderer/Private/DeferredShadingRenderer.cpp
--- a/Engine/Source/Runtime/Renderer/Private/DeferredShadingRenderer.cpp
+++ b/Engine/Source/Runtime/Renderer/Private/DeferredShadingRenderer.cpp
@@ -2009,6 +2009,8 @@ void FDeferredShadingSceneRenderer::Render(FRHICommandListImmediate& RHICmdList)
 	}
 
 	// Strata initialisation
+	const bool bStrataEnabled = Strata::IsStrataEnabled();
+	if (bStrataEnabled)
 	{		
 		Strata::InitialiseStrataFrameSceneData(*this, GraphBuilder);		
 	}
@@ -2647,6 +2649,11 @@ void FDeferredShadingSceneRenderer::Render(FRHICommandListImmediate& RHICmdList)
 		}
 	}
 
+	if (bStrataEnabled)
+	{
+		Strata::AddVisualizeMaterialPasses(GraphBuilder, Views, SceneColorTexture.Target);
+	}
+
 	if (bCanOverlayRayTracingOutput && ViewFamily.EngineShowFlags.LightShafts)
 	{
 		SCOPE_CYCLE_COUNTER(STAT_FDeferredShadingSceneRenderer_RenderLightShaftBloom);
diff --git a/Engine/Source/Runtime/Renderer/Private/SceneRendering.cpp b/Engine/Source/Runtime/Renderer/Private/SceneRendering.cpp
--- a/Engine/Source/Runtime/Renderer/Private/SceneRendering.cpp
+++ b/Engine/Source/Runtime/Renderer/Private/SceneRendering.cpp
@@ -3414,7 +3414,7 @@ void FSceneRenderer::OnStartRender(FRHICommandListImmediate& RHICmdList)
 
 bool FSceneRenderer::ShouldCompositeEditorPrimitives(const FViewInfo& View)
 {
-	if (View.Family->EngineShowFlags.VisualizeHDR || View.Family->UseDebugViewPS())
+	if (View.Family->EngineShowFlags.VisualizeHDR || View.Family->EngineShowFlags.VisualizeStrataMaterial || View.Family->UseDebugViewPS())
 	{
 		// certain visualize modes get obstructed too much
 		return false;
diff --git a/Engine/Source/Runtime/Renderer/Private/Strata/Strata.cpp b/Engine/Source/Runtime/Renderer/Private/Strata/Strata.cpp
--- a/Engine/Source/Runtime/Renderer/Private/Strata/Strata.cpp
+++ b/Engine/Source/Runtime/Renderer/Private/Strata/Strata.cpp
@@ -2,6 +2,7 @@
 
 #include "Strata.h"
 #include "HAL/IConsoleManager.h"
+#include "PixelShaderUtils.h"
 #include "SceneView.h"
 #include "ScenePrivate.h"
 #include "SceneRendering.h"
@@ -116,6 +117,73 @@ TUniformBufferRef<FStrataGlobalUniformParameters> BindStrataGlobalUniformParamet
 	return CreateUniformBufferImmediate(StrataUniformParameters, UniformBuffer_SingleDraw);
 }
 
+
+////////////////////////////////////////////////////////////////////////// Debug
+
+
+class FVisualizeMaterialPS : public FGlobalShader
+{
+	DECLARE_GLOBAL_SHADER(FVisualizeMaterialPS);
+	SHADER_USE_PARAMETER_STRUCT(FVisualizeMaterialPS, FGlobalShader);
+
+	using FPermutationDomain = TShaderPermutationDomain<>;
+
+	BEGIN_SHADER_PARAMETER_STRUCT(FParameters, )
+		SHADER_PARAMETER_STRUCT_REF(FViewUniformShaderParameters, ViewUniformBuffer)
+		SHADER_PARAMETER_STRUCT_REF(FStrataGlobalUniformParameters, Strata)
+		SHADER_PARAMETER_TEXTURE(Texture2D, MiniFontTexture)
+		RENDER_TARGET_BINDING_SLOTS()
+	END_SHADER_PARAMETER_STRUCT()
+
+	static FPermutationDomain RemapPermutation(FPermutationDomain PermutationVector)
+	{
+		return PermutationVector;
+	}
+
+	static bool ShouldCompilePermutation(const FGlobalShaderPermutationParameters& Parameters)
+	{
+		return GetMaxSupportedFeatureLevel(Parameters.Platform) >= ERHIFeatureLevel::SM5;
+	}
+
+	static void ModifyCompilationEnvironment(const FGlobalShaderPermutationParameters& Parameters, FShaderCompilerEnvironment& OutEnvironment)
+	{
+		FGlobalShader::ModifyCompilationEnvironment(Parameters, OutEnvironment);
+	}
+};
+IMPLEMENT_GLOBAL_SHADER(FVisualizeMaterialPS, "/Engine/Private/Strata/StrataVisualize.usf", "VisualizeMaterialPS", SF_Pixel);
+
+void AddVisualizeMaterialPasses(FRDGBuilder& GraphBuilder, const TArray<FViewInfo>& Views, FRDGTextureRef SceneColorTexture)
+{
+	RDG_EVENT_SCOPE_CONDITIONAL(GraphBuilder, IsStrataEnabled() && Views.Num() > 0, "StrataVisualizeMaterial");
+	if (!IsStrataEnabled())
+	{
+		return;
+	}
+
+	FRHIBlendState* PreMultipliedColorTransmittanceBlend = TStaticBlendState<CW_RGB, BO_Add, BF_One, BF_SourceAlpha, BO_Add, BF_Zero, BF_One>::GetRHI();
+
+	for (int32 i = 0; i < Views.Num(); ++i)
+	{
+		const FViewInfo& View = Views[i];
+
+		FVisualizeMaterialPS::FPermutationDomain PermutationVector;
+		TShaderMapRef<FVisualizeMaterialPS> PixelShader(View.ShaderMap, PermutationVector);
+
+		if (View.Family->EngineShowFlags.VisualizeStrataMaterial > 0)
+		{
+			FVisualizeMaterialPS::FParameters* PassParameters = GraphBuilder.AllocParameters<FVisualizeMaterialPS::FParameters>();
+			PassParameters->ViewUniformBuffer = View.ViewUniformBuffer;
+			PassParameters->Strata = Strata::BindStrataGlobalUniformParameters(View);
+			PassParameters->MiniFontTexture = GetMiniFontTexture();
+			PassParameters->RenderTargets[0] = FRenderTargetBinding(SceneColorTexture, ERenderTargetLoadAction::ELoad);
+
+			FPixelShaderUtils::AddFullscreenPass<FVisualizeMaterialPS>(GraphBuilder, View.ShaderMap, RDG_EVENT_NAME("StrataVisualizeMaterial"),
+				PixelShader, PassParameters, View.ViewRect, PreMultipliedColorTransmittanceBlend);
+		}
+	}
+}
+
+
 } // namespace Strata
 
 
diff --git a/Engine/Source/Runtime/Renderer/Private/Strata/Strata.h b/Engine/Source/Runtime/Renderer/Private/Strata/Strata.h
--- a/Engine/Source/Runtime/Renderer/Private/Strata/Strata.h
+++ b/Engine/Source/Runtime/Renderer/Private/Strata/Strata.h
@@ -54,6 +54,8 @@ void BindStrataBasePassUniformParameters(const FViewInfo& View, FStrataOpaquePas
 
 TUniformBufferRef<FStrataGlobalUniformParameters> BindStrataGlobalUniformParameters(const FViewInfo& View);
 
+void AddVisualizeMaterialPasses(FRDGBuilder& GraphBuilder, const TArray<FViewInfo>& Views, FRDGTextureRef SceneColorTexture);
+
 };
 
 
