Disabling some Starta code not compiling on Switch.
Added todo for later.

diff --git a/Engine/Shaders/Private/Strata/Strata.ush b/Engine/Shaders/Private/Strata/Strata.ush
--- a/Engine/Shaders/Private/Strata/Strata.ush
+++ b/Engine/Shaders/Private/Strata/Strata.ush
@@ -320,7 +320,29 @@ float3 UnpackAlbedoGamma2ToLinear(uint rgb)
 	return Gamma2 * Gamma2;	// Simple gamma 2.0
 }
 
+#if COMPILER_SWITCH // These platforms are complaining so simply disabling Strata there for now.
 
+#define STRATA_STORE_UINT1(x) 
+#define STRATA_STORE_UINT2(x) 
+#define STRATA_STORE_UINT3(x) 
+#define STRATA_STORE_UINT4(x) 
+
+#define STRATA_LOAD_UINT1(x) x = 0;
+#define STRATA_LOAD_UINT2(x) x = 0;
+#define STRATA_LOAD_UINT3(x) x = 0;
+#define STRATA_LOAD_UINT4(x) x = 0;
+
+#define STRATA_LOAD_FLOAT1(x) x = 0;
+#define STRATA_LOAD_FLOAT2(x) x = 0;
+#define STRATA_LOAD_FLOAT3(x) x = 0;
+#define STRATA_LOAD_FLOAT4(x) x = 0;
+
+uint  StrataLoadUint1(ByteAddressBuffer StrataBuffer, inout uint ByteOffset) { return 0; }
+uint2 StrataLoadUint2(ByteAddressBuffer StrataBuffer, inout uint ByteOffset) { return 0; }
+uint3 StrataLoadUint3(ByteAddressBuffer StrataBuffer, inout uint ByteOffset) { return 0; }
+uint4 StrataLoadUint4(ByteAddressBuffer StrataBuffer, inout uint ByteOffset) { return 0; }
+
+#else
 
 #define STRATA_STORE_UINT1(x) StrataBuffer.Store (ByteOffset, asuint(x)); ByteOffset += 4;
 #define STRATA_STORE_UINT2(x) StrataBuffer.Store2(ByteOffset, asuint(x)); ByteOffset += 8;
@@ -342,6 +364,8 @@ uint2 StrataLoadUint2(ByteAddressBuffer StrataBuffer, inout uint ByteOffset) { u
 uint3 StrataLoadUint3(ByteAddressBuffer StrataBuffer, inout uint ByteOffset) { uint3 x = StrataBuffer.Load3(ByteOffset); ByteOffset +=12; return x; }
 uint4 StrataLoadUint4(ByteAddressBuffer StrataBuffer, inout uint ByteOffset) { uint4 x = StrataBuffer.Load4(ByteOffset); ByteOffset +=16; return x; }
 
+#endif
+
 void StrataLoad_R16G16B16F_A16F(ByteAddressBuffer StrataBuffer, inout uint ByteOffset, inout float3 A, inout float B)
 {
 	float4 temp = UnpackR16G16B16A16F(StrataLoadUint2(StrataBuffer, ByteOffset));
