Implement Multiple Importance Sampling in PathTracer for all light types

Refactored the light API to have Trace/Sample methods that perform all work at once to reduce redundant operations across methods and avoid numerical issues.

Reworked the main path tracer loop to have a single trace call and less state maintained across loop iterations.

Flip the normal if viewing a shape from the wrong side to prevent light leaks

Expose a CVar to control what kind of sampling is done in the PathTracer (Material / Light / Both)

Expose a CVar to control if lights should be visible to camera rays (the skydome is always visible to camera to maintain the existing behavior).

Fixed a numerical precision issue in the DefaultLit material sampling function.

Temporary fix for lobe selection in DefaultLit to avoid under-sampling of unused lobes.

diff --git a/Engine/Shaders/Private/MonteCarlo.ush b/Engine/Shaders/Private/MonteCarlo.ush
--- a/Engine/Shaders/Private/MonteCarlo.ush
+++ b/Engine/Shaders/Private/MonteCarlo.ush
@@ -224,6 +224,30 @@ float4 UniformSampleCone( float2 E, float CosThetaMax )
 	return float4( L, PDF );
 }
 
+// Same as the function above, but uses SinThetaMax^2 as the parameter
+// so that the solid angle can be computed more accurately for very small angles
+// The caller is expected to ensure that SinThetaMax2 is <= 1
+float4 UniformSampleConeRobust(float2 E, float SinThetaMax2)
+{
+	float Phi = 2 * PI * E.x;
+	// The expression 1-sqrt(1-x) is susceptible to catastrophic cancelation.
+	// Instead, use a series expansion about 0 which is accurate within 10^-7
+	// and much more numerically stable.
+	float OneMinusCosThetaMax = SinThetaMax2 < 0.01 ? SinThetaMax2 * (0.5 + 0.125 * SinThetaMax2) : 1 - sqrt(1 - SinThetaMax2);
+
+	float CosTheta = 1 - OneMinusCosThetaMax * E.y;
+	float SinTheta = sqrt(1 - CosTheta * CosTheta);
+
+	float3 L;
+	L.x = SinTheta * cos(Phi);
+	L.y = SinTheta * sin(Phi);
+	L.z = CosTheta;
+
+	float PDF = 1.0 / (2 * PI * OneMinusCosThetaMax);
+
+	return float4(L, PDF);
+}
+
 float4 ImportanceSampleBlinn( float2 E, float a2 )
 {
 	float n = 2 / a2 - 2;
@@ -323,6 +347,42 @@ float MISWeight( uint Num, float PDF, uint OtherNum, float OtherPDF )
 	return Weight * Weight / (Weight * Weight + OtherWeight * OtherWeight);
 }
 
+// Multiple importance sampling power heuristic of two functions with a power of two. 
+float MISWeightRobust(float Pdf, float OtherPdf) {
+	// The straightforward implementation above is prone to numerical overflow and divisions by 0
+	// and does not work well with +inf inputs.
+
+	// We want this function to have the following properties:
+	//  0 <= w(a,b) <= 1 for all possible positive floats a and b (including 0 and +inf)
+	//  w(a, b) + w(b, a) == 1.0
+	
+	// The formulation below is much more stable across the range of all possible inputs
+	// and guarantees the sum always adds up to 1.0.
+
+	if (Pdf == OtherPdf)
+	{
+		// Catch potential NaNs from (0,0) and (+inf, +inf)
+		return 0.5f;
+	}
+
+	// Evaluate the expression using the ratio of the smaller value to the bigger one for greater
+	// numerical stability. The math would also work using the ratio of bigger to smaller value,
+	// which would underflow less but would make the weights asymmetric. Underflow to 0 is not a
+	// bad property to have in rendering application as it ensures more weights are exactly 0
+	// which allows some evaluations to be skipped.
+	if (OtherPdf < Pdf)
+	{
+		float x = OtherPdf / Pdf;
+		return 1.0 / (1.0 + x * x);
+	}
+	else
+	{
+		// this form guarantees the weights add back up to one when arguments are swapped
+		float x = Pdf / OtherPdf;
+		return 1.0 - 1.0 / (1.0 + x * x);
+	}
+}
+
 // Compute PDF of reflection (2 * dot( V, H ) * H - V).
 // [Heitz 2014, "Importance Sampling Microfacet-Based BSDFs using the Distribution of Visible Normals"]
 float RayPDFToReflectionRayPDF(float VoH, float RayPDF)
diff --git a/Engine/Shaders/Private/PathTracing/Light/PathTracingDirectionalLight.ush b/Engine/Shaders/Private/PathTracing/Light/PathTracingDirectionalLight.ush
--- a/Engine/Shaders/Private/PathTracing/Light/PathTracingDirectionalLight.ush
+++ b/Engine/Shaders/Private/PathTracing/Light/PathTracingDirectionalLight.ush
@@ -8,6 +8,54 @@ DirectionalLight.usf: Light sampling functions for Directional light implementat
 
 #include "PathTracingLightCommon.ush"
 
+float4 DirectionalLight_TraceLight(RayDesc Ray, int LightId, float HitT)
+{
+	if (Ray.TMax == HitT) // we hit the background
+	{
+		float3 LightDirection = GetNormal(LightId);
+		float SinThetaMax = GetRadius(LightId); // Sin(Angle / 2)
+		float SinThetaMax2 = SinThetaMax * SinThetaMax;
+		// See explanation in UniformSampleConeRobust
+		float OneMinusCosThetaMax = SinThetaMax2 < 0.01 ? SinThetaMax2 * (0.5 + 0.125 * SinThetaMax2) : 1 - sqrt(1 - SinThetaMax2);
+
+		float CosTheta = saturate(dot(normalize(Ray.Direction), LightDirection));
+
+		// Is ray pointing inside the cone of directions?
+		if (1 - CosTheta < OneMinusCosThetaMax)
+		{
+			// We divide the color of the light by the solid angle so
+			// that the overall illumination stays constant as the angle changes.
+			float SolidAngle = 2 * PI * OneMinusCosThetaMax;
+			float3 Radiance = GetColor(LightId) / SolidAngle;
+			return float4(Radiance, 1.0 / SolidAngle);
+		}
+	}
+	return 0.0;
+}
+
+bool DirectionalLight_SampleLight(
+	int LightId,
+	float4 RandSample,
+	float3 WorldPos,
+	float3 WorldNormal,
+	out float3 OutDirection,
+	out float  OutDistance,
+	out float3 OutRadianceOverPdf,
+	out float OutPdf)
+{
+	float SinThetaMax = GetRadius(LightId); // currently sin(Angle / 2)
+	float4 DirAndPdf = UniformSampleConeRobust(RandSample.yz, SinThetaMax * SinThetaMax);
+
+	OutDirection = TangentToWorld(DirAndPdf.xyz, GetNormal(LightId));
+	OutDistance = 1.0e27;
+
+	// Because the light is normalized by the solid angle, the ratio is just the color
+	OutRadianceOverPdf = GetColor(LightId);
+	OutPdf = DirAndPdf.w;
+
+	return true;
+}
+
 void DirectionalLight_GenerateLightRay(
 	RayDesc Ray,
 	int LightId,
@@ -17,7 +65,7 @@ void DirectionalLight_GenerateLightRay(
 {
 	OutLightRay.Origin = Ray.Origin;
 	OutLightRay.TMin = 0.0;
-	OutLightRay.Direction = normalize(LightUV);
+	OutLightRay.Direction = LightUV;
 	OutLightRay.TMax = 1.0e27;
 }
 
@@ -48,20 +96,24 @@ void DirectionalLight_EvalLight(
 	{
 		// #dxr_todo: Merge Eval/Sample to avoid redundant computation
 		float SinThetaMax = Radius; // currently sin(Angle / 2)
-		float CosThetaMax = sqrt(1.0 - SinThetaMax * SinThetaMax);
-		OutRadiance = Color * 1.0 / (2 * PI * (1 - CosThetaMax));
+
+		// See logic in UniformSampleConeRobust
+		// #dxr_todo: unify the sample and eval methods into one to avoid redundant calculations
+		float SinThetaMax2 = SinThetaMax * SinThetaMax;
+		float OneMinusCosThetaMax = SinThetaMax2 < 0.01 ? SinThetaMax2 * (0.5 + 0.125 * SinThetaMax2) : 1 - sqrt(1 - SinThetaMax2);
+		OutRadiance = Color * 1.0 / (2 * PI * OneMinusCosThetaMax);
 	}
 }
 
 void DirectionalLight_EstimateLight(
 	int LightId,
-	RayDesc Ray,
-	FMaterialClosestHitPayload Payload,
+	float3 WorldPos,
+	float3 WorldNormal,
 	out float OutIrradiance
 )
 {
 	float3 LightDirection = GetNormal(LightId);
-	float3 SurfaceNormal = Payload.WorldNormal;
+	float3 SurfaceNormal = WorldNormal;
 	float NoL = 1.0;
 	if (GetRadius(LightId) == 0.0)
 	{
@@ -82,7 +134,7 @@ void DirectionalLight_SampleLight(
 	out float OutPdf
 )
 {
-	float3 N = normalize(GetNormal(LightId));
+	float3 N = GetNormal(LightId);
 	float Radius = GetRadius(LightId);
 
 	if (Radius == 0)
@@ -96,9 +148,8 @@ void DirectionalLight_SampleLight(
 	{
 		// soft case - build a frame, and sample the distant cone
 		float SinThetaMax = Radius; // currently sin(Angle / 2)
-		float CosThetaMax = sqrt(1.0 - SinThetaMax * SinThetaMax);
 
-		float4 DirAndPdf = UniformSampleCone(RandSample.yz, CosThetaMax);
+		float4 DirAndPdf = UniformSampleConeRobust(RandSample.yz, SinThetaMax * SinThetaMax);
 		OutLightUV = TangentToWorld(DirAndPdf.xyz, N);
 		OutPdf = DirAndPdf.w;
 	}
diff --git a/Engine/Shaders/Private/PathTracing/Light/PathTracingLightCommon.ush b/Engine/Shaders/Private/PathTracing/Light/PathTracingLightCommon.ush
--- a/Engine/Shaders/Private/PathTracing/Light/PathTracingLightCommon.ush
+++ b/Engine/Shaders/Private/PathTracing/Light/PathTracingLightCommon.ush
@@ -82,6 +82,11 @@ float GetHeight(int LightId)
 	return SceneLightsData.Dimensions[LightId].y;
 }
 
+float2 GetRectSize(int LightId)
+{
+	return SceneLightsData.Dimensions[LightId].xy;
+}
+
 float GetDepth(int LightId)
 {
 	return SceneLightsData.Dimensions[LightId].z;
diff --git a/Engine/Shaders/Private/PathTracing/Light/PathTracingLightSampling.ush b/Engine/Shaders/Private/PathTracing/Light/PathTracingLightSampling.ush
--- a/Engine/Shaders/Private/PathTracing/Light/PathTracingLightSampling.ush
+++ b/Engine/Shaders/Private/PathTracing/Light/PathTracingLightSampling.ush
@@ -10,6 +10,8 @@ LightSampling.usf: Light sampling for next-event estimation
 	#define ENABLE_TRANSMISSION 1
 #endif
 
+//#define UNIFORM_LIGHT_SELECTION
+
 #include "PathTracingLightCommon.ush"
 #include "PathTracingDirectionalLight.ush"
 #include "PathTracingPointLight.ush"
@@ -17,6 +19,74 @@ LightSampling.usf: Light sampling for next-event estimation
 #include "PathTracingSkyLight.ush"
 #include "PathTracingSpotLight.ush"
 
+// Returns float4 packed with Radiance and Pdf
+float4 TraceLight(RayDesc Ray, int LightId, float HitT)
+{
+	if (IsPointLight(LightId))
+	{
+		return PointLight_TraceLight(Ray, LightId, HitT);
+	}
+	else if (IsSpotLight(LightId))
+	{
+		return SpotLight_TraceLight(Ray, LightId, HitT);
+	}
+	else if (IsDirectionalLight(LightId))
+	{
+		return DirectionalLight_TraceLight(Ray, LightId, HitT);
+	}
+	else if (IsEnvironmentLight(LightId))
+	{
+		return SkyLight_TraceLight(Ray, LightId, HitT);
+	}
+	else if (IsRectLight(LightId))
+	{
+		return RectLight_TraceLight(Ray, LightId, HitT);
+	}
+	// unknown light type?
+	return 0.0;
+}
+
+// Choose a point on a light as seen from the given shading point
+void SampleLight(
+	int LightId,
+	float4 RandSample,
+	float3 WorldPos,
+	float3 WorldNormal,
+	out float3 OutDirection,
+	out float  OutDistance,
+	out float3 OutRadianceOverPdf,
+	out float OutPdf)
+{
+	if (IsPointLight(LightId))
+	{
+		PointLight_SampleLight(LightId, RandSample, WorldPos, WorldNormal, OutDirection, OutDistance, OutRadianceOverPdf, OutPdf);
+	}
+	else if (IsSpotLight(LightId))
+	{
+		SpotLight_SampleLight(LightId, RandSample, WorldPos, WorldNormal, OutDirection, OutDistance, OutRadianceOverPdf, OutPdf);
+	}
+	else if (IsDirectionalLight(LightId))
+	{
+		DirectionalLight_SampleLight(LightId, RandSample, WorldPos, WorldNormal, OutDirection, OutDistance, OutRadianceOverPdf, OutPdf);
+	}
+	else if (IsEnvironmentLight(LightId))
+	{
+		SkyLight_SampleLight(LightId, RandSample, WorldPos, WorldNormal, OutDirection, OutDistance, OutRadianceOverPdf, OutPdf);
+	}
+	else if (IsRectLight(LightId))
+	{
+		RectLight_SampleLight(LightId, RandSample, WorldPos, WorldNormal, OutDirection, OutDistance, OutRadianceOverPdf, OutPdf);
+	}
+	else
+	{
+		// unknown light type?
+		OutDirection = 0;
+		OutDistance = 0;
+		OutRadianceOverPdf = 0;
+		OutPdf = 0;
+	}
+}
+
 void GenerateLightRay(RayDesc Ray, int LightId, float3 LightUV, inout RayDesc OutLightRay)
 {
 	OutLightRay = (RayDesc)0; // #dxr_todo: Return false in the case of no matching light type and handle case from caller rather than blindly using a degenerate ray
@@ -46,33 +116,34 @@ void GenerateLightRay(RayDesc Ray, int LightId, float3 LightUV, inout RayDesc Ou
 
 void EstimateLight(
 	int LightId,
-	RayDesc Ray,
-	FMaterialClosestHitPayload Payload,
+	float3 WorldPos,
+	float3 WorldNormal,
+	uint PrimitiveLightingChannelMask,
 	inout float OutIrradiance
 )
 {
 	OutIrradiance = 0.0f;
-	if ((GetLightingChannelMask(LightId) & Payload.PrimitiveLightingChannelMask) == 0) return;
+	if ((GetLightingChannelMask(LightId) & PrimitiveLightingChannelMask) == 0) return;
 
 	if (IsPointLight(LightId))
 	{
-		PointLight_EstimateLight(LightId, Ray, Payload, OutIrradiance);
+		PointLight_EstimateLight(LightId, WorldPos, WorldNormal, OutIrradiance);
 	}
 	else if (IsSpotLight(LightId))
 	{
-		SpotLight_EstimateLight(LightId, Ray, Payload, OutIrradiance);
+		SpotLight_EstimateLight(LightId, WorldPos, WorldNormal, OutIrradiance);
 	}
 	else if (IsEnvironmentLight(LightId))
 	{
-		SkyLight_EstimateLight(LightId, Ray, Payload, OutIrradiance);
+		SkyLight_EstimateLight(LightId, WorldPos, WorldNormal, OutIrradiance);
 	}
 	else if (IsDirectionalLight(LightId))
 	{
-		DirectionalLight_EstimateLight(LightId, Ray, Payload, OutIrradiance);
+		DirectionalLight_EstimateLight(LightId, WorldPos, WorldNormal, OutIrradiance);
 	}
 	else if (IsRectLight(LightId))
 	{
-		RectLight_EstimateLight(LightId, Ray, Payload, OutIrradiance);
+		RectLight_EstimateLight(LightId, WorldPos, WorldNormal, OutIrradiance);
 	}
 }
 
@@ -91,7 +162,6 @@ bool SampleLightSelection(
 		return false;
 	}
 
-//#define UNIFORM_LIGHT_SELECTION
 #ifdef UNIFORM_LIGHT_SELECTION
 	OutLightId = RandSample * SceneLightsData.Count;
 	OutLightSelectionPdf = 1.0 / SceneLightsData.Count;
@@ -100,11 +170,16 @@ bool SampleLightSelection(
 
 	// Build irradiance estimate prefix sum
 	float IrradianceEstimate[RAY_TRACING_LIGHT_COUNT_MAXIMUM];
-	EstimateLight(0, Ray, Payload, IrradianceEstimate[0]);
+
+	float3 WorldPos = Payload.WorldPos;
+	float3 WorldNormal = Payload.WorldNormal;
+	uint PrimitiveLightingChannelMask = Payload.PrimitiveLightingChannelMask;
+
+	EstimateLight(0, WorldPos, WorldNormal, PrimitiveLightingChannelMask, IrradianceEstimate[0]);
 	{
 		for (uint LightIndex = 1; LightIndex < SceneLightsData.Count; ++LightIndex)
 		{
-			EstimateLight(LightIndex, Ray, Payload, IrradianceEstimate[LightIndex]);
+			EstimateLight(LightIndex, WorldPos, WorldNormal, PrimitiveLightingChannelMask, IrradianceEstimate[LightIndex]);
 			IrradianceEstimate[LightIndex] += IrradianceEstimate[LightIndex - 1];
 		}
 	}
@@ -157,11 +232,16 @@ void PdfLightSelection(
 
 	// Build irradiance estimate prefix sum
 	float IrradianceEstimate[RAY_TRACING_LIGHT_COUNT_MAXIMUM];
-	EstimateLight(0, Ray, Payload, IrradianceEstimate[0]);
+
+	float3 WorldPos = Payload.WorldPos;
+	float3 WorldNormal = Payload.WorldNormal;
+	uint PrimitiveLightingChannelMask = Payload.PrimitiveLightingChannelMask;
+
+	EstimateLight(0, WorldPos, WorldNormal, PrimitiveLightingChannelMask, IrradianceEstimate[0]);
 	{
 		for (uint LightIndex = 1; LightIndex < SceneLightsData.Count; ++LightIndex)
 		{
-			EstimateLight(LightIndex, Ray, Payload, IrradianceEstimate[LightIndex]);
+			EstimateLight(LightIndex, WorldPos, WorldNormal, PrimitiveLightingChannelMask, IrradianceEstimate[LightIndex]);
 			IrradianceEstimate[LightIndex] += IrradianceEstimate[LightIndex - 1];
 		}
 	}
@@ -295,3 +375,70 @@ void SampleLight(
 
 	OutPdf = LightSelectionPdf * LightPdf;
 }
+
+// #dxr_todo: passing LightPickingCdf as anything but inout causes performance to drop significantly -- why ?
+void InitLightPickingCdf(inout float LightPickingCdf[RAY_TRACING_LIGHT_COUNT_MAXIMUM])
+{
+	for (uint LightIndex = 0; LightIndex < SceneLightsData.Count; ++LightIndex)
+	{
+		LightPickingCdf[LightIndex] = float(LightIndex + 1) / float(SceneLightsData.Count);
+	}
+}
+
+// #dxr_todo: passing LightPickingCdf as anything but inout causes performance to drop significantly -- why ?
+bool InitLightPickingCdf(
+	float3 WorldPos,
+	float3 WorldNormal,
+	uint PrimitiveLightingChannelMask,
+	inout float LightPickingCdf[RAY_TRACING_LIGHT_COUNT_MAXIMUM])
+{
+	// This code assumes SceneLightsData.Count > 0
+#ifdef UNIFORM_LIGHT_SELECTION
+	InitLightPickingCdf(LightPickingCdf);
+#else
+	// Build irradiance estimate prefix sum
+	EstimateLight(0, WorldPos, WorldNormal, PrimitiveLightingChannelMask, LightPickingCdf[0]);
+	{
+		for (uint LightIndex = 1; LightIndex < SceneLightsData.Count; ++LightIndex)
+		{
+			EstimateLight(LightIndex, WorldPos, WorldNormal, PrimitiveLightingChannelMask, LightPickingCdf[LightIndex]);
+			LightPickingCdf[LightIndex] += LightPickingCdf[LightIndex - 1];
+		}
+	}
+
+	if (LightPickingCdf[SceneLightsData.Count - 1] == 0)
+	{
+		// all the lights are invisible
+		return false;
+	}
+	{
+		// CDF normalization
+		for (uint LightIndex = 0; LightIndex < SceneLightsData.Count; ++LightIndex)
+		{
+			LightPickingCdf[LightIndex] /= LightPickingCdf[SceneLightsData.Count - 1];
+		}
+		LightPickingCdf[SceneLightsData.Count - 1] = 1.0;
+	}
+#endif
+	return true;
+}
+
+// #dxr_todo: passing LightPickingCdf as anything but inout causes performance to drop significantly -- why ?
+void SelectLight(float RandSample, inout float LightPickingCdf[RAY_TRACING_LIGHT_COUNT_MAXIMUM], out uint LightId, out float LightPickPdf)
+{
+	// This code assumes SceneLightsData.Count > 0
+
+	// Draw random variable
+	// #dxr_todo: binary search
+	float PreviousCdfValue = 0;
+	for (LightId = 0; LightId < SceneLightsData.Count - 1; ++LightId)
+	{
+		if (RandSample < LightPickingCdf[LightId])
+		{
+			break;
+		}
+		PreviousCdfValue = LightPickingCdf[LightId];
+	}
+
+	LightPickPdf = LightPickingCdf[LightId] - PreviousCdfValue;
+}
diff --git a/Engine/Shaders/Private/PathTracing/Light/PathTracingPointLight.ush b/Engine/Shaders/Private/PathTracing/Light/PathTracingPointLight.ush
--- a/Engine/Shaders/Private/PathTracing/Light/PathTracingPointLight.ush
+++ b/Engine/Shaders/Private/PathTracing/Light/PathTracingPointLight.ush
@@ -8,6 +8,82 @@ PointLight.usf: Light sampling functions for Point light implementation
 
 #include "PathTracingLightCommon.ush"
 
+float4 PointLight_TraceLight(RayDesc Ray, int LightId, float HitT)
+{
+	float3 LightPosition = GetPosition(LightId);
+	float LightRadius = GetRadius(LightId);
+	float LightRadius2 = Pow2(LightRadius);
+
+	float3 oc = Ray.Origin - LightPosition;
+
+	float LightDistanceSquared = dot(oc, oc);
+
+	// #dxr_todo: switch to numericaly stable ray/sphere test?
+	float a = dot(Ray.Direction, Ray.Direction);
+	float b = dot(oc, Ray.Direction);
+	float c = LightDistanceSquared - LightRadius2;
+	float h = b * b - a * c;
+	if (h >= 0)
+	{
+		h = sqrt(h);
+		float t = (-b - h) / a;
+		if (t > Ray.TMin && t < HitT)
+		{
+			// #dxr_todo: sphere area is 4*pi*r^2 -- but the factor of 4 is missing for some reason?
+			float3 LightPower = GetColor(LightId);
+			float3 LightRadiance = LightPower / (PI * LightRadius2);
+
+			float SinThetaMax2 = saturate(LightRadius2 / LightDistanceSquared);
+			float OneMinusCosThetaMax = SinThetaMax2 < 0.01 ? SinThetaMax2 * (0.5 + 0.125 * SinThetaMax2) : 1 - sqrt(1 - SinThetaMax2);
+
+			float SolidAngle = 2 * PI * OneMinusCosThetaMax;
+			return float4(LightRadiance, 1.0 / SolidAngle);
+		}
+		// #dxr_todo: process inside hit here ...
+		// How should we define radiance on the inside of the sphere?
+	}
+	return 0.0;
+}
+
+void PointLight_SampleLight(
+	int LightId,
+	float4 RandSample,
+	float3 WorldPos,
+	float3 WorldNormal,
+	out float3 OutDirection,
+	out float  OutDistance,
+	out float3 OutRadianceOverPdf,
+	out float OutPdf)
+{
+	float Radius = GetRadius(LightId);
+
+	// Sample the solid angle subtended by the sphere (which could be singgular, in which case the PDF will be infinite)
+	float3 LightDirection = GetPosition(LightId) - WorldPos;
+	float LightDistanceSquared = dot(LightDirection, LightDirection);
+	float Radius2 = Pow2(Radius);
+
+	// #dxr_todo: come up with a better definition when we are inside the light
+	float SinThetaMax2 = saturate(Radius2 / LightDistanceSquared);
+
+	// #dxr_todo: find a better way of handling the region inside the light than just clamping to 1.0 here
+	float4 DirAndPdf = UniformSampleConeRobust(RandSample.yz, SinThetaMax2);
+
+	float CosTheta = DirAndPdf.z;
+	float SinTheta2 = 1.0 - CosTheta * CosTheta;
+
+	OutDirection = TangentToWorld(DirAndPdf.xyz, normalize(LightDirection));
+	OutDistance = length(LightDirection) * (CosTheta - sqrt(max(SinThetaMax2 - SinTheta2, 0.0)));
+
+	OutPdf = DirAndPdf.w;
+
+	float3 LightPower = GetColor(LightId);
+	float3 LightRadiance = LightPower / (PI * Radius2);
+
+	// When the angle is very small, Radiance over pdf simplifies even more since SolidAngle ~= PI * SinThetaMax2
+	// Canceling out common factors further leads to the classic Power / D^2 formula
+	OutRadianceOverPdf = SinThetaMax2 < 0.001 ? LightPower / LightDistanceSquared : LightRadiance / OutPdf;
+}
+
 void PointLight_GenerateLightRay(
 	RayDesc Ray,
 	int LightId,
@@ -59,34 +135,27 @@ void PointLight_EvalLight(
 	float Radius = GetRadius(LightId);
 	float Falloff = ComputeAttenuationFalloff(Ray.TMax, GetAttenuation(LightId));
 	OutRadiance = Color * Falloff;
-#if 0
-	// doing the "right" thing leads to unstable math for small radius lights
-	// #dxr_todo: Revisit this when implementing MIS for all lights
+
 	if (Radius == 0)
 		OutRadiance *= 1 / LightDistanceSquared;
 	else
 		OutRadiance *= 1 / (PI * Radius * Radius);
-#else
-	// simplified expression for singular and soft lights - just like in the raster path
-	// this is more numerically stable and means we can ignore the PDF for now
-	OutRadiance *= 1.0 / LightDistanceSquared;
-#endif
 }
 
 void PointLight_EstimateLight(
 	int LightId,
-	RayDesc Ray,
-	FMaterialClosestHitPayload Payload,
+	float3 WorldPos,
+	float3 WorldNormal,
 	out float OutIrradiance
 )
 {
 	// Distance
-	float3 LightDirection = GetPosition(LightId) - Payload.WorldPos;
+	float3 LightDirection = GetPosition(LightId) - WorldPos;
 	float LightDistanceSquared = dot(LightDirection, LightDirection);
 	LightDirection = normalize(LightDirection);
 
 	// Geometric term
-	float3 SurfaceNormal = Payload.WorldNormal;
+	float3 SurfaceNormal = WorldNormal;
 
 	float NoL = 1.0;
 	if (GetRadius(LightId) == 0.0)
@@ -119,18 +188,11 @@ void PointLight_SampleLight(
 		
 		// outside the light
 		float SinThetaMax2 = Radius2 / LightDistanceSquared;
-		// #dxr_todo: find a better way of handling the region inside the light
-		float CosThetaMax = sqrt(max(1.0 - SinThetaMax2, 0.0));
-		float4 DirAndPdf = UniformSampleCone(RandSample.yz, CosThetaMax);
-
+		
+		// #dxr_todo: find a better way of handling the region inside the light than just clamping to 1.0 here
+		float4 DirAndPdf = UniformSampleConeRobust(RandSample.yz, min(SinThetaMax2, 1.0));
 		OutLightUV = DirAndPdf.xyz;
-#if 0
 		OutPdf = DirAndPdf.w;
-#else
-		// doing the "right" thing leads to unstable math for small radius lights
-		// #dxr_todo: Revisit this when implementing MIS for all lights
-		OutPdf = 1.0;
-#endif
 	}
 	else
 	{
diff --git a/Engine/Shaders/Private/PathTracing/Light/PathTracingRectLight.ush b/Engine/Shaders/Private/PathTracing/Light/PathTracingRectLight.ush
--- a/Engine/Shaders/Private/PathTracing/Light/PathTracingRectLight.ush
+++ b/Engine/Shaders/Private/PathTracing/Light/PathTracingRectLight.ush
@@ -9,6 +9,85 @@
 #include "../../RectLight.ush"
 #include "PathTracingLightCommon.ush"
 
+float4 RectLight_TraceLight(RayDesc Ray, int LightId, float HitT)
+{
+	float3 LightPosition = GetPosition(LightId);
+	float3 LightNormal = GetNormal(LightId);
+	float3 LightDirection = LightPosition - Ray.Origin;
+	float DoN = dot(Ray.Direction, LightNormal);
+	float t = dot(LightDirection, LightNormal) / DoN;
+	// ray points toward the plane and intersect it?
+	if (DoN < 0 && t > Ray.TMin && t < HitT)
+	{
+		float3 LightdPdu = GetdPdu(LightId);
+		float3 LightdPdv = GetdPdv(LightId);
+		float2 LightExtent = 0.5 * GetRectSize(LightId);
+
+		float3 P = t * Ray.Direction - LightDirection;
+		float2 UV = float2(dot(P, LightdPdu), dot(P, LightdPdv));
+		// test point against 
+		if (all(abs(UV) <= LightExtent))
+		{
+			// Clip the Rectangle by the barndoors
+			FRect Rect = GetRect(LightDirection,
+				-LightNormal,
+				LightdPdv,
+				LightExtent.x,
+				LightExtent.y,
+				GetRectLightBarnCosAngle(LightId),
+				GetRectLightBarnLength(LightId),
+				true);
+			P = t * Ray.Direction - Rect.Origin;
+			// test again with the clipped extents
+			UV = float2(dot(P, LightdPdu), dot(P, LightdPdv));
+			if (all(abs(UV) <= Rect.Extent))
+			{
+				// stored color is radiance
+				float3 Radiance = GetColor(LightId);
+				float SolidAngle = BuildSphericalRect(Rect).SolidAngle;
+				return float4(Radiance, 1.0 / SolidAngle);
+			}
+		}
+	}
+	return 0.0;
+}
+
+void RectLight_SampleLight(
+	int LightId,
+	float4 RandSample,
+	float3 WorldPos,
+	float3 WorldNormal,
+	out float3 OutDirection,
+	out float  OutDistance,
+	out float3 OutRadianceOverPdf,
+	out float OutPdf)
+{
+	float3 LightPosition = GetPosition(LightId);
+	float3 LightNormal = GetNormal(LightId);
+	float3 LightdPdu = GetdPdu(LightId);
+	float3 LightdPdv = GetdPdv(LightId);
+	float LightWidth = GetWidth(LightId);
+	float LightHeight = GetHeight(LightId);
+
+	// Define rectangle and compute solid angle
+	FRect Rect = GetRect(LightPosition - WorldPos,
+		-LightNormal,
+		LightdPdv,
+		0.5 * LightWidth,
+		0.5 * LightHeight,
+		GetRectLightBarnCosAngle(LightId),
+		GetRectLightBarnLength(LightId),
+		true /* bComputeVisibleRect */);
+	FSphericalRect SphericalRect = BuildSphericalRect(Rect);
+
+	float3 Direction = UniformSampleSphericalRect(RandSample.yz, SphericalRect);
+
+	OutDirection = normalize(Direction);
+	OutDistance = length(Direction);
+	OutRadianceOverPdf = GetColor(LightId) * SphericalRect.SolidAngle;
+	OutPdf = 1.0 / SphericalRect.SolidAngle;
+}
+
 void RectLight_GenerateLightRay(
 	RayDesc Ray,
 	int LightId,
@@ -47,14 +126,14 @@ void RectLight_EvalLight(
 
 void RectLight_EstimateLight(
 	int LightId,
-	RayDesc Ray,
-	FMaterialClosestHitPayload Payload,
+	float3 WorldPos,
+	float3 WorldNormal,
 	out float OutIrradiance
 )
 {
 	// Distance to centroid
 	// #dxr_todo: UE-72533 Use closest point, instead
-	float3 LightDirection = GetPosition(LightId) - Payload.WorldPos;
+	float3 LightDirection = GetPosition(LightId) - WorldPos;
 	float LightDistanceSquared = dot(LightDirection, LightDirection);
 	LightDirection = normalize(LightDirection);
 	float3 LightNormal = GetNormal(LightId);
@@ -67,7 +146,7 @@ void RectLight_EstimateLight(
 	}
 
 	// Approximate geometric term
-	float3 SurfaceNormal = Payload.WorldNormal;
+	float3 SurfaceNormal = WorldNormal;
 #if 0
 	float NoL = dot(SurfaceNormal, LightDirection);
 	NoL = ENABLE_TRANSMISSION && HasTransmission(LightId) ? abs(NoL) : saturate(NoL);
diff --git a/Engine/Shaders/Private/PathTracing/Light/PathTracingSkyLight.ush b/Engine/Shaders/Private/PathTracing/Light/PathTracingSkyLight.ush
--- a/Engine/Shaders/Private/PathTracing/Light/PathTracingSkyLight.ush
+++ b/Engine/Shaders/Private/PathTracing/Light/PathTracingSkyLight.ush
@@ -66,12 +66,11 @@ void SkyLight_EvalLight(
 
 void SkyLight_EstimateLight(
 	int LightId,
-	RayDesc Ray,
-	FMaterialClosestHitPayload Payload,
+	float3 WorldPos,
+	float3 WorldNormal,
 	out float OutIrradiance
 )
 {
-	float3 WorldNormal = Payload.WorldNormal;
 	// Modifying diffuse contribution to prevent false negatives with PDF estimation
 	float3 SkySHDiffuse = max(GetSkySHDiffuse(WorldNormal), 0.01);
 	float3 SkyLightIrradiance = GetColor(LightId) * SkySHDiffuse;
@@ -101,6 +100,7 @@ void SkyLight_SampleLight(
 	//uint MipCount = Log2(SkyLight.MipDimensions.x);
 	uint MipCount = log2(SkyLight.MipDimensions.x);
 	float FacePdf = 0.0;
+	// FIXME: RandSample.x has already been used to pick a light source ...
 	uint FaceIndex = SampleFace(MipCount, RandSample.x, FacePdf);
 
 	uint2 MipPixel = 0;
@@ -168,4 +168,58 @@ void SkyLight_SampleLight(
 	OutLightUV = RayDirection;
 	OutPdf = RayPdf;
 #endif
-}
\ No newline at end of file
+}
+
+
+//  NOTE: put new methods at the end for now, so we can re-use the methods above
+float4 SkyLight_TraceLight(RayDesc Ray, int LightId, float HitT) {
+	if (Ray.TMax == HitT)
+	{
+		// LightUV convention in the code below is just the world space ray direction
+		// so its safe to call into the method below for now ...
+
+		// #dxr_todo: write optimized form that combines radiance and pdf lookup in a single step
+
+		float3 Radiance = 0.0;
+		float Pdf = 0.0;
+
+		float3 LightUV = normalize(Ray.Direction);
+		SkyLight_EvalLight(LightId, LightUV, Ray, Radiance);
+		if (any(Radiance > 0))
+		{
+			SkyLight_PdfLight(Ray, LightId, LightUV, Pdf);
+		}
+		
+
+		return float4(Radiance, Pdf);
+	}
+	return 0.0;
+}
+
+void SkyLight_SampleLight(
+	int LightId,
+	float4 RandSample,
+	float3 WorldPos,
+	float3 WorldNormal,
+	out float3 OutDirection,
+	out float  OutDistance,
+	out float3 OutRadianceOverPdf,
+	out float OutPdf)
+{
+	// Make temp structures so we can re-use the calls above
+	// #dxr_todo: write optimized form that combines radiance and pdf lookup in a single step
+	RayDesc Ray;
+	Ray.Origin = WorldPos;
+	FMaterialClosestHitPayload Payload;
+	Payload.WorldNormal = WorldNormal;
+
+
+	float3 LightUV;
+	float3 Radiance;
+	SkyLight_SampleLight(Ray, Payload, RandSample, LightId, LightUV, OutPdf);
+	SkyLight_EvalLight(LightId, LightUV, Ray, Radiance);
+	
+	OutDirection = LightUV;
+	OutDistance = 1.0e27;
+	OutRadianceOverPdf = Radiance / OutPdf;
+}
diff --git a/Engine/Shaders/Private/PathTracing/Light/PathTracingSpotLight.ush b/Engine/Shaders/Private/PathTracing/Light/PathTracingSpotLight.ush
--- a/Engine/Shaders/Private/PathTracing/Light/PathTracingSpotLight.ush
+++ b/Engine/Shaders/Private/PathTracing/Light/PathTracingSpotLight.ush
@@ -26,6 +26,92 @@ float SpotAttenuation(float3 L, float3 SpotDirection, float2 SpotAngles)
 }
 
 
+float4 SpotLight_TraceLight(RayDesc Ray, int LightId, float HitT) {
+	// #dxr_todo: share with point light instead of copy/pasting logic?
+	float3 LightPosition = GetPosition(LightId);
+	float LightRadius = GetRadius(LightId);
+	float LightRadius2 = Pow2(LightRadius);
+
+	float3 oc = Ray.Origin - LightPosition;
+
+	float LightDistanceSquared = dot(oc, oc);
+
+	// #dxr_todo: switch to numericaly stable ray/sphere test?
+	float a = dot(Ray.Direction, Ray.Direction);
+	float b = dot(oc, Ray.Direction);
+	float c = LightDistanceSquared - LightRadius2;
+	float h = b * b - a * c;
+	if (h >= 0)
+	{
+		h = sqrt(h);
+		float t = (-b - h) / a;
+		if (t > Ray.TMin && t < HitT)
+		{
+			// #dxr_todo: sphere area is 4*pi*r^2 -- but the factor of 4 is missing for some reason?
+			float3 LightPower = GetColor(LightId);
+			float3 LightRadiance = LightPower / (PI * LightRadius2);
+
+			// adjust lighting according to spot light cone falloff
+			float3 LightNormal = GetNormal(LightId);
+			float2 CosConeAngles = GetCosConeAngles(LightId);
+			LightRadiance *= SpotAttenuation(-normalize(oc), LightNormal, CosConeAngles);
+
+			float SinThetaMax2 = saturate(LightRadius2 / LightDistanceSquared);
+			float OneMinusCosThetaMax = SinThetaMax2 < 0.01 ? SinThetaMax2 * (0.5 + 0.125 * SinThetaMax2) : 1 - sqrt(1 - SinThetaMax2);
+
+			float SolidAngle = 2 * PI * OneMinusCosThetaMax;
+			return float4(LightRadiance, 1.0 / SolidAngle);
+		}
+		// #dxr_todo: process inside hit here ...
+		// How should we define radiance on the inside of the sphere?
+	}
+	return 0.0;
+}
+
+void SpotLight_SampleLight(
+	int LightId,
+	float4 RandSample,
+	float3 WorldPos,
+	float3 WorldNormal,
+	out float3 OutDirection,
+	out float  OutDistance,
+	out float3 OutRadianceOverPdf,
+	out float OutPdf)
+{
+	float Radius = GetRadius(LightId);
+
+	// Sample the solid angle subtended by the sphere (which could be singgular, in which case the PDF will be infinite)
+	float3 LightDirection = GetPosition(LightId) - WorldPos;
+	float LightDistanceSquared = dot(LightDirection, LightDirection);
+	float Radius2 = Pow2(Radius);
+
+	// #dxr_todo: come up with a better definition when we are inside the light
+	float SinThetaMax2 = saturate(Radius2 / LightDistanceSquared);
+
+	// #dxr_todo: find a better way of handling the region inside the light than just clamping to 1.0 here
+	float4 DirAndPdf = UniformSampleConeRobust(RandSample.yz, SinThetaMax2);
+
+	float CosTheta = DirAndPdf.z;
+	float SinTheta2 = 1.0 - CosTheta * CosTheta;
+
+	OutDirection = TangentToWorld(DirAndPdf.xyz, normalize(LightDirection));
+	OutDistance = length(LightDirection) * (CosTheta - sqrt(max(SinThetaMax2 - SinTheta2, 0.0)));
+
+	OutPdf = DirAndPdf.w;
+
+	float3 LightPower = GetColor(LightId);
+	float3 LightRadiance = LightPower / (PI * Radius2);
+
+	// When the angle is very small, Radiance over pdf simplifies even more since SolidAngle ~= PI * SinThetaMax2
+	// Canceling out common factors further leads to the classic Power / D^2 formula
+	OutRadianceOverPdf = SinThetaMax2 < 0.001 ? LightPower / LightDistanceSquared : LightRadiance / OutPdf;
+
+	// adjust lighting according to spot light cone falloff
+	float3 LightNormal = GetNormal(LightId);
+	float2 CosConeAngles = GetCosConeAngles(LightId);
+	OutRadianceOverPdf *= SpotAttenuation(normalize(LightDirection), LightNormal, CosConeAngles);
+}
+
 void SpotLight_GenerateLightRay(
 	RayDesc Ray,
 	int LightId,
@@ -78,8 +164,11 @@ void SpotLight_EvalLight(
 	float Falloff = ComputeAttenuationFalloff(Ray.TMax, GetAttenuation(LightId));
 	OutRadiance = Color * Falloff;
 
-	// #dxr-todo: See math in point lights - ideally spot light would share that logic somehow
-	OutRadiance *= 1.0 / LightDistanceSquared;
+	// Same math as point lights
+	if (Radius == 0)
+		OutRadiance *= 1 / LightDistanceSquared;
+	else
+		OutRadiance *= 1 / (PI * Radius * Radius);
 
 	// This appears to be the only difference between point and spot lights
 	float3 LightNormal = GetNormal(LightId);
@@ -89,8 +178,8 @@ void SpotLight_EvalLight(
 
 void SpotLight_EstimateLight(
 	int LightId,
-	RayDesc Ray,
-	FMaterialClosestHitPayload Payload,
+	float3 WorldPos,
+	float3 WorldNormal,
 	out float OutIrradiance
 )
 {
@@ -100,12 +189,12 @@ void SpotLight_EstimateLight(
 	float Radius = GetRadius(LightId);
 	
 	// Distance
-	float3 LightDirection = GetPosition(LightId) - Payload.WorldPos;
+	float3 LightDirection = GetPosition(LightId) - WorldPos;
 	float LightDistanceSquared = dot(LightDirection, LightDirection);
 	LightDirection = normalize(LightDirection);
 
 	// Geometric term
-	float3 SurfaceNormal = Payload.WorldNormal;
+	float3 SurfaceNormal = WorldNormal;
 
 	float NoL = 1.0;
 	if (GetRadius(LightId) == 0.0)
@@ -140,18 +229,11 @@ void SpotLight_SampleLight(
 		
 		// outside the light
 		float SinThetaMax2 = Radius2 / LightDistanceSquared;
-		// #dxr_todo: find a better way of handling the region inside the light
-		float CosThetaMax = sqrt(max(1.0 - SinThetaMax2, 0.0));
-		float4 DirAndPdf = UniformSampleCone(RandSample.yz, CosThetaMax);
 
+		// #dxr_todo: find a better way of handling the region inside the light than just clamping to 1.0 here
+		float4 DirAndPdf = UniformSampleConeRobust(RandSample.yz, min(SinThetaMax2, 1.0));
 		OutLightUV = DirAndPdf.xyz;
-#if 0
 		OutPdf = DirAndPdf.w;
-#else
-		// doing the "right" thing leads to unstable math for small radius lights
-		// #dxr_todo: Revisit this when implementing MIS for all lights
-		OutPdf = 1.0;
-#endif
 	}
 	else
 	{
diff --git a/Engine/Shaders/Private/PathTracing/Material/PathTracingDefaultLit.ush b/Engine/Shaders/Private/PathTracing/Material/PathTracingDefaultLit.ush
--- a/Engine/Shaders/Private/PathTracing/Material/PathTracingDefaultLit.ush
+++ b/Engine/Shaders/Private/PathTracing/Material/PathTracingDefaultLit.ush
@@ -10,8 +10,14 @@
 
 float CalcLobeSelectionPdf(float3 V, FMaterialClosestHitPayload Payload)
 {
-	// dxr_todo: Reinstate weighted lobe selection criteria (eg: roughness)
-	return 0.5;
+	// #dxr_todo: Can do much better than this - but this at least prevents under-sampling when one of the lobes is black
+	float3 Diff = GetDiffuseColor(Payload);
+	float3 Spec = GetSpecularColor(Payload);
+
+	float D = Diff.x + Diff.y + Diff.z;
+	float S = Spec.x + Spec.y + Spec.z;
+	return D / (D + S + 1e-6);
+
 
 	float Pdf = min(GetRoughness(Payload), 0.9);
 	return Pdf;
@@ -108,14 +114,14 @@ void DefaultLit_SampleMaterial(
 		float3 F_Metallic = F_Schlick(GetSpecularColor(Payload), VoH);
 		float3 F_Dielectric = FresnelReflectance(VoH, Ior1, GetIor(Payload));
 		float3 F = lerp(F_Dielectric, F_Metallic, GetMetallic(Payload));
-		float D = D_GGX(a2, NoH);
+		float D = Result.w / Result.z;// D_GGX(a2, NoH);
 		// Note: Pedantic application of normalized weights is on purpose
 		float G = Vis_SmithJointApprox(a2, NoV, NoL) * (4.0 * NoV * NoL);
 
 		float MicrofacetPdf = 0.0;
 		if (VoH > 0.0)
 		{
-			MicrofacetPdf = D * NoH / (4.0 * VoH);
+			MicrofacetPdf = Result.w / (4.0 * VoH);
 		}
 
 		OutThroughput = 0.0;
diff --git a/Engine/Shaders/Private/PathTracing/Material/PathTracingMaterialCommon.ush b/Engine/Shaders/Private/PathTracing/Material/PathTracingMaterialCommon.ush
--- a/Engine/Shaders/Private/PathTracing/Material/PathTracingMaterialCommon.ush
+++ b/Engine/Shaders/Private/PathTracing/Material/PathTracingMaterialCommon.ush
@@ -111,6 +111,11 @@ bool IsLambert(FMaterialClosestHitPayload HitInfo)
 	return GetRoughness(HitInfo) > GLOSSY_MAX_ROUGHNESS;
 }
 
+bool IsUnlitMaterial(FMaterialClosestHitPayload HitInfo)
+{
+	return HitInfo.ShadingModelID == SHADINGMODELID_UNLIT;
+}
+
 bool IsDefaultLitMaterial(FMaterialClosestHitPayload HitInfo)
 {
 	return HitInfo.ShadingModelID == SHADINGMODELID_DEFAULT_LIT;
diff --git a/Engine/Shaders/Private/PathTracing/Material/PathTracingMaterialSampling.ush b/Engine/Shaders/Private/PathTracing/Material/PathTracingMaterialSampling.ush
--- a/Engine/Shaders/Private/PathTracing/Material/PathTracingMaterialSampling.ush
+++ b/Engine/Shaders/Private/PathTracing/Material/PathTracingMaterialSampling.ush
@@ -42,7 +42,7 @@ bool IsNeeValidMaterial(FMaterialClosestHitPayload Payload)
 	#if FORCE_ALL_MATERIALS_DIFFUSE
 		return true;
 	#endif
-	
+
 	if (IsPureSpecularReflection(Payload))
 	{
 		return false;
@@ -63,6 +63,10 @@ bool IsNeeValidMaterial(FMaterialClosestHitPayload Payload)
 	{
 		return true;
 	}
+	else if (IsUnlitMaterial(Payload))
+	{
+		return false;
+	}
 	else if(IsLambert(Payload))
 	{
 		return true;
@@ -124,6 +128,13 @@ void SampleMaterial(
 	{
 		TwoSidedFoliage_SampleMaterial(RayDirection, Payload, RandSample, OutDirection, OutThroughput, OutPdf, OutPositionBiasSign);
 	}
+	else if (IsUnlitMaterial(Payload))
+	{
+		OutDirection = 0;
+		OutThroughput = 0;
+		OutPdf = 0;
+		return;
+	}
 	else if(IsLambert(Payload))
 	{
 		Lambert_SampleMaterial(Payload, RandSample, OutDirection, OutThroughput, OutPdf);
@@ -177,6 +188,10 @@ void PdfMaterial(
 	{
 		OutPdf = TwoSidedFoliage_PdfMaterial(IncomingDirection, OutgoingDirection, Payload);
 	}
+	else if (IsUnlitMaterial(Payload))
+	{
+		OutPdf = 0.0;
+	}
 	else if (IsLambert(Payload))
 	{
 		OutPdf = Lambert_Pdf(OutgoingDirection, Payload);
@@ -235,6 +250,12 @@ void EvalMaterial(
 	{
 		TwoSidedFoliage_EvalMaterial(IncomingDirection, OutgoingDirection, Payload, OutThroughput, OutPdf);
 	}
+	else if (IsUnlitMaterial(Payload))
+	{
+		OutThroughput = 0.0;
+		OutPdf = 0.0;
+		return;
+	}
 	else if (IsLambert(Payload))
 	{
 		Lambert_EvalMaterial(OutgoingDirection, Payload, OutThroughput, OutPdf);
diff --git a/Engine/Shaders/Private/PathTracing/PathTracing.usf b/Engine/Shaders/Private/PathTracing/PathTracing.usf
--- a/Engine/Shaders/Private/PathTracing/PathTracing.usf
+++ b/Engine/Shaders/Private/PathTracing/PathTracing.usf
@@ -27,7 +27,6 @@ RWTexture2D<uint>   PixelPositionRT;
 RWTexture2D<uint>	RayCountPerPixelRT;
 RaytracingAccelerationStructure TLAS;
 
-
 RAY_TRACING_ENTRY_RAYGEN(PathTracingMainRG)
 {
 	uint2 LaunchIndex = DispatchRaysIndex().xy + View.ViewRectMin + PathTracingData.TileOffset;
@@ -39,7 +38,6 @@ RAY_TRACING_ENTRY_RAYGEN(PathTracingMainRG)
 	SampleCountRT[LaunchIndex] = 0;
 	PixelPositionRT[LaunchIndex] = 0;
 	RayCountPerPixelRT[LaunchIndex] = 0;
-	uint RayCounter = 0;
 	
 	// Initialize random sequence 
 	float2 BufferSize = View.BufferSizeAndInvSize.xy; 
@@ -53,6 +51,7 @@ RAY_TRACING_ENTRY_RAYGEN(PathTracingMainRG)
 	for (uint i = 0; i < SampleCount; ++i) 
 	{
 		float3 Irradiance = 0;
+		uint RayCounter = 0;
 
 		// Initialize ray and payload
 		RayDesc Ray;
@@ -64,8 +63,13 @@ RAY_TRACING_ENTRY_RAYGEN(PathTracingMainRG)
 			continue;
 		}
 
-		float3 RayThroughput = 1.0;
-		float RayPdf = 1.0;
+		// This array will hold a CDF for light picking
+		// Seed the array with a uniform CDF at first so that we always have a valid CDF
+		float LightPickingCdf[RAY_TRACING_LIGHT_COUNT_MAXIMUM];
+
+		InitLightPickingCdf(LightPickingCdf);
+
+		Ray.Direction = normalize(Ray.Direction);
 
 		FRayCone RayCone = (FRayCone)0;
 		RayCone.SpreadAngle = View.EyeToPixelSpreadAngle;
@@ -77,232 +81,197 @@ RAY_TRACING_ENTRY_RAYGEN(PathTracingMainRG)
 		const bool bIgnoreTranslucentMaterials = false;
 		const uint2 RenderTargetPos = LaunchIndex.xy + View.ViewRectMin.xy;
 
-		FMaterialClosestHitPayload Payload = TraceMaterialRay(
-			TLAS,
-			InitialRayFlags,
-			InitialInstanceInclusionMask,
-			Ray,
-			RayCone,
-			RenderTargetPos,
-			bInitialDisableSkyLightContribution,
-			bIgnoreTranslucentMaterials);
-		RayCounter++;
-
-		float MaterialPdf = 0.0;
-		FMaterialClosestHitPayload PrevMaterialPayload = (FMaterialClosestHitPayload)0;
-		uint Bounces = 0;
-		bool bIsNeeValid = true;
-		while (Bounces < PathTracingData.MaxBounces)
+		// path state variables (these cary information between bounces)
+		float3 PathThroughput = 1.0;
+
+		// We don't want the camera ray to have any MIS contribution against direct light sampling so we pretend that it came from
+		// a singular Material scattering event.
+		float PrevMaterialPdf = POSITIVE_INFINITY;
+
+		// 0: only Material sampling
+		// 1: only Light sampling
+		// 2: both Material and Light
+		const int MISMode = AdaptiveSamplingData.MISMode;
+
+		for (int Bounce = 0; Bounce <= PathTracingData.MaxBounces; Bounce++)
 		{
-			// Environment contribution
-			if (Payload.IsMiss())
+			FMaterialClosestHitPayload Payload = TraceMaterialRay(
+				TLAS,
+				InitialRayFlags,
+				InitialInstanceInclusionMask,
+				Ray,
+				RayCone,
+				RenderTargetPos,
+				bInitialDisableSkyLightContribution,
+				bIgnoreTranslucentMaterials);
+			RayCounter++;
+
+			if (MISMode != 1 || Bounce == 0)
 			{
-				// WARNING: Until ray payload encodes NEE LightId, the execution depends on this ordering!
-				uint SkyLightId = 0;
-				float3 EnvironmentRadiance = 0.0;
-				SkyLight_EvalLight(SkyLightId, Ray.Direction, Ray, EnvironmentRadiance);
-				if (length(EnvironmentRadiance) > 0.0)
+				float HitT = Payload.IsMiss() ? Ray.TMax : Payload.HitT;
+				// Loop over lights to capture their contribution
+				// #dxr_todo: if we have lots of lights, having some hierarchical structure would be better ....
+				float PreviousCdfValue = 0.0;
+				for (uint LightId = 0; LightId < SceneLightsData.Count; ++LightId)
 				{
-					if (IsUnidirectionalEnabled)
-					{
-						// Apply NEE Pdf for correct MIS weight
-						float MisWeight = 1.0;
-						if (IsNextEventEstimationEnabled && Bounces > 0 && bIsNeeValid)
-						{
-							float NeePdf = 0.0;
-							float3 LightUV = GetLightUV(SkyLightId, Ray, Payload);
-							PdfLight(Ray, PrevMaterialPayload, SkyLightId, LightUV, NeePdf);
-
-							MisWeight = MaterialPdf / (MaterialPdf + NeePdf);
-						}
-						Irradiance += EnvironmentRadiance * RayThroughput * MisWeight / RayPdf;
-					}
-					// Record primary environment hits with NEE
-					else if (IsNextEventEstimationEnabled && Bounces == 0)
-					{
-						Irradiance += EnvironmentRadiance * RayThroughput / RayPdf;
-					}
+					// Camera rays should only do the rest of the loop if we have visible lights OR if
+					// we have a skydome (which are always visible)
+					if (Bounce == 0 && AdaptiveSamplingData.VisibleLights == 0 && !IsEnvironmentLight(LightId))
+						continue;
+
+					float4 LightResult = TraceLight(Ray, LightId, HitT);
+					float3 LightRadiance = LightResult.xyz;
+					float LightPdf = LightResult.w;
+
+					float LightPickPdf = LightPickingCdf[LightId] - PreviousCdfValue;
+					PreviousCdfValue = LightPickingCdf[LightId];
+					Irradiance += PathThroughput * LightRadiance * MISWeightRobust(PrevMaterialPdf, LightPdf * LightPickPdf);
 				}
-				break;
 			}
 
-			// Emissive geometry contribution
-			float3 Emission = GetMaterialEmission(Payload);
-			if (any(Emission) > 0.0)
+			if (Payload.IsMiss())
 			{
-				if (IsUnidirectionalEnabled)
-				{
-					// Apply NEE Pdf for correct MIS weight
-					float MisWeight = 1.0;
-					if (IsNextEventEstimationEnabled && Bounces > 0 && bIsNeeValid)
-					{
-						float NeePdf = 0.0;
-						int LightId = GetLightId(Payload);
-						float3 LightUV = GetLightUV(LightId, Ray, Payload);
-						PdfLight(Ray, PrevMaterialPayload, LightId, LightUV, NeePdf);
-
-						float MisWeight = MaterialPdf / (MaterialPdf + NeePdf);
-					}
-
-					Irradiance += Emission * RayThroughput * MisWeight / RayPdf;
-				}
-				// Record primary environment hits with NEE
-				else if (IsNextEventEstimationEnabled && Bounces == 0)
-				{
-					Irradiance += Emission * RayThroughput / RayPdf;
-				}
+				// Ray didn't hit any real geometry, so nothing left to do
 				break;
 			}
 
-			Bounces = Bounces + 1;
+			// add in surface emission
+			Irradiance += PathThroughput * Payload.Radiance;
+
 			// Update origin
-			Ray.Origin = Ray.Origin + Ray.Direction * Payload.HitT;
-			FMaterialClosestHitPayload RayHitInfo = Payload; 
+			Ray.Origin = Payload.WorldPos;
 
-			bIsNeeValid = IsNeeValidMaterial(RayHitInfo);
+			// Make sure the surface normal points toward the viewer
+			if (dot(Ray.Direction, Payload.WorldNormal) > 0.0)
+			{
+				Payload.WorldNormal = -Payload.WorldNormal;
+			}
 
-			// Perform next-event estimation
-			if (IsNextEventEstimationEnabled && bIsNeeValid)
+			// Choose a random number for both Light sampling and BxDF sampling
+			float4 RandSample = RandomSequence_GenerateSample4D(RandSequence, SampleIndex);
+
+			// Does this material require NEE?
+			bool bIsNeeValid = IsNeeValidMaterial(Payload);
+
+			// If we are using Light sampling and we are not on the last bounce ....
+			if (MISMode != 0 && Bounce < PathTracingData.MaxBounces && bIsNeeValid && SceneLightsData.Count > 0)
 			{
-				int LightId;
-				float3 LightUV;
-				float NeePdf = 0.0;
-				float4 RandSample = RandomSequence_GenerateSample4D(RandSequence, SampleIndex);
-				SampleLight(Ray, RayHitInfo, RandSample, LightId, LightUV, NeePdf);
-
-				// Fire visibility ray
-				if (NeePdf > 0.0)
-				{
-					RayDesc LightRay;
-					GenerateLightRay(Ray, LightId, LightUV, LightRay);
+				// Choose a light and sample it
 
-					float SignedPositionBias = 1.0;
-					float NoL = dot(Payload.WorldNormal, LightRay.Direction);
-					if (ENABLE_TRANSMISSION && HasTransmission(LightId) && IsMaterialTransmissive(RayHitInfo) && NoL < 0.0)
-					{
-						SignedPositionBias = -1.0;
-					}
-					ApplyPositionBias(LightRay, SignedPositionBias * Payload.WorldNormal, AdaptiveSamplingData.MaxNormalBias);
+				float3 WorldPos = Payload.WorldPos;
+				float3 WorldNormal = Payload.WorldNormal;
+				uint PrimitiveLightingChannelMask = Payload.PrimitiveLightingChannelMask;
 
-					const uint RayFlags = RAY_FLAG_ACCEPT_FIRST_HIT_AND_END_SEARCH | RAY_FLAG_SKIP_CLOSEST_HIT_SHADER;
-					const uint InstanceInclusionMask = RAY_TRACING_MASK_SHADOW;
+				if (InitLightPickingCdf(WorldPos, WorldNormal, PrimitiveLightingChannelMask, LightPickingCdf))
+				{
+					// init worked
+					int LightId;
+					float LightPickPdf = 0;
 
-					FMinimalPayload NeePayload = TraceVisibilityRay(
-						TLAS,
-						RayFlags,
-						InstanceInclusionMask,
-						RenderTargetPos,
-						LightRay);
-					RayCounter++;
+					SelectLight(RandSample.x, LightPickingCdf, LightId, LightPickPdf);
 
-					if (NeePayload.IsMiss())
+					RayDesc LightRay;
+					LightRay.Origin = WorldPos;
+					LightRay.TMin = 0;
+					float3 RadianceOverPdf;
+					float LightPdf;
+					SampleLight(LightId, RandSample, WorldPos, WorldNormal,
+						LightRay.Direction,
+						LightRay.TMax,
+						RadianceOverPdf,
+						LightPdf);
+					RadianceOverPdf /= LightPickPdf;
+					LightPdf *= LightPickPdf;
+
+					if (LightPdf > 0)
 					{
-						FMaterialClosestHitPayload RayHitInfo = Payload; 
-#if 0
-						// DEBUG: Validate NEE PDF
-						float NeePdfTmp;
-						PdfLight(Ray, RayHitInfo, LightId, LightUV, NeePdfTmp);
-						if (abs(NeePdf - NeePdfTmp) > 0.0001)
+						float SignedPositionBias = 1.0;
+						float NoL = dot(Payload.WorldNormal, LightRay.Direction);
+						if (ENABLE_TRANSMISSION && HasTransmission(LightId) && IsMaterialTransmissive(Payload) && NoL < 0.0)
 						{
-							Irradiance = float3(-10, -10, -10);
-							break;
+							SignedPositionBias = -1.0;
 						}
-#endif
-						// Evaluate radiance
-						float3 Radiance;
-						EvalLight(LightId, LightUV, LightRay, Radiance);
-
-						// Evaluate material
-						float Ior1 = 1.0;
-						float3 MaterialThroughput;
-						float MaterialEvalPdf = 0.0;
-						EvalMaterial(Ray.Direction, LightRay.Direction, RayHitInfo, Ior1, false, MaterialThroughput, MaterialEvalPdf);
-
-						// Apply material Pdf for correct MIS weight
-						float MisWeight = 1.0;
-						if (IsUnidirectionalEnabled && IsPhysicalLight(LightId))
+						ApplyPositionBias(LightRay, SignedPositionBias * Payload.WorldNormal, AdaptiveSamplingData.MaxNormalBias);
+
+
+						const uint RayFlags = RAY_FLAG_ACCEPT_FIRST_HIT_AND_END_SEARCH | RAY_FLAG_SKIP_CLOSEST_HIT_SHADER;
+						const uint InstanceInclusionMask = RAY_TRACING_MASK_SHADOW;
+
+						FMinimalPayload NeePayload = TraceVisibilityRay(
+							TLAS,
+							RayFlags,
+							InstanceInclusionMask,
+							RenderTargetPos,
+							LightRay);
+						RayCounter++;
+
+						// #dxr_todo: Is it cheaper to fire the ray first? Or eval the material first?
+						if (NeePayload.IsMiss())
 						{
-#if 0
-							// DEBUG: Validate Material PDF
-							float MaterialPdfTmp;
-							PdfMaterial(Ray.Direction, LightRay, RayHitInfo, Ior1, false, MaterialPdfTmp);
-							if (abs(MaterialPdf - MaterialPdfTmp) > 0.0001)
-							{
-								MaterialThroughput *= -1.f;
-							}
-#endif
-
-							MisWeight = NeePdf / (NeePdf + MaterialEvalPdf);
-						}
+							// Evaluate material
+							float Ior1 = 1.0;
+							float3 MaterialEval;
+							float MaterialPdf = 0.0;
+							EvalMaterial(Ray.Direction, LightRay.Direction, Payload, Ior1, false, MaterialEval, MaterialPdf);
 
-						// Record the contribution
-						Irradiance += Radiance * MaterialThroughput * RayThroughput * MisWeight / (NeePdf * RayPdf);
+							float MISWeight = MISMode == 2 ? MISWeightRobust(LightPdf, MaterialPdf) : 1.0;
+
+							// Record the contribution
+							Irradiance += MISWeight * PathThroughput * RadianceOverPdf * MaterialEval;
+						}
 					}
 				}
 			}
 
+			if (Bounce == PathTracingData.MaxBounces)
+			{
+				// we've arrived at the last bounce, don't bother trying to sample a direction
+				break;
+			}
+
 			// Sample material
 			float Ior1 = 1.0;
 			float3 Direction;
-			float3 Throughput;
-			float4 RandSample = RandomSequence_GenerateSample4D(RandSequence, SampleIndex);
+			float3 MaterialThroughput;
 			float SignedPositionBias;
-			float SamplePdf;
-			SampleMaterial(Ray.Direction, RayHitInfo, Ior1, RandSample, false, Direction, Throughput, SamplePdf, SignedPositionBias);
-			if (SamplePdf == 0.0)
+			float MaterialSamplePdf;
+			SampleMaterial(Ray.Direction, Payload, Ior1, RandSample, false, Direction, MaterialThroughput, MaterialSamplePdf, SignedPositionBias);
+			if (!(MaterialSamplePdf > 0))
 			{
+				// No valid direction -- we are done
 				break;
 			}
 
-			PdfMaterial(Ray.Direction, Direction, Payload, Ior1, false, MaterialPdf);
-			// Debug: Evaluation of a candidate sample direction should maintain energy balance with pathspace MIS
-			// EvalMaterial(Ray.Direction, Direction, RayHitInfo, Ior1, false, Throughput, MaterialPdf);
-			// SamplePdf = MaterialPdf;
-
 			// Update ray
 			Ray.Direction = Direction;
+			Ray.TMin = 0;
 			ApplyPositionBias(Ray, SignedPositionBias * Payload.WorldNormal, AdaptiveSamplingData.MaxNormalBias);
-			RayThroughput *= Throughput;
-			RayPdf *= SamplePdf;
+			float3 NextPathThroughput = PathThroughput * MaterialThroughput / MaterialSamplePdf;
+			if (!any(NextPathThroughput > 0))
+			{
+				// no energy left in this path
+				break;
+			}
 
-			// Russian roulette
-			float PathEnergy = max(max(RayThroughput.x, RayThroughput.y), RayThroughput.z);
-			float EnergyThreshold = 1.0; 
-			if (PathEnergy < EnergyThreshold && Bounces >= PathTracingData.RussianRouletteStartingBounce)
+			// Russian roulette:
+			float ContinuationProb = saturate(max(NextPathThroughput.x, max(NextPathThroughput.y, NextPathThroughput.z)) / max(PathThroughput.x, max(PathThroughput.y, PathThroughput.z)));
+			if (ContinuationProb < 1)
 			{
-				float Q = 1.0 - PathEnergy;
-				if (RandomSequence_GenerateSample1D(RandSequence, SampleIndex) < Q)
+				// If there is some chance we should terminate the ray, draw an extra random value
+				if (RandomSequence_GenerateSample1D(RandSequence, SampleIndex) >= ContinuationProb)
 				{
+					// stochastically terminate the path
 					break;
 				}
-
-				RayThroughput /= (1.0 - Q);
+				PathThroughput = NextPathThroughput / ContinuationProb;
 			}
 
-			// Firefly rejection
-			PathEnergy = max(max(RayThroughput.x, RayThroughput.y), RayThroughput.z);
-			float FireflyThreshold = 5.0;
-			if (PathEnergy / RayPdf > PathTracingData.FireFlyRejectionThreshold)
+			if (MISMode == 2)
 			{
-				break;
+				// If we are doing MIS, the Material ray on the next loop through will need to know 
+				// #dxr_todo: the SamplePdf is not using the right measure for singular cases, this should be handled at the Material API level
+				PrevMaterialPdf = bIsNeeValid ? MaterialSamplePdf : 1e27;
 			}
-
-			// Trace again
-			const uint RayFlags = 0;
-			const uint InstanceInclusionMask = RAY_TRACING_MASK_ALL;
-			const bool bDisableSkyLightContribution = false;
-
-			PrevMaterialPayload = Payload;
-			Payload = TraceMaterialRay(
-				TLAS,
-				RayFlags,
-				InstanceInclusionMask,
-				Ray,
-				RayCone,
-				RenderTargetPos,
-				bDisableSkyLightContribution,
-				bIgnoreTranslucentMaterials);
-			RayCounter++;
 		}
 
 		RadianceRT[RenderTargetPos].rgb += Irradiance;
diff --git a/Engine/Source/Runtime/Renderer/Private/PathTracing.cpp b/Engine/Source/Runtime/Renderer/Private/PathTracing.cpp
--- a/Engine/Source/Runtime/Renderer/Private/PathTracing.cpp
+++ b/Engine/Source/Runtime/Renderer/Private/PathTracing.cpp
@@ -29,6 +29,24 @@ TAutoConsoleVariable<int32> CVarPathTracingSamplesPerPixel(
 	ECVF_RenderThreadSafe
 );
 
+TAutoConsoleVariable<int32> CVarPathTracingMISMode(
+	TEXT("r.PathTracing.MISMode"),
+	2,
+	TEXT("Selects the sampling techniques (default = 2 (MIS enabled))\n")
+	TEXT("0: Material sampling\n")
+	TEXT("1: Light sampling\n")
+	TEXT("2: MIS betwen material and light sampling (default)\n")
+);
+
+TAutoConsoleVariable<int32> CVarPathTracingVisibleLights(
+	TEXT("r.PathTracing.VisibleLights"),
+	0,
+	TEXT("Should light sources be visible to camera rays? (default = 0 (off))\n")
+	TEXT("0: Hide lights from camera rays (default)\n")
+	TEXT("1: Make lights visible to camera\n")
+);
+
+
 TAutoConsoleVariable<int32> CVarPathTracingFrameIndependentTemporalSeed(
 	TEXT("r.PathTracing.FrameIndependentTemporalSeed"),
 	1,
@@ -303,6 +321,25 @@ public:
 			AdaptiveSamplingData.MaxNormalBias = GetRaytracingMaxNormalBias();
 			AdaptiveSamplingData.UseAdaptiveSampling = CVarPathTracingAdaptiveSampling.GetValueOnRenderThread();
 			AdaptiveSamplingData.RandomSequence = CVarPathTracingRandomSequence.GetValueOnRenderThread();
+			AdaptiveSamplingData.MISMode = CVarPathTracingMISMode.GetValueOnRenderThread();
+			AdaptiveSamplingData.VisibleLights = CVarPathTracingVisibleLights.GetValueOnRenderThread();
+
+			// Changing MIS mode requires starting over
+			static uint32 PreviousMISMode = AdaptiveSamplingData.MISMode;
+			if (PreviousMISMode != AdaptiveSamplingData.MISMode)
+			{
+				Scene->bPathTracingNeedsInvalidation = true;
+				PreviousMISMode = AdaptiveSamplingData.MISMode;
+			}
+
+			// Changing VisibleLights requires starting over
+			static uint32 PreviousVisibleLights = AdaptiveSamplingData.VisibleLights;
+			if (PreviousVisibleLights != AdaptiveSamplingData.VisibleLights)
+			{
+				Scene->bPathTracingNeedsInvalidation = true;
+				PreviousVisibleLights = AdaptiveSamplingData.VisibleLights;
+			}
+
 			if (VarianceMipTree.NumBytes > 0)
 			{
 				AdaptiveSamplingData.Iteration = Iteration;
diff --git a/Engine/Source/Runtime/Renderer/Private/PathTracingUniformBuffers.h b/Engine/Source/Runtime/Renderer/Private/PathTracingUniformBuffers.h
--- a/Engine/Source/Runtime/Renderer/Private/PathTracingUniformBuffers.h
+++ b/Engine/Source/Runtime/Renderer/Private/PathTracingUniformBuffers.h
@@ -47,6 +47,8 @@ BEGIN_GLOBAL_SHADER_PARAMETER_STRUCT(FPathTracingAdaptiveSamplingData, )
 	SHADER_PARAMETER(uint32, MinimumSamplesPerPixel)
 	SHADER_PARAMETER(uint32, Iteration)
 	SHADER_PARAMETER(uint32, TemporalSeed)
+	SHADER_PARAMETER(uint32, MISMode)
+	SHADER_PARAMETER(uint32, VisibleLights)
 	SHADER_PARAMETER(float, MaxNormalBias)
 	SHADER_PARAMETER(FIntVector, VarianceDimensions)
 	SHADER_PARAMETER_SRV(Buffer<float>, VarianceMipTree)
