Fixed compilation issues and some warnings

diff --git a/Engine/Shaders/Private/Strata/Strata.ush b/Engine/Shaders/Private/Strata/Strata.ush
--- a/Engine/Shaders/Private/Strata/Strata.ush
+++ b/Engine/Shaders/Private/Strata/Strata.ush
@@ -593,11 +593,10 @@ FStrataPixelBSDFs UnpackStrataIn(ByteAddressBuffer StrataBuffer, uint ByteOffset
 	FStrataPixelBSDFs PixelBSDFs = (FStrataPixelBSDFs)0;
 
 	STRATA_LOAD_UINT1(PixelBSDFs.BSDFCount);
-	PixelBSDFs.BSDFCount = min(PixelBSDFs.BSDFCount, STRATA_MAX_TOTAL_BSDF);	// To make the compiler happy with unrolling
+	PixelBSDFs.BSDFCount = min(PixelBSDFs.BSDFCount, STRATA_MAX_TOTAL_BSDF);	// To make the compiler happy
 
 	float DummyFloat = 0.0f;
 
-	UNROLL
 	for (int i = 0; i < PixelBSDFs.BSDFCount; ++i)
 	{
 		{
diff --git a/Engine/Shaders/Private/Strata/StrataVisualize.usf b/Engine/Shaders/Private/Strata/StrataVisualize.usf
--- a/Engine/Shaders/Private/Strata/StrataVisualize.usf
+++ b/Engine/Shaders/Private/Strata/StrataVisualize.usf
@@ -6,13 +6,20 @@
 
 #include "/Engine/Private/MiniFontCommon.ush"
 
+
+// On some platforms, UnpackStrataIn causes the following error@
+// (error, code:7203) - not enough registers available for the entire program. In order to permit spilling to memory please either disable this error via '#pragma warning (disable:7203)' or reduce it to a warning via '#pragma warning (default:7203)'
+// To fix this and increase occupancy, we will have to progressively unpack the BSDFs on the fly (and thus only have a single BSDF allocating registers)
+#pragma warning (disable:7203)
+
+
 void VisualizeMaterialPS(
 	float4 SVPos : SV_POSITION,
 	out float4 OutColor : SV_Target0)
 {
-	const int2 PixelPos = SVPos.xy;
+	const int2 PixelPos = int2(SVPos.xy);
 	const int2 MaterialPos = View.CursorPosition;
-	const float InspectSquare = 10.0f;
+	const int  InspectSquare = 10;
 
 	OutColor = float4(0.0f, 0.0f, 0.0f, 1.0f);
 
@@ -23,7 +30,7 @@ void VisualizeMaterialPS(
 
 	if (Strata.MaterialLobesTexture[MaterialPos].x > 0.0f)
 	{
-		uint PixelStrataDataByteOffset = (MaterialPos.y * View.BufferSizeAndInvSize.x + MaterialPos.x) * Strata.MaxBytesPerPixel;
+		uint PixelStrataDataByteOffset = (MaterialPos.y * uint(View.BufferSizeAndInvSize.x) + MaterialPos.x) * Strata.MaxBytesPerPixel;
 		FStrataPixelBSDFs PixelBSDFs = UnpackStrataIn(Strata.MaterialLobesBuffer, PixelStrataDataByteOffset);
 
 		float3 BSDFFontColor = float3(0.25, 1.0, 0.25);
