Fix depth bias direction on cube maps at shadow quality <= 2

diff --git a/Engine/Shaders/Private/ShadowProjectionCommon.ush b/Engine/Shaders/Private/ShadowProjectionCommon.ush
--- a/Engine/Shaders/Private/ShadowProjectionCommon.ush
+++ b/Engine/Shaders/Private/ShadowProjectionCommon.ush
@@ -198,7 +198,7 @@ float CubemapHardwarePCF(
 
 #if SHADOW_QUALITY <= 2
 
-		Shadow = InShadowDepthCubeTexture.SampleCmpLevelZero(InShadowDepthCubeTextureSampler, WorldSampleToLightVec, CompareDistance + ShadowDepthBias);
+		Shadow = InShadowDepthCubeTexture.SampleCmpLevelZero(InShadowDepthCubeTextureSampler, WorldSampleToLightVec, CompareDistance - ShadowDepthBias);
 
 #elif SHADOW_QUALITY == 3
 		UNROLL for(int i = 0; i < 5; ++i)
