Fixed diffuse environment lighting for strata material. (diffuse convolution was applied twice)
Sky visibility is now a shared code path for both Disney and Strata path. This allows general AO and diffuse lighting for the dynamic sky light path to now matching accross the board.
Added strata material shading model to make sure it is not used as unlit where most passes optimise things out (i.e. DFAO, bent normal or SSAO)
Roughness is still ignored for diffuse lighting contribution.

Tested with strata disabled

diff --git a/Engine/Shaders/Private/BasePassPixelShader.usf b/Engine/Shaders/Private/BasePassPixelShader.usf
--- a/Engine/Shaders/Private/BasePassPixelShader.usf
+++ b/Engine/Shaders/Private/BasePassPixelShader.usf
@@ -1466,6 +1466,8 @@ void FPixelShaderInOut_MainPS(
 	GBuffer.BaseColor = 0;
 	GBuffer.Specular = 0;
 	GBuffer.Metallic = 0;
+	GBuffer.ShadingModelID = SHADINGMODELID_STRATA;
+	GBuffer.WorldNormal = MaterialParameters.WorldNormal;
 
 	OpaqueBasePass.Strata.MaterialLobesTextureUAV[uint2(In.SvPosition.xy)] = Strata.LayerCount; // Could store a bit pack of layer count, BSDF type masks for tiling, etc.
 
diff --git a/Engine/Shaders/Private/ReflectionEnvironmentPixelShader.usf b/Engine/Shaders/Private/ReflectionEnvironmentPixelShader.usf
--- a/Engine/Shaders/Private/ReflectionEnvironmentPixelShader.usf
+++ b/Engine/Shaders/Private/ReflectionEnvironmentPixelShader.usf
@@ -19,6 +19,96 @@
 #include "ReflectionEnvironmentComposite.ush"
 
 
+
+Texture2D ScreenSpaceReflectionsTexture;
+SamplerState ScreenSpaceReflectionsSampler;
+
+#if ENABLE_DYNAMIC_SKY_LIGHT
+#include "VolumetricCloudCommon.ush"
+Texture2D<float3> CloudSkyAOTexture;
+SamplerState CloudSkyAOSampler;
+float4x4 CloudSkyAOWorldToLightClipMatrix;
+float CloudSkyAOFarDepthKm;
+int CloudSkyAOEnabled;
+#endif
+
+Texture2D AmbientOcclusionTexture;
+SamplerState AmbientOcclusionSampler;
+
+float3 ContrastAndNormalizeMulAdd;
+float OcclusionExponent;
+float OcclusionCombineMode;
+
+
+
+struct FSkyLightVisibilityData
+{
+	float  SkyDiffuseLookUpMul;
+	float3 SkyDiffuseLookUpAdd;
+	float3 SkyDiffuseLookUpNormal;
+};
+FSkyLightVisibilityData GetSkyLightVisibilityData(FGBufferData GBuffer, float AmbientOcclusion, float3 BentNormal)
+{
+	float SkyVisibility = 1;
+	float DotProductFactor = 1;
+	float3 SkyLightingNormal = GBuffer.WorldNormal;
+	BRANCH if (GBuffer.ShadingModelID == SHADINGMODELID_CLEAR_COAT && CLEAR_COAT_BOTTOM_NORMAL)
+	{
+		const float2 oct1 = ((float2(GBuffer.CustomData.a, GBuffer.CustomData.z) * 2) - (256.0 / 255.0)) + UnitVectorToOctahedron(GBuffer.WorldNormal);
+		SkyLightingNormal = OctahedronToUnitVector(oct1);
+	}
+
+#if APPLY_SKY_SHADOWING
+#define USE_DIRECTIONAL_OCCLUSION_ON_SKY_DIFFUSE 1
+
+#if USE_DIRECTIONAL_OCCLUSION_ON_SKY_DIFFUSE
+	{
+		SkyVisibility = length(BentNormal);
+		float3 NormalizedBentNormal = BentNormal / (max(SkyVisibility, .00001f));
+
+		// Use more bent normal in corners
+		float BentNormalWeightFactor = SkyVisibility;
+
+		SkyLightingNormal = lerp(NormalizedBentNormal, GBuffer.WorldNormal, BentNormalWeightFactor);
+		DotProductFactor = lerp(dot(NormalizedBentNormal, GBuffer.WorldNormal), 1, BentNormalWeightFactor);
+	}
+#else
+	{
+		SkyVisibility = length(BentNormal);
+	}
+#endif
+
+	float ContrastCurve = 1 / (1 + exp(-ContrastAndNormalizeMulAdd.x * (SkyVisibility * 10 - 5)));
+	SkyVisibility = saturate(ContrastCurve * ContrastAndNormalizeMulAdd.y + ContrastAndNormalizeMulAdd.z);
+
+#endif
+
+	// Apply DFAO controls
+	SkyVisibility = pow(SkyVisibility, OcclusionExponent);
+	SkyVisibility = lerp(SkyVisibility, 1, OcclusionTintAndMinOcclusion.w);
+
+	// Combine with other AO sources
+	if (OcclusionCombineMode == 0)
+	{
+		// Combine with min which nicely avoids over-occlusion in cases where strong DFAO is present along with strong SSAO (distant trees)
+		SkyVisibility = min(SkyVisibility, min(GBuffer.GBufferAO, AmbientOcclusion));
+	}
+	else
+	{
+		// Combine with mul, which continues to add SSAO depth even indoors.  SSAO will need to be tweaked to be less strong.
+		SkyVisibility = SkyVisibility * min(GBuffer.GBufferAO, AmbientOcclusion);
+	}
+
+	// Apply AO to the sky diffuse and account for darkening due to the geometry term
+	// apply the Diffuse color to the lighting (including OcclusionTintAndMinOcclusion as it's considered another light, that fixes SubsurfaceProfile being too dark)
+	FSkyLightVisibilityData SkyVisData;
+	SkyVisData.SkyDiffuseLookUpMul = SkyVisibility * DotProductFactor;
+	SkyVisData.SkyDiffuseLookUpAdd = (1 - SkyVisibility) * OcclusionTintAndMinOcclusion.xyz;
+	SkyVisData.SkyDiffuseLookUpNormal = SkyLightingNormal;
+	return SkyVisData;
+}
+
+
 #if STRATA_ENABLED
 #include "/Engine/Private/Strata/Strata.ush"
 #include "/Engine/Private/ParticipatingMediaCommon.ush"
@@ -54,27 +144,6 @@
 		bCompositeSkylight);
 }
 
-
-Texture2D ScreenSpaceReflectionsTexture;
-SamplerState ScreenSpaceReflectionsSampler;
-
-#if ENABLE_DYNAMIC_SKY_LIGHT
-#include "VolumetricCloudCommon.ush"
-Texture2D<float3> CloudSkyAOTexture;
-SamplerState CloudSkyAOSampler;
-float4x4 CloudSkyAOWorldToLightClipMatrix;
-float CloudSkyAOFarDepthKm;
-int CloudSkyAOEnabled;
-#endif
-
-Texture2D AmbientOcclusionTexture;
-SamplerState AmbientOcclusionSampler;
-
-
-float3 ContrastAndNormalizeMulAdd;
-float OcclusionExponent;
-float OcclusionCombineMode;
-
 void RemapClearCoatDiffuseAndSpecularColor(FGBufferData GBuffer, float2 ScreenPosition, inout float3 DiffuseColor, inout float3 SpecularColor)
 {
 	if (GBuffer.ShadingModelID == SHADINGMODELID_CLEAR_COAT)
@@ -110,64 +179,14 @@
 	float2 UV = BufferUV;
 	float3 Lighting = 0;
 
-	float SkyVisibility = 1;
-	float DotProductFactor = 1;
-	float3 SkyLightingNormal = GBuffer.WorldNormal;
-	BRANCH if( GBuffer.ShadingModelID == SHADINGMODELID_CLEAR_COAT && CLEAR_COAT_BOTTOM_NORMAL)
-	{
-		const float2 oct1 = ((float2(GBuffer.CustomData.a, GBuffer.CustomData.z) * 2) - (256.0/255.0)) + UnitVectorToOctahedron(GBuffer.WorldNormal);
-		SkyLightingNormal = OctahedronToUnitVector(oct1);			
-	}
-	
-#if APPLY_SKY_SHADOWING
-#define USE_DIRECTIONAL_OCCLUSION_ON_SKY_DIFFUSE 1
+	FSkyLightVisibilityData SkyVisData = GetSkyLightVisibilityData(GBuffer, AmbientOcclusion, BentNormal);
 
-#if USE_DIRECTIONAL_OCCLUSION_ON_SKY_DIFFUSE
-	{
-		SkyVisibility = length(BentNormal);
-		float3 NormalizedBentNormal = BentNormal / (max(SkyVisibility, .00001f));
-
-		// Use more bent normal in corners
-		float BentNormalWeightFactor = SkyVisibility;
-
-		SkyLightingNormal = lerp(NormalizedBentNormal, GBuffer.WorldNormal, BentNormalWeightFactor);
-		DotProductFactor = lerp(dot(NormalizedBentNormal, GBuffer.WorldNormal), 1, BentNormalWeightFactor);
-	}
-#else
-	{
-		SkyVisibility = length(BentNormal);
-	}
-#endif
-
-	float ContrastCurve = 1 / (1 + exp(-ContrastAndNormalizeMulAdd.x * (SkyVisibility * 10 - 5)));
-	SkyVisibility = saturate(ContrastCurve * ContrastAndNormalizeMulAdd.y + ContrastAndNormalizeMulAdd.z);
-
-#endif
-
-	// Apply DFAO controls
-	SkyVisibility = pow(SkyVisibility, OcclusionExponent);
-	SkyVisibility = lerp(SkyVisibility, 1, OcclusionTintAndMinOcclusion.w);
-
-	// Combine with other AO sources
-	if (OcclusionCombineMode == 0)
-	{
-		// Combine with min which nicely avoids over-occlusion in cases where strong DFAO is present along with strong SSAO (distant trees)
-		SkyVisibility = min(SkyVisibility, min(GBuffer.GBufferAO, AmbientOcclusion));
-	}
-	else
-	{
-		// Combine with mul, which continues to add SSAO depth even indoors.  SSAO will need to be tweaked to be less strong.
-		SkyVisibility = SkyVisibility * min(GBuffer.GBufferAO, AmbientOcclusion);
-	}
-
-	float ScalarFactors = SkyVisibility;
-
 	BRANCH
 	if (GBuffer.ShadingModelID == SHADINGMODELID_TWOSIDED_FOLIAGE)
 	{
 		float3 SubsurfaceLookup = GetSkySHDiffuse(-GBuffer.WorldNormal) * View.SkyLightColor.rgb;
 		float3 SubsurfaceColor = ExtractSubsurfaceColor(GBuffer);
-		Lighting += ScalarFactors * SubsurfaceLookup * SubsurfaceColor;
+		Lighting += SkyVisData.SkyDiffuseLookUpMul * SubsurfaceLookup * SubsurfaceColor;
 	}
 
 	if (GBuffer.ShadingModelID == SHADINGMODELID_SUBSURFACE || GBuffer.ShadingModelID == SHADINGMODELID_PREINTEGRATED_SKIN)
@@ -183,7 +202,7 @@
 		float3 N = GBuffer.WorldNormal;
 		float3 V = -normalize(mul(float4(ScreenPosition, 1, 0), View.ScreenToWorld).xyz);
 		float3 L = normalize(V - N * dot(V, N));
-		SkyLightingNormal = L;
+		SkyVisData.SkyDiffuseLookUpNormal = L;
 		FHairTransmittanceData TransmittanceData = InitHairTransmittanceData(true);
 		DiffuseColor = PI * HairShading(GBuffer, L, V, N, 1, TransmittanceData, 0, 0.2, uint2(0, 0));
 	}
@@ -195,12 +214,10 @@
 	}
 
 	// Compute the preconvolved incoming lighting with the bent normal direction
-	float3 DiffuseLookup = GetSkySHDiffuse(SkyLightingNormal) * View.SkyLightColor.rgb;
+	float3 DiffuseLookup = GetSkySHDiffuse(SkyVisData.SkyDiffuseLookUpNormal) * View.SkyLightColor.rgb;
 
-	// Apply AO to the sky diffuse and account for darkening due to the geometry term
-	// apply the Diffuse color to the lighting (including OcclusionTintAndMinOcclusion as it's considered another light, that fixes SubsurfaceProfile being too dark)
-	Lighting += ((ScalarFactors * DotProductFactor) * DiffuseLookup + (1 - SkyVisibility) * OcclusionTintAndMinOcclusion.xyz) * DiffuseColor;
-
+	// And accumulate the lighting
+	Lighting += (SkyVisData.SkyDiffuseLookUpMul * DiffuseLookup + SkyVisData.SkyDiffuseLookUpAdd) * DiffuseColor;
 #if USE_PREEXPOSURE
 	Lighting *= View.PreExposure;
 #endif
@@ -362,12 +379,50 @@
 	
 	OutColor = 0.0f;
 
-	#if ENABLE_DYNAMIC_SKY_LIGHT
+#if ENABLE_DYNAMIC_SKY_LIGHT
+
+	uint2 PixelPos = SvPosition.xy;
+
+	const bool bUnlitMaterial = ShadingModelID == SHADINGMODELID_UNLIT;
+#if STRATA_ENABLED
+	const bool bStrataMaterial = Strata.MaterialLobesTexture[PixelPos].x > 0.0f;
+#else
+	const bool bStrataMaterial = false;
+#endif
+
+	float3 SkyLighting = 0;
 	BRANCH
-	// Only light pixels marked as lit
-	if (ShadingModelID != SHADINGMODELID_UNLIT)
+#if STRATA_ENABLED
+#if ENABLE_SKY_LIGHT
+	if (bStrataMaterial)
 	{
-		float CloudAmbientOcclusion = 1.0f;
+		uint PixelStrataDataByteOffset = (PixelPos.y * View.BufferSizeAndInvSize.x + PixelPos.x) * Strata.MaxBytesPerPixel;
+		FStrataPixelBSDFs PixelBSDFs = UnpackStrataIn(Strata.MaterialLobesBuffer, PixelStrataDataByteOffset);
+
+		float3 WorldPosition = mul(float4(ScreenPosition * GBuffer.Depth, GBuffer.Depth, 1), View.ScreenToWorld).xyz;
+		float3 CameraToPixel = normalize(WorldPosition - View.WorldCameraOrigin);
+		float3 V = -CameraToPixel;
+		
+		SkyLighting = StrataEnvironmentLighting(
+			BufferUV,
+			GBuffer,
+			WorldPosition,
+			CameraToPixel,
+			V,
+			AmbientOcclusion,
+			PixelBSDFs);
+	}
+	else
+#endif
+#endif
+	if (!bUnlitMaterial) // Only light pixels marked as lit
+	{
+		SkyLighting = SkyLightDiffuse(GBuffer, AmbientOcclusion, BufferUV, ScreenPosition, BentNormal, DiffuseColor);
+	}
+	
+	BRANCH 
+	if (!bUnlitMaterial || bStrataMaterial)
+	{
 		// Ideally we would compute spatially how much of the sky+cloud is visible for each point in the world. But we evaluate the sky light only once at the sky light component position as of today.
 		// To add some spatial variation, we can affect the sky light diffuse contribution according to cloud occlusion from above.
 		// This is an approximation because clouds are also occluding the sky in the sky light capture (a bit of a double contribution then). but it does help by adding spatially varying details.
@@ -376,52 +431,19 @@
 		{
 			float3 WorldPosition = mul(float4(ScreenPosition * GBuffer.Depth, GBuffer.Depth, 1), View.ScreenToWorld).xyz;
 			float OutOpticalDepth = 0.0f;
-			CloudAmbientOcclusion = GetCloudVolumetricShadow(WorldPosition, CloudSkyAOWorldToLightClipMatrix, CloudSkyAOFarDepthKm, CloudSkyAOTexture, CloudSkyAOSampler, OutOpticalDepth);
+			SkyLighting *= GetCloudVolumetricShadow(WorldPosition, CloudSkyAOWorldToLightClipMatrix, CloudSkyAOFarDepthKm, CloudSkyAOTexture, CloudSkyAOSampler, OutOpticalDepth);
 		}
 #endif
-
-		float3 SkyLighting = 0;
-#if STRATA_ENABLED
-#if ENABLE_SKY_LIGHT
-		uint2 PixelPos = SvPosition.xy;
-		if (Strata.MaterialLobesTexture[PixelPos].x > 0.0f)
-		{
-			uint PixelStrataDataByteOffset = (PixelPos.y * View.BufferSizeAndInvSize.x + PixelPos.x) * Strata.MaxBytesPerPixel;
-			FStrataPixelBSDFs PixelBSDFs = UnpackStrataIn(Strata.MaterialLobesBuffer, PixelStrataDataByteOffset);
-
-			float3 WorldPosition = mul(float4(ScreenPosition * GBuffer.Depth, GBuffer.Depth, 1), View.ScreenToWorld).xyz;
-			float3 CameraToPixel = normalize(WorldPosition - View.WorldCameraOrigin);
-			float3 V = -CameraToPixel;
-			float3 N = GBuffer.WorldNormal;
-
-			float3 R = 2 * dot(V, N) * N - V;
-
-			SkyLighting = StrataEnvironmentLighting(
-				GBuffer,
-				WorldPosition,
-				CameraToPixel,
-				V,
-				N,
-				R,
-				AmbientOcclusion,
-				PixelBSDFs);
-		}
-		else
-#endif
-#endif
-		{
-			SkyLighting = SkyLightDiffuse(GBuffer, AmbientOcclusion, BufferUV, ScreenPosition, BentNormal, DiffuseColor)* CloudAmbientOcclusion;
-		}
-
-		// Scene color alpha is used for ScreenSpaceSubsurfaceScattering (if that is not needed it can be disabled with SUBSURFACE_CHANNEL_MODE)
-		FLightAccumulator LightAccumulator = (FLightAccumulator)0;
-
-		const bool bNeedsSeparateSubsurfaceLightAccumulation = UseSubsurfaceProfile(ShadingModelID);
-
-		LightAccumulator_Add(LightAccumulator, SkyLighting, SkyLighting, 1.0f, bNeedsSeparateSubsurfaceLightAccumulation);
-		OutColor = LightAccumulator_GetResult(LightAccumulator);
 	}
-	#endif
 
+	// Scene color alpha is used for ScreenSpaceSubsurfaceScattering (if that is not needed it can be disabled with SUBSURFACE_CHANNEL_MODE)
+	FLightAccumulator LightAccumulator = (FLightAccumulator)0;
+	const bool bNeedsSeparateSubsurfaceLightAccumulation = UseSubsurfaceProfile(ShadingModelID);
+
+	LightAccumulator_Add(LightAccumulator, SkyLighting, SkyLighting, 1.0f, bNeedsSeparateSubsurfaceLightAccumulation);
+	OutColor = LightAccumulator_GetResult(LightAccumulator);
+
+#endif // ENABLE_DYNAMIC_SKY_LIGHT
+
 	BRANCH
 	if (ShadingModelID != SHADINGMODELID_UNLIT && ShadingModelID != SHADINGMODELID_HAIR)
diff --git a/Engine/Shaders/Private/ShadingCommon.ush b/Engine/Shaders/Private/ShadingCommon.ush
--- a/Engine/Shaders/Private/ShadingCommon.ush
+++ b/Engine/Shaders/Private/ShadingCommon.ush
@@ -17,7 +17,8 @@
 #define SHADINGMODELID_EYE					9
 #define SHADINGMODELID_SINGLELAYERWATER		10
 #define SHADINGMODELID_THIN_TRANSLUCENT		11
-#define SHADINGMODELID_NUM					12
+#define SHADINGMODELID_STRATA				12		// Temporary while we convert everything to Strata
+#define SHADINGMODELID_NUM					13
 #define SHADINGMODELID_MASK					0xF		// 4 bits reserved for ShadingModelID			
 
 // The flags are defined so that 0 value has no effect!
@@ -53,6 +54,7 @@ float3 GetShadingModelColor(uint ShadingModelID)
 	else if (ShadingModelID == SHADINGMODELID_EYE) return float3(0.3f, 1.0f, 1.0f); 
 	else if (ShadingModelID == SHADINGMODELID_SINGLELAYERWATER) return float3(0.5f, 0.5f, 1.0f);
 	else if (ShadingModelID == SHADINGMODELID_THIN_TRANSLUCENT) return float3(1.0f, 0.8f, 0.3f);
+	else if (ShadingModelID == SHADINGMODELID_STRATA) return float3(1.0f, 1.0f, 0.0f);
 	else return float3(1.0f, 1.0f, 1.0f); // White
 #else
 	switch(ShadingModelID)
@@ -69,6 +71,7 @@ float3 GetShadingModelColor(uint ShadingModelID)
 		case SHADINGMODELID_EYE: return float3(0.3f, 1.0f, 1.0f);
 		case SHADINGMODELID_SINGLELAYERWATER: return float3(0.5f, 0.5f, 1.0f);
 		case SHADINGMODELID_THIN_TRANSLUCENT: return float3(1.0f, 0.8f, 0.3f);
+		case SHADINGMODELID_STRATA: return float3(1.0f, 1.0f, 0.0f);
 		default: return float3(1.0f, 1.0f, 1.0f); // White
 	}
 #endif
diff --git a/Engine/Shaders/Private/Strata/StrataEnvironmentLighting.ush b/Engine/Shaders/Private/Strata/StrataEnvironmentLighting.ush
--- a/Engine/Shaders/Private/Strata/StrataEnvironmentLighting.ush
+++ b/Engine/Shaders/Private/Strata/StrataEnvironmentLighting.ush
@@ -5,12 +5,11 @@
 
 
 float3 StrataEnvironmentLighting(
+	float2 BufferUV,
 	FGBufferData GBuffer,
 	float3 WorldPosition,
 	float3 CameraToPixel,
 	float3 V,
-	float3 N,
-	float3 R,
 	float3 AmbientOcclusion,
 	FStrataPixelBSDFs PixelBSDFs
 	)
@@ -21,11 +20,6 @@ float3 StrataEnvironmentLighting(
 		FStrataBSDF BSDF = PixelBSDFs.BSDFs[i];
 		const float3 BSDFThroughput = BSDF.Weight;
 
-		BxDFContext Context = (BxDFContext)0;
-		float3 L = N;
-		Init(Context, N, V, L);
-		float SatNoL = saturate(Context.NoL);
-
 		float3 Albedo = 0.0f;
 		float3 Normal = 0.0f;
 		float3 SpecularLuminance = 0.0f;
@@ -41,45 +35,53 @@ float3 StrataEnvironmentLighting(
 		case STRATA_BSDF_TYPE_DIELECTRIC:
 		{
 			BRANCH
-				if (ReflectionStruct.SkyLightParameters.y > 0)
-				{
-					float F0 = DielectricIorToF0(BSDF.Dielectric.IOR);
-
-					const float SafeRoughness = MakeRoughnessSafe(BSDF.Dielectric.Roughness.x);
-					float a2 = Pow4(SafeRoughness);
-					float Energy = 1.0f; // See EnergyNormalization
-
-					// Generalized microfacet specular
-					//float D = D_GGX(a2, Context.NoH) * Energy;
-					//float Vis = Vis_SmithJointApprox(a2, Context.NoV, SatNoL);
-					//float3 F = F_Schlick(F0, Context.VoH);
-
-					float3 EvalEnvBRDF = EnvBRDFApprox(F0, SafeRoughness, Context.NoV);
-					float SkyAverageBrightness = 1.0f;
-					SpecularLuminance += GetSkyLightReflection(R, SafeRoughness, SkyAverageBrightness) * EvalEnvBRDF * BSDF.Dielectric.Tint;
-				}
+			if (ReflectionStruct.SkyLightParameters.y > 0)
+			{
+				float F0 = DielectricIorToF0(BSDF.Dielectric.IOR);
+
+				const float SafeRoughness = MakeRoughnessSafe(BSDF.Dielectric.Roughness.x);
+				float a2 = Pow4(SafeRoughness);
+				float Energy = 1.0f; // See EnergyNormalization
+
+				// Generalized microfacet specular
+				//float D = D_GGX(a2, Context.NoH) * Energy;
+				//float Vis = Vis_SmithJointApprox(a2, Context.NoV, SatNoL);
+				//float3 F = F_Schlick(F0, Context.VoH);
+
+				const float3 N = BSDF.Dielectric.Normal;
+				const float SatNoV = saturate(dot(N, V));
+				float3 R = 2 * dot(V, N) * N - V;
+
+				float3 EvalEnvBRDF = EnvBRDFApprox(F0, SafeRoughness, SatNoV);
+				float SkyAverageBrightness = 1.0f;
+				SpecularLuminance += GetSkyLightReflection(R, SafeRoughness, SkyAverageBrightness) * EvalEnvBRDF * BSDF.Dielectric.Tint;
+			}
 		}
 		break;
 		case STRATA_BSDF_TYPE_CONDUCTOR:
 		{
 			BRANCH
-				if (ReflectionStruct.SkyLightParameters.y > 0)
-				{
-					// See SpecularGGX, isotropic verion
-
-					const float SafeRoughness = MakeRoughnessSafe(BSDF.Conductor.Roughness.x);
-					float a2 = Pow4(SafeRoughness);
-					float Energy = 1.0f; // See EnergyNormalization
-
-					// Generalized microfacet specular
-					//float D = D_GGX(a2, Context.NoH) * Energy;
-					//float Vis = Vis_SmithJointApprox(a2, Context.NoV, SatNoL);
-					//float3 F = F_Schlick(BSDF.Conductor.Reflectivity, BSDF.Conductor.EdgeColor, Context.VoH); // STRATA_TODO check that this is correct 
-
-					float3 EvalEnvBRDF = EnvBRDFApprox(BSDF.Conductor.Reflectivity, BSDF.Conductor.EdgeColor, SafeRoughness, Context.NoV);
-					float SkyAverageBrightness = 1.0f;
-					SpecularLuminance += GetSkyLightReflection(R, SafeRoughness, SkyAverageBrightness) * EvalEnvBRDF;
-				}
+			if (ReflectionStruct.SkyLightParameters.y > 0)
+			{
+				// See SpecularGGX, isotropic verion
+
+				const float SafeRoughness = MakeRoughnessSafe(BSDF.Conductor.Roughness.x);
+				float a2 = Pow4(SafeRoughness);
+				float Energy = 1.0f; // See EnergyNormalization
+
+				// Generalized microfacet specular
+				//float D = D_GGX(a2, Context.NoH) * Energy;
+				//float Vis = Vis_SmithJointApprox(a2, Context.NoV, SatNoL);
+				//float3 F = F_Schlick(BSDF.Conductor.Reflectivity, BSDF.Conductor.EdgeColor, Context.VoH); // STRATA_TODO check that this is correct 
+
+				const float3 N = BSDF.Conductor.Normal;
+				const float SatNoV = saturate(dot(N, V));
+				float3 R = 2 * dot(V, N) * N - V;
+
+				float3 EvalEnvBRDF = EnvBRDFApprox(BSDF.Conductor.Reflectivity, BSDF.Conductor.EdgeColor, SafeRoughness, SatNoV);
+				float SkyAverageBrightness = 1.0f;
+				SpecularLuminance += GetSkyLightReflection(R, SafeRoughness, SkyAverageBrightness) * EvalEnvBRDF;
+			}
 		}
 		break;
 		case STRATA_BSDF_TYPE_VOLUME:
@@ -93,7 +95,8 @@ float3 StrataEnvironmentLighting(
 			const float3 SafeExtinctionThreshold = 0.000001f;
 			const float3 SafeExtinctionCoefficients = max(SafeExtinctionThreshold, ExtinctionCoefficients);
 
-			const float PathLength = Thickness / max(0.0001f, saturate(Context.NoV));
+			const float SatNoV = saturate(dot(GBuffer.WorldNormal, V));
+			const float PathLength = Thickness / max(0.0001f, SatNoV);
 			const float3 SafePathSegmentTransmittance = exp(-SafeExtinctionCoefficients * PathLength);
 			const float3 NormLuminanceTransfert = 1.0f * ScatteringCoefficients;
 			float3 NormLuminanceSingleScattering = (NormLuminanceTransfert - NormLuminanceTransfert * SafePathSegmentTransmittance) / SafeExtinctionCoefficients;
@@ -107,12 +110,26 @@ float3 StrataEnvironmentLighting(
 		break;
 		}
 
+		// Diffuse component
 		if (any(Albedo > 0.0f))
 		{
-			// BentNormal?
-			SkyLighting += BSDFThroughput * Diffuse_Lambert(Albedo) * GetSkySHDiffuse(Normal) * View.SkyLightColor.rgb * AmbientOcclusion;
+			// Apply bent normal if needed
+			#if APPLY_SKY_SHADOWING
+			{
+				Normal = UpsampleDFAO(BufferUV, GBuffer.Depth, Normal); // Bent normal
+			}
+			#endif
+
+			// Compute the common sky visibility factors
+			FSkyLightVisibilityData SkyVisData = GetSkyLightVisibilityData(GBuffer, AmbientOcclusion, Normal);
+
+			// Finally sample the sky diffuse contribution (spherical harmonic, Lambert BRDF)
+			float3 DiffuseLookup = GetSkySHDiffuse(Normal) * View.SkyLightColor.rgb;			// STRATA_TODO compute SH coefficients for ON or Chan
+			// And accumulate
+			SkyLighting += (SkyVisData.SkyDiffuseLookUpMul * DiffuseLookup + SkyVisData.SkyDiffuseLookUpAdd) * Albedo;
 		}
 
+		// Specular component
 		SkyLighting += BSDFThroughput * SpecularLuminance * View.SkyLightColor.rgb * AmbientOcclusion;
 	}
 #if USE_PREEXPOSURE
diff --git a/Engine/Shaders/Private/Strata/StrataForwardLighting.ush b/Engine/Shaders/Private/Strata/StrataForwardLighting.ush
--- a/Engine/Shaders/Private/Strata/StrataForwardLighting.ush
+++ b/Engine/Shaders/Private/Strata/StrataForwardLighting.ush
@@ -63,10 +63,10 @@ float3 StrataForwardLighting(
 
 #if ENABLE_SKY_LIGHT
 				BRANCH
-					if (ReflectionStruct.SkyLightParameters.y > 0)
-					{
-						Color += GetEffectiveSkySHDiffuse(N) * SkyLightCommonFactor * Diffuse_Lambert(BSDF.Diffuse.Albedo); // STRATA_TODO should not be lambert
-					}
+				if (ReflectionStruct.SkyLightParameters.y > 0)
+				{
+					Color += GetEffectiveSkySHDiffuse(N) * BSDF.Diffuse.Albedo * SkyLightCommonFactor; // STRATA_TODO compute SH coefficients for ON
+				}
 #endif
 
 				LayerCoverage	+= BSDFCoverage;
@@ -83,10 +83,10 @@ float3 StrataForwardLighting(
 
 #if ENABLE_SKY_LIGHT
 				BRANCH
-					if (ReflectionStruct.SkyLightParameters.y > 0)
-					{
-						Color += GetEffectiveSkySHDiffuse(N) * SkyLightCommonFactor * Diffuse_Lambert(BSDF.Diffuse.Albedo); // STRATA_TODO should not be lambert
-					}
+				if (ReflectionStruct.SkyLightParameters.y > 0)
+				{
+					Color += GetEffectiveSkySHDiffuse(N) * BSDF.Diffuse.Albedo * SkyLightCommonFactor; // STRATA_TODO compute SH coefficients for Chan
+				}
 #endif
 				
 				LayerCoverage	+= BSDFCoverage;
diff --git a/Engine/Source/Editor/PixelInspector/Private/PixelInspectorResult.cpp b/Engine/Source/Editor/PixelInspector/Private/PixelInspectorResult.cpp
--- a/Engine/Source/Editor/PixelInspector/Private/PixelInspectorResult.cpp
+++ b/Engine/Source/Editor/PixelInspector/Private/PixelInspectorResult.cpp
@@ -237,6 +237,8 @@ namespace PixelInspector
 			return EMaterialShadingModel::MSM_SingleLayerWater;
 		case PIXEL_INSPECTOR_SHADINGMODELID_THIN_TRANSLUCENT:
 			return EMaterialShadingModel::MSM_ThinTranslucent;
+		case PIXEL_INSPECTOR_SHADINGMODELID_STRATA:
+			return EMaterialShadingModel::MSM_Strata;
 		};
 		return EMaterialShadingModel::MSM_DefaultLit;
 	}
@@ -282,6 +284,7 @@ namespace PixelInspector
 		case EMaterialShadingModel::MSM_DefaultLit:
 		case EMaterialShadingModel::MSM_SingleLayerWater:
 		case EMaterialShadingModel::MSM_ThinTranslucent:
+		case EMaterialShadingModel::MSM_Strata:
 		{
 			SubSurfaceColor = FVector(0.0f);
 			Opacity = 0.0f;
diff --git a/Engine/Source/Editor/PixelInspector/Private/PixelInspectorResult.h b/Engine/Source/Editor/PixelInspector/Private/PixelInspectorResult.h
--- a/Engine/Source/Editor/PixelInspector/Private/PixelInspectorResult.h
+++ b/Engine/Source/Editor/PixelInspector/Private/PixelInspectorResult.h
@@ -17,6 +17,7 @@
 #define PIXEL_INSPECTOR_SHADINGMODELID_EYE 9
 #define PIXEL_INSPECTOR_SHADINGMODELID_SINGLELAYERWATER 10
 #define PIXEL_INSPECTOR_SHADINGMODELID_THIN_TRANSLUCENT 11
+#define PIXEL_INSPECTOR_SHADINGMODELID_STRATA 12
 #define PIXEL_INSPECTOR_SHADINGMODELID_MASK 0xF
 
 namespace PixelInspector
diff --git a/Engine/Source/Runtime/Engine/Classes/Engine/EngineTypes.h b/Engine/Source/Runtime/Engine/Classes/Engine/EngineTypes.h
--- a/Engine/Source/Runtime/Engine/Classes/Engine/EngineTypes.h
+++ b/Engine/Source/Runtime/Engine/Classes/Engine/EngineTypes.h
@@ -476,6 +476,7 @@ enum EMaterialShadingModel
 	MSM_Eye						UMETA(DisplayName="Eye"),
 	MSM_SingleLayerWater		UMETA(DisplayName="SingleLayerWater"),
 	MSM_ThinTranslucent			UMETA(DisplayName="Thin Translucent"),
+	MSM_Strata					UMETA(DisplayName="Strata"),
 	/** Number of unique shading models. */
 	MSM_NUM						UMETA(Hidden),
 	/** Shading model will be determined by the Material Expression Graph,
