Fixing EyeAdaptaitonInverse.

diff --git a/Engine/Shaders/Private/EyeAdaptationCommon.ush b/Engine/Shaders/Private/EyeAdaptationCommon.ush
--- a/Engine/Shaders/Private/EyeAdaptationCommon.ush
+++ b/Engine/Shaders/Private/EyeAdaptationCommon.ush
@@ -61,4 +61,23 @@ float EyeAdaptationLookup()
 #else
 	return 0.0f;
 #endif
-}
\ No newline at end of file
+}
+
+float3 EyeAdaptationInverseLookup(float3 LightValue, float Alpha)
+{
+	float Adaptation = EyeAdaptationLookup();
+
+	// When Alpha=0.0, we want to multiply by 1.0. when Alpha = 1.0, we want to multiply by 1/Adaptation.
+	// So the lerped value is:
+	//     LerpLogScale = Lerp(log(1),log(1/Adaptaiton),T)
+	// Which is simplified as:
+	//     LerpLogScale = Lerp(0,-log(Adaptation),T)
+	//     LerpLogScale = -T * logAdaptation;
+
+	float LerpLogScale = -Alpha * log(Adaptation);
+	float Scale = exp(LerpLogScale);
+	return LightValue * Scale;
+}
+
+
+
