Fixed voxel leaking when screen traces are enabled.  The screen trace would place the MinTraceDistance just at the surface, and the voxel trace clamps the RayTime to this value, causing the SDF gradient to be taken too close to the surface, resulting in incorrect normals and leaking.

diff --git a/Engine/Shaders/Private/GlobalDistanceFieldUtils.ush b/Engine/Shaders/Private/GlobalDistanceFieldUtils.ush
--- a/Engine/Shaders/Private/GlobalDistanceFieldUtils.ush
+++ b/Engine/Shaders/Private/GlobalDistanceFieldUtils.ush
@@ -128,7 +128,7 @@ FGlobalSDFTraceResult RayTraceGlobalDistanceField(FGlobalSDFTraceInput TraceInpu
 
 				if (DistanceField < ExpandSurfaceAmount)
 				{
-					TraceResult.HitTime = max(SampleRayTime + DistanceField - ExpandSurfaceAmount, TraceInput.MinTraceDistance);
+					TraceResult.HitTime = max(SampleRayTime + DistanceField - ExpandSurfaceAmount, 0.0f);
 					TraceResult.HitClipmapIndex = ClipmapIndex;
 					break;
 				}
