Fixed strata addressing to be uint only (to avoid hitting float precision limit).

diff --git a/Engine/Shaders/Private/BasePassPixelShader.usf b/Engine/Shaders/Private/BasePassPixelShader.usf
--- a/Engine/Shaders/Private/BasePassPixelShader.usf
+++ b/Engine/Shaders/Private/BasePassPixelShader.usf
@@ -1472,7 +1472,7 @@ void FPixelShaderInOut_MainPS(
 	OpaqueBasePass.Strata.MaterialLobesTextureUAV[uint2(In.SvPosition.xy)] = Strata.LayerCount; // Could store a bit pack of layer count, BSDF type masks for tiling, etc.
 
 	uint2 PixelPos = In.SvPosition.xy;
-	uint PixelStrataDataByteOffset = (PixelPos.y * ResolvedView.BufferSizeAndInvSize.x + PixelPos.x) * OpaqueBasePass.Strata.MaxBytesPerPixel;
+	uint PixelStrataDataByteOffset = GetStrataPixelDataByteOffset(PixelPos, ResolvedView.BufferSizeAndInvSize.x, OpaqueBasePass.Strata.MaxBytesPerPixel);
 
 	PackStrataOut(OpaqueBasePass.Strata.MaterialLobesBufferUAV, PixelStrataDataByteOffset, Strata, MaterialParameters.WorldNormal, MaterialParameters.CameraVector);
 #endif // STRATA_OPAQUE_ENABLED
diff --git a/Engine/Shaders/Private/DeferredLightPixelShaders.usf b/Engine/Shaders/Private/DeferredLightPixelShaders.usf
--- a/Engine/Shaders/Private/DeferredLightPixelShaders.usf
+++ b/Engine/Shaders/Private/DeferredLightPixelShaders.usf
@@ -209,7 +209,7 @@ void DeferredLightPixelMain(
 
 			if (LightMask > 0)
 			{
-				uint PixelStrataDataByteOffset = (PixelPos.y * View.BufferSizeAndInvSize.x + PixelPos.x) * Strata.MaxBytesPerPixel;
+				uint PixelStrataDataByteOffset = GetStrataPixelDataByteOffset(PixelPos, View.BufferSizeAndInvSize.x, Strata.MaxBytesPerPixel);
 				FStrataPixelBSDFs PixelBSDFs = UnpackStrataIn(Strata.MaterialLobesBuffer, PixelStrataDataByteOffset);
 
 				FShadowTerms ShadowTerms = { ScreenSpaceData.AmbientOcclusion, 1.0, 1.0, InitHairTransmittanceData() };
diff --git a/Engine/Shaders/Private/ReflectionEnvironmentPixelShader.usf b/Engine/Shaders/Private/ReflectionEnvironmentPixelShader.usf
--- a/Engine/Shaders/Private/ReflectionEnvironmentPixelShader.usf
+++ b/Engine/Shaders/Private/ReflectionEnvironmentPixelShader.usf
@@ -396,7 +396,7 @@ void ReflectionEnvironmentSkyLighting(
 #if ENABLE_SKY_LIGHT
 	if (bStrataMaterial)
 	{
-		uint PixelStrataDataByteOffset = (PixelPos.y * View.BufferSizeAndInvSize.x + PixelPos.x) * Strata.MaxBytesPerPixel;
+		uint PixelStrataDataByteOffset = GetStrataPixelDataByteOffset(PixelPos, View.BufferSizeAndInvSize.x, Strata.MaxBytesPerPixel);
 		FStrataPixelBSDFs PixelBSDFs = UnpackStrataIn(Strata.MaterialLobesBuffer, PixelStrataDataByteOffset);
 
 		float3 WorldPosition = mul(float4(ScreenPosition * GBuffer.Depth, GBuffer.Depth, 1), View.ScreenToWorld).xyz;
diff --git a/Engine/Shaders/Private/Strata/Strata.ush b/Engine/Shaders/Private/Strata/Strata.ush
--- a/Engine/Shaders/Private/Strata/Strata.ush
+++ b/Engine/Shaders/Private/Strata/Strata.ush
@@ -391,6 +391,13 @@ void StrataLoad_AlbedoGamma2ToLinear(ByteAddressBuffer StrataBuffer, inout uint
 
 
 
+uint GetStrataPixelDataByteOffset(uint2 PixelPos, uint ViewBufferSizeX, uint StrataMaxBytesPerPixel)
+{
+	return (PixelPos.y * ViewBufferSizeX + PixelPos.x) * StrataMaxBytesPerPixel;
+}
+
+
+
 void PackStrataOut(
 	RWByteAddressBuffer StrataBuffer, 
 	uint ByteOffset, 
diff --git a/Engine/Shaders/Private/Strata/StrataVisualize.usf b/Engine/Shaders/Private/Strata/StrataVisualize.usf
--- a/Engine/Shaders/Private/Strata/StrataVisualize.usf
+++ b/Engine/Shaders/Private/Strata/StrataVisualize.usf
@@ -17,10 +17,42 @@ void VisualizeMaterialPS(
 	float4 SVPos : SV_POSITION,
 	out float4 OutColor : SV_Target0)
 {
-	const int2 PixelPos = int2(SVPos.xy);
-	const int2 MaterialPos = View.CursorPosition;
+	const uint2 PixelPos = uint2(SVPos.xy);
+	const uint2 MaterialPos = View.CursorPosition;
 	const int  InspectSquare = 10;
 
+
+#if 0
+	OutColor = 0;
+	if (Strata.MaterialLobesTexture[PixelPos].x > 0.0f)
+	{
+		uint PixelStrataDataByteOffset = GetStrataPixelDataByteOffset(PixelPos, View.BufferSizeAndInvSize.x, Strata.MaxBytesPerPixel);
+		FStrataPixelBSDFs PixelBSDFs = UnpackStrataIn(Strata.MaterialLobesBuffer, PixelStrataDataByteOffset);
+		OutColor.w = 0.05f;
+		if (PixelBSDFs.BSDFCount == 1)
+		{
+			OutColor = float4(0, 0, 1, 0);
+		}
+		else if (PixelBSDFs.BSDFCount == 2)
+		{
+			OutColor = float4(0, 1, 0, 0);
+		}
+		else if (PixelBSDFs.BSDFCount == 3)
+		{
+			OutColor = float4(1, 1, 0, 0);
+		}
+		else if (PixelBSDFs.BSDFCount == 4)
+		{
+			OutColor = float4(1, 0, 0, 0);
+		}
+		else if (PixelBSDFs.BSDFCount > 4)
+		{
+			OutColor = float4(1, 1, 1, 0);
+		}
+	}
+	return;
+#endif
+
 	OutColor = float4(0.0f, 0.0f, 0.0f, 1.0f);
 
 	if (all(float2(PixelPos) > float2(MaterialPos + InspectSquare)))
@@ -30,7 +62,7 @@ void VisualizeMaterialPS(
 
 	if (Strata.MaterialLobesTexture[MaterialPos].x > 0.0f)
 	{
-		uint PixelStrataDataByteOffset = (MaterialPos.y * uint(View.BufferSizeAndInvSize.x) + MaterialPos.x) * Strata.MaxBytesPerPixel;
+		uint PixelStrataDataByteOffset = GetStrataPixelDataByteOffset(MaterialPos, View.BufferSizeAndInvSize.x, Strata.MaxBytesPerPixel);
 		FStrataPixelBSDFs PixelBSDFs = UnpackStrataIn(Strata.MaterialLobesBuffer, PixelStrataDataByteOffset);
 
 		float3 BSDFFontColor = float3(0.25, 1.0, 0.25);
@@ -45,6 +77,7 @@ void VisualizeMaterialPS(
 #define PRINTBSDFCHAR(X) PrintCharacter(PixelPos, OutColor.xyz, BSDFFontColor, TopLeft, X);
 
 #define PRINTBSDFWEIGHT()	PRINTBSDFCHAR(_SPC_); PRINTBSDFCHAR(_MINUS_); PRINTBSDFCHAR(_SPC_); PrintFloat(PixelPos, OutColor.xyz, FontColor, TopLeft, BSDF.Weight.x); TopLeft.x+=8*7; PRINTBSDFCHAR(_SPC_); PrintFloat(PixelPos, OutColor.xyz, FontColor, TopLeft, BSDF.Weight.y); TopLeft.x+=8*7; PRINTBSDFCHAR(_SPC_); PrintFloat(PixelPos, OutColor.xyz, FontColor, TopLeft, BSDF.Weight.z);
+#define PRINTBSDFCOUNT()	PRINTBSDFCHAR(min(_0_+PixelBSDFs.BSDFCount,_F_)); PRINTBSDFCHAR(_SPC_); PRINTBSDFCHAR(_B_); PRINTBSDFCHAR(_S_); PRINTBSDFCHAR(_D_); PRINTBSDFCHAR(_F_);
 #define PRINTFLOAT3(X) PrintFloat(PixelPos, OutColor.xyz, PARAMFontColor, TopLeft, X.x); TopLeft.x+=8*7; PRINTBSDFCHAR(_SPC_); PrintFloat(PixelPos, OutColor.xyz, PARAMFontColor, TopLeft, X.y); TopLeft.x+=8*7; PRINTBSDFCHAR(_SPC_); PrintFloat(PixelPos, OutColor.xyz, PARAMFontColor, TopLeft, X.z);
 #define PRINTFLOAT1(Value) PrintFloat(PixelPos, OutColor.xyz, PARAMFontColor, TopLeft, Value); TopLeft.x+=8*7;
 #define SIGNOFFLOAT1(Value) if(Value < 0.0f) { PrintCharacter(PixelPos, OutColor.xyz, PARAMFontColor, TopLeft, _MINUS_); } /*else { PrintCharacter(PixelPos, OutColor.xyz, PARAMFontColor, TopLeft, _SPC_);}*/
@@ -56,6 +89,9 @@ void VisualizeMaterialPS(
 		PRINTINFOCHAR(_T_); PRINTINFOCHAR(_O_); PRINTINFOCHAR(_P_);
 		NEXTLINE;
 		HORIZONTALSEPARATION;
+		PRINTBSDFCOUNT();
+		NEXTLINE;
+		HORIZONTALSEPARATION;
 
 		for (int i = 0; i < PixelBSDFs.BSDFCount; ++i)
 		{
