Fix crash on editor launch for certain project settings (UESP-4037).

diff --git a/Engine/Shaders/Private/GBufferHelpers.ush b/Engine/Shaders/Private/GBufferHelpers.ush
--- a/Engine/Shaders/Private/GBufferHelpers.ush
+++ b/Engine/Shaders/Private/GBufferHelpers.ush
@@ -378,7 +378,7 @@ void GBufferPostDecode(inout FGBufferData Ret, bool bChecker, bool bGetNormalize
 	Ret.StoredSpecular = Ret.Specular;
 
 #if GBUFFER_HAS_DIFFUSE_SAMPLE_OCCLUSION
-	Ret.DiffuseIndirectSampleOcclusion = 255 * Ret.GenericAO.x;
+	Ret.DiffuseIndirectSampleOcclusion = 255 * Ret.GenericAO;
 	Ret.GBufferAO = saturate(1.0 - float(countbits(Ret.DiffuseIndirectSampleOcclusion)) * rcp(float(INDIRECT_SAMPLE_COUNT)));
 	Ret.IndirectIrradiance = 1;
 #elif ALLOW_STATIC_LIGHTING
@@ -386,7 +386,7 @@ void GBufferPostDecode(inout FGBufferData Ret, bool bChecker, bool bGetNormalize
 	Ret.DiffuseIndirectSampleOcclusion = 0x0;
 	Ret.IndirectIrradiance = DecodeIndirectIrradiance(Ret.GenericAO.x);
 #else
-	Ret.GBufferAO = InGBufferC.a;
+	Ret.GBufferAO = Ret.GenericAO;
 	Ret.DiffuseIndirectSampleOcclusion = 0x0;
 	Ret.IndirectIrradiance = 1;
 #endif
diff --git a/Engine/Source/Runtime/Engine/Private/ShaderCompiler/ShaderGenerationUtil.cpp b/Engine/Source/Runtime/Engine/Private/ShaderCompiler/ShaderGenerationUtil.cpp
--- a/Engine/Source/Runtime/Engine/Private/ShaderCompiler/ShaderGenerationUtil.cpp
+++ b/Engine/Source/Runtime/Engine/Private/ShaderCompiler/ShaderGenerationUtil.cpp
@@ -1398,7 +1398,7 @@ static FString CreateGBufferDecodeFunctionVariation(const FGBufferInfo& BufferIn
 				}
 				else if (DecodeType == SceneTextures || DecodeType == SceneTexturesLoad)
 				{
-					TargetName = TEXT("VelocityBuffer");
+					TargetName = TEXT("GBufferVelocity");
 				}
 			}
 
