Fixing Reverb shader compile errors when INDIRECT_SAMPLE_COUNT isn't defined. Used value from DeferredShadingCommon.ush

diff --git a/Engine/Shaders/Private/GBufferHelpers.ush b/Engine/Shaders/Private/GBufferHelpers.ush
--- a/Engine/Shaders/Private/GBufferHelpers.ush
+++ b/Engine/Shaders/Private/GBufferHelpers.ush
@@ -366,6 +366,9 @@ void AdjustBaseColorAndSpecularColorForSubsurfaceProfileLightingCopyHack(inout f
 	}
 }
 
+#ifndef INDIRECT_SAMPLE_COUNT
+	#define INDIRECT_SAMPLE_COUNT 8
+#endif
 
 void GBufferPostDecode(inout FGBufferData Ret, bool bChecker, bool bGetNormalizedNormal)
 {
