GLES 3.1 fixes for new GBuffer layouts

diff --git a/Engine/Shaders/Private/GBufferHelpers.ush b/Engine/Shaders/Private/GBufferHelpers.ush
--- a/Engine/Shaders/Private/GBufferHelpers.ush
+++ b/Engine/Shaders/Private/GBufferHelpers.ush
@@ -24,7 +24,7 @@ float3 DecodeNormalHelper(float3 SrcNormal)
 
 uint EncodeQuantize6(float Value, float QuantizationBias)
 {
-	return min(uint(saturate(Value) * 63.0f + .5f + QuantizationBias),63);
+	return min(uint(saturate(Value) * 63.0f + .5f + QuantizationBias),63u);
 }
 
 float DecodeQuantize6(uint Value)
@@ -34,7 +34,7 @@ float DecodeQuantize6(uint Value)
 
 uint EncodeQuantize6Sqrt(float Value, float QuantizationBias)
 {
-	return min(uint(sqrt(saturate(Value)) * 63.0f + .5f + QuantizationBias),63);
+	return min(uint(sqrt(saturate(Value)) * 63.0f + .5f + QuantizationBias),63u);
 }
 
 float DecodeQuantize6Sqrt(uint Value)
@@ -44,7 +44,7 @@ float DecodeQuantize6Sqrt(uint Value)
 
 uint EncodeQuantize5(float Value, float QuantizationBias)
 {
-	return min(uint(saturate(Value) * 31.0f + .5f + QuantizationBias),31);
+	return min(uint(saturate(Value) * 31.0f + .5f + QuantizationBias),31u);
 }
 
 float DecodeQuantize5(uint Value)
@@ -54,7 +54,7 @@ float DecodeQuantize5(uint Value)
 
 uint EncodeQuantize5Sqrt(float Value, float QuantizationBias)
 {
-	return min(uint(sqrt(saturate(Value)) * 31.0f + .5f + QuantizationBias),31);
+	return min(uint(sqrt(saturate(Value)) * 31.0f + .5f + QuantizationBias),31u);
 }
 
 float DecodeQuantize5Sqrt(uint Value)
@@ -64,7 +64,7 @@ float DecodeQuantize5Sqrt(uint Value)
 
 uint EncodeQuantize4(float Value, float QuantizationBias)
 {
-	return min(uint(saturate(Value) * 15.0f + .5f + QuantizationBias),15);
+	return min(uint(saturate(Value) * 15.0f + .5f + QuantizationBias),15u);
 }
 
 float DecodeQuantize4(uint Value)
@@ -74,7 +74,7 @@ float DecodeQuantize4(uint Value)
 
 uint EncodeQuantize4Sqrt(float Value, float QuantizationBias)
 {
-	return min(uint(sqrt(saturate(Value)) * 15.0f + .5f + QuantizationBias),15);
+	return min(uint(sqrt(saturate(Value)) * 15.0f + .5f + QuantizationBias),15u);
 }
 
 float DecodeQuantize4Sqrt(uint Value)
@@ -85,7 +85,7 @@ float DecodeQuantize4Sqrt(uint Value)
 
 uint EncodeQuantize3(float Value, float QuantizationBias)
 {
-	return min(uint(saturate(Value) * 7.0f + .5f + QuantizationBias),7);
+	return min(uint(saturate(Value) * 7.0f + .5f + QuantizationBias),7u);
 }
 
 float DecodeQuantize3(uint Value)
@@ -95,7 +95,7 @@ float DecodeQuantize3(uint Value)
 
 uint EncodeQuantize3Sqrt(float Value, float QuantizationBias)
 {
-	return min(uint(sqrt(saturate(Value)) * 7.0f + .5f + QuantizationBias),7);
+	return min(uint(sqrt(saturate(Value)) * 7.0f + .5f + QuantizationBias),7u);
 }
 
 float DecodeQuantize3Sqrt(uint Value)
@@ -105,7 +105,7 @@ float DecodeQuantize3Sqrt(uint Value)
 
 uint EncodeQuantize2(float Value, float QuantizationBias)
 {
-	return min(uint(saturate(Value) * 3.0f + .5f + QuantizationBias),3);
+	return min(uint(saturate(Value) * 3.0f + .5f + QuantizationBias),3u);
 }
 
 float DecodeQuantize2(uint Value)
@@ -115,7 +115,7 @@ float DecodeQuantize2(uint Value)
 
 uint EncodeQuantize2Sqrt(float Value, float QuantizationBias)
 {
-	return min(uint(sqrt(saturate(Value)) * 3.0f + .5f + QuantizationBias),3);
+	return min(uint(sqrt(saturate(Value)) * 3.0f + .5f + QuantizationBias),3u);
 }
 
 float DecodeQuantize2Sqrt(uint Value)
@@ -125,7 +125,7 @@ float DecodeQuantize2Sqrt(uint Value)
 
 uint EncodeQuantize1(float Value, float QuantizationBias)
 {
-	return min(uint(saturate(Value) * 1.0f + .5f + QuantizationBias),1);
+	return min(uint(saturate(Value) * 1.0f + .5f + QuantizationBias),1u);
 }
 
 float DecodeQuantize1(uint Value)
@@ -135,7 +135,7 @@ float DecodeQuantize1(uint Value)
 
 uint EncodeQuantize1Sqrt(float Value, float QuantizationBias)
 {
-	return min(uint(sqrt(saturate(Value)) * 1.0f + .5f + QuantizationBias),1);
+	return min(uint(sqrt(saturate(Value)) * 1.0f + .5f + QuantizationBias),1u);
 }
 
 float DecodeQuantize1Sqrt(uint Value)
@@ -255,14 +255,15 @@ float3 DecodeQuantize332Sqrt(uint3 Value)
 }
 
 
-void EnvBRDFApproxFullyRoughHelper(inout half3 DiffuseColor, inout half3 SpecularColor)
+void EnvBRDFApproxFullyRoughHelper(inout float3 DiffuseColor, inout float3 SpecularColor)
 {
 	// Factors derived from EnvBRDFApprox( SpecularColor, 1, 1 ) == SpecularColor * 0.4524 - 0.0024
 	DiffuseColor += SpecularColor * 0.45;
 	SpecularColor = 0;
 	// We do not modify Roughness here as this is done differently at different places.
 }
-void EnvBRDFApproxFullyRoughHelper(inout half3 DiffuseColor, inout half SpecularColor)
+
+void EnvBRDFApproxFullyRoughHelper(inout float3 DiffuseColor, inout float SpecularColor)
 {
 	DiffuseColor += SpecularColor * 0.45;
 	SpecularColor = 0;
