Refactor material API for the path tracer to return Weight and Pdf from Eval and Sampl methods to promote term cancelation and remove special cases from singular BxDFs

Also include the NoL term in the BxDF itself in preparation for materials that don't have a surface measure (like hair or volumes)

Implement VisibleNormal sampling for DefaultLit and ClearCoat materials

Use single sample MIS to combine the probability distributions of multi-lobed materials

Simplify ClearCoat implementation so that the lobe selection can be performed independently of micronormal sampling, and simplify implementation to not refract into the bottom layer for now (still getting a reasonable visual match to Lit mode)

Add MaxPathIntensity CVar for firefly removal (off by default for now until we decide on a good default value)

Fix shadow tracing of non-opaque geometry (like glass). For now it seems more consistent to treat such objects as opaque.

Cleanup the path tracer's light setup and avoid adding a skydome to the light array if there isn't a valid one in the scene

Update LightMass and RTGI code to account for Throughput to Weight change in Material API

diff --git a/Engine/Plugins/Experimental/GPULightmass/Shaders/Private/LightmapPathTracing.usf b/Engine/Plugins/Experimental/GPULightmass/Shaders/Private/LightmapPathTracing.usf
--- a/Engine/Plugins/Experimental/GPULightmass/Shaders/Private/LightmapPathTracing.usf
+++ b/Engine/Plugins/Experimental/GPULightmass/Shaders/Private/LightmapPathTracing.usf
@@ -383,10 +383,10 @@ void PathTracingKernel(
 					EvalLight(LightId, LightUV, LightRay, Radiance);
 	
 					// Evaluate material
-					float Ior1 = 1.0;
-					float3 MaterialThroughput;
+					float3 MaterialEvalWeight = 0;
 					float MaterialEvalPdf = 0.0;
-					EvalMaterial(Ray.Direction, LightRay.Direction, RayHitInfo, Ior1, Bounces == 1, MaterialThroughput, MaterialEvalPdf);
+					EvalMaterial(Ray.Direction, LightRay.Direction, RayHitInfo, Bounces == 1, MaterialEvalWeight, MaterialEvalPdf);
+					float3 MaterialThroughput = MaterialEvalWeight * MaterialEvalPdf;
 	
 					// Apply material Pdf for correct MIS weight
 					float MisWeight = 1.0;
@@ -395,7 +395,7 @@ void PathTracingKernel(
 					#if 0
 						// DEBUG: Validate Material PDF
 						float MaterialPdfTmp;
-						PdfMaterial(Ray.Direction, LightRay, RayHitInfo, Ior1, Bounces == 1, MaterialPdfTmp);
+						PdfMaterial(Ray.Direction, LightRay, RayHitInfo, Bounces == 1, MaterialPdfTmp);
 						if (abs(MaterialPdf - MaterialPdfTmp) > 0.0001)
 						{
 							MaterialThroughput *= -1.f;
@@ -420,21 +420,21 @@ void PathTracingKernel(
 		}
 	
 		// Sample material
-		float Ior1 = 1.0;
 		float3 Direction;
-		float3 Throughput;
+		float3 Weight;
 		float4 RandSample = RandomSequence_GenerateSample4D(RandSequence, SampleIndex);
 		float SignedPositionBias;
 		float SamplePdf;
-		SampleMaterial(Ray.Direction, RayHitInfo, Ior1, RandSample, Bounces == 1, Direction, Throughput, SamplePdf, SignedPositionBias);
+		SampleMaterial(Ray.Direction, RayHitInfo, RandSample, Bounces == 1, Direction, Weight, SamplePdf, SignedPositionBias);
 		if (SamplePdf == 0.0)
 		{
 			break;
 		}
+		float3 Throughput = Weight * SamplePdf;
 	
-		PdfMaterial(Ray.Direction, Direction, Payload, Ior1, Bounces == 1, MaterialPdf);
+		PdfMaterial(Ray.Direction, Direction, Payload, Bounces == 1, MaterialPdf);
 		// Debug: Evaluation of a candidate sample direction should maintain energy balance with pathspace MIS
-		// EvalMaterial(Ray.Direction, Direction, RayHitInfo, Ior1, Bounces == 1, Throughput, MaterialPdf);
+		// EvalMaterial(Ray.Direction, Direction, RayHitInfo, Bounces == 1, Throughput, MaterialPdf);
 		// SamplePdf = MaterialPdf;
 
 		// Update ray
diff --git a/Engine/Shaders/Private/MonteCarlo.ush b/Engine/Shaders/Private/MonteCarlo.ush
--- a/Engine/Shaders/Private/MonteCarlo.ush
+++ b/Engine/Shaders/Private/MonteCarlo.ush
@@ -301,13 +301,27 @@ float VisibleGGXPDF(float3 V, float3 H, float a2)
 	return PDF;
 }
 
+float VisibleGGXPDF_aniso(float3 V, float3 H, float2 Alpha)
+{
+	float NoV = V.z;
+	float NoH = H.z;
+	float VoH = dot(V, H);
+	float CosThetaM4 = Pow2(NoH);
+	float TanThetaM2 = (Pow2(H.x / Alpha.x) + Pow2(H.y / Alpha.y)) / Pow2(NoH);
+	float D = 1 / (PI * Pow2(1 + TanThetaM2) * Alpha.x * Alpha.y * CosThetaM4);
+	float Lambda_i = 0.5 * (-1.0 + sqrt(1.0 + (Pow2(V.x * Alpha.x) + Pow2(V.y * Alpha.y)) / Pow2(V.z)));
+	float G1 = 1 / (Lambda_i + 1);
+	float Pdf = (G1 * D * VoH) / NoV;
+	return Pdf;
+}
+
 // [ Heitz 2018, "Sampling the GGX Distribution of Visible Normals" ]
 // http://jcgt.org/published/0007/04/01/
 // PDF = G_SmithV * VoH * D / NoV / (4 * VoH)
 // PDF = G_SmithV * D / (4 * NoV)
 float4 ImportanceSampleVisibleGGX( float2 DiskE, float a2, float3 V )
 {
-	// TODO float2 alpha for anisotropic
+	// NOTE: See below for anisotropic version that avoids this sqrt
 	float a = sqrt(a2);
 
 	// stretch
@@ -338,6 +352,31 @@ float4 ImportanceSampleVisibleGGX( float2 DiskE, float a2, float3 V )
 	return float4(H, VisibleGGXPDF(V, H, a2));
 }
 
+// [ Heitz 2018, "Sampling the GGX Distribution of Visible Normals" ]
+// http://jcgt.org/published/0007/04/01/
+// PDF = G_SmithV * VoH * D / NoV / (4 * VoH)
+// PDF = G_SmithV * D / (4 * NoV)
+float4 ImportanceSampleVisibleGGX_aniso(float2 DiskE, float2 Alpha, float3 V)
+{
+	// stretch
+	float3 Vh = normalize(float3(Alpha * V.xy, V.z));
+
+	// Stable tangent basis based on V
+	float3 Tx = (V.z < 0.9999) ? normalize(float3(-V.y, V.x, 0)) : float3(1, 0, 0);
+	float3 Ty = normalize(cross(Vh, Tx));
+
+	float2 p = DiskE;
+	float s = 0.5 + 0.5 * Vh.z;
+	p.y = lerp(sqrt(1 - p.x * p.x), p.y, s);
+
+	float3 H = p.x * Tx + p.y * Ty + sqrt(saturate(1 - dot(p, p))) * Vh;
+
+	// unstretch
+	H = normalize(float3(Alpha * H.xy, max(0.0, H.z)));
+
+	return float4(H, VisibleGGXPDF_aniso(V, H, Alpha));
+}
+
 // Multiple importance sampling power heuristic of two functions with a power of two. 
 // [Veach 1997, "Robust Monte Carlo Methods for Light Transport Simulation"]
 float MISWeight( uint Num, float PDF, uint OtherNum, float OtherPDF )
diff --git a/Engine/Shaders/Private/PathTracing/Light/PathTracingDirectionalLight.ush b/Engine/Shaders/Private/PathTracing/Light/PathTracingDirectionalLight.ush
--- a/Engine/Shaders/Private/PathTracing/Light/PathTracingDirectionalLight.ush
+++ b/Engine/Shaders/Private/PathTracing/Light/PathTracingDirectionalLight.ush
@@ -109,18 +109,14 @@ void DirectionalLight_EstimateLight(
 	int LightId,
 	float3 WorldPos,
 	float3 WorldNormal,
+	bool IsTransmissiveMaterial,
 	out float OutIrradiance
 )
 {
 	float3 LightDirection = GetNormal(LightId);
 	float3 SurfaceNormal = WorldNormal;
-	float NoL = 1.0;
-	if (GetRadius(LightId) == 0.0)
-	{
-		NoL = dot(SurfaceNormal, LightDirection);
-		NoL = ENABLE_TRANSMISSION && HasTransmission(LightId) ? abs(NoL) : saturate(NoL);
-	}
 
+	float NoL = 1.0; // trivial upper bound -- trying to be more accurate appears to reduce performance
 	float LightPower = Luminance(GetColor(LightId));
 	OutIrradiance = LightPower * NoL;
 }
diff --git a/Engine/Shaders/Private/PathTracing/Light/PathTracingLightSampling.ush b/Engine/Shaders/Private/PathTracing/Light/PathTracingLightSampling.ush
--- a/Engine/Shaders/Private/PathTracing/Light/PathTracingLightSampling.ush
+++ b/Engine/Shaders/Private/PathTracing/Light/PathTracingLightSampling.ush
@@ -119,6 +119,7 @@ void EstimateLight(
 	float3 WorldPos,
 	float3 WorldNormal,
 	uint PrimitiveLightingChannelMask,
+	bool IsTransmissiveMaterial,
 	inout float OutIrradiance
 )
 {
@@ -127,23 +128,23 @@ void EstimateLight(
 
 	if (IsPointLight(LightId))
 	{
-		PointLight_EstimateLight(LightId, WorldPos, WorldNormal, OutIrradiance);
+		PointLight_EstimateLight(LightId, WorldPos, WorldNormal, IsTransmissiveMaterial, OutIrradiance);
 	}
 	else if (IsSpotLight(LightId))
 	{
-		SpotLight_EstimateLight(LightId, WorldPos, WorldNormal, OutIrradiance);
+		SpotLight_EstimateLight(LightId, WorldPos, WorldNormal, IsTransmissiveMaterial, OutIrradiance);
 	}
 	else if (IsEnvironmentLight(LightId))
 	{
-		SkyLight_EstimateLight(LightId, WorldPos, WorldNormal, OutIrradiance);
+		SkyLight_EstimateLight(LightId, WorldPos, WorldNormal, IsTransmissiveMaterial, OutIrradiance);
 	}
 	else if (IsDirectionalLight(LightId))
 	{
-		DirectionalLight_EstimateLight(LightId, WorldPos, WorldNormal, OutIrradiance);
+		DirectionalLight_EstimateLight(LightId, WorldPos, WorldNormal, IsTransmissiveMaterial, OutIrradiance);
 	}
 	else if (IsRectLight(LightId))
 	{
-		RectLight_EstimateLight(LightId, WorldPos, WorldNormal, OutIrradiance);
+		RectLight_EstimateLight(LightId, WorldPos, WorldNormal, IsTransmissiveMaterial, OutIrradiance);
 	}
 }
 
@@ -174,12 +175,13 @@ bool SampleLightSelection(
 	float3 WorldPos = Payload.WorldPos;
 	float3 WorldNormal = Payload.WorldNormal;
 	uint PrimitiveLightingChannelMask = Payload.PrimitiveLightingChannelMask;
-
-	EstimateLight(0, WorldPos, WorldNormal, PrimitiveLightingChannelMask, IrradianceEstimate[0]);
+	
+	// TODO: query transmission from material instead of light?
+	EstimateLight(0, WorldPos, WorldNormal, PrimitiveLightingChannelMask, ENABLE_TRANSMISSION && HasTransmission(0), IrradianceEstimate[0]);
 	{
 		for (uint LightIndex = 1; LightIndex < SceneLightsData.Count; ++LightIndex)
 		{
-			EstimateLight(LightIndex, WorldPos, WorldNormal, PrimitiveLightingChannelMask, IrradianceEstimate[LightIndex]);
+			EstimateLight(LightIndex, WorldPos, WorldNormal, PrimitiveLightingChannelMask, ENABLE_TRANSMISSION && HasTransmission(LightIndex), IrradianceEstimate[LightIndex]);
 			IrradianceEstimate[LightIndex] += IrradianceEstimate[LightIndex - 1];
 		}
 	}
@@ -237,11 +239,11 @@ void PdfLightSelection(
 	float3 WorldNormal = Payload.WorldNormal;
 	uint PrimitiveLightingChannelMask = Payload.PrimitiveLightingChannelMask;
 
-	EstimateLight(0, WorldPos, WorldNormal, PrimitiveLightingChannelMask, IrradianceEstimate[0]);
+	EstimateLight(0, WorldPos, WorldNormal, PrimitiveLightingChannelMask, ENABLE_TRANSMISSION && HasTransmission(0), IrradianceEstimate[0]);
 	{
 		for (uint LightIndex = 1; LightIndex < SceneLightsData.Count; ++LightIndex)
 		{
-			EstimateLight(LightIndex, WorldPos, WorldNormal, PrimitiveLightingChannelMask, IrradianceEstimate[LightIndex]);
+			EstimateLight(LightIndex, WorldPos, WorldNormal, PrimitiveLightingChannelMask, ENABLE_TRANSMISSION && HasTransmission(LightIndex), IrradianceEstimate[LightIndex]);
 			IrradianceEstimate[LightIndex] += IrradianceEstimate[LightIndex - 1];
 		}
 	}
@@ -390,6 +392,7 @@ bool InitLightPickingCdf(
 	float3 WorldPos,
 	float3 WorldNormal,
 	uint PrimitiveLightingChannelMask,
+	bool IsTransmissiveMaterial,
 	inout float LightPickingCdf[RAY_TRACING_LIGHT_COUNT_MAXIMUM])
 {
 	// This code assumes SceneLightsData.Count > 0
@@ -397,11 +400,11 @@ bool InitLightPickingCdf(
 	InitLightPickingCdf(LightPickingCdf);
 #else
 	// Build irradiance estimate prefix sum
-	EstimateLight(0, WorldPos, WorldNormal, PrimitiveLightingChannelMask, LightPickingCdf[0]);
+	EstimateLight(0, WorldPos, WorldNormal, PrimitiveLightingChannelMask, IsTransmissiveMaterial, LightPickingCdf[0]);
 	{
 		for (uint LightIndex = 1; LightIndex < SceneLightsData.Count; ++LightIndex)
 		{
-			EstimateLight(LightIndex, WorldPos, WorldNormal, PrimitiveLightingChannelMask, LightPickingCdf[LightIndex]);
+			EstimateLight(LightIndex, WorldPos, WorldNormal, PrimitiveLightingChannelMask, IsTransmissiveMaterial, LightPickingCdf[LightIndex]);
 			LightPickingCdf[LightIndex] += LightPickingCdf[LightIndex - 1];
 		}
 	}
diff --git a/Engine/Shaders/Private/PathTracing/Light/PathTracingPointLight.ush b/Engine/Shaders/Private/PathTracing/Light/PathTracingPointLight.ush
--- a/Engine/Shaders/Private/PathTracing/Light/PathTracingPointLight.ush
+++ b/Engine/Shaders/Private/PathTracing/Light/PathTracingPointLight.ush
@@ -146,23 +146,16 @@ void PointLight_EstimateLight(
 	int LightId,
 	float3 WorldPos,
 	float3 WorldNormal,
+	bool IsTransmissiveMaterial,
 	out float OutIrradiance
 )
 {
 	// Distance
 	float3 LightDirection = GetPosition(LightId) - WorldPos;
 	float LightDistanceSquared = dot(LightDirection, LightDirection);
-	LightDirection = normalize(LightDirection);
 
 	// Geometric term
-	float3 SurfaceNormal = WorldNormal;
-
-	float NoL = 1.0;
-	if (GetRadius(LightId) == 0.0)
-	{
-		float NoL = dot(SurfaceNormal, LightDirection);
-		NoL = ENABLE_TRANSMISSION && HasTransmission(LightId) ? abs(NoL) : saturate(NoL);
-	}
+	float NoL = 1.0; // trivial upper bound -- trying to be more accurate appears to reduce performance
 
 	float LightPower = Luminance(GetColor(LightId));
 	float Falloff = ComputeAttenuationFalloff(sqrt(LightDistanceSquared), GetAttenuation(LightId));
diff --git a/Engine/Shaders/Private/PathTracing/Light/PathTracingRectLight.ush b/Engine/Shaders/Private/PathTracing/Light/PathTracingRectLight.ush
--- a/Engine/Shaders/Private/PathTracing/Light/PathTracingRectLight.ush
+++ b/Engine/Shaders/Private/PathTracing/Light/PathTracingRectLight.ush
@@ -128,6 +128,7 @@ void RectLight_EstimateLight(
 	int LightId,
 	float3 WorldPos,
 	float3 WorldNormal,
+	bool IsTransmissiveMaterial,
 	out float OutIrradiance
 )
 {
@@ -135,27 +136,21 @@ void RectLight_EstimateLight(
 	// #dxr_todo: UE-72533 Use closest point, instead
 	float3 LightDirection = GetPosition(LightId) - WorldPos;
 	float LightDistanceSquared = dot(LightDirection, LightDirection);
-	LightDirection = normalize(LightDirection);
 	float3 LightNormal = GetNormal(LightId);
 
-	// Light-normal culling
+	// Is the shading point behind the light?
 	if (dot(-LightDirection, LightNormal) <= 0.0)
 	{
 		OutIrradiance = 0.0;
 		return;
 	}
-
 	// Approximate geometric term
-	float3 SurfaceNormal = WorldNormal;
-#if 0
-	float NoL = dot(SurfaceNormal, LightDirection);
-	NoL = ENABLE_TRANSMISSION && HasTransmission(LightId) ? abs(NoL) : saturate(NoL);
-#else
-	// Just use a trivial upper bound on cosine term 
-	// #dxr_todo: UE-72533 least intrusive way of fixing the artifact cause by using LightCenter
+	float Width = GetWidth(LightId);
+	float Height = GetHeight(LightId);
+
+	// Don't bother trying to bound the N.L term as its in [0,1] and hard to estimate accurately and quickly
 	float NoL = 1.0;
-#endif
-	float Area = GetWidth(LightId) * GetHeight(LightId);
+	float Area = Width * Height;
 
 	float LightPower = Luminance(GetColor(LightId));
 	float Falloff = ComputeAttenuationFalloff(sqrt(LightDistanceSquared), GetAttenuation(LightId));
diff --git a/Engine/Shaders/Private/PathTracing/Light/PathTracingSkyLight.ush b/Engine/Shaders/Private/PathTracing/Light/PathTracingSkyLight.ush
--- a/Engine/Shaders/Private/PathTracing/Light/PathTracingSkyLight.ush
+++ b/Engine/Shaders/Private/PathTracing/Light/PathTracingSkyLight.ush
@@ -68,6 +68,7 @@ void SkyLight_EstimateLight(
 	int LightId,
 	float3 WorldPos,
 	float3 WorldNormal,
+	bool IsTransmissiveMaterial,
 	out float OutIrradiance
 )
 {
diff --git a/Engine/Shaders/Private/PathTracing/Light/PathTracingSpotLight.ush b/Engine/Shaders/Private/PathTracing/Light/PathTracingSpotLight.ush
--- a/Engine/Shaders/Private/PathTracing/Light/PathTracingSpotLight.ush
+++ b/Engine/Shaders/Private/PathTracing/Light/PathTracingSpotLight.ush
@@ -180,6 +180,7 @@ void SpotLight_EstimateLight(
 	int LightId,
 	float3 WorldPos,
 	float3 WorldNormal,
+	bool IsTransmissiveMaterial,
 	out float OutIrradiance
 )
 {
@@ -191,23 +192,15 @@ void SpotLight_EstimateLight(
 	// Distance
 	float3 LightDirection = GetPosition(LightId) - WorldPos;
 	float LightDistanceSquared = dot(LightDirection, LightDirection);
-	LightDirection = normalize(LightDirection);
 
 	// Geometric term
-	float3 SurfaceNormal = WorldNormal;
-
-	float NoL = 1.0;
-	if (GetRadius(LightId) == 0.0)
-	{
-		float NoL = dot(SurfaceNormal, LightDirection);
-		NoL = ENABLE_TRANSMISSION && HasTransmission(LightId) ? abs(NoL) : saturate(NoL);
-	}
+	float NoL = 1.0; // trivial upper bound -- trying to be more accurate appears to reduce performance
 
 	float Falloff = ComputeAttenuationFalloff(sqrt(LightDistanceSquared), GetAttenuation(LightId));
 	OutIrradiance = LightPower * Falloff * NoL / LightDistanceSquared;
 	
 	float2 CosConeAngles = GetCosConeAngles(LightId);
-	OutIrradiance *= SpotAttenuation(LightDirection, LightNormal, CosConeAngles);
+	OutIrradiance *= SpotAttenuation(normalize(LightDirection), LightNormal, CosConeAngles);
 }
 	
 void SpotLight_SampleLight(
diff --git a/Engine/Shaders/Private/PathTracing/Material/PathTracingClearCoat.ush b/Engine/Shaders/Private/PathTracing/Material/PathTracingClearCoat.ush
--- a/Engine/Shaders/Private/PathTracing/Material/PathTracingClearCoat.ush
+++ b/Engine/Shaders/Private/PathTracing/Material/PathTracingClearCoat.ush
@@ -8,7 +8,10 @@
 
 #define EVAL_TOP_LAYER 1
 #define EVAL_BOTTOM_LAYER 1
-#define EVAL_REFRACTION 1
+
+#define GREY_INNER_LAYER 0
+
+#define CLEAR_COAT_F0	0.04 // From hard-coded value in ShadingModels.ush
 
 #include "PathTracingFresnel.ush"
 
@@ -33,160 +36,110 @@ float3 CalcThinTransmission(float NoL, float NoV, FMaterialClosestHitPayload Pay
 	return Transmission;
 }
 
-float CalcClearCoatFresnel(float VoH, FMaterialClosestHitPayload HitInfo)
-{
-	float Ior1 = 1.0;
-	float Ior2 = GetClearCoatIor(HitInfo);
-	return FresnelReflectance(VoH, Ior1, Ior2);
-}
-
-float CalcClearCoatSelectionPdf(float3 V, float3 H, FMaterialClosestHitPayload Payload)
+float CalcClearCoatSelectionPdf(FMaterialClosestHitPayload Payload)
 {
-	// Lobe selection is determined by Fresnel term
-	float VoH = saturate(dot(V, H));
-	float Pdf = CalcClearCoatFresnel(VoH, Payload);
-
 #if EVAL_TOP_LAYER == 0
-	Pdf = 0.0;
+	// never select coating if top is disabled
+	return 0.0; 
 #elif EVAL_BOTTOM_LAYER == 0
-	Pdf = 1.0;
+	// always select coating if bottom is disabled
+	return 1.0; 
 #endif
+	// probability of choosing the coat
+	// TODO: this could be improved if we had a directional albedo tabulation
+	float Pdf = GetClearCoat(Payload) * CLEAR_COAT_F0;
 	return Pdf;
 }
 
-float ClearCoat_PdfMaterial(
-	float3 IncomingDirection,
-	float3 OutgoingDirection,
-	FMaterialClosestHitPayload Payload
-)
-{
-	// Compute GGX PDF
-	float3 V = -IncomingDirection;
-	float3 L = OutgoingDirection;
-	float3 H = normalize(V + L);
-	float3 N = GetWorldNormal(Payload);
-
-	float NoH = saturate(dot(N, H));
-	float HoV = saturate(dot(H, V));
-
-	// Compute top-layer pdf
-	float a2 = Pow4(GetClearCoatRoughness(Payload));
-	float D = D_GGX(a2, NoH);
-	float MicrofacetPdf = 0.0;
-	if (HoV > 0.0)
-	{
-		MicrofacetPdf = D * NoH / (4.0 * HoV);
-	}
-
-	// Refract V and L
-	float Eta = 1.0 / GetClearCoatIor(Payload);
-	#if EVAL_REFRACTION == 1
-	float3 V2;
-	float3 L2;
-	bool IsValidV2 = BSDF_Refract(V, H, Eta, V2);
-	bool IsValidL2 = BSDF_Refract(L, H, Eta, L2);
-	#else
-		float3 V2 = -V;
-		float3 L2 = -L;
-		bool IsValidV2 = true;
-		bool IsValidL2 = true;
-	#endif
-
-	// Compute bottom-layer pdf
-	float BottomLayerPdf = 0.0;
-	if (IsValidV2 && IsValidL2)
-	{
-		Payload.WorldNormal = GetClearCoatBottomNormal(Payload);
-		BottomLayerPdf = DefaultLit_PdfMaterial(V2, -L2, Payload, GetClearCoatIor(Payload));
-		Payload.WorldNormal = N;
-	}
-	
-	// Return piecewise pdf
-	float ClearCoatSelectionPdf = CalcClearCoatSelectionPdf(V, H, Payload);
-	float Pdf = ClearCoatSelectionPdf * MicrofacetPdf + (1.0 - ClearCoatSelectionPdf) * BottomLayerPdf;
-	return Pdf;
+float ClearCoatFresnel(float VoH) {
+	// from hardcoded logic in ClearCoatBxDF
+	float Fc = Pow5(1 - VoH);
+	return Fc + (1 - Fc) * CLEAR_COAT_F0;
 }
 
 void ClearCoat_EvalMaterial(
 	float3 IncomingDirection,
 	float3 OutgoingDirection,
 	inout FMaterialClosestHitPayload Payload,
-	out float3 OutThroughput,
+	out float3 OutWeight,
 	out float OutPdf
 )
 {
-	float3 V = -IncomingDirection;
-	float3 L = OutgoingDirection;
-	float3 H = normalize(V + L);
-	float3 N = GetWorldNormal(Payload);
+	OutWeight = 0;
+	OutPdf = 0;
 
-	float NoL = saturate(dot(N, L));
-	float NoV = saturate(dot(N, V));
-	float NoH = saturate(dot(N, H));
-	float VoH = saturate(dot(V, H));
+	float3 V_World = -IncomingDirection;
+	float3 L_World = OutgoingDirection;
+	float3 N_World = GetWorldNormal(Payload);
+
+	float CoatSelectionPdf = CalcClearCoatSelectionPdf(Payload);
+
+	// No anisotropy, just pick an arbitrary basis
+	float Roughness = max(GetClearCoatRoughness(Payload), GGX_MIN_ROUGHNESS);
+	float2 Alpha = Roughness * Roughness;
+	float3x3 Basis = GetTangentBasis(N_World);
 
-	// top-layer clear coat evaluation
-	OutThroughput = 0.0;
+	// move vectors into right shading frame
+	float3 V = mul(Basis, V_World);
+	float3 L = mul(Basis, L_World);
+	float3 H = normalize(V + L);
 
-#if EVAL_TOP_LAYER == 1
-	float CC_a2 = Pow4(GetClearCoatRoughness(Payload));
-	float CC_F = CalcClearCoatFresnel(VoH, Payload);
-	float CC_D = D_GGX(CC_a2, NoH);
-	float CC_G = Vis_SmithJointApprox(CC_a2, NoV, NoL) * (4.0 * NoV * NoL);
+	float NoL = saturate(L.z);
+	float NoV = saturate(V.z);
+	float NoH = saturate(H.z);
+	float VoH = saturate(dot(V, H));
 
-	float3 TopLayerThroughput = 0.0;
-	if (NoV * NoL > 0.0)
+	// Specular lobe
+	if (NoL > 0 && NoV > 0)
 	{
-		float ClearCoat = GetClearCoat(Payload);
-		// dxr_todo: NoL cancels
-		TopLayerThroughput = ClearCoat * ((CC_F * CC_D * CC_G) / (4.0 * NoV * NoL)) * NoL;
+		// Assume NoH > 0
+		// FIXME: this can be further simplified
+		const float CosThetaM4 = Pow2(NoH);
+		const float TanThetaM2 = (Pow2(H.x / Alpha.x) + Pow2(H.y / Alpha.y)) / Pow2(NoH);
+		const float D = 1 / (PI * Pow2(1 + TanThetaM2) * Alpha.x * Alpha.y * CosThetaM4);
+		const float Lambda_o = 0.5 * (-1.0 + sqrt(1.0 + (Pow2(L.x * Alpha.x) + Pow2(L.y * Alpha.y)) / Pow2(L.z)));
+		const float Lambda_i = 0.5 * (-1.0 + sqrt(1.0 + (Pow2(V.x * Alpha.x) + Pow2(V.y * Alpha.y)) / Pow2(V.z)));
+		const float G2 = 1 / (Lambda_i + Lambda_o + 1);
+		const float G1 = 1 / (Lambda_i + 1);
+
+		const float F = GetClearCoat(Payload) * ClearCoatFresnel(VoH);
+
+		const float3 SpecWeight = F * (G2 / G1);
+		const float SpecPdf = (G1 * D * 0.25f) / NoV;
+
+		AddLobeWithMIS(OutWeight, OutPdf, SpecWeight, SpecPdf, CoatSelectionPdf);
 	}
 
-	OutThroughput += TopLayerThroughput;
-#endif
-	
 	// bottom-layer clear coat evaluation
-#if EVAL_BOTTOM_LAYER == 1
-	float Eta = 1.0 / GetClearCoatIor(Payload);
-	#if EVAL_REFRACTION == 1
-	float3 V2;
-	float3 L2;
-	bool IsValidV2 = BSDF_Refract(V, H, Eta, V2);
-	bool IsValidL2 = BSDF_Refract(L, H, Eta, L2);
-	#else
-		float3 V2 = -V;
-		float3 L2 = -L;
-		bool IsValidV2 = true;
-		bool IsValidL2 = true;
-	#endif
-	if (IsValidV2 && IsValidL2)
-	{
-		V2 = -V2;
-		L2 = -L2;
-		float3 N2 = GetClearCoatBottomNormal(Payload);
-		float N2oL2 = saturate(dot(N2, L2));
-		float N2oV2 = saturate(dot(N2, V2));
-
-		// Evaluate bottom coat
-		Payload.WorldNormal = N2;
-		float3 BottomLayerThroughput;
-		float BottomLayerPdf;
-		DefaultLit_EvalMaterial(-V2, L2, Payload, GetClearCoatIor(Payload), BottomLayerThroughput, BottomLayerPdf);
-		Payload.WorldNormal = N;
-
-		// Calculate transmission through substrate
-		float3 Transmission = CalcThinTransmission(N2oL2, N2oV2, Payload);
-
-		// Apply fresnel through exit interface
-		float CC_F1 = CalcClearCoatFresnel(VoH, Payload);
-		float HoL = saturate(dot(H, L));
-		float CC_F2 = CalcClearCoatFresnel(HoL, Payload);
-
-		OutThroughput += (1.0 - CC_F1) * (1.0 - CC_F2) * Transmission * BottomLayerThroughput;
-	}
+	float3 BottomWeight = 0;
+	float BottomPdf = 0;
+	Payload.WorldNormal = GetClearCoatBottomNormal(Payload);
+#if GREY_INNER_LAYER == 0
+	// NOTE: Should we try to account for the IOR ratio between layers? We could try to darken Specular Color by the coating IOR
+	// but this would apply to both the dielctric and metal portions equally ...
+	// It seems better to punt until we have a Strata going
+	DefaultLit_EvalMaterial(IncomingDirection, OutgoingDirection, Payload, BottomWeight, BottomPdf);
+#else
+	float3 OldColor = Payload.BaseColor;
+	Payload.BaseColor = 0.18;
+	Lambert_EvalMaterial(OutgoingDirection, Payload, BottomWeight, BottomPdf);
+	Payload.BaseColor = OldColor;
 #endif
+	Payload.WorldNormal = N_World;
+
+	// Calculate transmission through substrate
+	float3 Transmission = CalcThinTransmission(NoL, NoV, Payload);
+
+	// Apply fresnel through exit interface
+	// TODO: This should actually use a measure of the directional albedo along V and L
+	// TODO: The original implementation tried to refract V and L, but doing this correctly would require a stochastic eval method
+	// NOTE: implementation in ClearCoatBxDF uses VoH for these which again makes the coating math dependent on the (stochastic) choice of H
+	float CC_F1 = ClearCoatFresnel(NoV);
+	float CC_F2 = ClearCoatFresnel(NoL);
+
+	BottomWeight *= lerp(1.0, Transmission * (1 - CC_F1) * (1 - CC_F2), GetClearCoat(Payload));
 
-	OutPdf = ClearCoat_PdfMaterial(IncomingDirection, OutgoingDirection, Payload);
+	AddLobeWithMIS(OutWeight, OutPdf, BottomWeight, BottomPdf, 1.0 - CoatSelectionPdf);
 }
 
 void ClearCoat_SampleMaterial(
@@ -194,108 +147,60 @@ void ClearCoat_SampleMaterial(
 	FMaterialClosestHitPayload Payload,
 	float4 RandSample,
 	out float3 OutDirection,
-	out float3 OutThroughput,
+	out float3 OutWeight,
 	out float OutPdf
 )
 {
-	float3 N = GetWorldNormal(Payload);
-	float3 V = -RayDirection;
+	float3 N_World = GetWorldNormal(Payload);
+	float3 V_World = -RayDirection;
+
 	OutDirection = 0.0;
-	OutThroughput = 0.0;
+	OutWeight = 0.0;
 	OutPdf = 0.0;
 
-	// Sample microfacet
-	float a2 = Pow4(GetClearCoatRoughness(Payload));
-	float4 Result = ImportanceSampleGGX(RandSample.zw, a2);
-	float3 H = TangentToWorld(Result.xyz, N);
-	float VoH = saturate(dot(V, H));
+	float CoatSelectionPdf = CalcClearCoatSelectionPdf(Payload);
 
-	float ClearCoatSelectionPdf = CalcClearCoatSelectionPdf(V, H, Payload);
-	if (RandSample.x < ClearCoatSelectionPdf)
+	// choose between coat and bottom layer
+	if (RandSample.x < CoatSelectionPdf)
 	{
-		float NoH = saturate(dot(N, H));
-		float D = D_GGX(a2, NoH);
-		float MicrofacetPdf = 0.0;
-		if (VoH > 0.0)
-		{
-			MicrofacetPdf = D * NoH / (4.0 * VoH);
-
-			// Lighting direction is the reflection vector about the sampled microfacet
-			float3 R = 2.0 * dot(V, H) * H - V;
-			OutDirection = R;
-			OutPdf = ClearCoatSelectionPdf * MicrofacetPdf;
-			{
-				// Calculate throughput
-				float3 L = OutDirection;
-				float NoL = saturate(dot(N, L));
-				float NoV = saturate(dot(N, V));
-
-				// top-layer clear coat evaluation
-				float CC_a2 = Pow4(GetClearCoatRoughness(Payload));
-				float CC_F = CalcClearCoatFresnel(VoH, Payload);
-				float CC_D = D_GGX(CC_a2, NoH);
-				float CC_G = Vis_SmithJointApprox(CC_a2, NoV, NoL) * (4.0 * NoV * NoL);
-
-				if (NoV * NoL > 0.0)
-				{
-					float ClearCoat = GetClearCoat(Payload);
-					// dxr_todo: NoL terms cancel
-					OutThroughput = ClearCoat * ((CC_F * CC_D * CC_G) / (4.0 * NoV * NoL)) * NoL;
-				}
-			}
+		RandSample.x /= CoatSelectionPdf;
+
+		// No anisotropy, just pick an arbitrary basis
+		float Roughness = max(GetClearCoatRoughness(Payload), GGX_MIN_ROUGHNESS);
+		float2 Alpha = Roughness * Roughness;
+		float3x3 Basis = GetTangentBasis(N_World);
+
+		float3 V = -mul(Basis, RayDirection);
+
+		// #dxr_todo: Switch to UniformSampleDiskConcentric after fixing 0/0 bug
+		float3 H = ImportanceSampleVisibleGGX_aniso(UniformSampleDisk(RandSample.xy), Alpha, V).xyz;
+
+		float3 L = reflect(-V, H);
+		if (L.z <= 0) {
+			// invalid output direction, skip some work
+			return;
 		}
+		// reflect and transform
+		OutDirection = mul(L, Basis);
 	}
 	else
 	{
-		RandSample.x -= ClearCoatSelectionPdf;
-		RandSample.x /= (1.0 - ClearCoatSelectionPdf);
-
-		// Refract V through clear coat
-		float Eta = 1.0 / GetClearCoatIor(Payload);
-		#if EVAL_REFRACTION == 1
-		float3 V2;
-		bool IsValidV2 = BSDF_Refract(V, N, Eta, V2);
-		V2 = -V2;
-		#else
-			bool IsValidV2 = true;
-			float3 V2 = V;
-		#endif
-
-		if (IsValidV2)
-		{
-			// Sample DefaultLit
-			float3 N2 = GetClearCoatBottomNormal(Payload);
-			Payload.WorldNormal = N2;
-
-			float3 L2;
-			float3 BottomLayerThroughput;
-			float BottomLayerPdf;
-			DefaultLit_SampleMaterial(-V2, Payload, GetClearCoatIor(Payload), RandSample, L2, BottomLayerThroughput, BottomLayerPdf);
-			Payload.WorldNormal = N;
-
-			// Refract outgoing ray
-			#if EVAL_REFRACTION == 1
-			bool IsValidL2 = BSDF_Refract(-L2, -N, 1.0 / Eta, OutDirection);
-			#else
-				OutDirection = L2;
-				bool IsValidL2 = true;
-			#endif
-			if (BottomLayerPdf > 0.0 && IsValidL2)
-			{
-				// Calculate transmission through substrate
-				float N2oL2 = saturate(dot(N2, L2));
-				float N2oV2 = saturate(dot(N2, V2));
-				float3 Transmission = CalcThinTransmission(N2oL2, N2oV2, Payload);
-
-				// Apply fresnel through exit interface
-				float NoV = saturate(dot(N, V));
-				float CC_F1 = CalcClearCoatFresnel(NoV, Payload);
-				float NoL = saturate(dot(N, OutDirection));
-				float CC_F2 = CalcClearCoatFresnel(NoL, Payload);
-
-				OutThroughput = (1.0 - CC_F1) * (1.0 - CC_F2) * Transmission * BottomLayerThroughput;
-				OutPdf = (1.0 - ClearCoatSelectionPdf) * BottomLayerPdf;
-			}
-		}
+		RandSample.x -= CoatSelectionPdf;
+		RandSample.x /= (1.0 - CoatSelectionPdf);
+
+#if GREY_INNER_LAYER == 0
+		float3 Weight = 0;
+		float Pdf = 0;
+		Payload.WorldNormal = GetClearCoatBottomNormal(Payload);
+		DefaultLit_SampleMaterial(RayDirection, Payload, RandSample, OutDirection, Weight, Pdf);
+		Payload.WorldNormal = N_World;
+#else
+		float3 Nb = GetClearCoatBottomNormal(Payload);
+		OutDirection = TangentToWorld(CosineSampleHemisphere(RandSample.xy).xyz, Nb);
+#endif
 	}
+
+	// defer to eval for now to get proper MIS mix of both events
+	// TODO: could use temporary results above to inline Eval logic here
+	ClearCoat_EvalMaterial(RayDirection, OutDirection, Payload, OutWeight, OutPdf);
 }
diff --git a/Engine/Shaders/Private/PathTracing/Material/PathTracingDefaultLit.ush b/Engine/Shaders/Private/PathTracing/Material/PathTracingDefaultLit.ush
--- a/Engine/Shaders/Private/PathTracing/Material/PathTracingDefaultLit.ush
+++ b/Engine/Shaders/Private/PathTracing/Material/PathTracingDefaultLit.ush
@@ -8,71 +8,100 @@
 
 #include "PathTracingFresnel.ush"
 
+// be conservative since the GPU may be aggressive with floating point optimizations
+#define GGX_MIN_ROUGHNESS 0.01
+
 float CalcLobeSelectionPdf(float3 V, FMaterialClosestHitPayload Payload)
 {
-	// #dxr_todo: Can do much better than this - but this at least prevents under-sampling when one of the lobes is black
+	// TODO: This could be improved if we had a directional albedo tabulation, but this formulation will at least prevent
+	// under-sampling when one lobe is black (like with metals)
+
 	float3 Diff = GetDiffuseColor(Payload);
 	float3 Spec = GetSpecularColor(Payload);
 
 	float D = Diff.x + Diff.y + Diff.z;
 	float S = Spec.x + Spec.y + Spec.z;
 	return D / (D + S + 1e-6);
-
-
-	float Pdf = min(GetRoughness(Payload), 0.9);
-	return Pdf;
 }
 
-float DefaultLit_PdfMaterial(
+void DefaultLit_EvalMaterial(
 	float3 IncomingDirection,
 	float3 OutgoingDirection,
 	FMaterialClosestHitPayload Payload,
-	float Ior1
+	out float3 OutWeight,
+	out float OutPdf
 )
 {
+	OutWeight = 0;
+	OutPdf = 0;
+
 	float3 V_World = -IncomingDirection;
 	float3 L_World = OutgoingDirection;
 	float3 N_World = GetWorldNormal(Payload);
-	float3 V = WorldToTangent(V_World, N_World);
-	float3 L = WorldToTangent(L_World, N_World);
-	float3 N = float3(0, 0, 1);
+
+	float Roughness = max(GetRoughness(Payload), GGX_MIN_ROUGHNESS);
+	float2 Alpha;
+	float3x3 Basis;
+	if (Payload.HasAnisotropy())
+	{
+		// Material is anisotropic, our shading frame should use the tangent vector
+		GetAnisotropicRoughness(Roughness * Roughness, Payload.Anisotropy, Alpha.x, Alpha.y);
+		float3 X = Payload.WorldTangent;
+		float3 Y = normalize(cross(N_World, X));
+		Basis = float3x3(cross(Y, N_World), Y, N_World);
+	}
+	else
+	{
+		// No anisotropy, just pick an arbitrary basis
+		Alpha = Roughness * Roughness;
+		Basis = GetTangentBasis(N_World);
+	}
+
+	// move vectors into right shading frame
+	float3 V = mul(Basis, V_World);
+	float3 L = mul(Basis, L_World);
 	float3 H = normalize(V + L);
 
-	float NoL = saturate(dot(N, L));
-	float NoH = saturate(dot(N, H));
-	float HoV = saturate(dot(H, V));
+	float NoL = saturate(L.z);
+	float NoV = saturate(V.z);
+	float NoH = saturate(H.z);
+	float VoH = saturate(dot(V, H));
 
+	// Probability of picking diffuse lobe vs. specular lobe
 	float LobeSelectionPdf = CalcLobeSelectionPdf(V_World, Payload);
-	float Roughness = GetRoughness(Payload);
-	float LambertPdf = NoL / PI;
 
-	float a2 = Pow4(Roughness);
-	float D = D_GGX(a2, NoH);
-	float GlossyPdf = 0.0;
-	if (HoV > 0.0)
+	// Diffuse Lobe
+	AddLobeWithMIS(OutWeight, OutPdf, GetDiffuseColor(Payload), NoL / PI, LobeSelectionPdf);
+
+	// Specular lobe
+	if (NoL > 0 && NoV > 0)
 	{
-		GlossyPdf = D * NoH / (4.0 * HoV);
+		// Assume NoH > 0
+		// FIXME: this can be further simplified
+		const float CosThetaM4 = Pow2(NoH);
+		const float TanThetaM2 = (Pow2(H.x / Alpha.x) + Pow2(H.y / Alpha.y)) / Pow2(NoH);
+		const float D = 1 / (PI * Pow2(1 + TanThetaM2) * Alpha.x * Alpha.y * CosThetaM4);
+		const float Lambda_o = 0.5 * (-1.0 + sqrt(1.0 + (Pow2(L.x * Alpha.x) + Pow2(L.y * Alpha.y)) / Pow2(L.z)));
+		const float Lambda_i = 0.5 * (-1.0 + sqrt(1.0 + (Pow2(V.x * Alpha.x) + Pow2(V.y * Alpha.y)) / Pow2(V.z)));
+		const float G2 = 1 / (Lambda_i + Lambda_o + 1);
+		const float G1 = 1 / (Lambda_i + 1);
+
+		// match behavior in SpecularGGX
+		const float3 F = F_Schlick(GetSpecularColor(Payload), VoH);
+
+		const float3 SpecWeight = F * (G2 / G1);
+		const float SpecPdf = (G1 * D * 0.25f) / NoV;
+
+		AddLobeWithMIS(OutWeight, OutPdf, SpecWeight, SpecPdf, 1.0 - LobeSelectionPdf);
 	}
-
-	return LobeSelectionPdf * LambertPdf + (1.0 - LobeSelectionPdf) * GlossyPdf;
 }
 
-void DefaultLit_EvalMaterial(
-	float3 IncomingDirection,
-	float3 OutgoingDirection,
-	FMaterialClosestHitPayload Payload,
-	float Ior1,
-	out float3 OutThroughput,
-	out float OutPdf
-);
-
 void DefaultLit_SampleMaterial(
 	float3 RayDirection,
 	FMaterialClosestHitPayload Payload,
-	float Ior1,
 	float4 RandSample,
 	out float3 OutDirection,
-	out float3 OutThroughput,
+	out float3 OutWeight,
 	out float OutPdf
 )
 {
@@ -81,7 +110,10 @@ void DefaultLit_SampleMaterial(
 	
 	float LobeSelectionPdf = CalcLobeSelectionPdf(V, Payload);
 
-	float MisWeight = 0;
+	OutDirection = 0;
+	OutWeight = 0;
+	OutPdf = 0;
+	
 	// Choose lobe based on roughness
 	if (RandSample.x < LobeSelectionPdf)
 	{
@@ -90,111 +122,46 @@ void DefaultLit_SampleMaterial(
 		// Lambert
 		float4 Result = CosineSampleHemisphere(RandSample.xy);
 		OutDirection = TangentToWorld(Result.xyz, N);
-		float NoL = saturate(dot(N, OutDirection));
-		OutThroughput = GetDiffuseColor(Payload) / PI;
-		OutPdf = LobeSelectionPdf * Result.w;
 	}
 	else
 	{
 		RandSample.x -= LobeSelectionPdf;
 		RandSample.x /= (1.0 - LobeSelectionPdf);
 
-		// GGX
-		float a2 = Pow4(GetRoughness(Payload));
-		float4 Result = ImportanceSampleGGX(RandSample.xy, a2);
-		float3 H = TangentToWorld(Result.xyz, N);
-		float3 L = 2.0 * dot(V, H) * H - V;
-		OutDirection = L;
-
-		float NoV = saturate(dot(N, V));
-		float NoH = saturate(dot(N, H));
-		float NoL = saturate(dot(N, L));
-		float VoH = saturate(dot(V, H));
-
-		float3 F_Metallic = F_Schlick(GetSpecularColor(Payload), VoH);
-		float3 F_Dielectric = FresnelReflectance(VoH, Ior1, GetIor(Payload));
-		float3 F = lerp(F_Dielectric, F_Metallic, GetMetallic(Payload));
-		float D = Result.w / Result.z;// D_GGX(a2, NoH);
-		// Note: Pedantic application of normalized weights is on purpose
-		float G = Vis_SmithJointApprox(a2, NoV, NoL) * (4.0 * NoV * NoL);
-
-		float MicrofacetPdf = 0.0;
-		if (VoH > 0.0)
+		float Roughness = max(GetRoughness(Payload), GGX_MIN_ROUGHNESS);
+		float2 Alpha;
+		float3x3 Basis;
+		if (Payload.HasAnisotropy())
 		{
-			MicrofacetPdf = Result.w / (4.0 * VoH);
+			// Material is anisotropic, our shading frame should use the tangent vector
+			GetAnisotropicRoughness(Roughness * Roughness, Payload.Anisotropy, Alpha.x, Alpha.y);
+			float3 X = Payload.WorldTangent;
+			float3 Y = normalize(cross(N, X));
+			Basis = float3x3(cross(Y, N), Y, N);
 		}
-
-		OutThroughput = 0.0;
-		if (NoV * NoL > 0.0)
+		else
 		{
-			OutThroughput = F * ((D * G) / (4.0 * NoV * NoL));
+			// No anisotropy, just pick an arbitrary basis
+			Alpha = Roughness * Roughness;
+			Basis = GetTangentBasis(N);
 		}
-		OutPdf = (1.0 - LobeSelectionPdf) * MicrofacetPdf;
-	}
-}
 
-void DefaultLit_EvalMaterial(
-	float3 IncomingDirection,
-	float3 OutgoingDirection,
-	FMaterialClosestHitPayload Payload,
-	float Ior1,
-	out float3 OutThroughput,
-	out float OutPdf
-)
-{
-	float3 V_World = -IncomingDirection;
-	float3 L_World = OutgoingDirection;
-	float3 N_World = GetWorldNormal(Payload);
-	float3 V = WorldToTangent(V_World, N_World);
-	float3 L = WorldToTangent(L_World, N_World);
-	float3 N = float3(0, 0, 1);
-	float3 H = normalize(V + L);
+		V = mul(Basis, V);
 
-	float NoL = saturate(dot(N, L));
-	float NoV = saturate(dot(N, V));
-	float NoH = saturate(dot(N, H));
-	float VoH = saturate(dot(V, H));
+		// #dxr_todo: Switch to UniformSampleDiskConcentric after fixing 0/0 bug
+		float3 H = ImportanceSampleVisibleGGX_aniso(UniformSampleDisk(RandSample.xy), Alpha, V).xyz;
 
-// dxr_todo:
-// Disabling EvaluateBXDF path as hierarchical composition of materials does not apply the same Payload properties
-//#ifndef PATH_TRACING
-#if 1
-	float a2 = Pow4(GetRoughness(Payload));
-	float D = D_GGX(a2, NoH);
-	// Note: Pedantic application of normalized weights is on purpose
-	float G = Vis_SmithJointApprox(a2, NoV, NoL) * (4.0 * NoV * NoL);
-	float3 F_Metallic = F_Schlick(GetSpecularColor(Payload), VoH);
-	float3 F_Dielectric = FresnelReflectance(VoH, Ior1, GetIor(Payload));
-	float3 F = lerp(F_Dielectric, F_Metallic, GetMetallic(Payload));
-
-	float3 LambertThroughput = GetDiffuseColor(Payload) / PI;
-	float3 GlossyThroughput = 0.0;
-	if (NoV * NoL > 0.0)
-	{
-		 GlossyThroughput = F * ((D * G) / (4.0 * NoV * NoL));
+		float3 L = reflect(-V, H);
+		if (L.z <= 0) {
+			// invalid output direction, exit early
+			return;
+		}
+		// transform to world space
+		OutDirection = normalize(mul(L, Basis));
 	}
-	OutThroughput = (LambertThroughput + GlossyThroughput);
-#else
-	// Create faux GBuffer to use with EvaluateBxDF
-	FGBufferData GBufferData = (FGBufferData)0;
-	GBufferData.Depth = 1.f; // Do not use depth
-	GBufferData.WorldNormal = Payload.WorldNormal;
-	GBufferData.BaseColor = Payload.BaseColor;
-	GBufferData.CustomData = Payload.CustomData;
-	GBufferData.GBufferAO = Payload.GBufferAO;
-	GBufferData.IndirectIrradiance = (Payload.IndirectIrradiance.x + Payload.IndirectIrradiance.y + Payload.IndirectIrradiance.z) / 3.f;
-	GBufferData.SpecularColor = Payload.SpecularColor;
-	GBufferData.DiffuseColor = Payload.DiffuseColor;
-	GBufferData.Metallic = Payload.Metallic;
-	GBufferData.Specular = Payload.Specular;
-	GBufferData.Roughness = Payload.Roughness;
-	GBufferData.ShadingModelID = Payload.ShadingModelID;
-	GBufferData.CustomData = Payload.CustomData;
-
-	FShadowTerms ShadowTerms = { 0.0, 0.0, 0.0, InitHairTransmittanceData() };
-	FDirectLighting LightingSample = EvaluateBxDF(GBufferData, N, V, L, NoL, ShadowTerms);
-	OutThroughput = LightingSample.Diffuse + LightingSample.Specular;
-#endif
-
-	OutPdf = DefaultLit_PdfMaterial(IncomingDirection, OutgoingDirection, Payload, Ior1);
-}
\ No newline at end of file
+
+	// Call eval to get the proper weight/pdf for the sampled direction
+	// TODO: might be cheaper to inline the eval here and take advantage of already computed terms?
+	// It is possible that the compiler is doing some of this for us already though ...
+	DefaultLit_EvalMaterial(RayDirection, OutDirection, Payload, OutWeight, OutPdf);
+}
diff --git a/Engine/Shaders/Private/PathTracing/Material/PathTracingFresnel.ush b/Engine/Shaders/Private/PathTracing/Material/PathTracingFresnel.ush
--- a/Engine/Shaders/Private/PathTracing/Material/PathTracingFresnel.ush
+++ b/Engine/Shaders/Private/PathTracing/Material/PathTracingFresnel.ush
@@ -5,102 +5,49 @@ PathTracingFresnel.usf: Fresnel utilities for material sampling functions
 ===============================================================================================*/
 #pragma once
 
-
-float FresnelDielectric(float CosThetaI, float IorI, float IorT)
+float FresnelReflectance(float CosI, float Eta)
 {
-	if (IorI == IorT)
-	{
-		return 1.0;
-	}
+	// This is a float only version of F_Fresnel that takes the Ior ratio as input directly (and handles TIR)
 
-	CosThetaI = clamp(CosThetaI, -1.0, 1.0);
+	// NOTE: this function tries uses the exact same math as the version below for consistency
 
-	if (CosThetaI < 0.0)
-	{
-		// Swap indices
-		CosThetaI = -CosThetaI;
-		float IorTmp = IorI;
-		IorI = IorT;
-		IorT = IorTmp;
-	}
+	float g2 = Eta * Eta - 1 + CosI * CosI;
 
-	const float SinThetaI = 1.0 - (CosThetaI * CosThetaI);
-	const float IorRatio = IorI / IorT;
-	const float d = 1.0 - SinThetaI * IorRatio*IorRatio;
-	if (d < 0.0)
-	{
-		return 1.0;
+	if (g2 >= 0) {
+		float c = abs(CosI);
+		float g = sqrt(g2);
+		float a2 = Pow2((g - c) / (g + c));
+		float b2 = Pow2((c * (g + c) - 1) / (c * (g - c) + 1));
+		return 0.5 * a2 * (1.0 + b2);
 	}
 
-	float CosThetaT = sqrt(d);
-	float CosI_A = IorI * CosThetaI;
-	float CosT_A = IorI * CosThetaT;
-	float CosI_B = IorT * CosThetaI;
-	float CosT_B = IorT * CosThetaT;
-
-	float Parl = (CosT_A - CosI_B) / (CosT_A + CosI_B);
-	float Perp = (CosI_A - CosT_B) / (CosI_A + CosT_B);
-	float Fr = min(0.5*(Perp*Perp + Parl * Parl), 1.0);
-
-	return Fr;
+	return 1.0;
 }
 
-void CalcFresnelCoefficients(float CosTheta1, float N1, float N2, out float CosTheta2, out float Rs, out float Rp)
+// returns true if we sampled a refraction, or false if we reflected (or TIR occured)
+bool SampleRefraction(float3 RayDirection, float3 Normal, float Eta, float RandSample, out float3 OutDirection)
 {
-	Rs = 0.0;
-	Rp = 0.0;
-	CosTheta2 = 0.0;
-	// Use Snell's law to compute CosTheta2:
-	// N1 * sin(Theta1) = N2 * sin(Theta2)
-	// [N1 * sin(Theta1)]^2 = [N2 * sin(Theta2)]^2
-	float SinTheta1_Squared = 1.0 - CosTheta1 * CosTheta1;
-	float SinTheta2_Squared = (SinTheta1_Squared * N1 * N1) / (N2 * N2);
-
-	float CosTheta1_Squared = 1.0 - SinTheta1_Squared;
-	float CosTheta2_Squared = 1.0 - SinTheta2_Squared;
-	if (CosTheta2_Squared >= 0.0)
-	{
-		CosTheta2 = sqrt(CosTheta2_Squared);
-
-		// Optimized Fresnel equations
-		// Rs = (N1 * CosTheta1 - N2 * CosTheta2) / (N1 * CosTheta1 + N2 * CosTheta2);
-		// Rp = (N2 * CosTheta1 - N1 * CosTheta2) / (N2 * CosTheta1 + N1 * CosTheta2);
-		float N1_CosTheta1 = N1 * CosTheta1;
-		float N1_CosTheta2 = N1 * CosTheta2;
-		float N2_CosTheta1 = N2 * CosTheta1;
-		float N2_CosTheta2 = N2 * CosTheta2;
-		Rs = (N1_CosTheta1 - N2_CosTheta2) / (N1_CosTheta1 + N2_CosTheta2);
-		Rp = (N2_CosTheta1 - N1_CosTheta2) / (N2_CosTheta1 + N1_CosTheta2);
+	// Compute refracted direction and Fresnel term at the same time to re-use intermediate values
+	// and to get matching results for the TIR condition
+	float CosI = abs(dot(RayDirection, Normal));
+	float g2 = Eta * Eta - 1.0 + CosI * CosI;
+
+	if (g2 >= 0.0) {
+		float c = CosI;
+		float g = sqrt(g2);
+		float a2 = Pow2((g - c) / (g + c));
+		float b2 = Pow2((c * (g + c) - 1) / (c * (g - c) + 1));
+		float Fr = 0.5 * a2 * (1.0 + b2);
+
+		// Stochastically decide to transmit or not based on Fr
+		// This assumes RandSample is in [0,1)
+		if (RandSample >= Fr)
+		{
+			OutDirection = normalize(RayDirection + Normal * (c - g));
+			return true;
+		}
 	}
+	// either TIR, or sampled reflection
+	OutDirection = RayDirection + (2 * CosI) * Normal;
+	return false;
 }
-
-// TODO: Unify FresnelReflectance() and FresnelDielectric()
-float FresnelReflectance(float CosTheta1, float N1, float N2)
-{
-	float CosTheta2;
-	float Rs;
-	float Rp;
-	CalcFresnelCoefficients(CosTheta1, N1, N2, CosTheta2, Rs, Rp);
-
-	float Fr = (Rs * Rs + Rp * Rp) / 2.0;
-	return Fr;
-}
-
-float FresnelTransmittance(float CosTheta1, float N1, float N2)
-{
-	float Ft = 0.0;
-	if (CosTheta1 > 0.0)
-	{
-		float CosTheta2;
-		float Rs;
-		float Rp;
-		CalcFresnelCoefficients(CosTheta1, N1, N2, CosTheta2, Rs, Rp);
-
-		float Ts = Rs + 1;
-		float Tp = (Rp + 1) * N1 / N2;
-		Ft = (Ts * Ts + Tp * Tp) / 2.0;
-		Ft *= (N2 * CosTheta2) / (N1 * CosTheta1);
-	}
-
-	return Ft;
-}
\ No newline at end of file
diff --git a/Engine/Shaders/Private/PathTracing/Material/PathTracingGlossy.ush b/Engine/Shaders/Private/PathTracing/Material/PathTracingGlossy.ush
--- a/Engine/Shaders/Private/PathTracing/Material/PathTracingGlossy.ush
+++ b/Engine/Shaders/Private/PathTracing/Material/PathTracingGlossy.ush
@@ -6,377 +6,3 @@ Glossy.usf: Microfacet GGX BRDF sampling functions
 
 #pragma once
 
-
-bool IsNanOrNotFinite(float p)
-{
-	//return (!isfinite(p) || isnan(p));
-	// Work around lack of support fo isnan and isfinite in current compiler
-	return (asuint(p) & (0xFF << 23)) == (0xFF << 23);
-}
-
-float ClampIfInvalid(float p)
-{
-	return IsNanOrNotFinite(p) ? 10000000 : p;
-}
-
-float GetSpecularEventProbability(float3 SpecularColor, float NoV)
-{
-	float3 F = F_Schlick(SpecularColor, NoV);
-	return (F.x + F.y + F.z) / 3.f;
-}
-
-// G_SmithV predivided by 4*NoV*NoL. #dxr-todo: move to BRDF.ush?
-float PredividedG_SmithV(float a2, float NoV, const float NoL)
-{
-	const float G_SmithV = NoV + sqrt((NoV - NoV * a2) * NoV + a2);
-	return 0.5f / (G_SmithV * NoL);
-}
-
-float4 ImportanceSampleVisibleGGX_Aniso(float2 DiskE, float2 a2, float3 V)
-{
-	// TODO float2 alpha for anisotropic
-	float ax = sqrt(a2.x);
-	float ay = sqrt(a2.y);
-
-	// stretch
-	float3 Vh = normalize(float3(ax * V.x, ax * V.x, V.z));
-
-	// Orthonormal basis
-	// Tangent0 is orthogonal to N.
-	float3 Tangent0 = (Vh.z < 0.9999) ? normalize(cross(float3(0, 0, 1), Vh)) : float3(1, 0, 0);
-	float3 Tangent1 = cross(Vh, Tangent0);
-
-	float2 p = DiskE;
-	float s = 0.5 + 0.5 * Vh.z;
-	p.y = (1 - s) * sqrt(1 - p.x * p.x) + s * p.y;
-
-	float3 H;
-	H = p.x * Tangent0;
-	H += p.y * Tangent1;
-	H += sqrt(saturate(1 - dot(p, p))) * Vh;
-
-	// unstretch
-	H = normalize(float3(ax * H.x, ay * H.y, max(0.0, H.z)));
-
-	float NoV = V.z;
-	float NoH = H.z;
-	float VoH = dot(V, H);
-
-	float d = (NoH * a2 - NoH) * NoH + 1;
-	float D = a2 / (PI*d*d);
-
-	float G_SmithV = 2 * NoV / (NoV + sqrt(NoV * (NoV - NoV * a2) + a2));
-
-	float PDF = G_SmithV * VoH * D / NoV;
-
-	return float4(H, PDF);
-}
-
-float Glossy_Pdf(float3 IncomingDirection, float3 OutgoingDirection, FMaterialClosestHitPayload Payload)
-{
-	float3 SpecularColor = GetSpecularColor(Payload);
-
-	//-------
-	SpecularColor = float3(0, 1, 0);
-	//-------
-
-
-	float3 N_World = GetWorldNormal(Payload);
-	float Roughness = max(GetRoughness(Payload), GLOSSY_MIN_ROUGHNESS);
-
-	float3 V_World = normalize(-IncomingDirection);
-	float3 V = WorldToTangent(V_World, N_World);
-	float3 L = WorldToTangent(normalize(OutgoingDirection), N_World);
-
-	float NoV = V.z;
-	float NoL = L.z;
-
-	const float ProbSpecular = GetSpecularEventProbability(SpecularColor, NoV);
-
-	// Specular
-	float SpecularPdf = 0.0f;
-	float3 H = normalize(V + L);
-	float VoH = saturate(dot(V, H));
-	float NoH = saturate(H.z);
-
-	if (NoL * NoV > 0.0f && NoH > 0.0f)
-	{
-		if (Payload.HasAnisotropy())
-		{
-			float ax = 0;
-			float ay = 0;
-			GetAnisotropicRoughness(Roughness*Roughness, Payload.Anisotropy, ax, ay);
-
-			//#todo check energy normalization factor
-			float3 X = Payload.WorldTangent;
-			float3 Y = normalize(cross(N_World, X));
-			float3 H_World = TangentToWorld(H, N_World);
-
-			float XoH = saturate(dot(X, H_World));
-			float YoH = saturate(dot(Y, H_World));
-
-			float D = D_GGXaniso(ax, ay, NoH, XoH, YoH);
-
-			if (VoH > 0.0)
-			{
-				SpecularPdf = D * NoH / (4.0 * VoH);
-			}
-		}
-		else
-		{
-			float a2 = Pow4(Roughness);
-			float D = D_GGX(a2, NoH);
-			float G1 = PredividedG_SmithV(a2, abs(NoV), abs(NoL));
-
-			SpecularPdf = ClampIfInvalid(G1 * D);
-			SpecularPdf *= NoL;
-		}
-	}
-
-	// Diffuse 
-	float DiffusePdf = NoL / PI;
-
-	float OutPdf = ProbSpecular * SpecularPdf + (1.0f - ProbSpecular) * DiffusePdf;
-	return (OutPdf);
-}
-
-void Glossy_SampleMaterial(
-	float3 RayDirection,
-	FMaterialClosestHitPayload Payload,
-	float4 RandSample,
-	out float3 OutDirection,
-	out float3 OutThroughput,
-	out float OutPdf)
-{
-	OutDirection = float3(0, 0, 0);
-	OutThroughput = float3(0, 0, 0);
-	OutPdf = 0.0f;
-
-	float3 N_World = GetWorldNormal(Payload);
-
-	// Diffuse and Specular color are calculated in the payload unpacking, both depend on BaseColor and Metallic attributes
-	float3 DiffuseColor = GetDiffuseColor(Payload);
-	float3 SpecularColor = GetSpecularColor(Payload);
-	float Roughness = max(GetRoughness(Payload), GLOSSY_MIN_ROUGHNESS);
-
-	//-------
-	DiffuseColor = float3(0, 1, 0);
-	SpecularColor = float3(0, 1, 0);
-	//-------
-
-
-	float3 V_World = normalize(-RayDirection);
-	float3 V = WorldToTangent(V_World, N_World);
-	float3 L;
-
-	const float NoV = V.z;
-	const float ProbSpecular = GetSpecularEventProbability(SpecularColor, NoV); 
-
-	if ((RandSample.z <= ProbSpecular)) 
-	{	
-		bool Flip = V.z < 0.0f;
-
-		float4 ResGGX = 0;
-		float ax = 0;
-		float ay = 0;
-
-		if (Payload.HasAnisotropy())
-		{
-			GetAnisotropicRoughness(Roughness*Roughness, Payload.Anisotropy, ax, ay);
-
-			float2 a2 = float2(ax, ay);
-			ResGGX = ImportanceSampleVisibleGGX_Aniso(UniformSampleDisk(RandSample.xy), Pow4(Roughness), Flip ? -V : V);
-		}
-		else
-		{
-			ResGGX = ImportanceSampleVisibleGGX(UniformSampleDisk(RandSample.xy), Pow4(Roughness), Flip ? -V : V);
-		}
-
-		float3 H = ResGGX.xyz;
-		if (Flip)
-		{
-			H = -H; //Use ResGGX.w for the PDF?
-		}
-
-		L = 2 * dot(V, H) * H - V;
-
-		float NoL = L.z;
-		float NoV = V.z;
-		float NoH = H.z;
-
-		if (L.z < 0.0f || NoL * NoV <= 0.0f || NoH == 0.0f)
-		{
-			return;
-		}
-
-		float D = 0;
-		float G2 = 0;
-		float G1 = 0;
-
-		if (Payload.HasAnisotropy())
-		{
-			//#todo check energy normalization factor
-			
-			float3 X = Payload.WorldTangent;
-			float3 Y = normalize(cross(N_World, X));
-			float3 H_World = TangentToWorld(H, N_World);
-
-			float XoH = saturate(dot(X, H_World));
-			float YoH = saturate(dot(Y, H_World));
-			float XoV = saturate(dot(X, V_World));
-			float YoV = saturate(dot(Y, V_World));
-
-			//#todo: optimize, we are doing it twice. Work in world coords all the time
-			float3 L_World = TangentToWorld(L, N_World);
-
-			float XoL = saturate(dot(X, L_World));
-			float YoL = saturate(dot(Y, L_World));
-
-			D = D_GGXaniso(ax, ay, NoH, XoH, YoH);
-			G2 = Vis_SmithJointAniso(ax, ay, abs(NoV), abs(NoL), XoV, XoL, YoV, YoL);
-
-			float a2_NoAniso = Pow4(Roughness); //todo
-			G1 = PredividedG_SmithV(a2_NoAniso, abs(NoV), abs(NoL));
-		}
-		else
-		{
-			float a2 = Pow4(Roughness);
-			D = D_GGX(a2, NoH);
-			G2 = Vis_SmithJointApprox(a2, abs(NoV), abs(NoL));
-			G1 = PredividedG_SmithV(a2, abs(NoV), abs(NoL));
-		}
-
-		float3 F = F_Schlick(SpecularColor, dot(V, H));
-
-		OutPdf = ClampIfInvalid(G1 * D) * abs(NoL);
-		OutThroughput = F * ClampIfInvalid(G2 / G1);
-
-		OutThroughput /= ProbSpecular;
-		OutPdf *= ProbSpecular;
-	}
-	else 
-	{
-		float4 DiffuseSampledValue = CosineSampleHemisphere(RandSample.xy);
-		L = DiffuseSampledValue.xyz;
-		L.z = abs(L.z);
-
-		OutThroughput = DiffuseColor; // Compensated below 
-		OutPdf = DiffuseSampledValue.w;
-
-		OutThroughput /= 1.0f - ProbSpecular;
-		OutPdf *= 1.0f - ProbSpecular;
-	}
-
-	// Througput is already previdided by the PDF so it has to be compensated
-	OutThroughput *= OutPdf;
-
-	OutDirection = TangentToWorld(L, N_World);
-}
-
-
-void Glossy_EvalMaterial(
-	float3 IncomingDirection,
-	float3 OutgoingDirection,
-	FMaterialClosestHitPayload Payload,
-	out float3 OutThroughput,
-	out float OutPdf
-)
-{
-	float3 DiffuseColor = GetDiffuseColor(Payload);
-	float3 SpecularColor = GetSpecularColor(Payload);
-	float Roughness = max(GetRoughness(Payload), 0.001);
-
-	//-------
-	DiffuseColor = float3(0, 1, 0);
-	SpecularColor = float3(0, 1, 0);
-	//-------
-
-	float3 N_World = GetWorldNormal(Payload);
-	float3 V_World = normalize(-IncomingDirection);
-	float3 L_World = normalize(OutgoingDirection);
-
-	float3 V = WorldToTangent(V_World, N_World);
-	float3 L = WorldToTangent(L_World, N_World);
-
-	float NoV = V.z;
-	float NoL = L.z;
-
-	// Specular eval
-	float3 SpecularThroughput = float3(0, 0, 0);
-	float SpecularPdf = 0.0;
-	
-	float3 H = normalize(V + L);
-	float VoH = saturate(dot(V, H));
-	float NoH = saturate(H.z);
-
-	if (NoL * NoV > 0.0f && NoH > 0.0f)
-	{
-		float D = 0;
-		float G2 = 0;
-		float G1 = 0;
-
-		if (Payload.HasAnisotropy())
-		{
-			float ax = 0;
-			float ay = 0;
-			GetAnisotropicRoughness(Roughness*Roughness, Payload.Anisotropy, ax, ay);
-
-			//#todo check energy normalization factor
-
-			float3 X = Payload.WorldTangent;
-			float3 Y = normalize(cross(N_World, X));
-
-			float3 H_World = TangentToWorld(H, N_World);
-
-			float XoH = saturate(dot(X, H_World));
-			float YoH = saturate(dot(Y, H_World));
-			float XoV = saturate(dot(X, V_World));
-			float YoV = saturate(dot(Y, V_World));
-			//#todo check world vs local
-			float XoL = saturate(dot(X, L_World));
-			float YoL = saturate(dot(Y, L_World));
-
-			float D = D_GGXaniso(ax, ay, NoH, XoH, YoH);
-			float G2 = Vis_SmithJointAniso(ax, ay, abs(NoV), abs(NoL), XoV, XoL, YoV, YoL);
-
-			float a2 = Pow4(Roughness);
-			float G1 = PredividedG_SmithV(a2, abs(NoV), abs(NoL));
-		}
-		else
-		{
-			float a2 = Pow4(Roughness);
-			float D = D_GGX(a2, NoH);
-			float G2 = Vis_SmithJointApprox(a2, abs(NoV), abs(NoL));
-			float G1 = PredividedG_SmithV(a2, abs(NoV), abs(NoL));
-		}
-
-		SpecularThroughput = ClampIfInvalid(G2 * D);
-		SpecularPdf = ClampIfInvalid(G1 * D);
-		SpecularPdf *= NoL;
-
-		if (SpecularPdf != 0)
-		{
-			float3 F = F_Schlick(SpecularColor, VoH);
-			SpecularThroughput *= F;
-		}	
-	}
-	   
-	// Diffuse eval
-	float3 DiffuseThroughput = DiffuseColor / PI;
-	float DiffusePdf = NoL / PI;
-
-	// Add contributions
-	OutThroughput = (DiffuseThroughput + SpecularThroughput);
-
-	float ProbSpecular = GetSpecularEventProbability(SpecularColor, NoV);
-	OutPdf = ProbSpecular * SpecularPdf + (1.0f - ProbSpecular) * DiffusePdf;
-}
-
-
-
-
-
-
-
-
-
diff --git a/Engine/Shaders/Private/PathTracing/Material/PathTracingLambert.ush b/Engine/Shaders/Private/PathTracing/Material/PathTracingLambert.ush
--- a/Engine/Shaders/Private/PathTracing/Material/PathTracingLambert.ush
+++ b/Engine/Shaders/Private/PathTracing/Material/PathTracingLambert.ush
@@ -9,36 +9,35 @@ void Lambert_SampleMaterial(
 	FMaterialClosestHitPayload Payload,
 	float4 RandSample,
 	out float3 OutDirection,
-	out float3 OutThroughput,
+	out float3 OutWeight,
 	out float OutPdf)
 {
 	float3 N_World = GetWorldNormal(Payload);
-	float3 DiffuseColor = GetDiffuseColor(Payload);
 
 	float4 SampledValue = CosineSampleHemisphere(RandSample.xy);
 	OutDirection = TangentToWorld(SampledValue.xyz, N_World);
 
-	OutThroughput = DiffuseColor / PI;
+	OutWeight = GetBaseColor(Payload);
 	OutPdf = SampledValue.w;
 }
 
 float Lambert_Pdf(float3 OutgoingDirection, FMaterialClosestHitPayload Payload)
 {
 	float3 N_World = GetWorldNormal(Payload);
-	float NoL = max(dot(N_World, OutgoingDirection), 0.0);
-	return (NoL / PI);
+	float NoL = saturate(dot(N_World, OutgoingDirection));
+	return NoL / PI;
 }
 
 void Lambert_EvalMaterial(
 	float3 OutgoingDirection,
 	FMaterialClosestHitPayload Payload,
-	out float3 OutThroughput,
+	out float3 OutWeight,
 	out float OutPdf
 )
 {
 	float3 N_World = GetWorldNormal(Payload);
-	float NoL = max(dot(N_World, OutgoingDirection), 0.0);
+	float NoL = saturate(dot(N_World, OutgoingDirection));
 
+	OutWeight = GetBaseColor(Payload);
 	OutPdf = NoL / PI;
-	OutThroughput = GetDiffuseColor(Payload) / PI;
-}
\ No newline at end of file
+}
diff --git a/Engine/Shaders/Private/PathTracing/Material/PathTracingMaterialCommon.ush b/Engine/Shaders/Private/PathTracing/Material/PathTracingMaterialCommon.ush
--- a/Engine/Shaders/Private/PathTracing/Material/PathTracingMaterialCommon.ush
+++ b/Engine/Shaders/Private/PathTracing/Material/PathTracingMaterialCommon.ush
@@ -51,7 +51,6 @@ float GetClearCoatRoughness(FMaterialClosestHitPayload HitInfo)
 
 float3 GetClearCoatBottomNormal(FMaterialClosestHitPayload HitInfo)
 {
-	float3 N = HitInfo.WorldNormal;
 #ifdef PATH_TRACING
 	#if CLEAR_COAT_BOTTOM_NORMAL
 	// Functionally inverse of behavior in RayTracingMaterialHitShader.usf
@@ -59,10 +58,10 @@ float3 GetClearCoatBottomNormal(FMaterialClosestHitPayload HitInfo)
 	float2 oct3 = float2(HitInfo.CustomData.a, HitInfo.CustomData.z);
 	float2 oct1 = (oct3 - (128.0 / 255.0)) * 2.0 + oct2;
 
-	N = OctahedronToUnitVector(oct1);
+	return OctahedronToUnitVector(oct1);
 	#endif
 #endif
-	return N;
+	return HitInfo.WorldNormal;
 }
 
 float GetClearCoatThickness(FMaterialClosestHitPayload HitInfo)
@@ -71,16 +70,6 @@ float GetClearCoatThickness(FMaterialClosestHitPayload HitInfo)
 	return 1.0;
 }
 
-float GetClearCoatIor(FMaterialClosestHitPayload HitInfo)
-{
-	return 1.5;
-}
-
-float GetIor(FMaterialClosestHitPayload HitInfo)
-{
-	return DielectricF0ToIor(DielectricSpecularToF0(HitInfo.Specular));
-}
-
 float3 GetSubsurfaceColor(FMaterialClosestHitPayload HitInfo)
 {
 	return HitInfo.CustomData.xyz;
@@ -91,16 +80,6 @@ float3 GetWorldNormal(FMaterialClosestHitPayload HitInfo)
 	return HitInfo.WorldNormal;
 }
 
-bool IsPureSpecular(FMaterialClosestHitPayload HitInfo)
-{
-	return GetRoughness(HitInfo) < GLOSSY_MIN_ROUGHNESS;
-}
-
-bool IsPureSpecularReflection(FMaterialClosestHitPayload HitInfo)
-{
-	return IsPureSpecular(HitInfo) && HitInfo.BlendingMode == 0 && HitInfo.Metallic == 1;
-}
-
 bool IsPureSpecularTransmission(FMaterialClosestHitPayload HitInfo)
 {
 	return HitInfo.BlendingMode > RAY_TRACING_BLEND_MODE_OPAQUE;
@@ -136,19 +115,29 @@ bool IsMaterialTransmissive(FMaterialClosestHitPayload HitInfo)
 	return HitInfo.ShadingModelID == SHADINGMODELID_TWOSIDED_FOLIAGE;
 }
 
-bool BSDF_Refract(float3 V, float3 N, float Eta, out float3 L)
+// Takes as input the sample weight and pdf for a certain lobe of a mixed model, together with the probability of picking that lobe
+// This function then updates a running total Weight and Pdf value that represents the overall contribution of the BxDF
+// This function should be called when a BxDF is made up of multiple lobes combined with a sum to correctly account for the probability
+// of sampling directions via all lobes.
+// NOTE: this function also contains special logic to handle cases with infinite pdfs cleanly
+void AddLobeWithMIS(inout float3 Weight, inout float Pdf, float3 LobeWeight, float LobePdf, float LobeProb)
 {
-	const float CosThetaI = dot(V, N);
-	const float CosThetaT2 = 1.f - Eta * Eta * (1.0f - CosThetaI * CosThetaI);
-	if (CosThetaT2 < 0.0f)
+	const float MinLobeProb = 1.1754943508e-38; // smallest normal float
+	if (LobeProb > MinLobeProb)
 	{
-		// Total internal reflection
-		L = reflect(-V, N);
-		return false;
+		LobePdf *= LobeProb;
+		LobeWeight *= 1 / LobeProb;
+
+		// See discussion in MISWeightRobust for why this is better than LobePdf / (Pdf + LobePdf)
+		float MISWeight;
+		if (Pdf < LobePdf)
+			MISWeight = 1 / (1 + Pdf / LobePdf);
+		else if (LobePdf < Pdf)
+			MISWeight = 1 - 1 / (1 + LobePdf / Pdf);
+		else
+			MISWeight = 0.5f; // avoid (rare) inf/inf
+
+		Weight = lerp(Weight, LobeWeight, MISWeight);
+		Pdf += LobePdf;
 	}
-
-	const float CosThetaT = -(CosThetaI >= 0.0f ? 1.0f : -1.0f) * sqrt(CosThetaT2);
-	L = (Eta * CosThetaI + CosThetaT) * N - Eta * V;
-
-	return true;
 }
diff --git a/Engine/Shaders/Private/PathTracing/Material/PathTracingMaterialSampling.ush b/Engine/Shaders/Private/PathTracing/Material/PathTracingMaterialSampling.ush
--- a/Engine/Shaders/Private/PathTracing/Material/PathTracingMaterialSampling.ush
+++ b/Engine/Shaders/Private/PathTracing/Material/PathTracingMaterialSampling.ush
@@ -21,33 +21,14 @@ PathTracingMaterialSampling.usf: Material sampling functions
 	#define FORCE_ALL_MATERIALS_DIFFUSE 0
 #endif
 
-void SampleMaterialAsWorldNormal(
-	float3 RayDirection,
-	FMaterialClosestHitPayload Payload,
-	float4 RandSample,
-	out float3 OutDirection,
-	out float3 OutThroughput,
-	out float OutPdf)
-{
-	float3 N_World = GetWorldNormal(Payload);
-	OutDirection = -RayDirection;
-	OutThroughput = (N_World / 2.0) + 0.5;
-	OutPdf = 1.0;
-}
-
 bool IsNeeValidMaterial(FMaterialClosestHitPayload Payload)
 {
 	// Checks done in the same order as other sampling and evaluation functions to ensure consistency
 	
-	#if FORCE_ALL_MATERIALS_DIFFUSE
-		return true;
-	#endif
-
-	if (IsPureSpecularReflection(Payload))
-	{
-		return false;
-	}
-	else if (IsPureSpecularTransmission(Payload))
+#if FORCE_ALL_MATERIALS_DIFFUSE
+	return true;
+#else
+	if (IsPureSpecularTransmission(Payload))
 	{
 		return false;
 	}
@@ -67,24 +48,17 @@ bool IsNeeValidMaterial(FMaterialClosestHitPayload Payload)
 	{
 		return false;
 	}
-	else if(IsLambert(Payload))
-	{
-		return true;
-	}
-	else
-	{
-		return true;
-	}
+	return true;
+#endif
 }
 
 void SampleMaterial(
 	float3 RayDirection,
 	FMaterialClosestHitPayload Payload,
-	float Ior1,
 	float4 RandSample,
 	bool bForceRadianceProbe,
 	out float3 OutDirection,
-	out float3 OutThroughput,
+	out float3 OutWeight,
 	out float OutPdf,
 	out float OutPositionBiasSign)
 {
@@ -92,179 +66,104 @@ void SampleMaterial(
 	
 	if (bForceRadianceProbe)
 	{
-		RadianceProbe_SampleMaterial(Payload, RandSample, OutDirection, OutThroughput, OutPdf);
+		RadianceProbe_SampleMaterial(Payload, RandSample, OutDirection, OutWeight, OutPdf);
 		return;
 	}
 	
-	#if FORCE_ALL_MATERIALS_DIFFUSE
+#if FORCE_ALL_MATERIALS_DIFFUSE
 	{
-		Lambert_SampleMaterial(Payload, RandSample, OutDirection, OutThroughput, OutPdf);
-		float NoL = saturate(dot(Payload.WorldNormal, OutDirection));
-		OutThroughput *= NoL;
-		return;
-	}
-	#endif
-	
-	// Note: Early-return dirac materials do not apply projected measure
-	if (IsPureSpecularReflection(Payload))
-	{
-		SpecularReflection_SampleMaterial(RayDirection, Payload, OutDirection, OutThroughput, OutPdf);
-		return;
+		Lambert_SampleMaterial(Payload, RandSample, OutDirection, OutWeight, OutPdf);
 	}
-	else if (IsPureSpecularTransmission(Payload))
+#else	
+	if (IsPureSpecularTransmission(Payload))
 	{
-		SpecularRefraction_SampleMaterial(RayDirection, Payload, RandSample, OutDirection, OutThroughput, OutPdf, OutPositionBiasSign);
-		return;
+		SpecularRefraction_SampleMaterial(RayDirection, Payload, RandSample, OutDirection, OutWeight, OutPdf, OutPositionBiasSign);
 	}
 	else if (IsDefaultLitMaterial(Payload))
 	{
-		DefaultLit_SampleMaterial(RayDirection, Payload, Ior1, RandSample, OutDirection, OutThroughput, OutPdf);
+		DefaultLit_SampleMaterial(RayDirection, Payload, RandSample, OutDirection, OutWeight, OutPdf);
 	}
 	else if (IsClearCoatMaterial(Payload))
 	{
-		ClearCoat_SampleMaterial(RayDirection, Payload, RandSample, OutDirection, OutThroughput, OutPdf);
+		ClearCoat_SampleMaterial(RayDirection, Payload, RandSample, OutDirection, OutWeight, OutPdf);
 	}
 	else if (IsTwoSidedFoliageMaterial(Payload))
 	{
-		TwoSidedFoliage_SampleMaterial(RayDirection, Payload, RandSample, OutDirection, OutThroughput, OutPdf, OutPositionBiasSign);
+		TwoSidedFoliage_SampleMaterial(RayDirection, Payload, RandSample, OutDirection, OutWeight, OutPdf, OutPositionBiasSign);
 	}
 	else if (IsUnlitMaterial(Payload))
 	{
+		// Nothing to do, material does not scatter light
 		OutDirection = 0;
-		OutThroughput = 0;
+		OutWeight = 0;
 		OutPdf = 0;
-		return;
 	}
-	else if(IsLambert(Payload))
-	{
-		Lambert_SampleMaterial(Payload, RandSample, OutDirection, OutThroughput, OutPdf);
-	}
-	else 
+	else
 	{
-		Glossy_SampleMaterial(RayDirection, Payload, RandSample, OutDirection, OutThroughput, OutPdf);
+		Lambert_SampleMaterial(Payload, RandSample, OutDirection, OutWeight, OutPdf);
 	}
+#endif
 
-	float NoL = saturate(dot(Payload.WorldNormal, OutDirection));
-	OutThroughput *= NoL;
 }
 
-void PdfMaterial(
+void EvalMaterial(
 	float3 IncomingDirection,
 	float3 OutgoingDirection,
 	FMaterialClosestHitPayload Payload,
-	float Ior1,
 	bool bForceRadianceProbe,
+	out float3 OutWeight,
 	out float OutPdf
 )
 {
 	if (bForceRadianceProbe)
 	{
-		OutPdf = RadianceProbe_Pdf(OutgoingDirection, Payload);
+		RadianceProbe_EvalMaterial(OutgoingDirection, Payload, OutWeight, OutPdf);
 		return;
 	}
 	
-	#if FORCE_ALL_MATERIALS_DIFFUSE
-		OutPdf = Lambert_Pdf(OutgoingDirection, Payload);
-		return;
-	#endif
-	
-	if (IsPureSpecularReflection(Payload))
+#if FORCE_ALL_MATERIALS_DIFFUSE
 	{
-		OutPdf = SpecularReflection_Pdf();
+		Lambert_EvalMaterial(OutgoingDirection, Payload, OutWeight, OutPdf);
 	}
-	else if (IsPureSpecularTransmission(Payload))
+#else
+	if (IsPureSpecularTransmission(Payload))
 	{
-		OutPdf = SpecularTransmission_Pdf();
+		SpecularTransmission_EvalMaterial(OutWeight, OutPdf);
 	}
 	else if (IsDefaultLitMaterial(Payload))
 	{
-		OutPdf = DefaultLit_PdfMaterial(IncomingDirection, OutgoingDirection, Payload, Ior1);
+		DefaultLit_EvalMaterial(IncomingDirection, OutgoingDirection, Payload, OutWeight, OutPdf);
 	}
 	else if (IsClearCoatMaterial(Payload))
 	{
-		OutPdf = ClearCoat_PdfMaterial(IncomingDirection, OutgoingDirection, Payload);
+		ClearCoat_EvalMaterial(IncomingDirection, OutgoingDirection, Payload, OutWeight, OutPdf);
 	}
 	else if (IsTwoSidedFoliageMaterial(Payload))
 	{
-		OutPdf = TwoSidedFoliage_PdfMaterial(IncomingDirection, OutgoingDirection, Payload);
+		TwoSidedFoliage_EvalMaterial(IncomingDirection, OutgoingDirection, Payload, OutWeight, OutPdf);
 	}
 	else if (IsUnlitMaterial(Payload))
 	{
+		OutWeight = 0.0;
 		OutPdf = 0.0;
 	}
-	else if (IsLambert(Payload))
-	{
-		OutPdf = Lambert_Pdf(OutgoingDirection, Payload);
-	}
 	else
 	{
-		OutPdf = Glossy_Pdf(IncomingDirection, OutgoingDirection, Payload);
+		Lambert_EvalMaterial(OutgoingDirection, Payload, OutWeight, OutPdf);
 	}
+#endif
 }
 
-void EvalMaterial(
+
+void PdfMaterial(
 	float3 IncomingDirection,
 	float3 OutgoingDirection,
 	FMaterialClosestHitPayload Payload,
-	float Ior1,
 	bool bForceRadianceProbe,
-	out float3 OutThroughput,
 	out float OutPdf
 )
 {
-	if (bForceRadianceProbe)
-	{
-		RadianceProbe_EvalMaterial(OutgoingDirection, Payload, OutThroughput, OutPdf);
-		return;
-	}
-	
-	#if FORCE_ALL_MATERIALS_DIFFUSE
-	{
-		Lambert_EvalMaterial(OutgoingDirection, Payload, OutThroughput, OutPdf);		
-		float NoL = saturate(dot(Payload.WorldNormal, OutgoingDirection));
-		OutThroughput *= NoL;
-		return;
-	}
-	#endif
-	
-	// Note: Early-return dirac materials do not apply projected measure
-	if (IsPureSpecularReflection(Payload))
-	{
-		SpecularReflection_EvalMaterial(OutThroughput, OutPdf);
-		return;
-	}
-	else if (IsPureSpecularTransmission(Payload))
-	{
-		SpecularTransmission_EvalMaterial(OutThroughput, OutPdf);
-		return;
-	}
-	else if (IsDefaultLitMaterial(Payload))
-	{
-		DefaultLit_EvalMaterial(IncomingDirection, OutgoingDirection, Payload, Ior1, OutThroughput, OutPdf);
-	}
-	else if (IsClearCoatMaterial(Payload))
-	{
-		ClearCoat_EvalMaterial(IncomingDirection, OutgoingDirection, Payload, OutThroughput, OutPdf);
-	}
-	else if (IsTwoSidedFoliageMaterial(Payload))
-	{
-		TwoSidedFoliage_EvalMaterial(IncomingDirection, OutgoingDirection, Payload, OutThroughput, OutPdf);
-	}
-	else if (IsUnlitMaterial(Payload))
-	{
-		OutThroughput = 0.0;
-		OutPdf = 0.0;
-		return;
-	}
-	else if (IsLambert(Payload))
-	{
-		Lambert_EvalMaterial(OutgoingDirection, Payload, OutThroughput, OutPdf);
-	}
-	else
-	{
-		Glossy_EvalMaterial(IncomingDirection, OutgoingDirection, Payload, OutThroughput, OutPdf);
-	}
-
-	float NoL = saturate(dot(Payload.WorldNormal, OutgoingDirection));
-	OutThroughput *= NoL;
-}
\ No newline at end of file
+	// NOTE: rely on compiler to optimize unused calculations
+	float3 Weight = 0;
+	EvalMaterial(IncomingDirection, OutgoingDirection, Payload, bForceRadianceProbe, Weight, OutPdf);
+}
diff --git a/Engine/Shaders/Private/PathTracing/Material/PathTracingRadianceProbe.ush b/Engine/Shaders/Private/PathTracing/Material/PathTracingRadianceProbe.ush
--- a/Engine/Shaders/Private/PathTracing/Material/PathTracingRadianceProbe.ush
+++ b/Engine/Shaders/Private/PathTracing/Material/PathTracingRadianceProbe.ush
@@ -58,7 +58,7 @@ void RadianceProbe_SampleMaterial(
 	FMaterialClosestHitPayload Payload,
 	float4 RandSample,
 	out float3 OutDirection,
-	out float3 OutThroughput,
+	out float3 OutWeight,
 	out float OutPdf)
 {
 	float PdfAdjust = 1.0f; 
@@ -126,7 +126,8 @@ void RadianceProbe_SampleMaterial(
 	float4 SampledValue = CosineSampleHemisphereConcentric(RandSample.xy);
 	OutDirection = TangentToWorld(SampledValue.xyz, N_World);
 
-	OutThroughput = 1.0f;
+	PdfAdjust = max(PdfAdjust, 0.00001f);
+	OutWeight = 1.0f / PdfAdjust;
 	OutPdf = SampledValue.w * PdfAdjust;
 }
 
@@ -191,14 +192,15 @@ float RadianceProbe_Pdf(float3 OutgoingDirection, FMaterialClosestHitPayload Pay
 #endif	
 #endif
 
+	PdfAdjust = max(PdfAdjust, 0.00001f);
 	float NoL = max(dot(N_World, OutgoingDirection), 0.0);
-	return (NoL / PI) * max(PdfAdjust, 0.00001f);
+	return (NoL / PI) * PdfAdjust;
 }
 
 void RadianceProbe_EvalMaterial(
 	float3 OutgoingDirection,
 	FMaterialClosestHitPayload Payload,
-	out float3 OutThroughput,
+	out float3 OutWeight,
 	out float OutPdf
 )
 {
@@ -263,8 +265,10 @@ void RadianceProbe_EvalMaterial(
 #endif
 	float NoL = max(dot(N_World, OutgoingDirection), 0.0);
 
-	OutPdf = NoL / PI * max(PdfAdjust, 0.00001f);
-	OutThroughput = 1.0f;
+	PdfAdjust = max(PdfAdjust, 0.00001f);
+
+	OutPdf = NoL / PI * PdfAdjust;
+	OutWeight = 1.0f / PdfAdjust;
 }
 
 #else
@@ -273,32 +277,37 @@ void RadianceProbe_SampleMaterial(
 	FMaterialClosestHitPayload Payload,
 	float4 RandSample,
 	out float3 OutDirection,
-	out float3 OutThroughput,
+	out float3 OutWeight,
 	out float OutPdf)
 {
 	float3 N_World = GetWorldNormal(Payload);
 
-	float4 SampledValue = UniformSampleHemisphere(RandSample.xy);
+	float4 SampledValue = CosineSampleHemisphere(RandSample.xy);
 	OutDirection = TangentToWorld(SampledValue.xyz, N_World);
 
-	OutThroughput = 1.0f;
-	OutPdf = 1.0f / (2 * PI);
+	OutWeight = 1;
+	OutPdf = SampledValue.w;
 }
 
 float RadianceProbe_Pdf(float3 OutgoingDirection, FMaterialClosestHitPayload Payload)
 {
-	return 1.0f / (2 * PI);
+	float3 N_World = GetWorldNormal(Payload);
+	float NoL = saturate(dot(N_World, OutgoingDirection));
+	return NoL / PI;
 }
 
 void RadianceProbe_EvalMaterial(
 	float3 OutgoingDirection,
 	FMaterialClosestHitPayload Payload,
-	out float3 OutThroughput,
+	out float3 OutWeight,
 	out float OutPdf
 )
 {
-	OutPdf = 1.0f / (2 * PI);
-	OutThroughput = 1.0f;
+	float3 N_World = GetWorldNormal(Payload);
+	float NoL = saturate(dot(N_World, OutgoingDirection));
+
+	OutWeight = GetDiffuseColor(Payload);
+	OutPdf = NoL / PI;
 }
 
 #endif
\ No newline at end of file
diff --git a/Engine/Shaders/Private/PathTracing/Material/PathTracingSpecularReflection.ush b/Engine/Shaders/Private/PathTracing/Material/PathTracingSpecularReflection.ush
--- a/Engine/Shaders/Private/PathTracing/Material/PathTracingSpecularReflection.ush
+++ b/Engine/Shaders/Private/PathTracing/Material/PathTracingSpecularReflection.ush
@@ -1,15 +1,18 @@
 // Copyright Epic Games, Inc. All Rights Reserved.
 
 /*=============================================================================================
-PathTracingSpecularReflection.usf: Lambertian BRDF sampling functions
+PathTracingSpecularReflection.usf: Singular Mirror BRDF sampling functions
 ===============================================================================================*/
+
 #pragma once
 
+// NOTE: this is no longer used, keeping around as reference
+
 void SpecularReflection_SampleMaterial(
 	float3 RayDirection,
 	FMaterialClosestHitPayload Payload,
 	out float3 OutDirection,
-	out float3 OutThroughput,
+	out float3 OutWeight,
 	out float OutPdf)
 {
 	float3 N_World = GetWorldNormal(Payload);
@@ -18,20 +21,14 @@ void SpecularReflection_SampleMaterial(
 	float3 V_World = normalize(-RayDirection);
 	OutDirection = reflect(-V_World, N_World);
 
-	float NoV = saturate(dot(V_World, N_World));
-	float VoH = NoV;
-	OutThroughput = F_Schlick(SpecularColor, VoH);
+	float NoV = abs(dot(V_World, N_World));
 
-	OutPdf = 1.0;
+	OutWeight = F_Schlick(SpecularColor, NoV);
+	OutPdf = POSITIVE_INFINITY;
 }
 
-float SpecularReflection_Pdf()
+void SpecularReflection_EvalMaterial(out float3 OutWeight, out float OutPdf)
 {
-	return 0.0;
+	OutWeight = 0;
+	OutPdf = 0;
 }
-
-void SpecularReflection_EvalMaterial(out float3 OutThroughput, out float OutPdf)
-{
-	OutThroughput = float3(0.0, 0.0, 0.0);
-	OutPdf = SpecularReflection_Pdf();
-}
\ No newline at end of file
diff --git a/Engine/Shaders/Private/PathTracing/Material/PathTracingSpecularTransmission.ush b/Engine/Shaders/Private/PathTracing/Material/PathTracingSpecularTransmission.ush
--- a/Engine/Shaders/Private/PathTracing/Material/PathTracingSpecularTransmission.ush
+++ b/Engine/Shaders/Private/PathTracing/Material/PathTracingSpecularTransmission.ush
@@ -11,7 +11,7 @@ void SpecularRefraction_SampleMaterial(
 	FMaterialClosestHitPayload Payload,
 	float4 RandSample,
 	out float3 OutDirection,
-	out float3 OutThroughput,
+	out float3 OutWeight,
 	out float OutPdf,
 	out float OutPositionBiasSign)
 {
@@ -23,43 +23,24 @@ void SpecularRefraction_SampleMaterial(
 	// #dxr_todo Determine if parameterization from Specular is still the proper path forward
 	//float Ior = Payload.Ior;
 	float Ior = DielectricF0ToIor(DielectricSpecularToF0(Payload.Specular));
-	float N1 = bIsEntering ? 1.0 : Ior;
-	float N2 = bIsEntering ? Ior : 1.0;
-	float Fr = bIsEntering ? FresnelDielectric(dot(-RayDirection, WorldNormal), N1, N2) : 0.0;
+	float Eta = bIsEntering ? Ior : 1.0 / Ior;
 
-	// Stochastically choose to reflect/refract based on Fresnel term
-	float ProbReflect = bIsEntering ? Fr : 0.0;
-	if (RandSample.x < ProbReflect)
+	// Throughput is 1 since we pick reflect/refract according to the fresnel term
+	OutWeight = 1.0;
+	OutPdf = POSITIVE_INFINITY;
+
+	if (SampleRefraction(RayDirection, WorldNormal, Eta, RandSample.x, OutDirection))
 	{
-		OutDirection = reflect(RayDirection, WorldNormal);
-		OutThroughput = Fr;
-		OutPdf = ProbReflect;
-		OutPositionBiasSign = 1.0;
+		OutPositionBiasSign = -1;
 	}
 	else
 	{
-		float Eta = N1 / N2;
-		OutThroughput = 1.0 - Fr;
-		OutPdf = 1.0 - ProbReflect;
-		OutDirection = refract(RayDirection, WorldNormal, Eta);
-		OutPositionBiasSign = -1.0;
-
-		// Total internal reflection
-		if (!any(OutDirection))
-		{
-			OutDirection = reflect(RayDirection, WorldNormal);
-			OutPositionBiasSign = 1.0;
-		}
+		OutPositionBiasSign = 1;
 	}
 } 
 
-float SpecularTransmission_Pdf()
-{
-	return 0.0;
-}
-
 void SpecularTransmission_EvalMaterial(out float3 OutThroughput, out float OutPdf)
 {
-	OutThroughput = float3(0.0, 0.0, 0.0);
-	OutPdf = SpecularTransmission_Pdf();
+	OutThroughput = 0;
+	OutPdf = 0;
 }
\ No newline at end of file
diff --git a/Engine/Shaders/Private/PathTracing/Material/PathTracingTestBrdfs.usf b/Engine/Shaders/Private/PathTracing/Material/PathTracingTestBrdfs.usf
--- a/Engine/Shaders/Private/PathTracing/Material/PathTracingTestBrdfs.usf
+++ b/Engine/Shaders/Private/PathTracing/Material/PathTracingTestBrdfs.usf
@@ -73,10 +73,10 @@ void TestBrdfIntegrityCS(
 	}
 	else if (BrdfType == 1)
 	{
-		Glossy_SampleMaterial(Wo, Payload, RandSample, Wi, SampledBRDF, SampledPDF);
-		Glossy_EvalMaterial(Wo, Wi, Payload, EvalBRDF, EvalPDF);
-		Glossy_EvalMaterial(Wi, Wo, Payload, EvalReversedBRDF, EvalPDF);
-		PDF = Glossy_Pdf(Wo, Wi, Payload);
+		DefaultLit_SampleMaterial(Wo, Payload, RandSample, Wi, SampledBRDF, SampledPDF);
+		DefaultLit_EvalMaterial(Wo, Wi, Payload, EvalBRDF, EvalPDF);
+		DefaultLit_EvalMaterial(Wi, Wo, Payload, EvalReversedBRDF, EvalPDF);
+		PDF = EvalPDF;
 	}
 	else
 	{
diff --git a/Engine/Shaders/Private/PathTracing/Material/PathTracingTestGenerateWiSamples.usf b/Engine/Shaders/Private/PathTracing/Material/PathTracingTestGenerateWiSamples.usf
--- a/Engine/Shaders/Private/PathTracing/Material/PathTracingTestGenerateWiSamples.usf
+++ b/Engine/Shaders/Private/PathTracing/Material/PathTracingTestGenerateWiSamples.usf
@@ -58,7 +58,7 @@ void TestGenerateWiSamplesCS(
 	}
 	else if (BrdfType == 1)
 	{
-		Glossy_SampleMaterial(Wo, Payload, RandSample, Wi, SampledBRDF, SampledPDF);
+		DefaultLit_SampleMaterial(Wo, Payload, RandSample, Wi, SampledBRDF, SampledPDF);
 	}
 	else
 	{
diff --git a/Engine/Shaders/Private/PathTracing/Material/PathTracingTestIntegrateHemispherePatch.usf b/Engine/Shaders/Private/PathTracing/Material/PathTracingTestIntegrateHemispherePatch.usf
--- a/Engine/Shaders/Private/PathTracing/Material/PathTracingTestIntegrateHemispherePatch.usf
+++ b/Engine/Shaders/Private/PathTracing/Material/PathTracingTestIntegrateHemispherePatch.usf
@@ -73,7 +73,8 @@ float IntegratePdfInSphericalPatch(float Theta0, float Phi0, float Theta1, float
 			}
 			else if (BrdfType == 1)
 			{
-				PDF = Glossy_Pdf(Wo, Wi, Payload);
+				float3 Weight = 0;
+				DefaultLit_EvalMaterial(Wo, Wi, Payload, Weight, PDF);
 			}
 			else
 			{
diff --git a/Engine/Shaders/Private/PathTracing/Material/PathTracingTestPdfIntegration.usf b/Engine/Shaders/Private/PathTracing/Material/PathTracingTestPdfIntegration.usf
--- a/Engine/Shaders/Private/PathTracing/Material/PathTracingTestPdfIntegration.usf
+++ b/Engine/Shaders/Private/PathTracing/Material/PathTracingTestPdfIntegration.usf
@@ -62,7 +62,8 @@ void TestPDFIntegratesToOneCS(
 	}
 	else if (BrdfType == 1)
 	{
-		PDF = Glossy_Pdf(Wo, Wi, Payload);
+		float3 Weight = 0;
+		DefaultLit_EvalMaterial(Wo, Wi, Payload, Weight, PDF);
 	}
 	else
 	{
diff --git a/Engine/Shaders/Private/PathTracing/Material/PathTracingTwoSidedFoliage.ush b/Engine/Shaders/Private/PathTracing/Material/PathTracingTwoSidedFoliage.ush
--- a/Engine/Shaders/Private/PathTracing/Material/PathTracingTwoSidedFoliage.ush
+++ b/Engine/Shaders/Private/PathTracing/Material/PathTracingTwoSidedFoliage.ush
@@ -8,56 +8,15 @@
 
 #include "PathTracingDefaultLit.ush"
 
-float3 TwoSidedFoliage_CalcTransmission(
-	float3 IncomingDirection,
-	FMaterialClosestHitPayload Payload
-)
-{
-	// Transmission throughput is governed by the amount of energy left after being reflected
-	float NoV = saturate(dot(-GetWorldNormal(Payload), IncomingDirection));
-	float3 Fresnel = F_Schlick(GetSpecularColor(Payload), NoV);
-	float3 DiffuseColor = GetDiffuseColor(Payload);
-	float3 ReflectedEnergy = Fresnel + DiffuseColor;
-	float3 Transmission = saturate(1.0 - ReflectedEnergy);
-	return Transmission;
-}
-
 float CalcTransmissionSelectionPdf(
 	float3 IncomingDirection,
 	FMaterialClosestHitPayload Payload
 )
 {
-	// TODO: Combined energy is reflected diffuse + transmitted diffuse..
-	float SubsurfaceLuminance = Luminance(GetSubsurfaceColor(Payload));
-	float DiffuseLuminance = Luminance(GetDiffuseColor(Payload));
-	//float DiffuseLuminance = Luminance(TwoSidedFoliage_CalcTransmission(IncomingDirection, Payload));
-
-	return SubsurfaceLuminance / (SubsurfaceLuminance + DiffuseLuminance);
-}
-
-float TwoSidedFoliage_PdfMaterial(
-	float3 IncomingDirection,
-	float3 OutgoingDirection,
-	FMaterialClosestHitPayload Payload
-)
-{
-	float TransmissionSelectionPdf = CalcTransmissionSelectionPdf(IncomingDirection, Payload);
-	float3 N = GetWorldNormal(Payload);
-	float3 L = OutgoingDirection;
-	float NoL = dot(N, L);
-	if (NoL < 0.0)
-	{
-		// Diffuse transmission
-		// TODO: Provide rough transmission as in ShadingModels.ush (Roughness = 0.6)
-		float LambertPdf = -NoL / PI;
-		return TransmissionSelectionPdf * LambertPdf;
-	}
-	else
-	{
-		// Default Lit
-		float Ior1 = 1.0;
-		return (1.0 - TransmissionSelectionPdf) * DefaultLit_PdfMaterial(IncomingDirection, OutgoingDirection, Payload, Ior1);
-	}
+	// pick back lobe proportionally to its contribution, and leave the rest to the front side
+	// TODO: estimate this better based on lobe colors ...
+	float3 Subsurf = GetSubsurfaceColor(Payload);
+	return max(Subsurf.x, max(Subsurf.y, Subsurf.z));
 }
 
 void TwoSidedFoliage_SampleMaterial(
@@ -65,7 +24,7 @@ void TwoSidedFoliage_SampleMaterial(
 	FMaterialClosestHitPayload Payload,
 	float4 RandSample,
 	out float3 OutDirection,
-	out float3 OutThroughput,
+	out float3 OutWeight,
 	out float OutPdf,
 	out float OutPositionBiasSign
 )
@@ -76,20 +35,13 @@ void TwoSidedFoliage_SampleMaterial(
 		RandSample.x /= TransmissionSelectionPdf;
 
 		float3 N = GetWorldNormal(Payload);
-		float V = -RayDirection;
-		float3 Transmission = TwoSidedFoliage_CalcTransmission(RayDirection, Payload);
 		float3 SubsurfaceColor = GetSubsurfaceColor(Payload);
 
-		// Diffuse transmission
-		// TODO: Provide rough transmission as in ShadingModels.ush (Roughness = 0.6)
 		float4 SampledValue = CosineSampleHemisphere(RandSample.yz);
 		OutDirection = TangentToWorld(SampledValue.xyz, -N);
 
-		float LambertPdf = SampledValue.w; // SampledValue.w = NoL / PI
-		float3 LambertThroughput = SubsurfaceColor * LambertPdf;
-		OutThroughput = Transmission * LambertThroughput;
-		OutPdf = TransmissionSelectionPdf * LambertPdf;
-
+		OutWeight = SubsurfaceColor / TransmissionSelectionPdf;
+		OutPdf = TransmissionSelectionPdf * SampledValue.w;
 		OutPositionBiasSign = -1.0;
 	}
 	else
@@ -97,8 +49,8 @@ void TwoSidedFoliage_SampleMaterial(
 		RandSample.x -= TransmissionSelectionPdf;
 		RandSample.x /= (1.0 - TransmissionSelectionPdf);
 
-		float Ior1 = 1.0;
-		DefaultLit_SampleMaterial(RayDirection, Payload, Ior1, RandSample, OutDirection, OutThroughput, OutPdf);
+		DefaultLit_SampleMaterial(RayDirection, Payload, RandSample, OutDirection, OutWeight, OutPdf);
+		OutWeight /= 1.0 - TransmissionSelectionPdf;
 		OutPdf *= 1.0 - TransmissionSelectionPdf;
 		OutPositionBiasSign = 1.0;
 	}
@@ -108,7 +60,7 @@ void TwoSidedFoliage_EvalMaterial(
 	float3 IncomingDirection,
 	float3 OutgoingDirection,
 	FMaterialClosestHitPayload Payload,
-	out float3 OutThroughput,
+	out float3 OutWeight,
 	out float OutPdf
 )
 {
@@ -118,22 +70,18 @@ void TwoSidedFoliage_EvalMaterial(
 	float NoL = dot(N, OutgoingDirection);
 	if (NoL < 0.0)
 	{
-		// Diffuse transmission
-		// TODO: Provide rough transmission as in ShadingModels.ush (Roughness = 0.6)
-		NoL = -NoL;
+		// Diffuse transmission (the implementation in TwoSidedBxDF does not appear to be a real BxDF)
 
-		float3 Transmission = TwoSidedFoliage_CalcTransmission(IncomingDirection, Payload);
 		float3 SubsurfaceColor = GetSubsurfaceColor(Payload);
 
-		float LambertPdf = NoL / PI;
-		float3 LambertThroughput = SubsurfaceColor * LambertPdf;
-		OutThroughput = Transmission * LambertThroughput;
+		float LambertPdf = -NoL / PI;
+		OutWeight = SubsurfaceColor / TransmissionSelectionPdf;
 		OutPdf = TransmissionSelectionPdf * LambertPdf;
 	}
 	else
 	{
-		float Ior1 = 1.0;
-		DefaultLit_EvalMaterial(IncomingDirection, OutgoingDirection, Payload, Ior1, OutThroughput, OutPdf);
+		DefaultLit_EvalMaterial(IncomingDirection, OutgoingDirection, Payload, OutWeight, OutPdf);
+		OutWeight /= 1.0 - TransmissionSelectionPdf;
 		OutPdf *= 1.0 - TransmissionSelectionPdf;
 	}
 }
\ No newline at end of file
diff --git a/Engine/Shaders/Private/PathTracing/PathTracing.usf b/Engine/Shaders/Private/PathTracing/PathTracing.usf
--- a/Engine/Shaders/Private/PathTracing/PathTracing.usf
+++ b/Engine/Shaders/Private/PathTracing/PathTracing.usf
@@ -27,13 +27,25 @@ RWTexture2D<uint>   PixelPositionRT;
 RWTexture2D<uint>	RayCountPerPixelRT;
 RaytracingAccelerationStructure TLAS;
 
+void AccumulateRadiance(inout float3 TotalRadiance, float3 PathRadiance)
+{
+	if (PathTracingData.MaxPathIntensity > 0)
+	{
+		// User asked for path contributions to be clamped to reduce fireflies.
+		// Depending on how aggressive this value is, the image could be quite biased
+		TotalRadiance += min(PathRadiance, PathTracingData.MaxPathIntensity);
+	}
+	else
+	{
+		// Just average values directly
+		TotalRadiance += PathRadiance;
+	}
+}
+
 RAY_TRACING_ENTRY_RAYGEN(PathTracingMainRG)
 {
-	uint2 LaunchIndex = DispatchRaysIndex().xy + View.ViewRectMin + PathTracingData.TileOffset;
+	uint2 LaunchIndex = DispatchRaysIndex().xy + View.ViewRectMin.xy + PathTracingData.TileOffset.xy;
 	
-	bool IsUnidirectionalEnabled = true;
-	bool IsNextEventEstimationEnabled = true;
-
 	RadianceRT[LaunchIndex] = 0;
 	SampleCountRT[LaunchIndex] = 0;
 	PixelPositionRT[LaunchIndex] = 0;
@@ -46,237 +58,254 @@ RAY_TRACING_ENTRY_RAYGEN(PathTracingMainRG)
 	RandomSequence_Initialize(RandSequence, LinearIndex, AdaptiveSamplingData.TemporalSeed);
 	RandSequence.Type = AdaptiveSamplingData.RandomSequence;
 	uint SampleIndex = 0;
-	uint SampleCount = 1;
 
-	for (uint i = 0; i < SampleCount; ++i) 
+	float3 Radiance = 0;
+	uint RayCounter = 0;
+
+	// Initialize ray and payload
+	RayDesc Ray;
+	uint2 Pixel;
+	float EmitterPdf;
+	SampleEmitter(LaunchIndex, RandSequence, SampleIndex, 0, 0, Ray, Pixel, EmitterPdf);
+	if (EmitterPdf <= 0.0)
 	{
-		float3 Irradiance = 0;
-		uint RayCounter = 0;
-
-		// Initialize ray and payload
-		RayDesc Ray;
-		uint2 Pixel;
-		float EmitterPdf;
-		SampleEmitter(LaunchIndex, RandSequence, SampleIndex, 0, 0, Ray, Pixel, EmitterPdf);
-		if (EmitterPdf <= 0.0)
-		{
-			continue;
-		}
+		return;
+	}
 
-		// This array will hold a CDF for light picking
-		// Seed the array with a uniform CDF at first so that we always have a valid CDF
-		float LightPickingCdf[RAY_TRACING_LIGHT_COUNT_MAXIMUM];
+	// This array will hold a CDF for light picking
+	// Seed the array with a uniform CDF at first so that we always have a valid CDF
+	float LightPickingCdf[RAY_TRACING_LIGHT_COUNT_MAXIMUM];
 
-		InitLightPickingCdf(LightPickingCdf);
+	InitLightPickingCdf(LightPickingCdf);
 
-		Ray.Direction = normalize(Ray.Direction);
+	Ray.Direction = normalize(Ray.Direction);
 
-		FRayCone RayCone = (FRayCone)0;
-		RayCone.SpreadAngle = View.EyeToPixelSpreadAngle;
+	FRayCone RayCone = (FRayCone)0;
+	RayCone.SpreadAngle = View.EyeToPixelSpreadAngle;
 
-		// Cast ray
-		const uint InitialRayFlags = 0;
-		const uint InitialInstanceInclusionMask = RAY_TRACING_MASK_ALL;
-		const bool bInitialDisableSkyLightContribution = false;
-		const bool bIgnoreTranslucentMaterials = false;
-		const uint2 RenderTargetPos = LaunchIndex.xy + View.ViewRectMin.xy;
+	// Cast ray
+	const uint InitialRayFlags = 0;
+	const bool bInitialDisableSkyLightContribution = false;
+	const bool bIgnoreTranslucentMaterials = false;
 
-		// path state variables (these cary information between bounces)
-		float3 PathThroughput = 1.0;
+	// path state variables (these cary information between bounces)
+	float3 PathThroughput = 1.0;
 
-		// We don't want the camera ray to have any MIS contribution against direct light sampling so we pretend that it came from
-		// a singular Material scattering event.
-		float PrevMaterialPdf = POSITIVE_INFINITY;
+	// We don't want the camera ray to have any MIS contribution against direct light sampling so we pretend that it came from
+	// a singular Material scattering event.
+	float PrevMaterialPdf = POSITIVE_INFINITY;
 
-		// 0: only Material sampling
-		// 1: only Light sampling
-		// 2: both Material and Light
-		const int MISMode = AdaptiveSamplingData.MISMode;
+	// 0: only Material sampling
+	// 1: only Light sampling
+	// 2: both Material and Light
+	const int MISMode = PathTracingData.MISMode;
 
-		for (int Bounce = 0; Bounce <= PathTracingData.MaxBounces; Bounce++)
+	for (int Bounce = 0; Bounce <= PathTracingData.MaxBounces; Bounce++)
+	{
+		FMaterialClosestHitPayload Payload = TraceMaterialRay(
+			TLAS,
+			InitialRayFlags,
+			RAY_TRACING_MASK_ALL,
+			Ray,
+			RayCone,
+			LaunchIndex,
+			bInitialDisableSkyLightContribution,
+			bIgnoreTranslucentMaterials);
+		RayCounter++;
+
+		if (MISMode != 1 || Bounce == 0)
 		{
-			FMaterialClosestHitPayload Payload = TraceMaterialRay(
-				TLAS,
-				InitialRayFlags,
-				InitialInstanceInclusionMask,
-				Ray,
-				RayCone,
-				RenderTargetPos,
-				bInitialDisableSkyLightContribution,
-				bIgnoreTranslucentMaterials);
-			RayCounter++;
-
-			if (MISMode != 1 || Bounce == 0)
+			float HitT = Payload.IsMiss() ? Ray.TMax : Payload.HitT;
+			// Loop over lights to capture their contribution
+			// #dxr_todo: if we have lots of lights, having some hierarchical structure would be better ....
+			float PreviousCdfValue = 0.0;
+			for (uint LightId = 0; LightId < SceneLightsData.Count; ++LightId)
 			{
-				float HitT = Payload.IsMiss() ? Ray.TMax : Payload.HitT;
-				// Loop over lights to capture their contribution
-				// #dxr_todo: if we have lots of lights, having some hierarchical structure would be better ....
-				float PreviousCdfValue = 0.0;
-				for (uint LightId = 0; LightId < SceneLightsData.Count; ++LightId)
-				{
-					// Camera rays should only do the rest of the loop if we have visible lights OR if
-					// we have a skydome (which are always visible)
-					if (Bounce == 0 && AdaptiveSamplingData.VisibleLights == 0 && !IsEnvironmentLight(LightId))
-						continue;
-
-					float4 LightResult = TraceLight(Ray, LightId, HitT);
-					float3 LightRadiance = LightResult.xyz;
-					float LightPdf = LightResult.w;
-
-					float LightPickPdf = LightPickingCdf[LightId] - PreviousCdfValue;
-					PreviousCdfValue = LightPickingCdf[LightId];
-					Irradiance += PathThroughput * LightRadiance * MISWeightRobust(PrevMaterialPdf, LightPdf * LightPickPdf);
-				}
-			}
+				// Camera rays should only do the rest of the loop if we have visible lights OR if
+				// we have a skydome (which are always visible)
+				if (Bounce == 0 && PathTracingData.VisibleLights == 0 && !IsEnvironmentLight(LightId))
+					continue;
 
-			if (Payload.IsMiss())
-			{
-				// Ray didn't hit any real geometry, so nothing left to do
-				break;
+				float4 LightResult = TraceLight(Ray, LightId, HitT);
+				float3 LightRadiance = LightResult.xyz;
+				float LightPdf = LightResult.w;
+
+				float LightPickPdf = LightPickingCdf[LightId] - PreviousCdfValue;
+				PreviousCdfValue = LightPickingCdf[LightId];
+
+				AccumulateRadiance(Radiance, PathThroughput * LightRadiance * MISWeightRobust(PrevMaterialPdf, LightPdf * LightPickPdf));
 			}
+		}
 
-			// add in surface emission
-			Irradiance += PathThroughput * Payload.Radiance;
+		if (Payload.IsMiss())
+		{
+			// Ray didn't hit any real geometry, so nothing left to do
+			break;
+		}
 
-			// Update origin
-			Ray.Origin = Payload.WorldPos;
+		// add in surface emission
+		AccumulateRadiance(Radiance, PathThroughput * Payload.Radiance);
 
-			// Make sure the surface normal points toward the viewer
-			if (dot(Ray.Direction, Payload.WorldNormal) > 0.0)
-			{
-				Payload.WorldNormal = -Payload.WorldNormal;
-			}
+		// Update origin
+		Ray.Origin = Payload.WorldPos;
 
-			// Choose a random number for both Light sampling and BxDF sampling
-			float4 RandSample = RandomSequence_GenerateSample4D(RandSequence, SampleIndex);
+		// Make sure the surface normal points toward the viewer
+		if (dot(Ray.Direction, Payload.WorldNormal) > 0.0)
+		{
+			Payload.WorldNormal = -Payload.WorldNormal;
+		}
 
-			// Does this material require NEE?
-			bool bIsNeeValid = IsNeeValidMaterial(Payload);
+		// Choose a random number for both Light sampling and BxDF sampling
+		float4 RandSample = RandomSequence_GenerateSample4D(RandSequence, SampleIndex);
 
-			// If we are using Light sampling and we are not on the last bounce ....
-			if (MISMode != 0 && Bounce < PathTracingData.MaxBounces && bIsNeeValid && SceneLightsData.Count > 0)
-			{
-				// Choose a light and sample it
+		// Does this material require NEE?
+		bool bIsNeeValid = IsNeeValidMaterial(Payload);
 
-				float3 WorldPos = Payload.WorldPos;
-				float3 WorldNormal = Payload.WorldNormal;
-				uint PrimitiveLightingChannelMask = Payload.PrimitiveLightingChannelMask;
+		// If we are using Light sampling and we are not on the last bounce ....
+		if (MISMode != 0 && Bounce < PathTracingData.MaxBounces && bIsNeeValid && SceneLightsData.Count > 0)
+		{
+			// Choose a light and sample it
+
+			float3 WorldPos = Payload.WorldPos;
+			float3 WorldNormal = Payload.WorldNormal;
+			uint PrimitiveLightingChannelMask = Payload.PrimitiveLightingChannelMask;
+
+			bool IsTransmissiveMaterial = ENABLE_TRANSMISSION && IsMaterialTransmissive(Payload);
 
-				if (InitLightPickingCdf(WorldPos, WorldNormal, PrimitiveLightingChannelMask, LightPickingCdf))
+			if (InitLightPickingCdf(WorldPos, WorldNormal, PrimitiveLightingChannelMask, IsTransmissiveMaterial, LightPickingCdf))
+			{
+				// init worked
+				int LightId;
+				float LightPickPdf = 0;
+
+				SelectLight(RandSample.x, LightPickingCdf, LightId, LightPickPdf);
+
+				RayDesc LightRay;
+				LightRay.Origin = WorldPos;
+				LightRay.TMin = 0;
+				float3 RadianceOverPdf;
+				float LightPdf;
+				SampleLight(LightId, RandSample, WorldPos, WorldNormal,
+					LightRay.Direction,
+					LightRay.TMax,
+					RadianceOverPdf,
+					LightPdf);
+				RadianceOverPdf /= LightPickPdf;
+				LightPdf *= LightPickPdf;
+
+				if (LightPdf > 0)
 				{
-					// init worked
-					int LightId;
-					float LightPickPdf = 0;
-
-					SelectLight(RandSample.x, LightPickingCdf, LightId, LightPickPdf);
-
-					RayDesc LightRay;
-					LightRay.Origin = WorldPos;
-					LightRay.TMin = 0;
-					float3 RadianceOverPdf;
-					float LightPdf;
-					SampleLight(LightId, RandSample, WorldPos, WorldNormal,
-						LightRay.Direction,
-						LightRay.TMax,
-						RadianceOverPdf,
-						LightPdf);
-					RadianceOverPdf /= LightPickPdf;
-					LightPdf *= LightPickPdf;
-
-					if (LightPdf > 0)
+					float SignedPositionBias = 1.0;
+					if (ENABLE_TRANSMISSION && IsMaterialTransmissive(Payload))
 					{
-						float SignedPositionBias = 1.0;
+						// for transmissive materials, bias the position to the other side of the surface if the light is coming from behind
 						float NoL = dot(Payload.WorldNormal, LightRay.Direction);
-						if (ENABLE_TRANSMISSION && HasTransmission(LightId) && IsMaterialTransmissive(Payload) && NoL < 0.0)
-						{
-							SignedPositionBias = -1.0;
-						}
-						ApplyPositionBias(LightRay, SignedPositionBias * Payload.WorldNormal, AdaptiveSamplingData.MaxNormalBias);
-
-
-						const uint RayFlags = RAY_FLAG_ACCEPT_FIRST_HIT_AND_END_SEARCH | RAY_FLAG_SKIP_CLOSEST_HIT_SHADER;
-						const uint InstanceInclusionMask = RAY_TRACING_MASK_SHADOW;
-
-						FMinimalPayload NeePayload = TraceVisibilityRay(
-							TLAS,
-							RayFlags,
-							InstanceInclusionMask,
-							RenderTargetPos,
-							LightRay);
-						RayCounter++;
-
-						// #dxr_todo: Is it cheaper to fire the ray first? Or eval the material first?
-						if (NeePayload.IsMiss())
-						{
-							// Evaluate material
-							float Ior1 = 1.0;
-							float3 MaterialEval;
-							float MaterialPdf = 0.0;
-							EvalMaterial(Ray.Direction, LightRay.Direction, Payload, Ior1, false, MaterialEval, MaterialPdf);
-
-							float MISWeight = MISMode == 2 ? MISWeightRobust(LightPdf, MaterialPdf) : 1.0;
-
-							// Record the contribution
-							Irradiance += MISWeight * PathThroughput * RadianceOverPdf * MaterialEval;
-						}
+						SignedPositionBias = sign(NoL);
 					}
-				}
-			}
+					ApplyPositionBias(LightRay, SignedPositionBias * Payload.WorldNormal, PathTracingData.MaxNormalBias);
 
-			if (Bounce == PathTracingData.MaxBounces)
-			{
-				// we've arrived at the last bounce, don't bother trying to sample a direction
-				break;
-			}
 
-			// Sample material
-			float Ior1 = 1.0;
-			float3 Direction;
-			float3 MaterialThroughput;
-			float SignedPositionBias;
-			float MaterialSamplePdf;
-			SampleMaterial(Ray.Direction, Payload, Ior1, RandSample, false, Direction, MaterialThroughput, MaterialSamplePdf, SignedPositionBias);
-			if (!(MaterialSamplePdf > 0))
-			{
-				// No valid direction -- we are done
-				break;
-			}
+					const uint RayFlags = RAY_FLAG_ACCEPT_FIRST_HIT_AND_END_SEARCH | RAY_FLAG_SKIP_CLOSEST_HIT_SHADER;
+					const uint InstanceInclusionMask = RAY_TRACING_MASK_ALL;
 
-			// Update ray
-			Ray.Direction = Direction;
-			Ray.TMin = 0;
-			ApplyPositionBias(Ray, SignedPositionBias * Payload.WorldNormal, AdaptiveSamplingData.MaxNormalBias);
-			float3 NextPathThroughput = PathThroughput * MaterialThroughput / MaterialSamplePdf;
-			if (!any(NextPathThroughput > 0))
-			{
-				// no energy left in this path
-				break;
-			}
+					FPackedMaterialClosestHitPayload PackedPayload = (FPackedMaterialClosestHitPayload)0;
 
-			// Russian roulette:
-			float ContinuationProb = saturate(max(NextPathThroughput.x, max(NextPathThroughput.y, NextPathThroughput.z)) / max(PathThroughput.x, max(PathThroughput.y, PathThroughput.z)));
-			if (ContinuationProb < 1)
-			{
-				// If there is some chance we should terminate the ray, draw an extra random value
-				if (RandomSequence_GenerateSample1D(RandSequence, SampleIndex) >= ContinuationProb)
-				{
-					// stochastically terminate the path
-					break;
+					TraceVisibilityRayPacked(PackedPayload, TLAS, RayFlags, InstanceInclusionMask, LaunchIndex, LightRay);
+
+					RayCounter++;
+
+					// #dxr_todo: Is it cheaper to fire the ray first? Or eval the material first?
+					if (PackedPayload.IsMiss())
+					{
+						// Evaluate material
+						float3 MaterialWeight;
+						float MaterialPdf = 0.0;
+						EvalMaterial(Ray.Direction, LightRay.Direction, Payload, false, MaterialWeight, MaterialPdf);
+
+						float3 MaterialEval = MaterialWeight * MaterialPdf;
+
+						float MISWeight = MISMode == 2 ? MISWeightRobust(LightPdf, MaterialPdf) : 1.0;
+
+						// Record the contribution
+						AccumulateRadiance(Radiance, MISWeight * PathThroughput * RadianceOverPdf * MaterialEval);
+					}
 				}
-				PathThroughput = NextPathThroughput / ContinuationProb;
 			}
+		}
+
+		if (Bounce == PathTracingData.MaxBounces)
+		{
+			// we've arrived at the last bounce, don't bother trying to sample a direction
+			break;
+		}
 
-			if (MISMode == 2)
+		// Sample material
+		float3 Direction;
+		float3 MaterialWeight;
+		float SignedPositionBias;
+		float MaterialSamplePdf;
+		SampleMaterial(Ray.Direction, Payload, RandSample, false, Direction, MaterialWeight, MaterialSamplePdf, SignedPositionBias);
+
+		if (MaterialSamplePdf < 0 || asuint(MaterialSamplePdf) > 0x7F800000)
+		{
+			// Pdf became invalid (either negative or NaN)
+			Radiance = float3(1, 0, 1);
+			break;
+		}
+
+		if (!(MaterialSamplePdf > 0))
+		{
+			// No valid direction -- we are done
+			break;
+		}
+
+		float3 NextPathThroughput = PathThroughput * MaterialWeight;
+		if (!any(NextPathThroughput > 0))
+		{
+			// no energy left in this path
+			break;
+		}
+
+		// Russian roulette:
+		//   The probability of keeping the path should be roughly proportional to the weight at the current shade point,
+		//  but just using MaterialWeight would miss out on cases where the path throughput changes color (like in a cornell
+		//  box when bouncing between walls of different colors). So use the ratio of the brightest color channel in the
+		//  previous and next throughput.
+		//   The second tweak is to add a sqrt() around the probability to soften the termination probability (paths will last
+		//  a little longer). This allows paths to go a bit deeper than the naive heuristic while still allowing them to terminate
+		//  early. This makes RR effective from the very first bounce without needing to delay it.
+		float ContinuationProb = sqrt(saturate(max(NextPathThroughput.x, max(NextPathThroughput.y, NextPathThroughput.z)) / max(PathThroughput.x, max(PathThroughput.y, PathThroughput.z))));
+		if (ContinuationProb < 1)
+		{
+			// If there is some chance we should terminate the ray, draw an extra random value
+			float RussianRouletteRand = RandSample.w; // SampleMaterial does not use this value at the moment
+			//RussianRouletteRand = RandomSequence_GenerateSample1D(RandSequence, SampleIndex);
+			if (RussianRouletteRand >= ContinuationProb)
 			{
-				// If we are doing MIS, the Material ray on the next loop through will need to know 
-				// #dxr_todo: the SamplePdf is not using the right measure for singular cases, this should be handled at the Material API level
-				PrevMaterialPdf = bIsNeeValid ? MaterialSamplePdf : 1e27;
+				// stochastically terminate the path
+				break;
 			}
+			PathThroughput = NextPathThroughput / ContinuationProb;
+		}
+		else
+		{
+			PathThroughput = NextPathThroughput;
 		}
 
-		RadianceRT[RenderTargetPos].rgb += Irradiance;
-		SampleCountRT[RenderTargetPos] += 1;
-		PixelPositionRT[RenderTargetPos] = Pixel.y * BufferSize.x + Pixel.x;
-		RayCountPerPixelRT[RenderTargetPos] += RayCounter;
+		if (MISMode == 2)
+		{
+			// If we are doing MIS, the Material ray on the next loop through will need to know the Pdf
+			PrevMaterialPdf = MaterialSamplePdf;
+		}
+
+		// Update ray and loop back around to trace again
+		Ray.Direction = Direction;
+		Ray.TMin = 0;
+		ApplyPositionBias(Ray, SignedPositionBias* Payload.WorldNormal, PathTracingData.MaxNormalBias);
 	}
+
+	RadianceRT[LaunchIndex].rgb += Radiance;
+	SampleCountRT[LaunchIndex] += 1;
+	PixelPositionRT[LaunchIndex] = Pixel.y * BufferSize.x + Pixel.x;
+	RayCountPerPixelRT[LaunchIndex] += RayCounter;
 }
diff --git a/Engine/Shaders/Private/RayTracing/RayTracingGlobalIlluminationRGS.usf b/Engine/Shaders/Private/RayTracing/RayTracingGlobalIlluminationRGS.usf
--- a/Engine/Shaders/Private/RayTracing/RayTracingGlobalIlluminationRGS.usf
+++ b/Engine/Shaders/Private/RayTracing/RayTracingGlobalIlluminationRGS.usf
@@ -272,7 +272,9 @@ RAY_TRACING_ENTRY_RAYGEN(GlobalIlluminationRGS)
 						float3 MaterialThroughput;
 						float MaterialEvalPdf = 0.0;
 #if USE_PATHTRACING_MATERIALS
-						EvalMaterial(Ray.Direction, LightRay.Direction, Payload, MaterialThroughput, MaterialEvalPdf);
+						float3 MaterialWeight = 0;
+						EvalMaterial(Ray.Direction, LightRay.Direction, Payload, MaterialWeight, MaterialEvalPdf);
+						MaterialThroughput = MaterialWeight * MaterialEvalPdf;
 #else
 						half3 N = Payload.WorldNormal;
 						half3 V = -Ray.Direction;
@@ -312,7 +314,9 @@ RAY_TRACING_ENTRY_RAYGEN(GlobalIlluminationRGS)
 				RandSample.xyz = QuasiRandom_Sample3(QuasiRandom);
 				RandSample.w = QuasiRandom_Sample3(QuasiRandom).x;
 				float PositionBiasSign;
-				SampleMaterial(Ray.Direction, Payload, RandSample, Direction, Throughput, MaterialPdf, PositionBiasSign);
+				float3 Weight = 0;
+				SampleMaterial(Ray.Direction, Payload, RandSample, Direction, Weight, MaterialPdf, PositionBiasSign);
+				Throughput = Weight * MaterialPdf;
 #else
 				float3 RandSample3 = QuasiRandom_Sample3(QuasiRandom);
 				float3 RayOrigin = Ray.Origin;
diff --git a/Engine/Source/Runtime/Renderer/Private/PathTracing.cpp b/Engine/Source/Runtime/Renderer/Private/PathTracing.cpp
--- a/Engine/Source/Runtime/Renderer/Private/PathTracing.cpp
+++ b/Engine/Source/Runtime/Renderer/Private/PathTracing.cpp
@@ -46,6 +46,11 @@ TAutoConsoleVariable<int32> CVarPathTracingVisibleLights(
 	TEXT("1: Make lights visible to camera\n")
 );
 
+TAutoConsoleVariable<int32> CVarPathTracingMaxPathIntensity(
+	TEXT("r.PathTracing.MaxPathIntensity"),
+	-1,
+	TEXT("When positive, light paths greater that this amount are clamped to prevent fireflies (default = -1 (off))")
+);
 
 TAutoConsoleVariable<int32> CVarPathTracingFrameIndependentTemporalSeed(
 	TEXT("r.PathTracing.FrameIndependentTemporalSeed"),
@@ -111,20 +116,6 @@ TAutoConsoleVariable<int32> CVarPathTracingWiperMode(
 	ECVF_RenderThreadSafe 
 );
 
-TAutoConsoleVariable<int32> CVarPathTracingRussianRouletteStartingBounce(
-	TEXT("r.PathTracing.RussianRouletteStartingBounce"),
-	5,
-	TEXT("Determines the starting bounce condition for Russian Roulette"),
-	ECVF_RenderThreadSafe
-);
-
-TAutoConsoleVariable<float> CVarPathTracingFireFlyRejectionTreshold(
-	TEXT("r.PathTracing.FireFlyRejectionThreshold"),
-	1024.0,
-	TEXT("Defines the maximum energy off a sample, after which it is categorized as a firefly"),
-	ECVF_RenderThreadSafe
-);
-
 IMPLEMENT_GLOBAL_SHADER_PARAMETER_STRUCT(FPathTracingData, "PathTracingData");
 IMPLEMENT_GLOBAL_SHADER_PARAMETER_STRUCT(FPathTracingLightData, "SceneLightsData");
 IMPLEMENT_GLOBAL_SHADER_PARAMETER_STRUCT(FPathTracingAdaptiveSamplingData, "AdaptiveSamplingData");
@@ -196,14 +187,43 @@ public:
 
 			int32 PathTracingMaxBounces = GPathTracingMaxBounces > -1 ? GPathTracingMaxBounces : View.FinalPostProcessSettings.PathTracingMaxBounces;
 			PathTracingData.MaxBounces = PathTracingMaxBounces;
+			PathTracingData.MaxNormalBias = GetRaytracingMaxNormalBias();
+			PathTracingData.MISMode = CVarPathTracingMISMode.GetValueOnRenderThread();
+			PathTracingData.VisibleLights = CVarPathTracingVisibleLights.GetValueOnRenderThread();
+			PathTracingData.MaxPathIntensity = CVarPathTracingMaxPathIntensity.GetValueOnRenderThread();
+
+			// If any of the parameters above changed since last time -- reset the accumulation
 			static uint32 PrevMaxBounces = PathTracingMaxBounces;
 			if (PathTracingData.MaxBounces != PrevMaxBounces)
 			{
 				Scene->bPathTracingNeedsInvalidation = true;
 				PrevMaxBounces = PathTracingData.MaxBounces;
 			}
-			PathTracingData.RussianRouletteStartingBounce = CVarPathTracingRussianRouletteStartingBounce.GetValueOnRenderThread();
-			PathTracingData.FireFlyRejectionThreshold = CVarPathTracingFireFlyRejectionTreshold.GetValueOnRenderThread();
+
+			// Changing MIS mode requires starting over
+			static uint32 PreviousMISMode = PathTracingData.MISMode;
+			if (PreviousMISMode != PathTracingData.MISMode)
+			{
+				Scene->bPathTracingNeedsInvalidation = true;
+				PreviousMISMode = PathTracingData.MISMode;
+			}
+
+			// Changing VisibleLights requires starting over
+			static uint32 PreviousVisibleLights = PathTracingData.VisibleLights;
+			if (PreviousVisibleLights != PathTracingData.VisibleLights)
+			{
+				Scene->bPathTracingNeedsInvalidation = true;
+				PreviousVisibleLights = PathTracingData.VisibleLights;
+			}
+
+			// Changing MaxPathIntensity requires starting over
+			static float PreviousMaxPathIntensity = PathTracingData.MaxPathIntensity;
+			if (PreviousMaxPathIntensity != PathTracingData.MaxPathIntensity)
+			{
+				Scene->bPathTracingNeedsInvalidation = true;
+				PreviousMaxPathIntensity = PathTracingData.MaxPathIntensity;
+			}
+
 			PathTracingData.TileOffset = TileOffset;
 
 			FUniformBufferRHIRef PathTracingDataUniformBuffer = RHICreateUniformBuffer(&PathTracingData, FPathTracingData::StaticStructMetadata.GetLayout(), EUniformBufferUsage::UniformBuffer_SingleDraw);
@@ -212,8 +232,9 @@ public:
 
 		// Sky light
 		FSkyLightData SkyLightData;
+		bool IsSkyLightValid = false;
 		{
-			SetupSkyLightParameters(*Scene, &SkyLightData);
+			IsSkyLightValid = SetupSkyLightParameters(*Scene, &SkyLightData);
 
 			FUniformBufferRHIRef SkyLightUniformBuffer = RHICreateUniformBuffer(&SkyLightData, FSkyLightData::StaticStructMetadata.GetLayout(), EUniformBufferUsage::UniformBuffer_SingleDraw);
 			GlobalResources.Set(SkyLightParameters, SkyLightUniformBuffer);
@@ -224,15 +245,17 @@ public:
 			FPathTracingLightData LightData;
 			LightData.Count = 0;
 
-			// Prepend SkyLight to light buffer
-			// WARNING: Until ray payload encodes Light data buffer, the execution depends on this ordering!
-			uint32 SkyLightIndex = 0;
-			uint8 SkyLightLightingChannelMask = 0xFF;
-			LightData.Type[SkyLightIndex] = 0;
-			LightData.Color[SkyLightIndex] = FVector(SkyLightData.Color);
-			LightData.Flags[SkyLightIndex] = SkyLightData.bTransmission & 0x01;
-			LightData.Flags[SkyLightIndex] |= (SkyLightLightingChannelMask & 0x7) << 1;
-			LightData.Count++;
+			// Prepend SkyLight to light buffer since it is not part of the regular light list
+			if (IsSkyLightValid)
+			{
+				uint32 SkyLightIndex = 0;
+				uint8 SkyLightLightingChannelMask = 0xFF;
+				LightData.Type[SkyLightIndex] = 0;
+				LightData.Color[SkyLightIndex] = FVector(SkyLightData.Color);
+				LightData.Flags[SkyLightIndex] = SkyLightData.bTransmission & 0x01;
+				LightData.Flags[SkyLightIndex] |= (SkyLightLightingChannelMask & 0x7) << 1;
+				LightData.Count++;
+			}
 
 			for (auto Light : Lights)
 			{
@@ -250,7 +273,6 @@ public:
 				ELightComponentType LightComponentType = (ELightComponentType)Light.LightSceneInfo->Proxy->GetLightType();
 				switch (LightComponentType)
 				{
-					// TODO: LightType_Spot
 					case LightType_Directional:
 					{
 						LightData.Type[LightData.Count] = 2;
@@ -279,23 +301,26 @@ public:
 						LightData.Type[LightData.Count] = 4;
 						LightData.Position[LightData.Count] = LightParameters.Position;
 						LightData.Normal[LightData.Count] = -LightParameters.Direction;
-						// #dxr_todo: UE-72556  define these differences from Lit..
 						LightData.Color[LightData.Count] = LightParameters.Color;
 						LightData.Dimensions[LightData.Count] = FVector(LightParameters.SpotAngles, LightParameters.SourceRadius);
 						LightData.Attenuation[LightData.Count] = 1.0 / LightParameters.InvRadius;
 						break;
 					}
 					case LightType_Point:
-					default:
 					{
 						LightData.Type[LightData.Count] = 1;
 						LightData.Position[LightData.Count] = LightParameters.Position;
-						// #dxr_todo: UE-72556  define these differences from Lit..
 						LightData.Color[LightData.Count] = LightParameters.Color;
 						LightData.Dimensions[LightData.Count] = FVector(0.0, 0.0, LightParameters.SourceRadius);
 						LightData.Attenuation[LightData.Count] = 1.0 / LightParameters.InvRadius;
 						break;
 					}
+					default:
+					{
+						// Just in case someone adds a new light type one day ...
+						checkNoEntry();
+						break;
+					}
 				};
 
 				LightData.Count++;
@@ -318,26 +343,15 @@ public:
 			}
 
 			FPathTracingAdaptiveSamplingData AdaptiveSamplingData;
-			AdaptiveSamplingData.MaxNormalBias = GetRaytracingMaxNormalBias();
 			AdaptiveSamplingData.UseAdaptiveSampling = CVarPathTracingAdaptiveSampling.GetValueOnRenderThread();
 			AdaptiveSamplingData.RandomSequence = CVarPathTracingRandomSequence.GetValueOnRenderThread();
-			AdaptiveSamplingData.MISMode = CVarPathTracingMISMode.GetValueOnRenderThread();
-			AdaptiveSamplingData.VisibleLights = CVarPathTracingVisibleLights.GetValueOnRenderThread();
 
-			// Changing MIS mode requires starting over
-			static uint32 PreviousMISMode = AdaptiveSamplingData.MISMode;
-			if (PreviousMISMode != AdaptiveSamplingData.MISMode)
-			{
-				Scene->bPathTracingNeedsInvalidation = true;
-				PreviousMISMode = AdaptiveSamplingData.MISMode;
-			}
-
-			// Changing VisibleLights requires starting over
-			static uint32 PreviousVisibleLights = AdaptiveSamplingData.VisibleLights;
-			if (PreviousVisibleLights != AdaptiveSamplingData.VisibleLights)
+			// Changing Adaptive sampling mode requires starting over
+			static uint32 PreviousUseAdaptiveSampling = AdaptiveSamplingData.UseAdaptiveSampling;
+			if (PreviousUseAdaptiveSampling != AdaptiveSamplingData.UseAdaptiveSampling)
 			{
 				Scene->bPathTracingNeedsInvalidation = true;
-				PreviousVisibleLights = AdaptiveSamplingData.VisibleLights;
+				PreviousUseAdaptiveSampling = AdaptiveSamplingData.UseAdaptiveSampling;
 			}
 
 			if (VarianceMipTree.NumBytes > 0)
diff --git a/Engine/Source/Runtime/Renderer/Private/PathTracingUniformBuffers.h b/Engine/Source/Runtime/Renderer/Private/PathTracingUniformBuffers.h
--- a/Engine/Source/Runtime/Renderer/Private/PathTracingUniformBuffers.h
+++ b/Engine/Source/Runtime/Renderer/Private/PathTracingUniformBuffers.h
@@ -11,8 +11,10 @@
 
 BEGIN_GLOBAL_SHADER_PARAMETER_STRUCT(FPathTracingData, )
 	SHADER_PARAMETER(uint32, MaxBounces)
-	SHADER_PARAMETER(uint32, RussianRouletteStartingBounce)
-	SHADER_PARAMETER(float, FireFlyRejectionThreshold)
+	SHADER_PARAMETER(uint32, MISMode)
+	SHADER_PARAMETER(uint32, VisibleLights)
+	SHADER_PARAMETER(float, MaxPathIntensity)
+	SHADER_PARAMETER(float, MaxNormalBias)
 	SHADER_PARAMETER(FIntVector, TileOffset)
 END_GLOBAL_SHADER_PARAMETER_STRUCT()
 
@@ -47,9 +49,6 @@ BEGIN_GLOBAL_SHADER_PARAMETER_STRUCT(FPathTracingAdaptiveSamplingData, )
 	SHADER_PARAMETER(uint32, MinimumSamplesPerPixel)
 	SHADER_PARAMETER(uint32, Iteration)
 	SHADER_PARAMETER(uint32, TemporalSeed)
-	SHADER_PARAMETER(uint32, MISMode)
-	SHADER_PARAMETER(uint32, VisibleLights)
-	SHADER_PARAMETER(float, MaxNormalBias)
 	SHADER_PARAMETER(FIntVector, VarianceDimensions)
 	SHADER_PARAMETER_SRV(Buffer<float>, VarianceMipTree)
 END_GLOBAL_SHADER_PARAMETER_STRUCT()
diff --git a/Engine/Source/Runtime/Renderer/Private/RayTracing/RayTracingSkyLight.h b/Engine/Source/Runtime/Renderer/Private/RayTracing/RayTracingSkyLight.h
--- a/Engine/Source/Runtime/Renderer/Private/RayTracing/RayTracingSkyLight.h
+++ b/Engine/Source/Runtime/Renderer/Private/RayTracing/RayTracingSkyLight.h
@@ -59,7 +59,7 @@ END_SHADER_PARAMETER_STRUCT()
 
 int32 GetRayTracingSkyLightDecoupleSampleGenerationCVarValue();
 
-extern RENDERER_API void SetupSkyLightParameters(const FScene& Scene, FSkyLightData* SkyLight);
+extern RENDERER_API bool SetupSkyLightParameters(const FScene& Scene, FSkyLightData* SkyLight);
 void SetupSkyLightQuasiRandomParameters(const FScene& Scene, const FViewInfo& View, FIntVector& OutBlueNoiseDimensions, FSkyLightQuasiRandomData* OutSkyLightQuasiRandomData);
 void SetupSkyLightVisibilityRaysParameters(FRDGBuilder& GraphBuilder, const FViewInfo& View, FSkyLightVisibilityRaysData* OutSkyLightVisibilityRaysData);
 
diff --git a/Engine/Source/Runtime/Renderer/Private/RayTracing/RaytracingSkylight.cpp b/Engine/Source/Runtime/Renderer/Private/RayTracing/RaytracingSkylight.cpp
--- a/Engine/Source/Runtime/Renderer/Private/RayTracing/RaytracingSkylight.cpp
+++ b/Engine/Source/Runtime/Renderer/Private/RayTracing/RaytracingSkylight.cpp
@@ -154,7 +154,7 @@ struct FSkyLightVisibilityRays
 	FVector4 DirectionAndPdf;
 };
 
-void SetupSkyLightParameters(
+bool SetupSkyLightParameters(
 	const FScene& Scene,
 	FSkyLightData* SkyLightData)
 {
@@ -195,6 +195,9 @@ void SetupSkyLightParameters(
 		SkyLightData->MipTreePdfPosZ = Scene.SkyLight->ImportanceSamplingData->MipTreePdfPosZ.SRV;
 		SkyLightData->MipTreePdfNegZ = Scene.SkyLight->ImportanceSamplingData->MipTreePdfNegZ.SRV;
 		SkyLightData->SolidAnglePdf = Scene.SkyLight->ImportanceSamplingData->SolidAnglePdf.SRV;
+
+		// we prepared a valid SkyLight
+		return true;
 	}
 	else
 	{
@@ -223,6 +226,9 @@ void SetupSkyLightParameters(
 		SkyLightData->MipTreePdfPosZ = GBlackTextureWithSRV->ShaderResourceViewRHI;
 		SkyLightData->MipTreePdfNegZ = GBlackTextureWithSRV->ShaderResourceViewRHI;
 		SkyLightData->SolidAnglePdf = GBlackTextureWithSRV->ShaderResourceViewRHI;
+
+		// we returned a dummy SkyLight
+		return false;
 	}
 }
 
