Add Owen-Sobol sampling sequence to the path tracer

This is optional for now, but converges faster than the current default, particularly on direct lighting

diff --git a/Engine/Shaders/Private/PathTracing/Utilities/PathTracingRandomSequence.ush b/Engine/Shaders/Private/PathTracing/Utilities/PathTracingRandomSequence.ush
--- a/Engine/Shaders/Private/PathTracing/Utilities/PathTracingRandomSequence.ush
+++ b/Engine/Shaders/Private/PathTracing/Utilities/PathTracingRandomSequence.ush
@@ -154,6 +154,102 @@ float Halton(uint Index, uint Base)
 	return r;
 }
 
+
+uint FastOwenScrambling(uint Index, uint Seed) {
+	// Laine and Karras / Stratified Sampling for Stochastic Transparency / EGSR 2011
+
+	// The operations below will mix bits toward the left, so temporarily reverse the order
+	// TODO: can fold this in with sobol logic to skip some reversals
+	Index = ReverseBits32(Index);
+	
+	Index += Seed; // randomize the index by our seed (pushes bits toward the left)
+
+	// A few different options below for doing the scrambling step:
+	// I'll stick to the simplest one until we have a good variance/time measurement system
+#if 1
+	// NOTE: Original publication has two extra xor-muls but they don't actually improve quality or convergence
+	Index ^= Index * 0x6C50B47Cu;
+	Index ^= Index * 0xB82F1E52u;
+#elif 0
+	// Better looking spectra on average - though still not perfect
+	Index ^= (Index + Seed) * 0x6C50B47Cu;
+	Index ^= (Index + Seed) * 0xB82F1E52u;
+#else
+	// This method achieves the closest results to real Owen-Scrambling
+	// but (probably?) costs a bit more perf
+	Index ^= Index * (Seed & ~1u);
+	Seed = IntegerHash(Seed);
+	Index ^= Index * (Seed & ~1u);
+#endif
+
+	// Undo the reverse so that we get left-to-right scrambling
+	// thereby emulating owen-scrambling
+	return ReverseBits32(Index);
+}
+
+// 32-bit Sobol matrices for dimension 1,2,3 from:
+// S. Joe and F. Y. Kuo, Constructing Sobol sequences with better two-dimensional projections, SIAM J. Sci. Comput. 30, 2635-2654 (2008)
+//    https://web.maths.unsw.edu.au/~fkuo/sobol/
+// NOTE: we don't bother storing dimension 0 since it is just a bit reversal
+// NOTE2: dimension 1 also follows a pattern that can be done procedurally but probably not worth the extra code?
+// NOTE3: if we did all our owen scrambling logic with 16 bits, we could avoid
+static const uint3 SobolMatrices[] = {
+	uint3(0x80000000, 0x80000000, 0x80000000),
+	uint3(0xc0000000, 0xc0000000, 0xc0000000),
+	uint3(0xa0000000, 0x60000000, 0x20000000),
+	uint3(0xf0000000, 0x90000000, 0x50000000),
+	uint3(0x88000000, 0xe8000000, 0xf8000000),
+	uint3(0xcc000000, 0x5c000000, 0x74000000),
+	uint3(0xaa000000, 0x8e000000, 0xa2000000),
+	uint3(0xff000000, 0xc5000000, 0x93000000),
+	uint3(0x80800000, 0x68800000, 0xd8800000),
+	uint3(0xc0c00000, 0x9cc00000, 0x25400000),
+	uint3(0xa0a00000, 0xee600000, 0x59e00000),
+	uint3(0xf0f00000, 0x55900000, 0xe6d00000),
+	uint3(0x88880000, 0x80680000, 0x78080000),
+	uint3(0xcccc0000, 0xc09c0000, 0xb40c0000),
+	uint3(0xaaaa0000, 0x60ee0000, 0x82020000),
+	uint3(0xffff0000, 0x90550000, 0xc3050000),
+	uint3(0x80008000, 0xe8808000, 0x208f8000),
+	uint3(0xc000c000, 0x5cc0c000, 0x51474000),
+	uint3(0xa000a000, 0x8e606000, 0xfbea2000),
+	uint3(0xf000f000, 0xc5909000, 0x75d93000),
+	uint3(0x88008800, 0x6868e800, 0xa0858800),
+	uint3(0xcc00cc00, 0x9c9c5c00, 0x914e5400),
+	uint3(0xaa00aa00, 0xeeee8e00, 0xdbe79e00),
+	uint3(0xff00ff00, 0x5555c500, 0x25db6d00),
+	uint3(0x80808080, 0x8000e880, 0x58800080),
+	uint3(0xc0c0c0c0, 0xc0005cc0, 0xe54000c0),
+	uint3(0xa0a0a0a0, 0x60008e60, 0x79e00020),
+	uint3(0xf0f0f0f0, 0x9000c590, 0xb6d00050),
+	uint3(0x88888888, 0xe8006868, 0x800800f8),
+	uint3(0xcccccccc, 0x5c009c9c, 0xc00c0074),
+	uint3(0xaaaaaaaa, 0x8e00eeee, 0x200200a2),
+	uint3(0xffffffff, 0xc5005555, 0x50050093)
+};
+
+float4 SobolSampler(uint SampleIndex, inout uint Seed) {
+	// We need a random sequence that we can advance quickly (5 steps per 4 samples)
+	// constant from: https://www.pcg-random.org/posts/does-it-beat-the-minimal-standard.html
+	const uint LCG_C = 1019135901;
+	// first scramble the index to decorelate from other 4-tuples
+	uint SobolIndex = FastOwenScrambling(SampleIndex, Seed = Seed * LCG_C + LCG_C);
+	// now get Sobol' point from this index
+	uint4 Result = uint4(ReverseBits32(SobolIndex), 0, 0, 0);
+	UNROLL // does this unrolling help?
+	for (int i = 0; i < 32; SobolIndex >>= 1, i++)
+		Result.yzw ^= (SobolIndex & 1) ? SobolMatrices[i] : 0;
+
+	// finally scramble the points to avoid structured artifacts
+	Result.x = FastOwenScrambling(Result.x, Seed = Seed * LCG_C + LCG_C);
+	Result.y = FastOwenScrambling(Result.y, Seed = Seed * LCG_C + LCG_C);
+	Result.z = FastOwenScrambling(Result.z, Seed = Seed * LCG_C + LCG_C);
+	Result.w = FastOwenScrambling(Result.w, Seed = Seed * LCG_C + LCG_C);
+	// output as float in [0,1) taking care not to skew the distribution
+	// due to the non-uniform spacing of floats in this range
+	return (Result >> 8) * 5.96046447754e-08; // * 2^-24
+}
+
 void RandomSequence_Initialize(inout RandomSequence RandSequence, uint PositionSeed, uint TimeSeed)
 {
 	RandSequence.Type = 0;
@@ -165,6 +261,9 @@ void RandomSequence_Initialize(inout RandomSequence RandSequence, uint PositionS
 
 	// Halton
 	RandSequence.HaltonDimensionIndex = 0;
+
+	// Halton and Sobol both want a hashed position seed - so precompute it here
+	RandSequence.PositionSeed = IntegerHash(PositionSeed);
 }
 
 bool RandomSequence_IsLCG(in RandomSequence RandSequence)
@@ -182,6 +281,12 @@ bool RandomSequence_IsScrambledHalton(in RandomSequence RandSequence)
 	return RandSequence.Type == 2;
 }
 
+bool RandomSequence_IsOwenSobol(in RandomSequence RandSequence)
+{
+	return RandSequence.Type == 3;
+}
+
+
 float RandomSequence_generateSample(inout RandomSequence RandSequence, in uint SampleIndex)
 {
 	// dxr_todo: Make choice of random a shader permutation
@@ -195,7 +300,7 @@ float RandomSequence_generateSample(inout RandomSequence RandSequence, in uint S
 	}
 	else if (RandomSequence_IsScrambledHalton(RandSequence))
 	{
-		uint HaltonSampleIndex = IntegerHash(RandSequence.PositionSeed) + RandSequence.TimeSeed;
+		uint HaltonSampleIndex = RandSequence.PositionSeed + RandSequence.TimeSeed;
 
 		float RandomSample = Halton(HaltonSampleIndex, Prime512(RandSequence.HaltonDimensionIndex));
 		RandSequence.HaltonDimensionIndex += 1;
@@ -210,7 +315,16 @@ float RandomSequence_generateSample(inout RandomSequence RandSequence, in uint S
 float RandomSequence_GenerateSample1D(inout RandomSequence RandSequence, inout uint SampleIndex)
 {
 	float Sample;
-	Sample = RandomSequence_generateSample(RandSequence, SampleIndex);
+	if (RandomSequence_IsOwenSobol(RandSequence))
+	{
+		// rely on compiler to optimize out dead code for now
+		// #dxr_todo: benchmark and specialize if needed
+		Sample = SobolSampler(RandSequence.TimeSeed, RandSequence.PositionSeed).x;
+	}
+	else
+	{
+		Sample = RandomSequence_generateSample(RandSequence, SampleIndex);
+	}
 	SampleIndex += 1;
 	return Sample;
 }
@@ -218,8 +332,17 @@ float RandomSequence_GenerateSample1D(inout RandomSequence RandSequence, inout u
 float2 RandomSequence_GenerateSample2D(inout RandomSequence RandSequence, inout uint SampleIndex)
 {
 	float2 Sample;
-	Sample.x = RandomSequence_generateSample(RandSequence, SampleIndex);
-	Sample.y = RandomSequence_generateSample(RandSequence, SampleIndex + 1);
+	if (RandomSequence_IsOwenSobol(RandSequence))
+	{
+		// rely on compiler to optimize out dead code for now
+		// #dxr_todo: benchmark and specialize if needed
+		Sample = SobolSampler(RandSequence.TimeSeed, RandSequence.PositionSeed).xy;
+	}
+	else
+	{
+		Sample.x = RandomSequence_generateSample(RandSequence, SampleIndex);
+		Sample.y = RandomSequence_generateSample(RandSequence, SampleIndex + 1);
+	}
 	SampleIndex += 2;
 	return Sample;
 }
@@ -227,9 +350,18 @@ float2 RandomSequence_GenerateSample2D(inout RandomSequence RandSequence, inout
 float3 RandomSequence_GenerateSample3D(inout RandomSequence RandSequence, inout uint SampleIndex)
 {
 	float3 Sample;
-	Sample.x = RandomSequence_generateSample(RandSequence, SampleIndex);
-	Sample.y = RandomSequence_generateSample(RandSequence, SampleIndex + 1);
-	Sample.z = RandomSequence_generateSample(RandSequence, SampleIndex + 2);
+	if (RandomSequence_IsOwenSobol(RandSequence))
+	{
+		// rely on compiler to optimize out dead code for now
+		// #dxr_todo: benchmark and specialize if needed
+		Sample = SobolSampler(RandSequence.TimeSeed, RandSequence.PositionSeed).xyz;
+	}
+	else
+	{
+		Sample.x = RandomSequence_generateSample(RandSequence, SampleIndex);
+		Sample.y = RandomSequence_generateSample(RandSequence, SampleIndex + 1);
+		Sample.z = RandomSequence_generateSample(RandSequence, SampleIndex + 2);
+	}
 	SampleIndex += 3;
 	return Sample;
 }
@@ -237,10 +369,17 @@ float3 RandomSequence_GenerateSample3D(inout RandomSequence RandSequence, inout
 float4 RandomSequence_GenerateSample4D(inout RandomSequence RandSequence, inout uint SampleIndex)
 {
 	float4 Sample;
-	Sample.x = RandomSequence_generateSample(RandSequence, SampleIndex);
-	Sample.y = RandomSequence_generateSample(RandSequence, SampleIndex + 1);
-	Sample.z = RandomSequence_generateSample(RandSequence, SampleIndex + 2);
-	Sample.w = RandomSequence_generateSample(RandSequence, SampleIndex + 3);
+	if (RandomSequence_IsOwenSobol(RandSequence))
+	{
+		Sample = SobolSampler(RandSequence.TimeSeed, RandSequence.PositionSeed);
+	}
+	else
+	{
+		Sample.x = RandomSequence_generateSample(RandSequence, SampleIndex);
+		Sample.y = RandomSequence_generateSample(RandSequence, SampleIndex + 1);
+		Sample.z = RandomSequence_generateSample(RandSequence, SampleIndex + 2);
+		Sample.w = RandomSequence_generateSample(RandSequence, SampleIndex + 3);
+	}
 	SampleIndex += 4;
 	return Sample;
 }
