Fix for hair in Lumen.

diff --git a/Engine/Shaders/Private/ShadingModelsSampling.ush b/Engine/Shaders/Private/ShadingModelsSampling.ush
--- a/Engine/Shaders/Private/ShadingModelsSampling.ush
+++ b/Engine/Shaders/Private/ShadingModelsSampling.ush
@@ -114,47 +114,15 @@ FBxDFSample SampleHairBxDF(uint TermMask, FGBufferData GBuffer, float3 V, float4
 	float3 N = GBuffer.WorldNormal;
 	float4 L = UniformSampleSphere(E.xy);
 	float NoL = dot(N, L.xyz);
-
-	const float Falloff = 1;
-
-#if 0
-	FAreaLight AreaLight;
-	AreaLight.SphereSinAlpha = 0;
-	AreaLight.SphereSinAlphaSoft = 0;
-	AreaLight.LineCosSubtended = 1;
-	AreaLight.FalloffColor = 1;
-	AreaLight.Rect = (FRect)0;
-	AreaLight.bIsRect = false;
-	AreaLight.Texture = InitRectTexture(LTCAmpTexture); // Dummy
-
-	FShadowTerms Shadow = { 1, 1, 1, InitHairTransmittanceData() };
-
-	//FDirectLighting Lighting = EvaluateBxDF(GBuffer, N, V, L.xyz, NoL, Shadow);
-	//FDirectLighting Lighting = IntegrateBxDF(GBuffer, N, V, L.xyz, Falloff, NoL, AreaLight, Shadow)
-	FDirectLighting Lighting = HairBxDF(GBuffer, N, V, L.xyz, Falloff, NoL, AreaLight, Shadow);
-#elif 1 // ReflectionEnvifornementPixelShader.usf
 	FHairTransmittanceData TransmittanceData = InitHairTransmittanceData();
-	bool bEvalMultiScatter = true;
-
-	FDirectLighting Lighting;
-	//TODO Lumen Team
-	Lighting.Diffuse = float3(1.0f, 0.0f, 0.0f); //HairShading(GBuffer, L.xyz, V, N, 1, TransmittanceData, 0, 0.2, uint2(0, 0), bEvalMultiScatter);
-	Lighting.Transmission = 0;
-	Lighting.Specular = 0;
-#else
-	float Shadow = 0.0;
-	float Backlit = 0.0;
-	float Area = 0.0;
-	uint2 Random = 0;
-	bool bEvalMultiScatter = false;
-	FHairTransmittanceData HairTransmittance = InitHairTransmittanceData();
 
 	FDirectLighting Lighting;
-	//TODO Lumen Team
-	Lighting.Diffuse = float3(1.0f, 0.0f, 0.0f); //HairShading(GBuffer, L.xyz, V, N, Shadow, HairTransmittance, Backlit, Area, Random, bEvalMultiScatter);
+	const float OpaqueVisibility = 1;	// No occlusion...
+	const float Area = 0;				// This is used for faking area light sources by increasing the roughness of the surface. Disabled = 0.
+	const float Backlit = 1; 			// This is used for suppressing the R & TT terms when when the lighting direction comes from behind. Disabled = 1.
+	Lighting.Diffuse = HairShading(GBuffer, L.xyz, V, N, OpaqueVisibility, TransmittanceData, Backlit, Area, uint2(0, 0));
 	Lighting.Transmission = 0;
 	Lighting.Specular = 0;
-#endif
 
 	FBxDFSample BxDFSample;
 	BxDFSample.L = L.xyz;
