Replace if-else chains in path tracer with switch statements

diff --git a/Engine/Shaders/Private/PathTracing/Light/PathTracingLightCommon.ush b/Engine/Shaders/Private/PathTracing/Light/PathTracingLightCommon.ush
--- a/Engine/Shaders/Private/PathTracing/Light/PathTracingLightCommon.ush
+++ b/Engine/Shaders/Private/PathTracing/Light/PathTracingLightCommon.ush
@@ -6,6 +6,13 @@ LightCommon.usf: Common utility functions for light sampling
 
 #pragma once
 
+// These ids must be kept in sync with the logic in PathTracer.cpp
+#define PATHTRACING_LIGHT_SKY				0
+#define PATHTRACING_LIGHT_POINT				1
+#define PATHTRACING_LIGHT_DIRECTIONAL		2
+#define PATHTRACING_LIGHT_RECT				3
+#define PATHTRACING_LIGHT_SPOT				4
+
 int GetLightId(FMaterialClosestHitPayload RayHitInfo)
 {
 	return -1;
@@ -13,27 +20,27 @@ int GetLightId(FMaterialClosestHitPayload RayHitInfo)
 
 bool IsEnvironmentLight(int LightId)
 {
-	return SceneLightsData.Type[LightId] == 0;
+	return SceneLightsData.Type[LightId] == PATHTRACING_LIGHT_SKY;
 }
 
 bool IsPointLight(int LightId)
 {
-	return SceneLightsData.Type[LightId] == 1;
+	return SceneLightsData.Type[LightId] == PATHTRACING_LIGHT_POINT;
 }
 
 bool IsDirectionalLight(int LightId)
 {
-	return SceneLightsData.Type[LightId] == 2;
+	return SceneLightsData.Type[LightId] == PATHTRACING_LIGHT_DIRECTIONAL;
 }
 
 bool IsRectLight(int LightId)
 {
-	return SceneLightsData.Type[LightId] == 3;
+	return SceneLightsData.Type[LightId] == PATHTRACING_LIGHT_RECT;
 }
 
 bool IsSpotLight(int LightId)
 {
-	return SceneLightsData.Type[LightId] == 4;
+	return SceneLightsData.Type[LightId] == PATHTRACING_LIGHT_SPOT;
 }
 
 // A light is a physical light if it can be intersected by a ray.
diff --git a/Engine/Shaders/Private/PathTracing/Light/PathTracingLightSampling.ush b/Engine/Shaders/Private/PathTracing/Light/PathTracingLightSampling.ush
--- a/Engine/Shaders/Private/PathTracing/Light/PathTracingLightSampling.ush
+++ b/Engine/Shaders/Private/PathTracing/Light/PathTracingLightSampling.ush
@@ -22,28 +22,15 @@ LightSampling.usf: Light sampling for next-event estimation
 // Returns float4 packed with Radiance and Pdf
 float4 TraceLight(RayDesc Ray, int LightId, float HitT)
 {
-	if (IsPointLight(LightId))
+	switch (SceneLightsData.Type[LightId])
 	{
-		return PointLight_TraceLight(Ray, LightId, HitT);
+		case PATHTRACING_LIGHT_SKY:			return SkyLight_TraceLight(Ray, LightId, HitT);
+		case PATHTRACING_LIGHT_POINT:		return PointLight_TraceLight(Ray, LightId, HitT);
+		case PATHTRACING_LIGHT_DIRECTIONAL: return DirectionalLight_TraceLight(Ray, LightId, HitT);
+		case PATHTRACING_LIGHT_RECT:		return RectLight_TraceLight(Ray, LightId, HitT);
+		case PATHTRACING_LIGHT_SPOT:		return SpotLight_TraceLight(Ray, LightId, HitT);
+		default: return 0; // unknown light type?
 	}
-	else if (IsSpotLight(LightId))
-	{
-		return SpotLight_TraceLight(Ray, LightId, HitT);
-	}
-	else if (IsDirectionalLight(LightId))
-	{
-		return DirectionalLight_TraceLight(Ray, LightId, HitT);
-	}
-	else if (IsEnvironmentLight(LightId))
-	{
-		return SkyLight_TraceLight(Ray, LightId, HitT);
-	}
-	else if (IsRectLight(LightId))
-	{
-		return RectLight_TraceLight(Ray, LightId, HitT);
-	}
-	// unknown light type?
-	return 0.0;
 }
 
 // Choose a point on a light as seen from the given shading point
@@ -57,33 +44,30 @@ void SampleLight(
 	out float3 OutRadianceOverPdf,
 	out float OutPdf)
 {
-	if (IsPointLight(LightId))
-	{
-		PointLight_SampleLight(LightId, RandSample, WorldPos, WorldNormal, OutDirection, OutDistance, OutRadianceOverPdf, OutPdf);
-	}
-	else if (IsSpotLight(LightId))
-	{
-		SpotLight_SampleLight(LightId, RandSample, WorldPos, WorldNormal, OutDirection, OutDistance, OutRadianceOverPdf, OutPdf);
-	}
-	else if (IsDirectionalLight(LightId))
-	{
-		DirectionalLight_SampleLight(LightId, RandSample, WorldPos, WorldNormal, OutDirection, OutDistance, OutRadianceOverPdf, OutPdf);
-	}
-	else if (IsEnvironmentLight(LightId))
-	{
-		SkyLight_SampleLight(LightId, RandSample, WorldPos, WorldNormal, OutDirection, OutDistance, OutRadianceOverPdf, OutPdf);
-	}
-	else if (IsRectLight(LightId))
-	{
-		RectLight_SampleLight(LightId, RandSample, WorldPos, WorldNormal, OutDirection, OutDistance, OutRadianceOverPdf, OutPdf);
-	}
-	else
-	{
-		// unknown light type?
-		OutDirection = 0;
-		OutDistance = 0;
-		OutRadianceOverPdf = 0;
-		OutPdf = 0;
+	switch (SceneLightsData.Type[LightId])
+	{
+		case PATHTRACING_LIGHT_SKY:
+			SkyLight_SampleLight(LightId, RandSample, WorldPos, WorldNormal, OutDirection, OutDistance, OutRadianceOverPdf, OutPdf);
+			return;
+		case PATHTRACING_LIGHT_POINT:
+			PointLight_SampleLight(LightId, RandSample, WorldPos, WorldNormal, OutDirection, OutDistance, OutRadianceOverPdf, OutPdf);
+			return;
+		case PATHTRACING_LIGHT_DIRECTIONAL:
+			DirectionalLight_SampleLight(LightId, RandSample, WorldPos, WorldNormal, OutDirection, OutDistance, OutRadianceOverPdf, OutPdf);
+			return;
+		case PATHTRACING_LIGHT_RECT:
+			RectLight_SampleLight(LightId, RandSample, WorldPos, WorldNormal, OutDirection, OutDistance, OutRadianceOverPdf, OutPdf);
+			return;
+		case PATHTRACING_LIGHT_SPOT:
+			SpotLight_SampleLight(LightId, RandSample, WorldPos, WorldNormal, OutDirection, OutDistance, OutRadianceOverPdf, OutPdf);
+			return;
+		default:
+			// unknown light type?
+			OutDirection = 0;
+			OutDistance = 0;
+			OutRadianceOverPdf = 0;
+			OutPdf = 0;
+			return;
 	}
 }
 
@@ -91,25 +75,25 @@ void GenerateLightRay(RayDesc Ray, int LightId, float3 LightUV, inout RayDesc Ou
 {
 	OutLightRay = (RayDesc)0; // #dxr_todo: Return false in the case of no matching light type and handle case from caller rather than blindly using a degenerate ray
 
-	if (IsPointLight(LightId))
-	{
-		PointLight_GenerateLightRay(Ray, LightId, LightUV, OutLightRay);
-	}
-	else if (IsSpotLight(LightId))
-	{
-		SpotLight_GenerateLightRay(Ray, LightId, LightUV, OutLightRay);
-	}
-	else if (IsDirectionalLight(LightId))
-	{
-		DirectionalLight_GenerateLightRay(Ray, LightId, LightUV, OutLightRay);
-	}
-	else if (IsEnvironmentLight(LightId))
-	{
-		SkyLight_GenerateLightRay(Ray, LightId, LightUV, OutLightRay);
-	}
-	else if (IsRectLight(LightId))
-	{
-		RectLight_GenerateLightRay(Ray, LightId, LightUV, OutLightRay);
+	switch (SceneLightsData.Type[LightId])
+	{
+		case PATHTRACING_LIGHT_SKY:
+			SkyLight_GenerateLightRay(Ray, LightId, LightUV, OutLightRay);
+			return;
+		case PATHTRACING_LIGHT_POINT:
+			PointLight_GenerateLightRay(Ray, LightId, LightUV, OutLightRay);
+			return;
+		case PATHTRACING_LIGHT_DIRECTIONAL:
+			DirectionalLight_GenerateLightRay(Ray, LightId, LightUV, OutLightRay);
+			return;
+		case PATHTRACING_LIGHT_RECT:
+			RectLight_GenerateLightRay(Ray, LightId, LightUV, OutLightRay);
+			return;
+		case PATHTRACING_LIGHT_SPOT:
+			SpotLight_GenerateLightRay(Ray, LightId, LightUV, OutLightRay);
+			return;
+		default:
+			return;
 	}
 }
 
@@ -126,25 +110,25 @@ void EstimateLight(
 	OutIrradiance = 0.0f;
 	if ((GetLightingChannelMask(LightId) & PrimitiveLightingChannelMask) == 0) return;
 
-	if (IsPointLight(LightId))
-	{
-		PointLight_EstimateLight(LightId, WorldPos, WorldNormal, IsTransmissiveMaterial, OutIrradiance);
-	}
-	else if (IsSpotLight(LightId))
-	{
-		SpotLight_EstimateLight(LightId, WorldPos, WorldNormal, IsTransmissiveMaterial, OutIrradiance);
-	}
-	else if (IsEnvironmentLight(LightId))
-	{
-		SkyLight_EstimateLight(LightId, WorldPos, WorldNormal, IsTransmissiveMaterial, OutIrradiance);
-	}
-	else if (IsDirectionalLight(LightId))
-	{
-		DirectionalLight_EstimateLight(LightId, WorldPos, WorldNormal, IsTransmissiveMaterial, OutIrradiance);
-	}
-	else if (IsRectLight(LightId))
-	{
-		RectLight_EstimateLight(LightId, WorldPos, WorldNormal, IsTransmissiveMaterial, OutIrradiance);
+	switch (SceneLightsData.Type[LightId])
+	{
+		case PATHTRACING_LIGHT_SKY:
+			SkyLight_EstimateLight(LightId, WorldPos, WorldNormal, IsTransmissiveMaterial, OutIrradiance);
+			return;
+		case PATHTRACING_LIGHT_POINT:
+			PointLight_EstimateLight(LightId, WorldPos, WorldNormal, IsTransmissiveMaterial, OutIrradiance);
+			return;
+		case PATHTRACING_LIGHT_DIRECTIONAL:
+			DirectionalLight_EstimateLight(LightId, WorldPos, WorldNormal, IsTransmissiveMaterial, OutIrradiance);
+			return;
+		case PATHTRACING_LIGHT_RECT:
+			RectLight_EstimateLight(LightId, WorldPos, WorldNormal, IsTransmissiveMaterial, OutIrradiance);
+			return;
+		case PATHTRACING_LIGHT_SPOT:
+			SpotLight_EstimateLight(LightId, WorldPos, WorldNormal, IsTransmissiveMaterial, OutIrradiance);
+			return;
+		default:
+			return;
 	}
 }
 
@@ -280,25 +264,25 @@ void PdfLight(
 	float LightPdf = 0.0;
 
 	// Light sampling
-	if (IsEnvironmentLight(LightId))
-	{
-		SkyLight_PdfLight(Ray, LightId, LightUV, LightPdf);
-	}
-	else if (IsPointLight(LightId))
-	{
-		PointLight_PdfLight(Ray, LightId, LightUV, LightPdf);
-	}
-	else if (IsSpotLight(LightId))
-	{
-		SpotLight_PdfLight(Ray, LightId, LightUV, LightPdf);
-	}
-	else if (IsDirectionalLight(LightId))
+	switch (SceneLightsData.Type[LightId])
 	{
-		DirectionalLight_PdfLight(Ray, LightId, LightUV, LightPdf);
-	}
-	else if (IsRectLight(LightId))
-	{
-		RectLight_PdfLight(Ray, LightId, LightUV, LightPdf);
+		case PATHTRACING_LIGHT_SKY:
+			SkyLight_PdfLight(Ray, LightId, LightUV, LightPdf);
+			break;
+		case PATHTRACING_LIGHT_POINT:
+			PointLight_PdfLight(Ray, LightId, LightUV, LightPdf);
+			break;
+		case PATHTRACING_LIGHT_DIRECTIONAL:
+			DirectionalLight_PdfLight(Ray, LightId, LightUV, LightPdf);
+			break;
+		case PATHTRACING_LIGHT_RECT:
+			RectLight_PdfLight(Ray, LightId, LightUV, LightPdf);
+			break;
+		case PATHTRACING_LIGHT_SPOT:
+			SpotLight_PdfLight(Ray, LightId, LightUV, LightPdf);
+			break;
+		default:
+			break;
 	}
 
 	OutPdf = LightSelectionPdf * LightPdf;
@@ -313,25 +297,25 @@ void EvalLight(
 {
 	OutRadiance = float3(0.0f, 0.0f, 0.0f);
 
-	if (IsPointLight(LightId))
-	{
-		PointLight_EvalLight(LightId, LightUV, Ray, OutRadiance);
-	}
-	else if (IsSpotLight(LightId))
-	{
-		SpotLight_EvalLight(LightId, LightUV, Ray, OutRadiance);
-	}
-	else if (IsEnvironmentLight(LightId))
-	{
-		SkyLight_EvalLight(LightId, LightUV, Ray, OutRadiance);
-	}
-	else if (IsDirectionalLight(LightId))
-	{
-		DirectionalLight_EvalLight(LightId, LightUV, Ray, OutRadiance);
-	}
-	else if (IsRectLight(LightId))
-	{
-		RectLight_EvalLight(LightId, LightUV, Ray,  OutRadiance);
+	switch (SceneLightsData.Type[LightId])
+	{
+		case PATHTRACING_LIGHT_SKY:
+			SkyLight_EvalLight(LightId, LightUV, Ray, OutRadiance);
+			return;
+		case PATHTRACING_LIGHT_POINT:
+			PointLight_EvalLight(LightId, LightUV, Ray, OutRadiance);
+			return;
+		case PATHTRACING_LIGHT_DIRECTIONAL:
+			DirectionalLight_EvalLight(LightId, LightUV, Ray, OutRadiance);
+			return;
+		case PATHTRACING_LIGHT_RECT:
+			RectLight_EvalLight(LightId, LightUV, Ray, OutRadiance);
+			return;
+		case PATHTRACING_LIGHT_SPOT:
+			SpotLight_EvalLight(LightId, LightUV, Ray, OutRadiance);
+			return;
+		default:
+			return;
 	}
 }
 
@@ -354,27 +338,27 @@ void SampleLight(
 
 	// Light sampling
 	float LightPdf = 0.0;
-	if (IsPointLight(OutLightId))
-	{
-		PointLight_SampleLight(Ray, Payload, RandSample, OutLightId, OutLightUV, LightPdf);
-	}
-	else if (IsSpotLight(OutLightId))
-	{
-		SpotLight_SampleLight(Ray, Payload, RandSample, OutLightId, OutLightUV, LightPdf);
-	}
-	else if (IsDirectionalLight(OutLightId))
-	{
-		DirectionalLight_SampleLight(Ray, Payload, RandSample, OutLightId, OutLightUV, LightPdf);
-	}
-	else if (IsEnvironmentLight(OutLightId))
-	{
-		SkyLight_SampleLight(Ray, Payload, RandSample, OutLightId, OutLightUV, LightPdf);
-	}
-	else if (IsRectLight(OutLightId))
+
+	switch (SceneLightsData.Type[OutLightId])
 	{
-		RectLight_SampleLight(Ray, Payload, RandSample, OutLightId, OutLightUV, LightPdf);
+		case PATHTRACING_LIGHT_SKY:
+			SkyLight_SampleLight(Ray, Payload, RandSample, OutLightId, OutLightUV, LightPdf);
+			break;
+		case PATHTRACING_LIGHT_POINT:
+			PointLight_SampleLight(Ray, Payload, RandSample, OutLightId, OutLightUV, LightPdf);
+			break;
+		case PATHTRACING_LIGHT_DIRECTIONAL:
+			DirectionalLight_SampleLight(Ray, Payload, RandSample, OutLightId, OutLightUV, LightPdf);
+			break;
+		case PATHTRACING_LIGHT_RECT:
+			RectLight_SampleLight(Ray, Payload, RandSample, OutLightId, OutLightUV, LightPdf);
+			break;
+		case PATHTRACING_LIGHT_SPOT:
+			SpotLight_SampleLight(Ray, Payload, RandSample, OutLightId, OutLightUV, LightPdf);
+			break;
+		default:
+			break;
 	}
-
 	OutPdf = LightSelectionPdf * LightPdf;
 }
 
diff --git a/Engine/Shaders/Private/PathTracing/Material/PathTracingMaterialSampling.ush b/Engine/Shaders/Private/PathTracing/Material/PathTracingMaterialSampling.ush
--- a/Engine/Shaders/Private/PathTracing/Material/PathTracingMaterialSampling.ush
+++ b/Engine/Shaders/Private/PathTracing/Material/PathTracingMaterialSampling.ush
@@ -28,27 +28,22 @@ bool IsNeeValidMaterial(FMaterialClosestHitPayload Payload)
 #if FORCE_ALL_MATERIALS_DIFFUSE
 	return true;
 #else
+
 	if (IsPureSpecularTransmission(Payload))
 	{
 		return false;
 	}
-	else if (IsDefaultLitMaterial(Payload))
-	{
-		return true;
-	}
-	else if (IsClearCoatMaterial(Payload))
-	{
-		return true;
-	}
-	else if (IsTwoSidedFoliageMaterial(Payload))
-	{
-		return true;
-	}
-	else if (IsUnlitMaterial(Payload))
+
+	switch (Payload.ShadingModelID)
 	{
-		return false;
+		case SHADINGMODELID_UNLIT:
+			return false;
+		case SHADINGMODELID_DEFAULT_LIT:		
+		case SHADINGMODELID_CLEAR_COAT:			
+		case SHADINGMODELID_TWOSIDED_FOLIAGE:	
+		default:								
+			return true;
 	}
-	return true;
 #endif
 }
 
@@ -78,29 +73,28 @@ void SampleMaterial(
 	if (IsPureSpecularTransmission(Payload))
 	{
 		SpecularRefraction_SampleMaterial(RayDirection, Payload, RandSample, OutDirection, OutWeight, OutPdf, OutPositionBiasSign);
+		return;
 	}
-	else if (IsDefaultLitMaterial(Payload))
-	{
-		DefaultLit_SampleMaterial(RayDirection, Payload, RandSample, OutDirection, OutWeight, OutPdf);
-	}
-	else if (IsClearCoatMaterial(Payload))
-	{
-		ClearCoat_SampleMaterial(RayDirection, Payload, RandSample, OutDirection, OutWeight, OutPdf);
-	}
-	else if (IsTwoSidedFoliageMaterial(Payload))
-	{
-		TwoSidedFoliage_SampleMaterial(RayDirection, Payload, RandSample, OutDirection, OutWeight, OutPdf, OutPositionBiasSign);
-	}
-	else if (IsUnlitMaterial(Payload))
-	{
-		// Nothing to do, material does not scatter light
-		OutDirection = 0;
-		OutWeight = 0;
-		OutPdf = 0;
-	}
-	else
-	{
-		Lambert_SampleMaterial(Payload, RandSample, OutDirection, OutWeight, OutPdf);
+	switch (Payload.ShadingModelID)
+	{
+		case SHADINGMODELID_UNLIT:
+			// Nothing to do, material does not scatter light
+			OutDirection = 0;
+			OutWeight = 0;
+			OutPdf = 0;
+			return;
+		case SHADINGMODELID_DEFAULT_LIT:
+			DefaultLit_SampleMaterial(RayDirection, Payload, RandSample, OutDirection, OutWeight, OutPdf);
+			return;
+		case SHADINGMODELID_CLEAR_COAT:
+			ClearCoat_SampleMaterial(RayDirection, Payload, RandSample, OutDirection, OutWeight, OutPdf);
+			return;
+		case SHADINGMODELID_TWOSIDED_FOLIAGE:
+			TwoSidedFoliage_SampleMaterial(RayDirection, Payload, RandSample, OutDirection, OutWeight, OutPdf, OutPositionBiasSign);
+			return;
+		default:
+			Lambert_SampleMaterial(Payload, RandSample, OutDirection, OutWeight, OutPdf);
+			break;
 	}
 #endif
 
@@ -129,27 +123,27 @@ void EvalMaterial(
 	if (IsPureSpecularTransmission(Payload))
 	{
 		SpecularTransmission_EvalMaterial(OutWeight, OutPdf);
+		return;
 	}
-	else if (IsDefaultLitMaterial(Payload))
-	{
-		DefaultLit_EvalMaterial(IncomingDirection, OutgoingDirection, Payload, OutWeight, OutPdf);
-	}
-	else if (IsClearCoatMaterial(Payload))
-	{
-		ClearCoat_EvalMaterial(IncomingDirection, OutgoingDirection, Payload, OutWeight, OutPdf);
-	}
-	else if (IsTwoSidedFoliageMaterial(Payload))
-	{
-		TwoSidedFoliage_EvalMaterial(IncomingDirection, OutgoingDirection, Payload, OutWeight, OutPdf);
-	}
-	else if (IsUnlitMaterial(Payload))
-	{
-		OutWeight = 0.0;
-		OutPdf = 0.0;
-	}
-	else
-	{
-		Lambert_EvalMaterial(OutgoingDirection, Payload, OutWeight, OutPdf);
+
+	switch (Payload.ShadingModelID)
+	{
+		case SHADINGMODELID_UNLIT:
+			OutWeight = 0.0;
+			OutPdf = 0.0;
+			break;
+		case SHADINGMODELID_DEFAULT_LIT:
+			DefaultLit_EvalMaterial(IncomingDirection, OutgoingDirection, Payload, OutWeight, OutPdf);
+			break;
+		case SHADINGMODELID_CLEAR_COAT:
+			ClearCoat_EvalMaterial(IncomingDirection, OutgoingDirection, Payload, OutWeight, OutPdf);
+			break;
+		case SHADINGMODELID_TWOSIDED_FOLIAGE:
+			TwoSidedFoliage_EvalMaterial(IncomingDirection, OutgoingDirection, Payload, OutWeight, OutPdf);
+			break;
+		default:
+			Lambert_EvalMaterial(OutgoingDirection, Payload, OutWeight, OutPdf);
+			break;
 	}
 #endif
 }
