Fixed throughput not being applied on the sky light diffuse contribution.

diff --git a/Engine/Shaders/Private/Strata/StrataEnvironmentLighting.ush b/Engine/Shaders/Private/Strata/StrataEnvironmentLighting.ush
--- a/Engine/Shaders/Private/Strata/StrataEnvironmentLighting.ush
+++ b/Engine/Shaders/Private/Strata/StrataEnvironmentLighting.ush
@@ -126,7 +126,7 @@ float3 StrataEnvironmentLighting(
 			// Finally sample the sky diffuse contribution (spherical harmonic, Lambert BRDF)
 			float3 DiffuseLookup = GetSkySHDiffuse(Normal) * View.SkyLightColor.rgb;			// STRATA_TODO compute SH coefficients for ON or Chan
 			// And accumulate
-			SkyLighting += (SkyVisData.SkyDiffuseLookUpMul * DiffuseLookup + SkyVisData.SkyDiffuseLookUpAdd) * Albedo;
+			SkyLighting += BSDFThroughput * (SkyVisData.SkyDiffuseLookUpMul * DiffuseLookup + SkyVisData.SkyDiffuseLookUpAdd) * Albedo;
 		}
 
 		// Specular component
