Select the random sequence in the path tracer at compile time instead of at runtime

The random sequence can still be chosen by defining the RANDSEQ pre-processor variable

diff --git a/Engine/Plugins/Experimental/GPULightmass/Shaders/Private/LightmapPathTracing.usf b/Engine/Plugins/Experimental/GPULightmass/Shaders/Private/LightmapPathTracing.usf
--- a/Engine/Plugins/Experimental/GPULightmass/Shaders/Private/LightmapPathTracing.usf
+++ b/Engine/Plugins/Experimental/GPULightmass/Shaders/Private/LightmapPathTracing.usf
@@ -591,8 +591,6 @@ void LightmapPathTracingMainRG()
 						(BatchedTiles[TileIndex].RenderPassIndex)
 				);
 				
-				RandSequence.Type = 2;
-				RandSequence.HaltonDimensionIndex = 2; // 0 and 1 are used for image plane sampling (lightmap GBuffer)
 				uint SampleIndex = 2;
 
 				float3 RadianceValue = 0;
@@ -794,8 +792,6 @@ void VolumetricLightmapPathTracingMainRG()
 	{
 		RandomSequence RandSequence;
 		RandomSequence_Initialize(RandSequence, Seed, FrameNumber / 2);	
-		RandSequence.Type = 2;
-		RandSequence.HaltonDimensionIndex = 3;
 		uint SampleIndex = 4;
 		
 		PathTracingKernel(
diff --git a/Engine/Shaders/Private/PathTracing/PathTracing.usf b/Engine/Shaders/Private/PathTracing/PathTracing.usf
--- a/Engine/Shaders/Private/PathTracing/PathTracing.usf
+++ b/Engine/Shaders/Private/PathTracing/PathTracing.usf
@@ -56,8 +56,6 @@ RAY_TRACING_ENTRY_RAYGEN(PathTracingMainRG)
 	uint LinearIndex = LaunchIndex.x + (LaunchIndex.y * BufferSize.x); 
 	RandomSequence RandSequence;
 	RandomSequence_Initialize(RandSequence, LinearIndex, AdaptiveSamplingData.TemporalSeed);
-	RandSequence.Type = AdaptiveSamplingData.RandomSequence;
-	uint SampleIndex = 0;
 
 	float3 Radiance = 0;
 	uint RayCounter = 0;
@@ -66,7 +64,7 @@ RAY_TRACING_ENTRY_RAYGEN(PathTracingMainRG)
 	RayDesc Ray;
 	uint2 Pixel;
 	float EmitterPdf;
-	SampleEmitter(LaunchIndex, RandSequence, SampleIndex, 0, 0, Ray, Pixel, EmitterPdf);
+	SampleEmitter(LaunchIndex, RandSequence, 0, 0, Ray, Pixel, EmitterPdf);
 	if (EmitterPdf <= 0.0)
 	{
 		return;
@@ -156,7 +154,7 @@ RAY_TRACING_ENTRY_RAYGEN(PathTracingMainRG)
 		}
 
 		// Choose a random number for both Light sampling and BxDF sampling
-		float4 RandSample = RandomSequence_GenerateSample4D(RandSequence, SampleIndex);
+		float4 RandSample = RandomSequence_GenerateSample4D(RandSequence, 0);
 
 		// Does this material require NEE?
 		bool bIsNeeValid = IsNeeValidMaterial(Payload);
@@ -279,7 +277,7 @@ RAY_TRACING_ENTRY_RAYGEN(PathTracingMainRG)
 		{
 			// If there is some chance we should terminate the ray, draw an extra random value
 			float RussianRouletteRand = RandSample.w; // SampleMaterial does not use this value at the moment
-			//RussianRouletteRand = RandomSequence_GenerateSample1D(RandSequence, SampleIndex);
+			//RussianRouletteRand = RandomSequence_GenerateSample1D(RandSequence, 0);
 			if (RussianRouletteRand >= ContinuationProb)
 			{
 				// stochastically terminate the path
diff --git a/Engine/Shaders/Private/PathTracing/PathTracingCamera.ush b/Engine/Shaders/Private/PathTracing/PathTracingCamera.ush
--- a/Engine/Shaders/Private/PathTracing/PathTracingCamera.ush
+++ b/Engine/Shaders/Private/PathTracing/PathTracingCamera.ush
@@ -13,7 +13,6 @@
 void SampleEmitter(
 	uint2 LaunchIndex,
 	inout RandomSequence RandSequence,
-	inout uint SampleIndex,
 	uint SourceEmitterId,
 	uint SensorId,
 	out RayDesc OutRay,
@@ -24,11 +23,11 @@ void SampleEmitter(
 	// Pinhole camera is preferred until adaptive variance reduction is a clear win
 	if (AdaptiveSamplingData.UseAdaptiveSampling == 1)
 	{
-		AdaptivePinholeCamera_SampleEmitter(LaunchIndex, RandSequence, SampleIndex, SourceEmitterId, SensorId, OutRay, OutPixel, OutPdf);
+		AdaptivePinholeCamera_SampleEmitter(LaunchIndex, RandSequence, SourceEmitterId, SensorId, OutRay, OutPixel, OutPdf);
 	}
 	else
 	{
-		PinholeCamera_SampleEmitter(LaunchIndex, RandSequence, SampleIndex, SourceEmitterId, SensorId, OutRay, OutPixel, OutPdf);
+		PinholeCamera_SampleEmitter(LaunchIndex, RandSequence, SourceEmitterId, SensorId, OutRay, OutPixel, OutPdf);
 	}
 }
 
diff --git a/Engine/Shaders/Private/PathTracing/Sensor/PathTracingAdaptivePinholeCamera.ush b/Engine/Shaders/Private/PathTracing/Sensor/PathTracingAdaptivePinholeCamera.ush
--- a/Engine/Shaders/Private/PathTracing/Sensor/PathTracingAdaptivePinholeCamera.ush
+++ b/Engine/Shaders/Private/PathTracing/Sensor/PathTracingAdaptivePinholeCamera.ush
@@ -54,7 +54,6 @@ bool DoAdaptiveSampling()
 void AdaptivePinholeCamera_SampleEmitter(
 	uint2 LaunchIndex,
 	inout RandomSequence RandSequence,
-	inout uint SampleIndex,
 	uint SourceEmitterId,
 	uint SensorId,
 	out RayDesc OutRay,
@@ -65,7 +64,7 @@ void AdaptivePinholeCamera_SampleEmitter(
 	// Fallback
 	if (!DoAdaptiveSampling())
 	{
-		return PinholeCamera_SampleEmitter(LaunchIndex, RandSequence, SampleIndex, SourceEmitterId, SensorId, OutRay, OutPixel, OutPdf);
+		return PinholeCamera_SampleEmitter(LaunchIndex, RandSequence, SourceEmitterId, SensorId, OutRay, OutPixel, OutPdf);
 	}
 	OutPdf = 1.0;
 
@@ -83,7 +82,7 @@ void AdaptivePinholeCamera_SampleEmitter(
 		VarianceCdf.w = VarianceCdf.z + AdaptiveSamplingData.VarianceMipTree[BufferOffsetAtPixel(MipPixel + uint2(1, 1), MipIndex)];
 		VarianceCdf /= VarianceCdf.w;
 
-		float RandSample = RandomSequence_GenerateSample1D(RandSequence, SampleIndex);
+		float RandSample = RandomSequence_GenerateSample1D(RandSequence, 0);
 		if (RandSample < VarianceCdf.x)
 		{
 			OutPdf *= VarianceCdf.x;
@@ -114,7 +113,7 @@ void AdaptivePinholeCamera_SampleEmitter(
 #endif
 	MipPixel = uint2(MipPixel.x >> 1, MipPixel.y >> 1);
 
-	float2 PixelOffset = RandomSequence_GenerateSample2D(RandSequence, SampleIndex);
+	float2 PixelOffset = RandomSequence_GenerateSample2D(RandSequence, 0);
 	float2 MipPixelUV = (MipPixel + PixelOffset) / float2(DimensionsAtLevel(StopLevel));
 	// Correct for ViewSize/BufferSize discrepancy
 	MipPixelUV *= View.ViewSizeAndInvSize.xy * View.BufferSizeAndInvSize.zw;
diff --git a/Engine/Shaders/Private/PathTracing/Sensor/PathTracingPinholeCamera.ush b/Engine/Shaders/Private/PathTracing/Sensor/PathTracingPinholeCamera.ush
--- a/Engine/Shaders/Private/PathTracing/Sensor/PathTracingPinholeCamera.ush
+++ b/Engine/Shaders/Private/PathTracing/Sensor/PathTracingPinholeCamera.ush
@@ -9,7 +9,6 @@
 void PinholeCamera_SampleEmitter(
 	uint2 LaunchIndex,
 	inout RandomSequence RandSequence,
-	inout uint SampleIndex,
 	uint SourceEmitterId,
 	uint SensorId,
 	out RayDesc OutRay,
@@ -17,7 +16,7 @@ void PinholeCamera_SampleEmitter(
 	out float OutPdf
 )
 {
-	float2 RandSample = RandomSequence_GenerateSample2D(RandSequence, SampleIndex);
+	float2 RandSample = RandomSequence_GenerateSample2D(RandSequence, 0);
 	float2 ViewportUV = (LaunchIndex.xy + RandSample.xy) * View.BufferSizeAndInvSize.zw;
 
 	OutRay = CreatePrimaryRay(ViewportUV);
diff --git a/Engine/Shaders/Private/PathTracing/Utilities/PathTracingRandomSequence.ush b/Engine/Shaders/Private/PathTracing/Utilities/PathTracingRandomSequence.ush
--- a/Engine/Shaders/Private/PathTracing/Utilities/PathTracingRandomSequence.ush
+++ b/Engine/Shaders/Private/PathTracing/Utilities/PathTracingRandomSequence.ush
@@ -9,18 +9,22 @@
 
 #include "../../QuasiRandom.ush" 
 
+#define RANDSEQ_LCG			0
+#define RANDSEQ_HALTON		1
+#define RANDSEQ_OWENSOBOL	2
 
-struct RandomSequence
-{
-	uint Type;
-	uint PositionSeed;
-	uint TimeSeed;
+#ifndef RANDSEQ
+#define RANDSEQ			RANDSEQ_OWENSOBOL
+#endif
 
-	// LCG
-	uint PseudoRandomSeed;
+#ifndef RANDSEQ_RANDOMIZED
+#define RANDSEQ_RANDOMIZED		1 // 1 to randomize per pixel, 0 to share the same sequence in all pixels
+#endif
 
-	// Halton options
-	uint HaltonDimensionIndex;
+struct RandomSequence
+{
+	uint SampleIndex;		// index into the random sequence
+	uint SampleSeed;		// changes as we draw samples to reflect the change in dimension
 };
 
 // TEA-based pseudo-random number generator
@@ -58,26 +62,6 @@ float Rand(inout uint Seed)
 }
 
 // #dxr_todo: convert prime factors to constant buffer
-static const uint Primes128LUT[] = {
-	2,3,5,7,11,13,17,19,23,29,
-	31,37,41,43,47,53,59,61,67,71,
-	73,79,83,89,97,101,103,107,109,113,
-	127,131,137,139,149,151,157,163,167,173,
-	179,181,191,193,197,199,211,223,227,229,
-	233,239,241,251,257,263,269,271,277,281,
-	283,293,307,311,313,317,331,337,347,349,
-	353,359,367,373,379,383,389,397,401,409,
-	419,421,431,433,439,443,449,457,461,463,
-	467,479,487,491,499,503,509,521,523,541,
-	547,557,563,569,571,577,587,593,599,601,
-	607,613,617,619,631,641,643,647,653,659,
-	661,673,677,683,691,701,709,719
-};
-uint Prime128(uint Dimension)
-{
-	return Primes128LUT[Dimension % 128];
-}
-
 static const uint Primes512LUT[] = {
 	2,3,5,7,11,13,17,19,23,29,
 	31,37,41,43,47,53,59,61,67,71,
@@ -252,134 +236,102 @@ float4 SobolSampler(uint SampleIndex, inout uint Seed) {
 
 void RandomSequence_Initialize(inout RandomSequence RandSequence, uint PositionSeed, uint TimeSeed)
 {
-	RandSequence.Type = 0;
-	RandSequence.PositionSeed = PositionSeed;
-	RandSequence.TimeSeed = TimeSeed;
-
-	// LCG
-	RandSequence.PseudoRandomSeed = RandInit(PositionSeed, TimeSeed);
-
-	// Halton
-	RandSequence.HaltonDimensionIndex = 0;
-
-	// Halton and Sobol both want a hashed position seed - so precompute it here
-	RandSequence.PositionSeed = IntegerHash(PositionSeed);
-}
-
-bool RandomSequence_IsLCG(in RandomSequence RandSequence)
-{
-	return RandSequence.Type == 0;
-}
-
-bool RandomSequence_IsHalton(in RandomSequence RandSequence)
-{
-	return RandSequence.Type == 1;
-}
-
-bool RandomSequence_IsScrambledHalton(in RandomSequence RandSequence)
-{
-	return RandSequence.Type == 2;
-}
+	// optionally disable randomization per pixel so that all pixels follow the same path in primary sample space
+#if RANDSEQ_RANDOMIZED == 0
+	PositionSeed = 0;
+#endif
 
-bool RandomSequence_IsOwenSobol(in RandomSequence RandSequence)
-{
-	return RandSequence.Type == 3;
+#if RANDSEQ == RANDSEQ_LCG
+	RandSequence.SampleIndex = 0; // not used
+	RandSequence.SampleSeed = RandInit(PositionSeed, TimeSeed);
+#elif RANDSEQ == RANDSEQ_HALTON
+	RandSequence.SampleIndex = TimeSeed + IntegerHash(PositionSeed);
+	RandSequence.SampleSeed = 0; // this sampler just counts dimensions
+#else
+	RandSequence.SampleIndex = TimeSeed;
+	RandSequence.SampleSeed  = IntegerHash(PositionSeed);
+#endif
 }
 
-
-float RandomSequence_generateSample(inout RandomSequence RandSequence, in uint SampleIndex)
+float RandomSequence_GenerateSample1D(inout RandomSequence RandSequence, uint DummyVariable)
 {
-	// dxr_todo: Make choice of random a shader permutation
-	if (RandomSequence_IsHalton(RandSequence))
-	{
-		uint HaltonSampleIndex = RandSequence.TimeSeed;
-
-		float RandomSample = Halton(HaltonSampleIndex, Prime512(RandSequence.HaltonDimensionIndex));
-		RandSequence.HaltonDimensionIndex += 1;
-		return RandomSample;
-	}
-	else if (RandomSequence_IsScrambledHalton(RandSequence))
-	{
-		uint HaltonSampleIndex = RandSequence.PositionSeed + RandSequence.TimeSeed;
-
-		float RandomSample = Halton(HaltonSampleIndex, Prime512(RandSequence.HaltonDimensionIndex));
-		RandSequence.HaltonDimensionIndex += 1;
-		return RandomSample;
-	}
-	else // LCG
-	{
-		return Rand(RandSequence.PseudoRandomSeed);
-	}
+	float Result;
+#if RANDSEQ == RANDSEQ_LCG
+	Result = Rand(RandSequence.SampleSeed);
+#elif RANDSEQ == RANDSEQ_HALTON
+	Result = Halton(RandSequence.SampleIndex, Prime512(RandSequence.SampleSeed + 0));
+	RandSequence.SampleSeed += 1;
+#elif RANDSEQ == RANDSEQ_OWENSOBOL
+	// rely on compiler to optimize out dead code for now
+	// #dxr_todo: benchmark and specialize if needed
+	Result = SobolSampler(RandSequence.SampleIndex, RandSequence.SampleSeed).x;
+#else
+#error "Unknown random sequence chosen in path tracer"
+#endif
+	return Result;
 }
 
-float RandomSequence_GenerateSample1D(inout RandomSequence RandSequence, inout uint SampleIndex)
+float2 RandomSequence_GenerateSample2D(inout RandomSequence RandSequence, uint DummyVariable)
 {
-	float Sample;
-	if (RandomSequence_IsOwenSobol(RandSequence))
-	{
-		// rely on compiler to optimize out dead code for now
-		// #dxr_todo: benchmark and specialize if needed
-		Sample = SobolSampler(RandSequence.TimeSeed, RandSequence.PositionSeed).x;
-	}
-	else
-	{
-		Sample = RandomSequence_generateSample(RandSequence, SampleIndex);
-	}
-	SampleIndex += 1;
-	return Sample;
+	float2 Result;
+#if RANDSEQ == RANDSEQ_LCG
+	Result.x = Rand(RandSequence.SampleSeed);
+	Result.y = Rand(RandSequence.SampleSeed);
+#elif RANDSEQ == RANDSEQ_HALTON
+	Result.x = Halton(RandSequence.SampleIndex, Prime512(RandSequence.SampleSeed + 0));
+	Result.y = Halton(RandSequence.SampleIndex, Prime512(RandSequence.SampleSeed + 1));
+	RandSequence.SampleSeed += 2;
+#elif RANDSEQ == RANDSEQ_OWENSOBOL
+	// rely on compiler to optimize out dead code for now
+	// #dxr_todo: benchmark and specialize if needed
+	Result = SobolSampler(RandSequence.SampleIndex, RandSequence.SampleSeed).xy;
+#else
+#error "Unknown random sequence chosen in path tracer"
+#endif
+	return Result;
 }
 
-float2 RandomSequence_GenerateSample2D(inout RandomSequence RandSequence, inout uint SampleIndex)
+float3 RandomSequence_GenerateSample3D(inout RandomSequence RandSequence, uint DummyVariable)
 {
-	float2 Sample;
-	if (RandomSequence_IsOwenSobol(RandSequence))
-	{
-		// rely on compiler to optimize out dead code for now
-		// #dxr_todo: benchmark and specialize if needed
-		Sample = SobolSampler(RandSequence.TimeSeed, RandSequence.PositionSeed).xy;
-	}
-	else
-	{
-		Sample.x = RandomSequence_generateSample(RandSequence, SampleIndex);
-		Sample.y = RandomSequence_generateSample(RandSequence, SampleIndex + 1);
-	}
-	SampleIndex += 2;
-	return Sample;
+	float3 Result;
+#if RANDSEQ == RANDSEQ_LCG
+	Result.x = Rand(RandSequence.SampleSeed);
+	Result.y = Rand(RandSequence.SampleSeed);
+	Result.z = Rand(RandSequence.SampleSeed);
+#elif RANDSEQ == RANDSEQ_HALTON
+	Result.x = Halton(RandSequence.SampleIndex, Prime512(RandSequence.SampleSeed + 0));
+	Result.y = Halton(RandSequence.SampleIndex, Prime512(RandSequence.SampleSeed + 1));
+	Result.z = Halton(RandSequence.SampleIndex, Prime512(RandSequence.SampleSeed + 2));
+	RandSequence.SampleSeed += 3;
+#elif RANDSEQ == RANDSEQ_OWENSOBOL
+	// rely on compiler to optimize out dead code for now
+	// #dxr_todo: benchmark and specialize if needed
+	Result = SobolSampler(RandSequence.SampleIndex, RandSequence.SampleSeed).xyz;
+#else
+#error "Unknown random sequence chosen in path tracer"
+#endif
+	return Result;
 }
 
-float3 RandomSequence_GenerateSample3D(inout RandomSequence RandSequence, inout uint SampleIndex)
-{
-	float3 Sample;
-	if (RandomSequence_IsOwenSobol(RandSequence))
-	{
-		// rely on compiler to optimize out dead code for now
-		// #dxr_todo: benchmark and specialize if needed
-		Sample = SobolSampler(RandSequence.TimeSeed, RandSequence.PositionSeed).xyz;
-	}
-	else
-	{
-		Sample.x = RandomSequence_generateSample(RandSequence, SampleIndex);
-		Sample.y = RandomSequence_generateSample(RandSequence, SampleIndex + 1);
-		Sample.z = RandomSequence_generateSample(RandSequence, SampleIndex + 2);
-	}
-	SampleIndex += 3;
-	return Sample;
-}
 
-float4 RandomSequence_GenerateSample4D(inout RandomSequence RandSequence, inout uint SampleIndex)
+float4 RandomSequence_GenerateSample4D(inout RandomSequence RandSequence, uint DummyVariable)
 {
-	float4 Sample;
-	if (RandomSequence_IsOwenSobol(RandSequence))
-	{
-		Sample = SobolSampler(RandSequence.TimeSeed, RandSequence.PositionSeed);
-	}
-	else
-	{
-		Sample.x = RandomSequence_generateSample(RandSequence, SampleIndex);
-		Sample.y = RandomSequence_generateSample(RandSequence, SampleIndex + 1);
-		Sample.z = RandomSequence_generateSample(RandSequence, SampleIndex + 2);
-		Sample.w = RandomSequence_generateSample(RandSequence, SampleIndex + 3);
-	}
-	SampleIndex += 4;
-	return Sample;
+	float4 Result;
+#if RANDSEQ == RANDSEQ_LCG
+	Result.x = Rand(RandSequence.SampleSeed);
+	Result.y = Rand(RandSequence.SampleSeed);
+	Result.z = Rand(RandSequence.SampleSeed);
+	Result.w = Rand(RandSequence.SampleSeed);
+#elif RANDSEQ == RANDSEQ_HALTON
+	Result.x = Halton(RandSequence.SampleIndex, Prime512(RandSequence.SampleSeed + 0));
+	Result.y = Halton(RandSequence.SampleIndex, Prime512(RandSequence.SampleSeed + 1));
+	Result.z = Halton(RandSequence.SampleIndex, Prime512(RandSequence.SampleSeed + 2));
+	Result.w = Halton(RandSequence.SampleIndex, Prime512(RandSequence.SampleSeed + 3));
+	RandSequence.SampleSeed += 4;
+#elif RANDSEQ == RANDSEQ_OWENSOBOL
+	Result = SobolSampler(RandSequence.SampleIndex, RandSequence.SampleSeed);
+#else
+#error "Unknown random sequence chosen in path tracer"
+#endif
+	return Result;
 }
diff --git a/Engine/Source/Runtime/Renderer/Private/PathTracing.cpp b/Engine/Source/Runtime/Renderer/Private/PathTracing.cpp
--- a/Engine/Source/Runtime/Renderer/Private/PathTracing.cpp
+++ b/Engine/Source/Runtime/Renderer/Private/PathTracing.cpp
@@ -64,17 +64,6 @@ TAutoConsoleVariable<int32> CVarPathTracingFrameIndependentTemporalSeed(
 	ECVF_RenderThreadSafe
 );
 
-TAutoConsoleVariable<int32> CVarPathTracingRandomSequence(
-	TEXT("r.PathTracing.RandomSequence"),
-	3,
-	TEXT("Changes the underlying random sequence\n")
-	TEXT("0: LCG\n")
-	TEXT("1: Halton\n")
-	TEXT("2: Scrambled Halton\n")
-	TEXT("3: Owen-Sobol (default)\n"),
-	ECVF_RenderThreadSafe
-);
-
 TAutoConsoleVariable<int32> CVarPathTracingAdaptiveSampling(
 	TEXT("r.PathTracing.AdaptiveSampling"),
 	1,
@@ -172,7 +161,6 @@ static bool PrepareShaderArgs(const FViewInfo& View, FPathTracingData& PathTraci
 	}
 
 	AdaptiveSamplingData.UseAdaptiveSampling = CVarPathTracingAdaptiveSampling.GetValueOnRenderThread();
-	AdaptiveSamplingData.RandomSequence = CVarPathTracingRandomSequence.GetValueOnRenderThread();
 
 	// Changing Adaptive sampling mode requires starting over
 	static uint32 PreviousUseAdaptiveSampling = AdaptiveSamplingData.UseAdaptiveSampling;
@@ -182,13 +170,6 @@ static bool PrepareShaderArgs(const FViewInfo& View, FPathTracingData& PathTraci
 		PreviousUseAdaptiveSampling = AdaptiveSamplingData.UseAdaptiveSampling;
 	}
 
-	static uint32 PreviousRandomSequence = AdaptiveSamplingData.RandomSequence;
-	if (PreviousRandomSequence != AdaptiveSamplingData.RandomSequence)
-	{
-		NeedInvalidation = true;
-		PreviousRandomSequence = AdaptiveSamplingData.RandomSequence;
-	}
-
 	// the rest of PathTracingData and AdaptiveSamplingData is filled in by SetParameters below
 	return NeedInvalidation;
 }
diff --git a/Engine/Source/Runtime/Renderer/Private/PathTracingUniformBuffers.h b/Engine/Source/Runtime/Renderer/Private/PathTracingUniformBuffers.h
--- a/Engine/Source/Runtime/Renderer/Private/PathTracingUniformBuffers.h
+++ b/Engine/Source/Runtime/Renderer/Private/PathTracingUniformBuffers.h
@@ -45,7 +45,6 @@ END_GLOBAL_SHADER_PARAMETER_STRUCT()
 
 BEGIN_GLOBAL_SHADER_PARAMETER_STRUCT(FPathTracingAdaptiveSamplingData, )
 	SHADER_PARAMETER(uint32, UseAdaptiveSampling)
-	SHADER_PARAMETER(uint32, RandomSequence)
 	SHADER_PARAMETER(uint32, MinimumSamplesPerPixel)
 	SHADER_PARAMETER(uint32, Iteration)
 	SHADER_PARAMETER(uint32, TemporalSeed)
