Optimize Owen-Sobol random sequence generation

About 5% faster on the conference scene with 32 bounces

diff --git a/Engine/Shaders/Private/PathTracing/Utilities/PathTracingRandomSequence.ush b/Engine/Shaders/Private/PathTracing/Utilities/PathTracingRandomSequence.ush
--- a/Engine/Shaders/Private/PathTracing/Utilities/PathTracingRandomSequence.ush
+++ b/Engine/Shaders/Private/PathTracing/Utilities/PathTracingRandomSequence.ush
@@ -143,8 +143,8 @@ uint FastOwenScrambling(uint Index, uint Seed) {
 	// Laine and Karras / Stratified Sampling for Stochastic Transparency / EGSR 2011
 
 	// The operations below will mix bits toward the left, so temporarily reverse the order
-	// TODO: can fold this in with sobol logic to skip some reversals
-	Index = ReverseBits32(Index);
+	// NOTE: This operation has been performed outside this call
+	// Index = ReverseBits32(Index);
 	
 	Index += Seed; // randomize the index by our seed (pushes bits toward the left)
 
@@ -176,40 +176,40 @@ uint FastOwenScrambling(uint Index, uint Seed) {
 //    https://web.maths.unsw.edu.au/~fkuo/sobol/
 // NOTE: we don't bother storing dimension 0 since it is just a bit reversal
 // NOTE2: dimension 1 also follows a pattern that can be done procedurally but probably not worth the extra code?
-// NOTE3: if we did all our owen scrambling logic with 16 bits, we could avoid
+// NOTE3: the matrix eleements are reversed to save one reverse in the owen scrambling
 static const uint3 SobolMatrices[] = {
-	uint3(0x80000000, 0x80000000, 0x80000000),
-	uint3(0xc0000000, 0xc0000000, 0xc0000000),
-	uint3(0xa0000000, 0x60000000, 0x20000000),
-	uint3(0xf0000000, 0x90000000, 0x50000000),
-	uint3(0x88000000, 0xe8000000, 0xf8000000),
-	uint3(0xcc000000, 0x5c000000, 0x74000000),
-	uint3(0xaa000000, 0x8e000000, 0xa2000000),
-	uint3(0xff000000, 0xc5000000, 0x93000000),
-	uint3(0x80800000, 0x68800000, 0xd8800000),
-	uint3(0xc0c00000, 0x9cc00000, 0x25400000),
-	uint3(0xa0a00000, 0xee600000, 0x59e00000),
-	uint3(0xf0f00000, 0x55900000, 0xe6d00000),
-	uint3(0x88880000, 0x80680000, 0x78080000),
-	uint3(0xcccc0000, 0xc09c0000, 0xb40c0000),
-	uint3(0xaaaa0000, 0x60ee0000, 0x82020000),
-	uint3(0xffff0000, 0x90550000, 0xc3050000),
-	uint3(0x80008000, 0xe8808000, 0x208f8000),
-	uint3(0xc000c000, 0x5cc0c000, 0x51474000),
-	uint3(0xa000a000, 0x8e606000, 0xfbea2000),
-	uint3(0xf000f000, 0xc5909000, 0x75d93000),
-	uint3(0x88008800, 0x6868e800, 0xa0858800),
-	uint3(0xcc00cc00, 0x9c9c5c00, 0x914e5400),
-	uint3(0xaa00aa00, 0xeeee8e00, 0xdbe79e00),
-	uint3(0xff00ff00, 0x5555c500, 0x25db6d00),
-	uint3(0x80808080, 0x8000e880, 0x58800080),
-	uint3(0xc0c0c0c0, 0xc0005cc0, 0xe54000c0),
-	uint3(0xa0a0a0a0, 0x60008e60, 0x79e00020),
-	uint3(0xf0f0f0f0, 0x9000c590, 0xb6d00050),
-	uint3(0x88888888, 0xe8006868, 0x800800f8),
-	uint3(0xcccccccc, 0x5c009c9c, 0xc00c0074),
-	uint3(0xaaaaaaaa, 0x8e00eeee, 0x200200a2),
-	uint3(0xffffffff, 0xc5005555, 0x50050093)
+	uint3(0x00000001, 0x00000001, 0x00000001),
+	uint3(0x00000003, 0x00000003, 0x00000003),
+	uint3(0x00000005, 0x00000006, 0x00000004),
+	uint3(0x0000000f, 0x00000009, 0x0000000a),
+	uint3(0x00000011, 0x00000017, 0x0000001f),
+	uint3(0x00000033, 0x0000003a, 0x0000002e),
+	uint3(0x00000055, 0x00000071, 0x00000045),
+	uint3(0x000000ff, 0x000000a3, 0x000000c9),
+	uint3(0x00000101, 0x00000116, 0x0000011b),
+	uint3(0x00000303, 0x00000339, 0x000002a4),
+	uint3(0x00000505, 0x00000677, 0x0000079a),
+	uint3(0x00000f0f, 0x000009aa, 0x00000b67),
+	uint3(0x00001111, 0x00001601, 0x0000101e),
+	uint3(0x00003333, 0x00003903, 0x0000302d),
+	uint3(0x00005555, 0x00007706, 0x00004041),
+	uint3(0x0000ffff, 0x0000aa09, 0x0000a0c3),
+	uint3(0x00010001, 0x00010117, 0x0001f104),
+	uint3(0x00030003, 0x0003033a, 0x0002e28a),
+	uint3(0x00050005, 0x00060671, 0x000457df),
+	uint3(0x000f000f, 0x000909a3, 0x000c9bae),
+	uint3(0x00110011, 0x00171616, 0x0011a105),
+	uint3(0x00330033, 0x003a3939, 0x002a7289),
+	uint3(0x00550055, 0x00717777, 0x0079e7db),
+	uint3(0x00ff00ff, 0x00a3aaaa, 0x00b6dba4),
+	uint3(0x01010101, 0x01170001, 0x0100011a),
+	uint3(0x03030303, 0x033a0003, 0x030002a7),
+	uint3(0x05050505, 0x06710006, 0x0400079e),
+	uint3(0x0f0f0f0f, 0x09a30009, 0x0a000b6d),
+	uint3(0x11111111, 0x16160017, 0x1f001001),
+	uint3(0x33333333, 0x3939003a, 0x2e003003),
+	uint3(0x55555555, 0x77770071, 0x45004004),
+	uint3(0xffffffff, 0xaaaa00a3, 0xc900a00a)
 };
 
 float4 SobolSampler(uint SampleIndex, inout uint Seed) {
@@ -219,10 +219,9 @@ float4 SobolSampler(uint SampleIndex, inout uint Seed) {
 	// first scramble the index to decorelate from other 4-tuples
 	uint SobolIndex = FastOwenScrambling(SampleIndex, Seed = Seed * LCG_C + LCG_C);
 	// now get Sobol' point from this index
-	uint4 Result = uint4(ReverseBits32(SobolIndex), 0, 0, 0);
-	UNROLL // does this unrolling help?
-	for (int i = 0; i < 32; SobolIndex >>= 1, i++)
-		Result.yzw ^= (SobolIndex & 1) ? SobolMatrices[i] : 0;
+	uint4 Result = uint4(SobolIndex, 0, 0, 0);
+	for (int i = 0; SobolIndex; SobolIndex >>= 1, i++)
+		Result.yzw ^= (SobolIndex & 1) * SobolMatrices[i];
 
 	// finally scramble the points to avoid structured artifacts
 	Result.x = FastOwenScrambling(Result.x, Seed = Seed * LCG_C + LCG_C);
@@ -247,9 +246,12 @@ void RandomSequence_Initialize(inout RandomSequence RandSequence, uint PositionS
 #elif RANDSEQ == RANDSEQ_HALTON
 	RandSequence.SampleIndex = TimeSeed + IntegerHash(PositionSeed);
 	RandSequence.SampleSeed = 0; // this sampler just counts dimensions
-#else
-	RandSequence.SampleIndex = TimeSeed;
+#elif RANDSEQ == RANDSEQ_OWENSOBOL
+	// pre-compute bit reversal needed for FastOwenScrambling since this index doesn't change
+	RandSequence.SampleIndex = ReverseBits32(TimeSeed);
 	RandSequence.SampleSeed  = IntegerHash(PositionSeed);
+#else
+#error "Unknown random sequence chosen in path tracer"
 #endif
 }
 
