Strata Volume BSDF nw uses mesh smooth normal to compute correct throughput to V.

diff --git a/Engine/Shaders/Private/BasePassPixelShader.usf b/Engine/Shaders/Private/BasePassPixelShader.usf
--- a/Engine/Shaders/Private/BasePassPixelShader.usf
+++ b/Engine/Shaders/Private/BasePassPixelShader.usf
@@ -1480,7 +1480,8 @@ void FPixelShaderInOut_MainPS(
 	uint2 PixelPos = In.SvPosition.xy;
 	uint PixelStrataDataByteOffset = GetStrataPixelDataByteOffset(PixelPos, ResolvedView.BufferSizeAndInvSize.x, OpaqueBasePass.Strata.MaxBytesPerPixel);
 
-	PackStrataOut(OpaqueBasePass.Strata.MaterialLobesBufferUAV, PixelStrataDataByteOffset, Strata, MaterialParameters.WorldNormal, MaterialParameters.CameraVector);
+	const float3 VertexNormal = MaterialParameters.TangentToWorld[2].xyz;
+	PackStrataOut(OpaqueBasePass.Strata.MaterialLobesBufferUAV, PixelStrataDataByteOffset, Strata, VertexNormal, MaterialParameters.CameraVector);
 #endif // STRATA_OPAQUE_ENABLED
 
 #if STRATA_TRANSLUCENT_ENABLED
diff --git a/Engine/Shaders/Private/Strata/Strata.ush b/Engine/Shaders/Private/Strata/Strata.ush
--- a/Engine/Shaders/Private/Strata/Strata.ush
+++ b/Engine/Shaders/Private/Strata/Strata.ush
@@ -52,6 +52,8 @@ struct FStrataVolume
 	float3	Extinction;
 	float	Anisotropy;
 	float	Thickness;
+
+	float3  SmoothSurfaceNormal;	// Typically vertex normal
 };
 
 
@@ -163,7 +165,7 @@ FStrataData GetStrataConductorBSDF(float3 Reflectivity, float3 EdgeColor, float2
 	return StrataData;
 }
 
-FStrataData GetStrataVolumeBSDF(float3 Albedo, float3 Extinction, float Anisotropy, float Thickness)
+FStrataData GetStrataVolumeBSDF(float3 Albedo, float3 Extinction, float Anisotropy, float Thickness, float3 SmoothSurfaceNormal)
 {
 	FStrataData StrataData = GetInitialisedStrataData();
 	StrataData.LayerCount++;
@@ -174,6 +176,7 @@ FStrataData GetStrataVolumeBSDF(float3 Albedo, float3 Extinction, float Anisotro
 	StrataData.Layers[0].BSDFs[0].Volume.Extinction				= Extinction;
 	StrataData.Layers[0].BSDFs[0].Volume.Anisotropy				= Anisotropy;
 	StrataData.Layers[0].BSDFs[0].Volume.Thickness				= Thickness;
+	StrataData.Layers[0].BSDFs[0].Volume.SmoothSurfaceNormal	= SmoothSurfaceNormal;
 
 	StrataData.Layers[0].BSDFs[0].Weight						= 1.0f;
 
@@ -402,7 +405,7 @@ void PackStrataOut(
 	RWByteAddressBuffer StrataBuffer, 
 	uint ByteOffset, 
 	FStrataData Strata, 
-	float3 MeshWorldNormal, 
+	float3 VertexNormal,
 	float3 V)
 {
 	// We cannot output (layers x bsdf) information because if we try to load them later, the reserved amount of VGPRS would be crazy high.
@@ -431,7 +434,7 @@ void PackStrataOut(
 			const float3 BSDFCoverage = BSDF.Weight;
 			BSDFCount++;
 
-			float3 N = MeshWorldNormal;
+			float3 N = VertexNormal;
 			switch (BSDF.BSDFType)		// TODO_STRATA: all have a normal, even subsurface?
 			{
 			case STRATA_BSDF_TYPE_DIFFUSE_ON:
@@ -441,6 +444,8 @@ void PackStrataOut(
 			{ N = BSDF.Dielectric.Normal; break; }
 			case STRATA_BSDF_TYPE_CONDUCTOR:
 			{ N = BSDF.Conductor.Normal; break; }
+			case STRATA_BSDF_TYPE_VOLUME:
+			{ N = BSDF.Volume.SmoothSurfaceNormal; break; }
 			}
 			N = normalize(N);
 
@@ -561,10 +566,13 @@ void PackStrataOut(
 			break;
 			case STRATA_BSDF_TYPE_VOLUME:
 			{
+				// We simply multiply the absorption and scattering coefficient by the thickness to simplify the material to be a normalise medium of thickness 1 meter.
+				// This gives the same result and reduces memory usage.
+				const float Thickness = max(0.0f, BSDF.Volume.Thickness);
+				STRATA_STORE_UINT1(PackR11G11B10F(BSDF.Volume.Extinction * Thickness));
 				STRATA_STORE_UINT1(PackR11G11B10F(BSDF.Volume.Albedo));
-				STRATA_STORE_UINT1(PackR11G11B10F(BSDF.Volume.Extinction));
-				STRATA_STORE_UINT1(PackR11G11B10F(float3(abs(BSDF.Volume.Anisotropy), BSDF.Volume.Thickness, BSDF.Volume.Anisotropy > 0.0f ? 2.0f : 0.0f)));
-				// 12 bytes
+				STRATA_STORE_UINT2(PackR16G16B16A16F(float4(BSDF.Volume.SmoothSurfaceNormal, BSDF.Volume.Anisotropy)));
+				// 16 bytes
 			}
 			break;
 			}
@@ -630,10 +638,8 @@ FStrataPixelBSDFs UnpackStrataIn(ByteAddressBuffer StrataBuffer, uint ByteOffset
 		{
 			StrataLoad_R11G11B10F			(StrataBuffer, ByteOffset, PixelBSDFs.BSDFs[i].Volume.Albedo);
 			StrataLoad_R11G11B10F			(StrataBuffer, ByteOffset, PixelBSDFs.BSDFs[i].Volume.Extinction);
-
-			float SignOfAnisotropy = 0.0f;
-			StrataLoad_R11_G11_B10F			(StrataBuffer, ByteOffset, PixelBSDFs.BSDFs[i].Volume.Anisotropy, PixelBSDFs.BSDFs[i].Volume.Thickness, SignOfAnisotropy);
-			PixelBSDFs.BSDFs[i].Volume.Anisotropy *= (SignOfAnisotropy - 1.0f);
+			StrataLoad_R16G16B16F_A16F		(StrataBuffer, ByteOffset, PixelBSDFs.BSDFs[i].Volume.SmoothSurfaceNormal, PixelBSDFs.BSDFs[i].Volume.Anisotropy);
+			PixelBSDFs.BSDFs[i].Volume.Thickness = 1.0f;
 		}
 		break;
 		}
diff --git a/Engine/Shaders/Private/Strata/StrataDeferredLighting.ush b/Engine/Shaders/Private/Strata/StrataDeferredLighting.ush
--- a/Engine/Shaders/Private/Strata/StrataDeferredLighting.ush
+++ b/Engine/Shaders/Private/Strata/StrataDeferredLighting.ush
@@ -33,6 +33,8 @@ float4 StrataDeferredLighting(
 		{ N = BSDF.Dielectric.Normal; break; }
 		case STRATA_BSDF_TYPE_CONDUCTOR:
 		{ N = BSDF.Conductor.Normal; break; }
+		case STRATA_BSDF_TYPE_VOLUME:
+		{ N = BSDF.Volume.SmoothSurfaceNormal; break; }
 		}
 		N = normalize(N);
 
diff --git a/Engine/Shaders/Private/Strata/StrataEnvironmentLighting.ush b/Engine/Shaders/Private/Strata/StrataEnvironmentLighting.ush
--- a/Engine/Shaders/Private/Strata/StrataEnvironmentLighting.ush
+++ b/Engine/Shaders/Private/Strata/StrataEnvironmentLighting.ush
@@ -95,7 +95,7 @@ float3 StrataEnvironmentLighting(
 			const float3 SafeExtinctionThreshold = 0.000001f;
 			const float3 SafeExtinctionCoefficients = max(SafeExtinctionThreshold, ExtinctionCoefficients);
 
-			const float SatNoV = saturate(dot(GBuffer.WorldNormal, V));
+			const float SatNoV = saturate(dot(BSDF.Volume.SmoothSurfaceNormal, V));
 			const float PathLength = Thickness / max(0.0001f, SatNoV);
 			const float3 SafePathSegmentTransmittance = exp(-SafeExtinctionCoefficients * PathLength);
 			const float3 NormLuminanceTransfert = 1.0f * ScatteringCoefficients;
@@ -105,7 +105,7 @@ float3 StrataEnvironmentLighting(
 			Albedo = NormLuminanceSingleScattering;
 
 			// STRATA_TODO: This is to fake some directionality. Have normal specified on the material as for MaterialX?
-			Normal = GBuffer.WorldNormal;
+			Normal = BSDF.Volume.SmoothSurfaceNormal;
 		}
 		break;
 		}
diff --git a/Engine/Shaders/Private/Strata/StrataForwardLighting.ush b/Engine/Shaders/Private/Strata/StrataForwardLighting.ush
--- a/Engine/Shaders/Private/Strata/StrataForwardLighting.ush
+++ b/Engine/Shaders/Private/Strata/StrataForwardLighting.ush
@@ -42,6 +42,8 @@ float3 StrataForwardLighting(
 			{ N = BSDF.Dielectric.Normal; break; }
 			case STRATA_BSDF_TYPE_CONDUCTOR:
 			{ N = BSDF.Conductor.Normal; break; }
+			case STRATA_BSDF_TYPE_VOLUME:
+			{ N = BSDF.Volume.SmoothSurfaceNormal; break; }
 			}
 			N = normalize(N);
 			float3 R = 2 * dot(V, N) * N - V;
diff --git a/Engine/Source/Runtime/Engine/Private/Materials/HLSLMaterialTranslator.cpp b/Engine/Source/Runtime/Engine/Private/Materials/HLSLMaterialTranslator.cpp
--- a/Engine/Source/Runtime/Engine/Private/Materials/HLSLMaterialTranslator.cpp
+++ b/Engine/Source/Runtime/Engine/Private/Materials/HLSLMaterialTranslator.cpp
@@ -7273,11 +7273,12 @@ int32 FHLSLMaterialTranslator::StrataConductorBSDF(int32 Reflectivity, int32 Edg
 int32 FHLSLMaterialTranslator::StrataVolumeBSDF(int32 Albedo, int32 Extinction, int32 Anisotropy, int32 Thickness)
 {
 	return AddInlinedCodeChunk(
-		MCT_Strata, TEXT("GetStrataVolumeBSDF(%s, %s, %s, %s)"),
+		MCT_Strata, TEXT("GetStrataVolumeBSDF(%s, %s, %s, %s, %s)"),
 		*GetParameterCode(Albedo),
 		*GetParameterCode(Extinction),
 		*GetParameterCode(Anisotropy),
-		*GetParameterCode(Thickness)
+		*GetParameterCode(Thickness),
+		*GetParameterCode(VertexNormal())
 	);
 }
 
