Added strata BSDF sanitize functions.
Also fixed bad read order on VolumeBSDF.

diff --git a/Engine/Shaders/Private/Strata/Strata.ush b/Engine/Shaders/Private/Strata/Strata.ush
--- a/Engine/Shaders/Private/Strata/Strata.ush
+++ b/Engine/Shaders/Private/Strata/Strata.ush
@@ -255,10 +255,37 @@ FStrataData StrataVerticalLayering(FStrataData Top, FStrataData Base)
 
 
 
+///////////////////////////////////////////////////////////////////////////////
+// Functions used to sanitize BSDF before they are used for lighting (when forward) or stored to memory (base pass).
+// We do not normalize normals, this is done before the lighting step.
 
+void SanitizeStrataDiffuse(inout FStrataDiffuse BSDF)
+{
+	BSDF.Albedo = saturate(BSDF.Albedo);
+	BSDF.Roughness = saturate(BSDF.Roughness);
+}
 
+void SanitizeStrataDielectric(inout FStrataDielectric BSDF)
+{
+	BSDF.IOR = clamp(BSDF.IOR, 1.0f, 5.0f);		// See https://en.wikipedia.org/wiki/List_of_refractive_indices
+	BSDF.Tint = saturate(BSDF.Tint);
+	BSDF.Roughness = saturate(BSDF.Roughness);
+}
 
+void SanitizeStrataConductor(inout FStrataConductor BSDF)
+{
+	BSDF.Reflectivity = saturate(BSDF.Reflectivity);
+	BSDF.EdgeColor = saturate(BSDF.EdgeColor);
+	BSDF.Roughness = saturate(BSDF.Roughness);
+}
 
+void SanitizeStrataVolume(inout FStrataVolume BSDF)
+{
+	BSDF.Albedo = saturate(BSDF.Albedo);
+	BSDF.Extinction = max(0.0f, BSDF.Extinction);
+	BSDF.Anisotropy = max(0.0f, BSDF.Anisotropy);
+	BSDF.Thickness = max(0.0f, BSDF.Thickness);
+}
 
 
 
@@ -439,13 +466,13 @@ void PackStrataOut(
 			{
 			case STRATA_BSDF_TYPE_DIFFUSE_ON:
 			case STRATA_BSDF_TYPE_DIFFUSE_CHAN:
-			{ N = BSDF.Diffuse.Normal; break; }
+			{ N = BSDF.Diffuse.Normal;				SanitizeStrataDiffuse(BSDF.Diffuse);		break; }
 			case STRATA_BSDF_TYPE_DIELECTRIC:
-			{ N = BSDF.Dielectric.Normal; break; }
+			{ N = BSDF.Dielectric.Normal;			SanitizeStrataDielectric(BSDF.Dielectric);	break; }
 			case STRATA_BSDF_TYPE_CONDUCTOR:
-			{ N = BSDF.Conductor.Normal; break; }
+			{ N = BSDF.Conductor.Normal;			SanitizeStrataConductor(BSDF.Conductor);	break; }
 			case STRATA_BSDF_TYPE_VOLUME:
-			{ N = BSDF.Volume.SmoothSurfaceNormal; break; }
+			{ N = BSDF.Volume.SmoothSurfaceNormal;	SanitizeStrataVolume(BSDF.Volume);			break; }
 			}
 			N = normalize(N);
 
@@ -636,8 +663,8 @@ FStrataPixelBSDFs UnpackStrataIn(ByteAddressBuffer StrataBuffer, uint ByteOffset
 		break;
 		case STRATA_BSDF_TYPE_VOLUME:
 		{
-			StrataLoad_R11G11B10F			(StrataBuffer, ByteOffset, PixelBSDFs.BSDFs[i].Volume.Albedo);
 			StrataLoad_R11G11B10F			(StrataBuffer, ByteOffset, PixelBSDFs.BSDFs[i].Volume.Extinction);
+			StrataLoad_R11G11B10F			(StrataBuffer, ByteOffset, PixelBSDFs.BSDFs[i].Volume.Albedo);
 			StrataLoad_R16G16B16F_A16F		(StrataBuffer, ByteOffset, PixelBSDFs.BSDFs[i].Volume.SmoothSurfaceNormal, PixelBSDFs.BSDFs[i].Volume.Anisotropy);
 			PixelBSDFs.BSDFs[i].Volume.Thickness = 1.0f;
 		}
diff --git a/Engine/Shaders/Private/Strata/StrataForwardLighting.ush b/Engine/Shaders/Private/Strata/StrataForwardLighting.ush
--- a/Engine/Shaders/Private/Strata/StrataForwardLighting.ush
+++ b/Engine/Shaders/Private/Strata/StrataForwardLighting.ush
@@ -37,13 +37,13 @@ float3 StrataForwardLighting(
 			{
 			case STRATA_BSDF_TYPE_DIFFUSE_ON:
 			case STRATA_BSDF_TYPE_DIFFUSE_CHAN:
-			{ N = BSDF.Diffuse.Normal; break; }
+			{ N = BSDF.Diffuse.Normal;				SanitizeStrataDiffuse(BSDF.Diffuse);		break; }
 			case STRATA_BSDF_TYPE_DIELECTRIC:
-			{ N = BSDF.Dielectric.Normal; break; }
+			{ N = BSDF.Dielectric.Normal;			SanitizeStrataDielectric(BSDF.Dielectric);	break; }
 			case STRATA_BSDF_TYPE_CONDUCTOR:
-			{ N = BSDF.Conductor.Normal; break; }
+			{ N = BSDF.Conductor.Normal;			SanitizeStrataConductor(BSDF.Conductor);	break; }
 			case STRATA_BSDF_TYPE_VOLUME:
-			{ N = BSDF.Volume.SmoothSurfaceNormal; break; }
+			{ N = BSDF.Volume.SmoothSurfaceNormal;	SanitizeStrataVolume(BSDF.Volume);		break; }
 			}
 			N = normalize(N);
 			float3 R = 2 * dot(V, N) * N - V;
