Minor consistency update on unused code yet.

diff --git a/Engine/Shaders/Private/Strata/StrataDeferredLighting.ush b/Engine/Shaders/Private/Strata/StrataDeferredLighting.ush
--- a/Engine/Shaders/Private/Strata/StrataDeferredLighting.ush
+++ b/Engine/Shaders/Private/Strata/StrataDeferredLighting.ush
@@ -124,7 +124,7 @@ float4 StrataDeferredLighting(
 #if 1
 			const float3 SelfShadowTransmittance = SatNoL;
 #else
-			const float LightPathLength = 1.0f / max(0.0001f, SatNoL);
+			const float LightPathLength = Thickness / max(0.0001f, SatNoL);
 			const float3 SelfShadowTransmittance = SatNoL <= 0.0 ? 0.0 : exp(-ExtinctionCoefficients * LightPathLength);
 #endif
 
diff --git a/Engine/Shaders/Private/Strata/StrataForwardLighting.ush b/Engine/Shaders/Private/Strata/StrataForwardLighting.ush
--- a/Engine/Shaders/Private/Strata/StrataForwardLighting.ush
+++ b/Engine/Shaders/Private/Strata/StrataForwardLighting.ush
@@ -175,7 +175,7 @@ float3 StrataForwardLighting(
 #if 1
 				const float3 SelfShadowTransmittance = SatNoL;
 #else
-				const float LightPathLength = 1.0f / max(0.0001f, SatNoL);
+				const float LightPathLength = Thickness / max(0.0001f, SatNoL);
 				const float3 SelfShadowTransmittance = SatNoL <= 0.0 ? 0.0 : exp(-ExtinctionCoefficients * LightPathLength);
 #endif
 
