Clean up Strata refelction code

diff --git a/Engine/Shaders/Private/Strata/StrataEnvironmentLighting.ush b/Engine/Shaders/Private/Strata/StrataEnvironmentLighting.ush
--- a/Engine/Shaders/Private/Strata/StrataEnvironmentLighting.ush
+++ b/Engine/Shaders/Private/Strata/StrataEnvironmentLighting.ush
@@ -10,7 +10,7 @@ float3 StrataEnvironmentLighting(
 	float3 WorldPosition,
 	float3 CameraToPixel,
 	float3 V,
-	float3 AmbientOcclusion,
+	float AmbientOcclusion,
 	FStrataPixelBSDFs PixelBSDFs
 	)
 {
@@ -20,16 +20,20 @@ float3 StrataEnvironmentLighting(
 		FStrataBSDF BSDF = PixelBSDFs.BSDFs[i];
 		const float3 BSDFThroughput = BSDF.Weight;
 
-		float3 Albedo = 0.0f;
-		float3 Normal = 0.0f;
-		float3 SpecularLuminance = 0.0f;
+		float3 DiffuseAlbedo = 0.0f;
+		float3 DiffuseNormal = 0.0f;
+
+		float3 SpecularDirection = 0.0f;
+		float3 SpecularWeight = 0.0f;
+		float  SpecularSafeRoughness = 0.0f;
+
 		switch (BSDF.BSDFType)
 		{
 		case STRATA_BSDF_TYPE_DIFFUSE_ON:
 		case STRATA_BSDF_TYPE_DIFFUSE_CHAN:
 		{
-			Albedo = BSDF.Diffuse.Albedo;
-			Normal = BSDF.Diffuse.Normal;
+			DiffuseAlbedo = BSDF.Diffuse.Albedo;
+			DiffuseNormal = BSDF.Diffuse.Normal;
 		}
 		break;
 		case STRATA_BSDF_TYPE_DIELECTRIC:
@@ -38,23 +42,18 @@ float3 StrataEnvironmentLighting(
 			if (ReflectionStruct.SkyLightParameters.y > 0)
 			{
 				float F0 = DielectricIorToF0(BSDF.Dielectric.IOR);
-
 				const float SafeRoughness = MakeRoughnessSafe(BSDF.Dielectric.Roughness.x);
 				float a2 = Pow4(SafeRoughness);
 				float Energy = 1.0f; // See EnergyNormalization
 
-				// Generalized microfacet specular
-				//float D = D_GGX(a2, Context.NoH) * Energy;
-				//float Vis = Vis_SmithJointApprox(a2, Context.NoV, SatNoL);
-				//float3 F = F_Schlick(F0, Context.VoH);
-
 				const float3 N = BSDF.Dielectric.Normal;
 				const float SatNoV = saturate(dot(N, V));
 				float3 R = 2 * dot(V, N) * N - V;
-
 				float3 EvalEnvBRDF = EnvBRDFApprox(F0, SafeRoughness, SatNoV);
-				float SkyAverageBrightness = 1.0f;
-				SpecularLuminance += GetSkyLightReflection(R, SafeRoughness, SkyAverageBrightness) * EvalEnvBRDF * BSDF.Dielectric.Tint;
+
+				SpecularDirection = R;
+				SpecularWeight = EvalEnvBRDF * BSDF.Dielectric.Tint;
+				SpecularSafeRoughness = a2;
 			}
 		}
 		break;
@@ -69,18 +68,14 @@ float3 StrataEnvironmentLighting(
 				float a2 = Pow4(SafeRoughness);
 				float Energy = 1.0f; // See EnergyNormalization
 
-				// Generalized microfacet specular
-				//float D = D_GGX(a2, Context.NoH) * Energy;
-				//float Vis = Vis_SmithJointApprox(a2, Context.NoV, SatNoL);
-				//float3 F = F_Schlick(BSDF.Conductor.Reflectivity, BSDF.Conductor.EdgeColor, Context.VoH); // STRATA_TODO check that this is correct 
-
 				const float3 N = BSDF.Conductor.Normal;
 				const float SatNoV = saturate(dot(N, V));
 				float3 R = 2 * dot(V, N) * N - V;
-
 				float3 EvalEnvBRDF = EnvBRDFApprox(BSDF.Conductor.Reflectivity, BSDF.Conductor.EdgeColor, SafeRoughness, SatNoV);
-				float SkyAverageBrightness = 1.0f;
-				SpecularLuminance += GetSkyLightReflection(R, SafeRoughness, SkyAverageBrightness) * EvalEnvBRDF;
+
+				SpecularDirection = R;
+				SpecularWeight = EvalEnvBRDF;
+				SpecularSafeRoughness = a2;
 			}
 		}
 		break;
@@ -102,35 +97,42 @@ float3 StrataEnvironmentLighting(
 			float3 NormLuminanceSingleScattering = (NormLuminanceTransfert - NormLuminanceTransfert * SafePathSegmentTransmittance) / SafeExtinctionCoefficients;
 
 			// No phase function because the SH already contains integration of BSDF with respect to hemisphere. Incorrect but needed.
-			Albedo = NormLuminanceSingleScattering;
+			DiffuseAlbedo = NormLuminanceSingleScattering;
 
 			// STRATA_TODO: This is to fake some directionality. Have normal specified on the material as for MaterialX?
-			Normal = BSDF.Volume.SmoothSurfaceNormal;
+			DiffuseNormal = BSDF.Volume.SmoothSurfaceNormal;
 		}
 		break;
 		}
 
 		// Diffuse component
-		if (any(Albedo > 0.0f))
+		if (any(DiffuseAlbedo > 0.0f))
 		{
 			// Apply bent normal if needed
 			#if APPLY_SKY_SHADOWING
 			{
-				Normal = UpsampleDFAO(BufferUV, GBuffer.Depth, Normal); // Bent normal
+				DiffuseNormal = UpsampleDFAO(BufferUV, GBuffer.Depth, DiffuseNormal); // Bent normal
 			}
 			#endif
 
 			// Compute the common sky visibility factors
-			FSkyLightVisibilityData SkyVisData = GetSkyLightVisibilityData(GBuffer, AmbientOcclusion, Normal);
+			FSkyLightVisibilityData SkyVisData = GetSkyLightVisibilityData(GBuffer, AmbientOcclusion, DiffuseNormal);
 
 			// Finally sample the sky diffuse contribution (spherical harmonic, Lambert BRDF)
-			float3 DiffuseLookup = GetSkySHDiffuse(Normal) * View.SkyLightColor.rgb;			// STRATA_TODO compute SH coefficients for ON or Chan
+			float3 DiffuseLookup = GetSkySHDiffuse(DiffuseNormal) * View.SkyLightColor.rgb;			// STRATA_TODO compute SH coefficients for ON or Chan
 			// And accumulate
-			SkyLighting += BSDFThroughput * (SkyVisData.SkyDiffuseLookUpMul * DiffuseLookup + SkyVisData.SkyDiffuseLookUpAdd) * Albedo;
+			SkyLighting += BSDFThroughput * (SkyVisData.SkyDiffuseLookUpMul * DiffuseLookup + SkyVisData.SkyDiffuseLookUpAdd) * DiffuseAlbedo;
 		}
 
 		// Specular component
-		SkyLighting += BSDFThroughput * SpecularLuminance * View.SkyLightColor.rgb * AmbientOcclusion;
+		if (any(SpecularWeight > 0.0f))
+		{
+			float SkyAverageBrightness = 1.0f;
+			const float3 SpecularLuminance = GetSkyLightReflection(SpecularDirection, SpecularSafeRoughness, SkyAverageBrightness);
+
+			SkyLighting += BSDFThroughput * SpecularLuminance * SpecularWeight * View.SkyLightColor.rgb * AmbientOcclusion;
+		}
+
 	}
 #if USE_PREEXPOSURE
 	SkyLighting *= View.PreExposure;
