Strata VGPRs optimisation.

Strata enable
SkyEnv   pass: Before 228VGPRs, now 144 (strata disable: 64 VGPRs)
DynLight pass: Before 216VGPRs, now 124 (strata disable: 64 VGPRs)

occupancy goes from 10 to 20% (good but the original one is 40%). There are still things we can do to get there.

diff --git a/Engine/Shaders/Private/Strata/Strata.ush b/Engine/Shaders/Private/Strata/Strata.ush
--- a/Engine/Shaders/Private/Strata/Strata.ush
+++ b/Engine/Shaders/Private/Strata/Strata.ush
@@ -20,45 +20,45 @@
 ///////////////////////////////////////////////////////////////////////////////
 // BSDFs representation
 
-struct FStrataDiffuse
-{
-	float3	Albedo;
-	float   Roughness;
-
-	float3	Normal;
-};
-
-struct FStrataDielectric
-{
-	float	IOR;
-	float3	Tint;
-	float2	Roughness;
-
-	float3	Normal;
-};
-
-struct FStrataConductor
-{
-	float3	Reflectivity;
-	float3	EdgeColor;
-	float2	Roughness;
-
-	float3	Normal;
-};
-
-struct FStrataVolume
-{
-	float3	Albedo;
-	float3	Extinction;
-	float	Anisotropy;
-	float	Thickness;
-
-	float3  SmoothSurfaceNormal;	// Typically vertex normal
-};
-
-
-
-// The largest Strata BSDF is FStrataConductor having as size of 24 bytes.
+//struct FStrataDiffuse
+//{
+//	float3	Albedo;
+//	float   Roughness;
+//
+//	float3	Normal;
+//};
+//
+//struct FStrataDielectric
+//{
+//	float	IOR;
+//	float3	Tint;
+//	float2	Roughness;
+//
+//	float3	Normal;
+//};
+//
+//struct FStrataConductor
+//{
+//	float3	Reflectivity;
+//	float3	EdgeColor;
+//	float2	Roughness;
+//
+//	float3	Normal;
+//};
+//
+//struct FStrataVolume
+//{
+//	float3	Albedo;
+//	float3	Extinction;
+//	float	Anisotropy;
+//	float	Thickness;
+//
+//	float3  SmoothSurfaceNormal;	// Typically vertex normal
+//};
+
+
+
+// The largest Strata BSDF is FStrataConductor having as size of 24 bytes (FStrataConductor)
 // So we will have 
 //		STRATA_MAX_TOTAL_BSDF * sizeof(FStrataConductor) +				// 4 * 24
 //		STRATA_MAX_TOTAL_BSDF * (sizeof(BSDFType) + sizeof(Weight))		// 4 * 8
@@ -69,13 +69,38 @@ struct FStrataBSDF
 	int					BSDFType;
 	float3				Weight;		// Weight is a float3 because, when loaded from the material buffer, it contains the BSDF throughput (coverage and transmittance combined)
 
-	// Even thought the compiler should optimise, it would be nice to have a union with reinterpretation of the register memory.
-	FStrataDiffuse		Diffuse;
-	FStrataDielectric	Dielectric;
-	FStrataConductor	Conductor;
-	FStrataVolume		Volume;
+	// We cannot have the structures above represented as a union in HLSL, so we enforce that using an array representing VGPRs and a bunch of preprocessor functions.
+	//union BSDFsUnion
+	//{
+	//	FStrataDiffuse		Diffuse;
+	//	FStrataDielectric	Dielectric;
+	//	FStrataConductor	Conductor;
+	//	FStrataVolume		Volume;
+	//};
+	//BSDFsUnion U;
+
+	float4 VGPRs[4];
 };
 
+#define DIFFUSE_ALBEDO(X)				X.VGPRs[0].xyz
+#define DIFFUSE_ROUGHNESS(X)			X.VGPRs[0].w
+#define DIFFUSE_NORMAL(X)				X.VGPRs[1].xyz
+
+#define DIELECTRIC_TINT(X)				X.VGPRs[0].xyz
+#define DIELECTRIC_IOR(X)				X.VGPRs[0].w
+#define DIELECTRIC_ROUGHNESS(X)			X.VGPRs[1].xy
+#define DIELECTRIC_NORMAL(X)			X.VGPRs[2].xyz
+
+#define CONDUCTOR_REFLECTIVITY(X)		X.VGPRs[0].xyz
+#define CONDUCTOR_EDGECOLOR(X)			X.VGPRs[1].xyz
+#define CONDUCTOR_ROUGHNESS(X)			X.VGPRs[2].xy
+#define CONDUCTOR_NORMAL(X)				X.VGPRs[3].xyz
+
+#define VOLUME_ALBEDO(X)				X.VGPRs[0].xyz
+#define VOLUME_EXTINCTION(X)			X.VGPRs[1].xyz
+#define VOLUME_ANISOTROPY(X)			X.VGPRs[2].x
+#define VOLUME_THICKNESS(X)				X.VGPRs[2].y
+#define VOLUME_SMOOTHSURFACENORMAL(X)	X.VGPRs[3].xyz
 
 
 struct FStrataLayer
@@ -106,9 +131,9 @@ FStrataData GetStrataDiffuseOrenNayarBSDF(float3 Albedo, float Roughness, float3
 
 	StrataData.Layers[0].BSDFCount++;
 	StrataData.Layers[0].BSDFs[0].BSDFType						= STRATA_BSDF_TYPE_DIFFUSE_ON;
-	StrataData.Layers[0].BSDFs[0].Diffuse.Albedo				= Albedo;
-	StrataData.Layers[0].BSDFs[0].Diffuse.Roughness				= Roughness;
-	StrataData.Layers[0].BSDFs[0].Diffuse.Normal				= Normal;
+	DIFFUSE_ALBEDO(StrataData.Layers[0].BSDFs[0])				= Albedo;
+	DIFFUSE_ROUGHNESS(StrataData.Layers[0].BSDFs[0])			= Roughness;
+	DIFFUSE_NORMAL(StrataData.Layers[0].BSDFs[0])				= Normal;
 
 	StrataData.Layers[0].BSDFs[0].Weight						= 1.0f;
 
@@ -122,9 +147,9 @@ FStrataData GetStrataDiffuseChanBSDF(float3 Albedo, float Roughness, float3 Norm
 
 	StrataData.Layers[0].BSDFCount++;
 	StrataData.Layers[0].BSDFs[0].BSDFType						= STRATA_BSDF_TYPE_DIFFUSE_CHAN;
-	StrataData.Layers[0].BSDFs[0].Diffuse.Albedo				= Albedo;
-	StrataData.Layers[0].BSDFs[0].Diffuse.Roughness				= Roughness;
-	StrataData.Layers[0].BSDFs[0].Diffuse.Normal				= Normal;
+	DIFFUSE_ALBEDO		(StrataData.Layers[0].BSDFs[0])			= Albedo;
+	DIFFUSE_ROUGHNESS	(StrataData.Layers[0].BSDFs[0])			= Roughness;
+	DIFFUSE_NORMAL		(StrataData.Layers[0].BSDFs[0])			= Normal;
 
 	StrataData.Layers[0].BSDFs[0].Weight						= 1.0f;
 
@@ -138,10 +163,10 @@ FStrataData GetStrataDielectricBSDF(float2 Roughness, float IOR, float3 Tint, fl
 
 	StrataData.Layers[0].BSDFCount++;
 	StrataData.Layers[0].BSDFs[0].BSDFType						= STRATA_BSDF_TYPE_DIELECTRIC;
-	StrataData.Layers[0].BSDFs[0].Dielectric.IOR				= IOR;
-	StrataData.Layers[0].BSDFs[0].Dielectric.Tint				= Tint;
-	StrataData.Layers[0].BSDFs[0].Dielectric.Roughness			= Roughness;
-	StrataData.Layers[0].BSDFs[0].Dielectric.Normal				= Normal;
+	DIELECTRIC_IOR			(StrataData.Layers[0].BSDFs[0])		= IOR;
+	DIELECTRIC_TINT			(StrataData.Layers[0].BSDFs[0])		= Tint;
+	DIELECTRIC_ROUGHNESS	(StrataData.Layers[0].BSDFs[0])		= Roughness;
+	DIELECTRIC_NORMAL		(StrataData.Layers[0].BSDFs[0])		= Normal;
 
 	StrataData.Layers[0].BSDFs[0].Weight						= 1.0f;
 
@@ -155,10 +180,10 @@ FStrataData GetStrataConductorBSDF(float3 Reflectivity, float3 EdgeColor, float2
 
 	StrataData.Layers[0].BSDFCount++;
 	StrataData.Layers[0].BSDFs[0].BSDFType						= STRATA_BSDF_TYPE_CONDUCTOR;
-	StrataData.Layers[0].BSDFs[0].Conductor.Reflectivity		= Reflectivity;
-	StrataData.Layers[0].BSDFs[0].Conductor.EdgeColor			= EdgeColor;
-	StrataData.Layers[0].BSDFs[0].Conductor.Roughness			= Roughness;
-	StrataData.Layers[0].BSDFs[0].Conductor.Normal				= Normal;
+	CONDUCTOR_REFLECTIVITY	(StrataData.Layers[0].BSDFs[0])		= Reflectivity;
+	CONDUCTOR_EDGECOLOR		(StrataData.Layers[0].BSDFs[0])		= EdgeColor;
+	CONDUCTOR_ROUGHNESS		(StrataData.Layers[0].BSDFs[0])		= Roughness;
+	CONDUCTOR_NORMAL		(StrataData.Layers[0].BSDFs[0])		= Normal;
 
 	StrataData.Layers[0].BSDFs[0].Weight						= 1.0f;
 
@@ -172,11 +197,11 @@ FStrataData GetStrataVolumeBSDF(float3 Albedo, float3 Extinction, float Anisotro
 
 	StrataData.Layers[0].BSDFCount++;
 	StrataData.Layers[0].BSDFs[0].BSDFType						= STRATA_BSDF_TYPE_VOLUME;
-	StrataData.Layers[0].BSDFs[0].Volume.Albedo					= Albedo;
-	StrataData.Layers[0].BSDFs[0].Volume.Extinction				= Extinction;
-	StrataData.Layers[0].BSDFs[0].Volume.Anisotropy				= Anisotropy;
-	StrataData.Layers[0].BSDFs[0].Volume.Thickness				= Thickness;
-	StrataData.Layers[0].BSDFs[0].Volume.SmoothSurfaceNormal	= SmoothSurfaceNormal;
+	VOLUME_ALBEDO				(StrataData.Layers[0].BSDFs[0])	= Albedo;
+	VOLUME_EXTINCTION			(StrataData.Layers[0].BSDFs[0])	= Extinction;
+	VOLUME_ANISOTROPY			(StrataData.Layers[0].BSDFs[0])	= Anisotropy;
+	VOLUME_THICKNESS			(StrataData.Layers[0].BSDFs[0])	= Thickness;
+	VOLUME_SMOOTHSURFACENORMAL	(StrataData.Layers[0].BSDFs[0])	= SmoothSurfaceNormal;
 
 	StrataData.Layers[0].BSDFs[0].Weight						= 1.0f;
 
@@ -259,32 +284,33 @@ FStrataData StrataVerticalLayering(FStrataData Top, FStrataData Base)
 // Functions used to sanitize BSDF before they are used for lighting (when forward) or stored to memory (base pass).
 // We do not normalize normals, this is done before the lighting step.
 
-void SanitizeStrataDiffuse(inout FStrataDiffuse BSDF)
+void SanitizeStrataDiffuse(inout FStrataBSDF BSDF)
 {
-	BSDF.Albedo = saturate(BSDF.Albedo);
-	BSDF.Roughness = saturate(BSDF.Roughness);
+	DIFFUSE_ALBEDO(BSDF) = saturate(DIFFUSE_ALBEDO(BSDF));
+	DIFFUSE_ROUGHNESS(BSDF) = saturate(DIFFUSE_ROUGHNESS(BSDF));
 }
 
-void SanitizeStrataDielectric(inout FStrataDielectric BSDF)
+void SanitizeStrataDielectric(inout FStrataBSDF BSDF)
 {
-	BSDF.IOR = clamp(BSDF.IOR, 1.0f, 5.0f);		// See https://en.wikipedia.org/wiki/List_of_refractive_indices
-	BSDF.Tint = saturate(BSDF.Tint);
-	BSDF.Roughness = saturate(BSDF.Roughness);
+
+	DIELECTRIC_IOR(BSDF) = clamp(DIELECTRIC_IOR(BSDF), 1.0f, 5.0f);		// See https://en.wikipedia.org/wiki/List_of_refractive_indices
+	DIELECTRIC_TINT(BSDF) = saturate(DIELECTRIC_TINT(BSDF));
+	DIELECTRIC_ROUGHNESS(BSDF) = saturate(DIELECTRIC_ROUGHNESS(BSDF));
 }
 
-void SanitizeStrataConductor(inout FStrataConductor BSDF)
+void SanitizeStrataConductor(inout FStrataBSDF BSDF)
 {
-	BSDF.Reflectivity = saturate(BSDF.Reflectivity);
-	BSDF.EdgeColor = saturate(BSDF.EdgeColor);
-	BSDF.Roughness = saturate(BSDF.Roughness);
+	CONDUCTOR_REFLECTIVITY(BSDF) = saturate(CONDUCTOR_REFLECTIVITY(BSDF));
+	CONDUCTOR_EDGECOLOR(BSDF) = saturate(CONDUCTOR_EDGECOLOR(BSDF));
+	CONDUCTOR_ROUGHNESS(BSDF) = saturate(CONDUCTOR_ROUGHNESS(BSDF));
 }
 
-void SanitizeStrataVolume(inout FStrataVolume BSDF)
+void SanitizeStrataVolume(inout FStrataBSDF BSDF)
 {
-	BSDF.Albedo = saturate(BSDF.Albedo);
-	BSDF.Extinction = max(0.0f, BSDF.Extinction);
-	BSDF.Anisotropy = max(0.0f, BSDF.Anisotropy);
-	BSDF.Thickness = max(0.0f, BSDF.Thickness);
+	VOLUME_ALBEDO(BSDF) = saturate(VOLUME_ALBEDO(BSDF));
+	VOLUME_EXTINCTION(BSDF) = max(0.0f, VOLUME_EXTINCTION(BSDF));
+	VOLUME_ANISOTROPY(BSDF) = max(0.0f, VOLUME_ANISOTROPY(BSDF));
+	VOLUME_THICKNESS(BSDF) = max(0.0f, VOLUME_THICKNESS(BSDF));
 }
 
 
@@ -466,13 +492,13 @@ void PackStrataOut(
 			{
 			case STRATA_BSDF_TYPE_DIFFUSE_ON:
 			case STRATA_BSDF_TYPE_DIFFUSE_CHAN:
-			{ N = BSDF.Diffuse.Normal;				SanitizeStrataDiffuse(BSDF.Diffuse);		break; }
+			{ N = DIFFUSE_NORMAL(BSDF);				SanitizeStrataDiffuse(BSDF);	break; }
 			case STRATA_BSDF_TYPE_DIELECTRIC:
-			{ N = BSDF.Dielectric.Normal;			SanitizeStrataDielectric(BSDF.Dielectric);	break; }
+			{ N = DIELECTRIC_NORMAL(BSDF);			SanitizeStrataDielectric(BSDF);	break; }
 			case STRATA_BSDF_TYPE_CONDUCTOR:
-			{ N = BSDF.Conductor.Normal;			SanitizeStrataConductor(BSDF.Conductor);	break; }
+			{ N = CONDUCTOR_NORMAL(BSDF);			SanitizeStrataConductor(BSDF);	break; }
 			case STRATA_BSDF_TYPE_VOLUME:
-			{ N = BSDF.Volume.SmoothSurfaceNormal;	SanitizeStrataVolume(BSDF.Volume);			break; }
+			{ N = VOLUME_SMOOTHSURFACENORMAL(BSDF);	SanitizeStrataVolume(BSDF);		break; }
 			}
 			N = normalize(N);
 
@@ -500,7 +526,7 @@ void PackStrataOut(
 
 			case STRATA_BSDF_TYPE_DIELECTRIC:
 			{
-				float F0 = DielectricIorToF0(BSDF.Dielectric.IOR);
+				float F0 = DielectricIorToF0(DIELECTRIC_IOR(BSDF));
 				float3 F = F_Schlick(F0, ContextNoLight.VoH);
 
 				LayerCoverage	+= BSDFCoverage;
@@ -510,7 +536,7 @@ void PackStrataOut(
 
 			case STRATA_BSDF_TYPE_CONDUCTOR:
 			{
-				float3 F = F_Schlick(BSDF.Conductor.Reflectivity, BSDF.Conductor.EdgeColor, ContextNoLight.VoH);
+				float3 F = F_Schlick(CONDUCTOR_REFLECTIVITY(BSDF), CONDUCTOR_EDGECOLOR(BSDF), ContextNoLight.VoH);
 
 				LayerCoverage	+= BSDFCoverage;
 				LayerThroughput += BSDFCoverage * OpaqueBSDFThroughput;
@@ -519,11 +545,11 @@ void PackStrataOut(
 
 			case STRATA_BSDF_TYPE_VOLUME:
 			{
-				const float3 Albedo					= BSDF.Volume.Albedo;
-				const float3 ExtinctionCoefficients = BSDF.Volume.Extinction;
+				const float3 Albedo					= VOLUME_ALBEDO(BSDF);
+				const float3 ExtinctionCoefficients = VOLUME_EXTINCTION(BSDF);
 				const float3 ScatteringCoefficients = Albedo * ExtinctionCoefficients;
 				const float3 AbsorptionCoefficients = ExtinctionCoefficients - ScatteringCoefficients;
-				const float  Thickness				= BSDF.Volume.Thickness;
+				const float  Thickness				= VOLUME_THICKNESS(BSDF);
 
 				const float3 SafeExtinctionThreshold = 0.000001f;
 				const float3 SafeExtinctionCoefficients = max(SafeExtinctionThreshold, ExtinctionCoefficients);
@@ -570,24 +596,24 @@ void PackStrataOut(
 			case STRATA_BSDF_TYPE_DIFFUSE_ON:
 			case STRATA_BSDF_TYPE_DIFFUSE_CHAN:
 			{
-				STRATA_STORE_UINT1(PackAlbedoLinearToGamma2(BSDF.Diffuse.Albedo));
-				STRATA_STORE_UINT2(PackR16G16B16A16F(float4(BSDF.Diffuse.Normal, BSDF.Diffuse.Roughness)));
+				STRATA_STORE_UINT1(PackAlbedoLinearToGamma2(DIFFUSE_ALBEDO(BSDF)));
+				STRATA_STORE_UINT2(PackR16G16B16A16F(float4(DIFFUSE_NORMAL(BSDF), DIFFUSE_ROUGHNESS(BSDF))));
 				// 12 bytes
 			}
 			break;
 			case STRATA_BSDF_TYPE_DIELECTRIC: //
 			{
-				STRATA_STORE_UINT1(PackR11G11B10F(BSDF.Dielectric.Tint));
-				STRATA_STORE_UINT1(PackR11G11B10F(float3(BSDF.Dielectric.Roughness, BSDF.Dielectric.IOR)));
-				STRATA_STORE_UINT2(PackR16G16B16A16F(float4(BSDF.Dielectric.Normal, 0.0f)));
+				STRATA_STORE_UINT1(PackR11G11B10F(DIELECTRIC_TINT(BSDF)));
+				STRATA_STORE_UINT1(PackR11G11B10F(float3(DIELECTRIC_ROUGHNESS(BSDF), DIELECTRIC_IOR(BSDF))));
+				STRATA_STORE_UINT2(PackR16G16B16A16F(float4(DIELECTRIC_NORMAL(BSDF), 0.0f)));
 				// 16 bytes
 			}
 			break;
 			case STRATA_BSDF_TYPE_CONDUCTOR:
 			{
-				STRATA_STORE_UINT2(PackR16G16B16A16F(float4(BSDF.Conductor.Reflectivity, BSDF.Conductor.Roughness.x)));
-				STRATA_STORE_UINT2(PackR16G16B16A16F(float4(BSDF.Conductor.EdgeColor, BSDF.Conductor.Roughness.y)));
-				STRATA_STORE_UINT2(PackR16G16B16A16F(float4(BSDF.Conductor.Normal, 0.0f)));
+				STRATA_STORE_UINT2(PackR16G16B16A16F(float4(CONDUCTOR_REFLECTIVITY(BSDF), CONDUCTOR_ROUGHNESS(BSDF).x)));
+				STRATA_STORE_UINT2(PackR16G16B16A16F(float4(CONDUCTOR_EDGECOLOR(BSDF), CONDUCTOR_ROUGHNESS(BSDF).y)));
+				STRATA_STORE_UINT2(PackR16G16B16A16F(float4(CONDUCTOR_NORMAL(BSDF), 0.0f)));
 				// 24 bytes
 			}
 			break;
@@ -595,10 +621,10 @@ void PackStrataOut(
 			{
 				// We simply multiply the absorption and scattering coefficient by the thickness to simplify the material to be a normalise medium of thickness 1 meter.
 				// This gives the same result and reduces memory usage.
-				const float Thickness = max(0.0f, BSDF.Volume.Thickness);
-				STRATA_STORE_UINT1(PackR11G11B10F(BSDF.Volume.Extinction * Thickness));
-				STRATA_STORE_UINT1(PackR11G11B10F(BSDF.Volume.Albedo));
-				STRATA_STORE_UINT2(PackR16G16B16A16F(float4(BSDF.Volume.SmoothSurfaceNormal, BSDF.Volume.Anisotropy)));
+				const float Thickness = VOLUME_THICKNESS(BSDF);
+				STRATA_STORE_UINT1(PackR11G11B10F(VOLUME_EXTINCTION(BSDF) * Thickness));
+				STRATA_STORE_UINT1(PackR11G11B10F(VOLUME_ALBEDO(BSDF)));
+				STRATA_STORE_UINT2(PackR16G16B16A16F(float4(VOLUME_SMOOTHSURFACENORMAL(BSDF), VOLUME_ANISOTROPY(BSDF))));
 				// 16 bytes
 			}
 			break;
@@ -643,30 +669,30 @@ FStrataPixelBSDFs UnpackStrataIn(ByteAddressBuffer StrataBuffer, uint ByteOffset
 		case STRATA_BSDF_TYPE_DIFFUSE_ON:
 		case STRATA_BSDF_TYPE_DIFFUSE_CHAN:
 		{
-			StrataLoad_AlbedoGamma2ToLinear	(StrataBuffer, ByteOffset, PixelBSDFs.BSDFs[i].Diffuse.Albedo);
-			StrataLoad_R16G16B16F_A16F		(StrataBuffer, ByteOffset, PixelBSDFs.BSDFs[i].Diffuse.Normal,			PixelBSDFs.BSDFs[i].Diffuse.Roughness);
+			StrataLoad_AlbedoGamma2ToLinear	(StrataBuffer, ByteOffset, DIFFUSE_ALBEDO(PixelBSDFs.BSDFs[i]));
+			StrataLoad_R16G16B16F_A16F		(StrataBuffer, ByteOffset, DIFFUSE_NORMAL(PixelBSDFs.BSDFs[i]),			DIFFUSE_ROUGHNESS(PixelBSDFs.BSDFs[i]));
 		}
 		break;
 		case STRATA_BSDF_TYPE_DIELECTRIC:
 		{
-			StrataLoad_R11G11B10F			(StrataBuffer, ByteOffset, PixelBSDFs.BSDFs[i].Dielectric.Tint);
-			StrataLoad_R11G11_B10F			(StrataBuffer, ByteOffset, PixelBSDFs.BSDFs[i].Dielectric.Roughness,	PixelBSDFs.BSDFs[i].Dielectric.IOR);
-			StrataLoad_R16G16B16F_A16F		(StrataBuffer, ByteOffset, PixelBSDFs.BSDFs[i].Dielectric.Normal,		DummyFloat);
+			StrataLoad_R11G11B10F			(StrataBuffer, ByteOffset, DIELECTRIC_TINT(PixelBSDFs.BSDFs[i]));
+			StrataLoad_R11G11_B10F			(StrataBuffer, ByteOffset, DIELECTRIC_ROUGHNESS(PixelBSDFs.BSDFs[i]),	DIELECTRIC_IOR(PixelBSDFs.BSDFs[i]));
+			StrataLoad_R16G16B16F_A16F		(StrataBuffer, ByteOffset, DIELECTRIC_NORMAL(PixelBSDFs.BSDFs[i]),		DummyFloat);
 		}
 		break;
 		case STRATA_BSDF_TYPE_CONDUCTOR:
 		{
-			StrataLoad_R16G16B16F_A16F		(StrataBuffer, ByteOffset, PixelBSDFs.BSDFs[i].Conductor.Reflectivity,	PixelBSDFs.BSDFs[i].Conductor.Roughness.x);
-			StrataLoad_R16G16B16F_A16F		(StrataBuffer, ByteOffset, PixelBSDFs.BSDFs[i].Conductor.EdgeColor,		PixelBSDFs.BSDFs[i].Conductor.Roughness.y);
-			StrataLoad_R16G16B16F_A16F		(StrataBuffer, ByteOffset, PixelBSDFs.BSDFs[i].Conductor.Normal,		DummyFloat);
+			StrataLoad_R16G16B16F_A16F		(StrataBuffer, ByteOffset, CONDUCTOR_REFLECTIVITY(PixelBSDFs.BSDFs[i]), CONDUCTOR_ROUGHNESS(PixelBSDFs.BSDFs[i]).x);
+			StrataLoad_R16G16B16F_A16F		(StrataBuffer, ByteOffset, CONDUCTOR_EDGECOLOR(PixelBSDFs.BSDFs[i]),	CONDUCTOR_ROUGHNESS(PixelBSDFs.BSDFs[i]).y);
+			StrataLoad_R16G16B16F_A16F		(StrataBuffer, ByteOffset, CONDUCTOR_NORMAL(PixelBSDFs.BSDFs[i]),		DummyFloat);
 		}
 		break;
 		case STRATA_BSDF_TYPE_VOLUME:
 		{
-			StrataLoad_R11G11B10F			(StrataBuffer, ByteOffset, PixelBSDFs.BSDFs[i].Volume.Extinction);
-			StrataLoad_R11G11B10F			(StrataBuffer, ByteOffset, PixelBSDFs.BSDFs[i].Volume.Albedo);
-			StrataLoad_R16G16B16F_A16F		(StrataBuffer, ByteOffset, PixelBSDFs.BSDFs[i].Volume.SmoothSurfaceNormal, PixelBSDFs.BSDFs[i].Volume.Anisotropy);
-			PixelBSDFs.BSDFs[i].Volume.Thickness = 1.0f;
+			StrataLoad_R11G11B10F			(StrataBuffer, ByteOffset, VOLUME_EXTINCTION(PixelBSDFs.BSDFs[i]));
+			StrataLoad_R11G11B10F			(StrataBuffer, ByteOffset, VOLUME_ALBEDO(PixelBSDFs.BSDFs[i]));
+			StrataLoad_R16G16B16F_A16F		(StrataBuffer, ByteOffset, VOLUME_SMOOTHSURFACENORMAL(PixelBSDFs.BSDFs[i]), VOLUME_ANISOTROPY(PixelBSDFs.BSDFs[i]));
+			VOLUME_THICKNESS(PixelBSDFs.BSDFs[i]) = 1.0f;
 		}
 		break;
 		}
diff --git a/Engine/Shaders/Private/Strata/StrataDeferredLighting.ush b/Engine/Shaders/Private/Strata/StrataDeferredLighting.ush
--- a/Engine/Shaders/Private/Strata/StrataDeferredLighting.ush
+++ b/Engine/Shaders/Private/Strata/StrataDeferredLighting.ush
@@ -28,13 +28,13 @@ float4 StrataDeferredLighting(
 		{
 		case STRATA_BSDF_TYPE_DIFFUSE_ON:
 		case STRATA_BSDF_TYPE_DIFFUSE_CHAN:
-		{ N = BSDF.Diffuse.Normal; break; }
+		{ N = DIFFUSE_NORMAL(BSDF);				break; }
 		case STRATA_BSDF_TYPE_DIELECTRIC:
-		{ N = BSDF.Dielectric.Normal; break; }
+		{ N = DIELECTRIC_NORMAL(BSDF);			break; }
 		case STRATA_BSDF_TYPE_CONDUCTOR:
-		{ N = BSDF.Conductor.Normal; break; }
+		{ N = CONDUCTOR_NORMAL(BSDF);			break; }
 		case STRATA_BSDF_TYPE_VOLUME:
-		{ N = BSDF.Volume.SmoothSurfaceNormal; break; }
+		{ N = VOLUME_SMOOTHSURFACENORMAL(BSDF);	break; }
 		}
 		N = normalize(N);
 
@@ -61,23 +61,23 @@ float4 StrataDeferredLighting(
 		{
 		case STRATA_BSDF_TYPE_DIFFUSE_ON:
 		{
-			const float SafeRoughness = MakeRoughnessSafe(BSDF.Diffuse.Roughness);
-			Radiance.rgb += Diffuse_OrenNayar(BSDF.Diffuse.Albedo, SafeRoughness, Context.NoV, SatNoL, Context.VoH) * DirLightGeometryFactor;
+			const float SafeRoughness = MakeRoughnessSafe(DIFFUSE_ROUGHNESS(BSDF));
+			Radiance.rgb += Diffuse_OrenNayar(DIFFUSE_ALBEDO(BSDF), SafeRoughness, Context.NoV, SatNoL, Context.VoH) * DirLightGeometryFactor;
 		}
 		break;
 		case STRATA_BSDF_TYPE_DIFFUSE_CHAN:
 		{
-			const float SafeRoughness = MakeRoughnessSafe(BSDF.Diffuse.Roughness);
+			const float SafeRoughness = MakeRoughnessSafe(DIFFUSE_ROUGHNESS(BSDF));
 			float a = SafeRoughness * SafeRoughness;
 			float a2 = a * a;
-			Radiance.rgb += Diffuse_Chan(BSDF.Diffuse.Albedo, a2, Context.NoV, SatNoL, Context.VoH, Context.NoH) * DirLightGeometryFactor;
+			Radiance.rgb += Diffuse_Chan(DIFFUSE_ALBEDO(BSDF), a2, Context.NoV, SatNoL, Context.VoH, Context.NoH) * DirLightGeometryFactor;
 		}
 		break;
 		case STRATA_BSDF_TYPE_DIELECTRIC:
 		{
-			float F0 = DielectricIorToF0(BSDF.Dielectric.IOR);
+			float F0 = DielectricIorToF0(DIELECTRIC_IOR(BSDF));
 
-			const float SafeRoughness = MakeRoughnessSafe(BSDF.Dielectric.Roughness.x);
+			const float SafeRoughness = MakeRoughnessSafe(DIELECTRIC_ROUGHNESS(BSDF).x);
 			float a2 = Pow4(SafeRoughness);
 			float Energy = 1.0f; // See EnergyNormalization
 
@@ -86,32 +86,33 @@ float4 StrataDeferredLighting(
 			float Vis = Vis_SmithJointApprox(a2, Context.NoV, SatNoL);
 			float3 F = F_Schlick(F0, Context.VoH);
 
-			Radiance.rgb += D * Vis * F * BSDF.Dielectric.Tint * DirLightGeometryFactor;
+			Radiance.rgb += D * Vis * F * DIELECTRIC_TINT(BSDF) * DirLightGeometryFactor;
 		}
 		break;
 		case STRATA_BSDF_TYPE_CONDUCTOR:
 		{
 			// See SpecularGGX, isotropic verion
 
-			const float SafeRoughness = MakeRoughnessSafe(BSDF.Conductor.Roughness.x);
+			const float SafeRoughness = MakeRoughnessSafe(CONDUCTOR_ROUGHNESS(BSDF).x);
 			float a2 = Pow4(SafeRoughness);
 			float Energy = 1.0f; // See EnergyNormalization
 
 			// Generalized microfacet specular
 			float D = D_GGX(a2, Context.NoH) * Energy;
 			float Vis = Vis_SmithJointApprox(a2, Context.NoV, SatNoL);
-			float3 F = F_Schlick(BSDF.Conductor.Reflectivity, BSDF.Conductor.EdgeColor, Context.VoH);
+			float3 F = F_Schlick(CONDUCTOR_REFLECTIVITY(BSDF), CONDUCTOR_EDGECOLOR(BSDF), Context.VoH);
 
 			Radiance.rgb += D * Vis * F * DirLightGeometryFactor;
 		}
 		break;
 		case STRATA_BSDF_TYPE_VOLUME:
 		{
-			const float3 Albedo = BSDF.Volume.Albedo;
-			const float3 ExtinctionCoefficients = BSDF.Volume.Extinction;
+			const float3 Albedo = VOLUME_ALBEDO(BSDF);
+			const float3 ExtinctionCoefficients = VOLUME_EXTINCTION(BSDF);
 			const float3 ScatteringCoefficients = Albedo * ExtinctionCoefficients;
 			const float3 AbsorptionCoefficients = ExtinctionCoefficients - ScatteringCoefficients;
-			const float  Thickness = BSDF.Volume.Thickness;
+			const float  Thickness = VOLUME_THICKNESS(BSDF);
+			const float  Phase = HenyeyGreensteinPhase(VOLUME_ANISOTROPY(BSDF), Context.VoL);
 
 			const float3 SafeExtinctionThreshold = 0.000001f;
 			const float3 SafeExtinctionCoefficients = max(SafeExtinctionThreshold, ExtinctionCoefficients);
@@ -128,7 +129,6 @@ float4 StrataDeferredLighting(
 			const float3 SelfShadowTransmittance = SatNoL <= 0.0 ? 0.0 : exp(-ExtinctionCoefficients * LightPathLength);
 #endif
 
-			const float Phase = HenyeyGreensteinPhase(BSDF.Volume.Anisotropy, Context.VoL);
 			Radiance.rgb += NormLuminanceSingleScattering * Phase * DirLightBaseFactor * SelfShadowTransmittance;
 		}
 		break;
diff --git a/Engine/Shaders/Private/Strata/StrataEnvironmentLighting.ush b/Engine/Shaders/Private/Strata/StrataEnvironmentLighting.ush
--- a/Engine/Shaders/Private/Strata/StrataEnvironmentLighting.ush
+++ b/Engine/Shaders/Private/Strata/StrataEnvironmentLighting.ush
@@ -32,8 +32,8 @@ float3 StrataEnvironmentLighting(
 		case STRATA_BSDF_TYPE_DIFFUSE_ON:
 		case STRATA_BSDF_TYPE_DIFFUSE_CHAN:
 		{
-			DiffuseAlbedo = BSDF.Diffuse.Albedo;
-			DiffuseNormal = BSDF.Diffuse.Normal;
+			DiffuseAlbedo = DIFFUSE_ALBEDO(BSDF);
+			DiffuseNormal = DIFFUSE_NORMAL(BSDF);
 		}
 		break;
 		case STRATA_BSDF_TYPE_DIELECTRIC:
@@ -41,18 +41,18 @@ float3 StrataEnvironmentLighting(
 			BRANCH
 			if (ReflectionStruct.SkyLightParameters.y > 0)
 			{
-				float F0 = DielectricIorToF0(BSDF.Dielectric.IOR);
-				const float SafeRoughness = MakeRoughnessSafe(BSDF.Dielectric.Roughness.x);
+				float F0 = DielectricIorToF0(DIELECTRIC_IOR(BSDF));
+				const float SafeRoughness = MakeRoughnessSafe(DIELECTRIC_ROUGHNESS(BSDF).x);
 				float a2 = Pow4(SafeRoughness);
 				float Energy = 1.0f; // See EnergyNormalization
 
-				const float3 N = BSDF.Dielectric.Normal;
+				const float3 N = DIELECTRIC_NORMAL(BSDF);
 				const float SatNoV = saturate(dot(N, V));
 				float3 R = 2 * dot(V, N) * N - V;
 				float3 EvalEnvBRDF = EnvBRDFApprox(F0, SafeRoughness, SatNoV);
 
 				SpecularDirection = R;
-				SpecularWeight = EvalEnvBRDF * BSDF.Dielectric.Tint;
+				SpecularWeight = EvalEnvBRDF * DIELECTRIC_TINT(BSDF);
 				SpecularSafeRoughness = a2;
 			}
 		}
@@ -64,14 +64,14 @@ float3 StrataEnvironmentLighting(
 			{
 				// See SpecularGGX, isotropic verion
 
-				const float SafeRoughness = MakeRoughnessSafe(BSDF.Conductor.Roughness.x);
+				const float SafeRoughness = MakeRoughnessSafe(CONDUCTOR_ROUGHNESS(BSDF).x);
 				float a2 = Pow4(SafeRoughness);
 				float Energy = 1.0f; // See EnergyNormalization
 
-				const float3 N = BSDF.Conductor.Normal;
+				const float3 N = CONDUCTOR_NORMAL(BSDF);
 				const float SatNoV = saturate(dot(N, V));
 				float3 R = 2 * dot(V, N) * N - V;
-				float3 EvalEnvBRDF = EnvBRDFApprox(BSDF.Conductor.Reflectivity, BSDF.Conductor.EdgeColor, SafeRoughness, SatNoV);
+				float3 EvalEnvBRDF = EnvBRDFApprox(CONDUCTOR_REFLECTIVITY(BSDF), CONDUCTOR_EDGECOLOR(BSDF), SafeRoughness, SatNoV);
 
 				SpecularDirection = R;
 				SpecularWeight = EvalEnvBRDF;
@@ -81,16 +81,16 @@ float3 StrataEnvironmentLighting(
 		break;
 		case STRATA_BSDF_TYPE_VOLUME:
 		{
-			const float3 VolumeAlbedo = BSDF.Volume.Albedo;
-			const float3 ExtinctionCoefficients = BSDF.Volume.Extinction;
+			const float3 VolumeAlbedo = VOLUME_ALBEDO(BSDF);
+			const float3 ExtinctionCoefficients = VOLUME_EXTINCTION(BSDF);
 			const float3 ScatteringCoefficients = VolumeAlbedo * ExtinctionCoefficients;
 			const float3 AbsorptionCoefficients = ExtinctionCoefficients - ScatteringCoefficients;
-			const float  Thickness = BSDF.Volume.Thickness;
+			const float  Thickness = VOLUME_THICKNESS(BSDF);
 
 			const float3 SafeExtinctionThreshold = 0.000001f;
 			const float3 SafeExtinctionCoefficients = max(SafeExtinctionThreshold, ExtinctionCoefficients);
 
-			const float SatNoV = saturate(dot(BSDF.Volume.SmoothSurfaceNormal, V));
+			const float SatNoV = saturate(dot(VOLUME_SMOOTHSURFACENORMAL(BSDF), V));
 			const float PathLength = Thickness / max(0.0001f, SatNoV);
 			const float3 SafePathSegmentTransmittance = exp(-SafeExtinctionCoefficients * PathLength);
 			const float3 NormLuminanceTransfert = 1.0f * ScatteringCoefficients;
@@ -100,7 +100,7 @@ float3 StrataEnvironmentLighting(
 			DiffuseAlbedo = NormLuminanceSingleScattering;
 
 			// STRATA_TODO: This is to fake some directionality. Have normal specified on the material as for MaterialX?
-			DiffuseNormal = BSDF.Volume.SmoothSurfaceNormal;
+			DiffuseNormal = VOLUME_SMOOTHSURFACENORMAL(BSDF);
 		}
 		break;
 		}
diff --git a/Engine/Shaders/Private/Strata/StrataForwardLighting.ush b/Engine/Shaders/Private/Strata/StrataForwardLighting.ush
--- a/Engine/Shaders/Private/Strata/StrataForwardLighting.ush
+++ b/Engine/Shaders/Private/Strata/StrataForwardLighting.ush
@@ -37,13 +37,13 @@ float3 StrataForwardLighting(
 			{
 			case STRATA_BSDF_TYPE_DIFFUSE_ON:
 			case STRATA_BSDF_TYPE_DIFFUSE_CHAN:
-			{ N = BSDF.Diffuse.Normal;				SanitizeStrataDiffuse(BSDF.Diffuse);		break; }
+			{ N = DIFFUSE_NORMAL(BSDF);				SanitizeStrataDiffuse(BSDF);	break; }
 			case STRATA_BSDF_TYPE_DIELECTRIC:
-			{ N = BSDF.Dielectric.Normal;			SanitizeStrataDielectric(BSDF.Dielectric);	break; }
+			{ N = DIELECTRIC_NORMAL(BSDF);			SanitizeStrataDielectric(BSDF);	break; }
 			case STRATA_BSDF_TYPE_CONDUCTOR:
-			{ N = BSDF.Conductor.Normal;			SanitizeStrataConductor(BSDF.Conductor);	break; }
+			{ N = CONDUCTOR_NORMAL(BSDF);			SanitizeStrataConductor(BSDF);	break; }
 			case STRATA_BSDF_TYPE_VOLUME:
-			{ N = BSDF.Volume.SmoothSurfaceNormal;	SanitizeStrataVolume(BSDF.Volume);		break; }
+			{ N = VOLUME_SMOOTHSURFACENORMAL(BSDF);	SanitizeStrataVolume(BSDF);		break; }
 			}
 			N = normalize(N);
 			float3 R = 2 * dot(V, N) * N - V;
@@ -60,14 +60,14 @@ float3 StrataForwardLighting(
 			{
 			case STRATA_BSDF_TYPE_DIFFUSE_ON:
 			{
-				const float SafeRoughness = MakeRoughnessSafe(BSDF.Diffuse.Roughness);
-				Color += Diffuse_OrenNayar(BSDF.Diffuse.Albedo, SafeRoughness, Context.NoV, SatNoL, Context.VoH) * DirLightGeometryFactor;
+				const float SafeRoughness = MakeRoughnessSafe(DIFFUSE_ROUGHNESS(BSDF));
+				Color += Diffuse_OrenNayar(DIFFUSE_ALBEDO(BSDF), SafeRoughness, Context.NoV, SatNoL, Context.VoH) * DirLightGeometryFactor;
 
 #if ENABLE_SKY_LIGHT
 				BRANCH
 				if (ReflectionStruct.SkyLightParameters.y > 0)
 				{
-					Color += GetEffectiveSkySHDiffuse(N) * BSDF.Diffuse.Albedo * SkyLightCommonFactor; // STRATA_TODO compute SH coefficients for ON
+					Color += GetEffectiveSkySHDiffuse(N) * DIFFUSE_ALBEDO(BSDF) * SkyLightCommonFactor; // STRATA_TODO compute SH coefficients for ON
 				}
 #endif
 
@@ -78,16 +78,16 @@ float3 StrataForwardLighting(
 
 			case STRATA_BSDF_TYPE_DIFFUSE_CHAN:
 			{
-				const float SafeRoughness = MakeRoughnessSafe(BSDF.Diffuse.Roughness);
+				const float SafeRoughness = MakeRoughnessSafe(DIFFUSE_ROUGHNESS(BSDF));
 				float a = SafeRoughness * SafeRoughness;
 				float a2 = a * a;
-				Color += Diffuse_Chan(BSDF.Diffuse.Albedo, a2, Context.NoV, SatNoL, Context.VoH, Context.NoH) * DirLightGeometryFactor;
+				Color += Diffuse_Chan(DIFFUSE_ALBEDO(BSDF), a2, Context.NoV, SatNoL, Context.VoH, Context.NoH) * DirLightGeometryFactor;
 
 #if ENABLE_SKY_LIGHT
 				BRANCH
 				if (ReflectionStruct.SkyLightParameters.y > 0)
 				{
-					Color += GetEffectiveSkySHDiffuse(N) * BSDF.Diffuse.Albedo * SkyLightCommonFactor; // STRATA_TODO compute SH coefficients for Chan
+					Color += GetEffectiveSkySHDiffuse(N) * DIFFUSE_ALBEDO(BSDF) * SkyLightCommonFactor; // STRATA_TODO compute SH coefficients for Chan
 				}
 #endif
 				
@@ -98,9 +98,9 @@ float3 StrataForwardLighting(
 
 			case STRATA_BSDF_TYPE_DIELECTRIC:
 			{
-				float F0 = DielectricIorToF0(BSDF.Dielectric.IOR);
+				float F0 = DielectricIorToF0(DIELECTRIC_IOR(BSDF));
 
-				const float SafeRoughness = MakeRoughnessSafe(BSDF.Dielectric.Roughness.x);
+				const float SafeRoughness = MakeRoughnessSafe(DIELECTRIC_ROUGHNESS(BSDF).x);
 				float a2 = Pow4(SafeRoughness);
 				float Energy = 1.0f; // See EnergyNormalization
 
@@ -109,7 +109,7 @@ float3 StrataForwardLighting(
 				float Vis = Vis_SmithJointApprox(a2, Context.NoV, SatNoL);
 				float3 F = F_Schlick(F0, Context.VoH);
 
-				Color += D * Vis * F * BSDF.Dielectric.Tint * DirLightGeometryFactor;
+				Color += D * Vis * F * DIELECTRIC_TINT(BSDF) * DirLightGeometryFactor;
 
 #if ENABLE_SKY_LIGHT
 				BRANCH
@@ -117,7 +117,7 @@ float3 StrataForwardLighting(
 					{
 						float3 EvalEnvBRDF = EnvBRDFApprox(F0, SafeRoughness, Context.NoV);
 						float SkyAverageBrightness = 1.0f;
-						Color += GetSkyLightReflection(R, SafeRoughness, SkyAverageBrightness) * EvalEnvBRDF * BSDF.Dielectric.Tint * SkyLightCommonFactor;
+						Color += GetSkyLightReflection(R, SafeRoughness, SkyAverageBrightness) * EvalEnvBRDF * DIELECTRIC_TINT(BSDF) * SkyLightCommonFactor;
 					}
 #endif
 				
@@ -130,14 +130,14 @@ float3 StrataForwardLighting(
 			{
 				// See SpecularGGX, isotropic verion
 
-				const float SafeRoughness = MakeRoughnessSafe(BSDF.Conductor.Roughness.x);
+				const float SafeRoughness = MakeRoughnessSafe(CONDUCTOR_ROUGHNESS(BSDF).x);
 				float a2 = Pow4(SafeRoughness);
 				float Energy = 1.0f; // See EnergyNormalization
 
 				// Generalized microfacet specular
 				float D = D_GGX(a2, Context.NoH) * Energy;
 				float Vis = Vis_SmithJointApprox(a2, Context.NoV, SatNoL);
-				float3 F = F_Schlick(BSDF.Conductor.Reflectivity, BSDF.Conductor.EdgeColor, Context.VoH);
+				float3 F = F_Schlick(CONDUCTOR_REFLECTIVITY(BSDF), CONDUCTOR_EDGECOLOR(BSDF), Context.VoH);
 
 				Color += D * Vis * F * DirLightGeometryFactor;
 
@@ -145,7 +145,7 @@ float3 StrataForwardLighting(
 				//BRANCH
 				if (ReflectionStruct.SkyLightParameters.y > 0)
 				{
-					float3 EvalEnvBRDF = EnvBRDFApprox(BSDF.Conductor.Reflectivity, BSDF.Conductor.EdgeColor, SafeRoughness, Context.NoV);
+					float3 EvalEnvBRDF = EnvBRDFApprox(CONDUCTOR_REFLECTIVITY(BSDF), CONDUCTOR_EDGECOLOR(BSDF), SafeRoughness, Context.NoV);
 					float SkyAverageBrightness = 1.0f;
 					Color += GetSkyLightReflection(R, SafeRoughness, SkyAverageBrightness) * EvalEnvBRDF * SkyLightCommonFactor;
 				}
@@ -158,11 +158,12 @@ float3 StrataForwardLighting(
 
 			case STRATA_BSDF_TYPE_VOLUME:
 			{
-				const float3 Albedo = BSDF.Volume.Albedo;
-				const float3 ExtinctionCoefficients = BSDF.Volume.Extinction;
+				const float3 Albedo = VOLUME_ALBEDO(BSDF);
+				const float3 ExtinctionCoefficients = VOLUME_EXTINCTION(BSDF);
 				const float3 ScatteringCoefficients = Albedo * ExtinctionCoefficients;
 				const float3 AbsorptionCoefficients = ExtinctionCoefficients - ScatteringCoefficients;
-				const float  Thickness = BSDF.Volume.Thickness;
+				const float  Thickness = VOLUME_THICKNESS(BSDF);
+				const float  Phase = HenyeyGreensteinPhase(VOLUME_ANISOTROPY(BSDF), Context.VoL);
 
 				const float3 SafeExtinctionThreshold = 0.000001f;
 				const float3 SafeExtinctionCoefficients = max(SafeExtinctionThreshold, ExtinctionCoefficients);
@@ -179,7 +180,6 @@ float3 StrataForwardLighting(
 				const float3 SelfShadowTransmittance = SatNoL <= 0.0 ? 0.0 : exp(-ExtinctionCoefficients * LightPathLength);
 #endif
 
-				const float Phase = HenyeyGreensteinPhase(BSDF.Volume.Anisotropy, Context.VoL);
 				Color += NormLuminanceSingleScattering * Phase * DirLightBaseFactor * SelfShadowTransmittance;
 
 #if ENABLE_SKY_LIGHT
@@ -188,7 +188,7 @@ float3 StrataForwardLighting(
 				{
 					const float RescaleWithPhase = 2.0f * PI * IsotropicPhase(); // times 2PI to recover luminance, and then times the phase function for the scattering event
 					float3 EvalEnvBRDF = NormLuminanceSingleScattering * RescaleWithPhase;
-					float SafeRoughnessFromG = clamp(1.0 - abs(BSDF.Volume.Anisotropy), 0.001, 1.0f);
+					float SafeRoughnessFromG = clamp(1.0 - abs(VOLUME_ANISOTROPY(BSDF)), 0.001, 1.0f); // STRATA_TODO map G to roughness correctly
 					float SkyAverageBrightness = 1.0f;
 					// STRATA_TODO 
 					//	- Using R is weird. We should use V. But for rought material when light is incoming mainly from upper hemisphere, using N would make more send.
diff --git a/Engine/Shaders/Private/Strata/StrataVisualize.usf b/Engine/Shaders/Private/Strata/StrataVisualize.usf
--- a/Engine/Shaders/Private/Strata/StrataVisualize.usf
+++ b/Engine/Shaders/Private/Strata/StrataVisualize.usf
@@ -108,10 +108,10 @@ void VisualizeMaterialPS(
 				NEXTLINE;
 
 				HORIZONTALTAB; PRINTPARAMCHAR(_A_); PRINTPARAMCHAR(_L_); PRINTPARAMCHAR(_B_); PRINTPARAMCHAR(_E_); PRINTPARAMCHAR(_D_); PRINTPARAMCHAR(_O_);
-				HORIZONTALTAB; PRINTFLOAT3(BSDF.Diffuse.Albedo);
+				HORIZONTALTAB; PRINTFLOAT3(DIFFUSE_ALBEDO(BSDF));
 				NEXTLINE;
 				HORIZONTALTAB; PRINTPARAMCHAR(_R_); PRINTPARAMCHAR(_O_); PRINTPARAMCHAR(_U_); PRINTPARAMCHAR(_G_); PRINTPARAMCHAR(_H_); PRINTPARAMCHAR(_SPC_);
-				HORIZONTALTAB; PRINTFLOAT1(BSDF.Diffuse.Roughness);
+				HORIZONTALTAB; PRINTFLOAT1(DIFFUSE_ROUGHNESS(BSDF));
 				NEXTLINE;
 
 				break;
@@ -123,10 +123,10 @@ void VisualizeMaterialPS(
 				NEXTLINE;
 
 				HORIZONTALTAB; PRINTPARAMCHAR(_A_); PRINTPARAMCHAR(_L_); PRINTPARAMCHAR(_B_); PRINTPARAMCHAR(_E_); PRINTPARAMCHAR(_D_); PRINTPARAMCHAR(_O_); 
-				HORIZONTALTAB; PRINTFLOAT3(BSDF.Diffuse.Albedo);
+				HORIZONTALTAB; PRINTFLOAT3(DIFFUSE_ALBEDO(BSDF));
 				NEXTLINE;
 				HORIZONTALTAB; PRINTPARAMCHAR(_R_); PRINTPARAMCHAR(_O_); PRINTPARAMCHAR(_U_); PRINTPARAMCHAR(_G_); PRINTPARAMCHAR(_H_); PRINTPARAMCHAR(_SPC_);
-				HORIZONTALTAB; PRINTFLOAT1(BSDF.Diffuse.Roughness);
+				HORIZONTALTAB; PRINTFLOAT1(DIFFUSE_ROUGHNESS(BSDF));
 				NEXTLINE;
 
 				break;
@@ -138,13 +138,13 @@ void VisualizeMaterialPS(
 				NEXTLINE;
 
 				HORIZONTALTAB; PRINTPARAMCHAR(_I_); PRINTPARAMCHAR(_O_); PRINTPARAMCHAR(_R_); PRINTPARAMCHAR(_SPC_); PRINTPARAMCHAR(_SPC_); PRINTPARAMCHAR(_SPC_);
-				HORIZONTALTAB; PRINTFLOAT1(BSDF.Dielectric.IOR);
+				HORIZONTALTAB; PRINTFLOAT1(DIELECTRIC_IOR(BSDF));
 				NEXTLINE;
 				HORIZONTALTAB; PRINTPARAMCHAR(_T_); PRINTPARAMCHAR(_I_); PRINTPARAMCHAR(_N_); PRINTPARAMCHAR(_T_); PRINTPARAMCHAR(_SPC_); PRINTPARAMCHAR(_SPC_);
-				HORIZONTALTAB; PRINTFLOAT3(BSDF.Dielectric.Tint);
+				HORIZONTALTAB; PRINTFLOAT3(DIELECTRIC_TINT(BSDF));
 				NEXTLINE;
 				HORIZONTALTAB; PRINTPARAMCHAR(_R_); PRINTPARAMCHAR(_O_); PRINTPARAMCHAR(_U_); PRINTPARAMCHAR(_G_); PRINTPARAMCHAR(_H_); PRINTPARAMCHAR(_SPC_);
-				HORIZONTALTAB; PRINTFLOAT1(BSDF.Dielectric.Roughness.x);
+				HORIZONTALTAB; PRINTFLOAT1(DIELECTRIC_ROUGHNESS(BSDF).x);
 				NEXTLINE;
 
 				break;
@@ -156,13 +156,13 @@ void VisualizeMaterialPS(
 				NEXTLINE;
 
 				HORIZONTALTAB; PRINTPARAMCHAR(_R_); PRINTPARAMCHAR(_E_); PRINTPARAMCHAR(_F_); PRINTPARAMCHAR(_L_); PRINTPARAMCHAR(_E_); PRINTPARAMCHAR(_C_);
-				HORIZONTALTAB; PRINTFLOAT3(BSDF.Conductor.Reflectivity);
+				HORIZONTALTAB; PRINTFLOAT3(CONDUCTOR_REFLECTIVITY(BSDF));
 				NEXTLINE;
 				HORIZONTALTAB; PRINTPARAMCHAR(_E_); PRINTPARAMCHAR(_D_); PRINTPARAMCHAR(_G_); PRINTPARAMCHAR(_E_); PRINTPARAMCHAR(_C_); PRINTPARAMCHAR(_SPC_);
-				HORIZONTALTAB; PRINTFLOAT3(BSDF.Conductor.EdgeColor);
+				HORIZONTALTAB; PRINTFLOAT3(CONDUCTOR_EDGECOLOR(BSDF));
 				NEXTLINE;
 				HORIZONTALTAB; PRINTPARAMCHAR(_R_); PRINTPARAMCHAR(_O_); PRINTPARAMCHAR(_U_); PRINTPARAMCHAR(_G_); PRINTPARAMCHAR(_H_); PRINTPARAMCHAR(_SPC_);
-				HORIZONTALTAB; PRINTFLOAT1(BSDF.Conductor.Roughness.x);
+				HORIZONTALTAB; PRINTFLOAT1(CONDUCTOR_ROUGHNESS(BSDF).x);
 				NEXTLINE;
 
 				break;
@@ -174,16 +174,16 @@ void VisualizeMaterialPS(
 				NEXTLINE;
 				
 				HORIZONTALTAB; PRINTPARAMCHAR(_A_); PRINTPARAMCHAR(_L_); PRINTPARAMCHAR(_B_); PRINTPARAMCHAR(_E_); PRINTPARAMCHAR(_D_); PRINTPARAMCHAR(_O_); 
-				HORIZONTALTAB; PRINTFLOAT3(BSDF.Volume.Albedo);
+				HORIZONTALTAB; PRINTFLOAT3(VOLUME_ALBEDO(BSDF));
 				NEXTLINE;
 				HORIZONTALTAB; PRINTPARAMCHAR(_E_); PRINTPARAMCHAR(_X_); PRINTPARAMCHAR(_T_); PRINTPARAMCHAR(_I_); PRINTPARAMCHAR(_N_); PRINTPARAMCHAR(_C_); 
-				HORIZONTALTAB; PRINTFLOAT3(BSDF.Volume.Extinction);
+				HORIZONTALTAB; PRINTFLOAT3(VOLUME_EXTINCTION(BSDF));
 				NEXTLINE;
 				HORIZONTALTAB; PRINTPARAMCHAR(_A_); PRINTPARAMCHAR(_N_); PRINTPARAMCHAR(_I_); PRINTPARAMCHAR(_S_); PRINTPARAMCHAR(_O_); PRINTPARAMCHAR(_T_); 
-				HORIZONTALTAB; SIGNOFFLOAT1(BSDF.Volume.Anisotropy); PRINTFLOAT1(abs(BSDF.Volume.Anisotropy));
+				HORIZONTALTAB; SIGNOFFLOAT1(VOLUME_ANISOTROPY(BSDF)); PRINTFLOAT1(abs(VOLUME_ANISOTROPY(BSDF)));
 				NEXTLINE;
 				HORIZONTALTAB; PRINTPARAMCHAR(_T_); PRINTPARAMCHAR(_H_); PRINTPARAMCHAR(_I_); PRINTPARAMCHAR(_C_); PRINTPARAMCHAR(_K_); PRINTPARAMCHAR(_N_); 
-				HORIZONTALTAB; PRINTFLOAT1(BSDF.Volume.Thickness*1000.0f); PRINTPARAMCHAR(_SPC_); PRINTPARAMCHAR(_M_); PRINTPARAMCHAR(_I_); PRINTPARAMCHAR(_L_); PRINTPARAMCHAR(_L_); PRINTPARAMCHAR(_I_);
+				HORIZONTALTAB; PRINTFLOAT1(VOLUME_THICKNESS(BSDF) *1000.0f); PRINTPARAMCHAR(_SPC_); PRINTPARAMCHAR(_M_); PRINTPARAMCHAR(_I_); PRINTPARAMCHAR(_L_); PRINTPARAMCHAR(_L_); PRINTPARAMCHAR(_I_);
 				NEXTLINE;
 
 				break;
