Fixed roughness on strata materia forl environment pass.

diff --git a/Engine/Shaders/Private/Strata/StrataEnvironmentLighting.ush b/Engine/Shaders/Private/Strata/StrataEnvironmentLighting.ush
--- a/Engine/Shaders/Private/Strata/StrataEnvironmentLighting.ush
+++ b/Engine/Shaders/Private/Strata/StrataEnvironmentLighting.ush
@@ -43,8 +43,6 @@ float3 StrataEnvironmentLighting(
 			{
 				float F0 = DielectricIorToF0(DIELECTRIC_IOR(BSDF));
 				const float SafeRoughness = MakeRoughnessSafe(DIELECTRIC_ROUGHNESS(BSDF).x);
-				float a2 = Pow4(SafeRoughness);
-				float Energy = 1.0f; // See EnergyNormalization
 
 				const float3 N = DIELECTRIC_NORMAL(BSDF);
 				const float SatNoV = saturate(dot(N, V));
@@ -53,7 +51,7 @@ float3 StrataEnvironmentLighting(
 
 				SpecularDirection = R;
 				SpecularWeight = EvalEnvBRDF * DIELECTRIC_TINT(BSDF);
-				SpecularSafeRoughness = a2;
+				SpecularSafeRoughness = SafeRoughness;
 			}
 		}
 		break;
@@ -65,8 +63,6 @@ float3 StrataEnvironmentLighting(
 				// See SpecularGGX, isotropic verion
 
 				const float SafeRoughness = MakeRoughnessSafe(CONDUCTOR_ROUGHNESS(BSDF).x);
-				float a2 = Pow4(SafeRoughness);
-				float Energy = 1.0f; // See EnergyNormalization
 
 				const float3 N = CONDUCTOR_NORMAL(BSDF);
 				const float SatNoV = saturate(dot(N, V));
@@ -75,7 +71,7 @@ float3 StrataEnvironmentLighting(
 
 				SpecularDirection = R;
 				SpecularWeight = EvalEnvBRDF;
-				SpecularSafeRoughness = a2;
+				SpecularSafeRoughness = SafeRoughness;
 			}
 		}
 		break;
