Strata now takes into account off specular peak

diff --git a/Engine/Shaders/Private/Strata/StrataEnvironmentLighting.ush b/Engine/Shaders/Private/Strata/StrataEnvironmentLighting.ush
--- a/Engine/Shaders/Private/Strata/StrataEnvironmentLighting.ush
+++ b/Engine/Shaders/Private/Strata/StrataEnvironmentLighting.ush
@@ -49,6 +49,8 @@ float3 StrataEnvironmentLighting(
 				float3 R = 2 * dot(V, N) * N - V;
 				float3 EvalEnvBRDF = EnvBRDFApprox(F0, SafeRoughness, SatNoV);
 
+				R = GetOffSpecularPeakReflectionDir(N, R, SafeRoughness);
+
 				SpecularDirection = R;
 				SpecularWeight = EvalEnvBRDF * DIELECTRIC_TINT(BSDF);
 				SpecularSafeRoughness = SafeRoughness;
@@ -69,6 +71,8 @@ float3 StrataEnvironmentLighting(
 				float3 R = 2 * dot(V, N) * N - V;
 				float3 EvalEnvBRDF = EnvBRDFApprox(CONDUCTOR_REFLECTIVITY(BSDF), CONDUCTOR_EDGECOLOR(BSDF), SafeRoughness, SatNoV);
 
+				R = GetOffSpecularPeakReflectionDir(N, R, SafeRoughness);
+
 				SpecularDirection = R;
 				SpecularWeight = EvalEnvBRDF;
 				SpecularSafeRoughness = SafeRoughness;
