Hinding Strata implementation behind BSDF evaluatino function. This is needed for
1- remove code duplication that was there in the first basic implementation
2- make it compatible with Lumen

VGPRs allocation of post processes reading strata remained the same.

diff --git a/Engine/Shaders/Private/ParticipatingMediaCommon.ush b/Engine/Shaders/Private/ParticipatingMediaCommon.ush
--- a/Engine/Shaders/Private/ParticipatingMediaCommon.ush
+++ b/Engine/Shaders/Private/ParticipatingMediaCommon.ush
@@ -2,6 +2,8 @@
 	ParticipatingMediaCommon.ush
 =============================================================================*/
 
+#pragma once
+
 #include "Common.ush"
 
 
diff --git a/Engine/Shaders/Private/Strata/Strata.ush b/Engine/Shaders/Private/Strata/Strata.ush
--- a/Engine/Shaders/Private/Strata/Strata.ush
+++ b/Engine/Shaders/Private/Strata/Strata.ush
@@ -5,6 +5,8 @@
 #include "../Common.ush"
 #include "../BRDF.ush"
 #include "../ShadingCommon.ush"
+#include "../ParticipatingMediaCommon.ush"
+#include "../MonteCarlo.ush"
 
 #include "/Engine/Shared/StrataDefinitions.h"
 
@@ -208,6 +210,276 @@ FStrataData GetStrataVolumeBSDF(float3 Albedo, float3 Extinction, float Anisotro
 	return StrataData;
 }
 
+// STRATA_TODO put in a common file
+// Point lobe in off-specular peak direction
+float3 StrataGetOffSpecularPeakReflectionDir(float3 Normal, float3 ReflectionVector, float Roughness)
+{
+	float a = Square(Roughness);
+	return lerp(Normal, ReflectionVector, (1 - a) * (sqrt(1 - a) + a));
+}
+
+///////////////////////////////////////////////////////////////////////////////
+// BSDF evaluate and sampling
+
+struct FStrataBSDFContext
+{
+	FStrataBSDF BSDF;
+
+	float3 N;
+	float3 V;
+	float3 L;
+	float3 R;
+	float3 H;
+
+	BxDFContext Context;
+	float SatNoL;
+	float SatNoV;
+
+	float3x3 TangentBasis;
+	float3 TangentV;
+	float3 TangentH;
+};
+
+FStrataBSDFContext StrataCreateBSDFContext(FStrataBSDF BSDF, float3 V, float3 L)
+{
+	FStrataBSDFContext BSDFContext = (FStrataBSDFContext)0;
+
+	BSDFContext.BSDF = BSDF;
+
+	BSDFContext.V = V;
+	BSDFContext.L = L;
+	BSDFContext.N = float3(0.0f, 0.0f, 1.0f);
+	switch (BSDF.BSDFType)		// TODO_STRATA: all have a normal, even subsurface?
+	{
+	case STRATA_BSDF_TYPE_DIFFUSE_ON:
+	case STRATA_BSDF_TYPE_DIFFUSE_CHAN:
+	{ BSDFContext.N = DIFFUSE_NORMAL(BSDF);				break; }
+	case STRATA_BSDF_TYPE_DIELECTRIC:
+	{ BSDFContext.N = DIELECTRIC_NORMAL(BSDF);			break; }
+	case STRATA_BSDF_TYPE_CONDUCTOR:
+	{ BSDFContext.N = CONDUCTOR_NORMAL(BSDF);			break; }
+	case STRATA_BSDF_TYPE_VOLUME:
+	{ BSDFContext.N = VOLUME_SMOOTHSURFACENORMAL(BSDF);	break; }
+	}
+	BSDFContext.N = normalize(BSDFContext.N);
+
+	BSDFContext.R = 2 * dot(BSDFContext.V, BSDFContext.N) * BSDFContext.N - BSDFContext.V;
+	BSDFContext.H = normalize(BSDFContext.V + BSDFContext.N);
+
+	BSDFContext.Context = (BxDFContext)0;
+	Init(BSDFContext.Context, BSDFContext.N, BSDFContext.V, BSDFContext.L);
+	BSDFContext.SatNoL = saturate(BSDFContext.Context.NoL);
+	BSDFContext.SatNoV = saturate(BSDFContext.Context.NoV);
+
+	BSDFContext.TangentBasis = GetTangentBasis(BSDFContext.N);
+	BSDFContext.TangentV = mul(BSDFContext.TangentBasis, BSDFContext.V);
+	BSDFContext.TangentH = mul(BSDFContext.TangentBasis, normalize(BSDFContext.V + BSDFContext.L));
+
+	return BSDFContext;
+}
+
+struct FStrataEvaluateResult
+{
+	float3 Value;
+	float  PDF;
+	float3 Throughput;	// Throughput to the next layer
+};
+
+FStrataEvaluateResult StrataEvaluateBSDF(FStrataBSDFContext BSDFContext)
+{
+	FStrataEvaluateResult Sample = (FStrataEvaluateResult)0;
+
+	const float OpaqueBSDFThroughput = 0.0f;
+
+	switch (BSDFContext.BSDF.BSDFType)
+	{
+		case STRATA_BSDF_TYPE_DIFFUSE_ON:
+		{
+			const float SafeRoughness = MakeRoughnessSafe(DIFFUSE_ROUGHNESS(BSDFContext.BSDF));
+
+			Sample.Value		= Diffuse_OrenNayar(DIFFUSE_ALBEDO(BSDFContext.BSDF), SafeRoughness, BSDFContext.Context.NoV, BSDFContext.SatNoL, BSDFContext.Context.VoH);
+			Sample.PDF			= BSDFContext.SatNoV / PI;
+			Sample.Throughput	= OpaqueBSDFThroughput;
+			break;
+		}
+		case STRATA_BSDF_TYPE_DIFFUSE_CHAN:
+		{
+			const float SafeRoughness = MakeRoughnessSafe(DIFFUSE_ROUGHNESS(BSDFContext.BSDF));
+			float a = SafeRoughness * SafeRoughness;
+			float a2 = a * a;
+
+			Sample.Value		= Diffuse_Chan(DIFFUSE_ALBEDO(BSDFContext.BSDF), a2, BSDFContext.Context.NoV, BSDFContext.SatNoL, BSDFContext.Context.VoH, BSDFContext.Context.NoH);
+			Sample.PDF			= BSDFContext.SatNoV / PI;
+			Sample.Throughput	= OpaqueBSDFThroughput;
+			break;
+		}
+		case STRATA_BSDF_TYPE_DIELECTRIC:
+		{
+			float F0 = DielectricIorToF0(DIELECTRIC_IOR(BSDFContext.BSDF));
+
+			const float SafeRoughness = MakeRoughnessSafe(DIELECTRIC_ROUGHNESS(BSDFContext.BSDF).x);
+			float a2 = Pow4(SafeRoughness);
+			float Energy = 1.0f; // See EnergyNormalization
+
+			// Generalized microfacet specular
+			float D = D_GGX(a2, BSDFContext.Context.NoH) * Energy;
+			float Vis = Vis_SmithJointApprox(a2, BSDFContext.Context.NoV, BSDFContext.SatNoL);
+			float3 F = F_Schlick(F0, BSDFContext.Context.VoH);
+
+			Sample.Value		= D * Vis * F * DIELECTRIC_TINT(BSDFContext.BSDF);
+			Sample.PDF			= VisibleGGXPDF(BSDFContext.TangentV, BSDFContext.TangentH, a2);
+			Sample.Throughput	= (1.0f - F); // STRATA_TODO (Kulla's (1-FDG))
+			break;
+		}
+		case STRATA_BSDF_TYPE_CONDUCTOR:
+		{
+			// See SpecularGGX, isotropic verion
+			const float SafeRoughness = MakeRoughnessSafe(CONDUCTOR_ROUGHNESS(BSDFContext.BSDF).x);
+			float a2 = Pow4(SafeRoughness);
+			float Energy = 1.0f; // See EnergyNormalization
+
+			// Generalized microfacet specular
+			float D = D_GGX(a2, BSDFContext.Context.NoH) * Energy;
+			float Vis = Vis_SmithJointApprox(a2, BSDFContext.Context.NoV, BSDFContext.SatNoL);
+			float3 F = F_Schlick(CONDUCTOR_REFLECTIVITY(BSDFContext.BSDF), CONDUCTOR_EDGECOLOR(BSDFContext.BSDF), BSDFContext.Context.VoH);
+
+			Sample.Value		= D * Vis * F;
+			Sample.PDF			= VisibleGGXPDF(BSDFContext.TangentV, BSDFContext.TangentH, a2);
+			Sample.Throughput	= OpaqueBSDFThroughput;
+			break;
+		}
+		case STRATA_BSDF_TYPE_VOLUME:
+		{
+			const float3 Albedo = VOLUME_ALBEDO(BSDFContext.BSDF);
+			const float3 ExtinctionCoefficients = VOLUME_EXTINCTION(BSDFContext.BSDF);
+			const float3 ScatteringCoefficients = Albedo * ExtinctionCoefficients;
+			const float3 AbsorptionCoefficients = ExtinctionCoefficients - ScatteringCoefficients;
+			const float  Thickness = VOLUME_THICKNESS(BSDFContext.BSDF);
+			const float  Phase = HenyeyGreensteinPhase(VOLUME_ANISOTROPY(BSDFContext.BSDF), BSDFContext.Context.VoL);
+
+			const float3 SafeExtinctionThreshold = 0.000001f;
+			const float3 SafeExtinctionCoefficients = max(SafeExtinctionThreshold, ExtinctionCoefficients);
+
+			const float PathLength = Thickness / max(0.0001f, abs(BSDFContext.Context.NoV));
+			const float3 SafePathSegmentTransmittance = exp(-SafeExtinctionCoefficients * PathLength);
+			const float3 NormLuminanceTransfert = 1.0f * ScatteringCoefficients;
+			float3 NormLuminanceSingleScattering = (NormLuminanceTransfert - NormLuminanceTransfert * SafePathSegmentTransmittance) / SafeExtinctionCoefficients;
+
+#if 1
+			const float3 SelfShadowTransmittance = 1.0;// BSDFContext.SatNoL;
+#else
+			const float LightPathLength = Thickness / max(0.0001f, BSDFContext.SatNoL);
+			const float3 SelfShadowTransmittance = BSDFContext.SatNoL <= 0.0 ? 0.0 : exp(-ExtinctionCoefficients * LightPathLength);
+#endif
+
+			Sample.Value		= NormLuminanceSingleScattering * Phase * SelfShadowTransmittance;
+			Sample.PDF			= Phase;
+			Sample.Throughput	= SafePathSegmentTransmittance;
+			break;
+		}
+	}
+
+	return Sample;
+}
+
+struct FStrataImportanceSampleResult
+{
+	float3 Direction;
+	float ConeHalfAngle;
+	// add a SH representation here?
+
+	FStrataEvaluateResult Eval;
+};
+
+FStrataImportanceSampleResult StrataSampleBSDF(FStrataBSDF BSDF, float2 U, BxDFContext Context)
+{
+	FStrataImportanceSampleResult Sample = (FStrataImportanceSampleResult)0;
+	// TODO
+	return Sample;
+}
+
+struct FStrataEnvLightResult
+{
+	float3 DiffuseAlbedo;
+	float3 DiffuseNormal;
+
+	float3 SpecularDirection;
+	float3 SpecularWeight;
+	float  SpecularSafeRoughness;
+};
+
+FStrataEnvLightResult StrataEvaluateForEnvLight(FStrataBSDFContext BSDFContext, bool bEnableSpecular)
+{
+	FStrataEnvLightResult StrataEnvLightResult = (FStrataEnvLightResult)0;
+
+	switch (BSDFContext.BSDF.BSDFType)
+	{
+		case STRATA_BSDF_TYPE_DIFFUSE_ON:
+		case STRATA_BSDF_TYPE_DIFFUSE_CHAN:
+		{
+			StrataEnvLightResult.DiffuseAlbedo = DIFFUSE_ALBEDO(BSDFContext.BSDF);
+			StrataEnvLightResult.DiffuseNormal = DIFFUSE_NORMAL(BSDFContext.BSDF);
+			break;
+		}
+		case STRATA_BSDF_TYPE_DIELECTRIC:
+		{
+			BRANCH
+			if (bEnableSpecular)
+			{
+				float F0 = DielectricIorToF0(DIELECTRIC_IOR(BSDFContext.BSDF));
+				const float SafeRoughness = MakeRoughnessSafe(DIELECTRIC_ROUGHNESS(BSDFContext.BSDF).x);
+
+				float3 EvalEnvBRDF = EnvBRDFApprox(F0, SafeRoughness, BSDFContext.SatNoV);
+
+				StrataEnvLightResult.SpecularDirection = StrataGetOffSpecularPeakReflectionDir(BSDFContext.N, BSDFContext.R, SafeRoughness);
+				StrataEnvLightResult.SpecularWeight = EvalEnvBRDF * DIELECTRIC_TINT(BSDFContext.BSDF);
+				StrataEnvLightResult.SpecularSafeRoughness = SafeRoughness;
+			}
+			break;
+		}
+		case STRATA_BSDF_TYPE_CONDUCTOR:
+		{
+			BRANCH
+			if (bEnableSpecular)
+			{
+				// See SpecularGGX, isotropic verion
+
+				const float SafeRoughness = MakeRoughnessSafe(CONDUCTOR_ROUGHNESS(BSDFContext.BSDF).x);
+
+				const float3 N = CONDUCTOR_NORMAL(BSDFContext.BSDF);
+				float3 EvalEnvBRDF = EnvBRDFApprox(CONDUCTOR_REFLECTIVITY(BSDFContext.BSDF), CONDUCTOR_EDGECOLOR(BSDFContext.BSDF), SafeRoughness, BSDFContext.SatNoV);
+
+				StrataEnvLightResult.SpecularDirection = StrataGetOffSpecularPeakReflectionDir(BSDFContext.N, BSDFContext.R, SafeRoughness);
+				StrataEnvLightResult.SpecularWeight = EvalEnvBRDF;
+				StrataEnvLightResult.SpecularSafeRoughness = SafeRoughness;
+			}
+			break;
+		}
+		case STRATA_BSDF_TYPE_VOLUME:
+		{
+			const float3 VolumeAlbedo = VOLUME_ALBEDO(BSDFContext.BSDF);
+			const float3 ExtinctionCoefficients = VOLUME_EXTINCTION(BSDFContext.BSDF);
+			const float3 ScatteringCoefficients = VolumeAlbedo * ExtinctionCoefficients;
+			const float3 AbsorptionCoefficients = ExtinctionCoefficients - ScatteringCoefficients;
+			const float  Thickness = VOLUME_THICKNESS(BSDFContext.BSDF);
+
+			const float3 SafeExtinctionThreshold = 0.000001f;
+			const float3 SafeExtinctionCoefficients = max(SafeExtinctionThreshold, ExtinctionCoefficients);
+
+			const float PathLength = Thickness / max(0.0001f, BSDFContext.SatNoV);
+			const float3 SafePathSegmentTransmittance = exp(-SafeExtinctionCoefficients * PathLength);
+			const float3 NormLuminanceTransfert = 1.0f * ScatteringCoefficients;
+			float3 NormLuminanceSingleScattering = (NormLuminanceTransfert - NormLuminanceTransfert * SafePathSegmentTransmittance) / SafeExtinctionCoefficients;
+
+			StrataEnvLightResult.DiffuseAlbedo = NormLuminanceSingleScattering;
+			StrataEnvLightResult.DiffuseNormal = VOLUME_SMOOTHSURFACENORMAL(BSDFContext.BSDF);
+			break;
+		}
+	}
+
+	return StrataEnvLightResult;
+}
+
 
 
 ///////////////////////////////////////////////////////////////////////////////
diff --git a/Engine/Shaders/Private/Strata/StrataDeferredLighting.ush b/Engine/Shaders/Private/Strata/StrataDeferredLighting.ush
--- a/Engine/Shaders/Private/Strata/StrataDeferredLighting.ush
+++ b/Engine/Shaders/Private/Strata/StrataDeferredLighting.ush
@@ -18,121 +18,34 @@ float4 StrataDeferredLighting(
 {
 	float4 Radiance = 0;
 
+
+	float3 UglyIntegrationFactor = 0; // This is bad we need to make sure it is done the right way
+	if (LightData.bRectLight)
+	{
+		FRect Rect = GetRect(ToLight, LightData);
+		UglyIntegrationFactor = IntegrateLight(Rect, RectTexture);
+	}
+	else
+	{
+		FCapsuleLight Capsule = GetCapsule(ToLight, LightData);
+		UglyIntegrationFactor = IntegrateLight(Capsule, LightData.bInverseSquared);
+	}
+
 	for (int i = 0; i < PixelBSDFs.BSDFCount; ++i)
 	{
 		FStrataBSDF BSDF = PixelBSDFs.BSDFs[i];
 		const float3 BSDFThroughput = BSDF.Weight;
 
-		float3 N = GBuffer.WorldNormal;
-		switch (BSDF.BSDFType)		// TODO_STRATA: all have a normal, even subsurface?
-		{
-		case STRATA_BSDF_TYPE_DIFFUSE_ON:
-		case STRATA_BSDF_TYPE_DIFFUSE_CHAN:
-		{ N = DIFFUSE_NORMAL(BSDF);				break; }
-		case STRATA_BSDF_TYPE_DIELECTRIC:
-		{ N = DIELECTRIC_NORMAL(BSDF);			break; }
-		case STRATA_BSDF_TYPE_CONDUCTOR:
-		{ N = CONDUCTOR_NORMAL(BSDF);			break; }
-		case STRATA_BSDF_TYPE_VOLUME:
-		{ N = VOLUME_SMOOTHSURFACENORMAL(BSDF);	break; }
-		}
-		N = normalize(N);
-
-		BxDFContext Context = (BxDFContext)0;
-		Init(Context, N, V, L);
-		float SatNoL = saturate(Context.NoL);
-
-		float3 UglyIntegrationFactor = 0; // This is bad we need to make sure it is done the right way
-		if (LightData.bRectLight)
-		{
-			FRect Rect = GetRect(ToLight, LightData);
-			UglyIntegrationFactor = IntegrateLight(Rect, RectTexture);
-		}
-		else
-		{
-			FCapsuleLight Capsule = GetCapsule(ToLight, LightData);
-			UglyIntegrationFactor = IntegrateLight(Capsule, LightData.bInverseSquared);
-		}
+		// Create the BSDF context
+		FStrataBSDFContext StrataBSDFContext = StrataCreateBSDFContext(BSDF, V, L);
 
+		// Compute some common factors
 		const float3 DirLightBaseFactor = BSDFThroughput * LightData.Color * LightMask * SurfaceShadow * UglyIntegrationFactor;
-		const float3 DirLightGeometryFactor = DirLightBaseFactor * SatNoL;
-
-		switch (BSDF.BSDFType)
-		{
-		case STRATA_BSDF_TYPE_DIFFUSE_ON:
-		{
-			const float SafeRoughness = MakeRoughnessSafe(DIFFUSE_ROUGHNESS(BSDF));
-			Radiance.rgb += Diffuse_OrenNayar(DIFFUSE_ALBEDO(BSDF), SafeRoughness, Context.NoV, SatNoL, Context.VoH) * DirLightGeometryFactor;
-		}
-		break;
-		case STRATA_BSDF_TYPE_DIFFUSE_CHAN:
-		{
-			const float SafeRoughness = MakeRoughnessSafe(DIFFUSE_ROUGHNESS(BSDF));
-			float a = SafeRoughness * SafeRoughness;
-			float a2 = a * a;
-			Radiance.rgb += Diffuse_Chan(DIFFUSE_ALBEDO(BSDF), a2, Context.NoV, SatNoL, Context.VoH, Context.NoH) * DirLightGeometryFactor;
-		}
-		break;
-		case STRATA_BSDF_TYPE_DIELECTRIC:
-		{
-			float F0 = DielectricIorToF0(DIELECTRIC_IOR(BSDF));
-
-			const float SafeRoughness = MakeRoughnessSafe(DIELECTRIC_ROUGHNESS(BSDF).x);
-			float a2 = Pow4(SafeRoughness);
-			float Energy = 1.0f; // See EnergyNormalization
-
-			// Generalized microfacet specular
-			float D = D_GGX(a2, Context.NoH) * Energy;
-			float Vis = Vis_SmithJointApprox(a2, Context.NoV, SatNoL);
-			float3 F = F_Schlick(F0, Context.VoH);
-
-			Radiance.rgb += D * Vis * F * DIELECTRIC_TINT(BSDF) * DirLightGeometryFactor;
-		}
-		break;
-		case STRATA_BSDF_TYPE_CONDUCTOR:
-		{
-			// See SpecularGGX, isotropic verion
-
-			const float SafeRoughness = MakeRoughnessSafe(CONDUCTOR_ROUGHNESS(BSDF).x);
-			float a2 = Pow4(SafeRoughness);
-			float Energy = 1.0f; // See EnergyNormalization
-
-			// Generalized microfacet specular
-			float D = D_GGX(a2, Context.NoH) * Energy;
-			float Vis = Vis_SmithJointApprox(a2, Context.NoV, SatNoL);
-			float3 F = F_Schlick(CONDUCTOR_REFLECTIVITY(BSDF), CONDUCTOR_EDGECOLOR(BSDF), Context.VoH);
-
-			Radiance.rgb += D * Vis * F * DirLightGeometryFactor;
-		}
-		break;
-		case STRATA_BSDF_TYPE_VOLUME:
-		{
-			const float3 Albedo = VOLUME_ALBEDO(BSDF);
-			const float3 ExtinctionCoefficients = VOLUME_EXTINCTION(BSDF);
-			const float3 ScatteringCoefficients = Albedo * ExtinctionCoefficients;
-			const float3 AbsorptionCoefficients = ExtinctionCoefficients - ScatteringCoefficients;
-			const float  Thickness = VOLUME_THICKNESS(BSDF);
-			const float  Phase = HenyeyGreensteinPhase(VOLUME_ANISOTROPY(BSDF), Context.VoL);
-
-			const float3 SafeExtinctionThreshold = 0.000001f;
-			const float3 SafeExtinctionCoefficients = max(SafeExtinctionThreshold, ExtinctionCoefficients);
-
-			const float PathLength = Thickness / max(0.0001f, abs(Context.NoV));
-			const float3 SafePathSegmentTransmittance = exp(-SafeExtinctionCoefficients * PathLength);
-			const float3 NormLuminanceTransfert = 1.0f * ScatteringCoefficients;
-			float3 NormLuminanceSingleScattering = (NormLuminanceTransfert - NormLuminanceTransfert * SafePathSegmentTransmittance) / SafeExtinctionCoefficients;
-
-#if 1
-			const float3 SelfShadowTransmittance = SatNoL;
-#else
-			const float LightPathLength = Thickness / max(0.0001f, SatNoL);
-			const float3 SelfShadowTransmittance = SatNoL <= 0.0 ? 0.0 : exp(-ExtinctionCoefficients * LightPathLength);
-#endif
+		const float3 DirLightGeometryFactor = DirLightBaseFactor * StrataBSDFContext.SatNoL;
 
-			Radiance.rgb += NormLuminanceSingleScattering * Phase * DirLightBaseFactor * SelfShadowTransmittance;
-		}
-		break;
-		}
+		// Evalaute the lighting coming from L
+		FStrataEvaluateResult StrataEvaluate = StrataEvaluateBSDF(StrataBSDFContext);
+		Radiance.rgb += StrataEvaluate.Value * DirLightGeometryFactor;
 	}
 
 	return Radiance;
diff --git a/Engine/Shaders/Private/Strata/StrataEnvironmentLighting.ush b/Engine/Shaders/Private/Strata/StrataEnvironmentLighting.ush
--- a/Engine/Shaders/Private/Strata/StrataEnvironmentLighting.ush
+++ b/Engine/Shaders/Private/Strata/StrataEnvironmentLighting.ush
@@ -20,117 +20,40 @@ float3 StrataEnvironmentLighting(
 		FStrataBSDF BSDF = PixelBSDFs.BSDFs[i];
 		const float3 BSDFThroughput = BSDF.Weight;
 
-		float3 DiffuseAlbedo = 0.0f;
-		float3 DiffuseNormal = 0.0f;
+		// Create the BSDF context
+		const float3 UnusedL = float3(0.0f, 0.0f, 1.0f);
+		FStrataBSDFContext StrataBSDFContext = StrataCreateBSDFContext(BSDF, V, UnusedL);
 
-		float3 SpecularDirection = 0.0f;
-		float3 SpecularWeight = 0.0f;
-		float  SpecularSafeRoughness = 0.0f;
-
-		switch (BSDF.BSDFType)
-		{
-		case STRATA_BSDF_TYPE_DIFFUSE_ON:
-		case STRATA_BSDF_TYPE_DIFFUSE_CHAN:
-		{
-			DiffuseAlbedo = DIFFUSE_ALBEDO(BSDF);
-			DiffuseNormal = DIFFUSE_NORMAL(BSDF);
-		}
-		break;
-		case STRATA_BSDF_TYPE_DIELECTRIC:
-		{
-			BRANCH
-			if (ReflectionStruct.SkyLightParameters.y > 0)
-			{
-				float F0 = DielectricIorToF0(DIELECTRIC_IOR(BSDF));
-				const float SafeRoughness = MakeRoughnessSafe(DIELECTRIC_ROUGHNESS(BSDF).x);
-
-				const float3 N = DIELECTRIC_NORMAL(BSDF);
-				const float SatNoV = saturate(dot(N, V));
-				float3 R = 2 * dot(V, N) * N - V;
-				float3 EvalEnvBRDF = EnvBRDFApprox(F0, SafeRoughness, SatNoV);
-
-				R = GetOffSpecularPeakReflectionDir(N, R, SafeRoughness);
-
-				SpecularDirection = R;
-				SpecularWeight = EvalEnvBRDF * DIELECTRIC_TINT(BSDF);
-				SpecularSafeRoughness = SafeRoughness;
-			}
-		}
-		break;
-		case STRATA_BSDF_TYPE_CONDUCTOR:
-		{
-			BRANCH
-			if (ReflectionStruct.SkyLightParameters.y > 0)
-			{
-				// See SpecularGGX, isotropic verion
-
-				const float SafeRoughness = MakeRoughnessSafe(CONDUCTOR_ROUGHNESS(BSDF).x);
-
-				const float3 N = CONDUCTOR_NORMAL(BSDF);
-				const float SatNoV = saturate(dot(N, V));
-				float3 R = 2 * dot(V, N) * N - V;
-				float3 EvalEnvBRDF = EnvBRDFApprox(CONDUCTOR_REFLECTIVITY(BSDF), CONDUCTOR_EDGECOLOR(BSDF), SafeRoughness, SatNoV);
-
-				R = GetOffSpecularPeakReflectionDir(N, R, SafeRoughness);
-
-				SpecularDirection = R;
-				SpecularWeight = EvalEnvBRDF;
-				SpecularSafeRoughness = SafeRoughness;
-			}
-		}
-		break;
-		case STRATA_BSDF_TYPE_VOLUME:
-		{
-			const float3 VolumeAlbedo = VOLUME_ALBEDO(BSDF);
-			const float3 ExtinctionCoefficients = VOLUME_EXTINCTION(BSDF);
-			const float3 ScatteringCoefficients = VolumeAlbedo * ExtinctionCoefficients;
-			const float3 AbsorptionCoefficients = ExtinctionCoefficients - ScatteringCoefficients;
-			const float  Thickness = VOLUME_THICKNESS(BSDF);
-
-			const float3 SafeExtinctionThreshold = 0.000001f;
-			const float3 SafeExtinctionCoefficients = max(SafeExtinctionThreshold, ExtinctionCoefficients);
-
-			const float SatNoV = saturate(dot(VOLUME_SMOOTHSURFACENORMAL(BSDF), V));
-			const float PathLength = Thickness / max(0.0001f, SatNoV);
-			const float3 SafePathSegmentTransmittance = exp(-SafeExtinctionCoefficients * PathLength);
-			const float3 NormLuminanceTransfert = 1.0f * ScatteringCoefficients;
-			float3 NormLuminanceSingleScattering = (NormLuminanceTransfert - NormLuminanceTransfert * SafePathSegmentTransmittance) / SafeExtinctionCoefficients;
-
-			// No phase function because the SH already contains integration of BSDF with respect to hemisphere. Incorrect but needed.
-			DiffuseAlbedo = NormLuminanceSingleScattering;
-
-			// STRATA_TODO: This is to fake some directionality. Have normal specified on the material as for MaterialX?
-			DiffuseNormal = VOLUME_SMOOTHSURFACENORMAL(BSDF);
-		}
-		break;
-		}
+		// Evaluate environment lighting
+		const bool bEnableSpecular = ReflectionStruct.SkyLightParameters.y > 0.0f;
+		FStrataEnvLightResult StrataEnvLight = StrataEvaluateForEnvLight(StrataBSDFContext, bEnableSpecular);
 
 		// Diffuse component
-		if (any(DiffuseAlbedo > 0.0f))
+		if (any(StrataEnvLight.DiffuseAlbedo > 0.0f))
 		{
 			// Apply bent normal if needed
 			#if APPLY_SKY_SHADOWING
 			{
-				DiffuseNormal = UpsampleDFAO(BufferUV, GBuffer.Depth, DiffuseNormal); // Bent normal
+				StrataEnvLight.DiffuseNormal = UpsampleDFAO(BufferUV, GBuffer.Depth, StrataEnvLight.DiffuseNormal); // Bent normal
 			}
 			#endif
 
 			// Compute the common sky visibility factors
-			FSkyLightVisibilityData SkyVisData = GetSkyLightVisibilityData(GBuffer, AmbientOcclusion, DiffuseNormal);
+			FSkyLightVisibilityData SkyVisData = GetSkyLightVisibilityData(GBuffer, AmbientOcclusion, StrataEnvLight.DiffuseNormal);
 
 			// Finally sample the sky diffuse contribution (spherical harmonic, Lambert BRDF)
-			float3 DiffuseLookup = GetSkySHDiffuse(DiffuseNormal) * View.SkyLightColor.rgb;			// STRATA_TODO compute SH coefficients for ON or Chan
+			float3 DiffuseLookup = GetSkySHDiffuse(StrataEnvLight.DiffuseNormal) * View.SkyLightColor.rgb;			// STRATA_TODO compute SH coefficients for ON or Chan
 			// And accumulate
-			SkyLighting += BSDFThroughput * (SkyVisData.SkyDiffuseLookUpMul * DiffuseLookup + SkyVisData.SkyDiffuseLookUpAdd) * DiffuseAlbedo;
+			SkyLighting += BSDFThroughput * (SkyVisData.SkyDiffuseLookUpMul * DiffuseLookup + SkyVisData.SkyDiffuseLookUpAdd) * StrataEnvLight.DiffuseAlbedo;
 		}
 
 		// Specular component
-		if (any(SpecularWeight > 0.0f))
+		if (any(StrataEnvLight.SpecularWeight > 0.0f))
 		{
 			float SkyAverageBrightness = 1.0f;
-			const float3 SpecularLuminance = GetSkyLightReflection(SpecularDirection, SpecularSafeRoughness, SkyAverageBrightness);
+			const float3 SpecularLuminance = GetSkyLightReflection(StrataEnvLight.SpecularDirection, StrataEnvLight.SpecularSafeRoughness, SkyAverageBrightness);
 
-			SkyLighting += BSDFThroughput * SpecularLuminance * SpecularWeight * View.SkyLightColor.rgb * AmbientOcclusion;
+			SkyLighting += BSDFThroughput * SpecularLuminance * StrataEnvLight.SpecularWeight * View.SkyLightColor.rgb * AmbientOcclusion;
 		}
 
 	}
diff --git a/Engine/Shaders/Private/Strata/StrataForwardLighting.ush b/Engine/Shaders/Private/Strata/StrataForwardLighting.ush
--- a/Engine/Shaders/Private/Strata/StrataForwardLighting.ush
+++ b/Engine/Shaders/Private/Strata/StrataForwardLighting.ush
@@ -32,177 +32,37 @@ float3 StrataForwardLighting(
 			FStrataBSDF BSDF = Strata.Layers[l].BSDFs[i];
 			const float3 BSDFCoverage = BSDF.Weight;
 
-			float3 N = MaterialParametersWorldNormal;
-			switch (BSDF.BSDFType)		// TODO_STRATA: all have a normal, even subsurface?
-			{
-			case STRATA_BSDF_TYPE_DIFFUSE_ON:
-			case STRATA_BSDF_TYPE_DIFFUSE_CHAN:
-			{ N = DIFFUSE_NORMAL(BSDF);				SanitizeStrataDiffuse(BSDF);	break; }
-			case STRATA_BSDF_TYPE_DIELECTRIC:
-			{ N = DIELECTRIC_NORMAL(BSDF);			SanitizeStrataDielectric(BSDF);	break; }
-			case STRATA_BSDF_TYPE_CONDUCTOR:
-			{ N = CONDUCTOR_NORMAL(BSDF);			SanitizeStrataConductor(BSDF);	break; }
-			case STRATA_BSDF_TYPE_VOLUME:
-			{ N = VOLUME_SMOOTHSURFACENORMAL(BSDF);	SanitizeStrataVolume(BSDF);		break; }
-			}
-			N = normalize(N);
-			float3 R = 2 * dot(V, N) * N - V;
-
-			BxDFContext Context = (BxDFContext)0;
-			Init(Context, N, V, L);
-			float SatNoL = saturate(Context.NoL);
+			// Create the BSDF context
+			FStrataBSDFContext StrataBSDFContext = StrataCreateBSDFContext(BSDF, V, L);
 
+			// Compute some common factors
 			const float3 DirLightBaseFactor = OutThroughput * BSDFCoverage * LColor * LShadow;
-			const float3 DirLightGeometryFactor = DirLightBaseFactor * SatNoL;
+			const float3 DirLightGeometryFactor = DirLightBaseFactor * StrataBSDFContext.SatNoL;
 			const float3 SkyLightCommonFactor = OutThroughput * BSDFCoverage * ResolvedView.SkyLightColor.rgb * IndirectOcclusion;
 
-			switch (BSDF.BSDFType)
+			// Evaluate environment lighting
+			const bool bEnableSpecular = ReflectionStruct.SkyLightParameters.y > 0.0f;
+			FStrataEnvLightResult StrataEnvLight = StrataEvaluateForEnvLight(StrataBSDFContext, bEnableSpecular);
+			if (any(StrataEnvLight.DiffuseAlbedo > 0.0f))
 			{
-			case STRATA_BSDF_TYPE_DIFFUSE_ON:
-			{
-				const float SafeRoughness = MakeRoughnessSafe(DIFFUSE_ROUGHNESS(BSDF));
-				Color += Diffuse_OrenNayar(DIFFUSE_ALBEDO(BSDF), SafeRoughness, Context.NoV, SatNoL, Context.VoH) * DirLightGeometryFactor;
-
-#if ENABLE_SKY_LIGHT
-				BRANCH
-				if (ReflectionStruct.SkyLightParameters.y > 0)
-				{
-					Color += GetEffectiveSkySHDiffuse(N) * DIFFUSE_ALBEDO(BSDF) * SkyLightCommonFactor; // STRATA_TODO compute SH coefficients for ON
-				}
-#endif
-
-				LayerCoverage	+= BSDFCoverage;
-				LayerThroughput += BSDFCoverage * OpaqueBSDFThroughput;
-				break;
-			}
+				// SkySHDiffuse is pre integrated according to a Lambert diffuse lobe. Thus, we need to rescale to account for the phase function operating over 4PI.
+				// We apply a 2PI factor to recover average luminance over half hemisphere, and then times the phase function to account for the scattering events.
+				const float VolumeOnlyFactor = BSDF.BSDFType == STRATA_BSDF_TYPE_VOLUME ? 2.0f * PI * IsotropicPhase() : 1.0f;
 
-			case STRATA_BSDF_TYPE_DIFFUSE_CHAN:
-			{
-				const float SafeRoughness = MakeRoughnessSafe(DIFFUSE_ROUGHNESS(BSDF));
-				float a = SafeRoughness * SafeRoughness;
-				float a2 = a * a;
-				Color += Diffuse_Chan(DIFFUSE_ALBEDO(BSDF), a2, Context.NoV, SatNoL, Context.VoH, Context.NoH) * DirLightGeometryFactor;
-
-#if ENABLE_SKY_LIGHT
-				BRANCH
-				if (ReflectionStruct.SkyLightParameters.y > 0)
-				{
-					Color += GetEffectiveSkySHDiffuse(N) * DIFFUSE_ALBEDO(BSDF) * SkyLightCommonFactor; // STRATA_TODO compute SH coefficients for Chan
-				}
-#endif
-				
-				LayerCoverage	+= BSDFCoverage;
-				LayerThroughput += BSDFCoverage * OpaqueBSDFThroughput;
-				break;
+				Color += GetEffectiveSkySHDiffuse(StrataEnvLight.DiffuseNormal) * StrataEnvLight.DiffuseAlbedo * SkyLightCommonFactor * VolumeOnlyFactor;
 			}
-
-			case STRATA_BSDF_TYPE_DIELECTRIC:
+			if (any(StrataEnvLight.SpecularWeight > 0.0f))
 			{
-				float F0 = DielectricIorToF0(DIELECTRIC_IOR(BSDF));
-
-				const float SafeRoughness = MakeRoughnessSafe(DIELECTRIC_ROUGHNESS(BSDF).x);
-				float a2 = Pow4(SafeRoughness);
-				float Energy = 1.0f; // See EnergyNormalization
-
-				// Generalized microfacet specular
-				float D = D_GGX(a2, Context.NoH) * Energy;
-				float Vis = Vis_SmithJointApprox(a2, Context.NoV, SatNoL);
-				float3 F = F_Schlick(F0, Context.VoH);
-
-				Color += D * Vis * F * DIELECTRIC_TINT(BSDF) * DirLightGeometryFactor;
-
-#if ENABLE_SKY_LIGHT
-				BRANCH
-					if (ReflectionStruct.SkyLightParameters.y > 0)
-					{
-						float3 EvalEnvBRDF = EnvBRDFApprox(F0, SafeRoughness, Context.NoV);
-						float SkyAverageBrightness = 1.0f;
-						Color += GetSkyLightReflection(R, SafeRoughness, SkyAverageBrightness) * EvalEnvBRDF * DIELECTRIC_TINT(BSDF) * SkyLightCommonFactor;
-					}
-#endif
-				
-				LayerCoverage	+= BSDFCoverage;
-				LayerThroughput += BSDFCoverage * (1.0f - F); // STRATA_TODO (Kulla's (1-FDG))
-				break;
+				float SkyAverageBrightness = 1.0f;
+				Color += GetSkyLightReflection(StrataEnvLight.SpecularDirection, StrataEnvLight.SpecularSafeRoughness, SkyAverageBrightness) * StrataEnvLight.SpecularWeight * SkyLightCommonFactor;
 			}
 
-			case STRATA_BSDF_TYPE_CONDUCTOR:
-			{
-				// See SpecularGGX, isotropic verion
-
-				const float SafeRoughness = MakeRoughnessSafe(CONDUCTOR_ROUGHNESS(BSDF).x);
-				float a2 = Pow4(SafeRoughness);
-				float Energy = 1.0f; // See EnergyNormalization
-
-				// Generalized microfacet specular
-				float D = D_GGX(a2, Context.NoH) * Energy;
-				float Vis = Vis_SmithJointApprox(a2, Context.NoV, SatNoL);
-				float3 F = F_Schlick(CONDUCTOR_REFLECTIVITY(BSDF), CONDUCTOR_EDGECOLOR(BSDF), Context.VoH);
-
-				Color += D * Vis * F * DirLightGeometryFactor;
-
-#if ENABLE_SKY_LIGHT
-				//BRANCH
-				if (ReflectionStruct.SkyLightParameters.y > 0)
-				{
-					float3 EvalEnvBRDF = EnvBRDFApprox(CONDUCTOR_REFLECTIVITY(BSDF), CONDUCTOR_EDGECOLOR(BSDF), SafeRoughness, Context.NoV);
-					float SkyAverageBrightness = 1.0f;
-					Color += GetSkyLightReflection(R, SafeRoughness, SkyAverageBrightness) * EvalEnvBRDF * SkyLightCommonFactor;
-				}
-#endif
-				
-				LayerCoverage	+= BSDFCoverage;
-				LayerThroughput += BSDFCoverage * OpaqueBSDFThroughput;
-				break;
-			}
+			// Evalaute the lighting coming from L (it will lso contain the BSDF throughput to the next layer)
+			FStrataEvaluateResult StrataEvaluate = StrataEvaluateBSDF(StrataBSDFContext);
+			Color += StrataEvaluate.Value * DirLightGeometryFactor;
 
-			case STRATA_BSDF_TYPE_VOLUME:
-			{
-				const float3 Albedo = VOLUME_ALBEDO(BSDF);
-				const float3 ExtinctionCoefficients = VOLUME_EXTINCTION(BSDF);
-				const float3 ScatteringCoefficients = Albedo * ExtinctionCoefficients;
-				const float3 AbsorptionCoefficients = ExtinctionCoefficients - ScatteringCoefficients;
-				const float  Thickness = VOLUME_THICKNESS(BSDF);
-				const float  Phase = HenyeyGreensteinPhase(VOLUME_ANISOTROPY(BSDF), Context.VoL);
-
-				const float3 SafeExtinctionThreshold = 0.000001f;
-				const float3 SafeExtinctionCoefficients = max(SafeExtinctionThreshold, ExtinctionCoefficients);
-
-				const float PathLength = Thickness / max(0.0001f, abs(Context.NoV));
-				const float3 SafePathSegmentTransmittance = exp(-SafeExtinctionCoefficients * PathLength);
-				const float3 NormLuminanceTransfert = 1.0f * ScatteringCoefficients;
-				float3 NormLuminanceSingleScattering = (NormLuminanceTransfert - NormLuminanceTransfert * SafePathSegmentTransmittance) / SafeExtinctionCoefficients;
-
-#if 1
-				const float3 SelfShadowTransmittance = SatNoL;
-#else
-				const float LightPathLength = Thickness / max(0.0001f, SatNoL);
-				const float3 SelfShadowTransmittance = SatNoL <= 0.0 ? 0.0 : exp(-ExtinctionCoefficients * LightPathLength);
-#endif
-
-				Color += NormLuminanceSingleScattering * Phase * DirLightBaseFactor * SelfShadowTransmittance;
-
-#if ENABLE_SKY_LIGHT
-				//BRANCH
-				if (ReflectionStruct.SkyLightParameters.y > 0)
-				{
-					const float RescaleWithPhase = 2.0f * PI * IsotropicPhase(); // times 2PI to recover luminance, and then times the phase function for the scattering event
-					float3 EvalEnvBRDF = NormLuminanceSingleScattering * RescaleWithPhase;
-					float SafeRoughnessFromG = clamp(1.0 - abs(VOLUME_ANISOTROPY(BSDF)), 0.001, 1.0f); // STRATA_TODO map G to roughness correctly
-					float SkyAverageBrightness = 1.0f;
-					// STRATA_TODO 
-					//	- Using R is weird. We should use V. But for rought material when light is incoming mainly from upper hemisphere, using N would make more send.
-					//	- Then what phase should be used? Phase(VoN)? We need to fund a good mapping on there. 
-					//	- We need to clarify what must be the weight there also (with env light in participating media coming from 2pi). See above attempt RescaleWithPhase.
-					Color += GetSkyLightReflection(R, SafeRoughnessFromG, SkyAverageBrightness) * EvalEnvBRDF * SkyLightCommonFactor;
-				}
-#endif
-
-				LayerCoverage	+= BSDFCoverage;
-				LayerThroughput += BSDFCoverage * SafePathSegmentTransmittance;
-				break;
-			}
-			}
+			LayerCoverage	+= BSDFCoverage;
+			LayerThroughput += BSDFCoverage * StrataEvaluate.Throughput;
 		}
 
 		// Take into account the non covered part for this layer and update throughput accordingly
