Fixing Thin Transparent objects in Frosty.

diff --git a/Engine/Shaders/Private/BasePassPixelShader.usf b/Engine/Shaders/Private/BasePassPixelShader.usf
--- a/Engine/Shaders/Private/BasePassPixelShader.usf
+++ b/Engine/Shaders/Private/BasePassPixelShader.usf
@@ -1652,7 +1652,7 @@ void FPixelShaderInOut_MainPS(
 	#else
 
 		// If not using the full gbuffer (forward shading) the velocity buffer can still be written to in the basepass.
-		#if GBUFFER_HAS_VELOCITY
+		#if GBUFFER_HAS_VELOCITY && !DUAL_SOURCE_COLOR_BLENDING_ENABLED
 			Out.MRT[1] = OutVelocity;
 		#endif
 
