Added function to importance sample Strata BSDF (Diffuse, GGX and phase function)
Made HenyeyGreensteinPhase follow PBRT convenstion.

diff --git a/Engine/Shaders/Private/ParticipatingMediaCommon.ush b/Engine/Shaders/Private/ParticipatingMediaCommon.ush
--- a/Engine/Shaders/Private/ParticipatingMediaCommon.ush
+++ b/Engine/Shaders/Private/ParticipatingMediaCommon.ush
@@ -13,10 +13,10 @@ float IsotropicPhase()
 	return 1.0f / (4.0f * PI);
 }
 
-float HenyeyGreensteinPhase(float g, float CosTheta)
+// Follows PBRT convention http://www.pbr-book.org/3ed-2018/Volume_Scattering/Phase_Functions.html#PhaseHG
+float HenyeyGreensteinPhase(float G, float CosTheta)
 {
-	g = -g;
-	return (1.0f - g * g) / (4.0f * PI * pow(1.0f + g * g - 2.0f * g * CosTheta, 1.5f));
+	return (1.0f - G * G) / (4.0f * PI * pow(1.0f + G * G + 2.0f * G * CosTheta, 1.5f));
 }
 
 float RaleighPhase(float CosTheta)
@@ -25,16 +25,45 @@ float RaleighPhase(float CosTheta)
 }
 
 // Schlick phase function approximating henyey-greenstein
-float SchlickPhaseFromK(float k, float CosTheta)
+float SchlickPhaseFromK(float K, float CosTheta)
 {
-	const float SchlickPhaseFactor = 1.0f + k * CosTheta;
-	const float PhaseValue = (1.0f - k * k) / (4.0f * PI * SchlickPhaseFactor * SchlickPhaseFactor);
+	const float SchlickPhaseFactor = 1.0f + K * CosTheta;
+	const float PhaseValue = (1.0f - K * K) / (4.0f * PI * SchlickPhaseFactor * SchlickPhaseFactor);
 	return PhaseValue;
 }
-float SchlickPhase(float g, float CosTheta)
+float SchlickPhase(float G, float CosTheta)
 {
-	const float k = 1.55f * g - 0.55f * g * g * g;
-	return SchlickPhaseFromK(k, CosTheta);
+	const float K = 1.55f * G - 0.55f * G * G * G;
+	return SchlickPhaseFromK(K, CosTheta);
 }
 
+// Follows PBRT convention http://www.pbr-book.org/3ed-2018/Light_Transport_II_Volume_Rendering/Sampling_Volume_Scattering.html#SamplingPhaseFunctions
+float HenyeyGreensteinPhaseInvertCDF(float E, float G)
+{
+	float OnePlusG2  = 1.0f + G * G;
+	float OneMinusG2 = 1.0f - G * G;
+	float OneOver2G  = 0.5f / G;
+	float t = OneMinusG2 / (1.0f + G - 2.0f * G * E);
+	return -OneOver2G * (OnePlusG2 - t * t);	// Careful: OneOver2G undefined for g close to 0
+}
+
+// Follows PBRT convention http://www.pbr-book.org/3ed-2018/Light_Transport_II_Volume_Rendering/Sampling_Volume_Scattering.html#SamplingPhaseFunctions
+float4 ImportanceSampleHenyeyGreensteinPhase(float2 E, float G)
+{
+	float Phi = 2.0f * PI * E.x;
+	float CosTheta = 0;
+	if (abs(G) < 0.001f)
+	{
+		CosTheta = 1.0f - 2.0f * E.y;
+	}
+	else
+	{
+		CosTheta = HenyeyGreensteinPhaseInvertCDF(E.y, G);
+	}
+	float SinTheta = sqrt(max(0.0f, 1.0f - CosTheta * CosTheta));
+
+	float3 H = float3(SinTheta * sin(Phi), SinTheta * cos(Phi), CosTheta);
+
+	return float4(H, HenyeyGreensteinPhase(G, CosTheta));
+}
 
diff --git a/Engine/Shaders/Private/Strata/Strata.ush b/Engine/Shaders/Private/Strata/Strata.ush
--- a/Engine/Shaders/Private/Strata/Strata.ush
+++ b/Engine/Shaders/Private/Strata/Strata.ush
@@ -227,10 +227,11 @@ struct FStrataBSDFContext
 
 	float3 N;
 	float3 V;
-	float3 L;
 	float3 R;
 	float3 H;
 
+	float3 L;	// There to initialise the BxDFContext. Only used by StrataEvaluateBSDF, not by StrataImportanceSampleBSDF or StrataEvaluateForEnvLight
+
 	BxDFContext Context;
 	float SatNoL;
 	float SatNoV;
@@ -264,7 +265,7 @@ FStrataBSDFContext StrataCreateBSDFContext(FStrataBSDF BSDF, float3 V, float3 L)
 	BSDFContext.N = normalize(BSDFContext.N);
 
 	BSDFContext.R = 2 * dot(BSDFContext.V, BSDFContext.N) * BSDFContext.N - BSDFContext.V;
-	BSDFContext.H = normalize(BSDFContext.V + BSDFContext.N);
+	BSDFContext.H = normalize(BSDFContext.V + BSDFContext.L);
 
 	BSDFContext.Context = (BxDFContext)0;
 	Init(BSDFContext.Context, BSDFContext.N, BSDFContext.V, BSDFContext.L);
@@ -273,11 +274,19 @@ FStrataBSDFContext StrataCreateBSDFContext(FStrataBSDF BSDF, float3 V, float3 L)
 
 	BSDFContext.TangentBasis = GetTangentBasis(BSDFContext.N);
 	BSDFContext.TangentV = mul(BSDFContext.TangentBasis, BSDFContext.V);
-	BSDFContext.TangentH = mul(BSDFContext.TangentBasis, normalize(BSDFContext.V + BSDFContext.L));
+	BSDFContext.TangentH = mul(BSDFContext.TangentBasis, BSDFContext.H);
 
 	return BSDFContext;
 }
 
+void StrataUpdateBSDFContext(inout FStrataBSDFContext BSDFContext, float3 NewL)
+{
+	// Update all the data related to L
+	Init(BSDFContext.Context, BSDFContext.N, BSDFContext.V, NewL);
+	BSDFContext.H = normalize(BSDFContext.V + BSDFContext.L);
+	BSDFContext.TangentH = mul(BSDFContext.TangentBasis, BSDFContext.H);
+}
+
 struct FStrataEvaluateResult
 {
 	float3 Value;
@@ -385,16 +394,71 @@ FStrataEvaluateResult StrataEvaluateBSDF(FStrataBSDFContext BSDFContext)
 struct FStrataImportanceSampleResult
 {
 	float3 Direction;
-	float ConeHalfAngle;
-	// add a SH representation here?
+	float ConeAngle;		// Cone angle is 1/PDF as defined by Lumen
 
-	FStrataEvaluateResult Eval;
+//	FStrataEvaluateResult Eval;
 };
 
-FStrataImportanceSampleResult StrataSampleBSDF(FStrataBSDF BSDF, float2 U, BxDFContext Context)
+/**
+ * Importance sample a Strata BSDF
+ * BSDF: the strata BSDF to importance sample
+ * E: two random numbers
+ * CameraVector: vector from the camera to the considered direction (not V)
+ */
+FStrataImportanceSampleResult StrataImportanceSampleBSDF(FStrataBSDF BSDF, float2 E, float3 CameraVector)
 {
 	FStrataImportanceSampleResult Sample = (FStrataImportanceSampleResult)0;
-	// TODO
+
+	// Create a context from the information we have
+	const float3 UnusedL = float3(0.0f, 0.0f, 1.0f);
+	const float3 V = -CameraVector;
+	FStrataBSDFContext BSDFContext = StrataCreateBSDFContext(BSDF, V, UnusedL);
+
+	float PDF = 0.0f;
+	switch (BSDFContext.BSDF.BSDFType)
+	{
+		case STRATA_BSDF_TYPE_DIFFUSE_ON:
+		case STRATA_BSDF_TYPE_DIFFUSE_CHAN:
+		{
+			// STRATA_TODO take into account roughness
+			float4 ImportanceSample = CosineSampleHemisphere(E);
+			Sample.Direction = mul(ImportanceSample.xyz, BSDFContext.TangentBasis);
+			PDF = ImportanceSample.w;
+			Sample.ConeAngle = 1.0f / PDF;
+			break;
+		}
+		case STRATA_BSDF_TYPE_DIELECTRIC:
+		{
+			float4 GGXSample = ImportanceSampleVisibleGGX(UniformSampleDisk(E), Pow4(DIELECTRIC_ROUGHNESS(BSDFContext.BSDF)), BSDFContext.TangentV);
+			float3 WorldH = mul(GGXSample.xyz, BSDFContext.TangentBasis);
+			Sample.Direction = reflect(CameraVector, WorldH);
+			PDF = GGXSample.w;
+			Sample.ConeAngle = 1.0f / PDF;
+			break;
+		}
+		case STRATA_BSDF_TYPE_CONDUCTOR:
+		{
+			float4 GGXSample = ImportanceSampleVisibleGGX(UniformSampleDisk(E), Pow4(CONDUCTOR_ROUGHNESS(BSDFContext.BSDF)), BSDFContext.TangentV);
+			float3 WorldH = mul(GGXSample.xyz, BSDFContext.TangentBasis);
+			Sample.Direction = reflect(CameraVector, WorldH);
+			PDF = GGXSample.w;
+			Sample.ConeAngle = 1.0f / PDF;
+			break;
+		}
+		case STRATA_BSDF_TYPE_VOLUME:
+		{
+			float4 PhaseSample = ImportanceSampleHenyeyGreensteinPhase(E, VOLUME_ANISOTROPY(BSDFContext.BSDF));
+			Sample.Direction = mul(PhaseSample.xyz, BSDFContext.TangentBasis);
+			PDF = PhaseSample.w;
+			Sample.ConeAngle = 1.0f / PDF;
+			break;
+		}
+	}
+
+	//StrataUpdateBSDFContext(BSDFContext, Sample.Direction);
+	//Sample.Eval = StrataEvaluateBSDF(BSDFContext);
+	//Sample.Eval.PDF = PDF;	// We override the pdf to make sure the compiler skips all pdf code from StrataEvaluateBSDF because GGXSample.w si the pdf already for instance.
+
 	return Sample;
 }
 
