Fix dimension assignment in SkyLight sampling for path tracer

The x component was already being used for light selection, so rendering with a skydome and finite light could produce visibly biased results compared to a reference render without hierarchical sampling.

diff --git a/Engine/Shaders/Private/PathTracing/Light/PathTracingSkyLight.ush b/Engine/Shaders/Private/PathTracing/Light/PathTracingSkyLight.ush
--- a/Engine/Shaders/Private/PathTracing/Light/PathTracingSkyLight.ush
+++ b/Engine/Shaders/Private/PathTracing/Light/PathTracingSkyLight.ush
@@ -101,8 +101,8 @@ void SkyLight_SampleLight(
 	//uint MipCount = Log2(SkyLight.MipDimensions.x);
 	uint MipCount = log2(SkyLight.MipDimensions.x);
 	float FacePdf = 0.0;
-	// FIXME: RandSample.x has already been used to pick a light source ...
-	uint FaceIndex = SampleFace(MipCount, RandSample.x, FacePdf);
+
+	uint FaceIndex = SampleFace(MipCount, RandSample.w, FacePdf);
 
 	uint2 MipPixel = 0;
 	float MipPdf = 1.0;
@@ -151,7 +151,7 @@ void SkyLight_SampleLight(
 	// DEBUG: assert cached PDF is equivalent to re-evaluation
 	//MipPdf = PdfSkyLightOcclusionRay(WorldNormal, RayDirection);
 
-	float2 PixelOffset = RandSample.zw;
+	float2 PixelOffset = float2(DecisionTreeSample, RandSample.z);
 	float2 MipPixelUV = (MipPixel + PixelOffset) / float2(DimensionsAtLevel(SkyLight.MipDimensions.xy, StopLevel));
 
 	// Construct Cubemap coordinate
