Fixed Strata compilation warning.

diff --git a/Engine/Shaders/Private/Strata/Strata.ush b/Engine/Shaders/Private/Strata/Strata.ush
--- a/Engine/Shaders/Private/Strata/Strata.ush
+++ b/Engine/Shaders/Private/Strata/Strata.ush
@@ -429,7 +429,7 @@ FStrataImportanceSampleResult StrataImportanceSampleBSDF(FStrataBSDF BSDF, float
 		}
 		case STRATA_BSDF_TYPE_DIELECTRIC:
 		{
-			float4 GGXSample = ImportanceSampleVisibleGGX(UniformSampleDisk(E), Pow4(DIELECTRIC_ROUGHNESS(BSDFContext.BSDF)), BSDFContext.TangentV);
+			float4 GGXSample = ImportanceSampleVisibleGGX(UniformSampleDisk(E), Pow4(DIELECTRIC_ROUGHNESS(BSDFContext.BSDF)).x, BSDFContext.TangentV);
 			float3 WorldH = mul(GGXSample.xyz, BSDFContext.TangentBasis);
 			Sample.Direction = reflect(CameraVector, WorldH);
 			PDF = GGXSample.w;
@@ -438,7 +438,7 @@ FStrataImportanceSampleResult StrataImportanceSampleBSDF(FStrataBSDF BSDF, float
 		}
 		case STRATA_BSDF_TYPE_CONDUCTOR:
 		{
-			float4 GGXSample = ImportanceSampleVisibleGGX(UniformSampleDisk(E), Pow4(CONDUCTOR_ROUGHNESS(BSDFContext.BSDF)), BSDFContext.TangentV);
+			float4 GGXSample = ImportanceSampleVisibleGGX(UniformSampleDisk(E), Pow4(CONDUCTOR_ROUGHNESS(BSDFContext.BSDF)).x, BSDFContext.TangentV);
 			float3 WorldH = mul(GGXSample.xyz, BSDFContext.TangentBasis);
 			Sample.Direction = reflect(CameraVector, WorldH);
 			PDF = GGXSample.w;
