Fixed cloud compilation issue.

diff --git a/Engine/Shaders/Private/VolumetricCloud.usf b/Engine/Shaders/Private/VolumetricCloud.usf
--- a/Engine/Shaders/Private/VolumetricCloud.usf
+++ b/Engine/Shaders/Private/VolumetricCloud.usf
@@ -18,6 +18,7 @@
 #include "/Engine/Generated/VertexFactory.ush"
 
 #include "ReflectionEnvironmentShared.ush"
+#include "ParticipatingMediaCommon.ush"
 
 
 struct RenderVolumetricCloudVSToPS
@@ -285,11 +286,6 @@ void UpdateMaterialCloudParam(inout FMaterialPixelParameters MaterialParameters,
 
 #endif
 
-float IsotropicPhase()
-{
-	return 1.0f / (4.0f * PI);
-}
-
 // TODO MAKE all phase function common between sky, cloud anmd volumetric
 float HgPhase(float G, float CosTheta)
 {
