Added StrataNormalForDiffuseTransfert for Lumen.

diff --git a/Engine/Shaders/Private/Strata/Strata.ush b/Engine/Shaders/Private/Strata/Strata.ush
--- a/Engine/Shaders/Private/Strata/Strata.ush
+++ b/Engine/Shaders/Private/Strata/Strata.ush
@@ -482,7 +482,7 @@ FStrataEnvLightResult StrataEvaluateForEnvLight(FStrataBSDFContext BSDFContext,
 		case STRATA_BSDF_TYPE_DIFFUSE_CHAN:
 		{
 			StrataEnvLightResult.DiffuseAlbedo = DIFFUSE_ALBEDO(BSDFContext.BSDF);
-			StrataEnvLightResult.DiffuseNormal = DIFFUSE_NORMAL(BSDFContext.BSDF);
+			StrataEnvLightResult.DiffuseNormal = BSDFContext.N;
 			break;
 		}
 		case STRATA_BSDF_TYPE_DIELECTRIC:
@@ -509,8 +509,6 @@ FStrataEnvLightResult StrataEvaluateForEnvLight(FStrataBSDFContext BSDFContext,
 				// See SpecularGGX, isotropic verion
 
 				const float SafeRoughness = MakeRoughnessSafe(CONDUCTOR_ROUGHNESS(BSDFContext.BSDF).x);
-
-				const float3 N = CONDUCTOR_NORMAL(BSDFContext.BSDF);
 				float3 EvalEnvBRDF = EnvBRDFApprox(CONDUCTOR_REFLECTIVITY(BSDFContext.BSDF), CONDUCTOR_EDGECOLOR(BSDFContext.BSDF), SafeRoughness, BSDFContext.SatNoV);
 
 				StrataEnvLightResult.SpecularDirection = StrataGetOffSpecularPeakReflectionDir(BSDFContext.N, BSDFContext.R, SafeRoughness);
@@ -536,7 +534,7 @@ FStrataEnvLightResult StrataEvaluateForEnvLight(FStrataBSDFContext BSDFContext,
 			float3 NormLuminanceSingleScattering = (NormLuminanceTransfert - NormLuminanceTransfert * SafePathSegmentTransmittance) / SafeExtinctionCoefficients;
 
 			StrataEnvLightResult.DiffuseAlbedo = NormLuminanceSingleScattering;
-			StrataEnvLightResult.DiffuseNormal = VOLUME_SMOOTHSURFACENORMAL(BSDFContext.BSDF);
+			StrataEnvLightResult.DiffuseNormal = BSDFContext.N;
 			break;
 		}
 	}
@@ -544,6 +542,11 @@ FStrataEnvLightResult StrataEvaluateForEnvLight(FStrataBSDFContext BSDFContext,
 	return StrataEnvLightResult;
 }
 
+// Always return a normal for diffuse transfer as needed by Lumen
+float3 StrataNormalForDiffuseTransfert(FStrataBSDFContext BSDFContext)
+{
+	return BSDFContext.N;
+}
 
 
 ///////////////////////////////////////////////////////////////////////////////
