Strata
- Common normal getter
- Clean up sanitize

diff --git a/Engine/Shaders/Private/Strata/Strata.ush b/Engine/Shaders/Private/Strata/Strata.ush
--- a/Engine/Shaders/Private/Strata/Strata.ush
+++ b/Engine/Shaders/Private/Strata/Strata.ush
@@ -241,6 +241,34 @@ struct FStrataBSDFContext
 	float3 TangentH;
 };
 
+void StrataGetBSDFNormal(in FStrataBSDF BSDF, inout float3 N)
+{
+	switch (BSDF.BSDFType)
+	{
+		case STRATA_BSDF_TYPE_DIFFUSE_ON:
+		case STRATA_BSDF_TYPE_DIFFUSE_CHAN:
+		{ 
+			N = DIFFUSE_NORMAL(BSDF);
+			break; 
+		}
+		case STRATA_BSDF_TYPE_DIELECTRIC:
+		{ 
+			N = DIELECTRIC_NORMAL(BSDF);
+			break; 
+		}
+		case STRATA_BSDF_TYPE_CONDUCTOR:
+		{
+			N = CONDUCTOR_NORMAL(BSDF);
+			break;
+		}
+		case STRATA_BSDF_TYPE_VOLUME:
+		{
+			N = VOLUME_SMOOTHSURFACENORMAL(BSDF);
+			break;
+		}
+	}
+}
+
 FStrataBSDFContext StrataCreateBSDFContext(FStrataBSDF BSDF, float3 V, float3 L)
 {
 	FStrataBSDFContext BSDFContext = (FStrataBSDFContext)0;
@@ -249,19 +277,9 @@ FStrataBSDFContext StrataCreateBSDFContext(FStrataBSDF BSDF, float3 V, float3 L)
 
 	BSDFContext.V = V;
 	BSDFContext.L = L;
+
 	BSDFContext.N = float3(0.0f, 0.0f, 1.0f);
-	switch (BSDF.BSDFType)		// TODO_STRATA: all have a normal, even subsurface?
-	{
-	case STRATA_BSDF_TYPE_DIFFUSE_ON:
-	case STRATA_BSDF_TYPE_DIFFUSE_CHAN:
-	{ BSDFContext.N = DIFFUSE_NORMAL(BSDF);				break; }
-	case STRATA_BSDF_TYPE_DIELECTRIC:
-	{ BSDFContext.N = DIELECTRIC_NORMAL(BSDF);			break; }
-	case STRATA_BSDF_TYPE_CONDUCTOR:
-	{ BSDFContext.N = CONDUCTOR_NORMAL(BSDF);			break; }
-	case STRATA_BSDF_TYPE_VOLUME:
-	{ BSDFContext.N = VOLUME_SMOOTHSURFACENORMAL(BSDF);	break; }
-	}
+	StrataGetBSDFNormal(BSDF, BSDFContext.N);
 	BSDFContext.N = normalize(BSDFContext.N);
 
 	BSDFContext.R = 2 * dot(BSDFContext.V, BSDFContext.N) * BSDFContext.N - BSDFContext.V;
@@ -652,6 +670,33 @@ void SanitizeStrataVolume(inout FStrataBSDF BSDF)
 	VOLUME_THICKNESS(BSDF) = max(0.0f, VOLUME_THICKNESS(BSDF));
 }
 
+void StrataSanitizeBSDF(inout FStrataBSDF BSDF)
+{
+	switch (BSDF.BSDFType)
+	{
+		case STRATA_BSDF_TYPE_DIFFUSE_ON:
+		case STRATA_BSDF_TYPE_DIFFUSE_CHAN:
+		{
+			SanitizeStrataDiffuse(BSDF);
+			break; 
+		}
+		case STRATA_BSDF_TYPE_DIELECTRIC:
+		{
+			SanitizeStrataDielectric(BSDF);
+			break;
+		}
+		case STRATA_BSDF_TYPE_CONDUCTOR:
+		{
+			SanitizeStrataConductor(BSDF);
+			break;
+		}
+		case STRATA_BSDF_TYPE_VOLUME:
+		{
+			SanitizeStrataVolume(BSDF);
+			break;
+		}
+	}
+}
 
 
 ///////////////////////////////////////////////////////////////////////////////
@@ -822,27 +867,18 @@ void PackStrataOut(
 
 		for (int i = 0; i < Strata.Layers[l].BSDFCount; ++i)
 		{
+			// Sanitize BSDF in place before it is used and packed out to memory
+			StrataSanitizeBSDF(Strata.Layers[l].BSDFs[i]);
+
 			FStrataBSDF BSDF = Strata.Layers[l].BSDFs[i];
 			const float3 BSDFCoverage = BSDF.Weight;
 			BSDFCount++;
 
+			// STRATA_TODO apply the same normal transform as for the material
 			float3 N = VertexNormal;
-			switch (BSDF.BSDFType)		// TODO_STRATA: all have a normal, even subsurface?
-			{
-			case STRATA_BSDF_TYPE_DIFFUSE_ON:
-			case STRATA_BSDF_TYPE_DIFFUSE_CHAN:
-			{ N = DIFFUSE_NORMAL(BSDF);				SanitizeStrataDiffuse(BSDF);	break; }
-			case STRATA_BSDF_TYPE_DIELECTRIC:
-			{ N = DIELECTRIC_NORMAL(BSDF);			SanitizeStrataDielectric(BSDF);	break; }
-			case STRATA_BSDF_TYPE_CONDUCTOR:
-			{ N = CONDUCTOR_NORMAL(BSDF);			SanitizeStrataConductor(BSDF);	break; }
-			case STRATA_BSDF_TYPE_VOLUME:
-			{ N = VOLUME_SMOOTHSURFACENORMAL(BSDF);	SanitizeStrataVolume(BSDF);		break; }
-			}
+			StrataGetBSDFNormal(BSDF, N);
 			N = normalize(N);
 
-			// STRATA_TODO apply the same normal transform as for the material
-
 			BxDFContext ContextNoLight = (BxDFContext)0;
 			const float3 FakeL = float3(0, 0, 1);
 			Init(ContextNoLight, N, V, FakeL);
diff --git a/Engine/Shaders/Private/Strata/StrataForwardLighting.ush b/Engine/Shaders/Private/Strata/StrataForwardLighting.ush
--- a/Engine/Shaders/Private/Strata/StrataForwardLighting.ush
+++ b/Engine/Shaders/Private/Strata/StrataForwardLighting.ush
@@ -30,6 +30,10 @@ float3 StrataForwardLighting(
 		for (int i = 0; i < Strata.Layers[l].BSDFCount; ++i)
 		{
 			FStrataBSDF BSDF = Strata.Layers[l].BSDFs[i];
+
+			// Sanitize BSDF before it is used for forward shading
+			StrataSanitizeBSDF(BSDF);
+
 			const float3 BSDFCoverage = BSDF.Weight;
 
 			// Create the BSDF context
