Fix compilation warning for shader referencing hair attributes node.

The resources referenced by the atttributes node where not define due to the vertex factory file not being include. This change moves the resources & function definition into a separate header which is included in both the VertexFactory and the MaterialTEmplate file.

This also uniformize the naming of the vertex factory define.

diff --git a/Engine/Plugins/Experimental/HairStrands/Source/HairStrandsCore/Private/HairCardsVertexFactory.cpp b/Engine/Plugins/Experimental/HairStrands/Source/HairStrandsCore/Private/HairCardsVertexFactory.cpp
--- a/Engine/Plugins/Experimental/HairStrands/Source/HairStrandsCore/Private/HairCardsVertexFactory.cpp
+++ b/Engine/Plugins/Experimental/HairStrands/Source/HairStrandsCore/Private/HairCardsVertexFactory.cpp
@@ -198,7 +198,7 @@ void FHairCardsVertexFactory::ModifyCompilationEnvironment(const FVertexFactoryS
 
 	const bool bUseGPUSceneAndPrimitiveIdStream = VF_CARDS_SUPPORT_GPU_SCENE && Parameters.VertexFactoryType->SupportsPrimitiveIdStream() && UseGPUScene(Parameters.Platform, GetMaxSupportedFeatureLevel(Parameters.Platform));
 	OutEnvironment.SetDefine(TEXT("VF_SUPPORTS_PRIMITIVE_SCENE_DATA"), bUseGPUSceneAndPrimitiveIdStream);
-	OutEnvironment.SetDefine(TEXT("VF_CARDS_HAIR"), TEXT("1"));
+	OutEnvironment.SetDefine(TEXT("HAIR_CARD_MESH_FACTORY"), TEXT("1"));
 	OutEnvironment.SetDefine(TEXT("VF_GPU_SCENE_TEXTURE"), bUseGPUSceneAndPrimitiveIdStream && GPUSceneUseTexture2D(Parameters.Platform));
 	OutEnvironment.SetDefine(TEXT("MANUAL_VERTEX_FETCH"), RHISupportsManualVertexFetch(Parameters.Platform) ? TEXT("1") : TEXT("0"));
 }
diff --git a/Engine/Plugins/Experimental/HairStrands/Source/HairStrandsCore/Private/HairStrandsVertexFactory.cpp b/Engine/Plugins/Experimental/HairStrands/Source/HairStrandsCore/Private/HairStrandsVertexFactory.cpp
--- a/Engine/Plugins/Experimental/HairStrands/Source/HairStrandsCore/Private/HairStrandsVertexFactory.cpp
+++ b/Engine/Plugins/Experimental/HairStrands/Source/HairStrandsCore/Private/HairStrandsVertexFactory.cpp
@@ -175,7 +175,7 @@ void FHairStrandsVertexFactory::ModifyCompilationEnvironment(const FVertexFactor
 {
 	const bool bUseGPUSceneAndPrimitiveIdStream = Parameters.VertexFactoryType->SupportsPrimitiveIdStream() && UseGPUScene(Parameters.Platform, GetMaxSupportedFeatureLevel(Parameters.Platform));
 	OutEnvironment.SetDefine(TEXT("VF_SUPPORTS_PRIMITIVE_SCENE_DATA"), bUseGPUSceneAndPrimitiveIdStream);
-	OutEnvironment.SetDefine(TEXT("VF_STRAND_HAIR"), TEXT("1"));
+	OutEnvironment.SetDefine(TEXT("HAIR_STRAND_MESH_FACTORY"), TEXT("1"));
 	OutEnvironment.SetDefine(TEXT("VF_GPU_SCENE_TEXTURE"), bUseGPUSceneAndPrimitiveIdStream && GPUSceneUseTexture2D(Parameters.Platform));
 }
 
diff --git a/Engine/Shaders/Private/HairStrands/HairCardsAttributeCommon.ush b/Engine/Shaders/Private/HairStrands/HairCardsAttributeCommon.ush
new file mode 100644
--- /dev/null
+++ b/Engine/Shaders/Private/HairStrands/HairCardsAttributeCommon.ush
@@ -0,0 +1,100 @@
+// Copyright Epic Games, Inc. All Rights Reserved.
+
+#pragma once
+
+#if HAIR_CARD_MESH_FACTORY
+
+#include "/Engine/Private/HairStrands/HairStrandsVertexFactoryCommon.ush"
+
+////////////////////////////////////////////////////////////////////////////////
+
+struct FHairVertexAttributes
+{
+	float2 HairUV;
+	float2 HairRootUV;
+	float2 HairDimensions;
+	float3 HairBaseColor;
+	float  HairSeed;
+	float  HairRoughness;
+};
+
+float GetHairStrandsDepth(uint HairPrimitiveId, float2 InUV, float InDeviceZ)
+{
+	InUV.y = HairCardsVF.bInvertUV ? 1 - InUV.y : InUV.y;
+	const float SceneDepthOffset = Texture2DSample(HairCardsVF.DepthTexture, HairCardsVF.DepthSampler, InUV).x;
+	const float SceneDepth = ConvertFromDeviceZ(InDeviceZ);
+
+	return ConvertToDeviceZ(SceneDepth + SceneDepthOffset);
+}
+
+float3 GetHairStrandsTangent(uint HairPrimitiveId, float2 InUV, half3x3 TangentToWorld, bool bUseTangentSpace)
+{
+	InUV.y = HairCardsVF.bInvertUV ? 1 - InUV.y : InUV.y;
+	const float3 LocalTangent = Texture2DSample(HairCardsVF.TangentTexture, HairCardsVF.TangentSampler, InUV).xyz * 2 - 1;
+	const float3 WorldTangent = mul(LocalTangent, TangentToWorld);
+	return bUseTangentSpace ? LocalTangent : WorldTangent;
+}
+
+float GetHairStrandsCoverage(uint HairPrimitiveId, float2 InUV)
+{
+	InUV.y = HairCardsVF.bInvertUV ? 1 - InUV.y : InUV.y;
+	return Texture2DSample(HairCardsVF.CoverageTexture, HairCardsVF.CoverageSampler, InUV).x;
+}
+
+FHairVertexAttributes GetHairStrandsAttributes(uint HairPrimitiveId, float2 InUV)
+{
+	InUV.y = HairCardsVF.bInvertUV ? 1 - InUV.y : InUV.y;
+	const float4 Value = Texture2DSample(HairCardsVF.AttributeTexture, HairCardsVF.AttributeSampler, InUV);
+
+	FHairVertexAttributes Out;
+	Out.HairUV = float2(Value.z, 0.5f);
+	Out.HairRootUV = Value.xy;
+	Out.HairDimensions = 0;
+	Out.HairBaseColor = 0;
+	Out.HairSeed = Value.w;
+	Out.HairRoughness = 0;
+	return Out;
+}
+
+float4 GetHairStrandsAuxilaryData(uint HairPrimitiveId, float2 InUV)
+{
+	return Texture2DSample(HairCardsVF.AuxilaryDataTexture, HairCardsVF.AuxilaryDataSampler, InUV);
+}
+
+float2 GetHairStrandsUV(uint HairPrimitiveId, float2 HairPrimitiveUV)
+{
+	return GetHairStrandsAttributes(HairPrimitiveId, HairPrimitiveUV).HairUV;
+}
+
+float2 GetHairStrandsDimensions(uint HairPrimitiveId, float2 HairPrimitiveUV)
+{
+	return GetHairStrandsAttributes(HairPrimitiveId, HairPrimitiveUV).HairDimensions;
+}
+
+float2 GetHairStrandsRootUV(uint HairPrimitiveId, float2 HairPrimitiveUV)
+{
+#if MANUAL_VERTEX_FETCH
+	return HairCardsVF.UVsBuffer[HairPrimitiveId].zw;
+#else
+	return GetHairStrandsAttributes(HairPrimitiveId, HairPrimitiveUV).HairRootUV; // TODO fix me?
+#endif
+}
+
+float  GetHairStrandsSeed(uint HairPrimitiveId, float2 HairPrimitiveUV)
+{
+	return GetHairStrandsAttributes(HairPrimitiveId, HairPrimitiveUV).HairSeed;
+}
+
+float3 GetHairStrandsBaseColor(uint HairPrimitiveId, float2 HairPrimitiveUV)
+{
+	return GetHairStrandsAttributes(HairPrimitiveId, HairPrimitiveUV).HairBaseColor;
+}
+
+float GetHairStrandsRoughness(uint HairPrimitiveId, float2 HairPrimitiveUV)
+{
+	return GetHairStrandsAttributes(HairPrimitiveId, HairPrimitiveUV).HairRoughness;
+}
+
+////////////////////////////////////////////////////////////////////////////////
+
+#endif // HAIR_CARD_MESH_FACTORY
\ No newline at end of file
diff --git a/Engine/Shaders/Private/HairStrands/HairCardsVertexFactory.ush b/Engine/Shaders/Private/HairStrands/HairCardsVertexFactory.ush
--- a/Engine/Shaders/Private/HairStrands/HairCardsVertexFactory.ush
+++ b/Engine/Shaders/Private/HairStrands/HairCardsVertexFactory.ush
@@ -7,11 +7,10 @@
 #include "../LocalVertexFactoryCommon.ush"
 #include "HairStrandsVisibilityCommon.ush"
 #include "HairStrandsVertexFactoryCommon.ush"
+#include "HairCardsAttributeCommon.ush"
 
 #include "/Engine/Generated/UniformBuffers/PrecomputedLightingBuffer.ush"
 
-#include "HairStrandsVisibilityCommon.ush"
-
 #if RAYHITGROUPSHADER
 #include "../RayTracing/RayTracingCommon.ush"
 #include "../RayTracing/RayTracingHitGroupCommon.ush"
@@ -25,7 +24,7 @@
 #endif
 
 ////////////////////////////////////////////////////////////////////////////////
-// VF_CARDS_HAIR
+// HAIR_CARD_MESH_FACTORY
 // Define used in certain shader like DeepShadow for running specific code (coverage computation for instance) 
 // when vertices are produced by the strand hair vertex factory 
 // This is set by compilation enviromenent of the vertex factory
@@ -37,103 +36,6 @@
 #define FPositionOnlyVertexFactoryInput FVertexFactoryInput
 #define FPositionAndNormalOnlyVertexFactoryInput FVertexFactoryInput
 
-struct FHairVertexAttributes
-{
-	float2 HairUV;
-	float2 HairRootUV;
-	float2 HairDimensions;
-	float3 HairBaseColor;
-	float  HairSeed;
-	float  HairRoughness;
-};
-
-float GetHairStrandsDepth(uint HairPrimitiveId, float2 InUV, float InDeviceZ)
-{
-	InUV.y = HairCardsVF.bInvertUV ? 1 - InUV.y : InUV.y;
-	const float SceneDepthOffset = Texture2DSample(HairCardsVF.DepthTexture, HairCardsVF.DepthSampler, InUV).x;
-	const float SceneDepth = ConvertFromDeviceZ(InDeviceZ);
-
-	return ConvertToDeviceZ(SceneDepth + SceneDepthOffset);
-}
-
-float3 GetHairStrandsTangent(uint HairPrimitiveId, float2 InUV, half3x3 TangentToWorld, bool bUseTangentSpace)
-{
-	InUV.y = HairCardsVF.bInvertUV ? 1 - InUV.y : InUV.y;
-	const float3 LocalTangent = Texture2DSample(HairCardsVF.TangentTexture, HairCardsVF.TangentSampler, InUV).xyz * 2 - 1;
-	const float3 WorldTangent = mul(LocalTangent, TangentToWorld);
-	return bUseTangentSpace ? LocalTangent : WorldTangent;
-}
-
-float GetHairStrandsCoverage(uint HairPrimitiveId, float2 InUV)
-{
-	InUV.y = HairCardsVF.bInvertUV ? 1 - InUV.y : InUV.y;
-	return Texture2DSample(HairCardsVF.CoverageTexture, HairCardsVF.CoverageSampler, InUV).x;
-}
-
-FHairVertexAttributes GetHairStrandsAttributes(uint HairPrimitiveId, float2 InUV)
-{
-	InUV.y = HairCardsVF.bInvertUV ? 1 - InUV.y : InUV.y;
-	const float4 Value = Texture2DSample(HairCardsVF.AttributeTexture, HairCardsVF.AttributeSampler, InUV);
-
-	FHairVertexAttributes Out;
-	Out.HairUV = float2(Value.z, 0.5f);
-	Out.HairRootUV = Value.xy;
-	Out.HairDimensions = 0;
-	Out.HairBaseColor = 0;
-	Out.HairSeed = Value.w;
-	Out.HairRoughness = 0;
-	return Out;
-}
-
-float4 GetHairStrandsAuxilaryData(uint HairPrimitiveId, float2 InUV)
-{
-	return Texture2DSample(HairCardsVF.AuxilaryDataTexture, HairCardsVF.AuxilaryDataSampler, InUV);
-}
-
-float2 GetHairStrandsUV(uint HairPrimitiveId, float2 HairPrimitiveUV)
-{
-	return GetHairStrandsAttributes(HairPrimitiveId, HairPrimitiveUV).HairUV;
-}
-
-float2 GetHairStrandsDimensions(uint HairPrimitiveId, float2 HairPrimitiveUV)
-{
-	return GetHairStrandsAttributes(HairPrimitiveId, HairPrimitiveUV).HairDimensions;
-}
-
-float2 GetHairStrandsRootUV(uint HairPrimitiveId, float2 HairPrimitiveUV)
-{
-#if MANUAL_VERTEX_FETCH
-	return HairCardsVF.UVsBuffer[HairPrimitiveId].zw;
-#else
-	return GetHairStrandsAttributes(HairPrimitiveId, HairPrimitiveUV).HairRootUV; // TODO fix me?
-#endif
-}
-
-float  GetHairStrandsSeed(uint HairPrimitiveId, float2 HairPrimitiveUV)
-{
-	return GetHairStrandsAttributes(HairPrimitiveId, HairPrimitiveUV).HairSeed;
-}
-
-float3 GetHairStrandsBaseColor(uint HairPrimitiveId, float2 HairPrimitiveUV)
-{
-	return GetHairStrandsAttributes(HairPrimitiveId, HairPrimitiveUV).HairBaseColor;
-}
-
-float GetHairStrandsRoughness(uint HairPrimitiveId, float2 HairPrimitiveUV)
-{
-	return GetHairStrandsAttributes(HairPrimitiveId, HairPrimitiveUV).HairRoughness;
-}
-
-//float GetWorldStrandRadius(FVertexFactoryInput Input)
-//{
-//	return 0;// GetVertexPosition(Input).WorldRadius;
-//}
-//
-//float GetWorldStrandLength(FVertexFactoryInput Input)
-//{
-//	return 0;//GetVertexPosition(Input).WorldLength;
-//}
-
 struct FVertexFactoryInput
 {
 #if MANUAL_VERTEX_FETCH
@@ -588,7 +490,7 @@ FVertexFactoryInterpolantsVSToDS VertexFactoryInterpolate(FVertexFactoryInterpol
 	TESSELLATION_INTERPOLATE_MEMBER(InterpolantsVSToPS.TexCoords[0]);
 #endif
 
-#if VF_STRAND_HAIR || VF_CARDS_HAIR
+#if HAIR_CARD_MESH_FACTORY
 	O.InterpolantsVSToPS.HairPrimitiveId = a.InterpolantsVSToPS.HairPrimitiveId;
 	TESSELLATION_INTERPOLATE_MEMBER(InterpolantsVSToPS.HairPrimitiveUV);
 #endif
diff --git a/Engine/Shaders/Private/HairStrands/HairStrandsAttributeCommon.ush b/Engine/Shaders/Private/HairStrands/HairStrandsAttributeCommon.ush
new file mode 100644
--- /dev/null
+++ b/Engine/Shaders/Private/HairStrands/HairStrandsAttributeCommon.ush
@@ -0,0 +1,110 @@
+// Copyright Epic Games, Inc. All Rights Reserved.
+
+#pragma once 
+
+#if HAIR_STRAND_MESH_FACTORY
+
+#include "/Engine/Private/HairStrands/HairStrandsVertexFactoryCommon.ush"
+
+////////////////////////////////////////////////////////////////////////////////
+
+float  HairStrandsVF_Radius;
+float  HairStrandsVF_Length;
+
+Buffer<uint4>  HairStrandsVF_PositionBuffer;
+Buffer<uint4>  HairStrandsVF_AttributeBuffer;
+Buffer<float4> HairStrandsVF_MaterialBuffer;
+
+struct FHairVertexAttributes
+{
+	float2 HairUV;
+	float2 HairRootUV;
+	float2 HairDimensions;
+	float3 HairBaseColor;
+	float  HairSeed;
+	float  HairRoughness;
+};
+
+FHairVertexAttributes GetHairStrandsAttributes(uint HairPrimitiveId, float2 VertexUV)
+{
+	const uint VertexId0 = HairPrimitiveId;
+	const uint VertexId1 = HairPrimitiveId+1;
+
+	// Instead of doing the actual interpolation we could rely on dithering as an optimization
+	const FHairAttributes Attributes0 = DecodeHairAttribute(HairStrandsVF_AttributeBuffer[VertexId0]);
+	const FHairAttributes Attributes1 = DecodeHairAttribute(HairStrandsVF_AttributeBuffer[VertexId1]);
+
+	const uint4 PackedData = HairStrandsVF_PositionBuffer[VertexId0];
+	const uint PackedAlpha = asuint(PackedData.a);
+	const float WorldRadius = ((PackedAlpha & 0x00FC) >> 2) / 63.f * HairStrandsVF_Radius;
+	const float WorldLength = ((PackedAlpha & 0xFF00) >> 8) / 255.f * HairStrandsVF_Length;
+
+	FHairVertexAttributes Out;
+	// Constant along the strand
+	Out.HairDimensions	= float2(WorldLength, WorldRadius);
+	Out.HairSeed		= Attributes0.Seed;
+	Out.HairRootUV		= Attributes0.RootUV;
+
+	// Varying along the strand
+	Out.HairUV.y		= VertexUV.y;
+	Out.HairUV.x		= lerp(Attributes0.UCoord, Attributes1.UCoord, VertexUV.x);
+
+	// Instead of doing the actual interpolation we could rely on dithering as an optimization
+	const float4 Material0 = HairStrandsVF_MaterialBuffer[VertexId0];
+	const float4 Material1 = HairStrandsVF_MaterialBuffer[VertexId1];
+
+	// Cheap/manual sRGB decoding as the underlying format is linear quantization
+	Out.HairBaseColor	= lerp(Material0.xyz*Material0.xyz, Material1.xyz*Material1.xyz, VertexUV.x);
+	Out.HairRoughness	= lerp(Material0.w, Material1.w, VertexUV.x);
+
+	return Out; 
+}
+
+float GetHairStrandsDepth(float2 InUV)
+{
+	return 0; // TODO: what should we do for the strands cards as we don't want to override the depth output, but we still want to share the same shader than the non cards one
+}
+
+float GetHairStrandsCoverage(uint HairPrimitiveId, float2 HairPrimitiveUV)
+{
+	return 1;
+}
+
+float2 GetHairStrandsUV(uint HairPrimitiveId, float2 HairPrimitiveUV)
+{
+	return GetHairStrandsAttributes(HairPrimitiveId, HairPrimitiveUV).HairUV;
+}
+
+float2 GetHairStrandsDimensions(uint HairPrimitiveId, float2 HairPrimitiveUV)
+{
+	return GetHairStrandsAttributes(HairPrimitiveId, HairPrimitiveUV).HairDimensions;
+}
+
+float2 GetHairStrandsRootUV(uint HairPrimitiveId, float2 HairPrimitiveUV)
+{
+	return GetHairStrandsAttributes(HairPrimitiveId, HairPrimitiveUV).HairRootUV;
+}
+
+float  GetHairStrandsSeed(uint HairPrimitiveId, float2 HairPrimitiveUV)
+{
+	return GetHairStrandsAttributes(HairPrimitiveId, HairPrimitiveUV).HairSeed;
+}
+
+float3 GetHairStrandsBaseColor(uint HairPrimitiveId, float2 HairPrimitiveUV)
+{
+	return GetHairStrandsAttributes(HairPrimitiveId, HairPrimitiveUV).HairBaseColor;
+}
+
+float GetHairStrandsRoughness(uint HairPrimitiveId, float2 HairPrimitiveUV)
+{
+	return GetHairStrandsAttributes(HairPrimitiveId, HairPrimitiveUV).HairRoughness;
+}
+
+float GetHairStrandsDepth(uint HairPrimitiveId, float2 InUV, float InDeviceZ)
+{
+	return 0.f;
+}
+
+////////////////////////////////////////////////////////////////////////////////
+
+#endif // HAIR_STRAND_MESH_FACTORY
\ No newline at end of file
diff --git a/Engine/Shaders/Private/HairStrands/HairStrandsDeepShadowPS.usf b/Engine/Shaders/Private/HairStrands/HairStrandsDeepShadowPS.usf
--- a/Engine/Shaders/Private/HairStrands/HairStrandsDeepShadowPS.usf
+++ b/Engine/Shaders/Private/HairStrands/HairStrandsDeepShadowPS.usf
@@ -71,7 +71,7 @@ void MainVoxel(
 	const float3 WorldPosition = WorldPositon_Coverage.xyz;
 	const float HairCoverage = WorldPositon_Coverage.w;
 
-#if VF_STRAND_HAIR  
+#if HAIR_STRAND_MESH_FACTORY  
 	const FPackedVirtualVoxelNodeDesc PackedNode = VoxelRasterPass.VirtualVoxel.NodeDescBuffer[VoxelRasterPass.MacroGroupId];
 	const FVirtualVoxelNodeDesc NodeDesc = UnpackVoxelNode(PackedNode, VoxelRasterPass.VirtualVoxel.PageResolution);
 
diff --git a/Engine/Shaders/Private/HairStrands/HairStrandsDeepShadowVS.usf b/Engine/Shaders/Private/HairStrands/HairStrandsDeepShadowVS.usf
--- a/Engine/Shaders/Private/HairStrands/HairStrandsDeepShadowVS.usf
+++ b/Engine/Shaders/Private/HairStrands/HairStrandsDeepShadowVS.usf
@@ -46,7 +46,7 @@ void Main(
 	OPTIONAL_VertexID
 	out VS_OUTPUT_TYPE Output
 )
-#if defined(VF_STRAND_HAIR) && MESH_RENDER_MODE == MESH_RENDER_VOXEL && SUPPORT_TANGENT_PROPERTY == 0 && SUPPORT_MATERIAL_PROPERTY == 0
+#if defined(HAIR_STRAND_MESH_FACTORY) && MESH_RENDER_MODE == MESH_RENDER_VOXEL && SUPPORT_TANGENT_PROPERTY == 0 && SUPPORT_MATERIAL_PROPERTY == 0
 {
 	ResolvedView = ResolveView();
 	
@@ -100,7 +100,7 @@ void Main(
 	float4x4 WorldToClipMatrix = DeepRasterPass.CPU_WorldToClipMatrix;
 #endif
 
-#ifdef VF_STRAND_HAIR	
+#ifdef HAIR_STRAND_MESH_FACTORY	
 	const FHairRenderInfo HairRenderInfo = GetHairRenderInfo(ResolvedView.HairRenderInfo, ResolvedView.HairRenderInfoBits);
 
 	FHairViewInfo HairViewInfo;
@@ -140,7 +140,7 @@ void Main(
 
 	Output.Position = mul(WorldPosition, WorldToClipMatrix);
 
-#if MESH_RENDER_MODE == MESH_RENDER_VOXEL && defined(VF_STRAND_HAIR)
+#if MESH_RENDER_MODE == MESH_RENDER_VOXEL && defined(HAIR_STRAND_MESH_FACTORY)
 	if (HairRenderInfo.bIsGPUDriven)
 	{
 		Output.Position.xy += float2(1, 1);
@@ -158,7 +158,7 @@ void Main(
 #endif
 
 	float HairCoverage = 1;
-  #ifdef VF_STRAND_HAIR	
+  #ifdef HAIR_STRAND_MESH_FACTORY	
 	const float CurrentDepth = Output.Position.z / Output.Position.w;
 	const float PixelRadius = HairRenderInfo.bIsOrthoView ? HairRenderInfo.RadiusAtDepth1Primary : (CurrentDepth * HairRenderInfo.RadiusAtDepth1Primary);
 	const float StrandRealRadius = VFIntermediates.HairDimensions.y;
diff --git a/Engine/Shaders/Private/HairStrands/HairStrandsDeepTransmittanceMask.usf b/Engine/Shaders/Private/HairStrands/HairStrandsDeepTransmittanceMask.usf
--- a/Engine/Shaders/Private/HairStrands/HairStrandsDeepTransmittanceMask.usf
+++ b/Engine/Shaders/Private/HairStrands/HairStrandsDeepTransmittanceMask.usf
@@ -15,9 +15,6 @@
 #include "HairStrandsVisibilityCommon.ush"
 #include "HairStrandsVoxelPageCommon.ush"
 
-#include "../ShaderDrawDebug.ush"
-#include "../ShaderPrintCommon.ush"
-
 #define PERMUTATION_TRANSMITTANCE_DEEP_SHADOW 0
 #define PERMUTATION_TRANSMITTANCE_VIRTUAL_VOXEL 1
 
@@ -30,8 +27,12 @@
 	#endif
 	#define VOXEL_TRAVERSAL_DEBUG 0
 	#include "HairStrandsVoxelPageTraversal.ush"
+	#if VOXEL_TRAVERSAL_DEBUG
+	#include "../ShaderPrintCommon.ush"
+	#endif
 #endif
 
+
 float3 ToWorldPosition(float2 UV, float Depth)
 {
 	float2 ScreenPosition = (UV - View.ScreenPositionScaleBias.wz) / View.ScreenPositionScaleBias.xy;
diff --git a/Engine/Shaders/Private/HairStrands/HairStrandsMaterialCommon.ush b/Engine/Shaders/Private/HairStrands/HairStrandsMaterialCommon.ush
--- a/Engine/Shaders/Private/HairStrands/HairStrandsMaterialCommon.ush
+++ b/Engine/Shaders/Private/HairStrands/HairStrandsMaterialCommon.ush
@@ -33,7 +33,7 @@ uint MaterialPass_LightChannelMask;
 #define HAIR_MATERIAL_DEBUG_OUTPUT 0
 #endif
 
-#if VF_STRAND_HAIR
+#if HAIR_STRAND_MESH_FACTORY
 struct FEvaluateOutput
 {
 	FHairSample NodeData;
diff --git a/Engine/Shaders/Private/HairStrands/HairStrandsVertexFactory.ush b/Engine/Shaders/Private/HairStrands/HairStrandsVertexFactory.ush
--- a/Engine/Shaders/Private/HairStrands/HairStrandsVertexFactory.ush
+++ b/Engine/Shaders/Private/HairStrands/HairStrandsVertexFactory.ush
@@ -7,11 +7,10 @@
 #include "../LocalVertexFactoryCommon.ush"
 #include "HairStrandsVisibilityCommon.ush"
 #include "HairStrandsVertexFactoryCommon.ush"
+#include "HairStrandsAttributeCommon.ush"
 
 #include "/Engine/Generated/UniformBuffers/PrecomputedLightingBuffer.ush"
 
-#include "HairStrandsVisibilityCommon.ush"
-
 #define USE_HAIR_COMPLEX_TRANSMITTANCE 1
 
 #if MANUAL_VERTEX_FETCH
@@ -32,7 +31,7 @@
 #endif
 
 ////////////////////////////////////////////////////////////////////////////////
-// VF_STRAND_HAIR
+// HAIR_STRAND_MESH_FACTORY
 // Define used in certain shader like DeepShadow for running specific code (coverage computation for instance) 
 // when vertices are produced by the strand hair vertex factory 
 // This is set by compilation enviromenent of the vertex factory
@@ -41,16 +40,11 @@
 float3 HairStrandsVF_PositionOffset;			// Vertex position (3x16bits) are relative to this offset for improving precision while keeping vertex footprint low
 float3 HairStrandsVF_PreviousPositionOffset;	// Vertex position (3x16bits) are relative to this offset for improving precision while keeping vertex footprint low
 float  HairStrandsVF_Density;
-float  HairStrandsVF_Radius;
-float  HairStrandsVF_Length;
 uint   HairStrandsVF_CullingEnable;
 uint   HairStrandsVF_bUseStableRasterization;
 uint   HairStrandsVF_bScatterSceneLighing;
 
-Buffer<uint4>  HairStrandsVF_PositionBuffer;
 Buffer<uint4>  HairStrandsVF_PreviousPositionBuffer;
-Buffer<uint4>  HairStrandsVF_AttributeBuffer;
-Buffer<float4> HairStrandsVF_MaterialBuffer;
 Buffer<float4> HairStrandsVF_TangentBuffer;
 
 // ## A quad is made of the following topology and indexing:
@@ -91,91 +85,6 @@ Buffer<float4> HairStrandsVF_TangentBuffer;
 #define FPositionOnlyVertexFactoryInput FVertexFactoryInput
 #define FPositionAndNormalOnlyVertexFactoryInput FVertexFactoryInput
 
-struct FHairVertexAttributes
-{
-	float2 HairUV;
-	float2 HairRootUV;
-	float2 HairDimensions;
-	float3 HairBaseColor;
-	float  HairSeed;
-	float  HairRoughness;
-};
-
-FHairVertexAttributes GetHairStrandsAttributes(uint HairPrimitiveId, float2 VertexUV)
-{
-	const uint VertexId0 = HairPrimitiveId;
-	const uint VertexId1 = HairPrimitiveId+1;
-
-	// Instead of doing the actual interpolation we could rely on dithering as an optimization
-	const FHairAttributes Attributes0 = DecodeHairAttribute(HairStrandsVF_AttributeBuffer[VertexId0]);
-	const FHairAttributes Attributes1 = DecodeHairAttribute(HairStrandsVF_AttributeBuffer[VertexId1]);
-
-	const uint4 PackedData = HairStrandsVF_PositionBuffer[VertexId0];
-	const uint PackedAlpha = asuint(PackedData.a);
-	const float WorldRadius = ((PackedAlpha & 0x00FC) >> 2) / 63.f * HairStrandsVF_Radius;
-	const float WorldLength = ((PackedAlpha & 0xFF00) >> 8) / 255.f * HairStrandsVF_Length;
-
-	FHairVertexAttributes Out;
-	// Constant along the strand
-	Out.HairDimensions	= float2(WorldLength, WorldRadius);
-	Out.HairSeed		= Attributes0.Seed;
-	Out.HairRootUV		= Attributes0.RootUV;
-
-	// Varying along the strand
-	Out.HairUV.y		= VertexUV.y;
-	Out.HairUV.x		= lerp(Attributes0.UCoord, Attributes1.UCoord, VertexUV.x);
-
-	// Instead of doing the actual interpolation we could rely on dithering as an optimization
-	const float4 Material0 = HairStrandsVF_MaterialBuffer[VertexId0];
-	const float4 Material1 = HairStrandsVF_MaterialBuffer[VertexId1];
-
-	// Cheap/manual sRGB decoding as the underlying format is linear quantization
-	Out.HairBaseColor	= lerp(Material0.xyz*Material0.xyz, Material1.xyz*Material1.xyz, VertexUV.x);
-	Out.HairRoughness	= lerp(Material0.w, Material1.w, VertexUV.x);
-
-	return Out; 
-}
-
-float GetHairStrandsDepth(float2 InUV)
-{
-	return 0; // TODO: what should we do for the strands cards as we don't want to override the depth output, but we still want to share the same shader than the non cards one
-}
-
-float GetHairStrandsCoverage(uint HairPrimitiveId, float2 HairPrimitiveUV)
-{
-	return 1;
-}
-
-float2 GetHairStrandsUV(uint HairPrimitiveId, float2 HairPrimitiveUV)
-{
-	return GetHairStrandsAttributes(HairPrimitiveId, HairPrimitiveUV).HairUV;
-}
-
-float2 GetHairStrandsDimensions(uint HairPrimitiveId, float2 HairPrimitiveUV)
-{
-	return GetHairStrandsAttributes(HairPrimitiveId, HairPrimitiveUV).HairDimensions;
-}
-
-float2 GetHairStrandsRootUV(uint HairPrimitiveId, float2 HairPrimitiveUV)
-{
-	return GetHairStrandsAttributes(HairPrimitiveId, HairPrimitiveUV).HairRootUV;
-}
-
-float  GetHairStrandsSeed(uint HairPrimitiveId, float2 HairPrimitiveUV)
-{
-	return GetHairStrandsAttributes(HairPrimitiveId, HairPrimitiveUV).HairSeed;
-}
-
-float3 GetHairStrandsBaseColor(uint HairPrimitiveId, float2 HairPrimitiveUV)
-{
-	return GetHairStrandsAttributes(HairPrimitiveId, HairPrimitiveUV).HairBaseColor;
-}
-
-float GetHairStrandsRoughness(uint HairPrimitiveId, float2 HairPrimitiveUV)
-{
-	return GetHairStrandsAttributes(HairPrimitiveId, HairPrimitiveUV).HairRoughness;
-}
-
 bool UseStableRasterization()
 {
 	return HairStrandsVF_bUseStableRasterization > 0;
diff --git a/Engine/Shaders/Private/HairStrands/HairStrandsVertexFactoryCommon.ush b/Engine/Shaders/Private/HairStrands/HairStrandsVertexFactoryCommon.ush
--- a/Engine/Shaders/Private/HairStrands/HairStrandsVertexFactoryCommon.ush
+++ b/Engine/Shaders/Private/HairStrands/HairStrandsVertexFactoryCommon.ush
@@ -1,9 +1,14 @@
 // Copyright Epic Games, Inc. All Rights Reserved.
 
+#pragma once
+
 #define HAIR_CONTROLPOINT_INSIDE 0
 #define HAIR_CONTROLPOINT_START	1
 #define HAIR_CONTROLPOINT_END	2
 
+////////////////////////////////////////////////////////////////////////////////////////////////////////////////////////////////
+// Hair control points
+
 struct FHairControlPoint
 {
 	float3 Position;
@@ -34,3 +39,27 @@ bool IsValidHairStrandsSegment(const FHairControlPoint CP0, const FHairControlPo
 	return !(CP0.Type == HAIR_CONTROLPOINT_END && CP1.Type == HAIR_CONTROLPOINT_START);
 }
 
+////////////////////////////////////////////////////////////////////////////////////////////////////////////////////////////////
+// Hair attributes
+
+struct FHairAttributes
+{
+	float2 RootUV;
+	float	Seed;
+	float	UCoord;
+};
+
+FHairAttributes DecodeHairAttribute(uint4 Attributes)
+{
+	float2 RootUV, RootIndex;
+	RootUV.x = float(Attributes.x & 0xFF) / 255.f;
+	RootUV.y = 1.0 - (float((Attributes.x >> 8) & 0xFF) / 255.f); // Invert V to compensate image origin flip
+	RootIndex.x = float(Attributes.z & 0xFF);
+	RootIndex.y = float((Attributes.z >> 8) & 0xFF);
+
+	FHairAttributes Out;
+	Out.RootUV = RootUV + RootIndex;
+	Out.UCoord = (Attributes.y & 0xFF) / 255.f;
+	Out.Seed = ((Attributes.y >> 8) & 0xFF) / 255.f;
+	return Out;
+}
\ No newline at end of file
diff --git a/Engine/Shaders/Private/HairStrands/HairStrandsVisibilityCommon.ush b/Engine/Shaders/Private/HairStrands/HairStrandsVisibilityCommon.ush
--- a/Engine/Shaders/Private/HairStrands/HairStrandsVisibilityCommon.ush
+++ b/Engine/Shaders/Private/HairStrands/HairStrandsVisibilityCommon.ush
@@ -317,31 +317,6 @@ struct FHairViewInfo
 	bool	bIsOrthoView;
 };
 
-////////////////////////////////////////////////////////////////////////////////////////////////////////////////////////////////
-// Hair Attributes
-
-struct FHairAttributes
-{
-	float2 RootUV;
-	float	Seed;
-	float	UCoord;
-};
-
-FHairAttributes DecodeHairAttribute(uint4 Attributes)
-{
-	float2 RootUV, RootIndex;
-	RootUV.x	= float(Attributes.x & 0xFF) / 255.f;
-	RootUV.y	= 1.0-(float((Attributes.x>>8) & 0xFF) / 255.f); // Invert V to compensate image origin flip
-	RootIndex.x	= float(Attributes.z & 0xFF);
-	RootIndex.y	= float((Attributes.z >> 8) & 0xFF);
-
-	FHairAttributes Out;
-	Out.RootUV	= RootUV + RootIndex;
-	Out.UCoord	= (Attributes.y & 0xFF) / 255.f;
-	Out.Seed	= ((Attributes.y>>8) & 0xFF) / 255.f;
-	return Out;
-}
-
 ////////////////////////////////////////////////////////////////////////////////////////////////////////////////////////////////
 // Hair transmittance/coverage
 
diff --git a/Engine/Shaders/Private/HairStrands/HairStrandsVisibilityPS.usf b/Engine/Shaders/Private/HairStrands/HairStrandsVisibilityPS.usf
--- a/Engine/Shaders/Private/HairStrands/HairStrandsVisibilityPS.usf
+++ b/Engine/Shaders/Private/HairStrands/HairStrandsVisibilityPS.usf
@@ -87,7 +87,7 @@ void MainVisibility(
 	float Coverage = 1;
 	{
 		bool bUseStableRasterization = false;
-	#if VF_STRAND_HAIR
+	#if HAIR_STRAND_MESH_FACTORY
 		bUseStableRasterization = UseStableRasterization();
 	#endif
 
@@ -107,7 +107,7 @@ void MainVisibility(
 #endif
 
 #if HAIR_RENDER_MODE == RENDER_MODE_MSAA || HAIR_RENDER_MODE == RENDER_MODE_PPLL
-  #if VF_STRAND_HAIR
+  #if HAIR_STRAND_MESH_FACTORY
 	float3 Tangent = normalize(Interpolants.TangentToWorld2.xyz);
 	const uint PrimitiveId = Interpolants.HairPrimitiveId;
 	const bool bUseScatterSceneLighting = UseScatterSceneLighting();
@@ -191,7 +191,7 @@ void MainVisibility(
 #endif
 
 #if HAIR_RENDER_MODE == RENDER_MODE_MSAA_VISIBILITY
-  #if VF_STRAND_HAIR
+  #if HAIR_STRAND_MESH_FACTORY
 	OutColor0 = PackHairVisSample(HairPrimitiveId, HairVisibilityPass_HairMaterialId);
   #else
 	OutColor0 = 0;
diff --git a/Engine/Shaders/Private/HairStrands/HairStrandsVisibilityVS.usf b/Engine/Shaders/Private/HairStrands/HairStrandsVisibilityVS.usf
--- a/Engine/Shaders/Private/HairStrands/HairStrandsVisibilityVS.usf
+++ b/Engine/Shaders/Private/HairStrands/HairStrandsVisibilityVS.usf
@@ -34,7 +34,7 @@ void Main(
 	OPTIONAL_VertexID
 	out VS_OUTPUT_TYPE Output
 )
-#if HAIR_RENDER_MODE == RENDER_MODE_MSAA_VISIBILITY && VF_STRAND_HAIR
+#if HAIR_RENDER_MODE == RENDER_MODE_MSAA_VISIBILITY && HAIR_STRAND_MESH_FACTORY
 {
 	ResolvedView = ResolveView();
 
@@ -60,7 +60,7 @@ void Main(
 	// #hair_todo: This is a hack/quick term solution, in order to generate valide light/view direction in the case of hair strand factory, as this is used for alighning the strand with the view.
 	//ResolvedView.TranslatedWorldCameraOrigin = -HairVisibilityPass.LightDirection * 1000; // Hack for strand rendering which use the camera position for deriving the view vector for forcing strand view facing
 	FVertexFactoryIntermediates VFIntermediates = GetVertexFactoryIntermediates(Input);
-#if VF_STRAND_HAIR
+#if HAIR_STRAND_MESH_FACTORY
 	// Use a raw position (rather than quad facing position) as the pre-position is not expanded into quad (for perf. reason)
 	float4 WorldPositionRaw = VertexFactoryGetWorldPositionRaw(Input, VFIntermediates);
 #else
@@ -68,7 +68,7 @@ void Main(
 #endif
 	
 	bool bUseStableRasterization = false;
-#if VF_STRAND_HAIR
+#if HAIR_STRAND_MESH_FACTORY
 	bUseStableRasterization = UseStableRasterization();
 #endif
 
@@ -90,7 +90,7 @@ void Main(
 	}
 	#endif
 
-#if VF_STRAND_HAIR
+#if HAIR_STRAND_MESH_FACTORY
 	FHairViewInfo HairViewInfo;
 	HairViewInfo.TranslatedWorldCameraOrigin = ResolvedView.TranslatedWorldCameraOrigin;
 	HairViewInfo.ViewForward = ResolvedView.ViewForward;
@@ -122,7 +122,7 @@ void Main(
 #else
 	Output.FactoryInterpolants = VertexFactoryGetInterpolantsVSToPS(Input, VFIntermediates, VertexParameters);
 
-  #ifdef VF_STRAND_HAIR
+  #ifdef HAIR_STRAND_MESH_FACTORY
 	Output.WorldStrandRadius = VFIntermediates.HairDimensions.y;
   #else
 	Output.WorldStrandRadius = 1;
diff --git a/Engine/Shaders/Private/LocalVertexFactoryCommon.ush b/Engine/Shaders/Private/LocalVertexFactoryCommon.ush
--- a/Engine/Shaders/Private/LocalVertexFactoryCommon.ush
+++ b/Engine/Shaders/Private/LocalVertexFactoryCommon.ush
@@ -36,7 +36,7 @@ struct FVertexFactoryInterpolantsVSToPS
 		nointerpolation uint LightmapDataIndex : LIGHTMAP_ID;
 	#endif
 #endif
-#if VF_STRAND_HAIR || VF_CARDS_HAIR
+#if HAIR_STRAND_MESH_FACTORY || HAIR_CARD_MESH_FACTORY
 	nointerpolation uint HairPrimitiveId	: HAIR_PRIMITIVE_ID; // Control point ID
 	float2 HairPrimitiveUV					: HAIR_PRIMITIVE_UV; // U: parameteric distance between the two surrounding control points. V: parametric distance along the width.
 #endif
diff --git a/Engine/Shaders/Private/MaterialTemplate.ush b/Engine/Shaders/Private/MaterialTemplate.ush
--- a/Engine/Shaders/Private/MaterialTemplate.ush
+++ b/Engine/Shaders/Private/MaterialTemplate.ush
@@ -16,6 +16,8 @@
 #include "/Engine/Private/PhysicsFieldSampler.ush"
 #include "/Engine/Private/SceneData.ush"
 #include "/Engine/Private/HairShadingCommon.ush"
+#include "/Engine/Private/HairStrands/HairCardsAttributeCommon.ush"
+#include "/Engine/Private/HairStrands/HairStrandsAttributeCommon.ush"
 #include "/Engine/Private/Strata/Strata.ush"
 
 #if USES_SPEEDTREE
@@ -122,6 +124,8 @@
 
 #define IS_MESHPARTICLE_FACTORY (PARTICLE_MESH_FACTORY || NIAGARA_MESH_FACTORY)
 
+#define IS_HAIR_FACTORY (HAIR_STRAND_MESH_FACTORY || HAIR_CARD_MESH_FACTORY)
+
 /**
  * Parameters used by vertex and pixel shaders to access particle properties.
  */
@@ -313,7 +317,7 @@ struct FMaterialPixelParameters
 	uint PrimitiveId;
 
 	// Actual primitive Id
-#if VF_STRAND_HAIR || VF_CARDS_HAIR
+#if IS_HAIR_FACTORY
 	uint	HairPrimitiveId;	// Control point ID
 	float2	HairPrimitiveUV;	// U: parametric distance between the two surrounding control point. V: parametric distance along hair width
 #endif
@@ -1064,23 +1068,9 @@ float3 MaterialExpressionBlackBody( float Temp )
 	return mul( XYZtoRGB, float3( X, Y, Z ) ) * pow( 0.0004 * Temp, 4 );
 }
 
-
-#if VF_STRAND_HAIR || VF_CARDS_HAIR
-float  GetHairStrandsDepth(uint HairPrimitiveId, float2 VertexUV, float InDeviceZ);
-float  GetHairStrandsCoverage(uint HairPrimitiveId, float2 VertexUV);
-float2 GetHairStrandsRootUV(uint HairPrimitiveId, float2 VertexUV);
-float2 GetHairStrandsUV(uint HairPrimitiveId, float2 VertexUV);
-float2 GetHairStrandsDimensions(uint HairPrimitiveId, float2 VertexUV);
-float  GetHairStrandsSeed(uint HairPrimitiveId, float2 VertexUV);
-float  GetHairStrandsRoughness(uint HairPrimitiveId, float2 HairPrimitiveUV);
-float3 GetHairStrandsBaseColor(uint HairPrimitiveId, float2 HairPrimitiveUV);
-float3 GetHairStrandsTangent(uint HairPrimitiveId, float2 HairPrimitiveUV, half3x3 TangentToWorld, bool bUseTangentSpace);
-float4 GetHairStrandsAuxilaryData(uint HairPrimitiveId, float2 VertexUV);
-#endif
-
 float2 MaterialExpressionGetHairRootUV(FMaterialPixelParameters Parameters)
 {
-#if VF_STRAND_HAIR || VF_CARDS_HAIR
+#if IS_HAIR_FACTORY
 	return GetHairStrandsRootUV(Parameters.HairPrimitiveId, Parameters.HairPrimitiveUV);
 #else
 	return float2(0, 0);
@@ -1089,7 +1079,7 @@ float2 MaterialExpressionGetHairRootUV(FMaterialPixelParameters Parameters)
 
 float2 MaterialExpressionGetHairUV(FMaterialPixelParameters Parameters)
 {	
-#if VF_STRAND_HAIR || VF_CARDS_HAIR
+#if IS_HAIR_FACTORY
 	return GetHairStrandsUV(Parameters.HairPrimitiveId, Parameters.HairPrimitiveUV);
 #else
 	return float2(0,0);
@@ -1098,7 +1088,7 @@ float2 MaterialExpressionGetHairUV(FMaterialPixelParameters Parameters)
 
 float2 MaterialExpressionGetHairDimensions(FMaterialPixelParameters Parameters)
 {
-#if VF_STRAND_HAIR || VF_CARDS_HAIR
+#if IS_HAIR_FACTORY
 	return GetHairStrandsDimensions(Parameters.HairPrimitiveId, Parameters.HairPrimitiveUV);
 #else
 	return float2(0, 0);
@@ -1107,7 +1097,7 @@ float2 MaterialExpressionGetHairDimensions(FMaterialPixelParameters Parameters)
 
 float MaterialExpressionGetHairSeed(FMaterialPixelParameters Parameters)
 {
-#if VF_STRAND_HAIR || VF_CARDS_HAIR
+#if IS_HAIR_FACTORY
 	return GetHairStrandsSeed(Parameters.HairPrimitiveId, Parameters.HairPrimitiveUV);
 #else
 	return 0;
@@ -1116,7 +1106,7 @@ float MaterialExpressionGetHairSeed(FMaterialPixelParameters Parameters)
 
 float3 MaterialExpressionGetHairBaseColor(FMaterialPixelParameters Parameters)
 {
-#if VF_STRAND_HAIR || VF_CARDS_HAIR
+#if IS_HAIR_FACTORY
 	return GetHairStrandsBaseColor(Parameters.HairPrimitiveId, Parameters.HairPrimitiveUV);
 #else
 	return float3(0,0,0);
@@ -1125,7 +1115,7 @@ float3 MaterialExpressionGetHairBaseColor(FMaterialPixelParameters Parameters)
 
 float MaterialExpressionGetHairRoughness(FMaterialPixelParameters Parameters)
 {
-#if VF_STRAND_HAIR || VF_CARDS_HAIR
+#if IS_HAIR_FACTORY
 	return GetHairStrandsRoughness(Parameters.HairPrimitiveId, Parameters.HairPrimitiveUV);
 #else
 	return 0;
@@ -1139,16 +1129,16 @@ float MaterialExpressionGetHairDepth(FMaterialVertexParameters Parameters)
 
 float MaterialExpressionGetHairDepth(FMaterialPixelParameters Parameters)
 {
-#if VF_CARDS_HAIR
+#if IS_HAIR_FACTORY
 	return GetHairStrandsDepth(Parameters.HairPrimitiveId, Parameters.HairPrimitiveUV, Parameters.SvPosition.z);
 #else
-	return 0; // TODO for VF_STRAND_HAIR?
+	return 0; 
 #endif
 }
 
 float MaterialExpressionGetHairCoverage(FMaterialPixelParameters Parameters)
 {
-#if VF_CARDS_HAIR || VF_STRAND_HAIR
+#if IS_HAIR_FACTORY
 	return GetHairStrandsCoverage(Parameters.HairPrimitiveId, Parameters.HairPrimitiveUV);
 #else
 	return 0;
@@ -1157,21 +1147,20 @@ float MaterialExpressionGetHairCoverage(FMaterialPixelParameters Parameters)
 
 float3 MaterialExpressionGetHairTangent(FMaterialPixelParameters Parameters, bool bUseTangentSpace)
 {
-#if VF_STRAND_HAIR
+#if HAIR_STRAND_MESH_FACTORY
 	return bUseTangentSpace ? float3(0,1,0) : Parameters.TangentToWorld[2];
-#elif VF_CARDS_HAIR
+#elif HAIR_CARD_MESH_FACTORY
 	return GetHairStrandsTangent(Parameters.HairPrimitiveId, Parameters.HairPrimitiveUV, Parameters.TangentToWorld, bUseTangentSpace);
 #else
 	return 0;
-#endif
-	
+#endif	
 }
 
 float2 MaterialExpressionGetAtlasUVs(FMaterialPixelParameters Parameters)
 {
-#if VF_STRAND_HAIR
+#if HAIR_STRAND_MESH_FACTORY
 	return float2(0,0);
-#elif VF_CARDS_HAIR
+#elif HAIR_CARD_MESH_FACTORY
 	return Parameters.HairPrimitiveUV;
 #else
 	return 0;
@@ -1181,7 +1170,7 @@ float2 MaterialExpressionGetAtlasUVs(FMaterialPixelParameters Parameters)
 
 float4 MaterialExpressionGetHairAuxilaryData(FMaterialPixelParameters Parameters)
 {
-#if VF_CARDS_HAIR
+#if HAIR_CARD_MESH_FACTORY
 	return GetHairStrandsAuxilaryData(Parameters.HairPrimitiveId, Parameters.HairPrimitiveUV);
 #else
 	return 0;
diff --git a/Engine/Source/Runtime/Renderer/Private/HairStrands/HairStrandsTransmittance.cpp b/Engine/Source/Runtime/Renderer/Private/HairStrands/HairStrandsTransmittance.cpp
--- a/Engine/Source/Runtime/Renderer/Private/HairStrands/HairStrandsTransmittance.cpp
+++ b/Engine/Source/Runtime/Renderer/Private/HairStrands/HairStrandsTransmittance.cpp
@@ -17,6 +17,7 @@
 #include "GpuDebugRendering.h"
 #include "ShaderPrintParameters.h"
 #include "LightSceneInfo.h"
+#include "ShaderPrint.h"
 
 ///////////////////////////////////////////////////////////////////////////////////////////////////
 
@@ -106,7 +107,6 @@ class FDeepTransmittanceMaskCS : public FGlobalShader
 	BEGIN_SHADER_PARAMETER_STRUCT(FParameters, )
 		SHADER_PARAMETER_STRUCT_INCLUDE(FSceneTextureParameters, SceneTextures)
 		SHADER_PARAMETER_STRUCT_INCLUDE(ShaderDrawDebug::FShaderDrawDebugParameters, ShaderDrawParameters)
-		SHADER_PARAMETER_STRUCT_INCLUDE(ShaderPrint::FShaderParameters, ShaderPrintParameters)
 
 		SHADER_PARAMETER_ARRAY(FIntVector4, DeepShadow_AtlasSlotOffsets_AtlasSlotIndex, [FHairStrandsDeepShadowData::MaxMacroGroupCount])
 		SHADER_PARAMETER_ARRAY(FMatrix, DeepShadow_CPUWorldToLightTransforms, [FHairStrandsDeepShadowData::MaxMacroGroupCount])
@@ -240,7 +240,6 @@ static FRDGBufferRef AddDeepShadowTransmittanceMaskPass(
 	if (ShaderDrawDebug::IsShaderDrawDebugEnabled(View))
 	{
 		ShaderDrawDebug::SetParameters(GraphBuilder, View.ShaderDrawData, Parameters->ShaderDrawParameters);
-		ShaderPrint::SetParameters(GraphBuilder, View, Parameters->ShaderPrintParameters);
 	}
 
 	memcpy(&(Parameters->DeepShadow_AtlasSlotOffsets_AtlasSlotIndex[0]), Params.DeepShadow_AtlasSlotOffsets_AtlasSlotIndex, sizeof(FIntVector4) * FHairStrandsDeepShadowData::MaxMacroGroupCount);
@@ -301,7 +300,6 @@ class FDeepShadowMaskPS : public FGlobalShader
 
 	BEGIN_SHADER_PARAMETER_STRUCT(FParameters, )
 		SHADER_PARAMETER_STRUCT_INCLUDE(ShaderDrawDebug::FShaderDrawDebugParameters, ShaderDrawParameters)
-		SHADER_PARAMETER_STRUCT_INCLUDE(ShaderPrint::FShaderParameters, ShaderPrintParameters)
 		                         
 		SHADER_PARAMETER(FIntPoint, DeepShadow_SlotOffset)
 		SHADER_PARAMETER(uint32, DeepShadow_SlotIndex)
@@ -401,7 +399,6 @@ static void AddDeepShadowOpaqueMaskPass(
 	if (ShaderDrawDebug::IsShaderDrawDebugEnabled(View))
 	{
 		ShaderDrawDebug::SetParameters(GraphBuilder, View.ShaderDrawData, Parameters->ShaderDrawParameters);
-		ShaderPrint::SetParameters(GraphBuilder, View, Parameters->ShaderPrintParameters);
 	}
 
 	FRDGTextureRef RayMarchMask = nullptr;
