Optimised Conductor BSDF size from 24bytes to 16bytes, reducing badnwitdh and strata worst case.

diff --git a/Engine/Shaders/Private/Strata/Strata.ush b/Engine/Shaders/Private/Strata/Strata.ush
--- a/Engine/Shaders/Private/Strata/Strata.ush
+++ b/Engine/Shaders/Private/Strata/Strata.ush
@@ -62,10 +62,10 @@
 
 // The largest Strata BSDF is FStrataConductor having as size of 24 bytes (FStrataConductor)
 // So we will have 
-//		STRATA_MAX_TOTAL_BSDF * sizeof(FStrataConductor) +				// 4 * 24
+//		STRATA_MAX_TOTAL_BSDF * sizeof(FStrataConductor) +				// 4 * 16
 //		STRATA_MAX_TOTAL_BSDF * (sizeof(BSDFType) + sizeof(Weight))		// 4 * 8
 //		sizeof(BSDFCount)												// UINT
-//		= 4 * 24 + 4 * 8 + 1 = 96 + 32 + 1 = 129						// See Strata.cpp
+//		= 4 * 16 + 4 * 8 + 1 = 96 + 32 + 1 = 97						// See Strata.cpp
 struct FStrataBSDF
 {
 	int					BSDFType;
@@ -751,6 +751,25 @@ float3 UnpackAlbedoGamma2ToLinear(uint rgb)
 	return Gamma2 * Gamma2;	// Simple gamma 2.0
 }
 
+uint PackAlbedoLinearToGamma2AlphaLinear(float4 rgba)
+{
+	uint r = (uint(sqrt(saturate(rgba.r)) * 255.0f) << 24);
+	uint g = (uint(sqrt(saturate(rgba.g)) * 255.0f) << 16);
+	uint b = (uint(sqrt(saturate(rgba.b)) * 255.0f) << 8);
+	uint a = (uint(     saturate(rgba.a)  * 255.0f) << 0);
+	return r | g | b | a;
+}
+
+float4 UnpackAlbedoGamma2ToLinearAlphaLinear(uint rgba)
+{
+	float r = float((rgba >> 24) & 0x000000FF);
+	float g = float((rgba >> 16) & 0x000000FF);
+	float b = float((rgba >> 8)  & 0x000000FF);
+	float a = float((rgba >> 0)  & 0x000000FF);
+	float3 Gamma2 = float3(r, g, b) * (1.0f / 255.0f);
+	return float4(Gamma2 * Gamma2, a * (1.0f / 255.0f));	// Simple gamma 2.0
+}
+
 #if COMPILER_SWITCH // These platforms are complaining so simply disabling Strata there for now.
 
 #define STRATA_STORE_UINT1(x) 
@@ -829,6 +848,13 @@ void StrataLoad_AlbedoGamma2ToLinear(ByteAddressBuffer StrataBuffer, inout uint
 	A = UnpackAlbedoGamma2ToLinear(StrataLoadUint1(StrataBuffer, ByteOffset));
 }
 
+void StrataLoad_AlbedoGamma2ToLinearAlphaLinear(ByteAddressBuffer StrataBuffer, inout uint ByteOffset, inout float3 RGB, inout float A)
+{
+	float4 Data = UnpackAlbedoGamma2ToLinearAlphaLinear(StrataLoadUint1(StrataBuffer, ByteOffset));
+	RGB = Data.rgb;
+	A = Data.a;
+}
+
 
 
 uint GetStrataPixelDataByteOffset(uint2 PixelPos, uint ViewBufferSizeX, uint StrataMaxBytesPerPixel)
@@ -986,10 +1012,10 @@ void PackStrataOut(
 			break;
 			case STRATA_BSDF_TYPE_CONDUCTOR:
 			{
-				STRATA_STORE_UINT2(PackR16G16B16A16F(float4(CONDUCTOR_REFLECTIVITY(BSDF), CONDUCTOR_ROUGHNESS(BSDF).x)));
-				STRATA_STORE_UINT2(PackR16G16B16A16F(float4(CONDUCTOR_EDGECOLOR(BSDF), CONDUCTOR_ROUGHNESS(BSDF).y)));
+				STRATA_STORE_UINT1(PackAlbedoLinearToGamma2AlphaLinear(float4(CONDUCTOR_REFLECTIVITY(BSDF), CONDUCTOR_ROUGHNESS(BSDF).x)));
+				STRATA_STORE_UINT1(PackAlbedoLinearToGamma2AlphaLinear(float4(CONDUCTOR_EDGECOLOR(BSDF),    CONDUCTOR_ROUGHNESS(BSDF).y)));
 				STRATA_STORE_UINT2(PackR16G16B16A16F(float4(CONDUCTOR_NORMAL(BSDF), 0.0f)));
-				// 24 bytes
+				// 16 bytes
 			}
 			break;
 			case STRATA_BSDF_TYPE_VOLUME:
@@ -1057,9 +1083,9 @@ FStrataPixelBSDFs UnpackStrataIn(ByteAddressBuffer StrataBuffer, uint ByteOffset
 		break;
 		case STRATA_BSDF_TYPE_CONDUCTOR:
 		{
-			StrataLoad_R16G16B16F_A16F		(StrataBuffer, ByteOffset, CONDUCTOR_REFLECTIVITY(PixelBSDFs.BSDFs[i]), CONDUCTOR_ROUGHNESS(PixelBSDFs.BSDFs[i]).x);
-			StrataLoad_R16G16B16F_A16F		(StrataBuffer, ByteOffset, CONDUCTOR_EDGECOLOR(PixelBSDFs.BSDFs[i]),	CONDUCTOR_ROUGHNESS(PixelBSDFs.BSDFs[i]).y);
-			StrataLoad_R16G16B16F_A16F		(StrataBuffer, ByteOffset, CONDUCTOR_NORMAL(PixelBSDFs.BSDFs[i]),		DummyFloat);
+			StrataLoad_AlbedoGamma2ToLinearAlphaLinear(StrataBuffer, ByteOffset, CONDUCTOR_REFLECTIVITY(PixelBSDFs.BSDFs[i]),	CONDUCTOR_ROUGHNESS(PixelBSDFs.BSDFs[i]).x);
+			StrataLoad_AlbedoGamma2ToLinearAlphaLinear(StrataBuffer, ByteOffset, CONDUCTOR_EDGECOLOR(PixelBSDFs.BSDFs[i]),		CONDUCTOR_ROUGHNESS(PixelBSDFs.BSDFs[i]).y);
+			StrataLoad_R16G16B16F_A16F		(StrataBuffer, ByteOffset,			 CONDUCTOR_NORMAL(PixelBSDFs.BSDFs[i]),			DummyFloat);
 		}
 		break;
 		case STRATA_BSDF_TYPE_VOLUME:
diff --git a/Engine/Source/Runtime/Renderer/Private/Strata/Strata.cpp b/Engine/Source/Runtime/Renderer/Private/Strata/Strata.cpp
--- a/Engine/Source/Runtime/Renderer/Private/Strata/Strata.cpp
+++ b/Engine/Source/Runtime/Renderer/Private/Strata/Strata.cpp
@@ -10,7 +10,7 @@
 #include "UniformBuffer.h"
 #include "SceneTextureParameters.h"
 
-
+//PRAGMA_DISABLE_OPTIMIZATION
 
 // The project setting for Strata
 static TAutoConsoleVariable<int32> CVarStrata(
@@ -50,8 +50,9 @@ void InitialiseStrataFrameSceneData(FSceneRenderer& SceneRenderer, FRDGBuilder&
 
 		// Previous GBuffer when complete was 28bytes
 		// check out Strata.ush to see how this is computed
-		const uint32 MaterialConservativeByteCountPerPixel = 129u;
-		StrataSceneData.MaxBytesPerPixel = FMath::DivideAndRoundUp(MaterialConservativeByteCountPerPixel, 4u) * 4u;
+		const uint32 MaterialConservativeByteCountPerPixel = 97u;
+		const uint32 RoundToValue = 4u;
+		StrataSceneData.MaxBytesPerPixel = FMath::DivideAndRoundUp(MaterialConservativeByteCountPerPixel, RoundToValue) * RoundToValue;
 	}
 	else
 	{
