Optional progress bar display for path tracer

Add a Cvar (off by default) that will overlay a progress bar on the image showing the progress toward the desired sample per pixel count. The progress bar dissapears once the target sample count is reached.

diff --git a/Engine/Shaders/Private/PathTracing/PathTracingCompositingPixelShader.usf b/Engine/Shaders/Private/PathTracing/PathTracingCompositingPixelShader.usf
--- a/Engine/Shaders/Private/PathTracing/PathTracingCompositingPixelShader.usf
+++ b/Engine/Shaders/Private/PathTracing/PathTracingCompositingPixelShader.usf
@@ -9,6 +9,7 @@ ReferencePathTracinPixelShader.ush: Reference path tracing  compositing pixel sh
 
 uint Iteration;
 uint MaxSamples;
+int ProgressDisplayEnabled;
 Texture2D<float4> RadianceTexture;
 
 
@@ -47,23 +48,55 @@ void CompositeMain(
 
 	OutColor = float4(Radiance.rgb * View.PreExposure, 1.0);
 
-#if 0
-	// TODO: make a configureable option for this
-	if (Iteration < MaxSamples)
+	if (ProgressDisplayEnabled && Iteration < MaxSamples)
 	{
-		// Draw a rectangle of how far along we are in accumulating samples
-		float Border = 0.05;
-		if (UV.y > 1 - 3 * Border && UV.y < 1 - Border)
+		// Draw a progress meter for how far along we are to the max number of passes
+		float Aspect = View.BufferSizeAndInvSize.y * View.BufferSizeAndInvSize.z;
+		float Size = 50.0 * View.BufferSizeAndInvSize.zz;
+		float2 P = float2(UV.x, UV.y * Aspect);
+		float2 RectMin = View.ViewRectMin.xy;
+		float2 RectMax = View.ViewRectMin.xy + View.ViewSizeAndInvSize.xy;
+
+		// center progress in middle-bottom of the frame
+		float2 Center = float2(lerp(RectMin.x, RectMax.x, 0.5), RectMax.y) * View.BufferSizeAndInvSize.zz - float2(0, Size * 0.75);
+		P -= Center;
+#if 0
+		// measure progress like a clock
+		float R2 = dot(P, P);
+		float ROuter2 = Square(0.50 * Size);
+		float RInner2 = Square(0.35 * Size);
+
+		float R2f = fwidth(R2);
+
+		// point is on the progress ring
+		if (R2 < ROuter2 + R2f && R2 > RInner2 - R2f)
 		{
 			float Fraction = saturate(float(Iteration) / float(MaxSamples));
-			if (UV.x > Border && (UV.x - Border) / (1 - 2 * Border) < Fraction)
-			{
-				float4 ProgressBarColor = float4(0.9, 0.1, 0.1, 1);
-				OutColor = lerp(OutColor, ProgressBarColor, 0.5);
-			}
+			float Angle = (atan2(-P.x, P.y) + PI) / (2 * PI);
+			// anti-alias progress bar edge
+			float Blend = smoothstep(Fraction - fwidth(Angle), Fraction + fwidth(Angle), Angle);
+			// Grey/Red blend
+			float4 ProgressBarColor = lerp(float4(1.0, 0.1, 0.1, 1), float4(0.5, 0.5, 0.5, 1), Blend);
+			// anti-alias ring edges
+			float BlendOuter = 1 - smoothstep(ROuter2 - R2f, ROuter2 + R2f, R2);
+			float BlendInner =     smoothstep(RInner2 - R2f, RInner2 + R2f, R2);
+			OutColor = lerp(OutColor, ProgressBarColor, 0.75 * BlendInner * BlendOuter);
+		}
+#else
+		// simple linear progress bar
+		float2 ProgressSize = float2(2.0 * Size, 0.15 * Size);
+		if (all(abs(P) < ProgressSize))
+		{
+			P += ProgressSize;
+			P /= ProgressSize * 2;
+			float Fraction = saturate(float(Iteration) / float(MaxSamples));
+			float Blend = smoothstep(Fraction - fwidth(P.x), Fraction + fwidth(P.x), P.x);
+			// Grey/Red blend
+			float4 ProgressBarColor = lerp(float4(1.0, 0.1, 0.1, 1), float4(0.5, 0.5, 0.5, 1), Blend);
+			OutColor = lerp(OutColor, ProgressBarColor, 0.75);
 		}
-	}
 #endif
+	}
 
 #if 0
     // split-screen debugging aids
diff --git a/Engine/Source/Runtime/Renderer/Private/PathTracing.cpp b/Engine/Source/Runtime/Renderer/Private/PathTracing.cpp
--- a/Engine/Source/Runtime/Renderer/Private/PathTracing.cpp
+++ b/Engine/Source/Runtime/Renderer/Private/PathTracing.cpp
@@ -79,6 +79,15 @@ TAutoConsoleVariable<int32> CVarPathTracingWiperMode(
 	ECVF_RenderThreadSafe 
 );
 
+TAutoConsoleVariable<int32> CVarPathTracingProgressDisplay(
+	TEXT("r.PathTracing.ProgressDisplay"),
+	0,
+	TEXT("Enables an in-frame display of progress towards the defined sample per pixel limit. The indicator dissapears when the maximum is reached and sample accumulation has stopped (default = 0)\n")
+	TEXT("0: off (default)\n")
+	TEXT("1: on\n"),
+	ECVF_RenderThreadSafe
+);
+
 IMPLEMENT_GLOBAL_SHADER_PARAMETER_STRUCT(FPathTracingData, "PathTracingData");
 IMPLEMENT_GLOBAL_SHADER_PARAMETER_STRUCT(FPathTracingLightData, "SceneLightsData");
 
@@ -259,6 +268,7 @@ class FPathTracingCompositorPS : public FGlobalShader
 		SHADER_PARAMETER_STRUCT_REF(FViewUniformShaderParameters, ViewUniformBuffer)
 		SHADER_PARAMETER(unsigned, Iteration)
 		SHADER_PARAMETER(unsigned, MaxSamples)
+		SHADER_PARAMETER(int, ProgressDisplayEnabled)
 
 		RENDER_TARGET_BINDING_SLOTS()
 	END_SHADER_PARAMETER_STRUCT()
@@ -411,6 +421,7 @@ void FDeferredShadingSceneRenderer::RenderPathTracing(
 	FPathTracingCompositorPS::FParameters* DisplayParameters = GraphBuilder.AllocParameters<FPathTracingCompositorPS::FParameters>();
 	DisplayParameters->Iteration = PathTracingData.Iteration;
 	DisplayParameters->MaxSamples = MaxSPP;
+	DisplayParameters->ProgressDisplayEnabled = CVarPathTracingProgressDisplay.GetValueOnRenderThread();
 	DisplayParameters->ViewUniformBuffer = View.ViewUniformBuffer;
 	DisplayParameters->RadianceTexture = GraphBuilder.CreateSRV(FRDGTextureSRVDesc::Create(RadianceTexture));
 	DisplayParameters->RenderTargets[0] = FRenderTargetBinding(SceneColorOutputTexture, ERenderTargetLoadAction::ENoAction);
