StrataNormalForDiffuseTransfert now outputs a SH3

diff --git a/Engine/Shaders/Private/Strata/Strata.ush b/Engine/Shaders/Private/Strata/Strata.ush
--- a/Engine/Shaders/Private/Strata/Strata.ush
+++ b/Engine/Shaders/Private/Strata/Strata.ush
@@ -7,6 +7,7 @@
 #include "../ShadingCommon.ush"
 #include "../ParticipatingMediaCommon.ush"
 #include "../MonteCarlo.ush"
+#include "../SHCommon.ush"
 
 #include "/Engine/Shared/StrataDefinitions.h"
 
@@ -560,10 +561,10 @@ FStrataEnvLightResult StrataEvaluateForEnvLight(FStrataBSDFContext BSDFContext,
 	return StrataEnvLightResult;
 }
 
-// Always return a normal for diffuse transfer as needed by Lumen
-float3 StrataNormalForDiffuseTransfert(FStrataBSDFContext BSDFContext)
+FThreeBandSHVector StrataBSDFToSH(FStrataBSDFContext BSDFContext)
 {
-	return BSDFContext.N;
+	// STRATA_TODO adapt the SH to BSDFs
+	return CalcDiffuseTransferSH3(BSDFContext.N, 1.0f);
 }
 
 
