Fixed negative numbers in Strata diffuse BSDFs.

diff --git a/Engine/Shaders/Private/Strata/Strata.ush b/Engine/Shaders/Private/Strata/Strata.ush
--- a/Engine/Shaders/Private/Strata/Strata.ush
+++ b/Engine/Shaders/Private/Strata/Strata.ush
@@ -325,7 +325,7 @@ FStrataEvaluateResult StrataEvaluateBSDF(FStrataBSDFContext BSDFContext)
 		{
 			const float SafeRoughness = MakeRoughnessSafe(DIFFUSE_ROUGHNESS(BSDFContext.BSDF));
 
-			Sample.Value		= Diffuse_OrenNayar(DIFFUSE_ALBEDO(BSDFContext.BSDF), SafeRoughness, BSDFContext.Context.NoV, BSDFContext.SatNoL, BSDFContext.Context.VoH);
+			Sample.Value		= Diffuse_OrenNayar(DIFFUSE_ALBEDO(BSDFContext.BSDF), SafeRoughness, BSDFContext.SatNoV, BSDFContext.SatNoL, BSDFContext.Context.VoH);
 			Sample.PDF			= BSDFContext.SatNoV / PI;
 			Sample.Throughput	= OpaqueBSDFThroughput;
 			break;
@@ -336,7 +336,7 @@ FStrataEvaluateResult StrataEvaluateBSDF(FStrataBSDFContext BSDFContext)
 			float a = SafeRoughness * SafeRoughness;
 			float a2 = a * a;
 
-			Sample.Value		= Diffuse_Chan(DIFFUSE_ALBEDO(BSDFContext.BSDF), a2, BSDFContext.Context.NoV, BSDFContext.SatNoL, BSDFContext.Context.VoH, BSDFContext.Context.NoH);
+			Sample.Value		= Diffuse_Chan(DIFFUSE_ALBEDO(BSDFContext.BSDF), a2, BSDFContext.SatNoV, BSDFContext.SatNoL, BSDFContext.Context.VoH, BSDFContext.Context.NoH);
 			Sample.PDF			= BSDFContext.SatNoV / PI;
 			Sample.Throughput	= OpaqueBSDFThroughput;
 			break;
