Reduce VGPR for strat lighting.

diff --git a/Engine/Shaders/Private/Strata/Strata.ush b/Engine/Shaders/Private/Strata/Strata.ush
--- a/Engine/Shaders/Private/Strata/Strata.ush
+++ b/Engine/Shaders/Private/Strata/Strata.ush
@@ -82,7 +82,7 @@ struct FStrataBSDF
 	//};
 	//BSDFsUnion U;
 
-	float4 VGPRs[4];
+	float4 VGPRs[3];
 };
 
 #define DIFFUSE_ALBEDO(X)				X.VGPRs[0].xyz
@@ -91,19 +91,21 @@ struct FStrataBSDF
 
 #define DIELECTRIC_TINT(X)				X.VGPRs[0].xyz
 #define DIELECTRIC_IOR(X)				X.VGPRs[0].w
-#define DIELECTRIC_ROUGHNESS(X)			X.VGPRs[1].xy
-#define DIELECTRIC_NORMAL(X)			X.VGPRs[2].xyz
+#define DIELECTRIC_ROUGHNESS0(X)		X.VGPRs[1].w
+#define DIELECTRIC_ROUGHNESS1(X)		X.VGPRs[2].w
+#define DIELECTRIC_NORMAL(X)			X.VGPRs[1].xyz
 
 #define CONDUCTOR_REFLECTIVITY(X)		X.VGPRs[0].xyz
 #define CONDUCTOR_EDGECOLOR(X)			X.VGPRs[1].xyz
-#define CONDUCTOR_ROUGHNESS(X)			X.VGPRs[2].xy
-#define CONDUCTOR_NORMAL(X)				X.VGPRs[3].xyz
+#define CONDUCTOR_ROUGHNESS0(X)			X.VGPRs[2].w
+#define CONDUCTOR_ROUGHNESS1(X)			X.VGPRs[1].w
+#define CONDUCTOR_NORMAL(X)				X.VGPRs[2].xyz
 
 #define VOLUME_ALBEDO(X)				X.VGPRs[0].xyz
 #define VOLUME_EXTINCTION(X)			X.VGPRs[1].xyz
-#define VOLUME_ANISOTROPY(X)			X.VGPRs[2].x
-#define VOLUME_THICKNESS(X)				X.VGPRs[2].y
-#define VOLUME_SMOOTHSURFACENORMAL(X)	X.VGPRs[3].xyz
+#define VOLUME_ANISOTROPY(X)			X.VGPRs[0].w
+#define VOLUME_THICKNESS(X)				X.VGPRs[1].w
+#define VOLUME_SMOOTHSURFACENORMAL(X)	X.VGPRs[2].xyz
 
 
 struct FStrataLayer
@@ -168,7 +170,8 @@ FStrataData GetStrataDielectricBSDF(float2 Roughness, float IOR, float3 Tint, fl
 	StrataData.Layers[0].BSDFs[0].BSDFType						= STRATA_BSDF_TYPE_DIELECTRIC;
 	DIELECTRIC_IOR			(StrataData.Layers[0].BSDFs[0])		= IOR;
 	DIELECTRIC_TINT			(StrataData.Layers[0].BSDFs[0])		= Tint;
-	DIELECTRIC_ROUGHNESS	(StrataData.Layers[0].BSDFs[0])		= Roughness;
+	DIELECTRIC_ROUGHNESS0	(StrataData.Layers[0].BSDFs[0])		= Roughness.x;
+	DIELECTRIC_ROUGHNESS1	(StrataData.Layers[0].BSDFs[0])		= Roughness.y;
 	DIELECTRIC_NORMAL		(StrataData.Layers[0].BSDFs[0])		= Normal;
 
 	StrataData.Layers[0].BSDFs[0].Weight						= 1.0f;
@@ -185,7 +188,8 @@ FStrataData GetStrataConductorBSDF(float3 Reflectivity, float3 EdgeColor, float2
 	StrataData.Layers[0].BSDFs[0].BSDFType						= STRATA_BSDF_TYPE_CONDUCTOR;
 	CONDUCTOR_REFLECTIVITY	(StrataData.Layers[0].BSDFs[0])		= Reflectivity;
 	CONDUCTOR_EDGECOLOR		(StrataData.Layers[0].BSDFs[0])		= EdgeColor;
-	CONDUCTOR_ROUGHNESS		(StrataData.Layers[0].BSDFs[0])		= Roughness;
+	CONDUCTOR_ROUGHNESS0	(StrataData.Layers[0].BSDFs[0])		= Roughness.x;
+	CONDUCTOR_ROUGHNESS1	(StrataData.Layers[0].BSDFs[0])		= Roughness.y;
 	CONDUCTOR_NORMAL		(StrataData.Layers[0].BSDFs[0])		= Normal;
 
 	StrataData.Layers[0].BSDFs[0].Weight						= 1.0f;
@@ -345,7 +349,7 @@ FStrataEvaluateResult StrataEvaluateBSDF(FStrataBSDFContext BSDFContext)
 		{
 			float F0 = DielectricIorToF0(DIELECTRIC_IOR(BSDFContext.BSDF));
 
-			const float SafeRoughness = MakeRoughnessSafe(DIELECTRIC_ROUGHNESS(BSDFContext.BSDF).x);
+			const float SafeRoughness = MakeRoughnessSafe(DIELECTRIC_ROUGHNESS0(BSDFContext.BSDF));
 			float a2 = Pow4(SafeRoughness);
 			float Energy = 1.0f; // See EnergyNormalization
 
@@ -362,7 +366,7 @@ FStrataEvaluateResult StrataEvaluateBSDF(FStrataBSDFContext BSDFContext)
 		case STRATA_BSDF_TYPE_CONDUCTOR:
 		{
 			// See SpecularGGX, isotropic verion
-			const float SafeRoughness = MakeRoughnessSafe(CONDUCTOR_ROUGHNESS(BSDFContext.BSDF).x);
+			const float SafeRoughness = MakeRoughnessSafe(CONDUCTOR_ROUGHNESS0(BSDFContext.BSDF));
 			float a2 = Pow4(SafeRoughness);
 			float Energy = 1.0f; // See EnergyNormalization
 
@@ -448,7 +452,7 @@ FStrataImportanceSampleResult StrataImportanceSampleBSDF(FStrataBSDF BSDF, float
 		}
 		case STRATA_BSDF_TYPE_DIELECTRIC:
 		{
-			float4 GGXSample = ImportanceSampleVisibleGGX(UniformSampleDisk(E), Pow4(DIELECTRIC_ROUGHNESS(BSDFContext.BSDF)).x, BSDFContext.TangentV);
+			float4 GGXSample = ImportanceSampleVisibleGGX(UniformSampleDisk(E), Pow4(DIELECTRIC_ROUGHNESS0(BSDFContext.BSDF)), BSDFContext.TangentV);
 			float3 WorldH = mul(GGXSample.xyz, BSDFContext.TangentBasis);
 			Sample.Direction = reflect(CameraVector, WorldH);
 			PDF = GGXSample.w;
@@ -457,7 +461,7 @@ FStrataImportanceSampleResult StrataImportanceSampleBSDF(FStrataBSDF BSDF, float
 		}
 		case STRATA_BSDF_TYPE_CONDUCTOR:
 		{
-			float4 GGXSample = ImportanceSampleVisibleGGX(UniformSampleDisk(E), Pow4(CONDUCTOR_ROUGHNESS(BSDFContext.BSDF)).x, BSDFContext.TangentV);
+			float4 GGXSample = ImportanceSampleVisibleGGX(UniformSampleDisk(E), Pow4(CONDUCTOR_ROUGHNESS0(BSDFContext.BSDF)), BSDFContext.TangentV);
 			float3 WorldH = mul(GGXSample.xyz, BSDFContext.TangentBasis);
 			Sample.Direction = reflect(CameraVector, WorldH);
 			PDF = GGXSample.w;
@@ -510,7 +514,7 @@ FStrataEnvLightResult StrataEvaluateForEnvLight(FStrataBSDFContext BSDFContext,
 			if (bEnableSpecular)
 			{
 				float F0 = DielectricIorToF0(DIELECTRIC_IOR(BSDFContext.BSDF));
-				const float SafeRoughness = MakeRoughnessSafe(DIELECTRIC_ROUGHNESS(BSDFContext.BSDF).x);
+				const float SafeRoughness = MakeRoughnessSafe(DIELECTRIC_ROUGHNESS0(BSDFContext.BSDF));
 
 				float3 EvalEnvBRDF = EnvBRDFApprox(F0, SafeRoughness, BSDFContext.SatNoV);
 
@@ -527,7 +531,7 @@ FStrataEnvLightResult StrataEvaluateForEnvLight(FStrataBSDFContext BSDFContext,
 			{
 				// See SpecularGGX, isotropic verion
 
-				const float SafeRoughness = MakeRoughnessSafe(CONDUCTOR_ROUGHNESS(BSDFContext.BSDF).x);
+				const float SafeRoughness = MakeRoughnessSafe(CONDUCTOR_ROUGHNESS0(BSDFContext.BSDF));
 				float3 EvalEnvBRDF = EnvBRDFApprox(CONDUCTOR_REFLECTIVITY(BSDFContext.BSDF), CONDUCTOR_EDGECOLOR(BSDFContext.BSDF), SafeRoughness, BSDFContext.SatNoV);
 
 				StrataEnvLightResult.SpecularDirection = StrataGetOffSpecularPeakReflectionDir(BSDFContext.N, BSDFContext.R, SafeRoughness);
@@ -653,14 +657,16 @@ void SanitizeStrataDielectric(inout FStrataBSDF BSDF)
 
 	DIELECTRIC_IOR(BSDF) = clamp(DIELECTRIC_IOR(BSDF), 1.0f, 5.0f);		// See https://en.wikipedia.org/wiki/List_of_refractive_indices
 	DIELECTRIC_TINT(BSDF) = saturate(DIELECTRIC_TINT(BSDF));
-	DIELECTRIC_ROUGHNESS(BSDF) = saturate(DIELECTRIC_ROUGHNESS(BSDF));
+	DIELECTRIC_ROUGHNESS0(BSDF) = saturate(DIELECTRIC_ROUGHNESS0(BSDF));
+	DIELECTRIC_ROUGHNESS1(BSDF) = saturate(DIELECTRIC_ROUGHNESS1(BSDF));
 }
 
 void SanitizeStrataConductor(inout FStrataBSDF BSDF)
 {
 	CONDUCTOR_REFLECTIVITY(BSDF) = saturate(CONDUCTOR_REFLECTIVITY(BSDF));
 	CONDUCTOR_EDGECOLOR(BSDF) = saturate(CONDUCTOR_EDGECOLOR(BSDF));
-	CONDUCTOR_ROUGHNESS(BSDF) = saturate(CONDUCTOR_ROUGHNESS(BSDF));
+	CONDUCTOR_ROUGHNESS0(BSDF) = saturate(CONDUCTOR_ROUGHNESS0(BSDF));
+	CONDUCTOR_ROUGHNESS1(BSDF) = saturate(CONDUCTOR_ROUGHNESS1(BSDF));
 }
 
 void SanitizeStrataVolume(inout FStrataBSDF BSDF)
@@ -829,13 +835,6 @@ void StrataLoad_R11G11B10F(ByteAddressBuffer StrataBuffer, inout uint ByteOffset
 	A = UnpackR11G11B10F(StrataLoadUint1(StrataBuffer, ByteOffset));
 }
 
-void StrataLoad_R11G11_B10F(ByteAddressBuffer StrataBuffer, inout uint ByteOffset, inout float2 A, inout float B)
-{
-	float3 temp = UnpackR11G11B10F(StrataLoadUint1(StrataBuffer, ByteOffset));
-	A = temp.xy;
-	B = temp.z;
-}
-
 void StrataLoad_R11_G11_B10F(ByteAddressBuffer StrataBuffer, inout uint ByteOffset, inout float A, inout float B, inout float C)
 {
 	float3 temp = UnpackR11G11B10F(StrataLoadUint1(StrataBuffer, ByteOffset));
@@ -1006,15 +1005,15 @@ void PackStrataOut(
 			case STRATA_BSDF_TYPE_DIELECTRIC: //
 			{
 				STRATA_STORE_UINT1(PackR11G11B10F(DIELECTRIC_TINT(BSDF)));
-				STRATA_STORE_UINT1(PackR11G11B10F(float3(DIELECTRIC_ROUGHNESS(BSDF), DIELECTRIC_IOR(BSDF))));
+				STRATA_STORE_UINT1(PackR11G11B10F(float3(DIELECTRIC_ROUGHNESS0(BSDF), DIELECTRIC_ROUGHNESS1(BSDF), DIELECTRIC_IOR(BSDF))));
 				STRATA_STORE_UINT2(PackR16G16B16A16F(float4(DIELECTRIC_NORMAL(BSDF), 0.0f)));
 				// 16 bytes
 			}
 			break;
 			case STRATA_BSDF_TYPE_CONDUCTOR:
 			{
-				STRATA_STORE_UINT1(PackAlbedoLinearToGamma2AlphaLinear(float4(CONDUCTOR_REFLECTIVITY(BSDF), CONDUCTOR_ROUGHNESS(BSDF).x)));
-				STRATA_STORE_UINT1(PackAlbedoLinearToGamma2AlphaLinear(float4(CONDUCTOR_EDGECOLOR(BSDF),    CONDUCTOR_ROUGHNESS(BSDF).y)));
+				STRATA_STORE_UINT1(PackAlbedoLinearToGamma2AlphaLinear(float4(CONDUCTOR_REFLECTIVITY(BSDF), CONDUCTOR_ROUGHNESS0(BSDF))));
+				STRATA_STORE_UINT1(PackAlbedoLinearToGamma2AlphaLinear(float4(CONDUCTOR_EDGECOLOR(BSDF),    CONDUCTOR_ROUGHNESS1(BSDF))));
 				STRATA_STORE_UINT2(PackR16G16B16A16F(float4(CONDUCTOR_NORMAL(BSDF), 0.0f)));
 				// 16 bytes
 			}
@@ -1078,14 +1077,14 @@ FStrataPixelBSDFs UnpackStrataIn(ByteAddressBuffer StrataBuffer, uint ByteOffset
 		case STRATA_BSDF_TYPE_DIELECTRIC:
 		{
 			StrataLoad_R11G11B10F			(StrataBuffer, ByteOffset, DIELECTRIC_TINT(PixelBSDFs.BSDFs[i]));
-			StrataLoad_R11G11_B10F			(StrataBuffer, ByteOffset, DIELECTRIC_ROUGHNESS(PixelBSDFs.BSDFs[i]),	DIELECTRIC_IOR(PixelBSDFs.BSDFs[i]));
+			StrataLoad_R11_G11_B10F			(StrataBuffer, ByteOffset, DIELECTRIC_ROUGHNESS0(PixelBSDFs.BSDFs[i]), DIELECTRIC_ROUGHNESS1(PixelBSDFs.BSDFs[i]), DIELECTRIC_IOR(PixelBSDFs.BSDFs[i]));
 			StrataLoad_R16G16B16F_A16F		(StrataBuffer, ByteOffset, DIELECTRIC_NORMAL(PixelBSDFs.BSDFs[i]),		DummyFloat);
 		}
 		break;
 		case STRATA_BSDF_TYPE_CONDUCTOR:
 		{
-			StrataLoad_AlbedoGamma2ToLinearAlphaLinear(StrataBuffer, ByteOffset, CONDUCTOR_REFLECTIVITY(PixelBSDFs.BSDFs[i]),	CONDUCTOR_ROUGHNESS(PixelBSDFs.BSDFs[i]).x);
-			StrataLoad_AlbedoGamma2ToLinearAlphaLinear(StrataBuffer, ByteOffset, CONDUCTOR_EDGECOLOR(PixelBSDFs.BSDFs[i]),		CONDUCTOR_ROUGHNESS(PixelBSDFs.BSDFs[i]).y);
+			StrataLoad_AlbedoGamma2ToLinearAlphaLinear(StrataBuffer, ByteOffset, CONDUCTOR_REFLECTIVITY(PixelBSDFs.BSDFs[i]),	CONDUCTOR_ROUGHNESS0(PixelBSDFs.BSDFs[i]));
+			StrataLoad_AlbedoGamma2ToLinearAlphaLinear(StrataBuffer, ByteOffset, CONDUCTOR_EDGECOLOR(PixelBSDFs.BSDFs[i]),		CONDUCTOR_ROUGHNESS1(PixelBSDFs.BSDFs[i]));
 			StrataLoad_R16G16B16F_A16F		(StrataBuffer, ByteOffset,			 CONDUCTOR_NORMAL(PixelBSDFs.BSDFs[i]),			DummyFloat);
 		}
 		break;
@@ -1102,5 +1101,3 @@ FStrataPixelBSDFs UnpackStrataIn(ByteAddressBuffer StrataBuffer, uint ByteOffset
 
 	return PixelBSDFs;
 }
-
-
diff --git a/Engine/Shaders/Private/Strata/StrataVisualize.usf b/Engine/Shaders/Private/Strata/StrataVisualize.usf
--- a/Engine/Shaders/Private/Strata/StrataVisualize.usf
+++ b/Engine/Shaders/Private/Strata/StrataVisualize.usf
@@ -149,7 +149,7 @@ void VisualizeMaterialPS(
 				HORIZONTALTAB; PRINTFLOAT3(DIELECTRIC_TINT(BSDF));
 				NEXTLINE;
 				HORIZONTALTAB; PRINTPARAMCHAR(_R_); PRINTPARAMCHAR(_O_); PRINTPARAMCHAR(_U_); PRINTPARAMCHAR(_G_); PRINTPARAMCHAR(_H_); PRINTPARAMCHAR(_SPC_);
-				HORIZONTALTAB; PRINTFLOAT1(DIELECTRIC_ROUGHNESS(BSDF).x);
+				HORIZONTALTAB; PRINTFLOAT1(DIELECTRIC_ROUGHNESS0(BSDF));
 				NEXTLINE;
 
 				break;
@@ -167,7 +167,7 @@ void VisualizeMaterialPS(
 				HORIZONTALTAB; PRINTFLOAT3(CONDUCTOR_EDGECOLOR(BSDF));
 				NEXTLINE;
 				HORIZONTALTAB; PRINTPARAMCHAR(_R_); PRINTPARAMCHAR(_O_); PRINTPARAMCHAR(_U_); PRINTPARAMCHAR(_G_); PRINTPARAMCHAR(_H_); PRINTPARAMCHAR(_SPC_);
-				HORIZONTALTAB; PRINTFLOAT1(CONDUCTOR_ROUGHNESS(BSDF).x);
+				HORIZONTALTAB; PRINTFLOAT1(CONDUCTOR_ROUGHNESS0(BSDF));
 				NEXTLINE;
 
 				break;
