Made Strata basepass work with SSR and SSAO/HBAO just because we can feed the gbuffer from starta top layer for now.
It is not 100% correct but it works well enough for now as a reference while we wait for Lumen.

diff --git a/Engine/Shaders/Private/BasePassPixelShader.usf b/Engine/Shaders/Private/BasePassPixelShader.usf
--- a/Engine/Shaders/Private/BasePassPixelShader.usf
+++ b/Engine/Shaders/Private/BasePassPixelShader.usf
@@ -1479,7 +1479,7 @@ void FPixelShaderInOut_MainPS(
 	uint PixelStrataDataByteOffset = GetStrataPixelDataByteOffset(PixelPos, ResolvedView.BufferSizeAndInvSize.x, OpaqueBasePass.Strata.MaxBytesPerPixel);
 
 	const float3 VertexNormal = MaterialParameters.TangentToWorld[2].xyz;
-	PackStrataOut(OpaqueBasePass.Strata.MaterialLobesBufferUAV, PixelStrataDataByteOffset, Strata, VertexNormal, MaterialParameters.CameraVector);
+	PackStrataOut(OpaqueBasePass.Strata.MaterialLobesBufferUAV, PixelStrataDataByteOffset, Strata, VertexNormal, MaterialParameters.CameraVector, GBuffer);
 #endif // STRATA_OPAQUE_ENABLED
 
 #if STRATA_TRANSLUCENT_ENABLED
diff --git a/Engine/Shaders/Private/ReflectionEnvironmentPixelShader.usf b/Engine/Shaders/Private/ReflectionEnvironmentPixelShader.usf
--- a/Engine/Shaders/Private/ReflectionEnvironmentPixelShader.usf
+++ b/Engine/Shaders/Private/ReflectionEnvironmentPixelShader.usf
@@ -411,6 +411,13 @@ void ReflectionEnvironmentSkyLighting(
 			V,
 			AmbientOcclusion,
 			PixelBSDFs);
+
+#if 1
+		// We use WorldNormal, Roughness and BaseColor (representing EnvBRDF) written by the strata top layer.
+		const float4 SSR = Texture2DSample(ScreenSpaceReflectionsTexture, ScreenSpaceReflectionsSampler, BufferUV);
+		const float3 EnvBRDF = GBuffer.BaseColor;
+		SkyLighting.rgb = SkyLighting.rgb * (1.0-SSR.a) + SSR.rgb * EnvBRDF;
+#endif
 	}
 	else
 #endif
@@ -446,7 +453,7 @@ void ReflectionEnvironmentSkyLighting(
 #endif // ENABLE_DYNAMIC_SKY_LIGHT
 
 	BRANCH
-	if (ShadingModelID != SHADINGMODELID_UNLIT && ShadingModelID != SHADINGMODELID_HAIR)
+	if (ShadingModelID != SHADINGMODELID_UNLIT && ShadingModelID != SHADINGMODELID_HAIR && ShadingModelID != SHADINGMODELID_STRATA)
 	{
 		OutColor.xyz += ReflectionEnvironment(GBuffer, AmbientOcclusion, BufferUV, ScreenPosition, SvPosition, BentNormal, SpecularColor);
 	}
diff --git a/Engine/Shaders/Private/Strata/Strata.ush b/Engine/Shaders/Private/Strata/Strata.ush
--- a/Engine/Shaders/Private/Strata/Strata.ush
+++ b/Engine/Shaders/Private/Strata/Strata.ush
@@ -8,6 +8,7 @@
 #include "../ParticipatingMediaCommon.ush"
 #include "../MonteCarlo.ush"
 #include "../SHCommon.ush"
+#include "../DeferredShadingCommon.ush"
 
 #include "/Engine/Shared/StrataDefinitions.h"
 
@@ -869,7 +870,8 @@ void PackStrataOut(
 	uint ByteOffset, 
 	FStrataData Strata, 
 	float3 VertexNormal,
-	float3 V)
+	float3 V,
+	inout FGBufferData GBuffer)
 {
 	// We cannot output (layers x bsdf) information because if we try to load them later, the reserved amount of VGPRS would be crazy high.
 	// So instead we simply store a list of BSDF and we allow only up to STRATA_MAX_TOTAL_BSDF of them.
@@ -881,6 +883,7 @@ void PackStrataOut(
 
 	const float OpaqueBSDFThroughput = 0.0f;
 	const float FullThroughput = 1.0f;
+	const float FullyRough = 1.0f;
 
 	uint BSDFCount = 0;
 	float3 Throughput = 1.0f;
@@ -891,6 +894,12 @@ void PackStrataOut(
 		float3 LayerThroughput = 0.0f;
 		float3 LayerCoverage = 0.0f;
 
+		float3 AccumTopNormal = 0.0f;
+		float3 AccumRoughness = 0.0f;
+		float  AccumRoughnessCount = 0.0f;
+		float3 AccumBaseColor = 0.0f;
+		const bool bTopLayer = l == 0;
+
 		for (int i = 0; i < Strata.Layers[l].BSDFCount; ++i)
 		{
 			// Sanitize BSDF in place before it is used and packed out to memory
@@ -905,6 +914,8 @@ void PackStrataOut(
 			StrataGetBSDFNormal(BSDF, N);
 			N = normalize(N);
 
+			AccumTopNormal += N;
+
 			BxDFContext ContextNoLight = (BxDFContext)0;
 			const float3 FakeL = float3(0, 0, 1);
 			Init(ContextNoLight, N, V, FakeL);
@@ -913,6 +924,12 @@ void PackStrataOut(
 			{
 			case STRATA_BSDF_TYPE_DIFFUSE_ON:
 			{
+				if (bTopLayer)
+				{
+					AccumRoughness		+= BSDFCoverage * FullyRough;
+					AccumRoughnessCount += BSDFCoverage;
+				}
+
 				LayerCoverage	+= BSDFCoverage;
 				LayerThroughput += BSDFCoverage * OpaqueBSDFThroughput;
 				break;
@@ -920,6 +937,12 @@ void PackStrataOut(
 
 			case STRATA_BSDF_TYPE_DIFFUSE_CHAN:
 			{
+				if (bTopLayer)
+				{
+					AccumRoughness		+= BSDFCoverage * FullyRough;
+					AccumRoughnessCount += BSDFCoverage;
+				}
+
 				LayerCoverage	+= BSDFCoverage;
 				LayerThroughput += BSDFCoverage * OpaqueBSDFThroughput;
 				break;
@@ -930,6 +953,15 @@ void PackStrataOut(
 				float F0 = DielectricIorToF0(DIELECTRIC_IOR(BSDF));
 				float3 F = F_Schlick(F0, ContextNoLight.VoH);
 
+				if (bTopLayer)
+				{
+					AccumRoughness		+= DIELECTRIC_ROUGHNESS0(BSDF);
+					AccumRoughnessCount += 1.0f;
+
+					const float SafeRoughness = MakeRoughnessSafe(DIELECTRIC_ROUGHNESS0(BSDF));
+					AccumBaseColor		+= BSDFCoverage * EnvBRDFApprox(F0, SafeRoughness, saturate(ContextNoLight.NoV));
+				}
+
 				LayerCoverage	+= BSDFCoverage;
 				LayerThroughput += BSDFCoverage * (1.0f - F); // STRATA_TODO use 1-FGD instead of simple 1-F, event if we ignore ray scattering/spreading
 				break;
@@ -939,6 +971,15 @@ void PackStrataOut(
 			{
 				float3 F = F_Schlick(CONDUCTOR_REFLECTIVITY(BSDF), CONDUCTOR_EDGECOLOR(BSDF), ContextNoLight.VoH);
 
+				if (bTopLayer)
+				{
+					AccumRoughness		+= CONDUCTOR_ROUGHNESS0(BSDF);
+					AccumRoughnessCount += 1.0f;
+
+					const float SafeRoughness = MakeRoughnessSafe(CONDUCTOR_ROUGHNESS0(BSDF));
+					AccumBaseColor		+= BSDFCoverage * EnvBRDFApprox(CONDUCTOR_REFLECTIVITY(BSDF), CONDUCTOR_EDGECOLOR(BSDF), SafeRoughness, saturate(ContextNoLight.NoV));;
+				}
+
 				LayerCoverage	+= BSDFCoverage;
 				LayerThroughput += BSDFCoverage * OpaqueBSDFThroughput;
 				break;
@@ -957,6 +998,12 @@ void PackStrataOut(
 
 				const float PathLength = Thickness / max(0.0001f, abs(ContextNoLight.NoV));
 				const float3 SafePathSegmentTransmittance = exp(-SafeExtinctionCoefficients * PathLength);
+				
+				if (bTopLayer)
+				{
+					AccumRoughness		+= BSDFCoverage * FullyRough;
+					AccumRoughnessCount += BSDFCoverage;
+				}
 
 				LayerCoverage	+= BSDFCoverage;
 				LayerThroughput += BSDFCoverage * SafePathSegmentTransmittance;
@@ -968,6 +1015,21 @@ void PackStrataOut(
 			Strata.Layers[l].BSDFs[i].Weight *= Throughput;
 		}
 
+		// Now modify the gbuffer according to the top layer information
+		if (bTopLayer)
+		{
+			// Normal is needed for some AO techniques
+			GBuffer.WorldNormal = normalize(AccumTopNormal);
+
+			// Roughness is needed for SSR, base color represents EnvBRDF
+			GBuffer.Roughness = AccumRoughnessCount > 0.0f ? saturate(AccumRoughness / AccumRoughnessCount) : FullyRough;
+			GBuffer.BaseColor = saturate(AccumBaseColor);
+			
+			// Approximative way to be able to have SSR produce colored specular reflections
+			GBuffer.Specular = 1;
+			GBuffer.Metallic = 1;
+		}
+
 		// Take into account the non covered part for this layer and update throughput accordingly
 		LayerThroughput += (1.0 - saturate(LayerCoverage)) * FullThroughput;
 
