Change strata BSDF to be loaded on demand.

Split data load into header + data.

diff --git a/Engine/Shaders/Private/DeferredLightPixelShaders.usf b/Engine/Shaders/Private/DeferredLightPixelShaders.usf
--- a/Engine/Shaders/Private/DeferredLightPixelShaders.usf
+++ b/Engine/Shaders/Private/DeferredLightPixelShaders.usf
@@ -209,8 +209,6 @@ void DeferredLightPixelMain(
 			if (LightMask > 0)
 			{
 				uint2 PixelPos = InputParams.PixelPos;
-				uint PixelStrataDataByteOffset = GetStrataPixelDataByteOffset(PixelPos, View.BufferSizeAndInvSize.x, Strata.MaxBytesPerPixel);
-				FStrataPixelBSDFs PixelBSDFs = UnpackStrataIn(Strata.MaterialLobesBuffer, PixelStrataDataByteOffset);
 
 				FShadowTerms ShadowTerms = { ScreenSpaceData.AmbientOcclusion, 1.0, 1.0, InitHairTransmittanceData() };
 				float4 LightAttenuation = GetPerPixelLightAttenuation(InputParams.ScreenUV);
@@ -226,7 +224,8 @@ void DeferredLightPixelMain(
 					LightMask,
 					SurfaceShadow,
 					RectTexture,
-					PixelBSDFs);
+					Strata.MaterialLobesBuffer,
+					GetStrataPixelDataByteOffset(PixelPos, View.BufferSizeAndInvSize.x, Strata.MaxBytesPerPixel));
 			}
 		}
 		else
diff --git a/Engine/Shaders/Private/ReflectionEnvironmentPixelShader.usf b/Engine/Shaders/Private/ReflectionEnvironmentPixelShader.usf
--- a/Engine/Shaders/Private/ReflectionEnvironmentPixelShader.usf
+++ b/Engine/Shaders/Private/ReflectionEnvironmentPixelShader.usf
@@ -396,9 +396,6 @@ void ReflectionEnvironmentSkyLighting(
 #if ENABLE_SKY_LIGHT
 	if (bStrataMaterial)
 	{
-		uint PixelStrataDataByteOffset = GetStrataPixelDataByteOffset(PixelPos, View.BufferSizeAndInvSize.x, Strata.MaxBytesPerPixel);
-		FStrataPixelBSDFs PixelBSDFs = UnpackStrataIn(Strata.MaterialLobesBuffer, PixelStrataDataByteOffset);
-
 		float3 WorldPosition = mul(float4(ScreenPosition * GBuffer.Depth, GBuffer.Depth, 1), View.ScreenToWorld).xyz;
 		float3 CameraToPixel = normalize(WorldPosition - View.WorldCameraOrigin);
 		float3 V = -CameraToPixel;
@@ -410,7 +407,8 @@ void ReflectionEnvironmentSkyLighting(
 			CameraToPixel,
 			V,
 			AmbientOcclusion,
-			PixelBSDFs);
+			Strata.MaterialLobesBuffer,
+			GetStrataPixelDataByteOffset(PixelPos, View.BufferSizeAndInvSize.x, Strata.MaxBytesPerPixel));
 
 #if 1
 		// We use WorldNormal, Roughness and BaseColor (representing EnvBRDF) written by the strata top layer.
diff --git a/Engine/Shaders/Private/Strata/Strata.ush b/Engine/Shaders/Private/Strata/Strata.ush
--- a/Engine/Shaders/Private/Strata/Strata.ush
+++ b/Engine/Shaders/Private/Strata/Strata.ush
@@ -1103,63 +1103,67 @@ void PackStrataOut(
 	}
 }
 
-
-struct FStrataPixelBSDFs
+struct FStrataPixelHeader
 {
-	int BSDFCount;
-	FStrataBSDF BSDFs[STRATA_MAX_TOTAL_BSDF];
+	uint BSDFCount;
 };
 
-FStrataPixelBSDFs UnpackStrataIn(ByteAddressBuffer StrataBuffer, uint ByteOffset)
+// Unpack only strata header
+FStrataPixelHeader UnpackStrataHeaderIn(ByteAddressBuffer StrataBuffer, inout uint ByteOffset)
 {
-	FStrataPixelBSDFs PixelBSDFs = (FStrataPixelBSDFs)0;
+	FStrataPixelHeader Out = (FStrataPixelHeader)0;
 
-	STRATA_LOAD_UINT1(PixelBSDFs.BSDFCount);
-	PixelBSDFs.BSDFCount = min(PixelBSDFs.BSDFCount, STRATA_MAX_TOTAL_BSDF);	// To make the compiler happy
+	STRATA_LOAD_UINT1(Out.BSDFCount);
+	Out.BSDFCount = min(Out.BSDFCount, STRATA_MAX_TOTAL_BSDF);	// To make the compiler happy
+	return Out;
+}
+
+// Unpack a single BSDF
+// Note: All BSDF lobes needs to be unpack in a sequential manner since each BSDF has a variable footprint.
+FStrataBSDF UnpackStrataBSDFIn(ByteAddressBuffer StrataBuffer, inout uint ByteOffset)
+{
+	FStrataBSDF OutBSDF = (FStrataBSDF)0;
 
 	float DummyFloat = 0.0f;
 
-	for (int i = 0; i < PixelBSDFs.BSDFCount; ++i)
 	{
-		{
-			float4 temp4  = UnpackR16G16B16A16F(StrataLoadUint2(StrataBuffer, ByteOffset));
-			PixelBSDFs.BSDFs[i].Weight   = temp4.xyz;
-			PixelBSDFs.BSDFs[i].BSDFType = int(temp4.a);
-		}
+		float4 temp4  = UnpackR16G16B16A16F(StrataLoadUint2(StrataBuffer, ByteOffset));
+		OutBSDF.Weight   = temp4.xyz;
+		OutBSDF.BSDFType = int(temp4.a);
+	}
 
-		switch (PixelBSDFs.BSDFs[i].BSDFType)
-		{
-		case STRATA_BSDF_TYPE_DIFFUSE_ON:
-		case STRATA_BSDF_TYPE_DIFFUSE_CHAN:
-		{
-			StrataLoad_AlbedoGamma2ToLinear	(StrataBuffer, ByteOffset, DIFFUSE_ALBEDO(PixelBSDFs.BSDFs[i]));
-			StrataLoad_R16G16B16F_A16F		(StrataBuffer, ByteOffset, DIFFUSE_NORMAL(PixelBSDFs.BSDFs[i]),			DIFFUSE_ROUGHNESS(PixelBSDFs.BSDFs[i]));
-		}
-		break;
-		case STRATA_BSDF_TYPE_DIELECTRIC:
-		{
-			StrataLoad_R11G11B10F			(StrataBuffer, ByteOffset, DIELECTRIC_TINT(PixelBSDFs.BSDFs[i]));
-			StrataLoad_R11_G11_B10F			(StrataBuffer, ByteOffset, DIELECTRIC_ROUGHNESS0(PixelBSDFs.BSDFs[i]), DIELECTRIC_ROUGHNESS1(PixelBSDFs.BSDFs[i]), DIELECTRIC_IOR(PixelBSDFs.BSDFs[i]));
-			StrataLoad_R16G16B16F_A16F		(StrataBuffer, ByteOffset, DIELECTRIC_NORMAL(PixelBSDFs.BSDFs[i]),		DummyFloat);
-		}
-		break;
-		case STRATA_BSDF_TYPE_CONDUCTOR:
-		{
-			StrataLoad_AlbedoGamma2ToLinearAlphaLinear(StrataBuffer, ByteOffset, CONDUCTOR_REFLECTIVITY(PixelBSDFs.BSDFs[i]),	CONDUCTOR_ROUGHNESS0(PixelBSDFs.BSDFs[i]));
-			StrataLoad_AlbedoGamma2ToLinearAlphaLinear(StrataBuffer, ByteOffset, CONDUCTOR_EDGECOLOR(PixelBSDFs.BSDFs[i]),		CONDUCTOR_ROUGHNESS1(PixelBSDFs.BSDFs[i]));
-			StrataLoad_R16G16B16F_A16F		(StrataBuffer, ByteOffset,			 CONDUCTOR_NORMAL(PixelBSDFs.BSDFs[i]),			DummyFloat);
-		}
-		break;
-		case STRATA_BSDF_TYPE_VOLUME:
-		{
-			StrataLoad_R11G11B10F			(StrataBuffer, ByteOffset, VOLUME_EXTINCTION(PixelBSDFs.BSDFs[i]));
-			StrataLoad_R11G11B10F			(StrataBuffer, ByteOffset, VOLUME_ALBEDO(PixelBSDFs.BSDFs[i]));
-			StrataLoad_R16G16B16F_A16F		(StrataBuffer, ByteOffset, VOLUME_SMOOTHSURFACENORMAL(PixelBSDFs.BSDFs[i]), VOLUME_ANISOTROPY(PixelBSDFs.BSDFs[i]));
-			VOLUME_THICKNESS(PixelBSDFs.BSDFs[i]) = 1.0f;
-		}
-		break;
-		}
+	switch (OutBSDF.BSDFType)
+	{
+	case STRATA_BSDF_TYPE_DIFFUSE_ON:
+	case STRATA_BSDF_TYPE_DIFFUSE_CHAN:
+	{
+		StrataLoad_AlbedoGamma2ToLinear	(StrataBuffer, ByteOffset, DIFFUSE_ALBEDO(OutBSDF));
+		StrataLoad_R16G16B16F_A16F		(StrataBuffer, ByteOffset, DIFFUSE_NORMAL(OutBSDF),			DIFFUSE_ROUGHNESS(OutBSDF));
+	}
+	break;
+	case STRATA_BSDF_TYPE_DIELECTRIC:
+	{
+		StrataLoad_R11G11B10F			(StrataBuffer, ByteOffset, DIELECTRIC_TINT(OutBSDF));
+		StrataLoad_R11_G11_B10F			(StrataBuffer, ByteOffset, DIELECTRIC_ROUGHNESS0(OutBSDF), DIELECTRIC_ROUGHNESS1(OutBSDF), DIELECTRIC_IOR(OutBSDF));
+		StrataLoad_R16G16B16F_A16F		(StrataBuffer, ByteOffset, DIELECTRIC_NORMAL(OutBSDF),		DummyFloat);
+	}
+	break;
+	case STRATA_BSDF_TYPE_CONDUCTOR:
+	{
+		StrataLoad_AlbedoGamma2ToLinearAlphaLinear(StrataBuffer, ByteOffset, CONDUCTOR_REFLECTIVITY(OutBSDF),	CONDUCTOR_ROUGHNESS0(OutBSDF));
+		StrataLoad_AlbedoGamma2ToLinearAlphaLinear(StrataBuffer, ByteOffset, CONDUCTOR_EDGECOLOR(OutBSDF),		CONDUCTOR_ROUGHNESS1(OutBSDF));
+		StrataLoad_R16G16B16F_A16F		(StrataBuffer, ByteOffset,			 CONDUCTOR_NORMAL(OutBSDF),			DummyFloat);
+	}
+	break;
+	case STRATA_BSDF_TYPE_VOLUME:
+	{
+		StrataLoad_R11G11B10F			(StrataBuffer, ByteOffset, VOLUME_EXTINCTION(OutBSDF));
+		StrataLoad_R11G11B10F			(StrataBuffer, ByteOffset, VOLUME_ALBEDO(OutBSDF));
+		StrataLoad_R16G16B16F_A16F		(StrataBuffer, ByteOffset, VOLUME_SMOOTHSURFACENORMAL(OutBSDF), VOLUME_ANISOTROPY(OutBSDF));
+		VOLUME_THICKNESS(OutBSDF) = 1.0f;
+	}
+	break;
 	}
 
-	return PixelBSDFs;
-}
+	return OutBSDF;
+}
\ No newline at end of file
diff --git a/Engine/Shaders/Private/Strata/StrataDeferredLighting.ush b/Engine/Shaders/Private/Strata/StrataDeferredLighting.ush
--- a/Engine/Shaders/Private/Strata/StrataDeferredLighting.ush
+++ b/Engine/Shaders/Private/Strata/StrataDeferredLighting.ush
@@ -2,8 +2,8 @@
 
 #pragma once
 
-
-
+// Analytical lighting evaluation for strata material.
+// Unpack BSDF on-the-fly
 float4 StrataDeferredLighting(
 	FDeferredLightData LightData,
 	FGBufferData GBuffer,
@@ -13,8 +13,8 @@ float4 StrataDeferredLighting(
 	float LightMask,
 	float SurfaceShadow,
 	FRectTexture RectTexture,
-	FStrataPixelBSDFs PixelBSDFs
-	)
+	ByteAddressBuffer StrataDataBuffer,
+	uint PixelStrataDataByteOffset)
 {
 	float4 Radiance = 0;
 
@@ -31,9 +31,12 @@ float4 StrataDeferredLighting(
 		UglyIntegrationFactor = IntegrateLight(Capsule, LightData.bInverseSquared);
 	}
 
-	for (int i = 0; i < PixelBSDFs.BSDFCount; ++i)
+	FStrataPixelHeader Header = UnpackStrataHeaderIn(StrataDataBuffer, PixelStrataDataByteOffset);
+	
+	LOOP
+	for (uint BSDFIndex = 0; BSDFIndex < Header.BSDFCount; ++BSDFIndex)
 	{
-		FStrataBSDF BSDF = PixelBSDFs.BSDFs[i];
+		const FStrataBSDF BSDF = UnpackStrataBSDFIn(StrataDataBuffer, PixelStrataDataByteOffset);
 		const float3 BSDFThroughput = BSDF.Weight;
 
 		// Create the BSDF context
@@ -50,5 +53,3 @@ float4 StrataDeferredLighting(
 
 	return Radiance;
 }
-
-
diff --git a/Engine/Shaders/Private/Strata/StrataEnvironmentLighting.ush b/Engine/Shaders/Private/Strata/StrataEnvironmentLighting.ush
--- a/Engine/Shaders/Private/Strata/StrataEnvironmentLighting.ush
+++ b/Engine/Shaders/Private/Strata/StrataEnvironmentLighting.ush
@@ -2,8 +2,8 @@
 
 #pragma once
 
-
-
+// Env. lighting evaluation for strata material.
+// Unpack BSDF on-the-fly
 float3 StrataEnvironmentLighting(
 	float2 BufferUV,
 	FGBufferData GBuffer,
@@ -11,13 +11,15 @@ float3 StrataEnvironmentLighting(
 	float3 CameraToPixel,
 	float3 V,
 	float AmbientOcclusion,
-	FStrataPixelBSDFs PixelBSDFs
-	)
+	ByteAddressBuffer StrataDataBuffer,
+	uint PixelStrataDataByteOffset)
 {
+	FStrataPixelHeader Header = UnpackStrataHeaderIn(StrataDataBuffer, PixelStrataDataByteOffset);
+
 	float3 SkyLighting = 0;
-	for (int i = 0; i < PixelBSDFs.BSDFCount; ++i)
+	for (uint BSDFIndex = 0; BSDFIndex < Header.BSDFCount; ++BSDFIndex)
 	{
-		FStrataBSDF BSDF = PixelBSDFs.BSDFs[i];
+		const FStrataBSDF BSDF = UnpackStrataBSDFIn(StrataDataBuffer, PixelStrataDataByteOffset);
 		const float3 BSDFThroughput = BSDF.Weight;
 
 		// Create the BSDF context
@@ -63,4 +65,3 @@ float3 StrataEnvironmentLighting(
 	return SkyLighting;
 }
 
-
diff --git a/Engine/Shaders/Private/Strata/StrataVisualize.usf b/Engine/Shaders/Private/Strata/StrataVisualize.usf
--- a/Engine/Shaders/Private/Strata/StrataVisualize.usf
+++ b/Engine/Shaders/Private/Strata/StrataVisualize.usf
@@ -31,26 +31,27 @@ void VisualizeMaterialPS(
 	OutColor = 0;
 	if (bStrataMaterial)
 	{
-		uint PixelStrataDataByteOffset = GetStrataPixelDataByteOffset(PixelPos, View.BufferSizeAndInvSize.x, Strata.MaxBytesPerPixel);
-		FStrataPixelBSDFs PixelBSDFs = UnpackStrataIn(Strata.MaterialLobesBuffer, PixelStrataDataByteOffset);
+		uint PixelStrataDataByteOffset = GetStrataPixelDataByteOffset(MaterialPos, View.BufferSizeAndInvSize.x, Strata.MaxBytesPerPixel);
+		FStrataPixelHeader Header = UnpackStrataHeaderIn(Strata.MaterialLobesBuffer, PixelStrataDataByteOffset);
+
 		OutColor.w = 0.05f;
-		if (PixelBSDFs.BSDFCount == 1)
+		if (Header.BSDFCount == 1)
 		{
 			OutColor = float4(0, 0, 1, 0);
 		}
-		else if (PixelBSDFs.BSDFCount == 2)
+		else if (Header.BSDFCount == 2)
 		{
 			OutColor = float4(0, 1, 0, 0);
 		}
-		else if (PixelBSDFs.BSDFCount == 3)
+		else if (Header.BSDFCount == 3)
 		{
 			OutColor = float4(1, 1, 0, 0);
 		}
-		else if (PixelBSDFs.BSDFCount == 4)
+		else if (Header.BSDFCount == 4)
 		{
 			OutColor = float4(1, 0, 0, 0);
 		}
-		else if (PixelBSDFs.BSDFCount > 4)
+		else if (Header.BSDFCount > 4)
 		{
 			OutColor = float4(1, 1, 1, 0);
 		}
@@ -68,7 +69,7 @@ void VisualizeMaterialPS(
 	if (bStrataMaterial)
 	{
 		uint PixelStrataDataByteOffset = GetStrataPixelDataByteOffset(MaterialPos, View.BufferSizeAndInvSize.x, Strata.MaxBytesPerPixel);
-		FStrataPixelBSDFs PixelBSDFs = UnpackStrataIn(Strata.MaterialLobesBuffer, PixelStrataDataByteOffset);
+		FStrataPixelHeader Header = UnpackStrataHeaderIn(Strata.MaterialLobesBuffer, PixelStrataDataByteOffset);
 
 		float3 BSDFFontColor = float3(0.25, 1.0, 0.25);
 		float3 INFOFontColor = float3(1.0, 0.5, 0.25);
@@ -82,7 +83,7 @@ void VisualizeMaterialPS(
 #define PRINTBSDFCHAR(X) PrintCharacter(PixelPos, OutColor.xyz, BSDFFontColor, TopLeft, X);
 
 #define PRINTBSDFWEIGHT()	PRINTBSDFCHAR(_SPC_); PRINTBSDFCHAR(_MINUS_); PRINTBSDFCHAR(_SPC_); PrintFloat(PixelPos, OutColor.xyz, FontColor, TopLeft, BSDF.Weight.x); TopLeft.x+=8*7; PRINTBSDFCHAR(_SPC_); PrintFloat(PixelPos, OutColor.xyz, FontColor, TopLeft, BSDF.Weight.y); TopLeft.x+=8*7; PRINTBSDFCHAR(_SPC_); PrintFloat(PixelPos, OutColor.xyz, FontColor, TopLeft, BSDF.Weight.z);
-#define PRINTBSDFCOUNT()	PRINTBSDFCHAR(min(_0_+PixelBSDFs.BSDFCount,_F_)); PRINTBSDFCHAR(_SPC_); PRINTBSDFCHAR(_B_); PRINTBSDFCHAR(_S_); PRINTBSDFCHAR(_D_); PRINTBSDFCHAR(_F_);
+#define PRINTBSDFCOUNT()	PRINTBSDFCHAR(min(_0_+Header.BSDFCount,_F_)); PRINTBSDFCHAR(_SPC_); PRINTBSDFCHAR(_B_); PRINTBSDFCHAR(_S_); PRINTBSDFCHAR(_D_); PRINTBSDFCHAR(_F_);
 #define PRINTFLOAT3(X) PrintFloat(PixelPos, OutColor.xyz, PARAMFontColor, TopLeft, X.x); TopLeft.x+=8*7; PRINTBSDFCHAR(_SPC_); PrintFloat(PixelPos, OutColor.xyz, PARAMFontColor, TopLeft, X.y); TopLeft.x+=8*7; PRINTBSDFCHAR(_SPC_); PrintFloat(PixelPos, OutColor.xyz, PARAMFontColor, TopLeft, X.z);
 #define PRINTFLOAT1(Value) PrintFloat(PixelPos, OutColor.xyz, PARAMFontColor, TopLeft, Value); TopLeft.x+=8*7;
 #define SIGNOFFLOAT1(Value) if(Value < 0.0f) { PrintCharacter(PixelPos, OutColor.xyz, PARAMFontColor, TopLeft, _MINUS_); } /*else { PrintCharacter(PixelPos, OutColor.xyz, PARAMFontColor, TopLeft, _SPC_);}*/
@@ -98,9 +99,9 @@ void VisualizeMaterialPS(
 		NEXTLINE;
 		HORIZONTALSEPARATION;
 
-		for (int i = 0; i < PixelBSDFs.BSDFCount; ++i)
+		for (uint BSDFIndex = 0; BSDFIndex < Header.BSDFCount; ++BSDFIndex)
 		{
-			FStrataBSDF BSDF = PixelBSDFs.BSDFs[i];
+			const FStrataBSDF BSDF = UnpackStrataBSDFIn(Strata.MaterialLobesBuffer, PixelStrataDataByteOffset);
 
 			TopLeft.x = TopLeftInit.x;
 
