Fixed PLATFORM_SUPPORTS_SM6_0_WAVE_OPERATIONS conflicting COMPILER_SUPPORTS_TO_SCALAR_MEMORY on some platforms

diff --git a/Engine/Shaders/Private/HairStrands/HairCardsVertexFactory.ush b/Engine/Shaders/Private/HairStrands/HairCardsVertexFactory.ush
--- a/Engine/Shaders/Private/HairStrands/HairCardsVertexFactory.ush
+++ b/Engine/Shaders/Private/HairStrands/HairCardsVertexFactory.ush
@@ -168,7 +168,7 @@ FMaterialPixelParameters GetMaterialPixelParameters(FVertexFactoryInterpolantsVS
 #endif
 
 	Result.TwoSidedSign = 1;
-	Result.PrimitiveId = GetPrimitiveId(Interpolants);
+	Result.PrimitiveId = ToScalarMemory(GetPrimitiveId(Interpolants));
 	Result.HairPrimitiveId = Interpolants.HairPrimitiveId;
 	Result.HairPrimitiveUV = float2(Interpolants.HairPrimitiveUV);
 #if NUM_TEX_COORD_INTERPOLATORS
@@ -179,10 +179,6 @@ FMaterialPixelParameters GetMaterialPixelParameters(FVertexFactoryInterpolantsVS
 	}
 #endif
 
-#if COMPILER_SUPPORTS_TO_SCALAR_MEMORY
-	Result.PrimitiveId = ToScalarMemory(Result.PrimitiveId);
-#endif
-
 	return Result;
 }
 
diff --git a/Engine/Shaders/Private/HairStrands/HairStrandsVertexFactory.ush b/Engine/Shaders/Private/HairStrands/HairStrandsVertexFactory.ush
--- a/Engine/Shaders/Private/HairStrands/HairStrandsVertexFactory.ush
+++ b/Engine/Shaders/Private/HairStrands/HairStrandsVertexFactory.ush
@@ -265,14 +265,10 @@ FMaterialPixelParameters GetMaterialPixelParameters(FVertexFactoryInterpolantsVS
 #endif
 
 	Result.TwoSidedSign = 1;
-	Result.PrimitiveId = GetPrimitiveId(Interpolants);
+	Result.PrimitiveId = ToScalarMemory(GetPrimitiveId(Interpolants));
 	Result.HairPrimitiveId = Interpolants.HairPrimitiveId;
 	Result.HairPrimitiveUV = float2(Interpolants.HairPrimitiveUV);
 
-#if COMPILER_SUPPORTS_TO_SCALAR_MEMORY
-	Result.PrimitiveId = ToScalarMemory(Result.PrimitiveId);
-#endif
-
 	return Result;
 }
 
diff --git a/Engine/Shaders/Public/Platform.ush b/Engine/Shaders/Public/Platform.ush
--- a/Engine/Shaders/Public/Platform.ush
+++ b/Engine/Shaders/Public/Platform.ush
@@ -808,12 +808,11 @@ uint MaskedBitCount( uint2 Bits, uint Index )
 }
 
 #if PLATFORM_SUPPORTS_SM6_0_WAVE_OPERATIONS
-#ifdef COMPILER_SUPPORTS_TO_SCALAR_MEMORY
-#error Conflict of the ToScalarMemory definition.
-#endif
 
+#ifndef COMPILER_SUPPORTS_TO_SCALAR_MEMORY
 #define COMPILER_SUPPORTS_TO_SCALAR_MEMORY 1
 #define ToScalarMemory(x) WaveReadLaneFirst(x)
+#endif
 
 #ifndef COMPILER_SUPPORTS_MASKED_BIT_COUNT
 #define COMPILER_SUPPORTS_MASKED_BIT_COUNT 1
