Fix mobile shader compilation for a deferred shading

diff --git a/Engine/Shaders/Private/MobileBasePassPixelShader.usf b/Engine/Shaders/Private/MobileBasePassPixelShader.usf
--- a/Engine/Shaders/Private/MobileBasePassPixelShader.usf
+++ b/Engine/Shaders/Private/MobileBasePassPixelShader.usf
@@ -650,6 +650,9 @@ void Main(
 	ComputeIndirect(Interpolants, MaterialParameters, ShadingModelContext, IndirectIrradiance, IndirectColor);
 	Color += IndirectColor;
 
+	// Shadow
+	half Shadow = GetPrimaryPrecomputedShadowMask(Interpolants, MaterialParameters).r;
+
 #if DEFERRED_SHADING_PATH
 	float4 OutGBufferD;
 	float4 OutGBufferE;
@@ -657,7 +660,7 @@ void Main(
 	float4 OutGBufferVelocity = 0;
 
 	GBuffer.IndirectIrradiance = IndirectIrradiance;
-	GBuffer.PrecomputedShadowFactors.r = GetPrimaryPrecomputedShadowMask(Interpolants).r;
+	GBuffer.PrecomputedShadowFactors.r = Shadow;
 
 	EncodeGBuffer(GBuffer, OutGBufferA, OutGBufferB, OutGBufferC, OutGBufferD, OutGBufferE, OutGBufferF, OutGBufferVelocity);
 	OutSceneDepthAux = SvPosition.z;
@@ -666,16 +669,14 @@ void Main(
 #if !MATERIAL_SHADINGMODEL_UNLIT
 
 #if ENABLE_SKY_LIGHT
-		half3 SkyDiffuseLighting = GetSkySHDiffuseSimple(MaterialParameters.WorldNormal);
-		half3 DiffuseLookup = SkyDiffuseLighting * ResolvedView.SkyLightColor.rgb;
-		IndirectIrradiance += Luminance(DiffuseLookup);
+	half3 SkyDiffuseLighting = GetSkySHDiffuseSimple(MaterialParameters.WorldNormal);
+	half3 DiffuseLookup = SkyDiffuseLighting * ResolvedView.SkyLightColor.rgb;
+	IndirectIrradiance += Luminance(DiffuseLookup);
 #endif
 			
 	Color *= MaterialAO;
 	IndirectIrradiance *= MaterialAO;
 
-	// Shadow
-	half Shadow = GetPrimaryPrecomputedShadowMask(Interpolants, MaterialParameters).r;
 #if DIRECTIONAL_LIGHT_CSM && !MATERIAL_SHADINGMODEL_SINGLELAYERWATER
 	// Cascaded Shadow Map
 	if(MobileBasePass.UseCSM)
