Strata - made IOR on dielectric node be Reflectivity (F0). Added a node to compute reflectivity from IOR.
Tested on strata level.

diff --git a/Engine/Shaders/Private/Strata/Strata.ush b/Engine/Shaders/Private/Strata/Strata.ush
--- a/Engine/Shaders/Private/Strata/Strata.ush
+++ b/Engine/Shaders/Private/Strata/Strata.ush
@@ -34,7 +34,7 @@
 //
 //struct FStrataDielectric
 //{
-//	float	IOR;
+//	float	Reflectivity;
 //	float3	Tint;
 //	float2	Roughness;
 //
@@ -91,7 +91,7 @@ struct FStrataBSDF
 #define DIFFUSE_NORMAL(X)				X.VGPRs[1].xyz
 
 #define DIELECTRIC_TINT(X)				X.VGPRs[0].xyz
-#define DIELECTRIC_IOR(X)				X.VGPRs[0].w
+#define DIELECTRIC_REFLECTIVITY(X)		X.VGPRs[0].w
 #define DIELECTRIC_ROUGHNESS0(X)		X.VGPRs[1].w
 #define DIELECTRIC_ROUGHNESS1(X)		X.VGPRs[2].w
 #define DIELECTRIC_NORMAL(X)			X.VGPRs[1].xyz
@@ -162,14 +162,14 @@ FStrataData GetStrataDiffuseChanBSDF(float3 Albedo, float Roughness, float3 Norm
 	return StrataData;
 }
 
-FStrataData GetStrataDielectricBSDF(float2 Roughness, float IOR, float3 Tint, float3 Normal)
+FStrataData GetStrataDielectricBSDF(float2 Roughness, float Reflectivity, float3 Tint, float3 Normal)
 {
 	FStrataData StrataData = GetInitialisedStrataData();
 	StrataData.LayerCount++;
 
 	StrataData.Layers[0].BSDFCount++;
 	StrataData.Layers[0].BSDFs[0].BSDFType						= STRATA_BSDF_TYPE_DIELECTRIC;
-	DIELECTRIC_IOR			(StrataData.Layers[0].BSDFs[0])		= IOR;
+	DIELECTRIC_REFLECTIVITY	(StrataData.Layers[0].BSDFs[0])		= Reflectivity;
 	DIELECTRIC_TINT			(StrataData.Layers[0].BSDFs[0])		= Tint;
 	DIELECTRIC_ROUGHNESS0	(StrataData.Layers[0].BSDFs[0])		= Roughness.x;
 	DIELECTRIC_ROUGHNESS1	(StrataData.Layers[0].BSDFs[0])		= Roughness.y;
@@ -348,7 +348,7 @@ FStrataEvaluateResult StrataEvaluateBSDF(FStrataBSDFContext BSDFContext)
 		}
 		case STRATA_BSDF_TYPE_DIELECTRIC:
 		{
-			float F0 = DielectricIorToF0(DIELECTRIC_IOR(BSDFContext.BSDF));
+			float F0 = DIELECTRIC_REFLECTIVITY(BSDFContext.BSDF);
 
 			const float SafeRoughness = MakeRoughnessSafe(DIELECTRIC_ROUGHNESS0(BSDFContext.BSDF));
 			float a2 = Pow4(SafeRoughness);
@@ -514,7 +514,7 @@ FStrataEnvLightResult StrataEvaluateForEnvLight(FStrataBSDFContext BSDFContext,
 			BRANCH
 			if (bEnableSpecular)
 			{
-				float F0 = DielectricIorToF0(DIELECTRIC_IOR(BSDFContext.BSDF));
+				float F0 = DIELECTRIC_REFLECTIVITY(BSDFContext.BSDF);
 				const float SafeRoughness = MakeRoughnessSafe(DIELECTRIC_ROUGHNESS0(BSDFContext.BSDF));
 
 				float3 EvalEnvBRDF = EnvBRDFApprox(F0, SafeRoughness, BSDFContext.SatNoV);
@@ -655,8 +655,7 @@ void SanitizeStrataDiffuse(inout FStrataBSDF BSDF)
 
 void SanitizeStrataDielectric(inout FStrataBSDF BSDF)
 {
-
-	DIELECTRIC_IOR(BSDF) = clamp(DIELECTRIC_IOR(BSDF), 1.0f, 5.0f);		// See https://en.wikipedia.org/wiki/List_of_refractive_indices
+	DIELECTRIC_REFLECTIVITY(BSDF) = saturate(DIELECTRIC_REFLECTIVITY(BSDF));
 	DIELECTRIC_TINT(BSDF) = saturate(DIELECTRIC_TINT(BSDF));
 	DIELECTRIC_ROUGHNESS0(BSDF) = saturate(DIELECTRIC_ROUGHNESS0(BSDF));
 	DIELECTRIC_ROUGHNESS1(BSDF) = saturate(DIELECTRIC_ROUGHNESS1(BSDF));
@@ -950,7 +949,7 @@ void PackStrataOut(
 
 			case STRATA_BSDF_TYPE_DIELECTRIC:
 			{
-				float F0 = DielectricIorToF0(DIELECTRIC_IOR(BSDF));
+				float F0 = DIELECTRIC_REFLECTIVITY(BSDF);
 				float3 F = F_Schlick(F0, ContextNoLight.VoH);
 
 				if (bTopLayer)
@@ -1067,7 +1066,7 @@ void PackStrataOut(
 			case STRATA_BSDF_TYPE_DIELECTRIC: //
 			{
 				STRATA_STORE_UINT1(PackR11G11B10F(DIELECTRIC_TINT(BSDF)));
-				STRATA_STORE_UINT1(PackR11G11B10F(float3(DIELECTRIC_ROUGHNESS0(BSDF), DIELECTRIC_ROUGHNESS1(BSDF), DIELECTRIC_IOR(BSDF))));
+				STRATA_STORE_UINT1(PackR11G11B10F(float3(DIELECTRIC_ROUGHNESS0(BSDF), DIELECTRIC_ROUGHNESS1(BSDF), DIELECTRIC_REFLECTIVITY(BSDF))));
 				STRATA_STORE_UINT2(PackR16G16B16A16F(float4(DIELECTRIC_NORMAL(BSDF), 0.0f)));
 				// 16 bytes
 			}
@@ -1144,7 +1143,7 @@ FStrataBSDF UnpackStrataBSDFIn(ByteAddressBuffer StrataBuffer, inout uint ByteOf
 	case STRATA_BSDF_TYPE_DIELECTRIC:
 	{
 		StrataLoad_R11G11B10F			(StrataBuffer, ByteOffset, DIELECTRIC_TINT(OutBSDF));
-		StrataLoad_R11_G11_B10F			(StrataBuffer, ByteOffset, DIELECTRIC_ROUGHNESS0(OutBSDF), DIELECTRIC_ROUGHNESS1(OutBSDF), DIELECTRIC_IOR(OutBSDF));
+		StrataLoad_R11_G11_B10F			(StrataBuffer, ByteOffset, DIELECTRIC_ROUGHNESS0(OutBSDF), DIELECTRIC_ROUGHNESS1(OutBSDF), DIELECTRIC_REFLECTIVITY(OutBSDF));
 		StrataLoad_R16G16B16F_A16F		(StrataBuffer, ByteOffset, DIELECTRIC_NORMAL(OutBSDF),		DummyFloat);
 	}
 	break;
diff --git a/Engine/Shaders/Private/Strata/StrataVisualize.usf b/Engine/Shaders/Private/Strata/StrataVisualize.usf
--- a/Engine/Shaders/Private/Strata/StrataVisualize.usf
+++ b/Engine/Shaders/Private/Strata/StrataVisualize.usf
@@ -143,8 +143,8 @@ void VisualizeMaterialPS(
 				PRINTBSDFWEIGHT();
 				NEXTLINE;
 
-				HORIZONTALTAB; PRINTPARAMCHAR(_I_); PRINTPARAMCHAR(_O_); PRINTPARAMCHAR(_R_); PRINTPARAMCHAR(_SPC_); PRINTPARAMCHAR(_SPC_); PRINTPARAMCHAR(_SPC_);
-				HORIZONTALTAB; PRINTFLOAT1(DIELECTRIC_IOR(BSDF));
+				HORIZONTALTAB; PRINTPARAMCHAR(_R_); PRINTPARAMCHAR(_E_); PRINTPARAMCHAR(_F_); PRINTPARAMCHAR(_L_); PRINTPARAMCHAR(_E_); PRINTPARAMCHAR(_C_);
+				HORIZONTALTAB; PRINTFLOAT1(DIELECTRIC_REFLECTIVITY(BSDF));
 				NEXTLINE;
 				HORIZONTALTAB; PRINTPARAMCHAR(_T_); PRINTPARAMCHAR(_I_); PRINTPARAMCHAR(_N_); PRINTPARAMCHAR(_T_); PRINTPARAMCHAR(_SPC_); PRINTPARAMCHAR(_SPC_);
 				HORIZONTALTAB; PRINTFLOAT3(DIELECTRIC_TINT(BSDF));
diff --git a/Engine/Source/Runtime/Engine/Classes/Materials/MaterialExpressionStrata.h b/Engine/Source/Runtime/Engine/Classes/Materials/MaterialExpressionStrata.h
--- a/Engine/Source/Runtime/Engine/Classes/Materials/MaterialExpressionStrata.h
+++ b/Engine/Source/Runtime/Engine/Classes/Materials/MaterialExpressionStrata.h
@@ -107,10 +107,10 @@ class UMaterialExpressionStrataDielectricBSDF : public UMaterialExpression
 	GENERATED_UCLASS_BODY()
 	
 	/**
-	 * The index of refraction of the surface (type = float, unit = unitless)
+	 * The reflectivity of the surface (type = float, unit = unitless, default=0.04 for glass)
 	 */
 	UPROPERTY()
-	FExpressionInput IOR;
+	FExpressionInput Reflectivity;
 
 	/**
 	 * A global color tint multiplied with the specular color, not physically based (type = float3, unit = unitless)
@@ -146,13 +146,13 @@ class UMaterialExpressionStrataConductorBSDF : public UMaterialExpression
 	GENERATED_UCLASS_BODY()
 		
 	/**
-	 * Reflectivity when view direction is perpendicular to the surface, also known as F0 (type = float3, unit = unitless)
+	 * Reflectivity when view direction is perpendicular to the surface, also known as F0 (type = float3, unit = unitless, defaults to gold)
 	 */
 	UPROPERTY()
 	FExpressionInput Reflectivity;
 	
 	/**
-	 * Reflectivity when the view direction is tangent to the surface (type = float3, unit = unitless)
+	 * Reflectivity when the view direction is tangent to the surface (type = float3, unit = unitless, defaults to gold)
 	 */
 	UPROPERTY()
 	FExpressionInput EdgeColor;
@@ -398,4 +398,24 @@ class UMaterialExpressionStrataPhysicalIOR : public UMaterialExpression
 	//~ End UMaterialExpression Interface
 };
 
+UCLASS(MinimalAPI, collapsecategories, hidecategories = Object)
+class UMaterialExpressionStrataDielectricIORToReflectivity : public UMaterialExpression
+{
+	GENERATED_UCLASS_BODY()
+
+	/**
+	 * The index of refraction of the surface (type = float, unit = unitless, default = 1.5 for glass)
+	 */
+	UPROPERTY()
+	FExpressionInput IOR;
+
+	//~ Begin UMaterialExpression Interface
+#if WITH_EDITOR
+	virtual int32 Compile(class FMaterialCompiler* Compiler, int32 OutputIndex) override;
+	virtual void GetCaption(TArray<FString>& OutCaptions) const override;
+	virtual uint32 GetOutputType(int32 OutputIndex) override;
+	virtual uint32 GetInputType(int32 InputIndex) override;
+#endif
+	//~ End UMaterialExpression Interface
+};
 
diff --git a/Engine/Source/Runtime/Engine/Private/Materials/HLSLMaterialTranslator.cpp b/Engine/Source/Runtime/Engine/Private/Materials/HLSLMaterialTranslator.cpp
--- a/Engine/Source/Runtime/Engine/Private/Materials/HLSLMaterialTranslator.cpp
+++ b/Engine/Source/Runtime/Engine/Private/Materials/HLSLMaterialTranslator.cpp
@@ -7254,12 +7254,12 @@ int32 FHLSLMaterialTranslator::StrataDiffuseChanBSDF(int32 Albedo, int32 Roughne
 	);
 }
 
-int32 FHLSLMaterialTranslator::StrataDielectricBSDF(int32 Roughness, int32 IOR, int32 Tint, int32 Normal)
+int32 FHLSLMaterialTranslator::StrataDielectricBSDF(int32 Roughness, int32 Reflectivity, int32 Tint, int32 Normal)
 {
 	return AddInlinedCodeChunk(
 		MCT_Strata, TEXT("GetStrataDielectricBSDF(%s, %s, %s, %s)"),
 		*GetParameterCode(Roughness),
-		*GetParameterCode(IOR),
+		*GetParameterCode(Reflectivity),
 		*GetParameterCode(Tint),
 		*GetParameterCode(Normal)
 	);
@@ -7369,6 +7369,18 @@ int32 FHLSLMaterialTranslator::StrataPhysicalIOR(int32 IOR, int32 Extinction, in
 	);
 }
 
+int32 FHLSLMaterialTranslator::StrataDielectricIORToReflectivity(int32 IOR, int32 OutputIndex)
+{
+	if (IOR == INDEX_NONE)
+	{
+		return INDEX_NONE;
+	}
+	return AddInlinedCodeChunk(
+		MCT_Float,
+		TEXT("DielectricIorToF0(%s)"),
+		*GetParameterCode(IOR));
+}
+
 int32 FHLSLMaterialTranslator::MapARPassthroughCameraUV(int32 UV)
 {
 	if (UV == INDEX_NONE)
diff --git a/Engine/Source/Runtime/Engine/Private/Materials/HLSLMaterialTranslator.h b/Engine/Source/Runtime/Engine/Private/Materials/HLSLMaterialTranslator.h
--- a/Engine/Source/Runtime/Engine/Private/Materials/HLSLMaterialTranslator.h
+++ b/Engine/Source/Runtime/Engine/Private/Materials/HLSLMaterialTranslator.h
@@ -781,7 +781,7 @@ protected:
 	virtual int32 FrontMaterial() override;
 	virtual int32 StrataDiffuseOrenNayarBSDF(int32 Albedo, int32 Roughness, int32 Normal) override;
 	virtual int32 StrataDiffuseChanBSDF(int32 Albedo, int32 Roughness, int32 Normal) override;
-	virtual int32 StrataDielectricBSDF(int32 Roughness, int32 IOR, int32 Tint, int32 Normal) override;
+	virtual int32 StrataDielectricBSDF(int32 Roughness, int32 Reflectivity, int32 Tint, int32 Normal) override;
 	virtual int32 StrataConductorBSDF(int32 Reflectivity, int32 EdgeColor, int32 Roughness, int32 Normal) override;
 	virtual int32 StrataVolumeBSDF(int32 Albedo, int32 Extinction, int32 Anisotropy, int32 Thickness) override;
 	virtual int32 StrataHorizontalMixing(int32 Foreground, int32 Background, int32 Mix) override;
@@ -790,6 +790,7 @@ protected:
 	virtual int32 StrataMultiply(int32 A, int32 Weight) override;
 	virtual int32 StrataArtisticIOR(int32 Reflectivity, int32 EdgeColor, int32 OutputIndex) override;
 	virtual int32 StrataPhysicalIOR(int32 IOR, int32 Extinction, int32 OutputIndex) override;
+	virtual int32 StrataDielectricIORToReflectivity(int32 IOR, int32 OutputIndex) override;
 
 	virtual void AddStrataCodeChunk(int32 CodeChunk, FStrataMaterialCompilationInfo& StrataMaterialCompilationInfo) override;
 	virtual bool ContainsStrataCodeChunk(int32 CodeChunk) override;
diff --git a/Engine/Source/Runtime/Engine/Private/Materials/MaterialExpressions.cpp b/Engine/Source/Runtime/Engine/Private/Materials/MaterialExpressions.cpp
--- a/Engine/Source/Runtime/Engine/Private/Materials/MaterialExpressions.cpp
+++ b/Engine/Source/Runtime/Engine/Private/Materials/MaterialExpressions.cpp
@@ -19154,7 +19154,7 @@ int32 UMaterialExpressionStrataDielectricBSDF::Compile(class FMaterialCompiler*
 {
 	int32 OutputCodeChunk = Compiler->StrataDielectricBSDF(
 		Roughness.GetTracedInput().Expression	? Roughness.Compile(Compiler)	: Compiler->Constant2(0.0f, 0.0f),
-		IOR.GetTracedInput().Expression			? IOR.Compile(Compiler)			: Compiler->Constant(1.5f),
+		Reflectivity.GetTracedInput().Expression? Reflectivity.Compile(Compiler): Compiler->Constant(0.04f),// Default to Glass
 		Tint.GetTracedInput().Expression		? Tint.Compile(Compiler)		: Compiler->Constant3(1.0f, 1.0f, 1.0f),
 		Normal.GetTracedInput().Expression		? Normal.Compile(Compiler)		: Compiler->PixelNormalWS());
 
@@ -19665,4 +19665,49 @@ uint32 UMaterialExpressionStrataPhysicalIOR::GetInputType(int32 InputIndex)
 
 
 
+UMaterialExpressionStrataDielectricIORToReflectivity::UMaterialExpressionStrataDielectricIORToReflectivity(const FObjectInitializer& ObjectInitializer)
+	: Super(ObjectInitializer)
+{
+	struct FConstructorStatics
+	{
+		FText NAME_Strata;
+		FConstructorStatics() : NAME_Strata(LOCTEXT("Strata", "Strata")) { }
+	};
+	static FConstructorStatics ConstructorStatics;
+#if WITH_EDITORONLY_DATA
+	MenuCategories.Add(ConstructorStatics.NAME_Strata);
+
+	bShowOutputNameOnPin = true;
+
+	Outputs.Reset();
+	Outputs.Add(FExpressionOutput(TEXT("Reflectivity")));
+#endif
+}
+
+#if WITH_EDITOR
+int32 UMaterialExpressionStrataDielectricIORToReflectivity::Compile(class FMaterialCompiler* Compiler, int32 OutputIndex)
+{
+	return Compiler->StrataDielectricIORToReflectivity(
+		IOR.GetTracedInput().Expression			? IOR.Compile(Compiler)			: Compiler->Constant(1.5f),// Default to Glass IOR
+		OutputIndex);
+}
+
+void UMaterialExpressionStrataDielectricIORToReflectivity::GetCaption(TArray<FString>& OutCaptions) const
+{
+	OutCaptions.Add(TEXT("Strata IOR to Reflectivity"));
+}
+
+uint32 UMaterialExpressionStrataDielectricIORToReflectivity::GetOutputType(int32 OutputIndex)
+{
+	return MCT_Float;
+}
+
+uint32 UMaterialExpressionStrataDielectricIORToReflectivity::GetInputType(int32 InputIndex)
+{
+	return MCT_Float;
+}
+#endif // WITH_EDITOR
+
+
+
 #undef LOCTEXT_NAMESPACE
diff --git a/Engine/Source/Runtime/Engine/Public/MaterialCompiler.h b/Engine/Source/Runtime/Engine/Public/MaterialCompiler.h
--- a/Engine/Source/Runtime/Engine/Public/MaterialCompiler.h
+++ b/Engine/Source/Runtime/Engine/Public/MaterialCompiler.h
@@ -393,7 +393,7 @@ public:
 	virtual int32 FrontMaterial() = 0;
 	virtual int32 StrataDiffuseOrenNayarBSDF(int32 Albedo, int32 Roughness, int32 Normal) = 0;
 	virtual int32 StrataDiffuseChanBSDF(int32 Albedo, int32 Roughness, int32 Normal) = 0;
-	virtual int32 StrataDielectricBSDF(int32 Roughness, int32 IOR, int32 Tint, int32 Normal) = 0;
+	virtual int32 StrataDielectricBSDF(int32 Roughness, int32 Reflectivity, int32 Tint, int32 Normal) = 0;
 	virtual int32 StrataConductorBSDF(int32 Reflectivity, int32 EdgeColor, int32 Roughness, int32 Normal) = 0;
 	virtual int32 StrataVolumeBSDF(int32 Albedo, int32 Extinction, int32 Anisotropy, int32 Thickness) = 0;
 	virtual int32 StrataHorizontalMixing(int32 Foreground, int32 Background, int32 Mix) = 0;
@@ -402,6 +402,7 @@ public:
 	virtual int32 StrataMultiply(int32 A, int32 Weight) = 0;
 	virtual int32 StrataArtisticIOR(int32 Reflectivity, int32 EdgeColor, int32 OutputIndex) = 0;
 	virtual int32 StrataPhysicalIOR(int32 IOR, int32 Extinction, int32 OutputIndex) = 0;
+	virtual int32 StrataDielectricIORToReflectivity(int32 IOR, int32 OutputIndex) = 0;
 
 	virtual void AddStrataCodeChunk(int32 CodeChunk, FStrataMaterialCompilationInfo& StrataMaterialCompilationInfo) = 0;
 	virtual bool ContainsStrataCodeChunk(int32 CodeChunk) = 0;
@@ -811,9 +812,9 @@ public:
 		return Compiler->StrataDiffuseChanBSDF(Albedo, Roughness, Normal);
 	}
 
-	virtual int32 StrataDielectricBSDF(int32 Roughness, int32 IOR, int32 Tint, int32 Normal) override
+	virtual int32 StrataDielectricBSDF(int32 Roughness, int32 Reflectivity, int32 Tint, int32 Normal) override
 	{
-		return Compiler->StrataDielectricBSDF(Roughness, IOR, Tint, Normal);
+		return Compiler->StrataDielectricBSDF(Roughness, Reflectivity, Tint, Normal);
 	}
 
 	virtual int32 StrataConductorBSDF(int32 Reflectivity, int32 EdgeColor, int32 Roughness, int32 Normal) override
@@ -856,6 +857,11 @@ public:
 		return Compiler->StrataPhysicalIOR(IOR, Extinction, OutputIndex);
 	}
 
+	virtual int32 StrataDielectricIORToReflectivity(int32 IOR, int32 OutputIndex) override
+	{
+		return Compiler->StrataDielectricIORToReflectivity(IOR, OutputIndex);
+	}
+
 	virtual void AddStrataCodeChunk(int32 CodeChunk, FStrataMaterialCompilationInfo& StrataMaterialCompilationInfo) override
 	{
 		Compiler->AddStrataCodeChunk(CodeChunk, StrataMaterialCompilationInfo);
