Fixed global SDF tracing artifacts between clipmap levels. When ray wasn't hitting the selected clipmap we were still moving it to the intersection point with that clipmap.

diff --git a/Engine/Shaders/Private/GlobalDistanceFieldUtils.ush b/Engine/Shaders/Private/GlobalDistanceFieldUtils.ush
--- a/Engine/Shaders/Private/GlobalDistanceFieldUtils.ush
+++ b/Engine/Shaders/Private/GlobalDistanceFieldUtils.ush
@@ -70,7 +70,7 @@ FGlobalSDFTraceResult RayTraceGlobalDistanceField(FGlobalSDFTraceInput TraceInpu
 	TraceResult.HitClipmapIndex = 0;
 	TraceResult.TotalStepsTaken = 0;
 
-	uint MinClipmapIndex = ComputeGlobalDistanceFieldClipmapIndex(TraceInput.WorldRayStart);
+	uint MinClipmapIndex = ComputeGlobalDistanceFieldClipmapIndex(TraceInput.WorldRayStart + TraceInput.MinTraceDistance * TraceInput.WorldRayDirection);
 	float MaxDistance = TraceInput.InitialMaxDistance;
 	float MinRayTime = TraceInput.MinTraceDistance;
 
@@ -88,11 +88,11 @@ FGlobalSDFTraceResult RayTraceGlobalDistanceField(FGlobalSDFTraceInput TraceInpu
 		IntersectionTimes.xy *= TraceInput.MaxTraceDistance;
 		IntersectionTimes.x = max(IntersectionTimes.x, MinRayTime);
 
-		// Update the trace start for the next clipmap
-		MinRayTime = IntersectionTimes.y;
-
 		if (IntersectionTimes.x < IntersectionTimes.y)
 		{
+			// Update the trace start for the next clipmap
+			MinRayTime = IntersectionTimes.y;
+
 			float SampleRayTime = IntersectionTimes.x;
 			const float ClipmapInfluenceRange = GLOBAL_DISTANCE_FIELD_INFLUENCE_RANGE_IN_VOXELS * 2.0f * GlobalVolumeCenterAndExtent[ClipmapIndex].w * GlobalVolumeTexelSize;
 
