Merge dieletric/conductor evaluation code.

diff --git a/Engine/Shaders/Private/Strata/Strata.ush b/Engine/Shaders/Private/Strata/Strata.ush
--- a/Engine/Shaders/Private/Strata/Strata.ush
+++ b/Engine/Shaders/Private/Strata/Strata.ush
@@ -338,8 +338,7 @@ FStrataEvaluateResult StrataEvaluateBSDF(FStrataBSDFContext BSDFContext)
 		case STRATA_BSDF_TYPE_DIFFUSE_CHAN:
 		{
 			const float SafeRoughness = MakeRoughnessSafe(DIFFUSE_ROUGHNESS(BSDFContext.BSDF));
-			float a = SafeRoughness * SafeRoughness;
-			float a2 = a * a;
+			float a2 = Pow4(SafeRoughness);
 
 			Sample.Value		= Diffuse_Chan(DIFFUSE_ALBEDO(BSDFContext.BSDF), a2, BSDFContext.SatNoV, BSDFContext.SatNoL, BSDFContext.Context.VoH, BSDFContext.Context.NoH);
 			Sample.PDF			= BSDFContext.SatNoV / PI;
@@ -347,38 +346,24 @@ FStrataEvaluateResult StrataEvaluateBSDF(FStrataBSDFContext BSDFContext)
 			break;
 		}
 		case STRATA_BSDF_TYPE_DIELECTRIC:
-		{
-			float F0 = DIELECTRIC_REFLECTIVITY(BSDFContext.BSDF);
-
-			const float SafeRoughness = MakeRoughnessSafe(DIELECTRIC_ROUGHNESS0(BSDFContext.BSDF));
-			float a2 = Pow4(SafeRoughness);
-			float Energy = 1.0f; // See EnergyNormalization
-
-			// Generalized microfacet specular
-			float D = D_GGX(a2, BSDFContext.Context.NoH) * Energy;
-			float Vis = Vis_SmithJointApprox(a2, BSDFContext.Context.NoV, BSDFContext.SatNoL);
-			float3 F = F_Schlick(F0, BSDFContext.Context.VoH);
-
-			Sample.Value		= D * Vis * F * DIELECTRIC_TINT(BSDFContext.BSDF);
-			Sample.PDF			= VisibleGGXPDF(BSDFContext.TangentV, BSDFContext.TangentH, a2);
-			Sample.Throughput	= (1.0f - F); // STRATA_TODO (Kulla's (1-FDG))
-			break;
-		}
 		case STRATA_BSDF_TYPE_CONDUCTOR:
 		{
-			// See SpecularGGX, isotropic verion
-			const float SafeRoughness = MakeRoughnessSafe(CONDUCTOR_ROUGHNESS0(BSDFContext.BSDF));
+			const bool bIsDieletric = BSDFContext.BSDF.BSDFType == STRATA_BSDF_TYPE_DIELECTRIC;
+			const float3 F0  = bIsDieletric ? DIELECTRIC_TINT(BSDFContext.BSDF) * DIELECTRIC_REFLECTIVITY(BSDFContext.BSDF) : CONDUCTOR_REFLECTIVITY(BSDFContext.BSDF);
+			const float3 F90 = bIsDieletric ? DIELECTRIC_TINT(BSDFContext.BSDF) : CONDUCTOR_EDGECOLOR(BSDFContext.BSDF);
+			const float SafeRoughness = MakeRoughnessSafe(bIsDieletric ? DIELECTRIC_ROUGHNESS0(BSDFContext.BSDF) : CONDUCTOR_ROUGHNESS0(BSDFContext.BSDF));
+
 			float a2 = Pow4(SafeRoughness);
 			float Energy = 1.0f; // See EnergyNormalization
 
 			// Generalized microfacet specular
 			float D = D_GGX(a2, BSDFContext.Context.NoH) * Energy;
 			float Vis = Vis_SmithJointApprox(a2, BSDFContext.Context.NoV, BSDFContext.SatNoL);
-			float3 F = F_Schlick(CONDUCTOR_REFLECTIVITY(BSDFContext.BSDF), CONDUCTOR_EDGECOLOR(BSDFContext.BSDF), BSDFContext.Context.VoH);
+			float3 F = F_Schlick(F0, F90, BSDFContext.Context.VoH);
 
 			Sample.Value		= D * Vis * F;
 			Sample.PDF			= VisibleGGXPDF(BSDFContext.TangentV, BSDFContext.TangentH, a2);
-			Sample.Throughput	= OpaqueBSDFThroughput;
+			Sample.Throughput	= bIsDieletric ? (1.0f - F) : OpaqueBSDFThroughput; // STRATA_TODO (Kulla's (1-FDG))
 			break;
 		}
 		case STRATA_BSDF_TYPE_VOLUME:
@@ -452,17 +437,11 @@ FStrataImportanceSampleResult StrataImportanceSampleBSDF(FStrataBSDF BSDF, float
 			break;
 		}
 		case STRATA_BSDF_TYPE_DIELECTRIC:
-		{
-			float4 GGXSample = ImportanceSampleVisibleGGX(UniformSampleDisk(E), Pow4(DIELECTRIC_ROUGHNESS0(BSDFContext.BSDF)), BSDFContext.TangentV);
-			float3 WorldH = mul(GGXSample.xyz, BSDFContext.TangentBasis);
-			Sample.Direction = reflect(CameraVector, WorldH);
-			PDF = GGXSample.w;
-			Sample.ConeAngle = 1.0f / PDF;
-			break;
-		}
 		case STRATA_BSDF_TYPE_CONDUCTOR:
 		{
-			float4 GGXSample = ImportanceSampleVisibleGGX(UniformSampleDisk(E), Pow4(CONDUCTOR_ROUGHNESS0(BSDFContext.BSDF)), BSDFContext.TangentV);
+			const bool bIsDielectric = BSDFContext.BSDF.BSDFType == STRATA_BSDF_TYPE_DIELECTRIC;
+			float Roughness = bIsDielectric ? DIELECTRIC_ROUGHNESS0(BSDFContext.BSDF) : CONDUCTOR_ROUGHNESS0(BSDFContext.BSDF);
+			float4 GGXSample = ImportanceSampleVisibleGGX(UniformSampleDisk(E), Pow4(Roughness), BSDFContext.TangentV);
 			float3 WorldH = mul(GGXSample.xyz, BSDFContext.TangentBasis);
 			Sample.Direction = reflect(CameraVector, WorldH);
 			PDF = GGXSample.w;
@@ -510,30 +489,17 @@ FStrataEnvLightResult StrataEvaluateForEnvLight(FStrataBSDFContext BSDFContext,
 			break;
 		}
 		case STRATA_BSDF_TYPE_DIELECTRIC:
-		{
-			BRANCH
-			if (bEnableSpecular)
-			{
-				float F0 = DIELECTRIC_REFLECTIVITY(BSDFContext.BSDF);
-				const float SafeRoughness = MakeRoughnessSafe(DIELECTRIC_ROUGHNESS0(BSDFContext.BSDF));
-
-				float3 EvalEnvBRDF = EnvBRDFApprox(F0, SafeRoughness, BSDFContext.SatNoV);
-
-				StrataEnvLightResult.SpecularDirection = StrataGetOffSpecularPeakReflectionDir(BSDFContext.N, BSDFContext.R, SafeRoughness);
-				StrataEnvLightResult.SpecularWeight = EvalEnvBRDF * DIELECTRIC_TINT(BSDFContext.BSDF);
-				StrataEnvLightResult.SpecularSafeRoughness = SafeRoughness;
-			}
-			break;
-		}
 		case STRATA_BSDF_TYPE_CONDUCTOR:
 		{
 			BRANCH
 			if (bEnableSpecular)
 			{
-				// See SpecularGGX, isotropic verion
+				const bool bIsDielectric = STRATA_BSDF_TYPE_DIELECTRIC == BSDFContext.BSDF.BSDFType;
+				const float3 F0  = bIsDielectric ? DIELECTRIC_TINT(BSDFContext.BSDF) * DIELECTRIC_REFLECTIVITY(BSDFContext.BSDF) : CONDUCTOR_REFLECTIVITY(BSDFContext.BSDF);
+				const float3 F90 = bIsDielectric ? DIELECTRIC_TINT(BSDFContext.BSDF) : CONDUCTOR_EDGECOLOR(BSDFContext.BSDF);
+				const float SafeRoughness = MakeRoughnessSafe(bIsDielectric ? DIELECTRIC_ROUGHNESS0(BSDFContext.BSDF) : CONDUCTOR_ROUGHNESS0(BSDFContext.BSDF));
 
-				const float SafeRoughness = MakeRoughnessSafe(CONDUCTOR_ROUGHNESS0(BSDFContext.BSDF));
-				float3 EvalEnvBRDF = EnvBRDFApprox(CONDUCTOR_REFLECTIVITY(BSDFContext.BSDF), CONDUCTOR_EDGECOLOR(BSDFContext.BSDF), SafeRoughness, BSDFContext.SatNoV);
+				float3 EvalEnvBRDF = EnvBRDFApprox(F0, F90, SafeRoughness, BSDFContext.SatNoV);
 
 				StrataEnvLightResult.SpecularDirection = StrataGetOffSpecularPeakReflectionDir(BSDFContext.N, BSDFContext.R, SafeRoughness);
 				StrataEnvLightResult.SpecularWeight = EvalEnvBRDF;
