Fix shader warning on hair bsdf

diff --git a/Engine/Shaders/Private/HairBsdf.ush b/Engine/Shaders/Private/HairBsdf.ush
--- a/Engine/Shaders/Private/HairBsdf.ush
+++ b/Engine/Shaders/Private/HairBsdf.ush
@@ -33,8 +33,8 @@ float3 KajiyaKayDiffuseAttenuation(FGBufferData GBuffer, float3 L, float3 V, hal
 	float NoL = saturate((dot(N, L) + Wrap) / Square(1 + Wrap));
 	float DiffuseScatter = (1 / PI) * lerp(NoL, KajiyaDiffuse, 0.33) * GBuffer.Metallic;
 	float Luma = Luminance(GBuffer.BaseColor);
-	float3 ScatterTint = pow(GBuffer.BaseColor / Luma, 1 - Shadow);
-	return sqrt(GBuffer.BaseColor) * DiffuseScatter * ScatterTint;
+	float3 ScatterTint = pow(abs(GBuffer.BaseColor / Luma), 1 - Shadow);
+	return sqrt(abs(GBuffer.BaseColor)) * DiffuseScatter * ScatterTint;
 }
 
 float3 EvaluateHairMultipleScattering(
@@ -349,7 +349,7 @@ float3 HairShading( FGBufferData GBuffer, float3 L, float3 V, half3 N, float Sha
 		float3 Tp = 0;
 		if (HairTransmittance.bUseLegacyAbsorption)
 		{
-			Tp = pow(GBuffer.BaseColor, 0.5 * sqrt(1 - Pow2(h * a)) / CosThetaD);
+			Tp = pow(abs(GBuffer.BaseColor), 0.5 * sqrt(1 - Pow2(h * a)) / CosThetaD);
 		}
 		else
 		{
@@ -378,7 +378,7 @@ float3 HairShading( FGBufferData GBuffer, float3 L, float3 V, half3 N, float Sha
 		float f = Hair_F( CosThetaD * 0.5 );
 		float Fp = Pow2(1 - f) * f;
 		//float3 Tp = pow( GBuffer.BaseColor, 1.6 / CosThetaD );
-		float3 Tp = pow( GBuffer.BaseColor, 0.8 / CosThetaD );
+		float3 Tp = pow(abs(GBuffer.BaseColor), 0.8 / CosThetaD );
 
 		//float s = 0.15;
 		//float Np = 0.75 * exp( Phi / s ) / ( s * Pow2( 1 + exp( Phi / s ) ) );
