Add area lightin support for strata material.

diff --git a/Engine/Shaders/Private/AreaLightCommon.ush b/Engine/Shaders/Private/AreaLightCommon.ush
new file mode 100644
--- /dev/null
+++ b/Engine/Shaders/Private/AreaLightCommon.ush
@@ -0,0 +1,53 @@
+// Copyright Epic Games, Inc. All Rights Reserved.
+
+#pragma once
+
+#include "RectLight.ush"
+
+struct FAreaLight
+{
+	float		SphereSinAlpha;
+	float		SphereSinAlphaSoft;
+	float		LineCosSubtended;
+
+	float3		FalloffColor;
+
+	FRect		Rect;
+	FRectTexture Texture;
+	bool		bIsRect;
+};
+
+struct FAreaLightIntegrateContext
+{
+	FAreaLight AreaLight;
+	float3 L;
+	float NoL;
+	float Falloff;
+};
+
+bool IsAreaLight(FAreaLight AreaLight)
+{
+	return (
+		AreaLight.bIsRect ||
+		AreaLight.SphereSinAlpha > 0.0f ||
+		AreaLight.SphereSinAlphaSoft > 0.0f ||
+		AreaLight.LineCosSubtended < 1.0f
+		) ? true : false;
+}
+
+FAreaLightIntegrateContext InitAreaLightIntegrateContext()
+{
+	// Manual initialization of the area light as the compiler is unhappy otherwise
+	FAreaLightIntegrateContext Out;
+	Out.AreaLight.SphereSinAlpha = 0;
+	Out.AreaLight.SphereSinAlphaSoft = 0;
+	Out.AreaLight.LineCosSubtended = 0;
+	Out.AreaLight.FalloffColor = 0;
+	Out.AreaLight.Rect = (FRect)0;
+	Out.AreaLight.bIsRect = false;
+//	Out.AreaLight.Texture;
+	Out.L = 0;
+	Out.NoL = 0;
+	Out.Falloff = 0;
+	return Out;
+}
diff --git a/Engine/Shaders/Private/CapsuleLightIntegrate.ush b/Engine/Shaders/Private/CapsuleLightIntegrate.ush
--- a/Engine/Shaders/Private/CapsuleLightIntegrate.ush
+++ b/Engine/Shaders/Private/CapsuleLightIntegrate.ush
@@ -4,6 +4,7 @@
 
 #include "DeferredShadingCommon.ush"
 #include "MonteCarlo.ush"
+#include "AreaLightCommon.ush"
 #include "ShadingModels.ush"
 #include "CapsuleLight.ush"
 #include "SobolRandom.ush"
@@ -34,8 +35,10 @@ float IntegrateLight( FCapsuleLight Capsule, bool bInverseSquared )
 	return Falloff;
 }
 
-FDirectLighting IntegrateBxDF( FGBufferData GBuffer, half3 N, half3 V, FCapsuleLight Capsule, FShadowTerms Shadow, bool bInverseSquared )
+FAreaLightIntegrateContext CreateCapsuleIntegrateContext(float Roughness, half3 N, half3 V, FCapsuleLight Capsule, bool bInverseSquared )
 {
+	FAreaLightIntegrateContext Out = InitAreaLightIntegrateContext();
+
 	float NoL;
 	float Falloff;
 	float LineCosSubtended = 1;
@@ -90,19 +93,28 @@ FDirectLighting IntegrateBxDF( FGBufferData GBuffer, half3 N, half3 V, FCapsuleL
 	float InvDist = rsqrt( DistSqr );
 	float3 L = ToLight * InvDist;
 	
-	GBuffer.Roughness = max( GBuffer.Roughness, View.MinRoughness );
-	float a = Pow2( GBuffer.Roughness );
+	Roughness = max( Roughness, View.MinRoughness );
+	float a = Pow2( Roughness );
     
-	FAreaLight AreaLight;
-	AreaLight.SphereSinAlpha = saturate( Capsule.Radius * InvDist * (1 - a) );
-	AreaLight.SphereSinAlphaSoft = saturate( Capsule.SoftRadius * InvDist );
-	AreaLight.LineCosSubtended = LineCosSubtended;
-	AreaLight.FalloffColor = 1;
-	AreaLight.Rect = (FRect)0;
-	AreaLight.bIsRect = false;
-	AreaLight.Texture = InitRectTexture(DummyRectLightTextureForCapsuleCompilerWarning);
-
-	return IntegrateBxDF( GBuffer, N, V, L, Falloff, NoL, AreaLight, Shadow );
+	Out.AreaLight.SphereSinAlpha = saturate( Capsule.Radius * InvDist * (1 - a) );
+	Out.AreaLight.SphereSinAlphaSoft = saturate( Capsule.SoftRadius * InvDist );
+	Out.AreaLight.LineCosSubtended = LineCosSubtended;
+	Out.AreaLight.FalloffColor = 1;
+	Out.AreaLight.Rect = (FRect)0;
+	Out.AreaLight.bIsRect = false;
+	Out.AreaLight.Texture = InitRectTexture(DummyRectLightTextureForCapsuleCompilerWarning);
+	Out.NoL = NoL;
+	Out.Falloff = Falloff;	
+	Out.L = L;
+	return Out;
+}
+
+
+FDirectLighting IntegrateBxDF( FGBufferData GBuffer, half3 N, half3 V, FCapsuleLight Capsule, FShadowTerms Shadow, bool bInverseSquared )
+{
+	GBuffer.Roughness = max( GBuffer.Roughness, View.MinRoughness );
+	FAreaLightIntegrateContext Context = CreateCapsuleIntegrateContext(GBuffer.Roughness, N, V, Capsule, bInverseSquared);
+	return IntegrateBxDF( GBuffer, N, V, Context.L, Context.Falloff, Context.NoL, Context.AreaLight, Shadow );
 }
 
 FDirectLighting IntegrateBxDF( FGBufferData GBuffer, half3 N, half3 V, FCapsuleLight Capsule, FShadowTerms Shadow, uint2 SVPos )
diff --git a/Engine/Shaders/Private/RectLight.ush b/Engine/Shaders/Private/RectLight.ush
--- a/Engine/Shaders/Private/RectLight.ush
+++ b/Engine/Shaders/Private/RectLight.ush
@@ -718,3 +718,8 @@ FRect GetRect(
 	return Rect;
 }
 
+bool IsRectVisible(FRect Rect)
+{
+	// No-visible rect light due to barn door occlusion
+	return Rect.Extent.x != 0 && Rect.Extent.y != 0;
+}
\ No newline at end of file
diff --git a/Engine/Shaders/Private/RectLightIntegrate.ush b/Engine/Shaders/Private/RectLightIntegrate.ush
--- a/Engine/Shaders/Private/RectLightIntegrate.ush
+++ b/Engine/Shaders/Private/RectLightIntegrate.ush
@@ -4,6 +4,7 @@
 
 #include "DeferredShadingCommon.ush"
 #include "MonteCarlo.ush"
+#include "AreaLightCommon.ush"
 #include "ShadingModels.ush"
 #include "RectLight.ush"
 
@@ -68,14 +69,12 @@ float IntegrateLight( FRect Rect, FRectTexture SourceTexture )
 	return Falloff;
 }
 
-FDirectLighting IntegrateBxDF( FGBufferData GBuffer, half3 N, half3 V, FRect Rect, FShadowTerms Shadow, FRectTexture SourceTexture )
+FAreaLightIntegrateContext CreateRectIntegrateContext( float Roughness, half3 N, half3 V, FRect Rect, FRectTexture SourceTexture )
 {
-	float NoL;
-	float Falloff;
+	float NoL = 0;
+	float Falloff = 0;
 	
-	// No-visible rect light due to barn door occlusion
-	if (Rect.Extent.x == 0 || Rect.Extent.y == 0) 
-        return (FDirectLighting)0;
+	FAreaLightIntegrateContext Out = InitAreaLightIntegrateContext();
 
 #if 0	// Reference quality
 	float3 L = RectIrradianceRef( N, Rect, Falloff, NoL );
@@ -226,28 +225,44 @@ FDirectLighting IntegrateBxDF( FGBufferData GBuffer, half3 N, half3 V, FRect Rec
 	}
 #endif
 	
-	GBuffer.Roughness = max( GBuffer.Roughness, 0.02 );
-
-	FAreaLight AreaLight;
-	AreaLight.SphereSinAlpha = sqrt( Falloff * (1.0 / PI) );
-	AreaLight.SphereSinAlphaSoft = 0;
-	AreaLight.LineCosSubtended = 1;
-	AreaLight.FalloffColor = 1;
-	AreaLight.Rect = Rect;
-	AreaLight.Texture = SourceTexture;
-	AreaLight.bIsRect = false;
+	Out.AreaLight.SphereSinAlpha = sqrt( Falloff * (1.0 / PI) );
+	Out.AreaLight.SphereSinAlphaSoft = 0;
+	Out.AreaLight.LineCosSubtended = 1;
+	Out.AreaLight.FalloffColor = 1;
+	Out.AreaLight.Rect = Rect;
+	Out.AreaLight.Texture = SourceTexture;
+	Out.AreaLight.bIsRect = false;
+	Out.L = L;
+	Out.NoL = NoL;
+	Out.Falloff = Falloff;
 #else
-	FAreaLight AreaLight;
-	AreaLight.SphereSinAlpha = 0;
-	AreaLight.SphereSinAlphaSoft = 0;
-	AreaLight.LineCosSubtended = 1;
-	AreaLight.FalloffColor = FalloffColor;
-	AreaLight.Rect = Rect;
-	AreaLight.Texture = SourceTexture;
-	AreaLight.bIsRect = true;
+	Out.AreaLight.SphereSinAlpha = 0;
+	Out.AreaLight.SphereSinAlphaSoft = 0;
+	Out.AreaLight.LineCosSubtended = 1;
+	Out.AreaLight.FalloffColor = FalloffColor;
+	Out.AreaLight.Rect = Rect;
+	Out.AreaLight.Texture = SourceTexture;
+	Out.AreaLight.bIsRect = true;
+	Out.L = L;
+	Out.NoL = NoL;
+	Out.Falloff = Falloff;
 #endif
+	return Out;
+}
 
-	return IntegrateBxDF( GBuffer, N, V, L, Falloff, NoL, AreaLight, Shadow );
+FDirectLighting IntegrateBxDF(FGBufferData GBuffer, half3 N, half3 V, FRect Rect, FShadowTerms Shadow, FRectTexture SourceTexture)
+{
+	// No-visible rect light due to barn door occlusion
+	FDirectLighting Out = (FDirectLighting)0;
+	if (IsRectVisible(Rect))
+	{
+		FAreaLightIntegrateContext Context = CreateRectIntegrateContext(GBuffer.Roughness, N, V, Rect, SourceTexture);
+	#if CHEAP_RECT
+		GBuffer.Roughness = max(GBuffer.Roughness, 0.02);
+	#endif
+		Out = IntegrateBxDF(GBuffer, N, V, Context.L, Context.Falloff, Context.NoL, Context.AreaLight, Shadow);
+	}
+	return Out;
 }
 
 FDirectLighting IntegrateBxDF( FGBufferData GBuffer, half3 N, half3 V, FRect Rect, FShadowTerms Shadow, FRectTexture SourceTexture, uint2 SVPos )
@@ -262,7 +277,7 @@ FDirectLighting IntegrateBxDF( FGBufferData GBuffer, half3 N, half3 V, FRect Rec
 		return Lighting;
 	
 	// No-visible rect light due to barn door occlusion
-	if (Rect.Extent.x == 0 || Rect.Extent.y == 0) 
+	if (!IsRectVisible(Rect))
 		return Lighting;
 
 	FSphericalRect SphericalRect = BuildSphericalRect( Rect );
diff --git a/Engine/Shaders/Private/ShadingModels.ush b/Engine/Shaders/Private/ShadingModels.ush
--- a/Engine/Shaders/Private/ShadingModels.ush
+++ b/Engine/Shaders/Private/ShadingModels.ush
@@ -7,6 +7,7 @@
 #include "FastMath.ush"
 #include "CapsuleLight.ush"
 #include "RectLight.ush"
+#include "AreaLightCommon.ush"
 #include "TransmissionCommon.ush"
 
 #if 0
@@ -84,19 +85,6 @@ FHairTransmittanceData InitHairStrandsTransmittanceData(bool bMultipleScatterEna
 	return o;
 }
 
-struct FAreaLight
-{
-	float		SphereSinAlpha;
-	float		SphereSinAlphaSoft;
-	float		LineCosSubtended;
-
-	float3		FalloffColor;
-
-	FRect		Rect;
-	FRectTexture Texture;
-	bool		bIsRect;
-};
-
 struct FDirectLighting
 {
 	float3	Diffuse;
@@ -114,16 +102,6 @@ struct FShadowTerms
 
 #include "HairBsdf.ush"
 
-bool IsAreaLight(FAreaLight AreaLight)
-{
-	return (
-		AreaLight.bIsRect ||
-		AreaLight.SphereSinAlpha > 0.0f ||
-		AreaLight.SphereSinAlphaSoft > 0.0f ||
-		AreaLight.LineCosSubtended < 1.0f
-		) ? true : false;
-}
-
 float New_a2( float a2, float SinAlpha, float VoH )
 {
 	return a2 + 0.25 * SinAlpha * (3.0 * sqrtFast(a2) + SinAlpha) / ( VoH + 0.001 );
diff --git a/Engine/Shaders/Private/Strata/Strata.ush b/Engine/Shaders/Private/Strata/Strata.ush
--- a/Engine/Shaders/Private/Strata/Strata.ush
+++ b/Engine/Shaders/Private/Strata/Strata.ush
@@ -245,6 +245,9 @@ struct FStrataBSDFContext
 	float3x3 TangentBasis;
 	float3 TangentV;
 	float3 TangentH;
+
+	// Preprocessed roughness override used for specular BSDF evaluation when integration area light (using MRP integrator)
+	float a2Override;
 };
 
 void StrataGetBSDFNormal(in FStrataBSDF BSDF, inout float3 N)
@@ -275,6 +278,43 @@ void StrataGetBSDFNormal(in FStrataBSDF BSDF, inout float3 N)
 	}
 }
 
+void StrataGetBSDFRoughness(in FStrataBSDF BSDF, inout float R)
+{
+	R = 0;
+	switch (BSDF.BSDFType)
+	{
+	case STRATA_BSDF_TYPE_DIFFUSE_ON:
+	case STRATA_BSDF_TYPE_DIFFUSE_CHAN:
+	{
+		R = DIFFUSE_ROUGHNESS(BSDF);
+		break;
+	}
+	case STRATA_BSDF_TYPE_DIELECTRIC:
+	{
+		R = DIELECTRIC_ROUGHNESS0(BSDF);
+		break;
+	}
+	case STRATA_BSDF_TYPE_CONDUCTOR:
+	{
+		R = CONDUCTOR_ROUGHNESS0(BSDF);
+		break;
+	}
+	case STRATA_BSDF_TYPE_VOLUME:
+	{
+		R = 0;
+		break;
+	}
+	}
+}
+
+// Return true if a BSDF has a custom area light integrator
+bool StrataHasAreaLightIntegrator(in FStrataBSDF BSDF)
+{
+	return 
+		BSDF.BSDFType == STRATA_BSDF_TYPE_CONDUCTOR ||
+		BSDF.BSDFType == STRATA_BSDF_TYPE_DIELECTRIC;
+}
+
 FStrataBSDFContext StrataCreateBSDFContext(FStrataBSDF BSDF, float3 V, float3 L)
 {
 	FStrataBSDFContext BSDFContext = (FStrataBSDFContext)0;
@@ -353,11 +393,10 @@ FStrataEvaluateResult StrataEvaluateBSDF(FStrataBSDFContext BSDFContext)
 			const float3 F90 = bIsDieletric ? DIELECTRIC_TINT(BSDFContext.BSDF) : CONDUCTOR_EDGECOLOR(BSDFContext.BSDF);
 			const float SafeRoughness = MakeRoughnessSafe(bIsDieletric ? DIELECTRIC_ROUGHNESS0(BSDFContext.BSDF) : CONDUCTOR_ROUGHNESS0(BSDFContext.BSDF));
 
-			float a2 = Pow4(SafeRoughness);
-			float Energy = 1.0f; // See EnergyNormalization
+			const float a2 = BSDFContext.a2Override>0 ? BSDFContext.a2Override : Pow4(SafeRoughness);
 
 			// Generalized microfacet specular
-			float D = D_GGX(a2, BSDFContext.Context.NoH) * Energy;
+			float D = D_GGX(a2, BSDFContext.Context.NoH);
 			float Vis = Vis_SmithJointApprox(a2, BSDFContext.Context.NoV, BSDFContext.SatNoL);
 			float3 F = F_Schlick(F0, F90, BSDFContext.Context.VoH);
 
diff --git a/Engine/Shaders/Private/Strata/StrataDeferredLighting.ush b/Engine/Shaders/Private/Strata/StrataDeferredLighting.ush
--- a/Engine/Shaders/Private/Strata/StrataDeferredLighting.ush
+++ b/Engine/Shaders/Private/Strata/StrataDeferredLighting.ush
@@ -2,6 +2,115 @@
 
 #pragma once
 
+#include "..\AreaLightCommon.ush"
+
+struct FStrataLightIntegrationContext
+{
+	float3 IncidentRadiance;
+	float3 IncidentIrradiance;	// Modified incident irradiance, based on custom are light integrator
+	float3 OutgoingRadiance;	// Continue the outgoing radiance if the BSDF evaluation/integration has been done with the custom area light integrator
+	bool bEvaluateBSDF;			// True if the BSDF need to be evalued, false, if the integration has already been done
+};
+
+FStrataLightIntegrationContext StrataIniLightIntegrationContext()
+{
+	FStrataLightIntegrationContext Out;
+	Out.bEvaluateBSDF = true;
+	Out.OutgoingRadiance = 0;
+	Out.IncidentIrradiance = 0;
+	Out.IncidentRadiance = 0;
+	return Out;
+}
+
+// Preprocess area lights and handle custom BSDF integrator
+// This computes the full area light x BSDF or precompute incident lighting & reweigth the BSDF parameters depending on the custom integrator type (LTC/MRP)
+void StrataPreprocessAreaLightIntegration(
+	in float3 ToLight,
+	in FDeferredLightData LightData, 
+	in FRectTexture RectTexture,
+	inout FStrataBSDFContext BSDFContext,
+	inout FStrataLightIntegrationContext Out)
+{
+	Out.bEvaluateBSDF = true;
+	Out.OutgoingRadiance = 0;
+
+	// Special area lighting integration
+	bool bIsValid = true;
+	float Roughness = 0;
+	StrataGetBSDFRoughness(BSDFContext.BSDF, Roughness);
+	FAreaLightIntegrateContext LightContext;
+	if (LightData.bRectLight)
+	{
+		FRect Rect = GetRect(ToLight, LightData);
+		bIsValid = IsRectVisible(Rect);
+		Out.bEvaluateBSDF = false;
+		// No-visible rect light due to barn door occlusion
+		LightContext = CreateRectIntegrateContext(Roughness, BSDFContext.N, BSDFContext.V, Rect, RectTexture);
+	}
+	else
+	{
+		FCapsuleLight Capsule = GetCapsule(ToLight, LightData);
+		LightContext = CreateCapsuleIntegrateContext(Roughness, BSDFContext.N, BSDFContext.V, Capsule, LightData.bInverseSquared);
+	}
+
+	// Special area light integrator for BSDF (opt-in)
+	// bEvaluateBSDF: Indicates if the BSDF need to be evaluate at the end, or if the custom area light integration handles the BSDF evaluation
+	if (bIsValid && IsAreaLight(LightContext.AreaLight) && StrataHasAreaLightIntegrator(BSDFContext.BSDF))
+	{
+		// Evaluate roughness / specular color
+		float SafeRoughness = 0;
+		float3 SpecularColor = 0;
+		switch (BSDFContext.BSDF.BSDFType)
+		{
+			case STRATA_BSDF_TYPE_DIELECTRIC:
+			{
+				SafeRoughness = MakeRoughnessSafe(DIELECTRIC_ROUGHNESS0(BSDFContext.BSDF));
+				SpecularColor = DielectricIorToF0(DIELECTRIC_REFLECTIVITY(BSDFContext.BSDF)) * DIELECTRIC_TINT(BSDFContext.BSDF);
+			}
+			break;
+			case STRATA_BSDF_TYPE_CONDUCTOR:
+			{
+				SafeRoughness = MakeRoughnessSafe(CONDUCTOR_ROUGHNESS0(BSDFContext.BSDF));
+				SpecularColor = CONDUCTOR_REFLECTIVITY(BSDFContext.BSDF);
+			}
+			break;
+		}		
+		
+		// Evaluate area-light preintegration
+		SafeRoughness = max(SafeRoughness, View.MinRoughness);
+		if (LightContext.AreaLight.bIsRect)
+		{
+			// Rect light uses LTC integrator, which computes the full BDSD x Area light integration. 
+			// This is why there is no need to evaluate the BSDF. The result is stored into a local 
+			// OutgoingRadiance, which will be accumulated outside of this function
+			Out.OutgoingRadiance.rgb += 
+				Out.IncidentRadiance * 
+				RectGGXApproxLTC(SafeRoughness, SpecularColor, BSDFContext.N, BSDFContext.V, LightContext.AreaLight.Rect, LightContext.AreaLight.Texture);
+			Out.bEvaluateBSDF = false;
+		}
+		else
+		{
+			// Sphere/Tube/Disk light use MRP integrator, which:
+			// 1. compute an energy rescaling factor (this is fold into the incident radiance)
+			// 2. change the BSDF roughness
+			// 3. compute a derived L 
+			// Update the BSDF context with the new L (and derived data) with value coming from MRP integrator
+			StrataUpdateBSDFContext(BSDFContext, LightContext.L);
+			SphereMaxNoH(BSDFContext.Context, LightContext.AreaLight.SphereSinAlpha, true);
+			BSDFContext.Context.NoV = saturate(abs(BSDFContext.Context.NoV) + 1e-5);
+
+			// Update incident irradiance to add the energy rescaling to due MRP integration
+			// The a2 override stored the overriden roughness, which will be used during the BSDF evaluation
+			BSDFContext.a2Override = Pow4(SafeRoughness);
+			Out.IncidentIrradiance =
+				Out.IncidentRadiance *
+				EnergyNormalization(BSDFContext.a2Override, BSDFContext.Context.VoH, LightContext.AreaLight) *
+				LightContext.AreaLight.FalloffColor * (LightContext.Falloff * LightContext.NoL);
+			Out.bEvaluateBSDF = true;
+		}
+	}
+}
+
 // Analytical lighting evaluation for strata material.
 // Unpack BSDF on-the-fly
 float4 StrataDeferredLighting(
@@ -16,40 +125,46 @@ float4 StrataDeferredLighting(
 	ByteAddressBuffer StrataDataBuffer,
 	uint PixelStrataDataByteOffset)
 {
-	float4 Radiance = 0;
+	float4 OutRadiance = 0;
 
-
-	float3 UglyIntegrationFactor = 0; // This is bad we need to make sure it is done the right way
+	FStrataPixelHeader Header = UnpackStrataHeaderIn(StrataDataBuffer, PixelStrataDataByteOffset);
+	
+	float3 PreintegratedIndicentIrradiance = 0;
 	if (LightData.bRectLight)
 	{
 		FRect Rect = GetRect(ToLight, LightData);
-		UglyIntegrationFactor = IntegrateLight(Rect, RectTexture);
+		PreintegratedIndicentIrradiance = IntegrateLight(Rect, RectTexture);
 	}
 	else
 	{
 		FCapsuleLight Capsule = GetCapsule(ToLight, LightData);
-		UglyIntegrationFactor = IntegrateLight(Capsule, LightData.bInverseSquared);
+		PreintegratedIndicentIrradiance = IntegrateLight(Capsule, LightData.bInverseSquared);
 	}
 
-	FStrataPixelHeader Header = UnpackStrataHeaderIn(StrataDataBuffer, PixelStrataDataByteOffset);
-	
 	LOOP
 	for (uint BSDFIndex = 0; BSDFIndex < Header.BSDFCount; ++BSDFIndex)
 	{
+		// Unpack BSDF data
 		const FStrataBSDF BSDF = UnpackStrataBSDFIn(StrataDataBuffer, PixelStrataDataByteOffset);
-		const float3 BSDFThroughput = BSDF.Weight;
-
-		// Create the BSDF context
-		FStrataBSDFContext StrataBSDFContext = StrataCreateBSDFContext(BSDF, V, L);
+		FStrataBSDFContext BSDFContext = StrataCreateBSDFContext(BSDF, V, L);
 
-		// Compute some common factors
-		const float3 DirLightBaseFactor = BSDFThroughput * LightData.Color * LightMask * SurfaceShadow * UglyIntegrationFactor;
-		const float3 DirLightGeometryFactor = DirLightBaseFactor * StrataBSDFContext.SatNoL;
+		// Evaluate incident lighting & (optionally) precompute area light integration
+		FStrataLightIntegrationContext IntegrationContext = StrataIniLightIntegrationContext();
+		IntegrationContext.IncidentRadiance = LightData.Color * LightMask * SurfaceShadow;
+		IntegrationContext.IncidentIrradiance = PreintegratedIndicentIrradiance * BSDFContext.SatNoL * IntegrationContext.IncidentRadiance;
+		StrataPreprocessAreaLightIntegration(ToLight, LightData, RectTexture, BSDFContext, IntegrationContext);
 
-		// Evalaute the lighting coming from L
-		FStrataEvaluateResult StrataEvaluate = StrataEvaluateBSDF(StrataBSDFContext);
-		Radiance.rgb += StrataEvaluate.Value * DirLightGeometryFactor;
+		// Evaluate BSDF
+		if (IntegrationContext.bEvaluateBSDF)
+		{
+			FStrataEvaluateResult BSDFEvaluate = StrataEvaluateBSDF(BSDFContext);
+			OutRadiance.rgb += BSDF.Weight * BSDFEvaluate.Value * IntegrationContext.IncidentIrradiance;
+		}
+		else
+		{
+			OutRadiance.rgb += IntegrationContext.OutgoingRadiance;
+		}
 	}
 
-	return Radiance;
-}
+	return OutRadiance;
+}
\ No newline at end of file
