Strata BSDF normals can now be shared between BSDFs.
Normals are now compressed as octahedral on 10bits x2
Wrost case per pixel byte size is now 84 instead of 104.
fixed AreaLightCommon.ush include

diff --git a/Engine/Shaders/Private/BasePassPixelShader.usf b/Engine/Shaders/Private/BasePassPixelShader.usf
--- a/Engine/Shaders/Private/BasePassPixelShader.usf
+++ b/Engine/Shaders/Private/BasePassPixelShader.usf
@@ -1479,13 +1479,23 @@ void FPixelShaderInOut_MainPS(
 	uint PixelStrataDataByteOffset = GetStrataPixelDataByteOffset(PixelPos, ResolvedView.BufferSizeAndInvSize.x, OpaqueBasePass.Strata.MaxBytesPerPixel);
 
 	const float3 VertexNormal = MaterialParameters.TangentToWorld[2].xyz;
-	PackStrataOut(OpaqueBasePass.Strata.MaterialLobesBufferUAV, PixelStrataDataByteOffset, Strata, VertexNormal, MaterialParameters.CameraVector, GBuffer);
+	PackStrataOut(
+		OpaqueBasePass.Strata.MaterialLobesBufferUAV, PixelStrataDataByteOffset, 
+		Strata, MaterialParameters.SharedNormals, 
+		VertexNormal, MaterialParameters.CameraVector, GBuffer);
 #endif // STRATA_OPAQUE_ENABLED
 
 #if STRATA_TRANSLUCENT_ENABLED
 	//FORWARD_SHADING
 	if (Strata.LayerCount > 0)
 	{
+		FStrataPixelHeader Header = (FStrataPixelHeader)0;
+		for (int l = 0; l < Strata.LayerCount; ++l)
+		{
+			Header.BSDFCount += Strata.Layers[l].BSDFCount;
+		}
+		Header.SharedNormals = MaterialParameters.SharedNormals;
+
 		float3 L = ResolvedView.DirectionalLightDirection;
 		float3 V = MaterialParameters.CameraVector;
 		float3 LightColor = ResolvedView.DirectionalLightColor.rgb * PI;
@@ -1499,6 +1509,7 @@ void FPixelShaderInOut_MainPS(
 			LightColor,
 			DirectionalLightShadow,
 			IndirectOcclusion,
+			Header,
 			Strata,
 			Throughput);
 
diff --git a/Engine/Shaders/Private/MaterialTemplate.ush b/Engine/Shaders/Private/MaterialTemplate.ush
--- a/Engine/Shaders/Private/MaterialTemplate.ush
+++ b/Engine/Shaders/Private/MaterialTemplate.ush
@@ -360,6 +360,10 @@ struct FMaterialPixelParameters
 	float CloudSampleNormAltitudeInLayer;
 	float3 VolumeSampleConservativeDensity;
 #endif
+
+#if PROJECT_STRATA && MATERIAL_IS_STRATA
+	FSharedNormals SharedNormals;
+#endif
 };
 
 /**
@@ -2341,6 +2345,10 @@ void CalcPixelMaterialInputs(in out FMaterialPixelParameters Parameters, in out
 	// The Normal is a special case as it might have its own expressions and also be used to calculate other inputs, so perform the assignment here
 %s
 
+#if PROJECT_STRATA && MATERIAL_IS_STRATA
+	Parameters.SharedNormals = StrataInitialiseSharedNormals();
+#endif
+
 	// Note that here MaterialNormal can be in world space or tangent space
 	float3 MaterialNormal = GetMaterialNormal(Parameters, PixelMaterialInputs);
 
diff --git a/Engine/Shaders/Private/Strata/Strata.ush b/Engine/Shaders/Private/Strata/Strata.ush
--- a/Engine/Shaders/Private/Strata/Strata.ush
+++ b/Engine/Shaders/Private/Strata/Strata.ush
@@ -21,92 +21,49 @@
 //	- fp16 on gen5
 
 
+
+
+
 ///////////////////////////////////////////////////////////////////////////////
 // BSDFs representation
-
-//struct FStrataDiffuse
-//{
-//	float3	Albedo;
-//	float   Roughness;
-//
-//	float3	Normal;
-//};
-//
-//struct FStrataDielectric
-//{
-//	float	IOR;
-//	float3	Tint;
-//	float2	Roughness;
-//
-//	float3	Normal;
-//};
-//
-//struct FStrataConductor
-//{
-//	float3	Reflectivity;
-//	float3	EdgeColor;
-//	float2	Roughness;
-//
-//	float3	Normal;
-//};
-//
-//struct FStrataVolume
-//{
-//	float3	Albedo;
-//	float3	Extinction;
-//	float	Anisotropy;
-//	float	Thickness;
-//
-//	float3  SmoothSurfaceNormal;	// Typically vertex normal
-//};
-
-
-
-// The largest Strata BSDF is FStrataConductor having as size of 24 bytes (FStrataConductor)
-// So we will have 
-//		STRATA_MAX_TOTAL_BSDF * sizeof(FStrataConductor) +				// 4 * 16
-//		STRATA_MAX_TOTAL_BSDF * (sizeof(BSDFType) + sizeof(Weight))		// 4 * 8
-//		sizeof(BSDFCount)												// UINT
-//		= 4 * 16 + 4 * 8 + 1 = 96 + 32 + 1 = 97						// See Strata.cpp
+//		STRATA_MAX_TOTAL_BSDF * sizeof(FStrataConductor) +				// 4 * 8
+//		STRATA_MAX_TOTAL_BSDF * (sizeof(State) + sizeof(Weight))		// 4 * 8
+//		sizeof(NormalCount_BSDFCount)									// 1 * 4
+//		STRATA_MAX_TOTAL_BSDF * sizeof(Normal)							// 4 * 4
+//		= 84															// See Strata.cpp
 struct FStrataBSDF
 {
-	int					BSDFType;
-	float3				Weight;		// Weight is a float3 because, when loaded from the material buffer, it contains the BSDF throughput (coverage and transmittance combined)
+	uint				State;		// (27 unused bits) (2bits as SharedNormalIndex) (3bits as BSDFType)
 
-	// We cannot have the structures above represented as a union in HLSL, so we enforce that using an array representing VGPRs and a bunch of preprocessor functions.
-	//union BSDFsUnion
-	//{
-	//	FStrataDiffuse		Diffuse;
-	//	FStrataDielectric	Dielectric;
-	//	FStrataConductor	Conductor;
-	//	FStrataVolume		Volume;
-	//};
-	//BSDFsUnion U;
+	float3				Weight;		// Weight is a float3 because, when loaded from the material buffer, it contains the BSDF throughput (coverage and transmittance combined)
 
 	float4 VGPRs[3];
 };
 
+#define BSDF_GETTYPE(X)					uint(X.State & 0x7)
+#define BSDF_SETTYPE(X, NewType)		X.State = (X.State & 0xFFF8) | ((NewType & 0x7))
+
+// This assumes Normals[STRATA_MAX_TOTAL_BSDF]
+#define BSDF_GETNORMALID(X)				uint((X.State>>3) & 0x3)
+#define BSDF_SETNORMALID(X, NrmlIdx)	X.State = (X.State & 0xFFE7) | ((NrmlIdx & 0x3) << 3)
+
 #define DIFFUSE_ALBEDO(X)				X.VGPRs[0].xyz
 #define DIFFUSE_ROUGHNESS(X)			X.VGPRs[0].w
-#define DIFFUSE_NORMAL(X)				X.VGPRs[1].xyz
 
 #define DIELECTRIC_TINT(X)				X.VGPRs[0].xyz
 #define DIELECTRIC_IOR(X)				X.VGPRs[0].w
 #define DIELECTRIC_ROUGHNESS0(X)		X.VGPRs[1].w
 #define DIELECTRIC_ROUGHNESS1(X)		X.VGPRs[2].w
-#define DIELECTRIC_NORMAL(X)			X.VGPRs[1].xyz
 
 #define CONDUCTOR_REFLECTIVITY(X)		X.VGPRs[0].xyz
 #define CONDUCTOR_EDGECOLOR(X)			X.VGPRs[1].xyz
 #define CONDUCTOR_ROUGHNESS0(X)			X.VGPRs[2].w
 #define CONDUCTOR_ROUGHNESS1(X)			X.VGPRs[1].w
-#define CONDUCTOR_NORMAL(X)				X.VGPRs[2].xyz
 
 #define VOLUME_ALBEDO(X)				X.VGPRs[0].xyz
 #define VOLUME_EXTINCTION(X)			X.VGPRs[1].xyz
 #define VOLUME_ANISOTROPY(X)			X.VGPRs[0].w
 #define VOLUME_THICKNESS(X)				X.VGPRs[1].w
-#define VOLUME_SMOOTHSURFACENORMAL(X)	X.VGPRs[2].xyz
 
 
 struct FStrataLayer
@@ -130,86 +87,86 @@ FStrataData GetInitialisedStrataData()
 	return StrataData;
 }
 
-FStrataData GetStrataDiffuseOrenNayarBSDF(float3 Albedo, float Roughness, float3 Normal)
+FStrataData GetStrataDiffuseOrenNayarBSDF(float3 Albedo, float Roughness, uint SharedNormalIndex)
 {
 	FStrataData StrataData = GetInitialisedStrataData();
 	StrataData.LayerCount++;
 
 	StrataData.Layers[0].BSDFCount++;
-	StrataData.Layers[0].BSDFs[0].BSDFType						= STRATA_BSDF_TYPE_DIFFUSE_ON;
-	DIFFUSE_ALBEDO(StrataData.Layers[0].BSDFs[0])				= Albedo;
-	DIFFUSE_ROUGHNESS(StrataData.Layers[0].BSDFs[0])			= Roughness;
-	DIFFUSE_NORMAL(StrataData.Layers[0].BSDFs[0])				= Normal;
+	BSDF_SETTYPE			(StrataData.Layers[0].BSDFs[0],		STRATA_BSDF_TYPE_DIFFUSE_ON);
+	BSDF_SETNORMALID		(StrataData.Layers[0].BSDFs[0],		SharedNormalIndex);
+	DIFFUSE_ALBEDO			(StrataData.Layers[0].BSDFs[0])		= Albedo;
+	DIFFUSE_ROUGHNESS		(StrataData.Layers[0].BSDFs[0])		= Roughness;
 
 	StrataData.Layers[0].BSDFs[0].Weight						= 1.0f;
 
 	return StrataData;
 }
 
-FStrataData GetStrataDiffuseChanBSDF(float3 Albedo, float Roughness, float3 Normal)
+FStrataData GetStrataDiffuseChanBSDF(float3 Albedo, float Roughness, uint SharedNormalIndex)
 {
 	FStrataData StrataData = GetInitialisedStrataData();
 	StrataData.LayerCount++;
 
 	StrataData.Layers[0].BSDFCount++;
-	StrataData.Layers[0].BSDFs[0].BSDFType						= STRATA_BSDF_TYPE_DIFFUSE_CHAN;
-	DIFFUSE_ALBEDO		(StrataData.Layers[0].BSDFs[0])			= Albedo;
-	DIFFUSE_ROUGHNESS	(StrataData.Layers[0].BSDFs[0])			= Roughness;
-	DIFFUSE_NORMAL		(StrataData.Layers[0].BSDFs[0])			= Normal;
+	BSDF_SETTYPE			(StrataData.Layers[0].BSDFs[0],		STRATA_BSDF_TYPE_DIFFUSE_CHAN);
+	BSDF_SETNORMALID		(StrataData.Layers[0].BSDFs[0],		SharedNormalIndex);
+	DIFFUSE_ALBEDO			(StrataData.Layers[0].BSDFs[0])		= Albedo;
+	DIFFUSE_ROUGHNESS		(StrataData.Layers[0].BSDFs[0])		= Roughness;
 
 	StrataData.Layers[0].BSDFs[0].Weight						= 1.0f;
 
 	return StrataData;
 }
 
-FStrataData GetStrataDielectricBSDF(float2 Roughness, float IOR, float3 Tint, float3 Normal)
+FStrataData GetStrataDielectricBSDF(float2 Roughness, float IOR, float3 Tint, uint SharedNormalIndex)
 {
 	FStrataData StrataData = GetInitialisedStrataData();
 	StrataData.LayerCount++;
 
 	StrataData.Layers[0].BSDFCount++;
-	StrataData.Layers[0].BSDFs[0].BSDFType						= STRATA_BSDF_TYPE_DIELECTRIC;
+	BSDF_SETTYPE			(StrataData.Layers[0].BSDFs[0],		STRATA_BSDF_TYPE_DIELECTRIC);
+	BSDF_SETNORMALID		(StrataData.Layers[0].BSDFs[0],		SharedNormalIndex);
 	DIELECTRIC_IOR			(StrataData.Layers[0].BSDFs[0])		= IOR;
 	DIELECTRIC_TINT			(StrataData.Layers[0].BSDFs[0])		= Tint;
 	DIELECTRIC_ROUGHNESS0	(StrataData.Layers[0].BSDFs[0])		= Roughness.x;
 	DIELECTRIC_ROUGHNESS1	(StrataData.Layers[0].BSDFs[0])		= Roughness.y;
-	DIELECTRIC_NORMAL		(StrataData.Layers[0].BSDFs[0])		= Normal;
 
 	StrataData.Layers[0].BSDFs[0].Weight						= 1.0f;
 
 	return StrataData;
 }
 
-FStrataData GetStrataConductorBSDF(float3 Reflectivity, float3 EdgeColor, float2 Roughness, float3 Normal)
+FStrataData GetStrataConductorBSDF(float3 Reflectivity, float3 EdgeColor, float2 Roughness, uint SharedNormalIndex)
 {
 	FStrataData StrataData = GetInitialisedStrataData();
 	StrataData.LayerCount++;
 
 	StrataData.Layers[0].BSDFCount++;
-	StrataData.Layers[0].BSDFs[0].BSDFType						= STRATA_BSDF_TYPE_CONDUCTOR;
+	BSDF_SETTYPE			(StrataData.Layers[0].BSDFs[0],		STRATA_BSDF_TYPE_CONDUCTOR);
+	BSDF_SETNORMALID		(StrataData.Layers[0].BSDFs[0],		SharedNormalIndex);
 	CONDUCTOR_REFLECTIVITY	(StrataData.Layers[0].BSDFs[0])		= Reflectivity;
 	CONDUCTOR_EDGECOLOR		(StrataData.Layers[0].BSDFs[0])		= EdgeColor;
 	CONDUCTOR_ROUGHNESS0	(StrataData.Layers[0].BSDFs[0])		= Roughness.x;
 	CONDUCTOR_ROUGHNESS1	(StrataData.Layers[0].BSDFs[0])		= Roughness.y;
-	CONDUCTOR_NORMAL		(StrataData.Layers[0].BSDFs[0])		= Normal;
 
 	StrataData.Layers[0].BSDFs[0].Weight						= 1.0f;
 
 	return StrataData;
 }
 
-FStrataData GetStrataVolumeBSDF(float3 Albedo, float3 Extinction, float Anisotropy, float Thickness, float3 SmoothSurfaceNormal)
+FStrataData GetStrataVolumeBSDF(float3 Albedo, float3 Extinction, float Anisotropy, float Thickness, uint SharedNormalIndex)
 {
 	FStrataData StrataData = GetInitialisedStrataData();
 	StrataData.LayerCount++;
 
 	StrataData.Layers[0].BSDFCount++;
-	StrataData.Layers[0].BSDFs[0].BSDFType						= STRATA_BSDF_TYPE_VOLUME;
+	BSDF_SETTYPE(				StrataData.Layers[0].BSDFs[0],	STRATA_BSDF_TYPE_VOLUME);
+	BSDF_SETNORMALID(			StrataData.Layers[0].BSDFs[0],	SharedNormalIndex);
 	VOLUME_ALBEDO				(StrataData.Layers[0].BSDFs[0])	= Albedo;
 	VOLUME_EXTINCTION			(StrataData.Layers[0].BSDFs[0])	= Extinction;
 	VOLUME_ANISOTROPY			(StrataData.Layers[0].BSDFs[0])	= Anisotropy;
 	VOLUME_THICKNESS			(StrataData.Layers[0].BSDFs[0])	= Thickness;
-	VOLUME_SMOOTHSURFACENORMAL	(StrataData.Layers[0].BSDFs[0])	= SmoothSurfaceNormal;
 
 	StrataData.Layers[0].BSDFs[0].Weight						= 1.0f;
 
@@ -224,6 +181,27 @@ float3 StrataGetOffSpecularPeakReflectionDir(float3 Normal, float3 ReflectionVec
 	return lerp(Normal, ReflectionVector, (1 - a) * (sqrt(1 - a) + a));
 }
 
+
+///////////////////////////////////////////////////////////////////////////////
+// Shared normals
+
+struct FSharedNormals
+{
+	uint NormalCount;
+	float3 Normals[STRATA_MAX_TOTAL_BSDF];
+};
+
+FSharedNormals StrataInitialiseSharedNormals()
+{
+	FSharedNormals SharedNormals = (FSharedNormals)0;
+	return SharedNormals;
+}
+
+float3 StrataGetBSDFSharedNormal(in FSharedNormals SharedNormals, in FStrataBSDF BSDF)
+{
+	return SharedNormals.Normals[BSDF_GETNORMALID(BSDF)];
+}
+
 ///////////////////////////////////////////////////////////////////////////////
 // BSDF evaluate and sampling
 
@@ -250,38 +228,16 @@ struct FStrataBSDFContext
 	float a2Override;
 };
 
-void StrataGetBSDFNormal(in FStrataBSDF BSDF, inout float3 N)
+void StrataGetBSDFNormal(in FSharedNormals SharedNormals, in FStrataBSDF BSDF, inout float3 N)
 {
-	switch (BSDF.BSDFType)
-	{
-		case STRATA_BSDF_TYPE_DIFFUSE_ON:
-		case STRATA_BSDF_TYPE_DIFFUSE_CHAN:
-		{ 
-			N = DIFFUSE_NORMAL(BSDF);
-			break; 
-		}
-		case STRATA_BSDF_TYPE_DIELECTRIC:
-		{ 
-			N = DIELECTRIC_NORMAL(BSDF);
-			break; 
-		}
-		case STRATA_BSDF_TYPE_CONDUCTOR:
-		{
-			N = CONDUCTOR_NORMAL(BSDF);
-			break;
-		}
-		case STRATA_BSDF_TYPE_VOLUME:
-		{
-			N = VOLUME_SMOOTHSURFACENORMAL(BSDF);
-			break;
-		}
-	}
+	N = StrataGetBSDFSharedNormal(SharedNormals, BSDF);
 }
 
 void StrataGetBSDFRoughness(in FStrataBSDF BSDF, inout float R)
 {
 	R = 0;
-	switch (BSDF.BSDFType)
+	const uint BSDFType = BSDF_GETTYPE(BSDF);
+	switch (BSDFType)
 	{
 	case STRATA_BSDF_TYPE_DIFFUSE_ON:
 	case STRATA_BSDF_TYPE_DIFFUSE_CHAN:
@@ -310,12 +266,13 @@ void StrataGetBSDFRoughness(in FStrataBSDF BSDF, inout float R)
 // Return true if a BSDF has a custom area light integrator
 bool StrataHasAreaLightIntegrator(in FStrataBSDF BSDF)
 {
+	const uint BSDFType = BSDF_GETTYPE(BSDF);
 	return 
-		BSDF.BSDFType == STRATA_BSDF_TYPE_CONDUCTOR ||
-		BSDF.BSDFType == STRATA_BSDF_TYPE_DIELECTRIC;
+		BSDFType == STRATA_BSDF_TYPE_CONDUCTOR ||
+		BSDFType == STRATA_BSDF_TYPE_DIELECTRIC;
 }
 
-FStrataBSDFContext StrataCreateBSDFContext(FStrataBSDF BSDF, float3 V, float3 L)
+FStrataBSDFContext StrataCreateBSDFContext(FSharedNormals SharedNormals, FStrataBSDF BSDF, float3 V, float3 L)
 {
 	FStrataBSDFContext BSDFContext = (FStrataBSDFContext)0;
 
@@ -324,8 +281,7 @@ FStrataBSDFContext StrataCreateBSDFContext(FStrataBSDF BSDF, float3 V, float3 L)
 	BSDFContext.V = V;
 	BSDFContext.L = L;
 
-	BSDFContext.N = float3(0.0f, 0.0f, 1.0f);
-	StrataGetBSDFNormal(BSDF, BSDFContext.N);
+	BSDFContext.N = StrataGetBSDFSharedNormal(SharedNormals, BSDF);
 	BSDFContext.N = normalize(BSDFContext.N);
 
 	BSDFContext.R = 2 * dot(BSDFContext.V, BSDFContext.N) * BSDFContext.N - BSDFContext.V;
@@ -364,7 +320,8 @@ FStrataEvaluateResult StrataEvaluateBSDF(FStrataBSDFContext BSDFContext)
 
 	const float OpaqueBSDFThroughput = 0.0f;
 
-	switch (BSDFContext.BSDF.BSDFType)
+	const uint BSDFType = BSDF_GETTYPE(BSDFContext.BSDF);
+	switch (BSDFType)
 	{
 		case STRATA_BSDF_TYPE_DIFFUSE_ON:
 		{
@@ -388,7 +345,7 @@ FStrataEvaluateResult StrataEvaluateBSDF(FStrataBSDFContext BSDFContext)
 		case STRATA_BSDF_TYPE_DIELECTRIC:
 		case STRATA_BSDF_TYPE_CONDUCTOR:
 		{
-			const bool bIsDieletric = BSDFContext.BSDF.BSDFType == STRATA_BSDF_TYPE_DIELECTRIC;
+			const bool bIsDieletric = BSDFType == STRATA_BSDF_TYPE_DIELECTRIC;
 			const float3 F0  = bIsDieletric ? DIELECTRIC_TINT(BSDFContext.BSDF) * DielectricIorToF0(DIELECTRIC_IOR(BSDFContext.BSDF)) : CONDUCTOR_REFLECTIVITY(BSDFContext.BSDF);
 			const float3 F90 = bIsDieletric ? DIELECTRIC_TINT(BSDFContext.BSDF) : CONDUCTOR_EDGECOLOR(BSDFContext.BSDF);
 			const float SafeRoughness = MakeRoughnessSafe(bIsDieletric ? DIELECTRIC_ROUGHNESS0(BSDFContext.BSDF) : CONDUCTOR_ROUGHNESS0(BSDFContext.BSDF));
@@ -453,17 +410,18 @@ struct FStrataImportanceSampleResult
  * E: two random numbers
  * CameraVector: vector from the camera to the considered direction (not V)
  */
-FStrataImportanceSampleResult StrataImportanceSampleBSDF(FStrataBSDF BSDF, float2 E, float3 CameraVector)
+FStrataImportanceSampleResult StrataImportanceSampleBSDF(FSharedNormals SharedNormals, FStrataBSDF BSDF, float2 E, float3 CameraVector)
 {
 	FStrataImportanceSampleResult Sample = (FStrataImportanceSampleResult)0;
 
 	// Create a context from the information we have
 	const float3 UnusedL = float3(0.0f, 0.0f, 1.0f);
 	const float3 V = -CameraVector;
-	FStrataBSDFContext BSDFContext = StrataCreateBSDFContext(BSDF, V, UnusedL);
+	FStrataBSDFContext BSDFContext = StrataCreateBSDFContext(SharedNormals, BSDF, V, UnusedL);
 
 	float PDF = 0.0f;
-	switch (BSDFContext.BSDF.BSDFType)
+	const uint BSDFType = BSDF_GETTYPE(BSDF);
+	switch (BSDFType)
 	{
 		case STRATA_BSDF_TYPE_DIFFUSE_ON:
 		case STRATA_BSDF_TYPE_DIFFUSE_CHAN:
@@ -478,7 +436,7 @@ FStrataImportanceSampleResult StrataImportanceSampleBSDF(FStrataBSDF BSDF, float
 		case STRATA_BSDF_TYPE_DIELECTRIC:
 		case STRATA_BSDF_TYPE_CONDUCTOR:
 		{
-			const bool bIsDielectric = BSDFContext.BSDF.BSDFType == STRATA_BSDF_TYPE_DIELECTRIC;
+			const bool bIsDielectric = BSDFType == STRATA_BSDF_TYPE_DIELECTRIC;
 			float Roughness = bIsDielectric ? DIELECTRIC_ROUGHNESS0(BSDFContext.BSDF) : CONDUCTOR_ROUGHNESS0(BSDFContext.BSDF);
 			float4 GGXSample = ImportanceSampleVisibleGGX(UniformSampleDisk(E), Pow4(Roughness), BSDFContext.TangentV);
 			float3 WorldH = mul(GGXSample.xyz, BSDFContext.TangentBasis);
@@ -518,7 +476,8 @@ FStrataEnvLightResult StrataEvaluateForEnvLight(FStrataBSDFContext BSDFContext,
 {
 	FStrataEnvLightResult StrataEnvLightResult = (FStrataEnvLightResult)0;
 
-	switch (BSDFContext.BSDF.BSDFType)
+	const uint BSDFType = BSDF_GETTYPE(BSDFContext.BSDF);
+	switch (BSDFType)
 	{
 		case STRATA_BSDF_TYPE_DIFFUSE_ON:
 		case STRATA_BSDF_TYPE_DIFFUSE_CHAN:
@@ -533,7 +492,7 @@ FStrataEnvLightResult StrataEvaluateForEnvLight(FStrataBSDFContext BSDFContext,
 			BRANCH
 			if (bEnableSpecular)
 			{
-				const bool bIsDielectric = STRATA_BSDF_TYPE_DIELECTRIC == BSDFContext.BSDF.BSDFType;
+				const bool bIsDielectric = STRATA_BSDF_TYPE_DIELECTRIC == BSDFType;
 				const float3 F0  = bIsDielectric ? DIELECTRIC_TINT(BSDFContext.BSDF) * DielectricIorToF0(DIELECTRIC_IOR(BSDFContext.BSDF)) : CONDUCTOR_REFLECTIVITY(BSDFContext.BSDF);
 				const float3 F90 = bIsDielectric ? DIELECTRIC_TINT(BSDFContext.BSDF) : CONDUCTOR_EDGECOLOR(BSDFContext.BSDF);
 				const float SafeRoughness = MakeRoughnessSafe(bIsDielectric ? DIELECTRIC_ROUGHNESS0(BSDFContext.BSDF) : CONDUCTOR_ROUGHNESS0(BSDFContext.BSDF));
@@ -684,7 +643,7 @@ void SanitizeStrataVolume(inout FStrataBSDF BSDF)
 
 void StrataSanitizeBSDF(inout FStrataBSDF BSDF)
 {
-	switch (BSDF.BSDFType)
+	switch (BSDF_GETTYPE(BSDF))
 	{
 		case STRATA_BSDF_TYPE_DIFFUSE_ON:
 		case STRATA_BSDF_TYPE_DIFFUSE_CHAN:
@@ -782,6 +741,42 @@ float4 UnpackAlbedoGamma2ToLinearAlphaLinear(uint rgba)
 	return float4(Gamma2 * Gamma2, a * (1.0f / 255.0f));	// Simple gamma 2.0
 }
 
+// Simple Octahedral from "A Survey of Efficient Representations for Independent Unit Vectors" http://jcgt.org/published/0003/02/01/
+float SignNotZero(in float k)
+{
+	return k >= 0.0 ? 1.0 : -1.0;
+}
+float2 SignNotZero(in float2 v)
+{
+	return float2(SignNotZero(v.x), SignNotZero(v.y));
+}
+uint StrataPackNormal(in float3 Normal)
+{
+	float l1norm = abs(Normal.x) + abs(Normal.y) + abs(Normal.z);
+	float2 result = Normal.xy * (1.0 / l1norm);
+	if (Normal.z < 0.0)
+	{
+		result = (1.0 - abs(result.yx)) * SignNotZero(result.xy);
+	}
+	uint X10Bits = uint(clamp(result.x * 511.0f + 512.0f, 0.0f, 1023.0f));
+	uint Y10Bits = uint(clamp(result.y * 511.0f + 512.0f, 0.0f, 1023.0f));
+	return (X10Bits | (Y10Bits << 10));
+}
+float3 StrataUnpackNormal(uint PackedNormal)
+{
+	uint X10Bits = 0x3FF & PackedNormal;
+	uint Y10Bits = 0x3FF & (PackedNormal >> 10);
+	float x = float(int(X10Bits) - 512) / 511.0f;
+	float y = float(int(Y10Bits) - 512) / 511.0f;
+
+	float3 Normal = float3(x, y, 1.0f - abs(x) - abs(y));
+	if (Normal.z < 0.0f)
+	{
+		Normal.xy = (1.0f - abs(Normal.yx)) * SignNotZero(Normal.xy);
+	}
+	return normalize(Normal);
+}
+
 #if COMPILER_SWITCH // These platforms are complaining so simply disabling Strata there for now.
 
 #define STRATA_STORE_UINT1(x) 
@@ -872,7 +867,8 @@ uint GetStrataPixelDataByteOffset(uint2 PixelPos, uint ViewBufferSizeX, uint Str
 void PackStrataOut(
 	RWByteAddressBuffer StrataBuffer, 
 	uint ByteOffset, 
-	FStrataData Strata, 
+	FStrataData Strata,
+	FSharedNormals SharedNormals,
 	float3 VertexNormal,
 	float3 V,
 	inout FGBufferData GBuffer)
@@ -915,7 +911,7 @@ void PackStrataOut(
 
 			// STRATA_TODO apply the same normal transform as for the material
 			float3 N = VertexNormal;
-			StrataGetBSDFNormal(BSDF, N);
+			StrataGetBSDFNormal(SharedNormals, BSDF, N);
 			N = normalize(N);
 
 			AccumTopNormal += N;
@@ -924,7 +920,7 @@ void PackStrataOut(
 			const float3 FakeL = float3(0, 0, 1);
 			Init(ContextNoLight, N, V, FakeL);
 
-			switch (BSDF.BSDFType)
+			switch (BSDF_GETTYPE(BSDF))
 			{
 			case STRATA_BSDF_TYPE_DIFFUSE_ON:
 			{
@@ -1042,11 +1038,20 @@ void PackStrataOut(
 	}
 
 
-	// Now write out the BSDFs
+	// Now write out Strata data
 
+	// 1. the header
 	BSDFCount = min(BSDFCount, uint(STRATA_MAX_TOTAL_BSDF));
-	STRATA_STORE_UINT1(BSDFCount);
 
+	uint SharedNormals_BSDFCount = (SharedNormals.NormalCount << 16) | BSDFCount;
+	STRATA_STORE_UINT1(SharedNormals_BSDFCount);
+
+	for (int i = 0; i < SharedNormals.NormalCount; ++i)
+	{
+		STRATA_STORE_UINT1(StrataPackNormal(SharedNormals.Normals[i]));
+	}
+
+	// 2. the list of BSDFs
 	BSDFCount = 0;
 	for (l = 0; l < Strata.LayerCount; ++l)
 	{
@@ -1056,32 +1061,31 @@ void PackStrataOut(
 		{
 			FStrataBSDF BSDF = Layer.BSDFs[i];
 
-			STRATA_STORE_UINT2(PackR16G16B16A16F(float4(BSDF.Weight, BSDF.BSDFType)));
+			STRATA_STORE_UINT1(PackR11G11B10F(BSDF.Weight));
+			STRATA_STORE_UINT1(BSDF.State);
 
-			switch (BSDF.BSDFType)
+			const uint BSDFType = BSDF_GETTYPE(BSDF);
+			switch (BSDF_GETTYPE(BSDF))
 			{
 			case STRATA_BSDF_TYPE_DIFFUSE_ON:
 			case STRATA_BSDF_TYPE_DIFFUSE_CHAN:
 			{
-				STRATA_STORE_UINT1(PackAlbedoLinearToGamma2(DIFFUSE_ALBEDO(BSDF)));
-				STRATA_STORE_UINT2(PackR16G16B16A16F(float4(DIFFUSE_NORMAL(BSDF), DIFFUSE_ROUGHNESS(BSDF))));
-				// 12 bytes
+				STRATA_STORE_UINT1(PackAlbedoLinearToGamma2AlphaLinear(float4(DIFFUSE_ALBEDO(BSDF), DIFFUSE_ROUGHNESS(BSDF))));
+				// 4 bytes
 			}
 			break;
 			case STRATA_BSDF_TYPE_DIELECTRIC: //
 			{
 				STRATA_STORE_UINT1(PackR11G11B10F(DIELECTRIC_TINT(BSDF)));
 				STRATA_STORE_UINT1(PackR11G11B10F(float3(DIELECTRIC_ROUGHNESS0(BSDF), DIELECTRIC_ROUGHNESS1(BSDF), DIELECTRIC_IOR(BSDF))));
-				STRATA_STORE_UINT2(PackR16G16B16A16F(float4(DIELECTRIC_NORMAL(BSDF), 0.0f)));
-				// 16 bytes
+				// 8 bytes
 			}
 			break;
 			case STRATA_BSDF_TYPE_CONDUCTOR:
 			{
 				STRATA_STORE_UINT1(PackAlbedoLinearToGamma2AlphaLinear(float4(CONDUCTOR_REFLECTIVITY(BSDF), CONDUCTOR_ROUGHNESS0(BSDF))));
 				STRATA_STORE_UINT1(PackAlbedoLinearToGamma2AlphaLinear(float4(CONDUCTOR_EDGECOLOR(BSDF),    CONDUCTOR_ROUGHNESS1(BSDF))));
-				STRATA_STORE_UINT2(PackR16G16B16A16F(float4(CONDUCTOR_NORMAL(BSDF), 0.0f)));
-				// 16 bytes
+				// 8 bytes
 			}
 			break;
 			case STRATA_BSDF_TYPE_VOLUME:
@@ -1090,9 +1094,8 @@ void PackStrataOut(
 				// This gives the same result and reduces memory usage.
 				const float Thickness = VOLUME_THICKNESS(BSDF);
 				STRATA_STORE_UINT1(PackR11G11B10F(VOLUME_EXTINCTION(BSDF) * Thickness));
-				STRATA_STORE_UINT1(PackR11G11B10F(VOLUME_ALBEDO(BSDF)));
-				STRATA_STORE_UINT2(PackR16G16B16A16F(float4(VOLUME_SMOOTHSURFACENORMAL(BSDF), VOLUME_ANISOTROPY(BSDF))));
-				// 16 bytes
+				STRATA_STORE_UINT1(PackAlbedoLinearToGamma2AlphaLinear(float4(VOLUME_ALBEDO(BSDF), VOLUME_ANISOTROPY(BSDF)*0.5f + 0.5f)));
+				// 8 bytes
 			}
 			break;
 			}
@@ -1110,6 +1113,8 @@ void PackStrataOut(
 struct FStrataPixelHeader
 {
 	uint BSDFCount;
+
+	FSharedNormals SharedNormals;
 };
 
 // Unpack only strata header
@@ -1117,8 +1122,18 @@ FStrataPixelHeader UnpackStrataHeaderIn(ByteAddressBuffer StrataBuffer, inout ui
 {
 	FStrataPixelHeader Out = (FStrataPixelHeader)0;
 
-	STRATA_LOAD_UINT1(Out.BSDFCount);
-	Out.BSDFCount = min(Out.BSDFCount, STRATA_MAX_TOTAL_BSDF);	// To make the compiler happy
+	uint SharedNormals_BSDFCount = 0;
+	STRATA_LOAD_UINT1(SharedNormals_BSDFCount);
+	Out.BSDFCount = SharedNormals_BSDFCount & 0x0000FFFF;
+	Out.SharedNormals.NormalCount = SharedNormals_BSDFCount >> 16;
+
+	for (int i = 0; i < Out.SharedNormals.NormalCount; ++i)
+	{
+		uint PackedNormal = 0;
+		STRATA_LOAD_UINT1(PackedNormal);
+		Out.SharedNormals.Normals[i] = StrataUnpackNormal(PackedNormal);
+	}
+
 	return Out;
 }
 
@@ -1130,40 +1145,36 @@ FStrataBSDF UnpackStrataBSDFIn(ByteAddressBuffer StrataBuffer, inout uint ByteOf
 
 	float DummyFloat = 0.0f;
 
-	{
-		float4 temp4  = UnpackR16G16B16A16F(StrataLoadUint2(StrataBuffer, ByteOffset));
-		OutBSDF.Weight   = temp4.xyz;
-		OutBSDF.BSDFType = int(temp4.a);
-	}
+	StrataLoad_R11G11B10F(StrataBuffer, ByteOffset, OutBSDF.Weight);
+	STRATA_LOAD_UINT1(OutBSDF.State);
 
-	switch (OutBSDF.BSDFType)
+	const uint BSDFType = BSDF_GETTYPE(OutBSDF);
+	BSDF_SETTYPE(OutBSDF, BSDFType);
+	switch (BSDFType)
 	{
 	case STRATA_BSDF_TYPE_DIFFUSE_ON:
 	case STRATA_BSDF_TYPE_DIFFUSE_CHAN:
 	{
-		StrataLoad_AlbedoGamma2ToLinear	(StrataBuffer, ByteOffset, DIFFUSE_ALBEDO(OutBSDF));
-		StrataLoad_R16G16B16F_A16F		(StrataBuffer, ByteOffset, DIFFUSE_NORMAL(OutBSDF),			DIFFUSE_ROUGHNESS(OutBSDF));
+		StrataLoad_AlbedoGamma2ToLinearAlphaLinear(StrataBuffer, ByteOffset, DIFFUSE_ALBEDO(OutBSDF), DIFFUSE_ROUGHNESS(OutBSDF));
 	}
 	break;
 	case STRATA_BSDF_TYPE_DIELECTRIC:
 	{
 		StrataLoad_R11G11B10F			(StrataBuffer, ByteOffset, DIELECTRIC_TINT(OutBSDF));
 		StrataLoad_R11_G11_B10F			(StrataBuffer, ByteOffset, DIELECTRIC_ROUGHNESS0(OutBSDF), DIELECTRIC_ROUGHNESS1(OutBSDF), DIELECTRIC_IOR(OutBSDF));
-		StrataLoad_R16G16B16F_A16F		(StrataBuffer, ByteOffset, DIELECTRIC_NORMAL(OutBSDF),		DummyFloat);
 	}
 	break;
 	case STRATA_BSDF_TYPE_CONDUCTOR:
 	{
 		StrataLoad_AlbedoGamma2ToLinearAlphaLinear(StrataBuffer, ByteOffset, CONDUCTOR_REFLECTIVITY(OutBSDF),	CONDUCTOR_ROUGHNESS0(OutBSDF));
 		StrataLoad_AlbedoGamma2ToLinearAlphaLinear(StrataBuffer, ByteOffset, CONDUCTOR_EDGECOLOR(OutBSDF),		CONDUCTOR_ROUGHNESS1(OutBSDF));
-		StrataLoad_R16G16B16F_A16F		(StrataBuffer, ByteOffset,			 CONDUCTOR_NORMAL(OutBSDF),			DummyFloat);
 	}
 	break;
 	case STRATA_BSDF_TYPE_VOLUME:
 	{
 		StrataLoad_R11G11B10F			(StrataBuffer, ByteOffset, VOLUME_EXTINCTION(OutBSDF));
-		StrataLoad_R11G11B10F			(StrataBuffer, ByteOffset, VOLUME_ALBEDO(OutBSDF));
-		StrataLoad_R16G16B16F_A16F		(StrataBuffer, ByteOffset, VOLUME_SMOOTHSURFACENORMAL(OutBSDF), VOLUME_ANISOTROPY(OutBSDF));
+		StrataLoad_AlbedoGamma2ToLinearAlphaLinear(StrataBuffer, ByteOffset, VOLUME_ALBEDO(OutBSDF), VOLUME_ANISOTROPY(OutBSDF));
+		VOLUME_ANISOTROPY(OutBSDF) = VOLUME_ANISOTROPY(OutBSDF) * 2.0f - 1.0f;
 		VOLUME_THICKNESS(OutBSDF) = 1.0f;
 	}
 	break;
diff --git a/Engine/Shaders/Private/Strata/StrataDeferredLighting.ush b/Engine/Shaders/Private/Strata/StrataDeferredLighting.ush
--- a/Engine/Shaders/Private/Strata/StrataDeferredLighting.ush
+++ b/Engine/Shaders/Private/Strata/StrataDeferredLighting.ush
@@ -2,7 +2,7 @@
 
 #pragma once
 
-#include "..\AreaLightCommon.ush"
+#include "../AreaLightCommon.ush"
 
 struct FStrataLightIntegrationContext
 {
@@ -60,12 +60,12 @@ void StrataPreprocessAreaLightIntegration(
 		// Evaluate roughness / specular color
 		float SafeRoughness = 0;
 		float3 SpecularColor = 0;
-		switch (BSDFContext.BSDF.BSDFType)
+		switch (BSDF_GETTYPE(BSDFContext.BSDF))
 		{
 			case STRATA_BSDF_TYPE_DIELECTRIC:
 			{
 				SafeRoughness = MakeRoughnessSafe(DIELECTRIC_ROUGHNESS0(BSDFContext.BSDF));
-				SpecularColor = DielectricIorToF0(DIELECTRIC_REFLECTIVITY(BSDFContext.BSDF)) * DIELECTRIC_TINT(BSDFContext.BSDF);
+				SpecularColor = DielectricIorToF0(DIELECTRIC_IOR(BSDFContext.BSDF)) * DIELECTRIC_TINT(BSDFContext.BSDF);
 			}
 			break;
 			case STRATA_BSDF_TYPE_CONDUCTOR:
@@ -146,7 +146,7 @@ float4 StrataDeferredLighting(
 	{
 		// Unpack BSDF data
 		const FStrataBSDF BSDF = UnpackStrataBSDFIn(StrataDataBuffer, PixelStrataDataByteOffset);
-		FStrataBSDFContext BSDFContext = StrataCreateBSDFContext(BSDF, V, L);
+		FStrataBSDFContext BSDFContext = StrataCreateBSDFContext(Header.SharedNormals, BSDF, V, L);
 
 		// Evaluate incident lighting & (optionally) precompute area light integration
 		FStrataLightIntegrationContext IntegrationContext = StrataIniLightIntegrationContext();
diff --git a/Engine/Shaders/Private/Strata/StrataEnvironmentLighting.ush b/Engine/Shaders/Private/Strata/StrataEnvironmentLighting.ush
--- a/Engine/Shaders/Private/Strata/StrataEnvironmentLighting.ush
+++ b/Engine/Shaders/Private/Strata/StrataEnvironmentLighting.ush
@@ -24,7 +24,7 @@ float3 StrataEnvironmentLighting(
 
 		// Create the BSDF context
 		const float3 UnusedL = float3(0.0f, 0.0f, 1.0f);
-		FStrataBSDFContext StrataBSDFContext = StrataCreateBSDFContext(BSDF, V, UnusedL);
+		FStrataBSDFContext StrataBSDFContext = StrataCreateBSDFContext(Header.SharedNormals, BSDF, V, UnusedL);
 
 		// Evaluate environment lighting
 		const bool bEnableSpecular = ReflectionStruct.SkyLightParameters.y > 0.0f;
diff --git a/Engine/Shaders/Private/Strata/StrataForwardLighting.ush b/Engine/Shaders/Private/Strata/StrataForwardLighting.ush
--- a/Engine/Shaders/Private/Strata/StrataForwardLighting.ush
+++ b/Engine/Shaders/Private/Strata/StrataForwardLighting.ush
@@ -12,6 +12,7 @@ float3 StrataForwardLighting(
 	float3 LColor,
 	float3 LShadow,
 	float3 IndirectOcclusion,
+	FStrataPixelHeader Header,
 	FStrataData Strata,
 	inout float3 OutThroughput
 	)
@@ -37,7 +38,7 @@ float3 StrataForwardLighting(
 			const float3 BSDFCoverage = BSDF.Weight;
 
 			// Create the BSDF context
-			FStrataBSDFContext StrataBSDFContext = StrataCreateBSDFContext(BSDF, V, L);
+			FStrataBSDFContext StrataBSDFContext = StrataCreateBSDFContext(Header.SharedNormals, BSDF, V, L);
 
 			// Compute some common factors
 			const float3 DirLightBaseFactor = OutThroughput * BSDFCoverage * LColor * LShadow;
@@ -51,7 +52,7 @@ float3 StrataForwardLighting(
 			{
 				// SkySHDiffuse is pre integrated according to a Lambert diffuse lobe. Thus, we need to rescale to account for the phase function operating over 4PI.
 				// We apply a 2PI factor to recover average luminance over half hemisphere, and then times the phase function to account for the scattering events.
-				const float VolumeOnlyFactor = BSDF.BSDFType == STRATA_BSDF_TYPE_VOLUME ? 2.0f * PI * IsotropicPhase() : 1.0f;
+				const float VolumeOnlyFactor = BSDF_GETTYPE(BSDF) == STRATA_BSDF_TYPE_VOLUME ? 2.0f * PI * IsotropicPhase() : 1.0f;
 
 				Color += GetEffectiveSkySHDiffuse(StrataEnvLight.DiffuseNormal) * StrataEnvLight.DiffuseAlbedo * SkyLightCommonFactor * VolumeOnlyFactor;
 			}
diff --git a/Engine/Shaders/Private/Strata/StrataVisualize.usf b/Engine/Shaders/Private/Strata/StrataVisualize.usf
--- a/Engine/Shaders/Private/Strata/StrataVisualize.usf
+++ b/Engine/Shaders/Private/Strata/StrataVisualize.usf
@@ -105,7 +105,7 @@ void VisualizeMaterialPS(
 
 			TopLeft.x = TopLeftInit.x;
 
-			switch (BSDF.BSDFType)
+			switch (BSDF_GETTYPE(BSDF))
 			{
 			case STRATA_BSDF_TYPE_DIFFUSE_ON:
 				PRINTBSDFCHAR(_D_); PRINTBSDFCHAR(_I_); PRINTBSDFCHAR(_F_); PRINTBSDFCHAR(_F_); PRINTBSDFCHAR(_O_); PRINTBSDFCHAR(_N_);
diff --git a/Engine/Source/Runtime/Engine/Private/Materials/HLSLMaterialTranslator.cpp b/Engine/Source/Runtime/Engine/Private/Materials/HLSLMaterialTranslator.cpp
--- a/Engine/Source/Runtime/Engine/Private/Materials/HLSLMaterialTranslator.cpp
+++ b/Engine/Source/Runtime/Engine/Private/Materials/HLSLMaterialTranslator.cpp
@@ -110,6 +110,7 @@ FHLSLMaterialTranslator::FHLSLMaterialTranslator(FMaterial* InMaterial,
 ,	AllocatedUserTexCoords()
 ,	AllocatedUserVertexTexCoords()
 ,	DynamicParticleParameterMask(0)
+,	NextFreeStrataShaderNormalIndex(0)
 ,	NumVtSamples(0)
 ,	TargetPlatform(InTargetPlatform)
 {
@@ -1107,16 +1108,18 @@ bool FHLSLMaterialTranslator::Translate()
 
 			FString StrataMaterialDescription = "";
 
-			StrataMaterialDescription += FString::Printf(TEXT("StrataCompilationInfo - TotalBSDFCount = %i\r\n"), StrataCompilationInfo.TotalBSDFCount);
+			StrataMaterialDescription += FString::Printf(TEXT("StrataCompilationInfo -\r\n"));
+			StrataMaterialDescription += FString::Printf(TEXT(" - TotalBSDFCount    = %i\r\n"), StrataCompilationInfo.TotalBSDFCount);
+			StrataMaterialDescription += FString::Printf(TEXT(" - SharedNormalCount = %i\r\n"), GetStrataSharedNormalCount());
 
 			for (uint32 LayerIt = 0; LayerIt < StrataCompilationInfo.LayerCount; ++LayerIt)
 			{
-				StrataMaterialDescription += FString::Printf(TEXT("    Layer %i BSDFs:\r\n"), LayerIt);
+				StrataMaterialDescription += FString::Printf(TEXT(" Layer %i BSDFs:\r\n"), LayerIt);
 
 				const FStrataMaterialCompilationInfo::FLayer& Layer = StrataCompilationInfo.Layers[LayerIt];
 				for (uint32 BSDFIt = 0; BSDFIt < Layer.BSDFCount; ++BSDFIt)
 				{
-					StrataMaterialDescription += FString::Printf(TEXT("        - %s\r\n"), *GetStrataBSDFName(Layer.BSDFs[BSDFIt].Type));
+					StrataMaterialDescription += FString::Printf(TEXT("     - %s - SharedNormalIndex = %u \r\n"), *GetStrataBSDFName(Layer.BSDFs[BSDFIt].Type), Layer.BSDFs[BSDFIt].SharedNormalIndex);
 				}
 			}
 
@@ -1584,6 +1587,16 @@ void FHLSLMaterialTranslator::GetSharedInputsMaterialCode(FString& PixelMembersD
 			PixelMembersInitializationEpilog += TranslatedCodeChunkDefinitions[LastProperty] + TEXT("\n");
 		}
 
+		if (CodeChunkToStrataSharedNormal.Num() > 0)
+		{
+			// Generate code to fill up the array containing normal shared between Strata BSDFs
+			for (TMap<uint64, FStrataSharedNormalInfo>::TIterator It(CodeChunkToStrataSharedNormal); It; ++It)
+			{
+				PixelInputInitializerValues += FString::Printf(TEXT("\tParameters.SharedNormals.Normals[%u] = %s;\n"), It->Value.SharedNormalIndex, *It->Value.SharedNormalCode);
+			}
+			PixelInputInitializerValues += FString::Printf(TEXT("\tParameters.SharedNormals.NormalCount = %u;\n"), CodeChunkToStrataSharedNormal.Num());
+		}
+
 		PixelMembersInitializationEpilog += PixelInputInitializerValues;
 	}
 }
@@ -7229,61 +7242,90 @@ const FStrataMaterialCompilationInfo& FHLSLMaterialTranslator::GetStrataCompilat
 	return CodeChunkToStrataCompilationInfoMap[CodeChunk];
 }
 
+uint8 FHLSLMaterialTranslator::GetStrataSharedNormalIndex(int32 NormalCodeChunk)
+{
+	check(NormalCodeChunk != INDEX_NONE);
+	check(NextFreeStrataShaderNormalIndex < 255);	// Out of shared normal slots
+	check(NextFreeStrataShaderNormalIndex < STRATA_MAX_TOTAL_BSDF);	// This is our current budget if normal per material. STRATA_TODO change that to a byte per pixel
+
+	const uint64 NormalCodeChunkHash = GetParameterHash(NormalCodeChunk);
+
+	if (CodeChunkToStrataSharedNormal.Find(NormalCodeChunkHash) == nullptr)
+	{
+		// Allocate a new slot for a new shared normal
+		uint8 NewSharedNormalIndex = NextFreeStrataShaderNormalIndex++;
+		CodeChunkToStrataSharedNormal.Add(NormalCodeChunkHash, { NormalCodeChunk , NewSharedNormalIndex, *GetParameterCode(NormalCodeChunk) });
+		return NewSharedNormalIndex;
+	}
+	// Return the existing code chunk
+	return CodeChunkToStrataSharedNormal[NormalCodeChunkHash].SharedNormalIndex;
+}
+
+uint8 FHLSLMaterialTranslator::GetStrataSharedNormalCount()
+{
+	return CodeChunkToStrataSharedNormal.Num();
+}
+
 int32 FHLSLMaterialTranslator::FrontMaterial()
 {
 	return AddInlinedCodeChunk(MCT_Strata, TEXT("GetInitialisedStrataData()"));
 }
 
-int32 FHLSLMaterialTranslator::StrataDiffuseOrenNayarBSDF(int32 Albedo, int32 Roughness, int32 Normal)
+int32 FHLSLMaterialTranslator::StrataDiffuseOrenNayarBSDF(int32 Albedo, int32 Roughness, int32 Normal, uint8 SharedNormalIndex)
 {
 	return AddCodeChunk(
-		MCT_Strata, TEXT("GetStrataDiffuseOrenNayarBSDF(%s, %s, %s)"),
+		MCT_Strata, TEXT("GetStrataDiffuseOrenNayarBSDF(%s, %s, %u) /* %s */"),
 		*GetParameterCode(Albedo),
 		*GetParameterCode(Roughness),
+		SharedNormalIndex,
 		*GetParameterCode(Normal)
 	);
 }
 
-int32 FHLSLMaterialTranslator::StrataDiffuseChanBSDF(int32 Albedo, int32 Roughness, int32 Normal)
+int32 FHLSLMaterialTranslator::StrataDiffuseChanBSDF(int32 Albedo, int32 Roughness, int32 Normal, uint8 SharedNormalIndex)
 {
 	return AddCodeChunk(
-		MCT_Strata, TEXT("GetStrataDiffuseChanBSDF(%s, %s, %s)"),
+		MCT_Strata, TEXT("GetStrataDiffuseChanBSDF(%s, %s, %u) /* %s */"),
 		*GetParameterCode(Albedo),
 		*GetParameterCode(Roughness),
+		SharedNormalIndex,
 		*GetParameterCode(Normal)
 	);
 }
 
-int32 FHLSLMaterialTranslator::StrataDielectricBSDF(int32 Roughness, int32 IOR, int32 Tint, int32 Normal)
+int32 FHLSLMaterialTranslator::StrataDielectricBSDF(int32 Roughness, int32 IOR, int32 Tint, int32 Normal, uint8 SharedNormalIndex)
 {
 	return AddCodeChunk(
-		MCT_Strata, TEXT("GetStrataDielectricBSDF(%s, %s, %s, %s)"),
+		MCT_Strata, TEXT("GetStrataDielectricBSDF(%s, %s, %s, %u) /* %s */"),
 		*GetParameterCode(Roughness),
 		*GetParameterCode(IOR),
 		*GetParameterCode(Tint),
+		SharedNormalIndex,
 		*GetParameterCode(Normal)
 	);
 }
 
-int32 FHLSLMaterialTranslator::StrataConductorBSDF(int32 Reflectivity, int32 EdgeColor, int32 Roughness, int32 Normal)
+int32 FHLSLMaterialTranslator::StrataConductorBSDF(int32 Reflectivity, int32 EdgeColor, int32 Roughness, int32 Normal, uint8 SharedNormalIndex)
 {
 	return AddCodeChunk(
-		MCT_Strata, TEXT("GetStrataConductorBSDF(%s, %s, %s, %s)"),
+		MCT_Strata, TEXT("GetStrataConductorBSDF(%s, %s, %s, %u) /* %s */"),
 		*GetParameterCode(Reflectivity),
 		*GetParameterCode(EdgeColor),
 		*GetParameterCode(Roughness),
+		SharedNormalIndex,
 		*GetParameterCode(Normal)
 	);
 }
 
-int32 FHLSLMaterialTranslator::StrataVolumeBSDF(int32 Albedo, int32 Extinction, int32 Anisotropy, int32 Thickness)
+int32 FHLSLMaterialTranslator::StrataVolumeBSDF(int32 Albedo, int32 Extinction, int32 Anisotropy, int32 Thickness, int32 Normal, uint8 SharedNormalIndex)
 {
 	return AddCodeChunk(
-		MCT_Strata, TEXT("GetStrataVolumeBSDF(%s, %s, %s, %s, %s)"),
+		MCT_Strata, TEXT("GetStrataVolumeBSDF(%s, %s, %s, %s, %u) /* %s */"),
 		*GetParameterCode(Albedo),
 		*GetParameterCode(Extinction),
 		*GetParameterCode(Anisotropy),
 		*GetParameterCode(Thickness),
+		SharedNormalIndex,
 		*GetParameterCode(VertexNormal())
 	);
 }
diff --git a/Engine/Source/Runtime/Engine/Private/Materials/HLSLMaterialTranslator.h b/Engine/Source/Runtime/Engine/Private/Materials/HLSLMaterialTranslator.h
--- a/Engine/Source/Runtime/Engine/Private/Materials/HLSLMaterialTranslator.h
+++ b/Engine/Source/Runtime/Engine/Private/Materials/HLSLMaterialTranslator.h
@@ -308,6 +308,14 @@ protected:
 	FMaterialShadingModelField ShadingModelsFromCompilation;
 
 	TMap<int32, FStrataMaterialCompilationInfo> CodeChunkToStrataCompilationInfoMap; // Mapping from node output to StrataMaterial
+	uint8 NextFreeStrataShaderNormalIndex;
+	struct FStrataSharedNormalInfo
+	{
+		int32 NormalCodeChunk;
+		uint8 SharedNormalIndex;
+		FString SharedNormalCode;
+	};
+	TMap<uint64, FStrataSharedNormalInfo> CodeChunkToStrataSharedNormal;
 
 	/** Tracks the total number of vt samples in the shader. */
 	uint32 NumVtSamples;
@@ -779,11 +787,11 @@ protected:
 
 	// Strata
 	virtual int32 FrontMaterial() override;
-	virtual int32 StrataDiffuseOrenNayarBSDF(int32 Albedo, int32 Roughness, int32 Normal) override;
-	virtual int32 StrataDiffuseChanBSDF(int32 Albedo, int32 Roughness, int32 Normal) override;
-	virtual int32 StrataDielectricBSDF(int32 Roughness, int32 IOR, int32 Tint, int32 Normal) override;
-	virtual int32 StrataConductorBSDF(int32 Reflectivity, int32 EdgeColor, int32 Roughness, int32 Normal) override;
-	virtual int32 StrataVolumeBSDF(int32 Albedo, int32 Extinction, int32 Anisotropy, int32 Thickness) override;
+	virtual int32 StrataDiffuseOrenNayarBSDF(int32 Albedo, int32 Roughness, int32 Normal, uint8 SharedNormalIndex) override;
+	virtual int32 StrataDiffuseChanBSDF(int32 Albedo, int32 Roughness, int32 Normal, uint8 SharedNormalIndex) override;
+	virtual int32 StrataDielectricBSDF(int32 Roughness, int32 IOR, int32 Tint, int32 Normal, uint8 SharedNormalIndex) override;
+	virtual int32 StrataConductorBSDF(int32 Reflectivity, int32 EdgeColor, int32 Roughness, int32 Normal, uint8 SharedNormalIndex) override;
+	virtual int32 StrataVolumeBSDF(int32 Albedo, int32 Extinction, int32 Anisotropy, int32 Thickness, int32 Normal, uint8 SharedNormalIndex) override;
 	virtual int32 StrataHorizontalMixing(int32 Foreground, int32 Background, int32 Mix) override;
 	virtual int32 StrataVerticalLayering(int32 Top, int32 Base) override;
 	virtual int32 StrataAdd(int32 A, int32 B) override;
@@ -794,6 +802,8 @@ protected:
 	virtual void AddStrataCodeChunk(int32 CodeChunk, FStrataMaterialCompilationInfo& StrataMaterialCompilationInfo) override;
 	virtual bool ContainsStrataCodeChunk(int32 CodeChunk) override;
 	virtual const FStrataMaterialCompilationInfo& GetStrataCompilationInfo(int32 CodeChunk) override;
+	virtual uint8 GetStrataSharedNormalIndex(int32 NormalCodeChunk) override;
+	virtual uint8 GetStrataSharedNormalCount() override;
 
 #if HANDLE_CUSTOM_OUTPUTS_AS_MATERIAL_ATTRIBUTES
 	/** Used to translate code for custom output attributes such as ClearCoatBottomNormal */
diff --git a/Engine/Source/Runtime/Engine/Private/Materials/MaterialExpressions.cpp b/Engine/Source/Runtime/Engine/Private/Materials/MaterialExpressions.cpp
--- a/Engine/Source/Runtime/Engine/Private/Materials/MaterialExpressions.cpp
+++ b/Engine/Source/Runtime/Engine/Private/Materials/MaterialExpressions.cpp
@@ -19030,12 +19030,15 @@ UMaterialExpressionStrataDiffuseBSDF::UMaterialExpressionStrataDiffuseBSDF(const
 #if WITH_EDITOR
 int32 UMaterialExpressionStrataDiffuseBSDF::Compile(class FMaterialCompiler* Compiler, int32 OutputIndex)
 {
+	int32 NormalCodeChunk = Normal.GetTracedInput().Expression ? Normal.Compile(Compiler) : Compiler->PixelNormalWS();
+	uint8 SharedNormalIndex = StrataCreateSharedNormal(Compiler, NormalCodeChunk);
+
 	int32 OutputCodeChunk = Compiler->StrataDiffuseOrenNayarBSDF(
 		Albedo.GetTracedInput().Expression		? Albedo.Compile(Compiler)		: Compiler->Constant3(0.18f, 0.18f, 0.18f),
 		Roughness.GetTracedInput().Expression	? Roughness.Compile(Compiler)	: Compiler->Constant(0.0f),
-		Normal.GetTracedInput().Expression		? Normal.Compile(Compiler)		: Compiler->PixelNormalWS());
-
-	StrataCreateSingleBSDFMaterial(Compiler, OutputCodeChunk, STRATA_BSDF_TYPE_DIFFUSE_ON);
+		NormalCodeChunk,
+		SharedNormalIndex);
+	StrataCreateSingleBSDFMaterial(Compiler, OutputCodeChunk, SharedNormalIndex, STRATA_BSDF_TYPE_DIFFUSE_ON);
 
 	return OutputCodeChunk;
 }
@@ -19091,12 +19094,15 @@ UMaterialExpressionStrataDiffuseChanBSDF::UMaterialExpressionStrataDiffuseChanBS
 #if WITH_EDITOR
 int32 UMaterialExpressionStrataDiffuseChanBSDF::Compile(class FMaterialCompiler* Compiler, int32 OutputIndex)
 {
+	int32 NormalCodeChunk = Normal.GetTracedInput().Expression ? Normal.Compile(Compiler) : Compiler->PixelNormalWS();
+	uint8 SharedNormalIndex = StrataCreateSharedNormal(Compiler, NormalCodeChunk);
+
 	int32 OutputCodeChunk = Compiler->StrataDiffuseChanBSDF(
 		Albedo.GetTracedInput().Expression		? Albedo.Compile(Compiler)		: Compiler->Constant3(0.18f, 0.18f, 0.18f),
 		Roughness.GetTracedInput().Expression	? Roughness.Compile(Compiler)	: Compiler->Constant(0.0f),
-		Normal.GetTracedInput().Expression		? Normal.Compile(Compiler)		: Compiler->PixelNormalWS());
-
-	StrataCreateSingleBSDFMaterial(Compiler, OutputCodeChunk, STRATA_BSDF_TYPE_DIFFUSE_CHAN);
+		NormalCodeChunk,
+		SharedNormalIndex);
+	StrataCreateSingleBSDFMaterial(Compiler, OutputCodeChunk, SharedNormalIndex, STRATA_BSDF_TYPE_DIFFUSE_CHAN);
 
 	return OutputCodeChunk;
 }
@@ -19152,13 +19158,16 @@ UMaterialExpressionStrataDielectricBSDF::UMaterialExpressionStrataDielectricBSDF
 #if WITH_EDITOR
 int32 UMaterialExpressionStrataDielectricBSDF::Compile(class FMaterialCompiler* Compiler, int32 OutputIndex)
 {
+	int32 NormalCodeChunk = Normal.GetTracedInput().Expression ? Normal.Compile(Compiler) : Compiler->PixelNormalWS();
+	uint8 SharedNormalIndex = StrataCreateSharedNormal(Compiler, NormalCodeChunk);
+
 	int32 OutputCodeChunk = Compiler->StrataDielectricBSDF(
 		Roughness.GetTracedInput().Expression	? Roughness.Compile(Compiler)	: Compiler->Constant2(0.0f, 0.0f),
 		IOR.GetTracedInput().Expression			? IOR.Compile(Compiler)			: Compiler->Constant(1.5f),// Default to Glass
 		Tint.GetTracedInput().Expression		? Tint.Compile(Compiler)		: Compiler->Constant3(1.0f, 1.0f, 1.0f),
-		Normal.GetTracedInput().Expression		? Normal.Compile(Compiler)		: Compiler->PixelNormalWS());
-
-	StrataCreateSingleBSDFMaterial(Compiler, OutputCodeChunk, STRATA_BSDF_TYPE_DIELECTRIC);
+		NormalCodeChunk,
+		SharedNormalIndex);
+	StrataCreateSingleBSDFMaterial(Compiler, OutputCodeChunk, SharedNormalIndex, STRATA_BSDF_TYPE_DIELECTRIC);
 
 	return OutputCodeChunk;
 }
@@ -19217,13 +19226,16 @@ UMaterialExpressionStrataConductorBSDF::UMaterialExpressionStrataConductorBSDF(c
 #if WITH_EDITOR
 int32 UMaterialExpressionStrataConductorBSDF::Compile(class FMaterialCompiler* Compiler, int32 OutputIndex)
 {
+	int32 NormalCodeChunk = Normal.GetTracedInput().Expression ? Normal.Compile(Compiler) : Compiler->PixelNormalWS();
+	uint8 SharedNormalIndex = StrataCreateSharedNormal(Compiler, NormalCodeChunk);
+
 	int32 OutputCodeChunk = Compiler->StrataConductorBSDF(
 		Reflectivity.GetTracedInput().Expression	? Reflectivity.Compile(Compiler): Compiler->Constant3(0.947f, 0.776f, 0.371f),	// Default to Gold
 		EdgeColor.GetTracedInput().Expression		? EdgeColor.Compile(Compiler)	: Compiler->Constant3(1.000f, 0.982f, 0.753f),	// Default to Gold
 		Roughness.GetTracedInput().Expression		? Roughness.Compile(Compiler)	: Compiler->Constant2(0.0f, 0.0f),
-		Normal.GetTracedInput().Expression			? Normal.Compile(Compiler)		: Compiler->PixelNormalWS());
-
-	StrataCreateSingleBSDFMaterial(Compiler, OutputCodeChunk, STRATA_BSDF_TYPE_CONDUCTOR);
+		NormalCodeChunk,
+		SharedNormalIndex);
+	StrataCreateSingleBSDFMaterial(Compiler, OutputCodeChunk, SharedNormalIndex, STRATA_BSDF_TYPE_CONDUCTOR);
 
 	return OutputCodeChunk;
 }
@@ -19282,13 +19294,17 @@ UMaterialExpressionStrataVolumeBSDF::UMaterialExpressionStrataVolumeBSDF(const F
 #if WITH_EDITOR
 int32 UMaterialExpressionStrataVolumeBSDF::Compile(class FMaterialCompiler* Compiler, int32 OutputIndex)
 {
+	int32 NormalCodeChunk = Compiler->VertexNormal();
+	uint8 SharedNormalIndex = StrataCreateSharedNormal(Compiler, NormalCodeChunk);
+
 	int32 OutputCodeChunk = Compiler->StrataVolumeBSDF(
 		Albedo.GetTracedInput().Expression		? Albedo.Compile(Compiler)			: Compiler->Constant3(0.0f, 0.0f, 0.0f),
 		Extinction.GetTracedInput().Expression	? Extinction.Compile(Compiler)		: Compiler->Constant3(0.0f, 0.0f, 0.0f),
 		Anisotropy.GetTracedInput().Expression	? Anisotropy.Compile(Compiler)		: Compiler->Constant(0.0f),
-		Thickness.GetTracedInput().Expression	? Thickness.Compile(Compiler)		: Compiler->Constant(0.001f)); // default = 1mm
-
-	StrataCreateSingleBSDFMaterial(Compiler, OutputCodeChunk, STRATA_BSDF_TYPE_VOLUME);
+		Thickness.GetTracedInput().Expression	? Thickness.Compile(Compiler)		: Compiler->Constant(0.001f), // default = 1mm
+		NormalCodeChunk,
+		SharedNormalIndex);
+	StrataCreateSingleBSDFMaterial(Compiler, OutputCodeChunk, SharedNormalIndex, STRATA_BSDF_TYPE_VOLUME);
 
 	return OutputCodeChunk;
 }
diff --git a/Engine/Source/Runtime/Engine/Private/Materials/StrataMaterial.cpp b/Engine/Source/Runtime/Engine/Private/Materials/StrataMaterial.cpp
--- a/Engine/Source/Runtime/Engine/Private/Materials/StrataMaterial.cpp
+++ b/Engine/Source/Runtime/Engine/Private/Materials/StrataMaterial.cpp
@@ -39,12 +39,19 @@ static void UpdateTotalBSDFCount(FStrataMaterialCompilationInfo& StrataInfo)
 }
 
 
-void StrataCreateSingleBSDFMaterial(FMaterialCompiler* Compiler, int32 CodeChunk, uint8 BSDFType)
+uint8 StrataCreateSharedNormal(FMaterialCompiler* Compiler, int32 NormalCodeChunk)
+{
+	return Compiler->GetStrataSharedNormalIndex(NormalCodeChunk);
+}
+
+
+void StrataCreateSingleBSDFMaterial(FMaterialCompiler* Compiler, int32 CodeChunk, uint8 SharedNormalIndex, uint8 BSDFType)
 {
 	FStrataMaterialCompilationInfo StrataInfo;
 	StrataInfo.LayerCount = 1;
 	StrataInfo.Layers[0].BSDFCount = 1;
 	StrataInfo.Layers[0].BSDFs[0].Type = BSDFType;
+	StrataInfo.Layers[0].BSDFs[0].SharedNormalIndex = SharedNormalIndex;
 	UpdateTotalBSDFCount(StrataInfo);
 	Compiler->AddStrataCodeChunk(CodeChunk, StrataInfo);
 }
diff --git a/Engine/Source/Runtime/Engine/Private/Materials/StrataMaterial.h b/Engine/Source/Runtime/Engine/Private/Materials/StrataMaterial.h
--- a/Engine/Source/Runtime/Engine/Private/Materials/StrataMaterial.h
+++ b/Engine/Source/Runtime/Engine/Private/Materials/StrataMaterial.h
@@ -4,6 +4,7 @@
 
 #include "CoreMinimal.h"
 #include "StrataDefinitions.h"
+#include "Containers/Map.h"
 
 
 
@@ -16,6 +17,7 @@ struct FStrataMaterialCompilationInfo
 	struct FBSDF
 	{
 		uint8 Type;
+		uint8 SharedNormalIndex;
 	};
 
 	struct FLayer
@@ -39,7 +41,8 @@ struct FStrataMaterialCompilationInfo
 
 FString GetStrataBSDFName(uint8 BSDFType);
 
-void StrataCreateSingleBSDFMaterial(FMaterialCompiler* Compiler, int32 CodeChunk, uint8 BSDFType);
+uint8 StrataCreateSharedNormal(FMaterialCompiler* Compiler, int32 NormalCodeChunk);
+void StrataCreateSingleBSDFMaterial(FMaterialCompiler* Compiler, int32 CodeChunk, uint8 SharedNormalIndex, uint8 BSDFType);
 
 FStrataMaterialCompilationInfo StrataAdd(FMaterialCompiler* Compiler, const FStrataMaterialCompilationInfo& A, const FStrataMaterialCompilationInfo& B);
 FStrataMaterialCompilationInfo StrataMultiply(FMaterialCompiler* Compiler, const FStrataMaterialCompilationInfo& A);
diff --git a/Engine/Source/Runtime/Engine/Public/MaterialCompiler.h b/Engine/Source/Runtime/Engine/Public/MaterialCompiler.h
--- a/Engine/Source/Runtime/Engine/Public/MaterialCompiler.h
+++ b/Engine/Source/Runtime/Engine/Public/MaterialCompiler.h
@@ -391,11 +391,11 @@ public:
 
 	// Strata
 	virtual int32 FrontMaterial() = 0;
-	virtual int32 StrataDiffuseOrenNayarBSDF(int32 Albedo, int32 Roughness, int32 Normal) = 0;
-	virtual int32 StrataDiffuseChanBSDF(int32 Albedo, int32 Roughness, int32 Normal) = 0;
-	virtual int32 StrataDielectricBSDF(int32 Roughness, int32 IOR, int32 Tint, int32 Normal) = 0;
-	virtual int32 StrataConductorBSDF(int32 Reflectivity, int32 EdgeColor, int32 Roughness, int32 Normal) = 0;
-	virtual int32 StrataVolumeBSDF(int32 Albedo, int32 Extinction, int32 Anisotropy, int32 Thickness) = 0;
+	virtual int32 StrataDiffuseOrenNayarBSDF(int32 Albedo, int32 Roughness, int32 Normal, uint8 SharedNormalIndex) = 0;
+	virtual int32 StrataDiffuseChanBSDF(int32 Albedo, int32 Roughness, int32 Normal, uint8 SharedNormalIndex) = 0;
+	virtual int32 StrataDielectricBSDF(int32 Roughness, int32 IOR, int32 Tint, int32 Normal, uint8 SharedNormalIndex) = 0;
+	virtual int32 StrataConductorBSDF(int32 Reflectivity, int32 EdgeColor, int32 Roughness, int32 Normal, uint8 SharedNormalIndex) = 0;
+	virtual int32 StrataVolumeBSDF(int32 Albedo, int32 Extinction, int32 Anisotropy, int32 Thickness, int32 Normal, uint8 SharedNormalIndex) = 0;
 	virtual int32 StrataHorizontalMixing(int32 Foreground, int32 Background, int32 Mix) = 0;
 	virtual int32 StrataVerticalLayering(int32 Top, int32 Base) = 0;
 	virtual int32 StrataAdd(int32 A, int32 B) = 0;
@@ -406,6 +406,8 @@ public:
 	virtual void AddStrataCodeChunk(int32 CodeChunk, FStrataMaterialCompilationInfo& StrataMaterialCompilationInfo) = 0;
 	virtual bool ContainsStrataCodeChunk(int32 CodeChunk) = 0;
 	virtual const FStrataMaterialCompilationInfo& GetStrataCompilationInfo(int32 CodeChunk) = 0;
+	virtual uint8 GetStrataSharedNormalIndex(int32 NormalCodeChunk) = 0;
+	virtual uint8 GetStrataSharedNormalCount() = 0;
 
 	// Water
 	virtual int32 SceneDepthWithoutWater(int32 Offset, int32 ViewportUV, bool bUseOffset, float FallbackDepth) = 0;
@@ -801,29 +803,29 @@ public:
 		return Compiler->FrontMaterial();
 	}
 
-	virtual int32 StrataDiffuseOrenNayarBSDF(int32 Albedo, int32 Roughness, int32 Normal) override
+	virtual int32 StrataDiffuseOrenNayarBSDF(int32 Albedo, int32 Roughness, int32 Normal, uint8 SharedNormalIndex) override
 	{
-		return Compiler->StrataDiffuseOrenNayarBSDF(Albedo, Roughness, Normal);
+		return Compiler->StrataDiffuseOrenNayarBSDF(Albedo, Roughness, Normal, SharedNormalIndex);
 	}
 
-	virtual int32 StrataDiffuseChanBSDF(int32 Albedo, int32 Roughness, int32 Normal) override
+	virtual int32 StrataDiffuseChanBSDF(int32 Albedo, int32 Roughness, int32 Normal, uint8 SharedNormalIndex) override
 	{
-		return Compiler->StrataDiffuseChanBSDF(Albedo, Roughness, Normal);
+		return Compiler->StrataDiffuseChanBSDF(Albedo, Roughness, Normal, SharedNormalIndex);
 	}
 
-	virtual int32 StrataDielectricBSDF(int32 Roughness, int32 IOR, int32 Tint, int32 Normal) override
+	virtual int32 StrataDielectricBSDF(int32 Roughness, int32 IOR, int32 Tint, int32 Normal, uint8 SharedNormalIndex) override
 	{
-		return Compiler->StrataDielectricBSDF(Roughness, IOR, Tint, Normal);
+		return Compiler->StrataDielectricBSDF(Roughness, IOR, Tint, Normal, SharedNormalIndex);
 	}
 
-	virtual int32 StrataConductorBSDF(int32 Reflectivity, int32 EdgeColor, int32 Roughness, int32 Normal) override
+	virtual int32 StrataConductorBSDF(int32 Reflectivity, int32 EdgeColor, int32 Roughness, int32 Normal, uint8 SharedNormalIndex) override
 	{
-		return Compiler->StrataConductorBSDF(Reflectivity, EdgeColor, Roughness, Normal);
+		return Compiler->StrataConductorBSDF(Reflectivity, EdgeColor, Roughness, Normal, SharedNormalIndex);
 	}
 
-	virtual int32 StrataVolumeBSDF(int32 Albedo, int32 Extinction, int32 Anisotropy, int32 Thickness) override
+	virtual int32 StrataVolumeBSDF(int32 Albedo, int32 Extinction, int32 Anisotropy, int32 Thickness, int32 Normal, uint8 SharedNormalIndex) override
 	{
-		return Compiler->StrataVolumeBSDF(Albedo, Extinction, Anisotropy, Thickness);
+		return Compiler->StrataVolumeBSDF(Albedo, Extinction, Anisotropy, Thickness, Normal, SharedNormalIndex);
 	}
 
 	virtual int32 StrataHorizontalMixing(int32 Foreground, int32 Background, int32 Mix) override
@@ -871,6 +873,16 @@ public:
 		return Compiler->GetStrataCompilationInfo(CodeChunk);
 	}
 
+	virtual uint8 GetStrataSharedNormalIndex(int32 NormalCodeChunk) override
+	{
+		return Compiler->GetStrataSharedNormalIndex(NormalCodeChunk);
+	}
+
+	virtual uint8 GetStrataSharedNormalCount() override
+	{
+		return Compiler->GetStrataSharedNormalCount();
+	}
+
 protected:
 		
 	FMaterialCompiler* Compiler;
diff --git a/Engine/Source/Runtime/Renderer/Private/Strata/Strata.cpp b/Engine/Source/Runtime/Renderer/Private/Strata/Strata.cpp
--- a/Engine/Source/Runtime/Renderer/Private/Strata/Strata.cpp
+++ b/Engine/Source/Runtime/Renderer/Private/Strata/Strata.cpp
@@ -50,7 +50,7 @@ void InitialiseStrataFrameSceneData(FSceneRenderer& SceneRenderer, FRDGBuilder&
 
 		// Previous GBuffer when complete was 28bytes
 		// check out Strata.ush to see how this is computed
-		const uint32 MaterialConservativeByteCountPerPixel = 97u;
+		const uint32 MaterialConservativeByteCountPerPixel = 84u;
 		const uint32 RoundToValue = 4u;
 		StrataSceneData.MaxBytesPerPixel = FMath::DivideAndRoundUp(MaterialConservativeByteCountPerPixel, RoundToValue) * RoundToValue;
 	}
