Add anisotropy support to strata dieletric and conductor.

* Add an (optional) tangent member to dieletric & conductor if the roughnesses have some anisotropy.
* Add a strata node to convert a base roughness & an anisotropy factor into two X/Y roughness

diff --git a/Engine/Shaders/Private/BRDF.ush b/Engine/Shaders/Private/BRDF.ush
--- a/Engine/Shaders/Private/BRDF.ush
+++ b/Engine/Shaders/Private/BRDF.ush
@@ -484,6 +484,7 @@ void ModifyGGXAnisotropicNormalRoughness(float3 WorldTangent, float Anisotropy,
 	}
 }
 
+// Convert a roughness and an anisotropy factor into GGX alpha values respectively for the major and minor axis of the tangent frame
 void GetAnisotropicRoughness(float Alpha, float Anisotropy, out float ax, out float ay)
 {
 #if 1
@@ -498,6 +499,27 @@ void GetAnisotropicRoughness(float Alpha, float Anisotropy, out float ax, out fl
 #endif
 }
 
+// Convert a roughness and an anisotropy factor, into two roughnesses respectively for the major and minor axis of the tangent frame
+float2 GetAnisotropicRoughness(float Roughness, float Anisotropy)
+{
+	// Anisotropic parameters: ax and ay are the roughness along the tangent and bitangent	
+	// Kulla 2017, "Revisiting Physically Based Shading at Imageworks"
+	float2 Out = saturate(Roughness);
+	Anisotropy = clamp(Anisotropy, -1.0, 1.0);
+	Out.x = max(Roughness * sqrt(1.0 + Anisotropy), 0.001f);
+	Out.y = max(Roughness * sqrt(1.0 - Anisotropy), 0.001f);
+	return Out;
+}
+
+// Return anisotropy factor (+1: Tangent is major axis, -1: Bitangent is the major axis), and the original roughness
+void GetAnisotropicFactor(float RoughnessX, float RoughnessY, inout float Anisotropy, inout float OriginalRoughness)
+{
+	const float MinRoughness = 0.001f;
+	float r = Pow2(max(RoughnessX, MinRoughness) / max(RoughnessY, MinRoughness));
+	Anisotropy = (r - 1.0) / (r + 1.0);
+	OriginalRoughness = (RoughnessX + RoughnessY) / (sqrt(1.0 + Anisotropy) + sqrt(1.0 - Anisotropy));
+}
+
 #ifndef PreIntegratedGF
 Texture2D		PreIntegratedGF;
 SamplerState	PreIntegratedGFSampler;
diff --git a/Engine/Shaders/Private/ShadingCommon.ush b/Engine/Shaders/Private/ShadingCommon.ush
--- a/Engine/Shaders/Private/ShadingCommon.ush
+++ b/Engine/Shaders/Private/ShadingCommon.ush
@@ -101,9 +101,9 @@ float3 ComputeF0(float Specular, float3 BaseColor, float Metallic)
 	return lerp(DielectricSpecularToF0(Specular).xxx, BaseColor, Metallic.xxx);
 }
 
-float MakeRoughnessSafe(float Roughness)
+float MakeRoughnessSafe(float Roughness, float MinRoughness=0.001f)
 {
-	return clamp(Roughness, 0.001f, 1.0f);
+	return clamp(Roughness, MinRoughness, 1.0f);
 }
 
 // [Gulbrandsen, "Artist Friendly Metallic Fresnel"]
diff --git a/Engine/Shaders/Private/Strata/Strata.ush b/Engine/Shaders/Private/Strata/Strata.ush
--- a/Engine/Shaders/Private/Strata/Strata.ush
+++ b/Engine/Shaders/Private/Strata/Strata.ush
@@ -23,29 +23,49 @@
 
 
 
-
 ///////////////////////////////////////////////////////////////////////////////
 // BSDFs representation
-//		STRATA_MAX_TOTAL_BSDF * sizeof(FStrataConductor) +				// 4 * 8
+//		STRATA_MAX_TOTAL_BSDF * sizeof(FStrataConductor) +				// 4 * 12
 //		STRATA_MAX_TOTAL_BSDF * (sizeof(State) + sizeof(Weight))		// 4 * 8
 //		sizeof(NormalCount_BSDFCount)									// 1 * 4
 //		STRATA_MAX_TOTAL_BSDF * sizeof(Normal)							// 4 * 4
-//		= 84															// See Strata.cpp
+//		= 100															// See Strata.cpp
 struct FStrataBSDF
 {
-	uint				State;		// (27 unused bits) (2bits as SharedNormalIndex) (3bits as BSDFType)
+	uint				State;		// (27 unused bits) (1bit as HasAnisotropy) (2bits as SharedNormalIndex) (3bits as BSDFType)
 
 	float3				Weight;		// Weight is a float3 because, when loaded from the material buffer, it contains the BSDF throughput (coverage and transmittance combined)
 
 	float4 VGPRs[3];
 };
 
-#define BSDF_GETTYPE(X)					uint(X.State & 0x7)
-#define BSDF_SETTYPE(X, NewType)		X.State = (X.State & 0xFFF8) | ((NewType & 0x7))
+// Num bits for each component store in BSDF state
+#define STATE_BITS_BSDF 3
+#define STATE_BITS_SHAREDNORMAL 2
+#define STATE_BITS_ANISOTROPY 1
+
+// Bit offset of each component store in BSDF state 
+#define STATE_BIT_OFFSET_BSDF			0
+#define STATE_BIT_OFFSET_SHAREDNORMAL	(STATE_BITS_BSDF)
+#define STATE_BIT_OFFSET_ANISOTROPY		(STATE_BITS_BSDF + STATE_BIT_OFFSET_SHAREDNORMAL)
+
+#define TO_BITMASK(BitCount)							((1<<BitCount)-1)
+#define BSDF_GETSTATE(X, BitCount, BitOffset)			uint((X.State>>BitOffset) & TO_BITMASK(BitCount))
+#define BSDF_SETSTATE(X, BitCount, BitOffset, Value)	X.State = (X.State & (~(TO_BITMASK(BitCount)<<BitOffset))) | ((Value & TO_BITMASK(BitCount)) << BitOffset)
+
+// Accessor for Type
+#define BSDF_GETTYPE(X)					BSDF_GETSTATE(X, STATE_BITS_BSDF, STATE_BIT_OFFSET_BSDF)
+#define BSDF_SETTYPE(X, NewType)		BSDF_SETSTATE(X, STATE_BITS_BSDF, STATE_BIT_OFFSET_BSDF, NewType)
 
+// Accessor for Shared Normal
 // This assumes Normals[STRATA_MAX_TOTAL_BSDF]
-#define BSDF_GETNORMALID(X)				uint((X.State>>3) & 0x3)
-#define BSDF_SETNORMALID(X, NrmlIdx)	X.State = (X.State & 0xFFE7) | ((NrmlIdx & 0x3) << 3)
+#define BSDF_GETNORMALID(X)				BSDF_GETSTATE(X, STATE_BITS_SHAREDNORMAL, STATE_BIT_OFFSET_SHAREDNORMAL)
+#define BSDF_SETNORMALID(X, NrmlIdx)	BSDF_SETSTATE(X, STATE_BITS_SHAREDNORMAL, STATE_BIT_OFFSET_SHAREDNORMAL, NrmlIdx)
+
+// Accessor for HasAnisotropy
+// Indicate of the BSDF has anisotrpic feature
+#define BSDF_GETANISOTROPY(X)			BSDF_GETSTATE(X, STATE_BITS_ANISOTROPY, STATE_BIT_OFFSET_ANISOTROPY)
+#define BSDF_SETANISOTROPY(X, Aniso)	BSDF_SETSTATE(X, STATE_BITS_ANISOTROPY, STATE_BIT_OFFSET_ANISOTROPY, Aniso)
 
 #define DIFFUSE_ALBEDO(X)				X.VGPRs[0].xyz
 #define DIFFUSE_ROUGHNESS(X)			X.VGPRs[0].w
@@ -54,17 +74,23 @@ struct FStrataBSDF
 #define DIELECTRIC_IOR(X)				X.VGPRs[0].w
 #define DIELECTRIC_ROUGHNESS0(X)		X.VGPRs[1].w
 #define DIELECTRIC_ROUGHNESS1(X)		X.VGPRs[2].w
+#define DIELECTRIC_TANGENT(X)			X.VGPRs[2].xyz
 
 #define CONDUCTOR_REFLECTIVITY(X)		X.VGPRs[0].xyz
 #define CONDUCTOR_EDGECOLOR(X)			X.VGPRs[1].xyz
 #define CONDUCTOR_ROUGHNESS0(X)			X.VGPRs[2].w
 #define CONDUCTOR_ROUGHNESS1(X)			X.VGPRs[1].w
+#define CONDUCTOR_TANGENT(X)			X.VGPRs[2].xyz
 
 #define VOLUME_ALBEDO(X)				X.VGPRs[0].xyz
 #define VOLUME_EXTINCTION(X)			X.VGPRs[1].xyz
 #define VOLUME_ANISOTROPY(X)			X.VGPRs[0].w
 #define VOLUME_THICKNESS(X)				X.VGPRs[1].w
 
+// Forward declaractions
+uint StrataPackNormal(in float3 Normal);
+float3 StrataUnpackNormal(uint PackedNormal);
+bool StrataHasAnisotropy(in float Roughness0, float Roughness1);
 
 struct FStrataLayer
 {
@@ -119,36 +145,44 @@ FStrataData GetStrataDiffuseChanBSDF(float3 Albedo, float Roughness, uint Shared
 	return StrataData;
 }
 
-FStrataData GetStrataDielectricBSDF(float2 Roughness, float IOR, float3 Tint, uint SharedNormalIndex)
+FStrataData GetStrataDielectricBSDF(float2 Roughness, float IOR, float3 Tint, float3 Tangent, uint SharedNormalIndex)
 {
 	FStrataData StrataData = GetInitialisedStrataData();
 	StrataData.LayerCount++;
 
+	const uint bHasAnisotropy = StrataHasAnisotropy(Roughness.x, Roughness.y) ? 1 : 0;
+
 	StrataData.Layers[0].BSDFCount++;
 	BSDF_SETTYPE			(StrataData.Layers[0].BSDFs[0],		STRATA_BSDF_TYPE_DIELECTRIC);
 	BSDF_SETNORMALID		(StrataData.Layers[0].BSDFs[0],		SharedNormalIndex);
+	BSDF_SETANISOTROPY		(StrataData.Layers[0].BSDFs[0],		bHasAnisotropy);
 	DIELECTRIC_IOR			(StrataData.Layers[0].BSDFs[0])		= IOR;
 	DIELECTRIC_TINT			(StrataData.Layers[0].BSDFs[0])		= Tint;
 	DIELECTRIC_ROUGHNESS0	(StrataData.Layers[0].BSDFs[0])		= Roughness.x;
 	DIELECTRIC_ROUGHNESS1	(StrataData.Layers[0].BSDFs[0])		= Roughness.y;
+	DIELECTRIC_TANGENT		(StrataData.Layers[0].BSDFs[0])		= Tangent;
 
 	StrataData.Layers[0].BSDFs[0].Weight						= 1.0f;
 
 	return StrataData;
 }
 
-FStrataData GetStrataConductorBSDF(float3 Reflectivity, float3 EdgeColor, float2 Roughness, uint SharedNormalIndex)
+FStrataData GetStrataConductorBSDF(float3 Reflectivity, float3 EdgeColor, float2 Roughness, float3 Tangent, uint SharedNormalIndex)
 {
 	FStrataData StrataData = GetInitialisedStrataData();
 	StrataData.LayerCount++;
 
+	const uint bHasAnisotropy = StrataHasAnisotropy(Roughness.x, Roughness.y) ? 1 : 0;
+
 	StrataData.Layers[0].BSDFCount++;
 	BSDF_SETTYPE			(StrataData.Layers[0].BSDFs[0],		STRATA_BSDF_TYPE_CONDUCTOR);
 	BSDF_SETNORMALID		(StrataData.Layers[0].BSDFs[0],		SharedNormalIndex);
+	BSDF_SETANISOTROPY		(StrataData.Layers[0].BSDFs[0],		bHasAnisotropy);
 	CONDUCTOR_REFLECTIVITY	(StrataData.Layers[0].BSDFs[0])		= Reflectivity;
 	CONDUCTOR_EDGECOLOR		(StrataData.Layers[0].BSDFs[0])		= EdgeColor;
 	CONDUCTOR_ROUGHNESS0	(StrataData.Layers[0].BSDFs[0])		= Roughness.x;
 	CONDUCTOR_ROUGHNESS1	(StrataData.Layers[0].BSDFs[0])		= Roughness.y;
+	CONDUCTOR_TANGENT		(StrataData.Layers[0].BSDFs[0])		= Tangent;
 
 	StrataData.Layers[0].BSDFs[0].Weight						= 1.0f;
 
@@ -210,6 +244,8 @@ struct FStrataBSDFContext
 	FStrataBSDF BSDF;
 
 	float3 N;
+	float3 X;
+	float3 Y;
 	float3 V;
 	float3 R;
 	float3 H;
@@ -263,6 +299,25 @@ void StrataGetBSDFRoughness(in FStrataBSDF BSDF, inout float R)
 	}
 }
 
+void StrataGetBSDFTangent(in FStrataBSDF BSDF, inout float3 T)
+{
+	T = 0;
+	const uint BSDFType = BSDF_GETTYPE(BSDF);
+	switch (BSDFType)
+	{
+	case STRATA_BSDF_TYPE_DIELECTRIC:
+	{
+		T = DIELECTRIC_TANGENT(BSDF);
+		break;
+	}
+	case STRATA_BSDF_TYPE_CONDUCTOR:
+	{
+		T = CONDUCTOR_TANGENT(BSDF);
+		break;
+	}
+	}
+}
+
 // Return true if a BSDF has a custom area light integrator
 bool StrataHasAreaLightIntegrator(in FStrataBSDF BSDF)
 {
@@ -272,6 +327,11 @@ bool StrataHasAreaLightIntegrator(in FStrataBSDF BSDF)
 		BSDFType == STRATA_BSDF_TYPE_DIELECTRIC;
 }
 
+bool StrataHasAnisotropy(in float Roughness0, float Roughness1)
+{
+	return Roughness0 != Roughness1;
+}
+
 FStrataBSDFContext StrataCreateBSDFContext(FSharedNormals SharedNormals, FStrataBSDF BSDF, float3 V, float3 L)
 {
 	FStrataBSDFContext BSDFContext = (FStrataBSDFContext)0;
@@ -288,14 +348,25 @@ FStrataBSDFContext StrataCreateBSDFContext(FSharedNormals SharedNormals, FStrata
 	BSDFContext.H = normalize(BSDFContext.V + BSDFContext.L);
 
 	BSDFContext.Context = (BxDFContext)0;
-	Init(BSDFContext.Context, BSDFContext.N, BSDFContext.V, BSDFContext.L);
-	BSDFContext.SatNoL = saturate(BSDFContext.Context.NoL);
-	BSDFContext.SatNoV = saturate(BSDFContext.Context.NoV);
-
-	BSDFContext.TangentBasis = GetTangentBasis(BSDFContext.N);
+	const bool bHasAnisotropy = BSDF_GETANISOTROPY(BSDF) != 0;
+	if (bHasAnisotropy)
+	{
+		StrataGetBSDFTangent(BSDF, BSDFContext.X);
+		BSDFContext.Y = normalize(cross(BSDFContext.N, BSDFContext.X));
+		Init(BSDFContext.Context, BSDFContext.N, BSDFContext.X, BSDFContext.Y, BSDFContext.V, BSDFContext.L);
+		BSDFContext.TangentBasis = float3x3(BSDFContext.X, BSDFContext.Y, BSDFContext.N);
+	}
+	else
+	{
+		Init(BSDFContext.Context, BSDFContext.N, BSDFContext.V, BSDFContext.L);
+		BSDFContext.TangentBasis = GetTangentBasis(BSDFContext.N);
+	}
 	BSDFContext.TangentV = mul(BSDFContext.TangentBasis, BSDFContext.V);
 	BSDFContext.TangentH = mul(BSDFContext.TangentBasis, BSDFContext.H);
 
+	BSDFContext.SatNoL = saturate(BSDFContext.Context.NoL);
+	BSDFContext.SatNoV = saturate(BSDFContext.Context.NoV);
+
 	return BSDFContext;
 }
 
@@ -348,17 +419,36 @@ FStrataEvaluateResult StrataEvaluateBSDF(FStrataBSDFContext BSDFContext)
 			const bool bIsDieletric = BSDFType == STRATA_BSDF_TYPE_DIELECTRIC;
 			const float3 F0  = bIsDieletric ? DIELECTRIC_TINT(BSDFContext.BSDF) * DielectricIorToF0(DIELECTRIC_IOR(BSDFContext.BSDF)) : CONDUCTOR_REFLECTIVITY(BSDFContext.BSDF);
 			const float3 F90 = bIsDieletric ? DIELECTRIC_TINT(BSDFContext.BSDF) : CONDUCTOR_EDGECOLOR(BSDFContext.BSDF);
-			const float SafeRoughness = MakeRoughnessSafe(bIsDieletric ? DIELECTRIC_ROUGHNESS0(BSDFContext.BSDF) : CONDUCTOR_ROUGHNESS0(BSDFContext.BSDF));
-
-			const float a2 = BSDFContext.a2Override>0 ? BSDFContext.a2Override : Pow4(SafeRoughness);
+			const float SafeRoughness0 = MakeRoughnessSafe(bIsDieletric ? DIELECTRIC_ROUGHNESS0(BSDFContext.BSDF) : CONDUCTOR_ROUGHNESS0(BSDFContext.BSDF));
+			const float SafeRoughness1 = MakeRoughnessSafe(bIsDieletric ? DIELECTRIC_ROUGHNESS1(BSDFContext.BSDF) : CONDUCTOR_ROUGHNESS1(BSDFContext.BSDF));
+
+			const bool bHasAnisotropy = BSDF_GETANISOTROPY(BSDFContext.BSDF);
+			float D = 0;
+			float Vis = 0;
+			float PDF = 0;
+			if (bHasAnisotropy)
+			{
+				// Generalized microfacet specular
+				float Alpha0 = Pow2(SafeRoughness0);
+				float Alpha1 = Pow2(SafeRoughness1);
+				D = D_GGXaniso(Alpha0, Alpha1, BSDFContext.Context.NoH, BSDFContext.Context.XoH, BSDFContext.Context.YoH);
+				Vis = Vis_SmithJointAniso(Alpha0, Alpha1, BSDFContext.Context.NoV, BSDFContext.Context.NoL, BSDFContext.Context.XoV, BSDFContext.Context.XoL, BSDFContext.Context.YoV, BSDFContext.Context.YoL);
+				PDF = VisibleGGXPDF_aniso(BSDFContext.TangentV, BSDFContext.TangentH, float2(Alpha0, Alpha1));
+			}
+			else
+			{
+				// Special override for roughness for supporting area light integrator with Sphere/Tube/Disk light, which modifies/increase roughness.
+				const float a2 = BSDFContext.a2Override > 0 ? BSDFContext.a2Override : Pow4(SafeRoughness0);
 
-			// Generalized microfacet specular
-			float D = D_GGX(a2, BSDFContext.Context.NoH);
-			float Vis = Vis_SmithJointApprox(a2, BSDFContext.Context.NoV, BSDFContext.SatNoL);
+				// Generalized microfacet specular
+				D = D_GGX(a2, BSDFContext.Context.NoH);
+				Vis = Vis_SmithJointApprox(a2, BSDFContext.Context.NoV, BSDFContext.SatNoL);
+				PDF = VisibleGGXPDF(BSDFContext.TangentV, BSDFContext.TangentH, a2);
+			}
 			float3 F = F_Schlick(F0, F90, BSDFContext.Context.VoH);
 
 			Sample.Value		= D * Vis * F;
-			Sample.PDF			= VisibleGGXPDF(BSDFContext.TangentV, BSDFContext.TangentH, a2);
+			Sample.PDF			= PDF;
 			Sample.Throughput	= bIsDieletric ? (1.0f - F) : OpaqueBSDFThroughput; // STRATA_TODO (Kulla's (1-FDG))
 			break;
 		}
@@ -437,8 +527,20 @@ FStrataImportanceSampleResult StrataImportanceSampleBSDF(FSharedNormals SharedNo
 		case STRATA_BSDF_TYPE_CONDUCTOR:
 		{
 			const bool bIsDielectric = BSDFType == STRATA_BSDF_TYPE_DIELECTRIC;
-			float Roughness = bIsDielectric ? DIELECTRIC_ROUGHNESS0(BSDFContext.BSDF) : CONDUCTOR_ROUGHNESS0(BSDFContext.BSDF);
-			float4 GGXSample = ImportanceSampleVisibleGGX(UniformSampleDisk(E), Pow4(Roughness), BSDFContext.TangentV);
+			float SafeRoughness0 = MakeRoughnessSafe(bIsDielectric ? DIELECTRIC_ROUGHNESS0(BSDFContext.BSDF) : CONDUCTOR_ROUGHNESS0(BSDFContext.BSDF));
+			float SafeRoughness1 = MakeRoughnessSafe(bIsDielectric ? DIELECTRIC_ROUGHNESS1(BSDFContext.BSDF) : CONDUCTOR_ROUGHNESS1(BSDFContext.BSDF));
+
+			const bool bHasAnisotropy = BSDF_GETANISOTROPY(BSDFContext.BSDF);
+			float4 GGXSample = 0;
+			if (bHasAnisotropy)
+			{
+				GGXSample = ImportanceSampleVisibleGGX_aniso(UniformSampleDisk(E), float2(Pow2(SafeRoughness0), Pow2(SafeRoughness1)), BSDFContext.TangentV);
+			}
+			else
+			{
+				GGXSample = ImportanceSampleVisibleGGX(UniformSampleDisk(E), Pow4(SafeRoughness0), BSDFContext.TangentV);
+			}
+			
 			float3 WorldH = mul(GGXSample.xyz, BSDFContext.TangentBasis);
 			Sample.Direction = reflect(CameraVector, WorldH);
 			PDF = GGXSample.w;
@@ -495,13 +597,28 @@ FStrataEnvLightResult StrataEvaluateForEnvLight(FStrataBSDFContext BSDFContext,
 				const bool bIsDielectric = STRATA_BSDF_TYPE_DIELECTRIC == BSDFType;
 				const float3 F0  = bIsDielectric ? DIELECTRIC_TINT(BSDFContext.BSDF) * DielectricIorToF0(DIELECTRIC_IOR(BSDFContext.BSDF)) : CONDUCTOR_REFLECTIVITY(BSDFContext.BSDF);
 				const float3 F90 = bIsDielectric ? DIELECTRIC_TINT(BSDFContext.BSDF) : CONDUCTOR_EDGECOLOR(BSDFContext.BSDF);
-				const float SafeRoughness = MakeRoughnessSafe(bIsDielectric ? DIELECTRIC_ROUGHNESS0(BSDFContext.BSDF) : CONDUCTOR_ROUGHNESS0(BSDFContext.BSDF));
+				float SafeRoughness0 = MakeRoughnessSafe(bIsDielectric ? DIELECTRIC_ROUGHNESS0(BSDFContext.BSDF) : CONDUCTOR_ROUGHNESS0(BSDFContext.BSDF));
+				float SafeRoughness1 = MakeRoughnessSafe(bIsDielectric ? DIELECTRIC_ROUGHNESS1(BSDFContext.BSDF) : CONDUCTOR_ROUGHNESS1(BSDFContext.BSDF));
+
+				const bool bHasAnisotropy = BSDF_GETANISOTROPY(BSDFContext.BSDF);
+				if (bHasAnisotropy)
+				{
+					// Modified the BSDF normal (and roughness)
+					float Anisotropy = 0;
+					GetAnisotropicFactor(SafeRoughness0, SafeRoughness1, Anisotropy, SafeRoughness0);
+					ModifyGGXAnisotropicNormalRoughness(BSDFContext.X, Anisotropy, SafeRoughness0, BSDFContext.N, BSDFContext.V);
+
+					// Update context (only needs: NoL/SatNoL/R) with the new N
+					BSDFContext.Context.NoL = dot(BSDFContext.N, BSDFContext.L);
+					BSDFContext.SatNoL = saturate(BSDFContext.Context.NoL);
+					BSDFContext.R = 2 * dot(BSDFContext.V, BSDFContext.N) * BSDFContext.N - BSDFContext.V;
+				}
 
-				float3 EvalEnvBRDF = EnvBRDFApprox(F0, F90, SafeRoughness, BSDFContext.SatNoV);
+				float3 EvalEnvBRDF = EnvBRDFApprox(F0, F90, SafeRoughness0, BSDFContext.SatNoV);
 
-				StrataEnvLightResult.SpecularDirection = StrataGetOffSpecularPeakReflectionDir(BSDFContext.N, BSDFContext.R, SafeRoughness);
+				StrataEnvLightResult.SpecularDirection = StrataGetOffSpecularPeakReflectionDir(BSDFContext.N, BSDFContext.R, SafeRoughness0);
 				StrataEnvLightResult.SpecularWeight = EvalEnvBRDF;
-				StrataEnvLightResult.SpecularSafeRoughness = SafeRoughness;
+				StrataEnvLightResult.SpecularSafeRoughness = SafeRoughness0;
 			}
 			break;
 		}
@@ -862,8 +979,6 @@ uint GetStrataPixelDataByteOffset(uint2 PixelPos, uint ViewBufferSizeX, uint Str
 	return (PixelPos.y * ViewBufferSizeX + PixelPos.x) * StrataMaxBytesPerPixel;
 }
 
-
-
 void PackStrataOut(
 	RWByteAddressBuffer StrataBuffer, 
 	uint ByteOffset, 
@@ -895,7 +1010,7 @@ void PackStrataOut(
 		float3 LayerCoverage = 0.0f;
 
 		float3 AccumTopNormal = 0.0f;
-		float3 AccumRoughness = 0.0f;
+		float  AccumRoughness = 0.0f;
 		float  AccumRoughnessCount = 0.0f;
 		float3 AccumBaseColor = 0.0f;
 		const bool bTopLayer = l == 0;
@@ -907,6 +1022,7 @@ void PackStrataOut(
 
 			FStrataBSDF BSDF = Strata.Layers[l].BSDFs[i];
 			const float3 BSDFCoverage = BSDF.Weight;
+			const float  BSDFCoverageAvg = dot(BSDF.Weight,1.0f/3.0f);
 			BSDFCount++;
 
 			// STRATA_TODO apply the same normal transform as for the material
@@ -926,8 +1042,8 @@ void PackStrataOut(
 			{
 				if (bTopLayer)
 				{
-					AccumRoughness		+= BSDFCoverage * FullyRough;
-					AccumRoughnessCount += BSDFCoverage;
+					AccumRoughness		+= BSDFCoverageAvg * FullyRough;
+					AccumRoughnessCount += BSDFCoverageAvg;
 				}
 
 				LayerCoverage	+= BSDFCoverage;
@@ -939,8 +1055,8 @@ void PackStrataOut(
 			{
 				if (bTopLayer)
 				{
-					AccumRoughness		+= BSDFCoverage * FullyRough;
-					AccumRoughnessCount += BSDFCoverage;
+					AccumRoughness		+= BSDFCoverageAvg * FullyRough;
+					AccumRoughnessCount += BSDFCoverageAvg;
 				}
 
 				LayerCoverage	+= BSDFCoverage;
@@ -1001,8 +1117,8 @@ void PackStrataOut(
 				
 				if (bTopLayer)
 				{
-					AccumRoughness		+= BSDFCoverage * FullyRough;
-					AccumRoughnessCount += BSDFCoverage;
+					AccumRoughness		+= BSDFCoverageAvg * FullyRough;
+					AccumRoughnessCount += BSDFCoverageAvg;
 				}
 
 				LayerCoverage	+= BSDFCoverage;
@@ -1046,7 +1162,7 @@ void PackStrataOut(
 	uint SharedNormals_BSDFCount = (SharedNormals.NormalCount << 16) | BSDFCount;
 	STRATA_STORE_UINT1(SharedNormals_BSDFCount);
 
-	for (int i = 0; i < SharedNormals.NormalCount; ++i)
+	for (uint i = 0; i < SharedNormals.NormalCount; ++i)
 	{
 		STRATA_STORE_UINT1(StrataPackNormal(SharedNormals.Normals[i]));
 	}
@@ -1074,18 +1190,28 @@ void PackStrataOut(
 				// 4 bytes
 			}
 			break;
-			case STRATA_BSDF_TYPE_DIELECTRIC: //
+			case STRATA_BSDF_TYPE_DIELECTRIC:
 			{
 				STRATA_STORE_UINT1(PackR11G11B10F(DIELECTRIC_TINT(BSDF)));
 				STRATA_STORE_UINT1(PackR11G11B10F(float3(DIELECTRIC_ROUGHNESS0(BSDF), DIELECTRIC_ROUGHNESS1(BSDF), DIELECTRIC_IOR(BSDF))));
-				// 8 bytes
+				const bool bHasAnisotropy = BSDF_GETANISOTROPY(BSDF);
+				if (bHasAnisotropy)
+				{
+					STRATA_STORE_UINT1(StrataPackNormal(DIELECTRIC_TANGENT(BSDF)));
+				}
+				// 12 bytes
 			}
 			break;
 			case STRATA_BSDF_TYPE_CONDUCTOR:
 			{
 				STRATA_STORE_UINT1(PackAlbedoLinearToGamma2AlphaLinear(float4(CONDUCTOR_REFLECTIVITY(BSDF), CONDUCTOR_ROUGHNESS0(BSDF))));
 				STRATA_STORE_UINT1(PackAlbedoLinearToGamma2AlphaLinear(float4(CONDUCTOR_EDGECOLOR(BSDF),    CONDUCTOR_ROUGHNESS1(BSDF))));
-				// 8 bytes
+				const bool bHasAnisotropy = BSDF_GETANISOTROPY(BSDF);
+				if (bHasAnisotropy)
+				{
+					STRATA_STORE_UINT1(StrataPackNormal(DIELECTRIC_TANGENT(BSDF)));
+				}
+				// 12 bytes
 			}
 			break;
 			case STRATA_BSDF_TYPE_VOLUME:
@@ -1162,12 +1288,28 @@ FStrataBSDF UnpackStrataBSDFIn(ByteAddressBuffer StrataBuffer, inout uint ByteOf
 	{
 		StrataLoad_R11G11B10F			(StrataBuffer, ByteOffset, DIELECTRIC_TINT(OutBSDF));
 		StrataLoad_R11_G11_B10F			(StrataBuffer, ByteOffset, DIELECTRIC_ROUGHNESS0(OutBSDF), DIELECTRIC_ROUGHNESS1(OutBSDF), DIELECTRIC_IOR(OutBSDF));
+
+		const bool bHasAnisotropy = BSDF_GETANISOTROPY(OutBSDF);
+		if (bHasAnisotropy)
+		{
+			uint PackedTangent;
+			STRATA_LOAD_UINT1(PackedTangent);
+			DIELECTRIC_TANGENT(OutBSDF) = StrataUnpackNormal(PackedTangent);
+		}
 	}
 	break;
 	case STRATA_BSDF_TYPE_CONDUCTOR:
 	{
 		StrataLoad_AlbedoGamma2ToLinearAlphaLinear(StrataBuffer, ByteOffset, CONDUCTOR_REFLECTIVITY(OutBSDF),	CONDUCTOR_ROUGHNESS0(OutBSDF));
 		StrataLoad_AlbedoGamma2ToLinearAlphaLinear(StrataBuffer, ByteOffset, CONDUCTOR_EDGECOLOR(OutBSDF),		CONDUCTOR_ROUGHNESS1(OutBSDF));
+
+		const bool bHasAnisotropy = BSDF_GETANISOTROPY(OutBSDF);
+		if (bHasAnisotropy)
+		{
+			uint PackedTangent;
+			STRATA_LOAD_UINT1(PackedTangent);
+			CONDUCTOR_TANGENT(OutBSDF) = StrataUnpackNormal(PackedTangent);
+		}
 	}
 	break;
 	case STRATA_BSDF_TYPE_VOLUME:
diff --git a/Engine/Shaders/Private/Strata/StrataDeferredLighting.ush b/Engine/Shaders/Private/Strata/StrataDeferredLighting.ush
--- a/Engine/Shaders/Private/Strata/StrataDeferredLighting.ush
+++ b/Engine/Shaders/Private/Strata/StrataDeferredLighting.ush
@@ -58,55 +58,62 @@ void StrataPreprocessAreaLightIntegration(
 	if (bIsValid && IsAreaLight(LightContext.AreaLight) && StrataHasAreaLightIntegrator(BSDFContext.BSDF))
 	{
 		// Evaluate roughness / specular color
-		float SafeRoughness = 0;
+		float SafeRoughness0 = 0;
+		float SafeRoughness1 = 0;
 		float3 SpecularColor = 0;
 		switch (BSDF_GETTYPE(BSDFContext.BSDF))
 		{
 			case STRATA_BSDF_TYPE_DIELECTRIC:
 			{
-				SafeRoughness = MakeRoughnessSafe(DIELECTRIC_ROUGHNESS0(BSDFContext.BSDF));
+				SafeRoughness0 = MakeRoughnessSafe(DIELECTRIC_ROUGHNESS0(BSDFContext.BSDF), View.MinRoughness);
+				SafeRoughness1 = MakeRoughnessSafe(DIELECTRIC_ROUGHNESS1(BSDFContext.BSDF), View.MinRoughness);
 				SpecularColor = DielectricIorToF0(DIELECTRIC_IOR(BSDFContext.BSDF)) * DIELECTRIC_TINT(BSDFContext.BSDF);
 			}
 			break;
 			case STRATA_BSDF_TYPE_CONDUCTOR:
 			{
-				SafeRoughness = MakeRoughnessSafe(CONDUCTOR_ROUGHNESS0(BSDFContext.BSDF));
+				SafeRoughness0 = MakeRoughnessSafe(CONDUCTOR_ROUGHNESS0(BSDFContext.BSDF), View.MinRoughness);
+				SafeRoughness1 = MakeRoughnessSafe(CONDUCTOR_ROUGHNESS1(BSDFContext.BSDF), View.MinRoughness);
 				SpecularColor = CONDUCTOR_REFLECTIVITY(BSDFContext.BSDF);
 			}
 			break;
-		}		
-		
-		// Evaluate area-light preintegration
-		SafeRoughness = max(SafeRoughness, View.MinRoughness);
-		if (LightContext.AreaLight.bIsRect)
-		{
-			// Rect light uses LTC integrator, which computes the full BDSD x Area light integration. 
-			// This is why there is no need to evaluate the BSDF. The result is stored into a local 
-			// OutgoingRadiance, which will be accumulated outside of this function
-			Out.OutgoingRadiance.rgb += 
-				Out.IncidentRadiance * 
-				RectGGXApproxLTC(SafeRoughness, SpecularColor, BSDFContext.N, BSDFContext.V, LightContext.AreaLight.Rect, LightContext.AreaLight.Texture);
-			Out.bEvaluateBSDF = false;
 		}
-		else
+
+		// Currently anistropic material don't support area lighting and will fallback on punctual lighting (that is a2Override == 0)
+		const bool bAnistropy = BSDF_GETANISOTROPY(BSDFContext.BSDF);
+		if (!bAnistropy)
 		{
-			// Sphere/Tube/Disk light use MRP integrator, which:
-			// 1. compute an energy rescaling factor (this is fold into the incident radiance)
-			// 2. change the BSDF roughness
-			// 3. compute a derived L 
-			// Update the BSDF context with the new L (and derived data) with value coming from MRP integrator
-			StrataUpdateBSDFContext(BSDFContext, LightContext.L);
-			SphereMaxNoH(BSDFContext.Context, LightContext.AreaLight.SphereSinAlpha, true);
-			BSDFContext.Context.NoV = saturate(abs(BSDFContext.Context.NoV) + 1e-5);
+			// Evaluate area-light preintegration
+			if (LightContext.AreaLight.bIsRect)
+			{
+				// Rect light uses LTC integrator, which computes the full BDSD x Area light integration. 
+				// This is why there is no need to evaluate the BSDF. The result is stored into a local 
+				// OutgoingRadiance, which will be accumulated outside of this function
+				Out.OutgoingRadiance.rgb += 
+					Out.IncidentRadiance * 
+					RectGGXApproxLTC(SafeRoughness0, SpecularColor, BSDFContext.N, BSDFContext.V, LightContext.AreaLight.Rect, LightContext.AreaLight.Texture);
+				Out.bEvaluateBSDF = false;
+			}
+			else
+			{
+				// Sphere/Tube/Disk light use MRP integrator, which:
+				// 1. compute an energy rescaling factor (this is fold into the incident radiance)
+				// 2. change the BSDF roughness
+				// 3. compute a derived L 
+				// Update the BSDF context with the new L (and derived data) with value coming from MRP integrator
+				StrataUpdateBSDFContext(BSDFContext, LightContext.L);
+				SphereMaxNoH(BSDFContext.Context, LightContext.AreaLight.SphereSinAlpha, true);
+				BSDFContext.Context.NoV = saturate(abs(BSDFContext.Context.NoV) + 1e-5);
 
-			// Update incident irradiance to add the energy rescaling to due MRP integration
-			// The a2 override stored the overriden roughness, which will be used during the BSDF evaluation
-			BSDFContext.a2Override = Pow4(SafeRoughness);
-			Out.IncidentIrradiance =
-				Out.IncidentRadiance *
-				EnergyNormalization(BSDFContext.a2Override, BSDFContext.Context.VoH, LightContext.AreaLight) *
-				LightContext.AreaLight.FalloffColor * (LightContext.Falloff * LightContext.NoL);
-			Out.bEvaluateBSDF = true;
+				// Update incident irradiance to add the energy rescaling to due MRP integration
+				// The a2 override stored the overriden roughness, which will be used during the BSDF evaluation
+				BSDFContext.a2Override = Pow4(SafeRoughness0);
+				Out.IncidentIrradiance =
+					Out.IncidentRadiance *
+					EnergyNormalization(BSDFContext.a2Override, BSDFContext.Context.VoH, LightContext.AreaLight) *
+					LightContext.AreaLight.FalloffColor * (LightContext.Falloff * LightContext.NoL);
+				Out.bEvaluateBSDF = true;
+			}
 		}
 	}
 }
diff --git a/Engine/Shaders/Private/Strata/StrataEnvironmentLighting.ush b/Engine/Shaders/Private/Strata/StrataEnvironmentLighting.ush
--- a/Engine/Shaders/Private/Strata/StrataEnvironmentLighting.ush
+++ b/Engine/Shaders/Private/Strata/StrataEnvironmentLighting.ush
@@ -58,6 +58,8 @@ float3 StrataEnvironmentLighting(
 			SkyLighting += BSDFThroughput * SpecularLuminance * StrataEnvLight.SpecularWeight * View.SkyLightColor.rgb * AmbientOcclusion;
 		}
 
+		// STRATA_TODO: Missing probe reflection here
+
 	}
 #if USE_PREEXPOSURE
 	SkyLighting *= View.PreExposure;
diff --git a/Engine/Source/Runtime/Engine/Classes/Materials/MaterialExpressionStrata.h b/Engine/Source/Runtime/Engine/Classes/Materials/MaterialExpressionStrata.h
--- a/Engine/Source/Runtime/Engine/Classes/Materials/MaterialExpressionStrata.h
+++ b/Engine/Source/Runtime/Engine/Classes/Materials/MaterialExpressionStrata.h
@@ -130,6 +130,12 @@ class UMaterialExpressionStrataDielectricBSDF : public UMaterialExpression
 	UPROPERTY()
 	FExpressionInput Normal;
 
+	/**
+	* Tangent (type = float3, unit = unitless)
+	*/
+	UPROPERTY()
+	FExpressionInput Tangent;
+
 	//~ Begin UMaterialExpression Interface
 #if WITH_EDITOR
 	virtual int32 Compile(class FMaterialCompiler* Compiler, int32 OutputIndex) override;
@@ -169,6 +175,12 @@ class UMaterialExpressionStrataConductorBSDF : public UMaterialExpression
 	UPROPERTY()
 	FExpressionInput Normal;
 
+	/**
+	 * Tangent (type = float3, unit = unitless)
+	 */
+	UPROPERTY()
+	FExpressionInput Tangent;
+
 	//~ Begin UMaterialExpression Interface
 #if WITH_EDITOR
 	virtual int32 Compile(class FMaterialCompiler* Compiler, int32 OutputIndex) override;
@@ -399,3 +411,30 @@ class UMaterialExpressionStrataPhysicalIOR : public UMaterialExpression
 };
 
 
+UCLASS(MinimalAPI, collapsecategories, hidecategories = Object)
+class UMaterialExpressionStrataAnisotropyToRoughness : public UMaterialExpression
+{
+	GENERATED_UCLASS_BODY()
+
+	/**
+	 * Input roughness
+	 */
+	UPROPERTY()
+	FExpressionInput Roughness;
+
+	/**
+	 * Anisotropy factor - 0: isotropic behavior, -1: anisotropy along the bitangent vector, 1: anisotropy along the tangent vector
+	 */
+	UPROPERTY()
+	FExpressionInput Anisotropy;
+
+	//~ Begin UMaterialExpression Interface
+#if WITH_EDITOR
+	virtual int32 Compile(class FMaterialCompiler* Compiler, int32 OutputIndex) override;
+	virtual void GetCaption(TArray<FString>& OutCaptions) const override;
+	virtual uint32 GetOutputType(int32 OutputIndex) override;
+	virtual uint32 GetInputType(int32 InputIndex) override;
+#endif
+	//~ End UMaterialExpression Interface
+};
+
diff --git a/Engine/Source/Runtime/Engine/Private/Materials/HLSLMaterialTranslator.cpp b/Engine/Source/Runtime/Engine/Private/Materials/HLSLMaterialTranslator.cpp
--- a/Engine/Source/Runtime/Engine/Private/Materials/HLSLMaterialTranslator.cpp
+++ b/Engine/Source/Runtime/Engine/Private/Materials/HLSLMaterialTranslator.cpp
@@ -7293,27 +7293,31 @@ int32 FHLSLMaterialTranslator::StrataDiffuseChanBSDF(int32 Albedo, int32 Roughne
 	);
 }
 
-int32 FHLSLMaterialTranslator::StrataDielectricBSDF(int32 Roughness, int32 IOR, int32 Tint, int32 Normal, uint8 SharedNormalIndex)
+int32 FHLSLMaterialTranslator::StrataDielectricBSDF(int32 Roughness, int32 IOR, int32 Tint, int32 Normal, int32 Tangent, uint8 SharedNormalIndex)
 {
 	return AddCodeChunk(
-		MCT_Strata, TEXT("GetStrataDielectricBSDF(%s, %s, %s, %u) /* %s */"),
+		MCT_Strata, TEXT("GetStrataDielectricBSDF(%s, %s, %s, %s, %u) /* Normal:%s Tangent:%s */"),
 		*GetParameterCode(Roughness),
 		*GetParameterCode(IOR),
 		*GetParameterCode(Tint),
+		*GetParameterCode(Tangent),
 		SharedNormalIndex,
-		*GetParameterCode(Normal)
+		*GetParameterCode(Normal),
+		*GetParameterCode(Tangent)
 	);
 }
 
-int32 FHLSLMaterialTranslator::StrataConductorBSDF(int32 Reflectivity, int32 EdgeColor, int32 Roughness, int32 Normal, uint8 SharedNormalIndex)
+int32 FHLSLMaterialTranslator::StrataConductorBSDF(int32 Reflectivity, int32 EdgeColor, int32 Roughness, int32 Normal, int32 Tangent, uint8 SharedNormalIndex)
 {
 	return AddCodeChunk(
-		MCT_Strata, TEXT("GetStrataConductorBSDF(%s, %s, %s, %u) /* %s */"),
+		MCT_Strata, TEXT("GetStrataConductorBSDF(%s, %s, %s, %s, %u) /* Normal:%s Tangent:%s */"),
 		*GetParameterCode(Reflectivity),
 		*GetParameterCode(EdgeColor),
 		*GetParameterCode(Roughness),
+		*GetParameterCode(Tangent),
 		SharedNormalIndex,
-		*GetParameterCode(Normal)
+		*GetParameterCode(Normal),
+		*GetParameterCode(Tangent)
 	);
 }
 
@@ -7411,6 +7415,23 @@ int32 FHLSLMaterialTranslator::StrataPhysicalIOR(int32 IOR, int32 Extinction, in
 	);
 }
 
+int32 FHLSLMaterialTranslator::StrataAnisotropyToRoughness(int32 Roughness, int32 Anisotropy, int32 OutputIndex)
+{
+	if (Roughness == INDEX_NONE)
+	{
+		return INDEX_NONE;
+	}
+	else if (Anisotropy == INDEX_NONE)
+	{
+		return INDEX_NONE;
+	}
+	return AddCodeChunk(
+		MCT_Float2,
+		TEXT("GetAnisotropicRoughness(%s,%s)"),
+		*GetParameterCode(Roughness),
+		*GetParameterCode(Anisotropy));
+}
+
 int32 FHLSLMaterialTranslator::MapARPassthroughCameraUV(int32 UV)
 {
 	if (UV == INDEX_NONE)
diff --git a/Engine/Source/Runtime/Engine/Private/Materials/HLSLMaterialTranslator.h b/Engine/Source/Runtime/Engine/Private/Materials/HLSLMaterialTranslator.h
--- a/Engine/Source/Runtime/Engine/Private/Materials/HLSLMaterialTranslator.h
+++ b/Engine/Source/Runtime/Engine/Private/Materials/HLSLMaterialTranslator.h
@@ -789,8 +789,8 @@ protected:
 	virtual int32 FrontMaterial() override;
 	virtual int32 StrataDiffuseOrenNayarBSDF(int32 Albedo, int32 Roughness, int32 Normal, uint8 SharedNormalIndex) override;
 	virtual int32 StrataDiffuseChanBSDF(int32 Albedo, int32 Roughness, int32 Normal, uint8 SharedNormalIndex) override;
-	virtual int32 StrataDielectricBSDF(int32 Roughness, int32 IOR, int32 Tint, int32 Normal, uint8 SharedNormalIndex) override;
-	virtual int32 StrataConductorBSDF(int32 Reflectivity, int32 EdgeColor, int32 Roughness, int32 Normal, uint8 SharedNormalIndex) override;
+	virtual int32 StrataDielectricBSDF(int32 Roughness, int32 IOR, int32 Tint, int32 Normal, int32 Tangent, uint8 SharedNormalIndex) override;
+	virtual int32 StrataConductorBSDF(int32 Reflectivity, int32 EdgeColor, int32 Roughness, int32 Normal, int32 Tangent, uint8 SharedNormalIndex) override;
 	virtual int32 StrataVolumeBSDF(int32 Albedo, int32 Extinction, int32 Anisotropy, int32 Thickness, int32 Normal, uint8 SharedNormalIndex) override;
 	virtual int32 StrataHorizontalMixing(int32 Foreground, int32 Background, int32 Mix) override;
 	virtual int32 StrataVerticalLayering(int32 Top, int32 Base) override;
@@ -798,6 +798,7 @@ protected:
 	virtual int32 StrataMultiply(int32 A, int32 Weight) override;
 	virtual int32 StrataArtisticIOR(int32 Reflectivity, int32 EdgeColor, int32 OutputIndex) override;
 	virtual int32 StrataPhysicalIOR(int32 IOR, int32 Extinction, int32 OutputIndex) override;
+	virtual int32 StrataAnisotropyToRoughness(int32 Roughness, int32 Anisotropy, int32 OutputIndex) override;
 
 	virtual void StrataCompilationInfoRegisterCodeChunk(int32 CodeChunk, FStrataMaterialCompilationInfo& StrataMaterialCompilationInfo) override;
 	virtual bool StrataCompilationInfoContainsCodeChunk(int32 CodeChunk) override;
diff --git a/Engine/Source/Runtime/Engine/Private/Materials/MaterialExpressions.cpp b/Engine/Source/Runtime/Engine/Private/Materials/MaterialExpressions.cpp
--- a/Engine/Source/Runtime/Engine/Private/Materials/MaterialExpressions.cpp
+++ b/Engine/Source/Runtime/Engine/Private/Materials/MaterialExpressions.cpp
@@ -19159,6 +19159,7 @@ UMaterialExpressionStrataDielectricBSDF::UMaterialExpressionStrataDielectricBSDF
 int32 UMaterialExpressionStrataDielectricBSDF::Compile(class FMaterialCompiler* Compiler, int32 OutputIndex)
 {
 	int32 NormalCodeChunk = Normal.GetTracedInput().Expression ? Normal.Compile(Compiler) : Compiler->PixelNormalWS();
+	int32 TangentCodeChunk = Tangent.GetTracedInput().Expression ? Tangent.Compile(Compiler) : Compiler->VertexTangent();
 	uint8 SharedNormalIndex = StrataCompilationInfoCreateSharedNormal(Compiler, NormalCodeChunk);
 
 	int32 OutputCodeChunk = Compiler->StrataDielectricBSDF(
@@ -19166,6 +19167,7 @@ int32 UMaterialExpressionStrataDielectricBSDF::Compile(class FMaterialCompiler*
 		IOR.GetTracedInput().Expression			? IOR.Compile(Compiler)			: Compiler->Constant(1.5f),// Default to Glass
 		Tint.GetTracedInput().Expression		? Tint.Compile(Compiler)		: Compiler->Constant3(1.0f, 1.0f, 1.0f),
 		NormalCodeChunk,
+		TangentCodeChunk,
 		SharedNormalIndex);
 	StrataCompilationInfoCreateSingleBSDFMaterial(Compiler, OutputCodeChunk, SharedNormalIndex, STRATA_BSDF_TYPE_DIELECTRIC);
 
@@ -19198,6 +19200,9 @@ uint32 UMaterialExpressionStrataDielectricBSDF::GetInputType(int32 InputIndex)
 	case 3:
 		return MCT_Float3;
 		break;
+	case 4:
+		return MCT_Float3;
+		break;
 	}
 
 	check(false);
@@ -19227,6 +19232,7 @@ UMaterialExpressionStrataConductorBSDF::UMaterialExpressionStrataConductorBSDF(c
 int32 UMaterialExpressionStrataConductorBSDF::Compile(class FMaterialCompiler* Compiler, int32 OutputIndex)
 {
 	int32 NormalCodeChunk = Normal.GetTracedInput().Expression ? Normal.Compile(Compiler) : Compiler->PixelNormalWS();
+	int32 TangentCodeChunk = Tangent.GetTracedInput().Expression ? Tangent.Compile(Compiler) : Compiler->VertexTangent();
 	uint8 SharedNormalIndex = StrataCompilationInfoCreateSharedNormal(Compiler, NormalCodeChunk);
 
 	int32 OutputCodeChunk = Compiler->StrataConductorBSDF(
@@ -19234,6 +19240,7 @@ int32 UMaterialExpressionStrataConductorBSDF::Compile(class FMaterialCompiler* C
 		EdgeColor.GetTracedInput().Expression		? EdgeColor.Compile(Compiler)	: Compiler->Constant3(1.000f, 0.982f, 0.753f),	// Default to Gold
 		Roughness.GetTracedInput().Expression		? Roughness.Compile(Compiler)	: Compiler->Constant2(0.0f, 0.0f),
 		NormalCodeChunk,
+		TangentCodeChunk,
 		SharedNormalIndex);
 	StrataCompilationInfoCreateSingleBSDFMaterial(Compiler, OutputCodeChunk, SharedNormalIndex, STRATA_BSDF_TYPE_CONDUCTOR);
 
@@ -19266,6 +19273,9 @@ uint32 UMaterialExpressionStrataConductorBSDF::GetInputType(int32 InputIndex)
 	case 3:
 		return MCT_Float3;
 		break;
+	case 4:
+		return MCT_Float3;
+		break;
 	}
 
 	check(false);
@@ -19681,4 +19691,60 @@ uint32 UMaterialExpressionStrataPhysicalIOR::GetInputType(int32 InputIndex)
 
 
 
+
+UMaterialExpressionStrataAnisotropyToRoughness::UMaterialExpressionStrataAnisotropyToRoughness(const FObjectInitializer& ObjectInitializer)
+	: Super(ObjectInitializer)
+{
+	struct FConstructorStatics
+	{
+		FText NAME_Strata;
+		FConstructorStatics() : NAME_Strata(LOCTEXT("Strata", "Strata")) { }
+	};
+	static FConstructorStatics ConstructorStatics;
+#if WITH_EDITORONLY_DATA
+	MenuCategories.Add(ConstructorStatics.NAME_Strata);
+
+	bShowOutputNameOnPin = true;
+
+	Outputs.Reset();
+	Outputs.Add(FExpressionOutput(TEXT("Roughnesses")));
+#endif
+}
+
+#if WITH_EDITOR
+int32 UMaterialExpressionStrataAnisotropyToRoughness::Compile(class FMaterialCompiler* Compiler, int32 OutputIndex)
+{
+	return Compiler->StrataAnisotropyToRoughness(
+		Roughness.GetTracedInput().Expression ? Roughness.Compile(Compiler) : Compiler->Constant(0.5f),
+		Anisotropy.GetTracedInput().Expression ? Anisotropy.Compile(Compiler) : Compiler->Constant(0.f),
+		OutputIndex);
+}
+
+void UMaterialExpressionStrataAnisotropyToRoughness::GetCaption(TArray<FString>& OutCaptions) const
+{
+	OutCaptions.Add(TEXT("Strata Anisotropy to roughnesses"));
+}
+
+uint32 UMaterialExpressionStrataAnisotropyToRoughness::GetOutputType(int32 OutputIndex)
+{
+	return MCT_Float2;
+}
+
+uint32 UMaterialExpressionStrataAnisotropyToRoughness::GetInputType(int32 InputIndex)
+{
+	switch (InputIndex)
+	{
+	case 0:
+		return MCT_Float;
+		break;
+	case 1:
+		return MCT_Float;
+		break;
+	}
+
+	check(false);
+	return MCT_Float1;
+}
+#endif // WITH_EDITOR
+
 #undef LOCTEXT_NAMESPACE
diff --git a/Engine/Source/Runtime/Engine/Public/MaterialCompiler.h b/Engine/Source/Runtime/Engine/Public/MaterialCompiler.h
--- a/Engine/Source/Runtime/Engine/Public/MaterialCompiler.h
+++ b/Engine/Source/Runtime/Engine/Public/MaterialCompiler.h
@@ -393,8 +393,8 @@ public:
 	virtual int32 FrontMaterial() = 0;
 	virtual int32 StrataDiffuseOrenNayarBSDF(int32 Albedo, int32 Roughness, int32 Normal, uint8 SharedNormalIndex) = 0;
 	virtual int32 StrataDiffuseChanBSDF(int32 Albedo, int32 Roughness, int32 Normal, uint8 SharedNormalIndex) = 0;
-	virtual int32 StrataDielectricBSDF(int32 Roughness, int32 IOR, int32 Tint, int32 Normal, uint8 SharedNormalIndex) = 0;
-	virtual int32 StrataConductorBSDF(int32 Reflectivity, int32 EdgeColor, int32 Roughness, int32 Normal, uint8 SharedNormalIndex) = 0;
+	virtual int32 StrataDielectricBSDF(int32 Roughness, int32 IOR, int32 Tint, int32 Normal, int32 Tangent, uint8 SharedNormalIndex) = 0;
+	virtual int32 StrataConductorBSDF(int32 Reflectivity, int32 EdgeColor, int32 Roughness, int32 Normal, int32 Tangent, uint8 SharedNormalIndex) = 0;
 	virtual int32 StrataVolumeBSDF(int32 Albedo, int32 Extinction, int32 Anisotropy, int32 Thickness, int32 Normal, uint8 SharedNormalIndex) = 0;
 	virtual int32 StrataHorizontalMixing(int32 Foreground, int32 Background, int32 Mix) = 0;
 	virtual int32 StrataVerticalLayering(int32 Top, int32 Base) = 0;
@@ -402,6 +402,7 @@ public:
 	virtual int32 StrataMultiply(int32 A, int32 Weight) = 0;
 	virtual int32 StrataArtisticIOR(int32 Reflectivity, int32 EdgeColor, int32 OutputIndex) = 0;
 	virtual int32 StrataPhysicalIOR(int32 IOR, int32 Extinction, int32 OutputIndex) = 0;
+	virtual int32 StrataAnisotropyToRoughness(int32 Roughness, int32 Anisotropy, int32 OutputIndex) = 0;
 
 	virtual void StrataCompilationInfoRegisterCodeChunk(int32 CodeChunk, FStrataMaterialCompilationInfo& StrataMaterialCompilationInfo) = 0;
 	virtual bool StrataCompilationInfoContainsCodeChunk(int32 CodeChunk) = 0;
@@ -813,14 +814,14 @@ public:
 		return Compiler->StrataDiffuseChanBSDF(Albedo, Roughness, Normal, SharedNormalIndex);
 	}
 
-	virtual int32 StrataDielectricBSDF(int32 Roughness, int32 IOR, int32 Tint, int32 Normal, uint8 SharedNormalIndex) override
+	virtual int32 StrataDielectricBSDF(int32 Roughness, int32 IOR, int32 Tint, int32 Normal, int32 Tangent, uint8 SharedNormalIndex) override
 	{
-		return Compiler->StrataDielectricBSDF(Roughness, IOR, Tint, Normal, SharedNormalIndex);
+		return Compiler->StrataDielectricBSDF(Roughness, IOR, Tint, Normal, Tangent, SharedNormalIndex);
 	}
 
-	virtual int32 StrataConductorBSDF(int32 Reflectivity, int32 EdgeColor, int32 Roughness, int32 Normal, uint8 SharedNormalIndex) override
+	virtual int32 StrataConductorBSDF(int32 Reflectivity, int32 EdgeColor, int32 Roughness, int32 Normal, int32 Tangent, uint8 SharedNormalIndex) override
 	{
-		return Compiler->StrataConductorBSDF(Reflectivity, EdgeColor, Roughness, Normal, SharedNormalIndex);
+		return Compiler->StrataConductorBSDF(Reflectivity, EdgeColor, Roughness, Normal, Tangent, SharedNormalIndex);
 	}
 
 	virtual int32 StrataVolumeBSDF(int32 Albedo, int32 Extinction, int32 Anisotropy, int32 Thickness, int32 Normal, uint8 SharedNormalIndex) override
@@ -858,6 +859,11 @@ public:
 		return Compiler->StrataPhysicalIOR(IOR, Extinction, OutputIndex);
 	}
 
+	virtual int32 StrataAnisotropyToRoughness(int32 Roughness, int32 Anisotropy, int32 OutputIndex) override
+	{
+		return Compiler->StrataAnisotropyToRoughness(Roughness, Anisotropy, OutputIndex);
+	}
+
 	virtual void StrataCompilationInfoRegisterCodeChunk(int32 CodeChunk, FStrataMaterialCompilationInfo& StrataMaterialCompilationInfo) override
 	{
 		Compiler->StrataCompilationInfoRegisterCodeChunk(CodeChunk, StrataMaterialCompilationInfo);
diff --git a/Engine/Source/Runtime/Renderer/Private/Strata/Strata.cpp b/Engine/Source/Runtime/Renderer/Private/Strata/Strata.cpp
--- a/Engine/Source/Runtime/Renderer/Private/Strata/Strata.cpp
+++ b/Engine/Source/Runtime/Renderer/Private/Strata/Strata.cpp
@@ -50,7 +50,7 @@ void InitialiseStrataFrameSceneData(FSceneRenderer& SceneRenderer, FRDGBuilder&
 
 		// Previous GBuffer when complete was 28bytes
 		// check out Strata.ush to see how this is computed
-		const uint32 MaterialConservativeByteCountPerPixel = 84u;
+		const uint32 MaterialConservativeByteCountPerPixel = 100u;
 		const uint32 RoundToValue = 4u;
 		StrataSceneData.MaxBytesPerPixel = FMath::DivideAndRoundUp(MaterialConservativeByteCountPerPixel, RoundToValue) * RoundToValue;
 	}
