Fixed SSR to only affect the top layer.

diff --git a/Engine/Shaders/Private/ReflectionEnvironmentPixelShader.usf b/Engine/Shaders/Private/ReflectionEnvironmentPixelShader.usf
--- a/Engine/Shaders/Private/ReflectionEnvironmentPixelShader.usf
+++ b/Engine/Shaders/Private/ReflectionEnvironmentPixelShader.usf
@@ -363,22 +363,25 @@ void ReflectionEnvironmentSkyLighting(
 #if ENABLE_SKY_LIGHT
 	if (bStrataMaterial)
 	{
-		SkyLighting = StrataEnvironmentLighting(
+		// We use WorldNormal, Roughness and BaseColor (representing EnvBRDF) written by the strata top layer.
+		const float4 SSR = Texture2DSample(ScreenSpaceReflectionsTexture, ScreenSpaceReflectionsSampler, BufferUV);
+		const float3 EnvBRDF = GBuffer.BaseColor;
+		SkyLighting.rgb += SSR.rgb * EnvBRDF;
+
+		// SSR is computed for the top level by averaging BSDF components.
+		// This is the top level specular contribution factor according to areas.
+		const float TopLayerSpecularContributionFactor = 1.0f - SSR.a;
+
+		SkyLighting += StrataEnvironmentLighting(
 			BufferUV,
 			GBuffer,
 			WorldPosition,
 			CameraToPixel,
 			V,
 			AmbientOcclusion,
+			TopLayerSpecularContributionFactor,
 			Strata.MaterialLobesBuffer,
 			GetStrataPixelDataByteOffset(PixelPos, View.BufferSizeAndInvSize.x, Strata.MaxBytesPerPixel));
-
-#if 1
-		// We use WorldNormal, Roughness and BaseColor (representing EnvBRDF) written by the strata top layer.
-		const float4 SSR = Texture2DSample(ScreenSpaceReflectionsTexture, ScreenSpaceReflectionsSampler, BufferUV);
-		const float3 EnvBRDF = GBuffer.BaseColor;
-		SkyLighting.rgb = SkyLighting.rgb * (1.0-SSR.a) + SSR.rgb * EnvBRDF;
-#endif
 	}
 	else
 #endif
diff --git a/Engine/Shaders/Private/Strata/Strata.ush b/Engine/Shaders/Private/Strata/Strata.ush
--- a/Engine/Shaders/Private/Strata/Strata.ush
+++ b/Engine/Shaders/Private/Strata/Strata.ush
@@ -40,14 +40,16 @@ struct FStrataBSDF
 };
 
 // Num bits for each component store in BSDF state
-#define STATE_BITS_BSDF 3
-#define STATE_BITS_SHAREDNORMAL 2
-#define STATE_BITS_ANISOTROPY 1
+#define STATE_BITS_BSDF					3
+#define STATE_BITS_SHAREDNORMAL			2
+#define STATE_BITS_ANISOTROPY			1
+#define STATE_BITS_TOPLAYER				1
 
 // Bit offset of each component store in BSDF state 
 #define STATE_BIT_OFFSET_BSDF			0
 #define STATE_BIT_OFFSET_SHAREDNORMAL	(STATE_BITS_BSDF)
 #define STATE_BIT_OFFSET_ANISOTROPY		(STATE_BITS_BSDF + STATE_BIT_OFFSET_SHAREDNORMAL)
+#define STATE_BIT_OFFSET_TOPLAYER		(STATE_BITS_BSDF + STATE_BIT_OFFSET_SHAREDNORMAL + STATE_BIT_OFFSET_ANISOTROPY)
 
 #define TO_BITMASK(BitCount)							((1<<BitCount)-1)
 #define BSDF_GETSTATE(X, BitCount, BitOffset)			uint((X.State>>BitOffset) & TO_BITMASK(BitCount))
@@ -63,10 +65,14 @@ struct FStrataBSDF
 #define BSDF_SETNORMALID(X, NrmlIdx)	BSDF_SETSTATE(X, STATE_BITS_SHAREDNORMAL, STATE_BIT_OFFSET_SHAREDNORMAL, NrmlIdx)
 
 // Accessor for HasAnisotropy
-// Indicate of the BSDF has anisotrpic feature
+// Indicate if the BSDF has anisotropic feature
 #define BSDF_GETANISOTROPY(X)			BSDF_GETSTATE(X, STATE_BITS_ANISOTROPY, STATE_BIT_OFFSET_ANISOTROPY)
 #define BSDF_SETANISOTROPY(X, Aniso)	BSDF_SETSTATE(X, STATE_BITS_ANISOTROPY, STATE_BIT_OFFSET_ANISOTROPY, Aniso)
 
+// Indicate if the BSDF is part of the top layer
+#define BSDF_GETISTOPLAYER(X)			BSDF_GETSTATE(X, STATE_BITS_TOPLAYER, STATE_BIT_OFFSET_TOPLAYER)
+#define BSDF_SETISTOPLAYER(X, IsTop)	BSDF_SETSTATE(X, STATE_BITS_TOPLAYER, STATE_BIT_OFFSET_TOPLAYER, IsTop)
+
 #define DIFFUSE_ALBEDO(X)				X.VGPRs[0].xyz
 #define DIFFUSE_ROUGHNESS(X)			X.VGPRs[0].w
 
@@ -121,6 +127,7 @@ FStrataData GetStrataDiffuseOrenNayarBSDF(float3 Albedo, float Roughness, uint S
 	StrataData.Layers[0].BSDFCount++;
 	BSDF_SETTYPE			(StrataData.Layers[0].BSDFs[0],		STRATA_BSDF_TYPE_DIFFUSE_ON);
 	BSDF_SETNORMALID		(StrataData.Layers[0].BSDFs[0],		SharedNormalIndex);
+	BSDF_SETISTOPLAYER		(StrataData.Layers[0].BSDFs[0],		0);
 	DIFFUSE_ALBEDO			(StrataData.Layers[0].BSDFs[0])		= Albedo;
 	DIFFUSE_ROUGHNESS		(StrataData.Layers[0].BSDFs[0])		= Roughness;
 
@@ -137,6 +144,7 @@ FStrataData GetStrataDiffuseChanBSDF(float3 Albedo, float Roughness, uint Shared
 	StrataData.Layers[0].BSDFCount++;
 	BSDF_SETTYPE			(StrataData.Layers[0].BSDFs[0],		STRATA_BSDF_TYPE_DIFFUSE_CHAN);
 	BSDF_SETNORMALID		(StrataData.Layers[0].BSDFs[0],		SharedNormalIndex);
+	BSDF_SETISTOPLAYER		(StrataData.Layers[0].BSDFs[0],		0);
 	DIFFUSE_ALBEDO			(StrataData.Layers[0].BSDFs[0])		= Albedo;
 	DIFFUSE_ROUGHNESS		(StrataData.Layers[0].BSDFs[0])		= Roughness;
 
@@ -156,6 +164,7 @@ FStrataData GetStrataDielectricBSDF(float2 Roughness, float IOR, float3 Tint, fl
 	BSDF_SETTYPE			(StrataData.Layers[0].BSDFs[0],		STRATA_BSDF_TYPE_DIELECTRIC);
 	BSDF_SETNORMALID		(StrataData.Layers[0].BSDFs[0],		SharedNormalIndex);
 	BSDF_SETANISOTROPY		(StrataData.Layers[0].BSDFs[0],		bHasAnisotropy);
+	BSDF_SETISTOPLAYER		(StrataData.Layers[0].BSDFs[0],		0);
 	DIELECTRIC_IOR			(StrataData.Layers[0].BSDFs[0])		= IOR;
 	DIELECTRIC_TINT			(StrataData.Layers[0].BSDFs[0])		= Tint;
 	DIELECTRIC_ROUGHNESS0	(StrataData.Layers[0].BSDFs[0])		= Roughness.x;
@@ -178,6 +187,7 @@ FStrataData GetStrataConductorBSDF(float3 Reflectivity, float3 EdgeColor, float2
 	BSDF_SETTYPE			(StrataData.Layers[0].BSDFs[0],		STRATA_BSDF_TYPE_CONDUCTOR);
 	BSDF_SETNORMALID		(StrataData.Layers[0].BSDFs[0],		SharedNormalIndex);
 	BSDF_SETANISOTROPY		(StrataData.Layers[0].BSDFs[0],		bHasAnisotropy);
+	BSDF_SETISTOPLAYER		(StrataData.Layers[0].BSDFs[0],		0);
 	CONDUCTOR_REFLECTIVITY	(StrataData.Layers[0].BSDFs[0])		= Reflectivity;
 	CONDUCTOR_EDGECOLOR		(StrataData.Layers[0].BSDFs[0])		= EdgeColor;
 	CONDUCTOR_ROUGHNESS0	(StrataData.Layers[0].BSDFs[0])		= Roughness.x;
@@ -197,6 +207,7 @@ FStrataData GetStrataVolumeBSDF(float3 Albedo, float3 Extinction, float Anisotro
 	StrataData.Layers[0].BSDFCount++;
 	BSDF_SETTYPE(				StrataData.Layers[0].BSDFs[0],	STRATA_BSDF_TYPE_VOLUME);
 	BSDF_SETNORMALID(			StrataData.Layers[0].BSDFs[0],	SharedNormalIndex);
+	BSDF_SETISTOPLAYER		(StrataData.Layers[0].BSDFs[0],		0);
 	VOLUME_ALBEDO				(StrataData.Layers[0].BSDFs[0])	= Albedo;
 	VOLUME_EXTINCTION			(StrataData.Layers[0].BSDFs[0])	= Extinction;
 	VOLUME_ANISOTROPY			(StrataData.Layers[0].BSDFs[0])	= Anisotropy;
@@ -1129,6 +1140,9 @@ void PackStrataOut(
 
 			// Update the weights to take into account throughput towards the view from layers above
 			Strata.Layers[l].BSDFs[i].Weight *= Throughput;
+			
+			// Notify that the BSDF is at the top for SSR to only affect reflection there and not on the lower layers
+			BSDF_SETISTOPLAYER(Strata.Layers[l].BSDFs[i], bTopLayer ? 1 : 0);
 		}
 
 		// Now modify the gbuffer according to the top layer information
diff --git a/Engine/Shaders/Private/Strata/StrataEnvironmentLighting.ush b/Engine/Shaders/Private/Strata/StrataEnvironmentLighting.ush
--- a/Engine/Shaders/Private/Strata/StrataEnvironmentLighting.ush
+++ b/Engine/Shaders/Private/Strata/StrataEnvironmentLighting.ush
@@ -11,6 +11,7 @@ float3 StrataEnvironmentLighting(
 	float3 CameraToPixel,
 	float3 V,
 	float AmbientOcclusion,
+	float TopLayerSpecularContributionFactor,
 	ByteAddressBuffer StrataDataBuffer,
 	uint PixelStrataDataByteOffset)
 {
@@ -55,7 +56,8 @@ float3 StrataEnvironmentLighting(
 			float SkyAverageBrightness = 1.0f;
 			const float3 SpecularLuminance = GetSkyLightReflection(StrataEnvLight.SpecularDirection, StrataEnvLight.SpecularSafeRoughness, SkyAverageBrightness);
 
-			SkyLighting += BSDFThroughput * SpecularLuminance * StrataEnvLight.SpecularWeight * View.SkyLightColor.rgb * AmbientOcclusion;
+			SkyLighting += BSDFThroughput * SpecularLuminance * StrataEnvLight.SpecularWeight * View.SkyLightColor.rgb
+				* AmbientOcclusion * (BSDF_GETISTOPLAYER(BSDF) ? TopLayerSpecularContributionFactor : 1.0f);
 		}
 
 		// STRATA_TODO: Missing probe reflection here
