Minor comments in strata.ush

diff --git a/Engine/Shaders/Private/Strata/Strata.ush b/Engine/Shaders/Private/Strata/Strata.ush
--- a/Engine/Shaders/Private/Strata/Strata.ush
+++ b/Engine/Shaders/Private/Strata/Strata.ush
@@ -1053,7 +1053,7 @@ void PackStrataOut(
 			{
 				if (bTopLayer)
 				{
-					AccumRoughness		+= BSDFCoverageAvg * FullyRough;
+					AccumRoughness		+= BSDFCoverageAvg * FullyRough; // Do not trigger SSR by using fully rough
 					AccumRoughnessCount += BSDFCoverageAvg;
 				}
 
@@ -1066,7 +1066,7 @@ void PackStrataOut(
 			{
 				if (bTopLayer)
 				{
-					AccumRoughness		+= BSDFCoverageAvg * FullyRough;
+					AccumRoughness		+= BSDFCoverageAvg * FullyRough; // Do not trigger SSR by using fully rough
 					AccumRoughnessCount += BSDFCoverageAvg;
 				}
 
@@ -1128,7 +1128,7 @@ void PackStrataOut(
 				
 				if (bTopLayer)
 				{
-					AccumRoughness		+= BSDFCoverageAvg * FullyRough;
+					AccumRoughness		+= BSDFCoverageAvg * FullyRough; // Do not trigger SSR by using fully rough
 					AccumRoughnessCount += BSDFCoverageAvg;
 				}
 
