Removed Strata Diffuse Oren-Nayar.

diff --git a/Engine/Shaders/Private/Strata/Strata.ush b/Engine/Shaders/Private/Strata/Strata.ush
--- a/Engine/Shaders/Private/Strata/Strata.ush
+++ b/Engine/Shaders/Private/Strata/Strata.ush
@@ -119,30 +119,13 @@ FStrataData GetInitialisedStrataData()
 	return StrataData;
 }
 
-FStrataData GetStrataDiffuseOrenNayarBSDF(float3 Albedo, float Roughness, uint SharedNormalIndex)
+FStrataData GetStrataDiffuseBSDF(float3 Albedo, float Roughness, uint SharedNormalIndex)
 {
 	FStrataData StrataData = GetInitialisedStrataData();
 	StrataData.LayerCount++;
 
 	StrataData.Layers[0].BSDFCount++;
-	BSDF_SETTYPE			(StrataData.Layers[0].BSDFs[0],		STRATA_BSDF_TYPE_DIFFUSE_ON);
-	BSDF_SETNORMALID		(StrataData.Layers[0].BSDFs[0],		SharedNormalIndex);
-	BSDF_SETISTOPLAYER		(StrataData.Layers[0].BSDFs[0],		0);
-	DIFFUSE_ALBEDO			(StrataData.Layers[0].BSDFs[0])		= Albedo;
-	DIFFUSE_ROUGHNESS		(StrataData.Layers[0].BSDFs[0])		= Roughness;
-
-	StrataData.Layers[0].BSDFs[0].Weight						= 1.0f;
-
-	return StrataData;
-}
-
-FStrataData GetStrataDiffuseChanBSDF(float3 Albedo, float Roughness, uint SharedNormalIndex)
-{
-	FStrataData StrataData = GetInitialisedStrataData();
-	StrataData.LayerCount++;
-
-	StrataData.Layers[0].BSDFCount++;
-	BSDF_SETTYPE			(StrataData.Layers[0].BSDFs[0],		STRATA_BSDF_TYPE_DIFFUSE_CHAN);
+	BSDF_SETTYPE			(StrataData.Layers[0].BSDFs[0],		STRATA_BSDF_TYPE_DIFFUSE);
 	BSDF_SETNORMALID		(StrataData.Layers[0].BSDFs[0],		SharedNormalIndex);
 	BSDF_SETISTOPLAYER		(StrataData.Layers[0].BSDFs[0],		0);
 	DIFFUSE_ALBEDO			(StrataData.Layers[0].BSDFs[0])		= Albedo;
@@ -286,8 +269,7 @@ void StrataGetBSDFRoughness(in FStrataBSDF BSDF, inout float R)
 	const uint BSDFType = BSDF_GETTYPE(BSDF);
 	switch (BSDFType)
 	{
-	case STRATA_BSDF_TYPE_DIFFUSE_ON:
-	case STRATA_BSDF_TYPE_DIFFUSE_CHAN:
+	case STRATA_BSDF_TYPE_DIFFUSE:
 	{
 		R = DIFFUSE_ROUGHNESS(BSDF);
 		break;
@@ -405,16 +387,7 @@ FStrataEvaluateResult StrataEvaluateBSDF(FStrataBSDFContext BSDFContext)
 	const uint BSDFType = BSDF_GETTYPE(BSDFContext.BSDF);
 	switch (BSDFType)
 	{
-		case STRATA_BSDF_TYPE_DIFFUSE_ON:
-		{
-			const float SafeRoughness = MakeRoughnessSafe(DIFFUSE_ROUGHNESS(BSDFContext.BSDF));
-
-			Sample.Value		= Diffuse_OrenNayar(DIFFUSE_ALBEDO(BSDFContext.BSDF), SafeRoughness, BSDFContext.SatNoV, BSDFContext.SatNoL, BSDFContext.Context.VoH);
-			Sample.PDF			= BSDFContext.SatNoV / PI;
-			Sample.Throughput	= OpaqueBSDFThroughput;
-			break;
-		}
-		case STRATA_BSDF_TYPE_DIFFUSE_CHAN:
+		case STRATA_BSDF_TYPE_DIFFUSE:
 		{
 			const float SafeRoughness = MakeRoughnessSafe(DIFFUSE_ROUGHNESS(BSDFContext.BSDF));
 			float a2 = Pow4(SafeRoughness);
@@ -524,8 +497,7 @@ FStrataImportanceSampleResult StrataImportanceSampleBSDF(FSharedNormals SharedNo
 	const uint BSDFType = BSDF_GETTYPE(BSDF);
 	switch (BSDFType)
 	{
-		case STRATA_BSDF_TYPE_DIFFUSE_ON:
-		case STRATA_BSDF_TYPE_DIFFUSE_CHAN:
+		case STRATA_BSDF_TYPE_DIFFUSE:
 		{
 			// STRATA_TODO take into account roughness
 			float4 ImportanceSample = CosineSampleHemisphere(E);
@@ -592,8 +564,7 @@ FStrataEnvLightResult StrataEvaluateForEnvLight(FStrataBSDFContext BSDFContext,
 	const uint BSDFType = BSDF_GETTYPE(BSDFContext.BSDF);
 	switch (BSDFType)
 	{
-		case STRATA_BSDF_TYPE_DIFFUSE_ON:
-		case STRATA_BSDF_TYPE_DIFFUSE_CHAN:
+		case STRATA_BSDF_TYPE_DIFFUSE:
 		{
 			StrataEnvLightResult.DiffuseAlbedo = DIFFUSE_ALBEDO(BSDFContext.BSDF);
 			StrataEnvLightResult.DiffuseNormal = BSDFContext.N;
@@ -773,8 +744,7 @@ void StrataSanitizeBSDF(inout FStrataBSDF BSDF)
 {
 	switch (BSDF_GETTYPE(BSDF))
 	{
-		case STRATA_BSDF_TYPE_DIFFUSE_ON:
-		case STRATA_BSDF_TYPE_DIFFUSE_CHAN:
+		case STRATA_BSDF_TYPE_DIFFUSE:
 		{
 			SanitizeStrataDiffuse(BSDF);
 			break; 
@@ -1049,20 +1019,7 @@ void PackStrataOut(
 
 			switch (BSDF_GETTYPE(BSDF))
 			{
-			case STRATA_BSDF_TYPE_DIFFUSE_ON:
-			{
-				if (bTopLayer)
-				{
-					AccumRoughness		+= BSDFCoverageAvg * FullyRough; // Do not trigger SSR by using fully rough
-					AccumRoughnessCount += BSDFCoverageAvg;
-				}
-
-				LayerCoverage	+= BSDFCoverage;
-				LayerThroughput += BSDFCoverage * OpaqueBSDFThroughput;
-				break;
-			}
-
-			case STRATA_BSDF_TYPE_DIFFUSE_CHAN:
+			case STRATA_BSDF_TYPE_DIFFUSE:
 			{
 				if (bTopLayer)
 				{
@@ -1197,8 +1154,7 @@ void PackStrataOut(
 			const uint BSDFType = BSDF_GETTYPE(BSDF);
 			switch (BSDF_GETTYPE(BSDF))
 			{
-			case STRATA_BSDF_TYPE_DIFFUSE_ON:
-			case STRATA_BSDF_TYPE_DIFFUSE_CHAN:
+			case STRATA_BSDF_TYPE_DIFFUSE:
 			{
 				STRATA_STORE_UINT1(PackAlbedoLinearToGamma2AlphaLinear(float4(DIFFUSE_ALBEDO(BSDF), DIFFUSE_ROUGHNESS(BSDF))));
 				// 4 bytes
@@ -1292,8 +1248,7 @@ FStrataBSDF UnpackStrataBSDFIn(ByteAddressBuffer StrataBuffer, inout uint ByteOf
 	BSDF_SETTYPE(OutBSDF, BSDFType);
 	switch (BSDFType)
 	{
-	case STRATA_BSDF_TYPE_DIFFUSE_ON:
-	case STRATA_BSDF_TYPE_DIFFUSE_CHAN:
+	case STRATA_BSDF_TYPE_DIFFUSE:
 	{
 		StrataLoad_AlbedoGamma2ToLinearAlphaLinear(StrataBuffer, ByteOffset, DIFFUSE_ALBEDO(OutBSDF), DIFFUSE_ROUGHNESS(OutBSDF));
 	}
diff --git a/Engine/Shaders/Private/Strata/StrataVisualize.usf b/Engine/Shaders/Private/Strata/StrataVisualize.usf
--- a/Engine/Shaders/Private/Strata/StrataVisualize.usf
+++ b/Engine/Shaders/Private/Strata/StrataVisualize.usf
@@ -107,24 +107,10 @@ void VisualizeMaterialPS(
 
 			switch (BSDF_GETTYPE(BSDF))
 			{
-			case STRATA_BSDF_TYPE_DIFFUSE_ON:
-				PRINTBSDFCHAR(_D_); PRINTBSDFCHAR(_I_); PRINTBSDFCHAR(_F_); PRINTBSDFCHAR(_F_); PRINTBSDFCHAR(_O_); PRINTBSDFCHAR(_N_);
-				PRINTBSDFCHAR(_SPC_); PRINTBSDFCHAR(_SPC_); PRINTBSDFCHAR(_SPC_); PRINTBSDFCHAR(_SPC_); PRINTBSDFCHAR(_SPC_); PRINTBSDFCHAR(_SPC_);
-				PRINTBSDFWEIGHT();
-				NEXTLINE;
-
-				HORIZONTALTAB; PRINTPARAMCHAR(_A_); PRINTPARAMCHAR(_L_); PRINTPARAMCHAR(_B_); PRINTPARAMCHAR(_E_); PRINTPARAMCHAR(_D_); PRINTPARAMCHAR(_O_);
-				HORIZONTALTAB; PRINTFLOAT3(DIFFUSE_ALBEDO(BSDF));
-				NEXTLINE;
-				HORIZONTALTAB; PRINTPARAMCHAR(_R_); PRINTPARAMCHAR(_O_); PRINTPARAMCHAR(_U_); PRINTPARAMCHAR(_G_); PRINTPARAMCHAR(_H_); PRINTPARAMCHAR(_SPC_);
-				HORIZONTALTAB; PRINTFLOAT1(DIFFUSE_ROUGHNESS(BSDF));
-				NEXTLINE;
-
-				break;
 
-			case STRATA_BSDF_TYPE_DIFFUSE_CHAN:
+			case STRATA_BSDF_TYPE_DIFFUSE:
 				PRINTBSDFCHAR(_D_); PRINTBSDFCHAR(_I_); PRINTBSDFCHAR(_F_); PRINTBSDFCHAR(_F_);
-				PRINTBSDFCHAR(_C_); PRINTBSDFCHAR(_H_); PRINTBSDFCHAR(_A_); PRINTBSDFCHAR(_N_);
+				PRINTBSDFCHAR(_U_); PRINTBSDFCHAR(_S_); PRINTBSDFCHAR(_E_); PRINTBSDFCHAR(_SPC_);
 				PRINTBSDFWEIGHT();
 				NEXTLINE;
 
diff --git a/Engine/Shaders/Shared/StrataDefinitions.h b/Engine/Shaders/Shared/StrataDefinitions.h
--- a/Engine/Shaders/Shared/StrataDefinitions.h
+++ b/Engine/Shaders/Shared/StrataDefinitions.h
@@ -18,10 +18,9 @@
 
 
 
-#define STRATA_BSDF_TYPE_DIFFUSE_ON			0
-#define STRATA_BSDF_TYPE_DIFFUSE_CHAN		1
-#define STRATA_BSDF_TYPE_DIELECTRIC			2
-#define STRATA_BSDF_TYPE_CONDUCTOR			3
-#define STRATA_BSDF_TYPE_VOLUME				4
+#define STRATA_BSDF_TYPE_DIFFUSE			0
+#define STRATA_BSDF_TYPE_DIELECTRIC			1
+#define STRATA_BSDF_TYPE_CONDUCTOR			2
+#define STRATA_BSDF_TYPE_VOLUME				3
 
 
diff --git a/Engine/Source/Runtime/Engine/Classes/Materials/MaterialExpressionStrata.h b/Engine/Source/Runtime/Engine/Classes/Materials/MaterialExpressionStrata.h
--- a/Engine/Source/Runtime/Engine/Classes/Materials/MaterialExpressionStrata.h
+++ b/Engine/Source/Runtime/Engine/Classes/Materials/MaterialExpressionStrata.h
@@ -68,39 +68,6 @@ class UMaterialExpressionStrataDiffuseBSDF : public UMaterialExpression // STRAT
 	//~ End UMaterialExpression Interface
 };
 
-UCLASS(MinimalAPI, collapsecategories, hidecategories = Object)
-class UMaterialExpressionStrataDiffuseChanBSDF : public UMaterialExpression
-{
-	GENERATED_UCLASS_BODY()
-
-	/**
-	 * Albedo (type = float3, unit = unitless)
-	 */
-	UPROPERTY()
-	FExpressionInput Albedo;
-
-	/**
-	 * Roughness (type = float, unit = unitless)
-	 */
-	UPROPERTY()
-	FExpressionInput Roughness;
-
-	/**
-	 * Normal (type = float3, unit = unitless)
-	 */
-	UPROPERTY()
-	FExpressionInput Normal;
-
-	//~ Begin UMaterialExpression Interface
-#if WITH_EDITOR
-	virtual int32 Compile(class FMaterialCompiler* Compiler, int32 OutputIndex) override;
-	virtual void GetCaption(TArray<FString>& OutCaptions) const override;
-	virtual uint32 GetOutputType(int32 OutputIndex) override;
-	virtual uint32 GetInputType(int32 InputIndex) override;
-#endif
-	//~ End UMaterialExpression Interface
-};
-
 UCLASS(MinimalAPI, collapsecategories, hidecategories = Object)
 class UMaterialExpressionStrataDielectricBSDF : public UMaterialExpression
 {
diff --git a/Engine/Source/Runtime/Engine/Private/Materials/HLSLMaterialTranslator.cpp b/Engine/Source/Runtime/Engine/Private/Materials/HLSLMaterialTranslator.cpp
--- a/Engine/Source/Runtime/Engine/Private/Materials/HLSLMaterialTranslator.cpp
+++ b/Engine/Source/Runtime/Engine/Private/Materials/HLSLMaterialTranslator.cpp
@@ -7271,21 +7271,10 @@ int32 FHLSLMaterialTranslator::FrontMaterial()
 	return AddInlinedCodeChunk(MCT_Strata, TEXT("GetInitialisedStrataData()"));
 }
 
-int32 FHLSLMaterialTranslator::StrataDiffuseOrenNayarBSDF(int32 Albedo, int32 Roughness, int32 Normal, uint8 SharedNormalIndex)
+int32 FHLSLMaterialTranslator::StrataDiffuseBSDF(int32 Albedo, int32 Roughness, int32 Normal, uint8 SharedNormalIndex)
 {
 	return AddCodeChunk(
-		MCT_Strata, TEXT("GetStrataDiffuseOrenNayarBSDF(%s, %s, %u) /* %s */"),
-		*GetParameterCode(Albedo),
-		*GetParameterCode(Roughness),
-		SharedNormalIndex,
-		*GetParameterCode(Normal)
-	);
-}
-
-int32 FHLSLMaterialTranslator::StrataDiffuseChanBSDF(int32 Albedo, int32 Roughness, int32 Normal, uint8 SharedNormalIndex)
-{
-	return AddCodeChunk(
-		MCT_Strata, TEXT("GetStrataDiffuseChanBSDF(%s, %s, %u) /* %s */"),
+		MCT_Strata, TEXT("GetStrataDiffuseBSDF(%s, %s, %u) /* %s */"),
 		*GetParameterCode(Albedo),
 		*GetParameterCode(Roughness),
 		SharedNormalIndex,
diff --git a/Engine/Source/Runtime/Engine/Private/Materials/HLSLMaterialTranslator.h b/Engine/Source/Runtime/Engine/Private/Materials/HLSLMaterialTranslator.h
--- a/Engine/Source/Runtime/Engine/Private/Materials/HLSLMaterialTranslator.h
+++ b/Engine/Source/Runtime/Engine/Private/Materials/HLSLMaterialTranslator.h
@@ -787,8 +787,7 @@ protected:
 
 	// Strata
 	virtual int32 FrontMaterial() override;
-	virtual int32 StrataDiffuseOrenNayarBSDF(int32 Albedo, int32 Roughness, int32 Normal, uint8 SharedNormalIndex) override;
-	virtual int32 StrataDiffuseChanBSDF(int32 Albedo, int32 Roughness, int32 Normal, uint8 SharedNormalIndex) override;
+	virtual int32 StrataDiffuseBSDF(int32 Albedo, int32 Roughness, int32 Normal, uint8 SharedNormalIndex) override;
 	virtual int32 StrataDielectricBSDF(int32 Roughness, int32 IOR, int32 Tint, int32 Normal, int32 Tangent, uint8 SharedNormalIndex) override;
 	virtual int32 StrataConductorBSDF(int32 Reflectivity, int32 EdgeColor, int32 Roughness, int32 Normal, int32 Tangent, uint8 SharedNormalIndex) override;
 	virtual int32 StrataVolumeBSDF(int32 Albedo, int32 Extinction, int32 Anisotropy, int32 Thickness, int32 Normal, uint8 SharedNormalIndex) override;
diff --git a/Engine/Source/Runtime/Engine/Private/Materials/MaterialExpressions.cpp b/Engine/Source/Runtime/Engine/Private/Materials/MaterialExpressions.cpp
--- a/Engine/Source/Runtime/Engine/Private/Materials/MaterialExpressions.cpp
+++ b/Engine/Source/Runtime/Engine/Private/Materials/MaterialExpressions.cpp
@@ -19033,19 +19033,19 @@ int32 UMaterialExpressionStrataDiffuseBSDF::Compile(class FMaterialCompiler* Com
 	int32 NormalCodeChunk = Normal.GetTracedInput().Expression ? Normal.Compile(Compiler) : Compiler->PixelNormalWS();
 	uint8 SharedNormalIndex = StrataCompilationInfoCreateSharedNormal(Compiler, NormalCodeChunk);
 
-	int32 OutputCodeChunk = Compiler->StrataDiffuseOrenNayarBSDF(
+	int32 OutputCodeChunk = Compiler->StrataDiffuseBSDF(
 		Albedo.GetTracedInput().Expression		? Albedo.Compile(Compiler)		: Compiler->Constant3(0.18f, 0.18f, 0.18f),
 		Roughness.GetTracedInput().Expression	? Roughness.Compile(Compiler)	: Compiler->Constant(0.0f),
 		NormalCodeChunk,
 		SharedNormalIndex);
-	StrataCompilationInfoCreateSingleBSDFMaterial(Compiler, OutputCodeChunk, SharedNormalIndex, STRATA_BSDF_TYPE_DIFFUSE_ON);
+	StrataCompilationInfoCreateSingleBSDFMaterial(Compiler, OutputCodeChunk, SharedNormalIndex, STRATA_BSDF_TYPE_DIFFUSE);
 
 	return OutputCodeChunk;
 }
 
 void UMaterialExpressionStrataDiffuseBSDF::GetCaption(TArray<FString>& OutCaptions) const
 {
-	OutCaptions.Add(TEXT("Strata Diffuse Oren-Nayar BSDF"));
+	OutCaptions.Add(TEXT("Strata Diffuse BSDF"));
 }
 
 uint32 UMaterialExpressionStrataDiffuseBSDF::GetOutputType(int32 OutputIndex)
@@ -19077,70 +19077,6 @@ uint32 UMaterialExpressionStrataDiffuseBSDF::GetInputType(int32 InputIndex)
 
 
 
-UMaterialExpressionStrataDiffuseChanBSDF::UMaterialExpressionStrataDiffuseChanBSDF(const FObjectInitializer& ObjectInitializer)
-	: Super(ObjectInitializer)
-{
-	struct FConstructorStatics
-	{
-		FText NAME_Strata;
-		FConstructorStatics() : NAME_Strata(LOCTEXT("Strata", "Strata")) { }
-	};
-	static FConstructorStatics ConstructorStatics;
-#if WITH_EDITORONLY_DATA
-	MenuCategories.Add(ConstructorStatics.NAME_Strata);
-#endif
-}
-
-#if WITH_EDITOR
-int32 UMaterialExpressionStrataDiffuseChanBSDF::Compile(class FMaterialCompiler* Compiler, int32 OutputIndex)
-{
-	int32 NormalCodeChunk = Normal.GetTracedInput().Expression ? Normal.Compile(Compiler) : Compiler->PixelNormalWS();
-	uint8 SharedNormalIndex = StrataCompilationInfoCreateSharedNormal(Compiler, NormalCodeChunk);
-
-	int32 OutputCodeChunk = Compiler->StrataDiffuseChanBSDF(
-		Albedo.GetTracedInput().Expression		? Albedo.Compile(Compiler)		: Compiler->Constant3(0.18f, 0.18f, 0.18f),
-		Roughness.GetTracedInput().Expression	? Roughness.Compile(Compiler)	: Compiler->Constant(0.0f),
-		NormalCodeChunk,
-		SharedNormalIndex);
-	StrataCompilationInfoCreateSingleBSDFMaterial(Compiler, OutputCodeChunk, SharedNormalIndex, STRATA_BSDF_TYPE_DIFFUSE_CHAN);
-
-	return OutputCodeChunk;
-}
-
-void UMaterialExpressionStrataDiffuseChanBSDF::GetCaption(TArray<FString>& OutCaptions) const
-{
-	OutCaptions.Add(TEXT("Strata Diffuse Chan BSDF"));
-}
-
-uint32 UMaterialExpressionStrataDiffuseChanBSDF::GetOutputType(int32 OutputIndex)
-{
-	return MCT_Strata;
-}
-
-uint32 UMaterialExpressionStrataDiffuseChanBSDF::GetInputType(int32 InputIndex)
-{
-	switch (InputIndex)
-	{
-	case 0:
-		return MCT_Float3;
-		break;
-	case 1:
-		return MCT_Float1;
-		break;
-	case 2:
-		return MCT_Float3;
-		break;
-	}
-
-	check(false);
-	return MCT_Float1;
-}
-#endif // WITH_EDITOR
-
-
-
-
-
 UMaterialExpressionStrataDielectricBSDF::UMaterialExpressionStrataDielectricBSDF(const FObjectInitializer& ObjectInitializer)
 	: Super(ObjectInitializer)
 {
diff --git a/Engine/Source/Runtime/Engine/Private/Materials/StrataMaterial.cpp b/Engine/Source/Runtime/Engine/Private/Materials/StrataMaterial.cpp
--- a/Engine/Source/Runtime/Engine/Private/Materials/StrataMaterial.cpp
+++ b/Engine/Source/Runtime/Engine/Private/Materials/StrataMaterial.cpp
@@ -9,11 +9,8 @@ FString GetStrataBSDFName(uint8 BSDFType)
 {
 	switch (BSDFType)
 	{
-	case STRATA_BSDF_TYPE_DIFFUSE_ON:
-		return TEXT("DIFFUSE_ON");
-		break;
-	case STRATA_BSDF_TYPE_DIFFUSE_CHAN:
-		return TEXT("DIFFUSE_CHAN");
+	case STRATA_BSDF_TYPE_DIFFUSE:
+		return TEXT("DIFFUSE");
 		break;
 	case STRATA_BSDF_TYPE_DIELECTRIC:
 		return TEXT("DIELECTRIC");
@@ -59,7 +56,7 @@ void StrataCompilationInfoCreateSingleBSDFMaterial(FMaterialCompiler* Compiler,
 
 FStrataMaterialCompilationInfo StrataCompilationInfoAdd(FMaterialCompiler* Compiler, const FStrataMaterialCompilationInfo& A, const FStrataMaterialCompilationInfo& B)
 {
-	if ((A.TotalBSDFCount + B.TotalBSDFCount) > STRATA_MAX_TOTAL_BSDF) // See StrataDefinitions.h
+	if ((A.TotalBSDFCount + B.TotalBSDFCount) > STRATA_MAX_TOTAL_BSDF)
 	{
 		Compiler->Error(TEXT("Adding would result in too many BSDFs"));
 		return A;
@@ -73,7 +70,7 @@ FStrataMaterialCompilationInfo StrataCompilationInfoAdd(FMaterialCompiler* Compi
 		const FStrataMaterialCompilationInfo::FLayer& ALayer = A.Layers[LayerIt];
 		const FStrataMaterialCompilationInfo::FLayer& BLayer = B.Layers[LayerIt];
 
-		if ((ALayer.BSDFCount + BLayer.BSDFCount) > STRATA_MAX_BSDF_COUNT_PER_LAYER) // See StrataDefinitions.h
+		if ((ALayer.BSDFCount + BLayer.BSDFCount) > STRATA_MAX_BSDF_COUNT_PER_LAYER)
 		{
 			Compiler->Error(TEXT("Adding would result in too many BSDFs in a Layer"));
 			return A;
@@ -102,7 +99,7 @@ FStrataMaterialCompilationInfo StrataCompilationInfoMultiply(FMaterialCompiler*
 
 FStrataMaterialCompilationInfo StrataCompilationInfoHorizontalMixing(FMaterialCompiler* Compiler, const FStrataMaterialCompilationInfo& A, const FStrataMaterialCompilationInfo& B)
 {
-	if ((A.TotalBSDFCount + B.TotalBSDFCount) > STRATA_MAX_TOTAL_BSDF) // See StrataDefinitions.h
+	if ((A.TotalBSDFCount + B.TotalBSDFCount) > STRATA_MAX_TOTAL_BSDF)
 	{
 		Compiler->Error(TEXT("Mixing would result in too many BSDFs"));
 		return A;
@@ -120,7 +117,7 @@ FStrataMaterialCompilationInfo StrataCompilationInfoVerticalLayering(FMaterialCo
 		return Base;
 	}
 
-	if ((Top.LayerCount + Base.LayerCount) > STRATA_MAX_LAYER_COUNT) // See StrataDefinitions.h
+	if ((Top.LayerCount + Base.LayerCount) > STRATA_MAX_LAYER_COUNT)
 	{
 		Compiler->Error(TEXT("Layering would result in too many Layers"));
 		return Base;
diff --git a/Engine/Source/Runtime/Engine/Public/MaterialCompiler.h b/Engine/Source/Runtime/Engine/Public/MaterialCompiler.h
--- a/Engine/Source/Runtime/Engine/Public/MaterialCompiler.h
+++ b/Engine/Source/Runtime/Engine/Public/MaterialCompiler.h
@@ -391,8 +391,7 @@ public:
 
 	// Strata
 	virtual int32 FrontMaterial() = 0;
-	virtual int32 StrataDiffuseOrenNayarBSDF(int32 Albedo, int32 Roughness, int32 Normal, uint8 SharedNormalIndex) = 0;
-	virtual int32 StrataDiffuseChanBSDF(int32 Albedo, int32 Roughness, int32 Normal, uint8 SharedNormalIndex) = 0;
+	virtual int32 StrataDiffuseBSDF(int32 Albedo, int32 Roughness, int32 Normal, uint8 SharedNormalIndex) = 0;
 	virtual int32 StrataDielectricBSDF(int32 Roughness, int32 IOR, int32 Tint, int32 Normal, int32 Tangent, uint8 SharedNormalIndex) = 0;
 	virtual int32 StrataConductorBSDF(int32 Reflectivity, int32 EdgeColor, int32 Roughness, int32 Normal, int32 Tangent, uint8 SharedNormalIndex) = 0;
 	virtual int32 StrataVolumeBSDF(int32 Albedo, int32 Extinction, int32 Anisotropy, int32 Thickness, int32 Normal, uint8 SharedNormalIndex) = 0;
@@ -804,14 +803,9 @@ public:
 		return Compiler->FrontMaterial();
 	}
 
-	virtual int32 StrataDiffuseOrenNayarBSDF(int32 Albedo, int32 Roughness, int32 Normal, uint8 SharedNormalIndex) override
+	virtual int32 StrataDiffuseBSDF(int32 Albedo, int32 Roughness, int32 Normal, uint8 SharedNormalIndex) override
 	{
-		return Compiler->StrataDiffuseOrenNayarBSDF(Albedo, Roughness, Normal, SharedNormalIndex);
-	}
-
-	virtual int32 StrataDiffuseChanBSDF(int32 Albedo, int32 Roughness, int32 Normal, uint8 SharedNormalIndex) override
-	{
-		return Compiler->StrataDiffuseChanBSDF(Albedo, Roughness, Normal, SharedNormalIndex);
+		return Compiler->StrataDiffuseBSDF(Albedo, Roughness, Normal, SharedNormalIndex);
 	}
 
 	virtual int32 StrataDielectricBSDF(int32 Roughness, int32 IOR, int32 Tint, int32 Normal, int32 Tangent, uint8 SharedNormalIndex) override
