Loading/unpacking shared normal on demand to reduce vgprs pressure and increase occupancy.

diff --git a/Engine/Shaders/Private/BasePassPixelShader.usf b/Engine/Shaders/Private/BasePassPixelShader.usf
--- a/Engine/Shaders/Private/BasePassPixelShader.usf
+++ b/Engine/Shaders/Private/BasePassPixelShader.usf
@@ -1461,6 +1461,14 @@ void FPixelShaderInOut_MainPS(
 
 	FStrataData Strata = PixelMaterialInputs.FrontMaterial;
 
+	// Initialise a Strata header with normal in registers
+	FStrataPixelHeader StrataPixelHeader = InitialiseStrataPixelHeader();
+	for (int l = 0; l < Strata.LayerCount; ++l)
+	{
+		StrataPixelHeader.BSDFCount += Strata.Layers[l].BSDFCount;
+	}
+	StrataPixelHeader.SharedNormals = MaterialParameters.SharedNormals;
+
 #if STRATA_OPAQUE_ENABLED
 
 	// Need to cancel that to make sure strata material are only lit using Strata lighting passes
@@ -1481,7 +1489,7 @@ void FPixelShaderInOut_MainPS(
 	const float3 VertexNormal = MaterialParameters.TangentToWorld[2].xyz;
 	PackStrataOut(
 		OpaqueBasePass.Strata.MaterialLobesBufferUAV, PixelStrataDataByteOffset, 
-		Strata, MaterialParameters.SharedNormals, 
+		StrataPixelHeader, Strata,
 		VertexNormal, MaterialParameters.CameraVector, GBuffer);
 #endif // STRATA_OPAQUE_ENABLED
 
@@ -1489,13 +1497,6 @@ void FPixelShaderInOut_MainPS(
 	//FORWARD_SHADING
 	if (Strata.LayerCount > 0)
 	{
-		FStrataPixelHeader Header = (FStrataPixelHeader)0;
-		for (int l = 0; l < Strata.LayerCount; ++l)
-		{
-			Header.BSDFCount += Strata.Layers[l].BSDFCount;
-		}
-		Header.SharedNormals = MaterialParameters.SharedNormals;
-
 		float3 L = ResolvedView.DirectionalLightDirection;
 		float3 V = MaterialParameters.CameraVector;
 		float3 LightColor = ResolvedView.DirectionalLightColor.rgb * PI;
@@ -1509,7 +1510,7 @@ void FPixelShaderInOut_MainPS(
 			LightColor,
 			DirectionalLightShadow,
 			IndirectOcclusion,
-			Header,
+			StrataPixelHeader,
 			Strata,
 			Throughput);
 
diff --git a/Engine/Shaders/Private/DeferredLightPixelShaders.usf b/Engine/Shaders/Private/DeferredLightPixelShaders.usf
--- a/Engine/Shaders/Private/DeferredLightPixelShaders.usf
+++ b/Engine/Shaders/Private/DeferredLightPixelShaders.usf
@@ -32,6 +32,7 @@ float CloudShadowmapStrength;
 #endif
 
 #if STRATA_ENABLED
+#define STRATA_SHARED_NORMALS_IN_REGISTERS 0
 #include "/Engine/Private/Strata/Strata.ush"
 #include "/Engine/Private/ParticipatingMediaCommon.ush"
 #include "/Engine/Private/Strata/StrataDeferredLighting.ush"
diff --git a/Engine/Shaders/Private/ReflectionEnvironmentPixelShader.usf b/Engine/Shaders/Private/ReflectionEnvironmentPixelShader.usf
--- a/Engine/Shaders/Private/ReflectionEnvironmentPixelShader.usf
+++ b/Engine/Shaders/Private/ReflectionEnvironmentPixelShader.usf
@@ -105,6 +105,7 @@ FSkyLightVisibilityData GetSkyLightVisibilityData(FGBufferData GBuffer, float Am
 
 
 #if STRATA_ENABLED
+#define STRATA_SHARED_NORMALS_IN_REGISTERS 0
 #include "/Engine/Private/Strata/Strata.ush"
 #include "/Engine/Private/ParticipatingMediaCommon.ush"
 #include "/Engine/Private/Strata/StrataEnvironmentLighting.ush"
diff --git a/Engine/Shaders/Private/Strata/Strata.ush b/Engine/Shaders/Private/Strata/Strata.ush
--- a/Engine/Shaders/Private/Strata/Strata.ush
+++ b/Engine/Shaders/Private/Strata/Strata.ush
@@ -12,13 +12,12 @@
 
 #include "/Engine/Shared/StrataDefinitions.h"
 
-// STRATA_TODO
-//	- What model exactly
-//	- Fat dielectric model? Or make it adaptative (for instance Medium_ could be optional to store/load)
-//	- Have a Disney model that is convernient to be super packed and be able to represent current mateiral?
-//	- Shared normal? Up to N normals? On CPU, inputs have a hash that could be use to decuplicate normals
-//	- Transmittance lob modelisation (GGX statistics?) and blocking surface such as FStrataDiffuse?
-//	- fp16 on gen5
+
+// During the base pass or forward rendering, shared normals are simply available from registers.
+// If a BSDF is loaded for processing from the strata buffer, normals are unpacked on demand to reduce VGPR pressure and increase occupancy.
+#ifndef STRATA_SHARED_NORMALS_IN_REGISTERS
+#define STRATA_SHARED_NORMALS_IN_REGISTERS 1
+#endif
 
 
 
@@ -213,10 +212,12 @@ float3 StrataGetOffSpecularPeakReflectionDir(float3 Normal, float3 ReflectionVec
 ///////////////////////////////////////////////////////////////////////////////
 // Shared normals
 
+#define STRATA_PACKED_NORMAL_STRIDE 4
+
 struct FSharedNormals
 {
 	uint NormalCount;
-	float3 Normals[STRATA_MAX_TOTAL_BSDF];
+	float3 Normals[STRATA_MAX_SHARED_NORMAL_REGISTERS];
 };
 
 FSharedNormals StrataInitialiseSharedNormals()
@@ -225,9 +226,41 @@ FSharedNormals StrataInitialiseSharedNormals()
 	return SharedNormals;
 }
 
-float3 StrataGetBSDFSharedNormal(in FSharedNormals SharedNormals, in FStrataBSDF BSDF)
+struct FStrataPixelHeader
 {
-	return SharedNormals.Normals[BSDF_GETNORMALID(BSDF)];
+	uint BSDFCount;
+
+#if STRATA_SHARED_NORMALS_IN_REGISTERS
+	// All the shared normal in VGPRS. This is only used in basepass and during forward rendering.
+	FSharedNormals SharedNormals;
+#else
+	ByteAddressBuffer StrataBuffer;
+	uint SharedNormalByteOffset;
+#endif
+};
+
+FStrataPixelHeader InitialiseStrataPixelHeader()
+{
+	FStrataPixelHeader Out;
+	Out.BSDFCount = 0;
+#if STRATA_SHARED_NORMALS_IN_REGISTERS
+	Out.SharedNormals = (FSharedNormals)0;
+#else
+	Out.SharedNormalByteOffset = 0;
+#endif
+	return Out;
+}
+
+float3 StrataGetBSDFSharedNormal(in FStrataPixelHeader StrataPixelHeader, in FStrataBSDF BSDF)
+{
+#if STRATA_SHARED_NORMALS_IN_REGISTERS
+	float3 Normal = StrataPixelHeader.SharedNormals.Normals[BSDF_GETNORMALID(BSDF)];
+#else
+	uint PackedNormal = StrataPixelHeader.StrataBuffer.Load(StrataPixelHeader.SharedNormalByteOffset + BSDF_GETNORMALID(BSDF) * STRATA_PACKED_NORMAL_STRIDE);
+	float3 Normal     = StrataUnpackNormal(PackedNormal);
+#endif
+
+	return Normal;
 }
 
 ///////////////////////////////////////////////////////////////////////////////
@@ -258,11 +291,6 @@ struct FStrataBSDFContext
 	float a2Override;
 };
 
-void StrataGetBSDFNormal(in FSharedNormals SharedNormals, in FStrataBSDF BSDF, inout float3 N)
-{
-	N = StrataGetBSDFSharedNormal(SharedNormals, BSDF);
-}
-
 void StrataGetBSDFRoughness(in FStrataBSDF BSDF, inout float R)
 {
 	R = 0;
@@ -325,7 +353,7 @@ bool StrataHasAnisotropy(in float Roughness0, float Roughness1)
 	return Roughness0 != Roughness1;
 }
 
-FStrataBSDFContext StrataCreateBSDFContext(FSharedNormals SharedNormals, FStrataBSDF BSDF, float3 V, float3 L)
+FStrataBSDFContext StrataCreateBSDFContext(FStrataPixelHeader StrataPixelHeader, FStrataBSDF BSDF, float3 V, float3 L)
 {
 	FStrataBSDFContext BSDFContext = (FStrataBSDFContext)0;
 
@@ -334,7 +362,7 @@ FStrataBSDFContext StrataCreateBSDFContext(FSharedNormals SharedNormals, FStrata
 	BSDFContext.V = V;
 	BSDFContext.L = L;
 
-	BSDFContext.N = StrataGetBSDFSharedNormal(SharedNormals, BSDF);
+	BSDFContext.N = StrataGetBSDFSharedNormal(StrataPixelHeader, BSDF);
 	BSDFContext.N = normalize(BSDFContext.N);
 
 	BSDFContext.R = 2 * dot(BSDFContext.V, BSDFContext.N) * BSDFContext.N - BSDFContext.V;
@@ -484,14 +512,14 @@ struct FStrataImportanceSampleResult
  * E: two random numbers
  * CameraVector: vector from the camera to the considered direction (not V)
  */
-FStrataImportanceSampleResult StrataImportanceSampleBSDF(FSharedNormals SharedNormals, FStrataBSDF BSDF, float2 E, float3 CameraVector)
+FStrataImportanceSampleResult StrataImportanceSampleBSDF(FStrataPixelHeader StrataPixelHeader, FStrataBSDF BSDF, float2 E, float3 CameraVector)
 {
 	FStrataImportanceSampleResult Sample = (FStrataImportanceSampleResult)0;
 
 	// Create a context from the information we have
 	const float3 UnusedL = float3(0.0f, 0.0f, 1.0f);
 	const float3 V = -CameraVector;
-	FStrataBSDFContext BSDFContext = StrataCreateBSDFContext(SharedNormals, BSDF, V, UnusedL);
+	FStrataBSDFContext BSDFContext = StrataCreateBSDFContext(StrataPixelHeader, BSDF, V, UnusedL);
 
 	float PDF = 0.0f;
 	const uint BSDFType = BSDF_GETTYPE(BSDF);
@@ -963,8 +991,8 @@ uint GetStrataPixelDataByteOffset(uint2 PixelPos, uint ViewBufferSizeX, uint Str
 void PackStrataOut(
 	RWByteAddressBuffer StrataBuffer, 
 	uint ByteOffset, 
+	FStrataPixelHeader StrataPixelHeader,
 	FStrataData Strata,
-	FSharedNormals SharedNormals,
 	float3 VertexNormal,
 	float3 V,
 	inout FGBufferData GBuffer)
@@ -1007,8 +1035,7 @@ void PackStrataOut(
 			BSDFCount++;
 
 			// STRATA_TODO apply the same normal transform as for the material
-			float3 N = VertexNormal;
-			StrataGetBSDFNormal(SharedNormals, BSDF, N);
+			float3 N = StrataGetBSDFSharedNormal(StrataPixelHeader, BSDF);
 			N = normalize(N);
 
 			AccumTopNormal += N;
@@ -1130,13 +1157,18 @@ void PackStrataOut(
 	// 1. the header
 	BSDFCount = min(BSDFCount, uint(STRATA_MAX_TOTAL_BSDF));
 
-	uint SharedNormals_BSDFCount = (SharedNormals.NormalCount << 16) | BSDFCount;
+	uint SharedNormals_BSDFCount = BSDFCount;
+#if STRATA_SHARED_NORMALS_IN_REGISTERS
+	SharedNormals_BSDFCount |= (StrataPixelHeader.SharedNormals.NormalCount << 16);
+#endif
 	STRATA_STORE_UINT1(SharedNormals_BSDFCount);
 
-	for (uint i = 0; i < SharedNormals.NormalCount; ++i)
+#if STRATA_SHARED_NORMALS_IN_REGISTERS
+	for (uint i = 0; i < StrataPixelHeader.SharedNormals.NormalCount; ++i)
 	{
-		STRATA_STORE_UINT1(StrataPackNormal(SharedNormals.Normals[i]));
+		STRATA_STORE_UINT1(StrataPackNormal(StrataPixelHeader.SharedNormals.Normals[i]));
 	}
+#endif
 
 	// 2. the list of BSDFs
 	BSDFCount = 0;
@@ -1206,29 +1238,25 @@ void PackStrataOut(
 	}
 }
 
-struct FStrataPixelHeader
-{
-	uint BSDFCount;
-
-	FSharedNormals SharedNormals;
-};
-
 // Unpack only strata header
 FStrataPixelHeader UnpackStrataHeaderIn(ByteAddressBuffer StrataBuffer, inout uint ByteOffset)
 {
-	FStrataPixelHeader Out = (FStrataPixelHeader)0;
+	FStrataPixelHeader Out = InitialiseStrataPixelHeader();
 
 	uint SharedNormals_BSDFCount = 0;
 	STRATA_LOAD_UINT1(SharedNormals_BSDFCount);
 	Out.BSDFCount = SharedNormals_BSDFCount & 0x0000FFFF;
-	Out.SharedNormals.NormalCount = SharedNormals_BSDFCount >> 16;
 
-	for (int i = 0; i < Out.SharedNormals.NormalCount; ++i)
-	{
-		uint PackedNormal = 0;
-		STRATA_LOAD_UINT1(PackedNormal);
-		Out.SharedNormals.Normals[i] = StrataUnpackNormal(PackedNormal);
-	}
+	uint NormalCount = SharedNormals_BSDFCount >> 16;
+
+#if STRATA_SHARED_NORMALS_IN_REGISTERS==0
+	// Only keep the offset to the shared normal memory. We are going to load them on demand.
+	Out.SharedNormalByteOffset = ByteOffset;
+	Out.StrataBuffer = StrataBuffer;
+
+	// Now skip over the shared normal memory to be able to load the BSDF content
+	ByteOffset += NormalCount * STRATA_PACKED_NORMAL_STRIDE;
+#endif
 
 	return Out;
 }
diff --git a/Engine/Shaders/Private/Strata/StrataDeferredLighting.ush b/Engine/Shaders/Private/Strata/StrataDeferredLighting.ush
--- a/Engine/Shaders/Private/Strata/StrataDeferredLighting.ush
+++ b/Engine/Shaders/Private/Strata/StrataDeferredLighting.ush
@@ -134,7 +134,7 @@ float4 StrataDeferredLighting(
 {
 	float4 OutRadiance = 0;
 
-	FStrataPixelHeader Header = UnpackStrataHeaderIn(StrataDataBuffer, PixelStrataDataByteOffset);
+	FStrataPixelHeader StrataPixelHeader = UnpackStrataHeaderIn(StrataDataBuffer, PixelStrataDataByteOffset);
 	
 	float3 PreintegratedIndicentIrradiance = 0;
 	if (LightData.bRectLight)
@@ -149,11 +149,11 @@ float4 StrataDeferredLighting(
 	}
 
 	LOOP
-	for (uint BSDFIndex = 0; BSDFIndex < Header.BSDFCount; ++BSDFIndex)
+	for (uint BSDFIndex = 0; BSDFIndex < StrataPixelHeader.BSDFCount; ++BSDFIndex)
 	{
 		// Unpack BSDF data
 		const FStrataBSDF BSDF = UnpackStrataBSDFIn(StrataDataBuffer, PixelStrataDataByteOffset);
-		FStrataBSDFContext BSDFContext = StrataCreateBSDFContext(Header.SharedNormals, BSDF, V, L);
+		FStrataBSDFContext BSDFContext = StrataCreateBSDFContext(StrataPixelHeader, BSDF, V, L);
 
 		// Evaluate incident lighting & (optionally) precompute area light integration
 		FStrataLightIntegrationContext IntegrationContext = StrataIniLightIntegrationContext();
diff --git a/Engine/Shaders/Private/Strata/StrataEnvironmentLighting.ush b/Engine/Shaders/Private/Strata/StrataEnvironmentLighting.ush
--- a/Engine/Shaders/Private/Strata/StrataEnvironmentLighting.ush
+++ b/Engine/Shaders/Private/Strata/StrataEnvironmentLighting.ush
@@ -15,17 +15,17 @@ float3 StrataEnvironmentLighting(
 	ByteAddressBuffer StrataDataBuffer,
 	uint PixelStrataDataByteOffset)
 {
-	FStrataPixelHeader Header = UnpackStrataHeaderIn(StrataDataBuffer, PixelStrataDataByteOffset);
+	FStrataPixelHeader StrataPixelHeader = UnpackStrataHeaderIn(StrataDataBuffer, PixelStrataDataByteOffset);
 
 	float3 SkyLighting = 0;
-	for (uint BSDFIndex = 0; BSDFIndex < Header.BSDFCount; ++BSDFIndex)
+	for (uint BSDFIndex = 0; BSDFIndex < StrataPixelHeader.BSDFCount; ++BSDFIndex)
 	{
 		const FStrataBSDF BSDF = UnpackStrataBSDFIn(StrataDataBuffer, PixelStrataDataByteOffset);
 		const float3 BSDFThroughput = BSDF.Weight;
 
 		// Create the BSDF context
 		const float3 UnusedL = float3(0.0f, 0.0f, 1.0f);
-		FStrataBSDFContext StrataBSDFContext = StrataCreateBSDFContext(Header.SharedNormals, BSDF, V, UnusedL);
+		FStrataBSDFContext StrataBSDFContext = StrataCreateBSDFContext(StrataPixelHeader, BSDF, V, UnusedL);
 
 		// Evaluate environment lighting
 		const bool bEnableSpecular = ReflectionStruct.SkyLightParameters.y > 0.0f;
diff --git a/Engine/Shaders/Private/Strata/StrataForwardLighting.ush b/Engine/Shaders/Private/Strata/StrataForwardLighting.ush
--- a/Engine/Shaders/Private/Strata/StrataForwardLighting.ush
+++ b/Engine/Shaders/Private/Strata/StrataForwardLighting.ush
@@ -38,7 +38,7 @@ float3 StrataForwardLighting(
 			const float3 BSDFCoverage = BSDF.Weight;
 
 			// Create the BSDF context
-			FStrataBSDFContext StrataBSDFContext = StrataCreateBSDFContext(Header.SharedNormals, BSDF, V, L);
+			FStrataBSDFContext StrataBSDFContext = StrataCreateBSDFContext(Header, BSDF, V, L);
 
 			// Compute some common factors
 			const float3 DirLightBaseFactor = OutThroughput * BSDFCoverage * LColor * LShadow;
diff --git a/Engine/Shaders/Private/Strata/StrataVisualize.usf b/Engine/Shaders/Private/Strata/StrataVisualize.usf
--- a/Engine/Shaders/Private/Strata/StrataVisualize.usf
+++ b/Engine/Shaders/Private/Strata/StrataVisualize.usf
@@ -3,6 +3,7 @@
 #include "/Engine/Private/Common.ush"
 #include "/Engine/Private/DeferredShadingCommon.ush"
 
+#define STRATA_SHARED_NORMALS_IN_REGISTERS 0
 #include "/Engine/Private/Strata/Strata.ush"
 
 #include "/Engine/Private/MiniFontCommon.ush"
diff --git a/Engine/Shaders/Shared/StrataDefinitions.h b/Engine/Shaders/Shared/StrataDefinitions.h
--- a/Engine/Shaders/Shared/StrataDefinitions.h
+++ b/Engine/Shaders/Shared/StrataDefinitions.h
@@ -15,6 +15,7 @@
 #define STRATA_MAX_BSDF_COUNT_PER_LAYER		4
 #define STRATA_MAX_LAYER_COUNT				4
 #define STRATA_MAX_TOTAL_BSDF				4
+#define STRATA_MAX_SHARED_NORMAL_REGISTERS	4
 
 
 
diff --git a/Engine/Source/Runtime/Engine/Private/Materials/HLSLMaterialTranslator.cpp b/Engine/Source/Runtime/Engine/Private/Materials/HLSLMaterialTranslator.cpp
--- a/Engine/Source/Runtime/Engine/Private/Materials/HLSLMaterialTranslator.cpp
+++ b/Engine/Source/Runtime/Engine/Private/Materials/HLSLMaterialTranslator.cpp
@@ -7246,7 +7246,7 @@ uint8 FHLSLMaterialTranslator::StrataCompilationInfoRegisterSharedNormalIndex(in
 {
 	check(NormalCodeChunk != INDEX_NONE);
 	check(NextFreeStrataShaderNormalIndex < 255);	// Out of shared normal slots
-	check(NextFreeStrataShaderNormalIndex < STRATA_MAX_TOTAL_BSDF);	// This is our current budget if normal per material. STRATA_TODO change that to a byte per pixel
+	check(NextFreeStrataShaderNormalIndex < STRATA_MAX_SHARED_NORMAL_REGISTERS);	// This is the size of the HLSL array limiting the maximum shared normal count.
 
 	const uint64 NormalCodeChunkHash = GetParameterHash(NormalCodeChunk);
 
