Made scattering/anisotropy optionnal on the volume BSDF to use less bytes.
This will be beneficial when switching to byte per pixel later.

diff --git a/Engine/Shaders/Private/Strata/Strata.ush b/Engine/Shaders/Private/Strata/Strata.ush
--- a/Engine/Shaders/Private/Strata/Strata.ush
+++ b/Engine/Shaders/Private/Strata/Strata.ush
@@ -43,12 +43,14 @@ struct FStrataBSDF
 #define STATE_BITS_SHAREDNORMAL			2
 #define STATE_BITS_ANISOTROPY			1
 #define STATE_BITS_TOPLAYER				1
+#define STATE_BITS_HASSCATTERING		1
 
 // Bit offset of each component store in BSDF state 
 #define STATE_BIT_OFFSET_BSDF			0
-#define STATE_BIT_OFFSET_SHAREDNORMAL	(STATE_BITS_BSDF)
-#define STATE_BIT_OFFSET_ANISOTROPY		(STATE_BITS_BSDF + STATE_BIT_OFFSET_SHAREDNORMAL)
-#define STATE_BIT_OFFSET_TOPLAYER		(STATE_BITS_BSDF + STATE_BIT_OFFSET_SHAREDNORMAL + STATE_BIT_OFFSET_ANISOTROPY)
+#define STATE_BIT_OFFSET_SHAREDNORMAL	(STATE_BIT_OFFSET_BSDF				+ STATE_BITS_BSDF)
+#define STATE_BIT_OFFSET_ANISOTROPY		(STATE_BIT_OFFSET_SHAREDNORMAL		+ STATE_BITS_SHAREDNORMAL)
+#define STATE_BIT_OFFSET_TOPLAYER		(STATE_BIT_OFFSET_ANISOTROPY		+ STATE_BIT_OFFSET_ANISOTROPY)
+#define STATE_BIT_OFFSET_HASSCATTERING	(STATE_BIT_OFFSET_TOPLAYER			+ STATE_BIT_OFFSET_TOPLAYER)
 
 #define TO_BITMASK(BitCount)							((1<<BitCount)-1)
 #define BSDF_GETSTATE(X, BitCount, BitOffset)			uint((X.State>>BitOffset) & TO_BITMASK(BitCount))
@@ -72,6 +74,10 @@ struct FStrataBSDF
 #define BSDF_GETISTOPLAYER(X)			BSDF_GETSTATE(X, STATE_BITS_TOPLAYER, STATE_BIT_OFFSET_TOPLAYER)
 #define BSDF_SETISTOPLAYER(X, IsTop)	BSDF_SETSTATE(X, STATE_BITS_TOPLAYER, STATE_BIT_OFFSET_TOPLAYER, IsTop)
 
+// Indicate if the BSDF is part of the top layer
+#define BSDF_GETHASSCATTERING(X)		BSDF_GETSTATE(X, STATE_BITS_HASSCATTERING, STATE_BIT_OFFSET_HASSCATTERING)
+#define BSDF_SETHASSCATTERING(X, Scatt)	BSDF_SETSTATE(X, STATE_BITS_HASSCATTERING, STATE_BIT_OFFSET_HASSCATTERING, Scatt)
+
 #define DIFFUSE_ALBEDO(X)				X.VGPRs[0].xyz
 #define DIFFUSE_ROUGHNESS(X)			X.VGPRs[0].w
 
@@ -93,9 +99,28 @@ struct FStrataBSDF
 #define VOLUME_THICKNESS(X)				X.VGPRs[1].w
 
 // Forward declaractions
-uint StrataPackNormal(in float3 Normal);
 float3 StrataUnpackNormal(uint PackedNormal);
-bool StrataHasAnisotropy(in float Roughness0, float Roughness1);
+
+
+bool StrataHasAnisotropy(in float Roughness0, float Roughness1)
+{
+	return Roughness0 != Roughness1;
+}
+
+// Return true if a BSDF has a custom area light integrator
+bool StrataHasAreaLightIntegrator(in FStrataBSDF BSDF)
+{
+	const uint BSDFType = BSDF_GETTYPE(BSDF);
+	return
+		BSDFType == STRATA_BSDF_TYPE_CONDUCTOR ||
+		BSDFType == STRATA_BSDF_TYPE_DIELECTRIC;
+}
+
+bool StrataHasScattering(in float3 Albedo)
+{
+	return any(Albedo > 0.0f);
+}
+
 
 struct FStrataLayer
 {
@@ -186,10 +211,13 @@ FStrataData GetStrataVolumeBSDF(float3 Albedo, float3 Extinction, float Anisotro
 	FStrataData StrataData = GetInitialisedStrataData();
 	StrataData.LayerCount++;
 
+	const uint bHasScattering = StrataHasScattering(Albedo) ? 1 : 0;
+
 	StrataData.Layers[0].BSDFCount++;
 	BSDF_SETTYPE(				StrataData.Layers[0].BSDFs[0],	STRATA_BSDF_TYPE_VOLUME);
 	BSDF_SETNORMALID(			StrataData.Layers[0].BSDFs[0],	SharedNormalIndex);
-	BSDF_SETISTOPLAYER		(StrataData.Layers[0].BSDFs[0],		0);
+	BSDF_SETISTOPLAYER			(StrataData.Layers[0].BSDFs[0],	0);
+	BSDF_SETHASSCATTERING		(StrataData.Layers[0].BSDFs[0], bHasScattering);
 	VOLUME_ALBEDO				(StrataData.Layers[0].BSDFs[0])	= Albedo;
 	VOLUME_EXTINCTION			(StrataData.Layers[0].BSDFs[0])	= Extinction;
 	VOLUME_ANISOTROPY			(StrataData.Layers[0].BSDFs[0])	= Anisotropy;
@@ -339,20 +367,6 @@ void StrataGetBSDFTangent(in FStrataBSDF BSDF, inout float3 T)
 	}
 }
 
-// Return true if a BSDF has a custom area light integrator
-bool StrataHasAreaLightIntegrator(in FStrataBSDF BSDF)
-{
-	const uint BSDFType = BSDF_GETTYPE(BSDF);
-	return 
-		BSDFType == STRATA_BSDF_TYPE_CONDUCTOR ||
-		BSDFType == STRATA_BSDF_TYPE_DIELECTRIC;
-}
-
-bool StrataHasAnisotropy(in float Roughness0, float Roughness1)
-{
-	return Roughness0 != Roughness1;
-}
-
 FStrataBSDFContext StrataCreateBSDFContext(FStrataPixelHeader StrataPixelHeader, FStrataBSDF BSDF, float3 V, float3 L)
 {
 	FStrataBSDFContext BSDFContext = (FStrataBSDFContext)0;
@@ -1222,7 +1236,10 @@ void PackStrataOut(
 				// This gives the same result and reduces memory usage.
 				const float Thickness = VOLUME_THICKNESS(BSDF);
 				STRATA_STORE_UINT1(PackR11G11B10F(VOLUME_EXTINCTION(BSDF) * Thickness));
-				STRATA_STORE_UINT1(PackAlbedoLinearToGamma2AlphaLinear(float4(VOLUME_ALBEDO(BSDF), VOLUME_ANISOTROPY(BSDF)*0.5f + 0.5f)));
+				if (BSDF_GETHASSCATTERING(BSDF))
+				{
+					STRATA_STORE_UINT1(PackAlbedoLinearToGamma2AlphaLinear(float4(VOLUME_ALBEDO(BSDF), VOLUME_ANISOTROPY(BSDF) * 0.5f + 0.5f)));
+				}
 				// 8 bytes
 			}
 			break;
@@ -1312,8 +1329,12 @@ FStrataBSDF UnpackStrataBSDFIn(ByteAddressBuffer StrataBuffer, inout uint ByteOf
 	case STRATA_BSDF_TYPE_VOLUME:
 	{
 		StrataLoad_R11G11B10F			(StrataBuffer, ByteOffset, VOLUME_EXTINCTION(OutBSDF));
-		StrataLoad_AlbedoGamma2ToLinearAlphaLinear(StrataBuffer, ByteOffset, VOLUME_ALBEDO(OutBSDF), VOLUME_ANISOTROPY(OutBSDF));
-		VOLUME_ANISOTROPY(OutBSDF) = VOLUME_ANISOTROPY(OutBSDF) * 2.0f - 1.0f;
+		if (BSDF_GETHASSCATTERING(OutBSDF))
+		{
+			StrataLoad_AlbedoGamma2ToLinearAlphaLinear(StrataBuffer, ByteOffset, VOLUME_ALBEDO(OutBSDF), VOLUME_ANISOTROPY(OutBSDF));
+			VOLUME_ANISOTROPY(OutBSDF) = VOLUME_ANISOTROPY(OutBSDF) * 2.0f - 1.0f;
+		}
+		// else, albedo and anisotropy have already cleared out to zero when creating OutBSDF.
 		VOLUME_THICKNESS(OutBSDF) = 1.0f;
 	}
 	break;
diff --git a/Engine/Source/Runtime/Engine/Private/Materials/MaterialExpressions.cpp b/Engine/Source/Runtime/Engine/Private/Materials/MaterialExpressions.cpp
--- a/Engine/Source/Runtime/Engine/Private/Materials/MaterialExpressions.cpp
+++ b/Engine/Source/Runtime/Engine/Private/Materials/MaterialExpressions.cpp
@@ -19243,8 +19243,10 @@ int32 UMaterialExpressionStrataVolumeBSDF::Compile(class FMaterialCompiler* Comp
 	int32 NormalCodeChunk = Compiler->VertexNormal();
 	uint8 SharedNormalIndex = StrataCompilationInfoCreateSharedNormal(Compiler, NormalCodeChunk);
 
+	const bool bHasScattering = Albedo.GetTracedInput().Expression != nullptr;	// STRATA_TODO use when we switch to byte per pixel
+
 	int32 OutputCodeChunk = Compiler->StrataVolumeBSDF(
-		Albedo.GetTracedInput().Expression		? Albedo.Compile(Compiler)			: Compiler->Constant3(0.0f, 0.0f, 0.0f),
+		bHasScattering							? Albedo.Compile(Compiler)			: Compiler->Constant3(0.0f, 0.0f, 0.0f),
 		Extinction.GetTracedInput().Expression	? Extinction.Compile(Compiler)		: Compiler->Constant3(0.0f, 0.0f, 0.0f),
 		Anisotropy.GetTracedInput().Expression	? Anisotropy.Compile(Compiler)		: Compiler->Constant(0.0f),
 		Thickness.GetTracedInput().Expression	? Thickness.Compile(Compiler)		: Compiler->Constant(0.001f), // default = 1mm
