Strata viz now show bytes used by header, bsdfs and total.

diff --git a/Engine/Shaders/Private/Strata/StrataVisualize.usf b/Engine/Shaders/Private/Strata/StrataVisualize.usf
--- a/Engine/Shaders/Private/Strata/StrataVisualize.usf
+++ b/Engine/Shaders/Private/Strata/StrataVisualize.usf
@@ -24,7 +24,7 @@ void VisualizeMaterialPS(
 	const int  InspectSquare = 10;
 
 	// Sample scene textures.
-	float2 BufferUV = SvPositionToBufferUV(SVPos);
+	float2 BufferUV = SvPositionToBufferUV(float4(MaterialPos, SVPos.zw));
 	FGBufferData GBuffer = GetGBufferDataFromSceneTextures(BufferUV);
 	bool bStrataMaterial = GBuffer.ShadingModelID == SHADINGMODELID_STRATA;
 
@@ -70,7 +70,10 @@ void VisualizeMaterialPS(
 	if (bStrataMaterial)
 	{
 		uint PixelStrataDataByteOffset = GetStrataPixelDataByteOffset(MaterialPos, View.BufferSizeAndInvSize.x, Strata.MaxBytesPerPixel);
+		const uint PixelStrataDataByteOffsetStart = PixelStrataDataByteOffset;
+
 		FStrataPixelHeader Header = UnpackStrataHeaderIn(Strata.MaterialLobesBuffer, PixelStrataDataByteOffset);
+		const uint PixelStrataDataByteOffsetPostHeader = PixelStrataDataByteOffset;
 
 		float3 BSDFFontColor = float3(0.25, 1.0, 0.25);
 		float3 INFOFontColor = float3(1.0, 0.5, 0.25);
@@ -87,6 +90,7 @@ void VisualizeMaterialPS(
 #define PRINTBSDFCOUNT()	PRINTBSDFCHAR(min(_0_+Header.BSDFCount,_F_)); PRINTBSDFCHAR(_SPC_); PRINTBSDFCHAR(_B_); PRINTBSDFCHAR(_S_); PRINTBSDFCHAR(_D_); PRINTBSDFCHAR(_F_);
 #define PRINTFLOAT3(X) PrintFloat(PixelPos, OutColor.xyz, PARAMFontColor, TopLeft, X.x); TopLeft.x+=8*7; PRINTBSDFCHAR(_SPC_); PrintFloat(PixelPos, OutColor.xyz, PARAMFontColor, TopLeft, X.y); TopLeft.x+=8*7; PRINTBSDFCHAR(_SPC_); PrintFloat(PixelPos, OutColor.xyz, PARAMFontColor, TopLeft, X.z);
 #define PRINTFLOAT1(Value) PrintFloat(PixelPos, OutColor.xyz, PARAMFontColor, TopLeft, Value); TopLeft.x+=8*7;
+#define PRINTFLOATNOFRAC(Value) PrintFloatNoFraction(PixelPos, OutColor.xyz, PARAMFontColor, TopLeft, Value); TopLeft.x+=8*7;
 #define SIGNOFFLOAT1(Value) if(Value < 0.0f) { PrintCharacter(PixelPos, OutColor.xyz, PARAMFontColor, TopLeft, _MINUS_); } /*else { PrintCharacter(PixelPos, OutColor.xyz, PARAMFontColor, TopLeft, _SPC_);}*/
 
 #define NEXTLINE  TopLeft.y += 8; TopLeft.x = TopLeftInit.x;
@@ -185,8 +189,27 @@ void VisualizeMaterialPS(
 			HORIZONTALSEPARATION;
 		}
 
+		const uint PixelStrataDataByteOffsetPostBSDF = PixelStrataDataByteOffset;
+
 		PRINTINFOCHAR(_B_); PRINTINFOCHAR(_O_); PRINTINFOCHAR(_T_); PRINTINFOCHAR(_T_); PRINTINFOCHAR(_O_); PRINTINFOCHAR(_M_);
 		NEXTLINE;
+		HORIZONTALSEPARATION;
+
+
+		PRINTBSDFCHAR(_H_); PRINTBSDFCHAR(_E_); PRINTBSDFCHAR(_A_); PRINTBSDFCHAR(_D_); PRINTBSDFCHAR(_E_); PRINTBSDFCHAR(_R_);
+		PRINTBSDFCHAR(_B_); PRINTBSDFCHAR(_Y_); PRINTBSDFCHAR(_T_); PRINTBSDFCHAR(_E_); PRINTBSDFCHAR(_S_); PRINTBSDFCHAR(_SPC_);
+		PRINTFLOATNOFRAC(float(PixelStrataDataByteOffsetPostHeader - PixelStrataDataByteOffsetStart));
+		NEXTLINE;
+
+		PRINTBSDFCHAR(_B_); PRINTBSDFCHAR(_S_); PRINTBSDFCHAR(_D_); PRINTBSDFCHAR(_F_); PRINTBSDFCHAR(_S_); PRINTBSDFCHAR(_SPC_);
+		PRINTBSDFCHAR(_B_); PRINTBSDFCHAR(_Y_); PRINTBSDFCHAR(_T_); PRINTBSDFCHAR(_E_); PRINTBSDFCHAR(_S_); PRINTBSDFCHAR(_SPC_);
+		PRINTFLOATNOFRAC(float(PixelStrataDataByteOffsetPostBSDF - PixelStrataDataByteOffsetPostHeader));
+		NEXTLINE;
+
+		PRINTBSDFCHAR(_T_); PRINTBSDFCHAR(_O_); PRINTBSDFCHAR(_T_); PRINTBSDFCHAR(_A_); PRINTBSDFCHAR(_L_); PRINTBSDFCHAR(_SPC_);
+		PRINTBSDFCHAR(_B_); PRINTBSDFCHAR(_Y_); PRINTBSDFCHAR(_T_); PRINTBSDFCHAR(_E_); PRINTBSDFCHAR(_S_); PRINTBSDFCHAR(_SPC_);
+		PRINTFLOATNOFRAC(float(PixelStrataDataByteOffsetPostBSDF - PixelStrataDataByteOffsetStart));
+		NEXTLINE;
 	}
 }
 
