Strata - when the BSDF weight is detected grey scale, it is stored as 10bits on the state instead to spare a full uint of storage.
+minor cleanups

diff --git a/Engine/Shaders/Private/Strata/Strata.ush b/Engine/Shaders/Private/Strata/Strata.ush
--- a/Engine/Shaders/Private/Strata/Strata.ush
+++ b/Engine/Shaders/Private/Strata/Strata.ush
@@ -40,44 +40,57 @@ struct FStrataBSDF
 
 // Num bits for each component store in BSDF state
 #define STATE_BITS_BSDF					3
-#define STATE_BITS_SHAREDNORMAL			2
+#define STATE_BITS_SHAREDNORMALID		2
 #define STATE_BITS_ANISOTROPY			1
 #define STATE_BITS_TOPLAYER				1
 #define STATE_BITS_HASSCATTERING		1
+#define STATE_BITS_GREYWEIGHT			1
+#define STATE_BITS___UNUSED___			13
+#define STATE_BITS_WEIGHT10F			10
 
 // Bit offset of each component store in BSDF state 
 #define STATE_BIT_OFFSET_BSDF			0
-#define STATE_BIT_OFFSET_SHAREDNORMAL	(STATE_BIT_OFFSET_BSDF				+ STATE_BITS_BSDF)
-#define STATE_BIT_OFFSET_ANISOTROPY		(STATE_BIT_OFFSET_SHAREDNORMAL		+ STATE_BITS_SHAREDNORMAL)
+#define STATE_BIT_OFFSET_SHAREDNORMALID	(STATE_BIT_OFFSET_BSDF				+ STATE_BITS_BSDF)
+#define STATE_BIT_OFFSET_ANISOTROPY		(STATE_BIT_OFFSET_SHAREDNORMALID	+ STATE_BITS_SHAREDNORMALID)
 #define STATE_BIT_OFFSET_TOPLAYER		(STATE_BIT_OFFSET_ANISOTROPY		+ STATE_BIT_OFFSET_ANISOTROPY)
 #define STATE_BIT_OFFSET_HASSCATTERING	(STATE_BIT_OFFSET_TOPLAYER			+ STATE_BIT_OFFSET_TOPLAYER)
+#define STATE_BIT_OFFSET_GREYWEIGHT		(STATE_BIT_OFFSET_HASSCATTERING		+ STATE_BITS_HASSCATTERING)
+#define STATE_BIT_OFFSET_WEIGHT10F		22
 
 #define TO_BITMASK(BitCount)							((1<<BitCount)-1)
 #define BSDF_GETSTATE(X, BitCount, BitOffset)			uint((X.State>>BitOffset) & TO_BITMASK(BitCount))
 #define BSDF_SETSTATE(X, BitCount, BitOffset, Value)	X.State = (X.State & (~(TO_BITMASK(BitCount)<<BitOffset))) | ((Value & TO_BITMASK(BitCount)) << BitOffset)
 
-// Accessor for Type
+// Accessors for the BSDF type
 #define BSDF_GETTYPE(X)					BSDF_GETSTATE(X, STATE_BITS_BSDF, STATE_BIT_OFFSET_BSDF)
 #define BSDF_SETTYPE(X, NewType)		BSDF_SETSTATE(X, STATE_BITS_BSDF, STATE_BIT_OFFSET_BSDF, NewType)
 
-// Accessor for Shared Normal
+// Accessors for Shared Normal
 // This assumes Normals[STRATA_MAX_TOTAL_BSDF]
-#define BSDF_GETNORMALID(X)				BSDF_GETSTATE(X, STATE_BITS_SHAREDNORMAL, STATE_BIT_OFFSET_SHAREDNORMAL)
-#define BSDF_SETNORMALID(X, NrmlIdx)	BSDF_SETSTATE(X, STATE_BITS_SHAREDNORMAL, STATE_BIT_OFFSET_SHAREDNORMAL, NrmlIdx)
+#define BSDF_GETNORMALID(X)				BSDF_GETSTATE(X, STATE_BITS_SHAREDNORMALID, STATE_BIT_OFFSET_SHAREDNORMALID)
+#define BSDF_SETNORMALID(X, NrmlIdx)	BSDF_SETSTATE(X, STATE_BITS_SHAREDNORMALID, STATE_BIT_OFFSET_SHAREDNORMALID, NrmlIdx)
 
-// Accessor for HasAnisotropy
-// Indicate if the BSDF has anisotropic feature
+// Accessors for HasAnisotropy
+// Indicates if the BSDF has anisotropic feature
 #define BSDF_GETANISOTROPY(X)			BSDF_GETSTATE(X, STATE_BITS_ANISOTROPY, STATE_BIT_OFFSET_ANISOTROPY)
 #define BSDF_SETANISOTROPY(X, Aniso)	BSDF_SETSTATE(X, STATE_BITS_ANISOTROPY, STATE_BIT_OFFSET_ANISOTROPY, Aniso)
 
-// Indicate if the BSDF is part of the top layer
+// Indicates if the BSDF is part of the top layer
 #define BSDF_GETISTOPLAYER(X)			BSDF_GETSTATE(X, STATE_BITS_TOPLAYER, STATE_BIT_OFFSET_TOPLAYER)
 #define BSDF_SETISTOPLAYER(X, IsTop)	BSDF_SETSTATE(X, STATE_BITS_TOPLAYER, STATE_BIT_OFFSET_TOPLAYER, IsTop)
 
-// Indicate if the BSDF is part of the top layer
+// Indicates if the BSDF is part of the top layer
 #define BSDF_GETHASSCATTERING(X)		BSDF_GETSTATE(X, STATE_BITS_HASSCATTERING, STATE_BIT_OFFSET_HASSCATTERING)
 #define BSDF_SETHASSCATTERING(X, Scatt)	BSDF_SETSTATE(X, STATE_BITS_HASSCATTERING, STATE_BIT_OFFSET_HASSCATTERING, Scatt)
 
+// Indicates if the BSDF weight is grey scale so that it can be stored in a single float (e.g. top layer BSDFs)
+#define BSDF_GETGREYWEIGHT(X)			BSDF_GETSTATE(X, STATE_BITS_GREYWEIGHT, STATE_BIT_OFFSET_GREYWEIGHT)
+#define BSDF_SETGREYWEIGHT(X, Grey)		BSDF_SETSTATE(X, STATE_BITS_GREYWEIGHT, STATE_BIT_OFFSET_GREYWEIGHT, Grey)
+
+// Accessors for the 10bits float weight used if is it detected being a grey scale
+#define BSDF_GETWEIGHT10F(X)			BSDF_GETSTATE(X, STATE_BITS_WEIGHT10F, STATE_BIT_OFFSET_WEIGHT10F)
+#define BSDF_SETWEIGHT10F(X, Weight)	BSDF_SETSTATE(X, STATE_BITS_WEIGHT10F, STATE_BIT_OFFSET_WEIGHT10F, Weight)
+
 #define DIFFUSE_ALBEDO(X)				X.VGPRs[0].xyz
 #define DIFFUSE_ROUGHNESS(X)			X.VGPRs[0].w
 
@@ -121,6 +134,11 @@ bool StrataHasScattering(in float3 Albedo)
 	return any(Albedo > 0.0f);
 }
 
+bool StrataHasGreyScaleWeight(in float3 Weight)
+{
+	return Weight.x == Weight.y && Weight.y == Weight.z;
+}
+
 
 struct FStrataLayer
 {
@@ -812,6 +830,17 @@ void StrataSanitizeBSDF(inout FStrataBSDF BSDF)
 
 ///////////////////////////////////////////////////////////////////////////////
 // BSDFs packing/unpacking
+// 11G11B10F and 10F are unsigned float format
+
+uint Pack10F(float Value)
+{
+	return (f32tof16(Value) >> 5) & 0x000003FF;
+}
+
+float Unpack10F(uint Value)
+{
+	return f16tof32((Value << 5) & 0x7FE0);
+}
 
 uint PackR11G11B10F(float3 rgb)
 {
@@ -1194,8 +1223,18 @@ void PackStrataOut(
 		{
 			FStrataBSDF BSDF = Layer.BSDFs[i];
 
-			STRATA_STORE_UINT1(PackR11G11B10F(BSDF.Weight));
-			STRATA_STORE_UINT1(BSDF.State);
+			const uint GreyScaleWeight = StrataHasGreyScaleWeight(BSDF.Weight) ? 1 : 0;
+			BSDF_SETGREYWEIGHT(BSDF, GreyScaleWeight);
+			if (GreyScaleWeight > 0)
+			{
+				BSDF_SETWEIGHT10F(BSDF, Pack10F(BSDF.Weight.x));
+				STRATA_STORE_UINT1(BSDF.State);
+			}
+			else
+			{
+				STRATA_STORE_UINT1(BSDF.State);
+				STRATA_STORE_UINT1(PackR11G11B10F(BSDF.Weight));
+			}
 
 			const uint BSDFType = BSDF_GETTYPE(BSDF);
 			switch (BSDF_GETTYPE(BSDF))
@@ -1286,8 +1325,16 @@ FStrataBSDF UnpackStrataBSDFIn(ByteAddressBuffer StrataBuffer, inout uint ByteOf
 
 	float DummyFloat = 0.0f;
 
-	StrataLoad_R11G11B10F(StrataBuffer, ByteOffset, OutBSDF.Weight);
 	STRATA_LOAD_UINT1(OutBSDF.State);
+	if (BSDF_GETGREYWEIGHT(OutBSDF))
+	{
+		OutBSDF.Weight = Unpack10F(BSDF_GETWEIGHT10F(OutBSDF));
+	}
+	else
+	{
+		StrataLoad_R11G11B10F(StrataBuffer, ByteOffset, OutBSDF.Weight);
+	}
+
 
 	const uint BSDFType = BSDF_GETTYPE(OutBSDF);
 	BSDF_SETTYPE(OutBSDF, BSDFType);
diff --git a/Engine/Shaders/Shared/StrataDefinitions.h b/Engine/Shaders/Shared/StrataDefinitions.h
--- a/Engine/Shaders/Shared/StrataDefinitions.h
+++ b/Engine/Shaders/Shared/StrataDefinitions.h
@@ -15,6 +15,8 @@
 #define STRATA_MAX_BSDF_COUNT_PER_LAYER		4
 #define STRATA_MAX_LAYER_COUNT				4
 #define STRATA_MAX_TOTAL_BSDF				4
+
+// It this is changed, STATE_BITS_SHAREDNORMAL also needs to be updated
 #define STRATA_MAX_SHARED_NORMAL_REGISTERS	4
 
 
