Updated the Strata material vis with the latest features, optimisation and byte usage.

diff --git a/Engine/Shaders/Private/MiniFontCommon.ush b/Engine/Shaders/Private/MiniFontCommon.ush
--- a/Engine/Shaders/Private/MiniFontCommon.ush
+++ b/Engine/Shaders/Private/MiniFontCommon.ush
@@ -126,6 +126,21 @@ void PrintFloatNoFractionLarge(int2 PixelPos, inout float3 OutColor, float3 Font
 	PrintCharacter(PixelPos, OutColor, FontColor, Cursor, ExtractDigitFromFloat(Number, 10));
 	PrintCharacter(PixelPos, OutColor, FontColor, Cursor, ExtractDigitFromFloat(Number, 1));
 }
+void PrintFloatNoFraction(int2 PixelPos, inout float3 OutColor, float3 FontColor, int2 LeftTop, float Number, uint DigitCount)
+{
+	int2 Cursor = LeftTop;
+
+	uint DigitValue = 1;
+	for (int i = 0; i < DigitCount-1; ++i)
+	{
+		DigitValue *= 10;
+	}
+	for (int i = 0; i < DigitCount; ++i)
+	{
+		PrintCharacter(PixelPos, OutColor, FontColor, Cursor, ExtractDigitFromFloat(Number, DigitValue));
+		DigitValue /= 10;
+	}
+}
 
 // for printf debugging in the shader, has to be positive (other abs is taken)
 // outputs a float number in the form: xx.yy
diff --git a/Engine/Shaders/Private/Strata/StrataVisualize.usf b/Engine/Shaders/Private/Strata/StrataVisualize.usf
--- a/Engine/Shaders/Private/Strata/StrataVisualize.usf
+++ b/Engine/Shaders/Private/Strata/StrataVisualize.usf
@@ -76,6 +76,7 @@ void VisualizeMaterialPS(
 		const uint PixelStrataDataByteOffsetPostHeader = PixelStrataDataByteOffset;
 
 		float3 BSDFFontColor = float3(0.25, 1.0, 0.25);
+		float3 BSDFFontColorGrey = float3(0.25, 0.25, 0.25);
 		float3 INFOFontColor = float3(1.0, 0.5, 0.25);
 		float3 PARAMFontColor = float3(0.8, 0.8, 0.8);
 		float3 FontColor = float3(0.25, 1.0, 1.0);
@@ -85,12 +86,18 @@ void VisualizeMaterialPS(
 #define PRINTPARAMCHAR(X) PrintCharacter(PixelPos, OutColor.xyz, PARAMFontColor, TopLeft, X);
 #define PRINTINFOCHAR(X) PrintCharacter(PixelPos, OutColor.xyz, INFOFontColor, TopLeft, X);
 #define PRINTBSDFCHAR(X) PrintCharacter(PixelPos, OutColor.xyz, BSDFFontColor, TopLeft, X);
-
-#define PRINTBSDFWEIGHT()	PRINTBSDFCHAR(_SPC_); PRINTBSDFCHAR(_MINUS_); PRINTBSDFCHAR(_SPC_); PrintFloat(PixelPos, OutColor.xyz, FontColor, TopLeft, BSDF.Weight.x); TopLeft.x+=8*7; PRINTBSDFCHAR(_SPC_); PrintFloat(PixelPos, OutColor.xyz, FontColor, TopLeft, BSDF.Weight.y); TopLeft.x+=8*7; PRINTBSDFCHAR(_SPC_); PrintFloat(PixelPos, OutColor.xyz, FontColor, TopLeft, BSDF.Weight.z);
+#define PRINTBSDFCHARGREY(X) PrintCharacter(PixelPos, OutColor.xyz, BSDFFontColorGrey, TopLeft, X);
+
+#define PRINTBSDFWEIGHT()	PRINTBSDFCHAR(_SPC_); PRINTBSDFCHAR(_MINUS_); PRINTBSDFCHAR(_W_); PRINTBSDFCHAR(_E_); PRINTBSDFCHAR(_I_); PRINTBSDFCHAR(_G_); PRINTBSDFCHAR(_H_);PRINTBSDFCHAR(_T_);  PRINTBSDFCHAR(_SPC_); PrintFloat(PixelPos, OutColor.xyz, FontColor, TopLeft, BSDF.Weight.x); TopLeft.x+=8*7; PRINTBSDFCHAR(_SPC_); PrintFloat(PixelPos, OutColor.xyz, FontColor, TopLeft, BSDF.Weight.y); TopLeft.x+=8*7; PRINTBSDFCHAR(_SPC_); PrintFloat(PixelPos, OutColor.xyz, FontColor, TopLeft, BSDF.Weight.z); TopLeft.x+=8*7; 
+#define PRINTBSDFGREYWEIGHT()PRINTBSDFCHAR(_SPC_);PRINTBSDFCHARGREY(_MINUS_); PRINTBSDFCHARGREY(_W_); PRINTBSDFCHARGREY(_E_); PRINTBSDFCHARGREY(_I_); PRINTBSDFCHARGREY(_G_); PRINTBSDFCHARGREY(_H_);PRINTBSDFCHARGREY(_T_);  PRINTBSDFCHAR(_SPC_); PrintFloat(PixelPos, OutColor.xyz, FontColor, TopLeft, BSDF_GETGREYWEIGHT(BSDF)); TopLeft.x+=8*23; 
+#define PRINTBSDFNORMALID()	PRINTBSDFCHAR(_SPC_); PRINTBSDFCHAR(_MINUS_); PRINTBSDFCHAR(_N_); PRINTBSDFCHAR(_R_); PRINTBSDFCHAR(_M_); PRINTBSDFCHAR(_I_); PRINTBSDFCHAR(_D_); PRINTBSDFCHAR(_SPC_); PrintFloatNoFraction(PixelPos, OutColor.xyz, FontColor, TopLeft, float(BSDF_GETNORMALID(BSDF)), 1); TopLeft.x+=8*1;
+#define PRINTBSDFANISO()	PRINTBSDFCHAR(_SPC_); PRINTBSDFCHAR(_MINUS_); PRINTBSDFCHAR(_A_); PRINTBSDFCHAR(_N_); PRINTBSDFCHAR(_I_); PRINTBSDFCHAR(_S_); PRINTBSDFCHAR(_O_); PRINTBSDFCHAR(_SPC_); PrintFloatNoFraction(PixelPos, OutColor.xyz, FontColor, TopLeft, float(BSDF_GETANISOTROPY(BSDF)), 1); TopLeft.x+=8*1;
+#define PRINTBSDFTOPLAY()	PRINTBSDFCHAR(_SPC_); PRINTBSDFCHAR(_MINUS_); PRINTBSDFCHAR(_T_); PRINTBSDFCHAR(_O_); PRINTBSDFCHAR(_P_); PRINTBSDFCHAR(_L_); PRINTBSDFCHAR(_A_); PRINTBSDFCHAR(_Y_); PRINTBSDFCHAR(_SPC_); PrintFloatNoFraction(PixelPos, OutColor.xyz, FontColor, TopLeft, float(BSDF_GETISTOPLAYER(BSDF)), 1); TopLeft.x+=8*1;
+#define PRINTBSDFHASSCAT()	PRINTBSDFCHAR(_SPC_); PRINTBSDFCHAR(_MINUS_); PRINTBSDFCHAR(_H_); PRINTBSDFCHAR(_A_); PRINTBSDFCHAR(_S_); PRINTBSDFCHAR(_S_); PRINTBSDFCHAR(_C_); PRINTBSDFCHAR(_A_); PRINTBSDFCHAR(_T_); PRINTBSDFCHAR(_SPC_); PrintFloatNoFraction(PixelPos, OutColor.xyz, FontColor, TopLeft, float(BSDF_GETHASSCATTERING(BSDF)), 1); TopLeft.x+=8*1;
 #define PRINTBSDFCOUNT()	PRINTBSDFCHAR(min(_0_+Header.BSDFCount,_F_)); PRINTBSDFCHAR(_SPC_); PRINTBSDFCHAR(_B_); PRINTBSDFCHAR(_S_); PRINTBSDFCHAR(_D_); PRINTBSDFCHAR(_F_);
 #define PRINTFLOAT3(X) PrintFloat(PixelPos, OutColor.xyz, PARAMFontColor, TopLeft, X.x); TopLeft.x+=8*7; PRINTBSDFCHAR(_SPC_); PrintFloat(PixelPos, OutColor.xyz, PARAMFontColor, TopLeft, X.y); TopLeft.x+=8*7; PRINTBSDFCHAR(_SPC_); PrintFloat(PixelPos, OutColor.xyz, PARAMFontColor, TopLeft, X.z);
 #define PRINTFLOAT1(Value) PrintFloat(PixelPos, OutColor.xyz, PARAMFontColor, TopLeft, Value); TopLeft.x+=8*7;
-#define PRINTFLOATNOFRAC(Value) PrintFloatNoFraction(PixelPos, OutColor.xyz, PARAMFontColor, TopLeft, Value); TopLeft.x+=8*7;
+#define PRINTFLOATNOFRAC(Value, DigitCount) PrintFloatNoFraction(PixelPos, OutColor.xyz, PARAMFontColor, TopLeft, Value, DigitCount); TopLeft.x+=8*7;
 #define SIGNOFFLOAT1(Value) if(Value < 0.0f) { PrintCharacter(PixelPos, OutColor.xyz, PARAMFontColor, TopLeft, _MINUS_); } /*else { PrintCharacter(PixelPos, OutColor.xyz, PARAMFontColor, TopLeft, _SPC_);}*/
 
 #define NEXTLINE  TopLeft.y += 8; TopLeft.x = TopLeftInit.x;
@@ -114,9 +121,19 @@ void VisualizeMaterialPS(
 			{
 
 			case STRATA_BSDF_TYPE_DIFFUSE:
-				PRINTBSDFCHAR(_D_); PRINTBSDFCHAR(_I_); PRINTBSDFCHAR(_F_); PRINTBSDFCHAR(_F_);
-				PRINTBSDFCHAR(_U_); PRINTBSDFCHAR(_S_); PRINTBSDFCHAR(_E_); PRINTBSDFCHAR(_SPC_);
-				PRINTBSDFWEIGHT();
+				PRINTBSDFCHAR(_D_); PRINTBSDFCHAR(_I_); PRINTBSDFCHAR(_F_); PRINTBSDFCHAR(_F_); PRINTBSDFCHAR(_U_); 
+				PRINTBSDFCHAR(_S_); PRINTBSDFCHAR(_E_); PRINTBSDFCHAR(_SPC_); PRINTBSDFCHAR(_SPC_); PRINTBSDFCHAR(_SPC_);
+				if (BSDF_GETGREYWEIGHT(BSDF))
+				{
+					PRINTBSDFGREYWEIGHT();
+				}
+				else
+				{
+					PRINTBSDFWEIGHT();
+				}
+				PRINTBSDFNORMALID();
+				PRINTBSDFANISO();
+				PRINTBSDFTOPLAY();
 				NEXTLINE;
 
 				HORIZONTALTAB; PRINTPARAMCHAR(_A_); PRINTPARAMCHAR(_L_); PRINTPARAMCHAR(_B_); PRINTPARAMCHAR(_E_); PRINTPARAMCHAR(_D_); PRINTPARAMCHAR(_O_); 
@@ -131,7 +148,17 @@ void VisualizeMaterialPS(
 			case STRATA_BSDF_TYPE_DIELECTRIC:
 				PRINTBSDFCHAR(_D_); PRINTBSDFCHAR(_I_); PRINTBSDFCHAR(_E_); PRINTBSDFCHAR(_L_); PRINTBSDFCHAR(_E_);
 				PRINTBSDFCHAR(_C_); PRINTBSDFCHAR(_T_); PRINTBSDFCHAR(_R_); PRINTBSDFCHAR(_I_); PRINTBSDFCHAR(_C_);
-				PRINTBSDFWEIGHT();
+				if (BSDF_GETGREYWEIGHT(BSDF))
+				{
+					PRINTBSDFGREYWEIGHT();
+				}
+				else
+				{
+					PRINTBSDFWEIGHT();
+				}
+				PRINTBSDFNORMALID();
+				PRINTBSDFANISO();
+				PRINTBSDFTOPLAY();
 				NEXTLINE;
 
 				HORIZONTALTAB; PRINTPARAMCHAR(_I_); PRINTPARAMCHAR(_O_); PRINTPARAMCHAR(_R_); PRINTPARAMCHAR(_SPC_); PRINTPARAMCHAR(_SPC_); PRINTPARAMCHAR(_SPC_);
@@ -149,7 +176,17 @@ void VisualizeMaterialPS(
 			case STRATA_BSDF_TYPE_CONDUCTOR:
 				PRINTBSDFCHAR(_C_); PRINTBSDFCHAR(_O_); PRINTBSDFCHAR(_N_); PRINTBSDFCHAR(_D_); PRINTBSDFCHAR(_U_);
 				PRINTBSDFCHAR(_C_); PRINTBSDFCHAR(_T_); PRINTBSDFCHAR(_O_); PRINTBSDFCHAR(_R_); PRINTBSDFCHAR(_SPC_);
-				PRINTBSDFWEIGHT();
+				if (BSDF_GETGREYWEIGHT(BSDF))
+				{
+					PRINTBSDFGREYWEIGHT();
+				}
+				else
+				{
+					PRINTBSDFWEIGHT();
+				}
+				PRINTBSDFNORMALID();
+				PRINTBSDFANISO();
+				PRINTBSDFTOPLAY();
 				NEXTLINE;
 
 				HORIZONTALTAB; PRINTPARAMCHAR(_R_); PRINTPARAMCHAR(_E_); PRINTPARAMCHAR(_F_); PRINTPARAMCHAR(_L_); PRINTPARAMCHAR(_E_); PRINTPARAMCHAR(_C_);
@@ -167,7 +204,18 @@ void VisualizeMaterialPS(
 			case STRATA_BSDF_TYPE_VOLUME:
 				PRINTBSDFCHAR(_V_); PRINTBSDFCHAR(_O_); PRINTBSDFCHAR(_L_); PRINTBSDFCHAR(_U_); PRINTBSDFCHAR(_M_); 
 				PRINTBSDFCHAR(_E_); PRINTBSDFCHAR(_SPC_); PRINTBSDFCHAR(_SPC_); PRINTBSDFCHAR(_SPC_); PRINTBSDFCHAR(_SPC_);
-				PRINTBSDFWEIGHT();
+				if (BSDF_GETGREYWEIGHT(BSDF))
+				{
+					PRINTBSDFGREYWEIGHT();
+				}
+				else
+				{
+					PRINTBSDFWEIGHT();
+				}
+				PRINTBSDFNORMALID();
+				PRINTBSDFANISO();
+				PRINTBSDFTOPLAY();
+				PRINTBSDFHASSCAT();
 				NEXTLINE;
 				
 				HORIZONTALTAB; PRINTPARAMCHAR(_A_); PRINTPARAMCHAR(_L_); PRINTPARAMCHAR(_B_); PRINTPARAMCHAR(_E_); PRINTPARAMCHAR(_D_); PRINTPARAMCHAR(_O_); 
@@ -198,17 +246,17 @@ void VisualizeMaterialPS(
 
 		PRINTBSDFCHAR(_H_); PRINTBSDFCHAR(_E_); PRINTBSDFCHAR(_A_); PRINTBSDFCHAR(_D_); PRINTBSDFCHAR(_E_); PRINTBSDFCHAR(_R_);
 		PRINTBSDFCHAR(_B_); PRINTBSDFCHAR(_Y_); PRINTBSDFCHAR(_T_); PRINTBSDFCHAR(_E_); PRINTBSDFCHAR(_S_); PRINTBSDFCHAR(_SPC_);
-		PRINTFLOATNOFRAC(float(PixelStrataDataByteOffsetPostHeader - PixelStrataDataByteOffsetStart));
+		PRINTFLOATNOFRAC(float(PixelStrataDataByteOffsetPostHeader - PixelStrataDataByteOffsetStart), 3);
 		NEXTLINE;
 
 		PRINTBSDFCHAR(_B_); PRINTBSDFCHAR(_S_); PRINTBSDFCHAR(_D_); PRINTBSDFCHAR(_F_); PRINTBSDFCHAR(_S_); PRINTBSDFCHAR(_SPC_);
 		PRINTBSDFCHAR(_B_); PRINTBSDFCHAR(_Y_); PRINTBSDFCHAR(_T_); PRINTBSDFCHAR(_E_); PRINTBSDFCHAR(_S_); PRINTBSDFCHAR(_SPC_);
-		PRINTFLOATNOFRAC(float(PixelStrataDataByteOffsetPostBSDF - PixelStrataDataByteOffsetPostHeader));
+		PRINTFLOATNOFRAC(float(PixelStrataDataByteOffsetPostBSDF - PixelStrataDataByteOffsetPostHeader), 3);
 		NEXTLINE;
 
 		PRINTBSDFCHAR(_T_); PRINTBSDFCHAR(_O_); PRINTBSDFCHAR(_T_); PRINTBSDFCHAR(_A_); PRINTBSDFCHAR(_L_); PRINTBSDFCHAR(_SPC_);
 		PRINTBSDFCHAR(_B_); PRINTBSDFCHAR(_Y_); PRINTBSDFCHAR(_T_); PRINTBSDFCHAR(_E_); PRINTBSDFCHAR(_S_); PRINTBSDFCHAR(_SPC_);
-		PRINTFLOATNOFRAC(float(PixelStrataDataByteOffsetPostBSDF - PixelStrataDataByteOffsetStart));
+		PRINTFLOATNOFRAC(float(PixelStrataDataByteOffsetPostBSDF - PixelStrataDataByteOffsetStart), 3);
 		NEXTLINE;
 	}
 }
