Forward default depth value in TShadowDepthPS if non-perspective-correct option is set to avoid empty pixel shader in Metal.

#rnx

diff --git a/Engine/Shaders/Private/ShadowDepthPixelShader.usf b/Engine/Shaders/Private/ShadowDepthPixelShader.usf
--- a/Engine/Shaders/Private/ShadowDepthPixelShader.usf
+++ b/Engine/Shaders/Private/ShadowDepthPixelShader.usf
@@ -28,7 +28,7 @@
 void Main( 
 	FShadowDepthVSToPS Inputs,
 	in float4 SvPosition : SV_Position		// after all interpolators
-#if PERSPECTIVE_CORRECT_DEPTH
+#if PERSPECTIVE_CORRECT_DEPTH || COMPILER_METAL
 	,out float OutDepth : SV_DEPTH
 #endif
 	)
@@ -64,5 +64,9 @@ void Main(
 	Inputs.ShadowDepth += Inputs.DepthBias;
 
 	OutDepth = saturate(Inputs.ShadowDepth);
+#elif COMPILER_METAL
+	// Metal fragment shader must not be empty,
+	// so output depth value explicitly if this shader permuation was not discarded
+	OutDepth = SvPosition.z;
 #endif
 }
diff --git a/Engine/Shaders/Public/Platform/Metal/MetalCommon.ush b/Engine/Shaders/Public/Platform/Metal/MetalCommon.ush
--- a/Engine/Shaders/Public/Platform/Metal/MetalCommon.ush
+++ b/Engine/Shaders/Public/Platform/Metal/MetalCommon.ush
@@ -5,7 +5,7 @@
 =============================================================================*/
 
 // Update this GUID to improve shader recompilation for Metal only shaders
-// GUID = 8C31264E-8AF0-4CE0-A988-F62D5B2CA088
+// GUID = 75802E83-9450-4B4A-A2A6-93B080C7C8DC
 
 #pragma once
 
