Fixed strata basepass and forward shaders to compile on console and avoid spilling. forward shader are not visible. We need to invstigate that next.

diff --git a/Engine/Shaders/Private/BasePassPixelShader.usf b/Engine/Shaders/Private/BasePassPixelShader.usf
--- a/Engine/Shaders/Private/BasePassPixelShader.usf
+++ b/Engine/Shaders/Private/BasePassPixelShader.usf
@@ -1463,6 +1463,7 @@ void FPixelShaderInOut_MainPS(
 
 	// Initialise a Strata header with normal in registers
 	FStrataPixelHeader StrataPixelHeader = InitialiseStrataPixelHeader();
+	UNROLL
 	for (int l = 0; l < Strata.LayerCount; ++l)
 	{
 		StrataPixelHeader.BSDFCount += Strata.Layers[l].BSDFCount;
diff --git a/Engine/Shaders/Private/Strata/Strata.ush b/Engine/Shaders/Private/Strata/Strata.ush
--- a/Engine/Shaders/Private/Strata/Strata.ush
+++ b/Engine/Shaders/Private/Strata/Strata.ush
@@ -705,10 +705,12 @@ FStrataData StrataAdd(FStrataData A, FStrataData B)
 	FStrataData StrataData = A;
 
 	// Append each BSDF from B to A, with same layer position
+	UNROLL
 	for (int Layer = 0; Layer < B.LayerCount; ++Layer)
 	{
 		FStrataLayer ALayer = A.Layers[Layer];
 		FStrataLayer BLayer = B.Layers[Layer];
+		UNROLL
 		for (int BSDF = 0; BSDF < BLayer.BSDFCount; BSDF++)
 		{
 			StrataData.Layers[Layer].BSDFs[ALayer.BSDFCount + BSDF] = BLayer.BSDFs[BSDF];
@@ -727,8 +729,10 @@ FStrataData StrataMultiply(FStrataData A, float Weight)
 	FStrataData StrataData = A;
 
 	// Apply the weight to each BSDF
+	UNROLL
 	for (int Layer = 0; Layer < StrataData.LayerCount; ++Layer)
 	{
+		UNROLL
 		for (int BSDF = 0; BSDF < StrataData.Layers[Layer].BSDFCount; ++BSDF)
 		{
 			StrataData.Layers[Layer].BSDFs[BSDF].Weight *= Weight;
@@ -754,6 +758,7 @@ FStrataData StrataVerticalLayering(FStrataData Top, FStrataData Base)
 
 	// Add each layer from Base under Top
 	int TopLayerCount = Top.LayerCount;
+	UNROLL
 	for (int Layer = 0; Layer < Base.LayerCount; ++Layer)
 	{
 		StrataData.Layers[TopLayerCount + Layer] = Base.Layers[Layer];
@@ -1056,6 +1061,7 @@ void PackStrataOut(
 	float3 Throughput = 1.0f;
 
 	int l;
+	UNROLL
 	for (l = 0; l < Strata.LayerCount; ++l)
 	{
 		float3 LayerThroughput = 0.0f;
@@ -1067,6 +1073,7 @@ void PackStrataOut(
 		float3 AccumBaseColor = 0.0f;
 		const bool bTopLayer = l == 0;
 
+		UNROLL
 		for (int i = 0; i < Strata.Layers[l].BSDFCount; ++i)
 		{
 			// Sanitize BSDF in place before it is used and packed out to memory
@@ -1207,6 +1214,7 @@ void PackStrataOut(
 	STRATA_STORE_UINT1(SharedNormals_BSDFCount);
 
 #if STRATA_SHARED_NORMALS_IN_REGISTERS
+	UNROLL
 	for (uint i = 0; i < StrataPixelHeader.SharedNormals.NormalCount; ++i)
 	{
 		STRATA_STORE_UINT1(StrataPackNormal(StrataPixelHeader.SharedNormals.Normals[i]));
@@ -1215,10 +1223,12 @@ void PackStrataOut(
 
 	// 2. the list of BSDFs
 	BSDFCount = 0;
+	UNROLL
 	for (l = 0; l < Strata.LayerCount; ++l)
 	{
 		FStrataLayer Layer = Strata.Layers[l];
 
+		UNROLL
 		for (int i = 0; i < Layer.BSDFCount; ++i)
 		{
 			FStrataBSDF BSDF = Layer.BSDFs[i];
diff --git a/Engine/Shaders/Private/Strata/StrataForwardLighting.ush b/Engine/Shaders/Private/Strata/StrataForwardLighting.ush
--- a/Engine/Shaders/Private/Strata/StrataForwardLighting.ush
+++ b/Engine/Shaders/Private/Strata/StrataForwardLighting.ush
@@ -23,11 +23,13 @@ float3 StrataForwardLighting(
 	float3 Color = 0;
 	OutThroughput = 1.0f;
 
+	UNROLL
 	for (int l = 0; l < Strata.LayerCount; ++l)
 	{
 		float3 LayerThroughput = 0.0f;
 		float3 LayerCoverage = 0.0f;
 
+		UNROLL
 		for (int i = 0; i < Strata.Layers[l].BSDFCount; ++i)
 		{
 			FStrataBSDF BSDF = Strata.Layers[l].BSDFs[i];
