Add tangent basis encoding into shared normal basis.

diff --git a/Engine/Shaders/Private/Strata/Strata.ush b/Engine/Shaders/Private/Strata/Strata.ush
--- a/Engine/Shaders/Private/Strata/Strata.ush
+++ b/Engine/Shaders/Private/Strata/Strata.ush
@@ -35,7 +35,7 @@ struct FStrataBSDF
 
 	float3				Weight;		// Weight is a float3 because, when loaded from the material buffer, it contains the BSDF throughput (coverage and transmittance combined)
 
-	float4 VGPRs[3];
+	float4 VGPRs[2];
 };
 
 // Num bits for each component store in BSDF state
@@ -96,25 +96,19 @@ struct FStrataBSDF
 
 #define DIELECTRIC_TINT(X)				X.VGPRs[0].xyz
 #define DIELECTRIC_IOR(X)				X.VGPRs[0].w
-#define DIELECTRIC_ROUGHNESS0(X)		X.VGPRs[1].w
-#define DIELECTRIC_ROUGHNESS1(X)		X.VGPRs[2].w
-#define DIELECTRIC_TANGENT(X)			X.VGPRs[2].xyz
+#define DIELECTRIC_ROUGHNESS0(X)		X.VGPRs[1].y
+#define DIELECTRIC_ROUGHNESS1(X)		X.VGPRs[1].x
 
 #define CONDUCTOR_REFLECTIVITY(X)		X.VGPRs[0].xyz
 #define CONDUCTOR_EDGECOLOR(X)			X.VGPRs[1].xyz
-#define CONDUCTOR_ROUGHNESS0(X)			X.VGPRs[2].w
+#define CONDUCTOR_ROUGHNESS0(X)			X.VGPRs[0].w
 #define CONDUCTOR_ROUGHNESS1(X)			X.VGPRs[1].w
-#define CONDUCTOR_TANGENT(X)			X.VGPRs[2].xyz
 
 #define VOLUME_ALBEDO(X)				X.VGPRs[0].xyz
 #define VOLUME_EXTINCTION(X)			X.VGPRs[1].xyz
 #define VOLUME_ANISOTROPY(X)			X.VGPRs[0].w
 #define VOLUME_THICKNESS(X)				X.VGPRs[1].w
 
-// Forward declaractions
-float3 StrataUnpackNormal(uint PackedNormal);
-
-
 bool StrataHasAnisotropy(in float Roughness0, float Roughness1)
 {
 	return Roughness0 != Roughness1;
@@ -178,7 +172,7 @@ FStrataData GetStrataDiffuseBSDF(float3 Albedo, float Roughness, uint SharedNorm
 	return StrataData;
 }
 
-FStrataData GetStrataDielectricBSDF(float2 Roughness, float IOR, float3 Tint, float3 Tangent, uint SharedNormalIndex)
+FStrataData GetStrataDielectricBSDF(float2 Roughness, float IOR, float3 Tint, uint SharedNormalIndex)
 {
 	FStrataData StrataData = GetInitialisedStrataData();
 	StrataData.LayerCount++;
@@ -194,14 +188,13 @@ FStrataData GetStrataDielectricBSDF(float2 Roughness, float IOR, float3 Tint, fl
 	DIELECTRIC_TINT			(StrataData.Layers[0].BSDFs[0])		= Tint;
 	DIELECTRIC_ROUGHNESS0	(StrataData.Layers[0].BSDFs[0])		= Roughness.x;
 	DIELECTRIC_ROUGHNESS1	(StrataData.Layers[0].BSDFs[0])		= Roughness.y;
-	DIELECTRIC_TANGENT		(StrataData.Layers[0].BSDFs[0])		= Tangent;
 
 	StrataData.Layers[0].BSDFs[0].Weight						= 1.0f;
 
 	return StrataData;
 }
 
-FStrataData GetStrataConductorBSDF(float3 Reflectivity, float3 EdgeColor, float2 Roughness, float3 Tangent, uint SharedNormalIndex)
+FStrataData GetStrataConductorBSDF(float3 Reflectivity, float3 EdgeColor, float2 Roughness, uint SharedNormalIndex)
 {
 	FStrataData StrataData = GetInitialisedStrataData();
 	StrataData.LayerCount++;
@@ -217,7 +210,6 @@ FStrataData GetStrataConductorBSDF(float3 Reflectivity, float3 EdgeColor, float2
 	CONDUCTOR_EDGECOLOR		(StrataData.Layers[0].BSDFs[0])		= EdgeColor;
 	CONDUCTOR_ROUGHNESS0	(StrataData.Layers[0].BSDFs[0])		= Roughness.x;
 	CONDUCTOR_ROUGHNESS1	(StrataData.Layers[0].BSDFs[0])		= Roughness.y;
-	CONDUCTOR_TANGENT		(StrataData.Layers[0].BSDFs[0])		= Tangent;
 
 	StrataData.Layers[0].BSDFs[0].Weight						= 1.0f;
 
@@ -258,12 +250,17 @@ float3 StrataGetOffSpecularPeakReflectionDir(float3 Normal, float3 ReflectionVec
 ///////////////////////////////////////////////////////////////////////////////
 // Shared normals
 
+float3 StrataUnpackNormal(uint PackedNormal);
+void StrataUnpackNormalAndTangent(inout float3 Normal, inout float3 Tangent, in uint InPacked);
+
 #define STRATA_PACKED_NORMAL_STRIDE 4
 
+#if STRATA_SHARED_NORMALS_IN_REGISTERS
 struct FSharedNormals
 {
 	uint NormalCount;
 	float3 Normals[STRATA_MAX_SHARED_NORMAL_REGISTERS];
+	float3 Tangents[STRATA_MAX_SHARED_NORMAL_REGISTERS];
 };
 
 FSharedNormals StrataInitialiseSharedNormals()
@@ -271,6 +268,7 @@ FSharedNormals StrataInitialiseSharedNormals()
 	FSharedNormals SharedNormals = (FSharedNormals)0;
 	return SharedNormals;
 }
+#endif
 
 struct FStrataPixelHeader
 {
@@ -297,16 +295,35 @@ FStrataPixelHeader InitialiseStrataPixelHeader()
 	return Out;
 }
 
-float3 StrataGetBSDFSharedNormal(in FStrataPixelHeader StrataPixelHeader, in FStrataBSDF BSDF)
+float3x3 StrataGetBSDFSharedBasis(in FStrataPixelHeader StrataPixelHeader, in FStrataBSDF BSDF)
 {
 #if STRATA_SHARED_NORMALS_IN_REGISTERS
-	float3 Normal = StrataPixelHeader.SharedNormals.Normals[BSDF_GETNORMALID(BSDF)];
+	// Assume Normal & Tangent are already been normalized
+	float3 Normal  = StrataPixelHeader.SharedNormals.Normals[BSDF_GETNORMALID(BSDF)];
+	float3 Tangent = StrataPixelHeader.SharedNormals.Tangents[BSDF_GETNORMALID(BSDF)];
+	float3 Bitangent = cross(Normal, Tangent);
+	return float3x3(Tangent, Bitangent, Normal);
 #else
-	uint PackedNormal = StrataPixelHeader.StrataBuffer.Load(StrataPixelHeader.SharedNormalByteOffset + BSDF_GETNORMALID(BSDF) * STRATA_PACKED_NORMAL_STRIDE);
-	float3 Normal     = StrataUnpackNormal(PackedNormal);
-#endif
+	const bool bHasAnisotropy = BSDF_GETANISOTROPY(BSDF) != 0;
+	uint PackedBasis = StrataPixelHeader.StrataBuffer.Load(StrataPixelHeader.SharedNormalByteOffset + BSDF_GETNORMALID(BSDF) * STRATA_PACKED_NORMAL_STRIDE);
+	float3x3 OutTangentBasis;
+	if (bHasAnisotropy)
+	{
+		float3 Normal;
+		float3 Tangent;
+		StrataUnpackNormalAndTangent(Normal, Tangent, PackedBasis);
 
-	return Normal;
+		OutTangentBasis[0] = Tangent;
+		OutTangentBasis[1] = cross(Normal, Tangent);
+		OutTangentBasis[2] = Normal;
+	}
+	else
+	{
+		float3 Normal = StrataUnpackNormal(PackedBasis);
+		OutTangentBasis = GetTangentBasis(Normal);
+	}
+	return OutTangentBasis;
+#endif
 }
 
 ///////////////////////////////////////////////////////////////////////////////
@@ -366,25 +383,6 @@ void StrataGetBSDFRoughness(in FStrataBSDF BSDF, inout float R)
 	}
 }
 
-void StrataGetBSDFTangent(in FStrataBSDF BSDF, inout float3 T)
-{
-	T = 0;
-	const uint BSDFType = BSDF_GETTYPE(BSDF);
-	switch (BSDFType)
-	{
-	case STRATA_BSDF_TYPE_DIELECTRIC:
-	{
-		T = DIELECTRIC_TANGENT(BSDF);
-		break;
-	}
-	case STRATA_BSDF_TYPE_CONDUCTOR:
-	{
-		T = CONDUCTOR_TANGENT(BSDF);
-		break;
-	}
-	}
-}
-
 FStrataBSDFContext StrataCreateBSDFContext(FStrataPixelHeader StrataPixelHeader, FStrataBSDF BSDF, float3 V, float3 L)
 {
 	FStrataBSDFContext BSDFContext = (FStrataBSDFContext)0;
@@ -394,26 +392,25 @@ FStrataBSDFContext StrataCreateBSDFContext(FStrataPixelHeader StrataPixelHeader,
 	BSDFContext.V = V;
 	BSDFContext.L = L;
 
-	BSDFContext.N = StrataGetBSDFSharedNormal(StrataPixelHeader, BSDF);
-	BSDFContext.N = normalize(BSDFContext.N);
+	float3x3 TangentBasis = StrataGetBSDFSharedBasis(StrataPixelHeader, BSDF);
+	BSDFContext.X = TangentBasis[0]; 
+	BSDFContext.Y = TangentBasis[1];
+	BSDFContext.N = TangentBasis[2]; 
 
 	BSDFContext.R = 2 * dot(BSDFContext.V, BSDFContext.N) * BSDFContext.N - BSDFContext.V;
 	BSDFContext.H = normalize(BSDFContext.V + BSDFContext.L);
 
 	BSDFContext.Context = (BxDFContext)0;
-	const bool bHasAnisotropy = BSDF_GETANISOTROPY(BSDF) != 0;
-	if (bHasAnisotropy)
+	if (BSDF_GETANISOTROPY(BSDF) != 0)
 	{
-		StrataGetBSDFTangent(BSDF, BSDFContext.X);
-		BSDFContext.Y = normalize(cross(BSDFContext.N, BSDFContext.X));
 		Init(BSDFContext.Context, BSDFContext.N, BSDFContext.X, BSDFContext.Y, BSDFContext.V, BSDFContext.L);
-		BSDFContext.TangentBasis = float3x3(BSDFContext.X, BSDFContext.Y, BSDFContext.N);
 	}
 	else
 	{
 		Init(BSDFContext.Context, BSDFContext.N, BSDFContext.V, BSDFContext.L);
-		BSDFContext.TangentBasis = GetTangentBasis(BSDFContext.N);
 	}
+	BSDFContext.TangentBasis = float3x3(BSDFContext.X, BSDFContext.Y, BSDFContext.N);
+
 	BSDFContext.TangentV = mul(BSDFContext.TangentBasis, BSDFContext.V);
 	BSDFContext.TangentH = mul(BSDFContext.TangentBasis, BSDFContext.H);
 
@@ -924,6 +921,7 @@ float2 SignNotZero(in float2 v)
 {
 	return float2(SignNotZero(v.x), SignNotZero(v.y));
 }
+
 uint StrataPackNormal(in float3 Normal)
 {
 	float l1norm = abs(Normal.x) + abs(Normal.y) + abs(Normal.z);
@@ -936,6 +934,7 @@ uint StrataPackNormal(in float3 Normal)
 	uint Y10Bits = uint(clamp(result.y * 511.0f + 512.0f, 0.0f, 1023.0f));
 	return (X10Bits | (Y10Bits << 10));
 }
+
 float3 StrataUnpackNormal(uint PackedNormal)
 {
 	uint X10Bits = 0x3FF & PackedNormal;
@@ -951,6 +950,99 @@ float3 StrataUnpackNormal(uint PackedNormal)
 	return normalize(Normal);
 }
 
+#define STRATA_TANGENT_SYMMETRY 1
+uint StrataPackTangentBasis(in float3 Normal, in float3 Tangent)
+{
+	uint X10Bits;
+	uint Y10Bits;
+	uint Z11Bits;
+	uint W1Bit;
+
+	// Oct. encoding for the normal
+	{
+		float l1norm = abs(Normal.x) + abs(Normal.y) + abs(Normal.z);
+		float2 result = Normal.xy * (1.0 / l1norm);
+		if (Normal.z < 0.0)
+		{
+			result = (1.0 - abs(result.yx)) * SignNotZero(result.xy);
+		}
+		X10Bits = uint(clamp(result.x * 511.0f + 512.0f, 0.0f, 1023.0f));
+		Y10Bits = uint(clamp(result.y * 511.0f + 512.0f, 0.0f, 1023.0f));
+	}
+
+	// Oct. encoding for the projected tangent, relative to an generated ortho frame based on normal
+	{
+		// Use the quantize normal to regenerate tje ortho. frame. This removes discontinuities
+		float x = float(int(X10Bits) - 512) / 511.0f;
+		float y = float(int(Y10Bits) - 512) / 511.0f;
+		Normal = float3(x, y, 1.0f - abs(x) - abs(y));
+		if (Normal.z < 0.0f)
+		{
+			Normal.xy = (1.0f - abs(Normal.yx)) * SignNotZero(Normal.xy);
+		}
+		//Normal = normalize(Normal);
+
+		float3x3 LocalRef = GetTangentBasis(Normal);
+		float2 ProjXY = float2(dot(LocalRef[0], Tangent), dot(LocalRef[1], Tangent));
+	#if STRATA_TANGENT_SYMMETRY
+		// Use symmetry to improve precision
+		ProjXY = ProjXY.x >= 0 ? ProjXY : -ProjXY.xy;
+	#endif
+		float l1norm = abs(ProjXY.x) + abs(ProjXY.y);
+		float2 result = ProjXY.xy * (1.0 / l1norm);
+
+	#if STRATA_TANGENT_SYMMETRY
+		Z11Bits = uint(clamp(result.x * 2047.0f, 0.0f, 2047.0f));
+	#else
+		Z11Bits = uint(clamp(result.x * 1023.0f + 1024.0f, 0.0f, 2047.0f));
+	#endif
+		W1Bit   = ProjXY.y >= 0 ? 1u : 0u;
+	}
+
+	return
+		 (0x3FF & X10Bits)        | 
+		((0x3FF & Y10Bits) << 10) | 
+		((0x7FF & Z11Bits) << 20) |
+		((0x1   & W1Bit  ) << 31);
+}
+
+void StrataUnpackNormalAndTangent(inout float3 Normal, inout float3 Tangent, in uint InPacked)
+{
+	uint X10Bits = 0x3FF &  InPacked;
+	uint Y10Bits = 0x3FF & (InPacked>>10);
+	uint Z11bits = 0x7FF & (InPacked>>20);
+	uint W1Bit   = 0x1   & (InPacked>>31);
+
+	// Oct. decoding for the normal
+	{
+		float x = float(int(X10Bits) - 512) / 511.0f;
+		float y = float(int(Y10Bits) - 512) / 511.0f;
+
+		Normal = float3(x, y, 1.0f - abs(x) - abs(y));
+		if (Normal.z < 0.0f)
+		{
+			Normal.xy = (1.0f - abs(Normal.yx)) * SignNotZero(Normal.xy);
+		}
+		Normal = normalize(Normal);
+	}
+
+	// Oct. decoding for the projected tangent, relative to an generated ortho frame based on normal
+	{
+		// Create intermediate basis
+		float3x3 LocalRef = GetTangentBasis(Normal);
+	#if STRATA_TANGENT_SYMMETRY
+		float x = float(Z11bits) / 2047.0f;
+	#else
+		float x = float(int(Z11bits) - 1024) / 1023.0f;
+	#endif
+		float Sign = (W1Bit > 0 ? 1 : -1);
+		float2 ProjX = float2(x, 1.0f - abs(x));
+		Tangent = LocalRef[0] * ProjX.x + LocalRef[1] * ProjX.y * Sign;
+		Tangent = normalize(Tangent);
+	}
+}
+
+
 #if COMPILER_SWITCH // These platforms are complaining so simply disabling Strata there for now.
 
 #define STRATA_STORE_UINT1(x) 
@@ -1085,8 +1177,9 @@ void PackStrataOut(
 			BSDFCount++;
 
 			// STRATA_TODO apply the same normal transform as for the material
-			float3 N = StrataGetBSDFSharedNormal(StrataPixelHeader, BSDF);
-			N = normalize(N);
+			float3x3 TangentBasis = StrataGetBSDFSharedBasis(StrataPixelHeader, BSDF);
+			float3 N = TangentBasis[2];
+			//N = normalize(N);
 
 			AccumTopNormal += N;
 
@@ -1217,7 +1310,7 @@ void PackStrataOut(
 	UNROLL
 	for (uint i = 0; i < StrataPixelHeader.SharedNormals.NormalCount; ++i)
 	{
-		STRATA_STORE_UINT1(StrataPackNormal(StrataPixelHeader.SharedNormals.Normals[i]));
+		STRATA_STORE_UINT1(StrataPackTangentBasis(StrataPixelHeader.SharedNormals.Normals[i], StrataPixelHeader.SharedNormals.Tangents[i]));
 	}
 #endif
 
@@ -1259,24 +1352,14 @@ void PackStrataOut(
 			{
 				STRATA_STORE_UINT1(PackR11G11B10F(DIELECTRIC_TINT(BSDF)));
 				STRATA_STORE_UINT1(PackR11G11B10F(float3(DIELECTRIC_ROUGHNESS0(BSDF), DIELECTRIC_ROUGHNESS1(BSDF), DIELECTRIC_IOR(BSDF))));
-				const bool bHasAnisotropy = BSDF_GETANISOTROPY(BSDF);
-				if (bHasAnisotropy)
-				{
-					STRATA_STORE_UINT1(StrataPackNormal(DIELECTRIC_TANGENT(BSDF)));
-				}
-				// 12 bytes
+				// 8 bytes
 			}
 			break;
 			case STRATA_BSDF_TYPE_CONDUCTOR:
 			{
 				STRATA_STORE_UINT1(PackAlbedoLinearToGamma2AlphaLinear(float4(CONDUCTOR_REFLECTIVITY(BSDF), CONDUCTOR_ROUGHNESS0(BSDF))));
 				STRATA_STORE_UINT1(PackAlbedoLinearToGamma2AlphaLinear(float4(CONDUCTOR_EDGECOLOR(BSDF),    CONDUCTOR_ROUGHNESS1(BSDF))));
-				const bool bHasAnisotropy = BSDF_GETANISOTROPY(BSDF);
-				if (bHasAnisotropy)
-				{
-					STRATA_STORE_UINT1(StrataPackNormal(DIELECTRIC_TANGENT(BSDF)));
-				}
-				// 12 bytes
+				// 8 bytes
 			}
 			break;
 			case STRATA_BSDF_TYPE_VOLUME:
@@ -1359,28 +1442,12 @@ FStrataBSDF UnpackStrataBSDFIn(ByteAddressBuffer StrataBuffer, inout uint ByteOf
 	{
 		StrataLoad_R11G11B10F			(StrataBuffer, ByteOffset, DIELECTRIC_TINT(OutBSDF));
 		StrataLoad_R11_G11_B10F			(StrataBuffer, ByteOffset, DIELECTRIC_ROUGHNESS0(OutBSDF), DIELECTRIC_ROUGHNESS1(OutBSDF), DIELECTRIC_IOR(OutBSDF));
-
-		const bool bHasAnisotropy = BSDF_GETANISOTROPY(OutBSDF);
-		if (bHasAnisotropy)
-		{
-			uint PackedTangent;
-			STRATA_LOAD_UINT1(PackedTangent);
-			DIELECTRIC_TANGENT(OutBSDF) = StrataUnpackNormal(PackedTangent);
-		}
 	}
 	break;
 	case STRATA_BSDF_TYPE_CONDUCTOR:
 	{
 		StrataLoad_AlbedoGamma2ToLinearAlphaLinear(StrataBuffer, ByteOffset, CONDUCTOR_REFLECTIVITY(OutBSDF),	CONDUCTOR_ROUGHNESS0(OutBSDF));
 		StrataLoad_AlbedoGamma2ToLinearAlphaLinear(StrataBuffer, ByteOffset, CONDUCTOR_EDGECOLOR(OutBSDF),		CONDUCTOR_ROUGHNESS1(OutBSDF));
-
-		const bool bHasAnisotropy = BSDF_GETANISOTROPY(OutBSDF);
-		if (bHasAnisotropy)
-		{
-			uint PackedTangent;
-			STRATA_LOAD_UINT1(PackedTangent);
-			CONDUCTOR_TANGENT(OutBSDF) = StrataUnpackNormal(PackedTangent);
-		}
 	}
 	break;
 	case STRATA_BSDF_TYPE_VOLUME:
diff --git a/Engine/Source/Runtime/Engine/Private/Materials/HLSLMaterialTranslator.cpp b/Engine/Source/Runtime/Engine/Private/Materials/HLSLMaterialTranslator.cpp
--- a/Engine/Source/Runtime/Engine/Private/Materials/HLSLMaterialTranslator.cpp
+++ b/Engine/Source/Runtime/Engine/Private/Materials/HLSLMaterialTranslator.cpp
@@ -1590,10 +1590,21 @@ void FHLSLMaterialTranslator::GetSharedInputsMaterialCode(FString& PixelMembersD
 		if (CodeChunkToStrataSharedNormal.Num() > 0)
 		{
 			// Generate code to fill up the array containing normal shared between Strata BSDFs
-			for (TMap<uint64, FStrataSharedNormalInfo>::TIterator It(CodeChunkToStrataSharedNormal); It; ++It)
+			for (TMultiMap<uint64, FStrataSharedNormalInfo>::TIterator It(CodeChunkToStrataSharedNormal); It; ++It)
 			{
 				PixelInputInitializerValues += FString::Printf(TEXT("\tParameters.SharedNormals.Normals[%u] = %s;\n"), It->Value.SharedNormalIndex, *It->Value.SharedNormalCode);
 			}
+			for (TMultiMap<uint64, FStrataSharedNormalInfo>::TIterator It(CodeChunkToStrataSharedNormal); It; ++It)
+			{
+				if (It->Value.TangentCodeChunk != -1)
+				{
+					PixelInputInitializerValues += FString::Printf(TEXT("\tParameters.SharedNormals.Tangents[%u] = %s;\n"), It->Value.SharedNormalIndex, *It->Value.SharedTangentCode);
+				}
+				else
+				{
+					PixelInputInitializerValues += FString::Printf(TEXT("\tParameters.SharedNormals.Tangents[%u] = float3(1,0,0);\n"), It->Value.SharedNormalIndex);
+				}
+			}
 			PixelInputInitializerValues += FString::Printf(TEXT("\tParameters.SharedNormals.NormalCount = %u;\n"), CodeChunkToStrataSharedNormal.Num());
 		}
 
@@ -7298,15 +7309,66 @@ uint8 FHLSLMaterialTranslator::StrataCompilationInfoRegisterSharedNormalIndex(in
 
 	const uint64 NormalCodeChunkHash = GetParameterHash(NormalCodeChunk);
 
-	if (CodeChunkToStrataSharedNormal.Find(NormalCodeChunkHash) == nullptr)
+	// Find any basis which match the Normal code chunk
+	// A normal can be duplicated when it is paired with different tangent, so find the first one which matches
+	TArray<FStrataSharedNormalInfo*> NormalInfos;
+	CodeChunkToStrataSharedNormal.MultiFindPointer(NormalCodeChunkHash, NormalInfos);
+	if (NormalInfos.Num() == 0)
 	{
 		// Allocate a new slot for a new shared normal
 		uint8 NewSharedNormalIndex = NextFreeStrataShaderNormalIndex++;
-		CodeChunkToStrataSharedNormal.Add(NormalCodeChunkHash, { NormalCodeChunk , NewSharedNormalIndex, *GetParameterCode(NormalCodeChunk) });
+		CodeChunkToStrataSharedNormal.Add(NormalCodeChunkHash, { NormalCodeChunk , -1, NewSharedNormalIndex, *GetParameterCode(NormalCodeChunk), FString() });
 		return NewSharedNormalIndex;
 	}
+	// Return the first existing code chunk which match the normal chunk code
+	return NormalInfos[0]->SharedNormalIndex;
+}
+
+uint8 FHLSLMaterialTranslator::StrataCompilationInfoRegisterSharedNormalIndex(int32 NormalCodeChunk, int32 TangentCodeChunk)
+{
+	check(NormalCodeChunk != INDEX_NONE);
+	check(TangentCodeChunk != INDEX_NONE);
+	check(NextFreeStrataShaderNormalIndex < 255);	// Out of shared normal slots
+	check(NextFreeStrataShaderNormalIndex < STRATA_MAX_TOTAL_BSDF);	// This is our current budget if normal per material. STRATA_TODO change that to a byte per pixel
+
+	const uint64 NormalCodeChunkHash = GetParameterHash(NormalCodeChunk);
+	const uint64 TangentCodeChunkHash = GetParameterHash(TangentCodeChunk);
+
+	// Find a basis which matches both the Normal & the Tangent code chunks
+	TArray<FStrataSharedNormalInfo*> NormalInfos;
+	CodeChunkToStrataSharedNormal.MultiFindPointer(NormalCodeChunkHash, NormalInfos);
+	bool bFound = false;
+	uint8 SharedNormalIndex = 0;
+	for (FStrataSharedNormalInfo* NormalInfo : NormalInfos)
+	{
+		// * Either we find a perfect match (normal & tangent matches)
+		// * Or we find a normal which doesn't have a tangent associated with, and we set the tangent for code
+		if (TangentCodeChunkHash == GetParameterHash(NormalInfo->TangentCodeChunk))
+		{
+			SharedNormalIndex = NormalInfo->SharedNormalIndex;
+			bFound = true;
+			break;
+		}
+		else if (NormalInfo->TangentCodeChunk == -1)
+		{
+			NormalInfo->TangentCodeChunk = TangentCodeChunk;
+			NormalInfo->SharedTangentCode = *GetParameterCode(TangentCodeChunk);
+			SharedNormalIndex = NormalInfo->SharedNormalIndex;
+			bFound = true;
+			break;
+		}
+	}
+	
+	if (NormalInfos.Num() == 0 || !bFound)
+	{
+		// Allocate a new slot for a new shared normal & tangent
+		uint8 NewSharedNormalIndex = NextFreeStrataShaderNormalIndex++;
+		CodeChunkToStrataSharedNormal.Add(NormalCodeChunkHash, { NormalCodeChunk , TangentCodeChunk, NewSharedNormalIndex, *GetParameterCode(NormalCodeChunk), *GetParameterCode(TangentCodeChunk) });
+		return NewSharedNormalIndex;
+	}
+
 	// Return the existing code chunk
-	return CodeChunkToStrataSharedNormal[NormalCodeChunkHash].SharedNormalIndex;
+	return SharedNormalIndex;
 }
 
 uint8 FHLSLMaterialTranslator::StrataCompilationInfoGetSharedNormalCount()
@@ -7333,12 +7395,11 @@ int32 FHLSLMaterialTranslator::StrataDiffuseBSDF(int32 Albedo, int32 Roughness,
 int32 FHLSLMaterialTranslator::StrataDielectricBSDF(int32 RoughnessX, int32 RoughnessY, int32 IOR, int32 Tint, int32 Normal, int32 Tangent, uint8 SharedNormalIndex)
 {
 	return AddCodeChunk(
-		MCT_Strata, TEXT("GetStrataDielectricBSDF(float2(%s, %s), %s, %s, %s, %u) /* Normal:%s Tangent:%s */"),
+		MCT_Strata, TEXT("GetStrataDielectricBSDF(float2(%s, %s), %s, %s, %u) /* Normal:%s Tangent:%s */"),
 		*GetParameterCode(RoughnessX),
 		*GetParameterCode(RoughnessY),
 		*GetParameterCode(IOR),
 		*GetParameterCode(Tint),
-		*GetParameterCode(Tangent),
 		SharedNormalIndex,
 		*GetParameterCode(Normal),
 		*GetParameterCode(Tangent)
@@ -7348,12 +7409,11 @@ int32 FHLSLMaterialTranslator::StrataDielectricBSDF(int32 RoughnessX, int32 Roug
 int32 FHLSLMaterialTranslator::StrataConductorBSDF(int32 Reflectivity, int32 EdgeColor, int32 RoughnessX, int32 RoughnessY, int32 Normal, int32 Tangent, uint8 SharedNormalIndex)
 {
 	return AddCodeChunk(
-		MCT_Strata, TEXT("GetStrataConductorBSDF(%s, %s, float2(%s, %s), %s, %u) /* Normal:%s Tangent:%s */"),
+		MCT_Strata, TEXT("GetStrataConductorBSDF(%s, %s, float2(%s, %s), %u) /* Normal:%s Tangent:%s */"),
 		*GetParameterCode(Reflectivity),
 		*GetParameterCode(EdgeColor),
 		*GetParameterCode(RoughnessX),
 		*GetParameterCode(RoughnessY),
-		*GetParameterCode(Tangent),
 		SharedNormalIndex,
 		*GetParameterCode(Normal),
 		*GetParameterCode(Tangent)
diff --git a/Engine/Source/Runtime/Engine/Private/Materials/HLSLMaterialTranslator.h b/Engine/Source/Runtime/Engine/Private/Materials/HLSLMaterialTranslator.h
--- a/Engine/Source/Runtime/Engine/Private/Materials/HLSLMaterialTranslator.h
+++ b/Engine/Source/Runtime/Engine/Private/Materials/HLSLMaterialTranslator.h
@@ -307,15 +307,18 @@ protected:
 	/** Will contain all the shading models picked up from the material expression graph */
 	FMaterialShadingModelField ShadingModelsFromCompilation;
 
+	/** Tracks shared normal & tangent for strata materials. */
 	TMap<int32, FStrataMaterialCompilationInfo> CodeChunkToStrataCompilationInfoMap; // Mapping from node output to StrataMaterial
 	uint8 NextFreeStrataShaderNormalIndex;
 	struct FStrataSharedNormalInfo
 	{
 		int32 NormalCodeChunk;
+		int32 TangentCodeChunk;
 		uint8 SharedNormalIndex;
 		FString SharedNormalCode;
+		FString SharedTangentCode;
 	};
-	TMap<uint64, FStrataSharedNormalInfo> CodeChunkToStrataSharedNormal;
+	TMultiMap<uint64, FStrataSharedNormalInfo> CodeChunkToStrataSharedNormal; // Stored into to a multi-map as a normal can be shared accross several tangent basis (i.e. having different tangent)
 
 	/** Tracks the total number of vt samples in the shader. */
 	uint32 NumVtSamples;
@@ -803,6 +806,7 @@ protected:
 	virtual bool StrataCompilationInfoContainsCodeChunk(int32 CodeChunk) override;
 	virtual const FStrataMaterialCompilationInfo& GetStrataCompilationInfo(int32 CodeChunk) override;
 	virtual uint8 StrataCompilationInfoRegisterSharedNormalIndex(int32 NormalCodeChunk) override;
+	virtual uint8 StrataCompilationInfoRegisterSharedNormalIndex(int32 NormalCodeChunk, int32 TangentCodeChunk) override;
 	virtual uint8 StrataCompilationInfoGetSharedNormalCount() override;
 
 #if HANDLE_CUSTOM_OUTPUTS_AS_MATERIAL_ATTRIBUTES
diff --git a/Engine/Source/Runtime/Engine/Private/Materials/MaterialExpressions.cpp b/Engine/Source/Runtime/Engine/Private/Materials/MaterialExpressions.cpp
--- a/Engine/Source/Runtime/Engine/Private/Materials/MaterialExpressions.cpp
+++ b/Engine/Source/Runtime/Engine/Private/Materials/MaterialExpressions.cpp
@@ -19096,7 +19096,7 @@ int32 UMaterialExpressionStrataDielectricBSDF::Compile(class FMaterialCompiler*
 {
 	int32 NormalCodeChunk = Normal.GetTracedInput().Expression ? Normal.Compile(Compiler) : Compiler->VertexNormal();
 	int32 TangentCodeChunk = Tangent.GetTracedInput().Expression ? Tangent.Compile(Compiler) : Compiler->VertexTangent();
-	uint8 SharedNormalIndex = StrataCompilationInfoCreateSharedNormal(Compiler, NormalCodeChunk);
+	uint8 SharedNormalIndex = StrataCompilationInfoCreateSharedNormal(Compiler, NormalCodeChunk, TangentCodeChunk);
 
 	int32 RoughnessXCodeChunk = RoughnessX.GetTracedInput().Expression ? RoughnessX.Compile(Compiler) : Compiler->Constant(0.0f);
 	// If not plugged in, RoughnessYCodeChunk is set to RoughnessXCodeChunk to get an isotropic behavior
@@ -19177,7 +19177,12 @@ int32 UMaterialExpressionStrataConductorBSDF::Compile(class FMaterialCompiler* C
 {
 	int32 NormalCodeChunk = Normal.GetTracedInput().Expression ? Normal.Compile(Compiler) : Compiler->VertexNormal();
 	int32 TangentCodeChunk = Tangent.GetTracedInput().Expression ? Tangent.Compile(Compiler) : Compiler->VertexTangent();
-	uint8 SharedNormalIndex = StrataCompilationInfoCreateSharedNormal(Compiler, NormalCodeChunk);
+
+	// For some reason, when editing shader, Normal.Compile() / Tangent.Compile() might return -1. In this case we fallback to vertex normal/tangent
+	if (NormalCodeChunk == -1)  { NormalCodeChunk = Compiler->VertexNormal(); }
+	if (TangentCodeChunk == -1)	{ TangentCodeChunk = Compiler->VertexTangent();	}
+
+	uint8 SharedNormalIndex = StrataCompilationInfoCreateSharedNormal(Compiler, NormalCodeChunk, TangentCodeChunk);
 
 	int32 RoughnessXCodeChunk = RoughnessX.GetTracedInput().Expression ? RoughnessX.Compile(Compiler) : Compiler->Constant(0.0f);
 	// If not plugged in, RoughnessYCodeChunk is set to RoughnessXCodeChunk to get an isotropic behavior
diff --git a/Engine/Source/Runtime/Engine/Private/Materials/StrataMaterial.cpp b/Engine/Source/Runtime/Engine/Private/Materials/StrataMaterial.cpp
--- a/Engine/Source/Runtime/Engine/Private/Materials/StrataMaterial.cpp
+++ b/Engine/Source/Runtime/Engine/Private/Materials/StrataMaterial.cpp
@@ -41,6 +41,10 @@ uint8 StrataCompilationInfoCreateSharedNormal(FMaterialCompiler* Compiler, int32
 	return Compiler->StrataCompilationInfoRegisterSharedNormalIndex(NormalCodeChunk);
 }
 
+uint8 StrataCompilationInfoCreateSharedNormal(FMaterialCompiler* Compiler, int32 NormalCodeChunk, int32 TangentCodeChunk)
+{
+	return Compiler->StrataCompilationInfoRegisterSharedNormalIndex(NormalCodeChunk, TangentCodeChunk);
+}
 
 void StrataCompilationInfoCreateSingleBSDFMaterial(FMaterialCompiler* Compiler, int32 CodeChunk, uint8 SharedNormalIndex, uint8 BSDFType)
 {
diff --git a/Engine/Source/Runtime/Engine/Private/Materials/StrataMaterial.h b/Engine/Source/Runtime/Engine/Private/Materials/StrataMaterial.h
--- a/Engine/Source/Runtime/Engine/Private/Materials/StrataMaterial.h
+++ b/Engine/Source/Runtime/Engine/Private/Materials/StrataMaterial.h
@@ -42,6 +42,7 @@ struct FStrataMaterialCompilationInfo
 FString GetStrataBSDFName(uint8 BSDFType);
 
 uint8 StrataCompilationInfoCreateSharedNormal(FMaterialCompiler* Compiler, int32 NormalCodeChunk);
+uint8 StrataCompilationInfoCreateSharedNormal(FMaterialCompiler* Compiler, int32 NormalCodeChunk, int32 TangentCodeChunk);
 void  StrataCompilationInfoCreateSingleBSDFMaterial(FMaterialCompiler* Compiler, int32 CodeChunk, uint8 SharedNormalIndex, uint8 BSDFType);
 
 FStrataMaterialCompilationInfo StrataCompilationInfoAdd(FMaterialCompiler* Compiler, const FStrataMaterialCompilationInfo& A, const FStrataMaterialCompilationInfo& B);
diff --git a/Engine/Source/Runtime/Engine/Public/MaterialCompiler.h b/Engine/Source/Runtime/Engine/Public/MaterialCompiler.h
--- a/Engine/Source/Runtime/Engine/Public/MaterialCompiler.h
+++ b/Engine/Source/Runtime/Engine/Public/MaterialCompiler.h
@@ -407,6 +407,7 @@ public:
 	virtual bool StrataCompilationInfoContainsCodeChunk(int32 CodeChunk) = 0;
 	virtual const FStrataMaterialCompilationInfo& GetStrataCompilationInfo(int32 CodeChunk) = 0;
 	virtual uint8 StrataCompilationInfoRegisterSharedNormalIndex(int32 NormalCodeChunk) = 0;
+	virtual uint8 StrataCompilationInfoRegisterSharedNormalIndex(int32 NormalCodeChunk, int32 TangentCodeChunk) = 0;
 	virtual uint8 StrataCompilationInfoGetSharedNormalCount() = 0;
 
 	// Water
@@ -878,6 +879,11 @@ public:
 		return Compiler->StrataCompilationInfoRegisterSharedNormalIndex(NormalCodeChunk);
 	}
 
+	virtual uint8 StrataCompilationInfoRegisterSharedNormalIndex(int32 NormalCodeChunk, int32 TangentCodeChunk) override
+	{
+		return Compiler->StrataCompilationInfoRegisterSharedNormalIndex(NormalCodeChunk, TangentCodeChunk);
+	}
+
 	virtual uint8 StrataCompilationInfoGetSharedNormalCount() override
 	{
 		return Compiler->StrataCompilationInfoGetSharedNormalCount();
diff --git a/Engine/Source/Runtime/Renderer/Private/Strata/Strata.cpp b/Engine/Source/Runtime/Renderer/Private/Strata/Strata.cpp
--- a/Engine/Source/Runtime/Renderer/Private/Strata/Strata.cpp
+++ b/Engine/Source/Runtime/Renderer/Private/Strata/Strata.cpp
@@ -136,6 +136,7 @@ class FVisualizeMaterialPS : public FGlobalShader
 		SHADER_PARAMETER_STRUCT_REF(FStrataGlobalUniformParameters, Strata)
 		SHADER_PARAMETER_TEXTURE(Texture2D, MiniFontTexture)
 		SHADER_PARAMETER_STRUCT_INCLUDE(FSceneTextureParameters, SceneTextures)
+		SHADER_PARAMETER_STRUCT_INCLUDE(ShaderDrawDebug::FShaderDrawDebugParameters, ShaderDrawParameters)
 		RENDER_TARGET_BINDING_SLOTS()
 	END_SHADER_PARAMETER_STRUCT()
 
@@ -182,6 +183,11 @@ void AddVisualizeMaterialPasses(FRDGBuilder& GraphBuilder, const TArray<FViewInf
 			PassParameters->SceneTextures = GetSceneTextureParameters(GraphBuilder);
 			PassParameters->RenderTargets[0] = FRenderTargetBinding(SceneColorTexture, ERenderTargetLoadAction::ELoad);
 
+			if (ShaderDrawDebug::IsShaderDrawDebugEnabled())
+			{
+				ShaderDrawDebug::SetParameters(GraphBuilder, View.ShaderDrawData, PassParameters->ShaderDrawParameters);
+			}
+
 			FPixelShaderUtils::AddFullscreenPass<FVisualizeMaterialPS>(GraphBuilder, View.ShaderMap, RDG_EVENT_NAME("StrataVisualizeMaterial"),
 				PixelShader, PassParameters, View.ViewRect, PreMultipliedColorTransmittanceBlend);
 		}
