Strata: now uses max byte per pixel
Green preflight

diff --git a/Engine/Shaders/Private/Strata/Strata.ush b/Engine/Shaders/Private/Strata/Strata.ush
--- a/Engine/Shaders/Private/Strata/Strata.ush
+++ b/Engine/Shaders/Private/Strata/Strata.ush
@@ -19,16 +19,15 @@
 #define STRATA_SHARED_NORMALS_IN_REGISTERS 1
 #endif
 
+#ifndef STRATA_CLAMPED_LAYER_COUNT
+#define STRATA_CLAMPED_LAYER_COUNT 1
+#endif
+
 
 
 
 ///////////////////////////////////////////////////////////////////////////////
 // BSDFs representation
-//		STRATA_MAX_TOTAL_BSDF * sizeof(FStrataConductor) +				// 4 * 12
-//		STRATA_MAX_TOTAL_BSDF * (sizeof(State) + sizeof(Weight))		// 4 * 8
-//		sizeof(NormalCount_BSDFCount)									// 1 * 4
-//		STRATA_MAX_TOTAL_BSDF * sizeof(Normal)							// 4 * 4
-//		= 100															// See Strata.cpp
 struct FStrataBSDF
 {
 	uint				State;		// (27 unused bits) (1bit as HasAnisotropy) (2bits as SharedNormalIndex) (3bits as BSDFType)
@@ -66,7 +65,7 @@ struct FStrataBSDF
 #define BSDF_SETTYPE(X, NewType)		BSDF_SETSTATE(X, STATE_BITS_BSDF, STATE_BIT_OFFSET_BSDF, NewType)
 
 // Accessors for Shared Normal
-// This assumes Normals[STRATA_MAX_TOTAL_BSDF]
+// This assumes Normals[STRATA_MAX_SHARED_NORMAL_REGISTERS]
 #define BSDF_GETNORMALID(X)				BSDF_GETSTATE(X, STATE_BITS_SHAREDNORMALID, STATE_BIT_OFFSET_SHAREDNORMALID)
 #define BSDF_SETNORMALID(X, NrmlIdx)	BSDF_SETSTATE(X, STATE_BITS_SHAREDNORMALID, STATE_BIT_OFFSET_SHAREDNORMALID, NrmlIdx)
 
@@ -253,8 +252,6 @@ float3 StrataGetOffSpecularPeakReflectionDir(float3 Normal, float3 ReflectionVec
 float3 StrataUnpackNormal(uint PackedNormal);
 void StrataUnpackNormalAndTangent(inout float3 Normal, inout float3 Tangent, in uint InPacked);
 
-#define STRATA_PACKED_NORMAL_STRIDE 4
-
 #if STRATA_SHARED_NORMALS_IN_REGISTERS
 struct FSharedNormals
 {
@@ -305,7 +302,7 @@ float3x3 StrataGetBSDFSharedBasis(in FStrataPixelHeader StrataPixelHeader, in FS
 	return float3x3(Tangent, Bitangent, Normal);
 #else
 	const bool bHasAnisotropy = BSDF_GETANISOTROPY(BSDF) != 0;
-	uint PackedBasis = StrataPixelHeader.StrataBuffer.Load(StrataPixelHeader.SharedNormalByteOffset + BSDF_GETNORMALID(BSDF) * STRATA_PACKED_NORMAL_STRIDE);
+	uint PackedBasis = StrataPixelHeader.StrataBuffer.Load(StrataPixelHeader.SharedNormalByteOffset + BSDF_GETNORMALID(BSDF) * STRATA_PACKED_NORMAL_STRIDE_BYTES);
 	float3x3 OutTangentBasis;
 	if (bHasAnisotropy)
 	{
@@ -1128,6 +1125,7 @@ uint GetStrataPixelDataByteOffset(uint2 PixelPos, uint ViewBufferSizeX, uint Str
 	return (PixelPos.y * ViewBufferSizeX + PixelPos.x) * StrataMaxBytesPerPixel;
 }
 
+// If this is changed, please update the compiler side material size evaluation in StrataMaterial.cpp
 void PackStrataOut(
 	RWByteAddressBuffer StrataBuffer, 
 	uint ByteOffset, 
@@ -1137,13 +1135,6 @@ void PackStrataOut(
 	float3 V,
 	inout FGBufferData GBuffer)
 {
-	// We cannot output (layers x bsdf) information because if we try to load them later, the reserved amount of VGPRS would be crazy high.
-	// So instead we simply store a list of BSDF and we allow only up to STRATA_MAX_TOTAL_BSDF of them.
-	// For each of the BSDF, we evaluated their final waight according to the material operation performed (layering, mixing, etc.).
-
-	// STRATA_TODO:
-	//	- We could clean up irrelevant BSDFs 
-	//	- We could generate transmittance BSDFs 
 
 	const float OpaqueBSDFThroughput = 0.0f;
 	const float FullThroughput = 1.0f;
@@ -1152,6 +1143,10 @@ void PackStrataOut(
 	uint BSDFCount = 0;
 	float3 Throughput = 1.0f;
 
+	// We clamp the layer count to the maximum one according to our byte per pixel budget (removing the bottom layer first)
+	Strata.LayerCount = STRATA_CLAMPED_LAYER_COUNT;
+
+	// For each of the BSDF, we evaluated their combined weight/throughput according to the material operation performed (layering, mixing, etc.).
 	int l;
 	UNROLL
 	for (l = 0; l < Strata.LayerCount; ++l)
@@ -1298,7 +1293,6 @@ void PackStrataOut(
 	// Now write out Strata data
 
 	// 1. the header
-	BSDFCount = min(BSDFCount, uint(STRATA_MAX_TOTAL_BSDF));
 
 	uint SharedNormals_BSDFCount = BSDFCount;
 #if STRATA_SHARED_NORMALS_IN_REGISTERS
@@ -1379,15 +1373,12 @@ void PackStrataOut(
 
 			// Stop writing if above budget.
 			BSDFCount++;
-			if (BSDFCount >= STRATA_MAX_TOTAL_BSDF)
-			{
-				return;
-			}
 		}
 	}
 }
 
 // Unpack only strata header
+// If this is changed, please update the compiler side material size evaluation in StrataMaterial.cpp
 FStrataPixelHeader UnpackStrataHeaderIn(ByteAddressBuffer StrataBuffer, inout uint ByteOffset)
 {
 	FStrataPixelHeader Out = InitialiseStrataPixelHeader();
@@ -1404,7 +1395,7 @@ FStrataPixelHeader UnpackStrataHeaderIn(ByteAddressBuffer StrataBuffer, inout ui
 	Out.StrataBuffer = StrataBuffer;
 
 	// Now skip over the shared normal memory to be able to load the BSDF content
-	ByteOffset += NormalCount * STRATA_PACKED_NORMAL_STRIDE;
+	ByteOffset += NormalCount * STRATA_PACKED_NORMAL_STRIDE_BYTES;
 #endif
 
 	return Out;
@@ -1412,6 +1403,7 @@ FStrataPixelHeader UnpackStrataHeaderIn(ByteAddressBuffer StrataBuffer, inout ui
 
 // Unpack a single BSDF
 // Note: All BSDF lobes needs to be unpack in a sequential manner since each BSDF has a variable footprint.
+// If this is changed, please update the compiler side material size evaluation in StrataMaterial.cpp
 FStrataBSDF UnpackStrataBSDFIn(ByteAddressBuffer StrataBuffer, inout uint ByteOffset)
 {
 	FStrataBSDF OutBSDF = (FStrataBSDF)0;
diff --git a/Engine/Shaders/Shared/StrataDefinitions.h b/Engine/Shaders/Shared/StrataDefinitions.h
--- a/Engine/Shaders/Shared/StrataDefinitions.h
+++ b/Engine/Shaders/Shared/StrataDefinitions.h
@@ -14,11 +14,11 @@
 
 #define STRATA_MAX_BSDF_COUNT_PER_LAYER		4
 #define STRATA_MAX_LAYER_COUNT				4
-#define STRATA_MAX_TOTAL_BSDF				4
 
 // It this is changed, STATE_BITS_SHAREDNORMAL also needs to be updated
 #define STRATA_MAX_SHARED_NORMAL_REGISTERS	4
 
+#define STRATA_PACKED_NORMAL_STRIDE_BYTES	4
 
 
 #define STRATA_BSDF_TYPE_DIFFUSE			0
diff --git a/Engine/Source/Runtime/Engine/Classes/Engine/RendererSettings.h b/Engine/Source/Runtime/Engine/Classes/Engine/RendererSettings.h
--- a/Engine/Source/Runtime/Engine/Classes/Engine/RendererSettings.h
+++ b/Engine/Source/Runtime/Engine/Classes/Engine/RendererSettings.h
@@ -711,14 +711,23 @@ class ENGINE_API URendererSettings : public UDeveloperSettings
 		uint32 bSupportSkyAtmosphereAffectsHeightFog : 1;
 
 	/**
-	"Strata settings."
+	"Enable Strata materials (Beta)."
 	*/
 	UPROPERTY(config, EditAnywhere, Category = Strata, meta = (
-		ConsoleVariable = "r.Strata", DisplayName = "Strata materials",
-		ToolTip = "Enable the new strata material worfklow and evaluation.",
+		ConsoleVariable = "r.Strata", DisplayName = "Strata materials (Beta)",
+		ToolTip = "Enable Strata materials (Beta).",
 		ConfigRestartRequired = true))
 		uint32 bEnableStrata : 1;
 
+	/**
+	"Strata allocated byte per pixel to store materials data. Higher value means more complex material can be represented."
+	*/
+	UPROPERTY(config, EditAnywhere, Category = Strata, meta = (
+		ConsoleVariable = "r.Strata.BytesPerPixel", DisplayName = "Strata bytes per pixel",
+		ToolTip = "Strata allocated byte per pixel to store materials data. Higher value means more complex material can be represented.",
+		ConfigRestartRequired = true))
+		uint32 StrataBytePerPixel;
+
 	/**
 	"Skin cache allows a compute shader to skin once each vertex, save those results into a new buffer and reuse those calculations when later running the depth, base and velocity passes. This also allows opting into the 'recompute tangents' for skinned mesh instance feature. Disabling will reduce the number of shader permutations required per material. Changing this setting requires restarting the editor."
 	*/
diff --git a/Engine/Source/Runtime/Engine/Private/Materials/HLSLMaterialTranslator.cpp b/Engine/Source/Runtime/Engine/Private/Materials/HLSLMaterialTranslator.cpp
--- a/Engine/Source/Runtime/Engine/Private/Materials/HLSLMaterialTranslator.cpp
+++ b/Engine/Source/Runtime/Engine/Private/Materials/HLSLMaterialTranslator.cpp
@@ -1106,8 +1106,14 @@ bool FHLSLMaterialTranslator::Translate()
 		{
 			const FStrataMaterialCompilationInfo& StrataCompilationInfo = CodeChunkToStrataCompilationInfoMap[Chunk[MP_FrontMaterial]];
 
-			FString StrataMaterialDescription = "";
+			static const auto CVarStrataBytePerPixel = IConsoleManager::Get().FindTConsoleVariableDataInt(TEXT("r.Strata.BytesPerPixel"));
+			check(CVarStrataBytePerPixel)
+			const uint32 StrataBytePerPixel = CVarStrataBytePerPixel ? CVarStrataBytePerPixel->GetValueOnAnyThread() : 0;
+			StrataMaterialAnalysis = StrataCompilationInfoMaterialAnalysis(this, StrataCompilationInfo, StrataBytePerPixel);
 
+			FString StrataMaterialDescription;
+
+			StrataMaterialDescription += FString::Printf(TEXT("----- STRATA -----\r\n"));
 			StrataMaterialDescription += FString::Printf(TEXT("StrataCompilationInfo -\r\n"));
 			StrataMaterialDescription += FString::Printf(TEXT(" - TotalBSDFCount    = %i\r\n"), StrataCompilationInfo.TotalBSDFCount);
 			StrataMaterialDescription += FString::Printf(TEXT(" - SharedNormalCount = %i\r\n"), StrataCompilationInfoGetSharedNormalCount());
@@ -1123,11 +1129,33 @@ bool FHLSLMaterialTranslator::Translate()
 				}
 			}
 
+			StrataMaterialDescription += FString::Printf(TEXT("Byte Per Pixel Budget      %u\r\n"), StrataBytePerPixel);
+			StrataMaterialDescription += FString::Printf(TEXT("Result.bFitInMemoryBudget  %s\r\n"), StrataMaterialAnalysis.bFitInMemoryBudget ? TEXT("YES") : TEXT("NO"));
+			StrataMaterialDescription += FString::Printf(TEXT("Result.RequestedLayerCount %u\r\n"), StrataMaterialAnalysis.RequestedLayerCount);
+			StrataMaterialDescription += FString::Printf(TEXT("Result.RequestedBSDFCount  %u\r\n"), StrataMaterialAnalysis.RequestedBSDFCount);
+			StrataMaterialDescription += FString::Printf(TEXT("Result.RequestedByteCount  %u\r\n"), StrataMaterialAnalysis.RequestedByteCount);
+			if (!StrataMaterialAnalysis.bFitInMemoryBudget)
+			{
+				StrataMaterialDescription += FString::Printf(TEXT("Result.ClampedLayerCount   %u\r\n"), StrataMaterialAnalysis.ClampedLayerCount);
+				StrataMaterialDescription += FString::Printf(TEXT("Result.ClampedBSDFCount    %u\r\n"), StrataMaterialAnalysis.ClampedBSDFCount);
+				StrataMaterialDescription += FString::Printf(TEXT("Result.UsedByteCount       %u\r\n"), StrataMaterialAnalysis.UsedByteCount);
+			}
+			StrataMaterialDescription += FString::Printf(TEXT("------------------\r\n"));
+
 			ResourcesString += TEXT("/*");
 			ResourcesString += StrataMaterialDescription;
 			ResourcesString += TEXT("*/");
 
 			MaterialCompilationOutput.StrataMaterialDescription = StrataMaterialDescription;
+
+			if (StrataMaterialAnalysis.ClampedLayerCount == 0)
+			{
+				UE_LOG(LogMaterial, Fatal, TEXT("Material %s cannot have any layers rendered due to its complexity (asset: %s).\r\n"), *Material->GetDebugName(), *Material->GetAssetPath());
+			}
+			else if (StrataMaterialAnalysis.RequestedLayerCount > StrataMaterialAnalysis.ClampedLayerCount)
+			{
+				UE_LOG(LogMaterial, Warning, TEXT("Material %s is not fitting the allocated byte per pixel budget for starta. Layers will be removed (asset: %s).\r\n"), *Material->GetDebugName(), *Material->GetAssetPath());
+			}
 		}
 		else
 		{
@@ -1464,6 +1492,12 @@ void FHLSLMaterialTranslator::GetMaterialEnvironment(EShaderPlatform InPlatform,
 
 	// if duals source blending (colored transmittance) is not supported on a platform, it will fall back to standard alpha blending (grey scale transmittance)
 	OutEnvironment.SetDefine(TEXT("DUAL_SOURCE_COLOR_BLENDING_ENABLED"), bMaterialRequestsDualSourceBlending && Material->IsDualBlendingEnabled(Platform) ? TEXT("1") : TEXT("0"));
+
+	// Translate() is called before getting here so we can create
+	if (StrataMaterialAnalysis.RequestedBSDFCount > 0)
+	{
+		OutEnvironment.SetDefine(TEXT("STRATA_CLAMPED_LAYER_COUNT"), StrataMaterialAnalysis.ClampedLayerCount);
+	}
 }
 
 // Assign custom interpolators to slots, packing them as much as possible in unused slots.
@@ -7329,7 +7363,7 @@ uint8 FHLSLMaterialTranslator::StrataCompilationInfoRegisterSharedNormalIndex(in
 	check(NormalCodeChunk != INDEX_NONE);
 	check(TangentCodeChunk != INDEX_NONE);
 	check(NextFreeStrataShaderNormalIndex < 255);	// Out of shared normal slots
-	check(NextFreeStrataShaderNormalIndex < STRATA_MAX_TOTAL_BSDF);	// This is our current budget if normal per material. STRATA_TODO change that to a byte per pixel
+	check(NextFreeStrataShaderNormalIndex < STRATA_MAX_SHARED_NORMAL_REGISTERS);	// This is the size of the HLSL array limiting the maximum shared normal count.
 
 	const uint64 NormalCodeChunkHash = GetParameterHash(NormalCodeChunk);
 	const uint64 TangentCodeChunkHash = GetParameterHash(TangentCodeChunk);
diff --git a/Engine/Source/Runtime/Engine/Private/Materials/HLSLMaterialTranslator.h b/Engine/Source/Runtime/Engine/Private/Materials/HLSLMaterialTranslator.h
--- a/Engine/Source/Runtime/Engine/Private/Materials/HLSLMaterialTranslator.h
+++ b/Engine/Source/Runtime/Engine/Private/Materials/HLSLMaterialTranslator.h
@@ -319,6 +319,7 @@ protected:
 		FString SharedTangentCode;
 	};
 	TMultiMap<uint64, FStrataSharedNormalInfo> CodeChunkToStrataSharedNormal; // Stored into to a multi-map as a normal can be shared accross several tangent basis (i.e. having different tangent)
+	FStrataMaterialAnalysisResult StrataMaterialAnalysis;
 
 	/** Tracks the total number of vt samples in the shader. */
 	uint32 NumVtSamples;
diff --git a/Engine/Source/Runtime/Engine/Private/Materials/StrataMaterial.cpp b/Engine/Source/Runtime/Engine/Private/Materials/StrataMaterial.cpp
--- a/Engine/Source/Runtime/Engine/Private/Materials/StrataMaterial.cpp
+++ b/Engine/Source/Runtime/Engine/Private/Materials/StrataMaterial.cpp
@@ -60,12 +60,6 @@ void StrataCompilationInfoCreateSingleBSDFMaterial(FMaterialCompiler* Compiler,
 
 FStrataMaterialCompilationInfo StrataCompilationInfoAdd(FMaterialCompiler* Compiler, const FStrataMaterialCompilationInfo& A, const FStrataMaterialCompilationInfo& B)
 {
-	if ((A.TotalBSDFCount + B.TotalBSDFCount) > STRATA_MAX_TOTAL_BSDF)
-	{
-		Compiler->Error(TEXT("Adding would result in too many BSDFs"));
-		return A;
-	}
-
 	FStrataMaterialCompilationInfo StrataInfo = A;
 
 	// Append each BSDF from B to A, with same layer position
@@ -103,24 +97,12 @@ FStrataMaterialCompilationInfo StrataCompilationInfoMultiply(FMaterialCompiler*
 
 FStrataMaterialCompilationInfo StrataCompilationInfoHorizontalMixing(FMaterialCompiler* Compiler, const FStrataMaterialCompilationInfo& A, const FStrataMaterialCompilationInfo& B)
 {
-	if ((A.TotalBSDFCount + B.TotalBSDFCount) > STRATA_MAX_TOTAL_BSDF)
-	{
-		Compiler->Error(TEXT("Mixing would result in too many BSDFs"));
-		return A;
-	}
-
 	return StrataCompilationInfoAdd(Compiler, A, B); // Mixing is a similar operation to Add when it comes to bsdf count
 }
 
 
 FStrataMaterialCompilationInfo StrataCompilationInfoVerticalLayering(FMaterialCompiler* Compiler, const FStrataMaterialCompilationInfo& Top, const FStrataMaterialCompilationInfo& Base)
 {
-	if ((Top.TotalBSDFCount + Base.TotalBSDFCount) > STRATA_MAX_TOTAL_BSDF) // See StrataDefinitions.h
-	{
-		Compiler->Error(TEXT("Layering would result in too many BSDFs"));
-		return Base;
-	}
-
 	if ((Top.LayerCount + Base.LayerCount) > STRATA_MAX_LAYER_COUNT)
 	{
 		Compiler->Error(TEXT("Layering would result in too many Layers"));
@@ -141,3 +123,104 @@ FStrataMaterialCompilationInfo StrataCompilationInfoVerticalLayering(FMaterialCo
 	return StrataInfo;
 }
 
+FStrataMaterialAnalysisResult::FStrataMaterialAnalysisResult()
+{
+	bFitInMemoryBudget = true;
+	RequestedLayerCount = 0;
+	RequestedBSDFCount = 0;
+	RequestedByteCount = 0;
+	ClampedLayerCount = 0;
+	ClampedBSDFCount = 0;
+	UsedByteCount = 0;
+}
+
+FStrataMaterialAnalysisResult StrataCompilationInfoMaterialAnalysis(FMaterialCompiler* Compiler, const FStrataMaterialCompilationInfo& Material, const uint32 StrataBytePerPixel)
+{
+	const uint32 UintByteSize = sizeof(uint32);
+
+	FStrataMaterialAnalysisResult Result;
+
+	// 1. Header
+
+	// SharedNormals_BSDFCount
+	Result.RequestedByteCount += UintByteSize;
+	// Shared normals between BSDFs
+	Result.RequestedByteCount += Compiler->StrataCompilationInfoGetSharedNormalCount() * STRATA_PACKED_NORMAL_STRIDE_BYTES;
+
+	// 2. The list of BSDFs
+
+	// We process layers from top to bottom to cull the bottom ones in case we run out of pixel bytes
+	for (uint32 LayerIt = 0; LayerIt < Material.LayerCount; LayerIt++)
+	{
+		const FStrataMaterialCompilationInfo::FLayer& Layer = Material.Layers[LayerIt];
+
+		for (uint32 BSDFIt = 0; BSDFIt < Layer.BSDFCount; BSDFIt++)
+		{
+			const FStrataMaterialCompilationInfo::FBSDF& BSDF = Layer.BSDFs[BSDFIt];
+
+			// BSDF state
+			Result.RequestedByteCount += UintByteSize;
+
+			// From the compiler side, we can only assume the top layer has grey scale weight/throughput
+			const bool bMayBeColoredWeight = LayerIt > 0;
+			if (bMayBeColoredWeight)
+			{
+				Result.RequestedByteCount += UintByteSize;
+			}
+
+			switch (BSDF.Type)
+			{
+			case STRATA_BSDF_TYPE_DIFFUSE:
+			{
+				Result.RequestedByteCount += UintByteSize;
+				break;
+			}
+			case STRATA_BSDF_TYPE_DIELECTRIC:
+			{
+				Result.RequestedByteCount += UintByteSize;
+				Result.RequestedByteCount += UintByteSize;
+				break;
+			}
+			case STRATA_BSDF_TYPE_CONDUCTOR:
+			{
+				Result.RequestedByteCount += UintByteSize;
+				Result.RequestedByteCount += UintByteSize;
+				break;
+			}
+			case STRATA_BSDF_TYPE_VOLUME:
+			{
+				Result.RequestedByteCount += UintByteSize;
+
+				const bool bMayHaveScattering = true;
+				if (bMayHaveScattering)
+				{
+					Result.RequestedByteCount += UintByteSize;
+				}
+				break;
+			}
+			}
+		}
+
+		Result.RequestedLayerCount++;
+		Result.RequestedBSDFCount += Layer.BSDFCount;
+
+		if (Result.RequestedByteCount <= StrataBytePerPixel && Result.bFitInMemoryBudget)
+		{
+			// We only validate all the BSDF of a layer if it remains within budget and we are not already out of budget
+			Result.ClampedBSDFCount += Layer.BSDFCount;
+			Result.ClampedLayerCount++;
+
+			// Set the current used bytes
+			Result.UsedByteCount = Result.RequestedByteCount;
+		}
+		else
+		{
+			// Used byte count remains unchanged and we notify that we start peeling off top layers now
+			Result.bFitInMemoryBudget = false;
+		}
+	}
+	return Result;
+
+}
+	
+
diff --git a/Engine/Source/Runtime/Engine/Private/Materials/StrataMaterial.h b/Engine/Source/Runtime/Engine/Private/Materials/StrataMaterial.h
--- a/Engine/Source/Runtime/Engine/Private/Materials/StrataMaterial.h
+++ b/Engine/Source/Runtime/Engine/Private/Materials/StrataMaterial.h
@@ -51,3 +51,15 @@ FStrataMaterialCompilationInfo StrataCompilationInfoHorizontalMixing(FMaterialCo
 FStrataMaterialCompilationInfo StrataCompilationInfoVerticalLayering(FMaterialCompiler* Compiler, const FStrataMaterialCompilationInfo& Top, const FStrataMaterialCompilationInfo& Base);
 
 
+struct FStrataMaterialAnalysisResult
+{
+	FStrataMaterialAnalysisResult();
+	bool bFitInMemoryBudget;
+	uint32 RequestedLayerCount;
+	uint32 RequestedBSDFCount;
+	uint32 RequestedByteCount;
+	uint32 ClampedLayerCount;
+	uint32 ClampedBSDFCount;
+	uint32 UsedByteCount;
+};
+FStrataMaterialAnalysisResult StrataCompilationInfoMaterialAnalysis(FMaterialCompiler* Compiler, const FStrataMaterialCompilationInfo& Material, const uint32 StrataBytePerPixel);
diff --git a/Engine/Source/Runtime/RenderCore/Private/Shader.cpp b/Engine/Source/Runtime/RenderCore/Private/Shader.cpp
--- a/Engine/Source/Runtime/RenderCore/Private/Shader.cpp
+++ b/Engine/Source/Runtime/RenderCore/Private/Shader.cpp
@@ -1699,10 +1699,17 @@ void ShaderMapAppendKeyString(EShaderPlatform Platform, FString& KeyString)
 
 	{
 		static const auto CVar = IConsoleManager::Get().FindTConsoleVariableDataInt(TEXT("r.Strata"));
-		if (CVar && CVar->GetValueOnAnyThread() > 0)
+		const bool bStrataEnabled = CVar && CVar->GetValueOnAnyThread() > 0;
+		if (bStrataEnabled)
 		{
 			KeyString += TEXT("_STRATA");
 		}
+
+		static const auto CVarBudget = IConsoleManager::Get().FindTConsoleVariableDataInt(TEXT("r.Strata.BytesPerPixel"));
+		if (bStrataEnabled && CVarBudget)
+		{
+			KeyString += FString::Printf(TEXT("_BUDGET%u"), CVarBudget->GetValueOnAnyThread());
+		}
 	}
 
 	{
diff --git a/Engine/Source/Runtime/Renderer/Private/Strata/Strata.cpp b/Engine/Source/Runtime/Renderer/Private/Strata/Strata.cpp
--- a/Engine/Source/Runtime/Renderer/Private/Strata/Strata.cpp
+++ b/Engine/Source/Runtime/Renderer/Private/Strata/Strata.cpp
@@ -16,9 +16,14 @@
 static TAutoConsoleVariable<int32> CVarStrata(
 	TEXT("r.Strata"),
 	0,
-	TEXT("Enables Strata."),
+	TEXT("Enable Strata materials (Beta)."),
 	ECVF_ReadOnly | ECVF_RenderThreadSafe);
 
+static TAutoConsoleVariable<int32> CVarStrataBytePerPixel(
+	TEXT("r.Strata.BytesPerPixel"),
+	80,
+	TEXT("Strata allocated byte per pixel to store materials data. Higher value means more complex material can be represented."),
+	ECVF_ReadOnly | ECVF_RenderThreadSafe);
 
 
 IMPLEMENT_GLOBAL_SHADER_PARAMETER_STRUCT(FStrataGlobalUniformParameters, "Strata");
