Add support for different encoding for normal only and normal/tangent shared normals.

diff --git a/Engine/Shaders/Private/Strata/Strata.ush b/Engine/Shaders/Private/Strata/Strata.ush
--- a/Engine/Shaders/Private/Strata/Strata.ush
+++ b/Engine/Shaders/Private/Strata/Strata.ush
@@ -108,6 +108,8 @@ struct FStrataBSDF
 #define VOLUME_ANISOTROPY(X)			X.VGPRs[0].w
 #define VOLUME_THICKNESS(X)				X.VGPRs[1].w
 
+void StrataRequestSharedNormalTangentBasis(inout uint NormalTypes, uint Index);
+
 bool StrataHasAnisotropy(in float Roughness0, float Roughness1)
 {
 	return Roughness0 != Roughness1;
@@ -171,12 +173,16 @@ FStrataData GetStrataDiffuseBSDF(float3 Albedo, float Roughness, uint SharedNorm
 	return StrataData;
 }
 
-FStrataData GetStrataDielectricBSDF(float2 Roughness, float IOR, float3 Tint, uint SharedNormalIndex)
+FStrataData GetStrataDielectricBSDF(float2 Roughness, float IOR, float3 Tint, uint SharedNormalIndex, inout uint SharedNormalTypes)
 {
 	FStrataData StrataData = GetInitialisedStrataData();
 	StrataData.LayerCount++;
 
 	const uint bHasAnisotropy = StrataHasAnisotropy(Roughness.x, Roughness.y) ? 1 : 0;
+	if (bHasAnisotropy)
+	{
+		StrataRequestSharedNormalTangentBasis(SharedNormalTypes, SharedNormalIndex);
+	}
 
 	StrataData.Layers[0].BSDFCount++;
 	BSDF_SETTYPE			(StrataData.Layers[0].BSDFs[0],		STRATA_BSDF_TYPE_DIELECTRIC);
@@ -193,12 +199,16 @@ FStrataData GetStrataDielectricBSDF(float2 Roughness, float IOR, float3 Tint, ui
 	return StrataData;
 }
 
-FStrataData GetStrataConductorBSDF(float3 Reflectivity, float3 EdgeColor, float2 Roughness, uint SharedNormalIndex)
+FStrataData GetStrataConductorBSDF(float3 Reflectivity, float3 EdgeColor, float2 Roughness, uint SharedNormalIndex, inout uint SharedNormalTypes)
 {
 	FStrataData StrataData = GetInitialisedStrataData();
 	StrataData.LayerCount++;
 
 	const uint bHasAnisotropy = StrataHasAnisotropy(Roughness.x, Roughness.y) ? 1 : 0;
+	if (bHasAnisotropy)
+	{
+		StrataRequestSharedNormalTangentBasis(SharedNormalTypes, SharedNormalIndex);
+	}
 
 	StrataData.Layers[0].BSDFCount++;
 	BSDF_SETTYPE			(StrataData.Layers[0].BSDFs[0],		STRATA_BSDF_TYPE_CONDUCTOR);
@@ -252,10 +262,13 @@ float3 StrataGetOffSpecularPeakReflectionDir(float3 Normal, float3 ReflectionVec
 float3 StrataUnpackNormal(uint PackedNormal);
 void StrataUnpackNormalAndTangent(inout float3 Normal, inout float3 Tangent, in uint InPacked);
 
+#define STRATA_BASIS_TYPE_NORMAL 0
+#define STRATA_BASIS_TYPE_TANGENT 1
 #if STRATA_SHARED_NORMALS_IN_REGISTERS
 struct FSharedNormals
 {
 	uint NormalCount;
+	uint NormalTypes;
 	float3 Normals[STRATA_MAX_SHARED_NORMAL_REGISTERS];
 	float3 Tangents[STRATA_MAX_SHARED_NORMAL_REGISTERS];
 };
@@ -267,6 +280,16 @@ FSharedNormals StrataInitialiseSharedNormals()
 }
 #endif
 
+void StrataRequestSharedNormalTangentBasis(inout uint NormalTypes, uint Index)
+{
+	NormalTypes = NormalTypes | (STRATA_BASIS_TYPE_TANGENT << Index);
+}
+
+uint StrataGetSharedNormalType(in uint NormalTypes, uint Index)
+{
+	return (NormalTypes >> Index) & 0x1 ? STRATA_BASIS_TYPE_TANGENT : STRATA_BASIS_TYPE_NORMAL;
+}
+
 struct FStrataPixelHeader
 {
 	uint BSDFCount;
@@ -277,6 +300,7 @@ struct FStrataPixelHeader
 #else
 	ByteAddressBuffer StrataBuffer;
 	uint SharedNormalByteOffset;
+	uint SharedNormalTypes;
 #endif
 };
 
@@ -288,6 +312,7 @@ FStrataPixelHeader InitialiseStrataPixelHeader()
 	Out.SharedNormals = (FSharedNormals)0;
 #else
 	Out.SharedNormalByteOffset = 0;
+	Out.SharedNormalTypes = 0;
 #endif
 	return Out;
 }
@@ -303,8 +328,14 @@ float3x3 StrataGetBSDFSharedBasis(in FStrataPixelHeader StrataPixelHeader, in FS
 #else
 	const bool bHasAnisotropy = BSDF_GETANISOTROPY(BSDF) != 0;
 	uint PackedBasis = StrataPixelHeader.StrataBuffer.Load(StrataPixelHeader.SharedNormalByteOffset + BSDF_GETNORMALID(BSDF) * STRATA_PACKED_NORMAL_STRIDE_BYTES);
+	uint BasisType   = StrataGetSharedNormalType(StrataPixelHeader.SharedNormalTypes, BSDF_GETNORMALID(BSDF));
 	float3x3 OutTangentBasis;
-	if (bHasAnisotropy)
+	if (BasisType == STRATA_BASIS_TYPE_NORMAL)
+	{
+		float3 Normal = StrataUnpackNormal(PackedBasis);
+		OutTangentBasis = GetTangentBasis(Normal);
+	}
+	else // if (BasisType == STRATA_BASIS_TYPE_TANGENT)
 	{
 		float3 Normal;
 		float3 Tangent;
@@ -314,11 +345,6 @@ float3x3 StrataGetBSDFSharedBasis(in FStrataPixelHeader StrataPixelHeader, in FS
 		OutTangentBasis[1] = cross(Normal, Tangent);
 		OutTangentBasis[2] = Normal;
 	}
-	else
-	{
-		float3 Normal = StrataUnpackNormal(PackedBasis);
-		OutTangentBasis = GetTangentBasis(Normal);
-	}
 	return OutTangentBasis;
 #endif
 }
@@ -927,17 +953,15 @@ uint StrataPackNormal(in float3 Normal)
 	{
 		result = (1.0 - abs(result.yx)) * SignNotZero(result.xy);
 	}
-	uint X10Bits = uint(clamp(result.x * 511.0f + 512.0f, 0.0f, 1023.0f));
-	uint Y10Bits = uint(clamp(result.y * 511.0f + 512.0f, 0.0f, 1023.0f));
-	return (X10Bits | (Y10Bits << 10));
+	uint PackedX = f32tof16(result.x);
+	uint PackedY = f32tof16(result.y);
+	return PackedX | (PackedY << 16);
 }
 
 float3 StrataUnpackNormal(uint PackedNormal)
 {
-	uint X10Bits = 0x3FF & PackedNormal;
-	uint Y10Bits = 0x3FF & (PackedNormal >> 10);
-	float x = float(int(X10Bits) - 512) / 511.0f;
-	float y = float(int(Y10Bits) - 512) / 511.0f;
+	float x = f16tof32(0xFFFFFFFF &  PackedNormal);
+	float y = f16tof32(0xFFFFFFFF & (PackedNormal>>16));
 
 	float3 Normal = float3(x, y, 1.0f - abs(x) - abs(y));
 	if (Normal.z < 0.0f)
@@ -948,7 +972,7 @@ float3 StrataUnpackNormal(uint PackedNormal)
 }
 
 #define STRATA_TANGENT_SYMMETRY 1
-uint StrataPackTangentBasis(in float3 Normal, in float3 Tangent)
+uint StrataPackNormalAndTangent(in float3 Normal, in float3 Tangent)
 {
 	uint X10Bits;
 	uint Y10Bits;
@@ -1118,7 +1142,7 @@ void StrataLoad_AlbedoGamma2ToLinearAlphaLinear(ByteAddressBuffer StrataBuffer,
 	A = Data.a;
 }
 
-
+uint PackStrataHeaderIn(uint BSDFCount, FStrataPixelHeader In);
 
 uint GetStrataPixelDataByteOffset(uint2 PixelPos, uint ViewBufferSizeX, uint StrataMaxBytesPerPixel)
 {
@@ -1293,18 +1317,20 @@ void PackStrataOut(
 	// Now write out Strata data
 
 	// 1. the header
-
-	uint SharedNormals_BSDFCount = BSDFCount;
-#if STRATA_SHARED_NORMALS_IN_REGISTERS
-	SharedNormals_BSDFCount |= (StrataPixelHeader.SharedNormals.NormalCount << 16);
-#endif
-	STRATA_STORE_UINT1(SharedNormals_BSDFCount);
-
+	STRATA_STORE_UINT1(PackStrataHeaderIn(BSDFCount, StrataPixelHeader));
 #if STRATA_SHARED_NORMALS_IN_REGISTERS
 	UNROLL
 	for (uint i = 0; i < StrataPixelHeader.SharedNormals.NormalCount; ++i)
 	{
-		STRATA_STORE_UINT1(StrataPackTangentBasis(StrataPixelHeader.SharedNormals.Normals[i], StrataPixelHeader.SharedNormals.Tangents[i]));
+		const uint BasisType = StrataGetSharedNormalType(StrataPixelHeader.SharedNormals.NormalTypes, i);
+		if (BasisType == STRATA_BASIS_TYPE_NORMAL)
+		{
+			STRATA_STORE_UINT1(StrataPackNormal(StrataPixelHeader.SharedNormals.Normals[i]));
+		}
+		else // if (BasisType == STRATA_BASIS_TYPE_TANGENT)
+		{
+			STRATA_STORE_UINT1(StrataPackNormalAndTangent(StrataPixelHeader.SharedNormals.Normals[i], StrataPixelHeader.SharedNormals.Tangents[i]));
+		}
 	}
 #endif
 
@@ -1377,22 +1403,34 @@ void PackStrataOut(
 	}
 }
 
+// Pack only strata header
+uint PackStrataHeaderIn(uint InBSDFCount, FStrataPixelHeader InHeader)
+{
+	uint Out = InBSDFCount;
+#if STRATA_SHARED_NORMALS_IN_REGISTERS
+	Out    |= (InHeader.SharedNormals.NormalCount << 8);
+	Out    |= (InHeader.SharedNormals.NormalTypes << 16);
+#endif
+	return Out;
+}
+
 // Unpack only strata header
 // If this is changed, please update the compiler side material size evaluation in StrataMaterial.cpp
 FStrataPixelHeader UnpackStrataHeaderIn(ByteAddressBuffer StrataBuffer, inout uint ByteOffset)
 {
 	FStrataPixelHeader Out = InitialiseStrataPixelHeader();
 
-	uint SharedNormals_BSDFCount = 0;
-	STRATA_LOAD_UINT1(SharedNormals_BSDFCount);
-	Out.BSDFCount = SharedNormals_BSDFCount & 0x0000FFFF;
-
-	uint NormalCount = SharedNormals_BSDFCount >> 16;
+	uint Packed = 0;
+	STRATA_LOAD_UINT1(Packed);
+	Out.BSDFCount    = 0xFF &  Packed;
+	uint NormalCount = 0xFF & (Packed >> 8);
+	uint NormalTypes = 0xFF & (Packed >> 16);
 
 #if STRATA_SHARED_NORMALS_IN_REGISTERS==0
 	// Only keep the offset to the shared normal memory. We are going to load them on demand.
 	Out.SharedNormalByteOffset = ByteOffset;
 	Out.StrataBuffer = StrataBuffer;
+	Out.SharedNormalTypes = NormalTypes;
 
 	// Now skip over the shared normal memory to be able to load the BSDF content
 	ByteOffset += NormalCount * STRATA_PACKED_NORMAL_STRIDE_BYTES;
diff --git a/Engine/Source/Runtime/Engine/Private/Materials/HLSLMaterialTranslator.cpp b/Engine/Source/Runtime/Engine/Private/Materials/HLSLMaterialTranslator.cpp
--- a/Engine/Source/Runtime/Engine/Private/Materials/HLSLMaterialTranslator.cpp
+++ b/Engine/Source/Runtime/Engine/Private/Materials/HLSLMaterialTranslator.cpp
@@ -7429,7 +7429,7 @@ int32 FHLSLMaterialTranslator::StrataDiffuseBSDF(int32 Albedo, int32 Roughness,
 int32 FHLSLMaterialTranslator::StrataDielectricBSDF(int32 RoughnessX, int32 RoughnessY, int32 IOR, int32 Tint, int32 Normal, int32 Tangent, uint8 SharedNormalIndex)
 {
 	return AddCodeChunk(
-		MCT_Strata, TEXT("GetStrataDielectricBSDF(float2(%s, %s), %s, %s, %u) /* Normal:%s Tangent:%s */"),
+		MCT_Strata, TEXT("GetStrataDielectricBSDF(float2(%s, %s), %s, %s, %u, Parameters.SharedNormals.NormalTypes) /* Normal:%s Tangent:%s */"),
 		*GetParameterCode(RoughnessX),
 		*GetParameterCode(RoughnessY),
 		*GetParameterCode(IOR),
@@ -7443,7 +7443,7 @@ int32 FHLSLMaterialTranslator::StrataDielectricBSDF(int32 RoughnessX, int32 Roug
 int32 FHLSLMaterialTranslator::StrataConductorBSDF(int32 Reflectivity, int32 EdgeColor, int32 RoughnessX, int32 RoughnessY, int32 Normal, int32 Tangent, uint8 SharedNormalIndex)
 {
 	return AddCodeChunk(
-		MCT_Strata, TEXT("GetStrataConductorBSDF(%s, %s, float2(%s, %s), %u) /* Normal:%s Tangent:%s */"),
+		MCT_Strata, TEXT("GetStrataConductorBSDF(%s, %s, float2(%s, %s), %u, Parameters.SharedNormals.NormalTypes) /* Normal:%s Tangent:%s */"),
 		*GetParameterCode(Reflectivity),
 		*GetParameterCode(EdgeColor),
 		*GetParameterCode(RoughnessX),
