Add sheen BSDF support.

diff --git a/Engine/Shaders/Private/Strata/Strata.ush b/Engine/Shaders/Private/Strata/Strata.ush
--- a/Engine/Shaders/Private/Strata/Strata.ush
+++ b/Engine/Shaders/Private/Strata/Strata.ush
@@ -108,6 +108,9 @@ struct FStrataBSDF
 #define VOLUME_ANISOTROPY(X)			X.VGPRs[0].w
 #define VOLUME_THICKNESS(X)				X.VGPRs[1].w
 
+#define SHEEN_ALBEDO(X)					X.VGPRs[0].xyz
+#define SHEEN_ROUGHNESS(X)				X.VGPRs[0].w
+
 void StrataRequestSharedNormalTangentBasis(inout uint NormalTypes, uint Index);
 
 bool StrataHasAnisotropy(in float Roughness0, float Roughness1)
@@ -247,6 +250,23 @@ FStrataData GetStrataVolumeBSDF(float3 Albedo, float3 Extinction, float Anisotro
 	return StrataData;
 }
 
+FStrataData GetStrataSheenBSDF(float3 Albedo, float Roughness, uint SharedNormalIndex)
+{
+	FStrataData StrataData = GetInitialisedStrataData();
+	StrataData.LayerCount++;
+
+	StrataData.Layers[0].BSDFCount++;
+	BSDF_SETTYPE			(StrataData.Layers[0].BSDFs[0],		STRATA_BSDF_TYPE_SHEEN);
+	BSDF_SETNORMALID		(StrataData.Layers[0].BSDFs[0],		SharedNormalIndex);
+	BSDF_SETISTOPLAYER		(StrataData.Layers[0].BSDFs[0],		0);
+	SHEEN_ALBEDO			(StrataData.Layers[0].BSDFs[0])		= Albedo;
+	SHEEN_ROUGHNESS			(StrataData.Layers[0].BSDFs[0])		= Roughness;
+
+	StrataData.Layers[0].BSDFs[0].Weight						= 1.0f;
+
+	return StrataData;
+}
+
 // STRATA_TODO put in a common file
 // Point lobe in off-specular peak direction
 float3 StrataGetOffSpecularPeakReflectionDir(float3 Normal, float3 ReflectionVector, float Roughness)
@@ -264,7 +284,6 @@ void StrataUnpackNormalAndTangent(inout float3 Normal, inout float3 Tangent, in
 
 #define STRATA_BASIS_TYPE_NORMAL 0
 #define STRATA_BASIS_TYPE_TANGENT 1
-#if STRATA_SHARED_NORMALS_IN_REGISTERS
 struct FSharedNormals
 {
 	uint NormalCount;
@@ -278,7 +297,6 @@ FSharedNormals StrataInitialiseSharedNormals()
 	FSharedNormals SharedNormals = (FSharedNormals)0;
 	return SharedNormals;
 }
-#endif
 
 void StrataRequestSharedNormalTangentBasis(inout uint NormalTypes, uint Index)
 {
@@ -403,6 +421,11 @@ void StrataGetBSDFRoughness(in FStrataBSDF BSDF, inout float R)
 		R = 0;
 		break;
 	}
+	case STRATA_BSDF_TYPE_SHEEN:
+	{
+		R = SHEEN_ROUGHNESS(BSDF);
+		break;
+	}
 	}
 }
 
@@ -545,6 +568,24 @@ FStrataEvaluateResult StrataEvaluateBSDF(FStrataBSDFContext BSDFContext)
 			Sample.Throughput	= SafePathSegmentTransmittance;
 			break;
 		}
+		case STRATA_BSDF_TYPE_SHEEN:
+		{
+			const float3 F0  = SHEEN_ALBEDO(BSDFContext.BSDF);
+			const float SafeRoughness = MakeRoughnessSafe(SHEEN_ROUGHNESS(BSDFContext.BSDF));
+
+			// Special override for roughness for supporting area light integrator with Sphere/Tube/Disk light, which modifies/increase roughness.
+			const float a2 = BSDFContext.a2Override > 0 ? BSDFContext.a2Override : Pow4(SafeRoughness);
+
+			// Generalized inverse microfacet specular
+			float D   = D_InvGGX(a2, BSDFContext.Context.NoH);
+			float Vis = Vis_Cloth(BSDFContext.Context.NoV, BSDFContext.SatNoL);
+			float3 F  = F_Schlick(F0, BSDFContext.Context.VoH);
+
+			Sample.Value = D * Vis * F;
+			Sample.PDF = BSDFContext.SatNoV / PI; // Per "Production Friendly Microfacet Sheen BRDF", hemispherical sampling give good result as the roughness is usually high.
+			Sample.Throughput = (1.0f - F);
+			break;
+		}
 	}
 
 	return Sample;
@@ -618,6 +659,16 @@ FStrataImportanceSampleResult StrataImportanceSampleBSDF(FStrataPixelHeader Stra
 			Sample.ConeAngle = 1.0f / PDF;
 			break;
 		}
+		case STRATA_BSDF_TYPE_SHEEN:
+		{
+			// Per "Production Friendly Microfacet Sheen BRDF", uniform hemispherical sampling gives good result as the roughness is 
+			// usually high and avoid issue at grazing angle
+			float4 ImportanceSample = CosineSampleHemisphere(E);
+			Sample.Direction = mul(ImportanceSample.xyz, BSDFContext.TangentBasis);
+			PDF = ImportanceSample.w;
+			Sample.ConeAngle = 1.0f / PDF;
+			break;
+		}
 	}
 
 	//StrataUpdateBSDFContext(BSDFContext, Sample.Direction);
@@ -704,6 +755,12 @@ FStrataEnvLightResult StrataEvaluateForEnvLight(FStrataBSDFContext BSDFContext,
 			StrataEnvLightResult.DiffuseNormal = BSDFContext.N;
 			break;
 		}
+		case STRATA_BSDF_TYPE_SHEEN:
+		{
+			StrataEnvLightResult.DiffuseAlbedo = SHEEN_ALBEDO(BSDFContext.BSDF);;
+			StrataEnvLightResult.DiffuseNormal = BSDFContext.N;
+			break;
+		}
 	}
 
 	return StrataEnvLightResult;
@@ -825,6 +882,12 @@ void SanitizeStrataVolume(inout FStrataBSDF BSDF)
 	VOLUME_THICKNESS(BSDF) = max(0.0f, VOLUME_THICKNESS(BSDF));
 }
 
+void SanitizeStrataSheen(inout FStrataBSDF BSDF)
+{
+	SHEEN_ALBEDO(BSDF) = saturate(SHEEN_ALBEDO(BSDF));
+	SHEEN_ROUGHNESS(BSDF) = saturate(SHEEN_ROUGHNESS(BSDF));
+}
+
 void StrataSanitizeBSDF(inout FStrataBSDF BSDF)
 {
 	switch (BSDF_GETTYPE(BSDF))
@@ -849,6 +912,11 @@ void StrataSanitizeBSDF(inout FStrataBSDF BSDF)
 			SanitizeStrataVolume(BSDF);
 			break;
 		}
+		case STRATA_BSDF_TYPE_SHEEN:
+		{
+			SanitizeStrataSheen(BSDF);
+			break;
+		}
 	}
 }
 
@@ -1282,6 +1350,26 @@ void PackStrataOut(
 				LayerThroughput += BSDFCoverage * SafePathSegmentTransmittance;
 				break;
 			}
+
+			case STRATA_BSDF_TYPE_SHEEN:
+			{
+				const float3 F0  = SHEEN_ALBEDO(BSDF);
+				const float3 F = F_Schlick(F0, ContextNoLight.VoH);
+
+				if (bTopLayer)
+				{
+					AccumRoughness		+= SHEEN_ROUGHNESS(BSDF);
+					AccumRoughnessCount += 1.0f;
+
+					const float SafeRoughness = MakeRoughnessSafe(SHEEN_ROUGHNESS(BSDF));
+					AccumBaseColor      += BSDFCoverage * EnvBRDFApprox(F0, SafeRoughness, saturate(ContextNoLight.NoV));
+				}
+
+				LayerCoverage	+= BSDFCoverage;
+				LayerThroughput += BSDFCoverage * (1.0f - F); // STRATA_TODO use 1-FGD instead of simple 1-F, event if we ignore ray scattering/spreading
+				break;
+			}
+
 			} // switch
 
 			// Update the weights to take into account throughput towards the view from layers above
@@ -1395,6 +1483,12 @@ void PackStrataOut(
 				// 8 bytes
 			}
 			break;
+			case STRATA_BSDF_TYPE_SHEEN:
+			{
+				STRATA_STORE_UINT1(PackAlbedoLinearToGamma2AlphaLinear(float4(SHEEN_ALBEDO(BSDF), SHEEN_ROUGHNESS(BSDF))));
+				// 4 bytes
+			}
+			break;
 			}
 
 			// Stop writing if above budget.
@@ -1492,6 +1586,11 @@ FStrataBSDF UnpackStrataBSDFIn(ByteAddressBuffer StrataBuffer, inout uint ByteOf
 		VOLUME_THICKNESS(OutBSDF) = 1.0f;
 	}
 	break;
+	case STRATA_BSDF_TYPE_SHEEN:
+	{
+		StrataLoad_AlbedoGamma2ToLinearAlphaLinear(StrataBuffer, ByteOffset, SHEEN_ALBEDO(OutBSDF), SHEEN_ROUGHNESS(OutBSDF));
+	}
+	break;
 	}
 
 	return OutBSDF;
diff --git a/Engine/Shaders/Private/Strata/StrataVisualize.usf b/Engine/Shaders/Private/Strata/StrataVisualize.usf
--- a/Engine/Shaders/Private/Strata/StrataVisualize.usf
+++ b/Engine/Shaders/Private/Strata/StrataVisualize.usf
@@ -231,6 +231,31 @@ void VisualizeMaterialPS(
 				HORIZONTALTAB; PRINTFLOAT1(VOLUME_THICKNESS(BSDF) *1000.0f); PRINTPARAMCHAR(_SPC_); PRINTPARAMCHAR(_M_); PRINTPARAMCHAR(_I_); PRINTPARAMCHAR(_L_); PRINTPARAMCHAR(_L_); PRINTPARAMCHAR(_I_);
 				NEXTLINE;
 
+				break;
+
+			case STRATA_BSDF_TYPE_SHEEN:
+				PRINTBSDFCHAR(_S_); PRINTBSDFCHAR(_H_); PRINTBSDFCHAR(_E_); PRINTBSDFCHAR(_E_); PRINTBSDFCHAR(_N_); 
+				PRINTBSDFCHAR(_SPC_); PRINTBSDFCHAR(_SPC_); PRINTBSDFCHAR(_SPC_); PRINTBSDFCHAR(_SPC_); PRINTBSDFCHAR(_SPC_);
+				if (BSDF_GETGREYWEIGHT(BSDF))
+				{
+					PRINTBSDFGREYWEIGHT();
+				}
+				else
+				{
+					PRINTBSDFWEIGHT();
+				}
+				PRINTBSDFNORMALID();
+				PRINTBSDFANISO();
+				PRINTBSDFTOPLAY();
+				NEXTLINE;
+
+				HORIZONTALTAB; PRINTPARAMCHAR(_A_); PRINTPARAMCHAR(_L_); PRINTPARAMCHAR(_B_); PRINTPARAMCHAR(_E_); PRINTPARAMCHAR(_D_); PRINTPARAMCHAR(_O_); 
+				HORIZONTALTAB; PRINTFLOAT3(SHEEN_ALBEDO(BSDF));
+				NEXTLINE;
+				HORIZONTALTAB; PRINTPARAMCHAR(_R_); PRINTPARAMCHAR(_O_); PRINTPARAMCHAR(_U_); PRINTPARAMCHAR(_G_); PRINTPARAMCHAR(_H_); PRINTPARAMCHAR(_SPC_);
+				HORIZONTALTAB; PRINTFLOAT1(SHEEN_ROUGHNESS(BSDF));
+				NEXTLINE;
+
 				break;
 			}
 
diff --git a/Engine/Shaders/Shared/StrataDefinitions.h b/Engine/Shaders/Shared/StrataDefinitions.h
--- a/Engine/Shaders/Shared/StrataDefinitions.h
+++ b/Engine/Shaders/Shared/StrataDefinitions.h
@@ -25,5 +25,6 @@
 #define STRATA_BSDF_TYPE_DIELECTRIC			1
 #define STRATA_BSDF_TYPE_CONDUCTOR			2
 #define STRATA_BSDF_TYPE_VOLUME				3
+#define STRATA_BSDF_TYPE_SHEEN				4
 
 
diff --git a/Engine/Source/Runtime/Engine/Classes/Materials/MaterialExpressionStrata.h b/Engine/Source/Runtime/Engine/Classes/Materials/MaterialExpressionStrata.h
--- a/Engine/Source/Runtime/Engine/Classes/Materials/MaterialExpressionStrata.h
+++ b/Engine/Source/Runtime/Engine/Classes/Materials/MaterialExpressionStrata.h
@@ -209,7 +209,40 @@ class UMaterialExpressionStrataVolumeBSDF : public UMaterialExpression
 	//~ End UMaterialExpression Interface
 };
 
-// STRATA_TODO Sheen, Subsurface, thinfilm, generalised schlick
+UCLASS(MinimalAPI, collapsecategories, hidecategories = Object)
+class UMaterialExpressionStrataSheenBSDF : public UMaterialExpression
+{
+	GENERATED_UCLASS_BODY()
+
+	/**
+	* Albedo (type = float3, unit = unitless)
+	*/
+	UPROPERTY()
+	FExpressionInput Albedo;
+
+	/**
+	 * Roughness (type = float, unit = unitless)
+	 */
+	UPROPERTY()
+	FExpressionInput Roughness;
+
+	/**
+	 * Normal (type = float3, unit = unitless)
+	 */
+	UPROPERTY()
+	FExpressionInput Normal;
+
+	//~ Begin UMaterialExpression Interface
+#if WITH_EDITOR
+	virtual int32 Compile(class FMaterialCompiler* Compiler, int32 OutputIndex) override;
+	virtual void GetCaption(TArray<FString>& OutCaptions) const override;
+	virtual uint32 GetOutputType(int32 OutputIndex) override;
+	virtual uint32 GetInputType(int32 InputIndex) override;
+#endif
+	//~ End UMaterialExpression Interface
+};
+
+// Subsurface, thinfilm, generalised schlick
 
 
 
diff --git a/Engine/Source/Runtime/Engine/Private/Materials/HLSLMaterialTranslator.cpp b/Engine/Source/Runtime/Engine/Private/Materials/HLSLMaterialTranslator.cpp
--- a/Engine/Source/Runtime/Engine/Private/Materials/HLSLMaterialTranslator.cpp
+++ b/Engine/Source/Runtime/Engine/Private/Materials/HLSLMaterialTranslator.cpp
@@ -7467,6 +7467,17 @@ int32 FHLSLMaterialTranslator::StrataVolumeBSDF(int32 Albedo, int32 Extinction,
 	);
 }
 
+int32 FHLSLMaterialTranslator::StrataSheenBSDF(int32 Albedo, int32 Roughness, int32 Normal, uint8 SharedNormalIndex)
+{
+	return AddCodeChunk(
+		MCT_Strata, TEXT("GetStrataSheenBSDF(%s, %s, %u) /* %s */"),
+		*GetParameterCode(Albedo),
+		*GetParameterCode(Roughness),
+		SharedNormalIndex,
+		*GetParameterCode(Normal)
+	);
+}
+
 int32 FHLSLMaterialTranslator::StrataHorizontalMixing(int32 Foreground, int32 Background, int32 Mix)
 {
 	if (Foreground == INDEX_NONE || Background == INDEX_NONE || Mix == INDEX_NONE)
diff --git a/Engine/Source/Runtime/Engine/Private/Materials/HLSLMaterialTranslator.h b/Engine/Source/Runtime/Engine/Private/Materials/HLSLMaterialTranslator.h
--- a/Engine/Source/Runtime/Engine/Private/Materials/HLSLMaterialTranslator.h
+++ b/Engine/Source/Runtime/Engine/Private/Materials/HLSLMaterialTranslator.h
@@ -795,6 +795,7 @@ protected:
 	virtual int32 StrataDielectricBSDF(int32 RoughnessX, int32 RoughnessY, int32 IOR, int32 Tint, int32 Normal, int32 Tangent, uint8 SharedNormalIndex) override;
 	virtual int32 StrataConductorBSDF(int32 Reflectivity, int32 EdgeColor, int32 RoughnessX, int32 RoughnessY, int32 Normal, int32 Tangent, uint8 SharedNormalIndex) override;
 	virtual int32 StrataVolumeBSDF(int32 Albedo, int32 Extinction, int32 Anisotropy, int32 Thickness, int32 Normal, uint8 SharedNormalIndex) override;
+	virtual int32 StrataSheenBSDF(int32 Albedo, int32 Roughness, int32 Normal, uint8 SharedNormalIndex) override;
 	virtual int32 StrataHorizontalMixing(int32 Foreground, int32 Background, int32 Mix) override;
 	virtual int32 StrataVerticalLayering(int32 Top, int32 Base) override;
 	virtual int32 StrataAdd(int32 A, int32 B) override;
diff --git a/Engine/Source/Runtime/Engine/Private/Materials/MaterialExpressions.cpp b/Engine/Source/Runtime/Engine/Private/Materials/MaterialExpressions.cpp
--- a/Engine/Source/Runtime/Engine/Private/Materials/MaterialExpressions.cpp
+++ b/Engine/Source/Runtime/Engine/Private/Materials/MaterialExpressions.cpp
@@ -19313,6 +19313,74 @@ uint32 UMaterialExpressionStrataVolumeBSDF::GetInputType(int32 InputIndex)
 
 
 
+UMaterialExpressionStrataSheenBSDF::UMaterialExpressionStrataSheenBSDF(const FObjectInitializer& ObjectInitializer)
+	: Super(ObjectInitializer)
+{
+	struct FConstructorStatics
+	{
+		FText NAME_Strata;
+		FConstructorStatics() : NAME_Strata(LOCTEXT("Strata", "Strata")) { }
+	};
+	static FConstructorStatics ConstructorStatics;
+#if WITH_EDITORONLY_DATA
+	MenuCategories.Add(ConstructorStatics.NAME_Strata);
+#endif
+}
+
+#if WITH_EDITOR
+int32 UMaterialExpressionStrataSheenBSDF::Compile(class FMaterialCompiler* Compiler, int32 OutputIndex)
+{
+	int32 NormalCodeChunk = Normal.GetTracedInput().Expression ? Normal.Compile(Compiler) : Compiler->VertexNormal();
+
+	// For some reason, when editing shader, Normal.Compile() / Tangent.Compile() might return -1. In this case we fallback to vertex normal/tangent
+	if (NormalCodeChunk == -1) { NormalCodeChunk = Compiler->VertexNormal(); }
+
+	uint8 SharedNormalIndex = StrataCompilationInfoCreateSharedNormal(Compiler, NormalCodeChunk);
+
+	int32 RoughnessCodeChunk = Roughness.GetTracedInput().Expression ? Roughness.Compile(Compiler) : Compiler->Constant(0.0f);
+
+	int32 OutputCodeChunk = Compiler->StrataSheenBSDF(
+		Albedo.GetTracedInput().Expression ? Albedo.Compile(Compiler) : Compiler->Constant3(0.0f, 0.0f, 0.0f),
+		RoughnessCodeChunk,
+		NormalCodeChunk,
+		SharedNormalIndex);
+	StrataCompilationInfoCreateSingleBSDFMaterial(Compiler, OutputCodeChunk, SharedNormalIndex, STRATA_BSDF_TYPE_SHEEN);
+
+	return OutputCodeChunk;
+}
+
+void UMaterialExpressionStrataSheenBSDF::GetCaption(TArray<FString>& OutCaptions) const
+{
+	OutCaptions.Add(TEXT("Strata Sheen BSDF"));
+}
+
+uint32 UMaterialExpressionStrataSheenBSDF::GetOutputType(int32 OutputIndex)
+{
+	return MCT_Strata;
+}
+
+uint32 UMaterialExpressionStrataSheenBSDF::GetInputType(int32 InputIndex)
+{
+	switch (InputIndex)
+	{
+	case 0:
+		return MCT_Float3;
+		break;
+	case 1:
+		return MCT_Float;
+		break;	
+	case 2:
+		return MCT_Float3;
+		break;
+	}
+
+	check(false);
+	return MCT_Float1;
+}
+#endif // WITH_EDITOR
+
+
+
 UMaterialExpressionStrataHorizontalMixing::UMaterialExpressionStrataHorizontalMixing(const FObjectInitializer& ObjectInitializer)
 	: Super(ObjectInitializer)
 {
diff --git a/Engine/Source/Runtime/Engine/Private/Materials/StrataMaterial.cpp b/Engine/Source/Runtime/Engine/Private/Materials/StrataMaterial.cpp
--- a/Engine/Source/Runtime/Engine/Private/Materials/StrataMaterial.cpp
+++ b/Engine/Source/Runtime/Engine/Private/Materials/StrataMaterial.cpp
@@ -21,6 +21,9 @@ FString GetStrataBSDFName(uint8 BSDFType)
 	case STRATA_BSDF_TYPE_VOLUME:
 		return TEXT("VOLUME");
 		break;
+	case STRATA_BSDF_TYPE_SHEEN:
+		return TEXT("SHEEN");
+		break;
 	}
 	check(false);
 	return "";
@@ -198,6 +201,11 @@ FStrataMaterialAnalysisResult StrataCompilationInfoMaterialAnalysis(FMaterialCom
 				}
 				break;
 			}
+			case STRATA_BSDF_TYPE_SHEEN:
+			{
+				Result.RequestedByteCount += UintByteSize;
+				break;
+			}
 			}
 		}
 
diff --git a/Engine/Source/Runtime/Engine/Public/MaterialCompiler.h b/Engine/Source/Runtime/Engine/Public/MaterialCompiler.h
--- a/Engine/Source/Runtime/Engine/Public/MaterialCompiler.h
+++ b/Engine/Source/Runtime/Engine/Public/MaterialCompiler.h
@@ -395,6 +395,7 @@ public:
 	virtual int32 StrataDielectricBSDF(int32 RoughnessX, int32 RoughnessY, int32 IOR, int32 Tint, int32 Normal, int32 Tangent, uint8 SharedNormalIndex) = 0;
 	virtual int32 StrataConductorBSDF(int32 Reflectivity, int32 EdgeColor, int32 RoughnessX, int32 RoughnessY, int32 Normal, int32 Tangent, uint8 SharedNormalIndex) = 0;
 	virtual int32 StrataVolumeBSDF(int32 Albedo, int32 Extinction, int32 Anisotropy, int32 Thickness, int32 Normal, uint8 SharedNormalIndex) = 0;
+	virtual int32 StrataSheenBSDF(int32 Albedo, int32 Roughness, int32 Normal, uint8 SharedNormalIndex) = 0;
 	virtual int32 StrataHorizontalMixing(int32 Foreground, int32 Background, int32 Mix) = 0;
 	virtual int32 StrataVerticalLayering(int32 Top, int32 Base) = 0;
 	virtual int32 StrataAdd(int32 A, int32 B) = 0;
@@ -824,6 +825,11 @@ public:
 		return Compiler->StrataVolumeBSDF(Albedo, Extinction, Anisotropy, Thickness, Normal, SharedNormalIndex);
 	}
 
+	virtual int32 StrataSheenBSDF(int32 Albedo, int32 Roughness, int32 Normal, uint8 SharedNormalIndex) override
+	{
+		return Compiler->StrataSheenBSDF(Albedo, Roughness, Normal, SharedNormalIndex);
+	}
+
 	virtual int32 StrataHorizontalMixing(int32 Foreground, int32 Background, int32 Mix) override
 	{
 		return Compiler->StrataHorizontalMixing(Foreground, Background, Mix);
