Add BSDFs plots to Strata vizualize.

diff --git a/Engine/Shaders/Private/Strata/Strata.ush b/Engine/Shaders/Private/Strata/Strata.ush
--- a/Engine/Shaders/Private/Strata/Strata.ush
+++ b/Engine/Shaders/Private/Strata/Strata.ush
@@ -429,7 +429,7 @@ void StrataGetBSDFRoughness(in FStrataBSDF BSDF, inout float R)
 	}
 }
 
-FStrataBSDFContext StrataCreateBSDFContext(FStrataPixelHeader StrataPixelHeader, FStrataBSDF BSDF, float3 V, float3 L)
+FStrataBSDFContext StrataCreateBSDFContext(FStrataPixelHeader StrataPixelHeader, FStrataBSDF BSDF, float3 V, float3 L, bool bOverrideWithLocalFrame=false)
 {
 	FStrataBSDFContext BSDFContext = (FStrataBSDFContext)0;
 
@@ -442,6 +442,12 @@ FStrataBSDFContext StrataCreateBSDFContext(FStrataPixelHeader StrataPixelHeader,
 	BSDFContext.X = TangentBasis[0]; 
 	BSDFContext.Y = TangentBasis[1];
 	BSDFContext.N = TangentBasis[2]; 
+	if (bOverrideWithLocalFrame)
+	{
+		BSDFContext.X = float3(1,0,0);
+		BSDFContext.Y = float3(0,1,0);
+		BSDFContext.N = float3(0,0,1);
+	}
 
 	BSDFContext.R = 2 * dot(BSDFContext.V, BSDFContext.N) * BSDFContext.N - BSDFContext.V;
 	BSDFContext.H = normalize(BSDFContext.V + BSDFContext.L);
diff --git a/Engine/Shaders/Private/Strata/StrataVisualize.usf b/Engine/Shaders/Private/Strata/StrataVisualize.usf
--- a/Engine/Shaders/Private/Strata/StrataVisualize.usf
+++ b/Engine/Shaders/Private/Strata/StrataVisualize.usf
@@ -7,13 +7,65 @@
 #include "/Engine/Private/Strata/Strata.ush"
 
 #include "/Engine/Private/MiniFontCommon.ush"
-
+#include "/Engine/Private/ShaderDrawDebug.ush"
 
 // On some platforms, UnpackStrataIn causes the following error@
 // (error, code:7203) - not enough registers available for the entire program. In order to permit spilling to memory please either disable this error via '#pragma warning (disable:7203)' or reduce it to a warning via '#pragma warning (default:7203)'
 // To fix this and increase occupancy, we will have to progressively unpack the BSDFs on the fly (and thus only have a single BSDF allocating registers)
 #pragma warning (disable:7203)
 
+void DrawReferential(float3 P, float3 X, float3 Y, float3 N, float3 InColor)
+{
+	const float Size = 10.f;
+	const float SizeWithTip = 12.0f;
+	const float4 Color = float4(InColor, 1);
+	const float4 ColorX = float4(1, 0, 0, 1);
+	const float4 ColorY = float4(0, 1, 0, 1);
+	const float4 ColorZ = float4(0, 0, 1, 1);
+
+	// Core
+	AddLine(P, P + X * Size, Color, Color);
+	AddLine(P, P + Y * Size, Color, Color);
+	AddLine(P, P + N * Size, Color, Color);
+	// Tips
+	AddLine(P + X * Size, P + X * SizeWithTip, ColorX, ColorX);
+	AddLine(P + Y * Size, P + Y * SizeWithTip, ColorY, ColorY);
+	AddLine(P + N * Size, P + N * SizeWithTip, ColorZ, ColorZ);
+}
+
+float3 DrawPlot(
+	const float3 V, 
+	const FStrataPixelHeader Header,
+	const FStrataBSDF BSDF, 
+	uint2 PixelCoord, 
+	uint2 RectOffset, 
+	uint2 RectSize)
+{
+	float3 Out = 0;
+
+	const uint2 LocalCoord = PixelCoord - RectOffset;
+	const float2 UV			= float2(PixelCoord - RectOffset) / float2(RectSize);
+	const float2 ProjCoord = UV * 2 - 1;
+	if (all(UV >= 0) && all(UV <= 1) && (Pow2(ProjCoord.x) + Pow2(ProjCoord.y) <1))
+	{
+		const float SinTheta = sqrt(Pow2(ProjCoord.x) + Pow2(ProjCoord.y));
+		const float CosTheta = sqrt(1 - Pow2(SinTheta));
+		const float3 L = float3(ProjCoord, CosTheta);
+
+		const float3x3 Basis = StrataGetBSDFSharedBasis(Header, BSDF);
+		float3 LocalV = mul(Basis, V);
+		if (LocalV.z < 0)
+		{
+			LocalV.z = 0;
+			LocalV = normalize(LocalV);
+		}
+
+		FStrataBSDFContext BSDFContext = StrataCreateBSDFContext(Header, BSDF, LocalV, L, true);
+		FStrataEvaluateResult BSDFEvaluate = StrataEvaluateBSDF(BSDFContext);
+		Out = BSDFEvaluate.Value * BSDFContext.SatNoL; // * BSDF.Weight
+	}
+	return Out;
+}
 
 void VisualizeMaterialPS(
 	float4 SVPos : SV_POSITION,
@@ -27,6 +79,10 @@ void VisualizeMaterialPS(
 	float2 BufferUV = SvPositionToBufferUV(float4(MaterialPos, SVPos.zw));
 	FGBufferData GBuffer = GetGBufferDataFromSceneTextures(BufferUV);
 	bool bStrataMaterial = GBuffer.ShadingModelID == SHADINGMODELID_STRATA;
+	const float2 ScreenPosition = SvPositionToScreenPosition(SVPos).xy;
+	const float2 ScreenMaterialPosition = SvPositionToScreenPosition(float4(MaterialPos, 0.5f, 1)).xy;
+	const float3 WorldPosition = mul(float4(ScreenMaterialPosition * GBuffer.Depth, GBuffer.Depth, 1), View.ScreenToWorld).xyz;
+	const float3 V = -normalize(WorldPosition - View.WorldCameraOrigin);
 
 #if 0
 	OutColor = 0;
@@ -62,13 +118,50 @@ void VisualizeMaterialPS(
 
 	OutColor = float4(0.0f, 0.0f, 0.0f, 1.0f);
 
-	if (all(float2(PixelPos) > float2(MaterialPos + InspectSquare)))
+	// Plot BSDFs
+	bool bPlotValid = false;
+	if (bStrataMaterial)
 	{
-		OutColor.w = 0.0f;				// Set background visiblity to 0 in the bottom right part of the screen to better reads the debug values
+		uint PixelStrataDataByteOffset = GetStrataPixelDataByteOffset(MaterialPos, View.BufferSizeAndInvSize.x, Strata.MaxBytesPerPixel);
+		const uint PixelStrataDataByteOffsetStart = PixelStrataDataByteOffset;
+
+		FStrataPixelHeader Header = UnpackStrataHeaderIn(Strata.MaterialLobesBuffer, PixelStrataDataByteOffset);
+		const uint PixelStrataDataByteOffsetPostHeader = PixelStrataDataByteOffset;
+
+		for (uint BSDFIndex = 0; BSDFIndex < Header.BSDFCount; ++BSDFIndex)
+		{
+			const FStrataBSDF BSDF = UnpackStrataBSDFIn(Strata.MaterialLobesBuffer, PixelStrataDataByteOffset);
+
+			// Referential
+			if (all(PixelPos == MaterialPos))
+			{
+				FStrataBSDFContext BSDFContext = StrataCreateBSDFContext(Header, BSDF, V, float3(0,0,1));
+				DrawReferential(WorldPosition, BSDFContext.X, BSDFContext.Y, BSDFContext.N, float3(1,1,0));
+			}
+
+			// Plot
+			{
+				const uint2 Resolution = 180;
+				const uint OffsetX = 50 + BSDFIndex * 200;
+				const uint OffsetY = 50;
+				OutColor.xyz += DrawPlot(V, Header, BSDF, PixelPos, uint2(OffsetX, OffsetY), Resolution);
+			}
+		}
+
+		// Background
+		if (dot(OutColor.xyz, float3(1, 1, 1)) > 0)
+		{
+			OutColor.w = 0;
+			bPlotValid = true;
+		}
 	}
 
+	// Print BSDFs attributes
 	if (bStrataMaterial)
 	{
+		uint2 RectMin = uint2(MaterialPos + InspectSquare);
+		uint2 RectMax = RectMin;
+
 		uint PixelStrataDataByteOffset = GetStrataPixelDataByteOffset(MaterialPos, View.BufferSizeAndInvSize.x, Strata.MaxBytesPerPixel);
 		const uint PixelStrataDataByteOffsetStart = PixelStrataDataByteOffset;
 
@@ -80,7 +173,7 @@ void VisualizeMaterialPS(
 		float3 INFOFontColor = float3(1.0, 0.5, 0.25);
 		float3 PARAMFontColor = float3(0.8, 0.8, 0.8);
 		float3 FontColor = float3(0.25, 1.0, 1.0);
-		int2 TopLeftInit = MaterialPos + InspectSquare + int2(3,3);
+		int2 TopLeftInit = MaterialPos + InspectSquare + int2(3, 3);
 		int2 TopLeft = TopLeftInit;
 
 #define PRINTPARAMCHAR(X) PrintCharacter(PixelPos, OutColor.xyz, PARAMFontColor, TopLeft, X);
@@ -100,7 +193,7 @@ void VisualizeMaterialPS(
 #define PRINTFLOATNOFRAC(Value, DigitCount) PrintFloatNoFraction(PixelPos, OutColor.xyz, PARAMFontColor, TopLeft, Value, DigitCount); TopLeft.x+=8*7;
 #define SIGNOFFLOAT1(Value) if(Value < 0.0f) { PrintCharacter(PixelPos, OutColor.xyz, PARAMFontColor, TopLeft, _MINUS_); } /*else { PrintCharacter(PixelPos, OutColor.xyz, PARAMFontColor, TopLeft, _SPC_);}*/
 
-#define NEXTLINE  TopLeft.y += 8; TopLeft.x = TopLeftInit.x;
+#define NEXTLINE  TopLeft.y += 8; RectMax.x = max(RectMax.x, TopLeft.x); RectMax.y = max(RectMax.y, TopLeft.y); TopLeft.x = TopLeftInit.x; 
 #define HORIZONTALSEPARATION  TopLeft.y += 4;
 #define HORIZONTALTAB  TopLeft.x += 8;
 
@@ -283,6 +376,12 @@ void VisualizeMaterialPS(
 		PRINTBSDFCHAR(_B_); PRINTBSDFCHAR(_Y_); PRINTBSDFCHAR(_T_); PRINTBSDFCHAR(_E_); PRINTBSDFCHAR(_S_); PRINTBSDFCHAR(_SPC_);
 		PRINTFLOATNOFRAC(float(PixelStrataDataByteOffsetPostBSDF - PixelStrataDataByteOffsetStart), 3);
 		NEXTLINE;
+
+		// Background
+		if (all(PixelPos > RectMin) && all(PixelPos < RectMax) && !bPlotValid)
+		{
+			OutColor.w = 0.0f;				// Set background visiblity to 0 in the bottom right part of the screen to better reads the debug values
+		}
 	}
 }
 
