Fix bit offset issue on Strata.ush

diff --git a/Engine/Shaders/Private/Strata/Strata.ush b/Engine/Shaders/Private/Strata/Strata.ush
--- a/Engine/Shaders/Private/Strata/Strata.ush
+++ b/Engine/Shaders/Private/Strata/Strata.ush
@@ -51,8 +51,8 @@ struct FStrataBSDF
 #define STATE_BIT_OFFSET_BSDF			0
 #define STATE_BIT_OFFSET_SHAREDNORMALID	(STATE_BIT_OFFSET_BSDF				+ STATE_BITS_BSDF)
 #define STATE_BIT_OFFSET_ANISOTROPY		(STATE_BIT_OFFSET_SHAREDNORMALID	+ STATE_BITS_SHAREDNORMALID)
-#define STATE_BIT_OFFSET_TOPLAYER		(STATE_BIT_OFFSET_ANISOTROPY		+ STATE_BIT_OFFSET_ANISOTROPY)
-#define STATE_BIT_OFFSET_HASSCATTERING	(STATE_BIT_OFFSET_TOPLAYER			+ STATE_BIT_OFFSET_TOPLAYER)
+#define STATE_BIT_OFFSET_TOPLAYER		(STATE_BIT_OFFSET_ANISOTROPY		+ STATE_BITS_ANISOTROPY)
+#define STATE_BIT_OFFSET_HASSCATTERING	(STATE_BIT_OFFSET_TOPLAYER			+ STATE_BITS_TOPLAYER)
 #define STATE_BIT_OFFSET_GREYWEIGHT		(STATE_BIT_OFFSET_HASSCATTERING		+ STATE_BITS_HASSCATTERING)
 #define STATE_BIT_OFFSET_WEIGHT10F		22
 
