Strata define renaming as it has now broader use cases.

diff --git a/Engine/Shaders/Private/DeferredLightPixelShaders.usf b/Engine/Shaders/Private/DeferredLightPixelShaders.usf
--- a/Engine/Shaders/Private/DeferredLightPixelShaders.usf
+++ b/Engine/Shaders/Private/DeferredLightPixelShaders.usf
@@ -32,7 +32,7 @@ float CloudShadowmapStrength;
 #endif
 
 #if STRATA_ENABLED
-#define STRATA_SHARED_NORMALS_IN_REGISTERS 0
+#define STRATA_INLINE_SHADING 0
 #include "/Engine/Private/Strata/Strata.ush"
 #include "/Engine/Private/ParticipatingMediaCommon.ush"
 #include "/Engine/Private/Strata/StrataDeferredLighting.ush"
diff --git a/Engine/Shaders/Private/ReflectionEnvironmentPixelShader.usf b/Engine/Shaders/Private/ReflectionEnvironmentPixelShader.usf
--- a/Engine/Shaders/Private/ReflectionEnvironmentPixelShader.usf
+++ b/Engine/Shaders/Private/ReflectionEnvironmentPixelShader.usf
@@ -105,7 +105,7 @@ FSkyLightVisibilityData GetSkyLightVisibilityData(FGBufferData GBuffer, float Am
 
 
 #if STRATA_ENABLED
-#define STRATA_SHARED_NORMALS_IN_REGISTERS 0
+#define STRATA_INLINE_SHADING 0
 #include "/Engine/Private/Strata/Strata.ush"
 #include "/Engine/Private/ParticipatingMediaCommon.ush"
 #include "/Engine/Private/Strata/StrataEnvironmentLighting.ush"
diff --git a/Engine/Shaders/Private/Strata/Strata.ush b/Engine/Shaders/Private/Strata/Strata.ush
--- a/Engine/Shaders/Private/Strata/Strata.ush
+++ b/Engine/Shaders/Private/Strata/Strata.ush
@@ -15,8 +15,8 @@
 
 // During the base pass or forward rendering, shared normals are simply available from registers.
 // If a BSDF is loaded for processing from the strata buffer, normals are unpacked on demand to reduce VGPR pressure and increase occupancy.
-#ifndef STRATA_SHARED_NORMALS_IN_REGISTERS
-#define STRATA_SHARED_NORMALS_IN_REGISTERS 1
+#ifndef STRATA_INLINE_SHADING
+#define STRATA_INLINE_SHADING 1
 #endif
 
 #ifndef STRATA_CLAMPED_LAYER_COUNT
@@ -32,7 +32,9 @@ struct FStrataBSDF
 	float3				Weight;		// Weight is a float3 because, when loaded from the material buffer, it contains the BSDF throughput (coverage and transmittance combined)
 	uint				State;		// (27 unused bits) (1bit as HasAnisotropy) (2bits as SharedNormalIndex) (3bits as BSDFType)
 
+#if STRATA_INLINE_SHADING
 	float3				Emissive;	// All bsdf can potientially have emissive. This is never written out to the strata gbuffer but instead accumulated in the emissive buffer during the base pass.
+#endif
 
 	float4 VGPRs[2];
 };
@@ -114,7 +116,13 @@ struct FStrataBSDF
 #define VOLUMETRICFOGCLOUD_ALBEDO(X)	X.VGPRs[0].xyz
 #define VOLUMETRICFOGCLOUD_EXTINCTION(X)X.VGPRs[1].xyz
 #define VOLUMETRICFOGCLOUD_AO(X)		X.VGPRs[0].w
-#define VOLUMETRICFOGCLOUD_EMISSIVE(X)	X.Emissive
+#if STRATA_INLINE_SHADING
+#define VOLUMETRICFOGCLOUD_SETEMISSIVE(X, V)	X.Emissive = V
+#define VOLUMETRICFOGCLOUD_GETEMISSIVE(X)		X.Emissive
+#else
+#define VOLUMETRICFOGCLOUD_SETEMISSIVE(X, V)	X
+#define VOLUMETRICFOGCLOUD_GETEMISSIVE(X)		0.0f
+#endif
 
 void StrataRequestSharedNormalTangentBasis(inout uint NormalTypes, uint Index);
 
@@ -281,8 +289,8 @@ FStrataData GetStrataVolumeFogCloudBSDF(float3 Albedo, float3 Extinction, float3
 	BSDF_SETTYPE					(StrataData.Layers[0].BSDFs[0], STRATA_BSDF_TYPE_VOLUMETRICFOGCLOUD);
 	VOLUMETRICFOGCLOUD_ALBEDO		(StrataData.Layers[0].BSDFs[0]) = Albedo;
 	VOLUMETRICFOGCLOUD_EXTINCTION	(StrataData.Layers[0].BSDFs[0]) = Extinction;
-	VOLUMETRICFOGCLOUD_EMISSIVE		(StrataData.Layers[0].BSDFs[0]) = Emissive;
 	VOLUMETRICFOGCLOUD_AO			(StrataData.Layers[0].BSDFs[0]) = AmbientOcclusion;
+	VOLUMETRICFOGCLOUD_SETEMISSIVE	(StrataData.Layers[0].BSDFs[0], Emissive);
 
 	StrataData.Layers[0].BSDFs[0].Weight = 1.0f;
 
@@ -334,7 +342,7 @@ struct FStrataPixelHeader
 {
 	uint BSDFCount;
 
-#if STRATA_SHARED_NORMALS_IN_REGISTERS
+#if STRATA_INLINE_SHADING
 	// All the shared normal in VGPRS. This is only used in basepass and during forward rendering.
 	FSharedNormals SharedNormals;
 #else
@@ -348,7 +356,7 @@ FStrataPixelHeader InitialiseStrataPixelHeader()
 {
 	FStrataPixelHeader Out;
 	Out.BSDFCount = 0;
-#if STRATA_SHARED_NORMALS_IN_REGISTERS
+#if STRATA_INLINE_SHADING
 	Out.SharedNormals = (FSharedNormals)0;
 #else
 	Out.SharedNormalByteOffset = 0;
@@ -359,7 +367,7 @@ FStrataPixelHeader InitialiseStrataPixelHeader()
 
 float3x3 StrataGetBSDFSharedBasis(in FStrataPixelHeader StrataPixelHeader, in FStrataBSDF BSDF)
 {
-#if STRATA_SHARED_NORMALS_IN_REGISTERS
+#if STRATA_INLINE_SHADING
 	// Assume Normal & Tangent are already been normalized
 	float3 Normal  = StrataPixelHeader.SharedNormals.Normals[BSDF_GETNORMALID(BSDF)];
 	float3 Tangent = StrataPixelHeader.SharedNormals.Tangents[BSDF_GETNORMALID(BSDF)];
@@ -1434,7 +1442,7 @@ void PackStrataOut(
 
 	// 1. the header
 	STRATA_STORE_UINT1(PackStrataHeaderIn(BSDFCount, StrataPixelHeader));
-#if STRATA_SHARED_NORMALS_IN_REGISTERS
+#if STRATA_INLINE_SHADING
 	UNROLL
 	for (uint i = 0; i < StrataPixelHeader.SharedNormals.NormalCount; ++i)
 	{
@@ -1529,7 +1537,7 @@ void PackStrataOut(
 uint PackStrataHeaderIn(uint InBSDFCount, FStrataPixelHeader InHeader)
 {
 	uint Out = InBSDFCount;
-#if STRATA_SHARED_NORMALS_IN_REGISTERS
+#if STRATA_INLINE_SHADING
 	Out    |= (InHeader.SharedNormals.NormalCount << 8);
 	Out    |= (InHeader.SharedNormals.NormalTypes << 16);
 #endif
@@ -1548,7 +1556,7 @@ FStrataPixelHeader UnpackStrataHeaderIn(ByteAddressBuffer StrataBuffer, inout ui
 	uint NormalCount = 0xFF & (Packed >> 8);
 	uint NormalTypes = 0xFF & (Packed >> 16);
 
-#if STRATA_SHARED_NORMALS_IN_REGISTERS==0
+#if STRATA_INLINE_SHADING==0
 	// Only keep the offset to the shared normal memory. We are going to load them on demand.
 	Out.SharedNormalByteOffset = ByteOffset;
 	Out.StrataBuffer = StrataBuffer;
diff --git a/Engine/Shaders/Private/Strata/StrataVisualize.usf b/Engine/Shaders/Private/Strata/StrataVisualize.usf
--- a/Engine/Shaders/Private/Strata/StrataVisualize.usf
+++ b/Engine/Shaders/Private/Strata/StrataVisualize.usf
@@ -3,7 +3,7 @@
 #include "/Engine/Private/Common.ush"
 #include "/Engine/Private/DeferredShadingCommon.ush"
 
-#define STRATA_SHARED_NORMALS_IN_REGISTERS 0
+#define STRATA_INLINE_SHADING 0
 #include "/Engine/Private/Strata/Strata.ush"
 
 #include "/Engine/Private/MiniFontCommon.ush"
diff --git a/Engine/Shaders/Private/VolumetricCloud.usf b/Engine/Shaders/Private/VolumetricCloud.usf
--- a/Engine/Shaders/Private/VolumetricCloud.usf
+++ b/Engine/Shaders/Private/VolumetricCloud.usf
@@ -51,7 +51,7 @@ float3 SampleEmissive(in FPixelMaterialInputs PixelMaterialInputs)
 	float3 EmissiveColor = 0.0f;
 #if STRATA_ENABLED
 	FStrataBSDF BSDF = PixelMaterialInputs.FrontMaterial.Layers[0].BSDFs[0];
-	EmissiveColor = VOLUMETRICFOGCLOUD_EMISSIVE(BSDF);
+	EmissiveColor = VOLUMETRICFOGCLOUD_GETEMISSIVE(BSDF);
 #else
 	EmissiveColor = GetMaterialEmissiveRaw(PixelMaterialInputs);
 #endif
diff --git a/Engine/Shaders/Private/VolumetricFogVoxelization.usf b/Engine/Shaders/Private/VolumetricFogVoxelization.usf
--- a/Engine/Shaders/Private/VolumetricFogVoxelization.usf
+++ b/Engine/Shaders/Private/VolumetricFogVoxelization.usf
@@ -40,7 +40,7 @@ float3 SampleEmissive(in FPixelMaterialInputs PixelMaterialInputs)
 	float3 EmissiveColor = 0.0f;
 #if STRATA_ENABLED
 	FStrataBSDF BSDF = PixelMaterialInputs.FrontMaterial.Layers[0].BSDFs[0];
-	EmissiveColor = VOLUMETRICFOGCLOUD_EMISSIVE(BSDF);
+	EmissiveColor = VOLUMETRICFOGCLOUD_GETEMISSIVE(BSDF);
 #else
 	EmissiveColor = GetMaterialEmissiveRaw(PixelMaterialInputs);
 #endif
