Add initial SSS support for strata material.

This CL adds support for SSS profile only, and use a single profile, set on the node material. If several SSS layer are present, the first diffuse layer on the stack wich has a non-0 weight and a non-0 radius will be used

diff --git a/Engine/Shaders/Private/DeferredShadingCommon.ush b/Engine/Shaders/Private/DeferredShadingCommon.ush
--- a/Engine/Shaders/Private/DeferredShadingCommon.ush
+++ b/Engine/Shaders/Private/DeferredShadingCommon.ush
@@ -349,12 +349,13 @@ bool IsSubsurfaceModel(int ShadingModel)
 		|| ShadingModel == SHADINGMODELID_SUBSURFACE_PROFILE
 		|| ShadingModel == SHADINGMODELID_TWOSIDED_FOLIAGE
 		|| ShadingModel == SHADINGMODELID_HAIR
-		|| ShadingModel == SHADINGMODELID_EYE;
+		|| ShadingModel == SHADINGMODELID_EYE 
+		|| ShadingModel == SHADINGMODELID_STRATA;
 }
 
 bool UseSubsurfaceProfile(int ShadingModel)
 {
-	return ShadingModel == SHADINGMODELID_SUBSURFACE_PROFILE || ShadingModel == SHADINGMODELID_EYE;
+	return ShadingModel == SHADINGMODELID_SUBSURFACE_PROFILE || ShadingModel == SHADINGMODELID_EYE || ShadingModel == SHADINGMODELID_STRATA;
 }
 
 bool HasCustomGBufferData(int ShadingModelID)
diff --git a/Engine/Shaders/Private/ShadowProjectionPixelShader.usf b/Engine/Shaders/Private/ShadowProjectionPixelShader.usf
--- a/Engine/Shaders/Private/ShadowProjectionPixelShader.usf
+++ b/Engine/Shaders/Private/ShadowProjectionPixelShader.usf
@@ -297,7 +297,7 @@ void Main(
 				Settings.ProjectionDepthBiasParameters = ProjectionDepthBiasParameters.xw;
 
 			  #if USE_TRANSMISSION
-				if (GBufferData.ShadingModelID == SHADINGMODELID_SUBSURFACE_PROFILE)
+				if (GBufferData.ShadingModelID == SHADINGMODELID_SUBSURFACE_PROFILE || GBufferData.ShadingModelID == SHADINGMODELID_STRATA)
 				{
 					SSSTransmission = CalcTransmissionThickness(ScreenPosition.xyz, LightPositionAndInvRadius.xyz, GBufferData, Settings);
 				}
diff --git a/Engine/Shaders/Private/Strata/Strata.ush b/Engine/Shaders/Private/Strata/Strata.ush
--- a/Engine/Shaders/Private/Strata/Strata.ush
+++ b/Engine/Shaders/Private/Strata/Strata.ush
@@ -94,6 +94,8 @@ struct FStrataBSDF
 
 #define DIFFUSE_ALBEDO(X)				X.VGPRs[0].xyz
 #define DIFFUSE_ROUGHNESS(X)			X.VGPRs[0].w
+#define DIFFUSE_SSSPROFILEID(X)			X.VGPRs[1].x
+#define DIFFUSE_SSSPROFILERADIUS(X)		X.VGPRs[1].y
 
 #define DIELECTRIC_TINT(X)				X.VGPRs[0].xyz
 #define DIELECTRIC_IOR(X)				X.VGPRs[0].w
@@ -165,14 +167,19 @@ struct FStrataData
 	FStrataLayer Layers[STRATA_MAX_LAYER_COUNT];	// Layer 0 is the top layer
 };
 
+uint GetStrataSubsurfaceProfile(float NormalizedId)
+{
+	// can be optimized
+	return uint(NormalizedId * 255.0f + 0.5f);
 
+}
 FStrataData GetInitialisedStrataData()
 {
 	FStrataData StrataData = (FStrataData)0;
 	return StrataData;
 }
 
-FStrataData GetStrataDiffuseBSDF(float3 Albedo, float Roughness, uint SharedNormalIndex)
+FStrataData GetStrataDiffuseBSDF(float3 Albedo, float Roughness, float SSSProfileID, float SSSRadiusScale, uint SharedNormalIndex)
 {
 	FStrataData StrataData = GetInitialisedStrataData();
 	StrataData.LayerCount++;
@@ -181,8 +188,11 @@ FStrataData GetStrataDiffuseBSDF(float3 Albedo, float Roughness, uint SharedNorm
 	BSDF_SETTYPE			(StrataData.Layers[0].BSDFs[0],		STRATA_BSDF_TYPE_DIFFUSE);
 	BSDF_SETNORMALID		(StrataData.Layers[0].BSDFs[0],		SharedNormalIndex);
 	BSDF_SETISTOPLAYER		(StrataData.Layers[0].BSDFs[0],		0);
+	BSDF_SETHASSCATTERING	(StrataData.Layers[0].BSDFs[0],		SSSRadiusScale > 0 ? 1 : 0);
 	DIFFUSE_ALBEDO			(StrataData.Layers[0].BSDFs[0])		= Albedo;
 	DIFFUSE_ROUGHNESS		(StrataData.Layers[0].BSDFs[0])		= Roughness;
+	DIFFUSE_SSSPROFILEID	(StrataData.Layers[0].BSDFs[0])		= SSSProfileID;
+	DIFFUSE_SSSPROFILERADIUS(StrataData.Layers[0].BSDFs[0])		= SSSRadiusScale;
 
 	StrataData.Layers[0].BSDFs[0].Weight						= 1.0f;
 
@@ -510,11 +520,16 @@ void StrataUpdateBSDFContext(inout FStrataBSDFContext BSDFContext, float3 NewL)
 	BSDFContext.TangentH = mul(BSDFContext.TangentBasis, BSDFContext.H);
 }
 
+#define STRATA_PATH_TYPE_DIFFUSE  0x1 
+#define STRATA_PATH_TYPE_SPECULAR 0x2
+#define STRATA_PATH_TYPE_SCATTER  0x4
+
 struct FStrataEvaluateResult
 {
 	float3 Value;
 	float  PDF;
 	float3 Throughput;	// Throughput to the next layer
+	uint   PathType;	// Type of path (specular|diffuse|scattering)
 };
 
 FStrataEvaluateResult StrataEvaluateBSDF(FStrataBSDFContext BSDFContext)
@@ -534,6 +549,7 @@ FStrataEvaluateResult StrataEvaluateBSDF(FStrataBSDFContext BSDFContext)
 			Sample.Value		= Diffuse_Chan(DIFFUSE_ALBEDO(BSDFContext.BSDF), a2, BSDFContext.SatNoV, BSDFContext.SatNoL, BSDFContext.Context.VoH, BSDFContext.Context.NoH);
 			Sample.PDF			= BSDFContext.SatNoV / PI;
 			Sample.Throughput	= OpaqueBSDFThroughput;
+			Sample.PathType		= STRATA_PATH_TYPE_DIFFUSE | (BSDF_GETHASSCATTERING(BSDFContext.BSDF) ? STRATA_PATH_TYPE_SCATTER : 0);
 			break;
 		}
 		case STRATA_BSDF_TYPE_DIELECTRIC:
@@ -573,6 +589,7 @@ FStrataEvaluateResult StrataEvaluateBSDF(FStrataBSDFContext BSDFContext)
 			Sample.Value		= D * Vis * F;
 			Sample.PDF			= PDF;
 			Sample.Throughput	= bIsDieletric ? (1.0f - F) : OpaqueBSDFThroughput; // STRATA_TODO (Kulla's (1-FDG))
+			Sample.PathType		= STRATA_PATH_TYPE_SPECULAR;
 			break;
 		}
 		case STRATA_BSDF_TYPE_VOLUME:
@@ -602,6 +619,7 @@ FStrataEvaluateResult StrataEvaluateBSDF(FStrataBSDFContext BSDFContext)
 			Sample.Value		= NormLuminanceSingleScattering * Phase * SelfShadowTransmittance;
 			Sample.PDF			= Phase;
 			Sample.Throughput	= SafePathSegmentTransmittance;
+			Sample.PathType		= STRATA_PATH_TYPE_DIFFUSE;
 			break;
 		}
 		case STRATA_BSDF_TYPE_SHEEN:
@@ -617,9 +635,10 @@ FStrataEvaluateResult StrataEvaluateBSDF(FStrataBSDFContext BSDFContext)
 			float Vis = Vis_Cloth(BSDFContext.Context.NoV, BSDFContext.SatNoL);
 			float3 F  = F_Schlick(F0, BSDFContext.Context.VoH);
 
-			Sample.Value = D * Vis * F;
-			Sample.PDF = BSDFContext.SatNoV / PI; // Per "Production Friendly Microfacet Sheen BRDF", hemispherical sampling give good result as the roughness is usually high.
-			Sample.Throughput = (1.0f - F);
+			Sample.Value		= D * Vis * F;
+			Sample.PDF			= BSDFContext.SatNoV / PI; // Per "Production Friendly Microfacet Sheen BRDF", hemispherical sampling give good result as the roughness is usually high.
+			Sample.Throughput	= (1.0f - F);
+			Sample.PathType		= STRATA_PATH_TYPE_DIFFUSE;
 			break;
 		}
 	}
@@ -1039,6 +1058,25 @@ float4 UnpackAlbedoGamma2ToLinearAlphaLinear(uint rgba)
 	return float4(Gamma2 * Gamma2, a * (1.0f / 255.0f));	// Simple gamma 2.0
 }
 
+uint PackRGBA8(float4 rgba)
+{
+	uint r = (uint(saturate(rgba.r) * 255.0f) << 0);
+	uint g = (uint(saturate(rgba.g) * 255.0f) << 8);
+	uint b = (uint(saturate(rgba.b) * 255.0f) << 16);
+	uint a = (uint(saturate(rgba.a) * 255.0f) << 24);
+	return r | g | b | a;
+}
+
+float4 UnpackRGBA8(uint rgba)
+{
+	float4 Out;
+	Out.r = float((rgba >> 0)  & 0xFF);
+	Out.g = float((rgba >> 8)  & 0xFF);
+	Out.b = float((rgba >> 16) & 0xFF);
+	Out.a = float((rgba >> 24) & 0xFF);
+	return Out;
+}
+
 // Simple Octahedral from "A Survey of Efficient Representations for Independent Unit Vectors" http://jcgt.org/published/0003/02/01/
 float SignNotZero(in float k)
 {
@@ -1274,6 +1312,12 @@ void PackStrataOut(
 	// We clamp the layer count to the maximum one according to our byte per pixel budget (removing the bottom layer first)
 	Strata.LayerCount = STRATA_CLAMPED_LAYER_COUNT;
 
+	// Subsurface scattering parameters
+	bool bHasSubsurfaceScattering = false;
+	float SubsurfaceProfile = 0;
+	float SubsurfaceRadiusScale = 0;
+	float3 SubsurfaceBaseColor = 0;
+
 	// For each of the BSDF, we evaluated their combined weight/throughput according to the material operation performed (layering, mixing, etc.).
 	int l;
 	UNROLL
@@ -1320,6 +1364,15 @@ void PackStrataOut(
 					AccumRoughnessCount += BSDFCoverageAvg;
 				}
 
+				// Select the first/toppest layer with subsurface scattering
+				if (BSDF_GETHASSCATTERING(BSDF) && BSDFCoverageAvg > 0 && DIFFUSE_SSSPROFILERADIUS(BSDF) > 0 && !bHasSubsurfaceScattering)
+				{
+					bHasSubsurfaceScattering = true;
+					SubsurfaceProfile = DIFFUSE_SSSPROFILEID(BSDF);
+					SubsurfaceRadiusScale = DIFFUSE_SSSPROFILERADIUS(BSDF);
+					SubsurfaceBaseColor = DIFFUSE_ALBEDO(BSDF);
+				}
+
 				LayerCoverage	+= BSDFCoverage;
 				LayerThroughput += BSDFCoverage * OpaqueBSDFThroughput;
 				break;
@@ -1430,6 +1483,13 @@ void PackStrataOut(
 			GBuffer.Metallic = 1;
 		}
 
+		if (bHasSubsurfaceScattering)
+		{
+			GBuffer.BaseColor      = SubsurfaceBaseColor;
+			GBuffer.CustomData.rgb = EncodeSubsurfaceProfile(SubsurfaceProfile);
+			GBuffer.CustomData.a   = SubsurfaceRadiusScale;
+		}
+
 		// Take into account the non covered part for this layer and update throughput accordingly
 		LayerThroughput += (1.0 - saturate(LayerCoverage)) * FullThroughput;
 
@@ -1437,7 +1497,6 @@ void PackStrataOut(
 		Throughput *= LayerThroughput;
 	}
 
-
 	// Now write out Strata data
 
 	// 1. the header
@@ -1489,7 +1548,11 @@ void PackStrataOut(
 			case STRATA_BSDF_TYPE_DIFFUSE:
 			{
 				STRATA_STORE_UINT1(PackAlbedoLinearToGamma2AlphaLinear(float4(DIFFUSE_ALBEDO(BSDF), DIFFUSE_ROUGHNESS(BSDF))));
-				// 4 bytes
+				if (BSDF_GETHASSCATTERING(BSDF))
+				{
+					STRATA_STORE_UINT1(PackRGBA8(float4(DIFFUSE_SSSPROFILEID(BSDF), DIFFUSE_SSSPROFILERADIUS(BSDF), 0, 0)))
+				}
+				// 8 bytes
 			}
 			break;
 			case STRATA_BSDF_TYPE_DIELECTRIC:
@@ -1596,6 +1659,12 @@ FStrataBSDF UnpackStrataBSDFIn(ByteAddressBuffer StrataBuffer, inout uint ByteOf
 	case STRATA_BSDF_TYPE_DIFFUSE:
 	{
 		StrataLoad_AlbedoGamma2ToLinearAlphaLinear(StrataBuffer, ByteOffset, DIFFUSE_ALBEDO(OutBSDF), DIFFUSE_ROUGHNESS(OutBSDF));
+		if (BSDF_GETHASSCATTERING(OutBSDF))
+		{
+			float4 Data = UnpackRGBA8(StrataLoadUint1(StrataBuffer, ByteOffset));
+			DIFFUSE_SSSPROFILEID(OutBSDF) = Data.r;
+			DIFFUSE_SSSPROFILERADIUS(OutBSDF) = Data.g;
+		}
 	}
 	break;
 	case STRATA_BSDF_TYPE_DIELECTRIC:
diff --git a/Engine/Shaders/Private/Strata/StrataDeferredLighting.ush b/Engine/Shaders/Private/Strata/StrataDeferredLighting.ush
--- a/Engine/Shaders/Private/Strata/StrataDeferredLighting.ush
+++ b/Engine/Shaders/Private/Strata/StrataDeferredLighting.ush
@@ -132,8 +132,6 @@ float4 StrataDeferredLighting(
 	ByteAddressBuffer StrataDataBuffer,
 	uint PixelStrataDataByteOffset)
 {
-	float4 OutRadiance = 0;
-
 	FStrataPixelHeader StrataPixelHeader = UnpackStrataHeaderIn(StrataDataBuffer, PixelStrataDataByteOffset);
 	
 	float3 PreintegratedIndicentIrradiance = 0;
@@ -148,6 +146,8 @@ float4 StrataDeferredLighting(
 		PreintegratedIndicentIrradiance = IntegrateLight(Capsule, LightData.bInverseSquared);
 	}
 
+	FLightAccumulator Out = (FLightAccumulator)0;
+
 	LOOP
 	for (uint BSDFIndex = 0; BSDFIndex < StrataPixelHeader.BSDFCount; ++BSDFIndex)
 	{
@@ -165,13 +165,28 @@ float4 StrataDeferredLighting(
 		if (IntegrationContext.bEvaluateBSDF)
 		{
 			FStrataEvaluateResult BSDFEvaluate = StrataEvaluateBSDF(BSDFContext);
-			OutRadiance.rgb += BSDF.Weight * BSDFEvaluate.Value * IntegrationContext.IncidentIrradiance;
+
+			const float3 DiffuseRadiance  = (BSDFEvaluate.PathType & STRATA_PATH_TYPE_DIFFUSE)  != 0 ? BSDF.Weight * BSDFEvaluate.Value : 0;
+			const float3 SpecularRadiance = (BSDFEvaluate.PathType & STRATA_PATH_TYPE_SPECULAR) != 0 ? BSDF.Weight * BSDFEvaluate.Value : 0;
+
+			float3 CommonFactor = IntegrationContext.IncidentIrradiance;
+			const bool bNeedsSeparateSubsurfaceLightAccumulation = BSDFEvaluate.PathType & STRATA_PATH_TYPE_SCATTER != 0;
+			LightAccumulator_AddSplit(Out, DiffuseRadiance, SpecularRadiance, DiffuseRadiance, CommonFactor, bNeedsSeparateSubsurfaceLightAccumulation);
 		}
 		else
 		{
-			OutRadiance.rgb += IntegrationContext.OutgoingRadiance;
+			
+			// STRATA_TODO Find a better solution for special of are lighting which handles specular path only, so we can hardcode weight.
+			const float3 DiffuseRadiance  = 0;
+			const float3 SpecularRadiance = 1;
+
+			float3 CommonFactor = IntegrationContext.OutgoingRadiance;
+			const bool bNeedsSeparateSubsurfaceLightAccumulation = false;
+			LightAccumulator_AddSplit(Out, DiffuseRadiance, SpecularRadiance, DiffuseRadiance, CommonFactor, bNeedsSeparateSubsurfaceLightAccumulation);
 		}
 	}
 
+	const FDeferredLightingSplit SplitLighting = LightAccumulator_GetResultSplit(Out);
+	const float4 OutRadiance = SplitLighting.SpecularLighting + SplitLighting.DiffuseLighting;
 	return OutRadiance;
 }
\ No newline at end of file
diff --git a/Engine/Shaders/Private/Strata/StrataVisualize.usf b/Engine/Shaders/Private/Strata/StrataVisualize.usf
--- a/Engine/Shaders/Private/Strata/StrataVisualize.usf
+++ b/Engine/Shaders/Private/Strata/StrataVisualize.usf
@@ -227,6 +227,7 @@ void VisualizeMaterialPS(
 				PRINTBSDFNORMALID();
 				PRINTBSDFANISO();
 				PRINTBSDFTOPLAY();
+				PRINTBSDFHASSCAT();
 				NEXTLINE;
 
 				HORIZONTALTAB; PRINTPARAMCHAR(_A_); PRINTPARAMCHAR(_L_); PRINTPARAMCHAR(_B_); PRINTPARAMCHAR(_E_); PRINTPARAMCHAR(_D_); PRINTPARAMCHAR(_O_); 
@@ -235,6 +236,11 @@ void VisualizeMaterialPS(
 				HORIZONTALTAB; PRINTPARAMCHAR(_R_); PRINTPARAMCHAR(_O_); PRINTPARAMCHAR(_U_); PRINTPARAMCHAR(_G_); PRINTPARAMCHAR(_H_); PRINTPARAMCHAR(_SPC_);
 				HORIZONTALTAB; PRINTFLOAT1(DIFFUSE_ROUGHNESS(BSDF));
 				NEXTLINE;
+				HORIZONTALTAB; PRINTPARAMCHAR(_S_); PRINTPARAMCHAR(_S_); PRINTPARAMCHAR(_S_); PRINTPARAMCHAR(_SPC_); PRINTPARAMCHAR(_I_); PRINTPARAMCHAR(_D_); PRINTPARAMCHAR(_SPC_);
+				HORIZONTALTAB; PRINTFLOAT1(DIFFUSE_SSSPROFILEID(BSDF));
+				NEXTLINE;
+				HORIZONTALTAB; PRINTPARAMCHAR(_S_); PRINTPARAMCHAR(_S_); PRINTPARAMCHAR(_S_); PRINTPARAMCHAR(_SPC_); PRINTPARAMCHAR(_R_); PRINTPARAMCHAR(_A_); PRINTPARAMCHAR(_SPC_);
+				HORIZONTALTAB; PRINTFLOAT1(DIFFUSE_SSSPROFILERADIUS(BSDF));
 
 				break;
 
diff --git a/Engine/Shaders/Private/TransmissionCommon.ush b/Engine/Shaders/Private/TransmissionCommon.ush
--- a/Engine/Shaders/Private/TransmissionCommon.ush
+++ b/Engine/Shaders/Private/TransmissionCommon.ush
@@ -17,13 +17,10 @@ struct FTransmissionProfileParams
 	float OneOverIOR;
 };
 
-FTransmissionProfileParams GetTransmissionProfileParams(FGBufferData GBufferData)
+FTransmissionProfileParams GetTransmissionProfileParams(uint SubsurfaceProfileInt)
 {
 	FTransmissionProfileParams Result;
 
-	// 0..255, which SubSurface profile to pick
-	uint SubsurfaceProfileInt = ExtractSubsurfaceProfileInt(GBufferData);
-
 	//X:ExtinctionScale, Y:Normal Scale, Z:ScatteringDistribution, W:OneOverIOR
 	float4 Data = ActualSSProfilesTexture.Load(int3(SSSS_TRANSMISSION_OFFSET, SubsurfaceProfileInt, 0));
 
@@ -35,10 +32,16 @@ FTransmissionProfileParams GetTransmissionProfileParams(FGBufferData GBufferData
 	return Result;
 }
 
-float4 GetTransmissionProfile(FGBufferData GBufferData, float Thickness)
+FTransmissionProfileParams GetTransmissionProfileParams(FGBufferData GBufferData)
 {
 	// 0..255, which SubSurface profile to pick
 	uint SubsurfaceProfileInt = ExtractSubsurfaceProfileInt(GBufferData);
+	return GetTransmissionProfileParams(SubsurfaceProfileInt);
+}
+
+float4 GetTransmissionProfile(uint SubsurfaceProfileInt, float Thickness)
+{
+	// 0..255, which SubSurface profile to pick
 	float3 TextureSize;
 	ActualSSProfilesTexture.GetDimensions(0, TextureSize.x, TextureSize.y,TextureSize.z); // TODO: it's slightly faster to pass texture dimensions through a constant buffer
 	bool bUseBurley = GetSubsurfaceTransmittanceProfileUseBurley(SubsurfaceProfileInt);
@@ -47,6 +50,12 @@ float4 GetTransmissionProfile(FGBufferData GBufferData, float Thickness)
 	return Texture2DSampleLevel(ActualSSProfilesTexture, TransmissionProfilesLinearSampler, uv, 0);
 }
 
+float4 GetTransmissionProfile(FGBufferData GBufferData, float Thickness)
+{
+	// 0..255, which SubSurface profile to pick
+	return GetTransmissionProfile(ExtractSubsurfaceProfileInt(GBufferData), Thickness);
+}
+
 float EncodeThickness(float Thickness)
 {
 	return 1 - Thickness;
diff --git a/Engine/Source/Runtime/Engine/Classes/Materials/MaterialExpressionStrata.h b/Engine/Source/Runtime/Engine/Classes/Materials/MaterialExpressionStrata.h
--- a/Engine/Source/Runtime/Engine/Classes/Materials/MaterialExpressionStrata.h
+++ b/Engine/Source/Runtime/Engine/Classes/Materials/MaterialExpressionStrata.h
@@ -58,6 +58,16 @@ class UMaterialExpressionStrataDiffuseBSDF : public UMaterialExpression // STRAT
 	UPROPERTY()
 	FExpressionInput Normal;
 
+	/**
+	 * SSS radius scale (type = float, unit = unitless)
+	 */
+	UPROPERTY()
+	FExpressionInput SubsurfaceProfileScale;
+
+	/** SubsurfaceProfile, for Screen Space Subsurface Scattering. The profile needs to be set up on both the Strata diffuse node, and the material node at the moment. */
+	UPROPERTY(EditAnywhere, BlueprintReadOnly, Category = Material, meta = (DisplayName = "Subsurface Profile"))
+	class USubsurfaceProfile* SubsurfaceProfile;
+
 	//~ Begin UMaterialExpression Interface
 #if WITH_EDITOR
 	virtual int32 Compile(class FMaterialCompiler* Compiler, int32 OutputIndex) override;
diff --git a/Engine/Source/Runtime/Engine/Private/Materials/HLSLMaterialTranslator.cpp b/Engine/Source/Runtime/Engine/Private/Materials/HLSLMaterialTranslator.cpp
--- a/Engine/Source/Runtime/Engine/Private/Materials/HLSLMaterialTranslator.cpp
+++ b/Engine/Source/Runtime/Engine/Private/Materials/HLSLMaterialTranslator.cpp
@@ -7476,12 +7476,14 @@ int32 FHLSLMaterialTranslator::StrataCreateAndRegisterNullMaterial()
 	return OutputCodeChunk;
 }
 
-int32 FHLSLMaterialTranslator::StrataDiffuseBSDF(int32 Albedo, int32 Roughness, int32 Normal, uint8 SharedNormalIndex)
+int32 FHLSLMaterialTranslator::StrataDiffuseBSDF(int32 Albedo, int32 Roughness, int32 SSSProfileId, int32 SSSProfileRadiusScale, int32 Normal, uint8 SharedNormalIndex)
 {
 	return AddCodeChunk(
-		MCT_Strata, TEXT("GetStrataDiffuseBSDF(%s, %s, %u) /* %s */"),
+		MCT_Strata, TEXT("GetStrataDiffuseBSDF(%s, %s, %s, %s, %u) /* %s */"),
 		*GetParameterCode(Albedo),
 		*GetParameterCode(Roughness),
+		*GetParameterCode(SSSProfileId),
+		*GetParameterCode(SSSProfileRadiusScale),
 		SharedNormalIndex,
 		*GetParameterCode(Normal)
 	);
diff --git a/Engine/Source/Runtime/Engine/Private/Materials/HLSLMaterialTranslator.h b/Engine/Source/Runtime/Engine/Private/Materials/HLSLMaterialTranslator.h
--- a/Engine/Source/Runtime/Engine/Private/Materials/HLSLMaterialTranslator.h
+++ b/Engine/Source/Runtime/Engine/Private/Materials/HLSLMaterialTranslator.h
@@ -791,7 +791,7 @@ protected:
 
 	// Strata
 	virtual int32 StrataCreateAndRegisterNullMaterial() override;
-	virtual int32 StrataDiffuseBSDF(int32 Albedo, int32 Roughness, int32 Normal, uint8 SharedNormalIndex) override;
+	virtual int32 StrataDiffuseBSDF(int32 Albedo, int32 Roughness, int32 SSSProfileId, int32 SSSProfileRadiusScale, int32 Normal, uint8 SharedNormalIndex) override;
 	virtual int32 StrataDielectricBSDF(int32 RoughnessX, int32 RoughnessY, int32 IOR, int32 Tint, int32 Normal, int32 Tangent, uint8 SharedNormalIndex) override;
 	virtual int32 StrataConductorBSDF(int32 Reflectivity, int32 EdgeColor, int32 RoughnessX, int32 RoughnessY, int32 Normal, int32 Tangent, uint8 SharedNormalIndex) override;
 	virtual int32 StrataVolumeBSDF(int32 Albedo, int32 Extinction, int32 Anisotropy, int32 Thickness, int32 Normal, uint8 SharedNormalIndex) override;
diff --git a/Engine/Source/Runtime/Engine/Private/Materials/MaterialExpressions.cpp b/Engine/Source/Runtime/Engine/Private/Materials/MaterialExpressions.cpp
--- a/Engine/Source/Runtime/Engine/Private/Materials/MaterialExpressions.cpp
+++ b/Engine/Source/Runtime/Engine/Private/Materials/MaterialExpressions.cpp
@@ -19122,12 +19122,23 @@ int32 UMaterialExpressionStrataDiffuseBSDF::Compile(class FMaterialCompiler* Com
 	int32 NormalCodeChunk = CompileWithDefaultCodeChunk(Compiler, Normal, Compiler->VertexNormal());
 	uint8 SharedNormalIndex = StrataCompilationInfoCreateSharedNormal(Compiler, NormalCodeChunk);
 
+	FName NameSubsurfaceProfile(FString(TEXT("__SubsurfaceProfile")));
+
+	int32 SSSProfileCodeChunk = INDEX_NONE;
+	const bool bHasScattering = SubsurfaceProfile != nullptr;
+	if (bHasScattering)
+	{
+		SSSProfileCodeChunk = Compiler->ForceCast(Compiler->ScalarParameter(NameSubsurfaceProfile, 1.0f), MCT_Float1);
+	}
+
 	int32 OutputCodeChunk = Compiler->StrataDiffuseBSDF(
 		CompileWithDefaultFloat3(Compiler, Albedo,		0.18f, 0.18f, 0.18f),
 		CompileWithDefaultFloat1(Compiler, Roughness,	0.0f),
+		SSSProfileCodeChunk != INDEX_NONE ? SSSProfileCodeChunk : Compiler->Constant(1.0f),
+		CompileWithDefaultFloat1(Compiler, SubsurfaceProfileScale,	0.0f),
 		NormalCodeChunk,
 		SharedNormalIndex);
-	StrataCompilationInfoCreateSingleBSDFMaterial(Compiler, OutputCodeChunk, SharedNormalIndex, STRATA_BSDF_TYPE_DIFFUSE);
+	StrataCompilationInfoCreateSingleBSDFMaterial(Compiler, OutputCodeChunk, SharedNormalIndex, STRATA_BSDF_TYPE_DIFFUSE, bHasScattering);
 
 	return OutputCodeChunk;
 }
@@ -19155,6 +19166,9 @@ uint32 UMaterialExpressionStrataDiffuseBSDF::GetInputType(int32 InputIndex)
 	case 2:
 		return MCT_Float3;
 		break;
+	case 3:
+		return MCT_Float1;
+		break;
 	}
 
 	check(false);
diff --git a/Engine/Source/Runtime/Engine/Private/Materials/StrataMaterial.cpp b/Engine/Source/Runtime/Engine/Private/Materials/StrataMaterial.cpp
--- a/Engine/Source/Runtime/Engine/Private/Materials/StrataMaterial.cpp
+++ b/Engine/Source/Runtime/Engine/Private/Materials/StrataMaterial.cpp
@@ -52,14 +52,14 @@ uint8 StrataCompilationInfoCreateSharedNormal(FMaterialCompiler* Compiler, int32
 	return Compiler->StrataCompilationInfoRegisterSharedNormalIndex(NormalCodeChunk, TangentCodeChunk);
 }
 
-void StrataCompilationInfoCreateSingleBSDFMaterial(FMaterialCompiler* Compiler, int32 CodeChunk, uint8 SharedNormalIndex, uint8 BSDFType, bool bVolumeHasScattering)
+void StrataCompilationInfoCreateSingleBSDFMaterial(FMaterialCompiler* Compiler, int32 CodeChunk, uint8 SharedNormalIndex, uint8 BSDFType, bool bHasScattering)
 {
 	FStrataMaterialCompilationInfo StrataInfo;
 	StrataInfo.LayerCount = 1;
 	StrataInfo.Layers[0].BSDFCount = 1;
 	StrataInfo.Layers[0].BSDFs[0].Type = BSDFType;
 	StrataInfo.Layers[0].BSDFs[0].SharedNormalIndex = SharedNormalIndex;
-	StrataInfo.Layers[0].BSDFs[0].bVolumeHasScattering = bVolumeHasScattering;
+	StrataInfo.Layers[0].BSDFs[0].bHasScattering = bHasScattering;
 	UpdateTotalBSDFCount(StrataInfo);
 	Compiler->StrataCompilationInfoRegisterCodeChunk(CodeChunk, StrataInfo);
 }
@@ -225,6 +225,10 @@ FStrataMaterialAnalysisResult StrataCompilationInfoMaterialAnalysis(FMaterialCom
 			case STRATA_BSDF_TYPE_DIFFUSE:
 			{
 				Result.RequestedByteCount += UintByteSize;
+				if (BSDF.bHasScattering)
+				{
+					Result.RequestedByteCount += UintByteSize;
+				}
 				break;
 			}
 			case STRATA_BSDF_TYPE_DIELECTRIC:
@@ -243,7 +247,7 @@ FStrataMaterialAnalysisResult StrataCompilationInfoMaterialAnalysis(FMaterialCom
 			{
 				Result.RequestedByteCount += UintByteSize;
 
-				if (BSDF.bVolumeHasScattering)
+				if (BSDF.bHasScattering)
 				{
 					Result.RequestedByteCount += UintByteSize;
 				}
diff --git a/Engine/Source/Runtime/Engine/Private/Materials/StrataMaterial.h b/Engine/Source/Runtime/Engine/Private/Materials/StrataMaterial.h
--- a/Engine/Source/Runtime/Engine/Private/Materials/StrataMaterial.h
+++ b/Engine/Source/Runtime/Engine/Private/Materials/StrataMaterial.h
@@ -18,7 +18,7 @@ struct FStrataMaterialCompilationInfo
 	{
 		uint8 Type;
 		uint8 SharedNormalIndex;
-		bool  bVolumeHasScattering;
+		bool  bHasScattering;
 	};
 
 	struct FLayer
@@ -45,7 +45,7 @@ FString GetStrataBSDFName(uint8 BSDFType);
 uint8 StrataCompilationInfoCreateSharedNormal(FMaterialCompiler* Compiler, int32 NormalCodeChunk);
 uint8 StrataCompilationInfoCreateSharedNormal(FMaterialCompiler* Compiler, int32 NormalCodeChunk, int32 TangentCodeChunk);
 void  StrataCompilationInfoCreateNullBSDF(FMaterialCompiler* Compiler, int32 CodeChunk);
-void  StrataCompilationInfoCreateSingleBSDFMaterial(FMaterialCompiler* Compiler, int32 CodeChunk, uint8 SharedNormalIndex, uint8 BSDFType, bool bVolumeHasScattering = false);
+void  StrataCompilationInfoCreateSingleBSDFMaterial(FMaterialCompiler* Compiler, int32 CodeChunk, uint8 SharedNormalIndex, uint8 BSDFType, bool bHasScattering = false);
 
 FStrataMaterialCompilationInfo StrataCompilationInfoAdd(FMaterialCompiler* Compiler, const FStrataMaterialCompilationInfo& A, const FStrataMaterialCompilationInfo& B);
 FStrataMaterialCompilationInfo StrataCompilationInfoMultiply(FMaterialCompiler* Compiler, const FStrataMaterialCompilationInfo& A);
diff --git a/Engine/Source/Runtime/Engine/Public/MaterialCompiler.h b/Engine/Source/Runtime/Engine/Public/MaterialCompiler.h
--- a/Engine/Source/Runtime/Engine/Public/MaterialCompiler.h
+++ b/Engine/Source/Runtime/Engine/Public/MaterialCompiler.h
@@ -391,7 +391,7 @@ public:
 
 	// Strata
 	virtual int32 StrataCreateAndRegisterNullMaterial() = 0;
-	virtual int32 StrataDiffuseBSDF(int32 Albedo, int32 Roughness, int32 Normal, uint8 SharedNormalIndex) = 0;
+	virtual int32 StrataDiffuseBSDF(int32 Albedo, int32 Roughness, int32 SSSProfileId, int32 SSSProfileRadiusScale, int32 Normal, uint8 SharedNormalIndex) = 0;
 	virtual int32 StrataDielectricBSDF(int32 RoughnessX, int32 RoughnessY, int32 IOR, int32 Tint, int32 Normal, int32 Tangent, uint8 SharedNormalIndex) = 0;
 	virtual int32 StrataConductorBSDF(int32 Reflectivity, int32 EdgeColor, int32 RoughnessX, int32 RoughnessY, int32 Normal, int32 Tangent, uint8 SharedNormalIndex) = 0;
 	virtual int32 StrataVolumeBSDF(int32 Albedo, int32 Extinction, int32 Anisotropy, int32 Thickness, int32 Normal, uint8 SharedNormalIndex) = 0;
@@ -806,9 +806,9 @@ public:
 		return Compiler->StrataCreateAndRegisterNullMaterial();
 	}
 
-	virtual int32 StrataDiffuseBSDF(int32 Albedo, int32 Roughness, int32 Normal, uint8 SharedNormalIndex) override
+	virtual int32 StrataDiffuseBSDF(int32 Albedo, int32 Roughness, int32 SSSProfileId, int32 SSSProfileRadiusScale, int32 Normal, uint8 SharedNormalIndex) override
 	{
-		return Compiler->StrataDiffuseBSDF(Albedo, Roughness, Normal, SharedNormalIndex);
+		return Compiler->StrataDiffuseBSDF(Albedo, Roughness, SSSProfileId, SSSProfileRadiusScale, Normal, SharedNormalIndex);
 	}
 
 	virtual int32 StrataDielectricBSDF(int32 RoughnessX, int32 RoughnessY, int32 IOR, int32 Tint, int32 Normal, int32 Tangent, uint8 SharedNormalIndex) override
