Make UniformSampleDiskConcentric branch free and fix NaN from 0/0 case

Add variants of cosine lobe sampling and uniform cone sampling using this mapping as examples of non-uniform sampling distributions.

diff --git a/Engine/Shaders/Private/MonteCarlo.ush b/Engine/Shaders/Private/MonteCarlo.ush
--- a/Engine/Shaders/Private/MonteCarlo.ush
+++ b/Engine/Shaders/Private/MonteCarlo.ush
@@ -110,22 +110,27 @@ float2 UniformSampleDisk( float2 E )
 	return Radius * float2( cos( Theta ), sin( Theta ) );
 }
 
-float2 UniformSampleDiskConcentric( float2 E )
+// Returns a point on the unit circle and a radius in z
+float3 ConcentricDiskSamplingHelper(float2 E)
 {
 	float2 p = 2 * E - 1;
-	float Radius;
-	float Phi;
-	if( abs( p.x ) > abs( p.y ) )
-	{
-		Radius = p.x;
-		Phi = (PI/4) * (p.y / p.x);
-	}
-	else
-	{
-		Radius = p.y;
-		Phi = (PI/2) - (PI/4) * (p.x / p.y);
-	}
-	return float2( Radius * cos( Phi ), Radius * sin( Phi ) );
+	float2 a = abs(p);
+	float Lo = min(a.x, a.y);
+	float Hi = max(a.x, a.y);
+	float Epsilon = 5.42101086243e-20; // 2^-64 (this avoids 0/0 without changing the rest of the mapping)
+	float Phi = (PI / 4) * (Lo / (Hi + Epsilon) + 2 * float(a.y >= a.x));
+	float Radius = Hi;
+	// copy sign bits from p
+	const uint SignMask = 0x80000000;
+	float2 Disk = asfloat((asuint(float2(cos(Phi), sin(Phi))) & ~SignMask) | (asuint(p) & SignMask));
+	// return point on the circle as well as the radius
+	return float3(Disk, Radius);
+}
+
+float2 UniformSampleDiskConcentric( float2 E )
+{
+	float3 Result = ConcentricDiskSamplingHelper(E);
+	return Result.xy * Result.z; // uniform sampling
 }
 
 // based on the approximate equal area transform from
@@ -184,17 +189,26 @@ float4 UniformSampleHemisphere( float2 E )
 float4 CosineSampleHemisphere( float2 E )
 {
 	float Phi = 2 * PI * E.x;
-	float CosTheta = sqrt( E.y );
-	float SinTheta = sqrt( 1 - CosTheta * CosTheta );
+	float CosTheta = sqrt(E.y);
+	float SinTheta = sqrt(1 - CosTheta * CosTheta);
 
 	float3 H;
-	H.x = SinTheta * cos( Phi );
-	H.y = SinTheta * sin( Phi );
+	H.x = SinTheta * cos(Phi);
+	H.y = SinTheta * sin(Phi);
 	H.z = CosTheta;
 
-	float PDF = CosTheta * (1.0 /  PI);
+	float PDF = CosTheta * (1.0 / PI);
 
-	return float4( H, PDF );
+	return float4(H, PDF);
+}
+
+// PDF = NoL / PI
+float4 CosineSampleHemisphereConcentric(float2 E)
+{
+	float3 Result = ConcentricDiskSamplingHelper(E);
+	float SinTheta = Result.z;
+	float CosTheta = sqrt(1 - SinTheta * SinTheta);
+	return float4(Result.xy * SinTheta, CosTheta, CosTheta * (1.0 / PI));
 }
 
 // PDF = NoL / PI
@@ -242,7 +256,23 @@ float4 UniformSampleConeRobust(float2 E, float SinThetaMax2)
 	L.x = SinTheta * cos(Phi);
 	L.y = SinTheta * sin(Phi);
 	L.z = CosTheta;
+	float PDF = 1.0 / (2 * PI * OneMinusCosThetaMax);
+
+	return float4(L, PDF);
+}
+
+// Same as the function above, but uses a concentric mapping
+float4 UniformSampleConeConcentricRobust(float2 E, float SinThetaMax2)
+{
+	// The expression 1-sqrt(1-x) is susceptible to catastrophic cancelation.
+	// Instead, use a series expansion about 0 which is accurate within 10^-7
+	// and much more numerically stable.
+	float OneMinusCosThetaMax = SinThetaMax2 < 0.01 ? SinThetaMax2 * (0.5 + 0.125 * SinThetaMax2) : 1 - sqrt(1 - SinThetaMax2);
+	float3 Result = ConcentricDiskSamplingHelper(E);
+	float SinTheta = Result.z * sqrt(SinThetaMax2);
+	float CosTheta = sqrt(1 - SinTheta * SinTheta);
 
+	float3 L = float3(Result.xy * SinTheta, CosTheta);
 	float PDF = 1.0 / (2 * PI * OneMinusCosThetaMax);
 
 	return float4(L, PDF);
diff --git a/Engine/Shaders/Private/PathTracing/Material/PathTracingRadianceProbe.ush b/Engine/Shaders/Private/PathTracing/Material/PathTracingRadianceProbe.ush
--- a/Engine/Shaders/Private/PathTracing/Material/PathTracingRadianceProbe.ush
+++ b/Engine/Shaders/Private/PathTracing/Material/PathTracingRadianceProbe.ush
@@ -11,12 +11,6 @@ PathTracingLambert.usf: Lambertian BRDF sampling functions
 	#include "FirstBounceRayGuidingCommon.ush"
 #endif
 
-float4 CosineSampleHemisphereConcentric( float2 E )
-{
-	float2 d = UniformSampleDiskConcentric(E);
-	return float4(d.x, d.y, sqrt(saturate(1 - d.x * d.x - d.y * d.y)), sqrt(saturate(1 - d.x * d.x - d.y * d.y)) / PI);
-}
-
 float2 InverseConcentricMapping(float2 p)
 {
 	// Handle degeneracy at the origin
