Fixed strata shader compilation issue on a console.

diff --git a/Engine/Shaders/Private/Strata/StrataVisualize.usf b/Engine/Shaders/Private/Strata/StrataVisualize.usf
--- a/Engine/Shaders/Private/Strata/StrataVisualize.usf
+++ b/Engine/Shaders/Private/Strata/StrataVisualize.usf
@@ -173,8 +173,8 @@ void VisualizeMaterialPS(
 		float3 INFOFontColor = float3(1.0, 0.5, 0.25);
 		float3 PARAMFontColor = float3(0.8, 0.8, 0.8);
 		float3 FontColor = float3(0.25, 1.0, 1.0);
-		int2 TopLeftInit = MaterialPos + InspectSquare + int2(3, 3);
-		int2 TopLeft = TopLeftInit;
+		uint2 TopLeftInit = MaterialPos + InspectSquare + uint2(3, 3);
+		uint2 TopLeft = TopLeftInit;
 
 #define PRINTPARAMCHAR(X) PrintCharacter(PixelPos, OutColor.xyz, PARAMFontColor, TopLeft, X);
 #define PRINTINFOCHAR(X) PrintCharacter(PixelPos, OutColor.xyz, INFOFontColor, TopLeft, X);
