Fix strata subsurface scattering affecting other object in the scene.

diff --git a/Engine/Shaders/Private/ReflectionEnvironmentPixelShader.usf b/Engine/Shaders/Private/ReflectionEnvironmentPixelShader.usf
--- a/Engine/Shaders/Private/ReflectionEnvironmentPixelShader.usf
+++ b/Engine/Shaders/Private/ReflectionEnvironmentPixelShader.usf
@@ -358,10 +358,9 @@ void ReflectionEnvironmentSkyLighting(
 	const bool bStrataMaterial = false;
 #endif
 
-	float3 SkyLighting = 0;
+	float4 SkyLighting = 0;
 	BRANCH
-#if STRATA_ENABLED
-#if ENABLE_SKY_LIGHT
+#if STRATA_ENABLED && ENABLE_SKY_LIGHT
 	if (bStrataMaterial)
 	{
 		// We use WorldNormal, Roughness and BaseColor (representing EnvBRDF) written by the strata top layer.
@@ -385,11 +384,10 @@ void ReflectionEnvironmentSkyLighting(
 			GetStrataPixelDataByteOffset(PixelPos, View.BufferSizeAndInvSize.x, Strata.MaxBytesPerPixel));
 	}
 	else
-#endif
 #endif
 	if (!bUnlitMaterial) // Only light pixels marked as lit
 	{
-		SkyLighting = SkyLightDiffuse(GBuffer, AmbientOcclusion, BufferUV, ScreenPosition, BentNormal, DiffuseColor);
+		SkyLighting.rgb = SkyLightDiffuse(GBuffer, AmbientOcclusion, BufferUV, ScreenPosition, BentNormal, DiffuseColor);
 	}
 	
 	BRANCH 
@@ -402,17 +400,21 @@ void ReflectionEnvironmentSkyLighting(
 		if (CloudSkyAOEnabled)
 		{
 			float OutOpticalDepth = 0.0f;
-			SkyLighting *= GetCloudVolumetricShadow(WorldPosition, CloudSkyAOWorldToLightClipMatrix, CloudSkyAOFarDepthKm, CloudSkyAOTexture, CloudSkyAOSampler, OutOpticalDepth);
+			SkyLighting.rgb *= GetCloudVolumetricShadow(WorldPosition, CloudSkyAOWorldToLightClipMatrix, CloudSkyAOFarDepthKm, CloudSkyAOTexture, CloudSkyAOSampler, OutOpticalDepth);
 		}
 #endif
 	}
 
 	// Scene color alpha is used for ScreenSpaceSubsurfaceScattering (if that is not needed it can be disabled with SUBSURFACE_CHANNEL_MODE)
+#if STRATA_ENABLED
+	OutColor = SkyLighting;
+#else
 	FLightAccumulator LightAccumulator = (FLightAccumulator)0;
 	const bool bNeedsSeparateSubsurfaceLightAccumulation = UseSubsurfaceProfile(ShadingModelID);
 
-	LightAccumulator_Add(LightAccumulator, SkyLighting, SkyLighting, 1.0f, bNeedsSeparateSubsurfaceLightAccumulation);
+	LightAccumulator_Add(LightAccumulator, SkyLighting.rgb, SkyLighting.rgb, 1.0f, bNeedsSeparateSubsurfaceLightAccumulation);
 	OutColor = LightAccumulator_GetResult(LightAccumulator);
+#endif
 
 #endif // ENABLE_DYNAMIC_SKY_LIGHT
 
diff --git a/Engine/Shaders/Private/Strata/Strata.ush b/Engine/Shaders/Private/Strata/Strata.ush
--- a/Engine/Shaders/Private/Strata/Strata.ush
+++ b/Engine/Shaders/Private/Strata/Strata.ush
@@ -759,6 +759,8 @@ struct FStrataEnvLightResult
 	float3 SpecularDirection;
 	float3 SpecularWeight;
 	float  SpecularSafeRoughness;
+
+	uint PathType;
 };
 
 FStrataEnvLightResult StrataEvaluateForEnvLight(FStrataBSDFContext BSDFContext, bool bEnableSpecular)
@@ -772,6 +774,7 @@ FStrataEnvLightResult StrataEvaluateForEnvLight(FStrataBSDFContext BSDFContext,
 		{
 			StrataEnvLightResult.DiffuseAlbedo = DIFFUSE_ALBEDO(BSDFContext.BSDF);
 			StrataEnvLightResult.DiffuseNormal = BSDFContext.N;
+			StrataEnvLightResult.PathType = STRATA_PATH_TYPE_DIFFUSE | (BSDF_GETHASSCATTERING(BSDFContext.BSDF) ? STRATA_PATH_TYPE_SCATTER : 0);
 			break;
 		}
 		case STRATA_BSDF_TYPE_DIELECTRIC:
@@ -805,6 +808,7 @@ FStrataEnvLightResult StrataEvaluateForEnvLight(FStrataBSDFContext BSDFContext,
 				StrataEnvLightResult.SpecularDirection = StrataGetOffSpecularPeakReflectionDir(BSDFContext.N, BSDFContext.R, SafeRoughness0);
 				StrataEnvLightResult.SpecularWeight = EvalEnvBRDF;
 				StrataEnvLightResult.SpecularSafeRoughness = SafeRoughness0;
+				StrataEnvLightResult.PathType = STRATA_PATH_TYPE_SPECULAR;
 			}
 			break;
 		}
@@ -826,12 +830,14 @@ FStrataEnvLightResult StrataEvaluateForEnvLight(FStrataBSDFContext BSDFContext,
 
 			StrataEnvLightResult.DiffuseAlbedo = NormLuminanceSingleScattering;
 			StrataEnvLightResult.DiffuseNormal = BSDFContext.N;
+			StrataEnvLightResult.PathType = STRATA_PATH_TYPE_DIFFUSE;
 			break;
 		}
 		case STRATA_BSDF_TYPE_SHEEN:
 		{
 			StrataEnvLightResult.DiffuseAlbedo = SHEEN_ALBEDO(BSDFContext.BSDF);;
 			StrataEnvLightResult.DiffuseNormal = BSDFContext.N;
+			StrataEnvLightResult.PathType = STRATA_PATH_TYPE_DIFFUSE;
 			break;
 		}
 	}
diff --git a/Engine/Shaders/Private/Strata/StrataDeferredLighting.ush b/Engine/Shaders/Private/Strata/StrataDeferredLighting.ush
--- a/Engine/Shaders/Private/Strata/StrataDeferredLighting.ush
+++ b/Engine/Shaders/Private/Strata/StrataDeferredLighting.ush
@@ -170,7 +170,7 @@ float4 StrataDeferredLighting(
 			const float3 SpecularRadiance = (BSDFEvaluate.PathType & STRATA_PATH_TYPE_SPECULAR) != 0 ? BSDF.Weight * BSDFEvaluate.Value : 0;
 
 			float3 CommonFactor = IntegrationContext.IncidentIrradiance;
-			const bool bNeedsSeparateSubsurfaceLightAccumulation = BSDFEvaluate.PathType & STRATA_PATH_TYPE_SCATTER != 0;
+			const bool bNeedsSeparateSubsurfaceLightAccumulation = (BSDFEvaluate.PathType & STRATA_PATH_TYPE_SCATTER) != 0;
 			LightAccumulator_AddSplit(Out, DiffuseRadiance, SpecularRadiance, DiffuseRadiance, CommonFactor, bNeedsSeparateSubsurfaceLightAccumulation);
 		}
 		else
diff --git a/Engine/Shaders/Private/Strata/StrataEnvironmentLighting.ush b/Engine/Shaders/Private/Strata/StrataEnvironmentLighting.ush
--- a/Engine/Shaders/Private/Strata/StrataEnvironmentLighting.ush
+++ b/Engine/Shaders/Private/Strata/StrataEnvironmentLighting.ush
@@ -4,7 +4,7 @@
 
 // Env. lighting evaluation for strata material.
 // Unpack BSDF on-the-fly
-float3 StrataEnvironmentLighting(
+float4 StrataEnvironmentLighting(
 	float2 BufferUV,
 	FGBufferData GBuffer,
 	float3 WorldPosition,
@@ -17,7 +17,13 @@ float3 StrataEnvironmentLighting(
 {
 	FStrataPixelHeader StrataPixelHeader = UnpackStrataHeaderIn(StrataDataBuffer, PixelStrataDataByteOffset);
 
-	float3 SkyLighting = 0;
+	FLightAccumulator Out = (FLightAccumulator)0;
+
+	float3 CommonFactor = 1;
+#if USE_PREEXPOSURE
+	CommonFactor *= View.PreExposure;
+#endif
+
 	for (uint BSDFIndex = 0; BSDFIndex < StrataPixelHeader.BSDFCount; ++BSDFIndex)
 	{
 		const FStrataBSDF BSDF = UnpackStrataBSDFIn(StrataDataBuffer, PixelStrataDataByteOffset);
@@ -32,6 +38,7 @@ float3 StrataEnvironmentLighting(
 		FStrataEnvLightResult StrataEnvLight = StrataEvaluateForEnvLight(StrataBSDFContext, bEnableSpecular);
 
 		// Diffuse component
+		float3 DiffuseSkyLighting = 0;
 		if (any(StrataEnvLight.DiffuseAlbedo > 0.0f))
 		{
 			// Apply bent normal if needed
@@ -47,25 +54,25 @@ float3 StrataEnvironmentLighting(
 			// Finally sample the sky diffuse contribution (spherical harmonic, Lambert BRDF)
 			float3 DiffuseLookup = GetSkySHDiffuse(StrataEnvLight.DiffuseNormal) * View.SkyLightColor.rgb;			// STRATA_TODO compute SH coefficients for ON or Chan
 			// And accumulate
-			SkyLighting += BSDFThroughput * (SkyVisData.SkyDiffuseLookUpMul * DiffuseLookup + SkyVisData.SkyDiffuseLookUpAdd) * StrataEnvLight.DiffuseAlbedo;
+			DiffuseSkyLighting = BSDFThroughput * (SkyVisData.SkyDiffuseLookUpMul * DiffuseLookup + SkyVisData.SkyDiffuseLookUpAdd) * StrataEnvLight.DiffuseAlbedo;
 		}
 
 		// Specular component
+		float3 SpecularSkyLighting = 0;
 		if (any(StrataEnvLight.SpecularWeight > 0.0f))
 		{
 			float SkyAverageBrightness = 1.0f;
 			const float3 SpecularLuminance = GetSkyLightReflection(StrataEnvLight.SpecularDirection, StrataEnvLight.SpecularSafeRoughness, SkyAverageBrightness);
 
-			SkyLighting += BSDFThroughput * SpecularLuminance * StrataEnvLight.SpecularWeight * View.SkyLightColor.rgb
+			SpecularSkyLighting = BSDFThroughput * SpecularLuminance * StrataEnvLight.SpecularWeight * View.SkyLightColor.rgb
 				* AmbientOcclusion * (BSDF_GETISTOPLAYER(BSDF) ? TopLayerSpecularContributionFactor : 1.0f);
 		}
 
-		// STRATA_TODO: Missing probe reflection here
+		const bool bNeedsSeparateSubsurfaceLightAccumulation = (StrataEnvLight.PathType & STRATA_PATH_TYPE_SCATTER) != 0;
+		LightAccumulator_AddSplit(Out, DiffuseSkyLighting, SpecularSkyLighting, DiffuseSkyLighting, CommonFactor, bNeedsSeparateSubsurfaceLightAccumulation);
 
+		// STRATA_TODO: Missing probe reflection here
 	}
-#if USE_PREEXPOSURE
-	SkyLighting *= View.PreExposure;
-#endif
-	return SkyLighting;
+	return LightAccumulator_GetResult(Out);
 }
 
