Improve error-distribution of Sobol sampling

Provide an alternative entry point to initialize random sequencss that takes the target number of samples into account as well as the pixel coordinates. By sharing a single sequence among pixels, shuffled along a space filling curve the visual distribution of error is reduced. Owen scrambling is still used to hide the structure of this space filling curve.

This is currently off by default, but can be enabled with the cvar: "r.PathTracing.UseErrorDiffusion=1"

diff --git a/Engine/Shaders/Private/PathTracing/PathTracing.usf b/Engine/Shaders/Private/PathTracing/PathTracing.usf
--- a/Engine/Shaders/Private/PathTracing/PathTracing.usf
+++ b/Engine/Shaders/Private/PathTracing/PathTracing.usf
@@ -48,7 +48,17 @@ RAY_TRACING_ENTRY_RAYGEN(PathTracingMainRG)
 	float2 BufferSize = View.BufferSizeAndInvSize.xy; 
 	uint LinearIndex = LaunchIndex.x + (LaunchIndex.y * BufferSize.x); 
 	RandomSequence RandSequence;
-	RandomSequence_Initialize(RandSequence, LinearIndex, PathTracingData.TemporalSeed);
+
+	if (PathTracingData.UseErrorDiffusion)
+	{
+		// z-sampler init
+		RandomSequence_Initialize(RandSequence, LaunchIndex, PathTracingData.Iteration, PathTracingData.TemporalSeed - PathTracingData.Iteration, PathTracingData.MaxSamples);
+	}
+	else
+	{
+		// random sobol init
+		RandomSequence_Initialize(RandSequence, LinearIndex, PathTracingData.TemporalSeed);
+	}
 
 	float3 Radiance = 0;
 	uint RayCounter = 0;
diff --git a/Engine/Shaders/Private/PathTracing/Utilities/PathTracingRandomSequence.ush b/Engine/Shaders/Private/PathTracing/Utilities/PathTracingRandomSequence.ush
--- a/Engine/Shaders/Private/PathTracing/Utilities/PathTracingRandomSequence.ush
+++ b/Engine/Shaders/Private/PathTracing/Utilities/PathTracingRandomSequence.ush
@@ -140,6 +140,18 @@ float Halton(uint Index, uint Base)
 
 
 uint FastOwenScrambling(uint Index, uint Seed) {
+#if 0
+	// reference implementation - scramble bits one at a time
+
+	// loop below does not require pre-inverted bits, so undo the inversion assumed in the optimized case below
+	Index = ReverseBits32(Index);
+	for (uint Mask = 1u << 31; Mask; Mask >>= 1) {
+		Seed = IntegerHash(Seed); // randomize state
+		Index ^= Seed & Mask;  // flip output (depending on state)
+		Seed ^= Index & Mask;  // flip state  (depending on output)
+	}
+	return Index;
+#else
 	// Laine and Karras / Stratified Sampling for Stochastic Transparency / EGSR 2011
 
 	// The operations below will mix bits toward the left, so temporarily reverse the order
@@ -169,6 +181,7 @@ uint FastOwenScrambling(uint Index, uint Seed) {
 	// Undo the reverse so that we get left-to-right scrambling
 	// thereby emulating owen-scrambling
 	return ReverseBits32(Index);
+#endif
 }
 
 // 32-bit Sobol matrices for dimension 1,2,3 from:
@@ -255,6 +268,64 @@ void RandomSequence_Initialize(inout RandomSequence RandSequence, uint PositionS
 #endif
 }
 
+#if RANDSEQ == RANDSEQ_OWENSOBOL
+
+// This is an alternative initialization to the method above when you know how many samples you are planning on taking.
+// The core idea is to use a single sobol pattern across the screen using a space filling curve to correlate nearby pixels
+// This achieves a much better error diffusion than the random sequence assignment of the method above
+// This algorithm is covered in the following paper:
+//   "Screen-Space Blue-Noise Diffusion of Monte Carlo Sampling Error via Hierarchical Ordering of Pixels"
+//   Siggraph Asia 2020
+//   http://abdallagafar.com/publications/zsampler/
+// Improvements made in this version are: 
+//   - switching the z-order curve for a hilbert curve
+//   - switching sobol points for owen-sobol points
+//   - recognizing that the hierarchical scrambling proposed in the original paper is equivalent to Owen-Scrambling of the index
+// Downsides of this method:
+//   - quality does not improve beyond the target sample count, making usage with adaptive sampling difficult
+//   - error-diffusion property is only achieved at the target sample count, earlier samples look random (sometimes slightly worse than random)
+void RandomSequence_Initialize(inout RandomSequence RandSequence, uint2 PixelCoord, uint SampleIndex, uint TimeSeed0, uint MaxSamples)
+{
+	uint X = PixelCoord.x;
+	uint Y = PixelCoord.y;
+
+#if 0
+	uint TileID = ((X & 255) + (Y & 255) * 256); // Scanline order
+#elif 0
+	uint TileID = ((Y & 1) ? 255 - (X & 255) : (X & 255)) + (Y & 255) * 256; // zig-zag curve
+#elif 0
+	uint TileID = MortonEncode(uint2(X, Y)); // Z-order curve
+#elif 0
+	// ordered dither curve
+	uint TileID = MortonEncode(uint2(X ^ Y, X));
+#elif 0
+    // pure random (breaks nice properties)
+	uint TileID = IntegerHash(X + IntegerHash(Y));
+#else
+	// Hilbert curve: [REDACTED_PROFILE_URL]/Hackers-Delight/blob/master/hilbert/hil_s_from_xy.c.txt
+	uint TileID = 0, HilbertState = 0;
+	for (int i = 7; i >= 0; i--) {
+		int xi = (X >> i) & 1;
+		int yi = (Y >> i) & 1;
+		uint Row = (4 * HilbertState) | (2 * xi) | yi;
+		TileID = (TileID << 2) | ((0x361E9CB4 >> 2 * Row) & 3);
+		HilbertState = (0x8FE65831 >> 2 * Row) & 3;
+	}
+#endif
+
+	// Combine pixel-level and sample-level bits into the sample index (visible structure will be hidden by owen scrambling of the index)
+	RandSequence.SampleIndex = ReverseBits32((65536 * TimeSeed0 + TileID) * MaxSamples + SampleIndex);
+
+	// progressive encodings (these kind of work, but quality is much worse)
+	//RandSequence.SampleIndex = ReverseBits32(MortonEncode(uint2(TimeSeed0, TileID)) * MaxSamples + SampleIndex);
+	//RandSequence.SampleIndex = ReverseBits32((MortonCode3(MaxSamples * TimeSeed0 + SampleIndex) * 4) | (MortonCode3(Y) * 2) | MortonCode3(X));
+	//RandSequence.SampleIndex = ReverseBits32(MortonEncode(uint2(MortonEncode(uint2(X, Y)), SampleIndex)));
+
+	RandSequence.SampleSeed = 0; // always use the same sequence
+}
+
+#endif
+
 float RandomSequence_GenerateSample1D(inout RandomSequence RandSequence)
 {
 	float Result;
diff --git a/Engine/Source/Runtime/Renderer/Private/PathTracing.cpp b/Engine/Source/Runtime/Renderer/Private/PathTracing.cpp
--- a/Engine/Source/Runtime/Renderer/Private/PathTracing.cpp
+++ b/Engine/Source/Runtime/Renderer/Private/PathTracing.cpp
@@ -29,6 +29,13 @@ TAutoConsoleVariable<int32> CVarPathTracingSamplesPerPixel(
 	ECVF_RenderThreadSafe
 );
 
+TAutoConsoleVariable<int32> CVarPathTracingUseErrorDiffusion(
+	TEXT("r.PathTracing.UseErrorDiffusion"),
+	0,
+	TEXT("Enables an experimental sampler that diffuses visible error in screen space. This generally produces better results when the target sample count can be reached. (default = 0 (disabled))"),
+	ECVF_RenderThreadSafe
+);
+
 TAutoConsoleVariable<int32> CVarPathTracingMISMode(
 	TEXT("r.PathTracing.MISMode"),
 	2,
@@ -101,6 +108,7 @@ static bool PrepareShaderArgs(const FViewInfo& View, FPathTracingData& PathTraci
 	PathTracingData.MISMode = CVarPathTracingMISMode.GetValueOnRenderThread();
 	PathTracingData.VisibleLights = CVarPathTracingVisibleLights.GetValueOnRenderThread();
 	PathTracingData.MaxPathIntensity = CVarPathTracingMaxPathIntensity.GetValueOnRenderThread();
+	PathTracingData.UseErrorDiffusion = CVarPathTracingUseErrorDiffusion.GetValueOnRenderThread();
 
 	bool NeedInvalidation = false;
 
@@ -138,6 +146,14 @@ static bool PrepareShaderArgs(const FViewInfo& View, FPathTracingData& PathTraci
 		PreviousMaxPathIntensity = PathTracingData.MaxPathIntensity;
 	}
 
+	// Changing sampler requires starting over
+	static uint32 PreviousUseErrorDiffusion = PathTracingData.UseErrorDiffusion;
+	if (PreviousUseErrorDiffusion != PathTracingData.UseErrorDiffusion)
+	{
+		NeedInvalidation = true;
+		PreviousUseErrorDiffusion = PathTracingData.UseErrorDiffusion;
+	}
+
 	// the rest of PathTracingData and AdaptiveSamplingData is filled in by SetParameters below
 	return NeedInvalidation;
 }
@@ -348,6 +364,7 @@ void FDeferredShadingSceneRenderer::RenderPathTracing(
 		PathTracingData.TemporalSeed = View.ViewState->PathTracingFrameIndependentTemporalSeed;
 	}
 	PathTracingData.Iteration = View.ViewState->PathTracingSPP;
+	PathTracingData.MaxSamples = MaxSPP;
 
 	// Prepare radiance buffer (will be shared with display pass)
 	FRDGTexture* RadianceTexture = nullptr;
diff --git a/Engine/Source/Runtime/Renderer/Private/PathTracingUniformBuffers.h b/Engine/Source/Runtime/Renderer/Private/PathTracingUniformBuffers.h
--- a/Engine/Source/Runtime/Renderer/Private/PathTracingUniformBuffers.h
+++ b/Engine/Source/Runtime/Renderer/Private/PathTracingUniformBuffers.h
@@ -12,6 +12,8 @@
 BEGIN_GLOBAL_SHADER_PARAMETER_STRUCT(FPathTracingData, )
 	SHADER_PARAMETER(uint32, Iteration)
 	SHADER_PARAMETER(uint32, TemporalSeed)
+	SHADER_PARAMETER(uint32, MaxSamples)
+	SHADER_PARAMETER(uint32, UseErrorDiffusion)
 	SHADER_PARAMETER(uint32, MaxBounces)
 	SHADER_PARAMETER(uint32, MISMode)
 	SHADER_PARAMETER(uint32, VisibleLights)
