Adjust shaders for DXC update. We can't support <float3> in StructuredBuffers for the Vulkan backend atm.

#rnx

diff --git a/Engine/Shaders/Private/Halton.ush b/Engine/Shaders/Private/Halton.ush
--- a/Engine/Shaders/Private/Halton.ush
+++ b/Engine/Shaders/Private/Halton.ush
@@ -8,7 +8,7 @@ Halton.ush: Halton sampling interface and utilities
 
 uint3 HaltonPrime3(uint DimensionDividedByThree)
 {
-	return HaltonPrimes.Primes[DimensionDividedByThree % HaltonPrimes.Dimensions];
+	return HaltonPrimes.Primes[DimensionDividedByThree % HaltonPrimes.Dimensions].xyz;
 }
 
 uint HaltonPrime(uint Dimension)
@@ -47,6 +47,6 @@ float3 HaltonSequence3(uint SequenceIndex, int IterationIndex, int DimensionInde
 	int SequenceOffset = SequenceIndex * HaltonIteration.IterationCount * HaltonIteration.Dimensions;
 	int IterationOffset = IterationIndex * HaltonIteration.Dimensions;
 	int SampleOffset = SequenceOffset + IterationOffset + DimensionIndex;
-	float3 RandSample = HaltonIteration.SequenceIteration[SampleOffset];
+	float3 RandSample = HaltonIteration.SequenceIteration[SampleOffset].xyz;
 	return RandSample;
 }
\ No newline at end of file
diff --git a/Engine/Source/Runtime/Renderer/Private/Halton.cpp b/Engine/Source/Runtime/Renderer/Private/Halton.cpp
--- a/Engine/Source/Runtime/Renderer/Private/Halton.cpp
+++ b/Engine/Source/Runtime/Renderer/Private/Halton.cpp
@@ -47,6 +47,7 @@ void FHaltonSequenceIteration::InitRHI()
 	struct FSequenceIterationData
 	{
 		float Sequence[3];
+		float Padding;
 	};
 
 	SCOPED_GPU_STAT(FRHICommandListExecutor::GetImmediateCommandList(), HaltonSequence);
@@ -158,6 +159,7 @@ void FHaltonPrimesResource::InitRHI()
 		struct FPrimeData
 		{
 			int Primes[3];
+			int Padding;
 		};
 
 		CreateInfo.DebugName = TEXT("HaltonPrimes");
diff --git a/Engine/Source/Runtime/Renderer/Private/Halton.h b/Engine/Source/Runtime/Renderer/Private/Halton.h
--- a/Engine/Source/Runtime/Renderer/Private/Halton.h
+++ b/Engine/Source/Runtime/Renderer/Private/Halton.h
@@ -15,7 +15,7 @@ SHADER_PARAMETER(int, Dimensions)
 SHADER_PARAMETER(int, SequenceRowCount)
 SHADER_PARAMETER(int, SequenceColumnCount)
 SHADER_PARAMETER(int, IterationCount)
-SHADER_PARAMETER_SRV(StructuredBuffer<float3>, SequenceIteration)
+SHADER_PARAMETER_SRV(StructuredBuffer<float4>, SequenceIteration)
 END_GLOBAL_SHADER_PARAMETER_STRUCT()
 
 /**
@@ -92,7 +92,7 @@ inline void InitializeHaltonSequenceIteration(const FHaltonSequenceIteration& Ha
 
 BEGIN_GLOBAL_SHADER_PARAMETER_STRUCT(FHaltonPrimes, )
 SHADER_PARAMETER(int, Dimensions)
-SHADER_PARAMETER_SRV(StructuredBuffer<uint3>, Primes)
+SHADER_PARAMETER_SRV(StructuredBuffer<uint4>, Primes)
 END_GLOBAL_SHADER_PARAMETER_STRUCT()
 
 /**
