Change sphere light shadows to use robust sampling method and change to concentric sampling to provide a more continuous gradient in primary sample space

This avoids artifacts when running with the error diffusion sampler

diff --git a/Engine/Shaders/Private/RayTracing/RayTracingSphereLight.ush b/Engine/Shaders/Private/RayTracing/RayTracingSphereLight.ush
--- a/Engine/Shaders/Private/RayTracing/RayTracingSphereLight.ush
+++ b/Engine/Shaders/Private/RayTracing/RayTracingSphereLight.ush
@@ -65,27 +65,27 @@ bool GenerateSphereLightOcclusionRayWithSolidAngleSampling(
 
 	// Sample uniformly about a cone aligned with the z-axis
 	float SinThetaMax2 = Radius2 / RayLength2;
-	float CosThetaMax2 = 1.0 - SinThetaMax2;
-	float CosThetaMax = sqrt(CosThetaMax2);
 
-	// NOTE: Same as UniformSampleCone(), but CosTheta and SinTheta are needed
-	// float3 RayDirectionTangent = UniformSampleCone(RandSample, CosThetaMax).xyz;
-	// RayPdf = UniformSampleCone(RandSample, CosThetaMax).w;
-	float CosTheta = lerp(1.0, CosThetaMax, RandSample.x);
+#if 0
+	float4 DirAndPdf = UniformSampleCone(RandSample, sqrt(1.0 - SinThetaMax2));
+#elif 0
+	float4 DirAndPdf = UniformSampleConeRobust(RandSample, SinThetaMax2);
+#else
+	float4 DirAndPdf = UniformSampleConeConcentricRobust(RandSample, SinThetaMax2);
+#endif
+
+	float CosTheta = DirAndPdf.z;
 	float SinTheta2 = 1.0 - CosTheta * CosTheta;
-	float SinTheta = sqrt(SinTheta2);
-	float Phi = RandSample.y * 2.0 * PI;
-	float3 RayDirectionTangent = float3(cos(Phi) * SinTheta, sin(Phi) * SinTheta, CosTheta);
 
 	RayOrigin = WorldPosition;
 	// Project ray direction to world-space, such that z-axis aligns with LightDirection
 	float RayLength = sqrt(RayLength2);
 	LightDirection *= rcp(RayLength + 1e-4);
-	RayDirection = TangentToWorld(RayDirectionTangent, LightDirection);
-	
+	RayDirection = TangentToWorld(DirAndPdf.xyz, LightDirection);
+
 	RayTMin = 0.0;
 	// Clip length to closest intersection with the sphere
 	RayTMax = RayLength * (CosTheta - sqrt(max(SinThetaMax2 - SinTheta2, 0.0)));
-	RayPdf = 1.0 / (2.0 * PI * (1 - CosThetaMax));
+	RayPdf = DirAndPdf.w;
 	return true;
 }
\ No newline at end of file
