Fix usage of the RandomSequence API in the ryatraced shadow occlusion pass and add support for error-diffusing sampler initialization (on by default)

The previous code was re-using a single RandomsSequence without considering the sample index in the multi-sample case. This meant the equivalent of random sampling was being used.

The ability to revert to the regular sampling init (per pixel randomization) is preserved via ifdefs to make it easy to compare methods in the future if needed.

Address TODO item in the hair strand ray tracing code (use a 4D sample instead of 3D+1D sampling)

diff --git a/Engine/Shaders/Private/HairStrands/HairStrandsRaytracing.ush b/Engine/Shaders/Private/HairStrands/HairStrandsRaytracing.ush
--- a/Engine/Shaders/Private/HairStrands/HairStrandsRaytracing.ush
+++ b/Engine/Shaders/Private/HairStrands/HairStrandsRaytracing.ush
@@ -89,9 +89,9 @@ float TraverseHair(uint2 PixelCoord, inout RandomSequence RandSequence, float3 R
 {
 	float HitT = InMinT;
 
-	// TODO: use 4D sample instead?
-	const float3 VoxelRandom = RandomSequence_GenerateSample3D(RandSequence);
-	const float RayRandom 	 = RandomSequence_GenerateSample1D(RandSequence);
+	const float4 Rnd = RandomSequence_GenerateSample4D(RandSequence);
+	const float3 VoxelRandom = Rnd.xyz;
+	const float RayRandom = Rnd.w;
 
 	return InternalTraverseHair(PixelCoord, RayOrigin, RayDirection, InMinT, InOcclusionThreshold, VoxelRandom, RayRandom);
 }
\ No newline at end of file
diff --git a/Engine/Shaders/Private/RayTracing/RayTracingOcclusionRGS.usf b/Engine/Shaders/Private/RayTracing/RayTracingOcclusionRGS.usf
--- a/Engine/Shaders/Private/RayTracing/RayTracingOcclusionRGS.usf
+++ b/Engine/Shaders/Private/RayTracing/RayTracingOcclusionRGS.usf
@@ -300,7 +300,6 @@ RayDesc Invert(in RayDesc Ray)
 
 FOcclusionResult ComputeOcclusion(
 	const uint2 PixelCoord,
-	RandomSequence RandSequence,
 	const uint ShadingModelID,
 	const uint RaytracingMask,
 	const float DeviceZ,
@@ -347,21 +346,38 @@ FOcclusionResult ComputeOcclusion(
 		bApplyNormalCulling = false;
 	}
 
+	uint TimeSeed = View.StateFrameIndex;
+
 #if ENABLE_MULTIPLE_SAMPLES_PER_PIXEL
 	LOOP for (uint SampleIndex = 0; SampleIndex < LocalSamplesPerPixel; ++SampleIndex)
 #else // ENABLE_MULTIPLE_SAMPLES_PER_PIXEL
 	do if (LocalSamplesPerPixel > 0)
 #endif // ENABLE_MULTIPLE_SAMPLES_PER_PIXEL
 	{
-		float2 RandSample = RandomSequence_GenerateSample2D(RandSequence);
+		RandomSequence RandSequence;
 
-		#if !ENABLE_MULTIPLE_SAMPLES_PER_PIXEL
-		{
-			uint2 RandomSeed = Rand3DPCG16(int3(PixelCoord, View.StateFrameIndex % 17)).xy;
-			
-			RandSample = Hammersley16(0, 1, RandomSeed);
-		}
-		#endif // !ENABLE_MULTIPLE_SAMPLES_PER_PIXEL
+// This initiailization improves the visual distribution of the error in screenspace
+#define USE_ERROR_DIFFUSION_SAMPLER 1 && RANDSEQ == RANDSEQ_OWENSOBOL
+
+#if USE_ERROR_DIFFUSION_SAMPLER
+
+#if ENABLE_MULTIPLE_SAMPLES_PER_PIXEL
+		RandomSequence_Initialize(RandSequence, PixelCoord, SampleIndex, TimeSeed, LocalSamplesPerPixel);
+#else
+		RandomSequence_Initialize(RandSequence, PixelCoord, 0, TimeSeed, 1);
+#endif
+
+#else
+
+#if ENABLE_MULTIPLE_SAMPLES_PER_PIXEL
+		RandomSequence_Initialize(RandSequence, CalcLinearIndex(PixelCoord), LocalSamplesPerPixel * TimeSeed + SampleIndex);
+#else
+		RandomSequence_Initialize(RandSequence, CalcLinearIndex(PixelCoord), TimeSeed);
+#endif
+
+#endif
+
+		float2 RandSample = RandomSequence_GenerateSample2D(RandSequence);
 
 		RayDesc Ray;
 		bool bIsValidRay = GenerateOcclusionRay(
@@ -498,6 +514,13 @@ FOcclusionResult ComputeOcclusion(
 
 	if (ENABLE_TRANSMISSION && LocalSamplesPerPixel > 0 && ShadingModelID == SHADINGMODELID_SUBSURFACE_PROFILE)
 	{
+		RandomSequence RandSequence;
+#if USE_ERROR_DIFFUSION_SAMPLER
+		RandomSequence_Initialize(RandSequence, PixelCoord, 0, TimeSeed, 1);
+#else
+		RandomSequence_Initialize(RandSequence, CalcLinearIndex(PixelCoord), TimeSeed);
+#endif
+
 		// Snag two variables to fix light sampling for all computations..
 		const float2 LightSample = RandomSequence_GenerateSample2D(RandSequence);
 		const float3 V = normalize(WorldPosition - View.WorldViewOrigin);
@@ -595,7 +618,6 @@ RAY_TRACING_ENTRY_RAYGEN(OcclusionRGS)
 
 		Occlusion = ComputeOcclusion(
 			PixelCoord,
-			RandSequence,
 			ShadingModelID,
 			RAY_TRACING_MASK_SHADOW | RAY_TRACING_MASK_THIN_SHADOW,
 			DeviceZ,
@@ -609,7 +631,6 @@ RAY_TRACING_ENTRY_RAYGEN(OcclusionRGS)
 		{
 			HairOcclusion = ComputeOcclusion(
 				PixelCoord,
-				RandSequence,
 				SHADINGMODELID_HAIR,
 				RAY_TRACING_MASK_SHADOW,
 				HairDeviceZ,
