Fix sky/GI contribution when SSS is not enabled on diffuse (radius scale set to 0).

diff --git a/Engine/Shaders/Private/BasePassPixelShader.usf b/Engine/Shaders/Private/BasePassPixelShader.usf
--- a/Engine/Shaders/Private/BasePassPixelShader.usf
+++ b/Engine/Shaders/Private/BasePassPixelShader.usf
@@ -1457,7 +1457,7 @@ void FPixelShaderInOut_MainPS(
 	}
 #endif // MATERIAL_SHADINGMODEL_THIN_TRANSLUCENT
 
-
+	bool bSubsurfaceEnable = true;
 #if STRATA_ENABLED && !MATERIAL_SHADINGMODEL_UNLIT
 
 	FStrataData Strata = PixelMaterialInputs.FrontMaterial;
@@ -1493,6 +1493,9 @@ void FPixelShaderInOut_MainPS(
 		OpaqueBasePass.Strata.MaterialLobesBufferUAV, PixelStrataDataByteOffset, 
 		StrataPixelHeader, Strata,
 		VertexNormal, MaterialParameters.CameraVector, GBuffer);
+
+	// Since we don't carry SSS on/off through the ShadingModel value (DefaultLit/SubSurfaceProfile), we update subsurface toggle, based on the MFP radius value
+	bSubsurfaceEnable = GBuffer.CustomData.a > 0;
 #endif // STRATA_OPAQUE_ENABLED
 
 #if STRATA_TRANSLUCENT_ENABLED
@@ -1594,7 +1597,7 @@ void FPixelShaderInOut_MainPS(
 			// Apply vertex fog
 			Color = Color * Fogging.a + Fogging.rgb;
 
-#if POST_PROCESS_SUBSURFACE
+		#if POST_PROCESS_SUBSURFACE
 			// Apply vertex fog to diffuse color
 			DiffuseColor = DiffuseColor * Fogging.a + Fogging.rgb;
 
@@ -1605,10 +1608,10 @@ void FPixelShaderInOut_MainPS(
 				// to the specular component, otherwise lighting is applied twice
 				Color *= !bChecker;
 			}
-			LightAccumulator_Add(LightAccumulator, Color + DiffuseColor, DiffuseColor, 1.0f, UseSubsurfaceProfile(GBuffer.ShadingModelID));
-#else
+			LightAccumulator_Add(LightAccumulator, Color + DiffuseColor, DiffuseColor, 1.0f, UseSubsurfaceProfile(GBuffer.ShadingModelID) && bSubsurfaceEnable);
+		#else
 			LightAccumulator_Add(LightAccumulator, Color, 0, 1.0f, false);
-#endif
+		#endif
 			Out.MRT[0] = RETURN_COLOR(LightAccumulator_GetResult(LightAccumulator));
 
 			#if !USES_GBUFFER
diff --git a/Engine/Shaders/Private/Strata/Strata.ush b/Engine/Shaders/Private/Strata/Strata.ush
--- a/Engine/Shaders/Private/Strata/Strata.ush
+++ b/Engine/Shaders/Private/Strata/Strata.ush
@@ -1507,13 +1507,6 @@ void PackStrataOut(
 			GBuffer.Metallic = 1;
 		}
 
-		if (bHasSubsurfaceScattering)
-		{
-			GBuffer.BaseColor      = SubsurfaceBaseColor;
-			GBuffer.CustomData.rgb = EncodeSubsurfaceProfile(SubsurfaceProfile);
-			GBuffer.CustomData.a   = SubsurfaceRadiusScale;
-		}
-
 		// Take into account the non covered part for this layer and update throughput accordingly
 		LayerThroughput += (1.0 - saturate(LayerCoverage)) * FullThroughput;
 
@@ -1521,6 +1514,15 @@ void PackStrataOut(
 		Throughput *= LayerThroughput;
 	}
 
+	// Patch Gbuffer data with SSS data if enabled
+	GBuffer.CustomData = 0;
+	if (bHasSubsurfaceScattering)
+	{
+		GBuffer.BaseColor		= SubsurfaceBaseColor;
+		GBuffer.CustomData.rgb	= EncodeSubsurfaceProfile(SubsurfaceProfile);
+		GBuffer.CustomData.a	= SubsurfaceRadiusScale;
+	}
+
 	// Now write out Strata data
 
 	// 1. the header
