Fix shader warning in MiniFont shader.

diff --git a/Engine/Shaders/Private/MiniFontCommon.ush b/Engine/Shaders/Private/MiniFontCommon.ush
--- a/Engine/Shaders/Private/MiniFontCommon.ush
+++ b/Engine/Shaders/Private/MiniFontCommon.ush
@@ -131,11 +131,11 @@ void PrintFloatNoFraction(int2 PixelPos, inout float3 OutColor, float3 FontColor
 	int2 Cursor = LeftTop;
 
 	uint DigitValue = 1;
-	for (int i = 0; i < DigitCount-1; ++i)
+	for (uint i = 0; i < DigitCount-1; ++i)
 	{
 		DigitValue *= 10;
 	}
-	for (int i = 0; i < DigitCount; ++i)
+	for (uint j = 0; j < DigitCount; ++j)
 	{
 		PrintCharacter(PixelPos, OutColor, FontColor, Cursor, ExtractDigitFromFloat(Number, DigitValue));
 		DigitValue /= 10;
