Fix a source of NaNs in the path tracer for two-sided foliage materials with a black subsurface color

diff --git a/Engine/Shaders/Private/PathTracing/Material/PathTracingTwoSidedFoliage.ush b/Engine/Shaders/Private/PathTracing/Material/PathTracingTwoSidedFoliage.ush
--- a/Engine/Shaders/Private/PathTracing/Material/PathTracingTwoSidedFoliage.ush
+++ b/Engine/Shaders/Private/PathTracing/Material/PathTracingTwoSidedFoliage.ush
@@ -75,7 +75,7 @@ void TwoSidedFoliage_EvalMaterial(
 		float3 SubsurfaceColor = GetSubsurfaceColor(Payload);
 
 		float LambertPdf = -NoL / PI;
-		OutWeight = SubsurfaceColor / TransmissionSelectionPdf;
+		OutWeight = TransmissionSelectionPdf > 0 ? SubsurfaceColor / TransmissionSelectionPdf : 0;
 		OutPdf = TransmissionSelectionPdf * LambertPdf;
 	}
 	else
