Add per-pixel subsurface support for Strata.

diff --git a/Engine/Shaders/Private/BurleyNormalizedSSSCommon.ush b/Engine/Shaders/Private/BurleyNormalizedSSSCommon.ush
--- a/Engine/Shaders/Private/BurleyNormalizedSSSCommon.ush
+++ b/Engine/Shaders/Private/BurleyNormalizedSSSCommon.ush
@@ -45,10 +45,15 @@ bool GetSubsurfaceTransmittanceProfileUseBurley(uint SubsurfaceProfileInt)
 
 bool GetSubsurfaceProfileUseBurley(uint SubsurfaceProfileInt)
 {
+#if STRATA_ENABLED
+	// STRATA_TODO: For now force Burley to be used by default
+	return true;
+#else
 	// 0..255, which SubSurface profile to pick
 	float Type = ActualSSProfilesTexture.Load(int3(SSSS_BOUNDARY_COLOR_BLEED_OFFSET, SubsurfaceProfileInt, 0)).a;
 
 	return abs(Type - SSS_TYPE_BURLEY) < 0.01f;
+#endif
 }
 
 bool GetSubsurfaceProfileUseBurley(float2 BufferUV)
@@ -60,10 +65,15 @@ bool GetSubsurfaceProfileUseBurley(float2 BufferUV)
 
 bool GetSubsurfaceProfileUseSeparable(uint SubsurfaceProfileInt)
 {
+#if STRATA_ENABLED
+	// STRATA_TODO: For now force Burley to be used by default
+	return false;
+#else
 	// 0..255, which SubSurface profile to pick
 	float Type = ActualSSProfilesTexture.Load(int3(SSSS_BOUNDARY_COLOR_BLEED_OFFSET, SubsurfaceProfileInt, 0)).a;
 
 	return abs(Type - SSS_TYPE_SSSS) < 0.01f;
+#endif
 }
 
 bool GetSubsurfaceProfileUseSeparable(float2 BufferUV)
diff --git a/Engine/Shaders/Private/DeferredShadingCommon.ush b/Engine/Shaders/Private/DeferredShadingCommon.ush
--- a/Engine/Shaders/Private/DeferredShadingCommon.ush
+++ b/Engine/Shaders/Private/DeferredShadingCommon.ush
@@ -716,7 +716,12 @@ float3 ExtractSubsurfaceColor(FGBufferData BufferData)
 uint ExtractSubsurfaceProfileInt(FGBufferData BufferData)
 {
 	// can be optimized
+#if STRATA_ENABLED
+	const bool bUseProfileData = BufferData.CustomData.y == 0 && BufferData.CustomData.z == 0;
+	return bUseProfileData ? uint(BufferData.CustomData.r * 255.0f + 0.5f) : SSS_PROFILE_ID_PERPIXEL;
+#else
 	return uint(BufferData.CustomData.r * 255.0f + 0.5f);
+#endif
 }
 
 #if SHADING_PATH_DEFERRED
diff --git a/Engine/Shaders/Private/PostProcessSubsurface.usf b/Engine/Shaders/Private/PostProcessSubsurface.usf
--- a/Engine/Shaders/Private/PostProcessSubsurface.usf
+++ b/Engine/Shaders/Private/PostProcessSubsurface.usf
@@ -5,6 +5,11 @@
 =============================================================================*/
 
 #include "Common.ush"
+
+#ifndef STRATA_ENABLED
+#define STRATA_ENABLED PROJECT_STRATA
+#endif
+
 #include "PostProcessCommon.ush"
 #include "DeferredShadingCommon.ush"
 #include "SubsurfaceProfileCommon.ush"
@@ -109,6 +114,9 @@ float GetSubsurfaceStrength(float2 UV)
 // @return .rgb is the weight for color channel, .a is the sample location
 float4 GetSubsurfaceProfileKernel(uint SampleIndex, uint SubsurfaceProfileInt)
 {
+#if STRATA_ENABLED
+	if (SubsurfaceProfileInt == SSS_PROFILE_ID_PERPIXEL) return float4(1, 1, 1, 1);
+#endif
 	const float4 TableMax = float4(1, 1, 1, SUBSURFACE_KERNEL_SIZE);
 
 	return ActualSSProfilesTexture.Load(int3(SampleIndex, SubsurfaceProfileInt, 0)) * TableMax;
@@ -118,7 +126,9 @@ float GetSubsurfaceProfileRadiusScale(FGBufferData GBufferData)
 {
 	// 0..255, which SubSurface profile to pick
 	uint SubsurfaceProfileInt = ExtractSubsurfaceProfileInt(GBufferData);
-
+#if STRATA_ENABLED
+	if (SubsurfaceProfileInt == SSS_PROFILE_ID_PERPIXEL) return 1.f;
+#endif
 	return GetSubsurfaceProfileKernel(SSSS_N_KERNELWEIGHTOFFSET + SSSS_N_KERNELWEIGHTCOUNT - 1, SubsurfaceProfileInt).a;
 }
 
@@ -126,7 +136,9 @@ float3 GetSubsurfaceProfileColor(FGBufferData GBufferData)
 {
 	// 0..255, which SubSurface profile to pick
 	uint SubsurfaceProfileInt = ExtractSubsurfaceProfileInt(GBufferData);
-
+#if STRATA_ENABLED
+	if (SubsurfaceProfileInt == SSS_PROFILE_ID_PERPIXEL) return GBufferData.BaseColor;
+#endif
 	return GetSubsurfaceProfileKernel(SSSS_SUBSURFACE_COLOR_OFFSET, SubsurfaceProfileInt).rgb;
 }
 
@@ -152,7 +164,9 @@ float3 GetSubsurfaceProfileBoundaryColorBleed(FGBufferData GBufferData)
 {
 	// 0..255, which SubSurface profile to pick
 	uint SubsurfaceProfileInt = ExtractSubsurfaceProfileInt(GBufferData);
-
+#if STRATA_ENABLED
+	if (SubsurfaceProfileInt == SSS_PROFILE_ID_PERPIXEL) return float3(1.f, 1.f, 1.f);
+#endif
 	return GetSubsurfaceProfileKernel(SSSS_BOUNDARY_COLOR_BLEED_OFFSET, SubsurfaceProfileInt).rgb;
 }
 
@@ -738,7 +752,16 @@ void SubsurfaceRecombinePS(noperspective float4 UVAndScreenPos : TEXCOORD0, floa
 	// combine potentially half res with full res
 	float3 SubsurfaceLighting = lerp(DiffuseAndSpecular.Diffuse, SSSColor, FadedSubsurfaceColor);
 
-	OutColor = float4(SubsurfaceLighting * StoredBaseColor + ExtractedNonSubsurface, 0);
+#if STRATA_ENABLED
+	if (ProfileId == SSS_PROFILE_ID_PERPIXEL)
+	{
+		OutColor = float4(SubsurfaceLighting + ExtractedNonSubsurface, 0);
+	}
+	else
+#endif
+	{
+		OutColor = float4(SubsurfaceLighting * StoredBaseColor + ExtractedNonSubsurface, 0);
+	}
 }
 
 void SubsurfaceViewportCopyPS(noperspective float4 InUV : TEXCOORD0, out float4 OutColor : SV_Target0)
diff --git a/Engine/Shaders/Private/ShadingCommon.ush b/Engine/Shaders/Private/ShadingCommon.ush
--- a/Engine/Shaders/Private/ShadingCommon.ush
+++ b/Engine/Shaders/Private/ShadingCommon.ush
@@ -37,6 +37,13 @@
 #define HAIR_COMPONENT_MULTISCATTER	0x20u
 #define HAIR_COMPONENT_TT_MODEL  	0x40u
 
+#define SSS_PROFILE_ID_INVALID  256
+#define SSS_PROFILE_ID_PERPIXEL 512
+
+// Max radius for per-pixel sub-surface scattering (non-profile version) in cm
+// Default value give a precision of 0.1mm
+#define SSS_PROFILE_PERPIXEL_MAXRADIUS 25.5f
+
 // for debugging and to visualize
 float3 GetShadingModelColor(uint ShadingModelID)
 {
diff --git a/Engine/Shaders/Private/Strata/Strata.ush b/Engine/Shaders/Private/Strata/Strata.ush
--- a/Engine/Shaders/Private/Strata/Strata.ush
+++ b/Engine/Shaders/Private/Strata/Strata.ush
@@ -45,6 +45,7 @@ struct FStrataBSDF
 #define STATE_BITS_ANISOTROPY			1
 #define STATE_BITS_TOPLAYER				1
 #define STATE_BITS_HASSCATTERING		1
+#define STATE_BITS_HASSSSPROFILE		1
 #define STATE_BITS_GREYWEIGHT			1
 #define STATE_BITS___UNUSED___			13
 #define STATE_BITS_WEIGHT10F			10
@@ -55,7 +56,8 @@ struct FStrataBSDF
 #define STATE_BIT_OFFSET_ANISOTROPY		(STATE_BIT_OFFSET_SHAREDNORMALID	+ STATE_BITS_SHAREDNORMALID)
 #define STATE_BIT_OFFSET_TOPLAYER		(STATE_BIT_OFFSET_ANISOTROPY		+ STATE_BITS_ANISOTROPY)
 #define STATE_BIT_OFFSET_HASSCATTERING	(STATE_BIT_OFFSET_TOPLAYER			+ STATE_BITS_TOPLAYER)
-#define STATE_BIT_OFFSET_GREYWEIGHT		(STATE_BIT_OFFSET_HASSCATTERING		+ STATE_BITS_HASSCATTERING)
+#define STATE_BIT_OFFSET_HASSSSPROFILE	(STATE_BIT_OFFSET_HASSCATTERING		+ STATE_BITS_HASSCATTERING)
+#define STATE_BIT_OFFSET_GREYWEIGHT		(STATE_BIT_OFFSET_HASSSSPROFILE		+ STATE_BITS_HASSSSPROFILE)
 #define STATE_BIT_OFFSET_WEIGHT10F		22
 
 #define TO_BITMASK(BitCount)							((1<<BitCount)-1)
@@ -80,10 +82,14 @@ struct FStrataBSDF
 #define BSDF_GETISTOPLAYER(X)			BSDF_GETSTATE(X, STATE_BITS_TOPLAYER, STATE_BIT_OFFSET_TOPLAYER)
 #define BSDF_SETISTOPLAYER(X, IsTop)	BSDF_SETSTATE(X, STATE_BITS_TOPLAYER, STATE_BIT_OFFSET_TOPLAYER, IsTop)
 
-// Indicates if the BSDF is part of the top layer
+// Indicates if the BSDF has scattering component
 #define BSDF_GETHASSCATTERING(X)		BSDF_GETSTATE(X, STATE_BITS_HASSCATTERING, STATE_BIT_OFFSET_HASSCATTERING)
 #define BSDF_SETHASSCATTERING(X, Scatt)	BSDF_SETSTATE(X, STATE_BITS_HASSCATTERING, STATE_BIT_OFFSET_HASSCATTERING, Scatt)
 
+// Indicates if the BSDF has a SSS profile
+#define BSDF_GETHASSSSPROFILE(X)		BSDF_GETSTATE(X, STATE_BITS_HASSSSPROFILE, STATE_BIT_OFFSET_HASSSSPROFILE)
+#define BSDF_SETHASSSSPROFILE(X, Profi)	BSDF_SETSTATE(X, STATE_BITS_HASSSSPROFILE, STATE_BIT_OFFSET_HASSSSPROFILE, Profi)
+
 // Indicates if the BSDF weight is grey scale so that it can be stored in a single float (e.g. top layer BSDFs)
 #define BSDF_GETGREYWEIGHT(X)			BSDF_GETSTATE(X, STATE_BITS_GREYWEIGHT, STATE_BIT_OFFSET_GREYWEIGHT)
 #define BSDF_SETGREYWEIGHT(X, Grey)		BSDF_SETSTATE(X, STATE_BITS_GREYWEIGHT, STATE_BIT_OFFSET_GREYWEIGHT, Grey)
@@ -102,8 +108,12 @@ struct FStrataBSDF
 
 #define DIFFUSE_ALBEDO(X)				X.VGPRs[0].xyz
 #define DIFFUSE_ROUGHNESS(X)			X.VGPRs[0].w
+// Either with SSS profile
 #define DIFFUSE_SSSPROFILEID(X)			X.VGPRs[1].x
 #define DIFFUSE_SSSPROFILERADIUS(X)		X.VGPRs[1].y
+// Or explicit DMFP
+#define DIFFUSE_SSSDMFPALBEDO(X)		X.VGPRs[1].xyz
+#define DIFFUSE_SSSDMFPRADIUS(X)		X.VGPRs[1].w
 
 #define DIELECTRIC_TINT(X)				X.VGPRs[0].xyz
 #define DIELECTRIC_IOR(X)				X.VGPRs[0].w
@@ -182,7 +192,7 @@ FStrataData GetInitialisedStrataData()
 	return StrataData;
 }
 
-FStrataData GetStrataDiffuseBSDF(float3 Albedo, float Roughness, float SSSProfileID, float SSSRadiusScale, uint SharedNormalIndex)
+FStrataData GetStrataDiffuseBSDF(float3 Albedo, float Roughness, float SSSProfileID, float SSSDMFPScale, float3 SSSDMFPAlbedo, float SSSDMFPRadius, uint SharedNormalIndex)
 {
 	FStrataData StrataData = GetInitialisedStrataData();
 	StrataData.LayerCount++;
@@ -191,11 +201,24 @@ FStrataData GetStrataDiffuseBSDF(float3 Albedo, float Roughness, float SSSProfil
 	BSDF_SETTYPE			(StrataData.Layers[0].BSDFs[0],		STRATA_BSDF_TYPE_DIFFUSE);
 	BSDF_SETNORMALID		(StrataData.Layers[0].BSDFs[0],		SharedNormalIndex);
 	BSDF_SETISTOPLAYER		(StrataData.Layers[0].BSDFs[0],		0);
-	BSDF_SETHASSCATTERING	(StrataData.Layers[0].BSDFs[0],		SSSRadiusScale > 0 ? 1 : 0);
+	BSDF_SETHASSCATTERING	(StrataData.Layers[0].BSDFs[0],		SSSDMFPScale > 0 ? 1 : 0);
+
+	const bool bHasSSSProfile = SSSProfileID > 0;
+	BSDF_SETHASSSSPROFILE	(StrataData.Layers[0].BSDFs[0],		bHasSSSProfile ? 1 : 0);
+
 	DIFFUSE_ALBEDO			(StrataData.Layers[0].BSDFs[0])		= Albedo;
 	DIFFUSE_ROUGHNESS		(StrataData.Layers[0].BSDFs[0])		= Roughness;
-	DIFFUSE_SSSPROFILEID	(StrataData.Layers[0].BSDFs[0])		= SSSProfileID;
-	DIFFUSE_SSSPROFILERADIUS(StrataData.Layers[0].BSDFs[0])		= SSSRadiusScale;
+
+	if (bHasSSSProfile)
+	{
+		DIFFUSE_SSSPROFILEID	(StrataData.Layers[0].BSDFs[0])	= SSSProfileID;
+		DIFFUSE_SSSPROFILERADIUS(StrataData.Layers[0].BSDFs[0])	= SSSDMFPScale;
+	}
+	else
+	{
+		DIFFUSE_SSSDMFPALBEDO	(StrataData.Layers[0].BSDFs[0])	= saturate(SSSDMFPAlbedo);
+		DIFFUSE_SSSDMFPRADIUS	(StrataData.Layers[0].BSDFs[0])	= clamp(SSSDMFPRadius * SSSDMFPScale, 0.f, SSS_PROFILE_PERPIXEL_MAXRADIUS);
+	}
 
 	StrataData.Layers[0].BSDFs[0].Weight						= 1.0f;
 
@@ -1337,10 +1360,13 @@ void PackStrataOut(
 	Strata.LayerCount = STRATA_CLAMPED_LAYER_COUNT;
 
 	// Subsurface scattering parameters
-	bool bHasSubsurfaceScattering = false;
-	float SubsurfaceProfile = 0;
-	float SubsurfaceRadiusScale = 0;
+	bool   bHasSubsurfaceScattering = false;
+	bool   bHasSubsurfaceProfile = false;
+	float  SubsurfaceProfile = 0;
+	float  SubsurfaceRadiusScale = 0;
 	float3 SubsurfaceBaseColor = 0;
+	float3 SubsurfaceDMFPAlbedo = 0;
+	float  SubsurfaceDMFPRadius = 0;
 
 	// For each of the BSDF, we evaluated their combined weight/throughput according to the material operation performed (layering, mixing, etc.).
 	int l;
@@ -1392,9 +1418,12 @@ void PackStrataOut(
 				if (BSDF_GETHASSCATTERING(BSDF) && BSDFCoverageAvg > 0 && DIFFUSE_SSSPROFILERADIUS(BSDF) > 0 && !bHasSubsurfaceScattering)
 				{
 					bHasSubsurfaceScattering = true;
-					SubsurfaceProfile = DIFFUSE_SSSPROFILEID(BSDF);
-					SubsurfaceRadiusScale = DIFFUSE_SSSPROFILERADIUS(BSDF);
-					SubsurfaceBaseColor = DIFFUSE_ALBEDO(BSDF);
+					bHasSubsurfaceProfile	= BSDF_GETHASSSSPROFILE(BSDF);
+					SubsurfaceProfile		= DIFFUSE_SSSPROFILEID(BSDF);
+					SubsurfaceRadiusScale	= DIFFUSE_SSSPROFILERADIUS(BSDF);
+					SubsurfaceBaseColor		= DIFFUSE_ALBEDO(BSDF);
+					SubsurfaceDMFPAlbedo	= DIFFUSE_SSSDMFPALBEDO(BSDF);
+					SubsurfaceDMFPRadius	= DIFFUSE_SSSDMFPRADIUS(BSDF);
 				}
 
 				LayerCoverage	+= BSDFCoverage;
@@ -1518,9 +1547,18 @@ void PackStrataOut(
 	GBuffer.CustomData = 0;
 	if (bHasSubsurfaceScattering)
 	{
-		GBuffer.BaseColor		= SubsurfaceBaseColor;
-		GBuffer.CustomData.rgb	= EncodeSubsurfaceProfile(SubsurfaceProfile);
-		GBuffer.CustomData.a	= SubsurfaceRadiusScale;
+		if (bHasSubsurfaceProfile)
+		{
+			GBuffer.BaseColor      = SubsurfaceBaseColor;
+			GBuffer.CustomData.rgb = EncodeSubsurfaceProfile(SubsurfaceProfile);
+			GBuffer.CustomData.a   = SubsurfaceRadiusScale;
+		}
+		else
+		{
+			GBuffer.BaseColor      = SubsurfaceBaseColor;
+			GBuffer.CustomData.rgb = sqrt(SubsurfaceDMFPAlbedo);
+			GBuffer.CustomData.a   = saturate(SubsurfaceDMFPRadius / SSS_PROFILE_PERPIXEL_MAXRADIUS);
+		}
 	}
 
 	// Now write out Strata data
@@ -1576,7 +1614,14 @@ void PackStrataOut(
 				STRATA_STORE_UINT1(PackAlbedoLinearToGamma2AlphaLinear(float4(DIFFUSE_ALBEDO(BSDF), DIFFUSE_ROUGHNESS(BSDF))));
 				if (BSDF_GETHASSCATTERING(BSDF))
 				{
-					STRATA_STORE_UINT1(PackRGBA8(float4(DIFFUSE_SSSPROFILEID(BSDF), DIFFUSE_SSSPROFILERADIUS(BSDF), 0, 0)))
+					if (BSDF_GETHASSSSPROFILE(BSDF))
+					{
+						STRATA_STORE_UINT1(PackRGBA8(float4(DIFFUSE_SSSPROFILEID(BSDF), DIFFUSE_SSSPROFILERADIUS(BSDF), 0, 0)));
+					}
+					else
+					{
+						STRATA_STORE_UINT1(PackAlbedoLinearToGamma2AlphaLinear(float4(DIFFUSE_SSSDMFPALBEDO(BSDF), DIFFUSE_SSSDMFPRADIUS(BSDF) / SSS_PROFILE_PERPIXEL_MAXRADIUS)));
+					}
 				}
 				// 8 bytes
 			}
@@ -1687,9 +1732,17 @@ FStrataBSDF UnpackStrataBSDFIn(ByteAddressBuffer StrataBuffer, inout uint ByteOf
 		StrataLoad_AlbedoGamma2ToLinearAlphaLinear(StrataBuffer, ByteOffset, DIFFUSE_ALBEDO(OutBSDF), DIFFUSE_ROUGHNESS(OutBSDF));
 		if (BSDF_GETHASSCATTERING(OutBSDF))
 		{
-			float4 Data = UnpackRGBA8(StrataLoadUint1(StrataBuffer, ByteOffset));
-			DIFFUSE_SSSPROFILEID(OutBSDF) = Data.r;
-			DIFFUSE_SSSPROFILERADIUS(OutBSDF) = Data.g;
+			if (BSDF_GETHASSSSPROFILE(OutBSDF))
+			{
+				float4 Data = UnpackRGBA8(StrataLoadUint1(StrataBuffer, ByteOffset));
+				DIFFUSE_SSSPROFILEID(OutBSDF)	  = Data.r;
+				DIFFUSE_SSSPROFILERADIUS(OutBSDF) = Data.g;
+			}
+			else
+			{
+				StrataLoad_AlbedoGamma2ToLinearAlphaLinear(StrataBuffer, ByteOffset, DIFFUSE_SSSDMFPALBEDO(OutBSDF), DIFFUSE_SSSDMFPRADIUS(OutBSDF));
+				DIFFUSE_SSSDMFPRADIUS(OutBSDF) *= SSS_PROFILE_PERPIXEL_MAXRADIUS;
+			}
 		}
 	}
 	break;
diff --git a/Engine/Shaders/Private/Strata/StrataVisualize.usf b/Engine/Shaders/Private/Strata/StrataVisualize.usf
--- a/Engine/Shaders/Private/Strata/StrataVisualize.usf
+++ b/Engine/Shaders/Private/Strata/StrataVisualize.usf
@@ -236,11 +236,28 @@ void VisualizeMaterialPS(
 				HORIZONTALTAB; PRINTPARAMCHAR(_R_); PRINTPARAMCHAR(_O_); PRINTPARAMCHAR(_U_); PRINTPARAMCHAR(_G_); PRINTPARAMCHAR(_H_); PRINTPARAMCHAR(_SPC_);
 				HORIZONTALTAB; PRINTFLOAT1(DIFFUSE_ROUGHNESS(BSDF));
 				NEXTLINE;
-				HORIZONTALTAB; PRINTPARAMCHAR(_S_); PRINTPARAMCHAR(_S_); PRINTPARAMCHAR(_S_); PRINTPARAMCHAR(_SPC_); PRINTPARAMCHAR(_I_); PRINTPARAMCHAR(_D_); PRINTPARAMCHAR(_SPC_);
-				HORIZONTALTAB; PRINTFLOAT1(DIFFUSE_SSSPROFILEID(BSDF));
-				NEXTLINE;
-				HORIZONTALTAB; PRINTPARAMCHAR(_S_); PRINTPARAMCHAR(_S_); PRINTPARAMCHAR(_S_); PRINTPARAMCHAR(_SPC_); PRINTPARAMCHAR(_R_); PRINTPARAMCHAR(_A_); PRINTPARAMCHAR(_SPC_);
-				HORIZONTALTAB; PRINTFLOAT1(DIFFUSE_SSSPROFILERADIUS(BSDF));
+
+				if (BSDF_GETHASSCATTERING(BSDF))
+				{
+					if (BSDF_GETHASSSSPROFILE(BSDF))
+					{
+						HORIZONTALTAB; PRINTPARAMCHAR(_S_); PRINTPARAMCHAR(_S_); PRINTPARAMCHAR(_S_); PRINTPARAMCHAR(_SPC_); PRINTPARAMCHAR(_I_); PRINTPARAMCHAR(_D_); PRINTPARAMCHAR(_SPC_);
+						HORIZONTALTAB; PRINTFLOAT1(DIFFUSE_SSSPROFILEID(BSDF));
+						NEXTLINE;
+						HORIZONTALTAB; PRINTPARAMCHAR(_S_); PRINTPARAMCHAR(_S_); PRINTPARAMCHAR(_S_); PRINTPARAMCHAR(_SPC_); PRINTPARAMCHAR(_R_); PRINTPARAMCHAR(_A_); PRINTPARAMCHAR(_SPC_);
+						HORIZONTALTAB; PRINTFLOAT1(DIFFUSE_SSSPROFILERADIUS(BSDF));
+						NEXTLINE;
+					}
+					else
+					{
+						HORIZONTALTAB; PRINTPARAMCHAR(_S_); PRINTPARAMCHAR(_S_); PRINTPARAMCHAR(_S_); PRINTPARAMCHAR(_SPC_); PRINTPARAMCHAR(_A_); PRINTPARAMCHAR(_L_); PRINTPARAMCHAR(_SPC_);
+						HORIZONTALTAB; PRINTFLOAT1(DIFFUSE_SSSDMFPALBEDO(BSDF));
+						NEXTLINE;
+						HORIZONTALTAB; PRINTPARAMCHAR(_S_); PRINTPARAMCHAR(_S_); PRINTPARAMCHAR(_S_); PRINTPARAMCHAR(_SPC_); PRINTPARAMCHAR(_R_); PRINTPARAMCHAR(_A_); PRINTPARAMCHAR(_SPC_);
+						HORIZONTALTAB; PRINTFLOAT1(DIFFUSE_SSSDMFPRADIUS(BSDF));
+						NEXTLINE;
+					}
+				}
 
 				break;
 
diff --git a/Engine/Shaders/Private/SubsurfaceBurleyNormalized.ush b/Engine/Shaders/Private/SubsurfaceBurleyNormalized.ush
--- a/Engine/Shaders/Private/SubsurfaceBurleyNormalized.ush
+++ b/Engine/Shaders/Private/SubsurfaceBurleyNormalized.ush
@@ -163,27 +163,18 @@ struct FBurleyParameter
 };
 
 //Surface albedo and mean free path length
-float4 GetSubsurfaceProfileSurfaceAlbedo(FGBufferData GBufferData)
+float4 GetSubsurfaceProfileSurfaceAlbedo(uint SubsurfaceProfileInt)
 {
-	// 0..255, which SubSurface profile to pick
-	uint SubsurfaceProfileInt = ExtractSubsurfaceProfileInt(GBufferData);
-
 	return ActualSSProfilesTexture.Load(int3(BSSS_SURFACEALBEDO_OFFSET, SubsurfaceProfileInt, 0));
 }
 
-float4 GetSubsurfaceProfileDiffuseMeanFreePath(FGBufferData GBufferData)
+float4 GetSubsurfaceProfileDiffuseMeanFreePath(uint SubsurfaceProfileInt)
 {
-	// 0..255, which SubSurface profile to pick
-	uint SubsurfaceProfileInt = ExtractSubsurfaceProfileInt(GBufferData);
-
 	return ActualSSProfilesTexture.Load(int3(BSSS_DMFP_OFFSET, SubsurfaceProfileInt, 0));
 }
 
-float GetSubsurfaceProfileWorldUnitScale(FGBufferData GBufferData)
+float GetSubsurfaceProfileWorldUnitScale(uint SubsurfaceProfileInt)
 {
-	// 0..255, which SubSurface profile to pick
-	uint SubsurfaceProfileInt = ExtractSubsurfaceProfileInt(GBufferData);
-
 	return ActualSSProfilesTexture.Load(int3(SSSS_SUBSURFACE_COLOR_OFFSET, SubsurfaceProfileInt, 0)).a;
 }
 
@@ -199,26 +190,47 @@ float SampleDepthTexture(float2 ScreenUV)
 	return Texture2DSampleLevel(SceneTexturesStruct.SceneDepthTexture, SceneTexturesStruct_SceneDepthTextureSampler, ScreenUV, 0).r;
 }
 
-
 // fetch surface albedo and diffuse mean free path
 // The diffuse mean free path is modulated by world scale and a custom term passed in
 // which is the opacity texture.
 FBurleyParameter GetBurleyParameters(FGBufferData GBuffer)
 {
-	FBurleyParameter BurleyParameter;
+	FBurleyParameter Out = (FBurleyParameter)0;
+#if STRATA_ENABLED
+	const bool bUseProfileData = GBuffer.CustomData.y == 0 && GBuffer.CustomData.z == 0;
+	if (GBuffer.ShadingModelID == SHADINGMODELID_STRATA && !bUseProfileData)
+	{
+		// Special data encoding from strata packing
+		const float3 MFPAlebdo			= Pow2(GBuffer.CustomData.xyz);
+		const float  MFPRadius			= GBuffer.CustomData.a * SSS_PROFILE_PERPIXEL_MAXRADIUS;
+		const float3 MFPChromaticRadius = MFPAlebdo * MFPRadius;
+
+		const float MaxAlbedo			= max3(MFPAlebdo.x, MFPAlebdo.y, MFPAlebdo.z);
+		const uint MaxAlbedoComponent	= MFPAlebdo.x == MaxAlbedo ? 0 : (MFPAlebdo.y == MaxAlbedo ? 1 : 2);
 
-	BurleyParameter.SurfaceAlbedo = GetSubsurfaceProfileSurfaceAlbedo(GBuffer);
-	float4 DiffuseMeanFreePath = GetSubsurfaceProfileDiffuseMeanFreePath(GBuffer);
-	BurleyParameter.DiffuseMeanFreePath = DecodeDiffuseMeanFreePath(DiffuseMeanFreePath);
-	float WorldUnitScale = DecodeWorldUnitScale(GetSubsurfaceProfileWorldUnitScale(GBuffer))*BURLEY_CM_2_MM;
-	BurleyParameter.WorldUnitScale = WorldUnitScale;
+		Out.SurfaceAlbedo		= float4(GBuffer.BaseColor, MaxAlbedoComponent);
+		Out.DiffuseMeanFreePath	= float4(MFPChromaticRadius, max3(MFPChromaticRadius.x, MFPChromaticRadius.y, MFPChromaticRadius.z));
+		Out.WorldUnitScale		= BURLEY_CM_2_MM;
 
-	// Mask should always be true, we can probably ignore it?
-	float Mask = UseSubsurfaceProfile(GBuffer.ShadingModelID) ? 1.0f : 0.0f;
-	BurleyParameter.SurfaceOpacity = Mask * GBuffer.CustomData.a;
-	BurleyParameter.DiffuseMeanFreePath *= BurleyParameter.SurfaceOpacity;
+		// Mask should always be true, we can probably ignore it?
+		// STRATA_TODO smooth transistion based on pixel footprint
+		Out.SurfaceOpacity		= MFPRadius > 0 ? 1.f : 0.f;
+	}
+	else
+#endif
+	{
+		uint SubsurfaceProfileInt = ExtractSubsurfaceProfileInt(GBuffer);
+
+		Out.SurfaceAlbedo		= GetSubsurfaceProfileSurfaceAlbedo(SubsurfaceProfileInt);
+		Out.DiffuseMeanFreePath	= DecodeDiffuseMeanFreePath(GetSubsurfaceProfileDiffuseMeanFreePath(SubsurfaceProfileInt));
+		Out.WorldUnitScale		= DecodeWorldUnitScale(GetSubsurfaceProfileWorldUnitScale(SubsurfaceProfileInt)) * BURLEY_CM_2_MM;
+
+		// Mask should always be true, we can probably ignore it?
+		Out.SurfaceOpacity		= UseSubsurfaceProfile(GBuffer.ShadingModelID) ? GBuffer.CustomData.a : 0.0f;
+		Out.DiffuseMeanFreePath*= Out.SurfaceOpacity;
+	}
 
-	return BurleyParameter;
+	return Out;
 }
 
 float2 Generate2DRandomNumber(int3 Seed)
@@ -449,12 +461,11 @@ int GetNumOfSamplesBasedOnQuality(float2 UV)
 	return NumOfSamples;
 }
 
-#define INVALID_PROFILE_ID 256
-// If the shading model is valid, return the profile. Otherwise, return INVALID_PROFILE_ID.
+// If the shading model is valid, return the profile. Otherwise, return SSS_PROFILE_ID_INVALID.
 // That way if we have a valid profile, we can check of the other profile is valid with a single boolean expression.
 uint ExtractSubsurfaceProfileIntWithInvalid(FGBufferData BufferData)
 {
-	uint ProfileID = INVALID_PROFILE_ID;
+	uint ProfileID = SSS_PROFILE_ID_INVALID;
 	if (UseSubsurfaceProfile(BufferData.ShadingModelID))
 	{
 		ProfileID = ExtractSubsurfaceProfileInt(BufferData);
@@ -470,7 +481,7 @@ uint2 ConvertUVPos2Input0GridPos(float2 BufferUV)
 
 uint ExtractSubsurfaceProfileIntWithInvalid(float2 UV,float Depth)
 {
-	uint ProfileID = INVALID_PROFILE_ID;
+	uint ProfileID = SSS_PROFILE_ID_INVALID;
 
 	// the invalid status be determined by the scene depth.
 	if (Depth > 0)
@@ -1011,7 +1022,7 @@ float4 BurleyNormalizedSS(float2 BufferUV, uint2 GroupThreadID)
 			// Determine the tint color, if the sampling pixel is not subsurface, we use tint color
 			// to mask out the sampling. Unless we specifically want the shadowing region.
 			float3 TintColor = (LocalProfile == SubsurfaceProfileInt) ? 1.0f : BoundaryColorBleed;
-			Mask = (LocalProfile != INVALID_PROFILE_ID) ? 1 : 0;
+			Mask = (LocalProfile != SSS_PROFILE_ID_INVALID) ? 1 : 0;
 
 			// Step 4: accumulate radiance from the diffusion profile rR(r)
 			// make sure the DiffuseMeanFreePath is not zero and in mm.
diff --git a/Engine/Source/Runtime/Engine/Classes/Materials/MaterialExpressionStrata.h b/Engine/Source/Runtime/Engine/Classes/Materials/MaterialExpressionStrata.h
--- a/Engine/Source/Runtime/Engine/Classes/Materials/MaterialExpressionStrata.h
+++ b/Engine/Source/Runtime/Engine/Classes/Materials/MaterialExpressionStrata.h
@@ -58,11 +58,23 @@ class UMaterialExpressionStrataDiffuseBSDF : public UMaterialExpression // STRAT
 	UPROPERTY()
 	FExpressionInput Normal;
 
+	/**
+	 * Mean free path albedo. Only used when no sub-surface profile is provided. (type = float3, unit = unitless)
+	 */
+	UPROPERTY(meta = (DisplayName = "DMFP Albedo"))
+	FExpressionInput DiffuseMeanFreePathAlbedo;
+
+	/**
+	 * Mean free path radius. Only used when no sub-surface profile is provided. (type = float, unit = cm)
+	 */
+	UPROPERTY(meta = (DisplayName = "DMFP Radius"))
+	FExpressionInput DiffuseMeanFreePathRadius;
+
 	/**
 	 * SSS radius scale (type = float, unit = unitless)
 	 */
-	UPROPERTY()
-	FExpressionInput SubsurfaceProfileScale;
+	UPROPERTY(meta = (DisplayName = "DMFP Scale"))
+	FExpressionInput DiffuseMeanFreePathRadiusScale;
 
 	/** SubsurfaceProfile, for Screen Space Subsurface Scattering. The profile needs to be set up on both the Strata diffuse node, and the material node at the moment. */
 	UPROPERTY(EditAnywhere, BlueprintReadOnly, Category = Material, meta = (DisplayName = "Subsurface Profile"))
@@ -76,6 +88,7 @@ class UMaterialExpressionStrataDiffuseBSDF : public UMaterialExpression // STRAT
 	virtual uint32 GetInputType(int32 InputIndex) override;
 	virtual bool IsResultStrataMaterial(int32 OutputIndex) override;
 	virtual void GatherStrataMaterialInfo(FStrataMaterialInfo& StrataMaterialInfo, int32 OutputIndex) override;
+	virtual FName GetInputName(int32 InputIndex) const override;
 #endif
 	//~ End UMaterialExpression Interface
 };
diff --git a/Engine/Source/Runtime/Engine/Private/Materials/HLSLMaterialTranslator.cpp b/Engine/Source/Runtime/Engine/Private/Materials/HLSLMaterialTranslator.cpp
--- a/Engine/Source/Runtime/Engine/Private/Materials/HLSLMaterialTranslator.cpp
+++ b/Engine/Source/Runtime/Engine/Private/Materials/HLSLMaterialTranslator.cpp
@@ -7483,14 +7483,16 @@ int32 FHLSLMaterialTranslator::StrataCreateAndRegisterNullMaterial()
 	return OutputCodeChunk;
 }
 
-int32 FHLSLMaterialTranslator::StrataDiffuseBSDF(int32 Albedo, int32 Roughness, int32 SSSProfileId, int32 SSSProfileRadiusScale, int32 Normal, uint8 SharedNormalIndex)
+int32 FHLSLMaterialTranslator::StrataDiffuseBSDF(int32 Albedo, int32 Roughness, int32 SSSProfileId, int32 SSSDMFPScale, int32 SSSDMFPAlbedo, int32 SSSDMFPRadius, int32 Normal, uint8 SharedNormalIndex)
 {
 	return AddCodeChunk(
-		MCT_Strata, TEXT("GetStrataDiffuseBSDF(%s, %s, %s, %s, %u) /* %s */"),
+		MCT_Strata, TEXT("GetStrataDiffuseBSDF(%s, %s, %s, %s, %s, %s, %u) /* %s */"),
 		*GetParameterCode(Albedo),
 		*GetParameterCode(Roughness),
 		*GetParameterCode(SSSProfileId),
-		*GetParameterCode(SSSProfileRadiusScale),
+		*GetParameterCode(SSSDMFPScale),
+		*GetParameterCode(SSSDMFPAlbedo),
+		*GetParameterCode(SSSDMFPRadius),
 		SharedNormalIndex,
 		*GetParameterCode(Normal)
 	);
diff --git a/Engine/Source/Runtime/Engine/Private/Materials/HLSLMaterialTranslator.h b/Engine/Source/Runtime/Engine/Private/Materials/HLSLMaterialTranslator.h
--- a/Engine/Source/Runtime/Engine/Private/Materials/HLSLMaterialTranslator.h
+++ b/Engine/Source/Runtime/Engine/Private/Materials/HLSLMaterialTranslator.h
@@ -791,7 +791,7 @@ protected:
 
 	// Strata
 	virtual int32 StrataCreateAndRegisterNullMaterial() override;
-	virtual int32 StrataDiffuseBSDF(int32 Albedo, int32 Roughness, int32 SSSProfileId, int32 SSSProfileRadiusScale, int32 Normal, uint8 SharedNormalIndex) override;
+	virtual int32 StrataDiffuseBSDF(int32 Albedo, int32 Roughness, int32 SSSProfileId, int32 SSSDMFPRadiusScale, int32 SSSDMFPAlebdo, int32 SSSDMFPRadius, int32 Normal, uint8 SharedNormalIndex) override;
 	virtual int32 StrataDielectricBSDF(int32 RoughnessX, int32 RoughnessY, int32 IOR, int32 Tint, int32 Normal, int32 Tangent, uint8 SharedNormalIndex) override;
 	virtual int32 StrataConductorBSDF(int32 Reflectivity, int32 EdgeColor, int32 RoughnessX, int32 RoughnessY, int32 Normal, int32 Tangent, uint8 SharedNormalIndex) override;
 	virtual int32 StrataVolumeBSDF(int32 Albedo, int32 Extinction, int32 Anisotropy, int32 Thickness, int32 Normal, uint8 SharedNormalIndex) override;
diff --git a/Engine/Source/Runtime/Engine/Private/Materials/MaterialExpressions.cpp b/Engine/Source/Runtime/Engine/Private/Materials/MaterialExpressions.cpp
--- a/Engine/Source/Runtime/Engine/Private/Materials/MaterialExpressions.cpp
+++ b/Engine/Source/Runtime/Engine/Private/Materials/MaterialExpressions.cpp
@@ -19198,20 +19198,21 @@ int32 UMaterialExpressionStrataDiffuseBSDF::Compile(class FMaterialCompiler* Com
 	int32 NormalCodeChunk = CompileWithDefaultCodeChunk(Compiler, Normal, Compiler->VertexNormal());
 	uint8 SharedNormalIndex = StrataCompilationInfoCreateSharedNormal(Compiler, NormalCodeChunk);
 
-	FName NameSubsurfaceProfile(FString(TEXT("__SubsurfaceProfile")));
-
 	int32 SSSProfileCodeChunk = INDEX_NONE;
 	const bool bHasScattering = SubsurfaceProfile != nullptr;
 	if (bHasScattering)
 	{
+		FName NameSubsurfaceProfile(FString(TEXT("__SubsurfaceProfile")));
 		SSSProfileCodeChunk = Compiler->ForceCast(Compiler->ScalarParameter(NameSubsurfaceProfile, 1.0f), MCT_Float1);
 	}
 
 	int32 OutputCodeChunk = Compiler->StrataDiffuseBSDF(
 		CompileWithDefaultFloat3(Compiler, Albedo,		0.18f, 0.18f, 0.18f),
 		CompileWithDefaultFloat1(Compiler, Roughness,	0.0f),
-		SSSProfileCodeChunk != INDEX_NONE ? SSSProfileCodeChunk : Compiler->Constant(1.0f),
-		CompileWithDefaultFloat1(Compiler, SubsurfaceProfileScale,	0.0f),
+		SSSProfileCodeChunk != INDEX_NONE ? SSSProfileCodeChunk : Compiler->Constant(0.0f),		
+		CompileWithDefaultFloat1(Compiler, DiffuseMeanFreePathRadiusScale,	0.0f),
+		CompileWithDefaultFloat3(Compiler, DiffuseMeanFreePathAlbedo, 0.0f, 0.0f, 0.0f),
+		CompileWithDefaultFloat1(Compiler, DiffuseMeanFreePathRadius, 0.0f),
 		NormalCodeChunk,
 		SharedNormalIndex);
 	StrataCompilationInfoCreateSingleBSDFMaterial(Compiler, OutputCodeChunk, SharedNormalIndex, STRATA_BSDF_TYPE_DIFFUSE, bHasScattering);
@@ -19243,6 +19244,12 @@ uint32 UMaterialExpressionStrataDiffuseBSDF::GetInputType(int32 InputIndex)
 		return MCT_Float3;
 		break;
 	case 3:
+		return MCT_Float3;
+		break;
+	case 4:
+		return MCT_Float1;
+		break;
+	case 5:
 		return MCT_Float1;
 		break;
 	}
@@ -19251,6 +19258,35 @@ uint32 UMaterialExpressionStrataDiffuseBSDF::GetInputType(int32 InputIndex)
 	return MCT_Float1;
 }
 
+FName UMaterialExpressionStrataDiffuseBSDF::GetInputName(int32 InputIndex) const
+{
+	if (InputIndex == 0)
+	{
+		return TEXT("Albedo");
+	}
+	else if (InputIndex == 1)
+	{
+		return TEXT("Roughness");
+	}
+	else if (InputIndex == 2)
+	{
+		return TEXT("Normal");
+	}
+	else if (InputIndex == 3)
+	{
+		return TEXT("DMFP Albedo");
+	}
+	else if (InputIndex == 4)
+	{
+		return TEXT("DMFP Radius");
+	}
+	else if (InputIndex == 5)
+	{
+		return TEXT("DMFP Scale");
+	}
+	return TEXT("Unknown");
+}
+
 bool UMaterialExpressionStrataDiffuseBSDF::IsResultStrataMaterial(int32 OutputIndex)
 {
 	return true;
diff --git a/Engine/Source/Runtime/Engine/Public/MaterialCompiler.h b/Engine/Source/Runtime/Engine/Public/MaterialCompiler.h
--- a/Engine/Source/Runtime/Engine/Public/MaterialCompiler.h
+++ b/Engine/Source/Runtime/Engine/Public/MaterialCompiler.h
@@ -391,7 +391,7 @@ public:
 
 	// Strata
 	virtual int32 StrataCreateAndRegisterNullMaterial() = 0;
-	virtual int32 StrataDiffuseBSDF(int32 Albedo, int32 Roughness, int32 SSSProfileId, int32 SSSProfileRadiusScale, int32 Normal, uint8 SharedNormalIndex) = 0;
+	virtual int32 StrataDiffuseBSDF(int32 Albedo, int32 Roughness, int32 SSSProfileId, int32 SSSDMFPRadiusScale, int32 SSSDMFPAlebdo, int32 SSSDMFPRadius, int32 Normal, uint8 SharedNormalIndex) = 0;
 	virtual int32 StrataDielectricBSDF(int32 RoughnessX, int32 RoughnessY, int32 IOR, int32 Tint, int32 Normal, int32 Tangent, uint8 SharedNormalIndex) = 0;
 	virtual int32 StrataConductorBSDF(int32 Reflectivity, int32 EdgeColor, int32 RoughnessX, int32 RoughnessY, int32 Normal, int32 Tangent, uint8 SharedNormalIndex) = 0;
 	virtual int32 StrataVolumeBSDF(int32 Albedo, int32 Extinction, int32 Anisotropy, int32 Thickness, int32 Normal, uint8 SharedNormalIndex) = 0;
@@ -807,9 +807,9 @@ public:
 		return Compiler->StrataCreateAndRegisterNullMaterial();
 	}
 
-	virtual int32 StrataDiffuseBSDF(int32 Albedo, int32 Roughness, int32 SSSProfileId, int32 SSSProfileRadiusScale, int32 Normal, uint8 SharedNormalIndex) override
+	virtual int32 StrataDiffuseBSDF(int32 Albedo, int32 Roughness, int32 SSSProfileId, int32 SSSDMFPRadiusScale, int32 SSSDMFPAlebdo, int32 SSSDMFPRadius, int32 Normal, uint8 SharedNormalIndex) override
 	{
-		return Compiler->StrataDiffuseBSDF(Albedo, Roughness, SSSProfileId, SSSProfileRadiusScale, Normal, SharedNormalIndex);
+		return Compiler->StrataDiffuseBSDF(Albedo, Roughness, SSSProfileId, SSSDMFPRadiusScale, SSSDMFPAlebdo, SSSDMFPRadius, Normal, SharedNormalIndex);
 	}
 
 	virtual int32 StrataDielectricBSDF(int32 RoughnessX, int32 RoughnessY, int32 IOR, int32 Tint, int32 Normal, int32 Tangent, uint8 SharedNormalIndex) override
diff --git a/Engine/Source/Runtime/Renderer/Private/PostProcess/PostProcessSubsurface.cpp b/Engine/Source/Runtime/Renderer/Private/PostProcess/PostProcessSubsurface.cpp
--- a/Engine/Source/Runtime/Renderer/Private/PostProcess/PostProcessSubsurface.cpp
+++ b/Engine/Source/Runtime/Renderer/Private/PostProcess/PostProcessSubsurface.cpp
@@ -18,6 +18,7 @@
 #include "ScenePrivate.h"
 #include "GenerateMips.h"
 #include "ClearQuad.h"
+#include "Strata/Strata.h"
 
 namespace
 {
@@ -428,10 +429,9 @@ public:
 
 	static void ModifyCompilationEnvironment(const FGlobalShaderPermutationParameters& Parameters, FShaderCompilerEnvironment& OutEnvironment)
 	{
-		FGlobalShader::ModifyCompilationEnvironment(Parameters, OutEnvironment);
-		OutEnvironment.SetDefine(TEXT("SUBSURFACE_RADIUS_SCALE"), SUBSURFACE_RADIUS_SCALE);
-		OutEnvironment.SetDefine(TEXT("SUBSURFACE_KERNEL_SIZE"), SUBSURFACE_KERNEL_SIZE);
+		FSubsurfaceShader::ModifyCompilationEnvironment(Parameters, OutEnvironment);
 		OutEnvironment.SetDefine(TEXT("SUBSURFACE_BURLEY_COMPUTE"), 1);
+		OutEnvironment.SetDefine(TEXT("STRATA_ENABLED"), Strata::IsStrataEnabled() ? 1 : 0); 
 	}
 
 	BEGIN_SHADER_PARAMETER_STRUCT(FParameters, )
@@ -455,10 +455,9 @@ public:
 
 	static void ModifyCompilationEnvironment(const FGlobalShaderPermutationParameters& Parameters, FShaderCompilerEnvironment& OutEnvironment)
 	{
-		FGlobalShader::ModifyCompilationEnvironment(Parameters, OutEnvironment);
-		OutEnvironment.SetDefine(TEXT("SUBSURFACE_RADIUS_SCALE"), SUBSURFACE_RADIUS_SCALE);
-		OutEnvironment.SetDefine(TEXT("SUBSURFACE_KERNEL_SIZE"), SUBSURFACE_KERNEL_SIZE);
+		FSubsurfaceShader::ModifyCompilationEnvironment(Parameters, OutEnvironment);
 		OutEnvironment.SetDefine(TEXT("SUBSURFACE_BURLEY_COMPUTE"), 1);
+		OutEnvironment.SetDefine(TEXT("STRATA_ENABLED"), Strata::IsStrataEnabled() ? 1 : 0);
 	}
 
 	BEGIN_SHADER_PARAMETER_STRUCT(FParameters, )
@@ -483,10 +482,9 @@ public:
 
 	static void ModifyCompilationEnvironment(const FGlobalShaderPermutationParameters& Parameters, FShaderCompilerEnvironment& OutEnvironment)
 	{
-		FGlobalShader::ModifyCompilationEnvironment(Parameters, OutEnvironment);
-		OutEnvironment.SetDefine(TEXT("SUBSURFACE_RADIUS_SCALE"), SUBSURFACE_RADIUS_SCALE);
-		OutEnvironment.SetDefine(TEXT("SUBSURFACE_KERNEL_SIZE"), SUBSURFACE_KERNEL_SIZE);
+		FSubsurfaceShader::ModifyCompilationEnvironment(Parameters, OutEnvironment);
 		OutEnvironment.SetDefine(TEXT("SUBSURFACE_BURLEY_COMPUTE"), 1);
+		OutEnvironment.SetDefine(TEXT("STRATA_ENABLED"), Strata::IsStrataEnabled() ? 1 : 0);
 	}
 
 	BEGIN_SHADER_PARAMETER_STRUCT(FParameters, )
@@ -518,12 +516,11 @@ public:
 
 	static void ModifyCompilationEnvironment(const FGlobalShaderPermutationParameters& Parameters, FShaderCompilerEnvironment& OutEnvironment)
 	{
-		FGlobalShader::ModifyCompilationEnvironment(Parameters, OutEnvironment);
-		OutEnvironment.SetDefine(TEXT("SUBSURFACE_RADIUS_SCALE"), SUBSURFACE_RADIUS_SCALE);
-		OutEnvironment.SetDefine(TEXT("SUBSURFACE_KERNEL_SIZE"), SUBSURFACE_KERNEL_SIZE);
+		FSubsurfaceShader::ModifyCompilationEnvironment(Parameters, OutEnvironment);
 		OutEnvironment.SetDefine(TEXT("SUBSURFACE_BURLEY_COMPUTE"), 1);
 		OutEnvironment.SetDefine(TEXT("ENABLE_VELOCITY"), 1);
 		OutEnvironment.SetDefine(TEXT("SUBSURFACE_GROUP_SIZE"), kSubsurfaceGroupSize);
+		OutEnvironment.SetDefine(TEXT("STRATA_ENABLED"), Strata::IsStrataEnabled() ? 1 : 0);
 	}
 
 	BEGIN_SHADER_PARAMETER_STRUCT(FParameters, )
@@ -784,6 +781,12 @@ class FSubsurfaceRecombinePS : public FSubsurfaceShader
 	class FRunningInSeparable : SHADER_PERMUTATION_BOOL("SUBSURFACE_FORCE_SEPARABLE");
 	using FPermutationDomain = TShaderPermutationDomain<FDimensionMode, FDimensionQuality, FDimensionCheckerboard, FDimensionHalfRes, FRunningInSeparable>;
 
+	static void ModifyCompilationEnvironment(const FGlobalShaderPermutationParameters& Parameters, FShaderCompilerEnvironment& OutEnvironment)
+	{
+		FSubsurfaceShader::ModifyCompilationEnvironment(Parameters, OutEnvironment);
+		OutEnvironment.SetDefine(TEXT("STRATA_ENABLED"), Strata::IsStrataEnabled() ? 1 : 0);
+	}
+
 	// Returns the Recombine quality level requested by the SSS Quality CVar setting.
 	static EQuality GetQuality(const FViewInfo& View)
 	{
