Update hashing functions to improve quality of Owen-Scrambling

The error-diffused sampler is a much more stringent test of Owen-Scrambling quality since it scrambles across pixels. I generated some new constants for the scrambling that are slightly better than the ones from the original Laine/Karras paper.

Use a slightly better seed evolution to feed the OwenScrambling calls in the sobol sampler. The previous LCG approach did not provide enough randomness. The chosen weak hash is carefully chosen to be cheap while still having just enough randomness to decorelate dimensions cleanly.

Remove dependency of RandomSequence header on QuasiRandom.h which only provided the IntegerHash function. Instead use more modern integer hashes that are faster and higher quality.

Replace RANDSEQ_LCG with a hash based version that produces a higher quality stream of random values so we have a solid pure monte carlo reference to compare against.

diff --git a/Engine/Plugins/Experimental/GPULightmass/Shaders/Private/LightmapPathTracing.usf b/Engine/Plugins/Experimental/GPULightmass/Shaders/Private/LightmapPathTracing.usf
--- a/Engine/Plugins/Experimental/GPULightmass/Shaders/Private/LightmapPathTracing.usf
+++ b/Engine/Plugins/Experimental/GPULightmass/Shaders/Private/LightmapPathTracing.usf
@@ -29,6 +29,7 @@ struct FIrradianceCacheRecord
 
 #include "BatchedTiles.ush"
 #include "/Engine/Private/PathTracing/Material/FirstBounceRayGuidingCommon.ush"
+#include "/Engine/Private/QuasiRandom.ush"
 
 int NumRayGuidingTrialSamples;
 RWTexture2D<uint> RayGuidingLuminance;
diff --git a/Engine/Shaders/Private/PathTracing/Utilities/PathTracingRandomSequence.ush b/Engine/Shaders/Private/PathTracing/Utilities/PathTracingRandomSequence.ush
--- a/Engine/Shaders/Private/PathTracing/Utilities/PathTracingRandomSequence.ush
+++ b/Engine/Shaders/Private/PathTracing/Utilities/PathTracingRandomSequence.ush
@@ -6,59 +6,72 @@
 
 #pragma once
 
+#include "../../MortonCode.ush"
 
-#include "../../QuasiRandom.ush" 
-
-#define RANDSEQ_LCG			0
+#define RANDSEQ_PURERANDOM	0
 #define RANDSEQ_HALTON		1
 #define RANDSEQ_OWENSOBOL	2
 
+// Select a default if none was specified
 #ifndef RANDSEQ
 #define RANDSEQ			RANDSEQ_OWENSOBOL
 #endif
 
 #ifndef RANDSEQ_RANDOMIZED
-#define RANDSEQ_RANDOMIZED		1 // 1 to randomize per pixel, 0 to share the same sequence in all pixels
+#define RANDSEQ_RANDOMIZED		1 // 1 to randomize per pixel, 0 to share the same sequence in all pixels (useful to benchmark coherent sampling)
 #endif
 
-struct RandomSequence
+// This hash mixes bits at the theoretical bias limit
+uint PerfectIntegerHash(uint x)
 {
-	uint SampleIndex;		// index into the random sequence
-	uint SampleSeed;		// changes as we draw samples to reflect the change in dimension
-};
+	// https://nullprogram.com/blog/2018/07/31/
+	// exact bias: 0.020888578919738908
+	x ^= x >> 17;
+	x *= 0xed5ad4bbu;
+	x ^= x >> 11;
+	x *= 0xac4c1b51u;
+	x ^= x >> 15;
+	x *= 0x31848babu;
+	x ^= x >> 14;
+	return x;
+}
 
-// TEA-based pseudo-random number generator
-uint RandInit(uint Seed0, uint Seed1)
+// High quality integer hash - this mixes bits almost perfectly
+uint StrongIntegerHash(uint x)
 {
-	// Constants cited from "GPU Random Numbers via the Tiny Encryption Algorithm"
-	const uint Delta = 0x9e3779b9;
-	const uint4 Key = uint4(0xa341316c, 0xc8013ea4, 0xad90777d, 0x7e95761e);
-	uint Rounds = 8;
-
-	uint Sum = 0;
-	uint2 Value = uint2(Seed0, Seed1);
-	for (uint Index = 0; Index < Rounds; ++Index)
-	{
-		Sum += Delta;
-		Value.x += (Value.y + Sum) ^ ((Value.y << 4) + Key.x) ^ ((Value.y >> 5) + Key.y);
-		Value.y += (Value.x + Sum) ^ ((Value.x << 4) + Key.z) ^ ((Value.x >> 5) + Key.w);
-	}
+	// From [REDACTED_PROFILE_URL]/hash-prospector
+	// bias = 0.16540778981744320
+	x ^= x >> 16;
+	x *= 0xa812d533;
+	x ^= x >> 15;
+	x *= 0xb278e4ad;
+	x ^= x >> 17;
+	return x;
+}
 
-	return Value.x;
+// This is a much weaker hash, but is faster and can be used to drive other hashes
+uint WeakIntegerHash(uint x)
+{
+	// Generated using [REDACTED_PROFILE_URL]/hash-prospector
+	// Estimated Bias ~586
+	x *= 0xe92955555u;
+	x ^= x >> 15;
+	return x;
 }
 
-// Linear congruential generator to evolve pseudo-random numbers
+struct RandomSequence
+{
+	uint SampleIndex;		// index into the random sequence
+	uint SampleSeed;		// changes as we draw samples to reflect the change in dimension
+};
+
 float Rand(inout uint Seed)
 {
-	// Scale and Bias coefficients are taken from the Park-Miller RNG
-	const uint Scale = 48271;
-	const uint Bias = 0;
-	Seed = Seed * Scale + Bias;
-
-	// Map to significand and divide into [0, 1) range
-	float Result = float(Seed & 0x00FFFFFF);
-	Result /= float(0x01000000);
-	return Result;
+	// Counter based PRNG -- safer than most small-state PRNGs since we use the random values directly here.
+	Seed += 1;
+	uint Output = StrongIntegerHash(Seed);
+	// take low 24 bits
+	return (Output & 0xFFFFFF) * 5.96046447754e-08; // * 2^-24
 }
 
 // #dxr_todo: convert prime factors to constant buffer
@@ -145,10 +158,10 @@ uint FastOwenScrambling(uint Index, uint Seed) {
 
 	// loop below does not require pre-inverted bits, so undo the inversion assumed in the optimized case below
 	Index = ReverseBits32(Index);
-	for (uint Mask = 1u << 31; Mask; Mask >>= 1) {
-		Seed = IntegerHash(Seed); // randomize state
+	for (uint Mask = 1u << 31, Input = Index; Mask; Mask >>= 1) {
+		Seed = PerfectIntegerHash(Seed); // randomize state
 		Index ^= Seed & Mask;  // flip output (depending on state)
-		Seed ^= Index & Mask;  // flip state  (depending on output)
+		Seed ^= Input & Mask;  // flip state  (depending on input)
 	}
 	return Index;
 #else
@@ -158,25 +171,26 @@ uint FastOwenScrambling(uint Index, uint Seed) {
 	// NOTE: This operation has been performed outside this call
 	// Index = ReverseBits32(Index);
 	
+	// This follows the basic construction from the paper above. The error-diffusion sampler which
+	// tiles a single point set across the screen, makes it much easier to visually "see" the impact of the scrambling.
+	// When the scrambling is not good enough, the structure of the space-filling curve is left behind.
+	// Care must be taken to ensure that the scrambling is still unbiased on average though. I found some cases
+	// that seemed to produce lower visual error but were actually biased due to not touching certain bits,
+	// leading to correlations across dimensions.
+
+	// After much experimentation with the hash prospector, I discovered the following two constants which appear to
+	// give results as good (or better) than the original 4 xor-mul constants from the Laine/Karras paper.
+	// It isn't entirely clear to me why some constants work better than others. Some hashes with slightly less
+    // bias produced visibly more error or worse looking power spectra.
+	// Estimates below from hash prospector for all hashes of the form: add,xmul=c0,xmul=c1 (with c0 and c1 being
+	// the constants below).
+	// Ran with score_quality=16 for about ~10000 random hashes
+	// Average bias: ~727.02
+	//   Best  bias: ~723.05
+	//   Worst bias: ~735.19
 	Index += Seed; // randomize the index by our seed (pushes bits toward the left)
-
-	// A few different options below for doing the scrambling step:
-	// I'll stick to the simplest one until we have a good variance/time measurement system
-#if 1
-	// NOTE: Original publication has two extra xor-muls but they don't actually improve quality or convergence
-	Index ^= Index * 0x6C50B47Cu;
-	Index ^= Index * 0xB82F1E52u;
-#elif 0
-	// Better looking spectra on average - though still not perfect
-	Index ^= (Index + Seed) * 0x6C50B47Cu;
-	Index ^= (Index + Seed) * 0xB82F1E52u;
-#else
-	// This method achieves the closest results to real Owen-Scrambling
-	// but (probably?) costs a bit more perf
-	Index ^= Index * (Seed & ~1u);
-	Seed = IntegerHash(Seed);
-	Index ^= Index * (Seed & ~1u);
-#endif
+	Index ^= Index * 0x9c117646u;
+	Index ^= Index * 0xe0705d72u;
 
 	// Undo the reverse so that we get left-to-right scrambling
 	// thereby emulating owen-scrambling
@@ -225,22 +239,42 @@ static const uint3 SobolMatrices[] = {
 	uint3(0xffffffff, 0xaaaa00a3, 0xc900a00a)
 };
 
-float4 SobolSampler(uint SampleIndex, inout uint Seed) {
-	// We need a random sequence that we can advance quickly (5 steps per 4 samples)
+uint EvolveSobolSeed(inout uint Seed)
+{
 	// constant from: https://www.pcg-random.org/posts/does-it-beat-the-minimal-standard.html
-	const uint LCG_C = 1019135901;
+	const uint MCG_C = 2739110765;
+#if 0
+	return Seed += MCG_C;
+#elif 0
+	Seed += MCG_C;
+	return Seed ^ (Seed >> 16); // scramble lower bits
+#elif 0
+	return PerfectIntegerHash(Seed += 1);
+#elif 0
+	return StrongIntegerHash(Seed += 1);
+#else
+	// a slightly weaker hash is ok since this drives FastOwenScrambling which is itself a hash
+	// Note that the Seed evolution is just an integer addition and the hash should optimize away
+	// when a particular dimension is not used
+	return WeakIntegerHash(Seed += MCG_C);
+#endif
+}
+
+float4 SobolSampler(uint SampleIndex, inout uint Seed)
+{
 	// first scramble the index to decorelate from other 4-tuples
-	uint SobolIndex = FastOwenScrambling(SampleIndex, Seed = Seed * LCG_C + LCG_C);
+	uint SobolIndex = FastOwenScrambling(SampleIndex, EvolveSobolSeed(Seed));
 	// now get Sobol' point from this index
 	uint4 Result = uint4(SobolIndex, 0, 0, 0);
 	for (int i = 0; SobolIndex; SobolIndex >>= 1, i++)
 		Result.yzw ^= (SobolIndex & 1) * SobolMatrices[i];
 
 	// finally scramble the points to avoid structured artifacts
-	Result.x = FastOwenScrambling(Result.x, Seed = Seed * LCG_C + LCG_C);
-	Result.y = FastOwenScrambling(Result.y, Seed = Seed * LCG_C + LCG_C);
-	Result.z = FastOwenScrambling(Result.z, Seed = Seed * LCG_C + LCG_C);
-	Result.w = FastOwenScrambling(Result.w, Seed = Seed * LCG_C + LCG_C);
+	Result.x = FastOwenScrambling(Result.x, EvolveSobolSeed(Seed));
+	Result.y = FastOwenScrambling(Result.y, EvolveSobolSeed(Seed));
+	Result.z = FastOwenScrambling(Result.z, EvolveSobolSeed(Seed));
+	Result.w = FastOwenScrambling(Result.w, EvolveSobolSeed(Seed));
+
 	// output as float in [0,1) taking care not to skew the distribution
 	// due to the non-uniform spacing of floats in this range
 	return (Result >> 8) * 5.96046447754e-08; // * 2^-24
@@ -253,16 +287,17 @@ void RandomSequence_Initialize(inout RandomSequence RandSequence, uint PositionS
 	PositionSeed = 0;
 #endif
 
-#if RANDSEQ == RANDSEQ_LCG
+#if RANDSEQ == RANDSEQ_PURERANDOM
 	RandSequence.SampleIndex = 0; // not used
-	RandSequence.SampleSeed = RandInit(PositionSeed, TimeSeed);
+	RandSequence.SampleSeed = StrongIntegerHash(PositionSeed + StrongIntegerHash(TimeSeed));
 #elif RANDSEQ == RANDSEQ_HALTON
-	RandSequence.SampleIndex = TimeSeed + IntegerHash(PositionSeed);
+	RandSequence.SampleIndex = TimeSeed + StrongIntegerHash(PositionSeed);
 	RandSequence.SampleSeed = 0; // this sampler just counts dimensions
 #elif RANDSEQ == RANDSEQ_OWENSOBOL
 	// pre-compute bit reversal needed for FastOwenScrambling since this index doesn't change
 	RandSequence.SampleIndex = ReverseBits32(TimeSeed);
-	RandSequence.SampleSeed  = IntegerHash(PositionSeed);
+	// change seed to get a unique sequence per pixel
+	RandSequence.SampleSeed  = StrongIntegerHash(PositionSeed);
 #else
 #error "Unknown random sequence chosen in path tracer"
 #endif
@@ -300,16 +335,16 @@ void RandomSequence_Initialize(inout RandomSequence RandSequence, uint2 PixelCoo
 	uint TileID = MortonEncode(uint2(X ^ Y, X));
 #elif 0
     // pure random (breaks nice properties)
-	uint TileID = IntegerHash(X + IntegerHash(Y));
+	uint TileID = StrongIntegerHash(X + StrongIntegerHash(Y));
 #else
 	// Hilbert curve: [REDACTED_PROFILE_URL]/Hackers-Delight/blob/master/hilbert/hil_s_from_xy.c.txt
 	uint TileID = 0, HilbertState = 0;
 	for (int i = 7; i >= 0; i--) {
-		int xi = (X >> i) & 1;
-		int yi = (Y >> i) & 1;
-		uint Row = (4 * HilbertState) | (2 * xi) | yi;
-		TileID = (TileID << 2) | ((0x361E9CB4 >> 2 * Row) & 3);
-		HilbertState = (0x8FE65831 >> 2 * Row) & 3;
+		uint xi = (X >> i) & 1;
+		uint yi = (Y >> i) & 1;
+		uint Row = 8 * HilbertState + 4 * xi + 2 * yi;
+		TileID = TileID * 4 + ((0x361E9CB4 >> Row) & 3);
+		HilbertState = (0x8FE65831 >> Row) & 3;
 	}
 #endif
 
@@ -324,12 +359,22 @@ void RandomSequence_Initialize(inout RandomSequence RandSequence, uint2 PixelCoo
 	RandSequence.SampleSeed = 0; // always use the same sequence
 }
 
+#else
+
+// Error-diffusion sampler is not available, but implement the same API signature for consistency
+void RandomSequence_Initialize(inout RandomSequence RandSequence, uint2 PixelCoord, uint SampleIndex, uint TimeSeed0, uint MaxSamples)
+{
+	// Error diffusion is not available, so just forward to the normal call
+	uint LinearIndex = (PixelCoord.x & 0xFFFF) + (PixelCoord.y * 65536);
+	RandomSequence_Initialize(RandSequence, LinearIndex, TimeSeed0 * MaxSamples + SampleIndex);
+}
+
 #endif
 
 float RandomSequence_GenerateSample1D(inout RandomSequence RandSequence)
 {
 	float Result;
-#if RANDSEQ == RANDSEQ_LCG
+#if RANDSEQ == RANDSEQ_PURERANDOM
 	Result = Rand(RandSequence.SampleSeed);
 #elif RANDSEQ == RANDSEQ_HALTON
 	Result = Halton(RandSequence.SampleIndex, Prime512(RandSequence.SampleSeed + 0));
@@ -347,7 +392,7 @@ float RandomSequence_GenerateSample1D(inout RandomSequence RandSequence)
 float2 RandomSequence_GenerateSample2D(inout RandomSequence RandSequence)
 {
 	float2 Result;
-#if RANDSEQ == RANDSEQ_LCG
+#if RANDSEQ == RANDSEQ_PURERANDOM
 	Result.x = Rand(RandSequence.SampleSeed);
 	Result.y = Rand(RandSequence.SampleSeed);
 #elif RANDSEQ == RANDSEQ_HALTON
@@ -367,7 +412,7 @@ float2 RandomSequence_GenerateSample2D(inout RandomSequence RandSequence)
 float3 RandomSequence_GenerateSample3D(inout RandomSequence RandSequence)
 {
 	float3 Result;
-#if RANDSEQ == RANDSEQ_LCG
+#if RANDSEQ == RANDSEQ_PURERANDOM
 	Result.x = Rand(RandSequence.SampleSeed);
 	Result.y = Rand(RandSequence.SampleSeed);
 	Result.z = Rand(RandSequence.SampleSeed);
@@ -390,7 +435,7 @@ float3 RandomSequence_GenerateSample3D(inout RandomSequence RandSequence)
 float4 RandomSequence_GenerateSample4D(inout RandomSequence RandSequence)
 {
 	float4 Result;
-#if RANDSEQ == RANDSEQ_LCG
+#if RANDSEQ == RANDSEQ_PURERANDOM
 	Result.x = Rand(RandSequence.SampleSeed);
 	Result.y = Rand(RandSequence.SampleSeed);
 	Result.z = Rand(RandSequence.SampleSeed);
