Fix usage of the RandomSequence API in the raytracing ambient occlusion pass and add support for error-diffusing sampler initialization

The previous code was re-using a single RandomsSequence without considering the sample index in the multi-sample case. This meant the equivalent of random sampling was being used.

diff --git a/Engine/Shaders/Private/RayTracing/RayTracingAmbientOcclusionRGS.usf b/Engine/Shaders/Private/RayTracing/RayTracingAmbientOcclusionRGS.usf
--- a/Engine/Shaders/Private/RayTracing/RayTracingAmbientOcclusionRGS.usf
+++ b/Engine/Shaders/Private/RayTracing/RayTracingAmbientOcclusionRGS.usf
@@ -52,9 +52,7 @@ void GenerateCosineNormalRay(
 	float4 Direction_Tangent;
 	#if 1
 	{
-		Direction_Tangent.xy = UniformSampleDiskConcentric(RandSample);
-
-		Direction_Tangent.z = sqrt(1.0 - dot(Direction_Tangent.xy, Direction_Tangent.xy));
+		Direction_Tangent = CosineSampleHemisphereConcentric(RandSample);
 		Direction_Tangent.w = Direction_Tangent.z;
 	}
 	#else
@@ -89,10 +87,6 @@ RAY_TRACING_ENTRY_RAYGEN(AmbientOcclusionRGS)
 {
 	uint2 PixelCoord = DispatchRaysIndex().xy + View.ViewRectMin;
 
-	RandomSequence RandSequence;
-	uint LinearIndex = CalcLinearIndex(PixelCoord);
-	RandomSequence_Initialize(RandSequence, LinearIndex, View.StateFrameIndex);
-
 	// Get G-Buffer surface data
 	float2 InvBufferSize = View.BufferSizeAndInvSize.zw;
 	float2 UV = (float2(PixelCoord) + 0.5) * InvBufferSize;
@@ -174,10 +168,21 @@ RAY_TRACING_ENTRY_RAYGEN(AmbientOcclusionRGS)
 		SamplesPerPixelLocal = 0.0;
 	}
 
+	uint TimeSeed = View.StateFrameIndex;
+
 	for (uint SampleIndex = 0; SampleIndex < SamplesPerPixelLocal; ++SampleIndex)
 	{
 		RayDesc Ray;
 		float RayPDF;
+
+#define USE_ERROR_DIFFUSION_SAMPLER 1 && RANDSEQ == RANDSEQ_OWENSOBOL
+
+		RandomSequence RandSequence;
+#if USE_ERROR_DIFFUSION_SAMPLER
+		RandomSequence_Initialize(RandSequence, PixelCoord, SampleIndex, TimeSeed, SamplesPerPixelLocal);
+#else
+		RandomSequence_Initialize(RandSequence, CalcLinearIndex(PixelCoord), SamplesPerPixelLocal * TimeSeed + SampleIndex);
+#endif
 		GenerateCosineNormalRay(PixelCoord, WorldPosition, WorldNormal, GeometricNormal, RandSequence, Ray.Origin, Ray.Direction, Ray.TMin, Ray.TMax, RayPDF);
 		ApplyCameraRelativeDepthBias(Ray, PixelCoord, DeviceZ, WorldNormal, MaxNormalBias);
 
