Added Strata HairBSDF for HairStrand (all supported) and deferred lighting path (some follow up will be required)

diff --git a/Engine/Shaders/Private/BasePassPixelShader.usf b/Engine/Shaders/Private/BasePassPixelShader.usf
--- a/Engine/Shaders/Private/BasePassPixelShader.usf
+++ b/Engine/Shaders/Private/BasePassPixelShader.usf
@@ -85,7 +85,7 @@
 #define STRATA_TRANSLUCENT_ENABLED		(STRATA_ENABLED && MATERIALBLENDING_ANY_TRANSLUCENT>0)
 
 #if STRATA_TRANSLUCENT_ENABLED
-#include "ParticipatingMediaCommon.ush"
+#include "/Engine/Private/Strata/StrataEvaluation.ush"
 #include "/Engine/Private/Strata/StrataForwardLighting.ush"
 #endif
 
diff --git a/Engine/Shaders/Private/DeferredLightPixelShaders.usf b/Engine/Shaders/Private/DeferredLightPixelShaders.usf
--- a/Engine/Shaders/Private/DeferredLightPixelShaders.usf
+++ b/Engine/Shaders/Private/DeferredLightPixelShaders.usf
@@ -33,8 +33,7 @@ float CloudShadowmapStrength;
 
 #if STRATA_ENABLED
 #define STRATA_INLINE_SHADING 0
-#include "/Engine/Private/Strata/Strata.ush"
-#include "/Engine/Private/ParticipatingMediaCommon.ush"
+#include "/Engine/Private/Strata/StrataEvaluation.ush"
 #include "/Engine/Private/Strata/StrataDeferredLighting.ush"
 #endif
 
@@ -223,7 +222,7 @@ void DeferredLightPixelMain(
 					L,
 					ToLight,
 					LightMask,
-					SurfaceShadow,
+					ShadowTerms,
 					RectTexture,
 					Strata.MaterialLobesBuffer,
 					GetStrataPixelDataByteOffset(PixelPos, View.BufferSizeAndInvSize.x, Strata.MaxBytesPerPixel));
diff --git a/Engine/Shaders/Private/HairStrands/HairStrandsMaterialCommon.ush b/Engine/Shaders/Private/HairStrands/HairStrandsMaterialCommon.ush
--- a/Engine/Shaders/Private/HairStrands/HairStrandsMaterialCommon.ush
+++ b/Engine/Shaders/Private/HairStrands/HairStrandsMaterialCommon.ush
@@ -186,27 +186,53 @@ FEvaluateOutput Evaluate(
 		float3 TranslatedWorldPosition = SvPositionToResolvedTranslatedWorld(SvPosition);
 		CalcMaterialParametersEx(MaterialParameters, PixelMaterialInputs, SvPosition, LocalScreenPosition, bIsFrontFace, TranslatedWorldPosition, TranslatedWorldPosition);
 
+#if PROJECT_STRATA && MATERIAL_IS_STRATA && MATERIAL_SHADINGMODEL_HAIR
+		FStrataBSDF HairBSDF = PixelMaterialInputs.FrontMaterial.Layers[0].BSDFs[0];
+
+		OutSample.BaseColor = HAIR_BASECOLOR(HairBSDF);
+		OutSample.Roughness = HAIR_ROUGHNESS(HairBSDF);
+		OutSample.Specular	= HAIR_SPECULAR(HairBSDF);
+		OutSample.Backlit	= HAIR_BACKLIT(HairBSDF);
+		OutSample.Emissive	= BSDF_GETEMISSIVE(HairBSDF);
+		float3 Tangent		= normalize(MaterialParameters.SharedNormals.Normals[BSDF_GETNORMALID(HairBSDF)]);
+		
+	  #if MATERIAL_TANGENTSPACENORMAL 
+		// Convert the tangent space input tangent 
+		// * from a frame with Z/(0,0,1) being up and Y/(0,1,0) being the main tangent direction (called green)
+		// * to   a frame with Y/(0,1,0) being up and Z/(0,0,1) being the main tangent direction (called blue)
+		const float3 TangentSpaceGreen = Tangent;
+		const float3 TangentSpaceBlue = float3(-TangentSpaceGreen.x, TangentSpaceGreen.z, TangentSpaceGreen.y);
+		OutSample.Tangent = mul(TangentSpaceBlue, MaterialParameters.TangentToWorld);
+		OutSample.Tangent = normalize(OutSample.Tangent);
+	  #else // Tangent in World space (likely provided by the HairAttributes node)
+		OutSample.Tangent = Tangent;
+	  #endif
+
+#else
+
 		OutSample.BaseColor = GetMaterialBaseColor(PixelMaterialInputs);
 		OutSample.Roughness = GetMaterialRoughness(PixelMaterialInputs);
-		OutSample.Specular = GetMaterialSpecular(PixelMaterialInputs);
-		OutSample.Backlit = saturate(GetMaterialCustomData0(MaterialParameters));
-		OutSample.Emissive = GetMaterialEmissive(PixelMaterialInputs);
+		OutSample.Specular	= GetMaterialSpecular(PixelMaterialInputs);
+		OutSample.Backlit	= saturate(GetMaterialCustomData0(MaterialParameters));
+		OutSample.Emissive	= GetMaterialEmissive(PixelMaterialInputs);
+		float3 Tangent		= normalize(GetMaterialNormalRaw(PixelMaterialInputs));
 
 	#if HAIRSTRANDS_HAS_NORMAL_CONNECTED
 	  #if MATERIAL_TANGENTSPACENORMAL 
 		// Convert the tangent space input tangent 
 		// * from a frame with Z/(0,0,1) being up and Y/(0,1,0) being the main tangent direction (called green)
 		// * to   a frame with Y/(0,1,0) being up and Z/(0,0,1) being the main tangent direction (called blue)
-		const float3 TangentSpaceGreen = normalize(GetMaterialNormalRaw(PixelMaterialInputs));
+		const float3 TangentSpaceGreen = Tangent;
 		const float3 TangentSpaceBlue = float3(-TangentSpaceGreen.x, TangentSpaceGreen.z, TangentSpaceGreen.y);
 		OutSample.Tangent = mul(TangentSpaceBlue, MaterialParameters.TangentToWorld);
 		OutSample.Tangent = normalize(OutSample.Tangent);
 	  #else // Tangent in World space (likely provided by the HairAttributes node)
-		OutSample.Tangent = normalize(GetMaterialNormalRaw(PixelMaterialInputs));
+		OutSample.Tangent = Tangent;
 	  #endif
 	#else
 		OutSample.Tangent = normalize(Interpolants.TangentToWorld2.xyz);
 	#endif
+#endif
 
 		// this feature is only needed for development/editor - we can compile it out for a shipping build (see r.CompileShadersForDevelopment cvar help)
 	#if USE_DEVELOPMENT_SHADERS
diff --git a/Engine/Shaders/Private/HairStrands/HairStrandsVisibilityPS.usf b/Engine/Shaders/Private/HairStrands/HairStrandsVisibilityPS.usf
--- a/Engine/Shaders/Private/HairStrands/HairStrandsVisibilityPS.usf
+++ b/Engine/Shaders/Private/HairStrands/HairStrandsVisibilityPS.usf
@@ -134,6 +134,18 @@ void MainVisibility(
 		float3 TranslatedWorldPosition = SvPositionToResolvedTranslatedWorld(SvPosition);
 		CalcMaterialParametersEx(MaterialParameters, PixelMaterialInputs, SvPosition, ScreenPosition, bIsFrontFace, TranslatedWorldPosition, TranslatedWorldPosition);
 
+#if PROJECT_STRATA && MATERIAL_IS_STRATA && MATERIAL_SHADINGMODEL_HAIR
+		FStrataBSDF HairBSDF = PixelMaterialInputs.FrontMaterial.Layers[0].BSDFs[0];
+
+		BaseColor		 = HAIR_BASECOLOR(HairBSDF);
+		Roughness		 = HAIR_ROUGHNESS(HairBSDF);
+		half  Metallic	 = HAIR_SCATTER(HairBSDF);
+		Specular		 = HAIR_SPECULAR(HairBSDF);
+		Backlit			 = HAIR_BACKLIT(HairBSDF);
+		float MaterialAO = GetMaterialAmbientOcclusion(PixelMaterialInputs);
+
+		Tangent = mul(normalize(GetMaterialNormalRaw(PixelMaterialInputs)), MaterialParameters.SharedNormals.Normals[BSDF_GETNORMALID(HairBSDF)]);
+#else
 		BaseColor		 = GetMaterialBaseColor(PixelMaterialInputs);
 		Roughness		 = GetMaterialRoughness(PixelMaterialInputs);
 		half  Metallic	 = GetMaterialMetallic(PixelMaterialInputs);
@@ -142,6 +154,7 @@ void MainVisibility(
 		float MaterialAO = GetMaterialAmbientOcclusion(PixelMaterialInputs);
 
 		Tangent = mul(normalize(GetMaterialNormalRaw(PixelMaterialInputs)), MaterialParameters.TangentToWorld);
+#endif
 	}
 	{
 		Velocity = Calculate3DVelocity(PackedVelocityA, PackedVelocityC).xy;
diff --git a/Engine/Shaders/Private/MaterialTemplate.ush b/Engine/Shaders/Private/MaterialTemplate.ush
--- a/Engine/Shaders/Private/MaterialTemplate.ush
+++ b/Engine/Shaders/Private/MaterialTemplate.ush
@@ -18,7 +18,23 @@
 #include "/Engine/Private/HairShadingCommon.ush"
 #include "/Engine/Private/HairStrands/HairCardsAttributeCommon.ush"
 #include "/Engine/Private/HairStrands/HairStrandsAttributeCommon.ush"
+
+#define TEMPLATE_USES_STRATA (PROJECT_STRATA && MATERIAL_IS_STRATA)
+
+#if TEMPLATE_USES_STRATA
 #include "/Engine/Private/Strata/Strata.ush"
+#else
+ // This is a simple way to not have strata included in vertex shader for instance while allowing generated code to compile.
+struct FStrataData
+{
+	uint a;
+};
+FStrataData GetInitialisedStrataData()
+{
+	FStrataData StrataData = (FStrataData)0;
+	return StrataData;
+}
+#endif
 
 #if USES_SPEEDTREE
 	#include "/Engine/Private/SpeedTreeCommon.ush"
@@ -362,7 +378,7 @@ struct FMaterialPixelParameters
 	float3 VolumeSampleConservativeDensity;
 #endif
 
-#if PROJECT_STRATA && MATERIAL_IS_STRATA
+#if TEMPLATE_USES_STRATA
 	FSharedNormals SharedNormals;
 #endif
 };
@@ -2346,7 +2362,7 @@ void CalcPixelMaterialInputs(in out FMaterialPixelParameters Parameters, in out
 	// The Normal is a special case as it might have its own expressions and also be used to calculate other inputs, so perform the assignment here
 %s
 
-#if PROJECT_STRATA && MATERIAL_IS_STRATA
+#if TEMPLATE_USES_STRATA
 	Parameters.SharedNormals = StrataInitialiseSharedNormals();
 #endif
 
diff --git a/Engine/Shaders/Private/ReflectionEnvironmentPixelShader.usf b/Engine/Shaders/Private/ReflectionEnvironmentPixelShader.usf
--- a/Engine/Shaders/Private/ReflectionEnvironmentPixelShader.usf
+++ b/Engine/Shaders/Private/ReflectionEnvironmentPixelShader.usf
@@ -106,8 +106,7 @@ FSkyLightVisibilityData GetSkyLightVisibilityData(FGBufferData GBuffer, float Am
 
 #if STRATA_ENABLED
 #define STRATA_INLINE_SHADING 0
-#include "/Engine/Private/Strata/Strata.ush"
-#include "/Engine/Private/ParticipatingMediaCommon.ush"
+#include "/Engine/Private/Strata/StrataEvaluation.ush"
 #include "/Engine/Private/Strata/StrataEnvironmentLighting.ush"
 #endif
 
diff --git a/Engine/Shaders/Private/ShadowProjectionPixelShader.usf b/Engine/Shaders/Private/ShadowProjectionPixelShader.usf
--- a/Engine/Shaders/Private/ShadowProjectionPixelShader.usf
+++ b/Engine/Shaders/Private/ShadowProjectionPixelShader.usf
@@ -268,6 +268,8 @@ void Main(
 			float Density = SubsurfaceDensityFromOpacity(Opacity);
 			if( GBufferData.ShadingModelID == SHADINGMODELID_HAIR || GBufferData.ShadingModelID == SHADINGMODELID_EYE )
 			{
+				// STRATA_TODO strata hair should also go there. Right now we would need to add SHADINGMODEL_STRATA there 
+				// and also enable HasCustomGBufferData just because we go through the GBuffer. We are going to fixcthat next the right way.
 				Opacity = 1;
 				Density = 1;
 			}
diff --git a/Engine/Shaders/Private/Strata/Strata.ush b/Engine/Shaders/Private/Strata/Strata.ush
--- a/Engine/Shaders/Private/Strata/Strata.ush
+++ b/Engine/Shaders/Private/Strata/Strata.ush
@@ -4,10 +4,6 @@
 
 #include "../Common.ush"
 #include "../BRDF.ush"
-#include "../ShadingCommon.ush"
-#include "../ParticipatingMediaCommon.ush"
-#include "../MonteCarlo.ush"
-#include "../SHCommon.ush"
 #include "../DeferredShadingCommon.ush"
 
 #include "/Engine/Shared/StrataDefinitions.h"
@@ -139,6 +135,12 @@ struct FStrataBSDF
 
 #define UNLIT_TRANSMITTANCE(X)			X.VGPRs[0].xyz
 
+#define HAIR_BASECOLOR(X)				X.VGPRs[0].xyz
+#define HAIR_SCATTER(X)					X.VGPRs[0].w
+#define HAIR_ROUGHNESS(X)				X.VGPRs[1].x
+#define HAIR_SPECULAR(X)				X.VGPRs[1].y
+#define HAIR_BACKLIT(X)					X.VGPRs[1].z
+
 void StrataRequestSharedNormalTangentBasis(inout uint NormalTypes, uint Index);
 
 bool StrataHasAnisotropy(in float Roughness0, float Roughness1)
@@ -348,15 +350,28 @@ FStrataData GetStrataUnlitBSDF(float3 EmissiveColor, float3 TransmittanceColor)
 	return StrataData;
 }
 
-// STRATA_TODO put in a common file
-// Point lobe in off-specular peak direction
-float3 StrataGetOffSpecularPeakReflectionDir(float3 Normal, float3 ReflectionVector, float Roughness)
+FStrataData GetStrataHairBSDF(float3 BaseColor, float Scatter, float Specular, float Roughness, float Backlit, float3 EmissiveColor, uint SharedNormalIndex)
 {
-	float a = Square(Roughness);
-	return lerp(Normal, ReflectionVector, (1 - a) * (sqrt(1 - a) + a));
+	FStrataData StrataData = GetInitialisedStrataData();
+	StrataData.LayerCount++;
+	StrataData.Layers[0].BSDFCount++;
+
+	BSDF_SETTYPE					(StrataData.Layers[0].BSDFs[0], STRATA_BSDF_TYPE_HAIR);
+	BSDF_SETNORMALID				(StrataData.Layers[0].BSDFs[0], SharedNormalIndex);
+	HAIR_BASECOLOR					(StrataData.Layers[0].BSDFs[0]) = BaseColor;
+	HAIR_SCATTER					(StrataData.Layers[0].BSDFs[0]) = Scatter;
+	HAIR_ROUGHNESS					(StrataData.Layers[0].BSDFs[0]) = Roughness;
+	HAIR_SPECULAR					(StrataData.Layers[0].BSDFs[0]) = Specular;
+	HAIR_BACKLIT					(StrataData.Layers[0].BSDFs[0]) = Backlit;
+	BSDF_SETEMISSIVE				(StrataData.Layers[0].BSDFs[0], EmissiveColor);
+
+	StrataData.Layers[0].BSDFs[0].Weight = 1.0f;
+
+	return StrataData;
 }
 
 
+
 ///////////////////////////////////////////////////////////////////////////////
 // Shared normals
 
@@ -448,437 +463,11 @@ float3x3 StrataGetBSDFSharedBasis(in FStrataPixelHeader StrataPixelHeader, in FS
 #endif
 }
 
-///////////////////////////////////////////////////////////////////////////////
-// BSDF evaluate and sampling
-
-struct FStrataBSDFContext
-{
-	FStrataBSDF BSDF;
-
-	float3 N;
-	float3 X;
-	float3 Y;
-	float3 V;
-	float3 R;
-	float3 H;
-
-	float3 L;	// There to initialise the BxDFContext. Only used by StrataEvaluateBSDF, not by StrataImportanceSampleBSDF or StrataEvaluateForEnvLight
-
-	BxDFContext Context;
-	float SatNoL;
-	float SatNoV;
-
-	float3x3 TangentBasis;
-	float3 TangentV;
-	float3 TangentH;
-
-	// Preprocessed roughness override used for specular BSDF evaluation when integration area light (using MRP integrator)
-	float a2Override;
-};
-
-void StrataGetBSDFRoughness(in FStrataBSDF BSDF, inout float R)
-{
-	R = 0;
-	const uint BSDFType = BSDF_GETTYPE(BSDF);
-	switch (BSDFType)
-	{
-	case STRATA_BSDF_TYPE_DIFFUSE:
-	{
-		R = DIFFUSE_ROUGHNESS(BSDF);
-		break;
-	}
-	case STRATA_BSDF_TYPE_DIELECTRIC:
-	{
-		R = DIELECTRIC_ROUGHNESS0(BSDF);
-		break;
-	}
-	case STRATA_BSDF_TYPE_CONDUCTOR:
-	{
-		R = CONDUCTOR_ROUGHNESS0(BSDF);
-		break;
-	}
-	case STRATA_BSDF_TYPE_VOLUME:
-	{
-		R = 0;
-		break;
-	}
-	case STRATA_BSDF_TYPE_SHEEN:
-	{
-		R = SHEEN_ROUGHNESS(BSDF);
-		break;
-	}
-	}
-}
-
-FStrataBSDFContext StrataCreateBSDFContext(FStrataPixelHeader StrataPixelHeader, FStrataBSDF BSDF, float3 V, float3 L, bool bOverrideWithLocalFrame=false)
-{
-	FStrataBSDFContext BSDFContext = (FStrataBSDFContext)0;
-
-	BSDFContext.BSDF = BSDF;
-
-	BSDFContext.V = V;
-	BSDFContext.L = L;
-
-	float3x3 TangentBasis = StrataGetBSDFSharedBasis(StrataPixelHeader, BSDF);
-	BSDFContext.X = TangentBasis[0]; 
-	BSDFContext.Y = TangentBasis[1];
-	BSDFContext.N = TangentBasis[2]; 
-	if (bOverrideWithLocalFrame)
-	{
-		BSDFContext.X = float3(1,0,0);
-		BSDFContext.Y = float3(0,1,0);
-		BSDFContext.N = float3(0,0,1);
-	}
-
-	BSDFContext.R = 2 * dot(BSDFContext.V, BSDFContext.N) * BSDFContext.N - BSDFContext.V;
-	BSDFContext.H = normalize(BSDFContext.V + BSDFContext.L);
-
-	BSDFContext.Context = (BxDFContext)0;
-	if (BSDF_GETANISOTROPY(BSDF) != 0)
-	{
-		Init(BSDFContext.Context, BSDFContext.N, BSDFContext.X, BSDFContext.Y, BSDFContext.V, BSDFContext.L);
-	}
-	else
-	{
-		Init(BSDFContext.Context, BSDFContext.N, BSDFContext.V, BSDFContext.L);
-	}
-	BSDFContext.TangentBasis = float3x3(BSDFContext.X, BSDFContext.Y, BSDFContext.N);
-
-	BSDFContext.TangentV = mul(BSDFContext.TangentBasis, BSDFContext.V);
-	BSDFContext.TangentH = mul(BSDFContext.TangentBasis, BSDFContext.H);
-
-	BSDFContext.SatNoL = saturate(BSDFContext.Context.NoL);
-	BSDFContext.SatNoV = saturate(BSDFContext.Context.NoV);
-
-	return BSDFContext;
-}
-
-void StrataUpdateBSDFContext(inout FStrataBSDFContext BSDFContext, float3 NewL)
-{
-	// Update all the data related to L
-	Init(BSDFContext.Context, BSDFContext.N, BSDFContext.V, NewL);
-	BSDFContext.H = normalize(BSDFContext.V + BSDFContext.L);
-	BSDFContext.TangentH = mul(BSDFContext.TangentBasis, BSDFContext.H);
-}
-
-#define STRATA_PATH_TYPE_DIFFUSE  0x1 
-#define STRATA_PATH_TYPE_SPECULAR 0x2
-#define STRATA_PATH_TYPE_SCATTER  0x4
-
-struct FStrataEvaluateResult
-{
-	float3 Value;
-	float  PDF;
-	float3 Throughput;	// Throughput to the next layer
-	uint   PathType;	// Type of path (specular|diffuse|scattering)
-};
-
-FStrataEvaluateResult StrataEvaluateBSDF(FStrataBSDFContext BSDFContext)
-{
-	FStrataEvaluateResult Sample = (FStrataEvaluateResult)0;
-
-	const float OpaqueBSDFThroughput = 0.0f;
-
-	const uint BSDFType = BSDF_GETTYPE(BSDFContext.BSDF);
-	switch (BSDFType)
-	{
-		case STRATA_BSDF_TYPE_DIFFUSE:
-		{
-			const float SafeRoughness = MakeRoughnessSafe(DIFFUSE_ROUGHNESS(BSDFContext.BSDF));
-			float a2 = Pow4(SafeRoughness);
-
-			Sample.Value		= Diffuse_Chan(DIFFUSE_ALBEDO(BSDFContext.BSDF), a2, BSDFContext.SatNoV, BSDFContext.SatNoL, BSDFContext.Context.VoH, BSDFContext.Context.NoH);
-			Sample.PDF			= BSDFContext.SatNoV / PI;
-			Sample.Throughput	= OpaqueBSDFThroughput;
-			Sample.PathType		= STRATA_PATH_TYPE_DIFFUSE | (BSDF_GETHASSCATTERING(BSDFContext.BSDF) ? STRATA_PATH_TYPE_SCATTER : 0);
-			break;
-		}
-		case STRATA_BSDF_TYPE_DIELECTRIC:
-		case STRATA_BSDF_TYPE_CONDUCTOR:
-		{
-			const bool bIsDieletric = BSDFType == STRATA_BSDF_TYPE_DIELECTRIC;
-			const float3 F0  = bIsDieletric ? DIELECTRIC_TINT(BSDFContext.BSDF) * DielectricIorToF0(DIELECTRIC_IOR(BSDFContext.BSDF)) : CONDUCTOR_REFLECTIVITY(BSDFContext.BSDF);
-			const float3 F90 = bIsDieletric ? DIELECTRIC_TINT(BSDFContext.BSDF) : CONDUCTOR_EDGECOLOR(BSDFContext.BSDF);
-			const float SafeRoughness0 = MakeRoughnessSafe(bIsDieletric ? DIELECTRIC_ROUGHNESS0(BSDFContext.BSDF) : CONDUCTOR_ROUGHNESS0(BSDFContext.BSDF));
-			const float SafeRoughness1 = MakeRoughnessSafe(bIsDieletric ? DIELECTRIC_ROUGHNESS1(BSDFContext.BSDF) : CONDUCTOR_ROUGHNESS1(BSDFContext.BSDF));
-
-			const bool bHasAnisotropy = BSDF_GETANISOTROPY(BSDFContext.BSDF);
-			float D = 0;
-			float Vis = 0;
-			float PDF = 0;
-			if (bHasAnisotropy)
-			{
-				// Generalized microfacet specular
-				float Alpha0 = Pow2(SafeRoughness0);
-				float Alpha1 = Pow2(SafeRoughness1);
-				D = D_GGXaniso(Alpha0, Alpha1, BSDFContext.Context.NoH, BSDFContext.Context.XoH, BSDFContext.Context.YoH);
-				Vis = Vis_SmithJointAniso(Alpha0, Alpha1, BSDFContext.Context.NoV, BSDFContext.Context.NoL, BSDFContext.Context.XoV, BSDFContext.Context.XoL, BSDFContext.Context.YoV, BSDFContext.Context.YoL);
-				PDF = VisibleGGXPDF_aniso(BSDFContext.TangentV, BSDFContext.TangentH, float2(Alpha0, Alpha1));
-			}
-			else
-			{
-				// Special override for roughness for supporting area light integrator with Sphere/Tube/Disk light, which modifies/increase roughness.
-				const float a2 = BSDFContext.a2Override > 0 ? BSDFContext.a2Override : Pow4(SafeRoughness0);
-
-				// Generalized microfacet specular
-				D = D_GGX(a2, BSDFContext.Context.NoH);
-				Vis = Vis_SmithJointApprox(a2, BSDFContext.Context.NoV, BSDFContext.SatNoL);
-				PDF = VisibleGGXPDF(BSDFContext.TangentV, BSDFContext.TangentH, a2);
-			}
-			float3 F = F_Schlick(F0, F90, BSDFContext.Context.VoH);
-
-			Sample.Value		= D * Vis * F;
-			Sample.PDF			= PDF;
-			Sample.Throughput	= bIsDieletric ? (1.0f - F) : OpaqueBSDFThroughput; // STRATA_TODO (Kulla's (1-FDG))
-			Sample.PathType		= STRATA_PATH_TYPE_SPECULAR;
-			break;
-		}
-		case STRATA_BSDF_TYPE_VOLUME:
-		{
-			const float3 Albedo = VOLUME_ALBEDO(BSDFContext.BSDF);
-			const float3 ExtinctionCoefficients = VOLUME_EXTINCTION(BSDFContext.BSDF);
-			const float3 ScatteringCoefficients = Albedo * ExtinctionCoefficients;
-			const float3 AbsorptionCoefficients = ExtinctionCoefficients - ScatteringCoefficients;
-			const float  Thickness = VOLUME_THICKNESS(BSDFContext.BSDF);
-			const float  Phase = HenyeyGreensteinPhase(VOLUME_ANISOTROPY(BSDFContext.BSDF), BSDFContext.Context.VoL);
-
-			const float3 SafeExtinctionThreshold = 0.000001f;
-			const float3 SafeExtinctionCoefficients = max(SafeExtinctionThreshold, ExtinctionCoefficients);
-
-			const float PathLength = Thickness / max(0.0001f, abs(BSDFContext.Context.NoV));
-			const float3 SafePathSegmentTransmittance = exp(-SafeExtinctionCoefficients * PathLength);
-			const float3 NormLuminanceTransfert = 1.0f * ScatteringCoefficients;
-			float3 NormLuminanceSingleScattering = (NormLuminanceTransfert - NormLuminanceTransfert * SafePathSegmentTransmittance) / SafeExtinctionCoefficients;
-
-#if 1
-			const float3 SelfShadowTransmittance = 1.0;// BSDFContext.SatNoL;
-#else
-			const float LightPathLength = Thickness / max(0.0001f, BSDFContext.SatNoL);
-			const float3 SelfShadowTransmittance = BSDFContext.SatNoL <= 0.0 ? 0.0 : exp(-ExtinctionCoefficients * LightPathLength);
-#endif
-
-			Sample.Value		= NormLuminanceSingleScattering * Phase * SelfShadowTransmittance;
-			Sample.PDF			= Phase;
-			Sample.Throughput	= SafePathSegmentTransmittance;
-			Sample.PathType		= STRATA_PATH_TYPE_DIFFUSE;
-			break;
-		}
-		case STRATA_BSDF_TYPE_SHEEN:
-		{
-			const float3 F0  = SHEEN_ALBEDO(BSDFContext.BSDF);
-			const float SafeRoughness = MakeRoughnessSafe(SHEEN_ROUGHNESS(BSDFContext.BSDF));
-
-			// Special override for roughness for supporting area light integrator with Sphere/Tube/Disk light, which modifies/increase roughness.
-			const float a2 = BSDFContext.a2Override > 0 ? BSDFContext.a2Override : Pow4(SafeRoughness);
-
-			// Generalized inverse microfacet specular
-			float D   = D_InvGGX(a2, BSDFContext.Context.NoH);
-			float Vis = Vis_Cloth(BSDFContext.Context.NoV, BSDFContext.SatNoL);
-			float3 F  = F_Schlick(F0, BSDFContext.Context.VoH);
-
-			Sample.Value		= D * Vis * F;
-			Sample.PDF			= BSDFContext.SatNoV / PI; // Per "Production Friendly Microfacet Sheen BRDF", hemispherical sampling give good result as the roughness is usually high.
-			Sample.Throughput	= (1.0f - F);
-			Sample.PathType		= STRATA_PATH_TYPE_DIFFUSE;
-			break;
-		}
-	}
-
-	return Sample;
-}
-
-struct FStrataImportanceSampleResult
-{
-	float3 Direction;
-	float ConeAngle;		// Cone angle is 1/PDF as defined by Lumen
-
-//	FStrataEvaluateResult Eval;
-};
-
-/**
- * Importance sample a Strata BSDF
- * BSDF: the strata BSDF to importance sample
- * E: two random numbers
- * CameraVector: vector from the camera to the considered direction (not V)
- */
-FStrataImportanceSampleResult StrataImportanceSampleBSDF(FStrataPixelHeader StrataPixelHeader, FStrataBSDF BSDF, float2 E, float3 CameraVector)
-{
-	FStrataImportanceSampleResult Sample = (FStrataImportanceSampleResult)0;
-
-	// Create a context from the information we have
-	const float3 UnusedL = float3(0.0f, 0.0f, 1.0f);
-	const float3 V = -CameraVector;
-	FStrataBSDFContext BSDFContext = StrataCreateBSDFContext(StrataPixelHeader, BSDF, V, UnusedL);
-
-	float PDF = 0.0f;
-	const uint BSDFType = BSDF_GETTYPE(BSDF);
-	switch (BSDFType)
-	{
-		case STRATA_BSDF_TYPE_DIFFUSE:
-		{
-			// STRATA_TODO take into account roughness
-			float4 ImportanceSample = CosineSampleHemisphere(E);
-			Sample.Direction = mul(ImportanceSample.xyz, BSDFContext.TangentBasis);
-			PDF = ImportanceSample.w;
-			Sample.ConeAngle = 1.0f / PDF;
-			break;
-		}
-		case STRATA_BSDF_TYPE_DIELECTRIC:
-		case STRATA_BSDF_TYPE_CONDUCTOR:
-		{
-			const bool bIsDielectric = BSDFType == STRATA_BSDF_TYPE_DIELECTRIC;
-			float SafeRoughness0 = MakeRoughnessSafe(bIsDielectric ? DIELECTRIC_ROUGHNESS0(BSDFContext.BSDF) : CONDUCTOR_ROUGHNESS0(BSDFContext.BSDF));
-			float SafeRoughness1 = MakeRoughnessSafe(bIsDielectric ? DIELECTRIC_ROUGHNESS1(BSDFContext.BSDF) : CONDUCTOR_ROUGHNESS1(BSDFContext.BSDF));
-
-			const bool bHasAnisotropy = BSDF_GETANISOTROPY(BSDFContext.BSDF);
-			float4 GGXSample = 0;
-			if (bHasAnisotropy)
-			{
-				GGXSample = ImportanceSampleVisibleGGX_aniso(UniformSampleDisk(E), float2(Pow2(SafeRoughness0), Pow2(SafeRoughness1)), BSDFContext.TangentV);
-			}
-			else
-			{
-				GGXSample = ImportanceSampleVisibleGGX(UniformSampleDisk(E), Pow4(SafeRoughness0), BSDFContext.TangentV);
-			}
-			
-			float3 WorldH = mul(GGXSample.xyz, BSDFContext.TangentBasis);
-			Sample.Direction = reflect(CameraVector, WorldH);
-			PDF = GGXSample.w;
-			Sample.ConeAngle = 1.0f / PDF;
-			break;
-		}
-		case STRATA_BSDF_TYPE_VOLUME:
-		{
-			float4 PhaseSample = ImportanceSampleHenyeyGreensteinPhase(E, VOLUME_ANISOTROPY(BSDFContext.BSDF));
-			Sample.Direction = mul(PhaseSample.xyz, BSDFContext.TangentBasis);
-			PDF = PhaseSample.w;
-			Sample.ConeAngle = 1.0f / PDF;
-			break;
-		}
-		case STRATA_BSDF_TYPE_SHEEN:
-		{
-			// Per "Production Friendly Microfacet Sheen BRDF", uniform hemispherical sampling gives good result as the roughness is 
-			// usually high and avoid issue at grazing angle
-			float4 ImportanceSample = CosineSampleHemisphere(E);
-			Sample.Direction = mul(ImportanceSample.xyz, BSDFContext.TangentBasis);
-			PDF = ImportanceSample.w;
-			Sample.ConeAngle = 1.0f / PDF;
-			break;
-		}
-	}
-
-	//StrataUpdateBSDFContext(BSDFContext, Sample.Direction);
-	//Sample.Eval = StrataEvaluateBSDF(BSDFContext);
-	//Sample.Eval.PDF = PDF;	// We override the pdf to make sure the compiler skips all pdf code from StrataEvaluateBSDF because GGXSample.w si the pdf already for instance.
-
-	return Sample;
-}
-
-struct FStrataEnvLightResult
-{
-	float3 DiffuseAlbedo;
-	float3 DiffuseNormal;
-
-	float3 SpecularDirection;
-	float3 SpecularWeight;
-	float  SpecularSafeRoughness;
-
-	uint PathType;
-};
-
-FStrataEnvLightResult StrataEvaluateForEnvLight(FStrataBSDFContext BSDFContext, bool bEnableSpecular)
-{
-	FStrataEnvLightResult StrataEnvLightResult = (FStrataEnvLightResult)0;
-
-	const uint BSDFType = BSDF_GETTYPE(BSDFContext.BSDF);
-	switch (BSDFType)
-	{
-		case STRATA_BSDF_TYPE_DIFFUSE:
-		{
-			StrataEnvLightResult.DiffuseAlbedo = DIFFUSE_ALBEDO(BSDFContext.BSDF);
-			StrataEnvLightResult.DiffuseNormal = BSDFContext.N;
-			StrataEnvLightResult.PathType = STRATA_PATH_TYPE_DIFFUSE | (BSDF_GETHASSCATTERING(BSDFContext.BSDF) ? STRATA_PATH_TYPE_SCATTER : 0);
-			break;
-		}
-		case STRATA_BSDF_TYPE_DIELECTRIC:
-		case STRATA_BSDF_TYPE_CONDUCTOR:
-		{
-			BRANCH
-			if (bEnableSpecular)
-			{
-				const bool bIsDielectric = STRATA_BSDF_TYPE_DIELECTRIC == BSDFType;
-				const float3 F0  = bIsDielectric ? DIELECTRIC_TINT(BSDFContext.BSDF) * DielectricIorToF0(DIELECTRIC_IOR(BSDFContext.BSDF)) : CONDUCTOR_REFLECTIVITY(BSDFContext.BSDF);
-				const float3 F90 = bIsDielectric ? DIELECTRIC_TINT(BSDFContext.BSDF) : CONDUCTOR_EDGECOLOR(BSDFContext.BSDF);
-				float SafeRoughness0 = MakeRoughnessSafe(bIsDielectric ? DIELECTRIC_ROUGHNESS0(BSDFContext.BSDF) : CONDUCTOR_ROUGHNESS0(BSDFContext.BSDF));
-				float SafeRoughness1 = MakeRoughnessSafe(bIsDielectric ? DIELECTRIC_ROUGHNESS1(BSDFContext.BSDF) : CONDUCTOR_ROUGHNESS1(BSDFContext.BSDF));
-
-				const bool bHasAnisotropy = BSDF_GETANISOTROPY(BSDFContext.BSDF);
-				if (bHasAnisotropy)
-				{
-					// Modified the BSDF normal (and roughness)
-					float Anisotropy = 0;
-					GetAnisotropicFactor(SafeRoughness0, SafeRoughness1, Anisotropy, SafeRoughness0);
-					ModifyGGXAnisotropicNormalRoughness(BSDFContext.X, Anisotropy, SafeRoughness0, BSDFContext.N, BSDFContext.V);
-
-					// Update context (only needs: NoL/SatNoL/R) with the new N
-					BSDFContext.Context.NoL = dot(BSDFContext.N, BSDFContext.L);
-					BSDFContext.SatNoL = saturate(BSDFContext.Context.NoL);
-					BSDFContext.R = 2 * dot(BSDFContext.V, BSDFContext.N) * BSDFContext.N - BSDFContext.V;
-				}
-
-				float3 EvalEnvBRDF = EnvBRDFApprox(F0, F90, SafeRoughness0, BSDFContext.SatNoV);
-
-				StrataEnvLightResult.SpecularDirection = StrataGetOffSpecularPeakReflectionDir(BSDFContext.N, BSDFContext.R, SafeRoughness0);
-				StrataEnvLightResult.SpecularWeight = EvalEnvBRDF;
-				StrataEnvLightResult.SpecularSafeRoughness = SafeRoughness0;
-				StrataEnvLightResult.PathType = STRATA_PATH_TYPE_SPECULAR;
-			}
-			break;
-		}
-		case STRATA_BSDF_TYPE_VOLUME:
-		{
-			const float3 VolumeAlbedo = VOLUME_ALBEDO(BSDFContext.BSDF);
-			const float3 ExtinctionCoefficients = VOLUME_EXTINCTION(BSDFContext.BSDF);
-			const float3 ScatteringCoefficients = VolumeAlbedo * ExtinctionCoefficients;
-			const float3 AbsorptionCoefficients = ExtinctionCoefficients - ScatteringCoefficients;
-			const float  Thickness = VOLUME_THICKNESS(BSDFContext.BSDF);
-
-			const float3 SafeExtinctionThreshold = 0.000001f;
-			const float3 SafeExtinctionCoefficients = max(SafeExtinctionThreshold, ExtinctionCoefficients);
-
-			const float PathLength = Thickness / max(0.0001f, BSDFContext.SatNoV);
-			const float3 SafePathSegmentTransmittance = exp(-SafeExtinctionCoefficients * PathLength);
-			const float3 NormLuminanceTransfert = 1.0f * ScatteringCoefficients;
-			float3 NormLuminanceSingleScattering = (NormLuminanceTransfert - NormLuminanceTransfert * SafePathSegmentTransmittance) / SafeExtinctionCoefficients;
-
-			StrataEnvLightResult.DiffuseAlbedo = NormLuminanceSingleScattering;
-			StrataEnvLightResult.DiffuseNormal = BSDFContext.N;
-			StrataEnvLightResult.PathType = STRATA_PATH_TYPE_DIFFUSE;
-			break;
-		}
-		case STRATA_BSDF_TYPE_SHEEN:
-		{
-			StrataEnvLightResult.DiffuseAlbedo = SHEEN_ALBEDO(BSDFContext.BSDF);;
-			StrataEnvLightResult.DiffuseNormal = BSDFContext.N;
-			StrataEnvLightResult.PathType = STRATA_PATH_TYPE_DIFFUSE;
-			break;
-		}
-	}
-
-	return StrataEnvLightResult;
-}
-
-FThreeBandSHVector StrataBSDFToSH(FStrataBSDFContext BSDFContext)
-{
-	// STRATA_TODO adapt the SH to BSDFs
-	return CalcDiffuseTransferSH3(BSDFContext.N, 1.0f);
-}
 
 
 ///////////////////////////////////////////////////////////////////////////////
 // Operators node
 
-
 FStrataData StrataAdd(FStrataData A, FStrataData B)
 {
 	FStrataData StrataData = A;
@@ -990,6 +579,16 @@ void SanitizeStrataSheen(inout FStrataBSDF BSDF)
 	SHEEN_ROUGHNESS(BSDF) = saturate(SHEEN_ROUGHNESS(BSDF));
 }
 
+void SanitizeStrataHair(inout FStrataBSDF BSDF)
+{
+	HAIR_BASECOLOR(BSDF) = saturate(HAIR_BASECOLOR(BSDF));
+	HAIR_SCATTER(BSDF) = saturate(HAIR_SCATTER(BSDF));
+	HAIR_ROUGHNESS(BSDF) = saturate(HAIR_ROUGHNESS(BSDF));
+	HAIR_SPECULAR(BSDF) = saturate(HAIR_SPECULAR(BSDF));
+	HAIR_BACKLIT(BSDF) = saturate(HAIR_BACKLIT(BSDF));
+	BSDF_SETEMISSIVE(BSDF, clamp(BSDF_GETEMISSIVE(BSDF), 0.0, Max10BitsFloat));
+}
+
 void StrataSanitizeBSDF(inout FStrataBSDF BSDF)
 {
 	switch (BSDF_GETTYPE(BSDF))
@@ -1019,10 +618,16 @@ void StrataSanitizeBSDF(inout FStrataBSDF BSDF)
 			SanitizeStrataSheen(BSDF);
 			break;
 		}
+		case STRATA_BSDF_TYPE_HAIR:
+		{
+			SanitizeStrataHair(BSDF);
+			break;
+		}
 	}
 }
 
 
+
 ///////////////////////////////////////////////////////////////////////////////
 // BSDFs packing/unpacking
 // 11G11B10F and 10F are unsigned float format
@@ -1512,6 +1117,20 @@ void PackStrataOut(
 				break;
 			}
 
+			case STRATA_BSDF_TYPE_HAIR:
+			{
+				// No SSR
+				if (bTopLayer)
+				{
+					AccumRoughness		+= BSDFCoverageAvg * FullyRough; // Do not trigger SSR by using fully rough
+					AccumRoughnessCount += BSDFCoverageAvg;
+				}
+				// Hair are considered opaque
+				LayerCoverage	+= BSDFCoverage;
+				LayerThroughput += BSDFCoverage * OpaqueBSDFThroughput;
+				break;
+			}
+
 			} // switch
 
 			// Update the weights to take into account throughput towards the view from layers above
@@ -1659,6 +1278,13 @@ void PackStrataOut(
 				// 4 bytes
 			}
 			break;
+			case STRATA_BSDF_TYPE_HAIR:
+			{
+				STRATA_STORE_UINT1(PackAlbedoLinearToGamma2AlphaLinear(float4(HAIR_BASECOLOR(BSDF), HAIR_ROUGHNESS(BSDF))));
+				STRATA_STORE_UINT1(PackR11G11B10F(float3(HAIR_SCATTER(BSDF), HAIR_SPECULAR(BSDF), HAIR_BACKLIT(BSDF))));
+				// 8 bytes
+			}
+			break;
 			}
 
 			// Stop writing if above budget.
@@ -1775,6 +1401,13 @@ FStrataBSDF UnpackStrataBSDFIn(ByteAddressBuffer StrataBuffer, inout uint ByteOf
 		StrataLoad_AlbedoGamma2ToLinearAlphaLinear(StrataBuffer, ByteOffset, SHEEN_ALBEDO(OutBSDF), SHEEN_ROUGHNESS(OutBSDF));
 	}
 	break;
+	case STRATA_BSDF_TYPE_HAIR:
+	{
+		StrataLoad_AlbedoGamma2ToLinearAlphaLinear(StrataBuffer, ByteOffset, HAIR_BASECOLOR(OutBSDF), HAIR_ROUGHNESS(OutBSDF));
+		StrataLoad_R11_G11_B10F(StrataBuffer, ByteOffset, HAIR_SCATTER(OutBSDF), HAIR_SPECULAR(OutBSDF), HAIR_BACKLIT(OutBSDF));
+		// 8 bytes
+	}
+	break;
 	}
 
 	return OutBSDF;
diff --git a/Engine/Shaders/Private/Strata/StrataDeferredLighting.ush b/Engine/Shaders/Private/Strata/StrataDeferredLighting.ush
--- a/Engine/Shaders/Private/Strata/StrataDeferredLighting.ush
+++ b/Engine/Shaders/Private/Strata/StrataDeferredLighting.ush
@@ -6,19 +6,19 @@
 
 struct FStrataLightIntegrationContext
 {
-	float3 IncidentRadiance;
-	float3 IncidentIrradiance;	// Modified incident irradiance, based on custom are light integrator
-	float3 OutgoingRadiance;	// Continue the outgoing radiance if the BSDF evaluation/integration has been done with the custom area light integrator
-	bool bEvaluateBSDF;			// True if the BSDF need to be evalued, false, if the integration has already been done
+	float3 IncidentRadianceNoShadow;
+	float3 IncidentIrradianceNoShadow;	// Modified incident irradiance, based on custom are light integrator
+	float3 OutgoingRadianceNoShadow;	// Continue the outgoing radiance if the BSDF evaluation/integration has been done with the custom area light integrator
+	bool bEvaluateBSDF;					// True if the BSDF need to be evalued, false, if the integration has already been done
 };
 
 FStrataLightIntegrationContext StrataIniLightIntegrationContext()
 {
 	FStrataLightIntegrationContext Out;
 	Out.bEvaluateBSDF = true;
-	Out.OutgoingRadiance = 0;
-	Out.IncidentIrradiance = 0;
-	Out.IncidentRadiance = 0;
+	Out.OutgoingRadianceNoShadow = 0;
+	Out.IncidentIrradianceNoShadow = 0;
+	Out.IncidentRadianceNoShadow = 0;
 	return Out;
 }
 
@@ -32,7 +32,7 @@ void StrataPreprocessAreaLightIntegration(
 	inout FStrataLightIntegrationContext Out)
 {
 	Out.bEvaluateBSDF = true;
-	Out.OutgoingRadiance = 0;
+	Out.OutgoingRadianceNoShadow = 0;
 
 	// Special area lighting integration
 	bool bIsValid = true;
@@ -88,9 +88,9 @@ void StrataPreprocessAreaLightIntegration(
 			{
 				// Rect light uses LTC integrator, which computes the full BDSD x Area light integration. 
 				// This is why there is no need to evaluate the BSDF. The result is stored into a local 
-				// OutgoingRadiance, which will be accumulated outside of this function
-				Out.OutgoingRadiance.rgb += 
-					Out.IncidentRadiance * 
+				// OutgoingRadianceNoShadow, which will be accumulated outside of this function
+				Out.OutgoingRadianceNoShadow.rgb +=
+					Out.IncidentRadianceNoShadow *
 					RectGGXApproxLTC(SafeRoughness0, SpecularColor, BSDFContext.N, BSDFContext.V, LightContext.AreaLight.Rect, LightContext.AreaLight.Texture);
 				Out.bEvaluateBSDF = false;
 			}
@@ -108,8 +108,8 @@ void StrataPreprocessAreaLightIntegration(
 				// Update incident irradiance to add the energy rescaling to due MRP integration
 				// The a2 override stored the overriden roughness, which will be used during the BSDF evaluation
 				BSDFContext.a2Override = Pow4(SafeRoughness0);
-				Out.IncidentIrradiance =
-					Out.IncidentRadiance *
+				Out.IncidentIrradianceNoShadow =
+					Out.IncidentRadianceNoShadow *
 					EnergyNormalization(BSDFContext.a2Override, BSDFContext.Context.VoH, LightContext.AreaLight) *
 					LightContext.AreaLight.FalloffColor * (LightContext.Falloff * LightContext.NoL);
 				Out.bEvaluateBSDF = true;
@@ -127,7 +127,7 @@ float4 StrataDeferredLighting(
 	float3 L,
 	float3 ToLight,
 	float LightMask,
-	float SurfaceShadow,
+	FShadowTerms ShadowTerms,
 	FRectTexture RectTexture,
 	ByteAddressBuffer StrataDataBuffer,
 	uint PixelStrataDataByteOffset)
@@ -157,32 +157,32 @@ float4 StrataDeferredLighting(
 
 		// Evaluate incident lighting & (optionally) precompute area light integration
 		FStrataLightIntegrationContext IntegrationContext = StrataIniLightIntegrationContext();
-		IntegrationContext.IncidentRadiance = LightData.Color * LightMask * SurfaceShadow;
-		IntegrationContext.IncidentIrradiance = PreintegratedIndicentIrradiance * BSDFContext.SatNoL * IntegrationContext.IncidentRadiance;
+		IntegrationContext.IncidentRadianceNoShadow   = LightData.Color * LightMask;
+		IntegrationContext.IncidentIrradianceNoShadow = PreintegratedIndicentIrradiance * IntegrationContext.IncidentRadianceNoShadow * StrataBSDFProjectedSolidAngleFactor(BSDFContext);
 		StrataPreprocessAreaLightIntegration(ToLight, LightData, RectTexture, BSDFContext, IntegrationContext);
 
 		// Evaluate BSDF
 		if (IntegrationContext.bEvaluateBSDF)
 		{
-			FStrataEvaluateResult BSDFEvaluate = StrataEvaluateBSDF(BSDFContext);
+			FStrataEvaluateResult BSDFEvaluate = StrataEvaluateBSDF(BSDFContext, ShadowTerms);
 
-			const float3 DiffuseRadiance  = (BSDFEvaluate.PathType & STRATA_PATH_TYPE_DIFFUSE)  != 0 ? BSDF.Weight * BSDFEvaluate.Value : 0;
-			const float3 SpecularRadiance = (BSDFEvaluate.PathType & STRATA_PATH_TYPE_SPECULAR) != 0 ? BSDF.Weight * BSDFEvaluate.Value : 0;
+			const float3 DiffusePathWeight		= BSDF.Weight * BSDFEvaluate.DiffusePathValue;
+			const float3 SpecularPathWeight		= BSDF.Weight * BSDFEvaluate.SpecularPathValue;
+			const float3 TransmissionPathWeight	= BSDF.Weight * BSDFEvaluate.TransmissionPathValue;
 
-			float3 CommonFactor = IntegrationContext.IncidentIrradiance;
-			const bool bNeedsSeparateSubsurfaceLightAccumulation = (BSDFEvaluate.PathType & STRATA_PATH_TYPE_SCATTER) != 0;
-			LightAccumulator_AddSplit(Out, DiffuseRadiance, SpecularRadiance, DiffuseRadiance, CommonFactor, bNeedsSeparateSubsurfaceLightAccumulation);
+			const bool bNeedsSeparateSubsurfaceLightAccumulation = BSDFEvaluate.bSubsurface;
+
+			LightAccumulator_AddSplit(Out, DiffusePathWeight,		SpecularPathWeight, DiffusePathWeight,		IntegrationContext.IncidentIrradianceNoShadow * ShadowTerms.SurfaceShadow,		bNeedsSeparateSubsurfaceLightAccumulation);
+			LightAccumulator_AddSplit(Out, TransmissionPathWeight,	0.0f,				TransmissionPathWeight,	IntegrationContext.IncidentIrradianceNoShadow * ShadowTerms.TransmissionShadow,	bNeedsSeparateSubsurfaceLightAccumulation);
 		}
 		else
 		{
-			
 			// STRATA_TODO Find a better solution for special of are lighting which handles specular path only, so we can hardcode weight.
 			const float3 DiffuseRadiance  = 0;
 			const float3 SpecularRadiance = 1;
-
-			float3 CommonFactor = IntegrationContext.OutgoingRadiance;
 			const bool bNeedsSeparateSubsurfaceLightAccumulation = false;
-			LightAccumulator_AddSplit(Out, DiffuseRadiance, SpecularRadiance, DiffuseRadiance, CommonFactor, bNeedsSeparateSubsurfaceLightAccumulation);
+
+			LightAccumulator_AddSplit(Out, DiffuseRadiance,			SpecularRadiance,	DiffuseRadiance,		IntegrationContext.OutgoingRadianceNoShadow * ShadowTerms.SurfaceShadow,		bNeedsSeparateSubsurfaceLightAccumulation);
 		}
 	}
 
diff --git a/Engine/Shaders/Private/Strata/StrataEnvironmentLighting.ush b/Engine/Shaders/Private/Strata/StrataEnvironmentLighting.ush
--- a/Engine/Shaders/Private/Strata/StrataEnvironmentLighting.ush
+++ b/Engine/Shaders/Private/Strata/StrataEnvironmentLighting.ush
@@ -68,7 +68,7 @@ float4 StrataEnvironmentLighting(
 				* AmbientOcclusion * (BSDF_GETISTOPLAYER(BSDF) ? TopLayerSpecularContributionFactor : 1.0f);
 		}
 
-		const bool bNeedsSeparateSubsurfaceLightAccumulation = (StrataEnvLight.PathType & STRATA_PATH_TYPE_SCATTER) != 0;
+		const bool bNeedsSeparateSubsurfaceLightAccumulation = StrataEnvLight.bSubsurface;
 		LightAccumulator_AddSplit(Out, DiffuseSkyLighting, SpecularSkyLighting, DiffuseSkyLighting, CommonFactor, bNeedsSeparateSubsurfaceLightAccumulation);
 
 		// STRATA_TODO: Missing probe reflection here
diff --git a/Engine/Shaders/Private/Strata/StrataEvaluation.ush b/Engine/Shaders/Private/Strata/StrataEvaluation.ush
new file mode 100644
--- /dev/null
+++ b/Engine/Shaders/Private/Strata/StrataEvaluation.ush
@@ -0,0 +1,526 @@
+// Copyright Epic Games, Inc. All Rights Reserved.
+
+#pragma once
+
+#include "Strata.ush"
+
+#include "../ParticipatingMediaCommon.ush"
+#include "../MonteCarlo.ush"
+#include "../SHCommon.ush"
+#include "../ShadingModels.ush"
+
+
+
+// STRATA_TODO put in a common file
+// Point lobe in off-specular peak direction
+float3 StrataGetOffSpecularPeakReflectionDir(float3 Normal, float3 ReflectionVector, float Roughness)
+{
+	float a = Square(Roughness);
+	return lerp(Normal, ReflectionVector, (1 - a) * (sqrt(1 - a) + a));
+}
+
+
+
+///////////////////////////////////////////////////////////////////////////////
+// BSDF evaluate and sampling
+
+struct FStrataBSDFContext
+{
+	FStrataBSDF BSDF;
+
+	float3 N;
+	float3 X;
+	float3 Y;
+	float3 V;
+	float3 R;
+	float3 H;
+
+	float3 L;	// There to initialise the BxDFContext. Only used by StrataEvaluateBSDF, not by StrataImportanceSampleBSDF or StrataEvaluateForEnvLight
+
+	BxDFContext Context;
+	float SatNoL;
+	float SatNoV;
+
+	float3x3 TangentBasis;
+	float3 TangentV;
+	float3 TangentH;
+
+	// Preprocessed roughness override used for specular BSDF evaluation when integration area light (using MRP integrator)
+	float a2Override;
+};
+
+void StrataGetBSDFRoughness(in FStrataBSDF BSDF, inout float R)
+{
+	R = 0;
+	const uint BSDFType = BSDF_GETTYPE(BSDF);
+	switch (BSDFType)
+	{
+	case STRATA_BSDF_TYPE_DIFFUSE:
+	{
+		R = DIFFUSE_ROUGHNESS(BSDF);
+		break;
+	}
+	case STRATA_BSDF_TYPE_DIELECTRIC:
+	{
+		R = DIELECTRIC_ROUGHNESS0(BSDF);
+		break;
+	}
+	case STRATA_BSDF_TYPE_CONDUCTOR:
+	{
+		R = CONDUCTOR_ROUGHNESS0(BSDF);
+		break;
+	}
+	case STRATA_BSDF_TYPE_VOLUME:
+	{
+		R = 0;
+		break;
+	}
+	case STRATA_BSDF_TYPE_SHEEN:
+	{
+		R = SHEEN_ROUGHNESS(BSDF);
+		break;
+	}
+	case STRATA_BSDF_TYPE_HAIR:
+	{
+		R = HAIR_ROUGHNESS(BSDF);
+		break;
+	}
+	}
+}
+
+FStrataBSDFContext StrataCreateBSDFContext(FStrataPixelHeader StrataPixelHeader, FStrataBSDF BSDF, float3 V, float3 L, bool bOverrideWithLocalFrame=false)
+{
+	FStrataBSDFContext BSDFContext = (FStrataBSDFContext)0;
+
+	BSDFContext.BSDF = BSDF;
+
+	BSDFContext.V = V;
+	BSDFContext.L = L;
+
+	float3x3 TangentBasis = StrataGetBSDFSharedBasis(StrataPixelHeader, BSDF);
+	BSDFContext.X = TangentBasis[0]; 
+	BSDFContext.Y = TangentBasis[1];
+	BSDFContext.N = TangentBasis[2]; 
+	if (bOverrideWithLocalFrame)
+	{
+		BSDFContext.X = float3(1,0,0);
+		BSDFContext.Y = float3(0,1,0);
+		BSDFContext.N = float3(0,0,1);
+	}
+
+	BSDFContext.R = 2 * dot(BSDFContext.V, BSDFContext.N) * BSDFContext.N - BSDFContext.V;
+	BSDFContext.H = normalize(BSDFContext.V + BSDFContext.L);
+
+	BSDFContext.Context = (BxDFContext)0;
+	if (BSDF_GETANISOTROPY(BSDF) != 0)
+	{
+		Init(BSDFContext.Context, BSDFContext.N, BSDFContext.X, BSDFContext.Y, BSDFContext.V, BSDFContext.L);
+	}
+	else
+	{
+		Init(BSDFContext.Context, BSDFContext.N, BSDFContext.V, BSDFContext.L);
+	}
+	BSDFContext.TangentBasis = float3x3(BSDFContext.X, BSDFContext.Y, BSDFContext.N);
+
+	BSDFContext.TangentV = mul(BSDFContext.TangentBasis, BSDFContext.V);
+	BSDFContext.TangentH = mul(BSDFContext.TangentBasis, BSDFContext.H);
+
+	BSDFContext.SatNoL = saturate(BSDFContext.Context.NoL);
+	BSDFContext.SatNoV = saturate(BSDFContext.Context.NoV);
+
+	return BSDFContext;
+}
+
+void StrataUpdateBSDFContext(inout FStrataBSDFContext BSDFContext, float3 NewL)
+{
+	// Update all the data related to L
+	Init(BSDFContext.Context, BSDFContext.N, BSDFContext.V, NewL);
+	BSDFContext.H = normalize(BSDFContext.V + BSDFContext.L);
+	BSDFContext.TangentH = mul(BSDFContext.TangentBasis, BSDFContext.H);
+}
+
+struct FStrataEvaluateResult
+{
+	float3 DiffusePathValue;
+	float3 SpecularPathValue;
+	float3 TransmissionPathValue;
+	float  PDF;
+	float3 Throughput;					// Throughput to the next layer
+	bool   bSubsurface;					// True if we need to separate the subsurface light contribution for the screen space diffusion process.
+	bool   bApplyProjectedSolidAngle;	// True if the Saturate(NoL) factor should be applied or not over the value of the BSDF
+};
+
+float StrataBSDFProjectedSolidAngleFactor(in FStrataBSDFContext BSDFContext)
+{
+	const uint BSDFType = BSDF_GETTYPE(BSDFContext.BSDF);
+	return BSDFType == STRATA_BSDF_TYPE_HAIR ? 1.0f : BSDFContext.SatNoL;
+}
+
+FStrataEvaluateResult StrataEvaluateBSDF(FStrataBSDFContext BSDFContext, FShadowTerms ShadowTerms)
+{
+	FStrataEvaluateResult Sample = (FStrataEvaluateResult)0;
+
+	const float OpaqueBSDFThroughput = 0.0f;
+
+	const uint BSDFType = BSDF_GETTYPE(BSDFContext.BSDF);
+	switch (BSDFType)
+	{
+		case STRATA_BSDF_TYPE_DIFFUSE:
+		{
+			const float SafeRoughness = MakeRoughnessSafe(DIFFUSE_ROUGHNESS(BSDFContext.BSDF));
+			float a2 = Pow4(SafeRoughness);
+
+			Sample.DiffusePathValue	= Diffuse_Chan(DIFFUSE_ALBEDO(BSDFContext.BSDF), a2, BSDFContext.SatNoV, BSDFContext.SatNoL, BSDFContext.Context.VoH, BSDFContext.Context.NoH);
+			Sample.PDF				= BSDFContext.SatNoV / PI;
+			Sample.Throughput		= OpaqueBSDFThroughput;
+			Sample.bSubsurface		= BSDF_GETHASSCATTERING(BSDFContext.BSDF) > 0;
+			break;
+		}
+		case STRATA_BSDF_TYPE_DIELECTRIC:
+		case STRATA_BSDF_TYPE_CONDUCTOR:
+		{
+			const bool bIsDieletric = BSDFType == STRATA_BSDF_TYPE_DIELECTRIC;
+			const float3 F0  = bIsDieletric ? DIELECTRIC_TINT(BSDFContext.BSDF) * DielectricIorToF0(DIELECTRIC_IOR(BSDFContext.BSDF)) : CONDUCTOR_REFLECTIVITY(BSDFContext.BSDF);
+			const float3 F90 = bIsDieletric ? DIELECTRIC_TINT(BSDFContext.BSDF) : CONDUCTOR_EDGECOLOR(BSDFContext.BSDF);
+			const float SafeRoughness0 = MakeRoughnessSafe(bIsDieletric ? DIELECTRIC_ROUGHNESS0(BSDFContext.BSDF) : CONDUCTOR_ROUGHNESS0(BSDFContext.BSDF));
+			const float SafeRoughness1 = MakeRoughnessSafe(bIsDieletric ? DIELECTRIC_ROUGHNESS1(BSDFContext.BSDF) : CONDUCTOR_ROUGHNESS1(BSDFContext.BSDF));
+
+			const bool bHasAnisotropy = BSDF_GETANISOTROPY(BSDFContext.BSDF);
+			float D = 0;
+			float Vis = 0;
+			float PDF = 0;
+			if (bHasAnisotropy)
+			{
+				// Generalized microfacet specular
+				float Alpha0 = Pow2(SafeRoughness0);
+				float Alpha1 = Pow2(SafeRoughness1);
+				D = D_GGXaniso(Alpha0, Alpha1, BSDFContext.Context.NoH, BSDFContext.Context.XoH, BSDFContext.Context.YoH);
+				Vis = Vis_SmithJointAniso(Alpha0, Alpha1, BSDFContext.Context.NoV, BSDFContext.Context.NoL, BSDFContext.Context.XoV, BSDFContext.Context.XoL, BSDFContext.Context.YoV, BSDFContext.Context.YoL);
+				PDF = VisibleGGXPDF_aniso(BSDFContext.TangentV, BSDFContext.TangentH, float2(Alpha0, Alpha1));
+			}
+			else
+			{
+				// Special override for roughness for supporting area light integrator with Sphere/Tube/Disk light, which modifies/increase roughness.
+				const float a2 = BSDFContext.a2Override > 0 ? BSDFContext.a2Override : Pow4(SafeRoughness0);
+
+				// Generalized microfacet specular
+				D = D_GGX(a2, BSDFContext.Context.NoH);
+				Vis = Vis_SmithJointApprox(a2, BSDFContext.Context.NoV, BSDFContext.SatNoL);
+				PDF = VisibleGGXPDF(BSDFContext.TangentV, BSDFContext.TangentH, a2);
+			}
+			float3 F = F_Schlick(F0, F90, BSDFContext.Context.VoH);
+
+			Sample.SpecularPathValue= D * Vis * F;
+			Sample.PDF				= PDF;
+			Sample.Throughput		= bIsDieletric ? (1.0f - F) : OpaqueBSDFThroughput; // STRATA_TODO (Kulla's (1-FDG))
+			break;
+		}
+		case STRATA_BSDF_TYPE_VOLUME:
+		{
+			const float3 Albedo = VOLUME_ALBEDO(BSDFContext.BSDF);
+			const float3 ExtinctionCoefficients = VOLUME_EXTINCTION(BSDFContext.BSDF);
+			const float3 ScatteringCoefficients = Albedo * ExtinctionCoefficients;
+			const float3 AbsorptionCoefficients = ExtinctionCoefficients - ScatteringCoefficients;
+			const float  Thickness = VOLUME_THICKNESS(BSDFContext.BSDF);
+			const float  Phase = HenyeyGreensteinPhase(VOLUME_ANISOTROPY(BSDFContext.BSDF), BSDFContext.Context.VoL);
+
+			const float3 SafeExtinctionThreshold = 0.000001f;
+			const float3 SafeExtinctionCoefficients = max(SafeExtinctionThreshold, ExtinctionCoefficients);
+
+			const float PathLength = Thickness / max(0.0001f, abs(BSDFContext.Context.NoV));
+			const float3 SafePathSegmentTransmittance = exp(-SafeExtinctionCoefficients * PathLength);
+			const float3 NormLuminanceTransfert = 1.0f * ScatteringCoefficients;
+			float3 NormLuminanceSingleScattering = (NormLuminanceTransfert - NormLuminanceTransfert * SafePathSegmentTransmittance) / SafeExtinctionCoefficients;
+
+#if 1
+			const float3 SelfShadowTransmittance = 1.0;// BSDFContext.SatNoL;
+#else
+			const float LightPathLength = Thickness / max(0.0001f, BSDFContext.SatNoL);
+			const float3 SelfShadowTransmittance = BSDFContext.SatNoL <= 0.0 ? 0.0 : exp(-ExtinctionCoefficients * LightPathLength);
+#endif
+
+			Sample.DiffusePathValue = NormLuminanceSingleScattering * Phase * SelfShadowTransmittance;
+			Sample.PDF				= Phase;
+			Sample.Throughput		= SafePathSegmentTransmittance;
+			break;
+		}
+		case STRATA_BSDF_TYPE_SHEEN:
+		{
+			const float3 F0  = SHEEN_ALBEDO(BSDFContext.BSDF);
+			const float SafeRoughness = MakeRoughnessSafe(SHEEN_ROUGHNESS(BSDFContext.BSDF));
+
+			// Special override for roughness for supporting area light integrator with Sphere/Tube/Disk light, which modifies/increase roughness.
+			const float a2 = BSDFContext.a2Override > 0 ? BSDFContext.a2Override : Pow4(SafeRoughness);
+
+			// Generalized inverse microfacet specular
+			float D   = D_InvGGX(a2, BSDFContext.Context.NoH);
+			float Vis = Vis_Cloth(BSDFContext.Context.NoV, BSDFContext.SatNoL);
+			float3 F  = F_Schlick(F0, BSDFContext.Context.VoH);
+
+			Sample.DiffusePathValue	= D * Vis * F;
+			Sample.PDF				= BSDFContext.SatNoV / PI; // Per "Production Friendly Microfacet Sheen BRDF", hemispherical sampling give good result as the roughness is usually high.
+			Sample.Throughput		= (1.0f - F);
+			break;
+		}
+		case STRATA_BSDF_TYPE_HAIR:
+		{
+			FGBufferData GBuffer	= (FGBufferData)0;
+			GBuffer.BaseColor		= HAIR_BASECOLOR(BSDFContext.BSDF);
+			GBuffer.Specular		= HAIR_SPECULAR(BSDFContext.BSDF);
+			GBuffer.Roughness		= HAIR_ROUGHNESS(BSDFContext.BSDF);
+			GBuffer.Metallic		= HAIR_SCATTER(BSDFContext.BSDF);
+			GBuffer.CustomData.z	= HAIR_BACKLIT(BSDFContext.BSDF);
+			GBuffer.ShadingModelID	= SHADINGMODELID_HAIR;
+			GBuffer.WorldNormal		= BSDFContext.N;
+
+			float BacklitEnabled = 1.0f;
+			float Area = 0.0f;
+			uint2 Random = uint2(0, 0);
+			Sample.TransmissionPathValue= HairShading(GBuffer, BSDFContext.L, BSDFContext.V, BSDFContext.N, ShadowTerms.TransmissionShadow, ShadowTerms.HairTransmittance, BacklitEnabled, Area, Random);
+			Sample.PDF					= 1.0f / (4.0f * PI);		// STRATA_TODO this currently match the uniform sphere sampling from StrataImportanceSampleBSDF
+			Sample.Throughput			= OpaqueBSDFThroughput; 
+		}
+		break;
+	}
+
+	return Sample;
+}
+
+
+FStrataEvaluateResult StrataEvaluateBSDF(FStrataBSDFContext BSDFContext)
+{
+	FShadowTerms Shadow = { 1, 1, 1, InitHairTransmittanceData() };
+	return StrataEvaluateBSDF(BSDFContext, Shadow);
+}
+
+struct FStrataImportanceSampleResult
+{
+	float3 Direction;
+	float ConeAngle;		// Cone angle is 1/PDF as defined by Lumen
+
+//	FStrataEvaluateResult Eval;
+};
+
+/**
+ * Importance sample a Strata BSDF
+ * BSDF: the strata BSDF to importance sample
+ * E: two random numbers
+ * CameraVector: vector from the camera to the considered direction (not V)
+ */
+FStrataImportanceSampleResult StrataImportanceSampleBSDF(FStrataPixelHeader StrataPixelHeader, FStrataBSDF BSDF, float2 E, float3 CameraVector)
+{
+	FStrataImportanceSampleResult Sample = (FStrataImportanceSampleResult)0;
+
+	// Create a context from the information we have
+	const float3 UnusedL = float3(0.0f, 0.0f, 1.0f);
+	const float3 V = -CameraVector;
+	FStrataBSDFContext BSDFContext = StrataCreateBSDFContext(StrataPixelHeader, BSDF, V, UnusedL);
+
+	float PDF = 0.0f;
+	const uint BSDFType = BSDF_GETTYPE(BSDF);
+	switch (BSDFType)
+	{
+		case STRATA_BSDF_TYPE_DIFFUSE:
+		{
+			// STRATA_TODO take into account roughness
+			float4 ImportanceSample = CosineSampleHemisphere(E);
+			Sample.Direction = mul(ImportanceSample.xyz, BSDFContext.TangentBasis);
+			PDF = ImportanceSample.w;
+			Sample.ConeAngle = 1.0f / PDF;
+			break;
+		}
+		case STRATA_BSDF_TYPE_DIELECTRIC:
+		case STRATA_BSDF_TYPE_CONDUCTOR:
+		{
+			const bool bIsDielectric = BSDFType == STRATA_BSDF_TYPE_DIELECTRIC;
+			float SafeRoughness0 = MakeRoughnessSafe(bIsDielectric ? DIELECTRIC_ROUGHNESS0(BSDFContext.BSDF) : CONDUCTOR_ROUGHNESS0(BSDFContext.BSDF));
+			float SafeRoughness1 = MakeRoughnessSafe(bIsDielectric ? DIELECTRIC_ROUGHNESS1(BSDFContext.BSDF) : CONDUCTOR_ROUGHNESS1(BSDFContext.BSDF));
+
+			const bool bHasAnisotropy = BSDF_GETANISOTROPY(BSDFContext.BSDF);
+			float4 GGXSample = 0;
+			if (bHasAnisotropy)
+			{
+				GGXSample = ImportanceSampleVisibleGGX_aniso(UniformSampleDisk(E), float2(Pow2(SafeRoughness0), Pow2(SafeRoughness1)), BSDFContext.TangentV);
+			}
+			else
+			{
+				GGXSample = ImportanceSampleVisibleGGX(UniformSampleDisk(E), Pow4(SafeRoughness0), BSDFContext.TangentV);
+			}
+			
+			float3 WorldH = mul(GGXSample.xyz, BSDFContext.TangentBasis);
+			Sample.Direction = reflect(CameraVector, WorldH);
+			PDF = GGXSample.w;
+			Sample.ConeAngle = 1.0f / PDF;
+			break;
+		}
+		case STRATA_BSDF_TYPE_VOLUME:
+		{
+			float4 PhaseSample = ImportanceSampleHenyeyGreensteinPhase(E, VOLUME_ANISOTROPY(BSDFContext.BSDF));
+			Sample.Direction = mul(PhaseSample.xyz, BSDFContext.TangentBasis);
+			PDF = PhaseSample.w;
+			Sample.ConeAngle = 1.0f / PDF;
+			break;
+		}
+		case STRATA_BSDF_TYPE_SHEEN:
+		{
+			// Per "Production Friendly Microfacet Sheen BRDF", uniform hemispherical sampling gives good result as the roughness is 
+			// usually high and avoid issue at grazing angle
+			float4 ImportanceSample = CosineSampleHemisphere(E);
+			Sample.Direction = mul(ImportanceSample.xyz, BSDFContext.TangentBasis);
+			PDF = ImportanceSample.w;
+			Sample.ConeAngle = 1.0f / PDF;
+			break;
+		}
+		case STRATA_BSDF_TYPE_HAIR:
+		{
+			// STRATA_TODO do something better when we get there with Lumen, and evaluate the different researched solution (e.g. Importance Sampling for Physically-Based Hair Fiber Models)
+			UniformSampleSphere(E);
+			break;
+		}
+	}
+
+	//StrataUpdateBSDFContext(BSDFContext, Sample.Direction);
+	//Sample.Eval = StrataEvaluateBSDF(BSDFContext);
+	//Sample.Eval.PDF = PDF;	// We override the pdf to make sure the compiler skips all pdf code from StrataEvaluateBSDF because GGXSample.w si the pdf already for instance.
+
+	return Sample;
+}
+
+struct FStrataEnvLightResult
+{
+	float3 DiffuseAlbedo;
+	float3 DiffuseNormal;
+
+	float3 SpecularDirection;
+	float3 SpecularWeight;
+	float  SpecularSafeRoughness;
+
+	bool   bSubsurface; // True if we need to separate the subsurface light contribution for the screen space diffusion process.
+};
+
+FStrataEnvLightResult StrataEvaluateForEnvLight(FStrataBSDFContext BSDFContext, bool bEnableSpecular)
+{
+	FStrataEnvLightResult StrataEnvLightResult = (FStrataEnvLightResult)0;
+
+	const uint BSDFType = BSDF_GETTYPE(BSDFContext.BSDF);
+	switch (BSDFType)
+	{
+		case STRATA_BSDF_TYPE_DIFFUSE:
+		{
+			StrataEnvLightResult.DiffuseAlbedo = DIFFUSE_ALBEDO(BSDFContext.BSDF);
+			StrataEnvLightResult.DiffuseNormal = BSDFContext.N;
+			StrataEnvLightResult.bSubsurface = BSDF_GETHASSCATTERING(BSDFContext.BSDF) > 0;
+			break;
+		}
+		case STRATA_BSDF_TYPE_DIELECTRIC:
+		case STRATA_BSDF_TYPE_CONDUCTOR:
+		{
+			BRANCH
+			if (bEnableSpecular)
+			{
+				const bool bIsDielectric = STRATA_BSDF_TYPE_DIELECTRIC == BSDFType;
+				const float3 F0  = bIsDielectric ? DIELECTRIC_TINT(BSDFContext.BSDF) * DielectricIorToF0(DIELECTRIC_IOR(BSDFContext.BSDF)) : CONDUCTOR_REFLECTIVITY(BSDFContext.BSDF);
+				const float3 F90 = bIsDielectric ? DIELECTRIC_TINT(BSDFContext.BSDF) : CONDUCTOR_EDGECOLOR(BSDFContext.BSDF);
+				float SafeRoughness0 = MakeRoughnessSafe(bIsDielectric ? DIELECTRIC_ROUGHNESS0(BSDFContext.BSDF) : CONDUCTOR_ROUGHNESS0(BSDFContext.BSDF));
+				float SafeRoughness1 = MakeRoughnessSafe(bIsDielectric ? DIELECTRIC_ROUGHNESS1(BSDFContext.BSDF) : CONDUCTOR_ROUGHNESS1(BSDFContext.BSDF));
+
+				const bool bHasAnisotropy = BSDF_GETANISOTROPY(BSDFContext.BSDF);
+				if (bHasAnisotropy)
+				{
+					// Modified the BSDF normal (and roughness)
+					float Anisotropy = 0;
+					GetAnisotropicFactor(SafeRoughness0, SafeRoughness1, Anisotropy, SafeRoughness0);
+					ModifyGGXAnisotropicNormalRoughness(BSDFContext.X, Anisotropy, SafeRoughness0, BSDFContext.N, BSDFContext.V);
+
+					// Update context (only needs: NoL/SatNoL/R) with the new N
+					BSDFContext.Context.NoL = dot(BSDFContext.N, BSDFContext.L);
+					BSDFContext.SatNoL = saturate(BSDFContext.Context.NoL);
+					BSDFContext.R = 2 * dot(BSDFContext.V, BSDFContext.N) * BSDFContext.N - BSDFContext.V;
+				}
+
+				float3 EvalEnvBRDF = EnvBRDFApprox(F0, F90, SafeRoughness0, BSDFContext.SatNoV);
+
+				StrataEnvLightResult.SpecularDirection = StrataGetOffSpecularPeakReflectionDir(BSDFContext.N, BSDFContext.R, SafeRoughness0);
+				StrataEnvLightResult.SpecularWeight = EvalEnvBRDF;
+				StrataEnvLightResult.SpecularSafeRoughness = SafeRoughness0;
+			}
+			break;
+		}
+		case STRATA_BSDF_TYPE_VOLUME:
+		{
+			const float3 VolumeAlbedo = VOLUME_ALBEDO(BSDFContext.BSDF);
+			const float3 ExtinctionCoefficients = VOLUME_EXTINCTION(BSDFContext.BSDF);
+			const float3 ScatteringCoefficients = VolumeAlbedo * ExtinctionCoefficients;
+			const float3 AbsorptionCoefficients = ExtinctionCoefficients - ScatteringCoefficients;
+			const float  Thickness = VOLUME_THICKNESS(BSDFContext.BSDF);
+
+			const float3 SafeExtinctionThreshold = 0.000001f;
+			const float3 SafeExtinctionCoefficients = max(SafeExtinctionThreshold, ExtinctionCoefficients);
+
+			const float PathLength = Thickness / max(0.0001f, BSDFContext.SatNoV);
+			const float3 SafePathSegmentTransmittance = exp(-SafeExtinctionCoefficients * PathLength);
+			const float3 NormLuminanceTransfert = 1.0f * ScatteringCoefficients;
+			float3 NormLuminanceSingleScattering = (NormLuminanceTransfert - NormLuminanceTransfert * SafePathSegmentTransmittance) / SafeExtinctionCoefficients;
+
+			StrataEnvLightResult.DiffuseAlbedo = NormLuminanceSingleScattering;
+			StrataEnvLightResult.DiffuseNormal = BSDFContext.N;
+			break;
+		}
+		case STRATA_BSDF_TYPE_SHEEN:
+		{
+			StrataEnvLightResult.DiffuseAlbedo = SHEEN_ALBEDO(BSDFContext.BSDF);;
+			StrataEnvLightResult.DiffuseNormal = BSDFContext.N;
+			break;
+		}
+		case STRATA_BSDF_TYPE_HAIR:
+		{
+			FGBufferData GBuffer	= (FGBufferData)0;
+			GBuffer.BaseColor		= HAIR_BASECOLOR(BSDFContext.BSDF);
+			GBuffer.Specular		= HAIR_SPECULAR(BSDFContext.BSDF);
+			GBuffer.Roughness		= HAIR_ROUGHNESS(BSDFContext.BSDF);
+			GBuffer.Metallic		= HAIR_SCATTER(BSDFContext.BSDF);
+			GBuffer.CustomData.z	= HAIR_BACKLIT(BSDFContext.BSDF);
+			GBuffer.ShadingModelID	= SHADINGMODELID_HAIR;
+			GBuffer.WorldNormal		= BSDFContext.N;
+			
+			float BacklitEnabled = 0.0f;
+			float Area = 0.2;
+			uint2 Random = uint2(0, 0);
+			float TransmissionShadow = 1.0f;
+			FHairTransmittanceData TransmittanceData = InitHairTransmittanceData(true);
+
+			const float3 N = BSDFContext.N;
+			const float3 V = BSDFContext.V;
+			float3 L = normalize(V - N * dot(V, N));
+			StrataEnvLightResult.DiffuseNormal = L;
+
+			StrataEnvLightResult.DiffuseAlbedo = PI * HairShading(GBuffer, L, V, N, TransmissionShadow, TransmittanceData, BacklitEnabled, Area, Random);
+			// No specular environment contribution as of today if not using the special HairStrand render path
+			break;
+		}
+	}
+
+	return StrataEnvLightResult;
+}
+
+FThreeBandSHVector StrataBSDFToSH(FStrataBSDFContext BSDFContext)
+{
+	FThreeBandSHVector SHVector;
+
+	const uint BSDFType = BSDF_GETTYPE(BSDFContext.BSDF);
+	if (BSDFType == STRATA_BSDF_TYPE_HAIR)
+	{
+		// Hack to avoid culling directions that hair will sample
+		SHVector = (FThreeBandSHVector)0;
+		SHVector.V0.x = 1.0f;
+	}
+	else
+	{
+		SHVector = CalcDiffuseTransferSH3(BSDFContext.N, 1.0f);
+	}
+
+	// STRATA_TODO adapt the SH to BSDFs
+	return SHVector;
+}
+
+
diff --git a/Engine/Shaders/Private/Strata/StrataForwardLighting.ush b/Engine/Shaders/Private/Strata/StrataForwardLighting.ush
--- a/Engine/Shaders/Private/Strata/StrataForwardLighting.ush
+++ b/Engine/Shaders/Private/Strata/StrataForwardLighting.ush
@@ -44,7 +44,7 @@ float3 StrataForwardLighting(
 
 			// Compute some common factors
 			const float3 DirLightBaseFactor = OutThroughput * BSDFCoverage * LColor * LShadow;
-			const float3 DirLightGeometryFactor = DirLightBaseFactor * StrataBSDFContext.SatNoL;
+			const float3 DirLightGeometryFactor = DirLightBaseFactor * StrataBSDFProjectedSolidAngleFactor(StrataBSDFContext);
 			const float3 SkyLightCommonFactor = OutThroughput * BSDFCoverage * ResolvedView.SkyLightColor.rgb * IndirectOcclusion;
 
 			// Evaluate environment lighting
@@ -66,7 +66,7 @@ float3 StrataForwardLighting(
 
 			// Evalaute the lighting coming from L (it will lso contain the BSDF throughput to the next layer)
 			FStrataEvaluateResult StrataEvaluate = StrataEvaluateBSDF(StrataBSDFContext);
-			Color += StrataEvaluate.Value * DirLightGeometryFactor;
+			Color += (StrataEvaluate.DiffusePathValue + StrataEvaluate.SpecularPathValue) * DirLightGeometryFactor;
 
 			LayerCoverage	+= BSDFCoverage;
 			LayerThroughput += BSDFCoverage * StrataEvaluate.Throughput;
diff --git a/Engine/Shaders/Private/Strata/StrataVisualize.usf b/Engine/Shaders/Private/Strata/StrataVisualize.usf
--- a/Engine/Shaders/Private/Strata/StrataVisualize.usf
+++ b/Engine/Shaders/Private/Strata/StrataVisualize.usf
@@ -4,7 +4,7 @@
 #include "/Engine/Private/DeferredShadingCommon.ush"
 
 #define STRATA_INLINE_SHADING 0
-#include "/Engine/Private/Strata/Strata.ush"
+#include "/Engine/Private/Strata/StrataEvaluation.ush"
 
 #include "/Engine/Private/MiniFontCommon.ush"
 #include "/Engine/Private/ShaderDrawDebug.ush"
@@ -62,7 +62,7 @@ float3 DrawPlot(
 
 		FStrataBSDFContext BSDFContext = StrataCreateBSDFContext(Header, BSDF, LocalV, L, true);
 		FStrataEvaluateResult BSDFEvaluate = StrataEvaluateBSDF(BSDFContext);
-		Out = BSDFEvaluate.Value * BSDFContext.SatNoL; // * BSDF.Weight
+		Out = (BSDFEvaluate.DiffusePathValue + BSDFEvaluate.SpecularPathValue + BSDFEvaluate.TransmissionPathValue) * BSDFContext.SatNoL; // * BSDF.Weight
 	}
 	return Out;
 }
@@ -373,6 +373,39 @@ void VisualizeMaterialPS(
 				NEXTLINE;
 
 				break;
+
+			case STRATA_BSDF_TYPE_HAIR:
+				PRINTBSDFCHAR(_H_); PRINTBSDFCHAR(_A_); PRINTBSDFCHAR(_I_); PRINTBSDFCHAR(_R_); PRINTBSDFCHAR(_SPC_);
+				PRINTBSDFCHAR(_SPC_); PRINTBSDFCHAR(_SPC_); PRINTBSDFCHAR(_SPC_); PRINTBSDFCHAR(_SPC_); PRINTBSDFCHAR(_SPC_);
+				if (BSDF_GETGREYWEIGHT(BSDF))
+				{
+					PRINTBSDFGREYWEIGHT();
+				}
+				else
+				{
+					PRINTBSDFWEIGHT();
+				}
+				PRINTBSDFNORMALID();
+				PRINTBSDFANISO();
+				PRINTBSDFTOPLAY();
+				NEXTLINE;
+
+				HORIZONTALTAB; PRINTPARAMCHAR(_A_); PRINTPARAMCHAR(_L_); PRINTPARAMCHAR(_B_); PRINTPARAMCHAR(_E_); PRINTPARAMCHAR(_D_); PRINTPARAMCHAR(_O_);
+				HORIZONTALTAB; PRINTFLOAT3(HAIR_BASECOLOR(BSDF));
+				NEXTLINE;
+				HORIZONTALTAB; PRINTPARAMCHAR(_R_); PRINTPARAMCHAR(_O_); PRINTPARAMCHAR(_U_); PRINTPARAMCHAR(_G_); PRINTPARAMCHAR(_H_); PRINTPARAMCHAR(_SPC_);
+				HORIZONTALTAB; PRINTFLOAT1(HAIR_ROUGHNESS(BSDF));
+				NEXTLINE;
+				HORIZONTALTAB; PRINTPARAMCHAR(_S_); PRINTPARAMCHAR(_C_); PRINTPARAMCHAR(_A_); PRINTPARAMCHAR(_T_); PRINTPARAMCHAR(_T_); PRINTPARAMCHAR(_SPC_);
+				HORIZONTALTAB; PRINTFLOAT1(HAIR_SCATTER(BSDF));
+				NEXTLINE;
+				HORIZONTALTAB; PRINTPARAMCHAR(_S_); PRINTPARAMCHAR(_P_); PRINTPARAMCHAR(_E_); PRINTPARAMCHAR(_C_); PRINTPARAMCHAR(_U_); PRINTPARAMCHAR(_SPC_);
+				HORIZONTALTAB; PRINTFLOAT1(HAIR_SPECULAR(BSDF));
+				NEXTLINE;
+				HORIZONTALTAB; PRINTPARAMCHAR(_B_); PRINTPARAMCHAR(_A_); PRINTPARAMCHAR(_C_); PRINTPARAMCHAR(_K_); PRINTPARAMCHAR(_L_); PRINTPARAMCHAR(_SPC_);
+				HORIZONTALTAB; PRINTFLOAT1(HAIR_BACKLIT(BSDF));
+				NEXTLINE;
+				break;
 			}
 
 			HORIZONTALSEPARATION;
diff --git a/Engine/Shaders/Shared/StrataDefinitions.h b/Engine/Shaders/Shared/StrataDefinitions.h
--- a/Engine/Shaders/Shared/StrataDefinitions.h
+++ b/Engine/Shaders/Shared/StrataDefinitions.h
@@ -28,6 +28,7 @@
 #define STRATA_BSDF_TYPE_SHEEN				4
 #define STRATA_BSDF_TYPE_VOLUMETRICFOGCLOUD	5
 #define STRATA_BSDF_TYPE_UNLIT				6
+#define STRATA_BSDF_TYPE_HAIR				7
 // When more than 8 BSDF must exists, please update STATE_BITS_BSDF in Strata.ush
 
 
diff --git a/Engine/Source/Runtime/Engine/Classes/Engine/EngineTypes.h b/Engine/Source/Runtime/Engine/Classes/Engine/EngineTypes.h
--- a/Engine/Source/Runtime/Engine/Classes/Engine/EngineTypes.h
+++ b/Engine/Source/Runtime/Engine/Classes/Engine/EngineTypes.h
@@ -544,6 +544,7 @@ enum EStrataShadingModel
 	SSM_DefaultLit				UMETA(DisplayName = "DefaultLit"),
 	SSM_SubsurfaceLit			UMETA(DisplayName = "SubsurfaceLit"),
 	SSM_VolumetricFogCloud		UMETA(DisplayName = "VolumetricFogCloud"),
+	SSM_Hair					UMETA(DisplayName = "Hair"),
 	/** Number of unique shading models. */
 	SSM_NUM						UMETA(Hidden),
 };
diff --git a/Engine/Source/Runtime/Engine/Classes/Materials/MaterialExpressionStrata.h b/Engine/Source/Runtime/Engine/Classes/Materials/MaterialExpressionStrata.h
--- a/Engine/Source/Runtime/Engine/Classes/Materials/MaterialExpressionStrata.h
+++ b/Engine/Source/Runtime/Engine/Classes/Materials/MaterialExpressionStrata.h
@@ -345,6 +345,65 @@ class UMaterialExpressionStrataUnlitBSDF : public UMaterialExpression
 	//~ End UMaterialExpression Interface
 };
 
+UCLASS(MinimalAPI, collapsecategories, hidecategories = Object)
+class UMaterialExpressionStrataHairBSDF : public UMaterialExpression
+{
+	GENERATED_UCLASS_BODY()
+		
+	/**
+	 * Hair fiber base color resulting from single and multiple scattering combined. (type = float3, unit = unitless, defaults to black)
+	 */
+	UPROPERTY()
+	FExpressionInput BaseColor;
+	
+	/**
+	 * Amount of light scattering, only available for non-HairStrand rendering (type = float, unit = unitless, defaults to 0.0)
+	 */
+	UPROPERTY()
+	FExpressionInput Scatter;
+		
+	/**
+	 * Specular (type = float, unit = unitless, defaults to 0.5)
+	 */
+	UPROPERTY()
+	FExpressionInput Specular;
+		
+	/**
+	 * Controls how rough the Material is. Roughness of 0 (smooth) is a mirror reflection and 1 (rough) is completely matte or diffuse (type = float, unit = unitless, defaults to 0.5)
+	 */
+	UPROPERTY()
+	FExpressionInput Roughness;
+
+	/**
+	 * How much light contributs when lighting hairs from the back side opposite from the view, only available for HairStrand rendering (type = float3, unit = unitless, defaults to 0.0)
+	 */
+	UPROPERTY()
+	FExpressionInput Backlit;
+
+	/**
+	 * Tangent (type = float3, unit = unitless, defaults to +X vector)
+	 */
+	UPROPERTY()
+	FExpressionInput Tangent;
+
+	/**
+	 * Emissive (type = float3, unit = luminance, defaults to 0.0)
+	 */
+	UPROPERTY()
+	FExpressionInput Emissive;
+
+	//~ Begin UMaterialExpression Interface
+#if WITH_EDITOR
+	virtual int32 Compile(class FMaterialCompiler* Compiler, int32 OutputIndex) override;
+	virtual void GetCaption(TArray<FString>& OutCaptions) const override;
+	virtual uint32 GetOutputType(int32 OutputIndex) override;
+	virtual uint32 GetInputType(int32 InputIndex) override;
+	virtual bool IsResultStrataMaterial(int32 OutputIndex) override;
+	virtual void GatherStrataMaterialInfo(FStrataMaterialInfo& StrataMaterialInfo, int32 OutputIndex) override;
+#endif
+	//~ End UMaterialExpression Interface
+};
+
 
 
 ///////////////////////////////////////////////////////////////////////////////
diff --git a/Engine/Source/Runtime/Engine/Private/Materials/HLSLMaterialTranslator.cpp b/Engine/Source/Runtime/Engine/Private/Materials/HLSLMaterialTranslator.cpp
--- a/Engine/Source/Runtime/Engine/Private/Materials/HLSLMaterialTranslator.cpp
+++ b/Engine/Source/Runtime/Engine/Private/Materials/HLSLMaterialTranslator.cpp
@@ -1125,7 +1125,7 @@ bool FHLSLMaterialTranslator::Translate()
 				// If the unlit node is used, it must be the only one used
 				if (!StrataIsVolumetricFogCloudOnly(this, StrataCompilationInfo))
 				{
-					FString ErrorMsg = FString::Printf(TEXT("Material %s contains Unlit BSDF but it is not the only one (asset: %s).\r\n"), *Material->GetDebugName(), *Material->GetAssetPath());
+					FString ErrorMsg = FString::Printf(TEXT("Material %s contains Unlit BSDF but it is not the only one representing the material asset: %s. It must be the single BSDF.\r\n"), *Material->GetDebugName(), *Material->GetAssetPath());
 					Error(*ErrorMsg);
 				}
 			}
@@ -1136,12 +1136,21 @@ bool FHLSLMaterialTranslator::Translate()
 				// If the unlit node is used, it must be the only one used
 				if (!StrataIsUnlitOnly(this, StrataCompilationInfo))
 				{
-					FString ErrorMsg = FString::Printf(TEXT("Material %s contains Unlit BSDF but it is not the only one (asset: %s).\r\n"), *Material->GetDebugName(), *Material->GetAssetPath());
+					FString ErrorMsg = FString::Printf(TEXT("Material %s contains Unlit BSDF but it is not the only one representing the material asset: %s. It must be the single BSDF.\r\n"), *Material->GetDebugName(), *Material->GetAssetPath());
 					Error(*ErrorMsg);
 				}
 			}
 
-			// STRATA_TODO if not volume domain: verify that there isn't any VolumetricFogCloud node in the material
+			// Hair must be used in isolation
+			if (StrataMaterialContainsAnyBSDF(this, StrataCompilationInfo, STRATA_BSDF_TYPE_HAIR))
+			{
+				// If the unlit node is used, it must be the only one used
+				if (!StrataIsHairOnly(this, StrataCompilationInfo))
+				{
+					FString ErrorMsg = FString::Printf(TEXT("Material %s contains hair BSDF but it is not the only representing the material asset: %s. It must be the single BSDF.\r\n"), *Material->GetDebugName(), *Material->GetAssetPath());
+					Error(*ErrorMsg);
+				}
+			}
 
 			// Output some debug info as comment in code and in the material stat window
 			static const auto CVarStrataBytePerPixel = IConsoleManager::Get().FindTConsoleVariableDataInt(TEXT("r.Strata.BytesPerPixel"));
@@ -7570,6 +7579,21 @@ int32 FHLSLMaterialTranslator::StrataUnlitBSDF(int32 EmissiveColor, int32 Transm
 	);
 }
 
+int32 FHLSLMaterialTranslator::StrataHairBSDF(int32 BaseColor, int32 Scatter, int32 Specular, int32 Roughness, int32 Backlit, int32 EmissiveColor, int32 Tangent, uint8 SharedNormalIndex)
+{
+	return AddCodeChunk(
+		MCT_Strata, TEXT("GetStrataHairBSDF(%s, %s, %s, %s, %s, %s, %u) /* %s */"),
+		*GetParameterCode(BaseColor),
+		*GetParameterCode(Scatter),
+		*GetParameterCode(Specular),
+		*GetParameterCode(Roughness),
+		*GetParameterCode(Backlit),
+		*GetParameterCode(EmissiveColor),
+		SharedNormalIndex,
+		*GetParameterCode(Tangent)
+	);
+}
+
 int32 FHLSLMaterialTranslator::StrataHorizontalMixing(int32 Foreground, int32 Background, int32 Mix)
 {
 	if (Foreground == INDEX_NONE || Background == INDEX_NONE || Mix == INDEX_NONE)
diff --git a/Engine/Source/Runtime/Engine/Private/Materials/HLSLMaterialTranslator.h b/Engine/Source/Runtime/Engine/Private/Materials/HLSLMaterialTranslator.h
--- a/Engine/Source/Runtime/Engine/Private/Materials/HLSLMaterialTranslator.h
+++ b/Engine/Source/Runtime/Engine/Private/Materials/HLSLMaterialTranslator.h
@@ -798,6 +798,7 @@ protected:
 	virtual int32 StrataSheenBSDF(int32 Albedo, int32 Roughness, int32 Normal, uint8 SharedNormalIndex) override;
 	virtual int32 StrataVolumetricFogCloudBSDF(int32 Albedo, int32 Extinction, int32 Emissive, int32 AmbientOcclusion) override;
 	virtual int32 StrataUnlitBSDF(int32 EmissiveColor, int32 TransmittanceColor) override;
+	virtual int32 StrataHairBSDF(int32 BaseColor, int32 Scatter, int32 Specular, int32 Roughness, int32 Backlit, int32 EmissiveColor, int32 Tangent, uint8 SharedNormalIndex) override;
 	virtual int32 StrataHorizontalMixing(int32 Foreground, int32 Background, int32 Mix) override;
 	virtual int32 StrataVerticalLayering(int32 Top, int32 Base) override;
 	virtual int32 StrataAdd(int32 A, int32 B) override;
diff --git a/Engine/Source/Runtime/Engine/Private/Materials/Material.cpp b/Engine/Source/Runtime/Engine/Private/Materials/Material.cpp
--- a/Engine/Source/Runtime/Engine/Private/Materials/Material.cpp
+++ b/Engine/Source/Runtime/Engine/Private/Materials/Material.cpp
@@ -4557,6 +4557,12 @@ void UMaterial::RebuildShadingModelField()
 			ShadingModel = MSM_DefaultLit;
 			BlendMode = EBlendMode::BLEND_Additive;
 		}
+		else if (StrataMaterialInfo.HasOnlyShadingModel(SSM_Hair))
+		{
+			MaterialDomain = EMaterialDomain::MD_Surface;
+			ShadingModel = MSM_Hair;
+			BlendMode = EBlendMode::BLEND_Opaque;
+		}
 
 		// Also update the ShadingModels for remaining pipeline operation
 		ShadingModels.AddShadingModel(ShadingModel);
diff --git a/Engine/Source/Runtime/Engine/Private/Materials/MaterialExpressions.cpp b/Engine/Source/Runtime/Engine/Private/Materials/MaterialExpressions.cpp
--- a/Engine/Source/Runtime/Engine/Private/Materials/MaterialExpressions.cpp
+++ b/Engine/Source/Runtime/Engine/Private/Materials/MaterialExpressions.cpp
@@ -19783,6 +19783,7 @@ uint32 UMaterialExpressionStrataUnlitBSDF::GetInputType(int32 InputIndex)
 		break;
 	case 1:
 		return MCT_Float3;
+		break;
 	}
 
 	check(false);
@@ -19802,6 +19803,96 @@ void UMaterialExpressionStrataUnlitBSDF::GatherStrataMaterialInfo(FStrataMateria
 
 
 
+UMaterialExpressionStrataHairBSDF::UMaterialExpressionStrataHairBSDF(const FObjectInitializer& ObjectInitializer)
+	: Super(ObjectInitializer)
+{
+	struct FConstructorStatics
+	{
+		FText NAME_Strata;
+		FConstructorStatics() : NAME_Strata(LOCTEXT("Strata BSDFs", "Strata BSDFs")) { }
+	};
+	static FConstructorStatics ConstructorStatics;
+#if WITH_EDITORONLY_DATA
+	MenuCategories.Add(ConstructorStatics.NAME_Strata);
+#endif
+}
+
+#if WITH_EDITOR
+int32 UMaterialExpressionStrataHairBSDF::Compile(class FMaterialCompiler* Compiler, int32 OutputIndex)
+{
+	// For hair, the shared normal index in fact represent the tangent
+	int32 TangentCodeChunk = CompileWithDefaultCodeChunk(Compiler, Tangent, Compiler->VertexTangent());
+	uint8 SharedNormalIndex = StrataCompilationInfoCreateSharedNormal(Compiler, TangentCodeChunk);
+
+	int32 OutputCodeChunk = Compiler->StrataHairBSDF(
+		CompileWithDefaultFloat3(Compiler, BaseColor,	0.0f, 0.0f, 0.0f),
+		CompileWithDefaultFloat1(Compiler, Scatter,		0.0f),
+		CompileWithDefaultFloat1(Compiler, Specular,	0.5f),
+		CompileWithDefaultFloat1(Compiler, Roughness,	0.5f),
+		CompileWithDefaultFloat1(Compiler, Backlit,		0.0f),
+		CompileWithDefaultFloat3(Compiler, Emissive,	1.0f, 0.0f, 0.0f),
+		TangentCodeChunk,
+		SharedNormalIndex);
+
+	StrataCompilationInfoCreateSingleBSDFMaterial(Compiler, OutputCodeChunk, SharedNormalIndex, STRATA_BSDF_TYPE_HAIR);
+
+	return OutputCodeChunk;
+}
+
+void UMaterialExpressionStrataHairBSDF::GetCaption(TArray<FString>& OutCaptions) const
+{
+	OutCaptions.Add(TEXT("Strata Hair BSDF"));
+}
+
+uint32 UMaterialExpressionStrataHairBSDF::GetOutputType(int32 OutputIndex)
+{
+	return MCT_Strata;
+}
+
+uint32 UMaterialExpressionStrataHairBSDF::GetInputType(int32 InputIndex)
+{
+	switch (InputIndex)
+	{
+	case 0:
+		return MCT_Float3;
+		break;
+	case 1:
+		return MCT_Float;
+		break;
+	case 2:
+		return MCT_Float;
+		break;
+	case 3:
+		return MCT_Float;
+		break;
+	case 4:
+		return MCT_Float;
+		break;
+	case 5:
+		return MCT_Float3;
+		break;
+	case 6:
+		return MCT_Float3;
+		break;
+	}
+
+	check(false);
+	return MCT_Float1;
+}
+
+bool UMaterialExpressionStrataHairBSDF::IsResultStrataMaterial(int32 OutputIndex)
+{
+	return true;
+}
+
+void UMaterialExpressionStrataHairBSDF::GatherStrataMaterialInfo(FStrataMaterialInfo& StrataMaterialInfo, int32 OutputIndex)
+{
+	StrataMaterialInfo.AddShadingModel(SSM_Hair);
+}
+#endif // WITH_EDITOR
+
+
+
 UMaterialExpressionStrataHorizontalMixing::UMaterialExpressionStrataHorizontalMixing(const FObjectInitializer& ObjectInitializer)
 	: Super(ObjectInitializer)
 {
diff --git a/Engine/Source/Runtime/Engine/Private/Materials/StrataMaterial.cpp b/Engine/Source/Runtime/Engine/Private/Materials/StrataMaterial.cpp
--- a/Engine/Source/Runtime/Engine/Private/Materials/StrataMaterial.cpp
+++ b/Engine/Source/Runtime/Engine/Private/Materials/StrataMaterial.cpp
@@ -30,6 +30,9 @@ FString GetStrataBSDFName(uint8 BSDFType)
 	case STRATA_BSDF_TYPE_UNLIT:
 		return TEXT("UNLIT");
 		break;
+	case STRATA_BSDF_TYPE_HAIR:
+		return TEXT("HAIR");
+		break;
 	}
 	check(false);
 	return "";
@@ -141,16 +144,25 @@ FStrataMaterialCompilationInfo StrataCompilationInfoVerticalLayering(FMaterialCo
 	return StrataInfo;
 }
 
-bool StrataIsVolumetricFogCloudOnly(FMaterialCompiler* Compiler, const FStrataMaterialCompilationInfo& Material)
+static bool StrataIsSingleBSDF(FMaterialCompiler* Compiler, const FStrataMaterialCompilationInfo& Material)
 {
 	if (Material.TotalBSDFCount == 0 || Material.LayerCount == 0)
 	{
-		Compiler->Error(TEXT("There is no layer or BSDF plugged in, but a material in the volume domain wants to read from a StrataVolumetricFogCloudBSDF."));
+		Compiler->Error(TEXT("There is no layer or BSDF plugged in, but one of the BSDF in the graph wants to enforce one and only one BSDF to be used."));
 		return false;
 	}
 	if (Material.TotalBSDFCount > 1 || Material.LayerCount > 1)
 	{
-		Compiler->Error(TEXT("There is more than one layer or BSDF, but a material in the volume domain wants to read from a single StrataVolumetricFogCloudBSDF only."));
+		Compiler->Error(TEXT("There is more than one layer or BSDF, but one of the BSDF in the graph wants to enforce one and only one BSDF to be used."));
+		return false;
+	}
+	return true;
+}
+
+bool StrataIsVolumetricFogCloudOnly(FMaterialCompiler* Compiler, const FStrataMaterialCompilationInfo& Material)
+{
+	if (!StrataIsSingleBSDF(Compiler, Material))
+	{
 		return false;
 	}
 	if (Material.Layers[0].BSDFs[0].Type != STRATA_BSDF_TYPE_VOLUMETRICFOGCLOUD)
@@ -164,19 +176,28 @@ bool StrataIsVolumetricFogCloudOnly(FMaterialCompiler* Compiler, const FStrataMa
 
 bool StrataIsUnlitOnly(FMaterialCompiler* Compiler, const FStrataMaterialCompilationInfo& Material)
 {
-	if (Material.TotalBSDFCount == 0 || Material.LayerCount == 0)
+	if (!StrataIsSingleBSDF(Compiler, Material))
 	{
-		Compiler->Error(TEXT("There is no layer or BSDF plugged in, but a material in the volume domain wants to read from a StrataUnlitBSDF."));
 		return false;
 	}
-	if (Material.TotalBSDFCount > 1 || Material.LayerCount > 1)
+	if (Material.Layers[0].BSDFs[0].Type != STRATA_BSDF_TYPE_UNLIT)
 	{
-		Compiler->Error(TEXT("There is more than one layer or BSDF, but a material in the volume domain wants to read from a single StrataUnlitBSDF only."));
+		Compiler->Error(TEXT("The single BSDF resulting from the graph is not of type Unlit."));
 		return false;
 	}
-	if (Material.Layers[0].BSDFs[0].Type != STRATA_BSDF_TYPE_UNLIT)
+
+	return true;
+}
+
+bool StrataIsHairOnly(FMaterialCompiler* Compiler, const FStrataMaterialCompilationInfo& Material)
+{
+	if (!StrataIsSingleBSDF(Compiler, Material))
 	{
-		Compiler->Error(TEXT("The single BSDF resulting from the graph is not of type Volume."));
+		return false;
+	}
+	if (Material.Layers[0].BSDFs[0].Type != STRATA_BSDF_TYPE_HAIR)
+	{
+		Compiler->Error(TEXT("The single BSDF resulting from the graph is not of type Hair."));
 		return false;
 	}
 
@@ -282,6 +303,12 @@ FStrataMaterialAnalysisResult StrataCompilationInfoMaterialAnalysis(FMaterialCom
 				Result.RequestedByteCount += UintByteSize;
 				break;
 			}
+			case STRATA_BSDF_TYPE_HAIR:
+			{
+				Result.RequestedByteCount += UintByteSize;
+				Result.RequestedByteCount += UintByteSize;
+				break;
+			}
 			}
 		}
 
diff --git a/Engine/Source/Runtime/Engine/Private/Materials/StrataMaterial.h b/Engine/Source/Runtime/Engine/Private/Materials/StrataMaterial.h
--- a/Engine/Source/Runtime/Engine/Private/Materials/StrataMaterial.h
+++ b/Engine/Source/Runtime/Engine/Private/Materials/StrataMaterial.h
@@ -54,6 +54,7 @@ FStrataMaterialCompilationInfo StrataCompilationInfoVerticalLayering(FMaterialCo
 
 bool StrataIsVolumetricFogCloudOnly(FMaterialCompiler* Compiler, const FStrataMaterialCompilationInfo& Material);
 bool StrataIsUnlitOnly(FMaterialCompiler* Compiler, const FStrataMaterialCompilationInfo& Material);
+bool StrataIsHairOnly(FMaterialCompiler* Compiler, const FStrataMaterialCompilationInfo& Material);
 bool StrataMaterialContainsAnyBSDF(FMaterialCompiler* Compiler, const FStrataMaterialCompilationInfo& Material, uint8 BSDFType);
 
 struct FStrataMaterialAnalysisResult
diff --git a/Engine/Source/Runtime/Engine/Public/MaterialCompiler.h b/Engine/Source/Runtime/Engine/Public/MaterialCompiler.h
--- a/Engine/Source/Runtime/Engine/Public/MaterialCompiler.h
+++ b/Engine/Source/Runtime/Engine/Public/MaterialCompiler.h
@@ -398,6 +398,7 @@ public:
 	virtual int32 StrataSheenBSDF(int32 Albedo, int32 Roughness, int32 Normal, uint8 SharedNormalIndex) = 0;
 	virtual int32 StrataVolumetricFogCloudBSDF(int32 Albedo, int32 Extinction, int32 Emissive, int32 AmbientOcclusion) = 0;
 	virtual int32 StrataUnlitBSDF(int32 EmissiveColor, int32 TransmittanceColor) = 0;
+	virtual int32 StrataHairBSDF(int32 BaseColor, int32 Scatter, int32 Specular, int32 Roughness, int32 Backlit, int32 EmissiveColor, int32 Tangent, uint8 SharedNormalIndex) = 0;
 	virtual int32 StrataHorizontalMixing(int32 Foreground, int32 Background, int32 Mix) = 0;
 	virtual int32 StrataVerticalLayering(int32 Top, int32 Base) = 0;
 	virtual int32 StrataAdd(int32 A, int32 B) = 0;
@@ -842,6 +843,11 @@ public:
 		return Compiler->StrataUnlitBSDF(EmissiveColor, TransmittanceColor);
 	}
 
+	virtual int32 StrataHairBSDF(int32 BaseColor, int32 Scatter, int32 Specular, int32 Roughness, int32 Backlit, int32 EmissiveColor, int32 Tangent, uint8 SharedNormalIndex) override
+	{
+		return Compiler->StrataHairBSDF(BaseColor, Scatter, Specular, Roughness, Backlit, EmissiveColor, Tangent, SharedNormalIndex);
+	}
+
 	virtual int32 StrataHorizontalMixing(int32 Foreground, int32 Background, int32 Mix) override
 	{
 		return Compiler->StrataHorizontalMixing(Foreground, Background, Mix);
