Replace QuasiRandom sampler in RTGI and Skydome passes with RandomSequence sampler

This enables error-diffusion in these passes which reduces visible noise.

diff --git a/Engine/Plugins/Experimental/GPULightmass/Shaders/Private/LightmapPathTracing.usf b/Engine/Plugins/Experimental/GPULightmass/Shaders/Private/LightmapPathTracing.usf
--- a/Engine/Plugins/Experimental/GPULightmass/Shaders/Private/LightmapPathTracing.usf
+++ b/Engine/Plugins/Experimental/GPULightmass/Shaders/Private/LightmapPathTracing.usf
@@ -29,7 +29,6 @@ struct FIrradianceCacheRecord
 
 #include "BatchedTiles.ush"
 #include "/Engine/Private/PathTracing/Material/FirstBounceRayGuidingCommon.ush"
-#include "/Engine/Private/QuasiRandom.ush"
 
 int NumRayGuidingTrialSamples;
 RWTexture2D<uint> RayGuidingLuminance;
@@ -1048,8 +1047,8 @@ void StationaryLightShadowTracingMainRG()
 		
 		{
 			float2 RandSample = float2(
-				Halton(IntegerHash(Seed) + LightSampleIndexArray[TileIndex], 5), 
-				Halton(IntegerHash(Seed) + LightSampleIndexArray[TileIndex], 7)
+				Halton(StrongIntegerHash(Seed) + LightSampleIndexArray[TileIndex], 5), 
+				Halton(StrongIntegerHash(Seed) + LightSampleIndexArray[TileIndex], 7)
 			);
 			
 			RayDesc Ray;
diff --git a/Engine/Shaders/Private/RayTracing/GenerateSkyLightVisibilityRaysCS.usf b/Engine/Shaders/Private/RayTracing/GenerateSkyLightVisibilityRaysCS.usf
--- a/Engine/Shaders/Private/RayTracing/GenerateSkyLightVisibilityRaysCS.usf
+++ b/Engine/Shaders/Private/RayTracing/GenerateSkyLightVisibilityRaysCS.usf
@@ -6,10 +6,10 @@
 
 #include "../Common.ush"
 #include "../MonteCarlo.ush"
-#include "../QuasiRandom.ush"
 #include "RayTracingCommon.ush"
 #include "RayTracingSkyLightCommon.ush"
 #include "SkyLightVisibilityRaysData.ush"
+#include "../PathTracing/Utilities/PathTracingRandomSequence.ush"
 
 uint3 SkyLightVisibilityRaysDimensions;
 RWStructuredBuffer<FSkyLightVisibilityRays> OutSkyLightVisibilityRays;
@@ -21,10 +21,6 @@ void MainCS(
 {
 	uint2 PixelCoord = uint2(DispatchThreadId.x, DispatchThreadId.y);
 
-	FQuasiRandom QuasiRandom;
-	QuasiRandom_Initialize(QuasiRandom, PixelCoord, 0, SkyLight.SamplesPerPixel);
-	float3 BlueNoiseShift = EvalBlueNoise(QuasiRandom.BlueNoiseCoord, 0);
-
 	// Cache Cubeface CDF
 	uint MipCount = log2(SkyLight.MipDimensions.x);
 	float FaceCdf[7];
@@ -32,13 +28,11 @@ void MainCS(
 
 	for (uint SampleIndex = 0; SampleIndex < SkyLight.SamplesPerPixel; ++SampleIndex)
 	{
-		QuasiRandom_Update(QuasiRandom, SampleIndex, SkyLight.SamplesPerPixel);
+		RandomSequence RandSequence;
+		RandomSequence_Initialize(RandSequence, PixelCoord, SampleIndex, View.StateFrameIndex, SkyLight.SamplesPerPixel);
 
 		// Determine sky light or lambert ray
-		float3 RandSample = QuasiRandom_Sample3(QuasiRandom);
-		RandSample = HaltonSequence3(QuasiRandom.HaltonSequence, QuasiRandom.HaltonIteration, 0);
-		RandSample += BlueNoiseShift;
-		RandSample = fmod(RandSample, 1.0);
+		float3 RandSample = RandomSequence_GenerateSample3D(RandSequence);
 
 		float3 RayDirection;
 		float SkyLightSamplingStrategyPdf = 0.5;
diff --git a/Engine/Shaders/Private/RayTracing/RayTracingCreateGatherPointsRGS.usf b/Engine/Shaders/Private/RayTracing/RayTracingCreateGatherPointsRGS.usf
--- a/Engine/Shaders/Private/RayTracing/RayTracingCreateGatherPointsRGS.usf
+++ b/Engine/Shaders/Private/RayTracing/RayTracingCreateGatherPointsRGS.usf
@@ -22,8 +22,7 @@
 
 #include "../PathTracing/Light/PathTracingLightSampling.ush"
 #include "../PathTracing/Material/PathTracingMaterialSampling.ush"
-
-#include "../QuasiRandom.ush"
+#include "../PathTracing/Utilities/PathTracingRandomSequence.ush"
 
 RaytracingAccelerationStructure TLAS;
 
@@ -168,18 +167,17 @@ RAY_TRACING_ENTRY_RAYGEN(RayTracingCreateGatherPointsTraceRGS)
 	uint2 BlueNoiseCoord = PixelCoord / GatherFilterWidth;
 	uint Iteration = (PixelCoord.y % GatherFilterWidth) * GatherFilterWidth + (PixelCoord.x % GatherFilterWidth);
 
-	FQuasiRandom QuasiRandom;
-	QuasiRandom_Initialize(QuasiRandom, BlueNoiseCoord, Iteration, SamplesPerPixel);
-
 	bool bTraceRay = ShouldTraceRay(DispatchThreadId, GBufferData);
 	if (bTraceRay)
 	{
-		float3 RandSample = QuasiRandom_Sample3(QuasiRandom);
+		RandomSequence RandSequence;
+		RandomSequence_Initialize(RandSequence, PixelCoord, SampleIndex, View.StateFrameIndex, SamplesPerPixel);
+		float2 RandSample = RandomSequence_GenerateSample2D(RandSequence);
 
 		// Initialize ray
 		RayDesc Ray;
 		float RayPdf = 1.0;
-		GenerateCosineNormalRay(WorldPosition, WorldNormal, RandSample.yz, Ray.Origin, Ray.Direction, Ray.TMin, Ray.TMax, RayPdf);
+		GenerateCosineNormalRay(WorldPosition, WorldNormal, RandSample, Ray.Origin, Ray.Direction, Ray.TMin, Ray.TMax, RayPdf);
 		//GenerateUniformRay(WorldPosition, WorldNormal, RandSample.yz, Ray.Origin, Ray.Direction, Ray.TMin, Ray.TMax, RayPdf);
 		//GenerateCosineNormalRay(WorldPosition, SampleCenterNormal, RandSample.yz, Ray.Origin, Ray.Direction, Ray.TMin, Ray.TMax, RayPdf);
 		ApplyCameraRelativeDepthBias(Ray, PixelCoord, DeviceZ, WorldNormal, MaxNormalBias);
@@ -284,9 +282,6 @@ RAY_TRACING_ENTRY_RAYGEN(RayTracingCreateGatherPointsRGS)
 	FGBufferData SampleCenterGBufferData = GetGBufferDataFromSceneTexturesLoad(SampleCenterCoord);
 	float3 SampleCenterNormal = SampleCenterGBufferData.WorldNormal;
 
-	FQuasiRandom QuasiRandom;
-	QuasiRandom_Initialize(QuasiRandom, BlueNoiseCoord, Iteration, SamplesPerPixel);
-
 	bool bTraceRay = ShouldTraceRay(DispatchThreadId, GBufferData);
 	if (bTraceRay)
 	{
@@ -301,12 +296,14 @@ RAY_TRACING_ENTRY_RAYGEN(RayTracingCreateGatherPointsRGS)
 		float3 Position = 0;
 		float3 Normal = 0;
 
-		float3 RandSample = QuasiRandom_Sample3(QuasiRandom);
+		RandomSequence RandSequence;
+		RandomSequence_Initialize(RandSequence, PixelCoord, SampleIndex, View.StateFrameIndex, SamplesPerPixel);
+		float2 RandSample = RandomSequence_GenerateSample2D(RandSequence);
 
 		// Initialize ray
 		RayDesc Ray;
 		float RayPdf = 1.0;
-		GenerateCosineNormalRay(WorldPosition, WorldNormal, RandSample.yz, Ray.Origin, Ray.Direction, Ray.TMin, Ray.TMax, RayPdf);
+		GenerateCosineNormalRay(WorldPosition, WorldNormal, RandSample, Ray.Origin, Ray.Direction, Ray.TMin, Ray.TMax, RayPdf);
 		//GenerateUniformRay(WorldPosition, WorldNormal, RandSample.yz, Ray.Origin, Ray.Direction, Ray.TMin, Ray.TMax, RayPdf);
 		//GenerateCosineNormalRay(WorldPosition, SampleCenterNormal, RandSample.yz, Ray.Origin, Ray.Direction, Ray.TMin, Ray.TMax, RayPdf);
 		// Distribution PDF is not used for gather points, only gather queries.
@@ -385,9 +382,8 @@ RAY_TRACING_ENTRY_RAYGEN(RayTracingCreateGatherPointsRGS)
 				float3 LightUV;
 				float NeePdf = 0.0;
 
-				float4 RandSample4;
-				RandSample4.xyz = QuasiRandom_Sample3(QuasiRandom);
-				RandSample4.w = QuasiRandom_Sample3(QuasiRandom).x;
+				float4 RandSample4 = RandomSequence_GenerateSample4D(RandSequence);
+
 				SampleLight(Ray, Payload, RandSample4, LightId, LightUV, NeePdf);
 
 				if (NeePdf > 0.0)
diff --git a/Engine/Shaders/Private/RayTracing/RayTracingFinalGatherRGS.usf b/Engine/Shaders/Private/RayTracing/RayTracingFinalGatherRGS.usf
--- a/Engine/Shaders/Private/RayTracing/RayTracingFinalGatherRGS.usf
+++ b/Engine/Shaders/Private/RayTracing/RayTracingFinalGatherRGS.usf
@@ -14,8 +14,6 @@
 #include "../PathTracing/Light/PathTracingLightSampling.ush"
 #include "../PathTracing/Material/PathTracingMaterialSampling.ush"
 
-#include "../QuasiRandom.ush"
-
 #define USE_PATHTRACING_MATERIALS 0
 #define USE_REJECTION_CRITERIA 0
 #define USE_REPROJECTION 1
diff --git a/Engine/Shaders/Private/RayTracing/RayTracingGlobalIlluminationRGS.usf b/Engine/Shaders/Private/RayTracing/RayTracingGlobalIlluminationRGS.usf
--- a/Engine/Shaders/Private/RayTracing/RayTracingGlobalIlluminationRGS.usf
+++ b/Engine/Shaders/Private/RayTracing/RayTracingGlobalIlluminationRGS.usf
@@ -12,8 +12,8 @@
 
 #include "../PathTracing/Light/PathTracingLightSampling.ush"
 #include "../PathTracing/Material/PathTracingMaterialSampling.ush"
+#include "../PathTracing/Utilities/PathTracingRandomSequence.ush"
 
-#include "../QuasiRandom.ush"
 
 #define USE_PATHTRACING_MATERIALS 0
 
@@ -102,19 +102,17 @@ RAY_TRACING_ENTRY_RAYGEN(GlobalIlluminationRGS)
 		return;
 	}
 
-	FQuasiRandom QuasiRandom;
-	QuasiRandom_Initialize(QuasiRandom, DispatchThreadId, 0, SamplesPerPixel);
-
 	float3 Irradiance = 0;
 	float HitDistance = 0.0;
 	float HitCount = 0.0;
 	float AmbientOcclusion = 0.0;
 	for (uint SampleIndex = 0; SampleIndex < SamplesPerPixel; ++SampleIndex)
 	{
-		QuasiRandom_Update(QuasiRandom, SampleIndex, SamplesPerPixel);
+		RandomSequence RandSequence;
+		RandomSequence_Initialize(RandSequence, PixelCoord, SampleIndex, View.StateFrameIndex, SamplesPerPixel);
 
 		float3 RayThroughput = 1.0;
-		float3 RandSample = QuasiRandom_Sample3(QuasiRandom);
+		float3 RandSample = RandomSequence_GenerateSample3D(RandSequence);
 
 		// Russian roulette based on DiffuseColor
 		if (UseRussianRoulette)
@@ -228,9 +226,8 @@ RAY_TRACING_ENTRY_RAYGEN(GlobalIlluminationRGS)
 				float3 LightUV;
 				float NeePdf = 0.0;
 
-				float4 RandSample4;
-				RandSample4.xyz = QuasiRandom_Sample3(QuasiRandom);
-				RandSample4.w = QuasiRandom_Sample3(QuasiRandom).x;
+				float4 RandSample4 = RandomSequence_GenerateSample4D(RandSequence);
+
 				SampleLight(Ray, Payload, RandSample4, LightId, LightUV, NeePdf);
 
 				if (NeePdf > 0.0)
@@ -317,15 +314,16 @@ RAY_TRACING_ENTRY_RAYGEN(GlobalIlluminationRGS)
 				float3 Direction;
 				float3 Throughput = 1.0;
 #if USE_PATHTRACING_MATERIALS
-				float4 RandSample;
-				RandSample.xyz = QuasiRandom_Sample3(QuasiRandom);
-				RandSample.w = QuasiRandom_Sample3(QuasiRandom).x;
+				RandSample4 = RandomSequence_GenerateSample4D(RandSequence);
+
 				float PositionBiasSign;
 				float3 Weight = 0;
 				SampleMaterial(Ray.Direction, Payload, RandSample, Direction, Weight, MaterialPdf, PositionBiasSign);
 				Throughput = Weight * MaterialPdf;
 #else
-				float3 RandSample3 = QuasiRandom_Sample3(QuasiRandom);
+
+				float3 RandSample3 = RandomSequence_GenerateSample3D(RandSequence);
+
 				float3 RayOrigin = Ray.Origin;
 				GenerateCosineNormalRay(RayOrigin, Payload.WorldNormal, RandSample3.xy, Ray.Origin, Direction, Ray.TMin, Ray.TMax, MaterialPdf);
 				
diff --git a/Engine/Shaders/Private/RayTracing/RayTracingSkyLightEvaluation.ush b/Engine/Shaders/Private/RayTracing/RayTracingSkyLightEvaluation.ush
--- a/Engine/Shaders/Private/RayTracing/RayTracingSkyLightEvaluation.ush
+++ b/Engine/Shaders/Private/RayTracing/RayTracingSkyLightEvaluation.ush
@@ -7,7 +7,6 @@
 #pragma once
 
 #include "../Common.ush"
-#include "../QuasiRandom.ush"
 #include "../ShadingModels.ush"
 #include "../DeferredShadingCommon.ush"
 #include "../ReflectionEnvironmentShared.ush"
@@ -60,19 +59,12 @@ void SkyLightEvaluate(
 	inout float HitDistance)
 {
 	float FaceCdf[7];
-	FQuasiRandom QuasiRandom;
-	const uint BlueNoiseDimension = 0;
-	float3 BlueNoiseShift;
 
 	if (!bDecoupleSampleGeneration)
 	{
 		// Cache Cubeface CDF
 		const uint MipCount = log2(SkyLight.MipDimensions.x);
 		BuildFaceCdf(MipCount, FaceCdf);
-
-		// Initialize a quasi random sequence for the current pixel
-		QuasiRandom_Initialize(QuasiRandom, SampleCoord, 0, SamplesPerPixel);
-		BlueNoiseShift = EvalBlueNoise(QuasiRandom.BlueNoiseCoord, BlueNoiseDimension);
 	}
 
 	float3 CurrentWorldNormal = WorldNormal;
@@ -107,13 +99,11 @@ void SkyLightEvaluate(
 		}
 		else
 		{
-			QuasiRandom_Update(QuasiRandom, SampleIndex, SamplesPerPixel);
+			RandomSequence RandSequence;
+			RandomSequence_Initialize(RandSequence, PixelCoord, SampleIndex, View.StateFrameIndex, SamplesPerPixel);
 
 			// Determine sky light or lambert ray
-			float3 RandSample = QuasiRandom_Sample3(QuasiRandom);
-			RandSample = HaltonSequence3(QuasiRandom.HaltonSequence, QuasiRandom.HaltonIteration, BlueNoiseDimension);
-			RandSample += BlueNoiseShift;
-			RandSample = fmod(RandSample, 1.0);
+			float3 RandSample = RandomSequence_GenerateSample3D(RandSequence);
 
 			// Generate a visibility ray for the current sample
 			float SkyLightSamplingStrategyPdf = 0.5;
