Centralize the logic for enabling the error diffusion sampler of the RandomSequence API

Client code will get this sampler by default, but can override this by using the RANDSEQ_ERROR_DIFFUSION define.

Also simplify the linear index calculation to not require knowing the buffer dimensions

diff --git a/Engine/Shaders/Private/PathTracing/Utilities/PathTracingRandomSequence.ush b/Engine/Shaders/Private/PathTracing/Utilities/PathTracingRandomSequence.ush
--- a/Engine/Shaders/Private/PathTracing/Utilities/PathTracingRandomSequence.ush
+++ b/Engine/Shaders/Private/PathTracing/Utilities/PathTracingRandomSequence.ush
@@ -21,6 +21,10 @@
 #define RANDSEQ_RANDOMIZED		1 // 1 to randomize per pixel, 0 to share the same sequence in all pixels (useful to benchmark coherent sampling)
 #endif
 
+#ifndef RANDSEQ_ERROR_DIFFUSION
+#define RANDSEQ_ERROR_DIFFUSION 1 // Enabled by default, only works when using RANDSEQ_OWENSOBOL
+#endif
+
 // This hash mixes bits at the theoretical bias limit
 uint PerfectIntegerHash(uint x)
 {
@@ -280,6 +284,9 @@ float4 SobolSampler(uint SampleIndex, inout uint Seed)
 	return (Result >> 8) * 5.96046447754e-08; // * 2^-24
 }
 
+// This initializes a multi-dimensional sampler for a particular pixel. The TimeSeed parameter should be set
+// such that it changes per sample. See the function below when you know how many samples you plan on taking
+// for a particular integral.
 void RandomSequence_Initialize(inout RandomSequence RandSequence, uint PositionSeed, uint TimeSeed)
 {
 	// optionally disable randomization per pixel so that all pixels follow the same path in primary sample space
@@ -303,24 +310,26 @@ void RandomSequence_Initialize(inout RandomSequence RandSequence, uint PositionS
 #endif
 }
 
-#if RANDSEQ == RANDSEQ_OWENSOBOL
-
-// This is an alternative initialization to the method above when you know how many samples you are planning on taking.
-// The core idea is to use a single sobol pattern across the screen using a space filling curve to correlate nearby pixels
-// This achieves a much better error diffusion than the random sequence assignment of the method above
-// This algorithm is covered in the following paper:
-//   "Screen-Space Blue-Noise Diffusion of Monte Carlo Sampling Error via Hierarchical Ordering of Pixels"
-//   Siggraph Asia 2020
-//   http://abdallagafar.com/publications/zsampler/
-// Improvements made in this version are: 
-//   - switching the z-order curve for a hilbert curve
-//   - switching sobol points for owen-sobol points
-//   - recognizing that the hierarchical scrambling proposed in the original paper is equivalent to Owen-Scrambling of the index
-// Downsides of this method:
-//   - quality does not improve beyond the target sample count, making usage with adaptive sampling difficult
-//   - error-diffusion property is only achieved at the target sample count, earlier samples look random (sometimes slightly worse than random)
+// This is an alternative initialization to the method above when you know how many samples you are planning on taking. This incorporates the 
+// TimeSeed0 (this should usually be the frame number when using all samples in one frame, or should be the frame on which the first sample was
+// taken if taking one sample per frame like the path tracer).
 void RandomSequence_Initialize(inout RandomSequence RandSequence, uint2 PixelCoord, uint SampleIndex, uint TimeSeed0, uint MaxSamples)
 {
+#if RANDSEQ_ERROR_DIFFUSION == 1 && RANDSEQ == RANDSEQ_OWENSOBOL
+	// The core idea is to use a single sobol pattern across the screen using a space filling curve to correlate nearby pixels
+	// This achieves a much better error diffusion than the random sequence assignment of the method above
+	// This algorithm is covered in the following paper:
+	//   "Screen-Space Blue-Noise Diffusion of Monte Carlo Sampling Error via Hierarchical Ordering of Pixels"
+	//   Siggraph Asia 2020
+	//   http://abdallagafar.com/publications/zsampler/
+	// Improvements made in this version are: 
+	//   - switching the z-order curve for a hilbert curve
+	//   - switching sobol points for owen-sobol points
+	//   - recognizing that the hierarchical scrambling proposed in the original paper is equivalent to Owen-Scrambling of the index
+	// Downsides of this method:
+	//   - quality does not improve beyond the target sample count, making usage with adaptive sampling difficult
+	//   - error-diffusion property is only achieved at the target sample count, earlier samples look random (sometimes slightly worse than random)
+
 	uint X = PixelCoord.x;
 	uint Y = PixelCoord.y;
 
@@ -357,19 +366,11 @@ void RandomSequence_Initialize(inout RandomSequence RandSequence, uint2 PixelCoo
 	//RandSequence.SampleIndex = ReverseBits32(MortonEncode(uint2(MortonEncode(uint2(X, Y)), SampleIndex)));
 
 	RandSequence.SampleSeed = 0; // always use the same sequence
-}
-
 #else
-
-// Error-diffusion sampler is not available, but implement the same API signature for consistency
-void RandomSequence_Initialize(inout RandomSequence RandSequence, uint2 PixelCoord, uint SampleIndex, uint TimeSeed0, uint MaxSamples)
-{
-	// Error diffusion is not available, so just forward to the normal call
-	uint LinearIndex = (PixelCoord.x & 0xFFFF) + (PixelCoord.y * 65536);
-	RandomSequence_Initialize(RandSequence, LinearIndex, TimeSeed0 * MaxSamples + SampleIndex);
-}
-
+	// Error diffusion sampler is disabled, just use the simple init function instead
+	RandomSequence_Initialize(RandSequence, PixelCoord.x + PixelCoord.y * 65536, TimeSeed0 * MaxSamples + SampleIndex);
 #endif
+}
 
 float RandomSequence_GenerateSample1D(inout RandomSequence RandSequence)
 {
diff --git a/Engine/Shaders/Private/RayTracing/RayTracingAmbientOcclusionRGS.usf b/Engine/Shaders/Private/RayTracing/RayTracingAmbientOcclusionRGS.usf
--- a/Engine/Shaders/Private/RayTracing/RayTracingAmbientOcclusionRGS.usf
+++ b/Engine/Shaders/Private/RayTracing/RayTracingAmbientOcclusionRGS.usf
@@ -175,14 +175,9 @@ RAY_TRACING_ENTRY_RAYGEN(AmbientOcclusionRGS)
 		RayDesc Ray;
 		float RayPDF;
 
-#define USE_ERROR_DIFFUSION_SAMPLER 1 && RANDSEQ == RANDSEQ_OWENSOBOL
-
 		RandomSequence RandSequence;
-#if USE_ERROR_DIFFUSION_SAMPLER
 		RandomSequence_Initialize(RandSequence, PixelCoord, SampleIndex, TimeSeed, SamplesPerPixelLocal);
-#else
-		RandomSequence_Initialize(RandSequence, CalcLinearIndex(PixelCoord), SamplesPerPixelLocal * TimeSeed + SampleIndex);
-#endif
+
 		GenerateCosineNormalRay(PixelCoord, WorldPosition, WorldNormal, GeometricNormal, RandSequence, Ray.Origin, Ray.Direction, Ray.TMin, Ray.TMax, RayPDF);
 		ApplyCameraRelativeDepthBias(Ray, PixelCoord, DeviceZ, WorldNormal, MaxNormalBias);
 
diff --git a/Engine/Shaders/Private/RayTracing/RayTracingOcclusionRGS.usf b/Engine/Shaders/Private/RayTracing/RayTracingOcclusionRGS.usf
--- a/Engine/Shaders/Private/RayTracing/RayTracingOcclusionRGS.usf
+++ b/Engine/Shaders/Private/RayTracing/RayTracingOcclusionRGS.usf
@@ -356,25 +356,10 @@ FOcclusionResult ComputeOcclusion(
 	{
 		RandomSequence RandSequence;
 
-// This initiailization improves the visual distribution of the error in screenspace
-#define USE_ERROR_DIFFUSION_SAMPLER 1 && RANDSEQ == RANDSEQ_OWENSOBOL
-
-#if USE_ERROR_DIFFUSION_SAMPLER
-
 #if ENABLE_MULTIPLE_SAMPLES_PER_PIXEL
 		RandomSequence_Initialize(RandSequence, PixelCoord, SampleIndex, TimeSeed, LocalSamplesPerPixel);
 #else
 		RandomSequence_Initialize(RandSequence, PixelCoord, 0, TimeSeed, 1);
-#endif
-
-#else
-
-#if ENABLE_MULTIPLE_SAMPLES_PER_PIXEL
-		RandomSequence_Initialize(RandSequence, CalcLinearIndex(PixelCoord), LocalSamplesPerPixel * TimeSeed + SampleIndex);
-#else
-		RandomSequence_Initialize(RandSequence, CalcLinearIndex(PixelCoord), TimeSeed);
-#endif
-
 #endif
 
 		float2 RandSample = RandomSequence_GenerateSample2D(RandSequence);
@@ -515,11 +500,7 @@ FOcclusionResult ComputeOcclusion(
 	if (ENABLE_TRANSMISSION && LocalSamplesPerPixel > 0 && ShadingModelID == SHADINGMODELID_SUBSURFACE_PROFILE)
 	{
 		RandomSequence RandSequence;
-#if USE_ERROR_DIFFUSION_SAMPLER
 		RandomSequence_Initialize(RandSequence, PixelCoord, 0, TimeSeed, 1);
-#else
-		RandomSequence_Initialize(RandSequence, CalcLinearIndex(PixelCoord), TimeSeed);
-#endif
 
 		// Snag two variables to fix light sampling for all computations..
 		const float2 LightSample = RandomSequence_GenerateSample2D(RandSequence);
