Disabling anisotropy on surfaces that don't have the bit set.

diff --git a/Engine/Shaders/Private/GBufferHelpers.ush b/Engine/Shaders/Private/GBufferHelpers.ush
--- a/Engine/Shaders/Private/GBufferHelpers.ush
+++ b/Engine/Shaders/Private/GBufferHelpers.ush
@@ -376,6 +376,7 @@ void GBufferPostDecode(inout FGBufferData Ret, bool bChecker, bool bGetNormalize
 	Ret.CustomData = !(Ret.SelectiveOutputMask & 0x1) ? Ret.CustomData : 0.0f;
 	Ret.PrecomputedShadowFactors = !(Ret.SelectiveOutputMask & 0x2) ? Ret.PrecomputedShadowFactors : ((Ret.SelectiveOutputMask & 0x4) ? 0.0f : 1.0f);
 	Ret.Velocity = !(Ret.SelectiveOutputMask & 0x8) ? Ret.Velocity : 0.0f;
+	bool bHasAnisotropy = (Ret.SelectiveOutputMask & 0x1);
 
 	Ret.StoredBaseColor = Ret.BaseColor;
 	Ret.StoredMetallic = Ret.Metallic;
@@ -429,6 +430,7 @@ void GBufferPostDecode(inout FGBufferData Ret, bool bChecker, bool bGetNormalize
 		#endif //USE_DEVELOPMENT_SHADERS
 	}
 	
+	if (bHasAnisotropy)
 	{
 		Ret.WorldTangent = DecodeNormal(Ret.WorldTangent);
 		Ret.Anisotropy = Ret.Anisotropy * 2.0f - 1.0f;
@@ -438,6 +440,11 @@ void GBufferPostDecode(inout FGBufferData Ret, bool bChecker, bool bGetNormalize
 			Ret.WorldTangent = normalize(Ret.WorldTangent);
 		}
 	}
+	else
+	{
+		Ret.WorldTangent = 0;
+		Ret.Anisotropy = 0;
+	}
 
 	// This requires cleanup. Shader code that uses GBuffer.SelectiveOutputMask expects the outputmask to be in
 	// bits [4:7], but it gets packed as bits [0:3] in the flexible gbuffer since we might move it around.
