Fixed Strata SSR bug in while we worj on decouple Strata from old shading model.

diff --git a/Engine/Shaders/Private/ReflectionEnvironmentPixelShader.usf b/Engine/Shaders/Private/ReflectionEnvironmentPixelShader.usf
--- a/Engine/Shaders/Private/ReflectionEnvironmentPixelShader.usf
+++ b/Engine/Shaders/Private/ReflectionEnvironmentPixelShader.usf
@@ -362,9 +362,15 @@ void ReflectionEnvironmentSkyLighting(
 #if STRATA_ENABLED && ENABLE_SKY_LIGHT
 	if (bStrataMaterial)
 	{
+		// Hack: AdjustBaseColorAndSpecularColorForSubsurfaceProfileLightingCopyHack forces GBuffer.BaseColor to be white when View.bSubsurfacePostprocessEnabled.
+		// And Starta always set UseSubsurfaceProfile to true to handle SSS. 
+		// STRATA_TODO We need to decouple starta from that, but for now we directly sample EnvBRDF-as-BaseColor here.
+		const float4 GBufferC = GBufferCTexture.SampleLevel(GBufferCTextureSampler, BufferUV, 0);
+		const float3 GBufferBaseColor = GBufferC.rgb;
+		const float3 EnvBRDF = GBufferBaseColor;
+
 		// We use WorldNormal, Roughness and BaseColor (representing EnvBRDF) written by the strata top layer.
 		const float4 SSR = Texture2DSample(ScreenSpaceReflectionsTexture, ScreenSpaceReflectionsSampler, BufferUV);
-		const float3 EnvBRDF = GBuffer.BaseColor;
 		SkyLighting.rgb += SSR.rgb * EnvBRDF;
 
 		// SSR is computed for the top level by averaging BSDF components.
