Improve robustness of primary ray generation for ortho cameras and account for near clip plane

Instead of tracing rays from the actual view origin, generate a point that is on the near clipping plane. This logic can probably be simplified more but this at least gives consistent results with raster views for both persp and ortho cameras

diff --git a/Engine/Shaders/Private/RayTracing/RayTracingCommon.ush b/Engine/Shaders/Private/RayTracing/RayTracingCommon.ush
--- a/Engine/Shaders/Private/RayTracing/RayTracingCommon.ush
+++ b/Engine/Shaders/Private/RayTracing/RayTracingCommon.ush
@@ -311,13 +311,17 @@ RayDesc CreatePrimaryRay(float2 UV)
 	float2 ScreenPosition = (UV - View.ScreenPositionScaleBias.wz) / View.ScreenPositionScaleBias.xy;
 	float3 RayStart_TranslatedWorld = View.TranslatedWorldCameraOrigin;
 	float4 RayEnd_TranslatedWorld = mul(float4(ScreenPosition, 1, 1), View.ScreenToTranslatedWorld);
-	RayEnd_TranslatedWorld *= rcp(RayEnd_TranslatedWorld.w);
 
 	RayDesc Ray;
-	Ray.Origin = View.WorldCameraOrigin;
-	Ray.Direction = normalize(RayEnd_TranslatedWorld.xyz - RayStart_TranslatedWorld);
-	Ray.TMin = 0.0;
-	Ray.TMax = 1.0e27;
+	Ray.Origin = mul(float4(ScreenPosition, 1, 1), View.ScreenToWorld).xyz;
+	Ray.Direction = normalize(RayEnd_TranslatedWorld.xyz - RayEnd_TranslatedWorld.w * RayStart_TranslatedWorld);
+
+	// shift the origin forward or backwards to line up with the clipping plane
+	float Nudge = ((View.NearPlane - 1) / dot(Ray.Direction, View.ViewForward));
+	Ray.Origin += Ray.Direction * Nudge;
+
+	Ray.TMin = 0;
+	Ray.TMax = 1.0e27; // TODO: how to extract the far clip plane?
 	return Ray;
 }
 
