Fix slate vertex shader warning

diff --git a/Engine/Shaders/Private/SlateVertexShader.usf b/Engine/Shaders/Private/SlateVertexShader.usf
--- a/Engine/Shaders/Private/SlateVertexShader.usf
+++ b/Engine/Shaders/Private/SlateVertexShader.usf
@@ -42,7 +42,7 @@ VertexToPixelInterpolants Main(
 	FMaterialVertexParameters VertexParameters = (FMaterialVertexParameters)0;
 	VertexParameters.WorldPosition = WorldPosition.xyz;
 
-	VertexParameters.VertexColor = FinalVertexColor;
+	VertexParameters.VertexColor = VOut.Color;
 
 #if NUM_MATERIAL_TEXCOORDS_VERTEX
 
