Fix some errors in EncodeTileStencilZBaseDelta

diff --git a/Engine/Shaders/Private/HTileEncoding.ush b/Engine/Shaders/Private/HTileEncoding.ush
--- a/Engine/Shaders/Private/HTileEncoding.ush
+++ b/Engine/Shaders/Private/HTileEncoding.ush
@@ -51,43 +51,43 @@ uint EncodeTileStencilZBaseDelta(float MinDepth, float MaxDepth, uint XX, uint S
 	
 	uint ZDelta = 0;
 
-	if (ZDiff & BitFieldMaskU32(1, 13) == 1)
+	if ((ZDiff & BitFieldMaskU32(1, 13)) != 0)
 	{
 		ZDelta = (0xF << 2u) | BitFieldExtractU32(ZDiff, 2, 11);
 	}
-	else if (ZDiff & BitFieldMaskU32(1, 12) == 1)
+	else if ((ZDiff & BitFieldMaskU32(1, 12)) != 0)
 	{
 		ZDelta = (0xE << 2u) | BitFieldExtractU32(ZDiff, 2, 10);
 	}
-	else if (ZDiff & BitFieldMaskU32(1, 11) == 1)
+	else if ((ZDiff & BitFieldMaskU32(1, 11)) != 0)
 	{
 		ZDelta = (0xD << 2u) | BitFieldExtractU32(ZDiff, 2, 9);
 	}
-	else if (ZDiff & BitFieldMaskU32(1, 10) == 1)
+	else if ((ZDiff & BitFieldMaskU32(1, 10)) != 0)
 	{
 		ZDelta = (0xC << 2u) | BitFieldExtractU32(ZDiff, 2, 8);
 	}
-	else if (ZDiff & BitFieldMaskU32(1, 9) == 1)
+	else if ((ZDiff & BitFieldMaskU32(1, 9)) != 0)
 	{
 		ZDelta = (0xB << 2u) | BitFieldExtractU32(ZDiff, 2, 7);
 	}
-	else if (ZDiff & BitFieldMaskU32(1, 8) == 1)
+	else if ((ZDiff & BitFieldMaskU32(1, 8)) != 0)
 	{
 		ZDelta = (0xA << 2u) | BitFieldExtractU32(ZDiff, 2, 6);
 	}
-	else if (ZDiff & BitFieldMaskU32(1, 7) == 1)
+	else if ((ZDiff & BitFieldMaskU32(1, 7)) != 0)
 	{
 		ZDelta = (0x9 << 2u) | BitFieldExtractU32(ZDiff, 2, 5);
 	}
-	else if (ZDiff & BitFieldMaskU32(1, 6) == 1)
+	else if ((ZDiff & BitFieldMaskU32(1, 6)) != 0)
 	{
 		ZDelta = (0x8 << 2u) | BitFieldExtractU32(ZDiff, 2, 4);
 	}
-	else if (ZDiff & BitFieldMaskU32(1, 5) == 1)
+	else if ((ZDiff & BitFieldMaskU32(1, 5)) != 0)
 	{
 		ZDelta = (0x3 << 3u) | BitFieldExtractU32(ZDiff, 3, 2);
 	}
-	else if (ZDiff & BitFieldMaskU32(1, 4) == 1)
+	else if ((ZDiff & BitFieldMaskU32(1, 4)) != 0)
 	{
 		ZDelta = (0x2 << 3u) | BitFieldExtractU32(ZDiff, 3, 1);
 	}
