Cleanup old path tracer light sampling methods, use the new methods for RTGI

diff --git a/Engine/Shaders/Private/PathTracing/Light/PathTracingDirectionalLight.ush b/Engine/Shaders/Private/PathTracing/Light/PathTracingDirectionalLight.ush
--- a/Engine/Shaders/Private/PathTracing/Light/PathTracingDirectionalLight.ush
+++ b/Engine/Shaders/Private/PathTracing/Light/PathTracingDirectionalLight.ush
@@ -56,55 +56,6 @@ bool DirectionalLight_SampleLight(
 	return true;
 }
 
-void DirectionalLight_GenerateLightRay(
-	RayDesc Ray,
-	int LightId,
-	float3 LightUV,
-	out RayDesc OutLightRay
-)
-{
-	OutLightRay.Origin = Ray.Origin;
-	OutLightRay.TMin = 0.0;
-	OutLightRay.Direction = LightUV;
-	OutLightRay.TMax = 1.0e27;
-}
-
-void DirectionalLight_PdfLight(
-	RayDesc Ray,
-	int LightId,
-	float3 LightUV,
-	out float OutPdf
-)
-{
-	OutPdf = 0.0;
-}
-
-void DirectionalLight_EvalLight(
-	int LightId,
-	float3 LightUV,
-	RayDesc Ray,
-	out float3 OutRadiance
-)
-{
-	float3 Color = GetColor(LightId);
-	float Radius = GetRadius(LightId);
-	if (Radius == 0)
-	{
-		OutRadiance = Color;
-	}
-	else
-	{
-		// #dxr_todo: Merge Eval/Sample to avoid redundant computation
-		float SinThetaMax = Radius; // currently sin(Angle / 2)
-
-		// See logic in UniformSampleConeRobust
-		// #dxr_todo: unify the sample and eval methods into one to avoid redundant calculations
-		float SinThetaMax2 = SinThetaMax * SinThetaMax;
-		float OneMinusCosThetaMax = SinThetaMax2 < 0.01 ? SinThetaMax2 * (0.5 + 0.125 * SinThetaMax2) : 1 - sqrt(1 - SinThetaMax2);
-		OutRadiance = Color * 1.0 / (2 * PI * OneMinusCosThetaMax);
-	}
-}
-
 void DirectionalLight_EstimateLight(
 	int LightId,
 	float3 WorldPos,
diff --git a/Engine/Shaders/Private/PathTracing/Light/PathTracingLightCommon.ush b/Engine/Shaders/Private/PathTracing/Light/PathTracingLightCommon.ush
--- a/Engine/Shaders/Private/PathTracing/Light/PathTracingLightCommon.ush
+++ b/Engine/Shaders/Private/PathTracing/Light/PathTracingLightCommon.ush
@@ -150,12 +150,19 @@ float ComputeAttenuationFalloff(float DistanceSquared, int LightId)
 }
 
 
+#ifndef USE_IES_TERM
+#define USE_IES_TERM 1
+#endif
+
+#if USE_IES_TERM
 Texture2DArray<float> IESTexture;
 SamplerState          IESTextureSampler;
+#endif
 
 float ComputeIESAttenuation(int LightId, float3 WorldPosition)
 {
 	float Result = 1.f;
+#if USE_IES_TERM
 	if (SceneLightsData.IESTextureSlice[LightId] >= 0)
 	{
 		float3 LightDirection = normalize(GetPosition(LightId) - WorldPosition);
@@ -171,6 +178,7 @@ float ComputeIESAttenuation(int LightId, float3 WorldPosition)
 		const float3 UVW = float3(NormAngle, NormTangentAngle, SceneLightsData.IESTextureSlice[LightId]);
 		Result = IESTexture.SampleLevel(IESTextureSampler, UVW, 0);
 	}
+#endif
 	return Result;
 }
 
diff --git a/Engine/Shaders/Private/PathTracing/Light/PathTracingLightSampling.ush b/Engine/Shaders/Private/PathTracing/Light/PathTracingLightSampling.ush
--- a/Engine/Shaders/Private/PathTracing/Light/PathTracingLightSampling.ush
+++ b/Engine/Shaders/Private/PathTracing/Light/PathTracingLightSampling.ush
@@ -71,33 +71,6 @@ void SampleLight(
 	}
 }
 
-void GenerateLightRay(RayDesc Ray, int LightId, float3 LightUV, inout RayDesc OutLightRay)
-{
-	OutLightRay = (RayDesc)0; // #dxr_todo: Return false in the case of no matching light type and handle case from caller rather than blindly using a degenerate ray
-
-	switch (SceneLightsData.Flags[LightId] & PATHTRACER_FLAG_TYPE_MASK)
-	{
-		case PATHTRACING_LIGHT_SKY:
-			SkyLight_GenerateLightRay(Ray, LightId, LightUV, OutLightRay);
-			return;
-		case PATHTRACING_LIGHT_POINT:
-			PointLight_GenerateLightRay(Ray, LightId, LightUV, OutLightRay);
-			return;
-		case PATHTRACING_LIGHT_DIRECTIONAL:
-			DirectionalLight_GenerateLightRay(Ray, LightId, LightUV, OutLightRay);
-			return;
-		case PATHTRACING_LIGHT_RECT:
-			RectLight_GenerateLightRay(Ray, LightId, LightUV, OutLightRay);
-			return;
-		case PATHTRACING_LIGHT_SPOT:
-			SpotLight_GenerateLightRay(Ray, LightId, LightUV, OutLightRay);
-			return;
-		default:
-			return;
-	}
-}
-
-
 void EstimateLight(
 	int LightId,
 	float3 WorldPos,
@@ -132,236 +105,6 @@ void EstimateLight(
 	}
 }
 
-bool SampleLightSelection(
-	RayDesc Ray,
-	FMaterialClosestHitPayload Payload,
-	float RandSample,
-	inout int OutLightId,
-	inout float OutLightSelectionPdf
-)
-{
-	OutLightSelectionPdf = 0.0;
-
-	if (SceneLightsData.Count == 0)
-	{
-		return false;
-	}
-
-#ifdef UNIFORM_LIGHT_SELECTION
-	OutLightId = RandSample * SceneLightsData.Count;
-	OutLightSelectionPdf = 1.0 / SceneLightsData.Count;
-	return true;
-#else
-
-	// Build irradiance estimate prefix sum
-	float IrradianceEstimate[RAY_TRACING_LIGHT_COUNT_MAXIMUM];
-
-	float3 WorldPos = Payload.WorldPos;
-	float3 WorldNormal = Payload.WorldNormal;
-	uint PrimitiveLightingChannelMask = Payload.PrimitiveLightingChannelMask;
-	
-	// TODO: query transmission from material instead of light?
-	EstimateLight(0, WorldPos, WorldNormal, PrimitiveLightingChannelMask, ENABLE_TRANSMISSION && HasTransmission(0), IrradianceEstimate[0]);
-	{
-		for (uint LightIndex = 1; LightIndex < SceneLightsData.Count; ++LightIndex)
-		{
-			EstimateLight(LightIndex, WorldPos, WorldNormal, PrimitiveLightingChannelMask, ENABLE_TRANSMISSION && HasTransmission(LightIndex), IrradianceEstimate[LightIndex]);
-			IrradianceEstimate[LightIndex] += IrradianceEstimate[LightIndex - 1];
-		}
-	}
-	
-	if (IrradianceEstimate[SceneLightsData.Count - 1] == 0)
-	{
-		return false;
-	}
-
-	// CDF normalization
-	{
-		for (uint LightIndex = 0; LightIndex < SceneLightsData.Count; ++LightIndex)
-		{
-			IrradianceEstimate[LightIndex] /= IrradianceEstimate[SceneLightsData.Count - 1];
-		}
-	}
-	IrradianceEstimate[SceneLightsData.Count - 1] = 1.0;
-
-	// Draw random variable
-	// dxr_todo: binary search
-	for (OutLightId = 0; OutLightId < SceneLightsData.Count; ++OutLightId)
-	{
-		if (RandSample < IrradianceEstimate[OutLightId]) break;
-	}
-
-	OutLightSelectionPdf = IrradianceEstimate[OutLightId];
-	if (OutLightId > 0)
-	{
-		OutLightSelectionPdf -= IrradianceEstimate[OutLightId - 1];
-	}
-
-	return true;
-#endif
-}
-
-void PdfLightSelection(
-	RayDesc Ray,
-	FMaterialClosestHitPayload Payload,
-	uint LightId,
-	inout float OutLightSelectionPdf
-)
-{
-	OutLightSelectionPdf = 0.0;
-
-#ifdef UNIFORM_LIGHT_SELECTION
-	OutLightSelectionPdf = 1.0 / SceneLightsData.Count;
-	return;
-#else
-	if (SceneLightsData.Count == 0) return;
-
-	// Build irradiance estimate prefix sum
-	float IrradianceEstimate[RAY_TRACING_LIGHT_COUNT_MAXIMUM];
-
-	float3 WorldPos = Payload.WorldPos;
-	float3 WorldNormal = Payload.WorldNormal;
-	uint PrimitiveLightingChannelMask = Payload.PrimitiveLightingChannelMask;
-
-	EstimateLight(0, WorldPos, WorldNormal, PrimitiveLightingChannelMask, ENABLE_TRANSMISSION && HasTransmission(0), IrradianceEstimate[0]);
-	{
-		for (uint LightIndex = 1; LightIndex < SceneLightsData.Count; ++LightIndex)
-		{
-			EstimateLight(LightIndex, WorldPos, WorldNormal, PrimitiveLightingChannelMask, ENABLE_TRANSMISSION && HasTransmission(LightIndex), IrradianceEstimate[LightIndex]);
-			IrradianceEstimate[LightIndex] += IrradianceEstimate[LightIndex - 1];
-		}
-	}
-	// dxr_todo: Consider when IrradianceEstimate[SceneLightsData.Count] == 0
-
-	// CDF normalization
-	{
-		for (uint LightIndex = 0; LightIndex < SceneLightsData.Count; ++LightIndex)
-		{
-			IrradianceEstimate[LightIndex] /= IrradianceEstimate[SceneLightsData.Count - 1];
-		}
-	}
-	IrradianceEstimate[SceneLightsData.Count - 1] = 1.0;
-
-	OutLightSelectionPdf = IrradianceEstimate[LightId];
-	if (LightId > 0)
-	{
-		OutLightSelectionPdf -= IrradianceEstimate[LightId - 1];
-	}
-#endif
-}
-
-void PdfLight(
-	RayDesc Ray,
-	FMaterialClosestHitPayload Payload,
-	int LightId,
-	float3 LightUV,
-	inout float OutPdf)
-{
-	// Light selection
-	// TODO: Use radiance estimate instead of uniform random selection
-	float LightSelectionPdf = 1.0 / SceneLightsData.Count;
-	PdfLightSelection(Ray, Payload, LightId, LightSelectionPdf);
-	float LightPdf = 0.0;
-
-	// Light sampling
-	switch (SceneLightsData.Flags[LightId] & PATHTRACER_FLAG_TYPE_MASK)
-	{
-		case PATHTRACING_LIGHT_SKY:
-			SkyLight_PdfLight(Ray, LightId, LightUV, LightPdf);
-			break;
-		case PATHTRACING_LIGHT_POINT:
-			PointLight_PdfLight(Ray, LightId, LightUV, LightPdf);
-			break;
-		case PATHTRACING_LIGHT_DIRECTIONAL:
-			DirectionalLight_PdfLight(Ray, LightId, LightUV, LightPdf);
-			break;
-		case PATHTRACING_LIGHT_RECT:
-			RectLight_PdfLight(Ray, LightId, LightUV, LightPdf);
-			break;
-		case PATHTRACING_LIGHT_SPOT:
-			SpotLight_PdfLight(Ray, LightId, LightUV, LightPdf);
-			break;
-		default:
-			break;
-	}
-
-	OutPdf = LightSelectionPdf * LightPdf;
-}
-
-void EvalLight(
-	int LightId,
-	float3 LightUV,
-	RayDesc Ray,
-	inout float3 OutRadiance
-)
-{
-	OutRadiance = float3(0.0f, 0.0f, 0.0f);
-
-	switch (SceneLightsData.Flags[LightId] & PATHTRACER_FLAG_TYPE_MASK)
-	{
-		case PATHTRACING_LIGHT_SKY:
-			SkyLight_EvalLight(LightId, LightUV, Ray, OutRadiance);
-			return;
-		case PATHTRACING_LIGHT_POINT:
-			PointLight_EvalLight(LightId, LightUV, Ray, OutRadiance);
-			return;
-		case PATHTRACING_LIGHT_DIRECTIONAL:
-			DirectionalLight_EvalLight(LightId, LightUV, Ray, OutRadiance);
-			return;
-		case PATHTRACING_LIGHT_RECT:
-			RectLight_EvalLight(LightId, LightUV, Ray, OutRadiance);
-			return;
-		case PATHTRACING_LIGHT_SPOT:
-			SpotLight_EvalLight(LightId, LightUV, Ray, OutRadiance);
-			return;
-		default:
-			return;
-	}
-}
-
-void SampleLight(
-	RayDesc Ray,
-	FMaterialClosestHitPayload Payload,
-	float4 RandSample,
-	inout int OutLightId,
-	inout float3 OutLightUV,
-	inout float OutPdf)
-{
-	OutPdf = 0.0f;
-
-	// Light selection
-	float LightSelectionPdf = 0.0;
-	if (!SampleLightSelection(Ray, Payload, RandSample.x, OutLightId, LightSelectionPdf))
-	{
-		return;
-	}
-
-	// Light sampling
-	float LightPdf = 0.0;
-
-	switch (SceneLightsData.Flags[OutLightId] & PATHTRACER_FLAG_TYPE_MASK)
-	{
-		case PATHTRACING_LIGHT_SKY:
-			SkyLight_SampleLight(Ray, Payload, RandSample, OutLightId, OutLightUV, LightPdf);
-			break;
-		case PATHTRACING_LIGHT_POINT:
-			PointLight_SampleLight(Ray, Payload, RandSample, OutLightId, OutLightUV, LightPdf);
-			break;
-		case PATHTRACING_LIGHT_DIRECTIONAL:
-			DirectionalLight_SampleLight(Ray, Payload, RandSample, OutLightId, OutLightUV, LightPdf);
-			break;
-		case PATHTRACING_LIGHT_RECT:
-			RectLight_SampleLight(Ray, Payload, RandSample, OutLightId, OutLightUV, LightPdf);
-			break;
-		case PATHTRACING_LIGHT_SPOT:
-			SpotLight_SampleLight(Ray, Payload, RandSample, OutLightId, OutLightUV, LightPdf);
-			break;
-		default:
-			break;
-	}
-	OutPdf = LightSelectionPdf * LightPdf;
-}
-
 // #dxr_todo: passing LightPickingCdf as anything but inout causes performance to drop significantly -- why ?
 void InitLightPickingCdf(inout float LightPickingCdf[RAY_TRACING_LIGHT_COUNT_MAXIMUM])
 {
@@ -384,16 +127,17 @@ bool InitLightPickingCdf(
 	InitLightPickingCdf(LightPickingCdf);
 #else
 	// Build irradiance estimate prefix sum
-	EstimateLight(0, WorldPos, WorldNormal, PrimitiveLightingChannelMask, IsTransmissiveMaterial, LightPickingCdf[0]);
+	float LastCdf = 0;
 	{
-		for (uint LightIndex = 1; LightIndex < SceneLightsData.Count; ++LightIndex)
+		for (uint LightIndex = 0; LightIndex < SceneLightsData.Count; ++LightIndex)
 		{
 			EstimateLight(LightIndex, WorldPos, WorldNormal, PrimitiveLightingChannelMask, IsTransmissiveMaterial, LightPickingCdf[LightIndex]);
-			LightPickingCdf[LightIndex] += LightPickingCdf[LightIndex - 1];
+			LightPickingCdf[LightIndex] += LastCdf;
+			LastCdf = LightPickingCdf[LightIndex];
 		}
 	}
 
-	if (LightPickingCdf[SceneLightsData.Count - 1] == 0)
+	if (LastCdf == 0)
 	{
 		// all the lights are invisible
 		return false;
@@ -402,7 +146,7 @@ bool InitLightPickingCdf(
 		// CDF normalization
 		for (uint LightIndex = 0; LightIndex < SceneLightsData.Count; ++LightIndex)
 		{
-			LightPickingCdf[LightIndex] /= LightPickingCdf[SceneLightsData.Count - 1];
+			LightPickingCdf[LightIndex] /= LastCdf;
 		}
 		LightPickingCdf[SceneLightsData.Count - 1] = 1.0;
 	}
@@ -429,3 +173,51 @@ void SelectLight(float RandSample, inout float LightPickingCdf[RAY_TRACING_LIGHT
 
 	LightPickPdf = LightPickingCdf[LightId] - PreviousCdfValue;
 }
+
+bool SampleLight(
+	float4 RandSample,
+	float3 WorldPos,
+	float3 WorldNormal,
+	uint PrimitiveLightingChannelMask,
+	bool IsTransmissiveMaterial,
+	out int OutLightId,
+	out float3 OutDirection,
+	out float  OutDistance,
+	out float3 OutRadianceOverPdf,
+	out float OutPdf)
+{
+	// set default values in case sampling fails
+	OutLightId = -1;
+	OutDirection = 0;
+	OutDistance = 0;
+	OutRadianceOverPdf = 0;
+	OutPdf = 0;
+
+	if (SceneLightsData.Count == 0)
+	{
+		// no lights in the scene
+		return false;
+	}
+
+	float LightPickingCdf[RAY_TRACING_LIGHT_COUNT_MAXIMUM];
+
+	if (!InitLightPickingCdf(WorldPos, WorldNormal, PrimitiveLightingChannelMask, IsTransmissiveMaterial, LightPickingCdf))
+	{
+		// light picking cdf is all 0
+		return false;
+	}
+
+	float LightPickPdf = 0;
+
+	SelectLight(RandSample.x, LightPickingCdf, OutLightId, LightPickPdf);
+
+	SampleLight(OutLightId, RandSample, WorldPos, WorldNormal,
+		OutDirection,
+		OutDistance,
+		OutRadianceOverPdf,
+		OutPdf);
+	OutRadianceOverPdf /= LightPickPdf;
+	OutPdf *= LightPickPdf;
+
+	return true;
+}
diff --git a/Engine/Shaders/Private/PathTracing/Light/PathTracingPointLight.ush b/Engine/Shaders/Private/PathTracing/Light/PathTracingPointLight.ush
--- a/Engine/Shaders/Private/PathTracing/Light/PathTracingPointLight.ush
+++ b/Engine/Shaders/Private/PathTracing/Light/PathTracingPointLight.ush
@@ -86,64 +86,6 @@ void PointLight_SampleLight(
 	OutRadianceOverPdf *= ComputeAttenuationFalloff(LightDistanceSquared, LightId);
 }
 
-void PointLight_GenerateLightRay(
-	RayDesc Ray,
-	int LightId,
-	float3 LightUV,
-	out RayDesc OutLightRay
-)
-{
-	float3 LightDirection = GetPosition(LightId) - Ray.Origin;
-
-	OutLightRay.Origin = Ray.Origin;
-	OutLightRay.TMin = 0.001;
-	OutLightRay.Direction = TangentToWorld(LightUV, normalize(LightDirection));
-
-	// Compute distance to sphere edge
-	float Radius = GetRadius(LightId);
-	float Radius2 = Pow2(Radius);
-
-	float RayLength2 = dot(LightDirection, LightDirection);
-	float RayLength = sqrt(RayLength2);
-	float SinThetaMax2 = Radius2 / RayLength2;
-
-	float CosTheta = LightUV.z;
-	float SinTheta2 = 1.0 - CosTheta * CosTheta;
-
-	OutLightRay.TMax = RayLength * (CosTheta - sqrt(max(SinThetaMax2 - SinTheta2, 0.0)));
-}
-
-void PointLight_PdfLight(
-	RayDesc Ray,
-	int LightId,
-	float3 LightUV,
-	out float OutPdf
-)
-{
-	OutPdf = 0.0;
-}
-
-void PointLight_EvalLight(
-	int LightId,
-	float3 LightUV,
-	RayDesc Ray,
-	out float3 OutRadiance
-)
-{
-	float3 LightDirection = GetPosition(LightId) - Ray.Origin;
-	float LightDistanceSquared = dot(LightDirection, LightDirection);
-
-	float3 Color = GetColor(LightId);
-	float Radius = GetRadius(LightId);
-	float Falloff = ComputeAttenuationFalloff(LightDistanceSquared, LightId);
-	OutRadiance = Color * Falloff;
-
-	if (Radius == 0)
-		OutRadiance *= 1 / LightDistanceSquared;
-	else
-		OutRadiance *= 1 / (PI * Radius * Radius);
-}
-
 void PointLight_EstimateLight(
 	int LightId,
 	float3 WorldPos,
@@ -163,37 +105,3 @@ void PointLight_EstimateLight(
 	float Falloff = ComputeAttenuationFalloff(LightDistanceSquared, LightId);
 	OutIrradiance = LightPower * Falloff * NoL / LightDistanceSquared;
 }
-
-void PointLight_SampleLight(
-	RayDesc Ray,
-	FMaterialClosestHitPayload Payload,
-	float4 RandSample,
-	int LightId,
-	out float3 OutLightUV,
-	out float OutPdf
-)
-{
-	float Radius = GetRadius(LightId);
-	if (Radius > 0.0)
-	{
-		// Soft case - sample the solid angle subtended by the sphere
-		float3 LightDirection = GetPosition(LightId) - Payload.WorldPos;
-		float LightDistanceSquared = dot(LightDirection, LightDirection);
-		float Radius2 = Pow2(Radius);
-		
-		// outside the light
-		float SinThetaMax2 = Radius2 / LightDistanceSquared;
-		
-		// #dxr_todo: find a better way of handling the region inside the light than just clamping to 1.0 here
-		float4 DirAndPdf = UniformSampleConeRobust(RandSample.yz, min(SinThetaMax2, 1.0));
-		OutLightUV = DirAndPdf.xyz;
-		OutPdf = DirAndPdf.w;
-	}
-	else
-	{
-		// Singular case
-		OutLightUV = float3(0, 0, 1);
-		// #dxr_todo: this isn't really part of the pdf, the pdf should be +inf for the singular case
-		OutPdf = 1.0;
-	}
-}
diff --git a/Engine/Shaders/Private/PathTracing/Light/PathTracingRectLight.ush b/Engine/Shaders/Private/PathTracing/Light/PathTracingRectLight.ush
--- a/Engine/Shaders/Private/PathTracing/Light/PathTracingRectLight.ush
+++ b/Engine/Shaders/Private/PathTracing/Light/PathTracingRectLight.ush
@@ -9,6 +9,53 @@
 #include "../../RectLight.ush"
 #include "PathTracingLightCommon.ush"
 
+float3 BilinearQuadWarp(float2 uv, float W00, float W01, float W10, float W11) {
+	// "Practical Product Sampling by Fitting and Composing Warps" - EGSR 2020
+	// https://casual-effects.com/research/Hart2020Sampling/index.html
+	// https://www.shadertoy.com/view/wljyDz
+
+	float a = lerp(W00, W01, .5);
+	float b = lerp(W10, W11, .5);
+	float u = a == b ? uv.x : (sqrt(lerp(a * a, b * b, uv.x)) - a) / (b - a);
+	float c = lerp(W00, W10, u);
+	float d = lerp(W01, W11, u);
+	float v = c == d ? uv.y : (sqrt(lerp(c * c, d * d, uv.y)) - c) / (d - c);
+	float area = lerp(a, b, .5);
+	float pdf = lerp(c, d, v) / area;
+	return float3(u, v, pdf);
+}
+
+// #dxr_todo: This method should be moved to RectLight.ush
+// #dxr_todo: Analyze the overall efficiency (MSE reduction vs. time)
+//  Projected solid angle sampling converges faster at equal sample count, but is a bit slower
+// #dxr_todo: Should we switch to plain area sampling at some distance?
+float4 SampleApproxProjectionSphericalRect(float2 Rand, FSphericalRect SphericalRect, float3 WorldNormal)
+{
+	// Construct the quad vertices
+	float3 S00 = float3(SphericalRect.x0, SphericalRect.y0, SphericalRect.z0);
+	float3 S10 = float3(SphericalRect.x1, SphericalRect.y0, SphericalRect.z0);
+	float3 S11 = float3(SphericalRect.x1, SphericalRect.y1, SphericalRect.z0);
+	float3 S01 = float3(SphericalRect.x0, SphericalRect.y1, SphericalRect.z0);
+
+	// Compute the cosine of the normal against each corner of the quad
+	float3 LightSpaceNormal = mul(SphericalRect.Axis, WorldNormal);
+	float W00 = saturate(dot(normalize(S00), LightSpaceNormal));
+	float W10 = saturate(dot(normalize(S10), LightSpaceNormal));
+	float W11 = saturate(dot(normalize(S11), LightSpaceNormal));
+	float W01 = saturate(dot(normalize(S01), LightSpaceNormal));
+
+	// Warp a 2D sample according to the cosine at each corner of the quad
+	float3 WarpedRand = BilinearQuadWarp(Rand, W00, W01, W10, W11);
+
+	// Proceed with uniform sampling of the spherical rectangle
+	float3 Direction = UniformSampleSphericalRect(WarpedRand.xy, SphericalRect);
+	float OutPdf = isfinite(SphericalRect.SolidAngle) ? 1.0 / SphericalRect.SolidAngle : 0.0;
+
+	// Return Direction vector and adjusted PDF
+	return float4(Direction, OutPdf * WarpedRand.z);
+}
+
+
 float4 RectLight_TraceLight(RayDesc Ray, int LightId, float HitT)
 {
 	float3 LightPosition = GetPosition(LightId);
@@ -91,42 +138,6 @@ void RectLight_SampleLight(
 	OutPdf = 1.0 / SphericalRect.SolidAngle;
 }
 
-void RectLight_GenerateLightRay(
-	RayDesc Ray,
-	int LightId,
-	float3 LightUV,
-	out RayDesc OutLightRay
-)
-{
-	OutLightRay.Origin = Ray.Origin;
-	OutLightRay.TMin = 0.001;
-	OutLightRay.Direction = normalize(LightUV);
-	OutLightRay.TMax = length(LightUV);
-}
-
-void RectLight_PdfLight(
-	RayDesc Ray,
-	int LightId,
-	float3 LightUV,
-	out float OutPdf
-)
-{
-	// TODO: Consider making RectLight a physical light. For now, it is analytical only.
-	OutPdf = 0.0;
-}
-
-void RectLight_EvalLight(
-	int LightId,
-	float3 LightUV,
-	RayDesc Ray,
-	out float3 OutRadiance
-)
-{
-	float3 Color = GetColor(LightId);
-	float Falloff = ComputeAttenuationFalloff(Square(Ray.TMax), LightId);
-	OutRadiance = Color * Falloff;
-}
-
 void RectLight_EstimateLight(
 	int LightId,
 	float3 WorldPos,
@@ -159,89 +170,3 @@ void RectLight_EstimateLight(
 	float Falloff = ComputeAttenuationFalloff(LightDistanceSquared, LightId);
 	OutIrradiance = LightPower * Falloff * Area * NoL / LightDistanceSquared;
 }
-
-float3 BilinearQuadWarp(float2 uv, float W00, float W01, float W10, float W11) {
-	// "Practical Product Sampling by Fitting and Composing Warps" - EGSR 2020
-	// https://casual-effects.com/research/Hart2020Sampling/index.html
-	// https://www.shadertoy.com/view/wljyDz
-
-	float a = lerp(W00, W01, .5);
-	float b = lerp(W10, W11, .5);
-	float u = a == b ? uv.x : (sqrt(lerp(a * a, b * b, uv.x)) - a) / (b - a);
-	float c = lerp(W00, W10, u);
-	float d = lerp(W01, W11, u);
-	float v = c == d ? uv.y : (sqrt(lerp(c * c, d * d, uv.y)) - c) / (d - c);
-	float area = lerp(a, b, .5);
-	float pdf = lerp(c, d, v) / area;
-	return float3(u, v, pdf);
-}
-
-// #dxr_todo: This method should be moved to RectLight.ush
-float4 SampleApproxProjectionSphericalRect(float2 Rand, FSphericalRect SphericalRect, float3 WorldNormal)
-{
-	// Construct the quad vertices
-	float3 S00 = float3(SphericalRect.x0, SphericalRect.y0, SphericalRect.z0);
-	float3 S10 = float3(SphericalRect.x1, SphericalRect.y0, SphericalRect.z0);
-	float3 S11 = float3(SphericalRect.x1, SphericalRect.y1, SphericalRect.z0);
-	float3 S01 = float3(SphericalRect.x0, SphericalRect.y1, SphericalRect.z0);
-	
-	// Compute the cosine of the normal against each corner of the quad
-	float3 LightSpaceNormal = mul(SphericalRect.Axis, WorldNormal);
-	float W00 = saturate(dot(normalize(S00), LightSpaceNormal));
-	float W10 = saturate(dot(normalize(S10), LightSpaceNormal));
-	float W11 = saturate(dot(normalize(S11), LightSpaceNormal));
-	float W01 = saturate(dot(normalize(S01), LightSpaceNormal));
-	
-	// Warp a 2D sample according to the cosine at each corner of the quad
-	float3 WarpedRand = BilinearQuadWarp(Rand, W00, W01, W10, W11);
-
-	// Proceed with uniform sampling of the spherical rectangle
-	float3 Direction = UniformSampleSphericalRect(WarpedRand.xy, SphericalRect);
-	float OutPdf = isfinite(SphericalRect.SolidAngle) ? 1.0 / SphericalRect.SolidAngle : 0.0;
-
-	// Return Direction vector and adjusted PDF
-	return float4(Direction, OutPdf * WarpedRand.z);
-}
-
-void RectLight_SampleLight(
-	RayDesc Ray,
-	FMaterialClosestHitPayload Payload,
-	float4 RandSample,
-	int LightId,
-	out float3 OutLightUV,
-	out float OutPdf
-)
-{
-	float3 WorldPosition = Payload.WorldPos;
-	float3 LightPosition = GetPosition(LightId);
-	float3 LightNormal = GetNormal(LightId);
-	float3 LightdPdu = GetdPdu(LightId);
-	float3 LightdPdv = GetdPdv(LightId);
-	float LightWidth = GetWidth(LightId);
-	float LightHeight = GetHeight(LightId);
-	
-	// Define rectangle and compute solid angle
-	FRect Rect = GetRect(LightPosition - WorldPosition,
-						 -LightNormal,
-						 LightdPdv,
-						 0.5 * LightWidth,
-						 0.5 * LightHeight,
-						 GetRectLightBarnCosAngle(LightId),
-						 GetRectLightBarnLength(LightId),
-						 true /* bComputeVisibleRect */);
-	FSphericalRect SphericalRect = BuildSphericalRect(Rect);
-
-	// #dxr_todo: Analyze the overall efficiency (MSE reduction vs. time)
-	//  Projected solid angle sampling converges faster at equal sample count, but is a bit slower
-	// #dxr_todo: Should we switch to plain area sampling at some distance?
-#if 1
-	// Approximate _projected_ solidangle sampling of the spherical rectangle
-	float4 Result = SampleApproxProjectionSphericalRect(RandSample.yz, SphericalRect, Payload.WorldNormal);
-	OutLightUV = Result.xyz;
-	OutPdf = Result.w;
-#else
-	// Ordinary solid angle sampling of the spherical rectangle
-	OutLightUV = UniformSampleSphericalRect(RandSample.yz, SphericalRect);
-	OutPdf = isfinite(SphericalRect.SolidAngle) ? 1.0 / SphericalRect.SolidAngle : 0.0;
-#endif
-}
\ No newline at end of file
diff --git a/Engine/Shaders/Private/PathTracing/Light/PathTracingSkyLight.ush b/Engine/Shaders/Private/PathTracing/Light/PathTracingSkyLight.ush
--- a/Engine/Shaders/Private/PathTracing/Light/PathTracingSkyLight.ush
+++ b/Engine/Shaders/Private/PathTracing/Light/PathTracingSkyLight.ush
@@ -11,59 +11,6 @@ SkyLight.usf: Light sampling functions for SkyLight implementation
 
 #define USE_HIERARCHICAL_IMPORTANCE_SAMPLING
 
-void SkyLight_GenerateLightRay(
-	RayDesc Ray,
-	int LightId,
-	float3 LightUV,
-	out RayDesc OutLightRay
-)
-{
-	OutLightRay.Origin = Ray.Origin;
-	OutLightRay.Direction = LightUV;
-	OutLightRay.TMin = 0.0;
-	OutLightRay.TMax = 1.0e27;
-}
-
-void SkyLight_PdfLight(
-	RayDesc Ray,
-	int LightId,
-	float3 LightUV,
-	out float OutPdf
-)
-{
-#ifndef USE_HIERARCHICAL_IMPORTANCE_SAMPLING
-	OutPdf = 1.0 / (2.0 * PI);
-#else
-	// Determine primary direction
-	//uint MipCount = Log2(SkyLight.MipDimensions.x);
-	uint MipCount = log2(SkyLight.MipDimensions.x);
-	uint2 TextureRes;
-	SkyLight.Texture.GetDimensions(TextureRes.x, TextureRes.y);
-	float3 LightDirection = LightUV;
-	uint3 TextureCoord = GetTextureCubeCoordinate(LightDirection, TextureRes);
-	float2 TextureUV = TextureCoord.xy / (float2)TextureRes;
-	uint2 MipCoord = TextureUV * SkyLight.MipDimensions.xy;
-	float FacePdf = PdfFace(MipCount, TextureCoord.z);
-
-	uint StopLevel = min(SkyLight.SamplingStopLevel, MipCount);
-	float MipPdf = PdfMipTree(uint3(MipCoord, TextureCoord.z), StopLevel);
-
-	uint LinearIndex = BufferOffsetAtPixel(MipCoord, StopLevel, SkyLight.MipDimensions.xy);
-	float AreaPdf = SkyLight.SolidAnglePdf[LinearIndex];
-	OutPdf = FacePdf * MipPdf * AreaPdf;
-#endif
-}
-
-void SkyLight_EvalLight(
-	int LightId,
-	float3 LightUV,
-	RayDesc Ray,
-	out float3 OutRadiance
-)
-{
-	OutRadiance = SkyLight.Color * TextureCubeSampleLevel(SkyLight.Texture, SkyLight.TextureSampler, LightUV, 0).rgb;
-}
-
 void SkyLight_EstimateLight(
 	int LightId,
 	float3 WorldPos,
@@ -80,23 +27,73 @@ void SkyLight_EstimateLight(
 	OutIrradiance = SkyLightIlluminance * 2.0f * PI;
 }
 
+float4 SkyLight_TraceLight(RayDesc Ray, int LightId, float HitT) {
+	if (Ray.TMax == HitT)
+	{
+		// LightUV convention in the code below is just the world space ray direction
+		// so its safe to call into the method below for now ...
+
+		// #dxr_todo: write optimized form that combines radiance and pdf lookup in a single step
+
+		float3 Radiance = 0.0;
+		float Pdf = 0.0;
+
+		float3 LightUV = normalize(Ray.Direction);
+		Radiance = SkyLight.Color * SkyLight.Texture.SampleLevel(SkyLight.TextureSampler, LightUV, 0).rgb;
+		if (any(Radiance > 0))
+		{
+#ifndef USE_HIERARCHICAL_IMPORTANCE_SAMPLING
+			Pdf = 1.0 / (2.0 * PI);
+#else
+			// Determine primary direction
+			//uint MipCount = Log2(SkyLight.MipDimensions.x);
+			uint MipCount = log2(SkyLight.MipDimensions.x);
+			uint2 TextureRes;
+			SkyLight.Texture.GetDimensions(TextureRes.x, TextureRes.y);
+			float3 LightDirection = LightUV;
+			uint3 TextureCoord = GetTextureCubeCoordinate(LightDirection, TextureRes);
+			float2 TextureUV = TextureCoord.xy / (float2)TextureRes;
+			uint2 MipCoord = TextureUV * SkyLight.MipDimensions.xy;
+			float FacePdf = PdfFace(MipCount, TextureCoord.z);
+
+			uint StopLevel = min(SkyLight.SamplingStopLevel, MipCount);
+			float MipPdf = PdfMipTree(uint3(MipCoord, TextureCoord.z), StopLevel);
+
+			uint LinearIndex = BufferOffsetAtPixel(MipCoord, StopLevel, SkyLight.MipDimensions.xy);
+			float AreaPdf = SkyLight.SolidAnglePdf[LinearIndex];
+			Pdf = FacePdf * MipPdf * AreaPdf;
+#endif
+		}
+		
+
+		return float4(Radiance, Pdf);
+	}
+	return 0.0;
+}
+
 void SkyLight_SampleLight(
-	RayDesc Ray,
-	FMaterialClosestHitPayload Payload,
-	float4 RandSample,
 	int LightId,
-	out float3 OutLightUV,
-	out float OutPdf
-)
+	float4 RandSample,
+	float3 WorldPos,
+	float3 WorldNormal,
+	out float3 OutDirection,
+	out float  OutDistance,
+	out float3 OutRadianceOverPdf,
+	out float OutPdf)
 {
+	// Make temp structures so we can re-use the calls above
+	// #dxr_todo: write optimized form that combines radiance and pdf lookup in a single step
+	RayDesc Ray;
+	Ray.Origin = WorldPos;
+	FMaterialClosestHitPayload Payload;
+	Payload.WorldNormal = WorldNormal;
+
+
 #ifndef USE_HIERARCHICAL_IMPORTANCE_SAMPLING
 	float4 Direction_Tangent = UniformSampleHemisphere(RandSample.yz);
-	OutLightUV = TangentToWorld(Direction_Tangent.xyz, Payload.WorldNormal);
+	OutDirection = TangentToWorld(Direction_Tangent.xyz, WorldNormal);
 	OutPdf = Direction_Tangent.w;
 #else
-	float3 WorldPosition = Ray.Origin;
-	float3 WorldNormal = Payload.WorldNormal;
-
 	// Determine the appropriate face to use
 	//uint MipCount = Log2(SkyLight.MipDimensions.x);
 	uint MipCount = log2(SkyLight.MipDimensions.x);
@@ -163,64 +160,12 @@ void SkyLight_SampleLight(
 	uint LinearIndex = BufferOffsetAtPixel(MipPixel, StopLevel, SkyLight.MipDimensions.xy);
 	float AreaPdf = SkyLight.SolidAnglePdf[LinearIndex];
 
-	float3 RayDirection = normalize(GetTextureCubeVector(TexelCoord, TextureDimensions));
-	float RayPdf = FacePdf * MipPdf * AreaPdf;
-
-	OutLightUV = RayDirection;
-	OutPdf = RayPdf;
+	OutDirection = normalize(GetTextureCubeVector(TexelCoord, TextureDimensions));
+	OutPdf = FacePdf * MipPdf * AreaPdf;
 #endif
-}
-
-
-//  NOTE: put new methods at the end for now, so we can re-use the methods above
-float4 SkyLight_TraceLight(RayDesc Ray, int LightId, float HitT) {
-	if (Ray.TMax == HitT)
-	{
-		// LightUV convention in the code below is just the world space ray direction
-		// so its safe to call into the method below for now ...
-
-		// #dxr_todo: write optimized form that combines radiance and pdf lookup in a single step
-
-		float3 Radiance = 0.0;
-		float Pdf = 0.0;
-
-		float3 LightUV = normalize(Ray.Direction);
-		SkyLight_EvalLight(LightId, LightUV, Ray, Radiance);
-		if (any(Radiance > 0))
-		{
-			SkyLight_PdfLight(Ray, LightId, LightUV, Pdf);
-		}
-		
-
-		return float4(Radiance, Pdf);
-	}
-	return 0.0;
-}
-
-void SkyLight_SampleLight(
-	int LightId,
-	float4 RandSample,
-	float3 WorldPos,
-	float3 WorldNormal,
-	out float3 OutDirection,
-	out float  OutDistance,
-	out float3 OutRadianceOverPdf,
-	out float OutPdf)
-{
-	// Make temp structures so we can re-use the calls above
-	// #dxr_todo: write optimized form that combines radiance and pdf lookup in a single step
-	RayDesc Ray;
-	Ray.Origin = WorldPos;
-	FMaterialClosestHitPayload Payload;
-	Payload.WorldNormal = WorldNormal;
 
+	float3 Radiance = SkyLight.Color * SkyLight.Texture.SampleLevel(SkyLight.TextureSampler, OutDirection, 0).rgb;
 
-	float3 LightUV;
-	float3 Radiance;
-	SkyLight_SampleLight(Ray, Payload, RandSample, LightId, LightUV, OutPdf);
-	SkyLight_EvalLight(LightId, LightUV, Ray, Radiance);
-	
-	OutDirection = LightUV;
 	OutDistance = 1.0e27;
 	OutRadianceOverPdf = Radiance / OutPdf;
 }
diff --git a/Engine/Shaders/Private/PathTracing/Light/PathTracingSpotLight.ush b/Engine/Shaders/Private/PathTracing/Light/PathTracingSpotLight.ush
--- a/Engine/Shaders/Private/PathTracing/Light/PathTracingSpotLight.ush
+++ b/Engine/Shaders/Private/PathTracing/Light/PathTracingSpotLight.ush
@@ -141,16 +141,6 @@ void SpotLight_GenerateLightRay(
 	OutLightRay.TMax = RayLength * (CosTheta - sqrt(max(SinThetaMax2 - SinTheta2, 0.0)));
 }
 
-void SpotLight_PdfLight(
-	RayDesc Ray,
-	int LightId,
-	float3 LightUV,
-	out float OutPdf
-)
-{
-	OutPdf = 0.0;
-}
-
 void SpotLight_EvalLight(
 	int LightId,
 	float3 LightUV,
diff --git a/Engine/Shaders/Private/RayTracing/RayTracingCreateGatherPointsRGS.usf b/Engine/Shaders/Private/RayTracing/RayTracingCreateGatherPointsRGS.usf
--- a/Engine/Shaders/Private/RayTracing/RayTracingCreateGatherPointsRGS.usf
+++ b/Engine/Shaders/Private/RayTracing/RayTracingCreateGatherPointsRGS.usf
@@ -9,6 +9,8 @@
 #define DIM_DEFERRED_MATERIAL_MODE DEFERRED_MATERIAL_MODE_NONE
 #endif
 
+#define USE_IES_TERM 0 // this shader doesn't bind the IES atlas, so disable the IES term from path tracer lights
+
 #include "../Common.ush"
 #include "../RectLight.ush"
 #include "../DeferredShadingCommon.ush"
@@ -127,7 +129,7 @@ bool ShouldTraceRay(uint2 DispatchThreadId, FGBufferData GBufferData)
 
 RAY_TRACING_ENTRY_RAYGEN(RayTracingCreateGatherPointsTraceRGS)
 {
-	uint2 DispatchThreadId = DispatchRaysIndex().xy + View.ViewRectMin + uint2(RenderTileOffsetX, RenderTileOffsetY);
+	uint2 DispatchThreadId = DispatchRaysIndex().xy + uint2(View.ViewRectMin.xy) + uint2(RenderTileOffsetX, RenderTileOffsetY);
 
 	FDeferredMaterialPayload DeferredMaterialPayload;
 	DeferredMaterialPayload.SortKey = RAY_TRACING_DEFERRED_MATERIAL_KEY_INVALID;
@@ -238,7 +240,7 @@ RAY_TRACING_ENTRY_RAYGEN(RayTracingCreateGatherPointsRGS)
 	}
 	else
 	{
-		DispatchThreadId = DispatchRaysIndex().xy + View.ViewRectMin + uint2(RenderTileOffsetX, RenderTileOffsetY);
+		DispatchThreadId = DispatchRaysIndex().xy + uint2(View.ViewRectMin.xy) + uint2(RenderTileOffsetX, RenderTileOffsetY);
 	}
 
 	uint GatherPointsIndex = DispatchThreadIdToLinearIndex(DispatchThreadId);
@@ -379,17 +381,25 @@ RAY_TRACING_ENTRY_RAYGEN(RayTracingCreateGatherPointsRGS)
 			{
 				// Light selection
 				int LightId;
-				float3 LightUV;
+				float3 LightRadianceOverPdf = 0;
 				float NeePdf = 0.0;
 
 				float4 RandSample4 = RandomSequence_GenerateSample4D(RandSequence);
 
-				SampleLight(Ray, Payload, RandSample4, LightId, LightUV, NeePdf);
-
-				if (NeePdf > 0.0)
+				RayDesc LightRay;
+				if (SampleLight(RandSample4,
+								Position,
+								Normal,
+								Payload.PrimitiveLightingChannelMask,
+								ENABLE_TRANSMISSION,
+								LightId,
+								LightRay.Direction,
+								LightRay.TMax,
+								LightRadianceOverPdf,
+								NeePdf))
 				{
-					RayDesc LightRay;
-					GenerateLightRay(Ray, LightId, LightUV, LightRay);
+					LightRay.Origin = Position;
+					LightRay.TMin = 0;
 					LightRay.TMax = min(LightRay.TMax, MaxShadowDistance);
 					bool bTransmission = HasTransmission(LightId);
 					float3 BiasNormal = Payload.WorldNormal;
@@ -417,10 +427,6 @@ RAY_TRACING_ENTRY_RAYGEN(RayTracingCreateGatherPointsRGS)
 					// No hit indicates successful next-event connection
 					if (NeePayload.IsMiss())
 					{
-						// Evaluate radiance
-						float3 Radiance;
-						EvalLight(LightId, LightUV, LightRay, Radiance);
-
 						// Evaluate material
 						half3 N = Payload.WorldNormal;
 						half3 V = -Ray.Direction;
@@ -435,7 +441,7 @@ RAY_TRACING_ENTRY_RAYGEN(RayTracingCreateGatherPointsRGS)
 						}
 
 						// Record the contribution
-						float3 ExitantRadianceSample = Radiance * MaterialThroughput * SplitFactor / NeePdf;
+						float3 ExitantRadianceSample = LightRadianceOverPdf * MaterialThroughput * SplitFactor;
 						Irradiance += isfinite(ExitantRadianceSample) ? ExitantRadianceSample : 0.0;
 					}
 				}
diff --git a/Engine/Shaders/Private/RayTracing/RayTracingGatherPoints.ush b/Engine/Shaders/Private/RayTracing/RayTracingGatherPoints.ush
--- a/Engine/Shaders/Private/RayTracing/RayTracingGatherPoints.ush
+++ b/Engine/Shaders/Private/RayTracing/RayTracingGatherPoints.ush
@@ -92,7 +92,7 @@ void CreateRejectionMask(
 		bShouldReject[Index] = dot(GPDistance2, GPDistance2) > RejectionCriteria.DistanceThreshold2;
 	}
 
-	for (int Index = SampleCount; Index < MAXIMUM_GATHER_POINTS_PER_PIXEL; ++Index)
+	for (Index = SampleCount; Index < MAXIMUM_GATHER_POINTS_PER_PIXEL; ++Index)
 	{
 		bShouldReject[Index] = true;
 	}
diff --git a/Engine/Shaders/Private/RayTracing/RayTracingGlobalIlluminationRGS.usf b/Engine/Shaders/Private/RayTracing/RayTracingGlobalIlluminationRGS.usf
--- a/Engine/Shaders/Private/RayTracing/RayTracingGlobalIlluminationRGS.usf
+++ b/Engine/Shaders/Private/RayTracing/RayTracingGlobalIlluminationRGS.usf
@@ -1,5 +1,7 @@
 // Copyright Epic Games, Inc. All Rights Reserved.
 
+#define USE_IES_TERM 0 // this shader doesn't bind the IES atlas, so disable the IES term from path tracer lights
+
 #include "../Common.ush"
 #include "../RectLight.ush"
 //#include "../MonteCarlo.ush"
@@ -60,7 +62,7 @@ void GenerateCosineNormalRay(
 
 RAY_TRACING_ENTRY_RAYGEN(GlobalIlluminationRGS)
 {
-	uint2 DispatchThreadId = DispatchRaysIndex().xy + View.ViewRectMin + uint2(RenderTileOffsetX, RenderTileOffsetY);
+	uint2 DispatchThreadId = DispatchRaysIndex().xy + uint2(View.ViewRectMin.xy) + uint2(RenderTileOffsetX, RenderTileOffsetY);
 	RWGlobalIlluminationUAV[DispatchThreadId] = 0.0;
 	RWGlobalIlluminationRayDistanceUAV[DispatchThreadId] = float2(-1.0, 0.0);
 
@@ -180,8 +182,7 @@ RAY_TRACING_ENTRY_RAYGEN(GlobalIlluminationRGS)
 				if (EvalSkyLight && Bounce > 0)
 				{
 					uint SkyLightId = 0;
-					float3 EnvironmentRadiance = 0.0;
-					SkyLight_EvalLight(SkyLightId, Ray.Direction, Ray, EnvironmentRadiance);
+					float3 EnvironmentRadiance = SkyLight_TraceLight(Ray, SkyLightId, Ray.TMax).xyz;
 					Irradiance += EnvironmentRadiance * RayThroughput / RayPdf;
 				}
 				break;
@@ -223,17 +224,25 @@ RAY_TRACING_ENTRY_RAYGEN(GlobalIlluminationRGS)
 			{
 				// Light selection
 				int LightId;
-				float3 LightUV;
+				float3 LightRadianceOverPdf = 0;
 				float NeePdf = 0.0;
 
 				float4 RandSample4 = RandomSequence_GenerateSample4D(RandSequence);
 
-				SampleLight(Ray, Payload, RandSample4, LightId, LightUV, NeePdf);
-
-				if (NeePdf > 0.0)
+				RayDesc LightRay;
+				if (SampleLight(RandSample4,
+								Ray.Origin,
+								Payload.WorldNormal,
+								Payload.PrimitiveLightingChannelMask,
+								ENABLE_TRANSMISSION,
+								LightId,
+								LightRay.Direction,
+								LightRay.TMax,
+								LightRadianceOverPdf,
+								NeePdf))
 				{
-					RayDesc LightRay;
-					GenerateLightRay(Ray, LightId, LightUV, LightRay);
+					LightRay.Origin = Ray.Origin;
+					LightRay.TMin = 0;
 					LightRay.TMax = min(LightRay.TMax, MaxShadowDistance);
 					bool bTransmission = HasTransmission(LightId);
 					float3 BiasNormal = Payload.WorldNormal;
@@ -261,10 +270,6 @@ RAY_TRACING_ENTRY_RAYGEN(GlobalIlluminationRGS)
 					// No hit indicates successful next-event connection
 					if (NeePayload.IsMiss())
 					{
-						// Evaluate radiance
-						float3 Radiance;
-						EvalLight(LightId, LightUV, LightRay, Radiance);
-
 						// Evaluate material
 						float3 MaterialThroughput;
 						float MaterialEvalPdf = 0.0;
@@ -286,16 +291,9 @@ RAY_TRACING_ENTRY_RAYGEN(GlobalIlluminationRGS)
 						}
 						MaterialEvalPdf = 1.0;
 #endif
-						// Apply material Pdf for correct MIS weight
-						float MisWeight = 1.0;
-#if 0
-						if (IsUnidirectionalEnabled && IsPhysicalLight(LightId))
-						{
-							MisWeight = NeePdf / (NeePdf + MaterialEvalPdf);
-						}
-#endif
+
 						// Record the contribution
-						float3 ExitantRadianceSample = Radiance * MaterialThroughput * RayThroughput * SplitFactor * MisWeight / (NeePdf * RayPdf);
+						float3 ExitantRadianceSample = LightRadianceOverPdf * MaterialThroughput * RayThroughput * SplitFactor / RayPdf;
 
 						// Tonemap for firefly suppression
 						if (UseFireflySuppression)
