Fix compile error in raytracing shader

I was trying to fix a warning of PC platforms, but some consoles seesm to handle scoping differently

diff --git a/Engine/Shaders/Private/RayTracing/RayTracingGatherPoints.ush b/Engine/Shaders/Private/RayTracing/RayTracingGatherPoints.ush
--- a/Engine/Shaders/Private/RayTracing/RayTracingGatherPoints.ush
+++ b/Engine/Shaders/Private/RayTracing/RayTracingGatherPoints.ush
@@ -92,7 +92,7 @@ void CreateRejectionMask(
 		bShouldReject[Index] = dot(GPDistance2, GPDistance2) > RejectionCriteria.DistanceThreshold2;
 	}
 
-	for (Index = SampleCount; Index < MAXIMUM_GATHER_POINTS_PER_PIXEL; ++Index)
+	for (int Index = SampleCount; Index < MAXIMUM_GATHER_POINTS_PER_PIXEL; ++Index)
 	{
 		bShouldReject[Index] = true;
 	}
