Initial shadow map ray tracing support for directional lights
- Disabled by defaul to start. Enable with r.Shadow.v.SMRT.RayCountDirectional 1
- Denoiser will now use temporal reprojection when SMRT is enabled (disable with r.Shadow.v.Denoiser.Temporal 0)
- Variety of fixes and refactors to existing VSM and projection code
- Still a few edge cases that will be improved in future around screen space trace and similar, but working fairly robustly for reasonable light source size settings

diff --git a/Engine/Shaders/Private/VirtualShadowMaps/CacheManagement.usf b/Engine/Shaders/Private/VirtualShadowMaps/CacheManagement.usf
--- a/Engine/Shaders/Private/VirtualShadowMaps/CacheManagement.usf
+++ b/Engine/Shaders/Private/VirtualShadowMaps/CacheManagement.usf
@@ -51,8 +51,15 @@ void VirtualSmInvalidateInstancePagesCS(
 
 			float4x4 LocalToTranslatedWorld = InstanceData.LocalToWorld;
 			LocalToTranslatedWorld[3].xyz += ProjectionData.ShadowPreViewTranslation;
-			// TODO: put mat mul in pre-computed
-			float4x4 LocalToClip = mul(LocalToTranslatedWorld, mul(ProjectionData.TranslatedWorldToShadowViewMatrix, ProjectionData.ShadowViewToClipMatrix));
+			
+			// Go back to clip space
+			float4x4 UVToClip;
+			UVToClip[0] = float4( 2,  0, 0, 0);
+			UVToClip[1] = float4( 0, -2, 0, 0);
+			UVToClip[2] = float4( 0,  0, 1, 0);
+			UVToClip[3] = float4(-1,  1, 0, 1);
+
+			float4x4 LocalToClip = mul(LocalToTranslatedWorld, mul(ProjectionData.TranslatedWorldToShadowUVMatrix, UVToClip));
 
 			FFrustumCullData Cull = BoxCullFrustum(InstanceData.LocalBoundsCenter, InstanceData.LocalBoundsExtent, LocalToClip, !bDirectionalLight, false);
 
diff --git a/Engine/Shaders/Private/VirtualShadowMaps/PageAccessCommon.ush b/Engine/Shaders/Private/VirtualShadowMaps/PageAccessCommon.ush
--- a/Engine/Shaders/Private/VirtualShadowMaps/PageAccessCommon.ush
+++ b/Engine/Shaders/Private/VirtualShadowMaps/PageAccessCommon.ush
@@ -162,12 +162,7 @@ bool VirtualToPhysicalTexel(uint ShadowMapID, uint Level, uint2 VirtualTexelAddr
 
 	uint2 PhysicalPageEntry = PageTable[CalcPageTableLevelOffset(ShadowMapID, Level) + CalcPageOffsetInLevel(Level, uint2(VPageX, VPageY))];
 	PhysicalTexelAddress = PhysicalPageEntry * VSM_PAGE_SIZE + (VirtualTexelAddress & VSM_PAGE_SIZE_MASK);
-	if (PhysicalPageEntry.x == VSM_INVALID_PHYSICAL_PAGE_ADDRESS)
-	{
-		return false;
-	}
-	
-	return true;
+	return (PhysicalPageEntry.x != VSM_INVALID_PHYSICAL_PAGE_ADDRESS);
 }
 
 // TODO: this surely needs some special handling of cases that contain the near plane etc.
diff --git a/Engine/Shaders/Private/VirtualShadowMaps/PageManagement.usf b/Engine/Shaders/Private/VirtualShadowMaps/PageManagement.usf
--- a/Engine/Shaders/Private/VirtualShadowMaps/PageManagement.usf
+++ b/Engine/Shaders/Private/VirtualShadowMaps/PageManagement.usf
@@ -198,7 +198,6 @@ StructuredBuffer<int> VirtualShadowMapIdRemap;
 uint NumDirectionalLightSmInds;
 uint bPostBasePass;
 float LodFootprintScale;
-uint LodPixelCountThreshold;
 
 // Convenience to group up a bunch of data computed for a screen pixel
 struct FScreenPixelData
@@ -289,10 +288,10 @@ uint GetMipLevel(int VirtualShadowMapId, FScreenPixelData Pixel)
 
 		float RadiusWorld = DepthScale * RadiusScreen;
 
-		// NOTE: Virtual shadow maps are currently always the same size in each axis (VSM_VIRTUAL_MAX_RESOLUTION_XY), so ShadowProj._m00 == ShadowProj._m11
-		float3 ShadowViewTranslatedWorldPosition = Pixel.WorldPosition + ShadowMapProjectionData[VirtualShadowMapId].ShadowPreViewTranslation;
-		float4 ShadowViewPosition = mul(float4(ShadowViewTranslatedWorldPosition, 1.0f), ShadowMapProjectionData[VirtualShadowMapId].TranslatedWorldToShadowViewMatrix);
-		float4 RadiusClipH = mul(float4(RadiusWorld, 0.0f, ShadowViewPosition.z, 1.0f), ShadowMapProjectionData[VirtualShadowMapId].ShadowViewToClipMatrix);
+		float3 TranslatedWorldPosition = Pixel.WorldPosition + ShadowMapProjectionData[VirtualShadowMapId].ShadowPreViewTranslation;
+		float4 ShadowUVz = mul(float4(TranslatedWorldPosition, 1.0f), ShadowMapProjectionData[VirtualShadowMapId].TranslatedWorldToShadowUVMatrix);
+
+		float4 RadiusClipH = mul(float4(RadiusWorld, 0.0f, ShadowUVz.w, 1.0f), ShadowMapProjectionData[VirtualShadowMapId].ShadowViewToClipMatrix);
 		float RadiusClip = abs(RadiusClipH.x / RadiusClipH.w);
 		Footprint = RadiusClip * float(2 * VSM_VIRTUAL_MAX_RESOLUTION_XY);
 	}
@@ -308,60 +307,22 @@ uint CalcPageOffsetFromUVMip(uint VirtualShadowMapId, float2 ShadowUv, uint MipL
 
 void MarkPage(uint VirtualShadowMapId, uint MipLevel, FScreenPixelData Pixel)
 {
-	FShadowPosition ShadowPosition = WorldToVirtualShadowMap(ShadowMapProjectionData[VirtualShadowMapId], Pixel.WorldPosition);
+	FVirtualShadowMapProjectionShaderData ProjectionData = ShadowMapProjectionData[VirtualShadowMapId];
+	float4 ShadowUVz = mul(float4(Pixel.WorldPosition + ProjectionData.ShadowPreViewTranslation, 1.0f), ProjectionData.TranslatedWorldToShadowUVMatrix);
+	ShadowUVz.xyz /= ShadowUVz.w;
 
 	// Check overlap vs the shadow map space
 	// NOTE: XY test not really needed anymore with the precise cone test in the caller, but we'll leave it for the moment
-	bool bInClip = ShadowPosition.ClipH.w > 0.0f && 
-		all(ShadowPosition.ClipH.xyz <= ShadowPosition.ClipH.w &&
-			ShadowPosition.ClipH.xyz >= float3(-ShadowPosition.ClipH.ww, 0.0f));
+	bool bInClip = ShadowUVz.w > 0.0f && 
+		all(ShadowUVz.xyz <= ShadowUVz.w &&
+			ShadowUVz.xyz >= float3(-ShadowUVz.ww, 0.0f));
 	if (!bInClip)
 	{
 		return;
 	}
 
-	// 2.1.4. Set covered flags to 1.
-	uint PageOffset = CalcPageOffsetFromUVMip(VirtualShadowMapId, ShadowPosition.UV, MipLevel);
-	if (LodPixelCountThreshold > 0 && MipLevel > 0)
-	{
-#if COMPILER_SUPPORTS_WAVE_VOTE
-		// Reduce any lanes that want to scatter to the same lane as the first across wave and issue single atomic
-		// Any lanes that don't match the first issue separate atomics
-		// This gets the vast majority of the benefits since the pages picked by a pixel are highly coherent
-		uint FirstLanePageOffset = WaveReadLaneFirst(PageOffset);
-		bool MatchesFirstLane = (PageOffset == FirstLanePageOffset);
-		uint MatchesFirstLaneCount = WaveActiveCountBits(MatchesFirstLane);
-
-		if (WaveIsFirstLane() || !MatchesFirstLane)
-		{
-			uint PixelCountToAdd = WaveIsFirstLane() ? MatchesFirstLaneCount : 1;
-			uint PixelCountInPage = 0;
-			InterlockedAdd(OutPageRequestFlags[PageOffset], PixelCountToAdd, PixelCountInPage);
-
-			if (PixelCountInPage > LodPixelCountThreshold)
-			{
-				// NOTE: This double-counts these pixels but for our current heuristic that is fine
-				MipLevel = MipLevel - 1;
-				PageOffset = CalcPageOffsetFromUVMip(VirtualShadowMapId, ShadowPosition.UV, MipLevel);
-				InterlockedAdd(OutPageRequestFlags[PageOffset], PixelCountToAdd, PixelCountInPage);
-			}
-		}
-#else
-		uint PixelCountInPage = 0;
-		InterlockedAdd(OutPageRequestFlags[PageOffset], 1, PixelCountInPage);
-
-		if (PixelCountInPage > LodPixelCountThreshold)
-		{
-			MipLevel = MipLevel - 1;
-			PageOffset = CalcPageOffsetFromUVMip(VirtualShadowMapId, ShadowPosition.UV, MipLevel);
-			InterlockedAdd(OutPageRequestFlags[PageOffset], 1, PixelCountInPage);
-		}
-#endif
-	}
-	else
-	{
-		OutPageRequestFlags[PageOffset] = 1;
-	}
+	uint PageOffset = CalcPageOffsetFromUVMip(VirtualShadowMapId, ShadowUVz.xy, MipLevel);
+	OutPageRequestFlags[PageOffset] = 1;
 }
 
 [numthreads(VSM_DEFAULT_CS_GROUP_XY, VSM_DEFAULT_CS_GROUP_XY, 1)]
@@ -395,21 +356,19 @@ void GeneratePageFlagsFromPixels(uint3 GroupId : SV_GroupID, uint GroupIndex : S
 		while (ProjectionIndex < NumDirectionalLightSmInds)
 		{
 			int ClipmapStartId = VirtualShadowMapIdRemap[ProjectionIndex];
-
-			const int FirstClipmapLevel          = ShadowMapProjectionData[ClipmapStartId].ClipmapLevel;
-			const int ClipmapLevelCount          = ShadowMapProjectionData[ClipmapStartId].ClipmapLevelCount;
-			const float ClipmapResolutionLodBias = ShadowMapProjectionData[ClipmapStartId].ClipmapResolutionLodBias;
+			FVirtualShadowMapProjectionShaderData BaseProjectionData = ShadowMapProjectionData[ClipmapStartId];
 
 			// Possibly use a single shadow projection data for the whole clipmap with just scales/biases for
 			// each level (to handle snapping). For now we maintain one per level.
-			int ClipmapLevel = CalcClipmapLevel(length(Pixel.ViewPosition), ClipmapResolutionLodBias);
-			int ClipmapIndex = max(0, ClipmapLevel - FirstClipmapLevel);
-			if (ClipmapIndex < ClipmapLevelCount)
+			//int ClipmapLevel = CalcClipmapLevel(length(Pixel.WorldPosition - ProjectionData.ClipmapWorldOrigin), ClipmapResolutionLodBias);
+			const int ClipmapLevel = CalcClipmapLevel(BaseProjectionData, Pixel.WorldPosition);
+			int ClipmapIndex = max(0, ClipmapLevel - BaseProjectionData.ClipmapLevel);
+			if (ClipmapIndex < BaseProjectionData.ClipmapLevelCount)
 			{
 				MarkPage(ClipmapStartId + ClipmapIndex, 0, Pixel);
 			}
 
-			ProjectionIndex += max(1, ClipmapLevelCount);		// Max just to ensure termination even with bad data
+			ProjectionIndex += max(1, BaseProjectionData.ClipmapLevelCount);		// Max just to ensure termination even with bad data
 		}
 	}
 
diff --git a/Engine/Shaders/Private/VirtualShadowMaps/ProjectionCommon.ush b/Engine/Shaders/Private/VirtualShadowMaps/ProjectionCommon.ush
--- a/Engine/Shaders/Private/VirtualShadowMaps/ProjectionCommon.ush
+++ b/Engine/Shaders/Private/VirtualShadowMaps/ProjectionCommon.ush
@@ -14,15 +14,14 @@ ProjectionCommon.ush:
 // Must match C++ structure in VirtualShadowMapArray.h
 struct FVirtualShadowMapProjectionShaderData
 {
-	// From shadow-pretranslated world-space to shadow view space, example use: (WorldSpacePos + ShadowPreViewTranslation) * TranslatedWorldToShadowViewMatrix
-	float4x4 TranslatedWorldToShadowViewMatrix;
 	float4x4 ShadowViewToClipMatrix;
-	float4x4 TranslatedWorldToShadowUvNormalMatrix;
+	float4x4 TranslatedWorldToShadowUVMatrix;
+	float4x4 TranslatedWorldToShadowUVNormalMatrix;
 	float3 ShadowPreViewTranslation;
 	uint LightType;		// Matches ELightComponentType via defines in SSDDefinitions.ush
 
 	float3 ClipmapWorldOrigin;
-	int VirtualShadowMapId;	
+	int VirtualShadowMapId;
 
 	int ClipmapLevel;
 	int ClipmapLevelCount;
@@ -48,58 +47,60 @@ uint CalcMipLevelFromFootprint(float Footprint, uint MipLevelCount = VSM_MAX_MIP
 	return MipLevel;
 }
 
-int CalcClipmapLevel(float DistanceToClipmapOrigin, float ResolutionLodBias)
+float CalcClipmapLevelFloat(FVirtualShadowMapProjectionShaderData BaseProjectionData, float3 WorldPosition)
 {
+	float DistanceToClipmapOrigin = length(WorldPosition - BaseProjectionData.ClipmapWorldOrigin);
 	float Log2Distance = log2(DistanceToClipmapOrigin);
-	int AbsoluteLevel = int(floor(Log2Distance + ResolutionLodBias));
+	float AbsoluteLevel = Log2Distance + BaseProjectionData.ClipmapResolutionLodBias;
 	return AbsoluteLevel;
 }
+int CalcClipmapLevel(FVirtualShadowMapProjectionShaderData BaseProjectionData, float3 WorldPosition)
+{
+	return int(floor(CalcClipmapLevelFloat(BaseProjectionData, WorldPosition)));
+}
 
 struct FVirtualSmSample
 {
 	float Depth;
+	uint MipLevel;
+	bool bValid;
 	uint2 VirtualTexelAddress;
 	float2 VirtualTexelAddressFloat;
 	uint2 PhysicalTexelAddress;
-	uint MipLevel;
-	bool bValid;
 };
 
+FVirtualSmSample InitVirtualSmSample()
+{
+	FVirtualSmSample Result;
+	Result.Depth = 0.0f;
+	Result.MipLevel = 0;
+	Result.bValid = false;
+	return Result;
+}
+
 FVirtualSmSample SampleVirtualSmLevel(int VirtualShadowMapId, float2 ShadowMapUV, uint VSMLevel)
 {
-	float2 VirtualTexelAddressFloat = ShadowMapUV * float(CalcLevelDimsTexels(VSMLevel));
-	uint2 VirtualTexelAddress = uint2(VirtualTexelAddressFloat);
+	FVirtualSmSample Result = InitVirtualSmSample();
+
+	Result.VirtualTexelAddressFloat = ShadowMapUV * float(CalcLevelDimsTexels(VSMLevel));
+	Result.VirtualTexelAddress = uint2(Result.VirtualTexelAddressFloat);
 
 	// Page map to physical address
 	uint2 PhysicalTexelAddress = uint2(0, 0);
-	bool PagePresent = VirtualToPhysicalTexel(VirtualShadowMapId, VSMLevel, VirtualTexelAddress, PhysicalTexelAddress);
+	bool PagePresent = VirtualToPhysicalTexel(VirtualShadowMapId, VSMLevel, Result.VirtualTexelAddress, Result.PhysicalTexelAddress);
 	if (PagePresent)
 	{
 		// TODO: maybe store metadata about Max depth (and/or use sparse HZB) for each tile, and early out if completely in shadow (must test using kernel size).
-		FVirtualSmSample Result;
-		Result.Depth = asfloat(PhysicalPagePool.Load(int3(PhysicalTexelAddress, 0)));
-		Result.VirtualTexelAddress = VirtualTexelAddress;
-		Result.PhysicalTexelAddress = PhysicalTexelAddress;
-		Result.VirtualTexelAddressFloat = VirtualTexelAddressFloat;
+		Result.Depth = asfloat(PhysicalPagePool.Load(int3(Result.PhysicalTexelAddress, 0)));
 		Result.MipLevel = VSMLevel;
 		Result.bValid = true;
-		return Result;
-	}
-	else
-	{
-		FVirtualSmSample Result;
-		Result.Depth = 0.0f;
-		Result.VirtualTexelAddress = uint2(0, 0);
-		Result.PhysicalTexelAddress = uint2(0, 0);
-		Result.MipLevel = 0;
-		Result.bValid = false;
-		return Result;
 	}
+	
+	return Result;
 }
 
 FVirtualSmSample SampleVirtualSm(int VirtualShadowMapId, float2 ShadowMapUV)
 {
-	// Compute virtual address, map to physical & sample
 	for (uint VSMLevel = 0; VSMLevel < VSM_MAX_MIP_LEVELS; ++VSMLevel)
 	{
 		FVirtualSmSample Result = SampleVirtualSmLevel(VirtualShadowMapId, ShadowMapUV, VSMLevel);
@@ -109,62 +110,36 @@ FVirtualSmSample SampleVirtualSm(int VirtualShadowMapId, float2 ShadowMapUV)
 		}
 	}
 
-	FVirtualSmSample Result;
-	Result.Depth = 0.0f;
-	Result.VirtualTexelAddress = uint2(0, 0);
-	Result.PhysicalTexelAddress = uint2(0, 0);
-	Result.MipLevel = 0;
-	Result.bValid = false;
-	return Result;
+	return InitVirtualSmSample();
 }
 
 float ComputeVirtualShadowMapOptimalSlopeBias(FVirtualShadowMapProjectionShaderData ProjectionData, FVirtualSmSample SmSample, float3 WorldPosition, float3 EstimatedGeoWorldNormal)
 {
 	// Transform geometry world-space plane eq to shadow 'UV' texture space [0-1] ranges
 	float4 NormalPlaneTranslatedWorld = float4(EstimatedGeoWorldNormal, -dot(EstimatedGeoWorldNormal, WorldPosition + ProjectionData.ShadowPreViewTranslation));
-	float4 NormalPlaneUv = mul(NormalPlaneTranslatedWorld, ProjectionData.TranslatedWorldToShadowUvNormalMatrix);
+	float4 NormalPlaneUV = mul(NormalPlaneTranslatedWorld, ProjectionData.TranslatedWorldToShadowUVNormalMatrix);
 
-	float2 DepthSlopeUv = -NormalPlaneUv.xy / NormalPlaneUv.z;
+	float2 DepthSlopeUV = -NormalPlaneUV.xy / NormalPlaneUV.z;
 	float MipLevelDim = float(CalcLevelDimsTexels(SmSample.MipLevel));
 	float2 TexelCenter = float2(SmSample.VirtualTexelAddress) + 0.5;
 	float2 TexelCenterOffset = TexelCenter - SmSample.VirtualTexelAddressFloat;
-	float2 TexelCenterOffsetUv = TexelCenterOffset / MipLevelDim;
+	float2 TexelCenterOffsetUV = TexelCenterOffset / MipLevelDim;
 	// 2x factor due to lack of precision (probably)
-	float OptimalSlopeBias = 2.0f * max(0.0f, dot(DepthSlopeUv, TexelCenterOffsetUv));
+	float OptimalSlopeBias = 2.0f * max(0.0f, dot(DepthSlopeUV, TexelCenterOffsetUV));
 	// Clamp to avoid excessive degenerate slope biases causing flickering lit pixels
-	OptimalSlopeBias = min(OptimalSlopeBias, 0.01f);
+	OptimalSlopeBias = min(OptimalSlopeBias, abs(100.0f * ProjectionData.ShadowViewToClipMatrix._33));
 
 	return OptimalSlopeBias;
 }
 
-// Small wrapper for point transformed into shadow space
-struct FShadowPosition
-{
-	float3 View;		// Shadow view space
-	float4 ClipH;		// NOT divided by W
-	float2 UV;
-	float Depth;
-};
-
-FShadowPosition WorldToVirtualShadowMap(FVirtualShadowMapProjectionShaderData ProjectionData, float3 WorldPosition)
-{
-	FShadowPosition ShadowPosition;
-	ShadowPosition.View = mul(float4(WorldPosition + ProjectionData.ShadowPreViewTranslation, 1.0f), ProjectionData.TranslatedWorldToShadowViewMatrix).xyz;
-
-	ShadowPosition.ClipH = mul(float4(ShadowPosition.View, 1.0f), ProjectionData.ShadowViewToClipMatrix);
-
-	float3 Clip = ShadowPosition.ClipH.xyz / ShadowPosition.ClipH.w;
-	ShadowPosition.UV = float2(0.5f, -0.5f) * Clip.xy + float2(0.5f, 0.5f);
-	ShadowPosition.Depth = Clip.z;
-
-	return ShadowPosition;
-}
-
 struct FVirtualShadowMapSampleResult
 {
 	bool bValid;
 	bool bOccluded;
 	float OccluderDistance;
+
+	// Debug data	
+	int ClipmapIndexOrMipLevel;
 };
 
 // Used for orthographic projections (i.e. directional lights)
@@ -181,13 +156,12 @@ float ComputeOccluderDistanceOrtho(float4x4 ShadowViewToClip, float OccluderDept
 
 // Used for perspective projections (i.e. spot lights)
 // Receiver depth is post-projection-divide.
-float ComputeOccluderDistancePerspective(float4x4 ShadowViewToClip, float OccluderDepth, float ReceiverDepth, float3 ReceiverShadowViewPosition)
+float ComputeOccluderDistancePerspective(float4x4 ShadowViewToClip, float OccluderDepth, float ReceiverDepth, float ReceiverDistance)
 {
 	float OccluderViewZ = ShadowViewToClip._43 / (OccluderDepth - ShadowViewToClip._33);
 	float ReceiverViewZ = ShadowViewToClip._43 / (ReceiverDepth - ShadowViewToClip._33);
 
 	// Similar triangles to compute euclidean distance in view/world space
-	float ReceiverDistance = length(ReceiverShadowViewPosition);
 	float OccluderDistance = (ReceiverDistance / ReceiverViewZ) * OccluderViewZ;
 	float Result = ReceiverDistance - OccluderDistance;	
 	return max(VIRTUAL_SHADOW_MAP_MIN_OCCLUDER_DISTANCE, Result);
@@ -206,25 +180,14 @@ float ComputeOccluderDistancePerspective(float4x4 ShadowViewToClip, float Occlud
 */
 FVirtualShadowMapSampleResult SampleVirtualShadowMap(int VirtualShadowMapId, float3 WorldPosition, float3 EstimatedGeoWorldNormal, float RayStartDistance = 0.0f)
 {
-	RayStartDistance= max(RayStartDistance, 0.0f);
+	RayStartDistance = max(RayStartDistance, 0.0f);
 
 	FVirtualShadowMapProjectionShaderData BaseProjectionData = GetVirtualShadowMapProjectionData(VirtualShadowMapId);
 
-	FVirtualSmSample SmSample;
-	SmSample.Depth = 0.0f;
-	SmSample.VirtualTexelAddress = uint2(0, 0);
-	SmSample.PhysicalTexelAddress = uint2(0, 0);
-	SmSample.MipLevel = 0;
-	SmSample.bValid = false;
-
-	float OccluderDistance = -1.0f;
-	bool bOccluded = false;
-
 	if (BaseProjectionData.LightType == LIGHT_TYPE_DIRECTIONAL)
 	{
 		const int FirstClipmapLevel = BaseProjectionData.ClipmapLevel;
-		float DistanceToClipmapOrigin = length(WorldPosition - BaseProjectionData.ClipmapWorldOrigin);
-		int ClipmapLevel = CalcClipmapLevel(DistanceToClipmapOrigin, BaseProjectionData.ClipmapResolutionLodBias);
+		const int ClipmapLevel = CalcClipmapLevel(BaseProjectionData, WorldPosition);
 		int ClipmapIndex = max(0, ClipmapLevel - FirstClipmapLevel);
 
 		// Check if sample is within the clipmap range (from camera)
@@ -234,50 +197,92 @@ FVirtualShadowMapSampleResult SampleVirtualShadowMap(int VirtualShadowMapId, flo
 			int ClipmapLevelVirtualShadowMapId = VirtualShadowMapId + ClipmapIndex;
 			FVirtualShadowMapProjectionShaderData ProjectionData = GetVirtualShadowMapProjectionData(ClipmapLevelVirtualShadowMapId);
 
-			FShadowPosition Position = WorldToVirtualShadowMap(ProjectionData, WorldPosition);
-			SmSample = SampleVirtualSmLevel(ClipmapLevelVirtualShadowMapId, Position.UV, 0);
+			// No perspective divided needed for ortho projection
+			float4 ShadowUVz = mul(float4(WorldPosition + ProjectionData.ShadowPreViewTranslation, 1.0f), ProjectionData.TranslatedWorldToShadowUVMatrix);
+
+			FVirtualSmSample SmSample = SampleVirtualSmLevel(ClipmapLevelVirtualShadowMapId, ShadowUVz.xy, 0);
 			if (SmSample.bValid)
 			{
+				FVirtualShadowMapSampleResult Result;
+				Result.bValid = true;
+				Result.bOccluded = false;
+				Result.OccluderDistance = -1.0f;
+				Result.ClipmapIndexOrMipLevel = ClipmapIndex;
+
 				float OptimalSlopeBias = ComputeVirtualShadowMapOptimalSlopeBias(ProjectionData, SmSample, WorldPosition, EstimatedGeoWorldNormal);
 				float RayStartBias = -RayStartDistance * ProjectionData.ShadowViewToClipMatrix._33;
 				float BiasedDepth = SmSample.Depth - OptimalSlopeBias - RayStartBias;
 
-				if (BiasedDepth > Position.Depth)
+				if (BiasedDepth > ShadowUVz.z)
 				{
-					bOccluded = true;
-					OccluderDistance = ComputeOccluderDistanceOrtho(
+					Result.bOccluded = true;
+					Result.OccluderDistance = ComputeOccluderDistanceOrtho(
 						ProjectionData.ShadowViewToClipMatrix,
 						SmSample.Depth,
-						Position.Depth);
+						ShadowUVz.z);
 				}
+
+				return Result;
 			}
 		}
 	}
 	else	// Spot light
 	{
-		FShadowPosition Position = WorldToVirtualShadowMap(BaseProjectionData, WorldPosition);
-		SmSample = SampleVirtualSm(VirtualShadowMapId, Position.UV);
+		float4 ShadowUVz = mul(float4(WorldPosition + BaseProjectionData.ShadowPreViewTranslation, 1.0f), BaseProjectionData.TranslatedWorldToShadowUVMatrix);
+		ShadowUVz.xyz /= ShadowUVz.w;
+
+		FVirtualSmSample SmSample = SampleVirtualSm(VirtualShadowMapId, ShadowUVz.xy);
 		if (SmSample.bValid)
 		{
+			FVirtualShadowMapSampleResult Result;
+			Result.bValid = true;
+			Result.bOccluded = false;
+			Result.OccluderDistance = -1.0f;
+			Result.ClipmapIndexOrMipLevel = SmSample.MipLevel;
+
 			float OptimalSlopeBias = ComputeVirtualShadowMapOptimalSlopeBias(BaseProjectionData, SmSample, WorldPosition, EstimatedGeoWorldNormal);
-			float RayStartBias = -RayStartDistance * BaseProjectionData.ShadowViewToClipMatrix._33 / Position.ClipH.w;
+			float RayStartBias = -RayStartDistance * BaseProjectionData.ShadowViewToClipMatrix._33 / ShadowUVz.w;
 			float BiasedDepth = SmSample.Depth - OptimalSlopeBias - RayStartBias;
 
-			if (BiasedDepth > Position.Depth)
+			if (BiasedDepth > ShadowUVz.z)
 			{
-				bOccluded = true;
-				OccluderDistance = RayStartDistance + ComputeOccluderDistancePerspective(
+				Result.bOccluded = true;
+				// Shadow view matrix is rotation
+				float ReceiverDistance = length(WorldPosition + BaseProjectionData.ShadowPreViewTranslation);
+				Result.OccluderDistance = RayStartDistance + ComputeOccluderDistancePerspective(
 					BaseProjectionData.ShadowViewToClipMatrix,
 					SmSample.Depth,
-					Position.Depth,
-					Position.View);
+					ShadowUVz.z,
+					ReceiverDistance);
 			}
+
+			return Result;
 		}
 	}
 
+	// Invalid
 	FVirtualShadowMapSampleResult Result;
-	Result.bValid = SmSample.bValid;
-	Result.bOccluded = bOccluded;
-	Result.OccluderDistance = OccluderDistance;
+	Result.bValid = false;
+	Result.bOccluded = false;
+	Result.OccluderDistance = -1.0f;
+	Result.ClipmapIndexOrMipLevel = -1;
 	return Result;
 }
+
+// Helper for debug mostly
+float3 VirtualShadowMapDebugIntToColor(uint Index)
+{
+	const float3 ColorTable[8] = {
+		float3(1.0f, 0.0f, 0.0f),
+		float3(0.0f, 1.0f, 0.0f),
+		float3(0.0f, 0.0f, 1.0f),
+
+		float3(1.0f, 1.0f, 0.0f),
+		float3(0.0f, 1.0f, 1.0f),
+		float3(1.0f, 0.0f, 1.0f),
+
+		float3(0.5f, 0.7f, 0.1f),
+		float3(1.0f, 1.0f, 1.0f)
+	};
+	return ColorTable[Index % 8];
+}
\ No newline at end of file
diff --git a/Engine/Shaders/Private/VirtualShadowMaps/ProjectionDirectional.ush b/Engine/Shaders/Private/VirtualShadowMaps/ProjectionDirectional.ush
new file mode 100644
--- /dev/null
+++ b/Engine/Shaders/Private/VirtualShadowMaps/ProjectionDirectional.ush
@@ -0,0 +1,251 @@
+// Copyright Epic Games, Inc. All Rights Reserved.
+
+/*=============================================================================
+ProjectionDirectional.ush: 
+=============================================================================*/
+#pragma once
+
+#include "../DeferredShadingCommon.ush"
+#include "../SceneTexturesCommon.ush"
+#include "../LightShaderParameters.ush"
+#include "PageAccessCommon.ush"
+#include "ProjectionCommon.ush"
+
+bool ShadowRayCastDirectional(
+	int VirtualShadowMapId,
+	float3 ViewPosition,
+	float3 RayStart,
+	float3 RayVector,	// RayEnd - RayStart
+	int NumSteps,
+	float StepOffset,
+	int DebugOutputType,
+	inout float3 DebugOutput,
+	inout float OccluderDistance)
+{	
+	const FVirtualShadowMapProjectionShaderData BaseProjectionData = GetVirtualShadowMapProjectionData(VirtualShadowMapId);
+
+	// The page at RayStart has a high chance of being mapped since that is close to the receiver point
+	// Thus we start our clipmap sampling around that page even as we walk through the depth space
+	// TODO: Branch out if beyond max clipmap range somewhere
+	const float ClipmapLevelFloat = CalcClipmapLevelFloat(BaseProjectionData, RayStart);	
+	const int ClipmapIndex = max(0, int(floor(ClipmapLevelFloat)) - BaseProjectionData.ClipmapLevel);
+	const float ClipmapLevelFrac = frac(ClipmapLevelFloat);
+
+	int ClipmapIndexAdjacent = ClipmapIndex + (ClipmapLevelFrac > 0.5f ? 1 : -1);
+	ClipmapIndexAdjacent = clamp(ClipmapIndexAdjacent, 0, BaseProjectionData.ClipmapLevelCount - 1);
+
+	int AbsoluteClipmapIndex = VirtualShadowMapId + ClipmapIndex;
+	int AbsoluteClipmapIndexAdjacent = VirtualShadowMapId + ClipmapIndexAdjacent;
+
+	// TODO: Experiment with perf of this, etc.
+	bool bSampleTwoClipmapLevels = (ClipmapIndexAdjacent != ClipmapIndex) && (abs(ClipmapLevelFrac - 0.5f) > 0.4f);
+#if 0
+	bSampleTwoClipmapLevels = false;
+#endif
+
+	const FVirtualShadowMapProjectionShaderData ProjectionData = GetVirtualShadowMapProjectionData(AbsoluteClipmapIndex);
+
+	const float3 RayStartUV = mul(float4(RayStart + ProjectionData.ShadowPreViewTranslation, 1.0f), ProjectionData.TranslatedWorldToShadowUVMatrix).xyz;
+	const float3 RayStepUV  = mul(float4(RayVector, 0.0f), ProjectionData.TranslatedWorldToShadowUVMatrix).xyz;
+
+	const float LevelDims = float(CalcLevelDimsTexels(0));
+	float3 RayStartShadowTexelsZ = float3(LevelDims, LevelDims, 1.0f) * RayStartUV.xyz;
+	float3 RayStepShadowTexelsZ  = float3(LevelDims, LevelDims, 1.0f) * RayStepUV.xyz;
+
+	// Z scale is consistent across clipmap levels
+	const float CompareScale = RayStepShadowTexelsZ.z * 2.0f / Pow2(NumSteps);
+	const float CompareBias = -CompareScale * (NumSteps - StepOffset + 1.0f);
+
+	float FirstHitTime = 2;
+	float DepthHistory = -1;
+	int SampleMissCount = 0;
+
+	// Work out scale/bias to convert to the adjacent UVs from the regular UVs
+	// NOTE: Since depth function is constant, probably worth refactoring the clipmap mappings to be scale/bias in general rather than
+	// a full transform matrix.
+	float2 AdjacentTexelScale = 1.0f;
+	float2 AdjacentTexelBias = 0.0f;
+	float AdjacentToBaseDepthScale = 1.0f;
+	float AdjacentToBaseDepthBias = 0.0f;
+	if (bSampleTwoClipmapLevels)
+	{
+		const FVirtualShadowMapProjectionShaderData ProjectionDataAdjacent = GetVirtualShadowMapProjectionData(AbsoluteClipmapIndexAdjacent);
+
+		AdjacentTexelScale = (ProjectionDataAdjacent.ShadowViewToClipMatrix._11_22 / ProjectionData.ShadowViewToClipMatrix._11_22);
+
+		// TranslatedWorldToShadowViewMatrix is identical for each of the cascades; only the pre-view translation changes
+		float3 UVz         = mul(float4(ProjectionData.ShadowPreViewTranslation, 1.0f), ProjectionData.TranslatedWorldToShadowUVMatrix).xyz;
+		float3 UVzAdjacent = mul(float4(ProjectionDataAdjacent.ShadowPreViewTranslation, 1.0f), ProjectionDataAdjacent.TranslatedWorldToShadowUVMatrix).xyz;
+
+		AdjacentTexelBias = (UVzAdjacent.xy - AdjacentTexelScale * UVz.xy);
+		AdjacentTexelBias *= float(CalcLevelDimsTexels(0));
+
+		// Always rescale adjacent samples back to the same depth mapping as the regular clipmap level to ensure
+		// precision is maintained and DepthHistory remains well defined.
+		AdjacentToBaseDepthScale = (ProjectionData.ShadowViewToClipMatrix._33 / ProjectionDataAdjacent.ShadowViewToClipMatrix._33);
+		AdjacentToBaseDepthBias = (UVz.z - AdjacentToBaseDepthScale * UVzAdjacent.z);
+	}
+
+	for (int i = 0; i < NumSteps; i++)
+	{
+		const float SampleTime = Pow2(saturate((-1.0f / NumSteps) * i + 1.0f - (StepOffset - 0.5f) / NumSteps));
+		const float3 SampleTexelsZ = RayStartShadowTexelsZ + RayStepShadowTexelsZ * SampleTime;
+
+		uint2 PhysicalTexelAddress = uint2(0, 0);
+		bool bValid = VirtualToPhysicalTexel(AbsoluteClipmapIndex, 0, uint2(SampleTexelsZ.xy), PhysicalTexelAddress);
+		
+		float ToBaseDepthScale = 1.0f;
+		float ToBaseDepthBias = 0.0f;
+		if (bSampleTwoClipmapLevels && !bValid)
+		{
+			// Replace with adjacent level sample
+			float2 AdjacentTexels = SampleTexelsZ.xy * AdjacentTexelScale + AdjacentTexelBias;
+			bValid = VirtualToPhysicalTexel(AbsoluteClipmapIndexAdjacent, 0, uint2(AdjacentTexels), PhysicalTexelAddress);
+			ToBaseDepthScale = AdjacentToBaseDepthScale;
+			ToBaseDepthBias  = AdjacentToBaseDepthBias;
+		}
+
+		if (bValid)
+		{
+			const float SampleDepth = ToBaseDepthScale * asfloat(PhysicalPagePool.Load(int3(PhysicalTexelAddress, 0))) + ToBaseDepthBias;
+			const float ReferenceDepth = SampleTexelsZ.z;
+			if (DepthHistory < 0)
+			{
+				// First valid sample we've seen. Do a regular depth compare.
+				DepthHistory = SampleDepth;
+				FirstHitTime = SampleDepth - ReferenceDepth > 0 ? SampleTime : 2;
+			}
+			else
+			{
+				const float CompareTolerance = abs(CompareScale * i + CompareBias);
+
+				bool bBehind = (SampleDepth - ReferenceDepth) > (2*CompareTolerance);
+				//bBehind = false;
+
+				float DepthSample = bBehind ? DepthHistory : SampleDepth;
+				DepthHistory = DepthSample;
+
+				float DepthDiff = ReferenceDepth - DepthSample;
+				bool bHit = abs(DepthDiff + CompareTolerance) < CompareTolerance;
+
+				if (bHit)
+				{
+					FirstHitTime = SampleTime;
+				}
+			}
+		}
+		else
+		{
+			++SampleMissCount;
+		}
+	}
+
+	bool bValidHit = FirstHitTime < 2;
+	if (bValidHit)
+	{
+		OccluderDistance = length(RayVector) * FirstHitTime;
+	}
+
+	float OutputDebugSignal = bValidHit ? 0.0f : 1.0f;
+	if (DebugOutputType == 1)	// Invalid page
+	{
+		DebugOutput = lerp(OutputDebugSignal.xxx, lerp(float3(0, 1, 0), float3(1, 0, 1), saturate(SampleMissCount / 4.0f)), 0.75f);
+	}
+	else if (DebugOutputType == 2)	// ClipmapIndex
+	{
+		DebugOutput = lerp(OutputDebugSignal.xxx, VirtualShadowMapDebugIntToColor(ClipmapIndex), 0.75f);
+	}
+	else if (DebugOutputType == 3)	// Two levels
+	{
+		DebugOutput = bSampleTwoClipmapLevels ? float3(1, 0, 0) : float3(0, 1, 0);
+	}
+
+	return bValidHit;
+}
+
+#define USE_RANDOM_SEQUENCE 1
+
+#if USE_RANDOM_SEQUENCE
+#include "/Engine/Private/PathTracing/Utilities/PathTracingRandomSequence.ush"
+#endif
+
+float3 GetRandomDirectionalLightRayDir(FLightShaderParameters Light, float2 E)
+{
+	float3 RayDir = Light.Direction;
+	{
+		float2 DiskUV = UniformSampleDiskConcentric(E) * Light.SourceRadius;
+		float3 N = RayDir;
+		float3 dPdu = cross(N, (abs(N.x) > 1e-6f) ? float3(1, 0, 0) : float3(0, 1, 0));
+		float3 dPdv = cross(dPdu, N);
+		RayDir += dPdu * DiskUV.x + dPdv * DiskUV.y;
+	}
+	return normalize(RayDir);
+}
+
+float TraceDirectional(
+	int VirtualShadowMapId,
+	FLightShaderParameters Light,
+	float4 SvPosition,
+	const float SceneDepth,
+	float3 TranslatedWorldPosition,
+	float RayStartOffset,
+	int RayCount,
+	int SamplesPerRay,
+	float RayLengthScale,
+	int DebugOutputType,
+	inout float3 DebugOutput,
+	out float OutHitDistance)
+{
+	float3 WorldPosition = TranslatedWorldPosition - View.PreViewTranslation;
+	float3 ViewPosition = mul(float4(TranslatedWorldPosition, 1.0f), View.TranslatedWorldToView).xyz;
+	float DistanceFromViewOrigin = length(ViewPosition);
+
+	// Compute max ray length based on view depth
+	// This affects both how far we are willing to trace across the shadow map (for clipmaps this is related to view depth)
+	// and the maximum size a penumbra can be.
+	// Too high values will cause shadows to detach from their contact points (unless more samples are used).
+	// Too low values will greatly restrict how large penumbras can be in screen space.
+	float RayLength = RayLengthScale * DistanceFromViewOrigin;		// TODO: Parameterize
+
+#if USE_RANDOM_SEQUENCE
+	uint Seed = uint(SvPosition.y * View.BufferSizeAndInvSize.x + SvPosition.x);
+	RandomSequence RandSequence;
+	RandomSequence_Initialize(RandSequence, Seed, View.StateFrameIndex);
+#else
+	uint2 SobolBase = SobolPixel(SvPosition.xy);
+	uint2 SobolFrame = SobolIndex(SobolBase, View.StateFrameIndexMod8, 3);
+#endif
+
+	float StepOffset = InterleavedGradientNoise(SvPosition.xy, View.StateFrameIndexMod8);
+
+	float HitDistance = 500000;		// TODO
+	float RayMissFactor = 0;
+	for (int i = 0; i < RayCount; ++i)
+	{
+#if USE_RANDOM_SEQUENCE
+		float2 E = RandomSequence_GenerateSample2D(RandSequence);
+#else
+		float2 E = float2(SobolIndex(SobolFrame, i << 3)) / 0x10000;
+#endif
+
+		float3 RayDir = GetRandomDirectionalLightRayDir(Light, E);
+		float3 RayStart = WorldPosition;// + RayDir * ContactShadowLengthWorld;
+
+		bool bHit = ShadowRayCastDirectional(
+			VirtualShadowMapId,
+			ViewPosition,
+			RayStart + RayStartOffset,
+			RayDir * RayLength,
+			SamplesPerRay,
+			StepOffset,
+			DebugOutputType,
+			DebugOutput,
+			HitDistance);
+
+		RayMissFactor += bHit ? 0 : 1;
+	}
+	RayMissFactor /= float(RayCount);
+
+	OutHitDistance = HitDistance;
+	return RayMissFactor;
+}
diff --git a/Engine/Shaders/Private/VirtualShadowMaps/VirtualShadowMapProjection.usf b/Engine/Shaders/Private/VirtualShadowMaps/VirtualShadowMapProjection.usf
--- a/Engine/Shaders/Private/VirtualShadowMaps/VirtualShadowMapProjection.usf
+++ b/Engine/Shaders/Private/VirtualShadowMaps/VirtualShadowMapProjection.usf
@@ -13,6 +13,7 @@
 #include "../LightShaderParameters.ush"
 #include "PageAccessCommon.ush"
 #include "ProjectionCommon.ush"
+#include "ProjectionDirectional.ush"
 
 // See VirtualShadowMapProjection.cpp
 // 0 = Output to denoiser
@@ -26,36 +27,28 @@ int VirtualShadowMapId;		// For the first layer for clipmaps
 int DebugOutputType;
 float ContactShadowLength;
 
-// TODO: Share LOD visualization color logic between Debug.usf and here (and possibly Nanite debug)
-float3 IntToColor(uint Index)
-{
-	const float3 ColorTable[8] = {
-		float3(1.0f, 0.0f, 0.0f),
-		float3(0.0f, 1.0f, 0.0f),
-		float3(0.0f, 0.0f, 1.0f),
-
-		float3(1.0f, 1.0f, 0.0f),
-		float3(0.0f, 1.0f, 1.0f),
-		float3(1.0f, 0.0f, 1.0f),
-
-		float3(0.5f, 0.7f, 0.1f),
-		float3(1.0f, 1.0f, 1.0f)
-	};
-	return ColorTable[Index % 8];
-}
+int SMRTRayCount;	// 0 = off
+int SMRTSamplesPerRay;
+float SMRTRayLengthScale;
+
 
-float ShadowRayCast(
+float VirtualShadowMapScreenRayCast(
 	float3 RayOriginTranslatedWorld, float3 RayDirection, float RayLength,
-	int NumSteps, float StepOffset
-)
+	int NumSteps, float StepOffset)
 {
+	// Just to make calling code a bit cleaner
+	if (RayLength == 0.0f)
+	{
+		return 0.0f;
+	}
+
 	float4 RayStartClip	= mul( float4( RayOriginTranslatedWorld, 1 ), View.TranslatedWorldToClip );
 	float4 RayDirClip	= mul( float4( RayDirection * RayLength, 0 ), View.TranslatedWorldToClip );
 	float4 RayEndClip	= RayStartClip + RayDirClip;
 
 	float3 RayStartScreen = RayStartClip.xyz / RayStartClip.w;
 	float3 RayEndScreen = RayEndClip.xyz / RayEndClip.w;
-	
+
 	float3 RayStepScreen = RayEndScreen - RayStartScreen;
 
 	float3 RayStartUVz = float3( RayStartScreen.xy * View.ScreenPositionScaleBias.xy + View.ScreenPositionScaleBias.wz, RayStartScreen.z );
@@ -86,15 +79,67 @@ float ShadowRayCast(
 		SampleTime += Step;
 	}
 
-	float Shadow = FirstHitTime > 0.0 ? 1.0 : 0.0;
+	float Shadow = FirstHitTime > 0.0f ? 1.0f : 0.0f;
 
 	// Off screen masking
-	float2 Vignette = max(6.0 * abs(RayStartScreen.xy + RayStepScreen.xy * FirstHitTime) - 5.0, 0.0);
-	Shadow *= saturate( 1.0 - dot( Vignette, Vignette ) );
+	bool bValidUV = all( abs( RayStartScreen.xy + RayStepScreen.xy * FirstHitTime ) < 1.0f );
+	Shadow = bValidUV ? Shadow : 0.0f;	
+
+	return 1.0f - Shadow;
+}
+
+
+float ProjectSingleSample(
+	FLightShaderParameters Light,
+	float3 L,
+	uint LightType,
+	const float4 SvPosition,
+	const float SceneDepth,
+	float3 TranslatedWorldPosition,
+	float3 EstimatedGeoWorldNormal,
+	float ContactShadowLengthWorld,
+	inout float3 DebugOutput,
+	out float OutHitDistance)
+{
+	float StepOffset = InterleavedGradientNoise(SvPosition.xy, View.StateFrameIndexMod8);
+
+	// NOTE: Could skip most of the rest of the shader if contact shadows say we're shadowed, but likely to be small, incoherent branching
+
+	float3 WorldPosition = TranslatedWorldPosition - View.PreViewTranslation;
+
+	// Offset ray to start at the end of the contact shadow trace
+	FVirtualShadowMapSampleResult VirtualShadowMapSample =
+		SampleVirtualShadowMap(VirtualShadowMapId, WorldPosition, EstimatedGeoWorldNormal, ContactShadowLengthWorld);
+
+#if OUTPUT_TYPE != 2		// Don't kill pixels in debug output mode
+	if (!VirtualShadowMapSample.bValid)
+	{
+		// Generally shouldn't occur if PT allocation is conservative
+		discard;
+	}
+#endif
+
+	float RayMissFactor = VirtualShadowMapSample.bOccluded ? 0.0f : 1.0f;
+	float HitDistance = VirtualShadowMapSample.OccluderDistance;
+
+	if (DebugOutputType == 1)	// Invalid page
+	{
+		DebugOutput = VirtualShadowMapSample.bValid ? (0.25f * RayMissFactor) : float3(1, 0, 1);
+	}
+	else if (DebugOutputType == 2)	// ClipmapIndex/Mip
+	{
+		DebugOutput = lerp(RayMissFactor, VirtualShadowMapDebugIntToColor(VirtualShadowMapSample.ClipmapIndexOrMipLevel), 0.75f);
+	}
+	else if (DebugOutputType == 3) // Occluder distance
+	{
+		DebugOutput = lerp(RayMissFactor, HitDistance.xxx / 1000.0f, 0.75f);
+	}
 
-	return 1 - Shadow;
+	OutHitDistance = HitDistance;
+	return RayMissFactor;
 }
 
+
 float4 EncodeLightAttenuationFromMask(float ShadowMask)
 {
 	const float ShadowFadeFraction = 1;
@@ -109,51 +154,23 @@ float4 EncodeLightAttenuationFromMask(float ShadowMask)
 	return EncodeLightAttenuation(half4(FadedShadow, FadedSSSShadow, FadedShadow, FadedSSSShadow));
 }
 
-void VirtualShadowMapProjectionPS(
-	in float4 SvPosition : SV_Position,
-	out float4 OutputTarget : SV_Target
-	)
+float4 EncodeOuput(float4 OutputSignal, float3 DebugOutput)
 {
-	// TODO: Unify this with the logic in PageManagement.usf - we need to get the same values in both passes
-	const float2 ScreenUV = SvPositionToBufferUV(SvPosition);
-	const float DeviceZ = LookupDeviceZ(ScreenUV);
-	const float SceneDepth = ConvertFromDeviceZ(DeviceZ);
-	
-	float3 TranslatedWorldPosition = SvPositionToTranslatedWorld(float4(SvPosition.xy, DeviceZ, 1.0f));
-		
-	// TODO: Early out on sky, etc.
-	bool bInLightRegion = true;
-
-	FVirtualShadowMapProjectionShaderData ProjectionData = GetVirtualShadowMapProjectionData(VirtualShadowMapId);
-	uint LightType = ProjectionData.LightType;
-
-	FLightShaderParameters Light = GetRootLightShaderParameters();
-	float3 ToLight = Light.Direction;
-	if (ProjectionData.LightType == LIGHT_TYPE_SPOT)
-	{
-		// Early out if we're outside the spot light cone		
-		ToLight = Light.Position - (TranslatedWorldPosition - View.PreViewTranslation);
-		bool bInLightRadius = length(ToLight) <= rcp(Light_InvRadius);
-		ToLight = normalize(ToLight);
-		bool bInLightCone = dot(ToLight, Light.Direction) >= Light.SpotAngles.x;
-		bInLightRegion = bInLightRegion && bInLightRadius && bInLightCone;
-	}
-	if (!bInLightRegion)
-	{
-		discard;
-	}
-	
-	float ContactShadow = 1.0f;
-	float ContactShadowLengthWorld = 0.0f;
-	if (ContactShadowLength > 0.0f)
-	{
-		ContactShadowLengthWorld = ContactShadowLength * View.ClipToView[1][1] * SceneDepth;
-		float StepOffset = InterleavedGradientNoise( SvPosition.xy, View.StateFrameIndexMod8 );
-		ContactShadow = ShadowRayCast( TranslatedWorldPosition, ToLight, ContactShadowLengthWorld, 12, StepOffset );
-	}
+#if OUTPUT_TYPE == 0
+	// Output feeds into denoiser
+	return OutputSignal;
+#elif OUTPUT_TYPE == 1
+	// Output to screen shadow mask
+	return EncodeLightAttenuationFromMask(OutputSignal.r);
+#elif OUTPUT_TYPE == 2
+	// Output debug for "vis".
+	return float4(DebugOutput, 1.0f);
+#endif // OUTPUT_TYPE
+}
 
-	// NOTE: Could skip most of the rest of the shader if contact shadows say we're shadowed, but likely to be small, incoherent branching
 
+float3 GetEstimatedGeoWorldNormal(float3 TranslatedWorldPosition, float2 ScreenUV)
+{
 	// Figure out slope, we do world space since that is the space where we might override using the shading normal...
 	float3 TranslatedWorldPositionDDX = DDX(TranslatedWorldPosition);
 	float3 TranslatedWorldPositionDDY = DDY(TranslatedWorldPosition);
@@ -168,71 +185,210 @@ void VirtualShadowMapProjectionPS(
 	}
 #endif
 
-	float3 WorldPosition = TranslatedWorldPosition - View.PreViewTranslation;
-	// Offset ray to start at the end of the contact shadow trace
-	float RayStartDistance = ContactShadowLengthWorld;
-	FVirtualShadowMapSampleResult VirtualShadowMapSample = SampleVirtualShadowMap(VirtualShadowMapId, WorldPosition, EstimatedGeoWorldNormal, RayStartDistance);
+	return EstimatedGeoWorldNormal;
+}
+
+
+void VirtualShadowMapProjectionSpotPS(
+	in float4 SvPosition : SV_Position,
+	out float4 OutputTarget : SV_Target)
+{
+	// TODO: Unify this with the logic in PageManagement.usf - we need to get the same values in both passes
+	const float2 ScreenUV = SvPositionToBufferUV(SvPosition);
+	const float DeviceZ = LookupDeviceZ(ScreenUV);
+	const float SceneDepth = ConvertFromDeviceZ(DeviceZ);	
+	float3 TranslatedWorldPosition = SvPositionToTranslatedWorld(float4(SvPosition.xy, DeviceZ, 1.0f));
 	
-#if OUTPUT_TYPE != 2		// Don't kill pixels in debug output mode
-	if (!VirtualShadowMapSample.bValid)
+	FLightShaderParameters Light = GetRootLightShaderParameters();
+
+	FVirtualShadowMapProjectionShaderData ProjectionData = GetVirtualShadowMapProjectionData(VirtualShadowMapId);
+
+	// Early out if we're outside the spot light cone		
+	float3 ToLight = Light.Position - (TranslatedWorldPosition - View.PreViewTranslation);
+	bool bInLightRadius = length(ToLight) <= rcp(Light_InvRadius);
+	float3 ToLightDir = normalize(ToLight);
+	bool bInLightCone = dot(ToLightDir, Light.Direction) >= Light.SpotAngles.x;
+	bool bInLightRegion = bInLightRadius && bInLightCone;
+
+	float3 DebugOutput = float3(0, 0, 0);
+	int SampleCount = 1;
+	float HitDistance = 100000;	// TODO
+	float RayMissFactor = 1.0f;
+
+	if (bInLightRadius && bInLightCone)
 	{
-		// Generally shouldn't occur if PT allocation is conservative
-		discard;
-	}
+		const float ContactShadowLengthWorld = ContactShadowLength * View.ClipToView[1][1] * SceneDepth;
+#if 0
+		if (SMRTRayCount > 0)
+		{
+			OutputSignal = TraceSpotLight(
+				Light,
+				SvPosition,
+				SceneDepth,
+				TranslatedWorldPosition,
+				DebugOutput);
+		}
+		else
 #endif
-
-	float RayMissFactor = VirtualShadowMapSample.bOccluded ? 0.0f : 1.0f;
-	float HitDistance = VirtualShadowMapSample.OccluderDistance;
-	if (ContactShadow < 1.0f)
+		{
+			float3 EstimatedGeoWorldNormal = GetEstimatedGeoWorldNormal(TranslatedWorldPosition, ScreenUV);
+
+			RayMissFactor = ProjectSingleSample(
+				Light,
+				ToLightDir,
+				LIGHT_TYPE_SPOT,
+				SvPosition,
+				SceneDepth,
+				TranslatedWorldPosition,
+				EstimatedGeoWorldNormal,
+				ContactShadowLengthWorld,
+				/* out */ DebugOutput,
+				HitDistance);
+		}
+
+		if (ContactShadowLengthWorld > 0.0f)
+		{
+			// TODO: Unify this with the random sampling inside SMRT
+			uint2 SobolBase = SobolPixel(SvPosition.xy);
+			uint2 SobolFrame = SobolIndex(SobolBase, View.StateFrameIndexMod8, 3);
+			float StepOffset = InterleavedGradientNoise(SvPosition.xy, View.StateFrameIndexMod8);
+
+			float2 E = float2(SobolIndex(SobolFrame, 8 << 3)) / 0x10000;	// TODO
+
+			float ContactShadow = VirtualShadowMapScreenRayCast(TranslatedWorldPosition, ToLightDir, ContactShadowLengthWorld, 12, StepOffset);		
+			if (ContactShadow < 1.0f)
+			{
+				RayMissFactor *= ContactShadow;
+				HitDistance = ContactShadowLengthWorld;
+			}
+		}
+	}
+	else
 	{
-		RayMissFactor *= ContactShadow;
-		HitDistance = ContactShadowLengthWorld;
+		//discard;
 	}
 
-	// See signal encoding for SIGNAL_BUFFER_LAYOUT_PENUMBRA_HISTORY in SSDSignalBufferEncoding.ush
-	// We fuse that logic in here to avoid needing another pass, but it's hard to call the functions directly as they depend on a large
-	// number of preprocessor defines and logic that we don't want to mirror here.
-	// Unfortunately that means this logic can get out of date when the denoiser changes, but that's the trade-off for now.
 	float4 OutputSignal = float4(0, 0, 0, 0);
-	float SampleCount = 1.0f;
-	float TransmissionDistance = RayMissFactor; // As per RayTracingOcclusionRGS.usf
 	if (RayMissFactor < 1.0f)
 	{
-		float WorldBluringRadius = ComputeLightSampleWorldBluringRadius(TranslatedWorldPosition, LightType, Light, HitDistance);
-
-		OutputSignal = float4(
-			RayMissFactor,
-			SampleCount,
-			WorldBluringRadius,
-			TransmissionDistance);
+		float WorldBluringRadius = ComputeLightSampleWorldBluringRadius(TranslatedWorldPosition, LIGHT_TYPE_SPOT, Light, HitDistance);
+		OutputSignal = float4(RayMissFactor, SampleCount, WorldBluringRadius, RayMissFactor);
 	}
 	else
 	{
 		// Miss
-		OutputSignal = float4(RayMissFactor, SampleCount, WORLD_RADIUS_MISS, TransmissionDistance);
+		OutputSignal = float4(1, SampleCount, WORLD_RADIUS_MISS, 1);
 	}
 
+	OutputTarget = EncodeOuput(OutputSignal, DebugOutput);
+}
+
+
+void VirtualShadowMapProjectionDirectionalPS(
+	in float4 SvPosition : SV_Position,
+	out float4 OutputTarget : SV_Target
+)
+{
+	// TODO: Unify this with the logic in PageManagement.usf - we need to get the same values in both passes
+	const float2 ScreenUV = SvPositionToBufferUV(SvPosition);
+	const float DeviceZ = LookupDeviceZ(ScreenUV);
+	const float SceneDepth = ConvertFromDeviceZ(DeviceZ);
+	float3 TranslatedWorldPosition = SvPositionToTranslatedWorld(float4(SvPosition.xy, DeviceZ, 1.0f));
+
+	FLightShaderParameters Light = GetRootLightShaderParameters();
+		
+	FGBufferData GBufferData = GetGBufferData(ScreenUV);
+	float3 ToLightDir = normalize(Light.Direction);
+	bool bFrontFace = dot(ToLightDir, GBufferData.WorldNormal) > 0.0f;
 	
+	float3 DebugOutput = float3(0, 0, 0);
+	int SampleCount = 1;
+	float HitDistance = 100000;	// TODO
+	float RayMissFactor = 1.0f;
 
-#if OUTPUT_TYPE == 0
-	// Output feeds into denoiser
-	OutputTarget = OutputSignal;
-#elif OUTPUT_TYPE == 1
-	// Output to screen shadow mask
-	OutputTarget = EncodeLightAttenuationFromMask(OutputSignal.r);
-#elif OUTPUT_TYPE == 2
-	// Output debug for "vis".
-	float3 Color = float3(0, 0, 0);
+	if (bFrontFace)
+	{
+		const float ContactShadowLengthWorld = ContactShadowLength * View.ClipToView[1][1] * SceneDepth;
+
+		if (SMRTRayCount > 0)
+		{
+			SampleCount = SMRTRayCount;
+			RayMissFactor = TraceDirectional(
+				VirtualShadowMapId,
+				Light,
+				SvPosition,
+				SceneDepth,
+				TranslatedWorldPosition,
+				ContactShadowLengthWorld,
+				SMRTRayCount,
+				SMRTSamplesPerRay,
+				SMRTRayLengthScale,
+				DebugOutputType,
+				DebugOutput,
+				HitDistance);
+		}
+		else
+		{
+			float3 EstimatedGeoWorldNormal = GetEstimatedGeoWorldNormal(TranslatedWorldPosition, ScreenUV);
+
+			SampleCount = 1;
+			RayMissFactor = ProjectSingleSample(
+				Light,
+				Light.Direction,
+				LIGHT_TYPE_DIRECTIONAL,
+				SvPosition,
+				SceneDepth,
+				TranslatedWorldPosition,
+				EstimatedGeoWorldNormal,
+				ContactShadowLengthWorld,
+				/* out */ DebugOutput,
+				HitDistance);
+		}
+
+		/*
+		if (ContactShadowLengthWorld > 0.0f)
+		{
+			// TODO: Unify this with the random sampling inside SMRT
+			uint2 SobolBase = SobolPixel(SvPosition.xy);
+			uint2 SobolFrame = SobolIndex(SobolBase, View.StateFrameIndexMod8, 3);
+			float StepOffset = InterleavedGradientNoise(SvPosition.xy, View.StateFrameIndexMod8);
+
+			float2 E = float2(SobolIndex(SobolFrame, 8 << 3)) / 0x10000;	// TODO
+
+			float ContactShadow = VirtualShadowMapScreenRayCast(TranslatedWorldPosition, ToLightDir, ContactShadowLengthWorld, 12, StepOffset);		
+			if (ContactShadow < 1.0f)
+			{
+				RayMissFactor *= ContactShadow;
+				HitDistance = ContactShadowLengthWorld;
+			}
+		}
+		*/
+	}
 
-	if (DebugOutputType == 1)	// Invalid page
+	// See signal encoding for SIGNAL_BUFFER_LAYOUT_PENUMBRA_RECONSTRUCTION in SSDSignalBufferEncoding.ush
+	// We fuse that logic in here to avoid needing another pass, but it's hard to call the functions directly as they depend on a large
+	// number of preprocessor defines and logic that we don't want to mirror here.
+	// Unfortunately that means this logic can get out of date when the denoiser changes, but that's the trade-off for now.
+	float4 OutputSignal = float4(0, 0, 0, 0);
+	if (RayMissFactor < 1.0f)
 	{
-		Color = VirtualShadowMapSample.bValid ? (0.25f * OutputSignal.rrr) : float3(1, 0, 1);
+		float WorldBluringRadius = ComputeLightSampleWorldBluringRadius(TranslatedWorldPosition, LIGHT_TYPE_DIRECTIONAL, Light, HitDistance);
+		OutputSignal = float4(RayMissFactor, SampleCount, WorldBluringRadius, RayMissFactor);
+	}
+	else
+	{
+		// Miss
+		OutputSignal = float4(1, SampleCount, WORLD_RADIUS_MISS, 1);
 	}
 
-	OutputTarget = float4(Color, 1.0f);	
-#endif // OUTPUT_TYPE
+	OutputTarget = EncodeOuput(OutputSignal, DebugOutput);
 }
 
+
+
+
+
+
 Texture2D<float4> InputSignal;
 
 void VirtualShadowMapCompositePS(
diff --git a/Engine/Source/Runtime/Renderer/Private/SceneRendering.h b/Engine/Source/Runtime/Renderer/Private/SceneRendering.h
--- a/Engine/Source/Runtime/Renderer/Private/SceneRendering.h
+++ b/Engine/Source/Runtime/Renderer/Private/SceneRendering.h
@@ -969,6 +969,9 @@ struct FPreviousViewInfo
 	// History for shadow denoising.
 	TMap<const ULightComponent*, FScreenSpaceDenoiserHistory> ShadowHistories;
 
+	// History for virtual shadow denoising.
+	TMap<const ULightComponent*, FScreenSpaceDenoiserHistory> VirtualShadowHistories;
+
 	// History for denoising all lights penumbra at once.
 	FScreenSpaceDenoiserHistory PolychromaticPenumbraHarmonicsHistory;
 
diff --git a/Engine/Source/Runtime/Renderer/Private/ScreenSpaceDenoise.cpp b/Engine/Source/Runtime/Renderer/Private/ScreenSpaceDenoise.cpp
--- a/Engine/Source/Runtime/Renderer/Private/ScreenSpaceDenoise.cpp
+++ b/Engine/Source/Runtime/Renderer/Private/ScreenSpaceDenoise.cpp
@@ -3046,7 +3046,8 @@ FSSDSignalTextures IScreenSpaceDenoiser::DenoiseVirtualShadowMapMask(
 	const FSceneTextureParameters& SceneTextures,
 	const FLightSceneInfo* LightSceneInfo,
 	FIntRect LightScissorRect,
-	const FVirtualShadowMapMaskInputs& InputParameters)
+	const FVirtualShadowMapMaskInputs& InputParameters,
+	bool bUseTemporal)
 {
 	RDG_GPU_STAT_SCOPE(GraphBuilder, VirtualShadowsDenoiser);
 	const FLightSceneProxy* Proxy = LightSceneInfo->Proxy;
@@ -3057,7 +3058,7 @@ FSSDSignalTextures IScreenSpaceDenoiser::DenoiseVirtualShadowMapMask(
 	Settings.SignalProcessing = ESignalProcessing::VirtualShadowMapMask;
 	Settings.ReconstructionSamples = CVarVirtualShadowReconstructionSampleCount.GetValueOnRenderThread();
 	Settings.PreConvolutionCount = CVarVirtualShadowPreConvolutionCount.GetValueOnRenderThread();
-	Settings.bUseTemporalAccumulation = false;
+	Settings.bUseTemporalAccumulation = bUseTemporal;
 	
 	Settings.LightSceneInfo[0] = LightSceneInfo;
 	Settings.SignalScissor[0] = LightScissorRect;
@@ -3069,8 +3070,13 @@ FSSDSignalTextures IScreenSpaceDenoiser::DenoiseVirtualShadowMapMask(
 	// No temporal reprojection
 	TStaticArray<FScreenSpaceDenoiserHistory*, IScreenSpaceDenoiser::kMaxBatchSize> PrevHistories;
 	TStaticArray<FScreenSpaceDenoiserHistory*, IScreenSpaceDenoiser::kMaxBatchSize> NewHistories;
-	PrevHistories[0] = nullptr;
+	PrevHistories[0] = PreviousViewInfos->VirtualShadowHistories.Find(LightSceneInfo->Proxy->GetLightComponent());
 	NewHistories[0] = nullptr;
+	if (!View.bStatePrevViewInfoIsReadOnly)
+	{
+		check(View.ViewState);
+		NewHistories[0] = &View.ViewState->PrevFrameViewInfo.VirtualShadowHistories.FindOrAdd(LightSceneInfo->Proxy->GetLightComponent());
+	}
 
 	FSSDSignalTextures InputSignal;
 	InputSignal.Textures[0] = InputParameters.Signal;
diff --git a/Engine/Source/Runtime/Renderer/Private/ScreenSpaceDenoise.h b/Engine/Source/Runtime/Renderer/Private/ScreenSpaceDenoise.h
--- a/Engine/Source/Runtime/Renderer/Private/ScreenSpaceDenoise.h
+++ b/Engine/Source/Runtime/Renderer/Private/ScreenSpaceDenoise.h
@@ -348,7 +348,8 @@ public:
 		const FSceneTextureParameters& SceneTextures,
 		const FLightSceneInfo* LightSceneInfo,
 		FIntRect LightScissorRect,
-		const FVirtualShadowMapMaskInputs& InputParameters);
+		const FVirtualShadowMapMaskInputs& InputParameters,
+		bool bUseTemporal);
 
 	/** Returns the interface of the default denoiser of the renderer. */
 	static const IScreenSpaceDenoiser* GetDefaultDenoiser();
diff --git a/Engine/Source/Runtime/Renderer/Private/ShadowRendering.cpp b/Engine/Source/Runtime/Renderer/Private/ShadowRendering.cpp
--- a/Engine/Source/Runtime/Renderer/Private/ShadowRendering.cpp
+++ b/Engine/Source/Runtime/Renderer/Private/ShadowRendering.cpp
@@ -194,6 +194,12 @@ static TAutoConsoleVariable<int32> CVarVirtualSmDenoiser(
 	TEXT("Apply denoiser on virtual shadow map masks."),
 	ECVF_RenderThreadSafe);
 
+static TAutoConsoleVariable<int32> CVarVirtualSmDenoiserTemporal(
+	TEXT("r.Shadow.v.Denoiser.Temporal"),
+	1,
+	TEXT("Use temporal reprojection (when applicable)."),
+	ECVF_RenderThreadSafe);
+
 DEFINE_GPU_DRAWCALL_STAT(ShadowProjection);
 
 // 0:off, 1:low, 2:med, 3:high, 4:very high, 5:max
@@ -1834,6 +1840,7 @@ void FDeferredShadingSceneRenderer::RenderShadowProjections(
 						}
 
 						// Project virtual shadow maps
+						bool bUseTemporal = false;
 						if (VisibleLightInfo.VirtualShadowMapClipmaps.Num() > 0)
 						{
 							check(VisibleLightInfo.VirtualShadowMapClipmaps.Num() == 1);
@@ -1843,7 +1850,8 @@ void FDeferredShadingSceneRenderer::RenderShadowProjections(
 								View,
 								VirtualShadowMapArray,
 								ScissorRect,
-								SignalTexture);
+								SignalTexture,
+								bUseTemporal);
 						}
 						else
 						{
@@ -1854,9 +1862,12 @@ void FDeferredShadingSceneRenderer::RenderShadowProjections(
 								View,
 								VirtualShadowMapArray,
 								ScissorRect,
-								SignalTexture);
+								SignalTexture,
+								bUseTemporal);
 						}
 
+						bUseTemporal = bUseTemporal && (CVarVirtualSmDenoiserTemporal.GetValueOnRenderThread() != 0);
+
 						// Shadow filtering via denoiser
 						FSSDSignalTextures DenoisedSignal;
 						{
@@ -1874,7 +1885,8 @@ void FDeferredShadingSceneRenderer::RenderShadowProjections(
 								SceneTextureParameters,
 								LightSceneInfo,
 								ScissorRect,
-								Inputs);
+								Inputs,
+								bUseTemporal);
 						}
 
 						// Composite into light's screen shadow mask
diff --git a/Engine/Source/Runtime/Renderer/Private/ShadowSetup.cpp b/Engine/Source/Runtime/Renderer/Private/ShadowSetup.cpp
--- a/Engine/Source/Runtime/Renderer/Private/ShadowSetup.cpp
+++ b/Engine/Source/Runtime/Renderer/Private/ShadowSetup.cpp
@@ -799,25 +799,9 @@ void FProjectedShadowInfo::SetupWholeSceneProjection(
 	
 	if ( VirtualSmCacheEntry != nullptr )
 	{
+		check(DependentView == nullptr);		// Directional lights use clipmaps instead
 		check(BorderSize == 0);
-		FVector SnappedSubjectWorldSpacePosition = -Initializer.PreShadowTranslation;
-		VirtualSmCacheEntry->Update(VirtualShadowMap->ID, TranslatedWorldToClipInnerMatrix, -Initializer.PreShadowTranslation, DependentView != nullptr, Initializer, SnappedSubjectWorldSpacePosition);
-
-		// Ignore viewport if not a view-dependent SM.
-		if (DependentView != nullptr)
-		{
-			PreShadowTranslation = -SnappedSubjectWorldSpacePosition;
-			
-			FVector2D AdjustedScales = Initializer.Scales * FVirtualShadowMapArrayCacheManager::ClipSpaceScaleFactor;
-			// Warning! Re-set up a bunch of things as the scales are changed
-			InnerScaleMatrix = OuterScaleMatrix = FScaleMatrix( FVector(AdjustedScales.X, AdjustedScales.Y, 1.0f ) );
-			WorldToViewScaledInner = TranslatedWorldToView * InnerScaleMatrix;
-			ViewToClipInner = ViewToClipOuter = InnerScaleMatrix * FShadowProjectionMatrix(MinSubjectZ, MaxSubjectZ, Initializer.WAxis);
-			TranslatedWorldToClipInnerMatrix = TranslatedWorldToClipOuterMatrix = TranslatedWorldToView * ViewToClipInner;
-
-			ResolutionX = FVirtualShadowMap::VirtualMaxResolutionXY;
-			ResolutionY = FVirtualShadowMap::VirtualMaxResolutionXY;
-		}
+		VirtualSmCacheEntry->Update(VirtualShadowMap->ID, Initializer);
 	}
 	
 	// Nothing to do with subject depth. Just a scale factor to map z to 0-1.
diff --git a/Engine/Source/Runtime/Renderer/Private/VirtualShadowMaps/VirtualShadowMapArray.cpp b/Engine/Source/Runtime/Renderer/Private/VirtualShadowMaps/VirtualShadowMapArray.cpp
--- a/Engine/Source/Runtime/Renderer/Private/VirtualShadowMaps/VirtualShadowMapArray.cpp
+++ b/Engine/Source/Runtime/Renderer/Private/VirtualShadowMaps/VirtualShadowMapArray.cpp
@@ -76,14 +76,21 @@ static TAutoConsoleVariable<float> CVarResolutionPixelCountPercent(
 	ECVF_RenderThreadSafe
 );
 
-FMatrix CalcTranslatedWorldToShadowUvNormalMatrix(
+FMatrix CalcTranslatedWorldToShadowUVMatrix(
 	const FMatrix& TranslatedWorldToShadowView,
 	const FMatrix& ViewToClip)
 {
 	FMatrix TranslatedWorldToShadowClip = TranslatedWorldToShadowView * ViewToClip;
-	FMatrix ScaleAndBiasToSmUv = FScaleMatrix(FVector(0.5f, -0.5f, 1.0f)) * FTranslationMatrix(FVector(0.5f, 0.5f, 0.0f));
-	FMatrix TranslatedWorldToShadowUv = TranslatedWorldToShadowClip * ScaleAndBiasToSmUv;
-	return TranslatedWorldToShadowUv.GetTransposed().Inverse();
+	FMatrix ScaleAndBiasToSmUV = FScaleMatrix(FVector(0.5f, -0.5f, 1.0f)) * FTranslationMatrix(FVector(0.5f, 0.5f, 0.0f));
+	FMatrix TranslatedWorldToShadowUv = TranslatedWorldToShadowClip * ScaleAndBiasToSmUV;
+	return TranslatedWorldToShadowUv;
+}
+
+FMatrix CalcTranslatedWorldToShadowUVNormalMatrix(
+	const FMatrix& TranslatedWorldToShadowView,
+	const FMatrix& ViewToClip)
+{
+	return CalcTranslatedWorldToShadowUVMatrix(TranslatedWorldToShadowView, ViewToClip).GetTransposed().Inverse();
 }
 
 FVirtualShadowMapProjectionShaderData GetVirtualShadowMapProjectionShaderData(const FProjectedShadowInfo* ShadowInfo)
@@ -105,11 +112,11 @@ FVirtualShadowMapProjectionShaderData GetVirtualShadowMapProjectionShaderData(co
 	}
 
 	FMatrix ViewToClip = ShadowInfo->ViewToClipInner;
-		
+	
 	FVirtualShadowMapProjectionShaderData Data;
-	Data.TranslatedWorldToShadowViewMatrix = ShadowInfo->TranslatedWorldToView;
 	Data.ShadowViewToClipMatrix = ViewToClip;
-	Data.TranslatedWorldToShadowUvNormalMatrix = CalcTranslatedWorldToShadowUvNormalMatrix(ShadowInfo->TranslatedWorldToView, ViewToClip);
+	Data.TranslatedWorldToShadowUVMatrix = CalcTranslatedWorldToShadowUVMatrix(ShadowInfo->TranslatedWorldToView, ViewToClip);
+	Data.TranslatedWorldToShadowUVNormalMatrix = CalcTranslatedWorldToShadowUVNormalMatrix(ShadowInfo->TranslatedWorldToView, ViewToClip);
 	Data.ShadowPreViewTranslation = FVector(ShadowInfo->PreShadowTranslation);
 	Data.VirtualShadowMapId = ShadowInfo->VirtualShadowMap->ID;
 	Data.LightType = ShadowInfo->GetLightSceneInfo().Proxy->GetLightType();
diff --git a/Engine/Source/Runtime/Renderer/Private/VirtualShadowMaps/VirtualShadowMapArray.h b/Engine/Source/Runtime/Renderer/Private/VirtualShadowMaps/VirtualShadowMapArray.h
--- a/Engine/Source/Runtime/Renderer/Private/VirtualShadowMaps/VirtualShadowMapArray.h
+++ b/Engine/Source/Runtime/Renderer/Private/VirtualShadowMaps/VirtualShadowMapArray.h
@@ -78,13 +78,9 @@ public:
 // as well as cached shadow maps
 struct FVirtualShadowMapProjectionShaderData
 {
-	/**
-	 * Transform from shadow-pre-translated world space to shadow view space, example use: (WorldSpacePos + ShadowPreViewTranslation) * TranslatedWorldToShadowViewMatrix
-	 * TODO: Why don't we call it a rotation and store in a 3x3? Does it ever have translation in?
-	 */
-	FMatrix TranslatedWorldToShadowViewMatrix;
 	FMatrix ShadowViewToClipMatrix;
-	FMatrix TranslatedWorldToShadowUvNormalMatrix;
+	FMatrix TranslatedWorldToShadowUVMatrix;
+	FMatrix TranslatedWorldToShadowUVNormalMatrix;
 	/**
 	 * Translation from world space to shadow space (add before transform by TranslatedWorldToShadowViewMatrix).
 	 */
@@ -105,7 +101,8 @@ static_assert((sizeof(FVirtualShadowMapProjectionShaderData) % 16) == 0, "FVirtu
 
 FVirtualShadowMapProjectionShaderData GetVirtualShadowMapProjectionShaderData(const FProjectedShadowInfo* ShadowInfo);
 
-FMatrix CalcTranslatedWorldToShadowUvNormalMatrix(const FMatrix& TranslatedWorldToShadowView, const FMatrix& ViewToClip);
+FMatrix CalcTranslatedWorldToShadowUVMatrix(const FMatrix& TranslatedWorldToShadowView, const FMatrix& ViewToClip);
+FMatrix CalcTranslatedWorldToShadowUVNormalMatrix(const FMatrix& TranslatedWorldToShadowView, const FMatrix& ViewToClip);
 
 BEGIN_SHADER_PARAMETER_STRUCT(FVirtualShadowMapCommonParameters, )
 	SHADER_PARAMETER_ARRAY(uint32, HPageFlagLevelOffsets, [FVirtualShadowMap::MaxMipLevels])
diff --git a/Engine/Source/Runtime/Renderer/Private/VirtualShadowMaps/VirtualShadowMapCacheManager.cpp b/Engine/Source/Runtime/Renderer/Private/VirtualShadowMaps/VirtualShadowMapCacheManager.cpp
--- a/Engine/Source/Runtime/Renderer/Private/VirtualShadowMaps/VirtualShadowMapCacheManager.cpp
+++ b/Engine/Source/Runtime/Renderer/Private/VirtualShadowMaps/VirtualShadowMapCacheManager.cpp
@@ -28,7 +28,11 @@ static TAutoConsoleVariable<int32> CVarCacheVirtualSMs(
 );
 
 
-void FVirtualShadowMapCacheEntry::UpdateClipmap(int32 VirtualShadowMapId, const FMatrix &WorldToLight, FIntPoint PageSpaceLocation, float GlobalDepth)
+void FVirtualShadowMapCacheEntry::UpdateClipmap(
+	int32 VirtualShadowMapId,
+	const FMatrix &WorldToLight,
+	FIntPoint PageSpaceLocation,
+	float GlobalDepth)
 {
 	// Swap previous frame data over.
 	PrevVirtualShadowMapId = CurrentVirtualShadowMapId;
@@ -46,80 +50,25 @@ void FVirtualShadowMapCacheEntry::UpdateClipmap(int32 VirtualShadowMapId, const
 	CurrentShadowMapGlobalDepth = GlobalDepth;
 }
 
-void FVirtualShadowMapCacheEntry::Update(int32 VirtualShadowMapId, const FMatrix &ShadowPreTranslatedWorldToShadowClip, 
-	const FVector &SubjectWorldSpacePosition, bool bIsViewDependent, const FWholeSceneProjectedShadowInitializer &InCacheValidKey, 
-	FVector &SnappedSubjectWorldSpacePosition)
+void FVirtualShadowMapCacheEntry::Update(int32 VirtualShadowMapId, const FWholeSceneProjectedShadowInitializer &InCacheValidKey)
 {
 	// Swap previous frame data over.
 	PrevPageSpaceLocation = CurrentPageSpaceLocation;
 	PrevVirtualShadowMapId = CurrentVirtualShadowMapId;
 	PrevShadowMapGlobalDepth = CurrentShadowMapGlobalDepth;
-	SnappedSubjectWorldSpacePosition = SubjectWorldSpacePosition;
 
 	// Check cache validity based of shadow setup
-	if (!bIsViewDependent)
+	if (!CacheValidKey.IsCachedShadowValid(InCacheValidKey))
 	{
-		if (!CacheValidKey.IsCachedShadowValid(InCacheValidKey))
-		{
-			// Mark as invalid
-			PrevVirtualShadowMapId = INDEX_NONE;
-		}
-	}
-	else
-	{
-		bool bCachedValid =	CacheValidKey.WorldToLight == InCacheValidKey.WorldToLight
-							&& CacheValidKey.Scales == InCacheValidKey.Scales
-							&& CacheValidKey.SubjectBounds.Origin == InCacheValidKey.SubjectBounds.Origin
-							&& CacheValidKey.SubjectBounds.BoxExtent == InCacheValidKey.SubjectBounds.BoxExtent
-							&& CacheValidKey.SubjectBounds.SphereRadius == InCacheValidKey.SubjectBounds.SphereRadius
-							&& CacheValidKey.WAxis == InCacheValidKey.WAxis
-							&& CacheValidKey.MinLightW == InCacheValidKey.MinLightW
-							&& CacheValidKey.MaxDistanceToCastInLightW == InCacheValidKey.MaxDistanceToCastInLightW
-							&& CacheValidKey.bRayTracedDistanceField == InCacheValidKey.bRayTracedDistanceField;
-
-		if (!bCachedValid)
-		{ 
-			// Mark as invalid
-			PrevVirtualShadowMapId = INDEX_NONE;
-			//UE_LOG(LogRenderer, Display, TEXT("Invalidated!"));
-		}
+		PrevVirtualShadowMapId = INDEX_NONE;
+		//UE_LOG(LogRenderer, Display, TEXT("Invalidated!"));
 	}
 
-	// Update key data
 	CacheValidKey = InCacheValidKey;
 
-	// Compute new
 	CurrentVirtualShadowMapId = VirtualShadowMapId;
-
-	// E.g., CSMs
-	if (bIsViewDependent)
-	{
-		FMatrix ScaleAndBiasToSmPage = FScaleMatrix(FVector(FVirtualShadowMapArrayCacheManager::EffectiveCacheResolutionPages, FVirtualShadowMapArrayCacheManager::EffectiveCacheResolutionPages, 1.0f)) * FScaleMatrix(FVector(0.5f, -0.5f, 1.0f)) * FTranslationMatrix(FVector(0.5f, 0.5f, 0.0f));
-		FMatrix WorldToGlobalShadowPage = ShadowPreTranslatedWorldToShadowClip * ScaleAndBiasToSmPage;
-		
-		FVector ShadowMapLocationInGlobalPageSpace = WorldToGlobalShadowPage.TransformPosition(SubjectWorldSpacePosition);
-
-		FVector MinPageSpace((ShadowMapLocationInGlobalPageSpace.X - float(FVirtualShadowMapArrayCacheManager::EffectiveCacheResolutionPages) / 2.0f),
-			(ShadowMapLocationInGlobalPageSpace.Y - float(FVirtualShadowMapArrayCacheManager::EffectiveCacheResolutionPages) / 2.0f),
-			0.0f
-		);
-			
-		float MinX = FMath::FloorToFloat(MinPageSpace.X / FVirtualShadowMapArrayCacheManager::AlignmentPages);
-		float MinY = FMath::FloorToFloat(MinPageSpace.Y / FVirtualShadowMapArrayCacheManager::AlignmentPages);
-
-		FVector MinPageSpaceAligned = FVector(MinX, MinY, 0.0f) * FVirtualShadowMapArrayCacheManager::AlignmentPages;
-		FVector SmLocPsAligned = MinPageSpaceAligned + FVector(FVirtualShadowMap::Level0DimPagesXY / 2, FVirtualShadowMap::Level0DimPagesXY / 2, ShadowMapLocationInGlobalPageSpace.Z);
-
-		CurrentPageSpaceLocation = FIntPoint(SmLocPsAligned.X, SmLocPsAligned.Y);
-		CurrentShadowMapGlobalDepth = ShadowMapLocationInGlobalPageSpace.Z;
-
-		SnappedSubjectWorldSpacePosition = WorldToGlobalShadowPage.InverseFast().TransformPosition(SmLocPsAligned);
-	}
-	else
-	{
-		PrevPageSpaceLocation = CurrentPageSpaceLocation = FIntPoint(0, 0);
-		PrevShadowMapGlobalDepth = CurrentShadowMapGlobalDepth = 0.0f;
-	}
+	PrevPageSpaceLocation = CurrentPageSpaceLocation = FIntPoint(0, 0);
+	PrevShadowMapGlobalDepth = CurrentShadowMapGlobalDepth = 0.0f;
 }
 
 
diff --git a/Engine/Source/Runtime/Renderer/Private/VirtualShadowMaps/VirtualShadowMapCacheManager.h b/Engine/Source/Runtime/Renderer/Private/VirtualShadowMaps/VirtualShadowMapCacheManager.h
--- a/Engine/Source/Runtime/Renderer/Private/VirtualShadowMaps/VirtualShadowMapCacheManager.h
+++ b/Engine/Source/Runtime/Renderer/Private/VirtualShadowMaps/VirtualShadowMapCacheManager.h
@@ -22,12 +22,10 @@ public:
 	bool IsValid() { return PrevVirtualShadowMapId != INDEX_NONE; }
 
 
-	void Update(int32 VirtualShadowMapId, const FMatrix &ShadowPreTranslatedWorldToShadowClip, 
-		const FVector &SubjectWorldSpacePosition, bool bIsViewDependent, 
-		const FWholeSceneProjectedShadowInitializer &InCacheValidKey, 
-		FVector &SnappedSubjectWorldSpacePosition);
+	void Update(int32 VirtualShadowMapId, const FWholeSceneProjectedShadowInitializer &InCacheValidKey);
 
-	void UpdateClipmap(int32 VirtualShadowMapId, const FMatrix &WorldToLight, FIntPoint PageSpaceLocation, float GlobalDepth);
+	void UpdateClipmap(int32 VirtualShadowMapId, const FMatrix &WorldToLight,
+		FIntPoint PageSpaceLocation, float GlobalDepth);
 
 	/**
 	 * Returns the PrevVirtualShadowMapId if cached data is valid (bValidData), or INDEX_NONE otherwise. 
diff --git a/Engine/Source/Runtime/Renderer/Private/VirtualShadowMaps/VirtualShadowMapClipmap.cpp b/Engine/Source/Runtime/Renderer/Private/VirtualShadowMaps/VirtualShadowMapClipmap.cpp
--- a/Engine/Source/Runtime/Renderer/Private/VirtualShadowMaps/VirtualShadowMapClipmap.cpp
+++ b/Engine/Source/Runtime/Renderer/Private/VirtualShadowMaps/VirtualShadowMapClipmap.cpp
@@ -24,7 +24,6 @@ static TAutoConsoleVariable<float> CVarVirtualShadowMapClipmapFirstLevel(
 	ECVF_RenderThreadSafe
 );
 
-
 // "Virtual" clipmap level to clipmap radius
 static float GetLevelRadius(int32 Level)
 {
@@ -103,9 +102,7 @@ FVirtualShadowMapClipmap::FVirtualShadowMapClipmap(
 
 		const FVector SnappedWorldCenter = ViewToWorldRotationMatrix.TransformPosition(ViewCenter);
 		const float SnappedLevelRadius = 0.5f * (PageSizeInWorldSpace * DimPages);
-				
-		// NOTE: Any ZScale/ZOffset change probably needs to invalidate the page cache...
-		// For now they should be stable for a given (absolute) clipmap level
+
 		const float ZScale = 0.5f / RawLevelRadius;
 		const float ZOffset = RawLevelRadius;
 
@@ -153,11 +150,12 @@ FVirtualShadowMapProjectionShaderData FVirtualShadowMapClipmap::GetProjectionSha
 {
 	check(ClipmapIndex >= 0 && ClipmapIndex < LevelData.Num());
 	const FLevelData& Level = LevelData[ClipmapIndex];
-
+	
+	// NOTE: Some shader logic (projection, etc) assumes some of these parameters are constant across all levels in a clipmap
 	FVirtualShadowMapProjectionShaderData Data;
-	Data.TranslatedWorldToShadowViewMatrix = WorldToViewRotationMatrix;
 	Data.ShadowViewToClipMatrix = Level.ViewToClip;
-	Data.TranslatedWorldToShadowUvNormalMatrix = CalcTranslatedWorldToShadowUvNormalMatrix(WorldToViewRotationMatrix, Level.ViewToClip);
+	Data.TranslatedWorldToShadowUVMatrix = CalcTranslatedWorldToShadowUVMatrix(WorldToViewRotationMatrix, Level.ViewToClip);
+	Data.TranslatedWorldToShadowUVNormalMatrix = CalcTranslatedWorldToShadowUVNormalMatrix(WorldToViewRotationMatrix, Level.ViewToClip);
 	Data.ShadowPreViewTranslation = -Level.WorldCenter;
 	Data.VirtualShadowMapId = Level.VirtualShadowMap->ID;
 	Data.LightType = ELightComponentType::LightType_Directional;
diff --git a/Engine/Source/Runtime/Renderer/Private/VirtualShadowMaps/VirtualShadowMapProjection.cpp b/Engine/Source/Runtime/Renderer/Private/VirtualShadowMaps/VirtualShadowMapProjection.cpp
--- a/Engine/Source/Runtime/Renderer/Private/VirtualShadowMaps/VirtualShadowMapProjection.cpp
+++ b/Engine/Source/Runtime/Renderer/Private/VirtualShadowMaps/VirtualShadowMapProjection.cpp
@@ -38,6 +38,45 @@ static TAutoConsoleVariable<int32> CVarVirtualShadowMapDebugProjection(
 	ECVF_RenderThreadSafe
 );
 
+static TAutoConsoleVariable<int32> CVarSMRTRayCountDirectional(
+	TEXT( "r.Shadow.v.SMRT.RayCountDirectional" ),
+	0,
+	TEXT( "Ray count for shadow map tracing of directional lights. 0 = disabled." ),
+	ECVF_RenderThreadSafe
+);
+
+#if 0
+static TAutoConsoleVariable<int32> CVarSMRTRayCountSpot(
+	TEXT( "r.Shadow.v.SMRT.RayCountSpot" ),
+	0,
+	TEXT( "Ray count for shadow map tracing of spot lights. 0 = disabled." ),
+	ECVF_RenderThreadSafe
+);
+
+static TAutoConsoleVariable<int32> CVarSMRTSamplesPerRaySpot(
+	TEXT( "r.Shadow.v.SMRT.SamplesPerRaySpot" ),
+	16,
+	TEXT( "Shadow map samples per ray for spot lights" ),
+	ECVF_RenderThreadSafe
+);
+#endif
+
+static TAutoConsoleVariable<int32> CVarSMRTSamplesPerRayDirectional(
+	TEXT( "r.Shadow.v.SMRT.SamplesPerRayDirectional" ),
+	12,
+	TEXT( "Shadow map samples per ray for directional lights" ),
+	ECVF_RenderThreadSafe
+);
+
+static TAutoConsoleVariable<float> CVarSMRTRayLengthScaleDirectional(
+	TEXT( "r.Shadow.v.SMRT.RayLengthScaleDirectional" ),
+	2.0f,
+	TEXT( "Length of ray to shoot for directional lights, scaled by distance to camera." )
+	TEXT( "Shorter rays limit the screen space size of shadow penumbra. " )
+	TEXT( "Longer rays require more samples to avoid shadows disconnecting from contact points. " ),
+	ECVF_RenderThreadSafe
+);
+
 BEGIN_SHADER_PARAMETER_STRUCT(FProjectionParameters, )
 	SHADER_PARAMETER_STRUCT_INCLUDE(FVirtualShadowMapSamplingParameters, ProjectionParameters)
 	SHADER_PARAMETER_STRUCT(FLightShaderParameters, Light)
@@ -46,13 +85,16 @@ BEGIN_SHADER_PARAMETER_STRUCT(FProjectionParameters, )
 	SHADER_PARAMETER(int32, VirtualShadowMapId)
 	SHADER_PARAMETER(float, ContactShadowLength)
 	SHADER_PARAMETER(int32, DebugOutputType)
+	SHADER_PARAMETER(int32, SMRTRayCount)
+	SHADER_PARAMETER(int32, SMRTSamplesPerRay)
+	SHADER_PARAMETER(float, SMRTRayLengthScale)
 	RENDER_TARGET_BINDING_SLOTS()
 END_SHADER_PARAMETER_STRUCT()
 
-class FVirtualShadowMapProjectionPS : public FGlobalShader
+class FVirtualShadowMapProjectionSpotPS : public FGlobalShader
 {
-	DECLARE_GLOBAL_SHADER(FVirtualShadowMapProjectionPS);
-	SHADER_USE_PARAMETER_STRUCT(FVirtualShadowMapProjectionPS, FGlobalShader);
+	DECLARE_GLOBAL_SHADER(FVirtualShadowMapProjectionSpotPS);
+	SHADER_USE_PARAMETER_STRUCT(FVirtualShadowMapProjectionSpotPS, FGlobalShader);
 	
 	class FOutputTypeDim : SHADER_PERMUTATION_ENUM_CLASS("OUTPUT_TYPE", EVirtualShadowMapProjectionOutputType);
 	using FPermutationDomain = TShaderPermutationDomain<FOutputTypeDim>;
@@ -69,7 +111,31 @@ class FVirtualShadowMapProjectionPS : public FGlobalShader
 		FVirtualShadowMapArray::SetShaderDefines(OutEnvironment);
 	}
 };
-IMPLEMENT_GLOBAL_SHADER(FVirtualShadowMapProjectionPS, "/Engine/Private/VirtualShadowMaps/VirtualShadowMapProjection.usf", "VirtualShadowMapProjectionPS", SF_Pixel);
+IMPLEMENT_GLOBAL_SHADER(FVirtualShadowMapProjectionSpotPS, "/Engine/Private/VirtualShadowMaps/VirtualShadowMapProjection.usf", "VirtualShadowMapProjectionSpotPS", SF_Pixel);
+
+class FVirtualShadowMapProjectionDirectionalPS : public FGlobalShader
+{
+	DECLARE_GLOBAL_SHADER(FVirtualShadowMapProjectionDirectionalPS);
+	SHADER_USE_PARAMETER_STRUCT(FVirtualShadowMapProjectionDirectionalPS, FGlobalShader);
+
+	class FOutputTypeDim : SHADER_PERMUTATION_ENUM_CLASS("OUTPUT_TYPE", EVirtualShadowMapProjectionOutputType);
+	using FPermutationDomain = TShaderPermutationDomain<FOutputTypeDim>;
+
+	using FParameters = FProjectionParameters;
+
+	static bool ShouldCompilePermutation(const FGlobalShaderPermutationParameters& Parameters)
+	{
+		return DoesPlatformSupportNanite(Parameters.Platform);
+	}
+
+	static void ModifyCompilationEnvironment(const FGlobalShaderPermutationParameters& Parameters, FShaderCompilerEnvironment& OutEnvironment)
+	{
+		FVirtualShadowMapArray::SetShaderDefines(OutEnvironment);
+	}
+};
+IMPLEMENT_GLOBAL_SHADER(FVirtualShadowMapProjectionDirectionalPS, "/Engine/Private/VirtualShadowMaps/VirtualShadowMapProjection.usf", "VirtualShadowMapProjectionDirectionalPS", SF_Pixel);
+
+
 
 BEGIN_SHADER_PARAMETER_STRUCT(FVirtualShadowMapProjectionCompositeParameters, )
 	SHADER_PARAMETER_RDG_TEXTURE(Texture2D<float4>, InputSignal)
@@ -98,7 +164,8 @@ class FVirtualShadowMapProjectionCompositePS : public FGlobalShader
 };
 IMPLEMENT_GLOBAL_SHADER(FVirtualShadowMapProjectionCompositePS, "/Engine/Private/VirtualShadowMaps/VirtualShadowMapProjection.usf", "VirtualShadowMapCompositePS", SF_Pixel);
 
-static void AddPass_RenderVirtualShadowMapProjection(
+// Returns true if temporal denoising should be used
+static bool AddPass_RenderVirtualShadowMapProjection(
 	const FLightSceneProxy* LightProxy,
 	FRDGBuilder& GraphBuilder,
 	const FViewInfo& View, 
@@ -127,21 +194,58 @@ static void AddPass_RenderVirtualShadowMapProjection(
 
 	PassParameters->RenderTargets[0] = RenderTargetBinding;
 
+	if (LightProxy->GetLightType() == LightType_Directional)
 	{
-		FVirtualShadowMapProjectionPS::FPermutationDomain PermutationVector;
-		PermutationVector.Set<FVirtualShadowMapProjectionPS::FOutputTypeDim>(OutputType);
-		auto PixelShader = ShaderMap->GetShader<FVirtualShadowMapProjectionPS>(PermutationVector);
+		PassParameters->SMRTRayCount = CVarSMRTRayCountDirectional.GetValueOnRenderThread();
+		PassParameters->SMRTSamplesPerRay = CVarSMRTSamplesPerRayDirectional.GetValueOnRenderThread();
+		PassParameters->SMRTRayLengthScale = CVarSMRTRayLengthScaleDirectional.GetValueOnRenderThread();
+
+		FVirtualShadowMapProjectionDirectionalPS::FPermutationDomain PermutationVector;
+		PermutationVector.Set<FVirtualShadowMapProjectionDirectionalPS::FOutputTypeDim>(OutputType);
+		auto PixelShader = ShaderMap->GetShader<FVirtualShadowMapProjectionDirectionalPS>(PermutationVector);
 		ValidateShaderParameters(PixelShader, *PassParameters);
 
 		// NOTE: We use SV_Position in the shader, so we don't need separate scissor/view rect
 		FPixelShaderUtils::AddFullscreenPass(GraphBuilder,
 			ShaderMap,
-			(OutputType == EVirtualShadowMapProjectionOutputType::Debug) ? RDG_EVENT_NAME("Debug Projection") : RDG_EVENT_NAME("Projection"),
+			(OutputType == EVirtualShadowMapProjectionOutputType::Debug) ? RDG_EVENT_NAME("Debug Directional Projection") : RDG_EVENT_NAME("Directional Projection"),
 			PixelShader,
 			PassParameters,
 			ScissorRect,
 			BlendState);
+
+		// Use temporal denoising with SMRT
+		return PassParameters->SMRTRayCount > 0;
 	}
+	else if (LightProxy->GetLightType() == LightType_Spot)
+	{
+		PassParameters->SMRTRayCount = 0; //CVarSMRTRayCountSpot.GetValueOnRenderThread();
+		PassParameters->SMRTSamplesPerRay = 0; //CVarSMRTSamplesPerRaySpot.GetValueOnRenderThread();
+		PassParameters->SMRTRayLengthScale = 0.0f;		// Currently unused in this path
+
+		FVirtualShadowMapProjectionSpotPS::FPermutationDomain PermutationVector;
+		PermutationVector.Set<FVirtualShadowMapProjectionSpotPS::FOutputTypeDim>(OutputType);
+		auto PixelShader = ShaderMap->GetShader<FVirtualShadowMapProjectionSpotPS>(PermutationVector);
+		ValidateShaderParameters(PixelShader, *PassParameters);
+
+		// NOTE: We use SV_Position in the shader, so we don't need separate scissor/view rect
+		FPixelShaderUtils::AddFullscreenPass(GraphBuilder,
+			ShaderMap,
+			(OutputType == EVirtualShadowMapProjectionOutputType::Debug) ? RDG_EVENT_NAME("Debug Spot Projection") : RDG_EVENT_NAME("Spot Projection"),
+			PixelShader,
+			PassParameters,
+			ScissorRect,
+			BlendState);
+
+		// Use temporal denoising with SMRT
+		return PassParameters->SMRTRayCount > 0;
+	}
+	else
+	{
+		check(false);	// Unsupported light type for VSM projection
+	}
+
+	return false;
 }
 
 static FRDGTextureRef CreateDebugOutput(FRDGBuilder& GraphBuilder, FIntPoint Extent)
@@ -154,7 +258,8 @@ static FRDGTextureRef CreateDebugOutput(FRDGBuilder& GraphBuilder, FIntPoint Ext
 	return GraphBuilder.CreateTexture(DebugOutputDesc, TEXT("VirtSmDebugProj"));
 }
 
-static void RenderVirtualShadowMapProjectionCommon(
+// Returns true if temporal denoising should be used
+static bool RenderVirtualShadowMapProjectionCommon(
 	const FLightSceneProxy* LightProxy,
 	FRDGBuilder& GraphBuilder,
 	const FViewInfo& View, 
@@ -166,7 +271,7 @@ static void RenderVirtualShadowMapProjectionCommon(
 	FRHIBlendState* BlendState)
 {
 	// Main Pass
-	AddPass_RenderVirtualShadowMapProjection(
+	bool bOutUseTemporalDenoising = AddPass_RenderVirtualShadowMapProjection(
 		LightProxy,
 		GraphBuilder,
 		View,
@@ -197,6 +302,8 @@ static void RenderVirtualShadowMapProjectionCommon(
 
 		GraphBuilder.QueueTextureExtraction(DebugOutput, &VirtualShadowMapArray.DebugVisualizationProjectionOutput);
 	}
+
+	return bOutUseTemporalDenoising;
 }
 
 void RenderVirtualShadowMapProjectionForDenoising(
@@ -205,10 +312,11 @@ void RenderVirtualShadowMapProjectionForDenoising(
 	const FViewInfo& View,
 	FVirtualShadowMapArray& VirtualShadowMapArray,
 	const FIntRect ScissorRect,
-	FRDGTextureRef SignalTexture)
+	FRDGTextureRef SignalTexture,
+	bool& bOutUseTemporalDenoising)
 {
 	FRenderTargetBinding RenderTargetBinding(SignalTexture, ERenderTargetLoadAction::EClear);
-	RenderVirtualShadowMapProjectionCommon(
+	bOutUseTemporalDenoising = RenderVirtualShadowMapProjectionCommon(
 		ShadowInfo->GetLightSceneInfo().Proxy,
 		GraphBuilder,
 		View,
@@ -226,10 +334,11 @@ void RenderVirtualShadowMapProjectionForDenoising(
 	const FViewInfo& View,
 	FVirtualShadowMapArray& VirtualShadowMapArray,
 	const FIntRect ScissorRect,
-	FRDGTextureRef SignalTexture)
+	FRDGTextureRef SignalTexture,
+	bool& bOutUseTemporalDenoising)
 {	
 	FRenderTargetBinding RenderTargetBinding(SignalTexture, ERenderTargetLoadAction::EClear);
-	RenderVirtualShadowMapProjectionCommon(
+	bOutUseTemporalDenoising = RenderVirtualShadowMapProjectionCommon(
 		Clipmap->GetLightSceneInfo().Proxy,
 		GraphBuilder,
 		View,
diff --git a/Engine/Source/Runtime/Renderer/Private/VirtualShadowMaps/VirtualShadowMapProjection.h b/Engine/Source/Runtime/Renderer/Private/VirtualShadowMaps/VirtualShadowMapProjection.h
--- a/Engine/Source/Runtime/Renderer/Private/VirtualShadowMaps/VirtualShadowMapProjection.h
+++ b/Engine/Source/Runtime/Renderer/Private/VirtualShadowMaps/VirtualShadowMapProjection.h
@@ -18,7 +18,8 @@ void RenderVirtualShadowMapProjectionForDenoising(
 	const FViewInfo& View,
 	FVirtualShadowMapArray& VirtualShadowMapArray,
 	const FIntRect ScissorRect,
-	FRDGTextureRef SignalTexture);
+	FRDGTextureRef SignalTexture,
+	bool& bOutUseTemporalDenoising);
 
 void RenderVirtualShadowMapProjectionForDenoising(
 	FProjectedShadowInfo* ShadowInfo,
@@ -26,7 +27,8 @@ void RenderVirtualShadowMapProjectionForDenoising(
 	const FViewInfo& View,
 	FVirtualShadowMapArray& VirtualShadowMapArray,
 	const FIntRect ScissorRect,
-	FRDGTextureRef SignalTexture);
+	FRDGTextureRef SignalTexture,
+	bool& bOutUseTemporalDenoising);
 
 void RenderVirtualShadowMapProjection(
 	const TSharedPtr<FVirtualShadowMapClipmap>& Clipmap,
