Fix/remove some debug code in previous commit

diff --git a/Engine/Shaders/Private/VirtualShadowMaps/ProjectionDirectional.ush b/Engine/Shaders/Private/VirtualShadowMaps/ProjectionDirectional.ush
--- a/Engine/Shaders/Private/VirtualShadowMaps/ProjectionDirectional.ush
+++ b/Engine/Shaders/Private/VirtualShadowMaps/ProjectionDirectional.ush
@@ -234,7 +234,7 @@ float TraceDirectional(
 		bool bHit = ShadowRayCastDirectional(
 			VirtualShadowMapId,
 			ViewPosition,
-			RayStart + RayStartOffset,
+			RayStart,
 			RayDir * RayLength,
 			SamplesPerRay,
 			StepOffset,
diff --git a/Engine/Shaders/Private/VirtualShadowMaps/VirtualShadowMapProjection.usf b/Engine/Shaders/Private/VirtualShadowMaps/VirtualShadowMapProjection.usf
--- a/Engine/Shaders/Private/VirtualShadowMaps/VirtualShadowMapProjection.usf
+++ b/Engine/Shaders/Private/VirtualShadowMaps/VirtualShadowMapProjection.usf
@@ -345,7 +345,6 @@ void VirtualShadowMapProjectionDirectionalPS(
 				HitDistance);
 		}
 
-		/*
 		if (ContactShadowLengthWorld > 0.0f)
 		{
 			// TODO: Unify this with the random sampling inside SMRT
@@ -362,7 +361,6 @@ void VirtualShadowMapProjectionDirectionalPS(
 				HitDistance = ContactShadowLengthWorld;
 			}
 		}
-		*/
 	}
 
 	// See signal encoding for SIGNAL_BUFFER_LAYOUT_PENUMBRA_RECONSTRUCTION in SSDSignalBufferEncoding.ush
