Add TopLayer data to classification pass and use its output to the ReflectionEnvironment pass.

diff --git a/Engine/Shaders/Private/ReflectionEnvironmentPixelShader.usf b/Engine/Shaders/Private/ReflectionEnvironmentPixelShader.usf
--- a/Engine/Shaders/Private/ReflectionEnvironmentPixelShader.usf
+++ b/Engine/Shaders/Private/ReflectionEnvironmentPixelShader.usf
@@ -362,16 +362,9 @@ void ReflectionEnvironmentSkyLighting(
 #if STRATA_ENABLED && ENABLE_SKY_LIGHT
 	if (bStrataMaterial)
 	{
-		// Hack: AdjustBaseColorAndSpecularColorForSubsurfaceProfileLightingCopyHack forces GBuffer.BaseColor to be white when View.bSubsurfacePostprocessEnabled.
-		// And Starta always set UseSubsurfaceProfile to true to handle SSS. 
-		// STRATA_TODO We need to decouple starta from that, but for now we directly sample EnvBRDF-as-BaseColor here.
-		const float4 GBufferC = GBufferCTexture.SampleLevel(GBufferCTextureSampler, BufferUV, 0);
-		const float3 GBufferBaseColor = GBufferC.rgb;
-		const float3 EnvBRDF = GBufferBaseColor;
-
-		// We use WorldNormal, Roughness and BaseColor (representing EnvBRDF) written by the strata top layer.
+		const FStrataTopLayerData TopLayerData = StrataUnpackTopLayerData(Strata.TopLayerTexture.Load(uint3(PixelPos, 0)));
 		const float4 SSR = Texture2DSample(ScreenSpaceReflectionsTexture, ScreenSpaceReflectionsSampler, BufferUV);
-		SkyLighting.rgb += SSR.rgb * EnvBRDF;
+		SkyLighting.rgb += SSR.rgb * TopLayerData.EnvBRDF;
 
 		// SSR is computed for the top level by averaging BSDF components.
 		// This is the top level specular contribution factor according to areas.
diff --git a/Engine/Shaders/Private/Strata/Strata.ush b/Engine/Shaders/Private/Strata/Strata.ush
--- a/Engine/Shaders/Private/Strata/Strata.ush
+++ b/Engine/Shaders/Private/Strata/Strata.ush
@@ -143,6 +143,13 @@ struct FStrataBSDF
 
 void StrataRequestSharedNormalTangentBasis(inout uint NormalTypes, uint Index);
 
+uint   PackR11G11B10F(float3 rgb);
+float3 UnpackR11G11B10F(uint rgb);
+uint   PackAlbedoLinearToGamma2AlphaLinear(float4 rgba);
+float4 UnpackAlbedoGamma2ToLinearAlphaLinear(uint rgba);
+uint   StrataPackNormal(in float3 Normal);
+float3 StrataUnpackNormal(uint PackedNormal);
+
 struct FStrataClassification
 {
 	uint ShadingModels;
@@ -154,22 +161,74 @@ struct FStrataClassification
 
 uint StrataPackClassificationData(FStrataClassification In)
 {
-	return 
-		(In.ShadingModels & 0xFF) | 
-		((In.bIsStrataMaterial		? 1 : 0) << 8) |
-		((In.bHasSubsurface			? 1 : 0) << 9) | 
-		((In.bHasSubsurfaceProfile	? 1 : 0) << 10)|
-		((In.bHasTransmission		? 1 : 0) << 11);
+	return
+		(In.ShadingModels & 0xFF) |
+		((In.bIsStrataMaterial ? 1 : 0) << 8) |
+		((In.bHasSubsurface ? 1 : 0) << 9) |
+		((In.bHasSubsurfaceProfile ? 1 : 0) << 10) |
+		((In.bHasTransmission ? 1 : 0) << 11);
 }
 
 FStrataClassification StrataUnpackClassificationData(uint In)
 {
 	FStrataClassification Out;
-	Out.ShadingModels			= (In & 0xFF);
-	Out.bIsStrataMaterial		= (In >>  8) & 0x1;
-	Out.bHasSubsurface			= (In >>  9) & 0x1;
-	Out.bHasSubsurfaceProfile	= (In >> 10) & 0x1;
-	Out.bHasTransmission		= (In >> 11) & 0x1;
+	Out.ShadingModels = (In & 0xFF);
+	Out.bIsStrataMaterial = (In >> 8) & 0x1;
+	Out.bHasSubsurface = (In >> 9) & 0x1;
+	Out.bHasSubsurfaceProfile = (In >> 10) & 0x1;
+	Out.bHasTransmission = (In >> 11) & 0x1;
+	return Out;
+}
+
+struct FStrataTopLayerData
+{
+	float3 Normal;
+	float3 BaseColor;
+	float  Roughness;
+	float3 EnvBRDF;
+
+	// transient data when accumulating top layer data
+	float AccumRoughnessCount;
+};
+
+struct FStrataSubsurfaceData
+{
+	bool   bHasSubsurface;
+	bool   bHasProfile;
+	float  RadiusScale;
+
+	float3 BaseColor;
+	uint   Profile;
+
+	float3 DMFPAlbedo;
+	float  DMFPRadius;
+};
+
+void FinalizeTopLayerData(inout FStrataTopLayerData Out)
+{
+	Out.Normal = normalize(Out.Normal);
+	Out.BaseColor = saturate(Out.BaseColor);
+	Out.Roughness = Out.AccumRoughnessCount > 0.0f ? saturate(Out.Roughness / Out.AccumRoughnessCount) : 1.f;
+	Out.EnvBRDF = Out.BaseColor; // Strata_todo: better approximation
+}
+
+uint2 StrataPackTopLayerData(FStrataTopLayerData In)
+{
+	uint2 Out = 0;
+	Out.x = PackAlbedoLinearToGamma2AlphaLinear(float4(In.BaseColor, In.Roughness));
+	Out.y = StrataPackNormal(In.Normal);
+	return Out;
+}
+
+FStrataTopLayerData StrataUnpackTopLayerData(uint2 In)
+{
+	const float4 BaseColorAndRoughness = UnpackAlbedoGamma2ToLinearAlphaLinear(In.x);
+
+	FStrataTopLayerData Out = (FStrataTopLayerData)0;
+	Out.Normal = StrataUnpackNormal(In.y);
+	Out.BaseColor = BaseColorAndRoughness.xyz;
+	Out.Roughness = BaseColorAndRoughness.w;
+	Out.EnvBRDF = Out.BaseColor; // Strata_todo: better approximation
 	return Out;
 }
 
@@ -1003,14 +1062,8 @@ void PackStrataOut(
 	// We clamp the layer count to the maximum one according to our byte per pixel budget (removing the bottom layer first)
 	Strata.LayerCount = STRATA_CLAMPED_LAYER_COUNT;
 
-	// Subsurface scattering parameters
-	bool   bHasSubsurfaceScattering = false;
-	bool   bHasSubsurfaceProfile = false;
-	float  SubsurfaceProfile = 0;
-	float  SubsurfaceRadiusScale = 0;
-	float3 SubsurfaceBaseColor = 0;
-	float3 SubsurfaceDMFPAlbedo = 0;
-	float  SubsurfaceDMFPRadius = 0;
+	FStrataTopLayerData TopLayer = (FStrataTopLayerData)0;
+	FStrataSubsurfaceData SSS = (FStrataSubsurfaceData)0;
 
 	// For each of the BSDF, we evaluated their combined weight/throughput according to the material operation performed (layering, mixing, etc.).
 	int l;
@@ -1020,10 +1073,6 @@ void PackStrataOut(
 		float3 LayerThroughput = 0.0f;
 		float3 LayerCoverage = 0.0f;
 
-		float3 AccumTopNormal = 0.0f;
-		float  AccumRoughness = 0.0f;
-		float  AccumRoughnessCount = 0.0f;
-		float3 AccumBaseColor = 0.0f;
 		const bool bTopLayer = l == 0;
 
 		UNROLL
@@ -1042,7 +1091,7 @@ void PackStrataOut(
 			float3 N = TangentBasis[2];
 			//N = normalize(N);
 
-			AccumTopNormal += N;
+			TopLayer.Normal += N;
 
 			BxDFContext ContextNoLight = (BxDFContext)0;
 			const float3 FakeL = float3(0, 0, 1);
@@ -1054,20 +1103,20 @@ void PackStrataOut(
 			{
 				if (bTopLayer)
 				{
-					AccumRoughness		+= BSDFCoverageAvg * FullyRough; // Do not trigger SSR by using fully rough
-					AccumRoughnessCount += BSDFCoverageAvg;
+					TopLayer.Roughness += BSDFCoverageAvg * FullyRough; // Do not trigger SSR by using fully rough
+					TopLayer.AccumRoughnessCount += BSDFCoverageAvg;
 				}
 
 				// Select the first/toppest layer with subsurface scattering
-				if (BSDF_GETHASSCATTERING(BSDF) && BSDFCoverageAvg > 0 && DIFFUSE_SSSPROFILERADIUS(BSDF) > 0 && !bHasSubsurfaceScattering)
+				if (BSDF_GETHASSCATTERING(BSDF) && BSDFCoverageAvg > 0 && DIFFUSE_SSSPROFILERADIUS(BSDF) > 0 && !SSS.bHasSubsurface)
 				{
-					bHasSubsurfaceScattering = true;
-					bHasSubsurfaceProfile	= BSDF_GETHASSSSPROFILE(BSDF);
-					SubsurfaceProfile		= DIFFUSE_SSSPROFILEID(BSDF);
-					SubsurfaceRadiusScale	= DIFFUSE_SSSPROFILERADIUS(BSDF);
-					SubsurfaceBaseColor		= DIFFUSE_ALBEDO(BSDF);
-					SubsurfaceDMFPAlbedo	= DIFFUSE_SSSDMFPALBEDO(BSDF);
-					SubsurfaceDMFPRadius	= DIFFUSE_SSSDMFPRADIUS(BSDF);
+					SSS.bHasSubsurface = true;
+					SSS.bHasProfile	= BSDF_GETHASSSSPROFILE(BSDF);
+					SSS.Profile		= DIFFUSE_SSSPROFILEID(BSDF);
+					SSS.RadiusScale	= DIFFUSE_SSSPROFILERADIUS(BSDF);
+					SSS.BaseColor	= DIFFUSE_ALBEDO(BSDF);
+					SSS.DMFPAlbedo	= DIFFUSE_SSSDMFPALBEDO(BSDF);
+					SSS.DMFPRadius	= DIFFUSE_SSSDMFPRADIUS(BSDF);
 				}
 
 				LayerCoverage	+= BSDFCoverage;
@@ -1082,11 +1131,11 @@ void PackStrataOut(
 
 				if (bTopLayer)
 				{
-					AccumRoughness		+= DIELECTRIC_ROUGHNESS0(BSDF);
-					AccumRoughnessCount += 1.0f;
+					TopLayer.Roughness += DIELECTRIC_ROUGHNESS0(BSDF);
+					TopLayer.AccumRoughnessCount += 1.0f;
 
 					const float SafeRoughness = MakeRoughnessSafe(DIELECTRIC_ROUGHNESS0(BSDF));
-					AccumBaseColor		+= BSDFCoverage * EnvBRDFApprox(F0, SafeRoughness, saturate(ContextNoLight.NoV));
+					TopLayer.BaseColor += BSDFCoverage * EnvBRDFApprox(F0, SafeRoughness, saturate(ContextNoLight.NoV));
 				}
 
 				LayerCoverage	+= BSDFCoverage;
@@ -1100,11 +1149,11 @@ void PackStrataOut(
 
 				if (bTopLayer)
 				{
-					AccumRoughness		+= CONDUCTOR_ROUGHNESS0(BSDF);
-					AccumRoughnessCount += 1.0f;
+					TopLayer.Roughness += CONDUCTOR_ROUGHNESS0(BSDF);
+					TopLayer.AccumRoughnessCount += 1.0f;
 
 					const float SafeRoughness = MakeRoughnessSafe(CONDUCTOR_ROUGHNESS0(BSDF));
-					AccumBaseColor		+= BSDFCoverage * EnvBRDFApprox(CONDUCTOR_REFLECTIVITY(BSDF), CONDUCTOR_EDGECOLOR(BSDF), SafeRoughness, saturate(ContextNoLight.NoV));;
+					TopLayer.BaseColor += BSDFCoverage * EnvBRDFApprox(CONDUCTOR_REFLECTIVITY(BSDF), CONDUCTOR_EDGECOLOR(BSDF), SafeRoughness, saturate(ContextNoLight.NoV));;
 				}
 
 				LayerCoverage	+= BSDFCoverage;
@@ -1128,8 +1177,8 @@ void PackStrataOut(
 				
 				if (bTopLayer)
 				{
-					AccumRoughness		+= BSDFCoverageAvg * FullyRough; // Do not trigger SSR by using fully rough
-					AccumRoughnessCount += BSDFCoverageAvg;
+					TopLayer.Roughness += BSDFCoverageAvg * FullyRough; // Do not trigger SSR by using fully rough
+					TopLayer.AccumRoughnessCount += BSDFCoverageAvg;
 				}
 
 				LayerCoverage	+= BSDFCoverage;
@@ -1144,11 +1193,11 @@ void PackStrataOut(
 
 				if (bTopLayer)
 				{
-					AccumRoughness		+= SHEEN_ROUGHNESS(BSDF);
-					AccumRoughnessCount += 1.0f;
+					TopLayer.Roughness += SHEEN_ROUGHNESS(BSDF);
+					TopLayer.AccumRoughnessCount += 1.0f;
 
 					const float SafeRoughness = MakeRoughnessSafe(SHEEN_ROUGHNESS(BSDF));
-					AccumBaseColor      += BSDFCoverage * EnvBRDFApprox(F0, SafeRoughness, saturate(ContextNoLight.NoV));
+					TopLayer.BaseColor += BSDFCoverage * EnvBRDFApprox(F0, SafeRoughness, saturate(ContextNoLight.NoV));
 				}
 
 				LayerCoverage	+= BSDFCoverage;
@@ -1161,8 +1210,8 @@ void PackStrataOut(
 				// No SSR
 				if (bTopLayer)
 				{
-					AccumRoughness		+= BSDFCoverageAvg * FullyRough; // Do not trigger SSR by using fully rough
-					AccumRoughnessCount += BSDFCoverageAvg;
+					TopLayer.Roughness += BSDFCoverageAvg * FullyRough; // Do not trigger SSR by using fully rough
+					TopLayer.AccumRoughnessCount += BSDFCoverageAvg;
 				}
 				// Hair are considered opaque
 				LayerCoverage	+= BSDFCoverage;
@@ -1182,12 +1231,14 @@ void PackStrataOut(
 		// Now modify the gbuffer according to the top layer information
 		if (bTopLayer)
 		{
+			FinalizeTopLayerData(TopLayer);
+
 			// Normal is needed for some AO techniques
-			GBuffer.WorldNormal = normalize(AccumTopNormal);
+			GBuffer.WorldNormal = normalize(TopLayer.Normal);
 
 			// Roughness is needed for SSR, base color represents EnvBRDF
-			GBuffer.Roughness = AccumRoughnessCount > 0.0f ? saturate(AccumRoughness / AccumRoughnessCount) : FullyRough;
-			GBuffer.BaseColor = saturate(AccumBaseColor);
+			GBuffer.Roughness = TopLayer.Roughness;
+			GBuffer.BaseColor = TopLayer.BaseColor;
 			
 			// Approximative way to be able to have SSR produce colored specular reflections
 			GBuffer.Specular = 1;
@@ -1203,19 +1254,19 @@ void PackStrataOut(
 
 	// Patch Gbuffer data with SSS data if enabled
 	GBuffer.CustomData = 0;
-	if (bHasSubsurfaceScattering)
+	if (SSS.bHasSubsurface)
 	{
-		if (bHasSubsurfaceProfile)
+		if (SSS.bHasProfile)
 		{
-			GBuffer.BaseColor      = SubsurfaceBaseColor;
-			GBuffer.CustomData.rgb = EncodeSubsurfaceProfile(SubsurfaceProfile);
-			GBuffer.CustomData.a   = SubsurfaceRadiusScale;
+			GBuffer.BaseColor      = SSS.BaseColor;
+			GBuffer.CustomData.rgb = EncodeSubsurfaceProfile(SSS.Profile);
+			GBuffer.CustomData.a   = SSS.RadiusScale;
 		}
 		else
 		{
-			GBuffer.BaseColor      = SubsurfaceBaseColor;
-			GBuffer.CustomData.rgb = sqrt(SubsurfaceDMFPAlbedo);
-			GBuffer.CustomData.a   = saturate(SubsurfaceDMFPRadius / SSS_PROFILE_PERPIXEL_MAXRADIUS);
+			GBuffer.BaseColor      = SSS.BaseColor;
+			GBuffer.CustomData.rgb = sqrt(SSS.DMFPAlbedo);
+			GBuffer.CustomData.a   = saturate(SSS.DMFPRadius / SSS_PROFILE_PERPIXEL_MAXRADIUS);
 		}
 	}
 
@@ -1450,4 +1501,5 @@ FStrataBSDF UnpackStrataBSDFIn(ByteAddressBuffer StrataBuffer, inout uint ByteOf
 	}
 
 	return OutBSDF;
-}
\ No newline at end of file
+}
+
diff --git a/Engine/Shaders/Private/Strata/StrataMaterialClassification.usf b/Engine/Shaders/Private/Strata/StrataMaterialClassification.usf
--- a/Engine/Shaders/Private/Strata/StrataMaterialClassification.usf
+++ b/Engine/Shaders/Private/Strata/StrataMaterialClassification.usf
@@ -6,56 +6,188 @@
 #define STRATA_INLINE_SHADING 0
 #include "/Engine/Private/Strata/Strata.ush"
 
+// Note: 
+// This code is partial copy of PackStrataOut in Strata.ush. 
+// The command part will progressively be removed as we remove data from the gbuffer.
+void ComputePostProcessBSDFData(
+	const FStrataPixelHeader Header,
+	const FStrataBSDF BSDF,
+	float3 V,
+	inout FStrataClassification Classification,
+	inout FStrataSubsurfaceData SSS,
+	inout FStrataTopLayerData TopLayer)
+{
+	const float OpaqueBSDFThroughput = 0.0f;
+	const float FullThroughput = 1.0f;
+	const float FullyRough = 1.0f;
+
+	const float3 BSDFCoverage = BSDF.Weight;
+	const float  BSDFCoverageAvg = dot(BSDF.Weight,1.0f/3.0f);
+
+	// STRATA_TODO apply the same normal transform as for the material
+	float3x3 TangentBasis = StrataGetBSDFSharedBasis(Header, BSDF);
+	float3 N = TangentBasis[2];
+	//N = normalize(N);
+
+	TopLayer.Normal += N;
+
+	BxDFContext ContextNoLight = (BxDFContext)0;
+	const float3 FakeL = float3(0, 0, 1);
+	Init(ContextNoLight, N, V, FakeL);
+
+	const bool bTopLayer = BSDF_GETISTOPLAYER(BSDF);
+
+	Classification.ShadingModels |= StrataShadingModelBit(BSDF_GETTYPE(BSDF));
+
+	switch (BSDF_GETTYPE(BSDF))
+	{
+	case STRATA_BSDF_TYPE_DIFFUSE:
+	{
+		if (bTopLayer)
+		{
+			TopLayer.Roughness += BSDFCoverageAvg * FullyRough; // Do not trigger SSR by using fully rough
+			TopLayer.AccumRoughnessCount += BSDFCoverageAvg;
+		}
+
+		// Select the first/toppest layer with subsurface scattering
+		if (BSDF_GETHASSCATTERING(BSDF) && BSDFCoverageAvg > 0 && DIFFUSE_SSSPROFILERADIUS(BSDF) > 0 && !SSS.bHasSubsurface)
+		{
+			SSS.bHasSubsurface = true;
+			SSS.bHasProfile	= BSDF_GETHASSSSPROFILE(BSDF);
+			SSS.Profile		= DIFFUSE_SSSPROFILEID(BSDF);
+			SSS.RadiusScale	= DIFFUSE_SSSPROFILERADIUS(BSDF);
+			SSS.BaseColor	= DIFFUSE_ALBEDO(BSDF);
+			SSS.DMFPAlbedo	= DIFFUSE_SSSDMFPALBEDO(BSDF);
+			SSS.DMFPRadius	= DIFFUSE_SSSDMFPRADIUS(BSDF);
+		}
+		break;
+	}
+
+	case STRATA_BSDF_TYPE_DIELECTRIC:
+	{
+		if (BSDF_GETHASSCATTERING(BSDF))
+		{
+			Classification.bHasSubsurface = true;
+			if (BSDF_GETHASSSSPROFILE(BSDF))
+			{
+				Classification.bHasSubsurfaceProfile = true;
+			}
+		}
+
+		float F0 = DielectricIorToF0(DIELECTRIC_IOR(BSDF));
+		float3 F = F_Schlick(F0, ContextNoLight.VoH);
+
+		if (bTopLayer)
+		{
+			TopLayer.Roughness += DIELECTRIC_ROUGHNESS0(BSDF);
+			TopLayer.AccumRoughnessCount += 1.0f;
+
+			const float SafeRoughness = MakeRoughnessSafe(DIELECTRIC_ROUGHNESS0(BSDF));
+			TopLayer.BaseColor += BSDFCoverage * EnvBRDFApprox(F0, SafeRoughness, saturate(ContextNoLight.NoV));
+		}
+		break;
+	}
+
+	case STRATA_BSDF_TYPE_CONDUCTOR:
+	{
+		float3 F = F_Schlick(CONDUCTOR_REFLECTIVITY(BSDF), CONDUCTOR_EDGECOLOR(BSDF), ContextNoLight.VoH);
+
+		if (bTopLayer)
+		{
+			TopLayer.Roughness += CONDUCTOR_ROUGHNESS0(BSDF);
+			TopLayer.AccumRoughnessCount += 1.0f;
+
+			const float SafeRoughness = MakeRoughnessSafe(CONDUCTOR_ROUGHNESS0(BSDF));
+			TopLayer.BaseColor += BSDFCoverage * EnvBRDFApprox(CONDUCTOR_REFLECTIVITY(BSDF), CONDUCTOR_EDGECOLOR(BSDF), SafeRoughness, saturate(ContextNoLight.NoV));;
+		}
+		break;
+	}
+
+	case STRATA_BSDF_TYPE_VOLUME:
+	{
+		const float3 Albedo					= VOLUME_ALBEDO(BSDF);
+		const float3 ExtinctionCoefficients = VOLUME_EXTINCTION(BSDF);
+		const float3 ScatteringCoefficients = Albedo * ExtinctionCoefficients;
+		const float3 AbsorptionCoefficients = ExtinctionCoefficients - ScatteringCoefficients;
+		const float  Thickness				= VOLUME_THICKNESS(BSDF);
+
+		const float3 SafeExtinctionThreshold = 0.000001f;
+		const float3 SafeExtinctionCoefficients = max(SafeExtinctionThreshold, ExtinctionCoefficients);
+
+		const float PathLength = Thickness / max(0.0001f, abs(ContextNoLight.NoV));
+		const float3 SafePathSegmentTransmittance = exp(-SafeExtinctionCoefficients * PathLength);
+				
+		if (bTopLayer)
+		{
+			TopLayer.Roughness += BSDFCoverageAvg * FullyRough; // Do not trigger SSR by using fully rough
+			TopLayer.AccumRoughnessCount += BSDFCoverageAvg;
+		}
+		break;
+	}
+
+	case STRATA_BSDF_TYPE_SHEEN:
+	{
+		const float3 F0  = SHEEN_ALBEDO(BSDF);
+		const float3 F   = F_Schlick(F0, ContextNoLight.VoH);
+
+		if (bTopLayer)
+		{
+			TopLayer.Roughness += SHEEN_ROUGHNESS(BSDF);
+			TopLayer.AccumRoughnessCount += 1.0f;
+
+			const float SafeRoughness = MakeRoughnessSafe(SHEEN_ROUGHNESS(BSDF));
+			TopLayer.BaseColor += BSDFCoverage * EnvBRDFApprox(F0, SafeRoughness, saturate(ContextNoLight.NoV));
+		}
+		break;
+	}
+
+	case STRATA_BSDF_TYPE_HAIR:
+	{
+		Classification.bHasTransmission = true;
+
+		// No SSR
+		if (bTopLayer)
+		{
+			TopLayer.Roughness += BSDFCoverageAvg * FullyRough; // Do not trigger SSR by using fully rough
+			TopLayer.AccumRoughnessCount += BSDFCoverageAvg;
+		}
+		break;
+	}
+
+	} // switch		
+}
+
 void MainPS(
 	float4 SVPos : SV_POSITION, 
-	out uint OutColor0 : SV_Target0)
+	out uint OutColor0 : SV_Target0,
+	out uint2 OutColor1 : SV_Target1)
 {
 	const uint2 PixelPos = uint2(SVPos.xy);
 
 	// Sample scene textures.
-	float2 BufferUV = SvPositionToBufferUV(float4(PixelPos, SVPos.zw));
+	float2 BufferUV		 = SvPositionToBufferUV(float4(PixelPos, SVPos.zw));
 	FGBufferData GBuffer = GetGBufferDataFromSceneTextures(BufferUV);
 	bool bStrataMaterial = GBuffer.ShadingModelID == SHADINGMODELID_STRATA;
 	
-	FStrataClassification Out = (FStrataClassification)0;
+	FStrataClassification ClassificationData	= (FStrataClassification)0;
+	FStrataSubsurfaceData SSSData				= (FStrataSubsurfaceData)0;
+	FStrataTopLayerData TopLayerData			= (FStrataTopLayerData)0;
+
 	if (bStrataMaterial)
 	{
-		Out.bIsStrataMaterial = true;
+		const float3 V = float3(0,0,1); // TODO
+
 		uint PixelStrataDataByteOffset = GetStrataPixelDataByteOffset(PixelPos, View.BufferSizeAndInvSize.x, Strata.MaxBytesPerPixel);
 		FStrataPixelHeader Header = UnpackStrataHeaderIn(Strata.MaterialLobesBuffer, PixelStrataDataByteOffset);
 		for (uint BSDFIndex = 0; BSDFIndex < Header.BSDFCount; ++BSDFIndex)
 		{
 			const FStrataBSDF BSDF = UnpackStrataBSDFIn(Strata.MaterialLobesBuffer, PixelStrataDataByteOffset);
-			Out.ShadingModels |= StrataShadingModelBit(BSDF_GETTYPE(BSDF));
-
-			switch (BSDF_GETTYPE(BSDF))
-			{
-			case STRATA_BSDF_TYPE_DIFFUSE:
-				if (BSDF_GETHASSCATTERING(BSDF))
-				{
-					Out.bHasSubsurface = true;
-					if (BSDF_GETHASSSSPROFILE(BSDF))
-					{
-						Out.bHasSubsurfaceProfile = true;
-					}
-				}
-				break;
-			case STRATA_BSDF_TYPE_HAIR:
-				Out.bHasTransmission = true;
-				break;
-			case STRATA_BSDF_TYPE_DIELECTRIC:
-				break;
-			case STRATA_BSDF_TYPE_CONDUCTOR:
-				break;
-			case STRATA_BSDF_TYPE_VOLUME:	
-				break;
-			case STRATA_BSDF_TYPE_SHEEN:	
-				break;
-			}
+			ComputePostProcessBSDFData(Header, BSDF, V, ClassificationData, SSSData, TopLayerData);
 		}
+
+		FinalizeTopLayerData(TopLayerData);
 	}
 
-	OutColor0 = StrataPackClassificationData(Out);
+	OutColor0 = StrataPackClassificationData(ClassificationData);
+	OutColor1 = StrataPackTopLayerData(TopLayerData);
 }
-
-
diff --git a/Engine/Source/Runtime/Renderer/Private/Strata/Strata.cpp b/Engine/Source/Runtime/Renderer/Private/Strata/Strata.cpp
--- a/Engine/Source/Runtime/Renderer/Private/Strata/Strata.cpp
+++ b/Engine/Source/Runtime/Renderer/Private/Strata/Strata.cpp
@@ -57,10 +57,19 @@ void InitialiseStrataFrameSceneData(FSceneRenderer& SceneRenderer, FRDGBuilder&
 		const uint32 RoundToValue = 4u;
 		StrataSceneData.MaxBytesPerPixel = FMath::DivideAndRoundUp(MaterialConservativeByteCountPerPixel, RoundToValue) * RoundToValue;
 
-		FRDGTextureDesc Desc = FRDGTextureDesc::Create2D(BufferSizeXY, PF_R32_UINT, FClearValueBinding::Black, TexCreate_ShaderResource | TexCreate_RenderTargetable);
-		FRDGTextureRef Texture = GraphBuilder.CreateTexture(Desc, TEXT("StrataClassificationTexture"));
-		GraphBuilder.PreallocateTexture(Texture);
-		StrataSceneData.ClassificationTexture = GraphBuilder.GetPooledTexture(Texture);
+		// Classification texture
+		{
+			FRDGTextureRef Texture = GraphBuilder.CreateTexture(FRDGTextureDesc::Create2D(BufferSizeXY, PF_R32_UINT, FClearValueBinding::Black, TexCreate_ShaderResource | TexCreate_RenderTargetable), TEXT("StrataClassificationTexture"));
+			GraphBuilder.PreallocateTexture(Texture);
+			StrataSceneData.ClassificationTexture = GraphBuilder.GetPooledTexture(Texture);
+		}
+
+		// Top layer texture
+		{
+			FRDGTextureRef Texture = GraphBuilder.CreateTexture(FRDGTextureDesc::Create2D(BufferSizeXY, PF_R32G32_UINT, FClearValueBinding::Black, TexCreate_ShaderResource | TexCreate_RenderTargetable), TEXT("StrataTopLayerTexture"));
+			GraphBuilder.PreallocateTexture(Texture);
+			StrataSceneData.TopLayerTexture = GraphBuilder.GetPooledTexture(Texture);
+		}
 	}
 	else
 	{
@@ -116,6 +125,7 @@ TUniformBufferRef<FStrataGlobalUniformParameters> BindStrataGlobalUniformParamet
 		StrataUniformParameters.MaxBytesPerPixel = View.StrataSceneData->MaxBytesPerPixel;
 		StrataUniformParameters.MaterialLobesBuffer = View.StrataSceneData->MaterialLobesBuffer.SRV;
 		StrataUniformParameters.ClassificationTexture = View.StrataSceneData->ClassificationTexture->GetRenderTargetItem().ShaderResourceTexture;
+		StrataUniformParameters.TopLayerTexture = View.StrataSceneData->TopLayerTexture->GetRenderTargetItem().ShaderResourceTexture;
 
 		View.StrataSceneData->StrataGlobalUniformParameters = CreateUniformBufferImmediate(StrataUniformParameters, UniformBuffer_SingleFrame);
 		return View.StrataSceneData->StrataGlobalUniformParameters;
@@ -126,6 +136,7 @@ TUniformBufferRef<FStrataGlobalUniformParameters> BindStrataGlobalUniformParamet
 		StrataUniformParameters.MaxBytesPerPixel = 0;
 		StrataUniformParameters.MaterialLobesBuffer = GEmptyVertexBufferWithUAV->ShaderResourceViewRHI;
 		StrataUniformParameters.ClassificationTexture = GSystemTextures.BlackDummy->GetRenderTargetItem().ShaderResourceTexture;
+		StrataUniformParameters.TopLayerTexture = GSystemTextures.BlackDummy->GetRenderTargetItem().ShaderResourceTexture;
 		return CreateUniformBufferImmediate(StrataUniformParameters, UniformBuffer_SingleDraw);
 	}
 }
@@ -262,6 +273,7 @@ void AddStrataMaterialClassificationPass(FRDGBuilder& GraphBuilder, const TArray
 		PassParameters->Strata = Strata::BindStrataGlobalUniformParameters(View);
 		PassParameters->SceneTextures = GetSceneTextureParameters(GraphBuilder);		
 		PassParameters->RenderTargets[0] = FRenderTargetBinding(GraphBuilder.RegisterExternalTexture(View.StrataSceneData->ClassificationTexture), ERenderTargetLoadAction::EClear);
+		PassParameters->RenderTargets[1] = FRenderTargetBinding(GraphBuilder.RegisterExternalTexture(View.StrataSceneData->TopLayerTexture), ERenderTargetLoadAction::EClear);
 
 		if (ShaderDrawDebug::IsShaderDrawDebugEnabled())
 		{
diff --git a/Engine/Source/Runtime/Renderer/Private/Strata/Strata.h b/Engine/Source/Runtime/Renderer/Private/Strata/Strata.h
--- a/Engine/Source/Runtime/Renderer/Private/Strata/Strata.h
+++ b/Engine/Source/Runtime/Renderer/Private/Strata/Strata.h
@@ -24,6 +24,7 @@ BEGIN_GLOBAL_SHADER_PARAMETER_STRUCT(FStrataGlobalUniformParameters, )
 	SHADER_PARAMETER(uint32, MaxBytesPerPixel)
 	SHADER_PARAMETER_SRV(ByteAddressBuffer, MaterialLobesBuffer)
 	SHADER_PARAMETER_TEXTURE(Texture2D<uint>, ClassificationTexture)
+	SHADER_PARAMETER_TEXTURE(Texture2D<uint2>, TopLayerTexture)
 END_GLOBAL_SHADER_PARAMETER_STRUCT()
 
 struct FStrataSceneData
@@ -31,6 +32,7 @@ struct FStrataSceneData
 	uint32 MaxBytesPerPixel;
 	FRWByteAddressBuffer MaterialLobesBuffer;				// This should be a RDG resource
 	TRefCountPtr<IPooledRenderTarget> ClassificationTexture;// This should be a RDG resource
+	TRefCountPtr<IPooledRenderTarget> TopLayerTexture;		// This should be a RDG resource
 
 	TUniformBufferRef<FStrataGlobalUniformParameters> StrataGlobalUniformParameters;
 
