Fix back-faces sometimes being marked as lit in the shadow mask.

diff --git a/Engine/Shaders/Private/VirtualShadowMaps/VirtualShadowMapProjection.usf b/Engine/Shaders/Private/VirtualShadowMaps/VirtualShadowMapProjection.usf
--- a/Engine/Shaders/Private/VirtualShadowMaps/VirtualShadowMapProjection.usf
+++ b/Engine/Shaders/Private/VirtualShadowMaps/VirtualShadowMapProjection.usf
@@ -111,14 +111,6 @@ float ProjectSingleSample(
 	FVirtualShadowMapSampleResult VirtualShadowMapSample =
 		SampleVirtualShadowMap(VirtualShadowMapId, WorldPosition, EstimatedGeoWorldNormal, ContactShadowLengthWorld);
 
-#if OUTPUT_TYPE != 2		// Don't kill pixels in debug output mode
-	if (!VirtualShadowMapSample.bValid)
-	{
-		// Generally shouldn't occur if PT allocation is conservative
-		discard;
-	}
-#endif
-
 	float RayMissFactor = VirtualShadowMapSample.bOccluded ? 0.0f : 1.0f;
 	float HitDistance = VirtualShadowMapSample.OccluderDistance;
 
@@ -213,7 +205,7 @@ void VirtualShadowMapProjectionSpotPS(
 	float3 DebugOutput = float3(0, 0, 0);
 	int SampleCount = 1;
 	float HitDistance = 100000;	// TODO
-	float RayMissFactor = 1.0f;
+	float RayMissFactor = 0.0f;
 
 	if (bInLightRadius && bInLightCone)
 	{
@@ -263,10 +255,6 @@ void VirtualShadowMapProjectionSpotPS(
 			}
 		}
 	}
-	else
-	{
-		//discard;
-	}
 
 	float4 OutputSignal = float4(0, 0, 0, 0);
 	if (RayMissFactor < 1.0f)
@@ -304,7 +292,7 @@ void VirtualShadowMapProjectionDirectionalPS(
 	float3 DebugOutput = float3(0, 0, 0);
 	int SampleCount = 1;
 	float HitDistance = 100000;	// TODO
-	float RayMissFactor = 1.0f;
+	float RayMissFactor = 0.0f;
 
 	if (bFrontFace)
 	{
