Fixes for Strata

diff --git a/Engine/Shaders/Private/Strata/Strata.ush b/Engine/Shaders/Private/Strata/Strata.ush
--- a/Engine/Shaders/Private/Strata/Strata.ush
+++ b/Engine/Shaders/Private/Strata/Strata.ush
@@ -43,7 +43,7 @@ struct FStrataBSDF
 #define STATE_BITS_HASSCATTERING		1
 #define STATE_BITS_HASSSSPROFILE		1
 #define STATE_BITS_GREYWEIGHT			1
-#define STATE_BITS___UNUSED___			13
+#define STATE_BITS___UNUSED___			12
 #define STATE_BITS_WEIGHT10F			10
 
 // Bit offset of each component store in BSDF state 
@@ -821,10 +821,10 @@ uint PackRGBA8(float4 rgba)
 float4 UnpackRGBA8(uint rgba)
 {
 	float4 Out;
-	Out.r = float((rgba >> 0)  & 0xFF);
-	Out.g = float((rgba >> 8)  & 0xFF);
-	Out.b = float((rgba >> 16) & 0xFF);
-	Out.a = float((rgba >> 24) & 0xFF);
+	Out.r = float((rgba >> 0)  & 0xFF) * (1.0f/255.0f);
+	Out.g = float((rgba >> 8)  & 0xFF) * (1.0f/255.0f);
+	Out.b = float((rgba >> 16) & 0xFF) * (1.0f/255.0f);
+	Out.a = float((rgba >> 24) & 0xFF) * (1.0f/255.0f);
 	return Out;
 }
 
diff --git a/Engine/Shaders/Private/Strata/StrataMaterialClassification.usf b/Engine/Shaders/Private/Strata/StrataMaterialClassification.usf
--- a/Engine/Shaders/Private/Strata/StrataMaterialClassification.usf
+++ b/Engine/Shaders/Private/Strata/StrataMaterialClassification.usf
@@ -167,13 +167,13 @@ void MainPS(
 	// Sample scene textures.
 	float2 BufferUV		 = SvPositionToBufferUV(float4(PixelPos, SVPos.zw));
 	FGBufferData GBuffer = GetGBufferDataFromSceneTextures(BufferUV);
-	bool bStrataMaterial = GBuffer.ShadingModelID == SHADINGMODELID_STRATA;
 	
 	FStrataClassification ClassificationData	= (FStrataClassification)0;
 	FStrataSubsurfaceData SSSData				= (FStrataSubsurfaceData)0;
 	FStrataTopLayerData TopLayerData			= (FStrataTopLayerData)0;
 
-	if (bStrataMaterial)
+	ClassificationData.bIsStrataMaterial = GBuffer.ShadingModelID == SHADINGMODELID_STRATA;
+	if (ClassificationData.bIsStrataMaterial)
 	{
 		const float3 V = float3(0,0,1); // TODO
 
