Strata support for single layer water.
A few renaming

diff --git a/Engine/Shaders/Private/BasePassPixelShader.usf b/Engine/Shaders/Private/BasePassPixelShader.usf
--- a/Engine/Shaders/Private/BasePassPixelShader.usf
+++ b/Engine/Shaders/Private/BasePassPixelShader.usf
@@ -957,6 +957,35 @@ void FPixelShaderInOut_MainPS(
 	}
 #endif
 
+#if STRATA_ENABLED && !MATERIAL_SHADINGMODEL_UNLIT
+	FStrataData Strata = PixelMaterialInputs.FrontMaterial;
+
+	// Initialise a Strata header with normal in registers
+	FStrataPixelHeader StrataPixelHeader = InitialiseStrataPixelHeader();
+	UNROLL
+		for (int l = 0; l < Strata.LayerCount; ++l)
+		{
+			StrataPixelHeader.BSDFCount += Strata.Layers[l].BSDFCount;
+		}
+	StrataPixelHeader.SharedNormals = MaterialParameters.SharedNormals;
+#endif
+
+#if STRATA_ENABLED && STRATA_INLINE_SINGLELAYERWATER
+	// Override GBuffer data with Strata SLW water BSDF
+	SanitizeStrataSingleLayerWater(PixelMaterialInputs.FrontMaterial.Layers[0].BSDFs[0]);
+	FStrataBSDF SLWBSDF = PixelMaterialInputs.FrontMaterial.Layers[0].BSDFs[0];
+
+	BaseColor = SLW_BASECOLOR(SLWBSDF);
+	Metallic = SLW_METALLIC(SLWBSDF);
+	Specular = SLW_SPECULAR(SLWBSDF);
+	Roughness = SLW_ROUGHNESS(SLWBSDF);
+	Opacity = SLW_TOPMATERIALOPACITY(SLWBSDF);
+	MaterialParameters.WorldNormal = StrataPixelHeader.SharedNormals.Normals[BSDF_GETNORMALID(SLWBSDF)];
+	MaterialParameters.TangentToWorld[2] = MaterialParameters.WorldNormal;
+	Anisotropy = 0.0f;
+	ShadingModel = MATERIAL_SHADINGMODEL_SINGLELAYERWATER;
+#endif
+
 	float DBufferOpacity = 1.0f;
 #if USE_DBUFFER && MATERIALDECALRESPONSEMASK && !MATERIALBLENDING_ANY_TRANSLUCENT && !MATERIAL_SHADINGMODEL_SINGLELAYERWATER
 	// apply decals from the DBuffer
@@ -1462,26 +1491,18 @@ void FPixelShaderInOut_MainPS(
 	bool bSubsurfaceEnable = true;
 #if STRATA_ENABLED && !MATERIAL_SHADINGMODEL_UNLIT
 
-	FStrataData Strata = PixelMaterialInputs.FrontMaterial;
-
-	// Initialise a Strata header with normal in registers
-	FStrataPixelHeader StrataPixelHeader = InitialiseStrataPixelHeader();
-	UNROLL
-	for (int l = 0; l < Strata.LayerCount; ++l)
-	{
-		StrataPixelHeader.BSDFCount += Strata.Layers[l].BSDFCount;
-	}
-	StrataPixelHeader.SharedNormals = MaterialParameters.SharedNormals;
-
 #if STRATA_OPAQUE_ENABLED
 
-	// Need to cancel that to make sure strata material are only lit using Strata lighting passes
+#if STRATA_INLINE_SINGLELAYERWATER==0
+	// Need to reset color to make sure strata material are only lit using Strata lighting passes.
+	// Except for water doing some specialised and simplified lighting during the base pass
 	Color = 0;
+#endif
 	// Clear only what should be and keep the rest
-	GBuffer.Roughness = 0;
 	GBuffer.DiffuseColor = 0;
 	GBuffer.SpecularColor = 0;
 	GBuffer.BaseColor = 0;
+	GBuffer.Roughness = 0;
 	GBuffer.Specular = 0;
 	GBuffer.Metallic = 0;
 	GBuffer.ShadingModelID = SHADINGMODELID_STRATA;
diff --git a/Engine/Shaders/Private/SingleLayerWaterComposite.usf b/Engine/Shaders/Private/SingleLayerWaterComposite.usf
--- a/Engine/Shaders/Private/SingleLayerWaterComposite.usf
+++ b/Engine/Shaders/Private/SingleLayerWaterComposite.usf
@@ -22,6 +22,10 @@
 #define REFLECTION_COMPOSITE_SUPPORT_SKYLIGHT_BLEND 0
 #include "ReflectionEnvironmentComposite.ush"
 
+#if STRATA_ENABLED
+#define STRATA_INLINE_SHADING 0
+#include "Strata/StrataEvaluation.ush"
+#endif
 
 Texture2D ScreenSpaceReflectionsTexture;
 SamplerState ScreenSpaceReflectionsSampler;
@@ -37,7 +41,7 @@ struct SingleLayerWaterCompositeOutput
 	float  Clip;
 };
 
-SingleLayerWaterCompositeOutput SingleLayerWaterComposite(float2 BufferUV, float2 ScreenPosition)
+SingleLayerWaterCompositeOutput SingleLayerWaterComposite(float2 BufferUV, float2 ScreenPosition, float4 SvPosition)
 {
 	// Sample scene textures.
 	FGBufferData GBuffer = GetGBufferDataFromSceneTextures(BufferUV);
@@ -58,7 +62,17 @@ SingleLayerWaterCompositeOutput SingleLayerWaterComposite(float2 BufferUV, float
 #else
 	const float PreExposure = 1.f;
 #endif
-	if (ShadingModelID == SHADINGMODELID_SINGLELAYERWATER)
+
+	bool bIsStrataSLW = false;
+#if STRATA_ENABLED
+	const FStrataClassification StrataData = StrataUnpackClassificationData(Strata.ClassificationTexture.Load(uint3(SvPosition.xy, 0)));
+	if (StrataData.bIsStrataMaterial)
+	{
+		bIsStrataSLW = StrataHasShadingModel(StrataData.ShadingModels, STRATA_BSDF_TYPE_SINGLELAYERWATER);
+	}
+#endif
+
+	if (ShadingModelID == SHADINGMODELID_SINGLELAYERWATER || bIsStrataSLW)
 	{
 		const float2 SceneNoWaterUV = clamp(BufferUV, SceneNoWaterMinMaxUV.xy, SceneNoWaterMinMaxUV.zw);
 
@@ -107,7 +121,25 @@ SingleLayerWaterCompositeOutput SingleLayerWaterComposite(float2 BufferUV, float
 
 			// Apply the BRDF reflection integral to reflection. 
 			// BRDF/Fresnel is already applied for the under water part in the scene.
-			float3 EnvBrdf = EnvBRDF(GBuffer.SpecularColor, GBuffer.Roughness, NoV);
+			float3 EnvBrdf = 0.0;
+#if STRATA_ENABLED
+			if (bIsStrataSLW)
+			{
+				// Water is using the higher quality EnvBRDF instead of EnvBRDFApprox using during the base pass.
+				// So we load the water BSDF here to make sure we compute the accurate reflection for water.
+				// STRATA_TODO use EnvBRDF and remove EnvBRDFApprox, or use a better approximation...
+				uint PixelStrataDataByteOffset = GetStrataPixelDataByteOffset(SvPosition.xy, View.BufferSizeAndInvSize.x, Strata.MaxBytesPerPixel);
+				FStrataPixelHeader StrataPixelHeader = UnpackStrataHeaderIn(Strata.MaterialLobesBuffer, PixelStrataDataByteOffset);
+				const FStrataBSDF BSDF = UnpackStrataBSDFIn(Strata.MaterialLobesBuffer, PixelStrataDataByteOffset);
+
+				const float3 F0 = ComputeF0(SLW_SPECULAR(BSDF), SLW_BASECOLOR(BSDF), SLW_METALLIC(BSDF));
+				EnvBrdf = EnvBRDF(F0, SLW_ROUGHNESS(BSDF), NoV);
+			}
+			else
+#endif
+			{
+				EnvBrdf = EnvBRDF(GBuffer.SpecularColor, GBuffer.Roughness, NoV);
+			}
 			Reflection *= EnvBrdf;
 
 			// Soft fading near shor to avoid seeing triangles.
@@ -181,6 +213,14 @@ void WaterTileCatergorisationCS(uint3 ThreadId : SV_DispatchThreadID, uint3 Grou
 				FGBufferData GBuffer = GetGBufferDataFromSceneTextures(BufferUV);
 
 				bContainsWater = bContainsWater || GBuffer.ShadingModelID == SHADINGMODELID_SINGLELAYERWATER;
+
+			#if STRATA_ENABLED
+				const FStrataClassification StrataData = StrataUnpackClassificationData(Strata.ClassificationTexture.Load(uint3(PixelPos,0)));
+				if (StrataData.bIsStrataMaterial)
+				{
+					bContainsWater = bContainsWater || StrataHasShadingModel(StrataData.ShadingModels, STRATA_BSDF_TYPE_SINGLELAYERWATER);
+				}
+			#endif
 			}
 		}
 		if (bContainsWater)
@@ -200,6 +240,14 @@ void WaterTileCatergorisationCS(uint3 ThreadId : SV_DispatchThreadID, uint3 Grou
 		float2 BufferUV = (PixelPos + 0.5f) * View.BufferSizeAndInvSize.zw;
 		FGBufferData GBuffer = GetGBufferDataFromSceneTextures(BufferUV);
 		bContainsWater = GBuffer.ShadingModelID == SHADINGMODELID_SINGLELAYERWATER;
+
+	#if STRATA_ENABLED
+		const FStrataClassification StrataData = StrataUnpackClassificationData(Strata.ClassificationTexture.Load(uint3(PixelPos, 0)));
+		if (StrataData.bIsStrataMaterial)
+		{
+			bContainsWater = bContainsWater || StrataHasShadingModel(StrataData.ShadingModels, STRATA_BSDF_TYPE_SINGLELAYERWATER);
+		}
+	#endif
 	}
 
 	const int LinearIndex = GroupThreadId.y * WORK_TILE_SIZE + GroupThreadId.x;
@@ -280,7 +328,7 @@ void SingleLayerWaterCompositePS(
 	float2 ScreenPosition = SvPositionToScreenPosition(SvPosition).xy;
 
 	// TODO use dual source blending to apply TransmittanceToScene
-	SingleLayerWaterCompositeOutput Result = SingleLayerWaterComposite(BufferUV, ScreenPosition);
+	SingleLayerWaterCompositeOutput Result = SingleLayerWaterComposite(BufferUV, ScreenPosition, SvPosition);
 
 	clip(Result.Clip); // Since this shader does not write to depth or stencil it should still benefit from EarlyZ (See AMD depth-in-depth documentation)
 	OutColor = Result.LuminanceTransmittance;
diff --git a/Engine/Shaders/Private/SingleLayerWaterShading.ush b/Engine/Shaders/Private/SingleLayerWaterShading.ush
--- a/Engine/Shaders/Private/SingleLayerWaterShading.ush
+++ b/Engine/Shaders/Private/SingleLayerWaterShading.ush
@@ -110,16 +110,27 @@ WaterVolumeLightingOutput EvaluateWaterVolumeLighting(
 
 		const float BehindWaterDeltaDepth = CameraIsUnderWater ? WaterDepth : max(0.0f, SceneDepth - WaterDepth);
 
+#if STRATA_INLINE_SINGLELAYERWATER
+		FStrataBSDF SLWBSDF = PixelMaterialInputs.FrontMaterial.Layers[0].BSDFs[0];
+		const float3 Albedo = SLW_WATERALBEDO(SLWBSDF);
+		const float3 ExtinctionCoeff = SLW_WATEREXTINCTION(SLWBSDF);
+		const float3 ScatteringCoeff = Albedo * ExtinctionCoeff;
+		const float3 AbsorptionCoeff = max(0.0, ExtinctionCoeff - ScatteringCoeff);
+		const float PhaseG = SLW_WATERPHASEG(SLWBSDF);
+		//Sample the optional Material Input ColorScaleBehindWater and fade it out at shorelines to avoid hard edge intersections
+		float3 ColorScaleBehindWater = lerp(1.0f, SLW_COLORSCALEBEHINDWATER(SLWBSDF), saturate(BehindWaterDeltaDepth * 0.02f));
+#else
 		const float3 ScatteringCoeff = max(0.0f, GetSingleLayerWaterMaterialOutput0(MaterialParameters));
 		const float3 AbsorptionCoeff = max(0.0f, GetSingleLayerWaterMaterialOutput1(MaterialParameters));
 		const float PhaseG = clamp(GetSingleLayerWaterMaterialOutput2(MaterialParameters), -1.0f, 1.0f);
 		//Sample the optional Material Input ColorScaleBehindWater and fade it out at shorelines to avoid hard edge intersections
 		float3 ColorScaleBehindWater = lerp(1.0f, max(0.0f, GetSingleLayerWaterMaterialOutput3(MaterialParameters)), saturate(BehindWaterDeltaDepth * 0.02f));
+		const float3 ExtinctionCoeff = ScatteringCoeff + AbsorptionCoeff;
+#endif
 		
-		const float3 ExtinctionCoeff = ScatteringCoeff + AbsorptionCoeff; 
 		// Max to avoid division by 0 with the analytical integral below.
 		// 1e-5 is high enough to avoid denorms on mobile
-		const float3 ExtinctionCoeffSafe = max(ScatteringCoeff + AbsorptionCoeff, 1e-5); 
+		const float3 ExtinctionCoeffSafe = max(ExtinctionCoeff, 1e-5);
 		
 		float DirLightPhaseValue = 0.0f; // Default when Total Internal Reflection happens.
 		{
diff --git a/Engine/Shaders/Private/Strata/Strata.ush b/Engine/Shaders/Private/Strata/Strata.ush
--- a/Engine/Shaders/Private/Strata/Strata.ush
+++ b/Engine/Shaders/Private/Strata/Strata.ush
@@ -19,7 +19,11 @@
 #define STRATA_CLAMPED_LAYER_COUNT 1
 #endif
 
+#ifndef MATERIAL_SHADINGMODEL_SINGLELAYERWATER
+#define MATERIAL_SHADINGMODEL_SINGLELAYERWATER 0 
+#endif
 
+#define STRATA_INLINE_SINGLELAYERWATER (PROJECT_STRATA && MATERIAL_IS_STRATA && STRATA_INLINE_SHADING && MATERIAL_SHADINGMODEL_SINGLELAYERWATER)
 
 ///////////////////////////////////////////////////////////////////////////////
 // BSDFs representation
@@ -32,18 +36,24 @@ struct FStrataBSDF
 	float3				Emissive;	// All bsdf can potientially have emissive. This is never written out to the strata gbuffer but instead accumulated in the emissive buffer during the base pass.
 #endif
 
+	// Water has too many parameters to be stores in the extra VGPRs below. So we extend this structure for the case where it is evaluated inline only.
+	// This is fine because when water is used, it is the only BSDF that can be used and it is only used to do in line shading.
+#if STRATA_INLINE_SINGLELAYERWATER
+	float4 InlineVGPRs[3];
+#endif
+
 	float4 VGPRs[2];
 };
 
 // Num bits for each component store in BSDF state
-#define STATE_BITS_BSDF					3
+#define STATE_BITS_BSDF					4
 #define STATE_BITS_SHAREDNORMALID		2
 #define STATE_BITS_ANISOTROPY			1
 #define STATE_BITS_TOPLAYER				1
 #define STATE_BITS_HASSCATTERING		1
 #define STATE_BITS_HASSSSPROFILE		1
 #define STATE_BITS_GREYWEIGHT			1
-#define STATE_BITS___UNUSED___			12
+#define STATE_BITS___UNUSED___			11
 #define STATE_BITS_WEIGHT10F			10
 
 // Bit offset of each component store in BSDF state 
@@ -141,6 +151,20 @@ struct FStrataBSDF
 #define HAIR_SPECULAR(X)				X.VGPRs[1].y
 #define HAIR_BACKLIT(X)					X.VGPRs[1].z
 
+#define SLW_BASECOLOR(X)				X.VGPRs[0].xyz
+#define SLW_METALLIC(X)					X.VGPRs[1].x
+#define SLW_SPECULAR(X)					X.VGPRs[1].y
+#define SLW_ROUGHNESS(X)				X.VGPRs[1].z
+#define SLW_TOPMATERIALOPACITY(X)		X.VGPRs[1].w
+#if STRATA_INLINE_SINGLELAYERWATER
+#define SLW_WATERALBEDO(X)				X.InlineVGPRs[0].xyz
+#define SLW_WATEREXTINCTION(X)			X.InlineVGPRs[1].xyz
+#define SLW_WATERPHASEG(X)				X.InlineVGPRs[0].w
+#define SLW_COLORSCALEBEHINDWATER(X)	X.InlineVGPRs[2].xyz
+#endif
+
+
+
 void StrataRequestSharedNormalTangentBasis(inout uint NormalTypes, uint Index);
 
 uint   PackR11G11B10F(float3 rgb);
@@ -161,22 +185,22 @@ struct FStrataClassification
 
 uint StrataPackClassificationData(FStrataClassification In)
 {
-	return
-		(In.ShadingModels & 0xFF) |
-		((In.bIsStrataMaterial ? 1 : 0) << 8) |
-		((In.bHasSubsurface ? 1 : 0) << 9) |
-		((In.bHasSubsurfaceProfile ? 1 : 0) << 10) |
-		((In.bHasTransmission ? 1 : 0) << 11);
+	return 
+		(In.ShadingModels & 0xFFFF) | 
+		((In.bIsStrataMaterial		? 1 : 0) << 16) |
+		((In.bHasSubsurface			? 1 : 0) << 17) | 
+		((In.bHasSubsurfaceProfile	? 1 : 0) << 18)|
+		((In.bHasTransmission		? 1 : 0) << 19);
 }
 
 FStrataClassification StrataUnpackClassificationData(uint In)
 {
 	FStrataClassification Out;
-	Out.ShadingModels = (In & 0xFF);
-	Out.bIsStrataMaterial = (In >> 8) & 0x1;
-	Out.bHasSubsurface = (In >> 9) & 0x1;
-	Out.bHasSubsurfaceProfile = (In >> 10) & 0x1;
-	Out.bHasTransmission = (In >> 11) & 0x1;
+	Out.ShadingModels			= (In & 0xFFFF);
+	Out.bIsStrataMaterial		= (In >> 16) & 0x1;
+	Out.bHasSubsurface			= (In >> 17) & 0x1;
+	Out.bHasSubsurfaceProfile	= (In >> 18) & 0x1;
+	Out.bHasTransmission		= (In >> 19) & 0x1;
 	return Out;
 }
 
@@ -424,32 +448,32 @@ FStrataData GetStrataVolumeFogCloudBSDF(float3 Albedo, float3 Extinction, float3
 	StrataData.Layers[0].BSDFCount++;
 
 	BSDF_SETTYPE					(StrataData.Layers[0].BSDFs[0], STRATA_BSDF_TYPE_VOLUMETRICFOGCLOUD);
+	BSDF_SETEMISSIVE				(StrataData.Layers[0].BSDFs[0], Emissive);
 	VOLUMETRICFOGCLOUD_ALBEDO		(StrataData.Layers[0].BSDFs[0]) = Albedo;
 	VOLUMETRICFOGCLOUD_EXTINCTION	(StrataData.Layers[0].BSDFs[0]) = Extinction;
 	VOLUMETRICFOGCLOUD_AO			(StrataData.Layers[0].BSDFs[0]) = AmbientOcclusion;
-	BSDF_SETEMISSIVE				(StrataData.Layers[0].BSDFs[0], Emissive);
 
 	StrataData.Layers[0].BSDFs[0].Weight = 1.0f;
 
 	return StrataData;
 }
 
-FStrataData GetStrataUnlitBSDF(float3 EmissiveColor, float3 TransmittanceColor)
+FStrataData GetStrataUnlitBSDF(float3 Emissive, float3 TransmittanceColor)
 {
 	FStrataData StrataData = GetInitialisedStrataData();
 	StrataData.LayerCount++;
 	StrataData.Layers[0].BSDFCount++;
 
 	BSDF_SETTYPE					(StrataData.Layers[0].BSDFs[0], STRATA_BSDF_TYPE_UNLIT);
+	BSDF_SETEMISSIVE				(StrataData.Layers[0].BSDFs[0], Emissive);
 	UNLIT_TRANSMITTANCE				(StrataData.Layers[0].BSDFs[0]) = TransmittanceColor;
-	BSDF_SETEMISSIVE				(StrataData.Layers[0].BSDFs[0], EmissiveColor);
 
 	StrataData.Layers[0].BSDFs[0].Weight = 1.0f;
 
 	return StrataData;
 }
 
-FStrataData GetStrataHairBSDF(float3 BaseColor, float Scatter, float Specular, float Roughness, float Backlit, float3 EmissiveColor, uint SharedNormalIndex)
+FStrataData GetStrataHairBSDF(float3 BaseColor, float Scatter, float Specular, float Roughness, float Backlit, float3 Emissive, uint SharedNormalIndex)
 {
 	FStrataData StrataData = GetInitialisedStrataData();
 	StrataData.LayerCount++;
@@ -457,12 +481,40 @@ FStrataData GetStrataHairBSDF(float3 BaseColor, float Scatter, float Specular, f
 
 	BSDF_SETTYPE					(StrataData.Layers[0].BSDFs[0], STRATA_BSDF_TYPE_HAIR);
 	BSDF_SETNORMALID				(StrataData.Layers[0].BSDFs[0], SharedNormalIndex);
+	BSDF_SETEMISSIVE				(StrataData.Layers[0].BSDFs[0], Emissive);
 	HAIR_BASECOLOR					(StrataData.Layers[0].BSDFs[0]) = BaseColor;
 	HAIR_SCATTER					(StrataData.Layers[0].BSDFs[0]) = Scatter;
 	HAIR_ROUGHNESS					(StrataData.Layers[0].BSDFs[0]) = Roughness;
 	HAIR_SPECULAR					(StrataData.Layers[0].BSDFs[0]) = Specular;
 	HAIR_BACKLIT					(StrataData.Layers[0].BSDFs[0]) = Backlit;
-	BSDF_SETEMISSIVE				(StrataData.Layers[0].BSDFs[0], EmissiveColor);
+
+	StrataData.Layers[0].BSDFs[0].Weight = 1.0f;
+
+	return StrataData;
+}
+
+FStrataData GetStrataSingleLayerWaterBSDF(
+	float3 BaseColor, float Metallic, float Specular, float Roughness, float3 Emissive, float TopMaterialOpacity,
+	float3 WaterAlbedo, float3 WaterExtinction, float WaterPhaseG, float3 ColorScaleBehindWater, uint SharedNormalIndex)
+{
+	FStrataData StrataData = GetInitialisedStrataData();
+	StrataData.LayerCount++;
+	StrataData.Layers[0].BSDFCount++;
+
+	BSDF_SETTYPE					(StrataData.Layers[0].BSDFs[0], STRATA_BSDF_TYPE_SINGLELAYERWATER);
+	BSDF_SETNORMALID				(StrataData.Layers[0].BSDFs[0], SharedNormalIndex);
+	BSDF_SETEMISSIVE				(StrataData.Layers[0].BSDFs[0], Emissive);
+	SLW_BASECOLOR					(StrataData.Layers[0].BSDFs[0]) = BaseColor;
+	SLW_METALLIC					(StrataData.Layers[0].BSDFs[0]) = Metallic;
+	SLW_SPECULAR					(StrataData.Layers[0].BSDFs[0]) = Specular;
+	SLW_ROUGHNESS					(StrataData.Layers[0].BSDFs[0]) = Roughness;
+	SLW_TOPMATERIALOPACITY			(StrataData.Layers[0].BSDFs[0]) = TopMaterialOpacity;
+#if STRATA_INLINE_SINGLELAYERWATER
+	SLW_WATERALBEDO					(StrataData.Layers[0].BSDFs[0]) = WaterAlbedo;
+	SLW_WATEREXTINCTION				(StrataData.Layers[0].BSDFs[0]) = WaterExtinction;
+	SLW_WATERPHASEG					(StrataData.Layers[0].BSDFs[0]) = WaterPhaseG;
+	SLW_COLORSCALEBEHINDWATER		(StrataData.Layers[0].BSDFs[0]) = ColorScaleBehindWater;
+#endif
 
 	StrataData.Layers[0].BSDFs[0].Weight = 1.0f;
 
@@ -680,14 +732,32 @@ void SanitizeStrataSheen(inout FStrataBSDF BSDF)
 
 void SanitizeStrataHair(inout FStrataBSDF BSDF)
 {
+	BSDF_SETEMISSIVE(BSDF, clamp(BSDF_GETEMISSIVE(BSDF), 0.0, Max10BitsFloat));
 	HAIR_BASECOLOR(BSDF) = saturate(HAIR_BASECOLOR(BSDF));
 	HAIR_SCATTER(BSDF) = saturate(HAIR_SCATTER(BSDF));
 	HAIR_ROUGHNESS(BSDF) = saturate(HAIR_ROUGHNESS(BSDF));
 	HAIR_SPECULAR(BSDF) = saturate(HAIR_SPECULAR(BSDF));
 	HAIR_BACKLIT(BSDF) = saturate(HAIR_BACKLIT(BSDF));
+}
+
+void SanitizeStrataSingleLayerWater(inout FStrataBSDF BSDF)
+{
 	BSDF_SETEMISSIVE(BSDF, clamp(BSDF_GETEMISSIVE(BSDF), 0.0, Max10BitsFloat));
+	SLW_BASECOLOR(BSDF) = saturate(SLW_BASECOLOR(BSDF));
+	SLW_METALLIC(BSDF) = saturate(SLW_METALLIC(BSDF));
+	SLW_SPECULAR(BSDF) = saturate(SLW_SPECULAR(BSDF));
+	SLW_ROUGHNESS(BSDF) = saturate(SLW_ROUGHNESS(BSDF));
+	SLW_TOPMATERIALOPACITY(BSDF) = saturate(SLW_TOPMATERIALOPACITY(BSDF));
+#if STRATA_INLINE_SINGLELAYERWATER
+	SLW_WATERALBEDO(BSDF) = saturate(SLW_WATERALBEDO(BSDF));
+	SLW_WATEREXTINCTION(BSDF) = clamp(SLW_WATEREXTINCTION(BSDF), 0.0f, Max10BitsFloat);
+	SLW_WATERPHASEG(BSDF) = clamp(SLW_WATERPHASEG(BSDF), -1.0f, 1.0f);
+	SLW_COLORSCALEBEHINDWATER(BSDF) = clamp(SLW_COLORSCALEBEHINDWATER(BSDF), 0.0f, Max10BitsFloat);
+#endif
 }
 
+
+
 void StrataSanitizeBSDF(inout FStrataBSDF BSDF)
 {
 	switch (BSDF_GETTYPE(BSDF))
@@ -722,6 +792,12 @@ void StrataSanitizeBSDF(inout FStrataBSDF BSDF)
 			SanitizeStrataHair(BSDF);
 			break;
 		}
+		case STRATA_BSDF_TYPE_SINGLELAYERWATER:
+		{
+			SanitizeStrataSingleLayerWater(BSDF);
+			break;
+		}
+		
 	}
 }
 
@@ -1219,6 +1295,27 @@ void PackStrataOut(
 				break;
 			}
 
+			case STRATA_BSDF_TYPE_SINGLELAYERWATER:
+			{
+				if (bTopLayer)
+				{
+					const float Roughness = SLW_ROUGHNESS(BSDF);
+
+					TopLayer.Roughness += BSDFCoverageAvg * Roughness;
+					TopLayer.AccumRoughnessCount += BSDFCoverageAvg;
+
+					// The specular over the water in completely in the hand of the user. We do not fade out metallness for instance.
+					const float3 F0 = ComputeF0(SLW_SPECULAR(BSDF), SLW_BASECOLOR(BSDF), SLW_METALLIC(BSDF));
+					const float SafeRoughness = MakeRoughnessSafe(Roughness);
+					TopLayer.BaseColor += BSDFCoverage * EnvBRDFApprox(F0, SafeRoughness, saturate(ContextNoLight.NoV));
+				}
+
+				// Water is considered opaque, it has a custom process to simulate translucency
+				LayerCoverage += BSDFCoverage;
+				LayerThroughput += BSDFCoverage * OpaqueBSDFThroughput;
+				break;
+			}
+
 			} // switch
 
 			// Update the weights to take into account throughput towards the view from layers above
@@ -1375,6 +1472,12 @@ void PackStrataOut(
 				// 8 bytes
 			}
 			break;
+			case STRATA_BSDF_TYPE_SINGLELAYERWATER:
+			{
+				STRATA_STORE_UINT1(PackAlbedoLinearToGamma2AlphaLinear(float4(SLW_BASECOLOR(BSDF), 0.0f)));
+				STRATA_STORE_UINT1(PackRGBA8(float4(SLW_METALLIC(BSDF), SLW_SPECULAR(BSDF), SLW_ROUGHNESS(BSDF), SLW_TOPMATERIALOPACITY(BSDF))));
+			}
+			break;
 			}
 
 			// Stop writing if above budget.
@@ -1498,6 +1601,17 @@ FStrataBSDF UnpackStrataBSDFIn(ByteAddressBuffer StrataBuffer, inout uint ByteOf
 		// 8 bytes
 	}
 	break;
+	case STRATA_BSDF_TYPE_SINGLELAYERWATER:
+	{
+		StrataLoad_AlbedoGamma2ToLinearAlphaLinear(StrataBuffer, ByteOffset, SLW_BASECOLOR(OutBSDF), DummyFloat);
+		float4 Data = UnpackRGBA8(StrataLoadUint1(StrataBuffer, ByteOffset));
+		SLW_METALLIC(OutBSDF) = Data.x;
+		SLW_SPECULAR(OutBSDF) = Data.y;
+		SLW_ROUGHNESS(OutBSDF) = Data.z;
+		SLW_TOPMATERIALOPACITY(OutBSDF) = Data.w;
+		// 8 bytes
+	}
+	break;
 	}
 
 	return OutBSDF;
diff --git a/Engine/Shaders/Private/Strata/StrataEvaluation.ush b/Engine/Shaders/Private/Strata/StrataEvaluation.ush
--- a/Engine/Shaders/Private/Strata/StrataEvaluation.ush
+++ b/Engine/Shaders/Private/Strata/StrataEvaluation.ush
@@ -85,6 +85,11 @@ void StrataGetBSDFRoughness(in FStrataBSDF BSDF, inout float R)
 		R = HAIR_ROUGHNESS(BSDF);
 		break;
 	}
+	case STRATA_BSDF_TYPE_SINGLELAYERWATER:
+	{
+		R = SLW_ROUGHNESS(BSDF);
+		break;
+	}
 	}
 }
 
@@ -281,6 +286,12 @@ FStrataEvaluateResult StrataEvaluateBSDF(FStrataBSDFContext BSDFContext, FShadow
 			Sample.Throughput			= OpaqueBSDFThroughput; 
 		}
 		break;
+
+		//case STRATA_BSDF_TYPE_VOLUMETRICFOGCLOUD:
+		//case STRATA_BSDF_TYPE_UNLIT:
+		//case STRATA_BSDF_TYPE_SINGLELAYERWATER:
+		//Nothing to do in this case because these BSDF are evaluated in other specialised passes.
+		//break;
 	}
 
 	return Sample;
@@ -377,6 +388,12 @@ FStrataImportanceSampleResult StrataImportanceSampleBSDF(FStrataPixelHeader Stra
 			UniformSampleSphere(E);
 			break;
 		}
+
+		//case STRATA_BSDF_TYPE_VOLUMETRICFOGCLOUD:
+		//case STRATA_BSDF_TYPE_UNLIT:
+		//case STRATA_BSDF_TYPE_SINGLELAYERWATER:
+		//Nothing to do in this case because these BSDF are evaluated in other specialised passes.
+		//break;
 	}
 
 	//StrataUpdateBSDFContext(BSDFContext, Sample.Direction);
@@ -498,6 +515,12 @@ FStrataEnvLightResult StrataEvaluateForEnvLight(FStrataBSDFContext BSDFContext,
 			// No specular environment contribution as of today if not using the special HairStrand render path
 			break;
 		}
+
+		//case STRATA_BSDF_TYPE_VOLUMETRICFOGCLOUD:
+		//case STRATA_BSDF_TYPE_UNLIT:
+		//case STRATA_BSDF_TYPE_SINGLELAYERWATER:
+		//Nothing to do in this case because these BSDF are evaluated in other specialised passes.
+		//break;
 	}
 
 	return StrataEnvLightResult;
diff --git a/Engine/Shaders/Private/Strata/StrataVisualize.usf b/Engine/Shaders/Private/Strata/StrataVisualize.usf
--- a/Engine/Shaders/Private/Strata/StrataVisualize.usf
+++ b/Engine/Shaders/Private/Strata/StrataVisualize.usf
@@ -251,7 +251,7 @@ void VisualizeMaterialPS(
 					else
 					{
 						HORIZONTALTAB; PRINTPARAMCHAR(_S_); PRINTPARAMCHAR(_S_); PRINTPARAMCHAR(_S_); PRINTPARAMCHAR(_SPC_); PRINTPARAMCHAR(_A_); PRINTPARAMCHAR(_L_); PRINTPARAMCHAR(_SPC_);
-						HORIZONTALTAB; PRINTFLOAT1(DIFFUSE_SSSDMFPALBEDO(BSDF));
+						HORIZONTALTAB; PRINTFLOAT3(DIFFUSE_SSSDMFPALBEDO(BSDF));
 						NEXTLINE;
 						HORIZONTALTAB; PRINTPARAMCHAR(_S_); PRINTPARAMCHAR(_S_); PRINTPARAMCHAR(_S_); PRINTPARAMCHAR(_SPC_); PRINTPARAMCHAR(_R_); PRINTPARAMCHAR(_A_); PRINTPARAMCHAR(_SPC_);
 						HORIZONTALTAB; PRINTFLOAT1(DIFFUSE_SSSDMFPRADIUS(BSDF));
@@ -406,6 +406,28 @@ void VisualizeMaterialPS(
 				HORIZONTALTAB; PRINTFLOAT1(HAIR_BACKLIT(BSDF));
 				NEXTLINE;
 				break;
+
+			case STRATA_BSDF_TYPE_SINGLELAYERWATER:
+				PRINTBSDFCHAR(_W_); PRINTBSDFCHAR(_A_); PRINTBSDFCHAR(_T_); PRINTBSDFCHAR(_E_); PRINTBSDFCHAR(_R_);
+				PRINTBSDFCHAR(_SPC_); PRINTBSDFCHAR(_SPC_); PRINTBSDFCHAR(_SPC_); PRINTBSDFCHAR(_SPC_); PRINTBSDFCHAR(_SPC_);
+				PRINTBSDFCHAR(_T_); PRINTBSDFCHAR(_O_); PRINTBSDFCHAR(_D_); PRINTBSDFCHAR(_O_); PRINTBSDFCHAR(_SPC_);
+
+				HORIZONTALTAB; PRINTPARAMCHAR(_A_); PRINTPARAMCHAR(_L_); PRINTPARAMCHAR(_B_); PRINTPARAMCHAR(_E_); PRINTPARAMCHAR(_D_); PRINTPARAMCHAR(_O_);
+				HORIZONTALTAB; PRINTFLOAT3(SLW_BASECOLOR(BSDF));
+				NEXTLINE;
+				HORIZONTALTAB; PRINTPARAMCHAR(_M_); PRINTPARAMCHAR(_E_); PRINTPARAMCHAR(_T_); PRINTPARAMCHAR(_A_); PRINTPARAMCHAR(_L_); PRINTPARAMCHAR(_O_);
+				HORIZONTALTAB; PRINTFLOAT1(SLW_METALLIC(BSDF));
+				NEXTLINE;
+				HORIZONTALTAB; PRINTPARAMCHAR(_S_); PRINTPARAMCHAR(_P_); PRINTPARAMCHAR(_E_); PRINTPARAMCHAR(_C_); PRINTPARAMCHAR(_U_); PRINTPARAMCHAR(_SPC_);
+				HORIZONTALTAB; PRINTFLOAT1(SLW_SPECULAR(BSDF));
+				NEXTLINE;
+				HORIZONTALTAB; PRINTPARAMCHAR(_R_); PRINTPARAMCHAR(_O_); PRINTPARAMCHAR(_U_); PRINTPARAMCHAR(_G_); PRINTPARAMCHAR(_H_); PRINTPARAMCHAR(_SPC_);
+				HORIZONTALTAB; PRINTFLOAT1(SLW_ROUGHNESS(BSDF));
+				NEXTLINE;
+				HORIZONTALTAB; PRINTPARAMCHAR(_O_); PRINTPARAMCHAR(_P_); PRINTPARAMCHAR(_A_); PRINTPARAMCHAR(_C_); PRINTPARAMCHAR(_I_); PRINTPARAMCHAR(_T_);
+				HORIZONTALTAB; PRINTFLOAT1(SLW_TOPMATERIALOPACITY(BSDF));
+				NEXTLINE;
+				break;
 			}
 
 			HORIZONTALSEPARATION;
diff --git a/Engine/Shaders/Shared/StrataDefinitions.h b/Engine/Shaders/Shared/StrataDefinitions.h
--- a/Engine/Shaders/Shared/StrataDefinitions.h
+++ b/Engine/Shaders/Shared/StrataDefinitions.h
@@ -29,6 +29,7 @@
 #define STRATA_BSDF_TYPE_VOLUMETRICFOGCLOUD	5
 #define STRATA_BSDF_TYPE_UNLIT				6
 #define STRATA_BSDF_TYPE_HAIR				7
-// When more than 8 BSDF must exists, please update STATE_BITS_BSDF in Strata.ush
+#define STRATA_BSDF_TYPE_SINGLELAYERWATER	8
+// When more than 16 BSDF must exists, please update STATE_BITS_BSDF and FStrataClassification.ShadingModels packing in Strata.ush
 
 
diff --git a/Engine/Source/Runtime/Engine/Classes/Engine/EngineTypes.h b/Engine/Source/Runtime/Engine/Classes/Engine/EngineTypes.h
--- a/Engine/Source/Runtime/Engine/Classes/Engine/EngineTypes.h
+++ b/Engine/Source/Runtime/Engine/Classes/Engine/EngineTypes.h
@@ -545,6 +545,7 @@ enum EStrataShadingModel
 	SSM_SubsurfaceLit			UMETA(DisplayName = "SubsurfaceLit"),
 	SSM_VolumetricFogCloud		UMETA(DisplayName = "VolumetricFogCloud"),
 	SSM_Hair					UMETA(DisplayName = "Hair"),
+	SSM_SingleLayerWater		UMETA(DisplayName = "SingleLayerWater"),
 	/** Number of unique shading models. */
 	SSM_NUM						UMETA(Hidden),
 };
diff --git a/Engine/Source/Runtime/Engine/Classes/Materials/MaterialExpressionStrata.h b/Engine/Source/Runtime/Engine/Classes/Materials/MaterialExpressionStrata.h
--- a/Engine/Source/Runtime/Engine/Classes/Materials/MaterialExpressionStrata.h
+++ b/Engine/Source/Runtime/Engine/Classes/Materials/MaterialExpressionStrata.h
@@ -322,10 +322,10 @@ class UMaterialExpressionStrataUnlitBSDF : public UMaterialExpression
 	GENERATED_UCLASS_BODY()
 
 	/**
-	* EmissiveColor (type = float3, unit = Luminance, default = 0)
+	* Emissive (type = float3, unit = Luminance, default = 0)
 	*/
 	UPROPERTY()
-	FExpressionInput EmissiveColor;
+	FExpressionInput Emissive;
 
 	/**
 	 * The amount of transmitted light from the back side of the surface to the front side of the surface ()
@@ -404,6 +404,89 @@ class UMaterialExpressionStrataHairBSDF : public UMaterialExpression
 	//~ End UMaterialExpression Interface
 };
 
+UCLASS(MinimalAPI, collapsecategories, hidecategories = Object)
+class UMaterialExpressionStrataSingleLayerWaterBSDF : public UMaterialExpression
+{
+	GENERATED_UCLASS_BODY()
+
+	/**
+	 * Surface base color resulting from single and multiple scattering combined. (type = float3, unit = unitless, defaults to black)
+	 */
+	UPROPERTY()
+	FExpressionInput BaseColor;
+
+	/**
+	 * TODO. (type = float3, unit = unitless, defaults to black)
+	 */
+	UPROPERTY()
+	FExpressionInput Metallic;
+
+	/**
+	 * Specular (type = float, unit = unitless, defaults to 0.5)
+	 */
+	UPROPERTY()
+	FExpressionInput Specular;
+
+	/**
+	 * Controls how rough the Material is. Roughness of 0 (smooth) is a mirror reflection and 1 (rough) is completely matte or diffuse (type = float, unit = unitless, defaults to 0.5)
+	 */
+	UPROPERTY()
+	FExpressionInput Roughness;
+
+	/**
+	 * Tangent (type = float3, unit = unitless, defaults to +Z vector)
+	 */
+	UPROPERTY()
+	FExpressionInput Normal;
+
+	/**
+	 * Emissive (type = float3, unit = luminance, defaults to 0.0)
+	 */
+	UPROPERTY()
+	FExpressionInput Emissive;
+
+	/**
+	 * Opacity of the material layered on top of the water (type = float3, unit = unitless, defaults to 0.0)
+	 */
+	UPROPERTY()
+	FExpressionInput TopMaterialOpacity;
+
+	/**
+	 * TODO. (type = float3, unit = unitless, defaults to black)
+	 */
+	UPROPERTY()
+	FExpressionInput WaterAlbedo;
+
+	/**
+	 * TODO. (type = float3, unit = 1/cm, defaults to black)
+	 */
+	UPROPERTY()
+	FExpressionInput WaterExtinction;
+
+	/**
+	 * TODO. (type = float3, unit = unitless, defaults to 0)
+	 */
+	UPROPERTY()
+	FExpressionInput WaterPhaseG;
+
+	/**
+	 * TODO. (type = float3, unit = unitless, defaults to white)
+	 */
+	UPROPERTY()
+	FExpressionInput ColorScaleBehindWater;
+
+	//~ Begin UMaterialExpression Interface
+#if WITH_EDITOR
+	virtual int32 Compile(class FMaterialCompiler* Compiler, int32 OutputIndex) override;
+	virtual void GetCaption(TArray<FString>& OutCaptions) const override;
+	virtual uint32 GetOutputType(int32 OutputIndex) override;
+	virtual uint32 GetInputType(int32 InputIndex) override;
+	virtual bool IsResultStrataMaterial(int32 OutputIndex) override;
+	virtual void GatherStrataMaterialInfo(FStrataMaterialInfo& StrataMaterialInfo, int32 OutputIndex) override;
+#endif
+	//~ End UMaterialExpression Interface
+};
+
 
 
 ///////////////////////////////////////////////////////////////////////////////
diff --git a/Engine/Source/Runtime/Engine/Private/Materials/HLSLMaterialTranslator.cpp b/Engine/Source/Runtime/Engine/Private/Materials/HLSLMaterialTranslator.cpp
--- a/Engine/Source/Runtime/Engine/Private/Materials/HLSLMaterialTranslator.cpp
+++ b/Engine/Source/Runtime/Engine/Private/Materials/HLSLMaterialTranslator.cpp
@@ -1154,6 +1154,17 @@ bool FHLSLMaterialTranslator::Translate()
 				}
 			}
 
+			// SingleLayerWater must be used in isolation
+			if (StrataMaterialContainsAnyBSDF(this, StrataCompilationInfo, STRATA_BSDF_TYPE_SINGLELAYERWATER))
+			{
+				// If the unlit node is used, it must be the only one used
+				if (!StrataIsSingleLayerWaterOnly(this, StrataCompilationInfo))
+				{
+					FString ErrorMsg = FString::Printf(TEXT("Material %s contains hair BSDF but it is not the only representing the material asset: %s. It must be the single BSDF.\r\n"), *Material->GetDebugName(), *Material->GetAssetPath().ToString());
+					Error(*ErrorMsg);
+				}
+			}
+
 			// Output some debug info as comment in code and in the material stat window
 			static const auto CVarStrataBytePerPixel = IConsoleManager::Get().FindTConsoleVariableDataInt(TEXT("r.Strata.BytesPerPixel"));
 			check(CVarStrataBytePerPixel)
@@ -7603,16 +7614,16 @@ int32 FHLSLMaterialTranslator::StrataVolumetricFogCloudBSDF(int32 Albedo, int32
 	);
 }
 
-int32 FHLSLMaterialTranslator::StrataUnlitBSDF(int32 EmissiveColor, int32 TransmittanceColor)
+int32 FHLSLMaterialTranslator::StrataUnlitBSDF(int32 Emissive, int32 TransmittanceColor)
 {
 	return AddCodeChunk(
 		MCT_Strata, TEXT("GetStrataUnlitBSDF(%s, %s)"),
-		*GetParameterCode(EmissiveColor),
+		*GetParameterCode(Emissive),
 		*GetParameterCode(TransmittanceColor)
 	);
 }
 
-int32 FHLSLMaterialTranslator::StrataHairBSDF(int32 BaseColor, int32 Scatter, int32 Specular, int32 Roughness, int32 Backlit, int32 EmissiveColor, int32 Tangent, uint8 SharedNormalIndex)
+int32 FHLSLMaterialTranslator::StrataHairBSDF(int32 BaseColor, int32 Scatter, int32 Specular, int32 Roughness, int32 Backlit, int32 Emissive, int32 Tangent, uint8 SharedNormalIndex)
 {
 	return AddCodeChunk(
 		MCT_Strata, TEXT("GetStrataHairBSDF(%s, %s, %s, %s, %s, %s, %u) /* %s */"),
@@ -7621,12 +7632,33 @@ int32 FHLSLMaterialTranslator::StrataHairBSDF(int32 BaseColor, int32 Scatter, in
 		*GetParameterCode(Specular),
 		*GetParameterCode(Roughness),
 		*GetParameterCode(Backlit),
-		*GetParameterCode(EmissiveColor),
+		*GetParameterCode(Emissive),
 		SharedNormalIndex,
 		*GetParameterCode(Tangent)
 	);
 }
 
+int32 FHLSLMaterialTranslator::StrataSingleLayerWaterBSDF(
+	int32 BaseColor, int32 Metallic, int32 Specular, int32 Roughness, int32 Emissive, int32 TopMaterialOpacity,
+	int32 WaterAlbedo, int32 WaterExtinction, int32 WaterPhaseG, int32 ColorScaleBehindWater, int32 Normal, uint8 SharedNormalIndex)
+{
+	return AddCodeChunk(
+		MCT_Strata, TEXT("GetStrataSingleLayerWaterBSDF(%s, %s, %s, %s, %s, %s, %s, %s, %s, %s, %u) /* %s */"),
+		*GetParameterCode(BaseColor),
+		*GetParameterCode(Metallic),
+		*GetParameterCode(Specular),
+		*GetParameterCode(Roughness),
+		*GetParameterCode(Emissive),
+		*GetParameterCode(TopMaterialOpacity),
+		*GetParameterCode(WaterAlbedo),
+		*GetParameterCode(WaterExtinction),
+		*GetParameterCode(WaterPhaseG),
+		*GetParameterCode(ColorScaleBehindWater),
+		SharedNormalIndex,
+		*GetParameterCode(Normal)
+	);
+}
+
 int32 FHLSLMaterialTranslator::StrataHorizontalMixing(int32 Foreground, int32 Background, int32 Mix)
 {
 	if (Foreground == INDEX_NONE || Background == INDEX_NONE || Mix == INDEX_NONE)
diff --git a/Engine/Source/Runtime/Engine/Private/Materials/HLSLMaterialTranslator.h b/Engine/Source/Runtime/Engine/Private/Materials/HLSLMaterialTranslator.h
--- a/Engine/Source/Runtime/Engine/Private/Materials/HLSLMaterialTranslator.h
+++ b/Engine/Source/Runtime/Engine/Private/Materials/HLSLMaterialTranslator.h
@@ -804,8 +804,9 @@ protected:
 	virtual int32 StrataVolumeBSDF(int32 Albedo, int32 Extinction, int32 Anisotropy, int32 Thickness, int32 Normal, uint8 SharedNormalIndex) override;
 	virtual int32 StrataSheenBSDF(int32 Albedo, int32 Roughness, int32 Normal, uint8 SharedNormalIndex) override;
 	virtual int32 StrataVolumetricFogCloudBSDF(int32 Albedo, int32 Extinction, int32 Emissive, int32 AmbientOcclusion) override;
-	virtual int32 StrataUnlitBSDF(int32 EmissiveColor, int32 TransmittanceColor) override;
-	virtual int32 StrataHairBSDF(int32 BaseColor, int32 Scatter, int32 Specular, int32 Roughness, int32 Backlit, int32 EmissiveColor, int32 Tangent, uint8 SharedNormalIndex) override;
+	virtual int32 StrataUnlitBSDF(int32 Emissive, int32 TransmittanceColor) override;
+	virtual int32 StrataHairBSDF(int32 BaseColor, int32 Scatter, int32 Specular, int32 Roughness, int32 Backlit, int32 Emissive, int32 Tangent, uint8 SharedNormalIndex) override;
+	virtual int32 StrataSingleLayerWaterBSDF(int32 BaseColor, int32 Metallic, int32 Specular, int32 Roughness, int32 Emissive, int32 TopMaterialOpacity, int32 WaterAlbedo, int32 WaterExtinction, int32 WaterPhaseG, int32 ColorScaleBehindWater, int32 Normal, uint8 SharedNormalIndex) override;
 	virtual int32 StrataHorizontalMixing(int32 Foreground, int32 Background, int32 Mix) override;
 	virtual int32 StrataVerticalLayering(int32 Top, int32 Base) override;
 	virtual int32 StrataAdd(int32 A, int32 B) override;
diff --git a/Engine/Source/Runtime/Engine/Private/Materials/Material.cpp b/Engine/Source/Runtime/Engine/Private/Materials/Material.cpp
--- a/Engine/Source/Runtime/Engine/Private/Materials/Material.cpp
+++ b/Engine/Source/Runtime/Engine/Private/Materials/Material.cpp
@@ -4563,6 +4563,12 @@ void UMaterial::RebuildShadingModelField()
 			ShadingModel = MSM_Hair;
 			BlendMode = EBlendMode::BLEND_Opaque;
 		}
+		else if (StrataMaterialInfo.HasOnlyShadingModel(SSM_SingleLayerWater))
+		{
+			MaterialDomain = EMaterialDomain::MD_Surface;
+			ShadingModel = MSM_SingleLayerWater;
+			BlendMode = EBlendMode::BLEND_Opaque; // STRATA_TODO water can also be masked: check Mask input from the main node to automatically enabled that?
+		}
 
 		// Also update the ShadingModels for remaining pipeline operation
 		ShadingModels.AddShadingModel(ShadingModel);
diff --git a/Engine/Source/Runtime/Engine/Private/Materials/MaterialExpressions.cpp b/Engine/Source/Runtime/Engine/Private/Materials/MaterialExpressions.cpp
--- a/Engine/Source/Runtime/Engine/Private/Materials/MaterialExpressions.cpp
+++ b/Engine/Source/Runtime/Engine/Private/Materials/MaterialExpressions.cpp
@@ -20163,7 +20163,7 @@ UMaterialExpressionStrataUnlitBSDF::UMaterialExpressionStrataUnlitBSDF(const FOb
 int32 UMaterialExpressionStrataUnlitBSDF::Compile(class FMaterialCompiler* Compiler, int32 OutputIndex)
 {
 	int32 OutputCodeChunk = Compiler->StrataUnlitBSDF(
-		CompileWithDefaultFloat3(Compiler, EmissiveColor, 0.0f, 0.0f, 0.0f),
+		CompileWithDefaultFloat3(Compiler, Emissive, 0.0f, 0.0f, 0.0f),
 		CompileWithDefaultFloat3(Compiler, TransmittanceColor, 1.0f, 1.0f, 1.0f));
 
 	uint8 FakeSharedNormalIndex = 0;
@@ -20301,6 +20301,111 @@ void UMaterialExpressionStrataHairBSDF::GatherStrataMaterialInfo(FStrataMaterial
 
 
 
+UMaterialExpressionStrataSingleLayerWaterBSDF::UMaterialExpressionStrataSingleLayerWaterBSDF(const FObjectInitializer& ObjectInitializer)
+	: Super(ObjectInitializer)
+{
+	struct FConstructorStatics
+	{
+		FText NAME_Strata;
+		FConstructorStatics() : NAME_Strata(LOCTEXT("Strata BSDFs", "Strata BSDFs")) { }
+	};
+	static FConstructorStatics ConstructorStatics;
+#if WITH_EDITORONLY_DATA
+	MenuCategories.Add(ConstructorStatics.NAME_Strata);
+#endif
+}
+
+#if WITH_EDITOR
+int32 UMaterialExpressionStrataSingleLayerWaterBSDF::Compile(class FMaterialCompiler* Compiler, int32 OutputIndex)
+{
+	int32 NormalCodeChunk = CompileWithDefaultNormalWS(Compiler, Normal);
+	uint8 SharedNormalIndex = StrataCompilationInfoCreateSharedNormal(Compiler, NormalCodeChunk);
+
+	int32 OutputCodeChunk = Compiler->StrataSingleLayerWaterBSDF(
+		CompileWithDefaultFloat3(Compiler, BaseColor, 0.0f, 0.0f, 0.0f),
+		CompileWithDefaultFloat1(Compiler, Metallic, 0.0f),
+		CompileWithDefaultFloat1(Compiler, Specular, 0.5f),
+		CompileWithDefaultFloat1(Compiler, Roughness, 0.5f),
+		CompileWithDefaultFloat3(Compiler, Emissive, 0.0f, 0.0f, 0.0f),
+		CompileWithDefaultFloat1(Compiler, TopMaterialOpacity, 0.0f),
+		CompileWithDefaultFloat3(Compiler, WaterAlbedo, 0.0f, 0.0f, 0.0f),
+		CompileWithDefaultFloat3(Compiler, WaterExtinction, 0.0f, 0.0f, 0.0f),
+		CompileWithDefaultFloat1(Compiler, WaterPhaseG, 0.0f),
+		CompileWithDefaultFloat3(Compiler, ColorScaleBehindWater, 1.0f, 1.0f, 1.0f),
+		NormalCodeChunk,
+		SharedNormalIndex);
+
+	StrataCompilationInfoCreateSingleBSDFMaterial(Compiler, OutputCodeChunk, SharedNormalIndex, STRATA_BSDF_TYPE_SINGLELAYERWATER);
+
+	return OutputCodeChunk;
+}
+
+void UMaterialExpressionStrataSingleLayerWaterBSDF::GetCaption(TArray<FString>& OutCaptions) const
+{
+	OutCaptions.Add(TEXT("Strata Single Layer Water BSDF"));
+}
+
+uint32 UMaterialExpressionStrataSingleLayerWaterBSDF::GetOutputType(int32 OutputIndex)
+{
+	return MCT_Strata;
+}
+
+uint32 UMaterialExpressionStrataSingleLayerWaterBSDF::GetInputType(int32 InputIndex)
+{
+	switch (InputIndex)
+	{
+	case 0:
+		return MCT_Float3; // BaseColor
+		break;
+	case 1:
+		return MCT_Float1; // Metallic
+		break;
+	case 2:
+		return MCT_Float1; // Specular
+		break;
+	case 3:
+		return MCT_Float1; // Roughness
+		break;
+	case 4:
+		return MCT_Float3; // Normal
+		break;
+	case 5:
+		return MCT_Float3; // Emissive
+		break;
+	case 6:
+		return MCT_Float1; // TopMaterialOpacity
+		break;
+	case 7:
+		return MCT_Float3; // WaterAlbedo
+		break;
+	case 8:
+		return MCT_Float3; // WaterExtinction
+		break;
+	case 9:
+		return MCT_Float1; // WaterPhaseG
+		break;
+	case 10:
+		return MCT_Float3; // ColorScaleBehindWater
+		break;
+	}
+	
+	check(false);
+	return MCT_Float1;
+}
+
+bool UMaterialExpressionStrataSingleLayerWaterBSDF::IsResultStrataMaterial(int32 OutputIndex)
+{
+	return true;
+}
+
+void UMaterialExpressionStrataSingleLayerWaterBSDF::GatherStrataMaterialInfo(FStrataMaterialInfo& StrataMaterialInfo, int32 OutputIndex)
+{
+	StrataMaterialInfo.AddShadingModel(SSM_SingleLayerWater);
+}
+#endif // WITH_EDITOR
+
+
+
 UMaterialExpressionStrataHorizontalMixing::UMaterialExpressionStrataHorizontalMixing(const FObjectInitializer& ObjectInitializer)
 	: Super(ObjectInitializer)
 {
diff --git a/Engine/Source/Runtime/Engine/Private/Materials/StrataMaterial.cpp b/Engine/Source/Runtime/Engine/Private/Materials/StrataMaterial.cpp
--- a/Engine/Source/Runtime/Engine/Private/Materials/StrataMaterial.cpp
+++ b/Engine/Source/Runtime/Engine/Private/Materials/StrataMaterial.cpp
@@ -33,6 +33,9 @@ FString GetStrataBSDFName(uint8 BSDFType)
 	case STRATA_BSDF_TYPE_HAIR:
 		return TEXT("HAIR");
 		break;
+	case STRATA_BSDF_TYPE_SINGLELAYERWATER:
+		return TEXT("SINGLELAYERWATER");
+		break;
 	}
 	check(false);
 	return "";
@@ -204,6 +207,21 @@ bool StrataIsHairOnly(FMaterialCompiler* Compiler, const FStrataMaterialCompilat
 	return true;
 }
 
+bool StrataIsSingleLayerWaterOnly(FMaterialCompiler* Compiler, const FStrataMaterialCompilationInfo& Material)
+{
+	if (!StrataIsSingleBSDF(Compiler, Material))
+	{
+		return false;
+	}
+	if (Material.Layers[0].BSDFs[0].Type != STRATA_BSDF_TYPE_SINGLELAYERWATER)
+	{
+		Compiler->Error(TEXT("The single BSDF resulting from the graph is not of type SingleLayerWater."));
+		return false;
+	}
+
+	return true;
+}
+
 bool StrataMaterialContainsAnyBSDF(FMaterialCompiler* Compiler, const FStrataMaterialCompilationInfo& Material, uint8 BSDFType)
 {
 	for (uint32 LayerIt = 0; LayerIt < Material.LayerCount; ++LayerIt)
@@ -309,6 +327,12 @@ FStrataMaterialAnalysisResult StrataCompilationInfoMaterialAnalysis(FMaterialCom
 				Result.RequestedByteCount += UintByteSize;
 				break;
 			}
+			case STRATA_BSDF_TYPE_SINGLELAYERWATER:
+			{
+				Result.RequestedByteCount += UintByteSize;
+				Result.RequestedByteCount += UintByteSize;
+				break;
+			}
 			}
 		}
 
diff --git a/Engine/Source/Runtime/Engine/Private/Materials/StrataMaterial.h b/Engine/Source/Runtime/Engine/Private/Materials/StrataMaterial.h
--- a/Engine/Source/Runtime/Engine/Private/Materials/StrataMaterial.h
+++ b/Engine/Source/Runtime/Engine/Private/Materials/StrataMaterial.h
@@ -55,6 +55,7 @@ FStrataMaterialCompilationInfo StrataCompilationInfoVerticalLayering(FMaterialCo
 bool StrataIsVolumetricFogCloudOnly(FMaterialCompiler* Compiler, const FStrataMaterialCompilationInfo& Material);
 bool StrataIsUnlitOnly(FMaterialCompiler* Compiler, const FStrataMaterialCompilationInfo& Material);
 bool StrataIsHairOnly(FMaterialCompiler* Compiler, const FStrataMaterialCompilationInfo& Material);
+bool StrataIsSingleLayerWaterOnly(FMaterialCompiler* Compiler, const FStrataMaterialCompilationInfo& Material);
 bool StrataMaterialContainsAnyBSDF(FMaterialCompiler* Compiler, const FStrataMaterialCompilationInfo& Material, uint8 BSDFType);
 
 struct FStrataMaterialAnalysisResult
diff --git a/Engine/Source/Runtime/Engine/Public/MaterialCompiler.h b/Engine/Source/Runtime/Engine/Public/MaterialCompiler.h
--- a/Engine/Source/Runtime/Engine/Public/MaterialCompiler.h
+++ b/Engine/Source/Runtime/Engine/Public/MaterialCompiler.h
@@ -398,8 +398,9 @@ public:
 	virtual int32 StrataVolumeBSDF(int32 Albedo, int32 Extinction, int32 Anisotropy, int32 Thickness, int32 Normal, uint8 SharedNormalIndex) = 0;
 	virtual int32 StrataSheenBSDF(int32 Albedo, int32 Roughness, int32 Normal, uint8 SharedNormalIndex) = 0;
 	virtual int32 StrataVolumetricFogCloudBSDF(int32 Albedo, int32 Extinction, int32 Emissive, int32 AmbientOcclusion) = 0;
-	virtual int32 StrataUnlitBSDF(int32 EmissiveColor, int32 TransmittanceColor) = 0;
-	virtual int32 StrataHairBSDF(int32 BaseColor, int32 Scatter, int32 Specular, int32 Roughness, int32 Backlit, int32 EmissiveColor, int32 Tangent, uint8 SharedNormalIndex) = 0;
+	virtual int32 StrataUnlitBSDF(int32 Emissive, int32 TransmittanceColor) = 0;
+	virtual int32 StrataHairBSDF(int32 BaseColor, int32 Scatter, int32 Specular, int32 Roughness, int32 Backlit, int32 Emissive, int32 Tangent, uint8 SharedNormalIndex) = 0;
+	virtual int32 StrataSingleLayerWaterBSDF(int32 BaseColor, int32 Metallic, int32 Specular, int32 Roughness, int32 Emissive, int32 TopMaterialOpacity, int32 WaterAlbedo, int32 WaterExtinction, int32 WaterPhaseG, int32 ColorScaleBehindWater, int32 Normal, uint8 SharedNormalIndex) = 0;
 	virtual int32 StrataHorizontalMixing(int32 Foreground, int32 Background, int32 Mix) = 0;
 	virtual int32 StrataVerticalLayering(int32 Top, int32 Base) = 0;
 	virtual int32 StrataAdd(int32 A, int32 B) = 0;
@@ -840,14 +841,19 @@ public:
 		return Compiler->StrataVolumetricFogCloudBSDF(Albedo, Extinction, Emissive, AmbientOcclusion);
 	}
 
-	virtual int32 StrataUnlitBSDF(int32 EmissiveColor, int32 TransmittanceColor) override
+	virtual int32 StrataUnlitBSDF(int32 Emissive, int32 TransmittanceColor) override
 	{
-		return Compiler->StrataUnlitBSDF(EmissiveColor, TransmittanceColor);
+		return Compiler->StrataUnlitBSDF(Emissive, TransmittanceColor);
 	}
 
-	virtual int32 StrataHairBSDF(int32 BaseColor, int32 Scatter, int32 Specular, int32 Roughness, int32 Backlit, int32 EmissiveColor, int32 Tangent, uint8 SharedNormalIndex) override
+	virtual int32 StrataHairBSDF(int32 BaseColor, int32 Scatter, int32 Specular, int32 Roughness, int32 Backlit, int32 Emissive, int32 Tangent, uint8 SharedNormalIndex) override
 	{
-		return Compiler->StrataHairBSDF(BaseColor, Scatter, Specular, Roughness, Backlit, EmissiveColor, Tangent, SharedNormalIndex);
+		return Compiler->StrataHairBSDF(BaseColor, Scatter, Specular, Roughness, Backlit, Emissive, Tangent, SharedNormalIndex);
+	}
+
+	virtual int32 StrataSingleLayerWaterBSDF(int32 BaseColor, int32 Metallic, int32 Specular, int32 Roughness, int32 Emissive, int32 TopMaterialOpacity, int32 WaterAlbedo, int32 WaterExtinction, int32 WaterPhaseG, int32 ColorScaleBehindWater, int32 Normal, uint8 SharedNormalIndex) override
+	{
+		return Compiler->StrataSingleLayerWaterBSDF(BaseColor, Metallic, Specular, Roughness, Emissive, TopMaterialOpacity, WaterAlbedo, WaterExtinction, WaterPhaseG, ColorScaleBehindWater, Normal, SharedNormalIndex);
 	}
 
 	virtual int32 StrataHorizontalMixing(int32 Foreground, int32 Background, int32 Mix) override
diff --git a/Engine/Source/Runtime/Renderer/Private/SingleLayerWaterRendering.cpp b/Engine/Source/Runtime/Renderer/Private/SingleLayerWaterRendering.cpp
--- a/Engine/Source/Runtime/Renderer/Private/SingleLayerWaterRendering.cpp
+++ b/Engine/Source/Runtime/Renderer/Private/SingleLayerWaterRendering.cpp
@@ -16,6 +16,7 @@
 #include "SceneRendering.h"
 #include "ScreenSpaceRayTracing.h"
 #include "SceneTextureParameters.h"
+#include "Strata/Strata.h"
 
 DECLARE_GPU_STAT_NAMED(RayTracingWaterReflections, TEXT("Ray Tracing Water Reflections"));
 
@@ -137,6 +138,7 @@ BEGIN_SHADER_PARAMETER_STRUCT(FSingleLayerWaterCommonShaderParameters, )
 	SHADER_PARAMETER_STRUCT_REF(FReflectionCaptureShaderData, ReflectionCaptureData)
 	SHADER_PARAMETER_STRUCT_REF(FReflectionUniformParameters, ReflectionsParameters)
 	SHADER_PARAMETER_STRUCT_REF(FForwardLightData, ForwardLightData)
+	SHADER_PARAMETER_STRUCT_REF(FStrataGlobalUniformParameters, Strata)
 END_SHADER_PARAMETER_STRUCT()
 
 class FSingleLayerWaterCompositePS : public FGlobalShader
@@ -147,7 +149,8 @@ class FSingleLayerWaterCompositePS : public FGlobalShader
 	class FScreenSpaceReflections : SHADER_PERMUTATION_BOOL("SCREEN_SPACE_REFLECTION");
 	class FHasBoxCaptures : SHADER_PERMUTATION_BOOL("REFLECTION_COMPOSITE_HAS_BOX_CAPTURES");
 	class FHasSphereCaptures : SHADER_PERMUTATION_BOOL("REFLECTION_COMPOSITE_HAS_SPHERE_CAPTURES");
-	using FPermutationDomain = TShaderPermutationDomain<FScreenSpaceReflections, FHasBoxCaptures, FHasSphereCaptures>;
+	class FStrata : SHADER_PERMUTATION_BOOL("STRATA_ENABLED");
+	using FPermutationDomain = TShaderPermutationDomain<FScreenSpaceReflections, FHasBoxCaptures, FHasSphereCaptures, FStrata>;
 
 	BEGIN_SHADER_PARAMETER_STRUCT(FParameters, )
 		SHADER_PARAMETER_STRUCT_INCLUDE(FSingleLayerWaterCommonShaderParameters, CommonParameters)
@@ -176,7 +179,8 @@ class FWaterTileCategorisationCS : public FGlobalShader
 	DECLARE_GLOBAL_SHADER(FWaterTileCategorisationCS);
 	SHADER_USE_PARAMETER_STRUCT(FWaterTileCategorisationCS, FGlobalShader)
 
-	using FPermutationDomain = TShaderPermutationDomain<>;
+	class FStrata : SHADER_PERMUTATION_BOOL("STRATA_ENABLED");
+	using FPermutationDomain = TShaderPermutationDomain<FStrata>;
 
 	static int32 GetTileSize()
 	{
@@ -404,12 +408,24 @@ void FDeferredShadingSceneRenderer::RenderSingleLayerWaterReflections(
 				Parameters.ReflectionsParameters = CreateUniformBufferImmediate(ReflectionUniformParameters, UniformBuffer_SingleDraw);
 			}
 			Parameters.ForwardLightData = View.ForwardLightingResources->ForwardLightDataUniformBuffer;
+			if (Strata::IsStrataEnabled())
+			{
+				Parameters.Strata = Strata::BindStrataGlobalUniformParameters(View);
+			}
 		};
 
 		const bool bRunTiled = UseSingleLayerWaterIndirectDraw(View.GetShaderPlatform()) && CVarWaterSingleLayerTiledComposite.GetValueOnRenderThread();
 		FTiledScreenSpaceReflection TiledScreenSpaceReflection = {nullptr, nullptr, nullptr, nullptr, nullptr, 8};
 		FIntVector ViewRes(View.ViewRect.Width(), View.ViewRect.Height(), 1);
 		FIntVector TiledViewRes = FIntVector::DivideAndRoundUp(ViewRes, TiledScreenSpaceReflection.TileSize);
+
+		// We need to update the classification texture once again because the GBuffer has been patched with single layer water data.
+		if (Strata::IsStrataEnabled() && bRunTiled)
+		{
+			Strata::AddStrataMaterialClassificationPass(GraphBuilder, Views);
+		}
+
+
 		if (bRunTiled)
 		{
 			TiledScreenSpaceReflection.DispatchIndirectParametersBuffer = GraphBuilder.CreateBuffer(FRDGBufferDesc::CreateIndirectDesc<FRHIDrawIndirectParameters>(), TEXT("WaterIndirectDrawParameters"));
@@ -424,6 +440,7 @@ void FDeferredShadingSceneRenderer::RenderSingleLayerWaterReflections(
 			// Categorization based on SHADING_MODEL_ID
 			{
 				FWaterTileCategorisationCS::FPermutationDomain PermutationVector;
+				PermutationVector.Set<FWaterTileCategorisationCS::FStrata>(Strata::IsStrataEnabled());
 				TShaderMapRef<FWaterTileCategorisationCS> ComputeShader(View.ShaderMap, PermutationVector);
 
 				FWaterTileCategorisationCS::FParameters* PassParameters = GraphBuilder.AllocParameters<FWaterTileCategorisationCS::FParameters>();
@@ -569,6 +586,7 @@ void FDeferredShadingSceneRenderer::RenderSingleLayerWaterReflections(
 			PermutationVector.Set<FSingleLayerWaterCompositePS::FScreenSpaceReflections>(bEnableSSR || bComposeLumenReflections);
 			PermutationVector.Set<FSingleLayerWaterCompositePS::FHasBoxCaptures>(bHasBoxCaptures);
 			PermutationVector.Set<FSingleLayerWaterCompositePS::FHasSphereCaptures>(bHasSphereCaptures);
+			PermutationVector.Set<FSingleLayerWaterCompositePS::FStrata>(Strata::IsStrataEnabled());
 			TShaderMapRef<FSingleLayerWaterCompositePS> PixelShader(View.ShaderMap, PermutationVector);
 
 			FWaterCompositeParameters* PassParameters = GraphBuilder.AllocParameters<FWaterCompositeParameters>();
