Made Strata refelction use StrataMaterialClassification output.
Fixed a issue with StrataMaterialClassification: V was wrong.

diff --git a/Engine/Shaders/Private/SingleLayerWaterComposite.usf b/Engine/Shaders/Private/SingleLayerWaterComposite.usf
--- a/Engine/Shaders/Private/SingleLayerWaterComposite.usf
+++ b/Engine/Shaders/Private/SingleLayerWaterComposite.usf
@@ -125,15 +125,8 @@ SingleLayerWaterCompositeOutput SingleLayerWaterComposite(float2 BufferUV, float
 #if STRATA_ENABLED
 			if (bIsStrataSLW)
 			{
-				// Water is using the higher quality EnvBRDF instead of EnvBRDFApprox using during the base pass.
-				// So we load the water BSDF here to make sure we compute the accurate reflection for water.
-				// STRATA_TODO use EnvBRDF and remove EnvBRDFApprox, or use a better approximation...
-				uint PixelStrataDataByteOffset = GetStrataPixelDataByteOffset(SvPosition.xy, View.BufferSizeAndInvSize.x, Strata.MaxBytesPerPixel);
-				FStrataPixelHeader StrataPixelHeader = UnpackStrataHeaderIn(Strata.MaterialLobesBuffer, PixelStrataDataByteOffset);
-				const FStrataBSDF BSDF = UnpackStrataBSDFIn(Strata.MaterialLobesBuffer, PixelStrataDataByteOffset);
-
-				const float3 F0 = ComputeF0(SLW_SPECULAR(BSDF), SLW_BASECOLOR(BSDF), SLW_METALLIC(BSDF));
-				EnvBrdf = EnvBRDF(F0, SLW_ROUGHNESS(BSDF), NoV);
+				const FStrataTopLayerData TopLayerData = StrataUnpackTopLayerData(Strata.TopLayerTexture.Load(uint3(SvPosition.xy, 0)));
+				EnvBrdf = TopLayerData.EnvBRDF;
 			}
 			else
 #endif
diff --git a/Engine/Shaders/Private/Strata/Strata.ush b/Engine/Shaders/Private/Strata/Strata.ush
--- a/Engine/Shaders/Private/Strata/Strata.ush
+++ b/Engine/Shaders/Private/Strata/Strata.ush
@@ -233,7 +233,7 @@ void FinalizeTopLayerData(inout FStrataTopLayerData Out)
 	Out.Normal = normalize(Out.Normal);
 	Out.BaseColor = saturate(Out.BaseColor);
 	Out.Roughness = Out.AccumRoughnessCount > 0.0f ? saturate(Out.Roughness / Out.AccumRoughnessCount) : 1.f;
-	Out.EnvBRDF = Out.BaseColor; // Strata_todo: better approximation
+	Out.EnvBRDF = Out.BaseColor; // When not subsurface, BaseColor stores the split-sum approximation EnvBRDF. STRATA_TODO better approximation
 }
 
 uint2 StrataPackTopLayerData(FStrataTopLayerData In)
@@ -252,7 +252,7 @@ FStrataTopLayerData StrataUnpackTopLayerData(uint2 In)
 	Out.Normal = StrataUnpackNormal(In.y);
 	Out.BaseColor = BaseColorAndRoughness.xyz;
 	Out.Roughness = BaseColorAndRoughness.w;
-	Out.EnvBRDF = Out.BaseColor; // Strata_todo: better approximation
+	Out.EnvBRDF = Out.BaseColor; // When not subsurface, BaseColor stores the split-sum approximation EnvBRDF. STRATA_TODO better approximation
 	return Out;
 }
 
diff --git a/Engine/Shaders/Private/Strata/StrataMaterialClassification.usf b/Engine/Shaders/Private/Strata/StrataMaterialClassification.usf
--- a/Engine/Shaders/Private/Strata/StrataMaterialClassification.usf
+++ b/Engine/Shaders/Private/Strata/StrataMaterialClassification.usf
@@ -154,6 +154,23 @@ void ComputePostProcessBSDFData(
 		break;
 	}
 
+	case STRATA_BSDF_TYPE_SINGLELAYERWATER:
+	{
+		if (bTopLayer)
+		{
+			const float Roughness = SLW_ROUGHNESS(BSDF);
+
+			TopLayer.Roughness += BSDFCoverageAvg * Roughness;
+			TopLayer.AccumRoughnessCount += BSDFCoverageAvg;
+
+			// The specular over the water in completely in the hand of the user. We do not fade out metallness for instance.
+			const float3 F0 = ComputeF0(SLW_SPECULAR(BSDF), SLW_BASECOLOR(BSDF), SLW_METALLIC(BSDF));
+			const float SafeRoughness = MakeRoughnessSafe(Roughness);
+			TopLayer.BaseColor += BSDFCoverageAvg * EnvBRDFApprox(F0, SafeRoughness, saturate(ContextNoLight.NoV));
+		}
+		break;
+	}
+
 	} // switch		
 }
 
@@ -175,7 +192,10 @@ void MainPS(
 	ClassificationData.bIsStrataMaterial = GBuffer.ShadingModelID == SHADINGMODELID_STRATA;
 	if (ClassificationData.bIsStrataMaterial)
 	{
-		const float3 V = float3(0,0,1); // TODO
+		float2 ScreenPosition = SvPositionToScreenPosition(SVPos).xy;
+		const float Depth = 1000000.0f;
+		float4 WorldPos = mul(float4(ScreenPosition * Depth, Depth, 1), View.ScreenToWorld);
+		const float3 V = normalize(View.WorldCameraOrigin - WorldPos.xyz);
 
 		uint PixelStrataDataByteOffset = GetStrataPixelDataByteOffset(PixelPos, View.BufferSizeAndInvSize.x, Strata.MaxBytesPerPixel);
 		FStrataPixelHeader Header = UnpackStrataHeaderIn(Strata.MaterialLobesBuffer, PixelStrataDataByteOffset);
