Fixing PVS 7.7 Issues: warning V1064: The '1' operand of integer division is less than the '895745041u' one. The result will always be zero.

diff --git a/Engine/Shaders/Private/HaltonUtilities.ush b/Engine/Shaders/Private/HaltonUtilities.ush
--- a/Engine/Shaders/Private/HaltonUtilities.ush
+++ b/Engine/Shaders/Private/HaltonUtilities.ush
@@ -29,6 +29,9 @@
 
 #pragma once
 
+// PVS-Studio lexer doesn't handle hex floating point constants properly and doesn't treat "-1" as part of the constant.
+//-V:0x1.fffffcp-1:1064 
+
 float Halton_Sample2(unsigned int Index)
 {
 	// int Dimension = 0;
