Work around invalid VSM back face test on hair.
Avoiding special-casing hair in this shader as well for now since this is "just" a performance optimization, but may eventually try and refactor hair Gbuffer encoding into a central place instead.

diff --git a/Engine/Shaders/Private/VirtualShadowMaps/VirtualShadowMapProjection.usf b/Engine/Shaders/Private/VirtualShadowMaps/VirtualShadowMapProjection.usf
--- a/Engine/Shaders/Private/VirtualShadowMaps/VirtualShadowMapProjection.usf
+++ b/Engine/Shaders/Private/VirtualShadowMaps/VirtualShadowMapProjection.usf
@@ -287,14 +287,16 @@ void VirtualShadowMapProjectionDirectionalPS(
 		
 	FGBufferData GBufferData = GetGBufferData(ScreenUV);
 	float3 ToLightDir = normalize(Light.Direction);
-	bool bFrontFace = dot(ToLightDir, GBufferData.WorldNormal) > 0.0f;
 	
 	float3 DebugOutput = float3(0, 0, 0);
 	int SampleCount = 1;
-	float HitDistance = 100000;	// TODO
+	float HitDistance = WORLD_RADIUS_MISS;
 	float RayMissFactor = 0.0f;
 
-	if (bFrontFace)
+	// NOTE: This logic is unreliable for hair and potentially other materials with alternate gbuffer encodings	
+	// TODO: Revisit other ways to get performance back here.
+	//bool bFrontFace = dot(ToLightDir, GBufferData.WorldNormal) > 0.0f;
+	//if (bFrontFace)
 	{
 		const float ContactShadowLengthWorld = ContactShadowLength * View.ClipToView[1][1] * SceneDepth;
 
