Fix for d3d crash when clearing structured buffer (should be Vulkan only)

diff --git a/Engine/Shaders/Private/ClearReplacementShaders.usf b/Engine/Shaders/Private/ClearReplacementShaders.usf
--- a/Engine/Shaders/Private/ClearReplacementShaders.usf
+++ b/Engine/Shaders/Private/ClearReplacementShaders.usf
@@ -51,9 +51,16 @@
 
 	#elif RESOURCE_TYPE == 4
 	
+	// This should only be run on Vulkan where structured buffers are not read-only
+	#if VULKAN_PROFILE
 		StructuredBuffer<VALUE_TYPE> ClearResource;
 		#define ADDR(p) (p).x
 		#define COMP(a,b) ((a).x < (b).x)
+	#else
+		RWBuffer<VALUE_TYPE> ClearResource;
+		#define ADDR(p) (p).x
+		#define COMP(a,b) ((a).x < (b).x)
+	#endif
 
 	#endif
 
