Improve primary ray creation to handle real ortho cameras and far clipping planes

Previous code only worked properly for "faux ortho" cameras (as used in the editor) but did not produce the expected results for actual ortho cameras (as used in the unit test framework).

diff --git a/Engine/Shaders/Private/RayTracing/RayTracingCommon.ush b/Engine/Shaders/Private/RayTracing/RayTracingCommon.ush
--- a/Engine/Shaders/Private/RayTracing/RayTracingCommon.ush
+++ b/Engine/Shaders/Private/RayTracing/RayTracingCommon.ush
@@ -316,19 +316,19 @@ void ApplyCameraRelativeDepthBias(inout float3 RayOrigin, float3 RayDirection, u
 RayDesc CreatePrimaryRay(float2 UV)
 {
 	float2 ScreenPosition = (UV - View.ScreenPositionScaleBias.wz) / View.ScreenPositionScaleBias.xy;
-	float3 RayStart_TranslatedWorld = View.TranslatedWorldCameraOrigin;
-	float4 RayEnd_TranslatedWorld = mul(float4(ScreenPosition, 1, 1), View.ScreenToTranslatedWorld);
 
 	RayDesc Ray;
-	Ray.Origin = mul(float4(ScreenPosition, 1, 1), View.ScreenToWorld).xyz;
-	Ray.Direction = normalize(RayEnd_TranslatedWorld.xyz - RayEnd_TranslatedWorld.w * RayStart_TranslatedWorld);
-
-	// shift the origin forward or backwards to line up with the clipping plane
-	float Nudge = ((View.NearPlane - 1) / dot(Ray.Direction, View.ViewForward));
-	Ray.Origin += Ray.Direction * Nudge;
-
+	float4 O = mul(float4(ScreenPosition, 1, 1), View.ClipToWorld); // near (world space)
+	float4 A = mul(float4(ScreenPosition, 1, 1), View.ClipToTranslatedWorld); // near 
+	float4 B = mul(float4(ScreenPosition, 0, 1), View.ClipToTranslatedWorld); // far
+	Ray.Origin = O.xyz / O.w;
+	Ray.Direction = normalize(A.w * B.xyz - B.w * A.xyz);
+
+	// extract far plane (ortho only, so ignore distinction between T and Z)
+	float4 Back = mul(float4(0, 0, 0, 1), View.ClipToView);
 	Ray.TMin = 0;
-	Ray.TMax = 1.0e27; // TODO: how to extract the far clip plane?
+	Ray.TMax = Back.w > 0 ? (Back.z / Back.w - View.NearPlane) : POSITIVE_INFINITY;
+
 	return Ray;
 }
 
