Fix/improve a variety of VSM artifacts:
- Use engine contact shadow trace function to respect a few tweaks there including the cast contact shadow flag
  - Had to stop short of including all the hair/eye-specific stuff as it is very hard-coded currently and not appropriate for general VSM shader...
- Fix self-intersection issues with short contact shadow traces
- Reduce default length of contact shadow trace to reduce artifacts
- Add normal-based offset to shadow lookup positions to reduce issues with projection aliasing and Nanite mesh LOD
- Add small page dilation region to avoid normal offset walking off mapped pages
  - Will likely eventually converge some of the sampling logic with SMRT to make it more robust as well, but dilation is a useful feature regardless
- Sample both clipmap levels with SMRT by default as performance of dynamically only sampling one doesn't seem to gain much and some artifacts are avoided.

diff --git a/Engine/Shaders/Private/DeferredLightingCommon.ush b/Engine/Shaders/Private/DeferredLightingCommon.ush
--- a/Engine/Shaders/Private/DeferredLightingCommon.ush
+++ b/Engine/Shaders/Private/DeferredLightingCommon.ush
@@ -104,6 +104,8 @@ float ShadowRayCast(
 
 	float FirstHitTime = -1.0;
 
+	const float StartDepth = SceneTexturesStruct.SceneDepthTexture.SampleLevel( SceneTexturesStruct_SceneDepthTextureSampler, RayStartUVz.xy, 0 ).r;
+
 	UNROLL
 	for( int i = 0; i < NumSteps; i++ )
 	{
@@ -112,6 +114,8 @@ float ShadowRayCast(
 
 		float DepthDiff = SampleUVz.z - SampleDepth;
 		bool Hit = abs( DepthDiff + CompareTolerance ) < CompareTolerance;
+		// Avoid self-intersection with the start pixel (exact comparison due to point sampling depth buffer)
+		Hit = Hit && ( SampleDepth != StartDepth );
 
 		FirstHitTime = (Hit && FirstHitTime < 0.0) ? SampleTime : FirstHitTime;
 
@@ -133,6 +137,7 @@ float ShadowRayCast(
 		HitDistance = bValidUV ? ( FirstHitTime * RayLength ) : -1.0;
 	}	
 	
+	
 	return HitDistance;
 }
 
diff --git a/Engine/Shaders/Private/VirtualShadowMaps/PageManagement.usf b/Engine/Shaders/Private/VirtualShadowMaps/PageManagement.usf
--- a/Engine/Shaders/Private/VirtualShadowMaps/PageManagement.usf
+++ b/Engine/Shaders/Private/VirtualShadowMaps/PageManagement.usf
@@ -198,6 +198,7 @@ StructuredBuffer<int> VirtualShadowMapIdRemap;
 uint NumDirectionalLightSmInds;
 uint bPostBasePass;
 float LodFootprintScale;
+float PageDilationBorderSize;
 
 // Convenience to group up a bunch of data computed for a screen pixel
 struct FScreenPixelData
@@ -299,13 +300,7 @@ uint GetMipLevel(int VirtualShadowMapId, FScreenPixelData Pixel)
 	return CalcMipLevelFromFootprint(LodFootprintScale * Footprint, VSM_MAX_MIP_LEVELS);
 }
 
-uint CalcPageOffsetFromUVMip(uint VirtualShadowMapId, float2 ShadowUv, uint MipLevel)
-{
-	uint2 PageAddress = uint2(ShadowUv * CalcLevelDimsPages(MipLevel));
-	return CalcPageOffset(VirtualShadowMapId, MipLevel, PageAddress);
-}
-
-void MarkPage(uint VirtualShadowMapId, uint MipLevel, FScreenPixelData Pixel)
+void MarkPage(uint VirtualShadowMapId, uint MipLevel, FScreenPixelData Pixel, float2 PageDilationOffset)
 {
 	FVirtualShadowMapProjectionShaderData ProjectionData = ShadowMapProjectionData[VirtualShadowMapId];
 	float4 ShadowUVz = mul(float4(Pixel.WorldPosition + ProjectionData.ShadowPreViewTranslation, 1.0f), ProjectionData.TranslatedWorldToShadowUVMatrix);
@@ -321,12 +316,31 @@ void MarkPage(uint VirtualShadowMapId, uint MipLevel, FScreenPixelData Pixel)
 		return;
 	}
 
-	uint PageOffset = CalcPageOffsetFromUVMip(VirtualShadowMapId, ShadowUVz.xy, MipLevel);
+	float2 PageAddressFloat = ShadowUVz.xy * CalcLevelDimsPages(MipLevel);
+	uint2 PageAddress = uint2(PageAddressFloat);
+	uint PageOffset = CalcPageOffset(VirtualShadowMapId, MipLevel, PageAddress);
 	OutPageRequestFlags[PageOffset] = 1;
+
+	// PageDilationBorderSize == 0 implies PageDilationOffset.xy == 0
+	if (PageDilationBorderSize > 0.0f)
+	{
+		uint2 PageAddress2 = uint2(PageAddressFloat + PageDilationOffset);
+		uint PageOffset2 = CalcPageOffset(VirtualShadowMapId, MipLevel, PageAddress2);
+		if (PageOffset2 != PageOffset)
+		{
+			OutPageRequestFlags[PageOffset2] = 1;
+		}
+		uint2 PageAddress3 = uint2(max(float2(0, 0), PageAddressFloat - PageDilationOffset));
+		uint PageOffset3 = CalcPageOffset(VirtualShadowMapId, MipLevel, PageAddress3);
+		if (PageOffset3 != PageOffset)
+		{
+			OutPageRequestFlags[PageOffset3] = 1;
+		}
+	}
 }
 
 [numthreads(VSM_DEFAULT_CS_GROUP_XY, VSM_DEFAULT_CS_GROUP_XY, 1)]
-void GeneratePageFlagsFromPixels(uint3 GroupId : SV_GroupID, uint GroupIndex : SV_GroupIndex)
+void GeneratePageFlagsFromPixels(uint3 GroupId : SV_GroupID, uint GroupIndex : SV_GroupIndex, uint3 DispatchThreadId : SV_DispatchThreadId)
 {
 	// Morton order within a group so page access/atomics are more coherent and wave-swizzled gradients are possible.
 	uint2 PixelPos = VSM_DEFAULT_CS_GROUP_XY * GroupId.xy + MortonDecode(GroupIndex);
@@ -337,7 +351,15 @@ void GeneratePageFlagsFromPixels(uint3 GroupId : SV_GroupID, uint GroupIndex : S
 	}
 	
 	FScreenPixelData Pixel = GetScreenPixelData(PixelPos);
-	
+
+	// Dither pattern for page dilation
+	// We don't need to to check all 8 adjacent pages; as long as there's at least a single pixel near the edge
+	// the adjacent one will get mapped. In practice only checking one diagonal seems to work fine and have minimal
+	// overhead.
+	float2 PageDilationOffset = PageDilationBorderSize * 
+		float2(GroupIndex & 1 ? 1.0f : -1.0f,
+			   GroupIndex & 2 ? 1.0f : -1.0f);
+		
 	// excluding unlit to avoid including processing sky dome
 	if (bPostBasePass != 0)
 	{
@@ -365,7 +387,7 @@ void GeneratePageFlagsFromPixels(uint3 GroupId : SV_GroupID, uint GroupIndex : S
 			int ClipmapIndex = max(0, ClipmapLevel - BaseProjectionData.ClipmapLevel);
 			if (ClipmapIndex < BaseProjectionData.ClipmapLevelCount)
 			{
-				MarkPage(ClipmapStartId + ClipmapIndex, 0, Pixel);
+				MarkPage(ClipmapStartId + ClipmapIndex, 0, Pixel, PageDilationOffset);
 			}
 
 			ProjectionIndex += max(1, BaseProjectionData.ClipmapLevelCount);		// Max just to ensure termination even with bad data
@@ -397,7 +419,7 @@ void GeneratePageFlagsFromPixels(uint3 GroupId : SV_GroupID, uint GroupIndex : S
 			if (VirtualShadowMapId != INDEX_NONE)
 			{
 				uint MipLevel = GetMipLevel(VirtualShadowMapId, Pixel);
-				MarkPage(VirtualShadowMapId, MipLevel, Pixel);
+				MarkPage(VirtualShadowMapId, MipLevel, Pixel, PageDilationOffset);
 			}
 		}
 	}
diff --git a/Engine/Shaders/Private/VirtualShadowMaps/ProjectionCommon.ush b/Engine/Shaders/Private/VirtualShadowMaps/ProjectionCommon.ush
--- a/Engine/Shaders/Private/VirtualShadowMaps/ProjectionCommon.ush
+++ b/Engine/Shaders/Private/VirtualShadowMaps/ProjectionCommon.ush
@@ -127,7 +127,7 @@ float ComputeVirtualShadowMapOptimalSlopeBias(FVirtualShadowMapProjectionShaderD
 	// 2x factor due to lack of precision (probably)
 	float OptimalSlopeBias = 2.0f * max(0.0f, dot(DepthSlopeUV, TexelCenterOffsetUV));
 	// Clamp to avoid excessive degenerate slope biases causing flickering lit pixels
-	OptimalSlopeBias = min(OptimalSlopeBias, abs(100.0f * ProjectionData.ShadowViewToClipMatrix._33));
+	OptimalSlopeBias = min(OptimalSlopeBias, abs(10.0f * ProjectionData.ShadowViewToClipMatrix._33));
 
 	return OptimalSlopeBias;
 }
diff --git a/Engine/Shaders/Private/VirtualShadowMaps/ProjectionDirectional.ush b/Engine/Shaders/Private/VirtualShadowMaps/ProjectionDirectional.ush
--- a/Engine/Shaders/Private/VirtualShadowMaps/ProjectionDirectional.ush
+++ b/Engine/Shaders/Private/VirtualShadowMaps/ProjectionDirectional.ush
@@ -38,9 +38,9 @@ bool ShadowRayCastDirectional(
 	int AbsoluteClipmapIndexAdjacent = VirtualShadowMapId + ClipmapIndexAdjacent;
 
 	// TODO: Experiment with perf of this, etc.
-	bool bSampleTwoClipmapLevels = (ClipmapIndexAdjacent != ClipmapIndex) && (abs(ClipmapLevelFrac - 0.5f) > 0.4f);
+	bool bSampleTwoClipmapLevels = true;
 #if 0
-	bSampleTwoClipmapLevels = false;
+	bSampleTwoClipmapLevels = (ClipmapIndexAdjacent != ClipmapIndex) && (abs(ClipmapLevelFrac - 0.5f) > 0.4f);
 #endif
 
 	const FVirtualShadowMapProjectionShaderData ProjectionData = GetVirtualShadowMapProjectionData(AbsoluteClipmapIndex);
@@ -194,7 +194,7 @@ float TraceDirectional(
 	float RayLengthScale,
 	int DebugOutputType,
 	inout float3 DebugOutput,
-	out float OutHitDistance)
+	inout float InOutHitDistance)
 {
 	float3 WorldPosition = TranslatedWorldPosition - View.PreViewTranslation;
 	float3 ViewPosition = mul(float4(TranslatedWorldPosition, 1.0f), View.TranslatedWorldToView).xyz;
@@ -205,7 +205,7 @@ float TraceDirectional(
 	// and the maximum size a penumbra can be.
 	// Too high values will cause shadows to detach from their contact points (unless more samples are used).
 	// Too low values will greatly restrict how large penumbras can be in screen space.
-	float RayLength = RayLengthScale * DistanceFromViewOrigin;		// TODO: Parameterize
+	float RayLength = RayLengthScale * DistanceFromViewOrigin;
 
 #if USE_RANDOM_SEQUENCE
 	uint Seed = uint(SvPosition.y * View.BufferSizeAndInvSize.x + SvPosition.x);
@@ -218,7 +218,6 @@ float TraceDirectional(
 
 	float StepOffset = InterleavedGradientNoise(SvPosition.xy, View.StateFrameIndexMod8);
 
-	float HitDistance = 500000;		// TODO
 	float RayMissFactor = 0;
 	for (int i = 0; i < RayCount; ++i)
 	{
@@ -229,7 +228,7 @@ float TraceDirectional(
 #endif
 
 		float3 RayDir = GetRandomDirectionalLightRayDir(Light, E);
-		float3 RayStart = WorldPosition;// + RayDir * ContactShadowLengthWorld;
+		float3 RayStart = WorldPosition;// + RayDir * RayStartOffset; TODO!
 
 		bool bHit = ShadowRayCastDirectional(
 			VirtualShadowMapId,
@@ -240,12 +239,11 @@ float TraceDirectional(
 			StepOffset,
 			DebugOutputType,
 			DebugOutput,
-			HitDistance);
+			InOutHitDistance);
 
 		RayMissFactor += bHit ? 0 : 1;
 	}
 	RayMissFactor /= float(RayCount);
 
-	OutHitDistance = HitDistance;
 	return RayMissFactor;
 }
diff --git a/Engine/Shaders/Private/VirtualShadowMaps/VirtualShadowMapProjection.usf b/Engine/Shaders/Private/VirtualShadowMaps/VirtualShadowMapProjection.usf
--- a/Engine/Shaders/Private/VirtualShadowMaps/VirtualShadowMapProjection.usf
+++ b/Engine/Shaders/Private/VirtualShadowMaps/VirtualShadowMapProjection.usf
@@ -4,9 +4,12 @@
 	VirtualShadowMapProjection.usf: 
 =============================================================================*/
 
+#define SUPPORT_CONTACT_SHADOWS 1
+
 #include "../Common.ush"
 #include "../SceneTexturesCommon.ush"
 #include "../DeferredShadingCommon.ush"
+#include "../DeferredLightingCommon.ush"
 #include "../ScreenSpaceDenoise/SSDPublic.ush"
 #include "../ScreenSpaceDenoise/SSDDefinitions.ush"
 #include "../ScreenSpaceDenoise/SSDPublicHarmonics.ush"
@@ -18,7 +21,7 @@
 // See VirtualShadowMapProjection.cpp
 // 0 = Output to denoiser
 // 1 = Output to screen shadow mask
-// 2 = Output debug LOD
+// 2 = Output debug
 #ifndef OUTPUT_TYPE
 #define OUTPUT_TYPE 0
 #endif
@@ -26,66 +29,26 @@
 int VirtualShadowMapId;		// For the first layer for clipmaps
 int DebugOutputType;
 float ContactShadowLength;
+float NormalOffsetWorld;	// Normal offset in world space
 
 int SMRTRayCount;	// 0 = off
 int SMRTSamplesPerRay;
 float SMRTRayLengthScale;
 
+#define CONTACT_SHADOW_SAMPLES 8
+
 
+// Returns hit distance or negative if miss
 float VirtualShadowMapScreenRayCast(
 	float3 RayOriginTranslatedWorld, float3 RayDirection, float RayLength,
 	int NumSteps, float StepOffset)
 {
-	// Just to make calling code a bit cleaner
-	if (RayLength == 0.0f)
-	{
-		return 0.0f;
-	}
-
-	float4 RayStartClip	= mul( float4( RayOriginTranslatedWorld, 1 ), View.TranslatedWorldToClip );
-	float4 RayDirClip	= mul( float4( RayDirection * RayLength, 0 ), View.TranslatedWorldToClip );
-	float4 RayEndClip	= RayStartClip + RayDirClip;
-
-	float3 RayStartScreen = RayStartClip.xyz / RayStartClip.w;
-	float3 RayEndScreen = RayEndClip.xyz / RayEndClip.w;
-
-	float3 RayStepScreen = RayEndScreen - RayStartScreen;
-
-	float3 RayStartUVz = float3( RayStartScreen.xy * View.ScreenPositionScaleBias.xy + View.ScreenPositionScaleBias.wz, RayStartScreen.z );
-	float3 RayStepUVz = float3( RayStepScreen.xy * View.ScreenPositionScaleBias.xy, RayStepScreen.z );
-
-	float4 RayDepthClip	= RayStartClip + mul( float4( 0, 0, RayLength, 0 ), View.ViewToClip );
-	float3 RayDepthScreen = RayDepthClip.xyz / RayDepthClip.w;
-
-	const float Step = 1.0 / NumSteps;
-
-	// *2 to get less moire pattern in extreme cases, larger values make object appear not grounded in reflections
-	const float CompareTolerance = abs( RayDepthScreen.z - RayStartScreen.z ) * Step * 2;
+	bool bHitCastContactShadow = true;
+	float HitDistance = ShadowRayCast(
+		RayOriginTranslatedWorld, RayDirection, RayLength,
+		NumSteps, StepOffset, bHitCastContactShadow);
 
-	float SampleTime = StepOffset * Step + Step;
-
-	float FirstHitTime = -1.0;
-
-	for( int i = 0; i < NumSteps; i++ )
-	{
-		float3 SampleUVz = RayStartUVz + RayStepUVz * SampleTime;
-		float SampleDepth = SceneTexturesStruct.SceneDepthTexture.SampleLevel( SceneTexturesStruct_SceneDepthTextureSampler, SampleUVz.xy, 0 ).r;
-
-		float DepthDiff = SampleUVz.z - SampleDepth;
-		bool Hit = abs( DepthDiff + CompareTolerance ) < CompareTolerance;
-
-		FirstHitTime = (Hit && FirstHitTime < 0.0) ? SampleTime : FirstHitTime;
-
-		SampleTime += Step;
-	}
-
-	float Shadow = FirstHitTime > 0.0f ? 1.0f : 0.0f;
-
-	// Off screen masking
-	bool bValidUV = all( abs( RayStartScreen.xy + RayStepScreen.xy * FirstHitTime ) < 1.0f );
-	Shadow = bValidUV ? Shadow : 0.0f;	
-
-	return 1.0f - Shadow;
+	return (bHitCastContactShadow) ? HitDistance : -1.0f;
 }
 
 
@@ -93,16 +56,12 @@ float ProjectSingleSample(
 	FLightShaderParameters Light,
 	float3 L,
 	uint LightType,
-	const float4 SvPosition,
-	const float SceneDepth,
 	float3 TranslatedWorldPosition,
 	float3 EstimatedGeoWorldNormal,
 	float ContactShadowLengthWorld,
 	inout float3 DebugOutput,
-	out float OutHitDistance)
+	inout float InOutHitDistance)
 {
-	float StepOffset = InterleavedGradientNoise(SvPosition.xy, View.StateFrameIndexMod8);
-
 	// NOTE: Could skip most of the rest of the shader if contact shadows say we're shadowed, but likely to be small, incoherent branching
 
 	float3 WorldPosition = TranslatedWorldPosition - View.PreViewTranslation;
@@ -111,8 +70,12 @@ float ProjectSingleSample(
 	FVirtualShadowMapSampleResult VirtualShadowMapSample =
 		SampleVirtualShadowMap(VirtualShadowMapId, WorldPosition, EstimatedGeoWorldNormal, ContactShadowLengthWorld);
 
-	float RayMissFactor = VirtualShadowMapSample.bOccluded ? 0.0f : 1.0f;
-	float HitDistance = VirtualShadowMapSample.OccluderDistance;
+	float RayMissFactor = 1.0f;
+	if (VirtualShadowMapSample.bOccluded)
+	{
+		RayMissFactor = 0.0f;
+		InOutHitDistance = VirtualShadowMapSample.OccluderDistance;
+	}
 
 	if (DebugOutputType == 1)	// Invalid page
 	{
@@ -124,10 +87,9 @@ float ProjectSingleSample(
 	}
 	else if (DebugOutputType == 3) // Occluder distance
 	{
-		DebugOutput = lerp(RayMissFactor, HitDistance.xxx / 1000.0f, 0.75f);
+		DebugOutput = lerp(RayMissFactor, InOutHitDistance.xxx / 1000.0f, 0.75f);
 	}
 
-	OutHitDistance = HitDistance;
 	return RayMissFactor;
 }
 
@@ -146,8 +108,31 @@ float4 EncodeLightAttenuationFromMask(float ShadowMask)
 	return EncodeLightAttenuation(half4(FadedShadow, FadedSSSShadow, FadedShadow, FadedSSSShadow));
 }
 
-float4 EncodeOuput(float4 OutputSignal, float3 DebugOutput)
+float4 EncodeOutput(
+	float3 TranslatedWorldPosition,
+	FLightShaderParameters Light,
+	uint LightType,
+	float RayMissFactor,
+	float HitDistance,
+	int SampleCount,
+	float3 DebugOutput)
 {
+	// See signal encoding for SIGNAL_BUFFER_LAYOUT_PENUMBRA_RECONSTRUCTION in SSDSignalBufferEncoding.ush
+	// We fuse that logic in here to avoid needing another pass, but it's hard to call the functions directly as they depend on a large
+	// number of preprocessor defines and logic that we don't want to mirror here.
+	// Unfortunately that means this logic can get out of date when the denoiser changes, but that's the trade-off for now.
+	float4 OutputSignal = float4(0, 0, 0, 0);
+	if (RayMissFactor < 1.0f)
+	{
+		float WorldBluringRadius = ComputeLightSampleWorldBluringRadius(TranslatedWorldPosition, LightType, Light, HitDistance);
+		OutputSignal = float4(RayMissFactor, SampleCount, WorldBluringRadius, RayMissFactor);
+	}
+	else
+	{
+		// Miss
+		OutputSignal = float4(1, SampleCount, WORLD_RADIUS_MISS, 1);
+	}
+
 #if OUTPUT_TYPE == 0
 	// Output feeds into denoiser
 	return OutputSignal;
@@ -166,9 +151,15 @@ float3 GetEstimatedGeoWorldNormal(float3 TranslatedWorldPosition, float2 ScreenU
 	// Figure out slope, we do world space since that is the space where we might override using the shading normal...
 	float3 TranslatedWorldPositionDDX = DDX(TranslatedWorldPosition);
 	float3 TranslatedWorldPositionDDY = DDY(TranslatedWorldPosition);
-	float3 EstimatedGeoWorldNormal = normalize(cross(TranslatedWorldPositionDDX, TranslatedWorldPositionDDY));
+	float3 EstimatedGeoWorldNormal = cross(TranslatedWorldPositionDDX, TranslatedWorldPositionDDY);
+
+	// Handle NaNs; will cause it to revert to shading normal below
+	float LengthSq = dot(EstimatedGeoWorldNormal, EstimatedGeoWorldNormal);
+	EstimatedGeoWorldNormal = LengthSq > 1e-8f ? normalize(EstimatedGeoWorldNormal) : float3(0, 0, 0);
 
 #if 1
+	// NOTE: Gbuffer normal is not the surface normal for hair; hair lighting takes a different path but possibly
+	// necessary to do some sort of special casing here (disable normal offset and biasing entirely?).
 	FGBufferData GBufferData = GetGBufferData(ScreenUV);
 	// If the estimated geo normal is too far out we assume it's broken (derivative includes other surfaces or background) and fall back to the shading normal
 	if (dot(GBufferData.WorldNormal, EstimatedGeoWorldNormal) < 0.25f)
@@ -193,8 +184,6 @@ void VirtualShadowMapProjectionSpotPS(
 	
 	FLightShaderParameters Light = GetRootLightShaderParameters();
 
-	FVirtualShadowMapProjectionShaderData ProjectionData = GetVirtualShadowMapProjectionData(VirtualShadowMapId);
-
 	// Early out if we're outside the spot light cone		
 	float3 ToLight = Light.Position - (TranslatedWorldPosition - View.PreViewTranslation);
 	bool bInLightRadius = length(ToLight) <= rcp(Light_InvRadius);
@@ -204,12 +193,16 @@ void VirtualShadowMapProjectionSpotPS(
 
 	float3 DebugOutput = float3(0, 0, 0);
 	int SampleCount = 1;
-	float HitDistance = 100000;	// TODO
+	float HitDistance = WORLD_RADIUS_MISS;
 	float RayMissFactor = 0.0f;
 
 	if (bInLightRadius && bInLightCone)
 	{
 		const float ContactShadowLengthWorld = ContactShadowLength * View.ClipToView[1][1] * SceneDepth;
+
+		float3 EstimatedGeoWorldNormal = GetEstimatedGeoWorldNormal(TranslatedWorldPosition, ScreenUV);
+		TranslatedWorldPosition += EstimatedGeoWorldNormal * NormalOffsetWorld;
+
 #if 0
 		if (SMRTRayCount > 0)
 		{
@@ -223,14 +216,10 @@ void VirtualShadowMapProjectionSpotPS(
 		else
 #endif
 		{
-			float3 EstimatedGeoWorldNormal = GetEstimatedGeoWorldNormal(TranslatedWorldPosition, ScreenUV);
-
 			RayMissFactor = ProjectSingleSample(
 				Light,
 				ToLightDir,
 				LIGHT_TYPE_SPOT,
-				SvPosition,
-				SceneDepth,
 				TranslatedWorldPosition,
 				EstimatedGeoWorldNormal,
 				ContactShadowLengthWorld,
@@ -240,35 +229,17 @@ void VirtualShadowMapProjectionSpotPS(
 
 		if (ContactShadowLengthWorld > 0.0f)
 		{
-			// TODO: Unify this with the random sampling inside SMRT
-			uint2 SobolBase = SobolPixel(SvPosition.xy);
-			uint2 SobolFrame = SobolIndex(SobolBase, View.StateFrameIndexMod8, 3);
 			float StepOffset = InterleavedGradientNoise(SvPosition.xy, View.StateFrameIndexMod8);
-
-			float2 E = float2(SobolIndex(SobolFrame, 8 << 3)) / 0x10000;	// TODO
-
-			float ContactShadow = VirtualShadowMapScreenRayCast(TranslatedWorldPosition, ToLightDir, ContactShadowLengthWorld, 12, StepOffset);		
-			if (ContactShadow < 1.0f)
+			float ContactShadowHitDistance = VirtualShadowMapScreenRayCast(TranslatedWorldPosition, ToLightDir, ContactShadowLengthWorld, CONTACT_SHADOW_SAMPLES, StepOffset);		
+			if (ContactShadowHitDistance >= 0.0f)
 			{
-				RayMissFactor *= ContactShadow;
-				HitDistance = ContactShadowLengthWorld;
+				RayMissFactor = 0.0f;
+				HitDistance = ContactShadowHitDistance;
 			}
 		}
 	}
 
-	float4 OutputSignal = float4(0, 0, 0, 0);
-	if (RayMissFactor < 1.0f)
-	{
-		float WorldBluringRadius = ComputeLightSampleWorldBluringRadius(TranslatedWorldPosition, LIGHT_TYPE_SPOT, Light, HitDistance);
-		OutputSignal = float4(RayMissFactor, SampleCount, WorldBluringRadius, RayMissFactor);
-	}
-	else
-	{
-		// Miss
-		OutputSignal = float4(1, SampleCount, WORLD_RADIUS_MISS, 1);
-	}
-
-	OutputTarget = EncodeOuput(OutputSignal, DebugOutput);
+	OutputTarget = EncodeOutput(TranslatedWorldPosition, Light, LIGHT_TYPE_SPOT, RayMissFactor, HitDistance, SampleCount, DebugOutput);
 }
 
 
@@ -300,6 +271,9 @@ void VirtualShadowMapProjectionDirectionalPS(
 	{
 		const float ContactShadowLengthWorld = ContactShadowLength * View.ClipToView[1][1] * SceneDepth;
 
+		float3 EstimatedGeoWorldNormal = GetEstimatedGeoWorldNormal(TranslatedWorldPosition, ScreenUV);
+		TranslatedWorldPosition += EstimatedGeoWorldNormal * NormalOffsetWorld;
+
 		if (SMRTRayCount > 0)
 		{
 			SampleCount = SMRTRayCount;
@@ -319,17 +293,13 @@ void VirtualShadowMapProjectionDirectionalPS(
 		}
 		else
 		{
-			float3 EstimatedGeoWorldNormal = GetEstimatedGeoWorldNormal(TranslatedWorldPosition, ScreenUV);
-
 			SampleCount = 1;
 			RayMissFactor = ProjectSingleSample(
 				Light,
 				Light.Direction,
 				LIGHT_TYPE_DIRECTIONAL,
-				SvPosition,
-				SceneDepth,
 				TranslatedWorldPosition,
-				EstimatedGeoWorldNormal,
+				float3(0, 0, 0),//EstimatedGeoWorldNormal,
 				ContactShadowLengthWorld,
 				/* out */ DebugOutput,
 				HitDistance);
@@ -337,39 +307,17 @@ void VirtualShadowMapProjectionDirectionalPS(
 
 		if (ContactShadowLengthWorld > 0.0f)
 		{
-			// TODO: Unify this with the random sampling inside SMRT
-			uint2 SobolBase = SobolPixel(SvPosition.xy);
-			uint2 SobolFrame = SobolIndex(SobolBase, View.StateFrameIndexMod8, 3);
 			float StepOffset = InterleavedGradientNoise(SvPosition.xy, View.StateFrameIndexMod8);
-
-			float2 E = float2(SobolIndex(SobolFrame, 8 << 3)) / 0x10000;	// TODO
-
-			float ContactShadow = VirtualShadowMapScreenRayCast(TranslatedWorldPosition, ToLightDir, ContactShadowLengthWorld, 12, StepOffset);		
-			if (ContactShadow < 1.0f)
+			float ContactShadowHitDistance = VirtualShadowMapScreenRayCast(TranslatedWorldPosition, ToLightDir, ContactShadowLengthWorld, CONTACT_SHADOW_SAMPLES, StepOffset);		
+			if (ContactShadowHitDistance >= 0.0f)
 			{
-				RayMissFactor *= ContactShadow;
-				HitDistance = ContactShadowLengthWorld;
+				RayMissFactor = 0.0f;
+				HitDistance = ContactShadowHitDistance;
 			}
 		}
 	}
 
-	// See signal encoding for SIGNAL_BUFFER_LAYOUT_PENUMBRA_RECONSTRUCTION in SSDSignalBufferEncoding.ush
-	// We fuse that logic in here to avoid needing another pass, but it's hard to call the functions directly as they depend on a large
-	// number of preprocessor defines and logic that we don't want to mirror here.
-	// Unfortunately that means this logic can get out of date when the denoiser changes, but that's the trade-off for now.
-	float4 OutputSignal = float4(0, 0, 0, 0);
-	if (RayMissFactor < 1.0f)
-	{
-		float WorldBluringRadius = ComputeLightSampleWorldBluringRadius(TranslatedWorldPosition, LIGHT_TYPE_DIRECTIONAL, Light, HitDistance);
-		OutputSignal = float4(RayMissFactor, SampleCount, WorldBluringRadius, RayMissFactor);
-	}
-	else
-	{
-		// Miss
-		OutputSignal = float4(1, SampleCount, WORLD_RADIUS_MISS, 1);
-	}
-
-	OutputTarget = EncodeOuput(OutputSignal, DebugOutput);
+	OutputTarget = EncodeOutput(TranslatedWorldPosition, Light, LIGHT_TYPE_DIRECTIONAL, RayMissFactor, HitDistance, SampleCount, DebugOutput);
 }
 
 
diff --git a/Engine/Source/Runtime/Renderer/Private/VirtualShadowMaps/VirtualShadowMapArray.cpp b/Engine/Source/Runtime/Renderer/Private/VirtualShadowMaps/VirtualShadowMapArray.cpp
--- a/Engine/Source/Runtime/Renderer/Private/VirtualShadowMaps/VirtualShadowMapArray.cpp
+++ b/Engine/Source/Runtime/Renderer/Private/VirtualShadowMaps/VirtualShadowMapArray.cpp
@@ -69,10 +69,11 @@ static TAutoConsoleVariable<float> CVarResolutionLodScale(
 	ECVF_RenderThreadSafe
 );
 
-static TAutoConsoleVariable<float> CVarResolutionPixelCountPercent(
-	TEXT("r.Shadow.v.ResolutionPixelCountPercent"),
-	0.0f,
-	TEXT("If more than this percent of the screen pixels fall into a single page, virtual resolution will be increased. 0 disables. 1-2% typical."),
+static TAutoConsoleVariable<float> CVarPageDilationBorderSize(
+	TEXT("r.Shadow.v.PageDilationBorderSize"),
+	0.05f,
+	TEXT("If a screen pixel falls within this fraction of a page border, the adacent page will also be mapped.")
+	TEXT("Higher values can reduce page misses at screen edges or disocclusions, but increase total page counts."),
 	ECVF_RenderThreadSafe
 );
 
@@ -273,7 +274,7 @@ class FGeneratePageFlagsFromPixelsCS : public FVirtualPageManagementShader
 		SHADER_PARAMETER(uint32, NumDirectionalLightSmInds)
 		SHADER_PARAMETER(uint32, bPostBasePass)
 		SHADER_PARAMETER(float, LodFootprintScale)
-		SHADER_PARAMETER(uint32, LodPixelCountThreshold)
+		SHADER_PARAMETER(float, PageDilationBorderSize)
 	END_SHADER_PARAMETER_STRUCT()
 };
 IMPLEMENT_GLOBAL_SHADER(FGeneratePageFlagsFromPixelsCS, "/Engine/Private/VirtualShadowMaps/PageManagement.usf", "GeneratePageFlagsFromPixels", SF_Compute);
@@ -528,6 +529,8 @@ void FVirtualShadowMapArray::BuildPageAllocations(FRDGBuilder& GraphBuilder, con
 	// Scale the projected footprint by the inverse scale factor such that 2x -> double the res.
 	const float LodFootprintScale = 1.0f / CVarResolutionLodScale.GetValueOnRenderThread();
 
+	const float PageDilationBorderSize = CVarPageDilationBorderSize.GetValueOnRenderThread();
+
 	const TArray<FSortedLightSceneInfo, SceneRenderingAllocator> &SortedLights = SortedLightsInfo.SortedLights;
 	if (ShadowMaps.Num())
 	{
@@ -649,12 +652,7 @@ void FVirtualShadowMapArray::BuildPageAllocations(FRDGBuilder& GraphBuilder, con
 					PassParameters->SceneTexturesStruct = CreateSceneTextureUniformBuffer(GraphBuilder, View.FeatureLevel, ESceneTextureSetupMode::SceneDepth);
 				}
 				PassParameters->bPostBasePass = bPostBasePass;
-				
-				// Number of pixels in a single page before we forcably bump the LOD
-				const float ResolutionPixelCountFactor = CVarResolutionPixelCountPercent.GetValueOnRenderThread() / 100.0f;
-				const uint32 LodPixelCountThreshold = ResolutionPixelCountFactor <= 0.0f ? 0 :
-					static_cast<uint32>(ResolutionPixelCountFactor * static_cast<float>(View.ViewRect.Area()));
-				
+								
 				PassParameters->VisBuffer64 = VisBuffer64;
 				PassParameters->View = View.ViewUniformBuffer;
 				PassParameters->OutPageRequestFlags = GraphBuilder.CreateUAV(PageRequestFlagsRDG);
@@ -663,7 +661,7 @@ void FVirtualShadowMapArray::BuildPageAllocations(FRDGBuilder& GraphBuilder, con
 				PassParameters->ShadowMapProjectionData = GraphBuilder.CreateSRV(ShadowMapProjectionDataRDG);
 				PassParameters->NumDirectionalLightSmInds = DirectionalLightSmInds.Num();
 				PassParameters->LodFootprintScale = LodFootprintScale;
-				PassParameters->LodPixelCountThreshold = LodPixelCountThreshold;
+				PassParameters->PageDilationBorderSize = PageDilationBorderSize;
 				
 				auto ComputeShader = View.ShaderMap->GetShader<FGeneratePageFlagsFromPixelsCS>(PermutationVector);
 
diff --git a/Engine/Source/Runtime/Renderer/Private/VirtualShadowMaps/VirtualShadowMapProjection.cpp b/Engine/Source/Runtime/Renderer/Private/VirtualShadowMaps/VirtualShadowMapProjection.cpp
--- a/Engine/Source/Runtime/Renderer/Private/VirtualShadowMaps/VirtualShadowMapProjection.cpp
+++ b/Engine/Source/Runtime/Renderer/Private/VirtualShadowMaps/VirtualShadowMapProjection.cpp
@@ -26,11 +26,19 @@
 
 static TAutoConsoleVariable<float> CVarContactShadowLength(
 	TEXT( "r.Shadow.v.ContactShadowLength" ),
-	0.04f,
+	0.02f,
 	TEXT( "Length of the screen space contact shadow trace (smart shadow bias) before the virtual shadow map lookup." ),
 	ECVF_RenderThreadSafe
 );
 
+static TAutoConsoleVariable<float> CVarNormalOffsetWorld(
+	TEXT( "r.Shadow.v.NormalOffsetWorld" ),
+	0.5f,
+	TEXT( "World space offset along surface normal for shadow lookup." )
+	TEXT( "Higher values avoid artifacts on surfaces nearly parallel to the light, but also visibility offset shadows and increase the chance of hitting unmapped pages." ),
+	ECVF_RenderThreadSafe
+);
+
 static TAutoConsoleVariable<int32> CVarVirtualShadowMapDebugProjection(
 	TEXT( "r.Shadow.v.DebugVisualizeProjection" ),
 	0,
@@ -70,7 +78,7 @@ static TAutoConsoleVariable<int32> CVarSMRTSamplesPerRayDirectional(
 
 static TAutoConsoleVariable<float> CVarSMRTRayLengthScaleDirectional(
 	TEXT( "r.Shadow.v.SMRT.RayLengthScaleDirectional" ),
-	2.0f,
+	1.0f,
 	TEXT( "Length of ray to shoot for directional lights, scaled by distance to camera." )
 	TEXT( "Shorter rays limit the screen space size of shadow penumbra. " )
 	TEXT( "Longer rays require more samples to avoid shadows disconnecting from contact points. " ),
@@ -84,6 +92,7 @@ BEGIN_SHADER_PARAMETER_STRUCT(FProjectionParameters, )
 	SHADER_PARAMETER_STRUCT_REF(FViewUniformShaderParameters, View)
 	SHADER_PARAMETER(int32, VirtualShadowMapId)
 	SHADER_PARAMETER(float, ContactShadowLength)
+	SHADER_PARAMETER(float, NormalOffsetWorld)
 	SHADER_PARAMETER(int32, DebugOutputType)
 	SHADER_PARAMETER(int32, SMRTRayCount)
 	SHADER_PARAMETER(int32, SMRTSamplesPerRay)
@@ -191,6 +200,7 @@ static bool AddPass_RenderVirtualShadowMapProjection(
 	PassParameters->VirtualShadowMapId = VirtualShadowMapId;
 	PassParameters->DebugOutputType = CVarVirtualShadowMapDebugProjection.GetValueOnRenderThread();
 	PassParameters->ContactShadowLength = CVarContactShadowLength.GetValueOnRenderThread();
+	PassParameters->NormalOffsetWorld = CVarNormalOffsetWorld.GetValueOnRenderThread();
 
 	PassParameters->RenderTargets[0] = RenderTargetBinding;
 
