Fix invalid hair cards deformation offset.

diff --git a/Engine/Plugins/Runtime/HairStrands/Source/HairStrandsCore/Private/HairStrandsRendering.cpp b/Engine/Plugins/Runtime/HairStrands/Source/HairStrandsCore/Private/HairStrandsRendering.cpp
--- a/Engine/Plugins/Runtime/HairStrands/Source/HairStrandsCore/Private/HairStrandsRendering.cpp
+++ b/Engine/Plugins/Runtime/HairStrands/Source/HairStrandsCore/Private/HairStrandsRendering.cpp
@@ -787,9 +787,9 @@ class FHairCardsDeformationCS : public FGlobalShader
 		SHADER_PARAMETER(uint32, CardsVertexCount)
 		SHADER_PARAMETER(uint32, GuideVertexCount)
 		SHADER_PARAMETER(FVector, GuideRestPositionOffset)
-		SHADER_PARAMETER(FVector, GuideDeformedPositionOffset)
 		SHADER_PARAMETER_RDG_BUFFER_SRV(Buffer, GuideRestPositionBuffer)
 		SHADER_PARAMETER_RDG_BUFFER_SRV(Buffer, GuideDeformedPositionBuffer)
+		SHADER_PARAMETER_RDG_BUFFER_SRV(Buffer, GuideDeformedPositionOffsetBuffer)
 		SHADER_PARAMETER_SRV(Buffer, CardsRestPositionBuffer)
 		SHADER_PARAMETER_RDG_BUFFER_SRV(Buffer, CardsInterpolationBuffer)
 
@@ -835,9 +835,9 @@ static void AddHairCardsDeformationPass(
 	FHairCardsDeformationCS::FParameters* Parameters = GraphBuilder.AllocParameters<FHairCardsDeformationCS::FParameters>();
 	Parameters->GuideVertexCount			= LOD.Guides.RestResource->GetVertexCount();
 	Parameters->GuideRestPositionOffset		= LOD.Guides.RestResource->PositionOffset;
-	Parameters->GuideDeformedPositionOffset = LOD.Guides.DeformedResource->GetPositionOffset(FHairStrandsDeformedResource::Current);
 	Parameters->GuideRestPositionBuffer		= RegisterAsSRV(GraphBuilder, LOD.Guides.RestResource->RestPositionBuffer);
 	Parameters->GuideDeformedPositionBuffer = RegisterAsSRV(GraphBuilder, LOD.Guides.DeformedResource->GetBuffer(FHairStrandsDeformedResource::Current));
+	Parameters->GuideDeformedPositionOffsetBuffer = RegisterAsSRV(GraphBuilder, LOD.Guides.DeformedResource->GetPositionOffsetBuffer(FHairStrandsDeformedResource::Current));
 	
 	Parameters->CardsVertexCount			= LOD.RestResource->VertexCount;
 	Parameters->CardsRestPositionBuffer		= LOD.RestResource->RestPositionBuffer.ShaderResourceViewRHI;
diff --git a/Engine/Shaders/Private/HairStrands/HairCardsDeformation.usf b/Engine/Shaders/Private/HairStrands/HairCardsDeformation.usf
--- a/Engine/Shaders/Private/HairStrands/HairCardsDeformation.usf
+++ b/Engine/Shaders/Private/HairStrands/HairCardsDeformation.usf
@@ -7,7 +7,7 @@
 uint CardsVertexCount;
 uint GuideVertexCount;
 float3 GuideRestPositionOffset;
-float3 GuideDeformedPositionOffset;
+Buffer<float4> GuideDeformedPositionOffsetBuffer;
 
 Buffer<uint4>  GuideRestPositionBuffer;
 Buffer<uint4>  GuideDeformedPositionBuffer;
@@ -43,7 +43,8 @@ float3 ComputeDynamicGeometryOffset(
 	float GuideVertexWeight,
 	FHairMeshTriangle RestTriangle,
 	FHairMeshTriangle DeformedTriangle,
-	float3 RootBarycentric)
+	float3 RootBarycentric,
+	float3 GuideDeformedPositionOffset)
 {
 	const float3 RestGuidePoint = f16tof32(GuideRestPositionBuffer[GuideIndex].xyz) + GuideRestPositionOffset;
 	const float3 LocalRestGuidePoint = ToTriangle(RestGuidePoint, RootBarycentric, RestTriangle);
@@ -64,6 +65,7 @@ void MainCS(uint2 DispatchThreadId : SV_DispatchThreadID)
 
 	const float GuideRadius = 1;
 	const float GuideLength = 1;
+	const float3 GuideDeformedPositionOffset = GuideDeformedPositionOffsetBuffer[0].xyz;
 
 #if PERMUTATION_DYNAMIC_GEOMETRY == 0
 	// Simple deformation based on world space offset
@@ -120,8 +122,8 @@ void MainCS(uint2 DispatchThreadId : SV_DispatchThreadID)
 	const FHairMeshTriangle RestTriangle = GetTriangleTransformation(RootIndex, TriangleRestPosition0Buffer, TriangleRestPosition1Buffer, TriangleRestPosition2Buffer);
 	const FHairMeshTriangle DeformedTriangle = GetTriangleTransformation(RootIndex, TriangleDeformedPosition0Buffer, TriangleDeformedPosition1Buffer, TriangleDeformedPosition2Buffer);
 
-	const float3 Offset0 = ComputeDynamicGeometryOffset(GuideIndex0, 1, RestTriangle, DeformedTriangle, RootBarycentric);
-	const float3 Offset1 = ComputeDynamicGeometryOffset(GuideIndex1, 1, RestTriangle, DeformedTriangle, RootBarycentric);
+	const float3 Offset0 = ComputeDynamicGeometryOffset(GuideIndex0, 1, RestTriangle, DeformedTriangle, RootBarycentric, GuideDeformedPositionOffset);
+	const float3 Offset1 = ComputeDynamicGeometryOffset(GuideIndex1, 1, RestTriangle, DeformedTriangle, RootBarycentric, GuideDeformedPositionOffset);
 	const float3 CurrOffset = VectorToWorld(lerp(Offset0, Offset1, GuideLerp), DeformedTriangle);
 
 	// Transform hair from rest post to deformed pose, based triangle deformation + simulation offset (optional depending of the simulation is actual running or not
