Changed ClearReplacementShaders to use RWStructuredBuffer for Vulkan (removes hack for D3D11/12).

diff --git a/Engine/Shaders/Private/ClearReplacementShaders.usf b/Engine/Shaders/Private/ClearReplacementShaders.usf
--- a/Engine/Shaders/Private/ClearReplacementShaders.usf
+++ b/Engine/Shaders/Private/ClearReplacementShaders.usf
@@ -51,16 +51,9 @@
 
 	#elif RESOURCE_TYPE == 4
 	
-	// This should only be run on Vulkan where structured buffers are not read-only
-	#if VULKAN_PROFILE
-		StructuredBuffer<VALUE_TYPE> ClearResource;
+		RWStructuredBuffer<VALUE_TYPE> ClearResource;
 		#define ADDR(p) (p).x
 		#define COMP(a,b) ((a).x < (b).x)
-	#else
-		RWBuffer<VALUE_TYPE> ClearResource;
-		#define ADDR(p) (p).x
-		#define COMP(a,b) ((a).x < (b).x)
-	#endif
 
 	#endif
 
