Strata: unified diffuse, dielectric and conductor behind a common slab interface with Disney.
There will be naming update after (not done here to not bloat the review) and update for the metal/volume representation after.

diff --git a/Engine/Shaders/Private/Strata/Strata.ush b/Engine/Shaders/Private/Strata/Strata.ush
--- a/Engine/Shaders/Private/Strata/Strata.ush
+++ b/Engine/Shaders/Private/Strata/Strata.ush
@@ -42,11 +42,11 @@ struct FStrataBSDF
 	float4 InlineVGPRs[3];
 #endif
 
-	float4 VGPRs[2];
+	float4 VGPRs[3];
 };
 
 // Num bits for each component store in BSDF state
-#define STATE_BITS_BSDF					4
+#define STATE_BITS_BSDF					3
 #define STATE_BITS_SHAREDNORMALID		2
 #define STATE_BITS_ANISOTROPY			1
 #define STATE_BITS_TOPLAYER				1
@@ -112,29 +112,17 @@ struct FStrataBSDF
 #define BSDF_GETEMISSIVE(X)		0.0f
 #endif
 
-#define DIFFUSE_ALBEDO(X)				X.VGPRs[0].xyz
-#define DIFFUSE_ROUGHNESS(X)			X.VGPRs[0].w
+#define SLAB_ALBEDO(X)					X.VGPRs[0].xyz
+#define SLAB_REFLECTIVITY(X)			X.VGPRs[0].w
+#define SLAB_METALLIC(X)				X.VGPRs[1].x
+#define SLAB_ROUGHNESS0(X)				X.VGPRs[1].y
+#define SLAB_ROUGHNESS1(X)				X.VGPRs[1].z
 // Either with SSS profile
-#define DIFFUSE_SSSPROFILEID(X)			X.VGPRs[1].x
-#define DIFFUSE_SSSPROFILERADIUS(X)		X.VGPRs[1].y
+#define SLAB_SSSPROFILEID(X)			X.VGPRs[2].x
+#define SLAB_SSSPROFILERADIUS(X)		X.VGPRs[2].y
 // Or explicit DMFP
-#define DIFFUSE_SSSDMFPALBEDO(X)		X.VGPRs[1].xyz
-#define DIFFUSE_SSSDMFPRADIUS(X)		X.VGPRs[1].w
-
-#define DIELECTRIC_TINT(X)				X.VGPRs[0].xyz
-#define DIELECTRIC_IOR(X)				X.VGPRs[0].w
-#define DIELECTRIC_ROUGHNESS0(X)		X.VGPRs[1].y
-#define DIELECTRIC_ROUGHNESS1(X)		X.VGPRs[1].x
-
-#define CONDUCTOR_REFLECTIVITY(X)		X.VGPRs[0].xyz
-#define CONDUCTOR_EDGECOLOR(X)			X.VGPRs[1].xyz
-#define CONDUCTOR_ROUGHNESS0(X)			X.VGPRs[0].w
-#define CONDUCTOR_ROUGHNESS1(X)			X.VGPRs[1].w
-
-#define VOLUME_ALBEDO(X)				X.VGPRs[0].xyz
-#define VOLUME_EXTINCTION(X)			X.VGPRs[1].xyz
-#define VOLUME_ANISOTROPY(X)			X.VGPRs[0].w
-#define VOLUME_THICKNESS(X)				X.VGPRs[1].w
+#define SLAB_SSSDMFPALBEDO(X)			X.VGPRs[2].xyz
+#define SLAB_SSSDMFPRADIUS(X)			X.VGPRs[2].w
 
 #define SHEEN_ALBEDO(X)					X.VGPRs[0].xyz
 #define SHEEN_ROUGHNESS(X)				X.VGPRs[0].w
@@ -186,21 +174,21 @@ struct FStrataClassification
 uint StrataPackClassificationData(FStrataClassification In)
 {
 	return 
-		(In.ShadingModels & 0xFFFF) | 
-		((In.bIsStrataMaterial		? 1 : 0) << 16) |
-		((In.bHasSubsurface			? 1 : 0) << 17) | 
-		((In.bHasSubsurfaceProfile	? 1 : 0) << 18)|
-		((In.bHasTransmission		? 1 : 0) << 19);
+		(In.ShadingModels & 0xFF) | 
+		((In.bIsStrataMaterial		? 1 : 0) << 8 ) |
+		((In.bHasSubsurface			? 1 : 0) << 9 ) | 
+		((In.bHasSubsurfaceProfile	? 1 : 0) << 10) |
+		((In.bHasTransmission		? 1 : 0) << 11) ;
 }
 
 FStrataClassification StrataUnpackClassificationData(uint In)
 {
 	FStrataClassification Out;
-	Out.ShadingModels			= (In & 0xFFFF);
-	Out.bIsStrataMaterial		= (In >> 16) & 0x1;
-	Out.bHasSubsurface			= (In >> 17) & 0x1;
-	Out.bHasSubsurfaceProfile	= (In >> 18) & 0x1;
-	Out.bHasTransmission		= (In >> 19) & 0x1;
+	Out.ShadingModels			= (In & 0xFF);
+	Out.bIsStrataMaterial		= (In >> 8 ) & 0x1;
+	Out.bHasSubsurface			= (In >> 9 ) & 0x1;
+	Out.bHasSubsurfaceProfile	= (In >> 10) & 0x1;
+	Out.bHasTransmission		= (In >> 11) & 0x1;
 	return Out;
 }
 
@@ -275,9 +263,7 @@ bool StrataHasAnisotropy(in float Roughness0, float Roughness1)
 bool StrataHasAreaLightIntegrator(in FStrataBSDF BSDF)
 {
 	const uint BSDFType = BSDF_GETTYPE(BSDF);
-	return
-		BSDFType == STRATA_BSDF_TYPE_CONDUCTOR ||
-		BSDFType == STRATA_BSDF_TYPE_DIELECTRIC;
+	return BSDFType == STRATA_BSDF_TYPE_SLAB;
 }
 
 bool StrataHasScattering(in float3 Albedo)
@@ -317,40 +303,10 @@ FStrataData GetInitialisedStrataData()
 	return StrataData;
 }
 
-FStrataData GetStrataDiffuseBSDF(float3 Albedo, float Roughness, float SSSProfileID, float SSSDMFPScale, float3 SSSDMFPAlbedo, float SSSDMFPRadius, uint SharedNormalIndex)
-{
-	FStrataData StrataData = GetInitialisedStrataData();
-	StrataData.LayerCount++;
-
-	StrataData.Layers[0].BSDFCount++;
-	BSDF_SETTYPE			(StrataData.Layers[0].BSDFs[0],		STRATA_BSDF_TYPE_DIFFUSE);
-	BSDF_SETNORMALID		(StrataData.Layers[0].BSDFs[0],		SharedNormalIndex);
-	BSDF_SETISTOPLAYER		(StrataData.Layers[0].BSDFs[0],		0);
-	BSDF_SETHASSCATTERING	(StrataData.Layers[0].BSDFs[0],		SSSDMFPScale > 0 ? 1 : 0);
-
-	const bool bHasSSSProfile = SSSProfileID > 0;
-	BSDF_SETHASSSSPROFILE	(StrataData.Layers[0].BSDFs[0],		bHasSSSProfile ? 1 : 0);
-
-	DIFFUSE_ALBEDO			(StrataData.Layers[0].BSDFs[0])		= Albedo;
-	DIFFUSE_ROUGHNESS		(StrataData.Layers[0].BSDFs[0])		= Roughness;
-
-	if (bHasSSSProfile)
-	{
-		DIFFUSE_SSSPROFILEID	(StrataData.Layers[0].BSDFs[0])	= SSSProfileID;
-		DIFFUSE_SSSPROFILERADIUS(StrataData.Layers[0].BSDFs[0])	= SSSDMFPScale;
-	}
-	else
-	{
-		DIFFUSE_SSSDMFPALBEDO	(StrataData.Layers[0].BSDFs[0])	= saturate(SSSDMFPAlbedo);
-		DIFFUSE_SSSDMFPRADIUS	(StrataData.Layers[0].BSDFs[0])	= clamp(SSSDMFPRadius * SSSDMFPScale, 0.f, SSS_PROFILE_PERPIXEL_MAXRADIUS);
-	}
-
-	StrataData.Layers[0].BSDFs[0].Weight						= 1.0f;
-
-	return StrataData;
-}
-
-FStrataData GetStrataDielectricBSDF(float2 Roughness, float IOR, float3 Tint, uint SharedNormalIndex, inout uint SharedNormalTypes)
+FStrataData GetStrataSlabBSDF(
+	float3 Albedo, float Reflectivity, float Metallic, float2 Roughness,
+	float SSSProfileID, float3 SSSDMFPAlbedo, float SSSDMFPRadius,
+	uint SharedNormalIndex, inout uint SharedNormalTypes)
 {
 	FStrataData StrataData = GetInitialisedStrataData();
 	StrataData.LayerCount++;
@@ -362,65 +318,32 @@ FStrataData GetStrataDielectricBSDF(float2 Roughness, float IOR, float3 Tint, ui
 	}
 
 	StrataData.Layers[0].BSDFCount++;
-	BSDF_SETTYPE			(StrataData.Layers[0].BSDFs[0],		STRATA_BSDF_TYPE_DIELECTRIC);
+	BSDF_SETTYPE			(StrataData.Layers[0].BSDFs[0],		STRATA_BSDF_TYPE_SLAB);
 	BSDF_SETNORMALID		(StrataData.Layers[0].BSDFs[0],		SharedNormalIndex);
 	BSDF_SETANISOTROPY		(StrataData.Layers[0].BSDFs[0],		bHasAnisotropy);
 	BSDF_SETISTOPLAYER		(StrataData.Layers[0].BSDFs[0],		0);
-	DIELECTRIC_IOR			(StrataData.Layers[0].BSDFs[0])		= IOR;
-	DIELECTRIC_TINT			(StrataData.Layers[0].BSDFs[0])		= Tint;
-	DIELECTRIC_ROUGHNESS0	(StrataData.Layers[0].BSDFs[0])		= Roughness.x;
-	DIELECTRIC_ROUGHNESS1	(StrataData.Layers[0].BSDFs[0])		= Roughness.y;
-
-	StrataData.Layers[0].BSDFs[0].Weight						= 1.0f;
+	BSDF_SETHASSCATTERING	(StrataData.Layers[0].BSDFs[0],		SSSDMFPRadius > 0 ? 1 : 0);
 
-	return StrataData;
-}
-
-FStrataData GetStrataConductorBSDF(float3 Reflectivity, float3 EdgeColor, float2 Roughness, uint SharedNormalIndex, inout uint SharedNormalTypes)
-{
-	FStrataData StrataData = GetInitialisedStrataData();
-	StrataData.LayerCount++;
+	const bool bHasSSSProfile = SSSProfileID > 0;
+	BSDF_SETHASSSSPROFILE	(StrataData.Layers[0].BSDFs[0],		bHasSSSProfile ? 1 : 0);
 
-	const uint bHasAnisotropy = StrataHasAnisotropy(Roughness.x, Roughness.y) ? 1 : 0;
-	if (bHasAnisotropy)
+	SLAB_ALBEDO				(StrataData.Layers[0].BSDFs[0])		= Albedo;
+	SLAB_REFLECTIVITY		(StrataData.Layers[0].BSDFs[0])		= Reflectivity;
+	SLAB_METALLIC			(StrataData.Layers[0].BSDFs[0])		= Metallic;
+	SLAB_ROUGHNESS0			(StrataData.Layers[0].BSDFs[0])		= Roughness.x;
+	SLAB_ROUGHNESS1			(StrataData.Layers[0].BSDFs[0])		= Roughness.y;
+	if (bHasSSSProfile)
 	{
-		StrataRequestSharedNormalTangentBasis(SharedNormalTypes, SharedNormalIndex);
+		SLAB_SSSPROFILEID	(StrataData.Layers[0].BSDFs[0])		= SSSProfileID;
+		SLAB_SSSPROFILERADIUS(StrataData.Layers[0].BSDFs[0])	= SSSDMFPRadius;
+	}
+	else
+	{
+		SLAB_SSSDMFPALBEDO	(StrataData.Layers[0].BSDFs[0])		= saturate(SSSDMFPAlbedo);
+		SLAB_SSSDMFPRADIUS	(StrataData.Layers[0].BSDFs[0])		= clamp(SSSDMFPRadius, 0.f, SSS_PROFILE_PERPIXEL_MAXRADIUS);
 	}
-
-	StrataData.Layers[0].BSDFCount++;
-	BSDF_SETTYPE			(StrataData.Layers[0].BSDFs[0],		STRATA_BSDF_TYPE_CONDUCTOR);
-	BSDF_SETNORMALID		(StrataData.Layers[0].BSDFs[0],		SharedNormalIndex);
-	BSDF_SETANISOTROPY		(StrataData.Layers[0].BSDFs[0],		bHasAnisotropy);
-	BSDF_SETISTOPLAYER		(StrataData.Layers[0].BSDFs[0],		0);
-	CONDUCTOR_REFLECTIVITY	(StrataData.Layers[0].BSDFs[0])		= Reflectivity;
-	CONDUCTOR_EDGECOLOR		(StrataData.Layers[0].BSDFs[0])		= EdgeColor;
-	CONDUCTOR_ROUGHNESS0	(StrataData.Layers[0].BSDFs[0])		= Roughness.x;
-	CONDUCTOR_ROUGHNESS1	(StrataData.Layers[0].BSDFs[0])		= Roughness.y;
-
-	StrataData.Layers[0].BSDFs[0].Weight						= 1.0f;
-
-	return StrataData;
-}
-
-FStrataData GetStrataVolumeBSDF(float3 Albedo, float3 Extinction, float Anisotropy, float Thickness, uint SharedNormalIndex)
-{
-	FStrataData StrataData = GetInitialisedStrataData();
-	StrataData.LayerCount++;
-
-	const uint bHasScattering = StrataHasScattering(Albedo) ? 1 : 0;
-
-	StrataData.Layers[0].BSDFCount++;
-	BSDF_SETTYPE(				StrataData.Layers[0].BSDFs[0],	STRATA_BSDF_TYPE_VOLUME);
-	BSDF_SETNORMALID(			StrataData.Layers[0].BSDFs[0],	SharedNormalIndex);
-	BSDF_SETISTOPLAYER			(StrataData.Layers[0].BSDFs[0],	0);
-	BSDF_SETHASSCATTERING		(StrataData.Layers[0].BSDFs[0], bHasScattering);
-	VOLUME_ALBEDO				(StrataData.Layers[0].BSDFs[0])	= Albedo;
-	VOLUME_EXTINCTION			(StrataData.Layers[0].BSDFs[0])	= Extinction;
-	VOLUME_ANISOTROPY			(StrataData.Layers[0].BSDFs[0])	= Anisotropy;
-	VOLUME_THICKNESS			(StrataData.Layers[0].BSDFs[0])	= Thickness;
 
 	StrataData.Layers[0].BSDFs[0].Weight						= 1.0f;
-
 	return StrataData;
 }
 
@@ -694,34 +617,16 @@ FStrataData StrataVerticalLayering(FStrataData Top, FStrataData Base)
 // Functions used to sanitize BSDF before they are used for lighting (when forward) or stored to memory (base pass).
 // We do not normalize normals, this is done before the lighting step.
 
-void SanitizeStrataDiffuse(inout FStrataBSDF BSDF)
+void SanitizeStrataSlab(inout FStrataBSDF BSDF)
 {
-	DIFFUSE_ALBEDO(BSDF) = saturate(DIFFUSE_ALBEDO(BSDF));
-	DIFFUSE_ROUGHNESS(BSDF) = saturate(DIFFUSE_ROUGHNESS(BSDF));
-}
-
-void SanitizeStrataDielectric(inout FStrataBSDF BSDF)
-{
-	DIELECTRIC_IOR(BSDF) = clamp(DIELECTRIC_IOR(BSDF), 1.0f, 5.0f);		// See https://en.wikipedia.org/wiki/List_of_refractive_indices
-	DIELECTRIC_TINT(BSDF) = saturate(DIELECTRIC_TINT(BSDF));
-	DIELECTRIC_ROUGHNESS0(BSDF) = saturate(DIELECTRIC_ROUGHNESS0(BSDF));
-	DIELECTRIC_ROUGHNESS1(BSDF) = saturate(DIELECTRIC_ROUGHNESS1(BSDF));
-}
-
-void SanitizeStrataConductor(inout FStrataBSDF BSDF)
-{
-	CONDUCTOR_REFLECTIVITY(BSDF) = saturate(CONDUCTOR_REFLECTIVITY(BSDF));
-	CONDUCTOR_EDGECOLOR(BSDF) = saturate(CONDUCTOR_EDGECOLOR(BSDF));
-	CONDUCTOR_ROUGHNESS0(BSDF) = saturate(CONDUCTOR_ROUGHNESS0(BSDF));
-	CONDUCTOR_ROUGHNESS1(BSDF) = saturate(CONDUCTOR_ROUGHNESS1(BSDF));
-}
-
-void SanitizeStrataVolume(inout FStrataBSDF BSDF)
-{
-	VOLUME_ALBEDO(BSDF) = saturate(VOLUME_ALBEDO(BSDF));
-	VOLUME_EXTINCTION(BSDF) = max(0.0f, VOLUME_EXTINCTION(BSDF));
-	VOLUME_ANISOTROPY(BSDF) = max(0.0f, VOLUME_ANISOTROPY(BSDF));
-	VOLUME_THICKNESS(BSDF) = max(0.0f, VOLUME_THICKNESS(BSDF));
+	SLAB_ALBEDO(BSDF) = saturate(SLAB_ALBEDO(BSDF));
+	SLAB_REFLECTIVITY(BSDF) = saturate(SLAB_REFLECTIVITY(BSDF));
+	SLAB_METALLIC(BSDF) = saturate(SLAB_METALLIC(BSDF));
+	SLAB_ROUGHNESS0(BSDF) = saturate(SLAB_ROUGHNESS0(BSDF));
+	SLAB_ROUGHNESS1(BSDF) = saturate(SLAB_ROUGHNESS1(BSDF));
+	SLAB_SSSPROFILERADIUS(BSDF) = max(0.0f, SLAB_SSSPROFILERADIUS(BSDF));
+	SLAB_SSSDMFPALBEDO(BSDF) = saturate(SLAB_SSSDMFPALBEDO(BSDF));
+	SLAB_SSSDMFPRADIUS(BSDF) = clamp(SLAB_SSSDMFPRADIUS(BSDF), 0.0f, SSS_PROFILE_PERPIXEL_MAXRADIUS);
 }
 
 void SanitizeStrataSheen(inout FStrataBSDF BSDF)
@@ -762,26 +667,11 @@ void StrataSanitizeBSDF(inout FStrataBSDF BSDF)
 {
 	switch (BSDF_GETTYPE(BSDF))
 	{
-		case STRATA_BSDF_TYPE_DIFFUSE:
+		case STRATA_BSDF_TYPE_SLAB:
 		{
-			SanitizeStrataDiffuse(BSDF);
+			SanitizeStrataSlab(BSDF);
 			break; 
 		}
-		case STRATA_BSDF_TYPE_DIELECTRIC:
-		{
-			SanitizeStrataDielectric(BSDF);
-			break;
-		}
-		case STRATA_BSDF_TYPE_CONDUCTOR:
-		{
-			SanitizeStrataConductor(BSDF);
-			break;
-		}
-		case STRATA_BSDF_TYPE_VOLUME:
-		{
-			SanitizeStrataVolume(BSDF);
-			break;
-		}
 		case STRATA_BSDF_TYPE_SHEEN:
 		{
 			SanitizeStrataSheen(BSDF);
@@ -1175,90 +1065,34 @@ void PackStrataOut(
 
 			switch (BSDF_GETTYPE(BSDF))
 			{
-			case STRATA_BSDF_TYPE_DIFFUSE:
+			case STRATA_BSDF_TYPE_SLAB:
 			{
+				const float3 F0 = ComputeF0(SLAB_REFLECTIVITY(BSDF), SLAB_ALBEDO(BSDF), SLAB_METALLIC(BSDF));
+
 				if (bTopLayer)
 				{
-					TopLayer.Roughness += BSDFCoverageAvg * FullyRough; // Do not trigger SSR by using fully rough
+					// Setup SSR for the top layer
+					TopLayer.Roughness += BSDFCoverageAvg * SLAB_ROUGHNESS0(BSDF);
 					TopLayer.AccumRoughnessCount += BSDFCoverageAvg;
+
+					const float SafeRoughness = MakeRoughnessSafe(SLAB_ROUGHNESS0(BSDF));
+					TopLayer.BaseColor += BSDFCoverage * EnvBRDFApprox(F0, SafeRoughness, saturate(ContextNoLight.NoV));
 				}
 
 				// Select the first/toppest layer with subsurface scattering
-				if (BSDF_GETHASSCATTERING(BSDF) && BSDFCoverageAvg > 0 && DIFFUSE_SSSPROFILERADIUS(BSDF) > 0 && !SSS.bHasSubsurface)
+				if (BSDF_GETHASSCATTERING(BSDF) && BSDFCoverageAvg > 0 && SLAB_SSSPROFILERADIUS(BSDF) > 0 && !SSS.bHasSubsurface)
 				{
 					SSS.bHasSubsurface = true;
 					SSS.bHasProfile	= BSDF_GETHASSSSPROFILE(BSDF);
-					SSS.Profile		= DIFFUSE_SSSPROFILEID(BSDF);
-					SSS.RadiusScale	= DIFFUSE_SSSPROFILERADIUS(BSDF);
-					SSS.BaseColor	= DIFFUSE_ALBEDO(BSDF);
-					SSS.DMFPAlbedo	= DIFFUSE_SSSDMFPALBEDO(BSDF);
-					SSS.DMFPRadius	= DIFFUSE_SSSDMFPRADIUS(BSDF);
-				}
-
-				LayerCoverage	+= BSDFCoverage;
-				LayerThroughput += BSDFCoverage * OpaqueBSDFThroughput;
-				break;
-			}
-
-			case STRATA_BSDF_TYPE_DIELECTRIC:
-			{
-				float F0 = DielectricIorToF0(DIELECTRIC_IOR(BSDF));
-				float3 F = F_Schlick(F0, ContextNoLight.VoH);
-
-				if (bTopLayer)
-				{
-					TopLayer.Roughness += DIELECTRIC_ROUGHNESS0(BSDF);
-					TopLayer.AccumRoughnessCount += 1.0f;
-
-					const float SafeRoughness = MakeRoughnessSafe(DIELECTRIC_ROUGHNESS0(BSDF));
-					TopLayer.BaseColor += BSDFCoverage * EnvBRDFApprox(F0, SafeRoughness, saturate(ContextNoLight.NoV));
-				}
-
-				LayerCoverage	+= BSDFCoverage;
-				LayerThroughput += BSDFCoverage * (1.0f - F); // STRATA_TODO use 1-FGD instead of simple 1-F, event if we ignore ray scattering/spreading
-				break;
-			}
-
-			case STRATA_BSDF_TYPE_CONDUCTOR:
-			{
-				float3 F = F_Schlick(CONDUCTOR_REFLECTIVITY(BSDF), CONDUCTOR_EDGECOLOR(BSDF), ContextNoLight.VoH);
-
-				if (bTopLayer)
-				{
-					TopLayer.Roughness += CONDUCTOR_ROUGHNESS0(BSDF);
-					TopLayer.AccumRoughnessCount += 1.0f;
-
-					const float SafeRoughness = MakeRoughnessSafe(CONDUCTOR_ROUGHNESS0(BSDF));
-					TopLayer.BaseColor += BSDFCoverage * EnvBRDFApprox(CONDUCTOR_REFLECTIVITY(BSDF), CONDUCTOR_EDGECOLOR(BSDF), SafeRoughness, saturate(ContextNoLight.NoV));;
-				}
-
-				LayerCoverage	+= BSDFCoverage;
-				LayerThroughput += BSDFCoverage * OpaqueBSDFThroughput;
-				break;
-			}
-
-			case STRATA_BSDF_TYPE_VOLUME:
-			{
-				const float3 Albedo					= VOLUME_ALBEDO(BSDF);
-				const float3 ExtinctionCoefficients = VOLUME_EXTINCTION(BSDF);
-				const float3 ScatteringCoefficients = Albedo * ExtinctionCoefficients;
-				const float3 AbsorptionCoefficients = ExtinctionCoefficients - ScatteringCoefficients;
-				const float  Thickness				= VOLUME_THICKNESS(BSDF);
-
-				const float3 SafeExtinctionThreshold = 0.000001f;
-				const float3 SafeExtinctionCoefficients = max(SafeExtinctionThreshold, ExtinctionCoefficients);
-
-				const float PathLength = Thickness / max(0.0001f, abs(ContextNoLight.NoV));
-				const float3 SafePathSegmentTransmittance = exp(-SafeExtinctionCoefficients * PathLength);
-				
-				if (bTopLayer)
-				{
-					TopLayer.Roughness += BSDFCoverageAvg * FullyRough; // Do not trigger SSR by using fully rough
-					TopLayer.AccumRoughnessCount += BSDFCoverageAvg;
+					SSS.Profile		= SLAB_SSSPROFILEID(BSDF);
+					SSS.RadiusScale	= SLAB_SSSPROFILERADIUS(BSDF);
+					SSS.BaseColor	= SLAB_ALBEDO(BSDF);
+					SSS.DMFPAlbedo	= SLAB_SSSDMFPALBEDO(BSDF);
+					SSS.DMFPRadius	= SLAB_SSSDMFPRADIUS(BSDF);
 				}
 
 				LayerCoverage	+= BSDFCoverage;
-				LayerThroughput += BSDFCoverage * SafePathSegmentTransmittance;
+				LayerThroughput += BSDFCoverage * OpaqueBSDFThroughput; // STRATA_TODO: take into accound SSS transmittance
 				break;
 			}
 
@@ -1269,8 +1103,8 @@ void PackStrataOut(
 
 				if (bTopLayer)
 				{
-					TopLayer.Roughness += SHEEN_ROUGHNESS(BSDF);
-					TopLayer.AccumRoughnessCount += 1.0f;
+					TopLayer.Roughness += BSDFCoverageAvg * SHEEN_ROUGHNESS(BSDF);
+					TopLayer.AccumRoughnessCount += BSDFCoverageAvg;
 
 					const float SafeRoughness = MakeRoughnessSafe(SHEEN_ROUGHNESS(BSDF));
 					TopLayer.BaseColor += BSDFCoverage * EnvBRDFApprox(F0, SafeRoughness, saturate(ContextNoLight.NoV));
@@ -1415,50 +1249,25 @@ void PackStrataOut(
 			const uint BSDFType = BSDF_GETTYPE(BSDF);
 			switch (BSDF_GETTYPE(BSDF))
 			{
-			case STRATA_BSDF_TYPE_DIFFUSE:
+			case STRATA_BSDF_TYPE_SLAB:
 			{
-				STRATA_STORE_UINT1(PackAlbedoLinearToGamma2AlphaLinear(float4(DIFFUSE_ALBEDO(BSDF), DIFFUSE_ROUGHNESS(BSDF))));
+				STRATA_STORE_UINT1(PackAlbedoLinearToGamma2AlphaLinear(float4(SLAB_ALBEDO(BSDF), SLAB_ROUGHNESS0(BSDF))));
+				STRATA_STORE_UINT1(PackRGBA8(float4(SLAB_REFLECTIVITY(BSDF), SLAB_METALLIC(BSDF), SLAB_ROUGHNESS1(BSDF), 0)));
+
 				if (BSDF_GETHASSCATTERING(BSDF))
 				{
 					if (BSDF_GETHASSSSPROFILE(BSDF))
 					{
-						STRATA_STORE_UINT1(PackRGBA8(float4(DIFFUSE_SSSPROFILEID(BSDF), DIFFUSE_SSSPROFILERADIUS(BSDF), 0, 0)));
+						STRATA_STORE_UINT1(PackRGBA8(float4(SLAB_SSSPROFILEID(BSDF), SLAB_SSSPROFILERADIUS(BSDF), 0, 0))); // STRATA_TODO SLAB_SSSPROFILERADIUS is in [0,1]?
 					}
 					else
 					{
-						STRATA_STORE_UINT1(PackAlbedoLinearToGamma2AlphaLinear(float4(DIFFUSE_SSSDMFPALBEDO(BSDF), DIFFUSE_SSSDMFPRADIUS(BSDF) / SSS_PROFILE_PERPIXEL_MAXRADIUS)));
+						STRATA_STORE_UINT1(PackAlbedoLinearToGamma2AlphaLinear(float4(SLAB_SSSDMFPALBEDO(BSDF), SLAB_SSSDMFPRADIUS(BSDF) / SSS_PROFILE_PERPIXEL_MAXRADIUS)));
 					}
 				}
 				// 8 bytes
 			}
 			break;
-			case STRATA_BSDF_TYPE_DIELECTRIC:
-			{
-				STRATA_STORE_UINT1(PackR11G11B10F(DIELECTRIC_TINT(BSDF)));
-				STRATA_STORE_UINT1(PackR11G11B10F(float3(DIELECTRIC_ROUGHNESS0(BSDF), DIELECTRIC_ROUGHNESS1(BSDF), DIELECTRIC_IOR(BSDF))));
-				// 8 bytes
-			}
-			break;
-			case STRATA_BSDF_TYPE_CONDUCTOR:
-			{
-				STRATA_STORE_UINT1(PackAlbedoLinearToGamma2AlphaLinear(float4(CONDUCTOR_REFLECTIVITY(BSDF), CONDUCTOR_ROUGHNESS0(BSDF))));
-				STRATA_STORE_UINT1(PackAlbedoLinearToGamma2AlphaLinear(float4(CONDUCTOR_EDGECOLOR(BSDF),    CONDUCTOR_ROUGHNESS1(BSDF))));
-				// 8 bytes
-			}
-			break;
-			case STRATA_BSDF_TYPE_VOLUME:
-			{
-				// We simply multiply the absorption and scattering coefficient by the thickness to simplify the material to be a normalise medium of thickness 1 meter.
-				// This gives the same result and reduces memory usage.
-				const float Thickness = VOLUME_THICKNESS(BSDF);
-				STRATA_STORE_UINT1(PackR11G11B10F(VOLUME_EXTINCTION(BSDF) * Thickness));
-				if (BSDF_GETHASSCATTERING(BSDF))
-				{
-					STRATA_STORE_UINT1(PackAlbedoLinearToGamma2AlphaLinear(float4(VOLUME_ALBEDO(BSDF), VOLUME_ANISOTROPY(BSDF) * 0.5f + 0.5f)));
-				}
-				// 8 bytes
-			}
-			break;
 			case STRATA_BSDF_TYPE_SHEEN:
 			{
 				STRATA_STORE_UINT1(PackAlbedoLinearToGamma2AlphaLinear(float4(SHEEN_ALBEDO(BSDF), SHEEN_ROUGHNESS(BSDF))));
@@ -1546,49 +1355,31 @@ FStrataBSDF UnpackStrataBSDFIn(ByteAddressBuffer StrataBuffer, inout uint ByteOf
 	BSDF_SETTYPE(OutBSDF, BSDFType);
 	switch (BSDFType)
 	{
-	case STRATA_BSDF_TYPE_DIFFUSE:
+	case STRATA_BSDF_TYPE_SLAB:
 	{
-		StrataLoad_AlbedoGamma2ToLinearAlphaLinear(StrataBuffer, ByteOffset, DIFFUSE_ALBEDO(OutBSDF), DIFFUSE_ROUGHNESS(OutBSDF));
+		StrataLoad_AlbedoGamma2ToLinearAlphaLinear(StrataBuffer, ByteOffset, SLAB_ALBEDO(OutBSDF), SLAB_ROUGHNESS0(OutBSDF));
+		{
+			float4 Data = UnpackRGBA8(StrataLoadUint1(StrataBuffer, ByteOffset));
+			SLAB_REFLECTIVITY(OutBSDF) = Data.x;
+			SLAB_METALLIC(OutBSDF) = Data.y;
+			SLAB_ROUGHNESS1(OutBSDF) = Data.z;
+		}
 		if (BSDF_GETHASSCATTERING(OutBSDF))
 		{
 			if (BSDF_GETHASSSSPROFILE(OutBSDF))
 			{
 				float4 Data = UnpackRGBA8(StrataLoadUint1(StrataBuffer, ByteOffset));
-				DIFFUSE_SSSPROFILEID(OutBSDF)	  = Data.r;
-				DIFFUSE_SSSPROFILERADIUS(OutBSDF) = Data.g;
+				SLAB_SSSPROFILEID(OutBSDF) = Data.r;
+				SLAB_SSSPROFILERADIUS(OutBSDF) = Data.g;
 			}
 			else
 			{
-				StrataLoad_AlbedoGamma2ToLinearAlphaLinear(StrataBuffer, ByteOffset, DIFFUSE_SSSDMFPALBEDO(OutBSDF), DIFFUSE_SSSDMFPRADIUS(OutBSDF));
-				DIFFUSE_SSSDMFPRADIUS(OutBSDF) *= SSS_PROFILE_PERPIXEL_MAXRADIUS;
+				StrataLoad_AlbedoGamma2ToLinearAlphaLinear(StrataBuffer, ByteOffset, SLAB_SSSDMFPALBEDO(OutBSDF), SLAB_SSSDMFPRADIUS(OutBSDF));
+				SLAB_SSSDMFPRADIUS(OutBSDF) *= SSS_PROFILE_PERPIXEL_MAXRADIUS;
 			}
 		}
 	}
 	break;
-	case STRATA_BSDF_TYPE_DIELECTRIC:
-	{
-		StrataLoad_R11G11B10F			(StrataBuffer, ByteOffset, DIELECTRIC_TINT(OutBSDF));
-		StrataLoad_R11_G11_B10F			(StrataBuffer, ByteOffset, DIELECTRIC_ROUGHNESS0(OutBSDF), DIELECTRIC_ROUGHNESS1(OutBSDF), DIELECTRIC_IOR(OutBSDF));
-	}
-	break;
-	case STRATA_BSDF_TYPE_CONDUCTOR:
-	{
-		StrataLoad_AlbedoGamma2ToLinearAlphaLinear(StrataBuffer, ByteOffset, CONDUCTOR_REFLECTIVITY(OutBSDF),	CONDUCTOR_ROUGHNESS0(OutBSDF));
-		StrataLoad_AlbedoGamma2ToLinearAlphaLinear(StrataBuffer, ByteOffset, CONDUCTOR_EDGECOLOR(OutBSDF),		CONDUCTOR_ROUGHNESS1(OutBSDF));
-	}
-	break;
-	case STRATA_BSDF_TYPE_VOLUME:
-	{
-		StrataLoad_R11G11B10F			(StrataBuffer, ByteOffset, VOLUME_EXTINCTION(OutBSDF));
-		if (BSDF_GETHASSCATTERING(OutBSDF))
-		{
-			StrataLoad_AlbedoGamma2ToLinearAlphaLinear(StrataBuffer, ByteOffset, VOLUME_ALBEDO(OutBSDF), VOLUME_ANISOTROPY(OutBSDF));
-			VOLUME_ANISOTROPY(OutBSDF) = VOLUME_ANISOTROPY(OutBSDF) * 2.0f - 1.0f;
-		}
-		// else, albedo and anisotropy have already cleared out to zero when creating OutBSDF.
-		VOLUME_THICKNESS(OutBSDF) = 1.0f;
-	}
-	break;
 	case STRATA_BSDF_TYPE_SHEEN:
 	{
 		StrataLoad_AlbedoGamma2ToLinearAlphaLinear(StrataBuffer, ByteOffset, SHEEN_ALBEDO(OutBSDF), SHEEN_ROUGHNESS(OutBSDF));
diff --git a/Engine/Shaders/Private/Strata/StrataDeferredLighting.ush b/Engine/Shaders/Private/Strata/StrataDeferredLighting.ush
--- a/Engine/Shaders/Private/Strata/StrataDeferredLighting.ush
+++ b/Engine/Shaders/Private/Strata/StrataDeferredLighting.ush
@@ -63,18 +63,11 @@ void StrataPreprocessAreaLightIntegration(
 		float3 SpecularColor = 0;
 		switch (BSDF_GETTYPE(BSDFContext.BSDF))
 		{
-			case STRATA_BSDF_TYPE_DIELECTRIC:
+			case STRATA_BSDF_TYPE_SLAB:
 			{
-				SafeRoughness0 = MakeRoughnessSafe(DIELECTRIC_ROUGHNESS0(BSDFContext.BSDF), View.MinRoughness);
-				SafeRoughness1 = MakeRoughnessSafe(DIELECTRIC_ROUGHNESS1(BSDFContext.BSDF), View.MinRoughness);
-				SpecularColor = DielectricIorToF0(DIELECTRIC_IOR(BSDFContext.BSDF)) * DIELECTRIC_TINT(BSDFContext.BSDF);
-			}
-			break;
-			case STRATA_BSDF_TYPE_CONDUCTOR:
-			{
-				SafeRoughness0 = MakeRoughnessSafe(CONDUCTOR_ROUGHNESS0(BSDFContext.BSDF), View.MinRoughness);
-				SafeRoughness1 = MakeRoughnessSafe(CONDUCTOR_ROUGHNESS1(BSDFContext.BSDF), View.MinRoughness);
-				SpecularColor = CONDUCTOR_REFLECTIVITY(BSDFContext.BSDF);
+				SafeRoughness0 = MakeRoughnessSafe(SLAB_ROUGHNESS0(BSDFContext.BSDF), View.MinRoughness);
+				SafeRoughness1 = MakeRoughnessSafe(SLAB_ROUGHNESS1(BSDFContext.BSDF), View.MinRoughness);
+				SpecularColor = ComputeF0(SLAB_REFLECTIVITY(BSDFContext.BSDF), SLAB_ALBEDO(BSDFContext.BSDF), SLAB_METALLIC(BSDFContext.BSDF));
 			}
 			break;
 		}
diff --git a/Engine/Shaders/Private/Strata/StrataEvaluation.ush b/Engine/Shaders/Private/Strata/StrataEvaluation.ush
--- a/Engine/Shaders/Private/Strata/StrataEvaluation.ush
+++ b/Engine/Shaders/Private/Strata/StrataEvaluation.ush
@@ -55,24 +55,9 @@ void StrataGetBSDFRoughness(in FStrataBSDF BSDF, inout float R)
 	const uint BSDFType = BSDF_GETTYPE(BSDF);
 	switch (BSDFType)
 	{
-	case STRATA_BSDF_TYPE_DIFFUSE:
+	case STRATA_BSDF_TYPE_SLAB:
 	{
-		R = DIFFUSE_ROUGHNESS(BSDF);
-		break;
-	}
-	case STRATA_BSDF_TYPE_DIELECTRIC:
-	{
-		R = DIELECTRIC_ROUGHNESS0(BSDF);
-		break;
-	}
-	case STRATA_BSDF_TYPE_CONDUCTOR:
-	{
-		R = CONDUCTOR_ROUGHNESS0(BSDF);
-		break;
-	}
-	case STRATA_BSDF_TYPE_VOLUME:
-	{
-		R = 0;
+		R = SLAB_ROUGHNESS0(BSDF);
 		break;
 	}
 	case STRATA_BSDF_TYPE_SHEEN:
@@ -149,7 +134,9 @@ struct FStrataEvaluateResult
 	float3 DiffusePathValue;
 	float3 SpecularPathValue;
 	float3 TransmissionPathValue;
-	float  PDF;
+	float  DiffusePDF;
+	float  SpecularPDF;
+	float  TransmissionPDF;
 	float3 Throughput;					// Throughput to the next layer
 	bool   bSubsurface;					// True if we need to separate the subsurface light contribution for the screen space diffusion process.
 	bool   bApplyProjectedSolidAngle;	// True if the Saturate(NoL) factor should be applied or not over the value of the BSDF
@@ -170,25 +157,24 @@ FStrataEvaluateResult StrataEvaluateBSDF(FStrataBSDFContext BSDFContext, FShadow
 	const uint BSDFType = BSDF_GETTYPE(BSDFContext.BSDF);
 	switch (BSDFType)
 	{
-		case STRATA_BSDF_TYPE_DIFFUSE:
+		case STRATA_BSDF_TYPE_SLAB:
 		{
-			const float SafeRoughness = MakeRoughnessSafe(DIFFUSE_ROUGHNESS(BSDFContext.BSDF));
+			const float3 BaseColor = SLAB_ALBEDO(BSDFContext.BSDF);
+			const float  Metallic = SLAB_METALLIC(BSDFContext.BSDF);
+			const float3 DiffuseColor = BaseColor - BaseColor * Metallic;
+			const float SafeRoughness = MakeRoughnessSafe(SLAB_ROUGHNESS0(BSDFContext.BSDF));
 			float a2 = Pow4(SafeRoughness);
 
-			Sample.DiffusePathValue	= Diffuse_Chan(DIFFUSE_ALBEDO(BSDFContext.BSDF), a2, BSDFContext.SatNoV, BSDFContext.SatNoL, BSDFContext.Context.VoH, BSDFContext.Context.NoH);
-			Sample.PDF				= BSDFContext.SatNoV / PI;
-			Sample.Throughput		= OpaqueBSDFThroughput;
+			Sample.DiffusePathValue	= Diffuse_Chan(DiffuseColor, a2, BSDFContext.SatNoV, BSDFContext.SatNoL, BSDFContext.Context.VoH, BSDFContext.Context.NoH);
+			Sample.DiffusePDF		= BSDFContext.SatNoV / PI;
 			Sample.bSubsurface		= BSDF_GETHASSCATTERING(BSDFContext.BSDF) > 0;
-			break;
-		}
-		case STRATA_BSDF_TYPE_DIELECTRIC:
-		case STRATA_BSDF_TYPE_CONDUCTOR:
-		{
-			const bool bIsDieletric = BSDFType == STRATA_BSDF_TYPE_DIELECTRIC;
-			const float3 F0  = bIsDieletric ? DIELECTRIC_TINT(BSDFContext.BSDF) * DielectricIorToF0(DIELECTRIC_IOR(BSDFContext.BSDF)) : CONDUCTOR_REFLECTIVITY(BSDFContext.BSDF);
-			const float3 F90 = bIsDieletric ? DIELECTRIC_TINT(BSDFContext.BSDF) : CONDUCTOR_EDGECOLOR(BSDFContext.BSDF);
-			const float SafeRoughness0 = MakeRoughnessSafe(bIsDieletric ? DIELECTRIC_ROUGHNESS0(BSDFContext.BSDF) : CONDUCTOR_ROUGHNESS0(BSDFContext.BSDF));
-			const float SafeRoughness1 = MakeRoughnessSafe(bIsDieletric ? DIELECTRIC_ROUGHNESS1(BSDFContext.BSDF) : CONDUCTOR_ROUGHNESS1(BSDFContext.BSDF));
+
+
+			const float3 F0 = ComputeF0(SLAB_REFLECTIVITY(BSDFContext.BSDF), BaseColor, Metallic);
+			// STRATA_TODO F82
+
+			const float SafeRoughness0 = MakeRoughnessSafe(SLAB_ROUGHNESS0(BSDFContext.BSDF));
+			const float SafeRoughness1 = MakeRoughnessSafe(SLAB_ROUGHNESS1(BSDFContext.BSDF));
 
 			const bool bHasAnisotropy = BSDF_GETANISOTROPY(BSDFContext.BSDF);
 			float D = 0;
@@ -213,42 +199,15 @@ FStrataEvaluateResult StrataEvaluateBSDF(FStrataBSDFContext BSDFContext, FShadow
 				Vis = Vis_SmithJointApprox(a2, BSDFContext.Context.NoV, BSDFContext.SatNoL);
 				PDF = VisibleGGXPDF(BSDFContext.TangentV, BSDFContext.TangentH, a2);
 			}
-			float3 F = F_Schlick(F0, F90, BSDFContext.Context.VoH);
+			float3 F = F_Schlick(F0, BSDFContext.Context.VoH);
 
 			Sample.SpecularPathValue= D * Vis * F;
-			Sample.PDF				= PDF;
-			Sample.Throughput		= bIsDieletric ? (1.0f - F) : OpaqueBSDFThroughput; // STRATA_TODO (Kulla's (1-FDG))
-			break;
-		}
-		case STRATA_BSDF_TYPE_VOLUME:
-		{
-			const float3 Albedo = VOLUME_ALBEDO(BSDFContext.BSDF);
-			const float3 ExtinctionCoefficients = VOLUME_EXTINCTION(BSDFContext.BSDF);
-			const float3 ScatteringCoefficients = Albedo * ExtinctionCoefficients;
-			const float3 AbsorptionCoefficients = ExtinctionCoefficients - ScatteringCoefficients;
-			const float  Thickness = VOLUME_THICKNESS(BSDFContext.BSDF);
-			const float  Phase = HenyeyGreensteinPhase(VOLUME_ANISOTROPY(BSDFContext.BSDF), BSDFContext.Context.VoL);
-
-			const float3 SafeExtinctionThreshold = 0.000001f;
-			const float3 SafeExtinctionCoefficients = max(SafeExtinctionThreshold, ExtinctionCoefficients);
-
-			const float PathLength = Thickness / max(0.0001f, abs(BSDFContext.Context.NoV));
-			const float3 SafePathSegmentTransmittance = exp(-SafeExtinctionCoefficients * PathLength);
-			const float3 NormLuminanceTransfert = 1.0f * ScatteringCoefficients;
-			float3 NormLuminanceSingleScattering = (NormLuminanceTransfert - NormLuminanceTransfert * SafePathSegmentTransmittance) / SafeExtinctionCoefficients;
-
-#if 1
-			const float3 SelfShadowTransmittance = 1.0;// BSDFContext.SatNoL;
-#else
-			const float LightPathLength = Thickness / max(0.0001f, BSDFContext.SatNoL);
-			const float3 SelfShadowTransmittance = BSDFContext.SatNoL <= 0.0 ? 0.0 : exp(-ExtinctionCoefficients * LightPathLength);
-#endif
-
-			Sample.DiffusePathValue = NormLuminanceSingleScattering * Phase * SelfShadowTransmittance;
-			Sample.PDF				= Phase;
-			Sample.Throughput		= SafePathSegmentTransmittance;
+			Sample.SpecularPDF		= PDF;
+
+			Sample.Throughput = OpaqueBSDFThroughput;
 			break;
 		}
+
 		case STRATA_BSDF_TYPE_SHEEN:
 		{
 			const float3 F0  = SHEEN_ALBEDO(BSDFContext.BSDF);
@@ -263,10 +222,11 @@ FStrataEvaluateResult StrataEvaluateBSDF(FStrataBSDFContext BSDFContext, FShadow
 			float3 F  = F_Schlick(F0, BSDFContext.Context.VoH);
 
 			Sample.DiffusePathValue	= D * Vis * F;
-			Sample.PDF				= BSDFContext.SatNoV / PI; // Per "Production Friendly Microfacet Sheen BRDF", hemispherical sampling give good result as the roughness is usually high.
+			Sample.DiffusePDF		= BSDFContext.SatNoV / PI; // Per "Production Friendly Microfacet Sheen BRDF", hemispherical sampling give good result as the roughness is usually high.
 			Sample.Throughput		= (1.0f - F);
 			break;
 		}
+
 		case STRATA_BSDF_TYPE_HAIR:
 		{
 			FGBufferData GBuffer	= (FGBufferData)0;
@@ -282,7 +242,7 @@ FStrataEvaluateResult StrataEvaluateBSDF(FStrataBSDFContext BSDFContext, FShadow
 			float Area = 0.0f;
 			uint2 Random = uint2(0, 0);
 			Sample.TransmissionPathValue= HairShading(GBuffer, BSDFContext.L, BSDFContext.V, BSDFContext.N, ShadowTerms.TransmissionShadow, ShadowTerms.HairTransmittance, BacklitEnabled, Area, Random);
-			Sample.PDF					= 1.0f / (4.0f * PI);		// STRATA_TODO this currently match the uniform sphere sampling from StrataImportanceSampleBSDF
+			Sample.TransmissionPDF		= 1.0f / (4.0f * PI);		// STRATA_TODO this currently match the uniform sphere sampling from StrataImportanceSampleBSDF
 			Sample.Throughput			= OpaqueBSDFThroughput; 
 		}
 		break;
@@ -297,7 +257,6 @@ FStrataEvaluateResult StrataEvaluateBSDF(FStrataBSDFContext BSDFContext, FShadow
 	return Sample;
 }
 
-
 FStrataEvaluateResult StrataEvaluateBSDF(FStrataBSDFContext BSDFContext)
 {
 	FShadowTerms Shadow = { 1, 1, 1, InitHairTransmittanceData() };
@@ -306,8 +265,11 @@ FStrataEvaluateResult StrataEvaluateBSDF(FStrataBSDFContext BSDFContext)
 
 struct FStrataImportanceSampleResult
 {
-	float3 Direction;
-	float ConeAngle;		// Cone angle is 1/PDF as defined by Lumen
+	float3 SpecularDirection;
+	float  SpecularConeAngle;		// bad approximate of Cone angle as 1/PDF as done by Lumen STRATA_TODO update when Lumen is pluged in...
+
+	float3 DiffuseDirection;
+	float  DiffuseConeAngle;		// bad approximate of Cone angle as 1/PDF as done by Lumen STRATA_TODO update when Lumen is pluged in...
 
 //	FStrataEvaluateResult Eval;
 };
@@ -331,61 +293,58 @@ FStrataImportanceSampleResult StrataImportanceSampleBSDF(FStrataPixelHeader Stra
 	const uint BSDFType = BSDF_GETTYPE(BSDF);
 	switch (BSDFType)
 	{
-		case STRATA_BSDF_TYPE_DIFFUSE:
+		case STRATA_BSDF_TYPE_SLAB:
 		{
+			const float3 BaseColor = SLAB_ALBEDO(BSDFContext.BSDF);
+			const float  Metallic = SLAB_METALLIC(BSDFContext.BSDF);
+
 			// STRATA_TODO take into account roughness
 			float4 ImportanceSample = CosineSampleHemisphere(E);
-			Sample.Direction = mul(ImportanceSample.xyz, BSDFContext.TangentBasis);
+			Sample.DiffuseDirection = mul(ImportanceSample.xyz, BSDFContext.TangentBasis);
 			PDF = ImportanceSample.w;
-			Sample.ConeAngle = 1.0f / PDF;
-			break;
-		}
-		case STRATA_BSDF_TYPE_DIELECTRIC:
-		case STRATA_BSDF_TYPE_CONDUCTOR:
-		{
-			const bool bIsDielectric = BSDFType == STRATA_BSDF_TYPE_DIELECTRIC;
-			float SafeRoughness0 = MakeRoughnessSafe(bIsDielectric ? DIELECTRIC_ROUGHNESS0(BSDFContext.BSDF) : CONDUCTOR_ROUGHNESS0(BSDFContext.BSDF));
-			float SafeRoughness1 = MakeRoughnessSafe(bIsDielectric ? DIELECTRIC_ROUGHNESS1(BSDFContext.BSDF) : CONDUCTOR_ROUGHNESS1(BSDFContext.BSDF));
+			Sample.DiffuseConeAngle = 1.0f / PDF;
 
+			const float3 F0 = ComputeF0(SLAB_REFLECTIVITY(BSDF), BaseColor, Metallic);
+			// STRATA_TODO F82
+
+			const float SafeRoughness0 = MakeRoughnessSafe(SLAB_ROUGHNESS0(BSDFContext.BSDF));
 			const bool bHasAnisotropy = BSDF_GETANISOTROPY(BSDFContext.BSDF);
 			float4 GGXSample = 0;
 			if (bHasAnisotropy)
 			{
+				const float SafeRoughness1 = MakeRoughnessSafe(SLAB_ROUGHNESS1(BSDFContext.BSDF));
 				GGXSample = ImportanceSampleVisibleGGX_aniso(UniformSampleDisk(E), float2(Pow2(SafeRoughness0), Pow2(SafeRoughness1)), BSDFContext.TangentV);
 			}
 			else
 			{
 				GGXSample = ImportanceSampleVisibleGGX(UniformSampleDisk(E), Pow4(SafeRoughness0), BSDFContext.TangentV);
 			}
-			
+
 			float3 WorldH = mul(GGXSample.xyz, BSDFContext.TangentBasis);
-			Sample.Direction = reflect(CameraVector, WorldH);
+			Sample.SpecularDirection = reflect(CameraVector, WorldH);
 			PDF = GGXSample.w;
-			Sample.ConeAngle = 1.0f / PDF;
-			break;
-		}
-		case STRATA_BSDF_TYPE_VOLUME:
-		{
-			float4 PhaseSample = ImportanceSampleHenyeyGreensteinPhase(E, VOLUME_ANISOTROPY(BSDFContext.BSDF));
-			Sample.Direction = mul(PhaseSample.xyz, BSDFContext.TangentBasis);
-			PDF = PhaseSample.w;
-			Sample.ConeAngle = 1.0f / PDF;
+			Sample.SpecularConeAngle = 1.0f / PDF;
 			break;
 		}
+
 		case STRATA_BSDF_TYPE_SHEEN:
 		{
 			// Per "Production Friendly Microfacet Sheen BRDF", uniform hemispherical sampling gives good result as the roughness is 
 			// usually high and avoid issue at grazing angle
 			float4 ImportanceSample = CosineSampleHemisphere(E);
-			Sample.Direction = mul(ImportanceSample.xyz, BSDFContext.TangentBasis);
+			Sample.DiffuseDirection = mul(ImportanceSample.xyz, BSDFContext.TangentBasis);
 			PDF = ImportanceSample.w;
-			Sample.ConeAngle = 1.0f / PDF;
+			Sample.DiffuseConeAngle = 1.0f / PDF;
 			break;
 		}
+
 		case STRATA_BSDF_TYPE_HAIR:
 		{
 			// STRATA_TODO do something better when we get there with Lumen, and evaluate the different researched solution (e.g. Importance Sampling for Physically-Based Hair Fiber Models)
-			UniformSampleSphere(E);
+			float4 ImportanceSample = CosineSampleHemisphere(E);
+			Sample.DiffuseDirection = mul(ImportanceSample.xyz, BSDFContext.TangentBasis);
+			PDF = ImportanceSample.w;
+			Sample.DiffuseConeAngle = 1.0f / PDF;
 			break;
 		}
 
@@ -422,28 +381,29 @@ FStrataEnvLightResult StrataEvaluateForEnvLight(FStrataBSDFContext BSDFContext,
 	const uint BSDFType = BSDF_GETTYPE(BSDFContext.BSDF);
 	switch (BSDFType)
 	{
-		case STRATA_BSDF_TYPE_DIFFUSE:
+		case STRATA_BSDF_TYPE_SLAB:
 		{
-			StrataEnvLightResult.DiffuseAlbedo = DIFFUSE_ALBEDO(BSDFContext.BSDF);
+			const float3 BaseColor = SLAB_ALBEDO(BSDFContext.BSDF);
+			const float  Metallic = SLAB_METALLIC(BSDFContext.BSDF);
+			const float3 DiffuseColor = BaseColor - BaseColor * Metallic;
+
+			StrataEnvLightResult.DiffuseAlbedo = DiffuseColor;
 			StrataEnvLightResult.DiffuseNormal = BSDFContext.N;
 			StrataEnvLightResult.bSubsurface = BSDF_GETHASSCATTERING(BSDFContext.BSDF) > 0;
-			break;
-		}
-		case STRATA_BSDF_TYPE_DIELECTRIC:
-		case STRATA_BSDF_TYPE_CONDUCTOR:
-		{
+
 			BRANCH
 			if (bEnableSpecular)
 			{
-				const bool bIsDielectric = STRATA_BSDF_TYPE_DIELECTRIC == BSDFType;
-				const float3 F0  = bIsDielectric ? DIELECTRIC_TINT(BSDFContext.BSDF) * DielectricIorToF0(DIELECTRIC_IOR(BSDFContext.BSDF)) : CONDUCTOR_REFLECTIVITY(BSDFContext.BSDF);
-				const float3 F90 = bIsDielectric ? DIELECTRIC_TINT(BSDFContext.BSDF) : CONDUCTOR_EDGECOLOR(BSDFContext.BSDF);
-				float SafeRoughness0 = MakeRoughnessSafe(bIsDielectric ? DIELECTRIC_ROUGHNESS0(BSDFContext.BSDF) : CONDUCTOR_ROUGHNESS0(BSDFContext.BSDF));
-				float SafeRoughness1 = MakeRoughnessSafe(bIsDielectric ? DIELECTRIC_ROUGHNESS1(BSDFContext.BSDF) : CONDUCTOR_ROUGHNESS1(BSDFContext.BSDF));
+				const float3 F0 = ComputeF0(SLAB_REFLECTIVITY(BSDFContext.BSDF), BaseColor, Metallic);
+				// STRATA_TODO F82
+
+				float SafeRoughness0 = MakeRoughnessSafe(SLAB_ROUGHNESS0(BSDFContext.BSDF));
 
 				const bool bHasAnisotropy = BSDF_GETANISOTROPY(BSDFContext.BSDF);
 				if (bHasAnisotropy)
 				{
+					const float SafeRoughness1 = MakeRoughnessSafe(SLAB_ROUGHNESS1(BSDFContext.BSDF));
+
 					// Modified the BSDF normal (and roughness)
 					float Anisotropy = 0;
 					GetAnisotropicFactor(SafeRoughness0, SafeRoughness1, Anisotropy, SafeRoughness0);
@@ -455,7 +415,7 @@ FStrataEnvLightResult StrataEvaluateForEnvLight(FStrataBSDFContext BSDFContext,
 					BSDFContext.R = 2 * dot(BSDFContext.V, BSDFContext.N) * BSDFContext.N - BSDFContext.V;
 				}
 
-				float3 EvalEnvBRDF = EnvBRDFApprox(F0, F90, SafeRoughness0, BSDFContext.SatNoV);
+				float3 EvalEnvBRDF = EnvBRDFApprox(F0, SafeRoughness0, BSDFContext.SatNoV);
 
 				StrataEnvLightResult.SpecularDirection = StrataGetOffSpecularPeakReflectionDir(BSDFContext.N, BSDFContext.R, SafeRoughness0);
 				StrataEnvLightResult.SpecularWeight = EvalEnvBRDF;
@@ -463,32 +423,14 @@ FStrataEnvLightResult StrataEvaluateForEnvLight(FStrataBSDFContext BSDFContext,
 			}
 			break;
 		}
-		case STRATA_BSDF_TYPE_VOLUME:
-		{
-			const float3 VolumeAlbedo = VOLUME_ALBEDO(BSDFContext.BSDF);
-			const float3 ExtinctionCoefficients = VOLUME_EXTINCTION(BSDFContext.BSDF);
-			const float3 ScatteringCoefficients = VolumeAlbedo * ExtinctionCoefficients;
-			const float3 AbsorptionCoefficients = ExtinctionCoefficients - ScatteringCoefficients;
-			const float  Thickness = VOLUME_THICKNESS(BSDFContext.BSDF);
 
-			const float3 SafeExtinctionThreshold = 0.000001f;
-			const float3 SafeExtinctionCoefficients = max(SafeExtinctionThreshold, ExtinctionCoefficients);
-
-			const float PathLength = Thickness / max(0.0001f, BSDFContext.SatNoV);
-			const float3 SafePathSegmentTransmittance = exp(-SafeExtinctionCoefficients * PathLength);
-			const float3 NormLuminanceTransfert = 1.0f * ScatteringCoefficients;
-			float3 NormLuminanceSingleScattering = (NormLuminanceTransfert - NormLuminanceTransfert * SafePathSegmentTransmittance) / SafeExtinctionCoefficients;
-
-			StrataEnvLightResult.DiffuseAlbedo = NormLuminanceSingleScattering;
-			StrataEnvLightResult.DiffuseNormal = BSDFContext.N;
-			break;
-		}
 		case STRATA_BSDF_TYPE_SHEEN:
 		{
 			StrataEnvLightResult.DiffuseAlbedo = SHEEN_ALBEDO(BSDFContext.BSDF);;
 			StrataEnvLightResult.DiffuseNormal = BSDFContext.N;
 			break;
 		}
+
 		case STRATA_BSDF_TYPE_HAIR:
 		{
 			FGBufferData GBuffer	= (FGBufferData)0;
diff --git a/Engine/Shaders/Private/Strata/StrataMaterialClassification.usf b/Engine/Shaders/Private/Strata/StrataMaterialClassification.usf
--- a/Engine/Shaders/Private/Strata/StrataMaterialClassification.usf
+++ b/Engine/Shaders/Private/Strata/StrataMaterialClassification.usf
@@ -41,86 +41,30 @@ void ComputePostProcessBSDFData(
 
 	switch (BSDF_GETTYPE(BSDF))
 	{
-	case STRATA_BSDF_TYPE_DIFFUSE:
+	case STRATA_BSDF_TYPE_SLAB:
 	{
+		const float3 F0 = ComputeF0(SLAB_REFLECTIVITY(BSDF), SLAB_ALBEDO(BSDF), SLAB_METALLIC(BSDF));
+
 		if (bTopLayer)
 		{
-			TopLayer.Roughness += BSDFCoverageAvg * FullyRough; // Do not trigger SSR by using fully rough
+			// Setup SSR for the top layer
+			TopLayer.Roughness += BSDFCoverageAvg * SLAB_ROUGHNESS0(BSDF);
 			TopLayer.AccumRoughnessCount += BSDFCoverageAvg;
+
+			const float SafeRoughness = MakeRoughnessSafe(SLAB_ROUGHNESS0(BSDF));
+			TopLayer.BaseColor += BSDFCoverage * EnvBRDFApprox(F0, SafeRoughness, saturate(ContextNoLight.NoV));
 		}
 
 		// Select the first/toppest layer with subsurface scattering
-		if (BSDF_GETHASSCATTERING(BSDF) && BSDFCoverageAvg > 0 && DIFFUSE_SSSPROFILERADIUS(BSDF) > 0 && !SSS.bHasSubsurface)
+		if (BSDF_GETHASSCATTERING(BSDF) && BSDFCoverageAvg > 0 && SLAB_SSSPROFILERADIUS(BSDF) > 0 && !SSS.bHasSubsurface)
 		{
 			SSS.bHasSubsurface = true;
 			SSS.bHasProfile	= BSDF_GETHASSSSPROFILE(BSDF);
-			SSS.Profile		= DIFFUSE_SSSPROFILEID(BSDF);
-			SSS.RadiusScale	= DIFFUSE_SSSPROFILERADIUS(BSDF);
-			SSS.BaseColor	= DIFFUSE_ALBEDO(BSDF);
-			SSS.DMFPAlbedo	= DIFFUSE_SSSDMFPALBEDO(BSDF);
-			SSS.DMFPRadius	= DIFFUSE_SSSDMFPRADIUS(BSDF);
-		}
-		break;
-	}
-
-	case STRATA_BSDF_TYPE_DIELECTRIC:
-	{
-		if (BSDF_GETHASSCATTERING(BSDF))
-		{
-			Classification.bHasSubsurface = true;
-			if (BSDF_GETHASSSSPROFILE(BSDF))
-			{
-				Classification.bHasSubsurfaceProfile = true;
-			}
-		}
-
-		float F0 = DielectricIorToF0(DIELECTRIC_IOR(BSDF));
-		float3 F = F_Schlick(F0, ContextNoLight.VoH);
-
-		if (bTopLayer)
-		{
-			TopLayer.Roughness += DIELECTRIC_ROUGHNESS0(BSDF);
-			TopLayer.AccumRoughnessCount += 1.0f;
-
-			const float SafeRoughness = MakeRoughnessSafe(DIELECTRIC_ROUGHNESS0(BSDF));
-			TopLayer.BaseColor += BSDFCoverage * EnvBRDFApprox(F0, SafeRoughness, saturate(ContextNoLight.NoV));
-		}
-		break;
-	}
-
-	case STRATA_BSDF_TYPE_CONDUCTOR:
-	{
-		float3 F = F_Schlick(CONDUCTOR_REFLECTIVITY(BSDF), CONDUCTOR_EDGECOLOR(BSDF), ContextNoLight.VoH);
-
-		if (bTopLayer)
-		{
-			TopLayer.Roughness += CONDUCTOR_ROUGHNESS0(BSDF);
-			TopLayer.AccumRoughnessCount += 1.0f;
-
-			const float SafeRoughness = MakeRoughnessSafe(CONDUCTOR_ROUGHNESS0(BSDF));
-			TopLayer.BaseColor += BSDFCoverage * EnvBRDFApprox(CONDUCTOR_REFLECTIVITY(BSDF), CONDUCTOR_EDGECOLOR(BSDF), SafeRoughness, saturate(ContextNoLight.NoV));;
-		}
-		break;
-	}
-
-	case STRATA_BSDF_TYPE_VOLUME:
-	{
-		const float3 Albedo					= VOLUME_ALBEDO(BSDF);
-		const float3 ExtinctionCoefficients = VOLUME_EXTINCTION(BSDF);
-		const float3 ScatteringCoefficients = Albedo * ExtinctionCoefficients;
-		const float3 AbsorptionCoefficients = ExtinctionCoefficients - ScatteringCoefficients;
-		const float  Thickness				= VOLUME_THICKNESS(BSDF);
-
-		const float3 SafeExtinctionThreshold = 0.000001f;
-		const float3 SafeExtinctionCoefficients = max(SafeExtinctionThreshold, ExtinctionCoefficients);
-
-		const float PathLength = Thickness / max(0.0001f, abs(ContextNoLight.NoV));
-		const float3 SafePathSegmentTransmittance = exp(-SafeExtinctionCoefficients * PathLength);
-				
-		if (bTopLayer)
-		{
-			TopLayer.Roughness += BSDFCoverageAvg * FullyRough; // Do not trigger SSR by using fully rough
-			TopLayer.AccumRoughnessCount += BSDFCoverageAvg;
+			SSS.Profile		= SLAB_SSSPROFILEID(BSDF);
+			SSS.RadiusScale	= SLAB_SSSPROFILERADIUS(BSDF);
+			SSS.BaseColor	= SLAB_ALBEDO(BSDF);
+			SSS.DMFPAlbedo	= SLAB_SSSDMFPALBEDO(BSDF);
+			SSS.DMFPRadius	= SLAB_SSSDMFPRADIUS(BSDF);
 		}
 		break;
 	}
diff --git a/Engine/Shaders/Private/Strata/StrataVisualize.usf b/Engine/Shaders/Private/Strata/StrataVisualize.usf
--- a/Engine/Shaders/Private/Strata/StrataVisualize.usf
+++ b/Engine/Shaders/Private/Strata/StrataVisualize.usf
@@ -213,9 +213,9 @@ void VisualizeMaterialPS(
 			switch (BSDF_GETTYPE(BSDF))
 			{
 
-			case STRATA_BSDF_TYPE_DIFFUSE:
-				PRINTBSDFCHAR(_D_); PRINTBSDFCHAR(_I_); PRINTBSDFCHAR(_F_); PRINTBSDFCHAR(_F_); PRINTBSDFCHAR(_U_); 
-				PRINTBSDFCHAR(_S_); PRINTBSDFCHAR(_E_); PRINTBSDFCHAR(_SPC_); PRINTBSDFCHAR(_SPC_); PRINTBSDFCHAR(_SPC_);
+			case STRATA_BSDF_TYPE_SLAB:
+				PRINTBSDFCHAR(_S_); PRINTBSDFCHAR(_L_); PRINTBSDFCHAR(_A_); PRINTBSDFCHAR(_B_); PRINTBSDFCHAR(_SPC_);
+				PRINTBSDFCHAR(_SPC_); PRINTBSDFCHAR(_SPC_); PRINTBSDFCHAR(_SPC_); PRINTBSDFCHAR(_SPC_); PRINTBSDFCHAR(_SPC_);
 				if (BSDF_GETGREYWEIGHT(BSDF))
 				{
 					PRINTBSDFGREYWEIGHT();
@@ -231,124 +231,52 @@ void VisualizeMaterialPS(
 				NEXTLINE;
 
 				HORIZONTALTAB; PRINTPARAMCHAR(_A_); PRINTPARAMCHAR(_L_); PRINTPARAMCHAR(_B_); PRINTPARAMCHAR(_E_); PRINTPARAMCHAR(_D_); PRINTPARAMCHAR(_O_); 
-				HORIZONTALTAB; PRINTFLOAT3(DIFFUSE_ALBEDO(BSDF));
+				HORIZONTALTAB; PRINTFLOAT3(SLAB_ALBEDO(BSDF));
 				NEXTLINE;
-				HORIZONTALTAB; PRINTPARAMCHAR(_R_); PRINTPARAMCHAR(_O_); PRINTPARAMCHAR(_U_); PRINTPARAMCHAR(_G_); PRINTPARAMCHAR(_H_); PRINTPARAMCHAR(_SPC_);
-				HORIZONTALTAB; PRINTFLOAT1(DIFFUSE_ROUGHNESS(BSDF));
+
+				HORIZONTALTAB; PRINTPARAMCHAR(_R_); PRINTPARAMCHAR(_E_); PRINTPARAMCHAR(_F_); PRINTPARAMCHAR(_L_); PRINTPARAMCHAR(_E_); PRINTPARAMCHAR(_C_);
+				HORIZONTALTAB; PRINTFLOAT1(SLAB_REFLECTIVITY(BSDF));
+				NEXTLINE;
+				HORIZONTALTAB; PRINTPARAMCHAR(_M_); PRINTPARAMCHAR(_E_); PRINTPARAMCHAR(_T_); PRINTPARAMCHAR(_A_); PRINTPARAMCHAR(_L_); PRINTPARAMCHAR(_SPC_);
+				HORIZONTALTAB; PRINTFLOAT1(SLAB_METALLIC(BSDF));
 				NEXTLINE;
 
+				HORIZONTALTAB; PRINTPARAMCHAR(_R_); PRINTPARAMCHAR(_O_); PRINTPARAMCHAR(_U_); PRINTPARAMCHAR(_G_); PRINTPARAMCHAR(_H_); PRINTPARAMCHAR(_0_);
+				HORIZONTALTAB; PRINTFLOAT1(SLAB_ROUGHNESS0(BSDF));
+				NEXTLINE;
+				if (BSDF_GETANISOTROPY(BSDF))
+				{
+					HORIZONTALTAB; PRINTPARAMCHAR(_R_); PRINTPARAMCHAR(_O_); PRINTPARAMCHAR(_U_); PRINTPARAMCHAR(_G_); PRINTPARAMCHAR(_H_); PRINTPARAMCHAR(_1_);
+					HORIZONTALTAB; PRINTFLOAT1(SLAB_ROUGHNESS1(BSDF));
+					NEXTLINE;
+				}
+
+				// STRATA_TODO
+
 				if (BSDF_GETHASSCATTERING(BSDF))
 				{
 					if (BSDF_GETHASSSSPROFILE(BSDF))
 					{
 						HORIZONTALTAB; PRINTPARAMCHAR(_S_); PRINTPARAMCHAR(_S_); PRINTPARAMCHAR(_S_); PRINTPARAMCHAR(_SPC_); PRINTPARAMCHAR(_I_); PRINTPARAMCHAR(_D_); PRINTPARAMCHAR(_SPC_);
-						HORIZONTALTAB; PRINTFLOAT1(DIFFUSE_SSSPROFILEID(BSDF));
+						HORIZONTALTAB; PRINTFLOAT1(SLAB_SSSPROFILEID(BSDF));
 						NEXTLINE;
 						HORIZONTALTAB; PRINTPARAMCHAR(_S_); PRINTPARAMCHAR(_S_); PRINTPARAMCHAR(_S_); PRINTPARAMCHAR(_SPC_); PRINTPARAMCHAR(_R_); PRINTPARAMCHAR(_A_); PRINTPARAMCHAR(_SPC_);
-						HORIZONTALTAB; PRINTFLOAT1(DIFFUSE_SSSPROFILERADIUS(BSDF));
+						HORIZONTALTAB; PRINTFLOAT1(SLAB_SSSPROFILERADIUS(BSDF));
 						NEXTLINE;
 					}
 					else
 					{
 						HORIZONTALTAB; PRINTPARAMCHAR(_S_); PRINTPARAMCHAR(_S_); PRINTPARAMCHAR(_S_); PRINTPARAMCHAR(_SPC_); PRINTPARAMCHAR(_A_); PRINTPARAMCHAR(_L_); PRINTPARAMCHAR(_SPC_);
-						HORIZONTALTAB; PRINTFLOAT3(DIFFUSE_SSSDMFPALBEDO(BSDF));
+						HORIZONTALTAB; PRINTFLOAT3(SLAB_SSSDMFPALBEDO(BSDF));
 						NEXTLINE;
 						HORIZONTALTAB; PRINTPARAMCHAR(_S_); PRINTPARAMCHAR(_S_); PRINTPARAMCHAR(_S_); PRINTPARAMCHAR(_SPC_); PRINTPARAMCHAR(_R_); PRINTPARAMCHAR(_A_); PRINTPARAMCHAR(_SPC_);
-						HORIZONTALTAB; PRINTFLOAT1(DIFFUSE_SSSDMFPRADIUS(BSDF));
+						HORIZONTALTAB; PRINTFLOAT1(SLAB_SSSDMFPRADIUS(BSDF));
 						NEXTLINE;
 					}
 				}
 
 				break;
 
-			case STRATA_BSDF_TYPE_DIELECTRIC:
-				PRINTBSDFCHAR(_D_); PRINTBSDFCHAR(_I_); PRINTBSDFCHAR(_E_); PRINTBSDFCHAR(_L_); PRINTBSDFCHAR(_E_);
-				PRINTBSDFCHAR(_C_); PRINTBSDFCHAR(_T_); PRINTBSDFCHAR(_R_); PRINTBSDFCHAR(_I_); PRINTBSDFCHAR(_C_);
-				if (BSDF_GETGREYWEIGHT(BSDF))
-				{
-					PRINTBSDFGREYWEIGHT();
-				}
-				else
-				{
-					PRINTBSDFWEIGHT();
-				}
-				PRINTBSDFNORMALID();
-				PRINTBSDFANISO();
-				PRINTBSDFTOPLAY();
-				NEXTLINE;
-
-				HORIZONTALTAB; PRINTPARAMCHAR(_I_); PRINTPARAMCHAR(_O_); PRINTPARAMCHAR(_R_); PRINTPARAMCHAR(_SPC_); PRINTPARAMCHAR(_SPC_); PRINTPARAMCHAR(_SPC_);
-				HORIZONTALTAB; PRINTFLOAT1(DIELECTRIC_IOR(BSDF));
-				NEXTLINE;
-				HORIZONTALTAB; PRINTPARAMCHAR(_T_); PRINTPARAMCHAR(_I_); PRINTPARAMCHAR(_N_); PRINTPARAMCHAR(_T_); PRINTPARAMCHAR(_SPC_); PRINTPARAMCHAR(_SPC_);
-				HORIZONTALTAB; PRINTFLOAT3(DIELECTRIC_TINT(BSDF));
-				NEXTLINE;
-				HORIZONTALTAB; PRINTPARAMCHAR(_R_); PRINTPARAMCHAR(_O_); PRINTPARAMCHAR(_U_); PRINTPARAMCHAR(_G_); PRINTPARAMCHAR(_H_); PRINTPARAMCHAR(_SPC_);
-				HORIZONTALTAB; PRINTFLOAT1(DIELECTRIC_ROUGHNESS0(BSDF));
-				NEXTLINE;
-
-				break;
-
-			case STRATA_BSDF_TYPE_CONDUCTOR:
-				PRINTBSDFCHAR(_C_); PRINTBSDFCHAR(_O_); PRINTBSDFCHAR(_N_); PRINTBSDFCHAR(_D_); PRINTBSDFCHAR(_U_);
-				PRINTBSDFCHAR(_C_); PRINTBSDFCHAR(_T_); PRINTBSDFCHAR(_O_); PRINTBSDFCHAR(_R_); PRINTBSDFCHAR(_SPC_);
-				if (BSDF_GETGREYWEIGHT(BSDF))
-				{
-					PRINTBSDFGREYWEIGHT();
-				}
-				else
-				{
-					PRINTBSDFWEIGHT();
-				}
-				PRINTBSDFNORMALID();
-				PRINTBSDFANISO();
-				PRINTBSDFTOPLAY();
-				NEXTLINE;
-
-				HORIZONTALTAB; PRINTPARAMCHAR(_R_); PRINTPARAMCHAR(_E_); PRINTPARAMCHAR(_F_); PRINTPARAMCHAR(_L_); PRINTPARAMCHAR(_E_); PRINTPARAMCHAR(_C_);
-				HORIZONTALTAB; PRINTFLOAT3(CONDUCTOR_REFLECTIVITY(BSDF));
-				NEXTLINE;
-				HORIZONTALTAB; PRINTPARAMCHAR(_E_); PRINTPARAMCHAR(_D_); PRINTPARAMCHAR(_G_); PRINTPARAMCHAR(_E_); PRINTPARAMCHAR(_C_); PRINTPARAMCHAR(_SPC_);
-				HORIZONTALTAB; PRINTFLOAT3(CONDUCTOR_EDGECOLOR(BSDF));
-				NEXTLINE;
-				HORIZONTALTAB; PRINTPARAMCHAR(_R_); PRINTPARAMCHAR(_O_); PRINTPARAMCHAR(_U_); PRINTPARAMCHAR(_G_); PRINTPARAMCHAR(_H_); PRINTPARAMCHAR(_SPC_);
-				HORIZONTALTAB; PRINTFLOAT1(CONDUCTOR_ROUGHNESS0(BSDF));
-				NEXTLINE;
-
-				break;
-
-			case STRATA_BSDF_TYPE_VOLUME:
-				PRINTBSDFCHAR(_V_); PRINTBSDFCHAR(_O_); PRINTBSDFCHAR(_L_); PRINTBSDFCHAR(_U_); PRINTBSDFCHAR(_M_); 
-				PRINTBSDFCHAR(_E_); PRINTBSDFCHAR(_SPC_); PRINTBSDFCHAR(_SPC_); PRINTBSDFCHAR(_SPC_); PRINTBSDFCHAR(_SPC_);
-				if (BSDF_GETGREYWEIGHT(BSDF))
-				{
-					PRINTBSDFGREYWEIGHT();
-				}
-				else
-				{
-					PRINTBSDFWEIGHT();
-				}
-				PRINTBSDFNORMALID();
-				PRINTBSDFANISO();
-				PRINTBSDFTOPLAY();
-				PRINTBSDFHASSCAT();
-				NEXTLINE;
-				
-				HORIZONTALTAB; PRINTPARAMCHAR(_A_); PRINTPARAMCHAR(_L_); PRINTPARAMCHAR(_B_); PRINTPARAMCHAR(_E_); PRINTPARAMCHAR(_D_); PRINTPARAMCHAR(_O_); 
-				HORIZONTALTAB; PRINTFLOAT3(VOLUME_ALBEDO(BSDF));
-				NEXTLINE;
-				HORIZONTALTAB; PRINTPARAMCHAR(_E_); PRINTPARAMCHAR(_X_); PRINTPARAMCHAR(_T_); PRINTPARAMCHAR(_I_); PRINTPARAMCHAR(_N_); PRINTPARAMCHAR(_C_); 
-				HORIZONTALTAB; PRINTFLOAT3(VOLUME_EXTINCTION(BSDF));
-				NEXTLINE;
-				HORIZONTALTAB; PRINTPARAMCHAR(_A_); PRINTPARAMCHAR(_N_); PRINTPARAMCHAR(_I_); PRINTPARAMCHAR(_S_); PRINTPARAMCHAR(_O_); PRINTPARAMCHAR(_T_); 
-				HORIZONTALTAB; SIGNOFFLOAT1(VOLUME_ANISOTROPY(BSDF)); PRINTFLOAT1(abs(VOLUME_ANISOTROPY(BSDF)));
-				NEXTLINE;
-				HORIZONTALTAB; PRINTPARAMCHAR(_T_); PRINTPARAMCHAR(_H_); PRINTPARAMCHAR(_I_); PRINTPARAMCHAR(_C_); PRINTPARAMCHAR(_K_); PRINTPARAMCHAR(_N_); 
-				HORIZONTALTAB; PRINTFLOAT1(VOLUME_THICKNESS(BSDF) *1000.0f); PRINTPARAMCHAR(_SPC_); PRINTPARAMCHAR(_M_); PRINTPARAMCHAR(_I_); PRINTPARAMCHAR(_L_); PRINTPARAMCHAR(_L_); PRINTPARAMCHAR(_I_);
-				NEXTLINE;
-
-				break;
-
 			case STRATA_BSDF_TYPE_SHEEN:
 				PRINTBSDFCHAR(_S_); PRINTBSDFCHAR(_H_); PRINTBSDFCHAR(_E_); PRINTBSDFCHAR(_E_); PRINTBSDFCHAR(_N_); 
 				PRINTBSDFCHAR(_SPC_); PRINTBSDFCHAR(_SPC_); PRINTBSDFCHAR(_SPC_); PRINTBSDFCHAR(_SPC_); PRINTBSDFCHAR(_SPC_);
diff --git a/Engine/Shaders/Shared/StrataDefinitions.h b/Engine/Shaders/Shared/StrataDefinitions.h
--- a/Engine/Shaders/Shared/StrataDefinitions.h
+++ b/Engine/Shaders/Shared/StrataDefinitions.h
@@ -21,15 +21,12 @@
 #define STRATA_PACKED_NORMAL_STRIDE_BYTES	4
 
 
-#define STRATA_BSDF_TYPE_DIFFUSE			0
-#define STRATA_BSDF_TYPE_DIELECTRIC			1
-#define STRATA_BSDF_TYPE_CONDUCTOR			2
-#define STRATA_BSDF_TYPE_VOLUME				3
-#define STRATA_BSDF_TYPE_SHEEN				4
-#define STRATA_BSDF_TYPE_VOLUMETRICFOGCLOUD	5
-#define STRATA_BSDF_TYPE_UNLIT				6
-#define STRATA_BSDF_TYPE_HAIR				7
-#define STRATA_BSDF_TYPE_SINGLELAYERWATER	8
-// When more than 16 BSDF must exists, please update STATE_BITS_BSDF and FStrataClassification.ShadingModels packing in Strata.ush
+#define STRATA_BSDF_TYPE_SLAB				0
+#define STRATA_BSDF_TYPE_SHEEN				1
+#define STRATA_BSDF_TYPE_VOLUMETRICFOGCLOUD	2
+#define STRATA_BSDF_TYPE_UNLIT				3
+#define STRATA_BSDF_TYPE_HAIR				4
+#define STRATA_BSDF_TYPE_SINGLELAYERWATER	5
+// When more than 8 BSDF must exists, please update STATE_BITS_BSDF and FStrataClassification.ShadingModels packing in Strata.ush
 
 
diff --git a/Engine/Source/Runtime/Engine/Classes/Materials/MaterialExpressionStrata.h b/Engine/Source/Runtime/Engine/Classes/Materials/MaterialExpressionStrata.h
--- a/Engine/Source/Runtime/Engine/Classes/Materials/MaterialExpressionStrata.h
+++ b/Engine/Source/Runtime/Engine/Classes/Materials/MaterialExpressionStrata.h
@@ -36,135 +36,32 @@ class UMaterialExpressionStrataBSDF : public UMaterialExpression
 }*/
 
 UCLASS(MinimalAPI, collapsecategories, hidecategories = Object)
-class UMaterialExpressionStrataDiffuseBSDF : public UMaterialExpression // STRATA_TODO the single diffuse model to keep when we remove al lthe tests
+class UMaterialExpressionStrataSlabBSDF : public UMaterialExpression // STRATA_TODO the single diffuse model to keep when we remove al lthe tests
 {
 	GENERATED_UCLASS_BODY()
 
 	/**
-	 * Albedo (type = float3, unit = unitless)
+	 * Albedo (type = float3, unit = unitless, defaults to 0.18)
 	 */
 	UPROPERTY()
-	FExpressionInput Albedo;
-
-	/**
-	 * Roughness (type = float, unit = unitless)
-	 */
-	UPROPERTY()
-	FExpressionInput Roughness;
+	FExpressionInput BaseColor;
 
 	/**
-	 * Normal (type = float3, unit = unitless)
+	 * 0 means dielectric, 1 means conductor (type = float, unit = unitless, defaults to 0)
 	 */
 	UPROPERTY()
-	FExpressionInput Normal;
-
-	/**
-	 * Mean free path albedo. Only used when no sub-surface profile is provided. (type = float3, unit = unitless)
-	 */
-	UPROPERTY(meta = (DisplayName = "DMFP Albedo"))
-	FExpressionInput DiffuseMeanFreePathAlbedo;
-
-	/**
-	 * Mean free path radius. Only used when no sub-surface profile is provided. (type = float, unit = cm)
-	 */
-	UPROPERTY(meta = (DisplayName = "DMFP Radius"))
-	FExpressionInput DiffuseMeanFreePathRadius;
-
-	/**
-	 * SSS radius scale (type = float, unit = unitless)
-	 */
-	UPROPERTY(meta = (DisplayName = "DMFP Scale"))
-	FExpressionInput DiffuseMeanFreePathRadiusScale;
-
-	/** SubsurfaceProfile, for Screen Space Subsurface Scattering. The profile needs to be set up on both the Strata diffuse node, and the material node at the moment. */
-	UPROPERTY(EditAnywhere, BlueprintReadOnly, Category = Material, meta = (DisplayName = "Subsurface Profile"))
-	class USubsurfaceProfile* SubsurfaceProfile;
-
-	//~ Begin UMaterialExpression Interface
-#if WITH_EDITOR
-	virtual int32 Compile(class FMaterialCompiler* Compiler, int32 OutputIndex) override;
-	virtual void GetCaption(TArray<FString>& OutCaptions) const override;
-	virtual uint32 GetOutputType(int32 OutputIndex) override;
-	virtual uint32 GetInputType(int32 InputIndex) override;
-	virtual bool IsResultStrataMaterial(int32 OutputIndex) override;
-	virtual void GatherStrataMaterialInfo(FStrataMaterialInfo& StrataMaterialInfo, int32 OutputIndex) override;
-	virtual FName GetInputName(int32 InputIndex) const override;
-#endif
-	//~ End UMaterialExpression Interface
-};
-
-UCLASS(MinimalAPI, collapsecategories, hidecategories = Object)
-class UMaterialExpressionStrataDielectricBSDF : public UMaterialExpression
-{
-	GENERATED_UCLASS_BODY()
+	FExpressionInput Metallic;
 	
 	/**
-	 * The index of refraction of the surface (type = float, unit = unitless, default=1.5 for glass)
-	 */
-	UPROPERTY()
-	FExpressionInput IOR;
-
-	/**
-	 * A global color tint multiplied with the specular color, not physically based (type = float3, unit = unitless)
-	 */
-	UPROPERTY()
-	FExpressionInput Tint;
-		
-	/**
-	 * RoughnessX (type = float, unit = unitless)
-	 */
-	UPROPERTY()
-	FExpressionInput RoughnessX;
-		
-	/**
-	 * RoughnessY (type = float, unit = unitless), If not plugged in, RoughnessY is set to RoughnessX to get an isotropic behavior.
-	 */
-	UPROPERTY()
-	FExpressionInput RoughnessY;
-
-	/**
-	 * Normal (type = float3, unit = unitless)
+	 * Reflectivity when view direction is perpendicular to the surface, also known as F0 (type = float, unit = unitless, defaults to plastic 0.5)
 	 */
 	UPROPERTY()
-	FExpressionInput Normal;
-
-	/**
-	* Tangent (type = float3, unit = unitless)
-	*/
-	UPROPERTY()
-	FExpressionInput Tangent;
+	FExpressionInput Reflectivity;
 
-	//~ Begin UMaterialExpression Interface
-#if WITH_EDITOR
-	virtual int32 Compile(class FMaterialCompiler* Compiler, int32 OutputIndex) override;
-	virtual void GetCaption(TArray<FString>& OutCaptions) const override;
-	virtual uint32 GetOutputType(int32 OutputIndex) override;
-	virtual uint32 GetInputType(int32 InputIndex) override;
-	virtual bool IsResultStrataMaterial(int32 OutputIndex) override;
-	virtual void GatherStrataMaterialInfo(FStrataMaterialInfo& StrataMaterialInfo, int32 OutputIndex) override;
-#endif
-	//~ End UMaterialExpression Interface
-};
+	// STRATA_TODO: edge or F82 EdgeColor?
 
-UCLASS(MinimalAPI, collapsecategories, hidecategories = Object)
-class UMaterialExpressionStrataConductorBSDF : public UMaterialExpression
-{
-	GENERATED_UCLASS_BODY()
-		
 	/**
-	 * Reflectivity when view direction is perpendicular to the surface, also known as F0 (type = float3, unit = unitless, defaults to gold)
-	 */
-	UPROPERTY()
-	FExpressionInput Reflectivity;
-	
-	/**
-	 * Reflectivity when the view direction is tangent to the surface (type = float3, unit = unitless, defaults to gold)
-	 */
-	UPROPERTY()
-	FExpressionInput EdgeColor;
-		
-	/**
-	 * RoughnessX (type = float, unit = unitless)
+	 * RoughnessX (type = float, unit = unitless, defaults to 0)
 	 */
 	UPROPERTY()
 	FExpressionInput RoughnessX;
@@ -176,57 +73,32 @@ class UMaterialExpressionStrataConductorBSDF : public UMaterialExpression
 	FExpressionInput RoughnessY;
 
 	/**
-	 * Normal (type = float3, unit = unitless)
+	 * Normal (type = float3, unit = unitless, defaults to vertex normal)
 	 */
 	UPROPERTY()
 	FExpressionInput Normal;
 
 	/**
-	 * Tangent (type = float3, unit = unitless)
-	 */
+	* Tangent (type = float3, unit = unitless, defaults to vertex tangent)
+	*/
 	UPROPERTY()
 	FExpressionInput Tangent;
 
-	//~ Begin UMaterialExpression Interface
-#if WITH_EDITOR
-	virtual int32 Compile(class FMaterialCompiler* Compiler, int32 OutputIndex) override;
-	virtual void GetCaption(TArray<FString>& OutCaptions) const override;
-	virtual uint32 GetOutputType(int32 OutputIndex) override;
-	virtual uint32 GetInputType(int32 InputIndex) override;
-	virtual bool IsResultStrataMaterial(int32 OutputIndex) override;
-	virtual void GatherStrataMaterialInfo(FStrataMaterialInfo& StrataMaterialInfo, int32 OutputIndex) override;
-#endif
-	//~ End UMaterialExpression Interface
-};
-
-UCLASS(MinimalAPI, collapsecategories, hidecategories = Object)
-class UMaterialExpressionStrataVolumeBSDF : public UMaterialExpression
-{
-	GENERATED_UCLASS_BODY()
-		
 	/**
-	 * Albedo (type = float3, unit = unitless)
-	 */
-	UPROPERTY()
-	FExpressionInput Albedo;
-	
-	/**
-	 * The rate at which light is absorbed or scattered by the medium. Mean Free Path = 1 / Extinction. (type = float3, unit = 1/m)
+	 * Mean free path albedo. Only used when no sub-surface profile is provided. (type = float3, unit = unitless)
 	 */
-	UPROPERTY()
-	FExpressionInput Extinction;
+	UPROPERTY(meta = (DisplayName = "DMFP Albedo"))
+	FExpressionInput DiffuseMeanFreePathAlbedo;
 
 	/**
-	 * Anisotropy of the volume with values lower than 0 representing back-scattering, equal 0 representing isotropic scattering and greater than 0 representing forward scattering. (type = float, unit = unitless, defaults to 0)
+	 * Mean free path radius. Only used when no sub-surface profile is provided. (type = float, unit = cm)
 	 */
-	UPROPERTY()
-	FExpressionInput Anisotropy;
+	UPROPERTY(meta = (DisplayName = "DMFP Radius"))
+	FExpressionInput DiffuseMeanFreePathRadius;
 
-	/**
-	 * Thickness (type = float, unit = meters, default = 1mm)
-	 */
-	UPROPERTY()
-	FExpressionInput Thickness;
+	/** SubsurfaceProfile, for Screen Space Subsurface Scattering. The profile needs to be set up on both the Strata diffuse node, and the material node at the moment. */
+	UPROPERTY(EditAnywhere, BlueprintReadOnly, Category = Material, meta = (DisplayName = "Subsurface Profile"))
+	class USubsurfaceProfile* SubsurfaceProfile;
 
 	//~ Begin UMaterialExpression Interface
 #if WITH_EDITOR
@@ -236,6 +108,7 @@ class UMaterialExpressionStrataVolumeBSDF : public UMaterialExpression
 	virtual uint32 GetInputType(int32 InputIndex) override;
 	virtual bool IsResultStrataMaterial(int32 OutputIndex) override;
 	virtual void GatherStrataMaterialInfo(FStrataMaterialInfo& StrataMaterialInfo, int32 OutputIndex) override;
+	virtual FName GetInputName(int32 InputIndex) const override;
 #endif
 	//~ End UMaterialExpression Interface
 };
diff --git a/Engine/Source/Runtime/Engine/Private/Materials/HLSLMaterialTranslator.cpp b/Engine/Source/Runtime/Engine/Private/Materials/HLSLMaterialTranslator.cpp
--- a/Engine/Source/Runtime/Engine/Private/Materials/HLSLMaterialTranslator.cpp
+++ b/Engine/Source/Runtime/Engine/Private/Materials/HLSLMaterialTranslator.cpp
@@ -7536,62 +7536,24 @@ int32 FHLSLMaterialTranslator::StrataCreateAndRegisterNullMaterial()
 	return OutputCodeChunk;
 }
 
-int32 FHLSLMaterialTranslator::StrataDiffuseBSDF(int32 Albedo, int32 Roughness, int32 SSSProfileId, int32 SSSDMFPScale, int32 SSSDMFPAlbedo, int32 SSSDMFPRadius, int32 Normal, uint8 SharedNormalIndex)
+int32 FHLSLMaterialTranslator::StrataSlabBSDF(int32 Albedo, int32 Reflectivity, int32 Metallic, int32 RoughnessX, int32 RoughnessY, int32 SSSProfileId, int32 SSSDMFPAlbedo, int32 SSSDMFPRadius, int32 Normal, int32 Tangent, uint8 SharedNormalIndex)
 {
 	return AddCodeChunk(
-		MCT_Strata, TEXT("GetStrataDiffuseBSDF(%s, %s, %s, %s, %s, %s, %u) /* %s */"),
+		MCT_Strata, TEXT("GetStrataSlabBSDF(%s, %s, %s, float2(%s, %s), %s, %s, %s, %u, Parameters.SharedNormals.NormalTypes) /* Normal = %s ; Tangent = %s */"),
 		*GetParameterCode(Albedo),
-		*GetParameterCode(Roughness),
-		*GetParameterCode(SSSProfileId),
-		*GetParameterCode(SSSDMFPScale),
-		*GetParameterCode(SSSDMFPAlbedo),
-		*GetParameterCode(SSSDMFPRadius),
-		SharedNormalIndex,
-		*GetParameterCode(Normal)
-	);
-}
-
-int32 FHLSLMaterialTranslator::StrataDielectricBSDF(int32 RoughnessX, int32 RoughnessY, int32 IOR, int32 Tint, int32 Normal, int32 Tangent, uint8 SharedNormalIndex)
-{
-	return AddCodeChunk(
-		MCT_Strata, TEXT("GetStrataDielectricBSDF(float2(%s, %s), %s, %s, %u, Parameters.SharedNormals.NormalTypes) /* Normal:%s Tangent:%s */"),
-		*GetParameterCode(RoughnessX),
-		*GetParameterCode(RoughnessY),
-		*GetParameterCode(IOR),
-		*GetParameterCode(Tint),
-		SharedNormalIndex,
-		*GetParameterCode(Normal),
-		*GetParameterCode(Tangent)
-	);
-}
-
-int32 FHLSLMaterialTranslator::StrataConductorBSDF(int32 Reflectivity, int32 EdgeColor, int32 RoughnessX, int32 RoughnessY, int32 Normal, int32 Tangent, uint8 SharedNormalIndex)
-{
-	return AddCodeChunk(
-		MCT_Strata, TEXT("GetStrataConductorBSDF(%s, %s, float2(%s, %s), %u, Parameters.SharedNormals.NormalTypes) /* Normal:%s Tangent:%s */"),
 		*GetParameterCode(Reflectivity),
-		*GetParameterCode(EdgeColor),
+		*GetParameterCode(Metallic),
 		*GetParameterCode(RoughnessX),
 		*GetParameterCode(RoughnessY),
+		*GetParameterCode(SSSProfileId),
+		*GetParameterCode(SSSDMFPAlbedo),
+		*GetParameterCode(SSSDMFPRadius),
 		SharedNormalIndex,
 		*GetParameterCode(Normal),
 		*GetParameterCode(Tangent)
 	);
 }
 
-int32 FHLSLMaterialTranslator::StrataVolumeBSDF(int32 Albedo, int32 Extinction, int32 Anisotropy, int32 Thickness, int32 Normal, uint8 SharedNormalIndex)
-{
-	return AddCodeChunk(
-		MCT_Strata, TEXT("GetStrataVolumeBSDF(%s, %s, %s, %s, %u) /* %s */"),
-		*GetParameterCode(Albedo),
-		*GetParameterCode(Extinction),
-		*GetParameterCode(Anisotropy),
-		*GetParameterCode(Thickness),
-		SharedNormalIndex,
-		*GetParameterCode(Normal)
-	);
-}
-
 int32 FHLSLMaterialTranslator::StrataSheenBSDF(int32 Albedo, int32 Roughness, int32 Normal, uint8 SharedNormalIndex)
 {
 	return AddCodeChunk(
diff --git a/Engine/Source/Runtime/Engine/Private/Materials/HLSLMaterialTranslator.h b/Engine/Source/Runtime/Engine/Private/Materials/HLSLMaterialTranslator.h
--- a/Engine/Source/Runtime/Engine/Private/Materials/HLSLMaterialTranslator.h
+++ b/Engine/Source/Runtime/Engine/Private/Materials/HLSLMaterialTranslator.h
@@ -798,10 +798,7 @@ protected:
 
 	// Strata
 	virtual int32 StrataCreateAndRegisterNullMaterial() override;
-	virtual int32 StrataDiffuseBSDF(int32 Albedo, int32 Roughness, int32 SSSProfileId, int32 SSSDMFPRadiusScale, int32 SSSDMFPAlebdo, int32 SSSDMFPRadius, int32 Normal, uint8 SharedNormalIndex) override;
-	virtual int32 StrataDielectricBSDF(int32 RoughnessX, int32 RoughnessY, int32 IOR, int32 Tint, int32 Normal, int32 Tangent, uint8 SharedNormalIndex) override;
-	virtual int32 StrataConductorBSDF(int32 Reflectivity, int32 EdgeColor, int32 RoughnessX, int32 RoughnessY, int32 Normal, int32 Tangent, uint8 SharedNormalIndex) override;
-	virtual int32 StrataVolumeBSDF(int32 Albedo, int32 Extinction, int32 Anisotropy, int32 Thickness, int32 Normal, uint8 SharedNormalIndex) override;
+	virtual int32 StrataSlabBSDF(int32 Albedo, int32 Reflectivity, int32 Metallic, int32 RoughnessX, int32 RoughnessY, int32 SSSProfileId, int32 SSSDMFPAlbedo, int32 SSSDMFPRadius, int32 Normal, int32 Tangent, uint8 SharedNormalIndex) override;
 	virtual int32 StrataSheenBSDF(int32 Albedo, int32 Roughness, int32 Normal, uint8 SharedNormalIndex) override;
 	virtual int32 StrataVolumetricFogCloudBSDF(int32 Albedo, int32 Extinction, int32 Emissive, int32 AmbientOcclusion) override;
 	virtual int32 StrataUnlitBSDF(int32 Emissive, int32 TransmittanceColor) override;
diff --git a/Engine/Source/Runtime/Engine/Private/Materials/MaterialExpressions.cpp b/Engine/Source/Runtime/Engine/Private/Materials/MaterialExpressions.cpp
--- a/Engine/Source/Runtime/Engine/Private/Materials/MaterialExpressions.cpp
+++ b/Engine/Source/Runtime/Engine/Private/Materials/MaterialExpressions.cpp
@@ -19646,7 +19646,7 @@ static int32 CompileWithDefaultTangentWS(class FMaterialCompiler* Compiler, FExp
 
 #endif // WITH_EDITOR
 
-UMaterialExpressionStrataDiffuseBSDF::UMaterialExpressionStrataDiffuseBSDF(const FObjectInitializer& ObjectInitializer)
+UMaterialExpressionStrataSlabBSDF::UMaterialExpressionStrataSlabBSDF(const FObjectInitializer& ObjectInitializer)
 	: Super(ObjectInitializer)
 {
 	struct FConstructorStatics
@@ -19661,10 +19661,17 @@ UMaterialExpressionStrataDiffuseBSDF::UMaterialExpressionStrataDiffuseBSDF(const
 }
 
 #if WITH_EDITOR
-int32 UMaterialExpressionStrataDiffuseBSDF::Compile(class FMaterialCompiler* Compiler, int32 OutputIndex)
+int32 UMaterialExpressionStrataSlabBSDF::Compile(class FMaterialCompiler* Compiler, int32 OutputIndex)
 {
 	int32 NormalCodeChunk = CompileWithDefaultNormalWS(Compiler, Normal);
-	uint8 SharedNormalIndex = StrataCompilationInfoCreateSharedNormal(Compiler, NormalCodeChunk);
+
+	bool bDefaultTangentIsUsed = false;
+	int32 TangentCodeChunk = CompileWithDefaultTangentWS(Compiler, Tangent, &bDefaultTangentIsUsed);
+	uint8 SharedNormalIndex = bDefaultTangentIsUsed ? StrataCompilationInfoCreateSharedNormal(Compiler, NormalCodeChunk) : StrataCompilationInfoCreateSharedNormal(Compiler, NormalCodeChunk, TangentCodeChunk);
+
+	int32 RoughnessXCodeChunk = CompileWithDefaultFloat1(Compiler, RoughnessX, 0.0f);
+	// If not plugged in, RoughnessYCodeChunk is set to RoughnessXCodeChunk to get an isotropic behavior
+	int32 RoughnessYCodeChunk = CompileWithDefaultCodeChunk(Compiler, RoughnessY, RoughnessXCodeChunk);
 
 	int32 SSSProfileCodeChunk = INDEX_NONE;
 	const bool bHasScattering = SubsurfaceProfile != nullptr;
@@ -19674,51 +19681,63 @@ int32 UMaterialExpressionStrataDiffuseBSDF::Compile(class FMaterialCompiler* Com
 		SSSProfileCodeChunk = Compiler->ForceCast(Compiler->ScalarParameter(NameSubsurfaceProfile, 1.0f), MCT_Float1);
 	}
 
-	int32 OutputCodeChunk = Compiler->StrataDiffuseBSDF(
-		CompileWithDefaultFloat3(Compiler, Albedo,		0.18f, 0.18f, 0.18f),
-		CompileWithDefaultFloat1(Compiler, Roughness,	0.0f),
-		SSSProfileCodeChunk != INDEX_NONE ? SSSProfileCodeChunk : Compiler->Constant(0.0f),		
-		CompileWithDefaultFloat1(Compiler, DiffuseMeanFreePathRadiusScale,	0.0f),
+	int32 OutputCodeChunk = Compiler->StrataSlabBSDF(
+		CompileWithDefaultFloat3(Compiler, BaseColor, 0.18f, 0.18f, 0.18f),
+		CompileWithDefaultFloat1(Compiler, Reflectivity, 0.5f),
+		CompileWithDefaultFloat1(Compiler, Metallic, 0.0f),
+		RoughnessXCodeChunk,
+		RoughnessYCodeChunk,
+		SSSProfileCodeChunk != INDEX_NONE ? SSSProfileCodeChunk : Compiler->Constant(0.0f),	
 		CompileWithDefaultFloat3(Compiler, DiffuseMeanFreePathAlbedo, 0.0f, 0.0f, 0.0f),
 		CompileWithDefaultFloat1(Compiler, DiffuseMeanFreePathRadius, 0.0f),
 		NormalCodeChunk,
+		TangentCodeChunk,
 		SharedNormalIndex);
-	StrataCompilationInfoCreateSingleBSDFMaterial(Compiler, OutputCodeChunk, SharedNormalIndex, STRATA_BSDF_TYPE_DIFFUSE, bHasScattering);
+	StrataCompilationInfoCreateSingleBSDFMaterial(Compiler, OutputCodeChunk, SharedNormalIndex, STRATA_BSDF_TYPE_SLAB, bHasScattering);
 
 	return OutputCodeChunk;
 }
 
-void UMaterialExpressionStrataDiffuseBSDF::GetCaption(TArray<FString>& OutCaptions) const
+void UMaterialExpressionStrataSlabBSDF::GetCaption(TArray<FString>& OutCaptions) const
 {
-	OutCaptions.Add(TEXT("Strata Diffuse BSDF"));
+	OutCaptions.Add(TEXT("Strata Slab BSDF"));
 }
 
-uint32 UMaterialExpressionStrataDiffuseBSDF::GetOutputType(int32 OutputIndex)
+uint32 UMaterialExpressionStrataSlabBSDF::GetOutputType(int32 OutputIndex)
 {
 	return MCT_Strata;
 }
 
-uint32 UMaterialExpressionStrataDiffuseBSDF::GetInputType(int32 InputIndex)
+uint32 UMaterialExpressionStrataSlabBSDF::GetInputType(int32 InputIndex)
 {
 	switch (InputIndex)
 	{
 	case 0:
-		return MCT_Float3;
+		return MCT_Float3; // Albedo
 		break;
 	case 1:
-		return MCT_Float1;
+		return MCT_Float1; // Metallic
 		break;
 	case 2:
-		return MCT_Float3;
+		return MCT_Float1; // Reflectivity
 		break;
 	case 3:
-		return MCT_Float3;
+		return MCT_Float1; // RoughnessX
 		break;
 	case 4:
-		return MCT_Float1;
+		return MCT_Float1; // RoughnessY
 		break;
 	case 5:
-		return MCT_Float1;
+		return MCT_Float3; // Normal
+		break;
+	case 6:
+		return MCT_Float3; // Tangent
+		break;
+	case 7:
+		return MCT_Float3; // DMFP Albedo
+		break;
+	case 8:
+		return MCT_Float1; // DMFP Radius
 		break;
 	}
 
@@ -19726,7 +19745,7 @@ uint32 UMaterialExpressionStrataDiffuseBSDF::GetInputType(int32 InputIndex)
 	return MCT_Float1;
 }
 
-FName UMaterialExpressionStrataDiffuseBSDF::GetInputName(int32 InputIndex) const
+FName UMaterialExpressionStrataSlabBSDF::GetInputName(int32 InputIndex) const
 {
 	if (InputIndex == 0)
 	{
@@ -19734,33 +19753,45 @@ FName UMaterialExpressionStrataDiffuseBSDF::GetInputName(int32 InputIndex) const
 	}
 	else if (InputIndex == 1)
 	{
-		return TEXT("Roughness");
+		return TEXT("Metallic");
 	}
 	else if (InputIndex == 2)
 	{
-		return TEXT("Normal");
+		return TEXT("Reflectivity");
 	}
 	else if (InputIndex == 3)
 	{
-		return TEXT("DMFP Albedo");
+		return TEXT("RoughnessX");
 	}
 	else if (InputIndex == 4)
 	{
-		return TEXT("DMFP Radius");
+		return TEXT("RoughnessY");
 	}
 	else if (InputIndex == 5)
 	{
-		return TEXT("DMFP Scale");
+		return TEXT("Normal");
+	}
+	else if (InputIndex == 6)
+	{
+		return TEXT("Tangent");
+	}
+	else if (InputIndex == 7)
+	{
+		return TEXT("DMFP Albedo");
+	}
+	else if (InputIndex == 8)
+	{
+		return TEXT("DMFP Radius");
 	}
 	return TEXT("Unknown");
 }
 
-bool UMaterialExpressionStrataDiffuseBSDF::IsResultStrataMaterial(int32 OutputIndex)
+bool UMaterialExpressionStrataSlabBSDF::IsResultStrataMaterial(int32 OutputIndex)
 {
 	return true;
 }
 
-void UMaterialExpressionStrataDiffuseBSDF::GatherStrataMaterialInfo(FStrataMaterialInfo& StrataMaterialInfo, int32 OutputIndex)
+void UMaterialExpressionStrataSlabBSDF::GatherStrataMaterialInfo(FStrataMaterialInfo& StrataMaterialInfo, int32 OutputIndex)
 {
 	const bool bIsScaleInputConnected = true; // STRATA_TODO
 	if (SubsurfaceProfile || bIsScaleInputConnected)
@@ -19779,265 +19810,6 @@ void UMaterialExpressionStrataDiffuseBSDF::GatherStrataMaterialInfo(FStrataMater
 #endif // WITH_EDITOR
 
 
-
-UMaterialExpressionStrataDielectricBSDF::UMaterialExpressionStrataDielectricBSDF(const FObjectInitializer& ObjectInitializer)
-	: Super(ObjectInitializer)
-{
-	struct FConstructorStatics
-	{
-		FText NAME_Strata;
-		FConstructorStatics() : NAME_Strata(LOCTEXT("Strata BSDFs", "Strata BSDFs")) { }
-	};
-	static FConstructorStatics ConstructorStatics;
-#if WITH_EDITORONLY_DATA
-	MenuCategories.Add(ConstructorStatics.NAME_Strata);
-#endif
-}
-
-#if WITH_EDITOR
-int32 UMaterialExpressionStrataDielectricBSDF::Compile(class FMaterialCompiler* Compiler, int32 OutputIndex)
-{
-	int32 NormalCodeChunk = CompileWithDefaultNormalWS(Compiler, Normal);
-	int32 TangentCodeChunk = CompileWithDefaultTangentWS(Compiler, Tangent);
-	uint8 SharedNormalIndex = StrataCompilationInfoCreateSharedNormal(Compiler, NormalCodeChunk, TangentCodeChunk);
-
-	int32 RoughnessXCodeChunk = CompileWithDefaultFloat1(Compiler, RoughnessX, 0.0f);
-	// If not plugged in, RoughnessYCodeChunk is set to RoughnessXCodeChunk to get an isotropic behavior
-	int32 RoughnessYCodeChunk = CompileWithDefaultCodeChunk(Compiler, RoughnessY, RoughnessXCodeChunk);
-
-	int32 OutputCodeChunk = Compiler->StrataDielectricBSDF(
-		RoughnessXCodeChunk,
-		RoughnessYCodeChunk,
-		CompileWithDefaultFloat1(Compiler, IOR,	1.5f),	// Default to Glass
-		CompileWithDefaultFloat3(Compiler, Tint,1.0f, 1.0f, 1.0f),
-		NormalCodeChunk,
-		TangentCodeChunk,
-		SharedNormalIndex);
-	StrataCompilationInfoCreateSingleBSDFMaterial(Compiler, OutputCodeChunk, SharedNormalIndex, STRATA_BSDF_TYPE_DIELECTRIC);
-
-	return OutputCodeChunk;
-}
-
-void UMaterialExpressionStrataDielectricBSDF::GetCaption(TArray<FString>& OutCaptions) const
-{
-	OutCaptions.Add(TEXT("Strata Dielectric BSDF"));
-}
-
-uint32 UMaterialExpressionStrataDielectricBSDF::GetOutputType(int32 OutputIndex)
-{
-	return MCT_Strata;
-}
-
-uint32 UMaterialExpressionStrataDielectricBSDF::GetInputType(int32 InputIndex)
-{
-	switch (InputIndex)
-	{
-	case 0:
-		return MCT_Float1;
-		break;
-	case 1:
-		return MCT_Float3;
-		break;
-	case 2:
-		return MCT_Float1;
-		break;
-	case 3:
-		return MCT_Float1;
-		break;
-	case 4:
-		return MCT_Float3;
-		break;
-	case 5:
-		return MCT_Float3;
-		break;
-	}
-
-	check(false);
-	return MCT_Float1;
-}
-
-bool UMaterialExpressionStrataDielectricBSDF::IsResultStrataMaterial(int32 OutputIndex)
-{
-	return true;
-}
-
-void UMaterialExpressionStrataDielectricBSDF::GatherStrataMaterialInfo(FStrataMaterialInfo& StrataMaterialInfo, int32 OutputIndex)
-{
-	StrataMaterialInfo.AddShadingModel(SSM_DefaultLit);
-}
-#endif // WITH_EDITOR
-
-
-
-UMaterialExpressionStrataConductorBSDF::UMaterialExpressionStrataConductorBSDF(const FObjectInitializer& ObjectInitializer)
-	: Super(ObjectInitializer)
-{
-	struct FConstructorStatics
-	{
-		FText NAME_Strata;
-		FConstructorStatics() : NAME_Strata(LOCTEXT("Strata BSDFs", "Strata BSDFs")) { }
-	};
-	static FConstructorStatics ConstructorStatics;
-#if WITH_EDITORONLY_DATA
-	MenuCategories.Add(ConstructorStatics.NAME_Strata);
-#endif
-}
-
-#if WITH_EDITOR
-int32 UMaterialExpressionStrataConductorBSDF::Compile(class FMaterialCompiler* Compiler, int32 OutputIndex)
-{
-	int32 NormalCodeChunk = CompileWithDefaultNormalWS(Compiler, Normal);
-	int32 TangentCodeChunk = CompileWithDefaultTangentWS(Compiler, Tangent);
-	uint8 SharedNormalIndex = StrataCompilationInfoCreateSharedNormal(Compiler, NormalCodeChunk, TangentCodeChunk);
-
-	int32 RoughnessXCodeChunk = CompileWithDefaultFloat1(Compiler, RoughnessX, 0.0f);
-	// If not plugged in, RoughnessYCodeChunk is set to RoughnessXCodeChunk to get an isotropic behavior
-	int32 RoughnessYCodeChunk = CompileWithDefaultCodeChunk(Compiler, RoughnessY, RoughnessXCodeChunk);
-
-	int32 OutputCodeChunk = Compiler->StrataConductorBSDF(
-		CompileWithDefaultFloat3(Compiler, Reflectivity,	0.947f, 0.776f, 0.371f),	// Default to Gold
-		CompileWithDefaultFloat3(Compiler, EdgeColor,		1.000f, 0.982f, 0.753f),	// Default to Gold
-		RoughnessXCodeChunk,
-		RoughnessYCodeChunk,
-		NormalCodeChunk,
-		TangentCodeChunk,
-		SharedNormalIndex);
-	StrataCompilationInfoCreateSingleBSDFMaterial(Compiler, OutputCodeChunk, SharedNormalIndex, STRATA_BSDF_TYPE_CONDUCTOR);
-
-	return OutputCodeChunk;
-}
-
-void UMaterialExpressionStrataConductorBSDF::GetCaption(TArray<FString>& OutCaptions) const
-{
-	OutCaptions.Add(TEXT("Strata Conductor BSDF"));
-}
-
-uint32 UMaterialExpressionStrataConductorBSDF::GetOutputType(int32 OutputIndex)
-{
-	return MCT_Strata;
-}
-
-uint32 UMaterialExpressionStrataConductorBSDF::GetInputType(int32 InputIndex)
-{
-	switch (InputIndex)
-	{
-	case 0:
-		return MCT_Float3;
-		break;
-	case 1:
-		return MCT_Float3;
-		break;
-	case 2:
-		return MCT_Float1;
-		break;
-	case 3:
-		return MCT_Float1;
-		break;
-	case 4:
-		return MCT_Float3;
-		break;
-	case 5:
-		return MCT_Float3;
-		break;
-	}
-
-	check(false);
-	return MCT_Float1;
-}
-
-bool UMaterialExpressionStrataConductorBSDF::IsResultStrataMaterial(int32 OutputIndex)
-{
-	return true;
-}
-
-void UMaterialExpressionStrataConductorBSDF::GatherStrataMaterialInfo(FStrataMaterialInfo& StrataMaterialInfo, int32 OutputIndex)
-{
-	StrataMaterialInfo.AddShadingModel(SSM_DefaultLit);
-}
-#endif // WITH_EDITOR
-
-
-
-UMaterialExpressionStrataVolumeBSDF::UMaterialExpressionStrataVolumeBSDF(const FObjectInitializer& ObjectInitializer)
-	: Super(ObjectInitializer)
-{
-	struct FConstructorStatics
-	{
-		FText NAME_Strata;
-		FConstructorStatics() : NAME_Strata(LOCTEXT("Strata BSDFs", "Strata BSDFs")) { }
-	};
-	static FConstructorStatics ConstructorStatics;
-#if WITH_EDITORONLY_DATA
-	MenuCategories.Add(ConstructorStatics.NAME_Strata);
-#endif
-}
-
-#if WITH_EDITOR
-int32 UMaterialExpressionStrataVolumeBSDF::Compile(class FMaterialCompiler* Compiler, int32 OutputIndex)
-{
-	int32 NormalCodeChunk = Compiler->VertexNormal();
-	uint8 SharedNormalIndex = StrataCompilationInfoCreateSharedNormal(Compiler, NormalCodeChunk);
-
-	bool DefaultAlbedoIsUsed = false;
-	int32 OutputCodeChunk = Compiler->StrataVolumeBSDF(
-		CompileWithDefaultFloat3(Compiler, Albedo,		0.0f, 0.0f, 0.0f, &DefaultAlbedoIsUsed),  // Default must be 0 to remove memory used by Albedo when it is not pluged (bHasScattering==false)
-		CompileWithDefaultFloat3(Compiler, Extinction,	0.0f, 0.0f, 0.0f),
-		CompileWithDefaultFloat1(Compiler, Anisotropy,	0.0f),
-		CompileWithDefaultFloat1(Compiler, Thickness,	0.001f),	// default = 1mm
-		NormalCodeChunk,
-		SharedNormalIndex);
-
-	const bool bHasScattering = !DefaultAlbedoIsUsed;
-	StrataCompilationInfoCreateSingleBSDFMaterial(Compiler, OutputCodeChunk, SharedNormalIndex, STRATA_BSDF_TYPE_VOLUME, bHasScattering);
-
-	return OutputCodeChunk;
-}
-
-void UMaterialExpressionStrataVolumeBSDF::GetCaption(TArray<FString>& OutCaptions) const
-{
-	OutCaptions.Add(TEXT("Strata Volume BSDF"));
-}
-
-uint32 UMaterialExpressionStrataVolumeBSDF::GetOutputType(int32 OutputIndex)
-{
-	return MCT_Strata;
-}
-
-uint32 UMaterialExpressionStrataVolumeBSDF::GetInputType(int32 InputIndex)
-{
-	switch (InputIndex)
-	{
-	case 0:
-		return MCT_Float3;
-		break;
-	case 1:
-		return MCT_Float3;
-		break;
-	case 2:
-		return MCT_Float1;
-		break;
-	case 3:
-		return MCT_Float1;
-		break;
-	}
-
-	check(false);
-	return MCT_Float1;
-}
-
-bool UMaterialExpressionStrataVolumeBSDF::IsResultStrataMaterial(int32 OutputIndex)
-{
-	return true;
-}
-
-void UMaterialExpressionStrataVolumeBSDF::GatherStrataMaterialInfo(FStrataMaterialInfo& StrataMaterialInfo, int32 OutputIndex)
-{
-	StrataMaterialInfo.AddShadingModel(SSM_DefaultLit);
-}
-#endif // WITH_EDITOR
-
-
-
 UMaterialExpressionStrataSheenBSDF::UMaterialExpressionStrataSheenBSDF(const FObjectInitializer& ObjectInitializer)
 	: Super(ObjectInitializer)
 {
diff --git a/Engine/Source/Runtime/Engine/Private/Materials/StrataMaterial.cpp b/Engine/Source/Runtime/Engine/Private/Materials/StrataMaterial.cpp
--- a/Engine/Source/Runtime/Engine/Private/Materials/StrataMaterial.cpp
+++ b/Engine/Source/Runtime/Engine/Private/Materials/StrataMaterial.cpp
@@ -9,17 +9,8 @@ FString GetStrataBSDFName(uint8 BSDFType)
 {
 	switch (BSDFType)
 	{
-	case STRATA_BSDF_TYPE_DIFFUSE:
-		return TEXT("DIFFUSE");
-		break;
-	case STRATA_BSDF_TYPE_DIELECTRIC:
-		return TEXT("DIELECTRIC");
-		break;
-	case STRATA_BSDF_TYPE_CONDUCTOR:
-		return TEXT("CONDUCTOR");
-		break;
-	case STRATA_BSDF_TYPE_VOLUME:
-		return TEXT("VOLUME");
+	case STRATA_BSDF_TYPE_SLAB:
+		return TEXT("SLAB");
 		break;
 	case STRATA_BSDF_TYPE_SHEEN:
 		return TEXT("SHEEN");
@@ -285,31 +276,10 @@ FStrataMaterialAnalysisResult StrataCompilationInfoMaterialAnalysis(FMaterialCom
 
 			switch (BSDF.Type)
 			{
-			case STRATA_BSDF_TYPE_DIFFUSE:
-			{
-				Result.RequestedByteCount += UintByteSize;
-				if (BSDF.bHasScattering)
-				{
-					Result.RequestedByteCount += UintByteSize;
-				}
-				break;
-			}
-			case STRATA_BSDF_TYPE_DIELECTRIC:
+			case STRATA_BSDF_TYPE_SLAB:
 			{
 				Result.RequestedByteCount += UintByteSize;
 				Result.RequestedByteCount += UintByteSize;
-				break;
-			}
-			case STRATA_BSDF_TYPE_CONDUCTOR:
-			{
-				Result.RequestedByteCount += UintByteSize;
-				Result.RequestedByteCount += UintByteSize;
-				break;
-			}
-			case STRATA_BSDF_TYPE_VOLUME:
-			{
-				Result.RequestedByteCount += UintByteSize;
-
 				if (BSDF.bHasScattering)
 				{
 					Result.RequestedByteCount += UintByteSize;
diff --git a/Engine/Source/Runtime/Engine/Public/MaterialCompiler.h b/Engine/Source/Runtime/Engine/Public/MaterialCompiler.h
--- a/Engine/Source/Runtime/Engine/Public/MaterialCompiler.h
+++ b/Engine/Source/Runtime/Engine/Public/MaterialCompiler.h
@@ -392,10 +392,7 @@ public:
 
 	// Strata
 	virtual int32 StrataCreateAndRegisterNullMaterial() = 0;
-	virtual int32 StrataDiffuseBSDF(int32 Albedo, int32 Roughness, int32 SSSProfileId, int32 SSSDMFPRadiusScale, int32 SSSDMFPAlebdo, int32 SSSDMFPRadius, int32 Normal, uint8 SharedNormalIndex) = 0;
-	virtual int32 StrataDielectricBSDF(int32 RoughnessX, int32 RoughnessY, int32 IOR, int32 Tint, int32 Normal, int32 Tangent, uint8 SharedNormalIndex) = 0;
-	virtual int32 StrataConductorBSDF(int32 Reflectivity, int32 EdgeColor, int32 RoughnessX, int32 RoughnessY, int32 Normal, int32 Tangent, uint8 SharedNormalIndex) = 0;
-	virtual int32 StrataVolumeBSDF(int32 Albedo, int32 Extinction, int32 Anisotropy, int32 Thickness, int32 Normal, uint8 SharedNormalIndex) = 0;
+	virtual int32 StrataSlabBSDF(int32 Albedo, int32 Reflectivity, int32 Metallic, int32 RoughnessX, int32 RoughnessY, int32 SSSProfileId, int32 SSSDMFPAlbedo, int32 SSSDMFPRadius, int32 Normal, int32 Tangent, uint8 SharedNormalIndex) = 0;
 	virtual int32 StrataSheenBSDF(int32 Albedo, int32 Roughness, int32 Normal, uint8 SharedNormalIndex) = 0;
 	virtual int32 StrataVolumetricFogCloudBSDF(int32 Albedo, int32 Extinction, int32 Emissive, int32 AmbientOcclusion) = 0;
 	virtual int32 StrataUnlitBSDF(int32 Emissive, int32 TransmittanceColor) = 0;
@@ -811,24 +808,9 @@ public:
 		return Compiler->StrataCreateAndRegisterNullMaterial();
 	}
 
-	virtual int32 StrataDiffuseBSDF(int32 Albedo, int32 Roughness, int32 SSSProfileId, int32 SSSDMFPRadiusScale, int32 SSSDMFPAlebdo, int32 SSSDMFPRadius, int32 Normal, uint8 SharedNormalIndex) override
+	virtual int32 StrataSlabBSDF(int32 Albedo, int32 Reflectivity, int32 Metallic, int32 RoughnessX, int32 RoughnessY, int32 SSSProfileId, int32 SSSDMFPAlbedo, int32 SSSDMFPRadius, int32 Normal, int32 Tangent, uint8 SharedNormalIndex) override
 	{
-		return Compiler->StrataDiffuseBSDF(Albedo, Roughness, SSSProfileId, SSSDMFPRadiusScale, SSSDMFPAlebdo, SSSDMFPRadius, Normal, SharedNormalIndex);
-	}
-
-	virtual int32 StrataDielectricBSDF(int32 RoughnessX, int32 RoughnessY, int32 IOR, int32 Tint, int32 Normal, int32 Tangent, uint8 SharedNormalIndex) override
-	{
-		return Compiler->StrataDielectricBSDF(RoughnessX, RoughnessY, IOR, Tint, Normal, Tangent, SharedNormalIndex);
-	}
-
-	virtual int32 StrataConductorBSDF(int32 Reflectivity, int32 EdgeColor, int32 RoughnessX, int32 RoughnessY, int32 Normal, int32 Tangent, uint8 SharedNormalIndex) override
-	{
-		return Compiler->StrataConductorBSDF(Reflectivity, EdgeColor, RoughnessX, RoughnessY, Normal, Tangent, SharedNormalIndex);
-	}
-
-	virtual int32 StrataVolumeBSDF(int32 Albedo, int32 Extinction, int32 Anisotropy, int32 Thickness, int32 Normal, uint8 SharedNormalIndex) override
-	{
-		return Compiler->StrataVolumeBSDF(Albedo, Extinction, Anisotropy, Thickness, Normal, SharedNormalIndex);
+		return Compiler->StrataSlabBSDF(Albedo, Reflectivity, Metallic, RoughnessX, RoughnessY, SSSProfileId, SSSDMFPAlbedo, SSSDMFPRadius, Normal, Tangent, SharedNormalIndex);
 	}
 
 	virtual int32 StrataSheenBSDF(int32 Albedo, int32 Roughness, int32 Normal, uint8 SharedNormalIndex) override
