Strata parameters renaming and tooltip clean up.
- RoughnessXY instead of Roughness01
- Specular instead of Reflectivity
- BaseColor+ DiffuseColor instead of Albedo.

diff --git a/Engine/Shaders/Private/Strata/Strata.ush b/Engine/Shaders/Private/Strata/Strata.ush
--- a/Engine/Shaders/Private/Strata/Strata.ush
+++ b/Engine/Shaders/Private/Strata/Strata.ush
@@ -105,18 +105,18 @@ struct FStrataBSDF
 #define BSDF_SETWEIGHT10F(X, Weight)	BSDF_SETSTATE(X, STATE_BITS_WEIGHT10F, STATE_BIT_OFFSET_WEIGHT10F, Weight)
 
 #if STRATA_INLINE_SHADING
-#define BSDF_SETEMISSIVE(X, V)	X.Emissive = V
-#define BSDF_GETEMISSIVE(X)		X.Emissive
+#define BSDF_SETEMISSIVE(X, V)			X.Emissive = V
+#define BSDF_GETEMISSIVE(X)				X.Emissive
 #else
-#define BSDF_SETEMISSIVE(X, V)	X
-#define BSDF_GETEMISSIVE(X)		0.0f
+#define BSDF_SETEMISSIVE(X, V)			X
+#define BSDF_GETEMISSIVE(X)				0.0f
 #endif
 
-#define SLAB_ALBEDO(X)					X.VGPRs[0].xyz
-#define SLAB_REFLECTIVITY(X)			X.VGPRs[0].w
+#define SLAB_BASECOLOR(X)				X.VGPRs[0].xyz
+#define SLAB_SPECULAR(X)				X.VGPRs[0].w
 #define SLAB_METALLIC(X)				X.VGPRs[1].x
-#define SLAB_ROUGHNESS0(X)				X.VGPRs[1].y
-#define SLAB_ROUGHNESS1(X)				X.VGPRs[1].z
+#define SLAB_ROUGHNESSX(X)				X.VGPRs[1].y
+#define SLAB_ROUGHNESSY(X)				X.VGPRs[1].z
 // Either with SSS profile
 #define SLAB_SSSPROFILEID(X)			X.VGPRs[2].x
 #define SLAB_SSSPROFILERADIUS(X)		X.VGPRs[2].y
@@ -124,7 +124,7 @@ struct FStrataBSDF
 #define SLAB_SSSDMFPALBEDO(X)			X.VGPRs[2].xyz
 #define SLAB_SSSDMFPRADIUS(X)			X.VGPRs[2].w
 
-#define SHEEN_ALBEDO(X)					X.VGPRs[0].xyz
+#define SHEEN_BASECOLOR(X)				X.VGPRs[0].xyz
 #define SHEEN_ROUGHNESS(X)				X.VGPRs[0].w
 
 #define VOLUMETRICFOGCLOUD_ALBEDO(X)	X.VGPRs[0].xyz
@@ -157,8 +157,8 @@ void StrataRequestSharedNormalTangentBasis(inout uint NormalTypes, uint Index);
 
 uint   PackR11G11B10F(float3 rgb);
 float3 UnpackR11G11B10F(uint rgb);
-uint   PackAlbedoLinearToGamma2AlphaLinear(float4 rgba);
-float4 UnpackAlbedoGamma2ToLinearAlphaLinear(uint rgba);
+uint   PackColorLinearToGamma2AlphaLinear(float4 rgba);
+float4 UnpackColorGamma2ToLinearAlphaLinear(uint rgba);
 uint   StrataPackNormal(in float3 Normal);
 float3 StrataUnpackNormal(uint PackedNormal);
 
@@ -227,14 +227,14 @@ void FinalizeTopLayerData(inout FStrataTopLayerData Out)
 uint2 StrataPackTopLayerData(FStrataTopLayerData In)
 {
 	uint2 Out = 0;
-	Out.x = PackAlbedoLinearToGamma2AlphaLinear(float4(In.BaseColor, In.Roughness));
+	Out.x = PackColorLinearToGamma2AlphaLinear(float4(In.BaseColor, In.Roughness));
 	Out.y = StrataPackNormal(In.Normal);
 	return Out;
 }
 
 FStrataTopLayerData StrataUnpackTopLayerData(uint2 In)
 {
-	const float4 BaseColorAndRoughness = UnpackAlbedoGamma2ToLinearAlphaLinear(In.x);
+	const float4 BaseColorAndRoughness = UnpackColorGamma2ToLinearAlphaLinear(In.x);
 
 	FStrataTopLayerData Out = (FStrataTopLayerData)0;
 	Out.Normal = StrataUnpackNormal(In.y);
@@ -254,9 +254,9 @@ bool StrataHasShadingModel(uint ShadingModels, uint ShadingModel)
 	return (ShadingModels & (1 << ShadingModel)) > 0;
 }
 
-bool StrataHasAnisotropy(in float Roughness0, float Roughness1)
+bool StrataHasAnisotropy(in float RoughnessX, float RoughnessY)
 {
-	return Roughness0 != Roughness1;
+	return RoughnessX != RoughnessY;
 }
 
 // Return true if a BSDF has a custom area light integrator
@@ -304,7 +304,7 @@ FStrataData GetInitialisedStrataData()
 }
 
 FStrataData GetStrataSlabBSDF(
-	float3 Albedo, float Reflectivity, float Metallic, float2 Roughness,
+	float3 BaseColor, float Specular, float Metallic, float2 Roughness,
 	float SSSProfileID, float3 SSSDMFPAlbedo, float SSSDMFPRadius,
 	uint SharedNormalIndex, inout uint SharedNormalTypes)
 {
@@ -327,11 +327,11 @@ FStrataData GetStrataSlabBSDF(
 	const bool bHasSSSProfile = SSSProfileID > 0;
 	BSDF_SETHASSSSPROFILE	(StrataData.Layers[0].BSDFs[0],		bHasSSSProfile ? 1 : 0);
 
-	SLAB_ALBEDO				(StrataData.Layers[0].BSDFs[0])		= Albedo;
-	SLAB_REFLECTIVITY		(StrataData.Layers[0].BSDFs[0])		= Reflectivity;
+	SLAB_BASECOLOR			(StrataData.Layers[0].BSDFs[0])		= BaseColor;
+	SLAB_SPECULAR			(StrataData.Layers[0].BSDFs[0])		= Specular;
 	SLAB_METALLIC			(StrataData.Layers[0].BSDFs[0])		= Metallic;
-	SLAB_ROUGHNESS0			(StrataData.Layers[0].BSDFs[0])		= Roughness.x;
-	SLAB_ROUGHNESS1			(StrataData.Layers[0].BSDFs[0])		= Roughness.y;
+	SLAB_ROUGHNESSX			(StrataData.Layers[0].BSDFs[0])		= Roughness.x;
+	SLAB_ROUGHNESSY			(StrataData.Layers[0].BSDFs[0])		= Roughness.y;
 	if (bHasSSSProfile)
 	{
 		SLAB_SSSPROFILEID	(StrataData.Layers[0].BSDFs[0])		= SSSProfileID;
@@ -347,7 +347,7 @@ FStrataData GetStrataSlabBSDF(
 	return StrataData;
 }
 
-FStrataData GetStrataSheenBSDF(float3 Albedo, float Roughness, uint SharedNormalIndex)
+FStrataData GetStrataSheenBSDF(float3 BaseColor, float Roughness, uint SharedNormalIndex)
 {
 	FStrataData StrataData = GetInitialisedStrataData();
 	StrataData.LayerCount++;
@@ -356,7 +356,7 @@ FStrataData GetStrataSheenBSDF(float3 Albedo, float Roughness, uint SharedNormal
 	BSDF_SETTYPE			(StrataData.Layers[0].BSDFs[0],		STRATA_BSDF_TYPE_SHEEN);
 	BSDF_SETNORMALID		(StrataData.Layers[0].BSDFs[0],		SharedNormalIndex);
 	BSDF_SETISTOPLAYER		(StrataData.Layers[0].BSDFs[0],		0);
-	SHEEN_ALBEDO			(StrataData.Layers[0].BSDFs[0])		= Albedo;
+	SHEEN_BASECOLOR			(StrataData.Layers[0].BSDFs[0])		= BaseColor;
 	SHEEN_ROUGHNESS			(StrataData.Layers[0].BSDFs[0])		= Roughness;
 
 	StrataData.Layers[0].BSDFs[0].Weight						= 1.0f;
@@ -619,11 +619,11 @@ FStrataData StrataVerticalLayering(FStrataData Top, FStrataData Base)
 
 void SanitizeStrataSlab(inout FStrataBSDF BSDF)
 {
-	SLAB_ALBEDO(BSDF) = saturate(SLAB_ALBEDO(BSDF));
-	SLAB_REFLECTIVITY(BSDF) = saturate(SLAB_REFLECTIVITY(BSDF));
+	SLAB_BASECOLOR(BSDF) = saturate(SLAB_BASECOLOR(BSDF));
+	SLAB_SPECULAR(BSDF) = saturate(SLAB_SPECULAR(BSDF));
 	SLAB_METALLIC(BSDF) = saturate(SLAB_METALLIC(BSDF));
-	SLAB_ROUGHNESS0(BSDF) = saturate(SLAB_ROUGHNESS0(BSDF));
-	SLAB_ROUGHNESS1(BSDF) = saturate(SLAB_ROUGHNESS1(BSDF));
+	SLAB_ROUGHNESSX(BSDF) = saturate(SLAB_ROUGHNESSX(BSDF));
+	SLAB_ROUGHNESSY(BSDF) = saturate(SLAB_ROUGHNESSY(BSDF));
 	SLAB_SSSPROFILERADIUS(BSDF) = max(0.0f, SLAB_SSSPROFILERADIUS(BSDF));
 	SLAB_SSSDMFPALBEDO(BSDF) = saturate(SLAB_SSSDMFPALBEDO(BSDF));
 	SLAB_SSSDMFPRADIUS(BSDF) = clamp(SLAB_SSSDMFPRADIUS(BSDF), 0.0f, SSS_PROFILE_PERPIXEL_MAXRADIUS);
@@ -631,7 +631,7 @@ void SanitizeStrataSlab(inout FStrataBSDF BSDF)
 
 void SanitizeStrataSheen(inout FStrataBSDF BSDF)
 {
-	SHEEN_ALBEDO(BSDF) = saturate(SHEEN_ALBEDO(BSDF));
+	SHEEN_BASECOLOR(BSDF) = saturate(SHEEN_BASECOLOR(BSDF));
 	SHEEN_ROUGHNESS(BSDF) = saturate(SHEEN_ROUGHNESS(BSDF));
 }
 
@@ -739,7 +739,7 @@ float4 UnpackR16G16B16A16F(uint2 AB)
 	return float4(r, g, b, a);
 }
 
-uint PackAlbedoLinearToGamma2(float3 rgb)
+uint PackColorLinearToGamma2(float3 rgb)
 {
 	uint r = (uint(sqrt(saturate(rgb.r)) * 255.0f) << 16);
 	uint g = (uint(sqrt(saturate(rgb.g)) * 255.0f) << 8);
@@ -747,7 +747,7 @@ uint PackAlbedoLinearToGamma2(float3 rgb)
 	return r | g | b;
 }
 
-float3 UnpackAlbedoGamma2ToLinear(uint rgb)
+float3 UnpackColorGamma2ToLinear(uint rgb)
 {
 	float r = float((rgb >> 16) & 0x000000FF);
 	float g = float((rgb >> 8)  & 0x000000FF);
@@ -756,7 +756,7 @@ float3 UnpackAlbedoGamma2ToLinear(uint rgb)
 	return Gamma2 * Gamma2;	// Simple gamma 2.0
 }
 
-uint PackAlbedoLinearToGamma2AlphaLinear(float4 rgba)
+uint PackColorLinearToGamma2AlphaLinear(float4 rgba)
 {
 	uint r = (uint(sqrt(saturate(rgba.r)) * 255.0f) << 24);
 	uint g = (uint(sqrt(saturate(rgba.g)) * 255.0f) << 16);
@@ -765,7 +765,7 @@ uint PackAlbedoLinearToGamma2AlphaLinear(float4 rgba)
 	return r | g | b | a;
 }
 
-float4 UnpackAlbedoGamma2ToLinearAlphaLinear(uint rgba)
+float4 UnpackColorGamma2ToLinearAlphaLinear(uint rgba)
 {
 	float r = float((rgba >> 24) & 0x000000FF);
 	float g = float((rgba >> 16) & 0x000000FF);
@@ -989,14 +989,14 @@ void StrataLoad_R11_G11_B10F(ByteAddressBuffer StrataBuffer, inout uint ByteOffs
 	C = temp.z;
 }
 
-void StrataLoad_AlbedoGamma2ToLinear(ByteAddressBuffer StrataBuffer, inout uint ByteOffset, inout float3 A)
+void StrataLoad_ColorGamma2ToLinear(ByteAddressBuffer StrataBuffer, inout uint ByteOffset, inout float3 A)
 {
-	A = UnpackAlbedoGamma2ToLinear(StrataLoadUint1(StrataBuffer, ByteOffset));
+	A = UnpackColorGamma2ToLinear(StrataLoadUint1(StrataBuffer, ByteOffset));
 }
 
-void StrataLoad_AlbedoGamma2ToLinearAlphaLinear(ByteAddressBuffer StrataBuffer, inout uint ByteOffset, inout float3 RGB, inout float A)
+void StrataLoad_ColorGamma2ToLinearAlphaLinear(ByteAddressBuffer StrataBuffer, inout uint ByteOffset, inout float3 RGB, inout float A)
 {
-	float4 Data = UnpackAlbedoGamma2ToLinearAlphaLinear(StrataLoadUint1(StrataBuffer, ByteOffset));
+	float4 Data = UnpackColorGamma2ToLinearAlphaLinear(StrataLoadUint1(StrataBuffer, ByteOffset));
 	RGB = Data.rgb;
 	A = Data.a;
 }
@@ -1067,15 +1067,15 @@ void PackStrataOut(
 			{
 			case STRATA_BSDF_TYPE_SLAB:
 			{
-				const float3 F0 = ComputeF0(SLAB_REFLECTIVITY(BSDF), SLAB_ALBEDO(BSDF), SLAB_METALLIC(BSDF));
+				const float3 F0 = ComputeF0(SLAB_SPECULAR(BSDF), SLAB_BASECOLOR(BSDF), SLAB_METALLIC(BSDF));
 
 				if (bTopLayer)
 				{
 					// Setup SSR for the top layer
-					TopLayer.Roughness += BSDFCoverageAvg * SLAB_ROUGHNESS0(BSDF);
+					TopLayer.Roughness += BSDFCoverageAvg * SLAB_ROUGHNESSX(BSDF);
 					TopLayer.AccumRoughnessCount += BSDFCoverageAvg;
 
-					const float SafeRoughness = MakeRoughnessSafe(SLAB_ROUGHNESS0(BSDF));
+					const float SafeRoughness = MakeRoughnessSafe(SLAB_ROUGHNESSX(BSDF));
 					TopLayer.BaseColor += BSDFCoverage * EnvBRDFApprox(F0, SafeRoughness, saturate(ContextNoLight.NoV));
 				}
 
@@ -1086,7 +1086,7 @@ void PackStrataOut(
 					SSS.bHasProfile	= BSDF_GETHASSSSPROFILE(BSDF);
 					SSS.Profile		= SLAB_SSSPROFILEID(BSDF);
 					SSS.RadiusScale	= SLAB_SSSPROFILERADIUS(BSDF);
-					SSS.BaseColor	= SLAB_ALBEDO(BSDF);
+					SSS.BaseColor	= SLAB_BASECOLOR(BSDF);
 					SSS.DMFPAlbedo	= SLAB_SSSDMFPALBEDO(BSDF);
 					SSS.DMFPRadius	= SLAB_SSSDMFPRADIUS(BSDF);
 				}
@@ -1098,7 +1098,7 @@ void PackStrataOut(
 
 			case STRATA_BSDF_TYPE_SHEEN:
 			{
-				const float3 F0  = SHEEN_ALBEDO(BSDF);
+				const float3 F0  = SHEEN_BASECOLOR(BSDF);
 				const float3 F = F_Schlick(F0, ContextNoLight.VoH);
 
 				if (bTopLayer)
@@ -1251,8 +1251,8 @@ void PackStrataOut(
 			{
 			case STRATA_BSDF_TYPE_SLAB:
 			{
-				STRATA_STORE_UINT1(PackAlbedoLinearToGamma2AlphaLinear(float4(SLAB_ALBEDO(BSDF), SLAB_ROUGHNESS0(BSDF))));
-				STRATA_STORE_UINT1(PackRGBA8(float4(SLAB_REFLECTIVITY(BSDF), SLAB_METALLIC(BSDF), SLAB_ROUGHNESS1(BSDF), 0)));
+				STRATA_STORE_UINT1(PackColorLinearToGamma2AlphaLinear(float4(SLAB_BASECOLOR(BSDF), SLAB_ROUGHNESSX(BSDF))));
+				STRATA_STORE_UINT1(PackRGBA8(float4(SLAB_SPECULAR(BSDF), SLAB_METALLIC(BSDF), SLAB_ROUGHNESSY(BSDF), 0)));
 
 				if (BSDF_GETHASSCATTERING(BSDF))
 				{
@@ -1262,7 +1262,7 @@ void PackStrataOut(
 					}
 					else
 					{
-						STRATA_STORE_UINT1(PackAlbedoLinearToGamma2AlphaLinear(float4(SLAB_SSSDMFPALBEDO(BSDF), SLAB_SSSDMFPRADIUS(BSDF) / SSS_PROFILE_PERPIXEL_MAXRADIUS)));
+						STRATA_STORE_UINT1(PackColorLinearToGamma2AlphaLinear(float4(SLAB_SSSDMFPALBEDO(BSDF), SLAB_SSSDMFPRADIUS(BSDF) / SSS_PROFILE_PERPIXEL_MAXRADIUS)));
 					}
 				}
 				// 8 bytes
@@ -1270,20 +1270,20 @@ void PackStrataOut(
 			break;
 			case STRATA_BSDF_TYPE_SHEEN:
 			{
-				STRATA_STORE_UINT1(PackAlbedoLinearToGamma2AlphaLinear(float4(SHEEN_ALBEDO(BSDF), SHEEN_ROUGHNESS(BSDF))));
+				STRATA_STORE_UINT1(PackColorLinearToGamma2AlphaLinear(float4(SHEEN_BASECOLOR(BSDF), SHEEN_ROUGHNESS(BSDF))));
 				// 4 bytes
 			}
 			break;
 			case STRATA_BSDF_TYPE_HAIR:
 			{
-				STRATA_STORE_UINT1(PackAlbedoLinearToGamma2AlphaLinear(float4(HAIR_BASECOLOR(BSDF), HAIR_ROUGHNESS(BSDF))));
+				STRATA_STORE_UINT1(PackColorLinearToGamma2AlphaLinear(float4(HAIR_BASECOLOR(BSDF), HAIR_ROUGHNESS(BSDF))));
 				STRATA_STORE_UINT1(PackR11G11B10F(float3(HAIR_SCATTER(BSDF), HAIR_SPECULAR(BSDF), HAIR_BACKLIT(BSDF))));
 				// 8 bytes
 			}
 			break;
 			case STRATA_BSDF_TYPE_SINGLELAYERWATER:
 			{
-				STRATA_STORE_UINT1(PackAlbedoLinearToGamma2AlphaLinear(float4(SLW_BASECOLOR(BSDF), 0.0f)));
+				STRATA_STORE_UINT1(PackColorLinearToGamma2AlphaLinear(float4(SLW_BASECOLOR(BSDF), 0.0f)));
 				STRATA_STORE_UINT1(PackRGBA8(float4(SLW_METALLIC(BSDF), SLW_SPECULAR(BSDF), SLW_ROUGHNESS(BSDF), SLW_TOPMATERIALOPACITY(BSDF))));
 			}
 			break;
@@ -1357,12 +1357,12 @@ FStrataBSDF UnpackStrataBSDFIn(ByteAddressBuffer StrataBuffer, inout uint ByteOf
 	{
 	case STRATA_BSDF_TYPE_SLAB:
 	{
-		StrataLoad_AlbedoGamma2ToLinearAlphaLinear(StrataBuffer, ByteOffset, SLAB_ALBEDO(OutBSDF), SLAB_ROUGHNESS0(OutBSDF));
+		StrataLoad_ColorGamma2ToLinearAlphaLinear(StrataBuffer, ByteOffset, SLAB_BASECOLOR(OutBSDF), SLAB_ROUGHNESSX(OutBSDF));
 		{
 			float4 Data = UnpackRGBA8(StrataLoadUint1(StrataBuffer, ByteOffset));
-			SLAB_REFLECTIVITY(OutBSDF) = Data.x;
+			SLAB_SPECULAR(OutBSDF) = Data.x;
 			SLAB_METALLIC(OutBSDF) = Data.y;
-			SLAB_ROUGHNESS1(OutBSDF) = Data.z;
+			SLAB_ROUGHNESSY(OutBSDF) = Data.z;
 		}
 		if (BSDF_GETHASSCATTERING(OutBSDF))
 		{
@@ -1374,7 +1374,7 @@ FStrataBSDF UnpackStrataBSDFIn(ByteAddressBuffer StrataBuffer, inout uint ByteOf
 			}
 			else
 			{
-				StrataLoad_AlbedoGamma2ToLinearAlphaLinear(StrataBuffer, ByteOffset, SLAB_SSSDMFPALBEDO(OutBSDF), SLAB_SSSDMFPRADIUS(OutBSDF));
+				StrataLoad_ColorGamma2ToLinearAlphaLinear(StrataBuffer, ByteOffset, SLAB_SSSDMFPALBEDO(OutBSDF), SLAB_SSSDMFPRADIUS(OutBSDF));
 				SLAB_SSSDMFPRADIUS(OutBSDF) *= SSS_PROFILE_PERPIXEL_MAXRADIUS;
 			}
 		}
@@ -1382,19 +1382,19 @@ FStrataBSDF UnpackStrataBSDFIn(ByteAddressBuffer StrataBuffer, inout uint ByteOf
 	break;
 	case STRATA_BSDF_TYPE_SHEEN:
 	{
-		StrataLoad_AlbedoGamma2ToLinearAlphaLinear(StrataBuffer, ByteOffset, SHEEN_ALBEDO(OutBSDF), SHEEN_ROUGHNESS(OutBSDF));
+		StrataLoad_ColorGamma2ToLinearAlphaLinear(StrataBuffer, ByteOffset, SHEEN_BASECOLOR(OutBSDF), SHEEN_ROUGHNESS(OutBSDF));
 	}
 	break;
 	case STRATA_BSDF_TYPE_HAIR:
 	{
-		StrataLoad_AlbedoGamma2ToLinearAlphaLinear(StrataBuffer, ByteOffset, HAIR_BASECOLOR(OutBSDF), HAIR_ROUGHNESS(OutBSDF));
+		StrataLoad_ColorGamma2ToLinearAlphaLinear(StrataBuffer, ByteOffset, HAIR_BASECOLOR(OutBSDF), HAIR_ROUGHNESS(OutBSDF));
 		StrataLoad_R11_G11_B10F(StrataBuffer, ByteOffset, HAIR_SCATTER(OutBSDF), HAIR_SPECULAR(OutBSDF), HAIR_BACKLIT(OutBSDF));
 		// 8 bytes
 	}
 	break;
 	case STRATA_BSDF_TYPE_SINGLELAYERWATER:
 	{
-		StrataLoad_AlbedoGamma2ToLinearAlphaLinear(StrataBuffer, ByteOffset, SLW_BASECOLOR(OutBSDF), DummyFloat);
+		StrataLoad_ColorGamma2ToLinearAlphaLinear(StrataBuffer, ByteOffset, SLW_BASECOLOR(OutBSDF), DummyFloat);
 		float4 Data = UnpackRGBA8(StrataLoadUint1(StrataBuffer, ByteOffset));
 		SLW_METALLIC(OutBSDF) = Data.x;
 		SLW_SPECULAR(OutBSDF) = Data.y;
diff --git a/Engine/Shaders/Private/Strata/StrataDeferredLighting.ush b/Engine/Shaders/Private/Strata/StrataDeferredLighting.ush
--- a/Engine/Shaders/Private/Strata/StrataDeferredLighting.ush
+++ b/Engine/Shaders/Private/Strata/StrataDeferredLighting.ush
@@ -58,16 +58,16 @@ void StrataPreprocessAreaLightIntegration(
 	if (bIsValid && IsAreaLight(LightContext.AreaLight) && StrataHasAreaLightIntegrator(BSDFContext.BSDF))
 	{
 		// Evaluate roughness / specular color
-		float SafeRoughness0 = 0;
-		float SafeRoughness1 = 0;
+		float SafeRoughnessX = 0;
+		float SafeRoughnessY = 0;
 		float3 SpecularColor = 0;
 		switch (BSDF_GETTYPE(BSDFContext.BSDF))
 		{
 			case STRATA_BSDF_TYPE_SLAB:
 			{
-				SafeRoughness0 = MakeRoughnessSafe(SLAB_ROUGHNESS0(BSDFContext.BSDF), View.MinRoughness);
-				SafeRoughness1 = MakeRoughnessSafe(SLAB_ROUGHNESS1(BSDFContext.BSDF), View.MinRoughness);
-				SpecularColor = ComputeF0(SLAB_REFLECTIVITY(BSDFContext.BSDF), SLAB_ALBEDO(BSDFContext.BSDF), SLAB_METALLIC(BSDFContext.BSDF));
+				SafeRoughnessX = MakeRoughnessSafe(SLAB_ROUGHNESSX(BSDFContext.BSDF), View.MinRoughness);
+				SafeRoughnessY = MakeRoughnessSafe(SLAB_ROUGHNESSY(BSDFContext.BSDF), View.MinRoughness);
+				SpecularColor = ComputeF0(SLAB_SPECULAR(BSDFContext.BSDF), SLAB_BASECOLOR(BSDFContext.BSDF), SLAB_METALLIC(BSDFContext.BSDF));
 			}
 			break;
 		}
@@ -84,7 +84,7 @@ void StrataPreprocessAreaLightIntegration(
 				// OutgoingRadianceNoShadow, which will be accumulated outside of this function
 				Out.OutgoingRadianceNoShadow.rgb +=
 					Out.IncidentRadianceNoShadow *
-					RectGGXApproxLTC(SafeRoughness0, SpecularColor, BSDFContext.N, BSDFContext.V, LightContext.AreaLight.Rect, LightContext.AreaLight.Texture);
+					RectGGXApproxLTC(SafeRoughnessX, SpecularColor, BSDFContext.N, BSDFContext.V, LightContext.AreaLight.Rect, LightContext.AreaLight.Texture);
 				Out.bEvaluateBSDF = false;
 			}
 			else
@@ -100,7 +100,7 @@ void StrataPreprocessAreaLightIntegration(
 
 				// Update incident irradiance to add the energy rescaling to due MRP integration
 				// The a2 override stored the overriden roughness, which will be used during the BSDF evaluation
-				BSDFContext.a2Override = Pow4(SafeRoughness0);
+				BSDFContext.a2Override = Pow4(SafeRoughnessX);
 				Out.IncidentIrradianceNoShadow =
 					Out.IncidentRadianceNoShadow *
 					EnergyNormalization(BSDFContext.a2Override, BSDFContext.Context.VoH, LightContext.AreaLight) *
diff --git a/Engine/Shaders/Private/Strata/StrataEnvironmentLighting.ush b/Engine/Shaders/Private/Strata/StrataEnvironmentLighting.ush
--- a/Engine/Shaders/Private/Strata/StrataEnvironmentLighting.ush
+++ b/Engine/Shaders/Private/Strata/StrataEnvironmentLighting.ush
@@ -39,7 +39,7 @@ float4 StrataEnvironmentLighting(
 
 		// Diffuse component
 		float3 DiffuseSkyLighting = 0;
-		if (any(StrataEnvLight.DiffuseAlbedo > 0.0f))
+		if (any(StrataEnvLight.DiffuseColor > 0.0f))
 		{
 			// Apply bent normal if needed
 			#if APPLY_SKY_SHADOWING
@@ -54,7 +54,7 @@ float4 StrataEnvironmentLighting(
 			// Finally sample the sky diffuse contribution (spherical harmonic, Lambert BRDF)
 			float3 DiffuseLookup = GetSkySHDiffuse(StrataEnvLight.DiffuseNormal) * View.SkyLightColor.rgb;			// STRATA_TODO compute SH coefficients for ON or Chan
 			// And accumulate
-			DiffuseSkyLighting = BSDFThroughput * (SkyVisData.SkyDiffuseLookUpMul * DiffuseLookup + SkyVisData.SkyDiffuseLookUpAdd) * StrataEnvLight.DiffuseAlbedo;
+			DiffuseSkyLighting = BSDFThroughput * (SkyVisData.SkyDiffuseLookUpMul * DiffuseLookup + SkyVisData.SkyDiffuseLookUpAdd) * StrataEnvLight.DiffuseColor;
 		}
 
 		// Specular component
diff --git a/Engine/Shaders/Private/Strata/StrataEvaluation.ush b/Engine/Shaders/Private/Strata/StrataEvaluation.ush
--- a/Engine/Shaders/Private/Strata/StrataEvaluation.ush
+++ b/Engine/Shaders/Private/Strata/StrataEvaluation.ush
@@ -57,7 +57,7 @@ void StrataGetBSDFRoughness(in FStrataBSDF BSDF, inout float R)
 	{
 	case STRATA_BSDF_TYPE_SLAB:
 	{
-		R = SLAB_ROUGHNESS0(BSDF);
+		R = SLAB_ROUGHNESSX(BSDF);
 		break;
 	}
 	case STRATA_BSDF_TYPE_SHEEN:
@@ -159,10 +159,10 @@ FStrataEvaluateResult StrataEvaluateBSDF(FStrataBSDFContext BSDFContext, FShadow
 	{
 		case STRATA_BSDF_TYPE_SLAB:
 		{
-			const float3 BaseColor = SLAB_ALBEDO(BSDFContext.BSDF);
+			const float3 BaseColor = SLAB_BASECOLOR(BSDFContext.BSDF);
 			const float  Metallic = SLAB_METALLIC(BSDFContext.BSDF);
 			const float3 DiffuseColor = BaseColor - BaseColor * Metallic;
-			const float SafeRoughness = MakeRoughnessSafe(SLAB_ROUGHNESS0(BSDFContext.BSDF));
+			const float SafeRoughness = MakeRoughnessSafe(SLAB_ROUGHNESSX(BSDFContext.BSDF));
 			float a2 = Pow4(SafeRoughness);
 
 			Sample.DiffusePathValue	= Diffuse_Chan(DiffuseColor, a2, BSDFContext.SatNoV, BSDFContext.SatNoL, BSDFContext.Context.VoH, BSDFContext.Context.NoH);
@@ -170,11 +170,11 @@ FStrataEvaluateResult StrataEvaluateBSDF(FStrataBSDFContext BSDFContext, FShadow
 			Sample.bSubsurface		= BSDF_GETHASSCATTERING(BSDFContext.BSDF) > 0;
 
 
-			const float3 F0 = ComputeF0(SLAB_REFLECTIVITY(BSDFContext.BSDF), BaseColor, Metallic);
+			const float3 F0 = ComputeF0(SLAB_SPECULAR(BSDFContext.BSDF), BaseColor, Metallic);
 			// STRATA_TODO F82
 
-			const float SafeRoughness0 = MakeRoughnessSafe(SLAB_ROUGHNESS0(BSDFContext.BSDF));
-			const float SafeRoughness1 = MakeRoughnessSafe(SLAB_ROUGHNESS1(BSDFContext.BSDF));
+			const float SafeRoughnessX = MakeRoughnessSafe(SLAB_ROUGHNESSX(BSDFContext.BSDF));
+			const float SafeRoughnessY = MakeRoughnessSafe(SLAB_ROUGHNESSY(BSDFContext.BSDF));
 
 			const bool bHasAnisotropy = BSDF_GETANISOTROPY(BSDFContext.BSDF);
 			float D = 0;
@@ -183,8 +183,8 @@ FStrataEvaluateResult StrataEvaluateBSDF(FStrataBSDFContext BSDFContext, FShadow
 			if (bHasAnisotropy)
 			{
 				// Generalized microfacet specular
-				float Alpha0 = Pow2(SafeRoughness0);
-				float Alpha1 = Pow2(SafeRoughness1);
+				float Alpha0 = Pow2(SafeRoughnessX);
+				float Alpha1 = Pow2(SafeRoughnessY);
 				D = D_GGXaniso(Alpha0, Alpha1, BSDFContext.Context.NoH, BSDFContext.Context.XoH, BSDFContext.Context.YoH);
 				Vis = Vis_SmithJointAniso(Alpha0, Alpha1, BSDFContext.Context.NoV, BSDFContext.Context.NoL, BSDFContext.Context.XoV, BSDFContext.Context.XoL, BSDFContext.Context.YoV, BSDFContext.Context.YoL);
 				PDF = VisibleGGXPDF_aniso(BSDFContext.TangentV, BSDFContext.TangentH, float2(Alpha0, Alpha1));
@@ -192,7 +192,7 @@ FStrataEvaluateResult StrataEvaluateBSDF(FStrataBSDFContext BSDFContext, FShadow
 			else
 			{
 				// Special override for roughness for supporting area light integrator with Sphere/Tube/Disk light, which modifies/increase roughness.
-				const float a2 = BSDFContext.a2Override > 0 ? BSDFContext.a2Override : Pow4(SafeRoughness0);
+				const float a2 = BSDFContext.a2Override > 0 ? BSDFContext.a2Override : Pow4(SafeRoughnessX);
 
 				// Generalized microfacet specular
 				D = D_GGX(a2, BSDFContext.Context.NoH);
@@ -210,7 +210,7 @@ FStrataEvaluateResult StrataEvaluateBSDF(FStrataBSDFContext BSDFContext, FShadow
 
 		case STRATA_BSDF_TYPE_SHEEN:
 		{
-			const float3 F0  = SHEEN_ALBEDO(BSDFContext.BSDF);
+			const float3 F0  = SHEEN_BASECOLOR(BSDFContext.BSDF);
 			const float SafeRoughness = MakeRoughnessSafe(SHEEN_ROUGHNESS(BSDFContext.BSDF));
 
 			// Special override for roughness for supporting area light integrator with Sphere/Tube/Disk light, which modifies/increase roughness.
@@ -295,7 +295,7 @@ FStrataImportanceSampleResult StrataImportanceSampleBSDF(FStrataPixelHeader Stra
 	{
 		case STRATA_BSDF_TYPE_SLAB:
 		{
-			const float3 BaseColor = SLAB_ALBEDO(BSDFContext.BSDF);
+			const float3 BaseColor = SLAB_BASECOLOR(BSDFContext.BSDF);
 			const float  Metallic = SLAB_METALLIC(BSDFContext.BSDF);
 
 			// STRATA_TODO take into account roughness
@@ -304,20 +304,20 @@ FStrataImportanceSampleResult StrataImportanceSampleBSDF(FStrataPixelHeader Stra
 			PDF = ImportanceSample.w;
 			Sample.DiffuseConeAngle = 1.0f / PDF;
 
-			const float3 F0 = ComputeF0(SLAB_REFLECTIVITY(BSDF), BaseColor, Metallic);
+			const float3 F0 = ComputeF0(SLAB_SPECULAR(BSDF), BaseColor, Metallic);
 			// STRATA_TODO F82
 
-			const float SafeRoughness0 = MakeRoughnessSafe(SLAB_ROUGHNESS0(BSDFContext.BSDF));
+			const float SafeRoughnessX = MakeRoughnessSafe(SLAB_ROUGHNESSX(BSDFContext.BSDF));
 			const bool bHasAnisotropy = BSDF_GETANISOTROPY(BSDFContext.BSDF);
 			float4 GGXSample = 0;
 			if (bHasAnisotropy)
 			{
-				const float SafeRoughness1 = MakeRoughnessSafe(SLAB_ROUGHNESS1(BSDFContext.BSDF));
-				GGXSample = ImportanceSampleVisibleGGX_aniso(UniformSampleDisk(E), float2(Pow2(SafeRoughness0), Pow2(SafeRoughness1)), BSDFContext.TangentV);
+				const float SafeRoughnessY = MakeRoughnessSafe(SLAB_ROUGHNESSY(BSDFContext.BSDF));
+				GGXSample = ImportanceSampleVisibleGGX_aniso(UniformSampleDisk(E), float2(Pow2(SafeRoughnessX), Pow2(SafeRoughnessY)), BSDFContext.TangentV);
 			}
 			else
 			{
-				GGXSample = ImportanceSampleVisibleGGX(UniformSampleDisk(E), Pow4(SafeRoughness0), BSDFContext.TangentV);
+				GGXSample = ImportanceSampleVisibleGGX(UniformSampleDisk(E), Pow4(SafeRoughnessX), BSDFContext.TangentV);
 			}
 
 			float3 WorldH = mul(GGXSample.xyz, BSDFContext.TangentBasis);
@@ -364,7 +364,7 @@ FStrataImportanceSampleResult StrataImportanceSampleBSDF(FStrataPixelHeader Stra
 
 struct FStrataEnvLightResult
 {
-	float3 DiffuseAlbedo;
+	float3 DiffuseColor;
 	float3 DiffuseNormal;
 
 	float3 SpecularDirection;
@@ -383,31 +383,30 @@ FStrataEnvLightResult StrataEvaluateForEnvLight(FStrataBSDFContext BSDFContext,
 	{
 		case STRATA_BSDF_TYPE_SLAB:
 		{
-			const float3 BaseColor = SLAB_ALBEDO(BSDFContext.BSDF);
+			const float3 BaseColor = SLAB_BASECOLOR(BSDFContext.BSDF);
 			const float  Metallic = SLAB_METALLIC(BSDFContext.BSDF);
-			const float3 DiffuseColor = BaseColor - BaseColor * Metallic;
 
-			StrataEnvLightResult.DiffuseAlbedo = DiffuseColor;
+			StrataEnvLightResult.DiffuseColor = BaseColor - BaseColor * Metallic;
 			StrataEnvLightResult.DiffuseNormal = BSDFContext.N;
 			StrataEnvLightResult.bSubsurface = BSDF_GETHASSCATTERING(BSDFContext.BSDF) > 0;
 
 			BRANCH
 			if (bEnableSpecular)
 			{
-				const float3 F0 = ComputeF0(SLAB_REFLECTIVITY(BSDFContext.BSDF), BaseColor, Metallic);
+				const float3 F0 = ComputeF0(SLAB_SPECULAR(BSDFContext.BSDF), BaseColor, Metallic);
 				// STRATA_TODO F82
 
-				float SafeRoughness0 = MakeRoughnessSafe(SLAB_ROUGHNESS0(BSDFContext.BSDF));
+				float SafeRoughnessX = MakeRoughnessSafe(SLAB_ROUGHNESSX(BSDFContext.BSDF));
 
 				const bool bHasAnisotropy = BSDF_GETANISOTROPY(BSDFContext.BSDF);
 				if (bHasAnisotropy)
 				{
-					const float SafeRoughness1 = MakeRoughnessSafe(SLAB_ROUGHNESS1(BSDFContext.BSDF));
+					const float SafeRoughnessY = MakeRoughnessSafe(SLAB_ROUGHNESSY(BSDFContext.BSDF));
 
 					// Modified the BSDF normal (and roughness)
 					float Anisotropy = 0;
-					GetAnisotropicFactor(SafeRoughness0, SafeRoughness1, Anisotropy, SafeRoughness0);
-					ModifyGGXAnisotropicNormalRoughness(BSDFContext.X, Anisotropy, SafeRoughness0, BSDFContext.N, BSDFContext.V);
+					GetAnisotropicFactor(SafeRoughnessX, SafeRoughnessY, Anisotropy, SafeRoughnessX);
+					ModifyGGXAnisotropicNormalRoughness(BSDFContext.X, Anisotropy, SafeRoughnessX, BSDFContext.N, BSDFContext.V);
 
 					// Update context (only needs: NoL/SatNoL/R) with the new N
 					BSDFContext.Context.NoL = dot(BSDFContext.N, BSDFContext.L);
@@ -415,18 +414,18 @@ FStrataEnvLightResult StrataEvaluateForEnvLight(FStrataBSDFContext BSDFContext,
 					BSDFContext.R = 2 * dot(BSDFContext.V, BSDFContext.N) * BSDFContext.N - BSDFContext.V;
 				}
 
-				float3 EvalEnvBRDF = EnvBRDFApprox(F0, SafeRoughness0, BSDFContext.SatNoV);
+				float3 EvalEnvBRDF = EnvBRDFApprox(F0, SafeRoughnessX, BSDFContext.SatNoV);
 
-				StrataEnvLightResult.SpecularDirection = StrataGetOffSpecularPeakReflectionDir(BSDFContext.N, BSDFContext.R, SafeRoughness0);
+				StrataEnvLightResult.SpecularDirection = StrataGetOffSpecularPeakReflectionDir(BSDFContext.N, BSDFContext.R, SafeRoughnessX);
 				StrataEnvLightResult.SpecularWeight = EvalEnvBRDF;
-				StrataEnvLightResult.SpecularSafeRoughness = SafeRoughness0;
+				StrataEnvLightResult.SpecularSafeRoughness = SafeRoughnessX;
 			}
 			break;
 		}
 
 		case STRATA_BSDF_TYPE_SHEEN:
 		{
-			StrataEnvLightResult.DiffuseAlbedo = SHEEN_ALBEDO(BSDFContext.BSDF);;
+			StrataEnvLightResult.DiffuseColor = SHEEN_BASECOLOR(BSDFContext.BSDF);;
 			StrataEnvLightResult.DiffuseNormal = BSDFContext.N;
 			break;
 		}
@@ -453,7 +452,7 @@ FStrataEnvLightResult StrataEvaluateForEnvLight(FStrataBSDFContext BSDFContext,
 			float3 L = normalize(V - N * dot(V, N));
 			StrataEnvLightResult.DiffuseNormal = L;
 
-			StrataEnvLightResult.DiffuseAlbedo = PI * HairShading(GBuffer, L, V, N, TransmissionShadow, TransmittanceData, BacklitEnabled, Area, Random);
+			StrataEnvLightResult.DiffuseColor = PI * HairShading(GBuffer, L, V, N, TransmissionShadow, TransmittanceData, BacklitEnabled, Area, Random);
 			// No specular environment contribution as of today if not using the special HairStrand render path
 			break;
 		}
diff --git a/Engine/Shaders/Private/Strata/StrataForwardLighting.ush b/Engine/Shaders/Private/Strata/StrataForwardLighting.ush
--- a/Engine/Shaders/Private/Strata/StrataForwardLighting.ush
+++ b/Engine/Shaders/Private/Strata/StrataForwardLighting.ush
@@ -50,13 +50,13 @@ float3 StrataForwardLighting(
 			// Evaluate environment lighting
 			const bool bEnableSpecular = ReflectionStruct.SkyLightParameters.y > 0.0f;
 			FStrataEnvLightResult StrataEnvLight = StrataEvaluateForEnvLight(StrataBSDFContext, bEnableSpecular);
-			if (any(StrataEnvLight.DiffuseAlbedo > 0.0f))
+			if (any(StrataEnvLight.DiffuseColor > 0.0f))
 			{
 				// SkySHDiffuse is pre integrated according to a Lambert diffuse lobe. Thus, we need to rescale to account for the phase function operating over 4PI.
 				// We apply a 2PI factor to recover average luminance over half hemisphere, and then times the phase function to account for the scattering events.
 				const float VolumeOnlyFactor = BSDF_GETTYPE(BSDF) == STRATA_BSDF_TYPE_VOLUME ? 2.0f * PI * IsotropicPhase() : 1.0f;
 
-				Color += GetEffectiveSkySHDiffuse(StrataEnvLight.DiffuseNormal) * StrataEnvLight.DiffuseAlbedo * SkyLightCommonFactor * VolumeOnlyFactor;
+				Color += GetEffectiveSkySHDiffuse(StrataEnvLight.DiffuseNormal) * StrataEnvLight.DiffuseColor * SkyLightCommonFactor * VolumeOnlyFactor;
 			}
 			if (any(StrataEnvLight.SpecularWeight > 0.0f))
 			{
diff --git a/Engine/Shaders/Private/Strata/StrataMaterialClassification.usf b/Engine/Shaders/Private/Strata/StrataMaterialClassification.usf
--- a/Engine/Shaders/Private/Strata/StrataMaterialClassification.usf
+++ b/Engine/Shaders/Private/Strata/StrataMaterialClassification.usf
@@ -43,15 +43,15 @@ void ComputePostProcessBSDFData(
 	{
 	case STRATA_BSDF_TYPE_SLAB:
 	{
-		const float3 F0 = ComputeF0(SLAB_REFLECTIVITY(BSDF), SLAB_ALBEDO(BSDF), SLAB_METALLIC(BSDF));
+		const float3 F0 = ComputeF0(SLAB_SPECULAR(BSDF), SLAB_BASECOLOR(BSDF), SLAB_METALLIC(BSDF));
 
 		if (bTopLayer)
 		{
 			// Setup SSR for the top layer
-			TopLayer.Roughness += BSDFCoverageAvg * SLAB_ROUGHNESS0(BSDF);
+			TopLayer.Roughness += BSDFCoverageAvg * SLAB_ROUGHNESSX(BSDF);
 			TopLayer.AccumRoughnessCount += BSDFCoverageAvg;
 
-			const float SafeRoughness = MakeRoughnessSafe(SLAB_ROUGHNESS0(BSDF));
+			const float SafeRoughness = MakeRoughnessSafe(SLAB_ROUGHNESSX(BSDF));
 			TopLayer.BaseColor += BSDFCoverage * EnvBRDFApprox(F0, SafeRoughness, saturate(ContextNoLight.NoV));
 		}
 
@@ -62,7 +62,7 @@ void ComputePostProcessBSDFData(
 			SSS.bHasProfile	= BSDF_GETHASSSSPROFILE(BSDF);
 			SSS.Profile		= SLAB_SSSPROFILEID(BSDF);
 			SSS.RadiusScale	= SLAB_SSSPROFILERADIUS(BSDF);
-			SSS.BaseColor	= SLAB_ALBEDO(BSDF);
+			SSS.BaseColor	= SLAB_BASECOLOR(BSDF);
 			SSS.DMFPAlbedo	= SLAB_SSSDMFPALBEDO(BSDF);
 			SSS.DMFPRadius	= SLAB_SSSDMFPRADIUS(BSDF);
 		}
@@ -71,7 +71,7 @@ void ComputePostProcessBSDFData(
 
 	case STRATA_BSDF_TYPE_SHEEN:
 	{
-		const float3 F0  = SHEEN_ALBEDO(BSDF);
+		const float3 F0  = SHEEN_BASECOLOR(BSDF);
 		const float3 F   = F_Schlick(F0, ContextNoLight.VoH);
 
 		if (bTopLayer)
diff --git a/Engine/Shaders/Private/Strata/StrataVisualize.usf b/Engine/Shaders/Private/Strata/StrataVisualize.usf
--- a/Engine/Shaders/Private/Strata/StrataVisualize.usf
+++ b/Engine/Shaders/Private/Strata/StrataVisualize.usf
@@ -231,23 +231,23 @@ void VisualizeMaterialPS(
 				NEXTLINE;
 
 				HORIZONTALTAB; PRINTPARAMCHAR(_A_); PRINTPARAMCHAR(_L_); PRINTPARAMCHAR(_B_); PRINTPARAMCHAR(_E_); PRINTPARAMCHAR(_D_); PRINTPARAMCHAR(_O_); 
-				HORIZONTALTAB; PRINTFLOAT3(SLAB_ALBEDO(BSDF));
+				HORIZONTALTAB; PRINTFLOAT3(SLAB_BASECOLOR(BSDF));
 				NEXTLINE;
 
 				HORIZONTALTAB; PRINTPARAMCHAR(_R_); PRINTPARAMCHAR(_E_); PRINTPARAMCHAR(_F_); PRINTPARAMCHAR(_L_); PRINTPARAMCHAR(_E_); PRINTPARAMCHAR(_C_);
-				HORIZONTALTAB; PRINTFLOAT1(SLAB_REFLECTIVITY(BSDF));
+				HORIZONTALTAB; PRINTFLOAT1(SLAB_SPECULAR(BSDF));
 				NEXTLINE;
 				HORIZONTALTAB; PRINTPARAMCHAR(_M_); PRINTPARAMCHAR(_E_); PRINTPARAMCHAR(_T_); PRINTPARAMCHAR(_A_); PRINTPARAMCHAR(_L_); PRINTPARAMCHAR(_SPC_);
 				HORIZONTALTAB; PRINTFLOAT1(SLAB_METALLIC(BSDF));
 				NEXTLINE;
 
 				HORIZONTALTAB; PRINTPARAMCHAR(_R_); PRINTPARAMCHAR(_O_); PRINTPARAMCHAR(_U_); PRINTPARAMCHAR(_G_); PRINTPARAMCHAR(_H_); PRINTPARAMCHAR(_0_);
-				HORIZONTALTAB; PRINTFLOAT1(SLAB_ROUGHNESS0(BSDF));
+				HORIZONTALTAB; PRINTFLOAT1(SLAB_ROUGHNESSX(BSDF));
 				NEXTLINE;
 				if (BSDF_GETANISOTROPY(BSDF))
 				{
 					HORIZONTALTAB; PRINTPARAMCHAR(_R_); PRINTPARAMCHAR(_O_); PRINTPARAMCHAR(_U_); PRINTPARAMCHAR(_G_); PRINTPARAMCHAR(_H_); PRINTPARAMCHAR(_1_);
-					HORIZONTALTAB; PRINTFLOAT1(SLAB_ROUGHNESS1(BSDF));
+					HORIZONTALTAB; PRINTFLOAT1(SLAB_ROUGHNESSY(BSDF));
 					NEXTLINE;
 				}
 
@@ -294,7 +294,7 @@ void VisualizeMaterialPS(
 				NEXTLINE;
 
 				HORIZONTALTAB; PRINTPARAMCHAR(_A_); PRINTPARAMCHAR(_L_); PRINTPARAMCHAR(_B_); PRINTPARAMCHAR(_E_); PRINTPARAMCHAR(_D_); PRINTPARAMCHAR(_O_); 
-				HORIZONTALTAB; PRINTFLOAT3(SHEEN_ALBEDO(BSDF));
+				HORIZONTALTAB; PRINTFLOAT3(SHEEN_BASECOLOR(BSDF));
 				NEXTLINE;
 				HORIZONTALTAB; PRINTPARAMCHAR(_R_); PRINTPARAMCHAR(_O_); PRINTPARAMCHAR(_U_); PRINTPARAMCHAR(_G_); PRINTPARAMCHAR(_H_); PRINTPARAMCHAR(_SPC_);
 				HORIZONTALTAB; PRINTFLOAT1(SHEEN_ROUGHNESS(BSDF));
diff --git a/Engine/Source/Runtime/Engine/Classes/Materials/MaterialExpressionStrata.h b/Engine/Source/Runtime/Engine/Classes/Materials/MaterialExpressionStrata.h
--- a/Engine/Source/Runtime/Engine/Classes/Materials/MaterialExpressionStrata.h
+++ b/Engine/Source/Runtime/Engine/Classes/Materials/MaterialExpressionStrata.h
@@ -41,45 +41,45 @@ class UMaterialExpressionStrataSlabBSDF : public UMaterialExpression // STRATA_T
 	GENERATED_UCLASS_BODY()
 
 	/**
-	 * Albedo (type = float3, unit = unitless, defaults to 0.18)
+	 * Defines the overall color of the Material. (type = float3, unit = unitless, defaults to 0.18)
 	 */
 	UPROPERTY()
 	FExpressionInput BaseColor;
 
 	/**
-	 * 0 means dielectric, 1 means conductor (type = float, unit = unitless, defaults to 0)
+	 * Controls how \"metal-like\" your surface looks like. 0 means dielectric, 1 means conductor (type = float, unit = unitless, defaults to 0)
 	 */
 	UPROPERTY()
 	FExpressionInput Metallic;
 	
 	/**
-	 * Reflectivity when view direction is perpendicular to the surface, also known as F0 (type = float, unit = unitless, defaults to plastic 0.5)
+	 * Used to scale the current amount of specularity on non-metallic surfaces and is a value between 0 and 1 (type = float, unit = unitless, defaults to plastic 0.5)
 	 */
 	UPROPERTY()
-	FExpressionInput Reflectivity;
+	FExpressionInput Specular;
 
 	// STRATA_TODO: edge or F82 EdgeColor?
 
 	/**
-	 * RoughnessX (type = float, unit = unitless, defaults to 0)
+	 * Controls how rough the Material is. Roughness of 0 (smooth) is a mirror reflection and 1 (rough) is completely matte or diffuse. When using anisotropy, it is the roughness used along the Tangent axis. (type = float, unit = unitless, defaults to 0)
 	 */
 	UPROPERTY()
 	FExpressionInput RoughnessX;
 		
 	/**
-	 * RoughnessY (type = float, unit = unitless), If not plugged in, RoughnessY is set to RoughnessX to get an isotropic behavior.
+	 * Controls the roughness along the secondary surface tangent vector (perpendicular to Tangent). (type = float, unit = unitless). If not plugged in, RoughnessY is set to RoughnessX to disable anisotropy, resulting in an isotropic behavior.
 	 */
 	UPROPERTY()
 	FExpressionInput RoughnessY;
 
 	/**
-	 * Normal (type = float3, unit = unitless, defaults to vertex normal)
+	 * Take the surface normal as input. The normal is considered tangent or world space according to the space properties on the main material node. (type = float3, unit = unitless, defaults to vertex normal)
 	 */
 	UPROPERTY()
 	FExpressionInput Normal;
 
 	/**
-	* Tangent (type = float3, unit = unitless, defaults to vertex tangent)
+	* Take a surface tangent as input. The tangent is considered tangent or world space according to the space properties on the main material node. (type = float3, unit = unitless, defaults to vertex tangent)
 	*/
 	UPROPERTY()
 	FExpressionInput Tangent;
@@ -119,10 +119,10 @@ class UMaterialExpressionStrataSheenBSDF : public UMaterialExpression
 	GENERATED_UCLASS_BODY()
 
 	/**
-	* Albedo (type = float3, unit = unitless)
+	* Defines the overall color of the Material. (type = float3, unit = unitless)
 	*/
 	UPROPERTY()
-	FExpressionInput Albedo;
+	FExpressionInput BaseColor;
 
 	/**
 	 * Roughness (type = float, unit = unitless)
@@ -131,7 +131,7 @@ class UMaterialExpressionStrataSheenBSDF : public UMaterialExpression
 	FExpressionInput Roughness;
 
 	/**
-	 * Normal (type = float3, unit = unitless)
+	 * Take the surface normal as input. The normal is considered tangent or world space according to the space properties on the main material node. (type = float3, unit = unitless)
 	 */
 	UPROPERTY()
 	FExpressionInput Normal;
@@ -154,7 +154,7 @@ class UMaterialExpressionStrataVolumetricFogCloudBSDF : public UMaterialExpressi
 	GENERATED_UCLASS_BODY()
 
 	/**
-	* Albedo (type = float3, unit = unitless, default = 0)
+	* The single scattering Albedo defining the overall color of the Material (type = float3, unit = unitless, default = 0)
 	*/
 	UPROPERTY()
 	FExpressionInput Albedo;
@@ -325,7 +325,7 @@ class UMaterialExpressionStrataSingleLayerWaterBSDF : public UMaterialExpression
 	FExpressionInput TopMaterialOpacity;
 
 	/**
-	 * Albedo (type = float3, unit = unitless, defaults to black)
+	* The single scattering Albedo defining the overall color of the Material (type = float3, unit = unitless, default = 0)
 	 */
 	UPROPERTY()
 	FExpressionInput WaterAlbedo;
diff --git a/Engine/Source/Runtime/Engine/Private/Materials/HLSLMaterialTranslator.cpp b/Engine/Source/Runtime/Engine/Private/Materials/HLSLMaterialTranslator.cpp
--- a/Engine/Source/Runtime/Engine/Private/Materials/HLSLMaterialTranslator.cpp
+++ b/Engine/Source/Runtime/Engine/Private/Materials/HLSLMaterialTranslator.cpp
@@ -7536,11 +7536,11 @@ int32 FHLSLMaterialTranslator::StrataCreateAndRegisterNullMaterial()
 	return OutputCodeChunk;
 }
 
-int32 FHLSLMaterialTranslator::StrataSlabBSDF(int32 Albedo, int32 Reflectivity, int32 Metallic, int32 RoughnessX, int32 RoughnessY, int32 SSSProfileId, int32 SSSDMFPAlbedo, int32 SSSDMFPRadius, int32 Normal, int32 Tangent, uint8 SharedNormalIndex)
+int32 FHLSLMaterialTranslator::StrataSlabBSDF(int32 BaseColor, int32 Reflectivity, int32 Metallic, int32 RoughnessX, int32 RoughnessY, int32 SSSProfileId, int32 SSSDMFPAlbedo, int32 SSSDMFPRadius, int32 Normal, int32 Tangent, uint8 SharedNormalIndex)
 {
 	return AddCodeChunk(
 		MCT_Strata, TEXT("GetStrataSlabBSDF(%s, %s, %s, float2(%s, %s), %s, %s, %s, %u, Parameters.SharedNormals.NormalTypes) /* Normal = %s ; Tangent = %s */"),
-		*GetParameterCode(Albedo),
+		*GetParameterCode(BaseColor),
 		*GetParameterCode(Reflectivity),
 		*GetParameterCode(Metallic),
 		*GetParameterCode(RoughnessX),
@@ -7554,11 +7554,11 @@ int32 FHLSLMaterialTranslator::StrataSlabBSDF(int32 Albedo, int32 Reflectivity,
 	);
 }
 
-int32 FHLSLMaterialTranslator::StrataSheenBSDF(int32 Albedo, int32 Roughness, int32 Normal, uint8 SharedNormalIndex)
+int32 FHLSLMaterialTranslator::StrataSheenBSDF(int32 BaseColor, int32 Roughness, int32 Normal, uint8 SharedNormalIndex)
 {
 	return AddCodeChunk(
 		MCT_Strata, TEXT("GetStrataSheenBSDF(%s, %s, %u) /* %s */"),
-		*GetParameterCode(Albedo),
+		*GetParameterCode(BaseColor),
 		*GetParameterCode(Roughness),
 		SharedNormalIndex,
 		*GetParameterCode(Normal)
diff --git a/Engine/Source/Runtime/Engine/Private/Materials/HLSLMaterialTranslator.h b/Engine/Source/Runtime/Engine/Private/Materials/HLSLMaterialTranslator.h
--- a/Engine/Source/Runtime/Engine/Private/Materials/HLSLMaterialTranslator.h
+++ b/Engine/Source/Runtime/Engine/Private/Materials/HLSLMaterialTranslator.h
@@ -798,8 +798,8 @@ protected:
 
 	// Strata
 	virtual int32 StrataCreateAndRegisterNullMaterial() override;
-	virtual int32 StrataSlabBSDF(int32 Albedo, int32 Reflectivity, int32 Metallic, int32 RoughnessX, int32 RoughnessY, int32 SSSProfileId, int32 SSSDMFPAlbedo, int32 SSSDMFPRadius, int32 Normal, int32 Tangent, uint8 SharedNormalIndex) override;
-	virtual int32 StrataSheenBSDF(int32 Albedo, int32 Roughness, int32 Normal, uint8 SharedNormalIndex) override;
+	virtual int32 StrataSlabBSDF(int32 BaseColor, int32 Reflectivity, int32 Metallic, int32 RoughnessX, int32 RoughnessY, int32 SSSProfileId, int32 SSSDMFPAlbedo, int32 SSSDMFPRadius, int32 Normal, int32 Tangent, uint8 SharedNormalIndex) override;
+	virtual int32 StrataSheenBSDF(int32 BaseColor, int32 Roughness, int32 Normal, uint8 SharedNormalIndex) override;
 	virtual int32 StrataVolumetricFogCloudBSDF(int32 Albedo, int32 Extinction, int32 Emissive, int32 AmbientOcclusion) override;
 	virtual int32 StrataUnlitBSDF(int32 Emissive, int32 TransmittanceColor) override;
 	virtual int32 StrataHairBSDF(int32 BaseColor, int32 Scatter, int32 Specular, int32 Roughness, int32 Backlit, int32 Emissive, int32 Tangent, uint8 SharedNormalIndex) override;
diff --git a/Engine/Source/Runtime/Engine/Private/Materials/MaterialExpressions.cpp b/Engine/Source/Runtime/Engine/Private/Materials/MaterialExpressions.cpp
--- a/Engine/Source/Runtime/Engine/Private/Materials/MaterialExpressions.cpp
+++ b/Engine/Source/Runtime/Engine/Private/Materials/MaterialExpressions.cpp
@@ -19683,7 +19683,7 @@ int32 UMaterialExpressionStrataSlabBSDF::Compile(class FMaterialCompiler* Compil
 
 	int32 OutputCodeChunk = Compiler->StrataSlabBSDF(
 		CompileWithDefaultFloat3(Compiler, BaseColor, 0.18f, 0.18f, 0.18f),
-		CompileWithDefaultFloat1(Compiler, Reflectivity, 0.5f),
+		CompileWithDefaultFloat1(Compiler, Specular, 0.5f),
 		CompileWithDefaultFloat1(Compiler, Metallic, 0.0f),
 		RoughnessXCodeChunk,
 		RoughnessYCodeChunk,
@@ -19713,13 +19713,13 @@ uint32 UMaterialExpressionStrataSlabBSDF::GetInputType(int32 InputIndex)
 	switch (InputIndex)
 	{
 	case 0:
-		return MCT_Float3; // Albedo
+		return MCT_Float3; // BaseColor
 		break;
 	case 1:
 		return MCT_Float1; // Metallic
 		break;
 	case 2:
-		return MCT_Float1; // Reflectivity
+		return MCT_Float1; // Specular
 		break;
 	case 3:
 		return MCT_Float1; // RoughnessX
@@ -19749,7 +19749,7 @@ FName UMaterialExpressionStrataSlabBSDF::GetInputName(int32 InputIndex) const
 {
 	if (InputIndex == 0)
 	{
-		return TEXT("Albedo");
+		return TEXT("BaseColor");
 	}
 	else if (InputIndex == 1)
 	{
@@ -19757,7 +19757,7 @@ FName UMaterialExpressionStrataSlabBSDF::GetInputName(int32 InputIndex) const
 	}
 	else if (InputIndex == 2)
 	{
-		return TEXT("Reflectivity");
+		return TEXT("Specular");
 	}
 	else if (InputIndex == 3)
 	{
@@ -19831,7 +19831,7 @@ int32 UMaterialExpressionStrataSheenBSDF::Compile(class FMaterialCompiler* Compi
 	uint8 SharedNormalIndex = StrataCompilationInfoCreateSharedNormal(Compiler, NormalCodeChunk);
 
 	int32 OutputCodeChunk = Compiler->StrataSheenBSDF(
-		CompileWithDefaultFloat3(Compiler, Albedo,		0.0f, 0.0f, 0.0f),
+		CompileWithDefaultFloat3(Compiler, BaseColor,	0.0f, 0.0f, 0.0f),
 		CompileWithDefaultFloat1(Compiler, Roughness,	0.0f),
 		NormalCodeChunk,
 		SharedNormalIndex);
diff --git a/Engine/Source/Runtime/Engine/Public/MaterialCompiler.h b/Engine/Source/Runtime/Engine/Public/MaterialCompiler.h
--- a/Engine/Source/Runtime/Engine/Public/MaterialCompiler.h
+++ b/Engine/Source/Runtime/Engine/Public/MaterialCompiler.h
@@ -392,8 +392,8 @@ public:
 
 	// Strata
 	virtual int32 StrataCreateAndRegisterNullMaterial() = 0;
-	virtual int32 StrataSlabBSDF(int32 Albedo, int32 Reflectivity, int32 Metallic, int32 RoughnessX, int32 RoughnessY, int32 SSSProfileId, int32 SSSDMFPAlbedo, int32 SSSDMFPRadius, int32 Normal, int32 Tangent, uint8 SharedNormalIndex) = 0;
-	virtual int32 StrataSheenBSDF(int32 Albedo, int32 Roughness, int32 Normal, uint8 SharedNormalIndex) = 0;
+	virtual int32 StrataSlabBSDF(int32 BaseColor, int32 Reflectivity, int32 Metallic, int32 RoughnessX, int32 RoughnessY, int32 SSSProfileId, int32 SSSDMFPAlbedo, int32 SSSDMFPRadius, int32 Normal, int32 Tangent, uint8 SharedNormalIndex) = 0;
+	virtual int32 StrataSheenBSDF(int32 BaseColor, int32 Roughness, int32 Normal, uint8 SharedNormalIndex) = 0;
 	virtual int32 StrataVolumetricFogCloudBSDF(int32 Albedo, int32 Extinction, int32 Emissive, int32 AmbientOcclusion) = 0;
 	virtual int32 StrataUnlitBSDF(int32 Emissive, int32 TransmittanceColor) = 0;
 	virtual int32 StrataHairBSDF(int32 BaseColor, int32 Scatter, int32 Specular, int32 Roughness, int32 Backlit, int32 Emissive, int32 Tangent, uint8 SharedNormalIndex) = 0;
@@ -808,14 +808,14 @@ public:
 		return Compiler->StrataCreateAndRegisterNullMaterial();
 	}
 
-	virtual int32 StrataSlabBSDF(int32 Albedo, int32 Reflectivity, int32 Metallic, int32 RoughnessX, int32 RoughnessY, int32 SSSProfileId, int32 SSSDMFPAlbedo, int32 SSSDMFPRadius, int32 Normal, int32 Tangent, uint8 SharedNormalIndex) override
+	virtual int32 StrataSlabBSDF(int32 BaseColor, int32 Reflectivity, int32 Metallic, int32 RoughnessX, int32 RoughnessY, int32 SSSProfileId, int32 SSSDMFPAlbedo, int32 SSSDMFPRadius, int32 Normal, int32 Tangent, uint8 SharedNormalIndex) override
 	{
-		return Compiler->StrataSlabBSDF(Albedo, Reflectivity, Metallic, RoughnessX, RoughnessY, SSSProfileId, SSSDMFPAlbedo, SSSDMFPRadius, Normal, Tangent, SharedNormalIndex);
+		return Compiler->StrataSlabBSDF(BaseColor, Reflectivity, Metallic, RoughnessX, RoughnessY, SSSProfileId, SSSDMFPAlbedo, SSSDMFPRadius, Normal, Tangent, SharedNormalIndex);
 	}
 
-	virtual int32 StrataSheenBSDF(int32 Albedo, int32 Roughness, int32 Normal, uint8 SharedNormalIndex) override
+	virtual int32 StrataSheenBSDF(int32 BaseColor, int32 Roughness, int32 Normal, uint8 SharedNormalIndex) override
 	{
-		return Compiler->StrataSheenBSDF(Albedo, Roughness, Normal, SharedNormalIndex);
+		return Compiler->StrataSheenBSDF(BaseColor, Roughness, Normal, SharedNormalIndex);
 	}
 
 	virtual int32 StrataVolumetricFogCloudBSDF(int32 Albedo, int32 Extinction, int32 Emissive, int32 AmbientOcclusion) override
