Refactor path tracer MIS calculations so shadow rays are handled consistently

Instead of re-using the indirect ray for shadow calculations with BxDF sampling, trace these rays with a dedicated loop so that shadow casting options can be hanlded consistently. This runs slowers but allows us to correctly handles options like non-shadow casting objects and materials.

Added support for non-shadow casting lights

Fixed an RDG validation error from the render target action

diff --git a/Engine/Shaders/Private/PathTracing/Light/PathTracingDirectionalLight.ush b/Engine/Shaders/Private/PathTracing/Light/PathTracingDirectionalLight.ush
--- a/Engine/Shaders/Private/PathTracing/Light/PathTracingDirectionalLight.ush
+++ b/Engine/Shaders/Private/PathTracing/Light/PathTracingDirectionalLight.ush
@@ -8,7 +8,7 @@ DirectionalLight.usf: Light sampling functions for Directional light implementat
 
 #include "PathTracingLightCommon.ush"
 
-float4 DirectionalLight_TraceLight(RayDesc Ray, int LightId, float HitT)
+float4 DirectionalLight_TraceLight(RayDesc Ray, int LightId, inout float HitT)
 {
 	if (Ray.TMax == HitT) // we hit the background
 	{
diff --git a/Engine/Shaders/Private/PathTracing/Light/PathTracingLightCommon.ush b/Engine/Shaders/Private/PathTracing/Light/PathTracingLightCommon.ush
--- a/Engine/Shaders/Private/PathTracing/Light/PathTracingLightCommon.ush
+++ b/Engine/Shaders/Private/PathTracing/Light/PathTracingLightCommon.ush
@@ -202,3 +202,8 @@ bool IsStationary(int LightId)
 {
 	return (SceneLights[LightId].Flags & PATHTRACER_FLAG_STATIONARY_MASK) != 0;
 }
+
+bool CastsShadow(int LightId)
+{
+	return (SceneLights[LightId].Flags & PATHTRACER_FLAG_CAST_SHADOW_MASK) != 0;
+}
diff --git a/Engine/Shaders/Private/PathTracing/Light/PathTracingLightSampling.ush b/Engine/Shaders/Private/PathTracing/Light/PathTracingLightSampling.ush
--- a/Engine/Shaders/Private/PathTracing/Light/PathTracingLightSampling.ush
+++ b/Engine/Shaders/Private/PathTracing/Light/PathTracingLightSampling.ush
@@ -20,7 +20,7 @@ LightSampling.usf: Light sampling for next-event estimation
 #include "PathTracingSpotLight.ush"
 
 // Returns float4 packed with Radiance and Pdf
-float4 TraceLight(RayDesc Ray, int LightId, float HitT)
+float4 TraceLight(RayDesc Ray, int LightId, inout float HitT)
 {
 	switch (SceneLights[LightId].Flags & PATHTRACER_FLAG_TYPE_MASK)
 	{
diff --git a/Engine/Shaders/Private/PathTracing/Light/PathTracingPointLight.ush b/Engine/Shaders/Private/PathTracing/Light/PathTracingPointLight.ush
--- a/Engine/Shaders/Private/PathTracing/Light/PathTracingPointLight.ush
+++ b/Engine/Shaders/Private/PathTracing/Light/PathTracingPointLight.ush
@@ -8,7 +8,7 @@ PointLight.usf: Light sampling functions for Point light implementation
 
 #include "PathTracingLightCommon.ush"
 
-float4 PointLight_TraceLight(RayDesc Ray, int LightId, float HitT)
+float4 PointLight_TraceLight(RayDesc Ray, int LightId, inout float HitT)
 {
 	float3 LightPosition = GetPosition(LightId);
 	float LightRadius = GetRadius(LightId);
@@ -29,6 +29,7 @@ float4 PointLight_TraceLight(RayDesc Ray, int LightId, float HitT)
 		float t = (-b - h) / a;
 		if (t > Ray.TMin && t < HitT)
 		{
+			HitT = t;
 			// #dxr_todo: sphere area is 4*pi*r^2 -- but the factor of 4 is missing for some reason?
 			float3 LightPower = GetColor(LightId) * ComputeIESAttenuation(LightId, Ray.Origin);
 			float3 LightRadiance = LightPower / (PI * LightRadius2);
diff --git a/Engine/Shaders/Private/PathTracing/Light/PathTracingRectLight.ush b/Engine/Shaders/Private/PathTracing/Light/PathTracingRectLight.ush
--- a/Engine/Shaders/Private/PathTracing/Light/PathTracingRectLight.ush
+++ b/Engine/Shaders/Private/PathTracing/Light/PathTracingRectLight.ush
@@ -56,7 +56,7 @@ float4 SampleApproxProjectionSphericalRect(float2 Rand, FSphericalRect Spherical
 }
 
 
-float4 RectLight_TraceLight(RayDesc Ray, int LightId, float HitT)
+float4 RectLight_TraceLight(RayDesc Ray, int LightId, inout float HitT)
 {
 	float3 LightPosition = GetPosition(LightId);
 	float3 LightNormal = GetNormal(LightId);
@@ -89,6 +89,7 @@ float4 RectLight_TraceLight(RayDesc Ray, int LightId, float HitT)
 			UV = float2(dot(P, LightdPdu), dot(P, LightdPdv));
 			if (all(abs(UV) <= Rect.Extent))
 			{
+				HitT = t;
 				// stored color is radiance
 				float3 Radiance = GetColor(LightId) * ComputeIESAttenuation(LightId, Ray.Origin);
 				Radiance *= ComputeAttenuationFalloff(dot(LightDirection, LightDirection), LightId);
diff --git a/Engine/Shaders/Private/PathTracing/Light/PathTracingSkyLight.ush b/Engine/Shaders/Private/PathTracing/Light/PathTracingSkyLight.ush
--- a/Engine/Shaders/Private/PathTracing/Light/PathTracingSkyLight.ush
+++ b/Engine/Shaders/Private/PathTracing/Light/PathTracingSkyLight.ush
@@ -27,7 +27,7 @@ void SkyLight_EstimateLight(
 	OutIrradiance = SkyLightIlluminance * 2.0f * PI;
 }
 
-float4 SkyLight_TraceLight(RayDesc Ray, int LightId, float HitT) {
+float4 SkyLight_TraceLight(RayDesc Ray, int LightId, inout float HitT) {
 	if (Ray.TMax == HitT)
 	{
 		// LightUV convention in the code below is just the world space ray direction
diff --git a/Engine/Shaders/Private/PathTracing/Light/PathTracingSpotLight.ush b/Engine/Shaders/Private/PathTracing/Light/PathTracingSpotLight.ush
--- a/Engine/Shaders/Private/PathTracing/Light/PathTracingSpotLight.ush
+++ b/Engine/Shaders/Private/PathTracing/Light/PathTracingSpotLight.ush
@@ -26,7 +26,7 @@ float SpotAttenuation(float3 L, float3 SpotDirection, float2 SpotAngles)
 }
 
 
-float4 SpotLight_TraceLight(RayDesc Ray, int LightId, float HitT) {
+float4 SpotLight_TraceLight(RayDesc Ray, int LightId, inout float HitT) {
 	// #dxr_todo: share with point light instead of copy/pasting logic?
 	float3 LightPosition = GetPosition(LightId);
 	float LightRadius = GetRadius(LightId);
@@ -47,6 +47,7 @@ float4 SpotLight_TraceLight(RayDesc Ray, int LightId, float HitT) {
 		float t = (-b - h) / a;
 		if (t > Ray.TMin && t < HitT)
 		{
+			HitT = t;
 			// #dxr_todo: sphere area is 4*pi*r^2 -- but the factor of 4 is missing for some reason?
 			float3 LightPower = GetColor(LightId) * ComputeIESAttenuation(LightId, Ray.Origin);
 			float3 LightRadiance = LightPower / (PI * LightRadius2);
diff --git a/Engine/Shaders/Private/PathTracing/PathTracing.usf b/Engine/Shaders/Private/PathTracing/PathTracing.usf
--- a/Engine/Shaders/Private/PathTracing/PathTracing.usf
+++ b/Engine/Shaders/Private/PathTracing/PathTracing.usf
@@ -61,7 +61,6 @@ RAY_TRACING_ENTRY_RAYGEN(PathTracingMainRG)
 	}
 
 	float3 Radiance = 0;
-	uint RayCounter = 0;
 
 	// Initialize ray and payload
 	RayDesc Ray;
@@ -88,11 +87,6 @@ RAY_TRACING_ENTRY_RAYGEN(PathTracingMainRG)
 	// path state variables (these cary information between bounces)
 	float3 PathThroughput = 1.0;
 
-	// We don't want the camera ray to have any MIS contribution against direct light sampling so we pretend that it came from
-	// a singular Material scattering event.
-	float PrevMaterialPdf = POSITIVE_INFINITY;
-	int PrevPrimitiveLightingChannelMask = 0xFF;
-
 	// 0: only Material sampling
 	// 1: only Light sampling
 	// 2: both Material and Light
@@ -109,37 +103,28 @@ RAY_TRACING_ENTRY_RAYGEN(PathTracingMainRG)
 			LaunchIndex,
 			bInitialDisableSkyLightContribution,
 			bIgnoreTranslucentMaterials);
-		RayCounter++;
 
-		if (MISMode != 1 || Bounce == 0)
+		// directly visible lights
+		if (Bounce == 0 && SceneLightCount > 0)
 		{
-			float HitT = Payload.IsMiss() ? Ray.TMax : Payload.HitT;
+			uint NumLights = 0;
+			if (PathTracingData.VisibleLights)
+			{
+				// user wants to see all lights
+				NumLights = SceneLightCount;
+			}
+			else if (IsEnvironmentLight(0))
+			{
+				// user doesn't want to see all lights - still check the first one in case its a skydome (there can be only one, and it always comes first if it is there)
+				NumLights = 1;
+			}
 			// Loop over lights to capture their contribution
 			// #dxr_todo: if we have lots of lights, having some hierarchical structure would be better ....
-			float PreviousCdfValue = 0.0;
-			for (uint LightId = 0; LightId < SceneLightCount; ++LightId)
+			for (uint LightId = 0; LightId < NumLights; ++LightId)
 			{
-				// Camera rays should only do the rest of the loop if we have visible lights OR if
-				// we have a skydome (which are always visible)
-				if (Bounce == 0 && PathTracingData.VisibleLights == 0 && !IsEnvironmentLight(LightId))
-				{
-					continue;
-				}
-
-				if ((PrevPrimitiveLightingChannelMask & GetLightingChannelMask(LightId)) == 0)
-				{
-					// light does not affect the current ray
-					continue;
-				}
-
-				float4 LightResult = TraceLight(Ray, LightId, HitT);
-				float3 LightRadiance = LightResult.xyz;
-				float LightPdf = LightResult.w;
-
-				float LightPickPdf = LightPickingCdf[LightId] - PreviousCdfValue;
-				PreviousCdfValue = LightPickingCdf[LightId];
-
-				AccumulateRadiance(Radiance, PathThroughput * LightRadiance * MISWeightRobust(PrevMaterialPdf, LightPdf * LightPickPdf));
+				float HitT = Payload.IsMiss() ? Ray.TMax : Payload.HitT;
+				float3 LightRadiance = TraceLight(Ray, LightId, HitT).xyz;
+				AccumulateRadiance(Radiance, PathThroughput * LightRadiance);
 			}
 		}
 
@@ -149,8 +134,25 @@ RAY_TRACING_ENTRY_RAYGEN(PathTracingMainRG)
 			break;
 		}
 
-		// add in surface emission
-		AccumulateRadiance(Radiance, PathThroughput * Payload.Radiance);
+		if (Payload.ShadingModelID == SHADINGMODELID_SUBSURFACE)
+		{
+			// for now just approximate SSS by moving it into the diffuse color
+			Payload.DiffuseColor += Payload.CustomData.xyz;
+			Payload.Opacity = 1; // overwrite radius
+			Payload.ShadingModelID == SHADINGMODELID_DEFAULT_LIT;
+		}
+
+		// add in surface emission (except for modulate blend mode which uses Radiance as the transparency multiplier)
+		if (Payload.BlendingMode != RAY_TRACING_BLEND_MODE_MODULATE)
+		{
+			AccumulateRadiance(Radiance, PathThroughput * Payload.Radiance);
+		}
+
+		if (Bounce == PathTracingData.MaxBounces)
+		{
+			// we've arrived at the last bounce - nothing left to do
+			break;
+		}
 
 		// Update origin
 		Ray.Origin = Payload.WorldPos;
@@ -164,14 +166,13 @@ RAY_TRACING_ENTRY_RAYGEN(PathTracingMainRG)
 		// Choose a random number for both Light sampling and BxDF sampling
 		float4 RandSample = RandomSequence_GenerateSample4D(RandSequence);
 
-		// Does this material require NEE?
+		// Does this material require NEE? (will be false if MaterialPdf is always +inf)
 		bool bIsNeeValid = IsNeeValidMaterial(Payload);
 
-		// If we are using Light sampling and we are not on the last bounce ....
-		if (MISMode != 0 && Bounce < PathTracingData.MaxBounces && bIsNeeValid && SceneLightCount > 0)
+		// If we are using Light sampling and the material can use it ...
+		if (MISMode != 0 && bIsNeeValid && SceneLightCount > 0)
 		{
 			// Choose a light and sample it
-
 			float3 WorldPos = Payload.WorldPos;
 			float3 WorldNormal = Payload.WorldNormal;
 			uint PrimitiveLightingChannelMask = Payload.PrimitiveLightingChannelMask;
@@ -198,7 +199,6 @@ RAY_TRACING_ENTRY_RAYGEN(PathTracingMainRG)
 					LightPdf);
 				RadianceOverPdf /= LightPickPdf;
 				LightPdf *= LightPickPdf;
-
 				if (LightPdf > 0)
 				{
 					float SignedPositionBias = 1.0;
@@ -210,18 +210,21 @@ RAY_TRACING_ENTRY_RAYGEN(PathTracingMainRG)
 					}
 					ApplyPositionBias(LightRay, SignedPositionBias * Payload.WorldNormal, PathTracingData.MaxNormalBias);
 
+					bool LightIsVisible = true;
+					if (CastsShadow(LightId))
+					{
+						const uint RayFlags = RAY_FLAG_ACCEPT_FIRST_HIT_AND_END_SEARCH | RAY_FLAG_SKIP_CLOSEST_HIT_SHADER;
+						const uint InstanceInclusionMask = RAY_TRACING_MASK_SHADOW;
 
-					const uint RayFlags = RAY_FLAG_ACCEPT_FIRST_HIT_AND_END_SEARCH | RAY_FLAG_SKIP_CLOSEST_HIT_SHADER;
-					const uint InstanceInclusionMask = RAY_TRACING_MASK_ALL;
-
-					FPackedMaterialClosestHitPayload PackedPayload = (FPackedMaterialClosestHitPayload)0;
+						FPackedMaterialClosestHitPayload PackedPayload = (FPackedMaterialClosestHitPayload)0;
 
-					TraceVisibilityRayPacked(PackedPayload, TLAS, RayFlags, InstanceInclusionMask, LaunchIndex, LightRay);
+						TraceVisibilityRayPacked(PackedPayload, TLAS, RayFlags, InstanceInclusionMask, LaunchIndex, LightRay);
 
-					RayCounter++;
+						LightIsVisible = PackedPayload.IsMiss();
+					}
 
 					// #dxr_todo: Is it cheaper to fire the ray first? Or eval the material first?
-					if (PackedPayload.IsMiss())
+					if (LightIsVisible)
 					{
 						// Evaluate material
 						float3 MaterialWeight;
@@ -239,12 +242,6 @@ RAY_TRACING_ENTRY_RAYGEN(PathTracingMainRG)
 			}
 		}
 
-		if (Bounce == PathTracingData.MaxBounces)
-		{
-			// we've arrived at the last bounce, don't bother trying to sample a direction
-			break;
-		}
-
 		// Sample material
 		float3 Direction;
 		float3 MaterialWeight;
@@ -298,18 +295,63 @@ RAY_TRACING_ENTRY_RAYGEN(PathTracingMainRG)
 			PathThroughput = NextPathThroughput;
 		}
 
-		if (MISMode == 2)
-		{
-			// If we are doing MIS, the Material ray on the next loop through will need to know the Pdf
-			PrevMaterialPdf = MaterialSamplePdf;
-		}
-
-		PrevPrimitiveLightingChannelMask = Payload.PrimitiveLightingChannelMask;
-
-		// Update ray and loop back around to trace again
+		// Update ray according to material sample
 		Ray.Direction = Direction;
 		Ray.TMin = 0;
+		Ray.TMax = POSITIVE_INFINITY;
 		ApplyPositionBias(Ray, SignedPositionBias * Payload.WorldNormal, PathTracingData.MaxNormalBias);
+
+		// If we are using Material sampling for lights
+		if (MISMode != 1)
+		{
+			// Check which lights can be seen by the material ray and trace a dedicated shadow ray
+			// While it would be possible to just loop around and use the indirect ray to do this, it would prevent the application
+			// of shadow ray specific logic for transparent shadows or various per light tricks like shadow casting
+			float PreviousCdfValue = 0.0;
+			for (uint LightId = 0; LightId < SceneLightCount; ++LightId)
+			{
+				if ((Payload.PrimitiveLightingChannelMask & GetLightingChannelMask(LightId)) == 0)
+				{
+					// light does not affect the current ray
+					continue;
+				}
+
+				float HitT = Ray.TMax;
+				float4 LightResult = TraceLight(Ray, LightId, HitT);
+				float3 LightRadiance = LightResult.xyz;
+
+				float3 LightContrib = PathThroughput * LightRadiance;
+				if (MISMode == 2)
+				{
+					float LightPdf = LightResult.w;
+					float LightPickPdf =  LightPickingCdf[LightId] - PreviousCdfValue;
+					PreviousCdfValue = LightPickingCdf[LightId];
+
+					LightPdf *= LightPickPdf;
+					LightContrib *= MISWeightRobust(MaterialSamplePdf, LightPdf);
+				}
+
+				if (any(LightContrib > 0))
+				{
+					bool LightIsVisible = true;
+					if (CastsShadow(LightId))
+					{
+						const uint RayFlags = RAY_FLAG_ACCEPT_FIRST_HIT_AND_END_SEARCH | RAY_FLAG_SKIP_CLOSEST_HIT_SHADER;
+						const uint InstanceInclusionMask = RAY_TRACING_MASK_SHADOW;
+						FPackedMaterialClosestHitPayload PackedPayload = (FPackedMaterialClosestHitPayload)0;
+						RayDesc LightRay = Ray;
+						LightRay.TMax = HitT;
+						TraceVisibilityRayPacked(PackedPayload, TLAS, RayFlags, InstanceInclusionMask, LaunchIndex, LightRay);
+						LightIsVisible = PackedPayload.IsMiss();
+					}
+					if (LightIsVisible)
+					{
+						// the light made some contribution, and there was nothing along the shadow ray
+						AccumulateRadiance(Radiance, LightContrib);
+					}
+				}
+			}
+		}
 	}
 
 
diff --git a/Engine/Shaders/Shared/RayTracingDefinitions.h b/Engine/Shaders/Shared/RayTracingDefinitions.h
--- a/Engine/Shaders/Shared/RayTracingDefinitions.h
+++ b/Engine/Shaders/Shared/RayTracingDefinitions.h
@@ -46,3 +46,4 @@
 #define PATHTRACING_LIGHT_POINT								(2 << 6)
 #define PATHTRACING_LIGHT_SPOT								(3 << 6)
 #define PATHTRACING_LIGHT_RECT								(4 << 6)
+#define PATHTRACER_FLAG_CAST_SHADOW_MASK 					(1 << 9)
\ No newline at end of file
diff --git a/Engine/Source/Runtime/Renderer/Private/PathTracing.cpp b/Engine/Source/Runtime/Renderer/Private/PathTracing.cpp
--- a/Engine/Source/Runtime/Renderer/Private/PathTracing.cpp
+++ b/Engine/Source/Runtime/Renderer/Private/PathTracing.cpp
@@ -229,6 +229,7 @@ void SetLightParameters(FRDGBuilder& GraphBuilder, FPathTracingRG::FParameters*
 		DestLight.Flags = SkyLightData.bTransmission ? PATHTRACER_FLAG_TRANSMISSION_MASK : 0;
 		DestLight.Flags |= PATHTRACER_FLAG_LIGHTING_CHANNEL_MASK;
 		DestLight.Flags |= PATHTRACING_LIGHT_SKY;
+		DestLight.Flags |= Scene->SkyLight->bCastShadows ? PATHTRACER_FLAG_CAST_SHADOW_MASK : 0;
 		DestLight.IESTextureSlice = -1;
 		LightCount++;
 	}
@@ -247,6 +248,7 @@ void SetLightParameters(FRDGBuilder& GraphBuilder, FPathTracingRG::FParameters*
 
 		DestLight.Flags = Transmission ? PATHTRACER_FLAG_TRANSMISSION_MASK : 0;
 		DestLight.Flags |= LightingChannelMask & PATHTRACER_FLAG_LIGHTING_CHANNEL_MASK;
+		DestLight.Flags |= Light.LightSceneInfo->Proxy->CastsDynamicShadow() ? PATHTRACER_FLAG_CAST_SHADOW_MASK : 0;
 		DestLight.IESTextureSlice = -1;
 
 		if (UseLightProfiles)
@@ -552,7 +554,7 @@ void FDeferredShadingSceneRenderer::RenderPathTracing(
 	DisplayParameters->ProgressDisplayEnabled = CVarPathTracingProgressDisplay.GetValueOnRenderThread();
 	DisplayParameters->ViewUniformBuffer = View.ViewUniformBuffer;
 	DisplayParameters->RadianceTexture = GraphBuilder.CreateSRV(FRDGTextureSRVDesc::Create(RadianceTexture));
-	DisplayParameters->RenderTargets[0] = FRenderTargetBinding(SceneColorOutputTexture, ERenderTargetLoadAction::ENoAction);
+	DisplayParameters->RenderTargets[0] = FRenderTargetBinding(SceneColorOutputTexture, ERenderTargetLoadAction::ELoad);
 
 	FScreenPassTextureViewport Viewport(SceneColorOutputTexture, View.ViewRect);
 
