Added VertexFactoryDefaultInterface header with new VF interface functions (disabled with GPUCULL_TODO define), and included this header from appropriate VF source files (relevant for incoming GPU instance culling work).

diff --git a/Engine/Plugins/Enterprise/LidarPointCloud/Shaders/Private/LidarPointCloudVertexFactory.ush b/Engine/Plugins/Enterprise/LidarPointCloud/Shaders/Private/LidarPointCloudVertexFactory.ush
--- a/Engine/Plugins/Enterprise/LidarPointCloud/Shaders/Private/LidarPointCloudVertexFactory.ush
+++ b/Engine/Plugins/Enterprise/LidarPointCloud/Shaders/Private/LidarPointCloudVertexFactory.ush
@@ -490,4 +490,6 @@ uint VertexFactoryGetEyeIndex(uint InstanceId)
 uint VertexFactoryGetPrimitiveId(FVertexFactoryInterpolantsVSToPS Interpolants)
 {
 	return 0;
-}
\ No newline at end of file
+}
+
+#include "/Engine/Private/VertexFactoryDefaultInterface.ush"
diff --git a/Engine/Plugins/Experimental/VirtualHeightfieldMesh/Shaders/Private/VirtualHeightfieldMeshVertexFactory.ush b/Engine/Plugins/Experimental/VirtualHeightfieldMesh/Shaders/Private/VirtualHeightfieldMeshVertexFactory.ush
--- a/Engine/Plugins/Experimental/VirtualHeightfieldMesh/Shaders/Private/VirtualHeightfieldMeshVertexFactory.ush
+++ b/Engine/Plugins/Experimental/VirtualHeightfieldMesh/Shaders/Private/VirtualHeightfieldMeshVertexFactory.ush
@@ -308,3 +308,5 @@ FMaterialPixelParameters GetMaterialPixelParameters(FVertexFactoryInterpolantsVS
 
 	return Result;
 }
+
+#include "/Engine/Private/VertexFactoryDefaultInterface.ush"
diff --git a/Engine/Plugins/Experimental/Water/Shaders/Private/WaterMeshVertexFactory.ush b/Engine/Plugins/Experimental/Water/Shaders/Private/WaterMeshVertexFactory.ush
--- a/Engine/Plugins/Experimental/Water/Shaders/Private/WaterMeshVertexFactory.ush
+++ b/Engine/Plugins/Experimental/Water/Shaders/Private/WaterMeshVertexFactory.ush
@@ -432,4 +432,6 @@ float3 VertexFactoryGetWorldNormal(FVertexFactoryInput Input, FVertexFactoryInte
 	{
 		return float3x3(1,0,0,0,1,0,0,0,1);
 	}
-#endif // #if USING_TESSELLATION
\ No newline at end of file
+#endif // #if USING_TESSELLATION
+
+#include "/Engine/Private/VertexFactoryDefaultInterface.ush"
diff --git a/Engine/Plugins/FX/Niagara/Shaders/Private/NiagaraMeshVertexFactory.ush b/Engine/Plugins/FX/Niagara/Shaders/Private/NiagaraMeshVertexFactory.ush
--- a/Engine/Plugins/FX/Niagara/Shaders/Private/NiagaraMeshVertexFactory.ush
+++ b/Engine/Plugins/FX/Niagara/Shaders/Private/NiagaraMeshVertexFactory.ush
@@ -1126,3 +1126,5 @@ float4 VertexFactoryGetTranslatedPrimitiveVolumeBounds(FVertexFactoryInterpolant
 		return TangentToWorld;
 	}
 #endif // #if USING_TESSELLATION
+
+#include "/Engine/Private/VertexFactoryDefaultInterface.ush"
diff --git a/Engine/Plugins/FX/Niagara/Shaders/Private/NiagaraRibbonVertexFactory.ush b/Engine/Plugins/FX/Niagara/Shaders/Private/NiagaraRibbonVertexFactory.ush
--- a/Engine/Plugins/FX/Niagara/Shaders/Private/NiagaraRibbonVertexFactory.ush
+++ b/Engine/Plugins/FX/Niagara/Shaders/Private/NiagaraRibbonVertexFactory.ush
@@ -946,3 +946,5 @@ FVertexFactoryInterpolantsVSToPS VertexFactoryAssignInterpolants(FVertexFactoryI
 }
 
 #endif
+
+#include "/Engine/Private/VertexFactoryDefaultInterface.ush"
diff --git a/Engine/Plugins/FX/Niagara/Shaders/Private/NiagaraSpriteVertexFactory.ush b/Engine/Plugins/FX/Niagara/Shaders/Private/NiagaraSpriteVertexFactory.ush
--- a/Engine/Plugins/FX/Niagara/Shaders/Private/NiagaraSpriteVertexFactory.ush
+++ b/Engine/Plugins/FX/Niagara/Shaders/Private/NiagaraSpriteVertexFactory.ush
@@ -1199,3 +1199,5 @@ uint VertexFactoryGetPrimitiveId(FVertexFactoryInterpolantsVSToPS Interpolants)
 {
 	return 0;
 }
+
+#include "/Engine/Private/VertexFactoryDefaultInterface.ush"
diff --git a/Engine/Shaders/Private/GeometryCacheVertexFactory.ush b/Engine/Shaders/Private/GeometryCacheVertexFactory.ush
--- a/Engine/Shaders/Private/GeometryCacheVertexFactory.ush
+++ b/Engine/Shaders/Private/GeometryCacheVertexFactory.ush
@@ -465,3 +465,5 @@ uint VertexFactoryGetPrimitiveId(FVertexFactoryInterpolantsVSToPS Interpolants)
 {
 	return 0;
 }
+
+#include "VertexFactoryDefaultInterface.ush"
\ No newline at end of file
diff --git a/Engine/Shaders/Private/GpuSkinVertexFactory.ush b/Engine/Shaders/Private/GpuSkinVertexFactory.ush
--- a/Engine/Shaders/Private/GpuSkinVertexFactory.ush
+++ b/Engine/Shaders/Private/GpuSkinVertexFactory.ush
@@ -331,7 +331,7 @@ float3 ClothingPosition(FClothVertex ClothVertexInfluences[NUM_INFLUENCES_PER_VE
 	int NumInfluences = 0;
 	float SumWeights = 0.0f;
 	float3 AveragedSimulatedPosition = float3(0, 0, 0);
-	
+
     for (int i = 0; i < NUM_INFLUENCES_PER_VERTEX; ++i )
     {
         const FClothVertex Influence = ClothVertexInfluences[i];
@@ -946,3 +946,5 @@ uint VertexFactoryGetPrimitiveId(FVertexFactoryInterpolantsVSToPS Interpolants)
 {
 	return 0;
 }
+
+#include "VertexFactoryDefaultInterface.ush"
diff --git a/Engine/Shaders/Private/HairStrands/HairCardsVertexFactory.ush b/Engine/Shaders/Private/HairStrands/HairCardsVertexFactory.ush
--- a/Engine/Shaders/Private/HairStrands/HairCardsVertexFactory.ush
+++ b/Engine/Shaders/Private/HairStrands/HairCardsVertexFactory.ush
@@ -18,7 +18,7 @@
 
 #define USE_HAIR_COMPLEX_TRANSMITTANCE 1
 
-#if MANUAL_VERTEX_FETCH 
+#if MANUAL_VERTEX_FETCH
 	Buffer<float4> VertexFetch_InstanceOriginBuffer;
 	Buffer<float4> VertexFetch_InstanceTransformBuffer;
 #endif
@@ -59,7 +59,7 @@ struct FVertexFactoryInput
 #endif	
 };
 
-#if RAYHITGROUPSHADER 
+#if RAYHITGROUPSHADER
 FVertexFactoryInput LoadVertexFactoryInputForHGS(uint TriangleIndex, int VertexIndex)
 {
 	// #hair_todo
@@ -504,4 +504,6 @@ float3x3 VertexFactoryGetTangentToLocalDS(FVertexFactoryInterpolantsVSToDS Inter
 
 	return TangentToWorld;
 }
-#endif // #if USING_TESSELLATION
\ No newline at end of file
+#endif // #if USING_TESSELLATION
+
+#include "/Engine/Private/VertexFactoryDefaultInterface.ush"
diff --git a/Engine/Shaders/Private/HairStrands/HairStrandsVertexFactory.ush b/Engine/Shaders/Private/HairStrands/HairStrandsVertexFactory.ush
--- a/Engine/Shaders/Private/HairStrands/HairStrandsVertexFactory.ush
+++ b/Engine/Shaders/Private/HairStrands/HairStrandsVertexFactory.ush
@@ -540,4 +540,6 @@ float3 VertexFactoryGetWorldNormal(FPositionAndNormalOnlyVertexFactoryInput Inpu
 float3 VertexFactoryGetWorldNormal(FVertexFactoryInput Input, FVertexFactoryIntermediates Intermediates)
 {
 	return float3(0,0,1);
-}
\ No newline at end of file
+}
+
+#include "/Engine/Private/VertexFactoryDefaultInterface.ush"
diff --git a/Engine/Shaders/Private/LandscapeVertexFactory.ush b/Engine/Shaders/Private/LandscapeVertexFactory.ush
--- a/Engine/Shaders/Private/LandscapeVertexFactory.ush
+++ b/Engine/Shaders/Private/LandscapeVertexFactory.ush
@@ -200,7 +200,7 @@ FVertexFactoryInput LoadVertexFactoryInputForHGS(uint TriangleIndex, int VertexI
 #if COMPUTESHADER
 FVertexFactoryInput LoadVertexFactoryInputForDynamicUpdate(uint TriangleIndex, int VertexIndex, uint PrimitiveId)
 {
-	FVertexFactoryInput Input;
+	FVertexFactoryInput Input = (FVertexFactoryInput)0;
 
 	float4 LodValues = GetLodValues(GetComponentLinearIndex());
 	uint SubsectionSizeQuadsCurrentLOD = LodValues.z;
@@ -629,7 +629,6 @@ void GetShadowMapCoordinate(FVertexFactoryInterpolantsVSToPS Interpolants, out f
 
 FVertexFactoryIntermediates GetVertexFactoryIntermediates(FVertexFactoryInput Input)
 {
-
 	FVertexFactoryIntermediates Intermediates;
 	Intermediates.ComponentIndex = GetComponentLinearIndex();
 	Intermediates.LodValues = GetLodValues(Intermediates.ComponentIndex);
@@ -913,3 +912,5 @@ uint VertexFactoryGetPrimitiveId(FVertexFactoryInterpolantsVSToPS Interpolants)
 {
 	return 0;
 }
+
+#include "VertexFactoryDefaultInterface.ush"
diff --git a/Engine/Shaders/Private/LocalVertexFactory.ush b/Engine/Shaders/Private/LocalVertexFactory.ush
--- a/Engine/Shaders/Private/LocalVertexFactory.ush
+++ b/Engine/Shaders/Private/LocalVertexFactory.ush
@@ -4,6 +4,7 @@
 	LocalVertexFactory.usf: Local vertex factory shader code.
 =============================================================================*/
 
+#include "SceneData.ush"
 #include "VertexFactoryCommon.ush"
 #include "LocalVertexFactoryCommon.ush"
 #include "LightmapData.ush"
@@ -1180,3 +1181,5 @@ uint VertexFactoryGetPrimitiveId(FVertexFactoryInterpolantsVSToPS Interpolants)
 {
 	return GetPrimitiveId(Interpolants);
 }
+
+#include "VertexFactoryDefaultInterface.ush"
diff --git a/Engine/Shaders/Private/MeshParticleVertexFactory.ush b/Engine/Shaders/Private/MeshParticleVertexFactory.ush
--- a/Engine/Shaders/Private/MeshParticleVertexFactory.ush
+++ b/Engine/Shaders/Private/MeshParticleVertexFactory.ush
@@ -717,3 +717,5 @@ uint VertexFactoryGetPrimitiveId(FVertexFactoryInterpolantsVSToPS Interpolants)
 {
 	return 0;
 }
+
+#include "VertexFactoryDefaultInterface.ush"
\ No newline at end of file
diff --git a/Engine/Shaders/Private/ParticleBeamTrailVertexFactory.ush b/Engine/Shaders/Private/ParticleBeamTrailVertexFactory.ush
--- a/Engine/Shaders/Private/ParticleBeamTrailVertexFactory.ush
+++ b/Engine/Shaders/Private/ParticleBeamTrailVertexFactory.ush
@@ -241,3 +241,5 @@ uint VertexFactoryGetPrimitiveId(FVertexFactoryInterpolantsVSToPS Interpolants)
 {
 	return 0;
 }
+
+#include "VertexFactoryDefaultInterface.ush"
\ No newline at end of file
diff --git a/Engine/Shaders/Private/ParticleGPUSpriteVertexFactory.ush b/Engine/Shaders/Private/ParticleGPUSpriteVertexFactory.ush
--- a/Engine/Shaders/Private/ParticleGPUSpriteVertexFactory.ush
+++ b/Engine/Shaders/Private/ParticleGPUSpriteVertexFactory.ush
@@ -695,3 +695,5 @@ uint VertexFactoryGetPrimitiveId(FVertexFactoryInterpolantsVSToPS Interpolants)
 {
 	return 0;
 }
+
+#include "VertexFactoryDefaultInterface.ush"
\ No newline at end of file
diff --git a/Engine/Shaders/Private/ParticleSpriteVertexFactory.ush b/Engine/Shaders/Private/ParticleSpriteVertexFactory.ush
--- a/Engine/Shaders/Private/ParticleSpriteVertexFactory.ush
+++ b/Engine/Shaders/Private/ParticleSpriteVertexFactory.ush
@@ -681,3 +681,5 @@ uint VertexFactoryGetPrimitiveId(FVertexFactoryInterpolantsVSToPS Interpolants)
 {
 	return 0;
 }
+
+#include "VertexFactoryDefaultInterface.ush"
\ No newline at end of file
diff --git a/Engine/Shaders/Private/VectorFieldVisualizationVertexFactory.ush b/Engine/Shaders/Private/VectorFieldVisualizationVertexFactory.ush
--- a/Engine/Shaders/Private/VectorFieldVisualizationVertexFactory.ush
+++ b/Engine/Shaders/Private/VectorFieldVisualizationVertexFactory.ush
@@ -199,3 +199,5 @@ uint VertexFactoryGetPrimitiveId(FVertexFactoryInterpolantsVSToPS Interpolants)
 {
 	return 0;
 }
+
+#include "VertexFactoryDefaultInterface.ush"
\ No newline at end of file
diff --git a/Engine/Shaders/Private/VertexFactoryDefaultInterface.ush b/Engine/Shaders/Private/VertexFactoryDefaultInterface.ush
new file mode 100644
--- /dev/null
+++ b/Engine/Shaders/Private/VertexFactoryDefaultInterface.ush
@@ -0,0 +1,92 @@
+// Copyright Epic Games, Inc. All Rights Reserved.
+
+// Note: include at the end of vertex factories to provide default/dummy implementation of factory functions that are not supported OR that follow the standard pattern
+// E.g., if the VF implements the three flavours of VertexFactoryGetViewIndex, it should define VF_IMPLEMENTED_GET_VIEW_INDEX
+#ifndef VF_IMPLEMENTED_GET_VIEW_INDEX
+#define VF_IMPLEMENTED_GET_VIEW_INDEX
+
+#if defined(GPUCULL_TODO)
+
+uint VertexFactoryGetViewIndex(FVertexFactoryIntermediates Intermediates)
+{
+#if VF_USE_PRIMITIVE_SCENE_DATA
+	return Intermediates.SceneData.ViewIndex;
+#else
+	return 0U;
+#endif
+}
+
+uint VertexFactoryGetInstanceIdLoadIndex(FVertexFactoryIntermediates Intermediates)
+{
+#if VF_USE_PRIMITIVE_SCENE_DATA
+	return Intermediates.SceneData.InstanceIdLoadIndex;
+#else
+	return 0U;
+#endif
+}
+
+float4x4 VertexFactoryGetLocalToWorld(FVertexFactoryIntermediates Intermediates)
+{
+#if VF_USE_PRIMITIVE_SCENE_DATA
+	return Intermediates.SceneData.InstanceData.LocalToWorld;
+#else
+	return Primitive.LocalToWorld;
+#endif
+}
+
+float4x4 VertexFactoryGetWorldToLocal(FVertexFactoryIntermediates Intermediates)
+{
+#if VF_USE_PRIMITIVE_SCENE_DATA
+	return Intermediates.SceneData.InstanceData.WorldToLocal;
+#else
+	return Primitive.WorldToLocal;
+#endif
+}
+
+#if POSITION_ONLY
+
+uint VertexFactoryGetViewIndex(FPositionOnlyVertexFactoryInput Input)
+{
+#if VF_USE_PRIMITIVE_SCENE_DATA
+	FSceneDataIntermediates SceneData = GetSceneDataIntermediates(Input.InstanceIdOffset, Input.DrawInstanceId);
+	return SceneData.ViewIndex;
+#else
+	return 0U;
+#endif
+}
+
+uint VertexFactoryGetViewIndex(FPositionAndNormalOnlyVertexFactoryInput Input)
+{
+#if VF_USE_PRIMITIVE_SCENE_DATA
+	FSceneDataIntermediates SceneData = GetSceneDataIntermediates(Input.InstanceIdOffset, Input.DrawInstanceId);
+	return SceneData.ViewIndex;
+#else
+	return 0U;
+#endif
+}
+
+uint VertexFactoryGetInstanceIdLoadIndex(FPositionOnlyVertexFactoryInput Input)
+{
+#if VF_USE_PRIMITIVE_SCENE_DATA
+	FSceneDataIntermediates SceneData = GetSceneDataIntermediates(Input.InstanceIdOffset, Input.DrawInstanceId);
+	return SceneData.InstanceIdLoadIndex;
+#else
+	return 0U;
+#endif
+}
+
+uint VertexFactoryGetInstanceIdLoadIndex(FPositionAndNormalOnlyVertexFactoryInput Input)
+{
+#if VF_USE_PRIMITIVE_SCENE_DATA
+	FSceneDataIntermediates SceneData = GetSceneDataIntermediates(Input.InstanceIdOffset, Input.DrawInstanceId);
+	return SceneData.InstanceIdLoadIndex;
+#else
+	return 0U;
+#endif
+}
+
+#endif // POSITION_ONLY
+
+#endif // GPUCULL_TODO
+
+#endif // VF_IMPLEMENTED_GET_VIEW_INDEX
