Added shader files for (disabled) instance culling

diff --git a/Engine/Shaders/Private/InstanceCulling/BuildInstanceDrawCommands.usf b/Engine/Shaders/Private/InstanceCulling/BuildInstanceDrawCommands.usf
new file mode 100644
--- /dev/null
+++ b/Engine/Shaders/Private/InstanceCulling/BuildInstanceDrawCommands.usf
@@ -0,0 +1,515 @@
+// Copyright Epic Games, Inc. All Rights Reserved.
+
+/*=============================================================================
+	BuildInstanceDrawCommands.usf: 
+=============================================================================*/
+
+#include "../Common.ush"
+#include "../SceneData.ush"
+
+struct FPrimCullingCommand
+{
+	uint BaseVertexIndex;
+	uint FirstIndex;
+	uint NumVerticesOrIndices;
+	uint FirstPrimitiveIdOffset;
+	uint FirstInstanceRunOffset;
+};
+
+struct FInstanceRun
+{
+	uint Start;
+	uint EndInclusive;
+	int PrimitiveId;
+};
+
+#if OUTPUT_COMMAND_IDS
+RWBuffer<uint> InstanceIdsBufferOut;
+RWBuffer<uint> DrawCommandIdsBufferOut;
+#else
+RWBuffer<uint> InstanceIdsBufferLegacyOut;
+#endif
+RWBuffer<uint> DrawIndirectArgsBufferOut;
+RWBuffer<uint> InstanceIdOffsetBufferOut;
+RWStructuredBuffer<uint> OutputOffsetBufferOut;
+
+StructuredBuffer<FInstanceRun> InstanceRuns;
+StructuredBuffer<uint> ViewIds;
+
+groupshared uint PrimitiveInstanceCountGroupShared;
+groupshared uint RunInstanceCountGroupShared;
+groupshared uint InstanceIdOffsetGroupShared;
+groupshared uint FinalInstanceCountGroupShared;
+
+StructuredBuffer<FPrimCullingCommand> PrimitiveCullingCommands;
+StructuredBuffer<uint> PrimitiveIds;
+
+// These flags are not supplied if the culling is disabled either for testing or because there is no data
+#if CULL_INSTANCES
+StructuredBuffer<uint> VisibleInstanceFlags;
+uint NumInstanceFlagWords;
+#endif
+
+int NumPrimitiveIds;
+int NumInstanceRuns;
+int NumCommands;
+uint NumViewIds;
+
+bool TestInstanceVisibilityBit(uint InstanceId, uint ViewId)
+{
+#if CULL_INSTANCES
+	uint ViewWordOffset = ViewId * NumInstanceFlagWords;
+	uint WordOffset = ViewWordOffset + InstanceId / 32U;
+	return (VisibleInstanceFlags[WordOffset] & (1U << (InstanceId % 32U))) != 0U;
+#else
+	return true;
+#endif
+}
+
+
+void WriteInstanceId(uint Offset, uint InstanceId, uint ViewIdIndex, uint DrawCommandId)
+{
+	uint PackedId = InstanceId | (ViewIdIndex << 28U);
+#if OUTPUT_COMMAND_IDS
+	InstanceIdsBufferOut[Offset] = PackedId;
+	DrawCommandIdsBufferOut[Offset] = DrawCommandId;
+#else
+	InstanceIdsBufferLegacyOut[Offset] = PackedId;
+#endif
+}
+
+#ifndef ENABLE_DETERMINISTIC_INSTANCE_CULLING
+	#define ENABLE_DETERMINISTIC_INSTANCE_CULLING 0
+#endif
+
+#if ENABLE_DETERMINISTIC_INSTANCE_CULLING
+
+RWStructuredBuffer<uint2> InstanceCountsOut;
+
+[numthreads(NUM_THREADS_PER_GROUP, 1, 1)]
+void ComputeInstanceIdOutputSize(uint GroupThreadIndex : SV_GroupIndex, uint GroupId : SV_GroupID)
+{
+	int CommandSlotIndex = int(GroupId);
+	FPrimCullingCommand Cmd = PrimitiveCullingCommands[CommandSlotIndex];
+
+	// 1. Figure out the total number of instances
+	PrimitiveInstanceCountGroupShared = 0;
+	RunInstanceCountGroupShared = 0;
+	FinalInstanceCountGroupShared = 0;
+	GroupMemoryBarrierWithGroupSync();
+
+	// Get start of next CMD or go to end.
+	int LastOffset = (CommandSlotIndex + 1 < NumCommands) ? int(PrimitiveCullingCommands[CommandSlotIndex + 1].FirstPrimitiveIdOffset) : NumPrimitiveIds;
+	uint NumPrimitiveIds = uint(LastOffset) - Cmd.FirstPrimitiveIdOffset;
+
+	int LastInstanceRunOffset = (CommandSlotIndex + 1 < NumCommands) ? int(PrimitiveCullingCommands[CommandSlotIndex + 1].FirstInstanceRunOffset) : NumInstanceRuns;
+	uint NumInstanceRuns = uint(LastInstanceRunOffset) - Cmd.FirstInstanceRunOffset;
+
+	for (uint Index = GroupThreadIndex; Index < NumPrimitiveIds + NumInstanceRuns; Index += NUM_THREADS_PER_GROUP)
+	{
+		uint Num = 0;
+		if (Index < NumPrimitiveIds)
+		{
+			uint PrimitiveId = PrimitiveIds[Cmd.FirstPrimitiveIdOffset + Index];
+			if (GetPrimitiveData(PrimitiveId).InstanceDataOffset >= 0)
+			{
+				Num = GetPrimitiveData(PrimitiveId).NumInstanceDataEntries;
+			}
+			else
+			{
+				// Fake entry for invalid ones -> shaders that will not use the instancing data anyway
+				Num = 1;
+			}
+			InterlockedAdd(PrimitiveInstanceCountGroupShared, Num);
+		}
+		else
+		{
+			FInstanceRun Run = InstanceRuns[Cmd.FirstInstanceRunOffset + Index - NumPrimitiveIds];
+			Num = (Run.EndInclusive + 1) - Run.Start;
+
+			InterlockedAdd(RunInstanceCountGroupShared, Num);
+		}
+
+	}
+
+	GroupMemoryBarrierWithGroupSync();
+	if (GroupThreadIndex == 0)
+	{
+		InstanceCountsOut[CommandSlotIndex] = uint2(PrimitiveInstanceCountGroupShared, RunInstanceCountGroupShared);
+	}
+}
+
+
+groupshared uint GroupSharedPrefixBuffer[NUM_THREADS_PER_GROUP];
+
+
+uint CalcPrefixSum(uint Value, inout uint InclusiveSum, uint ThreadIdx)
+{
+	// TODO: make efficient using: wave shuffle intrinsics when applicable
+	InclusiveSum = Value;
+	GroupSharedPrefixBuffer[ThreadIdx] = InclusiveSum;
+	UNROLL
+	for (uint Offset = 1; Offset < NUM_THREADS_PER_GROUP; Offset *= 2)
+	{
+		// Must sync so all threads have written data before loading it
+		GroupMemoryBarrierWithGroupSync();
+		if (ThreadIdx >= Offset)
+		{
+			InclusiveSum += GroupSharedPrefixBuffer[ThreadIdx - Offset];
+		}
+		// Must sync so all threads have loaded data before overwriting 
+		GroupMemoryBarrierWithGroupSync();
+		if (ThreadIdx >= Offset)
+		{
+			GroupSharedPrefixBuffer[ThreadIdx] = InclusiveSum;
+		}
+	}
+
+	uint ExclusiveSum = InclusiveSum - Value;
+	GroupMemoryBarrierWithGroupSync();
+	return ExclusiveSum;
+}
+
+uint CalcBinaryPrefixSum(bool bItemValid, inout uint InclusiveSum, uint ThreadIdx)
+{
+	// TODO: make efficient using: a ballot-based prefix for each warp/wave, or the built-in one in SM6, and then reduce the 
+	// wave offsets by brute force. 
+	return CalcPrefixSum(bItemValid ? 1 : 0, InclusiveSum, ThreadIdx);
+}
+
+StructuredBuffer<uint2> InstanceCounts;
+
+[numthreads(NUM_THREADS_PER_GROUP, 1, 1)]
+void CalcOutputOffsets(uint GroupThreadIndex : SV_GroupIndex, uint GroupId : SV_GroupID)
+{
+	// Run single thread group and compute prefix sum.
+	if (GroupThreadIndex == 0)
+	{
+		// We assume there is only one of these running at any given time (beware!)
+		// Load the start offset in the global output buffer and broadcast
+		FinalInstanceCountGroupShared = OutputOffsetBufferOut[0];
+	}
+	GroupMemoryBarrierWithGroupSync();
+
+	for (uint BaseIndex = 0; BaseIndex < NumCommands; BaseIndex += NUM_THREADS_PER_GROUP)
+	{
+		uint CommandSlotIndex = BaseIndex + GroupThreadIndex;
+		bool bValidIndex = CommandSlotIndex < NumCommands;
+
+		uint CommandInstanceCount = 0;
+		if (bValidIndex)
+		{
+			uint PrimitiveInstanceCount = InstanceCounts[CommandSlotIndex].x;
+			uint RunInstanceCount = InstanceCounts[CommandSlotIndex].y;
+			CommandInstanceCount = NumViewIds * (PrimitiveInstanceCount + RunInstanceCount);
+		}
+
+		uint InclusiveSum = 0U;
+		uint OutputOffset = CalcPrefixSum(CommandInstanceCount, InclusiveSum, GroupThreadIndex);
+		if (bValidIndex)
+		{
+			InstanceIdOffsetBufferOut[CommandSlotIndex] = OutputOffset + FinalInstanceCountGroupShared;
+		}
+		// broadcast the total
+		GroupMemoryBarrierWithGroupSync();
+		if (GroupThreadIndex == NUM_THREADS_PER_GROUP - 1)
+		{
+			FinalInstanceCountGroupShared += InclusiveSum;
+		}
+		GroupMemoryBarrierWithGroupSync();
+	}
+	GroupMemoryBarrierWithGroupSync();
+	if (GroupThreadIndex == 0)
+	{
+		// Write back the total count
+		OutputOffsetBufferOut[0] = FinalInstanceCountGroupShared;
+	}
+}
+
+StructuredBuffer<uint> InstanceIdOffsetBuffer;
+
+[numthreads(NUM_THREADS_PER_GROUP, 1, 1)]
+void OutputInstanceIdsAtOffset(uint GroupThreadIndex : SV_GroupIndex, uint GroupId : SV_GroupID)
+{
+	FinalInstanceCountGroupShared = 0;
+	GroupMemoryBarrierWithGroupSync();
+
+	int CommandSlotIndex = int(GroupId);
+	FPrimCullingCommand Cmd = PrimitiveCullingCommands[CommandSlotIndex];
+
+	uint PrimitiveInstanceCount = InstanceCounts[CommandSlotIndex].x; 
+	uint RunInstanceCount = InstanceCounts[CommandSlotIndex].y;
+	uint InstanceIdOutputOffset = InstanceIdOffsetBuffer[CommandSlotIndex];
+
+	// Note: NumViewIds - replicate for each view in instanced stereo or other multi-view
+	uint CommandInstanceCount = NumViewIds * (PrimitiveInstanceCount + RunInstanceCount);
+
+	// Bail if no instances to output
+	if (CommandInstanceCount > 0)
+	{
+		// 3. Broadcast to all threads
+		uint PrimitiveId = PrimitiveIds[Cmd.FirstPrimitiveIdOffset];
+		if (GroupThreadIndex == 0)
+		{
+			WriteInstanceId(InstanceIdOutputOffset, GetPrimitiveData(PrimitiveId).InstanceDataOffset, 0U, CommandSlotIndex);
+		}
+		GroupMemoryBarrierWithGroupSync();
+
+		int CurrentPrimitiveIdIndex = -1;
+
+		uint CurrentRangeStart = 0;
+		uint CurrentRangeEnd = 0;
+		int InstanceDataOffset = 0;
+
+		// 4. loop over output primitive instance number
+		for (uint BaseOutputIndex1 = 0; BaseOutputIndex1 < PrimitiveInstanceCount; BaseOutputIndex1 += NUM_THREADS_PER_GROUP)
+		{
+			uint OutputIndex = BaseOutputIndex1 + GroupThreadIndex;
+
+			bool bValidItemIndex = OutputIndex < PrimitiveInstanceCount;
+
+			if (bValidItemIndex)
+			{
+				// Naive search for range containing the output index (reasonably efficient for large runs, bad for many small runs)
+				while (OutputIndex >= CurrentRangeEnd)
+				{
+					CurrentPrimitiveIdIndex += 1;
+					uint PrimitiveId = PrimitiveIds[Cmd.FirstPrimitiveIdOffset + CurrentPrimitiveIdIndex];
+
+					CurrentRangeStart = CurrentRangeEnd;
+					InstanceDataOffset = GetPrimitiveData(PrimitiveId).InstanceDataOffset;
+					CurrentRangeEnd += max(1, GetPrimitiveData(PrimitiveId).NumInstanceDataEntries);
+				}
+			}
+			uint OutId = (bValidItemIndex && InstanceDataOffset >= 0) ? uint(InstanceDataOffset) + OutputIndex - CurrentRangeStart : 0U;
+			for (uint ViewIdIndex = 0; ViewIdIndex < NumViewIds; ++ViewIdIndex)
+			{
+				bool bWriteItem = bValidItemIndex && (InstanceDataOffset < 0 || TestInstanceVisibilityBit(OutId, ViewIds[ViewIdIndex]));
+				uint InclusiveSum = 0U;
+				uint OutputOffset = CalcBinaryPrefixSum(bWriteItem, InclusiveSum, GroupThreadIndex);
+				if (bWriteItem)
+				{
+					WriteInstanceId(InstanceIdOutputOffset + FinalInstanceCountGroupShared + OutputOffset, OutId, ViewIdIndex, CommandSlotIndex);
+				}
+				// broadcast the total
+				GroupMemoryBarrierWithGroupSync();
+				if (GroupThreadIndex == NUM_THREADS_PER_GROUP - 1)
+				{
+					FinalInstanceCountGroupShared += InclusiveSum;
+				}
+				GroupMemoryBarrierWithGroupSync();
+
+			}
+		}
+		// repeat process for the runs
+		CurrentRangeStart = 0;
+		CurrentRangeEnd = 0;
+		InstanceDataOffset = 0;
+		FInstanceRun Run = (FInstanceRun)0;
+		int CurrentRunIndex = -1;
+		for (uint BaseOutputIndex2 = 0; BaseOutputIndex2 < RunInstanceCount; BaseOutputIndex2 += NUM_THREADS_PER_GROUP)
+		{
+			uint OutputIndex = BaseOutputIndex2 + GroupThreadIndex;
+			bool bValidItemIndex = OutputIndex < RunInstanceCount;
+
+			if (bValidItemIndex)
+			{
+				// Naive search for range containing the output index (reasonably efficient for large runs, bad for many small runs)
+				while (OutputIndex >= CurrentRangeEnd)
+				{
+					CurrentRunIndex += 1;
+					Run = InstanceRuns[Cmd.FirstInstanceRunOffset + CurrentRunIndex];
+
+					CurrentRangeStart = CurrentRangeEnd;
+					CurrentRangeEnd += (Run.EndInclusive + 1) - Run.Start;
+					InstanceDataOffset = GetPrimitiveData(Run.PrimitiveId).InstanceDataOffset;
+				}
+			}
+			uint OutId = bValidItemIndex ? (uint(InstanceDataOffset) + Run.Start + OutputIndex - CurrentRangeStart) : 0;
+			for (uint ViewIdIndex = 0; ViewIdIndex < NumViewIds; ++ViewIdIndex)
+			{
+				bool bWriteItem = bValidItemIndex && TestInstanceVisibilityBit(OutId, ViewIds[ViewIdIndex]);
+
+				uint InclusiveSum = 0U;
+				uint OutputOffset = CalcBinaryPrefixSum(bWriteItem, InclusiveSum, GroupThreadIndex);
+				if (bWriteItem)
+				{
+					WriteInstanceId(InstanceIdOutputOffset + FinalInstanceCountGroupShared + OutputOffset, OutId, ViewIdIndex, CommandSlotIndex);
+				}
+				// broadcast the total
+				GroupMemoryBarrierWithGroupSync();
+				if (GroupThreadIndex == NUM_THREADS_PER_GROUP - 1)
+				{
+					FinalInstanceCountGroupShared += InclusiveSum;
+				}
+				GroupMemoryBarrierWithGroupSync();
+			}
+		}
+	}
+	GroupMemoryBarrierWithGroupSync();
+
+	// Generate draw command
+	if (GroupThreadIndex == 0)
+	{
+		DrawIndirectArgsBufferOut[CommandSlotIndex * INDIRECT_ARGS_NUM_WORDS + 0] = Cmd.NumVerticesOrIndices;
+		DrawIndirectArgsBufferOut[CommandSlotIndex * INDIRECT_ARGS_NUM_WORDS + 1] = FinalInstanceCountGroupShared;
+		DrawIndirectArgsBufferOut[CommandSlotIndex * INDIRECT_ARGS_NUM_WORDS + 2] = Cmd.FirstIndex;
+		DrawIndirectArgsBufferOut[CommandSlotIndex * INDIRECT_ARGS_NUM_WORDS + 3] = Cmd.BaseVertexIndex;
+		DrawIndirectArgsBufferOut[CommandSlotIndex * INDIRECT_ARGS_NUM_WORDS + 4] = 0U;
+
+	}
+	GroupMemoryBarrierWithGroupSync();
+}
+
+#endif // ENABLE_DETERMINISTIC_INSTANCE_CULLING
+
+/**
+ */
+[numthreads(NUM_THREADS_PER_GROUP, 1, 1)]
+void BuildInstanceIdBufferAndCommandsFromPrimitiveIdsCs(uint GroupThreadIndex : SV_GroupIndex, uint GroupId : SV_GroupID)
+{
+	int CommandSlotIndex = int(GroupId);
+	FPrimCullingCommand Cmd = PrimitiveCullingCommands[CommandSlotIndex];
+
+	// 1. Figure out the total number of instances
+	PrimitiveInstanceCountGroupShared = 0;
+	RunInstanceCountGroupShared = 0;
+	InstanceIdOffsetGroupShared = 0;
+	FinalInstanceCountGroupShared = 0;
+	GroupMemoryBarrierWithGroupSync();
+
+	// Get start of next CMD or go to end.
+	int LastOffset = (CommandSlotIndex + 1 < NumCommands) ? int(PrimitiveCullingCommands[CommandSlotIndex + 1].FirstPrimitiveIdOffset) : NumPrimitiveIds;
+	uint NumPrimitiveIds = uint(LastOffset) - Cmd.FirstPrimitiveIdOffset;
+
+	int LastInstanceRunOffset = (CommandSlotIndex + 1 < NumCommands) ? int(PrimitiveCullingCommands[CommandSlotIndex + 1].FirstInstanceRunOffset) : NumInstanceRuns;
+	uint NumInstanceRuns = uint(LastInstanceRunOffset) - Cmd.FirstInstanceRunOffset;
+
+	for (uint Index = GroupThreadIndex; Index < NumPrimitiveIds + NumInstanceRuns; Index += NUM_THREADS_PER_GROUP)
+	{
+		uint Num = 0;
+		if (Index < NumPrimitiveIds)
+		{
+			uint PrimitiveId = PrimitiveIds[Cmd.FirstPrimitiveIdOffset + Index];
+			if (GetPrimitiveData(PrimitiveId).InstanceDataOffset >= 0)
+			{
+				Num = GetPrimitiveData(PrimitiveId).NumInstanceDataEntries;
+			}
+			else
+			{
+				// Fake entry for invalid ones -> shaders that will not use the instancing data anyway
+				Num = 1;
+			}
+			InterlockedAdd(PrimitiveInstanceCountGroupShared, Num);
+		}
+		else
+		{
+			FInstanceRun Run = InstanceRuns[Cmd.FirstInstanceRunOffset + Index - NumPrimitiveIds];
+			Num = (Run.EndInclusive + 1) - Run.Start;
+
+			InterlockedAdd(RunInstanceCountGroupShared, Num);
+		}
+
+	}
+
+	GroupMemoryBarrierWithGroupSync();
+	uint PrimitiveInstanceCount = PrimitiveInstanceCountGroupShared;
+	uint RunInstanceCount = RunInstanceCountGroupShared;
+	// Note: NumViewIds - replicate for each view in instanced stereo or other multi-view
+	uint CommandInstanceCount = NumViewIds * (PrimitiveInstanceCount + RunInstanceCount);
+
+	// Bail if no instances to output
+	if (CommandInstanceCount > 0)
+	{
+		// 2. Allocate output offset (over-allocates for max number, culling may remove any number)
+		if (GroupThreadIndex == 0)
+		{
+			uint IdOutputStart = 0;
+			InterlockedAdd(OutputOffsetBufferOut[0], CommandInstanceCount, IdOutputStart);
+			InstanceIdOffsetGroupShared = IdOutputStart;
+			InstanceIdOffsetBufferOut[CommandSlotIndex] = IdOutputStart;
+		}
+		GroupMemoryBarrierWithGroupSync();
+
+		// 3. Broadcast to all threads
+		uint InstanceIdOutputOffset = InstanceIdOffsetGroupShared;
+		uint PrimitiveId = PrimitiveIds[Cmd.FirstPrimitiveIdOffset];
+		WriteInstanceId(InstanceIdOutputOffset, GetPrimitiveData(PrimitiveId).InstanceDataOffset, 0U, CommandSlotIndex);
+		GroupMemoryBarrierWithGroupSync();
+
+		int CurrentPrimitiveIdIndex = -1;
+
+		uint CurrentRangeStart = 0;
+		uint CurrentRangeEnd = 0;
+		int InstanceDataOffset = 0;
+
+		// 4. loop over output primitive instance number
+		for (uint OutputIndex1 = GroupThreadIndex; OutputIndex1 < PrimitiveInstanceCount; OutputIndex1 += NUM_THREADS_PER_GROUP)
+		{
+			// Naive search for range containing the output index (reasonably efficient for large runs, bad for many small runs)
+			while (OutputIndex1 >= CurrentRangeEnd)
+			{
+				CurrentPrimitiveIdIndex += 1;
+				uint PrimitiveId = PrimitiveIds[Cmd.FirstPrimitiveIdOffset + CurrentPrimitiveIdIndex];
+
+				CurrentRangeStart = CurrentRangeEnd;
+				InstanceDataOffset = GetPrimitiveData(PrimitiveId).InstanceDataOffset;
+				CurrentRangeEnd += max(1, GetPrimitiveData(PrimitiveId).NumInstanceDataEntries);
+			}
+
+			uint OutId = InstanceDataOffset >= 0 ? uint(InstanceDataOffset) + OutputIndex1 - CurrentRangeStart : 0U;
+			for (uint ViewIdIndex = 0; ViewIdIndex < NumViewIds; ++ViewIdIndex)
+			{
+				if (InstanceDataOffset < 0 || TestInstanceVisibilityBit(OutId, ViewIds[ViewIdIndex]))
+				{
+					uint OutputOffset = 0;
+					InterlockedAdd(FinalInstanceCountGroupShared, 1, OutputOffset);
+					WriteInstanceId(InstanceIdOutputOffset + OutputOffset, OutId, ViewIdIndex, CommandSlotIndex);
+				}
+			}
+		}
+		// repeat process for the runs
+		CurrentRangeStart = 0;
+		CurrentRangeEnd = 0;
+		InstanceDataOffset = 0;
+		FInstanceRun Run = (FInstanceRun)0;
+		int CurrentRunIndex = -1;
+		for (uint OutputIndex2 = GroupThreadIndex; OutputIndex2 < RunInstanceCount; OutputIndex2 += NUM_THREADS_PER_GROUP)
+		{
+			// Naive search for range containing the output index (reasonably efficient for large runs, bad for many small runs)
+			while (OutputIndex2 >= CurrentRangeEnd)
+			{
+				CurrentRunIndex += 1;
+				Run = InstanceRuns[Cmd.FirstInstanceRunOffset + CurrentRunIndex];
+
+				CurrentRangeStart = CurrentRangeEnd;
+				CurrentRangeEnd += (Run.EndInclusive + 1) - Run.Start;
+				InstanceDataOffset = GetPrimitiveData(Run.PrimitiveId).InstanceDataOffset;
+			}
+
+			uint OutId = uint(InstanceDataOffset) + Run.Start + OutputIndex2 - CurrentRangeStart;
+			for (uint ViewIdIndex = 0; ViewIdIndex < NumViewIds; ++ViewIdIndex)
+			{
+				if (TestInstanceVisibilityBit(OutId, ViewIds[ViewIdIndex]))
+				{
+					uint OutputOffset = 0;
+					InterlockedAdd(FinalInstanceCountGroupShared, 1, OutputOffset);
+					WriteInstanceId(InstanceIdOutputOffset + OutputOffset, OutId, ViewIdIndex, CommandSlotIndex);
+				}
+			}
+		}
+	}
+	GroupMemoryBarrierWithGroupSync();
+
+	// Generate draw command
+	if (GroupThreadIndex == 0)
+	{
+		DrawIndirectArgsBufferOut[CommandSlotIndex * INDIRECT_ARGS_NUM_WORDS + 0] = Cmd.NumVerticesOrIndices;
+		DrawIndirectArgsBufferOut[CommandSlotIndex * INDIRECT_ARGS_NUM_WORDS + 1] = FinalInstanceCountGroupShared;
+		DrawIndirectArgsBufferOut[CommandSlotIndex * INDIRECT_ARGS_NUM_WORDS + 2] = Cmd.FirstIndex;
+		DrawIndirectArgsBufferOut[CommandSlotIndex * INDIRECT_ARGS_NUM_WORDS + 3] = Cmd.BaseVertexIndex;
+		DrawIndirectArgsBufferOut[CommandSlotIndex * INDIRECT_ARGS_NUM_WORDS + 4] = 0U;
+
+	}
+	GroupMemoryBarrierWithGroupSync();
+}
diff --git a/Engine/Shaders/Private/InstanceCulling/CullInstances.usf b/Engine/Shaders/Private/InstanceCulling/CullInstances.usf
new file mode 100644
--- /dev/null
+++ b/Engine/Shaders/Private/InstanceCulling/CullInstances.usf
@@ -0,0 +1,62 @@
+// Copyright 1998-2019 Epic Games, Inc. All Rights Reserved.
+
+/*=============================================================================
+	CullInstances.usf: 
+=============================================================================*/
+
+#include "../Common.ush"
+#include "../SceneData.ush"
+#include "../Nanite/NaniteDataDecode.ush"
+#include "../Nanite/HZBCull.ush"
+
+//StructuredBuffer<FPackedNaniteView> Views;
+RWStructuredBuffer<uint> InstanceVisibilityFlagsOut;
+uint NumInstances;
+uint NumInstanceFlagWords;
+uint NumViews;
+uint InstanceDataSOAStride;
+
+/**
+ */
+[numthreads(NUM_THREADS_PER_GROUP, 1, 1)]
+void CullInstancesCs(uint InstanceId : SV_DispatchThreadID)
+{
+	if (InstanceId >= NumInstances)
+	{
+		return;
+	}
+
+	const bool bNearClip = true;
+
+	FInstanceSceneData InstanceData = GetInstanceData(InstanceId, InstanceDataSOAStride);
+	uint WordMask = 1U << (InstanceId % 32U);
+	uint InstanceWordOffset = InstanceId / 32U;
+
+	bool bIsVisible = InstanceData.PrimitiveId != 0xFFFFFFFFu;
+
+	if (!bIsVisible)
+	{
+		return;
+	}
+
+	for (uint ViewId = 0; ViewId < NumViews; ++ViewId)
+	{
+		FNaniteView NaniteView = GetNaniteView(ViewId);
+
+		float4x4 LocalToTranslatedWorld = InstanceData.LocalToWorld;
+		LocalToTranslatedWorld[3].xyz += NaniteView.PreViewTranslation.xyz;
+		float4x4 LocalToClip = mul(LocalToTranslatedWorld, NaniteView.TranslatedWorldToClip);
+
+		FFrustumCullData Cull = BoxCullFrustum(InstanceData.LocalBoundsCenter, InstanceData.LocalBoundsExtent, LocalToClip, bNearClip, false);
+
+		uint Flag = (Cull.bIsVisible ? WordMask : 0U);
+
+		// TODO: we can use ballot etc to do this without atomics (and without clearing) - but it would need to be portable and not sure if 
+		// we'll keep this code anyway.
+		if (Flag != 0U)
+		{
+			uint WordOffset = NumInstanceFlagWords * ViewId + InstanceWordOffset;
+			InterlockedOr(InstanceVisibilityFlagsOut[WordOffset], Flag);
+		}
+	}
+}
