Fix for sss to match between strata and current system.

diff --git a/Engine/Shaders/Private/Strata/Strata.ush b/Engine/Shaders/Private/Strata/Strata.ush
--- a/Engine/Shaders/Private/Strata/Strata.ush
+++ b/Engine/Shaders/Private/Strata/Strata.ush
@@ -210,7 +210,7 @@ struct FStrataSubsurfaceData
 	float  RadiusScale;
 
 	float3 BaseColor;
-	uint   Profile;
+	float  Profile;		// Integer ID packed in 8bits as float
 
 	float3 DMFPAlbedo;
 	float  DMFPRadius;
@@ -317,14 +317,16 @@ FStrataData GetStrataSlabBSDF(
 		StrataRequestSharedNormalTangentBasis(SharedNormalTypes, SharedNormalIndex);
 	}
 
+	// STRATA_TODO apply SSS as a function of pixel footprint w.r.t. MFP
+	const bool bHasSSSProfile = SSSProfileID  > 0.0f;	// When no profile is provide, we use Burley. STRATA_TODO we should clarify as burley is also an option of the profile.
+	const bool bHasSSS = SSSDMFPRadius > 0 || bHasSSSProfile;
+
 	StrataData.Layers[0].BSDFCount++;
 	BSDF_SETTYPE			(StrataData.Layers[0].BSDFs[0],		STRATA_BSDF_TYPE_SLAB);
 	BSDF_SETNORMALID		(StrataData.Layers[0].BSDFs[0],		SharedNormalIndex);
 	BSDF_SETANISOTROPY		(StrataData.Layers[0].BSDFs[0],		bHasAnisotropy);
 	BSDF_SETISTOPLAYER		(StrataData.Layers[0].BSDFs[0],		0);
-	BSDF_SETHASSCATTERING	(StrataData.Layers[0].BSDFs[0],		SSSDMFPRadius > 0 ? 1 : 0);
-
-	const bool bHasSSSProfile = SSSProfileID > 0;
+	BSDF_SETHASSCATTERING	(StrataData.Layers[0].BSDFs[0],		bHasSSS ? 1 : 0);
 	BSDF_SETHASSSSPROFILE	(StrataData.Layers[0].BSDFs[0],		bHasSSSProfile ? 1 : 0);
 
 	SLAB_BASECOLOR			(StrataData.Layers[0].BSDFs[0])		= BaseColor;
@@ -335,7 +337,7 @@ FStrataData GetStrataSlabBSDF(
 	if (bHasSSSProfile)
 	{
 		SLAB_SSSPROFILEID	(StrataData.Layers[0].BSDFs[0])		= SSSProfileID;
-		SLAB_SSSPROFILERADIUS(StrataData.Layers[0].BSDFs[0])	= SSSDMFPRadius;
+		SLAB_SSSPROFILERADIUS(StrataData.Layers[0].BSDFs[0])	= 1.0; // STRATA_TODO add a scale to control the radius per pixel
 	}
 	else
 	{
@@ -1080,7 +1082,7 @@ void PackStrataOut(
 				}
 
 				// Select the first/toppest layer with subsurface scattering
-				if (BSDF_GETHASSCATTERING(BSDF) && BSDFCoverageAvg > 0 && SLAB_SSSPROFILERADIUS(BSDF) > 0 && !SSS.bHasSubsurface)
+				if (BSDF_GETHASSCATTERING(BSDF) && BSDFCoverageAvg > 0 && !SSS.bHasSubsurface)
 				{
 					SSS.bHasSubsurface = true;
 					SSS.bHasProfile	= BSDF_GETHASSSSPROFILE(BSDF);
@@ -1191,12 +1193,12 @@ void PackStrataOut(
 		{
 			GBuffer.BaseColor      = SSS.BaseColor;
 			GBuffer.CustomData.rgb = EncodeSubsurfaceProfile(SSS.Profile);
-			GBuffer.CustomData.a   = SSS.RadiusScale;
+			GBuffer.CustomData.a   = 1.0f; // STRATA_TODO add a radius scale
 		}
 		else
 		{
 			GBuffer.BaseColor      = SSS.BaseColor;
-			GBuffer.CustomData.rgb = sqrt(SSS.DMFPAlbedo);
+			GBuffer.CustomData.rgb = EncodeSubsurfaceColor(SSS.DMFPAlbedo);
 			GBuffer.CustomData.a   = saturate(SSS.DMFPRadius / SSS_PROFILE_PERPIXEL_MAXRADIUS);
 		}
 	}
@@ -1258,7 +1260,7 @@ void PackStrataOut(
 				{
 					if (BSDF_GETHASSSSPROFILE(BSDF))
 					{
-						STRATA_STORE_UINT1(PackRGBA8(float4(SLAB_SSSPROFILEID(BSDF), SLAB_SSSPROFILERADIUS(BSDF), 0, 0))); // STRATA_TODO SLAB_SSSPROFILERADIUS is in [0,1]?
+						STRATA_STORE_UINT1(PackRGBA8(float4(SLAB_SSSPROFILEID(BSDF), SLAB_SSSPROFILERADIUS(BSDF), 0, 0)));
 					}
 					else
 					{
diff --git a/Engine/Shaders/Private/Strata/StrataMaterialClassification.usf b/Engine/Shaders/Private/Strata/StrataMaterialClassification.usf
--- a/Engine/Shaders/Private/Strata/StrataMaterialClassification.usf
+++ b/Engine/Shaders/Private/Strata/StrataMaterialClassification.usf
@@ -56,7 +56,7 @@ void ComputePostProcessBSDFData(
 		}
 
 		// Select the first/toppest layer with subsurface scattering
-		if (BSDF_GETHASSCATTERING(BSDF) && BSDFCoverageAvg > 0 && SLAB_SSSPROFILERADIUS(BSDF) > 0 && !SSS.bHasSubsurface)
+		if (BSDF_GETHASSCATTERING(BSDF) && BSDFCoverageAvg > 0 && !SSS.bHasSubsurface)
 		{
 			SSS.bHasSubsurface = true;
 			SSS.bHasProfile	= BSDF_GETHASSSSPROFILE(BSDF);
diff --git a/Engine/Shaders/Private/SubsurfaceBurleyNormalized.ush b/Engine/Shaders/Private/SubsurfaceBurleyNormalized.ush
--- a/Engine/Shaders/Private/SubsurfaceBurleyNormalized.ush
+++ b/Engine/Shaders/Private/SubsurfaceBurleyNormalized.ush
@@ -197,24 +197,38 @@ FBurleyParameter GetBurleyParameters(FGBufferData GBuffer)
 {
 	FBurleyParameter Out = (FBurleyParameter)0;
 #if STRATA_ENABLED
-	const bool bUseProfileData = GBuffer.CustomData.y == 0 && GBuffer.CustomData.z == 0;
-	if (GBuffer.ShadingModelID == SHADINGMODELID_STRATA && !bUseProfileData)
+	if (GBuffer.ShadingModelID == SHADINGMODELID_STRATA)
 	{
-		// Special data encoding from strata packing
-		const float3 MFPAlebdo			= Pow2(GBuffer.CustomData.xyz);
-		const float  MFPRadius			= GBuffer.CustomData.a * SSS_PROFILE_PERPIXEL_MAXRADIUS;
-		const float3 MFPChromaticRadius = MFPAlebdo * MFPRadius;
+		const bool bUseProfileData = GBuffer.CustomData.y == 0 && GBuffer.CustomData.z == 0;
+		if (bUseProfileData)
+		{
+			uint SubsurfaceProfileInt = ExtractSubsurfaceProfileInt(GBuffer);
+
+			Out.SurfaceAlbedo		= GetSubsurfaceProfileSurfaceAlbedo(SubsurfaceProfileInt);
+			Out.DiffuseMeanFreePath = DecodeDiffuseMeanFreePath(GetSubsurfaceProfileDiffuseMeanFreePath(SubsurfaceProfileInt));
+			Out.WorldUnitScale		= DecodeWorldUnitScale(GetSubsurfaceProfileWorldUnitScale(SubsurfaceProfileInt)) * BURLEY_CM_2_MM;
+
+			Out.SurfaceOpacity		= 1.0f;
+			Out.DiffuseMeanFreePath*= Out.SurfaceOpacity;
+		}
+		else
+		{
+			// Special data encoding from strata packing
+			const float3 MFPAlebdo			= Pow2(GBuffer.CustomData.xyz);
+			const float  MFPRadius			= GBuffer.CustomData.a * SSS_PROFILE_PERPIXEL_MAXRADIUS * BURLEY_CM_2_MM * 5.0f; // STRATA_TODO why a scale of 5 is needed here?
+			const float3 MFPChromaticRadius = MFPAlebdo * MFPRadius;
 
-		const float MaxAlbedo			= max3(MFPAlebdo.x, MFPAlebdo.y, MFPAlebdo.z);
-		const uint MaxAlbedoComponent	= MFPAlebdo.x == MaxAlbedo ? 0 : (MFPAlebdo.y == MaxAlbedo ? 1 : 2);
+			const float MaxDMFP				= max3(MFPChromaticRadius.x, MFPChromaticRadius.y, MFPChromaticRadius.z);
+			const uint IndexOfMaxDmfp		= (MFPChromaticRadius.r == MaxDMFP) ? 0 : ((MFPChromaticRadius.b == MaxDMFP) ? 1 : 2);// STARTA_TODO fix the component selection in SubsirfaceProfile, b should map to 2
 
-		Out.SurfaceAlbedo		= float4(GBuffer.BaseColor, MaxAlbedoComponent);
-		Out.DiffuseMeanFreePath	= float4(MFPChromaticRadius, max3(MFPChromaticRadius.x, MFPChromaticRadius.y, MFPChromaticRadius.z));
-		Out.WorldUnitScale		= BURLEY_CM_2_MM;
+			Out.SurfaceAlbedo				= float4(GBuffer.BaseColor, GBuffer.BaseColor[IndexOfMaxDmfp]);
+			Out.DiffuseMeanFreePath			= float4(MFPChromaticRadius, MaxDMFP);
+			Out.WorldUnitScale				= DecodeWorldUnitScale(1.0f * ENC_WORLDUNITSCALE_IN_CM_TO_UNIT) * BURLEY_CM_2_MM;
 
-		// Mask should always be true, we can probably ignore it?
-		// STRATA_TODO smooth transistion based on pixel footprint
-		Out.SurfaceOpacity		= MFPRadius > 0 ? 1.f : 0.f;
+			// STRATA_TODO smooth transistion based on pixel footprint w.r.t. MFPRadius
+			Out.SurfaceOpacity				= 1.0;
+			Out.DiffuseMeanFreePath			*= Out.SurfaceOpacity;
+		}
 	}
 	else
 #endif
