Add blending mode support to path tracer

Start implementing support for blending modes in the path tracer. This first implementation uses indirect rays to implement transparency so transparent shadows are technically treated as caustics.

Added an implementation of microfacet based rough glass for the translucent case.

Change ray traced refractions to use Ior instead of the specular parameter.

diff --git a/Engine/Shaders/Private/PathTracing/Material/PathTracingFresnel.ush b/Engine/Shaders/Private/PathTracing/Material/PathTracingFresnel.ush
--- a/Engine/Shaders/Private/PathTracing/Material/PathTracingFresnel.ush
+++ b/Engine/Shaders/Private/PathTracing/Material/PathTracingFresnel.ush
@@ -13,7 +13,8 @@ float FresnelReflectance(float CosI, float Eta)
 
 	float g2 = Eta * Eta - 1 + CosI * CosI;
 
-	if (g2 >= 0) {
+	if (g2 >= 0)
+	{
 		float c = abs(CosI);
 		float g = sqrt(g2);
 		float a2 = Pow2((g - c) / (g + c));
@@ -25,14 +26,15 @@ float FresnelReflectance(float CosI, float Eta)
 }
 
 // returns true if we sampled a refraction, or false if we reflected (or TIR occured)
-bool SampleRefraction(float3 RayDirection, float3 Normal, float Eta, float RandSample, out float3 OutDirection)
+bool SampleRefraction(float3 RayDirection, float3 Normal, float Eta, float RandSample, out float3 OutDirection, out float OutFresnel)
 {
 	// Compute refracted direction and Fresnel term at the same time to re-use intermediate values
 	// and to get matching results for the TIR condition
 	float CosI = abs(dot(RayDirection, Normal));
 	float g2 = Eta * Eta - 1.0 + CosI * CosI;
 
-	if (g2 >= 0.0) {
+	if (g2 >= 0.0)
+	{
 		float c = CosI;
 		float g = sqrt(g2);
 		float a2 = Pow2((g - c) / (g + c));
@@ -43,9 +45,16 @@ bool SampleRefraction(float3 RayDirection, float3 Normal, float Eta, float RandS
 		// This assumes RandSample is in [0,1)
 		if (RandSample >= Fr)
 		{
+			OutFresnel = 1 - Fr;
 			OutDirection = normalize(RayDirection + Normal * (c - g));
 			return true;
 		}
+		OutFresnel = Fr;
+	}
+	else
+	{
+		// TIR
+		OutFresnel = 1.0;
 	}
 	// either TIR, or sampled reflection
 	OutDirection = RayDirection + (2 * CosI) * Normal;
diff --git a/Engine/Shaders/Private/PathTracing/Material/PathTracingMaterialCommon.ush b/Engine/Shaders/Private/PathTracing/Material/PathTracingMaterialCommon.ush
--- a/Engine/Shaders/Private/PathTracing/Material/PathTracingMaterialCommon.ush
+++ b/Engine/Shaders/Private/PathTracing/Material/PathTracingMaterialCommon.ush
@@ -9,6 +9,9 @@ PathTracingMaterialCommon.usf: Brdf utility functions
 #include "../PathTracingCommon.ush"
 #include "/Engine/Private/BRDF.ush"
 
+// be conservative since the GPU may be aggressive with floating point optimizations
+#define GGX_MIN_ROUGHNESS 0.01
+
 float3 GetMaterialEmission(FMaterialClosestHitPayload HitInfo)
 {
 	return HitInfo.Radiance;
@@ -112,7 +115,8 @@ bool IsTwoSidedFoliageMaterial(FMaterialClosestHitPayload HitInfo)
 
 bool IsMaterialTransmissive(FMaterialClosestHitPayload HitInfo)
 {
-	return HitInfo.ShadingModelID == SHADINGMODELID_TWOSIDED_FOLIAGE;
+	return HitInfo.ShadingModelID == SHADINGMODELID_TWOSIDED_FOLIAGE ||
+		   (HitInfo.BlendingMode == RAY_TRACING_BLEND_MODE_TRANSLUCENT && HitInfo.Roughness != 0 && HitInfo.Opacity < 1);
 }
 
 // Takes as input the sample weight and pdf for a certain lobe of a mixed model, together with the probability of picking that lobe
diff --git a/Engine/Shaders/Private/PathTracing/Material/PathTracingMaterialSampling.ush b/Engine/Shaders/Private/PathTracing/Material/PathTracingMaterialSampling.ush
--- a/Engine/Shaders/Private/PathTracing/Material/PathTracingMaterialSampling.ush
+++ b/Engine/Shaders/Private/PathTracing/Material/PathTracingMaterialSampling.ush
@@ -29,7 +29,8 @@ bool IsNeeValidMaterial(FMaterialClosestHitPayload Payload)
 	return true;
 #else
 
-	if (IsPureSpecularTransmission(Payload))
+	if (Payload.BlendingMode == RAY_TRACING_BLEND_MODE_MODULATE ||
+		Payload.BlendingMode == RAY_TRACING_BLEND_MODE_ALPHA_HOLDOUT)
 	{
 		return false;
 	}
@@ -41,7 +42,7 @@ bool IsNeeValidMaterial(FMaterialClosestHitPayload Payload)
 		case SHADINGMODELID_DEFAULT_LIT:		
 		case SHADINGMODELID_CLEAR_COAT:			
 		case SHADINGMODELID_TWOSIDED_FOLIAGE:	
-		default:								
+		default:
 			return true;
 	}
 #endif
@@ -69,12 +70,88 @@ void SampleMaterial(
 	{
 		Lambert_SampleMaterial(Payload, RandSample, OutDirection, OutWeight, OutPdf);
 	}
-#else	
-	if (IsPureSpecularTransmission(Payload))
+#else
+
+	if (Payload.BlendingMode == RAY_TRACING_BLEND_MODE_MODULATE)
+	{
+		OutDirection = RayDirection;
+		OutWeight = Payload.Radiance;
+		OutPdf = POSITIVE_INFINITY;
+		OutPositionBiasSign = -1;
+		return;
+	}
+
+	if (Payload.BlendingMode == RAY_TRACING_BLEND_MODE_ALPHA_HOLDOUT)
 	{
-		SpecularRefraction_SampleMaterial(RayDirection, Payload, RandSample, OutDirection, OutWeight, OutPdf, OutPositionBiasSign);
+		OutDirection = RayDirection;
+		OutWeight = 1;
+		OutPdf = POSITIVE_INFINITY;
+		OutPositionBiasSign = -1;
 		return;
 	}
+
+	// Decide the overall weight of the BRDF portion vs. the translucent portion
+	// Note that the translucent portion could be refraction, transparency or SSS
+	float BrdfWeight = 1;
+	float TransWeight = 0;
+	float BrdfProb = 1;
+	switch (Payload.BlendingMode)
+	{
+		case RAY_TRACING_BLEND_MODE_TRANSLUCENT:
+			BrdfWeight = 1;
+			TransWeight = 1;
+			BrdfProb = Payload.Opacity;
+			break;
+		case RAY_TRACING_BLEND_MODE_ADDITIVE:
+			BrdfWeight = 1 + Payload.Opacity;
+			TransWeight = 1 + Payload.Opacity;
+			BrdfProb = Payload.Opacity / (1 + Payload.Opacity);
+			break;
+		case RAY_TRACING_BLEND_MODE_ALPHA_COMPOSITE:
+			BrdfWeight = 2 - Payload.Opacity;;
+			TransWeight = 2 - Payload.Opacity;
+			BrdfProb = 1 / (2 - Payload.Opacity);
+			break;
+		default:
+			// opaque or one of the modes handled above
+			break;
+	}
+
+	if (RandSample.x < BrdfProb)
+	{
+		// we randomly decided to sample the brdf portion of the material
+		RandSample.x /= BrdfProb;
+	}
+	else
+	{
+		// we randomly decided to sample the translucent portion
+		RandSample.x = (RandSample.x - BrdfProb) / (1 - BrdfProb);
+		// now decide what to do based on blend mode
+		if (Payload.BlendingMode == RAY_TRACING_BLEND_MODE_TRANSLUCENT && Payload.Ior != 1.0)
+		{
+			// refraction model
+			if (Payload.Roughness == 0)
+			{
+				SpecularRefraction_SampleMaterial(RayDirection, Payload, RandSample, OutDirection, OutWeight, OutPdf, OutPositionBiasSign);
+			}
+			else
+			{
+				RoughGlass_SampleMaterial(RayDirection, Payload, RandSample, OutDirection, OutWeight, OutPdf, OutPositionBiasSign);
+			}
+			OutWeight *= TransWeight;
+			OutPdf *= 1 - BrdfProb;
+		}
+		else
+		{
+			// plain transparency
+			OutDirection = RayDirection;
+			OutWeight = TransWeight;
+			OutPdf = POSITIVE_INFINITY; // * 1 - BrdfProb
+			OutPositionBiasSign = -1;
+		}
+		return;
+	}
+
 	switch (Payload.ShadingModelID)
 	{
 		case SHADINGMODELID_UNLIT:
@@ -85,19 +162,22 @@ void SampleMaterial(
 			return;
 		case SHADINGMODELID_DEFAULT_LIT:
 			DefaultLit_SampleMaterial(RayDirection, Payload, RandSample, OutDirection, OutWeight, OutPdf);
-			return;
+			break;
 		case SHADINGMODELID_CLEAR_COAT:
 			ClearCoat_SampleMaterial(RayDirection, Payload, RandSample, OutDirection, OutWeight, OutPdf);
-			return;
+			break;
 		case SHADINGMODELID_TWOSIDED_FOLIAGE:
 			TwoSidedFoliage_SampleMaterial(RayDirection, Payload, RandSample, OutDirection, OutWeight, OutPdf, OutPositionBiasSign);
-			return;
+			break;
 		default:
 			Lambert_SampleMaterial(Payload, RandSample, OutDirection, OutWeight, OutPdf);
 			break;
 	}
-#endif
 
+	// final adjustment to weight / prob based on the stochastic decision for translucency earlier
+	OutWeight *= BrdfWeight;
+	OutPdf *= BrdfProb;
+#endif
 }
 
 void EvalMaterial(
@@ -120,12 +200,50 @@ void EvalMaterial(
 		Lambert_EvalMaterial(OutgoingDirection, Payload, OutWeight, OutPdf);
 	}
 #else
-	if (IsPureSpecularTransmission(Payload))
+
+	if (Payload.BlendingMode == RAY_TRACING_BLEND_MODE_MODULATE ||
+		Payload.BlendingMode == RAY_TRACING_BLEND_MODE_ALPHA_HOLDOUT)
 	{
-		SpecularTransmission_EvalMaterial(OutWeight, OutPdf);
+		OutWeight = 0;
+		OutPdf = 0;
 		return;
 	}
 
+	// Decide the overall weight of the BRDF portion vs. the translucent portion
+	// Note that the translucent portion could be refraction, transparency or SSS
+	float BrdfWeight = 1;
+	float TransWeight = 0;
+	float BrdfProb = 1;
+	switch (Payload.BlendingMode)
+	{
+		case RAY_TRACING_BLEND_MODE_TRANSLUCENT:
+			// brdf:   Opacity
+			// glass:  1.0 - Opacity
+			BrdfWeight = 1;
+			TransWeight = 1;
+			BrdfProb = Payload.Opacity;
+			break;
+		case RAY_TRACING_BLEND_MODE_ADDITIVE:
+			// brdf:   Opacity
+			// transp: 1.0
+			BrdfWeight = 1 + Payload.Opacity;
+			TransWeight = 1 + Payload.Opacity;
+			BrdfProb = Payload.Opacity / (1 + Payload.Opacity);
+			break;
+		case RAY_TRACING_BLEND_MODE_ALPHA_COMPOSITE:
+			// brdf:   1.0
+			// transp: 1.0 - Opacity
+			BrdfWeight = 2 - Payload.Opacity;;
+			TransWeight = 2 - Payload.Opacity;
+			BrdfProb = 1 / (2 - Payload.Opacity);
+			break;
+		default:
+			// opaque or one of the modes handled above
+			break;
+	}
+
+	// NOTE: the translucent portions are all singular -- so we can ignore them in the eval method for now as they have Pdf=+inf for a single direction only
+
 	switch (Payload.ShadingModelID)
 	{
 		case SHADINGMODELID_UNLIT:
@@ -145,6 +263,26 @@ void EvalMaterial(
 			Lambert_EvalMaterial(OutgoingDirection, Payload, OutWeight, OutPdf);
 			break;
 	}
+
+	// final adjustment to weight / prob for translucency
+	OutWeight *= BrdfWeight;
+	OutPdf *= BrdfProb;
+
+	if (BrdfProb < 1.0)
+	{
+		if (Payload.BlendingMode == RAY_TRACING_BLEND_MODE_TRANSLUCENT && Payload.Ior != 1.0)
+		{
+			// refraction model (only evaluate for rough case as singular can't be evaluated)
+			if (Payload.Roughness != 0)
+			{
+				float3 GlassWeight = 0;
+				float GlassPdf = 0;
+				RoughGlass_EvalMaterial(IncomingDirection, OutgoingDirection, Payload, GlassWeight, GlassPdf);
+
+				AddLobeWithMIS(OutWeight, OutPdf, GlassWeight, GlassPdf, 1.0 - BrdfProb);
+			}
+		}
+	}
 #endif
 }
 
diff --git a/Engine/Shaders/Private/PathTracing/Material/PathTracingSpecularTransmission.ush b/Engine/Shaders/Private/PathTracing/Material/PathTracingSpecularTransmission.ush
--- a/Engine/Shaders/Private/PathTracing/Material/PathTracingSpecularTransmission.ush
+++ b/Engine/Shaders/Private/PathTracing/Material/PathTracingSpecularTransmission.ush
@@ -1,7 +1,7 @@
 // Copyright Epic Games, Inc. All Rights Reserved.
 
 /*=============================================================================================
-PathTracingSpecularTransmission.usf: Pure specular transmission BTDF sampling functions
+PathTracingSpecularTransmission.usf: Refractive BSDFs (pure singular and microfacet based)
 ===============================================================================================*/
 
 #pragma once
@@ -20,16 +20,15 @@ void SpecularRefraction_SampleMaterial(
 	float3 WorldNormal = GetWorldNormal(Payload);
 
 	// #dxr_todo: Maintain a refraction stack on the path tracing payload
-	// #dxr_todo Determine if parameterization from Specular is still the proper path forward
-	//float Ior = Payload.Ior;
-	float Ior = DielectricF0ToIor(DielectricSpecularToF0(Payload.Specular));
+	float Ior = Payload.Ior;
 	float Eta = bIsEntering ? Ior : 1.0 / Ior;
 
 	// Throughput is 1 since we pick reflect/refract according to the fresnel term
 	OutWeight = 1.0;
 	OutPdf = POSITIVE_INFINITY;
 
-	if (SampleRefraction(RayDirection, WorldNormal, Eta, RandSample.x, OutDirection))
+	float F = 0;
+	if (SampleRefraction(RayDirection, WorldNormal, Eta, RandSample.x, OutDirection, F))
 	{
 		OutPositionBiasSign = -1;
 	}
@@ -43,4 +42,226 @@ void SpecularTransmission_EvalMaterial(out float3 OutThroughput, out float OutPd
 {
 	OutThroughput = 0;
 	OutPdf = 0;
+}
+
+void RoughGlass_SampleMaterial(
+	float3 RayDirection,
+	FMaterialClosestHitPayload Payload,
+	float4 RandSample,
+	out float3 OutDirection,
+	out float3 OutWeight,
+	out float OutPdf,
+	out float OutPositionBiasSign)
+{
+
+	OutDirection = 0;
+	OutWeight = 0;
+	OutPdf = 0;
+	OutPositionBiasSign = 1;
+
+	// NOTE: IsFrontFace() determines polygon orientation, because the normal is always flipped towards in the incoming ray
+	bool bIsEntering = Payload.IsFrontFace();
+	float3 N_World = GetWorldNormal(Payload);
+	float3 V_World = -RayDirection;
+
+	// #dxr_todo: Maintain a refraction stack on the path tracing payload
+	float Ior = Payload.Ior;
+	float Eta = bIsEntering ? Ior : 1.0 / Ior;
+
+	float Roughness = max(GetRoughness(Payload), GGX_MIN_ROUGHNESS);
+	float2 Alpha;
+	float3x3 Basis;
+	if (Payload.HasAnisotropy())
+	{
+		// Material is anisotropic, our shading frame should use the tangent vector
+		GetAnisotropicRoughness(Roughness * Roughness, Payload.Anisotropy, Alpha.x, Alpha.y);
+		float3 X = Payload.WorldTangent;
+		float3 Y = normalize(cross(N_World, X));
+		Basis = float3x3(cross(Y, N_World), Y, N_World);
+	}
+	else
+	{
+		// No anisotropy, just pick an arbitrary basis
+		Alpha = Roughness * Roughness;
+		Basis = GetTangentBasis(N_World);
+	}
+
+
+	float3 V = mul(Basis, V_World);
+	float NoV = saturate(V.z);
+
+	if (NoV == 0)
+	{
+		// invalid grazing angle
+		return;
+	}
+
+	float3 H = ImportanceSampleVisibleGGX_aniso(UniformSampleDisk(RandSample.xy), Alpha, V).xyz;
+	float3 N = float3(0, 0, 1);
+	float NoH = saturate(H.z);
+	float VoH = saturate(dot(V, H));
+
+	float3 L = 0;
+	float Fresnel = 0;
+	if (SampleRefraction(-V, H, Eta, RandSample.z, L, Fresnel))
+	{
+		// ray should be refracted
+		if (L.z >= 0)
+		{
+			// masked direction
+			return;
+		}
+		float NoL = saturate(-L.z);
+		float LoH = dot(L, H);
+
+		// FIXME: this can be further simplified
+		const float CosThetaM4 = Pow2(NoH);
+		const float TanThetaM2 = (Pow2(H.x / Alpha.x) + Pow2(H.y / Alpha.y)) / Pow2(NoH);
+		const float D = 1 / (PI * Pow2(1 + TanThetaM2) * Alpha.x * Alpha.y * CosThetaM4);
+		const float Lambda_o = 0.5 * (-1.0 + sqrt(1.0 + (Pow2(L.x * Alpha.x) + Pow2(L.y * Alpha.y)) / Pow2(L.z)));
+		const float Lambda_i = 0.5 * (-1.0 + sqrt(1.0 + (Pow2(V.x * Alpha.x) + Pow2(V.y * Alpha.y)) / Pow2(V.z)));
+		const float G2 = 1 / (Lambda_i + Lambda_o + 1);
+		const float G1 = 1 / (Lambda_i + 1);
+
+		const float3 Ht = (Eta * L + V);
+
+		OutWeight = (G2 / G1);
+		OutPdf = (Fresnel * abs(LoH * VoH) * (Eta * Eta) * (G1 * D)) / (dot(Ht, Ht) * abs(V.z));
+
+		OutPositionBiasSign = -1;
+	}
+	else
+	{
+		// ray should be reflected
+		if (L.z <= 0)
+		{
+			// masked direction
+			return;
+		}
+		float NoL = saturate(L.z);
+
+		// Specular lobe
+		if (NoL > 0 && NoV > 0)
+		{
+			// Assume NoH > 0
+			// FIXME: this can be further simplified
+			const float CosThetaM4 = Pow2(NoH);
+			const float TanThetaM2 = (Pow2(H.x / Alpha.x) + Pow2(H.y / Alpha.y)) / Pow2(NoH);
+			const float D = 1 / (PI * Pow2(1 + TanThetaM2) * Alpha.x * Alpha.y * CosThetaM4);
+			const float Lambda_o = 0.5 * (-1.0 + sqrt(1.0 + (Pow2(L.x * Alpha.x) + Pow2(L.y * Alpha.y)) / Pow2(L.z)));
+			const float Lambda_i = 0.5 * (-1.0 + sqrt(1.0 + (Pow2(V.x * Alpha.x) + Pow2(V.y * Alpha.y)) / Pow2(V.z)));
+			const float G2 = 1 / (Lambda_i + Lambda_o + 1);
+			const float G1 = 1 / (Lambda_i + 1);
+
+			OutWeight = (G2 / G1);
+			OutPdf = (Fresnel * G1 * D * 0.25f) / NoV;
+		}
+	}
+
+	// transform to world space
+	OutDirection = normalize(mul(L, Basis));
+}
+
+void RoughGlass_EvalMaterial(
+	float3 IncomingDirection,
+	float3 OutgoingDirection,
+	FMaterialClosestHitPayload Payload,
+	out float3 OutWeight,
+	out float OutPdf)
+{
+	OutWeight = 0;
+	OutPdf = 0;
+
+	// NOTE: IsFrontFace() determines polygon orientation, because the normal is always flipped towards in the incoming ray
+	bool bIsEntering = Payload.IsFrontFace();
+	float3 V_World = -IncomingDirection;
+	float3 L_World = OutgoingDirection;
+	float3 N_World = GetWorldNormal(Payload);
+
+	// #dxr_todo: Maintain a refraction stack on the path tracing payload
+	float Ior = Payload.Ior;
+	float Eta = bIsEntering ? Ior : 1.0 / Ior;
+
+	float Roughness = max(GetRoughness(Payload), GGX_MIN_ROUGHNESS);
+	float2 Alpha;
+	float3x3 Basis;
+	if (Payload.HasAnisotropy())
+	{
+		// Material is anisotropic, our shading frame should use the tangent vector
+		GetAnisotropicRoughness(Roughness * Roughness, Payload.Anisotropy, Alpha.x, Alpha.y);
+		float3 X = Payload.WorldTangent;
+		float3 Y = normalize(cross(N_World, X));
+		Basis = float3x3(cross(Y, N_World), Y, N_World);
+	}
+	else
+	{
+		// No anisotropy, just pick an arbitrary basis
+		Alpha = Roughness * Roughness;
+		Basis = GetTangentBasis(N_World);
+	}
+
+	// move vectors into right shading frame
+	float3 V = mul(Basis, V_World);
+	float3 L = mul(Basis, L_World);
+
+	if (V.z <= 0)
+	{
+		// invalid input
+		return;
+	}
+
+	float NoV = saturate(V.z);
+	if (L.z > 0)
+	{
+		// reflection side
+		float3 H = normalize(L + V);
+		float NoL = saturate(L.z);
+		float NoH = saturate(H.z);
+		float VoH = saturate(dot(V, H));
+
+		// FIXME: this can be further simplified
+		const float CosThetaM4 = Pow2(NoH);
+		const float TanThetaM2 = (Pow2(H.x / Alpha.x) + Pow2(H.y / Alpha.y)) / Pow2(NoH);
+		const float D = 1 / (PI * Pow2(1 + TanThetaM2) * Alpha.x * Alpha.y * CosThetaM4);
+		const float Lambda_o = 0.5 * (-1.0 + sqrt(1.0 + (Pow2(L.x * Alpha.x) + Pow2(L.y * Alpha.y)) / Pow2(L.z)));
+		const float Lambda_i = 0.5 * (-1.0 + sqrt(1.0 + (Pow2(V.x * Alpha.x) + Pow2(V.y * Alpha.y)) / Pow2(V.z)));
+		const float G2 = 1 / (Lambda_i + Lambda_o + 1);
+		const float G1 = 1 / (Lambda_i + 1);
+
+		const float Fr = FresnelReflectance(VoH, Eta);
+
+		OutWeight = (G2 / G1);
+		OutPdf = (Fr * G1 * D * 0.25f) / NoV;
+	}
+	else if (L.z < 0)
+	{
+		// refraction side
+		float NoL = saturate(-L.z);
+		float3 Ht = -(Eta * L + V);
+		Ht = (Eta < 1.0f) ? -Ht : Ht;
+		float3 Htn = normalize(Ht);
+
+		float VoH = dot(V, Htn);
+		float Ft = 1.0f - FresnelReflectance(VoH, Eta);
+		if (Ft > 0)
+		{
+			const float LoH = dot(L, Htn);
+			const float NoH = Htn.z;
+			if (NoH <= 0.0)
+			{
+				return;
+			}
+			// FIXME: this can be further simplified
+			const float CosThetaM4 = Pow2(NoH);
+			const float TanThetaM2 = (Pow2(Htn.x / Alpha.x) + Pow2(Htn.y / Alpha.y)) / Pow2(NoH);
+			const float D = 1 / (PI * Pow2(1 + TanThetaM2) * Alpha.x * Alpha.y * CosThetaM4);
+			const float Lambda_o = 0.5 * (-1.0 + sqrt(1.0 + (Pow2(L.x * Alpha.x) + Pow2(L.y * Alpha.y)) / Pow2(L.z)));
+			const float Lambda_i = 0.5 * (-1.0 + sqrt(1.0 + (Pow2(V.x * Alpha.x) + Pow2(V.y * Alpha.y)) / Pow2(V.z)));
+			const float G2 = 1 / (Lambda_i + Lambda_o + 1);
+			const float G1 = 1 / (Lambda_i + 1);
+
+			OutWeight = (G2 / G1);
+			OutPdf = (Ft * abs(LoH * VoH) * (Eta * Eta) * (G1 * D)) / (dot(Ht, Ht) * abs(V.z));
+		}
+	}
 }
\ No newline at end of file
diff --git a/Engine/Shaders/Private/PathTracing/PathTracing.usf b/Engine/Shaders/Private/PathTracing/PathTracing.usf
--- a/Engine/Shaders/Private/PathTracing/PathTracing.usf
+++ b/Engine/Shaders/Private/PathTracing/PathTracing.usf
@@ -134,14 +134,6 @@ RAY_TRACING_ENTRY_RAYGEN(PathTracingMainRG)
 			break;
 		}
 
-		if (Payload.ShadingModelID == SHADINGMODELID_SUBSURFACE)
-		{
-			// for now just approximate SSS by moving it into the diffuse color
-			Payload.DiffuseColor += Payload.CustomData.xyz;
-			Payload.Opacity = 1; // overwrite radius
-			Payload.ShadingModelID == SHADINGMODELID_DEFAULT_LIT;
-		}
-
 		// add in surface emission (except for modulate blend mode which uses Radiance as the transparency multiplier)
 		if (Payload.BlendingMode != RAY_TRACING_BLEND_MODE_MODULATE)
 		{
@@ -163,6 +155,16 @@ RAY_TRACING_ENTRY_RAYGEN(PathTracingMainRG)
 			Payload.WorldNormal = -Payload.WorldNormal;
 		}
 
+		// material tweaks to workaround path tracer limitations
+		if (Payload.ShadingModelID == SHADINGMODELID_SUBSURFACE)
+		{
+			// for now just approximate SSS by moving it into the diffuse color
+			Payload.DiffuseColor += Payload.CustomData.xyz;
+			Payload.DiffuseColor = saturate(Payload.DiffuseColor);
+			Payload.Opacity = 1; // overwrite radius
+			Payload.ShadingModelID = SHADINGMODELID_DEFAULT_LIT;
+		}
+
 		// Choose a random number for both Light sampling and BxDF sampling
 		float4 RandSample = RandomSequence_GenerateSample4D(RandSequence);
 
diff --git a/Engine/Shaders/Private/RayTracing/RayTracingPrimaryRays.usf b/Engine/Shaders/Private/RayTracing/RayTracingPrimaryRays.usf
--- a/Engine/Shaders/Private/RayTracing/RayTracingPrimaryRays.usf
+++ b/Engine/Shaders/Private/RayTracing/RayTracingPrimaryRays.usf
@@ -273,9 +273,7 @@ RAY_TRACING_ENTRY_RAYGEN(RayTracingPrimaryRaysRGS)
 		float3 RefractedDirection = Ray.Direction;
 		if (TranslucencyRefraction)
 		{
-			// #dxr_todo Determine if parameterization from Specular is still the proper path forward
-			//float Ior = Payload.Ior;
-			float Ior = DielectricF0ToIor(DielectricSpecularToF0(Payload.Specular));
+			float Ior = Payload.Ior;
 			bHasScattered |= Ior > 1.0 ? true : false;
 
 			float ReflectionRefractionEventThroughput;
