Some shader changes related to GPU instance culling (USE_INSTANCE_CULLING, off by default)

diff --git a/Engine/Shaders/Private/HitProxyPixelShader.usf b/Engine/Shaders/Private/HitProxyPixelShader.usf
--- a/Engine/Shaders/Private/HitProxyPixelShader.usf
+++ b/Engine/Shaders/Private/HitProxyPixelShader.usf
@@ -18,7 +18,7 @@ float4 HitProxyId;
 
 void Main(
 	FVertexFactoryInterpolantsVSToPS Interpolants,
-#if PER_INSTANCE_HITPROXY_ID
+#if PER_INSTANCE_HITPROXY_ID || USE_INSTANCE_CULLING
 	float4 InstanceHitProxyId : HIT_PROXY_ID,
 #endif
 	float4 PixelPosition	: TEXCOORD6,
@@ -50,7 +50,7 @@ void Main(
 
 	GetMaterialCoverageAndClipping(MaterialParameters, PixelMaterialInputs);
 
-#if PER_INSTANCE_HITPROXY_ID
+#if PER_INSTANCE_HITPROXY_ID || USE_INSTANCE_CULLING
 	OutColor = HitProxyId + InstanceHitProxyId;
 #else
 	OutColor = HitProxyId;
diff --git a/Engine/Shaders/Private/HitProxyVertexShader.usf b/Engine/Shaders/Private/HitProxyVertexShader.usf
--- a/Engine/Shaders/Private/HitProxyVertexShader.usf
+++ b/Engine/Shaders/Private/HitProxyVertexShader.usf
@@ -21,7 +21,7 @@ Buffer<float4> VertexFetch_HitProxyIdBuffer;
 struct FHitProxyVSToPS
 {
 	FVertexFactoryInterpolantsVSToPS FactoryInterpolants;
-#if PER_INSTANCE_HITPROXY_ID
+#if PER_INSTANCE_HITPROXY_ID || USE_INSTANCE_CULLING
 	float4 InstanceHitProxyId : HIT_PROXY_ID;
 #endif
 	float4 PixelPosition : TEXCOORD6;
@@ -125,7 +125,7 @@ void Main(
 	#else
 		Output.InstanceHitProxyId = float4(0,0,0,0);
 	#endif
-#elif USE_INSTANCING || USE_VERTEXFACTORY_HITPROXY_ID
+#elif USE_INSTANCING || USE_VERTEXFACTORY_HITPROXY_ID|| USE_INSTANCE_CULLING
 	Output.InstanceHitProxyId = VertexFactoryGetInstanceHitProxyId(Input, VFIntermediates);
 #endif
 
diff --git a/Engine/Shaders/Private/LocalVertexFactory.ush b/Engine/Shaders/Private/LocalVertexFactory.ush
--- a/Engine/Shaders/Private/LocalVertexFactory.ush
+++ b/Engine/Shaders/Private/LocalVertexFactory.ush
@@ -15,7 +15,7 @@
 
 #include "/Engine/Generated/UniformBuffers/PrecomputedLightingBuffer.ush"
 
-#if USE_INSTANCING
+#if USE_INSTANCE_CULLING || USE_INSTANCING
 	#ifndef USE_DITHERED_LOD_TRANSITION_FOR_INSTANCED
 		#error "USE_DITHERED_LOD_TRANSITION_FOR_INSTANCED should have been defined"
 	#endif
@@ -27,19 +27,21 @@
 	#define USE_DITHERED_LOD_TRANSITION USE_DITHERED_LOD_TRANSITION_FROM_MATERIAL
 #endif
 
-#if USE_INSTANCING
-#if USE_DITHERED_LOD_TRANSITION
-	float4 InstancingViewZCompareZero;  // w contains random lod scale 
-	float4 InstancingViewZCompareOne;
-	float4 InstancingViewZConstant;
-	float4 InstancingWorldViewOriginZero;
-	float4 InstancingWorldViewOriginOne;
+#if USE_INSTANCE_CULLING || USE_INSTANCING
+	#if USE_DITHERED_LOD_TRANSITION
+		float4 InstancingViewZCompareZero;  // w contains random lod scale 
+		float4 InstancingViewZCompareOne;
+		float4 InstancingViewZConstant;
+		float4 InstancingWorldViewOriginZero;
+		float4 InstancingWorldViewOriginOne;
+	#endif
+	float4 InstancingFadeOutParams;
 #endif
 
+#if USE_INSTANCING
 float4 InstancingOffset;
-float4 InstancingFadeOutParams;
 uint InstanceOffset;
-#endif	// USE_INSTANCING
+#endif
 
 #if USE_SPLINEDEFORM
 float4 SplineParams[10];
@@ -135,7 +137,7 @@ struct FVertexFactoryInput
 	float2	TexCoords[1] : ATTRIBUTE4;
 #endif
 
-#if USE_INSTANCING && !MANUAL_VERTEX_FETCH
+#if (USE_INSTANCING || USE_INSTANCE_CULLING) && !MANUAL_VERTEX_FETCH
 	float4 InstanceOrigin : ATTRIBUTE8;  // per-instance random in w 
 	half4 InstanceTransform1 : ATTRIBUTE9;  // hitproxy.r + 256 * selected in .w
 	half4 InstanceTransform2 : ATTRIBUTE10; // hitproxy.g in .w
@@ -231,7 +233,7 @@ struct FPositionOnlyVertexFactoryInput
 {
 	float4	Position	: ATTRIBUTE0;
 
-#if USE_INSTANCING && !MANUAL_VERTEX_FETCH
+#if (USE_INSTANCING || USE_INSTANCE_CULLING) && !MANUAL_VERTEX_FETCH
 	float4 InstanceOrigin : ATTRIBUTE8;  // per-instance random in w 
 	half4 InstanceTransform1 : ATTRIBUTE9;  // hitproxy.r + 256 * selected in .w
 	half4 InstanceTransform2 : ATTRIBUTE10; // hitproxy.g in .w
@@ -289,6 +291,17 @@ struct FVertexFactoryIntermediates
 	half TangentToWorldSign;
 
 	half4 Color;
+
+#if USE_INSTANCE_CULLING
+	float4 HitProxyId;
+	float PerInstanceRandom;
+	float IsSelected;
+	float IsVisible;
+	// x = per-instance random, y = per-instance fade out amount, z = hide/show flag, w dither fade cutoff
+	float4 PerInstanceParams;
+	float4 InstanceLightmapAndShadowMapUVBias;
+#endif
+
 #if USE_INSTANCING
 		float4 InstanceOrigin;
 		float4 InstanceTransform1;
@@ -443,7 +456,7 @@ FMaterialPixelParameters GetMaterialPixelParameters(FVertexFactoryInterpolantsVS
 
 	// Required for previewing materials that use ParticleColor
 	Result.Particle.Color = half4(1,1,1,1);
-#if USE_INSTANCING
+#if USE_INSTANCING || USE_INSTANCE_CULLING 
 	Result.PerInstanceParams = Interpolants.PerInstanceParams;
 #endif
 
@@ -1160,6 +1173,13 @@ float4 VertexFactoryGetPreviousWorldPosition(FVertexFactoryInput Input, FVertexF
 	}
 #endif // #if USING_TESSELLATION
 
+#if USE_INSTANCE_CULLING
+float4 VertexFactoryGetInstanceHitProxyId(FVertexFactoryInput Input, FVertexFactoryIntermediates Intermediates)
+{
+	return Intermediates.HitProxyId;
+}
+#endif
+
 #if USE_INSTANCING
 float4 VertexFactoryGetInstanceHitProxyId(FVertexFactoryInput Input, FVertexFactoryIntermediates Intermediates)
 {
diff --git a/Engine/Shaders/Private/LocalVertexFactoryCommon.ush b/Engine/Shaders/Private/LocalVertexFactoryCommon.ush
--- a/Engine/Shaders/Private/LocalVertexFactoryCommon.ush
+++ b/Engine/Shaders/Private/LocalVertexFactoryCommon.ush
@@ -12,7 +12,7 @@ struct FVertexFactoryInterpolantsVSToPS
 	half4	Color : COLOR0;
 #endif
 
-#if USE_INSTANCING
+#if USE_INSTANCING || USE_INSTANCE_CULLING
 	// x = per-instance random, y = per-instance fade out amount, z = hide/show flag, w dither fade cutoff
 	float4  PerInstanceParams : COLOR1;
 #endif
diff --git a/Engine/Shaders/Private/MaterialTemplate.ush b/Engine/Shaders/Private/MaterialTemplate.ush
--- a/Engine/Shaders/Private/MaterialTemplate.ush
+++ b/Engine/Shaders/Private/MaterialTemplate.ush
@@ -95,18 +95,22 @@
 #include "/Engine/Private/PaniniProjection.ush"
 #include "/Engine/Private/DBufferNormalReprojection.ush"
 
+#ifndef USE_INSTANCE_CULLING
+	#define USE_INSTANCE_CULLING 0
+#endif
+
 #ifndef USE_DITHERED_LOD_TRANSITION
-    #if USE_INSTANCING
-	    #ifndef USE_DITHERED_LOD_TRANSITION_FOR_INSTANCED
-		    #error "USE_DITHERED_LOD_TRANSITION_FOR_INSTANCED should have been defined"
-	    #endif
-	    #define USE_DITHERED_LOD_TRANSITION USE_DITHERED_LOD_TRANSITION_FOR_INSTANCED
-    #else
-	    #ifndef USE_DITHERED_LOD_TRANSITION_FROM_MATERIAL
-		    #error "USE_DITHERED_LOD_TRANSITION_FROM_MATERIAL should have been defined"
-	    #endif
-	    #define USE_DITHERED_LOD_TRANSITION USE_DITHERED_LOD_TRANSITION_FROM_MATERIAL
-    #endif
+	#if USE_INSTANCING || USE_INSTANCE_CULLING
+		#ifndef USE_DITHERED_LOD_TRANSITION_FOR_INSTANCED
+			#error "USE_DITHERED_LOD_TRANSITION_FOR_INSTANCED should have been defined"
+		#endif
+		#define USE_DITHERED_LOD_TRANSITION USE_DITHERED_LOD_TRANSITION_FOR_INSTANCED
+	#else
+		#ifndef USE_DITHERED_LOD_TRANSITION_FROM_MATERIAL
+			#error "USE_DITHERED_LOD_TRANSITION_FROM_MATERIAL should have been defined"
+		#endif
+		#define USE_DITHERED_LOD_TRANSITION USE_DITHERED_LOD_TRANSITION_FROM_MATERIAL
+	#endif
 #endif
 
 #ifndef USE_STENCIL_LOD_DITHER
@@ -316,7 +320,7 @@ struct FMaterialPixelParameters
 	float2	LightmapUVs;
 #endif
 
-#if USE_INSTANCING || VF_NANITE
+#if USE_INSTANCING || USE_INSTANCE_CULLING || VF_NANITE
 	half4 PerInstanceParams;
 #endif
 
@@ -428,12 +432,16 @@ struct FMaterialVertexParameters
 	float3 WorldPosition;
 	// TangentToWorld[2] is WorldVertexNormal
 	half3x3 TangentToWorld;
-#if USE_INSTANCING
+#if USE_INSTANCING || USE_INSTANCE_CULLING
 	/** Per-instance properties. */
 	float4x4 InstanceLocalToWorld;
 	float3 InstanceLocalPosition;
 	float4 PerInstanceParams;
+#if USE_INSTANCE_CULLING
+	int PrimitiveLocalInstanceIndex;
+#else
 	uint InstanceId;
+#endif
 	uint InstanceOffset;
 
 #elif IS_MESHPARTICLE_FACTORY 
@@ -851,10 +859,16 @@ MaterialFloat2 GetDefaultSceneTextureUV(FMaterialPixelParameters Parameters, con
 /** Get the per-instance custom data when instancing */
 float GetPerInstanceCustomData(FMaterialVertexParameters Parameters, int Index, float DefaultValue)
 {
-#if USE_INSTANCING && USES_PER_INSTANCE_CUSTOM_DATA
-	if(InstanceVF.NumCustomDataFloats > 0)
+#if (USE_INSTANCING || USE_INSTANCE_CULLING) && USES_PER_INSTANCE_CUSTOM_DATA
+	if (InstanceVF.NumCustomDataFloats > 0)
 	{
-		int BufferStartIndex = (Parameters.InstanceId + Parameters.InstanceOffset) * InstanceVF.NumCustomDataFloats; 
+	#if USE_INSTANCE_CULLING
+		int InstanceDataIndex = Parameters.PrimitiveLocalInstanceIndex;
+	#else
+		int InstanceDataIndex = Parameters.InstanceId + Parameters.InstanceOffset;
+	#endif
+
+		int BufferStartIndex = InstanceDataIndex * InstanceVF.NumCustomDataFloats;
 		int FloatIndex = clamp(Index, 0, InstanceVF.NumCustomDataFloats - 1);
 		return InstanceVF.InstanceCustomDataBuffer[BufferStartIndex + FloatIndex];
 	}
@@ -892,7 +906,7 @@ MaterialFloat3 TransformTangentVectorToView(FMaterialPixelParameters Parameters,
 /** Transforms a vector from local space to world space (VS version) */
 MaterialFloat3 TransformLocalVectorToWorld(FMaterialVertexParameters Parameters,MaterialFloat3 InLocalVector)
 {
-	#if USE_INSTANCING || IS_MESHPARTICLE_FACTORY
+	#if USE_INSTANCING || USE_INSTANCE_CULLING || IS_MESHPARTICLE_FACTORY
 		return mul(InLocalVector, (MaterialFloat3x3)Parameters.InstanceLocalToWorld);
 	#else
 		return mul(InLocalVector, GetLocalToWorld3x3(Parameters.PrimitiveId));
@@ -922,7 +936,7 @@ float3 TransformLocalPositionToWorld(FMaterialPixelParameters Parameters,float3
 /** Transforms a position from local space to absolute world space */
 float3 TransformLocalPositionToWorld(FMaterialVertexParameters Parameters,float3 InLocalPosition)
 {
-	#if USE_INSTANCING || IS_MESHPARTICLE_FACTORY
+	#if USE_INSTANCING || USE_INSTANCE_CULLING || IS_MESHPARTICLE_FACTORY
 		return mul(float4(InLocalPosition, 1), Parameters.InstanceLocalToWorld).xyz;
 	#else
 		return mul(float4(InLocalPosition, 1), GetPrimitiveData(Parameters.PrimitiveId).LocalToWorld).xyz;
@@ -953,7 +967,7 @@ float3 GetObjectWorldPosition(FMaterialTessellationParameters Parameters)
 /** Return the object's position in world space. For instanced meshes, this returns the instance position. */
 float3 GetObjectWorldPosition(FMaterialVertexParameters Parameters)
 {
-	#if USE_INSTANCING || IS_MESHPARTICLE_FACTORY
+	#if USE_INSTANCING || USE_INSTANCE_CULLING || IS_MESHPARTICLE_FACTORY
 		return Parameters.InstanceLocalToWorld[3].xyz;
 	#else
 		return GetPrimitiveData(Parameters.PrimitiveId).ObjectWorldPositionAndRadius.xyz;
@@ -965,7 +979,7 @@ float3 GetObjectWorldPosition(FMaterialVertexParameters Parameters)
 /** Get the per-instance random value when instancing */
 float GetPerInstanceRandom(FMaterialVertexParameters Parameters)
 {
-#if USE_INSTANCING
+#if USE_INSTANCING || USE_INSTANCE_CULLING
 	return Parameters.PerInstanceParams.x;
 #else
 	return 0.0;
@@ -975,7 +989,7 @@ float GetPerInstanceRandom(FMaterialVertexParameters Parameters)
 /** Get the per-instance random value when instancing */
 float GetPerInstanceRandom(FMaterialPixelParameters Parameters)
 {
-#if USE_INSTANCING || VF_NANITE
+#if USE_INSTANCING || USE_INSTANCE_CULLING || VF_NANITE
 	return Parameters.PerInstanceParams.x;
 #else
 	return 0.0;
@@ -985,7 +999,7 @@ float GetPerInstanceRandom(FMaterialPixelParameters Parameters)
 /** Get the per-instance fade-out amount when instancing */
 float GetPerInstanceFadeAmount(FMaterialPixelParameters Parameters)
 {
-#if USE_INSTANCING
+#if USE_INSTANCING || USE_INSTANCE_CULLING
 	return float(Parameters.PerInstanceParams.y);
 #else
 	return float(1.0);
@@ -995,7 +1009,7 @@ float GetPerInstanceFadeAmount(FMaterialPixelParameters Parameters)
 /** Get the per-instance fade-out amount when instancing */
 float GetPerInstanceFadeAmount(FMaterialVertexParameters Parameters)
 {
-#if USE_INSTANCING
+#if USE_INSTANCING || USE_INSTANCE_CULLING
 	return float(Parameters.PerInstanceParams.y);
 #else
 	return float(1.0);
@@ -2197,7 +2211,7 @@ half GetMaterialMaskedOpacity(FPixelMaterialInputs PixelMaterialInputs)
 
 float3 GetMaterialWorldPositionOffset(FMaterialVertexParameters Parameters)
 {
-	#if USE_INSTANCING
+	#if USE_INSTANCING || USE_INSTANCE_CULLING
 		// skip if this instance is hidden
 		if (Parameters.PerInstanceParams.z < 1.f)
 		{
@@ -2209,7 +2223,7 @@ float3 GetMaterialWorldPositionOffset(FMaterialVertexParameters Parameters)
 
 float3 GetMaterialPreviousWorldPositionOffset(FMaterialVertexParameters Parameters)
 {
-	#if USE_INSTANCING
+	#if USE_INSTANCING || USE_INSTANCE_CULLING
 		// skip if this instance is hidden
 		if (Parameters.PerInstanceParams.z < 1.f)
 		{
@@ -2428,9 +2442,9 @@ void ClipLODTransition(FMaterialPixelParameters Parameters, float DitherFactor)
 }
 
 
-#define REQUIRES_VF_ATTRIBUTES_FOR_CLIPPING (USE_INSTANCING && USE_DITHERED_LOD_TRANSITION)
+#define REQUIRES_VF_ATTRIBUTES_FOR_CLIPPING ((USE_INSTANCING || USE_INSTANCE_CULLING) && USE_DITHERED_LOD_TRANSITION)
 
-#if USE_INSTANCING && USE_DITHERED_LOD_TRANSITION
+#if (USE_INSTANCING  || USE_INSTANCE_CULLING) && USE_DITHERED_LOD_TRANSITION
 void ClipLODTransition(FMaterialPixelParameters Parameters)
 {
 	ClipLODTransition(Parameters, Parameters.PerInstanceParams.w);
diff --git a/Engine/Shaders/Private/ShadowDepthPixelShader.usf b/Engine/Shaders/Private/ShadowDepthPixelShader.usf
--- a/Engine/Shaders/Private/ShadowDepthPixelShader.usf
+++ b/Engine/Shaders/Private/ShadowDepthPixelShader.usf
@@ -8,6 +8,7 @@
 #define SHADOW_DEPTH_SHADER 1
 #define USE_STENCIL_LOD_DITHER	0
 
+#include "Nanite/PackedNaniteView.ush"
 #include "Common.ush"
 
 // Reroute SceneTexturesStruct uniform buffer references to the shadow depth pass uniform buffer
