Disable VF position-only support for hair Cards/Strands.
- was not implemented fully and caused overloading issues.

#rnx

diff --git a/Engine/Plugins/Runtime/HairStrands/Source/HairStrandsCore/Private/HairCardsVertexFactory.cpp b/Engine/Plugins/Runtime/HairStrands/Source/HairStrandsCore/Private/HairCardsVertexFactory.cpp
--- a/Engine/Plugins/Runtime/HairStrands/Source/HairStrandsCore/Private/HairCardsVertexFactory.cpp
+++ b/Engine/Plugins/Runtime/HairStrands/Source/HairStrandsCore/Private/HairCardsVertexFactory.cpp
@@ -321,4 +321,4 @@ void FHairCardsVertexFactory::ReleaseRHI()
 	FVertexFactory::ReleaseRHI();
 }
 
-IMPLEMENT_VERTEX_FACTORY_TYPE_EX(FHairCardsVertexFactory,"/Engine/Private/HairStrands/HairCardsVertexFactory.ush",true,false,true,true,true,true,VF_CARDS_SUPPORT_GPU_SCENE,false);
+IMPLEMENT_VERTEX_FACTORY_TYPE_EX(FHairCardsVertexFactory,"/Engine/Private/HairStrands/HairCardsVertexFactory.ush",true,false,true,true,false,true,VF_CARDS_SUPPORT_GPU_SCENE,false);
diff --git a/Engine/Plugins/Runtime/HairStrands/Source/HairStrandsCore/Private/HairStrandsVertexFactory.cpp b/Engine/Plugins/Runtime/HairStrands/Source/HairStrandsCore/Private/HairStrandsVertexFactory.cpp
--- a/Engine/Plugins/Runtime/HairStrands/Source/HairStrandsCore/Private/HairStrandsVertexFactory.cpp
+++ b/Engine/Plugins/Runtime/HairStrands/Source/HairStrandsCore/Private/HairStrandsVertexFactory.cpp
@@ -254,4 +254,4 @@ void FHairStrandsVertexFactory::ReleaseRHI()
 	FVertexFactory::ReleaseRHI();
 }
 
-IMPLEMENT_VERTEX_FACTORY_TYPE_EX(FHairStrandsVertexFactory,"/Engine/Private/HairStrands/HairStrandsVertexFactory.ush",true,false,true,true,true,true,true,false);
+IMPLEMENT_VERTEX_FACTORY_TYPE_EX(FHairStrandsVertexFactory,"/Engine/Private/HairStrands/HairStrandsVertexFactory.ush",true,false,true,true,false,true,true,false);
diff --git a/Engine/Shaders/Private/HairStrands/HairCardsVertexFactory.ush b/Engine/Shaders/Private/HairStrands/HairCardsVertexFactory.ush
--- a/Engine/Shaders/Private/HairStrands/HairCardsVertexFactory.ush
+++ b/Engine/Shaders/Private/HairStrands/HairCardsVertexFactory.ush
@@ -33,8 +33,6 @@
 /**
  * Per-vertex inputs from bound vertex buffers
  */
-#define FPositionOnlyVertexFactoryInput FVertexFactoryInput
-#define FPositionAndNormalOnlyVertexFactoryInput FVertexFactoryInput
 
 struct FVertexFactoryInput
 {
@@ -341,7 +339,7 @@ FVertexFactoryInterpolantsVSToPS VertexFactoryGetInterpolantsVSToPS(FVertexFacto
 }
 
 /** for depth-only pass (Not used by the actual hair shaders)*/
-float4 VertexFactoryGetWorldPosition(FPositionOnlyVertexFactoryInput Input)
+float4 VertexFactoryGetWorldPosition(FVertexFactoryInput Input)
 {
 	return 0;
 }
@@ -375,7 +373,7 @@ uint VertexFactoryGetPrimitiveId(FVertexFactoryInterpolantsVSToPS Interpolants)
 	return GetPrimitiveId(Interpolants);
 }
 
-float3 VertexFactoryGetWorldNormal(FPositionAndNormalOnlyVertexFactoryInput Input)
+float3 VertexFactoryGetWorldNormal(FVertexFactoryInput Input)
 {
 	return float3(0, 0, 1); 
 }
diff --git a/Engine/Shaders/Private/HairStrands/HairStrandsVertexFactory.ush b/Engine/Shaders/Private/HairStrands/HairStrandsVertexFactory.ush
--- a/Engine/Shaders/Private/HairStrands/HairStrandsVertexFactory.ush
+++ b/Engine/Shaders/Private/HairStrands/HairStrandsVertexFactory.ush
@@ -84,8 +84,6 @@ Buffer<float4> HairStrandsVF_TangentBuffer;
 /**
  * Per-vertex inputs from bound vertex buffers
  */
-#define FPositionOnlyVertexFactoryInput FVertexFactoryInput
-#define FPositionAndNormalOnlyVertexFactoryInput FVertexFactoryInput
 
 bool UseStableRasterization()
 {
@@ -498,7 +496,7 @@ FVertexFactoryInterpolantsVSToPS VertexFactoryGetInterpolantsVSToPS(FVertexFacto
 }
 
 /** for depth-only pass (Not used by the actual hair shaders)*/
-float4 VertexFactoryGetWorldPosition(FPositionOnlyVertexFactoryInput Input)
+float4 VertexFactoryGetWorldPosition(FVertexFactoryInput Input)
 {
 	return 0;
 }
@@ -532,7 +530,7 @@ uint VertexFactoryGetPrimitiveId(FVertexFactoryInterpolantsVSToPS Interpolants)
 	return GetPrimitiveId(Interpolants);
 }
 
-float3 VertexFactoryGetWorldNormal(FPositionAndNormalOnlyVertexFactoryInput Input)
+float3 VertexFactoryGetWorldNormal(FVertexFactoryInput Input)
 {
 	return float3(0, 0, 1); 
 }
