Fixed SSS albedo component selection based on MaxDmfp.

diff --git a/Engine/Shaders/Private/SubsurfaceBurleyNormalized.ush b/Engine/Shaders/Private/SubsurfaceBurleyNormalized.ush
--- a/Engine/Shaders/Private/SubsurfaceBurleyNormalized.ush
+++ b/Engine/Shaders/Private/SubsurfaceBurleyNormalized.ush
@@ -215,11 +215,11 @@ FBurleyParameter GetBurleyParameters(FGBufferData GBuffer)
 		{
 			// Special data encoding from strata packing
 			const float3 MFPAlebdo			= Pow2(GBuffer.CustomData.xyz);
-			const float  MFPRadius			= GBuffer.CustomData.a * SSS_PROFILE_PERPIXEL_MAXRADIUS * BURLEY_CM_2_MM * 5.0f; // STRATA_TODO why a scale of 5 is needed here?
+			const float  MFPRadius			= GBuffer.CustomData.a * SSS_PROFILE_PERPIXEL_MAXRADIUS * BURLEY_CM_2_MM * 5.0f; // STRATA_TODO why a scale of 5 is needed here to match between regular and Strata?
 			const float3 MFPChromaticRadius = MFPAlebdo * MFPRadius;
 
 			const float MaxDMFP				= max3(MFPChromaticRadius.x, MFPChromaticRadius.y, MFPChromaticRadius.z);
-			const uint IndexOfMaxDmfp		= (MFPChromaticRadius.r == MaxDMFP) ? 0 : ((MFPChromaticRadius.b == MaxDMFP) ? 1 : 2);// STARTA_TODO fix the component selection in SubsirfaceProfile, b should map to 2
+			const uint IndexOfMaxDmfp		= (MFPChromaticRadius.r == MaxDMFP) ? 0 : ((MFPChromaticRadius.g == MaxDMFP) ? 1 : 2);
 
 			Out.SurfaceAlbedo				= float4(GBuffer.BaseColor, GBuffer.BaseColor[IndexOfMaxDmfp]);
 			Out.DiffuseMeanFreePath			= float4(MFPChromaticRadius, MaxDMFP);
diff --git a/Engine/Source/Runtime/Engine/Private/Rendering/SubsurfaceProfile.cpp b/Engine/Source/Runtime/Engine/Private/Rendering/SubsurfaceProfile.cpp
--- a/Engine/Source/Runtime/Engine/Private/Rendering/SubsurfaceProfile.cpp
+++ b/Engine/Source/Runtime/Engine/Private/Rendering/SubsurfaceProfile.cpp
@@ -190,7 +190,7 @@ void SetupSurfaceAlbedoAndDiffuseMeanFreePath(FLinearColor& SurfaceAlbedo, FLine
 	//This is an optimization to shift finding the max correspondence workload
 	//to CPU.
 	const float MaxDmfpComp = FMath::Max3(Dmfp.R, Dmfp.G, Dmfp.B);
-	const uint32 IndexOfMaxDmfp = (Dmfp.R == MaxDmfpComp) ? 0 : ((Dmfp.B == MaxDmfpComp) ? 1 : 2);
+	const uint32 IndexOfMaxDmfp = (Dmfp.R == MaxDmfpComp) ? 0 : ((Dmfp.G == MaxDmfpComp) ? 1 : 2);
 
 	SurfaceAlbedo.A = SurfaceAlbedo.Component(IndexOfMaxDmfp);
 	Dmfp.A = MaxDmfpComp;
