Added back the radius scale needed to control SSS profile asset radius with Strata.

diff --git a/Engine/Shaders/Private/Strata/Strata.ush b/Engine/Shaders/Private/Strata/Strata.ush
--- a/Engine/Shaders/Private/Strata/Strata.ush
+++ b/Engine/Shaders/Private/Strata/Strata.ush
@@ -119,7 +119,7 @@ struct FStrataBSDF
 #define SLAB_ROUGHNESSY(X)				X.VGPRs[1].z
 // Either with SSS profile
 #define SLAB_SSSPROFILEID(X)			X.VGPRs[2].x
-#define SLAB_SSSPROFILERADIUS(X)		X.VGPRs[2].y
+#define SLAB_SSSPROFILERADIUSSCALE(X)	X.VGPRs[2].y
 // Or explicit DMFP
 #define SLAB_SSSDMFPALBEDO(X)			X.VGPRs[2].xyz
 #define SLAB_SSSDMFPRADIUS(X)			X.VGPRs[2].w
@@ -207,10 +207,11 @@ struct FStrataSubsurfaceData
 {
 	bool   bHasSubsurface;
 	bool   bHasProfile;
-	float  RadiusScale;
 
 	float3 BaseColor;
-	float  Profile;		// Integer ID packed in 8bits as float
+
+	float  ProfilePackedId;		// Integer ID packed in 8bits as float
+	float  ProfileRadiusScale;
 
 	float3 DMFPAlbedo;
 	float  DMFPRadius;
@@ -305,7 +306,7 @@ FStrataData GetInitialisedStrataData()
 
 FStrataData GetStrataSlabBSDF(
 	float3 BaseColor, float Specular, float Metallic, float2 Roughness,
-	float SSSProfileID, float3 SSSDMFPAlbedo, float SSSDMFPRadius,
+	float SSSProfileID, float3 SSSDMFPAlbedo, float SSSDMFPRadius, float SSSRadiusScale,
 	uint SharedNormalIndex, inout uint SharedNormalTypes)
 {
 	FStrataData StrataData = GetInitialisedStrataData();
@@ -337,12 +338,12 @@ FStrataData GetStrataSlabBSDF(
 	if (bHasSSSProfile)
 	{
 		SLAB_SSSPROFILEID	(StrataData.Layers[0].BSDFs[0])		= SSSProfileID;
-		SLAB_SSSPROFILERADIUS(StrataData.Layers[0].BSDFs[0])	= 1.0; // STRATA_TODO add a scale to control the radius per pixel
+		SLAB_SSSPROFILERADIUSSCALE(StrataData.Layers[0].BSDFs[0])= saturate(SSSRadiusScale);
 	}
 	else
 	{
 		SLAB_SSSDMFPALBEDO	(StrataData.Layers[0].BSDFs[0])		= saturate(SSSDMFPAlbedo);
-		SLAB_SSSDMFPRADIUS	(StrataData.Layers[0].BSDFs[0])		= clamp(SSSDMFPRadius, 0.f, SSS_PROFILE_PERPIXEL_MAXRADIUS);
+		SLAB_SSSDMFPRADIUS	(StrataData.Layers[0].BSDFs[0])		= clamp(SSSDMFPRadius * saturate(SSSRadiusScale), 0.f, SSS_PROFILE_PERPIXEL_MAXRADIUS);
 	}
 
 	StrataData.Layers[0].BSDFs[0].Weight						= 1.0f;
@@ -626,7 +627,7 @@ void SanitizeStrataSlab(inout FStrataBSDF BSDF)
 	SLAB_METALLIC(BSDF) = saturate(SLAB_METALLIC(BSDF));
 	SLAB_ROUGHNESSX(BSDF) = saturate(SLAB_ROUGHNESSX(BSDF));
 	SLAB_ROUGHNESSY(BSDF) = saturate(SLAB_ROUGHNESSY(BSDF));
-	SLAB_SSSPROFILERADIUS(BSDF) = max(0.0f, SLAB_SSSPROFILERADIUS(BSDF));
+	SLAB_SSSPROFILERADIUSSCALE(BSDF) = saturate(SLAB_SSSPROFILERADIUSSCALE(BSDF));
 	SLAB_SSSDMFPALBEDO(BSDF) = saturate(SLAB_SSSDMFPALBEDO(BSDF));
 	SLAB_SSSDMFPRADIUS(BSDF) = clamp(SLAB_SSSDMFPRADIUS(BSDF), 0.0f, SSS_PROFILE_PERPIXEL_MAXRADIUS);
 }
@@ -1085,12 +1086,12 @@ void PackStrataOut(
 				if (BSDF_GETHASSCATTERING(BSDF) && BSDFCoverageAvg > 0 && !SSS.bHasSubsurface)
 				{
 					SSS.bHasSubsurface = true;
-					SSS.bHasProfile	= BSDF_GETHASSSSPROFILE(BSDF);
-					SSS.Profile		= SLAB_SSSPROFILEID(BSDF);
-					SSS.RadiusScale	= SLAB_SSSPROFILERADIUS(BSDF);
-					SSS.BaseColor	= SLAB_BASECOLOR(BSDF);
-					SSS.DMFPAlbedo	= SLAB_SSSDMFPALBEDO(BSDF);
-					SSS.DMFPRadius	= SLAB_SSSDMFPRADIUS(BSDF);
+					SSS.bHasProfile			= BSDF_GETHASSSSPROFILE(BSDF);
+					SSS.ProfilePackedId		= SLAB_SSSPROFILEID(BSDF);
+					SSS.ProfileRadiusScale	= SLAB_SSSPROFILERADIUSSCALE(BSDF);
+					SSS.BaseColor			= SLAB_BASECOLOR(BSDF);
+					SSS.DMFPAlbedo			= SLAB_SSSDMFPALBEDO(BSDF);
+					SSS.DMFPRadius			= SLAB_SSSDMFPRADIUS(BSDF);
 				}
 
 				LayerCoverage	+= BSDFCoverage;
@@ -1192,8 +1193,8 @@ void PackStrataOut(
 		if (SSS.bHasProfile)
 		{
 			GBuffer.BaseColor      = SSS.BaseColor;
-			GBuffer.CustomData.rgb = EncodeSubsurfaceProfile(SSS.Profile);
-			GBuffer.CustomData.a   = 1.0f; // STRATA_TODO add a radius scale
+			GBuffer.CustomData.rgb = EncodeSubsurfaceProfile(SSS.ProfilePackedId);
+			GBuffer.CustomData.a   = SSS.ProfileRadiusScale;
 		}
 		else
 		{
@@ -1260,7 +1261,7 @@ void PackStrataOut(
 				{
 					if (BSDF_GETHASSSSPROFILE(BSDF))
 					{
-						STRATA_STORE_UINT1(PackRGBA8(float4(SLAB_SSSPROFILEID(BSDF), SLAB_SSSPROFILERADIUS(BSDF), 0, 0)));
+						STRATA_STORE_UINT1(PackRGBA8(float4(SLAB_SSSPROFILEID(BSDF), SLAB_SSSPROFILERADIUSSCALE(BSDF), 0, 0)));
 					}
 					else
 					{
@@ -1372,7 +1373,7 @@ FStrataBSDF UnpackStrataBSDFIn(ByteAddressBuffer StrataBuffer, inout uint ByteOf
 			{
 				float4 Data = UnpackRGBA8(StrataLoadUint1(StrataBuffer, ByteOffset));
 				SLAB_SSSPROFILEID(OutBSDF) = Data.r;
-				SLAB_SSSPROFILERADIUS(OutBSDF) = Data.g;
+				SLAB_SSSPROFILERADIUSSCALE(OutBSDF) = Data.g;
 			}
 			else
 			{
diff --git a/Engine/Shaders/Private/Strata/StrataMaterialClassification.usf b/Engine/Shaders/Private/Strata/StrataMaterialClassification.usf
--- a/Engine/Shaders/Private/Strata/StrataMaterialClassification.usf
+++ b/Engine/Shaders/Private/Strata/StrataMaterialClassification.usf
@@ -59,12 +59,12 @@ void ComputePostProcessBSDFData(
 		if (BSDF_GETHASSCATTERING(BSDF) && BSDFCoverageAvg > 0 && !SSS.bHasSubsurface)
 		{
 			SSS.bHasSubsurface = true;
-			SSS.bHasProfile	= BSDF_GETHASSSSPROFILE(BSDF);
-			SSS.Profile		= SLAB_SSSPROFILEID(BSDF);
-			SSS.RadiusScale	= SLAB_SSSPROFILERADIUS(BSDF);
-			SSS.BaseColor	= SLAB_BASECOLOR(BSDF);
-			SSS.DMFPAlbedo	= SLAB_SSSDMFPALBEDO(BSDF);
-			SSS.DMFPRadius	= SLAB_SSSDMFPRADIUS(BSDF);
+			SSS.bHasProfile			= BSDF_GETHASSSSPROFILE(BSDF);
+			SSS.ProfilePackedId		= SLAB_SSSPROFILEID(BSDF);
+			SSS.ProfileRadiusScale	= SLAB_SSSPROFILERADIUSSCALE(BSDF);
+			SSS.BaseColor			= SLAB_BASECOLOR(BSDF);
+			SSS.DMFPAlbedo			= SLAB_SSSDMFPALBEDO(BSDF);
+			SSS.DMFPRadius			= SLAB_SSSDMFPRADIUS(BSDF);
 		}
 		break;
 	}
diff --git a/Engine/Shaders/Private/Strata/StrataVisualize.usf b/Engine/Shaders/Private/Strata/StrataVisualize.usf
--- a/Engine/Shaders/Private/Strata/StrataVisualize.usf
+++ b/Engine/Shaders/Private/Strata/StrataVisualize.usf
@@ -260,8 +260,8 @@ void VisualizeMaterialPS(
 						HORIZONTALTAB; PRINTPARAMCHAR(_S_); PRINTPARAMCHAR(_S_); PRINTPARAMCHAR(_S_); PRINTPARAMCHAR(_SPC_); PRINTPARAMCHAR(_I_); PRINTPARAMCHAR(_D_); PRINTPARAMCHAR(_SPC_);
 						HORIZONTALTAB; PRINTFLOAT1(SLAB_SSSPROFILEID(BSDF));
 						NEXTLINE;
-						HORIZONTALTAB; PRINTPARAMCHAR(_S_); PRINTPARAMCHAR(_S_); PRINTPARAMCHAR(_S_); PRINTPARAMCHAR(_SPC_); PRINTPARAMCHAR(_R_); PRINTPARAMCHAR(_A_); PRINTPARAMCHAR(_SPC_);
-						HORIZONTALTAB; PRINTFLOAT1(SLAB_SSSPROFILERADIUS(BSDF));
+						HORIZONTALTAB; PRINTPARAMCHAR(_S_); PRINTPARAMCHAR(_S_); PRINTPARAMCHAR(_S_); PRINTPARAMCHAR(_S_); PRINTPARAMCHAR(_C_); PRINTPARAMCHAR(_L_); PRINTPARAMCHAR(_SPC_);
+						HORIZONTALTAB; PRINTFLOAT1(SLAB_SSSPROFILERADIUSSCALE(BSDF));
 						NEXTLINE;
 					}
 					else
diff --git a/Engine/Shaders/Private/SubsurfaceBurleyNormalized.ush b/Engine/Shaders/Private/SubsurfaceBurleyNormalized.ush
--- a/Engine/Shaders/Private/SubsurfaceBurleyNormalized.ush
+++ b/Engine/Shaders/Private/SubsurfaceBurleyNormalized.ush
@@ -208,7 +208,7 @@ FBurleyParameter GetBurleyParameters(FGBufferData GBuffer)
 			Out.DiffuseMeanFreePath = DecodeDiffuseMeanFreePath(GetSubsurfaceProfileDiffuseMeanFreePath(SubsurfaceProfileInt));
 			Out.WorldUnitScale		= DecodeWorldUnitScale(GetSubsurfaceProfileWorldUnitScale(SubsurfaceProfileInt)) * BURLEY_CM_2_MM;
 
-			Out.SurfaceOpacity		= 1.0f;
+			Out.SurfaceOpacity		= GBuffer.CustomData.a;
 			Out.DiffuseMeanFreePath*= Out.SurfaceOpacity;
 		}
 		else
@@ -225,8 +225,8 @@ FBurleyParameter GetBurleyParameters(FGBufferData GBuffer)
 			Out.DiffuseMeanFreePath			= float4(MFPChromaticRadius, MaxDMFP);
 			Out.WorldUnitScale				= DecodeWorldUnitScale(1.0f * ENC_WORLDUNITSCALE_IN_CM_TO_UNIT) * BURLEY_CM_2_MM;
 
-			// STRATA_TODO smooth transistion based on pixel footprint w.r.t. MFPRadius
-			Out.SurfaceOpacity				= 1.0;
+			// STRATA_TODO smooth transistion based on pixel footprint w.r.t. MFPRadius and pixel footprint
+			Out.SurfaceOpacity				= 1.0f;
 			Out.DiffuseMeanFreePath			*= Out.SurfaceOpacity;
 		}
 	}
diff --git a/Engine/Source/Runtime/Engine/Classes/Materials/MaterialExpressionStrata.h b/Engine/Source/Runtime/Engine/Classes/Materials/MaterialExpressionStrata.h
--- a/Engine/Source/Runtime/Engine/Classes/Materials/MaterialExpressionStrata.h
+++ b/Engine/Source/Runtime/Engine/Classes/Materials/MaterialExpressionStrata.h
@@ -36,7 +36,7 @@ class UMaterialExpressionStrataBSDF : public UMaterialExpression
 }*/
 
 UCLASS(MinimalAPI, collapsecategories, hidecategories = Object)
-class UMaterialExpressionStrataSlabBSDF : public UMaterialExpression // STRATA_TODO the single diffuse model to keep when we remove al lthe tests
+class UMaterialExpressionStrataSlabBSDF : public UMaterialExpression
 {
 	GENERATED_UCLASS_BODY()
 
@@ -85,17 +85,23 @@ class UMaterialExpressionStrataSlabBSDF : public UMaterialExpression // STRATA_T
 	FExpressionInput Tangent;
 
 	/**
-	 * Mean free path albedo. Only used when no sub-surface profile is provided. (type = float3, unit = unitless)
+	 * Mean free path albedo. Only used when there is no sub-surface profile provided. (type = float3, unit = unitless)
 	 */
 	UPROPERTY(meta = (DisplayName = "DMFP Albedo"))
 	FExpressionInput DiffuseMeanFreePathAlbedo;
 
 	/**
-	 * Mean free path radius. Only used when no sub-surface profile is provided. (type = float, unit = cm)
+	 * Mean free path radius. Only used when there is no sub-surface profile provided. (type = float, unit = cm)
 	 */
 	UPROPERTY(meta = (DisplayName = "DMFP Radius"))
 	FExpressionInput DiffuseMeanFreePathRadius;
 
+	/**
+	 * Scale the mean free path radius. Always used, when a subsurface profile is provided or not. (type = float, unitless, defaults to 1)
+	 */
+	UPROPERTY(meta = (DisplayName = "SSS Radius Scale"))
+	FExpressionInput SSSRadiusScale;
+
 	/** SubsurfaceProfile, for Screen Space Subsurface Scattering. The profile needs to be set up on both the Strata diffuse node, and the material node at the moment. */
 	UPROPERTY(EditAnywhere, BlueprintReadOnly, Category = Material, meta = (DisplayName = "Subsurface Profile"))
 	class USubsurfaceProfile* SubsurfaceProfile;
diff --git a/Engine/Source/Runtime/Engine/Private/Materials/HLSLMaterialTranslator.cpp b/Engine/Source/Runtime/Engine/Private/Materials/HLSLMaterialTranslator.cpp
--- a/Engine/Source/Runtime/Engine/Private/Materials/HLSLMaterialTranslator.cpp
+++ b/Engine/Source/Runtime/Engine/Private/Materials/HLSLMaterialTranslator.cpp
@@ -7536,10 +7536,10 @@ int32 FHLSLMaterialTranslator::StrataCreateAndRegisterNullMaterial()
 	return OutputCodeChunk;
 }
 
-int32 FHLSLMaterialTranslator::StrataSlabBSDF(int32 BaseColor, int32 Reflectivity, int32 Metallic, int32 RoughnessX, int32 RoughnessY, int32 SSSProfileId, int32 SSSDMFPAlbedo, int32 SSSDMFPRadius, int32 Normal, int32 Tangent, uint8 SharedNormalIndex)
+int32 FHLSLMaterialTranslator::StrataSlabBSDF(int32 BaseColor, int32 Reflectivity, int32 Metallic, int32 RoughnessX, int32 RoughnessY, int32 SSSProfileId, int32 SSSDMFPAlbedo, int32 SSSDMFPRadius, int32 SSSRadiusScale, int32 Normal, int32 Tangent, uint8 SharedNormalIndex)
 {
 	return AddCodeChunk(
-		MCT_Strata, TEXT("GetStrataSlabBSDF(%s, %s, %s, float2(%s, %s), %s, %s, %s, %u, Parameters.SharedNormals.NormalTypes) /* Normal = %s ; Tangent = %s */"),
+		MCT_Strata, TEXT("GetStrataSlabBSDF(%s, %s, %s, float2(%s, %s), %s, %s, %s, %s, %u, Parameters.SharedNormals.NormalTypes) /* Normal = %s ; Tangent = %s */"),
 		*GetParameterCode(BaseColor),
 		*GetParameterCode(Reflectivity),
 		*GetParameterCode(Metallic),
@@ -7548,6 +7548,7 @@ int32 FHLSLMaterialTranslator::StrataSlabBSDF(int32 BaseColor, int32 Reflectivit
 		*GetParameterCode(SSSProfileId),
 		*GetParameterCode(SSSDMFPAlbedo),
 		*GetParameterCode(SSSDMFPRadius),
+		*GetParameterCode(SSSRadiusScale),
 		SharedNormalIndex,
 		*GetParameterCode(Normal),
 		*GetParameterCode(Tangent)
diff --git a/Engine/Source/Runtime/Engine/Private/Materials/HLSLMaterialTranslator.h b/Engine/Source/Runtime/Engine/Private/Materials/HLSLMaterialTranslator.h
--- a/Engine/Source/Runtime/Engine/Private/Materials/HLSLMaterialTranslator.h
+++ b/Engine/Source/Runtime/Engine/Private/Materials/HLSLMaterialTranslator.h
@@ -798,7 +798,7 @@ protected:
 
 	// Strata
 	virtual int32 StrataCreateAndRegisterNullMaterial() override;
-	virtual int32 StrataSlabBSDF(int32 BaseColor, int32 Reflectivity, int32 Metallic, int32 RoughnessX, int32 RoughnessY, int32 SSSProfileId, int32 SSSDMFPAlbedo, int32 SSSDMFPRadius, int32 Normal, int32 Tangent, uint8 SharedNormalIndex) override;
+	virtual int32 StrataSlabBSDF(int32 BaseColor, int32 Reflectivity, int32 Metallic, int32 RoughnessX, int32 RoughnessY, int32 SSSProfileId, int32 SSSDMFPAlbedo, int32 SSSDMFPRadius, int32 SSSRadiusScale, int32 Normal, int32 Tangent, uint8 SharedNormalIndex) override;
 	virtual int32 StrataSheenBSDF(int32 BaseColor, int32 Roughness, int32 Normal, uint8 SharedNormalIndex) override;
 	virtual int32 StrataVolumetricFogCloudBSDF(int32 Albedo, int32 Extinction, int32 Emissive, int32 AmbientOcclusion) override;
 	virtual int32 StrataUnlitBSDF(int32 Emissive, int32 TransmittanceColor) override;
diff --git a/Engine/Source/Runtime/Engine/Private/Materials/MaterialExpressions.cpp b/Engine/Source/Runtime/Engine/Private/Materials/MaterialExpressions.cpp
--- a/Engine/Source/Runtime/Engine/Private/Materials/MaterialExpressions.cpp
+++ b/Engine/Source/Runtime/Engine/Private/Materials/MaterialExpressions.cpp
@@ -19690,6 +19690,7 @@ int32 UMaterialExpressionStrataSlabBSDF::Compile(class FMaterialCompiler* Compil
 		SSSProfileCodeChunk != INDEX_NONE ? SSSProfileCodeChunk : Compiler->Constant(0.0f),	
 		CompileWithDefaultFloat3(Compiler, DiffuseMeanFreePathAlbedo, 0.0f, 0.0f, 0.0f),
 		CompileWithDefaultFloat1(Compiler, DiffuseMeanFreePathRadius, 0.0f),
+		CompileWithDefaultFloat1(Compiler, SSSRadiusScale, 1.0f),
 		NormalCodeChunk,
 		TangentCodeChunk,
 		SharedNormalIndex);
@@ -19739,6 +19740,9 @@ uint32 UMaterialExpressionStrataSlabBSDF::GetInputType(int32 InputIndex)
 	case 8:
 		return MCT_Float1; // DMFP Radius
 		break;
+	case 9:
+		return MCT_Float1; // SSS Radius Scale
+		break;
 	}
 
 	check(false);
@@ -19783,6 +19787,10 @@ FName UMaterialExpressionStrataSlabBSDF::GetInputName(int32 InputIndex) const
 	{
 		return TEXT("DMFP Radius");
 	}
+	else if (InputIndex == 9)
+	{
+		return TEXT("SSS Radius Scale");
+	}
 	return TEXT("Unknown");
 }
 
diff --git a/Engine/Source/Runtime/Engine/Public/MaterialCompiler.h b/Engine/Source/Runtime/Engine/Public/MaterialCompiler.h
--- a/Engine/Source/Runtime/Engine/Public/MaterialCompiler.h
+++ b/Engine/Source/Runtime/Engine/Public/MaterialCompiler.h
@@ -392,7 +392,7 @@ public:
 
 	// Strata
 	virtual int32 StrataCreateAndRegisterNullMaterial() = 0;
-	virtual int32 StrataSlabBSDF(int32 BaseColor, int32 Reflectivity, int32 Metallic, int32 RoughnessX, int32 RoughnessY, int32 SSSProfileId, int32 SSSDMFPAlbedo, int32 SSSDMFPRadius, int32 Normal, int32 Tangent, uint8 SharedNormalIndex) = 0;
+	virtual int32 StrataSlabBSDF(int32 BaseColor, int32 Reflectivity, int32 Metallic, int32 RoughnessX, int32 RoughnessY, int32 SSSProfileId, int32 SSSDMFPAlbedo, int32 SSSDMFPRadius, int32 SSSRadiusScale, int32 Normal, int32 Tangent, uint8 SharedNormalIndex) = 0;
 	virtual int32 StrataSheenBSDF(int32 BaseColor, int32 Roughness, int32 Normal, uint8 SharedNormalIndex) = 0;
 	virtual int32 StrataVolumetricFogCloudBSDF(int32 Albedo, int32 Extinction, int32 Emissive, int32 AmbientOcclusion) = 0;
 	virtual int32 StrataUnlitBSDF(int32 Emissive, int32 TransmittanceColor) = 0;
@@ -808,9 +808,9 @@ public:
 		return Compiler->StrataCreateAndRegisterNullMaterial();
 	}
 
-	virtual int32 StrataSlabBSDF(int32 BaseColor, int32 Reflectivity, int32 Metallic, int32 RoughnessX, int32 RoughnessY, int32 SSSProfileId, int32 SSSDMFPAlbedo, int32 SSSDMFPRadius, int32 Normal, int32 Tangent, uint8 SharedNormalIndex) override
+	virtual int32 StrataSlabBSDF(int32 BaseColor, int32 Reflectivity, int32 Metallic, int32 RoughnessX, int32 RoughnessY, int32 SSSProfileId, int32 SSSDMFPAlbedo, int32 SSSDMFPRadius, int32 SSSRadiusScale, int32 Normal, int32 Tangent, uint8 SharedNormalIndex) override
 	{
-		return Compiler->StrataSlabBSDF(BaseColor, Reflectivity, Metallic, RoughnessX, RoughnessY, SSSProfileId, SSSDMFPAlbedo, SSSDMFPRadius, Normal, Tangent, SharedNormalIndex);
+		return Compiler->StrataSlabBSDF(BaseColor, Reflectivity, Metallic, RoughnessX, RoughnessY, SSSProfileId, SSSDMFPAlbedo, SSSDMFPRadius, SSSRadiusScale, Normal, Tangent, SharedNormalIndex);
 	}
 
 	virtual int32 StrataSheenBSDF(int32 BaseColor, int32 Roughness, int32 Normal, uint8 SharedNormalIndex) override
