Removed a Strata define no longer use and left behind.

diff --git a/Engine/Shaders/Private/Strata/StrataForwardLighting.ush b/Engine/Shaders/Private/Strata/StrataForwardLighting.ush
--- a/Engine/Shaders/Private/Strata/StrataForwardLighting.ush
+++ b/Engine/Shaders/Private/Strata/StrataForwardLighting.ush
@@ -54,7 +54,7 @@ float3 StrataForwardLighting(
 			{
 				// SkySHDiffuse is pre integrated according to a Lambert diffuse lobe. Thus, we need to rescale to account for the phase function operating over 4PI.
 				// We apply a 2PI factor to recover average luminance over half hemisphere, and then times the phase function to account for the scattering events.
-				const float VolumeOnlyFactor = BSDF_GETTYPE(BSDF) == STRATA_BSDF_TYPE_VOLUME ? 2.0f * PI * IsotropicPhase() : 1.0f;
+				const float VolumeOnlyFactor = 1.0f; // STRATA_TODO BSDF_GETTYPE(BSDF) == STRATA_BSDF_TYPE_VOLUME ? 2.0f * PI * IsotropicPhase() : 1.0f;
 
 				Color += GetEffectiveSkySHDiffuse(StrataEnvLight.DiffuseNormal) * StrataEnvLight.DiffuseColor * SkyLightCommonFactor * VolumeOnlyFactor;
 			}
