Final set of GPU-Scene Instance Culling v0.1 changes (disabled by default)
- LocalVertexFactory can fetch instance data from GPU-Scene
- regular Shadow map rendering can use instance culling view index (route to cube face)
- HISM/ISM/SM registers instance data
- Support for overriding mesh draw commands with indirect dispatch

diff --git a/Engine/Plugins/Runtime/HairStrands/Source/HairStrandsCore/Private/HairStrandsVertexFactory.cpp b/Engine/Plugins/Runtime/HairStrands/Source/HairStrandsCore/Private/HairStrandsVertexFactory.cpp
--- a/Engine/Plugins/Runtime/HairStrands/Source/HairStrandsCore/Private/HairStrandsVertexFactory.cpp
+++ b/Engine/Plugins/Runtime/HairStrands/Source/HairStrandsCore/Private/HairStrandsVertexFactory.cpp
@@ -231,7 +231,7 @@ void FHairStrandsVertexFactory::InitRHI()
 	if (GetType()->SupportsPrimitiveIdStream() && bCanUseGPUScene)
 	{
 		// When the VF is used for rendering in normal mesh passes, this vertex buffer and offset will be overridden
-		Elements.Add(AccessStreamComponent(FVertexStreamComponent(&GPrimitiveIdDummy, 0, 0, sizeof(uint32), VET_UInt, EVertexStreamUsage::Instancing), 13));
+		Elements.Add(AccessStreamComponent(FVertexStreamComponent(&GPrimitiveIdDummy, 0, 0, PrimitiveIdStreamStride, VET_UInt, EVertexStreamUsage::Instancing), 13));
 		SetPrimitiveIdStreamIndex(EVertexInputStreamType::Default, Elements.Last().StreamIndex);
 		bNeedsDeclaration = true;
 	}
diff --git a/Engine/Shaders/Private/BasePassVertexShader.usf b/Engine/Shaders/Private/BasePassVertexShader.usf
--- a/Engine/Shaders/Private/BasePassVertexShader.usf
+++ b/Engine/Shaders/Private/BasePassVertexShader.usf
@@ -220,6 +220,10 @@ void Main(
 
 #if WRITES_VELOCITY_TO_GBUFFER
 	{
+#if defined(GPUCULL_TODO)
+		float4 PrevTranslatedWorldPosition = WorldPosition;
+		// GPUCULL_TODO: This here is a non-abstracted use VFIntermediates, needs to be abstracted using accessor function. 
+#else // !defined(GPUCULL_TODO)
 		float4 PrevTranslatedWorldPosition = float4(0, 0, 0, 1);
 		BRANCH
 		if (GetPrimitiveData(VFIntermediates.PrimitiveId).OutputVelocity > 0 || View.ForceDrawAllVelocities != 0)
@@ -232,6 +236,7 @@ void Main(
 				PrevTranslatedWorldPosition = mul(float4(PrevTranslatedWorldPosition.xyz, 1), ResolvedView.PrevTranslatedWorldToClip);
 			#endif
 		}
+#endif //!defined(GPUCULL_TODO)
 
 		#if USING_TESSELLATION
 			// We let the Domain Shader convert to post projection when tessellating
diff --git a/Engine/Shaders/Private/HairStrands/HairStrandsMaterialCommon.ush b/Engine/Shaders/Private/HairStrands/HairStrandsMaterialCommon.ush
--- a/Engine/Shaders/Private/HairStrands/HairStrandsMaterialCommon.ush
+++ b/Engine/Shaders/Private/HairStrands/HairStrandsMaterialCommon.ush
@@ -53,7 +53,14 @@ FVertexFactoryInterpolantsVSToPS GetInterpolants(uint PrimitiveId, uint HairPrim
 	FMaterialVertexParameters Empty = (FMaterialVertexParameters)0;
 
 	FVertexFactoryInput Input;
+#if defined(GPUCULL_TODO)
+	// GPUCULL_TODO: this should be switched over to use instance ID.
+	//Input.PrimitiveId = PrimitiveId;
+	Input.InstanceIdOffset = 0;
+	Input.DrawInstanceId = 0;
+#else //!defined(GPUCULL_TODO)
 	Input.PrimitiveId = PrimitiveId;
+#endif // defined(GPUCULL_TODO)
 	Input.VertexId = HairPrimitiveId * 6; // 6: the vertex factory expect triangles list index. A quad is made of two triangles with indices: 0,1,2,3,4,5.
 
 	FVertexFactoryIntermediates	Intermediates = GetVertexFactoryIntermediates(Input);
diff --git a/Engine/Shaders/Private/HairStrands/HairStrandsVertexFactory.ush b/Engine/Shaders/Private/HairStrands/HairStrandsVertexFactory.ush
--- a/Engine/Shaders/Private/HairStrands/HairStrandsVertexFactory.ush
+++ b/Engine/Shaders/Private/HairStrands/HairStrandsVertexFactory.ush
@@ -98,8 +98,13 @@ bool UseScatterSceneLighting()
 struct FVertexFactoryInput
 {
 #if VF_USE_PRIMITIVE_SCENE_DATA
+#if defined(GPUCULL_TODO)
+	uint InstanceIdOffset : ATTRIBUTE13;
+	uint DrawInstanceId : SV_InstanceID;
+#else //!defined(GPUCULL_TODO)
 	uint PrimitiveId : ATTRIBUTE13;
-#endif	
+#endif // defined(GPUCULL_TODO)
+#endif //VF_USE_PRIMITIVE_SCENE_DATA
 	uint VertexId : SV_VertexID;
 };
 
@@ -162,6 +167,11 @@ FVertexFactoryInput LoadVertexFactoryInput(uint TriangleIndex, int VertexIndex)
  */
 struct FVertexFactoryIntermediates
 {
+#if defined(GPUCULL_TODO)
+#if VF_USE_PRIMITIVE_SCENE_DATA
+	FSceneDataIntermediates SceneData;
+#endif
+#endif // defined(GPUCULL_TODO)
 	half3x3 TangentToLocal;
 	half3x3 TangentToWorld;
 	half TangentToWorldSign;
@@ -172,6 +182,8 @@ struct FVertexFactoryIntermediates
 	float  HairDensity;
 };
 
+#if !defined(GPUCULL_TODO)
+
 uint GetPrimitiveId(FVertexFactoryInput Input)
 {
 #if VF_USE_PRIMITIVE_SCENE_DATA
@@ -180,6 +192,7 @@ uint GetPrimitiveId(FVertexFactoryInput Input)
 	return 0;
 #endif
 }
+#endif // defined(GPUCULL_TODO)
 
 uint GetHairPrimitiveId(FVertexFactoryInput Input)
 {
@@ -337,6 +350,19 @@ half3x3 CalcTangentToWorld(FVertexFactoryIntermediates Intermediates, half3x3 Ta
 	return LocalTangentToWorld;
 }
 
+#if defined(GPUCULL_TODO)
+
+float3 GetTangent(FVertexFactoryIntermediates Intermediates)
+{
+	// Hair shader expec the tangent to be stored in place of the normal. This is what the StrandHairComponent is filling in:
+	// [0]: Real normal		| Should be Tangent
+	// [1]: Bitangent		| Should be BiTangent
+	// [2]: Real Tangent	| Should be Normal
+	return Intermediates.TangentToWorld[2].xyz;
+}
+
+#else //!defined(GPUCULL_TODO)
+
 float3 GetTangent(FVertexFactoryInput Input)
 {
 	FVertexFactoryIntermediates Intermediates;
@@ -354,10 +380,21 @@ float3 GetTangent(FVertexFactoryInput Input)
 	return Intermediates.TangentToWorld[2].xyz;
 }
 
+#endif // defined(GPUCULL_TODO)
+
 FVertexFactoryIntermediates GetVertexFactoryIntermediates(FVertexFactoryInput Input)
 {
 	FVertexFactoryIntermediates Intermediates;
+#if defined(GPUCULL_TODO)
+#if VF_USE_PRIMITIVE_SCENE_DATA
+	Intermediates.SceneData = GetSceneDataIntermediates(Input.InstanceIdOffset, Input.DrawInstanceId);
+	Intermediates.PrimitiveId = Intermediates.SceneData.InstanceData.PrimitiveId;
+#else
+	Intermediates.PrimitiveId = 0U;
+#endif
+#else //!defined(GPUCULL_TODO)
 	Intermediates.PrimitiveId		= GetPrimitiveId(Input);
+#endif // defined(GPUCULL_TODO)
 	Intermediates.HairPrimitiveId	= GetHairPrimitiveId(Input);
 	Intermediates.HairPrimitiveUV	= GetSegmentUVCoord(Input);
 	Intermediates.HairDensity		= HairStrandsVF_Density;
@@ -404,8 +441,14 @@ float4 VertexFactoryGetWorldPosition_Visibility(FVertexFactoryInput Input, FHair
 	OutHairPrimitiveId = VertexInfo.HairPrimitiveId;
 
 	FHairControlPoint ControlPoint = GetVertexPosition(VertexInfo);
+#if defined(GPUCULL_TODO)
+	FVertexFactoryIntermediates Intermediates = GetVertexFactoryIntermediates(Input);
+	const float4 WorldPosition = CalcWorldPosition(ControlPoint.Position, Intermediates.PrimitiveId);
+	const float3 WorldTangent = GetTangent(Intermediates);
+#else //!defined(GPUCULL_TODO)
 	const float4 WorldPosition = CalcWorldPosition(ControlPoint.Position, GetPrimitiveId(Input));
 	const float3 WorldTangent = GetTangent(Input);
+#endif // defined(GPUCULL_TODO)
 
 	// Minimal radius to snap the strand to a sample/pixel center (to avoid aliasing)
 	const float DistanceToCamera = length(HairViewInfo.TranslatedWorldCameraOrigin - WorldPosition.xyz);
@@ -427,8 +470,14 @@ float4 VertexFactoryGetWorldPosition_Voxelization(FVertexFactoryInput Input, FHa
 	OutHairWorldRadius = ControlPoint.WorldRadius;
 	OutHairDensity = HairStrandsVF_Density;
 
+#if defined(GPUCULL_TODO)
+	FVertexFactoryIntermediates Intermediates = GetVertexFactoryIntermediates(Input);
+	const float4 WorldPosition = CalcWorldPosition(ControlPoint.Position, Intermediates.PrimitiveId);
+	const float3 WorldTangent = GetTangent(Intermediates);
+#else //!defined(GPUCULL_TODO)
 	const float4 WorldPosition = CalcWorldPosition(ControlPoint.Position, GetPrimitiveId(Input));
 	const float3 WorldTangent = GetTangent(Input);
+#endif // defined(GPUCULL_TODO)
 
 	// Minimal radius to snap the strand to a sample/pixel center (to avoid aliasing)
 	const float3 ViewDir = -HairViewInfo.ViewForward;
@@ -446,7 +495,11 @@ float4 VertexFactoryGetWorldPosition(FVertexFactoryInput Input, FVertexFactoryIn
 	const float4 WorldPosition = CalcWorldPosition(VertexPosition, Intermediates.PrimitiveId);
 
 	// Hair shader expect the WorldNormal to be the tangent vector
+#if defined(GPUCULL_TODO)
+	const float3 WorldTangent = GetTangent(Intermediates);
+#else //!defined(GPUCULL_TODO)
 	const float3 WorldTangent = GetTangent(Input);
+#endif // defined(GPUCULL_TODO)
 	return ComputeViewAlignedWorldPosition(Input, WorldTangent, WorldPosition, ControlPoint.WorldRadius, HairViewInfo);
 }
 
diff --git a/Engine/Shaders/Private/LocalVertexFactory.ush b/Engine/Shaders/Private/LocalVertexFactory.ush
--- a/Engine/Shaders/Private/LocalVertexFactory.ush
+++ b/Engine/Shaders/Private/LocalVertexFactory.ush
@@ -145,17 +145,28 @@ struct FVertexFactoryInput
 	float4 InstanceLightmapAndShadowMapUVBias : ATTRIBUTE12; 
 #endif //USE_INSTANCING
 
+#if defined(GPUCULL_TODO)
+#if VF_USE_PRIMITIVE_SCENE_DATA
+	uint InstanceIdOffset : ATTRIBUTE13;
+	uint DrawInstanceId : SV_InstanceID;
+#elif USE_INSTANCING
+	uint InstanceId : SV_InstanceID;
+#endif
+#else //!defined(GPUCULL_TODO)
 #if VF_USE_PRIMITIVE_SCENE_DATA
 	uint PrimitiveId : ATTRIBUTE13;
 #endif
+#endif // defined(GPUCULL_TODO)
 
 #if NEEDS_LIGHTMAP_COORDINATE && !MANUAL_VERTEX_FETCH
 	float2	LightMapCoordinate : ATTRIBUTE15;
 #endif
 
+#if !defined(GPUCULL_TODO)
 #if USE_INSTANCING
 	uint InstanceId	: SV_InstanceID;
 #endif
+#endif //!defined(GPUCULL_TODO)
 
 #if GPUSKIN_PASS_THROUGH || MANUAL_VERTEX_FETCH
 	uint VertexId : SV_VertexID;
@@ -184,7 +195,10 @@ FVertexFactoryInput LoadVertexFactoryInputForHGS(uint TriangleIndex, int VertexI
 	Input.VertexId = Tri.Indices[VertexIndex];
 	Input.Position = float4(Tri.LocalPositions[VertexIndex], 1.0f);
 	#if VF_USE_PRIMITIVE_SCENE_DATA
-		Input.PrimitiveId = GetInstanceUserData();
+#if !defined(GPUCULL_TODO)
+	// GPUCULL_TODO: Need to work out what this is meant to do
+	Input.PrimitiveId = GetInstanceUserData();
+#endif //!defined(GPUCULL_TODO)
 	#endif
 
 #if USE_INSTANCING
@@ -209,7 +223,10 @@ FVertexFactoryInput LoadVertexFactoryInputForDynamicUpdate(uint TriangleIndex, i
 		Input.Position.y = GPUSkinCachePositionBuffer[VertexOffset * 3 + 1];
 		Input.Position.z = GPUSkinCachePositionBuffer[VertexOffset * 3 + 2];
 		#if VF_USE_PRIMITIVE_SCENE_DATA
+#if !defined(GPUCULL_TODO)
+			// GPUCULL_TODO: !!!
 			Input.PrimitiveId = PrimitiveId;
+#endif //!defined(GPUCULL_TODO)
 		#endif
 	#else
 		Input.VertexId = TriangleIndex * 3 + VertexIndex;
@@ -218,7 +235,10 @@ FVertexFactoryInput LoadVertexFactoryInputForDynamicUpdate(uint TriangleIndex, i
 		Input.Position.y = LocalVF.VertexFetch_PositionBuffer[VertexOffset + 1];
 		Input.Position.z = LocalVF.VertexFetch_PositionBuffer[VertexOffset + 2];
 		#if VF_USE_PRIMITIVE_SCENE_DATA
+#if !defined(GPUCULL_TODO)
+			// GPUCULL_TODO: !!!
 			Input.PrimitiveId = PrimitiveId;
+#endif //!defined(GPUCULL_TODO)
 		#endif
 	#endif
 	return Input;
@@ -240,14 +260,24 @@ struct FPositionOnlyVertexFactoryInput
 	half4 InstanceTransform3 : ATTRIBUTE11; // hitproxy.b in .w
 #endif	// USE_INSTANCING
 
+#if defined(GPUCULL_TODO)
+#if VF_USE_PRIMITIVE_SCENE_DATA
+	uint InstanceIdOffset : ATTRIBUTE1;
+	uint DrawInstanceId : SV_InstanceID;
+#elif USE_INSTANCING
+	uint InstanceId : SV_InstanceID;
+#endif
+#else //!defined(GPUCULL_TODO)
 #if VF_USE_PRIMITIVE_SCENE_DATA
 	uint PrimitiveId : ATTRIBUTE1;
 #endif
+#endif // defined(GPUCULL_TODO)
 
+#if !defined(GPUCULL_TODO)
 #if USE_INSTANCING
 	uint InstanceId	: SV_InstanceID;
 #endif
-
+#endif //!defined(GPUCULL_TODO)
 #if MANUAL_VERTEX_FETCH
 	uint VertexId : SV_VertexID;
 #endif
@@ -268,13 +298,24 @@ struct FPositionAndNormalOnlyVertexFactoryInput
 	half4 InstanceTransform3 : ATTRIBUTE11; // hitproxy.b in .w
 #endif	// USE_INSTANCING
 
+#if defined(GPUCULL_TODO)
+#if VF_USE_PRIMITIVE_SCENE_DATA
+	uint InstanceIdOffset : ATTRIBUTE1;
+	uint DrawInstanceId : SV_InstanceID;
+#elif USE_INSTANCING
+	uint InstanceId : SV_InstanceID;
+#endif
+#else //!defined(GPUCULL_TODO)
 #if VF_USE_PRIMITIVE_SCENE_DATA
 	uint PrimitiveId : ATTRIBUTE1;
 #endif
+#endif // defined(GPUCULL_TODO)
 
-#if  USE_INSTANCING
+#if !defined(GPUCULL_TODO)
+#if USE_INSTANCING
 	uint InstanceId	: SV_InstanceID;
 #endif
+#endif //!defined(GPUCULL_TODO)
 
 #if MANUAL_VERTEX_FETCH
 	uint VertexId : SV_VertexID;
@@ -286,6 +327,13 @@ struct FPositionAndNormalOnlyVertexFactoryInput
  */
 struct FVertexFactoryIntermediates
 {
+#if defined(GPUCULL_TODO)
+#if VF_USE_PRIMITIVE_SCENE_DATA
+	FSceneDataIntermediates SceneData;
+	int PrimitiveLocalInstanceIndex;
+#endif // VF_USE_PRIMITIVE_SCENE_DATA
+#endif // defined(GPUCULL_TODO)
+
 	half3x3 TangentToLocal;
 	half3x3 TangentToWorld;
 	half TangentToWorldSign;
@@ -324,6 +372,9 @@ struct FVertexFactoryIntermediates
 #endif
 
 	uint PrimitiveId;
+#if defined(GPUCULL_TODO)
+	uint InstanceId;
+#endif // defined(GPUCULL_TODO)
 
 	float3 PreSkinPosition;
 };
@@ -495,8 +546,13 @@ FMaterialPixelParameters GetMaterialPixelParameters(FVertexFactoryInterpolantsVS
 
 half3x3 CalcTangentToWorldNoScale(FVertexFactoryIntermediates Intermediates, half3x3 TangentToLocal)
 {
+#if defined(GPUCULL_TODO) && VF_USE_PRIMITIVE_SCENE_DATA
+	half3x3 LocalToWorld = (half3x3)Intermediates.SceneData.InstanceData.LocalToWorld;
+	half3 InvScale = Intermediates.SceneData.InstanceData.InvNonUniformScaleAndDeterminantSign.xyz;
+#else
 	half3x3 LocalToWorld = GetLocalToWorld3x3(Intermediates.PrimitiveId);
 	half3 InvScale = GetPrimitiveData(Intermediates.PrimitiveId).InvNonUniformScaleAndDeterminantSign.xyz;
+#endif
 	LocalToWorld[0] *= InvScale.x;
 	LocalToWorld[1] *= InvScale.y;
 	LocalToWorld[2] *= InvScale.z;
@@ -530,8 +586,18 @@ FMaterialVertexParameters GetMaterialVertexParameters(FVertexFactoryInput Input,
 		Result.PrevFrameLocalToWorld = mul(GetInstanceTransform(Intermediates), GetPrimitiveData(Intermediates.PrimitiveId).PreviousLocalToWorld);
 	#endif // USE_INSTANCING_BONEMAP
 #else
-	Result.PrevFrameLocalToWorld = GetPrimitiveData(Intermediates.PrimitiveId).PreviousLocalToWorld;
-#endif	// USE_INSTANCING
+	#if defined(GPUCULL_TODO) && VF_USE_PRIMITIVE_SCENE_DATA
+		Result.PrevFrameLocalToWorld = Intermediates.SceneData.InstanceData.PrevLocalToWorld;
+		#if USE_INSTANCE_CULLING
+			Result.InstanceLocalPosition = Input.Position.xyz;
+			Result.InstanceLocalToWorld = Intermediates.SceneData.InstanceData.LocalToWorld;
+			Result.PerInstanceParams = Intermediates.PerInstanceParams;
+			Result.PrimitiveLocalInstanceIndex = Intermediates.PrimitiveLocalInstanceIndex;
+		#endif
+	#else // !defined(GPUCULL_TODO) && VF_USE_PRIMITIVE_SCENE_DATA
+		Result.PrevFrameLocalToWorld = GetPrimitiveData(Intermediates.PrimitiveId).PreviousLocalToWorld;
+	#endif // defined(GPUCULL_TODO) && VF_USE_PRIMITIVE_SCENE_DATA
+#endif // USE_INSTANCING
 
 	Result.PreSkinnedPosition = Intermediates.PreSkinPosition.xyz;
 	Result.PreSkinnedNormal = TangentToLocal[2]; //TangentBias(Input.TangentZ.xyz);
@@ -705,10 +771,8 @@ float4 CalcWorldPosition(float4 Position, uint PrimitiveId)
 #endif
 }
 
-half3x3 CalcTangentToLocal(FVertexFactoryInput Input, out float TangentSign)
+half3x3 CalcTangentToLocal(FVertexFactoryInput Input, inout float TangentSign)
 {
-	half3x3 Result;
-	
 #if MANUAL_VERTEX_FETCH
 	half3 TangentInputX = LocalVF.VertexFetch_PackedTangentsBuffer[2 * (LocalVF.VertexFetch_Parameters[VF_VertexOffset] + Input.VertexId) + 0].xyz;
 	half4 TangentInputZ = LocalVF.VertexFetch_PackedTangentsBuffer[2 * (LocalVF.VertexFetch_Parameters[VF_VertexOffset] + Input.VertexId) + 1].xyzw;
@@ -741,6 +805,7 @@ half3x3 CalcTangentToLocal(FVertexFactoryInput Input, out float TangentSign)
 	// Recalculate TangentX off of the other two vectors
 	// This corrects quantization error since TangentX was passed in as a quantized vertex input
 	// The error shows up most in specular off of a mesh with a smoothed UV seam (normal is smooth, but tangents vary across the seam)
+	half3x3 Result;
 	Result[0] = cross(TangentY, TangentZ.xyz) * TangentZ.w;
 	Result[1] = TangentY;
 	Result[2] = TangentZ.xyz;
@@ -765,13 +830,78 @@ half3x3 CalcTangentToWorld(FVertexFactoryIntermediates Intermediates, half3x3 Ta
 
 FVertexFactoryIntermediates GetVertexFactoryIntermediates(FVertexFactoryInput Input)
 {
-	FVertexFactoryIntermediates Intermediates;
+	FVertexFactoryIntermediates Intermediates = (FVertexFactoryIntermediates)0;
 
 #if VF_USE_PRIMITIVE_SCENE_DATA
+#if defined(GPUCULL_TODO)
+	Intermediates.SceneData = GetSceneDataIntermediates(Input.InstanceIdOffset, Input.DrawInstanceId);
+	// Just aliases, remove later probably...
+	Intermediates.PrimitiveId = Intermediates.SceneData.PrimitiveId;
+	Intermediates.InstanceId = Intermediates.SceneData.InstanceId;
+
+	// Calculate the local (within the range of instance IDs belonging to a given primitive), as this can be used to load custom stuff
+	Intermediates.PrimitiveLocalInstanceIndex = Intermediates.SceneData.InstanceId - Intermediates.SceneData.Primitive.InstanceDataOffset;
+
+
+#if USE_INSTANCE_CULLING
+	{
+	#if MANUAL_VERTEX_FETCH
+		Intermediates.PerInstanceRandom = InstanceVF.VertexFetch_InstanceOriginBuffer[Intermediates.PrimitiveLocalInstanceIndex].w;
+
+		float R = InstanceVF.VertexFetch_InstanceTransformBuffer[3 * (Intermediates.PrimitiveLocalInstanceIndex) + 0].w;
+		float G = InstanceVF.VertexFetch_InstanceTransformBuffer[3 * (Intermediates.PrimitiveLocalInstanceIndex) + 1].w;
+		float B = InstanceVF.VertexFetch_InstanceTransformBuffer[3 * (Intermediates.PrimitiveLocalInstanceIndex) + 2].w;
+
+		Intermediates.InstanceLightmapAndShadowMapUVBias = InstanceVF.VertexFetch_InstanceLightmapBuffer[Intermediates.PrimitiveLocalInstanceIndex];
+	#else
+		Intermediates.PerInstanceRandom = Input.InstanceOrigin.w;
+
+		float R = Input.InstanceTransform1.w;
+		float G = Input.InstanceTransform2.w;
+		float B = Input.InstanceTransform3.w;
+
+		Intermediates.InstanceLightmapAndShadowMapUVBias = Input.InstanceLightmapAndShadowMapUVBias;
+	#endif
+
+		Intermediates.IsSelected = trunc(R * (1.0 / 256.0));;
+		R = R - 256.0 * Intermediates.IsSelected;
+		Intermediates.HitProxyId = float4(R / 255.0, G / 255.0, B / 255.0, 0);
+
+		// This repackaging is stupid-seeming but these params are propably packed off to some interpolator or something
+		Intermediates.PerInstanceParams.x = Intermediates.PerInstanceRandom;
+		
+		// PerInstanceParams.y: per-instance fade out amount
+		float3 InstanceLocation = TransformLocalToWorld(Intermediates.SceneData.InstanceData.LocalBoundsCenter, Intermediates.SceneData.InstanceData.LocalToWorld);
+		Intermediates.PerInstanceParams.y = 1.0 - saturate((length(InstanceLocation + ResolvedView.PreViewTranslation.xyz) - InstancingFadeOutParams.x) * InstancingFadeOutParams.y);
+		
+		// InstancingFadeOutParams.z,w are RenderSelected and RenderDeselected respectively.
+
+		// PerInstanceParams.z = hide / show flag, 
+		Intermediates.IsVisible = lerp(InstancingFadeOutParams.w, InstancingFadeOutParams.z, Intermediates.IsSelected);
+		Intermediates.PerInstanceParams.z = Intermediates.IsVisible;
+
+		// PerInstanceParams.w dither fade cutoff
+	#if USE_DITHERED_LOD_TRANSITION
+		float RandomLOD = InstancingViewZCompareZero.w * Intermediates.PerInstanceParams.x;
+		float ViewZZero = length(InstanceLocation - InstancingWorldViewOriginZero.xyz) + RandomLOD;
+		float ViewZOne = length(InstanceLocation - InstancingWorldViewOriginOne.xyz) + RandomLOD;
+		Intermediates.PerInstanceParams.w =
+			dot(float3(ViewZZero.xxx > InstancingViewZCompareZero.xyz), InstancingViewZConstant.xyz) * InstancingWorldViewOriginZero.w +
+			dot(float3(ViewZOne.xxx > InstancingViewZCompareOne.xyz), InstancingViewZConstant.xyz) * InstancingWorldViewOriginOne.w;
+		Intermediates.PerInstanceParams.z *= abs(Intermediates.PerInstanceParams.w) < .999;
+		Intermediates.IsVisible = Intermediates.PerInstanceParams.z;
+	#else
+		Intermediates.PerInstanceParams.w = 0;
+	#endif
+	}
+#endif // USE_INSTANCE_CULLING
+#else //!defined(GPUCULL_TODO)
 	Intermediates.PrimitiveId = Input.PrimitiveId;
-#else
+#endif // defined(GPUCULL_TODO)
+
+#else // !VF_USE_PRIMITIVE_SCENE_DATA
 	Intermediates.PrimitiveId = 0;
-#endif
+#endif // VF_USE_PRIMITIVE_SCENE_DATA
 
 #if MANUAL_VERTEX_FETCH
 	Intermediates.Color = LocalVF.VertexFetch_ColorComponentsBuffer[(LocalVF.VertexFetch_Parameters[VF_VertexOffset] + Input.VertexId) & LocalVF.VertexFetch_Parameters[VF_ColorIndexMask_Index]] FMANUALFETCH_COLOR_COMPONENT_SWIZZLE; // Swizzle vertex color.
@@ -807,17 +937,25 @@ FVertexFactoryIntermediates GetVertexFactoryIntermediates(FVertexFactoryInput In
 	Intermediates.InstanceLightmapAndShadowMapUVBias = Input.InstanceLightmapAndShadowMapUVBias;
 #endif
 
-	float TangentSign;
+	float TangentSign = 1.0;
 	Intermediates.TangentToLocal = CalcTangentToLocal(Input, TangentSign);
-	Intermediates.TangentToWorld = CalcTangentToWorld(Intermediates,Intermediates.TangentToLocal);
+	Intermediates.TangentToWorld = CalcTangentToWorld(Intermediates, Intermediates.TangentToLocal);
+#if defined(GPUCULL_TODO) && VF_USE_PRIMITIVE_SCENE_DATA
+	Intermediates.TangentToWorldSign = TangentSign * Intermediates.SceneData.InstanceData.InvNonUniformScaleAndDeterminantSign.w;
+#else
 	Intermediates.TangentToWorldSign = TangentSign * GetPrimitiveData(Intermediates.PrimitiveId).InvNonUniformScaleAndDeterminantSign.w;
+#endif
 
 #if USE_INSTANCING && !USE_INSTANCING_BONEMAP
 	// x = per-instance random, y = per-instance fade out factor, z = zero or one depending of if it is shown at all, w is dither cutoff 
 
 	// PerInstanceParams.z stores a hide/show flag for this instance
 	float SelectedValue = GetInstanceSelected(Intermediates);
+#if defined(GPUCULL_TODO)
+	Intermediates.PerInstanceParams.x = Intermediates.PerInstanceRandom;
+#else //!defined(GPUCULL_TODO)
 	Intermediates.PerInstanceParams.x = GetInstanceRandom(Intermediates);
+#endif // defined(GPUCULL_TODO)
 	float3 InstanceLocation = TransformLocalToWorld(GetInstanceOrigin(Intermediates), Intermediates.PrimitiveId).xyz;
 	Intermediates.PerInstanceParams.y = 1.0 - saturate((length(InstanceLocation + ResolvedView.PreViewTranslation.xyz) - InstancingFadeOutParams.x) * InstancingFadeOutParams.y);
 	// InstancingFadeOutParams.z,w are RenderSelected and RenderDeselected respectively.
@@ -867,11 +1005,20 @@ half3x3 VertexFactoryGetTangentToLocal( FVertexFactoryInput Input, FVertexFactor
 // @return translated world position
 float4 VertexFactoryGetWorldPosition(FVertexFactoryInput Input, FVertexFactoryIntermediates Intermediates)
 {
-#if USE_INSTANCING
-	return CalcWorldPosition(Input.Position, GetInstanceTransform(Intermediates), Intermediates.PrimitiveId) * Intermediates.PerInstanceParams.z;
-#else
-	return CalcWorldPosition(Input.Position, Intermediates.PrimitiveId);
+#if defined(GPUCULL_TODO) && VF_USE_PRIMITIVE_SCENE_DATA
+	#if USE_INSTANCE_CULLING
+		// Scale to zero if not visible, seems a bit wild but whatever
+		return TransformLocalToTranslatedWorld(Input.Position.xyz, Intermediates.SceneData.InstanceData) * Intermediates.IsVisible;
+	#else
+		return TransformLocalToTranslatedWorld(Input.Position.xyz, Intermediates.SceneData.InstanceData);
+	#endif
+#else // !defined(GPUCULL_TODO) && VF_USE_PRIMITIVE_SCENE_DATA
+	#if USE_INSTANCING
+		return CalcWorldPosition(Input.Position, GetInstanceTransform(Intermediates), Intermediates.PrimitiveId) * Intermediates.PerInstanceParams.z;
+	#else
+		return CalcWorldPosition(Input.Position, Intermediates.PrimitiveId);
 #endif	// USE_INSTANCING
+#endif // defined(GPUCULL_TODO) && VF_USE_PRIMITIVE_SCENE_DATA
 }
 
 float4 VertexFactoryGetRasterizedWorldPosition(FVertexFactoryInput Input, FVertexFactoryIntermediates Intermediates, float4 InWorldPosition)
@@ -924,8 +1071,13 @@ FVertexFactoryInterpolantsVSToPS VertexFactoryGetInterpolantsVSToPS(FVertexFacto
 	uint LightmapUVIndex   = 0;
 
 #if VF_USE_PRIMITIVE_SCENE_DATA
+#if defined(GPUCULL_TODO)
+	LightmapDataIndex = Intermediates.SceneData.Primitive.LightmapDataIndex + LocalVF.LODLightmapDataIndex;
+	LightmapUVIndex = Intermediates.SceneData.Primitive.LightmapUVIndex;
+#else //!defined(GPUCULL_TODO)
 	LightmapDataIndex = GetPrimitiveData(Intermediates.PrimitiveId).LightmapDataIndex + LocalVF.LODLightmapDataIndex;
 	LightmapUVIndex = GetPrimitiveData(Intermediates.PrimitiveId).LightmapUVIndex;
+#endif // defined(GPUCULL_TODO)
 #endif
 
 	float4 LightMapCoordinateScaleBias = GetLightmapData(LightmapDataIndex).LightMapCoordinateScaleBias;
@@ -964,6 +1116,63 @@ FVertexFactoryInterpolantsVSToPS VertexFactoryGetInterpolantsVSToPS(FVertexFacto
 	return Interpolants;
 }
 
+
+#if defined(GPUCULL_TODO)
+
+
+/** X for depth-only pass */
+float4 VertexFactoryGetWorldPosition(FPositionOnlyVertexFactoryInput Input)
+{
+#if VF_USE_PRIMITIVE_SCENE_DATA
+	FSceneDataIntermediates SceneData = GetSceneDataIntermediates(Input.InstanceIdOffset, Input.DrawInstanceId);
+	return TransformLocalToTranslatedWorld(Input.Position.xyz, SceneData.InstanceData);
+#else
+	float4 Position = Input.Position;
+	uint PrimitiveId = 0;
+    #if USE_INSTANCING
+	    return CalcWorldPosition(Position, GetInstanceTransform(Input), PrimitiveId);
+    #else
+	    return CalcWorldPosition(Position, PrimitiveId);
+    #endif
+#endif
+}
+
+/** for depth-only pass (slope depth bias) */
+float4 VertexFactoryGetWorldPosition(FPositionAndNormalOnlyVertexFactoryInput Input)
+{
+#if VF_USE_PRIMITIVE_SCENE_DATA
+	FSceneDataIntermediates SceneData = GetSceneDataIntermediates(Input.InstanceIdOffset, Input.DrawInstanceId);
+	return TransformLocalToTranslatedWorld(Input.Position.xyz, SceneData.InstanceData);
+#else
+	float4 Position = Input.Position;
+	uint PrimitiveId = 0;
+    #if USE_INSTANCING
+	    return CalcWorldPosition(Position, GetInstanceTransform(Input), PrimitiveId);
+    #else
+	    return CalcWorldPosition(Position, PrimitiveId);
+    #endif
+#endif
+}
+
+float3 VertexFactoryGetWorldNormal(FPositionAndNormalOnlyVertexFactoryInput Input)
+{
+#if VF_USE_PRIMITIVE_SCENE_DATA
+	FSceneDataIntermediates SceneData = GetSceneDataIntermediates(Input.InstanceIdOffset, Input.DrawInstanceId);
+	return RotateLocalToWorld(Input.Normal.xyz, SceneData.InstanceData);
+#else
+	uint PrimitiveId = 0;
+	float3 Normal = Input.Normal.xyz;
+    #if USE_INSTANCING
+	    const float3 InstanceTransformedNormal = mul(float4(Normal, 0), GetInstanceTransform(Input)).xyz;
+	    return RotateLocalToWorld(InstanceTransformedNormal, PrimitiveId);
+    #else
+	    return RotateLocalToWorld(Normal, PrimitiveId);
+    #endif
+#endif
+}
+
+#else //!defined(GPUCULL_TODO)
+
 /** for depth-only pass */
 float4 VertexFactoryGetWorldPosition(FPositionOnlyVertexFactoryInput Input)
 {
@@ -1018,6 +1227,7 @@ float3 VertexFactoryGetWorldNormal(FPositionAndNormalOnlyVertexFactoryInput Inpu
 #endif	// USE_INSTANCING
 }
 
+#endif // defined(GPUCULL_TODO)
 
 float3 VertexFactoryGetWorldNormal(FVertexFactoryInput Input, FVertexFactoryIntermediates Intermediates)
 {
@@ -1032,7 +1242,11 @@ Buffer<float> GPUSkinCachePreviousPositionBuffer;
 // @return previous translated world position
 float4 VertexFactoryGetPreviousWorldPosition(FVertexFactoryInput Input, FVertexFactoryIntermediates Intermediates)
 {
+#if defined(GPUCULL_TODO) && VF_USE_PRIMITIVE_SCENE_DATA
+	float4x4 PreviousLocalToWorldTranslated = Intermediates.SceneData.InstanceData.PrevLocalToWorld;
+#else
 	float4x4 PreviousLocalToWorldTranslated = GetPrimitiveData(Intermediates.PrimitiveId).PreviousLocalToWorld;
+#endif
 	PreviousLocalToWorldTranslated[3][0] += ResolvedView.PrevPreViewTranslation.x;
 	PreviousLocalToWorldTranslated[3][1] += ResolvedView.PrevPreViewTranslation.y;
 	PreviousLocalToWorldTranslated[3][2] += ResolvedView.PrevPreViewTranslation.z;
diff --git a/Engine/Shaders/Private/RayTracing/RayTracingMaterialHitShaders.usf b/Engine/Shaders/Private/RayTracing/RayTracingMaterialHitShaders.usf
--- a/Engine/Shaders/Private/RayTracing/RayTracingMaterialHitShaders.usf
+++ b/Engine/Shaders/Private/RayTracing/RayTracingMaterialHitShaders.usf
@@ -357,7 +357,9 @@ RAY_TRACING_ENTRY_CLOSEST_HIT(MaterialCHS,
 	#endif
 
 	Payload.BlendingMode = MaterialBlendingMode;
+#if !defined(GPUCULL_TODO)
 	Payload.PrimitiveLightingChannelMask = GetPrimitiveData(MaterialParameters.PrimitiveId).LightingChannelMask;
+#endif // defined(GPUCULL_TODO)
 
 	Payload.HitT = RayTCurrent();
 	if (HitKind() == HIT_KIND_TRIANGLE_FRONT_FACE)
diff --git a/Engine/Shaders/Private/SceneData.ush b/Engine/Shaders/Private/SceneData.ush
--- a/Engine/Shaders/Private/SceneData.ush
+++ b/Engine/Shaders/Private/SceneData.ush
@@ -9,9 +9,12 @@
 // Whether to fetch primitive values (eg LocalToWorld) by dynamically indexing a scene-wide buffer, or to reference a single Primitive uniform buffer
 #define VF_USE_PRIMITIVE_SCENE_DATA ((FEATURE_LEVEL >= FEATURE_LEVEL_SM5 || FEATURE_LEVEL == FEATURE_LEVEL_ES3_1) && VF_SUPPORTS_PRIMITIVE_SCENE_DATA)
 
-#if VF_USE_PRIMITIVE_SCENE_DATA
+#define INSTANCE_SCENE_DATA_SUPPORTED (VF_USE_PRIMITIVE_SCENE_DATA && (FEATURE_LEVEL >= FEATURE_LEVEL_SM5 || FEATURE_LEVEL == FEATURE_LEVEL_ES3_1))
 
 // Unpacked AoS layout - see FInstanceSceneShaderData::Setup() for SoA packed layout.
+#if INSTANCE_SCENE_DATA_SUPPORTED
+
+// Must match FInstanceUniformShaderParameters in C++
 struct FInstanceSceneData
 {
 	float4x4 LocalToWorld;
@@ -33,7 +36,6 @@ struct FInstanceSceneData
 
 #if USE_GLOBAL_GPU_SCENE_DATA
 StructuredBuffer<float4> GPUSceneInstanceSceneData;
-StructuredBuffer<float4> GPUScenePrimitiveSceneData;
 uint GPUSceneFrameNumber;
 #endif
 
@@ -99,6 +101,14 @@ FInstanceSceneData GetInstanceData(uint InstanceId, uint SOAStride)
 	return InstanceData;
 }
 
+#endif // INSTANCE_SCENE_DATA_SUPPORTED
+
+#if VF_USE_PRIMITIVE_SCENE_DATA
+
+#if USE_GLOBAL_GPU_SCENE_DATA
+StructuredBuffer<float4> GPUScenePrimitiveSceneData;
+#endif
+
 #define NUM_CUSTOM_PRIMITIVE_DATA 9 // Num float4s used for custom data. Must match FCustomPrimitiveData::NumCustomPrimitiveDataFloat4s in SceneTypes.h
 
 // Must match FPrimitiveUniformShaderParameters in C++
@@ -260,3 +270,54 @@ FPrimitiveSceneData GetPrimitiveData(uint PrimitiveId)
 #define GetInstanceData(x,y) PrimitiveInstance
 
 #endif // VF_USE_PRIMITIVE_SCENE_DATA
+
+#if defined(GPUCULL_TODO)
+
+struct FSceneDataIntermediates
+{
+	uint PrimitiveId;
+	uint InstanceId;
+	uint ViewIndex;
+	// Index from which we load the instance info, needed for the 
+	uint InstanceIdLoadIndex;
+#if INSTANCE_SCENE_DATA_SUPPORTED
+	FInstanceSceneData InstanceData;
+#endif // INSTANCE_SCENE_DATA_SUPPORTED
+
+#if VF_USE_PRIMITIVE_SCENE_DATA
+	// GPUCULL_TODO: Move this definition out of ifdef in SceneData so we can have fewer ifdefs?
+	// Also want to copy / reference the primitive UB so that we have a consistent access everywhere
+	FPrimitiveSceneData Primitive;
+#endif
+};
+
+/**
+ * Load scene data once given the inputs require.
+ * InstanceIdOffset - supplied as a vertex stream with 0 instance step rate (constant for all instances)
+ * DrawInstanceId - the instance ID (SV_InstanceID) in the current draw
+ */
+FSceneDataIntermediates GetSceneDataIntermediates(uint InstanceIdOffset, uint DrawInstanceId)
+{
+	FSceneDataIntermediates Intermediates = (FSceneDataIntermediates)0;
+#if VF_USE_PRIMITIVE_SCENE_DATA
+	Intermediates.InstanceIdLoadIndex = InstanceIdOffset + DrawInstanceId;
+	// TODO: move InstanceIdsBuffer out of view and into some global Scene structure?
+	Intermediates.InstanceId = View.InstanceIdsBuffer[InstanceIdOffset + DrawInstanceId] & ((1U << 28U) - 1);
+	// We store the view index (which can be used for instanced stereo or other multi-view in the top four bits of the instance ID)
+	// Note: this is an index of views for this render pass, not the view ID in the culling manager.
+	Intermediates.ViewIndex = View.InstanceIdsBuffer[InstanceIdOffset + DrawInstanceId] >> 28U;
+	Intermediates.InstanceData = GetInstanceData(Intermediates.InstanceId, View.InstanceDataSOAStride);
+	Intermediates.PrimitiveId = Intermediates.InstanceData.PrimitiveId;
+	Intermediates.Primitive = GetPrimitiveData(Intermediates.PrimitiveId);
+#else // Reference the uniform buffers
+	// GPUCULL_TODO: Need to get this from uniforms for the non-GPU scene path
+	//           Maybe we need to add an accessor function that can be a macro that resolves (like Primitive does) to the UBs
+	//Intermediates.ViewIndex = 0U;
+	//Intermediates.InstanceData = InstanceData;
+	//Intermediates.PrimitiveId = Intermediates.InstanceData.PrimitiveId;
+	//Intermediates.Primitive = Primitive;
+#endif
+	return Intermediates;
+}
+
+#endif // defined(GPUCULL_TODO)
diff --git a/Engine/Shaders/Private/ShadowDepthCommon.ush b/Engine/Shaders/Private/ShadowDepthCommon.ush
--- a/Engine/Shaders/Private/ShadowDepthCommon.ush
+++ b/Engine/Shaders/Private/ShadowDepthCommon.ush
@@ -37,6 +37,9 @@ struct FShadowDepthVSToGS : FShadowDepthVSToPS
 {
 	float4 GSPosition : TEXCOORD6;
 	float3 WorldNormal : TEXCOORD8;
+#if defined(GPUCULL_TODO)
+	nointerpolation uint ViewIndex : TEXCOORD9;
+#endif // defined(GPUCULL_TODO)
 };
 #endif
 
diff --git a/Engine/Shaders/Private/ShadowDepthVertexShader.usf b/Engine/Shaders/Private/ShadowDepthVertexShader.usf
--- a/Engine/Shaders/Private/ShadowDepthVertexShader.usf
+++ b/Engine/Shaders/Private/ShadowDepthVertexShader.usf
@@ -85,6 +85,10 @@ void SetShadowDepthOutputs(
 uint LayerId;
 #endif
 
+#if defined(GPUCULL_TODO)
+uint bUseGpuSceneInstancing;
+#endif // defined(GPUCULL_TODO)
+
 struct FShadowDepthGSToPS
 {
 #if INTERPOLATE_VF_ATTRIBUTES
@@ -110,13 +114,22 @@ bool IsTriangleInFrustum(float4 ClipSpacePositions[3])
 	return !( any(FrustumTests != 0) || any(FrustumTestsZ != 0) );
 }
 
-/** Allocate space for cloning to all 6 faces which is the worst case. */
+/** Allocate space for cloning to all 6 faces which is the worst case (no GPU-scene active). */
+/** GPU-Scene case: Instancing takes care of the clone, we just route to a layer. */
 [maxvertexcount(18)]
 void MainOnePassPointLightGS(triangle FShadowDepthVSToGS Input[3], inout TriangleStream<FShadowDepthGSToPS> OutStream)
 {
+#if defined(GPUCULL_TODO)
+	uint CubeFaceStartIndex = bUseGpuSceneInstancing ? Input[0].ViewIndex : 0;
+	uint CubeFaceEndIndex = bUseGpuSceneInstancing ? Input[0].ViewIndex + 1: 6;
+
+	// Clone the triangle to each face
+	for (int CubeFaceIndex = CubeFaceStartIndex; CubeFaceIndex < CubeFaceEndIndex; CubeFaceIndex++)
+#else //!defined(GPUCULL_TODO)
 	UNROLL
 	// Clone the triangle to each face
 	for (int CubeFaceIndex = 0; CubeFaceIndex < 6; CubeFaceIndex++)
+#endif // defined(GPUCULL_TODO)
 	{
 		float4 ClipSpacePositions[3];
 		UNROLL
@@ -300,13 +313,21 @@ void Main(
 	OutParameters.FactoryInterpolants = VertexFactoryGetInterpolantsVSToDS(Input, VFIntermediates, VertexParameters);
 
 	#if ONEPASS_POINTLIGHT_SHADOW && USING_VERTEX_SHADER_LAYER && VERTEXSHADER
+#if defined(GPUCULL_TODO)
+		OutParameters.LayerIndex = bUseGpuSceneInstancing ? VertexFactoryGetViewIndex(VFIntermediates) : LayerId;
+#else //!defined(GPUCULL_TODO)
 		OutParameters.LayerIndex = LayerId;
+#endif // defined(GPUCULL_TODO)
+
 	#endif
 
 #elif ONEPASS_POINTLIGHT_SHADOW
 
 	OutPosition = OutParameters.GSPosition = WorldPos;
 	OutParameters.WorldNormal = WorldNormal;
+#if defined(GPUCULL_TODO)
+	OutParameters.ViewIndex = VertexFactoryGetViewIndex(VFIntermediates);
+#endif // defined(GPUCULL_TODO)
 
 	#if INTERPOLATE_VF_ATTRIBUTES
 		// Masked materials need texture coords to clip
@@ -318,7 +339,11 @@ void Main(
 	#endif
 
 	#if FEATURE_LEVEL >= FEATURE_LEVEL_SM4 && USING_VERTEX_SHADER_LAYER && VERTEXSHADER
+#if defined(GPUCULL_TODO)
+		LayerIndex = bUseGpuSceneInstancing ? VertexFactoryGetViewIndex(VFIntermediates) : LayerId;
+#else //!defined(GPUCULL_TODO)
 		LayerIndex = LayerId;
+#endif // defined(GPUCULL_TODO)
 		OutPosition = mul(WorldPos, PassStruct.ShadowViewProjectionMatrices[LayerIndex]);
 		#if INTERPOLATE_VF_ATTRIBUTES
 			OutParameters.GSPosition = OutPosition;
@@ -427,8 +452,16 @@ void PositionOnlyMain(
 #if ONEPASS_POINTLIGHT_SHADOW
 	OutPosition = OutParameters.GSPosition = WorldPos;
 	OutParameters.WorldNormal = WorldNormal;
+#if defined(GPUCULL_TODO)
+	OutParameters.ViewIndex = VertexFactoryGetViewIndex(Input);
+#endif // defined(GPUCULL_TODO)
+
 	#if USING_ONEPASS_POINT_VERTEX_LAYER
+#if defined(GPUCULL_TODO)
+		LayerIndex = bUseGpuSceneInstancing ? VertexFactoryGetViewIndex(Input) : LayerId;
+#else //!defined(GPUCULL_TODO)
 		LayerIndex = LayerId;
+#endif // defined(GPUCULL_TODO)
 		OutPosition = mul(WorldPos, PassStruct.ShadowViewProjectionMatrices[LayerIndex]);
 		#if INTERPOLATE_VF_ATTRIBUTES
 			OutParameters.GSPosition = OutPosition;
diff --git a/Engine/Shaders/Private/VertexFactoryCommon.ush b/Engine/Shaders/Private/VertexFactoryCommon.ush
--- a/Engine/Shaders/Private/VertexFactoryCommon.ush
+++ b/Engine/Shaders/Private/VertexFactoryCommon.ush
@@ -2,6 +2,15 @@
 
 #include "Common.ush"
 
+#if defined(GPUCULL_TODO)
+float3 TransformLocalToWorld(float3 LocalPosition, float4x4 LocalToWorld)
+{
+	//expand the matrix multiply to avoid a code-gen bug on PS4.
+	//invariant code!
+	return INVARIANT((LocalToWorld[0].xyz * LocalPosition.xxx + LocalToWorld[1].xyz * LocalPosition.yyy + LocalToWorld[2].xyz * LocalPosition.zzz) + LocalToWorld[3].xyz);
+}
+#endif // defined(GPUCULL_TODO)
+
 float3 TransformLocalToWorld(float3 LocalPosition, uint PrimitiveId)
 {
 	//expand the matrix multiply to avoid a code-gen bug on PS4.
@@ -24,6 +33,31 @@ float4 TransformLocalToTranslatedWorld(float3 LocalPosition, uint PrimitiveId)
 	return float4(RotatedPosition + (LocalToWorld[3].xyz + ResolvedView.PreViewTranslation.xyz),1);
 }
 
+#if defined(GPUCULL_TODO)
+#if INSTANCE_SCENE_DATA_SUPPORTED
+
+float4 TransformLocalToTranslatedWorld(float3 LocalPosition, FInstanceSceneData InstanceSceneData)
+{
+	float4x4 LocalToWorld = InstanceSceneData.LocalToWorld;
+	float3 RotatedPosition = LocalToWorld[0].xyz * LocalPosition.xxx + LocalToWorld[1].xyz * LocalPosition.yyy + LocalToWorld[2].xyz * LocalPosition.zzz;
+	return float4(RotatedPosition + (LocalToWorld[3].xyz + ResolvedView.PreViewTranslation.xyz), 1);
+}
+
+float3 RotateLocalToWorld(float3 LocalDirection, FInstanceSceneData InstanceSceneData)
+{
+	const float4x4 LocalToWorld = InstanceSceneData.LocalToWorld;
+	const float3 InvScale = InstanceSceneData.InvNonUniformScaleAndDeterminantSign.xyz;
+	return
+		InvScale.x * LocalToWorld[0].xyz * LocalDirection.xxx +
+		InvScale.y * LocalToWorld[1].xyz * LocalDirection.yyy +
+		InvScale.z * LocalToWorld[2].xyz * LocalDirection.zzz;
+}
+
+#endif // INSTANCE_SCENE_DATA_SUPPORTED
+#endif // defined(GPUCULL_TODO)
+
+#if (!defined(GPUCULL_TODO) || !VF_USE_PRIMITIVE_SCENE_DATA)
+
 float4 TransformLocalToTranslatedWorld(float3 LocalPosition)
 {
 	float3 RotatedPosition = Primitive.LocalToWorld[0].xyz * LocalPosition.xxx + Primitive.LocalToWorld[1].xyz * LocalPosition.yyy + Primitive.LocalToWorld[2].xyz * LocalPosition.zzz;
@@ -61,6 +95,8 @@ float3 RotateWorldToLocal(float3 WorldDirection)
 	return mul(WorldDirection, InvRot);
 }
 
+#endif // (!defined(GPUCULL_TODO) || !VF_USE_PRIMITIVE_SCENE_DATA)
+
 // Octahedron Normal Vectors
 // [Cigolle 2014, "A Survey of Efficient Representations for Independent Unit Vectors"]
 //						Mean	Max
diff --git a/Engine/Source/Runtime/Engine/Private/HierarchicalInstancedStaticMesh.cpp b/Engine/Source/Runtime/Engine/Private/HierarchicalInstancedStaticMesh.cpp
--- a/Engine/Source/Runtime/Engine/Private/HierarchicalInstancedStaticMesh.cpp
+++ b/Engine/Source/Runtime/Engine/Private/HierarchicalInstancedStaticMesh.cpp
@@ -892,6 +892,39 @@ public:
 		, CaptureTag(0)
 #endif
 	{
+#if defined(GPUCULL_TODO)
+		if (InstancedRenderData.PerInstanceRenderData.IsValid() && InstancedRenderData.PerInstanceRenderData->InstanceBuffer_GameThread.IsValid())
+		{
+			const FStaticMeshInstanceData& StaticMeshInstanceData = *InstancedRenderData.PerInstanceRenderData->InstanceBuffer_GameThread;
+			
+			Instances.SetNum(StaticMeshInstanceData.GetNumInstances());
+			if (Instances.Num() == 0 && ClusterTree.Num() != 0)
+			{
+				UE_LOG(LogTemp, Display, TEXT("Instances.Num() == 0 && ClusterTree.Num() != 0"));
+			}
+
+			for (int32 InstanceIndex = 0; InstanceIndex < Instances.Num(); ++InstanceIndex)
+			{
+				FMatrix InstanceTransform;
+				StaticMeshInstanceData.GetInstanceTransform(InstanceIndex, InstanceTransform);
+				InstanceTransform.M[3][3] = 1.0f;
+
+				FPrimitiveInstance& Instance = Instances[InstanceIndex];
+				Instance.PrimitiveId = ~uint32(0);
+				Instance.InstanceToLocal = InstanceTransform;
+				Instance.LocalToInstance = InstanceTransform.Inverse();
+				// Filled in during GPU-Scene update...
+				Instance.LocalToWorld.SetIdentity();
+				Instance.RenderBounds = InComponent->GetStaticMesh()->GetBounds();
+				Instance.LocalBounds = Instance.RenderBounds.TransformBy(Instance.InstanceToLocal);
+			}
+		}
+		else
+		{
+			ensure(ClusterTree.Num() == 0);
+		}
+#endif // defined(GPUCULL_TODO)
+
 		SetupOcclusion(InComponent);
 	}
 
@@ -1337,6 +1370,77 @@ void FHierarchicalStaticMeshSceneProxy::FillDynamicMeshElements(FMeshElementColl
 					continue;
 				}
 
+
+#if defined(GPUCULL_TODO)
+				FMeshBatch& MeshBatch = Collector.AllocateMesh();
+				INC_DWORD_STAT(STAT_FoliageMeshBatches);
+
+				if (!FStaticMeshSceneProxy::GetMeshElement(LODIndex, 0, SectionIndex, GetDepthPriorityGroup(ElementParams.View), ElementParams.BatchRenderSelection[SelectionGroupIndex], true, MeshBatch))
+				{
+					continue;
+				}
+
+				checkSlow(MeshBatch.GetNumPrimitives() > 0);
+				MeshBatch.bCanApplyViewModeOverrides = true;
+				MeshBatch.bUseSelectionOutline = ElementParams.BatchRenderSelection[SelectionGroupIndex];
+				MeshBatch.bUseWireframeSelectionColoring = ElementParams.BatchRenderSelection[SelectionGroupIndex];
+				MeshBatch.bUseAsOccluder = ShouldUseAsOccluder();
+				MeshBatch.VertexFactory = &InstancedRenderData.VertexFactories[LODIndex];
+
+				FMeshBatchElement& MeshBatchElement = MeshBatch.Elements[0];
+				MeshBatchElement.UserData = ElementParams.PassUserData[SelectionGroupIndex];
+				MeshBatchElement.bUserDataIsColorVertexBuffer = false;
+				MeshBatchElement.MaxScreenSize = 1.0;
+				MeshBatchElement.MinScreenSize = 0.0;
+				MeshBatchElement.InstancedLODIndex = LODIndex;
+				MeshBatchElement.InstancedLODRange = bDitherLODEnabled ? 1 : 0;
+				MeshBatchElement.PrimitiveUniformBuffer = GetUniformBuffer();
+
+				int32 TotalInstances = bDitherLODEnabled ? Params.TotalMultipleLODInstances[LODIndex] : Params.TotalSingleLODInstances[LODIndex];
+				{
+					const int64 Tris = int64(TotalInstances) * int64(MeshBatchElement.NumPrimitives);
+					TotalTriangles += Tris;
+#if STATS
+					if (GFrameNumberRenderThread_CaptureFoliageRuns == GFrameNumberRenderThread)
+					{
+						if (ElementParams.FinalCullDistance > 9.9E8)
+						{
+							UE_LOG(LogStaticMesh, Display, TEXT("lod:%1d/%1d   sel:%1d   section:%1d/%1d   runs:%4d   inst:%8d   tris:%9lld   cast shadow:%1d   cull:-NONE!!-   shadow:%1d     %s %s"),
+								LODIndex, InstancedRenderData.VertexFactories.Num(), SelectionGroupIndex, SectionIndex, LODModel.Sections.Num(), RunArray.Num() / 2,
+								TotalInstances, Tris, (int)MeshBatch.CastShadow, ElementParams.ShadowFrustum,
+								*StaticMesh->GetPathName(),
+								*MeshBatch.MaterialRenderProxy->GetMaterial(ElementParams.FeatureLevel)->GetFriendlyName());
+						}
+						else
+						{
+							UE_LOG(LogStaticMesh, Display, TEXT("lod:%1d/%1d   sel:%1d   section:%1d/%1d   runs:%4d   inst:%8d   tris:%9lld   cast shadow:%1d   cull:%8.0f   shadow:%1d     %s %s"),
+								LODIndex, InstancedRenderData.VertexFactories.Num(), SelectionGroupIndex, SectionIndex, LODModel.Sections.Num(), RunArray.Num() / 2,
+								TotalInstances, Tris, (int)MeshBatch.CastShadow, ElementParams.FinalCullDistance, ElementParams.ShadowFrustum,
+								*StaticMesh->GetPathName(),
+								*MeshBatch.MaterialRenderProxy->GetMaterial(ElementParams.FeatureLevel)->GetFriendlyName());
+						}
+					}
+#endif // STATS
+				}
+
+				//MeshBatchElement.NumInstances = TotalInstances;
+				// The index was used as an offset, but the dynamic buffer thing uses a resource view to make this not needed (using PrimitiveInstanceDataOffset as a temp. debug help)
+				MeshBatchElement.UserIndex = 0;
+
+				// Note: this call overrides the UserIndex to mean the command index, which is used to fetch the offset to the instance array
+				//Collector.AllocateInstancedBatchArguments(ElementParams.ViewIndex, MeshBatch, PrimitiveInstanceDataOffset, PrimitiveInstanceDataCount, RunArray);
+				
+				// We use this existing hook to send info about the runs over to the visiblemeshbatch
+				MeshBatchElement.NumInstances = RunArray.Num() / 2;
+				MeshBatchElement.InstanceRuns = &RunArray[0];
+				MeshBatchElement.bIsInstanceRuns = true;
+
+				if (TotalTriangles < (int64)CVarMaxTrianglesToRender.GetValueOnRenderThread())
+				{
+					Collector.AddMesh(ElementParams.ViewIndex, MeshBatch);
+				}
+
+#else // !defined(GPUCULL_TODO)
 				int32 NumBatches = 1;
 				int32 CurrentRun = 0;
 				int32 CurrentInstance = 0;
@@ -1452,6 +1556,7 @@ void FHierarchicalStaticMeshSceneProxy::FillDynamicMeshElements(FMeshElementColl
 						Collector.AddMesh(ElementParams.ViewIndex, MeshElement);
 					}
 				}
+#endif // defined(GPUCULL_TODO)
 			}
 		}
 	}
@@ -1837,20 +1942,20 @@ void FHierarchicalStaticMeshSceneProxy::GetDynamicMeshElements(const TArray<cons
 							{
 								if (MinLOD < NumLODs)
 								{
-									InstanceParams.AddRun(MinLOD, MinLOD, FirstIndexInRun + FirstUnbuiltIndex, (Index - 1) + FirstUnbuiltIndex);
+									InstanceParams.AddRun(MinLOD, MinLOD, FirstIndexInRun + FirstUnbuiltIndex, (Index - 1) + FirstUnbuiltIndex - 1);
 								}
 								MinLOD = TempMinLOD;
 								FirstIndexInRun = Index;
 							}
 						}
-						InstanceParams.AddRun(MinLOD, MinLOD, FirstIndexInRun + FirstUnbuiltIndex, FirstIndexInRun + FirstUnbuiltIndex + UnbuiltInstanceCount);
+						InstanceParams.AddRun(MinLOD, MinLOD, FirstIndexInRun + FirstUnbuiltIndex, FirstUnbuiltIndex + UnbuiltInstanceCount - 1);
 					}
 				}
 				else
 				{
 					// more than 1000, render them all at lowest LOD (until we have an updated tree)
 					const int8 LowestLOD = (RenderData->LODResources.Num() - 1);
-					InstanceParams.AddRun(LowestLOD, LowestLOD, FirstUnbuiltIndex, FirstUnbuiltIndex + UnbuiltInstanceCount);
+					InstanceParams.AddRun(LowestLOD, LowestLOD, FirstUnbuiltIndex, FirstUnbuiltIndex + UnbuiltInstanceCount - 1);
 				}
 				FillDynamicMeshElements(Collector, ElementParams, InstanceParams);
 			}
@@ -2290,7 +2395,6 @@ bool UHierarchicalInstancedStaticMeshComponent::SetCustomDataValue(int32 Instanc
 	}
 
 	int32 RenderIndex = InstanceReorderTable.IsValidIndex(InstanceIndex) ? InstanceReorderTable[InstanceIndex] : InstanceIndex;
-	const float OldCustomDataValue = PerInstanceSMCustomData[InstanceIndex*NumCustomDataFloats+CustomDataIndex];
 
 	// if we are only updating rotation/scale we update the instance directly in the cluster tree
 	const bool bIsOmittedInstance = (RenderIndex == INDEX_NONE);
diff --git a/Engine/Source/Runtime/Engine/Private/LocalVertexFactory.cpp b/Engine/Source/Runtime/Engine/Private/LocalVertexFactory.cpp
--- a/Engine/Source/Runtime/Engine/Private/LocalVertexFactory.cpp
+++ b/Engine/Source/Runtime/Engine/Private/LocalVertexFactory.cpp
@@ -292,7 +292,7 @@ void FLocalVertexFactory::InitRHI()
 			if (GetType()->SupportsPrimitiveIdStream() && bCanUseGPUScene)
 			{
 				// When the VF is used for rendering in normal mesh passes, this vertex buffer and offset will be overridden
-				StreamElements.Add(AccessStreamComponent(FVertexStreamComponent(&GPrimitiveIdDummy, 0, 0, sizeof(uint32), VET_UInt, EVertexStreamUsage::Instancing), 1, InputStreamType));
+				StreamElements.Add(AccessStreamComponent(FVertexStreamComponent(&GPrimitiveIdDummy, 0, 0, PrimitiveIdStreamStride, VET_UInt, EVertexStreamUsage::Instancing), 1, InputStreamType));
 				PrimitiveIdStreamIndex[TypeIndex] = StreamElements.Last().StreamIndex;
 			}
 
@@ -315,7 +315,7 @@ void FLocalVertexFactory::InitRHI()
 		if (GetType()->SupportsPrimitiveIdStream() && bCanUseGPUScene)
 		{
 			// When the VF is used for rendering in normal mesh passes, this vertex buffer and offset will be overridden
-			Elements.Add(AccessStreamComponent(FVertexStreamComponent(&GPrimitiveIdDummy, 0, 0, sizeof(uint32), VET_UInt, EVertexStreamUsage::Instancing), 13));
+			Elements.Add(AccessStreamComponent(FVertexStreamComponent(&GPrimitiveIdDummy, 0, 0, PrimitiveIdStreamStride, VET_UInt, EVertexStreamUsage::Instancing), 13));
 			PrimitiveIdStreamIndex[Index] = Elements.Last().StreamIndex;
 		}
 	}
diff --git a/Engine/Source/Runtime/Engine/Private/StaticMeshRender.cpp b/Engine/Source/Runtime/Engine/Private/StaticMeshRender.cpp
--- a/Engine/Source/Runtime/Engine/Private/StaticMeshRender.cpp
+++ b/Engine/Source/Runtime/Engine/Private/StaticMeshRender.cpp
@@ -342,6 +342,17 @@ FStaticMeshSceneProxy::FStaticMeshSceneProxy(UStaticMeshComponent* InComponent,
 #endif
 
 	AddSpeedTreeWind();
+#if defined(GPUCULL_TODO)
+	Instances.SetNum(1);
+	FPrimitiveInstance& Instance = Instances[0];
+	Instance.PrimitiveId = ~uint32(0);
+	Instance.InstanceToLocal.SetIdentity();
+	Instance.LocalToInstance.SetIdentity();
+	Instance.LocalToWorld.SetIdentity();
+	Instance.RenderBounds = InComponent->GetStaticMesh()->GetBounds();
+	Instance.LocalBounds = Instance.RenderBounds;
+	bSupportsInstanceDataBuffer = true;
+#endif // defined(GPUCULL_TODO)
 }
 
 void FStaticMeshSceneProxy::SetEvaluateWorldPositionOffsetInRayTracing(bool NewValue)
diff --git a/Engine/Source/Runtime/Engine/Public/StaticMeshResources.h b/Engine/Source/Runtime/Engine/Public/StaticMeshResources.h
--- a/Engine/Source/Runtime/Engine/Public/StaticMeshResources.h
+++ b/Engine/Source/Runtime/Engine/Public/StaticMeshResources.h
@@ -1322,7 +1322,13 @@ public:
 		Transform.M[3][0] = Origin.X;
 		Transform.M[3][1] = Origin.Y;
 		Transform.M[3][2] = Origin.Z;
+#if 1 
+		// GPUCULL_TODO: This causes problems when using this as a normal real proper matrix and concatenating!
 		Transform.M[3][3] = 0.f;
+#else
+		// But this on the other hand seems to break some assumption or other in the HISM
+		Transform.M[3][3] = 1.0f;
+#endif
 	}
 
 	FORCEINLINE_DEBUGGABLE void GetInstanceShaderValues(int32 InstanceIndex, FVector4 (&InstanceTransform)[3], FVector4& InstanceLightmapAndShadowMapUVBias, FVector4& InstanceOrigin) const
diff --git a/Engine/Source/Runtime/RenderCore/Public/VertexFactory.h b/Engine/Source/Runtime/RenderCore/Public/VertexFactory.h
--- a/Engine/Source/Runtime/RenderCore/Public/VertexFactory.h
+++ b/Engine/Source/Runtime/RenderCore/Public/VertexFactory.h
@@ -740,6 +740,12 @@ protected:
 
 	int8 PrimitiveIdStreamIndex[static_cast<int>(EVertexInputStreamType::Count)];
 
+#if defined(GPUCULL_TODO)
+	static constexpr int32 PrimitiveIdStreamStride = 0;
+#else // !defined(GPUCULL_TODO)
+	static constexpr int32 PrimitiveIdStreamStride = sizeof(uint32);
+#endif // defined(GPUCULL_TODO)
+
 private:
 
 	/** The position only vertex stream used to render the factory during depth only passes. */
diff --git a/Engine/Source/Runtime/Renderer/Private/MeshDrawCommands.cpp b/Engine/Source/Runtime/Renderer/Private/MeshDrawCommands.cpp
--- a/Engine/Source/Runtime/Renderer/Private/MeshDrawCommands.cpp
+++ b/Engine/Source/Runtime/Renderer/Private/MeshDrawCommands.cpp
@@ -441,7 +441,13 @@ void BuildMeshDrawCommandPrimitiveIdBuffer(
 						VisibleMeshDrawCommand.StateBucketId,
 						VisibleMeshDrawCommand.MeshFillMode,
 						VisibleMeshDrawCommand.MeshCullMode,
+#if defined(GPUCULL_TODO)
+						VisibleMeshDrawCommand.SortKey,
+						VisibleMeshDrawCommand.RunArray,
+						VisibleMeshDrawCommand.NumRuns);
+#else //!defined(GPUCULL_TODO)
 						VisibleMeshDrawCommand.SortKey);
+#endif // defined(GPUCULL_TODO)
 
 					NewVisibleMeshDrawCommand.PrimitiveIdBufferOffset = PrimitiveIdIndex;
 					TempVisibleMeshDrawCommands.Emplace(MoveTemp(NewVisibleMeshDrawCommand));
@@ -805,7 +811,13 @@ void ApplyViewOverridesToMeshDrawCommands(
 					VisibleMeshDrawCommand.StateBucketId,
 					VisibleMeshDrawCommand.MeshFillMode,
 					VisibleMeshDrawCommand.MeshCullMode,
+#if defined(GPUCULL_TODO)
+					VisibleMeshDrawCommand.SortKey,
+					VisibleMeshDrawCommand.RunArray,
+					VisibleMeshDrawCommand.NumRuns);
+#else //!defined(GPUCULL_TODO)
 					VisibleMeshDrawCommand.SortKey);
+#endif // defined(GPUCULL_TODO)
 
 				TempVisibleMeshDrawCommands.Add(NewVisibleMeshDrawCommand);
 			}
@@ -952,6 +964,10 @@ public:
 
 			if (Context.bUseGPUScene)
 			{
+#if defined(GPUCULL_TODO)
+				// GPUCULL_TODO: Make a switch to control old / new behaviour, determine minimum reqs.
+				SetupGPUInstancedDraws(*Context.InstanceCullingContext, Context.MeshDrawCommands, Context.MaxInstances, Context.VisibleMeshDrawCommandsNum, Context.NewPassVisibleMeshDrawCommandsNum);
+#else //!defined(GPUCULL_TODO)
 				BuildMeshDrawCommandPrimitiveIdBuffer(
 					Context.bDynamicInstancing,
 					Context.MeshDrawCommands,
@@ -965,6 +981,7 @@ public:
 					Context.ShaderPlatform,
 					Context.InstanceFactor
 				);
+#endif // defined(GPUCULL_TODO)
 			}
 		}
 	}
@@ -1158,8 +1175,16 @@ void FParallelMeshDrawCommandPass::DispatchPassSetup(
 	TaskContext.NumDynamicMeshCommandBuildRequestElements = NumDynamicMeshCommandBuildRequestElements;
 
 	// Only apply instancing for ISR to main view passes
+#if defined(GPUCULL_TODO)
+	//const bool bIsMainViewPass = PassType != EMeshPass::Num && (FPassProcessorManager::GetPassFlags(TaskContext.ShadingPath, TaskContext.PassType) & EMeshPassFlags::MainView) != EMeshPassFlags::None;
+
+	// GPUCULL_TODO: Instance Factor should always be 1 when using GPU-side culling. Instead we'll make that stereo-aware (or something like that) - all that means is that the culling pass duplicates the 
+	// primitives (instances really) and culls them against the per-eye view etc. Some routing info must be passed along also to tell them the eye. Can probably borrow a bit somewhere.
+	TaskContext.InstanceFactor = 1;//(bIsMainViewPass && View.IsInstancedStereoPass()) ? 2 : 1;
+#else //!defined(GPUCULL_TODO)
 	const bool bIsMainViewPass = PassType != EMeshPass::Num && (FPassProcessorManager::GetPassFlags(TaskContext.ShadingPath, TaskContext.PassType) & EMeshPassFlags::MainView) != EMeshPassFlags::None;
 	TaskContext.InstanceFactor = (bIsMainViewPass && View.IsInstancedStereoPass()) ? 2 : 1;
+#endif // defined(GPUCULL_TODO)
 
 	TaskContext.InstanceCullingContext = InstanceCullingContext; 
 	TaskContext.InstanceCullingManager = InstanceCullingManager;
@@ -1343,10 +1368,15 @@ class FDrawVisibleMeshCommandsAnyThreadTask : public FRenderTask
 	FRHICommandList& RHICmdList;
 	const FMeshCommandOneFrameArray& VisibleMeshDrawCommands;
 	const FGraphicsMinimalPipelineStateSet& GraphicsMinimalPipelineStateSet;
+#if defined(GPUCULL_TODO)
+	FRHIVertexBuffer* InstanceIdOffsetBuffer;
+	FRHIVertexBuffer* DrawIndirectArgsBuffer;
+#else //!defined(GPUCULL_TODO)
 	FRHIVertexBuffer* PrimitiveIdsBuffer;
 	int32 BasePrimitiveIdsOffset;
 	bool bDynamicInstancing;
 	uint32 InstanceFactor;
+#endif // defined(GPUCULL_TODO)
 	int32 TaskIndex;
 	int32 TaskNum;
 
@@ -1356,20 +1386,30 @@ public:
 		FRHICommandList& InRHICmdList,
 		const FMeshCommandOneFrameArray& InVisibleMeshDrawCommands,
 		const FGraphicsMinimalPipelineStateSet& InGraphicsMinimalPipelineStateSet,
+#if defined(GPUCULL_TODO)
+		FRHIVertexBuffer* InInstanceIdOffsetBuffer,
+		FRHIVertexBuffer* InDrawIndirectArgsBuffer,
+#else //!defined(GPUCULL_TODO)
 		FRHIVertexBuffer* InPrimitiveIdsBuffer,
 		int32 InBasePrimitiveIdsOffset,
 		bool bInDynamicInstancing,
 		uint32 InInstanceFactor,
+#endif // defined(GPUCULL_TODO)
 		int32 InTaskIndex,
 		int32 InTaskNum
 	)
 		: RHICmdList(InRHICmdList)
 		, VisibleMeshDrawCommands(InVisibleMeshDrawCommands)
 		, GraphicsMinimalPipelineStateSet(InGraphicsMinimalPipelineStateSet)
+#if defined(GPUCULL_TODO)
+		, InstanceIdOffsetBuffer(InInstanceIdOffsetBuffer)
+		, DrawIndirectArgsBuffer(InDrawIndirectArgsBuffer)
+#else //!defined(GPUCULL_TODO)
 		, PrimitiveIdsBuffer(InPrimitiveIdsBuffer)
 		, BasePrimitiveIdsOffset(InBasePrimitiveIdsOffset)
 		, bDynamicInstancing(bInDynamicInstancing)
 		, InstanceFactor(InInstanceFactor)
+#endif // defined(GPUCULL_TODO)
 		, TaskIndex(InTaskIndex)
 		, TaskNum(InTaskNum)
 	{}
@@ -1395,46 +1435,23 @@ public:
 		const int32 NumDrawsPerTask = TaskIndex < DrawNum ? FMath::DivideAndRoundUp(DrawNum, TaskNum) : 0;
 		const int32 StartIndex = TaskIndex * NumDrawsPerTask;
 		const int32 NumDraws = FMath::Min(NumDrawsPerTask, DrawNum - StartIndex);
-
+#if defined(GPUCULL_TODO)
+		SubmitGPUInstancedMeshDrawCommandsRange(
+			VisibleMeshDrawCommands,
+			GraphicsMinimalPipelineStateSet,
+			StartIndex,
+			NumDraws,
+			InstanceIdOffsetBuffer,
+			DrawIndirectArgsBuffer,
+			RHICmdList);
+#else // !defined(GPUCULL_TODO)
 		SubmitMeshDrawCommandsRange(VisibleMeshDrawCommands, GraphicsMinimalPipelineStateSet, PrimitiveIdsBuffer, BasePrimitiveIdsOffset, bDynamicInstancing, StartIndex, NumDraws, InstanceFactor, RHICmdList);
-
+#endif // defined(GPUCULL_TODO)
 		RHICmdList.EndRenderPass();
 		RHICmdList.HandleRTThreadTaskCompletion(MyCompletionGraphEvent);
 	}
 };
 
-#if 0 // GPUCULL_TODO
-void FParallelMeshDrawCommandPass::DispatchInstanceCullingFinalize(FRHICommandList& RHICmdList)
-{
-#if 0
-
-	// 1. Add pre-req on the setup pass
-	if (TaskEventRef.IsValid())
-	{
-		RHICmdList.AddDispatchPrerequisite(TaskEventRef);
-	}
-	// 2. In the same way that the DispatchDraw did, queue up commands to upload the input data to the culling unpack - this will happen after the AddDispatchPrerequisite has been reached, so would be fine.
-	//    At this point we must have stable pointers such that none of the data moves, it must also be passed over to be owned by the RHI thread so it can delete the data.
-
-	RHICmdList.EnqueueLambda([
-			VertexBuffer = PrimitiveIdsBuffer,
-			VertexBufferData = TaskContext.PrimitiveIdBufferData,
-			VertexBufferDataSize = TaskContext.PrimitiveIdBufferDataSize,
-			PrimitiveIdVertexBufferPoolEntry = PrimitiveIdVertexBufferPoolEntry](FRHICommandListImmediate& CmdList)
-		{
-			// Upload vertex buffer data.
-			void* RESTRICT Data = (void* RESTRICT)CmdList.LockVertexBuffer(VertexBuffer, 0, VertexBufferDataSize, RLM_WriteOnly);
-			FMemory::Memcpy(Data, VertexBufferData, VertexBufferDataSize);
-			CmdList.UnlockVertexBuffer(VertexBuffer);
-
-			FMemory::Free(VertexBufferData);
-		});
-
-	RHICommandList.RHIThreadFence(true);
-#endif
-}
-#endif // GPUCULL_TODO
-
 void FParallelMeshDrawCommandPass::BuildRenderingCommands(FRDGBuilder& GraphBuilder, FGPUScene& GPUScene, FInstanceCullingDrawParams& OutInstanceCullingDrawParams)
 {
 #if defined(GPUCULL_TODO)
@@ -1474,6 +1491,18 @@ void FParallelMeshDrawCommandPass::DispatchDraw(FParallelCommandListSet* Paralle
 		return;
 	}
 
+#if defined(GPUCULL_TODO)
+	FRHIVertexBuffer* DrawIndirectArgsBuffer = nullptr;
+	FRHIVertexBuffer* InstanceIdOffsetBuffer = nullptr;
+	if (InstanceCullingDrawParams != nullptr 
+		&& InstanceCullingDrawParams->DrawIndirectArgsBuffer != nullptr 
+		&& InstanceCullingDrawParams->InstanceIdOffsetBuffer != nullptr)
+	{
+		DrawIndirectArgsBuffer = InstanceCullingDrawParams->DrawIndirectArgsBuffer->GetRHI();
+		InstanceIdOffsetBuffer = InstanceCullingDrawParams->InstanceIdOffsetBuffer->GetRHI();
+	}
+#endif // defined(GPUCULL_TODO)
+
 	FRHIVertexBuffer* PrimitiveIdsBuffer = PrimitiveIdVertexBufferPoolEntry.BufferRHI;
 	const int32 BasePrimitiveIdsOffset = 0;
 
@@ -1535,7 +1564,14 @@ void FParallelMeshDrawCommandPass::DispatchDraw(FParallelCommandListSet* Paralle
 			FRHICommandList* CmdList = ParallelCommandListSet->NewParallelCommandList();
 
 			FGraphEventRef AnyThreadCompletionEvent = TGraphTask<FDrawVisibleMeshCommandsAnyThreadTask>::CreateTask(&Prereqs, RenderThread)
+#if defined(GPUCULL_TODO)
+				.ConstructAndDispatchWhenReady(*CmdList, TaskContext.MeshDrawCommands, TaskContext.MinimalPipelineStatePassSet, 
+					InstanceIdOffsetBuffer,
+					DrawIndirectArgsBuffer,
+					TaskIndex, NumTasks);
+#else //!defined(GPUCULL_TODO)
 				.ConstructAndDispatchWhenReady(*CmdList, TaskContext.MeshDrawCommands, TaskContext.MinimalPipelineStatePassSet, PrimitiveIdsBuffer, BasePrimitiveIdsOffset, TaskContext.bDynamicInstancing, TaskContext.InstanceFactor, TaskIndex, NumTasks);
+#endif // defined(GPUCULL_TODO)
 			ParallelCommandListSet->AddParallelCommandList(CmdList, AnyThreadCompletionEvent, NumDraws);
 		}
 	}
@@ -1545,6 +1581,26 @@ void FParallelMeshDrawCommandPass::DispatchDraw(FParallelCommandListSet* Paralle
 
 		WaitForMeshPassSetupTask();
 
+#if defined(GPUCULL_TODO)
+		if (TaskContext.bUseGPUScene)
+		{
+			if (TaskContext.MeshDrawCommands.Num() > 0)
+			{
+				SubmitGPUInstancedMeshDrawCommandsRange(
+					TaskContext.MeshDrawCommands,
+					TaskContext.MinimalPipelineStatePassSet,
+					0,
+					TaskContext.MeshDrawCommands.Num(),
+					InstanceIdOffsetBuffer,
+					DrawIndirectArgsBuffer,
+					RHICmdList);
+			}
+		}
+		else
+		{
+			SubmitMeshDrawCommandsRange(TaskContext.MeshDrawCommands, TaskContext.MinimalPipelineStatePassSet, nullptr, 0, TaskContext.bDynamicInstancing, 0, TaskContext.MeshDrawCommands.Num(), TaskContext.InstanceFactor, RHICmdList);
+		}
+#else //!defined(GPUCULL_TODO)
 		if (TaskContext.bUseGPUScene)
 		{
 			// Can immediately upload vertex buffer data, as there is no parallel draw task.
@@ -1554,6 +1610,7 @@ void FParallelMeshDrawCommandPass::DispatchDraw(FParallelCommandListSet* Paralle
 		}
 
 		SubmitMeshDrawCommandsRange(TaskContext.MeshDrawCommands, TaskContext.MinimalPipelineStatePassSet, PrimitiveIdsBuffer, BasePrimitiveIdsOffset, TaskContext.bDynamicInstancing, 0, TaskContext.MeshDrawCommands.Num(), TaskContext.InstanceFactor, RHICmdList);
+#endif // defined(GPUCULL_TODO)
 	}
 }
 
diff --git a/Engine/Source/Runtime/Renderer/Private/MeshPassProcessor.cpp b/Engine/Source/Runtime/Renderer/Private/MeshPassProcessor.cpp
--- a/Engine/Source/Runtime/Renderer/Private/MeshPassProcessor.cpp
+++ b/Engine/Source/Runtime/Renderer/Private/MeshPassProcessor.cpp
@@ -779,7 +779,9 @@ void FMeshDrawCommand::SetDrawParametersAndFinalize(
 	const FMeshBatchElement& BatchElement = MeshBatch.Elements[BatchElementIndex];
 
 	check(!BatchElement.IndexBuffer || (BatchElement.IndexBuffer && BatchElement.IndexBuffer->IsInitialized() && BatchElement.IndexBuffer->IndexBufferRHI));
+#if !defined(GPUCULL_TODO)
 	checkSlow(!BatchElement.bIsInstanceRuns);
+#endif // defined(GPUCULL_TODO)
 	IndexBuffer = BatchElement.IndexBuffer ? BatchElement.IndexBuffer->IndexBufferRHI.GetReference() : nullptr;
 	FirstIndex = BatchElement.FirstIndex;
 	NumPrimitives = BatchElement.NumPrimitives;
@@ -1035,6 +1037,10 @@ void FMeshDrawCommand::SubmitDrawBegin(
 	FMeshDrawCommandStateCache& RESTRICT StateCache)
 {
 	checkSlow(MeshDrawCommand.CachedPipelineId.IsValid());
+#if defined(GPUCULL_TODO)
+	// Want to be sure that we supply GPU-scene instance data if required.
+	checkSlow(MeshDrawCommand.PrimitiveIdStreamIndex == -1 || ScenePrimitiveIdsBuffer != nullptr);
+#endif // defined(GPUCULL_TODO)
 
 #if WANTS_DRAW_MESH_EVENTS
 	FDrawEvent MeshEvent;
@@ -1103,11 +1109,15 @@ void FMeshDrawCommand::SubmitDrawBegin(
 	MeshDrawCommand.ShaderBindings.SetOnCommandList(RHICmdList, MeshPipelineState.BoundShaderState.AsBoundShaderState(), StateCache.ShaderBindings);
 }
 
-void FMeshDrawCommand::SubmitDrawEnd(const FMeshDrawCommand& MeshDrawCommand, uint32 InstanceFactor, FRHICommandList& RHICmdList)
+void FMeshDrawCommand::SubmitDrawEnd(const FMeshDrawCommand& MeshDrawCommand, uint32 InstanceFactor, FRHICommandList& RHICmdList,
+	FRHIVertexBuffer* IndirectArgsOverrideBuffer,
+	uint32 IndirectArgsOverrideByteOffset)
 {
+	const bool bDoOverrideArgs = IndirectArgsOverrideBuffer != nullptr && MeshDrawCommand.PrimitiveIdStreamIndex >= 0;
+
 	if (MeshDrawCommand.IndexBuffer)
 	{
-		if (MeshDrawCommand.NumPrimitives > 0)
+		if (MeshDrawCommand.NumPrimitives > 0 && !bDoOverrideArgs)
 		{
 			RHICmdList.DrawIndexedPrimitive(
 				MeshDrawCommand.IndexBuffer,
@@ -1123,14 +1133,14 @@ void FMeshDrawCommand::SubmitDrawEnd(const FMeshDrawCommand& MeshDrawCommand, ui
 		{
 			RHICmdList.DrawIndexedPrimitiveIndirect(
 				MeshDrawCommand.IndexBuffer,
-				MeshDrawCommand.IndirectArgs.Buffer,
-				MeshDrawCommand.IndirectArgs.Offset
+				bDoOverrideArgs ? IndirectArgsOverrideBuffer : MeshDrawCommand.IndirectArgs.Buffer,
+				bDoOverrideArgs ? IndirectArgsOverrideByteOffset : MeshDrawCommand.IndirectArgs.Offset
 			);
 		}
 	}
 	else
 	{
-		if (MeshDrawCommand.NumPrimitives > 0)
+		if (MeshDrawCommand.NumPrimitives > 0 && !bDoOverrideArgs)
 		{
 			RHICmdList.DrawPrimitive(
 				MeshDrawCommand.VertexParams.BaseVertexIndex + MeshDrawCommand.FirstIndex,
@@ -1140,8 +1150,8 @@ void FMeshDrawCommand::SubmitDrawEnd(const FMeshDrawCommand& MeshDrawCommand, ui
 		else
 		{
 			RHICmdList.DrawPrimitiveIndirect(
-				MeshDrawCommand.IndirectArgs.Buffer,
-				MeshDrawCommand.IndirectArgs.Offset
+				bDoOverrideArgs ? IndirectArgsOverrideBuffer : MeshDrawCommand.IndirectArgs.Buffer,
+				bDoOverrideArgs ? IndirectArgsOverrideByteOffset : MeshDrawCommand.IndirectArgs.Offset
 			);
 		}
 	}
@@ -1154,10 +1164,12 @@ void FMeshDrawCommand::SubmitDraw(
 	int32 PrimitiveIdOffset,
 	uint32 InstanceFactor,
 	FRHICommandList& RHICmdList,
-	FMeshDrawCommandStateCache& RESTRICT StateCache)
+	FMeshDrawCommandStateCache& RESTRICT StateCache,
+	FRHIVertexBuffer* IndirectArgsOverrideBuffer,
+	uint32 IndirectArgsOverrideByteOffset)
 {
 	SubmitDrawBegin(MeshDrawCommand, GraphicsMinimalPipelineStateSet, ScenePrimitiveIdsBuffer, PrimitiveIdOffset, InstanceFactor, RHICmdList, StateCache);
-	SubmitDrawEnd(MeshDrawCommand, InstanceFactor, RHICmdList);
+	SubmitDrawEnd(MeshDrawCommand, InstanceFactor, RHICmdList, IndirectArgsOverrideBuffer, IndirectArgsOverrideByteOffset);
 }
 
 #if MESH_DRAW_COMMAND_DEBUG_DATA
@@ -1245,7 +1257,13 @@ void ApplyViewOverridesToMeshDrawCommands(const FSceneView& View, FMeshCommandOn
 				VisibleMeshDrawCommand.StateBucketId,
 				VisibleMeshDrawCommand.MeshFillMode,
 				VisibleMeshDrawCommand.MeshCullMode,
+#if defined(GPUCULL_TODO)
+				VisibleMeshDrawCommand.SortKey,
+				VisibleMeshDrawCommand.RunArray,
+				VisibleMeshDrawCommand.NumRuns);
+#else //!defined(GPUCULL_TODO)
 				VisibleMeshDrawCommand.SortKey);
+#endif // defined(GPUCULL_TODO)
 
 			ViewOverriddenMeshCommands.Add(NewVisibleMeshDrawCommand);
 		}
@@ -1277,35 +1295,6 @@ void DrawDynamicMeshPassPrivate(
 	}
 }
 
-void DrawDynamicMeshPassPrivate(
-	const FSceneView& View,
-	const FScene &Scene,
-	FRHICommandListImmediate& RHICmdList,
-	FMeshCommandOneFrameArray& VisibleMeshDrawCommands,
-	FDynamicMeshDrawCommandStorage& DynamicMeshDrawCommandStorage,
-	FGraphicsMinimalPipelineStateSet& GraphicsMinimalPipelineStateSet,
-	bool& InNeedsShaderInitialization)
-{
-#if defined(GPUCULL_TODO)
-	check(View.bIsViewInfo);
-
-	if (VisibleMeshDrawCommands.Num() > 0)
-	{
-		const bool bDynamicInstancing = IsDynamicInstancingEnabled(View.GetFeatureLevel());
-
-		FRHIVertexBuffer* PrimitiveIdVertexBuffer = nullptr;
-
-		ApplyViewOverridesToMeshDrawCommands(View, VisibleMeshDrawCommands, DynamicMeshDrawCommandStorage, GraphicsMinimalPipelineStateSet, InNeedsShaderInitialization);
-		
-		// GPUCULL_TODO: determine if unsupportable OR just not supporting scene primitives (as this requires a multi-pass approach)
-
-		//SortAndSubmitDynamicMeshPassDrawCommands(View, Scene, VisibleMeshDrawCommands, DynamicMeshDrawCommandStorage, GraphicsMinimalPipelineStateSet, RHICmdList);
-
-		//SortAndMergeDynamicPassMeshDrawCommands(View.GetFeatureLevel(), VisibleMeshDrawCommands, DynamicMeshDrawCommandStorage, PrimitiveIdVertexBuffer, InstanceFactor);
-		//SubmitMeshDrawCommandsRange(VisibleMeshDrawCommands, GraphicsMinimalPipelineStateSet, PrimitiveIdVertexBuffer, 0, bDynamicInstancing, 0, VisibleMeshDrawCommands.Num(), InstanceFactor, RHICmdList);
-	}
-#endif
-}
 
 FMeshDrawCommandSortKey CalculateMeshStaticSortKey(const FMeshMaterialShader* VertexShader, const FMeshMaterialShader* PixelShader)
 {
diff --git a/Engine/Source/Runtime/Renderer/Private/SceneRendering.cpp b/Engine/Source/Runtime/Renderer/Private/SceneRendering.cpp
--- a/Engine/Source/Runtime/Renderer/Private/SceneRendering.cpp
+++ b/Engine/Source/Runtime/Renderer/Private/SceneRendering.cpp
@@ -1751,10 +1751,12 @@ void FViewInfo::SetupUniformBufferParameters(
 		if (InstanceSceneDataOverrideSRV)
 		{
 			ViewUniformShaderParameters.InstanceSceneData = InstanceSceneDataOverrideSRV;
+			ViewUniformShaderParameters.InstanceDataSOAStride = 0;
 		}
 		else if (Scene && Scene->GPUScene.InstanceDataBuffer.SRV)
 		{
 			ViewUniformShaderParameters.InstanceSceneData = Scene->GPUScene.InstanceDataBuffer.SRV;
+			ViewUniformShaderParameters.InstanceDataSOAStride = Scene->GPUScene.InstanceDataSOAStride;
 		}
 
 		if (LightmapSceneDataOverrideSRV)
@@ -3214,7 +3216,17 @@ void FSceneRenderer::SetupMeshPass(FViewInfo& View, FExclusiveDepthStencil::Type
 				Pass.SetDumpInstancingStats(GetMeshPassName(PassType));
 			}
 
-			FInstanceCullingContext* InstanceCullingContext = nullptr; // GPUCULL_TODO: InstanceCullingManager.CreateContext(ViewIds, NumViews);
+			// Only apply instancing for ISR to main view passes
+			const bool bIsMainViewPass = PassType != EMeshPass::Num && (FPassProcessorManager::GetPassFlags(Scene->GetShadingPath(), PassType) & EMeshPassFlags::MainView) != EMeshPassFlags::None;
+			int32 NumViews = bIsMainViewPass && View.IsInstancedStereoPass() ? 2 : 1;
+			int32 ViewIds[2] =
+			{
+				View.GPUSceneViewId,
+				// GPUCULL_TODO: This is a total hack! Fix fix fix (somehow retrieve the other view or pass in)
+				View.GPUSceneViewId + 1,
+			};
+
+			FInstanceCullingContext* InstanceCullingContext = InstanceCullingManager.CreateContext(ViewIds, NumViews);
 
 			Pass.DispatchPassSetup(
 				Scene,
diff --git a/Engine/Source/Runtime/Renderer/Private/ShadowDepthRendering.cpp b/Engine/Source/Runtime/Renderer/Private/ShadowDepthRendering.cpp
--- a/Engine/Source/Runtime/Renderer/Private/ShadowDepthRendering.cpp
+++ b/Engine/Source/Runtime/Renderer/Private/ShadowDepthRendering.cpp
@@ -201,6 +201,9 @@ class FShadowDepthShaderElementData : public FMeshMaterialShaderElementData
 {
 public:
 	int32 LayerId;
+#if defined(GPUCULL_TODO)
+	int32 bUseGpuSceneInstancing;
+#endif // defined(GPUCULL_TODO)
 };
 
 /**
@@ -229,6 +232,9 @@ public:
 		FMeshMaterialShader::GetShaderBindings(Scene, FeatureLevel, PrimitiveSceneProxy, MaterialRenderProxy, Material, DrawRenderState, ShaderElementData, ShaderBindings);
 
 		ShaderBindings.Add(LayerId, ShaderElementData.LayerId);
+#if defined(GPUCULL_TODO)
+		ShaderBindings.Add(bUseGpuSceneInstancing, ShaderElementData.bUseGpuSceneInstancing);
+#endif // defined(GPUCULL_TODO)
 	}
 
 	FShadowDepthVS() = default;
@@ -236,10 +242,16 @@ public:
 		FMeshMaterialShader(Initializer)
 	{
 		LayerId.Bind(Initializer.ParameterMap, TEXT("LayerId"));
+#if defined(GPUCULL_TODO)
+		bUseGpuSceneInstancing.Bind(Initializer.ParameterMap, TEXT("bUseGpuSceneInstancing"));
+#endif // defined(GPUCULL_TODO)
 	}
 
 private:
 	LAYOUT_FIELD(FShaderParameter, LayerId);
+#if defined(GPUCULL_TODO)
+	LAYOUT_FIELD(FShaderParameter, bUseGpuSceneInstancing);
+#endif // defined(GPUCULL_TODO)
 };
 
 enum EShadowDepthVertexShaderMode
@@ -396,6 +408,22 @@ class FOnePassPointShadowDepthGS : public FMeshMaterialShader
 	DECLARE_SHADER_TYPE(FOnePassPointShadowDepthGS, MeshMaterial);
 public:
 
+#if defined(GPUCULL_TODO)
+	void GetShaderBindings(
+		const FScene* Scene,
+		ERHIFeatureLevel::Type FeatureLevel,
+		const FPrimitiveSceneProxy* PrimitiveSceneProxy,
+		const FMaterialRenderProxy& MaterialRenderProxy,
+		const FMaterial& Material,
+		const FMeshPassProcessorRenderState& DrawRenderState,
+		const FShadowDepthShaderElementData& ShaderElementData,
+		FMeshDrawSingleShaderBindings& ShaderBindings) const
+	{
+		FMeshMaterialShader::GetShaderBindings(Scene, FeatureLevel, PrimitiveSceneProxy, MaterialRenderProxy, Material, DrawRenderState, ShaderElementData, ShaderBindings);
+		ShaderBindings.Add(bUseGpuSceneInstancing, ShaderElementData.bUseGpuSceneInstancing);
+	}
+#endif // defined(GPUCULL_TODO)
+
 	static bool ShouldCompilePermutation(const FMeshMaterialShaderPermutationParameters& Parameters)
 	{
 		return RHISupportsGeometryShaders(Parameters.Platform) && TShadowDepthVS<VertexShadowDepth_OnePassPointLight, false, true>::ShouldCompilePermutation(Parameters);
@@ -422,9 +450,15 @@ public:
 		{
 			PassUniformBuffer.Bind(Initializer.ParameterMap, FMobileShadowDepthPassUniformParameters::StaticStructMetadata.GetShaderVariableName());
 		}
+#if defined(GPUCULL_TODO)
+		bUseGpuSceneInstancing.Bind(Initializer.ParameterMap, TEXT("bUseGpuSceneInstancing"));
+#endif // defined(GPUCULL_TODO)
 	}
 
 	FOnePassPointShadowDepthGS() {}
+#if defined(GPUCULL_TODO)
+	LAYOUT_FIELD(FShaderParameter, bUseGpuSceneInstancing);
+#endif // defined(GPUCULL_TODO)
 };
 
 IMPLEMENT_SHADER_TYPE(, FOnePassPointShadowDepthGS, TEXT("/Engine/Private/ShadowDepthVertexShader.usf"), TEXT("MainOnePassPointLightGS"), SF_Geometry);
@@ -2211,10 +2245,22 @@ bool FShadowDepthPassMeshProcessor::Process(
 
 	const FMeshDrawCommandSortKey SortKey = CalculateMeshStaticSortKey(ShadowDepthPassShaders.VertexShader, ShadowDepthPassShaders.PixelShader);
 
+#if defined(GPUCULL_TODO)
+	const bool bUseGeometryShader = GShadowUseGS && RHISupportsGeometryShaders(GShaderPlatformForFeatureLevel[FeatureLevel]);
+
+	const bool bUseGpuSceneInstancing = UseGPUScene(GShaderPlatformForFeatureLevel[FeatureLevel], FeatureLevel) 
+		&& VertexFactory->GetPrimitiveIdStreamIndex(bUsePositionOnlyVS ? EVertexInputStreamType::PositionAndNormalOnly : EVertexInputStreamType::Default) != INDEX_NONE;
+
+	const uint32 InstanceFactor = bUseGpuSceneInstancing || !ShadowDepthType.bOnePassPointLightShadow || bUseGeometryShader ? 1 : 6;
+#else //!defined(GPUCULL_TODO)
 	const uint32 InstanceFactor = !ShadowDepthType.bOnePassPointLightShadow || (GShadowUseGS && RHISupportsGeometryShaders(GShaderPlatformForFeatureLevel[FeatureLevel])) ? 1 : 6;
+#endif // defined(GPUCULL_TODO)
 	for (uint32 i = 0; i < InstanceFactor; i++)
 	{
 		ShaderElementData.LayerId = i;
+#if defined(GPUCULL_TODO)
+		ShaderElementData.bUseGpuSceneInstancing = bUseGpuSceneInstancing;
+#endif // defined(GPUCULL_TODO)
 
 		BuildMeshDrawCommands(
 			MeshBatch,
diff --git a/Engine/Source/Runtime/Renderer/Private/ShadowSetup.cpp b/Engine/Source/Runtime/Renderer/Private/ShadowSetup.cpp
--- a/Engine/Source/Runtime/Renderer/Private/ShadowSetup.cpp
+++ b/Engine/Source/Runtime/Renderer/Private/ShadowSetup.cpp
@@ -2013,13 +2013,62 @@ void FProjectedShadowInfo::SetupMeshDrawCommandsForShadowDepth(FSceneRenderer& R
 		ShadowDepthPass.SetDumpInstancingStats(TEXT("ShadowDepth ") + PassNameForStats);
 	}
 	
+
+	if (GetShadowDepthType().bOnePassPointLightShadow)
+	{
+		ViewIds.AddDefaulted(6);
+		for (int32 CubemapFaceIndex = 0; CubemapFaceIndex < 6; CubemapFaceIndex++)
+		{
+			// We always render to a whole face at once
+			const FIntRect ShadowViewRect = FIntRect(X, Y, ResolutionX, ResolutionY);
+			// Setup packed view
+			TArray<Nanite::FPackedView, SceneRenderingAllocator> PackedViews;
+			{
+				FViewMatrices::FMinimalInitializer MatricesInitializer;
+				MatricesInitializer.ViewOrigin = -PreShadowTranslation;
+				MatricesInitializer.ViewRotationMatrix = OnePassShadowViewMatrices[CubemapFaceIndex];
+				MatricesInitializer.ProjectionMatrix = OnePassShadowFaceProjectionMatrix;
+				MatricesInitializer.ConstrainedViewRect = ShadowViewRect;
+
+				Nanite::FPackedViewParams Params;
+				Params.ViewMatrices = FViewMatrices(MatricesInitializer);
+				// TODO: Real prev frame matrices
+				Params.PrevViewMatrices = Params.ViewMatrices;	
+				Params.ViewRect = ShadowViewRect;
+				Params.RasterContextSize = FIntPoint(ResolutionX, ResolutionY);
+				ViewIds[CubemapFaceIndex] = InstanceCullingManager.RegisterView(Params);
+			}
+		}
+	}
+	else
+	{
+		Nanite::FPackedViewParams Params;
+		Params.ViewMatrices = ShadowDepthView->ViewMatrices;
+		// TODO: Real prev frame matrices
+		Params.PrevViewMatrices = Params.ViewMatrices;
+		Params.ViewRect = GetInnerViewRect();
+		Params.RasterContextSize = FIntPoint(ResolutionX, ResolutionY);
+		ViewIds.Add(InstanceCullingManager.RegisterView(Params));
+	}
+	// GPUCULL_TODO: Pass along any custom culling planes or whatever here (e.g., cacade bounds):
+	// GPUCULL_TODO: Add debug tags to context and views (so compute passes can be understood)
+
+	FInstanceCullingContext* InstanceCullingContext = InstanceCullingManager.CreateContext(ViewIds.GetData(), ViewIds.Num());
+
 	extern int32 GShadowUseGS;
+#if defined(GPUCULL_TODO)
+	// GPUCULL_TODO: Needed to support legacy, non-GPU-Scene culled, primitives, this is merely used to allocate enough space for CPU-side replication.
+	const bool bUseGeometryShader = GShadowUseGS && RHISupportsGeometryShaders(Renderer.Scene->GetShaderPlatform());
+	// Note: Iteracts with FShadowDepthPassMeshProcessor::Process and must be an overestimate of the actual replication done there.
+	const uint32 InstanceFactor = !GetShadowDepthType().bOnePassPointLightShadow || bUseGeometryShader ? 1 : 6;
+#else //!defined(GPUCULL_TODO)
 	const uint32 InstanceFactor = !GetShadowDepthType().bOnePassPointLightShadow || (GShadowUseGS && RHISupportsGeometryShaders(Renderer.Scene->GetShaderPlatform())) ? 1 : 6;
+#endif // defined(GPUCULL_TODO)
 
 	ShadowDepthPass.DispatchPassSetup(
 		Renderer.Scene,
 		*ShadowDepthView,
-		nullptr, // GPUCULL_TODO: Culling context
+		InstanceCullingContext,
 		&InstanceCullingManager,
 		EMeshPass::Num,
 		FExclusiveDepthStencil::DepthNop_StencilNop,
diff --git a/Engine/Source/Runtime/Renderer/Public/MeshPassProcessor.h b/Engine/Source/Runtime/Renderer/Public/MeshPassProcessor.h
--- a/Engine/Source/Runtime/Renderer/Public/MeshPassProcessor.h
+++ b/Engine/Source/Runtime/Renderer/Public/MeshPassProcessor.h
@@ -1105,6 +1105,7 @@ public:
 	static void SubmitDrawBegin(
 		const FMeshDrawCommand& RESTRICT MeshDrawCommand,
 		const FGraphicsMinimalPipelineStateSet& GraphicsMinimalPipelineStateSet,
+		// GPUCULL_TODO: Rename, and probably wrap in struct that links to GPU-Scene, maybe generalize (probably not)?
 		FRHIVertexBuffer* ScenePrimitiveIdsBuffer,
 		int32 PrimitiveIdOffset,
 		uint32 InstanceFactor,
@@ -1112,17 +1113,23 @@ public:
 		FMeshDrawCommandStateCache& RESTRICT StateCache);
 
 	/** Submits just the draw primitive portion of the draw command. */
-	static void SubmitDrawEnd(const FMeshDrawCommand& MeshDrawCommand, uint32 InstanceFactor, FRHICommandList& RHICmdList);
+	static void SubmitDrawEnd(const FMeshDrawCommand& MeshDrawCommand, uint32 InstanceFactor, FRHICommandList& RHICmdList,
+		// GPUCULL_TODO: Rename, and probably wrap in struct that links to GPU-Scene, maybe generalize (probably not)?
+		FRHIVertexBuffer* IndirectArgsOverrideBuffer = nullptr,
+		uint32 IndirectArgsOverrideByteOffset = 0U);
 
 	/** Submits commands to the RHI Commandlist to draw the MeshDrawCommand. */
 	static void SubmitDraw(
 		const FMeshDrawCommand& RESTRICT MeshDrawCommand,
 		const FGraphicsMinimalPipelineStateSet& GraphicsMinimalPipelineStateSet,
+		// GPUCULL_TODO: Rename, and probably wrap in struct that links to GPU-Scene, maybe generalize (probably not)?
 		FRHIVertexBuffer* ScenePrimitiveIdsBuffer,
 		int32 PrimitiveIdOffset,
 		uint32 InstanceFactor,
 		FRHICommandList& CommandList,
-		class FMeshDrawCommandStateCache& RESTRICT StateCache);
+		class FMeshDrawCommandStateCache& RESTRICT StateCache,
+		FRHIVertexBuffer* IndirectArgsOverrideBuffer = nullptr,
+		uint32 IndirectArgsOverrideByteOffset = 0U);
 
 	FORCENOINLINE friend uint32 GetTypeHash( const FMeshDrawCommand& Other )
 	{
@@ -1246,7 +1253,13 @@ public:
 		int32 InStateBucketId,
 		ERasterizerFillMode InMeshFillMode,
 		ERasterizerCullMode InMeshCullMode,
+#if defined(GPUCULL_TODO)
+		FMeshDrawCommandSortKey InSortKey,
+		const uint32* InRunArray = nullptr,
+		int32 InNumRuns = 0)
+#else //!defined(GPUCULL_TODO)
 		FMeshDrawCommandSortKey InSortKey)
+#endif // defined(GPUCULL_TODO)
 	{
 		MeshDrawCommand = InMeshDrawCommand;
 		DrawPrimitiveId = InDrawPrimitiveId;
@@ -1256,6 +1269,10 @@ public:
 		MeshFillMode = InMeshFillMode;
 		MeshCullMode = InMeshCullMode;
 		SortKey = InSortKey;
+#if defined(GPUCULL_TODO)
+		RunArray = InRunArray;
+		NumRuns = InNumRuns;
+#endif // defined(GPUCULL_TODO)
 	}
 
 	// Mesh Draw Command stored separately to avoid fetching its data during sorting
@@ -1278,6 +1295,11 @@ public:
 	// Any commands with the same StateBucketId can be merged into one draw call with instancing.
 	// A value of -1 means the draw is not in any state bucket and should be sorted by other factors instead.
 	int32 StateBucketId;
+#if defined(GPUCULL_TODO)
+	// Used for passing sub-selection of instances through to the culling
+	const uint32* RunArray;
+	int32 NumRuns;
+#endif // defined(GPUCULL_TODO)
 
 	// Needed for view overrides
 	ERasterizerFillMode MeshFillMode : ERasterizerFillMode_NumBits + 1;
@@ -1357,7 +1379,15 @@ public:
 		FVisibleMeshDrawCommand NewVisibleMeshDrawCommand;
 		//@todo MeshCommandPipeline - assign usable state ID for dynamic path draws
 		// Currently dynamic path draws will not get dynamic instancing, but they will be roughly sorted by state
+#if defined(GPUCULL_TODO)
+		const FMeshBatchElement& MeshBatchElement = MeshBatch.Elements[BatchElementIndex];
+		NewVisibleMeshDrawCommand.Setup(&MeshDrawCommand, DrawPrimitiveId, ScenePrimitiveId, -1, MeshFillMode, MeshCullMode, SortKey, 
+			MeshBatchElement.bIsInstanceRuns ? MeshBatchElement.InstanceRuns : nullptr,
+			MeshBatchElement.bIsInstanceRuns ? MeshBatchElement.NumInstances : 0
+			);
+#else //!defined(GPUCULL_TODO)
 		NewVisibleMeshDrawCommand.Setup(&MeshDrawCommand, DrawPrimitiveId, ScenePrimitiveId, -1, MeshFillMode, MeshCullMode, SortKey);
+#endif // defined(GPUCULL_TODO)
 		DrawList.Add(NewVisibleMeshDrawCommand);
 	}
 
@@ -1817,15 +1847,6 @@ RENDERER_API extern void DrawDynamicMeshPassPrivate(
 	bool& InNeedsShaderInitialisation,
 	uint32 InstanceFactor);
 
-RENDERER_API extern void DrawDynamicMeshPassPrivate(
-	const FSceneView& View,
-	const FScene& Scene,
-	FRHICommandListImmediate& RHICmdList,
-	FMeshCommandOneFrameArray& VisibleMeshDrawCommands,
-	FDynamicMeshDrawCommandStorage& DynamicMeshDrawCommandStorage,
-	FGraphicsMinimalPipelineStateSet& GraphicsMinimalPipelineStateSet,
-	bool& InNeedsShaderInitialization);
-
 RENDERER_API extern FMeshDrawCommandSortKey CalculateMeshStaticSortKey(const FMeshMaterialShader* VertexShader, const FMeshMaterialShader* PixelShader);
 
 inline FMeshDrawCommandSortKey CalculateMeshStaticSortKey(const TShaderRef<FMeshMaterialShader>& VertexShader, const TShaderRef<FMeshMaterialShader>& PixelShader)
diff --git a/Engine/Source/Runtime/Renderer/Public/MeshPassProcessor.inl b/Engine/Source/Runtime/Renderer/Public/MeshPassProcessor.inl
--- a/Engine/Source/Runtime/Renderer/Public/MeshPassProcessor.inl
+++ b/Engine/Source/Runtime/Renderer/Public/MeshPassProcessor.inl
@@ -177,25 +177,3 @@ void DrawDynamicMeshPass(const FSceneView& View, FRHICommandList& RHICmdList, co
 	const uint32 InstanceFactor = (!bForceStereoInstancingOff && View.IsInstancedStereoPass()) ? 2 : 1;
 	DrawDynamicMeshPassPrivate(View, RHICmdList, VisibleMeshDrawCommands, DynamicMeshDrawCommandStorage, GraphicsMinimalPipelineStateSet, NeedsShaderInitialisation, InstanceFactor);
 }
-
-#if defined(GPUCULL_TODO)
-
-/**
-* Overload that requires FViewInfo, need to be used when rendering draw commands that make use of Scene features (e.g., require instancing etc).
-*/
-template<typename LambdaType>
-void DrawDynamicMeshPass(const FSceneView& View, const FScene &Scene, FRHICommandListImmediate& RHICmdList, const LambdaType& BuildPassProcessorLambda, bool bForceStereoInstancingOff = false)
-{
-	FDynamicMeshDrawCommandStorage DynamicMeshDrawCommandStorage;
-	FMeshCommandOneFrameArray VisibleMeshDrawCommands;
-	FGraphicsMinimalPipelineStateSet GraphicsMinimalPipelineStateSet;
-	bool NeedsShaderInitialization;
-
-	FDynamicPassMeshDrawListContext DynamicMeshPassContext(DynamicMeshDrawCommandStorage, VisibleMeshDrawCommands, GraphicsMinimalPipelineStateSet, NeedsShaderInitialization);
-
-	BuildPassProcessorLambda(&DynamicMeshPassContext);
-
-	DrawDynamicMeshPassPrivate(View, Scene, RHICmdList, VisibleMeshDrawCommands, DynamicMeshDrawCommandStorage, GraphicsMinimalPipelineStateSet, NeedsShaderInitialization);
-}
-
-#endif // GPUCULL_TODO
