Fixed strata vizualise roughensses names.

diff --git a/Engine/Shaders/Private/Strata/StrataVisualize.usf b/Engine/Shaders/Private/Strata/StrataVisualize.usf
--- a/Engine/Shaders/Private/Strata/StrataVisualize.usf
+++ b/Engine/Shaders/Private/Strata/StrataVisualize.usf
@@ -241,12 +241,12 @@ void VisualizeMaterialPS(
 				HORIZONTALTAB; PRINTFLOAT1(SLAB_METALLIC(BSDF));
 				NEXTLINE;
 
-				HORIZONTALTAB; PRINTPARAMCHAR(_R_); PRINTPARAMCHAR(_O_); PRINTPARAMCHAR(_U_); PRINTPARAMCHAR(_G_); PRINTPARAMCHAR(_H_); PRINTPARAMCHAR(_0_);
+				HORIZONTALTAB; PRINTPARAMCHAR(_R_); PRINTPARAMCHAR(_O_); PRINTPARAMCHAR(_U_); PRINTPARAMCHAR(_G_); PRINTPARAMCHAR(_H_); PRINTPARAMCHAR(_X_);
 				HORIZONTALTAB; PRINTFLOAT1(SLAB_ROUGHNESSX(BSDF));
 				NEXTLINE;
 				if (BSDF_GETANISOTROPY(BSDF))
 				{
-					HORIZONTALTAB; PRINTPARAMCHAR(_R_); PRINTPARAMCHAR(_O_); PRINTPARAMCHAR(_U_); PRINTPARAMCHAR(_G_); PRINTPARAMCHAR(_H_); PRINTPARAMCHAR(_1_);
+					HORIZONTALTAB; PRINTPARAMCHAR(_R_); PRINTPARAMCHAR(_O_); PRINTPARAMCHAR(_U_); PRINTPARAMCHAR(_G_); PRINTPARAMCHAR(_H_); PRINTPARAMCHAR(_Y_);
 					HORIZONTALTAB; PRINTFLOAT1(SLAB_ROUGHNESSY(BSDF));
 					NEXTLINE;
 				}
