Revert computation of refracted Fresnel term for clear coat shading model of non-metallic surfaces.

While nested dielectrics have an Eta close to 1.0, and have a nearly unobstructed refracted Fresnel contribution, this is not true for a bottom interface with a separate normal, which can generate widly varying Fresnel values. This inconsistency amounts to a look regression for projects that have selected the "Clear Coat Enable Second Normal" project setting.

While reverting the behavior only for projects where this setting is enabled, carving out such a path for the bottom normal presents a modeling inconsistency (compared to original shading model without a varying bottom normal) in the case where the bottom normal and top normal are identical. Since we do not track IOR through interfaces, we cannot easily decouple modification of bottom-coat Fresnel due to Eta or normal variation. It is expected that this concept will be expressible in the future, but remains unsupported for now.

diff --git a/Engine/Shaders/Private/ShadingModels.ush b/Engine/Shaders/Private/ShadingModels.ush
--- a/Engine/Shaders/Private/ShadingModels.ush
+++ b/Engine/Shaders/Private/ShadingModels.ush
@@ -465,9 +465,7 @@ FDirectLighting ClearCoatBxDF( FGBufferData GBuffer, half3 N, half3 V, half3 L,
 			// NoL is chosen to provide better parity with DefaultLit when ClearCoat=0
 			Vis2 = Vis_SmithJointApprox(a2, BottomContext.NoV, NoL);
 		}
-		// When refracting into a non-metallic substrate, the IOR ratio (Eta) between top and bottom interfaces approaches 1.0 and drives Fresnel reflectance to 0
 		float3 F = F_Schlick(GBuffer.SpecularColor, BottomContext.VoH);
-		float3 RefractedF = lerp(0.0, F, GBuffer.Metallic);
 		
 		float Energy = 0;
 
@@ -491,7 +489,7 @@ FDirectLighting ClearCoatBxDF( FGBufferData GBuffer, half3 N, half3 V, half3 L,
 		// Note: reusing D, V, and F from refracted context to save computation for when ClearCoat < 1
 		float3 CommonSpecular = (Energy * Falloff * D2 * Vis2) * AreaLight.FalloffColor;
 		float3 DefaultSpecular = F * NoL;
-		float3 RefractedSpecular = FresnelCoeff * Transmission * RefractedF * BottomContext.NoL;
+		float3 RefractedSpecular = FresnelCoeff * Transmission * F * BottomContext.NoL;
 		Lighting.Specular += CommonSpecular * lerp(DefaultSpecular, RefractedSpecular, ClearCoat);
 	}
 
