Improve the IBL effect on mobile to conherent with pc as much as possible.

diff --git a/Engine/Shaders/Private/MobileBasePassPixelShader.usf b/Engine/Shaders/Private/MobileBasePassPixelShader.usf
--- a/Engine/Shaders/Private/MobileBasePassPixelShader.usf
+++ b/Engine/Shaders/Private/MobileBasePassPixelShader.usf
@@ -230,7 +230,7 @@ void GatherSpecularIBL(FMaterialPixelParameters MaterialParameters
 #endif
 )
 {
-	half3 ProjectedCaptureVector = MaterialParameters.ReflectionVector;
+	half3 ProjectedCaptureVector = GetOffSpecularPeakReflectionDir(MaterialParameters.WorldNormal, MaterialParameters.ReflectionVector, Roughness);
 
 #if HQ_REFLECTIONS
 	half DistanceAlpha = 0.0;
