Made Strata visualize shader use multiple pass to avoid scratch memory assert.

diff --git a/Engine/Shaders/Private/Strata/StrataVisualize.usf b/Engine/Shaders/Private/Strata/StrataVisualize.usf
--- a/Engine/Shaders/Private/Strata/StrataVisualize.usf
+++ b/Engine/Shaders/Private/Strata/StrataVisualize.usf
@@ -9,10 +9,6 @@
 #include "/Engine/Private/MiniFontCommon.ush"
 #include "/Engine/Private/ShaderDrawDebug.ush"
 
-// On some platforms, UnpackStrataIn causes the following error@
-// (error, code:7203) - not enough registers available for the entire program. In order to permit spilling to memory please either disable this error via '#pragma warning (disable:7203)' or reduce it to a warning via '#pragma warning (default:7203)'
-// To fix this and increase occupancy, we will have to progressively unpack the BSDFs on the fly (and thus only have a single BSDF allocating registers)
-#pragma warning (disable:7203)
 
 void DrawReferential(float3 P, float3 X, float3 Y, float3 N, float3 InColor)
 {
@@ -197,13 +193,21 @@ void VisualizeMaterialPS(
 #define HORIZONTALSEPARATION  TopLeft.y += 4;
 #define HORIZONTALTAB  TopLeft.x += 8;
 
+#if PERMUTATION_BSDF_PASS==0
 		PRINTINFOCHAR(_T_); PRINTINFOCHAR(_O_); PRINTINFOCHAR(_P_);
 		NEXTLINE;
 		HORIZONTALSEPARATION;
 		PRINTBSDFCOUNT();
 		NEXTLINE;
 		HORIZONTALSEPARATION;
+#else
+		NEXTLINE;
+		HORIZONTALSEPARATION;
+		NEXTLINE;
+		HORIZONTALSEPARATION;
+#endif
 
+		LOOP
 		for (uint BSDFIndex = 0; BSDFIndex < Header.BSDFCount; ++BSDFIndex)
 		{
 			const FStrataBSDF BSDF = UnpackStrataBSDFIn(Strata.MaterialLobesBuffer, PixelStrataDataByteOffset);
@@ -212,8 +216,8 @@ void VisualizeMaterialPS(
 
 			switch (BSDF_GETTYPE(BSDF))
 			{
-
 			case STRATA_BSDF_TYPE_SLAB:
+#if PERMUTATION_BSDF_PASS==0
 				PRINTBSDFCHAR(_S_); PRINTBSDFCHAR(_L_); PRINTBSDFCHAR(_A_); PRINTBSDFCHAR(_B_); PRINTBSDFCHAR(_SPC_);
 				PRINTBSDFCHAR(_SPC_); PRINTBSDFCHAR(_SPC_); PRINTBSDFCHAR(_SPC_); PRINTBSDFCHAR(_SPC_); PRINTBSDFCHAR(_SPC_);
 				if (BSDF_GETGREYWEIGHT(BSDF))
@@ -251,8 +255,6 @@ void VisualizeMaterialPS(
 					NEXTLINE;
 				}
 
-				// STRATA_TODO
-
 				if (BSDF_GETHASSCATTERING(BSDF))
 				{
 					if (BSDF_GETHASSSSPROFILE(BSDF))
@@ -274,10 +276,35 @@ void VisualizeMaterialPS(
 						NEXTLINE;
 					}
 				}
+#else
+				NEXTLINE;
+				NEXTLINE;
+				NEXTLINE;
+				NEXTLINE;
+				NEXTLINE;
+				if (BSDF_GETANISOTROPY(BSDF))
+				{
+					NEXTLINE;
+				}
 
+				if (BSDF_GETHASSCATTERING(BSDF))
+				{
+					if (BSDF_GETHASSSSPROFILE(BSDF))
+					{
+						NEXTLINE;
+						NEXTLINE;
+					}
+					else
+					{
+						NEXTLINE;
+						NEXTLINE;
+					}
+				}
+#endif
 				break;
 
 			case STRATA_BSDF_TYPE_SHEEN:
+#if PERMUTATION_BSDF_PASS==1
 				PRINTBSDFCHAR(_S_); PRINTBSDFCHAR(_H_); PRINTBSDFCHAR(_E_); PRINTBSDFCHAR(_E_); PRINTBSDFCHAR(_N_); 
 				PRINTBSDFCHAR(_SPC_); PRINTBSDFCHAR(_SPC_); PRINTBSDFCHAR(_SPC_); PRINTBSDFCHAR(_SPC_); PRINTBSDFCHAR(_SPC_);
 				if (BSDF_GETGREYWEIGHT(BSDF))
@@ -299,10 +326,15 @@ void VisualizeMaterialPS(
 				HORIZONTALTAB; PRINTPARAMCHAR(_R_); PRINTPARAMCHAR(_O_); PRINTPARAMCHAR(_U_); PRINTPARAMCHAR(_G_); PRINTPARAMCHAR(_H_); PRINTPARAMCHAR(_SPC_);
 				HORIZONTALTAB; PRINTFLOAT1(SHEEN_ROUGHNESS(BSDF));
 				NEXTLINE;
-
+#else
+				NEXTLINE;
+				NEXTLINE;
+				NEXTLINE;
+#endif
 				break;
 
 			case STRATA_BSDF_TYPE_HAIR:
+#if PERMUTATION_BSDF_PASS==2
 				PRINTBSDFCHAR(_H_); PRINTBSDFCHAR(_A_); PRINTBSDFCHAR(_I_); PRINTBSDFCHAR(_R_); PRINTBSDFCHAR(_SPC_);
 				PRINTBSDFCHAR(_SPC_); PRINTBSDFCHAR(_SPC_); PRINTBSDFCHAR(_SPC_); PRINTBSDFCHAR(_SPC_); PRINTBSDFCHAR(_SPC_);
 				if (BSDF_GETGREYWEIGHT(BSDF))
@@ -333,9 +365,18 @@ void VisualizeMaterialPS(
 				HORIZONTALTAB; PRINTPARAMCHAR(_B_); PRINTPARAMCHAR(_A_); PRINTPARAMCHAR(_C_); PRINTPARAMCHAR(_K_); PRINTPARAMCHAR(_L_); PRINTPARAMCHAR(_SPC_);
 				HORIZONTALTAB; PRINTFLOAT1(HAIR_BACKLIT(BSDF));
 				NEXTLINE;
+#else
+				NEXTLINE;
+				NEXTLINE;
+				NEXTLINE;
+				NEXTLINE;
+				NEXTLINE;
+				NEXTLINE;
+#endif
 				break;
 
 			case STRATA_BSDF_TYPE_SINGLELAYERWATER:
+#if PERMUTATION_BSDF_PASS==3
 				PRINTBSDFCHAR(_W_); PRINTBSDFCHAR(_A_); PRINTBSDFCHAR(_T_); PRINTBSDFCHAR(_E_); PRINTBSDFCHAR(_R_);
 				PRINTBSDFCHAR(_SPC_); PRINTBSDFCHAR(_SPC_); PRINTBSDFCHAR(_SPC_); PRINTBSDFCHAR(_SPC_); PRINTBSDFCHAR(_SPC_);
 				PRINTBSDFCHAR(_T_); PRINTBSDFCHAR(_O_); PRINTBSDFCHAR(_D_); PRINTBSDFCHAR(_O_); PRINTBSDFCHAR(_SPC_);
@@ -355,6 +396,13 @@ void VisualizeMaterialPS(
 				HORIZONTALTAB; PRINTPARAMCHAR(_O_); PRINTPARAMCHAR(_P_); PRINTPARAMCHAR(_A_); PRINTPARAMCHAR(_C_); PRINTPARAMCHAR(_I_); PRINTPARAMCHAR(_T_);
 				HORIZONTALTAB; PRINTFLOAT1(SLW_TOPMATERIALOPACITY(BSDF));
 				NEXTLINE;
+#else
+				NEXTLINE;
+				NEXTLINE;
+				NEXTLINE;
+				NEXTLINE;
+				NEXTLINE;
+#endif
 				break;
 			}
 
@@ -363,6 +411,7 @@ void VisualizeMaterialPS(
 
 		const uint PixelStrataDataByteOffsetPostBSDF = PixelStrataDataByteOffset;
 
+#if PERMUTATION_BSDF_PASS==0
 		PRINTINFOCHAR(_B_); PRINTINFOCHAR(_O_); PRINTINFOCHAR(_T_); PRINTINFOCHAR(_T_); PRINTINFOCHAR(_O_); PRINTINFOCHAR(_M_);
 		NEXTLINE;
 		HORIZONTALSEPARATION;
@@ -388,6 +437,7 @@ void VisualizeMaterialPS(
 		{
 			OutColor.w = 0.0f;				// Set background visiblity to 0 in the bottom right part of the screen to better reads the debug values
 		}
+#endif
 	}
 }
 
diff --git a/Engine/Source/Runtime/Renderer/Private/Strata/Strata.cpp b/Engine/Source/Runtime/Renderer/Private/Strata/Strata.cpp
--- a/Engine/Source/Runtime/Renderer/Private/Strata/Strata.cpp
+++ b/Engine/Source/Runtime/Renderer/Private/Strata/Strata.cpp
@@ -148,7 +148,8 @@ class FVisualizeMaterialPS : public FGlobalShader
 	DECLARE_GLOBAL_SHADER(FVisualizeMaterialPS);
 	SHADER_USE_PARAMETER_STRUCT(FVisualizeMaterialPS, FGlobalShader);
 
-	using FPermutationDomain = TShaderPermutationDomain<>;
+	class FBSDFPass : SHADER_PERMUTATION_INT("PERMUTATION_BSDF_PASS", 4);
+	using FPermutationDomain = TShaderPermutationDomain<FBSDFPass>;
 
 	BEGIN_SHADER_PARAMETER_STRUCT(FParameters, )
 		SHADER_PARAMETER_STRUCT_REF(FViewUniformShaderParameters, ViewUniformBuffer)
@@ -193,9 +194,6 @@ void AddVisualizeMaterialPasses(FRDGBuilder& GraphBuilder, const TArray<FViewInf
 	{
 		const FViewInfo& View = Views[i];
 
-		FVisualizeMaterialPS::FPermutationDomain PermutationVector;
-		TShaderMapRef<FVisualizeMaterialPS> PixelShader(View.ShaderMap, PermutationVector);
-
 		if (View.Family->EngineShowFlags.VisualizeStrataMaterial)
 		{
 			FVisualizeMaterialPS::FParameters* PassParameters = GraphBuilder.AllocParameters<FVisualizeMaterialPS::FParameters>();
@@ -210,8 +208,15 @@ void AddVisualizeMaterialPasses(FRDGBuilder& GraphBuilder, const TArray<FViewInf
 				ShaderDrawDebug::SetParameters(GraphBuilder, View.ShaderDrawData, PassParameters->ShaderDrawParameters);
 			}
 
-			FPixelShaderUtils::AddFullscreenPass<FVisualizeMaterialPS>(GraphBuilder, View.ShaderMap, RDG_EVENT_NAME("StrataVisualizeMaterial"),
-				PixelShader, PassParameters, View.ViewRect, PreMultipliedColorTransmittanceBlend);
+			for (uint32 j = 0; j < 4; ++j)
+			{
+				FVisualizeMaterialPS::FPermutationDomain PermutationVector;
+				PermutationVector.Set<typename FVisualizeMaterialPS::FBSDFPass>(j);
+				TShaderMapRef<FVisualizeMaterialPS> PixelShader(View.ShaderMap, PermutationVector);
+
+				FPixelShaderUtils::AddFullscreenPass<FVisualizeMaterialPS>(GraphBuilder, View.ShaderMap, RDG_EVENT_NAME("StrataVisualizeMaterial"),
+					PixelShader, PassParameters, View.ViewRect, PreMultipliedColorTransmittanceBlend);
+			}
 		}
 	}
 }
