Workaround for Dynamic Mesh Pass Draws

diff --git a/Engine/Shaders/Private/SceneData.ush b/Engine/Shaders/Private/SceneData.ush
--- a/Engine/Shaders/Private/SceneData.ush
+++ b/Engine/Shaders/Private/SceneData.ush
@@ -304,11 +304,26 @@ FSceneDataIntermediates GetSceneDataIntermediates(uint InstanceIdOffset, uint Dr
 	FSceneDataIntermediates Intermediates = (FSceneDataIntermediates)0;
 #if VF_USE_PRIMITIVE_SCENE_DATA
 	Intermediates.InstanceIdLoadIndex = InstanceIdOffset + DrawInstanceId;
-	// TODO: move InstanceIdsBuffer out of view and into some global Scene structure?
-	Intermediates.InstanceId = View.InstanceIdsBuffer[InstanceIdOffset + DrawInstanceId] & ((1U << 28U) - 1);
-	// We store the view index (which can be used for instanced stereo or other multi-view in the top four bits of the instance ID)
-	// Note: this is an index of views for this render pass, not the view ID in the culling manager.
-	Intermediates.ViewIndex = View.InstanceIdsBuffer[InstanceIdOffset + DrawInstanceId] >> 28U;
+	// GPUCULL_TODO: workaround for the fact that DrawDynamicMeshPassPrivate et al. don't work with GPU-Scene instancing
+	//				 instead they mark the top bit in the primitive ID and disable auto instancing such that there is an 1:1:1
+	//               drawcmd:primitive:instance. Then we can just look up the primitive and fetch the instance data index.
+	// We mark the PrimitiveID with the top bit in dynamic draw passes
+	if (InstanceIdOffset >= (1U << 31U))
+	{
+		// mask off the top bit
+		uint PrimitiveID = InstanceIdOffset & ((1U << 31U) - 1U);
+		Intermediates.InstanceId = GetPrimitiveData(PrimitiveID).InstanceDataOffset;
+		Intermediates.ViewIndex = 0;
+	}
+	else
+	{
+		// TODO: move InstanceIdsBuffer out of view and into some global Scene structure?
+		Intermediates.InstanceId = View.InstanceIdsBuffer[InstanceIdOffset + DrawInstanceId] & ((1U << 28U) - 1);
+		// We store the view index (which can be used for instanced stereo or other multi-view in the top four bits of the instance ID)
+		// Note: this is an index of views for this render pass, not the view ID in the culling manager.
+		Intermediates.ViewIndex = View.InstanceIdsBuffer[InstanceIdOffset + DrawInstanceId] >> 28U;
+	}
+
 	Intermediates.InstanceData = GetInstanceData(Intermediates.InstanceId, View.InstanceDataSOAStride);
 	Intermediates.PrimitiveId = Intermediates.InstanceData.PrimitiveId;
 	Intermediates.Primitive = GetPrimitiveData(Intermediates.PrimitiveId);
diff --git a/Engine/Source/Runtime/Renderer/Private/MeshDrawCommands.cpp b/Engine/Source/Runtime/Renderer/Private/MeshDrawCommands.cpp
--- a/Engine/Source/Runtime/Renderer/Private/MeshDrawCommands.cpp
+++ b/Engine/Source/Runtime/Renderer/Private/MeshDrawCommands.cpp
@@ -499,7 +499,12 @@ void BuildMeshDrawCommandPrimitiveIdBuffer(
 			for (uint32 InstanceFactorIndex = 0; InstanceFactorIndex < InstanceFactor; InstanceFactorIndex++, PrimitiveIdIndex++)
 			{
 				checkSlow(PrimitiveIdIndex < MaxPrimitiveId);
+			#if defined(GPUCULL_TODO)
+				// Append flag to mark this as a non-instance data index (which is then treated as a primitive ID in the SceneData.ush loading
+				PrimitiveIds[PrimitiveIdIndex] = VisibleMeshDrawCommand.DrawPrimitiveId | (1U << 31U);
+			#else
 				PrimitiveIds[PrimitiveIdIndex] = VisibleMeshDrawCommand.DrawPrimitiveId;
+			#endif
 			}
 		}
 	}
@@ -1068,7 +1073,14 @@ void SortAndMergeDynamicPassMeshDrawCommands(
 
 		if (bUseGPUScene)
 		{
+		#if defined(GPUCULL_TODO)
+			// GPUCULL_TODO: workaround for the fact that DrawDynamicMeshPassPrivate et al. don't work with GPU-Scene instancing
+			//               we don't support dynamic instancing for this path since we require one primitive per draw command
+			//               This is because the stride on the instance data buffer is set to 0 so only the first will ever be fetched.
+			const bool bDynamicInstancing = false;
+		#else
 			const bool bDynamicInstancing = IsDynamicInstancingEnabled(FeatureLevel);
+		#endif
 			if (bDynamicInstancing)
 			{
 				NewPassVisibleMeshDrawCommands.Empty(NumDrawCommands);
diff --git a/Engine/Source/Runtime/Renderer/Private/MeshPassProcessor.cpp b/Engine/Source/Runtime/Renderer/Private/MeshPassProcessor.cpp
--- a/Engine/Source/Runtime/Renderer/Private/MeshPassProcessor.cpp
+++ b/Engine/Source/Runtime/Renderer/Private/MeshPassProcessor.cpp
@@ -1257,13 +1257,13 @@ void ApplyViewOverridesToMeshDrawCommands(const FSceneView& View, FMeshCommandOn
 				VisibleMeshDrawCommand.StateBucketId,
 				VisibleMeshDrawCommand.MeshFillMode,
 				VisibleMeshDrawCommand.MeshCullMode,
-#if defined(GPUCULL_TODO)
+			#if defined(GPUCULL_TODO)
 				VisibleMeshDrawCommand.SortKey,
 				VisibleMeshDrawCommand.RunArray,
 				VisibleMeshDrawCommand.NumRuns);
-#else //!defined(GPUCULL_TODO)
+			#else
 				VisibleMeshDrawCommand.SortKey);
-#endif // defined(GPUCULL_TODO)
+			#endif
 
 			ViewOverriddenMeshCommands.Add(NewVisibleMeshDrawCommand);
 		}
@@ -1284,7 +1284,14 @@ void DrawDynamicMeshPassPrivate(
 {
 	if (VisibleMeshDrawCommands.Num() > 0)
 	{
+	#if defined(GPUCULL_TODO)
+		// GPUCULL_TODO: workaround for the fact that DrawDynamicMeshPassPrivate et al. don't work with GPU-Scene instancing
+		//               we don't support dynamic instancing for this path since we require one primitive per draw command
+		//               This is because the stride on the instance data buffer is set to 0 so only the first will ever be fetched.
+		const bool bDynamicInstancing = false;
+	#else
 		const bool bDynamicInstancing = IsDynamicInstancingEnabled(View.GetFeatureLevel());
+	#endif 
 
 		FRHIBuffer* PrimitiveIdVertexBuffer = nullptr;
 
