Fix shader errors when VIRTUAL_SM_ENABLED is set but GPUCULL_TODO is not

diff --git a/Engine/Shaders/Private/VertexFactoryDefaultInterface.ush b/Engine/Shaders/Private/VertexFactoryDefaultInterface.ush
--- a/Engine/Shaders/Private/VertexFactoryDefaultInterface.ush
+++ b/Engine/Shaders/Private/VertexFactoryDefaultInterface.ush
@@ -5,11 +5,9 @@
 #ifndef VF_IMPLEMENTED_GET_VIEW_INDEX
 #define VF_IMPLEMENTED_GET_VIEW_INDEX
 
-#if defined(GPUCULL_TODO)
-
 uint VertexFactoryGetViewIndex(FVertexFactoryIntermediates Intermediates)
 {
-#if VF_USE_PRIMITIVE_SCENE_DATA
+#if VF_USE_PRIMITIVE_SCENE_DATA && defined(GPUCULL_TODO)
 	return Intermediates.SceneData.ViewIndex;
 #else
 	return 0U;
@@ -18,7 +16,7 @@ uint VertexFactoryGetViewIndex(FVertexFactoryIntermediates Intermediates)
 
 uint VertexFactoryGetInstanceIdLoadIndex(FVertexFactoryIntermediates Intermediates)
 {
-#if VF_USE_PRIMITIVE_SCENE_DATA
+#if VF_USE_PRIMITIVE_SCENE_DATA && defined(GPUCULL_TODO)
 	return Intermediates.SceneData.InstanceIdLoadIndex;
 #else
 	return 0U;
@@ -27,7 +25,7 @@ uint VertexFactoryGetInstanceIdLoadIndex(FVertexFactoryIntermediates Intermediat
 
 float4x4 VertexFactoryGetLocalToWorld(FVertexFactoryIntermediates Intermediates)
 {
-#if VF_USE_PRIMITIVE_SCENE_DATA
+#if VF_USE_PRIMITIVE_SCENE_DATA && defined(GPUCULL_TODO)
 	return Intermediates.SceneData.InstanceData.LocalToWorld;
 #else
 	return Primitive.LocalToWorld;
@@ -36,7 +34,7 @@ float4x4 VertexFactoryGetLocalToWorld(FVertexFactoryIntermediates Intermediates)
 
 float4x4 VertexFactoryGetWorldToLocal(FVertexFactoryIntermediates Intermediates)
 {
-#if VF_USE_PRIMITIVE_SCENE_DATA
+#if VF_USE_PRIMITIVE_SCENE_DATA && defined(GPUCULL_TODO)
 	return Intermediates.SceneData.InstanceData.WorldToLocal;
 #else
 	return Primitive.WorldToLocal;
@@ -47,7 +45,7 @@ float4x4 VertexFactoryGetWorldToLocal(FVertexFactoryIntermediates Intermediates)
 
 uint VertexFactoryGetViewIndex(FPositionOnlyVertexFactoryInput Input)
 {
-#if VF_USE_PRIMITIVE_SCENE_DATA
+#if VF_USE_PRIMITIVE_SCENE_DATA && defined(GPUCULL_TODO)
 	FSceneDataIntermediates SceneData = GetSceneDataIntermediates(Input.InstanceIdOffset, Input.DrawInstanceId);
 	return SceneData.ViewIndex;
 #else
@@ -57,7 +55,7 @@ uint VertexFactoryGetViewIndex(FPositionOnlyVertexFactoryInput Input)
 
 uint VertexFactoryGetViewIndex(FPositionAndNormalOnlyVertexFactoryInput Input)
 {
-#if VF_USE_PRIMITIVE_SCENE_DATA
+#if VF_USE_PRIMITIVE_SCENE_DATA && defined(GPUCULL_TODO)
 	FSceneDataIntermediates SceneData = GetSceneDataIntermediates(Input.InstanceIdOffset, Input.DrawInstanceId);
 	return SceneData.ViewIndex;
 #else
@@ -67,7 +65,7 @@ uint VertexFactoryGetViewIndex(FPositionAndNormalOnlyVertexFactoryInput Input)
 
 uint VertexFactoryGetInstanceIdLoadIndex(FPositionOnlyVertexFactoryInput Input)
 {
-#if VF_USE_PRIMITIVE_SCENE_DATA
+#if VF_USE_PRIMITIVE_SCENE_DATA && defined(GPUCULL_TODO)
 	FSceneDataIntermediates SceneData = GetSceneDataIntermediates(Input.InstanceIdOffset, Input.DrawInstanceId);
 	return SceneData.InstanceIdLoadIndex;
 #else
@@ -77,7 +75,7 @@ uint VertexFactoryGetInstanceIdLoadIndex(FPositionOnlyVertexFactoryInput Input)
 
 uint VertexFactoryGetInstanceIdLoadIndex(FPositionAndNormalOnlyVertexFactoryInput Input)
 {
-#if VF_USE_PRIMITIVE_SCENE_DATA
+#if VF_USE_PRIMITIVE_SCENE_DATA && defined(GPUCULL_TODO)
 	FSceneDataIntermediates SceneData = GetSceneDataIntermediates(Input.InstanceIdOffset, Input.DrawInstanceId);
 	return SceneData.InstanceIdLoadIndex;
 #else
@@ -87,6 +85,4 @@ uint VertexFactoryGetInstanceIdLoadIndex(FPositionAndNormalOnlyVertexFactoryInpu
 
 #endif // POSITION_ONLY
 
-#endif // GPUCULL_TODO
-
 #endif // VF_IMPLEMENTED_GET_VIEW_INDEX
