Add native transparency handling to the path tracer

Instead of treating transparency with indirect rays, handle it with a special case in the integrator. This greatly reduces noise on transparent shadows and allows emissive surfaces behind transparency to be handled exactly.

The current strategy is to trace a ray through all transparent hits but only perform lighting calculations on one (proportionally to throughput*albedo) via weighted reservoir sampling. Shadow rays simply accumulate the transparency among all hits by running the material inside the anyhit shader.

diff --git a/Engine/Shaders/Private/PathTracing/Material/PathTracingMaterialSampling.ush b/Engine/Shaders/Private/PathTracing/Material/PathTracingMaterialSampling.ush
--- a/Engine/Shaders/Private/PathTracing/Material/PathTracingMaterialSampling.ush
+++ b/Engine/Shaders/Private/PathTracing/Material/PathTracingMaterialSampling.ush
@@ -48,6 +48,89 @@ bool IsNeeValidMaterial(FMaterialClosestHitPayload Payload)
 #endif
 }
 
+float3 GetMaterialTransparency(FMaterialClosestHitPayload Payload)
+{
+	switch (Payload.BlendingMode)
+	{
+		case RAY_TRACING_BLEND_MODE_TRANSLUCENT:
+			if (Payload.Ior == 1.0)
+			{
+				// material is not configured with refraction, can use transparency
+				return 1 - Payload.Opacity;
+			}
+			else
+			{
+				// not using transparency - material will be opaque
+				return 0;
+			}
+		case RAY_TRACING_BLEND_MODE_ADDITIVE:
+			return 1;
+		case RAY_TRACING_BLEND_MODE_MODULATE:
+			return Payload.Radiance;
+		case RAY_TRACING_BLEND_MODE_ALPHA_COMPOSITE:
+			return 1 - Payload.Opacity;
+		case RAY_TRACING_BLEND_MODE_ALPHA_HOLDOUT:
+			return 1;
+		default: 
+			return 0;
+	}
+}
+
+float3 EstimateMaterialAlbedo(FMaterialClosestHitPayload Payload)
+{
+	float BrdfWeight = 1;
+	bool HasRefraction = false;
+	switch (Payload.BlendingMode)
+	{
+		case RAY_TRACING_BLEND_MODE_TRANSLUCENT:
+			if (Payload.Ior == 1.0)
+			{
+				// material is not configured with refraction, can use transparency
+				// NOTE: the albedo when considering refraction is handled below
+				BrdfWeight = Payload.Opacity;
+			}
+			else
+			{
+				HasRefraction = true;
+			}
+			break;
+		case RAY_TRACING_BLEND_MODE_ADDITIVE:
+			BrdfWeight = Payload.Opacity;
+			break;
+		case RAY_TRACING_BLEND_MODE_MODULATE:
+		case RAY_TRACING_BLEND_MODE_ALPHA_HOLDOUT:
+			// In these modes, the light does not interact with the material at all, only passes through
+			return 0;
+		default:
+			break;
+	}
+
+	float3 Albedo = 0;
+	switch (Payload.ShadingModelID)
+	{
+		case SHADINGMODELID_UNLIT:
+			// this object does not scatter light
+			return 0;
+		case SHADINGMODELID_TWOSIDED_FOLIAGE:
+			Albedo = BrdfWeight * (Payload.DiffuseColor + Payload.SpecularColor + GetSubsurfaceColor(Payload));
+			break;
+		case SHADINGMODELID_DEFAULT_LIT:
+		case SHADINGMODELID_CLEAR_COAT:
+			// FIXME: could make this more accurate by estimating the directional albedo
+			Albedo = BrdfWeight * (Payload.DiffuseColor + Payload.SpecularColor);
+		default:
+			Albedo = BrdfWeight * Payload.BaseColor;
+			break;
+	}
+	if (HasRefraction)
+	{
+		// glass is currently perfectly clear - so has unit albedo
+		float3 White = 1;
+		Albedo = lerp(White, Albedo, Payload.Opacity);
+	}
+	return Albedo;
+}
+
 void SampleMaterial(
 	float3 RayDirection,
 	FMaterialClosestHitPayload Payload,
@@ -72,21 +155,13 @@ void SampleMaterial(
 	}
 #else
 
-	if (Payload.BlendingMode == RAY_TRACING_BLEND_MODE_MODULATE)
-	{
-		OutDirection = RayDirection;
-		OutWeight = Payload.Radiance;
-		OutPdf = POSITIVE_INFINITY;
-		OutPositionBiasSign = -1;
-		return;
-	}
-
-	if (Payload.BlendingMode == RAY_TRACING_BLEND_MODE_ALPHA_HOLDOUT)
+	if (Payload.BlendingMode == RAY_TRACING_BLEND_MODE_MODULATE ||
+		Payload.BlendingMode == RAY_TRACING_BLEND_MODE_ALPHA_HOLDOUT)
 	{
-		OutDirection = RayDirection;
-		OutWeight = 1;
-		OutPdf = POSITIVE_INFINITY;
-		OutPositionBiasSign = -1;
+		// transparency only, nothing to sample
+		OutDirection = 0;
+		OutWeight = 0;
+		OutPdf = 0;
 		return;
 	}
 
@@ -98,19 +173,21 @@ void SampleMaterial(
 	switch (Payload.BlendingMode)
 	{
 		case RAY_TRACING_BLEND_MODE_TRANSLUCENT:
-			BrdfWeight = 1;
-			TransWeight = 1;
-			BrdfProb = Payload.Opacity;
+			if (Payload.Ior == 1.0)
+			{
+				// scale down BRDF contribution, but don't pick any indirect ray because its already accounted for by transparency handling
+				BrdfWeight = Payload.Opacity;
+			}
+			else
+			{
+				// randomly choose between brdf and glass
+				BrdfWeight = 1;
+				TransWeight = 1;
+				BrdfProb = Payload.Opacity;
+			}
 			break;
 		case RAY_TRACING_BLEND_MODE_ADDITIVE:
-			BrdfWeight = 1 + Payload.Opacity;
-			TransWeight = 1 + Payload.Opacity;
-			BrdfProb = Payload.Opacity / (1 + Payload.Opacity);
-			break;
-		case RAY_TRACING_BLEND_MODE_ALPHA_COMPOSITE:
-			BrdfWeight = 2 - Payload.Opacity;;
-			TransWeight = 2 - Payload.Opacity;
-			BrdfProb = 1 / (2 - Payload.Opacity);
+			BrdfWeight = Payload.Opacity;
 			break;
 		default:
 			// opaque or one of the modes handled above
@@ -126,29 +203,18 @@ void SampleMaterial(
 	{
 		// we randomly decided to sample the translucent portion
 		RandSample.x = (RandSample.x - BrdfProb) / (1 - BrdfProb);
-		// now decide what to do based on blend mode
-		if (Payload.BlendingMode == RAY_TRACING_BLEND_MODE_TRANSLUCENT && Payload.Ior != 1.0)
+
+		// only case that should reach this point is the refraction model
+		if (Payload.Roughness == 0)
 		{
-			// refraction model
-			if (Payload.Roughness == 0)
-			{
-				SpecularRefraction_SampleMaterial(RayDirection, Payload, RandSample, OutDirection, OutWeight, OutPdf, OutPositionBiasSign);
-			}
-			else
-			{
-				RoughGlass_SampleMaterial(RayDirection, Payload, RandSample, OutDirection, OutWeight, OutPdf, OutPositionBiasSign);
-			}
-			OutWeight *= TransWeight;
-			OutPdf *= 1 - BrdfProb;
+			SpecularRefraction_SampleMaterial(RayDirection, Payload, RandSample, OutDirection, OutWeight, OutPdf, OutPositionBiasSign);
 		}
 		else
 		{
-			// plain transparency
-			OutDirection = RayDirection;
-			OutWeight = TransWeight;
-			OutPdf = POSITIVE_INFINITY; // * 1 - BrdfProb
-			OutPositionBiasSign = -1;
+			RoughGlass_SampleMaterial(RayDirection, Payload, RandSample, OutDirection, OutWeight, OutPdf, OutPositionBiasSign);
 		}
+		OutWeight *= TransWeight;
+		OutPdf *= 1 - BrdfProb;
 		return;
 	}
 
@@ -219,23 +285,22 @@ void EvalMaterial(
 		case RAY_TRACING_BLEND_MODE_TRANSLUCENT:
 			// brdf:   Opacity
 			// glass:  1.0 - Opacity
-			BrdfWeight = 1;
-			TransWeight = 1;
-			BrdfProb = Payload.Opacity;
+			if (Payload.Ior == 1.0)
+			{
+				// handled by native transpaarency, just weight down the brdf contrib
+				BrdfWeight = Payload.Opacity;
+			}
+			else
+			{
+				// brdf or glass
+				BrdfWeight = 1;
+				TransWeight = 1;
+				BrdfProb = Payload.Opacity;
+			}
 			break;
 		case RAY_TRACING_BLEND_MODE_ADDITIVE:
 			// brdf:   Opacity
-			// transp: 1.0
-			BrdfWeight = 1 + Payload.Opacity;
-			TransWeight = 1 + Payload.Opacity;
-			BrdfProb = Payload.Opacity / (1 + Payload.Opacity);
-			break;
-		case RAY_TRACING_BLEND_MODE_ALPHA_COMPOSITE:
-			// brdf:   1.0
-			// transp: 1.0 - Opacity
-			BrdfWeight = 2 - Payload.Opacity;;
-			TransWeight = 2 - Payload.Opacity;
-			BrdfProb = 1 / (2 - Payload.Opacity);
+			BrdfWeight = Payload.Opacity;
 			break;
 		default:
 			// opaque or one of the modes handled above
@@ -270,17 +335,14 @@ void EvalMaterial(
 
 	if (BrdfProb < 1.0)
 	{
-		if (Payload.BlendingMode == RAY_TRACING_BLEND_MODE_TRANSLUCENT && Payload.Ior != 1.0)
+		// refraction model (only evaluate for rough case as singular can't be evaluated)
+		if (Payload.Roughness != 0)
 		{
-			// refraction model (only evaluate for rough case as singular can't be evaluated)
-			if (Payload.Roughness != 0)
-			{
-				float3 GlassWeight = 0;
-				float GlassPdf = 0;
-				RoughGlass_EvalMaterial(IncomingDirection, OutgoingDirection, Payload, GlassWeight, GlassPdf);
+			float3 GlassWeight = 0;
+			float GlassPdf = 0;
+			RoughGlass_EvalMaterial(IncomingDirection, OutgoingDirection, Payload, GlassWeight, GlassPdf);
 
-				AddLobeWithMIS(OutWeight, OutPdf, GlassWeight, GlassPdf, 1.0 - BrdfProb);
-			}
+			AddLobeWithMIS(OutWeight, OutPdf, GlassWeight, GlassPdf, 1.0 - BrdfProb);
 		}
 	}
 #endif
diff --git a/Engine/Shaders/Private/PathTracing/PathTracing.usf b/Engine/Shaders/Private/PathTracing/PathTracing.usf
--- a/Engine/Shaders/Private/PathTracing/PathTracing.usf
+++ b/Engine/Shaders/Private/PathTracing/PathTracing.usf
@@ -40,6 +40,164 @@ void AccumulateRadiance(inout float3 TotalRadiance, float3 PathRadiance)
 	}
 }
 
+FMaterialClosestHitPayload TraceTransparentRay(RayDesc Ray, FRayCone RayCone, bool LastBounce, uint2 LaunchIndex, uint NumLights, inout RandomSequence RandSequence, inout float3 PathThroughput, inout float3 Radiance)
+{
+	float SelectionWeightSum = 0;
+	float3 PayloadThroughput = PathThroughput;
+	FMaterialClosestHitPayload Payload;
+	for (;;)
+	{
+		const uint InitialRayFlags = 0;
+		const bool bInitialDisableSkyLightContribution = false;
+		const bool bIgnoreTranslucentMaterials = false;
+		FMaterialClosestHitPayload HitPayload = TraceMaterialRay(
+			TLAS,
+			InitialRayFlags,
+			RAY_TRACING_MASK_ALL,
+			Ray,
+			RayCone,
+			LaunchIndex,
+			bInitialDisableSkyLightContribution,
+			bIgnoreTranslucentMaterials);
+
+		// Loop over lights to capture their contribution
+		// #dxr_todo: if we have lots of lights, having some hierarchical structure would be better ....
+		for (uint LightId = 0; LightId < NumLights; ++LightId)
+		{
+			float HitT = HitPayload.IsMiss() ? Ray.TMax : HitPayload.HitT;
+			float3 LightRadiance = TraceLight(Ray, LightId, HitT).xyz;
+			AccumulateRadiance(Radiance, PathThroughput * LightRadiance);
+		}
+
+		if (HitPayload.IsMiss())
+		{
+			// Ray didn't hit any real geometry, so nothing left to do
+			break;
+		}
+
+		// material tweaks to workaround temporary limitations
+		if (HitPayload.ShadingModelID == SHADINGMODELID_SUBSURFACE)
+		{
+			// for now just approximate SSS by moving it into the diffuse color
+			HitPayload.DiffuseColor += HitPayload.CustomData.xyz;
+			HitPayload.DiffuseColor = saturate(HitPayload.DiffuseColor);
+			HitPayload.Opacity = 1; // overwrite radius
+			HitPayload.ShadingModelID = SHADINGMODELID_DEFAULT_LIT;
+		}
+
+		// add in surface emission (except for modulate blend mode which uses Radiance as the transparency multiplier)
+		if (HitPayload.BlendingMode != RAY_TRACING_BLEND_MODE_MODULATE)
+		{
+			AccumulateRadiance(Radiance, PathThroughput * HitPayload.Radiance);
+		}
+
+		float3 Transparency = GetMaterialTransparency(HitPayload);
+
+		if (!LastBounce)
+		{
+			float3 Contrib = PathThroughput * EstimateMaterialAlbedo(HitPayload);
+
+			float SelectionWeight = max3(Contrib.x, Contrib.y, Contrib.z);
+			SelectionWeightSum += SelectionWeight;
+			bool AcceptHit = false;
+			float PickWeight = 1;
+
+			// weighted reservoir sampling
+			if (SelectionWeight > 0)
+			{
+				if (SelectionWeight < SelectionWeightSum)
+				{
+					// the acceptanace probability is not 1.0
+					// generate a random number to see if we should accept this hit
+					float RandValue = RandomSequence_GenerateSample1D(RandSequence);
+					AcceptHit = RandValue * SelectionWeightSum < SelectionWeight;
+					PickWeight = SelectionWeightSum / SelectionWeight;
+				}
+				else
+				{
+					// accept automatically on the first hit
+					AcceptHit = true;
+				}
+			}
+
+			if (AcceptHit)
+			{
+				// stash this hit for next time
+				Payload = HitPayload;
+				PayloadThroughput = PathThroughput * PickWeight;
+			}
+		}
+
+		// prepare next step around the loop
+		// retrace the exact same ray with TMin one ulp past the hit we just found
+		PathThroughput *= Transparency;
+		Ray.TMin = asfloat(asuint(HitPayload.HitT) + 1);
+
+		if (all(PathThroughput == 0))
+		{
+			break;
+		}
+	}
+
+	if (SelectionWeightSum > 0)
+	{
+		// if we stored a valid hit in the payload, reset the path throughput to this point
+		PathThroughput = PayloadThroughput;
+	}
+	else
+	{
+		PathThroughput = 0;
+		Payload.SetMiss();
+	}
+	return Payload;
+}
+
+float3 TraceTransparentVisibilityRay(RayDesc Ray, uint2 PixelCoord)
+{
+	const uint RayFlags = RAY_FLAG_ACCEPT_FIRST_HIT_AND_END_SEARCH | RAY_FLAG_SKIP_CLOSEST_HIT_SHADER;
+	const uint InstanceInclusionMask = RAY_TRACING_MASK_SHADOW;
+	const uint RayContributionToHitGroupIndex = RAY_TRACING_SHADER_SLOT_SHADOW;
+	const uint MultiplierForGeometryContributionToShaderIndex = RAY_TRACING_NUM_SHADER_SLOTS;
+	const uint MissShaderIndex = 0;
+
+	FPackedMaterialClosestHitPayload PackedPayload = (FPackedMaterialClosestHitPayload)0;
+
+	// By enabling minimal payload mode all other payload information is ignored, meaning these functions need no payload inputs
+	PackedPayload.SetMinimalPayloadMode();
+	PackedPayload.SetEvalOpacity();
+	PackedPayload.HitT = 0;
+	PackedPayload.RadianceAndNormal[0] = asuint(1.0);
+	PackedPayload.RadianceAndNormal[1] = asuint(1.0);
+	PackedPayload.RadianceAndNormal[2] = asuint(1.0);
+
+	PackedPayload.SetPixelCoord(PixelCoord);
+
+	// Trace the ray
+	TraceRay(
+		TLAS,
+		RayFlags,
+		InstanceInclusionMask,
+		RayContributionToHitGroupIndex,
+		MultiplierForGeometryContributionToShaderIndex,
+		MissShaderIndex,
+		Ray,
+		PackedPayload);
+
+	if (PackedPayload.IsMiss())
+	{
+		// we didn't hit any geo - just return whatever throughput is left
+		return float3(
+			asfloat(PackedPayload.RadianceAndNormal[0]),
+			asfloat(PackedPayload.RadianceAndNormal[1]),
+			asfloat(PackedPayload.RadianceAndNormal[2])
+		);
+	}
+
+	// we registered a hit -- assumed to be fully opaque
+	return 0;
+}
+
+
 RAY_TRACING_ENTRY_RAYGEN(PathTracingMainRG)
 {
 	uint2 LaunchIndex = DispatchRaysIndex().xy + View.ViewRectMin.xy + TileOffset.xy;
@@ -79,11 +237,6 @@ RAY_TRACING_ENTRY_RAYGEN(PathTracingMainRG)
 	FRayCone RayCone = (FRayCone)0;
 	RayCone.SpreadAngle = View.EyeToPixelSpreadAngle;
 
-	// Cast ray
-	const uint InitialRayFlags = 0;
-	const bool bInitialDisableSkyLightContribution = false;
-	const bool bIgnoreTranslucentMaterials = false;
-
 	// path state variables (these cary information between bounces)
 	float3 PathThroughput = 1.0;
 
@@ -92,57 +245,29 @@ RAY_TRACING_ENTRY_RAYGEN(PathTracingMainRG)
 	// 2: both Material and Light
 	const int MISMode = PathTracingData.MISMode;
 
-	for (int Bounce = 0; Bounce <= PathTracingData.MaxBounces; Bounce++)
+	// number of directly visible lights for the first bounce
+	uint NumLights = 0;
+	if (SceneLightCount > 0)
 	{
-		FMaterialClosestHitPayload Payload = TraceMaterialRay(
-			TLAS,
-			InitialRayFlags,
-			RAY_TRACING_MASK_ALL,
-			Ray,
-			RayCone,
-			LaunchIndex,
-			bInitialDisableSkyLightContribution,
-			bIgnoreTranslucentMaterials);
-
-		// directly visible lights
-		if (Bounce == 0 && SceneLightCount > 0)
+		if (PathTracingData.VisibleLights)
 		{
-			uint NumLights = 0;
-			if (PathTracingData.VisibleLights)
-			{
-				// user wants to see all lights
-				NumLights = SceneLightCount;
-			}
-			else if (IsEnvironmentLight(0))
-			{
-				// user doesn't want to see all lights - still check the first one in case its a skydome (there can be only one, and it always comes first if it is there)
-				NumLights = 1;
-			}
-			// Loop over lights to capture their contribution
-			// #dxr_todo: if we have lots of lights, having some hierarchical structure would be better ....
-			for (uint LightId = 0; LightId < NumLights; ++LightId)
-			{
-				float HitT = Payload.IsMiss() ? Ray.TMax : Payload.HitT;
-				float3 LightRadiance = TraceLight(Ray, LightId, HitT).xyz;
-				AccumulateRadiance(Radiance, PathThroughput * LightRadiance);
-			}
+			// user wants to see all lights
+			NumLights = SceneLightCount;
 		}
-
-		if (Payload.IsMiss())
+		else if (IsEnvironmentLight(0))
 		{
-			// Ray didn't hit any real geometry, so nothing left to do
-			break;
+			// user doesn't want to see all lights - still check the first one in case its a skydome (there can be only one, and it always comes first if it is there)
+			NumLights = 1;
 		}
+	}
 
-		// add in surface emission (except for modulate blend mode which uses Radiance as the transparency multiplier)
-		if (Payload.BlendingMode != RAY_TRACING_BLEND_MODE_MODULATE)
-		{
-			AccumulateRadiance(Radiance, PathThroughput * Payload.Radiance);
-		}
+	for (int Bounce = 0; Bounce <= PathTracingData.MaxBounces; Bounce++)
+	{
+		FMaterialClosestHitPayload Payload = TraceTransparentRay(Ray, RayCone, Bounce == PathTracingData.MaxBounces, LaunchIndex, NumLights, RandSequence, PathThroughput, Radiance);
 
-		if (Bounce == PathTracingData.MaxBounces)
+		if (Payload.IsMiss())
 		{
-			// we've arrived at the last bounce - nothing left to do
+			// we didn't hit anything selectable for further shading, we are done
 			break;
 		}
 
@@ -155,16 +280,6 @@ RAY_TRACING_ENTRY_RAYGEN(PathTracingMainRG)
 			Payload.WorldNormal = -Payload.WorldNormal;
 		}
 
-		// material tweaks to workaround path tracer limitations
-		if (Payload.ShadingModelID == SHADINGMODELID_SUBSURFACE)
-		{
-			// for now just approximate SSS by moving it into the diffuse color
-			Payload.DiffuseColor += Payload.CustomData.xyz;
-			Payload.DiffuseColor = saturate(Payload.DiffuseColor);
-			Payload.Opacity = 1; // overwrite radius
-			Payload.ShadingModelID = SHADINGMODELID_DEFAULT_LIT;
-		}
-
 		// Choose a random number for both Light sampling and BxDF sampling
 		float4 RandSample = RandomSequence_GenerateSample4D(RandSequence);
 
@@ -212,21 +327,13 @@ RAY_TRACING_ENTRY_RAYGEN(PathTracingMainRG)
 					}
 					ApplyPositionBias(LightRay, SignedPositionBias * Payload.WorldNormal, PathTracingData.MaxNormalBias);
 
-					bool LightIsVisible = true;
 					if (CastsShadow(LightId))
 					{
-						const uint RayFlags = RAY_FLAG_ACCEPT_FIRST_HIT_AND_END_SEARCH | RAY_FLAG_SKIP_CLOSEST_HIT_SHADER;
-						const uint InstanceInclusionMask = RAY_TRACING_MASK_SHADOW;
-
-						FPackedMaterialClosestHitPayload PackedPayload = (FPackedMaterialClosestHitPayload)0;
-
-						TraceVisibilityRayPacked(PackedPayload, TLAS, RayFlags, InstanceInclusionMask, LaunchIndex, LightRay);
-
-						LightIsVisible = PackedPayload.IsMiss();
+						RadianceOverPdf *= TraceTransparentVisibilityRay(LightRay, LaunchIndex);
 					}
 
 					// #dxr_todo: Is it cheaper to fire the ray first? Or eval the material first?
-					if (LightIsVisible)
+					if (any(RadianceOverPdf > 0))
 					{
 						// Evaluate material
 						float3 MaterialWeight;
@@ -335,25 +442,22 @@ RAY_TRACING_ENTRY_RAYGEN(PathTracingMainRG)
 
 				if (any(LightContrib > 0))
 				{
-					bool LightIsVisible = true;
 					if (CastsShadow(LightId))
 					{
-						const uint RayFlags = RAY_FLAG_ACCEPT_FIRST_HIT_AND_END_SEARCH | RAY_FLAG_SKIP_CLOSEST_HIT_SHADER;
-						const uint InstanceInclusionMask = RAY_TRACING_MASK_SHADOW;
-						FPackedMaterialClosestHitPayload PackedPayload = (FPackedMaterialClosestHitPayload)0;
 						RayDesc LightRay = Ray;
 						LightRay.TMax = HitT;
-						TraceVisibilityRayPacked(PackedPayload, TLAS, RayFlags, InstanceInclusionMask, LaunchIndex, LightRay);
-						LightIsVisible = PackedPayload.IsMiss();
-					}
-					if (LightIsVisible)
-					{
-						// the light made some contribution, and there was nothing along the shadow ray
-						AccumulateRadiance(Radiance, LightContrib);
+						LightContrib *= TraceTransparentVisibilityRay(LightRay, LaunchIndex);
 					}
+					// the light made some contribution, and there was nothing along the shadow ray
+					AccumulateRadiance(Radiance, LightContrib);
 				}
 			}
+			
 		}
+
+		// from this point on, we don't need to include lights in the trace call
+		// because NEE handled it for us
+		NumLights = 0;
 	}
 
 
diff --git a/Engine/Shaders/Private/RayTracing/RayTracingCommon.ush b/Engine/Shaders/Private/RayTracing/RayTracingCommon.ush
--- a/Engine/Shaders/Private/RayTracing/RayTracingCommon.ush
+++ b/Engine/Shaders/Private/RayTracing/RayTracingCommon.ush
@@ -225,6 +225,7 @@ uint2 GetPixelCoord(uint2 DispatchThreadId, uint UpscaleFactor)
 #define RAY_TRACING_PAYLOAD_FLAG_ENABLE_SKY_LIGHT_CONTRIBUTION (1 << 2) // Indicates that hit shaders should add in the Sky Light contribution in their output. Set by RayGen shader, before TraceRay().
 #define RAY_TRACING_PAYLOAD_FLAG_LUMEN_PAYLOAD				   (1 << 3) // Indicates that closest hit shader should evaluate geometry normal and instanceid.
 #define RAY_TRACING_PAYLOAD_FLAG_IGNORE_TRANSLUCENT            (1 << 4) // Indicates that translucent materials should be ignored by IgnoreHit() in anyhit shaders.
+#define RAY_TRACING_PAYLOAD_FLAG_EVAL_OPACITY                  (1 << 5) // Indicates that translucent materials should be evaluated and the overall transparency accumulated into the payload
 
 struct FMaterialClosestHitPayload : FMinimalPayload
 {
@@ -243,7 +244,7 @@ struct FMaterialClosestHitPayload : FMinimalPayload
 	float Roughness;					// 84        2       16bits
 	float Ior;							// 88        2       16bits
 	uint ShadingModelID;				// 92        1       4bits
-	uint BlendingMode;					// 96        0       4bits (packed with ShadingModelID)
+	uint BlendingMode;					// 96        0       3bits (packed with ShadingModelID)
 	uint PrimitiveLightingChannelMask;	// 100       0       3bits (packed with ShadingModelID)
 	float4 CustomData;					// 104       4       32bits
 	float GBufferAO;					// 120       0       (removed)
@@ -253,7 +254,7 @@ struct FMaterialClosestHitPayload : FMinimalPayload
 	// So keep it here and force to re-construct it in Unpack call using ray information.
 	// It is not packed in FRayHitInfoPacked
 	float3 WorldPos;					// 136       0       (derived)
-	uint Flags;							// 148       0       5bits (packed with ShadingModelID)
+	uint Flags;							// 148       0       6bits (packed with ShadingModelID)
 	float3 WorldTangent;				// 152       6       48bits
 	float Anisotropy;					// 164       2       16bits (packed with WorldTangent)
 										// 166 total
@@ -352,17 +353,20 @@ struct FPackedMaterialClosestHitPayload : FMinimalPayload
 	void SetPixelCoord(uint2 PixelCoord) { PackedPixelCoord = (PixelCoord.x & 0xFFFF) | (PixelCoord.y << 16); }
 	uint2 GetPixelCoord() { return uint2(PackedPixelCoord & 0xFFFF, PackedPixelCoord >> 16); }
 
-	void SetMinimalPayloadMode() { IorAndShadingModelIDAndBlendingModeAndFlagsAndPrimitiveLightingChannelMask |= RAY_TRACING_PAYLOAD_FLAG_MINIMAL_PAYLOAD << 24; }
-	bool IsMinimalPayloadMode() { return ((IorAndShadingModelIDAndBlendingModeAndFlagsAndPrimitiveLightingChannelMask >> 24) & RAY_TRACING_PAYLOAD_FLAG_MINIMAL_PAYLOAD) != 0; }
+	void SetMinimalPayloadMode() { IorAndShadingModelIDAndBlendingModeAndFlagsAndPrimitiveLightingChannelMask |= RAY_TRACING_PAYLOAD_FLAG_MINIMAL_PAYLOAD << 23; }
+	bool IsMinimalPayloadMode() { return ((IorAndShadingModelIDAndBlendingModeAndFlagsAndPrimitiveLightingChannelMask >> 23) & RAY_TRACING_PAYLOAD_FLAG_MINIMAL_PAYLOAD) != 0; }
 
-	void SetEnableSkyLightContribution() { IorAndShadingModelIDAndBlendingModeAndFlagsAndPrimitiveLightingChannelMask |= RAY_TRACING_PAYLOAD_FLAG_ENABLE_SKY_LIGHT_CONTRIBUTION << 24; }
-	bool IsEnableSkyLightContribution() { return ((IorAndShadingModelIDAndBlendingModeAndFlagsAndPrimitiveLightingChannelMask >> 24) & RAY_TRACING_PAYLOAD_FLAG_ENABLE_SKY_LIGHT_CONTRIBUTION) != 0; }
+	void SetEnableSkyLightContribution() { IorAndShadingModelIDAndBlendingModeAndFlagsAndPrimitiveLightingChannelMask |= RAY_TRACING_PAYLOAD_FLAG_ENABLE_SKY_LIGHT_CONTRIBUTION << 23; }
+	bool IsEnableSkyLightContribution() { return ((IorAndShadingModelIDAndBlendingModeAndFlagsAndPrimitiveLightingChannelMask >> 23) & RAY_TRACING_PAYLOAD_FLAG_ENABLE_SKY_LIGHT_CONTRIBUTION) != 0; }
 
-	void SetLumenPayload() { IorAndShadingModelIDAndBlendingModeAndFlagsAndPrimitiveLightingChannelMask |= RAY_TRACING_PAYLOAD_FLAG_LUMEN_PAYLOAD << 24; }
-	bool IsLumenPayload() { return  ((IorAndShadingModelIDAndBlendingModeAndFlagsAndPrimitiveLightingChannelMask >> 24) & RAY_TRACING_PAYLOAD_FLAG_LUMEN_PAYLOAD) != 0;}
+	void SetLumenPayload() { IorAndShadingModelIDAndBlendingModeAndFlagsAndPrimitiveLightingChannelMask |= RAY_TRACING_PAYLOAD_FLAG_LUMEN_PAYLOAD << 23; }
+	bool IsLumenPayload() { return  ((IorAndShadingModelIDAndBlendingModeAndFlagsAndPrimitiveLightingChannelMask >> 23) & RAY_TRACING_PAYLOAD_FLAG_LUMEN_PAYLOAD) != 0;}
 
-	void SetIgnoreTranslucentMaterials() { IorAndShadingModelIDAndBlendingModeAndFlagsAndPrimitiveLightingChannelMask |= RAY_TRACING_PAYLOAD_FLAG_IGNORE_TRANSLUCENT << 24; }
-	bool IsIgnoreTranslucentMaterials() { return ((IorAndShadingModelIDAndBlendingModeAndFlagsAndPrimitiveLightingChannelMask >> 24) & RAY_TRACING_PAYLOAD_FLAG_IGNORE_TRANSLUCENT) != 0; }
+	void SetIgnoreTranslucentMaterials() { IorAndShadingModelIDAndBlendingModeAndFlagsAndPrimitiveLightingChannelMask |= RAY_TRACING_PAYLOAD_FLAG_IGNORE_TRANSLUCENT << 23; }
+	bool IsIgnoreTranslucentMaterials() { return ((IorAndShadingModelIDAndBlendingModeAndFlagsAndPrimitiveLightingChannelMask >> 23) & RAY_TRACING_PAYLOAD_FLAG_IGNORE_TRANSLUCENT) != 0; }
+
+	void SetEvalOpacity() { IorAndShadingModelIDAndBlendingModeAndFlagsAndPrimitiveLightingChannelMask |= RAY_TRACING_PAYLOAD_FLAG_EVAL_OPACITY << 23; }
+	bool IsEvalOpacity() { return ((IorAndShadingModelIDAndBlendingModeAndFlagsAndPrimitiveLightingChannelMask >> 23) & RAY_TRACING_PAYLOAD_FLAG_EVAL_OPACITY) != 0; }
 
 	FRayCone GetRayCone() { return UnpackRayCone(PackedRayCone); }
 	void SetRayCone(FRayCone NewRayCone) { PackedRayCone = PackRayCone(NewRayCone); }
@@ -446,8 +450,8 @@ struct FPackedMaterialClosestHitPayload : FMinimalPayload
 	float GetRoughness()     { return f16tof32(MetallicAndSpecularAndRoughness >> 16); }
 	float GetIor()           { return f16tof32(IorAndShadingModelIDAndBlendingModeAndFlagsAndPrimitiveLightingChannelMask); }
 	uint GetShadingModelID() { return (IorAndShadingModelIDAndBlendingModeAndFlagsAndPrimitiveLightingChannelMask >> 16) & 0xF; }
-	uint GetBlendingMode()   { return (IorAndShadingModelIDAndBlendingModeAndFlagsAndPrimitiveLightingChannelMask >> 20) & 0xF; }
-	uint GetFlags()          { return (IorAndShadingModelIDAndBlendingModeAndFlagsAndPrimitiveLightingChannelMask >> 24) & 0x1F; }
+	uint GetBlendingMode()   { return (IorAndShadingModelIDAndBlendingModeAndFlagsAndPrimitiveLightingChannelMask >> 20) & 0x7; }
+	uint GetFlags()          { return (IorAndShadingModelIDAndBlendingModeAndFlagsAndPrimitiveLightingChannelMask >> 23) & 0x3F; }
 	uint GetPrimitiveLightingChannelMask() { return (IorAndShadingModelIDAndBlendingModeAndFlagsAndPrimitiveLightingChannelMask >> 29) & 0x7; }
 	float GetAnisotropy()    { return f16tof32(WorldTangentAndAnisotropy[1] >> 16); }
 
@@ -491,8 +495,8 @@ FPackedMaterialClosestHitPayload PackRayTracingPayload(FMaterialClosestHitPayloa
 	Output.MetallicAndSpecularAndRoughness |= f32tof16(Input.Roughness) << 16;
 	Output.IorAndShadingModelIDAndBlendingModeAndFlagsAndPrimitiveLightingChannelMask  = f32tof16(Input.Ior);                              // 16 bits
 	Output.IorAndShadingModelIDAndBlendingModeAndFlagsAndPrimitiveLightingChannelMask |= (Input.ShadingModelID & 0xF) << 16;               // 4 bits
-	Output.IorAndShadingModelIDAndBlendingModeAndFlagsAndPrimitiveLightingChannelMask |= (Input.BlendingMode & 0xF) << 20;                 // 4 bits
-	Output.IorAndShadingModelIDAndBlendingModeAndFlagsAndPrimitiveLightingChannelMask |= (Input.Flags & 0x1F) << 24;                       // 5 bits
+	Output.IorAndShadingModelIDAndBlendingModeAndFlagsAndPrimitiveLightingChannelMask |= (Input.BlendingMode & 0x7) << 20;                 // 3 bits
+	Output.IorAndShadingModelIDAndBlendingModeAndFlagsAndPrimitiveLightingChannelMask |= (Input.Flags & 0x3F) << 23;                       // 6 bits
 	Output.IorAndShadingModelIDAndBlendingModeAndFlagsAndPrimitiveLightingChannelMask |= (Input.PrimitiveLightingChannelMask & 0x7) << 29; // 3 bits
 	Output.PackedIndirectIrradiance[0]  = f32tof16(Input.IndirectIrradiance.x);
 	Output.PackedIndirectIrradiance[0] |= f32tof16(Input.IndirectIrradiance.y) << 16;
diff --git a/Engine/Shaders/Private/RayTracing/RayTracingMaterialHitShaders.usf b/Engine/Shaders/Private/RayTracing/RayTracingMaterialHitShaders.usf
--- a/Engine/Shaders/Private/RayTracing/RayTracingMaterialHitShaders.usf
+++ b/Engine/Shaders/Private/RayTracing/RayTracingMaterialHitShaders.usf
@@ -519,9 +519,45 @@ RAY_TRACING_ENTRY_ANY_HIT(MaterialAHS,
 	// Continue traversal by simply returning from this shader
 	// Usually this path is inactive because opaque materials don't have AHS bound
 	return;
-#endif // MATERIALBLENDING_SOLID
+#elif MATERIALBLENDING_ALPHAHOLDOUT || MATERIALBLENDING_ADDITIVE
+	if (PackedPayload.IsEvalOpacity())
+	{
+		// path tracer wants to treat these cases as fully transparent
+		// no material evaluation is required
+		IgnoreHit();
+	}
+	// otherwise accept hit
+#else
+	// All other blending modes _might_ need material evaluation - so we need
+	// to generate some code
 
 #if MATERIALBLENDING_MASKED
+	// NOTE: Masked mode execute always - regardless of ray flags
+	// This ensures that both material rays and shadow rays "see" the same thing
+#else
+	// For any other blending models, we have a few options:
+	if (PackedPayload.IsIgnoreTranslucentMaterials())
+	{
+		// special mode used by regular RT shadows - translucent materials are skipped, everything else is treated as opaque
+#if MATERIALBLENDING_TRANSLUCENT
+		IgnoreHit();
+#endif
+		return;
+	}
+	else if (PackedPayload.IsEvalOpacity())
+	{
+		// this case is used by the path tracer only currently
+		// We want to evaluate the material opacity and accumulate it
+		// Keep going and evaluate the material
+	}
+	else
+	{
+		// non-opaque blending mode but no special behavior was requested
+		return;
+	}
+#endif  // !MATERIALBLENDING_MASKED
+
+	// If we got here -- we have a reason to evaluate the material -- do the setup 
 
 	GlobalTextureMipBias = PackedPayload.GetMipBias();
 
@@ -553,22 +589,60 @@ RAY_TRACING_ENTRY_ANY_HIT(MaterialAHS,
 		// #endif
 	}
 
+	// Now extract the relevant info from PixelMaterialInputs according to the blending mode
+
+#if MATERIALBLENDING_MASKED
+	// Regardless of payload flags -- we always apply this 
 	if (GetMaterialMask(PixelMaterialInputs) < 0)
 	{
 		IgnoreHit();
 	}
-	else
+#else
+	// If we got this far -- we have a blending mode that can produce partial opacity _and_ we want to
+	// know how much it contributes to the throughput along this ray
+#if MATERIALBLENDING_MODULATE
+	// this material is the only one to lead to colored opacity
+	float3 Transparency = GetMaterialEmissive(PixelMaterialInputs);
+#elif MATERIALBLENDING_ALPHACOMPOSITE
+	float Opacity = GetMaterialOpacity(PixelMaterialInputs);
+	float Transparency = 1 - Opacity;
+#elif MATERIALBLENDING_TRANSLUCENT
+	float Ior = GetMaterialRefraction(PixelMaterialInputs).x;
+	if (Ior != 1.0)
 	{
-		// Continue traversal by simply returning from this shader
+		// refraction enabled, treat as opaque for now
+		// #dxr_todo: the path tracer will eventualy want an option for "fake caustics" 
 		return;
 	}
-#endif // MATERIALBLENDING_MASKED
+	// refraction disabled - treat as transparency
+	float Opacity = GetMaterialOpacity(PixelMaterialInputs);
+	float Transparency = 1 - Opacity;
+#else
+#error "Control flow does not account for all blending modes!"
+#endif
 
-	// For any other blending models, ignore hit in anyhit shader if we're tracing visibility rays until we have proper handling for translucent
-	if (PackedPayload.IsIgnoreTranslucentMaterials())
+	// Update the ray throughput (it is packed simply into the payload since we don't need to carry any other information across hits)
+	float3 RayThroughput = float3(
+		asfloat(PackedPayload.RadianceAndNormal[0]),
+		asfloat(PackedPayload.RadianceAndNormal[1]),
+		asfloat(PackedPayload.RadianceAndNormal[2])
+	);
+
+	RayThroughput *= Transparency;
+
+	PackedPayload.RadianceAndNormal[0] = asuint(RayThroughput.x);
+	PackedPayload.RadianceAndNormal[1] = asuint(RayThroughput.y);
+	PackedPayload.RadianceAndNormal[2] = asuint(RayThroughput.z);
+
+	if (any(RayThroughput > 0))
 	{
+		// not at full opacity yet -- ignore the hit so we can keep going
 		IgnoreHit();
 	}
+	// If RayThroughput reached 0 -- just return which will accept the hit
+#endif // !MATERIALBLENDING_MASKED
+
+#endif // !MATERIALBLENDING_SOLID && !MATERIALBLENDING_ALPHAHOLDOUT && !MATERIALBLENDING_ADDITIVE
 
 #endif // USE_MATERIAL_ANY_HIT_SHADER
 }
diff --git a/Engine/Source/Runtime/Renderer/Private/RayTracing/RayTracingMaterialHitShaders.cpp b/Engine/Source/Runtime/Renderer/Private/RayTracing/RayTracingMaterialHitShaders.cpp
--- a/Engine/Source/Runtime/Renderer/Private/RayTracing/RayTracingMaterialHitShaders.cpp
+++ b/Engine/Source/Runtime/Renderer/Private/RayTracing/RayTracingMaterialHitShaders.cpp
@@ -144,11 +144,7 @@ public:
 			return false;
 		}
 
-		const bool bWantAnyHitShader = GCompileRayTracingMaterialAHS 
-			&& (Parameters.MaterialParameters.bIsMasked || Parameters.MaterialParameters.BlendMode == BLEND_Translucent);
-
 		return IsSupportedVertexFactoryType(Parameters.VertexFactoryType)
-			&& (bWantAnyHitShader == UseAnyHitShader)
 			&& LightMapPolicyType::ShouldCompilePermutation(Parameters)
 			&& ShouldCompileRayTracingShadersForProject(Parameters.Platform)
 			&& (bool)GRayTracingUseTextureLod == UseRayConeTextureLod;
@@ -159,7 +155,7 @@ public:
 		// NOTE: Any CVars that are used in this function must be handled in ShaderMapAppendKeyString() to ensure shaders are recompiled when necessary.
 
 		OutEnvironment.SetDefine(TEXT("USE_MATERIAL_CLOSEST_HIT_SHADER"), GCompileRayTracingMaterialCHS ? 1 : 0);
-		OutEnvironment.SetDefine(TEXT("USE_MATERIAL_ANY_HIT_SHADER"), GCompileRayTracingMaterialAHS ? 1 : 0);
+		OutEnvironment.SetDefine(TEXT("USE_MATERIAL_ANY_HIT_SHADER"), UseAnyHitShader ? 1 : 0);
 		OutEnvironment.SetDefine(TEXT("USE_RAYTRACED_TEXTURE_RAYCONE_LOD"), UseRayConeTextureLod ? 1 : 0);
 		OutEnvironment.SetDefine(TEXT("SCENE_TEXTURES_DISABLED"), 1);
 		LightMapPolicyType::ModifyCompilationEnvironment(Parameters, OutEnvironment);
@@ -246,7 +242,7 @@ static TShaderRef<FMaterialCHS> GetMaterialHitShader(const FMaterial& RESTRICT M
 	const bool bMaterialsCompiled = GCompileRayTracingMaterialAHS || GCompileRayTracingMaterialCHS;
 	checkf(bMaterialsCompiled, TEXT(""));
 
-	if ((MaterialResource.IsMasked() || MaterialResource.GetBlendMode() == BLEND_Translucent) && GCompileRayTracingMaterialAHS)
+	if ((MaterialResource.IsMasked() || MaterialResource.GetBlendMode() != BLEND_Opaque) && GCompileRayTracingMaterialAHS)
 	{
 		if(UseTextureLod)
 		{ 
@@ -353,7 +349,6 @@ void FRayTracingMeshProcessor::AddMeshBatch(const FMeshBatch& RESTRICT MeshBatch
 
 		const EBlendMode BlendMode = Material.GetBlendMode();
 		const FMaterialShadingModelField ShadingModels = Material.GetShadingModels();
-		const bool bIsTranslucent = IsTranslucentBlendMode(BlendMode);
 
 		// Only draw opaque materials.
 		if ((!PrimitiveSceneProxy || PrimitiveSceneProxy->ShouldRenderInMainPass())
