Update a few non-nanite VSM paths with page table interface changes

diff --git a/Engine/Shaders/Private/VirtualShadowMaps/PageAccessCommon.ush b/Engine/Shaders/Private/VirtualShadowMaps/PageAccessCommon.ush
--- a/Engine/Shaders/Private/VirtualShadowMaps/PageAccessCommon.ush
+++ b/Engine/Shaders/Private/VirtualShadowMaps/PageAccessCommon.ush
@@ -44,6 +44,13 @@
 #define VSM_PACKED_PAGE_INFO_VIEW_ID_BITS (10U)
 #define VSM_PACKED_PAGE_INFO_SUPER_PAGE_LEVEL_BITS (2U)
 
+struct FPageInfo
+{
+	uint2 VirtualPage;
+	uint ViewId;
+	uint SuperPageLevel;
+};
+
 //Buffer containing VSM_ALLOCATED_FLAG for each page that is allocated and optionally | VSM_INVALID_FLAG if it needs to be rendered (not cached)
 StructuredBuffer<uint> PageFlags;
 StructuredBuffer<uint> HPageFlags;
diff --git a/Engine/Shaders/Private/VirtualShadowMaps/PageManagement.usf b/Engine/Shaders/Private/VirtualShadowMaps/PageManagement.usf
--- a/Engine/Shaders/Private/VirtualShadowMaps/PageManagement.usf
+++ b/Engine/Shaders/Private/VirtualShadowMaps/PageManagement.usf
@@ -773,7 +773,7 @@ void CalculatePageRects(uint3 DispatchThreadID : SV_DispatchThreadID)
 			InterlockedMax(PageRectBoundsOut[PageBoundIndex].z, vPage.x);
 			InterlockedMax(PageRectBoundsOut[PageBoundIndex].w, vPage.y);
 		}
-		OutPageTable[PageOffset] = uint2(VSM_INVALID_PHYSICAL_PAGE_ADDRESS, VSM_INVALID_PHYSICAL_PAGE_ADDRESS);
+		OutPageTable[PageOffset] = VSM_PHYSICAL_PAGE_INVALID;
 
 	}
 }
@@ -936,7 +936,7 @@ void AllocatePagesUsingRects(uint GroupThreadId : SV_GroupThreadID, uint2 Dispat
 						// Request rendering of pages that are not cached (AKA VSM_INVALID_PHYSICAL_PAGE_ADDRESS in cache).
 						// TODO: make sure caching machinery is use in this path also
 						OutPageFlags[PageOffset] = VSM_INVALID_FLAG | VSM_ALLOCATED_FLAG;
-						OutPageTable[PageOffset] = pPage;
+						OutPageTable[PageOffset] = ShadowEncodePageTable(pPage);
 					}
 				}
 
