When strata is enabled, only starta material are rendered correctly now.
To enable material conversion to strata locally, one must do the following and recompile:
- Set CONVERT_MATERIAL_TO_STRATA_ON_LOAD to 1
- Locally change ShaderVersion.ush GUID
Non strata code is gone from light end envlight passes. It is also gone from the basepass (SingleLayerWater is really special there as it is doing special forward shading not through strata and we should fix that later).
This reduces Starta basepas and lighting cost by 10& but this is far from enough to get previous cost...

Made opacity input used as coverage in strata now.
Added support for opacity-as-coverage on Strata translucent.

diff --git a/Engine/Shaders/Private/BasePassPixelShader.usf b/Engine/Shaders/Private/BasePassPixelShader.usf
--- a/Engine/Shaders/Private/BasePassPixelShader.usf
+++ b/Engine/Shaders/Private/BasePassPixelShader.usf
@@ -80,7 +80,7 @@
 #include "Lumen/LumenTranslucencyVolumeShared.ush"
 #endif
 
-#define STRATA_ENABLED					PROJECT_STRATA && MATERIAL_IS_STRATA
+#define STRATA_ENABLED					(PROJECT_STRATA && MATERIAL_IS_STRATA)
 #define STRATA_OPAQUE_ENABLED			(STRATA_ENABLED && MATERIALBLENDING_ANY_TRANSLUCENT<=0)
 #define STRATA_TRANSLUCENT_ENABLED		(STRATA_ENABLED && MATERIALBLENDING_ANY_TRANSLUCENT>0)
 
@@ -909,6 +909,8 @@ void FPixelShaderInOut_MainPS(
 #endif
 	}
 #endif
+
+#if !STRATA_ENABLED
 	// Store the results in local variables and reuse instead of calling the functions multiple times.
 	half3 BaseColor = GetMaterialBaseColor(PixelMaterialInputs);
 	half  Metallic = GetMaterialMetallic(PixelMaterialInputs);
@@ -918,8 +920,18 @@ void FPixelShaderInOut_MainPS(
 	float Roughness = GetMaterialRoughness(PixelMaterialInputs);
 	float Anisotropy = GetMaterialAnisotropy(PixelMaterialInputs);
 	uint ShadingModel = GetMaterialShadingModel(PixelMaterialInputs);
-
 	half Opacity = GetMaterialOpacity(PixelMaterialInputs);
+#else
+	half3 BaseColor = 0;
+	half  Metallic = 0;
+	half  Specular = 0;
+
+	float MaterialAO = 1;
+	float Roughness = 0;
+	float Anisotropy = 0;
+	uint ShadingModel = 0;
+	half Opacity = 0;
+#endif
 
 // Opacity for this model is the coverage of the top layer over the transmission surface. Since the
 // underlying material isn't allowed to be metallic, we are scaling the metallic value by the coverage.
@@ -971,7 +983,8 @@ void FPixelShaderInOut_MainPS(
 #endif
 
 #if STRATA_ENABLED && STRATA_INLINE_SINGLELAYERWATER
-	// Override GBuffer data with Strata SLW water BSDF
+	// Override GBuffer data with Strata SLW water BSDF to run forward shadfing code.
+	// STRATA_TODO: run the shading through a strata path? (by adding a special BSDF?)
 	SanitizeStrataSingleLayerWater(PixelMaterialInputs.FrontMaterial.Layers[0].BSDFs[0]);
 	FStrataBSDF SLWBSDF = PixelMaterialInputs.FrontMaterial.Layers[0].BSDFs[0];
 
@@ -1032,6 +1045,7 @@ void FPixelShaderInOut_MainPS(
 	GBuffer.Depth = MaterialParameters.ScreenPosition.w;
 	GBuffer.PrecomputedShadowFactors = GetPrecomputedShadowMasks(LightmapVTPageTableResult, Interpolants, MaterialParameters.PrimitiveId, MaterialParameters.AbsoluteWorldPosition, VolumetricLightmapBrickTextureUVs);
 
+#if !STRATA_ENABLED || STRATA_INLINE_SINGLELAYERWATER
 	const float GBufferDither = InterleavedGradientNoise(MaterialParameters.SvPosition.xy, View.StateFrameIndexMod8);
 	// Use GBuffer.ShadingModelID after SetGBufferForShadingModel(..) because the ShadingModel input might not be the same as the output
 	SetGBufferForShadingModel(
@@ -1048,6 +1062,7 @@ void FPixelShaderInOut_MainPS(
 		GBufferDither,
 		ShadingModel
 		);
+#endif // !STRATA_ENABLED
 
 #if USES_GBUFFER
 	// This requires cleanup. Shader code that uses GBuffer.SelectiveOutputMask expects the outputmask to be in
@@ -1085,7 +1100,10 @@ void FPixelShaderInOut_MainPS(
 	#endif
 	}
 #endif
-	
+
+	bool bChecker = CheckerFromPixelPos(MaterialParameters.SvPosition.xy);
+#if !STRATA_ENABLED || STRATA_INLINE_SINGLELAYERWATER
+
 	// So that the following code can still use DiffuseColor and SpecularColor.
 	GBuffer.SpecularColor = ComputeF0(Specular, BaseColor, Metallic);
 
@@ -1110,7 +1128,6 @@ void FPixelShaderInOut_MainPS(
 #if POST_PROCESS_SUBSURFACE
 	// SubsurfaceProfile applies the BaseColor in a later pass. Any lighting output in the base pass needs
 	// to separate specular and diffuse lighting in a checkerboard pattern
-	bool bChecker = CheckerFromPixelPos(MaterialParameters.SvPosition.xy);
 	if (UseSubsurfaceProfile(GBuffer.ShadingModelID))
 	{
 		AdjustBaseColorAndSpecularColorForSubsurfaceProfileLighting(BaseColor, GBuffer.SpecularColor, Specular, bChecker);
@@ -1188,6 +1205,8 @@ void FPixelShaderInOut_MainPS(
 	}
 	#endif
 
+#endif // !STRATA_ENABLED
+
 	half3 DiffuseColor = 0;
 	half3 Color = 0;
 	float IndirectIrradiance = 0;
@@ -1195,6 +1214,11 @@ void FPixelShaderInOut_MainPS(
 	half3 ColorSeparateSpecular = 0;
 	half3 ColorSeparateEmissive = 0;
 
+	float3 DiffuseIndirectLighting = 0;
+	float3 SubsurfaceIndirectLighting = 0;
+
+#if !STRATA_ENABLED || STRATA_INLINE_SINGLELAYERWATER
+
 	#if !MATERIAL_SHADINGMODEL_UNLIT
 
 		float3 DiffuseDir = BentNormal;
@@ -1232,8 +1256,6 @@ void FPixelShaderInOut_MainPS(
 		}
 		#endif
 
-		float3 DiffuseIndirectLighting;
-		float3 SubsurfaceIndirectLighting;
 		GetPrecomputedIndirectLightingAndSkyLight(MaterialParameters, Interpolants, BasePassInterpolants, LightmapVTPageTableResult, GBuffer, DiffuseDir, VolumetricLightmapBrickTextureUVs, DiffuseIndirectLighting, SubsurfaceIndirectLighting, IndirectIrradiance);
 		
 
@@ -1327,6 +1349,14 @@ void FPixelShaderInOut_MainPS(
 		#endif
 	#endif
 
+#else // !STRATA_ENABLED
+
+	// STRATA_TODO forward shadow and occlusion
+	float DirectionalLightShadow = 1.0f; 
+	float IndirectOcclusion = 1.0f;
+					
+#endif // !STRATA_ENABLED
+
 	#if NEEDS_BASEPASS_VERTEX_FOGGING
 		float4 HeightFogging = BasePassInterpolants.VertexFog;
 	#elif NEEDS_BASEPASS_PIXEL_FOGGING
@@ -1353,8 +1383,6 @@ void FPixelShaderInOut_MainPS(
 #if NEEDS_BASEPASS_PIXEL_FOGGING && PROJECT_SUPPORT_SKY_ATMOSPHERE && MATERIAL_IS_SKY==0 // Do not apply aerial perpsective on sky materials
 	if (ResolvedView.SkyAtmosphereApplyCameraAerialPerspectiveVolume > 0.0f)
 	{
-
-
 		// Sample the aerial perspective (AP).
 		Fogging = GetAerialPerspectiveLuminanceTransmittanceWithFogOver(
 			ResolvedView.RealTimeReflectionCapture, ResolvedView.SkyAtmosphereCameraAerialPerspectiveVolumeSizeAndInvSize,
@@ -1382,6 +1410,9 @@ void FPixelShaderInOut_MainPS(
 
 #endif
 
+	half3 Emissive = 0;
+#if !STRATA_ENABLED
+
 	// Volume lighting for lit translucency
 #if (MATERIAL_SHADINGMODEL_DEFAULT_LIT || MATERIAL_SHADINGMODEL_SUBSURFACE) && (MATERIALBLENDING_TRANSLUCENT || MATERIALBLENDING_ADDITIVE) && !SIMPLE_FORWARD_SHADING && !FORWARD_SHADING
 	if (GBuffer.ShadingModelID == SHADINGMODELID_DEFAULT_LIT || GBuffer.ShadingModelID == SHADINGMODELID_SUBSURFACE)
@@ -1397,7 +1428,9 @@ void FPixelShaderInOut_MainPS(
 		Color = lerp(Color, GBufferDiffuseColor, View.UnlitViewmodeMask);
 	#endif
 
-	half3 Emissive = GetMaterialEmissive(PixelMaterialInputs);
+	Emissive = GetMaterialEmissive(PixelMaterialInputs);
+
+#endif // !STRATA_ENABLED
 
 	// The following block is disabled on Vulkan because it triggers an nvidia driver bug (UE-101609).
 #if USE_DEVELOPMENT_SHADERS && !VULKAN_PROFILE_SM5
@@ -1416,6 +1449,14 @@ void FPixelShaderInOut_MainPS(
 	#endif
 #endif
 
+#if MATERIAL_WORKS_WITH_DUAL_SOURCE_COLOR_BLENDING
+	float3 DualBlendColorAdd = 0.0f;
+	float3 DualBlendColorMul = 1.0f;
+	float3 DualBlendColorCov = 1.0f;
+#endif
+
+#if !STRATA_ENABLED || STRATA_INLINE_SINGLELAYERWATER
+
 #if !POST_PROCESS_SUBSURFACE && !MATERIAL_SHADINGMODEL_THIN_TRANSLUCENT
  	// For skin we need to keep them separate. We also keep them separate for thin translucent.
 	// Otherwise just add them together.
@@ -1426,7 +1467,9 @@ void FPixelShaderInOut_MainPS(
 	Color += Emissive;
 #endif
 
-#if MATERIAL_SHADINGMODEL_SINGLELAYERWATER
+#endif // !STRATA_ENABLED
+
+#if MATERIAL_SHADINGMODEL_SINGLELAYERWATER || STRATA_INLINE_SINGLELAYERWATER
 	{
 		const bool CameraIsUnderWater = false;	// Fade out the material contribution over to water contribution according to material opacity.
 		const float3 SunIlluminance = ResolvedView.DirectionalLightColor.rgb * PI;	// times PI because it is divided by PI on CPU (=luminance) and we want illuminance here. 
@@ -1466,12 +1509,7 @@ void FPixelShaderInOut_MainPS(
 	}
 #endif // MATERIAL_SHADINGMODEL_SINGLELAYERWATER
 
-#if MATERIAL_WORKS_WITH_DUAL_SOURCE_COLOR_BLENDING
-	float3 DualBlendColorAdd = 0.0f;
-	float3 DualBlendColorMul = 1.0f;
-#endif
-
-#if MATERIAL_SHADINGMODEL_THIN_TRANSLUCENT
+#if MATERIAL_SHADINGMODEL_THIN_TRANSLUCENT && !STRATA_ENABLED
 	{
 		AccumulateThinTranslucentModel(
 								DualBlendColorAdd,
@@ -1488,6 +1526,7 @@ void FPixelShaderInOut_MainPS(
 	}
 #endif // MATERIAL_SHADINGMODEL_THIN_TRANSLUCENT
 
+
 	bool bSubsurfaceEnable = true;
 #if STRATA_ENABLED && !MATERIAL_SHADINGMODEL_UNLIT
 
@@ -1542,10 +1581,11 @@ void FPixelShaderInOut_MainPS(
 			Strata,
 			Throughput);
 
+		DualBlendColorCov = GetMaterialOpacity(PixelMaterialInputs); // Opacity represent coverage STRATA_TODO after validation, rename everything
 		DualBlendColorAdd = Color;
 		DualBlendColorMul = Throughput;
 		Color = 0.0f;
-		Opacity = 1.0f; // 1 is needed if using translucent. STRATA_TODO: we should only allow premultiplied alpha
+		Opacity = 1.0f; // nullify following operation
 	}
 #endif // STRATA_TRANSLUCENT_ENABLED
 
@@ -1556,10 +1596,11 @@ void FPixelShaderInOut_MainPS(
 #if STRATA_OPAQUE_ENABLED
 	Color = BSDF_GETEMISSIVE(UnlitBSDF);
 #else // STRATA_TRANSLUCENT_ENABLED
+	DualBlendColorCov = GetMaterialOpacity(PixelMaterialInputs); // Opacity represent coverage STRATA_TODO after validation, rename everything
 	DualBlendColorAdd = BSDF_GETEMISSIVE(UnlitBSDF);
 	DualBlendColorMul = UNLIT_TRANSMITTANCE(UnlitBSDF);
 	Color = 0.0f;
-	Opacity = 1.0f; // 1 is needed if using translucent. STRATA_TODO: we should only allow premultiplied alpha
+	Opacity = 1.0f; // nullify following operation
 #endif
 	
 #endif
@@ -1583,8 +1624,8 @@ void FPixelShaderInOut_MainPS(
 		//    AdjustedDualBlendMul = Fogging.a * DualBlendColorMul;
 		//    FinalColor = AdjustedDualBlendAdd + AdjustedDualBlendMul * BackgroundColor;
 
-		float3 AdjustedDualBlendAdd = Fogging.rgb + Fogging.a * DualBlendColorAdd;
-		float3 AdjustedDualBlendMul =               Fogging.a * DualBlendColorMul;
+		float3 AdjustedDualBlendAdd = DualBlendColorCov * (Fogging.rgb + Fogging.a * DualBlendColorAdd);
+		float3 AdjustedDualBlendMul = lerp(1.0f, Fogging.a * DualBlendColorMul, DualBlendColorCov);
 
 		#if DUAL_SOURCE_COLOR_BLENDING_ENABLED
 			// no RETURN_COLOR because these values are explicit multiplies and adds
diff --git a/Engine/Shaders/Private/DeferredLightPixelShaders.usf b/Engine/Shaders/Private/DeferredLightPixelShaders.usf
--- a/Engine/Shaders/Private/DeferredLightPixelShaders.usf
+++ b/Engine/Shaders/Private/DeferredLightPixelShaders.usf
@@ -230,11 +230,11 @@ void DeferredLightPixelMain(
 					GetStrataPixelDataByteOffset(PixelPos, View.BufferSizeAndInvSize.x, Strata.MaxBytesPerPixel));
 			}
 		}
-		else
-#endif
+#else
 		{
-		Radiance = GetDynamicLighting(DerivedParams.WorldPosition, DerivedParams.CameraVector, ScreenSpaceData.GBuffer, ScreenSpaceData.AmbientOcclusion, ScreenSpaceData.GBuffer.ShadingModelID, LightData, GetPerPixelLightAttenuation(InputParams.ScreenUV), Dither, uint2(InputParams.PixelPos), RectTexture, SurfaceShadow);
+			Radiance = GetDynamicLighting(DerivedParams.WorldPosition, DerivedParams.CameraVector, ScreenSpaceData.GBuffer, ScreenSpaceData.AmbientOcclusion, ScreenSpaceData.GBuffer.ShadingModelID, LightData, GetPerPixelLightAttenuation(InputParams.ScreenUV), Dither, uint2(InputParams.PixelPos), RectTexture, SurfaceShadow);
 		}
+#endif
 
 		const float  Attenuation = ComputeLightProfileMultiplier(DerivedParams.WorldPosition, DeferredLightUniforms.Position, -DeferredLightUniforms.Direction, DeferredLightUniforms.Tangent);
 
diff --git a/Engine/Shaders/Private/ReflectionEnvironmentPixelShader.usf b/Engine/Shaders/Private/ReflectionEnvironmentPixelShader.usf
--- a/Engine/Shaders/Private/ReflectionEnvironmentPixelShader.usf
+++ b/Engine/Shaders/Private/ReflectionEnvironmentPixelShader.usf
@@ -313,6 +313,20 @@ float3 ReflectionEnvironment(FGBufferData GBuffer, float AmbientOcclusion, float
 	return -min(-Color.rgb, 0.0);
 }
 
+void ApplyCloudVolumetricShadow(in float3 WorldPosition, inout float4 SkyLighting)
+{
+	// Ideally we would compute spatially how much of the sky+cloud is visible for each point in the world. But we evaluate the sky light only once at the sky light component position as of today.
+	// To add some spatial variation, we can affect the sky light diffuse contribution according to cloud occlusion from above.
+	// This is an approximation because clouds are also occluding the sky in the sky light capture (a bit of a double contribution then). but it does help by adding spatially varying details.
+#if 0	// DISABLED for now because it has artefact with specular not being affected (and thus looking too bright which can be confusing for users).
+	if (CloudSkyAOEnabled)
+	{
+		float OutOpticalDepth = 0.0f;
+		SkyLighting.rgb *= GetCloudVolumetricShadow(WorldPosition, CloudSkyAOWorldToLightClipMatrix, CloudSkyAOFarDepthKm, CloudSkyAOTexture, CloudSkyAOSampler, OutOpticalDepth);
+	}
+#endif
+}
+
 void ReflectionEnvironmentSkyLighting(
 	in float4 SvPosition : SV_Position,
 	out float4 OutColor : SV_Target0)
@@ -351,16 +365,14 @@ void ReflectionEnvironmentSkyLighting(
 	uint2 PixelPos = SvPosition.xy;
 
 	const bool bUnlitMaterial = ShadingModelID == SHADINGMODELID_UNLIT;
-#if STRATA_ENABLED
-	const bool bStrataMaterial = ShadingModelID == SHADINGMODELID_STRATA;
-#else
-	const bool bStrataMaterial = false;
-#endif
 
 	float4 SkyLighting = 0;
-	BRANCH
+
 #if STRATA_ENABLED && ENABLE_SKY_LIGHT
-	if (bStrataMaterial)
+
+	// STRATA_TODO handle the ENABLE_SKY_LIGHT && ALLOW_STATIC_LIGHTING case 
+	BRANCH
+	if (ShadingModelID == SHADINGMODELID_STRATA)
 	{
 		const FStrataTopLayerData TopLayerData = StrataUnpackTopLayerData(Strata.TopLayerTexture.Load(uint3(PixelPos, 0)));
 		const float4 SSR = Texture2DSample(ScreenSpaceReflectionsTexture, ScreenSpaceReflectionsSampler, BufferUV);
@@ -380,46 +392,36 @@ void ReflectionEnvironmentSkyLighting(
 			TopLayerSpecularContributionFactor,
 			Strata.MaterialLobesBuffer,
 			GetStrataPixelDataByteOffset(PixelPos, View.BufferSizeAndInvSize.x, Strata.MaxBytesPerPixel));
-	}
-	else
-#endif
+
+		ApplyCloudVolumetricShadow(WorldPosition, SkyLighting);
+
+		OutColor.xyz = SkyLighting.rgb;
+		}
+
+#else // STRATA_ENABLED && ENABLE_SKY_LIGHT
+
+	BRANCH
 	if (!bUnlitMaterial) // Only light pixels marked as lit
 	{
 		SkyLighting.rgb = SkyLightDiffuse(GBuffer, AmbientOcclusion, BufferUV, ScreenPosition, BentNormal, DiffuseColor);
-	}
-	
-	BRANCH 
-	if (!bUnlitMaterial || bStrataMaterial)
-	{
-		// Ideally we would compute spatially how much of the sky+cloud is visible for each point in the world. But we evaluate the sky light only once at the sky light component position as of today.
-		// To add some spatial variation, we can affect the sky light diffuse contribution according to cloud occlusion from above.
-		// This is an approximation because clouds are also occluding the sky in the sky light capture (a bit of a double contribution then). but it does help by adding spatially varying details.
-#if 0	// DISABLED for now because it has artefact with specular not being affected (and thus looking too bright which can be confusing for users).
-		if (CloudSkyAOEnabled)
-		{
-			float OutOpticalDepth = 0.0f;
-			SkyLighting.rgb *= GetCloudVolumetricShadow(WorldPosition, CloudSkyAOWorldToLightClipMatrix, CloudSkyAOFarDepthKm, CloudSkyAOTexture, CloudSkyAOSampler, OutOpticalDepth);
-		}
-#endif
-	}
 
-	// Scene color alpha is used for ScreenSpaceSubsurfaceScattering (if that is not needed it can be disabled with SUBSURFACE_CHANNEL_MODE)
-#if STRATA_ENABLED
-	OutColor = SkyLighting;
-#else
-	FLightAccumulator LightAccumulator = (FLightAccumulator)0;
-	const bool bNeedsSeparateSubsurfaceLightAccumulation = UseSubsurfaceProfile(ShadingModelID);
+		ApplyCloudVolumetricShadow(WorldPosition, SkyLighting);
 
-	LightAccumulator_Add(LightAccumulator, SkyLighting.rgb, SkyLighting.rgb, 1.0f, bNeedsSeparateSubsurfaceLightAccumulation);
-	OutColor = LightAccumulator_GetResult(LightAccumulator);
-#endif
+		FLightAccumulator LightAccumulator = (FLightAccumulator)0;
+		const bool bNeedsSeparateSubsurfaceLightAccumulation = UseSubsurfaceProfile(ShadingModelID);
 
-#endif // ENABLE_DYNAMIC_SKY_LIGHT
+		LightAccumulator_Add(LightAccumulator, SkyLighting.rgb, SkyLighting.rgb, 1.0f, bNeedsSeparateSubsurfaceLightAccumulation);
+		OutColor = LightAccumulator_GetResult(LightAccumulator);
+	}
 
 	BRANCH
-	if (ShadingModelID != SHADINGMODELID_UNLIT && ShadingModelID != SHADINGMODELID_HAIR && ShadingModelID != SHADINGMODELID_STRATA)
+	if (!bUnlitMaterial && ShadingModelID != SHADINGMODELID_HAIR)
 	{
 		OutColor.xyz += ReflectionEnvironment(GBuffer, AmbientOcclusion, BufferUV, WorldPosition, V, SvPosition, BentNormal, SpecularColor);
 	}
+
+#endif // STRATA_ENABLED && ENABLE_SKY_LIGHT
+
+#endif // ENABLE_DYNAMIC_SKY_LIGHT
 }
 
diff --git a/Engine/Source/Runtime/Engine/Classes/Materials/MaterialExpressionStrata.h b/Engine/Source/Runtime/Engine/Classes/Materials/MaterialExpressionStrata.h
--- a/Engine/Source/Runtime/Engine/Classes/Materials/MaterialExpressionStrata.h
+++ b/Engine/Source/Runtime/Engine/Classes/Materials/MaterialExpressionStrata.h
@@ -110,6 +110,8 @@ class UMaterialExpressionStrataSlabBSDF : public UMaterialExpressionStrataBSDF
 	virtual bool IsResultStrataMaterial(int32 OutputIndex) override;
 	virtual void GatherStrataMaterialInfo(FStrataMaterialInfo& StrataMaterialInfo, int32 OutputIndex) override;
 	virtual FName GetInputName(int32 InputIndex) const override;
+
+	bool HasScattering() const;
 #endif
 	//~ End UMaterialExpression Interface
 };
diff --git a/Engine/Source/Runtime/Engine/Private/Materials/MaterialExpressions.cpp b/Engine/Source/Runtime/Engine/Private/Materials/MaterialExpressions.cpp
--- a/Engine/Source/Runtime/Engine/Private/Materials/MaterialExpressions.cpp
+++ b/Engine/Source/Runtime/Engine/Private/Materials/MaterialExpressions.cpp
@@ -19698,7 +19698,7 @@ int32 UMaterialExpressionStrataSlabBSDF::Compile(class FMaterialCompiler* Compil
 	uint8 SharedNormalIndex = bAnisotropyPotentiallyUsed ? StrataCompilationInfoCreateSharedNormal(Compiler, NormalCodeChunk, TangentCodeChunk) : StrataCompilationInfoCreateSharedNormal(Compiler, NormalCodeChunk);
 
 	int32 SSSProfileCodeChunk = INDEX_NONE;
-	const bool bHasScattering = SubsurfaceProfile != nullptr;
+	const bool bHasScattering = HasScattering();
 	if (bHasScattering)
 	{
 		FName NameSubsurfaceProfile(FString(TEXT("__SubsurfaceProfile")));
@@ -19825,8 +19825,7 @@ bool UMaterialExpressionStrataSlabBSDF::IsResultStrataMaterial(int32 OutputIndex
 
 void UMaterialExpressionStrataSlabBSDF::GatherStrataMaterialInfo(FStrataMaterialInfo& StrataMaterialInfo, int32 OutputIndex)
 {
-	const bool bIsScaleInputConnected = true; // STRATA_TODO
-	if (SubsurfaceProfile || bIsScaleInputConnected)
+	if (HasScattering())
 	{
 		StrataMaterialInfo.AddShadingModel(SSM_SubsurfaceLit);
 		if (SubsurfaceProfile)
@@ -19839,6 +19838,11 @@ void UMaterialExpressionStrataSlabBSDF::GatherStrataMaterialInfo(FStrataMaterial
 		StrataMaterialInfo.AddShadingModel(SSM_DefaultLit);
 	}
 }
+
+bool UMaterialExpressionStrataSlabBSDF::HasScattering() const
+{
+	return SubsurfaceProfile != nullptr || DiffuseMeanFreePathRadius.IsConnected();
+}
 #endif // WITH_EDITOR
 
 
diff --git a/Engine/Source/Runtime/Engine/Private/Materials/StrataMaterial.h b/Engine/Source/Runtime/Engine/Private/Materials/StrataMaterial.h
--- a/Engine/Source/Runtime/Engine/Private/Materials/StrataMaterial.h
+++ b/Engine/Source/Runtime/Engine/Private/Materials/StrataMaterial.h
@@ -11,8 +11,6 @@
 // Material conversion to strata is disabled until ready.
 // To enable it locally:
 //		- Set CONVERT_MATERIAL_TO_STRATA_ON_LOAD to 1
-//		- Open UE5MainStreamObjectVersion.h
-//		- Add StrataMaterialConversion at the end of FUE5MainStreamObjectVersion::Type
 //		- Locally change ShaderVersion.ush GUID
 // !!!!!! !!!!!! DO NOT CHECK THIS IN SET TO 1 !!!!!! !!!!!!
 #define CONVERT_MATERIAL_TO_STRATA_ON_LOAD 0
