- Add pass to propogated mapped page information down the page table to enable O(1) lookups of any mapped page at the leaves
- Move some functions in PageManagement.usf to avoid confusing global resource scopes

diff --git a/Engine/Shaders/Private/VirtualShadowMaps/PageAccessCommon.ush b/Engine/Shaders/Private/VirtualShadowMaps/PageAccessCommon.ush
--- a/Engine/Shaders/Private/VirtualShadowMaps/PageAccessCommon.ush
+++ b/Engine/Shaders/Private/VirtualShadowMaps/PageAccessCommon.ush
@@ -158,7 +158,6 @@ uint ShadowEncodePageTable(uint2 PageIndex)
 {
 	return VSM_PHYSICAL_PAGE_ANY_MIP_VALID_FLAG | (PageIndex.y << 10) | (PageIndex.x);
 }
-
 uint ShadowEncodePageTable(uint2 PageIndex, uint MipOffset)
 {
 	return VSM_PHYSICAL_PAGE_ANY_MIP_VALID_FLAG | (MipOffset << 20) | (PageIndex.y << 10) | (PageIndex.x);
@@ -184,12 +183,40 @@ bool VirtualToPhysicalTexel(uint ShadowMapID, uint Level, uint2 VirtualTexelAddr
 	uint VPageX = VirtualTexelAddress.x >> VSM_LOG2_PAGE_SIZE;
 	uint VPageY = VirtualTexelAddress.y >> VSM_LOG2_PAGE_SIZE;
 
-	// TODO: Clever mip offset stuff
-	FShadowPhysicalPage PhysicalPageEntry = ShadowGetPhysicalPage(CalcPageTableLevelOffset(ShadowMapID, Level) + CalcPageOffsetInLevel(Level, uint2(VPageX, VPageY)));
+	FShadowPhysicalPage PhysicalPageEntry = ShadowGetPhysicalPage(CalcPageOffset(ShadowMapID, Level, uint2(VPageX, VPageY)));
 	PhysicalTexelAddress = PhysicalPageEntry.PageIndex * VSM_PAGE_SIZE + (VirtualTexelAddress & VSM_PAGE_SIZE_MASK);
 	return (PhysicalPageEntry.bThisMipValid);
 }
 
+struct FShadowPageTranslationResult
+{
+	bool bValid;
+	uint MipLevel;
+	uint2 VirtualTexelAddress;
+	float2 VirtualTexelAddressFloat;
+	uint2 PhysicalTexelAddress;
+};
+
+// Finds the best-resolution mapped page at the given UV
+FShadowPageTranslationResult ShadowVirtualToPhysicalUV(uint VirtualShadowMapID, float2 ShadowMapUV)
+{
+	float2 VirtualTexelAddressLevell0Float = ShadowMapUV * float(CalcLevelDimsTexels(0));
+	uint2 VirtualTexelAddressLevel0 = uint2(VirtualTexelAddressLevell0Float);
+	uint2 vPageLevel0 = uint2(VirtualTexelAddressLevel0 >> VSM_LOG2_PAGE_SIZE);
+
+	FShadowPhysicalPage PhysicalPageEntry = ShadowGetPhysicalPage(CalcPageOffset(VirtualShadowMapID, 0, vPageLevel0));
+
+	FShadowPageTranslationResult Result;
+	Result.bValid = PhysicalPageEntry.bAnyMipValid;
+	Result.MipLevel = PhysicalPageEntry.MipOffset;	
+	// TODO: Can optimize this slightly based on relative offset
+	Result.VirtualTexelAddressFloat = ShadowMapUV * float(CalcLevelDimsTexels(Result.MipLevel));
+	Result.VirtualTexelAddress = uint2(Result.VirtualTexelAddressFloat);
+	Result.PhysicalTexelAddress = PhysicalPageEntry.PageIndex * VSM_PAGE_SIZE + (Result.VirtualTexelAddress & VSM_PAGE_SIZE_MASK);
+
+	return Result;
+}
+
 // TODO: this surely needs some special handling of cases that contain the near plane etc.
 void TransformBoundingBox2D(float3 BoundsIn[2], float4x4 Transform, inout float2 BoundsOut[2])
 {
diff --git a/Engine/Shaders/Private/VirtualShadowMaps/PageManagement.usf b/Engine/Shaders/Private/VirtualShadowMaps/PageManagement.usf
--- a/Engine/Shaders/Private/VirtualShadowMaps/PageManagement.usf
+++ b/Engine/Shaders/Private/VirtualShadowMaps/PageManagement.usf
@@ -22,148 +22,8 @@
 RWStructuredBuffer<uint> OutStatsBuffer;
 #endif // VSM_GENERATE_STATS
 
-// Counter used to allocate physical pages from a linear range, which is mapped to a Physical 2D texture (by wrapping)
-RWStructuredBuffer<uint> AllocatedPagesOffset;
 // Flags generated by per-pixel pass to determine which pages are required to provide shadow for the visible geometry
 RWStructuredBuffer<uint> OutPageRequestFlags;
-RWStructuredBuffer<uint> OutPageTable;
 
-// Page flags generated by page allocation to indicate state to rendering passes (i.e., present / invalid)
-RWStructuredBuffer<uint> OutPageFlags;
-
-// Stores a uint2 with the physical location of cached pages in the previous frame physical texture, newly allocated ones are set to invalid index
-RWStructuredBuffer<FCachedPageInfo> OutCachedPageInfos;
-
-void CreatePageMapping( uint ShadowMapID, uint MipLevel, uint2 PageAddress, uint PageOffset, bool bPageRequestFlag )
-{	
-	if( bPageRequestFlag )
-	{
-		uint pPageIndex = 0;
-		InterlockedAdd( AllocatedPagesOffset[0], 1, pPageIndex );
-
-#if VSM_GENERATE_STATS
-		InterlockedAdd(OutStatsBuffer[0], 1);
-#endif // VSM_GENERATE_STATS
-
-		if (pPageIndex >= VirtualShadowMap.MaxPhysicalPages)
-		{
-			// We end up here if we're out of physical pages, this means some parts get no physical backing provided.
-			// Post this error condition back to the host somehow!
-			// Probably want to know if we're getting close even.
-			OutPageTable[PageOffset] = VSM_PHYSICAL_PAGE_INVALID;
-			OutPageFlags[PageOffset] = 0;
-			return;
-		}
-
-		// Wrap linear page index into physical 2D space
-		uint2 pPage;
-		pPage.x = pPageIndex  & VirtualShadowMap.PhysicalPageRowMask;
-		pPage.y = pPageIndex >> VirtualShadowMap.PhysicalPageRowShift;
-
-		OutPageTable[PageOffset] = ShadowEncodePageTable(pPage);
-
-		FCachedPageInfo CachedPageInfo;
-
-		CachedPageInfo.PhysPageAddress = uint2(VSM_INVALID_PHYSICAL_PAGE_ADDRESS, VSM_INVALID_PHYSICAL_PAGE_ADDRESS);
-		CachedPageInfo.DepthOffset = 0.0f;
-		CachedPageInfo.Padding = 0.0f;
-
-#if HAS_CACHE_DATA
-		if (ShadowMapCacheData[ShadowMapID].VirtualShadowMapId != INDEX_NONE)
-		{
-			CachedPageInfo.DepthOffset = ShadowMapCacheData[ShadowMapID].DepthOffset;
-
-			int2 PrevPageAddress = int2(PageAddress) - (ShadowMapCacheData[ShadowMapID].SmPageOffset >> MipLevel);
-			// A page cannot be cached if it is above the alignment level, or the light has not moved
-			// This means that for small movement a CSM can cache everything as well.
-			bool bCanCache = MipLevel < VSM_CACHE_ALIGNMENT_LEVEL || all(ShadowMapCacheData[ShadowMapID].SmPageOffset == 0);
-			if (all(uint2(PrevPageAddress) < CalcLevelDimsPages(MipLevel)) && bCanCache)
-			{
-				uint PrevOffset = CalcPageOffset(ShadowMapCacheData[ShadowMapID].VirtualShadowMapId, MipLevel, PrevPageAddress);
-				uint PrevPageFlag = PrevPageFlags[PrevOffset];
-				// True if the page had moving stuff drawn into it
-				bool PrevDynPageFlag = PrevDynamicCasterPageFlags[PrevOffset] != 0;
-
-#if VSM_GENERATE_STATS
-				// stats for those pages who die from moving stuff
-				if (PrevPageFlag && PrevDynPageFlag)
-				{
-					InterlockedAdd(OutStatsBuffer[2], 1);
-				}
-#endif // VSM_GENERATE_STATS
-				// If it was either rendered or cached previous frame, and not drawn to by a moving thing we may re-use
-				if (PrevPageFlag && !PrevDynPageFlag)
-				{
-					uint2 PrevPhysAddress = ShadowDecodePageTable(PrevPageTable[PrevOffset]).PageIndex;
-
-					// Only re-use this page if it was actually rendered to (i.e., mapped during in a render call)
-					FPhysicalPageMetaData MetaData = PrevPhysicalPageMetaData[PhysPageAddressToIndex(PrevPhysAddress)];
-					if (MetaData.State == VSM_PHYSICAL_PAGE_STATE_RENDERED)
-					{
-// 'rand robin' invalidation
-#if 0
-						// If cached for too long invalidate
-						if (MetaData.Age > 3 + (PageOffset & 3))
-						{
-#if VSM_GENERATE_STATS
-							InterlockedAdd(OutStatsBuffer[4], 1);
-#endif // VSM_GENERATE_STATS
-						}
-						else
-#endif
-						{
-#if VSM_GENERATE_STATS
-							InterlockedAdd(OutStatsBuffer[1], 1);
-#endif // VSM_GENERATE_STATS
-
-							CachedPageInfo.PhysPageAddress = PrevPhysAddress;
-						}
-					}
-				}
-// Needs more work (must disable page AND ensure fallback page is enabled)!
-#if 0
-				else
-				{
-					// Find fallback
-					uint PrevIncCoarserPageFlag = PrevPageFlag;
-					uint PrevCoarseOffset = 0;
-					// Check coarser levels (up to 4 levels away, arbitrary, hardcoded)
-					for (uint CoarseMipLevel = MipLevel + 1; CoarseMipLevel < min(MipLevel + 4, VSM_MAX_MIP_LEVELS); ++CoarseMipLevel)
-					{
-						PrevCoarseOffset = CalcPageOffset(ShadowMapCacheData[ShadowMapID].VirtualShadowMapId, CoarseMipLevel, PrevPageAddress >> (CoarseMipLevel - MipLevel));
-						uint PrevCoarserPageFlag = PrevPageFlags[PrevCoarseOffset];
-						if (PrevCoarserPageFlag)
-						{
-							PrevIncCoarserPageFlag = 1;
-							break;
-						}
-					}
-					if (PrevIncCoarserPageFlag)
-					{
-#if VSM_GENERATE_STATS
-#endif // VSM_GENERATE_STATS
-						InterlockedAdd(OutStatsBuffer[2], 1);
-						CachedPageInfo.PhysPageAddress = ShadowDecodePageTable(PrevPageTable[PrevCoarseOffset]).PageIndex;
-					}
-				}
-#endif
-			}
-		}
-#endif // HAS_CACHE_DATA
-
-		// Store data that is used later in the process.
-		OutCachedPageInfos[pPageIndex] = CachedPageInfo;
-		// Request rendering of pages that are not cached (AKA VSM_INVALID_PHYSICAL_PAGE_ADDRESS in cache).
-		uint Flags = (CachedPageInfo.PhysPageAddress.x == VSM_INVALID_PHYSICAL_PAGE_ADDRESS ? VSM_INVALID_FLAG : 0) | VSM_ALLOCATED_FLAG;
-		OutPageFlags[PageOffset] = Flags;
-	}
-	else
-	{
-		OutPageTable[PageOffset] = VSM_PHYSICAL_PAGE_INVALID;
-		OutPageFlags[PageOffset] = 0;
-	}
-}
-
-
 Texture2D<UlongType> VisBuffer64;
 StructuredBuffer<FVirtualShadowMapProjectionShaderData> ShadowMapProjectionData;
@@ -399,7 +259,6 @@
 }
 
 
-
 RWStructuredBuffer<uint> OutHPageFlags;
 RWStructuredBuffer<uint4> PageRectBoundsOut;
 
@@ -487,9 +346,149 @@
 	}
 }
 
+
+// Counter used to allocate physical pages from a linear range, which is mapped to a Physical 2D texture (by wrapping)
+RWStructuredBuffer<uint> AllocatedPagesOffset;
+// Page flags generated by page allocation to indicate state to rendering passes (i.e., present / invalid)
+RWStructuredBuffer<uint> OutPageFlags;
+RWStructuredBuffer<uint> OutPageTable;
+
+// Stores a uint2 with the physical location of cached pages in the previous frame physical texture, newly allocated ones are set to invalid index
+RWStructuredBuffer<FCachedPageInfo> OutCachedPageInfos;
+
 RWStructuredBuffer<uint> CoverageSummaryInOut;
 StructuredBuffer<uint> PageRequestFlags;
 
+void CreatePageMapping( uint ShadowMapID, uint MipLevel, uint2 PageAddress, uint PageOffset, bool bPageRequestFlag )
+{	
+	if( bPageRequestFlag )
+	{
+		uint pPageIndex = 0;
+		InterlockedAdd( AllocatedPagesOffset[0], 1, pPageIndex );
+
+#if VSM_GENERATE_STATS
+		InterlockedAdd(OutStatsBuffer[0], 1);
+#endif // VSM_GENERATE_STATS
+
+		if (pPageIndex >= VirtualShadowMap.MaxPhysicalPages)
+		{
+			// We end up here if we're out of physical pages, this means some parts get no physical backing provided.
+			// Post this error condition back to the host somehow!
+			// Probably want to know if we're getting close even.
+			OutPageTable[PageOffset] = VSM_PHYSICAL_PAGE_INVALID;
+			OutPageFlags[PageOffset] = 0;
+			return;
+		}
+
+		// Wrap linear page index into physical 2D space
+		uint2 pPage;
+		pPage.x = pPageIndex  & VirtualShadowMap.PhysicalPageRowMask;
+		pPage.y = pPageIndex >> VirtualShadowMap.PhysicalPageRowShift;
+
+		OutPageTable[PageOffset] = ShadowEncodePageTable(pPage);
+
+		FCachedPageInfo CachedPageInfo;
+
+		CachedPageInfo.PhysPageAddress = uint2(VSM_INVALID_PHYSICAL_PAGE_ADDRESS, VSM_INVALID_PHYSICAL_PAGE_ADDRESS);
+		CachedPageInfo.DepthOffset = 0.0f;
+		CachedPageInfo.Padding = 0.0f;
+
+#if HAS_CACHE_DATA
+		if (ShadowMapCacheData[ShadowMapID].VirtualShadowMapId != INDEX_NONE)
+		{
+			CachedPageInfo.DepthOffset = ShadowMapCacheData[ShadowMapID].DepthOffset;
+
+			int2 PrevPageAddress = int2(PageAddress) - (ShadowMapCacheData[ShadowMapID].SmPageOffset >> MipLevel);
+			// A page cannot be cached if it is above the alignment level, or the light has not moved
+			// This means that for small movement a CSM can cache everything as well.
+			bool bCanCache = MipLevel < VSM_CACHE_ALIGNMENT_LEVEL || all(ShadowMapCacheData[ShadowMapID].SmPageOffset == 0);
+			if (all(uint2(PrevPageAddress) < CalcLevelDimsPages(MipLevel)) && bCanCache)
+			{
+				uint PrevOffset = CalcPageOffset(ShadowMapCacheData[ShadowMapID].VirtualShadowMapId, MipLevel, PrevPageAddress);
+				uint PrevPageFlag = PrevPageFlags[PrevOffset];
+				// True if the page had moving stuff drawn into it
+				bool PrevDynPageFlag = PrevDynamicCasterPageFlags[PrevOffset] != 0;
+
+#if VSM_GENERATE_STATS
+				// stats for those pages who die from moving stuff
+				if (PrevPageFlag && PrevDynPageFlag)
+				{
+					InterlockedAdd(OutStatsBuffer[2], 1);
+				}
+#endif // VSM_GENERATE_STATS
+				// If it was either rendered or cached previous frame, and not drawn to by a moving thing we may re-use
+				if (PrevPageFlag && !PrevDynPageFlag)
+				{
+					uint2 PrevPhysAddress = ShadowDecodePageTable(PrevPageTable[PrevOffset]).PageIndex;
+
+					// Only re-use this page if it was actually rendered to (i.e., mapped during in a render call)
+					FPhysicalPageMetaData MetaData = PrevPhysicalPageMetaData[PhysPageAddressToIndex(PrevPhysAddress)];
+					if (MetaData.State == VSM_PHYSICAL_PAGE_STATE_RENDERED)
+					{
+// 'rand robin' invalidation
+#if 0
+						// If cached for too long invalidate
+						if (MetaData.Age > 3 + (PageOffset & 3))
+						{
+#if VSM_GENERATE_STATS
+							InterlockedAdd(OutStatsBuffer[4], 1);
+#endif // VSM_GENERATE_STATS
+						}
+						else
+#endif
+						{
+#if VSM_GENERATE_STATS
+							InterlockedAdd(OutStatsBuffer[1], 1);
+#endif // VSM_GENERATE_STATS
+
+							CachedPageInfo.PhysPageAddress = PrevPhysAddress;
+						}
+					}
+				}
+// Needs more work (must disable page AND ensure fallback page is enabled)!
+#if 0
+				else
+				{
+					// Find fallback
+					uint PrevIncCoarserPageFlag = PrevPageFlag;
+					uint PrevCoarseOffset = 0;
+					// Check coarser levels (up to 4 levels away, arbitrary, hardcoded)
+					for (uint CoarseMipLevel = MipLevel + 1; CoarseMipLevel < min(MipLevel + 4, VSM_MAX_MIP_LEVELS); ++CoarseMipLevel)
+					{
+						PrevCoarseOffset = CalcPageOffset(ShadowMapCacheData[ShadowMapID].VirtualShadowMapId, CoarseMipLevel, PrevPageAddress >> (CoarseMipLevel - MipLevel));
+						uint PrevCoarserPageFlag = PrevPageFlags[PrevCoarseOffset];
+						if (PrevCoarserPageFlag)
+						{
+							PrevIncCoarserPageFlag = 1;
+							break;
+						}
+					}
+					if (PrevIncCoarserPageFlag)
+					{
+#if VSM_GENERATE_STATS
+#endif // VSM_GENERATE_STATS
+						InterlockedAdd(OutStatsBuffer[2], 1);
+						CachedPageInfo.PhysPageAddress = ShadowDecodePageTable(PrevPageTable[PrevCoarseOffset]).PageIndex;
+					}
+				}
+#endif
+			}
+		}
+#endif // HAS_CACHE_DATA
+
+		// Store data that is used later in the process.
+		OutCachedPageInfos[pPageIndex] = CachedPageInfo;
+		// Request rendering of pages that are not cached (AKA VSM_INVALID_PHYSICAL_PAGE_ADDRESS in cache).
+		uint Flags = (CachedPageInfo.PhysPageAddress.x == VSM_INVALID_PHYSICAL_PAGE_ADDRESS ? VSM_INVALID_FLAG : 0) | VSM_ALLOCATED_FLAG;
+		OutPageFlags[PageOffset] = Flags;
+	}
+	else
+	{
+		OutPageTable[PageOffset] = VSM_PHYSICAL_PAGE_INVALID;
+		OutPageFlags[PageOffset] = 0;
+	}
+}
+
 /**
  * One thread per 2x2 page flags in level 0 (aka num in level 1), launched as 1d groups, with 2D grid with Y dim ==  NumShadowMaps.
  */
@@ -586,6 +585,63 @@
 	}
 }
 
+/**
+* One thread per page in level 0, launched as 1d groups, with 2D grid with Y dim ==  NumShadowMaps.
+* This is effectively just a big broadcast operation. There are more efficient ways to do this with
+* fewer threads and wave ops, but given the page counts just relying on memory coalescing is
+* good enough for now.
+*/
+[numthreads(VSM_DEFAULT_CS_GROUP_X, 1, 1)]
+void PropagateMappedMips(uint2 ThreadId : SV_DispatchThreadID, uint2 GroupId : SV_GroupID)
+{
+	const uint NumLevel0Entries = VSM_LEVEL0_DIM_PAGES_XY * VSM_LEVEL0_DIM_PAGES_XY;
+	if (ThreadId.x >= NumLevel0Entries)
+	{
+		return;
+	}
 
+	uint ShadowMapID = GroupId.y;
+	uint PageTableEntryIndex = ThreadId.x;
+
+	uint2 Level0Page;
+	Level0Page.x = PageTableEntryIndex & ((1U << VSM_LOG2_LEVEL0_DIM_PAGES_XY) - 1U);
+	Level0Page.y = PageTableEntryIndex >> VSM_LOG2_LEVEL0_DIM_PAGES_XY;
+
+	// Keep track of the last parent page that we saw that was mapped (if any)
+	int MappedPageLevel = VSM_MAX_MIP_LEVELS - 1;
+	uint2 MappedPageIndex = ShadowDecodePageTable(OutPageTable[CalcPageOffset(ShadowMapID, MappedPageLevel, uint2(0, 0))]).PageIndex;
+	
+	for (int Level = (VSM_MAX_MIP_LEVELS - 2); Level >= 0; --Level)
+	{
+		uint2 vPage = Level0Page >> Level;
+		uint PageOffset = CalcPageOffset(ShadowMapID, Level, vPage);
+		FShadowPhysicalPage pPage = ShadowDecodePageTable(OutPageTable[PageOffset]);
+
+		if (pPage.bThisMipValid)
+		{
+			// This page is mapped, so leave it alone and propogate downwards
+			MappedPageLevel = Level;
+			MappedPageIndex = pPage.PageIndex;
+		}
+		else
+		{
+			// This page is not mapped; replace it with our suitably offset parent mapped page
+			// Ensure only one thread writes each value to avoid races, but we read on all threads as the broadcast
+			// Note that this can race with the other threads reading this value, but since bThisMipValid will
+			// always be false on these updated pages the values will be ignored. As long as the writes to the page
+			// table are atomic (currently a single DWORD), this is safe.
+			if (all((vPage << Level) == Level0Page))
+			{
+				uint MipOffset = MappedPageLevel - Level;
+				OutPageTable[PageOffset] = ShadowEncodePageTable(MappedPageIndex, MipOffset);
+			}
+		}
+	}
+}
+
+
+
+
+
 #define LOG2_TILE_SIZE	5
 #define LOG2_TILES_PER_PAGE_X ( VSM_LOG2_PAGE_SIZE - LOG2_TILE_SIZE )
diff --git a/Engine/Shaders/Private/VirtualShadowMaps/ProjectionCommon.ush b/Engine/Shaders/Private/VirtualShadowMaps/ProjectionCommon.ush
--- a/Engine/Shaders/Private/VirtualShadowMaps/ProjectionCommon.ush
+++ b/Engine/Shaders/Private/VirtualShadowMaps/ProjectionCommon.ush
@@ -94,7 +94,6 @@ FVirtualSmSample SampleVirtualSmLevel(int VirtualShadowMapId, float2 ShadowMapUV
 	bool PagePresent = VirtualToPhysicalTexel(VirtualShadowMapId, VSMLevel, Result.VirtualTexelAddress, Result.PhysicalTexelAddress);
 	if (PagePresent)
 	{
-		// TODO: maybe store metadata about Max depth (and/or use sparse HZB) for each tile, and early out if completely in shadow (must test using kernel size).
 #if ENABLE_NON_NANITE_VSM
 		Result.Depth = max(
 			asfloat(PhysicalPagePool.Load(int3(Result.PhysicalTexelAddress, 0))),
@@ -112,13 +111,17 @@ FVirtualSmSample SampleVirtualSmLevel(int VirtualShadowMapId, float2 ShadowMapUV
 
 FVirtualSmSample SampleVirtualSm(int VirtualShadowMapId, float2 ShadowMapUV)
 {
-	for (uint VSMLevel = 0; VSMLevel < VSM_MAX_MIP_LEVELS; ++VSMLevel)
+	FShadowPageTranslationResult Page = ShadowVirtualToPhysicalUV(VirtualShadowMapId, ShadowMapUV);
+	if (Page.bValid)
 	{
-		FVirtualSmSample Result = SampleVirtualSmLevel(VirtualShadowMapId, ShadowMapUV, VSMLevel);
-		if (Result.bValid)
-		{
-			return Result;
-		}
+		FVirtualSmSample Result;
+		Result.bValid = true;
+		Result.MipLevel = Page.MipLevel;
+		Result.VirtualTexelAddress = Page.VirtualTexelAddress;
+		Result.VirtualTexelAddressFloat = Page.VirtualTexelAddressFloat;
+		Result.PhysicalTexelAddress = Page.PhysicalTexelAddress;
+		Result.Depth = asfloat(PhysicalPagePool.Load(int3(Result.PhysicalTexelAddress, 0)));
+		return Result;
 	}
 
 	return InitVirtualSmSample();
diff --git a/Engine/Source/Runtime/Renderer/Private/VirtualShadowMaps/VirtualShadowMapArray.cpp b/Engine/Source/Runtime/Renderer/Private/VirtualShadowMaps/VirtualShadowMapArray.cpp
--- a/Engine/Source/Runtime/Renderer/Private/VirtualShadowMaps/VirtualShadowMapArray.cpp
+++ b/Engine/Source/Runtime/Renderer/Private/VirtualShadowMaps/VirtualShadowMapArray.cpp
@@ -349,6 +349,18 @@ class FCreatePageMappingsCS : public FVirtualPageManagementShader
 };
 IMPLEMENT_GLOBAL_SHADER(FCreatePageMappingsCS, "/Engine/Private/VirtualShadowMaps/PageManagement.usf", "CreatePageMappings", SF_Compute);
 
+class FPropagateMappedMipsCS : public FVirtualPageManagementShader
+{
+	DECLARE_GLOBAL_SHADER(FPropagateMappedMipsCS);
+	SHADER_USE_PARAMETER_STRUCT(FPropagateMappedMipsCS, FVirtualPageManagementShader)
+
+	BEGIN_SHADER_PARAMETER_STRUCT(FParameters, )
+		SHADER_PARAMETER_RDG_UNIFORM_BUFFER(FVirtualShadowMapCommonParameters,	VirtualSmCommon)
+		SHADER_PARAMETER_RDG_BUFFER_UAV( RWStructuredBuffer< uint >,			OutPageTable )
+	END_SHADER_PARAMETER_STRUCT()
+};
+IMPLEMENT_GLOBAL_SHADER(FPropagateMappedMipsCS, "/Engine/Private/VirtualShadowMaps/PageManagement.usf", "PropagateMappedMips", SF_Compute);
+
 class FInitClearPhysicalPagesArgsCS : public FVirtualPageManagementShader
 {
 	DECLARE_GLOBAL_SHADER(FInitClearPhysicalPagesArgsCS);
@@ -980,6 +992,25 @@ void FVirtualShadowMapArray::BuildPageAllocations(
 				FIntVector(FMath::DivideAndRoundUp(CommonParameters.PageTableSize, FGenerateHierarchicalPageFlagsCS::DefaultCSGroupX), ShadowMaps.Num(), 1)
 			);
 		}
+
+		// NOTE: We could skip this (in shader) for shadow maps that only have 1 mip (ex. clipmaps)
+		{
+			// Propagate mapped mips down the hierarchy to allow O(1) lookup of coarser mapped pages
+			FPropagateMappedMipsCS::FParameters* PassParameters = GraphBuilder.AllocParameters<FPropagateMappedMipsCS::FParameters>();
+			PassParameters->VirtualSmCommon = GetCommonUniformBuffer(GraphBuilder);
+			PassParameters->OutPageTable	= GraphBuilder.CreateUAV(PageTableRDG);
+
+			const uint32 DispatchWidthThreads = FVirtualShadowMap::Level0DimPagesXY;
+
+			auto ComputeShader = Views[0].ShaderMap->GetShader<FPropagateMappedMipsCS>();
+			FComputeShaderUtils::AddPass(
+				GraphBuilder,
+				RDG_EVENT_NAME("PropagateMappedMips"),
+				ComputeShader,
+				PassParameters,
+				FIntVector(FMath::DivideAndRoundUp(FMath::Square(DispatchWidthThreads), FPropagateMappedMipsCS::DefaultCSGroupX), ShadowMaps.Num(), 1)
+			);
+		}
 	}
 }
 
