Fix flattened page tables propagating a garbage root page.

diff --git a/Engine/Shaders/Private/VirtualShadowMaps/PageManagement.usf b/Engine/Shaders/Private/VirtualShadowMaps/PageManagement.usf
--- a/Engine/Shaders/Private/VirtualShadowMaps/PageManagement.usf
+++ b/Engine/Shaders/Private/VirtualShadowMaps/PageManagement.usf
@@ -607,23 +607,23 @@ void PropagateMappedMips(uint2 ThreadId : SV_DispatchThreadID, uint2 GroupId : S
 	Level0Page.x = PageTableEntryIndex & ((1U << VSM_LOG2_LEVEL0_DIM_PAGES_XY) - 1U);
 	Level0Page.y = PageTableEntryIndex >> VSM_LOG2_LEVEL0_DIM_PAGES_XY;
 
-	// Keep track of the last parent page that we saw that was mapped (if any)
-	int MappedPageLevel = VSM_MAX_MIP_LEVELS - 1;
-	uint2 MappedPageIndex = ShadowDecodePageTable(OutPageTable[CalcPageOffset(ShadowMapID, MappedPageLevel, uint2(0, 0))]).PageIndex;
+	int MappedPageLevel = -1;
+	uint2 MappedPageIndex = 0;
 	
-	for (int Level = (VSM_MAX_MIP_LEVELS - 2); Level >= 0; --Level)
+	for (int Level = (VSM_MAX_MIP_LEVELS - 1); Level >= 0; --Level)
 	{
 		uint2 vPage = Level0Page >> Level;
 		uint PageOffset = CalcPageOffset(ShadowMapID, Level, vPage);
 		FShadowPhysicalPage pPage = ShadowDecodePageTable(OutPageTable[PageOffset]);
 
+		BRANCH
 		if (pPage.bThisMipValid)
 		{
 			// This page is mapped, so leave it alone and propogate downwards
 			MappedPageLevel = Level;
 			MappedPageIndex = pPage.PageIndex;
 		}
-		else
+		else if( MappedPageLevel >= 0 )
 		{
 			// This page is not mapped; replace it with our suitably offset parent mapped page
 			// Ensure only one thread writes each value to avoid races, but we read on all threads as the broadcast
