Fixing burley SSS precision issues and border bleed accumulation.

diff --git a/Engine/Shaders/Private/BurleyNormalizedSSSCommon.ush b/Engine/Shaders/Private/BurleyNormalizedSSSCommon.ush
--- a/Engine/Shaders/Private/BurleyNormalizedSSSCommon.ush
+++ b/Engine/Shaders/Private/BurleyNormalizedSSSCommon.ush
@@ -131,36 +131,43 @@ float GetDiffuseMeanFreePathForSampling(float4 DiffuseMeanFreePath)
 // Average relative error: 5.5% (reference to MC)
 float GetPerpendicularScalingFactor(float SurfaceAlbedo)
 {
-	return 1.85 - SurfaceAlbedo + 7 * UnClampedPow(SurfaceAlbedo - 0.8, 3);
+	// from the paper, the formula explicitly has an abs for (A-0.8)
+	float Value = abs(SurfaceAlbedo - 0.8);
+	return 1.85 - SurfaceAlbedo + 7 * Value * Value * Value;
 }
 
 float3 GetPerpendicularScalingFactor3D(float3 SurfaceAlbedo)
 {
-	return 1.85 - SurfaceAlbedo + 7 * UnClampedPow(SurfaceAlbedo - 0.8, 3);
+	float3 Value = abs(SurfaceAlbedo - 0.8);
+	return 1.85 - SurfaceAlbedo + 7 * Value * Value * Value;
 }
 
 // Method 2: Ideal diffuse transmission at the surface. More appropriate for rough surface.
 // Average relative error: 3.9% (reference to MC)
 float GetDiffuseSurfaceScalingFactor(float SurfaceAlbedo)
 {
-	return 1.9 - SurfaceAlbedo + 3.5*(SurfaceAlbedo - 0.8)*(SurfaceAlbedo - 0.8);
+	float Value = SurfaceAlbedo - 0.8;
+	return 1.9 - SurfaceAlbedo + 3.5 * Value * Value;
 }
 
 float3 GetDiffuseSurfaceScalingFactor3D(float3 SurfaceAlbedo)
 {
-	return 1.9 - SurfaceAlbedo + 3.5*(SurfaceAlbedo - 0.8)*(SurfaceAlbedo - 0.8);
+	float3 Value = SurfaceAlbedo - 0.8;
+	return 1.9 - SurfaceAlbedo + 3.5 * Value * Value;
 }
 
 // Method 3: The spectral of diffuse mean free path on the surface.
 // Avergate relative error: 7.7% (reference to MC)
 float GetSearchLightDiffuseScalingFactor(float SurfaceAlbedo)
 {
-	return 3.5 + 100 * UnClampedPow(SurfaceAlbedo - 0.33, 4);
+	float Value = SurfaceAlbedo - 0.33;
+	return 3.5 + 100 * Value * Value * Value * Value;
 }
 
 float3 GetSearchLightDiffuseScalingFactor3D(float3 SurfaceAlbedo)
 {
-	return 3.5 + 100.0 * UnClampedPow(max(SurfaceAlbedo - 0.33,0.0), 4.0);
+	float3 Value = SurfaceAlbedo - 0.33;
+	return 3.5 + 100 * Value * Value * Value * Value;
 }
 
 //@TODO Revisit here after offline comparison.
diff --git a/Engine/Shaders/Private/PostProcessSubsurface.usf b/Engine/Shaders/Private/PostProcessSubsurface.usf
--- a/Engine/Shaders/Private/PostProcessSubsurface.usf
+++ b/Engine/Shaders/Private/PostProcessSubsurface.usf
@@ -520,7 +520,6 @@ void SetupIndirectCS(uint3 DT_ID : SV_DispatchThreadID,
 	GroupMemoryBarrierWithGroupSync();
 
 	uint2 Pos = DT_ID.xy + Output_ViewportMin;
-	bool bHasSubsurface = false;
 	float2 BufferUV = ConvertGridPos2UV(Pos);
 	uint Type = 0;
 	float4 OutColor = 0;
@@ -536,7 +535,12 @@ void SetupIndirectCS(uint3 DT_ID : SV_DispatchThreadID,
 
 	float4 Sum = (A + B) + (C + D);
 
-	float Div = 1.0f / max(Sum.a, 0.00001f);
+	// Each alpha value of A/B/C/D should be either 1.0f or 0.0, so a check of 0.5 should be safe
+	bool bHasSubsurface = (Sum.a >= .50f);
+
+	// Sum.a should either be one of {0,1,2,3,4}, add a max of 1 to avoid floating point specials. If Sum.a is 0, Sum.rgb should be 0 as well so it
+	// shouldn't matter what Div is.
+	float Div = 1.0f / max(Sum.a, 1.0f);
 
 	OutColor.rgb = Sum.rgb * Div;
 
@@ -553,6 +557,8 @@ void SetupIndirectCS(uint3 DT_ID : SV_DispatchThreadID,
 	}
 	//BufferUV += min(float2(-0.25, -0.25f) * SubsurfaceInput0_ExtentInverse, SubsurfaceInput0_UVViewportBilinearMax);
 #else // SUBSURFACE_HALF_RES
+	
+	bool bHasSubsurface = false;
 	OutColor = SetupSubsurfaceForOnePixel(BufferUV);
 
 	if (OutColor.a > 0)
diff --git a/Engine/Shaders/Private/SubsurfaceBurleyNormalized.ush b/Engine/Shaders/Private/SubsurfaceBurleyNormalized.ush
--- a/Engine/Shaders/Private/SubsurfaceBurleyNormalized.ush
+++ b/Engine/Shaders/Private/SubsurfaceBurleyNormalized.ush
@@ -904,6 +904,7 @@ float4 BurleyNormalizedSS(float2 BufferUV, uint2 GroupThreadID)
 	float3 WeightingFactor = 0.0f;
 	float4 RadianceAccumulated = float4(0.0f, 0.0f, 0.0f, 1.0f);
 	float Mask = 1.0f;
+	float3 BoundaryColorBleedAccum = float3(0.0f, 0.0f, 0.0f);
 
 #if SUBSURFACE_SINGLE_PASS
 	int NumOfSamples = BURLEY_NUM_SAMPLES;
@@ -1035,7 +1036,7 @@ float4 BurleyNormalizedSS(float2 BufferUV, uint2 GroupThreadID)
 
 			// Determine the tint color, if the sampling pixel is not subsurface, we use tint color
 			// to mask out the sampling. Unless we specifically want the shadowing region.
-			float3 TintColor = (LocalProfile == SubsurfaceProfileInt) ? 1.0f : BoundaryColorBleed;
+			BoundaryColorBleedAccum += (LocalProfile == SubsurfaceProfileInt || LocalProfile == SSS_PROFILE_ID_INVALID) ? 1.0f : BoundaryColorBleed;
 			Mask = (LocalProfile != SSS_PROFILE_ID_INVALID) ? 1 : 0;
 
 			// Step 4: accumulate radiance from the diffusion profile rR(r)
@@ -1043,7 +1044,7 @@ float4 BurleyNormalizedSS(float2 BufferUV, uint2 GroupThreadID)
 			float3 DiffusionProfile = GetDiffuseReflectProfileWithDiffuseMeanFreePath(BurleyParameter.DiffuseMeanFreePath.xyz, S3D.xyz, RadiusSampledInMM);
 			float3 SampleWeight = (DiffusionProfile / BurleySampleInfo.Pdf) * Mask * NormalWeight;
 
-			RadianceAccumulated.xyz += SampleWeight * (SampledRadianceAndDepth.xyz*TintColor);
+			RadianceAccumulated.xyz += SampleWeight * (SampledRadianceAndDepth.xyz);
 
 #if USE_BILATERAL_FILTERING
 			WeightingFactor += SampleWeight;
@@ -1066,6 +1067,8 @@ float4 BurleyNormalizedSS(float2 BufferUV, uint2 GroupThreadID)
 	RadianceAccumulated.xyz *= (WeightingFactor == 0)? 0:(1.0f /WeightingFactor*EnergyNormalization);
 #endif
 
+	RadianceAccumulated.xyz *= BoundaryColorBleedAccum*InvNumOfSamples;
+
 #if REWEIGHT_CENTER_SAMPLE
 	// Apply lerp with center pixel
 	RadianceAccumulated.xyz = lerp(RadianceAccumulated.xyz,OriginalColor,CenterSampleWeight);
