One pass VSM projection, culled by the light grid.
Major SMRT optimization. Adaptive 7 rays per light with rough penumbra detection. Quality appears to surpass 1rpp with denoising.

diff --git a/Engine/Shaders/Private/ClusteredDeferredShadingPixelShader.usf b/Engine/Shaders/Private/ClusteredDeferredShadingPixelShader.usf
--- a/Engine/Shaders/Private/ClusteredDeferredShadingPixelShader.usf
+++ b/Engine/Shaders/Private/ClusteredDeferredShadingPixelShader.usf
@@ -30,6 +30,8 @@
 // Extracts out the simple lights into their own loop that is executed first, independently of shading model.
 #define ENABLE_SEPARATE_SIMPLE_LIGHT_LOOP 0
 
+Texture2D< uint > ShadowMaskBits;
+
 /**
  * Adds local lighting using the light grid, does not apply directional lights, as they are done elsewhere.
  * Does not support dynamic shadows, as these require the per-light shadow mask.
@@ -45,6 +47,8 @@ float3 GetLightGridLocalLighting(const FCulledLightsGridData InLightGridData, fl
 	// This prevents GPU hangs when the PS tries to read from uninitialized NumCulledLightsGrid buffer
 	const uint NumLocalLights = min(InLightGridData.NumLocalLights, GetNumLocalLights(EyeIndex));
 
+	uint ShadowMask = ~ShadowMaskBits[ int2( ScreenUV.xy * View.BufferSizeAndInvSize.xy ) ];
+
 	LOOP
 	for (uint LocalLightListIndex = FirstNonSimpleLightIndex; LocalLightListIndex < NumLocalLights; LocalLightListIndex++)
 	{
@@ -107,6 +111,9 @@ float3 GetLightGridLocalLighting(const FCulledLightsGridData InLightGridData, fl
 		// This matches tiled deferred, since there is no dynamic shadow on the lights used for clustered shading.
 		float4 LightAttenuation = 1.0f;
 
+		LightData.ShadowedBits = 1;
+		LightAttenuation = ( ( ShadowMask >> (LocalLightListIndex * 3) ) & 7 ) / 7.0;
+
 #if 1
 		FRectTexture RectTexture = InitRectTexture(ForwardLightData.DummyRectLightSourceTexture);
 		// NOTE: uses AO data like tiled deferred (as opposed to forward path)
diff --git a/Engine/Shaders/Private/VirtualShadowMaps/ProjectionSpot.ush b/Engine/Shaders/Private/VirtualShadowMaps/ProjectionSpot.ush
--- a/Engine/Shaders/Private/VirtualShadowMaps/ProjectionSpot.ush
+++ b/Engine/Shaders/Private/VirtualShadowMaps/ProjectionSpot.ush
@@ -8,6 +8,7 @@ ProjectionSpot.ush:
 #include "../DeferredShadingCommon.ush"
 #include "../SceneTexturesCommon.ush"
 #include "../LightShaderParameters.ush"
+#include "../PathTracing/Utilities/PathTracingRandomSequence.ush"
 #include "PageAccessCommon.ush"
 #include "ProjectionCommon.ush"
 
@@ -81,7 +82,8 @@ bool ShadowRayCastSpotlight(
 	float FirstHitTime = 2;
 	float DepthHistory = -1;
 
-	for (int i = 1; i <= NumSteps; i++)
+	bool bValidHit = false;
+	for (int i = 0; i <= NumSteps; i++)
 	{
 		//SampleTime = Pow2( 1.0 - (float)i / NumSteps ) * RayScale;
 		float SampleTime = Pow2(saturate((-1.0f / NumSteps) * i + 1.0f - (StepOffset - 0.5f) / NumSteps));
@@ -97,6 +99,18 @@ bool ShadowRayCastSpotlight(
 				// First valid sample we've seen. Do a regular depth compare.
 				DepthHistory = SampleDepth;
 				FirstHitTime = SampleDepth > ReferenceDepth ? SampleTime : 2;
+				if( SampleDepth > ReferenceDepth )
+				{
+					float ReceiverDistance = length( RayStart + BaseProjectionData.ShadowPreViewTranslation );
+					OccluderDistance = ComputeOccluderDistancePerspective(
+						BaseProjectionData.ShadowViewToClipMatrix,
+						SmSample.Depth,
+						SampleUVz.z,
+						ReceiverDistance);
+
+					bValidHit = true;
+					break;
+				}
 			}
 			else
 			{
@@ -108,38 +122,18 @@ bool ShadowRayCastSpotlight(
 
 				float DepthDiff = ReferenceDepth - DepthSample;
 				bool bHit = abs(DepthDiff + CompareTolerance) < CompareTolerance;
-	#if 0
-				//FirstHitTime = (bHit && FirstHitTime < 0.0) ? SampleTime : FirstHitTime;
-				FirstHitTime = min( FirstHitTime, bHit ? SampleTime : 2 );
-	#else
-				//bool bValidUV = all( abs( SampleUVz.xy - 0.5 ) < 1 );	// Ray clipped to frustum. Shouldn't be possible to be outside.
 				if (bHit)
 				{
 					FirstHitTime = SampleTime;
+					OccluderDistance = length(RayEnd - RayStart) * SampleTime;
+
+					bValidHit = true;
+					break;
 				}
 			}
-#endif
 		}
 	}
 
-	// Off screen masking
-	//bool bValidUV = all( abs( RayStartScreen.xy + RayStepScreen.xy * FirstHitTime ) < 1 );
-	//bool bValidHit = bValidUV && FirstHitTime < 2;
-	bool bValidHit = FirstHitTime < 2;
-
-	if (bValidHit)
-	{
-		float RayLength = length(RayEnd - RayStart);
-		OccluderDistance = RayLength * FirstHitTime;
-
-		//float HitDepth = RayStartUVz.z + RayStepUVz.z * FirstHitTime;
-		//OccluderDistance += ComputeOccluderDistancePerspective(
-		//	BaseProjectionData.ShadowViewToClipMatrix,
-		//	HitDepth,
-		//	RayStartUVz.z,
-		//	RayStartView );
-	}
-
 	return bValidHit;
 }
 
@@ -157,7 +151,7 @@ float TraceSpotLight(
 	inout float3 DebugOutput,
 	inout float InOutHitDistance)
 {
-	FGBufferData GBufferData = GetGBufferData(SvPositionToBufferUV(SvPosition));
+	FGBufferData GBufferData = GetGBufferDataUint( uint2( SvPosition.xy ), true );
 	const float3 WorldPosition = TranslatedWorldPosition - View.PreViewTranslation;
 
 	float3 ToLight = Light.Position - WorldPosition;
@@ -176,14 +170,20 @@ float TraceSpotLight(
 
 		float StepOffset = InterleavedGradientNoise(SvPosition.xy, View.StateFrameIndexMod8);
 
-#if 0
-		uint SampleCount = SMRTRayCount;
+#if 1
+		uint SampleCount = RayCount;
 
-		float RayMissFactor = 0;
 		for( uint i = 0; i < SampleCount; i++ )
 		{
+#if 1
+			RandomSequence RandSequence;
+			RandomSequence_Initialize( RandSequence, uint2( SvPosition.xy ), i, View.StateFrameIndexMod8, SampleCount );
+
+			float2 E = RandomSequence_GenerateSample2D(RandSequence);
+#else
 			// NOTE: 1-X to avoid 0/0 issues with UniformSampleDiskConcentric
 			float2 E = float2( SobolIndex( SobolFrame, i << 3 ) ) / 0x10000;
+#endif
 
 			// NOTE: Avoid 0/0 issues with UniformSampleDiskConcentric
 			float2 LightUV = UniformSampleDiskConcentricApprox( E ) * Light.SourceRadius;
@@ -194,45 +194,62 @@ float TraceSpotLight(
 				LightUV.y * cross( Light.Direction, Light.Tangent );
 
 			float3 Dir = normalize( RayEnd - WorldPosition );
-			float3 RayStart = WorldPosition + Dir * ContactShadowLengthWorld;
+			float3 RayStart = WorldPosition + Dir * RayStartOffset;
 
 			RayEnd = lerp( RayStart, RayEnd, 0.75 );
-			//RayEnd = RayStart + L * ContactShadowLengthWorld;
-
-			float HitDistance = ContactShadowLengthWorld;
-			bool bHit = ShadowRayCastSpotlight( RayStart, RayEnd, 16, StepOffset, MipLevel, HitDistance );
 
-			//ContactShadow = VirtualShadowMapScreenRayCast( TranslatedWorldPosition, Dir, ContactShadowLengthWorld, 12, StepOffset );
-			//RayMissFactor += bHit ? 0 : ContactShadow;
+			bool bHit = ShadowRayCastSpotlight(
+				VirtualShadowMapId,
+				RayStart,
+				RayEnd,
+				SamplesPerRay,
+				StepOffset,
+				InOutHitDistance);
 
 			RayMissFactor += bHit ? 0 : 1;
+			
+#if COMPILER_SUPPORTS_WAVE_VOTE
+			if( i == 0 )
+			{
+				if( bHit )
+				{
+					float WorldBluringRadius = ComputeLightSampleWorldBluringRadius( TranslatedWorldPosition, LIGHT_TYPE_SPOT, Light, InOutHitDistance );
+					float PixelBluringRadius = ProjectWorldDistanceToPixelDistance( SceneDepth, WorldBluringRadius );
+					bHit = PixelBluringRadius > 2;
+				}
+
+				// All lanes miss or hit something close enough to not be noisy.
+				bool bAllLanesMiss = all( WaveActiveBallot( bHit ).xy == 0 );
+				if( bAllLanesMiss )
+				{
+					RayMissFactor *= SampleCount;
+					break;
+				}
+			}
+			else if( i == 1 )
+			{
+				// After 2 iterations and all have hit, must be in umbra.
+				bool bAllLanesHit = all( WaveActiveBallot( RayMissFactor > 0 ).xy == 0 );
+				if( bAllLanesHit )
+				{
+					RayMissFactor *= 0.5 * SampleCount;
+					break;
+				}
+			}
+#endif
 		}
 		RayMissFactor /= SampleCount;
+#else
 
-		float ContactShadow = 1.0f;
-		const float ContactShadowLengthWorld = ContactShadowLength * View.ClipToView[1][1] * SceneDepth;
-		if( ContactShadowLengthWorld > 0.0f )
-		{
-			float2 E = float2( SobolIndex( SobolFrame, (View.Random & 7) << 3 ) ) / 0x10000;
-
-			// NOTE: Avoid 0/0 issues with UniformSampleDiskConcentric
-			float2 LightUV = UniformSampleDiskConcentricApprox( E ) * Light.SourceRadius;
-
-			float3 RayEnd =
-				Light.Position +
-				LightUV.x * Light.Tangent +
-				LightUV.y * cross( Light.Direction, Light.Tangent );
-
-			float3 Dir = normalize( RayEnd - WorldPosition );
-
-			ContactShadow = VirtualShadowMapScreenRayCast( TranslatedWorldPosition, Dir, ContactShadowLengthWorld, 12, StepOffset );
-		}
-		RayMissFactor *= ContactShadow;
+#if 0
+		RandomSequence RandSequence;
+		RandomSequence_Initialize( RandSequence, uint2( SvPosition.xy ), 0, View.StateFrameIndexMod8, 1 );
 
-		OutputSignal = float4( RayMissFactor, SampleCount, 0, RayMissFactor );
+		float2 E = RandomSequence_GenerateSample2D(RandSequence);
 #else
 		// TODO: Loop over number of rays RayCount
 		float2 E = float2(SobolIndex(SobolFrame, (View.Random & 7) << 3)) / 0x10000;
+#endif
 
 		// NOTE: Avoid 0/0 issues with UniformSampleDiskConcentric
 		// TODO: RandomSequence
diff --git a/Engine/Shaders/Private/VirtualShadowMaps/VirtualShadowMapProjection.usf b/Engine/Shaders/Private/VirtualShadowMaps/VirtualShadowMapProjection.usf
--- a/Engine/Shaders/Private/VirtualShadowMaps/VirtualShadowMapProjection.usf
+++ b/Engine/Shaders/Private/VirtualShadowMaps/VirtualShadowMapProjection.usf
@@ -14,6 +14,7 @@
 #include "../ScreenSpaceDenoise/SSDDefinitions.ush"
 #include "../ScreenSpaceDenoise/SSDPublicHarmonics.ush"
 #include "../LightShaderParameters.ush"
+#include "../LightGridCommon.ush"
 #include "PageAccessCommon.ush"
 #include "ProjectionCommon.ush"
 #include "ProjectionDirectional.ush"
@@ -54,9 +55,8 @@ float VirtualShadowMapScreenRayCast(
 
 
 float ProjectSingleSample(
-	FLightShaderParameters Light,
+	int VirtualShadowMapId,
 	float3 L,
-	uint LightType,
 	float3 TranslatedWorldPosition,
 	float3 EstimatedGeoWorldNormal,
 	float ContactShadowLengthWorld,
@@ -102,8 +102,10 @@ float4 EncodeLightAttenuationFromMask(float ShadowMask)
 
 	// 0 is shadowed, 1 is unshadowed
 	// RETURN_COLOR not needed unless writing to SceneColor;
-	float FadedShadow = lerp(1.0f, Square(ShadowMask), ShadowFadeFraction);
-	float FadedSSSShadow = lerp(1.0f, Square(SSSTransmission), ShadowFadeFraction);
+	//float FadedShadow = lerp(1.0f, Square(ShadowMask), ShadowFadeFraction);
+	//float FadedSSSShadow = lerp(1.0f, Square(SSSTransmission), ShadowFadeFraction);
+	float FadedShadow = lerp(1.0f, ShadowMask, ShadowFadeFraction);
+	float FadedSSSShadow = lerp(1.0f, SSSTransmission, ShadowFadeFraction);
 
 	// the channel assignment is documented in ShadowRendering.cpp (look for Light Attenuation channel assignment)
 	return EncodeLightAttenuation(half4(FadedShadow, FadedSSSShadow, FadedShadow, FadedSSSShadow));
@@ -188,8 +190,8 @@ void VirtualShadowMapProjectionSpotPS(
 	// Early out if we're outside the spot light cone		
 	float3 ToLight = Light.Position - (TranslatedWorldPosition - View.PreViewTranslation);
 	bool bInLightRadius = length(ToLight) <= rcp(Light_InvRadius);
-	float3 ToLightDir = normalize(ToLight);
-	bool bInLightCone = dot(ToLightDir, Light.Direction) >= Light.SpotAngles.x;
+	float3 L = normalize(ToLight);
+	bool bInLightCone = dot(L, Light.Direction) >= Light.SpotAngles.x;
 	bool bInLightRegion = bInLightRadius && bInLightCone;
 
 	float3 DebugOutput = float3(0, 0, 0);
@@ -202,16 +204,17 @@ void VirtualShadowMapProjectionSpotPS(
 		const float ContactShadowLengthWorld = ContactShadowLength * View.ClipToView[1][1] * SceneDepth;
 
 		float3 EstimatedGeoWorldNormal = GetEstimatedGeoWorldNormal(TranslatedWorldPosition, ScreenUV);
-		TranslatedWorldPosition += EstimatedGeoWorldNormal * NormalOffsetWorld;
 
-		if (SMRTRayCount > 0)
+		uint LightType = Light.SpotAngles.x > -2.0f ? LIGHT_TYPE_SPOT : LIGHT_TYPE_POINT;
+
+		if (SMRTRayCount > 0 && LightType == LIGHT_TYPE_SPOT)
 		{
 			RayMissFactor = TraceSpotLight(
 				VirtualShadowMapId,
 				Light,
 				SvPosition,
 				SceneDepth,
-				TranslatedWorldPosition,
+				TranslatedWorldPosition + EstimatedGeoWorldNormal * NormalOffsetWorld,
 				ContactShadowLengthWorld,
 				SMRTRayCount,
 				SMRTSamplesPerRay,
@@ -222,10 +225,9 @@ void VirtualShadowMapProjectionSpotPS(
 		else
 		{
 			RayMissFactor = ProjectSingleSample(
-				Light,
-				ToLightDir,
-				LIGHT_TYPE_SPOT,
-				TranslatedWorldPosition,
+				VirtualShadowMapId,
+				L,
+				TranslatedWorldPosition + EstimatedGeoWorldNormal * NormalOffsetWorld,
 				EstimatedGeoWorldNormal,
 				ContactShadowLengthWorld,
 				/* out */ DebugOutput,
@@ -235,7 +237,7 @@ void VirtualShadowMapProjectionSpotPS(
 		if (ContactShadowLengthWorld > 0.0f)
 		{
 			float StepOffset = InterleavedGradientNoise(SvPosition.xy, View.StateFrameIndexMod8);
-			float ContactShadowHitDistance = VirtualShadowMapScreenRayCast(TranslatedWorldPosition, ToLightDir, ContactShadowLengthWorld, CONTACT_SHADOW_SAMPLES, StepOffset);		
+			float ContactShadowHitDistance = VirtualShadowMapScreenRayCast(TranslatedWorldPosition, L, ContactShadowLengthWorld, CONTACT_SHADOW_SAMPLES, StepOffset);		
 			if (ContactShadowHitDistance >= 0.0f)
 			{
 				RayMissFactor = 0.0f;
@@ -262,7 +264,7 @@ void VirtualShadowMapProjectionDirectionalPS(
 	FLightShaderParameters Light = GetRootLightShaderParameters();
 		
 	FGBufferData GBufferData = GetGBufferData(ScreenUV);
-	float3 ToLightDir = normalize(Light.Direction);
+	float3 L = Light.Direction;
 	
 	float3 DebugOutput = float3(0, 0, 0);
 	int SampleCount = 1;
@@ -277,7 +279,6 @@ void VirtualShadowMapProjectionDirectionalPS(
 		const float ContactShadowLengthWorld = ContactShadowLength * View.ClipToView[1][1] * SceneDepth;
 
 		float3 EstimatedGeoWorldNormal = GetEstimatedGeoWorldNormal(TranslatedWorldPosition, ScreenUV);
-		TranslatedWorldPosition += EstimatedGeoWorldNormal * NormalOffsetWorld;
 
 		if (SMRTRayCount > 0)
 		{
@@ -287,7 +288,7 @@ void VirtualShadowMapProjectionDirectionalPS(
 				Light,
 				SvPosition,
 				SceneDepth,
-				TranslatedWorldPosition,
+				TranslatedWorldPosition + EstimatedGeoWorldNormal * NormalOffsetWorld,
 				ContactShadowLengthWorld,
 				SMRTRayCount,
 				SMRTSamplesPerRay,
@@ -300,10 +301,9 @@ void VirtualShadowMapProjectionDirectionalPS(
 		{
 			SampleCount = 1;
 			RayMissFactor = ProjectSingleSample(
-				Light,
-				Light.Direction,
-				LIGHT_TYPE_DIRECTIONAL,
-				TranslatedWorldPosition,
+				VirtualShadowMapId,
+				L,
+				TranslatedWorldPosition + EstimatedGeoWorldNormal * NormalOffsetWorld,
 				EstimatedGeoWorldNormal,
 				ContactShadowLengthWorld,
 				/* out */ DebugOutput,
@@ -313,7 +313,7 @@ void VirtualShadowMapProjectionDirectionalPS(
 		if (ContactShadowLengthWorld > 0.0f)
 		{
 			float StepOffset = InterleavedGradientNoise(SvPosition.xy, View.StateFrameIndexMod8);
-			float ContactShadowHitDistance = VirtualShadowMapScreenRayCast(TranslatedWorldPosition, ToLightDir, ContactShadowLengthWorld, CONTACT_SHADOW_SAMPLES, StepOffset);		
+			float ContactShadowHitDistance = VirtualShadowMapScreenRayCast(TranslatedWorldPosition, L, ContactShadowLengthWorld, CONTACT_SHADOW_SAMPLES, StepOffset);		
 			if (ContactShadowHitDistance >= 0.0f)
 			{
 				RayMissFactor = 0.0f;
@@ -345,3 +345,155 @@ void VirtualShadowMapCompositePS(
 	
 	OutShadowMask = EncodeLightAttenuationFromMask(Shadow);
 }
+
+FLightShaderParameters ConvertFromLocal( const FLocalLightData LightData )
+{
+	FLightShaderParameters Light = (FLightShaderParameters)0;
+	Light.Position			= LightData.LightPositionAndInvRadius.xyz;
+	Light.InvRadius			= LightData.LightPositionAndInvRadius.w;
+	Light.Color				= LightData.LightColorAndFalloffExponent.xyz;
+	Light.FalloffExponent	= LightData.LightColorAndFalloffExponent.w;
+	Light.Direction			= LightData.LightDirectionAndShadowMask.xyz;
+	Light.Tangent			= LightData.LightTangentAndSoftSourceRadius.xyz;
+	Light.SpotAngles		= LightData.SpotAnglesAndSourceRadiusPacked.xy;
+	Light.SpecularScale		= 1;
+	Light.SourceRadius		= LightData.SpotAnglesAndSourceRadiusPacked.z;
+	Light.SoftSourceRadius	= LightData.LightTangentAndSoftSourceRadius.w;
+	Light.SourceLength		= f16tof32(asuint(LightData.SpotAnglesAndSourceRadiusPacked.w));
+	return Light;
+}
+
+uint EncodeBlurRadius( float Radius )
+{
+	// Does WORLD_RADIUS_MISS really need to be that big?
+	const float Scale = 254 / log2( WORLD_RADIUS_MISS + 0.5 );
+	return max( Scale * log2( Radius + 0.5 ), 254 );
+}
+
+StructuredBuffer< int > VirtualShadowMapIdRemap;
+uint NumDirectionalLightSmInds;
+
+RWTexture2D< uint > RWShadowMaskBits;
+
+[numthreads(8, 8, 1)]
+void VirtualShadowMapProjection(
+	uint3	GroupId				: SV_GroupID,
+	uint	GroupIndex			: SV_GroupIndex,
+	uint3	DispatchThreadId	: SV_DispatchThreadID )
+{
+	// Morton order within a group so page access/atomics are more coherent and wave-swizzled gradients are possible.
+	uint2 LocalPixelPos = 8 * GroupId.xy + MortonDecode(GroupIndex);
+	uint2 PixelPos = LocalPixelPos + uint2( View.ViewRectMin.xy );
+
+	if (any(LocalPixelPos >= uint2(View.ViewSizeAndInvSize.xy)))
+	{
+		return;
+	}
+	
+	const float DeviceZ = SceneTexturesStruct.SceneDepthTexture.Load( int3( PixelPos, 0 ) ).r;
+	const float SceneDepth = ConvertFromDeviceZ( DeviceZ );
+
+	const float4 SvPosition = float4( float2( PixelPos ) + 0.5, DeviceZ, 1.0f );
+	const float2 ScreenUV = SvPositionToBufferUV( SvPosition );
+
+	const float3 TranslatedWorldPosition = SvPositionToTranslatedWorld( SvPosition );
+	const float3 WorldPosition = TranslatedWorldPosition - View.PreViewTranslation;
+
+	const float ContactShadowLengthWorld = ContactShadowLength * View.ClipToView[1][1] * SceneDepth;
+	const float Noise = InterleavedGradientNoise( SvPosition.xy, View.StateFrameIndexMod8 );
+
+	const float3 EstimatedGeoWorldNormal = GetEstimatedGeoWorldNormal(TranslatedWorldPosition, ScreenUV);
+
+	uint ShadowMask = 0xffffffff;
+	uint ShadowIndexes = 0xffffffff;
+	uint ShadowValues = 0;
+	
+	uint EyeIndex = 0; // ??
+	uint GridLinearIndex = ComputeLightGridCellIndex( LocalPixelPos, SceneDepth, EyeIndex );
+	const FCulledLightsGridData CulledLightGridData = GetCulledLightsGrid( GridLinearIndex, EyeIndex );
+
+	LOOP
+	for (uint Index = 0; Index < CulledLightGridData.NumLocalLights; ++Index)
+	{
+		const FLocalLightData LightData = GetLocalLightData(CulledLightGridData.DataStartIndex + Index, EyeIndex);
+		const uint LightGridLightIndex = ForwardLightData.CulledLightDataGrid[CulledLightGridData.DataStartIndex + Index];
+
+		// The Virtual Shadow Remap stores directional lights first
+		int VirtualShadowMapId = VirtualShadowMapIdRemap[NumDirectionalLightSmInds + LightGridLightIndex];
+		if (VirtualShadowMapId != INDEX_NONE)
+		{
+			float3 DebugOutput = float3(0, 0, 0);
+			int SampleCount = 1;
+			float HitDistance = WORLD_RADIUS_MISS;
+			float RayMissFactor = 0.0f;
+
+			float3 ToLight = LightData.LightPositionAndInvRadius.xyz - WorldPosition;
+			float d2 = dot( ToLight, ToLight );
+			float InvDist = rsqrt( d2 );
+			float3 L = ToLight * InvDist;
+			
+			bool bInLightRadius = InvDist >= LightData.LightPositionAndInvRadius.w;
+			bool bInLightCone = dot( L, LightData.LightDirectionAndShadowMask.xyz ) >= LightData.SpotAnglesAndSourceRadiusPacked.x;
+
+			if (bInLightRadius && bInLightCone)
+			{
+				FLightShaderParameters Light = ConvertFromLocal( LightData );
+
+				uint LightType = Light.SpotAngles.x > -2.0f ? LIGHT_TYPE_SPOT : LIGHT_TYPE_POINT;
+
+				if (SMRTRayCount > 0 && LightType == LIGHT_TYPE_SPOT)
+				{
+					RayMissFactor = TraceSpotLight(
+						VirtualShadowMapId,
+						Light,
+						SvPosition,
+						SceneDepth,
+						TranslatedWorldPosition + EstimatedGeoWorldNormal * NormalOffsetWorld,
+						ContactShadowLengthWorld,
+						SMRTRayCount,
+						SMRTSamplesPerRay,
+						DebugOutputType,
+						DebugOutput,
+						HitDistance);
+				}
+				else
+				{
+					RayMissFactor = ProjectSingleSample(
+						VirtualShadowMapId,
+						L,
+						TranslatedWorldPosition + EstimatedGeoWorldNormal * NormalOffsetWorld,
+						EstimatedGeoWorldNormal,
+						ContactShadowLengthWorld,
+						/* out */ DebugOutput,
+						HitDistance);
+				}
+
+				if (ContactShadowLengthWorld > 0.0f)
+				{
+					float ContactShadowHitDistance = VirtualShadowMapScreenRayCast(TranslatedWorldPosition, L, ContactShadowLengthWorld, CONTACT_SHADOW_SAMPLES, Noise);		
+					if (ContactShadowHitDistance >= 0.0f)
+					{
+						RayMissFactor = 0.0f;
+						HitDistance = ContactShadowHitDistance;
+					}
+				}
+
+				//ShadowMask ^= uint( RayMissFactor < 1.0 ) << Index;
+				ShadowMask ^= ( ~uint( round( RayMissFactor * 7.0 ) ) & 7 ) << (Index*3);
+
+				if( Index < 4 && LightGridLightIndex < 255 )
+				{
+					float WorldBluringRadius = ComputeLightSampleWorldBluringRadius(TranslatedWorldPosition, LightType, Light, HitDistance);
+					uint EncodedValue = RayMissFactor < 1.0 ? EncodeBlurRadius( WorldBluringRadius ) : 255;
+
+					ShadowIndexes = ( ShadowIndexes << 8 ) | LightGridLightIndex;
+					ShadowValues = ( ShadowValues << 8 ) | EncodedValue;
+				}
+			}
+
+			//OutputTarget = EncodeOutput(TranslatedWorldPosition, Light, LIGHT_TYPE_SPOT, RayMissFactor, HitDistance, SampleCount, DebugOutput);
+		}
+	}
+
+	RWShadowMaskBits[ PixelPos ] = ~ShadowMask;
+}
\ No newline at end of file
diff --git a/Engine/Source/Runtime/Renderer/Private/ClusteredDeferredShadingPass.cpp b/Engine/Source/Runtime/Renderer/Private/ClusteredDeferredShadingPass.cpp
--- a/Engine/Source/Runtime/Renderer/Private/ClusteredDeferredShadingPass.cpp
+++ b/Engine/Source/Runtime/Renderer/Private/ClusteredDeferredShadingPass.cpp
@@ -70,6 +70,7 @@ class FClusteredShadingPS : public FGlobalShader
 		SHADER_PARAMETER_STRUCT_REF(FForwardLightData, Forward)
 		SHADER_PARAMETER_STRUCT_REF(FViewUniformShaderParameters, View)
 		SHADER_PARAMETER_RDG_UNIFORM_BUFFER(FSceneTextureUniformParameters, SceneTextures)
+		SHADER_PARAMETER_RDG_TEXTURE(Texture2D, ShadowMaskBits)
 
 		SHADER_PARAMETER_TEXTURE(Texture2D, LTCMatTexture)
 		SHADER_PARAMETER_SAMPLER(SamplerState, LTCMatSampler)
@@ -101,7 +102,8 @@ IMPLEMENT_GLOBAL_SHADER(FClusteredShadingPS, "/Engine/Private/ClusteredDeferredS
 void FDeferredShadingSceneRenderer::AddClusteredDeferredShadingPass(
 	FRDGBuilder& GraphBuilder,
 	const FMinimalSceneTextures& SceneTextures,
-	const FSortedLightSetSceneInfo &SortedLightsSet)
+	const FSortedLightSetSceneInfo &SortedLightsSet,
+	FRDGTextureRef ShadowMaskBits)
 {
 	check(GUseClusteredDeferredShading);
 
@@ -123,6 +125,7 @@ void FDeferredShadingSceneRenderer::AddClusteredDeferredShadingPass(
 			PassParameters->View = View.ViewUniformBuffer;
 			PassParameters->Forward = View.ForwardLightingResources->ForwardLightDataUniformBuffer;
 			PassParameters->SceneTextures = SceneTextures.UniformBuffer;
+			PassParameters->ShadowMaskBits = ShadowMaskBits;
 
 			PassParameters->LTCMatTexture = GSystemTextures.LTCMat->GetShaderResourceRHI();
 			PassParameters->LTCMatSampler = TStaticSamplerState<SF_Bilinear, AM_Clamp, AM_Clamp, AM_Clamp>::GetRHI();
diff --git a/Engine/Source/Runtime/Renderer/Private/DeferredShadingRenderer.h b/Engine/Source/Runtime/Renderer/Private/DeferredShadingRenderer.h
--- a/Engine/Source/Runtime/Renderer/Private/DeferredShadingRenderer.h
+++ b/Engine/Source/Runtime/Renderer/Private/DeferredShadingRenderer.h
@@ -548,7 +548,8 @@ private:
 	void AddClusteredDeferredShadingPass(
 		FRDGBuilder& GraphBuilder,
 		const FMinimalSceneTextures& SceneTextures,
-		const FSortedLightSetSceneInfo& SortedLightsSet);
+		const FSortedLightSetSceneInfo& SortedLightsSet,
+		FRDGTextureRef ShadowMaskBits);
 
 	/** Renders the lights in SortedLights in the range [TiledDeferredLightsStart, TiledDeferredLightsEnd) using tiled deferred shading. */
 	void RenderTiledDeferredLighting(
diff --git a/Engine/Source/Runtime/Renderer/Private/LightGridInjection.cpp b/Engine/Source/Runtime/Renderer/Private/LightGridInjection.cpp
--- a/Engine/Source/Runtime/Renderer/Private/LightGridInjection.cpp
+++ b/Engine/Source/Runtime/Renderer/Private/LightGridInjection.cpp
@@ -73,6 +73,8 @@ FAutoConsoleVariableRef CVarLightCullingQuality(
 	ECVF_RenderThreadSafe
 );
 
+extern TAutoConsoleVariable<int32> CVarVirtualShadowOnePassProjection;
+
 /** A minimal forwarding lighting setup. */
 class FMinimalDummyForwardLightingResources : public FRenderResource
 {
@@ -770,7 +772,9 @@ void FSceneRenderer::ComputeLightGrid(FRDGBuilder& GraphBuilder, bool bCullLight
 
 void FDeferredShadingSceneRenderer::GatherLightsAndComputeLightGrid(FRDGBuilder& GraphBuilder, bool bNeedLightGrid, FSortedLightSetSceneInfo &SortedLightSet)
 {
-	GatherAndSortLights(SortedLightSet);
+	bool bShadowedLightsInClustered = ShouldUseClusteredDeferredShading() && CVarVirtualShadowOnePassProjection.GetValueOnRenderThread();
+
+	GatherAndSortLights(SortedLightSet, bShadowedLightsInClustered);
 	
 	if (!bNeedLightGrid)
 	{
diff --git a/Engine/Source/Runtime/Renderer/Private/LightRendering.cpp b/Engine/Source/Runtime/Renderer/Private/LightRendering.cpp
--- a/Engine/Source/Runtime/Renderer/Private/LightRendering.cpp
+++ b/Engine/Source/Runtime/Renderer/Private/LightRendering.cpp
@@ -21,6 +21,7 @@
 #include "SkyAtmosphereRendering.h"
 #include "VolumetricCloudRendering.h"
 #include "Strata/Strata.h"
+#include "VirtualShadowMaps/VirtualShadowMapProjection.h"
 
 // ENABLE_DEBUG_DISCARD_PROP is used to test the lighting code by allowing to discard lights to see how performance scales
 // It ought never to be enabled in a shipping build, and is probably only really useful when woring on the shading code.
@@ -42,6 +43,7 @@ IMPLEMENT_GLOBAL_SHADER_PARAMETER_STRUCT(FDeferredLightUniformStruct, "DeferredL
 extern int32 GUseTranslucentLightingVolumes;
 ENGINE_API IPooledRenderTarget* GetSubsufaceProfileTexture_RT(FRHICommandListImmediate& RHICmdList);
 
+extern TAutoConsoleVariable<int32> CVarVirtualShadowOnePassProjection;
 
 static int32 GAllowDepthBoundsTest = 1;
 static FAutoConsoleVariableRef CVarAllowDepthBoundsTest(
@@ -989,7 +991,7 @@ static bool LightRequiresDenosier(const FLightSceneInfo& LightSceneInfo)
 
 
 
-void FSceneRenderer::GatherAndSortLights(FSortedLightSetSceneInfo& OutSortedLights)
+void FSceneRenderer::GatherAndSortLights(FSortedLightSetSceneInfo& OutSortedLights, bool bShadowedLightsInClustered)
 {
 	if (bAllowSimpleLights)
 	{
@@ -1055,7 +1057,7 @@ void FSceneRenderer::GatherAndSortLights(FSortedLightSetSceneInfo& OutSortedLigh
 					// Rect lights are not supported as the performance impact is significant even if not used, for now, left for trad. deferred.
 					const bool bTiledOrClusteredDeferredSupported =
 						!SortedLightInfo->SortKey.Fields.bTextureProfile &&
-						!SortedLightInfo->SortKey.Fields.bShadowed &&
+						(!SortedLightInfo->SortKey.Fields.bShadowed || bShadowedLightsInClustered) &&
 						!SortedLightInfo->SortKey.Fields.bLightFunction &&
 						!SortedLightInfo->SortKey.Fields.bUsesLightingChannels
 						&& LightSceneInfoCompact.LightType != LightType_Directional
@@ -1141,7 +1143,7 @@ void FSceneRenderer::GatherAndSortLights(FSortedLightSetSceneInfo& OutSortedLigh
 			OutSortedLights.ClusteredSupportedEnd = LightIndex;
 		}
 
-		if (bDrawShadows || bDrawLightFunction || bLightingChannels)
+		if( (bDrawShadows || bDrawLightFunction || bLightingChannels) && SortedLightInfo.SortKey.Fields.bClusteredDeferredNotSupported )
 		{
 			// Once we find a shadowed light, we can exit the loop, these lights should never support tiled deferred rendering either
 			check(SortedLightInfo.SortKey.Fields.bTiledDeferredNotSupported);
@@ -1268,13 +1270,39 @@ void FDeferredShadingSceneRenderer::RenderLights(
 			// True if the clustered shading is enabled and the feature level is there, and that the light grid had lights injected.
 			if (ShouldUseClusteredDeferredShading() && AreLightsInLightGrid())
 			{
+				FRDGTextureRef ShadowMaskBits = nullptr;
+				if( VirtualShadowMapArray.IsAllocated() && CVarVirtualShadowOnePassProjection.GetValueOnRenderThread() )
+				{
+					const FRDGTextureDesc ShadowMaskDesc = FRDGTextureDesc::Create2D(
+						SceneTextures.Config.Extent,
+						PF_R32_UINT,
+						FClearValueBinding::None,
+						TexCreate_ShaderResource | TexCreate_UAV );
+
+					ShadowMaskBits = GraphBuilder.CreateTexture( ShadowMaskDesc, TEXT("ShadowMaskBits") );
+
+					for (int32 ViewIndex = 0, Num = Views.Num(); ViewIndex < Num; ViewIndex++)
+					{
+						RenderVirtualShadowMapProjection(
+							GraphBuilder,
+							SceneTextures,
+							Views[ ViewIndex ],
+							VirtualShadowMapArray,
+							ShadowMaskBits );
+					}
+				}
+				else
+				{
+					ShadowMaskBits = GraphBuilder.RegisterExternalTexture( GSystemTextures.ZeroUIntDummy );
+				}
+
 				// Tell the trad. deferred that the clustered deferred capable lights are taken care of.
 				// This includes the simple lights
 				StandardDeferredStart = SortedLightSet.ClusteredSupportedEnd;
 				// Tell the trad. deferred that the simple lights are spoken for.
 				bRenderSimpleLightsStandardDeferred = false;
 
-				AddClusteredDeferredShadingPass(GraphBuilder, SceneTextures, SortedLightSet);
+				AddClusteredDeferredShadingPass(GraphBuilder, SceneTextures, SortedLightSet, ShadowMaskBits);
 			}
 			else if (CanUseTiledDeferred())
 			{
diff --git a/Engine/Source/Runtime/Renderer/Private/SceneRendering.h b/Engine/Source/Runtime/Renderer/Private/SceneRendering.h
--- a/Engine/Source/Runtime/Renderer/Private/SceneRendering.h
+++ b/Engine/Source/Runtime/Renderer/Private/SceneRendering.h
@@ -2045,7 +2045,7 @@ protected:
 	 * Rounds up lights and sorts them according to what type of renderer supports them. The result is stored in OutSortedLights 
 	 * NOTE: Also extracts the SimpleLights AND adds them to the sorted range (first sub-range). 
 	 */
-	void GatherAndSortLights(FSortedLightSetSceneInfo& OutSortedLights);
+	void GatherAndSortLights(FSortedLightSetSceneInfo& OutSortedLights, bool bShadowedLightsInClustered = false);
 	
 	/** 
 	 * Culls local lights and reflection probes to a grid in frustum space, builds one light list and grid per view in the current Views.  
diff --git a/Engine/Source/Runtime/Renderer/Private/VirtualShadowMaps/VirtualShadowMapArray.cpp b/Engine/Source/Runtime/Renderer/Private/VirtualShadowMaps/VirtualShadowMapArray.cpp
--- a/Engine/Source/Runtime/Renderer/Private/VirtualShadowMaps/VirtualShadowMapArray.cpp
+++ b/Engine/Source/Runtime/Renderer/Private/VirtualShadowMaps/VirtualShadowMapArray.cpp
@@ -772,6 +772,7 @@ void FVirtualShadowMapArray::BuildPageAllocations(
 			}
 		}
 		ShadowMapProjectionDataRDG = CreateStructuredBuffer(GraphBuilder, TEXT("ShadowMapProjectionData"), ShadowMapProjectionData);
+		NumDirectionalLights = DirectionalLightSmInds.Num();
 
 		for (int32 ViewIndex = 0; ViewIndex < Views.Num(); ++ViewIndex)
 		{
@@ -811,7 +812,7 @@ void FVirtualShadowMapArray::BuildPageAllocations(
 			}
 
 			const FRDGSystemTextures& SystemTextures = FRDGSystemTextures::Get(GraphBuilder);
-			FRDGBufferRef VirtualShadowMapIdRemapRDG = CreateStructuredBuffer(GraphBuilder, TEXT("VirtualShadowMapIdRemap"), VirtualShadowMapIdRemap);
+			VirtualShadowMapIdRemapRDG.Add( CreateStructuredBuffer(GraphBuilder, TEXT("VirtualShadowMapIdRemap"), VirtualShadowMapIdRemap) );
 			FRDGTextureRef VisBuffer64 = NaniteVisBuffer64 != nullptr ? NaniteVisBuffer64 : SystemTextures.Black;
 
 			FRDGBufferRef ScreenSpaceGridBoundsRDG = nullptr;
@@ -831,7 +832,7 @@ void FVirtualShadowMapArray::BuildPageAllocations(
 				PassParameters->View = View.ViewUniformBuffer;
 				PassParameters->OutPageRequestFlags = GraphBuilder.CreateUAV(PageRequestFlagsRDG);
 				PassParameters->ForwardLightData = View.ForwardLightingResources->ForwardLightDataUniformBuffer;
-				PassParameters->VirtualShadowMapIdRemap = GraphBuilder.CreateSRV(VirtualShadowMapIdRemapRDG);
+				PassParameters->VirtualShadowMapIdRemap = GraphBuilder.CreateSRV(VirtualShadowMapIdRemapRDG[ViewIndex]);
 				PassParameters->ShadowMapProjectionData = GraphBuilder.CreateSRV(ShadowMapProjectionDataRDG);
 				PassParameters->NumDirectionalLightSmInds = DirectionalLightSmInds.Num();
 				PassParameters->LodFootprintScale = LodFootprintScale;
@@ -1023,7 +1024,7 @@ void FVirtualShadowMapArray::SetProjectionParameters(FRDGBuilder& GraphBuilder,
 	OutParameters.PhysicalPagePool = PhysicalPagePoolRDG != nullptr ? PhysicalPagePoolRDG : GraphBuilder.RegisterExternalTexture(GSystemTextures.BlackDummy);
 	OutParameters.PhysicalPagePoolHw = PhysicalPagePoolHw != nullptr ? PhysicalPagePoolHw : GraphBuilder.RegisterExternalTexture(GSystemTextures.BlackDummy);
 #else //!ENABLE_NON_NANITE_VSM
-	OutParameters.PhysicalPagePool = PhysicalPagePoolRDG;
+	OutParameters.PhysicalPagePool = PhysicalPagePoolRDG != nullptr ? PhysicalPagePoolRDG : GraphBuilder.RegisterExternalTexture(GSystemTextures.BlackDummy);
 #endif // ENABLE_NON_NANITE_VSM
 	OutParameters.VirtualShadowMapProjectionData = GraphBuilder.CreateSRV(ShadowMapProjectionDataRDG);
 }
diff --git a/Engine/Source/Runtime/Renderer/Private/VirtualShadowMaps/VirtualShadowMapArray.h b/Engine/Source/Runtime/Renderer/Private/VirtualShadowMaps/VirtualShadowMapArray.h
--- a/Engine/Source/Runtime/Renderer/Private/VirtualShadowMaps/VirtualShadowMapArray.h
+++ b/Engine/Source/Runtime/Renderer/Private/VirtualShadowMaps/VirtualShadowMapArray.h
@@ -271,6 +271,9 @@ public:
 	FRDGBufferRef AllocatedPageRectBoundsRDG = nullptr;
 #endif // ENABLE_NON_NANITE_VSM
 	FRDGBufferRef ShadowMapProjectionDataRDG = nullptr;
+	TArray< FRDGBufferRef, SceneRenderingAllocator > VirtualShadowMapIdRemapRDG;
+
+	int32 NumDirectionalLights = 0;
 
 	// These are not created or managed by this class - references are copied from the cache manager
 	// when needed as they come from the previous frame.
diff --git a/Engine/Source/Runtime/Renderer/Private/VirtualShadowMaps/VirtualShadowMapProjection.cpp b/Engine/Source/Runtime/Renderer/Private/VirtualShadowMaps/VirtualShadowMapProjection.cpp
--- a/Engine/Source/Runtime/Renderer/Private/VirtualShadowMaps/VirtualShadowMapProjection.cpp
+++ b/Engine/Source/Runtime/Renderer/Private/VirtualShadowMaps/VirtualShadowMapProjection.cpp
@@ -46,6 +46,13 @@ static TAutoConsoleVariable<int32> CVarVirtualShadowMapDebugProjection(
 	ECVF_RenderThreadSafe
 );
 
+TAutoConsoleVariable<int32> CVarVirtualShadowOnePassProjection(
+	TEXT("r.Shadow.v.OnePassProjection"),
+	0,
+	TEXT("Single pass projects all local VSMs culled with the light grid. Used in conjunction with clustered deferred shading."),
+	ECVF_RenderThreadSafe
+);
+
 static TAutoConsoleVariable<int32> CVarSMRTRayCountDirectional(
 	TEXT( "r.Shadow.v.SMRT.RayCountDirectional" ),
 	0,
@@ -429,3 +436,76 @@ void CompositeVirtualShadowMapMask(
 		ScissorRect,
 		BlendState);
 }
+
+class FVirtualShadowMapProjectionCS : public FGlobalShader
+{
+	DECLARE_GLOBAL_SHADER(FVirtualShadowMapProjectionCS);
+	SHADER_USE_PARAMETER_STRUCT(FVirtualShadowMapProjectionCS, FGlobalShader)
+	
+	BEGIN_SHADER_PARAMETER_STRUCT(FParameters, )
+		SHADER_PARAMETER_STRUCT_INCLUDE(FVirtualShadowMapSamplingParameters, ProjectionParameters)
+		SHADER_PARAMETER_RDG_UNIFORM_BUFFER(FSceneTextureUniformParameters, SceneTexturesStruct)
+		SHADER_PARAMETER_STRUCT_REF(FViewUniformShaderParameters, View)
+		SHADER_PARAMETER_STRUCT_REF(FForwardLightData, ForwardLightData)
+		SHADER_PARAMETER_RDG_BUFFER_SRV(StructuredBuffer< FVirtualShadowMapProjectionShaderData >, ShadowMapProjectionData)
+		SHADER_PARAMETER_RDG_BUFFER_SRV(StructuredBuffer< uint >, VirtualShadowMapIdRemap)
+		SHADER_PARAMETER_RDG_TEXTURE_UAV(RWTexture2D< uint >, RWShadowMaskBits)
+		SHADER_PARAMETER(uint32, NumDirectionalLightSmInds)
+		SHADER_PARAMETER(float, ContactShadowLength)
+		SHADER_PARAMETER(float, NormalOffsetWorld)
+		SHADER_PARAMETER(int32, DebugOutputType)
+		SHADER_PARAMETER(int32, SMRTRayCount)
+		SHADER_PARAMETER(int32, SMRTSamplesPerRay)
+		SHADER_PARAMETER(float, SMRTRayLengthScale)
+	END_SHADER_PARAMETER_STRUCT()
+
+	static void ModifyCompilationEnvironment( const FGlobalShaderPermutationParameters& Parameters, FShaderCompilerEnvironment& OutEnvironment )
+	{
+		FGlobalShader::ModifyCompilationEnvironment(Parameters, OutEnvironment);
+
+		FVirtualShadowMapArray::SetShaderDefines(OutEnvironment);
+		FForwardLightingParameters::ModifyCompilationEnvironment(Parameters.Platform, OutEnvironment);
+
+		OutEnvironment.CompilerFlags.Add(CFLAG_WaveOperations);
+	}
+};
+IMPLEMENT_GLOBAL_SHADER(FVirtualShadowMapProjectionCS, "/Engine/Private/VirtualShadowMaps/VirtualShadowMapProjection.usf", "VirtualShadowMapProjection", SF_Compute);
+
+void RenderVirtualShadowMapProjection(
+	FRDGBuilder& GraphBuilder,
+	const FMinimalSceneTextures& SceneTextures,
+	const FViewInfo& View,
+	FVirtualShadowMapArray& VirtualShadowMapArray,
+	FRDGTextureRef ShadowMaskBits )
+{
+	FVirtualShadowMapProjectionCS::FParameters* PassParameters = GraphBuilder.AllocParameters< FVirtualShadowMapProjectionCS::FParameters >();
+	VirtualShadowMapArray.SetProjectionParameters( GraphBuilder, PassParameters->ProjectionParameters );
+
+	PassParameters->SceneTexturesStruct = SceneTextures.UniformBuffer;
+	PassParameters->View = View.ViewUniformBuffer;
+	PassParameters->ForwardLightData = View.ForwardLightingResources->ForwardLightDataUniformBuffer;
+	PassParameters->VirtualShadowMapIdRemap = GraphBuilder.CreateSRV( VirtualShadowMapArray.VirtualShadowMapIdRemapRDG[0] );	// FIXME Index proper view
+	PassParameters->NumDirectionalLightSmInds = VirtualShadowMapArray.NumDirectionalLights;
+
+	PassParameters->DebugOutputType = CVarVirtualShadowMapDebugProjection.GetValueOnRenderThread();
+	PassParameters->ContactShadowLength = CVarContactShadowLength.GetValueOnRenderThread();
+	PassParameters->NormalOffsetWorld = CVarNormalOffsetWorld.GetValueOnRenderThread();
+
+	PassParameters->SMRTRayCount = CVarSMRTRayCountSpot.GetValueOnRenderThread();
+	PassParameters->SMRTSamplesPerRay = CVarSMRTSamplesPerRaySpot.GetValueOnRenderThread();
+	PassParameters->SMRTRayLengthScale = 0.0f;		// Currently unused in this path
+
+	PassParameters->RWShadowMaskBits = GraphBuilder.CreateUAV( ShadowMaskBits );
+				
+	auto ComputeShader = View.ShaderMap->GetShader< FVirtualShadowMapProjectionCS >();
+
+	const FIntPoint GroupCount = FIntPoint::DivideAndRoundUp( View.ViewRect.Size(), 8 );
+
+	FComputeShaderUtils::AddPass(
+		GraphBuilder,
+		RDG_EVENT_NAME("VirtualShadowMapProjection"),
+		ComputeShader,
+		PassParameters,
+		FIntVector( GroupCount.X, GroupCount.Y, 1 )
+	);
+}
\ No newline at end of file
diff --git a/Engine/Source/Runtime/Renderer/Private/VirtualShadowMaps/VirtualShadowMapProjection.h b/Engine/Source/Runtime/Renderer/Private/VirtualShadowMaps/VirtualShadowMapProjection.h
--- a/Engine/Source/Runtime/Renderer/Private/VirtualShadowMaps/VirtualShadowMapProjection.h
+++ b/Engine/Source/Runtime/Renderer/Private/VirtualShadowMaps/VirtualShadowMapProjection.h
@@ -49,6 +49,13 @@ void RenderVirtualShadowMapProjection(
 	FRDGTextureRef ScreenShadowMaskTexture,
 	bool bProjectingForForwardShading);
 
+void RenderVirtualShadowMapProjection(
+	FRDGBuilder& GraphBuilder,
+	const FMinimalSceneTextures& SceneTextures,
+	const FViewInfo& View,
+	FVirtualShadowMapArray& VirtualShadowMapArray,
+	FRDGTextureRef ShadowMaskBits );
+
 void CompositeVirtualShadowMapMask(
 	FRDGBuilder& GraphBuilder,
 	const FIntRect ScissorRect,
