Approximate caustics mode in the path tracer

Reduce noise from low roughness glass and metals by approximating caustic paths. First we clamp the material roughness to the max seen along a particular path. Second, we use transparent shadows for light to pass through glass more easily.

A CVar controls this effect so that it is still possible to generate ground truth images if needed.

diff --git a/Engine/Shaders/Private/PathTracing/Material/PathTracingMaterialSampling.ush b/Engine/Shaders/Private/PathTracing/Material/PathTracingMaterialSampling.ush
--- a/Engine/Shaders/Private/PathTracing/Material/PathTracingMaterialSampling.ush
+++ b/Engine/Shaders/Private/PathTracing/Material/PathTracingMaterialSampling.ush
@@ -131,6 +131,34 @@ float3 EstimateMaterialAlbedo(FMaterialClosestHitPayload Payload)
 	return Albedo;
 }
 
+void AdjustPathAndMaterialRoughness(inout FMaterialClosestHitPayload Payload, inout float PathRoughness)
+{
+	// boost material roughness according to what we've seen so far
+	float CoatRoughness = Payload.Roughness;
+	switch (Payload.ShadingModelID)
+	{
+		case SHADINGMODELID_UNLIT:
+			// this object does not scatter light
+			return;
+		case SHADINGMODELID_CLEAR_COAT:
+			// clamp coat roughness
+			CoatRoughness = Payload.CustomData.y = max(GetClearCoatRoughness(Payload), PathRoughness);
+			// fallthrough
+		case SHADINGMODELID_TWOSIDED_FOLIAGE:
+		case SHADINGMODELID_DEFAULT_LIT:
+			Payload.Roughness = max(Payload.Roughness, PathRoughness);
+			break;
+		default:
+			// fallback is diffuse
+			Payload.Roughness = 1.0;
+			break;
+	}
+	// now set the path roughness going forward
+	const float GlassPercent = (Payload.BlendingMode == RAY_TRACING_BLEND_MODE_TRANSLUCENT && Payload.Ior != 1.0) ? 1 - Payload.Opacity : 0.0;
+	PathRoughness = lerp(1.0 /* diffuse lobe */, max(CoatRoughness, Payload.Roughness), Payload.Metallic);
+	PathRoughness = lerp(PathRoughness, Payload.Roughness, GlassPercent);
+}
+
 void SampleMaterial(
 	float3 RayDirection,
 	FMaterialClosestHitPayload Payload,
diff --git a/Engine/Shaders/Private/PathTracing/PathTracing.usf b/Engine/Shaders/Private/PathTracing/PathTracing.usf
--- a/Engine/Shaders/Private/PathTracing/PathTracing.usf
+++ b/Engine/Shaders/Private/PathTracing/PathTracing.usf
@@ -85,6 +85,13 @@ FMaterialClosestHitPayload TraceTransparentRay(RayDesc Ray, FRayCone RayCone, bo
 			HitPayload.ShadingModelID = SHADINGMODELID_DEFAULT_LIT;
 		}
 
+		if (HitPayload.BlendingMode == RAY_TRACING_BLEND_MODE_ADDITIVE ||
+			HitPayload.BlendingMode == RAY_TRACING_BLEND_MODE_TRANSLUCENT)
+		{
+			// emission needs to be pre-multiplied by the opacity
+			HitPayload.Radiance *= HitPayload.Opacity;
+		}
+
 		// add in surface emission (except for modulate blend mode which uses Radiance as the transparency multiplier)
 		if (HitPayload.BlendingMode != RAY_TRACING_BLEND_MODE_MODULATE)
 		{
@@ -152,7 +159,7 @@ FMaterialClosestHitPayload TraceTransparentRay(RayDesc Ray, FRayCone RayCone, bo
 	return Payload;
 }
 
-float3 TraceTransparentVisibilityRay(RayDesc Ray, uint2 PixelCoord)
+float3 TraceTransparentVisibilityRay(RayDesc Ray, uint2 PixelCoord, float PathRoughness)
 {
 	const uint RayFlags = RAY_FLAG_ACCEPT_FIRST_HIT_AND_END_SEARCH | RAY_FLAG_SKIP_CLOSEST_HIT_SHADER;
 	const uint InstanceInclusionMask = RAY_TRACING_MASK_SHADOW;
@@ -162,13 +169,14 @@ float3 TraceTransparentVisibilityRay(RayDesc Ray, uint2 PixelCoord)
 
 	FPackedMaterialClosestHitPayload PackedPayload = (FPackedMaterialClosestHitPayload)0;
 
-	// By enabling minimal payload mode all other payload information is ignored, meaning these functions need no payload inputs
-	PackedPayload.SetMinimalPayloadMode();
-	PackedPayload.SetEvalOpacity();
-	PackedPayload.HitT = 0;
+	// Signal to the Any hit shader we want to evaluate the opacity
+	// The payload is used to carry the path throughput (for transparent shadows)
+	// and current path roughness (for approximate caustics)
+	PackedPayload.SetPathTracingVisibilityRay();
 	PackedPayload.RadianceAndNormal[0] = asuint(1.0);
 	PackedPayload.RadianceAndNormal[1] = asuint(1.0);
 	PackedPayload.RadianceAndNormal[2] = asuint(1.0);
+	PackedPayload.MetallicAndSpecularAndRoughness = asuint(PathRoughness);
 
 	PackedPayload.SetPixelCoord(PixelCoord);
 
@@ -239,6 +247,7 @@ RAY_TRACING_ENTRY_RAYGEN(PathTracingMainRG)
 
 	// path state variables (these cary information between bounces)
 	float3 PathThroughput = 1.0;
+	float PathRoughness = 0;
 
 	// 0: only Material sampling
 	// 1: only Light sampling
@@ -271,6 +280,16 @@ RAY_TRACING_ENTRY_RAYGEN(PathTracingMainRG)
 			break;
 		}
 
+		if (PathTracingData.ApproximateCaustics)
+		{
+			// modify the path roughness as well as the payload roughness to minimize difficult caustics
+			// This is inspired by a trick used in the Arnold renderer:
+			//   https://cgg.mff.cuni.cz/~jaroslav/gicourse2010/giai2010-02-marcos_fajardo-slides.pdf (slide 39)
+			//   https://www.arnoldrenderer.com/research/Arnold_TOG2018.pdf (section 4.2)
+			AdjustPathAndMaterialRoughness(Payload, PathRoughness);
+		}
+		
+
 		// Update origin
 		Ray.Origin = Payload.WorldPos;
 
@@ -329,7 +348,7 @@ RAY_TRACING_ENTRY_RAYGEN(PathTracingMainRG)
 
 					if (CastsShadow(LightId))
 					{
-						RadianceOverPdf *= TraceTransparentVisibilityRay(LightRay, LaunchIndex);
+						RadianceOverPdf *= TraceTransparentVisibilityRay(LightRay, LaunchIndex, PathRoughness);
 					}
 
 					// #dxr_todo: Is it cheaper to fire the ray first? Or eval the material first?
@@ -446,7 +465,7 @@ RAY_TRACING_ENTRY_RAYGEN(PathTracingMainRG)
 					{
 						RayDesc LightRay = Ray;
 						LightRay.TMax = HitT;
-						LightContrib *= TraceTransparentVisibilityRay(LightRay, LaunchIndex);
+						LightContrib *= TraceTransparentVisibilityRay(LightRay, LaunchIndex, PathRoughness);
 					}
 					// the light made some contribution, and there was nothing along the shadow ray
 					AccumulateRadiance(Radiance, LightContrib);
diff --git a/Engine/Shaders/Private/RayTracing/RayTracingCommon.ush b/Engine/Shaders/Private/RayTracing/RayTracingCommon.ush
--- a/Engine/Shaders/Private/RayTracing/RayTracingCommon.ush
+++ b/Engine/Shaders/Private/RayTracing/RayTracingCommon.ush
@@ -225,7 +225,7 @@ uint2 GetPixelCoord(uint2 DispatchThreadId, uint UpscaleFactor)
 #define RAY_TRACING_PAYLOAD_FLAG_ENABLE_SKY_LIGHT_CONTRIBUTION (1 << 2) // Indicates that hit shaders should add in the Sky Light contribution in their output. Set by RayGen shader, before TraceRay().
 #define RAY_TRACING_PAYLOAD_FLAG_LUMEN_PAYLOAD				   (1 << 3) // Indicates that closest hit shader should evaluate geometry normal and instanceid.
 #define RAY_TRACING_PAYLOAD_FLAG_IGNORE_TRANSLUCENT            (1 << 4) // Indicates that translucent materials should be ignored by IgnoreHit() in anyhit shaders.
-#define RAY_TRACING_PAYLOAD_FLAG_EVAL_OPACITY                  (1 << 5) // Indicates that translucent materials should be evaluated and the overall transparency accumulated into the payload
+#define RAY_TRACING_PAYLOAD_FLAG_PATHTRACER_VISIBILITY         (1 << 5) // Indicates that this is a path tracer visibility ray (which supports material evaluation for transparent shadows and fake caustics)
 
 struct FMaterialClosestHitPayload : FMinimalPayload
 {
@@ -365,8 +365,8 @@ struct FPackedMaterialClosestHitPayload : FMinimalPayload
 	void SetIgnoreTranslucentMaterials() { IorAndShadingModelIDAndBlendingModeAndFlagsAndPrimitiveLightingChannelMask |= RAY_TRACING_PAYLOAD_FLAG_IGNORE_TRANSLUCENT << 23; }
 	bool IsIgnoreTranslucentMaterials() { return ((IorAndShadingModelIDAndBlendingModeAndFlagsAndPrimitiveLightingChannelMask >> 23) & RAY_TRACING_PAYLOAD_FLAG_IGNORE_TRANSLUCENT) != 0; }
 
-	void SetEvalOpacity() { IorAndShadingModelIDAndBlendingModeAndFlagsAndPrimitiveLightingChannelMask |= RAY_TRACING_PAYLOAD_FLAG_EVAL_OPACITY << 23; }
-	bool IsEvalOpacity() { return ((IorAndShadingModelIDAndBlendingModeAndFlagsAndPrimitiveLightingChannelMask >> 23) & RAY_TRACING_PAYLOAD_FLAG_EVAL_OPACITY) != 0; }
+	void SetPathTracingVisibilityRay() { IorAndShadingModelIDAndBlendingModeAndFlagsAndPrimitiveLightingChannelMask |= RAY_TRACING_PAYLOAD_FLAG_PATHTRACER_VISIBILITY << 23; }
+	bool IsPathTracingVisibilityRay() { return ((IorAndShadingModelIDAndBlendingModeAndFlagsAndPrimitiveLightingChannelMask >> 23) & RAY_TRACING_PAYLOAD_FLAG_PATHTRACER_VISIBILITY) != 0; }
 
 	FRayCone GetRayCone() { return UnpackRayCone(PackedRayCone); }
 	void SetRayCone(FRayCone NewRayCone) { PackedRayCone = PackRayCone(NewRayCone); }
diff --git a/Engine/Shaders/Private/RayTracing/RayTracingMaterialHitShaders.usf b/Engine/Shaders/Private/RayTracing/RayTracingMaterialHitShaders.usf
--- a/Engine/Shaders/Private/RayTracing/RayTracingMaterialHitShaders.usf
+++ b/Engine/Shaders/Private/RayTracing/RayTracingMaterialHitShaders.usf
@@ -31,7 +31,8 @@
 #include "/Engine/Private/SHCommon.ush"
 #include "/Engine/Private/ReflectionEnvironmentShared.ush"
 #include "/Engine/Private/VirtualTextureCommon.ush" 
-#include "/Engine/Private/LightmapCommon.ush"     
+#include "/Engine/Private/LightmapCommon.ush"
+#include "/Engine/Private/PathTracing/Material/PathTracingFresnel.ush"
 
 /** Computes sky diffuse lighting, including precomputed shadowing. */
 void GetSkyLighting(
@@ -520,7 +521,7 @@ RAY_TRACING_ENTRY_ANY_HIT(MaterialAHS,
 	// Usually this path is inactive because opaque materials don't have AHS bound
 	return;
 #elif MATERIALBLENDING_ALPHAHOLDOUT || MATERIALBLENDING_ADDITIVE
-	if (PackedPayload.IsEvalOpacity())
+	if (PackedPayload.IsPathTracingVisibilityRay())
 	{
 		// path tracer wants to treat these cases as fully transparent
 		// no material evaluation is required
@@ -544,7 +545,7 @@ RAY_TRACING_ENTRY_ANY_HIT(MaterialAHS,
 #endif
 		return;
 	}
-	else if (PackedPayload.IsEvalOpacity())
+	else if (PackedPayload.IsPathTracingVisibilityRay())
 	{
 		// this case is used by the path tracer only currently
 		// We want to evaluate the material opacity and accumulate it
@@ -599,7 +600,8 @@ RAY_TRACING_ENTRY_ANY_HIT(MaterialAHS,
 	}
 #else
 	// If we got this far -- we have a blending mode that can produce partial opacity _and_ we want to
-	// know how much it contributes to the throughput along this ray
+	// know how much it contributes to the throughput along this ray. These cases are currently only
+	// exercised by the path tracer.
 #if MATERIALBLENDING_MODULATE
 	// this material is the only one to lead to colored opacity
 	float3 Transparency = GetMaterialEmissive(PixelMaterialInputs);
@@ -607,16 +609,32 @@ RAY_TRACING_ENTRY_ANY_HIT(MaterialAHS,
 	float Opacity = GetMaterialOpacity(PixelMaterialInputs);
 	float Transparency = 1 - Opacity;
 #elif MATERIALBLENDING_TRANSLUCENT
+	float Opacity = GetMaterialOpacity(PixelMaterialInputs);
+	float Transparency = 1 - Opacity;
 	float Ior = GetMaterialRefraction(PixelMaterialInputs).x;
-	if (Ior != 1.0)
+	// Is refraction enabled?
+	if (Transparency > 0 && Ior != 1.0)
 	{
-		// refraction enabled, treat as opaque for now
-		// #dxr_todo: the path tracer will eventualy want an option for "fake caustics" 
-		return;
+		float PathRoughness = asfloat(PackedPayload.MetallicAndSpecularAndRoughness);
+		float Roughness = GetMaterialRoughness(PixelMaterialInputs);
+		if (PathRoughness <= Roughness)
+		{
+			// not using fast caustic approximation - treat as opaque
+			return;
+		}
+		// The heuristic used here is inspired by the following presentations:
+		//   - Kulla & Conty: Revisiting Physically Based Shading at Imageworks
+		//       https://blog.selfshadow.com/publications/s2017-shading-course/imageworks/s2017_pbs_imageworks_slides_v2.pdf
+		//   - Colin Barré-Brisebois - Pica Pica & Nvidia Turing
+		//       https://www.ea.com/seed/news/siggraph-2018-picapica-nv-turing
+
+		// fake caustics
+		float3 N = normalize(MaterialParameters.WorldNormal);
+		float Eta = Ior;
+		float VoN = abs(dot(WorldRayDirection(), N));
+		float Fr = FresnelReflectance(VoN, Eta);
+		Transparency *= pow(1 - Fr, 2) * (1 - Roughness * Roughness) * saturate(PathRoughness - Roughness);
 	}
-	// refraction disabled - treat as transparency
-	float Opacity = GetMaterialOpacity(PixelMaterialInputs);
-	float Transparency = 1 - Opacity;
 #else
 #error "Control flow does not account for all blending modes!"
 #endif
diff --git a/Engine/Source/Runtime/Renderer/Private/PathTracing.cpp b/Engine/Source/Runtime/Renderer/Private/PathTracing.cpp
--- a/Engine/Source/Runtime/Renderer/Private/PathTracing.cpp
+++ b/Engine/Source/Runtime/Renderer/Private/PathTracing.cpp
@@ -63,6 +63,14 @@ TAutoConsoleVariable<float> CVarPathTracingMaxPathIntensity(
 	ECVF_RenderThreadSafe
 );
 
+TAutoConsoleVariable<int32> CVarPathTracingApproximateCaustics(
+	TEXT("r.PathTracing.ApproximateCaustics"),
+	1,
+	TEXT("When non-zero, the path tracer will approximate caustic paths to reduce noise. This reduces speckles and noise from low-roughness glass and metals. (default = 1 (enabled))"),
+	ECVF_RenderThreadSafe
+);
+
+
 TAutoConsoleVariable<int32> CVarPathTracingFrameIndependentTemporalSeed(
 	TEXT("r.PathTracing.FrameIndependentTemporalSeed"),
 	1,
@@ -109,6 +117,7 @@ static bool PrepareShaderArgs(const FViewInfo& View, FPathTracingData& PathTraci
 	PathTracingData.VisibleLights = CVarPathTracingVisibleLights.GetValueOnRenderThread();
 	PathTracingData.MaxPathIntensity = CVarPathTracingMaxPathIntensity.GetValueOnRenderThread();
 	PathTracingData.UseErrorDiffusion = CVarPathTracingUseErrorDiffusion.GetValueOnRenderThread();
+	PathTracingData.ApproximateCaustics = CVarPathTracingApproximateCaustics.GetValueOnRenderThread();
 
 	bool NeedInvalidation = false;
 
@@ -154,6 +163,14 @@ static bool PrepareShaderArgs(const FViewInfo& View, FPathTracingData& PathTraci
 		PreviousUseErrorDiffusion = PathTracingData.UseErrorDiffusion;
 	}
 
+	// Changing approximate caustics requires starting over
+	static uint32 PreviousApproximateCaustics = PathTracingData.ApproximateCaustics;
+	if (PreviousApproximateCaustics != PathTracingData.ApproximateCaustics)
+	{
+		NeedInvalidation = true;
+		PreviousApproximateCaustics = PathTracingData.ApproximateCaustics;
+	}
+
 	// the rest of PathTracingData and AdaptiveSamplingData is filled in by SetParameters below
 	return NeedInvalidation;
 }
diff --git a/Engine/Source/Runtime/Renderer/Private/PathTracingUniformBuffers.h b/Engine/Source/Runtime/Renderer/Private/PathTracingUniformBuffers.h
--- a/Engine/Source/Runtime/Renderer/Private/PathTracingUniformBuffers.h
+++ b/Engine/Source/Runtime/Renderer/Private/PathTracingUniformBuffers.h
@@ -17,6 +17,7 @@ BEGIN_GLOBAL_SHADER_PARAMETER_STRUCT(FPathTracingData, )
 	SHADER_PARAMETER(uint32, MaxBounces)
 	SHADER_PARAMETER(uint32, MISMode)
 	SHADER_PARAMETER(uint32, VisibleLights)
+	SHADER_PARAMETER(uint32, ApproximateCaustics)
 	SHADER_PARAMETER(float, MaxPathIntensity)
 	SHADER_PARAMETER(float, MaxNormalBias)
 END_GLOBAL_SHADER_PARAMETER_STRUCT()
