Fix uninitialized members that anger the shader compiler in debug mode

#rnx

diff --git a/Engine/Shaders/Private/VirtualShadowMaps/ProjectionCommon.ush b/Engine/Shaders/Private/VirtualShadowMaps/ProjectionCommon.ush
--- a/Engine/Shaders/Private/VirtualShadowMaps/ProjectionCommon.ush
+++ b/Engine/Shaders/Private/VirtualShadowMaps/ProjectionCommon.ush
@@ -79,6 +79,8 @@ FVirtualSmSample InitVirtualSmSample()
 	Result.Depth = 0.0f;
 	Result.MipLevel = 0;
 	Result.bValid = false;
+	Result.VirtualTexelAddress = Result.PhysicalTexelAddress = uint2(0U, 0U);
+	Result.VirtualTexelAddressFloat = float2(0.0f, 0.0f);
 	return Result;
 }
 
