Added Haziness v1.0 to Strata slab material. Only 50/50 blend of roughness and lerp(roughness, 1, Haziness).
Quite a lot of changes and interactions especialyl with SSR and special case for rect light LTC as well as anisotropy.
SSR progressively fades out as Haziness is revealed for environment lighting.

Minor renamings.

diff --git a/Engine/Shaders/Private/ReflectionEnvironmentPixelShader.usf b/Engine/Shaders/Private/ReflectionEnvironmentPixelShader.usf
--- a/Engine/Shaders/Private/ReflectionEnvironmentPixelShader.usf
+++ b/Engine/Shaders/Private/ReflectionEnvironmentPixelShader.usf
@@ -376,13 +376,15 @@ void ReflectionEnvironmentSkyLighting(
 	{
 		const FStrataTopLayerData TopLayerData = StrataUnpackTopLayerData(Strata.TopLayerTexture.Load(uint3(PixelPos, 0)));
 		const float4 SSR = Texture2DSample(ScreenSpaceReflectionsTexture, ScreenSpaceReflectionsSampler, BufferUV);
-		SkyLighting.rgb += SSR.rgb * TopLayerData.EnvBRDF;
 
 		// SSR is computed for the top level by averaging BSDF components.
 		// This is the top level specular contribution factor according to areas.
 		const float TopLayerSpecularContributionFactor = 1.0f - SSR.a;
 
-		SkyLighting += StrataEnvironmentLighting(
+		// A reduction of SSR contribution is needed when there is haziness because we only consider SSR for the sharpest lobe.
+		float SSRReductionFactor = 1.0f;
+
+		float4 StrataEnvLighting = StrataEnvironmentLighting(
 			BufferUV,
 			GBuffer,
 			WorldPosition,
@@ -391,7 +393,10 @@ void ReflectionEnvironmentSkyLighting(
 			AmbientOcclusion,
 			TopLayerSpecularContributionFactor,
 			Strata.MaterialLobesBuffer,
-			GetStrataPixelDataByteOffset(PixelPos, View.BufferSizeAndInvSize.x, Strata.MaxBytesPerPixel));
+			GetStrataPixelDataByteOffset(PixelPos, View.BufferSizeAndInvSize.x, Strata.MaxBytesPerPixel),
+			SSRReductionFactor);
+
+		SkyLighting.rgb = SSR.rgb * TopLayerData.EnvBRDF * SSRReductionFactor + StrataEnvLighting.rgb;
 
 		ApplyCloudVolumetricShadow(WorldPosition, SkyLighting);
 
diff --git a/Engine/Shaders/Private/ShadingCommon.ush b/Engine/Shaders/Private/ShadingCommon.ush
--- a/Engine/Shaders/Private/ShadingCommon.ush
+++ b/Engine/Shaders/Private/ShadingCommon.ush
@@ -113,3 +113,8 @@ float MakeRoughnessSafe(float Roughness, float MinRoughness=0.001f)
 {
 	return clamp(Roughness, MinRoughness, 1.0f);
 }
+
+float ComputeHazyLobeRoughness(float Roughness, float Haziness)
+{
+	return lerp(Roughness, 1.0f, Haziness);
+}
diff --git a/Engine/Shaders/Private/Strata/Strata.ush b/Engine/Shaders/Private/Strata/Strata.ush
--- a/Engine/Shaders/Private/Strata/Strata.ush
+++ b/Engine/Shaders/Private/Strata/Strata.ush
@@ -53,6 +53,7 @@ struct FStrataBSDF
 #define STATE_BITS_HASSCATTERING		1
 #define STATE_BITS_HASSSSPROFILE		1
 #define STATE_BITS_GREYWEIGHT			1
+#define STATE_BITS_HAZINESS				1
 #define STATE_BITS___UNUSED___			11
 #define STATE_BITS_WEIGHT10F			10
 
@@ -64,6 +65,7 @@ struct FStrataBSDF
 #define STATE_BIT_OFFSET_HASSCATTERING	(STATE_BIT_OFFSET_TOPLAYER			+ STATE_BITS_TOPLAYER)
 #define STATE_BIT_OFFSET_HASSSSPROFILE	(STATE_BIT_OFFSET_HASSCATTERING		+ STATE_BITS_HASSCATTERING)
 #define STATE_BIT_OFFSET_GREYWEIGHT		(STATE_BIT_OFFSET_HASSSSPROFILE		+ STATE_BITS_HASSSSPROFILE)
+#define STATE_BIT_OFFSET_HAZINESS		(STATE_BIT_OFFSET_GREYWEIGHT		+ STATE_BITS_GREYWEIGHT)
 #define STATE_BIT_OFFSET_WEIGHT10F		22
 
 #define TO_BITMASK(BitCount)							((1<<BitCount)-1)
@@ -100,6 +102,10 @@ struct FStrataBSDF
 #define BSDF_GETGREYWEIGHT(X)			BSDF_GETSTATE(X, STATE_BITS_GREYWEIGHT, STATE_BIT_OFFSET_GREYWEIGHT)
 #define BSDF_SETGREYWEIGHT(X, Grey)		BSDF_SETSTATE(X, STATE_BITS_GREYWEIGHT, STATE_BIT_OFFSET_GREYWEIGHT, Grey)
 
+// Indicates if the BSDF weight is grey scale so that it can be stored in a single float (e.g. top layer BSDFs)
+#define BSDF_GETHASHAZINESS(X)			BSDF_GETSTATE(X, STATE_BITS_HAZINESS, STATE_BIT_OFFSET_HAZINESS)
+#define BSDF_SETHASHAZINESS(X, Haze)	BSDF_SETSTATE(X, STATE_BITS_HAZINESS, STATE_BIT_OFFSET_HAZINESS, Haze)
+
 // Accessors for the 10bits float weight used if is it detected being a grey scale
 #define BSDF_GETWEIGHT10F(X)			BSDF_GETSTATE(X, STATE_BITS_WEIGHT10F, STATE_BIT_OFFSET_WEIGHT10F)
 #define BSDF_SETWEIGHT10F(X, Weight)	BSDF_SETSTATE(X, STATE_BITS_WEIGHT10F, STATE_BIT_OFFSET_WEIGHT10F, Weight)
@@ -117,6 +123,7 @@ struct FStrataBSDF
 #define SLAB_METALLIC(X)				X.VGPRs[1].x
 #define SLAB_ROUGHNESSX(X)				X.VGPRs[1].y
 #define SLAB_ROUGHNESSY(X)				X.VGPRs[1].z
+#define SLAB_HAZINESS(X)				X.VGPRs[1].w
 // Either with SSS profile
 #define SLAB_SSSPROFILEID(X)			X.VGPRs[2].x
 #define SLAB_SSSPROFILERADIUSSCALE(X)	X.VGPRs[2].y
@@ -307,12 +314,13 @@ FStrataData GetInitialisedStrataData()
 FStrataData GetStrataSlabBSDF(
 	float3 BaseColor, float Specular, float Metallic, float2 Roughness,
 	float SSSProfileID, float3 SSSDMFPAlbedo, float SSSDMFPRadius, float SSSRadiusScale,
-	float3 Emissive,
+	float3 Emissive, float Haziness, 
 	uint SharedNormalIndex, inout uint SharedNormalTypes)
 {
 	FStrataData StrataData = GetInitialisedStrataData();
 	StrataData.LayerCount++;
 
+	const uint bHasHaziness = Haziness > 0.0f ? 1 : 0;
 	const uint bHasAnisotropy = StrataHasAnisotropy(Roughness.x, Roughness.y) ? 1 : 0;
 	if (bHasAnisotropy)
 	{
@@ -331,12 +339,14 @@ FStrataData GetStrataSlabBSDF(
 	BSDF_SETISTOPLAYER		(StrataData.Layers[0].BSDFs[0],		0);
 	BSDF_SETHASSCATTERING	(StrataData.Layers[0].BSDFs[0],		bHasSSS ? 1 : 0);
 	BSDF_SETHASSSSPROFILE	(StrataData.Layers[0].BSDFs[0],		bHasSSSProfile ? 1 : 0);
+	BSDF_SETHASHAZINESS		(StrataData.Layers[0].BSDFs[0],		bHasHaziness);
 
 	SLAB_BASECOLOR			(StrataData.Layers[0].BSDFs[0])		= BaseColor;
 	SLAB_SPECULAR			(StrataData.Layers[0].BSDFs[0])		= Specular;
 	SLAB_METALLIC			(StrataData.Layers[0].BSDFs[0])		= Metallic;
 	SLAB_ROUGHNESSX			(StrataData.Layers[0].BSDFs[0])		= Roughness.x;
 	SLAB_ROUGHNESSY			(StrataData.Layers[0].BSDFs[0])		= Roughness.y;
+	SLAB_HAZINESS			(StrataData.Layers[0].BSDFs[0])		= Haziness;
 	if (bHasSSSProfile)
 	{
 		SLAB_SSSPROFILEID	(StrataData.Layers[0].BSDFs[0])		= SSSProfileID;
@@ -630,6 +640,7 @@ void SanitizeStrataSlab(inout FStrataBSDF BSDF)
 	SLAB_METALLIC(BSDF) = saturate(SLAB_METALLIC(BSDF));
 	SLAB_ROUGHNESSX(BSDF) = saturate(SLAB_ROUGHNESSX(BSDF));
 	SLAB_ROUGHNESSY(BSDF) = saturate(SLAB_ROUGHNESSY(BSDF));
+	SLAB_HAZINESS(BSDF) = saturate(SLAB_HAZINESS(BSDF));
 	SLAB_SSSPROFILERADIUSSCALE(BSDF) = saturate(SLAB_SSSPROFILERADIUSSCALE(BSDF));
 	SLAB_SSSDMFPALBEDO(BSDF) = saturate(SLAB_SSSDMFPALBEDO(BSDF));
 	SLAB_SSSDMFPRADIUS(BSDF) = clamp(SLAB_SSSDMFPRADIUS(BSDF), 0.0f, SSS_PROFILE_PERPIXEL_MAXRADIUS);
@@ -1264,7 +1275,7 @@ void PackStrataOut(
 			case STRATA_BSDF_TYPE_SLAB:
 			{
 				STRATA_STORE_UINT1(PackColorLinearToGamma2AlphaLinear(float4(SLAB_BASECOLOR(BSDF), SLAB_ROUGHNESSX(BSDF))));
-				STRATA_STORE_UINT1(PackRGBA8(float4(SLAB_SPECULAR(BSDF), SLAB_METALLIC(BSDF), SLAB_ROUGHNESSY(BSDF), 0)));
+				STRATA_STORE_UINT1(PackRGBA8(float4(SLAB_SPECULAR(BSDF), SLAB_METALLIC(BSDF), SLAB_ROUGHNESSY(BSDF), SLAB_HAZINESS(BSDF))));
 
 				if (BSDF_GETHASSCATTERING(BSDF))
 				{
@@ -1375,6 +1386,7 @@ FStrataBSDF UnpackStrataBSDFIn(ByteAddressBuffer StrataBuffer, inout uint ByteOf
 			SLAB_SPECULAR(OutBSDF) = Data.x;
 			SLAB_METALLIC(OutBSDF) = Data.y;
 			SLAB_ROUGHNESSY(OutBSDF) = Data.z;
+			SLAB_HAZINESS(OutBSDF) = Data.w;
 		}
 		if (BSDF_GETHASSCATTERING(OutBSDF))
 		{
diff --git a/Engine/Shaders/Private/Strata/StrataDeferredLighting.ush b/Engine/Shaders/Private/Strata/StrataDeferredLighting.ush
--- a/Engine/Shaders/Private/Strata/StrataDeferredLighting.ush
+++ b/Engine/Shaders/Private/Strata/StrataDeferredLighting.ush
@@ -9,6 +9,7 @@ struct FStrataLightIntegrationContext
 	float3 IncidentRadianceNoShadow;
 	float3 IncidentIrradianceNoShadow;	// Modified incident irradiance, based on custom are light integrator
 	float3 OutgoingRadianceNoShadow;	// Continue the outgoing radiance if the BSDF evaluation/integration has been done with the custom area light integrator
+	float3 IncidentIrradianceNoShadowHaze;
 	bool bEvaluateBSDF;					// True if the BSDF need to be evalued, false, if the integration has already been done
 };
 
@@ -18,6 +19,7 @@ FStrataLightIntegrationContext StrataIniLightIntegrationContext()
 	Out.bEvaluateBSDF = true;
 	Out.OutgoingRadianceNoShadow = 0;
 	Out.IncidentIrradianceNoShadow = 0;
+	Out.IncidentIrradianceNoShadowHaze = 0;
 	Out.IncidentRadianceNoShadow = 0;
 	return Out;
 }
@@ -74,17 +76,32 @@ void StrataPreprocessAreaLightIntegration(
 
 		// Currently anistropic material don't support area lighting and will fallback on punctual lighting (that is a2Override == 0)
 		const bool bAnistropy = BSDF_GETANISOTROPY(BSDFContext.BSDF);
+		const bool bHaziness = BSDF_GETHASHAZINESS(BSDFContext.BSDF);
 		if (!bAnistropy)
 		{
 			// Evaluate area-light preintegration
 			if (LightContext.AreaLight.bIsRect)
 			{
+				// Using a single SpecularColor and roughness so haziness is ignored for LTC rect lights
+
 				// Rect light uses LTC integrator, which computes the full BDSD x Area light integration. 
 				// This is why there is no need to evaluate the BSDF. The result is stored into a local 
 				// OutgoingRadianceNoShadow, which will be accumulated outside of this function
-				Out.OutgoingRadianceNoShadow.rgb +=
+				float3 LocalOutgoingRadianceNoShadow =
 					Out.IncidentRadianceNoShadow *
 					RectGGXApproxLTC(SafeRoughnessX, SpecularColor, BSDFContext.N, BSDFContext.V, LightContext.AreaLight.Rect, LightContext.AreaLight.Texture);
+
+				if (bHaziness)
+				{
+					Out.IncidentIrradianceNoShadow =
+						(1.0f - SPECULAR_HAZINESS_WEIGHT) * LocalOutgoingRadianceNoShadow +
+						SPECULAR_HAZINESS_WEIGHT * Out.IncidentRadianceNoShadow * RectGGXApproxLTC(SafeRoughnessX, SpecularColor, BSDFContext.N, BSDFContext.V, LightContext.AreaLight.Rect, LightContext.AreaLight.Texture);
+				}
+				else
+				{
+					Out.OutgoingRadianceNoShadow = LocalOutgoingRadianceNoShadow;
+				}
+
 				Out.bEvaluateBSDF = false;
 			}
 			else
@@ -105,6 +122,17 @@ void StrataPreprocessAreaLightIntegration(
 					Out.IncidentRadianceNoShadow *
 					EnergyNormalization(BSDFContext.a2Override, BSDFContext.Context.VoH, LightContext.AreaLight) *
 					LightContext.AreaLight.FalloffColor * (LightContext.Falloff * LightContext.NoL);
+
+				if (bHaziness)
+				{
+					const float SafeRoughnessHaze = ComputeHazyLobeRoughness(SafeRoughnessX, SLAB_HAZINESS(BSDFContext.BSDF));
+					BSDFContext.a2OverrideHazy = Pow4(SafeRoughnessHaze);
+					Out.IncidentIrradianceNoShadowHaze =
+						Out.IncidentRadianceNoShadow *
+						EnergyNormalization(BSDFContext.a2OverrideHazy, BSDFContext.Context.VoH, LightContext.AreaLight) *
+						LightContext.AreaLight.FalloffColor * (LightContext.Falloff * LightContext.NoL);
+				}
+
 				Out.bEvaluateBSDF = true;
 			}
 		}
@@ -152,6 +180,7 @@ float4 StrataDeferredLighting(
 		FStrataLightIntegrationContext IntegrationContext = StrataIniLightIntegrationContext();
 		IntegrationContext.IncidentRadianceNoShadow   = LightData.Color * LightMask;
 		IntegrationContext.IncidentIrradianceNoShadow = PreintegratedIndicentIrradiance * IntegrationContext.IncidentRadianceNoShadow * StrataBSDFProjectedSolidAngleFactor(BSDFContext);
+		IntegrationContext.IncidentIrradianceNoShadowHaze = IntegrationContext.IncidentIrradianceNoShadow;
 		StrataPreprocessAreaLightIntegration(ToLight, LightData, RectTexture, BSDFContext, IntegrationContext);
 
 		// Evaluate BSDF
@@ -160,13 +189,15 @@ float4 StrataDeferredLighting(
 			FStrataEvaluateResult BSDFEvaluate = StrataEvaluateBSDF(BSDFContext, ShadowTerms);
 
 			const float3 DiffusePathWeight		= BSDF.Weight * BSDFEvaluate.DiffusePathValue;
-			const float3 SpecularPathWeight		= BSDF.Weight * BSDFEvaluate.SpecularPathValue;
+			const float3 SpecularPathWeight		= BSDF.Weight * BSDFEvaluate.SpecularPathValue     * BSDFEvaluate.SpecularPathProbability;
+			const float3 SpecularHazePathWeight = BSDF.Weight * BSDFEvaluate.SpecularHazePathValue * BSDFEvaluate.SpecularHazePathProbability;
 			const float3 TransmissionPathWeight	= BSDF.Weight * BSDFEvaluate.TransmissionPathValue;
 
 			const bool bNeedsSeparateSubsurfaceLightAccumulation = BSDFEvaluate.bSubsurface;
 
-			LightAccumulator_AddSplit(Out, DiffusePathWeight,		SpecularPathWeight, DiffusePathWeight,		IntegrationContext.IncidentIrradianceNoShadow * ShadowTerms.SurfaceShadow,		bNeedsSeparateSubsurfaceLightAccumulation);
-			LightAccumulator_AddSplit(Out, TransmissionPathWeight,	0.0f,				TransmissionPathWeight,	IntegrationContext.IncidentIrradianceNoShadow * ShadowTerms.TransmissionShadow,	bNeedsSeparateSubsurfaceLightAccumulation);
+			LightAccumulator_AddSplit(Out, DiffusePathWeight,		SpecularPathWeight,		DiffusePathWeight,		IntegrationContext.IncidentIrradianceNoShadow * ShadowTerms.SurfaceShadow,		bNeedsSeparateSubsurfaceLightAccumulation);
+			LightAccumulator_AddSplit(Out, 0.0,						SpecularHazePathWeight, 0.0,					IntegrationContext.IncidentIrradianceNoShadowHaze * ShadowTerms.SurfaceShadow,	bNeedsSeparateSubsurfaceLightAccumulation);
+			LightAccumulator_AddSplit(Out, TransmissionPathWeight,	0.0f,					TransmissionPathWeight,	IntegrationContext.IncidentIrradianceNoShadow * ShadowTerms.TransmissionShadow,	bNeedsSeparateSubsurfaceLightAccumulation);
 		}
 		else
 		{
diff --git a/Engine/Shaders/Private/Strata/StrataEnvironmentLighting.ush b/Engine/Shaders/Private/Strata/StrataEnvironmentLighting.ush
--- a/Engine/Shaders/Private/Strata/StrataEnvironmentLighting.ush
+++ b/Engine/Shaders/Private/Strata/StrataEnvironmentLighting.ush
@@ -13,12 +13,15 @@ float4 StrataEnvironmentLighting(
 	float AmbientOcclusion,
 	float TopLayerSpecularContributionFactor,
 	ByteAddressBuffer StrataDataBuffer,
-	uint PixelStrataDataByteOffset)
+	uint PixelStrataDataByteOffset,
+	inout float SSRReductionFactor)
 {
 	FStrataPixelHeader StrataPixelHeader = UnpackStrataHeaderIn(StrataDataBuffer, PixelStrataDataByteOffset);
 
 	FLightAccumulator Out = (FLightAccumulator)0;
 
+	SSRReductionFactor = 1.0f;
+
 	float3 CommonFactor = 1;
 #if USE_PREEXPOSURE
 	CommonFactor *= View.PreExposure;
@@ -58,14 +61,26 @@ float4 StrataEnvironmentLighting(
 		}
 
 		// Specular component
+		const bool bIsTopLayer = BSDF_GETISTOPLAYER(BSDF);
 		float3 SpecularSkyLighting = 0;
 		if (any(StrataEnvLight.SpecularWeight > 0.0f))
 		{
 			float SkyAverageBrightness = 1.0f;
 			const float3 SpecularLuminance = GetSkyLightReflection(StrataEnvLight.SpecularDirection, StrataEnvLight.SpecularSafeRoughness, SkyAverageBrightness);
 
-			SpecularSkyLighting = BSDFThroughput * SpecularLuminance * StrataEnvLight.SpecularWeight * View.SkyLightColor.rgb
-				* AmbientOcclusion * (BSDF_GETISTOPLAYER(BSDF) ? TopLayerSpecularContributionFactor : 1.0f);
+			SpecularSkyLighting += BSDFThroughput * SpecularLuminance * StrataEnvLight.SpecularWeight * View.SkyLightColor.rgb * AmbientOcclusion * 
+				(bIsTopLayer ? TopLayerSpecularContributionFactor : 1.0f);
+		}
+		if (BSDF_GETHASHAZINESS(BSDF))
+		{
+			float SkyAverageBrightness = 1.0f;
+			const float3 SpecularLuminance = GetSkyLightReflection(StrataEnvLight.SpecularDirection, StrataEnvLight.SpecularHazeSafeRoughness, SkyAverageBrightness);
+
+			SpecularSkyLighting += BSDFThroughput * SpecularLuminance * StrataEnvLight.SpecularHazeWeight * View.SkyLightColor.rgb * AmbientOcclusion;
+
+			// SSR is traced for the sharpest lob. The smoothest one does not rely on SSR so we need to lower energy coming from SSR according to the lobe blend weight.
+			// And we also try to make the transition smooth using Haziness
+			SSRReductionFactor -= bIsTopLayer ? dot(BSDFThroughput, float3(1.0f / 3.0f, 1.0f / 3.0f, 1.0f / 3.0f)) * StrataEnvLight.SSRReduction : 0.0f;
 		}
 
 		const bool bNeedsSeparateSubsurfaceLightAccumulation = StrataEnvLight.bSubsurface;
@@ -73,6 +88,9 @@ float4 StrataEnvironmentLighting(
 
 		// STRATA_TODO: Missing probe reflection here
 	}
+
+	SSRReductionFactor = saturate(SSRReductionFactor);
+
 	return LightAccumulator_GetResult(Out);
 }
 
diff --git a/Engine/Shaders/Private/Strata/StrataEvaluation.ush b/Engine/Shaders/Private/Strata/StrataEvaluation.ush
--- a/Engine/Shaders/Private/Strata/StrataEvaluation.ush
+++ b/Engine/Shaders/Private/Strata/StrataEvaluation.ush
@@ -19,6 +19,9 @@ float3 StrataGetOffSpecularPeakReflectionDir(float3 Normal, float3 ReflectionVec
 	return lerp(Normal, ReflectionVector, (1 - a) * (sqrt(1 - a) + a));
 }
 
+// When haziness is used, both lobes take equal part in the blending. Only roughness changes.
+#define SPECULAR_HAZINESS_WEIGHT 0.5f
+
 
 
 ///////////////////////////////////////////////////////////////////////////////
@@ -47,6 +50,7 @@ struct FStrataBSDFContext
 
 	// Preprocessed roughness override used for specular BSDF evaluation when integration area light (using MRP integrator)
 	float a2Override;
+	float a2OverrideHazy;
 };
 
 void StrataGetBSDFRoughness(in FStrataBSDF BSDF, inout float R)
@@ -135,10 +139,17 @@ struct FStrataEvaluateResult
 
 	float3 DiffusePathValue;
 	float3 SpecularPathValue;
+	float3 SpecularHazePathValue;
 	float3 TransmissionPathValue;
 
+	// The following probabilities are part of the pdf to work correctly with a monter carlo based integrator and as expected are not applied on the xxxPathValue.
+	// So PathProbability should be applied on PathValue in the rasteriser to recover the correct balance.
+	float SpecularPathProbability;
+	float SpecularHazePathProbability;
+
 	float  DiffusePDF;
 	float  SpecularPDF;
+	float  SpecularHazePDF;
 	float  TransmissionPDF;
 
 	float3 Throughput;					// Throughput to the next layer
@@ -175,15 +186,19 @@ FStrataEvaluateResult StrataEvaluateBSDF(FStrataBSDFContext BSDFContext, FShadow
 
 
 			const float3 F0 = ComputeF0(SLAB_SPECULAR(BSDFContext.BSDF), BaseColor, Metallic);
-			// STRATA_TODO F82
+			// STRATA_TODO F82 or edge color
 
 			const float SafeRoughnessX = MakeRoughnessSafe(SLAB_ROUGHNESSX(BSDFContext.BSDF));
 			const float SafeRoughnessY = MakeRoughnessSafe(SLAB_ROUGHNESSY(BSDFContext.BSDF));
 
 			const bool bHasAnisotropy = BSDF_GETANISOTROPY(BSDFContext.BSDF);
+			const bool bHaziness = BSDF_GETHASHAZINESS(BSDFContext.BSDF);
 			float D = 0;
 			float Vis = 0;
 			float PDF = 0;
+			float HazeD = 0;
+			float HazeVis = 0;
+			float HazePDF = 0;
 			if (bHasAnisotropy)
 			{
 				// Generalized microfacet specular
@@ -192,6 +207,15 @@ FStrataEvaluateResult StrataEvaluateBSDF(FStrataBSDFContext BSDFContext, FShadow
 				D = D_GGXaniso(Alpha0, Alpha1, BSDFContext.Context.NoH, BSDFContext.Context.XoH, BSDFContext.Context.YoH);
 				Vis = Vis_SmithJointAniso(Alpha0, Alpha1, BSDFContext.Context.NoV, BSDFContext.Context.NoL, BSDFContext.Context.XoV, BSDFContext.Context.XoL, BSDFContext.Context.YoV, BSDFContext.Context.YoL);
 				PDF = VisibleGGXPDF_aniso(BSDFContext.TangentV, BSDFContext.TangentH, float2(Alpha0, Alpha1));
+
+				if (bHaziness)
+				{
+					Alpha0 = Pow2(ComputeHazyLobeRoughness(SafeRoughnessX, SLAB_HAZINESS(BSDFContext.BSDF)));
+					Alpha1 = Pow2(ComputeHazyLobeRoughness(SafeRoughnessY, SLAB_HAZINESS(BSDFContext.BSDF)));
+					HazeD = D_GGXaniso(Alpha0, Alpha1, BSDFContext.Context.NoH, BSDFContext.Context.XoH, BSDFContext.Context.YoH);
+					HazeVis = Vis_SmithJointAniso(Alpha0, Alpha1, BSDFContext.Context.NoV, BSDFContext.Context.NoL, BSDFContext.Context.XoV, BSDFContext.Context.XoL, BSDFContext.Context.YoV, BSDFContext.Context.YoL);
+					HazePDF = VisibleGGXPDF_aniso(BSDFContext.TangentV, BSDFContext.TangentH, float2(Alpha0, Alpha1));
+				}
 			}
 			else
 			{
@@ -202,11 +226,25 @@ FStrataEvaluateResult StrataEvaluateBSDF(FStrataBSDFContext BSDFContext, FShadow
 				D = D_GGX(a2, BSDFContext.Context.NoH);
 				Vis = Vis_SmithJointApprox(a2, BSDFContext.Context.NoV, BSDFContext.SatNoL);
 				PDF = VisibleGGXPDF(BSDFContext.TangentV, BSDFContext.TangentH, a2);
+
+				if (bHaziness)
+				{
+					const float a2Haze = BSDFContext.a2OverrideHazy > 0 ? BSDFContext.a2OverrideHazy : Pow4(ComputeHazyLobeRoughness(SafeRoughnessX, SLAB_HAZINESS(BSDFContext.BSDF)));
+					HazeD = D_GGX(a2Haze, BSDFContext.Context.NoH);
+					HazeVis = Vis_SmithJointApprox(a2Haze, BSDFContext.Context.NoV, BSDFContext.SatNoL);
+					HazePDF = VisibleGGXPDF(BSDFContext.TangentV, BSDFContext.TangentH, a2Haze);
+				}
 			}
 			float3 F = F_Schlick(F0, BSDFContext.Context.VoH);
 
-			Sample.SpecularPathValue= D * Vis * F;
-			Sample.SpecularPDF		= PDF;
+			Sample.SpecularPathProbability	= bHaziness		? (1.0f - SPECULAR_HAZINESS_WEIGHT) : 1.0f;
+			Sample.SpecularHazePathProbability= bHaziness	? SPECULAR_HAZINESS_WEIGHT			: 0.0f;
+
+			Sample.SpecularPathValue		= D * Vis * F;
+			Sample.SpecularPDF				= PDF * Sample.SpecularPathProbability;
+
+			Sample.SpecularHazePathValue	= HazeD * HazeVis * F;
+			Sample.SpecularHazePDF			= HazePDF * Sample.SpecularHazePathProbability;
 
 			Sample.EmissivePathValue= BSDF_GETEMISSIVE(BSDFContext.BSDF);
 
@@ -326,6 +364,8 @@ FStrataImportanceSampleResult StrataImportanceSampleBSDF(FStrataPixelHeader Stra
 				GGXSample = ImportanceSampleVisibleGGX(UniformSampleDisk(E), Pow4(SafeRoughnessX), BSDFContext.TangentV);
 			}
 
+			// STRATA_TODO: take into account Haziness
+
 			float3 WorldH = mul(GGXSample.xyz, BSDFContext.TangentBasis);
 			Sample.SpecularDirection = reflect(CameraVector, WorldH);
 			PDF = GGXSample.w;
@@ -377,6 +417,11 @@ struct FStrataEnvLightResult
 	float3 SpecularWeight;
 	float  SpecularSafeRoughness;
 
+	// Specular haze maintains the same off peak specular direction as the non hazy specular lobe
+	float3 SpecularHazeWeight;
+	float  SpecularHazeSafeRoughness;
+	float  SSRReduction;
+
 	bool   bSubsurface; // True if we need to separate the subsurface light contribution for the screen space diffusion process.
 };
 
@@ -425,6 +470,23 @@ FStrataEnvLightResult StrataEvaluateForEnvLight(FStrataBSDFContext BSDFContext,
 				StrataEnvLightResult.SpecularDirection = StrataGetOffSpecularPeakReflectionDir(BSDFContext.N, BSDFContext.R, SafeRoughnessX);
 				StrataEnvLightResult.SpecularWeight = EvalEnvBRDF;
 				StrataEnvLightResult.SpecularSafeRoughness = SafeRoughnessX;
+
+				const bool bHasHaziness = BSDF_GETHASHAZINESS(BSDFContext.BSDF);
+				if (bHasHaziness)
+				{
+					// Smoothly fade in haziness while fading out SSR to make sure we do not add energy and avoid popping.
+					const float Haziness = SLAB_HAZINESS(BSDFContext.BSDF);
+					const float HazinessFadeIn = saturate(Haziness / 0.1f);
+					StrataEnvLightResult.SSRReduction = lerp(0.0f, SPECULAR_HAZINESS_WEIGHT, HazinessFadeIn);
+
+					// Apply blend factor on the sharp specular contribution
+					StrataEnvLightResult.SpecularWeight *= lerp(1.0f, (1.0f - SPECULAR_HAZINESS_WEIGHT), HazinessFadeIn);
+
+					// Compute the second specular weight and roughness
+					float SafeRoughnessHaze = ComputeHazyLobeRoughness(SafeRoughnessX, Haziness);
+					StrataEnvLightResult.SpecularHazeWeight = lerp(0.0f, SPECULAR_HAZINESS_WEIGHT, HazinessFadeIn) * EnvBRDFApprox(F0, SafeRoughnessHaze, BSDFContext.SatNoV);
+					StrataEnvLightResult.SpecularHazeSafeRoughness = SafeRoughnessHaze;
+				}
 			}
 			break;
 		}
diff --git a/Engine/Shaders/Private/Strata/StrataVisualize.usf b/Engine/Shaders/Private/Strata/StrataVisualize.usf
--- a/Engine/Shaders/Private/Strata/StrataVisualize.usf
+++ b/Engine/Shaders/Private/Strata/StrataVisualize.usf
@@ -255,6 +255,10 @@ void VisualizeMaterialPS(
 					NEXTLINE;
 				}
 
+				HORIZONTALTAB; PRINTPARAMCHAR(_H_); PRINTPARAMCHAR(_A_); PRINTPARAMCHAR(_Z_); PRINTPARAMCHAR(_I_); PRINTPARAMCHAR(_N_); PRINTPARAMCHAR(_SPC_);
+				HORIZONTALTAB; PRINTFLOAT1(SLAB_HAZINESS(BSDF));
+				NEXTLINE;
+
 				if (BSDF_GETHASSCATTERING(BSDF))
 				{
 					if (BSDF_GETHASSSSPROFILE(BSDF))
diff --git a/Engine/Source/Runtime/Engine/Classes/Materials/MaterialExpressionStrata.h b/Engine/Source/Runtime/Engine/Classes/Materials/MaterialExpressionStrata.h
--- a/Engine/Source/Runtime/Engine/Classes/Materials/MaterialExpressionStrata.h
+++ b/Engine/Source/Runtime/Engine/Classes/Materials/MaterialExpressionStrata.h
@@ -103,6 +103,12 @@ class UMaterialExpressionStrataSlabBSDF : public UMaterialExpressionStrataBSDF
 	UPROPERTY()
 	FExpressionInput EmissiveColor;
 
+	/**
+	 * Haziness controls the relative roughness of a second specular lobe. 0 means disabled and 1 means the second lobe specular lobe will lerp the current roughness to fully rough. (type = float, unitless, default = 0)
+	 */
+	UPROPERTY()
+	FExpressionInput Haziness;
+
 	/** SubsurfaceProfile, for Screen Space Subsurface Scattering. The profile needs to be set up on both the Strata diffuse node, and the material node at the moment. */
 	UPROPERTY(EditAnywhere, BlueprintReadOnly, Category = Material, meta = (DisplayName = "Subsurface Profile"))
 	class USubsurfaceProfile* SubsurfaceProfile;
diff --git a/Engine/Source/Runtime/Engine/Private/Materials/HLSLMaterialTranslator.cpp b/Engine/Source/Runtime/Engine/Private/Materials/HLSLMaterialTranslator.cpp
--- a/Engine/Source/Runtime/Engine/Private/Materials/HLSLMaterialTranslator.cpp
+++ b/Engine/Source/Runtime/Engine/Private/Materials/HLSLMaterialTranslator.cpp
@@ -8527,12 +8527,15 @@ int32 FHLSLMaterialTranslator::StrataCreateAndRegisterNullMaterial()
 	return OutputCodeChunk;
 }
 
-int32 FHLSLMaterialTranslator::StrataSlabBSDF(int32 BaseColor, int32 Reflectivity, int32 Metallic, int32 RoughnessX, int32 RoughnessY, int32 SSSProfileId, int32 SSSDMFPAlbedo, int32 SSSDMFPRadius, int32 SSSRadiusScale, int32 EmissiveColor, int32 Normal, int32 Tangent, uint8 SharedNormalIndex)
+int32 FHLSLMaterialTranslator::StrataSlabBSDF(
+	int32 BaseColor, int32 Specular, int32 Metallic, int32 RoughnessX, int32 RoughnessY,
+	int32 SSSProfileId, int32 SSSDMFPAlbedo, int32 SSSDMFPRadius, int32 SSSRadiusScale, int32 EmissiveColor, 
+	int32 Haziness, int32 Normal, int32 Tangent, uint8 SharedNormalIndex)
 {
 	return AddCodeChunk(
-		MCT_Strata, TEXT("GetStrataSlabBSDF(%s, %s, %s, float2(%s, %s), %s, %s, %s, %s, %s, %u, Parameters.SharedNormals.NormalTypes) /* Normal = %s ; Tangent = %s */"),
+		MCT_Strata, TEXT("GetStrataSlabBSDF(%s, %s, %s, float2(%s, %s), %s, %s, %s, %s, %s, %s, %u, Parameters.SharedNormals.NormalTypes) /* Normal = %s ; Tangent = %s */"),
 		*GetParameterCode(BaseColor),
-		*GetParameterCode(Reflectivity),
+		*GetParameterCode(Specular),
 		*GetParameterCode(Metallic),
 		*GetParameterCode(RoughnessX),
 		*GetParameterCode(RoughnessY),
@@ -8541,6 +8544,7 @@ int32 FHLSLMaterialTranslator::StrataSlabBSDF(int32 BaseColor, int32 Reflectivit
 		*GetParameterCode(SSSDMFPRadius),
 		*GetParameterCode(SSSRadiusScale),
 		*GetParameterCode(EmissiveColor),
+		*GetParameterCode(Haziness),
 		SharedNormalIndex,
 		*GetParameterCode(Normal),
 		*GetParameterCode(Tangent)
diff --git a/Engine/Source/Runtime/Engine/Private/Materials/HLSLMaterialTranslator.h b/Engine/Source/Runtime/Engine/Private/Materials/HLSLMaterialTranslator.h
--- a/Engine/Source/Runtime/Engine/Private/Materials/HLSLMaterialTranslator.h
+++ b/Engine/Source/Runtime/Engine/Private/Materials/HLSLMaterialTranslator.h
@@ -920,12 +920,18 @@ protected:
 
 	// Strata
 	virtual int32 StrataCreateAndRegisterNullMaterial() override;
-	virtual int32 StrataSlabBSDF(int32 BaseColor, int32 Reflectivity, int32 Metallic, int32 RoughnessX, int32 RoughnessY, int32 SSSProfileId, int32 SSSDMFPAlbedo, int32 SSSDMFPRadius, int32 EmissiveColor, int32 SSSRadiusScale, int32 Normal, int32 Tangent, uint8 SharedNormalIndex) override;
+	virtual int32 StrataSlabBSDF(
+		int32 BaseColor, int32 Specular, int32 Metallic, int32 RoughnessX, int32 RoughnessY,
+		int32 SSSProfileId, int32 SSSDMFPAlbedo, int32 SSSDMFPRadius, int32 SSSRadiusScale, 
+		int32 EmissiveColor, int32 Haziness, int32 Normal, int32 Tangent, uint8 SharedNormalIndex) override;
 	virtual int32 StrataSheenBSDF(int32 BaseColor, int32 Roughness, int32 Normal, uint8 SharedNormalIndex) override;
 	virtual int32 StrataVolumetricFogCloudBSDF(int32 Albedo, int32 Extinction, int32 EmissiveColor, int32 AmbientOcclusion) override;
 	virtual int32 StrataUnlitBSDF(int32 EmissiveColor, int32 TransmittanceColor) override;
 	virtual int32 StrataHairBSDF(int32 BaseColor, int32 Scatter, int32 Specular, int32 Roughness, int32 Backlit, int32 EmissiveColor, int32 Tangent, uint8 SharedNormalIndex) override;
-	virtual int32 StrataSingleLayerWaterBSDF(int32 BaseColor, int32 Metallic, int32 Specular, int32 Roughness, int32 EmissiveColor, int32 TopMaterialOpacity, int32 WaterAlbedo, int32 WaterExtinction, int32 WaterPhaseG, int32 ColorScaleBehindWater, int32 Normal, uint8 SharedNormalIndex) override;
+	virtual int32 StrataSingleLayerWaterBSDF(
+		int32 BaseColor, int32 Metallic, int32 Specular, int32 Roughness, 
+		int32 EmissiveColor, int32 TopMaterialOpacity, int32 WaterAlbedo, int32 WaterExtinction, int32 WaterPhaseG, 
+		int32 ColorScaleBehindWater, int32 Normal, uint8 SharedNormalIndex) override;
 	virtual int32 StrataHorizontalMixing(int32 Foreground, int32 Background, int32 Mix) override;
 	virtual int32 StrataVerticalLayering(int32 Top, int32 Base) override;
 	virtual int32 StrataAdd(int32 A, int32 B) override;
diff --git a/Engine/Source/Runtime/Engine/Private/Materials/MaterialExpressions.cpp b/Engine/Source/Runtime/Engine/Private/Materials/MaterialExpressions.cpp
--- a/Engine/Source/Runtime/Engine/Private/Materials/MaterialExpressions.cpp
+++ b/Engine/Source/Runtime/Engine/Private/Materials/MaterialExpressions.cpp
@@ -19716,6 +19716,7 @@ int32 UMaterialExpressionStrataSlabBSDF::Compile(class FMaterialCompiler* Compil
 		CompileWithDefaultFloat1(Compiler, DiffuseMeanFreePathRadius, 0.0f),
 		CompileWithDefaultFloat1(Compiler, SSSRadiusScale, 1.0f),
 		CompileWithDefaultFloat3(Compiler, EmissiveColor, 0.0f, 0.0f, 0.0f),
+		CompileWithDefaultFloat1(Compiler, Haziness, 0.0f),
 		NormalCodeChunk,
 		TangentCodeChunk,
 		SharedNormalIndex);
@@ -19771,6 +19772,9 @@ uint32 UMaterialExpressionStrataSlabBSDF::GetInputType(int32 InputIndex)
 	case 10:
 		return MCT_Float3; // Emissive Color
 		break;
+	case 11:
+		return MCT_Float1; // Haziness
+		break;
 	}
 
 	check(false);
@@ -19823,6 +19827,10 @@ FName UMaterialExpressionStrataSlabBSDF::GetInputName(int32 InputIndex) const
 	{
 		return TEXT("Emissive Color");
 	}
+	else if (InputIndex == 11)
+	{
+		return TEXT("Haziness");
+	}
 	return TEXT("Unknown");
 }
 
diff --git a/Engine/Source/Runtime/Engine/Public/MaterialCompiler.h b/Engine/Source/Runtime/Engine/Public/MaterialCompiler.h
--- a/Engine/Source/Runtime/Engine/Public/MaterialCompiler.h
+++ b/Engine/Source/Runtime/Engine/Public/MaterialCompiler.h
@@ -392,12 +392,17 @@ public:
 
 	// Strata
 	virtual int32 StrataCreateAndRegisterNullMaterial() = 0;
-	virtual int32 StrataSlabBSDF(int32 BaseColor, int32 Reflectivity, int32 Metallic, int32 RoughnessX, int32 RoughnessY, int32 SSSProfileId, int32 SSSDMFPAlbedo, int32 SSSDMFPRadius, int32 SSSRadiusScale, int32 EmissiveColor, int32 Normal, int32 Tangent, uint8 SharedNormalIndex) = 0;
+	virtual int32 StrataSlabBSDF(
+		int32 BaseColor, int32 Specular, int32 Metallic, int32 RoughnessX, int32 RoughnessY, 
+		int32 SSSProfileId, int32 SSSDMFPAlbedo, int32 SSSDMFPRadius, int32 SSSRadiusScale, 
+		int32 EmissiveColor, int32 Haziness, int32 Normal, int32 Tangent, uint8 SharedNormalIndex) = 0;
 	virtual int32 StrataSheenBSDF(int32 BaseColor, int32 Roughness, int32 Normal, uint8 SharedNormalIndex) = 0;
 	virtual int32 StrataVolumetricFogCloudBSDF(int32 Albedo, int32 Extinction, int32 EmissiveColor, int32 AmbientOcclusion) = 0;
 	virtual int32 StrataUnlitBSDF(int32 EmissiveColor, int32 TransmittanceColor) = 0;
 	virtual int32 StrataHairBSDF(int32 BaseColor, int32 Scatter, int32 Specular, int32 Roughness, int32 Backlit, int32 EmissiveColor, int32 Tangent, uint8 SharedNormalIndex) = 0;
-	virtual int32 StrataSingleLayerWaterBSDF(int32 BaseColor, int32 Metallic, int32 Specular, int32 Roughness, int32 EmissiveColor, int32 TopMaterialOpacity, int32 WaterAlbedo, int32 WaterExtinction, int32 WaterPhaseG, int32 ColorScaleBehindWater, int32 Normal, uint8 SharedNormalIndex) = 0;
+	virtual int32 StrataSingleLayerWaterBSDF(
+		int32 BaseColor, int32 Metallic, int32 Specular, int32 Roughness, int32 EmissiveColor, int32 TopMaterialOpacity, 
+		int32 WaterAlbedo, int32 WaterExtinction, int32 WaterPhaseG, int32 ColorScaleBehindWater, int32 Normal, uint8 SharedNormalIndex) = 0;
 	virtual int32 StrataHorizontalMixing(int32 Foreground, int32 Background, int32 Mix) = 0;
 	virtual int32 StrataVerticalLayering(int32 Top, int32 Base) = 0;
 	virtual int32 StrataAdd(int32 A, int32 B) = 0;
@@ -806,9 +811,15 @@ public:
 		return Compiler->StrataCreateAndRegisterNullMaterial();
 	}
 
-	virtual int32 StrataSlabBSDF(int32 BaseColor, int32 Reflectivity, int32 Metallic, int32 RoughnessX, int32 RoughnessY, int32 SSSProfileId, int32 SSSDMFPAlbedo, int32 SSSDMFPRadius, int32 SSSRadiusScale, int32 EmissiveColor, int32 Normal, int32 Tangent, uint8 SharedNormalIndex) override
+	virtual int32 StrataSlabBSDF(
+		int32 BaseColor, int32 Specular, int32 Metallic, int32 RoughnessX, int32 RoughnessY,
+		int32 SSSProfileId, int32 SSSDMFPAlbedo, int32 SSSDMFPRadius, int32 SSSRadiusScale, 
+		int32 EmissiveColor, int32 Haziness, int32 Normal, int32 Tangent, uint8 SharedNormalIndex) override
 	{
-		return Compiler->StrataSlabBSDF(BaseColor, Reflectivity, Metallic, RoughnessX, RoughnessY, SSSProfileId, SSSDMFPAlbedo, SSSDMFPRadius, SSSRadiusScale, EmissiveColor, Normal, Tangent, SharedNormalIndex);
+		return Compiler->StrataSlabBSDF(
+			BaseColor, Specular, Metallic, RoughnessX, RoughnessY,
+			SSSProfileId, SSSDMFPAlbedo, SSSDMFPRadius, SSSRadiusScale, 
+			EmissiveColor, Haziness, Normal, Tangent, SharedNormalIndex);
 	}
 
 	virtual int32 StrataSheenBSDF(int32 BaseColor, int32 Roughness, int32 Normal, uint8 SharedNormalIndex) override
@@ -831,9 +842,13 @@ public:
 		return Compiler->StrataHairBSDF(BaseColor, Scatter, Specular, Roughness, Backlit, EmissiveColor, Tangent, SharedNormalIndex);
 	}
 
-	virtual int32 StrataSingleLayerWaterBSDF(int32 BaseColor, int32 Metallic, int32 Specular, int32 Roughness, int32 EmissiveColor, int32 TopMaterialOpacity, int32 WaterAlbedo, int32 WaterExtinction, int32 WaterPhaseG, int32 ColorScaleBehindWater, int32 Normal, uint8 SharedNormalIndex) override
+	virtual int32 StrataSingleLayerWaterBSDF(
+		int32 BaseColor, int32 Metallic, int32 Specular, int32 Roughness, int32 EmissiveColor, int32 TopMaterialOpacity, 
+		int32 WaterAlbedo, int32 WaterExtinction, int32 WaterPhaseG, int32 ColorScaleBehindWater, int32 Normal, uint8 SharedNormalIndex) override
 	{
-		return Compiler->StrataSingleLayerWaterBSDF(BaseColor, Metallic, Specular, Roughness, EmissiveColor, TopMaterialOpacity, WaterAlbedo, WaterExtinction, WaterPhaseG, ColorScaleBehindWater, Normal, SharedNormalIndex);
+		return Compiler->StrataSingleLayerWaterBSDF(
+			BaseColor, Metallic, Specular, Roughness, EmissiveColor, TopMaterialOpacity, 
+			WaterAlbedo, WaterExtinction, WaterPhaseG, ColorScaleBehindWater, Normal, SharedNormalIndex);
 	}
 
 	virtual int32 StrataHorizontalMixing(int32 Foreground, int32 Background, int32 Mix) override
