Fix for WaveActiveBallot and multiple platforms.

diff --git a/Engine/Shaders/Private/VirtualShadowMaps/ProjectionSpot.ush b/Engine/Shaders/Private/VirtualShadowMaps/ProjectionSpot.ush
--- a/Engine/Shaders/Private/VirtualShadowMaps/ProjectionSpot.ush
+++ b/Engine/Shaders/Private/VirtualShadowMaps/ProjectionSpot.ush
@@ -219,7 +219,7 @@ float TraceSpotLight(
 				}
 
 				// All lanes miss or hit something close enough to not be noisy.
-				bool bAllLanesMiss = all( WaveActiveBallot( bHit ).xy == 0 );
+				bool bAllLanesMiss = WaveActiveAllTrue( !bHit );
 				if( bAllLanesMiss )
 				{
 					RayMissFactor *= SampleCount;
@@ -229,7 +229,7 @@ float TraceSpotLight(
 			else if( i == 1 )
 			{
 				// After 2 iterations and all have hit, must be in umbra.
-				bool bAllLanesHit = all( WaveActiveBallot( RayMissFactor > 0 ).xy == 0 );
+				bool bAllLanesHit = WaveActiveAllTrue( RayMissFactor == 0 );
 				if( bAllLanesHit )
 				{
 					RayMissFactor *= 0.5 * SampleCount;
