Grab bag of virtual shadow fixes:
- Fix SMRT with non-nanite geometry (wasn't sampling the second map)
- Fix toggling distance field shadow cascade on when VSMs are enabled (trying to create a distance field complete shadow map)
  - This is a bit of a questionable case, but at least does not crash now
- Remove dependence of clipmap max distance on shadow cascade distance and replace with a cvar
  - Will likely want to do something similar for far culling distance for non-nanite stuff, but TBD
  - May eventually want a parameter back on the light, but for now just having a cvar set high is good as there is little overhead with Nanite meshes
- Add some debug output for SMRT spot lights

diff --git a/Engine/Shaders/Private/VirtualShadowMaps/PageManagement.usf b/Engine/Shaders/Private/VirtualShadowMaps/PageManagement.usf
--- a/Engine/Shaders/Private/VirtualShadowMaps/PageManagement.usf
+++ b/Engine/Shaders/Private/VirtualShadowMaps/PageManagement.usf
@@ -547,7 +547,9 @@ void CreatePageMappings(uint2 ThreadId : SV_DispatchThreadID, uint2 GroupId : SV
 
 		// Don't kill the last mip. It may be forced on.
 		if( Level == VSM_MAX_MIP_LEVELS - 1 )
+		{
 			bIsAllCovered = false;
+		}
 		
 		CreatePageMapping( ShadowMapID, Level, vPage, PageOffset, Flag && !bIsAllCovered );
 		
diff --git a/Engine/Shaders/Private/VirtualShadowMaps/ProjectionCommon.ush b/Engine/Shaders/Private/VirtualShadowMaps/ProjectionCommon.ush
--- a/Engine/Shaders/Private/VirtualShadowMaps/ProjectionCommon.ush
+++ b/Engine/Shaders/Private/VirtualShadowMaps/ProjectionCommon.ush
@@ -84,6 +84,18 @@ FVirtualSmSample InitVirtualSmSample()
 	return Result;
 }
 
+float SampleVirtualSmPhysicalDepth(uint2 PhysicalTexelAddress)
+{
+#if ENABLE_NON_NANITE_VSM
+	return max(
+		asfloat(PhysicalPagePool.Load(int3(PhysicalTexelAddress, 0))),
+		PhysicalPagePoolHw.Load(int3(PhysicalTexelAddress, 0))
+	);
+#else //!ENABLE_NON_NANITE_VSM
+	return asfloat(PhysicalPagePool.Load(int3(PhysicalTexelAddress, 0)));
+#endif // ENABLE_NON_NANITE_VSM
+}
+
 FVirtualSmSample SampleVirtualSmLevel(int VirtualShadowMapId, float2 ShadowMapUV, uint VSMLevel)
 {
 	FVirtualSmSample Result = InitVirtualSmSample();
@@ -91,19 +103,10 @@ FVirtualSmSample SampleVirtualSmLevel(int VirtualShadowMapId, float2 ShadowMapUV
 	Result.VirtualTexelAddressFloat = ShadowMapUV * float(CalcLevelDimsTexels(VSMLevel));
 	Result.VirtualTexelAddress = uint2(Result.VirtualTexelAddressFloat);
 
-	// Page map to physical address
-	uint2 PhysicalTexelAddress = uint2(0, 0);
 	bool PagePresent = VirtualToPhysicalTexel(VirtualShadowMapId, VSMLevel, Result.VirtualTexelAddress, Result.PhysicalTexelAddress);
 	if (PagePresent)
 	{
-#if ENABLE_NON_NANITE_VSM
-		Result.Depth = max(
-			asfloat(PhysicalPagePool.Load(int3(Result.PhysicalTexelAddress, 0))),
-			PhysicalPagePoolHw.Load(int3(Result.PhysicalTexelAddress, 0))
-		);
-#else //!ENABLE_NON_NANITE_VSM
-		Result.Depth = asfloat(PhysicalPagePool.Load(int3(Result.PhysicalTexelAddress, 0)));
-#endif // ENABLE_NON_NANITE_VSM
+		Result.Depth = SampleVirtualSmPhysicalDepth(Result.PhysicalTexelAddress);
 		Result.MipLevel = VSMLevel;
 		Result.bValid = true;
 	}
@@ -122,7 +125,7 @@ FVirtualSmSample SampleVirtualSm(int VirtualShadowMapId, float2 ShadowMapUV)
 		Result.VirtualTexelAddress = Page.VirtualTexelAddress;
 		Result.VirtualTexelAddressFloat = Page.VirtualTexelAddressFloat;
 		Result.PhysicalTexelAddress = Page.PhysicalTexelAddress;
-		Result.Depth = asfloat(PhysicalPagePool.Load(int3(Result.PhysicalTexelAddress, 0)));
+		Result.Depth = SampleVirtualSmPhysicalDepth(Result.PhysicalTexelAddress);
 		return Result;
 	}
 
diff --git a/Engine/Shaders/Private/VirtualShadowMaps/ProjectionDirectional.ush b/Engine/Shaders/Private/VirtualShadowMaps/ProjectionDirectional.ush
--- a/Engine/Shaders/Private/VirtualShadowMaps/ProjectionDirectional.ush
+++ b/Engine/Shaders/Private/VirtualShadowMaps/ProjectionDirectional.ush
@@ -107,7 +107,7 @@ bool ShadowRayCastDirectional(
 
 		if (bValid)
 		{
-			const float SampleDepth = ToBaseDepthScale * asfloat(PhysicalPagePool.Load(int3(PhysicalTexelAddress, 0))) + ToBaseDepthBias;
+			const float SampleDepth = ToBaseDepthScale * SampleVirtualSmPhysicalDepth(PhysicalTexelAddress) + ToBaseDepthBias;
 			const float ReferenceDepth = SampleTexelsZ.z;
 			if (DepthHistory < 0)
 			{
diff --git a/Engine/Shaders/Private/VirtualShadowMaps/ProjectionSpot.ush b/Engine/Shaders/Private/VirtualShadowMaps/ProjectionSpot.ush
--- a/Engine/Shaders/Private/VirtualShadowMaps/ProjectionSpot.ush
+++ b/Engine/Shaders/Private/VirtualShadowMaps/ProjectionSpot.ush
@@ -31,6 +31,8 @@ bool ShadowRayCastSpotlight(
 	float3 RayEnd,
 	int NumSteps,
 	float StepOffset,
+	int DebugOutputType,
+	inout float3 DebugOutput,
 	inout float OccluderDistance)
 {
 	FVirtualShadowMapProjectionShaderData BaseProjectionData = GetVirtualShadowMapProjectionData(VirtualShadowMapId);
@@ -79,8 +81,8 @@ bool ShadowRayCastSpotlight(
 	const float CompareScale = RayStepUVz.z * 2.0 / Pow2(NumSteps);
 	const float CompareBias = -CompareScale * (NumSteps - StepOffset + 1.0f);
 
-	float FirstHitTime = 2;
 	float DepthHistory = -1;
+	int SampleMissCount = 0;
 
 	bool bValidHit = false;
 	for (int i = 0; i <= NumSteps; i++)
@@ -98,7 +100,6 @@ bool ShadowRayCastSpotlight(
 			{
 				// First valid sample we've seen. Do a regular depth compare.
 				DepthHistory = SampleDepth;
-				FirstHitTime = SampleDepth > ReferenceDepth ? SampleTime : 2;
 				if( SampleDepth > ReferenceDepth )
 				{
 					float ReceiverDistance = length( RayStart + BaseProjectionData.ShadowPreViewTranslation );
@@ -124,7 +125,6 @@ bool ShadowRayCastSpotlight(
 				bool bHit = abs(DepthDiff + CompareTolerance) < CompareTolerance;
 				if (bHit)
 				{
-					FirstHitTime = SampleTime;
 					OccluderDistance = length(RayEnd - RayStart) * SampleTime;
 
 					bValidHit = true;
@@ -132,6 +132,16 @@ bool ShadowRayCastSpotlight(
 				}
 			}
 		}
+		else
+		{
+			++SampleMissCount;
+		}
+	}
+
+	float OutputDebugSignal = bValidHit ? 0.0f : 1.0f;
+	if (DebugOutputType == 1)	// Invalid page
+	{
+		DebugOutput = lerp(OutputDebugSignal.xxx, lerp(float3(0, 1, 0), float3(1, 0, 1), saturate(SampleMissCount / 4.0f)), 0.75f);
 	}
 
 	return bValidHit;
@@ -163,6 +173,7 @@ float TraceSpotLight(
 	bool bBackFace = dot(GBufferData.WorldNormal, normalize(ToLight)) < -SinAlpha;
 
 	float RayMissFactor = 0;
+	uint SampleCount = 0;
 	if (!bBackFace)
 	{
 		uint2 SobolBase = SobolPixel(SvPosition.xy);
@@ -170,15 +181,13 @@ float TraceSpotLight(
 
 		float StepOffset = InterleavedGradientNoise(SvPosition.xy, View.StateFrameIndexMod8);
 
-#if 1
-		uint SampleCount = RayCount;
-
-		for( uint i = 0; i < SampleCount; i++ )
+		uint i = 0;
+		for( ; i < uint(RayCount); i++ )
 		{
-#if 1
 			RandomSequence RandSequence;
-			RandomSequence_Initialize( RandSequence, uint2( SvPosition.xy ), i, View.StateFrameIndexMod8, SampleCount );
+			RandomSequence_Initialize( RandSequence, uint2( SvPosition.xy ), i, View.StateFrameIndexMod8, RayCount );
 
+#if 1
 			float2 E = RandomSequence_GenerateSample2D(RandSequence);
 #else
 			// NOTE: 1-X to avoid 0/0 issues with UniformSampleDiskConcentric
@@ -204,6 +213,8 @@ float TraceSpotLight(
 				RayEnd,
 				SamplesPerRay,
 				StepOffset,
+				DebugOutputType,
+				DebugOutput,
 				InOutHitDistance);
 
 			RayMissFactor += bHit ? 0 : 1;
@@ -222,7 +233,6 @@ float TraceSpotLight(
 				bool bAllLanesMiss = WaveActiveAllTrue( !bHit );
 				if( bAllLanesMiss )
 				{
-					RayMissFactor *= SampleCount;
 					break;
 				}
 			}
@@ -232,50 +242,20 @@ float TraceSpotLight(
 				bool bAllLanesHit = WaveActiveAllTrue( RayMissFactor == 0 );
 				if( bAllLanesHit )
 				{
-					RayMissFactor *= 0.5 * SampleCount;
 					break;
 				}
 			}
 #endif
 		}
-		RayMissFactor /= SampleCount;
-#else
-
-#if 0
-		RandomSequence RandSequence;
-		RandomSequence_Initialize( RandSequence, uint2( SvPosition.xy ), 0, View.StateFrameIndexMod8, 1 );
-
-		float2 E = RandomSequence_GenerateSample2D(RandSequence);
-#else
-		// TODO: Loop over number of rays RayCount
-		float2 E = float2(SobolIndex(SobolFrame, (View.Random & 7) << 3)) / 0x10000;
-#endif
-
-		// NOTE: Avoid 0/0 issues with UniformSampleDiskConcentric
-		// TODO: RandomSequence
-		float2 LightUV = UniformSampleDiskConcentricApprox(E) * Light.SourceRadius;
-
-		float3 RayEnd =
-			Light.Position +
-			LightUV.x * Light.Tangent +
-			LightUV.y * cross(Light.Direction, Light.Tangent);
-
-		float3 Dir = normalize(RayEnd - WorldPosition);
-
-		float3 RayStart = WorldPosition + Dir * RayStartOffset;
-		RayEnd = lerp(RayStart, RayEnd, 0.75);	// TODO!
-
-		bool bHit = ShadowRayCastSpotlight(
-			VirtualShadowMapId,
-			RayStart,
-			RayEnd,
-			SamplesPerRay,
-			StepOffset,
-			InOutHitDistance);
+		SampleCount = min(i + 1, RayCount);		// break vs regular for loop exit
+		RayMissFactor /= float(SampleCount);
+	}
 
-		RayMissFactor += bHit ? 0 : 1;
-#endif
+	if (DebugOutputType == 2)	// Sample count
+	{
+		float Ratio = float(SampleCount) / float(RayCount);
+		DebugOutput = lerp(RayMissFactor.xxx, lerp(float3(0, 1, 0), float3(1, 0, 0), Ratio), 1.0f);
 	}
 
 	return RayMissFactor;
-}
\ No newline at end of file
+}
diff --git a/Engine/Source/Runtime/Renderer/Private/ShadowSetup.cpp b/Engine/Source/Runtime/Renderer/Private/ShadowSetup.cpp
--- a/Engine/Source/Runtime/Renderer/Private/ShadowSetup.cpp
+++ b/Engine/Source/Runtime/Renderer/Private/ShadowSetup.cpp
@@ -4348,7 +4348,6 @@ void FSceneRenderer::AddViewDependentWholeSceneShadowsForView(
 
 			static_assert(INDEX_NONE == -1, "INDEX_NONE != -1!");
 
-			float MaxShadowCascadeDistance = 0.0f;
 #if ENABLE_NON_NANITE_VSM
 			float MaxNonFarCascadeDistance = 0.0f;
 #endif // ENABLE_NON_NANITE_VSM
@@ -4368,7 +4367,6 @@ void FSceneRenderer::AddViewDependentWholeSceneShadowsForView(
 
 				if (LightSceneInfo.Proxy->GetViewDependentWholeSceneProjectedShadowInitializer(View, LocalIndex, LightSceneInfo.IsPrecomputedLightingValid(), ProjectedShadowInitializer))
 				{
-					MaxShadowCascadeDistance = FMath::Max(MaxShadowCascadeDistance, ProjectedShadowInitializer.CascadeSettings.SplitFar);
 #if ENABLE_NON_NANITE_VSM
 					if (!ProjectedShadowInitializer.CascadeSettings.bFarShadowCascade)
 					{
@@ -4417,7 +4415,7 @@ void FSceneRenderer::AddViewDependentWholeSceneShadowsForView(
 					VisibleLightInfo.AllProjectedShadows.Add(ProjectedShadowInfo);
 					ShadowInfos.Add(ProjectedShadowInfo);
 
-					if (bNeedsVirtualShadowMap)
+					if (bNeedsVirtualShadowMap && !ProjectedShadowInitializer.bRayTracedDistanceField)
 					{
 						// If we have a virtual shadow map, disable nanite rendering into the regular shadow map or else we'd get double-shadowing
 						ProjectedShadowInfo->bNaniteGeometry = false;
@@ -4457,7 +4455,7 @@ void FSceneRenderer::AddViewDependentWholeSceneShadowsForView(
 				}
 			}
 
-			if (bNeedsVirtualShadowMap && MaxShadowCascadeDistance > 0.0f)
+			if (bNeedsVirtualShadowMap)
 			{
 				FMatrix WorldToLight = FInverseRotationMatrix(LightSceneInfo.Proxy->GetDirection().GetSafeNormal().Rotation());
 
@@ -4467,8 +4465,7 @@ void FSceneRenderer::AddViewDependentWholeSceneShadowsForView(
 					LightSceneInfo,
 					WorldToLight,
 					View.ViewMatrices,
-					View.ViewRect.Size(),
-					MaxShadowCascadeDistance
+					View.ViewRect.Size()
 				));
 
 				// TODO: Not clear we need both of these in this path, but keep it consistent for now
diff --git a/Engine/Source/Runtime/Renderer/Private/VirtualShadowMaps/VirtualShadowMapClipmap.cpp b/Engine/Source/Runtime/Renderer/Private/VirtualShadowMaps/VirtualShadowMapClipmap.cpp
--- a/Engine/Source/Runtime/Renderer/Private/VirtualShadowMaps/VirtualShadowMapClipmap.cpp
+++ b/Engine/Source/Runtime/Renderer/Private/VirtualShadowMaps/VirtualShadowMapClipmap.cpp
@@ -17,12 +17,18 @@ static TAutoConsoleVariable<float> CVarVirtualShadowMapClipmapResolutionLodBias(
 	TEXT( "" ),
 	ECVF_RenderThreadSafe
 );
-static TAutoConsoleVariable<float> CVarVirtualShadowMapClipmapFirstLevel(
+static TAutoConsoleVariable<int32> CVarVirtualShadowMapClipmapFirstLevel(
 	TEXT( "r.Shadow.v.Clipmap.FirstLevel" ),
-	8.0f,
+	8,
 	TEXT( "First level of the virtual clipmap. Lower values allow higher resolution shadows closer to the camera." ),
 	ECVF_RenderThreadSafe
 );
+static TAutoConsoleVariable<float> CVarVirtualShadowMapClipmapMaxRadius(
+	TEXT( "r.Shadow.v.Clipmap.MaxRadius" ),
+	1000000.0f,
+	TEXT( "Maximum distance the clipmap covers. Determines the number of clipmap levels. " ),
+	ECVF_RenderThreadSafe
+);
 
 // "Virtual" clipmap level to clipmap radius
 static float GetLevelRadius(int32 Level)
@@ -38,11 +44,8 @@ FVirtualShadowMapClipmap::FVirtualShadowMapClipmap(
 	const FLightSceneInfo& InLightSceneInfo,
 	const FMatrix& WorldToLightRotationMatrix,
 	const FViewMatrices& CameraViewMatrices,
-	FIntPoint CameraViewRectSize,
-
-	float InMaxRadius)
-	: LightSceneInfo(InLightSceneInfo), 
-	  MaxRadius(InMaxRadius)
+	FIntPoint CameraViewRectSize)
+	: LightSceneInfo(InLightSceneInfo)
 {
 	check(WorldToLightRotationMatrix.GetOrigin() == FVector(0, 0, 0));	// Should not contain translation or scaling
 
@@ -67,7 +70,8 @@ FVirtualShadowMapClipmap::FVirtualShadowMapClipmap(
 	// For now we adjust resolution by just biasing the page we look up in. This is wasteful in terms of page table vs.
 	// just resizing the virtual shadow maps for each clipmap, but convenient for now. This means we need to additionally bias
 	// which levels are present.
-	FirstLevel = FMath::FloorToInt(CVarVirtualShadowMapClipmapFirstLevel.GetValueOnRenderThread());
+	FirstLevel = CVarVirtualShadowMapClipmapFirstLevel.GetValueOnRenderThread();
+	MaxRadius  = CVarVirtualShadowMapClipmapMaxRadius.GetValueOnRenderThread();
 	int32 LastLevel = FMath::FloorToInt(FMath::Log2(MaxRadius) + ResolutionLodBias);
 	LastLevel = FMath::Max(FirstLevel, LastLevel);
 	int32 LevelCount = LastLevel - FirstLevel + 1;
diff --git a/Engine/Source/Runtime/Renderer/Private/VirtualShadowMaps/VirtualShadowMapClipmap.h b/Engine/Source/Runtime/Renderer/Private/VirtualShadowMaps/VirtualShadowMapClipmap.h
--- a/Engine/Source/Runtime/Renderer/Private/VirtualShadowMaps/VirtualShadowMapClipmap.h
+++ b/Engine/Source/Runtime/Renderer/Private/VirtualShadowMaps/VirtualShadowMapClipmap.h
@@ -23,8 +23,7 @@ public:
 		const FLightSceneInfo& InLightSceneInfo,
 		const FMatrix& WorldToLightRotationMatrix,
 		const FViewMatrices& CameraViewMatrices,
-		FIntPoint CameraViewRectSize,
-		float MaxRadius		// Maximum radius the clipmap must cover from the center point; used to compute level count
+		FIntPoint CameraViewRectSize
 	);
 
 	FViewMatrices GetViewMatrices(int32 ClipmapIndex) const;
