Fixed shader compile error due to ambiguous call to min()

diff --git a/Engine/Shaders/Private/VirtualShadowMaps/ProjectionSpot.ush b/Engine/Shaders/Private/VirtualShadowMaps/ProjectionSpot.ush
--- a/Engine/Shaders/Private/VirtualShadowMaps/ProjectionSpot.ush
+++ b/Engine/Shaders/Private/VirtualShadowMaps/ProjectionSpot.ush
@@ -247,7 +247,7 @@ float TraceSpotLight(
 			}
 #endif
 		}
-		SampleCount = min(i + 1, RayCount);		// break vs regular for loop exit
+		SampleCount = min(i + 1, uint(RayCount));		// break vs regular for loop exit
 		RayMissFactor /= float(SampleCount);
 	}
 
