Strata DMFP clean up:
- DMFP directly chromatic is now F11_11_10 so it allows for a higher range.
- DMFPRadius is remove.
- DMFPScale is kept around for now

Added Max11BitsFloat

diff --git a/Engine/Shaders/Private/BasePassPixelShader.usf b/Engine/Shaders/Private/BasePassPixelShader.usf
--- a/Engine/Shaders/Private/BasePassPixelShader.usf
+++ b/Engine/Shaders/Private/BasePassPixelShader.usf
@@ -1560,7 +1560,7 @@ void FPixelShaderInOut_MainPS(
 	Color += EmissiveLuminance;
 
 	// Since we don't carry SSS on/off through the ShadingModel value (DefaultLit/SubSurfaceProfile), we update subsurface toggle, based on the MFP radius value
-	bSubsurfaceEnable = GBuffer.CustomData.a > 0;
+	bSubsurfaceEnable = any(GBuffer.CustomData > 0);
 #endif // STRATA_OPAQUE_ENABLED
 
 #if STRATA_TRANSLUCENT_ENABLED
diff --git a/Engine/Shaders/Private/Common.ush b/Engine/Shaders/Private/Common.ush
--- a/Engine/Shaders/Private/Common.ush
+++ b/Engine/Shaders/Private/Common.ush
@@ -88,7 +88,9 @@
 
 const static MaterialFloat PI = 3.1415926535897932f;
 const static float MaxHalfFloat = 65504.0f;
+const static float Max11BitsFloat = 65024.0f;
 const static float Max10BitsFloat = 64512.0f;
+const static float3 Max111110BitsFloat3 = float3(Max11BitsFloat, Max11BitsFloat, Max10BitsFloat);
 
 // Control MIP level used for material texture fetches. By default only raytracing 
 // shaders (i.e., !PIXELSHADER) use manual MIP level selection. A material shader 
diff --git a/Engine/Shaders/Private/ShadingCommon.ush b/Engine/Shaders/Private/ShadingCommon.ush
--- a/Engine/Shaders/Private/ShadingCommon.ush
+++ b/Engine/Shaders/Private/ShadingCommon.ush
@@ -40,10 +40,6 @@
 #define SSS_PROFILE_ID_INVALID  256
 #define SSS_PROFILE_ID_PERPIXEL 512
 
-// Max radius for per-pixel sub-surface scattering (non-profile version) in cm
-// Default value give a precision of 0.1mm
-#define SSS_PROFILE_PERPIXEL_MAXRADIUS 25.5f
-
 // for debugging and to visualize
 float3 GetShadingModelColor(uint ShadingModelID)
 {
diff --git a/Engine/Shaders/Private/Strata/Strata.ush b/Engine/Shaders/Private/Strata/Strata.ush
--- a/Engine/Shaders/Private/Strata/Strata.ush
+++ b/Engine/Shaders/Private/Strata/Strata.ush
@@ -128,8 +128,7 @@ struct FStrataBSDF
 #define SLAB_SSSPROFILEID(X)			X.VGPRs[2].x
 #define SLAB_SSSPROFILERADIUSSCALE(X)	X.VGPRs[2].y
 // Or explicit DMFP
-#define SLAB_SSSDMFPALBEDO(X)			X.VGPRs[2].xyz
-#define SLAB_SSSDMFPRADIUS(X)			X.VGPRs[2].w
+#define SLAB_SSSDMFP(X)					X.VGPRs[2].xyz
 
 #define SHEEN_BASECOLOR(X)				X.VGPRs[0].xyz
 #define SHEEN_ROUGHNESS(X)				X.VGPRs[0].w
@@ -223,8 +222,7 @@ struct FStrataSubsurfaceData
 	float  ProfilePackedId;		// Integer ID packed in 8bits as float
 	float  ProfileRadiusScale;
 
-	float3 DMFPAlbedo;
-	float  DMFPRadius;
+	float3 DMFP;				// When the profile is not used, a per pixel DMFP is used
 };
 
 void FinalizeTopLayerData(inout FStrataTopLayerData Out)
@@ -316,7 +314,7 @@ FStrataData GetInitialisedStrataData()
 
 FStrataData GetStrataSlabBSDF(
 	float3 BaseColor, float Specular, float Metallic, float2 Roughness,
-	float SSSProfileID, float3 SSSDMFPAlbedo, float SSSDMFPRadius, float SSSRadiusScale,
+	float SSSProfileID, float3 SSSDMFP, float SSSDMFPScale,
 	float3 Emissive, float Haziness, 
 	uint SharedNormalIndex, inout uint SharedNormalTypes)
 {
@@ -332,7 +330,7 @@ FStrataData GetStrataSlabBSDF(
 
 	// STRATA_TODO apply SSS as a function of pixel footprint w.r.t. MFP
 	const bool bHasSSSProfile = SSSProfileID  > 0.0f;	// When no profile is provide, we use Burley. STRATA_TODO we should clarify as burley is also an option of the profile.
-	const bool bHasSSS = SSSDMFPRadius > 0 || bHasSSSProfile;
+	const bool bHasSSS = any(SSSDMFP > 0) || bHasSSSProfile;
 
 	StrataData.Layers[0].BSDFCount++;
 	BSDF_SETTYPE			(StrataData.Layers[0].BSDFs[0],		STRATA_BSDF_TYPE_SLAB);
@@ -353,12 +351,11 @@ FStrataData GetStrataSlabBSDF(
 	if (bHasSSSProfile)
 	{
 		SLAB_SSSPROFILEID	(StrataData.Layers[0].BSDFs[0])		= SSSProfileID;
-		SLAB_SSSPROFILERADIUSSCALE(StrataData.Layers[0].BSDFs[0])= saturate(SSSRadiusScale);
+		SLAB_SSSPROFILERADIUSSCALE(StrataData.Layers[0].BSDFs[0])= saturate(SSSDMFPScale);
 	}
 	else
 	{
-		SLAB_SSSDMFPALBEDO	(StrataData.Layers[0].BSDFs[0])		= saturate(SSSDMFPAlbedo);
-		SLAB_SSSDMFPRADIUS	(StrataData.Layers[0].BSDFs[0])		= clamp(SSSDMFPRadius * saturate(SSSRadiusScale), 0.f, SSS_PROFILE_PERPIXEL_MAXRADIUS);
+		SLAB_SSSDMFP	(StrataData.Layers[0].BSDFs[0])			= SSSDMFP * saturate(SSSDMFPScale);
 	}
 
 	StrataData.Layers[0].BSDFs[0].Weight						= 1.0f;
@@ -645,8 +642,7 @@ void SanitizeStrataSlab(inout FStrataBSDF BSDF)
 	SLAB_ROUGHNESSY(BSDF) = saturate(SLAB_ROUGHNESSY(BSDF));
 	SLAB_HAZINESS(BSDF) = saturate(SLAB_HAZINESS(BSDF));
 	SLAB_SSSPROFILERADIUSSCALE(BSDF) = saturate(SLAB_SSSPROFILERADIUSSCALE(BSDF));
-	SLAB_SSSDMFPALBEDO(BSDF) = saturate(SLAB_SSSDMFPALBEDO(BSDF));
-	SLAB_SSSDMFPRADIUS(BSDF) = clamp(SLAB_SSSDMFPRADIUS(BSDF), 0.0f, SSS_PROFILE_PERPIXEL_MAXRADIUS);
+	SLAB_SSSDMFP(BSDF) = clamp(SLAB_SSSDMFP(BSDF), 0.0f, Max111110BitsFloat3);
 }
 
 void SanitizeStrataSheen(inout FStrataBSDF BSDF)
@@ -1109,8 +1105,7 @@ void PackStrataOut(
 					SSS.ProfilePackedId		= SLAB_SSSPROFILEID(BSDF);
 					SSS.ProfileRadiusScale	= SLAB_SSSPROFILERADIUSSCALE(BSDF);
 					SSS.BaseColor			= SLAB_BASECOLOR(BSDF);
-					SSS.DMFPAlbedo			= SLAB_SSSDMFPALBEDO(BSDF);
-					SSS.DMFPRadius			= SLAB_SSSDMFPRADIUS(BSDF);
+					SSS.DMFP				= SLAB_SSSDMFP(BSDF);
 				}
 
 				EmissiveColor += BSDF_GETEMISSIVE(BSDF) * Throughput;
@@ -1222,8 +1217,10 @@ void PackStrataOut(
 		else
 		{
 			GBuffer.BaseColor      = SSS.BaseColor;
-			GBuffer.CustomData.rgb = EncodeSubsurfaceColor(SSS.DMFPAlbedo);
-			GBuffer.CustomData.a   = saturate(SSS.DMFPRadius / SSS_PROFILE_PERPIXEL_MAXRADIUS);
+
+			// This packing should match what is done in SubsurfaceBurleyNormalized.ush
+			uint PackedDMFP11_11_10 = PackR11G11B10F(SSS.DMFP);
+			GBuffer.CustomData = UnpackRGBA8(PackedDMFP11_11_10);
 		}
 	}
 
@@ -1288,7 +1285,7 @@ void PackStrataOut(
 					}
 					else
 					{
-						STRATA_STORE_UINT1(PackColorLinearToGamma2AlphaLinear(float4(SLAB_SSSDMFPALBEDO(BSDF), SLAB_SSSDMFPRADIUS(BSDF) / SSS_PROFILE_PERPIXEL_MAXRADIUS)));
+						STRATA_STORE_UINT1(PackR11G11B10F(SLAB_SSSDMFP(BSDF)));
 					}
 				}
 				// 8 bytes
@@ -1401,8 +1398,7 @@ FStrataBSDF UnpackStrataBSDFIn(ByteAddressBuffer StrataBuffer, inout uint ByteOf
 			}
 			else
 			{
-				StrataLoad_ColorGamma2ToLinearAlphaLinear(StrataBuffer, ByteOffset, SLAB_SSSDMFPALBEDO(OutBSDF), SLAB_SSSDMFPRADIUS(OutBSDF));
-				SLAB_SSSDMFPRADIUS(OutBSDF) *= SSS_PROFILE_PERPIXEL_MAXRADIUS;
+				StrataLoad_R11G11B10F(StrataBuffer, ByteOffset, SLAB_SSSDMFP(OutBSDF));
 			}
 		}
 	}
diff --git a/Engine/Shaders/Private/Strata/StrataMaterialClassification.usf b/Engine/Shaders/Private/Strata/StrataMaterialClassification.usf
--- a/Engine/Shaders/Private/Strata/StrataMaterialClassification.usf
+++ b/Engine/Shaders/Private/Strata/StrataMaterialClassification.usf
@@ -73,8 +73,7 @@ void ComputePostProcessBSDFData(
 			SSS.ProfilePackedId		= SLAB_SSSPROFILEID(BSDF);
 			SSS.ProfileRadiusScale	= SLAB_SSSPROFILERADIUSSCALE(BSDF);
 			SSS.BaseColor			= SLAB_BASECOLOR(BSDF);
-			SSS.DMFPAlbedo			= SLAB_SSSDMFPALBEDO(BSDF);
-			SSS.DMFPRadius			= SLAB_SSSDMFPRADIUS(BSDF);
+			SSS.DMFP				= SLAB_SSSDMFP(BSDF);
 		}
 		break;
 	}
diff --git a/Engine/Shaders/Private/Strata/StrataVisualize.usf b/Engine/Shaders/Private/Strata/StrataVisualize.usf
--- a/Engine/Shaders/Private/Strata/StrataVisualize.usf
+++ b/Engine/Shaders/Private/Strata/StrataVisualize.usf
@@ -272,11 +272,8 @@ void VisualizeMaterialPS(
 					}
 					else
 					{
-						HORIZONTALTAB; PRINTPARAMCHAR(_S_); PRINTPARAMCHAR(_S_); PRINTPARAMCHAR(_S_); PRINTPARAMCHAR(_SPC_); PRINTPARAMCHAR(_A_); PRINTPARAMCHAR(_L_); PRINTPARAMCHAR(_SPC_);
-						HORIZONTALTAB; PRINTFLOAT3(SLAB_SSSDMFPALBEDO(BSDF));
-						NEXTLINE;
-						HORIZONTALTAB; PRINTPARAMCHAR(_S_); PRINTPARAMCHAR(_S_); PRINTPARAMCHAR(_S_); PRINTPARAMCHAR(_SPC_); PRINTPARAMCHAR(_R_); PRINTPARAMCHAR(_A_); PRINTPARAMCHAR(_SPC_);
-						HORIZONTALTAB; PRINTFLOAT1(SLAB_SSSDMFPRADIUS(BSDF));
+						HORIZONTALTAB; PRINTPARAMCHAR(_S_); PRINTPARAMCHAR(_S_); PRINTPARAMCHAR(_S_); PRINTPARAMCHAR(_D_); PRINTPARAMCHAR(_M_); PRINTPARAMCHAR(_F_); PRINTPARAMCHAR(_P_);
+						HORIZONTALTAB; PRINTFLOAT3(SLAB_SSSDMFP(BSDF));
 						NEXTLINE;
 					}
 				}
@@ -301,7 +298,6 @@ void VisualizeMaterialPS(
 					else
 					{
 						NEXTLINE;
-						NEXTLINE;
 					}
 				}
 #endif
diff --git a/Engine/Shaders/Private/SubsurfaceBurleyNormalized.ush b/Engine/Shaders/Private/SubsurfaceBurleyNormalized.ush
--- a/Engine/Shaders/Private/SubsurfaceBurleyNormalized.ush
+++ b/Engine/Shaders/Private/SubsurfaceBurleyNormalized.ush
@@ -10,6 +10,10 @@
 #include "DeferredShadingCommon.ush"
 #include "MonteCarlo.ush"
 
+#if STRATA_ENABLED
+#include "Strata/Strata.ush"
+#endif
+
 // Setup the max number of Burley samples
 #define BURLEY_NUM_SAMPLES	64
 #define BURLEY_INV_NUM_SAMPLES (1.0f/BURLEY_NUM_SAMPLES)
@@ -213,16 +217,18 @@ FBurleyParameter GetBurleyParameters(FGBufferData GBuffer)
 		}
 		else
 		{
-			// Special data encoding from strata packing
-			const float3 MFPAlebdo			= Pow2(GBuffer.CustomData.xyz);
-			const float  MFPRadius			= GBuffer.CustomData.a * SSS_PROFILE_PERPIXEL_MAXRADIUS * BURLEY_CM_2_MM * 5.0f; // STRATA_TODO why a scale of 5 is needed here to match between regular and Strata?
-			const float3 MFPChromaticRadius = MFPAlebdo * MFPRadius;
+			// Decode the DMFP stored as Float111110 in a uint over a RGBA8 unorm
+			uint PackedDMFP11_11_10 = PackRGBA8(GBuffer.CustomData);
+			float3 DMFP = UnpackR11G11B10F(PackedDMFP11_11_10);
+
+			// Conversion to match other SSS encoding 
+			DMFP *= BURLEY_CM_2_MM * 5.0f; // STRATA_TODO why a scale of 5 is needed here to match between regular and Strata?
 
-			const float MaxDMFP				= max3(MFPChromaticRadius.x, MFPChromaticRadius.y, MFPChromaticRadius.z);
-			const uint IndexOfMaxDmfp		= (MFPChromaticRadius.r == MaxDMFP) ? 0 : ((MFPChromaticRadius.g == MaxDMFP) ? 1 : 2);
+			const float MaxDMFP				= max3(DMFP.x, DMFP.y, DMFP.z);
+			const uint IndexOfMaxDmfp		= (DMFP.r == MaxDMFP) ? 0 : ((DMFP.g == MaxDMFP) ? 1 : 2);
 
 			Out.SurfaceAlbedo				= float4(GBuffer.BaseColor, GBuffer.BaseColor[IndexOfMaxDmfp]);
-			Out.DiffuseMeanFreePath			= float4(MFPChromaticRadius, MaxDMFP);
+			Out.DiffuseMeanFreePath			= float4(DMFP, MaxDMFP);
 			Out.WorldUnitScale				= DecodeWorldUnitScale(1.0f * ENC_WORLDUNITSCALE_IN_CM_TO_UNIT) * BURLEY_CM_2_MM;
 
 			// STRATA_TODO smooth transistion based on pixel footprint w.r.t. MFPRadius and pixel footprint
diff --git a/Engine/Source/Runtime/Engine/Classes/Materials/MaterialExpressionStrata.h b/Engine/Source/Runtime/Engine/Classes/Materials/MaterialExpressionStrata.h
--- a/Engine/Source/Runtime/Engine/Classes/Materials/MaterialExpressionStrata.h
+++ b/Engine/Source/Runtime/Engine/Classes/Materials/MaterialExpressionStrata.h
@@ -80,22 +80,16 @@ class UMaterialExpressionStrataSlabBSDF : public UMaterialExpressionStrataBSDF
 	FExpressionInput Tangent;
 
 	/**
-	 * Mean free path albedo. Only used when there is no sub-surface profile provided. (type = float3, unit = unitless)
+	 * Chromatic mean free path . Only used when there is not any sub-surface profile provided. (type = float3, unit = unitless)
 	 */
-	UPROPERTY(meta = (DisplayName = "DMFP Albedo"))
-	FExpressionInput DiffuseMeanFreePathAlbedo;
-
-	/**
-	 * Mean free path radius. Only used when there is no sub-surface profile provided. (type = float, unit = cm)
-	 */
-	UPROPERTY(meta = (DisplayName = "DMFP Radius"))
-	FExpressionInput DiffuseMeanFreePathRadius;
+	UPROPERTY()
+	FExpressionInput SSSDMFP;
 
 	/**
-	 * Scale the mean free path radius. Always used, when a subsurface profile is provided or not. (type = float, unitless, defaults to 1)
+	 * Scale the mean free path radius of the SSS profile according to a value between 0 and 1. Always used, when a subsurface profile is provided or not. (type = float, unitless, defaults to 1)
 	 */
-	UPROPERTY(meta = (DisplayName = "SSS Radius Scale"))
-	FExpressionInput SSSRadiusScale;
+	UPROPERTY()
+	FExpressionInput SSSDMFPScale;
 
 	/**
 	 * Emissive color on top of the surface (type = float3, unit = luminance, default = 0)
diff --git a/Engine/Source/Runtime/Engine/Private/Materials/HLSLMaterialTranslator.cpp b/Engine/Source/Runtime/Engine/Private/Materials/HLSLMaterialTranslator.cpp
--- a/Engine/Source/Runtime/Engine/Private/Materials/HLSLMaterialTranslator.cpp
+++ b/Engine/Source/Runtime/Engine/Private/Materials/HLSLMaterialTranslator.cpp
@@ -8529,20 +8529,19 @@ int32 FHLSLMaterialTranslator::StrataCreateAndRegisterNullMaterial()
 
 int32 FHLSLMaterialTranslator::StrataSlabBSDF(
 	int32 BaseColor, int32 Specular, int32 Metallic, int32 RoughnessX, int32 RoughnessY,
-	int32 SSSProfileId, int32 SSSDMFPAlbedo, int32 SSSDMFPRadius, int32 SSSRadiusScale, int32 EmissiveColor, 
+	int32 SSSProfileId, int32 SSSDMFP, int32 SSSDMFPScale, int32 EmissiveColor,
 	int32 Haziness, int32 Normal, int32 Tangent, uint8 SharedNormalIndex)
 {
 	return AddCodeChunk(
-		MCT_Strata, TEXT("GetStrataSlabBSDF(%s, %s, %s, float2(%s, %s), %s, %s, %s, %s, %s, %s, %u, Parameters.SharedNormals.NormalTypes) /* Normal = %s ; Tangent = %s */"),
+		MCT_Strata, TEXT("GetStrataSlabBSDF(%s, %s, %s, float2(%s, %s), %s, %s, %s, %s, %s, %u, Parameters.SharedNormals.NormalTypes) /* Normal = %s ; Tangent = %s */"),
 		*GetParameterCode(BaseColor),
 		*GetParameterCode(Specular),
 		*GetParameterCode(Metallic),
 		*GetParameterCode(RoughnessX),
 		*GetParameterCode(RoughnessY),
 		*GetParameterCode(SSSProfileId),
-		*GetParameterCode(SSSDMFPAlbedo),
-		*GetParameterCode(SSSDMFPRadius),
-		*GetParameterCode(SSSRadiusScale),
+		*GetParameterCode(SSSDMFP),
+		*GetParameterCode(SSSDMFPScale),
 		*GetParameterCode(EmissiveColor),
 		*GetParameterCode(Haziness),
 		SharedNormalIndex,
diff --git a/Engine/Source/Runtime/Engine/Private/Materials/HLSLMaterialTranslator.h b/Engine/Source/Runtime/Engine/Private/Materials/HLSLMaterialTranslator.h
--- a/Engine/Source/Runtime/Engine/Private/Materials/HLSLMaterialTranslator.h
+++ b/Engine/Source/Runtime/Engine/Private/Materials/HLSLMaterialTranslator.h
@@ -922,7 +922,7 @@ protected:
 	virtual int32 StrataCreateAndRegisterNullMaterial() override;
 	virtual int32 StrataSlabBSDF(
 		int32 BaseColor, int32 Specular, int32 Metallic, int32 RoughnessX, int32 RoughnessY,
-		int32 SSSProfileId, int32 SSSDMFPAlbedo, int32 SSSDMFPRadius, int32 SSSRadiusScale, 
+		int32 SSSProfileId, int32 SSSDMFP, int32 SSSDMFPScale,
 		int32 EmissiveColor, int32 Haziness, int32 Normal, int32 Tangent, uint8 SharedNormalIndex) override;
 	virtual int32 StrataSheenBSDF(int32 BaseColor, int32 Roughness, int32 Normal, uint8 SharedNormalIndex) override;
 	virtual int32 StrataVolumetricFogCloudBSDF(int32 Albedo, int32 Extinction, int32 EmissiveColor, int32 AmbientOcclusion) override;
diff --git a/Engine/Source/Runtime/Engine/Private/Materials/MaterialExpressions.cpp b/Engine/Source/Runtime/Engine/Private/Materials/MaterialExpressions.cpp
--- a/Engine/Source/Runtime/Engine/Private/Materials/MaterialExpressions.cpp
+++ b/Engine/Source/Runtime/Engine/Private/Materials/MaterialExpressions.cpp
@@ -19712,9 +19712,8 @@ int32 UMaterialExpressionStrataSlabBSDF::Compile(class FMaterialCompiler* Compil
 		RoughnessXCodeChunk,
 		RoughnessYCodeChunk,
 		SSSProfileCodeChunk != INDEX_NONE ? SSSProfileCodeChunk : Compiler->Constant(0.0f),	
-		CompileWithDefaultFloat3(Compiler, DiffuseMeanFreePathAlbedo, 0.0f, 0.0f, 0.0f),
-		CompileWithDefaultFloat1(Compiler, DiffuseMeanFreePathRadius, 0.0f),
-		CompileWithDefaultFloat1(Compiler, SSSRadiusScale, 1.0f),
+		CompileWithDefaultFloat3(Compiler, SSSDMFP, 0.0f, 0.0f, 0.0f),
+		CompileWithDefaultFloat1(Compiler, SSSDMFPScale, 1.0f),
 		CompileWithDefaultFloat3(Compiler, EmissiveColor, 0.0f, 0.0f, 0.0f),
 		CompileWithDefaultFloat1(Compiler, Haziness, 0.0f),
 		NormalCodeChunk,
@@ -19761,18 +19760,15 @@ uint32 UMaterialExpressionStrataSlabBSDF::GetInputType(int32 InputIndex)
 		return MCT_Float3; // Tangent
 		break;
 	case 7:
-		return MCT_Float3; // DMFP Albedo
+		return MCT_Float3; // SSSDMFP
 		break;
 	case 8:
-		return MCT_Float1; // DMFP Radius
+		return MCT_Float1; // SSSDMFPScale
 		break;
 	case 9:
-		return MCT_Float1; // SSS Radius Scale
-		break;
-	case 10:
 		return MCT_Float3; // Emissive Color
 		break;
-	case 11:
+	case 10:
 		return MCT_Float1; // Haziness
 		break;
 	}
@@ -19813,21 +19809,17 @@ FName UMaterialExpressionStrataSlabBSDF::GetInputName(int32 InputIndex) const
 	}
 	else if (InputIndex == 7)
 	{
-		return TEXT("DMFP Albedo");
+		return TEXT("SSS Diffuse MFP");
 	}
 	else if (InputIndex == 8)
 	{
-		return TEXT("DMFP Radius");
+		return TEXT("SSS Diffuse MFP Scale");
 	}
 	else if (InputIndex == 9)
-	{
-		return TEXT("SSS Radius Scale");
-	}
-	else if (InputIndex == 10)
 	{
 		return TEXT("Emissive Color");
 	}
-	else if (InputIndex == 11)
+	else if (InputIndex == 10)
 	{
 		return TEXT("Haziness");
 	}
@@ -19857,7 +19849,7 @@ void UMaterialExpressionStrataSlabBSDF::GatherStrataMaterialInfo(FStrataMaterial
 
 bool UMaterialExpressionStrataSlabBSDF::HasScattering() const
 {
-	return SubsurfaceProfile != nullptr || DiffuseMeanFreePathRadius.IsConnected();
+	return SubsurfaceProfile != nullptr || SSSDMFP.IsConnected();
 }
 #endif // WITH_EDITOR
 
diff --git a/Engine/Source/Runtime/Engine/Public/MaterialCompiler.h b/Engine/Source/Runtime/Engine/Public/MaterialCompiler.h
--- a/Engine/Source/Runtime/Engine/Public/MaterialCompiler.h
+++ b/Engine/Source/Runtime/Engine/Public/MaterialCompiler.h
@@ -394,7 +394,7 @@ public:
 	virtual int32 StrataCreateAndRegisterNullMaterial() = 0;
 	virtual int32 StrataSlabBSDF(
 		int32 BaseColor, int32 Specular, int32 Metallic, int32 RoughnessX, int32 RoughnessY, 
-		int32 SSSProfileId, int32 SSSDMFPAlbedo, int32 SSSDMFPRadius, int32 SSSRadiusScale, 
+		int32 SSSProfileId, int32 SSSDMFP, int32 SSSDMFPScale,
 		int32 EmissiveColor, int32 Haziness, int32 Normal, int32 Tangent, uint8 SharedNormalIndex) = 0;
 	virtual int32 StrataSheenBSDF(int32 BaseColor, int32 Roughness, int32 Normal, uint8 SharedNormalIndex) = 0;
 	virtual int32 StrataVolumetricFogCloudBSDF(int32 Albedo, int32 Extinction, int32 EmissiveColor, int32 AmbientOcclusion) = 0;
@@ -813,12 +813,12 @@ public:
 
 	virtual int32 StrataSlabBSDF(
 		int32 BaseColor, int32 Specular, int32 Metallic, int32 RoughnessX, int32 RoughnessY,
-		int32 SSSProfileId, int32 SSSDMFPAlbedo, int32 SSSDMFPRadius, int32 SSSRadiusScale, 
+		int32 SSSProfileId, int32 SSSDMFP, int32 SSSDMFPScale,
 		int32 EmissiveColor, int32 Haziness, int32 Normal, int32 Tangent, uint8 SharedNormalIndex) override
 	{
 		return Compiler->StrataSlabBSDF(
 			BaseColor, Specular, Metallic, RoughnessX, RoughnessY,
-			SSSProfileId, SSSDMFPAlbedo, SSSDMFPRadius, SSSRadiusScale, 
+			SSSProfileId, SSSDMFP, SSSDMFPScale,
 			EmissiveColor, Haziness, Normal, Tangent, SharedNormalIndex);
 	}
 
