Fix bad constant in WeakIntegerHash

Not sure how I ended up with 9 digits here. I am removing the top digit since that was being ignored and has less bias than removing an extra trailing 5 from the end.

Sampling pattern is unchanged as this overflowing digit was being ignored on platforms (PC) where it did not cause a compiler error.

diff --git a/Engine/Shaders/Private/PathTracing/Utilities/PathTracingRandomSequence.ush b/Engine/Shaders/Private/PathTracing/Utilities/PathTracingRandomSequence.ush
--- a/Engine/Shaders/Private/PathTracing/Utilities/PathTracingRandomSequence.ush
+++ b/Engine/Shaders/Private/PathTracing/Utilities/PathTracingRandomSequence.ush
@@ -57,8 +57,8 @@ uint StrongIntegerHash(uint x)
 uint WeakIntegerHash(uint x)
 {
 	// Generated using [REDACTED_PROFILE_URL]/hash-prospector
-	// Estimated Bias ~586
-	x *= 0xe92955555u;
+	// Estimated Bias ~583
+	x *= 0x92955555u;
 	x ^= x >> 15;
 	return x;
 }
