Workaround for buffer corruption leading to a hang and cleanup of shader variable names.

diff --git a/Engine/Shaders/Private/InstanceCulling/BuildInstanceDrawCommands.usf b/Engine/Shaders/Private/InstanceCulling/BuildInstanceDrawCommands.usf
--- a/Engine/Shaders/Private/InstanceCulling/BuildInstanceDrawCommands.usf
+++ b/Engine/Shaders/Private/InstanceCulling/BuildInstanceDrawCommands.usf
@@ -46,8 +46,8 @@ StructuredBuffer<uint> VisibleInstanceFlags;
 uint NumInstanceFlagWords;
 #endif
 
-int NumPrimitiveIds;
-int NumInstanceRuns;
+uint NumPrimitiveIds;
+uint NumInstanceRuns;
 int NumCommands;
 uint NumViewIds;
 
@@ -95,16 +95,18 @@ void ComputeInstanceIdOutputSize(uint GroupThreadIndex : SV_GroupIndex, uint Gro
 	GroupMemoryBarrierWithGroupSync();
 
 	// Get start of next CMD or go to end.
-	int LastOffset = (CommandSlotIndex + 1 < NumCommands) ? int(PrimitiveCullingCommands[CommandSlotIndex + 1].FirstPrimitiveIdOffset) : NumPrimitiveIds;
-	uint NumPrimitiveIds = uint(LastOffset) - Cmd.FirstPrimitiveIdOffset;
+	uint LastOffset = (CommandSlotIndex + 1 < NumCommands) ? PrimitiveCullingCommands[CommandSlotIndex + 1].FirstPrimitiveIdOffset : NumPrimitiveIds;
+	// Note: the max here is purely to safe-guard against corrupted buffer data leading to an infinite loop
+	uint NumCmdPrimitiveIds = min(NumPrimitiveIds, uint(LastOffset) - Cmd.FirstPrimitiveIdOffset);
 
-	int LastInstanceRunOffset = (CommandSlotIndex + 1 < NumCommands) ? int(PrimitiveCullingCommands[CommandSlotIndex + 1].FirstInstanceRunOffset) : NumInstanceRuns;
-	uint NumInstanceRuns = uint(LastInstanceRunOffset) - Cmd.FirstInstanceRunOffset;
+	uint LastInstanceRunOffset = (CommandSlotIndex + 1 < NumCommands) ? PrimitiveCullingCommands[CommandSlotIndex + 1].FirstInstanceRunOffset : NumInstanceRuns;
+	// Note: the min here is purely to safe-guard against corrupted buffer data leading to an infinite loop
+	uint NumCmdInstanceRuns = min(NumInstanceRuns, LastInstanceRunOffset - Cmd.FirstInstanceRunOffset);
 
-	for (uint Index = GroupThreadIndex; Index < NumPrimitiveIds + NumInstanceRuns; Index += NUM_THREADS_PER_GROUP)
+	for (uint Index = GroupThreadIndex; Index < NumCmdPrimitiveIds + NumCmdInstanceRuns; Index += NUM_THREADS_PER_GROUP)
 	{
 		uint Num = 0;
-		if (Index < NumPrimitiveIds)
+		if (Index < NumCmdPrimitiveIds)
 		{
 			uint PrimitiveId = PrimitiveIds[Cmd.FirstPrimitiveIdOffset + Index];
 			if (GetPrimitiveData(PrimitiveId).InstanceDataOffset >= 0)
@@ -120,7 +122,7 @@ void ComputeInstanceIdOutputSize(uint GroupThreadIndex : SV_GroupIndex, uint Gro
 		}
 		else
 		{
-			FInstanceRun Run = InstanceRuns[Cmd.FirstInstanceRunOffset + Index - NumPrimitiveIds];
+			FInstanceRun Run = InstanceRuns[Cmd.FirstInstanceRunOffset + Index - NumCmdPrimitiveIds];
 			Num = (Run.EndInclusive + 1) - Run.Start;
 
 			InterlockedAdd(RunInstanceCountGroupShared, Num);
@@ -375,16 +377,18 @@ void BuildInstanceIdBufferAndCommandsFromPrimitiveIdsCs(uint GroupThreadIndex :
 	GroupMemoryBarrierWithGroupSync();
 
 	// Get start of next CMD or go to end.
-	int LastOffset = (CommandSlotIndex + 1 < NumCommands) ? int(PrimitiveCullingCommands[CommandSlotIndex + 1].FirstPrimitiveIdOffset) : NumPrimitiveIds;
-	uint NumPrimitiveIds = uint(LastOffset) - Cmd.FirstPrimitiveIdOffset;
+	uint LastOffset = (CommandSlotIndex + 1 < NumCommands) ? PrimitiveCullingCommands[CommandSlotIndex + 1].FirstPrimitiveIdOffset : NumPrimitiveIds;
+	// Note: the max here is purely to safe-guard against corrupted buffer data leading to an infinite loop
+	uint NumCmdPrimitiveIds = min(NumPrimitiveIds, LastOffset - Cmd.FirstPrimitiveIdOffset);
 
-	int LastInstanceRunOffset = (CommandSlotIndex + 1 < NumCommands) ? int(PrimitiveCullingCommands[CommandSlotIndex + 1].FirstInstanceRunOffset) : NumInstanceRuns;
-	uint NumInstanceRuns = uint(LastInstanceRunOffset) - Cmd.FirstInstanceRunOffset;
+	uint LastInstanceRunOffset = (CommandSlotIndex + 1 < NumCommands) ? PrimitiveCullingCommands[CommandSlotIndex + 1].FirstInstanceRunOffset : NumInstanceRuns;
+	// Note: the min here is purely to safe-guard against corrupted buffer data leading to an infinite loop
+	uint NumCmdInstanceRuns = min(NumInstanceRuns, LastInstanceRunOffset - Cmd.FirstInstanceRunOffset);
 
-	for (uint Index = GroupThreadIndex; Index < NumPrimitiveIds + NumInstanceRuns; Index += NUM_THREADS_PER_GROUP)
+	for (uint Index = GroupThreadIndex; Index < NumCmdPrimitiveIds + NumCmdInstanceRuns; Index += NUM_THREADS_PER_GROUP)
 	{
 		uint Num = 0;
-		if (Index < NumPrimitiveIds)
+		if (Index < NumCmdPrimitiveIds)
 		{
 			uint PrimitiveId = PrimitiveIds[Cmd.FirstPrimitiveIdOffset + Index];
 			if (GetPrimitiveData(PrimitiveId).InstanceDataOffset >= 0)
@@ -400,7 +404,7 @@ void BuildInstanceIdBufferAndCommandsFromPrimitiveIdsCs(uint GroupThreadIndex :
 		}
 		else
 		{
-			FInstanceRun Run = InstanceRuns[Cmd.FirstInstanceRunOffset + Index - NumPrimitiveIds];
+			FInstanceRun Run = InstanceRuns[Cmd.FirstInstanceRunOffset + Index - NumCmdPrimitiveIds];
 			Num = (Run.EndInclusive + 1) - Run.Start;
 
 			InterlockedAdd(RunInstanceCountGroupShared, Num);
