Fix VF input setup for ray tracing use-cases with GPU-Scene instance culling enabled.

diff --git a/Engine/Shaders/Private/HairStrands/HairCardsVertexFactory.ush b/Engine/Shaders/Private/HairStrands/HairCardsVertexFactory.ush
--- a/Engine/Shaders/Private/HairStrands/HairCardsVertexFactory.ush
+++ b/Engine/Shaders/Private/HairStrands/HairCardsVertexFactory.ush
@@ -66,8 +66,15 @@ FVertexFactoryInput LoadVertexFactoryInputForHGS(uint TriangleIndex, int VertexI
 	FVertexFactoryInput Input = (FVertexFactoryInput)0;
 	Input.VertexId = TriangleAttributes.Indices[VertexIndex];
 #if VF_USE_PRIMITIVE_SCENE_DATA
+#if defined(GPUCULL_TODO)
+	// Note: GetInstanceUserData() stores the GPU-Scene primitive ID
+	Input.InstanceIdOffset = GetInstanceUserData() | VF_TREAT_INSTANCE_ID_OFFSET_AS_PRIMITIVE_ID_FLAG;
+	Input.DrawInstanceId = 0;
+#else // !defined(GPUCULL_TODO)
+	// GPUCULL_TODO: Need to work out what this is meant to do
 	Input.PrimitiveId = GetInstanceUserData();
-#endif
+#endif //!defined(GPUCULL_TODO)
+#endif // VF_USE_PRIMITIVE_SCENE_DATA
 	return Input;
 }
 #endif
diff --git a/Engine/Shaders/Private/LocalVertexFactory.ush b/Engine/Shaders/Private/LocalVertexFactory.ush
--- a/Engine/Shaders/Private/LocalVertexFactory.ush
+++ b/Engine/Shaders/Private/LocalVertexFactory.ush
@@ -194,17 +194,25 @@ FVertexFactoryInput LoadVertexFactoryInputForHGS(uint TriangleIndex, int VertexI
 
 	Input.VertexId = Tri.Indices[VertexIndex];
 	Input.Position = float4(Tri.LocalPositions[VertexIndex], 1.0f);
-	#if VF_USE_PRIMITIVE_SCENE_DATA
-#if !defined(GPUCULL_TODO)
+#if VF_USE_PRIMITIVE_SCENE_DATA
+#if defined(GPUCULL_TODO)
+	// Note: GetInstanceUserData() stores the GPU-Scene primitive ID
+	Input.InstanceIdOffset = GetInstanceUserData() | VF_TREAT_INSTANCE_ID_OFFSET_AS_PRIMITIVE_ID_FLAG;
+	Input.DrawInstanceId = 0;
+#else // !defined(GPUCULL_TODO)
 	// GPUCULL_TODO: Need to work out what this is meant to do
 	Input.PrimitiveId = GetInstanceUserData();
 #endif //!defined(GPUCULL_TODO)
-	#endif
+#endif // VF_USE_PRIMITIVE_SCENE_DATA
 
 #if USE_INSTANCING
 	Input.InstanceId = InstanceIndex() - GetBaseInstanceIndex();
 #endif
 
+#if USE_INSTANCE_CULLING
+	Input.DrawInstanceId = InstanceIndex() - GetBaseInstanceIndex();
+#endif
+
 	return Input;
 }
 #endif
@@ -222,25 +230,24 @@ FVertexFactoryInput LoadVertexFactoryInputForDynamicUpdate(uint TriangleIndex, i
 		Input.Position.x = GPUSkinCachePositionBuffer[VertexOffset * 3 + 0];
 		Input.Position.y = GPUSkinCachePositionBuffer[VertexOffset * 3 + 1];
 		Input.Position.z = GPUSkinCachePositionBuffer[VertexOffset * 3 + 2];
-		#if VF_USE_PRIMITIVE_SCENE_DATA
-#if !defined(GPUCULL_TODO)
-			// GPUCULL_TODO: !!!
-			Input.PrimitiveId = PrimitiveId;
-#endif //!defined(GPUCULL_TODO)
-		#endif
 	#else
 		Input.VertexId = TriangleIndex * 3 + VertexIndex;
 		uint VertexOffset = LocalVF.VertexFetch_Parameters[VF_VertexOffset] + Input.VertexId * 3;
 		Input.Position.x = LocalVF.VertexFetch_PositionBuffer[VertexOffset + 0];
 		Input.Position.y = LocalVF.VertexFetch_PositionBuffer[VertexOffset + 1];
 		Input.Position.z = LocalVF.VertexFetch_PositionBuffer[VertexOffset + 2];
-		#if VF_USE_PRIMITIVE_SCENE_DATA
-#if !defined(GPUCULL_TODO)
-			// GPUCULL_TODO: !!!
-			Input.PrimitiveId = PrimitiveId;
-#endif //!defined(GPUCULL_TODO)
-		#endif
 	#endif
+
+#if VF_USE_PRIMITIVE_SCENE_DATA
+#if defined(GPUCULL_TODO)
+		// GPUCULL_TODO: Makes use of DynamicDraw workaround, should instead refactor VF interfaces and provide the instance ID data directly.
+		Input.InstanceIdOffset = PrimitiveId | VF_TREAT_INSTANCE_ID_OFFSET_AS_PRIMITIVE_ID_FLAG;
+		Input.DrawInstanceId = 0;
+#else // !defined(GPUCULL_TODO)
+		Input.PrimitiveId = PrimitiveId;
+#endif // defined(GPUCULL_TODO)
+#endif // VF_USE_PRIMITIVE_SCENE_DATA
+
 	return Input;
 }
 
diff --git a/Engine/Shaders/Private/SceneData.ush b/Engine/Shaders/Private/SceneData.ush
--- a/Engine/Shaders/Private/SceneData.ush
+++ b/Engine/Shaders/Private/SceneData.ush
@@ -11,6 +11,9 @@
 
 #define INSTANCE_SCENE_DATA_SUPPORTED (FEATURE_LEVEL >= FEATURE_LEVEL_SM5 || FEATURE_LEVEL == FEATURE_LEVEL_ES3_1)
 
+// GPUCULL_TODO: Eventually we need to remove this workaround
+#define VF_TREAT_INSTANCE_ID_OFFSET_AS_PRIMITIVE_ID_FLAG (1U << 31U)
+
 // Unpacked AoS layout - see FInstanceSceneShaderData::Setup() for SoA packed layout.
 #if INSTANCE_SCENE_DATA_SUPPORTED
 
@@ -313,12 +316,13 @@ FSceneDataIntermediates GetSceneDataIntermediates(uint InstanceIdOffset, uint Dr
 	// GPUCULL_TODO: workaround for the fact that DrawDynamicMeshPassPrivate et al. don't work with GPU-Scene instancing
 	//				 instead they mark the top bit in the primitive ID and disable auto instancing such that there is an 1:1:1
 	//               drawcmd:primitive:instance. Then we can just look up the primitive and fetch the instance data index.
+	// GPUCULL_TODO: Workaround also used for ray tracing interfacing with the VFs, that also supply a DrawInstanceId. 
 	// We mark the PrimitiveID with the top bit in dynamic draw passes
-	if (InstanceIdOffset >= (1U << 31U))
+	if ((InstanceIdOffset & VF_TREAT_INSTANCE_ID_OFFSET_AS_PRIMITIVE_ID_FLAG) != 0U)
 	{
-		// mask off the top bit
-		uint PrimitiveID = InstanceIdOffset & ((1U << 31U) - 1U);
-		Intermediates.InstanceId = GetPrimitiveData(PrimitiveID).InstanceDataOffset;
+		// mask off the flag
+		uint PrimitiveID = InstanceIdOffset & (VF_TREAT_INSTANCE_ID_OFFSET_AS_PRIMITIVE_ID_FLAG - 1U);
+		Intermediates.InstanceId = GetPrimitiveData(PrimitiveID).InstanceDataOffset + DrawInstanceId;
 		Intermediates.ViewIndex = 0;
 	}
 	else
