Fix strata classification mask

diff --git a/Engine/Shaders/Private/Strata/StrataMaterialClassification.usf b/Engine/Shaders/Private/Strata/StrataMaterialClassification.usf
--- a/Engine/Shaders/Private/Strata/StrataMaterialClassification.usf
+++ b/Engine/Shaders/Private/Strata/StrataMaterialClassification.usf
@@ -216,31 +216,31 @@ void TileMainCS(uint2 DispatchThreadId : SV_DispatchThreadID, uint LinearIndex :
 	GroupMemoryBarrierWithGroupSync();
 	if (LinearIndex < 32)
 	{
-		s_IsSimpleTile[LinearIndex] = s_IsSimpleTile[LinearIndex] || s_IsSimpleTile[LinearIndex + 32];
+		s_IsSimpleTile[LinearIndex] = s_IsSimpleTile[LinearIndex] && s_IsSimpleTile[LinearIndex + 32];
 	}
 	GroupMemoryBarrierWithGroupSync();
 	if (LinearIndex < 16)
 	{
-		s_IsSimpleTile[LinearIndex] = s_IsSimpleTile[LinearIndex] || s_IsSimpleTile[LinearIndex + 16];
+		s_IsSimpleTile[LinearIndex] = s_IsSimpleTile[LinearIndex] && s_IsSimpleTile[LinearIndex + 16];
 	}
 	GroupMemoryBarrierWithGroupSync();
 
 	if (LinearIndex < 8)
 	{
-		s_IsSimpleTile[LinearIndex] = s_IsSimpleTile[LinearIndex] || s_IsSimpleTile[LinearIndex + 8];
+		s_IsSimpleTile[LinearIndex] = s_IsSimpleTile[LinearIndex] && s_IsSimpleTile[LinearIndex + 8];
 	}
 	GroupMemoryBarrierWithGroupSync();
 	if (LinearIndex < 4)
 	{
-		s_IsSimpleTile[LinearIndex] = s_IsSimpleTile[LinearIndex] || s_IsSimpleTile[LinearIndex + 4];
+		s_IsSimpleTile[LinearIndex] = s_IsSimpleTile[LinearIndex] && s_IsSimpleTile[LinearIndex + 4];
 	}
 	GroupMemoryBarrierWithGroupSync();
 	if (LinearIndex < 2)
 	{
-		s_IsSimpleTile[LinearIndex] = s_IsSimpleTile[LinearIndex] || s_IsSimpleTile[LinearIndex + 2];
+		s_IsSimpleTile[LinearIndex] = s_IsSimpleTile[LinearIndex] && s_IsSimpleTile[LinearIndex + 2];
 	}
 	GroupMemoryBarrierWithGroupSync();
-	if (LinearIndex < 1 && (s_IsSimpleTile[LinearIndex] || s_IsSimpleTile[LinearIndex + 1]))
+	if (LinearIndex < 1 && (s_IsSimpleTile[LinearIndex] && s_IsSimpleTile[LinearIndex + 1]))
 	{
 		uint WriteToIndex;
 		InterlockedAdd(TileIndirectData[1], 1, WriteToIndex);
