Fixed reflection when not using dynamic sky.

diff --git a/Engine/Shaders/Private/ReflectionEnvironmentPixelShader.usf b/Engine/Shaders/Private/ReflectionEnvironmentPixelShader.usf
--- a/Engine/Shaders/Private/ReflectionEnvironmentPixelShader.usf
+++ b/Engine/Shaders/Private/ReflectionEnvironmentPixelShader.usf
@@ -313,7 +313,7 @@ float3 ReflectionEnvironment(FGBufferData GBuffer, float AmbientOcclusion, float
 	return -min(-Color.rgb, 0.0);
 }
 
-void ApplyCloudVolumetricShadow(in float3 WorldPosition, inout float4 SkyLighting)
+void ApplyCloudVolumetricShadow(in float3 WorldPosition, inout float3 SkyLighting)
 {
 	// Ideally we would compute spatially how much of the sky+cloud is visible for each point in the world. But we evaluate the sky light only once at the sky light component position as of today.
 	// To add some spatial variation, we can affect the sky light diffuse contribution according to cloud occlusion from above.
@@ -322,7 +322,7 @@ void ApplyCloudVolumetricShadow(in float3 WorldPosition, inout float4 SkyLightin
 	if (CloudSkyAOEnabled)
 	{
 		float OutOpticalDepth = 0.0f;
-		SkyLighting.rgb *= GetCloudVolumetricShadow(WorldPosition, CloudSkyAOWorldToLightClipMatrix, CloudSkyAOFarDepthKm, CloudSkyAOTexture, CloudSkyAOSampler, OutOpticalDepth);
+		SkyLighting *= GetCloudVolumetricShadow(WorldPosition, CloudSkyAOWorldToLightClipMatrix, CloudSkyAOFarDepthKm, CloudSkyAOTexture, CloudSkyAOSampler, OutOpticalDepth);
 	}
 #endif
 }
@@ -351,6 +351,9 @@ void ReflectionEnvironmentSkyLighting(
 	float AmbientOcclusion = AmbientOcclusionTexture.SampleLevel(AmbientOcclusionSampler, BufferUV, 0).r;
 	
 	uint ShadingModelID = GBuffer.ShadingModelID;
+	const bool bUnlitMaterial = ShadingModelID == SHADINGMODELID_UNLIT;
+	const bool bStrataMaterial = ShadingModelID == SHADINGMODELID_STRATA;
+
 	float3 BentNormal = GBuffer.WorldNormal;
 	#if APPLY_SKY_SHADOWING
 	{
@@ -362,18 +365,13 @@ void ReflectionEnvironmentSkyLighting(
 
 #if ENABLE_DYNAMIC_SKY_LIGHT
 
-	uint2 PixelPos = SvPosition.xy;
-
-	const bool bUnlitMaterial = ShadingModelID == SHADINGMODELID_UNLIT;
-
-	float4 SkyLighting = 0;
-
 #if STRATA_ENABLED && ENABLE_SKY_LIGHT
 
 	// STRATA_TODO handle the ENABLE_SKY_LIGHT && ALLOW_STATIC_LIGHTING case 
 	BRANCH
-	if (ShadingModelID == SHADINGMODELID_STRATA)
+	if (bStrataMaterial)
 	{
+		uint2 PixelPos = SvPosition.xy;
 		const FStrataTopLayerData TopLayerData = StrataUnpackTopLayerData(Strata.TopLayerTexture.Load(uint3(PixelPos, 0)));
 		const float4 SSR = Texture2DSample(ScreenSpaceReflectionsTexture, ScreenSpaceReflectionsSampler, BufferUV);
 
@@ -396,37 +394,40 @@ void ReflectionEnvironmentSkyLighting(
 			GetStrataPixelDataByteOffset(PixelPos, View.BufferSizeAndInvSize.x, Strata.MaxBytesPerPixel),
 			SSRReductionFactor);
 
-		SkyLighting.rgb = SSR.rgb * TopLayerData.EnvBRDF * SSRReductionFactor + StrataEnvLighting.rgb;
+		float3 SkyLighting = SSR.rgb * TopLayerData.EnvBRDF * SSRReductionFactor + StrataEnvLighting.rgb;
 
 		ApplyCloudVolumetricShadow(WorldPosition, SkyLighting);
 
-		OutColor.xyz = SkyLighting.rgb;
-		}
+		OutColor.xyz = SkyLighting;
+	}
 
 #else // STRATA_ENABLED && ENABLE_SKY_LIGHT
 
 	BRANCH
 	if (!bUnlitMaterial) // Only light pixels marked as lit
 	{
-		SkyLighting.rgb = SkyLightDiffuse(GBuffer, AmbientOcclusion, BufferUV, ScreenPosition, BentNormal, DiffuseColor);
+		float3 SkyLighting = SkyLightDiffuse(GBuffer, AmbientOcclusion, BufferUV, ScreenPosition, BentNormal, DiffuseColor);
 
 		ApplyCloudVolumetricShadow(WorldPosition, SkyLighting);
 
 		FLightAccumulator LightAccumulator = (FLightAccumulator)0;
 		const bool bNeedsSeparateSubsurfaceLightAccumulation = UseSubsurfaceProfile(ShadingModelID);
 
-		LightAccumulator_Add(LightAccumulator, SkyLighting.rgb, SkyLighting.rgb, 1.0f, bNeedsSeparateSubsurfaceLightAccumulation);
+		LightAccumulator_Add(LightAccumulator, SkyLighting, SkyLighting, 1.0f, bNeedsSeparateSubsurfaceLightAccumulation);
 		OutColor = LightAccumulator_GetResult(LightAccumulator);
 	}
 
+#endif // STRATA_ENABLED && ENABLE_SKY_LIGHT
+
+#endif // ENABLE_DYNAMIC_SKY_LIGHT
+
+
+#if !STRATA_ENABLED
 	BRANCH
 	if (!bUnlitMaterial && ShadingModelID != SHADINGMODELID_HAIR)
 	{
 		OutColor.xyz += ReflectionEnvironment(GBuffer, AmbientOcclusion, BufferUV, WorldPosition, V, SvPosition, BentNormal, SpecularColor);
 	}
-
-#endif // STRATA_ENABLED && ENABLE_SKY_LIGHT
-
-#endif // ENABLE_DYNAMIC_SKY_LIGHT
+#endif
 }
 
