Remove strata header loading when strata fast path is used.

diff --git a/Engine/Shaders/Private/Strata/Strata.ush b/Engine/Shaders/Private/Strata/Strata.ush
--- a/Engine/Shaders/Private/Strata/Strata.ush
+++ b/Engine/Shaders/Private/Strata/Strata.ush
@@ -1354,6 +1354,39 @@ FStrataPixelHeader UnpackStrataHeaderIn(ByteAddressBuffer StrataBuffer, inout ui
 	return Out;
 }
 
+// Unpack only strata header 
+// This version assume only a single slab (isotropic BSDF, with no subsurface scattering, and no haziness)
+FStrataPixelHeader UnpackFastPathStrataHeaderIn(ByteAddressBuffer StrataBuffer, inout uint ByteOffset)
+{
+	FStrataPixelHeader Out = InitialiseStrataPixelHeader();
+	
+	// Skip over the header packed data
+	ByteOffset += 4; ;
+
+	Out.BSDFCount = 1;
+#if STRATA_INLINE_SHADING==0
+	Out.SharedNormalByteOffset = ByteOffset;
+	Out.StrataBuffer = StrataBuffer;
+	Out.SharedNormalTypes = STRATA_BASIS_TYPE_NORMAL;
+#endif
+
+	// Skip over the shared normal memory to be able to load the BSDF content
+	ByteOffset += STRATA_PACKED_NORMAL_STRIDE_BYTES;
+
+	return Out;
+}
+
+// Force a BSDF state to a single slabe with isotropic BSDF, with no subsurface scattering, and no haziness
+void SetStrataFastPathState(inout FStrataBSDF BSDF)
+{
+	BSDF_SETTYPE(BSDF, STRATA_BSDF_TYPE_SLAB);
+	BSDF_SETANISOTROPY(BSDF, 0);
+	BSDF_SETISTOPLAYER(BSDF, 0);
+	BSDF_SETHASSCATTERING(BSDF, 0);
+	BSDF_SETHASSSSPROFILE(BSDF, 0);
+	BSDF_SETHASHAZINESS(BSDF, 0);
+}
+
 // Unpack a single BSDF
 // Note: All BSDF lobes needs to be unpack in a sequential manner since each BSDF has a variable footprint.
 // If this is changed, please update the compiler side material size evaluation in StrataMaterial.cpp
diff --git a/Engine/Shaders/Private/Strata/StrataDeferredLighting.ush b/Engine/Shaders/Private/Strata/StrataDeferredLighting.ush
--- a/Engine/Shaders/Private/Strata/StrataDeferredLighting.ush
+++ b/Engine/Shaders/Private/Strata/StrataDeferredLighting.ush
@@ -156,9 +156,7 @@ float4 StrataDeferredLighting(
 	FRectTexture RectTexture,
 	ByteAddressBuffer StrataDataBuffer,
 	uint PixelStrataDataByteOffset)
-{
-	FStrataPixelHeader StrataPixelHeader = UnpackStrataHeaderIn(StrataDataBuffer, PixelStrataDataByteOffset);
-	
+{	
 	float3 PreintegratedIndicentIrradiance = 0;
 	if (LightData.bRectLight)
 	{
@@ -175,8 +173,9 @@ float4 StrataDeferredLighting(
 
 #if STRATA_FASTPATH
 	uint BSDFIndex = 0;
-	if (StrataPixelHeader.BSDFCount == 1)
+	FStrataPixelHeader StrataPixelHeader = UnpackFastPathStrataHeaderIn(StrataDataBuffer, PixelStrataDataByteOffset);
 #else
+	FStrataPixelHeader StrataPixelHeader = UnpackStrataHeaderIn(StrataDataBuffer, PixelStrataDataByteOffset);
 	LOOP
 	for (uint BSDFIndex = 0; BSDFIndex < StrataPixelHeader.BSDFCount; ++BSDFIndex)
 #endif
@@ -184,12 +183,7 @@ float4 StrataDeferredLighting(
 		// Unpack BSDF data
 		FStrataBSDF BSDF = UnpackStrataBSDFIn(StrataDataBuffer, PixelStrataDataByteOffset);
 #if STRATA_FASTPATH
-		BSDF_SETTYPE(BSDF, STRATA_BSDF_TYPE_SLAB);
-		BSDF_SETANISOTROPY(BSDF, 0);
-		BSDF_SETISTOPLAYER(BSDF, 0);
-		BSDF_SETHASSCATTERING(BSDF, 0);
-		BSDF_SETHASSSSPROFILE(BSDF, 0);
-		BSDF_SETHASHAZINESS(BSDF, 0);
+		SetStrataFastPathState(BSDF);
 #endif
 		FStrataBSDFContext BSDFContext = StrataCreateBSDFContext(StrataPixelHeader, BSDF, V, L);
 
diff --git a/Engine/Shaders/Private/Strata/StrataEnvironmentLighting.ush b/Engine/Shaders/Private/Strata/StrataEnvironmentLighting.ush
--- a/Engine/Shaders/Private/Strata/StrataEnvironmentLighting.ush
+++ b/Engine/Shaders/Private/Strata/StrataEnvironmentLighting.ush
@@ -20,8 +20,6 @@ float4 StrataEnvironmentLighting(
 	uint PixelStrataDataByteOffset,
 	inout float SSRReductionFactor)
 {
-	FStrataPixelHeader StrataPixelHeader = UnpackStrataHeaderIn(StrataDataBuffer, PixelStrataDataByteOffset);
-
 	FLightAccumulator Out = (FLightAccumulator)0;
 
 	SSRReductionFactor = 1.0f;
@@ -31,21 +29,18 @@ float4 StrataEnvironmentLighting(
 	CommonFactor *= View.PreExposure;
 #endif
 
+
 #if STRATA_FASTPATH
 	uint BSDFIndex = 0;
-	if (StrataPixelHeader.BSDFCount == 1)
+	FStrataPixelHeader StrataPixelHeader = UnpackFastPathStrataHeaderIn(StrataDataBuffer, PixelStrataDataByteOffset);
 #else
+	FStrataPixelHeader StrataPixelHeader = UnpackStrataHeaderIn(StrataDataBuffer, PixelStrataDataByteOffset);
 	for (uint BSDFIndex = 0; BSDFIndex < StrataPixelHeader.BSDFCount; ++BSDFIndex)
 #endif
 	{
 		FStrataBSDF BSDF = UnpackStrataBSDFIn(StrataDataBuffer, PixelStrataDataByteOffset);
 #if STRATA_FASTPATH
-		BSDF_SETTYPE(BSDF, STRATA_BSDF_TYPE_SLAB);
-		BSDF_SETANISOTROPY(BSDF, 0);
-		BSDF_SETISTOPLAYER(BSDF, 0);
-		BSDF_SETHASSCATTERING(BSDF, 0);
-		BSDF_SETHASSSSPROFILE(BSDF, 0);
-		BSDF_SETHASHAZINESS(BSDF, 0);
+		SetStrataFastPathState(BSDF);
 #endif
 		const float3 BSDFThroughput = BSDF.Weight;
 
