Fixed tiles at right/bottom screen border forced receiving compelx strata material type. Now they can be simple.
Sky pixels are considered simple to allow simple on the edge of meshes over the sky.
Fixed precision issue causing flickering when reading GBuffer.

diff --git a/Engine/Shaders/Private/Strata/StrataMaterialClassification.usf b/Engine/Shaders/Private/Strata/StrataMaterialClassification.usf
--- a/Engine/Shaders/Private/Strata/StrataMaterialClassification.usf
+++ b/Engine/Shaders/Private/Strata/StrataMaterialClassification.usf
@@ -138,7 +138,7 @@ void MainPS(
 	const uint2 PixelPos = uint2(SVPos.xy);
 
 	// Sample scene textures.
-	float2 BufferUV		 = SvPositionToBufferUV(float4(PixelPos, SVPos.zw));
+	float2 BufferUV		 = SvPositionToBufferUV(SVPos);
 	FGBufferData GBuffer = GetGBufferDataFromSceneTextures(BufferUV);
 	
 	FStrataClassification ClassificationData	= (FStrataClassification)0;
@@ -146,6 +146,7 @@ void MainPS(
 	FStrataTopLayerData TopLayerData			= (FStrataTopLayerData)0;
 
 	ClassificationData.bIsStrataMaterial = GBuffer.ShadingModelID == SHADINGMODELID_STRATA;
+	ClassificationData.bIsSimple = true; // non strata material or sky pixels are considered as simple.
 	if (ClassificationData.bIsStrataMaterial)
 	{
 		float2 ScreenPosition = SvPositionToScreenPosition(SVPos).xy;
@@ -185,7 +186,7 @@ uint2 DecodeTile(uint In)
 #if SHADER_TILE_CATEGORIZATION
 int TileSize;
 int bRectPrimitive;
-int2 ClassificationResolution;
+int2 ViewResolution;
 Texture2D<uint> ClassificationTexture;
 RWBuffer<uint> TileIndirectData;
 RWBuffer<uint> TileListData;
@@ -204,8 +205,10 @@ void TileMainCS(uint2 DispatchThreadId : SV_DispatchThreadID, uint LinearIndex :
 
 	// TODO: add a SM6 permutation with ballot?
 	const uint2 PixelCoord = DispatchThreadId;
-	bool bIsSimpleTile = false;
-	if (all(PixelCoord < uint2(ClassificationResolution)))
+
+	// Pixels outside of the view area are considered simple to enable screen borders to receive the simple permutation when not aligned to shader group size.
+	bool bIsSimpleTile = true;
+	if (all(PixelCoord < uint2(ViewResolution)))
 	{
 		const uint PackClassification = ClassificationTexture.Load(uint3(PixelCoord, 0));
 		FStrataClassification Data = StrataUnpackClassificationData(PackClassification);
diff --git a/Engine/Source/Runtime/Renderer/Private/Strata/Strata.cpp b/Engine/Source/Runtime/Renderer/Private/Strata/Strata.cpp
--- a/Engine/Source/Runtime/Renderer/Private/Strata/Strata.cpp
+++ b/Engine/Source/Runtime/Renderer/Private/Strata/Strata.cpp
@@ -303,7 +303,7 @@ class FStrataMaterialTileClassificationPassCS : public FGlobalShader
 		SHADER_PARAMETER_STRUCT_REF(FViewUniformShaderParameters, ViewUniformBuffer)
 		SHADER_PARAMETER(int32, TileSize)
 		SHADER_PARAMETER(int32, bRectPrimitive)
-		SHADER_PARAMETER(FIntPoint, ClassificationResolution)
+		SHADER_PARAMETER(FIntPoint, ViewResolution)
 		SHADER_PARAMETER_RDG_TEXTURE(Texture2D, ClassificationTexture)
 		SHADER_PARAMETER_RDG_BUFFER_UAV(RWBuffer, TileIndirectData)
 		SHADER_PARAMETER_RDG_BUFFER_UAV(RWBuffer, TileListData)
@@ -546,7 +546,7 @@ void AddStrataMaterialClassificationPass(FRDGBuilder& GraphBuilder, const FMinim
 			FStrataMaterialTileClassificationPassCS::FParameters* PassParameters = GraphBuilder.AllocParameters<FStrataMaterialTileClassificationPassCS::FParameters>();
 			PassParameters->TileSize = GetStrataTileSize();
 			PassParameters->bRectPrimitive = GRHISupportsRectTopology ? 1 : 0;
-			PassParameters->ClassificationResolution = ClassificationTexture->Desc.Extent;
+			PassParameters->ViewResolution = View.ViewRect.Size();
 			PassParameters->ClassificationTexture = ClassificationTexture;
 			PassParameters->TileListData = GraphBuilder.CreateUAV(GraphBuilder.RegisterExternalBuffer(View.StrataSceneData->ClassificationTileListBuffer), PF_R32_UINT);
 			PassParameters->TileIndirectData = GraphBuilder.CreateUAV(GraphBuilder.RegisterExternalBuffer(View.StrataSceneData->ClassificationTileIndirectBuffer), PF_R32_UINT);
