Fix SSR weight with Strata material when running classification path.

diff --git a/Engine/Shaders/Private/Strata/Strata.ush b/Engine/Shaders/Private/Strata/Strata.ush
--- a/Engine/Shaders/Private/Strata/Strata.ush
+++ b/Engine/Shaders/Private/Strata/Strata.ush
@@ -1381,7 +1381,7 @@ void SetStrataFastPathState(inout FStrataBSDF BSDF)
 {
 	BSDF_SETTYPE(BSDF, STRATA_BSDF_TYPE_SLAB);
 	BSDF_SETANISOTROPY(BSDF, 0);
-	BSDF_SETISTOPLAYER(BSDF, 0);
+	BSDF_SETISTOPLAYER(BSDF, 1);
 	BSDF_SETHASSCATTERING(BSDF, 0);
 	BSDF_SETHASSSSPROFILE(BSDF, 0);
 	BSDF_SETHASHAZINESS(BSDF, 0);
