Thin translucent support in the path tracer

Similar to rough glass support when using the translucent blend mode, the thin translucent lighting model blends in a thin (rough) glass model when Opacity is less than 1. When Ior is 1, plain transparency is used (again similarly to the solid glass case). Both cases account for the thin translucent color set by the shader to account for absorption through the slab.

This patch also refactors the GGX calls into a common place to minimize code duplication.

diff --git a/Engine/Shaders/Private/PathTracing/Material/PathTracingClearCoat.ush b/Engine/Shaders/Private/PathTracing/Material/PathTracingClearCoat.ush
--- a/Engine/Shaders/Private/PathTracing/Material/PathTracingClearCoat.ush
+++ b/Engine/Shaders/Private/PathTracing/Material/PathTracingClearCoat.ush
@@ -74,10 +74,8 @@ void ClearCoat_EvalMaterial(
 
 	float CoatSelectionPdf = CalcClearCoatSelectionPdf(Payload);
 
-	// No anisotropy, just pick an arbitrary basis
-	float Roughness = max(GetClearCoatRoughness(Payload), GGX_MIN_ROUGHNESS);
-	float2 Alpha = Roughness * Roughness;
-	float3x3 Basis = GetTangentBasis(N_World);
+	float2 Alpha = 0;
+	float3x3 Basis = GetGGXBasis(GetClearCoatRoughness(Payload), Payload.WorldNormal, Alpha);
 
 	// move vectors into right shading frame
 	float3 V = mul(Basis, V_World);
@@ -90,25 +88,13 @@ void ClearCoat_EvalMaterial(
 	float VoH = saturate(dot(V, H));
 
 	// Specular lobe
-	if (NoL > 0 && NoV > 0)
-	{
-		// Assume NoH > 0
-		// FIXME: this can be further simplified
-		const float CosThetaM4 = Pow2(NoH);
-		const float TanThetaM2 = (Pow2(H.x / Alpha.x) + Pow2(H.y / Alpha.y)) / Pow2(NoH);
-		const float D = 1 / (PI * Pow2(1 + TanThetaM2) * Alpha.x * Alpha.y * CosThetaM4);
-		const float Lambda_o = 0.5 * (-1.0 + sqrt(1.0 + (Pow2(L.x * Alpha.x) + Pow2(L.y * Alpha.y)) / Pow2(L.z)));
-		const float Lambda_i = 0.5 * (-1.0 + sqrt(1.0 + (Pow2(V.x * Alpha.x) + Pow2(V.y * Alpha.y)) / Pow2(V.z)));
-		const float G2 = 1 / (Lambda_i + Lambda_o + 1);
-		const float G1 = 1 / (Lambda_i + 1);
-
-		const float F = GetClearCoat(Payload) * ClearCoatFresnel(VoH);
-
-		const float3 SpecWeight = F * (G2 / G1);
-		const float SpecPdf = (G1 * D * 0.25f) / NoV;
-
-		AddLobeWithMIS(OutWeight, OutPdf, SpecWeight, SpecPdf, CoatSelectionPdf);
-	}
+	const float2 GGXResult = GGXEvalReflection(L, V, H, Alpha);
+	const float F = GetClearCoat(Payload) * ClearCoatFresnel(VoH);
+	const float3 SpecWeight = F * GGXResult.x;
+	const float SpecPdf = GGXResult.y;
+
+	AddLobeWithMIS(OutWeight, OutPdf, SpecWeight, SpecPdf, CoatSelectionPdf);
+	
 
 	// bottom-layer clear coat evaluation
 	float3 BottomWeight = 0;
diff --git a/Engine/Shaders/Private/PathTracing/Material/PathTracingGlossy.ush b/Engine/Shaders/Private/PathTracing/Material/PathTracingGlossy.ush
--- a/Engine/Shaders/Private/PathTracing/Material/PathTracingGlossy.ush
+++ b/Engine/Shaders/Private/PathTracing/Material/PathTracingGlossy.ush
@@ -6,3 +6,98 @@ Glossy.usf: Microfacet GGX BRDF sampling functions
 
 #pragma once
 
+// be conservative since the GPU may be aggressive with floating point optimizations
+#define GGX_MIN_ROUGHNESS 0.01
+#define GGX_MIN_ANISOTROPY 0.001 // ignore anisotropy below this amount
+
+float3x3 GetGGXBasis(float Roughness, float Anisotropy, float3 WorldNormal, float3 WorldTangent, out float2 Alpha)
+{
+	Roughness = max(Roughness, GGX_MIN_ROUGHNESS);
+	float3x3 Basis;
+	if (abs(Anisotropy) >= GGX_MIN_ANISOTROPY)
+	{
+		// Material is anisotropic, our shading frame should use the tangent vector
+		GetAnisotropicRoughness(Roughness * Roughness, Anisotropy, Alpha.x, Alpha.y);
+		float3 X = WorldTangent;
+		float3 Y = normalize(cross(WorldNormal, X));
+		Basis = float3x3(cross(Y, WorldNormal), Y, WorldNormal);
+	}
+	else
+	{
+		// No anisotropy, just pick an arbitrary basis
+		Alpha = Roughness * Roughness;
+		Basis = GetTangentBasis(WorldNormal);
+	}
+	return Basis;
+}
+
+float3x3 GetGGXBasis(float Roughness, float3 WorldNormal, out float2 Alpha)
+{
+	// No anisotropy, just pick an arbitrary basis
+	Roughness = max(Roughness, GGX_MIN_ROUGHNESS);
+	Alpha = Roughness * Roughness;
+	return GetTangentBasis(WorldNormal);
+}
+
+float3 GGXHelper(float3 L, float3 V, float3 H, float2 Alpha)
+{
+	// TODO: Can this be further simplified / unified with other GGX routines?
+	const float CosThetaM4 = Pow2(H.z);
+	const float TanThetaM2 = (Pow2(H.x / Alpha.x) + Pow2(H.y / Alpha.y)) / Pow2(H.z);
+	const float D = 1 / (PI * Pow2(1 + TanThetaM2) * Alpha.x * Alpha.y * CosThetaM4);
+	const float Lambda_o = 0.5 * (-1.0 + sqrt(1.0 + (Pow2(L.x * Alpha.x) + Pow2(L.y * Alpha.y)) / Pow2(L.z)));
+	const float Lambda_i = 0.5 * (-1.0 + sqrt(1.0 + (Pow2(V.x * Alpha.x) + Pow2(V.y * Alpha.y)) / Pow2(V.z)));
+	const float G2 = 1 / (Lambda_i + Lambda_o + 1);
+	const float G1 = 1 / (Lambda_i + 1);
+
+	return float3(D, G2, G1);
+}
+
+
+// Compute GGX lobe Weight and Pdf (without Fresnel term) given a set of vectors in local space (Z up)
+float2 GGXEvalReflection(float3 L, float3 V, float3 H, float2 Alpha)
+{
+	const float NoL = saturate(L.z);
+	const float NoV = saturate(V.z);
+
+	if (NoL > 0 && NoV > 0)
+	{
+		const float3 Result = GGXHelper(L, V, H, Alpha);
+		const float D = Result.x;
+		const float G2 = Result.y;
+		const float G1 = Result.z;
+
+		const float Weight = G2 / G1;
+		const float Pdf = (G1 * D * 0.25f) / NoV;
+
+		return float2(Weight, Pdf);
+	}
+	return 0;
+}
+
+
+// Compute GGX lobe Weight and Pdf (without Fresnel term) given a set of vectors in local space (Z up)
+float2 GGXEvalRefraction(float3 L, float3 V, float3 H, float2 Alpha, float Eta)
+{
+	const float NoV = saturate(V.z);
+	const float NoL = saturate(-L.z); // NOTE: L should point down
+
+	// This checks that V and L point in opposite sides 
+	if (NoL > 0 && NoV > 0)
+	{
+		const float3 Result = GGXHelper(L, V, H, Alpha);
+		const float D = Result.x;
+		const float G2 = Result.y;
+		const float G1 = Result.z;
+
+		const float VoH = dot(V, H);
+		const float LoH = dot(L, H);
+		const float3 Ht = (Eta * L + V);
+
+		const float Weight = G2 / G1;
+		const float Pdf = (abs(LoH * VoH) * (Eta * Eta) * (G1 * D)) / (dot(Ht, Ht) * NoV);
+
+		return float2(Weight, Pdf);
+	}
+	return 0;
+}
diff --git a/Engine/Shaders/Private/PathTracing/Material/PathTracingMaterialCommon.ush b/Engine/Shaders/Private/PathTracing/Material/PathTracingMaterialCommon.ush
--- a/Engine/Shaders/Private/PathTracing/Material/PathTracingMaterialCommon.ush
+++ b/Engine/Shaders/Private/PathTracing/Material/PathTracingMaterialCommon.ush
@@ -9,9 +9,6 @@ PathTracingMaterialCommon.usf: Brdf utility functions
 #include "../PathTracingCommon.ush"
 #include "/Engine/Private/BRDF.ush"
 
-// be conservative since the GPU may be aggressive with floating point optimizations
-#define GGX_MIN_ROUGHNESS 0.01
-
 float3 GetMaterialEmission(FMaterialClosestHitPayload HitInfo)
 {
 	return HitInfo.Radiance;
@@ -119,6 +116,15 @@ bool IsMaterialTransmissive(FMaterialClosestHitPayload HitInfo)
 		   (HitInfo.BlendingMode == RAY_TRACING_BLEND_MODE_TRANSLUCENT && HitInfo.Roughness != 0 && HitInfo.Opacity < 1);
 }
 
+// Given two lobes that will roughly contribute colors A and B to the total (estimated for example using directional albedo)
+// return the probability of choosing lobe A
+float LobeSelectionProb(float3 A, float3 B)
+{
+	const float SumA = A.x + A.y + A.z;
+	const float SumB = B.x + B.y + B.z;
+	return SumA / (SumA + SumB + 1e-6);
+}
+
 // Takes as input the sample weight and pdf for a certain lobe of a mixed model, together with the probability of picking that lobe
 // This function then updates a running total Weight and Pdf value that represents the overall contribution of the BxDF
 // This function should be called when a BxDF is made up of multiple lobes combined with a sum to correctly account for the probability
diff --git a/Engine/Shaders/Private/PathTracing/Material/PathTracingMaterialSampling.ush b/Engine/Shaders/Private/PathTracing/Material/PathTracingMaterialSampling.ush
--- a/Engine/Shaders/Private/PathTracing/Material/PathTracingMaterialSampling.ush
+++ b/Engine/Shaders/Private/PathTracing/Material/PathTracingMaterialSampling.ush
@@ -13,6 +13,8 @@ PathTracingMaterialSampling.usf: Material sampling functions
 #include "PathTracingGlossy.ush"
 #include "PathTracingSpecularReflection.ush"
 #include "PathTracingSpecularTransmission.ush"
+#include "PathTracingSolidGlass.ush"
+#include "PathTracingThinGlass.ush"
 #include "PathTracingDefaultLit.ush"
 #include "PathTracingClearCoat.ush"
 #include "PathTracingTwoSidedFoliage.ush"
@@ -42,14 +44,27 @@ bool IsNeeValidMaterial(FMaterialClosestHitPayload Payload)
 		case SHADINGMODELID_DEFAULT_LIT:		
 		case SHADINGMODELID_CLEAR_COAT:			
 		case SHADINGMODELID_TWOSIDED_FOLIAGE:	
+		case SHADINGMODELID_THIN_TRANSLUCENT:
 		default:
 			return true;
 	}
 #endif
 }
 
-float3 GetMaterialTransparency(FMaterialClosestHitPayload Payload)
+float3 GetMaterialTransparency(FMaterialClosestHitPayload Payload, float3 RayDirection)
 {
+	if (Payload.ShadingModelID == SHADINGMODELID_THIN_TRANSLUCENT)
+	{
+		if (Payload.Ior == 1.0)
+		{
+			float VoN = dot(RayDirection, Payload.WorldNormal);
+			return (1 - Payload.Opacity) * pow(Payload.CustomData.xyz, 1.0 / abs(VoN));
+		}
+		else
+		{
+			return 0;
+		}
+	}
 	switch (Payload.BlendingMode)
 	{
 		case RAY_TRACING_BLEND_MODE_TRANSLUCENT:
@@ -125,6 +140,7 @@ float3 EstimateMaterialAlbedo(FMaterialClosestHitPayload Payload)
 	if (HasRefraction)
 	{
 		// glass is currently perfectly clear - so has unit albedo
+		// TODO: estimate some kind of albedo for thin-translucent case since that one uses a color?
 		float3 White = 1;
 		Albedo = lerp(White, Albedo, Payload.Opacity);
 	}
@@ -146,6 +162,7 @@ void AdjustPathAndMaterialRoughness(inout FMaterialClosestHitPayload Payload, in
 			// fallthrough
 		case SHADINGMODELID_TWOSIDED_FOLIAGE:
 		case SHADINGMODELID_DEFAULT_LIT:
+		case SHADINGMODELID_THIN_TRANSLUCENT:
 			Payload.Roughness = max(Payload.Roughness, PathRoughness);
 			break;
 		default:
@@ -232,8 +249,12 @@ void SampleMaterial(
 		// we randomly decided to sample the translucent portion
 		RandSample.x = (RandSample.x - BrdfProb) / (1 - BrdfProb);
 
-		// only case that should reach this point is the refraction model
-		if (Payload.Roughness == 0)
+		// only case that should reach this point is the refraction model (which can be thin or solid)
+		if (Payload.ShadingModelID == SHADINGMODELID_THIN_TRANSLUCENT)
+		{
+			RoughThinGlass_SampleMaterial(RayDirection, Payload, RandSample, OutDirection, OutWeight, OutPdf, OutPositionBiasSign);
+		}
+		else if (Payload.Roughness == 0)
 		{
 			SpecularRefraction_SampleMaterial(RayDirection, Payload, RandSample, OutDirection, OutWeight, OutPdf, OutPositionBiasSign);
 		}
@@ -368,24 +389,20 @@ void EvalMaterial(
 		{
 			float3 GlassWeight = 0;
 			float GlassPdf = 0;
-			RoughGlass_EvalMaterial(IncomingDirection, OutgoingDirection, Payload, GlassWeight, GlassPdf);
+			if (Payload.ShadingModelID == SHADINGMODELID_THIN_TRANSLUCENT)
+			{
+				// just a thin slab (reflection-like lobe on both sides)
+				RoughThinGlass_EvalMaterial(IncomingDirection, OutgoingDirection, Payload, GlassWeight, GlassPdf);
+			}
+			else
+			{
+				// real solid glass, reflection and refraction
+				RoughGlass_EvalMaterial(IncomingDirection, OutgoingDirection, Payload, GlassWeight, GlassPdf);
+			}
+			
 
 			AddLobeWithMIS(OutWeight, OutPdf, GlassWeight, GlassPdf, 1.0 - BrdfProb);
 		}
 	}
 #endif
 }
-
-
-void PdfMaterial(
-	float3 IncomingDirection,
-	float3 OutgoingDirection,
-	FMaterialClosestHitPayload Payload,
-	bool bForceRadianceProbe,
-	out float OutPdf
-)
-{
-	// NOTE: rely on compiler to optimize unused calculations
-	float3 Weight = 0;
-	EvalMaterial(IncomingDirection, OutgoingDirection, Payload, bForceRadianceProbe, Weight, OutPdf);
-}
diff --git a/Engine/Shaders/Private/PathTracing/Material/PathTracingSolidGlass.ush b/Engine/Shaders/Private/PathTracing/Material/PathTracingSolidGlass.ush
new file mode 100644
--- /dev/null
+++ b/Engine/Shaders/Private/PathTracing/Material/PathTracingSolidGlass.ush
@@ -0,0 +1,138 @@
+// Copyright Epic Games, Inc. All Rights Reserved.
+
+/*=============================================================================================
+PathTracingSolidGlass.ush: Microfacet Refractive BSDF for solid glass
+===============================================================================================*/
+
+#pragma once
+
+#include "PathTracingMaterialCommon.ush"
+#include "PathTracingFresnel.ush"
+#include "PathTracingGlossy.ush"
+
+void RoughGlass_SampleMaterial(
+	float3 RayDirection,
+	FMaterialClosestHitPayload Payload,
+	float4 RandSample,
+	out float3 OutDirection,
+	out float3 OutWeight,
+	out float OutPdf,
+	out float OutPositionBiasSign)
+{
+	OutDirection = 0;
+	OutWeight = 0;
+	OutPdf = 0;
+	OutPositionBiasSign = 1;
+
+	// NOTE: IsFrontFace() determines polygon orientation, because the normal is always flipped towards in the incoming ray
+	bool bIsEntering = Payload.IsFrontFace();
+	float3 N_World = GetWorldNormal(Payload);
+	float3 V_World = -RayDirection;
+
+	// #dxr_todo: Maintain a refraction stack on the path tracing payload
+	float Ior = Payload.Ior;
+	float Eta = bIsEntering ? Ior : 1.0 / Ior;
+
+	float2 Alpha;
+	float3x3 Basis = GetGGXBasis(Payload.Roughness, Payload.Anisotropy, Payload.WorldNormal, Payload.WorldTangent, Alpha);
+
+	float3 V = mul(Basis, V_World);
+	float NoV = saturate(V.z);
+
+	if (NoV == 0)
+	{
+		// invalid grazing angle
+		return;
+	}
+
+	float3 H = ImportanceSampleVisibleGGX_aniso(UniformSampleDisk(RandSample.xy), Alpha, V).xyz;
+	float3 N = float3(0, 0, 1);
+	float NoH = saturate(H.z);
+	float VoH = saturate(dot(V, H));
+
+	float3 L = 0;
+	float F = 0;
+	if (SampleRefraction(-V, H, Eta, RandSample.z, L, F))
+	{
+		const float2 GGXResult = GGXEvalRefraction(L, V, H, Alpha, Eta);
+
+		OutWeight = GGXResult.x;
+		OutPdf = F * GGXResult.y;
+
+		OutPositionBiasSign = -1;
+	}
+	else
+	{
+		const float2 GGXResult = GGXEvalReflection(L, V, H, Alpha);
+		OutWeight = GGXResult.x;
+		OutPdf = F * GGXResult.y;
+	}
+
+	// transform to world space
+	OutDirection = normalize(mul(L, Basis));
+}
+
+void RoughGlass_EvalMaterial(
+	float3 IncomingDirection,
+	float3 OutgoingDirection,
+	FMaterialClosestHitPayload Payload,
+	out float3 OutWeight,
+	out float OutPdf)
+{
+	OutWeight = 0;
+	OutPdf = 0;
+
+	// NOTE: IsFrontFace() determines polygon orientation, because the normal is always flipped towards in the incoming ray
+	bool bIsEntering = Payload.IsFrontFace();
+	float3 V_World = -IncomingDirection;
+	float3 L_World = OutgoingDirection;
+	float3 N_World = GetWorldNormal(Payload);
+
+	// #dxr_todo: Maintain a refraction stack on the path tracing payload
+	float Ior = Payload.Ior;
+	float Eta = bIsEntering ? Ior : 1.0 / Ior;
+
+	float2 Alpha;
+	float3x3 Basis = GetGGXBasis(Payload.Roughness, Payload.Anisotropy, Payload.WorldNormal, Payload.WorldTangent, Alpha);
+
+	// move vectors into right shading frame
+	float3 V = mul(Basis, V_World);
+	float3 L = mul(Basis, L_World);
+
+	if (V.z <= 0)
+	{
+		// invalid input
+		return;
+	}
+
+	float NoV = saturate(V.z);
+	if (L.z > 0)
+	{
+		// reflection side
+		const float3 H = normalize(L + V);
+		const float VoH = saturate(dot(V, H));
+		const float F = FresnelReflectance(VoH, Eta);
+
+		const float2 GGXResult = GGXEvalReflection(L, V, H, Alpha);
+		OutWeight = GGXResult.x;
+		OutPdf = F * GGXResult.y;
+	}
+	else if (L.z < 0)
+	{
+		// refraction side
+		float NoL = saturate(-L.z);
+		float3 Ht = -(Eta * L + V);
+		Ht = (Eta < 1.0f) ? -Ht : Ht;
+		float3 Htn = normalize(Ht);
+
+		float VoH = dot(V, Htn);
+		float F = 1.0f - FresnelReflectance(VoH, Eta);
+		if (F > 0)
+		{
+			const float2 GGXResult = GGXEvalRefraction(L, V, Htn, Alpha, Eta);
+
+			OutWeight = GGXResult.x;
+			OutPdf = F * GGXResult.y;
+		}
+	}
+}
diff --git a/Engine/Shaders/Private/PathTracing/Material/PathTracingSpecularTransmission.ush b/Engine/Shaders/Private/PathTracing/Material/PathTracingSpecularTransmission.ush
--- a/Engine/Shaders/Private/PathTracing/Material/PathTracingSpecularTransmission.ush
+++ b/Engine/Shaders/Private/PathTracing/Material/PathTracingSpecularTransmission.ush
@@ -1,11 +1,14 @@
 // Copyright Epic Games, Inc. All Rights Reserved.
 
 /*=============================================================================================
-PathTracingSpecularTransmission.usf: Refractive BSDFs (pure singular and microfacet based)
+PathTracingSpecularTransmission.ush: Refractive BSDFs (pure singular)
 ===============================================================================================*/
 
 #pragma once
 
+#include "PathTracingMaterialCommon.ush"
+#include "PathTracingGlossy.ush"
+
 void SpecularRefraction_SampleMaterial(
 	float3 RayDirection,
 	FMaterialClosestHitPayload Payload,
@@ -44,224 +47,3 @@ void SpecularTransmission_EvalMaterial(out float3 OutThroughput, out float OutPd
 	OutPdf = 0;
 }
 
-void RoughGlass_SampleMaterial(
-	float3 RayDirection,
-	FMaterialClosestHitPayload Payload,
-	float4 RandSample,
-	out float3 OutDirection,
-	out float3 OutWeight,
-	out float OutPdf,
-	out float OutPositionBiasSign)
-{
-
-	OutDirection = 0;
-	OutWeight = 0;
-	OutPdf = 0;
-	OutPositionBiasSign = 1;
-
-	// NOTE: IsFrontFace() determines polygon orientation, because the normal is always flipped towards in the incoming ray
-	bool bIsEntering = Payload.IsFrontFace();
-	float3 N_World = GetWorldNormal(Payload);
-	float3 V_World = -RayDirection;
-
-	// #dxr_todo: Maintain a refraction stack on the path tracing payload
-	float Ior = Payload.Ior;
-	float Eta = bIsEntering ? Ior : 1.0 / Ior;
-
-	float Roughness = max(GetRoughness(Payload), GGX_MIN_ROUGHNESS);
-	float2 Alpha;
-	float3x3 Basis;
-	if (Payload.HasAnisotropy())
-	{
-		// Material is anisotropic, our shading frame should use the tangent vector
-		GetAnisotropicRoughness(Roughness * Roughness, Payload.Anisotropy, Alpha.x, Alpha.y);
-		float3 X = Payload.WorldTangent;
-		float3 Y = normalize(cross(N_World, X));
-		Basis = float3x3(cross(Y, N_World), Y, N_World);
-	}
-	else
-	{
-		// No anisotropy, just pick an arbitrary basis
-		Alpha = Roughness * Roughness;
-		Basis = GetTangentBasis(N_World);
-	}
-
-
-	float3 V = mul(Basis, V_World);
-	float NoV = saturate(V.z);
-
-	if (NoV == 0)
-	{
-		// invalid grazing angle
-		return;
-	}
-
-	float3 H = ImportanceSampleVisibleGGX_aniso(UniformSampleDisk(RandSample.xy), Alpha, V).xyz;
-	float3 N = float3(0, 0, 1);
-	float NoH = saturate(H.z);
-	float VoH = saturate(dot(V, H));
-
-	float3 L = 0;
-	float Fresnel = 0;
-	if (SampleRefraction(-V, H, Eta, RandSample.z, L, Fresnel))
-	{
-		// ray should be refracted
-		if (L.z >= 0)
-		{
-			// masked direction
-			return;
-		}
-		float NoL = saturate(-L.z);
-		float LoH = dot(L, H);
-
-		// FIXME: this can be further simplified
-		const float CosThetaM4 = Pow2(NoH);
-		const float TanThetaM2 = (Pow2(H.x / Alpha.x) + Pow2(H.y / Alpha.y)) / Pow2(NoH);
-		const float D = 1 / (PI * Pow2(1 + TanThetaM2) * Alpha.x * Alpha.y * CosThetaM4);
-		const float Lambda_o = 0.5 * (-1.0 + sqrt(1.0 + (Pow2(L.x * Alpha.x) + Pow2(L.y * Alpha.y)) / Pow2(L.z)));
-		const float Lambda_i = 0.5 * (-1.0 + sqrt(1.0 + (Pow2(V.x * Alpha.x) + Pow2(V.y * Alpha.y)) / Pow2(V.z)));
-		const float G2 = 1 / (Lambda_i + Lambda_o + 1);
-		const float G1 = 1 / (Lambda_i + 1);
-
-		const float3 Ht = (Eta * L + V);
-
-		OutWeight = (G2 / G1);
-		OutPdf = (Fresnel * abs(LoH * VoH) * (Eta * Eta) * (G1 * D)) / (dot(Ht, Ht) * abs(V.z));
-
-		OutPositionBiasSign = -1;
-	}
-	else
-	{
-		// ray should be reflected
-		if (L.z <= 0)
-		{
-			// masked direction
-			return;
-		}
-		float NoL = saturate(L.z);
-
-		// Specular lobe
-		if (NoL > 0 && NoV > 0)
-		{
-			// Assume NoH > 0
-			// FIXME: this can be further simplified
-			const float CosThetaM4 = Pow2(NoH);
-			const float TanThetaM2 = (Pow2(H.x / Alpha.x) + Pow2(H.y / Alpha.y)) / Pow2(NoH);
-			const float D = 1 / (PI * Pow2(1 + TanThetaM2) * Alpha.x * Alpha.y * CosThetaM4);
-			const float Lambda_o = 0.5 * (-1.0 + sqrt(1.0 + (Pow2(L.x * Alpha.x) + Pow2(L.y * Alpha.y)) / Pow2(L.z)));
-			const float Lambda_i = 0.5 * (-1.0 + sqrt(1.0 + (Pow2(V.x * Alpha.x) + Pow2(V.y * Alpha.y)) / Pow2(V.z)));
-			const float G2 = 1 / (Lambda_i + Lambda_o + 1);
-			const float G1 = 1 / (Lambda_i + 1);
-
-			OutWeight = (G2 / G1);
-			OutPdf = (Fresnel * G1 * D * 0.25f) / NoV;
-		}
-	}
-
-	// transform to world space
-	OutDirection = normalize(mul(L, Basis));
-}
-
-void RoughGlass_EvalMaterial(
-	float3 IncomingDirection,
-	float3 OutgoingDirection,
-	FMaterialClosestHitPayload Payload,
-	out float3 OutWeight,
-	out float OutPdf)
-{
-	OutWeight = 0;
-	OutPdf = 0;
-
-	// NOTE: IsFrontFace() determines polygon orientation, because the normal is always flipped towards in the incoming ray
-	bool bIsEntering = Payload.IsFrontFace();
-	float3 V_World = -IncomingDirection;
-	float3 L_World = OutgoingDirection;
-	float3 N_World = GetWorldNormal(Payload);
-
-	// #dxr_todo: Maintain a refraction stack on the path tracing payload
-	float Ior = Payload.Ior;
-	float Eta = bIsEntering ? Ior : 1.0 / Ior;
-
-	float Roughness = max(GetRoughness(Payload), GGX_MIN_ROUGHNESS);
-	float2 Alpha;
-	float3x3 Basis;
-	if (Payload.HasAnisotropy())
-	{
-		// Material is anisotropic, our shading frame should use the tangent vector
-		GetAnisotropicRoughness(Roughness * Roughness, Payload.Anisotropy, Alpha.x, Alpha.y);
-		float3 X = Payload.WorldTangent;
-		float3 Y = normalize(cross(N_World, X));
-		Basis = float3x3(cross(Y, N_World), Y, N_World);
-	}
-	else
-	{
-		// No anisotropy, just pick an arbitrary basis
-		Alpha = Roughness * Roughness;
-		Basis = GetTangentBasis(N_World);
-	}
-
-	// move vectors into right shading frame
-	float3 V = mul(Basis, V_World);
-	float3 L = mul(Basis, L_World);
-
-	if (V.z <= 0)
-	{
-		// invalid input
-		return;
-	}
-
-	float NoV = saturate(V.z);
-	if (L.z > 0)
-	{
-		// reflection side
-		float3 H = normalize(L + V);
-		float NoL = saturate(L.z);
-		float NoH = saturate(H.z);
-		float VoH = saturate(dot(V, H));
-
-		// FIXME: this can be further simplified
-		const float CosThetaM4 = Pow2(NoH);
-		const float TanThetaM2 = (Pow2(H.x / Alpha.x) + Pow2(H.y / Alpha.y)) / Pow2(NoH);
-		const float D = 1 / (PI * Pow2(1 + TanThetaM2) * Alpha.x * Alpha.y * CosThetaM4);
-		const float Lambda_o = 0.5 * (-1.0 + sqrt(1.0 + (Pow2(L.x * Alpha.x) + Pow2(L.y * Alpha.y)) / Pow2(L.z)));
-		const float Lambda_i = 0.5 * (-1.0 + sqrt(1.0 + (Pow2(V.x * Alpha.x) + Pow2(V.y * Alpha.y)) / Pow2(V.z)));
-		const float G2 = 1 / (Lambda_i + Lambda_o + 1);
-		const float G1 = 1 / (Lambda_i + 1);
-
-		const float Fr = FresnelReflectance(VoH, Eta);
-
-		OutWeight = (G2 / G1);
-		OutPdf = (Fr * G1 * D * 0.25f) / NoV;
-	}
-	else if (L.z < 0)
-	{
-		// refraction side
-		float NoL = saturate(-L.z);
-		float3 Ht = -(Eta * L + V);
-		Ht = (Eta < 1.0f) ? -Ht : Ht;
-		float3 Htn = normalize(Ht);
-
-		float VoH = dot(V, Htn);
-		float Ft = 1.0f - FresnelReflectance(VoH, Eta);
-		if (Ft > 0)
-		{
-			const float LoH = dot(L, Htn);
-			const float NoH = Htn.z;
-			if (NoH <= 0.0)
-			{
-				return;
-			}
-			// FIXME: this can be further simplified
-			const float CosThetaM4 = Pow2(NoH);
-			const float TanThetaM2 = (Pow2(Htn.x / Alpha.x) + Pow2(Htn.y / Alpha.y)) / Pow2(NoH);
-			const float D = 1 / (PI * Pow2(1 + TanThetaM2) * Alpha.x * Alpha.y * CosThetaM4);
-			const float Lambda_o = 0.5 * (-1.0 + sqrt(1.0 + (Pow2(L.x * Alpha.x) + Pow2(L.y * Alpha.y)) / Pow2(L.z)));
-			const float Lambda_i = 0.5 * (-1.0 + sqrt(1.0 + (Pow2(V.x * Alpha.x) + Pow2(V.y * Alpha.y)) / Pow2(V.z)));
-			const float G2 = 1 / (Lambda_i + Lambda_o + 1);
-			const float G1 = 1 / (Lambda_i + 1);
-
-			OutWeight = (G2 / G1);
-			OutPdf = (Ft * abs(LoH * VoH) * (Eta * Eta) * (G1 * D)) / (dot(Ht, Ht) * abs(V.z));
-		}
-	}
-}
\ No newline at end of file
diff --git a/Engine/Shaders/Private/PathTracing/Material/PathTracingThinGlass.ush b/Engine/Shaders/Private/PathTracing/Material/PathTracingThinGlass.ush
new file mode 100644
--- /dev/null
+++ b/Engine/Shaders/Private/PathTracing/Material/PathTracingThinGlass.ush
@@ -0,0 +1,177 @@
+// Copyright Epic Games, Inc. All Rights Reserved.
+
+/*=============================================================================================
+PathTracingThinGlass.ush: Microfacet BSDF for thin glass
+===============================================================================================*/
+
+#pragma once
+
+#include "PathTracingMaterialCommon.ush"
+#include "PathTracingFresnel.ush"
+#include "PathTracingGlossy.ush"
+
+
+struct SlabWeights
+{
+	float3 Reflected;
+	float3 Transmitted;
+};
+
+// return the amount of reflected and transmitted energy for an infinite slab viewed from a particular direction
+// taking into account the (repeated) Fresnel terms and the transmittance through the slab
+SlabWeights ComputeThinSlabWeights(float3 SlabColor, float CosV, float Eta)
+{
+	// Compute Fresnel reflection off the top of the slab
+	// The Fresnel refletion at the bottom layer is exactly equal so we can assume energy is split into R and T at each bounce
+	const float R = FresnelReflectance(CosV, Eta);
+	const float T = 1.0 - R;
+
+
+	// The slab color is normalized to unit thickness, so we only need to account for the refracted viewing angle
+	const float3 A = pow(SlabColor, Eta / sqrt(Eta * Eta - 1.0 + CosV * CosV));
+
+	// figure out how much is reflected/transmitted from a single interaction
+	float3 SlabR = R;
+	float3 SlabT = T * A * T;
+
+	// if there was some transmission -- we must also account for the infinite series of bounces within the slab
+	// which can be computed exactly with geometric series
+	if (R < 1.0)
+	{
+		SlabR += R * Pow2(T * A) / (1.0 - Pow2(R * A));
+		SlabT /= 1.0 - Pow2(R * A);
+	}
+
+	SlabWeights Result = (SlabWeights) 0;
+	Result.Reflected = SlabR;
+	Result.Transmitted = SlabT;
+	return Result;
+}
+
+void RoughThinGlass_EvalMaterial(
+	float3 IncomingDirection,
+	float3 OutgoingDirection,
+	FMaterialClosestHitPayload Payload,
+	out float3 OutWeight,
+	out float OutPdf)
+{
+	OutWeight = 0;
+	OutPdf = 0;
+
+	float3 V_World = -IncomingDirection;
+	float3 L_World = OutgoingDirection;
+	float3 N_World = GetWorldNormal(Payload);
+
+
+	float2 Alpha;
+	float3x3 Basis = GetGGXBasis(Payload.Roughness, Payload.Anisotropy, Payload.WorldNormal, Payload.WorldTangent, Alpha);
+
+	// move vectors into right shading frame
+	float3 V = mul(Basis, V_World);
+	float3 L = mul(Basis, L_World);
+
+	if (V.z <= 0)
+	{
+		// invalid input
+		return;
+	}
+
+	const bool IsTransmission = L.z < 0;
+	L.z = abs(L.z); // push L to the same side as V
+
+	const float NoV = saturate(V.z);
+	const float3 H = normalize(L + V);
+	const float VoH = saturate(dot(V, H));
+
+	const float Eta = Payload.Ior;
+
+	SlabWeights Result = ComputeThinSlabWeights(Payload.CustomData.xyz, VoH, Payload.Ior);
+
+	const float ProbR = LobeSelectionProb(Result.Reflected, Result.Transmitted);
+
+	const float2 GGXResult = GGXEvalReflection(L, V, H, Alpha);
+	OutWeight = GGXResult.x;
+	OutPdf = GGXResult.y;
+
+	if (IsTransmission)
+	{
+		OutWeight *= Result.Transmitted / (1.0 - ProbR);
+		OutPdf *= 1.0 - ProbR;
+	}
+	else
+	{
+		OutWeight *= Result.Reflected / ProbR;
+		OutPdf *= ProbR;
+	}
+}
+
+
+void RoughThinGlass_SampleMaterial(
+	float3 RayDirection,
+	FMaterialClosestHitPayload Payload,
+	float4 RandSample,
+	out float3 OutDirection,
+	out float3 OutWeight,
+	out float OutPdf,
+	out float OutPositionBiasSign)
+{
+	OutDirection = 0;
+	OutWeight = 0;
+	OutPdf = 0;
+	OutPositionBiasSign = 1;
+
+	float3 N_World = GetWorldNormal(Payload);
+	float3 V_World = -RayDirection;
+
+	float Eta = Payload.Ior;
+
+	float2 Alpha;
+	float3x3 Basis = GetGGXBasis(Payload.Roughness, Payload.Anisotropy, Payload.WorldNormal, Payload.WorldTangent, Alpha);
+
+	float3 V = mul(Basis, V_World);
+	float NoV = saturate(V.z);
+
+	if (NoV == 0)
+	{
+		// invalid grazing angle
+		return;
+	}
+
+	float3 H = ImportanceSampleVisibleGGX_aniso(UniformSampleDisk(RandSample.xy), Alpha, V).xyz;
+	float3 N = float3(0, 0, 1);
+	float NoH = saturate(H.z);
+	float VoH = saturate(dot(V, H));
+
+	float3 L = reflect(-V, H);
+	if (L.z <= 0) {
+		// invalid output direction, exit early
+		return;
+	}
+
+	float NoL = saturate(L.z);
+
+	SlabWeights Result = ComputeThinSlabWeights(Payload.CustomData.xyz, VoH, Payload.Ior);
+
+	const float ProbR = LobeSelectionProb(Result.Reflected, Result.Transmitted);
+
+	const float2 GGXResult = GGXEvalReflection(L, V, H, Alpha);
+	OutWeight = GGXResult.x;
+	OutPdf = GGXResult.y;
+
+	// decide between reflection/transmission
+	if (RandSample.z < ProbR)
+	{
+		// reflected ray
+		OutWeight *= Result.Reflected / ProbR;
+		OutDirection = normalize(mul(L, Basis));
+		OutPdf *= ProbR;
+	}
+	else
+	{
+		// transmitted ray
+		OutWeight *= Result.Transmitted / (1.0 - ProbR);
+		OutDirection = normalize(mul(float3(L.xy, -L.z), Basis)); // flip reflection to other side
+		OutPdf *= 1.0 - ProbR;
+		OutPositionBiasSign = -1;
+	}
+}
diff --git a/Engine/Shaders/Private/PathTracing/PathTracing.usf b/Engine/Shaders/Private/PathTracing/PathTracing.usf
--- a/Engine/Shaders/Private/PathTracing/PathTracing.usf
+++ b/Engine/Shaders/Private/PathTracing/PathTracing.usf
@@ -98,7 +98,7 @@ FMaterialClosestHitPayload TraceTransparentRay(RayDesc Ray, FRayCone RayCone, bo
 			AccumulateRadiance(Radiance, PathThroughput * HitPayload.Radiance);
 		}
 
-		float3 Transparency = GetMaterialTransparency(HitPayload);
+		float3 Transparency = GetMaterialTransparency(HitPayload, Ray.Direction);
 
 		if (!LastBounce)
 		{
@@ -173,11 +173,8 @@ float3 TraceTransparentVisibilityRay(RayDesc Ray, uint2 PixelCoord, float PathRo
 	// The payload is used to carry the path throughput (for transparent shadows)
 	// and current path roughness (for approximate caustics)
 	PackedPayload.SetPathTracingVisibilityRay();
-	PackedPayload.RadianceAndNormal[0] = asuint(1.0);
-	PackedPayload.RadianceAndNormal[1] = asuint(1.0);
-	PackedPayload.RadianceAndNormal[2] = asuint(1.0);
-	PackedPayload.MetallicAndSpecularAndRoughness = asuint(PathRoughness);
-
+	PackedPayload.SetRayThroughput(1.0);
+	PackedPayload.SetPathRoughness(PathRoughness);
 	PackedPayload.SetPixelCoord(PixelCoord);
 
 	// Trace the ray
@@ -194,11 +191,7 @@ float3 TraceTransparentVisibilityRay(RayDesc Ray, uint2 PixelCoord, float PathRo
 	if (PackedPayload.IsMiss())
 	{
 		// we didn't hit any geo - just return whatever throughput is left
-		return float3(
-			asfloat(PackedPayload.RadianceAndNormal[0]),
-			asfloat(PackedPayload.RadianceAndNormal[1]),
-			asfloat(PackedPayload.RadianceAndNormal[2])
-		);
+		return PackedPayload.GetRayThroughput();
 	}
 
 	// we registered a hit -- assumed to be fully opaque
diff --git a/Engine/Shaders/Private/RayTracing/RayTracingCommon.ush b/Engine/Shaders/Private/RayTracing/RayTracingCommon.ush
--- a/Engine/Shaders/Private/RayTracing/RayTracingCommon.ush
+++ b/Engine/Shaders/Private/RayTracing/RayTracingCommon.ush
@@ -473,6 +473,33 @@ struct FPackedMaterialClosestHitPayload : FMinimalPayload
 
 	uint GetInstanceID()	 { return PackedCustomData; }
 	float3 GetGeometryNormal() { return GetWorldNormal(); }
+
+	// These method are meant to be used only when IsPathTracingVisibilityRay() is true
+	float3 GetRayThroughput()
+	{
+		return float3(
+			asfloat(RadianceAndNormal[0]),
+			asfloat(RadianceAndNormal[1]),
+			asfloat(RadianceAndNormal[2])
+		);
+	}
+
+	void SetRayThroughput(float3 RayThroughput)
+	{
+		RadianceAndNormal[0] = asuint(RayThroughput.x);
+		RadianceAndNormal[1] = asuint(RayThroughput.y);
+		RadianceAndNormal[2] = asuint(RayThroughput.z);
+	}
+
+	float GetPathRoughness()
+	{
+		return asfloat(MetallicAndSpecularAndRoughness);
+	}
+
+	void SetPathRoughness(float PathRoughness)
+	{
+		MetallicAndSpecularAndRoughness = asuint(PathRoughness);
+	}
 };
 
 FPackedMaterialClosestHitPayload PackRayTracingPayload(FMaterialClosestHitPayload Input, in FRayCone RayCone)
diff --git a/Engine/Shaders/Private/RayTracing/RayTracingMaterialHitShaders.usf b/Engine/Shaders/Private/RayTracing/RayTracingMaterialHitShaders.usf
--- a/Engine/Shaders/Private/RayTracing/RayTracingMaterialHitShaders.usf
+++ b/Engine/Shaders/Private/RayTracing/RayTracingMaterialHitShaders.usf
@@ -33,6 +33,7 @@
 #include "/Engine/Private/VirtualTextureCommon.ush" 
 #include "/Engine/Private/LightmapCommon.ush"
 #include "/Engine/Private/PathTracing/Material/PathTracingFresnel.ush"
+#include "/Engine/Private/PathTracing/Material/PathTracingThinGlass.ush"
 
 /** Computes sky diffuse lighting, including precomputed shadowing. */
 void GetSkyLighting(
@@ -400,6 +401,7 @@ RAY_TRACING_ENTRY_CLOSEST_HIT(MaterialCHS,
 	Payload.CustomData = float4(GetMaterialCustomData0(MaterialParameters), GetMaterialCustomData1(MaterialParameters), 0, 0);
 	
 	#if CLEAR_COAT_BOTTOM_NORMAL
+	if (Payload.ShadingModelID == SHADINGMODELID_CLEAR_COAT)
 	{
 		float2 oct2 = UnitVectorToOctahedron(Payload.WorldNormal);
 
@@ -434,6 +436,13 @@ RAY_TRACING_ENTRY_CLOSEST_HIT(MaterialCHS,
 		Payload.CustomData.a = saturate(GetMaterialCustomData0(MaterialParameters));
 	}
 
+#if MATERIAL_SHADINGMODEL_THIN_TRANSLUCENT
+	if (Payload.ShadingModelID == SHADINGMODELID_THIN_TRANSLUCENT)
+	{
+		Payload.CustomData.xyz = GetThinTranslucentMaterialOutput0(MaterialParameters);
+	}
+#endif
+
 #if MATERIAL_USES_ANISOTROPY
 	Payload.WorldTangent = CalculateAnisotropyTangent(MaterialParameters, PixelMaterialInputs);
 	Payload.Anisotropy = GetMaterialAnisotropy(PixelMaterialInputs);
@@ -608,14 +617,45 @@ RAY_TRACING_ENTRY_ANY_HIT(MaterialAHS,
 #elif MATERIALBLENDING_ALPHACOMPOSITE
 	float Opacity = GetMaterialOpacity(PixelMaterialInputs);
 	float Transparency = 1 - Opacity;
+#elif MATERIAL_SHADINGMODEL_THIN_TRANSLUCENT
+	// Assume Payload.ShadingModelID == SHADINGMODELID_THIN_TRANSLUCENT
+	float3 Transmission = GetThinTranslucentMaterialOutput0(MaterialParameters);
+	float Opacity = GetMaterialOpacity(PixelMaterialInputs);
+	float Ior = GetMaterialRefraction(PixelMaterialInputs).x;
+	float3 Transparency = 1 - Opacity;
+	float3 V = WorldRayDirection();
+	float3 N = normalize(MaterialParameters.WorldNormal);
+	float VoN = abs(dot(V, N));
+	if (Opacity < 1.0 && Ior != 1.0)
+	{
+		float PathRoughness = PackedPayload.GetPathRoughness();
+		float Roughness = GetMaterialRoughness(PixelMaterialInputs);
+		if (PathRoughness <= Roughness)
+		{
+			// not using fast caustic approximation - treat as opaque
+			return;
+		}
+		// compute transmission through the slab (fresnel + absorption)
+		Transparency *= ComputeThinSlabWeights(Transmission, VoN, Ior).Transmitted;
+
+		// fake caustic approximation (see comments below)
+		Transparency *= (1 - Roughness * Roughness) * saturate(PathRoughness - Roughness);
+	}
+	else
+	{
+		// no bending at the interface so no fresnel - equation above simplifies to just A
+		float3 A = pow(Transmission, 1.0 / VoN);
+		Transparency *= A;
+	}
 #elif MATERIALBLENDING_TRANSLUCENT
 	float Opacity = GetMaterialOpacity(PixelMaterialInputs);
 	float Transparency = 1 - Opacity;
 	float Ior = GetMaterialRefraction(PixelMaterialInputs).x;
+
 	// Is refraction enabled?
 	if (Transparency > 0 && Ior != 1.0)
 	{
-		float PathRoughness = asfloat(PackedPayload.MetallicAndSpecularAndRoughness);
+		float PathRoughness = PackedPayload.GetPathRoughness();
 		float Roughness = GetMaterialRoughness(PixelMaterialInputs);
 		if (PathRoughness <= Roughness)
 		{
@@ -640,17 +680,11 @@ RAY_TRACING_ENTRY_ANY_HIT(MaterialAHS,
 #endif
 
 	// Update the ray throughput (it is packed simply into the payload since we don't need to carry any other information across hits)
-	float3 RayThroughput = float3(
-		asfloat(PackedPayload.RadianceAndNormal[0]),
-		asfloat(PackedPayload.RadianceAndNormal[1]),
-		asfloat(PackedPayload.RadianceAndNormal[2])
-	);
+	float3 RayThroughput = PackedPayload.GetRayThroughput();
 
 	RayThroughput *= Transparency;
 
-	PackedPayload.RadianceAndNormal[0] = asuint(RayThroughput.x);
-	PackedPayload.RadianceAndNormal[1] = asuint(RayThroughput.y);
-	PackedPayload.RadianceAndNormal[2] = asuint(RayThroughput.z);
+	PackedPayload.SetRayThroughput(RayThroughput);
 
 	if (any(RayThroughput > 0))
 	{
