Clean up LocalVertexFactory.ush previous instance local to world transform specific to geometry collection / USE_INSTANCING_BONEMAP

diff --git a/Engine/Shaders/Private/LocalVertexFactory.ush b/Engine/Shaders/Private/LocalVertexFactory.ush
--- a/Engine/Shaders/Private/LocalVertexFactory.ush
+++ b/Engine/Shaders/Private/LocalVertexFactory.ush
@@ -387,6 +387,7 @@ struct FVertexFactoryIntermediates
 };
 
 #if USE_INSTANCING
+
 float4x4 GetInstanceTransform(FVertexFactoryIntermediates Intermediates)
 {
 	return float4x4(
@@ -396,16 +397,20 @@ float4x4 GetInstanceTransform(FVertexFactoryIntermediates Intermediates)
 		float4(Intermediates.InstanceOrigin.xyz, 1.0f));
 }
 
-#if USE_INSTANCING && USE_INSTANCING_BONEMAP
 float4x4 GetInstancePrevTransform(FVertexFactoryIntermediates Intermediates)
 {
+#if USE_INSTANCING_BONEMAP
+	// when using geometry collections, we can correctly compute velocities because we pass along previous transforms.
 	return float4x4(
 		float4(Intermediates.InstancePrevTransform1.xyz, 0.0f),
 		float4(Intermediates.InstancePrevTransform2.xyz, 0.0f),
 		float4(Intermediates.InstancePrevTransform3.xyz, 0.0f),
 		float4(Intermediates.InstancePrevOrigin.xyz, 1.0f));
-}
+#else
+	// Assumes instance transform never change, which means per-instance motion will cause TAA and motion blur artifacts.
+	return GetInstanceTransform(Intermediates);
 #endif
+}
 
 float4x4 GetInstanceTransform(FPositionOnlyVertexFactoryInput Input)
 {
@@ -580,18 +585,9 @@ FMaterialVertexParameters GetMaterialVertexParameters(FVertexFactoryInput Input,
 	Result.InstanceLocalToWorld = mul(GetInstanceTransform(Intermediates), GetPrimitiveData(Intermediates.PrimitiveId).LocalToWorld);
 	Result.InstanceLocalPosition = Input.Position.xyz;
 	Result.PerInstanceParams = Intermediates.PerInstanceParams;
-
 	Result.InstanceId = GetInstanceId(Input.InstanceId); 
 	Result.InstanceOffset = InstanceOffset;
-
-	#if USE_INSTANCING_BONEMAP
-		// when using geometry collections, we can correctly compute velocities because we pass along previous transforms
-		Result.PrevFrameLocalToWorld = mul(GetInstancePrevTransform(Intermediates), GetPrimitiveData(Intermediates.PrimitiveId).PreviousLocalToWorld);
-	#else
-		// Assumes instance transform never change, which means per-instance
-		// motion will cause TAA and motion blur artifacts
-		Result.PrevFrameLocalToWorld = mul(GetInstanceTransform(Intermediates), GetPrimitiveData(Intermediates.PrimitiveId).PreviousLocalToWorld);
-	#endif // USE_INSTANCING_BONEMAP
+	Result.PrevFrameLocalToWorld = mul(GetInstancePrevTransform(Intermediates), GetPrimitiveData(Intermediates.PrimitiveId).PreviousLocalToWorld);
 #else
 	#if defined(GPUCULL_TODO) && VF_USE_PRIMITIVE_SCENE_DATA
 		Result.PrevFrameLocalToWorld = Intermediates.SceneData.InstanceData.PrevLocalToWorld;
@@ -601,10 +597,10 @@ FMaterialVertexParameters GetMaterialVertexParameters(FVertexFactoryInput Input,
 			Result.PerInstanceParams = Intermediates.PerInstanceParams;
 			Result.PrimitiveLocalInstanceIndex = Intermediates.PrimitiveLocalInstanceIndex;
 		#endif
-	#else // !defined(GPUCULL_TODO) && VF_USE_PRIMITIVE_SCENE_DATA
+	#else
 		Result.PrevFrameLocalToWorld = GetPrimitiveData(Intermediates.PrimitiveId).PreviousLocalToWorld;
-	#endif // defined(GPUCULL_TODO) && VF_USE_PRIMITIVE_SCENE_DATA
-#endif // USE_INSTANCING
+	#endif
+#endif
 
 	Result.PreSkinnedPosition = Intermediates.PreSkinPosition.xyz;
 	Result.PreSkinnedNormal = TangentToLocal[2]; //TangentBias(Input.TangentZ.xyz);
@@ -1258,10 +1254,7 @@ float4 VertexFactoryGetPreviousWorldPosition(FVertexFactoryInput Input, FVertexF
 	PreviousLocalToWorldTranslated[3][1] += ResolvedView.PrevPreViewTranslation.y;
 	PreviousLocalToWorldTranslated[3][2] += ResolvedView.PrevPreViewTranslation.z;
 
-#if USE_INSTANCING && !USE_INSTANCING_BONEMAP
-	float4x4 InstanceTransform = GetInstanceTransform(Intermediates);
-	return mul(mul(Input.Position, InstanceTransform), PreviousLocalToWorldTranslated);
-#elif USE_INSTANCING && USE_INSTANCING_BONEMAP
+#if USE_INSTANCING
 	float4x4 InstanceTransform = GetInstancePrevTransform(Intermediates);
 	return mul(mul(Input.Position, InstanceTransform), PreviousLocalToWorldTranslated);
 #elif GPUSKIN_PASS_THROUGH
