The sky env light pass no longer reads the gbuffer when strata is enabled.
DFAO upsample is only done once for all BSDF.
Reflection with Strata off also tested.

diff --git a/Engine/Shaders/Private/DistanceFieldAOShared.ush b/Engine/Shaders/Private/DistanceFieldAOShared.ush
--- a/Engine/Shaders/Private/DistanceFieldAOShared.ush
+++ b/Engine/Shaders/Private/DistanceFieldAOShared.ush
@@ -173,8 +173,15 @@ float3 ComputeBentNormal(float3 RecordWorldNormal, uint RelativeRecordIndex)
 
 float2 AOBufferBilinearUVMax;
 
-float3 UpsampleDFAO(float2 BufferUV, float SceneDepth, float3 WorldNormal)
+struct FUpsampleDFAOOutput
+{
+	float3 BentNormal;
+	float FadeAlpha;
+};
+FUpsampleDFAOOutput UpsampleDFAO(float2 BufferUV, float SceneDepth)
 {
+	FUpsampleDFAOOutput UpsampleDFAOOutput;
+
 	// Distance field AO was computed at 0,0 regardless of viewrect min
 	float2 DistanceFieldUVs = BufferUV - View.ViewRectMin.xy * View.BufferSizeAndInvSize.zw;
 	DistanceFieldUVs = min(DistanceFieldUVs, AOBufferBilinearUVMax);
@@ -219,11 +226,22 @@ float3 UpsampleDFAO(float2 BufferUV, float SceneDepth, float3 WorldNormal)
 
 #endif
 
+	UpsampleDFAOOutput.BentNormal = BentNormal;
 	// Fade to unoccluded in the distance
-	float FadeAlpha = saturate((AOMaxViewDistance - SceneDepth) * DistanceFadeScale);
-	BentNormal = lerp(WorldNormal, BentNormal, FadeAlpha);
+	UpsampleDFAOOutput.FadeAlpha = saturate((AOMaxViewDistance - SceneDepth) * DistanceFadeScale);
+
+	return UpsampleDFAOOutput;
+}
 
-	return BentNormal;
+float3 ApplyDFAO(FUpsampleDFAOOutput UpsampleDFAOOutput, float3 WorldNormal)
+{
+	return lerp(WorldNormal, UpsampleDFAOOutput.BentNormal, UpsampleDFAOOutput.FadeAlpha);
+}
+
+float3 UpsampleDFAO(float2 BufferUV, float SceneDepth, float3 WorldNormal)
+{
+	FUpsampleDFAOOutput UpsampleDFAOOutput = UpsampleDFAO(BufferUV, SceneDepth);
+	return ApplyDFAO(UpsampleDFAOOutput, WorldNormal);
 }
 
 #ifndef CULLED_TILE_SIZEX
diff --git a/Engine/Shaders/Private/ReflectionEnvironmentPixelShader.usf b/Engine/Shaders/Private/ReflectionEnvironmentPixelShader.usf
--- a/Engine/Shaders/Private/ReflectionEnvironmentPixelShader.usf
+++ b/Engine/Shaders/Private/ReflectionEnvironmentPixelShader.usf
@@ -47,11 +47,10 @@ struct FSkyLightVisibilityData
 	float3 SkyDiffuseLookUpAdd;
 	float3 SkyDiffuseLookUpNormal;
 };
-FSkyLightVisibilityData GetSkyLightVisibilityData(FGBufferData GBuffer, float AmbientOcclusion, float3 BentNormal)
+FSkyLightVisibilityData GetSkyLightVisibilityData(float3 SkyLightingNormal, const float3 WorldNormal, const float GBufferAO, float ScreenAO, const float3 BentNormal)
 {
 	float SkyVisibility = 1;
 	float DotProductFactor = 1;
-	float3 SkyLightingNormal = GetClearCoatBottomNormal(GBuffer, GBuffer.WorldNormal);
 
 #if APPLY_SKY_SHADOWING
 #define USE_DIRECTIONAL_OCCLUSION_ON_SKY_DIFFUSE 1
@@ -64,8 +63,8 @@ FSkyLightVisibilityData GetSkyLightVisibilityData(FGBufferData GBuffer, float Am
 		// Use more bent normal in corners
 		float BentNormalWeightFactor = SkyVisibility;
 
-		SkyLightingNormal = lerp(NormalizedBentNormal, GBuffer.WorldNormal, BentNormalWeightFactor);
-		DotProductFactor = lerp(dot(NormalizedBentNormal, GBuffer.WorldNormal), 1, BentNormalWeightFactor);
+		SkyLightingNormal = lerp(NormalizedBentNormal, WorldNormal, BentNormalWeightFactor);
+		DotProductFactor = lerp(dot(NormalizedBentNormal, WorldNormal), 1, BentNormalWeightFactor);
 	}
 #else
 	{
@@ -86,12 +85,12 @@ FSkyLightVisibilityData GetSkyLightVisibilityData(FGBufferData GBuffer, float Am
 	if (OcclusionCombineMode == 0)
 	{
 		// Combine with min which nicely avoids over-occlusion in cases where strong DFAO is present along with strong SSAO (distant trees)
-		SkyVisibility = min(SkyVisibility, min(GBuffer.GBufferAO, AmbientOcclusion));
+		SkyVisibility = min(SkyVisibility, min(GBufferAO, ScreenAO));
 	}
 	else
 	{
 		// Combine with mul, which continues to add SSAO depth even indoors.  SSAO will need to be tweaked to be less strong.
-		SkyVisibility = SkyVisibility * min(GBuffer.GBufferAO, AmbientOcclusion);
+		SkyVisibility = SkyVisibility * min(GBufferAO, ScreenAO);
 	}
 
 	// Apply AO to the sky diffuse and account for darkening due to the geometry term
@@ -144,7 +143,8 @@ float3 SkyLightDiffuse(FGBufferData GBuffer, float AmbientOcclusion, float2 Buff
 	float2 UV = BufferUV;
 	float3 Lighting = 0;
 
-	FSkyLightVisibilityData SkyVisData = GetSkyLightVisibilityData(GBuffer, AmbientOcclusion, BentNormal);
+	float3 SkyLightingNormal = GetClearCoatBottomNormal(GBuffer, GBuffer.WorldNormal);
+	FSkyLightVisibilityData SkyVisData = GetSkyLightVisibilityData(SkyLightingNormal, GBuffer.WorldNormal, GBuffer.GBufferAO, AmbientOcclusion, BentNormal);
 
 	BRANCH
 	if (GBuffer.ShadingModelID == SHADINGMODELID_TWOSIDED_FOLIAGE)
@@ -331,48 +331,45 @@ void ReflectionEnvironmentSkyLighting(
 	in float4 SvPosition : SV_Position,
 	out float4 OutColor : SV_Target0)
 {
+	uint2 PixelPos = SvPosition.xy;
 	float2 BufferUV = SvPositionToBufferUV(SvPosition);
 	float2 ScreenPosition = SvPositionToScreenPosition(SvPosition).xy;
-
-	// Sample scene textures.
-	FGBufferData GBuffer = GetGBufferDataFromSceneTextures(BufferUV);
-
-	float3 WorldPosition = mul(float4(ScreenPosition * GBuffer.Depth, GBuffer.Depth, 1), View.ScreenToWorld).xyz;
-	float3 CameraToPixel = normalize(WorldPosition - View.WorldCameraOrigin);
-	float3 N = GBuffer.WorldNormal;
-	float3 V = -CameraToPixel;
-	float NoV = saturate( dot( N, V ) );
-
-	float3 DiffuseColor = GBuffer.DiffuseColor;
-	float3 SpecularColor = GBuffer.SpecularColor;
-	RemapClearCoatDiffuseAndSpecularColor(GBuffer, NoV, DiffuseColor, SpecularColor);
-
-	// Sample the ambient occlusion that is dynamically generated every frame.
-	float AmbientOcclusion = AmbientOcclusionTexture.SampleLevel(AmbientOcclusionSampler, BufferUV, 0).r;
-	
-	uint ShadingModelID = GBuffer.ShadingModelID;
-	const bool bUnlitMaterial = ShadingModelID == SHADINGMODELID_UNLIT;
-	const bool bStrataMaterial = ShadingModelID == SHADINGMODELID_STRATA;
-
-	float3 BentNormal = GBuffer.WorldNormal;
-	#if APPLY_SKY_SHADOWING
-	{
-		BentNormal = UpsampleDFAO(BufferUV, GBuffer.Depth, GBuffer.WorldNormal);
-	}
-	#endif
 	
 	OutColor = 0.0f;
 
+#if STRATA_ENABLED
+
 #if ENABLE_DYNAMIC_SKY_LIGHT
+#if ENABLE_SKY_LIGHT
 
-#if STRATA_ENABLED && ENABLE_SKY_LIGHT
+	FStrataClassification StrataClassification = StrataUnpackClassificationData(Strata.ClassificationTexture.Load(uint3(PixelPos, 0)));
 
 	// STRATA_TODO handle the ENABLE_SKY_LIGHT && ALLOW_STATIC_LIGHTING case 
 	BRANCH
-	if (bStrataMaterial)
+	if (StrataClassification.bIsStrataMaterial)
 	{
-		uint2 PixelPos = SvPosition.xy;
 		const FStrataTopLayerData TopLayerData = StrataUnpackTopLayerData(Strata.TopLayerTexture.Load(uint3(PixelPos, 0)));
+
+		float DeviceZ = SampleDeviceZFromSceneTextures(BufferUV);
+		float SceneDepth = ConvertFromDeviceZ(DeviceZ);
+
+		float3 WorldPosition = mul(float4(ScreenPosition * SceneDepth, SceneDepth, 1), View.ScreenToWorld).xyz;
+		float3 CameraToPixel = normalize(WorldPosition - View.WorldCameraOrigin);
+		float3 N = TopLayerData.Normal;
+		float3 V = -CameraToPixel;
+		float NoV = saturate(dot(N, V));
+
+		// Sample the ambient occlusion that is dynamically generated every frame.
+		float AmbientOcclusion = AmbientOcclusionTexture.SampleLevel(AmbientOcclusionSampler, BufferUV, 0).r;
+
+		// STRATA_TODO: the top layer should handle the bent normal...
+		float3 BentNormal = N; 
+#if APPLY_SKY_SHADOWING
+		{
+			BentNormal = UpsampleDFAO(BufferUV, SceneDepth, N);
+		}
+#endif
+
 		const float4 SSR = Texture2DSample(ScreenSpaceReflectionsTexture, ScreenSpaceReflectionsSampler, BufferUV);
 
 		// SSR is computed for the top level by averaging BSDF components.
@@ -384,7 +381,8 @@ void ReflectionEnvironmentSkyLighting(
 
 		float4 StrataEnvLighting = StrataEnvironmentLighting(
 			BufferUV,
-			GBuffer,
+			SceneDepth,
+			StrataClassification.GBufferAO,
 			WorldPosition,
 			CameraToPixel,
 			V,
@@ -401,7 +399,38 @@ void ReflectionEnvironmentSkyLighting(
 		OutColor.xyz = SkyLighting;
 	}
 
-#else // STRATA_ENABLED && ENABLE_SKY_LIGHT
+#endif // ENABLE_SKY_LIGHT
+#endif // ENABLE_DYNAMIC_SKY_LIGHT 
+
+#else // STRATA_ENABLED
+
+	// Sample scene textures.
+	FGBufferData GBuffer = GetGBufferDataFromSceneTextures(BufferUV);
+
+	uint ShadingModelID = GBuffer.ShadingModelID;
+	const bool bUnlitMaterial = ShadingModelID == SHADINGMODELID_UNLIT;
+
+	float3 WorldPosition = mul(float4(ScreenPosition * GBuffer.Depth, GBuffer.Depth, 1), View.ScreenToWorld).xyz;
+	float3 CameraToPixel = normalize(WorldPosition - View.WorldCameraOrigin);
+	float3 N = GBuffer.WorldNormal;
+	float3 V = -CameraToPixel;
+	float NoV = saturate(dot(N, V));
+
+	float3 DiffuseColor = GBuffer.DiffuseColor;
+	float3 SpecularColor = GBuffer.SpecularColor;
+	RemapClearCoatDiffuseAndSpecularColor(GBuffer, NoV, DiffuseColor, SpecularColor);
+
+	// Sample the ambient occlusion that is dynamically generated every frame.
+	float AmbientOcclusion = AmbientOcclusionTexture.SampleLevel(AmbientOcclusionSampler, BufferUV, 0).r;
+
+	float3 BentNormal = GBuffer.WorldNormal;
+#if APPLY_SKY_SHADOWING
+	{
+		BentNormal = UpsampleDFAO(BufferUV, GBuffer.Depth, GBuffer.WorldNormal);
+	}
+#endif
+
+#if ENABLE_DYNAMIC_SKY_LIGHT
 
 	BRANCH
 	if (!bUnlitMaterial) // Only light pixels marked as lit
@@ -417,17 +446,14 @@ void ReflectionEnvironmentSkyLighting(
 		OutColor = LightAccumulator_GetResult(LightAccumulator);
 	}
 
-#endif // STRATA_ENABLED && ENABLE_SKY_LIGHT
+#endif // ENABLE_DYNAMIC_SKY_LIGHT 
 
-#endif // ENABLE_DYNAMIC_SKY_LIGHT
-
-
-#if !STRATA_ENABLED
 	BRANCH
 	if (!bUnlitMaterial && ShadingModelID != SHADINGMODELID_HAIR)
 	{
 		OutColor.xyz += ReflectionEnvironment(GBuffer, AmbientOcclusion, BufferUV, WorldPosition, V, SvPosition, BentNormal, SpecularColor);
 	}
-#endif
+
+#endif // STRATA_ENABLED
 }
 
diff --git a/Engine/Shaders/Private/Strata/Strata.ush b/Engine/Shaders/Private/Strata/Strata.ush
--- a/Engine/Shaders/Private/Strata/Strata.ush
+++ b/Engine/Shaders/Private/Strata/Strata.ush
@@ -171,6 +171,7 @@ float3 StrataUnpackNormal(uint PackedNormal);
 struct FStrataClassification
 {
 	uint ShadingModels;
+	uint GBufferAO;
 	bool bIsStrataMaterial;
 	bool bHasSubsurface;
 	bool bHasSubsurfaceProfile;
@@ -182,22 +183,24 @@ uint StrataPackClassificationData(FStrataClassification In)
 {
 	return 
 		(In.ShadingModels & 0xFF) | 
-		((In.bIsStrataMaterial		? 1 : 0) << 8 ) |
-		((In.bHasSubsurface			? 1 : 0) << 9 ) | 
-		((In.bHasSubsurfaceProfile	? 1 : 0) << 10) |
-		((In.bHasTransmission		? 1 : 0) << 11) |
-		((In.bIsSimple				? 1 : 0) << 12);
+		((In.GBufferAO & 0xFF)				 << 8  ) |
+		((In.bIsStrataMaterial		? 1 : 0) << 16 ) |
+		((In.bHasSubsurface			? 1 : 0) << 17 ) | 
+		((In.bHasSubsurfaceProfile	? 1 : 0) << 18) |
+		((In.bHasTransmission		? 1 : 0) << 19) |
+		((In.bIsSimple				? 1 : 0) << 20);
 }
 
 FStrataClassification StrataUnpackClassificationData(uint In)
 {
 	FStrataClassification Out;
-	Out.ShadingModels			= (In & 0xFF);
-	Out.bIsStrataMaterial		= (In >> 8 ) & 0x1;
-	Out.bHasSubsurface			= (In >> 9 ) & 0x1;
-	Out.bHasSubsurfaceProfile	= (In >> 10) & 0x1;
-	Out.bHasTransmission		= (In >> 11) & 0x1;
-	Out.bIsSimple				= (In >> 12) & 0x1;
+	Out.ShadingModels			= (In >> 0 ) & 0xFF;
+	Out.GBufferAO				= (In >> 8 ) & 0xFF;
+	Out.bIsStrataMaterial		= (In >> 16) & 0x1;
+	Out.bHasSubsurface			= (In >> 17) & 0x1;
+	Out.bHasSubsurfaceProfile	= (In >> 18) & 0x1;
+	Out.bHasTransmission		= (In >> 19) & 0x1;
+	Out.bIsSimple				= (In >> 20) & 0x1;
 	return Out;
 }
 
diff --git a/Engine/Shaders/Private/Strata/StrataEnvironmentLighting.ush b/Engine/Shaders/Private/Strata/StrataEnvironmentLighting.ush
--- a/Engine/Shaders/Private/Strata/StrataEnvironmentLighting.ush
+++ b/Engine/Shaders/Private/Strata/StrataEnvironmentLighting.ush
@@ -10,7 +10,8 @@
 // Unpack BSDF on-the-fly
 float4 StrataEnvironmentLighting(
 	float2 BufferUV,
-	FGBufferData GBuffer,
+	float SceneDepth,
+	float GBufferAO,
 	float3 WorldPosition,
 	float3 CameraToPixel,
 	float3 V,
@@ -29,6 +30,8 @@ float4 StrataEnvironmentLighting(
 	CommonFactor *= View.PreExposure;
 #endif
 
+	// Sample DFAO only once
+	FUpsampleDFAOOutput UpsampleDFAOOutput = UpsampleDFAO(BufferUV, SceneDepth);
 
 #if STRATA_FASTPATH
 	uint BSDFIndex = 0;
@@ -59,12 +62,12 @@ float4 StrataEnvironmentLighting(
 			// Apply bent normal if needed
 			#if APPLY_SKY_SHADOWING
 			{
-				StrataEnvLight.DiffuseNormal = UpsampleDFAO(BufferUV, GBuffer.Depth, StrataEnvLight.DiffuseNormal); // Bent normal
+				StrataEnvLight.DiffuseNormal = ApplyDFAO(UpsampleDFAOOutput, StrataEnvLight.DiffuseNormal); // Bent normal
 			}
 			#endif
 
 			// Compute the common sky visibility factors
-			FSkyLightVisibilityData SkyVisData = GetSkyLightVisibilityData(GBuffer, AmbientOcclusion, StrataEnvLight.DiffuseNormal);
+			FSkyLightVisibilityData SkyVisData = GetSkyLightVisibilityData(StrataBSDFContext.N, StrataBSDFContext.N, GBufferAO, AmbientOcclusion, StrataEnvLight.DiffuseNormal);
 
 			// Finally sample the sky diffuse contribution (spherical harmonic, Lambert BRDF)
 			float3 DiffuseLookup = GetSkySHDiffuse(StrataEnvLight.DiffuseNormal) * View.SkyLightColor.rgb;			// STRATA_TODO compute SH coefficients for ON or Chan
diff --git a/Engine/Shaders/Private/Strata/StrataMaterialClassification.usf b/Engine/Shaders/Private/Strata/StrataMaterialClassification.usf
--- a/Engine/Shaders/Private/Strata/StrataMaterialClassification.usf
+++ b/Engine/Shaders/Private/Strata/StrataMaterialClassification.usf
@@ -146,6 +146,7 @@ void MainPS(
 	FStrataTopLayerData TopLayerData			= (FStrataTopLayerData)0;
 
 	ClassificationData.bIsStrataMaterial = GBuffer.ShadingModelID == SHADINGMODELID_STRATA;
+	ClassificationData.GBufferAO = GBuffer.GBufferAO;
 	ClassificationData.bIsSimple = true; // non strata material or sky pixels are considered as simple.
 	if (ClassificationData.bIsStrataMaterial)
 	{
