Add support back for edge color in strata.

Run with preflight.

diff --git a/Engine/Shaders/Private/BRDF.ush b/Engine/Shaders/Private/BRDF.ush
--- a/Engine/Shaders/Private/BRDF.ush
+++ b/Engine/Shaders/Private/BRDF.ush
@@ -441,7 +441,12 @@ float3 F_Schlick( float3 SpecularColor, float VoH )
 	
 	// Anything less than 2% is physically impossible and is instead considered to be shadowing
 	return saturate( 50.0 * SpecularColor.g ) * Fc + (1 - Fc) * SpecularColor;
-	
+}
+
+float3 F_Schlick(float3 F0, float3 F90, float VoH)
+{
+	float Fc = Pow5(1 - VoH);
+	return F90 * Fc + (1 - Fc) * F0;
 }
 
 float3 F_Fresnel( float3 SpecularColor, float VoH )
@@ -528,7 +533,7 @@ half3 EnvBRDF( half3 SpecularColor, half Roughness, half NoV )
 	return GF;
 }
 
-half2 EnvBRDFApproxLazarov(half3 SpecularColor, half Roughness, half NoV)
+half2 EnvBRDFApproxLazarov(half Roughness, half NoV)
 {
 	// [ Lazarov 2013, "Getting More Physical in Call of Duty: Black Ops II" ]
 	// Adaptation to fit our G term.
@@ -537,18 +542,24 @@ half2 EnvBRDFApproxLazarov(half3 SpecularColor, half Roughness, half NoV)
 	half4 r = Roughness * c0 + c1;
 	half a004 = min( r.x * r.x, exp2( -9.28 * NoV ) ) * r.x + r.y;
 	half2 AB = half2( -1.04, 1.04 ) * a004 + r.zw;
+	return AB;
+}
+
+half3 EnvBRDFApprox( half3 SpecularColor, half Roughness, half NoV )
+{
+	half2 AB = EnvBRDFApproxLazarov(Roughness, NoV);
 
 	// Anything less than 2% is physically impossible and is instead considered to be shadowing
 	// Note: this is needed for the 'specular' show flag to work, since it uses a SpecularColor of 0
-	AB.y *= saturate( 50.0 * SpecularColor.g );
+	float F90 = saturate( 50.0 * SpecularColor.g );
 
-	return AB;
+	return SpecularColor * AB.x + F90 * AB.y;
 }
 
-half3 EnvBRDFApprox( half3 SpecularColor, half Roughness, half NoV )
+half3 EnvBRDFApprox(half3 F0, half3 F90, half Roughness, half NoV)
 {
-	half2 AB = EnvBRDFApproxLazarov(SpecularColor, Roughness, NoV);
-	return SpecularColor * AB.x + AB.y;
+	half2 AB = EnvBRDFApproxLazarov(Roughness, NoV);
+	return F0 * AB.x + F90 * AB.y;
 }
 
 half EnvBRDFApproxNonmetal( half Roughness, half NoV )
diff --git a/Engine/Shaders/Private/ShadingCommon.ush b/Engine/Shaders/Private/ShadingCommon.ush
--- a/Engine/Shaders/Private/ShadingCommon.ush
+++ b/Engine/Shaders/Private/ShadingCommon.ush
@@ -105,6 +105,13 @@ float3 ComputeF0(float Specular, float3 BaseColor, float Metallic)
 	return lerp(DielectricSpecularToF0(Specular).xxx, BaseColor, Metallic.xxx);
 }
 
+float3 ComputeF90(float3 F0, float3 EdgeColor, float Metallic)
+{
+	// Anything less than 2% is physically impossible and is instead considered to be shadowing
+	const float MicroOcclusion = saturate(F0.g * 50.f);
+	return lerp(MicroOcclusion.xxx, MicroOcclusion * EdgeColor, Metallic.xxx);
+}
+
 float MakeRoughnessSafe(float Roughness, float MinRoughness=0.001f)
 {
 	return clamp(Roughness, MinRoughness, 1.0f);
diff --git a/Engine/Shaders/Private/Strata/Strata.ush b/Engine/Shaders/Private/Strata/Strata.ush
--- a/Engine/Shaders/Private/Strata/Strata.ush
+++ b/Engine/Shaders/Private/Strata/Strata.ush
@@ -42,7 +42,7 @@ struct FStrataBSDF
 	float4 InlineVGPRs[3];
 #endif
 
-	float4 VGPRs[3];
+	float4 VGPRs[4];
 };
 
 // Num bits for each component store in BSDF state
@@ -54,6 +54,7 @@ struct FStrataBSDF
 #define STATE_BITS_HASSSSPROFILE		1
 #define STATE_BITS_GREYWEIGHT			1
 #define STATE_BITS_HAZINESS				1
+#define STATE_BITS_EDGECOLOR			1
 #define STATE_BITS___UNUSED___			11
 #define STATE_BITS_WEIGHT10F			10
 
@@ -66,6 +67,7 @@ struct FStrataBSDF
 #define STATE_BIT_OFFSET_HASSSSPROFILE	(STATE_BIT_OFFSET_HASSCATTERING		+ STATE_BITS_HASSCATTERING)
 #define STATE_BIT_OFFSET_GREYWEIGHT		(STATE_BIT_OFFSET_HASSSSPROFILE		+ STATE_BITS_HASSSSPROFILE)
 #define STATE_BIT_OFFSET_HAZINESS		(STATE_BIT_OFFSET_GREYWEIGHT		+ STATE_BITS_GREYWEIGHT)
+#define STATE_BIT_OFFSET_EDGECOLOR		(STATE_BIT_OFFSET_HAZINESS			+ STATE_BITS_HAZINESS)
 #define STATE_BIT_OFFSET_WEIGHT10F		22
 
 #define TO_BITMASK(BitCount)							((1<<BitCount)-1)
@@ -102,10 +104,14 @@ struct FStrataBSDF
 #define BSDF_GETGREYWEIGHT(X)			BSDF_GETSTATE(X, STATE_BITS_GREYWEIGHT, STATE_BIT_OFFSET_GREYWEIGHT)
 #define BSDF_SETGREYWEIGHT(X, Grey)		BSDF_SETSTATE(X, STATE_BITS_GREYWEIGHT, STATE_BIT_OFFSET_GREYWEIGHT, Grey)
 
-// Indicates if the BSDF weight is grey scale so that it can be stored in a single float (e.g. top layer BSDFs)
+// Indicates if the BSDF has haziness data
 #define BSDF_GETHASHAZINESS(X)			BSDF_GETSTATE(X, STATE_BITS_HAZINESS, STATE_BIT_OFFSET_HAZINESS)
 #define BSDF_SETHASHAZINESS(X, Haze)	BSDF_SETSTATE(X, STATE_BITS_HAZINESS, STATE_BIT_OFFSET_HAZINESS, Haze)
 
+// Indicates if the BSDF has edge color data
+#define BSDF_GETHASEDGECOLOR(X)			BSDF_GETSTATE(X, STATE_BITS_EDGECOLOR, STATE_BIT_OFFSET_EDGECOLOR)
+#define BSDF_SETHASEDGECOLOR(X, Col)	BSDF_SETSTATE(X, STATE_BITS_EDGECOLOR, STATE_BIT_OFFSET_EDGECOLOR, Col)
+
 // Accessors for the 10bits float weight used if is it detected being a grey scale
 #define BSDF_GETWEIGHT10F(X)			BSDF_GETSTATE(X, STATE_BITS_WEIGHT10F, STATE_BIT_OFFSET_WEIGHT10F)
 #define BSDF_SETWEIGHT10F(X, Weight)	BSDF_SETSTATE(X, STATE_BITS_WEIGHT10F, STATE_BIT_OFFSET_WEIGHT10F, Weight)
@@ -124,6 +130,7 @@ struct FStrataBSDF
 #define SLAB_ROUGHNESSX(X)				X.VGPRs[1].y
 #define SLAB_ROUGHNESSY(X)				X.VGPRs[1].z
 #define SLAB_HAZINESS(X)				X.VGPRs[1].w
+#define SLAB_EDGECOLOR(X)				X.VGPRs[3].xyz
 // Either with SSS profile
 #define SLAB_SSSPROFILEID(X)			X.VGPRs[2].x
 #define SLAB_SSSPROFILERADIUSSCALE(X)	X.VGPRs[2].y
@@ -316,7 +323,7 @@ FStrataData GetInitialisedStrataData()
 }
 
 FStrataData GetStrataSlabBSDF(
-	float3 BaseColor, float Specular, float Metallic, float2 Roughness,
+	float3 BaseColor, float3 EdgeColor, float Specular, float Metallic, float2 Roughness,
 	float SSSProfileID, float3 SSSDMFP, float SSSDMFPScale,
 	float3 Emissive, float Haziness, 
 	uint SharedNormalIndex, inout uint SharedNormalTypes)
@@ -324,6 +331,7 @@ FStrataData GetStrataSlabBSDF(
 	FStrataData StrataData = GetInitialisedStrataData();
 	StrataData.LayerCount++;
 
+	const uint bHasEdgeColor = any(EdgeColor < 1);
 	const uint bHasHaziness = Haziness > 0.0f ? 1 : 0;
 	const uint bHasAnisotropy = StrataHasAnisotropy(Roughness.x, Roughness.y) ? 1 : 0;
 	if (bHasAnisotropy)
@@ -344,8 +352,10 @@ FStrataData GetStrataSlabBSDF(
 	BSDF_SETHASSCATTERING	(StrataData.Layers[0].BSDFs[0],		bHasSSS ? 1 : 0);
 	BSDF_SETHASSSSPROFILE	(StrataData.Layers[0].BSDFs[0],		bHasSSSProfile ? 1 : 0);
 	BSDF_SETHASHAZINESS		(StrataData.Layers[0].BSDFs[0],		bHasHaziness);
+	BSDF_SETHASEDGECOLOR	(StrataData.Layers[0].BSDFs[0],		bHasEdgeColor);
 
 	SLAB_BASECOLOR			(StrataData.Layers[0].BSDFs[0])		= BaseColor;
+	SLAB_EDGECOLOR			(StrataData.Layers[0].BSDFs[0])		= EdgeColor;
 	SLAB_SPECULAR			(StrataData.Layers[0].BSDFs[0])		= Specular;
 	SLAB_METALLIC			(StrataData.Layers[0].BSDFs[0])		= Metallic;
 	SLAB_ROUGHNESSX			(StrataData.Layers[0].BSDFs[0])		= Roughness.x;
@@ -639,6 +649,7 @@ void SanitizeStrataSlab(inout FStrataBSDF BSDF)
 {
 	BSDF_SETEMISSIVE(BSDF, clamp(BSDF_GETEMISSIVE(BSDF), 0.0, Max10BitsFloat));
 	SLAB_BASECOLOR(BSDF) = saturate(SLAB_BASECOLOR(BSDF));
+	SLAB_EDGECOLOR(BSDF) = saturate(SLAB_EDGECOLOR(BSDF));
 	SLAB_SPECULAR(BSDF) = saturate(SLAB_SPECULAR(BSDF));
 	SLAB_METALLIC(BSDF) = saturate(SLAB_METALLIC(BSDF));
 	SLAB_ROUGHNESSX(BSDF) = saturate(SLAB_ROUGHNESSX(BSDF));
@@ -1089,6 +1100,7 @@ void PackStrataOut(
 			case STRATA_BSDF_TYPE_SLAB:
 			{
 				const float3 F0 = ComputeF0(SLAB_SPECULAR(BSDF), SLAB_BASECOLOR(BSDF), SLAB_METALLIC(BSDF));
+				const float3 F90 = ComputeF90(F0, SLAB_EDGECOLOR(BSDF), SLAB_METALLIC(BSDF));
 
 				if (bTopLayer)
 				{
@@ -1097,7 +1109,7 @@ void PackStrataOut(
 					TopLayer.AccumRoughnessCount += BSDFCoverageAvg;
 
 					const float SafeRoughness = MakeRoughnessSafe(SLAB_ROUGHNESSX(BSDF));
-					TopLayer.BaseColor += BSDFCoverage * EnvBRDFApprox(F0, SafeRoughness, saturate(ContextNoLight.NoV));
+					TopLayer.BaseColor += BSDFCoverage * EnvBRDFApprox(F0, F90, SafeRoughness, saturate(ContextNoLight.NoV));
 				}
 
 				// Select the first/toppest layer with subsurface scattering
@@ -1280,6 +1292,11 @@ void PackStrataOut(
 				STRATA_STORE_UINT1(PackColorLinearToGamma2AlphaLinear(float4(SLAB_BASECOLOR(BSDF), SLAB_ROUGHNESSX(BSDF))));
 				STRATA_STORE_UINT1(PackRGBA8(float4(SLAB_SPECULAR(BSDF), SLAB_METALLIC(BSDF), SLAB_ROUGHNESSY(BSDF), SLAB_HAZINESS(BSDF))));
 
+				if (BSDF_GETHASEDGECOLOR(BSDF))
+				{
+					STRATA_STORE_UINT1(PackColorLinearToGamma2AlphaLinear(float4(SLAB_EDGECOLOR(BSDF), 0)));
+				}
+
 				if (BSDF_GETHASSCATTERING(BSDF))
 				{
 					if (BSDF_GETHASSSSPROFILE(BSDF))
@@ -1291,7 +1308,8 @@ void PackStrataOut(
 						STRATA_STORE_UINT1(PackR11G11B10F(SLAB_SSSDMFP(BSDF)));
 					}
 				}
-				// 8 bytes
+
+				// 8-16 bytes
 			}
 			break;
 			case STRATA_BSDF_TYPE_SHEEN:
@@ -1358,7 +1376,7 @@ FStrataPixelHeader UnpackStrataHeaderIn(ByteAddressBuffer StrataBuffer, inout ui
 }
 
 // Unpack only strata header 
-// This version assume only a single slab (isotropic BSDF, with no subsurface scattering, and no haziness)
+// This version assume only a single slab (isotropic BSDF, with no edge color, no subsurface scattering, and no haziness)
 FStrataPixelHeader UnpackFastPathStrataHeaderIn(ByteAddressBuffer StrataBuffer, inout uint ByteOffset)
 {
 	FStrataPixelHeader Out = InitialiseStrataPixelHeader();
@@ -1383,6 +1401,7 @@ FStrataPixelHeader UnpackFastPathStrataHeaderIn(ByteAddressBuffer StrataBuffer,
 void SetStrataFastPathState(inout FStrataBSDF BSDF)
 {
 	BSDF_SETTYPE(BSDF, STRATA_BSDF_TYPE_SLAB);
+	BSDF_SETHASEDGECOLOR(BSDF, 0);
 	BSDF_SETANISOTROPY(BSDF, 0);
 	BSDF_SETISTOPLAYER(BSDF, 1);
 	BSDF_SETHASSCATTERING(BSDF, 0);
@@ -1424,6 +1443,14 @@ FStrataBSDF UnpackStrataBSDFIn(ByteAddressBuffer StrataBuffer, inout uint ByteOf
 			SLAB_ROUGHNESSY(OutBSDF) = Data.z;
 			SLAB_HAZINESS(OutBSDF) = Data.w;
 		}
+
+		SLAB_EDGECOLOR(OutBSDF) = 1.0f;
+		if (BSDF_GETHASEDGECOLOR(OutBSDF))
+		{
+			float Unused;
+			StrataLoad_ColorGamma2ToLinearAlphaLinear(StrataBuffer, ByteOffset, SLAB_EDGECOLOR(OutBSDF), Unused);
+		}
+
 		if (BSDF_GETHASSCATTERING(OutBSDF))
 		{
 			if (BSDF_GETHASSSSPROFILE(OutBSDF))
diff --git a/Engine/Shaders/Private/Strata/StrataEvaluation.ush b/Engine/Shaders/Private/Strata/StrataEvaluation.ush
--- a/Engine/Shaders/Private/Strata/StrataEvaluation.ush
+++ b/Engine/Shaders/Private/Strata/StrataEvaluation.ush
@@ -184,9 +184,8 @@ FStrataEvaluateResult StrataEvaluateBSDF(FStrataBSDFContext BSDFContext, FShadow
 			Sample.DiffusePDF		= BSDFContext.SatNoV / PI;
 			Sample.bSubsurface		= BSDF_GETHASSCATTERING(BSDFContext.BSDF) > 0;
 
-
-			const float3 F0 = ComputeF0(SLAB_SPECULAR(BSDFContext.BSDF), BaseColor, Metallic);
-			// STRATA_TODO F82 or edge color
+			const float3 F0  = ComputeF0(SLAB_SPECULAR(BSDFContext.BSDF), BaseColor, Metallic);
+			const float3 F90 = ComputeF90(F0, SLAB_EDGECOLOR(BSDFContext.BSDF), Metallic);
 
 			const float SafeRoughnessX = MakeRoughnessSafe(SLAB_ROUGHNESSX(BSDFContext.BSDF));
 			const float SafeRoughnessY = MakeRoughnessSafe(SLAB_ROUGHNESSY(BSDFContext.BSDF));
@@ -235,7 +234,7 @@ FStrataEvaluateResult StrataEvaluateBSDF(FStrataBSDFContext BSDFContext, FShadow
 					HazePDF = VisibleGGXPDF(BSDFContext.TangentV, BSDFContext.TangentH, a2Haze);
 				}
 			}
-			float3 F = F_Schlick(F0, BSDFContext.Context.VoH);
+			float3 F = F_Schlick(F0, F90, BSDFContext.Context.VoH);
 
 			Sample.SpecularPathProbability	= bHaziness		? (1.0f - SPECULAR_HAZINESS_WEIGHT) : 1.0f;
 			Sample.SpecularHazePathProbability= bHaziness	? SPECULAR_HAZINESS_WEIGHT			: 0.0f;
@@ -445,7 +444,7 @@ FStrataEnvLightResult StrataEvaluateForEnvLight(FStrataBSDFContext BSDFContext,
 			if (bEnableSpecular)
 			{
 				const float3 F0 = ComputeF0(SLAB_SPECULAR(BSDFContext.BSDF), BaseColor, Metallic);
-				// STRATA_TODO F82
+				const float3 F90 = ComputeF90(F0, SLAB_EDGECOLOR(BSDFContext.BSDF), Metallic);
 
 				float SafeRoughnessX = MakeRoughnessSafe(SLAB_ROUGHNESSX(BSDFContext.BSDF));
 
@@ -465,7 +464,7 @@ FStrataEnvLightResult StrataEvaluateForEnvLight(FStrataBSDFContext BSDFContext,
 					BSDFContext.R = 2 * dot(BSDFContext.V, BSDFContext.N) * BSDFContext.N - BSDFContext.V;
 				}
 
-				float3 EvalEnvBRDF = EnvBRDFApprox(F0, SafeRoughnessX, BSDFContext.SatNoV);
+				float3 EvalEnvBRDF = EnvBRDFApprox(F0, F90, SafeRoughnessX, BSDFContext.SatNoV);
 
 				StrataEnvLightResult.SpecularDirection = StrataGetOffSpecularPeakReflectionDir(BSDFContext.N, BSDFContext.R, SafeRoughnessX);
 				StrataEnvLightResult.SpecularWeight = EvalEnvBRDF;
@@ -484,7 +483,7 @@ FStrataEnvLightResult StrataEvaluateForEnvLight(FStrataBSDFContext BSDFContext,
 
 					// Compute the second specular weight and roughness
 					float SafeRoughnessHaze = ComputeHazyLobeRoughness(SafeRoughnessX, Haziness);
-					StrataEnvLightResult.SpecularHazeWeight = lerp(0.0f, SPECULAR_HAZINESS_WEIGHT, HazinessFadeIn) * EnvBRDFApprox(F0, SafeRoughnessHaze, BSDFContext.SatNoV);
+					StrataEnvLightResult.SpecularHazeWeight = lerp(0.0f, SPECULAR_HAZINESS_WEIGHT, HazinessFadeIn) * EnvBRDFApprox(F0, F90, SafeRoughnessHaze, BSDFContext.SatNoV);
 					StrataEnvLightResult.SpecularHazeSafeRoughness = SafeRoughnessHaze;
 				}
 			}
diff --git a/Engine/Shaders/Private/Strata/StrataMaterialClassification.usf b/Engine/Shaders/Private/Strata/StrataMaterialClassification.usf
--- a/Engine/Shaders/Private/Strata/StrataMaterialClassification.usf
+++ b/Engine/Shaders/Private/Strata/StrataMaterialClassification.usf
@@ -49,7 +49,8 @@ void ComputePostProcessBSDFData(
 	{
 	case STRATA_BSDF_TYPE_SLAB:
 	{
-		const float3 F0 = ComputeF0(SLAB_SPECULAR(BSDF), SLAB_BASECOLOR(BSDF), SLAB_METALLIC(BSDF));
+		const float3 F0  = ComputeF0(SLAB_SPECULAR(BSDF), SLAB_BASECOLOR(BSDF), SLAB_METALLIC(BSDF));
+		const float3 F90 = ComputeF90(F0, SLAB_BASECOLOR(BSDF), SLAB_METALLIC(BSDF));
 
 		if (bTopLayer)
 		{
@@ -58,7 +59,7 @@ void ComputePostProcessBSDFData(
 			TopLayer.AccumRoughnessCount += BSDFCoverageAvg;
 
 			const float SafeRoughness = MakeRoughnessSafe(SLAB_ROUGHNESSX(BSDF));
-			TopLayer.BaseColor += BSDFCoverage * EnvBRDFApprox(F0, SafeRoughness, saturate(ContextNoLight.NoV));
+			TopLayer.BaseColor += BSDFCoverage * EnvBRDFApprox(F0, F90, SafeRoughness, saturate(ContextNoLight.NoV));
 		}
 
 		Classification.bIsSimple =
diff --git a/Engine/Shaders/Private/Strata/StrataVisualize.usf b/Engine/Shaders/Private/Strata/StrataVisualize.usf
--- a/Engine/Shaders/Private/Strata/StrataVisualize.usf
+++ b/Engine/Shaders/Private/Strata/StrataVisualize.usf
@@ -238,9 +238,17 @@ void VisualizeMaterialPS(
 				HORIZONTALTAB; PRINTFLOAT3(SLAB_BASECOLOR(BSDF));
 				NEXTLINE;
 
+				if (BSDF_GETHASEDGECOLOR(BSDF))
+				{
+					HORIZONTALTAB; PRINTPARAMCHAR(_E_); PRINTPARAMCHAR(_D_); PRINTPARAMCHAR(_G_); PRINTPARAMCHAR(_E_); PRINTPARAMCHAR(_SPC_); PRINTPARAMCHAR(_SPC_);
+					HORIZONTALTAB; PRINTFLOAT3(SLAB_EDGECOLOR(BSDF));
+					NEXTLINE;
+				}
+
 				HORIZONTALTAB; PRINTPARAMCHAR(_R_); PRINTPARAMCHAR(_E_); PRINTPARAMCHAR(_F_); PRINTPARAMCHAR(_L_); PRINTPARAMCHAR(_E_); PRINTPARAMCHAR(_C_);
 				HORIZONTALTAB; PRINTFLOAT1(SLAB_SPECULAR(BSDF));
 				NEXTLINE;
+
 				HORIZONTALTAB; PRINTPARAMCHAR(_M_); PRINTPARAMCHAR(_E_); PRINTPARAMCHAR(_T_); PRINTPARAMCHAR(_A_); PRINTPARAMCHAR(_L_); PRINTPARAMCHAR(_SPC_);
 				HORIZONTALTAB; PRINTFLOAT1(SLAB_METALLIC(BSDF));
 				NEXTLINE;
diff --git a/Engine/Source/Runtime/Engine/Classes/Materials/MaterialExpressionStrata.h b/Engine/Source/Runtime/Engine/Classes/Materials/MaterialExpressionStrata.h
--- a/Engine/Source/Runtime/Engine/Classes/Materials/MaterialExpressionStrata.h
+++ b/Engine/Source/Runtime/Engine/Classes/Materials/MaterialExpressionStrata.h
@@ -41,6 +41,12 @@ class UMaterialExpressionStrataSlabBSDF : public UMaterialExpressionStrataBSDF
 	UPROPERTY()
 	FExpressionInput BaseColor;
 
+	/**
+	 * Defines the edge color of the Material. This is only applied on metallic material (type = float3, unit = unitless, defaults to 1.0)
+	 */
+	UPROPERTY()
+	FExpressionInput EdgeColor;
+
 	/**
 	 * Controls how \"metal-like\" your surface looks like. 0 means dielectric, 1 means conductor (type = float, unit = unitless, defaults to 0)
 	 */
@@ -117,6 +123,7 @@ class UMaterialExpressionStrataSlabBSDF : public UMaterialExpressionStrataBSDF
 	virtual void GatherStrataMaterialInfo(FStrataMaterialInfo& StrataMaterialInfo, int32 OutputIndex) override;
 	virtual FName GetInputName(int32 InputIndex) const override;
 
+	bool HasEdgeColor() const;
 	bool HasScattering() const;
 #endif
 	//~ End UMaterialExpression Interface
diff --git a/Engine/Source/Runtime/Engine/Private/Materials/HLSLMaterialTranslator.cpp b/Engine/Source/Runtime/Engine/Private/Materials/HLSLMaterialTranslator.cpp
--- a/Engine/Source/Runtime/Engine/Private/Materials/HLSLMaterialTranslator.cpp
+++ b/Engine/Source/Runtime/Engine/Private/Materials/HLSLMaterialTranslator.cpp
@@ -9029,13 +9029,14 @@ int32 FHLSLMaterialTranslator::StrataCreateAndRegisterNullMaterial()
 }
 
 int32 FHLSLMaterialTranslator::StrataSlabBSDF(
-	int32 BaseColor, int32 Specular, int32 Metallic, int32 RoughnessX, int32 RoughnessY,
+	int32 BaseColor, int32 EdgeColor, int32 Specular, int32 Metallic, int32 RoughnessX, int32 RoughnessY,
 	int32 SSSProfileId, int32 SSSDMFP, int32 SSSDMFPScale, int32 EmissiveColor,
 	int32 Haziness, int32 Normal, int32 Tangent, uint8 SharedNormalIndex)
 {
 	return AddCodeChunk(
-		MCT_Strata, TEXT("GetStrataSlabBSDF(%s, %s, %s, float2(%s, %s), %s, %s, %s, %s, %s, %u, Parameters.SharedNormals.NormalTypes) /* Normal = %s ; Tangent = %s */"),
+		MCT_Strata, TEXT("GetStrataSlabBSDF(%s, %s, %s, %s, float2(%s, %s), %s, %s, %s, %s, %s, %u, Parameters.SharedNormals.NormalTypes) /* Normal = %s ; Tangent = %s */"),
 		*GetParameterCode(BaseColor),
+		*GetParameterCode(EdgeColor),
 		*GetParameterCode(Specular),
 		*GetParameterCode(Metallic),
 		*GetParameterCode(RoughnessX),
diff --git a/Engine/Source/Runtime/Engine/Private/Materials/HLSLMaterialTranslator.h b/Engine/Source/Runtime/Engine/Private/Materials/HLSLMaterialTranslator.h
--- a/Engine/Source/Runtime/Engine/Private/Materials/HLSLMaterialTranslator.h
+++ b/Engine/Source/Runtime/Engine/Private/Materials/HLSLMaterialTranslator.h
@@ -988,7 +988,7 @@ protected:
 	// Strata
 	virtual int32 StrataCreateAndRegisterNullMaterial() override;
 	virtual int32 StrataSlabBSDF(
-		int32 BaseColor, int32 Specular, int32 Metallic, int32 RoughnessX, int32 RoughnessY,
+		int32 BaseColor, int32 EdgeColor, int32 Specular, int32 Metallic, int32 RoughnessX, int32 RoughnessY,
 		int32 SSSProfileId, int32 SSSDMFP, int32 SSSDMFPScale,
 		int32 EmissiveColor, int32 Haziness, int32 Normal, int32 Tangent, uint8 SharedNormalIndex) override;
 	virtual int32 StrataSheenBSDF(int32 BaseColor, int32 Roughness, int32 Normal, uint8 SharedNormalIndex) override;
diff --git a/Engine/Source/Runtime/Engine/Private/Materials/Material.cpp b/Engine/Source/Runtime/Engine/Private/Materials/Material.cpp
--- a/Engine/Source/Runtime/Engine/Private/Materials/Material.cpp
+++ b/Engine/Source/Runtime/Engine/Private/Materials/Material.cpp
@@ -3677,10 +3677,10 @@ void UMaterial::ConvertMaterialToStrataMaterial()
 			}
 
 			MoveConnectionTo(BaseColor, SlabBSDF, 0);		// BaseColor
-			MoveConnectionTo(Metallic, SlabBSDF, 1);		// Metallic
-			MoveConnectionTo(Specular, SlabBSDF, 2);		// Specular
-			MoveConnectionTo(Normal, SlabBSDF, 5);			// Normal
-			MoveConnectionTo(Tangent, SlabBSDF, 6);			// Tangent
+			MoveConnectionTo(Metallic, SlabBSDF, 2);		// Metallic
+			MoveConnectionTo(Specular, SlabBSDF, 3);		// Specular
+			MoveConnectionTo(Normal, SlabBSDF, 6);			// Normal
+			MoveConnectionTo(Tangent, SlabBSDF, 7);			// Tangent
 			// Opacity can remain on the end point node
 			// Opacity mask can remain on the end point node
 			// WorldPositionOffset can remain on the end point node
diff --git a/Engine/Source/Runtime/Engine/Private/Materials/MaterialExpressions.cpp b/Engine/Source/Runtime/Engine/Private/Materials/MaterialExpressions.cpp
--- a/Engine/Source/Runtime/Engine/Private/Materials/MaterialExpressions.cpp
+++ b/Engine/Source/Runtime/Engine/Private/Materials/MaterialExpressions.cpp
@@ -19760,6 +19760,8 @@ int32 UMaterialExpressionStrataSlabBSDF::Compile(class FMaterialCompiler* Compil
 	int32 TangentCodeChunk = CompileWithDefaultTangentWS(Compiler, Tangent);
 	uint8 SharedNormalIndex = bAnisotropyPotentiallyUsed ? StrataCompilationInfoCreateSharedNormal(Compiler, NormalCodeChunk, TangentCodeChunk) : StrataCompilationInfoCreateSharedNormal(Compiler, NormalCodeChunk);
 
+	const bool bHasEdgeColor = HasEdgeColor();
+
 	int32 SSSProfileCodeChunk = INDEX_NONE;
 	const bool bHasScattering = HasScattering();
 	if (bHasScattering)
@@ -19770,6 +19772,7 @@ int32 UMaterialExpressionStrataSlabBSDF::Compile(class FMaterialCompiler* Compil
 
 	int32 OutputCodeChunk = Compiler->StrataSlabBSDF(
 		CompileWithDefaultFloat3(Compiler, BaseColor, 0.18f, 0.18f, 0.18f),
+		CompileWithDefaultFloat3(Compiler, EdgeColor, 1.0f, 1.0f, 1.0f),
 		CompileWithDefaultFloat1(Compiler, Specular, 0.5f),
 		CompileWithDefaultFloat1(Compiler, Metallic, 0.0f),
 		RoughnessXCodeChunk,
@@ -19782,7 +19785,7 @@ int32 UMaterialExpressionStrataSlabBSDF::Compile(class FMaterialCompiler* Compil
 		NormalCodeChunk,
 		TangentCodeChunk,
 		SharedNormalIndex);
-	StrataCompilationInfoCreateSingleBSDFMaterial(Compiler, OutputCodeChunk, SharedNormalIndex, STRATA_BSDF_TYPE_SLAB, bHasScattering);
+	StrataCompilationInfoCreateSingleBSDFMaterial(Compiler, OutputCodeChunk, SharedNormalIndex, STRATA_BSDF_TYPE_SLAB, bHasEdgeColor, bHasScattering);
 
 	return OutputCodeChunk;
 }
@@ -19805,33 +19808,36 @@ uint32 UMaterialExpressionStrataSlabBSDF::GetInputType(int32 InputIndex)
 		return MCT_Float3; // BaseColor
 		break;
 	case 1:
-		return MCT_Float1; // Metallic
+		return MCT_Float3; // EdgeColor
 		break;
 	case 2:
-		return MCT_Float1; // Specular
+		return MCT_Float1; // Metallic
 		break;
 	case 3:
-		return MCT_Float1; // RoughnessX
+		return MCT_Float1; // Specular
 		break;
 	case 4:
-		return MCT_Float1; // RoughnessY
+		return MCT_Float1; // RoughnessX
 		break;
 	case 5:
-		return MCT_Float3; // Normal
+		return MCT_Float1; // RoughnessY
 		break;
 	case 6:
-		return MCT_Float3; // Tangent
+		return MCT_Float3; // Normal
 		break;
 	case 7:
-		return MCT_Float3; // SSSDMFP
+		return MCT_Float3; // Tangent
 		break;
 	case 8:
-		return MCT_Float1; // SSSDMFPScale
+		return MCT_Float3; // SSSDMFP
 		break;
 	case 9:
-		return MCT_Float3; // Emissive Color
+		return MCT_Float1; // SSSDMFPScale
 		break;
 	case 10:
+		return MCT_Float3; // Emissive Color
+		break;
+	case 11:
 		return MCT_Float1; // Haziness
 		break;
 	}
@@ -19848,41 +19854,45 @@ FName UMaterialExpressionStrataSlabBSDF::GetInputName(int32 InputIndex) const
 	}
 	else if (InputIndex == 1)
 	{
-		return TEXT("Metallic");
+		return TEXT("EdgeColor");
 	}
 	else if (InputIndex == 2)
 	{
-		return TEXT("Specular");
+		return TEXT("Metallic");
 	}
 	else if (InputIndex == 3)
 	{
-		return TEXT("RoughnessX");
+		return TEXT("Specular");
 	}
 	else if (InputIndex == 4)
 	{
-		return TEXT("RoughnessY");
+		return TEXT("RoughnessX");
 	}
 	else if (InputIndex == 5)
 	{
-		return TEXT("Normal");
+		return TEXT("RoughnessY");
 	}
 	else if (InputIndex == 6)
 	{
-		return TEXT("Tangent");
+		return TEXT("Normal");
 	}
 	else if (InputIndex == 7)
 	{
-		return TEXT("SSS Diffuse MFP");
+		return TEXT("Tangent");
 	}
 	else if (InputIndex == 8)
 	{
-		return TEXT("SSS Diffuse MFP Scale");
+		return TEXT("SSS Diffuse MFP");
 	}
 	else if (InputIndex == 9)
 	{
-		return TEXT("Emissive Color");
+		return TEXT("SSS Diffuse MFP Scale");
 	}
 	else if (InputIndex == 10)
+	{
+		return TEXT("Emissive Color");
+	}
+	else if (InputIndex == 11)
 	{
 		return TEXT("Haziness");
 	}
@@ -19914,6 +19924,12 @@ bool UMaterialExpressionStrataSlabBSDF::HasScattering() const
 {
 	return SubsurfaceProfile != nullptr || SSSDMFP.IsConnected();
 }
+
+bool UMaterialExpressionStrataSlabBSDF::HasEdgeColor() const
+{
+	return EdgeColor.IsConnected();
+}
+
 #endif // WITH_EDITOR
 
 
diff --git a/Engine/Source/Runtime/Engine/Private/Materials/StrataMaterial.cpp b/Engine/Source/Runtime/Engine/Private/Materials/StrataMaterial.cpp
--- a/Engine/Source/Runtime/Engine/Private/Materials/StrataMaterial.cpp
+++ b/Engine/Source/Runtime/Engine/Private/Materials/StrataMaterial.cpp
@@ -52,7 +52,7 @@ uint8 StrataCompilationInfoCreateSharedNormal(FMaterialCompiler* Compiler, int32
 	return Compiler->StrataCompilationInfoRegisterSharedNormalIndex(NormalCodeChunk, TangentCodeChunk);
 }
 
-void StrataCompilationInfoCreateSingleBSDFMaterial(FMaterialCompiler* Compiler, int32 CodeChunk, uint8 SharedNormalIndex, uint8 BSDFType, bool bHasScattering)
+void StrataCompilationInfoCreateSingleBSDFMaterial(FMaterialCompiler* Compiler, int32 CodeChunk, uint8 SharedNormalIndex, uint8 BSDFType, bool bHasEdgeColor, bool bHasScattering)
 {
 	FStrataMaterialCompilationInfo StrataInfo;
 	StrataInfo.LayerCount = 1;
@@ -60,6 +60,7 @@ void StrataCompilationInfoCreateSingleBSDFMaterial(FMaterialCompiler* Compiler,
 	StrataInfo.Layers[0].BSDFs[0].Type = BSDFType;
 	StrataInfo.Layers[0].BSDFs[0].SharedNormalIndex = SharedNormalIndex;
 	StrataInfo.Layers[0].BSDFs[0].bHasScattering = bHasScattering;
+	StrataInfo.Layers[0].BSDFs[0].bHasEdgeColor = bHasEdgeColor;
 	UpdateTotalBSDFCount(StrataInfo);
 	Compiler->StrataCompilationInfoRegisterCodeChunk(CodeChunk, StrataInfo);
 }
@@ -280,6 +281,10 @@ FStrataMaterialAnalysisResult StrataCompilationInfoMaterialAnalysis(FMaterialCom
 			{
 				Result.RequestedByteCount += UintByteSize;
 				Result.RequestedByteCount += UintByteSize;
+				if (BSDF.bHasEdgeColor)
+				{
+					Result.RequestedByteCount += UintByteSize;
+				}
 				if (BSDF.bHasScattering)
 				{
 					Result.RequestedByteCount += UintByteSize;
diff --git a/Engine/Source/Runtime/Engine/Private/Materials/StrataMaterial.h b/Engine/Source/Runtime/Engine/Private/Materials/StrataMaterial.h
--- a/Engine/Source/Runtime/Engine/Private/Materials/StrataMaterial.h
+++ b/Engine/Source/Runtime/Engine/Private/Materials/StrataMaterial.h
@@ -28,6 +28,7 @@ struct FStrataMaterialCompilationInfo
 		uint8 Type;
 		uint8 SharedNormalIndex;
 		bool  bHasScattering;
+		bool  bHasEdgeColor;
 	};
 
 	struct FLayer
@@ -54,7 +55,7 @@ FString GetStrataBSDFName(uint8 BSDFType);
 uint8 StrataCompilationInfoCreateSharedNormal(FMaterialCompiler* Compiler, int32 NormalCodeChunk);
 uint8 StrataCompilationInfoCreateSharedNormal(FMaterialCompiler* Compiler, int32 NormalCodeChunk, int32 TangentCodeChunk);
 void  StrataCompilationInfoCreateNullBSDF(FMaterialCompiler* Compiler, int32 CodeChunk);
-void  StrataCompilationInfoCreateSingleBSDFMaterial(FMaterialCompiler* Compiler, int32 CodeChunk, uint8 SharedNormalIndex, uint8 BSDFType, bool bHasScattering = false);
+void  StrataCompilationInfoCreateSingleBSDFMaterial(FMaterialCompiler* Compiler, int32 CodeChunk, uint8 SharedNormalIndex, uint8 BSDFType, bool bHasEdgeColor = false, bool bHasScattering = false);
 
 FStrataMaterialCompilationInfo StrataCompilationInfoAdd(FMaterialCompiler* Compiler, const FStrataMaterialCompilationInfo& A, const FStrataMaterialCompilationInfo& B);
 FStrataMaterialCompilationInfo StrataCompilationInfoMultiply(FMaterialCompiler* Compiler, const FStrataMaterialCompilationInfo& A);
diff --git a/Engine/Source/Runtime/Engine/Public/MaterialCompiler.h b/Engine/Source/Runtime/Engine/Public/MaterialCompiler.h
--- a/Engine/Source/Runtime/Engine/Public/MaterialCompiler.h
+++ b/Engine/Source/Runtime/Engine/Public/MaterialCompiler.h
@@ -408,7 +408,7 @@ public:
 	// Strata
 	virtual int32 StrataCreateAndRegisterNullMaterial() = 0;
 	virtual int32 StrataSlabBSDF(
-		int32 BaseColor, int32 Specular, int32 Metallic, int32 RoughnessX, int32 RoughnessY, 
+		int32 BaseColor, int32 EdgeColor, int32 Specular, int32 Metallic, int32 RoughnessX, int32 RoughnessY, 
 		int32 SSSProfileId, int32 SSSDMFP, int32 SSSDMFPScale,
 		int32 EmissiveColor, int32 Haziness, int32 Normal, int32 Tangent, uint8 SharedNormalIndex) = 0;
 	virtual int32 StrataSheenBSDF(int32 BaseColor, int32 Roughness, int32 Normal, uint8 SharedNormalIndex) = 0;
@@ -888,12 +888,12 @@ public:
 	}
 
 	virtual int32 StrataSlabBSDF(
-		int32 BaseColor, int32 Specular, int32 Metallic, int32 RoughnessX, int32 RoughnessY,
+		int32 BaseColor, int32 EdgeColor, int32 Specular, int32 Metallic, int32 RoughnessX, int32 RoughnessY,
 		int32 SSSProfileId, int32 SSSDMFP, int32 SSSDMFPScale,
 		int32 EmissiveColor, int32 Haziness, int32 Normal, int32 Tangent, uint8 SharedNormalIndex) override
 	{
 		return Compiler->StrataSlabBSDF(
-			BaseColor, Specular, Metallic, RoughnessX, RoughnessY,
+			BaseColor, EdgeColor, Specular, Metallic, RoughnessX, RoughnessY,
 			SSSProfileId, SSSDMFP, SSSDMFPScale,
 			EmissiveColor, Haziness, Normal, Tangent, SharedNormalIndex);
 	}
