Add thim film support for Strata.

Only thickness parameter is exposed for simplifying storage. The thin film IOR is harcoded to 1.44 (~ oil). This CL provides a reference implementation and a simplified version. The simplied version does not take into account polarization neither proper Fresnel phase shift, but gives plausible result. The simplified version is used by default.

The direct lighting and sky lighting path are implemented. SSR is not properly weight at with by the Airy Fresnel term.

diff --git a/Engine/Shaders/Private/Strata/Strata.ush b/Engine/Shaders/Private/Strata/Strata.ush
--- a/Engine/Shaders/Private/Strata/Strata.ush
+++ b/Engine/Shaders/Private/Strata/Strata.ush
@@ -55,7 +55,8 @@ struct FStrataBSDF
 #define STATE_BITS_GREYWEIGHT			1
 #define STATE_BITS_HAZINESS				1
 #define STATE_BITS_EDGECOLOR			1
-#define STATE_BITS___UNUSED___			11
+#define STATE_BITS_THINFILM				1
+#define STATE_BITS___UNUSED___			10
 #define STATE_BITS_WEIGHT10F			10
 
 // Bit offset of each component store in BSDF state 
@@ -68,6 +69,7 @@ struct FStrataBSDF
 #define STATE_BIT_OFFSET_GREYWEIGHT		(STATE_BIT_OFFSET_HASSSSPROFILE		+ STATE_BITS_HASSSSPROFILE)
 #define STATE_BIT_OFFSET_HAZINESS		(STATE_BIT_OFFSET_GREYWEIGHT		+ STATE_BITS_GREYWEIGHT)
 #define STATE_BIT_OFFSET_EDGECOLOR		(STATE_BIT_OFFSET_HAZINESS			+ STATE_BITS_HAZINESS)
+#define STATE_BIT_OFFSET_THINFILM		(STATE_BIT_OFFSET_EDGECOLOR			+ STATE_BITS_EDGECOLOR)
 #define STATE_BIT_OFFSET_WEIGHT10F		22
 
 #define TO_BITMASK(BitCount)							((1<<BitCount)-1)
@@ -112,6 +114,10 @@ struct FStrataBSDF
 #define BSDF_GETHASEDGECOLOR(X)			BSDF_GETSTATE(X, STATE_BITS_EDGECOLOR, STATE_BIT_OFFSET_EDGECOLOR)
 #define BSDF_SETHASEDGECOLOR(X, Col)	BSDF_SETSTATE(X, STATE_BITS_EDGECOLOR, STATE_BIT_OFFSET_EDGECOLOR, Col)
 
+// Indicates if the BSDF has thin film
+#define BSDF_GETHASTHINFILM(X)			BSDF_GETSTATE(X, STATE_BITS_THINFILM, STATE_BIT_OFFSET_THINFILM)
+#define BSDF_SETHASTHINFILM(X, Thin)	BSDF_SETSTATE(X, STATE_BITS_THINFILM, STATE_BIT_OFFSET_THINFILM, Thin)
+
 // Accessors for the 10bits float weight used if is it detected being a grey scale
 #define BSDF_GETWEIGHT10F(X)			BSDF_GETSTATE(X, STATE_BITS_WEIGHT10F, STATE_BIT_OFFSET_WEIGHT10F)
 #define BSDF_SETWEIGHT10F(X, Weight)	BSDF_SETSTATE(X, STATE_BITS_WEIGHT10F, STATE_BIT_OFFSET_WEIGHT10F, Weight)
@@ -136,6 +142,9 @@ struct FStrataBSDF
 #define SLAB_SSSPROFILERADIUSSCALE(X)	X.VGPRs[2].y
 // Or explicit DMFP
 #define SLAB_SSSDMFP(X)					X.VGPRs[2].xyz
+// Thin film. IOR is harcoded to oil IOR in order to save VGPR
+#define SLAB_THINFILM_THICKNESS(X)		X.VGPRs[2].w
+#define SLAB_THINFILM_IOR(X)			1.44f
 
 #define SHEEN_BASECOLOR(X)				X.VGPRs[0].xyz
 #define SHEEN_ROUGHNESS(X)				X.VGPRs[0].w
@@ -164,8 +173,6 @@ struct FStrataBSDF
 #define SLW_COLORSCALEBEHINDWATER(X)	X.InlineVGPRs[2].xyz
 #endif
 
-
-
 void StrataRequestSharedNormalTangentBasis(inout uint NormalTypes, uint Index);
 
 uint   PackR11G11B10F(float3 rgb);
@@ -325,7 +332,7 @@ FStrataData GetInitialisedStrataData()
 FStrataData GetStrataSlabBSDF(
 	float3 BaseColor, float3 EdgeColor, float Specular, float Metallic, float2 Roughness,
 	float SSSProfileID, float3 SSSDMFP, float SSSDMFPScale,
-	float3 Emissive, float Haziness, 
+	float3 Emissive, float Haziness, float ThinFilmThickness, 
 	uint SharedNormalIndex, inout uint SharedNormalTypes)
 {
 	FStrataData StrataData = GetInitialisedStrataData();
@@ -333,6 +340,7 @@ FStrataData GetStrataSlabBSDF(
 
 	const uint bHasEdgeColor = any(EdgeColor < 1);
 	const uint bHasHaziness = Haziness > 0.0f ? 1 : 0;
+	const uint bHasThinFilm = ThinFilmThickness > 0.0f ? 1 : 0;
 	const uint bHasAnisotropy = StrataHasAnisotropy(Roughness.x, Roughness.y) ? 1 : 0;
 	if (bHasAnisotropy)
 	{
@@ -353,6 +361,7 @@ FStrataData GetStrataSlabBSDF(
 	BSDF_SETHASSSSPROFILE	(StrataData.Layers[0].BSDFs[0],		bHasSSSProfile ? 1 : 0);
 	BSDF_SETHASHAZINESS		(StrataData.Layers[0].BSDFs[0],		bHasHaziness);
 	BSDF_SETHASEDGECOLOR	(StrataData.Layers[0].BSDFs[0],		bHasEdgeColor);
+	BSDF_SETHASTHINFILM		(StrataData.Layers[0].BSDFs[0],		bHasThinFilm);
 
 	SLAB_BASECOLOR			(StrataData.Layers[0].BSDFs[0])		= BaseColor;
 	SLAB_EDGECOLOR			(StrataData.Layers[0].BSDFs[0])		= EdgeColor;
@@ -361,6 +370,7 @@ FStrataData GetStrataSlabBSDF(
 	SLAB_ROUGHNESSX			(StrataData.Layers[0].BSDFs[0])		= Roughness.x;
 	SLAB_ROUGHNESSY			(StrataData.Layers[0].BSDFs[0])		= Roughness.y;
 	SLAB_HAZINESS			(StrataData.Layers[0].BSDFs[0])		= Haziness;
+	SLAB_THINFILM_THICKNESS	(StrataData.Layers[0].BSDFs[0])		= ThinFilmThickness;
 	if (bHasSSSProfile)
 	{
 		SLAB_SSSPROFILEID	(StrataData.Layers[0].BSDFs[0])		= SSSProfileID;
@@ -655,6 +665,7 @@ void SanitizeStrataSlab(inout FStrataBSDF BSDF)
 	SLAB_ROUGHNESSX(BSDF) = saturate(SLAB_ROUGHNESSX(BSDF));
 	SLAB_ROUGHNESSY(BSDF) = saturate(SLAB_ROUGHNESSY(BSDF));
 	SLAB_HAZINESS(BSDF) = saturate(SLAB_HAZINESS(BSDF));
+	SLAB_THINFILM_THICKNESS(BSDF) = saturate(SLAB_THINFILM_THICKNESS(BSDF));
 	SLAB_SSSPROFILERADIUSSCALE(BSDF) = saturate(SLAB_SSSPROFILERADIUSSCALE(BSDF));
 	SLAB_SSSDMFP(BSDF) = clamp(SLAB_SSSDMFP(BSDF), 0.0f, Max111110BitsFloat3);
 }
@@ -1292,9 +1303,9 @@ void PackStrataOut(
 				STRATA_STORE_UINT1(PackColorLinearToGamma2AlphaLinear(float4(SLAB_BASECOLOR(BSDF), SLAB_ROUGHNESSX(BSDF))));
 				STRATA_STORE_UINT1(PackRGBA8(float4(SLAB_SPECULAR(BSDF), SLAB_METALLIC(BSDF), SLAB_ROUGHNESSY(BSDF), SLAB_HAZINESS(BSDF))));
 
-				if (BSDF_GETHASEDGECOLOR(BSDF))
+				if (BSDF_GETHASEDGECOLOR(BSDF) || BSDF_GETHASTHINFILM(BSDF))
 				{
-					STRATA_STORE_UINT1(PackColorLinearToGamma2AlphaLinear(float4(SLAB_EDGECOLOR(BSDF), 0)));
+					STRATA_STORE_UINT1(PackColorLinearToGamma2AlphaLinear(float4(SLAB_EDGECOLOR(BSDF), SLAB_THINFILM_THICKNESS(BSDF))));
 				}
 
 				if (BSDF_GETHASSCATTERING(BSDF))
@@ -1376,7 +1387,7 @@ FStrataPixelHeader UnpackStrataHeaderIn(ByteAddressBuffer StrataBuffer, inout ui
 }
 
 // Unpack only strata header 
-// This version assume only a single slab (isotropic BSDF, with no edge color, no subsurface scattering, and no haziness)
+// This version assume only a single slab (isotropic BSDF, with no edge color, no subsurface scattering, no haziness, and no thin film)
 FStrataPixelHeader UnpackFastPathStrataHeaderIn(ByteAddressBuffer StrataBuffer, inout uint ByteOffset)
 {
 	FStrataPixelHeader Out = InitialiseStrataPixelHeader();
@@ -1407,6 +1418,7 @@ void SetStrataFastPathState(inout FStrataBSDF BSDF)
 	BSDF_SETHASSCATTERING(BSDF, 0);
 	BSDF_SETHASSSSPROFILE(BSDF, 0);
 	BSDF_SETHASHAZINESS(BSDF, 0);
+	BSDF_SETHASTHINFILM(BSDF, 0);
 }
 
 // Unpack a single BSDF
@@ -1445,10 +1457,10 @@ FStrataBSDF UnpackStrataBSDFIn(ByteAddressBuffer StrataBuffer, inout uint ByteOf
 		}
 
 		SLAB_EDGECOLOR(OutBSDF) = 1.0f;
-		if (BSDF_GETHASEDGECOLOR(OutBSDF))
+		SLAB_THINFILM_THICKNESS(OutBSDF) = 0.0f;
+		if (BSDF_GETHASEDGECOLOR(OutBSDF) || BSDF_GETHASTHINFILM(OutBSDF))
 		{
-			float Unused;
-			StrataLoad_ColorGamma2ToLinearAlphaLinear(StrataBuffer, ByteOffset, SLAB_EDGECOLOR(OutBSDF), Unused);
+			StrataLoad_ColorGamma2ToLinearAlphaLinear(StrataBuffer, ByteOffset, SLAB_EDGECOLOR(OutBSDF), SLAB_THINFILM_THICKNESS(OutBSDF));
 		}
 
 		if (BSDF_GETHASSCATTERING(OutBSDF))
@@ -1493,4 +1505,3 @@ FStrataBSDF UnpackStrataBSDFIn(ByteAddressBuffer StrataBuffer, inout uint ByteOf
 
 	return OutBSDF;
 }
-
diff --git a/Engine/Shaders/Private/Strata/StrataEvaluation.ush b/Engine/Shaders/Private/Strata/StrataEvaluation.ush
--- a/Engine/Shaders/Private/Strata/StrataEvaluation.ush
+++ b/Engine/Shaders/Private/Strata/StrataEvaluation.ush
@@ -8,6 +8,7 @@
 #include "../MonteCarlo.ush"
 #include "../SHCommon.ush"
 #include "../ShadingModels.ush"
+#include "../ThinFilmBSDF.ush"
 
 
 
@@ -234,7 +235,20 @@ FStrataEvaluateResult StrataEvaluateBSDF(FStrataBSDFContext BSDFContext, FShadow
 					HazePDF = VisibleGGXPDF(BSDFContext.TangentV, BSDFContext.TangentH, a2Haze);
 				}
 			}
-			float3 F = F_Schlick(F0, F90, BSDFContext.Context.VoH);
+
+			float3 F = 0;
+			const bool bHasThinFilm = BSDF_GETHASTHINFILM(BSDFContext.BSDF);
+			BRANCH
+			if (!bHasThinFilm)
+			{
+				F = F_Schlick(F0, F90, BSDFContext.Context.VoH);
+			}
+			else
+			{
+				float ThinFilmIOR		= SLAB_THINFILM_IOR(BSDFContext.BSDF);
+				float ThinFilmThickness = SLAB_THINFILM_THICKNESS(BSDFContext.BSDF);
+				F = F_ThinFilm(BSDFContext.Context.NoV, BSDFContext.Context.NoL, BSDFContext.Context.VoH, F0, F90, ThinFilmIOR, ThinFilmThickness);
+			}
 
 			Sample.SpecularPathProbability	= bHaziness		? (1.0f - SPECULAR_HAZINESS_WEIGHT) : 1.0f;
 			Sample.SpecularHazePathProbability= bHaziness	? SPECULAR_HAZINESS_WEIGHT			: 0.0f;
@@ -443,9 +457,19 @@ FStrataEnvLightResult StrataEvaluateForEnvLight(FStrataBSDFContext BSDFContext,
 			BRANCH
 			if (bEnableSpecular)
 			{
-				const float3 F0 = ComputeF0(SLAB_SPECULAR(BSDFContext.BSDF), BaseColor, Metallic);
+				float3 F0 = ComputeF0(SLAB_SPECULAR(BSDFContext.BSDF), BaseColor, Metallic);
 				const float3 F90 = ComputeF90(F0, SLAB_EDGECOLOR(BSDFContext.BSDF), Metallic);
 
+				float3 FresnelWeight = 1;
+				const bool bHasThinFilm = BSDF_GETHASTHINFILM(BSDFContext.BSDF);
+				if (bHasThinFilm)
+				{
+					float ThinFilmIOR = SLAB_THINFILM_IOR(BSDFContext.BSDF);
+					float ThinFilmThickness = SLAB_THINFILM_THICKNESS(BSDFContext.BSDF);
+					FresnelWeight = F_ThinFilm(BSDFContext.Context.NoV, BSDFContext.Context.NoL, BSDFContext.Context.VoH, F0, F90, ThinFilmIOR, ThinFilmThickness);
+					F0 = 1.f;
+				}
+
 				float SafeRoughnessX = MakeRoughnessSafe(SLAB_ROUGHNESSX(BSDFContext.BSDF));
 
 				const bool bHasAnisotropy = BSDF_GETANISOTROPY(BSDFContext.BSDF);
@@ -464,7 +488,7 @@ FStrataEnvLightResult StrataEvaluateForEnvLight(FStrataBSDFContext BSDFContext,
 					BSDFContext.R = 2 * dot(BSDFContext.V, BSDFContext.N) * BSDFContext.N - BSDFContext.V;
 				}
 
-				float3 EvalEnvBRDF = EnvBRDFApprox(F0, F90, SafeRoughnessX, BSDFContext.SatNoV);
+				float3 EvalEnvBRDF = FresnelWeight * EnvBRDFApprox(F0, F90, SafeRoughnessX, BSDFContext.SatNoV);
 
 				StrataEnvLightResult.SpecularDirection = StrataGetOffSpecularPeakReflectionDir(BSDFContext.N, BSDFContext.R, SafeRoughnessX);
 				StrataEnvLightResult.SpecularWeight = EvalEnvBRDF;
diff --git a/Engine/Shaders/Private/Strata/StrataVisualize.usf b/Engine/Shaders/Private/Strata/StrataVisualize.usf
--- a/Engine/Shaders/Private/Strata/StrataVisualize.usf
+++ b/Engine/Shaders/Private/Strata/StrataVisualize.usf
@@ -285,6 +285,16 @@ void VisualizeMaterialPS(
 						NEXTLINE;
 					}
 				}
+
+				if (BSDF_GETHASTHINFILM(BSDF))
+				{
+					HORIZONTALTAB; PRINTPARAMCHAR(_T_); PRINTPARAMCHAR(_F_); PRINTPARAMCHAR(_SPC_); PRINTPARAMCHAR(_I_); PRINTPARAMCHAR(_O_); PRINTPARAMCHAR(_R_); PRINTPARAMCHAR(_SPC_);
+					HORIZONTALTAB; PRINTFLOAT1(SLAB_THINFILM_IOR(BSDF));
+					NEXTLINE;
+					HORIZONTALTAB; PRINTPARAMCHAR(_T_); PRINTPARAMCHAR(_F_); PRINTPARAMCHAR(_SPC_); PRINTPARAMCHAR(_T_); PRINTPARAMCHAR(_C_); PRINTPARAMCHAR(_K_); PRINTPARAMCHAR(_SPC_);
+					HORIZONTALTAB; PRINTFLOAT1(SLAB_THINFILM_THICKNESS(BSDF));
+					NEXTLINE;
+				}
 #else
 				NEXTLINE;
 				NEXTLINE;
diff --git a/Engine/Shaders/Private/ThinFilmBSDF.ush b/Engine/Shaders/Private/ThinFilmBSDF.ush
new file mode 100644
--- /dev/null
+++ b/Engine/Shaders/Private/ThinFilmBSDF.ush
@@ -0,0 +1,253 @@
+// Copyright Epic Games, Inc. All Rights Reserved.
+
+/////////////////////////////////////////////////////////////////////////////////////////////////
+// Iridescence
+
+// Depolarization functions for natural light
+float ThinFilmDepol(float2 polV) { return 0.5 * (polV.x + polV.y); }
+float3 ThinFilmDepolColor(float3 colS, float3 colP) { return 0.5 * (colS + colP); }
+
+// Fresnel equations for dielectric/dielectric interfaces.
+void ThinFilmFresnelDielectric(
+	in float ct1,
+	in float n1,
+	in float n2,
+	out float2 R,
+	out float2 phi)
+{
+	float st1 = (1 - ct1 * ct1); // Sinus theta1 'squared'
+	float nr = n1 / n2;
+
+	if (Pow2(nr) * st1 > 1)  // Total reflection
+	{
+		R = float2(1, 1);
+		phi = 2.0 * atan(float2(
+			-Pow2(nr) * sqrt(st1 - 1.0 / Pow2(nr)) / ct1,
+			-sqrt(st1 - 1.0 / Pow2(nr)) / ct1));
+	}
+	else    // Transmission & Reflection
+	{
+
+		float ct2 = sqrt(1 - Pow2(nr) * st1);
+		float2 r = float2(
+			(n2 * ct1 - n1 * ct2) / (n2 * ct1 + n1 * ct2),
+			(n1 * ct1 - n2 * ct2) / (n1 * ct1 + n2 * ct2));
+		phi.x = (r.x < 0.0) ? PI : 0.0;
+		phi.y = (r.y < 0.0) ? PI : 0.0;
+		R = Pow2(r);
+	}
+}
+
+// Fresnel equations for dielectric/conductor interfaces.
+void ThinFilmFresnelConductor(
+	in float ct1,
+	in float n1,
+	in float n2,
+	in float k,
+	out float2 R,
+	out float2 phi)
+{
+	if (k == 0)
+	{ // use dielectric formula to avoid numerical issues
+		ThinFilmFresnelDielectric(ct1, n1, n2, R, phi);
+	}
+	else
+	{
+		float A = Pow2(n2) * (1 - Pow2(k)) - Pow2(n1) * (1 - Pow2(ct1));
+		float B = sqrt(Pow2(A) + Pow2(2 * Pow2(n2) * k));
+		float U = sqrt((A + B) / 2.0);
+		float V = sqrt((B - A) / 2.0);
+
+		R.y = (Pow2(n1 * ct1 - U) + Pow2(V)) / (Pow2(n1 * ct1 + U) + Pow2(V));
+		phi.y = atan2(2 * n1 * V * ct1, Pow2(U) + Pow2(V) - Pow2(n1 * ct1)) + PI;
+
+		R.x = (Pow2(Pow2(n2) * (1 - Pow2(k)) * ct1 - n1 * U) + Pow2(2 * Pow2(n2) * k * ct1 - n1 * V)) / (Pow2(Pow2(n2) * (1 - Pow2(k)) * ct1 + n1 * U) + Pow2(2 * Pow2(n2) * k * ct1 + n1 * V));
+		phi.x = atan2(2 * n1 * Pow2(n2) * ct1 * (2 * k * U - (1 - Pow2(k)) * V), Pow2(Pow2(n2) * (1 + Pow2(k)) * ct1) - Pow2(n1) * (Pow2(U) + Pow2(V)));
+	}
+}
+
+// Evaluation XYZ sensitivity curves in Fourier space
+float3 ThinFilmEvalSensitivity(float opd, float shift)
+{
+	// Use Gaussian fits, given by 3 parameters: val, pos and var
+	float phase = 2 * PI * opd * 1.0e-6;
+	float3 val = float3(5.4856e-13, 4.4201e-13, 5.2481e-13);
+	float3 pos = float3(1.6810e+06, 1.7953e+06, 2.2084e+06);
+	float3 var = float3(4.3278e+09, 9.3046e+09, 6.6121e+09);
+	float3 xyz = val * sqrt(2 * PI * var) * cos(pos * phase + shift) * exp(-var * phase * phase);
+	xyz.x += 9.7470e-14 * sqrt(2 * PI * 4.5282e+09) * cos(2.2399e+06 * phase + shift) * exp(-4.5282e+09 * phase * phase);
+	return xyz / 1.0685e-7;
+	// As pointed by the suplemental material the oupout is not in the correct space, we need to tranfser 
+}
+
+
+// [Gulbrandsen, "Artist Friendly Metallic Fresnel"]
+// Compute complex index of refraction N = n + ik from F0 and edge tint.
+void ComputeComplexIORFromF0AndEdgeTint(in float3 Reflectivity, in float3 EdgeTint, inout float3 n, inout float3 k)
+{
+	const float3 g = EdgeTint;
+	const float3 r = clamp(Reflectivity, 0.0, 0.999);
+	const float3 SqrtR = sqrt(r);
+
+	n = g * (1 - r) / (1 + r) + (1 - g) * (1 + SqrtR) / (1 - SqrtR);
+
+	const float3 k2 = ((n + 1) * (n + 1) * r - (n - 1) * (n - 1)) / (1 - r);
+	k = sqrt(k2);
+}
+
+// Fresnel term for iridescent microfacet BRDF model 
+// This is a reference implementation adapted from the original paper's source code
+float3 F_ThinFilmRef(float NoL, float NoV, float VoH, float3 F0, float3 F90, float ThinFilmIOR, float ThinFilmTickness)
+{
+	// Translate F0/F90 to eta/kappa form
+	float3 n3 = 1;
+	float3 k3 = 0;
+	ComputeComplexIORFromF0AndEdgeTint(F0, F90, n3, k3);
+
+	// Thin film is limited (for now) to the top most layer (i.e., the upper layer need to be air)
+	//					|
+	//	  n1=1 / k1=0	| Air
+	// ----------------- 
+	//		n2 / k2=0	| Thin film
+	// ----------------- 
+	//		n3 / k3		| Regular material
+	//					|
+	//
+	//		        Min		Max		Default
+	// float Dinc	0.0		10.0	0.5
+	// float eta2	1.0		5.0		2.0
+	// float eta3	1.0		5.0		3.0
+	// float kappa3	0.0		5.0		0.0
+	// 
+	// n1 is fixed to air (~1.f)
+	float3 Avg = 1.f / 3.f;
+	float eta2	 = clamp(ThinFilmIOR, 1.f, 3.f);
+	float eta3	 = clamp(dot(Avg,n3), 1.f, 3.f);
+	float kappa3 = clamp(dot(Avg,k3), 1.f, 5.f);
+	float Dinc   = ThinFilmTickness * 10; // ThinFilmTickness is in micrometer. Range is 0 - 10um i.e, 0 - 10 000nm
+
+	// Force eta_2 -> 1.0 when Dinc -> 0.0
+	const float eta_2 = lerp(1.0, eta2, smoothstep(0.0, 0.03, Dinc));
+
+	// Compute dot products
+	const bool bIsValid = !(NoL < 0 || NoV < 0);
+
+	float3 I = 0;
+	if (bIsValid)
+	{
+		float cosTheta1 = VoH;
+		float cosTheta2 = sqrt(1.0 - Pow2(1.0 / eta_2) * (1 - Pow2(cosTheta1)));
+
+		// First interface
+		float2 R12, phi12;
+		ThinFilmFresnelDielectric(cosTheta1, 1.0, eta_2, R12, phi12);
+		float2 R21 = R12;
+		float2 T121 = float2(1, 1) - R12;
+		float2 phi21 = float2(PI, PI) - phi12;
+
+		// Second interface
+		float2 R23, phi23;
+		ThinFilmFresnelConductor(cosTheta2, eta_2, eta3, kappa3, R23, phi23);
+
+		// Phase shift
+		float OPD = Dinc * cosTheta2;
+		float2 phi2 = phi21 + phi23;
+
+		// Compound terms
+		float2 R123 = R12 * R23;
+		float2 r123 = sqrt(R123);
+		float2 Rs = Pow2(T121) * R23 / (1 - R123);
+
+		// Reflectance term for m=0 (DC term amplitude)
+		float2 C0 = R12 + Rs;
+		float3 S0 = ThinFilmEvalSensitivity(0.0, 0.0);
+		I += ThinFilmDepol(C0) * S0;
+
+		// Reflectance term for m>0 (pairs of diracs)
+		float2 Cm = Rs - T121;
+		for (int m = 1; m <= 3; ++m)
+		{
+			Cm *= r123;
+			float3 SmS = 2.0 * ThinFilmEvalSensitivity(m * OPD, m * phi2.x);
+			float3 SmP = 2.0 * ThinFilmEvalSensitivity(m * OPD, m * phi2.y);
+			I += ThinFilmDepolColor(Cm.x * SmS, Cm.y * SmP);
+		}
+
+		// Convert back to RGB reflectance
+		// XYZ to CIE 1931 RGB color space (using neutral E illuminant)
+		const float3x3 ThinFilm_XYZ_TO_RGB = float3x3(2.3706743, -0.5138850, 0.0052982, -0.9000405, 1.4253036, -0.0146949, -0.4706338, 0.0885814, 1.0093968);
+		I = clamp(mul(I, ThinFilm_XYZ_TO_RGB), float3(0, 0, 0), float3(1, 1, 1));
+	}
+	return I;
+}
+
+// Fresnel term for iridescent microfacet BRDF model 
+// Simplified version which relies on Schlick's Fresnel and de facto does not take into 
+// account Fresnel phase shift & polarization.
+float3 F_ThinFilm(float NoL, float NoV, float VoH, float3 F0, float3 F90, float ThinFilmIOR, float ThinFilmTickness)
+{
+	// Thin film is limited (for now) to the top most layer (i.e., the upper layer need to be air)
+	//					|
+	//	  n1=1 / k1=0	| Air
+	// ----------------- 
+	//		n2 / k2=0	| Thin film
+	// ----------------- 
+	//		n3 / k3		| Regular material
+	//					|	
+	
+	float Dinc = ThinFilmTickness * 10; // ThinFilmTickness is in micrometer. Range is 0 - 10um i.e, 0 - 10 000nm
+
+	// Force eta2 -> 1.0 when Dinc -> 0.0
+	const float eta2 = lerp(1.0, ThinFilmIOR, smoothstep(0.0, 0.03, Dinc));
+
+	// Compute dot products
+	const bool bIsValid = !(NoL < 0 || NoV < 0);
+
+	float3 I = 0;
+	if (bIsValid)
+	{
+		float cosTheta1 = VoH;
+		float cosTheta2 = sqrt(1.0 - Pow2(1.0 / eta2) * (1 - Pow2(cosTheta1)));
+
+		// First interface is dieletric so it is achromatic
+		float R12 = F_Schlick(DielectricIorToF0(ThinFilmIOR), VoH).x;
+		float phi12 = 0;
+		float R21 = R12;
+		float T121 = 1 - R12;
+		float phi21 = PI - phi12;
+
+		// Second interface is either conductor or dieletric, so can be chromatic. 
+		// Ideally we should recompute F0 (R23) by taking into account the thin layer IOR
+		float3 R23 = F0;
+		float phi23 = 0;
+
+		// Phase shift
+		float OPD = Dinc * cosTheta2;
+		float phi2 = phi21 + phi23;
+
+		// Compound terms
+		float3 R123 = R12 * R23;
+		float3 r123 = sqrt(R123);
+		float3 Rs = Pow2(T121) * R23 / (1 - R123);
+
+		// Reflectance term for m=0 (DC term amplitude)
+		float3 C0 = R12 + Rs;
+		float3 S0 = ThinFilmEvalSensitivity(0.0, 0.0);
+		I += C0 * S0;
+
+		// Reflectance term for m>0 (pairs of diracs)
+		// Simplified based as we don't take into account polarization
+		float3 Cm = Rs - T121;
+		for (int m = 1; m <= 3; ++m)
+		{
+			Cm *= r123;
+			I += Cm * 2.0 * ThinFilmEvalSensitivity(m * OPD, m * phi2);
+		}
+
+		// Convert back to RGB reflectance
+		// XYZ to CIE 1931 RGB color space (using neutral E illuminant)
+		const float3x3 ThinFilm_XYZ_TO_RGB = float3x3(2.3706743, -0.5138850, 0.0052982, -0.9000405, 1.4253036, -0.0146949, -0.4706338, 0.0885814, 1.0093968);
+		I = saturate(mul(I, ThinFilm_XYZ_TO_RGB));
+	}
+	return I;
+}
\ No newline at end of file
diff --git a/Engine/Source/Runtime/Engine/Classes/Materials/MaterialExpressionStrata.h b/Engine/Source/Runtime/Engine/Classes/Materials/MaterialExpressionStrata.h
--- a/Engine/Source/Runtime/Engine/Classes/Materials/MaterialExpressionStrata.h
+++ b/Engine/Source/Runtime/Engine/Classes/Materials/MaterialExpressionStrata.h
@@ -109,6 +109,12 @@ class UMaterialExpressionStrataSlabBSDF : public UMaterialExpressionStrataBSDF
 	UPROPERTY()
 	FExpressionInput Haziness;
 
+	/**
+	 * Thin film contosl the thin film layer coating the current slab. 0 means disabled and 1 means a coating layer of 10 micrometer. (type = float, unitless, default = 0)
+	 */
+	UPROPERTY()
+	FExpressionInput ThinFilmThickness;
+
 	/** SubsurfaceProfile, for Screen Space Subsurface Scattering. The profile needs to be set up on both the Strata diffuse node, and the material node at the moment. */
 	UPROPERTY(EditAnywhere, BlueprintReadOnly, Category = Material, meta = (DisplayName = "Subsurface Profile"))
 	TObjectPtr<class USubsurfaceProfile> SubsurfaceProfile;
@@ -125,6 +131,7 @@ class UMaterialExpressionStrataSlabBSDF : public UMaterialExpressionStrataBSDF
 
 	bool HasEdgeColor() const;
 	bool HasScattering() const;
+	bool HasThinFilm() const;
 #endif
 	//~ End UMaterialExpression Interface
 };
diff --git a/Engine/Source/Runtime/Engine/Private/Materials/HLSLMaterialTranslator.cpp b/Engine/Source/Runtime/Engine/Private/Materials/HLSLMaterialTranslator.cpp
--- a/Engine/Source/Runtime/Engine/Private/Materials/HLSLMaterialTranslator.cpp
+++ b/Engine/Source/Runtime/Engine/Private/Materials/HLSLMaterialTranslator.cpp
@@ -1800,8 +1800,8 @@ bool FHLSLMaterialTranslator::Translate()
 
 			// Output some debug info as comment in code and in the material stat window
 			static const auto CVarStrataBytePerPixel = IConsoleManager::Get().FindTConsoleVariableDataInt(TEXT("r.Strata.BytesPerPixel"));
-			check(CVarStrataBytePerPixel)
-				const uint32 StrataBytePerPixel = CVarStrataBytePerPixel ? CVarStrataBytePerPixel->GetValueOnAnyThread() : 0;
+			check(CVarStrataBytePerPixel);
+			const uint32 StrataBytePerPixel = CVarStrataBytePerPixel ? CVarStrataBytePerPixel->GetValueOnAnyThread() : 0;
 			StrataMaterialAnalysis = StrataCompilationInfoMaterialAnalysis(this, StrataCompilationInfo, StrataBytePerPixel);
 
 			FString StrataMaterialDescription;
@@ -9031,10 +9031,10 @@ int32 FHLSLMaterialTranslator::StrataCreateAndRegisterNullMaterial()
 int32 FHLSLMaterialTranslator::StrataSlabBSDF(
 	int32 BaseColor, int32 EdgeColor, int32 Specular, int32 Metallic, int32 RoughnessX, int32 RoughnessY,
 	int32 SSSProfileId, int32 SSSDMFP, int32 SSSDMFPScale, int32 EmissiveColor,
-	int32 Haziness, int32 Normal, int32 Tangent, uint8 SharedNormalIndex)
+	int32 Haziness, int32 ThinFilmThickness, int32 Normal, int32 Tangent, uint8 SharedNormalIndex)
 {
 	return AddCodeChunk(
-		MCT_Strata, TEXT("GetStrataSlabBSDF(%s, %s, %s, %s, float2(%s, %s), %s, %s, %s, %s, %s, %u, Parameters.SharedNormals.NormalTypes) /* Normal = %s ; Tangent = %s */"),
+		MCT_Strata, TEXT("GetStrataSlabBSDF(%s, %s, %s, %s, float2(%s, %s), %s, %s, %s, %s, %s, %s, %u, Parameters.SharedNormals.NormalTypes) /* Normal = %s ; Tangent = %s */"),
 		*GetParameterCode(BaseColor),
 		*GetParameterCode(EdgeColor),
 		*GetParameterCode(Specular),
@@ -9046,6 +9046,7 @@ int32 FHLSLMaterialTranslator::StrataSlabBSDF(
 		*GetParameterCode(SSSDMFPScale),
 		*GetParameterCode(EmissiveColor),
 		*GetParameterCode(Haziness),
+		*GetParameterCode(ThinFilmThickness),
 		SharedNormalIndex,
 		*GetParameterCode(Normal),
 		*GetParameterCode(Tangent)
diff --git a/Engine/Source/Runtime/Engine/Private/Materials/HLSLMaterialTranslator.h b/Engine/Source/Runtime/Engine/Private/Materials/HLSLMaterialTranslator.h
--- a/Engine/Source/Runtime/Engine/Private/Materials/HLSLMaterialTranslator.h
+++ b/Engine/Source/Runtime/Engine/Private/Materials/HLSLMaterialTranslator.h
@@ -990,7 +990,7 @@ protected:
 	virtual int32 StrataSlabBSDF(
 		int32 BaseColor, int32 EdgeColor, int32 Specular, int32 Metallic, int32 RoughnessX, int32 RoughnessY,
 		int32 SSSProfileId, int32 SSSDMFP, int32 SSSDMFPScale,
-		int32 EmissiveColor, int32 Haziness, int32 Normal, int32 Tangent, uint8 SharedNormalIndex) override;
+		int32 EmissiveColor, int32 Haziness, int32 ThinFilmThickness, int32 Normal, int32 Tangent, uint8 SharedNormalIndex) override;
 	virtual int32 StrataSheenBSDF(int32 BaseColor, int32 Roughness, int32 Normal, uint8 SharedNormalIndex) override;
 	virtual int32 StrataVolumetricFogCloudBSDF(int32 Albedo, int32 Extinction, int32 EmissiveColor, int32 AmbientOcclusion) override;
 	virtual int32 StrataUnlitBSDF(int32 EmissiveColor, int32 TransmittanceColor) override;
diff --git a/Engine/Source/Runtime/Engine/Private/Materials/MaterialExpressions.cpp b/Engine/Source/Runtime/Engine/Private/Materials/MaterialExpressions.cpp
--- a/Engine/Source/Runtime/Engine/Private/Materials/MaterialExpressions.cpp
+++ b/Engine/Source/Runtime/Engine/Private/Materials/MaterialExpressions.cpp
@@ -19761,6 +19761,7 @@ int32 UMaterialExpressionStrataSlabBSDF::Compile(class FMaterialCompiler* Compil
 	uint8 SharedNormalIndex = bAnisotropyPotentiallyUsed ? StrataCompilationInfoCreateSharedNormal(Compiler, NormalCodeChunk, TangentCodeChunk) : StrataCompilationInfoCreateSharedNormal(Compiler, NormalCodeChunk);
 
 	const bool bHasEdgeColor = HasEdgeColor();
+	const bool bHasThinFilm = HasThinFilm();
 
 	int32 SSSProfileCodeChunk = INDEX_NONE;
 	const bool bHasScattering = HasScattering();
@@ -19782,10 +19783,11 @@ int32 UMaterialExpressionStrataSlabBSDF::Compile(class FMaterialCompiler* Compil
 		CompileWithDefaultFloat1(Compiler, SSSDMFPScale, 1.0f),
 		CompileWithDefaultFloat3(Compiler, EmissiveColor, 0.0f, 0.0f, 0.0f),
 		CompileWithDefaultFloat1(Compiler, Haziness, 0.0f),
+		CompileWithDefaultFloat1(Compiler, ThinFilmThickness, 0.0f),
 		NormalCodeChunk,
 		TangentCodeChunk,
 		SharedNormalIndex);
-	StrataCompilationInfoCreateSingleBSDFMaterial(Compiler, OutputCodeChunk, SharedNormalIndex, STRATA_BSDF_TYPE_SLAB, bHasEdgeColor, bHasScattering);
+	StrataCompilationInfoCreateSingleBSDFMaterial(Compiler, OutputCodeChunk, SharedNormalIndex, STRATA_BSDF_TYPE_SLAB, bHasEdgeColor, bHasScattering, bHasThinFilm);
 
 	return OutputCodeChunk;
 }
@@ -19840,6 +19842,9 @@ uint32 UMaterialExpressionStrataSlabBSDF::GetInputType(int32 InputIndex)
 	case 11:
 		return MCT_Float1; // Haziness
 		break;
+	case 12:
+		return MCT_Float1; // ThinFilm Thickness
+		break;
 	}
 
 	check(false);
@@ -19896,6 +19901,10 @@ FName UMaterialExpressionStrataSlabBSDF::GetInputName(int32 InputIndex) const
 	{
 		return TEXT("Haziness");
 	}
+	else if (InputIndex == 12)
+	{
+		return TEXT("ThinFilm Thickness");
+	}
 	return TEXT("Unknown");
 }
 
@@ -19930,6 +19939,11 @@ bool UMaterialExpressionStrataSlabBSDF::HasEdgeColor() const
 	return EdgeColor.IsConnected();
 }
 
+bool UMaterialExpressionStrataSlabBSDF::HasThinFilm() const
+{
+	return ThinFilmThickness.IsConnected();
+}
+
 #endif // WITH_EDITOR
 
 
diff --git a/Engine/Source/Runtime/Engine/Private/Materials/StrataMaterial.cpp b/Engine/Source/Runtime/Engine/Private/Materials/StrataMaterial.cpp
--- a/Engine/Source/Runtime/Engine/Private/Materials/StrataMaterial.cpp
+++ b/Engine/Source/Runtime/Engine/Private/Materials/StrataMaterial.cpp
@@ -52,7 +52,7 @@ uint8 StrataCompilationInfoCreateSharedNormal(FMaterialCompiler* Compiler, int32
 	return Compiler->StrataCompilationInfoRegisterSharedNormalIndex(NormalCodeChunk, TangentCodeChunk);
 }
 
-void StrataCompilationInfoCreateSingleBSDFMaterial(FMaterialCompiler* Compiler, int32 CodeChunk, uint8 SharedNormalIndex, uint8 BSDFType, bool bHasEdgeColor, bool bHasScattering)
+void StrataCompilationInfoCreateSingleBSDFMaterial(FMaterialCompiler* Compiler, int32 CodeChunk, uint8 SharedNormalIndex, uint8 BSDFType, bool bHasEdgeColor, bool bHasScattering, bool bHasThinFilm)
 {
 	FStrataMaterialCompilationInfo StrataInfo;
 	StrataInfo.LayerCount = 1;
@@ -61,6 +61,7 @@ void StrataCompilationInfoCreateSingleBSDFMaterial(FMaterialCompiler* Compiler,
 	StrataInfo.Layers[0].BSDFs[0].SharedNormalIndex = SharedNormalIndex;
 	StrataInfo.Layers[0].BSDFs[0].bHasScattering = bHasScattering;
 	StrataInfo.Layers[0].BSDFs[0].bHasEdgeColor = bHasEdgeColor;
+	StrataInfo.Layers[0].BSDFs[0].bHasThinFilm = bHasThinFilm;
 	UpdateTotalBSDFCount(StrataInfo);
 	Compiler->StrataCompilationInfoRegisterCodeChunk(CodeChunk, StrataInfo);
 }
@@ -281,7 +282,7 @@ FStrataMaterialAnalysisResult StrataCompilationInfoMaterialAnalysis(FMaterialCom
 			{
 				Result.RequestedByteCount += UintByteSize;
 				Result.RequestedByteCount += UintByteSize;
-				if (BSDF.bHasEdgeColor)
+				if (BSDF.bHasEdgeColor || BSDF.bHasThinFilm)
 				{
 					Result.RequestedByteCount += UintByteSize;
 				}
diff --git a/Engine/Source/Runtime/Engine/Private/Materials/StrataMaterial.h b/Engine/Source/Runtime/Engine/Private/Materials/StrataMaterial.h
--- a/Engine/Source/Runtime/Engine/Private/Materials/StrataMaterial.h
+++ b/Engine/Source/Runtime/Engine/Private/Materials/StrataMaterial.h
@@ -29,6 +29,7 @@ struct FStrataMaterialCompilationInfo
 		uint8 SharedNormalIndex;
 		bool  bHasScattering;
 		bool  bHasEdgeColor;
+		bool  bHasThinFilm;
 	};
 
 	struct FLayer
@@ -55,7 +56,7 @@ FString GetStrataBSDFName(uint8 BSDFType);
 uint8 StrataCompilationInfoCreateSharedNormal(FMaterialCompiler* Compiler, int32 NormalCodeChunk);
 uint8 StrataCompilationInfoCreateSharedNormal(FMaterialCompiler* Compiler, int32 NormalCodeChunk, int32 TangentCodeChunk);
 void  StrataCompilationInfoCreateNullBSDF(FMaterialCompiler* Compiler, int32 CodeChunk);
-void  StrataCompilationInfoCreateSingleBSDFMaterial(FMaterialCompiler* Compiler, int32 CodeChunk, uint8 SharedNormalIndex, uint8 BSDFType, bool bHasEdgeColor = false, bool bHasScattering = false);
+void  StrataCompilationInfoCreateSingleBSDFMaterial(FMaterialCompiler* Compiler, int32 CodeChunk, uint8 SharedNormalIndex, uint8 BSDFType, bool bHasEdgeColor = false, bool bHasScattering = false, bool bHasThinFilm = false);
 
 FStrataMaterialCompilationInfo StrataCompilationInfoAdd(FMaterialCompiler* Compiler, const FStrataMaterialCompilationInfo& A, const FStrataMaterialCompilationInfo& B);
 FStrataMaterialCompilationInfo StrataCompilationInfoMultiply(FMaterialCompiler* Compiler, const FStrataMaterialCompilationInfo& A);
diff --git a/Engine/Source/Runtime/Engine/Public/MaterialCompiler.h b/Engine/Source/Runtime/Engine/Public/MaterialCompiler.h
--- a/Engine/Source/Runtime/Engine/Public/MaterialCompiler.h
+++ b/Engine/Source/Runtime/Engine/Public/MaterialCompiler.h
@@ -410,7 +410,7 @@ public:
 	virtual int32 StrataSlabBSDF(
 		int32 BaseColor, int32 EdgeColor, int32 Specular, int32 Metallic, int32 RoughnessX, int32 RoughnessY, 
 		int32 SSSProfileId, int32 SSSDMFP, int32 SSSDMFPScale,
-		int32 EmissiveColor, int32 Haziness, int32 Normal, int32 Tangent, uint8 SharedNormalIndex) = 0;
+		int32 EmissiveColor, int32 Haziness, int32 ThinFilmThickness, int32 Normal, int32 Tangent, uint8 SharedNormalIndex) = 0;
 	virtual int32 StrataSheenBSDF(int32 BaseColor, int32 Roughness, int32 Normal, uint8 SharedNormalIndex) = 0;
 	virtual int32 StrataVolumetricFogCloudBSDF(int32 Albedo, int32 Extinction, int32 EmissiveColor, int32 AmbientOcclusion) = 0;
 	virtual int32 StrataUnlitBSDF(int32 EmissiveColor, int32 TransmittanceColor) = 0;
@@ -890,12 +890,12 @@ public:
 	virtual int32 StrataSlabBSDF(
 		int32 BaseColor, int32 EdgeColor, int32 Specular, int32 Metallic, int32 RoughnessX, int32 RoughnessY,
 		int32 SSSProfileId, int32 SSSDMFP, int32 SSSDMFPScale,
-		int32 EmissiveColor, int32 Haziness, int32 Normal, int32 Tangent, uint8 SharedNormalIndex) override
+		int32 EmissiveColor, int32 Haziness, int32 ThinFilmThickness, int32 Normal, int32 Tangent, uint8 SharedNormalIndex) override
 	{
 		return Compiler->StrataSlabBSDF(
 			BaseColor, EdgeColor, Specular, Metallic, RoughnessX, RoughnessY,
 			SSSProfileId, SSSDMFP, SSSDMFPScale,
-			EmissiveColor, Haziness, Normal, Tangent, SharedNormalIndex);
+			EmissiveColor, Haziness, ThinFilmThickness, Normal, Tangent, SharedNormalIndex);
 	}
 
 	virtual int32 StrataSheenBSDF(int32 BaseColor, int32 Roughness, int32 Normal, uint8 SharedNormalIndex) override
