Invert depth used for mobile CSM shadow sampling to match to match reversed shadow projection matrices

diff --git a/Engine/Shaders/Private/MobileShadingModels.ush b/Engine/Shaders/Private/MobileShadingModels.ush
--- a/Engine/Shaders/Private/MobileShadingModels.ush
+++ b/Engine/Shaders/Private/MobileShadingModels.ush
@@ -264,7 +264,10 @@ half MobileDirectionalLightCSM(float2 ScreenPosition, float SceneDepth)
 	{
 		// Clamp pixel depth in light space for shadowing opaque, because areas of the shadow depth buffer that weren't rendered to will have been cleared to 1
 		// We want to force the shadow comparison to result in 'unshadowed' in that case, regardless of whether the pixel being shaded is in front or behind that plane
-		float LightSpacePixelDepthForOpaque = min(ShadowPosition.z, 0.99999f);
+
+		// Invert ShadowZ as the shadow space has been changed (but not yet the filtering code)
+		float ShadowZ = 1.0f - ShadowPosition.z;
+		float LightSpacePixelDepthForOpaque = min(ShadowZ, 0.99999f);
 		Settings.SceneDepth = LightSpacePixelDepthForOpaque;
 
 		ShadowMap = MobileShadowPCF(ShadowPosition.xy, Settings);
