UE5 - Use DXC intrinsic for reversebits

diff --git a/Engine/Shaders/Private/Common.ush b/Engine/Shaders/Private/Common.ush
--- a/Engine/Shaders/Private/Common.ush
+++ b/Engine/Shaders/Private/Common.ush
@@ -1727,7 +1727,7 @@ float DecodePackedHeight(float2 PackedHeight)
 /** Reverses all the 32 bits. */
 uint ReverseBits32( uint bits )
 {
-#if SM6_PROFILE || SM5_PROFILE || COMPILER_METAL
+#if SM6_PROFILE || SM5_PROFILE || COMPILER_METAL || COMPILER_DXC
 	return reversebits( bits );
 #else
 	bits = ( bits << 16) | ( bits >> 16);
