Small refactor to hide strata addressing behiond a structure to be able to change it easily later when experimenting.

diff --git a/Engine/Shaders/Private/BasePassPixelShader.usf b/Engine/Shaders/Private/BasePassPixelShader.usf
--- a/Engine/Shaders/Private/BasePassPixelShader.usf
+++ b/Engine/Shaders/Private/BasePassPixelShader.usf
@@ -1548,12 +1548,12 @@ void FPixelShaderInOut_MainPS(
 	GBuffer.WorldNormal = 0; // Normal will be overriden by the top layer normal
 
 	uint2 PixelPos = In.SvPosition.xy;
-	uint PixelStrataDataByteOffset = GetStrataPixelDataByteOffset(PixelPos, ResolvedView.BufferSizeAndInvSize.x, OpaqueBasePass.Strata.MaxBytesPerPixel);
 
 	float3 EmissiveLuminance = 0.0f;
 	const float3 VertexNormal = MaterialParameters.TangentToWorld[2].xyz;
 	PackStrataOut(
-		OpaqueBasePass.Strata.MaterialLobesBufferUAV, PixelStrataDataByteOffset, 
+		OpaqueBasePass.Strata.MaterialLobesBufferUAV, 
+		GetStrataPixelDataByteOffset(PixelPos, ResolvedView.BufferSizeAndInvSize.x, OpaqueBasePass.Strata.MaxBytesPerPixel),
 		StrataPixelHeader, Strata,
 		MaterialParameters.CameraVector, GBuffer, EmissiveLuminance);
 
diff --git a/Engine/Shaders/Private/Strata/Strata.ush b/Engine/Shaders/Private/Strata/Strata.ush
--- a/Engine/Shaders/Private/Strata/Strata.ush
+++ b/Engine/Shaders/Private/Strata/Strata.ush
@@ -483,6 +483,41 @@ FStrataData GetStrataSingleLayerWaterBSDF(
 }
 
 
+struct FStrataAddressing
+{
+	uint CurrentAddr;
+	uint ReadBytes;
+};
+FStrataAddressing GetStrataPixelDataByteOffset(uint2 PixelPos, uint ViewBufferSizeX, uint StrataMaxBytesPerPixel)
+{
+	FStrataAddressing StrataAddressing;
+	StrataAddressing.CurrentAddr = (PixelPos.y * ViewBufferSizeX + PixelPos.x) * StrataMaxBytesPerPixel;
+	StrataAddressing.ReadBytes = 0;
+	return StrataAddressing;
+}
+
+#if COMPILER_SWITCH // These platforms are complaining so simply disabling Strata there for now.
+
+#define STRATA_STORE_UINT1(x) 
+
+#define STRATA_LOAD_UINT1(x) x = 0;
+
+#define STRATA_LOAD_FLOAT1(x) x = 0;
+
+uint  StrataLoadUint1(ByteAddressBuffer StrataBuffer, inout FStrataAddressing StrataAddressing) { return 0; }
+
+#else
+
+#define STRATA_STORE_UINT1(x) StrataBuffer.Store (StrataAddressing.CurrentAddr, asuint(x)); StrataAddressing.CurrentAddr += 4;	 StrataAddressing.ReadBytes += 4;
+
+#define STRATA_LOAD_UINT1(x) x = StrataBuffer.Load (StrataAddressing.CurrentAddr); StrataAddressing.CurrentAddr += 4;	StrataAddressing.ReadBytes += 4;
+
+#define STRATA_LOAD_FLOAT1(x) x = asfloat(StrataBuffer.Load (StrataAddressing.CurrentAddr)); StrataAddressing.CurrentAddr += 4;		StrataAddressing.ReadBytes += 4;
+
+uint  StrataLoadUint1(ByteAddressBuffer StrataBuffer, inout FStrataAddressing StrataAddressing) { uint  x = StrataBuffer.Load(StrataAddressing.CurrentAddr); StrataAddressing.CurrentAddr += 4;	StrataAddressing.ReadBytes += 4;	return x; }
+
+#endif
+
 
 ///////////////////////////////////////////////////////////////////////////////
 // Shared normals
@@ -963,96 +998,37 @@ void StrataUnpackNormalAndTangent(inout float3 Normal, inout float3 Tangent, in
 	}
 }
 
-
-#if COMPILER_SWITCH // These platforms are complaining so simply disabling Strata there for now.
-
-#define STRATA_STORE_UINT1(x) 
-#define STRATA_STORE_UINT2(x) 
-#define STRATA_STORE_UINT3(x) 
-#define STRATA_STORE_UINT4(x) 
-
-#define STRATA_LOAD_UINT1(x) x = 0;
-#define STRATA_LOAD_UINT2(x) x = 0;
-#define STRATA_LOAD_UINT3(x) x = 0;
-#define STRATA_LOAD_UINT4(x) x = 0;
-
-#define STRATA_LOAD_FLOAT1(x) x = 0;
-#define STRATA_LOAD_FLOAT2(x) x = 0;
-#define STRATA_LOAD_FLOAT3(x) x = 0;
-#define STRATA_LOAD_FLOAT4(x) x = 0;
-
-uint  StrataLoadUint1(ByteAddressBuffer StrataBuffer, inout uint ByteOffset) { return 0; }
-uint2 StrataLoadUint2(ByteAddressBuffer StrataBuffer, inout uint ByteOffset) { return 0; }
-uint3 StrataLoadUint3(ByteAddressBuffer StrataBuffer, inout uint ByteOffset) { return 0; }
-uint4 StrataLoadUint4(ByteAddressBuffer StrataBuffer, inout uint ByteOffset) { return 0; }
-
-#else
-
-#define STRATA_STORE_UINT1(x) StrataBuffer.Store (ByteOffset, asuint(x)); ByteOffset += 4;
-#define STRATA_STORE_UINT2(x) StrataBuffer.Store2(ByteOffset, asuint(x)); ByteOffset += 8;
-#define STRATA_STORE_UINT3(x) StrataBuffer.Store3(ByteOffset, asuint(x)); ByteOffset += 12;
-#define STRATA_STORE_UINT4(x) StrataBuffer.Store4(ByteOffset, asuint(x)); ByteOffset += 16;
-
-#define STRATA_LOAD_UINT1(x) x = StrataBuffer.Load (ByteOffset); ByteOffset += 4;
-#define STRATA_LOAD_UINT2(x) x = StrataBuffer.Load2(ByteOffset); ByteOffset += 8;
-#define STRATA_LOAD_UINT3(x) x = StrataBuffer.Load3(ByteOffset); ByteOffset += 12;
-#define STRATA_LOAD_UINT4(x) x = StrataBuffer.Load4(ByteOffset); ByteOffset += 16;
-
-#define STRATA_LOAD_FLOAT1(x) x = asfloat(StrataBuffer.Load (ByteOffset)); ByteOffset += 4;
-#define STRATA_LOAD_FLOAT2(x) x = asfloat(StrataBuffer.Load2(ByteOffset)); ByteOffset += 8;
-#define STRATA_LOAD_FLOAT3(x) x = asfloat(StrataBuffer.Load3(ByteOffset)); ByteOffset += 12;
-#define STRATA_LOAD_FLOAT4(x) x = asfloat(StrataBuffer.Load4(ByteOffset)); ByteOffset += 16;
-
-uint  StrataLoadUint1(ByteAddressBuffer StrataBuffer, inout uint ByteOffset) { uint  x = StrataBuffer.Load (ByteOffset); ByteOffset += 4; return x; }
-uint2 StrataLoadUint2(ByteAddressBuffer StrataBuffer, inout uint ByteOffset) { uint2 x = StrataBuffer.Load2(ByteOffset); ByteOffset += 8; return x; }
-uint3 StrataLoadUint3(ByteAddressBuffer StrataBuffer, inout uint ByteOffset) { uint3 x = StrataBuffer.Load3(ByteOffset); ByteOffset +=12; return x; }
-uint4 StrataLoadUint4(ByteAddressBuffer StrataBuffer, inout uint ByteOffset) { uint4 x = StrataBuffer.Load4(ByteOffset); ByteOffset +=16; return x; }
-
-#endif
-
-void StrataLoad_R16G16B16F_A16F(ByteAddressBuffer StrataBuffer, inout uint ByteOffset, inout float3 A, inout float B)
-{
-	float4 temp = UnpackR16G16B16A16F(StrataLoadUint2(StrataBuffer, ByteOffset));
-	A = temp.xyz;
-	B = temp.a;
-}
-
-void StrataLoad_R11G11B10F(ByteAddressBuffer StrataBuffer, inout uint ByteOffset, inout float3 A)
+void StrataLoad_R11G11B10F(ByteAddressBuffer StrataBuffer, inout FStrataAddressing StrataAddressing, inout float3 A)
 {
-	A = UnpackR11G11B10F(StrataLoadUint1(StrataBuffer, ByteOffset));
+	A = UnpackR11G11B10F(StrataLoadUint1(StrataBuffer, StrataAddressing));
 }
 
-void StrataLoad_R11_G11_B10F(ByteAddressBuffer StrataBuffer, inout uint ByteOffset, inout float A, inout float B, inout float C)
+void StrataLoad_R11_G11_B10F(ByteAddressBuffer StrataBuffer, inout FStrataAddressing StrataAddressing, inout float A, inout float B, inout float C)
 {
-	float3 temp = UnpackR11G11B10F(StrataLoadUint1(StrataBuffer, ByteOffset));
+	float3 temp = UnpackR11G11B10F(StrataLoadUint1(StrataBuffer, StrataAddressing));
 	A = temp.x;
 	B = temp.y;
 	C = temp.z;
 }
 
-void StrataLoad_ColorGamma2ToLinear(ByteAddressBuffer StrataBuffer, inout uint ByteOffset, inout float3 A)
+void StrataLoad_ColorGamma2ToLinear(ByteAddressBuffer StrataBuffer, inout FStrataAddressing StrataAddressing, inout float3 A)
 {
-	A = UnpackColorGamma2ToLinear(StrataLoadUint1(StrataBuffer, ByteOffset));
+	A = UnpackColorGamma2ToLinear(StrataLoadUint1(StrataBuffer, StrataAddressing));
 }
 
-void StrataLoad_ColorGamma2ToLinearAlphaLinear(ByteAddressBuffer StrataBuffer, inout uint ByteOffset, inout float3 RGB, inout float A)
+void StrataLoad_ColorGamma2ToLinearAlphaLinear(ByteAddressBuffer StrataBuffer, inout FStrataAddressing StrataAddressing, inout float3 RGB, inout float A)
 {
-	float4 Data = UnpackColorGamma2ToLinearAlphaLinear(StrataLoadUint1(StrataBuffer, ByteOffset));
+	float4 Data = UnpackColorGamma2ToLinearAlphaLinear(StrataLoadUint1(StrataBuffer, StrataAddressing));
 	RGB = Data.rgb;
 	A = Data.a;
 }
 
 uint PackStrataHeaderIn(uint BSDFCount, FStrataPixelHeader In);
 
-uint GetStrataPixelDataByteOffset(uint2 PixelPos, uint ViewBufferSizeX, uint StrataMaxBytesPerPixel)
-{
-	return (PixelPos.y * ViewBufferSizeX + PixelPos.x) * StrataMaxBytesPerPixel;
-}
-
 // If this is changed, please update the compiler side material size evaluation in StrataMaterial.cpp
 void PackStrataOut(
 	RWByteAddressBuffer StrataBuffer, 
-	uint ByteOffset, 
+	FStrataAddressing StrataAddressing,
 	FStrataPixelHeader StrataPixelHeader,
 	FStrataData Strata,
 	float3 V,
@@ -1363,7 +1339,7 @@ uint PackStrataHeaderIn(uint InBSDFCount, FStrataPixelHeader InHeader)
 
 // Unpack only strata header
 // If this is changed, please update the compiler side material size evaluation in StrataMaterial.cpp
-FStrataPixelHeader UnpackStrataHeaderIn(ByteAddressBuffer StrataBuffer, inout uint ByteOffset)
+FStrataPixelHeader UnpackStrataHeaderIn(ByteAddressBuffer StrataBuffer, inout FStrataAddressing StrataAddressing)
 {
 	FStrataPixelHeader Out = InitialiseStrataPixelHeader();
 
@@ -1375,12 +1351,12 @@ FStrataPixelHeader UnpackStrataHeaderIn(ByteAddressBuffer StrataBuffer, inout ui
 
 #if STRATA_INLINE_SHADING==0
 	// Only keep the offset to the shared normal memory. We are going to load them on demand.
-	Out.SharedNormalByteOffset = ByteOffset;
+	Out.SharedNormalByteOffset = StrataAddressing.CurrentAddr;
 	Out.StrataBuffer = StrataBuffer;
 	Out.SharedNormalTypes = NormalTypes;
 
 	// Now skip over the shared normal memory to be able to load the BSDF content
-	ByteOffset += NormalCount * STRATA_PACKED_NORMAL_STRIDE_BYTES;
+	StrataAddressing.CurrentAddr += NormalCount * STRATA_PACKED_NORMAL_STRIDE_BYTES;
 #endif
 
 	return Out;
@@ -1388,22 +1364,22 @@ FStrataPixelHeader UnpackStrataHeaderIn(ByteAddressBuffer StrataBuffer, inout ui
 
 // Unpack only strata header 
 // This version assume only a single slab (isotropic BSDF, with no edge color, no subsurface scattering, no haziness, and no thin film)
-FStrataPixelHeader UnpackFastPathStrataHeaderIn(ByteAddressBuffer StrataBuffer, inout uint ByteOffset)
+FStrataPixelHeader UnpackFastPathStrataHeaderIn(ByteAddressBuffer StrataBuffer, inout FStrataAddressing StrataAddressing)
 {
 	FStrataPixelHeader Out = InitialiseStrataPixelHeader();
 	
 	// Skip over the header packed data
-	ByteOffset += 4; ;
+	StrataAddressing.CurrentAddr += 4; ;
 
 	Out.BSDFCount = 1;
 #if STRATA_INLINE_SHADING==0
-	Out.SharedNormalByteOffset = ByteOffset;
+	Out.SharedNormalByteOffset = StrataAddressing.CurrentAddr;
 	Out.StrataBuffer = StrataBuffer;
 	Out.SharedNormalTypes = STRATA_BASIS_TYPE_NORMAL;
 #endif
 
 	// Skip over the shared normal memory to be able to load the BSDF content
-	ByteOffset += STRATA_PACKED_NORMAL_STRIDE_BYTES;
+	StrataAddressing.CurrentAddr += STRATA_PACKED_NORMAL_STRIDE_BYTES;
 
 	return Out;
 }
@@ -1424,7 +1400,7 @@ void SetStrataFastPathState(inout FStrataBSDF BSDF)
 // Unpack a single BSDF
 // Note: All BSDF lobes needs to be unpack in a sequential manner since each BSDF has a variable footprint.
 // If this is changed, please update the compiler side material size evaluation in StrataMaterial.cpp
-FStrataBSDF UnpackStrataBSDFIn(ByteAddressBuffer StrataBuffer, inout uint ByteOffset)
+FStrataBSDF UnpackStrataBSDFIn(ByteAddressBuffer StrataBuffer, inout FStrataAddressing StrataAddressing)
 {
 	FStrataBSDF OutBSDF = (FStrataBSDF)0;
 
@@ -1437,7 +1413,7 @@ FStrataBSDF UnpackStrataBSDFIn(ByteAddressBuffer StrataBuffer, inout uint ByteOf
 	}
 	else
 	{
-		StrataLoad_R11G11B10F(StrataBuffer, ByteOffset, OutBSDF.Weight);
+		StrataLoad_R11G11B10F(StrataBuffer, StrataAddressing, OutBSDF.Weight);
 	}
 
 
@@ -1447,9 +1423,9 @@ FStrataBSDF UnpackStrataBSDFIn(ByteAddressBuffer StrataBuffer, inout uint ByteOf
 	{
 	case STRATA_BSDF_TYPE_SLAB:
 	{
-		StrataLoad_ColorGamma2ToLinearAlphaLinear(StrataBuffer, ByteOffset, SLAB_BASECOLOR(OutBSDF), SLAB_ROUGHNESSX(OutBSDF));
+		StrataLoad_ColorGamma2ToLinearAlphaLinear(StrataBuffer, StrataAddressing, SLAB_BASECOLOR(OutBSDF), SLAB_ROUGHNESSX(OutBSDF));
 		{
-			float4 Data = UnpackRGBA8(StrataLoadUint1(StrataBuffer, ByteOffset));
+			float4 Data = UnpackRGBA8(StrataLoadUint1(StrataBuffer, StrataAddressing));
 			SLAB_SPECULAR(OutBSDF) = Data.x;
 			SLAB_METALLIC(OutBSDF) = Data.y;
 			SLAB_ROUGHNESSY(OutBSDF) = Data.z;
@@ -1460,40 +1436,40 @@ FStrataBSDF UnpackStrataBSDFIn(ByteAddressBuffer StrataBuffer, inout uint ByteOf
 		SLAB_THINFILM_THICKNESS(OutBSDF) = 0.0f;
 		if (BSDF_GETHASEDGECOLOR(OutBSDF) || BSDF_GETHASTHINFILM(OutBSDF))
 		{
-			StrataLoad_ColorGamma2ToLinearAlphaLinear(StrataBuffer, ByteOffset, SLAB_EDGECOLOR(OutBSDF), SLAB_THINFILM_THICKNESS(OutBSDF));
+			StrataLoad_ColorGamma2ToLinearAlphaLinear(StrataBuffer, StrataAddressing, SLAB_EDGECOLOR(OutBSDF), SLAB_THINFILM_THICKNESS(OutBSDF));
 		}
 
 		if (BSDF_GETHASSCATTERING(OutBSDF))
 		{
 			if (BSDF_GETHASSSSPROFILE(OutBSDF))
 			{
-				float4 Data = UnpackRGBA8(StrataLoadUint1(StrataBuffer, ByteOffset));
+				float4 Data = UnpackRGBA8(StrataLoadUint1(StrataBuffer, StrataAddressing));
 				SLAB_SSSPROFILEID(OutBSDF) = Data.r;
 				SLAB_SSSPROFILERADIUSSCALE(OutBSDF) = Data.g;
 			}
 			else
 			{
-				StrataLoad_R11G11B10F(StrataBuffer, ByteOffset, SLAB_SSSDMFP(OutBSDF));
+				StrataLoad_R11G11B10F(StrataBuffer, StrataAddressing, SLAB_SSSDMFP(OutBSDF));
 			}
 		}
 	}
 	break;
 	case STRATA_BSDF_TYPE_SHEEN:
 	{
-		StrataLoad_ColorGamma2ToLinearAlphaLinear(StrataBuffer, ByteOffset, SHEEN_BASECOLOR(OutBSDF), SHEEN_ROUGHNESS(OutBSDF));
+		StrataLoad_ColorGamma2ToLinearAlphaLinear(StrataBuffer, StrataAddressing, SHEEN_BASECOLOR(OutBSDF), SHEEN_ROUGHNESS(OutBSDF));
 	}
 	break;
 	case STRATA_BSDF_TYPE_HAIR:
 	{
-		StrataLoad_ColorGamma2ToLinearAlphaLinear(StrataBuffer, ByteOffset, HAIR_BASECOLOR(OutBSDF), HAIR_ROUGHNESS(OutBSDF));
-		StrataLoad_R11_G11_B10F(StrataBuffer, ByteOffset, HAIR_SCATTER(OutBSDF), HAIR_SPECULAR(OutBSDF), HAIR_BACKLIT(OutBSDF));
+		StrataLoad_ColorGamma2ToLinearAlphaLinear(StrataBuffer, StrataAddressing, HAIR_BASECOLOR(OutBSDF), HAIR_ROUGHNESS(OutBSDF));
+		StrataLoad_R11_G11_B10F(StrataBuffer, StrataAddressing, HAIR_SCATTER(OutBSDF), HAIR_SPECULAR(OutBSDF), HAIR_BACKLIT(OutBSDF));
 		// 8 bytes
 	}
 	break;
 	case STRATA_BSDF_TYPE_SINGLELAYERWATER:
 	{
-		StrataLoad_ColorGamma2ToLinearAlphaLinear(StrataBuffer, ByteOffset, SLW_BASECOLOR(OutBSDF), DummyFloat);
-		float4 Data = UnpackRGBA8(StrataLoadUint1(StrataBuffer, ByteOffset));
+		StrataLoad_ColorGamma2ToLinearAlphaLinear(StrataBuffer, StrataAddressing, SLW_BASECOLOR(OutBSDF), DummyFloat);
+		float4 Data = UnpackRGBA8(StrataLoadUint1(StrataBuffer, StrataAddressing));
 		SLW_METALLIC(OutBSDF) = Data.x;
 		SLW_SPECULAR(OutBSDF) = Data.y;
 		SLW_ROUGHNESS(OutBSDF) = Data.z;
diff --git a/Engine/Shaders/Private/Strata/StrataDeferredLighting.ush b/Engine/Shaders/Private/Strata/StrataDeferredLighting.ush
--- a/Engine/Shaders/Private/Strata/StrataDeferredLighting.ush
+++ b/Engine/Shaders/Private/Strata/StrataDeferredLighting.ush
@@ -155,7 +155,7 @@ float4 StrataDeferredLighting(
 	FShadowTerms ShadowTerms,
 	FRectTexture RectTexture,
 	ByteAddressBuffer StrataDataBuffer,
-	uint PixelStrataDataByteOffset)
+	FStrataAddressing StrataAddressing)
 {	
 	float3 PreintegratedIndicentIrradiance = 0;
 	if (LightData.bRectLight)
@@ -173,15 +173,15 @@ float4 StrataDeferredLighting(
 
 #if STRATA_FASTPATH
 	uint BSDFIndex = 0;
-	FStrataPixelHeader StrataPixelHeader = UnpackFastPathStrataHeaderIn(StrataDataBuffer, PixelStrataDataByteOffset);
+	FStrataPixelHeader StrataPixelHeader = UnpackFastPathStrataHeaderIn(StrataDataBuffer, StrataAddressing);
 #else
-	FStrataPixelHeader StrataPixelHeader = UnpackStrataHeaderIn(StrataDataBuffer, PixelStrataDataByteOffset);
+	FStrataPixelHeader StrataPixelHeader = UnpackStrataHeaderIn(StrataDataBuffer, StrataAddressing);
 	LOOP
 	for (uint BSDFIndex = 0; BSDFIndex < StrataPixelHeader.BSDFCount; ++BSDFIndex)
 #endif
 	{
 		// Unpack BSDF data
-		FStrataBSDF BSDF = UnpackStrataBSDFIn(StrataDataBuffer, PixelStrataDataByteOffset);
+		FStrataBSDF BSDF = UnpackStrataBSDFIn(StrataDataBuffer, StrataAddressing);
 #if STRATA_FASTPATH
 		SetStrataFastPathState(BSDF);
 #endif
diff --git a/Engine/Shaders/Private/Strata/StrataEnvironmentLighting.ush b/Engine/Shaders/Private/Strata/StrataEnvironmentLighting.ush
--- a/Engine/Shaders/Private/Strata/StrataEnvironmentLighting.ush
+++ b/Engine/Shaders/Private/Strata/StrataEnvironmentLighting.ush
@@ -18,7 +18,7 @@ float4 StrataEnvironmentLighting(
 	float AmbientOcclusion,
 	float TopLayerSpecularContributionFactor,
 	ByteAddressBuffer StrataDataBuffer,
-	uint PixelStrataDataByteOffset,
+	FStrataAddressing StrataAddressing,
 	inout float SSRReductionFactor)
 {
 	FLightAccumulator Out = (FLightAccumulator)0;
@@ -35,13 +35,13 @@ float4 StrataEnvironmentLighting(
 
 #if STRATA_FASTPATH
 	uint BSDFIndex = 0;
-	FStrataPixelHeader StrataPixelHeader = UnpackFastPathStrataHeaderIn(StrataDataBuffer, PixelStrataDataByteOffset);
+	FStrataPixelHeader StrataPixelHeader = UnpackFastPathStrataHeaderIn(StrataDataBuffer, StrataAddressing);
 #else
-	FStrataPixelHeader StrataPixelHeader = UnpackStrataHeaderIn(StrataDataBuffer, PixelStrataDataByteOffset);
+	FStrataPixelHeader StrataPixelHeader = UnpackStrataHeaderIn(StrataDataBuffer, StrataAddressing);
 	for (uint BSDFIndex = 0; BSDFIndex < StrataPixelHeader.BSDFCount; ++BSDFIndex)
 #endif
 	{
-		FStrataBSDF BSDF = UnpackStrataBSDFIn(StrataDataBuffer, PixelStrataDataByteOffset);
+		FStrataBSDF BSDF = UnpackStrataBSDFIn(StrataDataBuffer, StrataAddressing);
 #if STRATA_FASTPATH
 		SetStrataFastPathState(BSDF);
 #endif
diff --git a/Engine/Shaders/Private/Strata/StrataMaterialClassification.usf b/Engine/Shaders/Private/Strata/StrataMaterialClassification.usf
--- a/Engine/Shaders/Private/Strata/StrataMaterialClassification.usf
+++ b/Engine/Shaders/Private/Strata/StrataMaterialClassification.usf
@@ -156,12 +156,12 @@ void MainPS(
 		float4 WorldPos = mul(float4(ScreenPosition * Depth, Depth, 1), View.ScreenToWorld);
 		const float3 V = normalize(View.WorldCameraOrigin - WorldPos.xyz);
 
-		uint PixelStrataDataByteOffset = GetStrataPixelDataByteOffset(PixelPos, View.BufferSizeAndInvSize.x, Strata.MaxBytesPerPixel);
-		FStrataPixelHeader Header = UnpackStrataHeaderIn(Strata.MaterialLobesBuffer, PixelStrataDataByteOffset);
+		FStrataAddressing StrataAddressing = GetStrataPixelDataByteOffset(PixelPos, View.BufferSizeAndInvSize.x, Strata.MaxBytesPerPixel);
+		FStrataPixelHeader Header = UnpackStrataHeaderIn(Strata.MaterialLobesBuffer, StrataAddressing);
 		ClassificationData.bIsSimple = Header.BSDFCount == 1;
 		for (uint BSDFIndex = 0; BSDFIndex < Header.BSDFCount; ++BSDFIndex)
 		{
-			const FStrataBSDF BSDF = UnpackStrataBSDFIn(Strata.MaterialLobesBuffer, PixelStrataDataByteOffset);
+			const FStrataBSDF BSDF = UnpackStrataBSDFIn(Strata.MaterialLobesBuffer, StrataAddressing);
 			ComputePostProcessBSDFData(Header, BSDF, V, ClassificationData, SSSData, TopLayerData);
 		}
 
diff --git a/Engine/Shaders/Private/Strata/StrataVisualize.usf b/Engine/Shaders/Private/Strata/StrataVisualize.usf
--- a/Engine/Shaders/Private/Strata/StrataVisualize.usf
+++ b/Engine/Shaders/Private/Strata/StrataVisualize.usf
@@ -84,8 +84,7 @@ void VisualizeMaterialPS(
 	OutColor = 0;
 	if (bStrataMaterial)
 	{
-		uint PixelStrataDataByteOffset = GetStrataPixelDataByteOffset(MaterialPos, View.BufferSizeAndInvSize.x, Strata.MaxBytesPerPixel);
-		FStrataPixelHeader Header = UnpackStrataHeaderIn(Strata.MaterialLobesBuffer, PixelStrataDataByteOffset);
+		FStrataPixelHeader Header = UnpackStrataHeaderIn(Strata.MaterialLobesBuffer, GetStrataPixelDataByteOffset(MaterialPos, View.BufferSizeAndInvSize.x, Strata.MaxBytesPerPixel));
 
 		OutColor.w = 0.05f;
 		if (Header.BSDFCount == 1)
@@ -118,15 +117,13 @@ void VisualizeMaterialPS(
 	bool bPlotValid = false;
 	if (bStrataMaterial)
 	{
-		uint PixelStrataDataByteOffset = GetStrataPixelDataByteOffset(MaterialPos, View.BufferSizeAndInvSize.x, Strata.MaxBytesPerPixel);
-		const uint PixelStrataDataByteOffsetStart = PixelStrataDataByteOffset;
+		FStrataAddressing StrataAddressing = GetStrataPixelDataByteOffset(PixelPos, View.BufferSizeAndInvSize.x, Strata.MaxBytesPerPixel);
 
-		FStrataPixelHeader Header = UnpackStrataHeaderIn(Strata.MaterialLobesBuffer, PixelStrataDataByteOffset);
-		const uint PixelStrataDataByteOffsetPostHeader = PixelStrataDataByteOffset;
+		FStrataPixelHeader Header = UnpackStrataHeaderIn(Strata.MaterialLobesBuffer, StrataAddressing);
 
 		for (uint BSDFIndex = 0; BSDFIndex < Header.BSDFCount; ++BSDFIndex)
 		{
-			const FStrataBSDF BSDF = UnpackStrataBSDFIn(Strata.MaterialLobesBuffer, PixelStrataDataByteOffset);
+			const FStrataBSDF BSDF = UnpackStrataBSDFIn(Strata.MaterialLobesBuffer, StrataAddressing);
 
 			// Referential
 			if (all(PixelPos == MaterialPos))
@@ -158,11 +155,11 @@ void VisualizeMaterialPS(
 		uint2 RectMin = uint2(MaterialPos + InspectSquare);
 		uint2 RectMax = RectMin;
 
-		uint PixelStrataDataByteOffset = GetStrataPixelDataByteOffset(MaterialPos, View.BufferSizeAndInvSize.x, Strata.MaxBytesPerPixel);
-		const uint PixelStrataDataByteOffsetStart = PixelStrataDataByteOffset;
+		FStrataAddressing StrataAddressing = GetStrataPixelDataByteOffset(MaterialPos, View.BufferSizeAndInvSize.x, Strata.MaxBytesPerPixel);
+		const uint PixelStrataDataByteOffsetStart = StrataAddressing.ReadBytes;
 
-		FStrataPixelHeader Header = UnpackStrataHeaderIn(Strata.MaterialLobesBuffer, PixelStrataDataByteOffset);
-		const uint PixelStrataDataByteOffsetPostHeader = PixelStrataDataByteOffset;
+		FStrataPixelHeader Header = UnpackStrataHeaderIn(Strata.MaterialLobesBuffer, StrataAddressing);
+		const uint PixelStrataDataByteOffsetPostHeader = StrataAddressing.ReadBytes;
 
 		float3 BSDFFontColor = float3(0.25, 1.0, 0.25);
 		float3 BSDFFontColorGrey = float3(0.25, 0.25, 0.25);
@@ -210,7 +207,7 @@ void VisualizeMaterialPS(
 		LOOP
 		for (uint BSDFIndex = 0; BSDFIndex < Header.BSDFCount; ++BSDFIndex)
 		{
-			const FStrataBSDF BSDF = UnpackStrataBSDFIn(Strata.MaterialLobesBuffer, PixelStrataDataByteOffset);
+			const FStrataBSDF BSDF = UnpackStrataBSDFIn(Strata.MaterialLobesBuffer, StrataAddressing);
 
 			TopLeft.x = TopLeftInit.x;
 
@@ -427,7 +424,7 @@ void VisualizeMaterialPS(
 			HORIZONTALSEPARATION;
 		}
 
-		const uint PixelStrataDataByteOffsetPostBSDF = PixelStrataDataByteOffset;
+		const uint PixelStrataDataByteOffsetPostBSDF = StrataAddressing.ReadBytes;
 
 #if PERMUTATION_BSDF_PASS==0
 		PRINTINFOCHAR(_B_); PRINTINFOCHAR(_O_); PRINTINFOCHAR(_T_); PRINTINFOCHAR(_T_); PRINTINFOCHAR(_O_); PRINTINFOCHAR(_M_);
