Implement Cloth shading model in path tracer

diff --git a/Engine/Shaders/Private/PathTracing/Material/PathTracingCloth.ush b/Engine/Shaders/Private/PathTracing/Material/PathTracingCloth.ush
new file mode 100644
--- /dev/null
+++ b/Engine/Shaders/Private/PathTracing/Material/PathTracingCloth.ush
@@ -0,0 +1,155 @@
+// Copyright Epic Games, Inc. All Rights Reserved.
+
+/*=============================================================================================
+	PathTracingCloth.ush: Path tracing BRDF model for cloth materials
+===============================================================================================*/
+
+#pragma once
+
+#include "PathTracingFresnel.ush"
+
+float ClothLobeSelectionPdf(float3 V, FMaterialClosestHitPayload Payload)
+{
+	// TODO: This could be improved if we had a directional albedo tabulation, but this formulation will at least prevent
+	// under-sampling when one lobe is black (like with metals)
+
+	const float3 Diff = Payload.DiffuseColor;
+	const float3 Spec = Payload.SpecularColor;
+	const float3 Fuzz = Payload.CustomData.xyz;
+	const float Cloth = Payload.CustomData.w;
+
+	return LobeSelectionProb(Diff + Cloth * Fuzz, (1.0 - Cloth) * Spec);
+}
+
+void Cloth_EvalMaterial(
+	float3 IncomingDirection,
+	float3 OutgoingDirection,
+	FMaterialClosestHitPayload Payload,
+	out float3 OutWeight,
+	out float OutPdf
+)
+{
+	OutWeight = 0;
+	OutPdf = 0;
+
+	const float3 V_World = -IncomingDirection;
+	const float3 L_World = OutgoingDirection;
+	const float3 N_World = GetWorldNormal(Payload);
+
+	float2 Alpha = 0;
+	const float3x3 Basis = GetGGXBasis(Payload.Roughness, Payload.Anisotropy, Payload.WorldNormal, Payload.WorldTangent, Alpha);
+
+	// move vectors into right shading frame
+	const float3 V = mul(Basis, V_World);
+	const float3 L = mul(Basis, L_World);
+	const float3 H = normalize(V + L);
+
+	const float NoL = saturate(L.z);
+	const float NoV = saturate(V.z);
+	const float VoH = saturate(dot(V, H));
+	const float NoH = saturate(H.z);
+
+	// Probability of picking diffuse lobe vs. specular lobe
+	const float LobeSelectionPdf = ClothLobeSelectionPdf(V_World, Payload);
+
+	const float Cloth = Payload.CustomData.w;
+
+	// Diffuse and Cloth Lobe (see ClothBxDF)
+	const float3 FuzzColor = Payload.CustomData.xyz;
+	const float DCloth = D_InvGGX(Pow4(Payload.Roughness), NoH);
+	const float3 FCloth = F_Schlick(FuzzColor, VoH);
+	const float3 ClothWeight = NoL > 0 && NoV > 0 ? (PI * Cloth * DCloth * Vis_Cloth(NoV, NoL)) * FCloth : 0.0;
+
+	AddLobeWithMIS(OutWeight, OutPdf, Payload.DiffuseColor + ClothWeight, NoL / PI, LobeSelectionPdf);
+
+	// Specular lobe
+	const float2 GGXResult = GGXEvalReflection(L, V, H, Alpha);
+	const float3 F = F_Schlick(Payload.SpecularColor, VoH);
+	const float3 SpecWeight = (1 - Cloth) * F * GGXResult.x;
+	const float SpecPdf = GGXResult.y;
+
+	AddLobeWithMIS(OutWeight, OutPdf, SpecWeight, SpecPdf, 1 - LobeSelectionPdf);
+}
+
+void Cloth_SampleMaterial(
+	float3 RayDirection,
+	FMaterialClosestHitPayload Payload,
+	float4 RandSample,
+	out float3 OutDirection,
+	out float3 OutWeight,
+	out float OutPdf
+)
+{
+	OutDirection = 0;
+	OutWeight = 0;
+	OutPdf = 0;
+
+	float3 N = GetWorldNormal(Payload);
+	float3 V = -RayDirection;
+
+	const float LobeSelectionPdf = ClothLobeSelectionPdf(V, Payload);
+
+	float2 Alpha;
+	float3x3 Basis = GetGGXBasis(Payload.Roughness, Payload.Anisotropy, Payload.WorldNormal, Payload.WorldTangent, Alpha);
+
+	V = mul(Basis, V);
+
+	// Randomly choose to sample diffuse, specular or cloth lobe
+	float3 L = 0, H = 0;
+	if (RandSample.x < LobeSelectionPdf)
+	{
+		RandSample.x /= LobeSelectionPdf;
+
+		// Lambert
+#if 1
+		float4 Result = CosineSampleHemisphere(RandSample.xy);
+#else
+		float4 Result = CosineSampleHemisphereConcentric(RandSample.xy);
+#endif
+
+		L = Result.xyz;
+		H = normalize(L + V);
+	}
+	else
+	{
+		RandSample.x -= LobeSelectionPdf;
+		RandSample.x /= 1 - LobeSelectionPdf;
+
+
+		// #dxr_todo: Evaluate UniformSampleDiskConcentric
+		H = ImportanceSampleVisibleGGX_aniso(UniformSampleDisk(RandSample.xy), Alpha, V).xyz;
+
+		L = reflect(-V, H);
+		if (L.z <= 0) {
+			// invalid output direction, exit early
+			return;
+		}
+	}
+
+	// With a valid direction in hand -- now evaluate the BxDF (taking advantage of already computed terms)
+
+	const float NoL = saturate(L.z);
+	const float NoV = saturate(V.z);
+	const float VoH = saturate(dot(V, H));
+	const float NoH = saturate(H.z);
+
+	const float Cloth = Payload.CustomData.w;
+
+	// Diffuse and Cloth Lobe (see ClothBxDF)
+	const float3 FuzzColor = Payload.CustomData.xyz;
+	const float DCloth = D_InvGGX(Pow4(Payload.Roughness), NoH);
+	const float3 FCloth = F_Schlick(FuzzColor, VoH);
+	const float3 ClothWeight = NoL > 0 && NoV > 0 ? (PI * Cloth * DCloth * Vis_Cloth(NoV, NoL)) * FCloth : 0.0;
+
+	AddLobeWithMIS(OutWeight, OutPdf, Payload.DiffuseColor + ClothWeight, NoL / PI, LobeSelectionPdf);
+
+	// Specular lobe
+	const float2 GGXResult = GGXEvalReflection(L, V, H, Alpha);
+	const float3 F = F_Schlick(Payload.SpecularColor, VoH);
+	const float3 SpecWeight = (1 - Cloth) * F * GGXResult.x;
+	const float SpecPdf = GGXResult.y;
+	AddLobeWithMIS(OutWeight, OutPdf, SpecWeight, SpecPdf, 1 - LobeSelectionPdf);
+
+	// transform to world space
+	OutDirection = normalize(mul(L, Basis));
+}
diff --git a/Engine/Shaders/Private/PathTracing/Material/PathTracingMaterialSampling.ush b/Engine/Shaders/Private/PathTracing/Material/PathTracingMaterialSampling.ush
--- a/Engine/Shaders/Private/PathTracing/Material/PathTracingMaterialSampling.ush
+++ b/Engine/Shaders/Private/PathTracing/Material/PathTracingMaterialSampling.ush
@@ -18,6 +18,7 @@ PathTracingMaterialSampling.usf: Material sampling functions
 #include "PathTracingDefaultLit.ush"
 #include "PathTracingClearCoat.ush"
 #include "PathTracingTwoSidedFoliage.ush"
+#include "PathTracingCloth.ush"
 
 #ifndef FORCE_ALL_MATERIALS_DIFFUSE
 	#define FORCE_ALL_MATERIALS_DIFFUSE 0
@@ -45,6 +46,7 @@ bool IsNeeValidMaterial(FMaterialClosestHitPayload Payload)
 		case SHADINGMODELID_CLEAR_COAT:			
 		case SHADINGMODELID_TWOSIDED_FOLIAGE:	
 		case SHADINGMODELID_THIN_TRANSLUCENT:
+		case SHADINGMODELID_CLOTH:
 		default:
 			return true;
 	}
@@ -133,6 +135,9 @@ float3 EstimateMaterialAlbedo(FMaterialClosestHitPayload Payload)
 		case SHADINGMODELID_CLEAR_COAT:
 			// FIXME: could make this more accurate by estimating the directional albedo
 			Albedo = BrdfWeight * (Payload.DiffuseColor + Payload.SpecularColor);
+		case SHADINGMODELID_CLOTH:
+			Albedo = BrdfWeight * (Payload.DiffuseColor + lerp(Payload.SpecularColor, Payload.CustomData.xyz, Payload.CustomData.w));
+			break;
 		default:
 			Albedo = BrdfWeight * Payload.BaseColor;
 			break;
@@ -163,6 +168,7 @@ void AdjustPathAndMaterialRoughness(inout FMaterialClosestHitPayload Payload, in
 		case SHADINGMODELID_TWOSIDED_FOLIAGE:
 		case SHADINGMODELID_DEFAULT_LIT:
 		case SHADINGMODELID_THIN_TRANSLUCENT:
+		case SHADINGMODELID_CLOTH:
 			Payload.Roughness = max(Payload.Roughness, PathRoughness);
 			break;
 		default:
@@ -278,6 +284,9 @@ void SampleMaterial(
 		case SHADINGMODELID_DEFAULT_LIT:
 			DefaultLit_SampleMaterial(RayDirection, Payload, RandSample, OutDirection, OutWeight, OutPdf);
 			break;
+		case SHADINGMODELID_CLOTH:
+			Cloth_SampleMaterial(RayDirection, Payload, RandSample, OutDirection, OutWeight, OutPdf);
+			break;
 		case SHADINGMODELID_CLEAR_COAT:
 			ClearCoat_SampleMaterial(RayDirection, Payload, RandSample, OutDirection, OutWeight, OutPdf);
 			break;
@@ -367,6 +376,9 @@ void EvalMaterial(
 		case SHADINGMODELID_DEFAULT_LIT:
 			DefaultLit_EvalMaterial(IncomingDirection, OutgoingDirection, Payload, OutWeight, OutPdf);
 			break;
+		case SHADINGMODELID_CLOTH:
+			Cloth_EvalMaterial(IncomingDirection, OutgoingDirection, Payload, OutWeight, OutPdf);
+			break;
 		case SHADINGMODELID_CLEAR_COAT:
 			ClearCoat_EvalMaterial(IncomingDirection, OutgoingDirection, Payload, OutWeight, OutPdf);
 			break;
