Strata env poass clean up and option:
- UpsampleDFAO has already been moved to the EvaluateEnv function and shared between BSDFs
- Sky pixels are correctly excluded
- EXplicit UnpackFastPathStrataBSDFIn function
- Slab fastpath loading is common code

Fixed starta debug view where normal bytes where no longer taken into account after a recent change.

diff --git a/Engine/Shaders/Private/ReflectionEnvironmentPixelShader.usf b/Engine/Shaders/Private/ReflectionEnvironmentPixelShader.usf
--- a/Engine/Shaders/Private/ReflectionEnvironmentPixelShader.usf
+++ b/Engine/Shaders/Private/ReflectionEnvironmentPixelShader.usf
@@ -341,35 +341,25 @@ void ReflectionEnvironmentSkyLighting(
 
 #if ENABLE_DYNAMIC_SKY_LIGHT
 #if ENABLE_SKY_LIGHT
+	// STRATA_TODO handle the ENABLE_SKY_LIGHT && ALLOW_STATIC_LIGHTING case 
 
 	FStrataClassification StrataClassification = StrataUnpackClassificationData(Strata.ClassificationTexture.Load(uint3(PixelPos, 0)));
-
-	// STRATA_TODO handle the ENABLE_SKY_LIGHT && ALLOW_STATIC_LIGHTING case 
 	BRANCH
-	if (StrataClassification.bIsStrataMaterial)
+	if (StrataClassification.bIsStrataMaterial) // This test is also enough to exclude sky pixels
 	{
-		const FStrataTopLayerData TopLayerData = StrataUnpackTopLayerData(Strata.TopLayerTexture.Load(uint3(PixelPos, 0)));
-
 		float DeviceZ = SampleDeviceZFromSceneTextures(BufferUV);
 		float SceneDepth = ConvertFromDeviceZ(DeviceZ);
 
+		// STRATA_TODO: when GCN TA command FIFO is not longer the bottleneck, try to compute TopLayerData.EnvBRDF in StrataEnvironmentLighting to avoid reading that TopLayerData texture.
+		const FStrataTopLayerData TopLayerData = StrataUnpackTopLayerData(Strata.TopLayerTexture.Load(uint3(PixelPos, 0)));
+
 		float3 WorldPosition = mul(float4(ScreenPosition * SceneDepth, SceneDepth, 1), View.ScreenToWorld).xyz;
 		float3 CameraToPixel = normalize(WorldPosition - View.WorldCameraOrigin);
-		float3 N = TopLayerData.Normal;
 		float3 V = -CameraToPixel;
-		float NoV = saturate(dot(N, V));
 
 		// Sample the ambient occlusion that is dynamically generated every frame.
 		float AmbientOcclusion = AmbientOcclusionTexture.SampleLevel(AmbientOcclusionSampler, BufferUV, 0).r;
 
-		// STRATA_TODO: the top layer should handle the bent normal...
-		float3 BentNormal = N; 
-#if APPLY_SKY_SHADOWING
-		{
-			BentNormal = UpsampleDFAO(BufferUV, SceneDepth, N);
-		}
-#endif
-
 		const float4 SSR = Texture2DSample(ScreenSpaceReflectionsTexture, ScreenSpaceReflectionsSampler, BufferUV);
 
 		// SSR is computed for the top level by averaging BSDF components.
diff --git a/Engine/Shaders/Private/Strata/Strata.ush b/Engine/Shaders/Private/Strata/Strata.ush
--- a/Engine/Shaders/Private/Strata/Strata.ush
+++ b/Engine/Shaders/Private/Strata/Strata.ush
@@ -1357,6 +1357,9 @@ FStrataPixelHeader UnpackStrataHeaderIn(ByteAddressBuffer StrataBuffer, inout FS
 
 	// Now skip over the shared normal memory to be able to load the BSDF content
 	StrataAddressing.CurrentAddr += NormalCount * STRATA_PACKED_NORMAL_STRIDE_BYTES;
+
+	// And account for the space used by normal for the debug information
+	StrataAddressing.ReadBytes += NormalCount * STRATA_PACKED_NORMAL_STRIDE_BYTES;
 #endif
 
 	return Out;
@@ -1384,17 +1387,45 @@ FStrataPixelHeader UnpackFastPathStrataHeaderIn(ByteAddressBuffer StrataBuffer,
 	return Out;
 }
 
-// Force a BSDF state to a single slabe with isotropic BSDF, with no subsurface scattering, and no haziness
-void SetStrataFastPathState(inout FStrataBSDF BSDF)
-{
-	BSDF_SETTYPE(BSDF, STRATA_BSDF_TYPE_SLAB);
-	BSDF_SETHASEDGECOLOR(BSDF, 0);
-	BSDF_SETANISOTROPY(BSDF, 0);
-	BSDF_SETISTOPLAYER(BSDF, 1);
-	BSDF_SETHASSCATTERING(BSDF, 0);
-	BSDF_SETHASSSSPROFILE(BSDF, 0);
-	BSDF_SETHASHAZINESS(BSDF, 0);
-	BSDF_SETHASTHINFILM(BSDF, 0);
+void UnpackFastPathSlabBSDF(inout FStrataBSDF OutBSDF, uint2 PackedData01)
+{
+	float4 Data0 = UnpackColorGamma2ToLinearAlphaLinear(PackedData01.x);
+	SLAB_BASECOLOR(OutBSDF)		= Data0.rgb;
+	SLAB_ROUGHNESSX(OutBSDF)	= Data0.a;
+
+	float4 Data1 = UnpackRGBA8(PackedData01.y);
+	SLAB_SPECULAR(OutBSDF)		= Data1.x;
+	SLAB_METALLIC(OutBSDF)		= Data1.y;
+	SLAB_ROUGHNESSY(OutBSDF)	= Data1.z;
+	SLAB_HAZINESS(OutBSDF)		= Data1.w;
+}
+
+FStrataBSDF UnpackFastPathStrataBSDFIn(ByteAddressBuffer StrataBuffer, inout FStrataAddressing StrataAddressing)
+{
+	FStrataBSDF OutBSDF = (FStrataBSDF)0;
+
+	OutBSDF.State = StrataLoadUint1(StrataBuffer, StrataAddressing);
+
+	// In the fast path the single slab BSDF is always on top so the weight must be grey scale.
+	OutBSDF.Weight = Unpack10F(BSDF_GETWEIGHT10F(OutBSDF));
+
+	// Now unpack the simple Slab BSDF
+	uint DataX = StrataLoadUint1(StrataBuffer, StrataAddressing);
+	uint DataY = StrataLoadUint1(StrataBuffer, StrataAddressing);
+	UnpackFastPathSlabBSDF(OutBSDF, uint2(DataX, DataY));
+
+	// Force a BSDF state to a single slab with isotropic specular, no subsurface scattering or any other options.
+	// This is to help the compiler understand disabled path.
+	BSDF_SETTYPE(OutBSDF, STRATA_BSDF_TYPE_SLAB);
+	BSDF_SETHASEDGECOLOR(OutBSDF, 0);
+	BSDF_SETANISOTROPY(OutBSDF, 0);
+	BSDF_SETISTOPLAYER(OutBSDF, 1);
+	BSDF_SETHASSCATTERING(OutBSDF, 0);
+	BSDF_SETHASSSSPROFILE(OutBSDF, 0);
+	BSDF_SETHASHAZINESS(OutBSDF, 0);
+	BSDF_SETHASTHINFILM(OutBSDF, 0);
+
+	return OutBSDF;
 }
 
 // Unpack a single BSDF
@@ -1423,14 +1454,9 @@ FStrataBSDF UnpackStrataBSDFIn(ByteAddressBuffer StrataBuffer, inout FStrataAddr
 	{
 	case STRATA_BSDF_TYPE_SLAB:
 	{
-		StrataLoad_ColorGamma2ToLinearAlphaLinear(StrataBuffer, StrataAddressing, SLAB_BASECOLOR(OutBSDF), SLAB_ROUGHNESSX(OutBSDF));
-		{
-			float4 Data = UnpackRGBA8(StrataLoadUint1(StrataBuffer, StrataAddressing));
-			SLAB_SPECULAR(OutBSDF) = Data.x;
-			SLAB_METALLIC(OutBSDF) = Data.y;
-			SLAB_ROUGHNESSY(OutBSDF) = Data.z;
-			SLAB_HAZINESS(OutBSDF) = Data.w;
-		}
+		uint DataX = StrataLoadUint1(StrataBuffer, StrataAddressing);
+		uint DataY = StrataLoadUint1(StrataBuffer, StrataAddressing);
+		UnpackFastPathSlabBSDF(OutBSDF, uint2(DataX, DataY));
 
 		SLAB_EDGECOLOR(OutBSDF) = 1.0f;
 		SLAB_THINFILM_THICKNESS(OutBSDF) = 0.0f;
diff --git a/Engine/Shaders/Private/Strata/StrataDeferredLighting.ush b/Engine/Shaders/Private/Strata/StrataDeferredLighting.ush
--- a/Engine/Shaders/Private/Strata/StrataDeferredLighting.ush
+++ b/Engine/Shaders/Private/Strata/StrataDeferredLighting.ush
@@ -181,9 +181,10 @@ float4 StrataDeferredLighting(
 #endif
 	{
 		// Unpack BSDF data
-		FStrataBSDF BSDF = UnpackStrataBSDFIn(StrataDataBuffer, StrataAddressing);
 #if STRATA_FASTPATH
-		SetStrataFastPathState(BSDF);
+		FStrataBSDF BSDF = UnpackFastPathStrataBSDFIn(StrataDataBuffer, StrataAddressing);
+#else
+		FStrataBSDF BSDF = UnpackStrataBSDFIn(StrataDataBuffer, StrataAddressing);
 #endif
 		FStrataBSDFContext BSDFContext = StrataCreateBSDFContext(StrataPixelHeader, BSDF, V, L);
 
diff --git a/Engine/Shaders/Private/Strata/StrataEnvironmentLighting.ush b/Engine/Shaders/Private/Strata/StrataEnvironmentLighting.ush
--- a/Engine/Shaders/Private/Strata/StrataEnvironmentLighting.ush
+++ b/Engine/Shaders/Private/Strata/StrataEnvironmentLighting.ush
@@ -41,9 +41,10 @@ float4 StrataEnvironmentLighting(
 	for (uint BSDFIndex = 0; BSDFIndex < StrataPixelHeader.BSDFCount; ++BSDFIndex)
 #endif
 	{
-		FStrataBSDF BSDF = UnpackStrataBSDFIn(StrataDataBuffer, StrataAddressing);
 #if STRATA_FASTPATH
-		SetStrataFastPathState(BSDF);
+		FStrataBSDF BSDF = UnpackFastPathStrataBSDFIn(StrataDataBuffer, StrataAddressing);
+#else
+		FStrataBSDF BSDF = UnpackStrataBSDFIn(StrataDataBuffer, StrataAddressing);
 #endif
 		const float3 BSDFThroughput = BSDF.Weight;
 
