Fixed Strata horizontal mixing weight to stay in safe [0,1] range.

diff --git a/Engine/Shaders/Private/Strata/Strata.ush b/Engine/Shaders/Private/Strata/Strata.ush
--- a/Engine/Shaders/Private/Strata/Strata.ush
+++ b/Engine/Shaders/Private/Strata/Strata.ush
@@ -659,10 +659,10 @@ FStrataData StrataMultiply(FStrataData A, float Weight)
 
 FStrataData StrataHorizontalMixing(FStrataData Foreground, FStrataData Background, float Mix)
 {
-	const float ForgroundMixFactor = Mix;
-	const float BackgroundMixFactor = 1.0 - Mix;
+	const float ForegroundMixFactor = saturate(Mix);
+	const float BackgroundMixFactor = 1.0 - ForegroundMixFactor;
 
-	FStrataData A = StrataMultiply(Foreground, ForgroundMixFactor);
+	FStrataData A = StrataMultiply(Foreground, ForegroundMixFactor);
 	FStrataData B = StrataMultiply(Background, BackgroundMixFactor);
 	return StrataAdd(A, B);
 }
