Fixed Reverb decal shader compile errors

diff --git a/Engine/Shaders/Private/DeferredDecal.usf b/Engine/Shaders/Private/DeferredDecal.usf
--- a/Engine/Shaders/Private/DeferredDecal.usf
+++ b/Engine/Shaders/Private/DeferredDecal.usf
@@ -323,6 +323,8 @@ void FPixelShaderInOut_MainPS(inout FPixelShaderIn In, inout FPixelShaderOut Out
 	GBufferData.ShadingModelID = SHADINGMODELID_DEFAULT_LIT;
 	GBufferData.SelectiveOutputMask = 0;
 	GBufferData.PerObjectGBufferData = 1; 
+	GBufferData.DiffuseIndirectSampleOcclusion = 0;
+	GBufferData.Velocity = 0;
 
 	DecalCommonOutput(In, Out, Color, Opacity, GBufferData);
 
diff --git a/Engine/Shaders/Private/MeshDecals.usf b/Engine/Shaders/Private/MeshDecals.usf
--- a/Engine/Shaders/Private/MeshDecals.usf
+++ b/Engine/Shaders/Private/MeshDecals.usf
@@ -223,6 +223,8 @@ void FPixelShaderInOut_MainPS(
 	GBufferData.ShadingModelID = SHADINGMODELID_DEFAULT_LIT;
 	GBufferData.SelectiveOutputMask = 0;
 	GBufferData.PerObjectGBufferData = 1;
+	GBufferData.DiffuseIndirectSampleOcclusion = 0;
+	GBufferData.Velocity = 0;
 
 	DecalCommonOutput(In, Out, Color, Opacity, GBufferData);
 
