Fixes history drift when using r.TemporalAA.R11G11B10History=1

diff --git a/Engine/Shaders/Private/TemporalAA/TAACommon.ush b/Engine/Shaders/Private/TemporalAA/TAACommon.ush
--- a/Engine/Shaders/Private/TemporalAA/TAACommon.ush
+++ b/Engine/Shaders/Private/TemporalAA/TAACommon.ush
@@ -442,3 +442,33 @@ taa_half UncompressParallaxFactor(uint CompressedParallaxFactor)
 {
 	return taa_half(CompressedParallaxFactor & SUB_PIXEL_MASK) * taa_half(MAX_PARALLAX_FACTOR / float(SUB_PIXEL_MASK)) + taa_half(1.0);
 }
+
+taa_half3 QuantizeForFloatRenderTarget(taa_half3 Color, taa_half E, float3 QuantizationError)
+{
+	taa_half3 Error = Color * taa_half3(QuantizationError);
+
+	#if 0
+		// NOP
+	#elif CONFIG_COMPILE_FP16 && 0
+	{
+		Error.x = asfloat16_t(asuint16_t(Error.x) & uint16_t(~0x03FF));
+		Error.y = asfloat16_t(asuint16_t(Error.y) & uint16_t(~0x03FF));
+		Error.z = asfloat16_t(asuint16_t(Error.z) & uint16_t(~0x03FF));
+	}
+	#elif CONFIG_COMPILE_FP16
+	{
+		// shader compilers haven't exposed asuint16_t() and asfloat16_t() :'(
+		Error.x = taa_half(asfloat(asuint(float(Error.x)) & ~0x007FFFFF));
+		Error.y = taa_half(asfloat(asuint(float(Error.y)) & ~0x007FFFFF));
+		Error.z = taa_half(asfloat(asuint(float(Error.z)) & ~0x007FFFFF));
+	}
+	#else
+	{
+		Error.x = asfloat(asuint(Error.x) & ~0x007FFFFF);
+		Error.y = asfloat(asuint(Error.y) & ~0x007FFFFF);
+		Error.z = asfloat(asuint(Error.z) & ~0x007FFFFF);
+	}
+	#endif
+	
+	return Color + Error * E;
+}
diff --git a/Engine/Shaders/Private/TemporalAA/TAADecimateHistory.usf b/Engine/Shaders/Private/TemporalAA/TAADecimateHistory.usf
--- a/Engine/Shaders/Private/TemporalAA/TAADecimateHistory.usf
+++ b/Engine/Shaders/Private/TemporalAA/TAADecimateHistory.usf
@@ -176,7 +176,7 @@ void MainCS(
 		uint2 Random = Rand3DPCG16(int3(InputPixelPos - InputInfo_ViewportMin, View.StateFrameIndexMod8)).xy;
 		float2 E = Hammersley16(0, 1, Random);
 
-		DecimatedColor += DecimatedColor * (E.x * OutputQuantizationError);
+		DecimatedColor = QuantizeForFloatRenderTarget(DecimatedColor, E.x, OutputQuantizationError);
 	}
 	#endif
 
diff --git a/Engine/Shaders/Private/TemporalAA/TAAFilterFrequencies.usf b/Engine/Shaders/Private/TemporalAA/TAAFilterFrequencies.usf
--- a/Engine/Shaders/Private/TemporalAA/TAAFilterFrequencies.usf
+++ b/Engine/Shaders/Private/TemporalAA/TAAFilterFrequencies.usf
@@ -284,9 +284,9 @@ void MainCS(
 	{
 		uint2 Random = Rand3DPCG16(int3(OutputPixelPos, View.StateFrameIndexMod8)).xy;
 		float2 E = Hammersley16(0, 1, Random);
-
-		FilteredInput += FilteredInput * (E.x * OutputQuantizationError);
-		FilteredHistory += FilteredHistory * (E.x * OutputQuantizationError);
+		
+		FilteredInput = QuantizeForFloatRenderTarget(FilteredInput, E.x, OutputQuantizationError);
+		FilteredHistory = QuantizeForFloatRenderTarget(FilteredHistory, E.x, OutputQuantizationError);
 	}
 	#endif
 
diff --git a/Engine/Shaders/Private/TemporalAA/TAAStandalone.usf b/Engine/Shaders/Private/TemporalAA/TAAStandalone.usf
--- a/Engine/Shaders/Private/TemporalAA/TAAStandalone.usf
+++ b/Engine/Shaders/Private/TemporalAA/TAAStandalone.usf
@@ -620,6 +620,17 @@ float ComputeNeightborSampleBilateralWeight(float CenterCocRadius, float SampleC
 
 #endif
 
+float3 QuantizeForFloatRenderTarget(float3 Color, float E, float3 QuantizationError)
+{
+	float3 Error = Color * QuantizationError;
+
+	Error.x = asfloat(asuint(Error.x) & ~0x007FFFFF);
+	Error.y = asfloat(asuint(Error.y) & ~0x007FFFFF);
+	Error.z = asfloat(asuint(Error.z) & ~0x007FFFFF);
+	
+	return Color + Error * E;
+}
+
 
 //------------------------------------------------------- HISTORY's PAYLOAD
 
@@ -2266,7 +2277,7 @@ void MainCS(
 			uint2 Random = Rand3DPCG16(int3(PixelPos, View.StateFrameIndexMod8)).xy;
 			float2 E = Hammersley16(0, 1, Random);
 
-			FinalOutput0.rgb += FinalOutput0.rgb * (E.x * OutputQuantizationError);
+			FinalOutput0.rgb = QuantizeForFloatRenderTarget(FinalOutput0.rgb, E.x, OutputQuantizationError);
 		}
 		#endif
 
diff --git a/Engine/Shaders/Private/TemporalAA/TAAUpdateHistory.usf b/Engine/Shaders/Private/TemporalAA/TAAUpdateHistory.usf
--- a/Engine/Shaders/Private/TemporalAA/TAAUpdateHistory.usf
+++ b/Engine/Shaders/Private/TemporalAA/TAAUpdateHistory.usf
@@ -685,8 +685,8 @@ void MainCS(
 				uint2 Random = Rand3DPCG16(int3(LocalHistoryPixelPos, View.StateFrameIndexMod8)).xy;
 				float2 E = Hammersley16(0, 1, Random);
 
-				FinalHistoryMoment1 += FinalHistoryMoment1 * taa_half3(E.x * HistoryQuantizationError);
-				FinalFallbackColor += FinalFallbackColor * taa_half3(E.x * HistoryQuantizationError);
+				FinalHistoryMoment1 = QuantizeForFloatRenderTarget(FinalHistoryMoment1, E.x, HistoryQuantizationError);
+				FinalFallbackColor = QuantizeForFloatRenderTarget(FinalFallbackColor, E.x, HistoryQuantizationError);
 			}
 			#endif
 
