Add define for clarifying the hair interpolation type.

diff --git a/Engine/Shaders/Private/HairStrands/HairStrandsInterpolation.usf b/Engine/Shaders/Private/HairStrands/HairStrandsInterpolation.usf
--- a/Engine/Shaders/Private/HairStrands/HairStrandsInterpolation.usf
+++ b/Engine/Shaders/Private/HairStrands/HairStrandsInterpolation.usf
@@ -16,6 +16,12 @@
 ///////////////////////////////////////////////////////////////////////////////////////////////////
 // Hair interpolation 
 
+#define INTERPOLATION_RIGID 0
+#define INTERPOLATION_SKINNING_OFFSET 1
+#define INTERPOLATION_SKINNING_TRANSLATION 2
+#define INTERPOLATION_SKINNING_ROTATION 3
+#define INTERPOLATION_SKINNING_TRANSLATION_AND_ROTATION 4
+
 #if SHADER_HAIRINTERPOLATION
 
 #if COMPILER_METAL || COMPILER_VULKAN
@@ -255,7 +261,7 @@ void MainCS(uint2 DispatchThreadId : SV_DispatchThreadID)
 		// 1. Rigid transformation
 		// ControlPoint is in the local hair referential
 		// CurrOffset takes only translation component into account, and is done in object local space (vs. triangle/hair local space)
-		#if PERMUTATION_DYNAMIC_GEOMETRY == 0
+		#if PERMUTATION_DYNAMIC_GEOMETRY == INTERPOLATION_RIGID
 		{
 			float3 CurrOffset = 0;
 			uint4 PackedRestPosition = RenderRestPosePositionBuffer[VertexIndex];
@@ -289,7 +295,7 @@ void MainCS(uint2 DispatchThreadId : SV_DispatchThreadID)
 
 		// 2. Skin transformation
 		// Apply dynamic mesh deformation (translation / rotation)
-		#if PERMUTATION_DYNAMIC_GEOMETRY == 1
+		#if PERMUTATION_DYNAMIC_GEOMETRY == INTERPOLATION_SKINNING_OFFSET
 		{
 			uint4 PackedRestPosition = RenderRestPosePositionBuffer[VertexIndex];
 			const float3 RestControlPoint = f16tof32(PackedRestPosition.xyz) + InRenderHairPositionOffset;
@@ -370,7 +376,7 @@ void MainCS(uint2 DispatchThreadId : SV_DispatchThreadID)
 		#endif	
 
 		// 3. Linear blend skinning (translation)
-		#if PERMUTATION_DYNAMIC_GEOMETRY == 2
+		#if PERMUTATION_DYNAMIC_GEOMETRY == INTERPOLATION_SKINNING_TRANSLATION
 		{
 			uint4 PackedRestPosition = RenderRestPosePositionBuffer[VertexIndex];
 			const float3 RestControlPoint = f16tof32(PackedRestPosition.xyz) + InRenderHairPositionOffset;
@@ -414,7 +420,7 @@ void MainCS(uint2 DispatchThreadId : SV_DispatchThreadID)
 		#endif	
 
 		// 4. Linear blend skinning (rotation)
-		#if PERMUTATION_DYNAMIC_GEOMETRY == 3
+		#if PERMUTATION_DYNAMIC_GEOMETRY == INTERPOLATION_SKINNING_ROTATION
 		{
 			uint4 PackedRestPosition = RenderRestPosePositionBuffer[VertexIndex];
 			const float3 RestControlPoint = f16tof32(PackedRestPosition.xyz) + InRenderHairPositionOffset;
@@ -465,7 +471,7 @@ void MainCS(uint2 DispatchThreadId : SV_DispatchThreadID)
 		#endif
 
 		// 5. Linear blend skinning (translation + rotation)
-		#if PERMUTATION_DYNAMIC_GEOMETRY == 4
+		#if PERMUTATION_DYNAMIC_GEOMETRY == INTERPOLATION_SKINNING_TRANSLATION_AND_ROTATION
 		{
 			uint4 PackedRestPosition = RenderRestPosePositionBuffer[VertexIndex];
 			const float3 RestControlPoint = f16tof32(PackedRestPosition.xyz) + InRenderHairPositionOffset;
