Disable base pass skylight when Lumen GI is enabled

diff --git a/Engine/Shaders/Private/BasePassPixelShader.usf b/Engine/Shaders/Private/BasePassPixelShader.usf
--- a/Engine/Shaders/Private/BasePassPixelShader.usf
+++ b/Engine/Shaders/Private/BasePassPixelShader.usf
@@ -19,8 +19,10 @@
 		#define PreIntegratedGFSampler	TranslucentBasePass.PreIntegratedGFSampler
 		#define SceneColorCopySampler	TranslucentBasePass.SceneColorCopySampler
 	#endif
+	#define UseBasePassSkylight TranslucentBasePass.Shared.UseBasePassSkylight
 #else
 	#define EyeAdaptationStruct OpaqueBasePass
+	#define UseBasePassSkylight OpaqueBasePass.Shared.UseBasePassSkylight
 #endif
 
 // Material setting to allow forward shading (including mobile) to use preintegrated GF lut for simple IBL.
@@ -333,7 +335,7 @@ void GetSkyLighting(FMaterialPixelParameters MaterialParameters, VTPageTableResu
 	}
 	else
 #endif
-
+	if (UseBasePassSkylight > 0)
 	{
 #if ENABLE_SKY_LIGHT
 	
diff --git a/Engine/Source/Runtime/Renderer/Private/BasePassRendering.cpp b/Engine/Source/Runtime/Renderer/Private/BasePassRendering.cpp
--- a/Engine/Source/Runtime/Renderer/Private/BasePassRendering.cpp
+++ b/Engine/Source/Runtime/Renderer/Private/BasePassRendering.cpp
@@ -573,6 +573,10 @@ void SetupSharedBasePassParameters(
 	{
 		SharedParameters.SSProfilesTexture = PooledRT->GetShaderResourceRHI();
 	}
+
+	// Skip base pass skylight if Lumen GI is enabled, as Lumen handles the skylight.
+	// Ideally we would choose a different shader permutation to skip skylight, but Lumen GI is only known per-view
+	SharedParameters.UseBasePassSkylight = View.FinalPostProcessSettings.DynamicGlobalIlluminationMethod == EDynamicGlobalIlluminationMethod::Lumen ? 0 : 1;
 }
 
 TRDGUniformBufferRef<FOpaqueBasePassUniformParameters> CreateOpaqueBasePassUniformBuffer(
diff --git a/Engine/Source/Runtime/Renderer/Private/BasePassRendering.h b/Engine/Source/Runtime/Renderer/Private/BasePassRendering.h
--- a/Engine/Source/Runtime/Renderer/Private/BasePassRendering.h
+++ b/Engine/Source/Runtime/Renderer/Private/BasePassRendering.h
@@ -67,6 +67,7 @@ BEGIN_GLOBAL_SHADER_PARAMETER_STRUCT(FSharedBasePassUniformParameters,)
 	SHADER_PARAMETER_STRUCT(FFogUniformParameters, Fog)
 	SHADER_PARAMETER_STRUCT(FFogUniformParameters, FogISR)
 	SHADER_PARAMETER_TEXTURE(Texture2D, SSProfilesTexture)
+	SHADER_PARAMETER(uint32, UseBasePassSkylight)
 END_GLOBAL_SHADER_PARAMETER_STRUCT()
 
 BEGIN_GLOBAL_SHADER_PARAMETER_STRUCT(FOpaqueBasePassUniformParameters,)
