Re-enable WRITES_VELOCITY_TO_GBUFFER functionality that was disabled for GPU-Scene instance culling path.

diff --git a/Engine/Shaders/Private/BasePassVertexShader.usf b/Engine/Shaders/Private/BasePassVertexShader.usf
--- a/Engine/Shaders/Private/BasePassVertexShader.usf
+++ b/Engine/Shaders/Private/BasePassVertexShader.usf
@@ -220,10 +220,7 @@ void Main(
 
 #if WRITES_VELOCITY_TO_GBUFFER
 	{
-#if defined(GPUCULL_TODO)
-		float4 PrevTranslatedWorldPosition = WorldPosition;
-		// GPUCULL_TODO: This here is a non-abstracted use VFIntermediates, needs to be abstracted using accessor function. 
-#else // !defined(GPUCULL_TODO)
+		// Warning: This here is a non-abstracted use of VFIntermediates.PrimitiveId - Only works because GetPrimitiveData(VFIntermediates.PrimitiveId) is a macro.
 		float4 PrevTranslatedWorldPosition = float4(0, 0, 0, 1);
 		BRANCH
 		if (GetPrimitiveData(VFIntermediates.PrimitiveId).OutputVelocity > 0 || View.ForceDrawAllVelocities != 0)
@@ -236,7 +233,6 @@ void Main(
 				PrevTranslatedWorldPosition = mul(float4(PrevTranslatedWorldPosition.xyz, 1), ResolvedView.PrevTranslatedWorldToClip);
 			#endif
 		}
-#endif //!defined(GPUCULL_TODO)
 
 		#if USING_TESSELLATION
 			// We let the Domain Shader convert to post projection when tessellating
