Strata data are now stored in 8x8 tiles and this prevent gpu being starved because data is not coming fast enough through the cache.

Default-no-strata => Strata before => Strata after simple tiled memory
Base pass: 	1.2 => 2.4 => 1.4 ms
Light pass: 	0.7 => 1.6 => 0.6 ms
SkyEnv:    	1.0 => 2.0 => 0.7 ms
Classification 	xx => 2ms => 0.15 ms

diff --git a/Engine/Shaders/Private/BasePassPixelShader.usf b/Engine/Shaders/Private/BasePassPixelShader.usf
--- a/Engine/Shaders/Private/BasePassPixelShader.usf
+++ b/Engine/Shaders/Private/BasePassPixelShader.usf
@@ -1555,7 +1555,7 @@ void FPixelShaderInOut_MainPS(
 	const float3 VertexNormal = MaterialParameters.TangentToWorld[2].xyz;
 	PackStrataOut(
 		OpaqueBasePass.Strata.MaterialLobesBufferUAV, 
-		GetStrataPixelDataByteOffset(PixelPos, ResolvedView.BufferSizeAndInvSize.x, OpaqueBasePass.Strata.MaxBytesPerPixel),
+		GetStrataPixelDataByteOffset(PixelPos, uint2(ResolvedView.BufferSizeAndInvSize.xy), OpaqueBasePass.Strata.MaxBytesPerPixel),
 		StrataPixelHeader, Strata,
 		MaterialParameters.CameraVector, GBuffer, EmissiveLuminance);
 
diff --git a/Engine/Shaders/Private/DeferredLightPixelShaders.usf b/Engine/Shaders/Private/DeferredLightPixelShaders.usf
--- a/Engine/Shaders/Private/DeferredLightPixelShaders.usf
+++ b/Engine/Shaders/Private/DeferredLightPixelShaders.usf
@@ -227,7 +227,7 @@ void DeferredLightPixelMain(
 					ShadowTerms,
 					RectTexture,
 					Strata.MaterialLobesBuffer,
-					GetStrataPixelDataByteOffset(PixelPos, View.BufferSizeAndInvSize.x, Strata.MaxBytesPerPixel));
+					GetStrataPixelDataByteOffset(PixelPos, uint2(View.BufferSizeAndInvSize.xy), Strata.MaxBytesPerPixel));
 			}
 		}
 #else
diff --git a/Engine/Shaders/Private/ReflectionEnvironmentPixelShader.usf b/Engine/Shaders/Private/ReflectionEnvironmentPixelShader.usf
--- a/Engine/Shaders/Private/ReflectionEnvironmentPixelShader.usf
+++ b/Engine/Shaders/Private/ReflectionEnvironmentPixelShader.usf
@@ -379,7 +379,7 @@ void ReflectionEnvironmentSkyLighting(
 			AmbientOcclusion,
 			TopLayerSpecularContributionFactor,
 			Strata.MaterialLobesBuffer,
-			GetStrataPixelDataByteOffset(PixelPos, View.BufferSizeAndInvSize.x, Strata.MaxBytesPerPixel),
+			GetStrataPixelDataByteOffset(PixelPos, uint2(View.BufferSizeAndInvSize.xy), Strata.MaxBytesPerPixel),
 			SSRReductionFactor);
 
 		float3 SkyLighting = SSR.rgb * TopLayerData.EnvBRDF * SSRReductionFactor + StrataEnvLighting.rgb;
diff --git a/Engine/Shaders/Private/Strata/Strata.ush b/Engine/Shaders/Private/Strata/Strata.ush
--- a/Engine/Shaders/Private/Strata/Strata.ush
+++ b/Engine/Shaders/Private/Strata/Strata.ush
@@ -482,16 +482,46 @@ FStrataData GetStrataSingleLayerWaterBSDF(
 	return StrataData;
 }
 
+// Optimised for power of two because it relies on a division done using bit shift
+uint DivideAndRoundUp(uint Dividend, uint Divisor, uint DivisorAsBitShift)
+{
+	return (Dividend + Divisor - 1) >> DivisorAsBitShift;
+}
 
 struct FStrataAddressing
 {
 	uint CurrentAddr;
+	uint ByteOffsetToNextUint;
 	uint ReadBytes;
 };
-FStrataAddressing GetStrataPixelDataByteOffset(uint2 PixelPos, uint ViewBufferSizeX, uint StrataMaxBytesPerPixel)
+FStrataAddressing GetStrataPixelDataByteOffset(uint2 PixelPos, uint2 ViewBufferSize, uint StrataMaxBytesPerPixel)
 {
-	FStrataAddressing StrataAddressing;
-	StrataAddressing.CurrentAddr = (PixelPos.y * ViewBufferSizeX + PixelPos.x) * StrataMaxBytesPerPixel;
+	FStrataAddressing StrataAddressing = (FStrataAddressing)0;
+
+#if 0
+	// Linear storage, pixel0 bytes, pixel1 bytes, etc.
+	StrataAddressing.CurrentAddr = (PixelPos.y * ViewBufferSize.x + PixelPos.x) * StrataMaxBytesPerPixel;
+	StrataAddressing.ByteOffsetToNextUint = 4;
+#else
+
+	// A tile of pixels has each uint that can be read packed in contiguous memory to be cache efficient. Then, the next set of uint to be read is in another "plane", after all uint from all the pixel/tile.
+	// Without that, L2 is rarely hit and GCN TA FIFO command buffer is getting full in no time because memory fetches are not served in time (GPU cost then increase by at least x2)
+
+	uint TilePixelSizeBytes			= 4;	// sizeof(uint)
+	uint TileSizeBytes				= STRATA_DATA_TILE_SIZE * STRATA_DATA_TILE_SIZE * TilePixelSizeBytes;
+	uint2 TileCount					= uint2(DivideAndRoundUp(ViewBufferSize.x, STRATA_DATA_TILE_SIZE, STRATA_DATA_TILE_SIZE_DIV_AS_SHIFT), DivideAndRoundUp(ViewBufferSize.y, STRATA_DATA_TILE_SIZE, STRATA_DATA_TILE_SIZE_DIV_AS_SHIFT));
+	uint2 TileCoord					= uint2(PixelPos.x >> STRATA_DATA_TILE_SIZE_DIV_AS_SHIFT, PixelPos.y >> STRATA_DATA_TILE_SIZE_DIV_AS_SHIFT);
+	uint2 PixelCoordInTile			= PixelPos - TileCoord * STRATA_DATA_TILE_SIZE;
+
+	uint PreviousTilesBytes			= (TileCoord.y * TileCount.x + TileCoord.x) * TileSizeBytes;
+	uint PreviousPixelsInTileBytes	= (PixelCoordInTile.y * STRATA_DATA_TILE_SIZE + PixelCoordInTile.x) * TilePixelSizeBytes;
+	StrataAddressing.CurrentAddr	= PreviousTilesBytes + PreviousPixelsInTileBytes;
+
+	// All
+	StrataAddressing.ByteOffsetToNextUint = TileCount.x * TileCount.y * TileSizeBytes;
+
+#endif
+
 	StrataAddressing.ReadBytes = 0;
 	return StrataAddressing;
 }
@@ -508,13 +538,25 @@ uint  StrataLoadUint1(ByteAddressBuffer StrataBuffer, inout FStrataAddressing St
 
 #else
 
-#define STRATA_STORE_UINT1(x) StrataBuffer.Store (StrataAddressing.CurrentAddr, asuint(x)); StrataAddressing.CurrentAddr += 4;	 StrataAddressing.ReadBytes += 4;
+#define STRATA_STORE_UINT1(x)	StrataBuffer.Store (StrataAddressing.CurrentAddr, asuint(x)); \
+								StrataAddressing.CurrentAddr += StrataAddressing.ByteOffsetToNextUint; \
+								StrataAddressing.ReadBytes += 4;
 
-#define STRATA_LOAD_UINT1(x) x = StrataBuffer.Load (StrataAddressing.CurrentAddr); StrataAddressing.CurrentAddr += 4;	StrataAddressing.ReadBytes += 4;
+#define STRATA_LOAD_UINT1(x)	x = StrataBuffer.Load (StrataAddressing.CurrentAddr); \
+								StrataAddressing.CurrentAddr += StrataAddressing.ByteOffsetToNextUint; \
+								StrataAddressing.ReadBytes += 4;
 
-#define STRATA_LOAD_FLOAT1(x) x = asfloat(StrataBuffer.Load (StrataAddressing.CurrentAddr)); StrataAddressing.CurrentAddr += 4;		StrataAddressing.ReadBytes += 4;
+#define STRATA_LOAD_FLOAT1(x)	x = asfloat(StrataBuffer.Load (StrataAddressing.CurrentAddr)); \
+								StrataAddressing.CurrentAddr += StrataAddressing.ByteOffsetToNextUint; \
+								StrataAddressing.ReadBytes += 4;
 
-uint  StrataLoadUint1(ByteAddressBuffer StrataBuffer, inout FStrataAddressing StrataAddressing) { uint  x = StrataBuffer.Load(StrataAddressing.CurrentAddr); StrataAddressing.CurrentAddr += 4;	StrataAddressing.ReadBytes += 4;	return x; }
+uint  StrataLoadUint1(ByteAddressBuffer StrataBuffer, inout FStrataAddressing StrataAddressing) 
+{ 
+	uint  x = StrataBuffer.Load(StrataAddressing.CurrentAddr); 
+	StrataAddressing.CurrentAddr += StrataAddressing.ByteOffsetToNextUint;	
+	StrataAddressing.ReadBytes += 4;
+	return x; 
+}
 
 #endif
 
@@ -578,7 +620,7 @@ FStrataPixelHeader InitialiseStrataPixelHeader()
 	return Out;
 }
 
-float3x3 StrataGetBSDFSharedBasis(in FStrataPixelHeader StrataPixelHeader, in FStrataBSDF BSDF)
+float3x3 StrataGetBSDFSharedBasis(in FStrataPixelHeader StrataPixelHeader, in FStrataBSDF BSDF, in FStrataAddressing StrataAddressing)
 {
 #if STRATA_INLINE_SHADING
 	// Assume Normal & Tangent are already been normalized
@@ -588,7 +630,7 @@ float3x3 StrataGetBSDFSharedBasis(in FStrataPixelHeader StrataPixelHeader, in FS
 	return float3x3(Tangent, Bitangent, Normal);
 #else
 	const bool bHasAnisotropy = BSDF_GETANISOTROPY(BSDF) != 0;
-	uint PackedBasis = StrataPixelHeader.StrataBuffer.Load(StrataPixelHeader.SharedNormalByteOffset + BSDF_GETNORMALID(BSDF) * STRATA_PACKED_NORMAL_STRIDE_BYTES);
+	uint PackedBasis = StrataPixelHeader.StrataBuffer.Load(StrataPixelHeader.SharedNormalByteOffset + BSDF_GETNORMALID(BSDF) * StrataAddressing.ByteOffsetToNextUint);
 	uint BasisType   = StrataGetSharedNormalType(StrataPixelHeader.SharedNormalTypes, BSDF_GETNORMALID(BSDF));
 	float3x3 OutTangentBasis;
 	if (BasisType == STRATA_BASIS_TYPE_NORMAL)
@@ -1072,7 +1114,7 @@ void PackStrataOut(
 			BSDFCount++;
 
 			// STRATA_TODO apply the same normal transform as for the material
-			float3x3 TangentBasis = StrataGetBSDFSharedBasis(StrataPixelHeader, BSDF);
+			float3x3 TangentBasis = StrataGetBSDFSharedBasis(StrataPixelHeader, BSDF, StrataAddressing);
 			float3 N = TangentBasis[2];
 			//N = normalize(N);
 
@@ -1356,7 +1398,7 @@ FStrataPixelHeader UnpackStrataHeaderIn(ByteAddressBuffer StrataBuffer, inout FS
 	Out.SharedNormalTypes = NormalTypes;
 
 	// Now skip over the shared normal memory to be able to load the BSDF content
-	StrataAddressing.CurrentAddr += NormalCount * STRATA_PACKED_NORMAL_STRIDE_BYTES;
+	StrataAddressing.CurrentAddr += NormalCount * StrataAddressing.ByteOffsetToNextUint;
 
 	// And account for the space used by normal for the debug information
 	StrataAddressing.ReadBytes += NormalCount * STRATA_PACKED_NORMAL_STRIDE_BYTES;
@@ -1370,20 +1412,25 @@ FStrataPixelHeader UnpackStrataHeaderIn(ByteAddressBuffer StrataBuffer, inout FS
 FStrataPixelHeader UnpackFastPathStrataHeaderIn(ByteAddressBuffer StrataBuffer, inout FStrataAddressing StrataAddressing)
 {
 	FStrataPixelHeader Out = InitialiseStrataPixelHeader();
-	
+
+	// This function should only be used when inline shading is 0
+#if STRATA_INLINE_SHADING==0
+
+	// When there  is only a single Slab BSDF loaded
+	uint NormalCount = 1;
+
 	// Skip over the header packed data
-	StrataAddressing.CurrentAddr += 4; ;
+	StrataAddressing.CurrentAddr += StrataAddressing.ByteOffsetToNextUint;
 
 	Out.BSDFCount = 1;
-#if STRATA_INLINE_SHADING==0
 	Out.SharedNormalByteOffset = StrataAddressing.CurrentAddr;
 	Out.StrataBuffer = StrataBuffer;
 	Out.SharedNormalTypes = STRATA_BASIS_TYPE_NORMAL;
-#endif
 
 	// Skip over the shared normal memory to be able to load the BSDF content
-	StrataAddressing.CurrentAddr += STRATA_PACKED_NORMAL_STRIDE_BYTES;
+	StrataAddressing.CurrentAddr += NormalCount * StrataAddressing.ByteOffsetToNextUint;
 
+#endif
 	return Out;
 }
 
diff --git a/Engine/Shaders/Private/Strata/StrataDeferredLighting.ush b/Engine/Shaders/Private/Strata/StrataDeferredLighting.ush
--- a/Engine/Shaders/Private/Strata/StrataDeferredLighting.ush
+++ b/Engine/Shaders/Private/Strata/StrataDeferredLighting.ush
@@ -186,7 +186,7 @@ float4 StrataDeferredLighting(
 #else
 		FStrataBSDF BSDF = UnpackStrataBSDFIn(StrataDataBuffer, StrataAddressing);
 #endif
-		FStrataBSDFContext BSDFContext = StrataCreateBSDFContext(StrataPixelHeader, BSDF, V, L);
+		FStrataBSDFContext BSDFContext = StrataCreateBSDFContext(StrataPixelHeader, BSDF, StrataAddressing, V, L);
 
 		// Evaluate incident lighting & (optionally) precompute area light integration
 		FStrataLightIntegrationContext IntegrationContext = StrataIniLightIntegrationContext();
diff --git a/Engine/Shaders/Private/Strata/StrataEnvironmentLighting.ush b/Engine/Shaders/Private/Strata/StrataEnvironmentLighting.ush
--- a/Engine/Shaders/Private/Strata/StrataEnvironmentLighting.ush
+++ b/Engine/Shaders/Private/Strata/StrataEnvironmentLighting.ush
@@ -50,7 +50,7 @@ float4 StrataEnvironmentLighting(
 
 		// Create the BSDF context
 		const float3 UnusedL = float3(0.0f, 0.0f, 1.0f);
-		FStrataBSDFContext StrataBSDFContext = StrataCreateBSDFContext(StrataPixelHeader, BSDF, V, UnusedL);
+		FStrataBSDFContext StrataBSDFContext = StrataCreateBSDFContext(StrataPixelHeader, BSDF, StrataAddressing, V, UnusedL);
 
 		// Evaluate environment lighting
 		const bool bEnableSpecular = ReflectionStruct.SkyLightParameters.y > 0.0f;
diff --git a/Engine/Shaders/Private/Strata/StrataEvaluation.ush b/Engine/Shaders/Private/Strata/StrataEvaluation.ush
--- a/Engine/Shaders/Private/Strata/StrataEvaluation.ush
+++ b/Engine/Shaders/Private/Strata/StrataEvaluation.ush
@@ -83,7 +83,7 @@ void StrataGetBSDFRoughness(in FStrataBSDF BSDF, inout float R)
 	}
 }
 
-FStrataBSDFContext StrataCreateBSDFContext(FStrataPixelHeader StrataPixelHeader, FStrataBSDF BSDF, float3 V, float3 L, bool bOverrideWithLocalFrame=false)
+FStrataBSDFContext StrataCreateBSDFContext(FStrataPixelHeader StrataPixelHeader, FStrataBSDF BSDF, const FStrataAddressing StrataAddressing, float3 V, float3 L, bool bOverrideWithLocalFrame=false)
 {
 	FStrataBSDFContext BSDFContext = (FStrataBSDFContext)0;
 
@@ -92,7 +92,7 @@ FStrataBSDFContext StrataCreateBSDFContext(FStrataPixelHeader StrataPixelHeader,
 	BSDFContext.V = V;
 	BSDFContext.L = L;
 
-	float3x3 TangentBasis = StrataGetBSDFSharedBasis(StrataPixelHeader, BSDF);
+	float3x3 TangentBasis = StrataGetBSDFSharedBasis(StrataPixelHeader, BSDF, StrataAddressing);
 	BSDFContext.X = TangentBasis[0]; 
 	BSDFContext.Y = TangentBasis[1];
 	BSDFContext.N = TangentBasis[2]; 
@@ -337,14 +337,14 @@ struct FStrataImportanceSampleResult
  * E: two random numbers
  * CameraVector: vector from the camera to the considered direction (not V)
  */
-FStrataImportanceSampleResult StrataImportanceSampleBSDF(FStrataPixelHeader StrataPixelHeader, FStrataBSDF BSDF, float2 E, float3 CameraVector)
+FStrataImportanceSampleResult StrataImportanceSampleBSDF(FStrataPixelHeader StrataPixelHeader, FStrataBSDF BSDF, const FStrataAddressing StrataAddressing, float2 E, float3 CameraVector)
 {
 	FStrataImportanceSampleResult Sample = (FStrataImportanceSampleResult)0;
 
 	// Create a context from the information we have
 	const float3 UnusedL = float3(0.0f, 0.0f, 1.0f);
 	const float3 V = -CameraVector;
-	FStrataBSDFContext BSDFContext = StrataCreateBSDFContext(StrataPixelHeader, BSDF, V, UnusedL);
+	FStrataBSDFContext BSDFContext = StrataCreateBSDFContext(StrataPixelHeader, BSDF, StrataAddressing, V, UnusedL);
 
 	float PDF = 0.0f;
 	const uint BSDFType = BSDF_GETTYPE(BSDF);
diff --git a/Engine/Shaders/Private/Strata/StrataForwardLighting.ush b/Engine/Shaders/Private/Strata/StrataForwardLighting.ush
--- a/Engine/Shaders/Private/Strata/StrataForwardLighting.ush
+++ b/Engine/Shaders/Private/Strata/StrataForwardLighting.ush
@@ -17,6 +17,8 @@ float3 StrataForwardLighting(
 	inout float3 OutThroughput
 	)
 {
+	FStrataAddressing NullStrataAddressing = (FStrataAddressing)0;	// Fake unused in StrataCreateBSDFContext when using Forward inline shading
+
 	const float OpaqueBSDFThroughput = 0.0f;
 	const float FullThroughput = 1.0f;
 
@@ -40,7 +42,7 @@ float3 StrataForwardLighting(
 			const float3 BSDFCoverage = BSDF.Weight;
 
 			// Create the BSDF context
-			FStrataBSDFContext StrataBSDFContext = StrataCreateBSDFContext(Header, BSDF, V, L);
+			FStrataBSDFContext StrataBSDFContext = StrataCreateBSDFContext(Header, BSDF, NullStrataAddressing, V, L);
 
 			// Compute some common factors
 			const float3 DirLightBaseFactor = OutThroughput * BSDFCoverage * LColor * LShadow;
diff --git a/Engine/Shaders/Private/Strata/StrataMaterialClassification.usf b/Engine/Shaders/Private/Strata/StrataMaterialClassification.usf
--- a/Engine/Shaders/Private/Strata/StrataMaterialClassification.usf
+++ b/Engine/Shaders/Private/Strata/StrataMaterialClassification.usf
@@ -17,6 +17,7 @@
 void ComputePostProcessBSDFData(
 	const FStrataPixelHeader Header,
 	const FStrataBSDF BSDF,
+	const FStrataAddressing StrataAddressing,
 	float3 V,
 	inout FStrataClassification Classification,
 	inout FStrataSubsurfaceData SSS,
@@ -30,7 +31,7 @@ void ComputePostProcessBSDFData(
 	const float  BSDFCoverageAvg = dot(BSDF.Weight,1.0f/3.0f);
 
 	// STRATA_TODO apply the same normal transform as for the material
-	float3x3 TangentBasis = StrataGetBSDFSharedBasis(Header, BSDF);
+	float3x3 TangentBasis = StrataGetBSDFSharedBasis(Header, BSDF, StrataAddressing);
 	float3 N = TangentBasis[2];
 	//N = normalize(N);
 
@@ -156,13 +157,13 @@ void MainPS(
 		float4 WorldPos = mul(float4(ScreenPosition * Depth, Depth, 1), View.ScreenToWorld);
 		const float3 V = normalize(View.WorldCameraOrigin - WorldPos.xyz);
 
-		FStrataAddressing StrataAddressing = GetStrataPixelDataByteOffset(PixelPos, View.BufferSizeAndInvSize.x, Strata.MaxBytesPerPixel);
+		FStrataAddressing StrataAddressing = GetStrataPixelDataByteOffset(PixelPos, uint2(View.BufferSizeAndInvSize.xy), Strata.MaxBytesPerPixel);
 		FStrataPixelHeader Header = UnpackStrataHeaderIn(Strata.MaterialLobesBuffer, StrataAddressing);
 		ClassificationData.bIsSimple = Header.BSDFCount == 1;
 		for (uint BSDFIndex = 0; BSDFIndex < Header.BSDFCount; ++BSDFIndex)
 		{
 			const FStrataBSDF BSDF = UnpackStrataBSDFIn(Strata.MaterialLobesBuffer, StrataAddressing);
-			ComputePostProcessBSDFData(Header, BSDF, V, ClassificationData, SSSData, TopLayerData);
+			ComputePostProcessBSDFData(Header, BSDF, StrataAddressing, V, ClassificationData, SSSData, TopLayerData);
 		}
 
 		FinalizeTopLayerData(TopLayerData);
diff --git a/Engine/Shaders/Private/Strata/StrataVisualize.usf b/Engine/Shaders/Private/Strata/StrataVisualize.usf
--- a/Engine/Shaders/Private/Strata/StrataVisualize.usf
+++ b/Engine/Shaders/Private/Strata/StrataVisualize.usf
@@ -33,6 +33,7 @@ float3 DrawPlot(
 	const float3 V, 
 	const FStrataPixelHeader Header,
 	const FStrataBSDF BSDF, 
+	const FStrataAddressing StrataAddressing,
 	uint2 PixelCoord, 
 	uint2 RectOffset, 
 	uint2 RectSize)
@@ -48,7 +49,7 @@ float3 DrawPlot(
 		const float CosTheta = sqrt(1 - Pow2(SinTheta));
 		const float3 L = float3(ProjCoord, CosTheta);
 
-		const float3x3 Basis = StrataGetBSDFSharedBasis(Header, BSDF);
+		const float3x3 Basis = StrataGetBSDFSharedBasis(Header, BSDF, StrataAddressing);
 		float3 LocalV = mul(Basis, V);
 		if (LocalV.z < 0)
 		{
@@ -56,7 +57,7 @@ float3 DrawPlot(
 			LocalV = normalize(LocalV);
 		}
 
-		FStrataBSDFContext BSDFContext = StrataCreateBSDFContext(Header, BSDF, LocalV, L, true);
+		FStrataBSDFContext BSDFContext = StrataCreateBSDFContext(Header, BSDF, StrataAddressing, LocalV, L, true);
 		FStrataEvaluateResult BSDFEvaluate = StrataEvaluateBSDF(BSDFContext);
 		Out = (BSDFEvaluate.DiffusePathValue + BSDFEvaluate.SpecularPathValue + BSDFEvaluate.TransmissionPathValue) * BSDFContext.SatNoL; // * BSDF.Weight
 	}
@@ -67,6 +68,7 @@ void VisualizeMaterialPS(
 	float4 SVPos : SV_POSITION,
 	out float4 OutColor : SV_Target0)
 {
+#if 1	// Here to disabled on some platforms having problems with that shader. STRATA_TODO investigate why this shader is spilling like crazy again
 	const uint2 PixelPos = uint2(SVPos.xy);
 	const uint2 MaterialPos = View.CursorPosition;
 	const int  InspectSquare = 10;
@@ -84,7 +86,7 @@ void VisualizeMaterialPS(
 	OutColor = 0;
 	if (bStrataMaterial)
 	{
-		FStrataPixelHeader Header = UnpackStrataHeaderIn(Strata.MaterialLobesBuffer, GetStrataPixelDataByteOffset(MaterialPos, View.BufferSizeAndInvSize.x, Strata.MaxBytesPerPixel));
+		FStrataPixelHeader Header = UnpackStrataHeaderIn(Strata.MaterialLobesBuffer, GetStrataPixelDataByteOffset(MaterialPos, uint2(View.BufferSizeAndInvSize.xy), Strata.MaxBytesPerPixel));
 
 		OutColor.w = 0.05f;
 		if (Header.BSDFCount == 1)
@@ -117,7 +119,7 @@ void VisualizeMaterialPS(
 	bool bPlotValid = false;
 	if (bStrataMaterial)
 	{
-		FStrataAddressing StrataAddressing = GetStrataPixelDataByteOffset(PixelPos, View.BufferSizeAndInvSize.x, Strata.MaxBytesPerPixel);
+		FStrataAddressing StrataAddressing = GetStrataPixelDataByteOffset(PixelPos, uint2(View.BufferSizeAndInvSize.xy), Strata.MaxBytesPerPixel);
 
 		FStrataPixelHeader Header = UnpackStrataHeaderIn(Strata.MaterialLobesBuffer, StrataAddressing);
 
@@ -128,7 +130,7 @@ void VisualizeMaterialPS(
 			// Referential
 			if (all(PixelPos == MaterialPos))
 			{
-				FStrataBSDFContext BSDFContext = StrataCreateBSDFContext(Header, BSDF, V, float3(0,0,1));
+				FStrataBSDFContext BSDFContext = StrataCreateBSDFContext(Header, BSDF, StrataAddressing, V, float3(0,0,1));
 				DrawReferential(WorldPosition, BSDFContext.X, BSDFContext.Y, BSDFContext.N, float3(1,1,0));
 			}
 
@@ -137,7 +139,7 @@ void VisualizeMaterialPS(
 				const uint2 Resolution = 180;
 				const uint OffsetX = 50 + BSDFIndex * 200;
 				const uint OffsetY = 50;
-				OutColor.xyz += DrawPlot(V, Header, BSDF, PixelPos, uint2(OffsetX, OffsetY), Resolution);
+				OutColor.xyz += DrawPlot(V, Header, BSDF, StrataAddressing, PixelPos, uint2(OffsetX, OffsetY), Resolution);
 			}
 		}
 
@@ -155,7 +157,7 @@ void VisualizeMaterialPS(
 		uint2 RectMin = uint2(MaterialPos + InspectSquare);
 		uint2 RectMax = RectMin;
 
-		FStrataAddressing StrataAddressing = GetStrataPixelDataByteOffset(MaterialPos, View.BufferSizeAndInvSize.x, Strata.MaxBytesPerPixel);
+		FStrataAddressing StrataAddressing = GetStrataPixelDataByteOffset(MaterialPos, uint2(View.BufferSizeAndInvSize.xy), Strata.MaxBytesPerPixel);
 		const uint PixelStrataDataByteOffsetStart = StrataAddressing.ReadBytes;
 
 		FStrataPixelHeader Header = UnpackStrataHeaderIn(Strata.MaterialLobesBuffer, StrataAddressing);
@@ -454,6 +456,7 @@ void VisualizeMaterialPS(
 		}
 #endif
 	}
+#endif
 }
 
 
diff --git a/Engine/Shaders/Shared/StrataDefinitions.h b/Engine/Shaders/Shared/StrataDefinitions.h
--- a/Engine/Shaders/Shared/StrataDefinitions.h
+++ b/Engine/Shaders/Shared/StrataDefinitions.h
@@ -29,4 +29,5 @@
 #define STRATA_BSDF_TYPE_SINGLELAYERWATER	5
 // When more than 8 BSDF must exists, please update STATE_BITS_BSDF and FStrataClassification.ShadingModels packing in Strata.ush
 
-
+#define STRATA_DATA_TILE_SIZE				8
+#define STRATA_DATA_TILE_SIZE_DIV_AS_SHIFT	3
diff --git a/Engine/Source/Runtime/Renderer/Private/Strata/Strata.cpp b/Engine/Source/Runtime/Renderer/Private/Strata/Strata.cpp
--- a/Engine/Source/Runtime/Renderer/Private/Strata/Strata.cpp
+++ b/Engine/Source/Runtime/Renderer/Private/Strata/Strata.cpp
@@ -9,6 +9,7 @@
 #include "RendererInterface.h"
 #include "UniformBuffer.h"
 #include "SceneTextureParameters.h"
+#include "StrataDefinitions.h"
 
 //PRAGMA_DISABLE_OPTIMIZATION
 
@@ -62,8 +63,9 @@ void InitialiseStrataFrameSceneData(FSceneRenderer& SceneRenderer, FRDGBuilder&
 	{
 		FIntPoint BufferSizeXY = GetSceneTextureExtent();
 		
-		ResolutionX = BufferSizeXY.X;
-		ResolutionY = BufferSizeXY.Y;
+		// We need to allocate enough for the tiled memory addressing to always work
+		ResolutionX = FMath::DivideAndRoundUp(BufferSizeXY.X, STRATA_DATA_TILE_SIZE) * STRATA_DATA_TILE_SIZE;
+		ResolutionY = FMath::DivideAndRoundUp(BufferSizeXY.Y, STRATA_DATA_TILE_SIZE) * STRATA_DATA_TILE_SIZE;
 
 		// Previous GBuffer when complete was 28bytes
 		// check out Strata.ush to see how this is computed
