Do not sample TopLayer data during the envlight pass.
Raise occupancy from 60 to 70, cost go down from 0.75 to 0.71.

diff --git a/Engine/Shaders/Private/ReflectionEnvironmentPixelShader.usf b/Engine/Shaders/Private/ReflectionEnvironmentPixelShader.usf
--- a/Engine/Shaders/Private/ReflectionEnvironmentPixelShader.usf
+++ b/Engine/Shaders/Private/ReflectionEnvironmentPixelShader.usf
@@ -350,9 +350,6 @@ void ReflectionEnvironmentSkyLighting(
 		float DeviceZ = SampleDeviceZFromSceneTextures(BufferUV);
 		float SceneDepth = ConvertFromDeviceZ(DeviceZ);
 
-		// STRATA_TODO: when GCN TA command FIFO is not longer the bottleneck, try to compute TopLayerData.EnvBRDF in StrataEnvironmentLighting to avoid reading that TopLayerData texture.
-		const FStrataTopLayerData TopLayerData = StrataUnpackTopLayerData(Strata.TopLayerTexture.Load(uint3(PixelPos, 0)));
-
 		float3 WorldPosition = mul(float4(ScreenPosition * SceneDepth, SceneDepth, 1), View.ScreenToWorld).xyz;
 		float3 CameraToPixel = normalize(WorldPosition - View.WorldCameraOrigin);
 		float3 V = -CameraToPixel;
@@ -368,6 +365,7 @@ void ReflectionEnvironmentSkyLighting(
 
 		// A reduction of SSR contribution is needed when there is haziness because we only consider SSR for the sharpest lobe.
 		float SSRReductionFactor = 1.0f;
+		float3 TopLayerEnvBRDF = 0.0f;
 
 		float4 StrataEnvLighting = StrataEnvironmentLighting(
 			BufferUV,
@@ -380,9 +378,10 @@ void ReflectionEnvironmentSkyLighting(
 			TopLayerSpecularContributionFactor,
 			Strata.MaterialLobesBuffer,
 			GetStrataPixelDataByteOffset(PixelPos, uint2(View.BufferSizeAndInvSize.xy), Strata.MaxBytesPerPixel),
-			SSRReductionFactor);
+			SSRReductionFactor,
+			TopLayerEnvBRDF);
 
-		float3 SkyLighting = SSR.rgb * TopLayerData.EnvBRDF * SSRReductionFactor + StrataEnvLighting.rgb;
+		float3 SkyLighting = SSR.rgb * TopLayerEnvBRDF * SSRReductionFactor + StrataEnvLighting.rgb;
 
 		ApplyCloudVolumetricShadow(WorldPosition, SkyLighting);
 
diff --git a/Engine/Shaders/Private/Strata/StrataEnvironmentLighting.ush b/Engine/Shaders/Private/Strata/StrataEnvironmentLighting.ush
--- a/Engine/Shaders/Private/Strata/StrataEnvironmentLighting.ush
+++ b/Engine/Shaders/Private/Strata/StrataEnvironmentLighting.ush
@@ -19,7 +19,8 @@ float4 StrataEnvironmentLighting(
 	float TopLayerSpecularContributionFactor,
 	ByteAddressBuffer StrataDataBuffer,
 	FStrataAddressing StrataAddressing,
-	inout float SSRReductionFactor)
+	inout float SSRReductionFactor,
+	inout float3 TopLayerEnvBRDF)
 {
 	FLightAccumulator Out = (FLightAccumulator)0;
 
@@ -99,6 +100,12 @@ float4 StrataEnvironmentLighting(
 			SSRReductionFactor -= bIsTopLayer ? dot(BSDFThroughput, float3(1.0f / 3.0f, 1.0f / 3.0f, 1.0f / 3.0f)) * StrataEnvLight.SSRReduction : 0.0f;
 		}
 
+		if (bIsTopLayer)
+		{
+			// The specular path weight applied on SSR. It must account for throughput even for top surface because it also contains coverage.
+			TopLayerEnvBRDF += BSDFThroughput * StrataEnvLight.SpecularWeight;
+		}
+
 		const bool bNeedsSeparateSubsurfaceLightAccumulation = StrataEnvLight.bSubsurface;
 		LightAccumulator_AddSplit(Out, DiffuseSkyLighting, SpecularSkyLighting, DiffuseSkyLighting, CommonFactor, bNeedsSeparateSubsurfaceLightAccumulation);
 
