Standardizes {int,uint,float}16_t{,2,3,4} as well as as{int,uint,float}16()

diff --git a/Engine/Shaders/Private/TemporalAA/TAACommon.ush b/Engine/Shaders/Private/TemporalAA/TAACommon.ush
--- a/Engine/Shaders/Private/TemporalAA/TAACommon.ush
+++ b/Engine/Shaders/Private/TemporalAA/TAACommon.ush
@@ -449,18 +449,11 @@ taa_half3 QuantizeForFloatRenderTarget(taa_half3 Color, taa_half E, float3 Quant
 
 	#if 0
 		// NOP
-	#elif CONFIG_COMPILE_FP16 && 0
-	{
-		Error.x = asfloat16_t(asuint16_t(Error.x) & uint16_t(~0x03FF));
-		Error.y = asfloat16_t(asuint16_t(Error.y) & uint16_t(~0x03FF));
-		Error.z = asfloat16_t(asuint16_t(Error.z) & uint16_t(~0x03FF));
-	}
 	#elif CONFIG_COMPILE_FP16
 	{
-		// shader compilers haven't exposed asuint16_t() and asfloat16_t() :'(
-		Error.x = taa_half(asfloat(asuint(float(Error.x)) & ~0x007FFFFF));
-		Error.y = taa_half(asfloat(asuint(float(Error.y)) & ~0x007FFFFF));
-		Error.z = taa_half(asfloat(asuint(float(Error.z)) & ~0x007FFFFF));
+		uint16_t3 UError = asint16(Error);
+
+		Error = asfloat16(UError & uint16_t(~0x03FF));
 	}
 	#else
 	{
