Make sure some code section are excluded when using Strata fast path.
Fixed fastpath selection bits.

diff --git a/Engine/Shaders/Private/Strata/Strata.ush b/Engine/Shaders/Private/Strata/Strata.ush
--- a/Engine/Shaders/Private/Strata/Strata.ush
+++ b/Engine/Shaders/Private/Strata/Strata.ush
@@ -1434,6 +1434,17 @@ FStrataPixelHeader UnpackFastPathStrataHeaderIn(ByteAddressBuffer StrataBuffer,
 	return Out;
 }
 
+bool IsStrataFastPathCompatible(in FStrataBSDF BSDF)
+{
+	return
+		!BSDF_GETANISOTROPY(BSDF) &&
+		!BSDF_GETHASSCATTERING(BSDF) &&
+		!BSDF_GETHASSSSPROFILE(BSDF) &&
+		!BSDF_GETHASHAZINESS(BSDF) &&
+		!BSDF_GETHASTHINFILM(BSDF) &&
+		!BSDF_GETHASEDGECOLOR(BSDF);
+}
+
 void UnpackFastPathSlabBSDF(inout FStrataBSDF OutBSDF, uint2 PackedData01)
 {
 	float4 Data0 = UnpackColorGamma2ToLinearAlphaLinear(PackedData01.x);
@@ -1471,6 +1482,7 @@ FStrataBSDF UnpackFastPathStrataBSDFIn(ByteAddressBuffer StrataBuffer, inout FSt
 	BSDF_SETHASSSSPROFILE(OutBSDF, 0);
 	BSDF_SETHASHAZINESS(OutBSDF, 0);
 	BSDF_SETHASTHINFILM(OutBSDF, 0);
+	BSDF_SETHASEDGECOLOR(OutBSDF, 0);
 
 	return OutBSDF;
 }
diff --git a/Engine/Shaders/Private/Strata/StrataEnvironmentLighting.ush b/Engine/Shaders/Private/Strata/StrataEnvironmentLighting.ush
--- a/Engine/Shaders/Private/Strata/StrataEnvironmentLighting.ush
+++ b/Engine/Shaders/Private/Strata/StrataEnvironmentLighting.ush
@@ -88,6 +88,8 @@ float4 StrataEnvironmentLighting(
 			SpecularSkyLighting += BSDFThroughput * SpecularLuminance * StrataEnvLight.SpecularWeight * View.SkyLightColor.rgb * AmbientOcclusion * 
 				(bIsTopLayer ? TopLayerSpecularContributionFactor : 1.0f);
 		}
+
+#if STRATA_FASTPATH==0
 		if (BSDF_GETHASHAZINESS(BSDF))
 		{
 			float SkyAverageBrightness = 1.0f;
@@ -99,6 +101,7 @@ float4 StrataEnvironmentLighting(
 			// And we also try to make the transition smooth using Haziness
 			SSRReductionFactor -= bIsTopLayer ? dot(BSDFThroughput, float3(1.0f / 3.0f, 1.0f / 3.0f, 1.0f / 3.0f)) * StrataEnvLight.SSRReduction : 0.0f;
 		}
+#endif 
 
 		if (bIsTopLayer)
 		{
diff --git a/Engine/Shaders/Private/Strata/StrataEvaluation.ush b/Engine/Shaders/Private/Strata/StrataEvaluation.ush
--- a/Engine/Shaders/Private/Strata/StrataEvaluation.ush
+++ b/Engine/Shaders/Private/Strata/StrataEvaluation.ush
@@ -12,6 +12,12 @@
 
 
 
+#ifndef STRATA_FASTPATH 
+#define STRATA_FASTPATH 0
+#endif
+
+
+
 // STRATA_TODO put in a common file
 // Point lobe in off-specular peak direction
 float3 StrataGetOffSpecularPeakReflectionDir(float3 Normal, float3 ReflectionVector, float Roughness)
@@ -199,6 +205,9 @@ FStrataEvaluateResult StrataEvaluateBSDF(FStrataBSDFContext BSDFContext, FShadow
 			float HazeD = 0;
 			float HazeVis = 0;
 			float HazePDF = 0;
+
+#if STRATA_FASTPATH==0
+			BRANCH
 			if (bHasAnisotropy)
 			{
 				// Generalized microfacet specular
@@ -218,6 +227,7 @@ FStrataEvaluateResult StrataEvaluateBSDF(FStrataBSDFContext BSDFContext, FShadow
 				}
 			}
 			else
+#endif
 			{
 				// Special override for roughness for supporting area light integrator with Sphere/Tube/Disk light, which modifies/increase roughness.
 				const float a2 = BSDFContext.a2Override > 0 ? BSDFContext.a2Override : Pow4(SafeRoughnessX);
@@ -227,6 +237,7 @@ FStrataEvaluateResult StrataEvaluateBSDF(FStrataBSDFContext BSDFContext, FShadow
 				Vis = Vis_SmithJointApprox(a2, BSDFContext.Context.NoV, BSDFContext.SatNoL);
 				PDF = VisibleGGXPDF(BSDFContext.TangentV, BSDFContext.TangentH, a2);
 
+#if STRATA_FASTPATH==0
 				if (bHaziness)
 				{
 					const float a2Haze = BSDFContext.a2OverrideHazy > 0 ? BSDFContext.a2OverrideHazy : Pow4(ComputeHazyLobeRoughness(SafeRoughnessX, SLAB_HAZINESS(BSDFContext.BSDF)));
@@ -234,9 +245,11 @@ FStrataEvaluateResult StrataEvaluateBSDF(FStrataBSDFContext BSDFContext, FShadow
 					HazeVis = Vis_SmithJointApprox(a2Haze, BSDFContext.Context.NoV, BSDFContext.SatNoL);
 					HazePDF = VisibleGGXPDF(BSDFContext.TangentV, BSDFContext.TangentH, a2Haze);
 				}
+#endif
 			}
 
 			float3 F = 0;
+#if STRATA_FASTPATH==0
 			const bool bHasThinFilm = BSDF_GETHASTHINFILM(BSDFContext.BSDF);
 			BRANCH
 			if (!bHasThinFilm)
@@ -244,6 +257,7 @@ FStrataEvaluateResult StrataEvaluateBSDF(FStrataBSDFContext BSDFContext, FShadow
 				F = F_Schlick(F0, F90, BSDFContext.Context.VoH);
 			}
 			else
+#endif
 			{
 				float ThinFilmIOR		= SLAB_THINFILM_IOR(BSDFContext.BSDF);
 				float ThinFilmThickness = SLAB_THINFILM_THICKNESS(BSDFContext.BSDF);
@@ -265,6 +279,7 @@ FStrataEvaluateResult StrataEvaluateBSDF(FStrataBSDFContext BSDFContext, FShadow
 			break;
 		}
 
+#if STRATA_FASTPATH==0
 		case STRATA_BSDF_TYPE_SHEEN:
 		{
 			const float3 F0  = SHEEN_BASECOLOR(BSDFContext.BSDF);
@@ -309,6 +324,7 @@ FStrataEvaluateResult StrataEvaluateBSDF(FStrataBSDFContext BSDFContext, FShadow
 		//case STRATA_BSDF_TYPE_SINGLELAYERWATER:
 		//Nothing to do in this case because these BSDF are evaluated in other specialised passes.
 		//break;
+#endif
 	}
 
 	return Sample;
@@ -461,6 +477,7 @@ FStrataEnvLightResult StrataEvaluateForEnvLight(FStrataBSDFContext BSDFContext,
 				const float3 F90 = ComputeF90(F0, SLAB_EDGECOLOR(BSDFContext.BSDF), Metallic);
 
 				float3 FresnelWeight = 1;
+#if STRATA_FASTPATH==0
 				const bool bHasThinFilm = BSDF_GETHASTHINFILM(BSDFContext.BSDF);
 				if (bHasThinFilm)
 				{
@@ -469,9 +486,11 @@ FStrataEnvLightResult StrataEvaluateForEnvLight(FStrataBSDFContext BSDFContext,
 					FresnelWeight = F_ThinFilm(BSDFContext.Context.NoV, BSDFContext.Context.NoL, BSDFContext.Context.VoH, F0, F90, ThinFilmIOR, ThinFilmThickness);
 					F0 = 1.f;
 				}
+#endif
 
 				float SafeRoughnessX = MakeRoughnessSafe(SLAB_ROUGHNESSX(BSDFContext.BSDF));
 
+#if STRATA_FASTPATH==0
 				const bool bHasAnisotropy = BSDF_GETANISOTROPY(BSDFContext.BSDF);
 				if (bHasAnisotropy)
 				{
@@ -487,6 +506,7 @@ FStrataEnvLightResult StrataEvaluateForEnvLight(FStrataBSDFContext BSDFContext,
 					BSDFContext.SatNoL = saturate(BSDFContext.Context.NoL);
 					BSDFContext.R = 2 * dot(BSDFContext.V, BSDFContext.N) * BSDFContext.N - BSDFContext.V;
 				}
+#endif
 
 				float3 EvalEnvBRDF = FresnelWeight * EnvBRDFApprox(F0, F90, SafeRoughnessX, BSDFContext.SatNoV);
 
@@ -494,6 +514,7 @@ FStrataEnvLightResult StrataEvaluateForEnvLight(FStrataBSDFContext BSDFContext,
 				StrataEnvLightResult.SpecularWeight = EvalEnvBRDF;
 				StrataEnvLightResult.SpecularSafeRoughness = SafeRoughnessX;
 
+#if STRATA_FASTPATH==0
 				const bool bHasHaziness = BSDF_GETHASHAZINESS(BSDFContext.BSDF);
 				if (bHasHaziness)
 				{
@@ -510,10 +531,12 @@ FStrataEnvLightResult StrataEvaluateForEnvLight(FStrataBSDFContext BSDFContext,
 					StrataEnvLightResult.SpecularHazeWeight = lerp(0.0f, SPECULAR_HAZINESS_WEIGHT, HazinessFadeIn) * EnvBRDFApprox(F0, F90, SafeRoughnessHaze, BSDFContext.SatNoV);
 					StrataEnvLightResult.SpecularHazeSafeRoughness = SafeRoughnessHaze;
 				}
+#endif
 			}
 			break;
 		}
 
+#if STRATA_FASTPATH==0
 		case STRATA_BSDF_TYPE_SHEEN:
 		{
 			StrataEnvLightResult.DiffuseColor = SHEEN_BASECOLOR(BSDFContext.BSDF);;
@@ -553,6 +576,7 @@ FStrataEnvLightResult StrataEvaluateForEnvLight(FStrataBSDFContext BSDFContext,
 		//case STRATA_BSDF_TYPE_SINGLELAYERWATER:
 		//Nothing to do in this case because these BSDF are evaluated in other specialised passes.
 		//break;
+#endif
 	}
 
 	return StrataEnvLightResult;
diff --git a/Engine/Shaders/Private/Strata/StrataMaterialClassification.usf b/Engine/Shaders/Private/Strata/StrataMaterialClassification.usf
--- a/Engine/Shaders/Private/Strata/StrataMaterialClassification.usf
+++ b/Engine/Shaders/Private/Strata/StrataMaterialClassification.usf
@@ -63,12 +63,7 @@ void ComputePostProcessBSDFData(
 			TopLayer.BaseColor += BSDFCoverage * EnvBRDFApprox(F0, F90, SafeRoughness, saturate(ContextNoLight.NoV));
 		}
 
-		Classification.bIsSimple =
-			Classification.bIsSimple &&
-			!BSDF_GETANISOTROPY(BSDF) &&
-			!BSDF_GETHASSCATTERING(BSDF) &&
-			!BSDF_GETHASSSSPROFILE(BSDF) &&
-			!BSDF_GETHASHAZINESS(BSDF);
+		Classification.bIsSimple = Classification.bIsSimple && IsStrataFastPathCompatible(BSDF);
 
 		// Select the first/toppest layer with subsurface scattering
 		if (BSDF_GETHASSCATTERING(BSDF) && BSDFCoverageAvg > 0 && !SSS.bHasSubsurface)
