Removed duplicated SetHasEdgeColor

diff --git a/Engine/Shaders/Private/Strata/Strata.ush b/Engine/Shaders/Private/Strata/Strata.ush
--- a/Engine/Shaders/Private/Strata/Strata.ush
+++ b/Engine/Shaders/Private/Strata/Strata.ush
@@ -1482,7 +1482,6 @@ FStrataBSDF UnpackFastPathStrataBSDFIn(ByteAddressBuffer StrataBuffer, inout FSt
 	BSDF_SETHASSSSPROFILE(OutBSDF, 0);
 	BSDF_SETHASHAZINESS(OutBSDF, 0);
 	BSDF_SETHASTHINFILM(OutBSDF, 0);
-	BSDF_SETHASEDGECOLOR(OutBSDF, 0);
 
 	return OutBSDF;
 }
