Implements TotalRejection in Gen5TAA

diff --git a/Engine/Shaders/Private/TemporalAA/TAAUpdateHistory.usf b/Engine/Shaders/Private/TemporalAA/TAAUpdateHistory.usf
--- a/Engine/Shaders/Private/TemporalAA/TAAUpdateHistory.usf
+++ b/Engine/Shaders/Private/TemporalAA/TAAUpdateHistory.usf
@@ -147,6 +147,9 @@ void MainCS(
 	// Detect if HistoryBufferUV would be outside of the viewport.
 	bool bOffScreen = IsOffScreen(bCameraCut, PrevScreenPos, ParallaxRejectionMask);
 	
+	taa_half TotalRejection = bOffScreen ? 0.0 : saturate(LowFrequencyRejection * 4.0);
+
+
 	// Filter input scene color at predictor frequency.
 	taa_half3 FilteredInputColor;
 	taa_half3 InputMinColor;
@@ -167,7 +170,7 @@ void MainCS(
 			taa_half FinalInputToHistoryFactor = bOffScreen ? taa_half(1.0) : InputToHistoryFactor;
 		#else
 			float InputToHistoryFactor = float(HistoryInfo_ViewportSize.x * InputInfo_ViewportSizeInverse.x);
-			float FinalInputToHistoryFactor = bOffScreen ? float(1.0) : InputToHistoryFactor;
+			float FinalInputToHistoryFactor = lerp(1.0, InputToHistoryFactor, TotalRejection);
 		#endif
 
 		InputMinColor = taa_half(INFINITE_FLOAT);
@@ -527,17 +530,13 @@ void MainCS(
 			}
 			#endif
 
-
-			//BlendFinal *= InputPixelAlignement; //, 0.05);
-
 			// Responsive forces 1/4 of new frame.
 			BlendFinal = bIsResponsiveAAPixel ? taa_half(1.0/4.0) : BlendFinal;
 
-			FLATTEN
-			if (bOffScreen)
+			// Completely discard the history
 			{
-				PrevFallbackSampleCount = taa_half(0.0);
-				BlendFinal = taa_half(1.0);
+				PrevFallbackSampleCount *= TotalRejection;
+				BlendFinal = lerp(1.0, BlendFinal, TotalRejection);
 			}
 
 			FinalFallbackWeight = saturate(CurrentFrameSampleCount + PrevFallbackSampleCount);
