Point light SMRT:
- Implement cube map SMRT using iteration across two faces. While a ray can technically span 3, no significant corner-related artifacts were observed.
- Rafactor SMRT loop out to a macro/include-based quasi "template" and reimplement SMRT for all light types in terms of it
- Few bug fixes and improvements to occluder distance approximation and similar

diff --git a/Engine/Shaders/Private/VirtualShadowMaps/ProjectionDirectional.ush b/Engine/Shaders/Private/VirtualShadowMaps/ProjectionDirectional.ush
--- a/Engine/Shaders/Private/VirtualShadowMaps/ProjectionDirectional.ush
+++ b/Engine/Shaders/Private/VirtualShadowMaps/ProjectionDirectional.ush
@@ -8,42 +8,38 @@ ProjectionDirectional.ush:
 #include "../DeferredShadingCommon.ush"
 #include "../SceneTexturesCommon.ush"
 #include "../LightShaderParameters.ush"
+#include "../PathTracing/Utilities/PathTracingRandomSequence.ush"
 #include "PageAccessCommon.ush"
 #include "ProjectionCommon.ush"
+#include "SMRTCommon.ush"
 
-bool ShadowRayCastDirectional(
+
+struct FSMRTClipmapRayState
+{
+	int ClipmapId;			// Pre-offset virtual shadow map ID of specific clipmap level
+	int ClipmapIdAdjacent;	// Pre-offset virtual shadow map ID of specific clipmap level
+	float3 RayStartShadowTexelsZ;	
+	float3 RayStepShadowTexelsZ;
+	float4 AdjacentTexelScaleBias;
+	float2 AdjacentToBaseDepthScaleBias;
+};
+
+FSMRTClipmapRayState SMRTClipmapRayInitialize(
 	int VirtualShadowMapId,
-	float3 ViewPosition,
 	float3 RayStart,
-	float3 RayVector,	// RayEnd - RayStart
-	int NumSteps,
-	float StepOffset,
-	int DebugOutputType,
-	inout float3 DebugOutput,
-	inout float OccluderDistance)
-{	
+	float3 RayVector)
+{
 	const FVirtualShadowMapProjectionShaderData BaseProjectionData = GetVirtualShadowMapProjectionData(VirtualShadowMapId);
 
 	// The page at RayStart has a high chance of being mapped since that is close to the receiver point
 	// Thus we start our clipmap sampling around that page even as we walk through the depth space
-	// TODO: Branch out if beyond max clipmap range somewhere
 	const float ClipmapLevelFloat = CalcClipmapLevelFloat(BaseProjectionData, RayStart);	
 	const int ClipmapIndex = max(0, int(floor(ClipmapLevelFloat)) - BaseProjectionData.ClipmapLevel);
 	const float ClipmapLevelFrac = frac(ClipmapLevelFloat);
-
 	int ClipmapIndexAdjacent = ClipmapIndex + (ClipmapLevelFrac > 0.5f ? 1 : -1);
 	ClipmapIndexAdjacent = clamp(ClipmapIndexAdjacent, 0, BaseProjectionData.ClipmapLevelCount - 1);
 
-	int AbsoluteClipmapIndex = VirtualShadowMapId + ClipmapIndex;
-	int AbsoluteClipmapIndexAdjacent = VirtualShadowMapId + ClipmapIndexAdjacent;
-
-	// TODO: Experiment with perf of this, etc.
-	bool bSampleTwoClipmapLevels = true;
-#if 0
-	bSampleTwoClipmapLevels = (ClipmapIndexAdjacent != ClipmapIndex) && (abs(ClipmapLevelFrac - 0.5f) > 0.4f);
-#endif
-
-	const FVirtualShadowMapProjectionShaderData ProjectionData = GetVirtualShadowMapProjectionData(AbsoluteClipmapIndex);
+	const FVirtualShadowMapProjectionShaderData ProjectionData = GetVirtualShadowMapProjectionData(VirtualShadowMapId + ClipmapIndex);
 
 	const float3 RayStartUV = mul(float4(RayStart + ProjectionData.ShadowPreViewTranslation, 1.0f), ProjectionData.TranslatedWorldToShadowUVMatrix).xyz;
 	const float3 RayStepUV  = mul(float4(RayVector, 0.0f), ProjectionData.TranslatedWorldToShadowUVMatrix).xyz;
@@ -52,120 +48,94 @@ bool ShadowRayCastDirectional(
 	float3 RayStartShadowTexelsZ = float3(LevelDims, LevelDims, 1.0f) * RayStartUV.xyz;
 	float3 RayStepShadowTexelsZ  = float3(LevelDims, LevelDims, 1.0f) * RayStepUV.xyz;
 
-	// Z scale is consistent across clipmap levels
-	const float CompareScale = RayStepShadowTexelsZ.z * 2.0f / Pow2(NumSteps);
-	const float CompareBias = -CompareScale * (NumSteps - StepOffset + 1.0f);
-
 	// Work out scale/bias to convert to the adjacent UVs from the regular UVs
 	// NOTE: Since depth function is constant, probably worth refactoring the clipmap mappings to be scale/bias in general rather than
 	// a full transform matrix.
-	float2 AdjacentTexelScale = 1.0f;
-	float2 AdjacentTexelBias = 0.0f;
-	float AdjacentToBaseDepthScale = 1.0f;
-	float AdjacentToBaseDepthBias = 0.0f;
-	if (bSampleTwoClipmapLevels)
-	{
-		const FVirtualShadowMapProjectionShaderData ProjectionDataAdjacent = GetVirtualShadowMapProjectionData(AbsoluteClipmapIndexAdjacent);
-
-		AdjacentTexelScale = (ProjectionDataAdjacent.ShadowViewToClipMatrix._11_22 / ProjectionData.ShadowViewToClipMatrix._11_22);
+	const FVirtualShadowMapProjectionShaderData ProjectionDataAdjacent = GetVirtualShadowMapProjectionData(VirtualShadowMapId + ClipmapIndexAdjacent);
+
+	// TranslatedWorldToShadowViewMatrix is identical for each of the cascades; only the pre-view translation changes
+	float3 UVz         = mul(float4(ProjectionData.ShadowPreViewTranslation, 1.0f), ProjectionData.TranslatedWorldToShadowUVMatrix).xyz;
+	float3 UVzAdjacent = mul(float4(ProjectionDataAdjacent.ShadowPreViewTranslation, 1.0f), ProjectionDataAdjacent.TranslatedWorldToShadowUVMatrix).xyz;
+
+	float2 AdjacentTexelScale = (ProjectionDataAdjacent.ShadowViewToClipMatrix._11_22 / ProjectionData.ShadowViewToClipMatrix._11_22);
+	float2 AdjacentTexelBias = LevelDims * (UVzAdjacent.xy - AdjacentTexelScale * UVz.xy);
+
+	// Always rescale adjacent samples back to the same depth mapping as the regular clipmap level to ensure
+	// precision is maintained and DepthHistory remains well defined.
+	float AdjacentToBaseDepthScale = (ProjectionData.ShadowViewToClipMatrix._33 / ProjectionDataAdjacent.ShadowViewToClipMatrix._33);
+	float AdjacentToBaseDepthBias = (UVz.z - AdjacentToBaseDepthScale * UVzAdjacent.z);
+	
+	FSMRTClipmapRayState Result;
+	Result.ClipmapId = VirtualShadowMapId + ClipmapIndex;
+	Result.ClipmapIdAdjacent = VirtualShadowMapId + ClipmapIndexAdjacent;
+	Result.RayStartShadowTexelsZ = RayStartShadowTexelsZ;
+	Result.RayStepShadowTexelsZ = RayStepShadowTexelsZ;
+	Result.AdjacentTexelScaleBias = float4(AdjacentTexelScale, AdjacentTexelBias);
+	Result.AdjacentToBaseDepthScaleBias = float2(AdjacentToBaseDepthScale, AdjacentToBaseDepthBias);
+	return Result;
 
-		// TranslatedWorldToShadowViewMatrix is identical for each of the cascades; only the pre-view translation changes
-		float3 UVz         = mul(float4(ProjectionData.ShadowPreViewTranslation, 1.0f), ProjectionData.TranslatedWorldToShadowUVMatrix).xyz;
-		float3 UVzAdjacent = mul(float4(ProjectionDataAdjacent.ShadowPreViewTranslation, 1.0f), ProjectionDataAdjacent.TranslatedWorldToShadowUVMatrix).xyz;
-
-		AdjacentTexelBias = (UVzAdjacent.xy - AdjacentTexelScale * UVz.xy);
-		AdjacentTexelBias *= float(CalcLevelDimsTexels(0));
+}
 
-		// Always rescale adjacent samples back to the same depth mapping as the regular clipmap level to ensure
-		// precision is maintained and DepthHistory remains well defined.
-		AdjacentToBaseDepthScale = (ProjectionData.ShadowViewToClipMatrix._33 / ProjectionDataAdjacent.ShadowViewToClipMatrix._33);
-		AdjacentToBaseDepthBias = (UVz.z - AdjacentToBaseDepthScale * UVzAdjacent.z);
-	}
+FSMRTSample SMRTFindSample(inout FSMRTClipmapRayState RayState, float SampleTime)
+{
+	FSMRTSample Sample;
+	Sample.bValid = false;
+	
+	const float3 SampleTexelsZ = RayState.RayStartShadowTexelsZ + RayState.RayStepShadowTexelsZ * SampleTime;
 
-	float FirstHitTime = 2;
-	float DepthHistory = -1;
-	int SampleMissCount = 0;
+	uint2 PhysicalTexelAddress = uint2(0, 0);
+	bool bValid = VirtualToPhysicalTexel(RayState.ClipmapId, 0, uint2(SampleTexelsZ.xy), PhysicalTexelAddress);
 
-	for (int i = 0; i < NumSteps; i++)
+	float2 ToBaseDepthScaleBias = float2(1.0f, 0.0f);
+	if (!bValid)
 	{
-		const float SampleTime = Pow2(saturate((-1.0f / NumSteps) * i + 1.0f - (StepOffset - 0.5f) / NumSteps));
-		const float3 SampleTexelsZ = RayStartShadowTexelsZ + RayStepShadowTexelsZ * SampleTime;
-
-		uint2 PhysicalTexelAddress = uint2(0, 0);
-		bool bValid = VirtualToPhysicalTexel(AbsoluteClipmapIndex, 0, uint2(SampleTexelsZ.xy), PhysicalTexelAddress);
-		
-		float ToBaseDepthScale = 1.0f;
-		float ToBaseDepthBias = 0.0f;
-		if (bSampleTwoClipmapLevels && !bValid)
-		{
-			// Replace with adjacent level sample
-			float2 AdjacentTexels = SampleTexelsZ.xy * AdjacentTexelScale + AdjacentTexelBias;
-			bValid = VirtualToPhysicalTexel(AbsoluteClipmapIndexAdjacent, 0, uint2(AdjacentTexels), PhysicalTexelAddress);
-			ToBaseDepthScale = AdjacentToBaseDepthScale;
-			ToBaseDepthBias  = AdjacentToBaseDepthBias;
-		}
-
-		if (bValid)
-		{
-			const float SampleDepth = ToBaseDepthScale * SampleVirtualSmPhysicalDepth(PhysicalTexelAddress) + ToBaseDepthBias;
-			const float ReferenceDepth = SampleTexelsZ.z;
-			if (DepthHistory < 0)
-			{
-				// First valid sample we've seen. Do a regular depth compare.
-				DepthHistory = SampleDepth;
-				FirstHitTime = SampleDepth > ReferenceDepth ? SampleTime : 2;
-			}
-			else
-			{
-				const float CompareTolerance = abs(CompareScale * i + CompareBias);
-
-				bool bBehind = (SampleDepth - ReferenceDepth) > (2.0f * CompareTolerance);
-				float DepthSample = bBehind ? DepthHistory : SampleDepth;
-				DepthHistory = DepthSample;
-
-				float DepthDiff = ReferenceDepth - DepthSample;
-				bool bHit = abs(DepthDiff + CompareTolerance) < CompareTolerance;
-
-				if (bHit)
-				{
-					FirstHitTime = SampleTime;
-				}
-			}
-		}
-		else
-		{
-			++SampleMissCount;
-		}
+		// Replace with adjacent level sample
+		float2 AdjacentTexels = SampleTexelsZ.xy * RayState.AdjacentTexelScaleBias.xy + RayState.AdjacentTexelScaleBias.zw;
+		bValid = VirtualToPhysicalTexel(RayState.ClipmapIdAdjacent, 0, uint2(AdjacentTexels), PhysicalTexelAddress);
+		ToBaseDepthScaleBias = RayState.AdjacentToBaseDepthScaleBias;
 	}
 
-	bool bValidHit = FirstHitTime < 2;
-	if (bValidHit)
+	if (bValid)
 	{
-		OccluderDistance = length(RayVector) * FirstHitTime;
+		float SampleDepth = ToBaseDepthScaleBias.x * SampleVirtualSmPhysicalDepth(PhysicalTexelAddress) + ToBaseDepthScaleBias.y;
+		Sample.bValid = true;
+		Sample.SampleDepth = SampleDepth;
+		Sample.ReferenceDepth = SampleTexelsZ.z;
+		Sample.CompareToleranceScale = RayState.RayStepShadowTexelsZ.z;
 	}
 
-	float OutputDebugSignal = bValidHit ? 0.0f : 1.0f;
-	if (DebugOutputType == 1)	// Invalid page
-	{
-		DebugOutput = lerp(OutputDebugSignal.xxx, lerp(float3(0, 1, 0), float3(1, 0, 1), saturate(SampleMissCount / 4.0f)), 0.75f);
-	}
-	else if (DebugOutputType == 2)	// ClipmapIndex
-	{
-		DebugOutput = lerp(OutputDebugSignal.xxx, VirtualShadowMapDebugIntToColor(ClipmapIndex), 0.75f);
-	}
-	else if (DebugOutputType == 3)	// Two levels
-	{
-		DebugOutput = bSampleTwoClipmapLevels ? float3(1, 0, 0) : float3(0, 1, 0);
-	}
+	return Sample;
+}
+
+float SMRTComputeOccluderDistance(FSMRTClipmapRayState RayState, float SampleDepth)
+{
+	float4x4 ViewToClip = GetVirtualShadowMapProjectionData(RayState.ClipmapId).ShadowViewToClipMatrix;
 
-	return bValidHit;
+	// No perspective projection, so simple difference gets us the distance
+	float OccluderViewZ = (SampleDepth                      - ViewToClip._43) / ViewToClip._33;
+	float ReceiverViewZ = (RayState.RayStartShadowTexelsZ.z - ViewToClip._43) / ViewToClip._33;
+	float Result = ReceiverViewZ - OccluderViewZ;
+	return max(VIRTUAL_SHADOW_MAP_MIN_OCCLUDER_DISTANCE, Result);
 }
 
-#define USE_RANDOM_SEQUENCE 1
+// Instantiate SMRTRayCast for FSMRTClipmapRayState
+#define SMRT_TEMPLATE_RAY_STRUCT FSMRTClipmapRayState
+#include "SMRTTemplate.ush"
+#undef SMRT_TEMPLATE_RAY_STRUCT
 
-#if USE_RANDOM_SEQUENCE
-#include "/Engine/Private/PathTracing/Utilities/PathTracingRandomSequence.ush"
-#endif
+
+
+FSMRTResult ShadowRayCastDirectional(
+	int VirtualShadowMapId,
+	float3 ViewPosition,
+	float3 RayStart,
+	float3 RayVector,	// RayEnd - RayStart
+	int NumSteps,
+	float StepOffset)
+{	
+	FSMRTClipmapRayState RayState = SMRTClipmapRayInitialize(VirtualShadowMapId, RayStart, RayVector);
+	return SMRTRayCast(RayState, NumSteps, StepOffset);
+}
 
 float3 GetRandomDirectionalLightRayDir(FLightShaderParameters Light, float2 E)
 {
@@ -205,41 +175,47 @@ float TraceDirectional(
 	// Too low values will greatly restrict how large penumbras can be in screen space.
 	float RayLength = RayLengthScale * DistanceFromViewOrigin;
 
-#if USE_RANDOM_SEQUENCE
-	uint Seed = uint(SvPosition.y * View.BufferSizeAndInvSize.x + SvPosition.x);
-	RandomSequence RandSequence;
-	RandomSequence_Initialize(RandSequence, Seed, View.StateFrameIndex);
-#else
-	uint2 SobolBase = SobolPixel(SvPosition.xy);
-	uint2 SobolFrame = SobolIndex(SobolBase, View.StateFrameIndexMod8, 3);
-#endif
-
 	float StepOffset = InterleavedGradientNoise(SvPosition.xy, View.StateFrameIndexMod8);
 
 	float RayMissFactor = 0;
 	for (int i = 0; i < RayCount; ++i)
 	{
-#if USE_RANDOM_SEQUENCE
+		RandomSequence RandSequence;
+		RandomSequence_Initialize(RandSequence, uint2(SvPosition.xy), i, View.StateFrameIndexMod8, RayCount);
+#if 1
 		float2 E = RandomSequence_GenerateSample2D(RandSequence);
 #else
 		float2 E = float2(SobolIndex(SobolFrame, i << 3)) / 0x10000;
 #endif
 
 		float3 RayDir = GetRandomDirectionalLightRayDir(Light, E);
-		float3 RayStart = WorldPosition;// + RayDir * RayStartOffset; TODO!
+		// TODO, move the clipmap selection outside/before this offset
+		float3 RayStart = WorldPosition + RayDir * RayStartOffset;
 
-		bool bHit = ShadowRayCastDirectional(
+		FSMRTResult Result = ShadowRayCastDirectional(
 			VirtualShadowMapId,
 			ViewPosition,
 			RayStart,
 			RayDir * RayLength,
 			SamplesPerRay,
-			StepOffset,
-			DebugOutputType,
-			DebugOutput,
-			InOutHitDistance);
+			StepOffset);
+
+		float OutputDebugSignal = Result.bValidHit ? 0.0f : 1.0f;
+		if (DebugOutputType == 1)	// Invalid page
+		{
+			DebugOutput = lerp(
+				OutputDebugSignal.xxx,
+				lerp( float3( 0, 1, 0 ), float3( 1, 0, 1 ), Result.SampleMissCount / 4.0f ),
+				0.75f );
+		}
+
+		if (Result.bValidHit)
+		{
+			// TODO: min over all the rays if we keep denoiser with > 1 spp?
+			InOutHitDistance = Result.OccluderDistance;
+		}
 
-		RayMissFactor += bHit ? 0 : 1;
+		RayMissFactor += Result.bValidHit ? 0 : 1;
 	}
 	RayMissFactor /= float(RayCount);
 
diff --git a/Engine/Shaders/Private/VirtualShadowMaps/ProjectionSpot.ush b/Engine/Shaders/Private/VirtualShadowMaps/ProjectionSpot.ush
--- a/Engine/Shaders/Private/VirtualShadowMaps/ProjectionSpot.ush
+++ b/Engine/Shaders/Private/VirtualShadowMaps/ProjectionSpot.ush
@@ -11,156 +11,213 @@ ProjectionSpot.ush:
 #include "../PathTracing/Utilities/PathTracingRandomSequence.ush"
 #include "PageAccessCommon.ush"
 #include "ProjectionCommon.ush"
+#include "SMRTCommon.ush"
 
 #ifndef SMRT_ADAPTIVE_RAY_COUNT
 #define SMRT_ADAPTIVE_RAY_COUNT 0
 #endif
 
-/** Return float multiplier to scale RayStepScreen such that it clip it right at the edge of the screen. */
-float GetStepScreenFactorToClipAtScreenEdge(float2 RayStartScreen, float2 RayStepScreen)
-{
-	// Computes the scale down factor for RayStepScreen required to fit on the X and Y axis in order to clip it in the viewport
-	const float RayStepScreenInvFactor = 0.5 * length(RayStepScreen);
-	const float2 S = 1 - max(abs(RayStepScreen + RayStartScreen * RayStepScreenInvFactor) - RayStepScreenInvFactor, 0.0f) / abs(RayStepScreen);
-
-	// Rescales RayStepScreen accordingly
-	const float RayStepFactor = min(S.x, S.y) / RayStepScreenInvFactor;
 
-	return RayStepFactor;
-}
+struct FSMRTSingleRayState
+{
+	int VirtualShadowMapId;
+	float3 RayStartUVz;	
+	float3 RayStepUVz;
+};
 
-bool ShadowRayCastSpotlight(
+FSMRTSingleRayState SMRTSingleRayInitialize(
 	int VirtualShadowMapId,
 	float3 RayStart,
-	float3 RayEnd,
-	int NumSteps,
-	float StepOffset,
-	int DebugOutputType,
-	inout float3 DebugOutput,
-	inout float OccluderDistance)
+	float3 RayEnd)
 {
-	FVirtualShadowMapProjectionShaderData BaseProjectionData = GetVirtualShadowMapProjectionData(VirtualShadowMapId);
+	FVirtualShadowMapProjectionShaderData ProjectionData = GetVirtualShadowMapProjectionData(VirtualShadowMapId);
 
-	// Spot light
-	float4 RayStartUVz = mul(float4(RayStart + BaseProjectionData.ShadowPreViewTranslation, 1), BaseProjectionData.TranslatedWorldToShadowUVMatrix).xyzw;
-	float4 RayEndUVz   = mul(float4(RayEnd   + BaseProjectionData.ShadowPreViewTranslation, 1), BaseProjectionData.TranslatedWorldToShadowUVMatrix).xyzw;
+	float4 RayStartUVz = mul(float4(RayStart + ProjectionData.ShadowPreViewTranslation, 1), ProjectionData.TranslatedWorldToShadowUVMatrix).xyzw;
+	float4 RayEndUVz   = mul(float4(RayEnd   + ProjectionData.ShadowPreViewTranslation, 1), ProjectionData.TranslatedWorldToShadowUVMatrix).xyzw;
 
-	if (RayEndUVz.w <= 0)
-	{
-		RayEndUVz = lerp(RayStartUVz, RayEndUVz, -RayStartUVz.w / (RayEndUVz.w - RayStartUVz.w));
-	}
+	// NOTE: We assume by construction that ray ends are not behind the light near plane as the warping
+	// due to SMRT ray tests gets severe long before the rays would otherwise clip at that plane.
+	// If clipping is ever necessary, it must be done to the ray endpoints in world space so that lights that
+	// need to walk multiple faces/clipmap levels in lock step can do it consistently.
 
 	RayStartUVz.xyz = RayStartUVz.xyz / RayStartUVz.w;
 	RayEndUVz.xyz = RayEndUVz.xyz / RayEndUVz.w;
-
 	float3 RayStepUVz = RayEndUVz.xyz - RayStartUVz.xyz;
 
-	/*
-#if 0
-	float ClipToEdgeScale = GetStepScreenFactorToClipAtScreenEdge( RayStartScreen.xy, RayStepScreen.xy );
-	RayStepScreen *= saturate( ClipToEdgeScale );
-#else
-	if( any( abs( RayEndScreen.xy ) > 1 ) )
+	FSMRTSingleRayState Result;
+	Result.VirtualShadowMapId = VirtualShadowMapId;
+	Result.RayStartUVz = RayStartUVz.xyz;
+	Result.RayStepUVz = RayStepUVz.xyz;
+	return Result;
+}
+
+FSMRTSample SMRTFindSample(inout FSMRTSingleRayState RayState, float SampleTime)
+{
+	FSMRTSample Sample;
+	Sample.bValid = false;
+	float3 UVz = RayState.RayStartUVz.xyz + RayState.RayStepUVz.xyz * SampleTime;
+
+	if (all(UVz.xy == saturate(UVz.xy)))
 	{
-		const float Epsilon = 1e-4;
-		float2 Replacement = RayStepScreen.xy > 0 ? Epsilon : -Epsilon;
-		RayStepScreen.xy = abs( RayStepScreen.xy ) < Epsilon ? Replacement : RayStepScreen.xy;
-
-		float2 InvDir = rcp( RayStepScreen.xy );
-		float2 Center = -RayStartScreen.xy * InvDir;
-		float2 MinIntersection = Center - abs( InvDir );
-		float2 MaxIntersection = Center + abs( InvDir );
-
-		float Time0 = max( MinIntersection.x, MinIntersection.y );
-		float Time1 = min( MaxIntersection.x, MaxIntersection.y );
-		RayStepScreen *= Time1;
+		FVirtualSmSample SmSample = SampleVirtualSm(RayState.VirtualShadowMapId, UVz.xy);
+		if (SmSample.bValid)
+		{
+			Sample.bValid = true;
+			Sample.SampleDepth = SmSample.Depth;
+			Sample.ReferenceDepth = UVz.z;
+			Sample.CompareToleranceScale = RayState.RayStepUVz.z;
+		}
 	}
-#endif
-*/
 
-	// Normalize 2D length
-	const float MaxPenumbra = 0.01;
-	RayStepUVz.xyz *= saturate(MaxPenumbra / length(RayStepUVz.xy));
+	return Sample;
+}
+
+float SMRTComputeOccluderDistance(FSMRTSingleRayState RayState, float SampleDepth)
+{
+	// Avoid keeping around the registers
+	float4x4 ViewToClip = GetVirtualShadowMapProjectionData(RayState.VirtualShadowMapId).ShadowViewToClipMatrix;
+
+	// View space Z delta is a good enough stand-in for euclidean distance for now given our bent ray tests and so on
+	float OccluderViewZ = ViewToClip._43 / (SampleDepth            - ViewToClip._33);
+	float ReceiverViewZ = ViewToClip._43 / (RayState.RayStartUVz.z - ViewToClip._33);		
+	float Result = ReceiverViewZ - OccluderViewZ;	
+	return max(VIRTUAL_SHADOW_MAP_MIN_OCCLUDER_DISTANCE, Result);
+}
 
-	const float CompareScale = RayStepUVz.z * 2.0 / Pow2(NumSteps);
-	const float CompareBias = -CompareScale * (NumSteps - StepOffset + 1.0f);
+// Instantiate SMRTRayCast for FSMRTSingleRayState
+#define SMRT_TEMPLATE_RAY_STRUCT FSMRTSingleRayState
+#include "SMRTTemplate.ush"
+#undef SMRT_TEMPLATE_RAY_STRUCT
 
-	float DepthHistory = -1;
-	int SampleMissCount = 0;
 
-	bool bValidHit = false;
-	for (int i = 0; i <= NumSteps; i++)
+struct FSMRTTwoCubeFaceRayState
+{
+	FSMRTSingleRayState Face0;
+	FSMRTSingleRayState Face1;
+	bool bSampleInFace0;
+};
+
+FSMRTTwoCubeFaceRayState SMRTTwoCubeFaceRayInitialize(
+	int VirtualShadowMapId0,
+	int VirtualShadowMapId1,
+	float3 RayStart,
+	float3 RayEnd)
+{
+	FSMRTTwoCubeFaceRayState Result;
+	Result.Face0 = SMRTSingleRayInitialize(VirtualShadowMapId0, RayStart, RayEnd);
+	Result.Face1 = SMRTSingleRayInitialize(VirtualShadowMapId1, RayStart, RayEnd);
+	Result.bSampleInFace0 = false;
+	return Result;
+}
+
+FSMRTSample SMRTFindSample(inout FSMRTTwoCubeFaceRayState RayState, float SampleTime)
+{
+	FSMRTSample Sample = SMRTFindSample(RayState.Face0, SampleTime);
+	RayState.bSampleInFace0 = Sample.bValid;
+	if (!RayState.bSampleInFace0)
 	{
-		//SampleTime = Pow2( 1.0 - (float)i / NumSteps ) * RayScale;
-		float SampleTime = Pow2(saturate((-1.0f / NumSteps) * i + 1.0f - (StepOffset - 0.5f) / NumSteps));
-		float3 SampleUVz = RayStartUVz.xyz + RayStepUVz.xyz * SampleTime;
-		FVirtualSmSample SmSample = SampleVirtualSm(VirtualShadowMapId, SampleUVz.xy);
+		Sample = SMRTFindSample(RayState.Face1, SampleTime);
+	}
+	return Sample;
+}
 
-		if (SmSample.bValid)
-		{
-			const float SampleDepth = SmSample.Depth;
-			const float ReferenceDepth = SampleUVz.z;
-			if (DepthHistory < 0)
-			{
-				// First valid sample we've seen. Do a regular depth compare.
-				DepthHistory = SampleDepth;
-				if( SampleDepth > ReferenceDepth )
-				{
-					float ReceiverDistance = length( RayStart + BaseProjectionData.ShadowPreViewTranslation );
-					OccluderDistance = ComputeOccluderDistancePerspective(
-						BaseProjectionData.ShadowViewToClipMatrix,
-						SmSample.Depth,
-						SampleUVz.z,
-						ReceiverDistance);
-
-					bValidHit = true;
-					break;
-				}
-			}
-			else
-			{
-				const float CompareTolerance = abs(CompareScale * i + CompareBias);
+float SMRTComputeOccluderDistance(FSMRTTwoCubeFaceRayState RayState, float SampleDepth)
+{
+	if (RayState.bSampleInFace0)
+	{
+		return SMRTComputeOccluderDistance(RayState.Face0, SampleDepth);
+	}
+	else
+	{
+		return SMRTComputeOccluderDistance(RayState.Face1, SampleDepth);
+	}
+}
 
-				bool bBehind = (SampleDepth - ReferenceDepth) > (2.0f * CompareTolerance);
-				float DepthSample = bBehind ? DepthHistory : SampleDepth;
-				DepthHistory = DepthSample;
+// Instantiate SMRTRayCast for FSMRTTwoCubeFaceRayState
+#define SMRT_TEMPLATE_RAY_STRUCT FSMRTTwoCubeFaceRayState
+#include "SMRTTemplate.ush"
+#undef SMRT_TEMPLATE_RAY_STRUCT
 
-				float DepthDiff = ReferenceDepth - DepthSample;
-				bool bHit = abs(DepthDiff + CompareTolerance) < CompareTolerance;
-				if (bHit)
-				{
-					OccluderDistance = length(RayEnd - RayStart) * SampleTime;
 
-					bValidHit = true;
-					break;
-				}
-			}
-		}
-		else
-		{
-			++SampleMissCount;
-		}
-	}
+FSMRTResult ShadowRayCastSpotLight(
+	int VirtualShadowMapId,
+	float3 RayStart,
+	float3 RayEnd,
+	int NumSteps,
+	float StepOffset)
+{
+	FSMRTSingleRayState RayState = SMRTSingleRayInitialize(VirtualShadowMapId, RayStart, RayEnd);
+	return SMRTRayCast(RayState, NumSteps, StepOffset);
+}
+
+FSMRTResult ShadowRayCastPointLight(
+	int VirtualShadowMapId,
+	float3 RayStart,
+	float3 RayEnd,
+	int NumSteps,
+	float StepOffset)
+{
+	FVirtualShadowMapProjectionShaderData BaseProjectionData = GetVirtualShadowMapProjectionData(VirtualShadowMapId);
 
-	float OutputDebugSignal = bValidHit ? 0.0f : 1.0f;
-	if (DebugOutputType == 1)	// Invalid page
+	uint RayStartFace = GetCubeFace(-(RayStart + BaseProjectionData.ShadowPreViewTranslation));
+	uint RayEndFace   = GetCubeFace(-(RayEnd   + BaseProjectionData.ShadowPreViewTranslation));
+		
+	FSMRTResult Result;	
+	if (RayStartFace != RayEndFace)
+	{
+		FSMRTTwoCubeFaceRayState RayState = SMRTTwoCubeFaceRayInitialize(
+			VirtualShadowMapId + RayStartFace,
+			VirtualShadowMapId + RayEndFace,
+			RayStart, RayEnd);
+		Result = SMRTRayCast(RayState, NumSteps, StepOffset);
+	}
+	else
 	{
-		DebugOutput = lerp(OutputDebugSignal.xxx, lerp(float3(0, 1, 0), float3(1, 0, 1), saturate(SampleMissCount / 4.0f)), 0.75f);
+		// Fast path: ray stays on a single cube map face (effectively a spot light)
+		FSMRTSingleRayState RayState = SMRTSingleRayInitialize(VirtualShadowMapId + RayEndFace, RayStart, RayEnd);
+		Result = SMRTRayCast(RayState, NumSteps, StepOffset);
 	}
 
-	return bValidHit;
+	return Result;
 }
 
+float3 GetRandomPointLightRayVector(FLightShaderParameters LightParameters, float3 WorldPosition, float2 E)
+{
+	float3 LightDirection = LightParameters.Position - WorldPosition;
+	float RayLength2 = dot(LightDirection, LightDirection);
+	float Radius2 = LightParameters.SourceRadius * LightParameters.SourceRadius;
+
+	// Clamping this instead of handling sampling point inside light radius properly is good enough for now
+	float SinThetaMax2 = saturate(Radius2 / RayLength2);
+
+	float4 DirAndPdf = UniformSampleConeConcentricRobust(E, SinThetaMax2);
+	float CosTheta = DirAndPdf.z;
+	float SinTheta2 = 1.0 - CosTheta * CosTheta;
 
-float TraceSpotLight(
+	// Project ray direction to world-space, such that z-axis aligns with LightDirection
+	float RayLength = sqrt(RayLength2);
+	LightDirection *= rcp(RayLength + 1e-4);
+	float3 RayDirection = TangentToWorld(DirAndPdf.xyz, LightDirection);
+
+	// Clip length to closest intersection with the sphere
+	//RayLength = RayLength * (CosTheta - sqrt(max(SinThetaMax2 - SinTheta2, 0.0)));
+
+	return RayLength * RayDirection;
+}
+
+
+
+float TraceLocalLight(
 	int VirtualShadowMapId,
 	FLightShaderParameters Light,
+	uint LightType,
 	float4 SvPosition,
 	const float SceneDepth,
 	float3 TranslatedWorldPosition,
 	float RayStartOffset,
 	int RayCount,
 	int SamplesPerRay,
+	float CotMaxRayAngleFromLight,
 	int DebugOutputType,
 	inout float3 DebugOutput,
 	inout float InOutHitDistance)
@@ -172,6 +229,9 @@ float TraceSpotLight(
 	float DistSqr = dot(ToLight, ToLight);
 	float Falloff = rcp(DistSqr + 1);
 
+	float DistToLight = sqrt(DistSqr);
+	float3 L = ToLight / DistToLight;	
+
 	float SinAlphaSqr = saturate(Pow2(Light.SourceRadius) * Falloff);
 	float SinAlpha = sqrt(SinAlphaSqr);
 	bool bBackFace = dot(GBufferData.WorldNormal, normalize(ToLight)) < -SinAlpha;
@@ -189,47 +249,92 @@ float TraceSpotLight(
 		for( ; i < uint(RayCount); i++ )
 		{
 			RandomSequence RandSequence;
-			RandomSequence_Initialize( RandSequence, uint2( SvPosition.xy ), i, View.StateFrameIndexMod8, RayCount );
+			RandomSequence_Initialize( RandSequence, uint2( SvPosition.xy ), i, View.StateFrameIndex, RayCount );
 
 #if 1
-			float2 E = RandomSequence_GenerateSample2D(RandSequence);
+			float2 E = RandomSequence_GenerateSample2D( RandSequence );
 #else
-			// NOTE: 1-X to avoid 0/0 issues with UniformSampleDiskConcentric
 			float2 E = float2( SobolIndex( SobolFrame, i << 3 ) ) / 0x10000;
 #endif
 
-			// NOTE: Avoid 0/0 issues with UniformSampleDiskConcentric
-			float2 LightUV = UniformSampleDiskConcentricApprox( E ) * Light.SourceRadius;
+			float3 RayEnd = float3(0, 0, 0);
+			if (LightType == LIGHT_TYPE_SPOT)
+			{
+				float2 LightUV = UniformSampleDiskConcentric( E ) * Light.SourceRadius;
 
-			float3 RayEnd =
-				Light.Position +
-				LightUV.x * Light.Tangent +
-				LightUV.y * cross( Light.Direction, Light.Tangent );
+				RayEnd =
+					Light.Position +
+					LightUV.x * Light.Tangent +
+					LightUV.y * cross( Light.Direction, Light.Tangent );
+			}
+			else
+			{
+				// Point light
+				RayEnd = WorldPosition + GetRandomPointLightRayVector( Light, WorldPosition, E );
+			}
 
 			float3 Dir = normalize( RayEnd - WorldPosition );
 			float3 RayStart = WorldPosition + Dir * RayStartOffset;
 
-			RayEnd = lerp( RayStart, RayEnd, 0.75 );
+			// The further afield our rays go the poorer our approximation is as the "bend" due to our testing
+			// against a shadow map instead of along the ray increases. Thus we avoid going all the way to the light
+			// where the projeciton becomes extreme.
+			RayEnd = lerp( RayStart, RayEnd, 0.75f );
+
+			{
+				// Clip ray to max angle to limit penumbra sizes
+				float CosLambda = dot( L, Dir );
+				float SinLambda = sqrt( 1.0f - CosLambda * CosLambda );
+				float MaxRayLength = DistToLight / ( CosLambda + CotMaxRayAngleFromLight * SinLambda );
+				float ClampedLength = saturate( MaxRayLength / length( RayEnd - RayStart ) );
 
-			bool bHit = ShadowRayCastSpotlight(
-				VirtualShadowMapId,
-				RayStart,
-				RayEnd,
-				SamplesPerRay,
-				StepOffset,
-				DebugOutputType,
-				DebugOutput,
-				InOutHitDistance);
+				RayEnd = lerp( RayStart, RayEnd, ClampedLength );
+			}
 
-			RayMissFactor += bHit ? 0 : 1;
+			FSMRTResult Result;
+			if ( LightType == LIGHT_TYPE_SPOT )
+			{
+				Result = ShadowRayCastSpotLight(
+					VirtualShadowMapId,
+					RayStart,
+					RayEnd,
+					SamplesPerRay,
+					StepOffset);
+			}
+			else
+			{
+				Result = ShadowRayCastPointLight(
+					VirtualShadowMapId,
+					RayStart,
+					RayEnd,
+					SamplesPerRay,
+					StepOffset);
+			}
+
+			float OutputDebugSignal = Result.bValidHit ? 0.0f : 1.0f;
+			if (DebugOutputType == 1)	// Invalid page
+			{
+				DebugOutput = lerp(	
+					OutputDebugSignal.xxx,
+					lerp( float3( 0, 1, 0 ), float3( 1, 0, 1 ), Result.SampleMissCount / 4.0f ),
+					0.75f );
+			}
+
+			if ( Result.bValidHit )
+			{
+				// TODO: Min over all the rays with > 1 spp?
+				InOutHitDistance = Result.OccluderDistance;
+			}
+
+			RayMissFactor += Result.bValidHit ? 0 : 1;
 			
-#if SMRT_ADAPTIVE_RAY_COUNT
-#if COMPILER_SUPPORTS_WAVE_VOTE
+#if SMRT_ADAPTIVE_RAY_COUNT && COMPILER_SUPPORTS_WAVE_VOTE
 			if( i == 0 )
 			{
+				bool bHit = Result.bValidHit;
 				if( bHit )
 				{
-					float WorldBluringRadius = ComputeLightSampleWorldBluringRadius( TranslatedWorldPosition, LIGHT_TYPE_SPOT, Light, InOutHitDistance );
+					float WorldBluringRadius = ComputeLightSampleWorldBluringRadius( TranslatedWorldPosition, LIGHT_TYPE_SPOT, Light, Result.OccluderDistance );
 					float PixelBluringRadius = ProjectWorldDistanceToPixelDistance( SceneDepth, WorldBluringRadius );
 					bHit = PixelBluringRadius > 2;
 				}
@@ -250,7 +355,6 @@ float TraceSpotLight(
 					break;
 				}
 			}
-#endif
 #endif
 		}
 		SampleCount = min(i + 1, uint(RayCount));		// break vs regular for loop exit
diff --git a/Engine/Shaders/Private/VirtualShadowMaps/SMRTCommon.ush b/Engine/Shaders/Private/VirtualShadowMaps/SMRTCommon.ush
new file mode 100644
--- /dev/null
+++ b/Engine/Shaders/Private/VirtualShadowMaps/SMRTCommon.ush
@@ -0,0 +1,22 @@
+// Copyright Epic Games, Inc. All Rights Reserved.
+
+/*=============================================================================
+SMRTTemplateCommon.ush: 
+=============================================================================*/
+
+#pragma once
+
+struct FSMRTSample
+{
+	bool bValid;	
+	float SampleDepth;
+	float ReferenceDepth;
+	float CompareToleranceScale;
+};
+
+struct FSMRTResult
+{
+	bool bValidHit;
+	float OccluderDistance;
+	int SampleMissCount;
+};
diff --git a/Engine/Shaders/Private/VirtualShadowMaps/SMRTTemplate.ush b/Engine/Shaders/Private/VirtualShadowMaps/SMRTTemplate.ush
new file mode 100644
--- /dev/null
+++ b/Engine/Shaders/Private/VirtualShadowMaps/SMRTTemplate.ush
@@ -0,0 +1,83 @@
+// Copyright Epic Games, Inc. All Rights Reserved.
+
+/*=============================================================================
+SMRTTemplate.ush: 
+
+Required parameters before including this file:
+1) Include "SMRTCommon.ush"
+2) Define a "ray state" structure (FRayState here, but must be unique)
+3) Implement FSMRTSample SMRTFindSample(inout FRayState RayState, float SampleTime)
+4) Implement SMRTComputeOccluderDistance(FRayState RayState, float SampleDepth)
+5) Include this file:
+#define SMRT_TEMPLATE_RAY_STRUCT FRayState
+#include "SMRTTemplate.ush"
+#undef SMRT_TEMPLATE_RAY_STRUCT
+=============================================================================*/
+
+// NOTE: No pragma(once) or inclusion guards because this file is meant to be
+// included multiple times as our hacky way to generate template-like instantiations
+// before HLSL supports templates natively.
+
+FSMRTResult SMRTRayCast(
+	SMRT_TEMPLATE_RAY_STRUCT RayState,
+	int NumSteps,
+	float StepOffset)
+{
+	float DepthHistory = -1;
+	int SampleMissCount = 0;
+	float OccluderDistance = 100000;
+
+	const float CompareScale = 2.0 / Pow2(NumSteps);
+	const float CompareBias = -CompareScale * (NumSteps - StepOffset + 1.0f);
+
+	bool bValidHit = false;
+	for (int i = 0; i <= NumSteps; i++)
+	{
+		float SampleTime = Pow2(saturate((-1.0f / NumSteps) * i + 1.0f - (StepOffset - 0.5f) / NumSteps));
+
+		FSMRTSample Sample = SMRTFindSample(RayState, SampleTime);
+		if (Sample.bValid)
+		{
+			const float SampleDepth = Sample.SampleDepth;
+			const float ReferenceDepth = Sample.ReferenceDepth;
+			if (DepthHistory < 0)
+			{
+				// First valid sample we've seen. Do a regular depth compare.
+				DepthHistory = SampleDepth;
+				if (SampleDepth > ReferenceDepth)
+				{
+					OccluderDistance = SMRTComputeOccluderDistance(RayState, SampleDepth);
+					bValidHit = true;
+					break;
+				}
+			}
+			else
+			{
+				const float CompareTolerance = abs(Sample.CompareToleranceScale * (CompareScale * i + CompareBias));
+
+				bool bBehind = (SampleDepth - ReferenceDepth) > (2.0f * CompareTolerance);
+				float DepthSample = bBehind ? DepthHistory : SampleDepth;
+				DepthHistory = DepthSample;
+
+				float DepthDiff = ReferenceDepth - DepthSample;
+				bool bHit = abs(DepthDiff + CompareTolerance) < CompareTolerance;
+				if (bHit)
+				{
+					OccluderDistance = SMRTComputeOccluderDistance(RayState, SampleDepth);
+					bValidHit = true;
+					break;
+				}
+			}
+		}
+		else
+		{
+			++SampleMissCount;
+		}
+	}
+
+	FSMRTResult Result;
+	Result.bValidHit = bValidHit;
+	Result.OccluderDistance = OccluderDistance;
+	Result.SampleMissCount = SampleMissCount;
+	return Result;
+}
diff --git a/Engine/Shaders/Private/VirtualShadowMaps/VirtualShadowMapProjection.usf b/Engine/Shaders/Private/VirtualShadowMaps/VirtualShadowMapProjection.usf
--- a/Engine/Shaders/Private/VirtualShadowMaps/VirtualShadowMapProjection.usf
+++ b/Engine/Shaders/Private/VirtualShadowMaps/VirtualShadowMapProjection.usf
@@ -35,7 +35,8 @@ float NormalOffsetWorld;	// Normal offset in world space
 
 int SMRTRayCount;	// 0 = off
 int SMRTSamplesPerRay;
-float SMRTRayLengthScale;
+float SMRTRayLengthScale;			// Directional lights
+float SMRTCotMaxRayAngleFromLight;	// Spot/point lights
 
 #define CONTACT_SHADOW_SAMPLES 8
 
@@ -199,25 +200,27 @@ void VirtualShadowMapProjectionSpotPS(
 	float HitDistance = WORLD_RADIUS_MISS;
 	float RayMissFactor = 0.0f;
 
+	uint LightType = Light.SpotAngles.x > -2.0f ? LIGHT_TYPE_SPOT : LIGHT_TYPE_POINT;
+
 	if (bInLightRadius && bInLightCone)
 	{
 		const float ContactShadowLengthWorld = ContactShadowLength * View.ClipToView[1][1] * SceneDepth;
 
 		float3 EstimatedGeoWorldNormal = GetEstimatedGeoWorldNormal(TranslatedWorldPosition, ScreenUV);
 
-		uint LightType = Light.SpotAngles.x > -2.0f ? LIGHT_TYPE_SPOT : LIGHT_TYPE_POINT;
-
-		if (SMRTRayCount > 0 && LightType == LIGHT_TYPE_SPOT)
+		if (SMRTRayCount > 0)
 		{
-			RayMissFactor = TraceSpotLight(
+			RayMissFactor = TraceLocalLight(
 				VirtualShadowMapId,
 				Light,
+				LightType,
 				SvPosition,
 				SceneDepth,
 				TranslatedWorldPosition + EstimatedGeoWorldNormal * NormalOffsetWorld,
 				ContactShadowLengthWorld,
 				SMRTRayCount,
 				SMRTSamplesPerRay,
+				SMRTCotMaxRayAngleFromLight,
 				DebugOutputType,
 				DebugOutput,
 				HitDistance);
@@ -246,7 +249,7 @@ void VirtualShadowMapProjectionSpotPS(
 		}
 	}
 
-	OutputTarget = EncodeOutput(TranslatedWorldPosition, Light, LIGHT_TYPE_SPOT, RayMissFactor, HitDistance, SampleCount, DebugOutput);
+	OutputTarget = EncodeOutput(TranslatedWorldPosition, Light, LightType, RayMissFactor, HitDistance, SampleCount, DebugOutput);
 }
 
 
@@ -441,17 +444,19 @@ void VirtualShadowMapProjection(
 
 				uint LightType = Light.SpotAngles.x > -2.0f ? LIGHT_TYPE_SPOT : LIGHT_TYPE_POINT;
 
-				if (SMRTRayCount > 0 && LightType == LIGHT_TYPE_SPOT)
+				if (SMRTRayCount > 0)
 				{
-					RayMissFactor = TraceSpotLight(
+					RayMissFactor = TraceLocalLight(
 						VirtualShadowMapId,
 						Light,
+						LightType,
 						SvPosition,
 						SceneDepth,
 						TranslatedWorldPosition + EstimatedGeoWorldNormal * NormalOffsetWorld,
 						ContactShadowLengthWorld,
 						SMRTRayCount,
 						SMRTSamplesPerRay,
+						SMRTCotMaxRayAngleFromLight,
 						DebugOutputType,
 						DebugOutput,
 						HitDistance);
diff --git a/Engine/Source/Runtime/Renderer/Private/VirtualShadowMaps/VirtualShadowMapProjection.cpp b/Engine/Source/Runtime/Renderer/Private/VirtualShadowMaps/VirtualShadowMapProjection.cpp
--- a/Engine/Source/Runtime/Renderer/Private/VirtualShadowMaps/VirtualShadowMapProjection.cpp
+++ b/Engine/Source/Runtime/Renderer/Private/VirtualShadowMaps/VirtualShadowMapProjection.cpp
@@ -53,24 +53,33 @@ TAutoConsoleVariable<int32> CVarVirtualShadowOnePassProjection(
 	ECVF_RenderThreadSafe
 );
 
-static TAutoConsoleVariable<int32> CVarSMRTRayCountDirectional(
-	TEXT( "r.Shadow.v.SMRT.RayCountDirectional" ),
+static TAutoConsoleVariable<int32> CVarSMRTRayCountLocal(
+	TEXT( "r.Shadow.v.SMRT.RayCountLocal" ),
 	0,
-	TEXT( "Ray count for shadow map tracing of directional lights. 0 = disabled." ),
+	TEXT( "Ray count for shadow map tracing of local lights. 0 = disabled." ),
 	ECVF_RenderThreadSafe
 );
 
-static TAutoConsoleVariable<int32> CVarSMRTRayCountSpot(
-	TEXT( "r.Shadow.v.SMRT.RayCountSpot" ),
-	0,
-	TEXT( "Ray count for shadow map tracing of spot lights. 0 = disabled." ),
+static TAutoConsoleVariable<int32> CVarSMRTSamplesPerRayLocal(
+	TEXT( "r.Shadow.v.SMRT.SamplesPerRayLocal" ),
+	16,
+	TEXT( "Shadow map samples per ray for local lights" ),
 	ECVF_RenderThreadSafe
 );
 
-static TAutoConsoleVariable<int32> CVarSMRTSamplesPerRaySpot(
-	TEXT( "r.Shadow.v.SMRT.SamplesPerRaySpot" ),
-	16,
-	TEXT( "Shadow map samples per ray for spot lights" ),
+static TAutoConsoleVariable<float> CVarSMRTMaxRayAngleFromLight(
+	TEXT( "r.Shadow.v.SMRT.MaxRayAngleFromLight" ),
+	0.03f,
+	TEXT( "Max angle (in radians) a ray is allowed to span from the light's perspective for local lights." )
+	TEXT( "Smaller angles limit the screen space size of shadow penumbra. " )
+	TEXT( "Larger angles lead to more noise. " ),
+	ECVF_RenderThreadSafe
+);
+
+static TAutoConsoleVariable<int32> CVarSMRTRayCountDirectional(
+	TEXT( "r.Shadow.v.SMRT.RayCountDirectional" ),
+	0,
+	TEXT( "Ray count for shadow map tracing of directional lights. 0 = disabled." ),
 	ECVF_RenderThreadSafe
 );
 
@@ -109,6 +118,7 @@ BEGIN_SHADER_PARAMETER_STRUCT(FProjectionParameters, )
 	SHADER_PARAMETER(int32, SMRTRayCount)
 	SHADER_PARAMETER(int32, SMRTSamplesPerRay)
 	SHADER_PARAMETER(float, SMRTRayLengthScale)
+	SHADER_PARAMETER(float, SMRTCotMaxRayAngleFromLight)
 	RENDER_TARGET_BINDING_SLOTS()
 END_SHADER_PARAMETER_STRUCT()
 
@@ -221,6 +231,7 @@ static bool AddPass_RenderVirtualShadowMapProjection(
 		PassParameters->SMRTRayCount = CVarSMRTRayCountDirectional.GetValueOnRenderThread();
 		PassParameters->SMRTSamplesPerRay = CVarSMRTSamplesPerRayDirectional.GetValueOnRenderThread();
 		PassParameters->SMRTRayLengthScale = CVarSMRTRayLengthScaleDirectional.GetValueOnRenderThread();
+		PassParameters->SMRTCotMaxRayAngleFromLight = 0.0f;	// unused in this path
 
 		FVirtualShadowMapProjectionDirectionalPS::FPermutationDomain PermutationVector;
 		PermutationVector.Set<FVirtualShadowMapProjectionDirectionalPS::FOutputTypeDim>(OutputType);
@@ -241,9 +252,10 @@ static bool AddPass_RenderVirtualShadowMapProjection(
 	}
 	else
 	{
-		PassParameters->SMRTRayCount = CVarSMRTRayCountSpot.GetValueOnRenderThread();
-		PassParameters->SMRTSamplesPerRay = CVarSMRTSamplesPerRaySpot.GetValueOnRenderThread();
-		PassParameters->SMRTRayLengthScale = 0.0f;		// Currently unused in this path
+		PassParameters->SMRTRayCount = CVarSMRTRayCountLocal.GetValueOnRenderThread();
+		PassParameters->SMRTSamplesPerRay = CVarSMRTSamplesPerRayLocal.GetValueOnRenderThread();
+		PassParameters->SMRTRayLengthScale = 0.0f;		// unused in this path
+		PassParameters->SMRTCotMaxRayAngleFromLight = 1.0f / FMath::Tan(CVarSMRTMaxRayAngleFromLight.GetValueOnRenderThread());
 
 		FVirtualShadowMapProjectionSpotPS::FPermutationDomain PermutationVector;
 		PermutationVector.Set<FVirtualShadowMapProjectionSpotPS::FOutputTypeDim>(OutputType);
@@ -467,6 +479,7 @@ class FVirtualShadowMapProjectionCS : public FGlobalShader
 		SHADER_PARAMETER(int32, SMRTRayCount)
 		SHADER_PARAMETER(int32, SMRTSamplesPerRay)
 		SHADER_PARAMETER(float, SMRTRayLengthScale)
+		SHADER_PARAMETER(float, SMRTCotMaxRayAngleFromLight)
 	END_SHADER_PARAMETER_STRUCT()
 
 	static void ModifyCompilationEnvironment( const FGlobalShaderPermutationParameters& Parameters, FShaderCompilerEnvironment& OutEnvironment )
@@ -511,9 +524,10 @@ void RenderVirtualShadowMapProjection(
 	PassParameters->ContactShadowLength = CVarContactShadowLength.GetValueOnRenderThread();
 	PassParameters->NormalOffsetWorld = CVarNormalOffsetWorld.GetValueOnRenderThread();
 
-	PassParameters->SMRTRayCount = CVarSMRTRayCountSpot.GetValueOnRenderThread();
-	PassParameters->SMRTSamplesPerRay = CVarSMRTSamplesPerRaySpot.GetValueOnRenderThread();
+	PassParameters->SMRTRayCount = CVarSMRTRayCountLocal.GetValueOnRenderThread();
+	PassParameters->SMRTSamplesPerRay = CVarSMRTSamplesPerRayLocal.GetValueOnRenderThread();
 	PassParameters->SMRTRayLengthScale = 0.0f;		// Currently unused in this path
+	PassParameters->SMRTCotMaxRayAngleFromLight = 1.0f / FMath::Tan(CVarSMRTMaxRayAngleFromLight.GetValueOnRenderThread());
 
 	PassParameters->RWShadowMaskBits = GraphBuilder.CreateUAV( ShadowMaskBits );
 	
