Source length support for point and spot lights in the path tracer

Refactored spot light code to call point light methods instead of duplicating logic.

diff --git a/Engine/Shaders/Private/PathTracing/Light/PathTracingCapsuleLight.ush b/Engine/Shaders/Private/PathTracing/Light/PathTracingCapsuleLight.ush
new file mode 100644
--- /dev/null
+++ b/Engine/Shaders/Private/PathTracing/Light/PathTracingCapsuleLight.ush
@@ -0,0 +1,151 @@
+// Copyright Epic Games, Inc. All Rights Reserved.
+
+/*=============================================================================================
+PathTracingCapsuleLight.ush: Light sampling functions for capsule case of point lights
+===============================================================================================*/
+
+#pragma once
+
+#include "PathTracingLightCommon.ush"
+
+
+// Adapted from:
+//    https://www.iquilezles.org/www/articles/intersectors/intersectors.htm
+// This code assume the ray direction is normalized and only returns the front hit
+float CapsuleIntersect(float3 Ro, float3 Rd, float3 Center, float3 Axis, float Radius)
+{
+	float3  ba = 0.5 * Axis;
+	float3  oa = Ro - Center;
+	float baba = dot(ba, ba);
+	float bard = dot(ba, Rd);
+	float baoa = dot(ba, oa);
+	float rdoa = dot(Rd, oa);
+	float oaoa = dot(oa, oa);
+	float a = baba - bard * bard;
+	float b = baba * rdoa - baoa * bard;
+	float c = baba * oaoa - baoa * baoa - Radius * Radius * baba;
+	float h = b * b - a * c;
+	if (h >= 0.0)
+	{
+		float t = (-b - sqrt(h)) / a;
+		float y = baoa + t * bard;
+		// body
+		if (abs(y) < baba)
+		{
+			return t;
+		}
+
+		// caps
+		float3 oc = oa - sign(y) * ba;
+		b = dot(Rd, oc);
+		c = dot(oc, oc) - Radius * Radius;
+		h = b * b - c;
+		if (h > 0.0)
+		{
+			return -b - sqrt(h);
+		}
+	}
+	return -1.0;
+}
+
+float3 CapsuleNormal(float3 Pos, float3 Center, float3 Axis)
+{
+	float3 pa = Pos - Center;
+	float h = clamp(dot(pa, Axis) / dot(Axis, Axis), -0.5, 0.5);
+	return normalize(pa - h * Axis);
+}
+
+float4 CapsuleLight_TraceLight(RayDesc Ray, int LightId, inout float HitT)
+{
+	float3 LightPosition = GetPosition(LightId);
+	float LightRadius = GetRadius(LightId);
+	float SourceLength = GetSourceLength(LightId);
+	float3 Axis = GetdPdv(LightId) * SourceLength;
+	float t = CapsuleIntersect(Ray.Origin, Ray.Direction, LightPosition, Axis, LightRadius);
+	if (t > Ray.TMin && t < HitT)
+	{
+		float CapsArea = Pow2(LightRadius);
+		float BodyArea = 0.5 * LightRadius * SourceLength;
+		HitT = t;
+		float3 LightPower = GetColor(LightId) * ComputeIESAttenuation(LightId, Ray.Origin);
+		float3 LightRadiance = LightPower / (PI * (CapsArea + BodyArea));
+		float3 LightDirection = GetPosition(LightId) - Ray.Origin;
+		float LightDistanceSquared = dot(LightDirection, LightDirection);
+		LightRadiance *= ComputeAttenuationFalloff(LightDistanceSquared, LightId);
+
+		float3 Normal = CapsuleNormal(Ray.Origin + t * Ray.Direction, LightPosition, Axis);
+
+		float CosTheta = saturate(-dot(Normal, Ray.Direction));
+
+		float Pdf = CosTheta > 0 ? t * t / (4 * PI * CosTheta * (CapsArea + BodyArea)) : 0.0;
+
+		return float4(LightRadiance, Pdf);
+	}
+	return 0.0;
+}
+
+void CapsuleLight_SampleLight(
+	int LightId,
+	float4 RandSample,
+	float3 WorldPos,
+	float3 WorldNormal,
+	out float3 OutDirection,
+	out float  OutDistance,
+	out float3 OutRadianceOverPdf,
+	out float OutPdf)
+{
+	// Capsule case
+	// #dxr_todo: only sample the visible portion of the capsule and account for the 1/d^2 falloff down the axis
+
+	// area sampling doesn't work if the shape has 0 area, so just clamp radius a bit
+	float Radius = max(GetRadius(LightId), 0.01);
+	float SourceLength = GetSourceLength(LightId);
+
+	// the caps are two halves of a full sphere
+	// the body is a cylinder
+	// the common factor of 4*PI is accounted for at the end
+	float CapsArea = Radius * Radius;
+	float BodyArea = 0.5 * Radius * SourceLength;
+
+	float Prob = CapsArea / (CapsArea + BodyArea);
+	float3 Axis = GetdPdv(LightId);
+	float3x3 CylinderBasis = GetTangentBasis(Axis);
+
+	float3 Point, Norm;
+	if (RandSample.y < Prob)
+	{
+		RandSample.y /= Prob;
+		// sample the caps
+		float4 Result = UniformSampleSphere(float2(RandSample.yz));
+
+		Point = Result.xyz * Radius;
+		Point.z += 0.5 * SourceLength * sign(Result.z);
+
+		Norm = Result.xyz;
+	}
+	else
+	{
+		RandSample.y -= Prob;
+		RandSample.y /= 1 - Prob;
+		// sample the cylinder body
+
+		float Phi = 2 * PI * RandSample.y;
+
+		Norm = float3(cos(Phi), sin(Phi), 0.0);
+		Point = float3(Radius * Norm.xy, (RandSample.z - 0.5) * SourceLength);
+	}
+	float3 LightDirection = GetPosition(LightId) - WorldPos;
+	float3 LocalLightDirection = mul(CylinderBasis, LightDirection); // World To Local
+
+	float CosTheta = saturate(dot(Norm, -normalize(LocalLightDirection + Point)));
+
+	Point = mul(Point, CylinderBasis) + LightDirection; // Local to World
+
+	float DistanceSquared = dot(Point, Point);
+	OutDistance = sqrt(DistanceSquared);
+	OutDirection = Point * rcp(OutDistance);
+
+	float3 LightPower = GetColor(LightId);
+	OutRadianceOverPdf = LightPower * CosTheta * 4 / DistanceSquared;
+	OutPdf = CosTheta > 0 ? DistanceSquared / (4 * PI * CosTheta * (CapsArea + BodyArea)) : 0.0;
+}
diff --git a/Engine/Shaders/Private/PathTracing/Light/PathTracingPointLight.ush b/Engine/Shaders/Private/PathTracing/Light/PathTracingPointLight.ush
--- a/Engine/Shaders/Private/PathTracing/Light/PathTracingPointLight.ush
+++ b/Engine/Shaders/Private/PathTracing/Light/PathTracingPointLight.ush
@@ -1,21 +1,26 @@
 // Copyright Epic Games, Inc. All Rights Reserved.
 
 /*=============================================================================================
-PointLight.usf: Light sampling functions for Point light implementation
+PathTracingPointLight.ush: Light sampling functions for Point lights
 ===============================================================================================*/
 
 #pragma once
 
 #include "PathTracingLightCommon.ush"
+#include "PathTracingCapsuleLight.ush"
+
 
 float4 PointLight_TraceLight(RayDesc Ray, int LightId, inout float HitT)
 {
+	if (GetSourceLength(LightId) > 0)
+	{
+		return CapsuleLight_TraceLight(Ray, LightId, HitT);
+	}
+
 	float3 LightPosition = GetPosition(LightId);
 	float LightRadius = GetRadius(LightId);
 	float LightRadius2 = Pow2(LightRadius);
-
 	float3 oc = Ray.Origin - LightPosition;
-
 	float LightDistanceSquared = dot(oc, oc);
 
 	// #dxr_todo: switch to numericaly stable ray/sphere test?
@@ -57,34 +62,46 @@ void PointLight_SampleLight(
 	out float3 OutRadianceOverPdf,
 	out float OutPdf)
 {
-	float Radius = GetRadius(LightId);
-
-	// Sample the solid angle subtended by the sphere (which could be singgular, in which case the PDF will be infinite)
 	float3 LightDirection = GetPosition(LightId) - WorldPos;
 	float LightDistanceSquared = dot(LightDirection, LightDirection);
-	float Radius2 = Pow2(Radius);
 
-	// #dxr_todo: come up with a better definition when we are inside the light
-	float SinThetaMax2 = saturate(Radius2 / LightDistanceSquared);
+	if (GetSourceLength(LightId) > 0)
+	{
+		CapsuleLight_SampleLight(LightId, RandSample, WorldPos, WorldNormal, OutDirection, OutDistance, OutRadianceOverPdf, OutPdf);
+	}
+	else
+	{
+		// Point light case
+
+		// Sample the solid angle subtended by the sphere (which could be singgular, in which case the PDF will be infinite)
+		float Radius = GetRadius(LightId);
+		float Radius2 = Pow2(Radius);
 
-	// #dxr_todo: find a better way of handling the region inside the light than just clamping to 1.0 here
-	float4 DirAndPdf = UniformSampleConeRobust(RandSample.yz, SinThetaMax2);
+		// #dxr_todo: come up with a better definition when we are inside the light
+		float SinThetaMax2 = saturate(Radius2 / LightDistanceSquared);
 
-	float CosTheta = DirAndPdf.z;
-	float SinTheta2 = 1.0 - CosTheta * CosTheta;
+		// #dxr_todo: find a better way of handling the region inside the light than just clamping to 1.0 here
+		float4 DirAndPdf = UniformSampleConeRobust(RandSample.yz, SinThetaMax2);
 
-	OutDirection = TangentToWorld(DirAndPdf.xyz, normalize(LightDirection));
-	OutDistance = length(LightDirection) * (CosTheta - sqrt(max(SinThetaMax2 - SinTheta2, 0.0)));
+		float CosTheta = DirAndPdf.z;
+		float SinTheta2 = 1.0 - CosTheta * CosTheta;
 
-	OutPdf = DirAndPdf.w;
+		OutDirection = TangentToWorld(DirAndPdf.xyz, normalize(LightDirection));
+		OutDistance = length(LightDirection) * (CosTheta - sqrt(max(SinThetaMax2 - SinTheta2, 0.0)));
 
-	float3 LightPower = GetColor(LightId) * ComputeIESAttenuation(LightId, WorldPos);
-	float3 LightRadiance = LightPower / (PI * Radius2);
+		OutPdf = DirAndPdf.w;
+
+		float3 LightPower = GetColor(LightId);
+		float3 LightRadiance = LightPower / (PI * Radius2);
+
+		// When the angle is very small, Radiance over pdf simplifies even more since SolidAngle ~= PI * SinThetaMax2
+		// Canceling out common factors further leads to the classic Power / D^2 formula
+		OutRadianceOverPdf = SinThetaMax2 < 0.001 ? LightPower / LightDistanceSquared : LightRadiance / OutPdf;
+	}
 
-	// When the angle is very small, Radiance over pdf simplifies even more since SolidAngle ~= PI * SinThetaMax2
-	// Canceling out common factors further leads to the classic Power / D^2 formula
-	OutRadianceOverPdf = SinThetaMax2 < 0.001 ? LightPower / LightDistanceSquared : LightRadiance / OutPdf;
+	// NOTE: uses distance to center to keep fadeoff mask consistent for all samples
 	OutRadianceOverPdf *= ComputeAttenuationFalloff(LightDistanceSquared, LightId);
+	OutRadianceOverPdf *= ComputeIESAttenuation(LightId, WorldPos);
 }
 
 void PointLight_EstimateLight(
diff --git a/Engine/Shaders/Private/PathTracing/Light/PathTracingSpotLight.ush b/Engine/Shaders/Private/PathTracing/Light/PathTracingSpotLight.ush
--- a/Engine/Shaders/Private/PathTracing/Light/PathTracingSpotLight.ush
+++ b/Engine/Shaders/Private/PathTracing/Light/PathTracingSpotLight.ush
@@ -1,12 +1,13 @@
 // Copyright Epic Games, Inc. All Rights Reserved.
 
 /*=============================================================================================
-	PathTracingSpotLight.usf: Light sampling functions for Spot light implementation
+	PathTracingSpotLight.ush: Light sampling functions for Spot light implementation
 ===============================================================================================*/
 
 #pragma once
 
 #include "PathTracingLightCommon.ush"
+#include "PathTracingSpotLight.ush"
 
 // #dxr_todo: These are stolen from DynamicLightingCommon.ush but I was having trouble getting that header to include cleanly
 /** 
@@ -27,47 +28,16 @@ float SpotAttenuation(float3 L, float3 SpotDirection, float2 SpotAngles)
 
 
 float4 SpotLight_TraceLight(RayDesc Ray, int LightId, inout float HitT) {
-	// #dxr_todo: share with point light instead of copy/pasting logic?
-	float3 LightPosition = GetPosition(LightId);
-	float LightRadius = GetRadius(LightId);
-	float LightRadius2 = Pow2(LightRadius);
-
-	float3 oc = Ray.Origin - LightPosition;
 
-	float LightDistanceSquared = dot(oc, oc);
-
-	// #dxr_todo: switch to numericaly stable ray/sphere test?
-	float a = dot(Ray.Direction, Ray.Direction);
-	float b = dot(oc, Ray.Direction);
-	float c = LightDistanceSquared - LightRadius2;
-	float h = b * b - a * c;
-	if (h >= 0)
+	float4 Result = PointLight_TraceLight(Ray, LightId, HitT);
+	if (Result.w > 0)
 	{
-		h = sqrt(h);
-		float t = (-b - h) / a;
-		if (t > Ray.TMin && t < HitT)
-		{
-			HitT = t;
-			// #dxr_todo: sphere area is 4*pi*r^2 -- but the factor of 4 is missing for some reason?
-			float3 LightPower = GetColor(LightId) * ComputeIESAttenuation(LightId, Ray.Origin);
-			float3 LightRadiance = LightPower / (PI * LightRadius2);
-
-			// adjust lighting according to spot light cone falloff
-			float3 LightNormal = GetNormal(LightId);
-			float2 CosConeAngles = GetCosConeAngles(LightId);
-			LightRadiance *= SpotAttenuation(-normalize(oc), LightNormal, CosConeAngles);
-			LightRadiance *= ComputeAttenuationFalloff(LightDistanceSquared, LightId);
-
-			float SinThetaMax2 = saturate(LightRadius2 / LightDistanceSquared);
-			float OneMinusCosThetaMax = SinThetaMax2 < 0.01 ? SinThetaMax2 * (0.5 + 0.125 * SinThetaMax2) : 1 - sqrt(1 - SinThetaMax2);
-
-			float SolidAngle = 2 * PI * OneMinusCosThetaMax;
-			return float4(LightRadiance, 1.0 / SolidAngle);
-		}
-		// #dxr_todo: process inside hit here ...
-		// How should we define radiance on the inside of the sphere?
+		float3 LightPosition = GetPosition(LightId);
+		float3 LightNormal = GetNormal(LightId);
+		float2 CosConeAngles = GetCosConeAngles(LightId);
+		Result.xyz *= SpotAttenuation(-normalize(Ray.Origin - LightPosition), LightNormal, CosConeAngles);
 	}
-	return 0.0;
+	return Result;
 }
 
 void SpotLight_SampleLight(
@@ -80,93 +50,12 @@ void SpotLight_SampleLight(
 	out float3 OutRadianceOverPdf,
 	out float OutPdf)
 {
-	float Radius = GetRadius(LightId);
-
-	// Sample the solid angle subtended by the sphere (which could be singgular, in which case the PDF will be infinite)
-	float3 LightDirection = GetPosition(LightId) - WorldPos;
-	float LightDistanceSquared = dot(LightDirection, LightDirection);
-	float Radius2 = Pow2(Radius);
-
-	// #dxr_todo: come up with a better definition when we are inside the light
-	float SinThetaMax2 = saturate(Radius2 / LightDistanceSquared);
-
-	// #dxr_todo: find a better way of handling the region inside the light than just clamping to 1.0 here
-	float4 DirAndPdf = UniformSampleConeRobust(RandSample.yz, SinThetaMax2);
-
-	float CosTheta = DirAndPdf.z;
-	float SinTheta2 = 1.0 - CosTheta * CosTheta;
-
-	OutDirection = TangentToWorld(DirAndPdf.xyz, normalize(LightDirection));
-	OutDistance = length(LightDirection) * (CosTheta - sqrt(max(SinThetaMax2 - SinTheta2, 0.0)));
-
-	OutPdf = DirAndPdf.w;
-
-	float3 LightPower = GetColor(LightId) * ComputeIESAttenuation(LightId, WorldPos);
-	float3 LightRadiance = LightPower / (PI * Radius2);
-
-	// When the angle is very small, Radiance over pdf simplifies even more since SolidAngle ~= PI * SinThetaMax2
-	// Canceling out common factors further leads to the classic Power / D^2 formula
-	OutRadianceOverPdf = SinThetaMax2 < 0.001 ? LightPower / LightDistanceSquared : LightRadiance / OutPdf;
-
-	// adjust lighting according to spot light cone falloff
-	float3 LightNormal = GetNormal(LightId);
-	float2 CosConeAngles = GetCosConeAngles(LightId);
-	OutRadianceOverPdf *= SpotAttenuation(normalize(LightDirection), LightNormal, CosConeAngles);
-	OutRadianceOverPdf *= ComputeAttenuationFalloff(LightDistanceSquared, LightId);
-}
-
-void SpotLight_GenerateLightRay(
-	RayDesc Ray,
-	int LightId,
-	float3 LightUV,
-	out RayDesc OutLightRay
-)
-{
-	float3 LightDirection = GetPosition(LightId) - Ray.Origin;
-
-	OutLightRay.Origin = Ray.Origin;
-	OutLightRay.TMin = 0.001;
-	OutLightRay.Direction = TangentToWorld(LightUV, normalize(LightDirection));
-
-	// Compute distance to sphere edge
-	float Radius = GetRadius(LightId);
-	float Radius2 = Pow2(Radius);
-
-	float RayLength2 = dot(LightDirection, LightDirection);
-	float RayLength = sqrt(RayLength2);
-	float SinThetaMax2 = Radius2 / RayLength2;
-
-	float CosTheta = LightUV.z;
-	float SinTheta2 = 1.0 - CosTheta * CosTheta;
-
-	OutLightRay.TMax = RayLength * (CosTheta - sqrt(max(SinThetaMax2 - SinTheta2, 0.0)));
-}
-
-void SpotLight_EvalLight(
-	int LightId,
-	float3 LightUV,
-	RayDesc Ray,
-	out float3 OutRadiance
-)
-{
-	float3 LightDirection = GetPosition(LightId) - Ray.Origin;
-	float LightDistanceSquared = dot(LightDirection, LightDirection);
-
-	float3 Color = GetColor(LightId);
-	float Radius = GetRadius(LightId);
-	float Falloff = ComputeAttenuationFalloff(LightDistanceSquared, LightId);
-	OutRadiance = Color * Falloff;
-
-	// Same math as point lights
-	if (Radius == 0)
-		OutRadiance *= 1 / LightDistanceSquared;
-	else
-		OutRadiance *= 1 / (PI * Radius * Radius);
+	PointLight_SampleLight(LightId, RandSample, WorldPos, WorldNormal, OutDirection, OutDistance, OutRadianceOverPdf, OutPdf);
 
-	// This appears to be the only difference between point and spot lights
+	float3 LightPosition = GetPosition(LightId);
 	float3 LightNormal = GetNormal(LightId);
 	float2 CosConeAngles = GetCosConeAngles(LightId);
-	OutRadiance *= SpotAttenuation(normalize(LightDirection), LightNormal, CosConeAngles);
+	OutRadianceOverPdf *= SpotAttenuation(-normalize(WorldPos - LightPosition), LightNormal, CosConeAngles);
 }
 
 void SpotLight_EstimateLight(
@@ -195,37 +84,4 @@ void SpotLight_EstimateLight(
 	float2 CosConeAngles = GetCosConeAngles(LightId);
 	OutIrradiance *= SpotAttenuation(normalize(LightDirection), LightNormal, CosConeAngles);
 }
-	
-void SpotLight_SampleLight(
-	RayDesc Ray,
-	FMaterialClosestHitPayload Payload,
-	float4 RandSample,
-	int LightId,
-	out float3 OutLightUV,
-	out float OutPdf
-)
-{
-	float Radius = GetRadius(LightId);
-	if (Radius > 0.0)
-	{
-		// Soft case - sample the solid angle subtended by the sphere
-		float3 LightDirection = GetPosition(LightId) - Payload.WorldPos;
-		float LightDistanceSquared = dot(LightDirection, LightDirection);
-		float Radius2 = Pow2(Radius);
-		
-		// outside the light
-		float SinThetaMax2 = Radius2 / LightDistanceSquared;
 
-		// #dxr_todo: find a better way of handling the region inside the light than just clamping to 1.0 here
-		float4 DirAndPdf = UniformSampleConeRobust(RandSample.yz, min(SinThetaMax2, 1.0));
-		OutLightUV = DirAndPdf.xyz;
-		OutPdf = DirAndPdf.w;
-	}
-	else
-	{
-		// Singular case
-		OutLightUV = float3(0, 0, 1);
-		// #dxr_todo: this isn't really part of the pdf, the pdf should be +inf for the singular case
-		OutPdf = 1.0;
-	}
-}
