Fixed exposure in Mesh Distance Fields visualization

diff --git a/Engine/Shaders/Private/DistanceFieldVisualization.usf b/Engine/Shaders/Private/DistanceFieldVisualization.usf
--- a/Engine/Shaders/Private/DistanceFieldVisualization.usf
+++ b/Engine/Shaders/Private/DistanceFieldVisualization.usf
@@ -202,9 +202,7 @@ void VisualizeDistanceFieldUpsamplePS(in float4 UVAndScreenPos : TEXCOORD0, out
 
 	float3 Value = Texture2DSampleLevel(VisualizeDistanceFieldTexture, VisualizeDistanceFieldSampler, DistanceFieldUVs, 0).xyz;
 
-#if USE_PREEXPOSURE
-	Value *= View.PreExposure;
-#endif
+	// Note: not applying pre-exposure as exposure is skipped in IsAutoExposureDebugMode
 
 	OutColor = float4(Value, 1);
 }
