Path tracer - add missing file from previous commit

The file got refactored as part of CL 15272638

diff --git a/Engine/Shaders/Private/PathTracing/Material/PathTracingDefaultLit.ush b/Engine/Shaders/Private/PathTracing/Material/PathTracingDefaultLit.ush
--- a/Engine/Shaders/Private/PathTracing/Material/PathTracingDefaultLit.ush
+++ b/Engine/Shaders/Private/PathTracing/Material/PathTracingDefaultLit.ush
@@ -8,9 +8,6 @@
 
 #include "PathTracingFresnel.ush"
 
-// be conservative since the GPU may be aggressive with floating point optimizations
-#define GGX_MIN_ROUGHNESS 0.01
-
 float CalcLobeSelectionPdf(float3 V, FMaterialClosestHitPayload Payload)
 {
 	// TODO: This could be improved if we had a directional albedo tabulation, but this formulation will at least prevent
@@ -19,9 +16,7 @@ float CalcLobeSelectionPdf(float3 V, FMaterialClosestHitPayload Payload)
 	float3 Diff = GetDiffuseColor(Payload);
 	float3 Spec = GetSpecularColor(Payload);
 
-	float D = Diff.x + Diff.y + Diff.z;
-	float S = Spec.x + Spec.y + Spec.z;
-	return D / (D + S + 1e-6);
+	return LobeSelectionProb(Diff, Spec);
 }
 
 void DefaultLit_EvalMaterial(
@@ -35,65 +30,34 @@ void DefaultLit_EvalMaterial(
 	OutWeight = 0;
 	OutPdf = 0;
 
-	float3 V_World = -IncomingDirection;
-	float3 L_World = OutgoingDirection;
-	float3 N_World = GetWorldNormal(Payload);
+	const float3 V_World = -IncomingDirection;
+	const float3 L_World = OutgoingDirection;
+	const float3 N_World = GetWorldNormal(Payload);
 
-	float Roughness = max(GetRoughness(Payload), GGX_MIN_ROUGHNESS);
-	float2 Alpha;
-	float3x3 Basis;
-	if (Payload.HasAnisotropy())
-	{
-		// Material is anisotropic, our shading frame should use the tangent vector
-		GetAnisotropicRoughness(Roughness * Roughness, Payload.Anisotropy, Alpha.x, Alpha.y);
-		float3 X = Payload.WorldTangent;
-		float3 Y = normalize(cross(N_World, X));
-		Basis = float3x3(cross(Y, N_World), Y, N_World);
-	}
-	else
-	{
-		// No anisotropy, just pick an arbitrary basis
-		Alpha = Roughness * Roughness;
-		Basis = GetTangentBasis(N_World);
-	}
+	float2 Alpha = 0;
+	const float3x3 Basis = GetGGXBasis(Payload.Roughness, Payload.Anisotropy, Payload.WorldNormal, Payload.WorldTangent, Alpha);
 
 	// move vectors into right shading frame
-	float3 V = mul(Basis, V_World);
-	float3 L = mul(Basis, L_World);
-	float3 H = normalize(V + L);
+	const float3 V = mul(Basis, V_World);
+	const float3 L = mul(Basis, L_World);
+	const float3 H = normalize(V + L);
 
-	float NoL = saturate(L.z);
-	float NoV = saturate(V.z);
-	float NoH = saturate(H.z);
-	float VoH = saturate(dot(V, H));
+	const float NoL = saturate(L.z);
+	const float VoH = saturate(dot(V, H));
 
 	// Probability of picking diffuse lobe vs. specular lobe
-	float LobeSelectionPdf = CalcLobeSelectionPdf(V_World, Payload);
+	const float LobeSelectionPdf = CalcLobeSelectionPdf(V_World, Payload);
 
 	// Diffuse Lobe
 	AddLobeWithMIS(OutWeight, OutPdf, GetDiffuseColor(Payload), NoL / PI, LobeSelectionPdf);
 
 	// Specular lobe
-	if (NoL > 0 && NoV > 0)
-	{
-		// Assume NoH > 0
-		// FIXME: this can be further simplified
-		const float CosThetaM4 = Pow2(NoH);
-		const float TanThetaM2 = (Pow2(H.x / Alpha.x) + Pow2(H.y / Alpha.y)) / Pow2(NoH);
-		const float D = 1 / (PI * Pow2(1 + TanThetaM2) * Alpha.x * Alpha.y * CosThetaM4);
-		const float Lambda_o = 0.5 * (-1.0 + sqrt(1.0 + (Pow2(L.x * Alpha.x) + Pow2(L.y * Alpha.y)) / Pow2(L.z)));
-		const float Lambda_i = 0.5 * (-1.0 + sqrt(1.0 + (Pow2(V.x * Alpha.x) + Pow2(V.y * Alpha.y)) / Pow2(V.z)));
-		const float G2 = 1 / (Lambda_i + Lambda_o + 1);
-		const float G1 = 1 / (Lambda_i + 1);
-
-		// match behavior in SpecularGGX
-		const float3 F = F_Schlick(GetSpecularColor(Payload), VoH);
-
-		const float3 SpecWeight = F * (G2 / G1);
-		const float SpecPdf = (G1 * D * 0.25f) / NoV;
-
-		AddLobeWithMIS(OutWeight, OutPdf, SpecWeight, SpecPdf, 1.0 - LobeSelectionPdf);
-	}
+	const float2 GGXResult = GGXEvalReflection(L, V, H, Alpha);
+	const float3 F = F_Schlick(GetSpecularColor(Payload), VoH);
+	const float3 SpecWeight = F * GGXResult.x;
+	const float SpecPdf = GGXResult.y;
+
+	AddLobeWithMIS(OutWeight, OutPdf, SpecWeight, SpecPdf, 1.0 - LobeSelectionPdf);
 }
 
 void DefaultLit_SampleMaterial(
@@ -105,63 +69,67 @@ void DefaultLit_SampleMaterial(
 	out float OutPdf
 )
 {
+	OutDirection = 0;
+	OutWeight = 0;
+	OutPdf = 0;
+
 	float3 N = GetWorldNormal(Payload);
 	float3 V = -RayDirection;
 	
 	float LobeSelectionPdf = CalcLobeSelectionPdf(V, Payload);
 
-	OutDirection = 0;
-	OutWeight = 0;
-	OutPdf = 0;
+	float2 Alpha;
+	float3x3 Basis = GetGGXBasis(Payload.Roughness, Payload.Anisotropy, Payload.WorldNormal, Payload.WorldTangent, Alpha);
+
+	V = mul(Basis, V);
 	
-	// Choose lobe based on roughness
+	// Randomly choose to sample diffuse or specular
+	float3 L = 0, H = 0;
 	if (RandSample.x < LobeSelectionPdf)
 	{
 		RandSample.x /= LobeSelectionPdf;
 
 		// Lambert
+#if 1
 		float4 Result = CosineSampleHemisphere(RandSample.xy);
-		OutDirection = TangentToWorld(Result.xyz, N);
+#else
+		float4 Result = CosineSampleHemisphereConcentric(RandSample.xy);
+#endif
+
+		L = Result.xyz;
+		H = normalize(L + V);
 	}
 	else
 	{
 		RandSample.x -= LobeSelectionPdf;
 		RandSample.x /= (1.0 - LobeSelectionPdf);
 
-		float Roughness = max(GetRoughness(Payload), GGX_MIN_ROUGHNESS);
-		float2 Alpha;
-		float3x3 Basis;
-		if (Payload.HasAnisotropy())
-		{
-			// Material is anisotropic, our shading frame should use the tangent vector
-			GetAnisotropicRoughness(Roughness * Roughness, Payload.Anisotropy, Alpha.x, Alpha.y);
-			float3 X = Payload.WorldTangent;
-			float3 Y = normalize(cross(N, X));
-			Basis = float3x3(cross(Y, N), Y, N);
-		}
-		else
-		{
-			// No anisotropy, just pick an arbitrary basis
-			Alpha = Roughness * Roughness;
-			Basis = GetTangentBasis(N);
-		}
-
-		V = mul(Basis, V);
 
-		// #dxr_todo: Switch to UniformSampleDiskConcentric after fixing 0/0 bug
-		float3 H = ImportanceSampleVisibleGGX_aniso(UniformSampleDisk(RandSample.xy), Alpha, V).xyz;
+		// #dxr_todo: Evaluate UniformSampleDiskConcentric
+		H = ImportanceSampleVisibleGGX_aniso(UniformSampleDisk(RandSample.xy), Alpha, V).xyz;
 
-		float3 L = reflect(-V, H);
+		L = reflect(-V, H);
 		if (L.z <= 0) {
 			// invalid output direction, exit early
 			return;
 		}
-		// transform to world space
-		OutDirection = normalize(mul(L, Basis));
 	}
 
-	// Call eval to get the proper weight/pdf for the sampled direction
-	// TODO: might be cheaper to inline the eval here and take advantage of already computed terms?
-	// It is possible that the compiler is doing some of this for us already though ...
-	DefaultLit_EvalMaterial(RayDirection, OutDirection, Payload, OutWeight, OutPdf);
+	// With a valid direction in hand -- now evaluate the BxDF (taking advantage of already computed terms)
+
+	const float NoL = saturate(L.z);
+	const float VoH = saturate(dot(V, H));
+
+	// Diffuse Lobe
+	AddLobeWithMIS(OutWeight, OutPdf, GetDiffuseColor(Payload), NoL / PI, LobeSelectionPdf);
+
+	// Specular lobe
+	const float2 GGXResult = GGXEvalReflection(L, V, H, Alpha);
+	const float3 F = F_Schlick(GetSpecularColor(Payload), VoH);
+	const float3 SpecWeight = F * GGXResult.x;
+	const float SpecPdf = GGXResult.y;
+	AddLobeWithMIS(OutWeight, OutPdf, SpecWeight, SpecPdf, 1.0 - LobeSelectionPdf);
+
+	// transform to world space
+	OutDirection = normalize(mul(L, Basis));
 }
