Fixed a wrong logic for starta fast path.

diff --git a/Engine/Shaders/Private/Strata/StrataEvaluation.ush b/Engine/Shaders/Private/Strata/StrataEvaluation.ush
--- a/Engine/Shaders/Private/Strata/StrataEvaluation.ush
+++ b/Engine/Shaders/Private/Strata/StrataEvaluation.ush
@@ -252,16 +252,16 @@ FStrataEvaluateResult StrataEvaluateBSDF(FStrataBSDFContext BSDFContext, FShadow
 #if STRATA_FASTPATH==0
 			const bool bHasThinFilm = BSDF_GETHASTHINFILM(BSDFContext.BSDF);
 			BRANCH
-			if (!bHasThinFilm)
+			if (bHasThinFilm)
 			{
-				F = F_Schlick(F0, F90, BSDFContext.Context.VoH);
+				float ThinFilmIOR = SLAB_THINFILM_IOR(BSDFContext.BSDF);
+				float ThinFilmThickness = SLAB_THINFILM_THICKNESS(BSDFContext.BSDF);
+				F = F_ThinFilm(BSDFContext.Context.NoV, BSDFContext.Context.NoL, BSDFContext.Context.VoH, F0, F90, ThinFilmIOR, ThinFilmThickness);
 			}
 			else
 #endif
 			{
-				float ThinFilmIOR		= SLAB_THINFILM_IOR(BSDFContext.BSDF);
-				float ThinFilmThickness = SLAB_THINFILM_THICKNESS(BSDFContext.BSDF);
-				F = F_ThinFilm(BSDFContext.Context.NoV, BSDFContext.Context.NoL, BSDFContext.Context.VoH, F0, F90, ThinFilmIOR, ThinFilmThickness);
+				F = F_Schlick(F0, F90, BSDFContext.Context.VoH);
 			}
 
 			Sample.SpecularPathProbability	= bHaziness		? (1.0f - SPECULAR_HAZINESS_WEIGHT) : 1.0f;
