Strata more optimisation and decoupling from GBuffer
- added a simple and fast clear path of the material buffer
- light and env light passes now reads AO and isStrataMaterial from the StrataHeader. No longer reads the classification or gbuffer data.
- Handle the max BSDF count case from the compiler.

Clear est 0.1ms at 1080p
SkyEnv pass cost 0.7->0.6ms et occupancy 60->70%

diff --git a/Engine/Shaders/Private/BasePassPixelShader.usf b/Engine/Shaders/Private/BasePassPixelShader.usf
--- a/Engine/Shaders/Private/BasePassPixelShader.usf
+++ b/Engine/Shaders/Private/BasePassPixelShader.usf
@@ -918,7 +918,6 @@ void FPixelShaderInOut_MainPS(
 	half  Metallic = GetMaterialMetallic(PixelMaterialInputs);
 	half  Specular = GetMaterialSpecular(PixelMaterialInputs);
 
-	float MaterialAO = GetMaterialAmbientOcclusion(PixelMaterialInputs);
 	float Roughness = GetMaterialRoughness(PixelMaterialInputs);
 	float Anisotropy = GetMaterialAnisotropy(PixelMaterialInputs);
 	uint ShadingModel = GetMaterialShadingModel(PixelMaterialInputs);
@@ -928,13 +927,14 @@ void FPixelShaderInOut_MainPS(
 	half  Metallic = 0;
 	half  Specular = 0;
 
-	float MaterialAO = 1;
 	float Roughness = 0;
 	float Anisotropy = 0;
 	uint ShadingModel = 0;
 	half Opacity = 0;
 #endif
 
+	float MaterialAO = GetMaterialAmbientOcclusion(PixelMaterialInputs);
+
 // Opacity for this model is the coverage of the top layer over the transmission surface. Since the
 // underlying material isn't allowed to be metallic, we are scaling the metallic value by the coverage.
 #if MATERIAL_SHADINGMODEL_THIN_TRANSLUCENT
@@ -1549,6 +1549,9 @@ void FPixelShaderInOut_MainPS(
 	GBuffer.ShadingModelID = SHADINGMODELID_STRATA;
 	GBuffer.WorldNormal = 0; // Normal will be overriden by the top layer normal
 
+	// STRATA_TODO GBUFFER_HAS_DIFFUSE_SAMPLE_OCCLUSION and ALLOW_STATIC_LIGHTING cases
+	StrataPixelHeader.MaterialAO = GBuffer.GBufferAO;
+
 	uint2 PixelPos = In.SvPosition.xy;
 
 	float3 EmissiveLuminance = 0.0f;
diff --git a/Engine/Shaders/Private/DeferredLightPixelShaders.usf b/Engine/Shaders/Private/DeferredLightPixelShaders.usf
--- a/Engine/Shaders/Private/DeferredLightPixelShaders.usf
+++ b/Engine/Shaders/Private/DeferredLightPixelShaders.usf
@@ -182,12 +182,14 @@ void DeferredLightPixelMain(
 		const FDerivedParams DerivedParams = GetDerivedParams(InputParams, SceneDepth);
 
 		FDeferredLightData LightData = SetupLightDataForStandardDeferred();
+#if STRATA_ENABLED == 0 || (STRATA_ENABLED && STRATA_FASTPATH == 0)
 		if (ScreenSpaceData.GBuffer.ShadingModelID == SHADINGMODELID_HAIR && ShouldUseHairComplexTransmittance(ScreenSpaceData.GBuffer))
 		{
 			// Prevent numerical issue causing NaN during hair lighting (i.e., during baseColor/luma(baseColor) computation)
 			ScreenSpaceData.GBuffer.BaseColor = max(0.0001f, ScreenSpaceData.GBuffer.BaseColor);
 			LightData.HairTransmittance = EvaluateDualScattering(ScreenSpaceData.GBuffer, DerivedParams.CameraVector, -DeferredLightUniforms.Direction);
 		}
+#endif
 
 		float Dither = InterleavedGradientNoise(InputParams.PixelPos, View.StateFrameIndexMod8 );
 
@@ -197,7 +199,10 @@ void DeferredLightPixelMain(
 		float4 Radiance = 0;
 
 #if STRATA_ENABLED
-		if (ScreenSpaceData.GBuffer.ShadingModelID == SHADINGMODELID_STRATA)
+		FStrataAddressing StrataAddressing = GetStrataPixelDataByteOffset(PixelPos, uint2(View.BufferSizeAndInvSize.xy), Strata.MaxBytesPerPixel);
+		FStrataPixelHeader StrataPixelHeader = UnpackStrataHeaderIn(Strata.MaterialLobesBuffer, StrataAddressing);
+		BRANCH
+		if (StrataPixelHeader.BSDFCount > 0) // This test is also enough to exclude sky pixels
 		{
 			float3 V =-DerivedParams.CameraVector;
 			float3 L = LightData.Direction;	// Already normalized
@@ -219,7 +224,6 @@ void DeferredLightPixelMain(
 
 				Radiance = StrataDeferredLighting(
 					LightData,
-					ScreenSpaceData.GBuffer,
 					V,
 					L,
 					ToLight,
@@ -227,7 +231,8 @@ void DeferredLightPixelMain(
 					ShadowTerms,
 					RectTexture,
 					Strata.MaterialLobesBuffer,
-					GetStrataPixelDataByteOffset(PixelPos, uint2(View.BufferSizeAndInvSize.xy), Strata.MaxBytesPerPixel));
+					StrataAddressing,
+					StrataPixelHeader);
 			}
 		}
 #else
diff --git a/Engine/Shaders/Private/ReflectionEnvironmentPixelShader.usf b/Engine/Shaders/Private/ReflectionEnvironmentPixelShader.usf
--- a/Engine/Shaders/Private/ReflectionEnvironmentPixelShader.usf
+++ b/Engine/Shaders/Private/ReflectionEnvironmentPixelShader.usf
@@ -343,9 +343,10 @@ void ReflectionEnvironmentSkyLighting(
 #if ENABLE_SKY_LIGHT
 	// STRATA_TODO handle the ENABLE_SKY_LIGHT && ALLOW_STATIC_LIGHTING case 
 
-	FStrataClassification StrataClassification = StrataUnpackClassificationData(Strata.ClassificationTexture.Load(uint3(PixelPos, 0)));
+	FStrataAddressing StrataAddressing = GetStrataPixelDataByteOffset(PixelPos, uint2(View.BufferSizeAndInvSize.xy), Strata.MaxBytesPerPixel);
+	FStrataPixelHeader StrataPixelHeader = UnpackStrataHeaderIn(Strata.MaterialLobesBuffer, StrataAddressing);
 	BRANCH
-	if (StrataClassification.bIsStrataMaterial) // This test is also enough to exclude sky pixels
+	if (StrataPixelHeader.BSDFCount > 0) // This test is also enough to exclude sky pixels
 	{
 		float DeviceZ = SampleDeviceZFromSceneTextures(BufferUV);
 		float SceneDepth = ConvertFromDeviceZ(DeviceZ);
@@ -370,14 +371,14 @@ void ReflectionEnvironmentSkyLighting(
 		float4 StrataEnvLighting = StrataEnvironmentLighting(
 			BufferUV,
 			SceneDepth,
-			StrataClassification.GBufferAO,
 			WorldPosition,
 			CameraToPixel,
 			V,
 			AmbientOcclusion,
 			TopLayerSpecularContributionFactor,
 			Strata.MaterialLobesBuffer,
-			GetStrataPixelDataByteOffset(PixelPos, uint2(View.BufferSizeAndInvSize.xy), Strata.MaxBytesPerPixel),
+			StrataAddressing,
+			StrataPixelHeader,
 			SSRReductionFactor,
 			TopLayerEnvBRDF);
 
diff --git a/Engine/Shaders/Private/Strata/Strata.ush b/Engine/Shaders/Private/Strata/Strata.ush
--- a/Engine/Shaders/Private/Strata/Strata.ush
+++ b/Engine/Shaders/Private/Strata/Strata.ush
@@ -45,82 +45,109 @@ struct FStrataBSDF
 	float4 VGPRs[4];
 };
 
-// Num bits for each component store in BSDF state
-#define STATE_BITS_BSDF					3
-#define STATE_BITS_SHAREDNORMALID		2
-#define STATE_BITS_ANISOTROPY			1
-#define STATE_BITS_TOPLAYER				1
-#define STATE_BITS_HASSCATTERING		1
-#define STATE_BITS_HASSSSPROFILE		1
-#define STATE_BITS_GREYWEIGHT			1
-#define STATE_BITS_HAZINESS				1
-#define STATE_BITS_EDGECOLOR			1
-#define STATE_BITS_THINFILM				1
-#define STATE_BITS___UNUSED___			10
-#define STATE_BITS_WEIGHT10F			10
+
+
+#define TO_BITMASK(BitCount)							((1<<BitCount)-1)
+#define READ_BITS(X, BitCount, BitOffset)				uint((X>>BitOffset) & TO_BITMASK(BitCount))
+#define WRITE_BITS(X, BitCount, BitOffset, NewValue)	X = (X & (~(TO_BITMASK(BitCount)<<BitOffset))) | ((NewValue & TO_BITMASK(BitCount)) << BitOffset)
+
+// Bit count for each component stored in a Strata header
+// For HEADER_BIT_BSDF_COUNT and  HEADER_BIT_SHAREDNORMAL_X, please refer to values in StrataDefinitions.h
+#define HEADER_BIT_COUNT_BSDF_COUNT				4
+#define HEADER_BIT_COUNT_SHAREDNORMAL_COUNT		3
+#define HEADER_BIT_COUNT_SHAREDNORMAL_TYPE		4
+#define HEADER_BIT_COUNT_AO						8
+#define HEADER_BIT_COUNT___UNUSED___			13
+
+#define HEADER_BIT_OFFSET_BSDF_COUNT			0
+#define HEADER_BIT_OFFSET_SHAREDNORMAL_COUNT	(HEADER_BIT_OFFSET_BSDF_COUNT				+ HEADER_BIT_COUNT_BSDF_COUNT)
+#define HEADER_BIT_OFFSET_SHAREDNORMAL_TYPE		(HEADER_BIT_OFFSET_SHAREDNORMAL_COUNT		+ HEADER_BIT_COUNT_SHAREDNORMAL_COUNT)
+#define HEADER_BIT_OFFSET_AO					(HEADER_BIT_OFFSET_SHAREDNORMAL_TYPE		+ HEADER_BIT_COUNT_SHAREDNORMAL_TYPE)
+
+#define HEADER_GETBSDFCOUNT(X)					READ_BITS( X, HEADER_BIT_COUNT_BSDF_COUNT, HEADER_BIT_OFFSET_BSDF_COUNT)
+#define HEADER_SETBSDFCOUNT(X, NewValue)		WRITE_BITS(X, HEADER_BIT_COUNT_BSDF_COUNT, HEADER_BIT_OFFSET_BSDF_COUNT, NewValue)
+
+#define HEADER_GETSHAREDNORMALCOUNT(X)			READ_BITS( X, HEADER_BIT_COUNT_SHAREDNORMAL_COUNT, HEADER_BIT_OFFSET_SHAREDNORMAL_COUNT)
+#define HEADER_SETSHAREDNORMALCOUNT(X, NewValue)WRITE_BITS(X, HEADER_BIT_COUNT_SHAREDNORMAL_COUNT, HEADER_BIT_OFFSET_SHAREDNORMAL_COUNT, NewValue)
+
+#define HEADER_GETSHAREDNORMALTYPE(X)			READ_BITS( X, HEADER_BIT_COUNT_SHAREDNORMAL_TYPE, HEADER_BIT_OFFSET_SHAREDNORMAL_TYPE)
+#define HEADER_SETSHAREDNORMALTYPE(X, NewValue)	WRITE_BITS(X, HEADER_BIT_COUNT_SHAREDNORMAL_TYPE, HEADER_BIT_OFFSET_SHAREDNORMAL_TYPE, NewValue)
+
+#define HEADER_GETAO(X)							READ_BITS( X, HEADER_BIT_COUNT_AO, HEADER_BIT_OFFSET_AO)
+#define HEADER_SETAO(X, NewValue)				WRITE_BITS(X, HEADER_BIT_COUNT_AO, HEADER_BIT_OFFSET_AO, NewValue)
+
+// Bit count for each component stored in a Strata BSDF state
+#define STATE_BIT_COUNT_BSDF					3
+#define STATE_BIT_COUNT_SHAREDNORMALID			2
+#define STATE_BIT_COUNT_ANISOTROPY				1
+#define STATE_BIT_COUNT_TOPLAYER				1
+#define STATE_BIT_COUNT_HASSCATTERING			1
+#define STATE_BIT_COUNT_HASSSSPROFILE			1
+#define STATE_BIT_COUNT_GREYWEIGHT				1
+#define STATE_BIT_COUNT_HAZINESS				1
+#define STATE_BIT_COUNT_EDGECOLOR				1
+#define STATE_BIT_COUNT_THINFILM				1
+#define STATE_BIT_COUNT___UNUSED___				10
+#define STATE_BIT_COUNT_WEIGHT10F				10
 
 // Bit offset of each component store in BSDF state 
 #define STATE_BIT_OFFSET_BSDF			0
-#define STATE_BIT_OFFSET_SHAREDNORMALID	(STATE_BIT_OFFSET_BSDF				+ STATE_BITS_BSDF)
-#define STATE_BIT_OFFSET_ANISOTROPY		(STATE_BIT_OFFSET_SHAREDNORMALID	+ STATE_BITS_SHAREDNORMALID)
-#define STATE_BIT_OFFSET_TOPLAYER		(STATE_BIT_OFFSET_ANISOTROPY		+ STATE_BITS_ANISOTROPY)
-#define STATE_BIT_OFFSET_HASSCATTERING	(STATE_BIT_OFFSET_TOPLAYER			+ STATE_BITS_TOPLAYER)
-#define STATE_BIT_OFFSET_HASSSSPROFILE	(STATE_BIT_OFFSET_HASSCATTERING		+ STATE_BITS_HASSCATTERING)
-#define STATE_BIT_OFFSET_GREYWEIGHT		(STATE_BIT_OFFSET_HASSSSPROFILE		+ STATE_BITS_HASSSSPROFILE)
-#define STATE_BIT_OFFSET_HAZINESS		(STATE_BIT_OFFSET_GREYWEIGHT		+ STATE_BITS_GREYWEIGHT)
-#define STATE_BIT_OFFSET_EDGECOLOR		(STATE_BIT_OFFSET_HAZINESS			+ STATE_BITS_HAZINESS)
-#define STATE_BIT_OFFSET_THINFILM		(STATE_BIT_OFFSET_EDGECOLOR			+ STATE_BITS_EDGECOLOR)
+#define STATE_BIT_OFFSET_SHAREDNORMALID	(STATE_BIT_OFFSET_BSDF				+ STATE_BIT_COUNT_BSDF)
+#define STATE_BIT_OFFSET_ANISOTROPY		(STATE_BIT_OFFSET_SHAREDNORMALID	+ STATE_BIT_COUNT_SHAREDNORMALID)
+#define STATE_BIT_OFFSET_TOPLAYER		(STATE_BIT_OFFSET_ANISOTROPY		+ STATE_BIT_COUNT_ANISOTROPY)
+#define STATE_BIT_OFFSET_HASSCATTERING	(STATE_BIT_OFFSET_TOPLAYER			+ STATE_BIT_COUNT_TOPLAYER)
+#define STATE_BIT_OFFSET_HASSSSPROFILE	(STATE_BIT_OFFSET_HASSCATTERING		+ STATE_BIT_COUNT_HASSCATTERING)
+#define STATE_BIT_OFFSET_GREYWEIGHT		(STATE_BIT_OFFSET_HASSSSPROFILE		+ STATE_BIT_COUNT_HASSSSPROFILE)
+#define STATE_BIT_OFFSET_HAZINESS		(STATE_BIT_OFFSET_GREYWEIGHT		+ STATE_BIT_COUNT_GREYWEIGHT)
+#define STATE_BIT_OFFSET_EDGECOLOR		(STATE_BIT_OFFSET_HAZINESS			+ STATE_BIT_COUNT_HAZINESS)
+#define STATE_BIT_OFFSET_THINFILM		(STATE_BIT_OFFSET_EDGECOLOR			+ STATE_BIT_COUNT_EDGECOLOR)
 #define STATE_BIT_OFFSET_WEIGHT10F		22
 
-#define TO_BITMASK(BitCount)							((1<<BitCount)-1)
-#define BSDF_GETSTATE(X, BitCount, BitOffset)			uint((X.State>>BitOffset) & TO_BITMASK(BitCount))
-#define BSDF_SETSTATE(X, BitCount, BitOffset, Value)	X.State = (X.State & (~(TO_BITMASK(BitCount)<<BitOffset))) | ((Value & TO_BITMASK(BitCount)) << BitOffset)
-
 // Accessors for the BSDF type
-#define BSDF_GETTYPE(X)					BSDF_GETSTATE(X, STATE_BITS_BSDF, STATE_BIT_OFFSET_BSDF)
-#define BSDF_SETTYPE(X, NewType)		BSDF_SETSTATE(X, STATE_BITS_BSDF, STATE_BIT_OFFSET_BSDF, NewType)
+#define BSDF_GETTYPE(X)					READ_BITS(X.State,  STATE_BIT_COUNT_BSDF, STATE_BIT_OFFSET_BSDF)
+#define BSDF_SETTYPE(X, NewType)		WRITE_BITS(X.State, STATE_BIT_COUNT_BSDF, STATE_BIT_OFFSET_BSDF, NewType)
 
 // Accessors for Shared Normal
 // This assumes Normals[STRATA_MAX_SHARED_NORMAL_REGISTERS]
-#define BSDF_GETNORMALID(X)				BSDF_GETSTATE(X, STATE_BITS_SHAREDNORMALID, STATE_BIT_OFFSET_SHAREDNORMALID)
-#define BSDF_SETNORMALID(X, NrmlIdx)	BSDF_SETSTATE(X, STATE_BITS_SHAREDNORMALID, STATE_BIT_OFFSET_SHAREDNORMALID, NrmlIdx)
+#define BSDF_GETNORMALID(X)				READ_BITS(X.State,  STATE_BIT_COUNT_SHAREDNORMALID, STATE_BIT_OFFSET_SHAREDNORMALID)
+#define BSDF_SETNORMALID(X, NrmlIdx)	WRITE_BITS(X.State, STATE_BIT_COUNT_SHAREDNORMALID, STATE_BIT_OFFSET_SHAREDNORMALID, NrmlIdx)
 
 // Accessors for HasAnisotropy
 // Indicates if the BSDF has anisotropic feature
-#define BSDF_GETANISOTROPY(X)			BSDF_GETSTATE(X, STATE_BITS_ANISOTROPY, STATE_BIT_OFFSET_ANISOTROPY)
-#define BSDF_SETANISOTROPY(X, Aniso)	BSDF_SETSTATE(X, STATE_BITS_ANISOTROPY, STATE_BIT_OFFSET_ANISOTROPY, Aniso)
+#define BSDF_GETANISOTROPY(X)			READ_BITS(X.State,  STATE_BIT_COUNT_ANISOTROPY, STATE_BIT_OFFSET_ANISOTROPY)
+#define BSDF_SETANISOTROPY(X, Aniso)	WRITE_BITS(X.State, STATE_BIT_COUNT_ANISOTROPY, STATE_BIT_OFFSET_ANISOTROPY, Aniso)
 
 // Indicates if the BSDF is part of the top layer
-#define BSDF_GETISTOPLAYER(X)			BSDF_GETSTATE(X, STATE_BITS_TOPLAYER, STATE_BIT_OFFSET_TOPLAYER)
-#define BSDF_SETISTOPLAYER(X, IsTop)	BSDF_SETSTATE(X, STATE_BITS_TOPLAYER, STATE_BIT_OFFSET_TOPLAYER, IsTop)
+#define BSDF_GETISTOPLAYER(X)			READ_BITS(X.State,  STATE_BIT_COUNT_TOPLAYER, STATE_BIT_OFFSET_TOPLAYER)
+#define BSDF_SETISTOPLAYER(X, IsTop)	WRITE_BITS(X.State, STATE_BIT_COUNT_TOPLAYER, STATE_BIT_OFFSET_TOPLAYER, IsTop)
 
 // Indicates if the BSDF has scattering component
-#define BSDF_GETHASSCATTERING(X)		BSDF_GETSTATE(X, STATE_BITS_HASSCATTERING, STATE_BIT_OFFSET_HASSCATTERING)
-#define BSDF_SETHASSCATTERING(X, Scatt)	BSDF_SETSTATE(X, STATE_BITS_HASSCATTERING, STATE_BIT_OFFSET_HASSCATTERING, Scatt)
+#define BSDF_GETHASSCATTERING(X)		READ_BITS(X.State,  STATE_BIT_COUNT_HASSCATTERING, STATE_BIT_OFFSET_HASSCATTERING)
+#define BSDF_SETHASSCATTERING(X, Scatt)	WRITE_BITS(X.State, STATE_BIT_COUNT_HASSCATTERING, STATE_BIT_OFFSET_HASSCATTERING, Scatt)
 
 // Indicates if the BSDF has a SSS profile
-#define BSDF_GETHASSSSPROFILE(X)		BSDF_GETSTATE(X, STATE_BITS_HASSSSPROFILE, STATE_BIT_OFFSET_HASSSSPROFILE)
-#define BSDF_SETHASSSSPROFILE(X, Profi)	BSDF_SETSTATE(X, STATE_BITS_HASSSSPROFILE, STATE_BIT_OFFSET_HASSSSPROFILE, Profi)
+#define BSDF_GETHASSSSPROFILE(X)		READ_BITS(X.State,  STATE_BIT_COUNT_HASSSSPROFILE, STATE_BIT_OFFSET_HASSSSPROFILE)
+#define BSDF_SETHASSSSPROFILE(X, Profi)	WRITE_BITS(X.State, STATE_BIT_COUNT_HASSSSPROFILE, STATE_BIT_OFFSET_HASSSSPROFILE, Profi)
 
 // Indicates if the BSDF weight is grey scale so that it can be stored in a single float (e.g. top layer BSDFs)
-#define BSDF_GETGREYWEIGHT(X)			BSDF_GETSTATE(X, STATE_BITS_GREYWEIGHT, STATE_BIT_OFFSET_GREYWEIGHT)
-#define BSDF_SETGREYWEIGHT(X, Grey)		BSDF_SETSTATE(X, STATE_BITS_GREYWEIGHT, STATE_BIT_OFFSET_GREYWEIGHT, Grey)
+#define BSDF_GETGREYWEIGHT(X)			READ_BITS(X.State,  STATE_BIT_COUNT_GREYWEIGHT, STATE_BIT_OFFSET_GREYWEIGHT)
+#define BSDF_SETGREYWEIGHT(X, Grey)		WRITE_BITS(X.State, STATE_BIT_COUNT_GREYWEIGHT, STATE_BIT_OFFSET_GREYWEIGHT, Grey)
 
 // Indicates if the BSDF has haziness data
-#define BSDF_GETHASHAZINESS(X)			BSDF_GETSTATE(X, STATE_BITS_HAZINESS, STATE_BIT_OFFSET_HAZINESS)
-#define BSDF_SETHASHAZINESS(X, Haze)	BSDF_SETSTATE(X, STATE_BITS_HAZINESS, STATE_BIT_OFFSET_HAZINESS, Haze)
+#define BSDF_GETHASHAZINESS(X)			READ_BITS(X.State,  STATE_BIT_COUNT_HAZINESS, STATE_BIT_OFFSET_HAZINESS)
+#define BSDF_SETHASHAZINESS(X, Haze)	WRITE_BITS(X.State, STATE_BIT_COUNT_HAZINESS, STATE_BIT_OFFSET_HAZINESS, Haze)
 
 // Indicates if the BSDF has edge color data
-#define BSDF_GETHASEDGECOLOR(X)			BSDF_GETSTATE(X, STATE_BITS_EDGECOLOR, STATE_BIT_OFFSET_EDGECOLOR)
-#define BSDF_SETHASEDGECOLOR(X, Col)	BSDF_SETSTATE(X, STATE_BITS_EDGECOLOR, STATE_BIT_OFFSET_EDGECOLOR, Col)
+#define BSDF_GETHASEDGECOLOR(X)			READ_BITS(X.State,  STATE_BIT_COUNT_EDGECOLOR, STATE_BIT_OFFSET_EDGECOLOR)
+#define BSDF_SETHASEDGECOLOR(X, Col)	WRITE_BITS(X.State, STATE_BIT_COUNT_EDGECOLOR, STATE_BIT_OFFSET_EDGECOLOR, Col)
 
 // Indicates if the BSDF has thin film
-#define BSDF_GETHASTHINFILM(X)			BSDF_GETSTATE(X, STATE_BITS_THINFILM, STATE_BIT_OFFSET_THINFILM)
-#define BSDF_SETHASTHINFILM(X, Thin)	BSDF_SETSTATE(X, STATE_BITS_THINFILM, STATE_BIT_OFFSET_THINFILM, Thin)
+#define BSDF_GETHASTHINFILM(X)			READ_BITS(X.State,  STATE_BIT_COUNT_THINFILM, STATE_BIT_OFFSET_THINFILM)
+#define BSDF_SETHASTHINFILM(X, Thin)	WRITE_BITS(X.State, STATE_BIT_COUNT_THINFILM, STATE_BIT_OFFSET_THINFILM, Thin)
 
 // Accessors for the 10bits float weight used if is it detected being a grey scale
-#define BSDF_GETWEIGHT10F(X)			BSDF_GETSTATE(X, STATE_BITS_WEIGHT10F, STATE_BIT_OFFSET_WEIGHT10F)
-#define BSDF_SETWEIGHT10F(X, Weight)	BSDF_SETSTATE(X, STATE_BITS_WEIGHT10F, STATE_BIT_OFFSET_WEIGHT10F, Weight)
+#define BSDF_GETWEIGHT10F(X)			READ_BITS(X.State,  STATE_BIT_COUNT_WEIGHT10F, STATE_BIT_OFFSET_WEIGHT10F)
+#define BSDF_SETWEIGHT10F(X, Weight)	WRITE_BITS(X.State, STATE_BIT_COUNT_WEIGHT10F, STATE_BIT_OFFSET_WEIGHT10F, Weight)
 
 #if STRATA_INLINE_SHADING
 #define BSDF_SETEMISSIVE(X, V)			X.Emissive = V
@@ -185,7 +212,7 @@ float3 StrataUnpackNormal(uint PackedNormal);
 struct FStrataClassification
 {
 	uint ShadingModels;
-	uint GBufferAO;
+	float MaterialAO;
 	bool bIsStrataMaterial;
 	bool bHasSubsurface;
 	bool bHasSubsurfaceProfile;
@@ -195,9 +222,10 @@ struct FStrataClassification
 
 uint StrataPackClassificationData(FStrataClassification In)
 {
+	const uint MaterialAO8Bits = In.MaterialAO * 255.0f;
 	return 
 		(In.ShadingModels & 0xFF) | 
-		((In.GBufferAO & 0xFF)				 << 8  ) |
+		((MaterialAO8Bits & 0xFF)			 << 8  ) |
 		((In.bIsStrataMaterial		? 1 : 0) << 16 ) |
 		((In.bHasSubsurface			? 1 : 0) << 17 ) | 
 		((In.bHasSubsurfaceProfile	? 1 : 0) << 18) |
@@ -207,14 +235,17 @@ uint StrataPackClassificationData(FStrataClassification In)
 
 FStrataClassification StrataUnpackClassificationData(uint In)
 {
+	uint MaterialAO8Bits = 0;
 	FStrataClassification Out;
 	Out.ShadingModels			= (In >> 0 ) & 0xFF;
-	Out.GBufferAO				= (In >> 8 ) & 0xFF;
+	MaterialAO8Bits				= (In >> 8 ) & 0xFF;
 	Out.bIsStrataMaterial		= (In >> 16) & 0x1;
 	Out.bHasSubsurface			= (In >> 17) & 0x1;
 	Out.bHasSubsurfaceProfile	= (In >> 18) & 0x1;
 	Out.bHasTransmission		= (In >> 19) & 0x1;
 	Out.bIsSimple				= (In >> 20) & 0x1;
+
+	Out.MaterialAO = float(MaterialAO8Bits) * (1.0f / 255.0f);
 	return Out;
 }
 
@@ -605,6 +636,8 @@ struct FStrataPixelHeader
 	uint SharedNormalByteOffset;
 	uint SharedNormalTypes;
 #endif
+
+	float MaterialAO;
 };
 
 FStrataPixelHeader InitialiseStrataPixelHeader()
@@ -617,6 +650,7 @@ FStrataPixelHeader InitialiseStrataPixelHeader()
 	Out.SharedNormalByteOffset = 0;
 	Out.SharedNormalTypes = 0;
 #endif
+	Out.MaterialAO = 1.0f;
 	return Out;
 }
 
@@ -1371,11 +1405,15 @@ void PackStrataOut(
 // Pack only strata header
 uint PackStrataHeaderIn(uint InBSDFCount, FStrataPixelHeader InHeader)
 {
-	uint Out = InBSDFCount;
+	uint Out = 0;
+	HEADER_SETBSDFCOUNT(Out, InBSDFCount);
+
 #if STRATA_INLINE_SHADING
-	Out    |= (InHeader.SharedNormals.NormalCount << 8);
-	Out    |= (InHeader.SharedNormals.NormalTypes << 16);
+	HEADER_SETSHAREDNORMALCOUNT(Out, InHeader.SharedNormals.NormalCount);
+	HEADER_SETSHAREDNORMALTYPE(Out, InHeader.SharedNormals.NormalTypes);
 #endif
+
+	HEADER_SETAO(Out, uint(InHeader.MaterialAO * 255.0f));
 	return Out;
 }
 
@@ -1385,11 +1423,11 @@ FStrataPixelHeader UnpackStrataHeaderIn(ByteAddressBuffer StrataBuffer, inout FS
 {
 	FStrataPixelHeader Out = InitialiseStrataPixelHeader();
 
-	uint Packed = 0;
-	STRATA_LOAD_UINT1(Packed);
-	Out.BSDFCount    = 0xFF &  Packed;
-	uint NormalCount = 0xFF & (Packed >> 8);
-	uint NormalTypes = 0xFF & (Packed >> 16);
+	uint PackedHeader = 0;
+	STRATA_LOAD_UINT1(PackedHeader);
+	Out.BSDFCount    = HEADER_GETBSDFCOUNT(PackedHeader);
+	uint NormalCount = HEADER_GETSHAREDNORMALCOUNT(PackedHeader);
+	uint NormalTypes = HEADER_GETSHAREDNORMALTYPE(PackedHeader);
 
 #if STRATA_INLINE_SHADING==0
 	// Only keep the offset to the shared normal memory. We are going to load them on demand.
@@ -1404,33 +1442,8 @@ FStrataPixelHeader UnpackStrataHeaderIn(ByteAddressBuffer StrataBuffer, inout FS
 	StrataAddressing.ReadBytes += NormalCount * STRATA_PACKED_NORMAL_STRIDE_BYTES;
 #endif
 
-	return Out;
-}
-
-// Unpack only strata header 
-// This version assume only a single slab (isotropic BSDF, with no edge color, no subsurface scattering, no haziness, and no thin film)
-FStrataPixelHeader UnpackFastPathStrataHeaderIn(ByteAddressBuffer StrataBuffer, inout FStrataAddressing StrataAddressing)
-{
-	FStrataPixelHeader Out = InitialiseStrataPixelHeader();
-
-	// This function should only be used when inline shading is 0
-#if STRATA_INLINE_SHADING==0
-
-	// When there  is only a single Slab BSDF loaded
-	uint NormalCount = 1;
-
-	// Skip over the header packed data
-	StrataAddressing.CurrentAddr += StrataAddressing.ByteOffsetToNextUint;
+	Out.MaterialAO = float(HEADER_GETAO(PackedHeader)) * (1.0f / 255.0f);
 
-	Out.BSDFCount = 1;
-	Out.SharedNormalByteOffset = StrataAddressing.CurrentAddr;
-	Out.StrataBuffer = StrataBuffer;
-	Out.SharedNormalTypes = STRATA_BASIS_TYPE_NORMAL;
-
-	// Skip over the shared normal memory to be able to load the BSDF content
-	StrataAddressing.CurrentAddr += NormalCount * StrataAddressing.ByteOffsetToNextUint;
-
-#endif
 	return Out;
 }
 
diff --git a/Engine/Shaders/Private/Strata/StrataDeferredLighting.ush b/Engine/Shaders/Private/Strata/StrataDeferredLighting.ush
--- a/Engine/Shaders/Private/Strata/StrataDeferredLighting.ush
+++ b/Engine/Shaders/Private/Strata/StrataDeferredLighting.ush
@@ -147,7 +147,6 @@ void StrataPreprocessAreaLightIntegration(
 // Unpack BSDF on-the-fly
 float4 StrataDeferredLighting(
 	FDeferredLightData LightData,
-	FGBufferData GBuffer,
 	float3 V,
 	float3 L,
 	float3 ToLight,
@@ -155,7 +154,8 @@ float4 StrataDeferredLighting(
 	FShadowTerms ShadowTerms,
 	FRectTexture RectTexture,
 	ByteAddressBuffer StrataDataBuffer,
-	FStrataAddressing StrataAddressing)
+	FStrataAddressing StrataAddressing,
+	FStrataPixelHeader StrataPixelHeader)
 {	
 	float3 PreintegratedIndicentIrradiance = 0;
 	if (LightData.bRectLight)
@@ -171,11 +171,7 @@ float4 StrataDeferredLighting(
 
 	FLightAccumulator Out = (FLightAccumulator)0;
 
-#if STRATA_FASTPATH
-	uint BSDFIndex = 0;
-	FStrataPixelHeader StrataPixelHeader = UnpackFastPathStrataHeaderIn(StrataDataBuffer, StrataAddressing);
-#else
-	FStrataPixelHeader StrataPixelHeader = UnpackStrataHeaderIn(StrataDataBuffer, StrataAddressing);
+#if STRATA_FASTPATH == 0
 	LOOP
 	for (uint BSDFIndex = 0; BSDFIndex < StrataPixelHeader.BSDFCount; ++BSDFIndex)
 #endif
diff --git a/Engine/Shaders/Private/Strata/StrataEnvironmentLighting.ush b/Engine/Shaders/Private/Strata/StrataEnvironmentLighting.ush
--- a/Engine/Shaders/Private/Strata/StrataEnvironmentLighting.ush
+++ b/Engine/Shaders/Private/Strata/StrataEnvironmentLighting.ush
@@ -11,7 +11,6 @@
 float4 StrataEnvironmentLighting(
 	float2 BufferUV,
 	float SceneDepth,
-	float GBufferAO,
 	float3 WorldPosition,
 	float3 CameraToPixel,
 	float3 V,
@@ -19,6 +18,7 @@ float4 StrataEnvironmentLighting(
 	float TopLayerSpecularContributionFactor,
 	ByteAddressBuffer StrataDataBuffer,
 	FStrataAddressing StrataAddressing,
+	FStrataPixelHeader StrataPixelHeader,
 	inout float SSRReductionFactor,
 	inout float3 TopLayerEnvBRDF)
 {
@@ -34,11 +34,7 @@ float4 StrataEnvironmentLighting(
 	// Sample DFAO only once
 	FUpsampleDFAOOutput UpsampleDFAOOutput = UpsampleDFAO(BufferUV, SceneDepth);
 
-#if STRATA_FASTPATH
-	uint BSDFIndex = 0;
-	FStrataPixelHeader StrataPixelHeader = UnpackFastPathStrataHeaderIn(StrataDataBuffer, StrataAddressing);
-#else
-	FStrataPixelHeader StrataPixelHeader = UnpackStrataHeaderIn(StrataDataBuffer, StrataAddressing);
+#if STRATA_FASTPATH == 0
 	for (uint BSDFIndex = 0; BSDFIndex < StrataPixelHeader.BSDFCount; ++BSDFIndex)
 #endif
 	{
@@ -69,7 +65,7 @@ float4 StrataEnvironmentLighting(
 			#endif
 
 			// Compute the common sky visibility factors
-			FSkyLightVisibilityData SkyVisData = GetSkyLightVisibilityData(StrataBSDFContext.N, StrataBSDFContext.N, GBufferAO, AmbientOcclusion, StrataEnvLight.DiffuseNormal);
+			FSkyLightVisibilityData SkyVisData = GetSkyLightVisibilityData(StrataBSDFContext.N, StrataBSDFContext.N, StrataPixelHeader.MaterialAO, AmbientOcclusion, StrataEnvLight.DiffuseNormal);
 
 			// Finally sample the sky diffuse contribution (spherical harmonic, Lambert BRDF)
 			float3 DiffuseLookup = GetSkySHDiffuse(StrataEnvLight.DiffuseNormal) * View.SkyLightColor.rgb;			// STRATA_TODO compute SH coefficients for ON or Chan
diff --git a/Engine/Shaders/Private/Strata/StrataMaterialClassification.usf b/Engine/Shaders/Private/Strata/StrataMaterialClassification.usf
--- a/Engine/Shaders/Private/Strata/StrataMaterialClassification.usf
+++ b/Engine/Shaders/Private/Strata/StrataMaterialClassification.usf
@@ -1,7 +1,6 @@
 // Copyright Epic Games, Inc. All Rights Reserved.
 
 #include "/Engine/Private/Common.ush"
-#include "/Engine/Private/DeferredShadingCommon.ush"
 
 #define STRATA_INLINE_SHADING 0
 #include "/Engine/Private/Strata/Strata.ush"
@@ -133,17 +132,16 @@ void MainPS(
 	out uint2 OutColor1 : SV_Target1)
 {
 	const uint2 PixelPos = uint2(SVPos.xy);
-
-	// Sample scene textures.
-	float2 BufferUV		 = SvPositionToBufferUV(SVPos);
-	FGBufferData GBuffer = GetGBufferDataFromSceneTextures(BufferUV);
 	
 	FStrataClassification ClassificationData	= (FStrataClassification)0;
 	FStrataSubsurfaceData SSSData				= (FStrataSubsurfaceData)0;
 	FStrataTopLayerData TopLayerData			= (FStrataTopLayerData)0;
 
-	ClassificationData.bIsStrataMaterial = GBuffer.ShadingModelID == SHADINGMODELID_STRATA;
-	ClassificationData.GBufferAO = GBuffer.GBufferAO;
+	FStrataAddressing StrataAddressing = GetStrataPixelDataByteOffset(PixelPos, uint2(View.BufferSizeAndInvSize.xy), Strata.MaxBytesPerPixel);
+	FStrataPixelHeader Header = UnpackStrataHeaderIn(Strata.MaterialLobesBuffer, StrataAddressing);
+
+	ClassificationData.bIsStrataMaterial = Header.BSDFCount > 0;
+	ClassificationData.MaterialAO = 1.0;
 	ClassificationData.bIsSimple = true; // non strata material or sky pixels are considered as simple.
 	if (ClassificationData.bIsStrataMaterial)
 	{
@@ -152,8 +150,7 @@ void MainPS(
 		float4 WorldPos = mul(float4(ScreenPosition * Depth, Depth, 1), View.ScreenToWorld);
 		const float3 V = normalize(View.WorldCameraOrigin - WorldPos.xyz);
 
-		FStrataAddressing StrataAddressing = GetStrataPixelDataByteOffset(PixelPos, uint2(View.BufferSizeAndInvSize.xy), Strata.MaxBytesPerPixel);
-		FStrataPixelHeader Header = UnpackStrataHeaderIn(Strata.MaterialLobesBuffer, StrataAddressing);
+		ClassificationData.MaterialAO = Header.MaterialAO;
 		ClassificationData.bIsSimple = Header.BSDFCount == 1;
 		for (uint BSDFIndex = 0; BSDFIndex < Header.BSDFCount; ++BSDFIndex)
 		{
@@ -279,3 +276,32 @@ void StencilMainPS(
 }
 
 #endif //SHADER_STENCIL_CATEGORIZATION
+
+////////////////////////////////////////////////////////////////////////////////////////////////////////////
+
+#if SHADER_CLEAR_MATERIAL_BUFFER
+
+RWByteAddressBuffer MaterialLobesBufferUAV;
+uint MaxBytesPerPixel;
+int2 TiledViewBufferResolution;
+
+[numthreads(8, 8, 1)]
+void ClearMaterialBufferMainCS(uint2 DispatchThreadId : SV_DispatchThreadID)
+{
+	if (any(int2(DispatchThreadId.xy) > TiledViewBufferResolution))
+	{
+		return;
+	}
+
+	// Custom clear of the Strata material buffer.
+	// The first layer of tiled uints contains the header that we need to clear so we only write a single uint per pixel instead of clearing the entire buffer which would be too slow.
+
+	FStrataPixelHeader StrataHeader = InitialiseStrataPixelHeader();
+	uint BSDFCount = 0;
+	uint PackedStrataHeader = PackStrataHeaderIn(BSDFCount, StrataHeader);
+
+	FStrataAddressing StrataAddressing = GetStrataPixelDataByteOffset(DispatchThreadId.xy, TiledViewBufferResolution, MaxBytesPerPixel);
+	MaterialLobesBufferUAV.Store(StrataAddressing.CurrentAddr, PackedStrataHeader);
+}
+
+#endif // SHADER_CLEAR_MATERIAL_BUFFER
diff --git a/Engine/Shaders/Private/Strata/StrataVisualize.usf b/Engine/Shaders/Private/Strata/StrataVisualize.usf
--- a/Engine/Shaders/Private/Strata/StrataVisualize.usf
+++ b/Engine/Shaders/Private/Strata/StrataVisualize.usf
@@ -68,26 +68,27 @@ void VisualizeMaterialPS(
 	float4 SVPos : SV_POSITION,
 	out float4 OutColor : SV_Target0)
 {
-#if 1	// Here to disabled on some platforms having problems with that shader. STRATA_TODO investigate why this shader is spilling like crazy again
+#if 0	// Here to disabled on some platforms having problems with that shader. STRATA_TODO investigate why this shader is spilling like crazy again
 	const uint2 PixelPos = uint2(SVPos.xy);
 	const uint2 MaterialPos = View.CursorPosition;
+	float2 BufferUV = SvPositionToBufferUV(float4(MaterialPos, SVPos.zw));
 	const int  InspectSquare = 10;
 
-	// Sample scene textures.
-	float2 BufferUV = SvPositionToBufferUV(float4(MaterialPos, SVPos.zw));
-	FGBufferData GBuffer = GetGBufferDataFromSceneTextures(BufferUV);
-	bool bStrataMaterial = GBuffer.ShadingModelID == SHADINGMODELID_STRATA;
+	FStrataAddressing StrataAddressing = GetStrataPixelDataByteOffset(MaterialPos, uint2(View.BufferSizeAndInvSize.xy), Strata.MaxBytesPerPixel);
+	FStrataPixelHeader Header = UnpackStrataHeaderIn(Strata.MaterialLobesBuffer, StrataAddressing);
+	bool bStrataMaterial = Header.BSDFCount > 0;
+
 	const float2 ScreenPosition = SvPositionToScreenPosition(SVPos).xy;
 	const float2 ScreenMaterialPosition = SvPositionToScreenPosition(float4(MaterialPos, 0.5f, 1)).xy;
-	const float3 WorldPosition = mul(float4(ScreenMaterialPosition * GBuffer.Depth, GBuffer.Depth, 1), View.ScreenToWorld).xyz;
+	float DeviceZ = SampleDeviceZFromSceneTextures(BufferUV);
+	float SceneDepth = ConvertFromDeviceZ(DeviceZ);
+	const float3 WorldPosition = mul(float4(ScreenMaterialPosition * SceneDepth, SceneDepth, 1), View.ScreenToWorld).xyz;
 	const float3 V = -normalize(WorldPosition - View.WorldCameraOrigin);
 
 #if 0
 	OutColor = 0;
-	if (bStrataMaterial)
+	if (Header.BSD)
 	{
-		FStrataPixelHeader Header = UnpackStrataHeaderIn(Strata.MaterialLobesBuffer, GetStrataPixelDataByteOffset(MaterialPos, uint2(View.BufferSizeAndInvSize.xy), Strata.MaxBytesPerPixel));
-
 		OutColor.w = 0.05f;
 		if (Header.BSDFCount == 1)
 		{
@@ -182,12 +183,14 @@ void VisualizeMaterialPS(
 #define PRINTBSDFANISO()	PRINTBSDFCHAR(_SPC_); PRINTBSDFCHAR(_MINUS_); PRINTBSDFCHAR(_A_); PRINTBSDFCHAR(_N_); PRINTBSDFCHAR(_I_); PRINTBSDFCHAR(_S_); PRINTBSDFCHAR(_O_); PRINTBSDFCHAR(_SPC_); PrintFloatNoFraction(PixelPos, OutColor.xyz, FontColor, TopLeft, float(BSDF_GETANISOTROPY(BSDF)), 1); TopLeft.x+=8*1;
 #define PRINTBSDFTOPLAY()	PRINTBSDFCHAR(_SPC_); PRINTBSDFCHAR(_MINUS_); PRINTBSDFCHAR(_T_); PRINTBSDFCHAR(_O_); PRINTBSDFCHAR(_P_); PRINTBSDFCHAR(_L_); PRINTBSDFCHAR(_A_); PRINTBSDFCHAR(_Y_); PRINTBSDFCHAR(_SPC_); PrintFloatNoFraction(PixelPos, OutColor.xyz, FontColor, TopLeft, float(BSDF_GETISTOPLAYER(BSDF)), 1); TopLeft.x+=8*1;
 #define PRINTBSDFHASSCAT()	PRINTBSDFCHAR(_SPC_); PRINTBSDFCHAR(_MINUS_); PRINTBSDFCHAR(_H_); PRINTBSDFCHAR(_A_); PRINTBSDFCHAR(_S_); PRINTBSDFCHAR(_S_); PRINTBSDFCHAR(_C_); PRINTBSDFCHAR(_A_); PRINTBSDFCHAR(_T_); PRINTBSDFCHAR(_SPC_); PrintFloatNoFraction(PixelPos, OutColor.xyz, FontColor, TopLeft, float(BSDF_GETHASSCATTERING(BSDF)), 1); TopLeft.x+=8*1;
-#define PRINTBSDFCOUNT()	PRINTBSDFCHAR(min(_0_+Header.BSDFCount,_F_)); PRINTBSDFCHAR(_SPC_); PRINTBSDFCHAR(_B_); PRINTBSDFCHAR(_S_); PRINTBSDFCHAR(_D_); PRINTBSDFCHAR(_F_);
 #define PRINTFLOAT3(X) PrintFloat(PixelPos, OutColor.xyz, PARAMFontColor, TopLeft, X.x); TopLeft.x+=8*7; PRINTBSDFCHAR(_SPC_); PrintFloat(PixelPos, OutColor.xyz, PARAMFontColor, TopLeft, X.y); TopLeft.x+=8*7; PRINTBSDFCHAR(_SPC_); PrintFloat(PixelPos, OutColor.xyz, PARAMFontColor, TopLeft, X.z);
 #define PRINTFLOAT1(Value) PrintFloat(PixelPos, OutColor.xyz, PARAMFontColor, TopLeft, Value); TopLeft.x+=8*7;
 #define PRINTFLOATNOFRAC(Value, DigitCount) PrintFloatNoFraction(PixelPos, OutColor.xyz, PARAMFontColor, TopLeft, Value, DigitCount); TopLeft.x+=8*7;
 #define SIGNOFFLOAT1(Value) if(Value < 0.0f) { PrintCharacter(PixelPos, OutColor.xyz, PARAMFontColor, TopLeft, _MINUS_); } /*else { PrintCharacter(PixelPos, OutColor.xyz, PARAMFontColor, TopLeft, _SPC_);}*/
 
+#define PRINTBSDFCOUNT()	PRINTBSDFCHAR(min(_0_+Header.BSDFCount,_F_)); PRINTBSDFCHAR(_SPC_); PRINTBSDFCHAR(_B_); PRINTBSDFCHAR(_S_); PRINTBSDFCHAR(_D_); PRINTBSDFCHAR(_F_);
+#define PRINTMATERIALAO()	PRINTBSDFCHAR(_A_); PRINTBSDFCHAR(_O_); PRINTBSDFCHAR(_SPC_); PRINTFLOAT1(Header.MaterialAO);
+
 #define NEXTLINE  TopLeft.y += 8; RectMax.x = max(RectMax.x, TopLeft.x); RectMax.y = max(RectMax.y, TopLeft.y); TopLeft.x = TopLeftInit.x; 
 #define HORIZONTALSEPARATION  TopLeft.y += 4;
 #define HORIZONTALTAB  TopLeft.x += 8;
@@ -198,11 +201,14 @@ void VisualizeMaterialPS(
 		HORIZONTALSEPARATION;
 		PRINTBSDFCOUNT();
 		NEXTLINE;
+		PRINTMATERIALAO();
+		NEXTLINE;
 		HORIZONTALSEPARATION;
 #else
 		NEXTLINE;
 		HORIZONTALSEPARATION;
 		NEXTLINE;
+		NEXTLINE;
 		HORIZONTALSEPARATION;
 #endif
 
diff --git a/Engine/Shaders/Shared/StrataDefinitions.h b/Engine/Shaders/Shared/StrataDefinitions.h
--- a/Engine/Shaders/Shared/StrataDefinitions.h
+++ b/Engine/Shaders/Shared/StrataDefinitions.h
@@ -14,6 +14,7 @@
 
 #define STRATA_MAX_BSDF_COUNT_PER_LAYER		4
 #define STRATA_MAX_LAYER_COUNT				4
+#define STRATA_MAX_BSDF_COUNT				15
 
 // It this is changed, STATE_BITS_SHAREDNORMAL also needs to be updated
 #define STRATA_MAX_SHARED_NORMAL_REGISTERS	4
diff --git a/Engine/Source/Runtime/Engine/Private/Materials/StrataMaterial.cpp b/Engine/Source/Runtime/Engine/Private/Materials/StrataMaterial.cpp
--- a/Engine/Source/Runtime/Engine/Private/Materials/StrataMaterial.cpp
+++ b/Engine/Source/Runtime/Engine/Private/Materials/StrataMaterial.cpp
@@ -32,13 +32,18 @@ FString GetStrataBSDFName(uint8 BSDFType)
 	return "";
 }
 
-static void UpdateTotalBSDFCount(FStrataMaterialCompilationInfo& StrataInfo)
+static void UpdateTotalBSDFCount(FMaterialCompiler* Compiler, FStrataMaterialCompilationInfo& StrataInfo)
 {
 	StrataInfo.TotalBSDFCount = 0;
 	for (uint32 LayerIt = 0; LayerIt < StrataInfo.LayerCount; ++LayerIt)
 	{
 		StrataInfo.TotalBSDFCount += StrataInfo.Layers[LayerIt].BSDFCount;
 	}
+
+	if ((StrataInfo.TotalBSDFCount) > STRATA_MAX_BSDF_COUNT)
+	{
+		Compiler->Error(TEXT("This material contains too many BSDFs"));
+	}
 }
 
 
@@ -62,7 +67,7 @@ void StrataCompilationInfoCreateSingleBSDFMaterial(FMaterialCompiler* Compiler,
 	StrataInfo.Layers[0].BSDFs[0].bHasScattering = bHasScattering;
 	StrataInfo.Layers[0].BSDFs[0].bHasEdgeColor = bHasEdgeColor;
 	StrataInfo.Layers[0].BSDFs[0].bHasThinFilm = bHasThinFilm;
-	UpdateTotalBSDFCount(StrataInfo);
+	UpdateTotalBSDFCount(Compiler, StrataInfo);
 	Compiler->StrataCompilationInfoRegisterCodeChunk(CodeChunk, StrataInfo);
 }
 
@@ -100,7 +105,7 @@ FStrataMaterialCompilationInfo StrataCompilationInfoAdd(FMaterialCompiler* Compi
 	}
 	StrataInfo.LayerCount = FMath::Max(A.LayerCount, B.LayerCount);
 
-	UpdateTotalBSDFCount(StrataInfo);
+	UpdateTotalBSDFCount(Compiler, StrataInfo);
 	return StrataInfo;
 }
 
@@ -136,7 +141,7 @@ FStrataMaterialCompilationInfo StrataCompilationInfoVerticalLayering(FMaterialCo
 	}
 	StrataInfo.LayerCount += Base.LayerCount;
 
-	UpdateTotalBSDFCount(StrataInfo);
+	UpdateTotalBSDFCount(Compiler, StrataInfo);
 	return StrataInfo;
 }
 
diff --git a/Engine/Source/Runtime/Renderer/Private/Strata/Strata.cpp b/Engine/Source/Runtime/Renderer/Private/Strata/Strata.cpp
--- a/Engine/Source/Runtime/Renderer/Private/Strata/Strata.cpp
+++ b/Engine/Source/Runtime/Renderer/Private/Strata/Strata.cpp
@@ -126,6 +126,8 @@ void InitialiseStrataFrameSceneData(FSceneRenderer& SceneRenderer, FRDGBuilder&
 
 	// Always reset the strata uniform buffer
 	SceneRenderer.Scene->StrataSceneData.StrataGlobalUniformParameters.SafeRelease();
+
+	AddStrataClearMaterialBufferPass(GraphBuilder, StrataSceneData.MaterialLobesBuffer.UAV, StrataSceneData.MaxBytesPerPixel, FIntPoint(ResolutionX, ResolutionY));
 }
 
 void BindStrataBasePassUniformParameters(const FViewInfo& View, FStrataOpaquePassUniformParameters& OutStrataUniformParameters)
@@ -298,6 +300,34 @@ IMPLEMENT_GLOBAL_SHADER(FStrataMaterialClassificationPassPS, "/Engine/Private/St
 
 ////////////////////////////////////////////////////////////////////////////////////////////////////////////////////////////////////////////////
 
+class FStrataClearMaterialBufferCS : public FGlobalShader
+{
+	DECLARE_GLOBAL_SHADER(FStrataClearMaterialBufferCS);
+	SHADER_USE_PARAMETER_STRUCT(FStrataClearMaterialBufferCS, FGlobalShader);
+
+	using FPermutationDomain = TShaderPermutationDomain<>;
+
+	BEGIN_SHADER_PARAMETER_STRUCT(FParameters, )
+		SHADER_PARAMETER_UAV(RWByteAddressBuffer, MaterialLobesBufferUAV)
+		SHADER_PARAMETER(uint32, MaxBytesPerPixel)
+		SHADER_PARAMETER(FIntPoint, TiledViewBufferResolution)
+	END_SHADER_PARAMETER_STRUCT()
+
+	static bool ShouldCompilePermutation(const FGlobalShaderPermutationParameters& Parameters)
+	{
+		return GetMaxSupportedFeatureLevel(Parameters.Platform) >= ERHIFeatureLevel::SM5 && Strata::IsStrataEnabled();
+	}
+
+	static void ModifyCompilationEnvironment(const FGlobalShaderPermutationParameters& Parameters, FShaderCompilerEnvironment& OutEnvironment)
+	{
+		FGlobalShader::ModifyCompilationEnvironment(Parameters, OutEnvironment);
+		OutEnvironment.SetDefine(TEXT("SHADER_CLEAR_MATERIAL_BUFFER"), 1);
+	}
+};
+IMPLEMENT_GLOBAL_SHADER(FStrataClearMaterialBufferCS, "/Engine/Private/Strata/StrataMaterialClassification.usf", "ClearMaterialBufferMainCS", SF_Compute);
+
+////////////////////////////////////////////////////////////////////////////////////////////////////////////////////////////////////////////////
+
 class FStrataMaterialTileClassificationPassCS : public FGlobalShader
 {
 	DECLARE_GLOBAL_SHADER(FStrataMaterialTileClassificationPassCS);
@@ -569,6 +599,23 @@ void AddStrataMaterialClassificationPass(FRDGBuilder& GraphBuilder, const FMinim
 	}
 }
 
+void AddStrataClearMaterialBufferPass(FRDGBuilder& GraphBuilder, FUnorderedAccessViewRHIRef MaterialLobesBufferUAV, uint32 MaxBytesPerPixel, FIntPoint TiledViewBufferResolution)
+{
+	TShaderMapRef<FStrataClearMaterialBufferCS> ComputeShader(GetGlobalShaderMap(GMaxRHIFeatureLevel));
+	FStrataClearMaterialBufferCS::FParameters* PassParameters = GraphBuilder.AllocParameters<FStrataClearMaterialBufferCS::FParameters>();
+	PassParameters->MaterialLobesBufferUAV = MaterialLobesBufferUAV;
+	PassParameters->MaxBytesPerPixel = MaxBytesPerPixel;
+	PassParameters->TiledViewBufferResolution = TiledViewBufferResolution;
+
+	const uint32 GroupSize = 8;
+	FComputeShaderUtils::AddPass(
+		GraphBuilder,
+		RDG_EVENT_NAME("StrataClearMaterialBuffer"),
+		ComputeShader,
+		PassParameters,
+		FComputeShaderUtils::GetGroupCount(TiledViewBufferResolution, GroupSize));
+}
+
 } // namespace Strata
 
 
diff --git a/Engine/Source/Runtime/Renderer/Private/Strata/Strata.h b/Engine/Source/Runtime/Renderer/Private/Strata/Strata.h
--- a/Engine/Source/Runtime/Renderer/Private/Strata/Strata.h
+++ b/Engine/Source/Runtime/Renderer/Private/Strata/Strata.h
@@ -52,16 +52,18 @@ bool IsClassificationEnabled();
 
 void InitialiseStrataFrameSceneData(FSceneRenderer& SceneRenderer, FRDGBuilder& GraphBuilder);
 
+void AddStrataClearMaterialBufferPass(FRDGBuilder& GraphBuilder, FUnorderedAccessViewRHIRef MaterialLobesBufferUAV, uint32 MaxBytesPerPixel, FIntPoint TiledViewBufferResolution);
+
 void BindStrataBasePassUniformParameters(const FViewInfo& View, FStrataOpaquePassUniformParameters& OutStrataUniformParameters);
 
 TUniformBufferRef<FStrataGlobalUniformParameters> BindStrataGlobalUniformParameters(const FViewInfo& View);
 
-void AddVisualizeMaterialPasses(FRDGBuilder& GraphBuilder, const TArray<FViewInfo>& Views, FRDGTextureRef SceneColorTexture);
-
 void AddStrataMaterialClassificationPass(FRDGBuilder& GraphBuilder, const FMinimalSceneTextures& SceneTextures, const TArray<FViewInfo>& Views);
 
 void AddStrataStencilPass(FRDGBuilder& GraphBuilder, const FViewInfo& View, const FMinimalSceneTextures& SceneTextures);
 void AddStrataStencilPass(FRDGBuilder& GraphBuilder, const TArray<FViewInfo>& Views, const FMinimalSceneTextures& SceneTextures);
+
+void AddVisualizeMaterialPasses(FRDGBuilder& GraphBuilder, const TArray<FViewInfo>& Views, FRDGTextureRef SceneColorTexture);
 };
 
 
