Fix warning as error when compiling with DXC

diff --git a/Engine/Shaders/Private/SobolRandom.ush b/Engine/Shaders/Private/SobolRandom.ush
--- a/Engine/Shaders/Private/SobolRandom.ush
+++ b/Engine/Shaders/Private/SobolRandom.ush
@@ -45,7 +45,7 @@ uint2 SobolIndex(uint2 Base, int Index, int Bits = 10)
 	uint2 Result = Base;
 	UNROLL for (int b = 0; b < 10 && b < Bits; ++b)
 	{
-		Result ^= (Index & (1 << b)) ? SobolNumbers[b] : 0;
+		Result ^= (Index & (1u << b)) ? SobolNumbers[b] : 0;
 	}
 	return Result;
 }
@@ -54,7 +54,7 @@ uint2 SobolIndex(uint2 Base, int Index, int Bits = 10)
 uint2 ComputePixelUniqueSobolRandSample(uint2 PixelCoord)
 {
 	const uint TemporalBits = 10;
-	uint FrameIndexMod1024 = ReverseBitsN(GetPowerOfTwoModulatedFrameIndex(1 << TemporalBits), TemporalBits);
+	uint FrameIndexMod1024 = ReverseBitsN(GetPowerOfTwoModulatedFrameIndex(1u << TemporalBits), TemporalBits);
 
 	uint2 SobolBase = SobolPixel(PixelCoord);
 	return SobolIndex(SobolBase, FrameIndexMod1024, TemporalBits);
diff --git a/Engine/Shaders/Private/Strata/Strata.ush b/Engine/Shaders/Private/Strata/Strata.ush
--- a/Engine/Shaders/Private/Strata/Strata.ush
+++ b/Engine/Shaders/Private/Strata/Strata.ush
@@ -47,7 +47,7 @@ struct FStrataBSDF
 
 
 
-#define TO_BITMASK(BitCount)							((1<<BitCount)-1)
+#define TO_BITMASK(BitCount)							((1u<<BitCount)-1)
 #define READ_BITS(X, BitCount, BitOffset)				uint((X>>BitOffset) & TO_BITMASK(BitCount))
 #define WRITE_BITS(X, BitCount, BitOffset, NewValue)	X = (X & (~(TO_BITMASK(BitCount)<<BitOffset))) | ((NewValue & TO_BITMASK(BitCount)) << BitOffset)
 
@@ -303,12 +303,12 @@ FStrataTopLayerData StrataUnpackTopLayerData(uint2 In)
 
 uint StrataShadingModelBit(uint ShadingModel)
 {
-	return 1 << ShadingModel;
+	return 1u << ShadingModel;
 }
 
 bool StrataHasShadingModel(uint ShadingModels, uint ShadingModel)
 {
-	return (ShadingModels & (1 << ShadingModel)) > 0;
+	return (ShadingModels & (1u << ShadingModel)) > 0;
 }
 
 bool StrataHasAnisotropy(in float RoughnessX, float RoughnessY)
@@ -598,8 +598,8 @@ uint  StrataLoadUint1(ByteAddressBuffer StrataBuffer, inout FStrataAddressing St
 float3 StrataUnpackNormal(uint PackedNormal);
 void StrataUnpackNormalAndTangent(inout float3 Normal, inout float3 Tangent, in uint InPacked);
 
-#define STRATA_BASIS_TYPE_NORMAL 0
-#define STRATA_BASIS_TYPE_TANGENT 1
+#define STRATA_BASIS_TYPE_NORMAL 0u
+#define STRATA_BASIS_TYPE_TANGENT 1u
 struct FSharedNormals
 {
 	uint NormalCount;
diff --git a/Engine/Shaders/Public/Platform.ush b/Engine/Shaders/Public/Platform.ush
--- a/Engine/Shaders/Public/Platform.ush
+++ b/Engine/Shaders/Public/Platform.ush
@@ -807,7 +807,7 @@ uint MaskedBitCount( uint2 Bits, uint Index )
 {
 	bool bLow = Index < 32;
 
-	uint Mask = 1 << ( Index - ( bLow ? 0 : 32 ) );
+	uint Mask = 1u << ( Index - ( bLow ? 0 : 32 ) );
 	Mask -= 1;
 
 	uint Offset;
