Fix warning-as-error issue on Strata loop with DXC

At compile time all the loop count value are known, which works fine on FXC/Wave. But DXC has (known) issued unroll/nested unroll. As the compiler get better, we will be able to remove this workaround.

diff --git a/Engine/Shaders/Private/Strata/Strata.ush b/Engine/Shaders/Private/Strata/Strata.ush
--- a/Engine/Shaders/Private/Strata/Strata.ush
+++ b/Engine/Shaders/Private/Strata/Strata.ush
@@ -1099,6 +1099,16 @@ void StrataLoad_ColorGamma2ToLinearAlphaLinear(ByteAddressBuffer StrataBuffer, i
 	A = Data.a;
 }
 
+// DXC has issue with unrolling certain part of the code. To address temporarly the warning-as-error issues, 
+// we redefine the unroll function and fall back to loop on DXC
+#if COMPILER_DXC == 1
+#define STRATA_UNROLL LOOP
+#define STRATA_UNROLL_N(X) LOOP
+#else
+#define STRATA_UNROLL UNROLL
+#define STRATA_UNROLL_N(X) UNROLL_N(X)
+#endif
+
 uint PackStrataHeaderIn(uint BSDFCount, FStrataPixelHeader In);
 
 // If this is changed, please update the compiler side material size evaluation in StrataMaterial.cpp
@@ -1128,7 +1138,7 @@ void PackStrataOut(
 
 	// For each of the BSDF, we evaluated their combined weight/throughput according to the material operation performed (layering, mixing, etc.).
 	int l;
-	UNROLL
+	STRATA_UNROLL_N(STRATA_CLAMPED_LAYER_COUNT)
 	for (l = 0; l < Strata.LayerCount; ++l)
 	{
 		float3 LayerThroughput = 0.0f;
@@ -1136,7 +1146,7 @@ void PackStrataOut(
 
 		const bool bTopLayer = l == 0;
 
-		UNROLL
+		STRATA_UNROLL
 		for (int i = 0; i < Strata.Layers[l].BSDFCount; ++i)
 		{
 			// Sanitize BSDF in place before it is used and packed out to memory
@@ -1324,12 +1334,12 @@ void PackStrataOut(
 
 	// 2. the list of BSDFs
 	BSDFCount = 0;
-	UNROLL
+	STRATA_UNROLL_N(STRATA_CLAMPED_LAYER_COUNT)
 	for (l = 0; l < Strata.LayerCount; ++l)
 	{
 		FStrataLayer Layer = Strata.Layers[l];
 
-		UNROLL
+		STRATA_UNROLL
 		for (int i = 0; i < Layer.BSDFCount; ++i)
 		{
 			FStrataBSDF BSDF = Layer.BSDFs[i];
