Fix compile errors in path tracer CapsuleLight on some console platforms

diff --git a/Engine/Shaders/Private/PathTracing/Light/PathTracingCapsuleLight.ush b/Engine/Shaders/Private/PathTracing/Light/PathTracingCapsuleLight.ush
--- a/Engine/Shaders/Private/PathTracing/Light/PathTracingCapsuleLight.ush
+++ b/Engine/Shaders/Private/PathTracing/Light/PathTracingCapsuleLight.ush
@@ -111,7 +111,8 @@ void CapsuleLight_SampleLight(
 	float3 Axis = GetdPdv(LightId);
 	float3x3 CylinderBasis = GetTangentBasis(Axis);
 
-	float3 Point, Norm;
+	float3 Point;
+	float3 Norm;
 	if (RandSample.y < Prob)
 	{
 		RandSample.y /= Prob;
