Fix edge color not properly setup when running fast path.

diff --git a/Engine/Shaders/Private/Strata/Strata.ush b/Engine/Shaders/Private/Strata/Strata.ush
--- a/Engine/Shaders/Private/Strata/Strata.ush
+++ b/Engine/Shaders/Private/Strata/Strata.ush
@@ -1472,6 +1472,7 @@ void UnpackFastPathSlabBSDF(inout FStrataBSDF OutBSDF, uint2 PackedData01)
 {
 	float4 Data0 = UnpackColorGamma2ToLinearAlphaLinear(PackedData01.x);
 	SLAB_BASECOLOR(OutBSDF)		= Data0.rgb;
+	SLAB_EDGECOLOR(OutBSDF)		= 1.0f;
 	SLAB_ROUGHNESSX(OutBSDF)	= Data0.a;
 
 	float4 Data1 = UnpackRGBA8(PackedData01.y);
