Merging 15277719 (backout of CL 15275547) - backout of hitproxy changes.

- Introduces errors in some materials that compiled fine otherwise.

diff --git a/Engine/Shaders/Private/HitProxyPixelShader.usf b/Engine/Shaders/Private/HitProxyPixelShader.usf
--- a/Engine/Shaders/Private/HitProxyPixelShader.usf
+++ b/Engine/Shaders/Private/HitProxyPixelShader.usf
@@ -21,7 +21,7 @@ void Main(
 #if PER_INSTANCE_HITPROXY_ID || USE_INSTANCE_CULLING
 	float4 InstanceHitProxyId : HIT_PROXY_ID,
 #endif
-	float4 PixelPosition	: TEXCOORD8,
+	float4 PixelPosition	: TEXCOORD6,
 #if USE_WORLD_POSITION_EXCLUDING_SHADER_OFFSETS
 	float3 PixelPositionExcludingWPO : TEXCOORD7,
 #endif
diff --git a/Engine/Shaders/Private/HitProxyVertexShader.usf b/Engine/Shaders/Private/HitProxyVertexShader.usf
--- a/Engine/Shaders/Private/HitProxyVertexShader.usf
+++ b/Engine/Shaders/Private/HitProxyVertexShader.usf
@@ -18,17 +18,13 @@
 Buffer<float4> VertexFetch_HitProxyIdBuffer;
 #endif
 
-#if NUM_TEX_COORD_INTERPOLATORS > 7
-	#error Too large number of VS-to-PS UV sets (NUM_TEX_COORD_INTERPOLATORS > 7)
-#endif
-
 struct FHitProxyVSToPS
 {
 	FVertexFactoryInterpolantsVSToPS FactoryInterpolants;
 #if PER_INSTANCE_HITPROXY_ID || USE_INSTANCE_CULLING
 	float4 InstanceHitProxyId : HIT_PROXY_ID;
 #endif
-	float4 PixelPosition : TEXCOORD8;
+	float4 PixelPosition : TEXCOORD6;
 #if USE_WORLD_POSITION_EXCLUDING_SHADER_OFFSETS
 	float3 PixelPositionExcludingWPO : TEXCOORD7;
 #endif
