Fixed SSS with clustered deferred.

diff --git a/Engine/Shaders/Private/ClusteredDeferredShadingPixelShader.usf b/Engine/Shaders/Private/ClusteredDeferredShadingPixelShader.usf
--- a/Engine/Shaders/Private/ClusteredDeferredShadingPixelShader.usf
+++ b/Engine/Shaders/Private/ClusteredDeferredShadingPixelShader.usf
@@ -36,9 +36,9 @@ Texture2D< uint > ShadowMaskBits;
  * Adds local lighting using the light grid, does not apply directional lights, as they are done elsewhere.
  * Does not support dynamic shadows, as these require the per-light shadow mask.
  */
-float3 GetLightGridLocalLighting(const FCulledLightsGridData InLightGridData, float3 WorldPosition, float3 CameraVector, FScreenSpaceData ScreenSpaceData, float2 ScreenUV, uint EyeIndex, float Dither, uint FirstNonSimpleLightIndex)
+float4 GetLightGridLocalLighting(const FCulledLightsGridData InLightGridData, float3 WorldPosition, float3 CameraVector, FScreenSpaceData ScreenSpaceData, float2 ScreenUV, uint EyeIndex, float Dither, uint FirstNonSimpleLightIndex)
 {
-	float3 DirectLighting = 0;
+	float4 DirectLighting = 0;
 
 	// TODO: Not part of LightGrid ligth data, does it need to be added?
 	float SpecularScale = 1.0f;
@@ -119,7 +119,7 @@ float3 GetLightGridLocalLighting(const FCulledLightsGridData InLightGridData, fl
 		// NOTE: uses AO data like tiled deferred (as opposed to forward path)
 		// TODO: Passing in 0,0 as SVPos, what is meaning of this? Only used if 'REFERENCE_QUALITY' is on.
 		float SurfaceShadow = 1.0f;
-		DirectLighting += GetDynamicLighting(WorldPosition, CameraVector, ScreenSpaceData.GBuffer, ScreenSpaceData.AmbientOcclusion, ScreenSpaceData.GBuffer.ShadingModelID, LightData,LightAttenuation, Dither, uint2(0, 0), RectTexture, SurfaceShadow).xyz;
+		DirectLighting += GetDynamicLighting(WorldPosition, CameraVector, ScreenSpaceData.GBuffer, ScreenSpaceData.AmbientOcclusion, ScreenSpaceData.GBuffer.ShadingModelID, LightData,LightAttenuation, Dither, uint2(0, 0), RectTexture, SurfaceShadow);
 
 #else // 1
 		FSimpleDeferredLightData SimpleLightData = (FSimpleDeferredLightData)0;
@@ -216,7 +216,7 @@ void ClusteredShadingPixelShader(
 		return;
 	}
 
-	float3 CompositedLighting = 0;
+	float4 CompositedLighting = 0;
 
 #if ENABLE_SEPARATE_SIMPLE_LIGHT_LOOP
 
@@ -284,9 +284,7 @@ void ClusteredShadingPixelShader(
 	
 #if USE_PREEXPOSURE && !VISUALIZE_LIGHT_CULLING
 	CompositedLighting *= View.PreExposure;
-#endif	 // USE_PREEXPOSURE && !VISUALIZE_LIGHT_CULLING
-	OutColor.rgb = CompositedLighting;
-
-	OutColor.w = 0.0f;
+#endif
+	OutColor = CompositedLighting;
 }
 
