Fix missing PrimitiveId in VFIntermediates when using GPU-Scene instance culling.

diff --git a/Engine/Shaders/Private/BasePassVertexShader.usf b/Engine/Shaders/Private/BasePassVertexShader.usf
--- a/Engine/Shaders/Private/BasePassVertexShader.usf
+++ b/Engine/Shaders/Private/BasePassVertexShader.usf
@@ -220,10 +220,15 @@ void Main(
 
 #if WRITES_VELOCITY_TO_GBUFFER
 	{
-		// Warning: This here is a non-abstracted use of VFIntermediates.PrimitiveId - Only works because GetPrimitiveData(VFIntermediates.PrimitiveId) is a macro.
 		float4 PrevTranslatedWorldPosition = float4(0, 0, 0, 1);
+#if defined(GPUCULL_TODO)
+		BRANCH
+		if (GetPrimitiveData(VFIntermediates).OutputVelocity > 0 || View.ForceDrawAllVelocities != 0)
+#else //!defined(GPUCULL_TODO)
+		// Warning: This here is a non-abstracted use of VFIntermediates.PrimitiveId - Only works because GetPrimitiveData(VFIntermediates.PrimitiveId) is a macro.
 		BRANCH
 		if (GetPrimitiveData(VFIntermediates.PrimitiveId).OutputVelocity > 0 || View.ForceDrawAllVelocities != 0)
+#endif // defined(GPUCULL_TODO)
 		{
 			PrevTranslatedWorldPosition = VertexFactoryGetPreviousWorldPosition( Input, VFIntermediates );	
 			VertexParameters = GetMaterialVertexParameters(Input, VFIntermediates, PrevTranslatedWorldPosition.xyz, TangentToLocal);
diff --git a/Engine/Shaders/Private/VertexFactoryDefaultInterface.ush b/Engine/Shaders/Private/VertexFactoryDefaultInterface.ush
--- a/Engine/Shaders/Private/VertexFactoryDefaultInterface.ush
+++ b/Engine/Shaders/Private/VertexFactoryDefaultInterface.ush
@@ -86,3 +86,13 @@ uint VertexFactoryGetInstanceIdLoadIndex(FPositionAndNormalOnlyVertexFactoryInpu
 #endif // POSITION_ONLY
 
 #endif // VF_IMPLEMENTED_GET_VIEW_INDEX
+
+
+#if VF_USE_PRIMITIVE_SCENE_DATA && defined(GPUCULL_TODO)
+
+FPrimitiveSceneData GetPrimitiveData(FVertexFactoryIntermediates Intermediates)
+{
+	return Intermediates.SceneData.Primitive;
+}
+
+#endif // VF_USE_PRIMITIVE_SCENE_DATA && defined(GPUCULL_TODO)
