Second attempt as fixing console warnings in PathTracingCapsuleLight

diff --git a/Engine/Shaders/Private/PathTracing/Light/PathTracingCapsuleLight.ush b/Engine/Shaders/Private/PathTracing/Light/PathTracingCapsuleLight.ush
--- a/Engine/Shaders/Private/PathTracing/Light/PathTracingCapsuleLight.ush
+++ b/Engine/Shaders/Private/PathTracing/Light/PathTracingCapsuleLight.ush
@@ -111,18 +111,18 @@ void CapsuleLight_SampleLight(
 	float3 Axis = GetdPdv(LightId);
 	float3x3 CylinderBasis = GetTangentBasis(Axis);
 
-	float3 Point;
-	float3 Norm;
+	float3 LightPoint;
+	float3 LightNormal;
 	if (RandSample.y < Prob)
 	{
 		RandSample.y /= Prob;
 		// sample the caps
 		float4 Result = UniformSampleSphere(float2(RandSample.yz));
 
-		Point = Result.xyz * Radius;
-		Point.z += 0.5 * SourceLength * sign(Result.z);
+		LightPoint = Result.xyz * Radius;
+		LightPoint.z += 0.5 * SourceLength * sign(Result.z);
 
-		Norm = Result.xyz;
+		LightNormal = Result.xyz;
 	}
 	else
 	{
@@ -132,19 +132,19 @@ void CapsuleLight_SampleLight(
 
 		float Phi = 2 * PI * RandSample.y;
 
-		Norm = float3(cos(Phi), sin(Phi), 0.0);
-		Point = float3(Radius * Norm.xy, (RandSample.z - 0.5) * SourceLength);
+		LightNormal = float3(cos(Phi), sin(Phi), 0.0);
+		LightPoint = float3(Radius * LightNormal.xy, (RandSample.z - 0.5) * SourceLength);
 	}
 	float3 LightDirection = GetPosition(LightId) - WorldPos;
 	float3 LocalLightDirection = mul(CylinderBasis, LightDirection); // World To Local
 
-	float CosTheta = saturate(dot(Norm, -normalize(LocalLightDirection + Point)));
+	float CosTheta = saturate(dot(LightNormal, -normalize(LocalLightDirection + LightPoint)));
 
-	Point = mul(Point, CylinderBasis) + LightDirection; // Local to World
+	LightPoint = mul(LightPoint, CylinderBasis) + LightDirection; // Local to World
 
-	float DistanceSquared = dot(Point, Point);
+	float DistanceSquared = dot(LightPoint, LightPoint);
 	OutDistance = sqrt(DistanceSquared);
-	OutDirection = Point * rcp(OutDistance);
+	OutDirection = LightPoint * rcp(OutDistance);
 
 	float3 LightPower = GetColor(LightId);
 	OutRadianceOverPdf = LightPower * CosTheta * 4 / DistanceSquared;
