Use tabulated FG texture for deferred lighting path.

diff --git a/Engine/Shaders/Private/BRDF.ush b/Engine/Shaders/Private/BRDF.ush
--- a/Engine/Shaders/Private/BRDF.ush
+++ b/Engine/Shaders/Private/BRDF.ush
@@ -533,6 +533,14 @@ half3 EnvBRDF( half3 SpecularColor, half Roughness, half NoV )
 	return GF;
 }
 
+half3 EnvBRDF(half3 F0, half3 F90, half Roughness, half NoV)
+{
+	// Importance sampled preintegrated G * F
+	float2 AB = Texture2DSampleLevel(PreIntegratedGF, PreIntegratedGFSampler, float2(NoV, Roughness), 0).rg;
+	float3 GF = F0 * AB.x + F90 * AB.y;
+	return GF;
+}
+
 half2 EnvBRDFApproxLazarov(half Roughness, half NoV)
 {
 	// [ Lazarov 2013, "Getting More Physical in Call of Duty: Black Ops II" ]
diff --git a/Engine/Shaders/Private/BasePassPixelShader.usf b/Engine/Shaders/Private/BasePassPixelShader.usf
--- a/Engine/Shaders/Private/BasePassPixelShader.usf
+++ b/Engine/Shaders/Private/BasePassPixelShader.usf
@@ -87,6 +87,7 @@
 #define STRATA_TRANSLUCENT_ENABLED		(STRATA_ENABLED && MATERIALBLENDING_ANY_TRANSLUCENT>0)
 
 #if STRATA_TRANSLUCENT_ENABLED
+#define STRATA_USE_ENV_BRDF_APPROX (MATERIAL_USE_PREINTEGRATED_GF == 0)
 #include "/Engine/Private/Strata/StrataEvaluation.ush"
 #include "/Engine/Private/Strata/StrataForwardLighting.ush"
 #endif
diff --git a/Engine/Shaders/Private/Strata/StrataEvaluation.ush b/Engine/Shaders/Private/Strata/StrataEvaluation.ush
--- a/Engine/Shaders/Private/Strata/StrataEvaluation.ush
+++ b/Engine/Shaders/Private/Strata/StrataEvaluation.ush
@@ -16,6 +16,9 @@
 #define STRATA_FASTPATH 0
 #endif
 
+#ifndef STRATA_USE_ENV_BRDF_APPROX
+#define STRATA_USE_ENV_BRDF_APPROX 0
+#endif
 
 
 // STRATA_TODO put in a common file
@@ -170,6 +173,16 @@ float StrataBSDFProjectedSolidAngleFactor(in FStrataBSDFContext BSDFContext)
 	return BSDFType == STRATA_BSDF_TYPE_HAIR ? 1.0f : BSDFContext.SatNoL;
 }
 
+float3 StrataEnvBRDF(float3 F0, float3 F90, float SafeRoughnessX, float SatNoV)
+{
+#if STRATA_USE_ENV_BRDF_APPROX
+	return EnvBRDFApprox(F0, F90, SafeRoughnessX, SatNoV);
+#else
+	return EnvBRDF(F0, F90, SafeRoughnessX, SatNoV);
+#endif
+}
+
+
 FStrataEvaluateResult StrataEvaluateBSDF(FStrataBSDFContext BSDFContext, FShadowTerms ShadowTerms)
 {
 	FStrataEvaluateResult Sample = (FStrataEvaluateResult)0;
@@ -508,7 +521,7 @@ FStrataEnvLightResult StrataEvaluateForEnvLight(FStrataBSDFContext BSDFContext,
 				}
 #endif
 
-				float3 EvalEnvBRDF = FresnelWeight * EnvBRDFApprox(F0, F90, SafeRoughnessX, BSDFContext.SatNoV);
+				float3 EvalEnvBRDF = FresnelWeight * StrataEnvBRDF(F0, F90, SafeRoughnessX, BSDFContext.SatNoV);
 
 				StrataEnvLightResult.SpecularDirection = StrataGetOffSpecularPeakReflectionDir(BSDFContext.N, BSDFContext.R, SafeRoughnessX);
 				StrataEnvLightResult.SpecularWeight = EvalEnvBRDF;
@@ -528,7 +541,7 @@ FStrataEnvLightResult StrataEvaluateForEnvLight(FStrataBSDFContext BSDFContext,
 
 					// Compute the second specular weight and roughness
 					float SafeRoughnessHaze = ComputeHazyLobeRoughness(SafeRoughnessX, Haziness);
-					StrataEnvLightResult.SpecularHazeWeight = lerp(0.0f, SPECULAR_HAZINESS_WEIGHT, HazinessFadeIn) * EnvBRDFApprox(F0, F90, SafeRoughnessHaze, BSDFContext.SatNoV);
+					StrataEnvLightResult.SpecularHazeWeight = lerp(0.0f, SPECULAR_HAZINESS_WEIGHT, HazinessFadeIn) * StrataEnvBRDF(F0, F90, SafeRoughnessHaze, BSDFContext.SatNoV);
 					StrataEnvLightResult.SpecularHazeSafeRoughness = SafeRoughnessHaze;
 				}
 #endif
