Refactor path tracer material functions to output small structs instead of using 'out' parameters

This makes the code a bit cleaner, and reduces the risk of leaving some outputs un-initialized. It also makes it easier to change what a material eval/sample should return.

diff --git a/Engine/Plugins/Experimental/GPULightmass/Shaders/Private/LightmapPathTracing.usf b/Engine/Plugins/Experimental/GPULightmass/Shaders/Private/LightmapPathTracing.usf
--- a/Engine/Plugins/Experimental/GPULightmass/Shaders/Private/LightmapPathTracing.usf
+++ b/Engine/Plugins/Experimental/GPULightmass/Shaders/Private/LightmapPathTracing.usf
@@ -423,25 +423,25 @@ void PathTracingKernel(
 					if (PackedPayload.IsMiss())
 					{
 						// Evaluate material
-						float3 MaterialWeight;
-						float MaterialPdf = 0.0;
-						EvalMaterial(Ray.Direction, LightRay.Direction, Payload, Bounce == 0, MaterialWeight, MaterialPdf);
+						FMaterialEval MaterialEval = EvalMaterial(Ray.Direction, LightRay.Direction, Payload, Bounce == 0);
 
-						float3 MaterialEval = MaterialWeight * MaterialPdf;
-
-						float MISWeight = MISMode == 2 ? MISWeightRobust(LightPdf, MaterialPdf) : 1.0;
+						float3 LightContrib = PathThroughput * RadianceOverPdf * MaterialEval.Weight * MaterialEval.Pdf;
+						if (MISMode == 2)
+						{
+							LightContrib *= MISWeightRobust(LightPdf, MaterialEval.Pdf);
+						}
 
 						// Record the contribution
 						if (Bounce > 0)
 						{
-							AccumulateRadiance(Radiance, MISWeight * PathThroughput * RadianceOverPdf * MaterialEval);
+							AccumulateRadiance(Radiance, LightContrib);
 						}
 						else
 						{
 							// GPULightmass records contribution of direct lighting separately, and only for static lights
 							if (!IsStationary(LightId))
 							{
-								DirectLightingNEERadianceValue += MISWeight * PathThroughput * RadianceOverPdf * MaterialEval;
+								DirectLightingNEERadianceValue += LightContrib;
 								DirectLightingNEERadianceDirection = LightRay.Direction;
 							}
 						}
@@ -457,26 +457,22 @@ void PathTracingKernel(
 		}
 	
 		// Sample material
-		float3 Direction;
-		float3 MaterialWeight;
-		float SignedPositionBias;
-		float MaterialSamplePdf;
-		SampleMaterial(Ray.Direction, Payload, RandSample, Bounce == 0, Direction, MaterialWeight, MaterialSamplePdf, SignedPositionBias);
+		FMaterialSample MaterialSample = SampleMaterial(Ray.Direction, Payload, RandSample, Bounce == 0);
 
-		if (MaterialSamplePdf < 0 || asuint(MaterialSamplePdf) > 0x7F800000)
+		if (MaterialSample.Pdf < 0 || asuint(MaterialSample.Pdf) > 0x7F800000)
 		{
 			// Pdf became invalid (either negative or NaN)
 			bIsValidSample = false;
 			break;
 		}
 
-		if (!(MaterialSamplePdf > 0))
+		if (!(MaterialSample.Pdf > 0))
 		{
 			// No valid direction -- we are done
 			break;
 		}
 
-		float3 NextPathThroughput = PathThroughput * MaterialWeight;
+		float3 NextPathThroughput = PathThroughput * MaterialSample.Weight;
 		if (!any(NextPathThroughput > 0))
 		{
 			// no energy left in this path
@@ -512,19 +508,19 @@ void PathTracingKernel(
 		if (MISMode == 2)
 		{
 			// If we are doing MIS, the Material ray on the next loop through will need to know the Pdf
-			PrevMaterialPdf = MaterialSamplePdf;
+			PrevMaterialPdf = MaterialSample.Pdf;
 		}
 
 		// Update ray and loop back around to trace again
-		Ray.Direction = Direction;
+		Ray.Direction = MaterialSample.Direction;
 		Ray.TMin = 0;
-		ApplyPositionBias(Ray, SignedPositionBias* Payload.WorldNormal, 0.01f);
+		ApplyPositionBias(Ray, MaterialSample.PositionBiasSign * Payload.WorldNormal, 0.01f);
 		
 		if (Bounce == 0)
 		{
 			RadianceDirection = Ray.Direction;
 		#if USE_IRRADIANCE_CACHING
-			RadianceProbe_Pdf = MaterialSamplePdf;
+			RadianceProbe_Pdf = MaterialSample.Pdf;
 		#endif
 			PrimaryRandSample = RandSample;
 		}
diff --git a/Engine/Shaders/Private/PathTracing/Material/PathTracingClearCoat.ush b/Engine/Shaders/Private/PathTracing/Material/PathTracingClearCoat.ush
--- a/Engine/Shaders/Private/PathTracing/Material/PathTracingClearCoat.ush
+++ b/Engine/Shaders/Private/PathTracing/Material/PathTracingClearCoat.ush
@@ -13,6 +13,7 @@
 
 #define CLEAR_COAT_F0	0.04 // From hard-coded value in ShadingModels.ush
 
+#include "PathTracingMaterialCommon.ush"
 #include "PathTracingFresnel.ush"
 
 float3 CalcThinTransmission(float NoL, float NoV, FMaterialClosestHitPayload Payload)
@@ -57,17 +58,14 @@ float ClearCoatFresnel(float VoH) {
 	return Fc + (1 - Fc) * CLEAR_COAT_F0;
 }
 
-void ClearCoat_EvalMaterial(
+FMaterialEval ClearCoat_EvalMaterial(
 	float3 IncomingDirection,
 	float3 OutgoingDirection,
-	inout FMaterialClosestHitPayload Payload,
-	out float3 OutWeight,
-	out float OutPdf
+	inout FMaterialClosestHitPayload Payload
 )
 {
-	OutWeight = 0;
-	OutPdf = 0;
-
+	FMaterialEval Result = NullMaterialEval();
+	
 	float3 V_World = -IncomingDirection;
 	float3 L_World = OutgoingDirection;
 	float3 N_World = GetWorldNormal(Payload);
@@ -93,22 +91,21 @@ void ClearCoat_EvalMaterial(
 	const float3 SpecWeight = F * GGXResult.x;
 	const float SpecPdf = GGXResult.y;
 
-	AddLobeWithMIS(OutWeight, OutPdf, SpecWeight, SpecPdf, CoatSelectionPdf);
+	AddLobeWithMIS(Result.Weight, Result.Pdf, SpecWeight, SpecPdf, CoatSelectionPdf);
 	
 
 	// bottom-layer clear coat evaluation
-	float3 BottomWeight = 0;
-	float BottomPdf = 0;
+	FMaterialEval BottomEval = NullMaterialEval();
 	Payload.WorldNormal = GetClearCoatBottomNormal(Payload);
 #if GREY_INNER_LAYER == 0
 	// NOTE: Should we try to account for the IOR ratio between layers? We could try to darken Specular Color by the coating IOR
 	// but this would apply to both the dielctric and metal portions equally ...
 	// It seems better to punt until we have a Strata going
-	DefaultLit_EvalMaterial(IncomingDirection, OutgoingDirection, Payload, BottomWeight, BottomPdf);
+	BottomEval = DefaultLit_EvalMaterial(IncomingDirection, OutgoingDirection, Payload);
 #else
 	float3 OldColor = Payload.BaseColor;
 	Payload.BaseColor = 0.18;
-	Lambert_EvalMaterial(OutgoingDirection, Payload, BottomWeight, BottomPdf);
+	BottomEval = Lambert_EvalMaterial(OutgoingDirection, Payload);
 	Payload.BaseColor = OldColor;
 #endif
 	Payload.WorldNormal = N_World;
@@ -123,30 +120,26 @@ void ClearCoat_EvalMaterial(
 	float CC_F1 = ClearCoatFresnel(NoV);
 	float CC_F2 = ClearCoatFresnel(NoL);
 
-	BottomWeight *= lerp(1.0, Transmission * (1 - CC_F1) * (1 - CC_F2), GetClearCoat(Payload));
+	BottomEval.Weight *= lerp(1.0, Transmission * (1 - CC_F1) * (1 - CC_F2), GetClearCoat(Payload));
 
-	AddLobeWithMIS(OutWeight, OutPdf, BottomWeight, BottomPdf, 1.0 - CoatSelectionPdf);
+	AddLobeWithMIS(Result.Weight, Result.Pdf, BottomEval.Weight, BottomEval.Pdf, 1.0 - CoatSelectionPdf);
+
+	return Result;
 }
 
-void ClearCoat_SampleMaterial(
+FMaterialSample ClearCoat_SampleMaterial(
 	float3 RayDirection,
 	FMaterialClosestHitPayload Payload,
-	float4 RandSample,
-	out float3 OutDirection,
-	out float3 OutWeight,
-	out float OutPdf
+	float4 RandSample
 )
 {
 	float3 N_World = GetWorldNormal(Payload);
 	float3 V_World = -RayDirection;
 
-	OutDirection = 0.0;
-	OutWeight = 0.0;
-	OutPdf = 0.0;
-
 	float CoatSelectionPdf = CalcClearCoatSelectionPdf(Payload);
 
 	// choose between coat and bottom layer
+	float3 OutDirection = 0;
 	if (RandSample.x < CoatSelectionPdf)
 	{
 		RandSample.x /= CoatSelectionPdf;
@@ -162,9 +155,10 @@ void ClearCoat_SampleMaterial(
 		float3 H = ImportanceSampleVisibleGGX_aniso(UniformSampleDisk(RandSample.xy), Alpha, V).xyz;
 
 		float3 L = reflect(-V, H);
-		if (L.z <= 0) {
+		if (L.z <= 0)
+		{
 			// invalid output direction, skip some work
-			return;
+			return NullMaterialSample();
 		}
 		// reflect and transform
 		OutDirection = mul(L, Basis);
@@ -178,7 +172,7 @@ void ClearCoat_SampleMaterial(
 		float3 Weight = 0;
 		float Pdf = 0;
 		Payload.WorldNormal = GetClearCoatBottomNormal(Payload);
-		DefaultLit_SampleMaterial(RayDirection, Payload, RandSample, OutDirection, Weight, Pdf);
+		OutDirection = DefaultLit_SampleMaterial(RayDirection, Payload, RandSample).Direction;
 		Payload.WorldNormal = N_World;
 #else
 		float3 Nb = GetClearCoatBottomNormal(Payload);
@@ -188,5 +182,6 @@ void ClearCoat_SampleMaterial(
 
 	// defer to eval for now to get proper MIS mix of both events
 	// TODO: could use temporary results above to inline Eval logic here
-	ClearCoat_EvalMaterial(RayDirection, OutDirection, Payload, OutWeight, OutPdf);
+	FMaterialEval Result = ClearCoat_EvalMaterial(RayDirection, OutDirection, Payload);
+	return CreateMaterialSample(OutDirection, Result.Weight, Result.Pdf, 1.0);
 }
diff --git a/Engine/Shaders/Private/PathTracing/Material/PathTracingCloth.ush b/Engine/Shaders/Private/PathTracing/Material/PathTracingCloth.ush
--- a/Engine/Shaders/Private/PathTracing/Material/PathTracingCloth.ush
+++ b/Engine/Shaders/Private/PathTracing/Material/PathTracingCloth.ush
@@ -21,16 +21,13 @@ float ClothLobeSelectionPdf(float3 V, FMaterialClosestHitPayload Payload)
 	return LobeSelectionProb(Diff + Cloth * Fuzz, (1.0 - Cloth) * Spec);
 }
 
-void Cloth_EvalMaterial(
+FMaterialEval Cloth_EvalMaterial(
 	float3 IncomingDirection,
 	float3 OutgoingDirection,
-	FMaterialClosestHitPayload Payload,
-	out float3 OutWeight,
-	out float OutPdf
+	FMaterialClosestHitPayload Payload
 )
 {
-	OutWeight = 0;
-	OutPdf = 0;
+	FMaterialEval Result = NullMaterialEval();
 
 	const float3 V_World = -IncomingDirection;
 	const float3 L_World = OutgoingDirection;
@@ -60,7 +57,7 @@ void Cloth_EvalMaterial(
 	const float3 FCloth = F_Schlick(FuzzColor, VoH);
 	const float3 ClothWeight = NoL > 0 && NoV > 0 ? (PI * Cloth * DCloth * Vis_Cloth(NoV, NoL)) * FCloth : 0.0;
 
-	AddLobeWithMIS(OutWeight, OutPdf, Payload.DiffuseColor + ClothWeight, NoL / PI, LobeSelectionPdf);
+	AddLobeWithMIS(Result.Weight, Result.Pdf, Payload.DiffuseColor + ClothWeight, NoL / PI, LobeSelectionPdf);
 
 	// Specular lobe
 	const float2 GGXResult = GGXEvalReflection(L, V, H, Alpha);
@@ -68,22 +65,17 @@ void Cloth_EvalMaterial(
 	const float3 SpecWeight = (1 - Cloth) * F * GGXResult.x;
 	const float SpecPdf = GGXResult.y;
 
-	AddLobeWithMIS(OutWeight, OutPdf, SpecWeight, SpecPdf, 1 - LobeSelectionPdf);
+	AddLobeWithMIS(Result.Weight, Result.Pdf, SpecWeight, SpecPdf, 1 - LobeSelectionPdf);
+
+	return Result;
 }
 
-void Cloth_SampleMaterial(
+FMaterialSample Cloth_SampleMaterial(
 	float3 RayDirection,
 	FMaterialClosestHitPayload Payload,
-	float4 RandSample,
-	out float3 OutDirection,
-	out float3 OutWeight,
-	out float OutPdf
+	float4 RandSample
 )
 {
-	OutDirection = 0;
-	OutWeight = 0;
-	OutPdf = 0;
-
 	float3 N = GetWorldNormal(Payload);
 	float3 V = -RayDirection;
 
@@ -120,9 +112,10 @@ void Cloth_SampleMaterial(
 		H = ImportanceSampleVisibleGGX_aniso(UniformSampleDisk(RandSample.xy), Alpha, V).xyz;
 
 		L = reflect(-V, H);
-		if (L.z <= 0) {
+		if (L.z <= 0)
+		{
 			// invalid output direction, exit early
-			return;
+			return NullMaterialSample();
 		}
 	}
 
@@ -141,6 +134,8 @@ void Cloth_SampleMaterial(
 	const float3 FCloth = F_Schlick(FuzzColor, VoH);
 	const float3 ClothWeight = NoL > 0 && NoV > 0 ? (PI * Cloth * DCloth * Vis_Cloth(NoV, NoL)) * FCloth : 0.0;
 
+	float3 OutWeight = 0;
+	float OutPdf = 0;
 	AddLobeWithMIS(OutWeight, OutPdf, Payload.DiffuseColor + ClothWeight, NoL / PI, LobeSelectionPdf);
 
 	// Specular lobe
@@ -151,5 +146,5 @@ void Cloth_SampleMaterial(
 	AddLobeWithMIS(OutWeight, OutPdf, SpecWeight, SpecPdf, 1 - LobeSelectionPdf);
 
 	// transform to world space
-	OutDirection = normalize(mul(L, Basis));
+	return CreateMaterialSample(normalize(mul(L, Basis)), OutWeight, OutPdf, 1.0);
 }
diff --git a/Engine/Shaders/Private/PathTracing/Material/PathTracingDefaultLit.ush b/Engine/Shaders/Private/PathTracing/Material/PathTracingDefaultLit.ush
--- a/Engine/Shaders/Private/PathTracing/Material/PathTracingDefaultLit.ush
+++ b/Engine/Shaders/Private/PathTracing/Material/PathTracingDefaultLit.ush
@@ -6,6 +6,7 @@
 
 #pragma once
 
+#include "PathTracingMaterialCommon.ush"
 #include "PathTracingFresnel.ush"
 
 float CalcLobeSelectionPdf(float3 V, FMaterialClosestHitPayload Payload)
@@ -19,16 +20,13 @@ float CalcLobeSelectionPdf(float3 V, FMaterialClosestHitPayload Payload)
 	return LobeSelectionProb(Diff, Spec);
 }
 
-void DefaultLit_EvalMaterial(
+FMaterialEval DefaultLit_EvalMaterial(
 	float3 IncomingDirection,
 	float3 OutgoingDirection,
-	FMaterialClosestHitPayload Payload,
-	out float3 OutWeight,
-	out float OutPdf
+	FMaterialClosestHitPayload Payload
 )
 {
-	OutWeight = 0;
-	OutPdf = 0;
+	FMaterialEval Result = NullMaterialEval();
 
 	const float3 V_World = -IncomingDirection;
 	const float3 L_World = OutgoingDirection;
@@ -49,7 +47,7 @@ void DefaultLit_EvalMaterial(
 	const float LobeSelectionPdf = CalcLobeSelectionPdf(V_World, Payload);
 
 	// Diffuse Lobe
-	AddLobeWithMIS(OutWeight, OutPdf, GetDiffuseColor(Payload), NoL / PI, LobeSelectionPdf);
+	AddLobeWithMIS(Result.Weight, Result.Pdf, GetDiffuseColor(Payload), NoL / PI, LobeSelectionPdf);
 
 	// Specular lobe
 	const float2 GGXResult = GGXEvalReflection(L, V, H, Alpha);
@@ -57,22 +55,17 @@ void DefaultLit_EvalMaterial(
 	const float3 SpecWeight = F * GGXResult.x;
 	const float SpecPdf = GGXResult.y;
 
-	AddLobeWithMIS(OutWeight, OutPdf, SpecWeight, SpecPdf, 1.0 - LobeSelectionPdf);
+	AddLobeWithMIS(Result.Weight, Result.Pdf, SpecWeight, SpecPdf, 1.0 - LobeSelectionPdf);
+
+	return Result;
 }
 
-void DefaultLit_SampleMaterial(
+FMaterialSample DefaultLit_SampleMaterial(
 	float3 RayDirection,
 	FMaterialClosestHitPayload Payload,
-	float4 RandSample,
-	out float3 OutDirection,
-	out float3 OutWeight,
-	out float OutPdf
+	float4 RandSample
 )
 {
-	OutDirection = 0;
-	OutWeight = 0;
-	OutPdf = 0;
-
 	float3 N = GetWorldNormal(Payload);
 	float3 V = -RayDirection;
 	
@@ -109,9 +102,10 @@ void DefaultLit_SampleMaterial(
 		H = ImportanceSampleVisibleGGX_aniso(UniformSampleDisk(RandSample.xy), Alpha, V).xyz;
 
 		L = reflect(-V, H);
-		if (L.z <= 0) {
+		if (L.z <= 0)
+		{
 			// invalid output direction, exit early
-			return;
+			return NullMaterialSample();
 		}
 	}
 
@@ -120,6 +114,9 @@ void DefaultLit_SampleMaterial(
 	const float NoL = saturate(L.z);
 	const float VoH = saturate(dot(V, H));
 
+	float3 OutWeight = 0;
+	float OutPdf = 0;
+
 	// Diffuse Lobe
 	AddLobeWithMIS(OutWeight, OutPdf, GetDiffuseColor(Payload), NoL / PI, LobeSelectionPdf);
 
@@ -131,5 +128,5 @@ void DefaultLit_SampleMaterial(
 	AddLobeWithMIS(OutWeight, OutPdf, SpecWeight, SpecPdf, 1.0 - LobeSelectionPdf);
 
 	// transform to world space
-	OutDirection = normalize(mul(L, Basis));
+	return CreateMaterialSample(normalize(mul(L, Basis)), OutWeight, OutPdf, 1.0);
 }
diff --git a/Engine/Shaders/Private/PathTracing/Material/PathTracingLambert.ush b/Engine/Shaders/Private/PathTracing/Material/PathTracingLambert.ush
--- a/Engine/Shaders/Private/PathTracing/Material/PathTracingLambert.ush
+++ b/Engine/Shaders/Private/PathTracing/Material/PathTracingLambert.ush
@@ -5,20 +5,15 @@ PathTracingLambert.usf: Lambertian BRDF sampling functions
 ===============================================================================================*/
 #pragma once
 
-void Lambert_SampleMaterial(
+FMaterialSample Lambert_SampleMaterial(
 	FMaterialClosestHitPayload Payload,
-	float4 RandSample,
-	out float3 OutDirection,
-	out float3 OutWeight,
-	out float OutPdf)
+	float4 RandSample)
 {
 	float3 N_World = GetWorldNormal(Payload);
 
 	float4 SampledValue = CosineSampleHemisphere(RandSample.xy);
-	OutDirection = TangentToWorld(SampledValue.xyz, N_World);
 
-	OutWeight = GetBaseColor(Payload);
-	OutPdf = SampledValue.w;
+	return CreateMaterialSample(TangentToWorld(SampledValue.xyz, N_World), GetBaseColor(Payload), SampledValue.w, 1.0);
 }
 
 float Lambert_Pdf(float3 OutgoingDirection, FMaterialClosestHitPayload Payload)
@@ -28,16 +23,13 @@ float Lambert_Pdf(float3 OutgoingDirection, FMaterialClosestHitPayload Payload)
 	return NoL / PI;
 }
 
-void Lambert_EvalMaterial(
+FMaterialEval Lambert_EvalMaterial(
 	float3 OutgoingDirection,
-	FMaterialClosestHitPayload Payload,
-	out float3 OutWeight,
-	out float OutPdf
+	FMaterialClosestHitPayload Payload
 )
 {
 	float3 N_World = GetWorldNormal(Payload);
 	float NoL = saturate(dot(N_World, OutgoingDirection));
 
-	OutWeight = GetBaseColor(Payload);
-	OutPdf = NoL / PI;
+	return CreateMaterialEval(GetBaseColor(Payload), NoL / PI);
 }
diff --git a/Engine/Shaders/Private/PathTracing/Material/PathTracingMaterialCommon.ush b/Engine/Shaders/Private/PathTracing/Material/PathTracingMaterialCommon.ush
--- a/Engine/Shaders/Private/PathTracing/Material/PathTracingMaterialCommon.ush
+++ b/Engine/Shaders/Private/PathTracing/Material/PathTracingMaterialCommon.ush
@@ -151,3 +151,48 @@ void AddLobeWithMIS(inout float3 Weight, inout float Pdf, float3 LobeWeight, flo
 		Pdf += LobePdf;
 	}
 }
+
+// The following structs are used as return types for the Eval/Sample methods
+
+struct FMaterialSample
+{
+	float3 Direction;
+	float3 Weight;
+	float Pdf;
+	float PositionBiasSign;
+};
+
+FMaterialSample NullMaterialSample()
+{
+	// return a zero initialized sample, for cases where we couldn't sample the material (such as rays below the horizon, etc..)
+	return (FMaterialSample)0;
+}
+
+FMaterialSample CreateMaterialSample(float3 Direction, float3 Weight, float Pdf, float PositionBiasSign)
+{
+	FMaterialSample Result;
+	Result.Direction = Direction;
+	Result.Weight = Weight;
+	Result.Pdf = Pdf;
+	Result.PositionBiasSign = PositionBiasSign;
+	return Result;
+}
+
+struct FMaterialEval
+{
+	float3 Weight;
+	float Pdf;
+};
+
+FMaterialEval NullMaterialEval()
+{
+	return (FMaterialEval)0;
+}
+
+FMaterialEval CreateMaterialEval(float3 Weight, float Pdf)
+{
+	FMaterialEval Result;
+	Result.Weight = Weight;
+	Result.Pdf = Pdf;
+	return Result;
+}
diff --git a/Engine/Shaders/Private/PathTracing/Material/PathTracingMaterialSampling.ush b/Engine/Shaders/Private/PathTracing/Material/PathTracingMaterialSampling.ush
--- a/Engine/Shaders/Private/PathTracing/Material/PathTracingMaterialSampling.ush
+++ b/Engine/Shaders/Private/PathTracing/Material/PathTracingMaterialSampling.ush
@@ -182,27 +182,21 @@ void AdjustPathAndMaterialRoughness(inout FMaterialClosestHitPayload Payload, in
 	PathRoughness = lerp(PathRoughness, Payload.Roughness, GlassPercent);
 }
 
-void SampleMaterial(
+FMaterialSample SampleMaterial(
 	float3 RayDirection,
 	FMaterialClosestHitPayload Payload,
 	float4 RandSample,
-	bool bForceRadianceProbe,
-	out float3 OutDirection,
-	out float3 OutWeight,
-	out float OutPdf,
-	out float OutPositionBiasSign)
+	bool bForceRadianceProbe
+)
 {
-	OutPositionBiasSign = 1;
-	
 	if (bForceRadianceProbe)
 	{
-		RadianceProbe_SampleMaterial(Payload, RandSample, OutDirection, OutWeight, OutPdf);
-		return;
+		return RadianceProbe_SampleMaterial(Payload, RandSample);
 	}
 	
 #if FORCE_ALL_MATERIALS_DIFFUSE
 	{
-		Lambert_SampleMaterial(Payload, RandSample, OutDirection, OutWeight, OutPdf);
+		return Lambert_SampleMaterial(Payload, RandSample);
 	}
 #else
 
@@ -210,10 +204,7 @@ void SampleMaterial(
 		Payload.BlendingMode == RAY_TRACING_BLEND_MODE_ALPHA_HOLDOUT)
 	{
 		// transparency only, nothing to sample
-		OutDirection = 0;
-		OutWeight = 0;
-		OutPdf = 0;
-		return;
+		return NullMaterialSample();
 	}
 
 	// Decide the overall weight of the BRDF portion vs. the translucent portion
@@ -245,6 +236,7 @@ void SampleMaterial(
 			break;
 	}
 
+	FMaterialSample Result = NullMaterialSample();
 	if (RandSample.x < BrdfProb)
 	{
 		// we randomly decided to sample the brdf portion of the material
@@ -258,79 +250,72 @@ void SampleMaterial(
 		// only case that should reach this point is the refraction model (which can be thin or solid)
 		if (Payload.ShadingModelID == SHADINGMODELID_THIN_TRANSLUCENT)
 		{
-			RoughThinGlass_SampleMaterial(RayDirection, Payload, RandSample, OutDirection, OutWeight, OutPdf, OutPositionBiasSign);
+			Result = RoughThinGlass_SampleMaterial(RayDirection, Payload, RandSample);
 		}
 		else if (Payload.Roughness == 0)
 		{
-			SpecularRefraction_SampleMaterial(RayDirection, Payload, RandSample, OutDirection, OutWeight, OutPdf, OutPositionBiasSign);
+			Result = SpecularRefraction_SampleMaterial(RayDirection, Payload, RandSample);
 		}
 		else
 		{
-			RoughGlass_SampleMaterial(RayDirection, Payload, RandSample, OutDirection, OutWeight, OutPdf, OutPositionBiasSign);
+			Result = RoughGlass_SampleMaterial(RayDirection, Payload, RandSample);
 		}
-		OutWeight *= TransWeight;
-		OutPdf *= 1 - BrdfProb;
-		return;
+		Result.Weight *= TransWeight;
+		Result.Pdf *= 1 - BrdfProb;
+		return Result;
 	}
 
 	switch (Payload.ShadingModelID)
 	{
 		case SHADINGMODELID_UNLIT:
 			// Nothing to do, material does not scatter light
-			OutDirection = 0;
-			OutWeight = 0;
-			OutPdf = 0;
-			return;
+			return Result;
 		case SHADINGMODELID_DEFAULT_LIT:
-			DefaultLit_SampleMaterial(RayDirection, Payload, RandSample, OutDirection, OutWeight, OutPdf);
+			Result = DefaultLit_SampleMaterial(RayDirection, Payload, RandSample);
 			break;
 		case SHADINGMODELID_CLOTH:
-			Cloth_SampleMaterial(RayDirection, Payload, RandSample, OutDirection, OutWeight, OutPdf);
+			Result = Cloth_SampleMaterial(RayDirection, Payload, RandSample);
 			break;
 		case SHADINGMODELID_CLEAR_COAT:
-			ClearCoat_SampleMaterial(RayDirection, Payload, RandSample, OutDirection, OutWeight, OutPdf);
+			Result = ClearCoat_SampleMaterial(RayDirection, Payload, RandSample);
 			break;
 		case SHADINGMODELID_TWOSIDED_FOLIAGE:
-			TwoSidedFoliage_SampleMaterial(RayDirection, Payload, RandSample, OutDirection, OutWeight, OutPdf, OutPositionBiasSign);
+			Result = TwoSidedFoliage_SampleMaterial(RayDirection, Payload, RandSample);
 			break;
 		default:
-			Lambert_SampleMaterial(Payload, RandSample, OutDirection, OutWeight, OutPdf);
+			Result = Lambert_SampleMaterial(Payload, RandSample);
 			break;
 	}
 
 	// final adjustment to weight / prob based on the stochastic decision for translucency earlier
-	OutWeight *= BrdfWeight;
-	OutPdf *= BrdfProb;
+	Result.Weight *= BrdfWeight;
+	Result.Pdf *= BrdfProb;
+	return Result;
 #endif
 }
 
-void EvalMaterial(
+FMaterialEval EvalMaterial(
 	float3 IncomingDirection,
 	float3 OutgoingDirection,
 	FMaterialClosestHitPayload Payload,
-	bool bForceRadianceProbe,
-	out float3 OutWeight,
-	out float OutPdf
+	bool bForceRadianceProbe
 )
 {
 	if (bForceRadianceProbe)
 	{
-		RadianceProbe_EvalMaterial(OutgoingDirection, Payload, OutWeight, OutPdf);
-		return;
+		return RadianceProbe_EvalMaterial(OutgoingDirection, Payload);
 	}
 	
 #if FORCE_ALL_MATERIALS_DIFFUSE
 	{
-		Lambert_EvalMaterial(OutgoingDirection, Payload, OutWeight, OutPdf);
+		return Lambert_EvalMaterial(OutgoingDirection, Payload);
 	}
 #else
 
 	if (Payload.BlendingMode == RAY_TRACING_BLEND_MODE_MODULATE ||
 		Payload.BlendingMode == RAY_TRACING_BLEND_MODE_ALPHA_HOLDOUT)
 	{
-		OutWeight = 0;
-		OutPdf = 0;
-		return;
+		return NullMaterialEval();
 	}
 
 	// Decide the overall weight of the BRDF portion vs. the translucent portion
@@ -367,54 +352,53 @@ void EvalMaterial(
 
 	// NOTE: the translucent portions are all singular -- so we can ignore them in the eval method for now as they have Pdf=+inf for a single direction only
 
+	FMaterialEval Result = NullMaterialEval();
 	switch (Payload.ShadingModelID)
 	{
 		case SHADINGMODELID_UNLIT:
-			OutWeight = 0.0;
-			OutPdf = 0.0;
 			break;
 		case SHADINGMODELID_DEFAULT_LIT:
-			DefaultLit_EvalMaterial(IncomingDirection, OutgoingDirection, Payload, OutWeight, OutPdf);
+			Result = DefaultLit_EvalMaterial(IncomingDirection, OutgoingDirection, Payload);
 			break;
 		case SHADINGMODELID_CLOTH:
-			Cloth_EvalMaterial(IncomingDirection, OutgoingDirection, Payload, OutWeight, OutPdf);
+			Result = Cloth_EvalMaterial(IncomingDirection, OutgoingDirection, Payload);
 			break;
 		case SHADINGMODELID_CLEAR_COAT:
-			ClearCoat_EvalMaterial(IncomingDirection, OutgoingDirection, Payload, OutWeight, OutPdf);
+			Result = ClearCoat_EvalMaterial(IncomingDirection, OutgoingDirection, Payload);
 			break;
 		case SHADINGMODELID_TWOSIDED_FOLIAGE:
-			TwoSidedFoliage_EvalMaterial(IncomingDirection, OutgoingDirection, Payload, OutWeight, OutPdf);
+			Result = TwoSidedFoliage_EvalMaterial(IncomingDirection, OutgoingDirection, Payload);
 			break;
 		default:
-			Lambert_EvalMaterial(OutgoingDirection, Payload, OutWeight, OutPdf);
+			Result = Lambert_EvalMaterial(OutgoingDirection, Payload);
 			break;
 	}
 
 	// final adjustment to weight / prob for translucency
-	OutWeight *= BrdfWeight;
-	OutPdf *= BrdfProb;
+	Result.Weight *= BrdfWeight;
+	Result.Pdf *= BrdfProb;
 
 	if (BrdfProb < 1.0)
 	{
 		// refraction model (only evaluate for rough case as singular can't be evaluated)
 		if (Payload.Roughness != 0)
 		{
-			float3 GlassWeight = 0;
-			float GlassPdf = 0;
+			FMaterialEval GlassEval;
 			if (Payload.ShadingModelID == SHADINGMODELID_THIN_TRANSLUCENT)
 			{
 				// just a thin slab (reflection-like lobe on both sides)
-				RoughThinGlass_EvalMaterial(IncomingDirection, OutgoingDirection, Payload, GlassWeight, GlassPdf);
+				GlassEval = RoughThinGlass_EvalMaterial(IncomingDirection, OutgoingDirection, Payload);
 			}
 			else
 			{
 				// real solid glass, reflection and refraction
-				RoughGlass_EvalMaterial(IncomingDirection, OutgoingDirection, Payload, GlassWeight, GlassPdf);
+				GlassEval = RoughGlass_EvalMaterial(IncomingDirection, OutgoingDirection, Payload);
 			}
 			
 
-			AddLobeWithMIS(OutWeight, OutPdf, GlassWeight, GlassPdf, 1.0 - BrdfProb);
+			AddLobeWithMIS(Result.Weight, Result.Pdf, GlassEval.Weight, GlassEval.Pdf, 1.0 - BrdfProb);
 		}
 	}
+	return Result;
 #endif
 }
diff --git a/Engine/Shaders/Private/PathTracing/Material/PathTracingRadianceProbe.ush b/Engine/Shaders/Private/PathTracing/Material/PathTracingRadianceProbe.ush
--- a/Engine/Shaders/Private/PathTracing/Material/PathTracingRadianceProbe.ush
+++ b/Engine/Shaders/Private/PathTracing/Material/PathTracingRadianceProbe.ush
@@ -48,12 +48,10 @@ float2 InverseConcentricMapping(float2 p)
 	return float2(x, y);
 }
 
-void RadianceProbe_SampleMaterial(
+FMaterialSample RadianceProbe_SampleMaterial(
 	FMaterialClosestHitPayload Payload,
-	float4 RandSample,
-	out float3 OutDirection,
-	out float3 OutWeight,
-	out float OutPdf)
+	float4 RandSample
+)
 {
 	float PdfAdjust = 1.0f; 
 	
@@ -118,12 +116,14 @@ void RadianceProbe_SampleMaterial(
 	float3 N_World = GetWorldNormal(Payload);
 
 	float4 SampledValue = CosineSampleHemisphereConcentric(RandSample.xy);
-	OutDirection = TangentToWorld(SampledValue.xyz, N_World);
+	float3 OutDirection = TangentToWorld(SampledValue.xyz, N_World);
 
 	PdfAdjust = max(PdfAdjust, 0.00001f);
-	OutPdf = SampledValue.w * PdfAdjust;
+	float OutPdf = SampledValue.w * PdfAdjust;
 	// As the name suggests, radiance probe has a combined throughput (Pdf * Weight) == 1
-	OutWeight = (OutPdf > 0.0f) ? (1.0f / OutPdf) : 0;
+	float OutWeight = (OutPdf > 0.0f) ? (1.0f / OutPdf) : 0;
+
+	return CreateMaterialSample(OutDirection, OutWeight, OutPdf, 1.0);
 }
 
 float RadianceProbe_Pdf(float3 OutgoingDirection, FMaterialClosestHitPayload Payload)
@@ -192,11 +192,9 @@ float RadianceProbe_Pdf(float3 OutgoingDirection, FMaterialClosestHitPayload Pay
 	return (NoL / PI) * PdfAdjust;
 }
 
-void RadianceProbe_EvalMaterial(
+FMaterialEval RadianceProbe_EvalMaterial(
 	float3 OutgoingDirection,
-	FMaterialClosestHitPayload Payload,
-	out float3 OutWeight,
-	out float OutPdf
+	FMaterialClosestHitPayload Payload
 )
 {
 	float3 N_World = GetWorldNormal(Payload);
@@ -262,27 +260,24 @@ void RadianceProbe_EvalMaterial(
 
 	PdfAdjust = max(PdfAdjust, 0.00001f);
 
-	OutPdf = NoL / PI * PdfAdjust;
+	float OutPdf = NoL / PI * PdfAdjust;
 	// As the name suggests, radiance probe has a combined throughput (Pdf * Weight) == 1
-	OutWeight = (OutPdf > 0.0f) ? (1.0f / OutPdf) : 0;
+	float OutWeight = (OutPdf > 0.0f) ? (1.0f / OutPdf) : 0;
+
+	return CreateMaterialEval(OutWeight, OutPdf);
 }
 
 #else
 
-void RadianceProbe_SampleMaterial(
+FMaterialSample RadianceProbe_SampleMaterial(
 	FMaterialClosestHitPayload Payload,
-	float4 RandSample,
-	out float3 OutDirection,
-	out float3 OutWeight,
-	out float OutPdf)
+	float4 RandSample)
 {
 	float3 N_World = GetWorldNormal(Payload);
 
 	float4 SampledValue = CosineSampleHemisphere(RandSample.xy);
-	OutDirection = TangentToWorld(SampledValue.xyz, N_World);
 
-	OutWeight = 1;
-	OutPdf = SampledValue.w;
+	return CreateMaterialSample(TangentToWorld(SampledValue.xyz, N_World), 1.0, SampledValue.w, 1.0);
 }
 
 float RadianceProbe_Pdf(float3 OutgoingDirection, FMaterialClosestHitPayload Payload)
@@ -292,18 +287,15 @@ float RadianceProbe_Pdf(float3 OutgoingDirection, FMaterialClosestHitPayload Pay
 	return NoL / PI;
 }
 
-void RadianceProbe_EvalMaterial(
+FMaterialEval RadianceProbe_EvalMaterial(
 	float3 OutgoingDirection,
-	FMaterialClosestHitPayload Payload,
-	out float3 OutWeight,
-	out float OutPdf
+	FMaterialClosestHitPayload Payload
 )
 {
 	float3 N_World = GetWorldNormal(Payload);
 	float NoL = saturate(dot(N_World, OutgoingDirection));
 
-	OutWeight = GetDiffuseColor(Payload);
-	OutPdf = NoL / PI;
+	return CreateMaterialEval(GetDiffuseColor(Payload), NoL / PI);
 }
 
 #endif
\ No newline at end of file
diff --git a/Engine/Shaders/Private/PathTracing/Material/PathTracingSolidGlass.ush b/Engine/Shaders/Private/PathTracing/Material/PathTracingSolidGlass.ush
--- a/Engine/Shaders/Private/PathTracing/Material/PathTracingSolidGlass.ush
+++ b/Engine/Shaders/Private/PathTracing/Material/PathTracingSolidGlass.ush
@@ -10,20 +10,11 @@ PathTracingSolidGlass.ush: Microfacet Refractive BSDF for solid glass
 #include "PathTracingFresnel.ush"
 #include "PathTracingGlossy.ush"
 
-void RoughGlass_SampleMaterial(
+FMaterialSample RoughGlass_SampleMaterial(
 	float3 RayDirection,
 	FMaterialClosestHitPayload Payload,
-	float4 RandSample,
-	out float3 OutDirection,
-	out float3 OutWeight,
-	out float OutPdf,
-	out float OutPositionBiasSign)
+	float4 RandSample)
 {
-	OutDirection = 0;
-	OutWeight = 0;
-	OutPdf = 0;
-	OutPositionBiasSign = 1;
-
 	// NOTE: IsFrontFace() determines polygon orientation, because the normal is always flipped towards in the incoming ray
 	bool bIsEntering = Payload.IsFrontFace();
 	float3 N_World = GetWorldNormal(Payload);
@@ -42,7 +33,7 @@ void RoughGlass_SampleMaterial(
 	if (NoV == 0)
 	{
 		// invalid grazing angle
-		return;
+		return NullMaterialSample();
 	}
 
 	float3 H = ImportanceSampleVisibleGGX_aniso(UniformSampleDisk(RandSample.xy), Alpha, V).xyz;
@@ -56,32 +47,24 @@ void RoughGlass_SampleMaterial(
 	{
 		const float2 GGXResult = GGXEvalRefraction(L, V, H, Alpha, Eta);
 
-		OutWeight = GGXResult.x;
-		OutPdf = F * GGXResult.y;
-
-		OutPositionBiasSign = -1;
+		// transform to world space
+		return CreateMaterialSample(normalize(mul(L, Basis)), GGXResult.x, F * GGXResult.y, -1.0);
 	}
 	else
 	{
 		const float2 GGXResult = GGXEvalReflection(L, V, H, Alpha);
-		OutWeight = GGXResult.x;
-		OutPdf = F * GGXResult.y;
-	}
 
-	// transform to world space
-	OutDirection = normalize(mul(L, Basis));
+		// transform to world space
+		return CreateMaterialSample(normalize(mul(L, Basis)), GGXResult.x, F * GGXResult.y, 1.0);
+	}
 }
 
-void RoughGlass_EvalMaterial(
+FMaterialEval RoughGlass_EvalMaterial(
 	float3 IncomingDirection,
 	float3 OutgoingDirection,
-	FMaterialClosestHitPayload Payload,
-	out float3 OutWeight,
-	out float OutPdf)
+	FMaterialClosestHitPayload Payload
+)
 {
-	OutWeight = 0;
-	OutPdf = 0;
-
 	// NOTE: IsFrontFace() determines polygon orientation, because the normal is always flipped towards in the incoming ray
 	bool bIsEntering = Payload.IsFrontFace();
 	float3 V_World = -IncomingDirection;
@@ -102,7 +85,7 @@ void RoughGlass_EvalMaterial(
 	if (V.z <= 0)
 	{
 		// invalid input
-		return;
+		return NullMaterialEval();
 	}
 
 	float NoV = saturate(V.z);
@@ -114,8 +97,7 @@ void RoughGlass_EvalMaterial(
 		const float F = FresnelReflectance(VoH, Eta);
 
 		const float2 GGXResult = GGXEvalReflection(L, V, H, Alpha);
-		OutWeight = GGXResult.x;
-		OutPdf = F * GGXResult.y;
+		return CreateMaterialEval(GGXResult.x, F * GGXResult.y);
 	}
 	else if (L.z < 0)
 	{
@@ -131,8 +113,8 @@ void RoughGlass_EvalMaterial(
 		{
 			const float2 GGXResult = GGXEvalRefraction(L, V, Htn, Alpha, Eta);
 
-			OutWeight = GGXResult.x;
-			OutPdf = F * GGXResult.y;
+			return CreateMaterialEval(GGXResult.x, F * GGXResult.y);
 		}
 	}
+	return NullMaterialEval();
 }
diff --git a/Engine/Shaders/Private/PathTracing/Material/PathTracingSpecularTransmission.ush b/Engine/Shaders/Private/PathTracing/Material/PathTracingSpecularTransmission.ush
--- a/Engine/Shaders/Private/PathTracing/Material/PathTracingSpecularTransmission.ush
+++ b/Engine/Shaders/Private/PathTracing/Material/PathTracingSpecularTransmission.ush
@@ -9,14 +9,10 @@ PathTracingSpecularTransmission.ush: Refractive BSDFs (pure singular)
 #include "PathTracingMaterialCommon.ush"
 #include "PathTracingGlossy.ush"
 
-void SpecularRefraction_SampleMaterial(
+FMaterialSample SpecularRefraction_SampleMaterial(
 	float3 RayDirection,
 	FMaterialClosestHitPayload Payload,
-	float4 RandSample,
-	out float3 OutDirection,
-	out float3 OutWeight,
-	out float OutPdf,
-	out float OutPositionBiasSign)
+	float4 RandSample)
 {
 	// NOTE: IsFrontFace() determines polygon orientation, because the normal is always flipped towards in the incoming ray
 	bool bIsEntering = Payload.IsFrontFace();
@@ -26,24 +22,20 @@ void SpecularRefraction_SampleMaterial(
 	float Ior = Payload.Ior;
 	float Eta = bIsEntering ? Ior : 1.0 / Ior;
 
-	// Throughput is 1 since we pick reflect/refract according to the fresnel term
-	OutWeight = 1.0;
-	OutPdf = POSITIVE_INFINITY;
 
-	float F = 0;
+	float F = 0, BiasSign = 1.0;
+	float3 OutDirection = 0;
 	if (SampleRefraction(RayDirection, WorldNormal, Eta, RandSample.x, OutDirection, F))
 	{
-		OutPositionBiasSign = -1;
-	}
-	else
-	{
-		OutPositionBiasSign = 1;
+		BiasSign = -1;
 	}
+
+	// Throughput is 1 since we pick reflect/refract according to the fresnel term
+	return CreateMaterialSample(OutDirection, 1.0, POSITIVE_INFINITY, BiasSign);
 } 
 
-void SpecularTransmission_EvalMaterial(out float3 OutThroughput, out float OutPdf)
+FMaterialEval SpecularTransmission_EvalMaterial()
 {
-	OutThroughput = 0;
-	OutPdf = 0;
+	return NullMaterialEval();
 }
 
diff --git a/Engine/Shaders/Private/PathTracing/Material/PathTracingTestBrdfs.usf b/Engine/Shaders/Private/PathTracing/Material/PathTracingTestBrdfs.usf
--- a/Engine/Shaders/Private/PathTracing/Material/PathTracingTestBrdfs.usf
+++ b/Engine/Shaders/Private/PathTracing/Material/PathTracingTestBrdfs.usf
@@ -63,6 +63,7 @@ void TestBrdfIntegrityCS(
 
 	float4 RandSample = RandomSequence_GenerateSample4D(RandSequence);
 
+#if 0
 	if (BrdfType == 0) 
 	{
 		Lambert_SampleMaterial(Payload, RandSample, Wi, SampledBRDF, SampledPDF);
@@ -78,6 +79,7 @@ void TestBrdfIntegrityCS(
 		PDF = EvalPDF;
 	}
 	else
+#endif
 	{
 		// Should never be here, generate invalid outgoing to trigger an error
 		Wi = float3(0.0, 0.0, 0.0); 
diff --git a/Engine/Shaders/Private/PathTracing/Material/PathTracingTestGenerateWiSamples.usf b/Engine/Shaders/Private/PathTracing/Material/PathTracingTestGenerateWiSamples.usf
--- a/Engine/Shaders/Private/PathTracing/Material/PathTracingTestGenerateWiSamples.usf
+++ b/Engine/Shaders/Private/PathTracing/Material/PathTracingTestGenerateWiSamples.usf
@@ -50,7 +50,7 @@ void TestGenerateWiSamplesCS(
 	float SampledPDF;
 
 	float4 RandSample = RandomSequence_GenerateSample4D(RandSequence);
-
+#if 0
 	if (BrdfType == 0)
 	{
 		Lambert_SampleMaterial(Payload, RandSample, Wi, SampledBRDF, SampledPDF);
@@ -60,6 +60,7 @@ void TestGenerateWiSamplesCS(
 		DefaultLit_SampleMaterial(Wo, Payload, RandSample, Wi, SampledBRDF, SampledPDF);
 	}
 	else
+#endif
 	{
 		// Should never be here, generate invalid outgoing to trigger an error
 		Wi = float3(0.0, 0.0, 0.0); 
diff --git a/Engine/Shaders/Private/PathTracing/Material/PathTracingTestIntegrateHemispherePatch.usf b/Engine/Shaders/Private/PathTracing/Material/PathTracingTestIntegrateHemispherePatch.usf
--- a/Engine/Shaders/Private/PathTracing/Material/PathTracingTestIntegrateHemispherePatch.usf
+++ b/Engine/Shaders/Private/PathTracing/Material/PathTracingTestIntegrateHemispherePatch.usf
@@ -65,8 +65,8 @@ float IntegratePdfInSphericalPatch(float Theta0, float Phi0, float Theta1, float
 
 			float3 Wi = normalize(float3(SinTheta*CosPhi, SinTheta*SinPhi, CosTheta));
 
-			float PDF;
-
+			float PDF = 0;
+#if 0
 			if (BrdfType == 0)
 			{
 				PDF = Lambert_Pdf(Wi, Payload);
@@ -77,6 +77,7 @@ float IntegratePdfInSphericalPatch(float Theta0, float Phi0, float Theta1, float
 				DefaultLit_EvalMaterial(Wo, Wi, Payload, Weight, PDF);
 			}
 			else
+#endif
 			{
 				// Should never be here, generate invalid outgoing to trigger an error
 				return -1.0;
diff --git a/Engine/Shaders/Private/PathTracing/Material/PathTracingTestPdfIntegration.usf b/Engine/Shaders/Private/PathTracing/Material/PathTracingTestPdfIntegration.usf
--- a/Engine/Shaders/Private/PathTracing/Material/PathTracingTestPdfIntegration.usf
+++ b/Engine/Shaders/Private/PathTracing/Material/PathTracingTestPdfIntegration.usf
@@ -56,6 +56,7 @@ void TestPDFIntegratesToOneCS(
 
 	float PDF;
 	
+#if 0
 	if (BrdfType == 0)
 	{
 		PDF = Lambert_Pdf(Wi, Payload);
@@ -66,6 +67,7 @@ void TestPDFIntegratesToOneCS(
 		DefaultLit_EvalMaterial(Wo, Wi, Payload, Weight, PDF);
 	}
 	else
+#endif
 	{
 		// Should never be here, generate invalid outgoing to trigger an error
 		Wi = float3(0.0, 0.0, 0.0);
diff --git a/Engine/Shaders/Private/PathTracing/Material/PathTracingThinGlass.ush b/Engine/Shaders/Private/PathTracing/Material/PathTracingThinGlass.ush
--- a/Engine/Shaders/Private/PathTracing/Material/PathTracingThinGlass.ush
+++ b/Engine/Shaders/Private/PathTracing/Material/PathTracingThinGlass.ush
@@ -48,16 +48,12 @@ SlabWeights ComputeThinSlabWeights(float3 SlabColor, float CosV, float Eta)
 	return Result;
 }
 
-void RoughThinGlass_EvalMaterial(
+FMaterialEval RoughThinGlass_EvalMaterial(
 	float3 IncomingDirection,
 	float3 OutgoingDirection,
-	FMaterialClosestHitPayload Payload,
-	out float3 OutWeight,
-	out float OutPdf)
+	FMaterialClosestHitPayload Payload
+)
 {
-	OutWeight = 0;
-	OutPdf = 0;
-
 	float3 V_World = -IncomingDirection;
 	float3 L_World = OutgoingDirection;
 	float3 N_World = GetWorldNormal(Payload);
@@ -73,7 +69,7 @@ void RoughThinGlass_EvalMaterial(
 	if (V.z <= 0)
 	{
 		// invalid input
-		return;
+		return NullMaterialEval();
 	}
 
 	const bool IsTransmission = L.z < 0;
@@ -90,36 +86,23 @@ void RoughThinGlass_EvalMaterial(
 	const float ProbR = LobeSelectionProb(Result.Reflected, Result.Transmitted);
 
 	const float2 GGXResult = GGXEvalReflection(L, V, H, Alpha);
-	OutWeight = GGXResult.x;
-	OutPdf = GGXResult.y;
 
 	if (IsTransmission)
 	{
-		OutWeight *= Result.Transmitted / (1.0 - ProbR);
-		OutPdf *= 1.0 - ProbR;
+		return CreateMaterialEval(GGXResult.x / (1.0 - ProbR) * Result.Transmitted, GGXResult.y * (1.0 - ProbR));
 	}
 	else
 	{
-		OutWeight *= Result.Reflected / ProbR;
-		OutPdf *= ProbR;
+		return CreateMaterialEval(GGXResult.x / ProbR * Result.Reflected, GGXResult.y * ProbR);
 	}
 }
 
 
-void RoughThinGlass_SampleMaterial(
+FMaterialSample RoughThinGlass_SampleMaterial(
 	float3 RayDirection,
 	FMaterialClosestHitPayload Payload,
-	float4 RandSample,
-	out float3 OutDirection,
-	out float3 OutWeight,
-	out float OutPdf,
-	out float OutPositionBiasSign)
+	float4 RandSample)
 {
-	OutDirection = 0;
-	OutWeight = 0;
-	OutPdf = 0;
-	OutPositionBiasSign = 1;
-
 	float3 N_World = GetWorldNormal(Payload);
 	float3 V_World = -RayDirection;
 
@@ -134,7 +117,7 @@ void RoughThinGlass_SampleMaterial(
 	if (NoV == 0)
 	{
 		// invalid grazing angle
-		return;
+		return NullMaterialSample();
 	}
 
 	float3 H = ImportanceSampleVisibleGGX_aniso(UniformSampleDisk(RandSample.xy), Alpha, V).xyz;
@@ -143,9 +126,10 @@ void RoughThinGlass_SampleMaterial(
 	float VoH = saturate(dot(V, H));
 
 	float3 L = reflect(-V, H);
-	if (L.z <= 0) {
+	if (L.z <= 0)
+	{
 		// invalid output direction, exit early
-		return;
+		return NullMaterialSample();
 	}
 
 	float NoL = saturate(L.z);
@@ -155,23 +139,24 @@ void RoughThinGlass_SampleMaterial(
 	const float ProbR = LobeSelectionProb(Result.Reflected, Result.Transmitted);
 
 	const float2 GGXResult = GGXEvalReflection(L, V, H, Alpha);
-	OutWeight = GGXResult.x;
-	OutPdf = GGXResult.y;
+	float3 OutWeight = GGXResult.x;
+	float OutPdf = GGXResult.y;
 
 	// decide between reflection/transmission
 	if (RandSample.z < ProbR)
 	{
 		// reflected ray
 		OutWeight *= Result.Reflected / ProbR;
-		OutDirection = normalize(mul(L, Basis));
 		OutPdf *= ProbR;
+		float3 OutDirection = normalize(mul(L, Basis));
+		return CreateMaterialSample(OutDirection, OutWeight, OutPdf, 1.0);
 	}
 	else
 	{
 		// transmitted ray
 		OutWeight *= Result.Transmitted / (1.0 - ProbR);
-		OutDirection = normalize(mul(float3(L.xy, -L.z), Basis)); // flip reflection to other side
 		OutPdf *= 1.0 - ProbR;
-		OutPositionBiasSign = -1;
+		float3 OutDirection = normalize(mul(float3(L.xy, -L.z), Basis)); // flip reflection to other side
+		return CreateMaterialSample(OutDirection, OutWeight, OutPdf, -1.0);
 	}
 }
diff --git a/Engine/Shaders/Private/PathTracing/Material/PathTracingTwoSidedFoliage.ush b/Engine/Shaders/Private/PathTracing/Material/PathTracingTwoSidedFoliage.ush
--- a/Engine/Shaders/Private/PathTracing/Material/PathTracingTwoSidedFoliage.ush
+++ b/Engine/Shaders/Private/PathTracing/Material/PathTracingTwoSidedFoliage.ush
@@ -13,20 +13,17 @@ float CalcTransmissionSelectionPdf(
 	FMaterialClosestHitPayload Payload
 )
 {
-	// pick back lobe proportionally to its contribution, and leave the rest to the front side
-	// TODO: estimate this better based on lobe colors ...
-	float3 Subsurf = GetSubsurfaceColor(Payload);
-	return max(Subsurf.x, max(Subsurf.y, Subsurf.z));
+	// pick back lobe proportionally to its contribution relative to front side
+	float3 Diff = GetDiffuseColor(Payload);
+	float3 Spec = GetSpecularColor(Payload);
+	float3 Back = GetSubsurfaceColor(Payload);
+	return LobeSelectionProb(Back, Diff + Spec);
 }
 
-void TwoSidedFoliage_SampleMaterial(
+FMaterialSample TwoSidedFoliage_SampleMaterial(
 	float3 RayDirection,
 	FMaterialClosestHitPayload Payload,
-	float4 RandSample,
-	out float3 OutDirection,
-	out float3 OutWeight,
-	out float OutPdf,
-	out float OutPositionBiasSign
+	float4 RandSample
 )
 {
 	float TransmissionSelectionPdf = CalcTransmissionSelectionPdf(RayDirection, Payload);
@@ -38,30 +35,25 @@ void TwoSidedFoliage_SampleMaterial(
 		float3 SubsurfaceColor = GetSubsurfaceColor(Payload);
 
 		float4 SampledValue = CosineSampleHemisphere(RandSample.yz);
-		OutDirection = TangentToWorld(SampledValue.xyz, -N);
 
-		OutWeight = SubsurfaceColor / TransmissionSelectionPdf;
-		OutPdf = TransmissionSelectionPdf * SampledValue.w;
-		OutPositionBiasSign = -1.0;
+		return CreateMaterialSample(TangentToWorld(SampledValue.xyz, -N), SubsurfaceColor / TransmissionSelectionPdf, TransmissionSelectionPdf * SampledValue.w, -1.0);
 	}
 	else
 	{
 		RandSample.x -= TransmissionSelectionPdf;
-		RandSample.x /= (1.0 - TransmissionSelectionPdf);
+		RandSample.x /= 1.0 - TransmissionSelectionPdf;
 
-		DefaultLit_SampleMaterial(RayDirection, Payload, RandSample, OutDirection, OutWeight, OutPdf);
-		OutWeight /= 1.0 - TransmissionSelectionPdf;
-		OutPdf *= 1.0 - TransmissionSelectionPdf;
-		OutPositionBiasSign = 1.0;
+		FMaterialSample Result = DefaultLit_SampleMaterial(RayDirection, Payload, RandSample);
+		Result.Weight /= 1.0 - TransmissionSelectionPdf;
+		Result.Pdf *= 1.0 - TransmissionSelectionPdf;
+		return Result;
 	}
 }
 
-void TwoSidedFoliage_EvalMaterial(
+FMaterialEval TwoSidedFoliage_EvalMaterial(
 	float3 IncomingDirection,
 	float3 OutgoingDirection,
-	FMaterialClosestHitPayload Payload,
-	out float3 OutWeight,
-	out float OutPdf
+	FMaterialClosestHitPayload Payload
 )
 {
 	float TransmissionSelectionPdf = CalcTransmissionSelectionPdf(IncomingDirection, Payload);
@@ -72,16 +64,20 @@ void TwoSidedFoliage_EvalMaterial(
 	{
 		// Diffuse transmission (the implementation in TwoSidedBxDF does not appear to be a real BxDF)
 
-		float3 SubsurfaceColor = GetSubsurfaceColor(Payload);
+		if (TransmissionSelectionPdf > 0)
+		{
+			float3 SubsurfaceColor = GetSubsurfaceColor(Payload);
+			float LambertPdf = -NoL / PI;
 
-		float LambertPdf = -NoL / PI;
-		OutWeight = TransmissionSelectionPdf > 0 ? SubsurfaceColor / TransmissionSelectionPdf : 0;
-		OutPdf = TransmissionSelectionPdf * LambertPdf;
+			return CreateMaterialEval(SubsurfaceColor / TransmissionSelectionPdf, TransmissionSelectionPdf * LambertPdf);
+		}
 	}
-	else
+	else if (TransmissionSelectionPdf < 1)
 	{
-		DefaultLit_EvalMaterial(IncomingDirection, OutgoingDirection, Payload, OutWeight, OutPdf);
-		OutWeight /= 1.0 - TransmissionSelectionPdf;
-		OutPdf *= 1.0 - TransmissionSelectionPdf;
+		FMaterialEval Result = DefaultLit_EvalMaterial(IncomingDirection, OutgoingDirection, Payload);
+		Result.Weight /= 1.0 - TransmissionSelectionPdf;
+		Result.Pdf *= 1.0 - TransmissionSelectionPdf;
+		return Result;
 	}
+	return NullMaterialEval();
 }
\ No newline at end of file
diff --git a/Engine/Shaders/Private/PathTracing/PathTracing.usf b/Engine/Shaders/Private/PathTracing/PathTracing.usf
--- a/Engine/Shaders/Private/PathTracing/PathTracing.usf
+++ b/Engine/Shaders/Private/PathTracing/PathTracing.usf
@@ -287,6 +287,7 @@ RAY_TRACING_ENTRY_RAYGEN(PathTracingMainRG)
 		// Update origin
 		Ray.Origin = Payload.WorldPos;
 
+
 		// Make sure the surface normal points toward the viewer
 		if (dot(Ray.Direction, Payload.WorldNormal) > 0.0)
 		{
@@ -349,42 +350,37 @@ RAY_TRACING_ENTRY_RAYGEN(PathTracingMainRG)
 					if (any(RadianceOverPdf > 0))
 					{
 						// Evaluate material
-						float3 MaterialWeight;
-						float MaterialPdf = 0.0;
-						EvalMaterial(Ray.Direction, LightRay.Direction, Payload, false, MaterialWeight, MaterialPdf);
-
-						float3 MaterialEval = MaterialWeight * MaterialPdf;
-
-						float MISWeight = MISMode == 2 ? MISWeightRobust(LightPdf, MaterialPdf) : 1.0;
+						FMaterialEval MaterialEval = EvalMaterial(Ray.Direction, LightRay.Direction, Payload, false);
 
 						// Record the contribution
-						AccumulateRadiance(Radiance, MISWeight * PathThroughput * RadianceOverPdf * MaterialEval);
+						float3 LightContrib = PathThroughput * RadianceOverPdf * MaterialEval.Weight * MaterialEval.Pdf;
+						if (MISMode == 2)
+						{
+							LightContrib *= MISWeightRobust(LightPdf, MaterialEval.Pdf);
+						}
+						AccumulateRadiance(Radiance, LightContrib);
 					}
 				}
 			}
 		}
 
 		// Sample material
-		float3 Direction;
-		float3 MaterialWeight;
-		float SignedPositionBias;
-		float MaterialSamplePdf;
-		SampleMaterial(Ray.Direction, Payload, RandSample, false, Direction, MaterialWeight, MaterialSamplePdf, SignedPositionBias);
+		FMaterialSample MaterialSample = SampleMaterial(Ray.Direction, Payload, RandSample, false);
 
-		if (MaterialSamplePdf < 0 || asuint(MaterialSamplePdf) > 0x7F800000)
+		if (MaterialSample.Pdf < 0 || asuint(MaterialSample.Pdf) > 0x7F800000)
 		{
 			// Pdf became invalid (either negative or NaN)
 			Radiance = float3(1, 0, 1);
 			break;
 		}
 
-		if (!(MaterialSamplePdf > 0))
+		if (!(MaterialSample.Pdf > 0))
 		{
 			// No valid direction -- we are done
 			break;
 		}
 
-		float3 NextPathThroughput = PathThroughput * MaterialWeight;
+		float3 NextPathThroughput = PathThroughput * MaterialSample.Weight;
 		if (!any(NextPathThroughput > 0))
 		{
 			// no energy left in this path
@@ -418,10 +414,10 @@ RAY_TRACING_ENTRY_RAYGEN(PathTracingMainRG)
 		}
 
 		// Update ray according to material sample
-		Ray.Direction = Direction;
+		Ray.Direction = MaterialSample.Direction;
 		Ray.TMin = 0;
 		Ray.TMax = POSITIVE_INFINITY;
-		ApplyPositionBias(Ray, SignedPositionBias * Payload.WorldNormal, PathTracingData.MaxNormalBias);
+		ApplyPositionBias(Ray, MaterialSample.PositionBiasSign * Payload.WorldNormal, PathTracingData.MaxNormalBias);
 
 		// If we are using Material sampling for lights
 		if (MISMode != 1)
@@ -450,7 +446,7 @@ RAY_TRACING_ENTRY_RAYGEN(PathTracingMainRG)
 					PreviousCdfValue = LightPickingCdf[LightId];
 
 					LightPdf *= LightPickPdf;
-					LightContrib *= MISWeightRobust(MaterialSamplePdf, LightPdf);
+					LightContrib *= MISWeightRobust(MaterialSample.Pdf, LightPdf);
 				}
 
 				if (any(LightContrib > 0))
