Cherry pick fix for recent deferred decal breakage.

diff --git a/Engine/Shaders/Private/DecalCommon.ush b/Engine/Shaders/Private/DecalCommon.ush
--- a/Engine/Shaders/Private/DecalCommon.ush
+++ b/Engine/Shaders/Private/DecalCommon.ush
@@ -76,13 +76,18 @@ void DecalCommonOutput(inout FPixelShaderIn In, inout FPixelShaderOut Out, float
 	#if DECAL_RENDERSTAGE == 2 || DECAL_RENDERSTAGE == 3
 	{
 		#if DECAL_BLEND_MODE == DECALBLENDMODEID_STAIN
-			Out.MRT[2].rgb *= Opacity;
+			// BaseColor MRT location changes if we have Normal MRT or not.
+			#if DECAL_OUTPUT_NORMAL
+				Out.MRT[3].rgb *= Opacity;
+			#else
+				Out.MRT[2].rgb *= Opacity;
+			#endif
 		#endif
 
-		// Note: As of 4.24, normal (GBufferA) is no longer written to MRT[1], so GBufferB is in MRT1 and GBufferC is in MRT2.
-		Out.MRT[0].a = Opacity;	// Emissive
-		Out.MRT[1].a = Opacity;	// Metallic, Specular, Roughness
-		Out.MRT[2].a = Opacity;	// BaseColor
+		Out.MRT[0].a = Opacity;
+		Out.MRT[1].a = Opacity;
+		Out.MRT[2].a = Opacity;
+		Out.MRT[3].a = Opacity;
 	}
 	#endif // DECAL_RENDERSTAGE
 #endif //DECAL_RENDERTARGET_COUNT > 1
