Strata SingleLayerWater:
- Removed use of gbuffer in SingleLayerWater passes
- SLW compute envbrdf in the compositing pass from strata water material.
- It is no longer needed to use the classification buffer, avoiding an extra classification pass just for water
- To do this, the SSR pass then also read the water BSDF directly
- Made room  in the top layer data for SSS data coming soon

diff --git a/Engine/Shaders/Private/BasePassPixelShader.usf b/Engine/Shaders/Private/BasePassPixelShader.usf
--- a/Engine/Shaders/Private/BasePassPixelShader.usf
+++ b/Engine/Shaders/Private/BasePassPixelShader.usf
@@ -993,6 +993,8 @@ void FPixelShaderInOut_MainPS(
 #endif
 
 #if STRATA_ENABLED && STRATA_INLINE_SINGLELAYERWATER
+	StrataPixelHeader.IsSingleLayerWater = true;
+
 	// Override GBuffer data with Strata SLW water BSDF to run forward shadfing code.
 	// STRATA_TODO: run the shading through a strata path? (by adding a special BSDF?)
 	SanitizeStrataSingleLayerWater(PixelMaterialInputs.FrontMaterial.Layers[0].BSDFs[0]);
diff --git a/Engine/Shaders/Private/SSRT/SSRTReflections.usf b/Engine/Shaders/Private/SSRT/SSRTReflections.usf
--- a/Engine/Shaders/Private/SSRT/SSRTReflections.usf
+++ b/Engine/Shaders/Private/SSRT/SSRTReflections.usf
@@ -32,6 +32,7 @@ SamplerState SceneColorSampler;
 Texture2D HZB;
 SamplerState HZBSampler;
 
+uint ShouldReflectOnlyWater;
 
 float GetRoughness(const in FGBufferData GBuffer)
 {
@@ -133,10 +134,37 @@ void ScreenSpaceReflections(
 #if STRATA_ENABLED
 
 	uint2 PixelCoord = SvPosition.xy;
-	const FStrataClassification StrataData = StrataUnpackClassificationData(Strata.ClassificationTexture.Load(uint3(PixelCoord, 0)));
-	const FStrataTopLayerData TopLayerData = StrataUnpackTopLayerData(Strata.TopLayerTexture.Load(uint3(PixelCoord, 0)));
 
-	float3 N = TopLayerData.Normal;
+	float3 N = 0;
+	float Roughness = 0.0;
+	bool bNoMaterial = true;
+	BRANCH
+	if (ShouldReflectOnlyWater)
+	{
+		// We do not rerun material classification after the Strata material buffer has been patched with SLW material.
+		// So in this case we simply load the Header and water BSDF directly because we know there can only be a single one.
+		FStrataAddressing StrataAddressing = GetStrataPixelDataByteOffset(PixelCoord, uint2(View.BufferSizeAndInvSize.xy), Strata.MaxBytesPerPixel);
+		FStrataPixelHeader StrataPixelHeader = UnpackStrataHeaderIn(Strata.MaterialLobesBuffer, StrataAddressing);
+		if (StrataPixelHeader.BSDFCount > 0 && StrataPixelHeader.IsSingleLayerWater)
+		{
+			const FStrataBSDF SLWaterBSDF = UnpackStrataBSDFIn(Strata.MaterialLobesBuffer, StrataAddressing);
+			N = StrataGetBSDFSharedBasis(StrataPixelHeader, SLWaterBSDF, StrataAddressing)[2];
+			Roughness = SLW_ROUGHNESS(SLWaterBSDF);
+			bNoMaterial = false;
+		}
+	}
+	else
+	{
+		const FStrataClassification StrataData = StrataUnpackClassificationData(Strata.ClassificationTexture.Load(uint3(PixelCoord, 0)));
+		const FStrataTopLayerData TopLayerData = StrataUnpackTopLayerData(Strata.TopLayerTexture.Load(uint3(PixelCoord, 0)));
+
+		// For now, SSR is only used to trace the top material layer.
+		// STRATA_TODO try stochastic selection
+		N = TopLayerData.Normal;
+		Roughness = TopLayerData.Roughness;
+		bNoMaterial = !StrataData.bIsStrataMaterial;
+	}
+
 	float DeviceZ = SampleDeviceZFromSceneTextures(UV);
 	float SceneDepth = ConvertFromDeviceZ(DeviceZ);
 
@@ -145,10 +173,6 @@ void ScreenSpaceReflections(
 
 	// STRATA_TODO Handle anisotrpoy using ModifyGGXAnisotropicNormalRoughness or similar approach
 
-	// For now, SSR is only used to trace the top material layer.
-	// STRATA_TODO try stochastic selection
-	const float Roughness = TopLayerData.Roughness;
-	const bool bNoMaterial = !StrataData.bIsStrataMaterial;
 
 #else // STRATA_ENABLED
 
diff --git a/Engine/Shaders/Private/SingleLayerWaterComposite.usf b/Engine/Shaders/Private/SingleLayerWaterComposite.usf
--- a/Engine/Shaders/Private/SingleLayerWaterComposite.usf
+++ b/Engine/Shaders/Private/SingleLayerWaterComposite.usf
@@ -43,9 +43,6 @@ struct SingleLayerWaterCompositeOutput
 
 SingleLayerWaterCompositeOutput SingleLayerWaterComposite(float2 BufferUV, float2 ScreenPosition, float4 SvPosition)
 {
-	// Sample scene textures.
-	FGBufferData GBuffer = GetGBufferDataFromSceneTextures(BufferUV);
-	uint ShadingModelID = GBuffer.ShadingModelID;
 	// No AO or DFAO
 
 	const float3 OneThird3 = float3(1.0f / 3.0f, 1.0f / 3.0f, 1.0f / 3.0f);
@@ -63,16 +60,79 @@ SingleLayerWaterCompositeOutput SingleLayerWaterComposite(float2 BufferUV, float
 	const float PreExposure = 1.f;
 #endif
 
-	bool bIsStrataSLW = false;
 #if STRATA_ENABLED
-	const FStrataClassification StrataData = StrataUnpackClassificationData(Strata.ClassificationTexture.Load(uint3(SvPosition.xy, 0)));
-	if (StrataData.bIsStrataMaterial)
+
+	// Water do nto read the classification data to avoid running it a second time when water is enabled.
+	FStrataAddressing StrataAddressing = GetStrataPixelDataByteOffset(SvPosition.xy, uint2(View.BufferSizeAndInvSize.xy), Strata.MaxBytesPerPixel);
+	FStrataPixelHeader StrataPixelHeader = UnpackStrataHeaderIn(Strata.MaterialLobesBuffer, StrataAddressing);
+	if (StrataPixelHeader.BSDFCount > 0 && StrataPixelHeader.IsSingleLayerWater)
 	{
-		bIsStrataSLW = StrataHasShadingModel(StrataData.ShadingModels, STRATA_BSDF_TYPE_SINGLELAYERWATER);
-	}
-#endif
+		float DeviceZ = SampleDeviceZFromSceneTextures(BufferUV);
+		float SceneDepth = ConvertFromDeviceZ(DeviceZ);
+
+		const float2 SceneNoWaterUV = clamp(BufferUV, SceneNoWaterMinMaxUV.xy, SceneNoWaterMinMaxUV.zw);
+
+		float OpaqueDepth = SceneNoWaterDepthTexture.SampleLevel(SceneNoWaterDepthSampler, SceneNoWaterUV, 0).x * SINGLE_LAYER_WATER_DEPTH_SCALE;
+		float WaterDepth = SceneDepth;
+
+		float DeltaDepth = CameraIsUnderWater ? WaterDepth : OpaqueDepth - WaterDepth;	// Inverted depth
+		if (DeltaDepth > 0.0)
+		{
+			const FStrataBSDF SLWaterBSDF = UnpackStrataBSDFIn(Strata.MaterialLobesBuffer, StrataAddressing);
+			const float3 WorldNormal = StrataGetBSDFSharedBasis(StrataPixelHeader, SLWaterBSDF, StrataAddressing)[2];
+			const float3 F0 = ComputeF0(SLW_SPECULAR(SLWaterBSDF), SLW_BASECOLOR(SLWaterBSDF), SLW_METALLIC(SLWaterBSDF));
+			const float Roughness = SLW_ROUGHNESS(SLWaterBSDF);
+			const float SafeRoughness = MakeRoughnessSafe(SLW_ROUGHNESS(SLWaterBSDF));
 
-	if (ShadingModelID == SHADINGMODELID_SINGLELAYERWATER || bIsStrataSLW)
+			// Compute the sky reflection contribution
+			float3 Reflection = 0.0f;
+
+			// TODO MAKE COMMON WITH ReflectionEnvironment?
+			float3 WorldPosition = mul(float4(ScreenPosition * WaterDepth, WaterDepth, 1), View.ScreenToWorld).xyz;
+			float3 CameraToPixel = normalize(WorldPosition - View.WorldCameraOrigin);
+			float3 ReflectionVector = reflect(CameraToPixel, WorldNormal);
+			float IndirectIrradiance = 0.0;
+			float IndirectSpecularOcclusion = 1.0f;
+			float3 ExtraIndirectSpecular = 0;
+			uint CaptureDataStartIndex = 0;
+			uint NumCulledReflectionCaptures = 0;
+			//
+			float3 N = WorldNormal;
+			float3 V = -CameraToPixel;
+			float3 R = 2 * dot(V, N) * N - V;
+			float NoV = saturate(dot(N, V));
+			// Point lobe in off-specular peak direction
+			R = GetOffSpecularPeakReflectionDir(N, R, SafeRoughness);
+			const bool bCompositeSkylight = true;
+			Reflection += PreExposure * CompositeReflectionCapturesAndSkylight(
+				1,
+				WorldPosition,
+				R,
+				SafeRoughness,
+				IndirectIrradiance,
+				IndirectSpecularOcclusion,
+				ExtraIndirectSpecular,
+				NumCulledReflectionCaptures,
+				CaptureDataStartIndex,
+				0,
+				bCompositeSkylight);
+
+			// Then combine reflection with SSR
+			float4 SSR = Texture2DSample(ScreenSpaceReflectionsTexture, ScreenSpaceReflectionsSampler, BufferUV);
+			Reflection = SSR.rgb + Reflection * (1 - SSR.a);
+
+			const float3 BSDFCoverage = 1.0f;	// Since the SLW BSDF is always isolated and composited onto the strata material buffer, this is assumed to be 1 here.
+
+			// The specular over the water in completely in the hand of the user. We do not fade out metalness for instance.
+			float3 EnvBrdfValue = BSDFCoverage * EnvBRDF(SLW_SPECULAR(SLWaterBSDF), SafeRoughness, saturate(NoV));
+
+#else // STRATA_ENABLED
+
+	// Sample scene textures.
+	FGBufferData GBuffer = GetGBufferDataFromSceneTextures(BufferUV);
+	uint ShadingModelID = GBuffer.ShadingModelID;
+
+	if (ShadingModelID == SHADINGMODELID_SINGLELAYERWATER)
 	{
 		const float2 SceneNoWaterUV = clamp(BufferUV, SceneNoWaterMinMaxUV.xy, SceneNoWaterMinMaxUV.zw);
 
@@ -121,19 +181,11 @@ SingleLayerWaterCompositeOutput SingleLayerWaterComposite(float2 BufferUV, float
 
 			// Apply the BRDF reflection integral to reflection. 
 			// BRDF/Fresnel is already applied for the under water part in the scene.
-			float3 EnvBrdf = 0.0;
-#if STRATA_ENABLED
-			if (bIsStrataSLW)
-			{
-				const FStrataTopLayerData TopLayerData = StrataUnpackTopLayerData(Strata.TopLayerTexture.Load(uint3(SvPosition.xy, 0)));
-				EnvBrdf = TopLayerData.EnvBRDF;
-			}
-			else
-#endif
-			{
-				EnvBrdf = EnvBRDF(GBuffer.SpecularColor, GBuffer.Roughness, NoV);
-			}
-			Reflection *= EnvBrdf;
+			float3 EnvBrdfValue = EnvBRDF(GBuffer.SpecularColor, GBuffer.Roughness, NoV);
+
+#endif // STRATA_ENABLED
+
+			Reflection *= EnvBrdfValue;
 
 			// Soft fading near shor to avoid seeing triangles.
 			const float ShoreOpacity = saturate(DeltaDepth * 0.02f);
@@ -142,7 +194,7 @@ SingleLayerWaterCompositeOutput SingleLayerWaterComposite(float2 BufferUV, float
 			// Apply transmittance to reflection if under water
 			if (CameraIsUnderWater)
 			{
-				TransmittanceToScene *= 1.0 - EnvBrdf;
+				TransmittanceToScene *= 1.0 - EnvBrdfValue;
 
 				// Using default under water material: compute transmittance and scattering
 				float3 WaterMediumScattering = 0.0f;
@@ -203,16 +255,14 @@ void WaterTileCatergorisationCS(uint3 ThreadId : SV_DispatchThreadID, uint3 Grou
 			{
 				uint2 PixelPos = (ThreadId.xy + uint2(i, j)) + View.ViewRectMin.xy;
 				float2 BufferUV = (PixelPos + 0.5f) * View.BufferSizeAndInvSize.zw;
-				FGBufferData GBuffer = GetGBufferDataFromSceneTextures(BufferUV);
-
-				bContainsWater = bContainsWater || GBuffer.ShadingModelID == SHADINGMODELID_SINGLELAYERWATER;
 
 			#if STRATA_ENABLED
-				const FStrataClassification StrataData = StrataUnpackClassificationData(Strata.ClassificationTexture.Load(uint3(PixelPos,0)));
-				if (StrataData.bIsStrataMaterial)
-				{
-					bContainsWater = bContainsWater || StrataHasShadingModel(StrataData.ShadingModels, STRATA_BSDF_TYPE_SINGLELAYERWATER);
-				}
+				FStrataAddressing StrataAddressing = GetStrataPixelDataByteOffset(PixelPos, uint2(View.BufferSizeAndInvSize.xy), Strata.MaxBytesPerPixel);
+				FStrataPixelHeader StrataPixelHeader = UnpackStrataHeaderIn(Strata.MaterialLobesBuffer, StrataAddressing);
+				bContainsWater = bContainsWater || (StrataPixelHeader.BSDFCount > 0 && StrataPixelHeader.IsSingleLayerWater);
+			#else
+				FGBufferData GBuffer = GetGBufferDataFromSceneTextures(BufferUV);
+				bContainsWater = bContainsWater || GBuffer.ShadingModelID == SHADINGMODELID_SINGLELAYERWATER;
 			#endif
 			}
 		}
@@ -231,15 +281,14 @@ void WaterTileCatergorisationCS(uint3 ThreadId : SV_DispatchThreadID, uint3 Grou
 	{
 		uint2 PixelPos = ThreadId.xy + View.ViewRectMin.xy;
 		float2 BufferUV = (PixelPos + 0.5f) * View.BufferSizeAndInvSize.zw;
-		FGBufferData GBuffer = GetGBufferDataFromSceneTextures(BufferUV);
-		bContainsWater = GBuffer.ShadingModelID == SHADINGMODELID_SINGLELAYERWATER;
 
 	#if STRATA_ENABLED
-		const FStrataClassification StrataData = StrataUnpackClassificationData(Strata.ClassificationTexture.Load(uint3(PixelPos, 0)));
-		if (StrataData.bIsStrataMaterial)
-		{
-			bContainsWater = bContainsWater || StrataHasShadingModel(StrataData.ShadingModels, STRATA_BSDF_TYPE_SINGLELAYERWATER);
-		}
+		FStrataAddressing StrataAddressing = GetStrataPixelDataByteOffset(PixelPos, uint2(View.BufferSizeAndInvSize.xy), Strata.MaxBytesPerPixel);
+		FStrataPixelHeader StrataPixelHeader = UnpackStrataHeaderIn(Strata.MaterialLobesBuffer, StrataAddressing);
+		bContainsWater = bContainsWater || (StrataPixelHeader.BSDFCount > 0 && StrataPixelHeader.IsSingleLayerWater);
+	#else
+		FGBufferData GBuffer = GetGBufferDataFromSceneTextures(BufferUV);
+		bContainsWater = bContainsWater || GBuffer.ShadingModelID == SHADINGMODELID_SINGLELAYERWATER;
 	#endif
 	}
 
diff --git a/Engine/Shaders/Private/Strata/Strata.ush b/Engine/Shaders/Private/Strata/Strata.ush
--- a/Engine/Shaders/Private/Strata/Strata.ush
+++ b/Engine/Shaders/Private/Strata/Strata.ush
@@ -59,7 +59,8 @@ struct FStrataBSDF
 #define HEADER_BIT_COUNT_AO						8
 #define HEADER_BIT_COUNT_SKIP_PRECSHADOW_MASK	1
 #define HEADER_BIT_COUNT_ZERO_PRECSHADOW_MASK	1
-#define HEADER_BIT_COUNT___UNUSED___			11
+#define HEADER_BIT_COUNT_ISSINGLELAYERWATER		1
+#define HEADER_BIT_COUNT___UNUSED___			10
 
 #define HEADER_BIT_OFFSET_BSDF_COUNT			0
 #define HEADER_BIT_OFFSET_SHAREDNORMAL_COUNT	(HEADER_BIT_OFFSET_BSDF_COUNT				+ HEADER_BIT_COUNT_BSDF_COUNT)
@@ -67,6 +68,7 @@ struct FStrataBSDF
 #define HEADER_BIT_OFFSET_AO					(HEADER_BIT_OFFSET_SHAREDNORMAL_TYPE		+ HEADER_BIT_COUNT_SHAREDNORMAL_TYPE)
 #define HEADER_BIT_OFFSET_SKIP_PRECSHADOW_MASK	(HEADER_BIT_OFFSET_AO						+ HEADER_BIT_COUNT_AO)
 #define HEADER_BIT_OFFSET_ZERO_PRECSHADOW_MASK	(HEADER_BIT_OFFSET_SKIP_PRECSHADOW_MASK		+ HEADER_BIT_COUNT_SKIP_PRECSHADOW_MASK)
+#define HEADER_BIT_OFFSET_ISSINGLELAYERWATER	(HEADER_BIT_OFFSET_ZERO_PRECSHADOW_MASK		+ HEADER_BIT_COUNT_ZERO_PRECSHADOW_MASK)
 
 #define HEADER_GETBSDFCOUNT(X)					READ_BITS( X, HEADER_BIT_COUNT_BSDF_COUNT, HEADER_BIT_OFFSET_BSDF_COUNT)
 #define HEADER_SETBSDFCOUNT(X, NewValue)		WRITE_BITS(X, HEADER_BIT_COUNT_BSDF_COUNT, HEADER_BIT_OFFSET_BSDF_COUNT, NewValue)
@@ -86,6 +88,9 @@ struct FStrataBSDF
 #define HEADER_GETZEROPRECSHADOWMASK(X)			READ_BITS(  X, HEADER_BIT_COUNT_ZERO_PRECSHADOW_MASK, HEADER_BIT_OFFSET_ZERO_PRECSHADOW_MASK)
 #define HEADER_SETZEROPRECSHADOWMASK(X, NewValue)WRITE_BITS(X, HEADER_BIT_COUNT_ZERO_PRECSHADOW_MASK, HEADER_BIT_OFFSET_ZERO_PRECSHADOW_MASK, NewValue)
 
+#define HEADER_GETISSINGLELAYERWATER(X)			READ_BITS(  X, HEADER_BIT_COUNT_ISSINGLELAYERWATER, HEADER_BIT_OFFSET_ISSINGLELAYERWATER)
+#define HEADER_SETISSINGLELAYERWATER(X, NewValue)WRITE_BITS(X, HEADER_BIT_COUNT_ISSINGLELAYERWATER, HEADER_BIT_OFFSET_ISSINGLELAYERWATER, NewValue)
+
 // Bit count for each component stored in a Strata BSDF state
 #define STATE_BIT_COUNT_BSDF					3
 #define STATE_BIT_COUNT_SHAREDNORMALID			2
@@ -262,9 +267,7 @@ FStrataClassification StrataUnpackClassificationData(uint In)
 struct FStrataTopLayerData
 {
 	float3 Normal;
-	float3 BaseColor;
 	float  Roughness;
-	float3 EnvBRDF;
 
 	// transient data when accumulating top layer data
 	float AccumRoughnessCount;
@@ -286,28 +289,25 @@ struct FStrataSubsurfaceData
 void FinalizeTopLayerData(inout FStrataTopLayerData Out)
 {
 	Out.Normal = normalize(Out.Normal);
-	Out.BaseColor = saturate(Out.BaseColor);
 	Out.Roughness = Out.AccumRoughnessCount > 0.0f ? saturate(Out.Roughness / Out.AccumRoughnessCount) : 1.f;
-	Out.EnvBRDF = Out.BaseColor; // When not subsurface, BaseColor stores the split-sum approximation EnvBRDF. STRATA_TODO better approximation
 }
 
 uint2 StrataPackTopLayerData(FStrataTopLayerData In)
 {
 	uint2 Out = 0;
-	Out.x = PackColorLinearToGamma2AlphaLinear(float4(In.BaseColor, In.Roughness));
+	float3 Unused = 0.0f;
+	Out.x = PackColorLinearToGamma2AlphaLinear(float4(Unused, In.Roughness));
 	Out.y = StrataPackNormal(In.Normal);
 	return Out;
 }
 
 FStrataTopLayerData StrataUnpackTopLayerData(uint2 In)
 {
-	const float4 BaseColorAndRoughness = UnpackColorGamma2ToLinearAlphaLinear(In.x);
+	const float4 UnusedAndRoughness = UnpackColorGamma2ToLinearAlphaLinear(In.x);
 
 	FStrataTopLayerData Out = (FStrataTopLayerData)0;
 	Out.Normal = StrataUnpackNormal(In.y);
-	Out.BaseColor = BaseColorAndRoughness.xyz;
-	Out.Roughness = BaseColorAndRoughness.w;
-	Out.EnvBRDF = Out.BaseColor; // When not subsurface, BaseColor stores the split-sum approximation EnvBRDF. STRATA_TODO better approximation
+	Out.Roughness = UnusedAndRoughness.w;
 	return Out;
 }
 
@@ -650,6 +650,7 @@ struct FStrataPixelHeader
 	float MaterialAO;
 	bool SkipPrecShadowMask;
 	bool ZeroPrecShadowMask;
+	bool IsSingleLayerWater;	// Used to avoid loading BSDF or Categorisation data when generating tiles containing SLW material. Avoid running another categorisation.
 };
 
 FStrataPixelHeader InitialiseStrataPixelHeader()
@@ -665,6 +666,7 @@ FStrataPixelHeader InitialiseStrataPixelHeader()
 	Out.MaterialAO = 1.0f;
 	Out.SkipPrecShadowMask = false;
 	Out.ZeroPrecShadowMask = false;
+	Out.IsSingleLayerWater = false;
 	return Out;
 }
 
@@ -1194,9 +1196,6 @@ void PackStrataOut(
 					// Setup SSR for the top layer
 					TopLayer.Roughness += BSDFCoverageAvg * SLAB_ROUGHNESSX(BSDF);
 					TopLayer.AccumRoughnessCount += BSDFCoverageAvg;
-
-					const float SafeRoughness = MakeRoughnessSafe(SLAB_ROUGHNESSX(BSDF));
-					TopLayer.BaseColor += BSDFCoverage * EnvBRDFApprox(F0, F90, SafeRoughness, saturate(ContextNoLight.NoV));
 				}
 
 				// Select the first/toppest layer with subsurface scattering
@@ -1226,9 +1225,6 @@ void PackStrataOut(
 				{
 					TopLayer.Roughness += BSDFCoverageAvg * SHEEN_ROUGHNESS(BSDF);
 					TopLayer.AccumRoughnessCount += BSDFCoverageAvg;
-
-					const float SafeRoughness = MakeRoughnessSafe(SHEEN_ROUGHNESS(BSDF));
-					TopLayer.BaseColor += BSDFCoverage * EnvBRDFApprox(F0, SafeRoughness, saturate(ContextNoLight.NoV));
 				}
 
 				LayerCoverage	+= BSDFCoverage;
@@ -1258,11 +1254,6 @@ void PackStrataOut(
 
 					TopLayer.Roughness += BSDFCoverageAvg * Roughness;
 					TopLayer.AccumRoughnessCount += BSDFCoverageAvg;
-
-					// The specular over the water in completely in the hand of the user. We do not fade out metallness for instance.
-					const float3 F0 = ComputeF0(SLW_SPECULAR(BSDF), SLW_BASECOLOR(BSDF), SLW_METALLIC(BSDF));
-					const float SafeRoughness = MakeRoughnessSafe(Roughness);
-					TopLayer.BaseColor += BSDFCoverage * EnvBRDFApprox(F0, SafeRoughness, saturate(ContextNoLight.NoV));
 				}
 
 				EmissiveColor += BSDF_GETEMISSIVE(BSDF) * Throughput;
@@ -1292,7 +1283,6 @@ void PackStrataOut(
 
 			// Roughness is needed for SSR, base color represents EnvBRDF
 			GBuffer.Roughness = TopLayer.Roughness;
-			GBuffer.BaseColor = TopLayer.BaseColor;
 			
 			// Approximative way to be able to have SSR produce colored specular reflections
 			GBuffer.Specular = 1;
@@ -1441,6 +1431,7 @@ uint PackStrataHeaderIn(uint InBSDFCount, FStrataPixelHeader InHeader)
 
 	HEADER_SETSKIPPRECSHADOWMASK(Out, InHeader.SkipPrecShadowMask ? 1u : 0u);
 	HEADER_SETZEROPRECSHADOWMASK(Out, InHeader.ZeroPrecShadowMask ? 1u : 0u);
+	HEADER_SETISSINGLELAYERWATER(Out, InHeader.IsSingleLayerWater ? 1u : 0u);
 	return Out;
 }
 
@@ -1473,6 +1464,7 @@ FStrataPixelHeader UnpackStrataHeaderIn(ByteAddressBuffer StrataBuffer, inout FS
 
 	Out.SkipPrecShadowMask = HEADER_GETSKIPPRECSHADOWMASK(PackedHeader);
 	Out.ZeroPrecShadowMask = HEADER_GETSKIPPRECSHADOWMASK(PackedHeader);
+	Out.IsSingleLayerWater = HEADER_GETISSINGLELAYERWATER(PackedHeader);
 
 	return Out;
 }
diff --git a/Engine/Shaders/Private/Strata/StrataMaterialClassification.usf b/Engine/Shaders/Private/Strata/StrataMaterialClassification.usf
--- a/Engine/Shaders/Private/Strata/StrataMaterialClassification.usf
+++ b/Engine/Shaders/Private/Strata/StrataMaterialClassification.usf
@@ -57,9 +57,6 @@ void ComputePostProcessBSDFData(
 			// Setup SSR for the top layer
 			TopLayer.Roughness += BSDFCoverageAvg * SLAB_ROUGHNESSX(BSDF);
 			TopLayer.AccumRoughnessCount += BSDFCoverageAvg;
-
-			const float SafeRoughness = MakeRoughnessSafe(SLAB_ROUGHNESSX(BSDF));
-			TopLayer.BaseColor += BSDFCoverage * EnvBRDFApprox(F0, F90, SafeRoughness, saturate(ContextNoLight.NoV));
 		}
 
 		Classification.bIsSimple = Classification.bIsSimple && IsStrataFastPathCompatible(BSDF);
@@ -86,9 +83,6 @@ void ComputePostProcessBSDFData(
 		{
 			TopLayer.Roughness += SHEEN_ROUGHNESS(BSDF);
 			TopLayer.AccumRoughnessCount += 1.0f;
-
-			const float SafeRoughness = MakeRoughnessSafe(SHEEN_ROUGHNESS(BSDF));
-			TopLayer.BaseColor += BSDFCoverage * EnvBRDFApprox(F0, SafeRoughness, saturate(ContextNoLight.NoV));
 		}
 		break;
 	}
@@ -114,11 +108,6 @@ void ComputePostProcessBSDFData(
 
 			TopLayer.Roughness += BSDFCoverageAvg * Roughness;
 			TopLayer.AccumRoughnessCount += BSDFCoverageAvg;
-
-			// The specular over the water in completely in the hand of the user. We do not fade out metallness for instance.
-			const float3 F0 = ComputeF0(SLW_SPECULAR(BSDF), SLW_BASECOLOR(BSDF), SLW_METALLIC(BSDF));
-			const float SafeRoughness = MakeRoughnessSafe(Roughness);
-			TopLayer.BaseColor += BSDFCoverageAvg * EnvBRDFApprox(F0, SafeRoughness, saturate(ContextNoLight.NoV));
 		}
 		break;
 	}
diff --git a/Engine/Shaders/Private/Strata/StrataVisualize.usf b/Engine/Shaders/Private/Strata/StrataVisualize.usf
--- a/Engine/Shaders/Private/Strata/StrataVisualize.usf
+++ b/Engine/Shaders/Private/Strata/StrataVisualize.usf
@@ -191,6 +191,7 @@ void VisualizeMaterialPS(
 #define PRINTMATERIALAO()	PRINTBSDFCHAR(_A_); PRINTBSDFCHAR(_O_); PRINTBSDFCHAR(_SPC_); PRINTFLOAT1(Header.MaterialAO);
 #define PRINTMATERIALSKIPPRECSHADOWMASK()	PRINTBSDFCHAR(_S_); PRINTBSDFCHAR(_K_); PRINTBSDFCHAR(_I_); PRINTBSDFCHAR(_P_); PRINTBSDFCHAR(_P_); PRINTBSDFCHAR(_R_); PRINTBSDFCHAR(_E_); PRINTBSDFCHAR(_S_); PRINTBSDFCHAR(_H_); PRINTBSDFCHAR(_A_); PRINTBSDFCHAR(_D_); PRINTBSDFCHAR(_SPC_); PRINTBSDFCHAR(_0_+(Header.SkipPrecShadowMask?1:0));
 #define PRINTMATERIALZEROPRECSHADOWMASK()	PRINTBSDFCHAR(_Z_); PRINTBSDFCHAR(_E_); PRINTBSDFCHAR(_R_); PRINTBSDFCHAR(_O_); PRINTBSDFCHAR(_P_); PRINTBSDFCHAR(_R_); PRINTBSDFCHAR(_E_); PRINTBSDFCHAR(_S_); PRINTBSDFCHAR(_H_); PRINTBSDFCHAR(_A_); PRINTBSDFCHAR(_D_); PRINTBSDFCHAR(_SPC_); PRINTBSDFCHAR(_0_+(Header.ZeroPrecShadowMask?1:0));
+#define PRINTMATERIALSLW()	PRINTBSDFCHAR(_S_); PRINTBSDFCHAR(_L_); PRINTBSDFCHAR(_W_); PRINTBSDFCHAR(_SPC_); PRINTBSDFCHAR(_0_+(Header.IsSingleLayerWater?1:0));
 
 #define NEXTLINE  TopLeft.y += 8; RectMax.x = max(RectMax.x, TopLeft.x); RectMax.y = max(RectMax.y, TopLeft.y); TopLeft.x = TopLeftInit.x; 
 #define HORIZONTALSEPARATION  TopLeft.y += 4;
@@ -208,6 +209,8 @@ void VisualizeMaterialPS(
 		NEXTLINE;
 		PRINTMATERIALZEROPRECSHADOWMASK();
 		NEXTLINE;
+		PRINTMATERIALSLW();
+		NEXTLINE;
 		HORIZONTALSEPARATION;
 #else
 		NEXTLINE;
@@ -216,6 +219,7 @@ void VisualizeMaterialPS(
 		NEXTLINE;
 		NEXTLINE;
 		NEXTLINE;
+		NEXTLINE;
 		HORIZONTALSEPARATION;
 #endif
 
diff --git a/Engine/Source/Runtime/Renderer/Private/ScreenSpaceRayTracing.cpp b/Engine/Source/Runtime/Renderer/Private/ScreenSpaceRayTracing.cpp
--- a/Engine/Source/Runtime/Renderer/Private/ScreenSpaceRayTracing.cpp
+++ b/Engine/Source/Runtime/Renderer/Private/ScreenSpaceRayTracing.cpp
@@ -344,6 +344,7 @@ BEGIN_SHADER_PARAMETER_STRUCT(FSSRPassCommonParameters, )
 	SHADER_PARAMETER(FVector4, HZBUvFactorAndInvFactor)
 	SHADER_PARAMETER(FVector4, PrevScreenPositionScaleBias)
 	SHADER_PARAMETER(float, PrevSceneColorPreExposureCorrection)
+	SHADER_PARAMETER(uint32, ShouldReflectOnlyWater)
 	SHADER_PARAMETER_RDG_TEXTURE(Texture2D, SceneColor)
 	SHADER_PARAMETER_SAMPLER(SamplerState, SceneColorSampler)
 	SHADER_PARAMETER_RDG_TEXTURE(Texture2D, HZB)
@@ -1027,6 +1028,7 @@ void RenderScreenSpaceReflections(
 	ESSRQuality SSRQuality,
 	bool bDenoiser,
 	IScreenSpaceDenoiser::FReflectionsInputs* DenoiserInputs,
+	bool bSingleLayerWater,
 	FTiledScreenSpaceReflection* TiledScreenSpaceReflection)
 {
 	FRDGTextureRef InputColor = CurrentSceneColor;
@@ -1166,6 +1168,7 @@ void RenderScreenSpaceReflections(
 			PassParameters->ScreenSpaceRayTracingDebugOutput = CreateScreenSpaceRayTracingDebugUAV(GraphBuilder, DenoiserInputs->Color->Desc, TEXT("DebugSSR"), true);
 		}
 		PassParameters->PrevSceneColorPreExposureCorrection = InputColor != CurrentSceneColor ? View.PreExposure / View.PrevViewInfo.SceneColorPreExposure : 1.0f;
+		PassParameters->ShouldReflectOnlyWater = bSingleLayerWater ? 1u : 0u;
 		
 		PassParameters->SceneColor = InputColor;
 		PassParameters->SceneColorSampler = GSSRHalfResSceneColor ? TStaticSamplerState<SF_Bilinear>::GetRHI() : TStaticSamplerState<SF_Point>::GetRHI();
diff --git a/Engine/Source/Runtime/Renderer/Private/ScreenSpaceRayTracing.h b/Engine/Source/Runtime/Renderer/Private/ScreenSpaceRayTracing.h
--- a/Engine/Source/Runtime/Renderer/Private/ScreenSpaceRayTracing.h
+++ b/Engine/Source/Runtime/Renderer/Private/ScreenSpaceRayTracing.h
@@ -101,6 +101,7 @@ void RenderScreenSpaceReflections(
 	ESSRQuality SSRQuality,
 	bool bDenoiser,
 	IScreenSpaceDenoiser::FReflectionsInputs* DenoiserInputs,
+	bool bSingleLayerWater = false,
 	FTiledScreenSpaceReflection* TiledScreenSpaceReflection = nullptr);
 
 bool IsScreenSpaceDiffuseIndirectSupported(const FViewInfo& View);
diff --git a/Engine/Source/Runtime/Renderer/Private/SingleLayerWaterRendering.cpp b/Engine/Source/Runtime/Renderer/Private/SingleLayerWaterRendering.cpp
--- a/Engine/Source/Runtime/Renderer/Private/SingleLayerWaterRendering.cpp
+++ b/Engine/Source/Runtime/Renderer/Private/SingleLayerWaterRendering.cpp
@@ -419,13 +419,6 @@ void FDeferredShadingSceneRenderer::RenderSingleLayerWaterReflections(
 		FIntVector ViewRes(View.ViewRect.Width(), View.ViewRect.Height(), 1);
 		FIntVector TiledViewRes = FIntVector::DivideAndRoundUp(ViewRes, TiledScreenSpaceReflection.TileSize);
 
-		// We need to update the classification texture once again because the GBuffer has been patched with single layer water data.
-		if (Strata::IsStrataEnabled() && bRunTiled)
-		{
-			Strata::AddStrataMaterialClassificationPass(GraphBuilder, SceneTextures, Views);
-		}
-
-
 		if (bRunTiled)
 		{
 			TiledScreenSpaceReflection.DispatchIndirectParametersBuffer = GraphBuilder.CreateBuffer(FRDGBufferDesc::CreateIndirectDesc<FRHIDrawIndirectParameters>(), TEXT("SLW.WaterIndirectDrawParameters"));
@@ -498,7 +491,6 @@ void FDeferredShadingSceneRenderer::RenderSingleLayerWaterReflections(
 				DenoiserInputs.RayImaginaryDepth = GraphBuilder.CreateTexture(Desc, TEXT("SLW.RayTracingReflectionsImaginaryDepth"));
 			}
 
-			bool bReflectOnlyWater = true;
 			RenderRayTracingReflections(
 				GraphBuilder,
 				SceneTextures,
@@ -548,8 +540,9 @@ void FDeferredShadingSceneRenderer::RenderSingleLayerWaterReflections(
 			RDG_EVENT_SCOPE(GraphBuilder, "Water ScreenSpaceReflections(Quality=%d)", int32(SSRQuality));
 
 			const bool bDenoise = false;
+			const bool bSingleLayerWater = true;
 			ScreenSpaceRayTracing::RenderScreenSpaceReflections(
-				GraphBuilder, SceneTextureParameters, SceneTextures.Color.Resolve, View, SSRQuality, bDenoise, &DenoiserInputs, bRunTiled ? &TiledScreenSpaceReflection : nullptr);
+				GraphBuilder, SceneTextureParameters, SceneTextures.Color.Resolve, View, SSRQuality, bDenoise, &DenoiserInputs, bSingleLayerWater, bRunTiled ? &TiledScreenSpaceReflection : nullptr);
 
 			ReflectionsColor = DenoiserInputs.Color;
 
