Refactor ligth sampling and trace functions in the path tracer to use struct output

This refactors the light side similarly to how materials are now handled.

diff --git a/Engine/Plugins/Experimental/GPULightmass/Shaders/Private/LightmapPathTracing.usf b/Engine/Plugins/Experimental/GPULightmass/Shaders/Private/LightmapPathTracing.usf
--- a/Engine/Plugins/Experimental/GPULightmass/Shaders/Private/LightmapPathTracing.usf
+++ b/Engine/Plugins/Experimental/GPULightmass/Shaders/Private/LightmapPathTracing.usf
@@ -245,10 +245,14 @@ void PathTracingKernel(
 		
 		if (MISMode != 1 && Bounce != 0) // In GPULightmass it is not possible to hit any light on its first fake ray
 		{
-			float HitT = Payload.IsMiss() ? Ray.TMax : Payload.HitT;
 			// Loop over lights to capture their contribution
 			// #dxr_todo: if we have lots of lights, having some hierarchical structure would be better ....
 			float PreviousCdfValue = 0.0;
+			RayDesc LightRay = Ray;
+			if (Payload.IsHit())
+			{
+				LightRay.TMax = Payload.HitT;
+			}
 			for (uint LightId = 0; LightId < SceneLightCount; ++LightId)
 			{
 				// Camera rays should only do the rest of the loop if we have visible lights OR if
@@ -256,9 +260,7 @@ void PathTracingKernel(
 				if (Bounce == 0 && !IsEnvironmentLight(LightId))
 					continue;
 
-				float4 LightResult = TraceLight(Ray, LightId, HitT);
-				float3 LightRadiance = LightResult.xyz;
-				float LightPdf = LightResult.w;
+				FLightHit LightResult = TraceLight(LightRay, LightId);
 
 				float LightPickPdf = LightPickingCdf[LightId] - PreviousCdfValue;
 				PreviousCdfValue = LightPickingCdf[LightId];
@@ -266,7 +268,7 @@ void PathTracingKernel(
 				// GPULightmass records direct contribution for only static lights
 				if (Bounce > 1 || (Bounce == 1 && !IsStationary(LightId)))
 				{
-					AccumulateRadiance(Radiance, PathThroughput * LightRadiance * MISWeightRobust(PrevMaterialPdf, LightPdf * LightPickPdf));
+					AccumulateRadiance(Radiance, PathThroughput * LightResult.Radiance * MISWeightRobust(PrevMaterialPdf, LightResult.Pdf * LightPickPdf));
 				}
 			}
 		}
@@ -387,30 +389,22 @@ void PathTracingKernel(
 
 				SelectLight(RandSample.x, LightPickingCdf, LightId, LightPickPdf);
 
-				RayDesc LightRay;
-				LightRay.Origin = WorldPos;
-				LightRay.TMin = 0;
-				float3 RadianceOverPdf;
-				float LightPdf;
-				SampleLight(LightId, RandSample, WorldPos, WorldNormal,
-					LightRay.Direction,
-					LightRay.TMax,
-					RadianceOverPdf,
-					LightPdf);
-				RadianceOverPdf /= LightPickPdf;
-				LightPdf *= LightPickPdf;
-
-				if (LightPdf > 0)
+				FLightSample LightSample = SampleLight(LightId, RandSample, WorldPos, WorldNormal);
+
+				LightSample.RadianceOverPdf /= LightPickPdf;
+				LightSample.Pdf *= LightPickPdf;
+
+				if (LightSample.Pdf > 0)
 				{
-					float SignedPositionBias = 1.0;
-					if (ENABLE_TRANSMISSION && IsMaterialTransmissive(Payload))
-					{
-						// for transmissive materials, bias the position to the other side of the surface if the light is coming from behind
-						float NoL = dot(Payload.WorldNormal, LightRay.Direction);
-						SignedPositionBias = sign(NoL);
-					}
-					ApplyPositionBias(LightRay, SignedPositionBias * Payload.WorldNormal, 0.01f);
+					// for transmissive materials, bias the position to the other side of the surface if the light is coming from behind
+					const float SignedPositionBias = IsTransmissiveMaterial ? sign(dot(Payload.WorldNormal, LightSample.Direction)) : 1.0;
 
+					RayDesc LightRay;
+					LightRay.Origin = WorldPos;
+					LightRay.TMin = 0;
+					LightRay.Direction = LightSample.Direction;
+					LightRay.TMax = LightSample.Distance;
+					ApplyPositionBias(LightRay, SignedPositionBias* Payload.WorldNormal, 0.01f);
 
 					const uint RayFlags = RAY_FLAG_ACCEPT_FIRST_HIT_AND_END_SEARCH | RAY_FLAG_SKIP_CLOSEST_HIT_SHADER;
 					const uint InstanceInclusionMask = RAY_TRACING_MASK_ALL;
@@ -423,12 +417,12 @@ void PathTracingKernel(
 					if (PackedPayload.IsMiss())
 					{
 						// Evaluate material
-						FMaterialEval MaterialEval = EvalMaterial(Ray.Direction, LightRay.Direction, Payload, Bounce == 0);
+						FMaterialEval MaterialEval = EvalMaterial(Ray.Direction, LightSample.Direction, Payload, Bounce == 0);
 
-						float3 LightContrib = PathThroughput * RadianceOverPdf * MaterialEval.Weight * MaterialEval.Pdf;
+						float3 LightContrib = PathThroughput * LightSample.RadianceOverPdf * MaterialEval.Weight * MaterialEval.Pdf;
 						if (MISMode == 2)
 						{
-							LightContrib *= MISWeightRobust(LightPdf, MaterialEval.Pdf);
+							LightContrib *= MISWeightRobust(LightSample.Pdf, MaterialEval.Pdf);
 						}
 
 						// Record the contribution
@@ -442,7 +436,7 @@ void PathTracingKernel(
 							if (!IsStationary(LightId))
 							{
 								DirectLightingNEERadianceValue += LightContrib;
-								DirectLightingNEERadianceDirection = LightRay.Direction;
+								DirectLightingNEERadianceDirection = LightSample.Direction;
 							}
 						}
 					}
diff --git a/Engine/Shaders/Private/PathTracing/Light/PathTracingCapsuleLight.ush b/Engine/Shaders/Private/PathTracing/Light/PathTracingCapsuleLight.ush
--- a/Engine/Shaders/Private/PathTracing/Light/PathTracingCapsuleLight.ush
+++ b/Engine/Shaders/Private/PathTracing/Light/PathTracingCapsuleLight.ush
@@ -55,18 +55,17 @@ float3 CapsuleNormal(float3 Pos, float3 Center, float3 Axis)
 	return normalize(pa - h * Axis);
 }
 
-float4 CapsuleLight_TraceLight(RayDesc Ray, int LightId, inout float HitT)
+FLightHit CapsuleLight_TraceLight(RayDesc Ray, int LightId)
 {
 	float3 LightPosition = GetPosition(LightId);
 	float LightRadius = GetRadius(LightId);
 	float SourceLength = GetSourceLength(LightId);
 	float3 Axis = GetdPdv(LightId) * SourceLength;
 	float t = CapsuleIntersect(Ray.Origin, Ray.Direction, LightPosition, Axis, LightRadius);
-	if (t > Ray.TMin && t < HitT)
+	if (t > Ray.TMin && t < Ray.TMax)
 	{
 		float CapsArea = Pow2(LightRadius);
 		float BodyArea = 0.5 * LightRadius * SourceLength;
-		HitT = t;
 		float3 LightPower = GetColor(LightId) * ComputeIESAttenuation(LightId, Ray.Origin);
 		float3 LightRadiance = LightPower / (PI * (CapsArea + BodyArea));
 		float3 LightDirection = GetPosition(LightId) - Ray.Origin;
@@ -79,20 +78,17 @@ float4 CapsuleLight_TraceLight(RayDesc Ray, int LightId, inout float HitT)
 
 		float Pdf = CosTheta > 0 ? t * t / (4 * PI * CosTheta * (CapsArea + BodyArea)) : 0.0;
 
-		return float4(LightRadiance, Pdf);
+		return CreateLightHit(LightRadiance, Pdf, t);
 	}
-	return 0.0;
+	return NullLightHit();
 }
 
-void CapsuleLight_SampleLight(
+FLightSample CapsuleLight_SampleLight(
 	int LightId,
 	float4 RandSample,
 	float3 WorldPos,
-	float3 WorldNormal,
-	out float3 OutDirection,
-	out float  OutDistance,
-	out float3 OutRadianceOverPdf,
-	out float OutPdf)
+	float3 WorldNormal
+)
 {
 	// Capsule case
 	// #dxr_todo: only sample the visible portion of the capsule and account for the 1/d^2 falloff down the axis
@@ -143,10 +139,12 @@ void CapsuleLight_SampleLight(
 	LightPoint = mul(LightPoint, CylinderBasis) + LightDirection; // Local to World
 
 	float DistanceSquared = dot(LightPoint, LightPoint);
-	OutDistance = sqrt(DistanceSquared);
-	OutDirection = LightPoint * rcp(OutDistance);
+	float Distance = sqrt(DistanceSquared);
+	float3 Direction = LightPoint * rcp(Distance);
 
 	float3 LightPower = GetColor(LightId);
-	OutRadianceOverPdf = LightPower * CosTheta * 4 / DistanceSquared;
-	OutPdf = CosTheta > 0 ? DistanceSquared / (4 * PI * CosTheta * (CapsArea + BodyArea)) : 0.0;
+	float3 RadianceOverPdf = LightPower * CosTheta * 4 / DistanceSquared;
+	float Pdf = CosTheta > 0 ? DistanceSquared / (4 * PI * CosTheta * (CapsArea + BodyArea)) : 0.0;
+
+	return CreateLightSample(RadianceOverPdf, Pdf, Direction, Distance);
 }
diff --git a/Engine/Shaders/Private/PathTracing/Light/PathTracingDirectionalLight.ush b/Engine/Shaders/Private/PathTracing/Light/PathTracingDirectionalLight.ush
--- a/Engine/Shaders/Private/PathTracing/Light/PathTracingDirectionalLight.ush
+++ b/Engine/Shaders/Private/PathTracing/Light/PathTracingDirectionalLight.ush
@@ -8,9 +8,9 @@ DirectionalLight.usf: Light sampling functions for Directional light implementat
 
 #include "PathTracingLightCommon.ush"
 
-float4 DirectionalLight_TraceLight(RayDesc Ray, int LightId, inout float HitT)
+FLightHit DirectionalLight_TraceLight(RayDesc Ray, int LightId)
 {
-	if (Ray.TMax == HitT) // we hit the background
+	if (Ray.TMax == POSITIVE_INFINITY) // we hit the background
 	{
 		float3 LightDirection = GetNormal(LightId);
 		float SinThetaMax = GetRadius(LightId); // Sin(Angle / 2)
@@ -27,41 +27,32 @@ float4 DirectionalLight_TraceLight(RayDesc Ray, int LightId, inout float HitT)
 			// that the overall illumination stays constant as the angle changes.
 			float SolidAngle = 2 * PI * OneMinusCosThetaMax;
 			float3 Radiance = GetColor(LightId) / SolidAngle;
-			return float4(Radiance, 1.0 / SolidAngle);
+			return CreateLightHit(Radiance, 1.0 / SolidAngle, POSITIVE_INFINITY);
 		}
 	}
-	return 0.0;
+	return NullLightHit();
 }
 
-bool DirectionalLight_SampleLight(
+FLightSample DirectionalLight_SampleLight(
 	int LightId,
 	float4 RandSample,
 	float3 WorldPos,
-	float3 WorldNormal,
-	out float3 OutDirection,
-	out float  OutDistance,
-	out float3 OutRadianceOverPdf,
-	out float OutPdf)
+	float3 WorldNormal
+)
 {
 	float SinThetaMax = GetRadius(LightId); // currently sin(Angle / 2)
 	float4 DirAndPdf = UniformSampleConeRobust(RandSample.yz, SinThetaMax * SinThetaMax);
+	float3 Direction = TangentToWorld(DirAndPdf.xyz, GetNormal(LightId));
 
-	OutDirection = TangentToWorld(DirAndPdf.xyz, GetNormal(LightId));
-	OutDistance = 1.0e27;
-
-	// Because the light is normalized by the solid angle, the ratio is just the color
-	OutRadianceOverPdf = GetColor(LightId);
-	OutPdf = DirAndPdf.w;
-
-	return true;
+	// Because the light is normalized by the solid angle, the radiance/pdf ratio is just the color
+	return CreateLightSample(GetColor(LightId), DirAndPdf.w, Direction, POSITIVE_INFINITY);
 }
 
-void DirectionalLight_EstimateLight(
+float DirectionalLight_EstimateLight(
 	int LightId,
 	float3 WorldPos,
 	float3 WorldNormal,
-	bool IsTransmissiveMaterial,
-	out float OutIrradiance
+	bool IsTransmissiveMaterial
 )
 {
 	float3 LightDirection = GetNormal(LightId);
@@ -69,35 +60,5 @@ void DirectionalLight_EstimateLight(
 
 	float NoL = 1.0; // trivial upper bound -- trying to be more accurate appears to reduce performance
 	float LightPower = Luminance(GetColor(LightId));
-	OutIrradiance = LightPower * NoL;
+	return LightPower * NoL;
 }
-
-void DirectionalLight_SampleLight(
-	RayDesc Ray,
-	FMaterialClosestHitPayload Payload,
-	float4 RandSample,
-	int LightId,
-	out float3 OutLightUV,
-	out float OutPdf
-)
-{
-	float3 N = GetNormal(LightId);
-	float Radius = GetRadius(LightId);
-
-	if (Radius == 0)
-	{
-		// singular case
-		// #dxr_todo: should return +inf PDF to disable MIS
-		OutLightUV = N;
-		OutPdf = 1.0;
-	}
-	else
-	{
-		// soft case - build a frame, and sample the distant cone
-		float SinThetaMax = Radius; // currently sin(Angle / 2)
-
-		float4 DirAndPdf = UniformSampleConeRobust(RandSample.yz, SinThetaMax * SinThetaMax);
-		OutLightUV = TangentToWorld(DirAndPdf.xyz, N);
-		OutPdf = DirAndPdf.w;
-	}
-}
\ No newline at end of file
diff --git a/Engine/Shaders/Private/PathTracing/Light/PathTracingLightCommon.ush b/Engine/Shaders/Private/PathTracing/Light/PathTracingLightCommon.ush
--- a/Engine/Shaders/Private/PathTracing/Light/PathTracingLightCommon.ush
+++ b/Engine/Shaders/Private/PathTracing/Light/PathTracingLightCommon.ush
@@ -11,6 +11,51 @@ LightCommon.usf: Common utility functions for light sampling
 uint SceneLightCount;
 StructuredBuffer<FPathTracingLight> SceneLights;
 
+struct FLightHit {
+	float3 Radiance;
+	float Pdf;
+	float HitT;
+
+	bool IsMiss() { return HitT <= 0; }
+	bool IsHit() { return HitT > 0; }
+};
+
+struct FLightSample {
+	float3 RadianceOverPdf;
+	float Pdf;
+	float3 Direction;
+	float Distance;
+};
+
+FLightHit NullLightHit()
+{
+	return (FLightHit)0;
+}
+
+FLightHit CreateLightHit(float3 Radiance, float Pdf, float HitT)
+{
+	FLightHit Result;
+	Result.Radiance = Radiance;
+	Result.Pdf = Pdf;
+	Result.HitT = HitT;
+	return Result;
+}
+
+FLightSample NullLightSample()
+{
+	return (FLightSample)0;
+}
+
+FLightSample CreateLightSample(float3 RadianceOverPdf, float Pdf, float3 Direction, float Distance)
+{
+	FLightSample Sample;
+	Sample.RadianceOverPdf = RadianceOverPdf;
+	Sample.Pdf = Pdf;
+	Sample.Direction = Direction;
+	Sample.Distance = Distance;
+	return Sample;
+}
+
 
 int GetLightId(FMaterialClosestHitPayload RayHitInfo)
 {
diff --git a/Engine/Shaders/Private/PathTracing/Light/PathTracingLightSampling.ush b/Engine/Shaders/Private/PathTracing/Light/PathTracingLightSampling.ush
--- a/Engine/Shaders/Private/PathTracing/Light/PathTracingLightSampling.ush
+++ b/Engine/Shaders/Private/PathTracing/Light/PathTracingLightSampling.ush
@@ -20,88 +20,59 @@ LightSampling.usf: Light sampling for next-event estimation
 #include "PathTracingSpotLight.ush"
 
 // Returns float4 packed with Radiance and Pdf
-float4 TraceLight(RayDesc Ray, int LightId, inout float HitT)
+FLightHit TraceLight(RayDesc Ray, int LightId)
 {
 	switch (SceneLights[LightId].Flags & PATHTRACER_FLAG_TYPE_MASK)
 	{
-		case PATHTRACING_LIGHT_SKY:			return SkyLight_TraceLight(Ray, LightId, HitT);
-		case PATHTRACING_LIGHT_POINT:		return PointLight_TraceLight(Ray, LightId, HitT);
-		case PATHTRACING_LIGHT_DIRECTIONAL: return DirectionalLight_TraceLight(Ray, LightId, HitT);
-		case PATHTRACING_LIGHT_RECT:		return RectLight_TraceLight(Ray, LightId, HitT);
-		case PATHTRACING_LIGHT_SPOT:		return SpotLight_TraceLight(Ray, LightId, HitT);
-		default: return 0; // unknown light type?
+		case PATHTRACING_LIGHT_SKY:			return SkyLight_TraceLight(Ray, LightId);
+		case PATHTRACING_LIGHT_POINT:		return PointLight_TraceLight(Ray, LightId);
+		case PATHTRACING_LIGHT_DIRECTIONAL: return DirectionalLight_TraceLight(Ray, LightId);
+		case PATHTRACING_LIGHT_RECT:		return RectLight_TraceLight(Ray, LightId);
+		case PATHTRACING_LIGHT_SPOT:		return SpotLight_TraceLight(Ray, LightId);
+		default:							return NullLightHit(); // unknown light type?
 	}
 }
 
 // Choose a point on a light as seen from the given shading point
-void SampleLight(
+FLightSample SampleLight(
 	int LightId,
 	float4 RandSample,
 	float3 WorldPos,
-	float3 WorldNormal,
-	out float3 OutDirection,
-	out float  OutDistance,
-	out float3 OutRadianceOverPdf,
-	out float OutPdf)
+	float3 WorldNormal
+)
 {
 	switch (SceneLights[LightId].Flags & PATHTRACER_FLAG_TYPE_MASK)
 	{
-		case PATHTRACING_LIGHT_SKY:
-			SkyLight_SampleLight(LightId, RandSample, WorldPos, WorldNormal, OutDirection, OutDistance, OutRadianceOverPdf, OutPdf);
-			return;
-		case PATHTRACING_LIGHT_POINT:
-			PointLight_SampleLight(LightId, RandSample, WorldPos, WorldNormal, OutDirection, OutDistance, OutRadianceOverPdf, OutPdf);
-			return;
-		case PATHTRACING_LIGHT_DIRECTIONAL:
-			DirectionalLight_SampleLight(LightId, RandSample, WorldPos, WorldNormal, OutDirection, OutDistance, OutRadianceOverPdf, OutPdf);
-			return;
-		case PATHTRACING_LIGHT_RECT:
-			RectLight_SampleLight(LightId, RandSample, WorldPos, WorldNormal, OutDirection, OutDistance, OutRadianceOverPdf, OutPdf);
-			return;
-		case PATHTRACING_LIGHT_SPOT:
-			SpotLight_SampleLight(LightId, RandSample, WorldPos, WorldNormal, OutDirection, OutDistance, OutRadianceOverPdf, OutPdf);
-			return;
-		default:
-			// unknown light type?
-			OutDirection = 0;
-			OutDistance = 0;
-			OutRadianceOverPdf = 0;
-			OutPdf = 0;
-			return;
+		case PATHTRACING_LIGHT_SKY:			return SkyLight_SampleLight(LightId, RandSample, WorldPos, WorldNormal);
+		case PATHTRACING_LIGHT_POINT:		return PointLight_SampleLight(LightId, RandSample, WorldPos, WorldNormal);
+		case PATHTRACING_LIGHT_DIRECTIONAL: return DirectionalLight_SampleLight(LightId, RandSample, WorldPos, WorldNormal);
+		case PATHTRACING_LIGHT_RECT:		return RectLight_SampleLight(LightId, RandSample, WorldPos, WorldNormal);
+		case PATHTRACING_LIGHT_SPOT:		return SpotLight_SampleLight(LightId, RandSample, WorldPos, WorldNormal);
+		default:							return NullLightSample(); // unknown light type?
 	}
 }
 
-void EstimateLight(
+float EstimateLight(
 	int LightId,
 	float3 WorldPos,
 	float3 WorldNormal,
 	uint PrimitiveLightingChannelMask,
-	bool IsTransmissiveMaterial,
-	inout float OutIrradiance
+	bool IsTransmissiveMaterial
 )
 {
-	OutIrradiance = 0.0f;
-	if ((GetLightingChannelMask(LightId) & PrimitiveLightingChannelMask) == 0) return;
+	if ((GetLightingChannelMask(LightId) & PrimitiveLightingChannelMask) == 0)
+	{
+		return 0.0;
+	}
 
 	switch (SceneLights[LightId].Flags & PATHTRACER_FLAG_TYPE_MASK)
 	{
-		case PATHTRACING_LIGHT_SKY:
-			SkyLight_EstimateLight(LightId, WorldPos, WorldNormal, IsTransmissiveMaterial, OutIrradiance);
-			return;
-		case PATHTRACING_LIGHT_POINT:
-			PointLight_EstimateLight(LightId, WorldPos, WorldNormal, IsTransmissiveMaterial, OutIrradiance);
-			return;
-		case PATHTRACING_LIGHT_DIRECTIONAL:
-			DirectionalLight_EstimateLight(LightId, WorldPos, WorldNormal, IsTransmissiveMaterial, OutIrradiance);
-			return;
-		case PATHTRACING_LIGHT_RECT:
-			RectLight_EstimateLight(LightId, WorldPos, WorldNormal, IsTransmissiveMaterial, OutIrradiance);
-			return;
-		case PATHTRACING_LIGHT_SPOT:
-			SpotLight_EstimateLight(LightId, WorldPos, WorldNormal, IsTransmissiveMaterial, OutIrradiance);
-			return;
-		default:
-			return;
+		case PATHTRACING_LIGHT_SKY:			return SkyLight_EstimateLight(LightId, WorldPos, WorldNormal, IsTransmissiveMaterial);
+		case PATHTRACING_LIGHT_POINT:		return PointLight_EstimateLight(LightId, WorldPos, WorldNormal, IsTransmissiveMaterial);
+		case PATHTRACING_LIGHT_DIRECTIONAL:	return DirectionalLight_EstimateLight(LightId, WorldPos, WorldNormal, IsTransmissiveMaterial);
+		case PATHTRACING_LIGHT_RECT:		return RectLight_EstimateLight(LightId, WorldPos, WorldNormal, IsTransmissiveMaterial);
+		case PATHTRACING_LIGHT_SPOT:		return SpotLight_EstimateLight(LightId, WorldPos, WorldNormal, IsTransmissiveMaterial);
+		default:							return 0.0;
 	}
 }
 
@@ -131,7 +102,7 @@ bool InitLightPickingCdf(
 	{
 		for (uint LightIndex = 0; LightIndex < SceneLightCount; ++LightIndex)
 		{
-			EstimateLight(LightIndex, WorldPos, WorldNormal, PrimitiveLightingChannelMask, IsTransmissiveMaterial, LightPickingCdf[LightIndex]);
+			LightPickingCdf[LightIndex] = EstimateLight(LightIndex, WorldPos, WorldNormal, PrimitiveLightingChannelMask, IsTransmissiveMaterial);
 			LightPickingCdf[LightIndex] += LastCdf;
 			LastCdf = LightPickingCdf[LightIndex];
 		}
@@ -211,13 +182,11 @@ bool SampleLight(
 
 	SelectLight(RandSample.x, LightPickingCdf, OutLightId, LightPickPdf);
 
-	SampleLight(OutLightId, RandSample, WorldPos, WorldNormal,
-		OutDirection,
-		OutDistance,
-		OutRadianceOverPdf,
-		OutPdf);
-	OutRadianceOverPdf /= LightPickPdf;
-	OutPdf *= LightPickPdf;
+	FLightSample Result = SampleLight(OutLightId, RandSample, WorldPos, WorldNormal);
+	OutDirection = Result.Direction;
+	OutDistance = Result.Distance;
+	OutRadianceOverPdf = Result.RadianceOverPdf / LightPickPdf;
+	OutPdf = Result.Pdf * LightPickPdf;
 
 	return true;
 }
diff --git a/Engine/Shaders/Private/PathTracing/Light/PathTracingPointLight.ush b/Engine/Shaders/Private/PathTracing/Light/PathTracingPointLight.ush
--- a/Engine/Shaders/Private/PathTracing/Light/PathTracingPointLight.ush
+++ b/Engine/Shaders/Private/PathTracing/Light/PathTracingPointLight.ush
@@ -10,11 +10,11 @@ PathTracingPointLight.ush: Light sampling functions for Point lights
 #include "PathTracingCapsuleLight.ush"
 
 
-float4 PointLight_TraceLight(RayDesc Ray, int LightId, inout float HitT)
+FLightHit PointLight_TraceLight(RayDesc Ray, int LightId)
 {
 	if (GetSourceLength(LightId) > 0)
 	{
-		return CapsuleLight_TraceLight(Ray, LightId, HitT);
+		return CapsuleLight_TraceLight(Ray, LightId);
 	}
 
 	float3 LightPosition = GetPosition(LightId);
@@ -32,9 +32,8 @@ float4 PointLight_TraceLight(RayDesc Ray, int LightId, inout float HitT)
 	{
 		h = sqrt(h);
 		float t = (-b - h) / a;
-		if (t > Ray.TMin && t < HitT)
+		if (t > Ray.TMin && t < Ray.TMax)
 		{
-			HitT = t;
 			// #dxr_todo: sphere area is 4*pi*r^2 -- but the factor of 4 is missing for some reason?
 			float3 LightPower = GetColor(LightId) * ComputeIESAttenuation(LightId, Ray.Origin);
 			float3 LightRadiance = LightPower / (PI * LightRadius2);
@@ -44,30 +43,28 @@ float4 PointLight_TraceLight(RayDesc Ray, int LightId, inout float HitT)
 			float OneMinusCosThetaMax = SinThetaMax2 < 0.01 ? SinThetaMax2 * (0.5 + 0.125 * SinThetaMax2) : 1 - sqrt(1 - SinThetaMax2);
 
 			float SolidAngle = 2 * PI * OneMinusCosThetaMax;
-			return float4(LightRadiance, 1.0 / SolidAngle);
+			return CreateLightHit(LightRadiance, 1.0 / SolidAngle, t);
 		}
 		// #dxr_todo: process inside hit here ...
 		// How should we define radiance on the inside of the sphere?
 	}
-	return 0.0;
+	return NullLightHit();
 }
 
-void PointLight_SampleLight(
+FLightSample PointLight_SampleLight(
 	int LightId,
 	float4 RandSample,
 	float3 WorldPos,
-	float3 WorldNormal,
-	out float3 OutDirection,
-	out float  OutDistance,
-	out float3 OutRadianceOverPdf,
-	out float OutPdf)
+	float3 WorldNormal
+)
 {
 	float3 LightDirection = GetPosition(LightId) - WorldPos;
 	float LightDistanceSquared = dot(LightDirection, LightDirection);
 
+	FLightSample Result = NullLightSample();
 	if (GetSourceLength(LightId) > 0)
 	{
-		CapsuleLight_SampleLight(LightId, RandSample, WorldPos, WorldNormal, OutDirection, OutDistance, OutRadianceOverPdf, OutPdf);
+		Result = CapsuleLight_SampleLight(LightId, RandSample, WorldPos, WorldNormal);
 	}
 	else
 	{
@@ -86,30 +83,30 @@ void PointLight_SampleLight(
 		float CosTheta = DirAndPdf.z;
 		float SinTheta2 = 1.0 - CosTheta * CosTheta;
 
-		OutDirection = TangentToWorld(DirAndPdf.xyz, normalize(LightDirection));
-		OutDistance = length(LightDirection) * (CosTheta - sqrt(max(SinThetaMax2 - SinTheta2, 0.0)));
+		Result.Direction = TangentToWorld(DirAndPdf.xyz, normalize(LightDirection));
+		Result.Distance = length(LightDirection) * (CosTheta - sqrt(max(SinThetaMax2 - SinTheta2, 0.0)));
 
-		OutPdf = DirAndPdf.w;
+		Result.Pdf = DirAndPdf.w;
 
 		float3 LightPower = GetColor(LightId);
 		float3 LightRadiance = LightPower / (PI * Radius2);
 
 		// When the angle is very small, Radiance over pdf simplifies even more since SolidAngle ~= PI * SinThetaMax2
 		// Canceling out common factors further leads to the classic Power / D^2 formula
-		OutRadianceOverPdf = SinThetaMax2 < 0.001 ? LightPower / LightDistanceSquared : LightRadiance / OutPdf;
+		Result.RadianceOverPdf = SinThetaMax2 < 0.001 ? LightPower / LightDistanceSquared : LightRadiance / Result.Pdf;
 	}
 
 	// NOTE: uses distance to center to keep fadeoff mask consistent for all samples
-	OutRadianceOverPdf *= ComputeAttenuationFalloff(LightDistanceSquared, LightId);
-	OutRadianceOverPdf *= ComputeIESAttenuation(LightId, WorldPos);
+	Result.RadianceOverPdf *= ComputeAttenuationFalloff(LightDistanceSquared, LightId);
+	Result.RadianceOverPdf *= ComputeIESAttenuation(LightId, WorldPos);
+	return Result;
 }
 
-void PointLight_EstimateLight(
+float PointLight_EstimateLight(
 	int LightId,
 	float3 WorldPos,
 	float3 WorldNormal,
-	bool IsTransmissiveMaterial,
-	out float OutIrradiance
+	bool IsTransmissiveMaterial
 )
 {
 	// Distance
@@ -121,5 +118,5 @@ void PointLight_EstimateLight(
 
 	float LightPower = Luminance(GetColor(LightId));
 	float Falloff = ComputeAttenuationFalloff(LightDistanceSquared, LightId);
-	OutIrradiance = LightPower * Falloff * NoL / LightDistanceSquared;
+	return LightPower * Falloff * NoL / LightDistanceSquared;
 }
diff --git a/Engine/Shaders/Private/PathTracing/Light/PathTracingRectLight.ush b/Engine/Shaders/Private/PathTracing/Light/PathTracingRectLight.ush
--- a/Engine/Shaders/Private/PathTracing/Light/PathTracingRectLight.ush
+++ b/Engine/Shaders/Private/PathTracing/Light/PathTracingRectLight.ush
@@ -56,7 +56,7 @@ float4 SampleApproxProjectionSphericalRect(float2 Rand, FSphericalRect Spherical
 }
 
 
-float4 RectLight_TraceLight(RayDesc Ray, int LightId, inout float HitT)
+FLightHit RectLight_TraceLight(RayDesc Ray, int LightId)
 {
 	float3 LightPosition = GetPosition(LightId);
 	float3 LightNormal = GetNormal(LightId);
@@ -64,7 +64,7 @@ float4 RectLight_TraceLight(RayDesc Ray, int LightId, inout float HitT)
 	float DoN = dot(Ray.Direction, LightNormal);
 	float t = dot(LightDirection, LightNormal) / DoN;
 	// ray points toward the plane and intersect it?
-	if (DoN < 0 && t > Ray.TMin && t < HitT)
+	if (DoN < 0 && t > Ray.TMin && t < Ray.TMax)
 	{
 		float3 LightdPdu = GetdPdu(LightId);
 		float3 LightdPdv = GetdPdv(LightId);
@@ -89,27 +89,23 @@ float4 RectLight_TraceLight(RayDesc Ray, int LightId, inout float HitT)
 			UV = float2(dot(P, LightdPdu), dot(P, LightdPdv));
 			if (all(abs(UV) <= Rect.Extent))
 			{
-				HitT = t;
 				// stored color is radiance
 				float3 Radiance = GetColor(LightId) * ComputeIESAttenuation(LightId, Ray.Origin);
 				Radiance *= ComputeAttenuationFalloff(dot(LightDirection, LightDirection), LightId);
 				float SolidAngle = BuildSphericalRect(Rect).SolidAngle;
-				return float4(Radiance, 1.0 / SolidAngle);
+				return CreateLightHit(Radiance, 1.0 / SolidAngle, t);
 			}
 		}
 	}
-	return 0.0;
+	return NullLightHit();
 }
 
-void RectLight_SampleLight(
+FLightSample RectLight_SampleLight(
 	int LightId,
 	float4 RandSample,
 	float3 WorldPos,
-	float3 WorldNormal,
-	out float3 OutDirection,
-	out float  OutDistance,
-	out float3 OutRadianceOverPdf,
-	out float OutPdf)
+	float3 WorldNormal
+)
 {
 	float3 LightPosition = GetPosition(LightId);
 	float3 LightNormal = GetNormal(LightId);
@@ -132,19 +128,17 @@ void RectLight_SampleLight(
 
 	float3 Direction = UniformSampleSphericalRect(RandSample.yz, SphericalRect);
 
-	OutDirection = normalize(Direction);
-	OutDistance = length(Direction);
-	OutRadianceOverPdf = GetColor(LightId) * SphericalRect.SolidAngle * ComputeIESAttenuation(LightId, WorldPos);
-	OutRadianceOverPdf *= ComputeAttenuationFalloff(dot(LightDirection, LightDirection), LightId);
-	OutPdf = 1.0 / SphericalRect.SolidAngle;
+	float3 RadianceOverPdf = GetColor(LightId) * SphericalRect.SolidAngle * ComputeIESAttenuation(LightId, WorldPos);
+	RadianceOverPdf *= ComputeAttenuationFalloff(dot(LightDirection, LightDirection), LightId);
+
+	return CreateLightSample(RadianceOverPdf, 1.0 / SphericalRect.SolidAngle, normalize(Direction), length(Direction));
 }
 
-void RectLight_EstimateLight(
+float RectLight_EstimateLight(
 	int LightId,
 	float3 WorldPos,
 	float3 WorldNormal,
-	bool IsTransmissiveMaterial,
-	out float OutIrradiance
+	bool IsTransmissiveMaterial
 )
 {
 	// Distance to centroid
@@ -154,10 +148,10 @@ void RectLight_EstimateLight(
 	float3 LightNormal = GetNormal(LightId);
 
 	// Is the shading point behind the light?
-	if (dot(-LightDirection, LightNormal) <= 0.0)
+	float LNoL = saturate(-dot(normalize(LightDirection), LightNormal));
+	if (LNoL <= 0.0)
 	{
-		OutIrradiance = 0.0;
-		return;
+		return 0.0;
 	}
 	// Approximate geometric term
 	float Width = GetWidth(LightId);
@@ -169,5 +163,5 @@ void RectLight_EstimateLight(
 
 	float LightPower = Luminance(GetColor(LightId));
 	float Falloff = ComputeAttenuationFalloff(LightDistanceSquared, LightId);
-	OutIrradiance = LightPower * Falloff * Area * NoL / LightDistanceSquared;
+	return LightPower * Falloff * Area * NoL * LNoL / LightDistanceSquared;
 }
diff --git a/Engine/Shaders/Private/PathTracing/Light/PathTracingSkyLight.ush b/Engine/Shaders/Private/PathTracing/Light/PathTracingSkyLight.ush
--- a/Engine/Shaders/Private/PathTracing/Light/PathTracingSkyLight.ush
+++ b/Engine/Shaders/Private/PathTracing/Light/PathTracingSkyLight.ush
@@ -11,12 +11,11 @@ SkyLight.usf: Light sampling functions for SkyLight implementation
 
 #define USE_HIERARCHICAL_IMPORTANCE_SAMPLING
 
-void SkyLight_EstimateLight(
+float SkyLight_EstimateLight(
 	int LightId,
 	float3 WorldPos,
 	float3 WorldNormal,
-	bool IsTransmissiveMaterial,
-	out float OutIrradiance
+	bool IsTransmissiveMaterial
 )
 {
 	// Modifying diffuse contribution to prevent false negatives with PDF estimation
@@ -24,11 +23,11 @@ void SkyLight_EstimateLight(
 	float3 SkyLightIrradiance = GetColor(LightId) * SkySHDiffuse;
 	float SkyLightIlluminance = Luminance(SkyLightIrradiance);
 
-	OutIrradiance = SkyLightIlluminance * 2.0f * PI;
+	return SkyLightIlluminance * 2.0f * PI;
 }
 
-float4 SkyLight_TraceLight(RayDesc Ray, int LightId, inout float HitT) {
-	if (Ray.TMax == HitT)
+FLightHit SkyLight_TraceLight(RayDesc Ray, int LightId) {
+	if (Ray.TMax == POSITIVE_INFINITY)
 	{
 		// LightUV convention in the code below is just the world space ray direction
 		// so its safe to call into the method below for now ...
@@ -65,21 +64,17 @@ float4 SkyLight_TraceLight(RayDesc Ray, int LightId, inout float HitT) {
 #endif
 		}
 		
-
-		return float4(Radiance, Pdf);
+		return CreateLightHit(Radiance, Pdf, POSITIVE_INFINITY);
 	}
-	return 0.0;
+	return NullLightHit();
 }
 
-void SkyLight_SampleLight(
+FLightSample SkyLight_SampleLight(
 	int LightId,
 	float4 RandSample,
 	float3 WorldPos,
-	float3 WorldNormal,
-	out float3 OutDirection,
-	out float  OutDistance,
-	out float3 OutRadianceOverPdf,
-	out float OutPdf)
+	float3 WorldNormal
+)
 {
 	// Make temp structures so we can re-use the calls above
 	// #dxr_todo: write optimized form that combines radiance and pdf lookup in a single step
@@ -160,12 +155,11 @@ void SkyLight_SampleLight(
 	uint LinearIndex = BufferOffsetAtPixel(MipPixel, StopLevel, SkyLight.MipDimensions.xy);
 	float AreaPdf = SkyLight.SolidAnglePdf[LinearIndex];
 
-	OutDirection = normalize(GetTextureCubeVector(TexelCoord, TextureDimensions));
-	OutPdf = FacePdf * MipPdf * AreaPdf;
+	float3 OutDirection = normalize(GetTextureCubeVector(TexelCoord, TextureDimensions));
+	float OutPdf = FacePdf * MipPdf * AreaPdf;
 #endif
 
 	float3 Radiance = SkyLight.Color * SkyLight.Texture.SampleLevel(SkyLight.TextureSampler, OutDirection, 0).rgb;
 
-	OutDistance = 1.0e27;
-	OutRadianceOverPdf = Radiance / OutPdf;
+	return CreateLightSample(Radiance / OutPdf, OutPdf, OutDirection, POSITIVE_INFINITY);
 }
diff --git a/Engine/Shaders/Private/PathTracing/Light/PathTracingSpotLight.ush b/Engine/Shaders/Private/PathTracing/Light/PathTracingSpotLight.ush
--- a/Engine/Shaders/Private/PathTracing/Light/PathTracingSpotLight.ush
+++ b/Engine/Shaders/Private/PathTracing/Light/PathTracingSpotLight.ush
@@ -27,43 +27,41 @@ float SpotAttenuation(float3 L, float3 SpotDirection, float2 SpotAngles)
 }
 
 
-float4 SpotLight_TraceLight(RayDesc Ray, int LightId, inout float HitT) {
+FLightHit SpotLight_TraceLight(RayDesc Ray, int LightId) {
 
-	float4 Result = PointLight_TraceLight(Ray, LightId, HitT);
-	if (Result.w > 0)
+	FLightHit Result = PointLight_TraceLight(Ray, LightId);
+	if (Result.IsHit())
 	{
 		float3 LightPosition = GetPosition(LightId);
 		float3 LightNormal = GetNormal(LightId);
 		float2 CosConeAngles = GetCosConeAngles(LightId);
-		Result.xyz *= SpotAttenuation(-normalize(Ray.Origin - LightPosition), LightNormal, CosConeAngles);
+		Result.Radiance *= SpotAttenuation(-normalize(Ray.Origin - LightPosition), LightNormal, CosConeAngles);
 	}
 	return Result;
 }
 
-void SpotLight_SampleLight(
+FLightSample SpotLight_SampleLight(
 	int LightId,
 	float4 RandSample,
 	float3 WorldPos,
-	float3 WorldNormal,
-	out float3 OutDirection,
-	out float  OutDistance,
-	out float3 OutRadianceOverPdf,
-	out float OutPdf)
+	float3 WorldNormal
+)
 {
-	PointLight_SampleLight(LightId, RandSample, WorldPos, WorldNormal, OutDirection, OutDistance, OutRadianceOverPdf, OutPdf);
+	FLightSample Result = PointLight_SampleLight(LightId, RandSample, WorldPos, WorldNormal);
 
 	float3 LightPosition = GetPosition(LightId);
 	float3 LightNormal = GetNormal(LightId);
 	float2 CosConeAngles = GetCosConeAngles(LightId);
-	OutRadianceOverPdf *= SpotAttenuation(-normalize(WorldPos - LightPosition), LightNormal, CosConeAngles);
+	Result.RadianceOverPdf *= SpotAttenuation(-normalize(WorldPos - LightPosition), LightNormal, CosConeAngles);
+
+	return Result;
 }
 
-void SpotLight_EstimateLight(
+float SpotLight_EstimateLight(
 	int LightId,
 	float3 WorldPos,
 	float3 WorldNormal,
-	bool IsTransmissiveMaterial,
-	out float OutIrradiance
+	bool IsTransmissiveMaterial
 )
 {
 	float3 LightPosition = GetPosition(LightId);
@@ -79,9 +77,11 @@ void SpotLight_EstimateLight(
 	float NoL = 1.0; // trivial upper bound -- trying to be more accurate appears to reduce performance
 
 	float Falloff = ComputeAttenuationFalloff(LightDistanceSquared, LightId);
-	OutIrradiance = LightPower * Falloff * NoL / LightDistanceSquared;
+	float OutIrradiance = LightPower * Falloff * NoL / LightDistanceSquared;
 	
 	float2 CosConeAngles = GetCosConeAngles(LightId);
 	OutIrradiance *= SpotAttenuation(normalize(LightDirection), LightNormal, CosConeAngles);
+
+	return OutIrradiance;
 }
 
diff --git a/Engine/Shaders/Private/PathTracing/PathTracing.usf b/Engine/Shaders/Private/PathTracing/PathTracing.usf
--- a/Engine/Shaders/Private/PathTracing/PathTracing.usf
+++ b/Engine/Shaders/Private/PathTracing/PathTracing.usf
@@ -64,8 +64,9 @@ FMaterialClosestHitPayload TraceTransparentRay(RayDesc Ray, FRayCone RayCone, bo
 		// #dxr_todo: if we have lots of lights, having some hierarchical structure would be better ....
 		for (uint LightId = 0; LightId < NumLights; ++LightId)
 		{
-			float HitT = HitPayload.IsMiss() ? Ray.TMax : HitPayload.HitT;
-			float3 LightRadiance = TraceLight(Ray, LightId, HitT).xyz;
+			RayDesc LightRay = Ray;
+			LightRay.TMax = HitPayload.IsMiss() ? Ray.TMax : HitPayload.HitT;
+			float3 LightRadiance = TraceLight(LightRay, LightId).Radiance;
 			AccumulateRadiance(Radiance, PathThroughput * LightRadiance);
 		}
 
@@ -318,45 +319,38 @@ RAY_TRACING_ENTRY_RAYGEN(PathTracingMainRG)
 
 				SelectLight(RandSample.x, LightPickingCdf, LightId, LightPickPdf);
 
-				RayDesc LightRay;
-				LightRay.Origin = WorldPos;
-				LightRay.TMin = 0;
-				float3 RadianceOverPdf;
-				float LightPdf;
-				SampleLight(LightId, RandSample, WorldPos, WorldNormal,
-					LightRay.Direction,
-					LightRay.TMax,
-					RadianceOverPdf,
-					LightPdf);
-				RadianceOverPdf /= LightPickPdf;
-				LightPdf *= LightPickPdf;
-				if (LightPdf > 0)
+				FLightSample LightSample = SampleLight(LightId, RandSample, WorldPos, WorldNormal);
+
+				LightSample.RadianceOverPdf /= LightPickPdf;
+				LightSample.Pdf *= LightPickPdf;
+				if (LightSample.Pdf > 0)
 				{
-					float SignedPositionBias = 1.0;
-					if (IsTransmissiveMaterial)
+					if (CastsShadow(LightId))
 					{
 						// for transmissive materials, bias the position to the other side of the surface if the light is coming from behind
-						float NoL = dot(Payload.WorldNormal, LightRay.Direction);
-						SignedPositionBias = sign(NoL);
-					}
-					ApplyPositionBias(LightRay, SignedPositionBias * Payload.WorldNormal, PathTracingData.MaxNormalBias);
+						const float SignedPositionBias = IsTransmissiveMaterial ? sign(dot(Payload.WorldNormal, LightSample.Direction)) : 1.0;
 
-					if (CastsShadow(LightId))
-					{
-						RadianceOverPdf *= TraceTransparentVisibilityRay(LightRay, LaunchIndex, PathRoughness);
+						RayDesc LightRay;
+						LightRay.Origin = WorldPos;
+						LightRay.TMin = 0;
+						LightRay.Direction = LightSample.Direction;
+						LightRay.TMax = LightSample.Distance;
+						ApplyPositionBias(LightRay, SignedPositionBias * Payload.WorldNormal, PathTracingData.MaxNormalBias);
+
+						LightSample.RadianceOverPdf *= TraceTransparentVisibilityRay(LightRay, LaunchIndex, PathRoughness);
 					}
 
 					// #dxr_todo: Is it cheaper to fire the ray first? Or eval the material first?
-					if (any(RadianceOverPdf > 0))
+					if (any(LightSample.RadianceOverPdf > 0))
 					{
 						// Evaluate material
-						FMaterialEval MaterialEval = EvalMaterial(Ray.Direction, LightRay.Direction, Payload, false);
+						FMaterialEval MaterialEval = EvalMaterial(Ray.Direction, LightSample.Direction, Payload, false);
 
 						// Record the contribution
-						float3 LightContrib = PathThroughput * RadianceOverPdf * MaterialEval.Weight * MaterialEval.Pdf;
+						float3 LightContrib = PathThroughput * LightSample.RadianceOverPdf * MaterialEval.Weight * MaterialEval.Pdf;
 						if (MISMode == 2)
 						{
-							LightContrib *= MISWeightRobust(LightPdf, MaterialEval.Pdf);
+							LightContrib *= MISWeightRobust(LightSample.Pdf, MaterialEval.Pdf);
 						}
 						AccumulateRadiance(Radiance, LightContrib);
 					}
@@ -434,19 +428,15 @@ RAY_TRACING_ENTRY_RAYGEN(PathTracingMainRG)
 					continue;
 				}
 
-				float HitT = Ray.TMax;
-				float4 LightResult = TraceLight(Ray, LightId, HitT);
-				float3 LightRadiance = LightResult.xyz;
+				FLightHit LightResult = TraceLight(Ray, LightId);
 
-				float3 LightContrib = PathThroughput * LightRadiance;
+				float3 LightContrib = PathThroughput * LightResult.Radiance;
 				if (MISMode == 2)
 				{
-					float LightPdf = LightResult.w;
 					float LightPickPdf =  LightPickingCdf[LightId] - PreviousCdfValue;
 					PreviousCdfValue = LightPickingCdf[LightId];
 
-					LightPdf *= LightPickPdf;
-					LightContrib *= MISWeightRobust(MaterialSample.Pdf, LightPdf);
+					LightContrib *= MISWeightRobust(MaterialSample.Pdf, LightResult.Pdf * LightPickPdf);
 				}
 
 				if (any(LightContrib > 0))
@@ -454,7 +444,7 @@ RAY_TRACING_ENTRY_RAYGEN(PathTracingMainRG)
 					if (CastsShadow(LightId))
 					{
 						RayDesc LightRay = Ray;
-						LightRay.TMax = HitT;
+						LightRay.TMax = LightResult.HitT;
 						LightContrib *= TraceTransparentVisibilityRay(LightRay, LaunchIndex, PathRoughness);
 					}
 					// the light made some contribution, and there was nothing along the shadow ray
diff --git a/Engine/Shaders/Private/RayTracing/RayTracingGlobalIlluminationRGS.usf b/Engine/Shaders/Private/RayTracing/RayTracingGlobalIlluminationRGS.usf
--- a/Engine/Shaders/Private/RayTracing/RayTracingGlobalIlluminationRGS.usf
+++ b/Engine/Shaders/Private/RayTracing/RayTracingGlobalIlluminationRGS.usf
@@ -182,7 +182,9 @@ RAY_TRACING_ENTRY_RAYGEN(GlobalIlluminationRGS)
 				if (EvalSkyLight && Bounce > 0)
 				{
 					uint SkyLightId = 0;
-					float3 EnvironmentRadiance = SkyLight_TraceLight(Ray, SkyLightId, Ray.TMax).xyz;
+					RayDesc LightRay = Ray;
+					LightRay.TMax = POSITIVE_INFINITY;
+					float3 EnvironmentRadiance = SkyLight_TraceLight(Ray, SkyLightId).Radiance;
 					Irradiance += EnvironmentRadiance * RayThroughput / RayPdf;
 				}
 				break;
