Strata clean up
- only a single uint used for normal, roughness is moved on the classification state free bits
- renaming Normal to WorldNormal
- Removed duplciated classification code from Basepass.
- AO is now gamma 2.0 compressed.

diff --git a/Engine/Shaders/Private/BasePassPixelShader.usf b/Engine/Shaders/Private/BasePassPixelShader.usf
--- a/Engine/Shaders/Private/BasePassPixelShader.usf
+++ b/Engine/Shaders/Private/BasePassPixelShader.usf
@@ -1539,7 +1539,7 @@ void FPixelShaderInOut_MainPS(
 #endif // MATERIAL_SHADINGMODEL_THIN_TRANSLUCENT
 
 
-	bool bSubsurfaceEnable = true;
+	bool bStrataSubsurfaceEnable = false;
 #if STRATA_ENABLED && !MATERIAL_SHADINGMODEL_UNLIT
 
 #if STRATA_OPAQUE_ENABLED
@@ -1549,19 +1549,14 @@ void FPixelShaderInOut_MainPS(
 	// Except for water doing some specialised and simplified lighting during the base pass
 	Color = 0;
 #endif
-	// Clear only what should be and keep the rest
-	GBuffer.DiffuseColor = 0;
-	GBuffer.SpecularColor = 0;
-	GBuffer.BaseColor = 0;
-	GBuffer.Roughness = 0;
-	GBuffer.Specular = 0;
-	GBuffer.Metallic = 0;
-	GBuffer.ShadingModelID = SHADINGMODELID_STRATA;
-	GBuffer.WorldNormal = 0; // Normal will be overriden by the top layer normal
 
 	// STRATA_TODO GBUFFER_HAS_DIFFUSE_SAMPLE_OCCLUSION and ALLOW_STATIC_LIGHTING cases
 	StrataPixelHeader.MaterialAO = GBuffer.GBufferAO;
 
+	// Clear only what should be and keep the rest
+	GBuffer = (FGBufferData)0;						// We should not rely on gbuffer for anything else to make any pass work.
+	GBuffer.ShadingModelID = SHADINGMODELID_STRATA;	// STRATA_TODO remove that after SSS has been moved to strata completely
+
 	uint2 PixelPos = In.SvPosition.xy;
 
 	float3 EmissiveLuminance = 0.0f;
@@ -1570,12 +1565,9 @@ void FPixelShaderInOut_MainPS(
 		OpaqueBasePass.Strata.MaterialLobesBufferUAV, 
 		GetStrataPixelDataByteOffset(PixelPos, uint2(ResolvedView.BufferSizeAndInvSize.xy), OpaqueBasePass.Strata.MaxBytesPerPixel),
 		StrataPixelHeader, Strata,
-		MaterialParameters.CameraVector, GBuffer, EmissiveLuminance);
+		MaterialParameters.CameraVector, bStrataSubsurfaceEnable, EmissiveLuminance);
 
 	Color += EmissiveLuminance;
-
-	// Since we don't carry SSS on/off through the ShadingModel value (DefaultLit/SubSurfaceProfile), we update subsurface toggle, based on the MFP radius value
-	bSubsurfaceEnable = any(GBuffer.CustomData > 0);
 #endif // STRATA_OPAQUE_ENABLED
 
 #if STRATA_TRANSLUCENT_ENABLED
@@ -1588,7 +1580,7 @@ void FPixelShaderInOut_MainPS(
 
 		float3 Throughput = 1.0f;
 		Color = StrataForwardLighting(
-			GBuffer,
+			GBuffer,		// STRATA_TODO remove that!
 			MaterialParameters.WorldNormal,
 			V,
 			L,
@@ -1690,7 +1682,7 @@ void FPixelShaderInOut_MainPS(
 				// to the specular component, otherwise lighting is applied twice
 				Color *= !bChecker;
 			}
-			LightAccumulator_Add(LightAccumulator, Color + DiffuseColor, DiffuseColor, 1.0f, UseSubsurfaceProfile(GBuffer.ShadingModelID) && bSubsurfaceEnable);
+			LightAccumulator_Add(LightAccumulator, Color + DiffuseColor, DiffuseColor, 1.0f, UseSubsurfaceProfile(GBuffer.ShadingModelID) || bStrataSubsurfaceEnable);
 		#else
 			LightAccumulator_Add(LightAccumulator, Color, 0, 1.0f, false);
 		#endif
diff --git a/Engine/Shaders/Private/DistanceFieldScreenGridLighting.usf b/Engine/Shaders/Private/DistanceFieldScreenGridLighting.usf
--- a/Engine/Shaders/Private/DistanceFieldScreenGridLighting.usf
+++ b/Engine/Shaders/Private/DistanceFieldScreenGridLighting.usf
@@ -28,7 +28,7 @@ void ComputeDistanceFieldNormalPS(
 	float SceneDepth = CalcSceneDepth(ScreenUV);
 
 #if STRATA_ENABLED
-	float3 WorldNormal = StrataUnpackTopLayerData(Strata.TopLayerTexture.Load(uint3(PixelCoord, 0))).Normal;
+	float3 WorldNormal = StrataUnpackTopLayerData(Strata.TopLayerNormalTexture.Load(uint3(PixelCoord, 0))).WorldNormal;
 #else
 	float3 WorldNormal = GetGBufferData(ScreenUV).WorldNormal.xyz;
 #endif
@@ -50,7 +50,7 @@ void ComputeDistanceFieldNormalCS(
 	float SceneDepth = CalcSceneDepth(ScreenUV);
 
 #if STRATA_ENABLED
-	float3 WorldNormal = StrataUnpackTopLayerData(Strata.TopLayerTexture.Load(uint3(PixelCoord, 0))).Normal;
+	float3 WorldNormal = StrataUnpackTopLayerData(Strata.TopLayerNormalTexture.Load(uint3(PixelCoord, 0))).WorldNormal;
 #else
 	float3 WorldNormal = GetGBufferData(ScreenUV).WorldNormal.xyz;
 #endif
diff --git a/Engine/Shaders/Private/PostProcessAmbientOcclusion.usf b/Engine/Shaders/Private/PostProcessAmbientOcclusion.usf
--- a/Engine/Shaders/Private/PostProcessAmbientOcclusion.usf
+++ b/Engine/Shaders/Private/PostProcessAmbientOcclusion.usf
@@ -264,9 +264,8 @@ void MainSetupPS(in noperspective float4 UVAndScreenPos : TEXCOORD0, float4 SvPo
 		Samples[i].rgb = normalize(ReconstructNormalFromDepthBuffer(float4(UV[i] * ViewPortSize, SvPosition.zw))) * 0.5f + 0.5f;
 #else
 #if STRATA_ENABLED
-		// STRATA_TODO TopLayerTexture is a UINT2 texture is that is going to be bandwith heavy as compared to non strata path. We should separate Normal and BaseColor-Roughness into two textures.
-		const FStrataTopLayerData TopLayerData = StrataUnpackTopLayerData(Strata.TopLayerTexture.Load(uint3(clamp(UV[i] * View.BufferSizeAndInvSize.xy, 0.0f, View.BufferSizeAndInvSize.xy - 1.0f), 0)));
-		Samples[i].rgb = TopLayerData.Normal * 0.5f + 0.5f;
+		const FStrataTopLayerData TopLayerData = StrataUnpackTopLayerData(Strata.TopLayerNormalTexture.Load(uint3(clamp(UV[i] * View.BufferSizeAndInvSize.xy, 0.0f, View.BufferSizeAndInvSize.xy - 1.0f), 0)));
+		Samples[i].rgb = TopLayerData.WorldNormal * 0.5f + 0.5f;
 #else
 		Samples[i].rgb = GetGBufferData(UV[i], true).WorldNormal * 0.5f + 0.5f;
 #endif
@@ -323,8 +322,8 @@ float3 GetWorldSpaceNormalFromAOInput(float2 UV, float4 SvPosition)
 		const FStrataClassification StrataData = StrataUnpackClassificationData(Strata.ClassificationTexture.Load(uint3(SvPosition.xy, 0)));
 		if (StrataData.ShadingModels != STRATA_BSDF_TYPE_HAIR)
 		{
-			const FStrataTopLayerData TopLayerData = StrataUnpackTopLayerData(Strata.TopLayerTexture.Load(uint3(SvPosition.xy, 0)));
-			WorldNormal = TopLayerData.Normal;
+			const FStrataTopLayerData TopLayerData = StrataUnpackTopLayerData(Strata.TopLayerNormalTexture.Load(uint3(SvPosition.xy, 0)));
+			WorldNormal = TopLayerData.WorldNormal;
 		}
 #else
 		// Otherwise sample Gbuffer normals if the shader model has normals.
@@ -983,8 +982,8 @@ float3 GetNormal(float2 UV, float3 ViewSpacePosMid)
 #if USE_NORMALBUFFER
 
 #if STRATA_ENABLED
-	const FStrataTopLayerData TopLayerData = StrataUnpackTopLayerData(Strata.TopLayerTexture.Load(uint3(clamp(UV * View.BufferSizeAndInvSize.xy, 0.0f, View.BufferSizeAndInvSize.xy-1.0f), 0)));
-	float3 WorldNormal = TopLayerData.Normal;
+	const FStrataTopLayerData TopLayerData = StrataUnpackTopLayerData(Strata.TopLayerNormalTexture.Load(uint3(clamp(UV * View.BufferSizeAndInvSize.xy, 0.0f, View.BufferSizeAndInvSize.xy-1.0f), 0)));
+	float3 WorldNormal = TopLayerData.WorldNormal;
 #else
 	// Get the normal from the normal buffer
 	float3 WorldNormal = GetGBufferData(UV, false).WorldNormal;
@@ -1361,8 +1360,8 @@ float GTAOInnerIntegratePSandCS(in float2 UV, in uint2 iPos)
 	// ViewspacePos and Normal
 	float3 ViewSpacePos		= ScreenToViewPos(UV, SceneDepth);
 #if STRATA_ENABLED
-	const FStrataTopLayerData TopLayerData = StrataUnpackTopLayerData(Strata.TopLayerTexture.Load(uint3(iPos, 0)));
-	float3 WorldNormal = TopLayerData.Normal;
+	const FStrataTopLayerData TopLayerData = StrataUnpackTopLayerData(Strata.TopLayerNormalTexture.Load(uint3(iPos, 0)));
+	float3 WorldNormal = TopLayerData.WorldNormal;
 #else
 	float3 WorldNormal		= GetGBufferData(UV, false).WorldNormal;
 #endif
diff --git a/Engine/Shaders/Private/SSRT/SSRTReflections.usf b/Engine/Shaders/Private/SSRT/SSRTReflections.usf
--- a/Engine/Shaders/Private/SSRT/SSRTReflections.usf
+++ b/Engine/Shaders/Private/SSRT/SSRTReflections.usf
@@ -78,9 +78,8 @@ void ScreenSpaceReflectionsStencilPS(
 
 	uint2 PixelCoord = SvPosition.xy;
 	const FStrataClassification StrataData = StrataUnpackClassificationData(Strata.ClassificationTexture.Load(uint3(PixelCoord, 0)));
-	const FStrataTopLayerData TopLayerData = StrataUnpackTopLayerData(Strata.TopLayerTexture.Load(uint3(PixelCoord, 0)));
 
-	const float Roughness = TopLayerData.Roughness;
+	const float Roughness = StrataData.TopLayerRoughness;
 	const bool bNoMaterial = !StrataData.bIsStrataMaterial;
 
 #else // STRATA_ENABLED
@@ -156,12 +155,12 @@ void ScreenSpaceReflections(
 	else
 	{
 		const FStrataClassification StrataData = StrataUnpackClassificationData(Strata.ClassificationTexture.Load(uint3(PixelCoord, 0)));
-		const FStrataTopLayerData TopLayerData = StrataUnpackTopLayerData(Strata.TopLayerTexture.Load(uint3(PixelCoord, 0)));
+		const FStrataTopLayerData TopLayerData = StrataUnpackTopLayerData(Strata.TopLayerNormalTexture.Load(uint3(PixelCoord, 0)));
 
 		// For now, SSR is only used to trace the top material layer.
 		// STRATA_TODO try stochastic selection
-		N = TopLayerData.Normal;
-		Roughness = TopLayerData.Roughness;
+		N = TopLayerData.WorldNormal;
+		Roughness = StrataData.TopLayerRoughness;
 		bNoMaterial = !StrataData.bIsStrataMaterial;
 	}
 
diff --git a/Engine/Shaders/Private/ShadowProjectionPixelShader.usf b/Engine/Shaders/Private/ShadowProjectionPixelShader.usf
--- a/Engine/Shaders/Private/ShadowProjectionPixelShader.usf
+++ b/Engine/Shaders/Private/ShadowProjectionPixelShader.usf
@@ -217,8 +217,8 @@ void Main(
 			const float3 LightDirection = bIsDirectional ? -LightPositionOrDirection.xyz : normalize(LightPositionOrDirection.xyz - WorldPosition);
 
 		#if STRATA_ENABLED
-			const FStrataTopLayerData TopLayerData = StrataUnpackTopLayerData(Strata.TopLayerTexture.Load(uint3(SVPos.xy, 0)));
-			const float3 WorldNormal = TopLayerData.Normal;
+			const FStrataTopLayerData TopLayerData = StrataUnpackTopLayerData(Strata.TopLayerNormalTexture.Load(uint3(SVPos.xy, 0)));
+			const float3 WorldNormal = TopLayerData.WorldNormal;
 		#else
 			FGBufferData GBufferData = GetGBufferData(ScreenUV);
 			const float3 WorldNormal = GBufferData.WorldNormal;
diff --git a/Engine/Shaders/Private/Strata/Strata.ush b/Engine/Shaders/Private/Strata/Strata.ush
--- a/Engine/Shaders/Private/Strata/Strata.ush
+++ b/Engine/Shaders/Private/Strata/Strata.ush
@@ -102,7 +102,7 @@ struct FStrataBSDF
 #define STATE_BIT_COUNT_HAZINESS				1
 #define STATE_BIT_COUNT_EDGECOLOR				1
 #define STATE_BIT_COUNT_THINFILM				1
-#define STATE_BIT_COUNT___UNUSED___				10
+#define STATE_BIT_COUNT___UNUSED___				9
 #define STATE_BIT_COUNT_WEIGHT10F				10
 
 // Bit offset of each component store in BSDF state 
@@ -228,6 +228,7 @@ struct FStrataClassification
 {
 	uint ShadingModels;
 	float MaterialAO;
+	float TopLayerRoughness;
 	bool bIsStrataMaterial;
 	bool bHasSubsurface;
 	bool bHasSubsurfaceProfile;
@@ -237,40 +238,41 @@ struct FStrataClassification
 
 uint StrataPackClassificationData(FStrataClassification In)
 {
-	const uint MaterialAO8Bits = In.MaterialAO * 255.0f;
+	const uint MaterialAO8Bits = sqrt(saturate(In.MaterialAO)) * 255.0f;	// Give more accuracy to the blacks to reduce banding
+	const uint TopLayerRoughness8Bits = In.TopLayerRoughness * 255.0f;
 	return 
 		(In.ShadingModels & 0xFF) | 
 		((MaterialAO8Bits & 0xFF)			 << 8  ) |
-		((In.bIsStrataMaterial		? 1 : 0) << 16 ) |
-		((In.bHasSubsurface			? 1 : 0) << 17 ) | 
-		((In.bHasSubsurfaceProfile	? 1 : 0) << 18) |
-		((In.bHasTransmission		? 1 : 0) << 19) |
-		((In.bIsSimple				? 1 : 0) << 20);
+		((TopLayerRoughness8Bits & 0xFF)	 << 16 ) |
+		((In.bIsStrataMaterial		? 1 : 0) << 24 ) |
+		((In.bHasSubsurface			? 1 : 0) << 25 ) | 
+		((In.bHasSubsurfaceProfile	? 1 : 0) << 26 ) |
+		((In.bHasTransmission		? 1 : 0) << 27 ) |
+		((In.bIsSimple				? 1 : 0) << 28 );
 }
 
 FStrataClassification StrataUnpackClassificationData(uint In)
 {
 	uint MaterialAO8Bits = 0;
+	uint TopLayerRoughness8Bits = 0;
 	FStrataClassification Out;
 	Out.ShadingModels			= (In >> 0 ) & 0xFF;
 	MaterialAO8Bits				= (In >> 8 ) & 0xFF;
-	Out.bIsStrataMaterial		= (In >> 16) & 0x1;
-	Out.bHasSubsurface			= (In >> 17) & 0x1;
-	Out.bHasSubsurfaceProfile	= (In >> 18) & 0x1;
-	Out.bHasTransmission		= (In >> 19) & 0x1;
-	Out.bIsSimple				= (In >> 20) & 0x1;
-
-	Out.MaterialAO = float(MaterialAO8Bits) * (1.0f / 255.0f);
+	TopLayerRoughness8Bits		= (In >> 16) & 0xFF;
+	Out.bIsStrataMaterial		= (In >> 24) & 0x1;
+	Out.bHasSubsurface			= (In >> 25) & 0x1;
+	Out.bHasSubsurfaceProfile	= (In >> 26) & 0x1;
+	Out.bHasTransmission		= (In >> 27) & 0x1;
+	Out.bIsSimple				= (In >> 28) & 0x1;
+
+	Out.MaterialAO = Square(float(MaterialAO8Bits) * (1.0f / 255.0f));
+	Out.TopLayerRoughness = float(TopLayerRoughness8Bits) * (1.0f / 255.0f);
 	return Out;
 }
 
 struct FStrataTopLayerData
 {
-	float3 Normal;
-	float  Roughness;
-
-	// transient data when accumulating top layer data
-	float AccumRoughnessCount;
+	float3 WorldNormal;
 };
 
 struct FStrataSubsurfaceData
@@ -286,28 +288,16 @@ struct FStrataSubsurfaceData
 	float3 DMFP;				// When the profile is not used, a per pixel DMFP is used
 };
 
-void FinalizeTopLayerData(inout FStrataTopLayerData Out)
-{
-	Out.Normal = normalize(Out.Normal);
-	Out.Roughness = Out.AccumRoughnessCount > 0.0f ? saturate(Out.Roughness / Out.AccumRoughnessCount) : 1.f;
-}
-
-uint2 StrataPackTopLayerData(FStrataTopLayerData In)
+uint StrataPackTopLayerData(FStrataTopLayerData In)
 {
-	uint2 Out = 0;
-	float3 Unused = 0.0f;
-	Out.x = PackColorLinearToGamma2AlphaLinear(float4(Unused, In.Roughness));
-	Out.y = StrataPackNormal(In.Normal);
+	uint Out = StrataPackNormal(In.WorldNormal);
 	return Out;
 }
 
-FStrataTopLayerData StrataUnpackTopLayerData(uint2 In)
+FStrataTopLayerData StrataUnpackTopLayerData(uint In)
 {
-	const float4 UnusedAndRoughness = UnpackColorGamma2ToLinearAlphaLinear(In.x);
-
 	FStrataTopLayerData Out = (FStrataTopLayerData)0;
-	Out.Normal = StrataUnpackNormal(In.y);
-	Out.Roughness = UnusedAndRoughness.w;
+	Out.WorldNormal = StrataUnpackNormal(In);
 	return Out;
 }
 
@@ -1134,9 +1124,10 @@ void PackStrataOut(
 	FStrataPixelHeader StrataPixelHeader,
 	FStrataData Strata,
 	float3 V,
-	inout FGBufferData GBuffer,
+	inout bool bStrataSubsurfaceEnable,
 	inout float3 EmissiveColor)
 {
+	bStrataSubsurfaceEnable = false;
 	EmissiveColor = 0.0f;
 
 	const float OpaqueBSDFThroughput = 0.0f;
@@ -1149,9 +1140,6 @@ void PackStrataOut(
 	// We clamp the layer count to the maximum one according to our byte per pixel budget (removing the bottom layer first)
 	Strata.LayerCount = STRATA_CLAMPED_LAYER_COUNT;
 
-	FStrataTopLayerData TopLayer = (FStrataTopLayerData)0;
-	FStrataSubsurfaceData SSS = (FStrataSubsurfaceData)0;
-
 	// For each of the BSDF, we evaluated their combined weight/throughput according to the material operation performed (layering, mixing, etc.).
 	int l;
 	STRATA_UNROLL_N(STRATA_CLAMPED_LAYER_COUNT)
@@ -1178,8 +1166,6 @@ void PackStrataOut(
 			float3 N = TangentBasis[2];
 			//N = normalize(N);
 
-			TopLayer.Normal += N;
-
 			BxDFContext ContextNoLight = (BxDFContext)0;
 			const float3 FakeL = float3(0, 0, 1);
 			Init(ContextNoLight, N, V, FakeL);
@@ -1191,26 +1177,10 @@ void PackStrataOut(
 				const float3 F0 = ComputeF0(SLAB_SPECULAR(BSDF), SLAB_BASECOLOR(BSDF), SLAB_METALLIC(BSDF));
 				const float3 F90 = ComputeF90(F0, SLAB_EDGECOLOR(BSDF), SLAB_METALLIC(BSDF));
 
-				if (bTopLayer)
-				{
-					// Setup SSR for the top layer
-					TopLayer.Roughness += BSDFCoverageAvg * SLAB_ROUGHNESSX(BSDF);
-					TopLayer.AccumRoughnessCount += BSDFCoverageAvg;
-				}
-
-				// Select the first/toppest layer with subsurface scattering
-				if (BSDF_GETHASSCATTERING(BSDF) && BSDFCoverageAvg > 0 && !SSS.bHasSubsurface)
-				{
-					SSS.bHasSubsurface = true;
-					SSS.bHasProfile			= BSDF_GETHASSSSPROFILE(BSDF);
-					SSS.ProfilePackedId		= SLAB_SSSPROFILEID(BSDF);
-					SSS.ProfileRadiusScale	= SLAB_SSSPROFILERADIUSSCALE(BSDF);
-					SSS.BaseColor			= SLAB_BASECOLOR(BSDF);
-					SSS.DMFP				= SLAB_SSSDMFP(BSDF);
-				}
-
 				EmissiveColor += BSDF_GETEMISSIVE(BSDF) * Throughput;
 
+				bStrataSubsurfaceEnable = bStrataSubsurfaceEnable || (BSDF_GETHASSCATTERING(BSDF) && BSDFCoverageAvg > 0);
+
 				LayerCoverage	+= BSDFCoverage;
 				LayerThroughput += BSDFCoverage * OpaqueBSDFThroughput; // STRATA_TODO: take into accound SSS transmittance
 				break;
@@ -1221,12 +1191,6 @@ void PackStrataOut(
 				const float3 F0  = SHEEN_BASECOLOR(BSDF);
 				const float3 F = F_Schlick(F0, ContextNoLight.VoH);
 
-				if (bTopLayer)
-				{
-					TopLayer.Roughness += BSDFCoverageAvg * SHEEN_ROUGHNESS(BSDF);
-					TopLayer.AccumRoughnessCount += BSDFCoverageAvg;
-				}
-
 				LayerCoverage	+= BSDFCoverage;
 				LayerThroughput += BSDFCoverage * (1.0f - F); // STRATA_TODO use 1-FGD instead of simple 1-F, event if we ignore ray scattering/spreading
 				break;
@@ -1234,12 +1198,6 @@ void PackStrataOut(
 
 			case STRATA_BSDF_TYPE_HAIR:
 			{
-				// No SSR
-				if (bTopLayer)
-				{
-					TopLayer.Roughness += BSDFCoverageAvg * FullyRough; // Do not trigger SSR by using fully rough
-					TopLayer.AccumRoughnessCount += BSDFCoverageAvg;
-				}
 				// Hair are considered opaque
 				LayerCoverage	+= BSDFCoverage;
 				LayerThroughput += BSDFCoverage * OpaqueBSDFThroughput;
@@ -1248,14 +1206,6 @@ void PackStrataOut(
 
 			case STRATA_BSDF_TYPE_SINGLELAYERWATER:
 			{
-				if (bTopLayer)
-				{
-					const float Roughness = SLW_ROUGHNESS(BSDF);
-
-					TopLayer.Roughness += BSDFCoverageAvg * Roughness;
-					TopLayer.AccumRoughnessCount += BSDFCoverageAvg;
-				}
-
 				EmissiveColor += BSDF_GETEMISSIVE(BSDF) * Throughput;
 
 				// Water is considered opaque, it has a custom process to simulate translucency
@@ -1273,22 +1223,6 @@ void PackStrataOut(
 			BSDF_SETISTOPLAYER(Strata.Layers[l].BSDFs[i], bTopLayer ? 1 : 0);
 		}
 
-		// Now modify the gbuffer according to the top layer information
-		if (bTopLayer)
-		{
-			FinalizeTopLayerData(TopLayer);
-
-			// Normal is needed for some AO techniques
-			GBuffer.WorldNormal = normalize(TopLayer.Normal);
-
-			// Roughness is needed for SSR, base color represents EnvBRDF
-			GBuffer.Roughness = TopLayer.Roughness;
-			
-			// Approximative way to be able to have SSR produce colored specular reflections
-			GBuffer.Specular = 1;
-			GBuffer.Metallic = 1;
-		}
-
 		// Take into account the non covered part for this layer and update throughput accordingly
 		LayerThroughput += (1.0 - saturate(LayerCoverage)) * FullThroughput;
 
@@ -1296,26 +1230,6 @@ void PackStrataOut(
 		Throughput *= LayerThroughput;
 	}
 
-	// Patch Gbuffer data with SSS data if enabled
-	GBuffer.CustomData = 0;
-	if (SSS.bHasSubsurface)
-	{
-		if (SSS.bHasProfile)
-		{
-			GBuffer.BaseColor      = SSS.BaseColor;
-			GBuffer.CustomData.rgb = EncodeSubsurfaceProfile(SSS.ProfilePackedId);
-			GBuffer.CustomData.a   = SSS.ProfileRadiusScale;
-		}
-		else
-		{
-			GBuffer.BaseColor      = SSS.BaseColor;
-
-			// This packing should match what is done in SubsurfaceBurleyNormalized.ush
-			uint PackedDMFP11_11_10 = PackR11G11B10F(SSS.DMFP);
-			GBuffer.CustomData = UnpackRGBA8(PackedDMFP11_11_10);
-		}
-	}
-
 	// Now write out Strata data
 
 	// 1. the header
diff --git a/Engine/Shaders/Private/Strata/StrataMaterialClassification.usf b/Engine/Shaders/Private/Strata/StrataMaterialClassification.usf
--- a/Engine/Shaders/Private/Strata/StrataMaterialClassification.usf
+++ b/Engine/Shaders/Private/Strata/StrataMaterialClassification.usf
@@ -20,7 +20,9 @@ void ComputePostProcessBSDFData(
 	float3 V,
 	inout FStrataClassification Classification,
 	inout FStrataSubsurfaceData SSS,
-	inout FStrataTopLayerData TopLayer)
+	inout FStrataTopLayerData TopLayer,
+	inout float TopLayerRoughness, 
+	inout float TopLayerRoughnessAverageCount)
 {
 	const float OpaqueBSDFThroughput = 0.0f;
 	const float FullThroughput = 1.0f;
@@ -34,7 +36,7 @@ void ComputePostProcessBSDFData(
 	float3 N = TangentBasis[2];
 	//N = normalize(N);
 
-	TopLayer.Normal += N;
+	TopLayer.WorldNormal += N;
 
 	BxDFContext ContextNoLight = (BxDFContext)0;
 	const float3 FakeL = float3(0, 0, 1);
@@ -55,8 +57,8 @@ void ComputePostProcessBSDFData(
 		if (bTopLayer)
 		{
 			// Setup SSR for the top layer
-			TopLayer.Roughness += BSDFCoverageAvg * SLAB_ROUGHNESSX(BSDF);
-			TopLayer.AccumRoughnessCount += BSDFCoverageAvg;
+			TopLayerRoughness += BSDFCoverageAvg * SLAB_ROUGHNESSX(BSDF);
+			TopLayerRoughnessAverageCount += BSDFCoverageAvg;
 		}
 
 		Classification.bIsSimple = Classification.bIsSimple && IsStrataFastPathCompatible(BSDF);
@@ -81,8 +83,8 @@ void ComputePostProcessBSDFData(
 
 		if (bTopLayer)
 		{
-			TopLayer.Roughness += SHEEN_ROUGHNESS(BSDF);
-			TopLayer.AccumRoughnessCount += 1.0f;
+			TopLayerRoughness += SHEEN_ROUGHNESS(BSDF);
+			TopLayerRoughnessAverageCount += 1.0f;
 		}
 		break;
 	}
@@ -94,8 +96,8 @@ void ComputePostProcessBSDFData(
 		// No SSR
 		if (bTopLayer)
 		{
-			TopLayer.Roughness += BSDFCoverageAvg * FullyRough; // Do not trigger SSR by using fully rough
-			TopLayer.AccumRoughnessCount += BSDFCoverageAvg;
+			TopLayerRoughness += BSDFCoverageAvg * FullyRough; // Do not trigger SSR by using fully rough
+			TopLayerRoughnessAverageCount += BSDFCoverageAvg;
 		}
 		break;
 	}
@@ -106,8 +108,8 @@ void ComputePostProcessBSDFData(
 		{
 			const float Roughness = SLW_ROUGHNESS(BSDF);
 
-			TopLayer.Roughness += BSDFCoverageAvg * Roughness;
-			TopLayer.AccumRoughnessCount += BSDFCoverageAvg;
+			TopLayerRoughness += BSDFCoverageAvg * Roughness;
+			TopLayerRoughnessAverageCount += BSDFCoverageAvg;
 		}
 		break;
 	}
@@ -118,7 +120,7 @@ void ComputePostProcessBSDFData(
 void MainPS(
 	float4 SVPos : SV_POSITION, 
 	out uint OutColor0 : SV_Target0,
-	out uint2 OutColor1 : SV_Target1)
+	out uint OutColor1 : SV_Target1)
 {
 	const uint2 PixelPos = uint2(SVPos.xy);
 	
@@ -141,17 +143,21 @@ void MainPS(
 
 		ClassificationData.MaterialAO = Header.MaterialAO;
 		ClassificationData.bIsSimple = Header.BSDFCount == 1;
+
+		float TopLayerRoughness = 0.0f;
+		float TopLayerRoughnessAverageCount = 0.0f;
 		for (uint BSDFIndex = 0; BSDFIndex < Header.BSDFCount; ++BSDFIndex)
 		{
 			const FStrataBSDF BSDF = UnpackStrataBSDFIn(Strata.MaterialLobesBuffer, StrataAddressing);
-			ComputePostProcessBSDFData(Header, BSDF, StrataAddressing, V, ClassificationData, SSSData, TopLayerData);
+			ComputePostProcessBSDFData(Header, BSDF, StrataAddressing, V, ClassificationData, SSSData, TopLayerData, TopLayerRoughness, TopLayerRoughnessAverageCount);
 		}
 
-		FinalizeTopLayerData(TopLayerData);
+		TopLayerData.WorldNormal = normalize(TopLayerData.WorldNormal);
+		ClassificationData.TopLayerRoughness = TopLayerRoughnessAverageCount > 0.0f ? TopLayerRoughness / TopLayerRoughnessAverageCount : 0.0f;
 	}
 	else
 	{
-		TopLayerData.Normal = normalize(float3(-1.0f, -1.0f, -1.0f));
+		TopLayerData.WorldNormal = normalize(float3(-1.0f, -1.0f, -1.0f));
 	}
 
 	OutColor0 = StrataPackClassificationData(ClassificationData);
diff --git a/Engine/Shaders/Private/Strata/StrataVisualize.usf b/Engine/Shaders/Private/Strata/StrataVisualize.usf
--- a/Engine/Shaders/Private/Strata/StrataVisualize.usf
+++ b/Engine/Shaders/Private/Strata/StrataVisualize.usf
@@ -157,6 +157,8 @@ void VisualizeMaterialPS(
 		uint2 RectMin = uint2(MaterialPos + InspectSquare);
 		uint2 RectMax = RectMin;
 
+		const FStrataClassification StrataData = StrataUnpackClassificationData(Strata.ClassificationTexture.Load(uint3(MaterialPos, 0)));
+
 		FStrataAddressing StrataAddressing = GetStrataPixelDataByteOffset(MaterialPos, uint2(View.BufferSizeAndInvSize.xy), Strata.MaxBytesPerPixel);
 		const uint PixelStrataDataByteOffsetStart = StrataAddressing.ReadBytes;
 
@@ -192,6 +194,7 @@ void VisualizeMaterialPS(
 #define PRINTMATERIALSKIPPRECSHADOWMASK()	PRINTBSDFCHAR(_S_); PRINTBSDFCHAR(_K_); PRINTBSDFCHAR(_I_); PRINTBSDFCHAR(_P_); PRINTBSDFCHAR(_P_); PRINTBSDFCHAR(_R_); PRINTBSDFCHAR(_E_); PRINTBSDFCHAR(_S_); PRINTBSDFCHAR(_H_); PRINTBSDFCHAR(_A_); PRINTBSDFCHAR(_D_); PRINTBSDFCHAR(_SPC_); PRINTBSDFCHAR(_0_+(Header.SkipPrecShadowMask?1:0));
 #define PRINTMATERIALZEROPRECSHADOWMASK()	PRINTBSDFCHAR(_Z_); PRINTBSDFCHAR(_E_); PRINTBSDFCHAR(_R_); PRINTBSDFCHAR(_O_); PRINTBSDFCHAR(_P_); PRINTBSDFCHAR(_R_); PRINTBSDFCHAR(_E_); PRINTBSDFCHAR(_S_); PRINTBSDFCHAR(_H_); PRINTBSDFCHAR(_A_); PRINTBSDFCHAR(_D_); PRINTBSDFCHAR(_SPC_); PRINTBSDFCHAR(_0_+(Header.ZeroPrecShadowMask?1:0));
 #define PRINTMATERIALSLW()	PRINTBSDFCHAR(_S_); PRINTBSDFCHAR(_L_); PRINTBSDFCHAR(_W_); PRINTBSDFCHAR(_SPC_); PRINTBSDFCHAR(_0_+(Header.IsSingleLayerWater?1:0));
+#define PRINTMATERIALROUGH()PRINTBSDFCHAR(_R_); PRINTBSDFCHAR(_O_); PRINTBSDFCHAR(_U_); PRINTBSDFCHAR(_G_); PRINTBSDFCHAR(_H_); PRINTBSDFCHAR(_SPC_); PRINTFLOAT1(StrataData.TopLayerRoughness);
 
 #define NEXTLINE  TopLeft.y += 8; RectMax.x = max(RectMax.x, TopLeft.x); RectMax.y = max(RectMax.y, TopLeft.y); TopLeft.x = TopLeftInit.x; 
 #define HORIZONTALSEPARATION  TopLeft.y += 4;
@@ -211,6 +214,8 @@ void VisualizeMaterialPS(
 		NEXTLINE;
 		PRINTMATERIALSLW();
 		NEXTLINE;
+		PRINTMATERIALROUGH();
+		NEXTLINE;
 		HORIZONTALSEPARATION;
 #else
 		NEXTLINE;
@@ -220,6 +225,7 @@ void VisualizeMaterialPS(
 		NEXTLINE;
 		NEXTLINE;
 		NEXTLINE;
+		NEXTLINE;
 		HORIZONTALSEPARATION;
 #endif
 
diff --git a/Engine/Source/Runtime/Renderer/Private/Strata/Strata.cpp b/Engine/Source/Runtime/Renderer/Private/Strata/Strata.cpp
--- a/Engine/Source/Runtime/Renderer/Private/Strata/Strata.cpp
+++ b/Engine/Source/Runtime/Renderer/Private/Strata/Strata.cpp
@@ -97,9 +97,9 @@ void InitialiseStrataFrameSceneData(FSceneRenderer& SceneRenderer, FRDGBuilder&
 
 		// Top layer texture
 		{
-			FRDGTextureRef Texture = GraphBuilder.CreateTexture(FRDGTextureDesc::Create2D(BufferSizeXY, PF_R32G32_UINT, FClearValueBinding::Black, TexCreate_ShaderResource | TexCreate_RenderTargetable), TEXT("StrataTopLayerTexture"));
+			FRDGTextureRef Texture = GraphBuilder.CreateTexture(FRDGTextureDesc::Create2D(BufferSizeXY, PF_R32G32_UINT, FClearValueBinding::Black, TexCreate_ShaderResource | TexCreate_RenderTargetable), TEXT("StrataTopLayerNormalTexture"));
 			GraphBuilder.PreallocateTexture(Texture);
-			StrataSceneData.TopLayerTexture = GraphBuilder.GetPooledTexture(Texture);
+			StrataSceneData.TopLayerNormalTexture = GraphBuilder.GetPooledTexture(Texture);
 		}
 	}
 	else
@@ -158,7 +158,7 @@ TUniformBufferRef<FStrataGlobalUniformParameters> BindStrataGlobalUniformParamet
 		StrataUniformParameters.MaxBytesPerPixel = View.StrataSceneData->MaxBytesPerPixel;
 		StrataUniformParameters.MaterialLobesBuffer = View.StrataSceneData->MaterialLobesBuffer.SRV;
 		StrataUniformParameters.ClassificationTexture = View.StrataSceneData->ClassificationTexture->GetRenderTargetItem().ShaderResourceTexture;
-		StrataUniformParameters.TopLayerTexture = View.StrataSceneData->TopLayerTexture->GetRenderTargetItem().ShaderResourceTexture;
+		StrataUniformParameters.TopLayerNormalTexture = View.StrataSceneData->TopLayerNormalTexture->GetRenderTargetItem().ShaderResourceTexture;
 
 		View.StrataSceneData->StrataGlobalUniformParameters = CreateUniformBufferImmediate(StrataUniformParameters, UniformBuffer_SingleFrame);
 		return View.StrataSceneData->StrataGlobalUniformParameters;
@@ -169,7 +169,7 @@ TUniformBufferRef<FStrataGlobalUniformParameters> BindStrataGlobalUniformParamet
 		StrataUniformParameters.MaxBytesPerPixel = 0;
 		StrataUniformParameters.MaterialLobesBuffer = GEmptyVertexBufferWithUAV->ShaderResourceViewRHI;
 		StrataUniformParameters.ClassificationTexture = GSystemTextures.BlackDummy->GetRenderTargetItem().ShaderResourceTexture;
-		StrataUniformParameters.TopLayerTexture = GSystemTextures.BlackDummy->GetRenderTargetItem().ShaderResourceTexture;
+		StrataUniformParameters.TopLayerNormalTexture = GSystemTextures.BlackDummy->GetRenderTargetItem().ShaderResourceTexture;
 		return CreateUniformBufferImmediate(StrataUniformParameters, UniformBuffer_SingleDraw);
 	}
 }
@@ -300,7 +300,7 @@ class FStrataMaterialClassificationPassPS : public FGlobalShader
 		FGlobalShader::ModifyCompilationEnvironment(Parameters, OutEnvironment);
 		OutEnvironment.SetDefine(TEXT("SHADER_CATEGORIZATION"), 1);
 		OutEnvironment.SetRenderTargetOutputFormat(0, PF_R32_UINT);
-		OutEnvironment.SetRenderTargetOutputFormat(1, PF_R32G32_UINT);
+		OutEnvironment.SetRenderTargetOutputFormat(1, PF_R32_UINT);
 	}
 };
 IMPLEMENT_GLOBAL_SHADER(FStrataMaterialClassificationPassPS, "/Engine/Private/Strata/StrataMaterialClassification.usf", "MainPS", SF_Pixel);
@@ -566,7 +566,7 @@ void AddStrataMaterialClassificationPass(FRDGBuilder& GraphBuilder, const FMinim
 			PassParameters->Strata = Strata::BindStrataGlobalUniformParameters(View);
 			PassParameters->SceneTextures = GetSceneTextureParameters(GraphBuilder);		
 			PassParameters->RenderTargets[0] = FRenderTargetBinding(ClassificationTexture, ERenderTargetLoadAction::EClear);
-			PassParameters->RenderTargets[1] = FRenderTargetBinding(GraphBuilder.RegisterExternalTexture(View.StrataSceneData->TopLayerTexture), ERenderTargetLoadAction::EClear);
+			PassParameters->RenderTargets[1] = FRenderTargetBinding(GraphBuilder.RegisterExternalTexture(View.StrataSceneData->TopLayerNormalTexture), ERenderTargetLoadAction::EClear);
 
 			if (ShaderDrawDebug::IsShaderDrawDebugEnabled())
 			{
diff --git a/Engine/Source/Runtime/Renderer/Private/Strata/Strata.h b/Engine/Source/Runtime/Renderer/Private/Strata/Strata.h
--- a/Engine/Source/Runtime/Renderer/Private/Strata/Strata.h
+++ b/Engine/Source/Runtime/Renderer/Private/Strata/Strata.h
@@ -23,7 +23,7 @@ BEGIN_GLOBAL_SHADER_PARAMETER_STRUCT(FStrataGlobalUniformParameters, )
 	SHADER_PARAMETER(uint32, MaxBytesPerPixel)
 	SHADER_PARAMETER_SRV(ByteAddressBuffer, MaterialLobesBuffer)
 	SHADER_PARAMETER_TEXTURE(Texture2D<uint>, ClassificationTexture)
-	SHADER_PARAMETER_TEXTURE(Texture2D<uint2>, TopLayerTexture)
+	SHADER_PARAMETER_TEXTURE(Texture2D<uint>, TopLayerNormalTexture)
 END_GLOBAL_SHADER_PARAMETER_STRUCT()
 
 struct FStrataSceneData
@@ -31,7 +31,7 @@ struct FStrataSceneData
 	uint32 MaxBytesPerPixel;
 	FRWByteAddressBuffer MaterialLobesBuffer;					// This should be a RDG resource
 	TRefCountPtr<IPooledRenderTarget> ClassificationTexture;	// This should be a RDG resource
-	TRefCountPtr<IPooledRenderTarget> TopLayerTexture;			// This should be a RDG resource
+	TRefCountPtr<IPooledRenderTarget> TopLayerNormalTexture;	// This should be a RDG resource
 
 	TRefCountPtr<FRDGPooledBuffer> ClassificationTileListBuffer;
 	TRefCountPtr<FRDGPooledBuffer> ClassificationTileIndirectBuffer;
