Strata better SSS profile ID visu

diff --git a/Engine/Shaders/Private/Strata/StrataVisualize.usf b/Engine/Shaders/Private/Strata/StrataVisualize.usf
--- a/Engine/Shaders/Private/Strata/StrataVisualize.usf
+++ b/Engine/Shaders/Private/Strata/StrataVisualize.usf
@@ -294,7 +294,7 @@ void VisualizeMaterialPS(
 					if (BSDF_GETHASSSSPROFILE(BSDF))
 					{
 						HORIZONTALTAB; PRINTPARAMCHAR(_S_); PRINTPARAMCHAR(_S_); PRINTPARAMCHAR(_S_); PRINTPARAMCHAR(_SPC_); PRINTPARAMCHAR(_I_); PRINTPARAMCHAR(_D_); PRINTPARAMCHAR(_SPC_);
-						HORIZONTALTAB; PRINTFLOAT1(SLAB_SSSPROFILEID(BSDF));
+						HORIZONTALTAB; PRINTFLOATNOFRAC(floor(SLAB_SSSPROFILEID(BSDF)*255.0f+0.5f),3);
 						NEXTLINE;
 						HORIZONTALTAB; PRINTPARAMCHAR(_S_); PRINTPARAMCHAR(_S_); PRINTPARAMCHAR(_S_); PRINTPARAMCHAR(_S_); PRINTPARAMCHAR(_C_); PRINTPARAMCHAR(_L_); PRINTPARAMCHAR(_SPC_);
 						HORIZONTALTAB; PRINTFLOAT1(SLAB_SSSPROFILERADIUSSCALE(BSDF));
