Fix minor shading warnings.

diff --git a/Engine/Shaders/Private/HairStrands/HairStrandsInterpolation.usf b/Engine/Shaders/Private/HairStrands/HairStrandsInterpolation.usf
--- a/Engine/Shaders/Private/HairStrands/HairStrandsInterpolation.usf
+++ b/Engine/Shaders/Private/HairStrands/HairStrandsInterpolation.usf
@@ -635,7 +635,7 @@ void ClusterAABBEvaluationCS(uint3 DispatchThreadId : SV_DispatchThreadID, uint3
 
 		if (ThreadReadOffset < IndexCount)
 		{
-			const float3 OutHairPositionOffset = RenderDeformedOffsetBuffer[0];
+			const float3 OutHairPositionOffset = RenderDeformedOffsetBuffer[0].xyz;
 
 			const uint VertexIndex = ClusterVertexIdBuffer[IndexOffset + ThreadReadOffset];
 			const FHairControlPoint CP = UnpackHairControlPoint(RenderDeformedPositionBuffer[VertexIndex], OutHairPositionOffset, 1, 1);
diff --git a/Engine/Shaders/Private/HairStrands/HairStrandsVoxelRasterCompute.usf b/Engine/Shaders/Private/HairStrands/HairStrandsVoxelRasterCompute.usf
--- a/Engine/Shaders/Private/HairStrands/HairStrandsVoxelRasterCompute.usf
+++ b/Engine/Shaders/Private/HairStrands/HairStrandsVoxelRasterCompute.usf
@@ -97,7 +97,7 @@ void MainCS(uint3 DispatchThreadID : SV_DispatchThreadID)
 	}
 #endif
 
-	const float4 PositionOffset = HairStrandsVF_PositionOffsetBuffer[0];
+	const float3 PositionOffset = HairStrandsVF_PositionOffsetBuffer[0].xyz;
 
 	const FHairControlPoint CP0 = UnpackHairControlPoint(
 		HairStrandsVF_PositionBuffer[VertexIndex0],
