Better decal depth bias calculation

diff --git a/Engine/Shaders/Private/MeshDecals.usf b/Engine/Shaders/Private/MeshDecals.usf
--- a/Engine/Shaders/Private/MeshDecals.usf
+++ b/Engine/Shaders/Private/MeshDecals.usf
@@ -143,7 +143,7 @@ void MainVS(
 // Move all geometry a little bit towards the camera (to not get z fighting with existing zbuffer of the same meshes rendered with slightly different float computations).
 #if !USING_TESSELLATION
 	// todo: tweak or expose (cvar or per mesh/material setting), if needed the artist always can create a WorldPositionOffset adjustment in the shader (e.g. -0.1*CameraDirectionVector )
-	Output.Position.z += 0.0001f * Output.Position.w;
+	Output.Position.z *= 1.001;
 #endif
 
 	OutputVertexID( Output );
