Rearrage clipmaps so that their snapping aligns more across levels
- Levels are now strictly pow2 sizes and snap to half their size (instead of internal page boundaries).
- Propogate mips propgates along clipmap levels instead of mips for directional lights
- Provide a O(1) lookup for clipmaps as well that will fall back to coarser data. Do not use it by default for current opaque shading where we know by construction where the mapped pages are.
Still some minor cleanup and optimization possible to come in a future commit.

diff --git a/Engine/Shaders/Private/VirtualShadowMaps/PageManagement.usf b/Engine/Shaders/Private/VirtualShadowMaps/PageManagement.usf
--- a/Engine/Shaders/Private/VirtualShadowMaps/PageManagement.usf
+++ b/Engine/Shaders/Private/VirtualShadowMaps/PageManagement.usf
@@ -26,7 +26,6 @@ RWStructuredBuffer<uint> OutStatsBuffer;
 RWStructuredBuffer<uint> OutPageRequestFlags;
 
 Texture2D<UlongType> VisBuffer64;
-StructuredBuffer<FVirtualShadowMapProjectionShaderData> ShadowMapProjectionData;
 StructuredBuffer<int> VirtualShadowMapIdRemap;
 uint NumDirectionalLightSmInds;
 uint bPostBasePass;
@@ -110,6 +109,7 @@ FScreenPixelData GetScreenPixelData(uint2 PixelPos)
 uint GetMipLevel(int VirtualShadowMapId, FScreenPixelData Pixel)
 {
 	float Footprint = float(VSM_VIRTUAL_MAX_RESOLUTION_XY);
+	FVirtualShadowMapProjectionShaderData ProjectionData = GetVirtualShadowMapProjectionData(VirtualShadowMapId);
 
 	{
 		// Compute footprint by projecting the approximate size of a camera pixel at the given depth to shadow space
@@ -122,10 +122,10 @@ uint GetMipLevel(int VirtualShadowMapId, FScreenPixelData Pixel)
 
 		float RadiusWorld = DepthScale * RadiusScreen;
 
-		float3 TranslatedWorldPosition = Pixel.WorldPosition + ShadowMapProjectionData[VirtualShadowMapId].ShadowPreViewTranslation;
-		float4 ShadowUVz = mul(float4(TranslatedWorldPosition, 1.0f), ShadowMapProjectionData[VirtualShadowMapId].TranslatedWorldToShadowUVMatrix);
+		float3 TranslatedWorldPosition = Pixel.WorldPosition + ProjectionData.ShadowPreViewTranslation;
+		float4 ShadowUVz = mul(float4(TranslatedWorldPosition, 1.0f), ProjectionData.TranslatedWorldToShadowUVMatrix);
 
-		float4 RadiusClipH = mul(float4(RadiusWorld, 0.0f, ShadowUVz.w, 1.0f), ShadowMapProjectionData[VirtualShadowMapId].ShadowViewToClipMatrix);
+		float4 RadiusClipH = mul(float4(RadiusWorld, 0.0f, ShadowUVz.w, 1.0f), ProjectionData.ShadowViewToClipMatrix);
 		float RadiusClip = abs(RadiusClipH.x / RadiusClipH.w);
 		Footprint = RadiusClip * float(2 * VSM_VIRTUAL_MAX_RESOLUTION_XY);
 	}
@@ -135,7 +135,7 @@ uint GetMipLevel(int VirtualShadowMapId, FScreenPixelData Pixel)
 
 void MarkPage(uint VirtualShadowMapId, uint MipLevel, FScreenPixelData Pixel, float2 PageDilationOffset)
 {
-	FVirtualShadowMapProjectionShaderData ProjectionData = ShadowMapProjectionData[VirtualShadowMapId];
+	FVirtualShadowMapProjectionShaderData ProjectionData = GetVirtualShadowMapProjectionData(VirtualShadowMapId);
 	float4 ShadowUVz = mul(float4(Pixel.WorldPosition + ProjectionData.ShadowPreViewTranslation, 1.0f), ProjectionData.TranslatedWorldToShadowUVMatrix);
 	ShadowUVz.xyz /= ShadowUVz.w;
 
@@ -208,7 +208,7 @@ void GeneratePageFlagsFromPixels(uint3 GroupId : SV_GroupID, uint GroupIndex : S
 		// Mark one page per directional light. Lights can have varying numbers of clipmap levels allocated which
 		// will all be contiguous in the array, so we hop between them using the counts.
 		int ClipmapStartId = VirtualShadowMapIdRemap[Index];
-		FVirtualShadowMapProjectionShaderData BaseProjectionData = ShadowMapProjectionData[ClipmapStartId];
+		FVirtualShadowMapProjectionShaderData BaseProjectionData = GetVirtualShadowMapProjectionData(ClipmapStartId);
 
 		// Possibly use a single shadow projection data for the whole clipmap with just scales/biases for
 		// each level (to handle snapping). For now we maintain one per level.
@@ -602,40 +602,94 @@ void PropagateMappedMips(uint2 ThreadId : SV_DispatchThreadID, uint2 GroupId : S
 		return;
 	}
 
-	uint ShadowMapID = GroupId.y;
+	int VirtualShadowMapID = int(GroupId.y);
 	uint PageTableEntryIndex = ThreadId.x;
 
 	uint2 Level0Page;
 	Level0Page.x = PageTableEntryIndex & ((1U << VSM_LOG2_LEVEL0_DIM_PAGES_XY) - 1U);
 	Level0Page.y = PageTableEntryIndex >> VSM_LOG2_LEVEL0_DIM_PAGES_XY;
 
-	int MappedPageLevel = -1;
-	uint2 MappedPageIndex = 0;
-	
-	for (int Level = (VSM_MAX_MIP_LEVELS - 1); Level >= 0; --Level)
+	FVirtualShadowMapProjectionShaderData ProjectionData = GetVirtualShadowMapProjectionData(VirtualShadowMapID);
+
+	if (ProjectionData.LightType == LIGHT_TYPE_DIRECTIONAL)
 	{
-		uint2 vPage = Level0Page >> Level;
-		uint PageOffset = CalcPageOffset(ShadowMapID, Level, vPage);
-		FShadowPhysicalPage pPage = ShadowDecodePageTable(OutPageTable[PageOffset]);
+		// Directional lights propagate pages to their coarser/larger clipmap levels (and only use mip0 pages)
+		// Each clipmap level is a separate VSM, so we gather any mapped coarser pages as necessary and write only our own page output
+		// There's also technically a race similar to below with other threads writing the PT data we are reading,
+		// but it's still deterministic as long as we only look at pages with "bThisMipValid".
+		// There's some redundant work of course, but this shader is pretty cheap overall
+
+		uint Page0Offset = CalcPageOffset(VirtualShadowMapID, 0, Level0Page);
+		FShadowPhysicalPage pPage0 = ShadowDecodePageTable(OutPageTable[Page0Offset]);
 
 		BRANCH
-		if (pPage.bThisMipValid)
+		if (!pPage0.bThisMipValid)
 		{
-			// This page is mapped, so leave it alone and propogate downwards
-			MappedPageLevel = Level;
-			MappedPageIndex = pPage.PageIndex;
+			const int OffsetScale = (VSM_LEVEL0_DIM_PAGES_XY >> 2);
+			int2 BaseOffset = OffsetScale * (ProjectionData.ClipmapCenterSnapUnits * int2(-1, 1) + 2);
+			int2 BasePage   = int2(Level0Page) - BaseOffset;
+
+			// Search for first mapped page past this one
+			uint RemainingLevels = ProjectionData.ClipmapLevelCount - ProjectionData.ClipmapIndex;
+			for (uint ClipmapOffset = 1; ClipmapOffset < RemainingLevels; ++ClipmapOffset)
+			{
+				const int ClipmapLevelId = VirtualShadowMapID + int(ClipmapOffset);
+
+				FVirtualShadowMapProjectionShaderData LevelProjectionData = GetVirtualShadowMapProjectionData(ClipmapLevelId);								
+				int2 LevelOffset = OffsetScale * (LevelProjectionData.ClipmapCenterSnapUnits * int2(-1, 1) + 2);
+
+				const int LevelScale = 1 << ClipmapOffset;
+				int2 LevelPage = (BasePage + LevelScale * LevelOffset) / LevelScale;
+
+				if (all(LevelPage >= 0) && all(LevelPage < VSM_LEVEL0_DIM_PAGES_XY))
+				{
+					uint LevelPageOffset = CalcPageOffset(ClipmapLevelId, 0, uint2(LevelPage));
+					FShadowPhysicalPage pPage = ShadowDecodePageTable(OutPageTable[LevelPageOffset]);
+					if (pPage.bThisMipValid)
+					{
+						OutPageTable[Page0Offset] = ShadowEncodePageTable(pPage.PageIndex, ClipmapOffset);
+						break;
+					}
+				}
+				else
+				{
+					// TODO: We're off the edge... can we ever get back on?
+					// TODO: Can this ever even happen in practice given the construction?
+				}
+			}
 		}
-		else if( MappedPageLevel >= 0 )
+	}
+	else
+	{
+		// Local lights propagate pages to their coarser mips
+		int MappedPageLevel = -1;
+		uint2 MappedPageIndex = 0;
+	
+		for (int Level = (VSM_MAX_MIP_LEVELS - 1); Level >= 0; --Level)
 		{
-			// This page is not mapped; replace it with our suitably offset parent mapped page
-			// Ensure only one thread writes each value to avoid races, but we read on all threads as the broadcast
-			// Note that this can race with the other threads reading this value, but since bThisMipValid will
-			// always be false on these updated pages the values will be ignored. As long as the writes to the page
-			// table are atomic (currently a single DWORD), this is safe.
-			if (all((vPage << Level) == Level0Page))
+			uint2 vPage = Level0Page >> Level;
+			uint PageOffset = CalcPageOffset(VirtualShadowMapID, Level, vPage);
+			FShadowPhysicalPage pPage = ShadowDecodePageTable(OutPageTable[PageOffset]);
+
+			BRANCH
+			if (pPage.bThisMipValid)
+			{
+				// This page is mapped, so leave it alone and propogate downwards
+				MappedPageLevel = Level;
+				MappedPageIndex = pPage.PageIndex;
+			}
+			else if( MappedPageLevel >= 0 )
 			{
-				uint MipOffset = MappedPageLevel - Level;
-				OutPageTable[PageOffset] = ShadowEncodePageTable(MappedPageIndex, MipOffset);
+				// This page is not mapped; replace it with our suitably offset parent mapped page
+				// Ensure only one thread writes each value to avoid races, but we read on all threads as the broadcast
+				// Note that this can race with the other threads reading this value, but since bThisMipValid will
+				// always be false on these updated pages the values will be ignored. As long as the writes to the page
+				// table are atomic (currently a single DWORD), this is safe.
+				if (all((vPage << Level) == Level0Page))
+				{
+					uint MipOffset = MappedPageLevel - Level;
+					OutPageTable[PageOffset] = ShadowEncodePageTable(MappedPageIndex, MipOffset);
+				}
 			}
 		}
 	}
diff --git a/Engine/Shaders/Private/VirtualShadowMaps/ProjectionCommon.ush b/Engine/Shaders/Private/VirtualShadowMaps/ProjectionCommon.ush
--- a/Engine/Shaders/Private/VirtualShadowMaps/ProjectionCommon.ush
+++ b/Engine/Shaders/Private/VirtualShadowMaps/ProjectionCommon.ush
@@ -18,18 +18,21 @@ struct FVirtualShadowMapProjectionShaderData
 	float4x4 ShadowViewToClipMatrix;
 	float4x4 TranslatedWorldToShadowUVMatrix;
 	float4x4 TranslatedWorldToShadowUVNormalMatrix;
+
 	float3 ShadowPreViewTranslation;
 	uint LightType;		// Matches ELightComponentType via defines in SSDDefinitions.ush
 
 	float3 ClipmapWorldOrigin;
 	int VirtualShadowMapId;
-
+	
+	int2 ClipmapCenterSnapUnits;
+	int ClipmapIndex;		// 0..ClipmapLevelCount-1
 	int ClipmapLevel;
-	int ClipmapLevelCount;
-	float ClipmapResolutionLodBias;
+	int ClipmapLevelCount;	
+	float ClipmapResolutionLodBias;	
 
 	// Seems the FMatrix forces 16-byte alignment, so we need to pad
-	float Padding[1];
+	float Padding[2];
 };
 
 // FVirtualShadowMapProjectionParameters
@@ -67,6 +70,7 @@ struct FVirtualSmSample
 {
 	float Depth;
 	uint MipLevel;
+	int VirtualShadowMapId;		// May be offset from input for clipmaps
 	bool bValid;
 	uint2 VirtualTexelAddress;
 	float2 VirtualTexelAddressFloat;
@@ -78,6 +82,7 @@ FVirtualSmSample InitVirtualSmSample()
 	FVirtualSmSample Result;
 	Result.Depth = 0.0f;
 	Result.MipLevel = 0;
+	Result.VirtualShadowMapId = -1; 
 	Result.bValid = false;
 	Result.VirtualTexelAddress = Result.PhysicalTexelAddress = uint2(0U, 0U);
 	Result.VirtualTexelAddressFloat = float2(0.0f, 0.0f);
@@ -108,12 +113,87 @@ FVirtualSmSample SampleVirtualSmLevel(int VirtualShadowMapId, float2 ShadowMapUV
 	{
 		Result.Depth = SampleVirtualSmPhysicalDepth(Result.PhysicalTexelAddress);
 		Result.MipLevel = VSMLevel;
+		Result.VirtualShadowMapId = VirtualShadowMapId;
 		Result.bValid = true;
 	}
 	
 	return Result;
 }
 
+// Data to convert UV's and depths from one clipmap level basis to another
+struct FVirtualShadowMapClipmapRelativeTransform
+{
+	float Scale;
+	float3 Bias;
+};
+
+// NOTE: Two virtual shadow IDs must reference different levels from the same clipmap
+// Resulting transform takes coordinates and depths from A to B
+FVirtualShadowMapClipmapRelativeTransform CalcClipmapRelativeTransform(int VirtualShadowIdA, int VirtualShadowIdB)
+{
+	const FVirtualShadowMapProjectionShaderData ProjectionDataA = GetVirtualShadowMapProjectionData(VirtualShadowIdA);
+	const FVirtualShadowMapProjectionShaderData ProjectionDataB = GetVirtualShadowMapProjectionData(VirtualShadowIdB);
+
+	// TODO: Simplify/optimize
+	float2 UVA = float2(ProjectionDataA.ClipmapCenterSnapUnits) * float2(-0.25f, 0.25f) + float2(0.5f, 0.5f);
+	float2 UVB = float2(ProjectionDataB.ClipmapCenterSnapUnits) * float2(-0.25f, 0.25f) + float2(0.5f, 0.5f);
+
+	FVirtualShadowMapClipmapRelativeTransform Result;
+	Result.Scale = pow(2.0f, float(VirtualShadowIdA - VirtualShadowIdB));
+	Result.Bias.xy = (UVB - Result.Scale * UVA);
+	// NOTE: Currently no snapping in Z
+	Result.Bias.z = 0.5f - Result.Scale * 0.5f;
+	return Result;
+}
+
+
+
+
+// Will sample from the given clipmap level Id, onwards to coarser levels if no valid data is present
+// ShadowMapUVs are in terms of the given clipmap level/virtual shadow map
+FVirtualSmSample SampleVirtualSmClipmap(int VirtualShadowMapId, float2 ShadowMapUV)
+{
+	FVirtualSmSample Result = InitVirtualSmSample();
+		
+#if 1
+	FShadowPageTranslationResult Page = ShadowVirtualToPhysicalUV(VirtualShadowMapId, ShadowMapUV);
+	if (Page.bValid)
+	{
+		// TODO: Clean up - clipmap offset in mip level
+		int ClipmapId = VirtualShadowMapId + Page.MipLevel;
+
+		FVirtualShadowMapClipmapRelativeTransform Transform = CalcClipmapRelativeTransform(VirtualShadowMapId, ClipmapId);
+
+		float2 ClipmapUV = ShadowMapUV * Transform.Scale + Transform.Bias.xy;
+		Result = SampleVirtualSmLevel(ClipmapId, ClipmapUV, 0);
+		if (Result.bValid)
+		{
+			// Convert depth into back into the original reference clipmap level range
+			Result.Depth = (Result.Depth - Transform.Bias.z) / Transform.Scale;
+		}
+	}
+#else
+	FVirtualShadowMapProjectionShaderData BaseProjectionData = GetVirtualShadowMapProjectionData(VirtualShadowMapId);
+	const int LastId = VirtualShadowMapId + BaseProjectionData.ClipmapLevelCount;
+	for (int ClipmapId = VirtualShadowMapId; ClipmapId < VirtualShadowMapId + 2; ++ClipmapId)
+	{
+		FVirtualShadowMapClipmapRelativeTransform Transform = CalcClipmapRelativeTransform(VirtualShadowMapId, ClipmapId);
+		
+		float2 ClipmapUV = ShadowMapUV * Transform.Scale + Transform.Bias.xy;
+		Result = SampleVirtualSmLevel(ClipmapId, ClipmapUV, 0);
+		if (Result.bValid)
+		{
+			// Convert depth into back into the original reference clipmap level range
+			Result.Depth = (Result.Depth - Transform.Bias.z) / Transform.Scale;
+			break;
+		}
+	}
+#endif
+
+	return Result;
+}
+
+
 FVirtualSmSample SampleVirtualSm(int VirtualShadowMapId, float2 ShadowMapUV)
 {
 	FShadowPageTranslationResult Page = ShadowVirtualToPhysicalUV(VirtualShadowMapId, ShadowMapUV);
@@ -122,6 +202,7 @@ FVirtualSmSample SampleVirtualSm(int VirtualShadowMapId, float2 ShadowMapUV)
 		FVirtualSmSample Result;
 		Result.bValid = true;
 		Result.MipLevel = Page.MipLevel;
+		Result.VirtualShadowMapId = VirtualShadowMapId;
 		Result.VirtualTexelAddress = Page.VirtualTexelAddress;
 		Result.VirtualTexelAddressFloat = Page.VirtualTexelAddressFloat;
 		Result.PhysicalTexelAddress = Page.PhysicalTexelAddress;
@@ -132,8 +213,10 @@ FVirtualSmSample SampleVirtualSm(int VirtualShadowMapId, float2 ShadowMapUV)
 	return InitVirtualSmSample();
 }
 
-float ComputeVirtualShadowMapOptimalSlopeBias(FVirtualShadowMapProjectionShaderData ProjectionData, FVirtualSmSample SmSample, float3 WorldPosition, float3 EstimatedGeoWorldNormal)
+float ComputeVirtualShadowMapOptimalSlopeBias(FVirtualSmSample SmSample, float3 WorldPosition, float3 EstimatedGeoWorldNormal)
 {
+	FVirtualShadowMapProjectionShaderData ProjectionData = GetVirtualShadowMapProjectionData(SmSample.VirtualShadowMapId);
+
 	// Transform geometry world-space plane eq to shadow 'UV' texture space [0-1] ranges
 	float4 NormalPlaneTranslatedWorld = float4(EstimatedGeoWorldNormal, -dot(EstimatedGeoWorldNormal, WorldPosition + ProjectionData.ShadowPreViewTranslation));
 	float4 NormalPlaneUV = mul(NormalPlaneTranslatedWorld, ProjectionData.TranslatedWorldToShadowUVNormalMatrix);
@@ -229,17 +312,22 @@ FVirtualShadowMapSampleResult SampleVirtualShadowMap(int VirtualShadowMapId, flo
 
 			// No perspective divided needed for ortho projection
 			float4 ShadowUVz = mul(float4(WorldPosition + ProjectionData.ShadowPreViewTranslation, 1.0f), ProjectionData.TranslatedWorldToShadowUVMatrix);
-
-			FVirtualSmSample SmSample = SampleVirtualSmLevel(ClipmapLevelVirtualShadowMapId, ShadowUVz.xy, 0);
+			
+			FVirtualSmSample SmSample;
+			SmSample = SampleVirtualSmLevel(ClipmapLevelVirtualShadowMapId, ShadowUVz.xy, 0);
+			//SmSample = SampleVirtualSmClipmap(ClipmapLevelVirtualShadowMapId, ShadowUVz.xy);
+			
 			if (SmSample.bValid)
 			{
+				int SampledClipmapIndex = SmSample.VirtualShadowMapId - VirtualShadowMapId;
+
 				FVirtualShadowMapSampleResult Result;
 				Result.bValid = true;
 				Result.bOccluded = false;
 				Result.OccluderDistance = -1.0f;
-				Result.ClipmapIndexOrMipLevel = ClipmapIndex;
+				Result.ClipmapIndexOrMipLevel = SampledClipmapIndex;
 
-				float OptimalSlopeBias = ComputeVirtualShadowMapOptimalSlopeBias(ProjectionData, SmSample, WorldPosition, EstimatedGeoWorldNormal);
+				float OptimalSlopeBias = ComputeVirtualShadowMapOptimalSlopeBias(SmSample, WorldPosition, EstimatedGeoWorldNormal);
 				float RayStartBias = -RayStartDistance * ProjectionData.ShadowViewToClipMatrix._33;
 				float BiasedDepth = SmSample.Depth - OptimalSlopeBias - RayStartBias;
 
@@ -279,7 +367,7 @@ FVirtualShadowMapSampleResult SampleVirtualShadowMap(int VirtualShadowMapId, flo
 			Result.OccluderDistance = -1.0f;
 			Result.ClipmapIndexOrMipLevel = SmSample.MipLevel;
 
-			float OptimalSlopeBias = ComputeVirtualShadowMapOptimalSlopeBias(BaseProjectionData, SmSample, WorldPosition, EstimatedGeoWorldNormal);
+			float OptimalSlopeBias = ComputeVirtualShadowMapOptimalSlopeBias(SmSample, WorldPosition, EstimatedGeoWorldNormal);
 			float RayStartBias = -RayStartDistance * BaseProjectionData.ShadowViewToClipMatrix._33 / ShadowUVz.w;
 			float BiasedDepth = SmSample.Depth - OptimalSlopeBias - RayStartBias;
 
diff --git a/Engine/Shaders/Private/VirtualShadowMaps/ProjectionDirectional.ush b/Engine/Shaders/Private/VirtualShadowMaps/ProjectionDirectional.ush
--- a/Engine/Shaders/Private/VirtualShadowMaps/ProjectionDirectional.ush
+++ b/Engine/Shaders/Private/VirtualShadowMaps/ProjectionDirectional.ush
@@ -49,8 +49,6 @@ FSMRTClipmapRayState SMRTClipmapRayInitialize(
 	float3 RayStepShadowTexelsZ  = float3(LevelDims, LevelDims, 1.0f) * RayStepUV.xyz;
 
 	// Work out scale/bias to convert to the adjacent UVs from the regular UVs
-	// NOTE: Since depth function is constant, probably worth refactoring the clipmap mappings to be scale/bias in general rather than
-	// a full transform matrix.
 	const FVirtualShadowMapProjectionShaderData ProjectionDataAdjacent = GetVirtualShadowMapProjectionData(VirtualShadowMapId + ClipmapIndexAdjacent);
 
 	// TranslatedWorldToShadowViewMatrix is identical for each of the cascades; only the pre-view translation changes
diff --git a/Engine/Source/Runtime/Renderer/Private/VirtualShadowMaps/VirtualShadowMapArray.cpp b/Engine/Source/Runtime/Renderer/Private/VirtualShadowMaps/VirtualShadowMapArray.cpp
--- a/Engine/Source/Runtime/Renderer/Private/VirtualShadowMaps/VirtualShadowMapArray.cpp
+++ b/Engine/Source/Runtime/Renderer/Private/VirtualShadowMaps/VirtualShadowMapArray.cpp
@@ -286,7 +286,7 @@ class FGeneratePageFlagsFromPixelsCS : public FVirtualPageManagementShader
 		SHADER_PARAMETER_STRUCT_REF(FForwardLightData, ForwardLightData)
 		SHADER_PARAMETER_RDG_TEXTURE(Texture2D<uint2>, VisBuffer64)
 		SHADER_PARAMETER_RDG_BUFFER_UAV(RWStructuredBuffer<uint>, OutPageRequestFlags)
-		SHADER_PARAMETER_RDG_BUFFER_SRV(StructuredBuffer< FVirtualShadowMapProjectionShaderData >, ShadowMapProjectionData)
+		SHADER_PARAMETER_RDG_BUFFER_SRV(StructuredBuffer< FVirtualShadowMapProjectionShaderData >, VirtualShadowMapProjectionData)
 		SHADER_PARAMETER_RDG_BUFFER_SRV(StructuredBuffer< uint >, VirtualShadowMapIdRemap)
 		SHADER_PARAMETER(uint32, NumDirectionalLightSmInds)
 		SHADER_PARAMETER(uint32, bPostBasePass)
@@ -355,7 +355,8 @@ class FPropagateMappedMipsCS : public FVirtualPageManagementShader
 	SHADER_USE_PARAMETER_STRUCT(FPropagateMappedMipsCS, FVirtualPageManagementShader)
 
 	BEGIN_SHADER_PARAMETER_STRUCT(FParameters, )
-		SHADER_PARAMETER_RDG_UNIFORM_BUFFER(FVirtualShadowMapCommonParameters,	VirtualSmCommon)
+		SHADER_PARAMETER_RDG_UNIFORM_BUFFER( FVirtualShadowMapCommonParameters,	VirtualSmCommon )
+		SHADER_PARAMETER_RDG_BUFFER_SRV( StructuredBuffer< FVirtualShadowMapProjectionShaderData >, VirtualShadowMapProjectionData )
 		SHADER_PARAMETER_RDG_BUFFER_UAV( RWStructuredBuffer< uint >,			OutPageTable )
 	END_SHADER_PARAMETER_STRUCT()
 };
@@ -815,7 +816,7 @@ void FVirtualShadowMapArray::BuildPageAllocations(
 				PassParameters->OutPageRequestFlags = GraphBuilder.CreateUAV(PageRequestFlagsRDG);
 				PassParameters->ForwardLightData = View.ForwardLightingResources->ForwardLightDataUniformBuffer;
 				PassParameters->VirtualShadowMapIdRemap = GraphBuilder.CreateSRV(VirtualShadowMapIdRemapRDG[ViewIndex]);
-				PassParameters->ShadowMapProjectionData = GraphBuilder.CreateSRV(ShadowMapProjectionDataRDG);
+				PassParameters->VirtualShadowMapProjectionData = GraphBuilder.CreateSRV(ShadowMapProjectionDataRDG);
 				PassParameters->NumDirectionalLightSmInds = DirectionalLightSmInds.Num();
 				PassParameters->LodFootprintScale = LodFootprintScale;
 				PassParameters->PageDilationBorderSize = PageDilationBorderSize;
@@ -928,7 +929,7 @@ void FVirtualShadowMapArray::BuildPageAllocations(
 			FCreatePageMappingsCS::FParameters* PassParameters = GraphBuilder.AllocParameters< FCreatePageMappingsCS::FParameters >();
 			PassParameters->VirtualSmCommon = GetCommonUniformBuffer(GraphBuilder);
 
-			PassParameters->PageRequestFlags		= GraphBuilder.CreateSRV(PageRequestFlagsRDG );
+			PassParameters->PageRequestFlags		= GraphBuilder.CreateSRV( PageRequestFlagsRDG );
 			PassParameters->CoverageSummaryInOut	= GraphBuilder.CreateUAV( CoverageSummary );
 			PassParameters->AllocatedPagesOffset	= GraphBuilder.CreateUAV( AllocatedPagesOffsetRDG);
 			PassParameters->OutPageTable			= GraphBuilder.CreateUAV( PageTableRDG );
@@ -981,6 +982,7 @@ void FVirtualShadowMapArray::BuildPageAllocations(
 			// Propagate mapped mips down the hierarchy to allow O(1) lookup of coarser mapped pages
 			FPropagateMappedMipsCS::FParameters* PassParameters = GraphBuilder.AllocParameters<FPropagateMappedMipsCS::FParameters>();
 			PassParameters->VirtualSmCommon = GetCommonUniformBuffer(GraphBuilder);
+			PassParameters->VirtualShadowMapProjectionData = GraphBuilder.CreateSRV(ShadowMapProjectionDataRDG);
 			PassParameters->OutPageTable	= GraphBuilder.CreateUAV(PageTableRDG);
 
 			const uint32 DispatchWidthThreads = FVirtualShadowMap::Level0DimPagesXY;
diff --git a/Engine/Source/Runtime/Renderer/Private/VirtualShadowMaps/VirtualShadowMapArray.h b/Engine/Source/Runtime/Renderer/Private/VirtualShadowMaps/VirtualShadowMapArray.h
--- a/Engine/Source/Runtime/Renderer/Private/VirtualShadowMaps/VirtualShadowMapArray.h
+++ b/Engine/Source/Runtime/Renderer/Private/VirtualShadowMaps/VirtualShadowMapArray.h
@@ -90,21 +90,21 @@ struct FVirtualShadowMapProjectionShaderData
 	FMatrix ShadowViewToClipMatrix;
 	FMatrix TranslatedWorldToShadowUVMatrix;
 	FMatrix TranslatedWorldToShadowUVNormalMatrix;
-	/**
-	 * Translation from world space to shadow space (add before transform by TranslatedWorldToShadowViewMatrix).
-	 */
+
 	FVector ShadowPreViewTranslation;
 	uint32 LightType = ELightComponentType::LightType_Directional;
 	
 	FVector ClipmapWorldOrigin;
 	int32 VirtualShadowMapId = INDEX_NONE;
-
+	
+	FIntPoint ClipmapCenterSnapUnits;
+	int32 ClipmapIndex = 0;					// 0 .. ClipmapLevelCount-1
 	int32 ClipmapLevel = 0;					// "Absolute" level, can be negative
 	int32 ClipmapLevelCount = 0;
 	float ClipmapResolutionLodBias = 0.0f;
 
 	// Seems the FMatrix forces 16-byte alignment
-	//float Padding[3];
+	float Padding[2];
 };
 static_assert((sizeof(FVirtualShadowMapProjectionShaderData) % 16) == 0, "FVirtualShadowMapProjectionShaderData size should be a multiple of 16-bytes for alignment.");
 
diff --git a/Engine/Source/Runtime/Renderer/Private/VirtualShadowMaps/VirtualShadowMapClipmap.cpp b/Engine/Source/Runtime/Renderer/Private/VirtualShadowMaps/VirtualShadowMapClipmap.cpp
--- a/Engine/Source/Runtime/Renderer/Private/VirtualShadowMaps/VirtualShadowMapClipmap.cpp
+++ b/Engine/Source/Runtime/Renderer/Private/VirtualShadowMaps/VirtualShadowMapClipmap.cpp
@@ -26,7 +26,7 @@ static TAutoConsoleVariable<int32> CVarVirtualShadowMapClipmapFirstLevel(
 static TAutoConsoleVariable<float> CVarVirtualShadowMapClipmapMaxRadius(
 	TEXT( "r.Shadow.v.Clipmap.MaxRadius" ),
 	1000000.0f,
-	TEXT( "Maximum distance the clipmap covers. Determines the number of clipmap levels. " ),
+	TEXT( "Maximum distance the clipmap covers. Determines the number of clipmap levels." ),
 	ECVF_RenderThreadSafe
 );
 
@@ -62,9 +62,10 @@ FVirtualShadowMapClipmap::FVirtualShadowMapClipmap(
 	FMatrix ViewToWorldRotationMatrix = WorldToViewRotationMatrix.GetTransposed();
 	
 	// NOTE: Rotational (roll) invariance of the directional light depends on square pixels so we just base everything on the camera X scales/resolution
-	float LodScale = 1.0f / CameraViewMatrices.GetProjectionScale().X;
+	// NOTE: 0.5 because we double the size of the clipmap region below to handle snapping
+	float LodScale = 0.5f / CameraViewMatrices.GetProjectionScale().X;
 	LodScale *= float(FVirtualShadowMap::VirtualMaxResolutionXY) / float(CameraViewRectSize.X);
-	
+		
 	ResolutionLodBias = CVarVirtualShadowMapClipmapResolutionLodBias.GetValueOnRenderThread() + FMath::Log2(LodScale);
 	// Clamp negative absolute resolution biases as they would exceed the maximum resolution/ranges allocated
 	ResolutionLodBias = FMath::Max(0.0f, ResolutionLodBias);
@@ -94,30 +95,26 @@ FVirtualShadowMapClipmap::FVirtualShadowMapClipmap(
 		Level.VirtualShadowMap = VirtualShadowMapArray.Allocate();
 		ensure(Index == 0 || (Level.VirtualShadowMap->ID == (LevelData[Index-1].VirtualShadowMap->ID + 1)));
 
-		// Snap world center point to page boundaries
-		// Max error from snapping with round to nearest is the diagonal sqrt(2) * PageSize
-		// Must ensure radius is adjusted to cover the entire original radius
-		constexpr float DimPages = static_cast<float>(FVirtualShadowMap::Level0DimPagesXY);
 		const float RawLevelRadius = GetLevelRadius(LevelIndex);
-		float PageSizeInWorldSpace = (2.0f * RawLevelRadius) / (DimPages - UE_SQRT_2);
 
-		// NOTE: We may eventually want to snap/quantize Z as well, but right now we just adjust it each time
-		// we update the cache page. With clipmaps this could ertainly be simplified.
+		float SnappedLevelRadius = 2.0f * RawLevelRadius;
+		float SnapSize = RawLevelRadius;
+
 		FVector ViewCenter = WorldToViewRotationMatrix.TransformPosition(WorldOrigin);
-		FIntPoint PageSpaceCenter(
-			FMath::RoundToInt(ViewCenter.X / PageSizeInWorldSpace),
-			FMath::RoundToInt(ViewCenter.Y / PageSizeInWorldSpace));
-		ViewCenter.X = PageSpaceCenter.X * PageSizeInWorldSpace;
-		ViewCenter.Y = PageSpaceCenter.Y * PageSizeInWorldSpace;
+		FIntPoint CenterSnapUnits(
+			FMath::RoundToInt(ViewCenter.X / SnapSize),
+			FMath::RoundToInt(ViewCenter.Y / SnapSize));
+		ViewCenter.X = CenterSnapUnits.X * SnapSize;
+		ViewCenter.Y = CenterSnapUnits.Y * SnapSize;
 
 		const FVector SnappedWorldCenter = ViewToWorldRotationMatrix.TransformPosition(ViewCenter);
-		const float SnappedLevelRadius = 0.5f * (PageSizeInWorldSpace * DimPages);
 
 		const float ZScale = 0.5f / RawLevelRadius;
 		const float ZOffset = RawLevelRadius;
 
 		Level.WorldCenter = SnappedWorldCenter;
 		Level.ViewToClip = FReversedZOrthoMatrix(SnappedLevelRadius, SnappedLevelRadius, ZScale, ZOffset);
+		Level.CenterSnapUnits = CenterSnapUnits;
 
 		if (VirtualShadowMapArrayCacheManager)
 		{
@@ -125,7 +122,9 @@ FVirtualShadowMapClipmap::FVirtualShadowMapClipmap(
 			TSharedPtr<FVirtualShadowMapCacheEntry> CacheEntry = VirtualShadowMapArrayCacheManager->FindCreateCacheEntry(LightSceneInfo.Id, LevelIndex);
 			if (CacheEntry)
 			{
-				FIntPoint PageOffset(PageSpaceCenter);
+				// We snap to half the size of the VSM at each level
+				check((FVirtualShadowMap::Level0DimPagesXY & 1) == 0);
+				FIntPoint PageOffset(CenterSnapUnits * (FVirtualShadowMap::Level0DimPagesXY >> 1));
 				PageOffset.Y = -PageOffset.Y;		// Viewport
 				float DepthOffset = -ViewCenter.Z * ZScale;
 
@@ -171,9 +170,11 @@ FVirtualShadowMapProjectionShaderData FVirtualShadowMapClipmap::GetProjectionSha
 	Data.VirtualShadowMapId = Level.VirtualShadowMap->ID;
 	Data.LightType = ELightComponentType::LightType_Directional;
 	Data.ClipmapWorldOrigin = WorldOrigin;
+	Data.ClipmapIndex = ClipmapIndex;
 	Data.ClipmapLevel = FirstLevel + ClipmapIndex;
 	Data.ClipmapLevelCount = LevelData.Num();
 	Data.ClipmapResolutionLodBias = ResolutionLodBias;
+	Data.ClipmapCenterSnapUnits = Level.CenterSnapUnits;
 
 	return Data;
 }
diff --git a/Engine/Source/Runtime/Renderer/Private/VirtualShadowMaps/VirtualShadowMapClipmap.h b/Engine/Source/Runtime/Renderer/Private/VirtualShadowMaps/VirtualShadowMapClipmap.h
--- a/Engine/Source/Runtime/Renderer/Private/VirtualShadowMaps/VirtualShadowMapClipmap.h
+++ b/Engine/Source/Runtime/Renderer/Private/VirtualShadowMaps/VirtualShadowMapClipmap.h
@@ -88,6 +88,7 @@ private:
 		FVirtualShadowMap* VirtualShadowMap;
 		FMatrix ViewToClip;
 		FVector WorldCenter;
+		FIntPoint CenterSnapUnits;	// Center of clipmap in world units / level radius
 	};
 
 	TArray< FLevelData, TInlineAllocator<16> > LevelData;
