More Strata and basepass classification clean up.
Shadow/transmission passes now are strata only.
FStrataSubsurfaceData is now output from classification.
Removed unused bHasTransmission.

diff --git a/Engine/Shaders/Private/BasePassPixelShader.usf b/Engine/Shaders/Private/BasePassPixelShader.usf
--- a/Engine/Shaders/Private/BasePassPixelShader.usf
+++ b/Engine/Shaders/Private/BasePassPixelShader.usf
@@ -1555,7 +1555,6 @@ void FPixelShaderInOut_MainPS(
 
 	// Clear only what should be and keep the rest
 	GBuffer = (FGBufferData)0;						// We should not rely on gbuffer for anything else to make any pass work.
-	GBuffer.ShadingModelID = SHADINGMODELID_STRATA;	// STRATA_TODO remove that after SSS has been moved to strata completely
 
 	uint2 PixelPos = In.SvPosition.xy;
 
diff --git a/Engine/Shaders/Private/DeferredShadingCommon.ush b/Engine/Shaders/Private/DeferredShadingCommon.ush
--- a/Engine/Shaders/Private/DeferredShadingCommon.ush
+++ b/Engine/Shaders/Private/DeferredShadingCommon.ush
@@ -713,15 +713,14 @@ float3 ExtractSubsurfaceColor(FGBufferData BufferData)
 	return Square(BufferData.CustomData.rgb);
 }
 
+uint ExtractSubsurfaceProfileInt(float ProfileNormFloat)
+{
+	return uint(ProfileNormFloat * 255.0f + 0.5f);
+}
+
 uint ExtractSubsurfaceProfileInt(FGBufferData BufferData)
 {
-	// can be optimized
-#if STRATA_ENABLED
-	const bool bUseProfileData = BufferData.CustomData.y == 0 && BufferData.CustomData.z == 0;
-	return bUseProfileData ? uint(BufferData.CustomData.r * 255.0f + 0.5f) : SSS_PROFILE_ID_PERPIXEL;
-#else
-	return uint(BufferData.CustomData.r * 255.0f + 0.5f);
-#endif
+	return ExtractSubsurfaceProfileInt(BufferData.CustomData.r);
 }
 
 #if SHADING_PATH_DEFERRED
diff --git a/Engine/Shaders/Private/DistanceFieldShadowing.usf b/Engine/Shaders/Private/DistanceFieldShadowing.usf
--- a/Engine/Shaders/Private/DistanceFieldShadowing.usf
+++ b/Engine/Shaders/Private/DistanceFieldShadowing.usf
@@ -711,9 +711,7 @@ void DistanceFieldShadowingCS(
 			BRANCH
 			if (IsSubsurfaceModel(GBufferData.ShadingModelID))
 			{
-				// Note: this has to match shadowmap logic
-				// Derive density from a heuristic using opacity, tweaked for useful falloff ranges and to give a linear depth falloff with opacity
-				SubsurfaceDensity = -.05f * log(1 - min(GBufferData.CustomData.a, .999f));
+				SubsurfaceDensity = SubsurfaceDensityFromOpacity(GBufferData.CustomData.a);
 				bUseSubsurfaceTransmission = true;
 			}
 #endif
diff --git a/Engine/Shaders/Private/ShadowProjectionPixelShader.usf b/Engine/Shaders/Private/ShadowProjectionPixelShader.usf
--- a/Engine/Shaders/Private/ShadowProjectionPixelShader.usf
+++ b/Engine/Shaders/Private/ShadowProjectionPixelShader.usf
@@ -257,32 +257,51 @@ void Main(
 
 	#if FEATURE_LEVEL >= FEATURE_LEVEL_SM4 && !FORWARD_SHADING && !APPLY_TRANSLUCENCY_SHADOWS && !SUBPIXEL_SHADOW
 
-		FGBufferData GBufferData	= GetGBufferData(ScreenUV); 
-		bool bHasSubsurface			= IsSubsurfaceModel(GBufferData.ShadingModelID);
-		bool bHasSubsurfaceProfile	= GBufferData.ShadingModelID == SHADINGMODELID_SUBSURFACE_PROFILE;
-		bool bIsEyeOrHair 			= GBufferData.ShadingModelID == SHADINGMODELID_HAIR || GBufferData.ShadingModelID == SHADINGMODELID_EYE;
-		#if STRATA_ENABLED
+		bool bHasSubsurface			= false;
+		bool bHasSubsurfaceProfile	= false;
+		bool bIsEyeOrHair 			= false;
+	#if STRATA_ENABLED
+		const FStrataClassification StrataData = StrataUnpackClassificationData(Strata.ClassificationTexture.Load(PixelCoord));
 		{
-			const FStrataClassification StrataData = StrataUnpackClassificationData(Strata.ClassificationTexture.Load(PixelCoord));
 			if (StrataData.bIsStrataMaterial)
 			{
-				bIsEyeOrHair			= bIsEyeOrHair || StrataHasShadingModel(StrataData.ShadingModels, STRATA_BSDF_TYPE_HAIR);
-				bHasSubsurface			= bHasSubsurface || StrataData.bHasSubsurface || bIsEyeOrHair;
-				bHasSubsurfaceProfile	= bHasSubsurfaceProfile || StrataData.bHasSubsurfaceProfile;
+				bIsEyeOrHair			= StrataHasShadingModel(StrataData.ShadingModels, STRATA_BSDF_TYPE_HAIR);
+				bHasSubsurface			= StrataData.bHasSubsurface || bIsEyeOrHair;
+				bHasSubsurfaceProfile	= StrataData.bHasSubsurfaceProfile;
 			}
 		}
-		#endif
+	#else // STRATA_ENABLED
+		FGBufferData GBufferData = GetGBufferData(ScreenUV);
+		{
+			bHasSubsurface = IsSubsurfaceModel(GBufferData.ShadingModelID);
+			bHasSubsurfaceProfile = GBufferData.ShadingModelID == SHADINGMODELID_SUBSURFACE_PROFILE;
+			bIsEyeOrHair = GBufferData.ShadingModelID == SHADINGMODELID_HAIR || GBufferData.ShadingModelID == SHADINGMODELID_EYE;
+		}
+	#endif // STRATA_ENABLED
 
 		BRANCH
 		if (bIsSubsurfaceCompatible && bHasSubsurface)
 		{
+		#if STRATA_ENABLED
+			const FStrataSubsurfaceData SSSData = StrataUnpackSSSData(Strata.SSSTexture.Load(uint3(SVPos.xy, 0)));
+
+			float Opacity = 1;
+			float Density = 1;
+			if (!bIsEyeOrHair && SSSData.bIsProfile)
+			{
+				Opacity = SSSData.ProfileRadiusScale;
+				Density = SubsurfaceDensityFromOpacity(Opacity);
+			}
+		#else // STRATA_ENABLED
 			float Opacity = GBufferData.CustomData.a;
 			float Density = SubsurfaceDensityFromOpacity(Opacity);
-			if(bIsEyeOrHair)
+			if (bIsEyeOrHair)
 			{
 				Opacity = 1;
 				Density = 1;
 			}
+		#endif // STRATA_ENABLED
+
 			//@todo - Use a larger kernel filter based on material opacity
 			// This is not being done currently because a varying opacity causes discontinuities in the shadow result
 			// Note: Snapping based on the largest size kernel is done in FProjectedShadowInfo::FProjectedShadowInfo, MaxDownsampleFactor is set to 4
@@ -311,7 +330,15 @@ void Main(
 			  #if USE_TRANSMISSION
 				if (bHasSubsurfaceProfile)
 				{
-					SSSTransmission = CalcTransmissionThickness(ScreenPosition.xyz, LightPositionAndInvRadius.xyz, GBufferData, Settings);
+		#if STRATA_ENABLED
+					const FStrataTopLayerData TopLayerData = StrataUnpackTopLayerData(Strata.TopLayerNormalTexture.Load(uint3(SVPos.xy, 0)));
+					const float3 WorldNormal = TopLayerData.WorldNormal;
+					FTransmissionProfileParams TransmissionParams = GetTransmissionProfileParams(ExtractSubsurfaceProfileInt(SSSData.ProfilePackedId));
+		#else // STRATA_ENABLED
+					float3 WorldNormal = GBufferData.WorldNormal;
+					FTransmissionProfileParams TransmissionParams = GetTransmissionProfileParams(GBufferData);
+		#endif // STRATA_ENABLED
+					SSSTransmission = CalcTransmissionThickness(ScreenPosition.xyz, LightPositionAndInvRadius.xyz, WorldNormal, TransmissionParams, Settings);
 				}
 				else
 			  #endif
@@ -369,6 +396,12 @@ float3 PointLightDepthBias;
 // xy: depth projection parameters
 float2 PointLightProjParameters;
 
+#if STRATA_ENABLED
+Texture2D<uint> StrataClassificationTexture;
+Texture2D<uint> StrataTopLayerNormalTexture;
+Texture2D<uint2> StrataSSSTexture;
+#endif
+
 /** Pixel shader for projecting a one pass point light shadow from a cube map. */
 void MainOnePassPointLightPS(
 	in float4 SVPos : SV_POSITION,
@@ -411,10 +444,17 @@ void MainOnePassPointLightPS(
 	float SlopeBias = 0;
 #if !SUBPIXEL_SHADOW
 	{
+#if STRATA_ENABLED
+		const FStrataTopLayerData TopLayerData = StrataUnpackTopLayerData(StrataTopLayerNormalTexture.Load(uint3(SVPos.xy, 0)));
+		const float3 WorldNormal = TopLayerData.WorldNormal;
+#else // STRATA_ENABLED
 		FGBufferData GBufferData = GetGBufferData(ScreenUV);
+		const float3 WorldNormal = GBufferData.WorldNormal;
+#endif // STRATA_ENABLED
+
 		const float3 LightDirection = normalize(WorldSampleToLightVec);
-		const float NoL = saturate(dot(GBufferData.WorldNormal, LightDirection));
-		SlopeBias = (1-NoL) * PointLightDepthBias.y;
+		const float NoL = saturate(dot(WorldNormal, LightDirection));
+		SlopeBias = (1 - NoL) * PointLightDepthBias.y;
 	}
 #endif
 
@@ -440,13 +480,32 @@ void MainOnePassPointLightPS(
 
 #if USE_TRANSMISSION
 	float Thickness = 1.0f;
+
+#if STRATA_ENABLED
+
+	const FStrataClassification StrataData = StrataUnpackClassificationData(StrataClassificationTexture.Load(uint3(SVPos.xy, 0)));
+	const FStrataTopLayerData TopLayerData = StrataUnpackTopLayerData(StrataTopLayerNormalTexture.Load(uint3(SVPos.xy, 0)));
+	const FStrataSubsurfaceData SSSData = StrataUnpackSSSData(StrataSSSTexture.Load(uint3(SVPos.xy, 0)));
+
+	const bool IsSubsurfaceProfile = StrataData.bHasSubsurfaceProfile;
+	const float3 WorldNormal = TopLayerData.WorldNormal;
+	FTransmissionProfileParams TransmissionParams = GetTransmissionProfileParams(ExtractSubsurfaceProfileInt(SSSData.ProfilePackedId));
+
+#else // STRATA_ENABLED
+
 	// TODO: Use existing GBuffer data, instead of overwriting?
 	FGBufferData GBufferData = GetGBufferData(ScreenUV);
-	//if bSubsurface get SSS shadow,else opaque shadow.
+
 	bool IsSubsurfaceProfile = GBufferData.ShadingModelID == SHADINGMODELID_SUBSURFACE_PROFILE;
+	float3 WorldNormal = GBufferData.WorldNormal;
+	FTransmissionProfileParams TransmissionParams = GetTransmissionProfileParams(GBufferData);
+
+#endif // STRATA_ENABLED
+
+	//if bSubsurface get SSS shadow,else opaque shadow.
 	if (IsSubsurfaceProfile)
 	{
-		Thickness = CalcTransmissionThickness(GBufferData, WorldSampleToLightVec, PointLightProjParameters, LightPositionAndInvRadius.w, 0);
+		Thickness = CalcTransmissionThickness(WorldNormal, TransmissionParams, WorldSampleToLightVec, PointLightProjParameters, LightPositionAndInvRadius.w, 0);
 	}
 	OutColor.a = IsSubsurfaceProfile ? EncodeLightAttenuation(Thickness) : OutColor.b;
 #endif // USE_TRANSMISSION
diff --git a/Engine/Shaders/Private/Strata/Strata.ush b/Engine/Shaders/Private/Strata/Strata.ush
--- a/Engine/Shaders/Private/Strata/Strata.ush
+++ b/Engine/Shaders/Private/Strata/Strata.ush
@@ -217,6 +217,8 @@ struct FStrataBSDF
 
 void StrataRequestSharedNormalTangentBasis(inout uint NormalTypes, uint Index);
 
+uint   PackRGBA8(float4 rgba);
+float4 UnpackRGBA8(uint rgba);
 uint   PackR11G11B10F(float3 rgb);
 float3 UnpackR11G11B10F(uint rgb);
 uint   PackColorLinearToGamma2AlphaLinear(float4 rgba);
@@ -232,7 +234,6 @@ struct FStrataClassification
 	bool bIsStrataMaterial;
 	bool bHasSubsurface;
 	bool bHasSubsurfaceProfile;
-	bool bHasTransmission;
 	bool bIsSimple;
 };
 
@@ -247,8 +248,7 @@ uint StrataPackClassificationData(FStrataClassification In)
 		((In.bIsStrataMaterial		? 1 : 0) << 24 ) |
 		((In.bHasSubsurface			? 1 : 0) << 25 ) | 
 		((In.bHasSubsurfaceProfile	? 1 : 0) << 26 ) |
-		((In.bHasTransmission		? 1 : 0) << 27 ) |
-		((In.bIsSimple				? 1 : 0) << 28 );
+		((In.bIsSimple				? 1 : 0) << 27 );
 }
 
 FStrataClassification StrataUnpackClassificationData(uint In)
@@ -262,8 +262,7 @@ FStrataClassification StrataUnpackClassificationData(uint In)
 	Out.bIsStrataMaterial		= (In >> 24) & 0x1;
 	Out.bHasSubsurface			= (In >> 25) & 0x1;
 	Out.bHasSubsurfaceProfile	= (In >> 26) & 0x1;
-	Out.bHasTransmission		= (In >> 27) & 0x1;
-	Out.bIsSimple				= (In >> 28) & 0x1;
+	Out.bIsSimple				= (In >> 27) & 0x1;
 
 	Out.MaterialAO = Square(float(MaterialAO8Bits) * (1.0f / 255.0f));
 	Out.TopLayerRoughness = float(TopLayerRoughness8Bits) * (1.0f / 255.0f);
@@ -275,29 +274,56 @@ struct FStrataTopLayerData
 	float3 WorldNormal;
 };
 
-struct FStrataSubsurfaceData
+uint StrataPackTopLayerData(FStrataTopLayerData In)
 {
-	bool   bHasSubsurface;
-	bool   bHasProfile;
+	uint Out = StrataPackNormal(In.WorldNormal);
+	return Out;
+}
 
+FStrataTopLayerData StrataUnpackTopLayerData(uint In)
+{
+	FStrataTopLayerData Out = (FStrataTopLayerData)0;
+	Out.WorldNormal = StrataUnpackNormal(In);
+	return Out;
+}
+
+struct FStrataSubsurfaceData
+{
 	float3 BaseColor;
 
+	bool   bIsProfile;
 	float  ProfilePackedId;		// Integer ID packed in 8bits as float
 	float  ProfileRadiusScale;
 
 	float3 DMFP;				// When the profile is not used, a per pixel DMFP is used
 };
 
-uint StrataPackTopLayerData(FStrataTopLayerData In)
+uint2 StrataPackSSSData(FStrataSubsurfaceData In)
 {
-	uint Out = StrataPackNormal(In.WorldNormal);
+	uint2 Out;
+	Out.x = PackColorLinearToGamma2AlphaLinear(float4(In.BaseColor, In.bIsProfile ? 1.0f : 0.0f));
+	Out.y = In.bIsProfile ? PackRGBA8(float4(In.ProfilePackedId, In.ProfileRadiusScale, 0.0f, 0.0f)) : PackR11G11B10F(In.DMFP);
 	return Out;
 }
 
-FStrataTopLayerData StrataUnpackTopLayerData(uint In)
+FStrataSubsurfaceData StrataUnpackSSSData(uint2 In)
 {
-	FStrataTopLayerData Out = (FStrataTopLayerData)0;
-	Out.WorldNormal = StrataUnpackNormal(In);
+	FStrataSubsurfaceData Out = (FStrataSubsurfaceData)0;
+
+	float4 Data0 = UnpackColorGamma2ToLinearAlphaLinear(In.x);
+	Out.BaseColor = Data0.rgb;
+	Out.bIsProfile= Data0.a > 0.0f;
+
+	if (Out.bIsProfile)
+	{
+		float4 Data1 = UnpackRGBA8(In.y);
+		Out.ProfilePackedId = Data1.x;
+		Out.ProfileRadiusScale = Data1.y;
+	}
+	else
+	{
+		Out.DMFP = PackR11G11B10F(In.y);
+	}
 	return Out;
 }
 
diff --git a/Engine/Shaders/Private/Strata/StrataMaterialClassification.usf b/Engine/Shaders/Private/Strata/StrataMaterialClassification.usf
--- a/Engine/Shaders/Private/Strata/StrataMaterialClassification.usf
+++ b/Engine/Shaders/Private/Strata/StrataMaterialClassification.usf
@@ -28,8 +28,8 @@ void ComputePostProcessBSDFData(
 	const float FullThroughput = 1.0f;
 	const float FullyRough = 1.0f;
 
-	const float3 BSDFCoverage = BSDF.Weight;
-	const float  BSDFCoverageAvg = dot(BSDF.Weight,1.0f/3.0f);
+	const float3 BSDFWeight = BSDF.Weight;
+	const float  BSDFWeightAvg = dot(BSDF.Weight,1.0f/3.0f);
 
 	// STRATA_TODO apply the same normal transform as for the material
 	float3x3 TangentBasis = StrataGetBSDFSharedBasis(Header, BSDF, StrataAddressing);
@@ -57,20 +57,28 @@ void ComputePostProcessBSDFData(
 		if (bTopLayer)
 		{
 			// Setup SSR for the top layer
-			TopLayerRoughness += BSDFCoverageAvg * SLAB_ROUGHNESSX(BSDF);
-			TopLayerRoughnessAverageCount += BSDFCoverageAvg;
+			TopLayerRoughness += BSDFWeightAvg * SLAB_ROUGHNESSX(BSDF);
+			TopLayerRoughnessAverageCount += BSDFWeightAvg;
 		}
 
 		Classification.bIsSimple = Classification.bIsSimple && IsStrataFastPathCompatible(BSDF);
 
-		// Select the first/toppest layer with subsurface scattering
-		if (BSDF_GETHASSCATTERING(BSDF) && BSDFCoverageAvg > 0 && !SSS.bHasSubsurface)
+		// Select the latest encoutered SSS BSDF
+		if (BSDF_GETHASSCATTERING(BSDF) && BSDFWeightAvg > 0)
 		{
-			SSS.bHasSubsurface = true;
-			SSS.bHasProfile			= BSDF_GETHASSSSPROFILE(BSDF);
+			Classification.bHasSubsurface = true;
+
+			SSS.BaseColor = SLAB_BASECOLOR(BSDF);
+
+			// STRATA_TODO when layered translucency is figured out once again, only bottom layer should run SSS
+
+			// Profile ID: we only keep the latest ID, and profile will overtake any non-profile based SSS
+			SSS.bIsProfile			= SSS.bIsProfile || BSDF_GETHASSSSPROFILE(BSDF);
 			SSS.ProfilePackedId		= SLAB_SSSPROFILEID(BSDF);
 			SSS.ProfileRadiusScale	= SLAB_SSSPROFILERADIUSSCALE(BSDF);
-			SSS.BaseColor			= SLAB_BASECOLOR(BSDF);
+
+			// STRATA_TODO This should be accumulated according to the BSDF weight when we run this in the base pass.
+			// We cannot do that now because BSDFWeight contains weight*coverage*trhougput all mixed together.
 			SSS.DMFP				= SLAB_SSSDMFP(BSDF);
 		}
 		break;
@@ -91,13 +99,11 @@ void ComputePostProcessBSDFData(
 
 	case STRATA_BSDF_TYPE_HAIR:
 	{
-		Classification.bHasTransmission = true;
-
 		// No SSR
 		if (bTopLayer)
 		{
-			TopLayerRoughness += BSDFCoverageAvg * FullyRough; // Do not trigger SSR by using fully rough
-			TopLayerRoughnessAverageCount += BSDFCoverageAvg;
+			TopLayerRoughness += BSDFWeightAvg * FullyRough; // Do not trigger SSR by using fully rough
+			TopLayerRoughnessAverageCount += BSDFWeightAvg;
 		}
 		break;
 	}
@@ -108,8 +114,8 @@ void ComputePostProcessBSDFData(
 		{
 			const float Roughness = SLW_ROUGHNESS(BSDF);
 
-			TopLayerRoughness += BSDFCoverageAvg * Roughness;
-			TopLayerRoughnessAverageCount += BSDFCoverageAvg;
+			TopLayerRoughness += BSDFWeightAvg * Roughness;
+			TopLayerRoughnessAverageCount += BSDFWeightAvg;
 		}
 		break;
 	}
@@ -119,8 +125,9 @@ void ComputePostProcessBSDFData(
 
 void MainPS(
 	float4 SVPos : SV_POSITION, 
-	out uint OutColor0 : SV_Target0,
-	out uint OutColor1 : SV_Target1)
+	out uint  OutColor0 : SV_Target0,
+	out uint  OutColor1 : SV_Target1,
+	out uint2 OutColor2 : SV_Target2)
 {
 	const uint2 PixelPos = uint2(SVPos.xy);
 	
@@ -162,6 +169,7 @@ void MainPS(
 
 	OutColor0 = StrataPackClassificationData(ClassificationData);
 	OutColor1 = StrataPackTopLayerData(TopLayerData);
+	OutColor2 = StrataPackSSSData(SSSData);
 }
 #endif // SHADER_CATEGORIZATION
 
diff --git a/Engine/Shaders/Private/Strata/StrataVisualize.usf b/Engine/Shaders/Private/Strata/StrataVisualize.usf
--- a/Engine/Shaders/Private/Strata/StrataVisualize.usf
+++ b/Engine/Shaders/Private/Strata/StrataVisualize.usf
@@ -158,6 +158,7 @@ void VisualizeMaterialPS(
 		uint2 RectMax = RectMin;
 
 		const FStrataClassification StrataData = StrataUnpackClassificationData(Strata.ClassificationTexture.Load(uint3(MaterialPos, 0)));
+		const FStrataSubsurfaceData SSSData = StrataUnpackSSSData(Strata.SSSTexture.Load(uint3(MaterialPos, 0)));
 
 		FStrataAddressing StrataAddressing = GetStrataPixelDataByteOffset(MaterialPos, uint2(View.BufferSizeAndInvSize.xy), Strata.MaxBytesPerPixel);
 		const uint PixelStrataDataByteOffsetStart = StrataAddressing.ReadBytes;
@@ -194,7 +195,13 @@ void VisualizeMaterialPS(
 #define PRINTMATERIALSKIPPRECSHADOWMASK()	PRINTBSDFCHAR(_S_); PRINTBSDFCHAR(_K_); PRINTBSDFCHAR(_I_); PRINTBSDFCHAR(_P_); PRINTBSDFCHAR(_P_); PRINTBSDFCHAR(_R_); PRINTBSDFCHAR(_E_); PRINTBSDFCHAR(_S_); PRINTBSDFCHAR(_H_); PRINTBSDFCHAR(_A_); PRINTBSDFCHAR(_D_); PRINTBSDFCHAR(_SPC_); PRINTBSDFCHAR(_0_+(Header.SkipPrecShadowMask?1:0));
 #define PRINTMATERIALZEROPRECSHADOWMASK()	PRINTBSDFCHAR(_Z_); PRINTBSDFCHAR(_E_); PRINTBSDFCHAR(_R_); PRINTBSDFCHAR(_O_); PRINTBSDFCHAR(_P_); PRINTBSDFCHAR(_R_); PRINTBSDFCHAR(_E_); PRINTBSDFCHAR(_S_); PRINTBSDFCHAR(_H_); PRINTBSDFCHAR(_A_); PRINTBSDFCHAR(_D_); PRINTBSDFCHAR(_SPC_); PRINTBSDFCHAR(_0_+(Header.ZeroPrecShadowMask?1:0));
 #define PRINTMATERIALSLW()	PRINTBSDFCHAR(_S_); PRINTBSDFCHAR(_L_); PRINTBSDFCHAR(_W_); PRINTBSDFCHAR(_SPC_); PRINTBSDFCHAR(_0_+(Header.IsSingleLayerWater?1:0));
-#define PRINTMATERIALROUGH()PRINTBSDFCHAR(_R_); PRINTBSDFCHAR(_O_); PRINTBSDFCHAR(_U_); PRINTBSDFCHAR(_G_); PRINTBSDFCHAR(_H_); PRINTBSDFCHAR(_SPC_); PRINTFLOAT1(StrataData.TopLayerRoughness);
+
+#define PRINTCLASSIFROUGH()		PRINTBSDFCHAR(_C_); PRINTBSDFCHAR(_L_); PRINTBSDFCHAR(_A_); PRINTBSDFCHAR(_S_); PRINTBSDFCHAR(_DOT_); PRINTBSDFCHAR(_R_) PRINTBSDFCHAR(_O_) PRINTBSDFCHAR(_U_) PRINTBSDFCHAR(_G_) PRINTBSDFCHAR(_H_) PRINTBSDFCHAR(_SPC_);	PRINTBSDFCHAR(_SPC_); PRINTFLOAT1(StrataData.TopLayerRoughness);
+#define PRINTCLASSIFBASCOL()	PRINTBSDFCHAR(_C_); PRINTBSDFCHAR(_L_); PRINTBSDFCHAR(_A_); PRINTBSDFCHAR(_S_); PRINTBSDFCHAR(_DOT_); PRINTBSDFCHAR(_B_) PRINTBSDFCHAR(_A_) PRINTBSDFCHAR(_S_) PRINTBSDFCHAR(_C_) PRINTBSDFCHAR(_O_) PRINTBSDFCHAR(_L_);	PRINTBSDFCHAR(_SPC_); PRINTFLOAT3(SSSData.BaseColor);
+#define PRINTCLASSIFISPROF()	PRINTBSDFCHAR(_C_); PRINTBSDFCHAR(_L_); PRINTBSDFCHAR(_A_); PRINTBSDFCHAR(_S_); PRINTBSDFCHAR(_DOT_); PRINTBSDFCHAR(_I_) PRINTBSDFCHAR(_S_) PRINTBSDFCHAR(_P_) PRINTBSDFCHAR(_R_) PRINTBSDFCHAR(_O_) PRINTBSDFCHAR(_F_);	PRINTBSDFCHAR(_SPC_); PRINTBSDFCHAR(SSSData.bIsProfile?_1_:_0_);
+#define PRINTCLASSIFPROID()		PRINTBSDFCHAR(_C_); PRINTBSDFCHAR(_L_); PRINTBSDFCHAR(_A_); PRINTBSDFCHAR(_S_); PRINTBSDFCHAR(_DOT_); PRINTBSDFCHAR(_P_) PRINTBSDFCHAR(_R_) PRINTBSDFCHAR(_O_) PRINTBSDFCHAR(_I_) PRINTBSDFCHAR(_D_) PRINTBSDFCHAR(_SPC_);	PRINTBSDFCHAR(_SPC_); PRINTFLOATNOFRAC(ExtractSubsurfaceProfileInt(SSSData.ProfilePackedId), 3);
+#define PRINTCLASSIFPROSCL()	PRINTBSDFCHAR(_C_); PRINTBSDFCHAR(_L_); PRINTBSDFCHAR(_A_); PRINTBSDFCHAR(_S_); PRINTBSDFCHAR(_DOT_); PRINTBSDFCHAR(_P_) PRINTBSDFCHAR(_R_) PRINTBSDFCHAR(_O_) PRINTBSDFCHAR(_S_) PRINTBSDFCHAR(_C_) PRINTBSDFCHAR(_L_);	PRINTBSDFCHAR(_SPC_); PRINTFLOAT1(SSSData.ProfileRadiusScale);
+#define PRINTCLASSIFSSDMFP()	PRINTBSDFCHAR(_C_); PRINTBSDFCHAR(_L_); PRINTBSDFCHAR(_A_); PRINTBSDFCHAR(_S_); PRINTBSDFCHAR(_DOT_); PRINTBSDFCHAR(_S_) PRINTBSDFCHAR(_S_) PRINTBSDFCHAR(_D_) PRINTBSDFCHAR(_M_) PRINTBSDFCHAR(_F_) PRINTBSDFCHAR(_P_);	PRINTBSDFCHAR(_SPC_); PRINTFLOAT3(SSSData.DMFP);
 
 #define NEXTLINE  TopLeft.y += 8; RectMax.x = max(RectMax.x, TopLeft.x); RectMax.y = max(RectMax.y, TopLeft.y); TopLeft.x = TopLeftInit.x; 
 #define HORIZONTALSEPARATION  TopLeft.y += 4;
@@ -208,14 +215,41 @@ void VisualizeMaterialPS(
 		NEXTLINE;
 		PRINTMATERIALAO();
 		NEXTLINE;
+		PRINTMATERIALSLW();
+		NEXTLINE;
 		PRINTMATERIALSKIPPRECSHADOWMASK();
 		NEXTLINE;
 		PRINTMATERIALZEROPRECSHADOWMASK();
 		NEXTLINE;
-		PRINTMATERIALSLW();
-		NEXTLINE;
-		PRINTMATERIALROUGH();
+
+		PRINTCLASSIFROUGH();
 		NEXTLINE;
+		if (StrataData.bHasSubsurface)
+		{
+			PRINTCLASSIFBASCOL();
+			NEXTLINE;
+			PRINTCLASSIFISPROF();
+			NEXTLINE;
+			if (SSSData.bIsProfile)
+			{
+				PRINTCLASSIFPROID();
+				NEXTLINE;
+				PRINTCLASSIFPROSCL();
+				NEXTLINE;
+			}
+			else
+			{
+				PRINTCLASSIFSSDMFP();
+				NEXTLINE;
+				NEXTLINE;
+			}
+		}
+		else
+		{
+			PRINTBSDFCHAR(_N_); PRINTBSDFCHAR(_O_); PRINTBSDFCHAR(_S_); PRINTBSDFCHAR(_S_); PRINTBSDFCHAR(_S_);
+			NEXTLINE;
+		}
+
 		HORIZONTALSEPARATION;
 #else
 		NEXTLINE;
@@ -225,7 +259,28 @@ void VisualizeMaterialPS(
 		NEXTLINE;
 		NEXTLINE;
 		NEXTLINE;
+
 		NEXTLINE;
+		if (StrataData.bHasSubsurface)
+		{
+			NEXTLINE;
+			NEXTLINE;
+			if (SSSData.bIsProfile)
+			{
+				NEXTLINE;
+				NEXTLINE;
+			}
+			else
+			{
+				NEXTLINE;
+				NEXTLINE;
+			}
+		}
+		else
+		{
+			NEXTLINE;
+		}
+
 		HORIZONTALSEPARATION;
 #endif
 
diff --git a/Engine/Shaders/Private/TransmissionThickness.ush b/Engine/Shaders/Private/TransmissionThickness.ush
--- a/Engine/Shaders/Private/TransmissionThickness.ush
+++ b/Engine/Shaders/Private/TransmissionThickness.ush
@@ -113,7 +113,7 @@ float CalculateThickness(float ShadowmapDepth, FPCFSamplerSettings Settings, flo
 	return Thickness;
 }
 
-float CalcTransmissionThickness(float3 ViewPos, float3 LightPos, FGBufferData GBufferData, FPCFSamplerSettings Settings)
+float CalcTransmissionThickness(float3 ViewPos, float3 LightPos, float3 WorldNormal, FTransmissionProfileParams TransmissionParams, FPCFSamplerSettings Settings)
 {
 	///Get -NoL, maybe unnecessary.
 	float4 HomogeneousWorldPosition = mul(float4(ViewPos.xyz, 1), View.ScreenToWorld);
@@ -121,12 +121,10 @@ float CalcTransmissionThickness(float3 ViewPos, float3 LightPos, FGBufferData GB
 	float3 ToLight = LightPos.xyz - WorldPosition;
 	float DistanceSqr = dot(ToLight, ToLight);
 	float3 L = ToLight * rsqrt(DistanceSqr);
-	float NoL = 1;// dot(-L, GBufferData.WorldNormal);
-
-	FTransmissionProfileParams TransmissionParams = GetTransmissionProfileParams(GBufferData);
+	float NoL = 1;// dot(-L, WorldNormal);
 
 	//Shrink Pos in View Space.
-	float3 ViewNormal = mul(GBufferData.WorldNormal, (float3x3)View.TranslatedWorldToView);
+	float3 ViewNormal = mul(WorldNormal, (float3x3)View.TranslatedWorldToView);
 
 	//Shrink world position with neg worldNormal to get thickness
 	//These scalar const is based on adjust expericence.
@@ -167,17 +165,15 @@ float CalcTransmissionThickness(float3 ViewPos, float3 LightPos, FGBufferData GB
 
 
 /** Computes thickness for a given world position from a cubemap shadowmap used on a point light. */
-float CalcTransmissionThickness(FGBufferData GBufferData, float3 WorldSampleToLightVec, float2 ProjParam, float LightInvRadius, float DepthBias)
+float CalcTransmissionThickness(float3 WorldNormal, FTransmissionProfileParams TransmissionParams, float3 WorldSampleToLightVec, float2 ProjParam, float LightInvRadius, float DepthBias)
 {
 	float Thickness = 1;
 
-	FTransmissionProfileParams TransmissionParams = GetTransmissionProfileParams(GBufferData);
-
 	float DepthScale = TransmissionParams.ExtinctionScale * 3.1f;
 	float NormalScale = TransmissionParams.NormalScale * 0.5;
 
 	//Shrink pos aligh -WorldNormal
-	float3 Normal = GBufferData.WorldNormal;
+	float3 Normal = WorldNormal;
 	float3 LightRelativeBiasedPosition = -WorldSampleToLightVec - Normal * NormalScale;
 	float DistanceSqr = dot(WorldSampleToLightVec, WorldSampleToLightVec);
 	float3 L = WorldSampleToLightVec * rsqrt(DistanceSqr);
diff --git a/Engine/Source/Runtime/Renderer/Private/ShadowRendering.h b/Engine/Source/Runtime/Renderer/Private/ShadowRendering.h
--- a/Engine/Source/Runtime/Renderer/Private/ShadowRendering.h
+++ b/Engine/Source/Runtime/Renderer/Private/ShadowRendering.h
@@ -1473,6 +1473,9 @@ public:
 		PointLightProjParameters.Bind(Initializer.ParameterMap, TEXT("PointLightProjParameters"));
 		TransmissionProfilesTexture.Bind(Initializer.ParameterMap, TEXT("SSProfilesTexture"));
 		HairCategorizationTexture.Bind(Initializer.ParameterMap, TEXT("HairCategorizationTexture"));
+		ClassificationTexture.Bind(Initializer.ParameterMap, TEXT("StrataClassificationTexture"));
+		TopLayerNormalTexture.Bind(Initializer.ParameterMap, TEXT("StrataTopLayerNormalTexture"));
+		SSSTexture.Bind(Initializer.ParameterMap, TEXT("StrataSSSTexture"));
 	}
 
 	static void ModifyCompilationEnvironment(const FGlobalShaderPermutationParameters& Parameters, FShaderCompilerEnvironment& OutEnvironment)
@@ -1492,7 +1495,7 @@ public:
 	void SetParameters(
 		FRHICommandList& RHICmdList, 
 		int32 ViewIndex,
-		const FSceneView& View,
+		const FViewInfo& View,
 		const FHairStrandsVisibilityData* HairVisibilityData,
 		const FProjectedShadowInfo* ShadowInfo)
 	{
@@ -1519,6 +1522,13 @@ public:
 			SetTextureParameter(RHICmdList, ShaderRHI, HairCategorizationTexture, HairVisibilityData->CategorizationTexture->GetPooledRenderTarget()->GetRenderTargetItem().ShaderResourceTexture);
 		}
 
+		if (Strata::IsStrataEnabled())
+		{
+			SetTextureParameter(RHICmdList, ShaderRHI, ClassificationTexture, Strata::GetClassificationTexture(View));
+			SetTextureParameter(RHICmdList, ShaderRHI, TopLayerNormalTexture, Strata::GetTopLayerNormalTexture(View));
+			SetTextureParameter(RHICmdList, ShaderRHI, SSSTexture, Strata::GetSSSTexture(View));
+		}
+
 		{
 			const IPooledRenderTarget* PooledRT = GetSubsufaceProfileTexture_RT((FRHICommandListImmediate&)RHICmdList);
 
@@ -1560,6 +1570,9 @@ private:
 	LAYOUT_FIELD(FShaderParameter, PointLightProjParameters);
 	LAYOUT_FIELD(FShaderResourceParameter, TransmissionProfilesTexture);
 	LAYOUT_FIELD(FShaderResourceParameter, HairCategorizationTexture);
+	LAYOUT_FIELD(FShaderResourceParameter, ClassificationTexture);
+	LAYOUT_FIELD(FShaderResourceParameter, TopLayerNormalTexture);
+	LAYOUT_FIELD(FShaderResourceParameter, SSSTexture);
 };
 
 // Reversed Z
diff --git a/Engine/Source/Runtime/Renderer/Private/Strata/Strata.cpp b/Engine/Source/Runtime/Renderer/Private/Strata/Strata.cpp
--- a/Engine/Source/Runtime/Renderer/Private/Strata/Strata.cpp
+++ b/Engine/Source/Runtime/Renderer/Private/Strata/Strata.cpp
@@ -143,11 +143,18 @@ void InitialiseStrataFrameSceneData(FSceneRenderer& SceneRenderer, FRDGBuilder&
 
 		// Top layer texture
 		{
-			FRDGTextureRef Texture = GraphBuilder.CreateTexture(FRDGTextureDesc::Create2D(BufferSizeXY, PF_R32G32_UINT, FClearValueBinding::Black, TexCreate_ShaderResource | TexCreate_RenderTargetable), TEXT("StrataTopLayerNormalTexture"));
+			FRDGTextureRef Texture = GraphBuilder.CreateTexture(FRDGTextureDesc::Create2D(BufferSizeXY, PF_R32_UINT, FClearValueBinding::Black, TexCreate_ShaderResource | TexCreate_RenderTargetable), TEXT("StrataTopLayerNormalTexture"));
 			GraphBuilder.PreallocateTexture(Texture);
 			StrataSceneData.TopLayerNormalTexture = GraphBuilder.GetPooledTexture(Texture);
 		}
 
+		// SSS texture
+		{
+			FRDGTextureRef Texture = GraphBuilder.CreateTexture(FRDGTextureDesc::Create2D(BufferSizeXY, PF_R32G32_UINT, FClearValueBinding::Black, TexCreate_ShaderResource | TexCreate_RenderTargetable), TEXT("StrataSSSTexture"));
+			GraphBuilder.PreallocateTexture(Texture);
+			StrataSceneData.SSSTexture = GraphBuilder.GetPooledTexture(Texture);
+		}
+
 		// Energy LUT
 		const uint32 LUTResolution = GetStrataGGXEnergyLUTResolution();
 		bUpdateLUT = StrataSceneData.GGXEnergyLUT2DTexture == nullptr || StrataSceneData.GGXEnergyLUT2DTexture->GetDesc().Extent.X != LUTResolution || CVarStrataLUTContinousUpdate.GetValueOnAnyThread() > 0;
@@ -243,6 +250,7 @@ TUniformBufferRef<FStrataGlobalUniformParameters> BindStrataGlobalUniformParamet
 		StrataUniformParameters.GGXEnergyLUT2DTexture = View.StrataSceneData->GGXEnergyLUT2DTexture->GetRenderTargetItem().ShaderResourceTexture;
 		StrataUniformParameters.GGXEnergyLUTSampler = TStaticSamplerState<SF_Bilinear, AM_Clamp, AM_Clamp, AM_Clamp>::GetRHI();
 		StrataUniformParameters.GGXEnergyLUTScaleBias = GetStrataGGXEnergyLUTScaleBias();
+		StrataUniformParameters.SSSTexture = View.StrataSceneData->SSSTexture->GetRenderTargetItem().ShaderResourceTexture;
 		View.StrataSceneData->StrataGlobalUniformParameters = CreateUniformBufferImmediate(StrataUniformParameters, UniformBuffer_SingleFrame);
 		return View.StrataSceneData->StrataGlobalUniformParameters;
 	}
@@ -251,8 +259,9 @@ TUniformBufferRef<FStrataGlobalUniformParameters> BindStrataGlobalUniformParamet
 		// Create each time. This path will go away when Strata is always enabled anyway.
 		StrataUniformParameters.MaxBytesPerPixel = 0;
 		StrataUniformParameters.MaterialLobesBuffer = GEmptyVertexBufferWithUAV->ShaderResourceViewRHI;
-		StrataUniformParameters.ClassificationTexture = GSystemTextures.BlackDummy->GetRenderTargetItem().ShaderResourceTexture;
-		StrataUniformParameters.TopLayerNormalTexture = GSystemTextures.BlackDummy->GetRenderTargetItem().ShaderResourceTexture;
+		StrataUniformParameters.ClassificationTexture = GSystemTextures.ZeroUIntDummy->GetRenderTargetItem().ShaderResourceTexture;
+		StrataUniformParameters.TopLayerNormalTexture = GSystemTextures.ZeroUIntDummy->GetRenderTargetItem().ShaderResourceTexture;
+		StrataUniformParameters.SSSTexture = GSystemTextures.ZeroUIntDummy->GetRenderTargetItem().ShaderResourceTexture;
 		StrataUniformParameters.GGXEnergyLUT3DTexture = GSystemTextures.BlackDummy->GetRenderTargetItem().ShaderResourceTexture;
 		StrataUniformParameters.GGXEnergyLUT2DTexture = GSystemTextures.BlackDummy->GetRenderTargetItem().ShaderResourceTexture;
 		StrataUniformParameters.GGXEnergyLUTSampler = TStaticSamplerState<SF_Bilinear, AM_Clamp, AM_Clamp, AM_Clamp>::GetRHI();
@@ -261,6 +270,31 @@ TUniformBufferRef<FStrataGlobalUniformParameters> BindStrataGlobalUniformParamet
 	}
 }
 
+FTextureRHIRef GetClassificationTexture(const FViewInfo& View)
+{
+	if (View.StrataSceneData)
+	{
+		return View.StrataSceneData->ClassificationTexture->GetRenderTargetItem().ShaderResourceTexture;
+	}
+	return GSystemTextures.ZeroUIntDummy->GetRenderTargetItem().ShaderResourceTexture;
+}
+FTextureRHIRef GetTopLayerNormalTexture(const FViewInfo& View)
+{
+	if (View.StrataSceneData)
+	{
+		return View.StrataSceneData->TopLayerNormalTexture->GetRenderTargetItem().ShaderResourceTexture;
+	}
+	return GSystemTextures.ZeroUIntDummy->GetRenderTargetItem().ShaderResourceTexture;
+}
+FTextureRHIRef GetSSSTexture(const FViewInfo& View)
+{
+	if (View.StrataSceneData)
+	{
+		return View.StrataSceneData->SSSTexture->GetRenderTargetItem().ShaderResourceTexture;
+	}
+	return GSystemTextures.ZeroUIntDummy->GetRenderTargetItem().ShaderResourceTexture;
+}
+
 ////////////////////////////////////////////////////////////////////////// 
 // Debug
 
@@ -388,6 +422,7 @@ class FStrataMaterialClassificationPassPS : public FGlobalShader
 		OutEnvironment.SetDefine(TEXT("SHADER_CATEGORIZATION"), 1);
 		OutEnvironment.SetRenderTargetOutputFormat(0, PF_R32_UINT);
 		OutEnvironment.SetRenderTargetOutputFormat(1, PF_R32_UINT);
+		OutEnvironment.SetRenderTargetOutputFormat(2, PF_R32G32_UINT);
 	}
 };
 IMPLEMENT_GLOBAL_SHADER(FStrataMaterialClassificationPassPS, "/Engine/Private/Strata/StrataMaterialClassification.usf", "MainPS", SF_Pixel);
@@ -654,6 +689,7 @@ void AddStrataMaterialClassificationPass(FRDGBuilder& GraphBuilder, const FMinim
 			PassParameters->SceneTextures = GetSceneTextureParameters(GraphBuilder);		
 			PassParameters->RenderTargets[0] = FRenderTargetBinding(ClassificationTexture, ERenderTargetLoadAction::EClear);
 			PassParameters->RenderTargets[1] = FRenderTargetBinding(GraphBuilder.RegisterExternalTexture(View.StrataSceneData->TopLayerNormalTexture), ERenderTargetLoadAction::EClear);
+			PassParameters->RenderTargets[2] = FRenderTargetBinding(GraphBuilder.RegisterExternalTexture(View.StrataSceneData->SSSTexture), ERenderTargetLoadAction::EClear);
 
 			if (ShaderDrawDebug::IsShaderDrawDebugEnabled())
 			{
diff --git a/Engine/Source/Runtime/Renderer/Private/Strata/Strata.h b/Engine/Source/Runtime/Renderer/Private/Strata/Strata.h
--- a/Engine/Source/Runtime/Renderer/Private/Strata/Strata.h
+++ b/Engine/Source/Runtime/Renderer/Private/Strata/Strata.h
@@ -28,6 +28,7 @@ BEGIN_GLOBAL_SHADER_PARAMETER_STRUCT(FStrataGlobalUniformParameters, )
 	SHADER_PARAMETER_TEXTURE(Texture3D<float2>, GGXEnergyLUT3DTexture)
 	SHADER_PARAMETER_TEXTURE(Texture2D<float4>, GGXEnergyLUT2DTexture)
 	SHADER_PARAMETER_SAMPLER(SamplerState, GGXEnergyLUTSampler)
+	SHADER_PARAMETER_TEXTURE(Texture2D<uint2>, SSSTexture)
 END_GLOBAL_SHADER_PARAMETER_STRUCT()
 
 struct FStrataSceneData
@@ -38,6 +39,7 @@ struct FStrataSceneData
 	TRefCountPtr<IPooledRenderTarget> TopLayerNormalTexture;	// This should be a RDG resource
 	TRefCountPtr<IPooledRenderTarget> GGXEnergyLUT3DTexture;
 	TRefCountPtr<IPooledRenderTarget> GGXEnergyLUT2DTexture;
+	TRefCountPtr<IPooledRenderTarget> SSSTexture;				// This should be a RDG resource
 
 	TRefCountPtr<FRDGPooledBuffer> ClassificationTileListBuffer;
 	TRefCountPtr<FRDGPooledBuffer> ClassificationTileIndirectBuffer;
@@ -64,6 +66,10 @@ void BindStrataBasePassUniformParameters(const FViewInfo& View, FStrataOpaquePas
 
 TUniformBufferRef<FStrataGlobalUniformParameters> BindStrataGlobalUniformParameters(const FViewInfo& View);
 
+FTextureRHIRef GetClassificationTexture(const FViewInfo& View);
+FTextureRHIRef GetTopLayerNormalTexture(const FViewInfo& View);
+FTextureRHIRef GetSSSTexture(const FViewInfo& View);
+
 void AddStrataMaterialClassificationPass(FRDGBuilder& GraphBuilder, const FMinimalSceneTextures& SceneTextures, const TArray<FViewInfo>& Views);
 
 void AddStrataStencilPass(FRDGBuilder& GraphBuilder, const FViewInfo& View, const FMinimalSceneTextures& SceneTextures);
