Forgoten shader file that should have gone with 15406566

diff --git a/Engine/Shaders/Private/ShadowProjectionPixelShader.usf b/Engine/Shaders/Private/ShadowProjectionPixelShader.usf
--- a/Engine/Shaders/Private/ShadowProjectionPixelShader.usf
+++ b/Engine/Shaders/Private/ShadowProjectionPixelShader.usf
@@ -396,12 +396,6 @@ float3 PointLightDepthBias;
 // xy: depth projection parameters
 float2 PointLightProjParameters;
 
-#if STRATA_ENABLED
-Texture2D<uint> StrataClassificationTexture;
-Texture2D<uint> StrataTopLayerNormalTexture;
-Texture2D<uint2> StrataSSSTexture;
-#endif
-
 /** Pixel shader for projecting a one pass point light shadow from a cube map. */
 void MainOnePassPointLightPS(
 	in float4 SVPos : SV_POSITION,
@@ -445,7 +439,7 @@ void MainOnePassPointLightPS(
 #if !SUBPIXEL_SHADOW
 	{
 #if STRATA_ENABLED
-		const FStrataTopLayerData TopLayerData = StrataUnpackTopLayerData(StrataTopLayerNormalTexture.Load(uint3(SVPos.xy, 0)));
+		const FStrataTopLayerData TopLayerData = StrataUnpackTopLayerData(Strata.TopLayerNormalTexture.Load(uint3(SVPos.xy, 0)));
 		const float3 WorldNormal = TopLayerData.WorldNormal;
 #else // STRATA_ENABLED
 		FGBufferData GBufferData = GetGBufferData(ScreenUV);
@@ -483,9 +477,9 @@ void MainOnePassPointLightPS(
 
 #if STRATA_ENABLED
 
-	const FStrataClassification StrataData = StrataUnpackClassificationData(StrataClassificationTexture.Load(uint3(SVPos.xy, 0)));
-	const FStrataTopLayerData TopLayerData = StrataUnpackTopLayerData(StrataTopLayerNormalTexture.Load(uint3(SVPos.xy, 0)));
-	const FStrataSubsurfaceData SSSData = StrataUnpackSSSData(StrataSSSTexture.Load(uint3(SVPos.xy, 0)));
+	const FStrataClassification StrataData = StrataUnpackClassificationData(Strata.ClassificationTexture.Load(uint3(SVPos.xy, 0)));
+	const FStrataTopLayerData TopLayerData = StrataUnpackTopLayerData(Strata.TopLayerNormalTexture.Load(uint3(SVPos.xy, 0)));
+	const FStrataSubsurfaceData SSSData = StrataUnpackSSSData(Strata.SSSTexture.Load(uint3(SVPos.xy, 0)));
 
 	const bool IsSubsurfaceProfile = StrataData.bHasSubsurfaceProfile;
 	const float3 WorldNormal = TopLayerData.WorldNormal;
