Restore use of LightingChannelMask for the primitive in ray tracing material shaders when GPU-Scene instance culling is enabled.

diff --git a/Engine/Shaders/Private/RayTracing/RayTracingMaterialHitShaders.usf b/Engine/Shaders/Private/RayTracing/RayTracingMaterialHitShaders.usf
--- a/Engine/Shaders/Private/RayTracing/RayTracingMaterialHitShaders.usf
+++ b/Engine/Shaders/Private/RayTracing/RayTracingMaterialHitShaders.usf
@@ -366,9 +366,7 @@ RAY_TRACING_ENTRY_CLOSEST_HIT(MaterialCHS,
 	#endif
 
 	Payload.BlendingMode = MaterialBlendingMode;
-#if !defined(GPUCULL_TODO)
 	Payload.PrimitiveLightingChannelMask = GetPrimitiveData(MaterialParameters.PrimitiveId).LightingChannelMask;
-#endif // defined(GPUCULL_TODO)
 
 	Payload.HitT = RayTCurrent();
 	if (HitKind() == HIT_KIND_TRIANGLE_FRONT_FACE)
