Added new HALF_TYPE to toggle between half/float on mobile vs desktop

diff --git a/Engine/Shaders/Private/DepthOnlyPixelShader.usf b/Engine/Shaders/Private/DepthOnlyPixelShader.usf
--- a/Engine/Shaders/Private/DepthOnlyPixelShader.usf
+++ b/Engine/Shaders/Private/DepthOnlyPixelShader.usf
@@ -28,8 +28,8 @@ void Main(
 	OPTIONAL_OutDepthConservative,
 #endif
 #if OUTPUT_MOBILE_COLOR_VALUE
-	out half OutCustomDepth : SV_Target0,
-	out half OutCustomStencil : SV_Target1
+	out HALF_TYPE OutCustomDepth : SV_Target0,
+	out HALF_TYPE OutCustomStencil : SV_Target1
 #else
 	out float4 OutColor : SV_Target0
 #endif
diff --git a/Engine/Shaders/Private/DistortApplyScreenPS.usf b/Engine/Shaders/Private/DistortApplyScreenPS.usf
--- a/Engine/Shaders/Private/DistortApplyScreenPS.usf
+++ b/Engine/Shaders/Private/DistortApplyScreenPS.usf
@@ -88,7 +88,7 @@ SamplerState	DistortionAccumulateSampler;
 
 void Merge_Mobile(
 	in noperspective float4 TexCoord: TEXCOORD0,
-	out half4 OutColor : SV_Target0
+	out HALF4_TYPE OutColor : SV_Target0
 )
 {
 	half4 AccumDist = DistortionAccumulateTexture.Sample(DistortionAccumulateSampler, TexCoord.xy);
diff --git a/Engine/Shaders/Private/GeometryCacheVertexFactory.ush b/Engine/Shaders/Private/GeometryCacheVertexFactory.ush
--- a/Engine/Shaders/Private/GeometryCacheVertexFactory.ush
+++ b/Engine/Shaders/Private/GeometryCacheVertexFactory.ush
@@ -41,10 +41,10 @@ struct FVertexFactoryInput
 {
 	float4	Position		: ATTRIBUTE0;
 	// 0..1
-	half3	TangentX		: ATTRIBUTE1;
+	HALF3_TYPE	TangentX		: ATTRIBUTE1;
 	// 0..1
 	// TangentZ.w contains sign of tangent basis determinant
-	half4	TangentZ		: ATTRIBUTE2;	
+	HALF4_TYPE	TangentZ		: ATTRIBUTE2;	
 
 	// Per vertex color 
 	float4 Color            : ATTRIBUTE3;
diff --git a/Engine/Shaders/Private/GpuSkinVertexFactory.ush b/Engine/Shaders/Private/GpuSkinVertexFactory.ush
--- a/Engine/Shaders/Private/GpuSkinVertexFactory.ush
+++ b/Engine/Shaders/Private/GpuSkinVertexFactory.ush
@@ -63,10 +63,10 @@ struct FVertexFactoryInput
 {
 	float4	Position		: ATTRIBUTE0;
 	// 0..1
-	half3	TangentX		: ATTRIBUTE1;
+	HALF3_TYPE	TangentX		: ATTRIBUTE1;
 	// 0..1
 	// TangentZ.w contains sign of tangent basis determinant
-	half4	TangentZ		: ATTRIBUTE2;	
+	HALF4_TYPE	TangentZ		: ATTRIBUTE2;	
 
 #if GPUSKIN_UNLIMITED_BONE_INFLUENCE
 	uint	BlendOffsetCount	: ATTRIBUTE3;
diff --git a/Engine/Shaders/Private/LandscapeGrassWeight.usf b/Engine/Shaders/Private/LandscapeGrassWeight.usf
--- a/Engine/Shaders/Private/LandscapeGrassWeight.usf
+++ b/Engine/Shaders/Private/LandscapeGrassWeight.usf
@@ -53,7 +53,7 @@ int OutputPass;
 void PSMain( 
 	FLandscapeGrassWeightInterpolantsVSToPS Interpolants
 	OPTIONAL_IsFrontFace,
-	out half4 OutColor	: SV_Target0
+	out HALF4_TYPE OutColor	: SV_Target0
 	)
 {
 	ResolvedView = ResolveView();
diff --git a/Engine/Shaders/Private/LandscapePhysicalMaterial.usf b/Engine/Shaders/Private/LandscapePhysicalMaterial.usf
--- a/Engine/Shaders/Private/LandscapePhysicalMaterial.usf
+++ b/Engine/Shaders/Private/LandscapePhysicalMaterial.usf
@@ -50,7 +50,7 @@ float2 UpdateDominantMaterial(in float Index, in float Contribution, in float2 D
 void PSMain(
 	FLandscapeGrassWeightInterpolantsVSToPS Interpolants
 	OPTIONAL_IsFrontFace,
-	out half4 OutColor : SV_Target0
+	out HALF4_TYPE OutColor : SV_Target0
 )
 {
 	ResolvedView = ResolveView();
diff --git a/Engine/Shaders/Private/LocalVertexFactory.ush b/Engine/Shaders/Private/LocalVertexFactory.ush
--- a/Engine/Shaders/Private/LocalVertexFactory.ush
+++ b/Engine/Shaders/Private/LocalVertexFactory.ush
@@ -102,11 +102,11 @@ struct FVertexFactoryInput
 
 		float4	Color		: ATTRIBUTE3;
 	#else
-		half3	TangentX	: ATTRIBUTE1;
+		HALF3_TYPE	TangentX	: ATTRIBUTE1;
 		// TangentZ.w contains sign of tangent basis determinant
-		half4	TangentZ	: ATTRIBUTE2;
+		HALF4_TYPE	TangentZ	: ATTRIBUTE2;
 
-		half4	Color		: ATTRIBUTE3;
+		HALF4_TYPE	Color		: ATTRIBUTE3;
 	#endif
 #endif
 
diff --git a/Engine/Shaders/Private/MeshParticleVertexFactory.ush b/Engine/Shaders/Private/MeshParticleVertexFactory.ush
--- a/Engine/Shaders/Private/MeshParticleVertexFactory.ush
+++ b/Engine/Shaders/Private/MeshParticleVertexFactory.ush
@@ -38,10 +38,10 @@
 struct FVertexFactoryInput
 {
 	float4	Position	: ATTRIBUTE0;
-	half3	TangentX	: ATTRIBUTE1;
+	HALF3_TYPE	TangentX	: ATTRIBUTE1;
 	// TangentZ.w contains sign of tangent basis determinant
-	half4	TangentZ	: ATTRIBUTE2;
-	half4	VertexColor : ATTRIBUTE3;
+	HALF4_TYPE	TangentZ	: ATTRIBUTE2;
+	HALF4_TYPE	VertexColor : ATTRIBUTE3;
 
 #if NUM_MATERIAL_TEXCOORDS_VERTEX
 	float2	TexCoords[NUM_MATERIAL_TEXCOORDS_VERTEX] : ATTRIBUTE4;
diff --git a/Engine/Shaders/Private/MobileBasePassPixelShader.usf b/Engine/Shaders/Private/MobileBasePassPixelShader.usf
--- a/Engine/Shaders/Private/MobileBasePassPixelShader.usf
+++ b/Engine/Shaders/Private/MobileBasePassPixelShader.usf
@@ -533,13 +533,13 @@ void Main(
 	OPTIONAL_IsFrontFace
 #if DEFERRED_SHADING_PATH
 #if !USE_GLES_FBF_DEFERRED
-	, out half4 OutColor	: SV_Target0
-	, out half4 OutGBufferA	: SV_Target1
-	, out half4 OutGBufferB	: SV_Target2
-	, out half4 OutGBufferC	: SV_Target3  
+	, out HALF4_TYPE OutColor		: SV_Target0
+	, out HALF4_TYPE OutGBufferA	: SV_Target1
+	, out HALF4_TYPE OutGBufferB	: SV_Target2
+	, out HALF4_TYPE OutGBufferC	: SV_Target3  
 #endif
 #else
-	, out half4 OutColor	: SV_Target0
+	, out HALF4_TYPE OutColor		: SV_Target0
 #endif
 #if USE_SCENE_DEPTH_AUX
 	, out float OutSceneDepthAux : SV_Target4
diff --git a/Engine/Shaders/Private/MobileDeferredShading.usf b/Engine/Shaders/Private/MobileDeferredShading.usf
--- a/Engine/Shaders/Private/MobileDeferredShading.usf
+++ b/Engine/Shaders/Private/MobileDeferredShading.usf
@@ -236,7 +236,7 @@ half3 GetLightGridLocalLighting(
 void MobileDirectLightPS(
 	noperspective float4 UVAndScreenPos : TEXCOORD0, 
 	float4 SvPosition : SV_POSITION, 
-	out half4 OutColor : SV_Target0)
+	out HALF4_TYPE OutColor : SV_Target0)
 {
 	FGBufferData GBuffer = (FGBufferData)0;
 	float SceneDepth = 0; 
@@ -336,7 +336,7 @@ void MobileDirectLightPS(
 void MobileRadialLightPS(
 	float4 InScreenPosition : TEXCOORD0,
 	float4 SVPos			: SV_POSITION,
-	out half4 OutColor		: SV_Target0
+	out float4 OutColor		: SV_Target0
 )
 {
 	FGBufferData GBuffer = (FGBufferData)0;
diff --git a/Engine/Shaders/Private/PostProcessAmbientOcclusionMobile.usf b/Engine/Shaders/Private/PostProcessAmbientOcclusionMobile.usf
--- a/Engine/Shaders/Private/PostProcessAmbientOcclusionMobile.usf
+++ b/Engine/Shaders/Private/PostProcessAmbientOcclusionMobile.usf
@@ -13,11 +13,11 @@
 float4 FadeRadiusMulAdd_FadeDistance_AttenFactor;
 float4 WorldRadiusAdj_SinDeltaAngle_CosDeltaAngle_Thickness;
 
-half4 Power_Intensity_ScreenPixelsToSearch;
-half4 DepthBufferSizeAndInvSize;
-half4 BufferSizeAndInvSize;
-half4 ViewSizeAndInvSize;
-half4 ViewRectMin;
+float4 Power_Intensity_ScreenPixelsToSearch;
+float4 DepthBufferSizeAndInvSize;
+float4 BufferSizeAndInvSize;
+float4 ViewSizeAndInvSize;
+float4 ViewRectMin;
 
 half2 TexturePosToBufferUV(int2 TexturePos)
 {
@@ -836,7 +836,7 @@ void GTAOSpatialFilterCS(
 void GTAOHorizonSearchIntegralPS(
 	float4 InUVPos : TEXCOORD0,
 	in float4 SvPosition : SV_Position,
-	out half4 OutColor : SV_Target0
+	out HALF4_TYPE OutColor : SV_Target0
 )
 {
 	ResolvedView = ResolveView();
@@ -855,7 +855,7 @@ void GTAOHorizonSearchIntegralPS(
 void GTAOSpatialFilterPS(
 	float4 InUVPos : TEXCOORD0,
 	in float4 SvPosition : SV_Position, 
-	out half4 OutColor : SV_Target0)
+	out HALF4_TYPE OutColor : SV_Target0)
 {
 	ResolvedView = ResolveView();
 
diff --git a/Engine/Shaders/Private/PostProcessMobile.usf b/Engine/Shaders/Private/PostProcessMobile.usf
--- a/Engine/Shaders/Private/PostProcessMobile.usf
+++ b/Engine/Shaders/Private/PostProcessMobile.usf
@@ -139,9 +139,9 @@ half2 SunConstDepthMaskScaleBias()
 
 void SunMaskPS_Mobile(
 	float4 InUVPos : TEXCOORD0,
-	out half OutSunShaftAndDof : SV_Target0
+	out HALF_TYPE OutSunShaftAndDof : SV_Target0
 #if !MOBILE_USEDEPTHTEXTURE
-	, out half4 OutColor : SV_Target1
+	, out HALF4_TYPE OutColor : SV_Target1
 #endif
 	)
 {
@@ -195,7 +195,7 @@ void SunMaskPS_Mobile(
 
 void PreTonemapMSAA_Mobile(
 	float4 InUVPos : TEXCOORD0,
-	out half4 OutColor : SV_Target0
+	out HALF4_TYPE OutColor : SV_Target0
 	)
 {
 #if (METAL_PROFILE && !MAC)
@@ -210,7 +210,7 @@ SamplerState InputSampler;
 
 void MSAADecodeAndCopyRectPS(
 	noperspective float4 UVAndScreenPos : TEXCOORD0,
-	out half4 OutColor : SV_Target0
+	out HALF4_TYPE OutColor : SV_Target0
 	)
 {
 #if (METAL_PROFILE && !MAC)
@@ -244,20 +244,20 @@ void BloomPS_Mobile(
 	float2 InUVs[4] : TEXCOORD0
 #if MOBILE_USEBLOOM
 
-	, out half4 OutColor : SV_Target0
+	, out HALF4_TYPE OutColor : SV_Target0
 
 	#if MOBILE_USEDOF || MOBILE_USESUN
 
-		,out half OutSunShaftAndDof : SV_Target1
+		,out HALF_TYPE OutSunShaftAndDof : SV_Target1
 
 		#if MOBILE_USEEYEADAPTATION
-			, out half OutEyeAdaptation : SV_Target2
+			, out HALF_TYPE OutEyeAdaptation : SV_Target2
 		#endif
 
 	#else
 
 		#if MOBILE_USEEYEADAPTATION
-			, out half OutEyeAdaptation : SV_Target1
+			, out HALF_TYPE OutEyeAdaptation : SV_Target1
 		#endif
 
 	#endif
@@ -267,16 +267,16 @@ void BloomPS_Mobile(
 
 	#if MOBILE_USEDOF || MOBILE_USESUN
 
-		, out half OutSunShaftAndDof : SV_Target0
+		, out HALF_TYPE OutSunShaftAndDof : SV_Target0
 
 		#if MOBILE_USEEYEADAPTATION
-			, out half OutEyeAdaptation : SV_Target1
+			, out HALF_TYPE OutEyeAdaptation : SV_Target1
 		#endif
 
 	#else
 
 		#if MOBILE_USEEYEADAPTATION
-			, out half OutEyeAdaptation : SV_Target0
+			, out HALF_TYPE OutEyeAdaptation : SV_Target0
 		#endif
 
 	#endif
@@ -428,7 +428,7 @@ void BloomDownVS_Mobile(
 
 void BloomDownPS_Mobile(
 	float4 InUVs[8] : TEXCOORD0,
-	out half4 OutColor : SV_Target0
+	out HALF4_TYPE OutColor : SV_Target0
 	)
 {
 	half4 N0 = BloomDownSourceTexture.Sample(BloomDownSourceSampler, InUVs[0].xy).rgba;
@@ -517,7 +517,7 @@ float4 BloomTintB;
 
 void BloomUpPS_Mobile(
 	float4 InUVs[8] : TEXCOORD0,
-	out half4 OutColor : SV_Target0
+	out HALF4_TYPE OutColor : SV_Target0
 	)
 {
 	half3 A0 = BloomUpSourceATexture.Sample(BloomUpSourceASampler, InUVs[0].xy).rgb;
@@ -592,7 +592,7 @@ void DofNearVS_Mobile(
 void DofNearPS_Mobile(
 	float2 InUVs2 : TEXCOORD0,
 	float4 InUVs[4] : TEXCOORD1,
-	out half OutColor : SV_Target0
+	out HALF_TYPE OutColor : SV_Target0
 	)
 {
 	half N0 = SunShaftAndDofTexture.Sample(SunShaftAndDofSampler, InUVs2).r;
@@ -662,7 +662,7 @@ void DofDownVS_Mobile(
 
 void DofDownPS_Mobile(
 	float2 InUVs[5] : TEXCOORD0,
-	out half4 OutColor : SV_Target0
+	out HALF4_TYPE OutColor : SV_Target0
 	)
 {
 	// This shader needs float precision to work.
@@ -797,7 +797,7 @@ void DofBlurVS_Mobile(
 
 void DofBlurPS_Mobile(
 	float2 InUVs[5] : TEXCOORD0,
-	out half4 OutColor : SV_Target0
+	out HALF4_TYPE OutColor : SV_Target0
 	)
 {
 	// Near diolation size is copied into alpha for the tonemapper pass.
@@ -852,7 +852,7 @@ void IntegrateDOFVS_Mobile(
 void IntegrateDOFPS_Mobile(
 	in float4 TexCoords : TEXCOORD0,
 	in float2 FineDofGrain : TEXCOORD1,
-	out half4 OutColor : SV_Target0
+	out HALF4_TYPE OutColor : SV_Target0
 )
 {
 	half4 SceneColor = Texture2DSample(SceneColorTexture, SceneColorSampler, TexCoords.xy);
@@ -959,7 +959,7 @@ half SunUnBias(half A)
 
 void SunAlphaPS_Mobile(
 	float2 InUVs[8] : TEXCOORD0,
-	out half OutColor : SV_Target0
+	out HALF_TYPE OutColor : SV_Target0
 	)
 {
 	OutColor = 
@@ -1006,7 +1006,7 @@ void SunBlurVS_Mobile(
 
 void SunBlurPS_Mobile(
 	float2 InUVs[8] : TEXCOORD0,
-	out half OutColor : SV_Target0
+	out HALF4_TYPE OutColor : SV_Target0
 	)
 {
 	OutColor = 
@@ -1076,7 +1076,7 @@ float4 BloomDirtMaskTint;
 void SunMergePS_Mobile(
 	float4 InUVVignette : TEXCOORD0,
 	float4 InUVs[7] : TEXCOORD1,
-	out half4 OutColor : SV_Target0
+	out HALF4_TYPE OutColor : SV_Target0
 	)
 {
 	#if MOBILE_USEBLOOM
@@ -1142,7 +1142,7 @@ void SunMergePS_Mobile(
 
 void SunAvgPS_Mobile(
 	float4 InUV : TEXCOORD0,
-	out half4 OutColor : SV_Target0
+	out HALF4_TYPE OutColor : SV_Target0
 	)
 {
 	OutColor = 
@@ -1184,7 +1184,7 @@ void AaVS_Mobile(
 
 void AaPS_Mobile(
 	float2 UV[6] : TEXCOORD0,
-	out half4 OutColor : SV_Target0
+	out HALF4_TYPE OutColor : SV_Target0
 	)
 {
 	//   N
diff --git a/Engine/Shaders/Private/PostProcessPixelProjectedReflectionMobile.usf b/Engine/Shaders/Private/PostProcessPixelProjectedReflectionMobile.usf
--- a/Engine/Shaders/Private/PostProcessPixelProjectedReflectionMobile.usf
+++ b/Engine/Shaders/Private/PostProcessPixelProjectedReflectionMobile.usf
@@ -32,9 +32,9 @@ Texture2D SceneDepthTexture;
 SamplerState SceneDepthSampler;
 
 float4 ReflectionPlane;
-half4 BufferSizeAndInvSize;
-half4 ViewSizeAndInvSize;
-half4 ViewRectMin;
+float4 BufferSizeAndInvSize;
+float4 ViewSizeAndInvSize;
+float4 ViewRectMin;
 
 float2 ProjectionPixelToBufferUV(int2 ReflectionPixel)
 {
@@ -297,7 +297,7 @@ void ReflectionPassPS(
 #if MOBILE_MULTI_VIEW
 	in float MultiViewId : VIEW_ID,
 #endif
-	out half4 OutColor : SV_Target0)
+	out HALF4_TYPE OutColor : SV_Target0)
 {
 #if MOBILE_MULTI_VIEW
 	ResolvedView = ResolveView(uint(MultiViewId));
diff --git a/Engine/Shaders/Private/SlateElementPixelShader.usf b/Engine/Shaders/Private/SlateElementPixelShader.usf
--- a/Engine/Shaders/Private/SlateElementPixelShader.usf
+++ b/Engine/Shaders/Private/SlateElementPixelShader.usf
@@ -317,7 +317,7 @@ half4 GetLineSegmentElementColor( VertexToPixelInterpolants InVertex )
 } 
 #endif
 
-half4 Main( VertexToPixelInterpolants VIn ) : SV_Target0
+HALF4_TYPE Main( VertexToPixelInterpolants VIn ) : SV_Target0
 {
 	half4 OutColor;
 #if SHADER_TYPE == ST_Default || SHADER_TYPE == ST_Custom
diff --git a/Engine/Shaders/Public/Platform.ush b/Engine/Shaders/Public/Platform.ush
--- a/Engine/Shaders/Public/Platform.ush
+++ b/Engine/Shaders/Public/Platform.ush
@@ -458,6 +458,18 @@ void ShaderYield()
 #define COMPILER_SUPPORTS_GATHER_UINT 0
 #endif
 
+#if !COMPILER_HLSLCC && ES3_1_PROFILE
+#define HALF_TYPE half
+#define HALF2_TYPE half2
+#define HALF3_TYPE half3
+#define HALF4_TYPE half4
+#else
+#define HALF_TYPE float
+#define HALF2_TYPE float2
+#define HALF3_TYPE float3
+#define HALF4_TYPE float4
+#endif
+
 // ---------------------------------------------------- Compiler attributes
 
 #if SM6_PROFILE || SM5_PROFILE || COMPILER_SUPPORTS_ATTRIBUTES
