Add missing shader update for Strata classification

diff --git a/Engine/Shaders/Private/Strata/StrataMaterialClassification.usf b/Engine/Shaders/Private/Strata/StrataMaterialClassification.usf
--- a/Engine/Shaders/Private/Strata/StrataMaterialClassification.usf
+++ b/Engine/Shaders/Private/Strata/StrataMaterialClassification.usf
@@ -265,7 +265,8 @@ Buffer<uint> TileListBuffer;
 void StencilMainVS(
 	in uint InVertexId : SV_VertexID,
 	in uint InInstanceId : SV_InstanceID,
-	out float4 OutPosition : SV_POSITION)
+	out float4 OutPosition : SV_POSITION,
+	out uint2 OutTileCoord : TILE_COORD)
 {
 	const uint2 TileCoord = DecodeTile(TileListBuffer[InInstanceId]);
 
@@ -273,13 +274,18 @@ void StencilMainVS(
 	TileVertex.x += InVertexId == 1 || InVertexId == 2 || InVertexId == 4 ? TILE_SIZE : 0;
 	TileVertex.y += InVertexId == 2 || InVertexId == 4 || InVertexId == 5 ? TILE_SIZE : 0;
 	OutPosition = float4(float2(TileVertex) / OutputResolution * float2(2.0f, -2.0f) + float2(-1.0, 1.0f), 0.5f, 1.0f);
+	OutTileCoord = TileCoord;
 }
 
 void StencilMainPS(
-	float4 SVPos : SV_POSITION,
-	out uint OutColor0 : SV_Target0)
+	in float4 SVPos : SV_POSITION,
+	in uint2 InTileCoord : TILE_COORD,
+	out float4 OutColor0 : SV_Target0)
 {
-	OutColor0 = 255;
+	const bool bTileX = (InTileCoord.x & 1) == 0;
+	const bool bTileY = (InTileCoord.y & 1) == 0;
+	const bool bChecker = (bTileX && bTileY) || (!bTileX && !bTileY);
+	OutColor0 = bChecker ? float4(0.0f, 0.5f, 0.f, 1.0f) : float4(0.f, 1.0f, 0.f, 1.0f);
 }
 
 #endif //SHADER_STENCIL_CATEGORIZATION
